Fix VS2017 compile error.

diff --git a/Engine/Source/Runtime/D3D12RHI/Private/D3D12StateCachePrivate.h b/Engine/Source/Runtime/D3D12RHI/Private/D3D12StateCachePrivate.h
--- a/Engine/Source/Runtime/D3D12RHI/Private/D3D12StateCachePrivate.h
+++ b/Engine/Source/Runtime/D3D12RHI/Private/D3D12StateCachePrivate.h
@@ -713,7 +713,7 @@ public:
 	}
 
 	template <bool IsCompute = false>
-	D3D12_STATE_CACHE_INLINE void FD3D12StateCacheBase::SetPipelineState(FD3D12PipelineState* PSO)
+	D3D12_STATE_CACHE_INLINE void SetPipelineState(FD3D12PipelineState* PSO)
 	{
 		// Save the PSO
 		if (PSO)
