Fix an innocuous SA warning

diff --git a/Engine/Source/Runtime/D3D12RHI/Private/D3D12Resources.cpp b/Engine/Source/Runtime/D3D12RHI/Private/D3D12Resources.cpp
--- a/Engine/Source/Runtime/D3D12RHI/Private/D3D12Resources.cpp
+++ b/Engine/Source/Runtime/D3D12RHI/Private/D3D12Resources.cpp
@@ -311,13 +311,13 @@ HRESULT FD3D12Adapter::CreatePlacedResourceWithHeap(const D3D12_RESOURCE_DESC& I
 	HeapDesc.SizeInBytes = ResInfo.SizeInBytes;
 	HeapDesc.Alignment = 0;
 	HeapDesc.Flags = D3D12_HEAP_FLAG_NONE;
-	if (D3D12RHI_RESOURCE_FLAG_ALLOW_INDIRECT_BUFFER != 0)
+	
+#if D3D12RHI_RESOURCE_FLAG_ALLOW_INDIRECT_BUFFER != 0
+	if (InDesc.Flags & D3D12RHI_RESOURCE_FLAG_ALLOW_INDIRECT_BUFFER)
 	{
-		if (InDesc.Flags & D3D12RHI_RESOURCE_FLAG_ALLOW_INDIRECT_BUFFER)
-		{
-			HeapDesc.Flags |= D3D12RHI_HEAP_FLAG_ALLOW_INDIRECT_BUFFERS;
-		}
+		HeapDesc.Flags |= D3D12RHI_HEAP_FLAG_ALLOW_INDIRECT_BUFFERS;
 	}
+#endif
 	hr = RootDevice->CreateHeap(&HeapDesc, IID_PPV_ARGS(D3DHeap.GetInitReference()));
 	check(SUCCEEDED(hr));
 	if (SUCCEEDED(hr))
