Fix PC/editor warning introduced with previous change

diff --git a/Engine/Source/Runtime/D3D12RHI/Private/D3D12Resources.cpp b/Engine/Source/Runtime/D3D12RHI/Private/D3D12Resources.cpp
--- a/Engine/Source/Runtime/D3D12RHI/Private/D3D12Resources.cpp
+++ b/Engine/Source/Runtime/D3D12RHI/Private/D3D12Resources.cpp
@@ -312,7 +312,7 @@ HRESULT FD3D12Adapter::CreatePlacedResourceWithHeap(const D3D12_RESOURCE_DESC& I
 	HeapDesc.Alignment = 0;
 	HeapDesc.Flags = D3D12_HEAP_FLAG_NONE;
 	
-#if D3D12RHI_RESOURCE_FLAG_ALLOW_INDIRECT_BUFFER != 0
+#if PLATFORM_XBOXONE
 	if (InDesc.Flags & D3D12RHI_RESOURCE_FLAG_ALLOW_INDIRECT_BUFFER)
 	{
 		HeapDesc.Flags |= D3D12RHI_HEAP_FLAG_ALLOW_INDIRECT_BUFFERS;
