Fix compile error on VS2017 with /permissive- option.

diff --git a/Engine/Source/Runtime/D3D12RHI/Public/D3D12Viewport.h b/Engine/Source/Runtime/D3D12RHI/Public/D3D12Viewport.h
--- a/Engine/Source/Runtime/D3D12RHI/Public/D3D12Viewport.h
+++ b/Engine/Source/Runtime/D3D12RHI/Public/D3D12Viewport.h
@@ -107,7 +107,7 @@ public:
 	void IssueFrameEvent();
 
 	IDXGISwapChain1* GetSwapChain() const { return SwapChain1; }
-	IDXGISwapChain1* GetSDRSwapChain() const { return (PixelFormat == SDRPixelFormat) ? GetSwapChain() : SDRSwapChain1; }
+	IDXGISwapChain1* GetSDRSwapChain() const { return (PixelFormat == SDRPixelFormat) ? GetSwapChain() : (IDXGISwapChain1*)SDRSwapChain1; }
 
 	virtual void* GetNativeSwapChain() const override { return GetSwapChain(); }
 	virtual void* GetNativeBackBufferTexture() const override { return GetBackBuffer_RHIThread()->GetResource(); }
