Fix compile-time errors when /permissive- compiler option is set

diff --git a/Engine/Source/Runtime/D3D12RHI/Private/D3D12View.h b/Engine/Source/Runtime/D3D12RHI/Private/D3D12View.h
--- a/Engine/Source/Runtime/D3D12RHI/Private/D3D12View.h
+++ b/Engine/Source/Runtime/D3D12RHI/Private/D3D12View.h
@@ -630,7 +630,8 @@ public:
 		: FD3D12DeviceChild(InParentDevice)
 	{
 		// Allocate descriptor slot
-		FD3D12OfflineDescriptorManager& DescriptorAllocator = GetParentDevice()->GetViewDescriptorAllocator<TDesc>();
+		FD3D12Device* Device = GetParentDevice();
+		FD3D12OfflineDescriptorManager& DescriptorAllocator = Device->template GetViewDescriptorAllocator<TDesc>();
 		Handle = DescriptorAllocator.AllocateHeapSlot(Index);
 		check(Handle.ptr != 0);
 	}
@@ -638,7 +639,8 @@ public:
 	~TD3D12ViewDescriptorHandle()
 	{
 		// Free descriptor slot
-		FD3D12OfflineDescriptorManager& DescriptorAllocator = GetParentDevice()->GetViewDescriptorAllocator<TDesc>();
+		FD3D12Device* Device = GetParentDevice();
+		FD3D12OfflineDescriptorManager& DescriptorAllocator = Device->template GetViewDescriptorAllocator<TDesc>();
 		DescriptorAllocator.FreeHeapSlot(Handle, Index);
 		Handle.ptr = 0;
 	}
diff --git a/Engine/Source/Runtime/ShaderCore/Public/Shader.h b/Engine/Source/Runtime/ShaderCore/Public/Shader.h
--- a/Engine/Source/Runtime/ShaderCore/Public/Shader.h
+++ b/Engine/Source/Runtime/ShaderCore/Public/Shader.h
@@ -1657,11 +1657,11 @@ public:
 	void GetShaderList(TMap<FShaderId, FShader*>& OutShaders) const
 	{
 		check(bHasBeenRegistered);
-		for(TMap<FShaderPrimaryKey,TRefCountPtr<FShader> >::TConstIterator ShaderIt(Shaders);ShaderIt;++ShaderIt)
+		for (typename TMap<FShaderPrimaryKey, TRefCountPtr<FShader>>::TConstIterator ShaderIt(Shaders); ShaderIt; ++ShaderIt)
 		{
-			if(ShaderIt.Value())
+			if (ShaderIt.Value())
 			{
-				OutShaders.Add(ShaderIt.Value()->GetId(),ShaderIt.Value());
+				OutShaders.Add(ShaderIt.Value()->GetId(), ShaderIt.Value());
 			}
 		}
 	}
