Copying //UE4/Dev-Rendering@4072789 to Dev-Main (//UE4/Dev-Main)
Copy missing DX12 compile fixes.

diff --git a/Engine/Source/Runtime/D3D12RHI/Public/D3D12Util.h b/Engine/Source/Runtime/D3D12RHI/Public/D3D12Util.h
--- a/Engine/Source/Runtime/D3D12RHI/Public/D3D12Util.h
+++ b/Engine/Source/Runtime/D3D12RHI/Public/D3D12Util.h
@@ -32,16 +32,6 @@
 
 #endif
 
-#if PLATFORM_XBOXONE
-struct D3D12_RT_FORMAT_ARRAY
-{
-    DXGI_FORMAT RTFormats[D3D12_SIMULTANEOUS_RENDER_TARGET_COUNT];
-    UINT NumRenderTargets;
-};
-typedef D3D12_DEPTH_STENCIL_DESC D3D12_DEPTH_STENCIL_DESC1;
-typedef CD3DX12_DEPTH_STENCIL_DESC CD3DX12_DEPTH_STENCIL_DESC1;
-#endif
-
 namespace D3D12RHI
 {
 	/**
