Fix ambiguous use of dependent name with latest version of MSVC in /permissive- mode.

diff --git a/Engine/Source/Runtime/D3D12RHI/Private/Windows/WindowsD3D12PipelineState.cpp b/Engine/Source/Runtime/D3D12RHI/Private/Windows/WindowsD3D12PipelineState.cpp
--- a/Engine/Source/Runtime/D3D12RHI/Private/Windows/WindowsD3D12PipelineState.cpp
+++ b/Engine/Source/Runtime/D3D12RHI/Private/Windows/WindowsD3D12PipelineState.cpp
@@ -622,7 +622,7 @@ ID3D12PipelineState* CreatePipelineStateWrapper(FD3D12Adapter* Adapter, const TD
 	ID3D12Device2* const pDevice2 = Adapter->GetD3DDevice2();
 	if (pDevice2)
 	{
-		TPSOStreamFunctionMap<TDesc>::D3D12PipelineStateStreamType Stream = CreationArgs->Desc->Desc.PipelineStateStream();
+		typename TPSOStreamFunctionMap<TDesc>::D3D12PipelineStateStreamType Stream = CreationArgs->Desc->Desc.PipelineStateStream();
 		const D3D12_PIPELINE_STATE_STREAM_DESC StreamDesc = { sizeof(Stream), &Stream };
 		return CreatePipelineStateFromStream<TDesc>(pDevice2, &StreamDesc, static_cast<ID3D12PipelineLibrary1*>(CreationArgs->Library), *Name);	// Static cast to ID3D12PipelineLibrary1 since we already checked for ID3D12Device2.
 	}
