Re-enabled pix on windows

diff --git a/Engine/Source/Programs/UnrealBuildTool/Platform/Windows/UEBuildWindows.cs b/Engine/Source/Programs/UnrealBuildTool/Platform/Windows/UEBuildWindows.cs
--- a/Engine/Source/Programs/UnrealBuildTool/Platform/Windows/UEBuildWindows.cs
+++ b/Engine/Source/Programs/UnrealBuildTool/Platform/Windows/UEBuildWindows.cs
@@ -106,7 +106,7 @@ namespace UnrealBuildTool
 		/// Enable PIX debugging (automatically disabled in Shipping and Test configs)
 		/// </summary>
 		[ConfigFile(ConfigHierarchyType.Engine, "/Script/WindowsTargetPlatform.WindowsTargetSettings", "bEnablePIXProfiling")]
-		public bool bPixProfilingEnabled = false;
+		public bool bPixProfilingEnabled = true;
 
 		/// <summary>
 		/// The name of the company (author, provider) that created the project.
diff --git a/Engine/Source/Runtime/D3D12RHI/Private/D3D12Device.cpp b/Engine/Source/Runtime/D3D12RHI/Private/D3D12Device.cpp
--- a/Engine/Source/Runtime/D3D12RHI/Private/D3D12Device.cpp
+++ b/Engine/Source/Runtime/D3D12RHI/Private/D3D12Device.cpp
@@ -116,6 +116,11 @@ bool FD3D12Device::IsGPUIdle()
 	return Fence.IsFenceComplete(Fence.GetLastSignaledFence());
 }
 
+#if PLATFORM_WINDOWS
+typedef HRESULT(WINAPI *FDXGIGetDebugInterface1)(UINT, REFIID, void **);
+#endif
+
+
 void FD3D12Device::SetupAfterDeviceCreation()
 {
 	ID3D12Device* Direct3DDevice = GetParentAdapter()->GetD3DDevice();
@@ -144,18 +149,34 @@ void FD3D12Device::SetupAfterDeviceCreation()
 		// Running on AMD with RGP profiling enabled, so enable capturing mode
 		bUnderGPUCapture = true;
 	}
-
 #if USE_PIX
-	// PIX
+	// PIX (note that DXGIGetDebugInterface1 requires Windows 8.1 and up)
+	if (FWindowsPlatformMisc::VerifyWindowsVersion(6, 3))
 	{
-		IID GraphicsAnalysisID;
-		if (SUCCEEDED(IIDFromString(L"{9F251514-9D4D-4902-9D60-18988AB7D4B5}", &GraphicsAnalysisID)))
+		FDXGIGetDebugInterface1 DXGIGetDebugInterface1FnPtr = nullptr;
+
+		// CreateDXGIFactory2 is only available on Win8.1+, find it if it exists
+		HMODULE DxgiDLL = LoadLibraryA("dxgi.dll");
+		if (DxgiDLL)
 		{
-			TRefCountPtr<IUnknown> GraphicsAnalysis;
-			if (SUCCEEDED(DXGIGetDebugInterface1(0, GraphicsAnalysisID, (void**)GraphicsAnalysis.GetInitReference())))
+#pragma warning(push)
+#pragma warning(disable: 4191) // disable the "unsafe conversion from 'FARPROC' to 'blah'" warning
+			DXGIGetDebugInterface1FnPtr = (FDXGIGetDebugInterface1)(GetProcAddress(DxgiDLL, "DXGIGetDebugInterface1"));
+#pragma warning(pop)
+			FreeLibrary(DxgiDLL);
+		}
+		
+		if (DXGIGetDebugInterface1FnPtr)
+		{
+			IID GraphicsAnalysisID;
+			if (SUCCEEDED(IIDFromString(L"{9F251514-9D4D-4902-9D60-18988AB7D4B5}", &GraphicsAnalysisID)))
 			{
-				// Running under PIX, so enable capturing mode
-				bUnderGPUCapture = true;
+				TRefCountPtr<IUnknown> GraphicsAnalysis;
+				if (SUCCEEDED(DXGIGetDebugInterface1FnPtr(0, GraphicsAnalysisID, (void**)GraphicsAnalysis.GetInitReference())))
+				{
+					// Running under PIX, so enable capturing mode
+					bUnderGPUCapture = true;
+				}
 			}
 		}
 	}
