D3D12Viewport

- Fix for D3D12 Invalid Error message

diff --git a/Engine/Source/Runtime/D3D12RHI/Private/Windows/WindowsD3D12Viewport.cpp b/Engine/Source/Runtime/D3D12RHI/Private/Windows/WindowsD3D12Viewport.cpp
--- a/Engine/Source/Runtime/D3D12RHI/Private/Windows/WindowsD3D12Viewport.cpp
+++ b/Engine/Source/Runtime/D3D12RHI/Private/Windows/WindowsD3D12Viewport.cpp
@@ -85,6 +85,7 @@ void FD3D12Viewport::Init()
 	const DXGI_MODE_DESC BufferDesc = SetupDXGI_MODE_DESC();
 
 	// Create the swapchain.
+	if(Adapter->GetOwningRHI()->IsQuadBufferStereoEnabled())
 	{
 		TRefCountPtr<IDXGIFactory2> Factory2;
 		Factory->QueryInterface(IID_PPV_ARGS(Factory2.GetInitReference()));
