Changed D3D12 so that we don't use the boundshadestate anymore in the state cache. This fixes an issue where reloading the same map again and again ends up using more and more memory.

diff --git a/Engine/Source/Runtime/D3D12RHI/Private/D3D12Commands.cpp b/Engine/Source/Runtime/D3D12RHI/Private/D3D12Commands.cpp
--- a/Engine/Source/Runtime/D3D12RHI/Private/D3D12Commands.cpp
+++ b/Engine/Source/Runtime/D3D12RHI/Private/D3D12Commands.cpp
@@ -432,32 +432,8 @@ void FD3D12CommandContext::RHISetScissorRect(bool bEnable, uint32 MinX, uint32 M
 */
 void FD3D12CommandContext::RHISetBoundShaderState(FBoundShaderStateRHIParamRef BoundShaderStateRHI)
 {
-	//SCOPE_CYCLE_COUNTER(STAT_D3D12SetBoundShaderState);
-	FD3D12BoundShaderState* const BoundShaderState = FD3D12DynamicRHI::ResourceCast(BoundShaderStateRHI);
-
-	StateCache.SetBoundShaderState(BoundShaderState);
-
-	if (BoundShaderState->GetHullShader() && BoundShaderState->GetDomainShader())
-	{
-		bUsingTessellation = true;
-
-		// Ensure the command buffers are reset to reduce the amount of data that needs to be versioned.
-		HSConstantBuffer.Reset();
-		DSConstantBuffer.Reset();
-	}
-	else
-	{
-		bUsingTessellation = false;
-	}
-
-	// @TODO : really should only discard the constants if the shader state has actually changed.
-	bDiscardSharedConstants = true;
-
-	// Ensure the command buffers are reset to reduce the amount of data that needs to be versioned.
-	VSConstantBuffer.Reset();
-	PSConstantBuffer.Reset();
-	GSConstantBuffer.Reset();
-	CSConstantBuffer.Reset();	// Should this be here or in RHISetComputeShader? Might need a new bDiscardSharedConstants for CS.
+	// This shouldn't be used. Instead use RHISetGraphicsPipelineState.
+	unimplemented();
 }
 
 D3D_PRIMITIVE_TOPOLOGY GetD3D12PrimitiveType(uint32 PrimitiveType, bool bUsingTessellation)
@@ -541,10 +517,19 @@ void FD3D12CommandContext::RHISetGraphicsPipelineState(FGraphicsPipelineStateRHI
 	FD3D12GraphicsPipelineState* GraphicsPipelineState = FD3D12DynamicRHI::ResourceCast(GraphicsState);
 
 	// TODO: [PSO API] Every thing inside this scope is only necessary to keep the PSO shadow in sync while we convert the high level to only use PSOs
-	{
-		RHISetBoundShaderState(GraphicsPipelineState->BoundShaderState);
-		RHIEnableDepthBoundsTest(GraphicsPipelineState->PipelineStateInitializer.bDepthBounds);
-	}
+	bUsingTessellation = GraphicsPipelineState->GetHullShader() && GraphicsPipelineState->GetDomainShader();
+	// Ensure the command buffers are reset to reduce the amount of data that needs to be versioned.
+	VSConstantBuffer.Reset();
+	PSConstantBuffer.Reset();
+	HSConstantBuffer.Reset();
+	DSConstantBuffer.Reset();
+	GSConstantBuffer.Reset();
+	// Should this be here or in RHISetComputeShader? Might need a new bDiscardSharedConstants for CS.
+	CSConstantBuffer.Reset();
+	// @TODO : really should only discard the constants if the shader state has actually changed.
+	bDiscardSharedConstants = true;
+
+	RHIEnableDepthBoundsTest(GraphicsPipelineState->PipelineStateInitializer.bDepthBounds);
 
 	StateCache.SetGraphicsPipelineState(GraphicsPipelineState);
 }
@@ -1197,13 +1182,13 @@ void FD3D12CommandContext::CommitNonComputeShaderConstants()
 {
 	//SCOPE_CYCLE_COUNTER(STAT_D3D12CommitGraphicsConstants);
 
-	FD3D12BoundShaderState* RESTRICT CurrentBoundShaderStateRef = StateCache.GetBoundShaderState();
+	const FD3D12GraphicsPipelineState* const RESTRICT GraphicPSO = StateCache.GetGraphicsPipelineState();
 
-	check(CurrentBoundShaderStateRef);
+	check(GraphicPSO);
 
 	// Only set the constant buffer if this shader needs the global constant buffer bound
 	// Otherwise we will overwrite a different constant buffer
-	if (CurrentBoundShaderStateRef->bShaderNeedsGlobalConstantBuffer[SF_Vertex])
+	if (GraphicPSO->bShaderNeedsGlobalConstantBuffer[SF_Vertex])
 	{
 		StateCache.SetConstantBuffer<SF_Vertex>(VSConstantBuffer, bDiscardSharedConstants);
 	}
@@ -1214,23 +1199,23 @@ void FD3D12CommandContext::CommitNonComputeShaderConstants()
 	// is always reset whenever bUsingTessellation changes in SetBoundShaderState()
 	if (bUsingTessellation)
 	{
-		if (CurrentBoundShaderStateRef->bShaderNeedsGlobalConstantBuffer[SF_Hull])
+		if (GraphicPSO->bShaderNeedsGlobalConstantBuffer[SF_Hull])
 		{
 			StateCache.SetConstantBuffer<SF_Hull>(HSConstantBuffer, bDiscardSharedConstants);
 		}
 
-		if (CurrentBoundShaderStateRef->bShaderNeedsGlobalConstantBuffer[SF_Domain])
+		if (GraphicPSO->bShaderNeedsGlobalConstantBuffer[SF_Domain])
 		{
 			StateCache.SetConstantBuffer<SF_Domain>(DSConstantBuffer, bDiscardSharedConstants);
 		}
 	}
 
-	if (CurrentBoundShaderStateRef->bShaderNeedsGlobalConstantBuffer[SF_Geometry])
+	if (GraphicPSO->bShaderNeedsGlobalConstantBuffer[SF_Geometry])
 	{
 		StateCache.SetConstantBuffer<SF_Geometry>(GSConstantBuffer, bDiscardSharedConstants);
 	}
 
-	if (CurrentBoundShaderStateRef->bShaderNeedsGlobalConstantBuffer[SF_Pixel])
+	if (GraphicPSO->bShaderNeedsGlobalConstantBuffer[SF_Pixel])
 	{
 		StateCache.SetConstantBuffer<SF_Pixel>(PSConstantBuffer, bDiscardSharedConstants);
 	}
@@ -1377,26 +1362,26 @@ void FD3D12CommandContext::CommitGraphicsResourceTables()
 {
 	//SCOPE_CYCLE_COUNTER(STAT_D3D12CommitResourceTables);
 
-	const FD3D12BoundShaderState* const RESTRICT CurrentBoundShaderStateRef = StateCache.GetBoundShaderState();
-	check(CurrentBoundShaderStateRef);
+	const FD3D12GraphicsPipelineState* const RESTRICT GraphicPSO = StateCache.GetGraphicsPipelineState();
+	check(GraphicPSO);
 
-	if (auto* Shader = CurrentBoundShaderStateRef->GetVertexShader())
+	if (auto* Shader = GraphicPSO->GetVertexShader())
 	{
 		SetResourcesFromTables(Shader);
 	}
-	if (auto* Shader = CurrentBoundShaderStateRef->GetPixelShader())
+	if (auto* Shader = GraphicPSO->GetPixelShader())
 	{
 		SetResourcesFromTables(Shader);
 	}
-	if (auto* Shader = CurrentBoundShaderStateRef->GetHullShader())
+	if (auto* Shader = GraphicPSO->GetHullShader())
 	{
 		SetResourcesFromTables(Shader);
 	}
-	if (auto* Shader = CurrentBoundShaderStateRef->GetDomainShader())
+	if (auto* Shader = GraphicPSO->GetDomainShader())
 	{
 		SetResourcesFromTables(Shader);
 	}
-	if (auto* Shader = CurrentBoundShaderStateRef->GetGeometryShader())
+	if (auto* Shader = GraphicPSO->GetGeometryShader())
 	{
 		SetResourcesFromTables(Shader);
 	}
diff --git a/Engine/Source/Runtime/D3D12RHI/Private/D3D12PipelineState.h b/Engine/Source/Runtime/D3D12RHI/Private/D3D12PipelineState.h
--- a/Engine/Source/Runtime/D3D12RHI/Private/D3D12PipelineState.h
+++ b/Engine/Source/Runtime/D3D12RHI/Private/D3D12PipelineState.h
@@ -304,21 +304,21 @@ protected:
 
 struct FD3D12GraphicsPipelineState : public FRHIGraphicsPipelineState
 {
-	explicit FD3D12GraphicsPipelineState(
-		const FGraphicsPipelineStateInitializer& Initializer,
-		FD3D12BoundShaderState* InBoundShaderState,
-		FD3D12PipelineState* InPipelineState)
-		: PipelineStateInitializer(Initializer)
-		, BoundShaderState(InBoundShaderState)
-		, PipelineState(InPipelineState)
-	{
-	}
-
+	explicit FD3D12GraphicsPipelineState(const FGraphicsPipelineStateInitializer& Initializer, FD3D12BoundShaderState* InBoundShaderState, FD3D12PipelineState* InPipelineState);
 	~FD3D12GraphicsPipelineState();
 
 	const FGraphicsPipelineStateInitializer PipelineStateInitializer;
-	TRefCountPtr<FD3D12BoundShaderState> BoundShaderState;
+	const FD3D12RootSignature* RootSignature;
+	uint16 StreamStrides[MaxVertexElementCount];
+	bool bShaderNeedsGlobalConstantBuffer[SF_NumFrequencies];
+
 	FD3D12PipelineState* PipelineState;
+
+	FORCEINLINE class FD3D12VertexShader*   GetVertexShader() const { return (FD3D12VertexShader*)PipelineStateInitializer.BoundShaderState.VertexShaderRHI; }
+	FORCEINLINE class FD3D12PixelShader*    GetPixelShader() const { return (FD3D12PixelShader*)PipelineStateInitializer.BoundShaderState.PixelShaderRHI; }
+	FORCEINLINE class FD3D12HullShader*     GetHullShader() const { return (FD3D12HullShader*)PipelineStateInitializer.BoundShaderState.HullShaderRHI; }
+	FORCEINLINE class FD3D12DomainShader*   GetDomainShader() const { return (FD3D12DomainShader*)PipelineStateInitializer.BoundShaderState.DomainShaderRHI; }
+	FORCEINLINE class FD3D12GeometryShader* GetGeometryShader() const { return (FD3D12GeometryShader*)PipelineStateInitializer.BoundShaderState.GeometryShaderRHI; }
 };
 
 struct FD3D12ComputePipelineState : public FRHIComputePipelineState
@@ -379,9 +379,9 @@ protected:
 
 	// Thread access mutual exclusion
 	mutable FRWLock InitializerToGraphicsPipelineMapMutex;
-	mutable FRWLock LowLevelDescToGraphicsPipelineMapMutex;
+	mutable FRWLock LowLevelGraphicsPipelineStateCacheMutex;
 	mutable FRWLock ComputeShaderToComputePipelineMapMutex;
-	mutable FRWLock LowLevelDescToComputePipelineMapMutex;
+	mutable FRWLock ComputePipelineStateCacheMutex;
 
 	FCriticalSection DiskCachesCS;
 
diff --git a/Engine/Source/Runtime/D3D12RHI/Private/D3D12Pipelinestate.cpp b/Engine/Source/Runtime/D3D12RHI/Private/D3D12Pipelinestate.cpp
--- a/Engine/Source/Runtime/D3D12RHI/Private/D3D12Pipelinestate.cpp
+++ b/Engine/Source/Runtime/D3D12RHI/Private/D3D12Pipelinestate.cpp
@@ -287,6 +287,22 @@ ID3D12PipelineState* FD3D12PipelineState::GetPipelineState()
 	return PipelineState.GetReference();
 }
 
+FD3D12GraphicsPipelineState::FD3D12GraphicsPipelineState(
+	const FGraphicsPipelineStateInitializer& Initializer,
+	FD3D12BoundShaderState* InBoundShaderState,
+	FD3D12PipelineState* InPipelineState)
+	: PipelineStateInitializer(Initializer)
+	, RootSignature(InBoundShaderState->pRootSignature)
+	, PipelineState(InPipelineState)
+{
+	FMemory::Memcpy(StreamStrides, InBoundShaderState->StreamStrides, sizeof(StreamStrides));
+	bShaderNeedsGlobalConstantBuffer[SF_Vertex] = GetVertexShader() && GetVertexShader()->ResourceCounts.bGlobalUniformBufferUsed;
+	bShaderNeedsGlobalConstantBuffer[SF_Pixel] = GetPixelShader() && GetPixelShader()->ResourceCounts.bGlobalUniformBufferUsed;
+	bShaderNeedsGlobalConstantBuffer[SF_Hull] = GetHullShader() && GetHullShader()->ResourceCounts.bGlobalUniformBufferUsed;
+	bShaderNeedsGlobalConstantBuffer[SF_Domain] = GetDomainShader() && GetDomainShader()->ResourceCounts.bGlobalUniformBufferUsed;
+	bShaderNeedsGlobalConstantBuffer[SF_Geometry] = GetGeometryShader() && GetGeometryShader()->ResourceCounts.bGlobalUniformBufferUsed;
+}
+
 FD3D12GraphicsPipelineState::~FD3D12GraphicsPipelineState()
 {
 	// At this point the object is not safe to use in the PSO cache.
@@ -309,39 +325,51 @@ FD3D12ComputePipelineState::~FD3D12ComputePipelineState()
 
 void FD3D12PipelineStateCacheBase::CleanupPipelineStateCaches()
 {
-	// The runtime caches manage the lifetime of their FD3D12GraphicsPipelineState and FD3D12ComputePipelineState.
-	// We need to release them.
-	for (auto Pair : InitializerToGraphicsPipelineMap)
 	{
-		FD3D12GraphicsPipelineState* GraphicsPipelineState = Pair.Value;
-		ensure(GIsRHIInitialized || (!GIsRHIInitialized && GraphicsPipelineState->GetRefCount() == 1));
-		GraphicsPipelineState->Release();
+		FRWScopeLock Lock(InitializerToGraphicsPipelineMapMutex, FRWScopeLockType::SLT_Write);
+		// The runtime caches manage the lifetime of their FD3D12GraphicsPipelineState and FD3D12ComputePipelineState.
+		// We need to release them.
+		for (auto Pair : InitializerToGraphicsPipelineMap)
+		{
+			FD3D12GraphicsPipelineState* GraphicsPipelineState = Pair.Value;
+			ensure(GIsRHIInitialized || (!GIsRHIInitialized && GraphicsPipelineState->GetRefCount() == 1));
+			GraphicsPipelineState->Release();
+		}
+		InitializerToGraphicsPipelineMap.Reset();
 	}
-	InitializerToGraphicsPipelineMap.Reset();
 
-	for (auto Pair : ComputeShaderToComputePipelineMap)
 	{
-		FD3D12ComputePipelineState* ComputePipelineState = Pair.Value;
-		ensure(GIsRHIInitialized || (!GIsRHIInitialized && ComputePipelineState->GetRefCount() == 1));
-		ComputePipelineState->Release();
+		FRWScopeLock Lock(ComputeShaderToComputePipelineMapMutex, FRWScopeLockType::SLT_Write);
+		for (auto Pair : ComputeShaderToComputePipelineMap)
+		{
+			FD3D12ComputePipelineState* ComputePipelineState = Pair.Value;
+			ensure(GIsRHIInitialized || (!GIsRHIInitialized && ComputePipelineState->GetRefCount() == 1));
+			ComputePipelineState->Release();
+		}
+		ComputeShaderToComputePipelineMap.Reset();
 	}
-	ComputeShaderToComputePipelineMap.Reset();
 
-	// The low level graphics and compute maps manage the lifetime of their PSOs.
-	// We need to delete each element before we empty it.
-	for (auto Iter = LowLevelGraphicsPipelineStateCache.CreateConstIterator(); Iter; ++Iter)
 	{
-		const FD3D12PipelineState* const PipelineState = Iter.Value();
-		delete PipelineState;
+		FRWScopeLock Lock(LowLevelGraphicsPipelineStateCacheMutex, FRWScopeLockType::SLT_Write);
+		// The low level graphics and compute maps manage the lifetime of their PSOs.
+		// We need to delete each element before we empty it.
+		for (auto Iter = LowLevelGraphicsPipelineStateCache.CreateConstIterator(); Iter; ++Iter)
+		{
+			const FD3D12PipelineState* const PipelineState = Iter.Value();
+			delete PipelineState;
+		}
+		LowLevelGraphicsPipelineStateCache.Empty();
 	}
-	LowLevelGraphicsPipelineStateCache.Empty();
 
-	for (auto Iter = ComputePipelineStateCache.CreateConstIterator(); Iter; ++Iter)
 	{
-		const FD3D12PipelineState* const PipelineState = Iter.Value();
-		delete PipelineState;
+		FRWScopeLock Lock(ComputePipelineStateCacheMutex, FRWScopeLockType::SLT_Write);
+		for (auto Iter = ComputePipelineStateCache.CreateConstIterator(); Iter; ++Iter)
+		{
+			const FD3D12PipelineState* const PipelineState = Iter.Value();
+			delete PipelineState;
+		}
+		ComputePipelineStateCache.Empty();
 	}
-	ComputePipelineStateCache.Empty();
 }
 
 FD3D12GraphicsPipelineState* FD3D12PipelineStateCacheBase::AddToRuntimeCache(const FGraphicsPipelineStateInitializer& Initializer, uint32 InitializerHash, FD3D12BoundShaderState* BoundShaderState, FD3D12PipelineState* PipelineState)
@@ -368,7 +396,7 @@ FD3D12PipelineState* FD3D12PipelineStateCacheBase::FindInLowLevelCache(const FD3
 	check(Desc.CombinedHash != 0);
 
 	{
-		FRWScopeLock Lock(LowLevelDescToGraphicsPipelineMapMutex, FRWScopeLockType::SLT_ReadOnly);
+		FRWScopeLock Lock(LowLevelGraphicsPipelineStateCacheMutex, FRWScopeLockType::SLT_ReadOnly);
 		FD3D12PipelineState** Found = LowLevelGraphicsPipelineStateCache.Find(Desc);
 		if (Found)
 		{
@@ -399,7 +427,7 @@ void FD3D12PipelineStateCacheBase::AddToLowLevelCache(const FD3D12LowLevelGraphi
 
 	// Double check the desc doesn't already exist while the lock is taken.
 	// This avoids having multiple threads try to create the same PSO.
-	FRWScopeLock Lock(LowLevelDescToGraphicsPipelineMapMutex, FRWScopeLockType::SLT_Write);
+	FRWScopeLock Lock(LowLevelGraphicsPipelineStateCacheMutex, FRWScopeLockType::SLT_Write);
 	FD3D12PipelineState** const PipelineState = LowLevelGraphicsPipelineStateCache.Find(Desc);
 	if (PipelineState)
 	{
@@ -443,7 +471,7 @@ FD3D12PipelineState* FD3D12PipelineStateCacheBase::FindInLowLevelCache(const FD3
 	check(Desc.CombinedHash != 0);
 
 	{
-		FRWScopeLock Lock(LowLevelDescToComputePipelineMapMutex, FRWScopeLockType::SLT_ReadOnly);
+		FRWScopeLock Lock(ComputePipelineStateCacheMutex, FRWScopeLockType::SLT_ReadOnly);
 		FD3D12PipelineState** Found = ComputePipelineStateCache.Find(Desc);
 		if (Found)
 		{
@@ -475,7 +503,7 @@ void FD3D12PipelineStateCacheBase::AddToLowLevelCache(const FD3D12ComputePipelin
 
 	// Double check the desc doesn't already exist while the lock is taken.
 	// This avoids having multiple threads try to create the same PSO.
-	FRWScopeLock Lock(LowLevelDescToComputePipelineMapMutex, FRWScopeLockType::SLT_Write);
+	FRWScopeLock Lock(ComputePipelineStateCacheMutex, FRWScopeLockType::SLT_Write);
 	FD3D12PipelineState** const PipelineState = ComputePipelineStateCache.Find(Desc);
 	if (PipelineState)
 	{
diff --git a/Engine/Source/Runtime/D3D12RHI/Private/D3D12StateCache.cpp b/Engine/Source/Runtime/D3D12RHI/Private/D3D12StateCache.cpp
--- a/Engine/Source/Runtime/D3D12RHI/Private/D3D12StateCache.cpp
+++ b/Engine/Source/Runtime/D3D12RHI/Private/D3D12StateCache.cpp
@@ -353,8 +353,7 @@ void FD3D12StateCacheBase::ApplyState()
 #endif
 
 	FD3D12CommandListHandle& CommandList = CmdContext->CommandListHandle;
-	const FD3D12RootSignature* const pRootSignature = IsCompute ?
-		PipelineState.Compute.CurrentPipelineStateObject->ComputeShader->pRootSignature : PipelineState.Graphics.CurrentPipelineStateObject->BoundShaderState->pRootSignature;
+	const FD3D12RootSignature* const pRootSignature = IsCompute ? PipelineState.Compute.CurrentPipelineStateObject->ComputeShader->pRootSignature : GetGraphicsRootSignature();
 
 	// PSO
 	if (IsCompute)
diff --git a/Engine/Source/Runtime/D3D12RHI/Private/D3D12StateCachePrivate.h b/Engine/Source/Runtime/D3D12RHI/Private/D3D12StateCachePrivate.h
--- a/Engine/Source/Runtime/D3D12RHI/Private/D3D12StateCachePrivate.h
+++ b/Engine/Source/Runtime/D3D12RHI/Private/D3D12StateCachePrivate.h
@@ -365,6 +365,7 @@ protected:
 	{ \
 		static const EShaderFrequency Frequency = SF_##Name; \
 		static FD3D12##Name##Shader* GetShader(FD3D12BoundShaderState* BSS) { return BSS ? BSS->Get##Name##Shader() : nullptr; } \
+		static FD3D12##Name##Shader* GetShader(FD3D12GraphicsPipelineState* PSO) { return PSO ? (FD3D12##Name##Shader*)PSO->PipelineStateInitializer.BoundShaderState.##Name##ShaderRHI : nullptr; } \
 	}
 	DECLARE_SHADER_TRAITS(Vertex);
 	DECLARE_SHADER_TRAITS(Pixel);
@@ -376,7 +377,8 @@ protected:
 	template <typename TShader> D3D12_STATE_CACHE_INLINE void SetShader(TShader* Shader)
 	{
 		typedef StateCacheShaderTraits<TShader> Traits;
-		TShader* OldShader = Traits::GetShader(GetBoundShaderState());
+		TShader* OldShader = Traits::GetShader(GetGraphicsPipelineState());
+
 		if (OldShader != Shader)
 		{
 			PipelineState.Common.CurrentShaderSamplerCounts[Traits::Frequency] = (Shader) ? Shader->ResourceCounts.NumSamplers : 0;
@@ -391,7 +393,7 @@ protected:
 
 	template <typename TShader> D3D12_STATE_CACHE_INLINE void GetShader(TShader** Shader)
 	{
-		*Shader = StateCacheShaderTraits<TShader>::GetShader(GetBoundShaderState());
+		*Shader = StateCacheShaderTraits<TShader>::GetShader(GetGraphicsPipelineState());
 	}
 
 	template <bool IsCompute = false>
@@ -432,21 +434,19 @@ public:
 		return &DescriptorCache;
 	}
 
-	FD3D12BoundShaderState* GetBoundShaderState() const
+	FD3D12GraphicsPipelineState* GetGraphicsPipelineState() const
 	{
-		return PipelineState.Graphics.CurrentPipelineStateObject ? PipelineState.Graphics.CurrentPipelineStateObject->BoundShaderState : nullptr;
+		return PipelineState.Graphics.CurrentPipelineStateObject;
 	}
 
 	const FD3D12RootSignature* GetGraphicsRootSignature()
 	{
-		return PipelineState.Graphics.CurrentPipelineStateObject ?
-			PipelineState.Graphics.CurrentPipelineStateObject->BoundShaderState->pRootSignature : nullptr;
+		return PipelineState.Graphics.CurrentPipelineStateObject ? PipelineState.Graphics.CurrentPipelineStateObject->RootSignature : nullptr;
 	}
 
 	const FD3D12RootSignature* GetComputeRootSignature()
 	{
-		return PipelineState.Compute.CurrentPipelineStateObject ?
-			PipelineState.Compute.CurrentPipelineStateObject->ComputeShader->pRootSignature : nullptr;
+		return PipelineState.Compute.CurrentPipelineStateObject ? PipelineState.Compute.CurrentPipelineStateObject->ComputeShader->pRootSignature : nullptr;
 	}
 
 	void ClearSRVs();
@@ -656,46 +656,24 @@ public:
 		GetShader(Shader);
 	}
 
-	D3D12_STATE_CACHE_INLINE void SetBoundShaderState(FD3D12BoundShaderState* BoundShaderState)
+	D3D12_STATE_CACHE_INLINE void SetGraphicsPipelineState(FD3D12GraphicsPipelineState* GraphicsPipelineState)
 	{
-		if (BoundShaderState)
-		{
-			SetStreamStrides(BoundShaderState->StreamStrides);
-			SetShader(BoundShaderState->GetVertexShader());
-			SetShader(BoundShaderState->GetPixelShader());
-			SetShader(BoundShaderState->GetDomainShader());
-			SetShader(BoundShaderState->GetHullShader());
-			SetShader(BoundShaderState->GetGeometryShader());
-		}
-		else
+		check(GraphicsPipelineState);
+		if (PipelineState.Graphics.CurrentPipelineStateObject != GraphicsPipelineState)
 		{
-			uint16 NullStrides[MaxVertexElementCount] = {0};
-			SetStreamStrides(NullStrides);
-			SetShader<FD3D12VertexShader>(nullptr);
-			SetShader<FD3D12PixelShader>(nullptr);
-			SetShader<FD3D12HullShader>(nullptr);
-			SetShader<FD3D12DomainShader>(nullptr);
-			SetShader<FD3D12GeometryShader>(nullptr);
-		}
+			SetStreamStrides(GraphicsPipelineState->StreamStrides);
+			SetShader(GraphicsPipelineState->GetVertexShader());
+			SetShader(GraphicsPipelineState->GetPixelShader());
+			SetShader(GraphicsPipelineState->GetDomainShader());
+			SetShader(GraphicsPipelineState->GetHullShader());
+			SetShader(GraphicsPipelineState->GetGeometryShader());
 
-		FD3D12BoundShaderState* CurrentBoundShaderState = GetBoundShaderState();
-		if (CurrentBoundShaderState != BoundShaderState)
-		{
 			// See if we need to change the root signature
-			const FD3D12RootSignature* const pCurrentRootSignature = CurrentBoundShaderState ? CurrentBoundShaderState->pRootSignature : nullptr;
-			const FD3D12RootSignature* const pNewRootSignature = BoundShaderState ? BoundShaderState->pRootSignature : nullptr;
-			if (pCurrentRootSignature != pNewRootSignature)
+			if (GetGraphicsRootSignature() != GraphicsPipelineState->RootSignature)
 			{
 				PipelineState.Graphics.bNeedSetRootSignature = true;
 			}
-		}
-	}
 
-	D3D12_STATE_CACHE_INLINE void SetGraphicsPipelineState(FD3D12GraphicsPipelineState* GraphicsPipelineState)
-	{
-		check(GraphicsPipelineState);
-		if (PipelineState.Graphics.CurrentPipelineStateObject != GraphicsPipelineState)
-		{
 			// Save the PSO
 			PipelineState.Common.bNeedSetPSO = true;
 			PipelineState.Graphics.CurrentPipelineStateObject = GraphicsPipelineState;
