Changed warning to verbose in the d3d12 when waiting for query results.

diff --git a/Engine/Source/Runtime/D3D12RHI/Private/D3D12Query.cpp b/Engine/Source/Runtime/D3D12RHI/Private/D3D12Query.cpp
--- a/Engine/Source/Runtime/D3D12RHI/Private/D3D12Query.cpp
+++ b/Engine/Source/Runtime/D3D12RHI/Private/D3D12Query.cpp
@@ -122,14 +122,14 @@ bool FD3D12Device::GetQueryData(FD3D12RenderQuery& Query, bool bWait)
 		}
 
 		// It's reasonable to wait for things like occlusion query results. But waiting for timestamps should be avoided.
-		UE_CLOG(Query.Type == RQT_AbsoluteTime, LogD3D12RHI, Warning, TEXT("Waiting for a GPU timestamp query's result to be available. This should be avoided when possible."));
+		UE_CLOG(Query.Type == RQT_AbsoluteTime, LogD3D12RHI, Verbose, TEXT("Waiting for a GPU timestamp query's result to be available. This should be avoided when possible."));
 
 		const uint32 IdleStart = FPlatformTime::Cycles();
 
 		if (SyncPoint.IsOpen())
 		{
 			// We should really try to avoid this!
-			UE_LOG(LogD3D12RHI, Warning, TEXT("Stalling the RHI thread and flushing GPU commands to wait for a RenderQuery that hasn't been submitted to the GPU yet."));
+			UE_LOG(LogD3D12RHI, Verbose, TEXT("Stalling the RHI thread and flushing GPU commands to wait for a RenderQuery that hasn't been submitted to the GPU yet."));
 
 			// The query is on a command list that hasn't been submitted yet.
 			// We need to flush, but the RHI thread may be using the default command list...so stall it first.
