Merging //UE4/Dev-Main to Dev-Sequencer (//UE4/Dev-Sequencer)

diff --git a/Engine/Build/BuildFarm/Main.pl b/Engine/Build/BuildFarm/Main.pl
deleted file mode 100644
diff --git a/Engine/Build/BuildFarm/PostpExtensions.pl b/Engine/Build/BuildFarm/PostpExtensions.pl
deleted file mode 100644
diff --git a/Engine/Build/BuildFarm/PostpFilter.pl b/Engine/Build/BuildFarm/PostpFilter.pl
deleted file mode 100644
diff --git a/Engine/Build/BuildFarm/PostpTest.bat b/Engine/Build/BuildFarm/PostpTest.bat
deleted file mode 100644
diff --git a/Engine/Plugins/FX/HoudiniNiagara/Source/HoudiniNiagaraEditor/Private/HoudiniNiagaraEditor.cpp b/Engine/Plugins/FX/HoudiniNiagara/Source/HoudiniNiagaraEditor/Private/HoudiniNiagaraEditor.cpp
--- a/Engine/Plugins/FX/HoudiniNiagara/Source/HoudiniNiagaraEditor/Private/HoudiniNiagaraEditor.cpp
+++ b/Engine/Plugins/FX/HoudiniNiagara/Source/HoudiniNiagaraEditor/Private/HoudiniNiagaraEditor.cpp
@@ -93,4 +93,4 @@ void FHoudiniNiagaraEditorModule::ShutdownModule()
 
 #undef LOCTEXT_NAMESPACE
 	
-IMPLEMENT_MODULE(FHoudiniNiagaraEditorModule, HoudiniNiagara)
+IMPLEMENT_MODULE(FHoudiniNiagaraEditorModule, HoudiniNiagaraEditor)
diff --git a/Engine/Plugins/FX/Niagara/Source/NiagaraEditor/Private/NiagaraEditorStyle.cpp b/Engine/Plugins/FX/Niagara/Source/NiagaraEditor/Private/NiagaraEditorStyle.cpp
--- a/Engine/Plugins/FX/Niagara/Source/NiagaraEditor/Private/NiagaraEditorStyle.cpp
+++ b/Engine/Plugins/FX/Niagara/Source/NiagaraEditor/Private/NiagaraEditorStyle.cpp
@@ -34,7 +34,7 @@ FName FNiagaraEditorStyle::GetStyleSetName()
 	return StyleSetName;
 }
 
-FString RelativePathToPluginPath(const FString& RelativePath, const ANSICHAR* Extension)
+NIAGARAEDITOR_API FString RelativePathToPluginPath(const FString& RelativePath, const ANSICHAR* Extension)
 {
 	static FString ContentDir = IPluginManager::Get().FindPlugin(TEXT("Niagara"))->GetContentDir();
 	return (ContentDir / RelativePath) + Extension;
diff --git a/Engine/Plugins/FX/Niagara/Source/NiagaraEditor/Private/NiagaraNodeReroute.cpp b/Engine/Plugins/FX/Niagara/Source/NiagaraEditor/Private/NiagaraNodeReroute.cpp
--- a/Engine/Plugins/FX/Niagara/Source/NiagaraEditor/Private/NiagaraNodeReroute.cpp
+++ b/Engine/Plugins/FX/Niagara/Source/NiagaraEditor/Private/NiagaraNodeReroute.cpp
@@ -6,7 +6,7 @@
 
 #define LOCTEXT_NAMESPACE "NiagaraNodeReroute"
 
-const char* PC_Wildcard = "wildcard";
+const char* PC_Wildcard_Niagara = "wildcard";
 
 UNiagaraNodeReroute::UNiagaraNodeReroute(const FObjectInitializer& ObjectInitializer)
 	: Super(ObjectInitializer)
@@ -32,11 +32,11 @@ void UNiagaraNodeReroute::PostLoad()
 void UNiagaraNodeReroute::AllocateDefaultPins()
 {
 	const FName InputPinName(TEXT("InputPin"));
-	UEdGraphPin* MyInputPin = CreatePin(EGPD_Input, PC_Wildcard, InputPinName);
+	UEdGraphPin* MyInputPin = CreatePin(EGPD_Input, PC_Wildcard_Niagara, InputPinName);
 	MyInputPin->bDefaultValueIsIgnored = true;
 
 	const FName OutputPinName(TEXT("OutputPin"));
-	CreatePin(EGPD_Output, PC_Wildcard, OutputPinName);
+	CreatePin(EGPD_Output, PC_Wildcard_Niagara, OutputPinName);
 }
 
 FText UNiagaraNodeReroute::GetTooltipText() const
@@ -138,7 +138,7 @@ void UNiagaraNodeReroute::PropagatePinType()
 
 	for (UEdGraphPin* Inputs : MyInputPin->LinkedTo)
 	{
-		if (Inputs->PinType.PinCategory != PC_Wildcard)
+		if (Inputs->PinType.PinCategory != PC_Wildcard_Niagara)
 		{
 			PropagatePinTypeFromDirection(true);
 			return;
@@ -147,7 +147,7 @@ void UNiagaraNodeReroute::PropagatePinType()
 
 	for (UEdGraphPin* Outputs : MyOutputPin->LinkedTo)
 	{
-		if (Outputs->PinType.PinCategory != PC_Wildcard)
+		if (Outputs->PinType.PinCategory != PC_Wildcard_Niagara)
 		{
 			PropagatePinTypeFromDirection(false);
 			return;
@@ -170,11 +170,11 @@ void UNiagaraNodeReroute::PropagatePinType()
 		// Revert to wildcard
 		MyInputPin->BreakAllPinLinks();
 		MyInputPin->PinType.ResetToDefaults();
-		MyInputPin->PinType.PinCategory = PC_Wildcard;
+		MyInputPin->PinType.PinCategory = PC_Wildcard_Niagara;
 
 		MyOutputPin->BreakAllPinLinks();
 		MyOutputPin->PinType.ResetToDefaults();
-		MyOutputPin->PinType.PinCategory = PC_Wildcard;
+		MyOutputPin->PinType.PinCategory = PC_Wildcard_Niagara;
 	}
 }
 
diff --git a/Engine/Plugins/FX/Niagara/Source/NiagaraEditor/Private/Sequencer/NiagaraSequence/Sections/MovieSceneNiagaraEmitterSection.cpp b/Engine/Plugins/FX/Niagara/Source/NiagaraEditor/Private/Sequencer/NiagaraSequence/Sections/MovieSceneNiagaraEmitterSection.cpp
--- a/Engine/Plugins/FX/Niagara/Source/NiagaraEditor/Private/Sequencer/NiagaraSequence/Sections/MovieSceneNiagaraEmitterSection.cpp
+++ b/Engine/Plugins/FX/Niagara/Source/NiagaraEditor/Private/Sequencer/NiagaraSequence/Sections/MovieSceneNiagaraEmitterSection.cpp
@@ -99,8 +99,9 @@ TSharedPtr<FStructOnScope> GetKeyStruct(TMovieSceneChannelHandle<FMovieSceneNiag
 	int32 KeyValueIndex = Channel.Get()->GetData().GetIndex(InHandle);
 	if (KeyValueIndex != INDEX_NONE)
 	{
-		FNiagaraEmitterSectionKey& KeyValue = Channel.Get()->GetData().GetValues()[KeyValueIndex];
-		return MakeShared<FStructOnScope>(KeyValue.Value.GetType().GetStruct(), KeyValue.Value.GetData());
+		FNiagaraTypeDefinition KeyType = Channel.Get()->GetData().GetValues()[KeyValueIndex].Value.GetType();
+		uint8* KeyData = Channel.Get()->GetData().GetValues()[KeyValueIndex].Value.GetData();
+		return MakeShared<FStructOnScope>(KeyType.GetStruct(), KeyData);
 	}
 	return TSharedPtr<FStructOnScope>();
 }
diff --git a/Engine/Plugins/FX/Niagara/Source/NiagaraEditor/Private/Widgets/SNiagaraParameterMapView.cpp b/Engine/Plugins/FX/Niagara/Source/NiagaraEditor/Private/Widgets/SNiagaraParameterMapView.cpp
--- a/Engine/Plugins/FX/Niagara/Source/NiagaraEditor/Private/Widgets/SNiagaraParameterMapView.cpp
+++ b/Engine/Plugins/FX/Niagara/Source/NiagaraEditor/Private/Widgets/SNiagaraParameterMapView.cpp
@@ -35,8 +35,6 @@
 
 #define LOCTEXT_NAMESPACE "NiagaraParameterMapView"
 
-DEFINE_LOG_CATEGORY(LogNiagara);
-
 FText NiagaraParameterMapSectionID::OnGetSectionTitle(const NiagaraParameterMapSectionID::Type InSection)
 {
 	/* Setup an appropriate name for the section for this node */
diff --git a/Engine/Plugins/FX/Niagara/Source/NiagaraEditorWidgets/Private/NiagaraEditorWidgetsStyle.cpp b/Engine/Plugins/FX/Niagara/Source/NiagaraEditorWidgets/Private/NiagaraEditorWidgetsStyle.cpp
--- a/Engine/Plugins/FX/Niagara/Source/NiagaraEditorWidgets/Private/NiagaraEditorWidgetsStyle.cpp
+++ b/Engine/Plugins/FX/Niagara/Source/NiagaraEditorWidgets/Private/NiagaraEditorWidgetsStyle.cpp
@@ -34,11 +34,7 @@ FName FNiagaraEditorWidgetsStyle::GetStyleSetName()
 	return StyleSetName;
 }
 
-FString RelativePathToPluginPath(const FString& RelativePath, const ANSICHAR* Extension)
-{
-	static FString ContentDir = IPluginManager::Get().FindPlugin(TEXT("Niagara"))->GetContentDir();
-	return (ContentDir / RelativePath) + Extension;
-}
+NIAGARAEDITOR_API FString RelativePathToPluginPath(const FString& RelativePath, const ANSICHAR* Extension);
 
 #define IMAGE_PLUGIN_BRUSH( RelativePath, ... ) FSlateImageBrush( RelativePathToPluginPath( RelativePath, ".png" ), __VA_ARGS__ )
 #define IMAGE_BRUSH( RelativePath, ... ) FSlateImageBrush( Style->RootToContentDir( RelativePath, TEXT(".png") ), __VA_ARGS__ )
diff --git a/Engine/Shaders/Private/ReflectionEnvironmentPixelShader.usf b/Engine/Shaders/Private/ReflectionEnvironmentPixelShader.usf
--- a/Engine/Shaders/Private/ReflectionEnvironmentPixelShader.usf
+++ b/Engine/Shaders/Private/ReflectionEnvironmentPixelShader.usf
@@ -298,7 +298,7 @@ void ReflectionEnvironmentSkyLighting(
 {
 	float2 ScreenUV = UVAndScreenPos.xy;
 	FScreenSpaceData ScreenSpaceData = GetScreenSpaceData(ScreenUV);
-	uint ShadingModelID = GetShadingModelId(ScreenUV);
+	uint ShadingModelID = ScreenSpaceData.GBuffer.ShadingModelID;
 
 	OutColor = 0.0f;
 
diff --git a/Engine/Source/Developer/FunctionalTesting/Private/FunctionalTestingModule.cpp b/Engine/Source/Developer/FunctionalTesting/Private/FunctionalTestingModule.cpp
--- a/Engine/Source/Developer/FunctionalTesting/Private/FunctionalTestingModule.cpp
+++ b/Engine/Source/Developer/FunctionalTesting/Private/FunctionalTestingModule.cpp
@@ -11,6 +11,7 @@
 #include "EngineGlobals.h"
 #include "ARFilter.h"
 #include "AssetRegistryModule.h"
+#include "Misc/CommandLine.h"
 
 #define LOCTEXT_NAMESPACE "FunctionalTesting"
 
@@ -38,9 +39,20 @@ private:
 	void OnGetAssetTagsForWorld(const UWorld* World, TArray<UObject::FAssetRegistryTag>& OutTags);
 
 	void BuildTestBlacklistFromConfig();
-	bool IsBlacklisted(const FString& MapName, const FString& TestName) const;
+	bool IsBlacklisted(const FString& MapName, const FString& TestName, FString* OutReason=nullptr, bool* OutWarn=nullptr) const;
 
-	TArray<FString> TestBlacklist;
+	struct FBlacklistEntry
+	{
+		FBlacklistEntry() :
+			bWarn(false) {}
+
+		FString	Map;
+		FString Test;
+		FString Reason;
+		bool bWarn;
+	};
+
+	TMap<FString, FBlacklistEntry> TestBlacklist;
 	TWeakObjectPtr<class UFunctionalTestingManager> TestManager;
 	bool bPendingActivation;
 };
@@ -148,12 +160,25 @@ void FFunctionalTestingModule::GetMapTests(bool bEditorOnlyTests, TArray<FString
 
 							if (MapTest.Split(TEXT("|"), &BeautifulTestName, &RealTestName))
 							{
-								if (!IsBlacklisted(MapPackageName, RealTestName))
+								FString BlacklistReason;
+								bool bWarn(false);
+								if (!IsBlacklisted(MapPackageName, RealTestName, &BlacklistReason, &bWarn))
 								{
 									OutBeautifiedNames.Add(MapPackageName + TEXT(".") + *BeautifulTestName);
 									OutTestCommands.Add(MapAssetPath + TEXT(";") + MapAsset.PackageName.ToString() + TEXT(";") + *RealTestName);
 									OutTestMapAssets.AddUnique(MapAssetPath);
 								}
+								else
+								{
+									if (bWarn)
+									{
+										UE_LOG(LogFunctionalTest, Warning, TEXT("Test '%s' is blacklisted. %s"), *MapTest, *BlacklistReason);
+									}
+									else
+									{
+										UE_LOG(LogFunctionalTest, Display, TEXT("Test '%s' is blacklisted. %s"), *MapTest, *BlacklistReason);
+									}
+								}
 							}
 						}
 					}
@@ -266,30 +291,49 @@ void FFunctionalTestingModule::BuildTestBlacklistFromConfig()
 	TestBlacklist.Empty();
 	if (GConfig)
 	{
+	
+		const FString CommandLine = FCommandLine::Get();
+
 		for (const TPair<FString,FConfigFile>& Config : *GConfig)
 		{
 			FConfigSection* BlacklistSection = GConfig->GetSectionPrivate(TEXT("AutomationTestBlacklist"), false, true, Config.Key);
 			if (BlacklistSection)
 			{
-				// Parse all blacklist definitions of the format "BlacklistTest=(Map=/Game/Tests/MapName, Test=TestName)"
+				// Parse all blacklist definitions of the format "BlacklistTest=(Map=/Game/Tests/MapName, Test=TestName, Reason="Foo")"
 				for (FConfigSection::TIterator Section(*BlacklistSection); Section; ++Section)
 				{
 					if (Section.Key() == TEXT("BlacklistTest"))
 					{
 						FString BlacklistValue = Section.Value().GetValue();
-						FString Map, Test;
+						FString Map, Test, Reason, Warn;
 						bool bSuccess = false;
 						
 						if (FParse::Value(*BlacklistValue, TEXT("Map="), Map, true) && FParse::Value(*BlacklistValue, TEXT("Test="), Test, true))
 						{
+							FParse::Value(*BlacklistValue, TEXT("Reason="), Reason);
+							FParse::Value(*BlacklistValue, TEXT("Warn="), Warn);
+
 							// These are used as folders so ensure they match the expected layout
 							if (Map.StartsWith(TEXT("/")))
 							{
 								FString ListName = Map + TEXT("/") + Test;
 								ListName.RemoveSpacesInline();
-
-								TestBlacklist.AddUnique(ListName);
 								bSuccess = true;
+
+								if (CommandLine.Contains(Map) || CommandLine.Contains(Test))
+								{
+									UE_LOG(LogFunctionalTest, Warning, TEXT("Test '%s' is blacklisted but allowing due to command line."), *BlacklistValue);
+								}
+								else
+								{
+									// convert Pretty.Name into PrettyName as we compare on the latter.
+									ListName = ListName.Replace(TEXT("."), TEXT(""));
+									FBlacklistEntry& Entry = TestBlacklist.Add(ListName);
+									Entry.Map = Map;
+									Entry.Test = Test;
+									Entry.Reason = Reason;
+									Entry.bWarn = Warn.ToBool();
+								}
 							}
 						}
 						
@@ -306,18 +350,34 @@ void FFunctionalTestingModule::BuildTestBlacklistFromConfig()
 	if (TestBlacklist.Num() > 0)
 	{
 		UE_LOG(LogFunctionalTest, Log, TEXT("Automated Test Blacklist:"));
-		for (FString& Test : TestBlacklist)
+		for (auto& KV : TestBlacklist)
 		{
-			UE_LOG(LogFunctionalTest, Log, TEXT("\tTest: %s"), *Test);
+			UE_LOG(LogFunctionalTest, Log, TEXT("\tTest: %s"), *KV.Key);
 		}
 	}
 }
 
-bool FFunctionalTestingModule::IsBlacklisted(const FString& MapName, const FString& TestName) const
+bool FFunctionalTestingModule::IsBlacklisted(const FString& MapName, const FString& TestName, FString* OutReason, bool *OutWarn) const
 {
 	FString ListName = MapName + TEXT("/") + TestName;
 	ListName.RemoveSpacesInline();
-	return TestBlacklist.Contains(ListName);
+
+	const FBlacklistEntry* Entry = TestBlacklist.Find(ListName);
+
+	if (Entry)
+	{
+		if (OutReason != nullptr)
+		{
+			*OutReason = Entry->Reason;
+		}
+
+		if (OutWarn != nullptr)
+		{
+			*OutWarn = Entry->bWarn;
+		}
+	}
+
+	return Entry != nullptr;
 }
 
 //////////////////////////////////////////////////////////////////////////
diff --git a/Engine/Source/Developer/MeshSimplifier/Private/QuadricMeshReduction.cpp b/Engine/Source/Developer/MeshSimplifier/Private/QuadricMeshReduction.cpp
--- a/Engine/Source/Developer/MeshSimplifier/Private/QuadricMeshReduction.cpp
+++ b/Engine/Source/Developer/MeshSimplifier/Private/QuadricMeshReduction.cpp
@@ -552,6 +552,9 @@ public:
 
 			for (const FMeshTriangle& MeshTriangle : Triangles)
 			{
+				int32 CurrentFaceIndex = FaceIndex;
+				//Increment face index here because there is many continue in this for loop
+				++FaceIndex;
 				for (int32 CornerIndex = 0; CornerIndex < 3; ++CornerIndex)
 				{
 					VertexInstanceIDs[CornerIndex] = MeshTriangle.GetVertexInstanceID(CornerIndex);
@@ -570,8 +573,7 @@ public:
 				int32 VertexIndices[3];
 				for (int32 CornerIndex = 0; CornerIndex < 3; CornerIndex++)
 				{
-					int32 WedgeIndex = FaceIndex * 3 + CornerIndex;
-					++FaceIndex;
+					int32 WedgeIndex = CurrentFaceIndex * 3 + CornerIndex;
 
 					TVertSimp< NumTexCoords > NewVert;
 
diff --git a/Engine/Source/Editor/UnrealEd/Private/ObjectTools.cpp b/Engine/Source/Editor/UnrealEd/Private/ObjectTools.cpp
--- a/Engine/Source/Editor/UnrealEd/Private/ObjectTools.cpp
+++ b/Engine/Source/Editor/UnrealEd/Private/ObjectTools.cpp
@@ -785,15 +785,15 @@ namespace ObjectTools
 			}
 		}
 
-		for (int CurrentIndex = 0; CurrentIndex < ReferencingPropertiesMapKeys.Num(); CurrentIndex++)
 		{
-			int CheckForCircles = 0;
-			FFindReferencersArchive FindDependentArchive(ReferencingPropertiesMapKeys[CurrentIndex], ReferencingPropertiesMapKeys);
-			for (int DependentIndex = CurrentIndex + 1; DependentIndex < ReferencingPropertiesMapKeys.Num(); DependentIndex++)
+			TBitArray<> TouchedThisItteration(false, ReferencingPropertiesMapKeys.Num());
+			for (int CurrentIndex = 0; CurrentIndex < ReferencingPropertiesMapKeys.Num(); CurrentIndex++)
 			{
-				if (FindDependentArchive.GetReferenceCount(ReferencingPropertiesMapKeys[DependentIndex]) > 0)
+				TouchedThisItteration.Init(false, ReferencingPropertiesMapKeys.Num());
+				FFindReferencersArchive FindDependentArchive(ReferencingPropertiesMapKeys[CurrentIndex], ReferencingPropertiesMapKeys);
+				for (int DependentIndex = CurrentIndex + 1; DependentIndex < ReferencingPropertiesMapKeys.Num(); DependentIndex++)
 				{
-					if (CheckForCircles != DependentIndex)
+					if (!TouchedThisItteration[DependentIndex] && FindDependentArchive.GetReferenceCount(ReferencingPropertiesMapKeys[DependentIndex]) > 0)
 					{
 						UObject* Key = ReferencingPropertiesMapKeys[CurrentIndex];
 						PropertyArrayType Value = ReferencingPropertiesMapValues[CurrentIndex];
@@ -803,12 +803,7 @@ namespace ObjectTools
 						ReferencingPropertiesMapValues[DependentIndex] = Value;
 
 						FindDependentArchive.ResetPotentialReferencer(ReferencingPropertiesMapKeys[CurrentIndex]);
-						
-						if(CheckForCircles == 0)
-						{
-							CheckForCircles = DependentIndex;
-						}
-						
+						TouchedThisItteration[DependentIndex] = true;
 						DependentIndex = CurrentIndex;
 					}
 				}
diff --git a/Engine/Source/Programs/AutomationTool/BuildGraph/Script.cs b/Engine/Source/Programs/AutomationTool/BuildGraph/Script.cs
--- a/Engine/Source/Programs/AutomationTool/BuildGraph/Script.cs
+++ b/Engine/Source/Programs/AutomationTool/BuildGraph/Script.cs
@@ -794,12 +794,10 @@ namespace AutomationTool
 				{
 					if (Types.Length > 0 && Agent.PossibleTypes.Length > 0)
 					{
-						string[] NewTypes = Agent.PossibleTypes.Intersect(Types, StringComparer.InvariantCultureIgnoreCase).ToArray();
-						if (NewTypes.Length == 0)
+						if (Types.Length != Agent.PossibleTypes.Length || !Types.SequenceEqual(Agent.PossibleTypes, StringComparer.InvariantCultureIgnoreCase))
 						{
-							LogError(Element, "No common agent types with previous agent definition");
+							LogError(Element, "Agent types ({0}) were different than previous agent definition with types ({1}). Must either be empty or match exactly.", string.Join(",", Types), string.Join(",", Agent.PossibleTypes));
 						}
-						Agent.PossibleTypes = NewTypes;
 					}
 				}
 				else
diff --git a/Engine/Source/Programs/UnrealBuildTool/Platform/Android/AndroidToolChain.cs b/Engine/Source/Programs/UnrealBuildTool/Platform/Android/AndroidToolChain.cs
--- a/Engine/Source/Programs/UnrealBuildTool/Platform/Android/AndroidToolChain.cs
+++ b/Engine/Source/Programs/UnrealBuildTool/Platform/Android/AndroidToolChain.cs
@@ -1665,12 +1665,71 @@ namespace UnrealBuildTool
 						LinkAction.CommandArguments = String.Format("/c \"{0} {1}\"", LinkAction.CommandPath, LinkAction.CommandArguments);
 						LinkAction.CommandPath = "cmd.exe";
 					}
+
+					// Windows can run into an issue with too long of a commandline when clang tries to call ld to link.
+					// To work around this we call clang to just get the command it would execute and generate a
+					// second response file to directly call ld with the right arguments instead of calling through clang.
+/* disable while tracking down some linker errors this introduces
+					if (!Utils.IsRunningOnMono)
+					{
+						// capture the actual link command without running it
+						ProcessStartInfo StartInfo = new ProcessStartInfo();
+						StartInfo.WorkingDirectory = LinkEnvironment.IntermediateDirectory.FullName;
+						StartInfo.FileName = LinkAction.CommandPath;
+						StartInfo.Arguments = "-### " + LinkAction.CommandArguments;
+						StartInfo.UseShellExecute = false;
+						StartInfo.CreateNoWindow = true;
+						StartInfo.RedirectStandardError = true;
+
+						LinkerCommandline = "";
+
+						Process Proc = new Process();
+						Proc.StartInfo = StartInfo;
+						Proc.ErrorDataReceived += new DataReceivedEventHandler(OutputReceivedForLinker);
+						Proc.Start();
+						Proc.BeginErrorReadLine();
+						Proc.WaitForExit(5000);
+
+						LinkerCommandline = LinkerCommandline.Trim();
+
+						// the command should be in quotes; if not we'll just use clang to link as usual
+						int FirstQuoteIndex = LinkerCommandline.IndexOf('"');
+						if (FirstQuoteIndex >= 0)
+						{
+							int SecondQuoteIndex = LinkerCommandline.Substring(FirstQuoteIndex + 1).IndexOf('"');
+							if (SecondQuoteIndex >= 0)
+							{
+								LinkAction.CommandPath = LinkerCommandline.Substring(FirstQuoteIndex + 1, SecondQuoteIndex - FirstQuoteIndex);
+								LinkAction.CommandArguments = LinkerCommandline.Substring(FirstQuoteIndex + SecondQuoteIndex + 3);
+
+								// replace double backslashes
+								LinkAction.CommandPath = LinkAction.CommandPath.Replace("\\\\", "/");
+
+								// now create a response file for the full command using ld directly
+								FileReference FinalResponseFileName = FileReference.Combine(LinkEnvironment.IntermediateDirectory, OutputFile.Location.GetFileName() + ".responseFinal");
+								FileItem FinalResponseFileItem = FileItem.CreateIntermediateTextFile(FinalResponseFileName, LinkAction.CommandArguments);
+								LinkAction.CommandArguments = string.Format("@\"{0}\"", FinalResponseFileName);
+								LinkAction.PrerequisiteItems.Add(FinalResponseFileItem);
+							}
+						}
+					}
+*/
 				}
 			}
 
 			return Outputs.ToArray();
 		}
 
+		// captures stderr from clang
+		private static string LinkerCommandline = "";
+		static public void OutputReceivedForLinker(Object Sender, DataReceivedEventArgs Line)
+		{
+			if ((Line != null) && (Line.Data != null) && (Line.Data.Contains("--sysroot")))
+			{
+				LinkerCommandline += Line.Data;
+			}
+		}
+
 		private void ExportObjectFilePaths(LinkEnvironment LinkEnvironment, string FileName)
 		{
 			// Write the list of object file directories
diff --git a/Engine/Source/Runtime/Core/Private/HAL/LowLevelMemTracker.cpp b/Engine/Source/Runtime/Core/Private/HAL/LowLevelMemTracker.cpp
--- a/Engine/Source/Runtime/Core/Private/HAL/LowLevelMemTracker.cpp
+++ b/Engine/Source/Runtime/Core/Private/HAL/LowLevelMemTracker.cpp
@@ -61,6 +61,7 @@ DECLARE_LLM_MEMORY_STAT(TEXT("Networking"), STAT_NetworkingLLM, STATGROUP_LLMFUL
 DECLARE_LLM_MEMORY_STAT(TEXT("Meshes"), STAT_MeshesLLM, STATGROUP_LLMFULL);
 DECLARE_LLM_MEMORY_STAT(TEXT("Stats"), STAT_StatsLLM, STATGROUP_LLMFULL);
 DECLARE_LLM_MEMORY_STAT(TEXT("Shaders"), STAT_ShadersLLM, STATGROUP_LLMFULL);
+DECLARE_LLM_MEMORY_STAT(TEXT("PSO"), STAT_PSOLLM, STATGROUP_LLMFULL);
 DECLARE_LLM_MEMORY_STAT(TEXT("Textures"), STAT_TexturesLLM, STATGROUP_LLMFULL);
 DECLARE_LLM_MEMORY_STAT(TEXT("Render Targets"), STAT_RenderTargetsLLM, STATGROUP_LLMFULL);
 DECLARE_LLM_MEMORY_STAT(TEXT("RHIMisc"), STAT_RHIMiscLLM, STATGROUP_LLMFULL);
diff --git a/Engine/Source/Runtime/Core/Public/HAL/LowLevelMemTracker.h b/Engine/Source/Runtime/Core/Public/HAL/LowLevelMemTracker.h
--- a/Engine/Source/Runtime/Core/Public/HAL/LowLevelMemTracker.h
+++ b/Engine/Source/Runtime/Core/Public/HAL/LowLevelMemTracker.h
@@ -138,6 +138,7 @@ enum class ELLMTagSet : uint8
 	macro(Meshes,								"Meshes",						GET_STATFNAME(STAT_MeshesLLM),								GET_STATFNAME(STAT_MeshesSummaryLLM))		\
 	macro(Stats,								"Stats",						GET_STATFNAME(STAT_StatsLLM),								GET_STATFNAME(STAT_EngineSummaryLLM))		\
 	macro(Shaders,								"Shaders",						GET_STATFNAME(STAT_ShadersLLM),								GET_STATFNAME(STAT_EngineSummaryLLM))		\
+	macro(PSO,									"PSO",							GET_STATFNAME(STAT_PSOLLM),									GET_STATFNAME(STAT_EngineSummaryLLM))		\
 	macro(Textures,								"Textures",						GET_STATFNAME(STAT_TexturesLLM),							GET_STATFNAME(STAT_TexturesSummaryLLM))		\
 	macro(RenderTargets,						"RenderTargets",				GET_STATFNAME(STAT_RenderTargetsLLM),						GET_STATFNAME(STAT_EngineSummaryLLM))		\
 	macro(RHIMisc,								"RHIMisc",						GET_STATFNAME(STAT_RHIMiscLLM),								GET_STATFNAME(STAT_EngineSummaryLLM))		\
diff --git a/Engine/Source/Runtime/D3D12RHI/Private/D3D12Commands.cpp b/Engine/Source/Runtime/D3D12RHI/Private/D3D12Commands.cpp
--- a/Engine/Source/Runtime/D3D12RHI/Private/D3D12Commands.cpp
+++ b/Engine/Source/Runtime/D3D12RHI/Private/D3D12Commands.cpp
@@ -432,32 +432,8 @@ void FD3D12CommandContext::RHISetScissorRect(bool bEnable, uint32 MinX, uint32 M
 */
 void FD3D12CommandContext::RHISetBoundShaderState(FBoundShaderStateRHIParamRef BoundShaderStateRHI)
 {
-	//SCOPE_CYCLE_COUNTER(STAT_D3D12SetBoundShaderState);
-	FD3D12BoundShaderState* const BoundShaderState = FD3D12DynamicRHI::ResourceCast(BoundShaderStateRHI);
-
-	StateCache.SetBoundShaderState(BoundShaderState);
-
-	if (BoundShaderState->GetHullShader() && BoundShaderState->GetDomainShader())
-	{
-		bUsingTessellation = true;
-
-		// Ensure the command buffers are reset to reduce the amount of data that needs to be versioned.
-		HSConstantBuffer.Reset();
-		DSConstantBuffer.Reset();
-	}
-	else
-	{
-		bUsingTessellation = false;
-	}
-
-	// @TODO : really should only discard the constants if the shader state has actually changed.
-	bDiscardSharedConstants = true;
-
-	// Ensure the command buffers are reset to reduce the amount of data that needs to be versioned.
-	VSConstantBuffer.Reset();
-	PSConstantBuffer.Reset();
-	GSConstantBuffer.Reset();
-	CSConstantBuffer.Reset();	// Should this be here or in RHISetComputeShader? Might need a new bDiscardSharedConstants for CS.
+	// This shouldn't be used. Instead use RHISetGraphicsPipelineState.
+	unimplemented();
 }
 
 D3D_PRIMITIVE_TOPOLOGY GetD3D12PrimitiveType(uint32 PrimitiveType, bool bUsingTessellation)
@@ -541,10 +517,19 @@ void FD3D12CommandContext::RHISetGraphicsPipelineState(FGraphicsPipelineStateRHI
 	FD3D12GraphicsPipelineState* GraphicsPipelineState = FD3D12DynamicRHI::ResourceCast(GraphicsState);
 
 	// TODO: [PSO API] Every thing inside this scope is only necessary to keep the PSO shadow in sync while we convert the high level to only use PSOs
-	{
-		RHISetBoundShaderState(GraphicsPipelineState->BoundShaderState);
-		RHIEnableDepthBoundsTest(GraphicsPipelineState->PipelineStateInitializer.bDepthBounds);
-	}
+	bUsingTessellation = GraphicsPipelineState->GetHullShader() && GraphicsPipelineState->GetDomainShader();
+	// Ensure the command buffers are reset to reduce the amount of data that needs to be versioned.
+	VSConstantBuffer.Reset();
+	PSConstantBuffer.Reset();
+	HSConstantBuffer.Reset();
+	DSConstantBuffer.Reset();
+	GSConstantBuffer.Reset();
+	// Should this be here or in RHISetComputeShader? Might need a new bDiscardSharedConstants for CS.
+	CSConstantBuffer.Reset();
+	// @TODO : really should only discard the constants if the shader state has actually changed.
+	bDiscardSharedConstants = true;
+
+	RHIEnableDepthBoundsTest(GraphicsPipelineState->PipelineStateInitializer.bDepthBounds);
 
 	StateCache.SetGraphicsPipelineState(GraphicsPipelineState);
 }
@@ -1197,13 +1182,13 @@ void FD3D12CommandContext::CommitNonComputeShaderConstants()
 {
 	//SCOPE_CYCLE_COUNTER(STAT_D3D12CommitGraphicsConstants);
 
-	FD3D12BoundShaderState* RESTRICT CurrentBoundShaderStateRef = StateCache.GetBoundShaderState();
+	const FD3D12GraphicsPipelineState* const RESTRICT GraphicPSO = StateCache.GetGraphicsPipelineState();
 
-	check(CurrentBoundShaderStateRef);
+	check(GraphicPSO);
 
 	// Only set the constant buffer if this shader needs the global constant buffer bound
 	// Otherwise we will overwrite a different constant buffer
-	if (CurrentBoundShaderStateRef->bShaderNeedsGlobalConstantBuffer[SF_Vertex])
+	if (GraphicPSO->bShaderNeedsGlobalConstantBuffer[SF_Vertex])
 	{
 		StateCache.SetConstantBuffer<SF_Vertex>(VSConstantBuffer, bDiscardSharedConstants);
 	}
@@ -1214,23 +1199,23 @@ void FD3D12CommandContext::CommitNonComputeShaderConstants()
 	// is always reset whenever bUsingTessellation changes in SetBoundShaderState()
 	if (bUsingTessellation)
 	{
-		if (CurrentBoundShaderStateRef->bShaderNeedsGlobalConstantBuffer[SF_Hull])
+		if (GraphicPSO->bShaderNeedsGlobalConstantBuffer[SF_Hull])
 		{
 			StateCache.SetConstantBuffer<SF_Hull>(HSConstantBuffer, bDiscardSharedConstants);
 		}
 
-		if (CurrentBoundShaderStateRef->bShaderNeedsGlobalConstantBuffer[SF_Domain])
+		if (GraphicPSO->bShaderNeedsGlobalConstantBuffer[SF_Domain])
 		{
 			StateCache.SetConstantBuffer<SF_Domain>(DSConstantBuffer, bDiscardSharedConstants);
 		}
 	}
 
-	if (CurrentBoundShaderStateRef->bShaderNeedsGlobalConstantBuffer[SF_Geometry])
+	if (GraphicPSO->bShaderNeedsGlobalConstantBuffer[SF_Geometry])
 	{
 		StateCache.SetConstantBuffer<SF_Geometry>(GSConstantBuffer, bDiscardSharedConstants);
 	}
 
-	if (CurrentBoundShaderStateRef->bShaderNeedsGlobalConstantBuffer[SF_Pixel])
+	if (GraphicPSO->bShaderNeedsGlobalConstantBuffer[SF_Pixel])
 	{
 		StateCache.SetConstantBuffer<SF_Pixel>(PSConstantBuffer, bDiscardSharedConstants);
 	}
@@ -1377,26 +1362,26 @@ void FD3D12CommandContext::CommitGraphicsResourceTables()
 {
 	//SCOPE_CYCLE_COUNTER(STAT_D3D12CommitResourceTables);
 
-	const FD3D12BoundShaderState* const RESTRICT CurrentBoundShaderStateRef = StateCache.GetBoundShaderState();
-	check(CurrentBoundShaderStateRef);
+	const FD3D12GraphicsPipelineState* const RESTRICT GraphicPSO = StateCache.GetGraphicsPipelineState();
+	check(GraphicPSO);
 
-	if (auto* Shader = CurrentBoundShaderStateRef->GetVertexShader())
+	if (auto* Shader = GraphicPSO->GetVertexShader())
 	{
 		SetResourcesFromTables(Shader);
 	}
-	if (auto* Shader = CurrentBoundShaderStateRef->GetPixelShader())
+	if (auto* Shader = GraphicPSO->GetPixelShader())
 	{
 		SetResourcesFromTables(Shader);
 	}
-	if (auto* Shader = CurrentBoundShaderStateRef->GetHullShader())
+	if (auto* Shader = GraphicPSO->GetHullShader())
 	{
 		SetResourcesFromTables(Shader);
 	}
-	if (auto* Shader = CurrentBoundShaderStateRef->GetDomainShader())
+	if (auto* Shader = GraphicPSO->GetDomainShader())
 	{
 		SetResourcesFromTables(Shader);
 	}
-	if (auto* Shader = CurrentBoundShaderStateRef->GetGeometryShader())
+	if (auto* Shader = GraphicPSO->GetGeometryShader())
 	{
 		SetResourcesFromTables(Shader);
 	}
diff --git a/Engine/Source/Runtime/D3D12RHI/Private/D3D12PipelineState.h b/Engine/Source/Runtime/D3D12RHI/Private/D3D12PipelineState.h
--- a/Engine/Source/Runtime/D3D12RHI/Private/D3D12PipelineState.h
+++ b/Engine/Source/Runtime/D3D12RHI/Private/D3D12PipelineState.h
@@ -304,21 +304,21 @@ protected:
 
 struct FD3D12GraphicsPipelineState : public FRHIGraphicsPipelineState
 {
-	explicit FD3D12GraphicsPipelineState(
-		const FGraphicsPipelineStateInitializer& Initializer,
-		FD3D12BoundShaderState* InBoundShaderState,
-		FD3D12PipelineState* InPipelineState)
-		: PipelineStateInitializer(Initializer)
-		, BoundShaderState(InBoundShaderState)
-		, PipelineState(InPipelineState)
-	{
-	}
-
+	explicit FD3D12GraphicsPipelineState(const FGraphicsPipelineStateInitializer& Initializer, FD3D12BoundShaderState* InBoundShaderState, FD3D12PipelineState* InPipelineState);
 	~FD3D12GraphicsPipelineState();
 
 	const FGraphicsPipelineStateInitializer PipelineStateInitializer;
-	TRefCountPtr<FD3D12BoundShaderState> BoundShaderState;
+	const FD3D12RootSignature* RootSignature;
+	uint16 StreamStrides[MaxVertexElementCount];
+	bool bShaderNeedsGlobalConstantBuffer[SF_NumFrequencies];
+
 	FD3D12PipelineState* PipelineState;
+
+	FORCEINLINE class FD3D12VertexShader*   GetVertexShader() const { return (FD3D12VertexShader*)PipelineStateInitializer.BoundShaderState.VertexShaderRHI; }
+	FORCEINLINE class FD3D12PixelShader*    GetPixelShader() const { return (FD3D12PixelShader*)PipelineStateInitializer.BoundShaderState.PixelShaderRHI; }
+	FORCEINLINE class FD3D12HullShader*     GetHullShader() const { return (FD3D12HullShader*)PipelineStateInitializer.BoundShaderState.HullShaderRHI; }
+	FORCEINLINE class FD3D12DomainShader*   GetDomainShader() const { return (FD3D12DomainShader*)PipelineStateInitializer.BoundShaderState.DomainShaderRHI; }
+	FORCEINLINE class FD3D12GeometryShader* GetGeometryShader() const { return (FD3D12GeometryShader*)PipelineStateInitializer.BoundShaderState.GeometryShaderRHI; }
 };
 
 struct FD3D12ComputePipelineState : public FRHIComputePipelineState
@@ -379,9 +379,9 @@ protected:
 
 	// Thread access mutual exclusion
 	mutable FRWLock InitializerToGraphicsPipelineMapMutex;
-	mutable FRWLock LowLevelDescToGraphicsPipelineMapMutex;
+	mutable FRWLock LowLevelGraphicsPipelineStateCacheMutex;
 	mutable FRWLock ComputeShaderToComputePipelineMapMutex;
-	mutable FRWLock LowLevelDescToComputePipelineMapMutex;
+	mutable FRWLock ComputePipelineStateCacheMutex;
 
 	FCriticalSection DiskCachesCS;
 
diff --git a/Engine/Source/Runtime/D3D12RHI/Private/D3D12Pipelinestate.cpp b/Engine/Source/Runtime/D3D12RHI/Private/D3D12Pipelinestate.cpp
--- a/Engine/Source/Runtime/D3D12RHI/Private/D3D12Pipelinestate.cpp
+++ b/Engine/Source/Runtime/D3D12RHI/Private/D3D12Pipelinestate.cpp
@@ -287,6 +287,22 @@ ID3D12PipelineState* FD3D12PipelineState::GetPipelineState()
 	return PipelineState.GetReference();
 }
 
+FD3D12GraphicsPipelineState::FD3D12GraphicsPipelineState(
+	const FGraphicsPipelineStateInitializer& Initializer,
+	FD3D12BoundShaderState* InBoundShaderState,
+	FD3D12PipelineState* InPipelineState)
+	: PipelineStateInitializer(Initializer)
+	, RootSignature(InBoundShaderState->pRootSignature)
+	, PipelineState(InPipelineState)
+{
+	FMemory::Memcpy(StreamStrides, InBoundShaderState->StreamStrides, sizeof(StreamStrides));
+	bShaderNeedsGlobalConstantBuffer[SF_Vertex] = GetVertexShader() && GetVertexShader()->ResourceCounts.bGlobalUniformBufferUsed;
+	bShaderNeedsGlobalConstantBuffer[SF_Pixel] = GetPixelShader() && GetPixelShader()->ResourceCounts.bGlobalUniformBufferUsed;
+	bShaderNeedsGlobalConstantBuffer[SF_Hull] = GetHullShader() && GetHullShader()->ResourceCounts.bGlobalUniformBufferUsed;
+	bShaderNeedsGlobalConstantBuffer[SF_Domain] = GetDomainShader() && GetDomainShader()->ResourceCounts.bGlobalUniformBufferUsed;
+	bShaderNeedsGlobalConstantBuffer[SF_Geometry] = GetGeometryShader() && GetGeometryShader()->ResourceCounts.bGlobalUniformBufferUsed;
+}
+
 FD3D12GraphicsPipelineState::~FD3D12GraphicsPipelineState()
 {
 	// At this point the object is not safe to use in the PSO cache.
@@ -309,39 +325,51 @@ FD3D12ComputePipelineState::~FD3D12ComputePipelineState()
 
 void FD3D12PipelineStateCacheBase::CleanupPipelineStateCaches()
 {
-	// The runtime caches manage the lifetime of their FD3D12GraphicsPipelineState and FD3D12ComputePipelineState.
-	// We need to release them.
-	for (auto Pair : InitializerToGraphicsPipelineMap)
 	{
-		FD3D12GraphicsPipelineState* GraphicsPipelineState = Pair.Value;
-		ensure(GIsRHIInitialized || (!GIsRHIInitialized && GraphicsPipelineState->GetRefCount() == 1));
-		GraphicsPipelineState->Release();
+		FRWScopeLock Lock(InitializerToGraphicsPipelineMapMutex, FRWScopeLockType::SLT_Write);
+		// The runtime caches manage the lifetime of their FD3D12GraphicsPipelineState and FD3D12ComputePipelineState.
+		// We need to release them.
+		for (auto Pair : InitializerToGraphicsPipelineMap)
+		{
+			FD3D12GraphicsPipelineState* GraphicsPipelineState = Pair.Value;
+			ensure(GIsRHIInitialized || (!GIsRHIInitialized && GraphicsPipelineState->GetRefCount() == 1));
+			GraphicsPipelineState->Release();
+		}
+		InitializerToGraphicsPipelineMap.Reset();
 	}
-	InitializerToGraphicsPipelineMap.Reset();
 
-	for (auto Pair : ComputeShaderToComputePipelineMap)
 	{
-		FD3D12ComputePipelineState* ComputePipelineState = Pair.Value;
-		ensure(GIsRHIInitialized || (!GIsRHIInitialized && ComputePipelineState->GetRefCount() == 1));
-		ComputePipelineState->Release();
+		FRWScopeLock Lock(ComputeShaderToComputePipelineMapMutex, FRWScopeLockType::SLT_Write);
+		for (auto Pair : ComputeShaderToComputePipelineMap)
+		{
+			FD3D12ComputePipelineState* ComputePipelineState = Pair.Value;
+			ensure(GIsRHIInitialized || (!GIsRHIInitialized && ComputePipelineState->GetRefCount() == 1));
+			ComputePipelineState->Release();
+		}
+		ComputeShaderToComputePipelineMap.Reset();
 	}
-	ComputeShaderToComputePipelineMap.Reset();
 
-	// The low level graphics and compute maps manage the lifetime of their PSOs.
-	// We need to delete each element before we empty it.
-	for (auto Iter = LowLevelGraphicsPipelineStateCache.CreateConstIterator(); Iter; ++Iter)
 	{
-		const FD3D12PipelineState* const PipelineState = Iter.Value();
-		delete PipelineState;
+		FRWScopeLock Lock(LowLevelGraphicsPipelineStateCacheMutex, FRWScopeLockType::SLT_Write);
+		// The low level graphics and compute maps manage the lifetime of their PSOs.
+		// We need to delete each element before we empty it.
+		for (auto Iter = LowLevelGraphicsPipelineStateCache.CreateConstIterator(); Iter; ++Iter)
+		{
+			const FD3D12PipelineState* const PipelineState = Iter.Value();
+			delete PipelineState;
+		}
+		LowLevelGraphicsPipelineStateCache.Empty();
 	}
-	LowLevelGraphicsPipelineStateCache.Empty();
 
-	for (auto Iter = ComputePipelineStateCache.CreateConstIterator(); Iter; ++Iter)
 	{
-		const FD3D12PipelineState* const PipelineState = Iter.Value();
-		delete PipelineState;
+		FRWScopeLock Lock(ComputePipelineStateCacheMutex, FRWScopeLockType::SLT_Write);
+		for (auto Iter = ComputePipelineStateCache.CreateConstIterator(); Iter; ++Iter)
+		{
+			const FD3D12PipelineState* const PipelineState = Iter.Value();
+			delete PipelineState;
+		}
+		ComputePipelineStateCache.Empty();
 	}
-	ComputePipelineStateCache.Empty();
 }
 
 FD3D12GraphicsPipelineState* FD3D12PipelineStateCacheBase::AddToRuntimeCache(const FGraphicsPipelineStateInitializer& Initializer, uint32 InitializerHash, FD3D12BoundShaderState* BoundShaderState, FD3D12PipelineState* PipelineState)
@@ -368,7 +396,7 @@ FD3D12PipelineState* FD3D12PipelineStateCacheBase::FindInLowLevelCache(const FD3
 	check(Desc.CombinedHash != 0);
 
 	{
-		FRWScopeLock Lock(LowLevelDescToGraphicsPipelineMapMutex, FRWScopeLockType::SLT_ReadOnly);
+		FRWScopeLock Lock(LowLevelGraphicsPipelineStateCacheMutex, FRWScopeLockType::SLT_ReadOnly);
 		FD3D12PipelineState** Found = LowLevelGraphicsPipelineStateCache.Find(Desc);
 		if (Found)
 		{
@@ -399,7 +427,7 @@ void FD3D12PipelineStateCacheBase::AddToLowLevelCache(const FD3D12LowLevelGraphi
 
 	// Double check the desc doesn't already exist while the lock is taken.
 	// This avoids having multiple threads try to create the same PSO.
-	FRWScopeLock Lock(LowLevelDescToGraphicsPipelineMapMutex, FRWScopeLockType::SLT_Write);
+	FRWScopeLock Lock(LowLevelGraphicsPipelineStateCacheMutex, FRWScopeLockType::SLT_Write);
 	FD3D12PipelineState** const PipelineState = LowLevelGraphicsPipelineStateCache.Find(Desc);
 	if (PipelineState)
 	{
@@ -443,7 +471,7 @@ FD3D12PipelineState* FD3D12PipelineStateCacheBase::FindInLowLevelCache(const FD3
 	check(Desc.CombinedHash != 0);
 
 	{
-		FRWScopeLock Lock(LowLevelDescToComputePipelineMapMutex, FRWScopeLockType::SLT_ReadOnly);
+		FRWScopeLock Lock(ComputePipelineStateCacheMutex, FRWScopeLockType::SLT_ReadOnly);
 		FD3D12PipelineState** Found = ComputePipelineStateCache.Find(Desc);
 		if (Found)
 		{
@@ -475,7 +503,7 @@ void FD3D12PipelineStateCacheBase::AddToLowLevelCache(const FD3D12ComputePipelin
 
 	// Double check the desc doesn't already exist while the lock is taken.
 	// This avoids having multiple threads try to create the same PSO.
-	FRWScopeLock Lock(LowLevelDescToComputePipelineMapMutex, FRWScopeLockType::SLT_Write);
+	FRWScopeLock Lock(ComputePipelineStateCacheMutex, FRWScopeLockType::SLT_Write);
 	FD3D12PipelineState** const PipelineState = ComputePipelineStateCache.Find(Desc);
 	if (PipelineState)
 	{
diff --git a/Engine/Source/Runtime/D3D12RHI/Private/D3D12StateCache.cpp b/Engine/Source/Runtime/D3D12RHI/Private/D3D12StateCache.cpp
--- a/Engine/Source/Runtime/D3D12RHI/Private/D3D12StateCache.cpp
+++ b/Engine/Source/Runtime/D3D12RHI/Private/D3D12StateCache.cpp
@@ -353,8 +353,7 @@ void FD3D12StateCacheBase::ApplyState()
 #endif
 
 	FD3D12CommandListHandle& CommandList = CmdContext->CommandListHandle;
-	const FD3D12RootSignature* const pRootSignature = IsCompute ?
-		PipelineState.Compute.CurrentPipelineStateObject->ComputeShader->pRootSignature : PipelineState.Graphics.CurrentPipelineStateObject->BoundShaderState->pRootSignature;
+	const FD3D12RootSignature* const pRootSignature = IsCompute ? PipelineState.Compute.CurrentPipelineStateObject->ComputeShader->pRootSignature : GetGraphicsRootSignature();
 
 	// PSO
 	if (IsCompute)
diff --git a/Engine/Source/Runtime/D3D12RHI/Private/D3D12StateCachePrivate.h b/Engine/Source/Runtime/D3D12RHI/Private/D3D12StateCachePrivate.h
--- a/Engine/Source/Runtime/D3D12RHI/Private/D3D12StateCachePrivate.h
+++ b/Engine/Source/Runtime/D3D12RHI/Private/D3D12StateCachePrivate.h
@@ -365,6 +365,7 @@ protected:
 	{ \
 		static const EShaderFrequency Frequency = SF_##Name; \
 		static FD3D12##Name##Shader* GetShader(FD3D12BoundShaderState* BSS) { return BSS ? BSS->Get##Name##Shader() : nullptr; } \
+		static FD3D12##Name##Shader* GetShader(FD3D12GraphicsPipelineState* PSO) { return PSO ? (FD3D12##Name##Shader*)PSO->PipelineStateInitializer.BoundShaderState.##Name##ShaderRHI : nullptr; } \
 	}
 	DECLARE_SHADER_TRAITS(Vertex);
 	DECLARE_SHADER_TRAITS(Pixel);
@@ -376,7 +377,8 @@ protected:
 	template <typename TShader> D3D12_STATE_CACHE_INLINE void SetShader(TShader* Shader)
 	{
 		typedef StateCacheShaderTraits<TShader> Traits;
-		TShader* OldShader = Traits::GetShader(GetBoundShaderState());
+		TShader* OldShader = Traits::GetShader(GetGraphicsPipelineState());
+
 		if (OldShader != Shader)
 		{
 			PipelineState.Common.CurrentShaderSamplerCounts[Traits::Frequency] = (Shader) ? Shader->ResourceCounts.NumSamplers : 0;
@@ -391,7 +393,7 @@ protected:
 
 	template <typename TShader> D3D12_STATE_CACHE_INLINE void GetShader(TShader** Shader)
 	{
-		*Shader = StateCacheShaderTraits<TShader>::GetShader(GetBoundShaderState());
+		*Shader = StateCacheShaderTraits<TShader>::GetShader(GetGraphicsPipelineState());
 	}
 
 	template <bool IsCompute = false>
@@ -432,21 +434,19 @@ public:
 		return &DescriptorCache;
 	}
 
-	FD3D12BoundShaderState* GetBoundShaderState() const
+	FD3D12GraphicsPipelineState* GetGraphicsPipelineState() const
 	{
-		return PipelineState.Graphics.CurrentPipelineStateObject ? PipelineState.Graphics.CurrentPipelineStateObject->BoundShaderState : nullptr;
+		return PipelineState.Graphics.CurrentPipelineStateObject;
 	}
 
 	const FD3D12RootSignature* GetGraphicsRootSignature()
 	{
-		return PipelineState.Graphics.CurrentPipelineStateObject ?
-			PipelineState.Graphics.CurrentPipelineStateObject->BoundShaderState->pRootSignature : nullptr;
+		return PipelineState.Graphics.CurrentPipelineStateObject ? PipelineState.Graphics.CurrentPipelineStateObject->RootSignature : nullptr;
 	}
 
 	const FD3D12RootSignature* GetComputeRootSignature()
 	{
-		return PipelineState.Compute.CurrentPipelineStateObject ?
-			PipelineState.Compute.CurrentPipelineStateObject->ComputeShader->pRootSignature : nullptr;
+		return PipelineState.Compute.CurrentPipelineStateObject ? PipelineState.Compute.CurrentPipelineStateObject->ComputeShader->pRootSignature : nullptr;
 	}
 
 	void ClearSRVs();
@@ -656,46 +656,24 @@ public:
 		GetShader(Shader);
 	}
 
-	D3D12_STATE_CACHE_INLINE void SetBoundShaderState(FD3D12BoundShaderState* BoundShaderState)
+	D3D12_STATE_CACHE_INLINE void SetGraphicsPipelineState(FD3D12GraphicsPipelineState* GraphicsPipelineState)
 	{
-		if (BoundShaderState)
-		{
-			SetStreamStrides(BoundShaderState->StreamStrides);
-			SetShader(BoundShaderState->GetVertexShader());
-			SetShader(BoundShaderState->GetPixelShader());
-			SetShader(BoundShaderState->GetDomainShader());
-			SetShader(BoundShaderState->GetHullShader());
-			SetShader(BoundShaderState->GetGeometryShader());
-		}
-		else
+		check(GraphicsPipelineState);
+		if (PipelineState.Graphics.CurrentPipelineStateObject != GraphicsPipelineState)
 		{
-			uint16 NullStrides[MaxVertexElementCount] = {0};
-			SetStreamStrides(NullStrides);
-			SetShader<FD3D12VertexShader>(nullptr);
-			SetShader<FD3D12PixelShader>(nullptr);
-			SetShader<FD3D12HullShader>(nullptr);
-			SetShader<FD3D12DomainShader>(nullptr);
-			SetShader<FD3D12GeometryShader>(nullptr);
-		}
+			SetStreamStrides(GraphicsPipelineState->StreamStrides);
+			SetShader(GraphicsPipelineState->GetVertexShader());
+			SetShader(GraphicsPipelineState->GetPixelShader());
+			SetShader(GraphicsPipelineState->GetDomainShader());
+			SetShader(GraphicsPipelineState->GetHullShader());
+			SetShader(GraphicsPipelineState->GetGeometryShader());
 
-		FD3D12BoundShaderState* CurrentBoundShaderState = GetBoundShaderState();
-		if (CurrentBoundShaderState != BoundShaderState)
-		{
 			// See if we need to change the root signature
-			const FD3D12RootSignature* const pCurrentRootSignature = CurrentBoundShaderState ? CurrentBoundShaderState->pRootSignature : nullptr;
-			const FD3D12RootSignature* const pNewRootSignature = BoundShaderState ? BoundShaderState->pRootSignature : nullptr;
-			if (pCurrentRootSignature != pNewRootSignature)
+			if (GetGraphicsRootSignature() != GraphicsPipelineState->RootSignature)
 			{
 				PipelineState.Graphics.bNeedSetRootSignature = true;
 			}
-		}
-	}
 
-	D3D12_STATE_CACHE_INLINE void SetGraphicsPipelineState(FD3D12GraphicsPipelineState* GraphicsPipelineState)
-	{
-		check(GraphicsPipelineState);
-		if (PipelineState.Graphics.CurrentPipelineStateObject != GraphicsPipelineState)
-		{
 			// Save the PSO
 			PipelineState.Common.bNeedSetPSO = true;
 			PipelineState.Graphics.CurrentPipelineStateObject = GraphicsPipelineState;
diff --git a/Engine/Source/Runtime/Engine/Private/StaticMesh.cpp b/Engine/Source/Runtime/Engine/Private/StaticMesh.cpp
--- a/Engine/Source/Runtime/Engine/Private/StaticMesh.cpp
+++ b/Engine/Source/Runtime/Engine/Private/StaticMesh.cpp
@@ -1350,7 +1350,7 @@ FArchive& operator<<(FArchive& Ar, FMeshBuildSettings& BuildSettings)
 // differences, etc.) replace the version GUID below with a new one.
 // In case of merge conflicts with DDC versions, you MUST generate a new GUID
 // and set this new GUID as the version.                                       
-#define STATICMESH_DERIVEDDATA_VER TEXT("42B9CFE0CAFE46F29F43E9049AA21282")
+#define STATICMESH_DERIVEDDATA_VER TEXT("4FDA6C9D166447EAB2DA9EDA8A184591")
 
 static const FString& GetStaticMeshDerivedDataVersion()
 {
diff --git a/Engine/Source/Runtime/RHI/Private/PipelineStateCache.cpp b/Engine/Source/Runtime/RHI/Private/PipelineStateCache.cpp
--- a/Engine/Source/Runtime/RHI/Private/PipelineStateCache.cpp
+++ b/Engine/Source/Runtime/RHI/Private/PipelineStateCache.cpp
@@ -483,6 +483,8 @@ public:
 
 	void DoTask(ENamedThreads::Type CurrentThread, const FGraphEventRef& MyCompletionGraphEvent)
 	{
+		LLM_SCOPE(ELLMTag::PSO);
+
 		if (Pipeline->IsCompute())
 		{
 			FComputePipelineState* ComputePipeline = static_cast<FComputePipelineState*>(Pipeline);
@@ -713,7 +715,7 @@ FRHIComputePipelineState* ExecuteSetComputePipelineState(FComputePipelineState*
 
 FGraphicsPipelineState* PipelineStateCache::GetAndOrCreateGraphicsPipelineState(FRHICommandList& RHICmdList, const FGraphicsPipelineStateInitializer& OriginalInitializer, EApplyRendertargetOption ApplyFlags)
 {
-
+	LLM_SCOPE(ELLMTag::PSO);
 	SCOPE_CYCLE_COUNTER(STAT_GetOrCreatePSO);
 
 	FGraphicsPipelineStateInitializer NewInitializer;
diff --git a/Engine/Source/Runtime/Renderer/Private/PostProcess/PostProcessing.cpp b/Engine/Source/Runtime/Renderer/Private/PostProcess/PostProcessing.cpp
--- a/Engine/Source/Runtime/Renderer/Private/PostProcess/PostProcessing.cpp
+++ b/Engine/Source/Runtime/Renderer/Private/PostProcess/PostProcessing.cpp
@@ -932,11 +932,11 @@ static void AddTemporalAA( FPostprocessContext& Context, FRenderingCompositeOutp
 
 	FSceneViewState* ViewState = Context.View.ViewState;
 
-	FRCPassPostProcessTemporalAA* TemporalAAPass = new(FMemStack::Get()) FRCPassPostProcessTemporalAA(
+	FRCPassPostProcessTemporalAA* TemporalAAPass = Context.Graph.RegisterPass(new(FMemStack::Get()) FRCPassPostProcessTemporalAA(
 		Context,
 		Parameters,
 		Context.View.PrevViewInfo.TemporalAAHistory,
-		&ViewState->PendingPrevFrameViewInfo.TemporalAAHistory);
+		&ViewState->PendingPrevFrameViewInfo.TemporalAAHistory));
 
 	TemporalAAPass->SetInput( ePId_Input0, Context.FinalOutput );
 	TemporalAAPass->SetInput( ePId_Input2, VelocityInput );
@@ -1726,7 +1726,6 @@ void FPostProcessing::Process(FRHICommandListImmediate& RHICmdList, const FViewI
 			// some views don't have a state (thumbnail rendering)
 			if(!GIsHighResScreenshot && View.State && IStereoRendering::IsAPrimaryView(StereoPass))
 			{
-				
 				const bool bUseBasicEyeAdaptation = (AutoExposure.MethodId == EAutoExposureMethod::AEM_Basic);
 
 				if (bUseBasicEyeAdaptation) // log average ps reduction ( non histogram ) 
