UE4 - Fix clang warnings
#rnx

diff --git a/Engine/Source/Runtime/D3D12RHI/Private/D3D12Device.cpp b/Engine/Source/Runtime/D3D12RHI/Private/D3D12Device.cpp
--- a/Engine/Source/Runtime/D3D12RHI/Private/D3D12Device.cpp
+++ b/Engine/Source/Runtime/D3D12RHI/Private/D3D12Device.cpp
@@ -18,6 +18,12 @@ FD3D12Device::FD3D12Device() :
 	}
 
 FD3D12Device::FD3D12Device(FRHIGPUMask InGPUMask, FD3D12Adapter* InAdapter) :
+	FD3D12SingleNodeGPUObject(InGPUMask),
+	FD3D12AdapterChild(InAdapter)
+	CommandListManager(nullptr),
+	CopyCommandListManager(nullptr),
+	AsyncCommandListManager(nullptr),
+	TextureStreamingCommandAllocatorManager(this, D3D12_COMMAND_LIST_TYPE_COPY),
 	RTVAllocator(InGPUMask, D3D12_DESCRIPTOR_HEAP_TYPE_RTV, 256),
 	DSVAllocator(InGPUMask, D3D12_DESCRIPTOR_HEAP_TYPE_DSV, 256),
 	SRVAllocator(InGPUMask, D3D12_DESCRIPTOR_HEAP_TYPE_CBV_SRV_UAV, 1024),
@@ -26,20 +32,14 @@ FD3D12Device::FD3D12Device(FRHIGPUMask InGPUMask, FD3D12Adapter* InAdapter) :
 	CBVAllocator(InGPUMask, D3D12_DESCRIPTOR_HEAP_TYPE_CBV_SRV_UAV, 2048),
 #endif
 	SamplerAllocator(InGPUMask, FD3D12_DESCRIPTOR_HEAP_TYPE_SAMPLER, 128),
-	SamplerID(0),
+	GlobalSamplerHeap(this, InGPUMask),
+	GlobalViewHeap(this, InGPUMask),
 	OcclusionQueryHeap(this, D3D12_QUERY_HEAP_TYPE_OCCLUSION, 65536, 4 /*frames to keep results */ * 1 /*batches per frame*/),
 	TimestampQueryHeap(this, D3D12_QUERY_HEAP_TYPE_TIMESTAMP, 8192, 4 /*frames to keep results */ * 5 /*batches per frame*/ ),
 	DefaultBufferAllocator(this, InGPUMask), //Note: Cross node buffers are possible 
-	CommandListManager(nullptr),
-	CopyCommandListManager(nullptr),
-	AsyncCommandListManager(nullptr),
-	TextureStreamingCommandAllocatorManager(this, D3D12_COMMAND_LIST_TYPE_COPY),
-	GlobalSamplerHeap(this, InGPUMask),
-	GlobalViewHeap(this, InGPUMask),
+	SamplerID(0),
 	DefaultFastAllocator(this, InGPUMask, D3D12_HEAP_TYPE_UPLOAD, 1024 * 1024 * 4),
 	TextureAllocator(this, FRHIGPUMask::All()),
-	FD3D12SingleNodeGPUObject(InGPUMask),
-	FD3D12AdapterChild(InAdapter)
 {
 	InitPlatformSpecific();
 }
diff --git a/Engine/Source/Runtime/D3D12RHI/Private/D3D12DirectCommandListManager.cpp b/Engine/Source/Runtime/D3D12RHI/Private/D3D12DirectCommandListManager.cpp
--- a/Engine/Source/Runtime/D3D12RHI/Private/D3D12DirectCommandListManager.cpp
+++ b/Engine/Source/Runtime/D3D12RHI/Private/D3D12DirectCommandListManager.cpp
@@ -15,10 +15,10 @@ FComputeFenceRHIRef FD3D12DynamicRHI::RHICreateComputeFence(const FName& Name)
 }
 
 FD3D12FenceCore::FD3D12FenceCore(FD3D12Adapter* Parent, uint64 InitialValue, uint32 InGPUIndex)
-	: hFenceCompleteEvent(INVALID_HANDLE_VALUE)
+	: FD3D12AdapterChild(Parent)
 	, FenceValueAvailableAt(0)
 	, GPUIndex(InGPUIndex)
-	, FD3D12AdapterChild(Parent)
+	, hFenceCompleteEvent(INVALID_HANDLE_VALUE)
 {
 	check(Parent);
 	hFenceCompleteEvent = CreateEvent(nullptr, false, false, nullptr);
