UE4 - Typo fix

diff --git a/Engine/Source/Runtime/D3D12RHI/Private/D3D12Device.cpp b/Engine/Source/Runtime/D3D12RHI/Private/D3D12Device.cpp
--- a/Engine/Source/Runtime/D3D12RHI/Private/D3D12Device.cpp
+++ b/Engine/Source/Runtime/D3D12RHI/Private/D3D12Device.cpp
@@ -19,7 +19,7 @@ FD3D12Device::FD3D12Device() :
 
 FD3D12Device::FD3D12Device(FRHIGPUMask InGPUMask, FD3D12Adapter* InAdapter) :
 	FD3D12SingleNodeGPUObject(InGPUMask),
-	FD3D12AdapterChild(InAdapter)
+	FD3D12AdapterChild(InAdapter),
 	CommandListManager(nullptr),
 	CopyCommandListManager(nullptr),
 	AsyncCommandListManager(nullptr),
@@ -39,7 +39,7 @@ FD3D12Device::FD3D12Device(FRHIGPUMask InGPUMask, FD3D12Adapter* InAdapter) :
 	DefaultBufferAllocator(this, InGPUMask), //Note: Cross node buffers are possible 
 	SamplerID(0),
 	DefaultFastAllocator(this, InGPUMask, D3D12_HEAP_TYPE_UPLOAD, 1024 * 1024 * 4),
-	TextureAllocator(this, FRHIGPUMask::All()),
+	TextureAllocator(this, FRHIGPUMask::All())
 {
 	InitPlatformSpecific();
 }
