Dont call cooking function when running with cooked data.

diff --git a/Engine/Source/Runtime/D3D12RHI/Private/Windows/WindowsD3D12Device.cpp b/Engine/Source/Runtime/D3D12RHI/Private/Windows/WindowsD3D12Device.cpp
--- a/Engine/Source/Runtime/D3D12RHI/Private/Windows/WindowsD3D12Device.cpp
+++ b/Engine/Source/Runtime/D3D12RHI/Private/Windows/WindowsD3D12Device.cpp
@@ -702,7 +702,7 @@ void FD3D12DynamicRHI::Init()
 
 void FD3D12DynamicRHI::PostInit()
 {
-	if (GRHISupportsRayTracing || GRHISupportsRHIThread)
+	if (!FPlatformProperties::RequiresCookedData() && (GRHISupportsRayTracing || GRHISupportsRHIThread))
 	{
 		// Make sure all global shaders are complete at this point
 		extern RENDERCORE_API const int32 GlobalShaderMapId;
