ensure place resource sub-allocations are tracked correctly with LLM in non-Windows D3D12 builds.

diff --git a/Engine/Source/Runtime/D3D12RHI/Private/D3D12Allocation.cpp b/Engine/Source/Runtime/D3D12RHI/Private/D3D12Allocation.cpp
--- a/Engine/Source/Runtime/D3D12RHI/Private/D3D12Allocation.cpp
+++ b/Engine/Source/Runtime/D3D12RHI/Private/D3D12Allocation.cpp
@@ -303,8 +303,7 @@ void FD3D12BuddyAllocator::Allocate(uint32 SizeInBytes, uint32 Alignment, FD3D12
 
 	// track the allocation
 #if !PLATFORM_WINDOWS
-	LLM(uint64 Addr = (ResourceLocation.GetGPUVirtualAddress() != 0ull) ? (uint64)ResourceLocation.GetGPUVirtualAddress() : AlignedOffsetFromResourceBase);
-	LLM(FLowLevelMemTracker::Get().OnLowLevelAlloc(ELLMTracker::Default, (void*)Addr, SizeInBytes));
+	LLM(FLowLevelMemTracker::Get().OnLowLevelAlloc(ELLMTracker::Default, &ResourceLocation, SizeInBytes));
 	// Note: Disabling this LLM hook for Windows is due to a work-around in the way that d3d12 buffers are tracked
 	// by LLM. LLM tracks buffer data in the UpdateBufferStats function because that is the easiest place to ensure that LLM
 	// can be updated whenever a buffer is created or released. Unfortunately, some buffers allocate from this allocator
@@ -365,8 +364,7 @@ void FD3D12BuddyAllocator::Deallocate(FD3D12ResourceLocation& ResourceLocation)
 	// which means that the memory would be counted twice. Because of this the tracking had to be disabled here.
 	// This does mean that non-buffer memory that goes through this allocator won't be tracked, so this does need a better solution.
 	// see UpdateBufferStats for a more detailed explanation.
-	LLM(uint64 Addr = (ResourceLocation.GetGPUVirtualAddress() != 0ull) ? (uint64)ResourceLocation.GetGPUVirtualAddress() : ResourceLocation.GetOffsetFromBaseOfResource());
-	LLM(FLowLevelMemTracker::Get().OnLowLevelFree(ELLMTracker::Default, (void*)Addr));
+	LLM(FLowLevelMemTracker::Get().OnLowLevelFree(ELLMTracker::Default, &ResourceLocation));
 #endif
 }
 
diff --git a/Engine/Source/Runtime/D3D12RHI/Private/D3D12Resources.cpp b/Engine/Source/Runtime/D3D12RHI/Private/D3D12Resources.cpp
--- a/Engine/Source/Runtime/D3D12RHI/Private/D3D12Resources.cpp
+++ b/Engine/Source/Runtime/D3D12RHI/Private/D3D12Resources.cpp
@@ -468,6 +468,14 @@ void FD3D12ResourceLocation::TransferOwnership(FD3D12ResourceLocation& Destinati
 
 	FMemory::Memmove(&Destination, &Source, sizeof(FD3D12ResourceLocation));
 
+	// update tracked allocation
+#if !PLATFORM_WINDOWS && ENABLE_LOW_LEVEL_MEM_TRACKER
+	if (Source.GetType() == ResourceLocationType::eSubAllocation)
+	{
+		FLowLevelMemTracker::Get().OnLowLevelAllocMoved( ELLMTracker::Default, &Destination, &Source );
+	}
+#endif
+
 	// Destroy the source but don't invoke any resource destruction
 	Source.InternalClear<false>();
 }
