Temporary fix to a startup crash when raytracing is enabled

diff --git a/Engine/Source/Runtime/D3D12RHI/Private/D3D12Buffer.cpp b/Engine/Source/Runtime/D3D12RHI/Private/D3D12Buffer.cpp
--- a/Engine/Source/Runtime/D3D12RHI/Private/D3D12Buffer.cpp
+++ b/Engine/Source/Runtime/D3D12RHI/Private/D3D12Buffer.cpp
@@ -222,7 +222,7 @@ BufferType* FD3D12Adapter::CreateRHIBuffer(FRHICommandListImmediate* RHICmdList,
 				}
 			};
 			
-			if (!RHICmdList)
+			if (!IsInRHIThread() && !IsInRenderingThread())
 			{
 				// Need to update buffer content on RHI thread (immediate context) because the buffer can be a
 				// sub-allocation and its backing resource may be in a state incompatible with the copy queue.
@@ -250,13 +250,15 @@ BufferType* FD3D12Adapter::CreateRHIBuffer(FRHICommandListImmediate* RHICmdList,
 					delete SrcResourceLoc_Heap;
 				});
 			}
-			else if (RHICmdList->Bypass())
+			else if (!RHICmdList || RHICmdList->Bypass())
 			{
+				// On RHIT or RT (when bypassing), we can access immediate context directly
 				FD3D12RHICommandInitializeBuffer Command(BufferOut, SrcResourceLoc, Size);
 				Command.ExecuteNoCmdList();
 			}
 			else
 			{
+				// On RT but not bypassing
 				new (RHICmdList->AllocCommand<FD3D12RHICommandInitializeBuffer>()) FD3D12RHICommandInitializeBuffer(BufferOut, SrcResourceLoc, Size);
 			}
 		}
