Added checks to help diagnosing UE-71662

diff --git a/Engine/Source/Runtime/D3D12RHI/Private/D3D12RayTracing.cpp b/Engine/Source/Runtime/D3D12RHI/Private/D3D12RayTracing.cpp
--- a/Engine/Source/Runtime/D3D12RHI/Private/D3D12RayTracing.cpp
+++ b/Engine/Source/Runtime/D3D12RHI/Private/D3D12RayTracing.cpp
@@ -409,6 +409,9 @@ struct FD3D12RayTracingDescriptorHeap : public FD3D12DeviceChild
 
 		CPUBase = D3D12Heap->GetCPUDescriptorHandleForHeapStart();
 		GPUBase = D3D12Heap->GetGPUDescriptorHandleForHeapStart();
+
+		checkf(CPUBase.ptr, TEXT("Ray tracing descriptor heap of type %d returned from descriptor heap cache is invalid."), Type);
+
 		DescriptorSize = GetParentDevice()->GetDevice()->GetDescriptorHandleIncrementSize(Type);
 	}
 
@@ -516,6 +519,7 @@ public:
 			DescriptorTableBaseIndex = Heap.Allocate(NumDescriptors);
 
 			D3D12_CPU_DESCRIPTOR_HANDLE DestDescriptor = Heap.GetDescriptorCPU(DescriptorTableBaseIndex);
+			checkf(Heap.CPUBase.ptr, TEXT("Ray tracing descriptor heap of type %d assigned to descriptor cache is invalid."), Type);
 			GetParentDevice()->GetDevice()->CopyDescriptors(1, &DestDescriptor, &NumDescriptors, NumDescriptors, Descriptors, nullptr, Type);
 
 			Map.Add(Key, DescriptorTableBaseIndex);
