Disable an unnecessary RHI thread flush in LockTexture2D. It was previously assumed this was unsafe on UMA hardware due to virtual texture implementation writing directly to GPU memory, but in the VT case, it should always be safe, because we only ever write to this memory immediately after mapping the new pages in the texture streamer
Add a cvar D3D12.LockTexture2DRHIFlush so we can hotfix it on if necessary

diff --git a/Engine/Source/Runtime/D3D12RHI/Private/D3D12Texture.cpp b/Engine/Source/Runtime/D3D12RHI/Private/D3D12Texture.cpp
--- a/Engine/Source/Runtime/D3D12RHI/Private/D3D12Texture.cpp
+++ b/Engine/Source/Runtime/D3D12RHI/Private/D3D12Texture.cpp
@@ -19,6 +19,13 @@ static FAutoConsoleVariableRef CVarAdjustTexturePoolSizeBasedOnBudget(
 	TEXT("Indicates if the RHI should lower the texture pool size when the application is over the memory budget provided by the OS. This can result in lower quality textures (but hopefully improve performance).")
 	);
 
+static TAutoConsoleVariable<int32> CVarD3D12Texture2DRHIFlush(
+	TEXT("D3D12.LockTexture2DRHIFlush"),
+	0,
+	TEXT("If enabled, we do RHIThread flush on LockTexture2D. Likely not required on any platform, but keeping just for testing for now")
+	TEXT(" 0: off (default)\n")
+	TEXT(" 1: on"),
+	ECVF_RenderThreadSafe);
 
 static TAutoConsoleVariable<int32> CVarUseUpdateTexture3DComputeShader(
 	TEXT("D3D12.UseUpdateTexture3DComputeShader"),
@@ -1854,14 +1861,7 @@ void TD3D12Texture2D<RHIResourceType>::UpdateTexture2D(class FRHICommandListImme
 
 void* FD3D12DynamicRHI::LockTexture2D_RenderThread(class FRHICommandListImmediate& RHICmdList, FTexture2DRHIParamRef TextureRHI, uint32 MipIndex, EResourceLockMode LockMode, uint32& DestStride, bool bLockWithinMiptail, bool bNeedsDefaultRHIFlush)
 {
-	// TODO: Remove when we are sure this branch is unnecessary
-	// bNeedsDefaultRHIFlush is default to true, which causes unecessary stall on render thread
-	// TD3D12Texture2D<...>::Lock handles different type of locks:
-	//   RLM_WriteOnly - system memory is allocated and returned. Buffer renaming happens on
-	//     RHI thread and is triggered on unlock
-	//   RLM_ReadOnly - slow path, flush RHI thread
-	// XB1 uses virtual textures and read/writes GPU memory directly so a flush may still be needed
-	if (!PLATFORM_WINDOWS && bNeedsDefaultRHIFlush)
+	if (CVarD3D12Texture2DRHIFlush.GetValueOnRenderThread() && bNeedsDefaultRHIFlush)
 	{
 		QUICK_SCOPE_CYCLE_COUNTER(STAT_RHIMETHOD_LockTexture2D_Flush);
 		RHICmdList.ImmediateFlush(EImmediateFlushType::FlushRHIThread);
@@ -1882,8 +1882,7 @@ void* FD3D12DynamicRHI::RHILockTexture2D(FTexture2DRHIParamRef TextureRHI, uint3
 
 void FD3D12DynamicRHI::UnlockTexture2D_RenderThread(class FRHICommandListImmediate& RHICmdList, FTexture2DRHIParamRef TextureRHI, uint32 MipIndex, bool bLockWithinMiptail, bool bNeedsDefaultRHIFlush)
 {
-	// TODO: Remove when we are sure this branch is unnecessary
-	if (!PLATFORM_WINDOWS && bNeedsDefaultRHIFlush)
+	if (CVarD3D12Texture2DRHIFlush.GetValueOnRenderThread() && bNeedsDefaultRHIFlush)
 	{
 		QUICK_SCOPE_CYCLE_COUNTER(STAT_RHIMETHOD_UnlockTexture2D_Flush);
 		RHICmdList.ImmediateFlush(EImmediateFlushType::FlushRHIThread);
