Cherry-pick 5735905: Reduced the number of descriptor heaps that are kept alive in D3D12RayTracing. Default shader tables no longer allocate heaps. Unused heaps are now released after 100 frames.

diff --git a/Engine/Source/Runtime/D3D12RHI/Private/D3D12RayTracing.cpp b/Engine/Source/Runtime/D3D12RHI/Private/D3D12RayTracing.cpp
--- a/Engine/Source/Runtime/D3D12RHI/Private/D3D12RayTracing.cpp
+++ b/Engine/Source/Runtime/D3D12RHI/Private/D3D12RayTracing.cpp
@@ -343,6 +343,10 @@ public:
 			}
 		}
 
+		// Compatible heap was not found in cache, so create a new one.
+
+		ReleaseStaleEntries(100, CompletedFenceValue); // Release heaps that were not used for 100 frames before allocating new.
+
 		D3D12_DESCRIPTOR_HEAP_DESC Desc = {};
 
 		Desc.Flags = D3D12_DESCRIPTOR_HEAP_FLAG_SHADER_VISIBLE;
@@ -362,6 +366,25 @@ public:
 		return Result;
 	}
 
+	void ReleaseStaleEntries(uint32 MaxAge, uint64 CompletedFenceValue)
+	{
+		int32 EntryIndex = 0;
+		while (EntryIndex < Entries.Num())
+		{
+			Entry& It = Entries[EntryIndex];
+			if ((It.FenceValue + MaxAge) <= CompletedFenceValue)
+			{
+				It.Heap->Release();
+				Entries[EntryIndex] = Entries.Last();
+				Entries.Pop(false);
+			}
+			else
+			{
+				EntryIndex++;
+			}
+		}
+	}
+
 	void Flush()
 	{
 		FD3D12Device* Device = GetParentDevice();
@@ -569,10 +592,14 @@ public:
 
 	FD3D12RayTracingShaderTable(FD3D12Device* Device)
 		: FD3D12DeviceChild(Device)
-		, DescriptorCache(Device)
 	{
 	}
 
+	~FD3D12RayTracingShaderTable()
+	{
+		delete DescriptorCache;
+	}
+
 	void Init(uint32 InNumRayGenShaders, uint32 InNumMissShaders, uint32 InNumHitRecords, uint32 LocalRootDataSize)
 	{
 		checkf(LocalRootDataSize <= 4096, TEXT("The maximum size of a local root signature is 4KB.")); // as per section 4.22.1 of DXR spec v1.0
@@ -581,17 +608,25 @@ public:
 		HitRecordSizeUnaligned = ShaderIdentifierSize + LocalRootDataSize;
 		HitRecordStride = RoundUpToNextMultiple(HitRecordSizeUnaligned, D3D12_RAYTRACING_SHADER_RECORD_BYTE_ALIGNMENT);
 
-		// Minimum number of descriptors required to support binding global resources (arbitrarily chosen)
-		// #dxr_todo: Remove this when RT descriptors are sub-allocated from the global view descriptor heap.
-		const uint32 MinNumViewDescriptors = 1024;
-		const uint32 ApproximateDescriptorsPerRecord = 32; // #dxr_todo: calculate this based on shader reflection data
+		// Custom descriptor cache is only required when local resources may be bound.
+		// If only global resources are used, then transient descriptor cache can be used.
+		const bool bNeedsDescriptorCache = InNumHitRecords * LocalRootDataSize != 0;
 
-		// D3D12 is guaranteed to support 1M (D3D12_MAX_SHADER_VISIBLE_DESCRIPTOR_HEAP_SIZE_TIER_1) descriptors in a CBV/SRV/UAV heap, so clamp the size to this.
-		// https://docs.microsoft.com/en-us/windows/desktop/direct3d12/hardware-support
-		const uint32 NumViewDescriptors = FMath::Max(MinNumViewDescriptors, FMath::Min<uint32>(InNumHitRecords * ApproximateDescriptorsPerRecord, D3D12_MAX_SHADER_VISIBLE_DESCRIPTOR_HEAP_SIZE_TIER_1));
-		const uint32 NumSamplerDescriptors = D3D12_MAX_SHADER_VISIBLE_SAMPLER_HEAP_SIZE;
-
-		DescriptorCache.Init(NumViewDescriptors, NumSamplerDescriptors);
+		if (bNeedsDescriptorCache)
+		{
+			// Minimum number of descriptors required to support binding global resources (arbitrarily chosen)
+			// #dxr_todo UE-72158: Remove this when RT descriptors are sub-allocated from the global view descriptor heap.
+			const uint32 MinNumViewDescriptors = 1024;
+			const uint32 ApproximateDescriptorsPerRecord = 32; // #dxr_todo: calculate this based on shader reflection data
+
+			// D3D12 is guaranteed to support 1M (D3D12_MAX_SHADER_VISIBLE_DESCRIPTOR_HEAP_SIZE_TIER_1) descriptors in a CBV/SRV/UAV heap, so clamp the size to this.
+			// https://docs.microsoft.com/en-us/windows/desktop/direct3d12/hardware-support
+			const uint32 NumViewDescriptors = FMath::Max(MinNumViewDescriptors, FMath::Min<uint32>(InNumHitRecords * ApproximateDescriptorsPerRecord, D3D12_MAX_SHADER_VISIBLE_DESCRIPTOR_HEAP_SIZE_TIER_1));
+			const uint32 NumSamplerDescriptors = D3D12_MAX_SHADER_VISIBLE_SAMPLER_HEAP_SIZE;
+
+			DescriptorCache = new FD3D12RayTracingDescriptorCache(GetParentDevice());
+			DescriptorCache->Init(NumViewDescriptors, NumSamplerDescriptors);
+		}
 
 		NumRayGenShaders = InNumRayGenShaders;
 		NumMissShaders = InNumMissShaders;
@@ -767,8 +802,7 @@ public:
 	TRefCountPtr<FD3D12MemBuffer> Buffer;
 
 	// SBTs have their own descriptor heaps
-	FD3D12RayTracingDescriptorCache DescriptorCache;
-
+	FD3D12RayTracingDescriptorCache* DescriptorCache = nullptr;
 
 #if ENABLE_RESIDENCY_MANAGEMENT
 	// A set of all resources referenced by this shader table for the purpose of updating residency before ray tracing work dispatch.
@@ -2275,11 +2309,13 @@ static void DispatchRays(FD3D12CommandContext& CommandContext,
 
 	FD3D12RayTracingShader* RayGenShader = Pipeline->RayGenShaders.Shaders[RayGenShaderIndex];
 
-	if (OptShaderTable)
+	if (OptShaderTable && OptShaderTable->DescriptorCache)
 	{
-		OptShaderTable->DescriptorCache.SetDescriptorHeaps(CommandContext);
+		FD3D12RayTracingDescriptorCache* DescriptorCache = OptShaderTable->DescriptorCache;
+
+		DescriptorCache->SetDescriptorHeaps(CommandContext);
 		FD3D12RayTracingGlobalResourceBinder ResourceBinder(CommandContext);
-		SetRayTracingShaderResources(CommandContext, RayGenShader, GlobalBindings, OptShaderTable->DescriptorCache, ResourceBinder);
+		SetRayTracingShaderResources(CommandContext, RayGenShader, GlobalBindings, *DescriptorCache, ResourceBinder);
 
 		// #dxr_todo: avoid updating residency if this scene was already used on the current command list (i.e. multiple ray dispatches are performed back-to-back)
 		OptShaderTable->UpdateResidency(CommandContext);
@@ -2441,13 +2477,14 @@ void FD3D12CommandContext::RHISetRayTracingHitGroup(
 	const FD3D12RayTracingShader* Shader = Pipeline->HitGroupShaders.Shaders[HitGroupIndex];
 
 	FD3D12RayTracingLocalResourceBinder ResourceBinder(*this, ShaderTable, Shader->pRootSignature, RecordIndex);
+	check(ShaderTable->DescriptorCache);
 	SetRayTracingShaderResources(*this, Shader,
 		0, nullptr, // Textures
 		0, nullptr, // SRVs
 		NumUniformBuffers, UniformBuffers,
 		0, nullptr, // Samplers
 		0, nullptr, // UAVs
-		ShaderTable->DescriptorCache, ResourceBinder);
+		*(ShaderTable->DescriptorCache), ResourceBinder);
 }
 
 #endif // D3D12_RHI_RAYTRACING
