Cherry-pick 5773457: Shader reflection data is now used to more accurately compute how many view descriptors to allocate for ray tracing resource descriptor heaps.

diff --git a/Engine/Source/Runtime/D3D12RHI/Private/D3D12RayTracing.cpp b/Engine/Source/Runtime/D3D12RHI/Private/D3D12RayTracing.cpp
--- a/Engine/Source/Runtime/D3D12RHI/Private/D3D12RayTracing.cpp
+++ b/Engine/Source/Runtime/D3D12RHI/Private/D3D12RayTracing.cpp
@@ -600,7 +600,7 @@ public:
 		delete DescriptorCache;
 	}
 
-	void Init(uint32 InNumRayGenShaders, uint32 InNumMissShaders, uint32 InNumHitRecords, uint32 LocalRootDataSize)
+	void Init(uint32 InNumRayGenShaders, uint32 InNumMissShaders, uint32 InNumHitRecords, uint32 LocalRootDataSize, uint32 MaxViewDescriptorsPerRecord)
 	{
 		checkf(LocalRootDataSize <= 4096, TEXT("The maximum size of a local root signature is 4KB.")); // as per section 4.22.1 of DXR spec v1.0
 		checkf(InNumRayGenShaders >= 1, TEXT("All shader tables must contain at least one raygen shader."));
@@ -617,11 +617,10 @@ public:
 			// Minimum number of descriptors required to support binding global resources (arbitrarily chosen)
 			// #dxr_todo UE-72158: Remove this when RT descriptors are sub-allocated from the global view descriptor heap.
 			const uint32 MinNumViewDescriptors = 1024;
-			const uint32 ApproximateDescriptorsPerRecord = 32; // #dxr_todo: calculate this based on shader reflection data
 
 			// D3D12 is guaranteed to support 1M (D3D12_MAX_SHADER_VISIBLE_DESCRIPTOR_HEAP_SIZE_TIER_1) descriptors in a CBV/SRV/UAV heap, so clamp the size to this.
 			// https://docs.microsoft.com/en-us/windows/desktop/direct3d12/hardware-support
-			const uint32 NumViewDescriptors = FMath::Max(MinNumViewDescriptors, FMath::Min<uint32>(InNumHitRecords * ApproximateDescriptorsPerRecord, D3D12_MAX_SHADER_VISIBLE_DESCRIPTOR_HEAP_SIZE_TIER_1));
+			const uint32 NumViewDescriptors = FMath::Max(MinNumViewDescriptors, FMath::Min<uint32>(InNumHitRecords * MaxViewDescriptorsPerRecord, D3D12_MAX_SHADER_VISIBLE_DESCRIPTOR_HEAP_SIZE_TIER_1));
 			const uint32 NumSamplerDescriptors = D3D12_MAX_SHADER_VISIBLE_SAMPLER_HEAP_SIZE;
 
 			DescriptorCache = new FD3D12RayTracingDescriptorCache(GetParentDevice());
@@ -976,6 +975,8 @@ public:
 		TArray<FHitGroupEntryIndices> HitGroupEntryIndices;
 		HitGroupEntryIndices.Reserve(InitializerHitGroups.Num());
 
+		MaxHitGroupViewDescriptors = 0;
+
 		for (const FRayTracingShaderRHIParamRef ShaderRHI : InitializerHitGroups)
 		{
 			FD3D12RayTracingShader* Shader = FD3D12DynamicRHI::ResourceCast(ShaderRHI);
@@ -983,6 +984,9 @@ public:
 			checkf(Shader, TEXT("A valid ray tracing hit group shader must be provided for all elements in the FRayTracingPipelineStateInitializer hit group table."));
 			checkf(!Shader->ResourceCounts.bGlobalUniformBufferUsed, TEXT("Global uniform buffers are not implemented for ray tracing shaders"));
 
+			const uint32 ShaderViewDescriptors = Shader->ResourceCounts.NumSRVs + Shader->ResourceCounts.NumUAVs;
+			MaxHitGroupViewDescriptors = FMath::Max(MaxHitGroupViewDescriptors, ShaderViewDescriptors);
+
 			HitGroupShaders.Shaders.Add(Shader);
 
 			MaxLocalRootSignatureSize = FMath::Max(MaxLocalRootSignatureSize, Shader->pRootSignature->GetTotalRootSignatureSizeInBytes());
@@ -1195,7 +1199,7 @@ public:
 		// Hit record indexing and local resources access is disabled when using using this SBT.
 
 		const uint32 DefaultLocalRootDataSize = 0; // Shaders in default SBT are not allowed to access any local resources
-		DefaultShaderTable.Init(RayGenShaders.Num(), MissShaders.Num(), 0, DefaultLocalRootDataSize);
+		DefaultShaderTable.Init(RayGenShaders.Num(), MissShaders.Num(), 0, DefaultLocalRootDataSize, 0);
 		DefaultShaderTable.SetRayGenIdentifiers(RayGenShaders.Identifiers);
 		DefaultShaderTable.SetMissIdentifiers(MissShaders.Identifiers);
 		DefaultShaderTable.SetDefaultHitGroupIdentifier(HitGroupShaders.Identifiers[0]);
@@ -1221,6 +1225,7 @@ public:
 	bool bAllowHitGroupIndexing = true;
 
 	uint32 MaxLocalRootSignatureSize = 0;
+	uint32 MaxHitGroupViewDescriptors = 0;
 };
 
 class FD3D12BasicRayTracingPipeline
@@ -1805,7 +1810,8 @@ FD3D12RayTracingShaderTable* FD3D12RayTracingScene::FindOrCreateShaderTable(cons
 		Pipeline->RayGenShaders.Num(),
 		Pipeline->MissShaders.Num(),
 		NumHitGroupSlots,
-		Pipeline->MaxLocalRootSignatureSize);
+		Pipeline->MaxLocalRootSignatureSize,
+		Pipeline->MaxHitGroupViewDescriptors);
 
 	CreatedShaderTable->SetRayGenIdentifiers(Pipeline->RayGenShaders.Identifiers);
 	CreatedShaderTable->SetMissIdentifiers(Pipeline->MissShaders.Identifiers);
