Cherry-pick 5807081: Added basic resource usage stats for ray tracing. Use `stat D3D12RayTracing` to show the report.

diff --git a/Engine/Source/Runtime/D3D12RHI/Private/D3D12RayTracing.cpp b/Engine/Source/Runtime/D3D12RHI/Private/D3D12RayTracing.cpp
--- a/Engine/Source/Runtime/D3D12RHI/Private/D3D12RayTracing.cpp
+++ b/Engine/Source/Runtime/D3D12RHI/Private/D3D12RayTracing.cpp
@@ -48,6 +48,33 @@ static FAutoConsoleVariableRef CVarRayTracingDebugDisableTriangleCull(
 	TEXT("Forces all ray tracing geometry instances to be double-sided by disabling back-face culling. This is useful for debugging and profiling. (default = 0)")
 );
 
+// Ray tracing stat counters
+
+DECLARE_STATS_GROUP(TEXT("D3D12RHI: Ray Tracing"), STATGROUP_D3D12RayTracing, STATCAT_Advanced);
+
+DECLARE_DWORD_ACCUMULATOR_STAT(TEXT("Created pipelines (total)"), STAT_D3D12RayTracingCreatedPipelines, STATGROUP_D3D12RayTracing);
+DECLARE_DWORD_ACCUMULATOR_STAT(TEXT("Compiled shaders (total)"), STAT_D3D12RayTracingCompiledShaders, STATGROUP_D3D12RayTracing);
+
+DECLARE_DWORD_ACCUMULATOR_STAT(TEXT("Allocated bottom level acceleration structures"), STAT_D3D12RayTracingAllocatedBLAS, STATGROUP_D3D12RayTracing);
+DECLARE_DWORD_ACCUMULATOR_STAT(TEXT("Allocated top level acceleration structures"), STAT_D3D12RayTracingAllocatedTLAS, STATGROUP_D3D12RayTracing);
+DECLARE_DWORD_ACCUMULATOR_STAT(TEXT("Triangles in all BL acceleration structures"), STAT_D3D12RayTracingTrianglesBLAS, STATGROUP_D3D12RayTracing);
+
+DECLARE_DWORD_ACCUMULATOR_STAT(TEXT("Allocated sampler descriptor heaps"), STAT_D3D12RayTracingSamplerDescriptorHeaps, STATGROUP_D3D12RayTracing);
+DECLARE_DWORD_ACCUMULATOR_STAT(TEXT("Allocated sampler descriptors"), STAT_D3D12RayTracingSamplerDescriptors, STATGROUP_D3D12RayTracing);
+
+DECLARE_DWORD_ACCUMULATOR_STAT(TEXT("Allocated view descriptor heaps"), STAT_D3D12RayTracingViewDescriptorHeaps, STATGROUP_D3D12RayTracing);
+DECLARE_DWORD_ACCUMULATOR_STAT(TEXT("Allocated view descriptors"), STAT_D3D12RayTracingViewDescriptors, STATGROUP_D3D12RayTracing);
+
+DECLARE_DWORD_COUNTER_STAT(TEXT("Used sampler descriptors (per frame)"), STAT_D3D12RayTracingUsedSamplerDescriptors, STATGROUP_D3D12RayTracing);
+DECLARE_DWORD_COUNTER_STAT(TEXT("Used view descriptors (per frame)"), STAT_D3D12RayTracingUsedViewDescriptors, STATGROUP_D3D12RayTracing);
+
+DECLARE_DWORD_COUNTER_STAT(TEXT("Built BL AS (per frame)"), STAT_D3D12RayTracingBuiltBLAS, STATGROUP_D3D12RayTracing);
+DECLARE_DWORD_COUNTER_STAT(TEXT("Updated BL AS (per frame)"), STAT_D3D12RayTracingUpdatedBLAS, STATGROUP_D3D12RayTracing);
+DECLARE_DWORD_COUNTER_STAT(TEXT("Built TL AS (per frame)"), STAT_D3D12RayTracingBuiltTLAS, STATGROUP_D3D12RayTracing);
+DECLARE_DWORD_COUNTER_STAT(TEXT("Updated TL AS (per frame)"), STAT_D3D12RayTracingUpdatedTLAS, STATGROUP_D3D12RayTracing);
+
+DECLARE_MEMORY_STAT(TEXT("Used Video Memory"), STAT_D3D12RayTracingUsedVideoMemory, STATGROUP_D3D12RayTracing);
+
 // Built-in local root parameters that are always bound to all hit shaders
 struct FHitGroupSystemParameters
 {
@@ -165,6 +192,14 @@ static TRefCountPtr<ID3D12StateObject> CreateRayTracingStateObject(
 {
 	checkf(LocalRootSignatureAssociations.Num() == Exports.Num(), TEXT("There must be exactly one local root signature association per export."));
 
+#if !NO_LOGGING
+	const uint32 NumShadersWarningThreshold = 1000;
+	if (Exports.Num() > NumShadersWarningThreshold)
+	{
+		UE_LOG(LogD3D12RHI, Warning, TEXT("Creating ray tracing pipeline with %d shaders. Reduce the number of unique shaders in the scene to avoid CPU stalls."), Exports.Num());
+	}
+#endif
+
 	TRefCountPtr<ID3D12StateObject> Result;
 
 	// There are several pipeline sub-objects that are always required:
@@ -268,6 +303,9 @@ static TRefCountPtr<ID3D12StateObject> CreateRayTracingStateObject(
 
 	VERIFYD3D12RESULT(RayTracingDevice->CreateStateObject(&Desc, IID_PPV_ARGS(Result.GetInitReference())));
 
+	INC_DWORD_STAT(STAT_D3D12RayTracingCreatedPipelines);
+	INC_DWORD_STAT_BY(STAT_D3D12RayTracingCompiledShaders, NumExports);
+
 	return Result;
 }
 
@@ -303,6 +341,17 @@ public:
 		FScopeLock Lock(&CriticalSection);
 		for (const Entry& It : Entries)
 		{
+			if (It.Type == D3D12_DESCRIPTOR_HEAP_TYPE_CBV_SRV_UAV)
+			{
+				DEC_DWORD_STAT(STAT_D3D12RayTracingViewDescriptorHeaps);
+				DEC_DWORD_STAT_BY(STAT_D3D12RayTracingViewDescriptors, It.NumDescriptors);
+			}
+			else
+			{
+				DEC_DWORD_STAT(STAT_D3D12RayTracingSamplerDescriptorHeaps);
+				DEC_DWORD_STAT_BY(STAT_D3D12RayTracingSamplerDescriptors, It.NumDescriptors);
+			}
+
 			It.Heap->Release();
 		}
 		Entries.Empty();
@@ -363,6 +412,17 @@ public:
 		Result.Type = Type;
 		Result.Heap = D3D12Heap;
 
+		if (Desc.Type == D3D12_DESCRIPTOR_HEAP_TYPE_CBV_SRV_UAV)
+		{
+			INC_DWORD_STAT(STAT_D3D12RayTracingViewDescriptorHeaps);
+			INC_DWORD_STAT_BY(STAT_D3D12RayTracingViewDescriptors, NumDescriptors);
+		}
+		else
+		{
+			INC_DWORD_STAT(STAT_D3D12RayTracingSamplerDescriptorHeaps);
+			INC_DWORD_STAT_BY(STAT_D3D12RayTracingSamplerDescriptors, NumDescriptors);
+		}
+
 		return Result;
 	}
 
@@ -374,7 +434,19 @@ public:
 			Entry& It = Entries[EntryIndex];
 			if ((It.FenceValue + MaxAge) <= CompletedFenceValue)
 			{
+				if (It.Type == D3D12_DESCRIPTOR_HEAP_TYPE_CBV_SRV_UAV)
+				{
+					DEC_DWORD_STAT(STAT_D3D12RayTracingViewDescriptorHeaps);
+					DEC_DWORD_STAT_BY(STAT_D3D12RayTracingViewDescriptors, It.NumDescriptors);
+				}
+				else
+				{
+					DEC_DWORD_STAT(STAT_D3D12RayTracingSamplerDescriptorHeaps);
+					DEC_DWORD_STAT_BY(STAT_D3D12RayTracingSamplerDescriptors, It.NumDescriptors);
+				}
+
 				It.Heap->Release();
+
 				Entries[EntryIndex] = Entries.Last();
 				Entries.Pop(false);
 			}
@@ -546,6 +618,15 @@ public:
 			GetParentDevice()->GetDevice()->CopyDescriptors(1, &DestDescriptor, &NumDescriptors, NumDescriptors, Descriptors, nullptr, Type);
 
 			Map.Add(Key, DescriptorTableBaseIndex);
+
+			if (Type == D3D12_DESCRIPTOR_HEAP_TYPE_CBV_SRV_UAV)
+			{
+				INC_DWORD_STAT_BY(STAT_D3D12RayTracingUsedViewDescriptors, NumDescriptors);
+			}
+			else
+			{
+				INC_DWORD_STAT_BY(STAT_D3D12RayTracingUsedSamplerDescriptors, NumDescriptors);
+			}
 		}
 
 		return DescriptorTableBaseIndex;
@@ -1396,6 +1477,8 @@ FRayTracingGeometryRHIRef FD3D12DynamicRHI::RHICreateRayTracingGeometry(const FR
 		return Mesh;
 	});
 
+	INC_DWORD_STAT_BY(STAT_D3D12RayTracingTrianglesBLAS, Initializer.TotalPrimitiveCount);
+
 	return Result;
 }
 
@@ -1425,6 +1508,28 @@ FRayTracingSceneRHIRef FD3D12DynamicRHI::RHICreateRayTracingScene(const FRayTrac
 	});
 }
 
+FD3D12RayTracingGeometry::FD3D12RayTracingGeometry(FD3D12Device* Device) 
+	: FD3D12DeviceChild(Device) 
+{
+	INC_DWORD_STAT(STAT_D3D12RayTracingAllocatedBLAS);
+}
+
+FD3D12RayTracingGeometry::~FD3D12RayTracingGeometry()
+{
+	if (AccelerationStructureBuffer)
+	{
+		DEC_MEMORY_STAT_BY(STAT_D3D12RayTracingUsedVideoMemory, AccelerationStructureBuffer->GetSize());
+	}
+
+	if (ScratchBuffer)
+	{
+		DEC_MEMORY_STAT_BY(STAT_D3D12RayTracingUsedVideoMemory, ScratchBuffer->GetSize());
+	}
+
+	DEC_DWORD_STAT_BY(STAT_D3D12RayTracingTrianglesBLAS, TotalPrimitiveCount);
+	DEC_DWORD_STAT(STAT_D3D12RayTracingAllocatedBLAS);
+}
+
 void FD3D12RayTracingGeometry::TransitionBuffers(FD3D12CommandContext& CommandContext)
 {
 	// Transition vertex and index resources..
@@ -1579,7 +1684,10 @@ void FD3D12RayTracingGeometry::BuildAccelerationStructure(FD3D12CommandContext&
 
 		CreateAccelerationStructureBuffers(AccelerationStructureBuffer, ScratchBuffer, Adapter, GPUIndex, PrebuildInfo);
 
-		// #dxr_todo: scratch buffers should be created in UAV state from the start
+		INC_MEMORY_STAT_BY(STAT_D3D12RayTracingUsedVideoMemory, AccelerationStructureBuffer->GetSize());
+		INC_MEMORY_STAT_BY(STAT_D3D12RayTracingUsedVideoMemory, ScratchBuffer->GetSize());
+
+		// #dxr_todo UE-72161: scratch buffers should be created in UAV state from the start
 		FD3D12DynamicRHI::TransitionResource(CommandContext.CommandListHandle, ScratchBuffer.GetReference()->GetResource(), D3D12_RESOURCE_STATE_UNORDERED_ACCESS, 0);
 	}
 
@@ -1602,21 +1710,47 @@ void FD3D12RayTracingGeometry::BuildAccelerationStructure(FD3D12CommandContext&
 		ID3D12GraphicsCommandList4* RayTracingCommandList = CommandContext.CommandListHandle.RayTracingCommandList();
 		RayTracingCommandList->BuildRaytracingAccelerationStructure(&BuildDesc, 0, nullptr);
 		bIsAccelerationStructureDirty = false;
+
+		if (bIsUpdate)
+		{
+			INC_DWORD_STAT(STAT_D3D12RayTracingUpdatedBLAS);
+		}
+		else
+		{
+			INC_DWORD_STAT(STAT_D3D12RayTracingBuiltBLAS);
+		}
 	}
 
 	// We don't need to keep a scratch buffer after initial build if acceleration structure is static.
 	if (!(BuildFlags & D3D12_RAYTRACING_ACCELERATION_STRUCTURE_BUILD_FLAG_ALLOW_UPDATE))
 	{
+		DEC_MEMORY_STAT_BY(STAT_D3D12RayTracingUsedVideoMemory, ScratchBuffer->GetSize());
 		ScratchBuffer = nullptr;
 	}
 }
 
+FD3D12RayTracingScene::FD3D12RayTracingScene(FD3D12Device* Device)
+	: FD3D12DeviceChild(Device)
+	, AccelerationStructureView(new FD3D12ShaderResourceView(Device))
+{
+	INC_DWORD_STAT(STAT_D3D12RayTracingAllocatedTLAS);
+
+	ShaderResourceView = AccelerationStructureView;
+};
+
 FD3D12RayTracingScene::~FD3D12RayTracingScene()
 {
 	for (auto Item : ShaderTables)
 	{
 		delete Item.Value;
 	}
+
+	if (AccelerationStructureBuffer)
+	{
+		DEC_MEMORY_STAT_BY(STAT_D3D12RayTracingUsedVideoMemory, AccelerationStructureBuffer->GetSize());
+	}
+
+	DEC_DWORD_STAT(STAT_D3D12RayTracingAllocatedTLAS);
 }
 
 void FD3D12RayTracingScene::BuildAccelerationStructure(FD3D12CommandContext& CommandContext, D3D12_RAYTRACING_ACCELERATION_STRUCTURE_BUILD_FLAGS BuildFlags)
@@ -1642,7 +1776,9 @@ void FD3D12RayTracingScene::BuildAccelerationStructure(FD3D12CommandContext& Com
 
 	CreateAccelerationStructureBuffers(AccelerationStructureBuffer, ScratchBuffer, Adapter, GPUIndex, PrebuildInfo);
 
-	// #dxr_todo: scratch buffers should be created in UAV state from the start
+	INC_MEMORY_STAT_BY(STAT_D3D12RayTracingUsedVideoMemory, AccelerationStructureBuffer->GetSize());
+
+	// #dxr_todo UE-72161: scratch buffers should be created in UAV state from the start
 	FD3D12DynamicRHI::TransitionResource(CommandContext.CommandListHandle, ScratchBuffer.GetReference()->GetResource(), D3D12_RESOURCE_STATE_UNORDERED_ACCESS, 0);
 
 	if (bAccelerationStructureViewInitialized)
@@ -1664,6 +1800,8 @@ void FD3D12RayTracingScene::BuildAccelerationStructure(FD3D12CommandContext& Com
 
 	// Create and fill instance buffer
 
+	TotalPrimitiveCount = 0;
+
 	if (Instances.Num())
 	{
 		FRHIResourceCreateInfo CreateInfo;
@@ -1733,6 +1871,8 @@ void FD3D12RayTracingScene::BuildAccelerationStructure(FD3D12CommandContext& Com
 
 			MappedData[InstanceIndex] = InstanceDesc;
 			++InstanceIndex;
+
+			TotalPrimitiveCount += Geometry->TotalPrimitiveCount;
 		}
 
 		Adapter->GetOwningRHI()->UnlockBuffer(nullptr, InstanceBuffer.GetReference());
@@ -1763,6 +1903,15 @@ void FD3D12RayTracingScene::BuildAccelerationStructure(FD3D12CommandContext& Com
 	ID3D12GraphicsCommandList4* RayTracingCommandList = CommandContext.CommandListHandle.RayTracingCommandList();
 	RayTracingCommandList->BuildRaytracingAccelerationStructure(&BuildDesc, 0, nullptr);
 
+	if (bIsUpdateMode)
+	{
+		INC_DWORD_STAT(STAT_D3D12RayTracingUpdatedTLAS);
+	}
+	else
+	{
+		INC_DWORD_STAT(STAT_D3D12RayTracingBuiltTLAS);
+	}
+
 	// UAV barrier is used here to ensure that the acceleration structure build is complete before any rays are traced
 	// #dxr_todo: these barriers should ideally be inserted by the high level code to allow more overlapped execution
 	CommandContext.CommandListHandle.AddUAVBarrier();
diff --git a/Engine/Source/Runtime/D3D12RHI/Private/D3D12RayTracing.h b/Engine/Source/Runtime/D3D12RHI/Private/D3D12RayTracing.h
--- a/Engine/Source/Runtime/D3D12RHI/Private/D3D12RayTracing.h
+++ b/Engine/Source/Runtime/D3D12RHI/Private/D3D12RayTracing.h
@@ -16,7 +16,9 @@ class FD3D12RayTracingGeometry : public FRHIRayTracingGeometry, public FD3D12Dev
 {
 public:
 
-	FD3D12RayTracingGeometry(FD3D12Device* Device) : FD3D12DeviceChild(Device) {}
+	FD3D12RayTracingGeometry(FD3D12Device* Device);
+	~FD3D12RayTracingGeometry();
+
 	void TransitionBuffers(FD3D12CommandContext& CommandContext);
 	void BuildAccelerationStructure(FD3D12CommandContext& CommandContext, bool bIsUpdate);
 
@@ -45,13 +47,7 @@ class FD3D12RayTracingScene : public FRHIRayTracingScene, public FD3D12DeviceChi
 {
 public:
 
-	FD3D12RayTracingScene(FD3D12Device* Device) 
-		: FD3D12DeviceChild(Device)
-		, AccelerationStructureView(new FD3D12ShaderResourceView(Device))
-	{
-		ShaderResourceView = AccelerationStructureView;
-	};
-
+	FD3D12RayTracingScene(FD3D12Device* Device);
 	~FD3D12RayTracingScene();
 
 	void BuildAccelerationStructure(FD3D12CommandContext& CommandContext, D3D12_RAYTRACING_ACCELERATION_STRUCTURE_BUILD_FLAGS BuildFlags);
@@ -75,7 +71,9 @@ public:
 	uint32 NumTotalSegments = 0;
 	uint32 GetHitRecordBaseIndex(uint32 InstanceIndex, uint32 SegmentIndex) const { return (SegmentPrefixSum[InstanceIndex] + SegmentIndex) * ShaderSlotsPerGeometrySegment; }
 
-	// #dxr_todo: shader tables should be explicitly registered and unregistered with the scene
+	uint32 TotalPrimitiveCount = 0; // Combined number of primitives in all geometry instances
+
+	// #dxr_todo UE-68230: shader tables should be explicitly registered and unregistered with the scene
 	FD3D12RayTracingShaderTable* FindOrCreateShaderTable(const FD3D12RayTracingPipelineState* Pipeline);
 	FD3D12RayTracingShaderTable* FindExistingShaderTable(const FD3D12RayTracingPipelineState* Pipeline) const;
 
