Make query heaps resizable, needed to avoid crash in UE-71291. This also addresses the warning observed in UE-63932, UE-65727 and UE-67959.

diff --git a/Engine/Source/Runtime/D3D12RHI/Private/D3D12Commands.cpp b/Engine/Source/Runtime/D3D12RHI/Private/D3D12Commands.cpp
--- a/Engine/Source/Runtime/D3D12RHI/Private/D3D12Commands.cpp
+++ b/Engine/Source/Runtime/D3D12RHI/Private/D3D12Commands.cpp
@@ -60,8 +60,6 @@ void FD3D12DynamicRHI::SetupRecursiveResources()
 		PixelShader->GetPixelShader();
 	}
 
-	extern ENGINE_API TGlobalResource<FScreenVertexDeclaration> GScreenVertexDeclaration;
-
 	{
 		TShaderMapRef<FLongGPUTaskPS> PixelShader(ShaderMap);
 		PixelShader->GetPixelShader();
@@ -1159,8 +1157,7 @@ void FD3D12CommandContext::RHIBeginRenderQuery(FRenderQueryRHIParamRef QueryRHI)
 	check(IsDefaultContext());
 	check(Query->Type == RQT_Occlusion);
 
-	Query->Reset();
-	Query->HeapIndex = GetParentDevice()->GetOcclusionQueryHeap()->BeginQuery(*this);
+	GetParentDevice()->GetOcclusionQueryHeap()->BeginQuery(*this, Query);
 
 #if EXECUTE_DEBUG_COMMAND_LISTS
 	GIsDoingQuery = true;
@@ -1172,31 +1169,26 @@ void FD3D12CommandContext::RHIEndRenderQuery(FRenderQueryRHIParamRef QueryRHI)
 	FD3D12RenderQuery* Query = RetrieveObject<FD3D12RenderQuery>(QueryRHI);
 	check(IsDefaultContext());
 
+	FD3D12QueryHeap* QueryHeap = nullptr;
+
 	switch (Query->Type)
 	{
 	case RQT_Occlusion:
-	{
-		GetParentDevice()->GetOcclusionQueryHeap()->EndQuery(*this, Query->HeapIndex, Query);
-		// Multi-GPU support : by setting a timestamp, we can filter only the relevant GPUs when getting the query results.
-		Query->Timestamp = GFrameNumberRenderThread;
+		QueryHeap = GetParentDevice()->GetOcclusionQueryHeap();
 		break;
-	}
 
 	case RQT_AbsoluteTime:
-	{
-		FD3D12QueryHeap* pQueryHeap = GetParentDevice()->GetTimestampQueryHeap();
-		Query->Reset();
-		Query->HeapIndex = pQueryHeap->AllocQuery(*this);
-		pQueryHeap->EndQuery(*this, Query->HeapIndex, Query);
-		// Multi-GPU support : by setting a timestamp, we can filter only the relevant GPUs when getting the query results.
-		Query->Timestamp = GFrameNumberRenderThread;
+		QueryHeap = GetParentDevice()->GetTimestampQueryHeap();
 		break;
-	}
 
 	default:
 		ensure(false);
 	}
 
+	QueryHeap->EndQuery(*this, Query);
+	// Multi-GPU support : by setting a timestamp, we can filter only the relevant GPUs when getting the query results.
+	Query->Timestamp = GFrameNumberRenderThread;
+
 	// Query data isn't ready until it has been resolved.
 	ensure(Query->bResultIsCached == false && Query->bResolved == false);
 
diff --git a/Engine/Source/Runtime/D3D12RHI/Private/D3D12Device.cpp b/Engine/Source/Runtime/D3D12RHI/Private/D3D12Device.cpp
--- a/Engine/Source/Runtime/D3D12RHI/Private/D3D12Device.cpp
+++ b/Engine/Source/Runtime/D3D12RHI/Private/D3D12Device.cpp
@@ -34,8 +34,8 @@ FD3D12Device::FD3D12Device(FRHIGPUMask InGPUMask, FD3D12Adapter* InAdapter) :
 	SamplerAllocator(InGPUMask, FD3D12_DESCRIPTOR_HEAP_TYPE_SAMPLER, 128),
 	GlobalSamplerHeap(this, InGPUMask),
 	GlobalViewHeap(this, InGPUMask),
-	OcclusionQueryHeap(this, D3D12_QUERY_HEAP_TYPE_OCCLUSION, 65536, 4 /*frames to keep results */ * 1 /*batches per frame*/),
-	TimestampQueryHeap(this, D3D12_QUERY_HEAP_TYPE_TIMESTAMP, 8192, 4 /*frames to keep results */ * 5 /*batches per frame*/ ),
+	OcclusionQueryHeap(this, D3D12_QUERY_TYPE_OCCLUSION, 65536, 4 /*frames to keep results */ * 1 /*batches per frame*/),
+	TimestampQueryHeap(this, D3D12_QUERY_TYPE_TIMESTAMP, 8192, 4 /*frames to keep results */ * 5 /*batches per frame*/ ),
 #if WITH_PROFILEGPU
 	CmdListExecTimeQueryHeap(this, D3D12_QUERY_HEAP_TYPE_TIMESTAMP, 8192),
 #endif
diff --git a/Engine/Source/Runtime/D3D12RHI/Private/D3D12Query.cpp b/Engine/Source/Runtime/D3D12RHI/Private/D3D12Query.cpp
--- a/Engine/Source/Runtime/D3D12RHI/Private/D3D12Query.cpp
+++ b/Engine/Source/Runtime/D3D12RHI/Private/D3D12Query.cpp
@@ -46,7 +46,7 @@ bool FD3D12DynamicRHI::RHIGetRenderQueryResult(FRenderQueryRHIParamRef QueryRHI,
 	FRHIGPUMask RelevantNodeMask = FRHIGPUMask::GPU0();
 	if (GNumExplicitGPUsForRendering > 1)
 	{
-		uint64 LatestTimestamp = 0;
+		uint32 LatestTimestamp = 0;
 		for (FD3D12RenderQuery* Query = FD3D12DynamicRHI::ResourceCast(QueryRHI); Query; Query = Query->GetNextObject())
 		{
 			if (Query->Timestamp > LatestTimestamp)
@@ -150,8 +150,7 @@ bool FD3D12Device::GetQueryData(FD3D12RenderQuery& Query, bool bWait)
 	static const CD3DX12_RANGE EmptyRange(0, 0);
 
 	{
-		FD3D12Resource* const ResultBuffer = Query.Type == RQT_Occlusion ? OcclusionQueryHeap.GetResultBuffer() : TimestampQueryHeap.GetResultBuffer();
-		const FD3D12ScopeMap<uint64> MappedData(ResultBuffer, 0, &ReadRange, &EmptyRange /* Not writing any data */);
+		const FD3D12ScopeMap<uint64> MappedData(Query.ResultBuffer, 0, &ReadRange, &EmptyRange /* Not writing any data */);
 		Query.Result = MappedData[Query.HeapIndex];
 	}
 
@@ -160,7 +159,7 @@ bool FD3D12Device::GetQueryData(FD3D12RenderQuery& Query, bool bWait)
 
 void FD3D12CommandContext::RHIBeginOcclusionQueryBatch(uint32 NumQueriesInBatch)
 {
-	// Nothing to do here, we always start a new batch during RHIEndOcclusionQueryBatch().
+	GetParentDevice()->GetOcclusionQueryHeap()->StartQueryBatch(*this, NumQueriesInBatch);
 }
 
 void FD3D12CommandContext::RHIEndOcclusionQueryBatch()
@@ -175,86 +174,38 @@ void FD3D12CommandContext::RHIEndOcclusionQueryBatch()
 * class FD3D12QueryHeap
 *=============================================================================*/
 
-FD3D12QueryHeap::FD3D12QueryHeap(FD3D12Device* InParent, const D3D12_QUERY_HEAP_TYPE &InQueryHeapType, uint32 InQueryHeapCount, uint32 InMaxActiveBatches)
+FD3D12QueryHeap::FD3D12QueryHeap(FD3D12Device* InParent, const D3D12_QUERY_TYPE InQueryType, uint32 InQueryHeapCount, uint32 InMaxActiveBatches)
 	: FD3D12DeviceChild(InParent)
 	, FD3D12SingleNodeGPUObject(InParent->GetGPUMask())
-	, MaxActiveBatches(InMaxActiveBatches)
 	, LastBatch(InMaxActiveBatches - 1)
 	, ActiveAllocatedElementCount(0)
 	, LastAllocatedElement(InQueryHeapCount - 1)
-	, ResultSize(8)
-	, pResultData(nullptr)
+	, QueryType(InQueryType)
+	, QueryHeapCount(InQueryHeapCount)
+	, QueryHeap(nullptr)
+	, ResultBuffer(nullptr)
 {
-	if (InQueryHeapType == D3D12_QUERY_HEAP_TYPE_OCCLUSION)
-	{
-		QueryType = D3D12_QUERY_TYPE_OCCLUSION;
-	}
-	else if (InQueryHeapType == D3D12_QUERY_HEAP_TYPE_TIMESTAMP)
-	{
-		QueryType = D3D12_QUERY_TYPE_TIMESTAMP;
-	}
-	else
-	{
-		check(false);
-	}
-
-	// Setup the query heap desc
-	QueryHeapDesc = {};
-	QueryHeapDesc.Type = InQueryHeapType;
-	QueryHeapDesc.Count = InQueryHeapCount;
-	QueryHeapDesc.NodeMask = (uint32)GetGPUMask();
+	check(QueryType == D3D12_QUERY_TYPE_OCCLUSION || QueryType == D3D12_QUERY_TYPE_TIMESTAMP);
 
 	CurrentQueryBatch.Clear();
 
-	ActiveQueryBatches.Reserve(MaxActiveBatches);
-	ActiveQueryBatches.AddZeroed(MaxActiveBatches);
+	ActiveQueryBatches.Reserve(InMaxActiveBatches);
+	ActiveQueryBatches.AddZeroed(InMaxActiveBatches);
 
 	// Don't Init() until the RHI has created the device
 }
 
-FD3D12QueryHeap::~FD3D12QueryHeap()
-{
-	// Unmap the result buffer
-	if (pResultData)
-	{
-		ResultBuffer->GetResource()->Unmap(0, nullptr);
-		pResultData = nullptr;
-	}
-}
-
 void FD3D12QueryHeap::Init()
 {
 	check(GetParentDevice());
 	check(GetParentDevice()->GetDevice());
 
-	// Create the query heap
 	CreateQueryHeap();
-
-	// Create the result buffer
-	CreateResultBuffer();
-
-	// Start out with an open query batch
-	StartQueryBatch();
 }
 
 void FD3D12QueryHeap::Destroy()
 {
-	if (pResultData)
-	{
-		ResultBuffer->GetResource()->Unmap(0, nullptr);
-		pResultData = nullptr;
-	}
-
-#if ENABLE_RESIDENCY_MANAGEMENT
-	if (D3DX12Residency::IsInitialized(QueryHeapResidencyHandle))
-	{
-		D3DX12Residency::EndTrackingObject(GetParentDevice()->GetResidencyManager(), QueryHeapResidencyHandle);
-		QueryHeapResidencyHandle = {};
-	}
-#endif
-
-	QueryHeap = nullptr;
-	ResultBuffer = nullptr;
+	DestroyQueryHeap(false);
 }
 
 uint32 FD3D12QueryHeap::GetNextElement(uint32 InElement)
@@ -263,7 +214,7 @@ uint32 FD3D12QueryHeap::GetNextElement(uint32 InElement)
 	InElement++;
 
 	// See if we need to wrap around to the begining of the heap
-	if (InElement >= GetQueryHeapCount())
+	if (InElement >= QueryHeapCount)
 	{
 		InElement = 0;
 	}
@@ -277,7 +228,7 @@ uint32 FD3D12QueryHeap::GetNextBatchElement(uint32 InBatchElement)
 	InBatchElement++;
 
 	// See if we need to wrap around to the begining of the heap
-	if (InBatchElement >= MaxActiveBatches)
+	if (InBatchElement >= (uint32) ActiveQueryBatches.Num())
 	{
 		InBatchElement = 0;
 	}
@@ -288,42 +239,62 @@ uint32 FD3D12QueryHeap::GetNextBatchElement(uint32 InBatchElement)
 uint32 FD3D12QueryHeap::AllocQuery(FD3D12CommandContext& CmdContext)
 {
 	check(CmdContext.IsDefaultContext());
-	check(CurrentQueryBatch.bOpen);
 
 	// Get the element for this allocation
 	const uint32 CurrentElement = GetNextElement(LastAllocatedElement);
 
-	if (CurrentQueryBatch.StartElement > CurrentElement)
+	if (QueryType == D3D12_QUERY_TYPE_OCCLUSION)
 	{
-		// We're in the middle of a batch, but we're at the end of the heap
-		// We need to split the batch in two and resolve the first piece
-		EndQueryBatchAndResolveQueryData(CmdContext);
-		check(CurrentQueryBatch.bOpen && CurrentQueryBatch.ElementCount == 0);
+		check(CurrentQueryBatch.bOpen);
+	}
+	else
+	{
+		if (!CurrentQueryBatch.bOpen)
+		{
+			StartQueryBatch(CmdContext, 256);
+			check(CurrentQueryBatch.bOpen && CurrentQueryBatch.ElementCount == 0);
+		}
+
+		if (CurrentQueryBatch.StartElement > CurrentElement)
+		{
+			// We're in the middle of a batch, but we're at the end of the heap
+			// We need to split the batch in two and resolve the first piece
+			EndQueryBatchAndResolveQueryData(CmdContext);
+		}
 	}
 
 	// Increment the count for the current batch
 	CurrentQueryBatch.ElementCount++;
 
 	LastAllocatedElement = CurrentElement;
-	check(CurrentElement < GetQueryHeapCount());
+	check(CurrentElement < QueryHeapCount);
 	return CurrentElement;
 }
 
-void FD3D12QueryHeap::StartQueryBatch()
+void FD3D12QueryHeap::StartQueryBatch(FD3D12CommandContext& CmdContext, uint32 NumQueriesInBatch)
 {
-	//#todo-rco: Use NumQueriesInBatch!
-	static bool bWarned = false;
-	if (!bWarned)
-	{ 
-		bWarned = true;
-		UE_LOG(LogD3D12RHI, Warning, TEXT("NumQueriesInBatch is not used in FD3D12QueryHeap::StartQueryBatch(), this helpful warning exists to remind you about that. Remove it when this is fixed."));
-	}
+	check(!CurrentQueryBatch.bOpen);
+
 
 	if (!CurrentQueryBatch.bOpen)
 	{
 		// Clear the current batch
 		CurrentQueryBatch.Clear();
 
+		if (ActiveAllocatedElementCount + NumQueriesInBatch > QueryHeapCount)
+		{
+			DestroyQueryHeap(true);
+
+			QueryHeapCount = Align(NumQueriesInBatch + QueryHeapCount, 65536 / ResultSize);
+
+			CreateQueryHeap();
+
+			UE_LOG(LogD3D12RHI, Display, TEXT("QueryHeapCount is now %d elements"), QueryHeapCount);
+
+			ActiveAllocatedElementCount = 0;
+			LastAllocatedElement = QueryHeapCount - 1;
+		}
+
 		// Start a new batch
 		CurrentQueryBatch.StartElement = GetNextElement(LastAllocatedElement);
 		CurrentQueryBatch.bOpen = true;
@@ -333,20 +304,26 @@ void FD3D12QueryHeap::StartQueryBatch()
 void FD3D12QueryHeap::EndQueryBatchAndResolveQueryData(FD3D12CommandContext& CmdContext)
 {
 	check(CmdContext.IsDefaultContext());
-	check(CurrentQueryBatch.bOpen);
 
-	// Discard empty batches
-	if (CurrentQueryBatch.ElementCount == 0)
+	if (!CurrentQueryBatch.bOpen)
 	{
 		return;
 	}
 
+	check(CurrentQueryBatch.bOpen);
+
 	// Close the current batch
 	CurrentQueryBatch.bOpen = false;
 
+	// Discard empty batches
+	if (CurrentQueryBatch.ElementCount == 0)
+	{
+		return;
+	}
+
 	// Increment the active element count
 	ActiveAllocatedElementCount += CurrentQueryBatch.ElementCount;
-	checkf(ActiveAllocatedElementCount <= GetQueryHeapCount(), TEXT("The query heap is too small. Either increase the heap count (larger resource) or decrease MAX_ACTIVE_BATCHES."));
+	checkf(ActiveAllocatedElementCount <= QueryHeapCount, TEXT("The query heap is too small. Either increase the heap count (larger resource) or decrease MAX_ACTIVE_BATCHES."));
 
 	// Track the current active batches (application is using the data)
 	LastBatch = GetNextBatchElement(LastBatch);
@@ -357,9 +334,23 @@ void FD3D12QueryHeap::EndQueryBatchAndResolveQueryData(FD3D12CommandContext& Cmd
 	ActiveAllocatedElementCount -= OldestBatch.ElementCount;
 
 	CmdContext.otherWorkCounter++;
-	CmdContext.CommandListHandle->ResolveQueryData(
-		QueryHeap, QueryType, CurrentQueryBatch.StartElement, CurrentQueryBatch.ElementCount,
-		ResultBuffer->GetResource(), GetResultBufferOffsetForElement(CurrentQueryBatch.StartElement));
+	if (CurrentQueryBatch.StartElement + CurrentQueryBatch.ElementCount <= QueryHeapCount)
+	{
+		// Single range
+		CmdContext.CommandListHandle->ResolveQueryData(
+			QueryHeap, QueryType, CurrentQueryBatch.StartElement, CurrentQueryBatch.ElementCount,
+			ResultBuffer->GetResource(), GetResultBufferOffsetForElement(CurrentQueryBatch.StartElement));
+	}
+	else
+	{
+		// Wrapping around heap border, need two resolves for end of heap and beginning of new range
+		CmdContext.CommandListHandle->ResolveQueryData(
+			QueryHeap, QueryType, CurrentQueryBatch.StartElement, QueryHeapCount - CurrentQueryBatch.StartElement,
+			ResultBuffer->GetResource(), GetResultBufferOffsetForElement(CurrentQueryBatch.StartElement));
+		CmdContext.CommandListHandle->ResolveQueryData(
+			QueryHeap, QueryType, 0, CurrentQueryBatch.ElementCount - (QueryHeapCount - CurrentQueryBatch.StartElement),
+			ResultBuffer->GetResource(), 0);
+	}
 
 	CmdContext.CommandListHandle.UpdateResidency(&QueryHeapResidencyHandle);
 	CmdContext.CommandListHandle.UpdateResidency(ResultBuffer);
@@ -368,56 +359,64 @@ void FD3D12QueryHeap::EndQueryBatchAndResolveQueryData(FD3D12CommandContext& Cmd
 	// so we know what sync point to wait for. The query's data isn't ready to read until the above ResolveQueryData completes on the GPU.
 	for (int32 i = 0; i < CurrentQueryBatch.RenderQueries.Num(); i++)
 	{
-		CurrentQueryBatch.RenderQueries[i]->MarkResolved(CmdContext.CommandListHandle);
+		CurrentQueryBatch.RenderQueries[i]->MarkResolved(CmdContext.CommandListHandle, ResultBuffer);
 	}
-
-	// Start a new batch
-	StartQueryBatch();
 }
 
-uint32 FD3D12QueryHeap::BeginQuery(FD3D12CommandContext& CmdContext)
+void FD3D12QueryHeap::BeginQuery(FD3D12CommandContext& CmdContext, FD3D12RenderQuery* RenderQuery)
 {
 	check(CmdContext.IsDefaultContext());
 	check(CurrentQueryBatch.bOpen);
-	const uint32 Element = AllocQuery(CmdContext);
+
+	RenderQuery->Reset();
+	RenderQuery->HeapIndex = AllocQuery(CmdContext);
+
 	CmdContext.otherWorkCounter++;
-	CmdContext.CommandListHandle->BeginQuery(QueryHeap, QueryType, Element);
+	CmdContext.CommandListHandle->BeginQuery(QueryHeap, QueryType, RenderQuery->HeapIndex);
 
 	CmdContext.CommandListHandle.UpdateResidency(&QueryHeapResidencyHandle);
-
-	return Element;
 }
 
-void FD3D12QueryHeap::EndQuery(FD3D12CommandContext& CmdContext, uint32 InElement, FD3D12RenderQuery* InRenderQuery)
+void FD3D12QueryHeap::EndQuery(FD3D12CommandContext& CmdContext, FD3D12RenderQuery* RenderQuery)
 {
 	check(CmdContext.IsDefaultContext());
-	check(CurrentQueryBatch.bOpen);
+
+	if (QueryType == D3D12_QUERY_TYPE_OCCLUSION)
+	{
+		check(CurrentQueryBatch.bOpen);
+	}
+	else
+	{
+		RenderQuery->Reset();
+		RenderQuery->HeapIndex = AllocQuery(CmdContext);
+	}
+
 	CmdContext.otherWorkCounter++;
-	CmdContext.CommandListHandle->EndQuery(QueryHeap, QueryType, InElement);
+	CmdContext.CommandListHandle->EndQuery(QueryHeap, QueryType, RenderQuery->HeapIndex);
 
 	CmdContext.CommandListHandle.UpdateResidency(&QueryHeapResidencyHandle);
 
 	// Track which render queries are used in this batch.
-	if (InRenderQuery)
-	{
-		CurrentQueryBatch.RenderQueries.Push(InRenderQuery);
-	}
+	CurrentQueryBatch.RenderQueries.Push(RenderQuery);
 }
 
 void FD3D12QueryHeap::CreateQueryHeap()
 {
+	// Setup the query heap desc
+	D3D12_QUERY_HEAP_DESC QueryHeapDesc;
+	QueryHeapDesc.Type = (QueryType == D3D12_QUERY_TYPE_OCCLUSION)? D3D12_QUERY_HEAP_TYPE_OCCLUSION : D3D12_QUERY_HEAP_TYPE_TIMESTAMP;
+	QueryHeapDesc.Count = QueryHeapCount;
+	QueryHeapDesc.NodeMask = (uint32)GetGPUMask();
+
 	// Create the upload heap
-	VERIFYD3D12RESULT(GetParentDevice()->GetDevice()->CreateQueryHeap(&QueryHeapDesc, IID_PPV_ARGS(QueryHeap.GetInitReference())));
+	VERIFYD3D12RESULT(GetParentDevice()->GetDevice()->CreateQueryHeap(&QueryHeapDesc, IID_PPV_ARGS(&QueryHeap)));
 	SetName(QueryHeap, L"Query Heap");
 
 #if ENABLE_RESIDENCY_MANAGEMENT
-	D3DX12Residency::Initialize(QueryHeapResidencyHandle, QueryHeap.GetReference(), ResultSize * QueryHeapDesc.Count);
+	D3DX12Residency::Initialize(QueryHeapResidencyHandle, QueryHeap, ResultSize * QueryHeapDesc.Count);
 	D3DX12Residency::BeginTrackingObject(GetParentDevice()->GetResidencyManager(), QueryHeapResidencyHandle);
 #endif
-}
 
-void FD3D12QueryHeap::CreateResultBuffer()
-{
 	FD3D12Adapter* Adapter = GetParentDevice()->GetParentAdapter();
 
 	const D3D12_HEAP_PROPERTIES ResultBufferHeapProperties = CD3DX12_HEAP_PROPERTIES(D3D12_HEAP_TYPE_READBACK, (uint32)GetGPUMask(), (uint32)GetVisibilityMask());
@@ -429,11 +428,43 @@ void FD3D12QueryHeap::CreateResultBuffer()
 		ResultBufferHeapProperties,
 		D3D12_RESOURCE_STATE_COPY_DEST,
 		nullptr,
-		ResultBuffer.GetInitReference(),
+		&ResultBuffer,
 		TEXT("Query Heap Result Buffer")));
+}
 
-	// Map the result buffer (and keep it mapped)
-	VERIFYD3D12RESULT(ResultBuffer->GetResource()->Map(0, nullptr, &pResultData));
+void FD3D12QueryHeap::DestroyQueryHeap(bool bDeferDelete)
+{
+#if ENABLE_RESIDENCY_MANAGEMENT
+	if (D3DX12Residency::IsInitialized(QueryHeapResidencyHandle))
+	{
+		D3DX12Residency::EndTrackingObject(GetParentDevice()->GetResidencyManager(), QueryHeapResidencyHandle);
+		QueryHeapResidencyHandle = {};
+	}
+#endif
+	if (QueryHeap)
+	{
+		if (bDeferDelete)
+		{
+			GetParentDevice()->GetParentAdapter()->GetDeferredDeletionQueue().EnqueueResource(QueryHeap);
+		}
+		else
+		{
+			QueryHeap->Release();
+		}
+		QueryHeap = nullptr;
+	}
+	if (ResultBuffer)
+	{
+		if (bDeferDelete)
+		{
+			ResultBuffer->DeferDelete();
+		}
+		else
+		{
+			ResultBuffer->Release();
+		}
+		ResultBuffer = nullptr;
+	}
 }
 
 /*=============================================================================
diff --git a/Engine/Source/Runtime/D3D12RHI/Private/D3D12Query.h b/Engine/Source/Runtime/D3D12RHI/Private/D3D12Query.h
--- a/Engine/Source/Runtime/D3D12RHI/Private/D3D12Query.h
+++ b/Engine/Source/Runtime/D3D12RHI/Private/D3D12Query.h
@@ -9,12 +9,19 @@ D3D12Query.h: Implementation of D3D12 Query
 class FD3D12RenderQuery : public FRHIRenderQuery, public FD3D12DeviceChild, public FD3D12LinkedAdapterObject<FD3D12RenderQuery>
 {
 public:
+	FD3D12Resource* ResultBuffer;
+
+	/** The cached query result. */
+	uint64 Result;
+
+	// A timestamp so that LDA query results only handle object from the most recent frames.
+	uint32 Timestamp;
 
 	/** The query's index in its heap. */
 	uint32 HeapIndex;
 
-	/** The cached query result. */
-	uint64 Result;
+	/** The query's type. */
+	const ERenderQueryType Type;
 
 	/** True if the query's result is cached. */
 	bool bResultIsCached : 1;
@@ -22,18 +29,13 @@ public:
 	/** True if the query has been resolved. */
 	bool bResolved : 1;
 
-	/** The query's type. */
-	const ERenderQueryType Type;
-
-	// A timestamp so that LDA query results only handle object from the most recent frames.
-	uint64 Timestamp;
-
 	/** Initialization constructor. */
 	FD3D12RenderQuery(FD3D12Device* Parent, ERenderQueryType InQueryType) :
 		FD3D12DeviceChild(Parent),
+		ResultBuffer(nullptr),
 		Result(0),
-		Type(InQueryType),
-		Timestamp(0)
+		Timestamp(0),
+		Type(InQueryType)
 	{
 		Reset();
 	}
@@ -46,9 +48,10 @@ public:
 	}
 
 	// Indicate the command list that was used to resolve the query.
-	inline void MarkResolved(FD3D12CommandListHandle& CommandList)
+	inline void MarkResolved(FD3D12CommandListHandle& CommandList, FD3D12Resource* InResultBuffer)
 	{
 		CLSyncPoint = CommandList;
+		ResultBuffer = InResultBuffer;
 		bResolved = true;
 	}
 
@@ -101,35 +104,27 @@ private:
 	};
 
 public:
-	FD3D12QueryHeap(class FD3D12Device* InParent, const D3D12_QUERY_HEAP_TYPE& InQueryHeapType, uint32 InQueryHeapCount, uint32 InMaxActiveBatches);
-	~FD3D12QueryHeap();
+	FD3D12QueryHeap(class FD3D12Device* InParent, const D3D12_QUERY_TYPE InQueryType, uint32 InQueryHeapCount, uint32 InMaxActiveBatches);
 
 	void Init();
-
 	void Destroy();
 
+	// Start tracking a new batch of begin/end query calls that will be resolved together
+	void StartQueryBatch(FD3D12CommandContext& CmdContext, uint32 NumQueriesInBatch);
 	// Stop tracking the current batch of begin/end query calls that will be resolved together.
-	// This implicitly starts a new batch.
 	void EndQueryBatchAndResolveQueryData(FD3D12CommandContext& CmdContext);  
 
 	uint32 AllocQuery(FD3D12CommandContext& CmdContext); // Some query types don't need a BeginQuery call. Instead just alloc a slot to EndQuery with.
-	uint32 BeginQuery(FD3D12CommandContext& CmdContext); // Obtain a query from the store of available queries
-	void EndQuery(FD3D12CommandContext& CmdContext, uint32 InElement, FD3D12RenderQuery* InRenderQuery);
-
-	uint32 GetQueryHeapCount() const { return QueryHeapDesc.Count; }
-	uint32 GetResultSize() const { return ResultSize; }
-
-	inline FD3D12Resource* GetResultBuffer() { return ResultBuffer.GetReference(); }
+	void BeginQuery(FD3D12CommandContext& CmdContext, FD3D12RenderQuery* RenderQuery); // Obtain a query from the store of available queries
+	void EndQuery(FD3D12CommandContext& CmdContext, FD3D12RenderQuery* RenderQuery);
 
 private:
-	void StartQueryBatch(); // Start tracking a new batch of begin/end query calls that will be resolved together
-
 	uint32 GetNextElement(uint32 InElement); // Get the next element, after the specified element. Handles overflow.
 
 	uint32 GetNextBatchElement(uint32 InBatchElement);
 
 	void CreateQueryHeap();
-	void CreateResultBuffer();
+	void DestroyQueryHeap(bool bDeferDelete);
 
 	uint64 GetResultBufferOffsetForElement(uint32 InElement) const { return ResultSize * InElement; };
 
@@ -137,20 +132,18 @@ private:
 	QueryBatch CurrentQueryBatch;                       // The current recording batch.
 
 	TArray<QueryBatch> ActiveQueryBatches;              // List of active query batches. The data for these is in use.
-
-	const uint32 MaxActiveBatches;                      // The max number of query batches that will be held.
 	uint32 LastBatch;                                   // The index of the newest batch.
 
 	uint32 ActiveAllocatedElementCount;         // The number of elements that are in use (Active). Between the head and the tail.
 
 	uint32 LastAllocatedElement;                // The last element that was allocated for BeginQuery
-	const uint32 ResultSize;                    // The byte size of a result for a single query
-	D3D12_QUERY_HEAP_DESC QueryHeapDesc;        // The description of the current query heap
-	D3D12_QUERY_TYPE QueryType;
-	TRefCountPtr<ID3D12QueryHeap> QueryHeap;    // The query heap where all elements reside
+	const D3D12_QUERY_TYPE QueryType;
+	uint32 QueryHeapCount;
+	ID3D12QueryHeap* QueryHeap;                 // The query heap where all elements reside
 	FD3D12ResidencyHandle QueryHeapResidencyHandle;
-	TRefCountPtr<FD3D12Resource> ResultBuffer;  // The buffer where all query results are stored
-	void* pResultData;
+	FD3D12Resource* ResultBuffer;               // The buffer where all query results are stored
+
+	static const uint32 ResultSize = 8;         // The byte size of a result for a single query
 };
 
 /**
