Duplicate cl 5142655 from Dev-Rendering. Fix for aftermath crash on exit

diff --git a/Engine/Source/Runtime/D3D12RHI/Private/D3D12Query.cpp b/Engine/Source/Runtime/D3D12RHI/Private/D3D12Query.cpp
--- a/Engine/Source/Runtime/D3D12RHI/Private/D3D12Query.cpp
+++ b/Engine/Source/Runtime/D3D12RHI/Private/D3D12Query.cpp
@@ -777,6 +777,8 @@ void FD3D12BufferedGPUTiming::ReleaseDynamicRHI()
 	delete(TimestampQueryHeap);
 	TimestampQueryHeap = nullptr;
 	TimestampQueryHeapBuffer = nullptr;
+
+	TimestampListHandles.Reset();
 }
 
 /**
