Fix stat gpu not working on D3D12

diff --git a/Engine/Source/Runtime/D3D12RHI/Private/D3D12Query.cpp b/Engine/Source/Runtime/D3D12RHI/Private/D3D12Query.cpp
--- a/Engine/Source/Runtime/D3D12RHI/Private/D3D12Query.cpp
+++ b/Engine/Source/Runtime/D3D12RHI/Private/D3D12Query.cpp
@@ -261,6 +261,13 @@ uint32 FD3D12QueryHeap::AllocQuery(FD3D12CommandContext& CmdContext)
 			// We need to split the batch in two and resolve the first piece
 			EndQueryBatchAndResolveQueryData(CmdContext);
 		}
+
+		// check for the the batch being closed due to wrap and open a new one
+		if (!CurrentQueryBatch.bOpen)
+		{
+			StartQueryBatch(CmdContext, 256);
+			check(CurrentQueryBatch.bOpen && CurrentQueryBatch.ElementCount == 0);
+		}
 	}
 
 	// Increment the count for the current batch
