Cherry-pick 5934928 and cpp changes from 5646022: Implemented separate compilation and run-time caching of ray tracing shaders and linking into final pipeline state objects.

Significantly improves RTPSO creation performance when  loading new shaders into an existing large scene.

diff --git a/Engine/Source/Runtime/D3D12RHI/Private/D3D12Device.h b/Engine/Source/Runtime/D3D12RHI/Private/D3D12Device.h
--- a/Engine/Source/Runtime/D3D12RHI/Private/D3D12Device.h
+++ b/Engine/Source/Runtime/D3D12RHI/Private/D3D12Device.h
@@ -11,6 +11,7 @@ D3D12Device.h: D3D12 Device Interfaces
 class FD3D12DynamicRHI;
 class FD3D12BasicRayTracingPipeline;
 class FD3D12RayTracingDescriptorHeapCache;
+class FD3D12RayTracingPipelineCache;
 
 class FD3D12Device : public FD3D12SingleNodeGPUObject, public FNoncopyable, public FD3D12AdapterChild
 {
@@ -47,6 +48,7 @@ public:
 	ID3D12Device5*							GetRayTracingDevice();
 	const FD3D12BasicRayTracingPipeline*	GetBasicRayTracingPipeline() const { return BasicRayTracingPipeline; }
 	FD3D12RayTracingDescriptorHeapCache*	GetRayTracingDescriptorHeapCache() { return RayTracingDescriptorHeapCache; }
+	FD3D12RayTracingPipelineCache*			GetRayTracingPipelineCache() { return RayTracingPipelineCache; }
 #endif // D3D12_RHI_RAYTRACING
 
 	FD3D12DynamicRHI* GetOwningRHI();
@@ -187,7 +189,7 @@ protected:
 
 #if D3D12_RHI_RAYTRACING
 	FD3D12BasicRayTracingPipeline* BasicRayTracingPipeline = nullptr;
-
+	FD3D12RayTracingPipelineCache* RayTracingPipelineCache = nullptr;
 	// #dxr_todo: unify RT descriptor cache with main FD3D12DescriptorCache
 	FD3D12RayTracingDescriptorHeapCache* RayTracingDescriptorHeapCache = nullptr;
 	void DestroyRayTracingDescriptorCache();
diff --git a/Engine/Source/Runtime/D3D12RHI/Private/D3D12RayTracing.cpp b/Engine/Source/Runtime/D3D12RHI/Private/D3D12RayTracing.cpp
--- a/Engine/Source/Runtime/D3D12RHI/Private/D3D12RayTracing.cpp
+++ b/Engine/Source/Runtime/D3D12RHI/Private/D3D12RayTracing.cpp
@@ -185,21 +185,15 @@ static TRefCountPtr<ID3D12StateObject> CreateRayTracingStateObject(
 	const TArrayView<LPCWSTR>& Exports,
 	uint32 MaxPayloadSizeInBytes,
 	const TArrayView<const D3D12_HIT_GROUP_DESC>& HitGroups,
-	const FD3D12RootSignature& GlobalRootSignature,
+	const ID3D12RootSignature* GlobalRootSignature,
 	const TArrayView<ID3D12RootSignature*>& LocalRootSignatures,
-	const TArrayView<uint32>& LocalRootSignatureAssociations // indices into LocalRootSignatures, one per export
+	const TArrayView<uint32>& LocalRootSignatureAssociations, // indices into LocalRootSignatures, one per export
+	const TArrayView<D3D12_EXISTING_COLLECTION_DESC>& ExistingCollections,
+	D3D12_STATE_OBJECT_TYPE StateObjectType // Full RTPSO or a Collection
 )
 {
 	checkf(LocalRootSignatureAssociations.Num() == Exports.Num(), TEXT("There must be exactly one local root signature association per export."));
 
-#if !NO_LOGGING
-	const uint32 NumShadersWarningThreshold = 1000;
-	if (Exports.Num() > NumShadersWarningThreshold)
-	{
-		UE_LOG(LogD3D12RHI, Warning, TEXT("Creating ray tracing pipeline with %d shaders. Reduce the number of unique shaders in the scene to avoid CPU stalls."), Exports.Num());
-	}
-#endif
-
 	TRefCountPtr<ID3D12StateObject> Result;
 
 	// There are several pipeline sub-objects that are always required:
@@ -209,15 +203,16 @@ static TRefCountPtr<ID3D12StateObject> CreateRayTracingStateObject(
 	// 4) Global root signature
 	static constexpr uint32 NumRequiredSubobjects = 4;
 
-	TArray<D3D12_STATE_SUBOBJECT, TInlineAllocator<16>> Subobjects;
+	TArray<D3D12_STATE_SUBOBJECT> Subobjects;
 	Subobjects.SetNumUninitialized(NumRequiredSubobjects
 		+ ShaderLibraries.Num()
 		+ HitGroups.Num()
 		+ LocalRootSignatures.Num()
 		+ LocalRootSignatureAssociations.Num()
+		+ ExistingCollections.Num()
 	);
 
-	TArray<D3D12_SUBOBJECT_TO_EXPORTS_ASSOCIATION, TInlineAllocator<16>> ExportAssociations;
+	TArray<D3D12_SUBOBJECT_TO_EXPORTS_ASSOCIATION> ExportAssociations;
 	ExportAssociations.SetNumUninitialized(LocalRootSignatureAssociations.Num());
 
 	uint32 Index = 0;
@@ -247,7 +242,7 @@ static TRefCountPtr<ID3D12StateObject> CreateRayTracingStateObject(
 	ShaderConfigAssociation.pSubobjectToAssociate = &Subobjects[ShaderConfigIndex];
 	Subobjects[Index++] = D3D12_STATE_SUBOBJECT{ D3D12_STATE_SUBOBJECT_TYPE_SUBOBJECT_TO_EXPORTS_ASSOCIATION, &ShaderConfigAssociation };
 
-	// Intersection hit group
+	// Hit groups
 
 	for (const D3D12_HIT_GROUP_DESC& HitGroupDesc : HitGroups)
 	{
@@ -263,8 +258,7 @@ static TRefCountPtr<ID3D12StateObject> CreateRayTracingStateObject(
 
 	// Global root signature
 
-	ID3D12RootSignature* GlobalRootSignaturePtr = GlobalRootSignature.GetRootSignature();
-	Subobjects[Index++] = D3D12_STATE_SUBOBJECT{ D3D12_STATE_SUBOBJECT_TYPE_GLOBAL_ROOT_SIGNATURE, &GlobalRootSignaturePtr };
+	Subobjects[Index++] = D3D12_STATE_SUBOBJECT{ D3D12_STATE_SUBOBJECT_TYPE_GLOBAL_ROOT_SIGNATURE, &GlobalRootSignature };
 
 	// Local root signatures
 
@@ -292,6 +286,15 @@ static TRefCountPtr<ID3D12StateObject> CreateRayTracingStateObject(
 		Subobjects[Index++] = D3D12_STATE_SUBOBJECT{ D3D12_STATE_SUBOBJECT_TYPE_SUBOBJECT_TO_EXPORTS_ASSOCIATION, &ExportAssociations[ExportIndex] };
 	}
 
+	// Existing collection objects
+
+	for (int32 CollectionIndex = 0; CollectionIndex < ExistingCollections.Num(); ++CollectionIndex)
+	{
+		Subobjects[Index++] = D3D12_STATE_SUBOBJECT{ D3D12_STATE_SUBOBJECT_TYPE_EXISTING_COLLECTION, &ExistingCollections[CollectionIndex] };
+	}
+
+	// Done!
+
 	checkf(Index == Subobjects.Num(), TEXT("All pipeline subobjects must be initialized."));
 
 	// Create ray tracing pipeline state object
@@ -299,7 +302,7 @@ static TRefCountPtr<ID3D12StateObject> CreateRayTracingStateObject(
 	D3D12_STATE_OBJECT_DESC Desc = {};
 	Desc.NumSubobjects = Index;
 	Desc.pSubobjects = &Subobjects[0];
-	Desc.Type = D3D12_STATE_OBJECT_TYPE_RAYTRACING_PIPELINE;
+	Desc.Type = StateObjectType;
 
 	VERIFYD3D12RESULT(RayTracingDevice->CreateStateObject(&Desc, IID_PPV_ARGS(Result.GetInitReference())));
 
@@ -309,7 +312,266 @@ static TRefCountPtr<ID3D12StateObject> CreateRayTracingStateObject(
 	return Result;
 }
 
-// #dxr_todo: FD3D12Device::GlobalViewHeap/GlobalSamplerHeap should be used instead of ad-hoc heaps here.
+inline uint64 GetShaderHash64(FRayTracingShaderRHIParamRef ShaderRHI)
+{
+	uint64 ShaderHash; // 64 bits from the shader SHA1
+	FMemory::Memcpy(&ShaderHash, ShaderRHI->GetHash().Hash, sizeof(ShaderHash));
+	return ShaderHash;
+}
+
+// Generates a stable symbol name for a ray tracing shader, used for RT PSO creation.
+
+inline FString GenerateShaderName(const TCHAR* Prefix, uint64 Hash)
+{
+	return FString::Printf(TEXT("%s_%016llx"), Prefix, Hash);
+}
+
+inline FString GenerateShaderName(FRayTracingShaderRHIParamRef ShaderRHI)
+{
+	const FD3D12RayTracingShader* Shader = FD3D12DynamicRHI::ResourceCast(ShaderRHI);
+	uint64 ShaderHash = GetShaderHash64(ShaderRHI);
+	return GenerateShaderName(*(Shader->EntryPoint), ShaderHash);
+}
+
+// Cache for ray tracing pipeline collection objects, containing single shaders that can be linked into full pipelines.
+class FD3D12RayTracingPipelineCache
+{
+public:
+
+	FD3D12RayTracingPipelineCache(FD3D12Device* Device)
+		: DefaultLocalRootSignature(Device->GetParentAdapter())
+	{
+		// Default empty local root signature
+
+		D3D12_VERSIONED_ROOT_SIGNATURE_DESC LocalRootSignatureDesc = {};
+		LocalRootSignatureDesc.Version = D3D_ROOT_SIGNATURE_VERSION_1_0;
+		LocalRootSignatureDesc.Desc_1_0.Flags |= D3D12_ROOT_SIGNATURE_FLAG_LOCAL_ROOT_SIGNATURE;
+		DefaultLocalRootSignature.Init(LocalRootSignatureDesc);
+	}
+
+	struct FKey
+	{
+		uint64 ShaderHash = 0;
+		uint32 MaxPayloadSizeInBytes = 0;
+		const ID3D12RootSignature* GlobalRootSignature = nullptr;
+		const ID3D12RootSignature* LocalRootSignature = nullptr;
+
+		bool operator == (const FKey& Other) const
+		{
+			return ShaderHash == Other.ShaderHash
+				&& MaxPayloadSizeInBytes == Other.MaxPayloadSizeInBytes
+				&& GlobalRootSignature == Other.GlobalRootSignature
+				&& LocalRootSignature == Other.LocalRootSignature;
+		}
+
+		inline friend uint32 GetTypeHash(const FKey& Key)
+		{
+			return Key.ShaderHash;
+		}
+	};
+
+	struct FEntry
+	{
+		// Move-only type
+		FEntry() = default;
+		FEntry(FEntry&& Other) = default;
+
+		FEntry(const FEntry&) = delete;
+		FEntry& operator = (const FEntry&) = delete;
+		FEntry& operator = (FEntry&& Other) = delete;
+
+		D3D12_EXISTING_COLLECTION_DESC GetCollectionDesc()
+		{
+			D3D12_EXISTING_COLLECTION_DESC Result = {};
+			Result.pExistingCollection = StateObject;
+			Result.pExports = ExportDescs.GetData();
+			Result.NumExports = ExportDescs.Num();
+			return Result;
+		}
+
+		TRefCountPtr<FD3D12RayTracingShader> Shader;
+
+		TRefCountPtr<ID3D12StateObject> StateObject;
+
+		TArray<FString> ExportNames;
+		TArray<D3D12_EXPORT_DESC> ExportDescs;
+	};
+
+	enum class ECollectionType
+	{
+		RayGen,
+		Miss,
+		HitGroup,
+	};
+
+	static const TCHAR* GetCollectionTypeName(ECollectionType Type)
+	{
+		switch (Type)
+		{
+		case ECollectionType::RayGen:
+			return TEXT("RayGen");
+			break;
+		case ECollectionType::Miss:
+			return TEXT("Miss");
+			break;
+		case ECollectionType::HitGroup:
+			return TEXT("HitGroup");
+		default:
+			return TEXT("");
+		}
+	}
+
+	D3D12_EXISTING_COLLECTION_DESC GetOrCompileShader(
+		ID3D12Device5* RayTracingDevice,
+		FD3D12RayTracingShader* Shader,
+		ID3D12RootSignature* GlobalRootSignature,
+		uint32 MaxPayloadSizeInBytes,
+		ECollectionType CollectionType,
+		bool* bOutCacheHit = nullptr)
+	{
+		const uint64 ShaderHash = GetShaderHash64(Shader);
+
+		ID3D12RootSignature* LocalRootSignature = nullptr;
+		if (CollectionType == ECollectionType::HitGroup)
+		{
+			// Only hit group shaders have a local root signature
+			LocalRootSignature = Shader->pRootSignature->GetRootSignature();
+		}
+		else
+		{
+			// ... all other shaders share a default empty local root signature
+			LocalRootSignature = DefaultLocalRootSignature.GetRootSignature();
+		}
+
+		FKey CacheKey;
+		CacheKey.ShaderHash = ShaderHash;
+		CacheKey.MaxPayloadSizeInBytes = MaxPayloadSizeInBytes;
+		CacheKey.GlobalRootSignature = GlobalRootSignature;
+		CacheKey.LocalRootSignature = LocalRootSignature;
+
+		FEntry& Entry = Cache.FindOrAdd(CacheKey);
+
+		if (Entry.StateObject)
+		{
+			if (bOutCacheHit) *bOutCacheHit = true;
+		}
+		else
+		{
+			if (bOutCacheHit) *bOutCacheHit = false;
+
+			Entry.Shader = Shader;
+
+			LPCWSTR EntryName = *(Shader->EntryPoint);
+
+			TArray<LPCWSTR, TInlineAllocator<4>> ExportNames;
+			TArray<LPCWSTR, TInlineAllocator<3>> OriginalEntryPoints;
+			TArray<LPCWSTR, TInlineAllocator<3>> RenamedEntryPoints;
+			TArray<uint32, TInlineAllocator<3>> LocalRootSignatureAssociations;
+
+			// Typical collection can have up to 4 exports: HitGroup/RGS/MS, CHS, AHS, IS
+			const uint32 MaxExports = 4;
+			Entry.ExportNames.Reserve(MaxExports);
+			Entry.ExportDescs.Reserve(MaxExports);
+
+			const uint32 NumHitGroups = CollectionType == ECollectionType::HitGroup ? 1 : 0;
+
+			FString& PrimaryExportName = Entry.ExportNames.Add_GetRef(GenerateShaderName(GetCollectionTypeName(CollectionType), ShaderHash));
+
+			D3D12_HIT_GROUP_DESC HitGroupDesc = {};
+
+			ExportNames.Add(*PrimaryExportName);
+
+			if (CollectionType == ECollectionType::HitGroup)
+			{
+				HitGroupDesc.HitGroupExport = *PrimaryExportName;
+
+				{
+					FString& ExportName = Entry.ExportNames.Add_GetRef(GenerateShaderName(TEXT("CHS"), ShaderHash));
+
+					HitGroupDesc.ClosestHitShaderImport = *ExportName;
+
+					OriginalEntryPoints.Add(*(Shader->EntryPoint));
+					RenamedEntryPoints.Add(*ExportName);
+					ExportNames.Add(*ExportName);
+
+					LocalRootSignatureAssociations.Add(0);
+				}
+
+				if (!Shader->AnyHitEntryPoint.IsEmpty())
+				{
+					FString& ExportName = Entry.ExportNames.Add_GetRef(GenerateShaderName(TEXT("AHS"), ShaderHash));
+
+					HitGroupDesc.AnyHitShaderImport = *ExportName;
+
+					OriginalEntryPoints.Add(*(Shader->AnyHitEntryPoint));
+					RenamedEntryPoints.Add(*ExportName);
+					ExportNames.Add(*ExportName);
+
+					LocalRootSignatureAssociations.Add(0);
+				}
+				if (!Shader->IntersectionEntryPoint.IsEmpty())
+				{
+					FString& ExportName = Entry.ExportNames.Add_GetRef(GenerateShaderName(TEXT("IS"), ShaderHash));
+
+					HitGroupDesc.IntersectionShaderImport = *ExportName;
+
+					OriginalEntryPoints.Add(*(Shader->IntersectionEntryPoint));
+					RenamedEntryPoints.Add(*ExportName);
+					ExportNames.Add(*ExportName);
+
+					LocalRootSignatureAssociations.Add(0);
+				}
+			}
+			else
+			{
+				checkf(CollectionType == ECollectionType::Miss || CollectionType == ECollectionType::RayGen, TEXT("Unexpected RT sahder collection type"));
+
+				OriginalEntryPoints.Add(*(Shader->EntryPoint));
+				RenamedEntryPoints.Add(*PrimaryExportName);
+
+				LocalRootSignatureAssociations.Add(0);
+			}
+
+			FDXILLibrary Library;
+			Library.InitFromDXIL(Shader->ShaderBytecode, OriginalEntryPoints.GetData(), RenamedEntryPoints.GetData(), OriginalEntryPoints.Num());
+
+			const FDXILLibrary* LibraryPtr = &Library;
+
+			Entry.StateObject = CreateRayTracingStateObject(
+				RayTracingDevice,
+				MakeArrayView(&LibraryPtr, 1),
+				RenamedEntryPoints,
+				MaxPayloadSizeInBytes,
+				MakeArrayView(&HitGroupDesc, NumHitGroups),
+				GlobalRootSignature,
+				MakeArrayView(&LocalRootSignature, 1),
+				LocalRootSignatureAssociations,
+				{},
+				D3D12_STATE_OBJECT_TYPE_COLLECTION);
+
+			for (LPCWSTR ExportName : ExportNames)
+			{
+				D3D12_EXPORT_DESC ExportDesc = {};
+				ExportDesc.Name = ExportName;
+				ExportDesc.Flags = D3D12_EXPORT_FLAG_NONE;
+				Entry.ExportDescs.Add(ExportDesc);
+			}
+
+			// Validate that memory reservation was correct
+			check(Entry.ExportNames.Num() <= MaxExports);
+			check(Entry.ExportDescs.Num() <= MaxExports);
+		}
+
+		return Entry.GetCollectionDesc();
+	}
+
+private:
+
+	TMap<FKey, FEntry> Cache;
+	FD3D12RootSignature DefaultLocalRootSignature; // Default empty root signature used for default hit shaders.
+};
+
+// #dxr_todo UE-72158: FD3D12Device::GlobalViewHeap/GlobalSamplerHeap should be used instead of ad-hoc heaps here.
 // Unfortunately, this requires a major refactor of how global heaps work.
 // FD3D12CommandContext-s should not get static chunks of the global heap, but instead should dynamically allocate
 // chunks on as-needed basis and release them when possible.
@@ -941,29 +1203,10 @@ struct FD3D12RayTracingShaderLibrary
 	void Reserve(uint32 NumShaders)
 	{
 		Shaders.Reserve(NumShaders);
-		DXILLibraries.Reserve(NumShaders);
 		Identifiers.Reserve(NumShaders);
 	}
 
-	// Single-entry shaders, such as ray-gen, miss and callable
-	void AddSingleShader(FRayTracingShaderRHIParamRef ShaderRHI, LPCWSTR OptExportName = nullptr)
-	{
-		FD3D12RayTracingShader* Shader = FD3D12DynamicRHI::ResourceCast(ShaderRHI);
-		Shaders.Add(Shader);
-
-		LPCWSTR EntryName = *(Shader->EntryPoint);
-		FDXILLibrary& Library = DXILLibraries.AddDefaulted_GetRef();
-		LPCWSTR ExportName = OptExportName ? OptExportName : EntryName;
-		Library.InitFromDXIL(Shader->ShaderBytecode, &EntryName, &ExportName, 1);
-	}
-
-	int32 Num() const
-	{
-		return Identifiers.Num();
-	}
-
 	TArray<TRefCountPtr<FD3D12RayTracingShader>> Shaders;
-	TArray<FDXILLibrary> DXILLibraries;
 	TArray<FD3D12ShaderIdentifier> Identifiers;
 };
 
@@ -975,54 +1218,20 @@ public:
 
 	FD3D12RayTracingPipelineState(FD3D12Device* Device, const FRayTracingPipelineStateInitializer& Initializer)
 		: DefaultShaderTable(Device)
-		, DefaultLocalRootSignature(Device->GetParentAdapter())
 	{
-		FD3D12Adapter* Adapter = Device->GetParentAdapter();
-		ID3D12Device5* RayTracingDevice = Device->GetRayTracingDevice();
-
 		checkf(Initializer.GetRayGenTable().Num() > 0, TEXT("Ray tracing pipelines must have at leat one ray generation shader."));
 
-		// If no custom hit groups were provided, then disable SBT indexing and force default shader on all primitives
-		bAllowHitGroupIndexing = Initializer.GetHitGroupTable().Num() ? Initializer.bAllowHitGroupIndexing : false;
-
-		RayGenShaders.Reserve(Initializer.GetRayGenTable().Num());
-
-		const auto& RayGenShaderTable = Initializer.GetRayGenTable();
-		for (int32 RayGenShaderIndex = 0; RayGenShaderIndex < RayGenShaderTable.Num(); ++RayGenShaderIndex)
-		{
-			FD3D12RayTracingShader* Shader = FD3D12DynamicRHI::ResourceCast(RayGenShaderTable[RayGenShaderIndex]);
-			FString RenamedEntryPoint = FString::Printf(TEXT("RayGen_%04d_%s"), RayGenShaderIndex, *(Shader->EntryPoint));
-			RayGenShaders.AddSingleShader(Shader, *RenamedEntryPoint);
-		}
-
-		// Add miss shaders (either custom ones provided by user or default one otherwise)
-		if (Initializer.GetMissTable().Num())
-		{
-			MissShaders.Reserve(Initializer.GetMissTable().Num());
-			for (FRayTracingShaderRHIParamRef Shader : Initializer.GetMissTable())
-			{
-				MissShaders.AddSingleShader(Shader);
-			}
-		}
-		else
-		{
-			MissShaders.AddSingleShader(GetBuildInRayTracingShader<FDefaultMainMS>());
-		}
-
-		// All raygen and miss shaders must share the same root signature, so take the first one and validate the rest
-		GlobalRootSignature = RayGenShaders.Shaders[0]->pRootSignature;
+		uint64 TotalCreationTime = 0;
+		uint64 CompileTime = 0;
+		uint64 LinkTime = 0;
+		uint32 NumCacheHits = 0;
 
-		for (auto Shader : RayGenShaders.Shaders)
-		{
-			checkf(Shader->pRootSignature == GlobalRootSignature, TEXT("All raygen and miss shaders must share the same root signature"));
-		}
+		TotalCreationTime -= FPlatformTime::Cycles64();
 
-		for (auto Shader : MissShaders.Shaders)
-		{
-			checkf(Shader->pRootSignature == GlobalRootSignature, TEXT("All raygen and miss shaders must share the same root signature"));
-		}
+		FD3D12Adapter* Adapter = Device->GetParentAdapter();
+		ID3D12Device5* RayTracingDevice = Device->GetRayTracingDevice();
 
-		// Use hit shaders from initializer or fall back to default if none were provided
+		// Use hit and miss shaders from initializer or fall back to default ones if none were provided
 
 		FRayTracingShaderRHIParamRef DefaultHitShader = GetBuildInRayTracingShader<FDefaultMainCHS>();
 		FRayTracingShaderRHIParamRef DefaultHitGroupTable[] = { DefaultHitShader };
@@ -1031,207 +1240,140 @@ public:
 			? Initializer.GetHitGroupTable()
 			: DefaultHitGroupTable;
 
-		TArray<ID3D12RootSignature*> LocalRootSignatures;
-		LocalRootSignatures.Reserve(1 + InitializerHitGroups.Num()); // default empty signature (for raygen and miss shaders) + one per hit group
+		FRayTracingShaderRHIParamRef DefaultMissShader = GetBuildInRayTracingShader<FDefaultMainMS>();
+		FRayTracingShaderRHIParamRef DefaultMissTable[] = { DefaultMissShader };
 
-		MaxLocalRootSignatureSize = 0;
+		TArrayView<const FRayTracingShaderRHIParamRef> InitializerMissShaders = Initializer.GetMissTable().Num()
+			? Initializer.GetMissTable()
+			: DefaultMissTable;
 
-		const uint32 HitShaderRootSignatureBaseIndex = LocalRootSignatures.Num();
+		TArrayView<const FRayTracingShaderRHIParamRef> InitializerRayGenShaders = Initializer.GetRayGenTable();
 
-		// Initialize hit group shader libraries
+		const uint32 MaxTotalShaders = InitializerRayGenShaders.Num() + InitializerMissShaders.Num() + InitializerHitGroups.Num();
 		// #dxr_todo: hit group libraries *also* could come from precompiled pipeline sub-objects, when those are supported in the future
-		HitGroupShaders.Reserve(InitializerHitGroups.Num());
+		// All raygen and miss shaders must share the same global root signature, so take the first one and validate the rest
 
-		// Each shader within RTPSO must have a unique name, therefore we must rename original shader entry points.
-		TArray<FString> RenamedHitGroupEntryPoints;
-		RenamedHitGroupEntryPoints.Reserve(InitializerHitGroups.Num() * 3); // Up to 3 entry points may exist per hit group
+		GlobalRootSignature = FD3D12DynamicRHI::ResourceCast(InitializerRayGenShaders[0])->pRootSignature->GetRootSignature();
 
-		struct FHitGroupEntryIndices
-		{
-			int32 ClosestHit   = INDEX_NONE;
-			int32 AnyHit       = INDEX_NONE;
-			int32 Intersection = INDEX_NONE;
-		};
+		// Helper function to acquire a D3D12_EXISTING_COLLECTION_DESC for a compiled shader via cache
 
-		TArray<FHitGroupEntryIndices> HitGroupEntryIndices;
-		HitGroupEntryIndices.Reserve(InitializerHitGroups.Num());
+		TSet<uint64> UniqueShaderHashes;
+		UniqueShaderHashes.Reserve(MaxTotalShaders);
 
-		MaxHitGroupViewDescriptors = 0;
+		TArray<D3D12_EXISTING_COLLECTION_DESC> UniqueShaderCollections;
+		UniqueShaderCollections.Reserve(MaxTotalShaders);
 
-		for (const FRayTracingShaderRHIParamRef ShaderRHI : InitializerHitGroups)
+		auto AddShaderCollection = [Device, RayTracingDevice, GlobalRootSignature = this->GlobalRootSignature,
+										&UniqueShaderHashes, &UniqueShaderCollections, &Initializer, &NumCacheHits, &CompileTime]
+			(FD3D12RayTracingShader* Shader, FD3D12RayTracingPipelineCache::ECollectionType CollectionType)
 		{
-			FD3D12RayTracingShader* Shader = FD3D12DynamicRHI::ResourceCast(ShaderRHI);
-
-			checkf(Shader, TEXT("A valid ray tracing hit group shader must be provided for all elements in the FRayTracingPipelineStateInitializer hit group table."));
-			checkf(!Shader->ResourceCounts.bGlobalUniformBufferUsed, TEXT("Global uniform buffers are not implemented for ray tracing shaders"));
-
-			const uint32 ShaderViewDescriptors = Shader->ResourceCounts.NumSRVs + Shader->ResourceCounts.NumUAVs;
-			MaxHitGroupViewDescriptors = FMath::Max(MaxHitGroupViewDescriptors, ShaderViewDescriptors);
-
-			HitGroupShaders.Shaders.Add(Shader);
-
-			MaxLocalRootSignatureSize = FMath::Max(MaxLocalRootSignatureSize, Shader->pRootSignature->GetTotalRootSignatureSizeInBytes());
-			LocalRootSignatures.Add(Shader->pRootSignature->GetRootSignature());
-
-			FDXILLibrary& Library = HitGroupShaders.DXILLibraries.AddDefaulted_GetRef();
-
-			LPCWSTR OriginalGroupEntryPoints[3];
-			LPCWSTR RenamedGroupEntryPoints[3];
-			uint32 NumGroupEntryPoints = 0;
+			bool bIsAlreadyInSet = false;
+			const uint64 ShaderHash = GetShaderHash64(Shader);
+			UniqueShaderHashes.Add(ShaderHash, &bIsAlreadyInSet);
 
-			// #dxr_todo: A unique name for all ray tracing shaders could be auto-generated in the shader pipeline instead of run-time
-			auto RegisterHitGroupEntry = 
-				[&OriginalGroupEntryPoints, &RenamedGroupEntryPoints, &NumGroupEntryPoints, &RenamedHitGroupEntryPoints]
-			(LPCWSTR EntryNameChars)
-			{
-				const int32 EntryIndex = RenamedHitGroupEntryPoints.Num();
-				FString& RenamedEntryString = RenamedHitGroupEntryPoints.AddDefaulted_GetRef();
-				RenamedEntryString = FString::Printf(TEXT("HitGroup_%04d_%s"), EntryIndex, EntryNameChars);
+			FD3D12RayTracingPipelineCache* PipelineCache = Device->GetRayTracingPipelineCache();
 
-				OriginalGroupEntryPoints[NumGroupEntryPoints] = EntryNameChars;
-				RenamedGroupEntryPoints[NumGroupEntryPoints] = *RenamedEntryString;
-				++NumGroupEntryPoints;
+			bool bCacheHit = false;
 
-				return EntryIndex;
-			};
+			CompileTime -= FPlatformTime::Cycles64();
 
-			FHitGroupEntryIndices EntryIndices;
+			D3D12_EXISTING_COLLECTION_DESC CollectionDesc = PipelineCache->GetOrCompileShader(
+				RayTracingDevice, Shader, GlobalRootSignature,
+				Initializer.MaxPayloadSizeInBytes,
+				CollectionType,
+				&bCacheHit);
 
-			EntryIndices.ClosestHit = RegisterHitGroupEntry(*(Shader->EntryPoint));
+			CompileTime += FPlatformTime::Cycles64();
 
-			if (!Shader->AnyHitEntryPoint.IsEmpty())
+			if (!bIsAlreadyInSet)
 			{
-				EntryIndices.AnyHit = RegisterHitGroupEntry(*(Shader->AnyHitEntryPoint));
-			}
+				UniqueShaderCollections.Add(CollectionDesc);
 
-			if (!Shader->IntersectionEntryPoint.IsEmpty())
-			{
-				EntryIndices.Intersection = RegisterHitGroupEntry(*(Shader->IntersectionEntryPoint));
+				if (bCacheHit) NumCacheHits++;
 			}
 
-			HitGroupEntryIndices.Add(EntryIndices);
-
-			Library.InitFromDXIL(Shader->ShaderBytecode, OriginalGroupEntryPoints, RenamedGroupEntryPoints, NumGroupEntryPoints);
-		}
-
-		// Default empty local root signature (#dxr_todo: move this into a single shared place)
-
-		const uint32 EmptyLocalRootSignatureIndex = LocalRootSignatures.Num();
-
-		{
-			D3D12_VERSIONED_ROOT_SIGNATURE_DESC LocalRootSignatureDesc = {};
-			LocalRootSignatureDesc.Version = D3D_ROOT_SIGNATURE_VERSION_1_0;
-			LocalRootSignatureDesc.Desc_1_0.Flags |= D3D12_ROOT_SIGNATURE_FLAG_LOCAL_ROOT_SIGNATURE;
-			DefaultLocalRootSignature.Init(LocalRootSignatureDesc);
-			LocalRootSignatures.Add(DefaultLocalRootSignature.GetRootSignature());
-		}
-
-		// Initialize StateObject
+			return CollectionDesc;
+		};
 
-		TArray<const FDXILLibrary*> Libraries;
-		TArray<LPCWSTR> Exports;
-		TArray<uint32> LocalRootSignatureAssociations;
+		// If no custom hit groups were provided, then disable SBT indexing and force default shader on all primitives
 
-		// All shader types other than hit groups
-		const int32 NumSystemShaders = RayGenShaders.DXILLibraries.Num() + MissShaders.DXILLibraries.Num();
+		bAllowHitGroupIndexing = Initializer.GetHitGroupTable().Num() ? Initializer.bAllowHitGroupIndexing : false;
 
-		// Reserve space for all custom hit groups + system shaders
-		const int32 MaxNumLibraries = HitGroupShaders.DXILLibraries.Num() + NumSystemShaders;
-		Libraries.Reserve(MaxNumLibraries);
+		// Add ray generation shaders
 
-		const int32 MaxNumExports = RenamedHitGroupEntryPoints.Num() + NumSystemShaders;
-		Exports.Reserve(MaxNumExports);
-		LocalRootSignatureAssociations.Reserve(MaxNumExports); // One RS association per export
+		TArray<LPCWSTR> RayGenShaderNames;
 
-		// Ray generation shaders
+		RayGenShaders.Reserve(InitializerRayGenShaders.Num());
+		RayGenShaderNames.Reserve(InitializerRayGenShaders.Num());
 
-		for (const FDXILLibrary& Library : RayGenShaders.DXILLibraries)
+		for (FRayTracingShaderRHIParamRef ShaderRHI : InitializerRayGenShaders)
 		{
-			Libraries.Add(&Library);
-
-			check(Library.ExportNames.Num() == 1);
-			Exports.Add(*(Library.ExportNames[0]));
+			FD3D12RayTracingShader* Shader = FD3D12DynamicRHI::ResourceCast(ShaderRHI);
+			checkf(Shader->pRootSignature->GetRootSignature() == GlobalRootSignature, TEXT("All raygen and miss shaders must share the same root signature"));
 
-			// RayGen shaders don't have parameters that come from SBT, so associate with empty local RS
-			LocalRootSignatureAssociations.Add(EmptyLocalRootSignatureIndex);
+			D3D12_EXISTING_COLLECTION_DESC CollectionDesc = AddShaderCollection(Shader, FD3D12RayTracingPipelineCache::ECollectionType::RayGen);
+			RayGenShaderNames.Add(CollectionDesc.pExports[0].Name);
+			RayGenShaders.Shaders.Add(Shader);
 		}
 
-		// Miss shaders
-
-		for (const FDXILLibrary& Library : MissShaders.DXILLibraries)
-		{
-			Libraries.Add(&Library);
-
-			check(Library.ExportNames.Num() == 1);
-			Exports.Add(*(Library.ExportNames[0]));
-
-			// Miss shaders don't have parameters that come from SBT, so associate with empty local RS
-			LocalRootSignatureAssociations.Add(EmptyLocalRootSignatureIndex);
-		}
+		// Add miss shaders
 
-		// Add hit group shaders
+		TArray<LPCWSTR> MissShaderNames;
+		MissShaders.Reserve(InitializerMissShaders.Num());
+		MissShaderNames.Reserve(InitializerMissShaders.Num());
 
-		uint32 CurrentHitGroupRootSignatureIndex = HitShaderRootSignatureBaseIndex;
-		for (const FDXILLibrary& Library : HitGroupShaders.DXILLibraries)
+		for (FRayTracingShaderRHIParamRef ShaderRHI : InitializerMissShaders)
 		{
-			Libraries.Add(&Library);
-			for (const FString& ExportName : Library.ExportNames)
-			{
-				Exports.Add(*ExportName);
+			FD3D12RayTracingShader* Shader = FD3D12DynamicRHI::ResourceCast(ShaderRHI);
+			checkf(Shader->pRootSignature->GetRootSignature() == GlobalRootSignature, TEXT("All raygen and miss shaders must share the same root signature"));
 
-				// NOTE: Same local root signature is associated with all shaders in a hit group: closest hit, any hit and intersection (if they are present)
-				LocalRootSignatureAssociations.Add(CurrentHitGroupRootSignatureIndex);
-			}
-			++CurrentHitGroupRootSignatureIndex;
+			D3D12_EXISTING_COLLECTION_DESC CollectionDesc = AddShaderCollection(Shader, FD3D12RayTracingPipelineCache::ECollectionType::Miss);
+			MissShaderNames.Add(CollectionDesc.pExports[0].Name);
+			MissShaders.Shaders.Add(Shader);
 		}
 
-		// Hit groups
+		// Add hit groups
 
-		TArray<D3D12_HIT_GROUP_DESC> HitGroups;
-		TArray<FString> HitGroupNames;
+		MaxHitGroupViewDescriptors = 0;
+		MaxLocalRootSignatureSize = 0;
 
-		HitGroups.Reserve(InitializerHitGroups.Num());
+		TArray<LPCWSTR> HitGroupNames;
+		HitGroupShaders.Reserve(InitializerHitGroups.Num());
 		HitGroupNames.Reserve(InitializerHitGroups.Num());
 
 		for (const FRayTracingShaderRHIParamRef ShaderRHI : InitializerHitGroups)
 		{
-			const int32 HitGroupIndex = HitGroupNames.Num(); // #dxr_todo: this would need to be a unique index if we support pipeline sub-object linking
-			FString& HitGroupName = HitGroupNames.AddDefaulted_GetRef();
-			HitGroupName = FString::Printf(TEXT("HitGroup_%d"), HitGroupIndex);
+			FD3D12RayTracingShader* Shader = FD3D12DynamicRHI::ResourceCast(ShaderRHI);
 
-			D3D12_HIT_GROUP_DESC HitGroup = {};
-			HitGroup.HitGroupExport = *HitGroupName;
+			checkf(Shader, TEXT("A valid ray tracing hit group shader must be provided for all elements in the FRayTracingPipelineStateInitializer hit group table."));
+			checkf(!Shader->ResourceCounts.bGlobalUniformBufferUsed, TEXT("Global uniform buffers are not implemented for ray tracing shaders"));
 
-			const FHitGroupEntryIndices& EntryIndices = HitGroupEntryIndices[HitGroupIndex];
+			const uint32 ShaderViewDescriptors = Shader->ResourceCounts.NumSRVs + Shader->ResourceCounts.NumUAVs;
+			MaxHitGroupViewDescriptors = FMath::Max(MaxHitGroupViewDescriptors, ShaderViewDescriptors);
+			MaxLocalRootSignatureSize = FMath::Max(MaxLocalRootSignatureSize, Shader->pRootSignature->GetTotalRootSignatureSizeInBytes());
 
-			HitGroup.ClosestHitShaderImport = *(RenamedHitGroupEntryPoints[EntryIndices.ClosestHit]);
-			if (EntryIndices.AnyHit != INDEX_NONE)
-			{
-				HitGroup.AnyHitShaderImport = *(RenamedHitGroupEntryPoints[EntryIndices.AnyHit]);
-			}
-			if (EntryIndices.Intersection != INDEX_NONE)
-			{
-				HitGroup.IntersectionShaderImport = *(RenamedHitGroupEntryPoints[EntryIndices.Intersection]);
-			}
+			D3D12_EXISTING_COLLECTION_DESC CollectionDesc = AddShaderCollection(Shader, FD3D12RayTracingPipelineCache::ECollectionType::HitGroup);
 
-			HitGroups.Add(HitGroup);
+			check(CollectionDesc.NumExports != 0);
+			HitGroupNames.Add(CollectionDesc.pExports[0].Name);
+			HitGroupShaders.Shaders.Add(Shader);
 		}
 
-		// Create the pipeline
-
-		check(Libraries.Num() == MaxNumLibraries); // Confirm that our memory reservation assumptions hold up
-		check(Exports.Num() == MaxNumExports); // Confirm that our memory reservation assumptions hold up
-		check(Exports.Num() == LocalRootSignatureAssociations.Num()); // Confirm that we have associated local root signatures with all shaders
+		// Link final RTPSO from shader collections
 
+		LinkTime -= FPlatformTime::Cycles64();
 		StateObject = CreateRayTracingStateObject(
 			RayTracingDevice,
-			Libraries,
-			Exports,
+			{}, // Libraries,
+			{}, // LibraryExports,
 			Initializer.MaxPayloadSizeInBytes,
-			HitGroups,
-			*GlobalRootSignature,
-			LocalRootSignatures,
-			LocalRootSignatureAssociations);
+			{}, // HitGroups
+			GlobalRootSignature,
+			{}, // LocalRootSignatures
+			{}, // LocalRootSignatureAssociations,
+			UniqueShaderCollections,
+			D3D12_STATE_OBJECT_TYPE_RAYTRACING_PIPELINE);
+		LinkTime += FPlatformTime::Cycles64();
 
 		HRESULT QueryInterfaceResult = StateObject->QueryInterface(IID_PPV_ARGS(PipelineProperties.GetInitReference()));
 		checkf(SUCCEEDED(QueryInterfaceResult), TEXT("Failed to query pipeline properties from the ray tracing pipeline state object. Result=%08x"), QueryInterfaceResult);
@@ -1258,21 +1400,21 @@ public:
 		HitGroupShaders.Identifiers.SetNumUninitialized(InitializerHitGroups.Num());
 		for (int32 HitGroupIndex = 0; HitGroupIndex < HitGroupNames.Num(); ++HitGroupIndex)
 		{
-			LPCWSTR ExportNameChars = *HitGroupNames[HitGroupIndex];
+			LPCWSTR ExportNameChars = HitGroupNames[HitGroupIndex];
 			HitGroupShaders.Identifiers[HitGroupIndex] = GetShaderIdentifier(ExportNameChars);
 		}
 
-		RayGenShaders.Identifiers.SetNumUninitialized(RayGenShaders.DXILLibraries.Num());
-		for (int32 ShaderIndex = 0; ShaderIndex < RayGenShaders.DXILLibraries.Num(); ++ShaderIndex)
+		RayGenShaders.Identifiers.SetNumUninitialized(RayGenShaderNames.Num());
+		for (int32 ShaderIndex = 0; ShaderIndex < RayGenShaderNames.Num(); ++ShaderIndex)
 		{
-			LPCWSTR ExportNameChars = *(RayGenShaders.DXILLibraries[ShaderIndex].ExportNames[0]);
+			LPCWSTR ExportNameChars = RayGenShaderNames[ShaderIndex];
 			RayGenShaders.Identifiers[ShaderIndex] = GetShaderIdentifier(ExportNameChars);
 		}
 
-		MissShaders.Identifiers.SetNumUninitialized(MissShaders.DXILLibraries.Num());
-		for (int32 ShaderIndex = 0; ShaderIndex < MissShaders.DXILLibraries.Num(); ++ShaderIndex)
+		MissShaders.Identifiers.SetNumUninitialized(MissShaderNames.Num());
+		for (int32 ShaderIndex = 0; ShaderIndex < MissShaderNames.Num(); ++ShaderIndex)
 		{
-			LPCWSTR ExportNameChars = *(MissShaders.DXILLibraries[ShaderIndex].ExportNames[0]);
+			LPCWSTR ExportNameChars = MissShaderNames[ShaderIndex];
 			MissShaders.Identifiers[ShaderIndex] = GetShaderIdentifier(ExportNameChars);
 		}
 
@@ -1280,10 +1422,28 @@ public:
 		// Hit record indexing and local resources access is disabled when using using this SBT.
 
 		const uint32 DefaultLocalRootDataSize = 0; // Shaders in default SBT are not allowed to access any local resources
-		DefaultShaderTable.Init(RayGenShaders.Num(), MissShaders.Num(), 0, DefaultLocalRootDataSize, 0);
+		DefaultShaderTable.Init(RayGenShaders.Identifiers.Num(), MissShaders.Identifiers.Num(), 0, DefaultLocalRootDataSize, 0);
 		DefaultShaderTable.SetRayGenIdentifiers(RayGenShaders.Identifiers);
 		DefaultShaderTable.SetMissIdentifiers(MissShaders.Identifiers);
 		DefaultShaderTable.SetDefaultHitGroupIdentifier(HitGroupShaders.Identifiers[0]);
+
+		TotalCreationTime += FPlatformTime::Cycles64();
+
+		// Report stats for pipelines that take a long time to create
+
+#if !NO_LOGGING
+		const double TotalCreatimTimeMS = 1000.0 * FPlatformTime::ToSeconds64(TotalCreationTime);
+		const float CreationTimeWarningThresholdMS = 10.0f;
+		if (TotalCreatimTimeMS > CreationTimeWarningThresholdMS)
+		{
+			const double CompileTimeMS = 1000.0 * FPlatformTime::ToSeconds64(CompileTime);
+			const double LinkTimeMS = 1000.0 * FPlatformTime::ToSeconds64(LinkTime);
+			uint32 NumUniqueShaders = UniqueShaderCollections.Num();
+			UE_LOG(LogD3D12RHI, Log,
+				TEXT("Creating RTPSO with %d shaders (%d cached, %d new) took %.2f ms. Compile time %.2f ms, link time %.2f ms."),
+				NumUniqueShaders, NumCacheHits, NumUniqueShaders - NumCacheHits, TotalCreatimTimeMS, CompileTimeMS, LinkTimeMS);
+		}
+#endif //!NO_LOGGING
 	}
 
 	FD3D12RayTracingShaderLibrary RayGenShaders;
@@ -1294,10 +1454,7 @@ public:
 	// This is useful for the case where user only provides default RayGen, Miss and HitGroup shaders.
 	FD3D12RayTracingShaderTable DefaultShaderTable;
 
-	// Default empty root signature used for default hit shaders.
-	FD3D12RootSignature DefaultLocalRootSignature; // #dxr_todo: move this into a single shared place
-
-	const FD3D12RootSignature* GlobalRootSignature = nullptr;
+	ID3D12RootSignature* GlobalRootSignature = nullptr;
 
 	TRefCountPtr<ID3D12StateObject> StateObject;
 	TRefCountPtr<ID3D12StateObjectProperties> PipelineProperties;
@@ -1363,6 +1520,9 @@ public:
 
 void FD3D12Device::InitRayTracing()
 {
+	check(RayTracingPipelineCache == nullptr);
+	RayTracingPipelineCache = new FD3D12RayTracingPipelineCache(this);
+
 	check(RayTracingDescriptorHeapCache == nullptr);
 	RayTracingDescriptorHeapCache = new FD3D12RayTracingDescriptorHeapCache(this);
 
@@ -1375,6 +1535,9 @@ void FD3D12Device::CleanupRayTracing()
 	delete BasicRayTracingPipeline;
 	BasicRayTracingPipeline = nullptr;
 
+	delete RayTracingPipelineCache;
+	RayTracingPipelineCache = nullptr;
+
 	// Note: RayTracingDescriptorHeapCache is destroyed in ~FD3D12Device, after all deferred deletion is processed
 }
 
@@ -1956,8 +2119,8 @@ FD3D12RayTracingShaderTable* FD3D12RayTracingScene::FindOrCreateShaderTable(cons
 	checkf(Pipeline->MaxLocalRootSignatureSize >= sizeof(FHitGroupSystemParameters), TEXT("All local root signatures are expected to contain ray tracing system root parameters (2x root buffers + 4x root DWORD)"));
 
 	CreatedShaderTable->Init(
-		Pipeline->RayGenShaders.Num(),
-		Pipeline->MissShaders.Num(),
+		Pipeline->RayGenShaders.Identifiers.Num(),
+		Pipeline->MissShaders.Identifiers.Num(),
 		NumHitGroupSlots,
 		Pipeline->MaxLocalRootSignatureSize,
 		Pipeline->MaxHitGroupViewDescriptors);
@@ -2460,7 +2623,7 @@ static void DispatchRays(FD3D12CommandContext& CommandContext,
 	// Invalidate state cache to ensure all root parameters for regular shaders are reset when non-RT work is dispatched later.
 	CommandContext.StateCache.TransitionComputeState(D3D12PT_RayTracing);
 
-	CommandContext.CommandListHandle.GraphicsCommandList()->SetComputeRootSignature(Pipeline->GlobalRootSignature->GetRootSignature());
+	CommandContext.CommandListHandle.GraphicsCommandList()->SetComputeRootSignature(Pipeline->GlobalRootSignature);
 
 	FD3D12RayTracingShader* RayGenShader = Pipeline->RayGenShaders.Shaders[RayGenShaderIndex];
 
diff --git a/Engine/Source/Runtime/RenderCore/Private/Shader.cpp b/Engine/Source/Runtime/RenderCore/Private/Shader.cpp
--- a/Engine/Source/Runtime/RenderCore/Private/Shader.cpp
+++ b/Engine/Source/Runtime/RenderCore/Private/Shader.cpp
@@ -937,6 +937,8 @@ void FShaderResource::InitRHI()
 			{
 				RayTracingMaterialLibraryIndex = AddToRayTracingLibrary(RayTracingShader);
 			}
+
+			RayTracingShader->SetHash(OutputHash);
 		}
 	}
 #endif // RHI_RAYTRACING
