Adding crash protection against geometry and vertex buffers which do not exist at acceleration structure update time.

diff --git a/Engine/Source/Runtime/D3D12RHI/Private/D3D12RayTracing.cpp b/Engine/Source/Runtime/D3D12RHI/Private/D3D12RayTracing.cpp
--- a/Engine/Source/Runtime/D3D12RHI/Private/D3D12RayTracing.cpp
+++ b/Engine/Source/Runtime/D3D12RHI/Private/D3D12RayTracing.cpp
@@ -2197,9 +2197,15 @@ void FD3D12CommandContext::RHIUpdateAccelerationStructures(const TArrayView<cons
 	// First batch up all barriers
 	for (const FAccelerationStructureUpdateParams P : Params)
 	{
-		FD3D12RayTracingGeometry* Geometry = FD3D12DynamicRHI::ResourceCast(P.Geometry);
-		Geometry->PositionVertexBuffer = FD3D12DynamicRHI::ResourceCast(P.VertexBuffer);
-		Geometry->TransitionBuffers(*this);
+		if (P.Geometry->IsValid())
+		{
+			FD3D12RayTracingGeometry* Geometry = FD3D12DynamicRHI::ResourceCast(P.Geometry);
+			if (P.VertexBuffer->IsValid())
+			{
+				Geometry->PositionVertexBuffer = FD3D12DynamicRHI::ResourceCast(P.VertexBuffer);
+				Geometry->TransitionBuffers(*this);
+			}
+		}
 	}
 	CommandListHandle.FlushResourceBarriers();
 
@@ -2207,10 +2213,13 @@ void FD3D12CommandContext::RHIUpdateAccelerationStructures(const TArrayView<cons
 	for (const FAccelerationStructureUpdateParams P : Params)
 	{
 		FD3D12RayTracingGeometry* Geometry = FD3D12DynamicRHI::ResourceCast(P.Geometry);
-		Geometry->bIsAccelerationStructureDirty = true;
+		if (P.VertexBuffer->IsValid())
+		{
+			Geometry->bIsAccelerationStructureDirty = true;
 
-		const bool bIsUpdate = !GRayTracingDebugForceFullBuild;
-		Geometry->BuildAccelerationStructure(*this, bIsUpdate);
+			const bool bIsUpdate = !GRayTracingDebugForceFullBuild;
+			Geometry->BuildAccelerationStructure(*this, bIsUpdate);
+		}
 	}
 }
 
