Show clearer message when running UE4 with -d3ddebug and the Windows Graphics Tools are not installed.`gpuvalidation`is also accepted together with `d3d12gpuvalidation` to enable the validation layer to be consistent with the Vulkan RHI

diff --git a/Engine/Source/Runtime/D3D12RHI/Private/D3D12Adapter.cpp b/Engine/Source/Runtime/D3D12RHI/Private/D3D12Adapter.cpp
--- a/Engine/Source/Runtime/D3D12RHI/Private/D3D12Adapter.cpp
+++ b/Engine/Source/Runtime/D3D12RHI/Private/D3D12Adapter.cpp
@@ -96,25 +96,12 @@ void FD3D12Adapter::CreateRootDevice(bool bWithDebug)
 	if (bWithDebug)
 	{
 		TRefCountPtr<ID3D12Debug> DebugController;
-		VERIFYD3D12RESULT(D3D12GetDebugInterface(IID_PPV_ARGS(DebugController.GetInitReference())));
-		DebugController->EnableDebugLayer();
-
-		// TODO: MSFT: BEGIN TEMPORARY WORKAROUND for a debug layer issue with the Editor creating lots of viewports (swapchains).
-		// Without this you could see this error: D3D12 ERROR: ID3D12CommandQueue::ExecuteCommandLists: Up to 8 swapchains can be written to by a single command queue. Present must be called on one of the swapchains to enable a command queue to execute command lists that write to more.  [ STATE_SETTING ERROR #906: COMMAND_QUEUE_TOO_MANY_SWAPCHAIN_REFERENCES]
-		if (GIsEditor)
+		if (SUCCEEDED(D3D12GetDebugInterface(IID_PPV_ARGS(DebugController.GetInitReference()))))
 		{
-			TRefCountPtr<ID3D12Debug1> DebugController1;
-			const HRESULT HrDebugController1 = D3D12GetDebugInterface(IID_PPV_ARGS(DebugController1.GetInitReference()));
-			if (DebugController1.GetReference())
-			{
-				DebugController1->SetEnableSynchronizedCommandQueueValidation(false);
-				UE_LOG(LogD3D12RHI, Warning, TEXT("Disabling the debug layer's Synchronized Command Queue Validation. This means many debug layer features won't do anything. This code should be removed as soon as possible with an update debug layer."));
-			}
-		}
-		// END TEMPORARY WORKAROUND
+			DebugController->EnableDebugLayer();
 
 		bool bD3d12gpuvalidation = false;
-		if (FParse::Param(FCommandLine::Get(), TEXT("d3d12gpuvalidation")))
+			if (FParse::Param(FCommandLine::Get(), TEXT("d3d12gpuvalidation")) || FParse::Param(FCommandLine::Get(), TEXT("gpuvalidation")))
 		{
 			TRefCountPtr<ID3D12Debug1> DebugController1;
 			VERIFYD3D12RESULT(DebugController->QueryInterface(IID_PPV_ARGS(DebugController1.GetInitReference())));
@@ -122,7 +109,13 @@ void FD3D12Adapter::CreateRootDevice(bool bWithDebug)
 			bD3d12gpuvalidation = true;
 		}
 
-		UE_LOG(LogD3D12RHI, Log, TEXT("InitD3DDevice: -D3DDebug = %s -D3D12GPUValidation = %s"), bWithDebug ? TEXT("on") : TEXT("off"), bD3d12gpuvalidation ? TEXT("on") : TEXT("off") );
+			UE_LOG(LogD3D12RHI, Log, TEXT("InitD3DDevice: -D3DDebug = %s -D3D12GPUValidation = %s"), bWithDebug ? TEXT("on") : TEXT("off"), bD3d12gpuvalidation ? TEXT("on") : TEXT("off"));
+		}
+		else
+		{
+			bWithDebug = false;
+			UE_LOG(LogD3D12RHI, Fatal, TEXT("The debug interface requires the D3D12 SDK Layers. Please install the Graphics Tools for Windows. See: https://docs.microsoft.com/en-us/windows/uwp/gaming/use-the-directx-runtime-and-visual-studio-graphics-diagnostic-features"));
+		}
 	}
 #endif // PLATFORM_WINDOWS
 
