Fixed crash in dx12 mode with laptops that have Nvidia and Intel devices. The fix is a workaround to prevent deadlocks when resources are destroyed in a separate thread.

diff --git a/Engine/Source/Runtime/D3D12RHI/Private/D3D12Resources.cpp b/Engine/Source/Runtime/D3D12RHI/Private/D3D12Resources.cpp
--- a/Engine/Source/Runtime/D3D12RHI/Private/D3D12Resources.cpp
+++ b/Engine/Source/Runtime/D3D12RHI/Private/D3D12Resources.cpp
@@ -8,6 +8,20 @@ D3D12Resources.cpp: D3D RHI utility implementation.
 #include "EngineModule.h"
 #include "HAL/LowLevelMemTracker.h"
 
+int32 GD3D12AsyncDeferredDeletion = ASYNC_DEFERRED_DELETION;
+static FAutoConsoleVariableRef CVarAsyncDeferredDeletion(
+	TEXT("D3D12.AsyncDeferredDeletion"),
+	GD3D12AsyncDeferredDeletion,
+	TEXT("Controls whether D3D12 resources will be released on a separate thread (default = ")
+#if ASYNC_DEFERRED_DELETION
+	TEXT("on")
+#else
+	TEXT("off")
+#endif
+	TEXT(")."),
+	ECVF_ReadOnly
+);
+
 /////////////////////////////////////////////////////////////////////
 //	FD3D12 Deferred Deletion Queue
 /////////////////////////////////////////////////////////////////////
@@ -59,19 +73,6 @@ bool FD3D12DeferredDeletionQueue::ReleaseResources(bool DeleteImmediately)
 {
 	FD3D12Adapter* Adapter = GetParentAdapter();
 
-#if ASYNC_DEFERRED_DELETION
-	if (DeleteImmediately)
-	{
-		FAsyncTask<FD3D12AsyncDeletionWorker>* DeleteTask = nullptr;
-		// Call back all threads
-		while (DeleteTasks.Peek(DeleteTask))
-		{
-			DeleteTasks.Dequeue(DeleteTask);
-			DeleteTask->EnsureCompletion(true);
-			delete(DeleteTask);
-		}
-#endif
-
 		struct FDequeueFenceObject
 		{
 			FDequeueFenceObject(FD3D12Fence& InFence)
@@ -88,25 +89,19 @@ bool FD3D12DeferredDeletionQueue::ReleaseResources(bool DeleteImmediately)
 			FD3D12Fence& FrameFence;
 		};
 
-		FencedObjectType FenceObject;
-		const FDequeueFenceObject DequeueFenceObject(Adapter->GetFrameFence());
-
-		while (DeferredReleaseQueue.Dequeue(FenceObject, DequeueFenceObject))
+	if (GD3D12AsyncDeferredDeletion)
 		{
-			if (FenceObject.Type == EObjectType::RHI)
+		if (DeleteImmediately)
 			{
-				FenceObject.RHIObject->Release();
-			}
-			else
+			FAsyncTask<FD3D12AsyncDeletionWorker>* DeleteTask = nullptr;
+			// Call back all threads
+			while (DeleteTasks.Peek(DeleteTask))
 			{
-				FenceObject.D3DObject->Release();
+				DeleteTasks.Dequeue(DeleteTask);
+				DeleteTask->EnsureCompletion(true);
+				delete(DeleteTask);
 			}
 		}
-
-		return DeferredReleaseQueue.IsEmpty();
-
-#if ASYNC_DEFERRED_DELETION
-	}
 	else
 	{
 		FAsyncTask<FD3D12AsyncDeletionWorker>* DeleteTask = nullptr;
@@ -123,7 +118,24 @@ bool FD3D12DeferredDeletionQueue::ReleaseResources(bool DeleteImmediately)
 
 		return false;
 	}
-#endif
+	}
+
+	FencedObjectType FenceObject;
+	const FDequeueFenceObject DequeueFenceObject(Adapter->GetFrameFence());
+
+	while (DeferredReleaseQueue.Dequeue(FenceObject, DequeueFenceObject))
+	{
+		if (FenceObject.Type == EObjectType::RHI)
+		{
+			FenceObject.RHIObject->Release();
+		}
+		else
+		{
+			FenceObject.D3DObject->Release();
+		}
+	}
+
+	return DeferredReleaseQueue.IsEmpty();
 }
 
 FD3D12DeferredDeletionQueue::FD3D12AsyncDeletionWorker::FD3D12AsyncDeletionWorker(FD3D12Adapter* Adapter, FThreadsafeQueue<FencedObjectType>* DeletionQueue)
diff --git a/Engine/Source/Runtime/D3D12RHI/Private/Windows/WindowsD3D12Device.cpp b/Engine/Source/Runtime/D3D12RHI/Private/Windows/WindowsD3D12Device.cpp
--- a/Engine/Source/Runtime/D3D12RHI/Private/Windows/WindowsD3D12Device.cpp
+++ b/Engine/Source/Runtime/D3D12RHI/Private/Windows/WindowsD3D12Device.cpp
@@ -390,6 +390,14 @@ void FD3D12DynamicRHIModule::FindAdapter()
 		}
 	}
 
+	if (bIsAnyIntel && bIsAnyNVIDIA)
+	{
+		UE_LOG(LogD3D12RHI, Log, TEXT("Forcing D3D12.AsyncDeferredDeletion=0 as a workaround for a deadlock on hybrid NVIDIA+Intel systems."));
+
+		extern RHI_API int32 GD3D12AsyncDeferredDeletion;
+		GD3D12AsyncDeferredDeletion = 0;
+	}
+
 	TSharedPtr<FD3D12Adapter> NewAdapter;
 	if (bFavorNonIntegrated && (bIsAnyAMD || bIsAnyNVIDIA))
 	{
