Linker warning fix
#rnx

diff --git a/Engine/Source/Runtime/D3D12RHI/Private/D3D12Resources.cpp b/Engine/Source/Runtime/D3D12RHI/Private/D3D12Resources.cpp
--- a/Engine/Source/Runtime/D3D12RHI/Private/D3D12Resources.cpp
+++ b/Engine/Source/Runtime/D3D12RHI/Private/D3D12Resources.cpp
@@ -8,7 +8,8 @@ D3D12Resources.cpp: D3D RHI utility implementation.
 #include "EngineModule.h"
 #include "HAL/LowLevelMemTracker.h"
 
-int32 GD3D12AsyncDeferredDeletion = ASYNC_DEFERRED_DELETION;
+D3D12RHI_API int32 GD3D12AsyncDeferredDeletion = ASYNC_DEFERRED_DELETION;
+
 static FAutoConsoleVariableRef CVarAsyncDeferredDeletion(
 	TEXT("D3D12.AsyncDeferredDeletion"),
 	GD3D12AsyncDeferredDeletion,
diff --git a/Engine/Source/Runtime/D3D12RHI/Private/Windows/WindowsD3D12Device.cpp b/Engine/Source/Runtime/D3D12RHI/Private/Windows/WindowsD3D12Device.cpp
--- a/Engine/Source/Runtime/D3D12RHI/Private/Windows/WindowsD3D12Device.cpp
+++ b/Engine/Source/Runtime/D3D12RHI/Private/Windows/WindowsD3D12Device.cpp
@@ -394,7 +394,7 @@ void FD3D12DynamicRHIModule::FindAdapter()
 	{
 		UE_LOG(LogD3D12RHI, Log, TEXT("Forcing D3D12.AsyncDeferredDeletion=0 as a workaround for a deadlock on hybrid NVIDIA+Intel systems."));
 
-		extern RHI_API int32 GD3D12AsyncDeferredDeletion;
+		extern D3D12RHI_API int32 GD3D12AsyncDeferredDeletion;
 		GD3D12AsyncDeferredDeletion = 0;
 	}
 
