UE4 now triggers an error when launched with ray tracing but SkinCache.ComputeShaders disabled.

diff --git a/Engine/Source/Runtime/D3D12RHI/Private/D3D12Adapter.cpp b/Engine/Source/Runtime/D3D12RHI/Private/D3D12Adapter.cpp
--- a/Engine/Source/Runtime/D3D12RHI/Private/D3D12Adapter.cpp
+++ b/Engine/Source/Runtime/D3D12RHI/Private/D3D12Adapter.cpp
@@ -177,6 +177,12 @@ void FD3D12Adapter::CreateRootDevice(bool bWithDebug)
 		if (RootRayTracingDevice)
 		{
 			UE_LOG(LogD3D12RHI, Log, TEXT("D3D12 ray tracing enabled."));
+
+			static auto CVarSkinCache = IConsoleManager::Get().FindConsoleVariable(TEXT("r.SkinCache.CompileShaders"));
+			if (CVarSkinCache->GetInt() <= 0)
+			{
+				UE_LOG(LogD3D12RHI, Fatal, TEXT("D3D12 ray tracing requires skin cache to be enabled. Set r.SkinCache.CompileShaders=1."));
+			}
 		}
 		else
 		{
