Workaround implemented in CL 6411800 for laptops that have Nvidia and Intel devices in DX12 now is enclosed in WITH_EDITOR.

diff --git a/Engine/Source/Runtime/D3D12RHI/Private/Windows/WindowsD3D12Device.cpp b/Engine/Source/Runtime/D3D12RHI/Private/Windows/WindowsD3D12Device.cpp
--- a/Engine/Source/Runtime/D3D12RHI/Private/Windows/WindowsD3D12Device.cpp
+++ b/Engine/Source/Runtime/D3D12RHI/Private/Windows/WindowsD3D12Device.cpp
@@ -390,6 +390,7 @@ void FD3D12DynamicRHIModule::FindAdapter()
 		}
 	}
 
+#if WITH_EDITOR
 	if (bIsAnyIntel && bIsAnyNVIDIA)
 	{
 		UE_LOG(LogD3D12RHI, Log, TEXT("Forcing D3D12.AsyncDeferredDeletion=0 as a workaround for a deadlock on hybrid NVIDIA+Intel systems."));
@@ -397,6 +398,7 @@ void FD3D12DynamicRHIModule::FindAdapter()
 		extern D3D12RHI_API int32 GD3D12AsyncDeferredDeletion;
 		GD3D12AsyncDeferredDeletion = 0;
 	}
+#endif
 
 	TSharedPtr<FD3D12Adapter> NewAdapter;
 	if (bFavorNonIntegrated && (bIsAnyAMD || bIsAnyNVIDIA))
