Extend FD3D12FastAllocator to support multi-frame allocations

diff --git a/Engine/Source/Runtime/D3D12RHI/Private/D3D12Allocation.cpp b/Engine/Source/Runtime/D3D12RHI/Private/D3D12Allocation.cpp
--- a/Engine/Source/Runtime/D3D12RHI/Private/D3D12Allocation.cpp
+++ b/Engine/Source/Runtime/D3D12RHI/Private/D3D12Allocation.cpp
@@ -1287,8 +1287,8 @@ HRESULT FD3D12TextureAllocatorPool::AllocateTexture(D3D12_RESOURCE_DESC Desc, co
 //	Fast Allocation
 //-----------------------------------------------------------------------------
 
-template void* FD3D12FastAllocator::Allocate<FD3D12ScopeLock>(uint32 Size, uint32 Alignment, class FD3D12ResourceLocation* ResourceLocation);
-template void* FD3D12FastAllocator::Allocate<FD3D12ScopeNoLock>(uint32 Size, uint32 Alignment, class FD3D12ResourceLocation* ResourceLocation);
+template void* FD3D12FastAllocator::Allocate<FD3D12ScopeLock>(uint32 Size, uint32 Alignment, class FD3D12ResourceLocation* ResourceLocation, bool bMultiFrame);
+template void* FD3D12FastAllocator::Allocate<FD3D12ScopeNoLock>(uint32 Size, uint32 Alignment, class FD3D12ResourceLocation* ResourceLocation, bool bMultiFrame);
 
 template void FD3D12FastAllocator::CleanupPages<FD3D12ScopeLock>(uint64 FrameLag);
 template void FD3D12FastAllocator::CleanupPages<FD3D12ScopeNoLock>(uint64 FrameLag);
@@ -1311,7 +1311,7 @@ FD3D12FastAllocator::FD3D12FastAllocator(FD3D12Device* Parent, FRHIGPUMask Visib
 {}
 
 template<typename LockType>
-void* FD3D12FastAllocator::Allocate(uint32 Size, uint32 Alignment, class FD3D12ResourceLocation* ResourceLocation)
+void* FD3D12FastAllocator::Allocate(uint32 Size, uint32 Alignment, class FD3D12ResourceLocation* ResourceLocation, bool bMultiFrame)
 {
 	LockType Lock(&CS);
 
@@ -1364,7 +1364,8 @@ void* FD3D12FastAllocator::Allocate(uint32 Size, uint32 Alignment, class FD3D12R
 			Size,
 			CurrentAllocatorPage->FastAllocBuffer->GetGPUVirtualAddress(),
 			CurrentAllocatorPage->FastAllocData,
-			CurrentOffset);
+			CurrentOffset,
+			bMultiFrame);
 
 		CurrentAllocatorPage->NextFastAllocOffset = CurrentOffset + Size;
 
diff --git a/Engine/Source/Runtime/D3D12RHI/Private/D3D12Allocation.h b/Engine/Source/Runtime/D3D12RHI/Private/D3D12Allocation.h
--- a/Engine/Source/Runtime/D3D12RHI/Private/D3D12Allocation.h
+++ b/Engine/Source/Runtime/D3D12RHI/Private/D3D12Allocation.h
@@ -493,7 +493,7 @@ public:
 	FD3D12FastAllocator(FD3D12Device* Parent, FRHIGPUMask VisibiltyMask, const D3D12_HEAP_PROPERTIES& InHeapProperties, uint32 PageSize);
 
 	template<typename LockType>
-	void* Allocate(uint32 Size, uint32 Alignment, class FD3D12ResourceLocation* ResourceLocation);
+	void* Allocate(uint32 Size, uint32 Alignment, class FD3D12ResourceLocation* ResourceLocation, bool bMultiFrame = false);
 
 	template<typename LockType>
 	void Destroy();
diff --git a/Engine/Source/Runtime/D3D12RHI/Private/D3D12Buffer.cpp b/Engine/Source/Runtime/D3D12RHI/Private/D3D12Buffer.cpp
--- a/Engine/Source/Runtime/D3D12RHI/Private/D3D12Buffer.cpp
+++ b/Engine/Source/Runtime/D3D12RHI/Private/D3D12Buffer.cpp
@@ -168,9 +168,11 @@ BufferType* FD3D12Adapter::CreateRHIBuffer(FRHICommandListImmediate* RHICmdList,
 		{
 			check(Size == CreateInfo.ResourceArray->GetResourceDataSize());
 
+			const bool bOnAsyncThread = !IsInRHIThread() && !IsInRenderingThread();
+
 			// Get an upload heap and initialize data
 			FD3D12ResourceLocation SrcResourceLoc(BufferOut->GetParentDevice());
-			void* pData = SrcResourceLoc.GetParentDevice()->GetDefaultFastAllocator().Allocate<FD3D12ScopeLock>(Size, 4UL, &SrcResourceLoc);
+			void* pData = SrcResourceLoc.GetParentDevice()->GetDefaultFastAllocator().Allocate<FD3D12ScopeLock>(Size, 4UL, &SrcResourceLoc, bOnAsyncThread);
 			check(pData);
 			FMemory::Memcpy(pData, CreateInfo.ResourceArray->GetResourceData(), Size);
 
@@ -215,6 +217,7 @@ BufferType* FD3D12Adapter::CreateRHIBuffer(FRHICommandListImmediate* RHICmdList,
 								SrcResourceLoc.GetOffsetFromBaseOfResource(), Size);
 
 							hCommandList.UpdateResidency(Destination);
+							hCommandList.UpdateResidency(SrcResourceLoc.GetResource());
 						}
 
 						CurrentBuffer = CurrentBuffer->GetNextObject();
@@ -222,7 +225,7 @@ BufferType* FD3D12Adapter::CreateRHIBuffer(FRHICommandListImmediate* RHICmdList,
 				}
 			};
 			
-			if (!IsInRHIThread() && !IsInRenderingThread())
+			if (bOnAsyncThread)
 			{
 				// Need to update buffer content on RHI thread (immediate context) because the buffer can be a
 				// sub-allocation and its backing resource may be in a state incompatible with the copy queue.
diff --git a/Engine/Source/Runtime/D3D12RHI/Private/D3D12Resources.cpp b/Engine/Source/Runtime/D3D12RHI/Private/D3D12Resources.cpp
--- a/Engine/Source/Runtime/D3D12RHI/Private/D3D12Resources.cpp
+++ b/Engine/Source/Runtime/D3D12RHI/Private/D3D12Resources.cpp
@@ -586,6 +586,12 @@ void FD3D12ResourceLocation::ReleaseResource()
 		}
 		break;
 	}
+	case ResourceLocationType::eMultiFrameFastAllocation:
+	{
+		check(UnderlyingResource->GetRefCount() > 1);
+		UnderlyingResource->Release();
+		break;
+	}
 	case ResourceLocationType::eFastAllocation:
 	case ResourceLocationType::eUndefined:
 	default:
diff --git a/Engine/Source/Runtime/D3D12RHI/Public/D3D12Resources.h b/Engine/Source/Runtime/D3D12RHI/Public/D3D12Resources.h
--- a/Engine/Source/Runtime/D3D12RHI/Public/D3D12Resources.h
+++ b/Engine/Source/Runtime/D3D12RHI/Public/D3D12Resources.h
@@ -420,6 +420,7 @@ public:
 		eStandAlone,
 		eSubAllocation,
 		eFastAllocation,
+		eMultiFrameFastAllocation,
 		eAliased, // Oculus is the only API that uses this
 		eNodeReference,
 		eHeapAliased, 
@@ -495,9 +496,17 @@ public:
 	}
 
 
-	inline void AsFastAllocation(FD3D12Resource* Resource, uint32 BufferSize, D3D12_GPU_VIRTUAL_ADDRESS GPUBase, void* CPUBase, uint64 Offset)
+	inline void AsFastAllocation(FD3D12Resource* Resource, uint32 BufferSize, D3D12_GPU_VIRTUAL_ADDRESS GPUBase, void* CPUBase, uint64 Offset, bool bMultiFrame = false)
 	{
-		SetType(FD3D12ResourceLocation::ResourceLocationType::eFastAllocation);
+		if (bMultiFrame)
+		{
+			Resource->AddRef();
+			SetType(ResourceLocationType::eMultiFrameFastAllocation);
+		}
+		else
+		{
+			SetType(ResourceLocationType::eFastAllocation);
+		}
 		SetResource(Resource);
 		SetSize(BufferSize);
 		SetOffsetFromBaseOfResource(Offset);
