Added Bank Rotation to texture creation on XboxOne

diff --git a/Engine/Source/Runtime/D3D12RHI/Private/D3D12Texture.cpp b/Engine/Source/Runtime/D3D12RHI/Private/D3D12Texture.cpp
--- a/Engine/Source/Runtime/D3D12RHI/Private/D3D12Texture.cpp
+++ b/Engine/Source/Runtime/D3D12RHI/Private/D3D12Texture.cpp
@@ -298,7 +298,13 @@ void FD3D12TextureStats::D3D12TextureAllocated(TD3D12Texture2D<BaseResourceType>
 		}
 		else
 		{
+#if PLATFORM_XBOXONE
+			D3D12_RESOURCE_DESC Desc = D3D12Texture2D->GetDesc();
+			Desc.Layout = D3D12XBOX_BANK_ROTATED_TILE_MODE(Desc.Layout, 0);	// Remove bank rotation as this can cause the GetResourceAllocationInfo() to return 0
+#else
 			const D3D12_RESOURCE_DESC& Desc = D3D12Texture2D->GetDesc();
+#endif
+
 			const D3D12_RESOURCE_ALLOCATION_INFO AllocationInfo = Texture.GetParentDevice()->GetDevice()->GetResourceAllocationInfo(0, 1, &Desc);
 			const int64 TextureSize = AllocationInfo.SizeInBytes;
 
diff --git a/Engine/Source/Runtime/D3D12RHI/Public/D3D12Util.h b/Engine/Source/Runtime/D3D12RHI/Public/D3D12Util.h
--- a/Engine/Source/Runtime/D3D12RHI/Public/D3D12Util.h
+++ b/Engine/Source/Runtime/D3D12RHI/Public/D3D12Util.h
@@ -1143,4 +1143,38 @@ private:
 	TArray<Type*> AllObjects;
 	FCriticalSection CS;
 	static __declspec(thread) void* ThisThreadObject;
-};
\ No newline at end of file
+};
+
+#if PLATFORM_XBOXONE
+#include "XboxOneAllowPlatformTypes.h"
+#include <xg.h>
+#include "XboxOneHidePlatformTypes.h"
+
+FORCEINLINE XG_TILE_MODE GetXgOptimalTileMode(const D3D12_RESOURCE_DESC& TextureDesc, XG_BIND_FLAG BindFlags, XG_RESOURCE_MISC_FLAG MiscFlags)
+{
+	XG_RESOURCE_DIMENSION xgDimension = (XG_RESOURCE_DIMENSION)TextureDesc.Dimension;
+	check(xgDimension == XG_RESOURCE_DIMENSION_TEXTURE2D || xgDimension == XG_RESOURCE_DIMENSION_TEXTURE3D);  
+	XG_TILE_MODE iTileMode = XG_TILE_MODE_INVALID;
+	if ((BindFlags & XG_BIND_DEPTH_STENCIL) == XG_BIND_DEPTH_STENCIL)
+	{
+		XG_TILE_MODE iTileModeStencil = XG_TILE_MODE_INVALID;
+		const bool bCompressed = (MiscFlags & XG_RESOURCE_MISC_NO_DEPTH_COMPRESSION) != 0;
+
+#if _XDK_EDITION >= 170600
+		const bool TileResource = false;
+		const bool TextureCompatible = false;
+		XGComputeOptimalDepthStencilTileModes((XG_FORMAT)TextureDesc.Format, TextureDesc.Width, TextureDesc.Height, TextureDesc.DepthOrArraySize, TextureDesc.SampleDesc.Count, bCompressed, TileResource, TextureCompatible, &iTileMode, &iTileModeStencil);
+#else
+		XGComputeOptimalDepthStencilTileModes((XG_FORMAT)TextureDesc.Format, TextureDesc.Width, TextureDesc.Height, TextureDesc.DepthOrArraySize, TextureDesc.SampleDesc.Count, bCompressed, &iTileMode, &iTileModeStencil);
+#endif
+	}
+	else
+	{
+		iTileMode = XGComputeOptimalTileMode(xgDimension, (XG_FORMAT)TextureDesc.Format, TextureDesc.Width, TextureDesc.Height, TextureDesc.Dimension, TextureDesc.SampleDesc.Count, BindFlags, MiscFlags);
+
+	}
+	check(iTileMode != XG_TILE_MODE_INVALID);
+	return iTileMode;
+}
+
+#endif
\ No newline at end of file
