Cherry-pick 6628360: D3D12 frame fence is now advanced in FD3D12Adapter::EndFrame(), which happens regardless of whether application is minimized.

DefaultContext.ReleaseCommandAllocator() and GetGPUProfiler().EndFrame() steps are now also performed in FD3D12Adapter::EndFrame(), which fixes command allocator memory bloat when game is minimized.

This solved GPU resource release/recycling when the game is minimized that lead to resource exhaustion and eventual crash.

diff --git a/Engine/Source/Runtime/D3D12RHI/Private/D3D12Adapter.cpp b/Engine/Source/Runtime/D3D12RHI/Private/D3D12Adapter.cpp
--- a/Engine/Source/Runtime/D3D12RHI/Private/D3D12Adapter.cpp
+++ b/Engine/Source/Runtime/D3D12RHI/Private/D3D12Adapter.cpp
@@ -23,25 +23,6 @@ static TAutoConsoleVariable<int32> CVarResidencyManagement(
 );
 #endif // ENABLE_RESIDENCY_MANAGEMENT
 
-struct FRHICommandSignalFrameFence final : public FRHICommand<FRHICommandSignalFrameFence>
-{
-	ED3D12CommandQueueType QueueType;
-	FD3D12ManualFence* const Fence;
-	const uint64 Value;
-	FORCEINLINE_DEBUGGABLE FRHICommandSignalFrameFence(ED3D12CommandQueueType InQueueType, FD3D12ManualFence* InFence, uint64 InValue)
-		: QueueType(InQueueType)
-		, Fence(InFence)
-		, Value(InValue)
-	{ 
-	}
-
-	void Execute(FRHICommandListBase& CmdList)
-	{
-		Fence->Signal(QueueType, Value);
-		check(Fence->GetLastSignaledFence() == Value);
-	}
-};
-
 FD3D12Adapter::FD3D12Adapter(FD3D12AdapterDesc& DescIn)
 	: OwningRHI(nullptr)
 	, bDepthBoundsTestSupported(false)
@@ -662,26 +643,6 @@ void FD3D12Adapter::EndFrame()
 	GetDeferredDeletionQueue().ReleaseResources();
 }
 
-void FD3D12Adapter::SignalFrameFence_RenderThread(FRHICommandListImmediate& RHICmdList)
-{
-	check(IsInRenderingThread());
-	check(RHICmdList.IsImmediate());
-
-	// Increment the current fence (on render thread timeline).
-	const uint64 PreviousFence = FrameFence->IncrementCurrentFence();
-
-	// Queue a command to signal the frame fence is complete on the GPU (on the RHI thread timeline if using an RHI thread).
-	if (RHICmdList.Bypass() || !IsRunningRHIInSeparateThread())
-	{
-		FRHICommandSignalFrameFence Cmd(ED3D12CommandQueueType::Default, FrameFence, PreviousFence);
-		Cmd.Execute(RHICmdList);
-	}
-	else
-	{
-		ALLOC_COMMAND_CL(RHICmdList, FRHICommandSignalFrameFence)(ED3D12CommandQueueType::Default, FrameFence, PreviousFence);
-	}
-}
-
 FD3D12TemporalEffect* FD3D12Adapter::GetTemporalEffect(const FName& EffectName)
 {
 	FD3D12TemporalEffect* Effect = TemporalEffectMap.Find(EffectName);
diff --git a/Engine/Source/Runtime/D3D12RHI/Private/D3D12CommandContext.cpp b/Engine/Source/Runtime/D3D12RHI/Private/D3D12CommandContext.cpp
--- a/Engine/Source/Runtime/D3D12RHI/Private/D3D12CommandContext.cpp
+++ b/Engine/Source/Runtime/D3D12RHI/Private/D3D12CommandContext.cpp
@@ -414,6 +414,20 @@ void FD3D12CommandContextBase::RHIEndFrame()
 	{
 		FD3D12Device* Device = ParentAdapter->GetDevice(GPUIndex);
 
+		FD3D12CommandContext& DefaultContext = Device->GetDefaultCommandContext();
+		DefaultContext.CommandListHandle.FlushResourceBarriers();
+
+		DefaultContext.ReleaseCommandAllocator();
+		DefaultContext.ClearState();
+		DefaultContext.FlushCommands();
+
+		if (GEnableAsyncCompute)
+		{
+			FD3D12CommandContext& DefaultAsyncComputeContext = Device->GetDefaultAsyncComputeContext();
+			DefaultAsyncComputeContext.ReleaseCommandAllocator();
+			DefaultAsyncComputeContext.ClearState();
+		}
+
 		const uint32 NumContexts = Device->GetNumContexts();
 		for (uint32 i = 0; i < NumContexts; ++i)
 		{
@@ -450,8 +464,18 @@ void FD3D12CommandContextBase::RHIEndFrame()
 		Device->GetCommandListManager().ReleaseResourceBarrierCommandListAllocator();
 	}
 
-		UpdateMemoryStats();
-	}
+	UpdateMemoryStats();
+
+	// Stop Timing at the very last moment
+
+	ParentAdapter->GetGPUProfiler().EndFrame(ParentAdapter->GetOwningRHI());
+
+
+	// Advance frame fence
+
+	FD3D12ManualFence& FrameFence = ParentAdapter->GetFrameFence();
+	FrameFence.Signal(ED3D12CommandQueueType::Default, FrameFence.IncrementCurrentFence());
+}
 
 void FD3D12CommandContextBase::UpdateMemoryStats()
 {
diff --git a/Engine/Source/Runtime/D3D12RHI/Private/D3D12Viewport.cpp b/Engine/Source/Runtime/D3D12RHI/Private/D3D12Viewport.cpp
--- a/Engine/Source/Runtime/D3D12RHI/Private/D3D12Viewport.cpp
+++ b/Engine/Source/Runtime/D3D12RHI/Private/D3D12Viewport.cpp
@@ -666,27 +666,7 @@ bool FD3D12Viewport::Present(bool bLockToVsync)
 			FD3D12DynamicRHI::TransitionResource(DefaultContext.CommandListHandle, DeviceSDRBackBuffer->GetShaderResourceView(), D3D12_RESOURCE_STATE_PRESENT);
 		}
 		DefaultContext.CommandListHandle.FlushResourceBarriers();
-	}
-
-	// Stop Timing at the very last moment
-	Adapter->GetGPUProfiler().EndFrame(Adapter->GetOwningRHI());
-
-	for (uint32 GPUIndex : FRHIGPUMask::All())
-	{
-		FD3D12Device* Device = Adapter->GetDevice(GPUIndex);
-		FD3D12CommandContext& DefaultContext = Device->GetDefaultCommandContext();
-
-		// Execute the current command lists, and then open a new command list with a new command allocator.
-		DefaultContext.ReleaseCommandAllocator();
-		DefaultContext.ClearState();
 		DefaultContext.FlushCommands();
-
-		if (GEnableAsyncCompute)
-		{
-			FD3D12CommandContext& DefaultAsyncComputeContext = Device->GetDefaultAsyncComputeContext();
-			DefaultAsyncComputeContext.ReleaseCommandAllocator();
-			DefaultAsyncComputeContext.ClearState();
-		}
 	}
 
 #if WITH_MGPU
@@ -907,15 +887,6 @@ void FD3D12DynamicRHI::RHIAdvanceFrameForGetViewportBackBuffer(FViewportRHIParam
 	UE_LOG(LogD3D12RHI, Log, TEXT("Thread %s: RHIAdvanceFrameForGetViewportBackBuffer"), ThreadName.GetCharArray().GetData());
 #endif
 
-	// Queue a command to signal the current frame is a complete on the GPU.
-	// Note: No need to handle multiple adapters yet, eventually this function will take a viewport as input.
-	FRHICommandListImmediate& RHICmdList = FRHICommandListExecutor::GetImmediateCommandList();
-	auto& Adapter = ChosenAdapters[0];
-	{
-		check(Adapter != nullptr);
-		Adapter->SignalFrameFence_RenderThread(RHICmdList);
-	}
-
 	// Advance frame so the next call to RHIGetViewportBackBuffer returns the next buffer in the swap chain.
 	FD3D12Viewport* Viewport = FD3D12DynamicRHI::ResourceCast(ViewportRHI);
 	Viewport->AdvanceBackBufferFrame_RenderThread();
