Added default input resource desc for D3D12 Texture stats

diff --git a/Engine/Source/Runtime/D3D12RHI/Private/D3D12Texture.cpp b/Engine/Source/Runtime/D3D12RHI/Private/D3D12Texture.cpp
--- a/Engine/Source/Runtime/D3D12RHI/Private/D3D12Texture.cpp
+++ b/Engine/Source/Runtime/D3D12RHI/Private/D3D12Texture.cpp
@@ -62,9 +62,9 @@ template TD3D12Texture2D<FD3D12BaseTextureCube>::~TD3D12Texture2D();
 
 /// @cond DOXYGEN_WARNINGS
 
-template void FD3D12TextureStats::D3D12TextureAllocated(TD3D12Texture2D<FD3D12BaseTexture2D>& Texture);
-template void FD3D12TextureStats::D3D12TextureAllocated(TD3D12Texture2D<FD3D12BaseTexture2DArray>& Texture);
-template void FD3D12TextureStats::D3D12TextureAllocated(TD3D12Texture2D<FD3D12BaseTextureCube>& Texture);
+template void FD3D12TextureStats::D3D12TextureAllocated(TD3D12Texture2D<FD3D12BaseTexture2D>& Texture, const D3D12_RESOURCE_DESC *Desc);
+template void FD3D12TextureStats::D3D12TextureAllocated(TD3D12Texture2D<FD3D12BaseTexture2DArray>& Texture, const D3D12_RESOURCE_DESC *Desc);
+template void FD3D12TextureStats::D3D12TextureAllocated(TD3D12Texture2D<FD3D12BaseTextureCube>& Texture, const D3D12_RESOURCE_DESC *Desc);
 
 template void FD3D12TextureStats::D3D12TextureDeleted(TD3D12Texture2D<FD3D12BaseTexture2D>& Texture);
 template void FD3D12TextureStats::D3D12TextureDeleted(TD3D12Texture2D<FD3D12BaseTexture2DArray>& Texture);
@@ -285,7 +285,7 @@ void FD3D12TextureStats::UpdateD3D12TextureStats(const D3D12_RESOURCE_DESC& Desc
 }
 
 template<typename BaseResourceType>
-void FD3D12TextureStats::D3D12TextureAllocated(TD3D12Texture2D<BaseResourceType>& Texture)
+void FD3D12TextureStats::D3D12TextureAllocated(TD3D12Texture2D<BaseResourceType>& Texture, const D3D12_RESOURCE_DESC *Desc)
 {
 	FD3D12Resource* D3D12Texture2D = Texture.GetResource();
 
@@ -298,13 +298,17 @@ void FD3D12TextureStats::D3D12TextureAllocated(TD3D12Texture2D<BaseResourceType>
 		}
 		else
 		{
-			const D3D12_RESOURCE_DESC& Desc = D3D12Texture2D->GetDesc();
-			const D3D12_RESOURCE_ALLOCATION_INFO AllocationInfo = Texture.GetParentDevice()->GetDevice()->GetResourceAllocationInfo(0, 1, &Desc);
+			if (!Desc)
+			{
+				Desc = &D3D12Texture2D->GetDesc();
+			}
+
+			const D3D12_RESOURCE_ALLOCATION_INFO AllocationInfo = Texture.GetParentDevice()->GetDevice()->GetResourceAllocationInfo(0, 1, Desc);
 			const int64 TextureSize = AllocationInfo.SizeInBytes;
 
 			Texture.SetMemorySize(TextureSize);
 
-			UpdateD3D12TextureStats(Desc, TextureSize, false, Texture.IsCubemap());
+			UpdateD3D12TextureStats(*Desc, TextureSize, false, Texture.IsCubemap());
 
 #if PLATFORM_WINDOWS
 			// On Windows there is no way to hook into the low level d3d allocations and frees.
diff --git a/Engine/Source/Runtime/D3D12RHI/Private/D3D12Texture.h b/Engine/Source/Runtime/D3D12RHI/Private/D3D12Texture.h
--- a/Engine/Source/Runtime/D3D12RHI/Private/D3D12Texture.h
+++ b/Engine/Source/Runtime/D3D12RHI/Private/D3D12Texture.h
@@ -502,7 +502,7 @@ public:
 	static void UpdateD3D12TextureStats(const D3D12_RESOURCE_DESC& Desc, int64 TextureSize, bool b3D, bool bCubeMap);
 
 	template<typename BaseResourceType>
-	static void D3D12TextureAllocated(TD3D12Texture2D<BaseResourceType>& Texture);
+	static void D3D12TextureAllocated(TD3D12Texture2D<BaseResourceType>& Texture, const D3D12_RESOURCE_DESC *Desc = nullptr);
 
 	template<typename BaseResourceType>
 	static void D3D12TextureDeleted(TD3D12Texture2D<BaseResourceType>& Texture);
