UE4.23 - Deprecate FUniformBufferRHIParamRef, FIndexBufferRHIParamRef, FVertexBufferRHIParamRef, FStructuredBufferRHIParamRef, FTextureRHIParamRef, FTexture2DRHIParamRef, FTexture2DArrayRHIParamRef, F

diff --git a/Engine/Source/Runtime/D3D12RHI/Private/D3D12CommandContext.h b/Engine/Source/Runtime/D3D12RHI/Private/D3D12CommandContext.h
--- a/Engine/Source/Runtime/D3D12RHI/Private/D3D12CommandContext.h
+++ b/Engine/Source/Runtime/D3D12RHI/Private/D3D12CommandContext.h
@@ -35,7 +35,7 @@ public:
 
 	FD3D12CommandContextBase(class FD3D12Adapter* InParent, FRHIGPUMask InGPUMask, bool InIsDefaultContext, bool InIsAsyncComputeContext);
 
-	void RHIBeginDrawingViewport(FRHIViewport* Viewport, FTextureRHIParamRef RenderTargetRHI) final override;
+	void RHIBeginDrawingViewport(FRHIViewport* Viewport, FRHITexture* RenderTargetRHI) final override;
 	void RHIEndDrawingViewport(FRHIViewport* Viewport, bool bPresent, bool bLockToVsync) final override;
 
 	void RHIBeginFrame() final override;
@@ -152,7 +152,7 @@ public:
 	/** Set to true when the current shading setup uses tessellation */
 	bool bUsingTessellation;
 
-	virtual void FlushMetadata(FTextureRHIParamRef* InTextures, int32 NumTextures) {};
+	virtual void FlushMetadata(FRHITexture** InTextures, int32 NumTextures) {};
 
 	D3D12_RESOURCE_STATES SkipFastClearEliminateState;
 
@@ -231,14 +231,14 @@ public:
 	virtual void RHISetComputeShader(FRHIComputeShader* ComputeShader) final override;
 	virtual void RHISetComputePipelineState(FRHIComputePipelineState* ComputePipelineState) final override;
 	virtual void RHIDispatchComputeShader(uint32 ThreadGroupCountX, uint32 ThreadGroupCountY, uint32 ThreadGroupCountZ) final override;
-	virtual void RHIDispatchIndirectComputeShader(FVertexBufferRHIParamRef ArgumentBuffer, uint32 ArgumentOffset) final override;
+	virtual void RHIDispatchIndirectComputeShader(FRHIVertexBuffer* ArgumentBuffer, uint32 ArgumentOffset) final override;
 	virtual void RHITransitionResources(EResourceTransitionAccess TransitionType, EResourceTransitionPipeline TransitionPipeline, FRHIUnorderedAccessView** InUAVs, int32 NumUAVs, FRHIComputeFence* WriteComputeFenceRHI) final override;
-	virtual void RHISetShaderTexture(FRHIComputeShader* PixelShader, uint32 TextureIndex, FTextureRHIParamRef NewTexture) final override;
+	virtual void RHISetShaderTexture(FRHIComputeShader* PixelShader, uint32 TextureIndex, FRHITexture* NewTexture) final override;
 	virtual void RHISetShaderSampler(FRHIComputeShader* ComputeShader, uint32 SamplerIndex, FRHISamplerState* NewState) final override;
 	virtual void RHISetUAVParameter(FRHIComputeShader* ComputeShader, uint32 UAVIndex, FRHIUnorderedAccessView* UAV) final override;
 	virtual void RHISetUAVParameter(FRHIComputeShader* ComputeShader, uint32 UAVIndex, FRHIUnorderedAccessView* UAV, uint32 InitialCount) final override;
 	virtual void RHISetShaderResourceViewParameter(FRHIComputeShader* ComputeShader, uint32 SamplerIndex, FRHIShaderResourceView* SRV) final override;
-	virtual void RHISetShaderUniformBuffer(FRHIComputeShader* ComputeShader, uint32 BufferIndex, FUniformBufferRHIParamRef Buffer) final override;
+	virtual void RHISetShaderUniformBuffer(FRHIComputeShader* ComputeShader, uint32 BufferIndex, FRHIUniformBuffer* Buffer) final override;
 	virtual void RHISetShaderParameter(FRHIComputeShader* ComputeShader, uint32 BufferIndex, uint32 BaseIndex, uint32 NumBytes, const void* NewValue) final override;
 	virtual void RHIPushEvent(const TCHAR* Name, FColor Color) final override;
 	virtual void RHIPopEvent() final override;
@@ -249,10 +249,10 @@ public:
 	virtual void RHIFlushComputeShaderCache() final override;
 	virtual void RHISetMultipleViewports(uint32 Count, const FViewportBounds* Data) final override;
 	virtual void RHIClearTinyUAV(FRHIUnorderedAccessView* UnorderedAccessViewRHI, const uint32* Values) final override;
-	virtual void RHICopyToResolveTarget(FTextureRHIParamRef SourceTexture, FTextureRHIParamRef DestTexture, const FResolveParams& ResolveParams) override;
-	virtual void RHICopyTexture(FTextureRHIParamRef SourceTexture, FTextureRHIParamRef DestTexture, const FRHICopyTextureInfo& CopyInfo) final override;
-	virtual void RHITransitionResources(EResourceTransitionAccess TransitionType, FTextureRHIParamRef* InTextures, int32 NumTextures) final override;
-	virtual void RHICopyToStagingBuffer(FVertexBufferRHIParamRef SourceBuffer, FRHIStagingBuffer* DestinationStagingBuffer, uint32 Offset, uint32 NumBytes) final override;
+	virtual void RHICopyToResolveTarget(FRHITexture* SourceTexture, FRHITexture* DestTexture, const FResolveParams& ResolveParams) override;
+	virtual void RHICopyTexture(FRHITexture* SourceTexture, FRHITexture* DestTexture, const FRHICopyTextureInfo& CopyInfo) final override;
+	virtual void RHITransitionResources(EResourceTransitionAccess TransitionType, FRHITexture** InTextures, int32 NumTextures) final override;
+	virtual void RHICopyToStagingBuffer(FRHIVertexBuffer* SourceBuffer, FRHIStagingBuffer* DestinationStagingBuffer, uint32 Offset, uint32 NumBytes) final override;
 	virtual void RHIWriteGPUFence(FRHIGPUFence* Fence) final override;
 	virtual void RHIBeginRenderQuery(FRHIRenderQuery* RenderQuery) final override;
 	virtual void RHIEndRenderQuery(FRHIRenderQuery* RenderQuery) final override;
@@ -260,15 +260,15 @@ public:
 	void RHIEndOcclusionQueryBatch();
 	virtual void RHIBeginScene() final override;
 	virtual void RHIEndScene() final override;
-	virtual void RHISetStreamSource(uint32 StreamIndex, FVertexBufferRHIParamRef VertexBuffer, uint32 Offset) final override;
+	virtual void RHISetStreamSource(uint32 StreamIndex, FRHIVertexBuffer* VertexBuffer, uint32 Offset) final override;
 	virtual void RHISetViewport(uint32 MinX, uint32 MinY, float MinZ, uint32 MaxX, uint32 MaxY, float MaxZ) final override;
 	virtual void RHISetScissorRect(bool bEnable, uint32 MinX, uint32 MinY, uint32 MaxX, uint32 MaxY) final override;
 	virtual void RHISetGraphicsPipelineState(FRHIGraphicsPipelineState* GraphicsPipelineState) final override;
-	virtual void RHISetShaderTexture(FRHIVertexShader* VertexShader, uint32 TextureIndex, FTextureRHIParamRef NewTexture) final override;
-	virtual void RHISetShaderTexture(FRHIHullShader* HullShader, uint32 TextureIndex, FTextureRHIParamRef NewTexture) final override;
-	virtual void RHISetShaderTexture(FRHIDomainShader* DomainShader, uint32 TextureIndex, FTextureRHIParamRef NewTexture) final override;
-	virtual void RHISetShaderTexture(FRHIGeometryShader* GeometryShader, uint32 TextureIndex, FTextureRHIParamRef NewTexture) final override;
-	virtual void RHISetShaderTexture(FRHIPixelShader* PixelShader, uint32 TextureIndex, FTextureRHIParamRef NewTexture) final override;
+	virtual void RHISetShaderTexture(FRHIVertexShader* VertexShader, uint32 TextureIndex, FRHITexture* NewTexture) final override;
+	virtual void RHISetShaderTexture(FRHIHullShader* HullShader, uint32 TextureIndex, FRHITexture* NewTexture) final override;
+	virtual void RHISetShaderTexture(FRHIDomainShader* DomainShader, uint32 TextureIndex, FRHITexture* NewTexture) final override;
+	virtual void RHISetShaderTexture(FRHIGeometryShader* GeometryShader, uint32 TextureIndex, FRHITexture* NewTexture) final override;
+	virtual void RHISetShaderTexture(FRHIPixelShader* PixelShader, uint32 TextureIndex, FRHITexture* NewTexture) final override;
 	virtual void RHISetShaderSampler(FRHIVertexShader* VertexShader, uint32 SamplerIndex, FRHISamplerState* NewState) final override;
 	virtual void RHISetShaderSampler(FRHIGeometryShader* GeometryShader, uint32 SamplerIndex, FRHISamplerState* NewState) final override;
 	virtual void RHISetShaderSampler(FRHIDomainShader* DomainShader, uint32 SamplerIndex, FRHISamplerState* NewState) final override;
@@ -279,11 +279,11 @@ public:
 	virtual void RHISetShaderResourceViewParameter(FRHIHullShader* HullShader, uint32 SamplerIndex, FRHIShaderResourceView* SRV) final override;
 	virtual void RHISetShaderResourceViewParameter(FRHIDomainShader* DomainShader, uint32 SamplerIndex, FRHIShaderResourceView* SRV) final override;
 	virtual void RHISetShaderResourceViewParameter(FRHIGeometryShader* GeometryShader, uint32 SamplerIndex, FRHIShaderResourceView* SRV) final override;
-	virtual void RHISetShaderUniformBuffer(FRHIVertexShader* VertexShader, uint32 BufferIndex, FUniformBufferRHIParamRef Buffer) final override;
-	virtual void RHISetShaderUniformBuffer(FRHIHullShader* HullShader, uint32 BufferIndex, FUniformBufferRHIParamRef Buffer) final override;
-	virtual void RHISetShaderUniformBuffer(FRHIDomainShader* DomainShader, uint32 BufferIndex, FUniformBufferRHIParamRef Buffer) final override;
-	virtual void RHISetShaderUniformBuffer(FRHIGeometryShader* GeometryShader, uint32 BufferIndex, FUniformBufferRHIParamRef Buffer) final override;
-	virtual void RHISetShaderUniformBuffer(FRHIPixelShader* PixelShader, uint32 BufferIndex, FUniformBufferRHIParamRef Buffer) final override;
+	virtual void RHISetShaderUniformBuffer(FRHIVertexShader* VertexShader, uint32 BufferIndex, FRHIUniformBuffer* Buffer) final override;
+	virtual void RHISetShaderUniformBuffer(FRHIHullShader* HullShader, uint32 BufferIndex, FRHIUniformBuffer* Buffer) final override;
+	virtual void RHISetShaderUniformBuffer(FRHIDomainShader* DomainShader, uint32 BufferIndex, FRHIUniformBuffer* Buffer) final override;
+	virtual void RHISetShaderUniformBuffer(FRHIGeometryShader* GeometryShader, uint32 BufferIndex, FRHIUniformBuffer* Buffer) final override;
+	virtual void RHISetShaderUniformBuffer(FRHIPixelShader* PixelShader, uint32 BufferIndex, FRHIUniformBuffer* Buffer) final override;
 	virtual void RHISetShaderParameter(FRHIVertexShader* VertexShader, uint32 BufferIndex, uint32 BaseIndex, uint32 NumBytes, const void* NewValue) final override;
 	virtual void RHISetShaderParameter(FRHIPixelShader* PixelShader, uint32 BufferIndex, uint32 BaseIndex, uint32 NumBytes, const void* NewValue) final override;
 	virtual void RHISetShaderParameter(FRHIHullShader* HullShader, uint32 BufferIndex, uint32 BaseIndex, uint32 NumBytes, const void* NewValue) final override;
@@ -295,12 +295,12 @@ public:
 	virtual void RHISetRenderTargetsAndClear(const FRHISetRenderTargetsInfo& RenderTargetsInfo) final override;
 	virtual void RHIBindClearMRTValues(bool bClearColor, bool bClearDepth, bool bClearStencil) final override;
 	virtual void RHIDrawPrimitive(uint32 BaseVertexIndex, uint32 NumPrimitives, uint32 NumInstances) final override;
-	virtual void RHIDrawPrimitiveIndirect(FVertexBufferRHIParamRef ArgumentBuffer, uint32 ArgumentOffset) final override;
-	virtual void RHIDrawIndexedIndirect(FIndexBufferRHIParamRef IndexBufferRHI, FStructuredBufferRHIParamRef ArgumentsBufferRHI, int32 DrawArgumentsIndex, uint32 NumInstances) final override;
-	virtual void RHIDrawIndexedPrimitive(FIndexBufferRHIParamRef IndexBuffer, int32 BaseVertexIndex, uint32 FirstInstance, uint32 NumVertices, uint32 StartIndex, uint32 NumPrimitives, uint32 NumInstances) final override;
-	virtual void RHIDrawIndexedPrimitiveIndirect(FIndexBufferRHIParamRef IndexBuffer, FVertexBufferRHIParamRef ArgumentBuffer, uint32 ArgumentOffset) final override;
+	virtual void RHIDrawPrimitiveIndirect(FRHIVertexBuffer* ArgumentBuffer, uint32 ArgumentOffset) final override;
+	virtual void RHIDrawIndexedIndirect(FRHIIndexBuffer* IndexBufferRHI, FRHIStructuredBuffer* ArgumentsBufferRHI, int32 DrawArgumentsIndex, uint32 NumInstances) final override;
+	virtual void RHIDrawIndexedPrimitive(FRHIIndexBuffer* IndexBuffer, int32 BaseVertexIndex, uint32 FirstInstance, uint32 NumVertices, uint32 StartIndex, uint32 NumPrimitives, uint32 NumInstances) final override;
+	virtual void RHIDrawIndexedPrimitiveIndirect(FRHIIndexBuffer* IndexBuffer, FRHIVertexBuffer* ArgumentBuffer, uint32 ArgumentOffset) final override;
 	virtual void RHISetDepthBounds(float MinDepth, float MaxDepth) final override;
-	virtual void RHIUpdateTextureReference(FRHITextureReference* TextureRef, FTextureRHIParamRef NewTexture) final override;
+	virtual void RHIUpdateTextureReference(FRHITextureReference* TextureRef, FRHITexture* NewTexture) final override;
 
 	virtual void RHIClearMRTImpl(bool bClearColor, int32 NumClearColors, const FLinearColor* ColorArray, bool bClearDepth, float Depth, bool bClearStencil, uint32 Stencil);
 
@@ -329,10 +329,10 @@ public:
 	virtual void RHIWaitForTemporalEffect(const FName& InEffectName) final AFR_API_OVERRIDE;
 
 	// This should be called right after the effect generates the resources which will be used in subsequent frame(s).
-	virtual void RHIBroadcastTemporalEffect(const FName& InEffectName, FTextureRHIParamRef* InTextures, int32 NumTextures) final AFR_API_OVERRIDE;
+	virtual void RHIBroadcastTemporalEffect(const FName& InEffectName, FRHITexture** InTextures, int32 NumTextures) final AFR_API_OVERRIDE;
 
 #if D3D12_RHI_RAYTRACING
-	virtual void RHICopyBufferRegion(FVertexBufferRHIParamRef DestBuffer, uint64 DstOffset, FVertexBufferRHIParamRef SourceBuffer, uint64 SrcOffset, uint64 NumBytes) final override;
+	virtual void RHICopyBufferRegion(FRHIVertexBuffer* DestBuffer, uint64 DstOffset, FRHIVertexBuffer* SourceBuffer, uint64 SrcOffset, uint64 NumBytes) final override;
 	virtual void RHICopyBufferRegions(const TArrayView<const FCopyBufferRegionParams> Params) final override;
 	virtual void RHIBuildAccelerationStructure(FRHIRayTracingGeometry* Geometry) final override;
 	virtual void RHIUpdateAccelerationStructures(const TArrayView<const FAccelerationStructureUpdateParams> Params) final override;
@@ -353,13 +353,13 @@ public:
 	virtual void RHISetRayTracingHitGroup(
 		FRHIRayTracingScene* Scene, uint32 InstanceIndex, uint32 SegmentIndex, uint32 ShaderSlot,
 		FRHIRayTracingPipelineState* Pipeline, uint32 HitGroupIndex,
-		uint32 NumUniformBuffers, const FUniformBufferRHIParamRef* UniformBuffers,
+		uint32 NumUniformBuffers, FRHIUniformBuffer* const* UniformBuffers,
 		uint32 LooseParameterDataSize, const void* LooseParameterData,
 		uint32 UserData) final override;
 	virtual void RHISetRayTracingCallableShader(
 		FRHIRayTracingScene* Scene, uint32 ShaderSlotInScene,
 		FRHIRayTracingPipelineState* Pipeline, uint32 ShaderIndexInPipeline,
-		uint32 NumUniformBuffers, const FUniformBufferRHIParamRef* UniformBuffers,
+		uint32 NumUniformBuffers, FRHIUniformBuffer* const* UniformBuffers,
 		uint32 UserData) final override;
 #endif // D3D12_RHI_RAYTRACING
 
@@ -455,14 +455,14 @@ public:
 	{
 		ContextRedirect(RHIDispatchComputeShader(ThreadGroupCountX, ThreadGroupCountY, ThreadGroupCountZ));
 	}
-	FORCEINLINE virtual void RHIDispatchIndirectComputeShader(FVertexBufferRHIParamRef ArgumentBuffer, uint32 ArgumentOffset) final override
+	FORCEINLINE virtual void RHIDispatchIndirectComputeShader(FRHIVertexBuffer* ArgumentBuffer, uint32 ArgumentOffset) final override
 	{
 		ContextRedirect(RHIDispatchIndirectComputeShader(ArgumentBuffer, ArgumentOffset));
 	}
 	// Special implementation that only signal the fence once.
 	virtual void RHITransitionResources(EResourceTransitionAccess TransitionType, EResourceTransitionPipeline TransitionPipeline, FRHIUnorderedAccessView** InUAVs, int32 NumUAVs, FRHIComputeFence* WriteComputeFenceRHI) final override;
 
-	FORCEINLINE virtual void RHICopyToStagingBuffer(FVertexBufferRHIParamRef SourceBuffer, FRHIStagingBuffer* DestinationStagingBuffer, uint32 Offset, uint32 NumBytes) final override
+	FORCEINLINE virtual void RHICopyToStagingBuffer(FRHIVertexBuffer* SourceBuffer, FRHIStagingBuffer* DestinationStagingBuffer, uint32 Offset, uint32 NumBytes) final override
 	{
 		ContextRedirect(RHICopyToStagingBuffer(SourceBuffer, DestinationStagingBuffer, Offset, NumBytes));
 	}
@@ -470,7 +470,7 @@ public:
 	{
 		ContextRedirect(RHIWriteGPUFence(Fence));
 	}
-	FORCEINLINE virtual void RHISetShaderTexture(FRHIComputeShader* PixelShader, uint32 TextureIndex, FTextureRHIParamRef NewTexture) final override
+	FORCEINLINE virtual void RHISetShaderTexture(FRHIComputeShader* PixelShader, uint32 TextureIndex, FRHITexture* NewTexture) final override
 	{
 		ContextRedirect(RHISetShaderTexture(PixelShader, TextureIndex, NewTexture));
 	}
@@ -490,7 +490,7 @@ public:
 	{
 		ContextRedirect(RHISetShaderResourceViewParameter(ComputeShader, SamplerIndex, SRV));
 	}
-	FORCEINLINE virtual void RHISetShaderUniformBuffer(FRHIComputeShader* ComputeShader, uint32 BufferIndex, FUniformBufferRHIParamRef Buffer) final override
+	FORCEINLINE virtual void RHISetShaderUniformBuffer(FRHIComputeShader* ComputeShader, uint32 BufferIndex, FRHIUniformBuffer* Buffer) final override
 	{
 		ContextRedirect(RHISetShaderUniformBuffer(ComputeShader, BufferIndex, Buffer));
 	}
@@ -528,11 +528,11 @@ public:
 	{
 		ContextRedirect(RHIClearTinyUAV(UnorderedAccessViewRHI, Values));
 	}
-	FORCEINLINE virtual void RHICopyToResolveTarget(FTextureRHIParamRef SourceTexture, FTextureRHIParamRef DestTexture, const FResolveParams& ResolveParams) final override
+	FORCEINLINE virtual void RHICopyToResolveTarget(FRHITexture* SourceTexture, FRHITexture* DestTexture, const FResolveParams& ResolveParams) final override
 	{
 		ContextRedirect(RHICopyToResolveTarget(SourceTexture, DestTexture, ResolveParams));
 	}
-	FORCEINLINE virtual void RHITransitionResources(EResourceTransitionAccess TransitionType, FTextureRHIParamRef* InTextures, int32 NumTextures) final override
+	FORCEINLINE virtual void RHITransitionResources(EResourceTransitionAccess TransitionType, FRHITexture** InTextures, int32 NumTextures) final override
 	{
 		ContextRedirect(RHITransitionResources(TransitionType, InTextures, NumTextures));
 	}
@@ -552,7 +552,7 @@ public:
 	{
 		ContextRedirect(RHIEndScene());
 	}
-	FORCEINLINE virtual void RHISetStreamSource(uint32 StreamIndex, FVertexBufferRHIParamRef VertexBuffer, uint32 Offset) final override
+	FORCEINLINE virtual void RHISetStreamSource(uint32 StreamIndex, FRHIVertexBuffer* VertexBuffer, uint32 Offset) final override
 	{
 		ContextRedirect(RHISetStreamSource(StreamIndex, VertexBuffer, Offset));
 	}
@@ -568,23 +568,23 @@ public:
 	{
 		ContextRedirect(RHISetGraphicsPipelineState(GraphicsPipelineState));
 	}
-	FORCEINLINE virtual void RHISetShaderTexture(FRHIVertexShader* VertexShader, uint32 TextureIndex, FTextureRHIParamRef NewTexture) final override
+	FORCEINLINE virtual void RHISetShaderTexture(FRHIVertexShader* VertexShader, uint32 TextureIndex, FRHITexture* NewTexture) final override
 	{
 		ContextRedirect(RHISetShaderTexture(VertexShader, TextureIndex, NewTexture));
 	}
-	FORCEINLINE virtual void RHISetShaderTexture(FRHIHullShader* HullShader, uint32 TextureIndex, FTextureRHIParamRef NewTexture) final override
+	FORCEINLINE virtual void RHISetShaderTexture(FRHIHullShader* HullShader, uint32 TextureIndex, FRHITexture* NewTexture) final override
 	{
 		ContextRedirect(RHISetShaderTexture(HullShader, TextureIndex, NewTexture));
 	}
-	FORCEINLINE virtual void RHISetShaderTexture(FRHIDomainShader* DomainShader, uint32 TextureIndex, FTextureRHIParamRef NewTexture) final override
+	FORCEINLINE virtual void RHISetShaderTexture(FRHIDomainShader* DomainShader, uint32 TextureIndex, FRHITexture* NewTexture) final override
 	{
 		ContextRedirect(RHISetShaderTexture(DomainShader, TextureIndex, NewTexture));
 	}
-	FORCEINLINE virtual void RHISetShaderTexture(FRHIGeometryShader* GeometryShader, uint32 TextureIndex, FTextureRHIParamRef NewTexture) final override
+	FORCEINLINE virtual void RHISetShaderTexture(FRHIGeometryShader* GeometryShader, uint32 TextureIndex, FRHITexture* NewTexture) final override
 	{
 		ContextRedirect(RHISetShaderTexture(GeometryShader, TextureIndex, NewTexture));
 	}
-	FORCEINLINE virtual void RHISetShaderTexture(FRHIPixelShader* PixelShader, uint32 TextureIndex, FTextureRHIParamRef NewTexture) final override
+	FORCEINLINE virtual void RHISetShaderTexture(FRHIPixelShader* PixelShader, uint32 TextureIndex, FRHITexture* NewTexture) final override
 	{
 		ContextRedirect(RHISetShaderTexture(PixelShader, TextureIndex, NewTexture));
 	}
@@ -628,23 +628,23 @@ public:
 	{
 		ContextRedirect(RHISetShaderResourceViewParameter(GeometryShader, SamplerIndex, SRV));
 	}
-	FORCEINLINE virtual void RHISetShaderUniformBuffer(FRHIVertexShader* VertexShader, uint32 BufferIndex, FUniformBufferRHIParamRef Buffer) final override
+	FORCEINLINE virtual void RHISetShaderUniformBuffer(FRHIVertexShader* VertexShader, uint32 BufferIndex, FRHIUniformBuffer* Buffer) final override
 	{
 		ContextRedirect(RHISetShaderUniformBuffer(VertexShader, BufferIndex, Buffer));
 	}
-	FORCEINLINE virtual void RHISetShaderUniformBuffer(FRHIHullShader* HullShader, uint32 BufferIndex, FUniformBufferRHIParamRef Buffer) final override
+	FORCEINLINE virtual void RHISetShaderUniformBuffer(FRHIHullShader* HullShader, uint32 BufferIndex, FRHIUniformBuffer* Buffer) final override
 	{
 		ContextRedirect(RHISetShaderUniformBuffer(HullShader, BufferIndex, Buffer));
 	}
-	FORCEINLINE virtual void RHISetShaderUniformBuffer(FRHIDomainShader* DomainShader, uint32 BufferIndex, FUniformBufferRHIParamRef Buffer) final override
+	FORCEINLINE virtual void RHISetShaderUniformBuffer(FRHIDomainShader* DomainShader, uint32 BufferIndex, FRHIUniformBuffer* Buffer) final override
 	{
 		ContextRedirect(RHISetShaderUniformBuffer(DomainShader, BufferIndex, Buffer));
 	}
-	FORCEINLINE virtual void RHISetShaderUniformBuffer(FRHIGeometryShader* GeometryShader, uint32 BufferIndex, FUniformBufferRHIParamRef Buffer) final override
+	FORCEINLINE virtual void RHISetShaderUniformBuffer(FRHIGeometryShader* GeometryShader, uint32 BufferIndex, FRHIUniformBuffer* Buffer) final override
 	{
 		ContextRedirect(RHISetShaderUniformBuffer(GeometryShader, BufferIndex, Buffer));
 	}
-	FORCEINLINE virtual void RHISetShaderUniformBuffer(FRHIPixelShader* PixelShader, uint32 BufferIndex, FUniformBufferRHIParamRef Buffer) final override
+	FORCEINLINE virtual void RHISetShaderUniformBuffer(FRHIPixelShader* PixelShader, uint32 BufferIndex, FRHIUniformBuffer* Buffer) final override
 	{
 		ContextRedirect(RHISetShaderUniformBuffer(PixelShader, BufferIndex, Buffer));
 	}
@@ -692,19 +692,19 @@ public:
 	{
 		ContextRedirect(RHIDrawPrimitive(BaseVertexIndex, NumPrimitives, NumInstances));
 	}
-	FORCEINLINE virtual void RHIDrawPrimitiveIndirect(FVertexBufferRHIParamRef ArgumentBuffer, uint32 ArgumentOffset) final override
+	FORCEINLINE virtual void RHIDrawPrimitiveIndirect(FRHIVertexBuffer* ArgumentBuffer, uint32 ArgumentOffset) final override
 	{
 		ContextRedirect(RHIDrawPrimitiveIndirect(ArgumentBuffer, ArgumentOffset));
 	}
-	FORCEINLINE virtual void RHIDrawIndexedIndirect(FIndexBufferRHIParamRef IndexBufferRHI, FStructuredBufferRHIParamRef ArgumentsBufferRHI, int32 DrawArgumentsIndex, uint32 NumInstances) final override
+	FORCEINLINE virtual void RHIDrawIndexedIndirect(FRHIIndexBuffer* IndexBufferRHI, FRHIStructuredBuffer* ArgumentsBufferRHI, int32 DrawArgumentsIndex, uint32 NumInstances) final override
 	{
 		ContextRedirect(RHIDrawIndexedIndirect(IndexBufferRHI, ArgumentsBufferRHI, DrawArgumentsIndex, NumInstances));
 	}
-	FORCEINLINE virtual void RHIDrawIndexedPrimitive(FIndexBufferRHIParamRef IndexBuffer, int32 BaseVertexIndex, uint32 FirstInstance, uint32 NumVertices, uint32 StartIndex, uint32 NumPrimitives, uint32 NumInstances) final override
+	FORCEINLINE virtual void RHIDrawIndexedPrimitive(FRHIIndexBuffer* IndexBuffer, int32 BaseVertexIndex, uint32 FirstInstance, uint32 NumVertices, uint32 StartIndex, uint32 NumPrimitives, uint32 NumInstances) final override
 	{
 		ContextRedirect(RHIDrawIndexedPrimitive(IndexBuffer, BaseVertexIndex, FirstInstance, NumVertices, StartIndex, NumPrimitives, NumInstances));
 	}
-	FORCEINLINE virtual void RHIDrawIndexedPrimitiveIndirect(FIndexBufferRHIParamRef IndexBuffer, FVertexBufferRHIParamRef ArgumentBuffer, uint32 ArgumentOffset) final override
+	FORCEINLINE virtual void RHIDrawIndexedPrimitiveIndirect(FRHIIndexBuffer* IndexBuffer, FRHIVertexBuffer* ArgumentBuffer, uint32 ArgumentOffset) final override
 	{
 		ContextRedirect(RHIDrawIndexedPrimitiveIndirect(IndexBuffer, ArgumentBuffer, ArgumentOffset));
 	}
@@ -714,7 +714,7 @@ public:
 		ContextRedirect(RHISetDepthBounds(MinDepth, MaxDepth));
 	}
 
-	FORCEINLINE virtual void RHIUpdateTextureReference(FRHITextureReference* TextureRef, FTextureRHIParamRef NewTexture) final override
+	FORCEINLINE virtual void RHIUpdateTextureReference(FRHITextureReference* TextureRef, FRHITexture* NewTexture) final override
 	{
 		ContextRedirect(RHIUpdateTextureReference(TextureRef, NewTexture));
 	}
@@ -728,7 +728,7 @@ public:
 		ContextRedirect(RHIWaitForTemporalEffect(InEffectName));
 	}
 
-	FORCEINLINE virtual void RHIBroadcastTemporalEffect(const FName& InEffectName, FTextureRHIParamRef* InTextures, int32 NumTextures) final AFR_API_OVERRIDE
+	FORCEINLINE virtual void RHIBroadcastTemporalEffect(const FName& InEffectName, FRHITexture** InTextures, int32 NumTextures) final AFR_API_OVERRIDE
 	{
 		ContextRedirect(RHIBroadcastTemporalEffect(InEffectName, InTextures, NumTextures));
 	}
@@ -745,7 +745,7 @@ public:
 
 
 #if RHI_RAYTRACING
-	virtual void RHICopyBufferRegion(FVertexBufferRHIParamRef DestBuffer, uint64 DstOffset, FVertexBufferRHIParamRef SourceBuffer, uint64 SrcOffset, uint64 NumBytes) final override
+	virtual void RHICopyBufferRegion(FRHIVertexBuffer* DestBuffer, uint64 DstOffset, FRHIVertexBuffer* SourceBuffer, uint64 SrcOffset, uint64 NumBytes) final override
 	{
 		ContextRedirect(RHICopyBufferRegion(DestBuffer, DstOffset, SourceBuffer, SrcOffset, NumBytes));
 	}
@@ -802,7 +802,7 @@ public:
 	virtual void RHISetRayTracingHitGroup(
 		FRHIRayTracingScene* Scene, uint32 InstanceIndex, uint32 SegmentIndex, uint32 ShaderSlot,
 		FRHIRayTracingPipelineState* Pipeline, uint32 HitGroupIndex,
-		uint32 NumUniformBuffers, const FUniformBufferRHIParamRef* UniformBuffers,
+		uint32 NumUniformBuffers, FRHIUniformBuffer* const* UniformBuffers,
 		uint32 LooseParameterDataSize, const void* LooseParameterData,
 		uint32 UserData) final override
 	{
@@ -812,7 +812,7 @@ public:
 	virtual void RHISetRayTracingCallableShader(
 		FRHIRayTracingScene* Scene, uint32 ShaderSlotInScene,
 		FRHIRayTracingPipelineState* Pipeline, uint32 ShaderIndexInPipeline,
-		uint32 NumUniformBuffers, const FUniformBufferRHIParamRef* UniformBuffers,
+		uint32 NumUniformBuffers, FRHIUniformBuffer* const* UniformBuffers,
 		uint32 UserData) final override
 	{
 		ContextRedirect(RHISetRayTracingCallableShader(Scene, ShaderSlotInScene, Pipeline, ShaderIndexInPipeline, NumUniformBuffers, UniformBuffers, UserData));
diff --git a/Engine/Source/Runtime/D3D12RHI/Private/D3D12Commands.cpp b/Engine/Source/Runtime/D3D12RHI/Private/D3D12Commands.cpp
--- a/Engine/Source/Runtime/D3D12RHI/Private/D3D12Commands.cpp
+++ b/Engine/Source/Runtime/D3D12RHI/Private/D3D12Commands.cpp
@@ -86,7 +86,7 @@ void FD3D12DynamicRHI::SetupRecursiveResources()
 }
 
 // Vertex state.
-void FD3D12CommandContext::RHISetStreamSource(uint32 StreamIndex, FVertexBufferRHIParamRef VertexBufferRHI, uint32 Offset)
+void FD3D12CommandContext::RHISetStreamSource(uint32 StreamIndex, FRHIVertexBuffer* VertexBufferRHI, uint32 Offset)
 {
 	FD3D12VertexBuffer* VertexBuffer = RetrieveObject<FD3D12VertexBuffer>(VertexBufferRHI);
 
@@ -94,7 +94,7 @@ void FD3D12CommandContext::RHISetStreamSource(uint32 StreamIndex, FVertexBufferR
 }
 
 // Stream-Out state.
-void FD3D12DynamicRHI::RHISetStreamOutTargets(uint32 NumTargets, const FVertexBufferRHIParamRef* VertexBuffers, const uint32* Offsets)
+void FD3D12DynamicRHI::RHISetStreamOutTargets(uint32 NumTargets, FRHIVertexBuffer* const* VertexBuffers, const uint32* Offsets)
 {
 	// Multi-GPU support: this might need a node mask parameter or broadcast to all GPUs.
 	FD3D12CommandContext& CmdContext = GetRHIDevice()->GetDefaultCommandContext();
@@ -171,7 +171,7 @@ void FD3D12CommandContext::RHIDispatchComputeShader(uint32 ThreadGroupCountX, ui
 	DEBUG_EXECUTE_COMMAND_LIST(this);
 }
 
-void FD3D12CommandContext::RHIDispatchIndirectComputeShader(FVertexBufferRHIParamRef ArgumentBufferRHI, uint32 ArgumentOffset)
+void FD3D12CommandContext::RHIDispatchIndirectComputeShader(FRHIVertexBuffer* ArgumentBufferRHI, uint32 ArgumentOffset)
 {
 	FD3D12VertexBuffer* ArgumentBuffer = FD3D12DynamicRHI::ResourceCast(ArgumentBufferRHI);
 
@@ -209,7 +209,7 @@ void FD3D12CommandContext::RHIDispatchIndirectComputeShader(FVertexBufferRHIPara
 }
 
 
-void FD3D12CommandContext::RHITransitionResources(EResourceTransitionAccess TransitionType, FTextureRHIParamRef* InTextures, int32 NumTextures)
+void FD3D12CommandContext::RHITransitionResources(EResourceTransitionAccess TransitionType, FRHITexture** InTextures, int32 NumTextures)
 {
 #if !USE_D3D12RHI_RESOURCE_STATE_TRACKING
 	// TODO: Make sure that EMetaData is supported with an aliasing barrier, otherwise the CMask decal optimisation will break.
@@ -391,7 +391,7 @@ void FD3D12CommandContext::RHITransitionResources(EResourceTransitionAccess Tran
 }
 
 
-void FD3D12CommandContext::RHICopyToStagingBuffer(FVertexBufferRHIParamRef SourceBufferRHI, FRHIStagingBuffer* StagingBufferRHI, uint32 Offset, uint32 NumBytes)
+void FD3D12CommandContext::RHICopyToStagingBuffer(FRHIVertexBuffer* SourceBufferRHI, FRHIStagingBuffer* StagingBufferRHI, uint32 Offset, uint32 NumBytes)
 {
 	FD3D12StagingBuffer* StagingBuffer = FD3D12DynamicRHI::ResourceCast(StagingBufferRHI);
 	check(StagingBuffer);
@@ -588,42 +588,42 @@ void FD3D12CommandContext::RHISetComputePipelineState(FRHIComputePipelineState*
 	StateCache.SetComputePipelineState(ComputePipelineState);
 }
 
-void FD3D12CommandContext::RHISetShaderTexture(FRHIVertexShader* VertexShaderRHI, uint32 TextureIndex, FTextureRHIParamRef NewTextureRHI)
+void FD3D12CommandContext::RHISetShaderTexture(FRHIVertexShader* VertexShaderRHI, uint32 TextureIndex, FRHITexture* NewTextureRHI)
 {
 	VALIDATE_BOUND_SHADER(VertexShaderRHI);
 	FD3D12TextureBase* const NewTexture = RetrieveTextureBase(NewTextureRHI);
 	StateCache.SetShaderResourceView<SF_Vertex>(NewTexture ? NewTexture->GetShaderResourceView() : nullptr, TextureIndex);
 }
 
-void FD3D12CommandContext::RHISetShaderTexture(FRHIHullShader* HullShaderRHI, uint32 TextureIndex, FTextureRHIParamRef NewTextureRHI)
+void FD3D12CommandContext::RHISetShaderTexture(FRHIHullShader* HullShaderRHI, uint32 TextureIndex, FRHITexture* NewTextureRHI)
 {
 	VALIDATE_BOUND_SHADER(HullShaderRHI);
 	FD3D12TextureBase* const NewTexture = RetrieveTextureBase(NewTextureRHI);
 	StateCache.SetShaderResourceView<SF_Hull>(NewTexture ? NewTexture->GetShaderResourceView() : nullptr, TextureIndex);
 }
 
-void FD3D12CommandContext::RHISetShaderTexture(FRHIDomainShader* DomainShaderRHI, uint32 TextureIndex, FTextureRHIParamRef NewTextureRHI)
+void FD3D12CommandContext::RHISetShaderTexture(FRHIDomainShader* DomainShaderRHI, uint32 TextureIndex, FRHITexture* NewTextureRHI)
 {
 	VALIDATE_BOUND_SHADER(DomainShaderRHI);
 	FD3D12TextureBase* const NewTexture = RetrieveTextureBase(NewTextureRHI);
 	StateCache.SetShaderResourceView<SF_Domain>(NewTexture ? NewTexture->GetShaderResourceView() : nullptr, TextureIndex);
 }
 
-void FD3D12CommandContext::RHISetShaderTexture(FRHIGeometryShader* GeometryShaderRHI, uint32 TextureIndex, FTextureRHIParamRef NewTextureRHI)
+void FD3D12CommandContext::RHISetShaderTexture(FRHIGeometryShader* GeometryShaderRHI, uint32 TextureIndex, FRHITexture* NewTextureRHI)
 {
 	VALIDATE_BOUND_SHADER(GeometryShaderRHI);
 	FD3D12TextureBase* const NewTexture = RetrieveTextureBase(NewTextureRHI);
 	StateCache.SetShaderResourceView<SF_Geometry>(NewTexture ? NewTexture->GetShaderResourceView() : nullptr, TextureIndex);
 }
 
-void FD3D12CommandContext::RHISetShaderTexture(FRHIPixelShader* PixelShaderRHI, uint32 TextureIndex, FTextureRHIParamRef NewTextureRHI)
+void FD3D12CommandContext::RHISetShaderTexture(FRHIPixelShader* PixelShaderRHI, uint32 TextureIndex, FRHITexture* NewTextureRHI)
 {
 	VALIDATE_BOUND_SHADER(PixelShaderRHI);
 	FD3D12TextureBase* const NewTexture = RetrieveTextureBase(NewTextureRHI);
 	StateCache.SetShaderResourceView<SF_Pixel>(NewTexture ? NewTexture->GetShaderResourceView() : nullptr, TextureIndex);
 }
 
-void FD3D12CommandContext::RHISetShaderTexture(FRHIComputeShader* ComputeShaderRHI, uint32 TextureIndex, FTextureRHIParamRef NewTextureRHI)
+void FD3D12CommandContext::RHISetShaderTexture(FRHIComputeShader* ComputeShaderRHI, uint32 TextureIndex, FRHITexture* NewTextureRHI)
 {
 	//VALIDATE_BOUND_SHADER(ComputeShaderRHI);
 	FD3D12TextureBase* const NewTexture = RetrieveTextureBase(NewTextureRHI);
@@ -745,7 +745,7 @@ void FD3D12CommandContext::RHISetShaderSampler(FRHIComputeShader* ComputeShaderR
 	StateCache.SetSamplerState<SF_Compute>(NewState, SamplerIndex);
 }
 
-void FD3D12CommandContext::RHISetShaderUniformBuffer(FRHIVertexShader* VertexShader, uint32 BufferIndex, FUniformBufferRHIParamRef BufferRHI)
+void FD3D12CommandContext::RHISetShaderUniformBuffer(FRHIVertexShader* VertexShader, uint32 BufferIndex, FRHIUniformBuffer* BufferRHI)
 {
 	//SCOPE_CYCLE_COUNTER(STAT_D3D12SetShaderUniformBuffer);
 	VALIDATE_BOUND_SHADER(VertexShader);
@@ -762,7 +762,7 @@ void FD3D12CommandContext::RHISetShaderUniformBuffer(FRHIVertexShader* VertexSha
 	DirtyUniformBuffers[SF_Vertex] |= (1 << BufferIndex);
 }
 
-void FD3D12CommandContext::RHISetShaderUniformBuffer(FRHIHullShader* HullShader, uint32 BufferIndex, FUniformBufferRHIParamRef BufferRHI)
+void FD3D12CommandContext::RHISetShaderUniformBuffer(FRHIHullShader* HullShader, uint32 BufferIndex, FRHIUniformBuffer* BufferRHI)
 {
 	//SCOPE_CYCLE_COUNTER(STAT_D3D12SetShaderUniformBuffer);
 	VALIDATE_BOUND_SHADER(HullShader);
@@ -779,7 +779,7 @@ void FD3D12CommandContext::RHISetShaderUniformBuffer(FRHIHullShader* HullShader,
 	DirtyUniformBuffers[SF_Hull] |= (1 << BufferIndex);
 }
 
-void FD3D12CommandContext::RHISetShaderUniformBuffer(FRHIDomainShader* DomainShader, uint32 BufferIndex, FUniformBufferRHIParamRef BufferRHI)
+void FD3D12CommandContext::RHISetShaderUniformBuffer(FRHIDomainShader* DomainShader, uint32 BufferIndex, FRHIUniformBuffer* BufferRHI)
 {
 	//SCOPE_CYCLE_COUNTER(STAT_D3D12SetShaderUniformBuffer);
 	VALIDATE_BOUND_SHADER(DomainShader);
@@ -796,7 +796,7 @@ void FD3D12CommandContext::RHISetShaderUniformBuffer(FRHIDomainShader* DomainSha
 	DirtyUniformBuffers[SF_Domain] |= (1 << BufferIndex);
 }
 
-void FD3D12CommandContext::RHISetShaderUniformBuffer(FRHIGeometryShader* GeometryShader, uint32 BufferIndex, FUniformBufferRHIParamRef BufferRHI)
+void FD3D12CommandContext::RHISetShaderUniformBuffer(FRHIGeometryShader* GeometryShader, uint32 BufferIndex, FRHIUniformBuffer* BufferRHI)
 {
 	//SCOPE_CYCLE_COUNTER(STAT_D3D12SetShaderUniformBuffer);
 	VALIDATE_BOUND_SHADER(GeometryShader);
@@ -813,7 +813,7 @@ void FD3D12CommandContext::RHISetShaderUniformBuffer(FRHIGeometryShader* Geometr
 	DirtyUniformBuffers[SF_Geometry] |= (1 << BufferIndex);
 }
 
-void FD3D12CommandContext::RHISetShaderUniformBuffer(FRHIPixelShader* PixelShader, uint32 BufferIndex, FUniformBufferRHIParamRef BufferRHI)
+void FD3D12CommandContext::RHISetShaderUniformBuffer(FRHIPixelShader* PixelShader, uint32 BufferIndex, FRHIUniformBuffer* BufferRHI)
 {
 	//SCOPE_CYCLE_COUNTER(STAT_D3D12SetShaderUniformBuffer);
 	VALIDATE_BOUND_SHADER(PixelShader);
@@ -830,7 +830,7 @@ void FD3D12CommandContext::RHISetShaderUniformBuffer(FRHIPixelShader* PixelShade
 	DirtyUniformBuffers[SF_Pixel] |= (1 << BufferIndex);
 }
 
-void FD3D12CommandContext::RHISetShaderUniformBuffer(FRHIComputeShader* ComputeShader, uint32 BufferIndex, FUniformBufferRHIParamRef BufferRHI)
+void FD3D12CommandContext::RHISetShaderUniformBuffer(FRHIComputeShader* ComputeShader, uint32 BufferIndex, FRHIUniformBuffer* BufferRHI)
 {
 	//SCOPE_CYCLE_COUNTER(STAT_D3D12SetShaderUniformBuffer);
 	//VALIDATE_BOUND_SHADER(ComputeShader);
@@ -1446,7 +1446,7 @@ void FD3D12CommandContext::RHIDrawPrimitive(uint32 BaseVertexIndex, uint32 NumPr
 
 }
 
-void FD3D12CommandContext::RHIDrawPrimitiveIndirect(FVertexBufferRHIParamRef ArgumentBufferRHI, uint32 ArgumentOffset)
+void FD3D12CommandContext::RHIDrawPrimitiveIndirect(FRHIVertexBuffer* ArgumentBufferRHI, uint32 ArgumentOffset)
 {
 	FD3D12VertexBuffer* ArgumentBuffer = RetrieveObject<FD3D12VertexBuffer>(ArgumentBufferRHI);
 
@@ -1486,7 +1486,7 @@ void FD3D12CommandContext::RHIDrawPrimitiveIndirect(FVertexBufferRHIParamRef Arg
 
 }
 
-void FD3D12CommandContext::RHIDrawIndexedIndirect(FIndexBufferRHIParamRef IndexBufferRHI, FStructuredBufferRHIParamRef ArgumentsBufferRHI, int32 DrawArgumentsIndex, uint32 NumInstances)
+void FD3D12CommandContext::RHIDrawIndexedIndirect(FRHIIndexBuffer* IndexBufferRHI, FRHIStructuredBuffer* ArgumentsBufferRHI, int32 DrawArgumentsIndex, uint32 NumInstances)
 {
 	FD3D12IndexBuffer* IndexBuffer = RetrieveObject<FD3D12IndexBuffer>(IndexBufferRHI);
 	FD3D12StructuredBuffer* ArgumentsBuffer = RetrieveObject<FD3D12StructuredBuffer>(ArgumentsBufferRHI);
@@ -1531,7 +1531,7 @@ void FD3D12CommandContext::RHIDrawIndexedIndirect(FIndexBufferRHIParamRef IndexB
 	DEBUG_EXECUTE_COMMAND_LIST(this);
 }
 
-void FD3D12CommandContext::RHIDrawIndexedPrimitive(FIndexBufferRHIParamRef IndexBufferRHI, int32 BaseVertexIndex, uint32 FirstInstance, uint32 NumVertices, uint32 StartIndex, uint32 NumPrimitives, uint32 NumInstances)
+void FD3D12CommandContext::RHIDrawIndexedPrimitive(FRHIIndexBuffer* IndexBufferRHI, int32 BaseVertexIndex, uint32 FirstInstance, uint32 NumVertices, uint32 StartIndex, uint32 NumPrimitives, uint32 NumInstances)
 {
 	FD3D12IndexBuffer* IndexBuffer = RetrieveObject<FD3D12IndexBuffer>(IndexBufferRHI);
 
@@ -1572,7 +1572,7 @@ void FD3D12CommandContext::RHIDrawIndexedPrimitive(FIndexBufferRHIParamRef Index
 
 }
 
-void FD3D12CommandContext::RHIDrawIndexedPrimitiveIndirect(FIndexBufferRHIParamRef IndexBufferRHI, FVertexBufferRHIParamRef ArgumentBufferRHI, uint32 ArgumentOffset)
+void FD3D12CommandContext::RHIDrawIndexedPrimitiveIndirect(FRHIIndexBuffer* IndexBufferRHI, FRHIVertexBuffer* ArgumentBufferRHI, uint32 ArgumentOffset)
 {
 	FD3D12IndexBuffer* IndexBuffer = RetrieveObject<FD3D12IndexBuffer>(IndexBufferRHI);
 	FD3D12VertexBuffer* ArgumentBuffer = RetrieveObject<FD3D12VertexBuffer>(ArgumentBufferRHI);
@@ -1870,7 +1870,7 @@ void FD3D12CommandContext::RHIWaitForTemporalEffect(const FName& InEffectName)
 #endif
 }
 
-void FD3D12CommandContext::RHIBroadcastTemporalEffect(const FName& InEffectName, FTextureRHIParamRef* InTextures, int32 NumTextures)
+void FD3D12CommandContext::RHIBroadcastTemporalEffect(const FName& InEffectName, FRHITexture** InTextures, int32 NumTextures)
 {
 	check(IsDefaultContext());
 	FD3D12Device* Device = GetParentDevice();
diff --git a/Engine/Source/Runtime/D3D12RHI/Private/D3D12IndexBuffer.cpp b/Engine/Source/Runtime/D3D12RHI/Private/D3D12IndexBuffer.cpp
--- a/Engine/Source/Runtime/D3D12RHI/Private/D3D12IndexBuffer.cpp
+++ b/Engine/Source/Runtime/D3D12RHI/Private/D3D12IndexBuffer.cpp
@@ -104,17 +104,17 @@ FIndexBufferRHIRef FD3D12DynamicRHI::RHICreateIndexBuffer(uint32 Stride, uint32
 	return Buffer;
 }
 
-void* FD3D12DynamicRHI::RHILockIndexBuffer(FIndexBufferRHIParamRef IndexBufferRHI, uint32 Offset, uint32 Size, EResourceLockMode LockMode)
+void* FD3D12DynamicRHI::RHILockIndexBuffer(FRHIIndexBuffer* IndexBufferRHI, uint32 Offset, uint32 Size, EResourceLockMode LockMode)
 {
 	return LockBuffer(nullptr, FD3D12DynamicRHI::ResourceCast(IndexBufferRHI), Offset, Size, LockMode);
 }
 
-void FD3D12DynamicRHI::RHIUnlockIndexBuffer(FIndexBufferRHIParamRef IndexBufferRHI)
+void FD3D12DynamicRHI::RHIUnlockIndexBuffer(FRHIIndexBuffer* IndexBufferRHI)
 {
 	UnlockBuffer(nullptr, FD3D12DynamicRHI::ResourceCast(IndexBufferRHI));
 }
 
-void FD3D12DynamicRHI::RHITransferIndexBufferUnderlyingResource(FIndexBufferRHIParamRef DestIndexBuffer, FIndexBufferRHIParamRef SrcIndexBuffer)
+void FD3D12DynamicRHI::RHITransferIndexBufferUnderlyingResource(FRHIIndexBuffer* DestIndexBuffer, FRHIIndexBuffer* SrcIndexBuffer)
 {
 	check(DestIndexBuffer);
 	FD3D12IndexBuffer* Dest = ResourceCast(DestIndexBuffer);
@@ -149,12 +149,12 @@ FIndexBufferRHIRef FD3D12DynamicRHI::CreateIndexBuffer_RenderThread(class FRHICo
 	return Buffer;
 }
 
-void* FD3D12DynamicRHI::LockIndexBuffer_RenderThread(class FRHICommandListImmediate& RHICmdList, FIndexBufferRHIParamRef IndexBufferRHI, uint32 Offset, uint32 SizeRHI, EResourceLockMode LockMode)
+void* FD3D12DynamicRHI::LockIndexBuffer_RenderThread(class FRHICommandListImmediate& RHICmdList, FRHIIndexBuffer* IndexBufferRHI, uint32 Offset, uint32 SizeRHI, EResourceLockMode LockMode)
 {
 	return LockBuffer(&RHICmdList, FD3D12DynamicRHI::ResourceCast(IndexBufferRHI), Offset, SizeRHI, LockMode);
 }
 
-void FD3D12DynamicRHI::UnlockIndexBuffer_RenderThread(class FRHICommandListImmediate& RHICmdList, FIndexBufferRHIParamRef IndexBufferRHI)
+void FD3D12DynamicRHI::UnlockIndexBuffer_RenderThread(class FRHICommandListImmediate& RHICmdList, FRHIIndexBuffer* IndexBufferRHI)
 {
 	// Pull down the above RHI implementation so that we can flush only when absolutely necessary
 	QUICK_SCOPE_CYCLE_COUNTER(STAT_FDynamicRHI_UnlockIndexBuffer_RenderThread);
diff --git a/Engine/Source/Runtime/D3D12RHI/Private/D3D12RHIPrivate.h b/Engine/Source/Runtime/D3D12RHI/Private/D3D12RHIPrivate.h
--- a/Engine/Source/Runtime/D3D12RHI/Private/D3D12RHIPrivate.h
+++ b/Engine/Source/Runtime/D3D12RHI/Private/D3D12RHIPrivate.h
@@ -281,28 +281,28 @@ public:
 	virtual FGraphicsPipelineStateRHIRef RHICreateGraphicsPipelineState(const FGraphicsPipelineStateInitializer& Initializer) final override;
 	virtual TRefCountPtr<FRHIComputePipelineState> RHICreateComputePipelineState(FRHIComputeShader* ComputeShader) final override;
 	virtual FUniformBufferRHIRef RHICreateUniformBuffer(const void* Contents, const FRHIUniformBufferLayout& Layout, EUniformBufferUsage Usage, EUniformBufferValidation Validation) final override;
-	virtual void RHIUpdateUniformBuffer(FUniformBufferRHIParamRef UniformBufferRHI, const void* Contents) final override;
+	virtual void RHIUpdateUniformBuffer(FRHIUniformBuffer* UniformBufferRHI, const void* Contents) final override;
 	virtual FIndexBufferRHIRef RHICreateIndexBuffer(uint32 Stride, uint32 Size, uint32 InUsage, FRHIResourceCreateInfo& CreateInfo) final override;
-	virtual void* RHILockIndexBuffer(FIndexBufferRHIParamRef IndexBuffer, uint32 Offset, uint32 Size, EResourceLockMode LockMode) final override;
-	virtual void RHIUnlockIndexBuffer(FIndexBufferRHIParamRef IndexBuffer) final override;
-	virtual void RHITransferIndexBufferUnderlyingResource(FIndexBufferRHIParamRef DestIndexBuffer, FIndexBufferRHIParamRef SrcIndexBuffer) final override;
+	virtual void* RHILockIndexBuffer(FRHIIndexBuffer* IndexBuffer, uint32 Offset, uint32 Size, EResourceLockMode LockMode) final override;
+	virtual void RHIUnlockIndexBuffer(FRHIIndexBuffer* IndexBuffer) final override;
+	virtual void RHITransferIndexBufferUnderlyingResource(FRHIIndexBuffer* DestIndexBuffer, FRHIIndexBuffer* SrcIndexBuffer) final override;
 	virtual FVertexBufferRHIRef RHICreateVertexBuffer(uint32 Size, uint32 InUsage, FRHIResourceCreateInfo& CreateInfo) final override;
-	virtual void* RHILockVertexBuffer(FVertexBufferRHIParamRef VertexBuffer, uint32 Offset, uint32 SizeRHI, EResourceLockMode LockMode) final override;
-	virtual void RHIUnlockVertexBuffer(FVertexBufferRHIParamRef VertexBuffer) final override;
-	virtual void RHICopyVertexBuffer(FVertexBufferRHIParamRef SourceBuffer, FVertexBufferRHIParamRef DestBuffer) final override;
-	virtual void RHITransferVertexBufferUnderlyingResource(FVertexBufferRHIParamRef DestVertexBuffer, FVertexBufferRHIParamRef SrcVertexBuffer) final override;
+	virtual void* RHILockVertexBuffer(FRHIVertexBuffer* VertexBuffer, uint32 Offset, uint32 SizeRHI, EResourceLockMode LockMode) final override;
+	virtual void RHIUnlockVertexBuffer(FRHIVertexBuffer* VertexBuffer) final override;
+	virtual void RHICopyVertexBuffer(FRHIVertexBuffer* SourceBuffer, FRHIVertexBuffer* DestBuffer) final override;
+	virtual void RHITransferVertexBufferUnderlyingResource(FRHIVertexBuffer* DestVertexBuffer, FRHIVertexBuffer* SrcVertexBuffer) final override;
 	virtual FStructuredBufferRHIRef RHICreateStructuredBuffer(uint32 Stride, uint32 Size, uint32 InUsage, FRHIResourceCreateInfo& CreateInfo) final override;
-	virtual void* RHILockStructuredBuffer(FStructuredBufferRHIParamRef StructuredBuffer, uint32 Offset, uint32 SizeRHI, EResourceLockMode LockMode) final override;
-	virtual void RHIUnlockStructuredBuffer(FStructuredBufferRHIParamRef StructuredBuffer) final override;
-	virtual FUnorderedAccessViewRHIRef RHICreateUnorderedAccessView(FStructuredBufferRHIParamRef StructuredBuffer, bool bUseUAVCounter, bool bAppendBuffer) final override;
-	virtual FUnorderedAccessViewRHIRef RHICreateUnorderedAccessView(FTextureRHIParamRef Texture, uint32 MipLevel) final override;
-	virtual FUnorderedAccessViewRHIRef RHICreateUnorderedAccessView(FVertexBufferRHIParamRef VertexBuffer, uint8 Format) final override;
-	virtual FUnorderedAccessViewRHIRef RHICreateUnorderedAccessView(FIndexBufferRHIParamRef IndexBuffer, uint8 Format) final override;
-	virtual FShaderResourceViewRHIRef RHICreateShaderResourceView(FStructuredBufferRHIParamRef StructuredBuffer) final override;
-	virtual FShaderResourceViewRHIRef RHICreateShaderResourceView(FVertexBufferRHIParamRef VertexBuffer, uint32 Stride, uint8 Format) final override;
-	virtual FShaderResourceViewRHIRef RHICreateShaderResourceView(FIndexBufferRHIParamRef Buffer) final override;
-	virtual void RHIUpdateShaderResourceView(FRHIShaderResourceView* SRV, FVertexBufferRHIParamRef VertexBuffer, uint32 Stride, uint8 Format) final override;
-	virtual void RHIUpdateShaderResourceView(FRHIShaderResourceView* SRV, FIndexBufferRHIParamRef IndexBuffer) final override;
+	virtual void* RHILockStructuredBuffer(FRHIStructuredBuffer* StructuredBuffer, uint32 Offset, uint32 SizeRHI, EResourceLockMode LockMode) final override;
+	virtual void RHIUnlockStructuredBuffer(FRHIStructuredBuffer* StructuredBuffer) final override;
+	virtual FUnorderedAccessViewRHIRef RHICreateUnorderedAccessView(FRHIStructuredBuffer* StructuredBuffer, bool bUseUAVCounter, bool bAppendBuffer) final override;
+	virtual FUnorderedAccessViewRHIRef RHICreateUnorderedAccessView(FRHITexture* Texture, uint32 MipLevel) final override;
+	virtual FUnorderedAccessViewRHIRef RHICreateUnorderedAccessView(FRHIVertexBuffer* VertexBuffer, uint8 Format) final override;
+	virtual FUnorderedAccessViewRHIRef RHICreateUnorderedAccessView(FRHIIndexBuffer* IndexBuffer, uint8 Format) final override;
+	virtual FShaderResourceViewRHIRef RHICreateShaderResourceView(FRHIStructuredBuffer* StructuredBuffer) final override;
+	virtual FShaderResourceViewRHIRef RHICreateShaderResourceView(FRHIVertexBuffer* VertexBuffer, uint32 Stride, uint8 Format) final override;
+	virtual FShaderResourceViewRHIRef RHICreateShaderResourceView(FRHIIndexBuffer* Buffer) final override;
+	virtual void RHIUpdateShaderResourceView(FRHIShaderResourceView* SRV, FRHIVertexBuffer* VertexBuffer, uint32 Stride, uint8 Format) final override;
+	virtual void RHIUpdateShaderResourceView(FRHIShaderResourceView* SRV, FRHIIndexBuffer* IndexBuffer) final override;
 	virtual uint64 RHICalcTexture2DPlatformSize(uint32 SizeX, uint32 SizeY, uint8 Format, uint32 NumMips, uint32 NumSamples, uint32 Flags, uint32& OutAlign) override;
 	virtual uint64 RHICalcTexture3DPlatformSize(uint32 SizeX, uint32 SizeY, uint32 SizeZ, uint8 Format, uint32 NumMips, uint32 Flags, uint32& OutAlign) final override;
 	virtual uint64 RHICalcTextureCubePlatformSize(uint32 Size, uint8 Format, uint32 NumMips, uint32 Flags, uint32& OutAlign) final override;
@@ -311,41 +311,41 @@ public:
 	virtual FTextureReferenceRHIRef RHICreateTextureReference(FLastRenderTimeContainer* LastRenderTime) final override;
 	virtual FTexture2DRHIRef RHICreateTexture2D(uint32 SizeX, uint32 SizeY, uint8 Format, uint32 NumMips, uint32 NumSamples, uint32 Flags, FRHIResourceCreateInfo& CreateInfo) override;
 	virtual FTexture2DRHIRef RHIAsyncCreateTexture2D(uint32 SizeX, uint32 SizeY, uint8 Format, uint32 NumMips, uint32 Flags, void** InitialMipData, uint32 NumInitialMips) final override;
-	virtual void RHICopySharedMips(FTexture2DRHIParamRef DestTexture2D, FTexture2DRHIParamRef SrcTexture2D) final override;
+	virtual void RHICopySharedMips(FRHITexture2D* DestTexture2D, FRHITexture2D* SrcTexture2D) final override;
 	void RHIMultiGPULockstep(FRHIGPUMask GPUMask);
-	virtual void RHITransferTexture(FTexture2DRHIParamRef Texture, FIntRect Rect, uint32 SrcGPUIndex, uint32 DestGPUIndex, bool PullData) final override;
+	virtual void RHITransferTexture(FRHITexture2D* Texture, FIntRect Rect, uint32 SrcGPUIndex, uint32 DestGPUIndex, bool PullData) final override;
 	virtual FTexture2DArrayRHIRef RHICreateTexture2DArray(uint32 SizeX, uint32 SizeY, uint32 SizeZ, uint8 Format, uint32 NumMips, uint32 NumSamples, uint32 Flags, FRHIResourceCreateInfo& CreateInfo) override;
 	virtual FTexture3DRHIRef RHICreateTexture3D(uint32 SizeX, uint32 SizeY, uint32 SizeZ, uint8 Format, uint32 NumMips, uint32 Flags, FRHIResourceCreateInfo& CreateInfo) override;
-	virtual void RHIGetResourceInfo(FTextureRHIParamRef Ref, FRHIResourceInfo& OutInfo) override;
-	virtual FShaderResourceViewRHIRef RHICreateShaderResourceView(FTextureRHIParamRef Texture, const FRHITextureSRVCreateInfo& CreateInfo) final override;
-	virtual uint32 RHIComputeMemorySize(FTextureRHIParamRef TextureRHI) final override;
-	virtual FTexture2DRHIRef RHIAsyncReallocateTexture2D(FTexture2DRHIParamRef Texture2D, int32 NewMipCount, int32 NewSizeX, int32 NewSizeY, FThreadSafeCounter* RequestStatus) override;
-	virtual ETextureReallocationStatus RHIFinalizeAsyncReallocateTexture2D(FTexture2DRHIParamRef Texture2D, bool bBlockUntilCompleted) final override;
-	virtual ETextureReallocationStatus RHICancelAsyncReallocateTexture2D(FTexture2DRHIParamRef Texture2D, bool bBlockUntilCompleted) final override;
-	virtual void* RHILockTexture2D(FTexture2DRHIParamRef Texture, uint32 MipIndex, EResourceLockMode LockMode, uint32& DestStride, bool bLockWithinMiptail) final override;
-	virtual void RHIUnlockTexture2D(FTexture2DRHIParamRef Texture, uint32 MipIndex, bool bLockWithinMiptail) final override;
-	virtual void* RHILockTexture2DArray(FTexture2DArrayRHIParamRef Texture, uint32 TextureIndex, uint32 MipIndex, EResourceLockMode LockMode, uint32& DestStride, bool bLockWithinMiptail) final override;
-	virtual void RHIUnlockTexture2DArray(FTexture2DArrayRHIParamRef Texture, uint32 TextureIndex, uint32 MipIndex, bool bLockWithinMiptail) final override;
-	virtual void RHIUpdateTexture2D(FTexture2DRHIParamRef Texture, uint32 MipIndex, const struct FUpdateTextureRegion2D& UpdateRegion, uint32 SourcePitch, const uint8* SourceData) final override;
-	virtual void RHIUpdateTexture3D(FTexture3DRHIParamRef Texture, uint32 MipIndex, const struct FUpdateTextureRegion3D& UpdateRegion, uint32 SourceRowPitch, uint32 SourceDepthPitch, const uint8* SourceData) final override;
-	virtual FUpdateTexture3DData BeginUpdateTexture3D_RenderThread(class FRHICommandListImmediate& RHICmdList, FTexture3DRHIParamRef Texture, uint32 MipIndex, const struct FUpdateTextureRegion3D& UpdateRegion) final override;
+	virtual void RHIGetResourceInfo(FRHITexture* Ref, FRHIResourceInfo& OutInfo) override;
+	virtual FShaderResourceViewRHIRef RHICreateShaderResourceView(FRHITexture* Texture, const FRHITextureSRVCreateInfo& CreateInfo) final override;
+	virtual uint32 RHIComputeMemorySize(FRHITexture* TextureRHI) final override;
+	virtual FTexture2DRHIRef RHIAsyncReallocateTexture2D(FRHITexture2D* Texture2D, int32 NewMipCount, int32 NewSizeX, int32 NewSizeY, FThreadSafeCounter* RequestStatus) override;
+	virtual ETextureReallocationStatus RHIFinalizeAsyncReallocateTexture2D(FRHITexture2D* Texture2D, bool bBlockUntilCompleted) final override;
+	virtual ETextureReallocationStatus RHICancelAsyncReallocateTexture2D(FRHITexture2D* Texture2D, bool bBlockUntilCompleted) final override;
+	virtual void* RHILockTexture2D(FRHITexture2D* Texture, uint32 MipIndex, EResourceLockMode LockMode, uint32& DestStride, bool bLockWithinMiptail) final override;
+	virtual void RHIUnlockTexture2D(FRHITexture2D* Texture, uint32 MipIndex, bool bLockWithinMiptail) final override;
+	virtual void* RHILockTexture2DArray(FRHITexture2DArray* Texture, uint32 TextureIndex, uint32 MipIndex, EResourceLockMode LockMode, uint32& DestStride, bool bLockWithinMiptail) final override;
+	virtual void RHIUnlockTexture2DArray(FRHITexture2DArray* Texture, uint32 TextureIndex, uint32 MipIndex, bool bLockWithinMiptail) final override;
+	virtual void RHIUpdateTexture2D(FRHITexture2D* Texture, uint32 MipIndex, const struct FUpdateTextureRegion2D& UpdateRegion, uint32 SourcePitch, const uint8* SourceData) final override;
+	virtual void RHIUpdateTexture3D(FRHITexture3D* Texture, uint32 MipIndex, const struct FUpdateTextureRegion3D& UpdateRegion, uint32 SourceRowPitch, uint32 SourceDepthPitch, const uint8* SourceData) final override;
+	virtual FUpdateTexture3DData BeginUpdateTexture3D_RenderThread(class FRHICommandListImmediate& RHICmdList, FRHITexture3D* Texture, uint32 MipIndex, const struct FUpdateTextureRegion3D& UpdateRegion) final override;
 	virtual void EndUpdateTexture3D_RenderThread(class FRHICommandListImmediate& RHICmdList, FUpdateTexture3DData& UpdateData) final override;
 	virtual void EndMultiUpdateTexture3D_RenderThread(class FRHICommandListImmediate& RHICmdList, TArray<FUpdateTexture3DData>& UpdateDataArray) final override;
 	virtual FTextureCubeRHIRef RHICreateTextureCube(uint32 Size, uint8 Format, uint32 NumMips, uint32 Flags, FRHIResourceCreateInfo& CreateInfo) override;
 	virtual FTextureCubeRHIRef RHICreateTextureCubeArray(uint32 Size, uint32 ArraySize, uint8 Format, uint32 NumMips, uint32 Flags, FRHIResourceCreateInfo& CreateInfo) override;
-	virtual void* RHILockTextureCubeFace(FTextureCubeRHIParamRef Texture, uint32 FaceIndex, uint32 ArrayIndex, uint32 MipIndex, EResourceLockMode LockMode, uint32& DestStride, bool bLockWithinMiptail) final override;
-	virtual void RHIUnlockTextureCubeFace(FTextureCubeRHIParamRef Texture, uint32 FaceIndex, uint32 ArrayIndex, uint32 MipIndex, bool bLockWithinMiptail) final override;
-	virtual void RHIBindDebugLabelName(FTextureRHIParamRef Texture, const TCHAR* Name) final override;
-	virtual void RHIReadSurfaceData(FTextureRHIParamRef Texture, FIntRect Rect, TArray<FColor>& OutData, FReadSurfaceDataFlags InFlags) final override;
-	virtual void RHIMapStagingSurface(FTextureRHIParamRef Texture, void*& OutData, int32& OutWidth, int32& OutHeight) final override;
-	virtual void RHIUnmapStagingSurface(FTextureRHIParamRef Texture) final override;
-	virtual void RHIReadSurfaceFloatData(FTextureRHIParamRef Texture, FIntRect Rect, TArray<FFloat16Color>& OutData, ECubeFace CubeFace, int32 ArrayIndex, int32 MipIndex) final override;
-	virtual void RHIRead3DSurfaceFloatData(FTextureRHIParamRef Texture, FIntRect Rect, FIntPoint ZMinMax, TArray<FFloat16Color>& OutData) final override;
+	virtual void* RHILockTextureCubeFace(FRHITextureCube* Texture, uint32 FaceIndex, uint32 ArrayIndex, uint32 MipIndex, EResourceLockMode LockMode, uint32& DestStride, bool bLockWithinMiptail) final override;
+	virtual void RHIUnlockTextureCubeFace(FRHITextureCube* Texture, uint32 FaceIndex, uint32 ArrayIndex, uint32 MipIndex, bool bLockWithinMiptail) final override;
+	virtual void RHIBindDebugLabelName(FRHITexture* Texture, const TCHAR* Name) final override;
+	virtual void RHIReadSurfaceData(FRHITexture* Texture, FIntRect Rect, TArray<FColor>& OutData, FReadSurfaceDataFlags InFlags) final override;
+	virtual void RHIMapStagingSurface(FRHITexture* Texture, void*& OutData, int32& OutWidth, int32& OutHeight) final override;
+	virtual void RHIUnmapStagingSurface(FRHITexture* Texture) final override;
+	virtual void RHIReadSurfaceFloatData(FRHITexture* Texture, FIntRect Rect, TArray<FFloat16Color>& OutData, ECubeFace CubeFace, int32 ArrayIndex, int32 MipIndex) final override;
+	virtual void RHIRead3DSurfaceFloatData(FRHITexture* Texture, FIntRect Rect, FIntPoint ZMinMax, TArray<FFloat16Color>& OutData) final override;
 	virtual FRenderQueryRHIRef RHICreateRenderQuery(ERenderQueryType QueryType) final override;
 	virtual bool RHIGetRenderQueryResult(FRHIRenderQuery* RenderQuery, uint64& OutResult, bool bWait) final override;
 	virtual uint32 RHIGetViewportNextPresentGPUIndex(FRHIViewport* Viewport) final override;
 	virtual FTexture2DRHIRef RHIGetViewportBackBuffer(FRHIViewport* Viewport) final override;
-	virtual void RHIAliasTextureResources(FTextureRHIParamRef DestTexture, FTextureRHIParamRef SrcTexture) final override;
+	virtual void RHIAliasTextureResources(FRHITexture* DestTexture, FRHITexture* SrcTexture) final override;
 	virtual void RHIAdvanceFrameForGetViewportBackBuffer(FRHIViewport* Viewport) final override;
 	virtual void RHIAcquireThreadOwnership() final override;
 	virtual void RHIReleaseThreadOwnership() final override;
@@ -355,13 +355,13 @@ public:
 	virtual void RHIResizeViewport(FRHIViewport* Viewport, uint32 SizeX, uint32 SizeY, bool bIsFullscreen) final override;
 	virtual void RHIResizeViewport(FRHIViewport* ViewportRHI, uint32 SizeX, uint32 SizeY, bool bIsFullscreen, EPixelFormat PreferredPixelFormat) final override;
 	virtual void RHITick(float DeltaTime) final override;
-	virtual void RHISetStreamOutTargets(uint32 NumTargets, const FVertexBufferRHIParamRef* VertexBuffers, const uint32* Offsets) final override;
+	virtual void RHISetStreamOutTargets(uint32 NumTargets, FRHIVertexBuffer* const* VertexBuffers, const uint32* Offsets) final override;
 	virtual void RHIBlockUntilGPUIdle() final override;
 	virtual void RHISubmitCommandsAndFlushGPU() final override;
 	virtual bool RHIGetAvailableResolutions(FScreenResolutionArray& Resolutions, bool bIgnoreRefreshRate) final override;
 	virtual void RHIGetSupportedResolution(uint32& Width, uint32& Height) final override;
-	virtual void RHIVirtualTextureSetFirstMipInMemory(FTexture2DRHIParamRef Texture, uint32 FirstMip) override;
-	virtual void RHIVirtualTextureSetFirstMipVisible(FTexture2DRHIParamRef Texture, uint32 FirstMip) override;
+	virtual void RHIVirtualTextureSetFirstMipInMemory(FRHITexture2D* Texture, uint32 FirstMip) override;
+	virtual void RHIVirtualTextureSetFirstMipVisible(FRHITexture2D* Texture, uint32 FirstMip) override;
 	virtual void RHIExecuteCommandList(FRHICommandList* CmdList) final override;
 	virtual void* RHIGetNativeDevice() final override;
 	virtual class IRHICommandContext* RHIGetDefaultContext() final override;
@@ -422,30 +422,30 @@ public:
 
 	virtual FVertexBufferRHIRef CreateVertexBuffer_RenderThread(class FRHICommandListImmediate& RHICmdList, uint32 Size, uint32 InUsage, FRHIResourceCreateInfo& CreateInfo);
 	virtual FStructuredBufferRHIRef CreateStructuredBuffer_RenderThread(class FRHICommandListImmediate& RHICmdList, uint32 Stride, uint32 Size, uint32 InUsage, FRHIResourceCreateInfo& CreateInfo);
-	virtual void* LockVertexBuffer_RenderThread(class FRHICommandListImmediate& RHICmdList, FVertexBufferRHIParamRef VertexBuffer, uint32 Offset, uint32 SizeRHI, EResourceLockMode LockMode);
-	virtual void UnlockVertexBuffer_RenderThread(class FRHICommandListImmediate& RHICmdList, FVertexBufferRHIParamRef VertexBuffer);
+	virtual void* LockVertexBuffer_RenderThread(class FRHICommandListImmediate& RHICmdList, FRHIVertexBuffer* VertexBuffer, uint32 Offset, uint32 SizeRHI, EResourceLockMode LockMode);
+	virtual void UnlockVertexBuffer_RenderThread(class FRHICommandListImmediate& RHICmdList, FRHIVertexBuffer* VertexBuffer);
 	virtual FVertexBufferRHIRef CreateAndLockVertexBuffer_RenderThread(class FRHICommandListImmediate& RHICmdList, uint32 Size, uint32 InUsage, FRHIResourceCreateInfo& CreateInfo, void*& OutDataBuffer);
 	virtual FIndexBufferRHIRef CreateAndLockIndexBuffer_RenderThread(class FRHICommandListImmediate& RHICmdList, uint32 Stride, uint32 Size, uint32 InUsage, FRHIResourceCreateInfo& CreateInfo, void*& OutDataBuffer);
-	// virtual FTexture2DRHIRef AsyncReallocateTexture2D_RenderThread(class FRHICommandListImmediate& RHICmdList, FTexture2DRHIParamRef Texture2D, int32 NewMipCount, int32 NewSizeX, int32 NewSizeY, FThreadSafeCounter* RequestStatus);
-	// virtual ETextureReallocationStatus FinalizeAsyncReallocateTexture2D_RenderThread(class FRHICommandListImmediate& RHICmdList, FTexture2DRHIParamRef Texture2D, bool bBlockUntilCompleted);
-	// virtual ETextureReallocationStatus CancelAsyncReallocateTexture2D_RenderThread(class FRHICommandListImmediate& RHICmdList, FTexture2DRHIParamRef Texture2D, bool bBlockUntilCompleted);
+	// virtual FTexture2DRHIRef AsyncReallocateTexture2D_RenderThread(class FRHICommandListImmediate& RHICmdList, FRHITexture2D* Texture2D, int32 NewMipCount, int32 NewSizeX, int32 NewSizeY, FThreadSafeCounter* RequestStatus);
+	// virtual ETextureReallocationStatus FinalizeAsyncReallocateTexture2D_RenderThread(class FRHICommandListImmediate& RHICmdList, FRHITexture2D* Texture2D, bool bBlockUntilCompleted);
+	// virtual ETextureReallocationStatus CancelAsyncReallocateTexture2D_RenderThread(class FRHICommandListImmediate& RHICmdList, FRHITexture2D* Texture2D, bool bBlockUntilCompleted);
 	virtual FIndexBufferRHIRef CreateIndexBuffer_RenderThread(class FRHICommandListImmediate& RHICmdList, uint32 Stride, uint32 Size, uint32 InUsage, FRHIResourceCreateInfo& CreateInfo);
-	virtual void* LockIndexBuffer_RenderThread(class FRHICommandListImmediate& RHICmdList, FIndexBufferRHIParamRef IndexBuffer, uint32 Offset, uint32 SizeRHI, EResourceLockMode LockMode);
-	virtual void UnlockIndexBuffer_RenderThread(class FRHICommandListImmediate& RHICmdList, FIndexBufferRHIParamRef IndexBuffer);
-	virtual void UpdateTexture2D_RenderThread(class FRHICommandListImmediate& RHICmdList, FTexture2DRHIParamRef Texture, uint32 MipIndex, const struct FUpdateTextureRegion2D& UpdateRegion, uint32 SourcePitch, const uint8* SourceData);
-	virtual void* LockTexture2D_RenderThread(class FRHICommandListImmediate& RHICmdList, FTexture2DRHIParamRef Texture, uint32 MipIndex, EResourceLockMode LockMode, uint32& DestStride, bool bLockWithinMiptail, bool bNeedsDefaultRHIFlush = true);
-	virtual void UnlockTexture2D_RenderThread(class FRHICommandListImmediate& RHICmdList, FTexture2DRHIParamRef Texture, uint32 MipIndex, bool bLockWithinMiptail, bool bNeedsDefaultRHIFlush = true);
+	virtual void* LockIndexBuffer_RenderThread(class FRHICommandListImmediate& RHICmdList, FRHIIndexBuffer* IndexBuffer, uint32 Offset, uint32 SizeRHI, EResourceLockMode LockMode);
+	virtual void UnlockIndexBuffer_RenderThread(class FRHICommandListImmediate& RHICmdList, FRHIIndexBuffer* IndexBuffer);
+	virtual void UpdateTexture2D_RenderThread(class FRHICommandListImmediate& RHICmdList, FRHITexture2D* Texture, uint32 MipIndex, const struct FUpdateTextureRegion2D& UpdateRegion, uint32 SourcePitch, const uint8* SourceData);
+	virtual void* LockTexture2D_RenderThread(class FRHICommandListImmediate& RHICmdList, FRHITexture2D* Texture, uint32 MipIndex, EResourceLockMode LockMode, uint32& DestStride, bool bLockWithinMiptail, bool bNeedsDefaultRHIFlush = true);
+	virtual void UnlockTexture2D_RenderThread(class FRHICommandListImmediate& RHICmdList, FRHITexture2D* Texture, uint32 MipIndex, bool bLockWithinMiptail, bool bNeedsDefaultRHIFlush = true);
 	// 
 	virtual FTexture2DRHIRef RHICreateTexture2D_RenderThread(class FRHICommandListImmediate& RHICmdList, uint32 SizeX, uint32 SizeY, uint8 Format, uint32 NumMips, uint32 NumSamples, uint32 Flags, FRHIResourceCreateInfo& CreateInfo);
 	virtual FTexture2DArrayRHIRef RHICreateTexture2DArray_RenderThread(class FRHICommandListImmediate& RHICmdList, uint32 SizeX, uint32 SizeY, uint32 SizeZ, uint8 Format, uint32 NumMips, uint32 NumSamples, uint32 Flags, FRHIResourceCreateInfo& CreateInfo);
 	virtual FTexture3DRHIRef RHICreateTexture3D_RenderThread(class FRHICommandListImmediate& RHICmdList, uint32 SizeX, uint32 SizeY, uint32 SizeZ, uint8 Format, uint32 NumMips, uint32 Flags, FRHIResourceCreateInfo& CreateInfo);
-	virtual FUnorderedAccessViewRHIRef RHICreateUnorderedAccessView_RenderThread(class FRHICommandListImmediate& RHICmdList, FStructuredBufferRHIParamRef StructuredBuffer, bool bUseUAVCounter, bool bAppendBuffer);
-	virtual FUnorderedAccessViewRHIRef RHICreateUnorderedAccessView_RenderThread(class FRHICommandListImmediate& RHICmdList, FTextureRHIParamRef Texture, uint32 MipLevel);
-	virtual FUnorderedAccessViewRHIRef RHICreateUnorderedAccessView_RenderThread(class FRHICommandListImmediate& RHICmdList, FVertexBufferRHIParamRef VertexBuffer, uint8 Format);
-	virtual FShaderResourceViewRHIRef RHICreateShaderResourceView_RenderThread(class FRHICommandListImmediate& RHICmdList, FTextureRHIParamRef Texture, const FRHITextureSRVCreateInfo& CreateInfo);
-	virtual FShaderResourceViewRHIRef RHICreateShaderResourceView_RenderThread(class FRHICommandListImmediate& RHICmdList, FVertexBufferRHIParamRef VertexBuffer, uint32 Stride, uint8 Format);
-	virtual FShaderResourceViewRHIRef CreateShaderResourceView_RenderThread(class FRHICommandListImmediate& RHICmdList, FVertexBufferRHIParamRef VertexBuffer, uint32 Stride, uint8 Format);
-	virtual FShaderResourceViewRHIRef RHICreateShaderResourceView_RenderThread(class FRHICommandListImmediate& RHICmdList, FStructuredBufferRHIParamRef StructuredBuffer);
+	virtual FUnorderedAccessViewRHIRef RHICreateUnorderedAccessView_RenderThread(class FRHICommandListImmediate& RHICmdList, FRHIStructuredBuffer* StructuredBuffer, bool bUseUAVCounter, bool bAppendBuffer);
+	virtual FUnorderedAccessViewRHIRef RHICreateUnorderedAccessView_RenderThread(class FRHICommandListImmediate& RHICmdList, FRHITexture* Texture, uint32 MipLevel);
+	virtual FUnorderedAccessViewRHIRef RHICreateUnorderedAccessView_RenderThread(class FRHICommandListImmediate& RHICmdList, FRHIVertexBuffer* VertexBuffer, uint8 Format);
+	virtual FShaderResourceViewRHIRef RHICreateShaderResourceView_RenderThread(class FRHICommandListImmediate& RHICmdList, FRHITexture* Texture, const FRHITextureSRVCreateInfo& CreateInfo);
+	virtual FShaderResourceViewRHIRef RHICreateShaderResourceView_RenderThread(class FRHICommandListImmediate& RHICmdList, FRHIVertexBuffer* VertexBuffer, uint32 Stride, uint8 Format);
+	virtual FShaderResourceViewRHIRef CreateShaderResourceView_RenderThread(class FRHICommandListImmediate& RHICmdList, FRHIVertexBuffer* VertexBuffer, uint32 Stride, uint8 Format);
+	virtual FShaderResourceViewRHIRef RHICreateShaderResourceView_RenderThread(class FRHICommandListImmediate& RHICmdList, FRHIStructuredBuffer* StructuredBuffer);
 	virtual FTextureCubeRHIRef RHICreateTextureCube_RenderThread(class FRHICommandListImmediate& RHICmdList, uint32 Size, uint8 Format, uint32 NumMips, uint32 Flags, FRHIResourceCreateInfo& CreateInfo);
 	virtual FTextureCubeRHIRef RHICreateTextureCubeArray_RenderThread(class FRHICommandListImmediate& RHICmdList, uint32 Size, uint32 ArraySize, uint8 Format, uint32 NumMips, uint32 Flags, FRHIResourceCreateInfo& CreateInfo);
 
@@ -454,7 +454,7 @@ public:
 		return RHICreateRenderQuery(QueryType);
 	}
 
-	virtual void RHICopySubTextureRegion(FTexture2DRHIParamRef SourceTexture, FTexture2DRHIParamRef DestinationTexture, FBox2D SourceBox, FBox2D DestinationBox) final override;
+	virtual void RHICopySubTextureRegion(FRHITexture2D* SourceTexture, FRHITexture2D* DestinationTexture, FBox2D SourceBox, FBox2D DestinationBox) final override;
 
 	void RHICalibrateTimers() override;
 
@@ -505,7 +505,7 @@ public:
 		// Not supported on PC
 	}
 
-	FUpdateTexture3DData BeginUpdateTexture3D_Internal(FTexture3DRHIParamRef Texture, uint32 MipIndex, const struct FUpdateTextureRegion3D& UpdateRegion);
+	FUpdateTexture3DData BeginUpdateTexture3D_Internal(FRHITexture3D* Texture, uint32 MipIndex, const struct FUpdateTextureRegion3D& UpdateRegion);
 	void EndUpdateTexture3D_Internal(FUpdateTexture3DData& UpdateData);
 
 	void UpdateBuffer(FD3D12Resource* Dest, uint32 DestOffset, FD3D12Resource* Source, uint32 SourceOffset, uint32 NumBytes);
@@ -980,11 +980,11 @@ protected:
 	* @param StagingRectOUT - parameter is filled with the rectangle to read from the returned texture.
 	* @return The CPU readable Texture object.
 	*/
-	TRefCountPtr<FD3D12Resource> GetStagingTexture(FTextureRHIParamRef TextureRHI, FIntRect InRect, FIntRect& OutRect, FReadSurfaceDataFlags InFlags, D3D12_PLACED_SUBRESOURCE_FOOTPRINT &readBackHeapDesc);
+	TRefCountPtr<FD3D12Resource> GetStagingTexture(FRHITexture* TextureRHI, FIntRect InRect, FIntRect& OutRect, FReadSurfaceDataFlags InFlags, D3D12_PLACED_SUBRESOURCE_FOOTPRINT &readBackHeapDesc);
 
-	void ReadSurfaceDataNoMSAARaw(FTextureRHIParamRef TextureRHI, FIntRect Rect, TArray<uint8>& OutData, FReadSurfaceDataFlags InFlags);
+	void ReadSurfaceDataNoMSAARaw(FRHITexture* TextureRHI, FIntRect Rect, TArray<uint8>& OutData, FReadSurfaceDataFlags InFlags);
 
-	void ReadSurfaceDataMSAARaw(FRHICommandList_RecursiveHazardous& RHICmdList, FTextureRHIParamRef TextureRHI, FIntRect Rect, TArray<uint8>& OutData, FReadSurfaceDataFlags InFlags);
+	void ReadSurfaceDataMSAARaw(FRHICommandList_RecursiveHazardous& RHICmdList, FRHITexture* TextureRHI, FIntRect Rect, TArray<uint8>& OutData, FReadSurfaceDataFlags InFlags);
 
 	void SetupRecursiveResources();
 
diff --git a/Engine/Source/Runtime/D3D12RHI/Private/D3D12RayTracing.cpp b/Engine/Source/Runtime/D3D12RHI/Private/D3D12RayTracing.cpp
--- a/Engine/Source/Runtime/D3D12RHI/Private/D3D12RayTracing.cpp
+++ b/Engine/Source/Runtime/D3D12RHI/Private/D3D12RayTracing.cpp
@@ -2641,9 +2641,9 @@ template <typename ResourceBinderType>
 static void SetRayTracingShaderResources(
 	FD3D12CommandContext& CommandContext,
 	const FD3D12RayTracingShader* Shader,
-	uint32 InNumTextures, const FTextureRHIParamRef* Textures,
+	uint32 InNumTextures, FRHITexture* const* Textures,
 	uint32 InNumSRVs, FRHIShaderResourceView* const* SRVs,
-	uint32 InNumUniformBuffers, const FUniformBufferRHIParamRef* UniformBuffers,
+	uint32 InNumUniformBuffers, FRHIUniformBuffer* const* UniformBuffers,
 	uint32 InNumSamplers, FRHISamplerState* const* Samplers,
 	uint32 InNumUAVs, FRHIUnorderedAccessView* const* UAVs,
 	uint32 InLooseParameterDataSize, const void* InLooseParameterData,
@@ -2673,7 +2673,7 @@ static void SetRayTracingShaderResources(
 
 	for (uint32 SRVIndex = 0; SRVIndex < InNumTextures; ++SRVIndex)
 	{
-		FTextureRHIParamRef Resource = Textures[SRVIndex];
+		FRHITexture* Resource = Textures[SRVIndex];
 		if (Resource)
 		{
 			FD3D12TextureBase* Texture = CommandContext.RetrieveTextureBase(Resource);
@@ -2699,7 +2699,7 @@ static void SetRayTracingShaderResources(
 
 	for (uint32 CBVIndex = 0; CBVIndex < InNumUniformBuffers; ++CBVIndex)
 	{
-		FUniformBufferRHIParamRef Resource = UniformBuffers[CBVIndex];
+		FRHIUniformBuffer* Resource = UniformBuffers[CBVIndex];
 		if (Resource)
 		{
 			FD3D12UniformBuffer* CBV = CommandContext.RetrieveObject<FD3D12UniformBuffer>(Resource);
@@ -3113,7 +3113,7 @@ void FD3D12CommandContext::RHIRayTraceDispatch(FRHIRayTracingPipelineState* InRa
 void FD3D12CommandContext::RHISetRayTracingHitGroup(
 	FRHIRayTracingScene* InScene, uint32 InstanceIndex, uint32 SegmentIndex, uint32 ShaderSlot,
 	FRHIRayTracingPipelineState* InPipeline, uint32 HitGroupIndex,
-	uint32 NumUniformBuffers, const FUniformBufferRHIParamRef* UniformBuffers,
+	uint32 NumUniformBuffers, FRHIUniformBuffer* const* UniformBuffers,
 	uint32 LooseParameterDataSize, const void* LooseParameterData,
 	uint32 UserData)
 {
@@ -3150,7 +3150,7 @@ void FD3D12CommandContext::RHISetRayTracingHitGroup(
 void FD3D12CommandContext::RHISetRayTracingCallableShader(
 	FRHIRayTracingScene* InScene, uint32 ShaderSlotInScene,
 	FRHIRayTracingPipelineState* InPipeline, uint32 ShaderIndexInPipeline,
-	uint32 NumUniformBuffers, const FUniformBufferRHIParamRef* UniformBuffers,
+	uint32 NumUniformBuffers, FRHIUniformBuffer* const* UniformBuffers,
 	uint32 UserData)
 {
 	FD3D12RayTracingScene* Scene = FD3D12DynamicRHI::ResourceCast(InScene);
diff --git a/Engine/Source/Runtime/D3D12RHI/Private/D3D12RenderTarget.cpp b/Engine/Source/Runtime/D3D12RHI/Private/D3D12RenderTarget.cpp
--- a/Engine/Source/Runtime/D3D12RHI/Private/D3D12RenderTarget.cpp
+++ b/Engine/Source/Runtime/D3D12RHI/Private/D3D12RenderTarget.cpp
@@ -232,7 +232,7 @@ void FD3D12CommandContext::ResolveTextureUsingShader(
 * @param SourceSurface - surface with a resolve texture to copy to
 * @param ResolveParams - optional resolve params
 */
-void FD3D12CommandContext::RHICopyToResolveTarget(FTextureRHIParamRef SourceTextureRHI, FTextureRHIParamRef DestTextureRHI, const FResolveParams& ResolveParams)
+void FD3D12CommandContext::RHICopyToResolveTarget(FRHITexture* SourceTextureRHI, FRHITexture* DestTextureRHI, const FResolveParams& ResolveParams)
 {
 	if (!SourceTextureRHI || !DestTextureRHI)
 	{
@@ -528,7 +528,7 @@ void FD3D12DynamicRHI::RHIMultiGPULockstep(FRHIGPUMask GPUMask)
 	}
 }
 
-void FD3D12DynamicRHI::RHITransferTexture(FTexture2DRHIParamRef TextureRHI, FIntRect Rect, uint32 SrcGPUIndex, uint32 DestGPUIndex, bool PullData)
+void FD3D12DynamicRHI::RHITransferTexture(FRHITexture2D* TextureRHI, FIntRect Rect, uint32 SrcGPUIndex, uint32 DestGPUIndex, bool PullData)
 {
 	FD3D12Adapter& Adapter = GetAdapter();
 
@@ -661,7 +661,7 @@ static uint32 ComputeBytesPerPixel(DXGI_FORMAT Format)
 	return BytesPerPixel;
 }
 
-TRefCountPtr<FD3D12Resource> FD3D12DynamicRHI::GetStagingTexture(FTextureRHIParamRef TextureRHI, FIntRect InRect, FIntRect& StagingRectOUT, FReadSurfaceDataFlags InFlags, D3D12_PLACED_SUBRESOURCE_FOOTPRINT &readbackHeapDesc)
+TRefCountPtr<FD3D12Resource> FD3D12DynamicRHI::GetStagingTexture(FRHITexture* TextureRHI, FIntRect InRect, FIntRect& StagingRectOUT, FReadSurfaceDataFlags InFlags, D3D12_PLACED_SUBRESOURCE_FOOTPRINT &readbackHeapDesc)
 {
 	FD3D12Device* Device = GetRHIDevice();
 	FD3D12Adapter* Adapter = Device->GetParentAdapter();
@@ -771,7 +771,7 @@ TRefCountPtr<FD3D12Resource> FD3D12DynamicRHI::GetStagingTexture(FTextureRHIPara
 	return TempTexture2D;
 }
 
-void FD3D12DynamicRHI::ReadSurfaceDataNoMSAARaw(FTextureRHIParamRef TextureRHI, FIntRect InRect, TArray<uint8>& OutData, FReadSurfaceDataFlags InFlags)
+void FD3D12DynamicRHI::ReadSurfaceDataNoMSAARaw(FRHITexture* TextureRHI, FIntRect InRect, TArray<uint8>& OutData, FReadSurfaceDataFlags InFlags)
 {
 	FD3D12TextureBase* Texture = GetD3D12TextureFromRHITexture(TextureRHI);
 
@@ -1108,7 +1108,7 @@ static void ConvertRAWSurfaceDataToFColor(DXGI_FORMAT Format, uint32 Width, uint
 	}
 }
 
-void FD3D12DynamicRHI::RHIReadSurfaceData(FTextureRHIParamRef TextureRHI, FIntRect InRect, TArray<FColor>& OutData, FReadSurfaceDataFlags InFlags)
+void FD3D12DynamicRHI::RHIReadSurfaceData(FRHITexture* TextureRHI, FIntRect InRect, TArray<FColor>& OutData, FReadSurfaceDataFlags InFlags)
 {
 	if (!ensure(TextureRHI))
 	{
@@ -1167,7 +1167,7 @@ void FD3D12DynamicRHI::RHIReadSurfaceData(FTextureRHIParamRef TextureRHI, FIntRe
 	ConvertRAWSurfaceDataToFColor((DXGI_FORMAT)FormatInfo.PlatformFormat, SizeX, SizeY, OutDataRaw.GetData(), SrcPitch, OutData.GetData(), InFlags);
 }
 
-void FD3D12DynamicRHI::ReadSurfaceDataMSAARaw(FRHICommandList_RecursiveHazardous& RHICmdList, FTextureRHIParamRef TextureRHI, FIntRect InRect, TArray<uint8>& OutData, FReadSurfaceDataFlags InFlags)
+void FD3D12DynamicRHI::ReadSurfaceDataMSAARaw(FRHICommandList_RecursiveHazardous& RHICmdList, FRHITexture* TextureRHI, FIntRect InRect, TArray<uint8>& OutData, FReadSurfaceDataFlags InFlags)
 {
 	FD3D12Device* Device = GetRHIDevice();
 	FD3D12Adapter* Adapter = Device->GetParentAdapter();
@@ -1328,7 +1328,7 @@ void FD3D12DynamicRHI::ReadSurfaceDataMSAARaw(FRHICommandList_RecursiveHazardous
 	}
 }
 
-void FD3D12DynamicRHI::RHIMapStagingSurface(FTextureRHIParamRef TextureRHI, void*& OutData, int32& OutWidth, int32& OutHeight)
+void FD3D12DynamicRHI::RHIMapStagingSurface(FRHITexture* TextureRHI, void*& OutData, int32& OutWidth, int32& OutHeight)
 {
 	FD3D12Resource* Texture = GetD3D12TextureFromRHITexture(TextureRHI)->GetResource();
 
@@ -1380,14 +1380,14 @@ void FD3D12DynamicRHI::RHIMapStagingSurface(FTextureRHIParamRef TextureRHI, void
 	}
 }
 
-void FD3D12DynamicRHI::RHIUnmapStagingSurface(FTextureRHIParamRef TextureRHI)
+void FD3D12DynamicRHI::RHIUnmapStagingSurface(FRHITexture* TextureRHI)
 {
 	ID3D12Resource* Texture = GetD3D12TextureFromRHITexture(TextureRHI)->GetResource()->GetResource();
 
 	Texture->Unmap(0, nullptr);
 }
 
-void FD3D12DynamicRHI::RHIReadSurfaceFloatData(FTextureRHIParamRef TextureRHI, FIntRect InRect, TArray<FFloat16Color>& OutData, ECubeFace CubeFace, int32 ArrayIndex, int32 MipIndex)
+void FD3D12DynamicRHI::RHIReadSurfaceFloatData(FRHITexture* TextureRHI, FIntRect InRect, TArray<FFloat16Color>& OutData, ECubeFace CubeFace, int32 ArrayIndex, int32 MipIndex)
 {
 	FD3D12Device* Device = GetRHIDevice();
 	FD3D12Adapter* Adapter = Device->GetParentAdapter();
@@ -1515,7 +1515,7 @@ void FD3D12DynamicRHI::RHIReadSurfaceFloatData(FTextureRHIParamRef TextureRHI, F
 	TempTexture2D->GetResource()->Unmap(0, nullptr);
 }
 
-void FD3D12DynamicRHI::RHIRead3DSurfaceFloatData(FTextureRHIParamRef TextureRHI, FIntRect InRect, FIntPoint ZMinMax, TArray<FFloat16Color>& OutData)
+void FD3D12DynamicRHI::RHIRead3DSurfaceFloatData(FRHITexture* TextureRHI, FIntRect InRect, FIntPoint ZMinMax, TArray<FFloat16Color>& OutData)
 {
 	FD3D12Device* Device = GetRHIDevice();
 	FD3D12Adapter* Adapter = Device->GetParentAdapter();
diff --git a/Engine/Source/Runtime/D3D12RHI/Private/D3D12StructuredBuffer.cpp b/Engine/Source/Runtime/D3D12RHI/Private/D3D12StructuredBuffer.cpp
--- a/Engine/Source/Runtime/D3D12RHI/Private/D3D12StructuredBuffer.cpp
+++ b/Engine/Source/Runtime/D3D12RHI/Private/D3D12StructuredBuffer.cpp
@@ -109,12 +109,12 @@ void FD3D12StructuredBuffer::RenameLDAChain(FD3D12ResourceLocation& NewLocation)
 	}
 }
 
-void* FD3D12DynamicRHI::RHILockStructuredBuffer(FStructuredBufferRHIParamRef StructuredBufferRHI, uint32 Offset, uint32 Size, EResourceLockMode LockMode)
+void* FD3D12DynamicRHI::RHILockStructuredBuffer(FRHIStructuredBuffer* StructuredBufferRHI, uint32 Offset, uint32 Size, EResourceLockMode LockMode)
 {
 	return LockBuffer(nullptr, FD3D12DynamicRHI::ResourceCast(StructuredBufferRHI), Offset, Size, LockMode);
 }
 
-void FD3D12DynamicRHI::RHIUnlockStructuredBuffer(FStructuredBufferRHIParamRef StructuredBufferRHI)
+void FD3D12DynamicRHI::RHIUnlockStructuredBuffer(FRHIStructuredBuffer* StructuredBufferRHI)
 {
 	UnlockBuffer(nullptr, FD3D12DynamicRHI::ResourceCast(StructuredBufferRHI));
 }
diff --git a/Engine/Source/Runtime/D3D12RHI/Private/D3D12Texture.cpp b/Engine/Source/Runtime/D3D12RHI/Private/D3D12Texture.cpp
--- a/Engine/Source/Runtime/D3D12RHI/Private/D3D12Texture.cpp
+++ b/Engine/Source/Runtime/D3D12RHI/Private/D3D12Texture.cpp
@@ -1307,7 +1307,7 @@ FTexture2DRHIRef FD3D12DynamicRHI::RHIAsyncCreateTexture2D(uint32 SizeX, uint32
 	return TextureOut;
 }
 
-void FD3D12DynamicRHI::RHICopySharedMips(FTexture2DRHIParamRef DestTexture2DRHI, FTexture2DRHIParamRef SrcTexture2DRHI)
+void FD3D12DynamicRHI::RHICopySharedMips(FRHITexture2D* DestTexture2DRHI, FRHITexture2D* SrcTexture2DRHI)
 {
 	FD3D12Texture2D*  DestTexture2D = FD3D12DynamicRHI::ResourceCast(DestTexture2DRHI);
 	FD3D12Texture2D*  SrcTexture2D = FD3D12DynamicRHI::ResourceCast(SrcTexture2DRHI);
@@ -1389,7 +1389,7 @@ FTexture3DRHIRef FD3D12DynamicRHI::RHICreateTexture3D(uint32 SizeX, uint32 SizeY
 #endif // PLATFORM_WINDOWS || PLATFORM_HOLOLENS
 }
 
-void FD3D12DynamicRHI::RHIGetResourceInfo(FTextureRHIParamRef Ref, FRHIResourceInfo& OutInfo)
+void FD3D12DynamicRHI::RHIGetResourceInfo(FRHITexture* Ref, FRHIResourceInfo& OutInfo)
 {
 	if (Ref)
 	{
@@ -1403,7 +1403,7 @@ void FD3D12DynamicRHI::RHIGetResourceInfo(FTextureRHIParamRef Ref, FRHIResourceI
  * @param	TextureRHI		- Texture we want to know the size of
  * @return					- Size in Bytes
  */
-uint32 FD3D12DynamicRHI::RHIComputeMemorySize(FTextureRHIParamRef TextureRHI)
+uint32 FD3D12DynamicRHI::RHIComputeMemorySize(FRHITexture* TextureRHI)
 {
 	if (!TextureRHI)
 	{
@@ -1429,7 +1429,7 @@ uint32 FD3D12DynamicRHI::RHIComputeMemorySize(FTextureRHIParamRef TextureRHI)
  * @param RequestStatus	- Will be decremented by 1 when the reallocation is complete (success or failure).
  * @return				- New reference to the texture, or an invalid reference upon failure
  */
-FTexture2DRHIRef FD3D12DynamicRHI::RHIAsyncReallocateTexture2D(FTexture2DRHIParamRef Texture2DRHI, int32 NewMipCount, int32 NewSizeX, int32 NewSizeY, FThreadSafeCounter* RequestStatus)
+FTexture2DRHIRef FD3D12DynamicRHI::RHIAsyncReallocateTexture2D(FRHITexture2D* Texture2DRHI, int32 NewMipCount, int32 NewSizeX, int32 NewSizeY, FThreadSafeCounter* RequestStatus)
 {
 	FD3D12Texture2D*  Texture2D = FD3D12DynamicRHI::ResourceCast(Texture2DRHI);
 
@@ -1500,7 +1500,7 @@ FTexture2DRHIRef FD3D12DynamicRHI::RHIAsyncReallocateTexture2D(FTexture2DRHIPara
  * @param Texture2D		- Texture to check the reallocation status for
  * @return				- Current reallocation status
  */
-ETextureReallocationStatus FD3D12DynamicRHI::RHIFinalizeAsyncReallocateTexture2D(FTexture2DRHIParamRef Texture2D, bool bBlockUntilCompleted)
+ETextureReallocationStatus FD3D12DynamicRHI::RHIFinalizeAsyncReallocateTexture2D(FRHITexture2D* Texture2D, bool bBlockUntilCompleted)
 {
 	return TexRealloc_Succeeded;
 }
@@ -1513,7 +1513,7 @@ ETextureReallocationStatus FD3D12DynamicRHI::RHIFinalizeAsyncReallocateTexture2D
  * @param bBlockUntilCompleted	If true, blocks until the cancellation is fully completed
  * @return						Reallocation status
  */
-ETextureReallocationStatus FD3D12DynamicRHI::RHICancelAsyncReallocateTexture2D(FTexture2DRHIParamRef Texture2D, bool bBlockUntilCompleted)
+ETextureReallocationStatus FD3D12DynamicRHI::RHICancelAsyncReallocateTexture2D(FRHITexture2D* Texture2D, bool bBlockUntilCompleted)
 {
 	return TexRealloc_Succeeded;
 }
@@ -1910,7 +1910,7 @@ void TD3D12Texture2D<RHIResourceType>::UpdateTexture2D(class FRHICommandListImme
 	}
 }
 
-void* FD3D12DynamicRHI::LockTexture2D_RenderThread(class FRHICommandListImmediate& RHICmdList, FTexture2DRHIParamRef TextureRHI, uint32 MipIndex, EResourceLockMode LockMode, uint32& DestStride, bool bLockWithinMiptail, bool bNeedsDefaultRHIFlush)
+void* FD3D12DynamicRHI::LockTexture2D_RenderThread(class FRHICommandListImmediate& RHICmdList, FRHITexture2D* TextureRHI, uint32 MipIndex, EResourceLockMode LockMode, uint32& DestStride, bool bLockWithinMiptail, bool bNeedsDefaultRHIFlush)
 {
 	if (CVarD3D12Texture2DRHIFlush.GetValueOnRenderThread() && bNeedsDefaultRHIFlush)
 	{
@@ -1924,14 +1924,14 @@ void* FD3D12DynamicRHI::LockTexture2D_RenderThread(class FRHICommandListImmediat
 	return Texture->Lock(&RHICmdList, MipIndex, 0, LockMode, DestStride);
 }
 
-void* FD3D12DynamicRHI::RHILockTexture2D(FTexture2DRHIParamRef TextureRHI, uint32 MipIndex, EResourceLockMode LockMode, uint32& DestStride, bool bLockWithinMiptail)
+void* FD3D12DynamicRHI::RHILockTexture2D(FRHITexture2D* TextureRHI, uint32 MipIndex, EResourceLockMode LockMode, uint32& DestStride, bool bLockWithinMiptail)
 {
 	check(TextureRHI);
 	FD3D12Texture2D*  Texture = FD3D12DynamicRHI::ResourceCast(TextureRHI);
 	return Texture->Lock(nullptr, MipIndex, 0, LockMode, DestStride);
 }
 
-void FD3D12DynamicRHI::UnlockTexture2D_RenderThread(class FRHICommandListImmediate& RHICmdList, FTexture2DRHIParamRef TextureRHI, uint32 MipIndex, bool bLockWithinMiptail, bool bNeedsDefaultRHIFlush)
+void FD3D12DynamicRHI::UnlockTexture2D_RenderThread(class FRHICommandListImmediate& RHICmdList, FRHITexture2D* TextureRHI, uint32 MipIndex, bool bLockWithinMiptail, bool bNeedsDefaultRHIFlush)
 {
 	if (CVarD3D12Texture2DRHIFlush.GetValueOnRenderThread() && bNeedsDefaultRHIFlush)
 	{
@@ -1946,42 +1946,42 @@ void FD3D12DynamicRHI::UnlockTexture2D_RenderThread(class FRHICommandListImmedia
 	Texture->Unlock(&RHICmdList, MipIndex, 0);
 }
 
-void FD3D12DynamicRHI::RHIUnlockTexture2D(FTexture2DRHIParamRef TextureRHI, uint32 MipIndex, bool bLockWithinMiptail)
+void FD3D12DynamicRHI::RHIUnlockTexture2D(FRHITexture2D* TextureRHI, uint32 MipIndex, bool bLockWithinMiptail)
 {
 	check(TextureRHI);
 	FD3D12Texture2D*  Texture = FD3D12DynamicRHI::ResourceCast(TextureRHI);
 	Texture->Unlock(nullptr, MipIndex, 0);
 }
 
-void* FD3D12DynamicRHI::RHILockTexture2DArray(FTexture2DArrayRHIParamRef TextureRHI, uint32 TextureIndex, uint32 MipIndex, EResourceLockMode LockMode, uint32& DestStride, bool bLockWithinMiptail)
+void* FD3D12DynamicRHI::RHILockTexture2DArray(FRHITexture2DArray* TextureRHI, uint32 TextureIndex, uint32 MipIndex, EResourceLockMode LockMode, uint32& DestStride, bool bLockWithinMiptail)
 {
 	check(TextureRHI);
 	FD3D12Texture2DArray*  Texture = FD3D12DynamicRHI::ResourceCast(TextureRHI);
 	return Texture->Lock(nullptr, MipIndex, TextureIndex, LockMode, DestStride);
 }
 
-void FD3D12DynamicRHI::RHIUnlockTexture2DArray(FTexture2DArrayRHIParamRef TextureRHI, uint32 TextureIndex, uint32 MipIndex, bool bLockWithinMiptail)
+void FD3D12DynamicRHI::RHIUnlockTexture2DArray(FRHITexture2DArray* TextureRHI, uint32 TextureIndex, uint32 MipIndex, bool bLockWithinMiptail)
 {
 	check(TextureRHI);
 	FD3D12Texture2DArray*  Texture = FD3D12DynamicRHI::ResourceCast(TextureRHI);
 	Texture->Unlock(nullptr, MipIndex, TextureIndex);
 }
 
-void FD3D12DynamicRHI::UpdateTexture2D_RenderThread(class FRHICommandListImmediate& RHICmdList, FTexture2DRHIParamRef TextureRHI, uint32 MipIndex, const struct FUpdateTextureRegion2D& UpdateRegion, uint32 SourcePitch, const uint8* SourceData)
+void FD3D12DynamicRHI::UpdateTexture2D_RenderThread(class FRHICommandListImmediate& RHICmdList, FRHITexture2D* TextureRHI, uint32 MipIndex, const struct FUpdateTextureRegion2D& UpdateRegion, uint32 SourcePitch, const uint8* SourceData)
 {
 	check(TextureRHI);
 	FD3D12Texture2D* Texture = FD3D12DynamicRHI::ResourceCast(TextureRHI);
 	Texture->UpdateTexture2D(&RHICmdList, MipIndex, UpdateRegion, SourcePitch, SourceData);
 }
 
-void FD3D12DynamicRHI::RHIUpdateTexture2D(FTexture2DRHIParamRef TextureRHI, uint32 MipIndex, const FUpdateTextureRegion2D& UpdateRegion, uint32 SourcePitch, const uint8* SourceData)
+void FD3D12DynamicRHI::RHIUpdateTexture2D(FRHITexture2D* TextureRHI, uint32 MipIndex, const FUpdateTextureRegion2D& UpdateRegion, uint32 SourcePitch, const uint8* SourceData)
 {
 	check(TextureRHI);
 	FD3D12Texture2D* Texture = FD3D12DynamicRHI::ResourceCast(TextureRHI);
 	Texture->UpdateTexture2D(nullptr, MipIndex, UpdateRegion, SourcePitch, SourceData);
 }
 
-FUpdateTexture3DData FD3D12DynamicRHI::BeginUpdateTexture3D_RenderThread(class FRHICommandListImmediate& RHICmdList, FTexture3DRHIParamRef Texture, uint32 MipIndex, const struct FUpdateTextureRegion3D& UpdateRegion)
+FUpdateTexture3DData FD3D12DynamicRHI::BeginUpdateTexture3D_RenderThread(class FRHICommandListImmediate& RHICmdList, FRHITexture3D* Texture, uint32 MipIndex, const struct FUpdateTextureRegion3D& UpdateRegion)
 {
 	check(IsInRenderingThread());
 	// This stall could potentially be removed, provided the fast allocator is thread-safe. However we 
@@ -2159,7 +2159,7 @@ void FD3D12DynamicRHI::EndMultiUpdateTexture3D_RenderThread(class FRHICommandLis
 	}
 }
 
-void FD3D12DynamicRHI::RHIUpdateTexture3D(FTexture3DRHIParamRef TextureRHI, uint32 MipIndex, const FUpdateTextureRegion3D& InUpdateRegion, uint32 SourceRowPitch, uint32 SourceDepthPitch, const uint8* SourceData)
+void FD3D12DynamicRHI::RHIUpdateTexture3D(FRHITexture3D* TextureRHI, uint32 MipIndex, const FUpdateTextureRegion3D& InUpdateRegion, uint32 SourceRowPitch, uint32 SourceDepthPitch, const uint8* SourceData)
 {
 	check(IsInRenderingThread());
 
@@ -2197,7 +2197,7 @@ void FD3D12DynamicRHI::RHIUpdateTexture3D(FTexture3DRHIParamRef TextureRHI, uint
 }
 
 
-FUpdateTexture3DData FD3D12DynamicRHI::BeginUpdateTexture3D_Internal(FTexture3DRHIParamRef TextureRHI, uint32 MipIndex, const struct FUpdateTextureRegion3D& UpdateRegion)
+FUpdateTexture3DData FD3D12DynamicRHI::BeginUpdateTexture3D_Internal(FRHITexture3D* TextureRHI, uint32 MipIndex, const struct FUpdateTextureRegion3D& UpdateRegion)
 {
 	check(IsInRenderingThread());
 	FUpdateTexture3DData UpdateData(TextureRHI, MipIndex, UpdateRegion, 0, 0, nullptr, 0, GFrameNumberRenderThread);
@@ -2382,21 +2382,21 @@ FTextureCubeRHIRef FD3D12DynamicRHI::RHICreateTextureCubeArray(uint32 Size, uint
 	return CreateD3D12Texture2D<FD3D12BaseTextureCube>(nullptr, Size, Size, 6 * ArraySize, true, true, (EPixelFormat) Format, NumMips, 1, Flags, CreateInfo);
 }
 
-void* FD3D12DynamicRHI::RHILockTextureCubeFace(FTextureCubeRHIParamRef TextureCubeRHI, uint32 FaceIndex, uint32 ArrayIndex, uint32 MipIndex, EResourceLockMode LockMode, uint32& DestStride, bool bLockWithinMiptail)
+void* FD3D12DynamicRHI::RHILockTextureCubeFace(FRHITextureCube* TextureCubeRHI, uint32 FaceIndex, uint32 ArrayIndex, uint32 MipIndex, EResourceLockMode LockMode, uint32& DestStride, bool bLockWithinMiptail)
 {
 	FD3D12TextureCube*  TextureCube = FD3D12DynamicRHI::ResourceCast(TextureCubeRHI);
 	GetRHIDevice()->GetDefaultCommandContext().ConditionalClearShaderResource(&TextureCube->ResourceLocation);
 	uint32 D3DFace = GetD3D12CubeFace((ECubeFace)FaceIndex);
 	return TextureCube->Lock(nullptr, MipIndex, D3DFace + ArrayIndex * 6, LockMode, DestStride);
 }
-void FD3D12DynamicRHI::RHIUnlockTextureCubeFace(FTextureCubeRHIParamRef TextureCubeRHI, uint32 FaceIndex, uint32 ArrayIndex, uint32 MipIndex, bool bLockWithinMiptail)
+void FD3D12DynamicRHI::RHIUnlockTextureCubeFace(FRHITextureCube* TextureCubeRHI, uint32 FaceIndex, uint32 ArrayIndex, uint32 MipIndex, bool bLockWithinMiptail)
 {
 	FD3D12TextureCube*  TextureCube = FD3D12DynamicRHI::ResourceCast(TextureCubeRHI);
 	uint32 D3DFace = GetD3D12CubeFace((ECubeFace)FaceIndex);
 	TextureCube->Unlock(nullptr, MipIndex, D3DFace + ArrayIndex * 6);
 }
 
-void FD3D12DynamicRHI::RHIBindDebugLabelName(FTextureRHIParamRef TextureRHI, const TCHAR* Name)
+void FD3D12DynamicRHI::RHIBindDebugLabelName(FRHITexture* TextureRHI, const TCHAR* Name)
 {
 #if NAME_OBJECTS
 	FD3D12TextureBase* BaseTexture = GetD3D12TextureFromRHITexture(TextureRHI);
@@ -2423,11 +2423,11 @@ void FD3D12DynamicRHI::RHIBindDebugLabelName(FTextureRHIParamRef TextureRHI, con
 #endif
 }
 
-void FD3D12DynamicRHI::RHIVirtualTextureSetFirstMipInMemory(FTexture2DRHIParamRef TextureRHI, uint32 FirstMip)
+void FD3D12DynamicRHI::RHIVirtualTextureSetFirstMipInMemory(FRHITexture2D* TextureRHI, uint32 FirstMip)
 {
 }
 
-void FD3D12DynamicRHI::RHIVirtualTextureSetFirstMipVisible(FTexture2DRHIParamRef TextureRHI, uint32 FirstMip)
+void FD3D12DynamicRHI::RHIVirtualTextureSetFirstMipVisible(FRHITexture2D* TextureRHI, uint32 FirstMip)
 {
 }
 
@@ -2436,7 +2436,7 @@ FTextureReferenceRHIRef FD3D12DynamicRHI::RHICreateTextureReference(FLastRenderT
 	return new FD3D12TextureReference(GetRHIDevice(), LastRenderTime);
 }
 
-void FD3D12CommandContext::RHIUpdateTextureReference(FRHITextureReference* TextureRefRHI, FTextureRHIParamRef NewTextureRHI)
+void FD3D12CommandContext::RHIUpdateTextureReference(FRHITextureReference* TextureRefRHI, FRHITexture* NewTextureRHI)
 {
 #if 0
 	// Updating texture references is disallowed while the RHI could be caching them in referenced resource tables.
@@ -2794,7 +2794,7 @@ FTextureCubeRHIRef FD3D12DynamicRHI::RHICreateTextureCubeFromResource(EPixelForm
 	return TextureCube;
 }
 
-void FD3D12DynamicRHI::RHIAliasTextureResources(FTextureRHIParamRef DestTextureRHI, FTextureRHIParamRef SrcTextureRHI)
+void FD3D12DynamicRHI::RHIAliasTextureResources(FRHITexture* DestTextureRHI, FRHITexture* SrcTextureRHI)
 {	
 	FD3D12TextureBase* DestTexture = GetD3D12TextureFromRHITexture(DestTextureRHI);
 	FD3D12TextureBase* SrcTexture = GetD3D12TextureFromRHITexture(SrcTextureRHI);
@@ -2808,7 +2808,7 @@ void FD3D12DynamicRHI::RHIAliasTextureResources(FTextureRHIParamRef DestTextureR
 	}
 }
 
-void FD3D12DynamicRHI::RHICopySubTextureRegion(FTexture2DRHIParamRef SourceTextureRHI, FTexture2DRHIParamRef DestTextureRHI, FBox2D SourceBox, FBox2D DestinationBox)
+void FD3D12DynamicRHI::RHICopySubTextureRegion(FRHITexture2D* SourceTextureRHI, FRHITexture2D* DestTextureRHI, FBox2D SourceBox, FBox2D DestinationBox)
 {
 	FD3D12TextureBase* SourceTexture = GetD3D12TextureFromRHITexture(SourceTextureRHI);
 	FD3D12TextureBase* DestTexture = GetD3D12TextureFromRHITexture(DestTextureRHI);
@@ -2835,7 +2835,7 @@ void FD3D12DynamicRHI::RHICopySubTextureRegion(FTexture2DRHIParamRef SourceTextu
 	}
 }
 
-void FD3D12CommandContext::RHICopyTexture(FTextureRHIParamRef SourceTextureRHI, FTextureRHIParamRef DestTextureRHI, const FRHICopyTextureInfo& CopyInfo)
+void FD3D12CommandContext::RHICopyTexture(FRHITexture* SourceTextureRHI, FRHITexture* DestTextureRHI, const FRHICopyTextureInfo& CopyInfo)
 {
 	FD3D12TextureBase* SourceTexture = GetD3D12TextureFromRHITexture(SourceTextureRHI);
 	FD3D12TextureBase* DestTexture = GetD3D12TextureFromRHITexture(DestTextureRHI);
diff --git a/Engine/Source/Runtime/D3D12RHI/Private/D3D12UAV.cpp b/Engine/Source/Runtime/D3D12RHI/Private/D3D12UAV.cpp
--- a/Engine/Source/Runtime/D3D12RHI/Private/D3D12UAV.cpp
+++ b/Engine/Source/Runtime/D3D12RHI/Private/D3D12UAV.cpp
@@ -27,7 +27,7 @@ inline FD3D12UnorderedAccessView* CreateUAV(D3D12_UNORDERED_ACCESS_VIEW_DESC& De
 	});
 }
 
-FUnorderedAccessViewRHIRef FD3D12DynamicRHI::RHICreateUnorderedAccessView(FStructuredBufferRHIParamRef StructuredBufferRHI, bool bUseUAVCounter, bool bAppendBuffer)
+FUnorderedAccessViewRHIRef FD3D12DynamicRHI::RHICreateUnorderedAccessView(FRHIStructuredBuffer* StructuredBufferRHI, bool bUseUAVCounter, bool bAppendBuffer)
 {
 	FD3D12StructuredBuffer*  StructuredBuffer = FD3D12DynamicRHI::ResourceCast(StructuredBufferRHI);
 
@@ -74,7 +74,7 @@ FUnorderedAccessViewRHIRef FD3D12DynamicRHI::RHICreateUnorderedAccessView(FStruc
 	return CreateUAV(UAVDesc, StructuredBuffer, bNeedsCounterResource);
 }
 
-FUnorderedAccessViewRHIRef FD3D12DynamicRHI::RHICreateUnorderedAccessView(FTextureRHIParamRef TextureRHI, uint32 MipLevel)
+FUnorderedAccessViewRHIRef FD3D12DynamicRHI::RHICreateUnorderedAccessView(FRHITexture* TextureRHI, uint32 MipLevel)
 {
 	FD3D12TextureBase* Texture = GetD3D12TextureFromRHITexture(TextureRHI);
 
@@ -126,7 +126,7 @@ FUnorderedAccessViewRHIRef FD3D12DynamicRHI::RHICreateUnorderedAccessView(FTextu
 	}
 }
 
-FUnorderedAccessViewRHIRef FD3D12DynamicRHI::RHICreateUnorderedAccessView(FVertexBufferRHIParamRef VertexBufferRHI, uint8 Format)
+FUnorderedAccessViewRHIRef FD3D12DynamicRHI::RHICreateUnorderedAccessView(FRHIVertexBuffer* VertexBufferRHI, uint8 Format)
 {
 	FD3D12VertexBuffer*  VertexBuffer = FD3D12DynamicRHI::ResourceCast(VertexBufferRHI);
 	FD3D12ResourceLocation& Location = VertexBuffer->ResourceLocation;
@@ -164,7 +164,7 @@ FUnorderedAccessViewRHIRef FD3D12DynamicRHI::RHICreateUnorderedAccessView(FVerte
 	return CreateUAV(UAVDesc, VertexBuffer, false);
 }
 
-FUnorderedAccessViewRHIRef FD3D12DynamicRHI::RHICreateUnorderedAccessView(FIndexBufferRHIParamRef IndexBufferRHI, uint8 Format)
+FUnorderedAccessViewRHIRef FD3D12DynamicRHI::RHICreateUnorderedAccessView(FRHIIndexBuffer* IndexBufferRHI, uint8 Format)
 {
 	FD3D12IndexBuffer* IndexBuffer = FD3D12DynamicRHI::ResourceCast(IndexBufferRHI);
 	FD3D12ResourceLocation& Location = IndexBuffer->ResourceLocation;
@@ -230,7 +230,7 @@ void FD3D12CommandContext::RHIClearTinyUAV(FRHIUnorderedAccessView* UnorderedAcc
 	}
 }
 
-FUnorderedAccessViewRHIRef FD3D12DynamicRHI::RHICreateUnorderedAccessView_RenderThread(FRHICommandListImmediate& RHICmdList, FStructuredBufferRHIParamRef StructuredBufferRHI, bool bUseUAVCounter, bool bAppendBuffer)
+FUnorderedAccessViewRHIRef FD3D12DynamicRHI::RHICreateUnorderedAccessView_RenderThread(FRHICommandListImmediate& RHICmdList, FRHIStructuredBuffer* StructuredBufferRHI, bool bUseUAVCounter, bool bAppendBuffer)
 {
 	FD3D12StructuredBuffer* StructuredBuffer = FD3D12DynamicRHI::ResourceCast(StructuredBufferRHI);
 	// TODO: we have to stall the RHI thread when creating SRVs of dynamic buffers because they get renamed.
@@ -243,12 +243,12 @@ FUnorderedAccessViewRHIRef FD3D12DynamicRHI::RHICreateUnorderedAccessView_Render
 	return RHICreateUnorderedAccessView(StructuredBufferRHI, bUseUAVCounter, bAppendBuffer);
 }
 
-FUnorderedAccessViewRHIRef FD3D12DynamicRHI::RHICreateUnorderedAccessView_RenderThread(FRHICommandListImmediate& RHICmdList, FTextureRHIParamRef Texture, uint32 MipLevel)
+FUnorderedAccessViewRHIRef FD3D12DynamicRHI::RHICreateUnorderedAccessView_RenderThread(FRHICommandListImmediate& RHICmdList, FRHITexture* Texture, uint32 MipLevel)
 {
 	return FD3D12DynamicRHI::RHICreateUnorderedAccessView(Texture, MipLevel);
 }
 
-FUnorderedAccessViewRHIRef FD3D12DynamicRHI::RHICreateUnorderedAccessView_RenderThread(FRHICommandListImmediate& RHICmdList, FVertexBufferRHIParamRef VertexBufferRHI, uint8 Format)
+FUnorderedAccessViewRHIRef FD3D12DynamicRHI::RHICreateUnorderedAccessView_RenderThread(FRHICommandListImmediate& RHICmdList, FRHIVertexBuffer* VertexBufferRHI, uint8 Format)
 {
 	FD3D12VertexBuffer* VertexBuffer = FD3D12DynamicRHI::ResourceCast(VertexBufferRHI);
 
diff --git a/Engine/Source/Runtime/D3D12RHI/Private/D3D12UniformBuffer.cpp b/Engine/Source/Runtime/D3D12RHI/Private/D3D12UniformBuffer.cpp
--- a/Engine/Source/Runtime/D3D12RHI/Private/D3D12UniformBuffer.cpp
+++ b/Engine/Source/Runtime/D3D12RHI/Private/D3D12UniformBuffer.cpp
@@ -152,7 +152,7 @@ struct FRHICommandD3D12UpdateUniformBuffer final : public FRHICommand<FRHIComman
 	}
 };
 
-void FD3D12DynamicRHI::RHIUpdateUniformBuffer(FUniformBufferRHIParamRef UniformBufferRHI, const void* Contents)
+void FD3D12DynamicRHI::RHIUpdateUniformBuffer(FRHIUniformBuffer* UniformBufferRHI, const void* Contents)
 {
 	check(IsInRenderingThread());
 	check(UniformBufferRHI);
