Cherry-pick 6953813: Added memory stat counters for bottom and top level acceleration structures

diff --git a/Engine/Source/Runtime/D3D12RHI/Private/D3D12RayTracing.cpp b/Engine/Source/Runtime/D3D12RHI/Private/D3D12RayTracing.cpp
--- a/Engine/Source/Runtime/D3D12RHI/Private/D3D12RayTracing.cpp
+++ b/Engine/Source/Runtime/D3D12RHI/Private/D3D12RayTracing.cpp
@@ -73,7 +73,9 @@ DECLARE_DWORD_COUNTER_STAT(TEXT("Updated BL AS (per frame)"), STAT_D3D12RayTraci
 DECLARE_DWORD_COUNTER_STAT(TEXT("Built TL AS (per frame)"), STAT_D3D12RayTracingBuiltTLAS, STATGROUP_D3D12RayTracing);
 DECLARE_DWORD_COUNTER_STAT(TEXT("Updated TL AS (per frame)"), STAT_D3D12RayTracingUpdatedTLAS, STATGROUP_D3D12RayTracing);
 
-DECLARE_MEMORY_STAT(TEXT("Used Video Memory"), STAT_D3D12RayTracingUsedVideoMemory, STATGROUP_D3D12RayTracing);
+DECLARE_MEMORY_STAT(TEXT("BL AS Memory"), STAT_D3D12RayTracingBLASMemory, STATGROUP_D3D12RayTracing);
+DECLARE_MEMORY_STAT(TEXT("TL AS Memory"), STAT_D3D12RayTracingTLASMemory, STATGROUP_D3D12RayTracing);
+DECLARE_MEMORY_STAT(TEXT("Total Used Video Memory"), STAT_D3D12RayTracingUsedVideoMemory, STATGROUP_D3D12RayTracing);
 
 DECLARE_CYCLE_STAT(TEXT("RTPSO Compile Shader"), STAT_RTPSO_CompileShader, STATGROUP_D3D12RayTracing);
 DECLARE_CYCLE_STAT(TEXT("RTPSO Create Pipeline"), STAT_RTPSO_CreatePipeline, STATGROUP_D3D12RayTracing);
@@ -1925,6 +1927,7 @@ FD3D12RayTracingGeometry::~FD3D12RayTracingGeometry()
 		if (Buffer)
 		{
 			DEC_MEMORY_STAT_BY(STAT_D3D12RayTracingUsedVideoMemory, Buffer->GetSize());
+			DEC_MEMORY_STAT_BY(STAT_D3D12RayTracingBLASMemory, Buffer->GetSize());
 		}
 	}
 
@@ -1933,6 +1936,7 @@ FD3D12RayTracingGeometry::~FD3D12RayTracingGeometry()
 		if (Buffer)
 		{
 			DEC_MEMORY_STAT_BY(STAT_D3D12RayTracingUsedVideoMemory, Buffer->GetSize());
+			DEC_MEMORY_STAT_BY(STAT_D3D12RayTracingBLASMemory, Buffer->GetSize());
 		}
 	}
 
@@ -2115,6 +2119,9 @@ void FD3D12RayTracingGeometry::BuildAccelerationStructure(FD3D12CommandContext&
 		INC_MEMORY_STAT_BY(STAT_D3D12RayTracingUsedVideoMemory, AccelerationStructureBuffers[GPUIndex]->GetSize());
 		INC_MEMORY_STAT_BY(STAT_D3D12RayTracingUsedVideoMemory, ScratchBuffers[GPUIndex]->GetSize());
 
+		INC_MEMORY_STAT_BY(STAT_D3D12RayTracingBLASMemory, AccelerationStructureBuffers[GPUIndex]->GetSize());
+		INC_MEMORY_STAT_BY(STAT_D3D12RayTracingBLASMemory, ScratchBuffers[GPUIndex]->GetSize());
+
 		// #dxr_todo UE-72161: scratch buffers should be created in UAV state from the start
 		FD3D12DynamicRHI::TransitionResource(CommandContext.CommandListHandle, ScratchBuffers[GPUIndex].GetReference()->GetResource(), D3D12_RESOURCE_STATE_UNORDERED_ACCESS, 0);
 	}
@@ -2153,6 +2160,7 @@ void FD3D12RayTracingGeometry::BuildAccelerationStructure(FD3D12CommandContext&
 	if (!(BuildFlags & D3D12_RAYTRACING_ACCELERATION_STRUCTURE_BUILD_FLAG_ALLOW_UPDATE))
 	{
 		DEC_MEMORY_STAT_BY(STAT_D3D12RayTracingUsedVideoMemory, ScratchBuffers[GPUIndex]->GetSize());
+		DEC_MEMORY_STAT_BY(STAT_D3D12RayTracingBLASMemory, ScratchBuffers[GPUIndex]->GetSize());
 		ScratchBuffers[GPUIndex] = nullptr;
 	}
 }
@@ -2182,6 +2190,7 @@ FD3D12RayTracingScene::~FD3D12RayTracingScene()
 		if (AccelerationStructureBuffer)
 		{
 			DEC_MEMORY_STAT_BY(STAT_D3D12RayTracingUsedVideoMemory, AccelerationStructureBuffer->GetSize());
+			DEC_MEMORY_STAT_BY(STAT_D3D12RayTracingTLASMemory, AccelerationStructureBuffer->GetSize());
 		}
 	}
 	
@@ -2213,6 +2222,7 @@ void FD3D12RayTracingScene::BuildAccelerationStructure(FD3D12CommandContext& Com
 	CreateAccelerationStructureBuffers(AccelerationStructureBuffer, ScratchBuffer, Adapter, GPUIndex, PrebuildInfo);
 
 	INC_MEMORY_STAT_BY(STAT_D3D12RayTracingUsedVideoMemory, AccelerationStructureBuffer->GetSize());
+	INC_MEMORY_STAT_BY(STAT_D3D12RayTracingTLASMemory, AccelerationStructureBuffer->GetSize());
 
 	// #dxr_todo UE-72161: scratch buffers should be created in UAV state from the start
 	FD3D12DynamicRHI::TransitionResource(CommandContext.CommandListHandle, ScratchBuffer.GetReference()->GetResource(), D3D12_RESOURCE_STATE_UNORDERED_ACCESS, 0);
