Cherry-pick 6969888: Implemented mGPU virtualization on a single GPU in D3D12, enabled via -vmgpu command line argument.

Virtualization works by creating multiple FD3D12Device-s, as if real SLI/LDA mode was used. All resources are created multiple times on the same physical device to simulate multiple LDA nodes.
Multiple 3D, copy and compute queues are created (one per device). Explicit transfers and synchronization must be performed to share resources between virtual devices, exactly like on a real mGPU machine.

This mode allows development and debugging of mGPU AFR and SFR features on any machine, using standard tools such as PIX, Nsight and RenderDoc that don't support real multi-node configurations.

Up to 4 GPUs can be virtualized this way (same number as maximum supported real LDA GPUs).

Example usage command line: MyProject -game -d3d12 -MaxGPUCount=2 -VMGPU -AFR

diff --git a/Engine/Source/Runtime/D3D12RHI/Private/D3D12Adapter.cpp b/Engine/Source/Runtime/D3D12RHI/Private/D3D12Adapter.cpp
--- a/Engine/Source/Runtime/D3D12RHI/Private/D3D12Adapter.cpp
+++ b/Engine/Source/Runtime/D3D12RHI/Private/D3D12Adapter.cpp
@@ -50,8 +50,21 @@ FD3D12Adapter::FD3D12Adapter(FD3D12AdapterDesc& DescIn)
 			MaxGPUCount = MAX_NUM_GPUS;
 		}
 	}
+	if (FParse::Param(FCommandLine::Get(), TEXT("VMGPU")))
+	{
+		GVirtualMGPU = 1;
+		UE_LOG(LogD3D12RHI, Log, TEXT("Enabling virtual multi-GPU mode"), Desc.NumDeviceNodes);
+	}
 #endif
-	Desc.NumDeviceNodes = FMath::Min3<uint32>(Desc.NumDeviceNodes, MaxGPUCount, MAX_NUM_GPUS);
+
+	if (GVirtualMGPU)
+	{
+		Desc.NumDeviceNodes = FMath::Min<uint32>(MaxGPUCount, MAX_NUM_GPUS);
+	}
+	else
+	{
+		Desc.NumDeviceNodes = FMath::Min3<uint32>(Desc.NumDeviceNodes, MaxGPUCount, MAX_NUM_GPUS);
+	}
 }
 
 void FD3D12Adapter::Initialize(FD3D12DynamicRHI* RHI)
@@ -676,12 +689,15 @@ void FD3D12Adapter::GetLocalVideoMemoryInfo(DXGI_QUERY_VIDEO_MEMORY_INFO* LocalV
 
 	VERIFYD3D12RESULT(Adapter3->QueryVideoMemoryInfo(0, DXGI_MEMORY_SEGMENT_GROUP_LOCAL, LocalVideoMemoryInfo));
 
-	for (uint32 Index = 1; Index < GNumExplicitGPUsForRendering; ++Index)
+	if (!GVirtualMGPU)
 	{
-		DXGI_QUERY_VIDEO_MEMORY_INFO TempVideoMemoryInfo;
-		VERIFYD3D12RESULT(Adapter3->QueryVideoMemoryInfo(Index, DXGI_MEMORY_SEGMENT_GROUP_LOCAL, &TempVideoMemoryInfo));
-		LocalVideoMemoryInfo->Budget = FMath::Min(LocalVideoMemoryInfo->Budget, TempVideoMemoryInfo.Budget);
-		LocalVideoMemoryInfo->Budget = FMath::Min(LocalVideoMemoryInfo->CurrentUsage, TempVideoMemoryInfo.CurrentUsage);
+		for (uint32 Index = 1; Index < GNumExplicitGPUsForRendering; ++Index)
+		{
+			DXGI_QUERY_VIDEO_MEMORY_INFO TempVideoMemoryInfo;
+			VERIFYD3D12RESULT(Adapter3->QueryVideoMemoryInfo(Index, DXGI_MEMORY_SEGMENT_GROUP_LOCAL, &TempVideoMemoryInfo));
+			LocalVideoMemoryInfo->Budget = FMath::Min(LocalVideoMemoryInfo->Budget, TempVideoMemoryInfo.Budget);
+			LocalVideoMemoryInfo->Budget = FMath::Min(LocalVideoMemoryInfo->CurrentUsage, TempVideoMemoryInfo.CurrentUsage);
+		}
 	}
 #endif
 }
diff --git a/Engine/Source/Runtime/D3D12RHI/Private/D3D12Allocation.cpp b/Engine/Source/Runtime/D3D12RHI/Private/D3D12Allocation.cpp
--- a/Engine/Source/Runtime/D3D12RHI/Private/D3D12Allocation.cpp
+++ b/Engine/Source/Runtime/D3D12RHI/Private/D3D12Allocation.cpp
@@ -1625,7 +1625,7 @@ FD3D12SegHeap* FD3D12SegList::CreateBackingHeap(
 	ID3D12Heap* D3DHeap;
 	D3D12_HEAP_DESC Desc = {};
 	Desc.SizeInBytes = HeapSize;
-	Desc.Properties = CD3DX12_HEAP_PROPERTIES(HeapType, Parent->GetGPUMask(), VisibleNodeMask);
+	Desc.Properties = CD3DX12_HEAP_PROPERTIES(HeapType, (uint32)Parent->GetGPUMask(), (uint32)VisibleNodeMask);
 	Desc.Flags = HeapFlags;
 
 	VERIFYD3D12RESULT(Parent->GetDevice()->CreateHeap(&Desc, IID_PPV_ARGS(&D3DHeap)));
diff --git a/Engine/Source/Runtime/D3D12RHI/Private/D3D12Device.cpp b/Engine/Source/Runtime/D3D12RHI/Private/D3D12Device.cpp
--- a/Engine/Source/Runtime/D3D12RHI/Private/D3D12Device.cpp
+++ b/Engine/Source/Runtime/D3D12RHI/Private/D3D12Device.cpp
@@ -279,10 +279,10 @@ void FD3D12Device::SetupAfterDeviceCreation()
 	OcclusionQueryHeap.Init();
 	TimestampQueryHeap.Init();
 
-	CommandListManager->Create(L"3D Queue");
+	CommandListManager->Create(*FString::Printf(TEXT("3D Queue %d"), GetGPUIndex()));
 	gD3D12CommandQueue = CommandListManager->GetD3DCommandQueue();
-	CopyCommandListManager->Create(L"Copy Queue");
-	AsyncCommandListManager->Create(L"Async Compute Queue", 0, AsyncComputePriority_Default);
+	CopyCommandListManager->Create(*FString::Printf(TEXT("Copy Queue %d"), GetGPUIndex()));
+	AsyncCommandListManager->Create(*FString::Printf(TEXT("Compute Queue %d"), GetGPUIndex()), 0, AsyncComputePriority_Default);
 
 	// Needs to be called before creating command contexts
 	UpdateConstantBufferPageProperties();
diff --git a/Engine/Source/Runtime/D3D12RHI/Private/Windows/WindowsD3D12Device.cpp b/Engine/Source/Runtime/D3D12RHI/Private/Windows/WindowsD3D12Device.cpp
--- a/Engine/Source/Runtime/D3D12RHI/Private/Windows/WindowsD3D12Device.cpp
+++ b/Engine/Source/Runtime/D3D12RHI/Private/Windows/WindowsD3D12Device.cpp
@@ -752,7 +752,8 @@ void FD3D12Device::Initialize()
 #if ENABLE_RESIDENCY_MANAGEMENT
 	IDXGIAdapter3* DxgiAdapter3 = nullptr;
 	VERIFYD3D12RESULT(GetParentAdapter()->GetAdapter()->QueryInterface(IID_PPV_ARGS(&DxgiAdapter3)));
-	D3DX12Residency::InitializeResidencyManager(ResidencyManager, GetDevice(), GetGPUIndex(), DxgiAdapter3, RESIDENCY_PIPELINE_DEPTH);
+	const uint32 GPUIndex = GVirtualMGPU ? 0 : GetGPUIndex(); // GPU node index is used by residency manager to query budget
+	D3DX12Residency::InitializeResidencyManager(ResidencyManager, GetDevice(), GPUIndex, DxgiAdapter3, RESIDENCY_PIPELINE_DEPTH);
 #endif // ENABLE_RESIDENCY_MANAGEMENT
 
 	SetupAfterDeviceCreation();
diff --git a/Engine/Source/Runtime/RHI/Private/MultiGPU.cpp b/Engine/Source/Runtime/RHI/Private/MultiGPU.cpp
--- a/Engine/Source/Runtime/RHI/Private/MultiGPU.cpp
+++ b/Engine/Source/Runtime/RHI/Private/MultiGPU.cpp
@@ -10,4 +10,5 @@
 #if WITH_SLI || WITH_MGPU
 uint32 GNumAlternateFrameRenderingGroups = 1;
 uint32 GNumExplicitGPUsForRendering = 1;
+uint32 GVirtualMGPU = 0;
 #endif
diff --git a/Engine/Source/Runtime/RHI/Public/MultiGPU.h b/Engine/Source/Runtime/RHI/Public/MultiGPU.h
--- a/Engine/Source/Runtime/RHI/Public/MultiGPU.h
+++ b/Engine/Source/Runtime/RHI/Public/MultiGPU.h
@@ -15,12 +15,14 @@
 #define MAX_NUM_GPUS 4
 extern RHI_API uint32 GNumExplicitGPUsForRendering;
 extern RHI_API uint32 GNumAlternateFrameRenderingGroups;
+extern RHI_API uint32 GVirtualMGPU;
 #else
 #define WITH_SLI 0
 #define WITH_MGPU 0
 #define MAX_NUM_GPUS 1
 #define GNumExplicitGPUsForRendering 1
 #define GNumAlternateFrameRenderingGroups 1
+#define GVirtualMGPU 0
 #endif
 
 /** A mask where each bit is a GPU index. Can not be empty so that non SLI platforms can optimize it to be always 1.  */
@@ -86,10 +88,12 @@ public:
 	FORCEINLINE bool Intersects(const FRHIGPUMask& Rhs) const { return (GPUMask & Rhs.GPUMask) != 0; }
 
 	FORCEINLINE bool operator ==(const FRHIGPUMask& Rhs) const { return GPUMask == Rhs.GPUMask; }
+	FORCEINLINE bool operator !=(const FRHIGPUMask& Rhs) const { return GPUMask != Rhs.GPUMask; }
 
 	void operator |=(const FRHIGPUMask& Rhs) { GPUMask |= Rhs.GPUMask; }
 	void operator &=(const FRHIGPUMask& Rhs) { GPUMask &= Rhs.GPUMask; }
-	FORCEINLINE operator uint32() const { return GPUMask; }
+
+	FORCEINLINE explicit operator uint32() const { return GVirtualMGPU ? 1 : GPUMask; }
 
 	FORCEINLINE FRHIGPUMask operator &(const FRHIGPUMask& Rhs) const
 	{
@@ -106,13 +110,17 @@ public:
 
 	struct FIterator
 	{
-		FORCEINLINE FIterator(const uint32 InGPUMask) : GPUMask(InGPUMask), FirstGPUIndexInMask(0)
+		FORCEINLINE explicit FIterator(const uint32 InGPUMask) : GPUMask(InGPUMask), FirstGPUIndexInMask(0)
 		{
 #if WITH_MGPU
 			FirstGPUIndexInMask = FPlatformMath::CountTrailingZeros(InGPUMask);
 #endif
 		}
 
+		FORCEINLINE explicit FIterator(const FRHIGPUMask& InGPUMask) : FIterator(InGPUMask.GPUMask)
+		{
+		}
+
 		FORCEINLINE void operator++()
 		{
 #if WITH_MGPU
