Cherry-pick 6970401: Fixed compiler warning due to GPUIndex variable shadowing

diff --git a/Engine/Source/Runtime/D3D12RHI/Private/Windows/WindowsD3D12Device.cpp b/Engine/Source/Runtime/D3D12RHI/Private/Windows/WindowsD3D12Device.cpp
--- a/Engine/Source/Runtime/D3D12RHI/Private/Windows/WindowsD3D12Device.cpp
+++ b/Engine/Source/Runtime/D3D12RHI/Private/Windows/WindowsD3D12Device.cpp
@@ -752,8 +752,8 @@ void FD3D12Device::Initialize()
 #if ENABLE_RESIDENCY_MANAGEMENT
 	IDXGIAdapter3* DxgiAdapter3 = nullptr;
 	VERIFYD3D12RESULT(GetParentAdapter()->GetAdapter()->QueryInterface(IID_PPV_ARGS(&DxgiAdapter3)));
-	const uint32 GPUIndex = GVirtualMGPU ? 0 : GetGPUIndex(); // GPU node index is used by residency manager to query budget
-	D3DX12Residency::InitializeResidencyManager(ResidencyManager, GetDevice(), GPUIndex, DxgiAdapter3, RESIDENCY_PIPELINE_DEPTH);
+	const uint32 ResidencyMangerGPUIndex = GVirtualMGPU ? 0 : GetGPUIndex(); // GPU node index is used by residency manager to query budget
+	D3DX12Residency::InitializeResidencyManager(ResidencyManager, GetDevice(), ResidencyMangerGPUIndex, DxgiAdapter3, RESIDENCY_PIPELINE_DEPTH);
 #endif // ENABLE_RESIDENCY_MANAGEMENT
 
 	SetupAfterDeviceCreation();
