Cherry-pick 6971873: Take viewport offset into account for scissor disable case

diff --git a/Engine/Source/Runtime/D3D12RHI/Private/D3D12Commands.cpp b/Engine/Source/Runtime/D3D12RHI/Private/D3D12Commands.cpp
--- a/Engine/Source/Runtime/D3D12RHI/Private/D3D12Commands.cpp
+++ b/Engine/Source/Runtime/D3D12RHI/Private/D3D12Commands.cpp
@@ -469,7 +469,8 @@ void FD3D12CommandContext::RHISetScissorRect(bool bEnable, uint32 MinX, uint32 M
 	}
 	else
 	{
-		const CD3DX12_RECT ScissorRect(0, 0, (LONG) StateCache.GetViewport().Width, (LONG) StateCache.GetViewport().Height);
+		const D3D12_VIEWPORT& Viewport = StateCache.GetViewport();
+		const CD3DX12_RECT ScissorRect((LONG) Viewport.TopLeftX, (LONG) Viewport.TopLeftY, (LONG) Viewport.TopLeftX + (LONG) Viewport.Width, (LONG) Viewport.TopLeftY + (LONG) Viewport.Height);
 		StateCache.SetScissorRect(ScissorRect);
 	}
 }
