Cherry-pick 7062333: Implemented an optimization for ray tracing resource binding. SBT records are now cached using uniform buffer pointers as key. Records with matching keys are simply copied, avoiding the entire cost of resource binding.

CPU performance in large scenes with many identical mesh instances is significantly improved by this (up to ~4-5x speedup in ArchVizRayTracing, ~1.5x speedup in Infiltrator).
This optimization assumes that contents of uniform buffers will not change during a frame. CVar r.RayTracing.CacheShaderRecords (default = 1) can be used to turn off caching.

diff --git a/Engine/Source/Runtime/D3D12RHI/Private/D3D12RayTracing.cpp b/Engine/Source/Runtime/D3D12RHI/Private/D3D12RayTracing.cpp
--- a/Engine/Source/Runtime/D3D12RHI/Private/D3D12RayTracing.cpp
+++ b/Engine/Source/Runtime/D3D12RHI/Private/D3D12RayTracing.cpp
@@ -48,6 +48,14 @@ static FAutoConsoleVariableRef CVarRayTracingDebugDisableTriangleCull(
 	TEXT("Forces all ray tracing geometry instances to be double-sided by disabling back-face culling. This is useful for debugging and profiling. (default = 0)")
 );
 
+static int32 GRayTracingCacheShaderRecords = 1;
+static FAutoConsoleVariableRef CVarRayTracingShaderRecordCache(
+	TEXT("r.RayTracing.CacheShaderRecords"),
+	GRayTracingCacheShaderRecords,
+	TEXT("Automatically cache and re-use SBT hit group records. This significantly improves CPU performance in large scenes with many identical mesh instances. (default = 1)\n")
+	TEXT("This mode assumes that contents of uniform buffers does not change during ray tracing resource binding.")
+);
+
 // Ray tracing stat counters
 
 DECLARE_STATS_GROUP(TEXT("D3D12RHI: Ray Tracing"), STATGROUP_D3D12RayTracing, STATCAT_Advanced);
@@ -1141,6 +1149,22 @@ public:
 		WriteLocalShaderRecord(RecordIndex, ShaderIdentifierSize + InOffsetWithinRootSignature, InData, InDataSize);
 	}
 
+	void CopyLocalShaderParameters(uint32 InDestRecordIndex, uint32 InSourceRecordIndex, uint32 InOffsetWithinRootSignature)
+	{
+		const uint32 BaseOffset = LocalShaderTableOffset + ShaderIdentifierSize + InOffsetWithinRootSignature;
+		const uint32 DestOffset   = BaseOffset + LocalRecordStride * InDestRecordIndex;
+		const uint32 SourceOffset = BaseOffset + LocalRecordStride * InSourceRecordIndex;
+		const uint32 CopySize = LocalRecordStride - ShaderIdentifierSize - InOffsetWithinRootSignature;
+		checkSlow(CopySize <= LocalRecordStride);
+
+		FMemory::Memcpy(
+			Data.GetData() + DestOffset,
+			Data.GetData() + SourceOffset,
+			CopySize);
+
+		bIsDirty = true;
+	}
+
 	void SetLocalShaderIdentifier(uint32 RecordIndex, const void* ShaderIdentifierData, uint32 InShaderIdentifierSize)
 	{
 		checkSlow(InShaderIdentifierSize == ShaderIdentifierSize);
@@ -1296,6 +1320,44 @@ public:
 	// SBTs have their own descriptor heaps
 	FD3D12RayTracingDescriptorCache* DescriptorCache = nullptr;
 
+	struct FShaderRecordCacheKey
+	{
+		static constexpr uint32 MaxUniformBuffers = 4;
+		const FUniformBufferRHIParamRef* UniformBuffers[MaxUniformBuffers];
+		uint64 Hash;
+		uint32 NumUniformBuffers;
+
+		FShaderRecordCacheKey() = default;
+		FShaderRecordCacheKey(uint32 InNumUniformBuffers, const FUniformBufferRHIParamRef* InUniformBuffers)
+		{
+			check(InNumUniformBuffers <= MaxUniformBuffers);
+			NumUniformBuffers = FMath::Min(MaxUniformBuffers, InNumUniformBuffers);
+
+			const uint64 DataSizeInBytes = sizeof(FUniformBufferRHIParamRef) * NumUniformBuffers;
+			FMemory::Memcpy(UniformBuffers, InUniformBuffers, DataSizeInBytes);
+			Hash = CityHash64(reinterpret_cast<const char*>(UniformBuffers), DataSizeInBytes);
+		}
+
+		bool operator == (const FShaderRecordCacheKey& Other) const
+		{
+			if (Hash != Other.Hash) return false;
+			if (NumUniformBuffers != Other.NumUniformBuffers) return false;
+
+			for (uint32 BufferIndex = 0; BufferIndex < NumUniformBuffers; ++BufferIndex)
+			{
+				if (UniformBuffers[BufferIndex] != Other.UniformBuffers[BufferIndex]) return false;
+			}
+
+			return true;
+		}
+
+		friend uint32 GetTypeHash(const FShaderRecordCacheKey& Key)
+		{
+			return uint32(Key.Hash);
+		}
+	};
+	TMap<FShaderRecordCacheKey, uint32> ShaderRecordCache;
+
 #if ENABLE_RESIDENCY_MANAGEMENT
 	// A set of all resources referenced by this shader table for the purpose of updating residency before ray tracing work dispatch.
 	// #dxr_todo UE-72159: remove resources from this set when SBT slot entries are replaced
@@ -3144,6 +3206,27 @@ void FD3D12CommandContext::RHISetRayTracingHitGroup(
 
 	const FD3D12RayTracingShader* Shader = Pipeline->HitGroupShaders.Shaders[HitGroupIndex];
 
+	FD3D12RayTracingShaderTable::FShaderRecordCacheKey CacheKey;
+
+	const bool bCanUseRecordCache = GRayTracingCacheShaderRecords
+		&& Scene->Lifetime == RTSL_SingleFrame
+		&& LooseParameterDataSize == 0
+		&& NumUniformBuffers > 0
+		&& NumUniformBuffers <= CacheKey.MaxUniformBuffers;
+
+	if (bCanUseRecordCache)
+	{
+		CacheKey = FD3D12RayTracingShaderTable::FShaderRecordCacheKey(NumUniformBuffers, UniformBuffers);
+
+		uint32* ExistingRecordIndex = ShaderTable->ShaderRecordCache.Find(CacheKey);
+		if (ExistingRecordIndex)
+		{
+			const uint32 OffsetFromRootSignatureStart = sizeof(FHitGroupSystemParameters);
+			ShaderTable->CopyLocalShaderParameters(RecordIndex, *ExistingRecordIndex, OffsetFromRootSignatureStart);
+			return;
+		}
+	}
+
 	FD3D12RayTracingLocalResourceBinder ResourceBinder(*this, ShaderTable, Shader->pRootSignature, RecordIndex);
 	check(ShaderTable->DescriptorCache);
 	SetRayTracingShaderResources(*this, Shader,
@@ -3154,6 +3237,11 @@ void FD3D12CommandContext::RHISetRayTracingHitGroup(
 		0, nullptr, // UAVs
 		LooseParameterDataSize, LooseParameterData,
 		*(ShaderTable->DescriptorCache), ResourceBinder);
+
+	if (bCanUseRecordCache)
+	{
+		ShaderTable->ShaderRecordCache.Add(CacheKey, RecordIndex);
+	}
 }
 
 
