Fixed FUniformBufferRHIParamRef deprecation warning

diff --git a/Engine/Source/Runtime/D3D12RHI/Private/D3D12RayTracing.cpp b/Engine/Source/Runtime/D3D12RHI/Private/D3D12RayTracing.cpp
--- a/Engine/Source/Runtime/D3D12RHI/Private/D3D12RayTracing.cpp
+++ b/Engine/Source/Runtime/D3D12RHI/Private/D3D12RayTracing.cpp
@@ -1323,17 +1323,17 @@ public:
 	struct FShaderRecordCacheKey
 	{
 		static constexpr uint32 MaxUniformBuffers = 4;
-		const FUniformBufferRHIParamRef* UniformBuffers[MaxUniformBuffers];
+		FRHIUniformBuffer* const* UniformBuffers[MaxUniformBuffers];
 		uint64 Hash;
 		uint32 NumUniformBuffers;
 
 		FShaderRecordCacheKey() = default;
-		FShaderRecordCacheKey(uint32 InNumUniformBuffers, const FUniformBufferRHIParamRef* InUniformBuffers)
+		FShaderRecordCacheKey(uint32 InNumUniformBuffers, FRHIUniformBuffer* const* InUniformBuffers)
 		{
 			check(InNumUniformBuffers <= MaxUniformBuffers);
 			NumUniformBuffers = FMath::Min(MaxUniformBuffers, InNumUniformBuffers);
 
-			const uint64 DataSizeInBytes = sizeof(FUniformBufferRHIParamRef) * NumUniformBuffers;
+			const uint64 DataSizeInBytes = sizeof(FRHIUniformBuffer*) * NumUniformBuffers;
 			FMemory::Memcpy(UniformBuffers, InUniformBuffers, DataSizeInBytes);
 			Hash = CityHash64(reinterpret_cast<const char*>(UniformBuffers), DataSizeInBytes);
 		}
