Fix D3D12 GetShaderResourceViews #jira UE-76329 #rb rolando.caloca #5953

diff --git a/Engine/Source/Runtime/D3D12RHI/Private/D3D12StateCachePrivate.h b/Engine/Source/Runtime/D3D12RHI/Private/D3D12StateCachePrivate.h
--- a/Engine/Source/Runtime/D3D12RHI/Private/D3D12StateCachePrivate.h
+++ b/Engine/Source/Runtime/D3D12RHI/Private/D3D12StateCachePrivate.h
@@ -498,25 +498,7 @@ public:
 
 	template <EShaderFrequency ShaderFrequency>
 	void SetShaderResourceView(FD3D12ShaderResourceView* SRV, uint32 ResourceIndex);
-
-	template <EShaderFrequency ShaderFrequency>
-	D3D12_STATE_CACHE_INLINE void GetShaderResourceViews(uint32 StartResourceIndex, uint32& NumResources, FD3D12ShaderResourceView** SRV)
-	{
-		{
-			uint32 NumLoops = D3D12_COMMONSHADER_INPUT_RESOURCE_SLOT_COUNT - StartResourceIndex;
-			NumResources = 0;
-			for (uint32 ResourceLoop = 0; ResourceLoop < NumLoops; ResourceLoop++)
-			{
-				SRV[ResourceLoop] = PipelineState.Common.CurrentShaderResourceViews[ShaderFrequency][ResourceLoop + StartResourceIndex];
-				if (SRV[ResourceLoop])
-				{
-					SRV[ResourceLoop]->AddRef();
-					NumResources = ResourceLoop;
-				}
-			}
-		}
-	}
-
+	
 	void SetScissorRects(uint32 Count, const D3D12_RECT* const ScissorRects);
 	void SetScissorRect(const D3D12_RECT& ScissorRect);
 
