Fix a bug where StaticMesh LLM tag isn't counted correctly when LOD streaming is enabled

diff --git a/Engine/Source/Runtime/D3D12RHI/Private/D3D12Resources.cpp b/Engine/Source/Runtime/D3D12RHI/Private/D3D12Resources.cpp
--- a/Engine/Source/Runtime/D3D12RHI/Private/D3D12Resources.cpp
+++ b/Engine/Source/Runtime/D3D12RHI/Private/D3D12Resources.cpp
@@ -483,7 +483,7 @@ void FD3D12ResourceLocation::TransferOwnership(FD3D12ResourceLocation& Destinati
 
 	// update tracked allocation
 #if !PLATFORM_WINDOWS && ENABLE_LOW_LEVEL_MEM_TRACKER
-	if (Source.GetType() == ResourceLocationType::eSubAllocation)
+	if (Source.GetType() == ResourceLocationType::eSubAllocation && Source.AllocatorType != AT_SegList)
 	{
 		FLowLevelMemTracker::Get().OnLowLevelAllocMoved( ELLMTracker::Default, &Destination, &Source );
 	}
@@ -495,6 +495,25 @@ void FD3D12ResourceLocation::TransferOwnership(FD3D12ResourceLocation& Destinati
 
 void FD3D12ResourceLocation::Swap(FD3D12ResourceLocation& Other)
 {
+	// TODO: Probably shouldn't manually track suballocations. It's error-prone and inaccurate
+#if !PLATFORM_WINDOWS && ENABLE_LOW_LEVEL_MEM_TRACKER
+	const bool bRequiresManualTracking = GetType() == ResourceLocationType::eSubAllocation && AllocatorType != AT_SegList;
+	const bool bOtherRequiresManualTracking = Other.GetType() == ResourceLocationType::eSubAllocation && Other.AllocatorType != AT_SegList;
+
+	if (bRequiresManualTracking)
+	{
+		FLowLevelMemTracker::Get().OnLowLevelFree(ELLMTracker::Default, this);
+	}
+	if (bOtherRequiresManualTracking)
+	{
+		FLowLevelMemTracker::Get().OnLowLevelAllocMoved(ELLMTracker::Default, this, &Other);
+	}
+	if (bRequiresManualTracking)
+	{
+		FLowLevelMemTracker::Get().OnLowLevelAlloc(ELLMTracker::Default, &Other, GetSize());
+	}
+#endif
+
 	::Swap(*this, Other);
 }
 
