IRHICommandContext::RHIEnqueueStagedRead().
This removes the need to call flush on the RHI threads.

Removed IRHICommandContext::RHIInsertGPUFence() as its implementation is now within IRHICommandContext::RHIEnqueueStagedRead().
Removed FRHIGPUFence::Wait() as there is no use case for it.

NiagaraEmitterInstanceBatcher now use the normal command list instead of FRHICommandListExecutor::GetImmediateCommandList(),
improving support for future use.

Added global implementation for RHILockStagingBuffer() and RHIUnlockStagingBuffer()

diff --git a/Engine/Plugins/FX/Niagara/Source/Niagara/Classes/NiagaraEmitterInstanceBatcher.h b/Engine/Plugins/FX/Niagara/Source/Niagara/Classes/NiagaraEmitterInstanceBatcher.h
--- a/Engine/Plugins/FX/Niagara/Source/Niagara/Classes/NiagaraEmitterInstanceBatcher.h
+++ b/Engine/Plugins/FX/Niagara/Source/Niagara/Classes/NiagaraEmitterInstanceBatcher.h
@@ -88,13 +88,11 @@ public:
 	{
 		CurQueueIndex ^= 0x1;
 
-		// @TODO : we can't use RHICmdList here because we call EnqueueStagingRead which 
-		// require to flush the simulation dispatch before sending the copy to staging buffer command.
-		ExecuteAll(FRHICommandListExecutor::GetImmediateCommandList(), ViewUniformBuffer);
+		ExecuteAll(RHICmdList, ViewUniformBuffer);
 	}
 
-	void ExecuteAll(FRHICommandListImmediate &RHICmdList, FUniformBufferRHIParamRef ViewUniformBuffer);
-	void TickSingle(FNiagaraComputeExecutionContext *Context, FRHICommandListImmediate &RHICmdList, FUniformBufferRHIParamRef ViewUniformBuffer) const;
+	void ExecuteAll(FRHICommandList &RHICmdList, FUniformBufferRHIParamRef ViewUniformBuffer);
+	void TickSingle(FNiagaraComputeExecutionContext *Context, FRHICommandList &RHICmdList, FUniformBufferRHIParamRef ViewUniformBuffer) const;
 
 	void SetPrevDataStrideParams(const FNiagaraDataSet *Set, FNiagaraShader *Shader, FRHICommandList &RHICmdList) const;
 
@@ -114,7 +112,7 @@ public:
 	void RunEventHandlers(const FNiagaraComputeExecutionContext *Context, uint32 NumInstancesAfterSim, uint32 NumInstancesAfterSpawn, uint32 NumInstancesAfterNonEventSpawn, FRHICommandList &RhiCmdList) const;
 
 	void ClearIndexBufferCur(FRHICommandList &RHICmdList, FNiagaraComputeExecutionContext *Context) const;
-	void ResolveDatasetWrites(FRHICommandListImmediate &RHICmdList, FNiagaraComputeExecutionContext *Context) const;
+	void ResolveDatasetWrites(FRHICommandList &RHICmdList, FNiagaraComputeExecutionContext *Context) const;
 	void ResizeCurrentBuffer(FRHICommandList &RHICmdList, FNiagaraComputeExecutionContext *Context, uint32 NewNumInstances, uint32 PrevNumInstances) const;
 private:
 	static NiagaraEmitterInstanceBatcher* BatcherSingleton;
diff --git a/Engine/Plugins/FX/Niagara/Source/Niagara/Private/NiagaraEmitterInstanceBatcher.cpp b/Engine/Plugins/FX/Niagara/Source/Niagara/Private/NiagaraEmitterInstanceBatcher.cpp
--- a/Engine/Plugins/FX/Niagara/Source/Niagara/Private/NiagaraEmitterInstanceBatcher.cpp
+++ b/Engine/Plugins/FX/Niagara/Source/Niagara/Private/NiagaraEmitterInstanceBatcher.cpp
@@ -56,7 +56,7 @@ void NiagaraEmitterInstanceBatcher::Remove(FNiagaraComputeExecutionContext *InCo
 		});
 }
 
-void NiagaraEmitterInstanceBatcher::ExecuteAll(FRHICommandListImmediate &RHICmdList, FUniformBufferRHIParamRef ViewUniformBuffer)
+void NiagaraEmitterInstanceBatcher::ExecuteAll(FRHICommandList &RHICmdList, FUniformBufferRHIParamRef ViewUniformBuffer)
 {
 	const uint32 TickedQueueIndex = (CurQueueIndex ^ 0x1);
 	const uint32 TickedQueueIndexMask = (1 << TickedQueueIndex);
@@ -76,7 +76,7 @@ void NiagaraEmitterInstanceBatcher::ExecuteAll(FRHICommandListImmediate &RHICmdL
 	WorkQueue.Empty();
 }
 
-void NiagaraEmitterInstanceBatcher::TickSingle(FNiagaraComputeExecutionContext *Context, FRHICommandListImmediate &RHICmdList, FUniformBufferRHIParamRef ViewUniformBuffer) const
+void NiagaraEmitterInstanceBatcher::TickSingle(FNiagaraComputeExecutionContext *Context, FRHICommandList &RHICmdList, FUniformBufferRHIParamRef ViewUniformBuffer) const
 {
 	SCOPE_CYCLE_COUNTER(STAT_NiagaraGPUSimTick_RT);
 
@@ -178,7 +178,7 @@ void NiagaraEmitterInstanceBatcher::ClearIndexBufferCur(FRHICommandList &RHICmdL
  *	overallocating a bit, and the CPU might think we have more particles alive than we actually do;
  *	since we use DrawIndirect with the GPU determining draw call parameters, that's not an issue
  */
-void NiagaraEmitterInstanceBatcher::ResolveDatasetWrites(FRHICommandListImmediate &RHICmdList, FNiagaraComputeExecutionContext *Context) const
+void NiagaraEmitterInstanceBatcher::ResolveDatasetWrites(FRHICommandList &RHICmdList, FNiagaraComputeExecutionContext *Context) const
 {
 	FRWBuffer &DatasetIndexBufferWrite = Context->MainDataSet->GetCurDataSetIndices();
 	uint32 SpawnedThisFrame = Context->SpawnRateInstances + Context->EventSpawnTotal;
diff --git a/Engine/Source/Runtime/Apple/MetalRHI/Private/MetalUAV.cpp b/Engine/Source/Runtime/Apple/MetalRHI/Private/MetalUAV.cpp
--- a/Engine/Source/Runtime/Apple/MetalRHI/Private/MetalUAV.cpp
+++ b/Engine/Source/Runtime/Apple/MetalRHI/Private/MetalUAV.cpp
@@ -805,18 +805,50 @@ void FMetalRHICommandContext::RHIWaitComputeFence(FComputeFenceRHIParamRef InFen
 	}
 }
 
-void FMetalRHICommandContext::RHIInsertGPUFence(FGPUFenceRHIParamRef InFence)
+void FMetalGPUFence::WriteInternal(mtlpp::CommandBuffer& CmdBuffer)
+{
+	Fence = CmdBuffer.GetCompletionFence();
+	check(Fence);
+}
+
+void FMetalRHICommandContext::RHIEnqueueStagedRead(FStagingBufferRHIParamRef StagingBuffer, FGPUFenceRHIParamRef InFence, uint32 Offset, uint32 NumBytes)
 {
 	@autoreleasepool {
-	if (InFence)
-	{
-		FMetalGPUFence* Fence = ResourceCast(InFence);
-		Fence->Fence = Context->GetCurrentCommandBuffer().GetCompletionFence();
-		check(Fence->Fence);
-	}
+	
+		check(StagingBuffer);
+
+		FMetalStagingBuffer* StageBuffer = ResourceCast(StagingBuffer);
+		FMetalVertexBuffer* VertexBuffer = ResourceCast(StageBuffer->GetBackingBuffer());
+		switch (VertexBuffer->Buffer.GetStorageMode())
+		{
+	#if PLATFORM_MAC
+			case mtlpp::StorageMode::Managed:
+			{
+				GetMetalDeviceContext().SynchroniseResource(VertexBuffer->Buffer);
+				break;
+			}
+	#endif
+			case mtlpp::StorageMode::Private:
+			{
+				VertexBuffer->Alloc(VertexBuffer->Buffer.GetLength(), RLM_ReadOnly);
+				GetMetalDeviceContext().CopyFromBufferToBuffer(VertexBuffer->Buffer, Offset, VertexBuffer->CPUBuffer, Offset, NumBytes);
+				break;
+			}
+			default:
+			{
+				break;
+			}
+		}
+
+		if (InFence)
+		{
+			FMetalGPUFence* Fence = ResourceCast(InFence);
+			Fence->WriteInternal(Context->GetCurrentCommandBuffer());
+		}
 	}
 }
 
+
 FGPUFenceRHIRef FMetalDynamicRHI::RHICreateGPUFence(const FName &Name)
 {
 	@autoreleasepool {
@@ -835,13 +867,3 @@ bool FMetalGPUFence::Poll() const
 		return false;
 	}
 }
-
-bool FMetalGPUFence::Wait(float TimeoutMs) const
-{
-	if (Fence.GetPtr() == nil)
-	{
-		FRHICommandListExecutor::GetImmediateCommandList().ImmediateFlush(EImmediateFlushType::FlushRHIThread);
-	}
-	check(Fence);
-	return Fence.Wait(TimeoutMs);
-}
diff --git a/Engine/Source/Runtime/Apple/MetalRHI/Private/MetalVertexBuffer.cpp b/Engine/Source/Runtime/Apple/MetalRHI/Private/MetalVertexBuffer.cpp
--- a/Engine/Source/Runtime/Apple/MetalRHI/Private/MetalVertexBuffer.cpp
+++ b/Engine/Source/Runtime/Apple/MetalRHI/Private/MetalVertexBuffer.cpp
@@ -707,36 +707,6 @@ FVertexBufferRHIRef FMetalDynamicRHI::CreateVertexBuffer_RenderThread(class FRHI
 	}
 }
 
-void FMetalDynamicRHI::RHIEnqueueStagedRead(FStagingBufferRHIParamRef StagingBuffer, FGPUFenceRHIParamRef Fence, uint32 Offset, uint32 NumBytes)
-{
-	check(StagingBuffer && Fence);
-
-	FMetalStagingBuffer* StageBuffer = ResourceCast(StagingBuffer);
-	FMetalVertexBuffer* VertexBuffer = ResourceCast(StageBuffer->GetBackingBuffer());
-	switch (VertexBuffer->Buffer.GetStorageMode())
-	{
-#if PLATFORM_MAC
-		case mtlpp::StorageMode::Managed:
-		{
-			GetMetalDeviceContext().SynchroniseResource(VertexBuffer->Buffer);
-			break;
-		}
-#endif
-		case mtlpp::StorageMode::Private:
-		{
-			VertexBuffer->Alloc(VertexBuffer->Buffer.GetLength(), RLM_ReadOnly);
-			GetMetalDeviceContext().CopyFromBufferToBuffer(VertexBuffer->Buffer, Offset, VertexBuffer->CPUBuffer, Offset, NumBytes);
-			break;
-		}
-		default:
-		{
-			break;
-		}
-	}
-
-	RHIGetDefaultContext()->RHIInsertGPUFence(Fence);
-}
-
 void* FMetalDynamicRHI::RHILockStagingBuffer(FStagingBufferRHIParamRef StagingBuffer, uint32 Offset, uint32 SizeRHI)
 {
 	FMetalStagingBuffer* Buffer = ResourceCast(StagingBuffer);
diff --git a/Engine/Source/Runtime/Apple/MetalRHI/Public/MetalRHI.h b/Engine/Source/Runtime/Apple/MetalRHI/Public/MetalRHI.h
--- a/Engine/Source/Runtime/Apple/MetalRHI/Public/MetalRHI.h
+++ b/Engine/Source/Runtime/Apple/MetalRHI/Public/MetalRHI.h
@@ -194,7 +194,6 @@ public:
 	virtual void RHIUnlockStagingBuffer(FStagingBufferRHIParamRef StagingBuffer) final override;
 	virtual void* LockStagingBuffer_RenderThread(class FRHICommandListImmediate& RHICmdList, FStagingBufferRHIParamRef StagingBuffer, uint32 Offset, uint32 SizeRHI) final override;
 	virtual void UnlockStagingBuffer_RenderThread(class FRHICommandListImmediate& RHICmdList, FStagingBufferRHIParamRef StagingBuffer) final override;	
-	virtual void RHIEnqueueStagedRead(FStagingBufferRHIParamRef StagingBuffer, FGPUFenceRHIParamRef Fence, uint32 Offset, uint32 NumBytes) final override;
 	
 	virtual void RHISetResourceAliasability_RenderThread(class FRHICommandListImmediate& RHICmdList, EResourceAliasability AliasMode, FTextureRHIParamRef* InTextures, int32 NumTextures) final override;
 	
diff --git a/Engine/Source/Runtime/Apple/MetalRHI/Public/MetalRHIContext.h b/Engine/Source/Runtime/Apple/MetalRHI/Public/MetalRHIContext.h
--- a/Engine/Source/Runtime/Apple/MetalRHI/Public/MetalRHIContext.h
+++ b/Engine/Source/Runtime/Apple/MetalRHI/Public/MetalRHIContext.h
@@ -22,8 +22,6 @@ public:
 	
 	/** Get the profiler pointer */
 	FORCEINLINE class FMetalProfiler* GetProfiler() const { return Profiler; }
-	
-	virtual void RHIInsertGPUFence(FGPUFenceRHIParamRef Fence) final override;
 
 	/**
 	 *Sets the current compute shader.  Mostly for compliance with platforms
@@ -320,6 +318,8 @@ public:
 	 */
 	virtual void RHITransitionResources(EResourceTransitionAccess TransitionType, EResourceTransitionPipeline TransitionPipeline, FUnorderedAccessViewRHIParamRef* InUAVs, int32 NumUAVs, FComputeFenceRHIParamRef WriteComputeFence) final override;
 	
+	virtual void RHIEnqueueStagedRead(FStagingBufferRHIParamRef StagingBuffer, FGPUFenceRHIParamRef InFence, uint32 Offset, uint32 NumBytes) final override;
+
 	/**
 	 * Compute queue will wait for the fence to be written before continuing.
 	 */
diff --git a/Engine/Source/Runtime/Apple/MetalRHI/Public/MetalResources.h b/Engine/Source/Runtime/Apple/MetalRHI/Public/MetalResources.h
--- a/Engine/Source/Runtime/Apple/MetalRHI/Public/MetalResources.h
+++ b/Engine/Source/Runtime/Apple/MetalRHI/Public/MetalResources.h
@@ -991,13 +991,13 @@ public:
 
 	~FMetalGPUFence()
 	{
-
 	}
 
-	virtual bool Poll() const override final;
+	void WriteInternal(mtlpp::CommandBuffer& CmdBuffer);
 
-	virtual bool Wait(float TimeoutMs) const override final;
+	virtual bool Poll() const override final;
 
+private:
 	mtlpp::CommandBufferFence Fence;
 };
 
diff --git a/Engine/Source/Runtime/D3D12RHI/Private/D3D12CommandContext.h b/Engine/Source/Runtime/D3D12RHI/Private/D3D12CommandContext.h
--- a/Engine/Source/Runtime/D3D12RHI/Private/D3D12CommandContext.h
+++ b/Engine/Source/Runtime/D3D12RHI/Private/D3D12CommandContext.h
@@ -263,6 +263,7 @@ public:
 	virtual void RHIClearTinyUAV(FUnorderedAccessViewRHIParamRef UnorderedAccessViewRHI, const uint32* Values) final override;
 	virtual void RHICopyToResolveTarget(FTextureRHIParamRef SourceTexture, FTextureRHIParamRef DestTexture, const FResolveParams& ResolveParams) final override;
 	virtual void RHITransitionResources(EResourceTransitionAccess TransitionType, FTextureRHIParamRef* InTextures, int32 NumTextures) final override;
+    virtual void RHIEnqueueStagedRead(FStagingBufferRHIParamRef StagingBuffer, FGPUFenceRHIParamRef Fence, uint32 Offset, uint32 NumBytes) final override;
 	virtual void RHIBeginRenderQuery(FRenderQueryRHIParamRef RenderQuery) final override;
 	virtual void RHIEndRenderQuery(FRenderQueryRHIParamRef RenderQuery) final override;
 	void RHIBeginOcclusionQueryBatch(uint32 NumQueriesInBatch);
@@ -449,6 +450,10 @@ public:
 	// Special implementation that only signal the fence once.
 	virtual void RHITransitionResources(EResourceTransitionAccess TransitionType, EResourceTransitionPipeline TransitionPipeline, FUnorderedAccessViewRHIParamRef* InUAVs, int32 NumUAVs, FComputeFenceRHIParamRef WriteComputeFenceRHI) final override;
 
+	FORCEINLINE virtual void RHIEnqueueStagedRead(FStagingBufferRHIParamRef StagingBuffer, FGPUFenceRHIParamRef Fence, uint32 Offset, uint32 NumBytes) final override
+	{
+		ContextRedirect(RHIEnqueueStagedRead(StagingBuffer, Fence, Offset, NumBytes));
+	}
 	FORCEINLINE virtual void RHISetShaderTexture(FComputeShaderRHIParamRef PixelShader, uint32 TextureIndex, FTextureRHIParamRef NewTexture) final override
 	{
 		ContextRedirect(RHISetShaderTexture(PixelShader, TextureIndex, NewTexture));
diff --git a/Engine/Source/Runtime/D3D12RHI/Private/D3D12Commands.cpp b/Engine/Source/Runtime/D3D12RHI/Private/D3D12Commands.cpp
--- a/Engine/Source/Runtime/D3D12RHI/Private/D3D12Commands.cpp
+++ b/Engine/Source/Runtime/D3D12RHI/Private/D3D12Commands.cpp
@@ -395,6 +395,44 @@ void FD3D12CommandContext::RHITransitionResources(EResourceTransitionAccess Tran
 	}
 }
 
+void FD3D12CommandContext::RHIEnqueueStagedRead(FStagingBufferRHIParamRef StagingBufferRHI, FGPUFenceRHIParamRef FenceRHI, uint32 Offset, uint32 NumBytes)
+{
+	FD3D12StagingBuffer* StagingBuffer = FD3D12DynamicRHI::ResourceCast(StagingBufferRHI);
+	check(StagingBuffer);
+
+	FD3D12VertexBuffer* VertexBuffer = FD3D12DynamicRHI::ResourceCast(StagingBuffer->GetBackingBuffer());
+	check(VertexBuffer);
+
+	// Only get data from the first gpu for now.
+	FD3D12Device* StagingDevice = VertexBuffer->GetParentDevice();
+	StagingBuffer->StagedRead.SafeRelease();
+	VERIFYD3D12RESULT(GetParentDevice()->GetParentAdapter()->CreateBuffer(D3D12_HEAP_TYPE_READBACK, GetGPUMask(), GetGPUMask(), Offset + NumBytes, StagingBuffer->StagedRead.GetInitReference()));
+	
+	{
+		FD3D12Resource* pSourceResource = VertexBuffer->ResourceLocation.GetResource();
+		D3D12_RESOURCE_DESC const& SourceBufferDesc = pSourceResource->GetDesc();
+
+		FD3D12Resource* pDestResource = StagingBuffer->StagedRead.GetReference();
+		D3D12_RESOURCE_DESC const& DestBufferDesc = pDestResource->GetDesc();
+
+		FD3D12DynamicRHI::TransitionResource(CommandListHandle, pSourceResource, D3D12_RESOURCE_STATE_COPY_SOURCE, 0);
+		CommandListHandle.FlushResourceBarriers();	// Must flush so the desired state is actually set.
+
+		numCopies++;
+
+		CommandListHandle->CopyBufferRegion(pDestResource->GetResource(), Offset, pSourceResource->GetResource(), Offset, NumBytes);
+		CommandListHandle.UpdateResidency(pDestResource);
+		CommandListHandle.UpdateResidency(pSourceResource);
+	}
+
+	if (FenceRHI)
+	{
+		RHISubmitCommandsHint();
+		FD3D12GPUFence* Fence = FD3D12DynamicRHI::ResourceCast(FenceRHI);
+		Fence->WriteInternal(ED3D12CommandQueueType::Default);
+	}
+}
+
 void FD3D12CommandContext::RHISetViewport(uint32 MinX, uint32 MinY, float MinZ, uint32 MaxX, uint32 MaxY, float MaxZ)
 {
 	// These are the maximum viewport extents for D3D12. Exceeding them leads to badness.
diff --git a/Engine/Source/Runtime/D3D12RHI/Private/D3D12DirectCommandListManager.cpp b/Engine/Source/Runtime/D3D12RHI/Private/D3D12DirectCommandListManager.cpp
--- a/Engine/Source/Runtime/D3D12RHI/Private/D3D12DirectCommandListManager.cpp
+++ b/Engine/Source/Runtime/D3D12RHI/Private/D3D12DirectCommandListManager.cpp
@@ -14,51 +14,18 @@ FComputeFenceRHIRef FD3D12DynamicRHI::RHICreateComputeFence(const FName& Name)
 	return Fence;
 }
 
-// =============================================================================
-
-class FD3D12GPUFence : public FRHIGPUFence
+void FD3D12GPUFence::WriteInternal(ED3D12CommandQueueType QueueType)
 {
-public:
-
-	FD3D12GPUFence(FName InName, FD3D12Fence* InFence) 
-		: FRHIGPUFence(InName)
-		, Fence(InFence)
-		, Value(0)
-	{}
-
-	virtual void Write() final override
-	{
-		if (Fence)
-		{
-			Value = Fence->Signal(ED3D12CommandQueueType::Default);
-		}
-	}
-
-	virtual bool Poll() const final override
-	{
-		return !Value || (Fence && Fence->IsFenceComplete(Value));
-	}
-
-	virtual bool Wait(float TimeoutMs) const final override
+	if (Fence)
 	{
-		if (Value && Fence)
-		{
-			Fence->WaitForFence(Value);
-		}
-		return Value != 0;
+		Value = Fence->Signal(QueueType);
 	}
+}
 
-protected:
-
-	TRefCountPtr<FD3D12Fence> Fence;
-	uint64 Value;
-};
-
-template<>
-struct TD3D12ResourceTraits<FRHIGPUFence>
+bool FD3D12GPUFence::Poll() const
 {
-	typedef FD3D12GPUFence TConcreteType;
-};
+	return !Value || (Fence && Fence->IsFenceComplete(Value));
+}
 
 
 FGPUFenceRHIRef FD3D12DynamicRHI::RHICreateGPUFence(const FName& Name)
@@ -66,74 +33,11 @@ FGPUFenceRHIRef FD3D12DynamicRHI::RHICreateGPUFence(const FName& Name)
 	return new FD3D12GPUFence(Name, GetAdapter().GetStagingFence());
 }
 
-// =============================================================================
-
-class FD3D12StagingBuffer : public FRHIStagingBuffer
-{
-public:
-
-	FD3D12StagingBuffer(FVertexBufferRHIRef InBuffer)
-		: FRHIStagingBuffer(InBuffer)
-	{
-	}
-
-	TRefCountPtr<FD3D12Resource> StagedRead;
-};
-
-template<>
-struct TD3D12ResourceTraits<FRHIStagingBuffer>
-{
-	typedef FD3D12StagingBuffer TConcreteType;
-};
-
-
 FStagingBufferRHIRef FD3D12DynamicRHI::RHICreateStagingBuffer(FVertexBufferRHIParamRef VertexBufferRHI)
 {
 	return new FD3D12StagingBuffer(VertexBufferRHI);
 }
 
-void FD3D12DynamicRHI::RHIEnqueueStagedRead(FStagingBufferRHIParamRef StagingBufferRHI, FGPUFenceRHIParamRef FenceRHI, uint32 Offset, uint32 NumBytes)
-{
-	FD3D12StagingBuffer* StagingBuffer = FD3D12DynamicRHI::ResourceCast(StagingBufferRHI);
-	check(StagingBuffer);
-
-	FD3D12VertexBuffer* VertexBuffer = FD3D12DynamicRHI::ResourceCast(StagingBuffer->GetBackingBuffer());
-	check(VertexBuffer);
-
-	// Only get data from the first gpu for now.
-	FD3D12Device* StagingDevice = VertexBuffer->GetParentDevice();
-	FRHIGPUMask StagingGPUMask = StagingDevice->GetGPUMask();
-	StagingBuffer->StagedRead.SafeRelease();
-	VERIFYD3D12RESULT(GetAdapter().CreateBuffer(D3D12_HEAP_TYPE_READBACK, StagingGPUMask, StagingGPUMask, Offset + NumBytes, StagingBuffer->StagedRead.GetInitReference()));
-	
-	FD3D12CommandContext& Context = StagingDevice->GetDefaultCommandContext();
-
-	{
-		FD3D12Resource* pSourceResource = VertexBuffer->ResourceLocation.GetResource();
-		D3D12_RESOURCE_DESC const& SourceBufferDesc = pSourceResource->GetDesc();
-
-		FD3D12Resource* pDestResource = StagingBuffer->StagedRead.GetReference();
-		D3D12_RESOURCE_DESC const& DestBufferDesc = pDestResource->GetDesc();
-
-		TransitionResource(Context.CommandListHandle, pSourceResource, D3D12_RESOURCE_STATE_COPY_SOURCE, 0);
-		// TransitionResource(Context.CommandListHandle, pDestResource, D3D12_RESOURCE_STATE_COPY_DEST, 0);
-		Context.CommandListHandle.FlushResourceBarriers();	// Must flush so the desired state is actually set.
-
-		Context.numCopies++;
-
-		Context.CommandListHandle->CopyBufferRegion(pDestResource->GetResource(), Offset, pSourceResource->GetResource(), Offset, NumBytes);
-		Context.CommandListHandle.UpdateResidency(pDestResource);
-		Context.CommandListHandle.UpdateResidency(pSourceResource);
-
-		Context.FlushCommands();
-
-}
-
-	// Signal the fence
-	FenceRHI->Write();
-}
-
-
 void* FD3D12DynamicRHI::RHILockStagingBuffer(FStagingBufferRHIParamRef StagingBufferRHI, uint32 Offset, uint32 SizeRHI)
 {
 	FD3D12StagingBuffer* StagingBuffer = FD3D12DynamicRHI::ResourceCast(StagingBufferRHI);
diff --git a/Engine/Source/Runtime/D3D12RHI/Private/D3D12RHIPrivate.h b/Engine/Source/Runtime/D3D12RHI/Private/D3D12RHIPrivate.h
--- a/Engine/Source/Runtime/D3D12RHI/Private/D3D12RHIPrivate.h
+++ b/Engine/Source/Runtime/D3D12RHI/Private/D3D12RHIPrivate.h
@@ -236,7 +236,6 @@ public:
 	virtual FComputeFenceRHIRef RHICreateComputeFence(const FName& Name) final override;
 	virtual FGPUFenceRHIRef RHICreateGPUFence(const FName& Name) final override;
 	virtual FStagingBufferRHIRef RHICreateStagingBuffer(FVertexBufferRHIParamRef VertexBuffer) final override;
-    virtual void RHIEnqueueStagedRead(FStagingBufferRHIParamRef StagingBuffer, FGPUFenceRHIParamRef Fence, uint32 Offset, uint32 NumBytes) final override;
 	virtual void* RHILockStagingBuffer(FStagingBufferRHIParamRef StagingBuffer, uint32 Offset, uint32 SizeRHI) final override;
     virtual void RHIUnlockStagingBuffer(FStagingBufferRHIParamRef StagingBuffer) final override;
 	virtual FBoundShaderStateRHIRef RHICreateBoundShaderState(FVertexDeclarationRHIParamRef VertexDeclaration, FVertexShaderRHIParamRef VertexShader, FHullShaderRHIParamRef HullShader, FDomainShaderRHIParamRef DomainShader, FPixelShaderRHIParamRef PixelShader, FGeometryShaderRHIParamRef GeometryShader) final override;
diff --git a/Engine/Source/Runtime/D3D12RHI/Public/D3D12Resources.h b/Engine/Source/Runtime/D3D12RHI/Public/D3D12Resources.h
--- a/Engine/Source/Runtime/D3D12RHI/Public/D3D12Resources.h
+++ b/Engine/Source/Runtime/D3D12RHI/Public/D3D12Resources.h
@@ -848,6 +848,34 @@ private:
 	FD3D12ResourceLocation ResourceLocation;
 };
 
+class FD3D12StagingBuffer : public FRHIStagingBuffer
+{
+public:
+	FD3D12StagingBuffer(FVertexBufferRHIRef InBuffer)
+		: FRHIStagingBuffer(InBuffer)
+	{}
+
+	TRefCountPtr<FD3D12Resource> StagedRead;
+};
+
+class FD3D12GPUFence : public FRHIGPUFence
+{
+public:
+	FD3D12GPUFence(FName InName, FD3D12Fence* InFence) 
+		: FRHIGPUFence(InName)
+		, Fence(InFence)
+		, Value(0)
+	{}
+
+	void WriteInternal(ED3D12CommandQueueType QueueType);
+	virtual bool Poll() const final override;
+
+protected:
+
+	TRefCountPtr<FD3D12Fence> Fence;
+	uint64 Value;
+};
+
 template<class T>
 struct TD3D12ResourceTraits
 {
@@ -907,3 +935,14 @@ struct TD3D12ResourceTraits<FRHIComputePipelineState>
 {
 	typedef FD3D12ComputePipelineState TConcreteType;
 };
+template<>
+struct TD3D12ResourceTraits<FRHIGPUFence>
+{
+	typedef FD3D12GPUFence TConcreteType;
+};
+template<>
+struct TD3D12ResourceTraits<FRHIStagingBuffer>
+{
+	typedef FD3D12StagingBuffer TConcreteType;
+};
+
diff --git a/Engine/Source/Runtime/RHI/Private/RHICommandList.cpp b/Engine/Source/Runtime/RHI/Private/RHICommandList.cpp
--- a/Engine/Source/Runtime/RHI/Private/RHICommandList.cpp
+++ b/Engine/Source/Runtime/RHI/Private/RHICommandList.cpp
@@ -2071,11 +2071,11 @@ void FDynamicRHI::UnlockIndexBuffer_RenderThread(class FRHICommandListImmediate&
 
 void* FDynamicRHI::RHILockStagingBuffer(FStagingBufferRHIParamRef StagingBuffer, uint32 Offset, uint32 SizeRHI)
 {
-	return RHILockVertexBuffer(StagingBuffer->GetBackingBuffer(), Offset, SizeRHI, RLM_ReadOnly);
+	return GDynamicRHI->RHILockVertexBuffer(StagingBuffer->GetBackingBuffer(), Offset, SizeRHI, RLM_ReadOnly);
 }
 void FDynamicRHI::RHIUnlockStagingBuffer(FStagingBufferRHIParamRef StagingBuffer)
 {
-	RHIUnlockVertexBuffer(StagingBuffer->GetBackingBuffer());
+	GDynamicRHI->RHIUnlockVertexBuffer(StagingBuffer->GetBackingBuffer());
 }
 
 void* FDynamicRHI::LockStagingBuffer_RenderThread(class FRHICommandListImmediate& RHICmdList, FStagingBufferRHIParamRef StagingBuffer, uint32 Offset, uint32 SizeRHI)
@@ -2084,7 +2084,7 @@ void* FDynamicRHI::LockStagingBuffer_RenderThread(class FRHICommandListImmediate
 	check(IsInRenderingThread());
 
 	FScopedRHIThreadStaller StallRHIThread(RHICmdList);
-	return RHILockStagingBuffer(StagingBuffer, Offset, SizeRHI);
+	return GDynamicRHI->RHILockStagingBuffer(StagingBuffer, Offset, SizeRHI);
 }
 void FDynamicRHI::UnlockStagingBuffer_RenderThread(class FRHICommandListImmediate& RHICmdList, FStagingBufferRHIParamRef StagingBuffer)
 {
@@ -2092,7 +2092,7 @@ void FDynamicRHI::UnlockStagingBuffer_RenderThread(class FRHICommandListImmediat
 	check(IsInRenderingThread());
 	
 	FScopedRHIThreadStaller StallRHIThread(RHICmdList);
-	RHIUnlockStagingBuffer(StagingBuffer);
+	GDynamicRHI->RHIUnlockStagingBuffer(StagingBuffer);
 }
 
 FTexture2DRHIRef FDynamicRHI::AsyncReallocateTexture2D_RenderThread(class FRHICommandListImmediate& RHICmdList, FTexture2DRHIParamRef Texture2D, int32 NewMipCount, int32 NewSizeX, int32 NewSizeY, FThreadSafeCounter* RequestStatus)
@@ -2478,12 +2478,3 @@ void FDynamicRHI::RHICopySubTextureRegion_RenderThread(class FRHICommandListImme
 	FScopedRHIThreadStaller StallRHIThread(RHICmdList);
 	return GDynamicRHI->RHICopySubTextureRegion(SourceTexture, DestinationTexture, SourceBox, DestinationBox);
 }
-
-void FDynamicRHI::RHIEnqueueStagedRead(FStagingBufferRHIParamRef StagingBuffer, FGPUFenceRHIParamRef Fence, uint32 Offset, uint32 NumBytes)
-{
-	IRHICommandContext* Context = RHIGetDefaultContext();
-	check(Context);
-	check(Fence);
-	Context->RHIInsertGPUFence(Fence);
-}
-
diff --git a/Engine/Source/Runtime/RHI/Private/RHIGPUReadback.cpp b/Engine/Source/Runtime/RHI/Private/RHIGPUReadback.cpp
--- a/Engine/Source/Runtime/RHI/Private/RHIGPUReadback.cpp
+++ b/Engine/Source/Runtime/RHI/Private/RHIGPUReadback.cpp
@@ -32,12 +32,3 @@ bool FGenericRHIGPUFence::Poll() const
 	}
 	return false;
 }
-
-bool FGenericRHIGPUFence::Wait(float TimeoutMs) const
-{
-	if (GFrameNumberRenderThread > InsertedFrameNumber)
-	{
-		return true;
-	}
-	return false;
-}
diff --git a/Engine/Source/Runtime/RHI/Public/DynamicRHI.h b/Engine/Source/Runtime/RHI/Public/DynamicRHI.h
--- a/Engine/Source/Runtime/RHI/Public/DynamicRHI.h
+++ b/Engine/Source/Runtime/RHI/Public/DynamicRHI.h
@@ -1045,17 +1045,7 @@ public:
 	/* Copy the source box pixels in the destination box texture, return true if implemented for the current platform*/
 	virtual void RHICopySubTextureRegion_RenderThread(class FRHICommandListImmediate& RHICmdList, FTexture2DRHIParamRef SourceTexture, FTexture2DRHIParamRef DestinationTexture, FBox2D SourceBox, FBox2D DestinationBox);
 	virtual void RHICopySubTextureRegion(FTexture2DRHIParamRef SourceTexture, FTexture2DRHIParamRef DestinationTexture, FBox2D SourceBox, FBox2D DestinationBox) { }
-	
-    /**
-     * Enqueues on the GPU timeline any necessary operations to make the contents of 'StagingBuffer' accessible to the CPU, flushing outstanding GPU writes and/or transferring from inaccessible non-unified GPU memory to local CPU memory.
-     * @discussion The RHI will call RHIInsertGPUFence on 'Fence', use FRHIGPUFence::Wait or FRHIGPUFence::Poll to determine when the fence has been signalled by the GPU. Only once this occurs may 'StagingBuffer' be locked.
-     * @param StagingBuffer The buffer to stage. Must not be null.
-     * @param Fence A GPU fence that will be inserted into the GPU timeline and which must then be tested on the CPU to know when the StagedRead was completed. Must not be null.
-     * @param Offset The start of the data in 'StagingBuffer' to make available.
-     * @param NumBytes The lenght of data in 'StagingBuffer' to make available.
-     */
-    virtual void RHIEnqueueStagedRead(FStagingBufferRHIParamRef StagingBuffer, FGPUFenceRHIParamRef Fence, uint32 Offset, uint32 NumBytes);
-	
+
 	virtual FRHIFlipDetails RHIWaitForFlip(double TimeoutInSeconds) { return FRHIFlipDetails(); }
 	virtual void RHISignalFlipEvent() { }
 
diff --git a/Engine/Source/Runtime/RHI/Public/RHICommandList.h b/Engine/Source/Runtime/RHI/Public/RHICommandList.h
--- a/Engine/Source/Runtime/RHI/Public/RHICommandList.h
+++ b/Engine/Source/Runtime/RHI/Public/RHICommandList.h
@@ -1459,12 +1459,18 @@ struct FRHICommandWaitComputeFence final : public FRHICommand<FRHICommandWaitCom
 };
 
 template<ECmdList CmdListType>
-struct FRHICommandInsertGPUFence final : public FRHICommand<FRHICommandInsertGPUFence<CmdListType>>
+struct FRHICommandEnqueueStagedRead final : public FRHICommand<FRHICommandEnqueueStagedRead<CmdListType>>
 {
+	FStagingBufferRHIParamRef StagingBuffer;
 	FGPUFenceRHIParamRef Fence;
+	uint32 Offset;
+	uint32 NumBytes;
 
-	FORCEINLINE_DEBUGGABLE FRHICommandInsertGPUFence(FGPUFenceRHIParamRef InFence)
-		: Fence(InFence)
+	FORCEINLINE_DEBUGGABLE FRHICommandEnqueueStagedRead(FStagingBufferRHIParamRef InStagingBuffer, FGPUFenceRHIParamRef InFence, uint32 InOffset, uint32 InNumBytes)
+		: StagingBuffer(InStagingBuffer)
+		, Fence(InFence)
+		, Offset(InOffset)
+		, NumBytes(InNumBytes)
 	{		
 	}
 	RHI_API void Execute(FRHICommandListBase& CmdList);
@@ -2703,14 +2709,14 @@ public:
 		new (AllocCommand<FRHICommandWaitComputeFence<ECmdList::EGfx>>()) FRHICommandWaitComputeFence<ECmdList::EGfx>(WaitFence);
 	}
 
-	FORCEINLINE_DEBUGGABLE void InsertGPUFence(FGPUFenceRHIParamRef Fence)
+	FORCEINLINE_DEBUGGABLE void EnqueueStagedRead(FStagingBufferRHIParamRef StagingBuffer, FGPUFenceRHIParamRef Fence, uint32 Offset, uint32 NumBytes)
 	{
 		if (Bypass())
 		{
-			CMD_CONTEXT(RHIInsertGPUFence)(Fence);
+			CMD_CONTEXT(RHIEnqueueStagedRead)(StagingBuffer, Fence, Offset, NumBytes);
 			return;
 		}
-		new (AllocCommand<FRHICommandInsertGPUFence<ECmdList::EGfx>>()) FRHICommandInsertGPUFence<ECmdList::EGfx>(Fence);
+		new (AllocCommand<FRHICommandEnqueueStagedRead<ECmdList::EGfx>>()) FRHICommandEnqueueStagedRead<ECmdList::EGfx>(StagingBuffer, Fence, Offset, NumBytes);
 	}
 
 	FORCEINLINE_DEBUGGABLE void BeginRenderPass(const FRHIRenderPassInfo& InInfo, const TCHAR* Name)
@@ -3098,14 +3104,14 @@ public:
 		new (AllocCommand<FRHICommandWaitComputeFence<ECmdList::ECompute>>()) FRHICommandWaitComputeFence<ECmdList::ECompute>(WaitFence);
 	}
 
-	FORCEINLINE_DEBUGGABLE void InsertGPUFence(FGPUFenceRHIParamRef Fence)
+	FORCEINLINE_DEBUGGABLE void EnqueueStagedRead(FStagingBufferRHIParamRef StagingBuffer, FGPUFenceRHIParamRef Fence, uint32 Offset, uint32 NumBytes)
 	{
 		if (Bypass())
 		{
-			COMPUTE_CONTEXT(RHIInsertGPUFence)(Fence);
+			CMD_CONTEXT(RHIEnqueueStagedRead)(StagingBuffer, Fence, Offset, NumBytes);
 			return;
 		}
-		new (AllocCommand<FRHICommandInsertGPUFence<ECmdList::ECompute>>()) FRHICommandInsertGPUFence<ECmdList::ECompute>(Fence);
+		new (AllocCommand<FRHICommandEnqueueStagedRead<ECmdList::ECompute>>()) FRHICommandEnqueueStagedRead<ECmdList::ECompute>(StagingBuffer, Fence, Offset, NumBytes);
 	}
 };
 
@@ -3988,12 +3994,6 @@ public:
 	{
 		GDynamicRHI->RHICopySubTextureRegion_RenderThread(*this, SourceTexture, DestinationTexture, SourceBox, DestinationBox);
 	}
-
-	FORCEINLINE void EnqueueStagedRead(FStagingBufferRHIParamRef StagingBuffer, FGPUFenceRHIParamRef Fence, uint32 Offset, uint32 NumBytes)
-	{
-		ImmediateFlush(EImmediateFlushType::FlushRHIThread); 
-		GDynamicRHI->RHIEnqueueStagedRead(StagingBuffer, Fence, Offset, NumBytes);
-	}
 	
 	FORCEINLINE void ExecuteCommandList(FRHICommandList* CmdList)
 	{
@@ -4583,5 +4583,15 @@ FORCEINLINE FRHIShaderLibraryRef RHICreateShaderLibrary(EShaderPlatform Platform
     return GDynamicRHI->RHICreateShaderLibrary(Platform, FilePath, Name);
 }
 
+FORCEINLINE void* RHILockStagingBuffer(FStagingBufferRHIParamRef StagingBuffer, uint32 Offset, uint32 Size)
+{
+	return FRHICommandListExecutor::GetImmediateCommandList().LockStagingBuffer(StagingBuffer, Offset, Size);
+}
+
+FORCEINLINE void RHIUnlockStagingBuffer(FStagingBufferRHIParamRef StagingBuffer)
+{
+	 FRHICommandListExecutor::GetImmediateCommandList().UnlockStagingBuffer(StagingBuffer);
+}
+
 
 #include "RHICommandList.inl"
diff --git a/Engine/Source/Runtime/RHI/Public/RHICommandListCommandExecutes.inl b/Engine/Source/Runtime/RHI/Public/RHICommandListCommandExecutes.inl
--- a/Engine/Source/Runtime/RHI/Public/RHICommandListCommandExecutes.inl
+++ b/Engine/Source/Runtime/RHI/Public/RHICommandListCommandExecutes.inl
@@ -494,13 +494,13 @@ template struct FRHICommandWaitComputeFence<ECmdList::EGfx>;
 template struct FRHICommandWaitComputeFence<ECmdList::ECompute>;
 
 template<ECmdList CmdListType>
-void FRHICommandInsertGPUFence<CmdListType>::Execute(FRHICommandListBase& CmdList)
+void FRHICommandEnqueueStagedRead<CmdListType>::Execute(FRHICommandListBase& CmdList)
 {
-	RHISTAT(InsertGPUFence);
-	INTERNAL_DECORATOR_CONTEXT(RHIInsertGPUFence)(Fence);
+	RHISTAT(EnqueueStagedRead);
+	INTERNAL_DECORATOR_CONTEXT(RHIEnqueueStagedRead)(StagingBuffer, Fence, Offset, NumBytes);
 }
-template struct FRHICommandInsertGPUFence<ECmdList::EGfx>;
-template struct FRHICommandInsertGPUFence<ECmdList::ECompute>;
+template struct FRHICommandEnqueueStagedRead<ECmdList::EGfx>;
+template struct FRHICommandEnqueueStagedRead<ECmdList::ECompute>;
 
 void FRHICommandBuildLocalGraphicsPipelineState::Execute(FRHICommandListBase& CmdList)
 {
diff --git a/Engine/Source/Runtime/RHI/Public/RHIContext.h b/Engine/Source/Runtime/RHI/Public/RHIContext.h
--- a/Engine/Source/Runtime/RHI/Public/RHIContext.h
+++ b/Engine/Source/Runtime/RHI/Public/RHIContext.h
@@ -127,16 +127,19 @@ public:
 	 */
 	virtual void RHIInvalidateCachedState() {};
 
-
-	/**
-	 * Insert a GPU-to-CPU fence into the GPU command-stream. 
-	 * The default implementation will simply write the current frame-number into the fence.
-	 * RHIs that can implement true GPU fences should override this to do so.
-	 */
-	virtual void RHIInsertGPUFence(FGPUFenceRHIParamRef Fence) 
-	{ 
-		check(Fence); 
-		Fence->Write(); 
+    /**
+     * Enqueues on the GPU timeline any necessary operations to make the contents of 'StagingBuffer' accessible to the CPU, flushing outstanding GPU writes and/or transferring from inaccessible non-unified GPU memory to local CPU memory.
+     * @param StagingBuffer The buffer to stage. Must not be null.
+     * @param Fence A GPU fence that will be inserted into the GPU timeline and which must then be tested on the CPU to know when the StagedRead was completed. Must not be null.
+     * @param Offset The start of the data in 'StagingBuffer' to make available.
+     * @param NumBytes The lenght of data in 'StagingBuffer' to make available.
+     */
+	virtual void RHIEnqueueStagedRead(FStagingBufferRHIParamRef StagingBuffer, FGPUFenceRHIParamRef Fence, uint32 InOffset, uint32 InNumBytes)
+	{
+		if (Fence)
+		{
+			Fence->Write();
+		}
 	}
 };
 
@@ -258,7 +261,6 @@ public:
 		RHITransitionResources(TransitionType, TransitionPipeline, InUAVs, NumUAVs, nullptr);
 	}
 
-
 	virtual void RHIBeginRenderQuery(FRenderQueryRHIParamRef RenderQuery) = 0;
 
 	virtual void RHIEndRenderQuery(FRenderQueryRHIParamRef RenderQuery) = 0;
diff --git a/Engine/Source/Runtime/RHI/Public/RHIGPUReadback.h b/Engine/Source/Runtime/RHI/Public/RHIGPUReadback.h
--- a/Engine/Source/Runtime/RHI/Public/RHIGPUReadback.h
+++ b/Engine/Source/Runtime/RHI/Public/RHIGPUReadback.h
@@ -29,14 +29,9 @@ public:
 	{
 	}
 
-	void Insert(FRHICommandListImmediate& RHICmdList)
+	void Insert(FRHICommandList& RHICmdList, uint32 NumBytes = 0)
 	{
-		RHICmdList.EnqueueStagedRead(StagingBuffer, Fence, 0, GPUVertexBuffer->GetSize());
-	}
-
-	bool WaitForFence(float Timeout)
-	{
-		return Fence->Wait(Timeout);
+		RHICmdList.EnqueueStagedRead(StagingBuffer, Fence, 0, NumBytes ? NumBytes : GPUVertexBuffer->GetSize());
 	}
 
 	bool IsReady()
@@ -47,14 +42,12 @@ public:
 	void *RetrieveData(uint32 NumBytes)
 	{
 		ensure(Fence->Poll());
-		FRHICommandListImmediate& RHICmdList = FRHICommandListExecutor::GetImmediateCommandList();
-		return RHICmdList.LockStagingBuffer(StagingBuffer, 0, NumBytes);
+		return RHILockStagingBuffer(StagingBuffer, 0, NumBytes);
 	}
 
 	void Finish()
 	{
-		FRHICommandListImmediate& RHICmdList = FRHICommandListExecutor::GetImmediateCommandList();
-		RHICmdList.UnlockStagingBuffer(StagingBuffer);
+		RHIUnlockStagingBuffer(StagingBuffer);
 	}
 
 private:
diff --git a/Engine/Source/Runtime/RHI/Public/RHIResources.h b/Engine/Source/Runtime/RHI/Public/RHIResources.h
--- a/Engine/Source/Runtime/RHI/Public/RHIResources.h
+++ b/Engine/Source/Runtime/RHI/Public/RHIResources.h
@@ -795,13 +795,6 @@ public:
      */
 	virtual bool Poll() const = 0;
 
-    /**
-     * Wait for the GPU to pass and signal the fence.
-     * @param TimeoutMs The maximum time to wait for the fence in milliseconds.
-     * @returns True if and only if the GPU fence has been inserted and the GPU has signalled the fence.
-     */
-	virtual bool Wait(float TimeoutMs) const = 0;
-
 private:
 	FName FenceName;
 };
@@ -817,9 +810,6 @@ public:
     /** @discussion RHI implementations must be thread-safe and must correctly handle being called before RHIInsertFence if an RHI thread is active. */
 	virtual bool Poll() const final override;
 
-    /** @discussion RHI implementations must be thread-safe and must correctly handle being called before RHIInsertFence if an RHI thread is active. */
-	virtual bool Wait(float TimeoutMs) const final override;
-
 private:
 	uint32 InsertedFrameNumber;
 };
@@ -1034,7 +1024,6 @@ protected:
 typedef FRHIStagingBuffer*				FStagingBufferRHIParamRef;
 typedef TRefCountPtr<FRHIStagingBuffer>	FStagingBufferRHIRef;
 
-
 class FRHIRenderTargetView
 {
 public:
diff --git a/Engine/Source/Runtime/VulkanRHI/Private/VulkanCommands.cpp b/Engine/Source/Runtime/VulkanRHI/Private/VulkanCommands.cpp
--- a/Engine/Source/Runtime/VulkanRHI/Private/VulkanCommands.cpp
+++ b/Engine/Source/Runtime/VulkanRHI/Private/VulkanCommands.cpp
@@ -1191,6 +1191,30 @@ void FVulkanCommandListContext::PrepareParallelFromBase(const FVulkanCommandList
 	TransitionAndLayoutManager.TempCopy(BaseContext.TransitionAndLayoutManager);
 }
 
+void FVulkanCommandListContext::RHIEnqueueStagedRead(FStagingBufferRHIParamRef StagingBufferRHI, FGPUFenceRHIParamRef FenceRHI, uint32 Offset, uint32 NumBytes)
+{
+	FVulkanCmdBuffer* CmdBuffer = CommandBufferManager->GetActiveCmdBuffer();
+
+	FVulkanGPUFence* Fence = ResourceCast(FenceRHI);
+	Fence->CmdBuffer = CmdBuffer;
+	Fence->FenceSignaledCounter = CmdBuffer->GetFenceSignaledCounter();
+
+	ensure(CmdBuffer->IsOutsideRenderPass());
+	VulkanRHI::FStagingBuffer* ReadbackStagingBuffer = Device->GetStagingManager().AcquireBuffer(NumBytes);
+	FVulkanStagingBuffer* StagingBuffer = ResourceCast(StagingBufferRHI);
+	StagingBuffer->StagingBuffer = ReadbackStagingBuffer;
+	StagingBuffer->QueuedOffset = Offset;
+	StagingBuffer->QueuedNumBytes = NumBytes;
+
+	VkBufferCopy Region;
+	FMemory::Memzero(Region);
+	Region.size = NumBytes;
+	FVulkanVertexBuffer* VertexBuffer = ResourceCast(StagingBufferRHI->GetBackingBuffer());
+	Region.srcOffset = Offset + VertexBuffer->GetOffset();
+	//Region.dstOffset = 0;
+	VulkanRHI::vkCmdCopyBuffer(CmdBuffer->GetHandle(), VertexBuffer->GetHandle(), ReadbackStagingBuffer->GetHandle(), 1, &Region);
+}
+
 
 FVulkanCommandContextContainer::FVulkanCommandContextContainer(FVulkanDevice* InDevice)
 	: VulkanRHI::FDeviceChild(InDevice)
diff --git a/Engine/Source/Runtime/VulkanRHI/Private/VulkanUAV.cpp b/Engine/Source/Runtime/VulkanRHI/Private/VulkanUAV.cpp
--- a/Engine/Source/Runtime/VulkanRHI/Private/VulkanUAV.cpp
+++ b/Engine/Source/Runtime/VulkanRHI/Private/VulkanUAV.cpp
@@ -349,16 +349,6 @@ bool FVulkanGPUFence::Poll() const
 	return FenceSignaledCounter < CmdBuffer->GetFenceSignaledCounter();
 }
 
-bool FVulkanGPUFence::Wait(float TimeoutMs) const
-{
-	if (!Poll())
-	{
-		CmdBuffer->GetOwner()->GetMgr().WaitForCmdBuffer(CmdBuffer, TimeoutMs / 1000.0f);
-	}
-
-	return Poll();
-}
-
 FGPUFenceRHIRef FVulkanDynamicRHI::RHICreateGPUFence(const FName& Name)
 {
 	return new FVulkanGPUFence(Name);
diff --git a/Engine/Source/Runtime/VulkanRHI/Private/VulkanUtil.cpp b/Engine/Source/Runtime/VulkanRHI/Private/VulkanUtil.cpp
--- a/Engine/Source/Runtime/VulkanRHI/Private/VulkanUtil.cpp
+++ b/Engine/Source/Runtime/VulkanRHI/Private/VulkanUtil.cpp
@@ -97,33 +97,6 @@ FStagingBufferRHIRef FVulkanDynamicRHI::RHICreateStagingBuffer(FVertexBufferRHIP
 	return new FVulkanStagingBuffer(BackingBuffer);
 }
 
-void FVulkanDynamicRHI::RHIEnqueueStagedRead(FStagingBufferRHIParamRef StagingBufferRHI, FGPUFenceRHIParamRef FenceRHI, uint32 Offset, uint32 NumBytes)
-{
-	FVulkanCommandListContext& Context = Device->GetImmediateContext();
-	FVulkanCmdBuffer* CmdBuffer = Context.GetCommandBufferManager()->GetActiveCmdBuffer();
-
-	FVulkanGPUFence* Fence = ResourceCast(FenceRHI);
-	Fence->CmdBuffer = CmdBuffer;
-	Fence->FenceSignaledCounter = CmdBuffer->GetFenceSignaledCounter();
-
-	ensure(CmdBuffer->IsOutsideRenderPass());
-	VulkanRHI::FStagingBuffer* ReadbackStagingBuffer = Device->GetStagingManager().AcquireBuffer(NumBytes);
-	FVulkanStagingBuffer* StagingBuffer = ResourceCast(StagingBufferRHI);
-	StagingBuffer->StagingBuffer = ReadbackStagingBuffer;
-	StagingBuffer->QueuedOffset = Offset;
-	StagingBuffer->QueuedNumBytes = NumBytes;
-
-	VkBufferCopy Region;
-	FMemory::Memzero(Region);
-	Region.size = NumBytes;
-	FVulkanVertexBuffer* VertexBuffer = ResourceCast(StagingBufferRHI->GetBackingBuffer());
-	Region.srcOffset = Offset + VertexBuffer->GetOffset();
-	//Region.dstOffset = 0;
-	VulkanRHI::vkCmdCopyBuffer(CmdBuffer->GetHandle(), VertexBuffer->GetHandle(), ReadbackStagingBuffer->GetHandle(), 1, &Region);
-
-	return FDynamicRHI::RHIEnqueueStagedRead(StagingBuffer, Fence, Offset, NumBytes);
-}
-
 void* FVulkanDynamicRHI::RHILockStagingBuffer(FStagingBufferRHIParamRef StagingBufferRHI, uint32 Offset, uint32 SizeRHI)
 {
 	FVulkanStagingBuffer* StagingBuffer = ResourceCast(StagingBufferRHI);
diff --git a/Engine/Source/Runtime/VulkanRHI/Public/VulkanContext.h b/Engine/Source/Runtime/VulkanRHI/Public/VulkanContext.h
--- a/Engine/Source/Runtime/VulkanRHI/Public/VulkanContext.h
+++ b/Engine/Source/Runtime/VulkanRHI/Public/VulkanContext.h
@@ -240,6 +240,7 @@ public:
 	virtual void RHICopyTexture(FTextureRHIParamRef SourceTexture, FTextureRHIParamRef DestTexture, const FRHICopyTextureInfo& CopyInfo) final override;
 	virtual void RHITransitionResources(EResourceTransitionAccess TransitionType, FTextureRHIParamRef* InRenderTargets, int32 NumTextures) final override;
 	virtual void RHITransitionResources(EResourceTransitionAccess TransitionType, EResourceTransitionPipeline TransitionPipeline, FUnorderedAccessViewRHIParamRef* InUAVs, int32 NumUAVs, FComputeFenceRHIParamRef WriteComputeFence) final override;
+	virtual void RHIEnqueueStagedRead(FStagingBufferRHIParamRef StagingBuffer, FGPUFenceRHIParamRef Fence, uint32 Offset, uint32 NumBytes) final override;
 
 	// Render time measurement
 	virtual void RHIBeginRenderQuery(FRenderQueryRHIParamRef RenderQuery) final override;
diff --git a/Engine/Source/Runtime/VulkanRHI/Public/VulkanDynamicRHI.h b/Engine/Source/Runtime/VulkanRHI/Public/VulkanDynamicRHI.h
--- a/Engine/Source/Runtime/VulkanRHI/Public/VulkanDynamicRHI.h
+++ b/Engine/Source/Runtime/VulkanRHI/Public/VulkanDynamicRHI.h
@@ -376,7 +376,6 @@ public:
 
 	virtual void VulkanSetImageLayout( VkCommandBuffer CmdBuffer, VkImage Image, VkImageLayout OldLayout, VkImageLayout NewLayout, const VkImageSubresourceRange& SubresourceRange );
 	
-	virtual void RHIEnqueueStagedRead(FStagingBufferRHIParamRef StagingBuffer, FGPUFenceRHIParamRef Fence, uint32 Offset, uint32 NumBytes) final override;
 	virtual void* RHILockStagingBuffer(FStagingBufferRHIParamRef StagingBuffer, uint32 Offset, uint32 SizeRHI) final override;
 	virtual void RHIUnlockStagingBuffer(FStagingBufferRHIParamRef StagingBuffer) final override;
 
diff --git a/Engine/Source/Runtime/VulkanRHI/Public/VulkanResources.h b/Engine/Source/Runtime/VulkanRHI/Public/VulkanResources.h
--- a/Engine/Source/Runtime/VulkanRHI/Public/VulkanResources.h
+++ b/Engine/Source/Runtime/VulkanRHI/Public/VulkanResources.h
@@ -1480,13 +1480,12 @@ public:
 	}
 
 	virtual bool Poll() const final override;
-	virtual bool Wait(float TimeoutMs) const final override;
 
 protected:
 	FVulkanCmdBuffer*	CmdBuffer = nullptr;
 	uint64				FenceSignaledCounter = 0;
 
-	friend class FVulkanDynamicRHI;
+	friend class FVulkanCommandListContext;
 };
 
 template<class T>
