fix for async compute fence bug

diff --git a/Engine/Source/Runtime/D3D12RHI/Private/D3D12Commands.cpp b/Engine/Source/Runtime/D3D12RHI/Private/D3D12Commands.cpp
--- a/Engine/Source/Runtime/D3D12RHI/Private/D3D12Commands.cpp
+++ b/Engine/Source/Runtime/D3D12RHI/Private/D3D12Commands.cpp
@@ -391,7 +391,7 @@ void FD3D12CommandContext::RHITransitionResources(EResourceTransitionAccess Tran
 		FD3D12Fence* Fence = FD3D12DynamicRHI::ResourceCast(WriteComputeFenceRHI);
 		Fence->WriteFence();
 
-		Fence->Signal(ED3D12CommandQueueType::Default);
+		Fence->Signal(bIsAsyncComputeContext ? ED3D12CommandQueueType::Async : ED3D12CommandQueueType::Default);
 	}
 }
 
