Fix for DX12 PSO cache crashes

diff --git a/Engine/Source/Runtime/D3D12RHI/Private/D3D12PipelineState.h b/Engine/Source/Runtime/D3D12RHI/Private/D3D12PipelineState.h
--- a/Engine/Source/Runtime/D3D12RHI/Private/D3D12PipelineState.h
+++ b/Engine/Source/Runtime/D3D12RHI/Private/D3D12PipelineState.h
@@ -287,22 +287,10 @@ public:
 		return *this;
 	}
 
-	// Indicates this PSO should be added to any disk caches.
-	void MarkForDiskCacheAdd()
-	{
-		bAddToDiskCache = PipelineState.GetReference() ? true : false;
-	}
-
-	bool ShouldAddToDiskCache() const
-	{
-		return bAddToDiskCache;
-	}
-
 protected:
 	TRefCountPtr<ID3D12PipelineState> PipelineState;
 	FAsyncTask<FD3D12PipelineStateWorker>* Worker;
 	volatile int32 PendingWaitOnWorkerCalls;
-	bool bAddToDiskCache;
 };
 
 struct FD3D12GraphicsPipelineState : public FRHIGraphicsPipelineState
diff --git a/Engine/Source/Runtime/D3D12RHI/Private/D3D12Pipelinestate.cpp b/Engine/Source/Runtime/D3D12RHI/Private/D3D12Pipelinestate.cpp
--- a/Engine/Source/Runtime/D3D12RHI/Private/D3D12Pipelinestate.cpp
+++ b/Engine/Source/Runtime/D3D12RHI/Private/D3D12Pipelinestate.cpp
@@ -232,7 +232,6 @@ FD3D12PipelineState::FD3D12PipelineState(FD3D12Adapter* Parent)
 	, FD3D12MultiNodeGPUObject(FRHIGPUMask::All(), FRHIGPUMask::All()) //Create on all, visible on all
 	, Worker(nullptr)
 	, PendingWaitOnWorkerCalls(0)
-	, bAddToDiskCache(false)
 {
 	INC_DWORD_STAT(STAT_D3D12NumPSOs);
 }
diff --git a/Engine/Source/Runtime/D3D12RHI/Private/Windows/WindowsD3D12DiskCache.cpp b/Engine/Source/Runtime/D3D12RHI/Private/Windows/WindowsD3D12DiskCache.cpp
--- a/Engine/Source/Runtime/D3D12RHI/Private/Windows/WindowsD3D12DiskCache.cpp
+++ b/Engine/Source/Runtime/D3D12RHI/Private/Windows/WindowsD3D12DiskCache.cpp
@@ -42,7 +42,6 @@ void FDiskCacheInterface::Init(FString &filename, bool bEnable)
 		}
 	}
 	bool fileFound = mCacheExists;
-	mCurrentFileMapSize = 1;
 	GrowMapping(64 * 1024, true);
 
 	if (fileFound && mFileStart)
@@ -69,12 +68,9 @@ void FDiskCacheInterface::GrowMapping(SIZE_T size, bool firstrun)
 		return;
 	}
 
-	if (mCurrentFileMapSize - mCurrentOffset < size)
+	if (mCurrentOffset + size > mCurrentFileMapSize)
 	{
-		while ((mCurrentFileMapSize - mCurrentOffset) < size)
-		{
-			mCurrentFileMapSize += mFileGrowSize;
-		}
+		mCurrentFileMapSize = Align(mCurrentOffset + size, mFileGrowSize);
 	}
 	else
 	{
diff --git a/Engine/Source/Runtime/D3D12RHI/Private/Windows/WindowsD3D12DiskCache.h b/Engine/Source/Runtime/D3D12RHI/Private/Windows/WindowsD3D12DiskCache.h
--- a/Engine/Source/Runtime/D3D12RHI/Private/Windows/WindowsD3D12DiskCache.h
+++ b/Engine/Source/Runtime/D3D12RHI/Private/Windows/WindowsD3D12DiskCache.h
@@ -10,7 +10,7 @@ class FDiskCacheInterface
 {
 	// Increment this if changes are made to the
 	// disk caches so stale caches get updated correctly
-	static const uint32 mCurrentHeaderVersion = 7;
+	static const uint32 mCurrentHeaderVersion = 8;
 	struct FDiskCacheHeader
 	{
 		uint32 mHeaderVersion;
diff --git a/Engine/Source/Runtime/D3D12RHI/Private/Windows/WindowsD3D12PipelineState.cpp b/Engine/Source/Runtime/D3D12RHI/Private/Windows/WindowsD3D12PipelineState.cpp
--- a/Engine/Source/Runtime/D3D12RHI/Private/Windows/WindowsD3D12PipelineState.cpp
+++ b/Engine/Source/Runtime/D3D12RHI/Private/Windows/WindowsD3D12PipelineState.cpp
@@ -109,8 +109,8 @@ void FD3D12PipelineStateCache::OnPSOCreated(FD3D12PipelineState* PipelineState,
 	// Actually create the PSO.
 	PipelineState->Create(GraphicsPipelineCreationArgs(&Desc, PipelineLibrary.GetReference()));
 
-	// Indicate we should add this PSO to necessary disk caches.
-	PipelineState->MarkForDiskCacheAdd();
+	// Save this PSO to disk cache.
+	AddToDiskCache(Desc, PipelineState);
 }
 
 void FD3D12PipelineStateCache::OnPSOCreated(FD3D12PipelineState* PipelineState, const FD3D12ComputePipelineStateDesc& Desc)
@@ -118,8 +118,8 @@ void FD3D12PipelineStateCache::OnPSOCreated(FD3D12PipelineState* PipelineState,
 	// Actually create the PSO.
 	PipelineState->Create(ComputePipelineCreationArgs(&Desc, PipelineLibrary.GetReference()));
 
-	// Indicate we should add this PSO to necessary disk caches.
-	PipelineState->MarkForDiskCacheAdd();
+	// Save this PSO to disk cache.
+	AddToDiskCache(Desc, PipelineState);
 }
 
 void FD3D12PipelineStateCache::RebuildFromDiskCache(ID3D12RootSignature* GraphicsRootSignature, ID3D12RootSignature* ComputeRootSignature)
@@ -231,9 +231,6 @@ void FD3D12PipelineStateCache::RebuildFromDiskCache(ID3D12RootSignature* Graphic
 					// Actually create the PSO.
 					const GraphicsPipelineCreationArgs Args(&Desc, PipelineLibrary.GetReference());
 					(*PipelineState)->CreateAsync(Args);
-
-					// Sanity check that we won't add this back to the disk cache again.
-					check(!(*PipelineState)->ShouldAddToDiskCache())
 				});
 			}
 		}
@@ -286,9 +283,6 @@ void FD3D12PipelineStateCache::RebuildFromDiskCache(ID3D12RootSignature* Graphic
 					// Actually create the PSO.
 					const ComputePipelineCreationArgs Args(&Desc, PipelineLibrary.GetReference());
 					PipelineState->CreateAsync(Args);
-
-					// Sanity check that we won't add this back to the disk cache again.
-					check(!PipelineState->ShouldAddToDiskCache())
 				});
 			}
 		}
@@ -455,55 +449,6 @@ void FD3D12PipelineStateCache::Close()
 {
 	FScopeLock Lock(&DiskCachesCS);
 
-	FDiskCacheInterface& DiskGraphicsCache = DiskCaches[PSO_CACHE_GRAPHICS];
-	FDiskCacheInterface& DiskComputeCache = DiskCaches[PSO_CACHE_COMPUTE];
-
-	DiskGraphicsCache.Reset(FDiskCacheInterface::RESET_TO_AFTER_LAST_OBJECT);
-	DiskComputeCache.Reset(FDiskCacheInterface::RESET_TO_AFTER_LAST_OBJECT);
-	DiskBinaryCache.Reset(FDiskCacheInterface::RESET_TO_AFTER_LAST_OBJECT);
-
-	// Write graphics PSOs to the disk cache(s).
-	if (!DiskGraphicsCache.IsInErrorState())
-	{
-		const uint32 NumGraphicsPSOsBefore = DiskGraphicsCache.GetNumPSOs();
-		for (auto Iter = LowLevelGraphicsPipelineStateCache.CreateConstIterator(); Iter; ++Iter)
-		{
-			FD3D12PipelineState* const PipelineState = Iter.Value();
-			if (PipelineState->ShouldAddToDiskCache())
-			{
-				const FD3D12LowLevelGraphicsPipelineStateDesc& Desc = Iter.Key();
-				AddToDiskCache(Desc, PipelineState);
-			}
-		}
-
-		const uint32 NumGraphicsPSOsAfter = DiskGraphicsCache.GetNumPSOs();
-		const uint32 NumNewGraphicsPSOs = NumGraphicsPSOsAfter - NumGraphicsPSOsBefore;
-		UE_CLOG(NumNewGraphicsPSOs, LogD3D12RHI, Log, TEXT("Added %u new Graphics PSO(s) to the disk cache."), NumNewGraphicsPSOs);
-		UE_LOG(LogD3D12RHI, Log, TEXT("Closing Graphics PSO disk cache. Cache contains %u PSO(s)."), NumGraphicsPSOsAfter);
-		DiskGraphicsCache.Close(NumGraphicsPSOsAfter);
-	}
-
-	// Write compute PSOs to the disk cache(s).
-	if (!DiskComputeCache.IsInErrorState())
-	{
-		const uint32 NumComputePSOsBefore = DiskComputeCache.GetNumPSOs();
-		for (auto Iter = ComputePipelineStateCache.CreateConstIterator(); Iter; ++Iter)
-		{
-			FD3D12PipelineState* const PipelineState = Iter.Value();
-			if (PipelineState->ShouldAddToDiskCache())
-			{
-				const FD3D12ComputePipelineStateDesc& Desc = Iter.Key();
-				AddToDiskCache(Desc, PipelineState);
-			}
-		}
-
-		const uint32 NumComputePSOsAfter = DiskComputeCache.GetNumPSOs();
-		const uint32 NumNewComputePSOs = NumComputePSOsAfter - NumComputePSOsBefore;
-		UE_CLOG(NumNewComputePSOs, LogD3D12RHI, Log, TEXT("Added %u new Compute PSO(s) to the disk cache."), NumNewComputePSOs);
-		UE_LOG(LogD3D12RHI, Log, TEXT("Closing Compute PSO disk cache. Cache contains %u PSO(s)."), NumComputePSOsAfter);
-		DiskComputeCache.Close(NumComputePSOsAfter);
-	}
-
 	// Write driver-optimized PSOs to the disk cache.
 	TArray<BYTE> LibraryData;
 	const bool bOverwriteExistingPipelineLibrary = true;
@@ -529,6 +474,8 @@ void FD3D12PipelineStateCache::Close()
 	DiskBinaryCache.Close(0);
 
 	CleanupPipelineStateCaches();
+
+	PipelineLibrary.SafeRelease();
 }
 
 void FD3D12PipelineStateCache::Init(FString& GraphicsCacheFileName, FString& ComputeCacheFileName, FString& DriverBlobFileName)
@@ -609,75 +556,6 @@ FD3D12PipelineStateCache::~FD3D12PipelineStateCache()
 {
 }
 
-template ID3D12PipelineState* CreatePipelineStateWrapper<GraphicsPipelineCreationArgs_POD>(FD3D12Adapter* Adapter, const GraphicsPipelineCreationArgs_POD* CreationArgs);
-template ID3D12PipelineState* CreatePipelineStateWrapper<ComputePipelineCreationArgs_POD>(FD3D12Adapter* Adapter, const ComputePipelineCreationArgs_POD* CreationArgs);
-
-template <typename TDesc>
-ID3D12PipelineState* CreatePipelineStateWrapper(FD3D12Adapter* Adapter, const TDesc* CreationArgs)
-{
-	// Get the pipeline state name, currently based on the hash.
-	FString Name = FString::Printf(TEXT("%llu"), CreationArgs->Desc->CombinedHash);
-
-	// Use pipeline streams if the system supports it.
-	ID3D12Device2* const pDevice2 = Adapter->GetD3DDevice2();
-	if (pDevice2)
-	{
-		typename TPSOStreamFunctionMap<TDesc>::D3D12PipelineStateStreamType Stream = CreationArgs->Desc->Desc.PipelineStateStream();
-		const D3D12_PIPELINE_STATE_STREAM_DESC StreamDesc = { sizeof(Stream), &Stream };
-		return CreatePipelineStateFromStream<TDesc>(pDevice2, &StreamDesc, static_cast<ID3D12PipelineLibrary1*>(CreationArgs->Library), *Name);	// Static cast to ID3D12PipelineLibrary1 since we already checked for ID3D12Device2.
-	}
-	else
-	{
-		// Use the older pipeline descs.
-		const typename TPSOStreamFunctionMap<TDesc>::D3D12PipelineStateDescV0Type Desc = (CreationArgs->Desc->Desc.*TPSOStreamFunctionMap<TDesc>::GetPipelineStateDescV0())();
-		return CreatePipelineState(Adapter->GetD3DDevice(), &Desc, CreationArgs->Library, *Name);
-	}
-}
-
-void FD3D12PipelineState::Create(const ComputePipelineCreationArgs& InCreationArgs)
-{
-	check(PipelineState.GetReference() == nullptr);
-	PipelineState = CreatePipelineStateWrapper(GetParentAdapter(), &InCreationArgs.Args);
-}
-
-void FD3D12PipelineState::CreateAsync(const ComputePipelineCreationArgs& InCreationArgs)
-{
-	check(PipelineState.GetReference() == nullptr && Worker == nullptr);
-	Worker = new FAsyncTask<FD3D12PipelineStateWorker>(GetParentAdapter(), InCreationArgs);
-	if (Worker)
-	{
-		Worker->StartBackgroundTask();
-	}
-}
-
-void FD3D12PipelineState::Create(const GraphicsPipelineCreationArgs& InCreationArgs)
-{
-	check(PipelineState.GetReference() == nullptr);
-	PipelineState = CreatePipelineStateWrapper(GetParentAdapter(), &InCreationArgs.Args);
-}
-
-void FD3D12PipelineState::CreateAsync(const GraphicsPipelineCreationArgs& InCreationArgs)
-{
-	check(PipelineState.GetReference() == nullptr && Worker == nullptr);
-	Worker = new FAsyncTask<FD3D12PipelineStateWorker>(GetParentAdapter(), InCreationArgs);
-	if (Worker)
-	{
-		Worker->StartBackgroundTask();
-	}
-}
-
-void FD3D12PipelineStateWorker::DoWork()
-{
-	if (bIsGraphics)
-	{
-		PSO = CreatePipelineStateWrapper(GetParentAdapter(), &CreationArgs.GraphicsArgs);
-	}
-	else
-	{
-		PSO = CreatePipelineStateWrapper(GetParentAdapter(), &CreationArgs.ComputeArgs);
-	}
-}
-
 #if LOG_PSO_CREATES
 	/** Accumulative time spent creating pipeline states. */
 	FTotalTimeAndCount GD3D12CreatePSOTime;
@@ -685,19 +563,10 @@ void FD3D12PipelineStateWorker::DoWork()
 
 DECLARE_CYCLE_STAT(TEXT("Create time"), STAT_PSOCreateTime, STATGROUP_D3D12PipelineState);
 
-template ID3D12PipelineState* CreatePipelineState<D3D12_GRAPHICS_PIPELINE_STATE_DESC>(ID3D12Device* Device, const D3D12_GRAPHICS_PIPELINE_STATE_DESC* Desc, ID3D12PipelineLibrary* Library, const TCHAR* Name);
-template ID3D12PipelineState* CreatePipelineState<D3D12_COMPUTE_PIPELINE_STATE_DESC>(ID3D12Device* Device, const D3D12_COMPUTE_PIPELINE_STATE_DESC* Desc, ID3D12PipelineLibrary* Library, const TCHAR* Name);
-
 // Thread-safe create graphics/compute pipeline state. Conditionally load/store the PSO using a Pipeline Library.
 template <typename TDesc>
-ID3D12PipelineState* CreatePipelineState(ID3D12Device* Device, const TDesc* Desc, ID3D12PipelineLibrary* Library, const TCHAR* Name)
+static void CreatePipelineState(ID3D12PipelineState*&PSO, ID3D12Device* Device, const TDesc* Desc, ID3D12PipelineLibrary* Library, const TCHAR* Name)
 {
-#if LOG_PSO_CREATES
-	const FString CreatePipelineStateMessage = FString::Printf(TEXT("CreatePipelineState (%s, Hash = %s)"), *TPSOFunctionMap<TDesc>::GetString(), Name);
-	SCOPE_LOG_TIME(*CreatePipelineStateMessage, &GD3D12CreatePSOTime);
-#endif
-
-	ID3D12PipelineState* PSO = nullptr;
 	if (Library)
 	{
 		// Try to load the PSO from the library.
@@ -732,28 +601,16 @@ ID3D12PipelineState* CreatePipelineState(ID3D12Device* Device, const TDesc* Desc
 		HRESULT r = (Device->*TPSOFunctionMap<TDesc>::GetCreatePipelineState())(Desc, IID_PPV_ARGS(&PSO));
 		if (FAILED(r))
 		{
-			UE_LOG(LogD3D12RHI, Error, TEXT("Failed to create PipelineState %s with hash %s"), *TPSOFunctionMap<TDesc>::GetString(), Name);
-			return nullptr;
+			UE_LOG(LogD3D12RHI, Error, TEXT("Failed to create PipelineState with hash %s"), Name);
 		}
 	}
 
 	check(PSO);
-	return PSO;
 }
 
-template ID3D12PipelineState* CreatePipelineStateFromStream<GraphicsPipelineCreationArgs_POD>(ID3D12Device2* Device, const D3D12_PIPELINE_STATE_STREAM_DESC* Desc, ID3D12PipelineLibrary1* Library, const TCHAR* Name);
-template ID3D12PipelineState* CreatePipelineStateFromStream<ComputePipelineCreationArgs_POD>(ID3D12Device2* Device, const D3D12_PIPELINE_STATE_STREAM_DESC* Desc, ID3D12PipelineLibrary1* Library, const TCHAR* Name);
-
 // Thread-safe create graphics/compute pipeline state. Conditionally load/store the PSO using a Pipeline Library.
-template <typename TDesc>
-ID3D12PipelineState* CreatePipelineStateFromStream(ID3D12Device2* Device, const D3D12_PIPELINE_STATE_STREAM_DESC* Desc, ID3D12PipelineLibrary1* Library, const TCHAR* Name)
+static void CreatePipelineStateFromStream(ID3D12PipelineState*& PSO, ID3D12Device2* Device, const D3D12_PIPELINE_STATE_STREAM_DESC* Desc, ID3D12PipelineLibrary1* Library, const TCHAR* Name)
 {
-#if LOG_PSO_CREATES
-	const FString CreatePipelineStateMessage = FString::Printf(TEXT("CreatePipelineState (%s, Hash = %s)"), *TPSOStreamFunctionMap<TDesc>::GetString(), Name);
-	SCOPE_LOG_TIME(*CreatePipelineStateMessage, &GD3D12CreatePSOTime);
-#endif
-
-	ID3D12PipelineState* PSO = nullptr;
 	if (Library)
 	{
 		// Try to load the PSO from the library.
@@ -789,11 +646,91 @@ ID3D12PipelineState* CreatePipelineStateFromStream(ID3D12Device2* Device, const
 
 		if (FAILED(r))
 		{
-			UE_LOG(LogD3D12RHI, Error, TEXT("Failed to create PipelineState"));
-			return nullptr;
+			UE_LOG(LogD3D12RHI, Error, TEXT("Failed to create PipelineState with hash %s"), Name);
 		}
 	}
 
 	check(PSO);
-	return PSO;
+}
+
+static inline void FastHashName(wchar_t Name[9], uint32 Hash)
+{
+	for (int i = 0; i < 8; i++)
+	{
+		Name[i] = (Hash & 0xF) + 'A';
+		Hash >>= 4;
+	}
+	Name[8] = 0;
+}
+
+template <typename TDesc>
+static void CreatePipelineStateWrapper(ID3D12PipelineState** PSO, FD3D12Adapter* Adapter, const TDesc* CreationArgs)
+{
+	// Get the pipeline state name, currently based on the hash.
+	wchar_t Name[9];
+	FastHashName(Name, CreationArgs->Desc->CombinedHash);
+
+#if LOG_PSO_CREATES
+	const FString CreatePipelineStateMessage = FString::Printf(TEXT("CreatePipelineState (%s, Hash = %s)"), TPSOStreamFunctionMap<TDesc>::GetString(), Name);
+	SCOPE_LOG_TIME(*CreatePipelineStateMessage, &GD3D12CreatePSOTime);
+#endif
+
+	// Use pipeline streams if the system supports it.
+	ID3D12Device2* const pDevice2 = Adapter->GetD3DDevice2();
+	if (pDevice2)
+	{
+		typename TPSOStreamFunctionMap<TDesc>::D3D12PipelineStateStreamType Stream = CreationArgs->Desc->Desc.PipelineStateStream();
+		const D3D12_PIPELINE_STATE_STREAM_DESC StreamDesc = { sizeof(Stream), &Stream };
+		CreatePipelineStateFromStream(*PSO, pDevice2, &StreamDesc, static_cast<ID3D12PipelineLibrary1*>(CreationArgs->Library), Name);	// Static cast to ID3D12PipelineLibrary1 since we already checked for ID3D12Device2.
+	}
+	else
+	{
+		// Use the older pipeline descs.
+		const typename TPSOStreamFunctionMap<TDesc>::D3D12PipelineStateDescV0Type Desc = (CreationArgs->Desc->Desc.*TPSOStreamFunctionMap<TDesc>::GetPipelineStateDescV0())();
+		CreatePipelineState(*PSO, Adapter->GetD3DDevice(), &Desc, CreationArgs->Library, Name);
+	}
+}
+
+void FD3D12PipelineState::Create(const ComputePipelineCreationArgs& InCreationArgs)
+{
+	check(PipelineState.GetReference() == nullptr);
+	CreatePipelineStateWrapper(PipelineState.GetInitReference(), GetParentAdapter(), &InCreationArgs.Args);
+}
+
+void FD3D12PipelineState::CreateAsync(const ComputePipelineCreationArgs& InCreationArgs)
+{
+	check(PipelineState.GetReference() == nullptr && Worker == nullptr);
+	Worker = new FAsyncTask<FD3D12PipelineStateWorker>(GetParentAdapter(), InCreationArgs);
+	if (Worker)
+	{
+		Worker->StartBackgroundTask();
+	}
+}
+
+void FD3D12PipelineState::Create(const GraphicsPipelineCreationArgs& InCreationArgs)
+{
+	check(PipelineState.GetReference() == nullptr);
+	CreatePipelineStateWrapper(PipelineState.GetInitReference(), GetParentAdapter(), &InCreationArgs.Args);
+}
+
+void FD3D12PipelineState::CreateAsync(const GraphicsPipelineCreationArgs& InCreationArgs)
+{
+	check(PipelineState.GetReference() == nullptr && Worker == nullptr);
+	Worker = new FAsyncTask<FD3D12PipelineStateWorker>(GetParentAdapter(), InCreationArgs);
+	if (Worker)
+	{
+		Worker->StartBackgroundTask();
+	}
+}
+
+void FD3D12PipelineStateWorker::DoWork()
+{
+	if (bIsGraphics)
+	{
+		CreatePipelineStateWrapper(PSO.GetInitReference(), GetParentAdapter(), &CreationArgs.GraphicsArgs);
+	}
+	else
+	{
+		CreatePipelineStateWrapper(PSO.GetInitReference(), GetParentAdapter(), &CreationArgs.ComputeArgs);
+	}
 }
diff --git a/Engine/Source/Runtime/D3D12RHI/Private/Windows/WindowsD3D12PipelineState.h b/Engine/Source/Runtime/D3D12RHI/Private/Windows/WindowsD3D12PipelineState.h
--- a/Engine/Source/Runtime/D3D12RHI/Private/Windows/WindowsD3D12PipelineState.h
+++ b/Engine/Source/Runtime/D3D12RHI/Private/Windows/WindowsD3D12PipelineState.h
@@ -148,11 +148,6 @@ private:
 	TRefCountPtr<ID3D12PipelineLibrary> PipelineLibrary;
 	bool bUseAPILibaries;
 
-
-	// Add_NoLock
-	void AddToDiskCache(const FD3D12LowLevelGraphicsPipelineStateDesc& Desc, FD3D12PipelineState* PipelineState);
-	void AddToDiskCache(const FD3D12ComputePipelineStateDesc& Desc, FD3D12PipelineState* PipelineState);
-
 	void WriteOutShaderBlob(PSO_CACHE_TYPE Cache, ID3D12PipelineState* APIPso);
 
 	template<typename PipelineStateDescType>
@@ -197,6 +192,9 @@ protected:
 	void OnPSOCreated(FD3D12PipelineState* PipelineState, const FD3D12LowLevelGraphicsPipelineStateDesc& Desc) final override;
 	void OnPSOCreated(FD3D12PipelineState* PipelineState, const FD3D12ComputePipelineStateDesc& Desc) final override;
 
+	void AddToDiskCache(const FD3D12LowLevelGraphicsPipelineStateDesc& Desc, FD3D12PipelineState* PipelineState);
+	void AddToDiskCache(const FD3D12ComputePipelineStateDesc& Desc, FD3D12PipelineState* PipelineState);
+
 public:
 
 	void RebuildFromDiskCache(ID3D12RootSignature* GraphicsRootSignature, ID3D12RootSignature* ComputeRootSignature);
@@ -215,13 +213,11 @@ template<> struct TPSOFunctionMap < D3D12_GRAPHICS_PIPELINE_STATE_DESC >
 {
 	static decltype(&ID3D12Device::CreateGraphicsPipelineState) GetCreatePipelineState() { return &ID3D12Device::CreateGraphicsPipelineState; }
 	static decltype(&ID3D12PipelineLibrary::LoadGraphicsPipeline) GetLoadPipeline() { return &ID3D12PipelineLibrary::LoadGraphicsPipeline; }
-	static FString GetString() { return TEXT("Graphics"); }
 };
 template<> struct TPSOFunctionMap < D3D12_COMPUTE_PIPELINE_STATE_DESC >
 {
 	static decltype(&ID3D12Device::CreateComputePipelineState) GetCreatePipelineState() { return &ID3D12Device::CreateComputePipelineState; }
 	static decltype(&ID3D12PipelineLibrary::LoadComputePipeline) GetLoadPipeline() { return &ID3D12PipelineLibrary::LoadComputePipeline; }
-	static FString GetString() { return TEXT("Compute"); }
 };
 
 template <typename TDesc> struct TPSOStreamFunctionMap;
@@ -230,23 +226,12 @@ template<> struct TPSOStreamFunctionMap < GraphicsPipelineCreationArgs_POD >
 	static decltype(&FD3D12_GRAPHICS_PIPELINE_STATE_DESC::GraphicsDescV0) GetPipelineStateDescV0() { return &FD3D12_GRAPHICS_PIPELINE_STATE_DESC::GraphicsDescV0; }
 	typedef D3D12_GRAPHICS_PIPELINE_STATE_DESC D3D12PipelineStateDescV0Type;
 	typedef FD3D12_GRAPHICS_PIPELINE_STATE_STREAM D3D12PipelineStateStreamType;
-	static FString GetString() { return TEXT("Graphics Stream"); }
+	static TCHAR* GetString() { return TEXT("Graphics"); }
 };
 template<> struct TPSOStreamFunctionMap < ComputePipelineCreationArgs_POD >
 {
 	static decltype(&FD3D12_COMPUTE_PIPELINE_STATE_DESC::ComputeDescV0) GetPipelineStateDescV0() { return &FD3D12_COMPUTE_PIPELINE_STATE_DESC::ComputeDescV0; }
 	typedef D3D12_COMPUTE_PIPELINE_STATE_DESC D3D12PipelineStateDescV0Type;
 	typedef FD3D12_COMPUTE_PIPELINE_STATE_STREAM D3D12PipelineStateStreamType;
-	static FString GetString() { return TEXT("Compute Stream"); }
+	static TCHAR* GetString() { return TEXT("Compute"); }
 };
-
-// Thread-safe create graphics/compute pipeline state. Conditionally load/store the PSO using a Pipeline Library.
-template <typename TDesc>
-ID3D12PipelineState* CreatePipelineState(ID3D12Device* Device, const TDesc* Desc, ID3D12PipelineLibrary* Library, const TCHAR* Name);
-
-// Thread-safe create graphics/compute pipeline state. Conditionally load/store the PSO using a Pipeline Library.
-template <typename TDesc>
-ID3D12PipelineState* CreatePipelineStateFromStream(ID3D12Device2* Device, const D3D12_PIPELINE_STATE_STREAM_DESC* Desc, ID3D12PipelineLibrary1* Library, const TCHAR* Name);
-
-template <typename TDesc>
-ID3D12PipelineState* CreatePipelineStateWrapper(FD3D12Adapter* Adapter, const TDesc* CreationArgs);
