Fix clang warning from cl 4484754

diff --git a/Engine/Source/Runtime/D3D12RHI/Private/Windows/WindowsD3D12PipelineState.h b/Engine/Source/Runtime/D3D12RHI/Private/Windows/WindowsD3D12PipelineState.h
--- a/Engine/Source/Runtime/D3D12RHI/Private/Windows/WindowsD3D12PipelineState.h
+++ b/Engine/Source/Runtime/D3D12RHI/Private/Windows/WindowsD3D12PipelineState.h
@@ -226,12 +226,12 @@ template<> struct TPSOStreamFunctionMap < GraphicsPipelineCreationArgs_POD >
 	static decltype(&FD3D12_GRAPHICS_PIPELINE_STATE_DESC::GraphicsDescV0) GetPipelineStateDescV0() { return &FD3D12_GRAPHICS_PIPELINE_STATE_DESC::GraphicsDescV0; }
 	typedef D3D12_GRAPHICS_PIPELINE_STATE_DESC D3D12PipelineStateDescV0Type;
 	typedef FD3D12_GRAPHICS_PIPELINE_STATE_STREAM D3D12PipelineStateStreamType;
-	static TCHAR* GetString() { return TEXT("Graphics"); }
+	static const TCHAR* GetString() { return TEXT("Graphics"); }
 };
 template<> struct TPSOStreamFunctionMap < ComputePipelineCreationArgs_POD >
 {
 	static decltype(&FD3D12_COMPUTE_PIPELINE_STATE_DESC::ComputeDescV0) GetPipelineStateDescV0() { return &FD3D12_COMPUTE_PIPELINE_STATE_DESC::ComputeDescV0; }
 	typedef D3D12_COMPUTE_PIPELINE_STATE_DESC D3D12PipelineStateDescV0Type;
 	typedef FD3D12_COMPUTE_PIPELINE_STATE_STREAM D3D12PipelineStateStreamType;
-	static TCHAR* GetString() { return TEXT("Compute"); }
+	static const TCHAR* GetString() { return TEXT("Compute"); }
 };
