Workaround for low repro-rate PSO cache crash

diff --git a/Engine/Source/Runtime/D3D12RHI/Private/Windows/WindowsD3D12PipelineState.cpp b/Engine/Source/Runtime/D3D12RHI/Private/Windows/WindowsD3D12PipelineState.cpp
--- a/Engine/Source/Runtime/D3D12RHI/Private/Windows/WindowsD3D12PipelineState.cpp
+++ b/Engine/Source/Runtime/D3D12RHI/Private/Windows/WindowsD3D12PipelineState.cpp
@@ -9,6 +9,16 @@
 #include "Misc/ScopeRWLock.h"
 #include "Stats/StatsMisc.h"
 
+// UE-65533
+// Using asynchronous PSO creation to preload the PSO cache significantly speeds up startup.
+// An crash bug of low repro rate currently prevents us from using this feature, so as a workaround PSOs are created synchronously.
+// The effect of this bug is that a previously verified valid PSO has been overwritten/deleted or otherwise corrupted by the time it is
+// first accessed. The root cause of this problem has not be established. The symptom is a crash in GetPipelineState() where the corrupt
+// PSO is copied over from the async worker but the memory has been wiped and the v-table pointer is garbage, causing AddRef() to crash.
+#ifndef D3D12RHI_USE_ASYNC_PRELOAD
+#define D3D12RHI_USE_ASYNC_PRELOAD 0
+#endif
+
 static TAutoConsoleVariable<int32> CVarPipelineStateDiskCache(
 	TEXT("D3D12.PSO.DiskCache"),
 	1,
@@ -230,7 +240,11 @@ void FD3D12PipelineStateCache::RebuildFromDiskCache(ID3D12RootSignature* Graphic
 				{
 					// Actually create the PSO.
 					const GraphicsPipelineCreationArgs Args(&Desc, PipelineLibrary.GetReference());
+				#if D3D12RHI_USE_ASYNC_PRELOAD
 					(*PipelineState)->CreateAsync(Args);
+				#else
+					(*PipelineState)->Create(Args);
+				#endif
 				});
 			}
 		}
@@ -282,7 +296,11 @@ void FD3D12PipelineStateCache::RebuildFromDiskCache(ID3D12RootSignature* Graphic
 				{
 					// Actually create the PSO.
 					const ComputePipelineCreationArgs Args(&Desc, PipelineLibrary.GetReference());
+				#if D3D12RHI_USE_ASYNC_PRELOAD
 					PipelineState->CreateAsync(Args);
+				#else
+					PipelineState->Create(Args);
+				#endif
 				});
 			}
 		}
