Copy 4364019 - Fix MSVC 2017 compile errors with bStrictConformanceMode (/permissive-). Original change made by [REDACTED_PERSON].

diff --git a/Engine/Source/Runtime/D3D12RHI/Private/D3D12StateCachePrivate.h b/Engine/Source/Runtime/D3D12RHI/Private/D3D12StateCachePrivate.h
--- a/Engine/Source/Runtime/D3D12RHI/Private/D3D12StateCachePrivate.h
+++ b/Engine/Source/Runtime/D3D12RHI/Private/D3D12StateCachePrivate.h
@@ -4,7 +4,7 @@
 //	thread performance by removing redundant device context calls.
 
 #pragma once
-#include "Containers/Queue.h"
+
 #include "D3D12DirectCommandListManager.h"
 
 //-----------------------------------------------------------------------------
@@ -387,7 +387,7 @@ protected:
 			PipelineState.Common.CurrentShaderUAVCounts[Traits::Frequency]     = (Shader) ? Shader->ResourceCounts.NumUAVs     : 0;
 		
 			// Shader changed so its resource table is dirty
-			CmdContext->DirtyUniformBuffers[Traits::Frequency] = 0xffff;
+			this->CmdContext->DirtyUniformBuffers[Traits::Frequency] = 0xffff;
 		}
 	}
 
@@ -415,7 +415,7 @@ protected:
 		if (bNeedSetPSO)
 		{
 			check(CurrentPSO);
-			CmdContext->CommandListHandle->SetPipelineState(CurrentPSO);
+			this->CmdContext->CommandListHandle->SetPipelineState(CurrentPSO);
 			PipelineState.Common.bNeedSetPSO = false;
 		}
 	}
@@ -553,9 +553,10 @@ public:
 	D3D12_STATE_CACHE_INLINE void GetSamplerState(uint32 StartSamplerIndex, uint32 NumSamplerIndexes, FD3D12SamplerState** SamplerStates) const
 	{
 		check(StartSamplerIndex + NumSamplerIndexes <= D3D12_COMMONSHADER_INPUT_RESOURCE_SLOT_COUNT);
+		auto& CurrentShaderResourceViews = PipelineState.Common.SRVCache.Views[ShaderFrequency];
 		for (uint32 StateLoop = 0; StateLoop < NumSamplerIndexes; StateLoop++)
 		{
-			SamplerStates[StateLoop] = CurrentShaderResourceViews[ShaderFrequency][StateLoop + StartSamplerIndex];
+			SamplerStates[StateLoop] = CurrentShaderResourceViews[StateLoop + StartSamplerIndex];
 			if (SamplerStates[StateLoop])
 			{
 				SamplerStates[StateLoop]->AddRef();
diff --git a/Engine/Source/Runtime/D3D12RHI/Private/D3D12View.h b/Engine/Source/Runtime/D3D12RHI/Private/D3D12View.h
--- a/Engine/Source/Runtime/D3D12RHI/Private/D3D12View.h
+++ b/Engine/Source/Runtime/D3D12RHI/Private/D3D12View.h
@@ -613,6 +613,8 @@ inline uint32 CViewSubresourceSubset::MaxSubresource() const
 	return (*(--end())).Value;
 }
 
+class FD3D12OfflineDescriptorManager;
+
 template <typename TDesc>
 class TD3D12ViewDescriptorHandle : public FD3D12DeviceChild
 {
