Fix DX12 texture upload crash by deferring upload to RHI thread. This also avoids stalling RHI thread.

Crash report with repro:
https://udn.unrealengine.com/questions/452179/rhicreatetexture-crash-on-dx12.html

diff --git a/Engine/Source/Runtime/D3D12RHI/Private/D3D12RHIPrivate.h b/Engine/Source/Runtime/D3D12RHI/Private/D3D12RHIPrivate.h
--- a/Engine/Source/Runtime/D3D12RHI/Private/D3D12RHIPrivate.h
+++ b/Engine/Source/Runtime/D3D12RHI/Private/D3D12RHIPrivate.h
@@ -907,10 +907,10 @@ protected:
 	uint32 ZeroBufferSize;
 
 	template<typename BaseResourceType>
-	TD3D12Texture2D<BaseResourceType>* CreateD3D12Texture2D(uint32 SizeX, uint32 SizeY, uint32 SizeZ, bool bTextureArray, bool CubeTexture, uint8 Format,
+	TD3D12Texture2D<BaseResourceType>* CreateD3D12Texture2D(class FRHICommandListImmediate* RHICmdList, uint32 SizeX, uint32 SizeY, uint32 SizeZ, bool bTextureArray, bool CubeTexture, EPixelFormat Format,
 		uint32 NumMips, uint32 NumSamples, uint32 Flags, FRHIResourceCreateInfo& CreateInfo);
 
-	FD3D12Texture3D* CreateD3D12Texture3D(uint32 SizeX, uint32 SizeY, uint32 SizeZ, uint8 Format, uint32 NumMips, uint32 Flags, FRHIResourceCreateInfo& CreateInfo);
+	FD3D12Texture3D* CreateD3D12Texture3D(class FRHICommandListImmediate* RHICmdList, uint32 SizeX, uint32 SizeY, uint32 SizeZ, EPixelFormat Format, uint32 NumMips, uint32 Flags, FRHIResourceCreateInfo& CreateInfo);
 
 	/**
 	 * Gets the best supported MSAA settings from the provided MSAA count to check against.
diff --git a/Engine/Source/Runtime/D3D12RHI/Private/D3D12Texture.cpp b/Engine/Source/Runtime/D3D12RHI/Private/D3D12Texture.cpp
--- a/Engine/Source/Runtime/D3D12RHI/Private/D3D12Texture.cpp
+++ b/Engine/Source/Runtime/D3D12RHI/Private/D3D12Texture.cpp
@@ -97,6 +97,99 @@ struct FRHICommandUpdateTexture final : public FRHICommand<FRHICommandUpdateText
 	}
 };
 
+struct FD3D12RHICommandInitializeTexture final : public FRHICommand<FD3D12RHICommandInitializeTexture>
+{
+	FD3D12TextureBase* TextureBase;
+	FD3D12ResourceLocation SrcResourceLoc;
+	uint32 SizeX;
+	uint32 SizeY;
+	uint32 SizeZ;
+	uint32 NumSlices;
+	uint32 NumMips;
+	EPixelFormat Format;
+	D3D12_RESOURCE_STATES DestinationState;
+
+	FORCEINLINE_DEBUGGABLE FD3D12RHICommandInitializeTexture(FD3D12TextureBase* InTexture, FD3D12ResourceLocation& InSrcResourceLoc, uint32 InSizeX, uint32 InSizeY, uint32 InSizeZ, uint32 InNumSlices, uint32 InNumMips, EPixelFormat InFormat, D3D12_RESOURCE_STATES InDestinationState)
+		: TextureBase(InTexture)
+		, SrcResourceLoc(InSrcResourceLoc.GetParentDevice())
+		, SizeX(InSizeX)
+		, SizeY(InSizeY)
+		, SizeZ(InSizeZ)
+		, NumSlices(InNumSlices)
+		, NumMips(InNumMips)
+		, Format(InFormat)
+		, DestinationState(InDestinationState)
+	{
+		FD3D12ResourceLocation::TransferOwnership(SrcResourceLoc, InSrcResourceLoc);
+	}
+
+	void Execute(FRHICommandListBase& /* unused */)
+	{
+		ExecuteNoCmdList();
+	}
+
+	void ExecuteNoCmdList()
+	{
+		DXGI_FORMAT PlatformFormat = (DXGI_FORMAT) GPixelFormats[Format].PlatformFormat;
+
+		// Begin with the root texture (i.e. device node 0)
+		FD3D12TextureBase* CurrentTexture = TextureBase;
+
+		// Initialize all the textures in the chain
+		while (CurrentTexture != nullptr)
+		{
+			FD3D12Device* Device = CurrentTexture->GetParentDevice();
+			FD3D12Resource* Resource = CurrentTexture->GetResource();
+
+			FD3D12CommandListHandle& hCommandList = Device->GetDefaultCommandContext().CommandListHandle;
+			hCommandList.GetCurrentOwningContext()->numCopies += NumSlices * NumMips;
+
+			ID3D12Device* D3DDevice = Device->GetDevice();
+			ID3D12GraphicsCommandList* CmdList = hCommandList.GraphicsCommandList();
+
+			D3D12_TEXTURE_COPY_LOCATION Dst;
+			Dst.pResource = Resource->GetResource();
+			Dst.Type = D3D12_TEXTURE_COPY_TYPE_SUBRESOURCE_INDEX;
+
+			D3D12_TEXTURE_COPY_LOCATION Src;
+			Src.pResource = SrcResourceLoc.GetResource()->GetResource();
+			Src.Type = D3D12_TEXTURE_COPY_TYPE_PLACED_FOOTPRINT;
+
+			uint64 Offset = SrcResourceLoc.GetOffsetFromBaseOfResource();
+			for (uint32 ArraySliceIndex = 0; ArraySliceIndex < NumSlices; ++ArraySliceIndex)
+			{
+				for (uint32 MipIndex = 0; MipIndex < NumMips; ++MipIndex)
+				{
+					uint32 NumBlocksX = FMath::Max(1U, (SizeX >> MipIndex) / GPixelFormats[Format].BlockSizeX);
+					uint32 NumBlocksY = FMath::Max(1U, (SizeY >> MipIndex) / GPixelFormats[Format].BlockSizeY);
+					uint32 NumBlocksZ = FMath::Max(1U, (SizeZ >> MipIndex) / GPixelFormats[Format].BlockSizeZ);
+
+					uint32 RowPitch = Align(NumBlocksX * GPixelFormats[Format].BlockBytes, FD3D12_TEXTURE_DATA_PITCH_ALIGNMENT);
+
+					Dst.SubresourceIndex                   = ArraySliceIndex * NumMips + MipIndex;
+					Src.PlacedFootprint.Offset             = Offset;
+					Src.PlacedFootprint.Footprint.Format   = PlatformFormat;
+					Src.PlacedFootprint.Footprint.Width    = FMath::Max(1U, SizeX >> MipIndex);
+					Src.PlacedFootprint.Footprint.Height   = FMath::Max(1U, SizeY >> MipIndex);
+					Src.PlacedFootprint.Footprint.Depth    = FMath::Max(1U, SizeZ >> MipIndex);
+					Src.PlacedFootprint.Footprint.RowPitch = RowPitch;
+
+					CmdList->CopyTextureRegion(&Dst, 0, 0, 0, &Src, nullptr);
+
+					Offset += RowPitch * NumBlocksY * NumBlocksZ;
+				}
+			}
+
+			hCommandList.UpdateResidency(Resource);
+
+			hCommandList.AddTransitionBarrier(Resource, D3D12_RESOURCE_STATE_COPY_DEST, DestinationState, D3D12_RESOURCE_BARRIER_ALL_SUBRESOURCES);
+
+			CurrentTexture = CurrentTexture->GetNextObject();
+		}
+
+	}
+};
+
 ///////////////////////////////////////////////////////////////////////////////////////////
 // Texture Stats
 ///////////////////////////////////////////////////////////////////////////////////////////
@@ -560,9 +653,8 @@ void SafeCreateTexture2D(FD3D12Device* pDevice,
 #endif // #if GUARDED_TEXTURE_CREATES
 }
 
-
 template<typename BaseResourceType>
-TD3D12Texture2D<BaseResourceType>* FD3D12DynamicRHI::CreateD3D12Texture2D(uint32 SizeX, uint32 SizeY, uint32 SizeZ, bool bTextureArray, bool bCubeTexture, uint8 Format,
+TD3D12Texture2D<BaseResourceType>* FD3D12DynamicRHI::CreateD3D12Texture2D(FRHICommandListImmediate* RHICmdList, uint32 SizeX, uint32 SizeY, uint32 SizeZ, bool bTextureArray, bool bCubeTexture, EPixelFormat Format,
 	uint32 NumMips, uint32 NumSamples, uint32 Flags, FRHIResourceCreateInfo& CreateInfo)
 {
 #if PLATFORM_WINDOWS
@@ -735,7 +827,7 @@ TD3D12Texture2D<BaseResourceType>* FD3D12DynamicRHI::CreateD3D12Texture2D(uint32
 			&Location,
 			Format,
 			Flags,
-			(CreateInfo.BulkData != NULL) ? D3D12_RESOURCE_STATE_COPY_DEST : DestinationState);
+			(CreateInfo.BulkData != nullptr) ? D3D12_RESOURCE_STATE_COPY_DEST : DestinationState);
 
 		uint32 RTVIndex = 0;
 
@@ -889,72 +981,10 @@ TD3D12Texture2D<BaseResourceType>* FD3D12DynamicRHI::CreateD3D12Texture2D(uint32
 	FD3D12TextureStats::D3D12TextureAllocated(*D3D12TextureOut);
 
 	// Initialize if data is given
-	if (CreateInfo.BulkData != NULL)
+	if (CreateInfo.BulkData != nullptr)
 	{
-		TArray<D3D12_SUBRESOURCE_DATA> SubResourceData;
-
-		uint8* Data = (uint8*)CreateInfo.BulkData->GetResourceBulkData();
-
-		// each mip of each array slice counts as a subresource
-		SubResourceData.AddZeroed(NumMips * SizeZ);
-
-		uint32 SliceOffset = 0;
-		for (uint32 ArraySliceIndex = 0; ArraySliceIndex < SizeZ; ++ArraySliceIndex)
-		{
-			uint32 MipOffset = 0;
-			for (uint32 MipIndex = 0; MipIndex < NumMips; ++MipIndex)
-			{
-				uint32 DataOffset = SliceOffset + MipOffset;
-				uint32 SubResourceIndex = ArraySliceIndex * NumMips + MipIndex;
-
-				uint32 NumBlocksX = FMath::Max<uint32>(1, (SizeX >> MipIndex) / GPixelFormats[Format].BlockSizeX);
-				uint32 NumBlocksY = FMath::Max<uint32>(1, (SizeY >> MipIndex) / GPixelFormats[Format].BlockSizeY);
-
-				SubResourceData[SubResourceIndex].pData = &Data[DataOffset];
-				SubResourceData[SubResourceIndex].RowPitch = NumBlocksX * GPixelFormats[Format].BlockBytes;
-				SubResourceData[SubResourceIndex].SlicePitch = NumBlocksX * NumBlocksY * SubResourceData[MipIndex].RowPitch;
-
-				MipOffset += NumBlocksY * SubResourceData[MipIndex].RowPitch;
-			}
-			SliceOffset += MipOffset;
-		}
-
-		uint64 Size = GetRequiredIntermediateSize(D3D12TextureOut->GetResource()->GetResource(), 0, NumMips * SizeZ);
-
-		auto& FastAllocator = Adapter->GetDevice(0)->GetDefaultFastAllocator();
-
-		FD3D12ResourceLocation TempResourceLocation(FastAllocator.GetParentDevice());
-		void* pData = FastAllocator.Allocate<FD3D12ScopeLock>(Size, D3D12_TEXTURE_DATA_PLACEMENT_ALIGNMENT, &TempResourceLocation);
-
-		// Begin with the root texture (i.e. device node 0)
-		TD3D12Texture2D<BaseResourceType>* CurrentTexture = D3D12TextureOut;
-
-		// Initialize all the textures in the chain
-		while (CurrentTexture != nullptr)
-		{
-			FD3D12Device* Device = CurrentTexture->GetParentDevice();
-			FD3D12CommandListHandle& hCommandList = Device->GetDefaultCommandContext().CommandListHandle;
-
-			FD3D12Resource* Resource = CurrentTexture->GetResource();
-
-			hCommandList.GetCurrentOwningContext()->numCopies++;
-			UpdateSubresources(hCommandList.GraphicsCommandList(),
-				Resource->GetResource(),
-				TempResourceLocation.GetResource()->GetResource(),
-				TempResourceLocation.GetOffsetFromBaseOfResource(),
-				0, NumMips * SizeZ,
-				SubResourceData.GetData());
+		D3D12TextureOut->InitializeTextureData(RHICmdList, CreateInfo.BulkData->GetResourceBulkData(), CreateInfo.BulkData->GetResourceBulkDataSize(), SizeX, SizeY, 1, SizeZ, NumMips, Format, DestinationState);
 
-			hCommandList.UpdateResidency(Resource);
-
-			hCommandList.AddTransitionBarrier(Resource, D3D12_RESOURCE_STATE_COPY_DEST, DestinationState, D3D12_RESOURCE_BARRIER_ALL_SUBRESOURCES);
-
-			CurrentTexture = (TD3D12Texture2D<BaseResourceType>*) CurrentTexture->GetNextObject();
-		}
-	}
-
-	if (CreateInfo.BulkData)
-	{
 		CreateInfo.BulkData->Discard();
 	}
 
@@ -965,7 +995,7 @@ TD3D12Texture2D<BaseResourceType>* FD3D12DynamicRHI::CreateD3D12Texture2D(uint32
 #endif // PLATFORM_WINDOWS
 }
 
-FD3D12Texture3D* FD3D12DynamicRHI::CreateD3D12Texture3D(uint32 SizeX, uint32 SizeY, uint32 SizeZ, uint8 Format, uint32 NumMips, uint32 Flags, FRHIResourceCreateInfo& CreateInfo)
+FD3D12Texture3D* FD3D12DynamicRHI::CreateD3D12Texture3D(FRHICommandListImmediate* RHICmdList, uint32 SizeX, uint32 SizeY, uint32 SizeZ, EPixelFormat Format, uint32 NumMips, uint32 Flags, FRHIResourceCreateInfo& CreateInfo)
 {
 #if PLATFORM_WINDOWS
 	SCOPE_CYCLE_COUNTER(STAT_D3D12CreateTextureTime);
@@ -1017,9 +1047,9 @@ FD3D12Texture3D* FD3D12DynamicRHI::CreateD3D12Texture3D(uint32 SizeX, uint32 Siz
 	FD3D12Adapter* Adapter = &GetAdapter();
 	FD3D12Texture3D* D3D12TextureOut = Adapter->CreateLinkedObject<FD3D12Texture3D>(FRHIGPUMask::All(), [&](FD3D12Device* Device)
 	{
-		FD3D12Texture3D* Texture3D = new FD3D12Texture3D(Device, SizeX, SizeY, SizeZ, NumMips, (EPixelFormat)Format, Flags, CreateInfo.ClearValueBinding);
+		FD3D12Texture3D* Texture3D = new FD3D12Texture3D(Device, SizeX, SizeY, SizeZ, NumMips, Format, Flags, CreateInfo.ClearValueBinding);
 
-		HRESULT hr = Device->GetTextureAllocator().AllocateTexture(TextureDesc, ClearValuePtr, Format, Texture3D->ResourceLocation, (CreateInfo.BulkData != NULL) ? D3D12_RESOURCE_STATE_COPY_DEST : DestinationState);
+		HRESULT hr = Device->GetTextureAllocator().AllocateTexture(TextureDesc, ClearValuePtr, Format, Texture3D->ResourceLocation, (CreateInfo.BulkData != nullptr) ? D3D12_RESOURCE_STATE_COPY_DEST : DestinationState);
 		check(SUCCEEDED(hr));
 
 		if (bCreateRTV)
@@ -1052,55 +1082,15 @@ FD3D12Texture3D* FD3D12DynamicRHI::CreateD3D12Texture3D(uint32 SizeX, uint32 Siz
 	// Intialize if data given
 	if (D3D12TextureOut)
 	{
-		if (CreateInfo.BulkData)
+		if (CreateInfo.BulkData != nullptr)
 		{
-			TArray<D3D12_SUBRESOURCE_DATA> SubResourceData;
-
-			uint8* Data = (uint8*)CreateInfo.BulkData->GetResourceBulkData();
-			SubResourceData.AddZeroed(NumMips);
-			uint32 MipOffset = 0;
-			for (uint32 MipIndex = 0; MipIndex < NumMips; ++MipIndex)
-			{
-				SubResourceData[MipIndex].pData = &Data[MipOffset];
-				SubResourceData[MipIndex].RowPitch = FMath::Max<uint32>(1, SizeX >> MipIndex) * GPixelFormats[Format].BlockBytes;
-				SubResourceData[MipIndex].SlicePitch = FMath::Max<uint32>(1, SizeY >> MipIndex) * SubResourceData[MipIndex].RowPitch;
-				MipOffset += FMath::Max<uint32>(1, SizeZ >> MipIndex) * SubResourceData[MipIndex].SlicePitch;
-			}
-
-			const uint64 Size = GetRequiredIntermediateSize(D3D12TextureOut->GetResource()->GetResource(), 0, NumMips);
-
-			FD3D12ResourceLocation TempResourceLocation(GetRHIDevice());
-			void* pData = GetRHIDevice()->GetDefaultFastAllocator().Allocate<FD3D12ScopeLock>(Size, D3D12_TEXTURE_DATA_PLACEMENT_ALIGNMENT, &TempResourceLocation);
-
-			FD3D12Texture3D* CurrentTexture = D3D12TextureOut;
-
-			// Initialize each resource in the chain
-			while (CurrentTexture != nullptr)
-			{
-				FD3D12Device* Device = CurrentTexture->GetParentDevice();
-				FD3D12Resource* Resource = CurrentTexture->GetResource();
-
-				FD3D12CommandListHandle& hCommandList = Device->GetDefaultCommandContext().CommandListHandle;
-				hCommandList.GetCurrentOwningContext()->numCopies++;
-				UpdateSubresources(
-					hCommandList.GraphicsCommandList(),
-					Resource->GetResource(),
-					TempResourceLocation.GetResource()->GetResource(),
-					TempResourceLocation.GetOffsetFromBaseOfResource(),
-					0, NumMips,
-					SubResourceData.GetData());
-				hCommandList.UpdateResidency(Resource);
-
-				hCommandList.AddTransitionBarrier(Resource, D3D12_RESOURCE_STATE_COPY_DEST, DestinationState, D3D12_RESOURCE_BARRIER_ALL_SUBRESOURCES);
-
-				CurrentTexture = (FD3D12Texture3D*)CurrentTexture->GetNextObject();
-			}
+			D3D12TextureOut->InitializeTextureData(RHICmdList, CreateInfo.BulkData->GetResourceBulkData(), CreateInfo.BulkData->GetResourceBulkDataSize(), SizeX, SizeY, SizeZ, 1, NumMips, Format, DestinationState);
 		}
 
 		FD3D12TextureStats::D3D12TextureAllocated(*D3D12TextureOut);
 	}
 
-	if (CreateInfo.BulkData)
+	if (CreateInfo.BulkData != nullptr)
 	{
 		CreateInfo.BulkData->Discard();
 	}
@@ -1118,19 +1108,12 @@ FD3D12Texture3D* FD3D12DynamicRHI::CreateD3D12Texture3D(uint32 SizeX, uint32 Siz
 
 FTexture2DRHIRef FD3D12DynamicRHI::RHICreateTexture2D_RenderThread(class FRHICommandListImmediate& RHICmdList, uint32 SizeX, uint32 SizeY, uint8 Format, uint32 NumMips, uint32 NumSamples, uint32 Flags, FRHIResourceCreateInfo& CreateInfo)
 {
-	if (CreateInfo.BulkData != nullptr)
-	{
-		FScopedRHIThreadStaller StallRHIThread(RHICmdList);
-
-		return GDynamicRHI->RHICreateTexture2D(SizeX, SizeY, Format, NumMips, NumSamples, Flags, CreateInfo);
-	}
-
-	return GDynamicRHI->RHICreateTexture2D(SizeX, SizeY, Format, NumMips, NumSamples, Flags, CreateInfo);
+	return CreateD3D12Texture2D<FD3D12BaseTexture2D>(&RHICmdList, SizeX, SizeY, 1, false, false, (EPixelFormat) Format, NumMips, NumSamples, Flags, CreateInfo);
 }
 
 FTexture2DRHIRef FD3D12DynamicRHI::RHICreateTexture2D(uint32 SizeX, uint32 SizeY, uint8 Format, uint32 NumMips, uint32 NumSamples, uint32 Flags, FRHIResourceCreateInfo& CreateInfo)
 {
-	return CreateD3D12Texture2D<FD3D12BaseTexture2D>(SizeX, SizeY, 1, false, false, Format, NumMips, NumSamples, Flags, CreateInfo);
+	return CreateD3D12Texture2D<FD3D12BaseTexture2D>(nullptr, SizeX, SizeY, 1, false, false, (EPixelFormat) Format, NumMips, NumSamples, Flags, CreateInfo);
 }
 
 FTexture2DRHIRef FD3D12DynamicRHI::RHIAsyncCreateTexture2D(uint32 SizeX, uint32 SizeY, uint8 Format, uint32 NumMips, uint32 Flags, void** InitialMipData, uint32 NumInitialMips)
@@ -1343,39 +1326,26 @@ void FD3D12DynamicRHI::RHICopySharedMips(FTexture2DRHIParamRef DestTexture2DRHI,
 
 FTexture2DArrayRHIRef FD3D12DynamicRHI::RHICreateTexture2DArray_RenderThread(class FRHICommandListImmediate& RHICmdList, uint32 SizeX, uint32 SizeY, uint32 SizeZ, uint8 Format, uint32 NumMips, uint32 Flags, FRHIResourceCreateInfo& CreateInfo)
 {
-	if (CreateInfo.BulkData != nullptr)
-	{
-		FScopedRHIThreadStaller StallRHIThread(RHICmdList);
-
-		return RHICreateTexture2DArray(SizeX, SizeY, SizeZ, Format, NumMips, Flags, CreateInfo);
-	}
-
-	return RHICreateTexture2DArray(SizeX, SizeY, SizeZ, Format, NumMips, Flags, CreateInfo);
+	check(SizeZ >= 1);
+	return CreateD3D12Texture2D<FD3D12BaseTexture2DArray>(&RHICmdList, SizeX, SizeY, SizeZ, true, false, (EPixelFormat) Format, NumMips, 1, Flags, CreateInfo);
 }
 
 FTexture2DArrayRHIRef FD3D12DynamicRHI::RHICreateTexture2DArray(uint32 SizeX, uint32 SizeY, uint32 SizeZ, uint8 Format, uint32 NumMips, uint32 Flags, FRHIResourceCreateInfo& CreateInfo)
 {
 	check(SizeZ >= 1);
-	return CreateD3D12Texture2D<FD3D12BaseTexture2DArray>(SizeX, SizeY, SizeZ, true, false, Format, NumMips, 1, Flags, CreateInfo);
+	return CreateD3D12Texture2D<FD3D12BaseTexture2DArray>(nullptr, SizeX, SizeY, SizeZ, true, false, (EPixelFormat) Format, NumMips, 1, Flags, CreateInfo);
 }
 
 FTexture3DRHIRef FD3D12DynamicRHI::RHICreateTexture3D_RenderThread(class FRHICommandListImmediate& RHICmdList, uint32 SizeX, uint32 SizeY, uint32 SizeZ, uint8 Format, uint32 NumMips, uint32 Flags, FRHIResourceCreateInfo& CreateInfo)
 {
-	if (CreateInfo.BulkData != nullptr)
-	{
-		FScopedRHIThreadStaller StallRHIThread(RHICmdList);
-
-		return RHICreateTexture3D(SizeX, SizeY, SizeZ, Format, NumMips, Flags, CreateInfo);
-	}
-
-	return RHICreateTexture3D(SizeX, SizeY, SizeZ, Format, NumMips, Flags, CreateInfo);
+	return CreateD3D12Texture3D(&RHICmdList, SizeX, SizeY, SizeZ, (EPixelFormat) Format, NumMips, Flags, CreateInfo);
 }
 
 FTexture3DRHIRef FD3D12DynamicRHI::RHICreateTexture3D(uint32 SizeX, uint32 SizeY, uint32 SizeZ, uint8 Format, uint32 NumMips, uint32 Flags, FRHIResourceCreateInfo& CreateInfo)
 {
 	check(SizeZ >= 1);
 #if PLATFORM_WINDOWS
-	return CreateD3D12Texture3D(SizeX, SizeY, SizeZ, Format, NumMips, Flags, CreateInfo);
+	return CreateD3D12Texture3D(nullptr, SizeX, SizeY, SizeZ, (EPixelFormat) Format, NumMips, Flags, CreateInfo);
 #else
 	checkf(false, TEXT("XBOX_CODE_MERGE : Removed. The Xbox platform version should be used."));
 	return nullptr;
@@ -1439,7 +1409,7 @@ FTexture2DRHIRef FD3D12DynamicRHI::RHIAsyncReallocateTexture2D(FTexture2DRHIPara
 
 	// Allocate a new texture.
 	FRHIResourceCreateInfo CreateInfo;
-	FD3D12Texture2D* NewTexture2D = CreateD3D12Texture2D<FD3D12BaseTexture2D>(NewSizeX, NewSizeY, 1, false, false, Texture2D->GetFormat(), NewMipCount, 1, Texture2D->GetFlags(), CreateInfo);
+	FD3D12Texture2D* NewTexture2D = CreateD3D12Texture2D<FD3D12BaseTexture2D>(nullptr, NewSizeX, NewSizeY, 1, false, false, Texture2D->GetFormat(), NewMipCount, 1, Texture2D->GetFlags(), CreateInfo);
 	FD3D12Texture2D* OriginalTexture = NewTexture2D;
 
 	// Use the GPU to asynchronously copy the old mip-maps into the new texture.
@@ -1683,6 +1653,64 @@ void FD3D12TextureBase::UpdateTexture(const D3D12_TEXTURE_COPY_LOCATION& DestCop
 	DEBUG_EXECUTE_COMMAND_CONTEXT(DefaultContext);
 }
 
+void FD3D12TextureBase::InitializeTextureData(FRHICommandListImmediate* RHICmdList, const void* InitData, uint32 InitDataSize, uint32 SizeX, uint32 SizeY, uint32 SizeZ, uint32 NumSlices, uint32 NumMips, EPixelFormat Format, D3D12_RESOURCE_STATES DestinationState)
+{
+	// each mip of each array slice counts as a subresource
+	uint32 NumSubresources = NumMips * NumSlices;
+	uint64 Size = GetRequiredIntermediateSize(GetResource()->GetResource(), 0, NumSubresources);
+
+	FD3D12ResourceLocation SrcResourceLoc(GetParentDevice());
+	uint8* DstData = (uint8*) SrcResourceLoc.GetParentDevice()->GetDefaultFastAllocator().Allocate<FD3D12ScopeLock>(Size, D3D12_TEXTURE_DATA_PLACEMENT_ALIGNMENT, &SrcResourceLoc);
+	uint32 DstDataOffset = 0;
+
+	const uint8* SrcData = (const uint8*) InitData;
+	for (uint32 ArraySliceIndex = 0; ArraySliceIndex < NumSlices; ++ArraySliceIndex)
+	{
+		for (uint32 MipIndex = 0; MipIndex < NumMips; ++MipIndex)
+		{
+			uint32 NumBlocksX = FMath::Max(1U, (SizeX >> MipIndex) / GPixelFormats[Format].BlockSizeX);
+			uint32 NumBlocksY = FMath::Max(1U, (SizeY >> MipIndex) / GPixelFormats[Format].BlockSizeY);
+			uint32 NumBlocksZ = FMath::Max(1U, (SizeZ >> MipIndex) / GPixelFormats[Format].BlockSizeZ);
+
+			uint32 RowPitch = NumBlocksX * GPixelFormats[Format].BlockBytes;
+			uint32 NumRows = NumBlocksY * NumBlocksZ;
+
+			// If RowPitch is aligned, which is typically the case for the bulk of the data (most large mips, POT textures), we turn the below loop into a single large memcpy
+			if (IsAligned(RowPitch, FD3D12_TEXTURE_DATA_PITCH_ALIGNMENT))
+			{
+				RowPitch *= NumRows;
+				NumRows = 1;
+			}
+
+			for (uint32 Row = 0; Row < NumRows; ++Row)
+			{
+				DstDataOffset = Align<uint32>(DstDataOffset, FD3D12_TEXTURE_DATA_PITCH_ALIGNMENT);
+				memcpy(DstData + DstDataOffset, SrcData, RowPitch);
+
+				SrcData += RowPitch;
+				DstDataOffset += RowPitch;
+			}
+		}
+	}
+
+#if PLATFORM_WINDOWS
+	ensure(DstDataOffset == Size);
+#else
+	ensure(Align(DstDataOffset, FD3D12_TEXTURE_DATA_PITCH_ALIGNMENT) == Size);
+#endif
+	ensure(SrcData == (uint8*) InitData + InitDataSize);
+
+	if (ShouldDeferCmdListOperation(RHICmdList))
+	{
+		new (RHICmdList->AllocCommand<FD3D12RHICommandInitializeTexture>()) FD3D12RHICommandInitializeTexture(this, SrcResourceLoc, SizeX, SizeY, SizeZ, NumSlices, NumMips, Format, DestinationState);
+	}
+	else
+	{
+		FD3D12RHICommandInitializeTexture Command(this, SrcResourceLoc, SizeX, SizeY, SizeZ, NumSlices, NumMips, Format, DestinationState);
+		Command.ExecuteNoCmdList();
+	}
+}
+
 template<typename RHIResourceType>
 void TD3D12Texture2D<RHIResourceType>::Unlock(class FRHICommandListImmediate* RHICmdList, uint32 MipIndex, uint32 ArrayIndex)
 {
@@ -2067,36 +2095,22 @@ void FD3D12DynamicRHI::EndUpdateTexture3D_Internal(FUpdateTexture3DData& UpdateD
 	-----------------------------------------------------------------------------*/
 FTextureCubeRHIRef FD3D12DynamicRHI::RHICreateTextureCube_RenderThread(class FRHICommandListImmediate& RHICmdList, uint32 Size, uint8 Format, uint32 NumMips, uint32 Flags, FRHIResourceCreateInfo& CreateInfo)
 {
-	if (CreateInfo.BulkData != nullptr)
-	{
-		FScopedRHIThreadStaller StallRHIThread(RHICmdList);
-
-		return RHICreateTextureCube(Size, Format, NumMips, Flags, CreateInfo);
-	}
-
-	return RHICreateTextureCube(Size, Format, NumMips, Flags, CreateInfo);
+	return CreateD3D12Texture2D<FD3D12BaseTextureCube>(&RHICmdList, Size, Size, 6, false, true, (EPixelFormat) Format, NumMips, 1, Flags, CreateInfo);
 }
 
 FTextureCubeRHIRef FD3D12DynamicRHI::RHICreateTextureCube(uint32 Size, uint8 Format, uint32 NumMips, uint32 Flags, FRHIResourceCreateInfo& CreateInfo)
 {
-	return CreateD3D12Texture2D<FD3D12BaseTextureCube>(Size, Size, 6, false, true, Format, NumMips, 1, Flags, CreateInfo);
+	return CreateD3D12Texture2D<FD3D12BaseTextureCube>(nullptr, Size, Size, 6, false, true, (EPixelFormat) Format, NumMips, 1, Flags, CreateInfo);
 }
 
 FTextureCubeRHIRef FD3D12DynamicRHI::RHICreateTextureCubeArray_RenderThread(class FRHICommandListImmediate& RHICmdList, uint32 Size, uint32 ArraySize, uint8 Format, uint32 NumMips, uint32 Flags, FRHIResourceCreateInfo& CreateInfo)
 {
-	if (CreateInfo.BulkData != nullptr)
-	{
-		FScopedRHIThreadStaller StallRHIThread(RHICmdList);
-
-		return RHICreateTextureCubeArray(Size, ArraySize, Format, NumMips, Flags, CreateInfo);
-	}
-
-	return RHICreateTextureCubeArray(Size, ArraySize, Format, NumMips, Flags, CreateInfo);
+	return CreateD3D12Texture2D<FD3D12BaseTextureCube>(&RHICmdList, Size, Size, 6 * ArraySize, true, true, (EPixelFormat) Format, NumMips, 1, Flags, CreateInfo);
 }
 
 FTextureCubeRHIRef FD3D12DynamicRHI::RHICreateTextureCubeArray(uint32 Size, uint32 ArraySize, uint8 Format, uint32 NumMips, uint32 Flags, FRHIResourceCreateInfo& CreateInfo)
 {
-	return CreateD3D12Texture2D<FD3D12BaseTextureCube>(Size, Size, 6 * ArraySize, true, true, Format, NumMips, 1, Flags, CreateInfo);
+	return CreateD3D12Texture2D<FD3D12BaseTextureCube>(nullptr, Size, Size, 6 * ArraySize, true, true, (EPixelFormat) Format, NumMips, 1, Flags, CreateInfo);
 }
 
 void* FD3D12DynamicRHI::RHILockTextureCubeFace(FTextureCubeRHIParamRef TextureCubeRHI, uint32 FaceIndex, uint32 ArrayIndex, uint32 MipIndex, EResourceLockMode LockMode, uint32& DestStride, bool bLockWithinMiptail)
diff --git a/Engine/Source/Runtime/D3D12RHI/Private/D3D12Texture.h b/Engine/Source/Runtime/D3D12RHI/Private/D3D12Texture.h
--- a/Engine/Source/Runtime/D3D12RHI/Private/D3D12Texture.h
+++ b/Engine/Source/Runtime/D3D12RHI/Private/D3D12Texture.h
@@ -91,7 +91,23 @@ public:
 	FD3D12BaseShaderResource* GetBaseShaderResource() const { return BaseShaderResource; }
 	void SetShaderResourceView(FD3D12ShaderResourceView* InShaderResourceView) { ShaderResourceView = InShaderResourceView; }
 
+	static inline bool ShouldDeferCmdListOperation(FRHICommandList* RHICmdList)
+	{
+		if (RHICmdList == nullptr)
+		{
+			return false;
+		}
+
+		if (RHICmdList->Bypass() || !IsRunningRHIInSeparateThread())
+		{
+			return false;
+		}
+
+		return true;
+	}
+
 	void UpdateTexture(const D3D12_TEXTURE_COPY_LOCATION& DestCopyLocation, uint32 DestX, uint32 DestY, uint32 DestZ, const D3D12_TEXTURE_COPY_LOCATION& SourceCopyLocation);
+	void InitializeTextureData(class FRHICommandListImmediate* RHICmdList, const void* InitData, uint32 InitDataSize, uint32 SizeX, uint32 SizeY, uint32 SizeZ, uint32 NumSlices, uint32 NumMips, EPixelFormat Format, D3D12_RESOURCE_STATES DestinationState);
 
 	/**
 	* Get the render target view for the specified mip and array slice.
@@ -253,21 +269,6 @@ public:
 	//* Update the contents of the Texture2D using a Copy command */
 	void UpdateTexture2D(class FRHICommandListImmediate* RHICmdList, uint32 MipIndex, const struct FUpdateTextureRegion2D& UpdateRegion, uint32 SourcePitch, const uint8* SourceData);
 
-	inline bool ShouldDeferCmdListOperation(FRHICommandList* RHICmdList)
-	{
-		if (RHICmdList == nullptr)
-		{
-			return false;
-		}
-
-		if (RHICmdList->Bypass() || !IsRunningRHIInSeparateThread())
-		{
-			return false;
-		}
-
-		return true;
-	}
-
 	// Accessors.
 	FD3D12Resource* GetResource() const { return (FD3D12Resource*)FD3D12TextureBase::GetResource(); }
 	const D3D12_PLACED_SUBRESOURCE_FOOTPRINT& GetReadBackHeapDesc() const
diff --git a/Engine/Source/Runtime/D3D12RHI/Public/D3D12Resources.h b/Engine/Source/Runtime/D3D12RHI/Public/D3D12Resources.h
--- a/Engine/Source/Runtime/D3D12RHI/Public/D3D12Resources.h
+++ b/Engine/Source/Runtime/D3D12RHI/Public/D3D12Resources.h
@@ -377,7 +377,7 @@ struct FD3D12BlockAllocatorPrivateData
 
 class FD3D12ResourceAllocator;
 // A very light-weight and cache friendly way of accessing a GPU resource
-class FD3D12ResourceLocation : public FD3D12DeviceChild
+class FD3D12ResourceLocation : public FD3D12DeviceChild, public FNoncopyable
 {
 public:
 
