DR - Remove unused functions/variables

diff --git a/Engine/Source/Runtime/D3D12RHI/Private/D3D12Query.cpp b/Engine/Source/Runtime/D3D12RHI/Private/D3D12Query.cpp
--- a/Engine/Source/Runtime/D3D12RHI/Private/D3D12Query.cpp
+++ b/Engine/Source/Runtime/D3D12RHI/Private/D3D12Query.cpp
@@ -180,8 +180,6 @@ FD3D12QueryHeap::FD3D12QueryHeap(FD3D12Device* InParent, const D3D12_QUERY_HEAP_
 	, FD3D12SingleNodeGPUObject(InParent->GetGPUMask())
 	, MaxActiveBatches(InMaxActiveBatches)
 	, LastBatch(InMaxActiveBatches - 1)
-	, HeadActiveElement(0)
-	, TailActiveElement(0)
 	, ActiveAllocatedElementCount(0)
 	, LastAllocatedElement(InQueryHeapCount - 1)
 	, ResultSize(8)
@@ -265,20 +263,6 @@ uint32 FD3D12QueryHeap::GetNextElement(uint32 InElement)
 	return InElement;
 }
 
-uint32 FD3D12QueryHeap::GetPreviousElement(uint32 InElement)
-{
-	// Decrement the provided element
-	InElement--;
-
-	// See if we need to wrap around to the end of the heap
-	if (InElement == UINT_MAX)
-	{
-		InElement = GetQueryHeapCount() - 1;
-	}
-
-	return InElement;
-}
-
 uint32 FD3D12QueryHeap::GetNextBatchElement(uint32 InBatchElement)
 {
 	// Increment the provided element
@@ -293,26 +277,6 @@ uint32 FD3D12QueryHeap::GetNextBatchElement(uint32 InBatchElement)
 	return InBatchElement;
 }
 
-uint32 FD3D12QueryHeap::GetPreviousBatchElement(uint32 InBatchElement)
-{
-	// Decrement the provided element
-	InBatchElement--;
-
-	// See if we need to wrap around to the end of the heap
-	if (InBatchElement == UINT_MAX)
-	{
-		InBatchElement = MaxActiveBatches - 1;
-	}
-
-	return InBatchElement;
-}
-
-bool FD3D12QueryHeap::IsHeapFull()
-{
-	// Find the next element after the active tail and compare with the head
-	return (GetNextElement(TailActiveElement) == HeadActiveElement);
-}
-
 uint32 FD3D12QueryHeap::AllocQuery(FD3D12CommandContext& CmdContext)
 {
 	check(CmdContext.IsDefaultContext());
@@ -372,13 +336,6 @@ void FD3D12QueryHeap::EndQueryBatchAndResolveQueryData(FD3D12CommandContext& Cmd
 	// Close the current batch
 	CurrentQueryBatch.bOpen = false;
 
-	// Update the end element
-	CurrentQueryBatch.EndElement = CurrentQueryBatch.StartElement + CurrentQueryBatch.ElementCount - 1;
-
-	// Update the tail
-	TailActiveElement = CurrentQueryBatch.EndElement;
-	check(TailActiveElement < GetQueryHeapCount());
-
 	// Increment the active element count
 	ActiveAllocatedElementCount += CurrentQueryBatch.ElementCount;
 	checkf(ActiveAllocatedElementCount <= GetQueryHeapCount(), TEXT("The query heap is too small. Either increase the heap count (larger resource) or decrease MAX_ACTIVE_BATCHES."));
@@ -389,7 +346,6 @@ void FD3D12QueryHeap::EndQueryBatchAndResolveQueryData(FD3D12CommandContext& Cmd
 
 	// Update the head
 	QueryBatch& OldestBatch = ActiveQueryBatches[GetNextBatchElement(LastBatch)];
-	HeadActiveElement = OldestBatch.StartElement;
 	ActiveAllocatedElementCount -= OldestBatch.ElementCount;
 
 	CmdContext.otherWorkCounter++;
diff --git a/Engine/Source/Runtime/D3D12RHI/Private/D3D12Query.h b/Engine/Source/Runtime/D3D12RHI/Private/D3D12Query.h
--- a/Engine/Source/Runtime/D3D12RHI/Private/D3D12Query.h
+++ b/Engine/Source/Runtime/D3D12RHI/Private/D3D12Query.h
@@ -78,7 +78,6 @@ private:
 	{
 	public:
 		uint32 StartElement;    // The first element in the batch (inclusive)
-		uint32 EndElement;      // The last element in the batch (inclusive)
 		uint32 ElementCount;    // The number of elements in the batch
 		bool bOpen;             // Is the batch still open for more begin/end queries?
 		
@@ -95,7 +94,6 @@ private:
 		inline void Clear()
 		{
 			StartElement = 0;
-			EndElement = 0;
 			ElementCount = 0;
 			bOpen = false;
 			RenderQueries.Reset();
@@ -127,12 +125,8 @@ private:
 	void StartQueryBatch(); // Start tracking a new batch of begin/end query calls that will be resolved together
 
 	uint32 GetNextElement(uint32 InElement); // Get the next element, after the specified element. Handles overflow.
-	uint32 GetPreviousElement(uint32 InElement); // Get the previous element, before the specified element. Handles underflow.
-	bool IsHeapFull();
-	bool IsHeapEmpty() const { return ActiveAllocatedElementCount == 0; }
 
 	uint32 GetNextBatchElement(uint32 InBatchElement);
-	uint32 GetPreviousBatchElement(uint32 InBatchElement);
 
 	void CreateQueryHeap();
 	void CreateResultBuffer();
@@ -147,8 +141,6 @@ private:
 	const uint32 MaxActiveBatches;                      // The max number of query batches that will be held.
 	uint32 LastBatch;                                   // The index of the newest batch.
 
-	uint32 HeadActiveElement;                   // The oldest element that is in use (Active). The data for this element is being used.
-	uint32 TailActiveElement;                   // The most recent element that is in use (Active). The data for this element is being used.
 	uint32 ActiveAllocatedElementCount;         // The number of elements that are in use (Active). Between the head and the tail.
 
 	uint32 LastAllocatedElement;                // The last element that was allocated for BeginQuery
