Fix debug runtime warning

diff --git a/Engine/Source/Runtime/D3D12RHI/Private/D3D12Resources.cpp b/Engine/Source/Runtime/D3D12RHI/Private/D3D12Resources.cpp
--- a/Engine/Source/Runtime/D3D12RHI/Private/D3D12Resources.cpp
+++ b/Engine/Source/Runtime/D3D12RHI/Private/D3D12Resources.cpp
@@ -173,7 +173,11 @@ FD3D12Resource::FD3D12Resource(FD3D12Device* ParentDevice,
 	FPlatformAtomics::InterlockedIncrement(&TotalResourceCount);
 #endif
 
-	if (Resource)
+	if (Resource
+#if PLATFORM_WINDOWS
+		&& Desc.Dimension == D3D12_RESOURCE_DIMENSION_BUFFER
+#endif
+		)
 	{
 		GPUVirtualAddress = Resource->GetGPUVirtualAddress();
 	}
