Better PSO initialization, removing AddRef() and Release().

diff --git a/Engine/Source/Runtime/D3D12RHI/Private/Windows/WindowsD3D12PipelineState.cpp b/Engine/Source/Runtime/D3D12RHI/Private/Windows/WindowsD3D12PipelineState.cpp
--- a/Engine/Source/Runtime/D3D12RHI/Private/Windows/WindowsD3D12PipelineState.cpp
+++ b/Engine/Source/Runtime/D3D12RHI/Private/Windows/WindowsD3D12PipelineState.cpp
@@ -609,11 +609,11 @@ FD3D12PipelineStateCache::~FD3D12PipelineStateCache()
 {
 }
 
-template ID3D12PipelineState* CreatePipelineStateWrapper<GraphicsPipelineCreationArgs_POD>(FD3D12Adapter* Adapter, const GraphicsPipelineCreationArgs_POD* CreationArgs);
-template ID3D12PipelineState* CreatePipelineStateWrapper<ComputePipelineCreationArgs_POD>(FD3D12Adapter* Adapter, const ComputePipelineCreationArgs_POD* CreationArgs);
+template void CreatePipelineStateWrapper<GraphicsPipelineCreationArgs_POD>(ID3D12PipelineState** PSO, FD3D12Adapter* Adapter, const GraphicsPipelineCreationArgs_POD* CreationArgs);
+template void CreatePipelineStateWrapper<ComputePipelineCreationArgs_POD>(ID3D12PipelineState** PSO, FD3D12Adapter* Adapter, const ComputePipelineCreationArgs_POD* CreationArgs);
 
 template <typename TDesc>
-ID3D12PipelineState* CreatePipelineStateWrapper(FD3D12Adapter* Adapter, const TDesc* CreationArgs)
+void CreatePipelineStateWrapper(ID3D12PipelineState** PSO, FD3D12Adapter* Adapter, const TDesc* CreationArgs)
 {
 	// Get the pipeline state name, currently based on the hash.
 	FString Name = FString::Printf(TEXT("%llu"), CreationArgs->Desc->CombinedHash);
@@ -624,21 +624,20 @@ ID3D12PipelineState* CreatePipelineStateWrapper(FD3D12Adapter* Adapter, const TD
 	{
 		typename TPSOStreamFunctionMap<TDesc>::D3D12PipelineStateStreamType Stream = CreationArgs->Desc->Desc.PipelineStateStream();
 		const D3D12_PIPELINE_STATE_STREAM_DESC StreamDesc = { sizeof(Stream), &Stream };
-		return CreatePipelineStateFromStream<TDesc>(pDevice2, &StreamDesc, static_cast<ID3D12PipelineLibrary1*>(CreationArgs->Library), *Name);	// Static cast to ID3D12PipelineLibrary1 since we already checked for ID3D12Device2.
+		CreatePipelineStateFromStream<TDesc>(*PSO, pDevice2, &StreamDesc, static_cast<ID3D12PipelineLibrary1*>(CreationArgs->Library), *Name);	// Static cast to ID3D12PipelineLibrary1 since we already checked for ID3D12Device2.
 	}
 	else
 	{
 		// Use the older pipeline descs.
 		const typename TPSOStreamFunctionMap<TDesc>::D3D12PipelineStateDescV0Type Desc = (CreationArgs->Desc->Desc.*TPSOStreamFunctionMap<TDesc>::GetPipelineStateDescV0())();
-		return CreatePipelineState(Adapter->GetD3DDevice(), &Desc, CreationArgs->Library, *Name);
+		CreatePipelineState(*PSO, Adapter->GetD3DDevice(), &Desc, CreationArgs->Library, *Name);
 	}
 }
 
 void FD3D12PipelineState::Create(const ComputePipelineCreationArgs& InCreationArgs)
 {
 	check(PipelineState.GetReference() == nullptr);
-	PipelineState = CreatePipelineStateWrapper(GetParentAdapter(), &InCreationArgs.Args);
-	PipelineState->Release();
+	CreatePipelineStateWrapper(PipelineState.GetInitReference(), GetParentAdapter(), &InCreationArgs.Args);
 }
 
 void FD3D12PipelineState::CreateAsync(const ComputePipelineCreationArgs& InCreationArgs)
@@ -654,8 +653,7 @@ void FD3D12PipelineState::CreateAsync(const ComputePipelineCreationArgs& InCreat
 void FD3D12PipelineState::Create(const GraphicsPipelineCreationArgs& InCreationArgs)
 {
 	check(PipelineState.GetReference() == nullptr);
-	PipelineState = CreatePipelineStateWrapper(GetParentAdapter(), &InCreationArgs.Args);
-	PipelineState->Release();
+	CreatePipelineStateWrapper(PipelineState.GetInitReference(), GetParentAdapter(), &InCreationArgs.Args);
 }
 
 void FD3D12PipelineState::CreateAsync(const GraphicsPipelineCreationArgs& InCreationArgs)
@@ -672,16 +670,12 @@ void FD3D12PipelineStateWorker::DoWork()
 {
 	if (bIsGraphics)
 	{
-		PSO = CreatePipelineStateWrapper(GetParentAdapter(), &CreationArgs.GraphicsArgs);
+		CreatePipelineStateWrapper(PSO.GetInitReference(), GetParentAdapter(), &CreationArgs.GraphicsArgs);
 	}
 	else
 	{
-		PSO = CreatePipelineStateWrapper(GetParentAdapter(), &CreationArgs.ComputeArgs);
+		CreatePipelineStateWrapper(PSO.GetInitReference(), GetParentAdapter(), &CreationArgs.ComputeArgs);
 	}
-	// Newly created d3d object has a ref-count of 1
-	// If we don't call Release here, its ref-count will not drop to 0
-	// even if we don't reference it anymore
-	PSO->Release();
 }
 
 #if LOG_PSO_CREATES
@@ -691,19 +685,18 @@ void FD3D12PipelineStateWorker::DoWork()
 
 DECLARE_CYCLE_STAT(TEXT("Create time"), STAT_PSOCreateTime, STATGROUP_D3D12PipelineState);
 
-template ID3D12PipelineState* CreatePipelineState<D3D12_GRAPHICS_PIPELINE_STATE_DESC>(ID3D12Device* Device, const D3D12_GRAPHICS_PIPELINE_STATE_DESC* Desc, ID3D12PipelineLibrary* Library, const TCHAR* Name);
-template ID3D12PipelineState* CreatePipelineState<D3D12_COMPUTE_PIPELINE_STATE_DESC>(ID3D12Device* Device, const D3D12_COMPUTE_PIPELINE_STATE_DESC* Desc, ID3D12PipelineLibrary* Library, const TCHAR* Name);
+template void CreatePipelineState<D3D12_GRAPHICS_PIPELINE_STATE_DESC>(ID3D12PipelineState*&PSO, ID3D12Device* Device, const D3D12_GRAPHICS_PIPELINE_STATE_DESC* Desc, ID3D12PipelineLibrary* Library, const TCHAR* Name);
+template void CreatePipelineState<D3D12_COMPUTE_PIPELINE_STATE_DESC>(ID3D12PipelineState*&PSO, ID3D12Device* Device, const D3D12_COMPUTE_PIPELINE_STATE_DESC* Desc, ID3D12PipelineLibrary* Library, const TCHAR* Name);
 
 // Thread-safe create graphics/compute pipeline state. Conditionally load/store the PSO using a Pipeline Library.
 template <typename TDesc>
-ID3D12PipelineState* CreatePipelineState(ID3D12Device* Device, const TDesc* Desc, ID3D12PipelineLibrary* Library, const TCHAR* Name)
+void CreatePipelineState(ID3D12PipelineState*&PSO, ID3D12Device* Device, const TDesc* Desc, ID3D12PipelineLibrary* Library, const TCHAR* Name)
 {
 #if LOG_PSO_CREATES
 	const FString CreatePipelineStateMessage = FString::Printf(TEXT("CreatePipelineState (%s, Hash = %s)"), *TPSOFunctionMap<TDesc>::GetString(), Name);
 	SCOPE_LOG_TIME(*CreatePipelineStateMessage, &GD3D12CreatePSOTime);
 #endif
 
-	ID3D12PipelineState* PSO = nullptr;
 	if (Library)
 	{
 		// Try to load the PSO from the library.
@@ -739,27 +732,24 @@ ID3D12PipelineState* CreatePipelineState(ID3D12Device* Device, const TDesc* Desc
 		if (FAILED(r))
 		{
 			UE_LOG(LogD3D12RHI, Error, TEXT("Failed to create PipelineState %s with hash %s"), *TPSOFunctionMap<TDesc>::GetString(), Name);
-			return nullptr;
 		}
 	}
 
 	check(PSO);
-	return PSO;
 }
 
-template ID3D12PipelineState* CreatePipelineStateFromStream<GraphicsPipelineCreationArgs_POD>(ID3D12Device2* Device, const D3D12_PIPELINE_STATE_STREAM_DESC* Desc, ID3D12PipelineLibrary1* Library, const TCHAR* Name);
-template ID3D12PipelineState* CreatePipelineStateFromStream<ComputePipelineCreationArgs_POD>(ID3D12Device2* Device, const D3D12_PIPELINE_STATE_STREAM_DESC* Desc, ID3D12PipelineLibrary1* Library, const TCHAR* Name);
+template void CreatePipelineStateFromStream<GraphicsPipelineCreationArgs_POD>(ID3D12PipelineState*& PSO, ID3D12Device2* Device, const D3D12_PIPELINE_STATE_STREAM_DESC* Desc, ID3D12PipelineLibrary1* Library, const TCHAR* Name);
+template void CreatePipelineStateFromStream<ComputePipelineCreationArgs_POD>(ID3D12PipelineState*& PSO, ID3D12Device2* Device, const D3D12_PIPELINE_STATE_STREAM_DESC* Desc, ID3D12PipelineLibrary1* Library, const TCHAR* Name);
 
 // Thread-safe create graphics/compute pipeline state. Conditionally load/store the PSO using a Pipeline Library.
 template <typename TDesc>
-ID3D12PipelineState* CreatePipelineStateFromStream(ID3D12Device2* Device, const D3D12_PIPELINE_STATE_STREAM_DESC* Desc, ID3D12PipelineLibrary1* Library, const TCHAR* Name)
+void CreatePipelineStateFromStream(ID3D12PipelineState*& PSO, ID3D12Device2* Device, const D3D12_PIPELINE_STATE_STREAM_DESC* Desc, ID3D12PipelineLibrary1* Library, const TCHAR* Name)
 {
 #if LOG_PSO_CREATES
 	const FString CreatePipelineStateMessage = FString::Printf(TEXT("CreatePipelineState (%s, Hash = %s)"), *TPSOStreamFunctionMap<TDesc>::GetString(), Name);
 	SCOPE_LOG_TIME(*CreatePipelineStateMessage, &GD3D12CreatePSOTime);
 #endif
 
-	ID3D12PipelineState* PSO = nullptr;
 	if (Library)
 	{
 		// Try to load the PSO from the library.
@@ -796,10 +786,8 @@ ID3D12PipelineState* CreatePipelineStateFromStream(ID3D12Device2* Device, const
 		if (FAILED(r))
 		{
 			UE_LOG(LogD3D12RHI, Error, TEXT("Failed to create PipelineState"));
-			return nullptr;
 		}
 	}
 
 	check(PSO);
-	return PSO;
 }
diff --git a/Engine/Source/Runtime/D3D12RHI/Private/Windows/WindowsD3D12PipelineState.h b/Engine/Source/Runtime/D3D12RHI/Private/Windows/WindowsD3D12PipelineState.h
--- a/Engine/Source/Runtime/D3D12RHI/Private/Windows/WindowsD3D12PipelineState.h
+++ b/Engine/Source/Runtime/D3D12RHI/Private/Windows/WindowsD3D12PipelineState.h
@@ -245,11 +245,11 @@ template<> struct TPSOStreamFunctionMap < ComputePipelineCreationArgs_POD >
 
 // Thread-safe create graphics/compute pipeline state. Conditionally load/store the PSO using a Pipeline Library.
 template <typename TDesc>
-ID3D12PipelineState* CreatePipelineState(ID3D12Device* Device, const TDesc* Desc, ID3D12PipelineLibrary* Library, const TCHAR* Name);
+void CreatePipelineState(ID3D12PipelineState*&PSO, ID3D12Device* Device, const TDesc* Desc, ID3D12PipelineLibrary* Library, const TCHAR* Name);
 
 // Thread-safe create graphics/compute pipeline state. Conditionally load/store the PSO using a Pipeline Library.
 template <typename TDesc>
-ID3D12PipelineState* CreatePipelineStateFromStream(ID3D12Device2* Device, const D3D12_PIPELINE_STATE_STREAM_DESC* Desc, ID3D12PipelineLibrary1* Library, const TCHAR* Name);
+void CreatePipelineStateFromStream(ID3D12PipelineState*& PSO, ID3D12Device2* Device, const D3D12_PIPELINE_STATE_STREAM_DESC* Desc, ID3D12PipelineLibrary1* Library, const TCHAR* Name);
 
 template <typename TDesc>
-ID3D12PipelineState* CreatePipelineStateWrapper(FD3D12Adapter* Adapter, const TDesc* CreationArgs);
+void CreatePipelineStateWrapper(ID3D12PipelineState** PSO, FD3D12Adapter* Adapter, const TDesc* CreationArgs);
