Copy 4447980 - Implement GetInitializer for some dx12 RHI classes

diff --git a/Engine/Source/Runtime/D3D12RHI/Private/D3D12Shader.h b/Engine/Source/Runtime/D3D12RHI/Private/D3D12Shader.h
--- a/Engine/Source/Runtime/D3D12RHI/Private/D3D12Shader.h
+++ b/Engine/Source/Runtime/D3D12RHI/Private/D3D12Shader.h
@@ -30,6 +30,8 @@ public:
 	{
 		FMemory::Memcpy(StreamStrides, InStrides, sizeof(StreamStrides));
 	}
+
+	virtual bool GetInitializer(FVertexDeclarationElementList& Init) final override;
 };
 
 /** This represents a vertex shader that hasn't been combined with a specific declaration to create a bound shader. */
diff --git a/Engine/Source/Runtime/D3D12RHI/Private/D3D12State.cpp b/Engine/Source/Runtime/D3D12RHI/Private/D3D12State.cpp
--- a/Engine/Source/Runtime/D3D12RHI/Private/D3D12State.cpp
+++ b/Engine/Source/Runtime/D3D12RHI/Private/D3D12State.cpp
@@ -35,6 +35,16 @@ static D3D12_CULL_MODE TranslateCullMode(ERasterizerCullMode CullMode)
 	};
 }
 
+static ERasterizerCullMode ReverseTranslateCullMode(D3D12_CULL_MODE CullMode)
+{
+	switch (CullMode)
+	{
+	case D3D12_CULL_MODE_BACK: return CM_CW;
+	case D3D12_CULL_MODE_FRONT: return CM_CCW;
+	default: return CM_None;
+	}
+}
+
 static D3D12_FILL_MODE TranslateFillMode(ERasterizerFillMode FillMode)
 {
 	switch (FillMode)
@@ -44,6 +54,15 @@ static D3D12_FILL_MODE TranslateFillMode(ERasterizerFillMode FillMode)
 	};
 }
 
+static ERasterizerFillMode ReverseTranslateFillMode(D3D12_FILL_MODE FillMode)
+{
+	switch (FillMode)
+	{
+	case D3D12_FILL_MODE_WIREFRAME: return FM_Wireframe;
+	default: return FM_Solid;
+	}
+}
+
 static D3D12_COMPARISON_FUNC TranslateCompareFunction(ECompareFunction CompareFunction)
 {
 	switch (CompareFunction)
@@ -59,6 +78,21 @@ static D3D12_COMPARISON_FUNC TranslateCompareFunction(ECompareFunction CompareFu
 	};
 }
 
+static ECompareFunction ReverseTranslateCompareFunction(D3D12_COMPARISON_FUNC CompareFunction)
+{
+	switch (CompareFunction)
+	{
+	case D3D12_COMPARISON_FUNC_LESS: return CF_Less;
+	case D3D12_COMPARISON_FUNC_LESS_EQUAL: return CF_LessEqual;
+	case D3D12_COMPARISON_FUNC_GREATER: return CF_Greater;
+	case D3D12_COMPARISON_FUNC_GREATER_EQUAL: return CF_GreaterEqual;
+	case D3D12_COMPARISON_FUNC_EQUAL: return CF_Equal;
+	case D3D12_COMPARISON_FUNC_NOT_EQUAL: return CF_NotEqual;
+	case D3D12_COMPARISON_FUNC_NEVER: return CF_Never;
+	default: return CF_Always;
+	}
+}
+
 static D3D12_COMPARISON_FUNC TranslateSamplerCompareFunction(ESamplerCompareFunction SamplerComparisonFunction)
 {
 	switch (SamplerComparisonFunction)
@@ -84,6 +118,21 @@ static D3D12_STENCIL_OP TranslateStencilOp(EStencilOp StencilOp)
 	};
 }
 
+static EStencilOp ReverseTranslateStencilOp(D3D12_STENCIL_OP StencilOp)
+{
+	switch (StencilOp)
+	{
+	case D3D12_STENCIL_OP_ZERO: return SO_Zero;
+	case D3D12_STENCIL_OP_REPLACE: return SO_Replace;
+	case D3D12_STENCIL_OP_INCR_SAT: return SO_SaturatedIncrement;
+	case D3D12_STENCIL_OP_DECR_SAT: return SO_SaturatedDecrement;
+	case D3D12_STENCIL_OP_INVERT: return SO_Invert;
+	case D3D12_STENCIL_OP_INCR: return SO_Increment;
+	case D3D12_STENCIL_OP_DECR: return SO_Decrement;
+	default: return SO_Keep;
+	};
+}
+
 static D3D12_BLEND_OP TranslateBlendOp(EBlendOperation BlendOp)
 {
 	switch (BlendOp)
@@ -96,6 +145,18 @@ static D3D12_BLEND_OP TranslateBlendOp(EBlendOperation BlendOp)
 	};
 }
 
+static EBlendOperation ReverseTranslateBlendOp(D3D12_BLEND_OP BlendOp)
+{
+	switch (BlendOp)
+	{
+	case D3D12_BLEND_OP_SUBTRACT: return BO_Subtract;
+	case D3D12_BLEND_OP_MIN: return BO_Min;
+	case D3D12_BLEND_OP_MAX: return BO_Max;
+	case D3D12_BLEND_OP_REV_SUBTRACT: return BO_ReverseSubtract;
+	default: return BO_Add;
+	};
+}
+
 static D3D12_BLEND TranslateBlendFactor(EBlendFactor BlendFactor)
 {
 	switch (BlendFactor)
@@ -115,6 +176,25 @@ static D3D12_BLEND TranslateBlendFactor(EBlendFactor BlendFactor)
 	};
 }
 
+static EBlendFactor ReverseTranslateBlendFactor(D3D12_BLEND BlendFactor)
+{
+	switch (BlendFactor)
+	{
+	case D3D12_BLEND_ONE: return BF_One;
+	case D3D12_BLEND_SRC_COLOR: return BF_SourceColor;
+	case D3D12_BLEND_INV_SRC_COLOR: return BF_InverseSourceColor;
+	case D3D12_BLEND_SRC_ALPHA: return BF_SourceAlpha;
+	case D3D12_BLEND_INV_SRC_ALPHA: return BF_InverseSourceAlpha;
+	case D3D12_BLEND_DEST_ALPHA: return BF_DestAlpha;
+	case D3D12_BLEND_INV_DEST_ALPHA: return BF_InverseDestAlpha;
+	case D3D12_BLEND_DEST_COLOR: return BF_DestColor;
+	case D3D12_BLEND_INV_DEST_COLOR: return BF_InverseDestColor;
+	case D3D12_BLEND_BLEND_FACTOR: return BF_ConstantBlendFactor;
+	case D3D12_BLEND_INV_BLEND_FACTOR: return BF_InverseConstantBlendFactor;
+	default: return BF_Zero;
+	};
+}
+
 bool operator==(const D3D12_SAMPLER_DESC& lhs, const D3D12_SAMPLER_DESC& rhs)
 {
 	return 0 == memcmp(&lhs, &rhs, sizeof(lhs));
@@ -225,6 +305,18 @@ FRasterizerStateRHIRef FD3D12DynamicRHI::RHICreateRasterizerState(const FRasteri
 	return RasterizerState;
 }
 
+bool FD3D12RasterizerState::GetInitializer(struct FRasterizerStateInitializerRHI& Init)
+{
+	Init.FillMode = ReverseTranslateFillMode(Desc.FillMode);
+	Init.CullMode = ReverseTranslateCullMode(Desc.CullMode);
+	Init.DepthBias = Desc.DepthBias / static_cast<float>(1 << 24);
+	check(Desc.DepthBias == FMath::FloorToInt(Init.DepthBias * static_cast<float>(1 << 24)));
+	Init.SlopeScaleDepthBias = Desc.SlopeScaledDepthBias;
+	Init.bAllowMSAA = Desc.MultisampleEnable;
+	Init.bEnableLineAA = false;
+	return true;
+}
+
 FDepthStencilStateRHIRef FD3D12DynamicRHI::RHICreateDepthStencilState(const FDepthStencilStateInitializerRHI& Initializer)
 {
 	FD3D12DepthStencilState* DepthStencilState = new FD3D12DepthStencilState;
@@ -275,6 +367,30 @@ FDepthStencilStateRHIRef FD3D12DynamicRHI::RHICreateDepthStencilState(const FDep
 	return DepthStencilState;
 }
 
+bool FD3D12DepthStencilState::GetInitializer(struct FDepthStencilStateInitializerRHI& Init)
+{
+	Init.bEnableDepthWrite = Desc.DepthWriteMask != D3D12_DEPTH_WRITE_MASK_ZERO;
+	Init.DepthTest = ReverseTranslateCompareFunction(Desc.DepthFunc);
+	Init.bEnableFrontFaceStencil = Desc.StencilEnable;
+	Init.FrontFaceStencilTest = ReverseTranslateCompareFunction(Desc.FrontFace.StencilFunc);
+	Init.FrontFaceStencilFailStencilOp = ReverseTranslateStencilOp(Desc.FrontFace.StencilFailOp);
+	Init.FrontFaceDepthFailStencilOp = ReverseTranslateStencilOp(Desc.FrontFace.StencilDepthFailOp);
+	Init.FrontFacePassStencilOp = ReverseTranslateStencilOp(Desc.FrontFace.StencilPassOp);
+	Init.bEnableBackFaceStencil =
+		Desc.StencilEnable &&
+		(Desc.FrontFace.StencilFunc != Desc.BackFace.StencilFunc ||
+			Desc.FrontFace.StencilFailOp != Desc.BackFace.StencilFailOp ||
+			Desc.FrontFace.StencilDepthFailOp != Desc.BackFace.StencilDepthFailOp ||
+			Desc.FrontFace.StencilPassOp != Desc.BackFace.StencilPassOp);
+	Init.BackFaceStencilTest = ReverseTranslateCompareFunction(Desc.BackFace.StencilFunc);
+	Init.BackFaceStencilFailStencilOp = ReverseTranslateStencilOp(Desc.BackFace.StencilFailOp);
+	Init.BackFaceDepthFailStencilOp = ReverseTranslateStencilOp(Desc.BackFace.StencilDepthFailOp);
+	Init.BackFacePassStencilOp = ReverseTranslateStencilOp(Desc.BackFace.StencilPassOp);
+	Init.StencilReadMask = Desc.StencilReadMask;
+	Init.StencilWriteMask = Desc.StencilWriteMask;
+	return true;
+}
+
 FBlendStateRHIRef FD3D12DynamicRHI::RHICreateBlendState(const FBlendStateInitializerRHI& Initializer)
 {
 	FD3D12BlendState* BlendState = new FD3D12BlendState;
@@ -308,6 +424,29 @@ FBlendStateRHIRef FD3D12DynamicRHI::RHICreateBlendState(const FBlendStateInitial
 	return BlendState;
 }
 
+bool FD3D12BlendState::GetInitializer(class FBlendStateInitializerRHI& Init)
+{
+	for (int32 Idx = 0; Idx < MaxSimultaneousRenderTargets; ++Idx)
+	{
+		const D3D12_RENDER_TARGET_BLEND_DESC& Src = Desc.RenderTarget[Idx];
+		FBlendStateInitializerRHI::FRenderTarget& Dst = Init.RenderTargets[Idx];
+
+		Dst.ColorBlendOp = ReverseTranslateBlendOp(Src.BlendOp);
+		Dst.ColorSrcBlend = ReverseTranslateBlendFactor(Src.SrcBlend);
+		Dst.ColorDestBlend = ReverseTranslateBlendFactor(Src.DestBlend);
+		Dst.AlphaBlendOp = ReverseTranslateBlendOp(Src.BlendOpAlpha);
+		Dst.AlphaSrcBlend = ReverseTranslateBlendFactor(Src.SrcBlendAlpha);
+		Dst.AlphaDestBlend = ReverseTranslateBlendFactor(Src.DestBlendAlpha);
+		Dst.ColorWriteMask = TEnumAsByte<EColorWriteMask>(
+			((Src.RenderTargetWriteMask & D3D12_COLOR_WRITE_ENABLE_RED) ? CW_RED : 0)
+			| ((Src.RenderTargetWriteMask & D3D12_COLOR_WRITE_ENABLE_GREEN) ? CW_GREEN : 0)
+			| ((Src.RenderTargetWriteMask & D3D12_COLOR_WRITE_ENABLE_BLUE) ? CW_BLUE : 0)
+			| ((Src.RenderTargetWriteMask & D3D12_COLOR_WRITE_ENABLE_ALPHA) ? CW_ALPHA : 0));
+	}
+	Init.bUseIndependentRenderTargetBlendStates = Desc.IndependentBlendEnable;
+	return true;
+}
+
 FGraphicsPipelineStateRHIRef FD3D12DynamicRHI::RHICreateGraphicsPipelineState(const FGraphicsPipelineStateInitializer& Initializer)
 {
 	SCOPE_CYCLE_COUNTER(STAT_PSOGraphicsFindOrCreateTime);
diff --git a/Engine/Source/Runtime/D3D12RHI/Private/D3D12VertexDeclaration.cpp b/Engine/Source/Runtime/D3D12RHI/Private/D3D12VertexDeclaration.cpp
--- a/Engine/Source/Runtime/D3D12RHI/Private/D3D12VertexDeclaration.cpp
+++ b/Engine/Source/Runtime/D3D12RHI/Private/D3D12VertexDeclaration.cpp
@@ -152,3 +152,46 @@ FVertexDeclarationRHIRef FD3D12DynamicRHI::RHICreateVertexDeclaration(const FVer
 
 	return *VertexDeclarationRefPtr;
 }
+
+bool FD3D12VertexDeclaration::GetInitializer(FVertexDeclarationElementList& Init)
+{
+	const int32 NumVertElems = VertexElements.Num();
+	Init.Empty(NumVertElems);
+	Init.AddUninitialized(NumVertElems);
+
+	for (int32 Idx = 0; Idx < NumVertElems; ++Idx)
+	{
+		const D3D12_INPUT_ELEMENT_DESC& Elem = VertexElements[Idx];
+		FVertexElement& Out = Init[Idx];
+
+		Out.StreamIndex = Elem.InputSlot;
+		Out.Offset = Elem.AlignedByteOffset;
+		switch (Elem.Format)
+		{
+		case DXGI_FORMAT_R32_FLOAT: Out.Type = VET_Float1; break;
+		case DXGI_FORMAT_R32G32_FLOAT: Out.Type = VET_Float2; break;
+		case DXGI_FORMAT_R32G32B32_FLOAT: Out.Type = VET_Float3; break;
+		case DXGI_FORMAT_R32G32B32A32_FLOAT: Out.Type = VET_Float4; break;
+		case DXGI_FORMAT_R8G8B8A8_SNORM: Out.Type = VET_PackedNormal; break;
+		case DXGI_FORMAT_R8G8B8A8_UINT: Out.Type = VET_UByte4; break;
+		case DXGI_FORMAT_R8G8B8A8_UNORM: Out.Type = VET_UByte4N; break;
+		case DXGI_FORMAT_B8G8R8A8_UNORM: Out.Type = VET_Color; break;
+		case DXGI_FORMAT_R16G16_SINT: Out.Type = VET_Short2; break;
+		case DXGI_FORMAT_R16G16B16A16_SINT: Out.Type = VET_Short4; break;
+		case DXGI_FORMAT_R16G16_SNORM: Out.Type = VET_Short2N; break;
+		case DXGI_FORMAT_R16G16_FLOAT: Out.Type = VET_Half2; break;
+		case DXGI_FORMAT_R16G16B16A16_FLOAT: Out.Type = VET_Half4; break;
+		case DXGI_FORMAT_R16G16B16A16_SNORM: Out.Type = VET_Short4N; break;
+		case DXGI_FORMAT_R16G16_UINT: Out.Type = VET_UShort2; break;
+		case DXGI_FORMAT_R16G16B16A16_UINT: Out.Type = VET_UShort4; break;
+		case DXGI_FORMAT_R16G16_UNORM: Out.Type = VET_UShort2N; break;
+		case DXGI_FORMAT_R16G16B16A16_UNORM: Out.Type = VET_UShort4N; break;
+		case DXGI_FORMAT_R10G10B10A2_UNORM: Out.Type = VET_URGB10A2N; break;
+		default: UE_LOG(LogD3D12RHI, Fatal, TEXT("Unknown D3D vertex element type %u"), Elem.Format);
+		}
+		Out.AttributeIndex = Elem.SemanticIndex;
+		Out.Stride = StreamStrides[Elem.InputSlot];
+		Out.bUseInstanceIndex = Elem.InputSlotClass == D3D12_INPUT_CLASSIFICATION_PER_INSTANCE_DATA;
+	}
+	return true;
+}
diff --git a/Engine/Source/Runtime/D3D12RHI/Public/D3D12State.h b/Engine/Source/Runtime/D3D12RHI/Public/D3D12State.h
--- a/Engine/Source/Runtime/D3D12RHI/Public/D3D12State.h
+++ b/Engine/Source/Runtime/D3D12RHI/Public/D3D12State.h
@@ -21,6 +21,8 @@ class FD3D12RasterizerState : public FRHIRasterizerState
 {
 public:
 	D3D12_RASTERIZER_DESC Desc;
+
+	virtual bool GetInitializer(struct FRasterizerStateInitializerRHI& Init) final override;
 };
 
 class FD3D12DepthStencilState : public FRHIDepthStencilState
@@ -31,6 +33,8 @@ public:
 
 	/* Describes the read/write state of the separate depth and stencil components of the DSV. */
 	FExclusiveDepthStencil AccessType;
+
+	virtual bool GetInitializer(struct FDepthStencilStateInitializerRHI& Init) final override;
 };
 
 class FD3D12BlendState : public FRHIBlendState
@@ -38,5 +42,7 @@ class FD3D12BlendState : public FRHIBlendState
 public:
 
 	D3D12_BLEND_DESC Desc;
+
+	virtual bool GetInitializer(class FBlendStateInitializerRHI& Init) final override;
 };
 
