Copy 4452949 - Fix a VS2015 compile error. Credits to Saad Nader

diff --git a/Engine/Source/Runtime/D3D12RHI/Private/D3D12State.cpp b/Engine/Source/Runtime/D3D12RHI/Private/D3D12State.cpp
--- a/Engine/Source/Runtime/D3D12RHI/Private/D3D12State.cpp
+++ b/Engine/Source/Runtime/D3D12RHI/Private/D3D12State.cpp
@@ -312,7 +312,7 @@ bool FD3D12RasterizerState::GetInitializer(struct FRasterizerStateInitializerRHI
 	Init.DepthBias = Desc.DepthBias / static_cast<float>(1 << 24);
 	check(Desc.DepthBias == FMath::FloorToInt(Init.DepthBias * static_cast<float>(1 << 24)));
 	Init.SlopeScaleDepthBias = Desc.SlopeScaledDepthBias;
-	Init.bAllowMSAA = Desc.MultisampleEnable;
+	Init.bAllowMSAA = static_cast<bool>(Desc.MultisampleEnable);
 	Init.bEnableLineAA = false;
 	return true;
 }
@@ -371,7 +371,7 @@ bool FD3D12DepthStencilState::GetInitializer(struct FDepthStencilStateInitialize
 {
 	Init.bEnableDepthWrite = Desc.DepthWriteMask != D3D12_DEPTH_WRITE_MASK_ZERO;
 	Init.DepthTest = ReverseTranslateCompareFunction(Desc.DepthFunc);
-	Init.bEnableFrontFaceStencil = Desc.StencilEnable;
+	Init.bEnableFrontFaceStencil = static_cast<bool>(Desc.StencilEnable);
 	Init.FrontFaceStencilTest = ReverseTranslateCompareFunction(Desc.FrontFace.StencilFunc);
 	Init.FrontFaceStencilFailStencilOp = ReverseTranslateStencilOp(Desc.FrontFace.StencilFailOp);
 	Init.FrontFaceDepthFailStencilOp = ReverseTranslateStencilOp(Desc.FrontFace.StencilDepthFailOp);
@@ -443,7 +443,7 @@ bool FD3D12BlendState::GetInitializer(class FBlendStateInitializerRHI& Init)
 			| ((Src.RenderTargetWriteMask & D3D12_COLOR_WRITE_ENABLE_BLUE) ? CW_BLUE : 0)
 			| ((Src.RenderTargetWriteMask & D3D12_COLOR_WRITE_ENABLE_ALPHA) ? CW_ALPHA : 0));
 	}
-	Init.bUseIndependentRenderTargetBlendStates = Desc.IndependentBlendEnable;
+	Init.bUseIndependentRenderTargetBlendStates = static_cast<bool>(Desc.IndependentBlendEnable);
 	return true;
 }
 
