Add CREATEPIPELINELIBRARY_DRIVERVERSIONMISMATCH to -d3ddebug ignore list

diff --git a/Engine/Source/Runtime/D3D12RHI/Private/D3D12Adapter.cpp b/Engine/Source/Runtime/D3D12RHI/Private/D3D12Adapter.cpp
--- a/Engine/Source/Runtime/D3D12RHI/Private/D3D12Adapter.cpp
+++ b/Engine/Source/Runtime/D3D12RHI/Private/D3D12Adapter.cpp
@@ -212,6 +212,11 @@ void FD3D12Adapter::CreateRootDevice(bool bWithDebug)
 				// but will hopefully be resolved when the RHI switches to use the engine's resource tracking system.
 				(D3D12_MESSAGE_ID)1008,
 
+				// This error gets generated on the first run when you install a new driver. The code handles this error properly and resets the PipelineLibrary,
+				// so we can safely ignore this message. It could possibly be avoided by adding driver version to the PSO cache filename, but an average user is unlikely
+				// to be interested in keeping PSO caches associated with old drivers around on disk, so it's better to just reset.
+				D3D12_MESSAGE_ID_CREATEPIPELINELIBRARY_DRIVERVERSIONMISMATCH,
+
 #if ENABLE_RESIDENCY_MANAGEMENT
 				// TODO: Remove this when the debug layers work for executions which are guarded by a fence
 				D3D12_MESSAGE_ID_INVALID_USE_OF_NON_RESIDENT_RESOURCE
