Copy 4474736 - remove an unnecessary RT stall in d3d12rhi

diff --git a/Engine/Source/Runtime/D3D12RHI/Private/D3D12Texture.cpp b/Engine/Source/Runtime/D3D12RHI/Private/D3D12Texture.cpp
--- a/Engine/Source/Runtime/D3D12RHI/Private/D3D12Texture.cpp
+++ b/Engine/Source/Runtime/D3D12RHI/Private/D3D12Texture.cpp
@@ -1859,12 +1859,20 @@ void TD3D12Texture2D<RHIResourceType>::UpdateTexture2D(class FRHICommandListImme
 
 void* FD3D12DynamicRHI::LockTexture2D_RenderThread(class FRHICommandListImmediate& RHICmdList, FTexture2DRHIParamRef TextureRHI, uint32 MipIndex, EResourceLockMode LockMode, uint32& DestStride, bool bLockWithinMiptail, bool bNeedsDefaultRHIFlush)
 {
+	// TODO: Remove when we are sure this branch is unnecessary
+	// bNeedsDefaultRHIFlush is default to true, which causes unecessary stall on render thread
+	// TD3D12Texture2D<...>::Lock handles different type of locks:
+	//   RLM_WriteOnly - system memory is allocated and returned. Buffer renaming happens on
+	//     RHI thread and is triggered on unlock
+	//   RLM_ReadOnly - slow path, flush RHI thread
+#if 0
 	if (bNeedsDefaultRHIFlush)
 	{
 		QUICK_SCOPE_CYCLE_COUNTER(STAT_RHIMETHOD_LockTexture2D_Flush);
 		RHICmdList.ImmediateFlush(EImmediateFlushType::FlushRHIThread);
 		return RHILockTexture2D(TextureRHI, MipIndex, LockMode, DestStride, bLockWithinMiptail);
 	}
+#endif
 
 	check(TextureRHI);
 	FD3D12Texture2D* Texture = FD3D12DynamicRHI::ResourceCast(TextureRHI);
@@ -1880,6 +1888,8 @@ void* FD3D12DynamicRHI::RHILockTexture2D(FTexture2DRHIParamRef TextureRHI, uint3
 
 void FD3D12DynamicRHI::UnlockTexture2D_RenderThread(class FRHICommandListImmediate& RHICmdList, FTexture2DRHIParamRef TextureRHI, uint32 MipIndex, bool bLockWithinMiptail, bool bNeedsDefaultRHIFlush)
 {
+	// TODO: Remove when we are sure this branch is unnecessary
+#if 0
 	if (bNeedsDefaultRHIFlush)
 	{
 		QUICK_SCOPE_CYCLE_COUNTER(STAT_RHIMETHOD_UnlockTexture2D_Flush);
@@ -1887,6 +1897,7 @@ void FD3D12DynamicRHI::UnlockTexture2D_RenderThread(class FRHICommandListImmedia
 		RHIUnlockTexture2D(TextureRHI, MipIndex, bLockWithinMiptail);
 		return;
 	}
+#endif
 
 	check(TextureRHI);
 	FD3D12Texture2D* Texture = FD3D12DynamicRHI::ResourceCast(TextureRHI);
