Ignore bNeedsDefaultRHIFlush only on PC dx12

diff --git a/Engine/Source/Runtime/D3D12RHI/Private/D3D12Texture.cpp b/Engine/Source/Runtime/D3D12RHI/Private/D3D12Texture.cpp
--- a/Engine/Source/Runtime/D3D12RHI/Private/D3D12Texture.cpp
+++ b/Engine/Source/Runtime/D3D12RHI/Private/D3D12Texture.cpp
@@ -1865,14 +1865,13 @@ void* FD3D12DynamicRHI::LockTexture2D_RenderThread(class FRHICommandListImmediat
 	//   RLM_WriteOnly - system memory is allocated and returned. Buffer renaming happens on
 	//     RHI thread and is triggered on unlock
 	//   RLM_ReadOnly - slow path, flush RHI thread
-#if 0
-	if (bNeedsDefaultRHIFlush)
+	// XB1 uses virtual textures and read/writes GPU memory directly so a flush may still be needed
+	if (!PLATFORM_WINDOWS && bNeedsDefaultRHIFlush)
 	{
 		QUICK_SCOPE_CYCLE_COUNTER(STAT_RHIMETHOD_LockTexture2D_Flush);
 		RHICmdList.ImmediateFlush(EImmediateFlushType::FlushRHIThread);
 		return RHILockTexture2D(TextureRHI, MipIndex, LockMode, DestStride, bLockWithinMiptail);
 	}
-#endif
 
 	check(TextureRHI);
 	FD3D12Texture2D* Texture = FD3D12DynamicRHI::ResourceCast(TextureRHI);
@@ -1889,15 +1888,13 @@ void* FD3D12DynamicRHI::RHILockTexture2D(FTexture2DRHIParamRef TextureRHI, uint3
 void FD3D12DynamicRHI::UnlockTexture2D_RenderThread(class FRHICommandListImmediate& RHICmdList, FTexture2DRHIParamRef TextureRHI, uint32 MipIndex, bool bLockWithinMiptail, bool bNeedsDefaultRHIFlush)
 {
 	// TODO: Remove when we are sure this branch is unnecessary
-#if 0
-	if (bNeedsDefaultRHIFlush)
+	if (!PLATFORM_WINDOWS && bNeedsDefaultRHIFlush)
 	{
 		QUICK_SCOPE_CYCLE_COUNTER(STAT_RHIMETHOD_UnlockTexture2D_Flush);
 		RHICmdList.ImmediateFlush(EImmediateFlushType::FlushRHIThread);
 		RHIUnlockTexture2D(TextureRHI, MipIndex, bLockWithinMiptail);
 		return;
 	}
-#endif
 
 	check(TextureRHI);
 	FD3D12Texture2D* Texture = FD3D12DynamicRHI::ResourceCast(TextureRHI);
