Clean up a bunch of unreleased D3D objects

diff --git a/Engine/Source/Runtime/D3D12RHI/Private/D3D12Adapter.cpp b/Engine/Source/Runtime/D3D12RHI/Private/D3D12Adapter.cpp
--- a/Engine/Source/Runtime/D3D12RHI/Private/D3D12Adapter.cpp
+++ b/Engine/Source/Runtime/D3D12RHI/Private/D3D12Adapter.cpp
@@ -513,6 +513,12 @@ void FD3D12Adapter::Cleanup()
 	}
 
 	PipelineStateCache.Close();
+	RootSignatureManager.Destroy();
+
+	DrawIndirectCommandSignature.SafeRelease();
+	DrawIndexedIndirectCommandSignature.SafeRelease();
+	DispatchIndirectCommandSignature.SafeRelease();
+
 	FenceCorePool.Destroy();
 }
 
diff --git a/Engine/Source/Runtime/D3D12RHI/Private/D3D12RootSignature.cpp b/Engine/Source/Runtime/D3D12RHI/Private/D3D12RootSignature.cpp
--- a/Engine/Source/Runtime/D3D12RHI/Private/D3D12RootSignature.cpp
+++ b/Engine/Source/Runtime/D3D12RHI/Private/D3D12RootSignature.cpp
@@ -456,13 +456,14 @@ void FD3D12RootSignature::InternalAnalyzeSignature(const RootSignatureDescType&
 	}
 }
 
-FD3D12RootSignatureManager::~FD3D12RootSignatureManager()
+void FD3D12RootSignatureManager::Destroy()
 {
 	for (auto Iter = RootSignatureMap.CreateIterator(); Iter; ++Iter)
 	{
 		FD3D12RootSignature* pRootSignature = Iter.Value();
 		delete pRootSignature;
 	}
+	RootSignatureMap.Reset();
 }
 
 FD3D12RootSignature* FD3D12RootSignatureManager::GetRootSignature(const FD3D12QuantizedBoundShaderState& QBSS)
diff --git a/Engine/Source/Runtime/D3D12RHI/Private/D3D12RootSignature.h b/Engine/Source/Runtime/D3D12RHI/Private/D3D12RootSignature.h
--- a/Engine/Source/Runtime/D3D12RHI/Private/D3D12RootSignature.h
+++ b/Engine/Source/Runtime/D3D12RHI/Private/D3D12RootSignature.h
@@ -452,7 +452,12 @@ public:
 		: FD3D12AdapterChild(InParent)
 	{
 	}
-	~FD3D12RootSignatureManager();
+	~FD3D12RootSignatureManager()
+	{
+		check(RootSignatureMap.Num() == 0);
+	}
+	void Destroy();
+
 	FD3D12RootSignature* GetRootSignature(const FD3D12QuantizedBoundShaderState& QBSS);
 	FD3D12QuantizedBoundShaderState GetQuantizedBoundShaderState(const FD3D12RootSignature* const RootSignature);
 
