Minor cleanup of PSO creation. Cut most ties to BoundShaderState.

diff --git a/Engine/Source/Runtime/D3D12RHI/Private/D3D12PipelineState.h b/Engine/Source/Runtime/D3D12RHI/Private/D3D12PipelineState.h
--- a/Engine/Source/Runtime/D3D12RHI/Private/D3D12PipelineState.h
+++ b/Engine/Source/Runtime/D3D12RHI/Private/D3D12PipelineState.h
@@ -112,7 +112,7 @@ struct FD3D12ComputePipelineStateDesc
 };
 
 
-FD3D12LowLevelGraphicsPipelineStateDesc GetLowLevelGraphicsPipelineStateDesc(const FGraphicsPipelineStateInitializer& Initializer, FD3D12BoundShaderState* BoundShaderState);
+FD3D12LowLevelGraphicsPipelineStateDesc GetLowLevelGraphicsPipelineStateDesc(const FGraphicsPipelineStateInitializer& Initializer, const FD3D12RootSignature* RootSignature);
 FD3D12ComputePipelineStateDesc GetComputePipelineStateDesc(const FD3D12ComputeShader* ComputeShader);
 
 #define PSO_IF_NOT_EQUAL_RETURN_FALSE( value ) if(lhs.##value != rhs.##value){ return false; }
@@ -308,7 +308,7 @@ protected:
 
 struct FD3D12GraphicsPipelineState : public FRHIGraphicsPipelineState
 {
-	explicit FD3D12GraphicsPipelineState(const FGraphicsPipelineStateInitializer& Initializer, FD3D12BoundShaderState* InBoundShaderState, FD3D12PipelineState* InPipelineState);
+	explicit FD3D12GraphicsPipelineState(const FGraphicsPipelineStateInitializer& Initializer, const FD3D12RootSignature* InRootSignature, FD3D12PipelineState* InPipelineState);
 	~FD3D12GraphicsPipelineState();
 
 	const FGraphicsPipelineStateInitializer PipelineStateInitializer;
@@ -428,7 +428,7 @@ protected:
 	typedef TFunction<void(FD3D12PipelineState*, const FD3D12ComputePipelineStateDesc&)> FPostCreateComputeCallback;
 
 #if D3D12RHI_USE_HIGH_LEVEL_PSO_CACHE
-	virtual FD3D12GraphicsPipelineState* AddToRuntimeCache(const FGraphicsPipelineStateInitializer& Initializer, uint32 InitializerHash, FD3D12BoundShaderState* BoundShaderState, FD3D12PipelineState* PipelineState);
+	virtual FD3D12GraphicsPipelineState* AddToRuntimeCache(const FGraphicsPipelineStateInitializer& Initializer, uint32 InitializerHash, const FD3D12RootSignature* RootSignature, FD3D12PipelineState* PipelineState);
 	FD3D12ComputePipelineState* AddToRuntimeCache(FD3D12ComputeShader* ComputeShader, FD3D12PipelineState* PipelineState);
 #endif
 
@@ -443,14 +443,14 @@ protected:
 	virtual void OnPSOCreated(FD3D12PipelineState* PipelineState, const FD3D12ComputePipelineStateDesc& Desc) = 0;
 
 #if !D3D12RHI_USE_HIGH_LEVEL_PSO_CACHE
-	FD3D12GraphicsPipelineState* FindInLoadedCache(const FGraphicsPipelineStateInitializer& Initializer, FD3D12BoundShaderState* BoundShaderState, FD3D12LowLevelGraphicsPipelineStateDesc& OutLowLevelDesc);
-	FD3D12GraphicsPipelineState* CreateAndAdd(const FGraphicsPipelineStateInitializer& Initializer, FD3D12BoundShaderState* BoundShaderState, const FD3D12LowLevelGraphicsPipelineStateDesc& LowLevelDesc);
+	FD3D12GraphicsPipelineState* FindInLoadedCache(const FGraphicsPipelineStateInitializer& Initializer, const FD3D12RootSignature* RootSignature, FD3D12LowLevelGraphicsPipelineStateDesc& OutLowLevelDesc);
+	FD3D12GraphicsPipelineState* CreateAndAdd(const FGraphicsPipelineStateInitializer& Initializer, const FD3D12RootSignature* RootSignature, const FD3D12LowLevelGraphicsPipelineStateDesc& LowLevelDesc);
 #endif
 public:
 #if D3D12RHI_USE_HIGH_LEVEL_PSO_CACHE
 	FD3D12GraphicsPipelineState* FindInRuntimeCache(const FGraphicsPipelineStateInitializer& Initializer, uint32& OutHash);
-	FD3D12GraphicsPipelineState* FindInLoadedCache(const FGraphicsPipelineStateInitializer& Initializer, uint32 InitializerHash, FD3D12BoundShaderState* BoundShaderState, FD3D12LowLevelGraphicsPipelineStateDesc& OutLowLevelDesc);
-	FD3D12GraphicsPipelineState* CreateAndAdd(const FGraphicsPipelineStateInitializer& Initializer, uint32 InitializerHash, FD3D12BoundShaderState* BoundShaderState, const FD3D12LowLevelGraphicsPipelineStateDesc& LowLevelDesc);
+	FD3D12GraphicsPipelineState* FindInLoadedCache(const FGraphicsPipelineStateInitializer& Initializer, uint32 InitializerHash, const FD3D12RootSignature* RootSignature, FD3D12LowLevelGraphicsPipelineStateDesc& OutLowLevelDesc);
+	FD3D12GraphicsPipelineState* CreateAndAdd(const FGraphicsPipelineStateInitializer& Initializer, uint32 InitializerHash, const FD3D12RootSignature* RootSignature, const FD3D12LowLevelGraphicsPipelineStateDesc& LowLevelDesc);
 
 	FD3D12ComputePipelineState* FindInRuntimeCache(const FD3D12ComputeShader* ComputeShader);
 #endif
diff --git a/Engine/Source/Runtime/D3D12RHI/Private/D3D12Pipelinestate.cpp b/Engine/Source/Runtime/D3D12RHI/Private/D3D12Pipelinestate.cpp
--- a/Engine/Source/Runtime/D3D12RHI/Private/D3D12Pipelinestate.cpp
+++ b/Engine/Source/Runtime/D3D12RHI/Private/D3D12Pipelinestate.cpp
@@ -24,14 +24,14 @@ static TAutoConsoleVariable<float> CVarPSOStallTimeoutInMs(
 
 /// @cond DOXYGEN_WARNINGS
 
-FD3D12LowLevelGraphicsPipelineStateDesc GetLowLevelGraphicsPipelineStateDesc(const FGraphicsPipelineStateInitializer& Initializer, FD3D12BoundShaderState* BoundShaderState)
+FD3D12LowLevelGraphicsPipelineStateDesc GetLowLevelGraphicsPipelineStateDesc(const FGraphicsPipelineStateInitializer& Initializer, const FD3D12RootSignature* RootSignature)
 {
 	// Memzero because we hash using the entire struct and we need to clear any padding.
 	FD3D12LowLevelGraphicsPipelineStateDesc Desc;
 	FMemory::Memzero(&Desc, sizeof(Desc));
 
-	Desc.pRootSignature = BoundShaderState->pRootSignature;
-	Desc.Desc.pRootSignature = Desc.pRootSignature->GetRootSignature();
+	Desc.pRootSignature = RootSignature;
+	Desc.Desc.pRootSignature = RootSignature->GetRootSignature();
 
 #if !D3D12_USE_DERIVED_PSO
 	Desc.Desc.BlendState = Initializer.BlendState ? FD3D12DynamicRHI::ResourceCast(Initializer.BlendState)->Desc : CD3DX12_BLEND_DESC(D3D12_DEFAULT);
@@ -40,7 +40,7 @@ FD3D12LowLevelGraphicsPipelineStateDesc GetLowLevelGraphicsPipelineStateDesc(con
 	Desc.Desc.DepthStencilState = Initializer.DepthStencilState ? CD3DX12_DEPTH_STENCIL_DESC1(FD3D12DynamicRHI::ResourceCast(Initializer.DepthStencilState)->Desc) : CD3DX12_DEPTH_STENCIL_DESC1(D3D12_DEFAULT);
 #endif // !D3D12_USE_DERIVED_PSO
 
-	if (BoundShaderState->GetHullShader() && BoundShaderState->GetDomainShader())
+	if (Initializer.BoundShaderState.HullShaderRHI && Initializer.BoundShaderState.DomainShaderRHI)
 	{
 		Desc.Desc.PrimitiveTopologyType = D3D12_PRIMITIVE_TOPOLOGY_TYPE_PATCH;
 	}
@@ -54,15 +54,19 @@ FD3D12LowLevelGraphicsPipelineStateDesc GetLowLevelGraphicsPipelineStateDesc(con
 	Desc.Desc.SampleDesc.Count = Initializer.NumSamples;
 	Desc.Desc.SampleDesc.Quality = GetMaxMSAAQuality(Initializer.NumSamples);
 
-	Desc.Desc.InputLayout = BoundShaderState->InputLayout;
+	if (FD3D12VertexDeclaration* InputLayout = (FD3D12VertexDeclaration*) Initializer.BoundShaderState.VertexDeclarationRHI)
+	{
+		Desc.Desc.InputLayout.NumElements        = InputLayout->VertexElements.Num();
+		Desc.Desc.InputLayout.pInputElementDescs = InputLayout->VertexElements.GetData();
+	}
 
-	if (BoundShaderState->GetGeometryShader())
+	if (Initializer.BoundShaderState.GeometryShaderRHI)
 	{
-		Desc.Desc.StreamOutput = BoundShaderState->GetGeometryShader()->StreamOutput;
+		Desc.Desc.StreamOutput = ((FD3D12GeometryShader*) Initializer.BoundShaderState.GeometryShaderRHI)->StreamOutput;
 	}
 
 #define COPY_SHADER(Initial, Name) \
-	if (FD3D12##Name##Shader* Shader = BoundShaderState->Get##Name##Shader()) \
+	if (FD3D12##Name##Shader* Shader = (FD3D12##Name##Shader*) Initializer.BoundShaderState.##Name##ShaderRHI) \
 	{ \
 		Desc.Desc.Initial##S = Shader->ShaderBytecode.GetShaderBytecode(); \
 		Desc.Initial##SHash = Shader->ShaderBytecode.GetHash(); \
@@ -304,13 +308,16 @@ ID3D12PipelineState* FD3D12PipelineState::GetPipelineState()
 
 FD3D12GraphicsPipelineState::FD3D12GraphicsPipelineState(
 	const FGraphicsPipelineStateInitializer& Initializer,
-	FD3D12BoundShaderState* InBoundShaderState,
+	const FD3D12RootSignature* InRootSignature,
 	FD3D12PipelineState* InPipelineState)
 	: PipelineStateInitializer(Initializer)
-	, RootSignature(InBoundShaderState->pRootSignature)
+	, RootSignature(InRootSignature)
 	, PipelineState(InPipelineState)
 {
-	FMemory::Memcpy(StreamStrides, InBoundShaderState->StreamStrides, sizeof(StreamStrides));
+	if (Initializer.BoundShaderState.VertexDeclarationRHI)
+		FMemory::Memcpy(StreamStrides, ((FD3D12VertexDeclaration*) Initializer.BoundShaderState.VertexDeclarationRHI)->StreamStrides, sizeof(StreamStrides));
+	else
+		FMemory::Memzero(StreamStrides, sizeof(StreamStrides));
 	bShaderNeedsGlobalConstantBuffer[SF_Vertex] = GetVertexShader() && GetVertexShader()->ResourceCounts.bGlobalUniformBufferUsed;
 	bShaderNeedsGlobalConstantBuffer[SF_Pixel] = GetPixelShader() && GetPixelShader()->ResourceCounts.bGlobalUniformBufferUsed;
 	bShaderNeedsGlobalConstantBuffer[SF_Hull] = GetHullShader() && GetHullShader()->ResourceCounts.bGlobalUniformBufferUsed;
@@ -389,11 +396,11 @@ void FD3D12PipelineStateCacheBase::CleanupPipelineStateCaches()
 }
 
 #if D3D12RHI_USE_HIGH_LEVEL_PSO_CACHE
-FD3D12GraphicsPipelineState* FD3D12PipelineStateCacheBase::AddToRuntimeCache(const FGraphicsPipelineStateInitializer& Initializer, uint32 InitializerHash, FD3D12BoundShaderState* BoundShaderState, FD3D12PipelineState* PipelineState)
+FD3D12GraphicsPipelineState* FD3D12PipelineStateCacheBase::AddToRuntimeCache(const FGraphicsPipelineStateInitializer& Initializer, uint32 InitializerHash, const FD3D12RootSignature* RootSignature, FD3D12PipelineState* PipelineState)
 {
 	// Lifetime managed by the runtime cache. AddRef() so the upper level doesn't delete the FD3D12GraphicsPipelineState objects while they're still in the runtime cache.
 	// One alternative is to remove the object from the runtime cache in the FD3D12GraphicsPipelineState destructor.
-	FD3D12GraphicsPipelineState* const GraphicsPipelineState = new FD3D12GraphicsPipelineState(Initializer, BoundShaderState, PipelineState);
+	FD3D12GraphicsPipelineState* const GraphicsPipelineState = new FD3D12GraphicsPipelineState(Initializer, RootSignature, PipelineState);
 	GraphicsPipelineState->AddRef();
 
 	check(GraphicsPipelineState && InitializerHash != 0);
@@ -580,11 +587,11 @@ FD3D12GraphicsPipelineState* FD3D12PipelineStateCacheBase::FindInLoadedCache(
 #if D3D12RHI_USE_HIGH_LEVEL_PSO_CACHE
 	uint32 InitializerHash,
 #endif
-	FD3D12BoundShaderState* BoundShaderState,
+	const FD3D12RootSignature* RootSignature,
 	FD3D12LowLevelGraphicsPipelineStateDesc& OutLowLevelDesc)
 {
 	// TODO: For now PSOs will be created on every node of the LDA chain.
-	OutLowLevelDesc = GetLowLevelGraphicsPipelineStateDesc(Initializer, BoundShaderState);
+	OutLowLevelDesc = GetLowLevelGraphicsPipelineStateDesc(Initializer, RootSignature);
 	OutLowLevelDesc.Desc.NodeMask = (uint32)FRHIGPUMask::All();
 
 	OutLowLevelDesc.CombinedHash = FD3D12PipelineStateCacheBase::HashPSODesc(OutLowLevelDesc);
@@ -595,7 +602,7 @@ FD3D12GraphicsPipelineState* FD3D12PipelineStateCacheBase::FindInLoadedCache(
 	if (PipelineState)
 	{
 		// Add the PSO to the runtime cache for better performance next time.
-		return AddToRuntimeCache(Initializer, InitializerHash, BoundShaderState, PipelineState);
+		return AddToRuntimeCache(Initializer, InitializerHash, RootSignature, PipelineState);
 	}
 
 	// TODO: Try to load from a PipelineLibrary now instead of at Create time.
@@ -604,7 +611,7 @@ FD3D12GraphicsPipelineState* FD3D12PipelineStateCacheBase::FindInLoadedCache(
 #else
 	if (PipelineState)
 	{
-		return new FD3D12GraphicsPipelineState(Initializer, BoundShaderState, PipelineState);
+		return new FD3D12GraphicsPipelineState(Initializer, RootSignature, PipelineState);
 	}
 	else
 	{
@@ -618,7 +625,7 @@ FD3D12GraphicsPipelineState* FD3D12PipelineStateCacheBase::CreateAndAdd(
 #if D3D12RHI_USE_HIGH_LEVEL_PSO_CACHE
 	uint32 InitializerHash,
 #endif
-	FD3D12BoundShaderState* BoundShaderState,
+	const FD3D12RootSignature* RootSignature,
 	const FD3D12LowLevelGraphicsPipelineStateDesc& LowLevelDesc)
 {
 	FD3D12PipelineState* const PipelineState = CreateAndAddToLowLevelCache(LowLevelDesc);
@@ -629,11 +636,11 @@ FD3D12GraphicsPipelineState* FD3D12PipelineStateCacheBase::CreateAndAdd(
 	}
 
 	// Add the PSO to the runtime cache for better performance next time.
-	return AddToRuntimeCache(Initializer, InitializerHash, BoundShaderState, PipelineState);
+	return AddToRuntimeCache(Initializer, InitializerHash, RootSignature, PipelineState);
 #else
 	if (PipelineState)
 	{
-		return new FD3D12GraphicsPipelineState(Initializer, BoundShaderState, PipelineState);
+		return new FD3D12GraphicsPipelineState(Initializer, RootSignature, PipelineState);
 	}
 	else
 	{
diff --git a/Engine/Source/Runtime/D3D12RHI/Private/D3D12Shader.h b/Engine/Source/Runtime/D3D12RHI/Private/D3D12Shader.h
--- a/Engine/Source/Runtime/D3D12RHI/Private/D3D12Shader.h
+++ b/Engine/Source/Runtime/D3D12RHI/Private/D3D12Shader.h
@@ -169,7 +169,7 @@ public:
 * Combined shader state and vertex definition for rendering geometry.
 * Each unique instance consists of a vertex decl, vertex shader, and pixel shader.
 */
-class FD3D12BoundShaderState : public FRHIBoundShaderState, public FD3D12DeviceChild
+class FD3D12BoundShaderState : public FRHIBoundShaderState
 {
 public:
 
@@ -179,11 +179,6 @@ public:
 	FCachedBoundShaderStateLink CacheLink;
 #endif
 
-	D3D12_INPUT_LAYOUT_DESC InputLayout;
-	uint16 StreamStrides[MaxVertexElementCount];
-
-	bool bShaderNeedsGlobalConstantBuffer[SF_NumFrequencies];
-	uint64 UniqueID;
 	const FD3D12RootSignature* pRootSignature;
 
 	/** Initialization constructor. */
@@ -202,11 +197,12 @@ public:
 	/**
 	* Get the shader for the given frequency.
 	*/
-	FORCEINLINE FD3D12VertexShader*   GetVertexShader() const { return (FD3D12VertexShader*)CacheLink.GetVertexShader(); }
-	FORCEINLINE FD3D12PixelShader*    GetPixelShader() const { return (FD3D12PixelShader*)CacheLink.GetPixelShader(); }
-	FORCEINLINE FD3D12HullShader*     GetHullShader() const { return (FD3D12HullShader*)CacheLink.GetHullShader(); }
-	FORCEINLINE FD3D12DomainShader*   GetDomainShader() const { return (FD3D12DomainShader*)CacheLink.GetDomainShader(); }
-	FORCEINLINE FD3D12GeometryShader* GetGeometryShader() const { return (FD3D12GeometryShader*)CacheLink.GetGeometryShader(); }
+	FORCEINLINE FD3D12VertexDeclaration* GetVertexDeclaration() const { return (FD3D12VertexDeclaration*) CacheLink.GetVertexDeclaration(); }
+	FORCEINLINE FD3D12VertexShader*      GetVertexShader()      const { return (FD3D12VertexShader*)      CacheLink.GetVertexShader();      }
+	FORCEINLINE FD3D12PixelShader*       GetPixelShader()       const { return (FD3D12PixelShader*)       CacheLink.GetPixelShader();       }
+	FORCEINLINE FD3D12HullShader*        GetHullShader()        const { return (FD3D12HullShader*)        CacheLink.GetHullShader();        }
+	FORCEINLINE FD3D12DomainShader*      GetDomainShader()      const { return (FD3D12DomainShader*)      CacheLink.GetDomainShader();      }
+	FORCEINLINE FD3D12GeometryShader*    GetGeometryShader()    const { return (FD3D12GeometryShader*)    CacheLink.GetGeometryShader();    }
 };
 
 template<>
diff --git a/Engine/Source/Runtime/D3D12RHI/Private/D3D12Shaders.cpp b/Engine/Source/Runtime/D3D12RHI/Private/D3D12Shaders.cpp
--- a/Engine/Source/Runtime/D3D12RHI/Private/D3D12Shaders.cpp
+++ b/Engine/Source/Runtime/D3D12RHI/Private/D3D12Shaders.cpp
@@ -281,8 +281,6 @@ void FD3D12CommandContext::RHISetMultipleViewports(uint32 Count, const FViewport
 	StateCache.SetViewports(Count, reinterpret_cast<const D3D12_VIEWPORT*>(Data));
 }
 
-static volatile uint64 BoundShaderStateID = 0;
-
 FD3D12BoundShaderState::FD3D12BoundShaderState(
 	FVertexDeclarationRHIParamRef InVertexDeclarationRHI,
 	FVertexShaderRHIParamRef InVertexShaderRHI,
@@ -292,44 +290,11 @@ FD3D12BoundShaderState::FD3D12BoundShaderState(
 	FGeometryShaderRHIParamRef InGeometryShaderRHI,
 	FD3D12Device* InDevice
 	) :
-	FD3D12DeviceChild(InDevice),
-	CacheLink(InVertexDeclarationRHI, InVertexShaderRHI, InPixelShaderRHI, InHullShaderRHI, InDomainShaderRHI, InGeometryShaderRHI, this),
-	UniqueID(FPlatformAtomics::InterlockedIncrement(reinterpret_cast<volatile int64*>(&BoundShaderStateID)))
+	CacheLink(InVertexDeclarationRHI, InVertexShaderRHI, InPixelShaderRHI, InHullShaderRHI, InDomainShaderRHI, InGeometryShaderRHI, this)
 {
 	INC_DWORD_STAT(STAT_D3D12NumBoundShaderState);
 
-	// Warning: Input layout desc contains padding which must be zero-initialized to prevent PSO cache misses
-	FMemory::Memzero(&InputLayout, sizeof(InputLayout));
-
-	FD3D12VertexDeclaration*  InVertexDeclaration = FD3D12DynamicRHI::ResourceCast(InVertexDeclarationRHI);
-	FD3D12VertexShader*  InVertexShader = FD3D12DynamicRHI::ResourceCast(InVertexShaderRHI);
-	FD3D12PixelShader*  InPixelShader = FD3D12DynamicRHI::ResourceCast(InPixelShaderRHI);
-	FD3D12HullShader*  InHullShader = FD3D12DynamicRHI::ResourceCast(InHullShaderRHI);
-	FD3D12DomainShader*  InDomainShader = FD3D12DynamicRHI::ResourceCast(InDomainShaderRHI);
-	FD3D12GeometryShader*  InGeometryShader = FD3D12DynamicRHI::ResourceCast(InGeometryShaderRHI);
-
-	// Create an input layout for this combination of vertex declaration and vertex shader.
-	InputLayout.NumElements = (InVertexDeclaration ? InVertexDeclaration->VertexElements.Num() : 0);
-	InputLayout.pInputElementDescs = (InVertexDeclaration ? InVertexDeclaration->VertexElements.GetData() : nullptr);
-
-	bShaderNeedsGlobalConstantBuffer[SF_Vertex] = InVertexShader ? InVertexShader->ResourceCounts.bGlobalUniformBufferUsed : false;
-	bShaderNeedsGlobalConstantBuffer[SF_Hull] = InHullShader ? InHullShader->ResourceCounts.bGlobalUniformBufferUsed : false;
-	bShaderNeedsGlobalConstantBuffer[SF_Domain] = InDomainShader ? InDomainShader->ResourceCounts.bGlobalUniformBufferUsed : false;
-	bShaderNeedsGlobalConstantBuffer[SF_Pixel] = InPixelShader ? InPixelShader->ResourceCounts.bGlobalUniformBufferUsed : false;
-	bShaderNeedsGlobalConstantBuffer[SF_Geometry] = InGeometryShader ? InGeometryShader->ResourceCounts.bGlobalUniformBufferUsed : false;
-
-	static_assert(ARRAY_COUNT(bShaderNeedsGlobalConstantBuffer) == SF_NumFrequencies, "EShaderFrequency size should match with array count of bShaderNeedsGlobalConstantBuffer.");
-
-	if (InVertexDeclaration)
-	{
-		FMemory::Memcpy(StreamStrides, InVertexDeclaration->StreamStrides, sizeof(StreamStrides));
-	}
-	else
-	{
-		FMemory::Memzero(StreamStrides, sizeof(StreamStrides));
-	}
-
-	FD3D12Adapter* Adapter = GetParentDevice()->GetParentAdapter();
+	FD3D12Adapter* Adapter = InDevice->GetParentAdapter();
 
 #if USE_STATIC_ROOT_SIGNATURE
 	pRootSignature = Adapter->GetStaticGraphicsRootSignature();
diff --git a/Engine/Source/Runtime/D3D12RHI/Private/D3D12State.cpp b/Engine/Source/Runtime/D3D12RHI/Private/D3D12State.cpp
--- a/Engine/Source/Runtime/D3D12RHI/Private/D3D12State.cpp
+++ b/Engine/Source/Runtime/D3D12RHI/Private/D3D12State.cpp
@@ -464,7 +464,7 @@ FGraphicsPipelineStateRHIRef FD3D12DynamicRHI::RHICreateGraphicsPipelineState(co
 		return Found;
 	}
 #endif
-	// TODO: Remove the need for BoundShaderState objects. Currently they are needed for the Root Signature and resource counts.
+	// TODO: Remove the need for BoundShaderState objects. Currently they are needed for the Root Signature.
 	FBoundShaderStateRHIRef const BSS = RHICreateBoundShaderState(
 		Initializer.BoundShaderState.VertexDeclarationRHI,
 		Initializer.BoundShaderState.VertexShaderRHI,
@@ -475,11 +475,13 @@ FGraphicsPipelineStateRHIRef FD3D12DynamicRHI::RHICreateGraphicsPipelineState(co
 
 	// Next try to find the PSO based on the hash of its desc.
 	FD3D12BoundShaderState* const BoundShaderState = FD3D12DynamicRHI::ResourceCast(BSS.GetReference());
+	const FD3D12RootSignature* RootSignature = BoundShaderState->pRootSignature;
+
 	FD3D12LowLevelGraphicsPipelineStateDesc LowLevelDesc;
 #if D3D12RHI_USE_HIGH_LEVEL_PSO_CACHE
-	Found = PSOCache.FindInLoadedCache(Initializer, InitializerHash, BoundShaderState, LowLevelDesc);
+	Found = PSOCache.FindInLoadedCache(Initializer, InitializerHash, RootSignature, LowLevelDesc);
 #else
-	FD3D12GraphicsPipelineState* Found = PSOCache.FindInLoadedCache(Initializer, BoundShaderState, LowLevelDesc);
+	FD3D12GraphicsPipelineState* Found = PSOCache.FindInLoadedCache(Initializer, RootSignature, LowLevelDesc);
 #endif
 	if (Found)
 	{
@@ -488,9 +490,9 @@ FGraphicsPipelineStateRHIRef FD3D12DynamicRHI::RHICreateGraphicsPipelineState(co
 
 	// We need to actually create a PSO.
 #if D3D12RHI_USE_HIGH_LEVEL_PSO_CACHE
-	return PSOCache.CreateAndAdd(Initializer, InitializerHash, BoundShaderState, LowLevelDesc);
+	return PSOCache.CreateAndAdd(Initializer, InitializerHash, RootSignature, LowLevelDesc);
 #else
-	return PSOCache.CreateAndAdd(Initializer, BoundShaderState, LowLevelDesc);
+	return PSOCache.CreateAndAdd(Initializer, RootSignature, LowLevelDesc);
 #endif
 }
 
diff --git a/Engine/Source/Runtime/D3D12RHI/Private/D3D12Util.cpp b/Engine/Source/Runtime/D3D12RHI/Private/D3D12Util.cpp
--- a/Engine/Source/Runtime/D3D12RHI/Private/D3D12Util.cpp
+++ b/Engine/Source/Runtime/D3D12RHI/Private/D3D12Util.cpp
@@ -364,7 +364,7 @@ void QuantizeBoundShaderState(
 	// The objective is to allow a single root signature to represent many bound shader state objects.
 	// The bigger the quantization step sizes, the fewer the root signatures.
 	FMemory::Memzero(&QBSS, sizeof(QBSS));
-	QBSS.bAllowIAInputLayout = BSS->InputLayout.NumElements > 0;	// Does the root signature need access to vertex buffers?
+	QBSS.bAllowIAInputLayout = BSS->GetVertexDeclaration() != nullptr;	// Does the root signature need access to vertex buffers?
 
 	const FD3D12VertexShader* const VertexShader = BSS->GetVertexShader();
 	const FD3D12PixelShader* const PixelShader = BSS->GetPixelShader();
