Copy 4492773. dx12 - fix double allocation on first dynamic buffer lock

diff --git a/Engine/Source/Runtime/D3D12RHI/Private/D3D12Adapter.h b/Engine/Source/Runtime/D3D12RHI/Private/D3D12Adapter.h
--- a/Engine/Source/Runtime/D3D12RHI/Private/D3D12Adapter.h
+++ b/Engine/Source/Runtime/D3D12RHI/Private/D3D12Adapter.h
@@ -206,7 +206,6 @@ public:
 		const D3D12_RESOURCE_DESC& Desc,
 		uint32 Alignment, uint32 Stride, uint32 Size, uint32 InUsage,
 		FRHIResourceCreateInfo& CreateInfo,
-		bool SkipCreate,
 		FRHIGPUMask GPUMask = FRHIGPUMask::All());
 
 	// Queue up a command to signal the frame fence on the command list. This should only be called from the rendering thread.
diff --git a/Engine/Source/Runtime/D3D12RHI/Private/D3D12Buffer.cpp b/Engine/Source/Runtime/D3D12RHI/Private/D3D12Buffer.cpp
--- a/Engine/Source/Runtime/D3D12RHI/Private/D3D12Buffer.cpp
+++ b/Engine/Source/Runtime/D3D12RHI/Private/D3D12Buffer.cpp
@@ -19,21 +19,18 @@ template FD3D12VertexBuffer* FD3D12Adapter::CreateRHIBuffer<FD3D12VertexBuffer>(
 	const D3D12_RESOURCE_DESC& Desc,
 	uint32 Alignment, uint32 Stride, uint32 Size, uint32 InUsage,
 	FRHIResourceCreateInfo& CreateInfo,
-	bool SkipCreate,
 	FRHIGPUMask GPUMask);
 
 template FD3D12IndexBuffer* FD3D12Adapter::CreateRHIBuffer<FD3D12IndexBuffer>(FRHICommandListImmediate* RHICmdList,
 	const D3D12_RESOURCE_DESC& Desc,
 	uint32 Alignment, uint32 Stride, uint32 Size, uint32 InUsage,
 	FRHIResourceCreateInfo& CreateInfo,
-	bool SkipCreate,
 	FRHIGPUMask GPUMask);
 
 template FD3D12StructuredBuffer* FD3D12Adapter::CreateRHIBuffer<FD3D12StructuredBuffer>(FRHICommandListImmediate* RHICmdList,
 	const D3D12_RESOURCE_DESC& Desc,
 	uint32 Alignment, uint32 Stride, uint32 Size, uint32 InUsage,
 	FRHIResourceCreateInfo& CreateInfo,
-	bool SkipCreate,
 	FRHIGPUMask GPUMask);
 
 struct FRHICommandUpdateBuffer final : public FRHICommand<FRHICommandUpdateBuffer>
@@ -122,7 +119,6 @@ BufferType* FD3D12Adapter::CreateRHIBuffer(FRHICommandListImmediate* RHICmdList,
 	uint32 Size,
 	uint32 InUsage,
 	FRHIResourceCreateInfo& CreateInfo,
-	bool SkipCreate,
 	FRHIGPUMask GPUMask)
 {
 	SCOPE_CYCLE_COUNTER(STAT_D3D12CreateBufferTime);
@@ -139,18 +135,15 @@ BufferType* FD3D12Adapter::CreateRHIBuffer(FRHICommandListImmediate* RHICmdList,
 			BufferType* NewBuffer = new BufferType(Device, Stride, Size, InUsage);
 			NewBuffer->BufferAlignment = Alignment;
 
-			if (!SkipCreate)
+			if (Device->GetGPUIndex() == FirstGPUIndex)
 			{
-				if (Device->GetGPUIndex() == FirstGPUIndex)
-				{
-					AllocateBuffer(Device, InDesc, Size, InUsage, CreateInfo, Alignment, *NewBuffer, NewBuffer->ResourceLocation);
-					NewBuffer0 = NewBuffer;
-				}
-				else
-				{
-					check(NewBuffer0);
-					FD3D12ResourceLocation::ReferenceNode(Device, NewBuffer->ResourceLocation, NewBuffer0->ResourceLocation);
-				}
+				AllocateBuffer(Device, InDesc, Size, InUsage, CreateInfo, Alignment, *NewBuffer, NewBuffer->ResourceLocation);
+				NewBuffer0 = NewBuffer;
+			}
+			else
+			{
+				check(NewBuffer0);
+				FD3D12ResourceLocation::ReferenceNode(Device, NewBuffer->ResourceLocation, NewBuffer0->ResourceLocation);
 			}
 
 			return NewBuffer;
@@ -163,10 +156,7 @@ BufferType* FD3D12Adapter::CreateRHIBuffer(FRHICommandListImmediate* RHICmdList,
 			BufferType* NewBuffer = new BufferType(Device, Stride, Size, InUsage);
 			NewBuffer->BufferAlignment = Alignment;
 
-			if (!SkipCreate)
-			{
-				AllocateBuffer(Device, InDesc, Size, InUsage, CreateInfo, Alignment, *NewBuffer, NewBuffer->ResourceLocation);
-			}
+			AllocateBuffer(Device, InDesc, Size, InUsage, CreateInfo, Alignment, *NewBuffer, NewBuffer->ResourceLocation);
 
 			return NewBuffer;
 		});
@@ -266,22 +256,32 @@ void* FD3D12DynamicRHI::LockBuffer(FRHICommandListImmediate* RHICmdList, BufferT
 
 	if (bIsDynamic)
 	{
-		FD3D12Device* Device = Buffer->GetParentDevice();
 		check(LockMode == RLM_WriteOnly);
 
-		// If on the RenderThread, queue up a command on the RHIThread to rename this buffer at the correct time
-		if (ShouldDeferBufferLockOperation(RHICmdList))
+		if (LockedData.bHasNeverBeenLocked)
 		{
-			FRHICommandRenameUploadBuffer<BufferType>* Command = new (RHICmdList->AllocCommand<FRHICommandRenameUploadBuffer<BufferType>>()) FRHICommandRenameUploadBuffer<BufferType>(Buffer, Device);
-
-			Data = Adapter.GetUploadHeapAllocator(Device->GetGPUIndex()).AllocUploadResource(Buffer->GetSize(), Buffer->BufferAlignment, Command->NewLocation);
-			RHICmdList->RHIThreadFence(true);
+			// Buffers on upload heap are mapped right after creation
+			Data = Buffer->ResourceLocation.GetMappedBaseAddress();
+			check(!!Data);
 		}
 		else
 		{
-			FD3D12ResourceLocation Location(Buffer->GetParentDevice());
-			Data = Adapter.GetUploadHeapAllocator(Device->GetGPUIndex()).AllocUploadResource(Buffer->GetSize(), Buffer->BufferAlignment, Location);
-			Buffer->RenameLDAChain(Location);
+			FD3D12Device* Device = Buffer->GetParentDevice();
+
+			// If on the RenderThread, queue up a command on the RHIThread to rename this buffer at the correct time
+			if (ShouldDeferBufferLockOperation(RHICmdList))
+			{
+				FRHICommandRenameUploadBuffer<BufferType>* Command = new (RHICmdList->AllocCommand<FRHICommandRenameUploadBuffer<BufferType>>()) FRHICommandRenameUploadBuffer<BufferType>(Buffer, Device);
+
+				Data = Adapter.GetUploadHeapAllocator(Device->GetGPUIndex()).AllocUploadResource(Buffer->GetSize(), Buffer->BufferAlignment, Command->NewLocation);
+				RHICmdList->RHIThreadFence(true);
+			}
+			else
+			{
+				FD3D12ResourceLocation Location(Buffer->GetParentDevice());
+				Data = Adapter.GetUploadHeapAllocator(Device->GetGPUIndex()).AllocUploadResource(Buffer->GetSize(), Buffer->BufferAlignment, Location);
+				Buffer->RenameLDAChain(Location);
+			}
 		}
 	}
 	else
@@ -354,6 +354,7 @@ void* FD3D12DynamicRHI::LockBuffer(FRHICommandListImmediate* RHICmdList, BufferT
 	LockedData.LockedOffset = Offset;
 	LockedData.LockedPitch = Size;
 	LockedData.bLocked = true;
+	LockedData.bHasNeverBeenLocked = false;
 
 	// Return the offset pointer
 	check(Data != nullptr);
diff --git a/Engine/Source/Runtime/D3D12RHI/Private/D3D12IndexBuffer.cpp b/Engine/Source/Runtime/D3D12RHI/Private/D3D12IndexBuffer.cpp
--- a/Engine/Source/Runtime/D3D12RHI/Private/D3D12IndexBuffer.cpp
+++ b/Engine/Source/Runtime/D3D12RHI/Private/D3D12IndexBuffer.cpp
@@ -64,7 +64,7 @@ FIndexBufferRHIRef FD3D12DynamicRHI::RHICreateIndexBuffer(uint32 Stride, uint32
 	const D3D12_RESOURCE_DESC Desc = CreateIndexBufferResourceDesc(Size, InUsage);
 	const uint32 Alignment = 4;
 
-	FD3D12IndexBuffer* Buffer = GetAdapter().CreateRHIBuffer<FD3D12IndexBuffer>(nullptr, Desc, Alignment, Stride, Size, InUsage, CreateInfo, false);
+	FD3D12IndexBuffer* Buffer = GetAdapter().CreateRHIBuffer<FD3D12IndexBuffer>(nullptr, Desc, Alignment, Stride, Size, InUsage, CreateInfo);
 	if (Buffer->ResourceLocation.IsTransient())
 	{
 		// TODO: this should ideally be set in platform-independent code, since this tracking is for the high level
@@ -91,7 +91,7 @@ FIndexBufferRHIRef FD3D12DynamicRHI::CreateIndexBuffer_RenderThread(class FRHICo
 	const D3D12_RESOURCE_DESC Desc = CreateIndexBufferResourceDesc(Size, InUsage);
 	const uint32 Alignment = 4;
 
-	FD3D12IndexBuffer* Buffer = GetAdapter().CreateRHIBuffer<FD3D12IndexBuffer>(&RHICmdList, Desc, Alignment, Stride, Size, InUsage, CreateInfo, false);
+	FD3D12IndexBuffer* Buffer = GetAdapter().CreateRHIBuffer<FD3D12IndexBuffer>(&RHICmdList, Desc, Alignment, Stride, Size, InUsage, CreateInfo);
 	if (Buffer->ResourceLocation.IsTransient())
 	{
 		// TODO: this should ideally be set in platform-independent code, since this tracking is for the high level
@@ -122,8 +122,7 @@ FIndexBufferRHIRef FD3D12DynamicRHI::CreateAndLockIndexBuffer_RenderThread(class
 	const D3D12_RESOURCE_DESC Desc = CreateIndexBufferResourceDesc(Size, InUsage);
 	const uint32 Alignment = 4;
 
-	const bool bIsDynamic = (InUsage & BUF_AnyDynamic) ? true : false;
-	FD3D12IndexBuffer* Buffer = GetAdapter().CreateRHIBuffer<FD3D12IndexBuffer>(&RHICmdList, Desc, Alignment, Stride, Size, InUsage, CreateInfo, bIsDynamic);
+	FD3D12IndexBuffer* Buffer = GetAdapter().CreateRHIBuffer<FD3D12IndexBuffer>(&RHICmdList, Desc, Alignment, Stride, Size, InUsage, CreateInfo);
 	if (Buffer->ResourceLocation.IsTransient())
 	{
 		// TODO: this should ideally be set in platform-independent code, since this tracking is for the high level
diff --git a/Engine/Source/Runtime/D3D12RHI/Private/D3D12StructuredBuffer.cpp b/Engine/Source/Runtime/D3D12RHI/Private/D3D12StructuredBuffer.cpp
--- a/Engine/Source/Runtime/D3D12RHI/Private/D3D12StructuredBuffer.cpp
+++ b/Engine/Source/Runtime/D3D12RHI/Private/D3D12StructuredBuffer.cpp
@@ -37,7 +37,7 @@ FStructuredBufferRHIRef FD3D12DynamicRHI::CreateStructuredBuffer_RenderThread(FR
 	// can be addressed correctly with element based offsets.
 	const uint32 Alignment = ((InUsage & (BUF_ByteAddressBuffer | BUF_DrawIndirect)) == 0) ? Stride : 4;
 
-	FD3D12StructuredBuffer* NewBuffer = GetAdapter().CreateRHIBuffer<FD3D12StructuredBuffer>(&RHICmdList, Desc, Alignment, Stride, Size, InUsage, CreateInfo, false);
+	FD3D12StructuredBuffer* NewBuffer = GetAdapter().CreateRHIBuffer<FD3D12StructuredBuffer>(&RHICmdList, Desc, Alignment, Stride, Size, InUsage, CreateInfo);
 	if (NewBuffer->ResourceLocation.IsTransient())
 	{
 		// TODO: this should ideally be set in platform-independent code, since this tracking is for the high level
@@ -60,7 +60,7 @@ FStructuredBufferRHIRef FD3D12DynamicRHI::RHICreateStructuredBuffer(uint32 Strid
 	// can be addressed correctly with element based offsets.
 	const uint32 Alignment = ((InUsage & (BUF_ByteAddressBuffer | BUF_DrawIndirect)) == 0) ? Stride : 4;
 
-	FD3D12StructuredBuffer* NewBuffer = GetAdapter().CreateRHIBuffer<FD3D12StructuredBuffer>(nullptr, Desc, Alignment, Stride, Size, InUsage, CreateInfo, false);
+	FD3D12StructuredBuffer* NewBuffer = GetAdapter().CreateRHIBuffer<FD3D12StructuredBuffer>(nullptr, Desc, Alignment, Stride, Size, InUsage, CreateInfo);
 	if (NewBuffer->ResourceLocation.IsTransient())
 	{
 		// TODO: this should ideally be set in platform-independent code, since this tracking is for the high level
diff --git a/Engine/Source/Runtime/D3D12RHI/Private/D3D12VertexBuffer.cpp b/Engine/Source/Runtime/D3D12RHI/Private/D3D12VertexBuffer.cpp
--- a/Engine/Source/Runtime/D3D12RHI/Private/D3D12VertexBuffer.cpp
+++ b/Engine/Source/Runtime/D3D12RHI/Private/D3D12VertexBuffer.cpp
@@ -84,7 +84,7 @@ FVertexBufferRHIRef FD3D12DynamicRHI::RHICreateVertexBuffer(uint32 Size, uint32
 	const D3D12_RESOURCE_DESC Desc = CreateVertexBufferResourceDesc(Size, InUsage);
 	const uint32 Alignment = 4;
 
-	FD3D12VertexBuffer* Buffer = GetAdapter().CreateRHIBuffer<FD3D12VertexBuffer>(nullptr, Desc, Alignment, 0, Size, InUsage, CreateInfo, false);
+	FD3D12VertexBuffer* Buffer = GetAdapter().CreateRHIBuffer<FD3D12VertexBuffer>(nullptr, Desc, Alignment, 0, Size, InUsage, CreateInfo);
 	if (Buffer->ResourceLocation.IsTransient() )
 	{
 		// TODO: this should ideally be set in platform-independent code, since this tracking is for the high level
@@ -111,7 +111,7 @@ FVertexBufferRHIRef FD3D12DynamicRHI::CreateVertexBuffer_RenderThread(FRHIComman
 	const D3D12_RESOURCE_DESC Desc = CreateVertexBufferResourceDesc(Size, InUsage);
 	const uint32 Alignment = 4;
 
-	FD3D12VertexBuffer* Buffer = GetAdapter().CreateRHIBuffer<FD3D12VertexBuffer>(&RHICmdList, Desc, Alignment, 0, Size, InUsage, CreateInfo, false);
+	FD3D12VertexBuffer* Buffer = GetAdapter().CreateRHIBuffer<FD3D12VertexBuffer>(&RHICmdList, Desc, Alignment, 0, Size, InUsage, CreateInfo);
 	if (Buffer->ResourceLocation.IsTransient())
 	{
 		// TODO: this should ideally be set in platform-independent code, since this tracking is for the high level
@@ -179,7 +179,7 @@ FVertexBufferRHIRef FD3D12DynamicRHI::CreateAndLockVertexBuffer_RenderThread(FRH
 	const D3D12_RESOURCE_DESC Desc = CreateVertexBufferResourceDesc(Size, InUsage);
 	const uint32 Alignment = 4;
 
-	FD3D12VertexBuffer* Buffer = GetAdapter().CreateRHIBuffer<FD3D12VertexBuffer>(nullptr, Desc, Alignment, 0, Size, InUsage, CreateInfo, false);
+	FD3D12VertexBuffer* Buffer = GetAdapter().CreateRHIBuffer<FD3D12VertexBuffer>(nullptr, Desc, Alignment, 0, Size, InUsage, CreateInfo);
 	if (Buffer->ResourceLocation.IsTransient())
 	{
 		// TODO: this should ideally be set in platform-independent code, since this tracking is for the high level
diff --git a/Engine/Source/Runtime/D3D12RHI/Public/D3D12Resources.h b/Engine/Source/Runtime/D3D12RHI/Public/D3D12Resources.h
--- a/Engine/Source/Runtime/D3D12RHI/Public/D3D12Resources.h
+++ b/Engine/Source/Runtime/D3D12RHI/Public/D3D12Resources.h
@@ -565,6 +565,7 @@ struct FD3D12LockedResource : public FD3D12DeviceChild
 		, LockedPitch(0)
 		, bLocked(false)
 		, bLockedForReadOnly(false)
+		, bHasNeverBeenLocked(true)
 	{}
 
 	inline void Reset()
@@ -581,6 +582,7 @@ struct FD3D12LockedResource : public FD3D12DeviceChild
 	uint32 LockedPitch;
 	uint32 bLocked : 1;
 	uint32 bLockedForReadOnly : 1;
+	uint32 bHasNeverBeenLocked : 1;
 };
 
 /** The base class of resources that may be bound as shader resources. */
