Fix a couple of new DX12 startup errors

diff --git a/Engine/Source/Runtime/D3D12RHI/Private/D3D12Texture.cpp b/Engine/Source/Runtime/D3D12RHI/Private/D3D12Texture.cpp
--- a/Engine/Source/Runtime/D3D12RHI/Private/D3D12Texture.cpp
+++ b/Engine/Source/Runtime/D3D12RHI/Private/D3D12Texture.cpp
@@ -1036,7 +1036,7 @@ FD3D12Texture3D* FD3D12DynamicRHI::CreateD3D12Texture3D(FRHICommandListImmediate
 	D3D12_CLEAR_VALUE ClearValue;
 	if (bCreateRTV && CreateInfo.ClearValueBinding.ColorBinding == EClearBinding::EColorBound)
 	{
-		ClearValue = CD3DX12_CLEAR_VALUE(PlatformResourceFormat, CreateInfo.ClearValueBinding.Value.Color);
+		ClearValue = CD3DX12_CLEAR_VALUE(PlatformRenderTargetFormat, CreateInfo.ClearValueBinding.Value.Color);
 		ClearValuePtr = &ClearValue;
 	}
 
@@ -2089,10 +2089,9 @@ public:
 			CD3DX12_TEXTURE_COPY_LOCATION DestCopyLocation(TextureLink->GetResource()->GetResource(), MipIdx);
 			CD3DX12_TEXTURE_COPY_LOCATION SourceCopyLocation(UploadBuffer->GetResource(), PlacedSubresourceFootprint);
 
-			FScopeResourceBarrier ScopeResourceBarrierDest(
+			FConditionalScopeResourceBarrier ScopeResourceBarrierDest(
 				NativeCmdList,
 				TextureLink->GetResource(),
-				TextureLink->GetResource()->GetDefaultResourceState(),
 				D3D12_RESOURCE_STATE_COPY_DEST,
 				DestCopyLocation.SubresourceIndex);
 
