Initialize ResidencyTracking for FD3D12SegHeap. This fixes a spurious E_OUTOFMEMORY error because of an empty residency set and attempting to allocate 0 bytes and getting nullptr back and mistaking that for an out-of-memory condition.

diff --git a/Engine/Source/Runtime/D3D12RHI/Private/D3D12Allocation.cpp b/Engine/Source/Runtime/D3D12RHI/Private/D3D12Allocation.cpp
--- a/Engine/Source/Runtime/D3D12RHI/Private/D3D12Allocation.cpp
+++ b/Engine/Source/Runtime/D3D12RHI/Private/D3D12Allocation.cpp
@@ -1577,7 +1577,7 @@ FD3D12SegHeap* FD3D12SegList::CreateBackingHeap(
 
 	VERIFYD3D12RESULT(Parent->GetDevice()->CreateHeap(&Desc, IID_PPV_ARGS(&D3DHeap)));
 
-	FD3D12SegHeap* Ret = new FD3D12SegHeap(Parent, VisibleNodeMask, D3DHeap, this, FreeHeaps.Num());
+	FD3D12SegHeap* Ret = new FD3D12SegHeap(Parent, VisibleNodeMask, D3DHeap, HeapSize, this, FreeHeaps.Num());
 	FreeHeaps.Add(Ret);
 	return Ret;
 }
diff --git a/Engine/Source/Runtime/D3D12RHI/Private/D3D12Allocation.h b/Engine/Source/Runtime/D3D12RHI/Private/D3D12Allocation.h
--- a/Engine/Source/Runtime/D3D12RHI/Private/D3D12Allocation.h
+++ b/Engine/Source/Runtime/D3D12RHI/Private/D3D12Allocation.h
@@ -644,6 +644,7 @@ private:
 		FD3D12Device* Parent,
 		FRHIGPUMask VisibileNodeMask,
 		ID3D12Heap* NewHeap,
+		uint64 HeapSize,
 		FD3D12SegList* Owner,
 		uint32 Idx) :
 		FD3D12Heap(Parent, VisibileNodeMask),
@@ -652,6 +653,7 @@ private:
 		FirstFreeOffset(0)
 	{
 		this->SetHeap(NewHeap);
+		BeginTrackingResidency(HeapSize);
 	}
 
 	virtual ~FD3D12SegHeap() = default;
