Copy 4545845. PC dx12 - segregated list read-only texture allocator

diff --git a/Engine/Source/Runtime/D3D12RHI/Private/D3D12Allocation.cpp b/Engine/Source/Runtime/D3D12RHI/Private/D3D12Allocation.cpp
--- a/Engine/Source/Runtime/D3D12RHI/Private/D3D12Allocation.cpp
+++ b/Engine/Source/Runtime/D3D12RHI/Private/D3D12Allocation.cpp
@@ -9,6 +9,30 @@
 #include "D3D12Allocation.h"
 #include "Misc/BufferedOutputDevice.h"
 
+#if D3D12RHI_SEGREGATED_TEXTURE_ALLOC
+static int32 GD3D12ReadOnlyTextureAllocatorMinPoolSize = 4 * 1024 * 1024;
+static FAutoConsoleVariableRef CVarD3D12ReadOnlyTextureAllocatorMinPoolSize(
+	TEXT("d3d12.ReadOnlyTextureAllocator.MinPoolSize"),
+	GD3D12ReadOnlyTextureAllocatorMinPoolSize,
+	TEXT("Minimum allocation granularity (in bytes) of each size list"),
+	ECVF_ReadOnly);
+
+static int32 GD3D12ReadOnlyTextureAllocatorMinNumToPool = 8;
+static FAutoConsoleVariableRef CVarD3D12ReadOnlyTextureAllocatorMinNumToPool(
+	TEXT("d3d12.ReadOnlyTextureAllocator.MinNumToPool"),
+	GD3D12ReadOnlyTextureAllocatorMinNumToPool,
+	TEXT("Texture pool of each size list must be large enough to store this")
+	TEXT("many textures unless constrained by maximum allocation granularity"),
+	ECVF_ReadOnly);
+
+static int32 GD3D12ReadOnlyTextureAllocatorMaxPoolSize = 20 * 1024 * 1024;
+static FAutoConsoleVariableRef CVarD3D12ReadOnlyTextureAllocatorMaxPoolSize(
+	TEXT("d3d12.ReadOnlyTextureAllocator.MaxPoolSize"),
+	GD3D12ReadOnlyTextureAllocatorMaxPoolSize,
+	TEXT("Maximum allocation granularity (in bytes) of each size list"),
+	ECVF_ReadOnly);
+#endif
+
 namespace ED3D12AllocatorID
 {
 	enum Type
@@ -1048,6 +1072,74 @@ void FD3D12DefaultBufferAllocator::CleanupFreeBlocks()
 //	Texture Allocator
 //-----------------------------------------------------------------------------
 
+#if D3D12RHI_SEGREGATED_TEXTURE_ALLOC
+FD3D12TextureAllocatorPool::FD3D12TextureAllocatorPool(FD3D12Device* Device, FRHIGPUMask VisibilityNode) :
+	FD3D12DeviceChild(Device),
+	FD3D12MultiNodeGPUObject(Device->GetGPUMask(), VisibilityNode),
+	ReadOnlyTexturePool(
+		Device,
+		VisibilityNode,
+		D3D12_HEAP_TYPE_DEFAULT,
+		D3D12_HEAP_FLAG_ALLOW_ONLY_NON_RT_DS_TEXTURES,
+		GD3D12ReadOnlyTextureAllocatorMinPoolSize,
+		GD3D12ReadOnlyTextureAllocatorMinNumToPool,
+		GD3D12ReadOnlyTextureAllocatorMaxPoolSize)
+{
+}
+
+HRESULT FD3D12TextureAllocatorPool::AllocateTexture(
+	D3D12_RESOURCE_DESC Desc,
+	const D3D12_CLEAR_VALUE* ClearValue,
+	uint8 UEFormat,
+	FD3D12ResourceLocation& TextureLocation,
+	const D3D12_RESOURCE_STATES InitialState)
+{
+	HRESULT RetCode = S_OK;
+	FD3D12Resource* NewResource = nullptr;
+	FD3D12Device* Device = GetParentDevice();
+	FD3D12Adapter* Adapter = Device->GetParentAdapter();
+
+	TextureLocation.Clear();
+
+	if (!(Desc.Flags & D3D12_RESOURCE_FLAG_ALLOW_RENDER_TARGET ||
+		Desc.Flags & D3D12_RESOURCE_FLAG_ALLOW_DEPTH_STENCIL ||
+		Desc.Flags & D3D12_RESOURCE_FLAG_ALLOW_UNORDERED_ACCESS) &&
+		Desc.SampleDesc.Count == 1)
+	{
+		// The top mip level must be less than 64 KB to use 4 KB alignment
+		Desc.Alignment = TextureCanBe4KAligned(Desc, UEFormat) ?
+			D3D12_SMALL_RESOURCE_PLACEMENT_ALIGNMENT :
+			D3D12_DEFAULT_RESOURCE_PLACEMENT_ALIGNMENT;
+		const D3D12_RESOURCE_ALLOCATION_INFO Info = Device->GetDevice()->GetResourceAllocationInfo(0, 1, &Desc);
+
+		TRefCountPtr<FD3D12SegHeap> BackingHeap;
+		const uint32 Offset = ReadOnlyTexturePool.Allocate(Info.SizeInBytes, Info.Alignment, BackingHeap);
+
+		if (Offset != FD3D12SegListAllocator::InvalidOffset)
+		{
+			RetCode = Adapter->CreatePlacedResource(Desc, BackingHeap.GetReference(), Offset, InitialState, ClearValue, &NewResource);
+
+			FD3D12SegListAllocatorPrivateData& PrivateData = TextureLocation.GetSegListAllocatorPrivateData();
+			PrivateData.Offset = Offset;
+
+			TextureLocation.SetType(FD3D12ResourceLocation::ResourceLocationType::eSubAllocation);
+			TextureLocation.SetSegListAllocator(&ReadOnlyTexturePool);
+			TextureLocation.SetSize(Info.SizeInBytes);
+			TextureLocation.SetOffsetFromBaseOfResource(Offset);
+			TextureLocation.SetResource(NewResource);
+			return RetCode;
+		}
+	}
+
+	const D3D12_HEAP_PROPERTIES HeapProps = CD3DX12_HEAP_PROPERTIES(D3D12_HEAP_TYPE_DEFAULT, (uint32)GetGPUMask(), (uint32)GetVisibilityMask());
+	Desc.Alignment = 0;
+	RetCode = Adapter->CreateCommittedResource(Desc, HeapProps, InitialState, ClearValue, &NewResource);
+
+	TextureLocation.SetType(FD3D12ResourceLocation::ResourceLocationType::eStandAlone);
+	TextureLocation.SetResource(NewResource);
+	return RetCode;
+}
+#else
 FD3D12TextureAllocator::FD3D12TextureAllocator(FD3D12Device* Device,
 	FRHIGPUMask VisibleNodes,
 	const FString& Name,
@@ -1085,7 +1177,7 @@ HRESULT FD3D12TextureAllocator::AllocateTexture(D3D12_RESOURCE_DESC Desc, const
 	TextureLocation.Clear();
 
 	D3D12_RESOURCE_ALLOCATION_INFO Info = Device->GetDevice()->GetResourceAllocationInfo(0, 1, &Desc);
-	
+
 	if (Info.SizeInBytes < D3D12_DEFAULT_RESOURCE_PLACEMENT_ALIGNMENT)
 	{
 		if (TryAllocate(Info.SizeInBytes, Info.Alignment, TextureLocation))
@@ -1146,6 +1238,7 @@ HRESULT FD3D12TextureAllocatorPool::AllocateTexture(D3D12_RESOURCE_DESC Desc, co
 
 	return hr;
 }
+#endif
 
 //-----------------------------------------------------------------------------
 //	Fast Allocation
@@ -1468,3 +1561,147 @@ void* FD3D12FastConstantAllocator::Allocate(uint32 Bytes, FD3D12ResourceLocation
 	return OutLocation.GetMappedBaseAddress();
 }
 
+FD3D12SegHeap* FD3D12SegList::CreateBackingHeap(
+	FD3D12Device* Parent,
+	FRHIGPUMask VisibleNodeMask,
+	D3D12_HEAP_TYPE HeapType,
+	D3D12_HEAP_FLAGS HeapFlags)
+{
+	// CS can be unlocked at this point and re-locked before adding it to FreeHeaps
+	// but doing so may cause multiple heaps to be created
+	ID3D12Heap* D3DHeap;
+	D3D12_HEAP_DESC Desc = {};
+	Desc.SizeInBytes = HeapSize;
+	Desc.Properties = CD3DX12_HEAP_PROPERTIES(HeapType, Parent->GetGPUMask(), VisibleNodeMask);
+	Desc.Flags = HeapFlags;
+
+	VERIFYD3D12RESULT(Parent->GetDevice()->CreateHeap(&Desc, IID_PPV_ARGS(&D3DHeap)));
+
+	FD3D12SegHeap* Ret = new FD3D12SegHeap(Parent, VisibleNodeMask, D3DHeap, this, FreeHeaps.Num());
+	FreeHeaps.Add(Ret);
+	return Ret;
+}
+
+FD3D12SegListAllocator::FD3D12SegListAllocator(
+	FD3D12Device* Parent,
+	FRHIGPUMask VisibilityMask,
+	D3D12_HEAP_TYPE InHeapType,
+	D3D12_HEAP_FLAGS InHeapFlags,
+	uint32 InMinPoolSize,
+	uint32 InMinNumToPool,
+	uint32 InMaxPoolSize)
+	: FD3D12DeviceChild(Parent)
+	, FD3D12MultiNodeGPUObject(Parent->GetGPUMask(), VisibilityMask)
+	, HeapType(InHeapType)
+	, HeapFlags(InHeapFlags)
+	, MinPoolSize(InMinPoolSize)
+	, MinNumToPool(InMinNumToPool)
+	, MaxPoolSize(InMaxPoolSize)
+#if D3D12RHI_SEGLIST_ALLOC_TRACK_WASTAGE
+	, TotalBytesRequested(0)
+#endif
+{}
+
+void FD3D12SegListAllocator::Deallocate(
+	FD3D12Resource* PlacedResource,
+	uint32 Offset,
+	uint32 SizeInBytes)
+{
+	FD3D12Device* Device = this->GetParentDevice();
+	FD3D12Adapter* Adapter = Device->GetParentAdapter();
+	uint64 CurFenceValue = Adapter->GetFrameFence().GetCurrentFence();
+	{
+		FScopeLock Lock(&DeferredDeletionCS);
+
+		check(FenceValues.Num() == DeferredDeletionQueue.Num());
+		check(!FenceValues.Num() || FenceValues.Last() <= CurFenceValue);
+
+		int32 LastIdx = FenceValues.Num() - 1;
+		if (LastIdx < 0 || FenceValues[LastIdx] != CurFenceValue)
+		{
+			++LastIdx;
+			FenceValues.Add(CurFenceValue);
+			DeferredDeletionQueue.AddDefaulted();
+		}
+		new (DeferredDeletionQueue[LastIdx]) FRetiredBlock(PlacedResource, Offset, SizeInBytes);
+	}
+}
+
+template <typename AllocX, typename AllocY>
+void FD3D12SegListAllocator::FreeRetiredBlocks(TArray<TArray<FRetiredBlock, AllocX>, AllocY>& PendingDeletes)
+{
+	for (int32 Y = 0; Y < PendingDeletes.Num(); ++Y)
+	{
+		TArray<FRetiredBlock>& RetiredBlocks = PendingDeletes[Y];
+		for (int32 X = 0; X < RetiredBlocks.Num(); ++X)
+		{
+			FRetiredBlock& Block = RetiredBlocks[X];
+			FD3D12SegHeap* BackingHeap = static_cast<FD3D12SegHeap*>(Block.PlacedResource->GetHeap());
+			check(Block.PlacedResource->GetRefCount() == 1);
+			Block.PlacedResource->Release();
+			FD3D12SegList* Owner = BackingHeap->OwnerList;
+			check(!!Owner);
+			Owner->FreeBlock(BackingHeap, Block.Offset);
+#if D3D12RHI_SEGLIST_ALLOC_TRACK_WASTAGE
+			TotalBytesRequested -= Block.ResourceSize;
+#endif
+		}
+	}
+}
+
+void FD3D12SegListAllocator::CleanUpAllocations()
+{
+	TArray<TArray<FRetiredBlock>, TInlineAllocator<1>> PendingDeletes;
+	{
+		int32 NumToRemove = 0;
+		FD3D12Device* Device = this->GetParentDevice();
+		FD3D12Adapter* Adapter = Device->GetParentAdapter();
+		FD3D12Fence& FrameFence = Adapter->GetFrameFence();
+
+		FScopeLock Lock(&DeferredDeletionCS);
+
+		for (int32 Idx = 0; Idx < DeferredDeletionQueue.Num(); ++Idx)
+		{
+			if (FrameFence.IsFenceComplete(FenceValues[Idx]))
+			{
+				++NumToRemove;
+				PendingDeletes.Add(MoveTemp(DeferredDeletionQueue[Idx]));
+			}
+			else
+			{
+				break;
+			}
+		}
+		if (!!NumToRemove)
+		{
+			FenceValues.RemoveAt(0, NumToRemove);
+			DeferredDeletionQueue.RemoveAt(0, NumToRemove);
+		}
+	}
+	FreeRetiredBlocks(PendingDeletes);
+}
+
+void FD3D12SegListAllocator::Destroy()
+{
+	{
+		FScopeLock Lock(&DeferredDeletionCS);
+		check(FenceValues.Num() == DeferredDeletionQueue.Num());
+		FreeRetiredBlocks(DeferredDeletionQueue);
+		FenceValues.Empty();
+		DeferredDeletionQueue.Empty();
+#if D3D12RHI_SEGLIST_ALLOC_TRACK_WASTAGE
+		check(!TotalBytesRequested);
+#endif
+	}
+	{
+		FRWScopeLock Lock(SegListsRWLock, SLT_Write);
+		for (auto& Pair : SegLists)
+		{
+			FD3D12SegList*& SegList = Pair.Value;
+			check(!!SegList);
+			delete SegList;
+			SegList = nullptr;
+		}
+		SegLists.Empty();
+	}
+}
diff --git a/Engine/Source/Runtime/D3D12RHI/Private/D3D12Allocation.h b/Engine/Source/Runtime/D3D12RHI/Private/D3D12Allocation.h
--- a/Engine/Source/Runtime/D3D12RHI/Private/D3D12Allocation.h
+++ b/Engine/Source/Runtime/D3D12RHI/Private/D3D12Allocation.h
@@ -7,6 +7,30 @@ D3D12Allocation.h: A Collection of allocators
 #pragma once
 #include "D3D12Resources.h"
 
+// Segregated free list texture allocator
+// Description:
+// - Binned read-only texture allocation based on sizes
+// Suggestions:
+// - You can check memory wastage using "stat d3d12rhi" in a dev build
+// - Tune d3d12.ReadOnlyTextureAllocator.MinPoolSize/MinNumToPool/MaxPoolSize
+//   according to video memory budget
+// - Memory overhead is slightly over 200 MB in internal tests but consider
+//   adjusting above cvars or disabling if it fails your use case
+// - The purpose of this allocator is pooling texture allocations because
+//   creating committed resources is slow on PC. But if committed resource
+//   creation ever becomes fast, there is no need for this allocator
+// Internal test statistics (11/6/2018):
+// - Average read-only texture alloc time reduced from ~420 us to ~72 us
+// - Number of allocations over 1 ms reduced from 8145 to 504 (from 14.76% to
+//   0.92%) over a 17 minutes 11 seconds game replay
+// - Peak memory overhead was ~207 MB (from 2666.58 MB to 2874.08 MB)
+// TODO:
+// - Defragmentation support
+#define D3D12RHI_SEGREGATED_TEXTURE_ALLOC (PLATFORM_WINDOWS)
+#define D3D12RHI_SEGLIST_ALLOC_TRACK_WASTAGE (!(UE_BUILD_TEST || UE_BUILD_SHIPPING))
+
+class FD3D12SegList;
+
 class FD3D12ResourceAllocator : public FD3D12DeviceChild, public FD3D12MultiNodeGPUObject
 {
 public:
@@ -381,35 +405,6 @@ private:
 	}
 };
 
-//-----------------------------------------------------------------------------
-//	FD3D12TextureAllocator
-//-----------------------------------------------------------------------------
-
-class FD3D12TextureAllocator : public FD3D12MultiBuddyAllocator
-{
-public:
-	FD3D12TextureAllocator(FD3D12Device* Device, FRHIGPUMask VisibleNodes, const FString& Name, uint32 HeapSize, D3D12_HEAP_FLAGS Flags);
-
-	~FD3D12TextureAllocator();
-
-	HRESULT AllocateTexture(D3D12_RESOURCE_DESC Desc, const D3D12_CLEAR_VALUE* ClearValue, FD3D12ResourceLocation& TextureLocation, const D3D12_RESOURCE_STATES InitialState );
-};
-
-class FD3D12TextureAllocatorPool : public FD3D12DeviceChild, public FD3D12MultiNodeGPUObject
-{
-public:
-	FD3D12TextureAllocatorPool(FD3D12Device* Device, FRHIGPUMask VisibilityNode);
-
-	HRESULT AllocateTexture(D3D12_RESOURCE_DESC Desc, const D3D12_CLEAR_VALUE* ClearValue, uint8 UEFormat, FD3D12ResourceLocation& TextureLocation, const D3D12_RESOURCE_STATES InitialState );
-
-	void CleanUpAllocations() { ReadOnlyTexturePool.CleanUpAllocations(); }
-
-	void Destroy() { ReadOnlyTexturePool.Destroy(); }
-
-private:
-	FD3D12TextureAllocator ReadOnlyTexturePool;
-};
-
 //-----------------------------------------------------------------------------
 //	Fast Allocation
 //-----------------------------------------------------------------------------
@@ -636,4 +631,414 @@ private:
 
 	uint32 PageSize;
 	FD3D12AbstractRingBuffer RingBuffer;
-};
\ No newline at end of file
+};
+
+//-----------------------------------------------------------------------------
+//	FD3D12SegListAllocator
+//-----------------------------------------------------------------------------
+
+class FD3D12SegHeap : public FD3D12Heap
+{
+private:
+	FD3D12SegHeap(
+		FD3D12Device* Parent,
+		FRHIGPUMask VisibileNodeMask,
+		ID3D12Heap* NewHeap,
+		FD3D12SegList* Owner,
+		uint32 Idx) :
+		FD3D12Heap(Parent, VisibileNodeMask),
+		OwnerList(Owner),
+		ArrayIdx(Idx),
+		FirstFreeOffset(0)
+	{
+		this->SetHeap(NewHeap);
+	}
+
+	virtual ~FD3D12SegHeap() = default;
+
+	FD3D12SegHeap(const FD3D12SegHeap&) = delete;
+	FD3D12SegHeap(FD3D12SegHeap&&) = delete;
+
+	FD3D12SegHeap& operator=(const FD3D12SegHeap&) = delete;
+	FD3D12SegHeap& operator=(FD3D12SegHeap&&) = delete;
+
+	bool IsArrayIdxValid() const { return ArrayIdx >= 0; }
+
+	bool IsFull(uint32 HeapSize) const
+	{
+		check(FirstFreeOffset <= HeapSize);
+		return !FreeBlockOffsets.Num() && FirstFreeOffset == HeapSize;
+	}
+
+	bool IsEmpty(uint32 BlockSize) const
+	{
+		return FreeBlockOffsets.Num() * BlockSize == FirstFreeOffset;
+	}
+
+	// @return - In-heap offset of the allocated block
+	uint32 AllocateBlock(uint32 BlockSize)
+	{
+		if (!FreeBlockOffsets.Num())
+		{
+			uint32 Ret = FirstFreeOffset;
+			FirstFreeOffset += BlockSize;
+			return Ret;
+		}
+		else
+		{
+			return FreeBlockOffsets.Pop();
+		}
+	}
+
+	TArray<uint32> FreeBlockOffsets;
+	FD3D12SegList* OwnerList;
+	int32 ArrayIdx;
+	uint32 FirstFreeOffset;
+
+	friend class FD3D12SegList;
+	friend class FD3D12SegListAllocator;
+};
+
+class FD3D12SegList
+{
+private:
+	FD3D12SegList(uint32 InBlockSize, uint32 InHeapSize)
+		: BlockSize(InBlockSize)
+		, HeapSize(InHeapSize)
+#if D3D12RHI_SEGLIST_ALLOC_TRACK_WASTAGE
+		, TotalBytesAllocated(0)
+#endif
+	{
+		check(!(HeapSize % BlockSize));
+		check(HeapSize / BlockSize > 1);
+	}
+
+	~FD3D12SegList()
+	{
+		FScopeLock Lock(&CS);
+		check(!!BlockSize);
+		check(!!HeapSize);
+		for (const auto& Heap : FreeHeaps)
+		{
+			check(Heap->GetRefCount() == 1);
+		}
+	}
+
+	FD3D12SegList(const FD3D12SegList&) = delete;
+	FD3D12SegList(FD3D12SegList&&) = delete;
+
+	FD3D12SegList& operator=(const FD3D12SegList&) = delete;
+	FD3D12SegList& operator=(FD3D12SegList&&) = delete;
+
+	// @return - In-heap offset of the allocated block
+	uint32 AllocateBlock(
+		FD3D12Device* Device,
+		FRHIGPUMask VisibleNodeMask,
+		D3D12_HEAP_TYPE HeapType,
+		D3D12_HEAP_FLAGS HeapFlags,
+		TRefCountPtr<FD3D12SegHeap>& OutHeap)
+	{
+		FScopeLock Lock(&CS);
+		uint32 Offset;
+
+		if (!!FreeHeaps.Num())
+		{
+			const int32 LastHeapIdx = FreeHeaps.Num() - 1;
+			OutHeap = FreeHeaps[LastHeapIdx];
+			Offset = OutHeap->AllocateBlock(BlockSize);
+			check(Offset <= HeapSize - BlockSize);
+			if (OutHeap->IsFull(HeapSize))
+			{
+				// Heap is full
+				OutHeap->ArrayIdx = INDEX_NONE;
+				FreeHeaps.RemoveAt(LastHeapIdx);
+			}
+		}
+		else
+		{
+			OutHeap = CreateBackingHeap(Device, VisibleNodeMask, HeapType, HeapFlags);
+			Offset = OutHeap->AllocateBlock(BlockSize);
+#if D3D12RHI_SEGLIST_ALLOC_TRACK_WASTAGE
+			TotalBytesAllocated += HeapSize;
+#endif
+		}
+		return Offset;
+	}
+
+	// Deferred deletion is handled by FD3D12SegListAllocator
+	void FreeBlock(FD3D12SegHeap* RESTRICT Heap, uint32 Offset)
+	{
+		FScopeLock Lock(&CS);
+
+		check(!(Offset % BlockSize));
+		check(Offset <= HeapSize - BlockSize);
+		check(this == Heap->OwnerList);
+
+		const bool bFull = Heap->IsFull(HeapSize);
+		Heap->FreeBlockOffsets.Add(Offset);
+
+		if (bFull)
+		{
+			// Heap was full
+			check(!Heap->IsArrayIdxValid());
+			Heap->ArrayIdx = FreeHeaps.Add(Heap);
+		}
+		else if (Heap->IsEmpty(BlockSize))
+		{
+			// Heap is empty
+			check(Heap->GetRefCount() == 1);
+			check(Heap->IsArrayIdxValid());
+			check(FreeHeaps.Num() > Heap->ArrayIdx);
+			const int32 Idx = Heap->ArrayIdx;
+			const int32 LastIdx = FreeHeaps.Num() - 1;
+			FreeHeaps.RemoveAtSwap(Idx);
+			if (Idx != LastIdx)
+			{
+				FreeHeaps[Idx]->ArrayIdx = Idx;
+			}
+#if D3D12RHI_SEGLIST_ALLOC_TRACK_WASTAGE
+			TotalBytesAllocated -= HeapSize;
+#endif
+		}
+	}
+
+	FD3D12SegHeap* CreateBackingHeap(
+		FD3D12Device* Parent,
+		FRHIGPUMask VisibleNodeMask,
+		D3D12_HEAP_TYPE HeapType,
+		D3D12_HEAP_FLAGS HeapFlags);
+
+	TArray<TRefCountPtr<FD3D12SegHeap>> FreeHeaps;
+	FCriticalSection CS;
+	uint32 BlockSize;
+	uint32 HeapSize;
+
+#if D3D12RHI_SEGLIST_ALLOC_TRACK_WASTAGE
+	uint64 TotalBytesAllocated;
+#endif
+
+	friend class FD3D12SegListAllocator;
+};
+
+#if !D3D12RHI_SEGLIST_ALLOC_TRACK_WASTAGE
+static_assert(sizeof(FD3D12SegList) <= 64, "Try to make it fit in a single cacheline");
+#endif
+
+class FD3D12SegListAllocator : public FD3D12DeviceChild, public FD3D12MultiNodeGPUObject
+{
+public:
+	static constexpr uint32 InvalidOffset = 0xffffffff;
+
+	FD3D12SegListAllocator(
+		FD3D12Device* Parent,
+		FRHIGPUMask VisibilityMask,
+		D3D12_HEAP_TYPE InHeapType,
+		D3D12_HEAP_FLAGS InHeapFlags,
+		uint32 InMinPoolSize,
+		uint32 InMinNumToPool,
+		uint32 InMaxPoolSize);
+
+	~FD3D12SegListAllocator()
+	{
+		check(!SegLists.Num());
+		check(!FenceValues.Num());
+		check(!DeferredDeletionQueue.Num());
+#if D3D12RHI_SEGLIST_ALLOC_TRACK_WASTAGE
+		check(!TotalBytesRequested);
+#endif
+	}
+
+	FD3D12SegListAllocator(const FD3D12SegListAllocator&) = delete;
+	FD3D12SegListAllocator(FD3D12SegListAllocator&&) = delete;
+
+	FD3D12SegListAllocator& operator=(const FD3D12SegListAllocator&) = delete;
+	FD3D12SegListAllocator& operator=(FD3D12SegListAllocator&&) = delete;
+
+	uint32 Allocate(uint32 SizeInBytes, uint32 Alignment, TRefCountPtr<FD3D12SegHeap>& OutHeap)
+	{
+		check(!(Alignment & Alignment - 1));
+
+		const uint32 BlockSize = CalculateBlockSize(SizeInBytes, Alignment);
+		if (ShouldPool(BlockSize))
+		{
+			FD3D12SegList* SegList;
+			{
+				FRWScopeLock Lock(SegListsRWLock, SLT_ReadOnly);
+				FD3D12SegList** SegListPtr = SegLists.Find(BlockSize);
+				SegList = !!SegListPtr ? *SegListPtr : nullptr;
+			}
+			if (!SegList)
+			{
+				const uint32 HeapSize = CalculateHeapSize(BlockSize);
+				{
+					FRWScopeLock Lock(SegListsRWLock, SLT_Write);
+					FD3D12SegList** SegListPtr = SegLists.Find(BlockSize);
+					SegList = !!SegListPtr ?
+						*SegListPtr :
+						SegLists.Add(BlockSize, new FD3D12SegList(BlockSize, HeapSize));
+				}
+			}
+			check(!!SegList);
+			uint32 Ret = SegList->AllocateBlock(
+				this->GetParentDevice(),
+				this->GetVisibilityMask(),
+				HeapType,
+				HeapFlags,
+				OutHeap);
+			check(Ret != InvalidOffset);
+#if D3D12RHI_SEGLIST_ALLOC_TRACK_WASTAGE
+			TotalBytesRequested += SizeInBytes;
+#endif
+			return Ret;
+		}
+		OutHeap = nullptr;
+		return InvalidOffset;
+	}
+
+	void Deallocate(FD3D12Resource* PlacedResource, uint32 Offset, uint32 SizeInBytes);
+
+	void CleanUpAllocations();
+
+	void Destroy();
+
+	bool GetMemoryStats(uint64& OutTotalAllocated, uint64& OutTotalUnused) const
+	{
+#if D3D12RHI_SEGLIST_ALLOC_TRACK_WASTAGE
+		FScopeLock LockCS(&DeferredDeletionCS);
+		FRWScopeLock LockRW(SegListsRWLock, SLT_Write);
+
+		OutTotalAllocated = 0;
+		for (const auto& Pair : SegLists)
+		{
+			const FD3D12SegList* SegList = Pair.Value;
+			OutTotalAllocated += SegList->TotalBytesAllocated;
+		}
+		OutTotalUnused = OutTotalAllocated - TotalBytesRequested.Load(EMemoryOrder::Relaxed);
+		return true;
+#else
+		return false;
+#endif
+	}
+
+private:
+	struct FRetiredBlock
+	{
+		// FD3D12Resource knows which heap it is from
+		FD3D12Resource* PlacedResource;
+		uint32 Offset;
+		uint32 ResourceSize;
+
+		FRetiredBlock(
+			FD3D12Resource* InResource,
+			uint32 InOffset,
+			uint32 InResourceSize) :
+			PlacedResource(InResource),
+			Offset(InOffset),
+			ResourceSize(InResourceSize)
+		{}
+	};
+
+	static constexpr uint32 CalculateBlockSize(uint32 SizeInBytes, uint32 Alignment)
+	{
+		return (SizeInBytes + Alignment - 1) & ~(Alignment - 1);
+	}
+
+	bool ShouldPool(uint32 BlockSize) const
+	{
+		return BlockSize * 2 <= MaxPoolSize;
+	}
+
+	uint32 CalculateHeapSize(uint32 BlockSize) const
+	{
+		check(MinPoolSize + BlockSize - 1 > MinPoolSize);
+		uint32 NumPooled = (MinPoolSize + BlockSize - 1) / BlockSize;
+		if (NumPooled < MinNumToPool)
+		{
+			NumPooled = MinNumToPool;
+		}
+		const uint32 MaxNumPooled = MaxPoolSize / BlockSize;
+		if (NumPooled > MaxNumPooled)
+		{
+			NumPooled = MaxNumPooled;
+		}
+		check(NumPooled > 1);
+		check(NumPooled * BlockSize >= MinPoolSize);
+		check(NumPooled * BlockSize <= MaxPoolSize);
+		return NumPooled * BlockSize;
+	}
+
+	template <typename AllocX, typename AllocY>
+	void FreeRetiredBlocks(TArray<TArray<FRetiredBlock, AllocX>, AllocY>& PendingDeletes);
+
+	TMap<uint32, FD3D12SegList*> SegLists;
+	TArray<uint64> FenceValues;
+	TArray<TArray<FRetiredBlock>> DeferredDeletionQueue;
+	mutable FRWLock SegListsRWLock;
+	mutable FCriticalSection DeferredDeletionCS;
+	const D3D12_HEAP_TYPE HeapType;
+	const D3D12_HEAP_FLAGS HeapFlags;
+	const uint32 MinPoolSize;
+	const uint32 MinNumToPool;
+	const uint32 MaxPoolSize;
+
+#if D3D12RHI_SEGLIST_ALLOC_TRACK_WASTAGE
+	TAtomic<uint64> TotalBytesRequested;
+#endif
+};
+
+//-----------------------------------------------------------------------------
+//	FD3D12TextureAllocator
+//-----------------------------------------------------------------------------
+
+#if D3D12RHI_SEGREGATED_TEXTURE_ALLOC
+class FD3D12TextureAllocatorPool : public FD3D12DeviceChild, public FD3D12MultiNodeGPUObject
+{
+public:
+	FD3D12TextureAllocatorPool(FD3D12Device* Device, FRHIGPUMask VisibilityNode);
+
+	HRESULT AllocateTexture(D3D12_RESOURCE_DESC Desc, const D3D12_CLEAR_VALUE* ClearValue, uint8 UEFormat, FD3D12ResourceLocation& TextureLocation, const D3D12_RESOURCE_STATES InitialState);
+
+	void CleanUpAllocations()
+	{
+		ReadOnlyTexturePool.CleanUpAllocations();
+	}
+
+	void Destroy()
+	{
+		ReadOnlyTexturePool.Destroy();
+	}
+
+	bool GetMemoryStats(uint64& OutTotalAllocated, uint64& OutTotalUnused) const
+	{
+		return ReadOnlyTexturePool.GetMemoryStats(OutTotalAllocated, OutTotalUnused);
+	}
+
+private:
+	FD3D12SegListAllocator ReadOnlyTexturePool;
+};
+#else
+class FD3D12TextureAllocator : public FD3D12MultiBuddyAllocator
+{
+public:
+	FD3D12TextureAllocator(FD3D12Device* Device, FRHIGPUMask VisibleNodes, const FString& Name, uint32 HeapSize, D3D12_HEAP_FLAGS Flags);
+
+	~FD3D12TextureAllocator();
+
+	HRESULT AllocateTexture(D3D12_RESOURCE_DESC Desc, const D3D12_CLEAR_VALUE* ClearValue, FD3D12ResourceLocation& TextureLocation, const D3D12_RESOURCE_STATES InitialState);
+};
+
+class FD3D12TextureAllocatorPool : public FD3D12DeviceChild, public FD3D12MultiNodeGPUObject
+{
+public:
+	FD3D12TextureAllocatorPool(FD3D12Device* Device, FRHIGPUMask VisibilityNode);
+
+	HRESULT AllocateTexture(D3D12_RESOURCE_DESC Desc, const D3D12_CLEAR_VALUE* ClearValue, uint8 UEFormat, FD3D12ResourceLocation& TextureLocation, const D3D12_RESOURCE_STATES InitialState);
+
+	void CleanUpAllocations() { ReadOnlyTexturePool.CleanUpAllocations(); }
+
+	void Destroy() { ReadOnlyTexturePool.Destroy(); }
+
+private:
+	FD3D12TextureAllocator ReadOnlyTexturePool;
+};
+#endif
diff --git a/Engine/Source/Runtime/D3D12RHI/Private/D3D12CommandContext.cpp b/Engine/Source/Runtime/D3D12RHI/Private/D3D12CommandContext.cpp
--- a/Engine/Source/Runtime/D3D12RHI/Private/D3D12CommandContext.cpp
+++ b/Engine/Source/Runtime/D3D12RHI/Private/D3D12CommandContext.cpp
@@ -441,6 +441,19 @@ void FD3D12CommandContextBase::UpdateMemoryStats()
 	SET_MEMORY_STAT(STAT_D3D12UsedVideoMemory, LocalVideoMemoryInfo.CurrentUsage);
 	SET_MEMORY_STAT(STAT_D3D12AvailableVideoMemory, AvailableSpace);
 	SET_MEMORY_STAT(STAT_D3D12TotalVideoMemory, Budget);
+
+#if D3D12RHI_SEGREGATED_TEXTURE_ALLOC && D3D12RHI_SEGLIST_ALLOC_TRACK_WASTAGE
+	uint64 MaxTexAllocWastage = 0;
+	for (uint32 GPUIndex : GPUMask)
+	{
+		FD3D12Device* Device = ParentAdapter->GetDevice(GPUIndex);
+		uint64 TotalAllocated;
+		uint64 TotalUnused;
+		Device->GetTextureAllocator().GetMemoryStats(TotalAllocated, TotalUnused);
+		MaxTexAllocWastage = FMath::Max(MaxTexAllocWastage, TotalUnused);
+	}
+	SET_MEMORY_STAT(STAT_D3D12TextureAllocatorWastage, MaxTexAllocWastage);
+#endif
 #endif
 }
 
diff --git a/Engine/Source/Runtime/D3D12RHI/Private/D3D12Resources.cpp b/Engine/Source/Runtime/D3D12RHI/Private/D3D12Resources.cpp
--- a/Engine/Source/Runtime/D3D12RHI/Private/D3D12Resources.cpp
+++ b/Engine/Source/Runtime/D3D12RHI/Private/D3D12Resources.cpp
@@ -370,6 +370,7 @@ FD3D12ResourceLocation::FD3D12ResourceLocation(FD3D12Device* Parent)
 	, OffsetFromBaseOfResource(0)
 	, Size(0)
 	, bTransient(false)
+	, AllocatorType(AT_Unknown)
 {
 	FMemory::Memzero(AllocatorData);
 }
@@ -406,6 +407,7 @@ void FD3D12ResourceLocation::InternalClear()
 	FMemory::Memzero(AllocatorData);
 
 	Allocator = nullptr;
+	AllocatorType = AT_Unknown;
 }
 
 void FD3D12ResourceLocation::TransferOwnership(FD3D12ResourceLocation& Destination, FD3D12ResourceLocation& Source)
@@ -469,7 +471,17 @@ void FD3D12ResourceLocation::ReleaseResource()
 	case ResourceLocationType::eSubAllocation:
 	{
 		check(Allocator != nullptr);
-		Allocator->Deallocate(*this);
+		if (AllocatorType == AT_SegList)
+		{
+			SegListAllocator->Deallocate(
+				GetResource(),
+				GetSegListAllocatorPrivateData().Offset,
+				GetSize());
+		}
+		else
+		{
+			Allocator->Deallocate(*this);
+		}
 		break;
 	}
 	case ResourceLocationType::eNodeReference:
diff --git a/Engine/Source/Runtime/D3D12RHI/Private/D3D12Stats.h b/Engine/Source/Runtime/D3D12RHI/Private/D3D12Stats.h
--- a/Engine/Source/Runtime/D3D12RHI/Private/D3D12Stats.h
+++ b/Engine/Source/Runtime/D3D12RHI/Private/D3D12Stats.h
@@ -60,6 +60,7 @@ DECLARE_CYCLE_STAT_EXTERN(TEXT("WaitForFence time"), STAT_D3D12WaitForFenceTime,
 DECLARE_MEMORY_STAT_EXTERN(TEXT("Used Video Memory"), STAT_D3D12UsedVideoMemory, STATGROUP_D3D12RHI, );
 DECLARE_MEMORY_STAT_EXTERN(TEXT("Available Video Memory"), STAT_D3D12AvailableVideoMemory, STATGROUP_D3D12RHI, );
 DECLARE_MEMORY_STAT_EXTERN(TEXT("Total Video Memory"), STAT_D3D12TotalVideoMemory, STATGROUP_D3D12RHI, );
+DECLARE_MEMORY_STAT_EXTERN(TEXT("Texture allocator wastage"), STAT_D3D12TextureAllocatorWastage, STATGROUP_D3D12RHI, );
 
 /**
 * Detailed Descriptor heap stats
diff --git a/Engine/Source/Runtime/D3D12RHI/Private/D3D12Util.cpp b/Engine/Source/Runtime/D3D12RHI/Private/D3D12Util.cpp
--- a/Engine/Source/Runtime/D3D12RHI/Private/D3D12Util.cpp
+++ b/Engine/Source/Runtime/D3D12RHI/Private/D3D12Util.cpp
@@ -781,6 +781,7 @@ DEFINE_STAT(STAT_D3D12WaitForFenceTime);
 DEFINE_STAT(STAT_D3D12UsedVideoMemory);
 DEFINE_STAT(STAT_D3D12AvailableVideoMemory);
 DEFINE_STAT(STAT_D3D12TotalVideoMemory);
+DEFINE_STAT(STAT_D3D12TextureAllocatorWastage);
 
 DEFINE_STAT(STAT_ViewHeapChanged);
 DEFINE_STAT(STAT_SamplerHeapChanged);
diff --git a/Engine/Source/Runtime/D3D12RHI/Private/Windows/WindowsD3D12Device.cpp b/Engine/Source/Runtime/D3D12RHI/Private/Windows/WindowsD3D12Device.cpp
--- a/Engine/Source/Runtime/D3D12RHI/Private/Windows/WindowsD3D12Device.cpp
+++ b/Engine/Source/Runtime/D3D12RHI/Private/Windows/WindowsD3D12Device.cpp
@@ -658,6 +658,11 @@ void FD3D12DynamicRHI::Init()
 	// Indicate that the RHI needs to use the engine's deferred deletion queue.
 	GRHINeedsExtraDeletionLatency = true;
 
+	// There is no need to defer deletion of streaming textures
+	// - Suballocated ones are defer-deleted by their allocators
+	// - Standalones are added to the deferred deletion queue of its parent FD3D12Adapter
+	GRHIForceNoDeletionLatencyForStreamingTextures = !!PLATFORM_WINDOWS;
+
 	// Set the RHI initialized flag.
 	GIsRHIInitialized = true;
 }
diff --git a/Engine/Source/Runtime/D3D12RHI/Public/D3D12Resources.h b/Engine/Source/Runtime/D3D12RHI/Public/D3D12Resources.h
--- a/Engine/Source/Runtime/D3D12RHI/Public/D3D12Resources.h
+++ b/Engine/Source/Runtime/D3D12RHI/Public/D3D12Resources.h
@@ -15,6 +15,7 @@ class FD3D12StateCacheBase;
 class FD3D12CommandListManager;
 class FD3D12CommandContext;
 class FD3D12CommandListHandle;
+class FD3D12SegListAllocator;
 struct FD3D12GraphicsPipelineState;
 typedef FD3D12StateCacheBase FD3D12StateCache;
 
@@ -70,7 +71,7 @@ public:
 	~FD3D12Heap();
 
 	inline ID3D12Heap* GetHeap() { return Heap.GetReference(); }
-	inline void SetHeap(ID3D12Heap* HeapIn) { Heap = HeapIn; }
+	inline void SetHeap(ID3D12Heap* HeapIn) { *Heap.GetInitReference() = HeapIn; }
 
 	void UpdateResidency(FD3D12CommandListHandle& CommandList);
 
@@ -375,6 +376,16 @@ struct FD3D12BlockAllocatorPrivateData
 	}
 };
 
+struct FD3D12SegListAllocatorPrivateData
+{
+	uint32 Offset;
+
+	void Init()
+	{
+		Offset = 0;
+	}
+};
+
 class FD3D12ResourceAllocator;
 // A very light-weight and cache friendly way of accessing a GPU resource
 class FD3D12ResourceLocation : public FD3D12DeviceChild, public FNoncopyable
@@ -392,6 +403,13 @@ public:
 		eHeapAliased, 
 	};
 
+	enum EAllocatorType : uint8
+	{
+		AT_Default, // FD3D12BaseAllocatorType
+		AT_SegList, // FD3D12SegListAllocator
+		AT_Unknown = 0xff
+	};
+
 	FD3D12ResourceLocation(FD3D12Device* Parent);
 	~FD3D12ResourceLocation();
 
@@ -404,7 +422,8 @@ public:
 	void SetResource(FD3D12Resource* Value);
 	inline void SetType(ResourceLocationType Value) { Type = Value;}
 
-	inline void SetAllocator(FD3D12BaseAllocatorType* Value) { Allocator = Value; }
+	inline void SetAllocator(FD3D12BaseAllocatorType* Value) { Allocator = Value; AllocatorType = AT_Default; }
+	inline void SetSegListAllocator(FD3D12SegListAllocator* Value) { SegListAllocator = Value; AllocatorType = AT_SegList; }
 	inline void SetMappedBaseAddress(void* Value) { MappedBaseAddress = Value; }
 	inline void SetGPUVirtualAddress(D3D12_GPU_VIRTUAL_ADDRESS Value) { GPUVirtualAddress = Value; }
 	inline void SetOffsetFromBaseOfResource(uint64 Value) { OffsetFromBaseOfResource = Value; }
@@ -412,7 +431,8 @@ public:
 
 	// Getters
 	inline ResourceLocationType GetType() const { return Type; }
-	inline FD3D12BaseAllocatorType* GetAllocator() { return Allocator; }
+	inline FD3D12BaseAllocatorType* GetAllocator() { check(AT_Default == AllocatorType); return Allocator; }
+	inline FD3D12SegListAllocator* GetSegListAllocator() { check(AT_SegList == AllocatorType); return SegListAllocator; }
 	inline FD3D12Resource* GetResource() const { return UnderlyingResource; }
 	inline void* GetMappedBaseAddress() const { return MappedBaseAddress; }
 	inline D3D12_GPU_VIRTUAL_ADDRESS GetGPUVirtualAddress() const { return GPUVirtualAddress; }
@@ -421,6 +441,7 @@ public:
 	inline FD3D12ResidencyHandle* GetResidencyHandle() { return ResidencyHandle; }
 	inline FD3D12BuddyAllocatorPrivateData& GetBuddyAllocatorPrivateData() { return AllocatorData.BuddyAllocatorPrivateData; }
 	inline FD3D12BlockAllocatorPrivateData& GetBlockAllocatorPrivateData() { return AllocatorData.BlockAllocatorPrivateData; }
+	inline FD3D12SegListAllocatorPrivateData& GetSegListAllocatorPrivateData() { return AllocatorData.SegListAllocatorPrivateData; }
 
 	const inline bool IsValid() const { return Type != ResourceLocationType::eUndefined; }
 
@@ -494,13 +515,18 @@ private:
 	FD3D12ResidencyHandle* ResidencyHandle;
 
 	// Which allocator this belongs to
-	FD3D12BaseAllocatorType* Allocator;
+	union
+	{
+		FD3D12BaseAllocatorType* Allocator;
+		FD3D12SegListAllocator* SegListAllocator;
+	};
 
 	// Union to save memory
 	union PrivateAllocatorData
 	{
 		FD3D12BuddyAllocatorPrivateData BuddyAllocatorPrivateData;
 		FD3D12BlockAllocatorPrivateData BlockAllocatorPrivateData;
+		FD3D12SegListAllocatorPrivateData SegListAllocatorPrivateData;
 	} AllocatorData;
 
 	// Note: These values refer to the start of this location including any padding *NOT* the start of the underlying resource
@@ -512,6 +538,8 @@ private:
 	uint64 Size;
 
 	bool bTransient;
+
+	EAllocatorType AllocatorType;
 };
 
 class FD3D12DeferredDeletionQueue : public FD3D12AdapterChild
diff --git a/Engine/Source/Runtime/Engine/Private/Texture2D.cpp b/Engine/Source/Runtime/Engine/Private/Texture2D.cpp
--- a/Engine/Source/Runtime/Engine/Private/Texture2D.cpp
+++ b/Engine/Source/Runtime/Engine/Private/Texture2D.cpp
@@ -1621,6 +1621,11 @@ void FTexture2DResource::UpdateTexture(FTexture2DRHIRef& InTextureRHI, int32 InN
 			INC_DWORD_STAT_FNAME_BY( LODGroupStatName, TextureSize );
 		}
 
+		if (GRHIForceNoDeletionLatencyForStreamingTextures)
+		{
+			TextureRHI->DoNoDeferDelete();
+		}
+
 		TextureRHI		= InTextureRHI;
 		Texture2DRHI	= InTextureRHI;
 		CurrentFirstMip = InNewFirstMip;
diff --git a/Engine/Source/Runtime/RHI/Private/RHI.cpp b/Engine/Source/Runtime/RHI/Private/RHI.cpp
--- a/Engine/Source/Runtime/RHI/Private/RHI.cpp
+++ b/Engine/Source/Runtime/RHI/Private/RHI.cpp
@@ -476,6 +476,7 @@ bool GSupportsWideMRT = true;
 float GMinClipZ = 0.0f;
 float GProjectionSignY = 1.0f;
 bool GRHINeedsExtraDeletionLatency = false;
+bool GRHIForceNoDeletionLatencyForStreamingTextures = false;
 TRHIGlobal<int32> GMaxComputeDispatchDimension((1 << 16) - 1);
 bool GRHILazyShaderCodeLoading = false;
 bool GRHISupportsLazyShaderCodeLoading = false;
diff --git a/Engine/Source/Runtime/RHI/Public/RHI.h b/Engine/Source/Runtime/RHI/Public/RHI.h
--- a/Engine/Source/Runtime/RHI/Public/RHI.h
+++ b/Engine/Source/Runtime/RHI/Public/RHI.h
@@ -333,6 +333,9 @@ extern RHI_API float GProjectionSignY;
 /** Does this RHI need to wait for deletion of resources due to ref counting. */
 extern RHI_API bool GRHINeedsExtraDeletionLatency;
 
+/** Allow opt-out default RHI resource deletion latency for streaming textures */
+extern RHI_API bool GRHIForceNoDeletionLatencyForStreamingTextures;
+
 /** The maximum size allowed for a computeshader dispatch. */
 extern RHI_API TRHIGlobal<int32> GMaxComputeDispatchDimension;
 
