PR #5291: Fixed typo (Contributed by [REDACTED_PERSON])

diff --git a/Engine/Source/Runtime/D3D12RHI/Private/D3D12Allocation.h b/Engine/Source/Runtime/D3D12RHI/Private/D3D12Allocation.h
--- a/Engine/Source/Runtime/D3D12RHI/Private/D3D12Allocation.h
+++ b/Engine/Source/Runtime/D3D12RHI/Private/D3D12Allocation.h
@@ -94,10 +94,10 @@ enum eBuddyAllocationStrategy
 {
 	// This strategy uses Placed Resources to sub-allocate a buffer out of an underlying ID3D12Heap.
 	// The benefit of this is that each buffer can have it's own resource state and can be treated
-	// as any other buffer. The downside of this strategy is the API limitiation which enforces
+	// as any other buffer. The downside of this strategy is the API limitation which enforces
 	// the minimum buffer size to 64k leading to large internal fragmentation in the allocator
 	kPlacedResourceStrategy,
-	// The alternative is to manualy sub-allocate out of a single large buffer which allows block
+	// The alternative is to manually sub-allocate out of a single large buffer which allows block
 	// allocation granularity down to 1 byte. However, this strategy is only really valid for buffers which
 	// will be treated as read-only after their creation (i.e. most Index and Vertex buffers). This 
 	// is because the underlying resource can only have one state at a time.
