Make DX12 stencilref behavior consistent with other RHIs

diff --git a/Engine/Source/Runtime/D3D12RHI/Private/D3D12Commands.cpp b/Engine/Source/Runtime/D3D12RHI/Private/D3D12Commands.cpp
--- a/Engine/Source/Runtime/D3D12RHI/Private/D3D12Commands.cpp
+++ b/Engine/Source/Runtime/D3D12RHI/Private/D3D12Commands.cpp
@@ -570,6 +570,7 @@ void FD3D12CommandContext::RHISetGraphicsPipelineState(FGraphicsPipelineStateRHI
 	RHIEnableDepthBoundsTest(GraphicsPipelineState->PipelineStateInitializer.bDepthBounds);
 
 	StateCache.SetGraphicsPipelineState(GraphicsPipelineState);
+	StateCache.SetStencilRef(0);
 }
 
 void FD3D12CommandContext::RHISetComputePipelineState(FRHIComputePipelineState* ComputeState)
