Added debug console variables to disable back-face culling and any-hit shaders for all ray tracing geometry instances.

diff --git a/Engine/Source/Runtime/D3D12RHI/Private/D3D12RayTracing.cpp b/Engine/Source/Runtime/D3D12RHI/Private/D3D12RayTracing.cpp
--- a/Engine/Source/Runtime/D3D12RHI/Private/D3D12RayTracing.cpp
+++ b/Engine/Source/Runtime/D3D12RHI/Private/D3D12RayTracing.cpp
@@ -11,8 +11,23 @@
 #include "CommonRayTracingBuiltInResources.ush"
 #include "Hash/CityHash.h"
 #include "HAL/CriticalSection.h"
+#include "HAL/IConsoleManager.h"
 #include "Misc/ScopeLock.h"
 
+static int32 GRayTracingDebugForceOpaque = 0;
+static FAutoConsoleVariableRef CVarRayTracingDebugForceOpaque(
+	TEXT("r.RayTracing.DebugForceOpaque"),
+	GRayTracingDebugForceOpaque,
+	TEXT("Forces all ray tracing geometry instances to be opaque, effectively disabling any-hit shaders. This is useful for debugging and profiling. (default = 0)")
+);
+
+static int32 GRayTracingDebugDisableTriangleCull = 0;
+static FAutoConsoleVariableRef CVarRayTracingDebugDisableTriangleCull(
+	TEXT("r.RayTracing.DebugDisableTriangleCull"),
+	GRayTracingDebugDisableTriangleCull,
+	TEXT("Forces all ray tracing geometry instances to be double-sided by disabling back-face culling. This is useful for debugging and profiling. (default = 0)")
+);
+
 #define VERIFYHRESULT(expr) { HRESULT HR##__LINE__ = expr; if (FAILED(HR##__LINE__)) { UE_LOG(LogD3D12RHI, Fatal, TEXT(#expr " failed: Result=%08x"), HR##__LINE__); } }
 
 // Built-in local root parameters that are always bound to all hit shaders
@@ -1537,6 +1552,17 @@ void FD3D12RayTracingScene::BuildAccelerationStructure(FD3D12CommandContext& Com
 			InstanceDesc.InstanceMask = 0xFF; // Instance mask is currently unused
 			InstanceDesc.InstanceContributionToHitGroupIndex = SegmentPrefixSum[InstanceIndex];
 			InstanceDesc.Flags = D3D12_RAYTRACING_INSTANCE_FLAG_TRIANGLE_FRONT_COUNTERCLOCKWISE; // #dxr_todo: convert cull mode based on instance mirroring or double-sidedness
+			
+			if (GRayTracingDebugForceOpaque)
+			{
+				InstanceDesc.Flags |= D3D12_RAYTRACING_INSTANCE_FLAG_FORCE_OPAQUE;
+			}
+
+			if (GRayTracingDebugDisableTriangleCull)
+			{
+				InstanceDesc.Flags |= D3D12_RAYTRACING_INSTANCE_FLAG_TRIANGLE_CULL_DISABLE;
+			}
+
 			InstanceDesc.AccelerationStructure = Geometry->AccelerationStructureBuffer->ResourceLocation.GetGPUVirtualAddress();
 
 			MappedData[InstanceIndex] = InstanceDesc;
