Fixed UE-67773 (Crash when launching with -raytracing arg when raytracing project setting is disabled)

diff --git a/Engine/Source/Runtime/D3D12RHI/Private/Windows/WindowsD3D12Device.cpp b/Engine/Source/Runtime/D3D12RHI/Private/Windows/WindowsD3D12Device.cpp
--- a/Engine/Source/Runtime/D3D12RHI/Private/Windows/WindowsD3D12Device.cpp
+++ b/Engine/Source/Runtime/D3D12RHI/Private/Windows/WindowsD3D12Device.cpp
@@ -669,9 +669,12 @@ void FD3D12DynamicRHI::Init()
 
 void FD3D12DynamicRHI::PostInit()
 {
-	for (FD3D12Adapter*& Adapter : ChosenAdapters)
+	if (IsRayTracingSupportedForThisProject())
 	{
-		Adapter->InitializeRayTracing();
+		for (FD3D12Adapter*& Adapter : ChosenAdapters)
+		{
+			Adapter->InitializeRayTracing();
+		}
 	}
 
 	if (GRHISupportsRHIThread)
