FD3D12DynamicRHI::RHIGetRayTracingSupport now returns 2 when DXR is supported, which corresponds to full ray tracing support.

diff --git a/Engine/Source/Runtime/D3D12RHI/Private/D3D12RayTracing.cpp b/Engine/Source/Runtime/D3D12RHI/Private/D3D12RayTracing.cpp
--- a/Engine/Source/Runtime/D3D12RHI/Private/D3D12RayTracing.cpp
+++ b/Engine/Source/Runtime/D3D12RHI/Private/D3D12RayTracing.cpp
@@ -249,7 +249,7 @@ static TRefCountPtr<ID3D12StateObject> CreateRayTracingStateObject(
 
 uint32 FD3D12DynamicRHI::RHIGetRayTracingSupport()
 {
-	return GetAdapter().GetD3DRayTracingDevice() ? 1 : 0;
+	return GetAdapter().GetD3DRayTracingDevice() ? 2 : 0;
 }
 
 // #dxr_todo: FD3D12Device::GlobalViewHeap/GlobalSamplerHeap should be used instead of ad-hoc heaps here.
