Reintroduce graphical corruption fix in cl 4616014 that got lost in recent merges.

diff --git a/Engine/Source/Runtime/D3D12RHI/Private/D3D12UniformBuffer.cpp b/Engine/Source/Runtime/D3D12RHI/Private/D3D12UniformBuffer.cpp
--- a/Engine/Source/Runtime/D3D12RHI/Private/D3D12UniformBuffer.cpp
+++ b/Engine/Source/Runtime/D3D12RHI/Private/D3D12UniformBuffer.cpp
@@ -144,6 +144,10 @@ struct FRHICommandD3D12UpdateUniformBuffer final : public FRHICommand<FRHIComman
 			check(UniformBuffer->ResourceTable[i]);
 		}
 		FD3D12ResourceLocation::TransferOwnership(UniformBuffer->ResourceLocation, *UpdatedLocation);
+#if USE_STATIC_ROOT_SIGNATURE
+		const uint32 NumBytes = Align(UniformBuffer->GetLayout().ConstantBufferSize, D3D12_CONSTANT_BUFFER_DATA_PLACEMENT_ALIGNMENT);
+		UniformBuffer->View->Create(UniformBuffer->ResourceLocation.GetGPUVirtualAddress(), NumBytes);
+#endif
 	}
 };
 
