Fix assert fail case. UAVs needs separate assert conditions.

diff --git a/Engine/Source/Runtime/D3D12RHI/Private/D3D12Allocation.cpp b/Engine/Source/Runtime/D3D12RHI/Private/D3D12Allocation.cpp
--- a/Engine/Source/Runtime/D3D12RHI/Private/D3D12Allocation.cpp
+++ b/Engine/Source/Runtime/D3D12RHI/Private/D3D12Allocation.cpp
@@ -974,12 +974,23 @@ void FD3D12DefaultBufferPool::AllocDefaultResource(const D3D12_RESOURCE_DESC& De
 
 	if (PoolResource)
 	{
+		const bool bPlacedResource = (Allocator->GetAllocationStrategy() == kPlacedResourceStrategy);
+		
 		// Ensure we're allocating from the correct pool
-		check(Desc.Flags == Allocator->ResourceFlags);
+		if (bPlacedResource)
+		{
+			// Writeable resources get separate ID3D12Resource* with their own resource state by using placed resources. Just make sure it's UAV, other flags are free to differ.
+			check((Desc.Flags & Allocator->ResourceFlags & D3D12_RESOURCE_FLAG_ALLOW_UNORDERED_ACCESS) != 0);
+		}
+		else
+		{
+			// Read-only resources get suballocated from big resources, thus share ID3D12Resource* and resource state with other resources. Ensure it's suballocated from a resource with identical flags.
+			check(Desc.Flags == Allocator->ResourceFlags);
+		}
 	
 		if (Allocator->TryAllocate(Desc.Width, Alignment, ResourceLocation))
 		{
-			if (Allocator->GetAllocationStrategy() == kPlacedResourceStrategy)
+			if (bPlacedResource)
 			{
 				check(ResourceLocation.GetResource() == nullptr);
 
@@ -991,6 +1002,10 @@ void FD3D12DefaultBufferPool::AllocDefaultResource(const D3D12_RESOURCE_DESC& De
 
 				ResourceLocation.SetResource(NewResource);
 			}
+			else
+			{
+				// Nothing to do for suballocated resources
+			}
 
 			// Successfully sub-allocated
 			return;
@@ -998,9 +1013,8 @@ void FD3D12DefaultBufferPool::AllocDefaultResource(const D3D12_RESOURCE_DESC& De
 	}
 
 	// Allocate Standalone
-
 	FD3D12Resource* NewResource = nullptr;
-	VERIFYD3D12RESULT(Adapter->CreateBuffer(D3D12_HEAP_TYPE_DEFAULT, GetGPUMask(), GetVisibilityMask(), InitialState, Desc.Width, &NewResource, Allocator->ResourceFlags));
+	VERIFYD3D12RESULT(Adapter->CreateBuffer(D3D12_HEAP_TYPE_DEFAULT, GetGPUMask(), GetVisibilityMask(), InitialState, Desc.Width, &NewResource, Desc.Flags));
 	SetName(NewResource, L"Stand Alone Default Buffer");
 
 	ResourceLocation.AsStandAlone(NewResource, Desc.Width);
