Uniform buffers are now used to set all local ray tracing parameters. Global parameters are still set through FRayTracingShaderBindings for now.

diff --git a/Engine/Source/Runtime/D3D12RHI/Private/D3D12CommandContext.h b/Engine/Source/Runtime/D3D12RHI/Private/D3D12CommandContext.h
--- a/Engine/Source/Runtime/D3D12RHI/Private/D3D12CommandContext.h
+++ b/Engine/Source/Runtime/D3D12RHI/Private/D3D12CommandContext.h
@@ -370,7 +370,7 @@ public:
 	virtual void RHISetRayTracingHitGroup(
 		FRayTracingSceneRHIParamRef Scene, uint32 InstanceIndex, uint32 SegmentIndex,
 		FRayTracingPipelineStateRHIParamRef Pipeline, uint32 HitGroupIndex,
-		const FRayTracingShaderBindings& ResourceBindings) final override;
+		uint32 NumUniformBuffers, const FUniformBufferRHIParamRef* UniformBuffers) final override;
 #endif // D3D12_RHI_RAYTRACING
 
 	template<typename ObjectType, typename RHIType, typename Predicate>
diff --git a/Engine/Source/Runtime/D3D12RHI/Private/D3D12RayTracing.cpp b/Engine/Source/Runtime/D3D12RHI/Private/D3D12RayTracing.cpp
--- a/Engine/Source/Runtime/D3D12RHI/Private/D3D12RayTracing.cpp
+++ b/Engine/Source/Runtime/D3D12RHI/Private/D3D12RayTracing.cpp
@@ -1799,7 +1799,11 @@ template <typename ResourceBinderType>
 static void SetRayTracingShaderResources(
 	FD3D12CommandContext& CommandContext,
 	const FD3D12RayTracingShader* Shader,
-	const FRayTracingShaderBindings& ResourceBindings,
+	uint32 InNumTextures, const FTextureRHIParamRef* Textures,
+	uint32 InNumSRVs, const FShaderResourceViewRHIParamRef* SRVs,
+	uint32 InNumUniformBuffers, const FUniformBufferRHIParamRef* UniformBuffers,
+	uint32 InNumSamplers, const FSamplerStateRHIParamRef* Samplers,
+	uint32 InNumUAVs, const FUnorderedAccessViewRHIParamRef* UAVs,
 	FD3D12RayTracingDescriptorCache& DescriptorCache,
 	ResourceBinderType& Binder)
 {
@@ -1817,9 +1821,9 @@ static void SetRayTracingShaderResources(
 	uint64 BoundUAVMask = 0;
 	uint64 BoundSamplerMask = 0;
 
-	for (uint32 SRVIndex = 0; SRVIndex < ARRAY_COUNT(ResourceBindings.Textures); ++SRVIndex)
+	for (uint32 SRVIndex = 0; SRVIndex < InNumTextures; ++SRVIndex)
 	{
-		FTextureRHIParamRef Resource = ResourceBindings.Textures[SRVIndex];
+		FTextureRHIParamRef Resource = Textures[SRVIndex];
 		if (Resource)
 		{
 			LocalSRVs[SRVIndex] = CommandContext.RetrieveTextureBase(Resource)->GetShaderResourceView()->GetView();
@@ -1827,9 +1831,9 @@ static void SetRayTracingShaderResources(
 		}
 	}
 
-	for (uint32 SRVIndex = 0; SRVIndex < ARRAY_COUNT(ResourceBindings.SRVs); ++SRVIndex)
+	for (uint32 SRVIndex = 0; SRVIndex < InNumSRVs; ++SRVIndex)
 	{
-		FShaderResourceViewRHIParamRef Resource = ResourceBindings.SRVs[SRVIndex];
+		FShaderResourceViewRHIParamRef Resource = SRVs[SRVIndex];
 		if (Resource)
 		{
 			LocalSRVs[SRVIndex] = FD3D12DynamicRHI::ResourceCast(Resource)->GetView();
@@ -1837,9 +1841,9 @@ static void SetRayTracingShaderResources(
 		}
 	}
 
-	for (uint32 CBVIndex = 0; CBVIndex < ARRAY_COUNT(ResourceBindings.UniformBuffers); ++CBVIndex)
+	for (uint32 CBVIndex = 0; CBVIndex < InNumUniformBuffers; ++CBVIndex)
 	{
-		FUniformBufferRHIParamRef Resource = ResourceBindings.UniformBuffers[CBVIndex];
+		FUniformBufferRHIParamRef Resource = UniformBuffers[CBVIndex];
 		if (Resource)
 		{
 			LocalCBVs[CBVIndex] = FD3D12DynamicRHI::ResourceCast(Resource);
@@ -1847,9 +1851,9 @@ static void SetRayTracingShaderResources(
 		}
 	}
 
-	for (uint32 SamplerIndex = 0; SamplerIndex < ARRAY_COUNT(ResourceBindings.Samplers); ++SamplerIndex)
+	for (uint32 SamplerIndex = 0; SamplerIndex < InNumSamplers; ++SamplerIndex)
 	{
-		FSamplerStateRHIParamRef Resource = ResourceBindings.Samplers[SamplerIndex];
+		FSamplerStateRHIParamRef Resource = Samplers[SamplerIndex];
 		if (Resource)
 		{
 			LocalSamplers[SamplerIndex] = FD3D12DynamicRHI::ResourceCast(Resource)->Descriptor;
@@ -1857,9 +1861,9 @@ static void SetRayTracingShaderResources(
 		}
 	}
 
-	for (uint32 UAVIndex = 0; UAVIndex < ARRAY_COUNT(ResourceBindings.UAVs); ++UAVIndex)
+	for (uint32 UAVIndex = 0; UAVIndex < InNumUAVs; ++UAVIndex)
 	{
-		FUnorderedAccessViewRHIParamRef Resource = ResourceBindings.UAVs[UAVIndex];
+		FUnorderedAccessViewRHIParamRef Resource = UAVs[UAVIndex];
 		if (Resource)
 		{
 			LocalUAVs[UAVIndex] = FD3D12DynamicRHI::ResourceCast(Resource)->GetView();
@@ -1877,8 +1881,8 @@ static void SetRayTracingShaderResources(
 		const uint32 LowestBitMask = (DirtyBits)& (-(int32)DirtyBits);
 		const int32 BufferIndex = FMath::FloorLog2(LowestBitMask); // todo: This has a branch on zero, we know it could never be zero...
 		DirtyBits ^= LowestBitMask;
-		check(BufferIndex < ARRAY_COUNT(ResourceBindings.UniformBuffers));
-		FD3D12UniformBuffer* Buffer = FD3D12DynamicRHI::ResourceCast(ResourceBindings.UniformBuffers[BufferIndex]);
+		check(uint32(BufferIndex) < InNumUniformBuffers);
+		FD3D12UniformBuffer* Buffer = FD3D12DynamicRHI::ResourceCast(UniformBuffers[BufferIndex]);
 		check(Buffer);
 		check(BufferIndex < ShaderResourceTable.ResourceTableLayoutHashes.Num());
 		check(Buffer->GetLayout().GetHash() == ShaderResourceTable.ResourceTableLayoutHashes[BufferIndex]);
@@ -2028,6 +2032,23 @@ static void SetRayTracingShaderResources(
 	}
 }
 
+template <typename ResourceBinderType>
+static void SetRayTracingShaderResources(
+	FD3D12CommandContext& CommandContext,
+	const FD3D12RayTracingShader* Shader,
+	const FRayTracingShaderBindings& ResourceBindings,
+	FD3D12RayTracingDescriptorCache& DescriptorCache,
+	ResourceBinderType& Binder)
+{
+	SetRayTracingShaderResources(CommandContext, Shader,
+		ARRAY_COUNT(ResourceBindings.Textures), ResourceBindings.Textures,
+		ARRAY_COUNT(ResourceBindings.SRVs), ResourceBindings.SRVs,
+		ARRAY_COUNT(ResourceBindings.UniformBuffers), ResourceBindings.UniformBuffers,
+		ARRAY_COUNT(ResourceBindings.Samplers), ResourceBindings.Samplers,
+		ARRAY_COUNT(ResourceBindings.UAVs), ResourceBindings.UAVs,
+		DescriptorCache, Binder);
+}
+
 static void DispatchRays(FD3D12CommandContext& CommandContext,
 	const FRayTracingShaderBindings& GlobalBindings,
 	const FD3D12RayTracingPipelineState* Pipeline,
@@ -2198,7 +2219,7 @@ void FD3D12CommandContext::RHIRayTraceDispatch(FRayTracingPipelineStateRHIParamR
 void FD3D12CommandContext::RHISetRayTracingHitGroup(
 	FRayTracingSceneRHIParamRef InScene, uint32 InstanceIndex, uint32 SegmentIndex,
 	FRayTracingPipelineStateRHIParamRef InPipeline, uint32 HitGroupIndex,
-	const FRayTracingShaderBindings& ResourceBindings)
+	uint32 NumUniformBuffers, const FUniformBufferRHIParamRef* UniformBuffers)
 {
 	FD3D12RayTracingScene* Scene = FD3D12DynamicRHI::ResourceCast(InScene);
 	FD3D12RayTracingPipelineState* Pipeline = FD3D12DynamicRHI::ResourceCast(InPipeline);
@@ -2211,7 +2232,13 @@ void FD3D12CommandContext::RHISetRayTracingHitGroup(
 	const FD3D12RayTracingShader* Shader = Pipeline->HitGroupShaders[HitGroupIndex];
 
 	FD3D12RayTracingLocalResourceBinder ResourceBinder(*this, ShaderTable, Shader->pRootSignature, RecordIndex);
-	SetRayTracingShaderResources(*this, Shader, ResourceBindings, ShaderTable->DescriptorCache, ResourceBinder);
+	SetRayTracingShaderResources(*this, Shader,
+		0, nullptr, // Textures
+		0, nullptr, // SRVs
+		NumUniformBuffers, UniformBuffers, 
+		0, nullptr, // Samplers
+		0, nullptr, // UAVs
+		ShaderTable->DescriptorCache, ResourceBinder);
 }
 
 #undef VERIFYHRESULT
diff --git a/Engine/Source/Runtime/RHI/Public/RHICommandList.h b/Engine/Source/Runtime/RHI/Public/RHICommandList.h
--- a/Engine/Source/Runtime/RHI/Public/RHICommandList.h
+++ b/Engine/Source/Runtime/RHI/Public/RHICommandList.h
@@ -1909,16 +1909,19 @@ struct FRHICommandSetRayTracingHitGroup final : public FRHICommand<FRHICommandSe
 	uint32 SegmentIndex;
 	FRayTracingPipelineStateRHIParamRef Pipeline;
 	uint32 HitGroupIndex;
-	FRayTracingShaderBindings ResourceBindings;
+	uint32 NumUniformBuffers;
+	const FUniformBufferRHIParamRef* UniformBuffers; // Pointer to an array of uniform buffers, allocated inline with the command list
 
-	FRHICommandSetRayTracingHitGroup(FRayTracingSceneRHIParamRef InScene, uint32 InInstanceIndex, uint32 InSegmentIndex, FRayTracingPipelineStateRHIParamRef InPipeline, uint32 InHitGroupIndex, const FRayTracingShaderBindings& InResourceBindings)
+	FRHICommandSetRayTracingHitGroup(FRayTracingSceneRHIParamRef InScene, uint32 InInstanceIndex, uint32 InSegmentIndex, FRayTracingPipelineStateRHIParamRef InPipeline, uint32 InHitGroupIndex, uint32 InNumUniformBuffers, const FUniformBufferRHIParamRef* InUniformBuffers)
 		: Scene(InScene)
 		, InstanceIndex(InInstanceIndex)
 		, SegmentIndex(InSegmentIndex)
 		, Pipeline(InPipeline)
 		, HitGroupIndex(InHitGroupIndex)
-		, ResourceBindings(InResourceBindings)
-	{}
+		, NumUniformBuffers(InNumUniformBuffers)
+		, UniformBuffers(InUniformBuffers)
+	{
+	}
 
 	RHI_API void Execute(FRHICommandListBase& CmdList);
 };
@@ -3015,15 +3018,25 @@ public:
 	FORCEINLINE_DEBUGGABLE void SetRayTracingHitGroup(
 		FRayTracingSceneRHIParamRef Scene, uint32 InstanceIndex, uint32 SegmentIndex,
 		FRayTracingPipelineStateRHIParamRef Pipeline, uint32 HitGroupIndex,
-		const FRayTracingShaderBindings& ResourceBindings)
+		uint32 NumUniformBuffers, const FUniformBufferRHIParamRef* UniformBuffers)
 	{
 		if (Bypass())
 		{
-			GetContext().RHISetRayTracingHitGroup(Scene, InstanceIndex, SegmentIndex, Pipeline, HitGroupIndex, ResourceBindings);
+			GetContext().RHISetRayTracingHitGroup(Scene, InstanceIndex, SegmentIndex, Pipeline, HitGroupIndex, NumUniformBuffers, UniformBuffers);
 		}
 		else
 		{
-			ALLOC_COMMAND(FRHICommandSetRayTracingHitGroup)(Scene, InstanceIndex, SegmentIndex, Pipeline, HitGroupIndex, ResourceBindings);
+			FUniformBufferRHIParamRef* InlineUniformBuffers = nullptr;
+			if (NumUniformBuffers)
+			{
+				InlineUniformBuffers = (FUniformBufferRHIParamRef*)Alloc(sizeof(FUniformBufferRHIParamRef) * NumUniformBuffers, alignof(FUniformBufferRHIParamRef));
+				for (uint32 Index = 0; Index < NumUniformBuffers; ++Index)
+				{
+					InlineUniformBuffers[Index] = UniformBuffers[Index];
+				}
+			}
+
+			ALLOC_COMMAND(FRHICommandSetRayTracingHitGroup)(Scene, InstanceIndex, SegmentIndex, Pipeline, HitGroupIndex, NumUniformBuffers, InlineUniformBuffers);
 		}
 	}
 
diff --git a/Engine/Source/Runtime/RHI/Public/RHICommandListCommandExecutes.inl b/Engine/Source/Runtime/RHI/Public/RHICommandListCommandExecutes.inl
--- a/Engine/Source/Runtime/RHI/Public/RHICommandListCommandExecutes.inl
+++ b/Engine/Source/Runtime/RHI/Public/RHICommandListCommandExecutes.inl
@@ -599,7 +599,7 @@ void FRHICommandRayTraceDispatch::Execute(FRHICommandListBase& CmdList)
 void FRHICommandSetRayTracingHitGroup::Execute(FRHICommandListBase& CmdList)
 {
 	RHISTAT(SetRayTracingHitGroup);
-	INTERNAL_DECORATOR(RHISetRayTracingHitGroup)(Scene, InstanceIndex, SegmentIndex, Pipeline, HitGroupIndex, ResourceBindings);
+	INTERNAL_DECORATOR(RHISetRayTracingHitGroup)(Scene, InstanceIndex, SegmentIndex, Pipeline, HitGroupIndex, NumUniformBuffers, UniformBuffers);
 }
 
 #endif // RHI_RAYTRACING
diff --git a/Engine/Source/Runtime/RHI/Public/RHIContext.h b/Engine/Source/Runtime/RHI/Public/RHIContext.h
--- a/Engine/Source/Runtime/RHI/Public/RHIContext.h
+++ b/Engine/Source/Runtime/RHI/Public/RHIContext.h
@@ -651,7 +651,7 @@ public:
 	virtual void RHISetRayTracingHitGroup(
 		FRayTracingSceneRHIParamRef Scene, uint32 InstanceIndex, uint32 SegmentIndex,
 		FRayTracingPipelineStateRHIParamRef Pipeline, uint32 HitGroupIndex,
-		const FRayTracingShaderBindings& ResourceBindings)
+		uint32 NumUniformBuffers, const FUniformBufferRHIParamRef* UniformBuffers)
 	{
 		checkNoEntry();
 	}
diff --git a/Engine/Source/Runtime/Renderer/Private/MeshPassProcessor.cpp b/Engine/Source/Runtime/Renderer/Private/MeshPassProcessor.cpp
--- a/Engine/Source/Runtime/Renderer/Private/MeshPassProcessor.cpp
+++ b/Engine/Source/Runtime/Renderer/Private/MeshPassProcessor.cpp
@@ -258,56 +258,33 @@ void FMeshDrawShaderBindings::SetOnRayTracingStructure(FRHICommandList& RHICmdLi
 	check(ShaderLayouts.Num() == 1);
 
 	FReadOnlyMeshDrawSingleShaderBindings SingleShaderBindings(ShaderLayouts[0], GetData());
-	FRayTracingShaderBindings RayTracingShaderBindings = {};
+
+	FUniformBufferRHIParamRef LocalUniformBuffers[MAX_UNIFORM_BUFFERS_PER_SHADER_STAGE] = {};
 
 	const FUniformBufferRHIParamRef* RESTRICT UniformBufferBindings = SingleShaderBindings.GetUniformBufferStart();
 	const FShaderParameterInfo* RESTRICT UniformBufferParameters = SingleShaderBindings.ParameterMapInfo.UniformBuffers.GetData();
 	const int32 NumUniformBuffers = SingleShaderBindings.ParameterMapInfo.UniformBuffers.Num();
 
+	int32 MaxUniformBufferUsed = -1;
 	for (int32 UniformBufferIndex = 0; UniformBufferIndex < NumUniformBuffers; UniformBufferIndex++)
 	{
 		FShaderParameterInfo Parameter = UniformBufferParameters[UniformBufferIndex];
-		checkSlow(Parameter.BaseIndex <= ARRAY_COUNT(RayTracingShaderBindings.UniformBuffers));
+		checkSlow(Parameter.BaseIndex <= ARRAY_COUNT(LocalUniformBuffers));
 		FUniformBufferRHIParamRef UniformBuffer = UniformBufferBindings[UniformBufferIndex];
-		RayTracingShaderBindings.UniformBuffers[Parameter.BaseIndex] = UniformBuffer;
-	}
-
-	const FSamplerStateRHIParamRef* RESTRICT SamplerBindings = SingleShaderBindings.GetSamplerStart();
-	const FShaderParameterInfo* RESTRICT TextureSamplerParameters = SingleShaderBindings.ParameterMapInfo.TextureSamplers.GetData();
-	const int32 NumTextureSamplers = SingleShaderBindings.ParameterMapInfo.TextureSamplers.Num();
-
-	for (int32 SamplerIndex = 0; SamplerIndex < NumTextureSamplers; SamplerIndex++)
-	{
-		FShaderParameterInfo Parameter = TextureSamplerParameters[SamplerIndex];
-		checkSlow(Parameter.BaseIndex <= ARRAY_COUNT(RayTracingShaderBindings.Samplers));
-		FSamplerStateRHIParamRef Sampler = SamplerBindings[SamplerIndex];
-		RayTracingShaderBindings.Samplers[Parameter.BaseIndex] = Sampler;
-	}
-
-	const FShaderResourceViewRHIParamRef* RESTRICT SRVBindings = SingleShaderBindings.GetSRVStart();
-	const FShaderParameterInfo* RESTRICT SRVParameters = SingleShaderBindings.ParameterMapInfo.SRVs.GetData();
-	const int32 NumSRVs = SingleShaderBindings.ParameterMapInfo.SRVs.Num();
-
-	for (int32 SRVIndex = 0; SRVIndex < NumSRVs; SRVIndex++)
-	{
-		FShaderParameterInfo Parameter = SRVParameters[SRVIndex];
-		checkSlow(Parameter.BaseIndex <= ARRAY_COUNT(RayTracingShaderBindings.SRVs));
-		FShaderResourceViewRHIParamRef SRV = SRVBindings[SRVIndex];
-		RayTracingShaderBindings.SRVs[Parameter.BaseIndex] = SRV;
+		if (Parameter.BaseIndex <= ARRAY_COUNT(LocalUniformBuffers))
+		{
+			LocalUniformBuffers[Parameter.BaseIndex] = UniformBuffer;
+			MaxUniformBufferUsed = FMath::Max((int32)Parameter.BaseIndex, MaxUniformBufferUsed);
+		}
 	}
 
-	const FTextureRHIParamRef* RESTRICT TextureBindings = SingleShaderBindings.GetTextureStart();
-
-	for (int32 TextureIndex = 0; TextureIndex < NumSRVs; TextureIndex++)
-	{
-		FShaderParameterInfo Parameter = SRVParameters[TextureIndex];
-		checkSlow(Parameter.BaseIndex <= ARRAY_COUNT(RayTracingShaderBindings.Textures));
-		FTextureRHIParamRef Texture = TextureBindings[TextureIndex];
-		RayTracingShaderBindings.Textures[Parameter.BaseIndex] = Texture;
-	}
+	checkf(SingleShaderBindings.ParameterMapInfo.TextureSamplers.Num() == 0, TEXT("Texture sampler parameters are not supported for ray tracing. UniformBuffers must be used for all resource binding."));
+	checkf(SingleShaderBindings.ParameterMapInfo.SRVs.Num() == 0, TEXT("SRV parameters are not supported for ray tracing. UniformBuffers must be used for all resource binding."));
+	checkf(SingleShaderBindings.ParameterMapInfo.LooseParameterBuffers.Num() == 0, TEXT("Loose parameter buffers are not supported for ray tracing. UniformBuffers must be used for all resource binding."));
 
 	check(SegmentIndex < 0xFF);
-	RHICmdList.SetRayTracingHitGroup(Scene, InstanceIndex, SegmentIndex, PipelineState, HitGroupIndex, RayTracingShaderBindings);
+	uint32 NumUniformBuffersToSet = MaxUniformBufferUsed + 1;
+	RHICmdList.SetRayTracingHitGroup(Scene, InstanceIndex, SegmentIndex, PipelineState, HitGroupIndex, NumUniformBuffersToSet, LocalUniformBuffers);
 }
 #endif // RHI_RAYTRACING
 
