Fix for final-delete while in flight crash

diff --git a/Engine/Source/Runtime/D3D12RHI/Private/D3D12Commands.cpp b/Engine/Source/Runtime/D3D12RHI/Private/D3D12Commands.cpp
--- a/Engine/Source/Runtime/D3D12RHI/Private/D3D12Commands.cpp
+++ b/Engine/Source/Runtime/D3D12RHI/Private/D3D12Commands.cpp
@@ -402,14 +402,14 @@ void FD3D12CommandContext::RHIEnqueueStagedRead(FStagingBufferRHIParamRef Stagin
 
 	// Only get data from the first gpu for now.
 	FD3D12Device* StagingDevice = VertexBuffer->GetParentDevice();
-	StagingBuffer->StagedRead.SafeRelease();
-	VERIFYD3D12RESULT(GetParentDevice()->GetParentAdapter()->CreateBuffer(D3D12_HEAP_TYPE_READBACK, GetGPUMask(), GetGPUMask(), Offset + NumBytes, StagingBuffer->StagedRead.GetInitReference()));
+	StagingBuffer->SafeRelease();
+	VERIFYD3D12RESULT(GetParentDevice()->GetParentAdapter()->CreateBuffer(D3D12_HEAP_TYPE_READBACK, GetGPUMask(), GetGPUMask(), Offset + NumBytes, &StagingBuffer->StagedRead));
 	
 	{
 		FD3D12Resource* pSourceResource = VertexBuffer->ResourceLocation.GetResource();
 		D3D12_RESOURCE_DESC const& SourceBufferDesc = pSourceResource->GetDesc();
 
-		FD3D12Resource* pDestResource = StagingBuffer->StagedRead.GetReference();
+		FD3D12Resource* pDestResource = StagingBuffer->StagedRead;
 		D3D12_RESOURCE_DESC const& DestBufferDesc = pDestResource->GetDesc();
 
 		FD3D12DynamicRHI::TransitionResource(CommandListHandle, pSourceResource, D3D12_RESOURCE_STATE_COPY_SOURCE, 0);
diff --git a/Engine/Source/Runtime/D3D12RHI/Private/D3D12DirectCommandListManager.cpp b/Engine/Source/Runtime/D3D12RHI/Private/D3D12DirectCommandListManager.cpp
--- a/Engine/Source/Runtime/D3D12RHI/Private/D3D12DirectCommandListManager.cpp
+++ b/Engine/Source/Runtime/D3D12RHI/Private/D3D12DirectCommandListManager.cpp
@@ -43,7 +43,7 @@ void* FD3D12DynamicRHI::RHILockStagingBuffer(FStagingBufferRHIParamRef StagingBu
 	FD3D12StagingBuffer* StagingBuffer = FD3D12DynamicRHI::ResourceCast(StagingBufferRHI);
 	check(StagingBuffer);
 
-	FD3D12Resource* pResource = StagingBuffer->StagedRead.GetReference();
+	FD3D12Resource* pResource = StagingBuffer->StagedRead;
 	if (pResource)
 	{
 		D3D12_RANGE ReadRange;
@@ -62,7 +62,7 @@ void FD3D12DynamicRHI::RHIUnlockStagingBuffer(FStagingBufferRHIParamRef StagingB
 	FD3D12StagingBuffer* StagingBuffer = FD3D12DynamicRHI::ResourceCast(StagingBufferRHI);
 	check(StagingBuffer);
 
-	FD3D12Resource* pResource = StagingBuffer->StagedRead.GetReference();
+	FD3D12Resource* pResource = StagingBuffer->StagedRead;
 	if (pResource)
 	{
 		pResource->Unmap();
diff --git a/Engine/Source/Runtime/D3D12RHI/Private/D3D12Resources.cpp b/Engine/Source/Runtime/D3D12RHI/Private/D3D12Resources.cpp
--- a/Engine/Source/Runtime/D3D12RHI/Private/D3D12Resources.cpp
+++ b/Engine/Source/Runtime/D3D12RHI/Private/D3D12Resources.cpp
@@ -252,6 +252,11 @@ void FD3D12Resource::UpdateResidency(FD3D12CommandListHandle& CommandList)
 #endif
 }
 
+void FD3D12Resource::DeferDelete()
+{
+	GetParentDevice()->GetParentAdapter()->GetDeferredDeletionQueue().EnqueueResource(this);
+}
+
 /////////////////////////////////////////////////////////////////////
 //	FD3D12 Heap
 /////////////////////////////////////////////////////////////////////
@@ -500,7 +505,7 @@ void FD3D12ResourceLocation::ReleaseResource()
 		
 		if (UnderlyingResource->ShouldDeferDelete())
 		{
-			GetParentDevice()->GetParentAdapter()->GetDeferredDeletionQueue().EnqueueResource(UnderlyingResource);
+			UnderlyingResource->DeferDelete();
 		}
 		else
 		{
@@ -529,7 +534,7 @@ void FD3D12ResourceLocation::ReleaseResource()
 	{
 		if (UnderlyingResource->ShouldDeferDelete() && UnderlyingResource->GetRefCount() == 1)
 		{
-			GetParentDevice()->GetParentAdapter()->GetDeferredDeletionQueue().EnqueueResource(UnderlyingResource);
+			UnderlyingResource->DeferDelete();
 		}
 		else
 		{
@@ -542,7 +547,7 @@ void FD3D12ResourceLocation::ReleaseResource()
 		check(UnderlyingResource->GetRefCount() == 1);
 		if (UnderlyingResource->ShouldDeferDelete())
 		{
-			GetParentDevice()->GetParentAdapter()->GetDeferredDeletionQueue().EnqueueResource(UnderlyingResource);
+			UnderlyingResource->DeferDelete();
 		}
 		else
 		{
diff --git a/Engine/Source/Runtime/D3D12RHI/Public/D3D12Resources.h b/Engine/Source/Runtime/D3D12RHI/Public/D3D12Resources.h
--- a/Engine/Source/Runtime/D3D12RHI/Public/D3D12Resources.h
+++ b/Engine/Source/Runtime/D3D12RHI/Public/D3D12Resources.h
@@ -197,6 +197,8 @@ public:
 	}
 
 	inline bool ShouldDeferDelete() const { return bDeferDelete; }
+	void DeferDelete();
+
 	inline bool IsPlacedResource() const { return Heap.GetReference() != nullptr; }
 	inline FD3D12Heap* GetHeap() const { return Heap; };
 	inline bool IsDepthStencilResource() const { return bDepthStencil; }
@@ -949,9 +951,25 @@ class FD3D12StagingBuffer : public FRHIStagingBuffer
 public:
 	FD3D12StagingBuffer(FVertexBufferRHIRef InBuffer)
 		: FRHIStagingBuffer(InBuffer)
+		, StagedRead(nullptr)
 	{}
+	~FD3D12StagingBuffer()
+	{
+		if (StagedRead)
+		{
+			StagedRead->DeferDelete();
+		}
+	}
+	void SafeRelease()
+	{
+		if (StagedRead)
+		{
+			StagedRead->Release();
+			StagedRead = nullptr;
+		}
+	}
 
-	TRefCountPtr<FD3D12Resource> StagedRead;
+	FD3D12Resource* StagedRead;
 };
 
 class FD3D12GPUFence : public FRHIGPUFence
