Ray tracing acceleration structure view is now persistently allocated and renamed on the RHIT timeline if necessary. This fixes ImmediateFlush workaround in FDeferredShadingSceneRenderer::DispatchRayTracingWorldUpdates.
Removed RHIGetAccelerationStructureShaderResourceView() and added SRV member directly to the FRHIRayTracingScene.
DefaultShaderTable for ray tracing pipelines is now created / uploaded on demand, like the other shader tables. This fixes invalid concurrent access to d3d12 command list.

diff --git a/Engine/Source/Runtime/D3D12RHI/Private/D3D12RHIPrivate.h b/Engine/Source/Runtime/D3D12RHI/Private/D3D12RHIPrivate.h
--- a/Engine/Source/Runtime/D3D12RHI/Private/D3D12RHIPrivate.h
+++ b/Engine/Source/Runtime/D3D12RHI/Private/D3D12RHIPrivate.h
@@ -460,11 +460,7 @@ public:
 #if D3D12_RHI_RAYTRACING
 
 	virtual FRayTracingGeometryRHIRef RHICreateRayTracingGeometry(const FRayTracingGeometryInitializer& Initializer) final override;
-
 	virtual FRayTracingSceneRHIRef RHICreateRayTracingScene(const FRayTracingSceneInitializer& Initializer) final override;
-
-	virtual FShaderResourceViewRHIParamRef RHIGetAccelerationStructureShaderResourceView(FRayTracingSceneRHIParamRef AccelerationStructure) final override;
-
 	virtual FRayTracingShaderRHIRef RHICreateRayTracingShader(const TArray<uint8>& Code, EShaderFrequency ShaderFrequency) final override;
 	virtual FRayTracingPipelineStateRHIRef RHICreateRayTracingPipelineState(const FRayTracingPipelineStateInitializer& Initializer) final override;
 
diff --git a/Engine/Source/Runtime/D3D12RHI/Private/D3D12RayTracing.cpp b/Engine/Source/Runtime/D3D12RHI/Private/D3D12RayTracing.cpp
--- a/Engine/Source/Runtime/D3D12RHI/Private/D3D12RayTracing.cpp
+++ b/Engine/Source/Runtime/D3D12RHI/Private/D3D12RayTracing.cpp
@@ -656,6 +656,8 @@ public:
 
 	void CopyToGPU()
 	{
+		check(IsInRHIThread() || !IsRunningRHIInSeparateThread());
+
 		FD3D12Device* Device = GetParentDevice();
 
 		checkf(Data.Num(), TEXT("Shader table is expected to be initialized before copying to GPU."));
@@ -681,6 +683,7 @@ public:
 
 	D3D12_GPU_VIRTUAL_ADDRESS GetShaderTableAddress() const
 	{
+		checkf(!bIsDirty, TEXT("Shader table update is pending, therefore GPU address is not available. Use CopyToGPU() to upload data and acquire a valid GPU buffer address."));
 		return Buffer->ResourceLocation.GetGPUVirtualAddress();
 	}
 
@@ -703,6 +706,12 @@ public:
 			Desc.HitGroupTable.StrideInBytes = HitGroupRecordStride;
 			Desc.HitGroupTable.SizeInBytes = NumHitGroups * HitGroupRecordStride;
 		}
+		else
+		{
+			Desc.HitGroupTable.StartAddress = ShaderTableAddress + DefaultHitGroupShaderTableOffset;
+			Desc.HitGroupTable.StrideInBytes = 0; // Zero stride effectively disables SBT indexing
+			Desc.HitGroupTable.SizeInBytes = D3D12_RAYTRACING_SHADER_TABLE_BYTE_ALIGNMENT; // Minimal table with only one record
+		}
 
 		return Desc;
 	}
@@ -1059,15 +1068,6 @@ public:
 		DefaultShaderTable.SetRayGenIdentifier(0, RayGenShaderIdentifier);
 		DefaultShaderTable.SetMissIdentifier(0, MissShaderIdentifier);
 		DefaultShaderTable.SetDefaultHitGroupIdentifier(DefaultHitGroupIdentifier);
-
-		DefaultShaderTable.CopyToGPU();
-
-		DefaultDispatchDesc = DefaultShaderTable.GetDispatchRaysDesc();
-
-		// Default dispatch desc refers to exactly one hit shader
-		DefaultDispatchDesc.HitGroupTable.StartAddress = DefaultShaderTable.GetShaderTableAddress() + DefaultShaderTable.DefaultHitGroupShaderTableOffset;
-		DefaultDispatchDesc.HitGroupTable.StrideInBytes = 0; // Zero stride effectively disables SBT indexing
-		DefaultDispatchDesc.HitGroupTable.SizeInBytes = D3D12_RAYTRACING_SHADER_TABLE_BYTE_ALIGNMENT; // Minimal table with only one record
 	}
 
 	FD3D12ShaderIdentifier GetShaderIdentifier(LPCWSTR ExportName)
@@ -1111,7 +1111,6 @@ public:
 
 	TRefCountPtr<ID3D12StateObject> StateObject;
 	TRefCountPtr<ID3D12StateObjectProperties> PipelineProperties;
-	D3D12_DISPATCH_RAYS_DESC DefaultDispatchDesc = {};
 
 	static constexpr uint32 ShaderIdentifierSize = D3D12_SHADER_IDENTIFIER_SIZE_IN_BYTES;
 	TArray<FD3D12ShaderIdentifier> HitGroupShaderIdentifiers;
@@ -1287,11 +1286,6 @@ FRayTracingSceneRHIRef FD3D12DynamicRHI::RHICreateRayTracingScene(const FRayTrac
 	});
 }
 
-FShaderResourceViewRHIParamRef FD3D12DynamicRHI::RHIGetAccelerationStructureShaderResourceView(FRayTracingSceneRHIParamRef InAccelerationStructure)
-{
-	return FD3D12DynamicRHI::ResourceCast(InAccelerationStructure)->AccelerationStructureView;
-}
-
 void FD3D12RayTracingGeometry::TransitionBuffers(FD3D12CommandContext& CommandContext)
 {
 	// Transition vertex and index resources..
@@ -1307,6 +1301,8 @@ void FD3D12RayTracingGeometry::TransitionBuffers(FD3D12CommandContext& CommandCo
 
 static void CreateAccelerationStructureBuffers(TRefCountPtr<FD3D12MemBuffer>& AccelerationStructureBuffer, TRefCountPtr<FD3D12MemBuffer>&  ScratchBuffer, FD3D12Adapter* Adapter, uint32 GPUIndex, const D3D12_RAYTRACING_ACCELERATION_STRUCTURE_PREBUILD_INFO& PrebuildInfo)
 {
+	check(IsInRHIThread() || !IsRunningRHIInSeparateThread());
+
 	FRHIResourceCreateInfo CreateInfo;
 
 	D3D12_RESOURCE_DESC AccelerationStructureBufferDesc = CD3DX12_RESOURCE_DESC::Buffer(
@@ -1492,6 +1488,11 @@ void FD3D12RayTracingScene::BuildAccelerationStructure(FD3D12CommandContext& Com
 	// #dxr_todo: scratch buffers should be created in UAV state from the start
 	FD3D12DynamicRHI::TransitionResource(CommandContext.CommandListHandle, ScratchBuffer.GetReference()->GetResource(), D3D12_RESOURCE_STATE_UNORDERED_ACCESS, 0);
 
+	if (bAccelerationStructureViewInitialized)
+	{
+		AccelerationStructureView->Rename(AccelerationStructureBuffer->ResourceLocation);
+	}
+	else
 	{
 		D3D12_SHADER_RESOURCE_VIEW_DESC SRVDesc = {};
 		SRVDesc.Format = DXGI_FORMAT_UNKNOWN;
@@ -1499,9 +1500,9 @@ void FD3D12RayTracingScene::BuildAccelerationStructure(FD3D12CommandContext& Com
 		SRVDesc.Shader4ComponentMapping = D3D12_DEFAULT_SHADER_4_COMPONENT_MAPPING;
 		SRVDesc.RaytracingAccelerationStructure.Location = AccelerationStructureBuffer->ResourceLocation.GetGPUVirtualAddress();
 
-		AccelerationStructureView = new FD3D12ShaderResourceView(
-			AccelerationStructureBuffer->GetParentDevice(), 
-			SRVDesc, AccelerationStructureBuffer->ResourceLocation, 4);
+		AccelerationStructureView->Initialize(SRVDesc, AccelerationStructureBuffer->ResourceLocation, 4);
+
+		bAccelerationStructureViewInitialized = true;
 	}
 
 	// Create and fill instance buffer
@@ -2116,12 +2117,15 @@ void FD3D12CommandContext::RHIRayTraceOcclusion(FRayTracingSceneRHIParamRef InSc
 	FD3D12ShaderResourceView* Rays = FD3D12DynamicRHI::ResourceCast(InRays);
 	FD3D12UnorderedAccessView* Output = FD3D12DynamicRHI::ResourceCast(InOutput);
 
-	const FD3D12RayTracingPipelineState* Pipeline = GetParentDevice()->GetBasicRayTracingPipeline()->Occlusion;
-	const FD3D12RayTracingShaderTable& ShaderTable = Pipeline->DefaultShaderTable;
+	FD3D12RayTracingPipelineState* Pipeline = GetParentDevice()->GetBasicRayTracingPipeline()->Occlusion;
+	FD3D12RayTracingShaderTable& ShaderTable = Pipeline->DefaultShaderTable;
 
-	ID3D12GraphicsCommandList4* RayTracingCommandList = CommandListHandle.RayTracingCommandList();
+	if (ShaderTable.bIsDirty)
+	{
+		ShaderTable.CopyToGPU();
+	}
 
-	D3D12_GPU_VIRTUAL_ADDRESS ShaderTableAddress = ShaderTable.Buffer->ResourceLocation.GetGPUVirtualAddress();
+	ID3D12GraphicsCommandList4* RayTracingCommandList = CommandListHandle.RayTracingCommandList();
 
 	D3D12_DISPATCH_RAYS_DESC DispatchDesc = ShaderTable.GetDispatchRaysDesc();
 
@@ -2148,19 +2152,18 @@ void FD3D12CommandContext::RHIRayTraceIntersection(FRayTracingSceneRHIParamRef I
 	FD3D12ShaderResourceView* Rays = FD3D12DynamicRHI::ResourceCast(InRays);
 	FD3D12UnorderedAccessView* Output = FD3D12DynamicRHI::ResourceCast(InOutput);
 
-	const FD3D12RayTracingPipelineState* Pipeline = GetParentDevice()->GetBasicRayTracingPipeline()->Intersection;
-	const FD3D12RayTracingShaderTable& ShaderTable = Pipeline->DefaultShaderTable;
+	FD3D12RayTracingPipelineState* Pipeline = GetParentDevice()->GetBasicRayTracingPipeline()->Intersection;
+	FD3D12RayTracingShaderTable& ShaderTable = Pipeline->DefaultShaderTable;
 
-	ID3D12GraphicsCommandList4* RayTracingCommandList = CommandListHandle.RayTracingCommandList();
+	if (ShaderTable.bIsDirty)
+	{
+		ShaderTable.CopyToGPU();
+	}
 
-	D3D12_GPU_VIRTUAL_ADDRESS ShaderTableAddress = ShaderTable.Buffer->ResourceLocation.GetGPUVirtualAddress();
+	ID3D12GraphicsCommandList4* RayTracingCommandList = CommandListHandle.RayTracingCommandList();
 
 	D3D12_DISPATCH_RAYS_DESC DispatchDesc = ShaderTable.GetDispatchRaysDesc();
 
-	DispatchDesc.HitGroupTable.StartAddress = ShaderTableAddress + ShaderTable.DefaultHitGroupShaderTableOffset;
-	DispatchDesc.HitGroupTable.StrideInBytes = 0; // Zero stride effectively disables SBT indexing
-	DispatchDesc.HitGroupTable.SizeInBytes = D3D12_RAYTRACING_SHADER_TABLE_BYTE_ALIGNMENT; // Minimal table with only one record
-
 	DispatchDesc.Width = NumRays;
 	DispatchDesc.Height = 1;
 	DispatchDesc.Depth = 1;
@@ -2186,7 +2189,6 @@ void FD3D12CommandContext::RHIRayTraceDispatch(FRayTracingPipelineStateRHIParamR
 
 	FD3D12RayTracingShaderTable* ShaderTable = Scene->FindOrCreateShaderTable(Pipeline);
 
-	// #dxr_todo: shader table update should be an explicit operation and this code should only perform validation.
 	if (ShaderTable->bIsDirty)
 	{
 		ShaderTable->CopyToGPU();
@@ -2205,9 +2207,15 @@ void FD3D12CommandContext::RHIRayTraceDispatch(FRayTracingPipelineStateRHIParamR
 	const FRayTracingShaderBindings& GlobalResourceBindings,
 	uint32 Width, uint32 Height)
 {
-	const FD3D12RayTracingPipelineState* Pipeline = FD3D12DynamicRHI::ResourceCast(InRayTracingPipelineState);
+	FD3D12RayTracingPipelineState* Pipeline = FD3D12DynamicRHI::ResourceCast(InRayTracingPipelineState);
+	FD3D12RayTracingShaderTable& ShaderTable = Pipeline->DefaultShaderTable;
 
-	D3D12_DISPATCH_RAYS_DESC DispatchDesc = Pipeline->DefaultDispatchDesc;
+	if (ShaderTable.bIsDirty)
+	{
+		ShaderTable.CopyToGPU();
+	}
+
+	D3D12_DISPATCH_RAYS_DESC DispatchDesc = ShaderTable.GetDispatchRaysDesc();
 
 	DispatchDesc.Width = Width;
 	DispatchDesc.Height = Height;
diff --git a/Engine/Source/Runtime/D3D12RHI/Private/D3D12RayTracing.h b/Engine/Source/Runtime/D3D12RHI/Private/D3D12RayTracing.h
--- a/Engine/Source/Runtime/D3D12RHI/Private/D3D12RayTracing.h
+++ b/Engine/Source/Runtime/D3D12RHI/Private/D3D12RayTracing.h
@@ -47,7 +47,10 @@ public:
 
 	FD3D12RayTracingScene(FD3D12Device* Device) 
 		: FD3D12DeviceChild(Device)
-	{};
+		, AccelerationStructureView(new FD3D12ShaderResourceView(Device))
+	{
+		ShaderResourceView = AccelerationStructureView;
+	};
 
 	~FD3D12RayTracingScene();
 
@@ -56,6 +59,7 @@ public:
 	// #dxr_todo: only GPU resources should be managed as LinkedAdapterObjects
 	TRefCountPtr<FD3D12MemBuffer> AccelerationStructureBuffer;
 	TRefCountPtr<FD3D12ShaderResourceView> AccelerationStructureView;
+	bool bAccelerationStructureViewInitialized = false;
 
 	// #dxr_todo: add API to update geometries and instances
 	TArray<FRayTracingGeometryInstance> Instances;
diff --git a/Engine/Source/Runtime/RHI/Public/DynamicRHI.h b/Engine/Source/Runtime/RHI/Public/DynamicRHI.h
--- a/Engine/Source/Runtime/RHI/Public/DynamicRHI.h
+++ b/Engine/Source/Runtime/RHI/Public/DynamicRHI.h
@@ -1123,12 +1123,6 @@ public:
 		checkNoEntry();
 		return nullptr;
 	}
-
-	virtual FShaderResourceViewRHIParamRef RHIGetAccelerationStructureShaderResourceView(FRayTracingSceneRHIParamRef AccelerationStructure)
-	{
-		checkNoEntry();
-		return nullptr;
-	}
 #endif // RHI_RAYTRACING
 
 protected:
@@ -1327,11 +1321,6 @@ FORCEINLINE FRayTracingSceneRHIRef RHICreateRayTracingScene(const FRayTracingSce
 	return GDynamicRHI->RHICreateRayTracingScene(Initializer);
 }
 
-FORCEINLINE FShaderResourceViewRHIParamRef RHIGetAccelerationStructureShaderResourceView(FRayTracingSceneRHIParamRef AccelerationStructure)
-{
-	return GDynamicRHI->RHIGetAccelerationStructureShaderResourceView(AccelerationStructure);
-}
-
 FORCEINLINE FRayTracingShaderRHIRef RHICreateRayTracingShader(const TArray<uint8>& Code, EShaderFrequency ShaderFrequency)
 {
 	return GDynamicRHI->RHICreateRayTracingShader(Code, ShaderFrequency);
diff --git a/Engine/Source/Runtime/RHI/Public/RHIResources.h b/Engine/Source/Runtime/RHI/Public/RHIResources.h
--- a/Engine/Source/Runtime/RHI/Public/RHIResources.h
+++ b/Engine/Source/Runtime/RHI/Public/RHIResources.h
@@ -232,16 +232,6 @@ class FRHIGraphicsPipelineState : public FRHIResource {};
 class FRHIComputePipelineState : public FRHIResource {};
 class FRHIRayTracingPipelineState : public FRHIResource {};
 
-//
-// Ray tracing resources
-//
-
-/** Bottom level ray tracing acceleration structure (contains triangles). */
-class FRHIRayTracingGeometry : public FRHIResource {};
-
-/** Top level ray tracing acceleration structure (contains instances of meshes). */
-class FRHIRayTracingScene : public FRHIResource {};
-
 //
 // Buffers
 //
@@ -1010,7 +1000,6 @@ class FRHIUnorderedAccessView : public FRHIResource {};
 class FRHIShaderResourceView : public FRHIResource {};
 
 
-
 typedef FRHISamplerState*              FSamplerStateRHIParamRef;
 typedef TRefCountPtr<FRHISamplerState> FSamplerStateRHIRef;
 
@@ -1103,14 +1092,32 @@ typedef TRefCountPtr<FRHIShaderResourceView> FShaderResourceViewRHIRef;
 typedef FRHIGraphicsPipelineState*              FGraphicsPipelineStateRHIParamRef;
 typedef TRefCountPtr<FRHIGraphicsPipelineState> FGraphicsPipelineStateRHIRef;
 
+typedef FRHIRayTracingPipelineState*              FRayTracingPipelineStateRHIParamRef;
+typedef TRefCountPtr<FRHIRayTracingPipelineState> FRayTracingPipelineStateRHIRef;
+
+
+//
+// Ray tracing resources
+//
+
+/** Bottom level ray tracing acceleration structure (contains triangles). */
+class FRHIRayTracingGeometry : public FRHIResource {};
+
 typedef FRHIRayTracingGeometry*                  FRayTracingGeometryRHIParamRef;
 typedef TRefCountPtr<FRHIRayTracingGeometry>     FRayTracingGeometryRHIRef;
 
+/** Top level ray tracing acceleration structure (contains instances of meshes). */
+class FRHIRayTracingScene : public FRHIResource
+{
+public:
+	FShaderResourceViewRHIParamRef GetShaderResourceView() { return ShaderResourceView; }
+protected:
+	FShaderResourceViewRHIRef ShaderResourceView;
+};
+
 typedef FRHIRayTracingScene*                     FRayTracingSceneRHIParamRef;
 typedef TRefCountPtr<FRHIRayTracingScene>        FRayTracingSceneRHIRef;
 
-typedef FRHIRayTracingPipelineState*              FRayTracingPipelineStateRHIParamRef;
-typedef TRefCountPtr<FRHIRayTracingPipelineState> FRayTracingPipelineStateRHIRef;
 
 /* Generic staging buffer class used by FRHIGPUMemoryReadback
 * RHI specific staging buffers derive from this
diff --git a/Engine/Source/Runtime/Renderer/Private/DeferredShadingRenderer.cpp b/Engine/Source/Runtime/Renderer/Private/DeferredShadingRenderer.cpp
--- a/Engine/Source/Runtime/Renderer/Private/DeferredShadingRenderer.cpp
+++ b/Engine/Source/Runtime/Renderer/Private/DeferredShadingRenderer.cpp
@@ -890,10 +890,6 @@ bool FDeferredShadingSceneRenderer::DispatchRayTracingWorldUpdates(FRHICommandLi
 			Initializer.Instances = View.RayTracingGeometryInstances;
 			View.PerViewRayTracingScene.RayTracingSceneRHI = RHICreateRayTracingScene(Initializer);
 			RHICmdList.BuildAccelerationStructure(View.PerViewRayTracingScene.RayTracingSceneRHI);
-
-			// #dxr_todo: RHIGetAccelerationStructureShaderResourceView should return a handle that's safe to use immediately (which is internally multi-buffered and fenced)  to avoid the need for flush after BuildAccelerationStructure.
-			// FlushRHIThread is required here to ensure that acceleration structure SRV is created by the RHI before it is referenced on RenderThread.
-			RHICmdList.ImmediateFlush(EImmediateFlushType::FlushRHIThread);
 		}
 	}
 	return true;
diff --git a/Engine/Source/Runtime/Renderer/Private/PathTracing.cpp b/Engine/Source/Runtime/Renderer/Private/PathTracing.cpp
--- a/Engine/Source/Runtime/Renderer/Private/PathTracing.cpp
+++ b/Engine/Source/Runtime/Renderer/Private/PathTracing.cpp
@@ -130,7 +130,7 @@ public:
 		FUnorderedAccessViewRHIParamRef RayCountPerPixelUAV)
 	{
 
-		GlobalResources.Set(TLASParameter, RHIGetAccelerationStructureShaderResourceView(RayTracingScene.RayTracingSceneRHI));
+		GlobalResources.Set(TLASParameter, RayTracingScene.RayTracingSceneRHI->GetShaderResourceView());
 		GlobalResources.Set(ViewParameter, ViewUniformBuffer);
 
 		// Lights
diff --git a/Engine/Source/Runtime/Renderer/Private/RayTracing/RayTracingAmbientOcclusion.cpp b/Engine/Source/Runtime/Renderer/Private/RayTracing/RayTracingAmbientOcclusion.cpp
--- a/Engine/Source/Runtime/Renderer/Private/RayTracing/RayTracingAmbientOcclusion.cpp
+++ b/Engine/Source/Runtime/Renderer/Private/RayTracing/RayTracingAmbientOcclusion.cpp
@@ -109,7 +109,7 @@ public:
 		FRHIRayTracingPipelineState* Pipeline = PipelineStateCache::GetAndOrCreateRayTracingPipelineState(Initializer); // #dxr_todo: this should be done once at load-time and cached
 
 		FRayTracingShaderBindingsWriter GlobalResources;
-		GlobalResources.Set(TLASParameter, RHIGetAccelerationStructureShaderResourceView(RayTracingScene.RayTracingSceneRHI));
+		GlobalResources.Set(TLASParameter, RayTracingScene.RayTracingSceneRHI->GetShaderResourceView());
 		GlobalResources.Set(ViewParameter, ViewUniformBuffer);
 		GlobalResources.Set(SceneTexturesParameter, SceneTexturesUniformBuffer);
 		GlobalResources.Set(AmbientOcclusionParameter, AmbientOcclusionUniformBuffer);
diff --git a/Engine/Source/Runtime/Renderer/Private/RayTracing/RayTracingBarycentrics.cpp b/Engine/Source/Runtime/Renderer/Private/RayTracing/RayTracingBarycentrics.cpp
--- a/Engine/Source/Runtime/Renderer/Private/RayTracing/RayTracingBarycentrics.cpp
+++ b/Engine/Source/Runtime/Renderer/Private/RayTracing/RayTracingBarycentrics.cpp
@@ -94,7 +94,7 @@ void FDeferredShadingSceneRenderer::RenderRayTracedBarycentrics(FRHICommandListI
 
 	FRayTracingBarycentricsRGS::FParameters* RayGenParameters = GraphBuilder.AllocParameters<FRayTracingBarycentricsRGS::FParameters>();
 
-	RayGenParameters->TLAS = RHIGetAccelerationStructureShaderResourceView(RayTracingSceneRHI);
+	RayGenParameters->TLAS = RayTracingSceneRHI->GetShaderResourceView();
 	RayGenParameters->ViewUniformBuffer = View.ViewUniformBuffer;
 	RayGenParameters->Output = GraphBuilder.CreateUAV(GraphBuilder.RegisterExternalTexture(SceneContext.GetSceneColor()));
 
diff --git a/Engine/Source/Runtime/Renderer/Private/RayTracing/RayTracingDebug.cpp b/Engine/Source/Runtime/Renderer/Private/RayTracing/RayTracingDebug.cpp
--- a/Engine/Source/Runtime/Renderer/Private/RayTracing/RayTracingDebug.cpp
+++ b/Engine/Source/Runtime/Renderer/Private/RayTracing/RayTracingDebug.cpp
@@ -126,7 +126,7 @@ void FDeferredShadingSceneRenderer::RenderRayTracedDebug(FRHICommandListImmediat
 	FRayTracingDebugRGS::FParameters* RayGenParameters = GraphBuilder.AllocParameters<FRayTracingDebugRGS::FParameters>();
 
 	RayGenParameters->VisualizationMode = DebugVisualizationMode;
-	RayGenParameters->TLAS = RHIGetAccelerationStructureShaderResourceView(RayTracingSceneRHI);
+	RayGenParameters->TLAS = RayTracingSceneRHI->GetShaderResourceView();
 	RayGenParameters->ViewUniformBuffer = View.ViewUniformBuffer;
 	RayGenParameters->Output = GraphBuilder.CreateUAV(GraphBuilder.RegisterExternalTexture(SceneContext.GetSceneColor()));
 
diff --git a/Engine/Source/Runtime/Renderer/Private/RayTracing/RayTracingReflections.cpp b/Engine/Source/Runtime/Renderer/Private/RayTracing/RayTracingReflections.cpp
--- a/Engine/Source/Runtime/Renderer/Private/RayTracing/RayTracingReflections.cpp
+++ b/Engine/Source/Runtime/Renderer/Private/RayTracing/RayTracingReflections.cpp
@@ -216,7 +216,7 @@ void FDeferredShadingSceneRenderer::RayTraceReflections(
 	PassParameters->LTCAmpTexture = GSystemTextures.LTCAmp->GetRenderTargetItem().ShaderResourceTexture;
 	PassParameters->LTCAmpSampler = TStaticSamplerState<SF_Bilinear, AM_Clamp, AM_Clamp, AM_Clamp>::GetRHI();
 
-	PassParameters->TLAS = RHIGetAccelerationStructureShaderResourceView(View.PerViewRayTracingScene.RayTracingSceneRHI);
+	PassParameters->TLAS = View.PerViewRayTracingScene.RayTracingSceneRHI->GetShaderResourceView();
 	PassParameters->ViewUniformBuffer = View.ViewUniformBuffer;
 	{
 		FReflectionsLightData LightData;
diff --git a/Engine/Source/Runtime/Renderer/Private/RayTracing/RayTracingShadows.cpp b/Engine/Source/Runtime/Renderer/Private/RayTracing/RayTracingShadows.cpp
--- a/Engine/Source/Runtime/Renderer/Private/RayTracing/RayTracingShadows.cpp
+++ b/Engine/Source/Runtime/Renderer/Private/RayTracing/RayTracingShadows.cpp
@@ -109,7 +109,7 @@ void FDeferredShadingSceneRenderer::RenderRayTracingOcclusion(
 		PassParameters->RWRayDistanceUAV = GraphBuilder.CreateUAV(FRDGTextureUAVDesc(RayDistanceTexture));
 		PassParameters->SamplesPerPixel = GRayTracingOcclusionSamplesPerPixel;
 		LightSceneProxy->GetLightShaderParameters(PassParameters->Light);
-		PassParameters->TLAS = RHIGetAccelerationStructureShaderResourceView(View.PerViewRayTracingScene.RayTracingSceneRHI);
+		PassParameters->TLAS = View.PerViewRayTracingScene.RayTracingSceneRHI->GetShaderResourceView();
 		PassParameters->ViewUniformBuffer = View.ViewUniformBuffer;
 		PassParameters->SceneTexturesStruct = CreateUniformBufferImmediate(SceneTextures, EUniformBufferUsage::UniformBuffer_SingleDraw);
 
diff --git a/Engine/Source/Runtime/Renderer/Private/RayTracing/RayTracingShadowsForRectLight.cpp b/Engine/Source/Runtime/Renderer/Private/RayTracing/RayTracingShadowsForRectLight.cpp
--- a/Engine/Source/Runtime/Renderer/Private/RayTracing/RayTracingShadowsForRectLight.cpp
+++ b/Engine/Source/Runtime/Renderer/Private/RayTracing/RayTracingShadowsForRectLight.cpp
@@ -144,7 +144,7 @@ public:
 		FRHIRayTracingPipelineState* Pipeline = PipelineStateCache::GetAndOrCreateRayTracingPipelineState(Initializer); // #dxr_todo: this should be done once at load-time and cached
 
 		FRayTracingShaderBindingsWriter GlobalResources;
-		GlobalResources.Set(TLASParameter, RHIGetAccelerationStructureShaderResourceView(RayTracingScene.RayTracingSceneRHI));
+		GlobalResources.Set(TLASParameter, RayTracingScene.RayTracingSceneRHI->GetShaderResourceView());
 		GlobalResources.Set(ViewParameter, ViewUniformBuffer);
 		GlobalResources.Set(SceneTexturesParameter, SceneTexturesUniformBuffer);
 		GlobalResources.Set(RectLightParameter, RectLightUniformBuffer);
diff --git a/Engine/Source/Runtime/Renderer/Private/RayTracing/RaytracingSkylight.cpp b/Engine/Source/Runtime/Renderer/Private/RayTracing/RaytracingSkylight.cpp
--- a/Engine/Source/Runtime/Renderer/Private/RayTracing/RaytracingSkylight.cpp
+++ b/Engine/Source/Runtime/Renderer/Private/RayTracing/RaytracingSkylight.cpp
@@ -123,7 +123,7 @@ void FDeferredShadingSceneRenderer::RenderRayTracingSkyLight(
 	PassParameters->SkyLightColor = FVector(1.0);
 	//PassParameters->SkyLightTexture = Scene->SkyLight->ProcessedTexture->TextureRHI;
 	//PassParameters->SkyLightTextureSampler = Scene->SkyLight->ProcessedTexture->SamplerStateRHI;
-	PassParameters->TLAS = RHIGetAccelerationStructureShaderResourceView(View.PerViewRayTracingScene.RayTracingSceneRHI);
+	PassParameters->TLAS = View.PerViewRayTracingScene.RayTracingSceneRHI->GetShaderResourceView();
 	PassParameters->ViewUniformBuffer = View.ViewUniformBuffer;
 	PassParameters->SceneTexturesStruct = CreateUniformBufferImmediate(SceneTextures, EUniformBufferUsage::UniformBuffer_SingleDraw);
 
@@ -226,7 +226,7 @@ public:
 		FRHIRayTracingPipelineState* Pipeline = PipelineStateCache::GetAndOrCreateRayTracingPipelineState(Initializer); // #dxr_todo: this should be done once at load-time and cached
 
 		FRayTracingShaderBindingsWriter GlobalResources;
-		GlobalResources.Set(TLASParameter, RHIGetAccelerationStructureShaderResourceView(RayTracingScene.RayTracingSceneRHI));
+		GlobalResources.Set(TLASParameter, RayTracingScene.RayTracingSceneRHI->GetShaderResourceView());
 		GlobalResources.Set(ViewParameter, ViewUniformBuffer);
 		GlobalResources.Set(SceneTexturesParameter, SceneTexturesUniformBuffer);
 		GlobalResources.Set(SkyLightParameter, SkyLightUniformBuffer);
