Added UE_NONCOPYABLE macro to easily disable copy and move operations for a type

diff --git a/Engine/Source/Runtime/Core/Public/Misc/CoreMiscDefines.h b/Engine/Source/Runtime/Core/Public/Misc/CoreMiscDefines.h
--- a/Engine/Source/Runtime/Core/Public/Misc/CoreMiscDefines.h
+++ b/Engine/Source/Runtime/Core/Public/Misc/CoreMiscDefines.h
@@ -259,3 +259,22 @@ struct TStaticDeprecateExpression
 		static constexpr int condition(TStaticDeprecateExpression<false>) { return 1; } \
 	}; \
 	enum class PREPROCESSOR_JOIN(EDeprecationMsg_, __LINE__) { Value = PREPROCESSOR_JOIN(FDeprecationMsg_, __LINE__)::condition(TStaticDeprecateExpression<!!(bExpression)>()) }
+
+/**
+* Makes a type non-copyable and non-movable by deleting copy/move constructors and assignment/move operators.
+* The macro should be placed in the public section of the type for better compiler diagnostic messages.
+* Example usage:
+*
+*	class FMyClassName
+*	{
+*	public:
+*		UE_NONCOPYABLE(FMyClassName)
+*		FMyClassName() = default;
+*	};
+*/
+#define UE_NONCOPYABLE(TypeName) \
+	TypeName(TypeName&&) = delete; \
+	TypeName(const TypeName&) = delete; \
+	TypeName& operator=(const TypeName&) = delete; \
+	TypeName& operator=(TypeName&&) = delete;
+
diff --git a/Engine/Source/Runtime/D3D12RHI/Private/D3D12RayTracing.cpp b/Engine/Source/Runtime/D3D12RHI/Private/D3D12RayTracing.cpp
--- a/Engine/Source/Runtime/D3D12RHI/Private/D3D12RayTracing.cpp
+++ b/Engine/Source/Runtime/D3D12RHI/Private/D3D12RayTracing.cpp
@@ -68,13 +68,10 @@ static_assert(sizeof(FD3D12ShaderIdentifier) == D3D12_SHADER_IDENTIFIER_SIZE_IN_
 
 struct FDXILLibrary
 {
-	FDXILLibrary() = default;
+	// No copy assignment or move because FDXILLibrary points to internal struct memory
+	UE_NONCOPYABLE(FDXILLibrary)
 
-	//no copy assignment or move because FDXILLibrary points to internal struct memory
-	FDXILLibrary(FDXILLibrary&&) = delete;
-	FDXILLibrary(const FDXILLibrary&) = delete;
-	FDXILLibrary& operator=(const FDXILLibrary&) = delete;
-	FDXILLibrary& operator=(FDXILLibrary&&) = delete;
+	FDXILLibrary() = default;
 
 	void InitFromDXIL(const void* Bytecode, uint32 BytecodeLength, const LPCWSTR* InEntryNames, const LPCWSTR* InExportNames, uint32 NumEntryNames)
 	{
@@ -254,13 +251,10 @@ static TRefCountPtr<ID3D12StateObject> CreateRayTracingStateObject(
 // This would allow ray tracing code to sub-allocate heap blocks from the same global heap.
 class FD3D12RayTracingDescriptorHeapCache : FD3D12DeviceChild
 {
-	FD3D12RayTracingDescriptorHeapCache(FD3D12RayTracingDescriptorHeapCache&&) = delete;
-	FD3D12RayTracingDescriptorHeapCache(const FD3D12RayTracingDescriptorHeapCache&) = delete;
-	FD3D12RayTracingDescriptorHeapCache& operator=(const FD3D12RayTracingDescriptorHeapCache&) = delete;
-	FD3D12RayTracingDescriptorHeapCache& operator=(FD3D12RayTracingDescriptorHeapCache&&) = delete;
-
 public:
 
+	UE_NONCOPYABLE(FD3D12RayTracingDescriptorHeapCache)
+
 	struct Entry
 	{
 		ID3D12DescriptorHeap* Heap = nullptr;
@@ -355,10 +349,7 @@ public:
 
 struct FD3D12RayTracingDescriptorHeap : public FD3D12DeviceChild
 {
-	FD3D12RayTracingDescriptorHeap(FD3D12RayTracingDescriptorHeap&&) = delete;
-	FD3D12RayTracingDescriptorHeap(const FD3D12RayTracingDescriptorHeap&) = delete;
-	FD3D12RayTracingDescriptorHeap& operator=(const FD3D12RayTracingDescriptorHeap&) = delete;
-	FD3D12RayTracingDescriptorHeap& operator=(FD3D12RayTracingDescriptorHeap&&) = delete;
+	UE_NONCOPYABLE(FD3D12RayTracingDescriptorHeap)
 
 	FD3D12RayTracingDescriptorHeap(FD3D12Device* Device)
 		: FD3D12DeviceChild(Device)
@@ -437,13 +428,10 @@ struct FD3D12RayTracingDescriptorHeap : public FD3D12DeviceChild
 
 class FD3D12RayTracingDescriptorCache : public FD3D12DeviceChild
 {
-	FD3D12RayTracingDescriptorCache(FD3D12RayTracingDescriptorCache&&) = delete;
-	FD3D12RayTracingDescriptorCache(const FD3D12RayTracingDescriptorCache&) = delete;
-	FD3D12RayTracingDescriptorCache& operator=(const FD3D12RayTracingDescriptorCache&) = delete;
-	FD3D12RayTracingDescriptorCache& operator=(FD3D12RayTracingDescriptorCache&&) = delete;
-
 public:
 
+	UE_NONCOPYABLE(FD3D12RayTracingDescriptorCache)
+
 	FD3D12RayTracingDescriptorCache(FD3D12Device* Device)
 		: FD3D12DeviceChild(Device)
 		, ViewHeap(Device)
@@ -514,11 +502,6 @@ public:
 
 class FD3D12RayTracingShaderTable : public FD3D12DeviceChild
 {
-	FD3D12RayTracingShaderTable(FD3D12RayTracingShaderTable&&) = delete;
-	FD3D12RayTracingShaderTable(const FD3D12RayTracingShaderTable&) = delete;
-	FD3D12RayTracingShaderTable& operator=(const FD3D12RayTracingShaderTable&) = delete;
-	FD3D12RayTracingShaderTable& operator=(FD3D12RayTracingShaderTable&&) = delete;
-
 private:
 
 	void WriteData(uint32 WriteOffset, const void* InData, uint32 InDataSize)
@@ -547,6 +530,8 @@ private:
 
 public:
 
+	UE_NONCOPYABLE(FD3D12RayTracingShaderTable)
+
 	FD3D12RayTracingShaderTable(FD3D12Device* Device)
 		: FD3D12DeviceChild(Device)
 		, DescriptorCache(Device)
@@ -773,13 +758,10 @@ void FD3D12Device::DestroyRayTracingDescriptorCache()
 
 class FD3D12RayTracingPipelineState : public FRHIRayTracingPipelineState
 {
-	FD3D12RayTracingPipelineState(FD3D12RayTracingPipelineState&&) = delete;
-	FD3D12RayTracingPipelineState(const FD3D12RayTracingPipelineState&) = delete;
-	FD3D12RayTracingPipelineState& operator=(const FD3D12RayTracingPipelineState&) = delete;
-	FD3D12RayTracingPipelineState& operator=(FD3D12RayTracingPipelineState&&) = delete;
-
 public:
 
+	UE_NONCOPYABLE(FD3D12RayTracingPipelineState)
+
 	FD3D12RayTracingPipelineState(FD3D12Device* Device, const FRayTracingPipelineStateInitializer& Initializer)
 		: DefaultShaderTable(Device)
 		, DefaultLocalRootSignature(Device->GetParentAdapter())
@@ -1121,12 +1103,10 @@ public:
 
 class FD3D12BasicRayTracingPipeline
 {
-	FD3D12BasicRayTracingPipeline(FD3D12BasicRayTracingPipeline&&) = delete;
-	FD3D12BasicRayTracingPipeline(const FD3D12BasicRayTracingPipeline&) = delete;
-	FD3D12BasicRayTracingPipeline& operator=(const FD3D12BasicRayTracingPipeline&) = delete;
-	FD3D12BasicRayTracingPipeline& operator=(FD3D12BasicRayTracingPipeline&&) = delete;
-
 public:
+
+	UE_NONCOPYABLE(FD3D12BasicRayTracingPipeline)
+
 	FD3D12BasicRayTracingPipeline(FD3D12Device* Device)
 	{
 		FRayTracingPipelineStateInitializer OcclusionInitializer;
