DR - Switch depth buffer to 32 bit on PC D3D11 and D3D12 and allow switching back to 24bit using r.d3d12.Depth24Bit

diff --git a/Engine/Source/Developer/StandaloneRenderer/Private/Windows/D3D/SlateD3DRenderer.cpp b/Engine/Source/Developer/StandaloneRenderer/Private/Windows/D3D/SlateD3DRenderer.cpp
--- a/Engine/Source/Developer/StandaloneRenderer/Private/Windows/D3D/SlateD3DRenderer.cpp
+++ b/Engine/Source/Developer/StandaloneRenderer/Private/Windows/D3D/SlateD3DRenderer.cpp
@@ -466,11 +466,7 @@ void FSlateD3DRenderer::CreateDepthStencilBuffer( FSlateD3DViewport& Viewport )
 	DescDepth.Height = Viewport.ViewportInfo.Height;
 	DescDepth.MipLevels = 1;
 	DescDepth.ArraySize = 1;
-#if DEPTH_32_BIT_CONVERSION
-	DescDepth.Format = DXGI_FORMAT_D32_FLOAT_S8X24_UINT;
-#else
 	DescDepth.Format = DXGI_FORMAT_D24_UNORM_S8_UINT;
-#endif
 	DescDepth.SampleDesc.Count = 1;
 	DescDepth.SampleDesc.Quality = 0;
 	DescDepth.Usage = D3D11_USAGE_DEFAULT;
@@ -481,11 +477,7 @@ void FSlateD3DRenderer::CreateDepthStencilBuffer( FSlateD3DViewport& Viewport )
 	if (SUCCEEDED(Hr))
 	{
 		D3D11_DEPTH_STENCIL_VIEW_DESC DescDSV;
-#if DEPTH_32_BIT_CONVERSION
-		DescDSV.Format = DXGI_FORMAT_D32_FLOAT_S8X24_UINT;
-#else
 		DescDSV.Format = DXGI_FORMAT_D24_UNORM_S8_UINT;
-#endif
 		DescDSV.ViewDimension = D3D11_DSV_DIMENSION_TEXTURE2D;
 		DescDSV.Texture2D.MipSlice = 0;
 		DescDSV.Flags = 0;
diff --git a/Engine/Source/Programs/UnrealBuildTool/Platform/Windows/UEBuildWindows.cs b/Engine/Source/Programs/UnrealBuildTool/Platform/Windows/UEBuildWindows.cs
--- a/Engine/Source/Programs/UnrealBuildTool/Platform/Windows/UEBuildWindows.cs
+++ b/Engine/Source/Programs/UnrealBuildTool/Platform/Windows/UEBuildWindows.cs
@@ -1522,8 +1522,6 @@ namespace UnrealBuildTool
 			CompileEnvironment.Definitions.Add(String.Format("WINVER=0x{0:X4}", Target.WindowsPlatform.TargetWindowsVersion));
 			CompileEnvironment.Definitions.Add("PLATFORM_WINDOWS=1");
 
-			CompileEnvironment.Definitions.Add("DEPTH_32_BIT_CONVERSION=0");
-
 			FileReference MorpheusShaderPath = FileReference.Combine(UnrealBuildTool.EngineDirectory, "Shaders", "Private", "Platform", "PS4", "PostProcessHMDMorpheus.usf");
 			if (FileReference.Exists(MorpheusShaderPath))
 			{
diff --git a/Engine/Source/Runtime/D3D12RHI/Private/D3D12RHI.cpp b/Engine/Source/Runtime/D3D12RHI/Private/D3D12RHI.cpp
--- a/Engine/Source/Runtime/D3D12RHI/Private/D3D12RHI.cpp
+++ b/Engine/Source/Runtime/D3D12RHI/Private/D3D12RHI.cpp
@@ -20,6 +20,15 @@
 
 DEFINE_LOG_CATEGORY(LogD3D12RHI);
 
+
+static TAutoConsoleVariable<int32> CVarD3D12UseD24(
+	TEXT("r.d3d12.Depth24Bit"),
+	0,
+	TEXT("0: Use 32bit float(default)\n1: Use 24bit fixed point"),
+	ECVF_ReadOnly
+);
+
+
 TAutoConsoleVariable<int32> CVarD3D12ZeroBufferSizeInMB(
 	TEXT("D3D12.ZeroBufferSizeInMB"),
 	4,
@@ -83,19 +92,22 @@ FD3D12DynamicRHI::FD3D12DynamicRHI(TArray<FD3D12Adapter*>& ChosenAdaptersIn) :
 	GPixelFormats[PF_DXT5			].PlatformFormat = DXGI_FORMAT_BC3_TYPELESS;
 	GPixelFormats[PF_BC4			].PlatformFormat = DXGI_FORMAT_BC4_UNORM;
 	GPixelFormats[PF_UYVY			].PlatformFormat = DXGI_FORMAT_UNKNOWN;		// TODO: Not supported in D3D11
-#if DEPTH_32_BIT_CONVERSION
-	GPixelFormats[PF_DepthStencil	].PlatformFormat = DXGI_FORMAT_R32G8X24_TYPELESS;
-	GPixelFormats[PF_DepthStencil	].BlockBytes = 5;
-	GPixelFormats[PF_DepthStencil	].Supported = true;
-	GPixelFormats[PF_X24_G8			].PlatformFormat = DXGI_FORMAT_X32_TYPELESS_G8X24_UINT;
-	GPixelFormats[PF_X24_G8			].BlockBytes = 5;
-#else
-	GPixelFormats[PF_DepthStencil	].PlatformFormat = DXGI_FORMAT_R24G8_TYPELESS;
-	GPixelFormats[PF_DepthStencil	].BlockBytes = 4;
-	GPixelFormats[PF_DepthStencil	].Supported = true;
-	GPixelFormats[PF_X24_G8			].PlatformFormat = DXGI_FORMAT_X24_TYPELESS_G8_UINT;
-	GPixelFormats[PF_X24_G8			].BlockBytes = 4;
-#endif
+	if (CVarD3D12UseD24.GetValueOnAnyThread())
+	{
+		GPixelFormats[PF_DepthStencil].PlatformFormat = DXGI_FORMAT_R24G8_TYPELESS;
+		GPixelFormats[PF_DepthStencil].BlockBytes = 4;
+		GPixelFormats[PF_DepthStencil].Supported = true;
+		GPixelFormats[PF_X24_G8].PlatformFormat = DXGI_FORMAT_X24_TYPELESS_G8_UINT;
+		GPixelFormats[PF_X24_G8].BlockBytes = 4;
+	}
+	else
+	{
+		GPixelFormats[PF_DepthStencil].PlatformFormat = DXGI_FORMAT_R32G8X24_TYPELESS;
+		GPixelFormats[PF_DepthStencil].BlockBytes = 5;
+		GPixelFormats[PF_DepthStencil].Supported = true;
+		GPixelFormats[PF_X24_G8].PlatformFormat = DXGI_FORMAT_X32_TYPELESS_G8X24_UINT;
+		GPixelFormats[PF_X24_G8].BlockBytes = 5;
+	}
 	GPixelFormats[PF_ShadowDepth	].PlatformFormat = DXGI_FORMAT_R16_TYPELESS;
 	GPixelFormats[PF_ShadowDepth	].BlockBytes = 2;
 	GPixelFormats[PF_ShadowDepth	].Supported = true;
diff --git a/Engine/Source/Runtime/D3D12RHI/Private/D3D12RenderTarget.cpp b/Engine/Source/Runtime/D3D12RHI/Private/D3D12RenderTarget.cpp
--- a/Engine/Source/Runtime/D3D12RHI/Private/D3D12RenderTarget.cpp
+++ b/Engine/Source/Runtime/D3D12RHI/Private/D3D12RenderTarget.cpp
@@ -675,12 +675,10 @@ static uint32 ComputeBytesPerPixel(DXGI_FORMAT Format)
 	case DXGI_FORMAT_R16G16B16A16_UNORM:
 		BytesPerPixel = 8;
 		break;
-#if DEPTH_32_BIT_CONVERSION
 		// Changing Depth Buffers to 32 bit on Dingo as D24S8 is actually implemented as a 32 bit buffer in the hardware
 	case DXGI_FORMAT_R32G8X24_TYPELESS:
 	case DXGI_FORMAT_D32_FLOAT_S8X24_UINT:
 	case DXGI_FORMAT_R32_FLOAT_X8X24_TYPELESS:
-#endif
 		BytesPerPixel = 5;
 		break;
 	case DXGI_FORMAT_R32G32B32A32_FLOAT:
@@ -1074,7 +1072,6 @@ static void ConvertRAWSurfaceDataToFColor(DXGI_FORMAT Format, uint32 Width, uint
 			}
 		}
 	}
-#if DEPTH_32_BIT_CONVERSION
 	// Changing Depth Buffers to 32 bit on Dingo as D24S8 is actually implemented as a 32 bit buffer in the hardware
 	else if (Format == DXGI_FORMAT_R32G8X24_TYPELESS)
 	{
@@ -1097,7 +1094,6 @@ static void ConvertRAWSurfaceDataToFColor(DXGI_FORMAT Format, uint32 Width, uint
 			}
 		}
 	}
-#endif
 	else if (Format == DXGI_FORMAT_R16G16B16A16_UNORM)
 	{
 		// Read the data out of the buffer, converting it to FColor.
diff --git a/Engine/Source/Runtime/D3D12RHI/Private/D3D12Util.cpp b/Engine/Source/Runtime/D3D12RHI/Private/D3D12Util.cpp
--- a/Engine/Source/Runtime/D3D12RHI/Private/D3D12Util.cpp
+++ b/Engine/Source/Runtime/D3D12RHI/Private/D3D12Util.cpp
@@ -127,10 +127,8 @@ namespace D3D12RHI
 				D3DFORMATCASE(DXGI_FORMAT_R32G32B32A32_FLOAT)
 				D3DFORMATCASE(DXGI_FORMAT_UNKNOWN)
 				D3DFORMATCASE(DXGI_FORMAT_R8_UNORM)
-#if DEPTH_32_BIT_CONVERSION
 				D3DFORMATCASE(DXGI_FORMAT_D32_FLOAT_S8X24_UINT)
 				D3DFORMATCASE(DXGI_FORMAT_R32_FLOAT_X8X24_TYPELESS)
-#endif
 				D3DFORMATCASE(DXGI_FORMAT_R32G8X24_TYPELESS)
 				D3DFORMATCASE(DXGI_FORMAT_D24_UNORM_S8_UINT)
 				D3DFORMATCASE(DXGI_FORMAT_R24_UNORM_X8_TYPELESS)
diff --git a/Engine/Source/Runtime/D3D12RHI/Public/D3D12Util.h b/Engine/Source/Runtime/D3D12RHI/Public/D3D12Util.h
--- a/Engine/Source/Runtime/D3D12RHI/Public/D3D12Util.h
+++ b/Engine/Source/Runtime/D3D12RHI/Public/D3D12Util.h
@@ -832,10 +832,8 @@ inline DXGI_FORMAT FindShaderResourceDXGIFormat(DXGI_FORMAT InFormat, bool bSRGB
 	case DXGI_FORMAT_R24G8_TYPELESS: return DXGI_FORMAT_R24_UNORM_X8_TYPELESS;
 	case DXGI_FORMAT_R32_TYPELESS: return DXGI_FORMAT_R32_FLOAT;
 	case DXGI_FORMAT_R16_TYPELESS: return DXGI_FORMAT_R16_UNORM;
-#if DEPTH_32_BIT_CONVERSION
 		// Changing Depth Buffers to 32 bit on Dingo as D24S8 is actually implemented as a 32 bit buffer in the hardware
 	case DXGI_FORMAT_R32G8X24_TYPELESS: return DXGI_FORMAT_R32_FLOAT_X8X24_TYPELESS;
-#endif
 	}
 	return InFormat;
 }
@@ -858,11 +856,9 @@ inline DXGI_FORMAT FindDepthStencilDXGIFormat(DXGI_FORMAT InFormat)
 	{
 	case DXGI_FORMAT_R24G8_TYPELESS:
 		return DXGI_FORMAT_D24_UNORM_S8_UINT;
-#if DEPTH_32_BIT_CONVERSION
 		// Changing Depth Buffers to 32 bit on Dingo as D24S8 is actually implemented as a 32 bit buffer in the hardware
 	case DXGI_FORMAT_R32G8X24_TYPELESS:
 		return DXGI_FORMAT_D32_FLOAT_S8X24_UINT;
-#endif
 	case DXGI_FORMAT_R32_TYPELESS:
 		return DXGI_FORMAT_D32_FLOAT;
 	case DXGI_FORMAT_R16_TYPELESS:
@@ -880,10 +876,8 @@ inline bool HasStencilBits(DXGI_FORMAT InFormat)
 	switch (InFormat)
 	{
 	case DXGI_FORMAT_D24_UNORM_S8_UINT:
-#if  DEPTH_32_BIT_CONVERSION
 		// Changing Depth Buffers to 32 bit on Dingo as D24S8 is actually implemented as a 32 bit buffer in the hardware
 	case DXGI_FORMAT_D32_FLOAT_S8X24_UINT:
-#endif
 		return true;
 	};
 
@@ -1039,12 +1033,10 @@ inline DXGI_FORMAT FindDepthStencilParentDXGIFormat(DXGI_FORMAT InFormat)
 	case DXGI_FORMAT_D24_UNORM_S8_UINT:
 	case DXGI_FORMAT_X24_TYPELESS_G8_UINT:
 		return DXGI_FORMAT_R24G8_TYPELESS;
-#if DEPTH_32_BIT_CONVERSION
 		// Changing Depth Buffers to 32 bit on Dingo as D24S8 is actually implemented as a 32 bit buffer in the hardware
 	case DXGI_FORMAT_D32_FLOAT_S8X24_UINT:
 	case DXGI_FORMAT_X32_TYPELESS_G8X24_UINT:
 		return DXGI_FORMAT_R32G8X24_TYPELESS;
-#endif
 	case DXGI_FORMAT_D32_FLOAT:
 		return DXGI_FORMAT_R32_TYPELESS;
 	case DXGI_FORMAT_D16_UNORM:
diff --git a/Engine/Source/Runtime/Windows/D3D11RHI/Private/D3D11Device.cpp b/Engine/Source/Runtime/Windows/D3D11RHI/Private/D3D11Device.cpp
--- a/Engine/Source/Runtime/Windows/D3D11RHI/Private/D3D11Device.cpp
+++ b/Engine/Source/Runtime/Windows/D3D11RHI/Private/D3D11Device.cpp
@@ -31,6 +31,15 @@ bool D3D11RHI_ShouldAllowAsyncResourceCreation()
 
 IMPLEMENT_MODULE(FD3D11DynamicRHIModule, D3D11RHI);
 
+static TAutoConsoleVariable<int32> CVarD3D11UseD24(
+	TEXT("r.d3d11.Depth24Bit"),
+	0,
+	TEXT("0: Use 32bit float(default)\n1: Use 24bit fixed point"),
+	ECVF_ReadOnly
+);
+
+
+
 TAutoConsoleVariable<int32> CVarD3D11ZeroBufferSizeInMB(
 	TEXT("d3d11.ZeroBufferSizeInMB"),
 	4,
@@ -127,17 +136,20 @@ FD3D11DynamicRHI::FD3D11DynamicRHI(IDXGIFactory1* InDXGIFactory1,D3D_FEATURE_LEV
 	GPixelFormats[ PF_DXT5			].PlatformFormat	= DXGI_FORMAT_BC3_TYPELESS;
 	GPixelFormats[ PF_BC4			].PlatformFormat	= DXGI_FORMAT_BC4_UNORM;
 	GPixelFormats[ PF_UYVY			].PlatformFormat	= DXGI_FORMAT_UNKNOWN;		// TODO: Not supported in D3D11
-#if DEPTH_32_BIT_CONVERSION
-	GPixelFormats[ PF_DepthStencil	].PlatformFormat	= DXGI_FORMAT_R32G8X24_TYPELESS; 
-	GPixelFormats[ PF_DepthStencil	].BlockBytes		= 5;
-	GPixelFormats[ PF_X24_G8 ].PlatformFormat			= DXGI_FORMAT_X32_TYPELESS_G8X24_UINT;
-	GPixelFormats[ PF_X24_G8].BlockBytes				= 5;
-#else
-	GPixelFormats[ PF_DepthStencil	].PlatformFormat	= DXGI_FORMAT_R24G8_TYPELESS;
-	GPixelFormats[ PF_DepthStencil	].BlockBytes		= 4;
-	GPixelFormats[ PF_X24_G8 ].PlatformFormat			= DXGI_FORMAT_X24_TYPELESS_G8_UINT;
-	GPixelFormats[ PF_X24_G8].BlockBytes				= 4;
-#endif
+	if (CVarD3D11UseD24.GetValueOnAnyThread())
+	{
+		GPixelFormats[PF_DepthStencil].PlatformFormat = DXGI_FORMAT_R24G8_TYPELESS;
+		GPixelFormats[PF_DepthStencil].BlockBytes = 4;
+		GPixelFormats[PF_X24_G8].PlatformFormat = DXGI_FORMAT_X24_TYPELESS_G8_UINT;
+		GPixelFormats[PF_X24_G8].BlockBytes = 4;
+	}
+	else
+	{
+		GPixelFormats[PF_DepthStencil].PlatformFormat = DXGI_FORMAT_R32G8X24_TYPELESS;
+		GPixelFormats[PF_DepthStencil].BlockBytes = 5;
+		GPixelFormats[PF_X24_G8].PlatformFormat = DXGI_FORMAT_X32_TYPELESS_G8X24_UINT;
+		GPixelFormats[PF_X24_G8].BlockBytes = 5;
+	}
 	GPixelFormats[ PF_ShadowDepth	].PlatformFormat	= DXGI_FORMAT_R16_TYPELESS;
 	GPixelFormats[ PF_ShadowDepth	].BlockBytes		= 2;
 	GPixelFormats[ PF_R32_FLOAT		].PlatformFormat	= DXGI_FORMAT_R32_FLOAT;
diff --git a/Engine/Source/Runtime/Windows/D3D11RHI/Private/D3D11RHIPrivate.h b/Engine/Source/Runtime/Windows/D3D11RHI/Private/D3D11RHIPrivate.h
--- a/Engine/Source/Runtime/Windows/D3D11RHI/Private/D3D11RHIPrivate.h
+++ b/Engine/Source/Runtime/Windows/D3D11RHI/Private/D3D11RHIPrivate.h
@@ -1060,10 +1060,8 @@ inline DXGI_FORMAT FindShaderResourceDXGIFormat(DXGI_FORMAT InFormat,bool bSRGB)
 		case DXGI_FORMAT_R24G8_TYPELESS: return DXGI_FORMAT_R24_UNORM_X8_TYPELESS;
 		case DXGI_FORMAT_R32_TYPELESS: return DXGI_FORMAT_R32_FLOAT;
 		case DXGI_FORMAT_R16_TYPELESS: return DXGI_FORMAT_R16_UNORM;
-#if DEPTH_32_BIT_CONVERSION
 		// Changing Depth Buffers to 32 bit on Dingo as D24S8 is actually implemented as a 32 bit buffer in the hardware
 		case DXGI_FORMAT_R32G8X24_TYPELESS: return DXGI_FORMAT_R32_FLOAT_X8X24_TYPELESS; 
-#endif
 	}
 	return InFormat;
 }
@@ -1086,11 +1084,9 @@ inline DXGI_FORMAT FindDepthStencilDXGIFormat(DXGI_FORMAT InFormat)
 	{
 		case DXGI_FORMAT_R24G8_TYPELESS:
 			return DXGI_FORMAT_D24_UNORM_S8_UINT;
-#if DEPTH_32_BIT_CONVERSION
 		// Changing Depth Buffers to 32 bit on Dingo as D24S8 is actually implemented as a 32 bit buffer in the hardware
 		case DXGI_FORMAT_R32G8X24_TYPELESS:
 			return DXGI_FORMAT_D32_FLOAT_S8X24_UINT;
-#endif
 		case DXGI_FORMAT_R32_TYPELESS:
 			return DXGI_FORMAT_D32_FLOAT;
 		case DXGI_FORMAT_R16_TYPELESS:
@@ -1109,11 +1105,9 @@ inline bool HasStencilBits(DXGI_FORMAT InFormat)
 	{
 	case DXGI_FORMAT_D24_UNORM_S8_UINT:
 		return true;
-#if  DEPTH_32_BIT_CONVERSION
 	// Changing Depth Buffers to 32 bit on Dingo as D24S8 is actually implemented as a 32 bit buffer in the hardware
 	case DXGI_FORMAT_D32_FLOAT_S8X24_UINT:
 		return true;
-#endif
 	};
 	return false;
 }
diff --git a/Engine/Source/Runtime/Windows/D3D11RHI/Private/D3D11RenderTarget.cpp b/Engine/Source/Runtime/Windows/D3D11RHI/Private/D3D11RenderTarget.cpp
--- a/Engine/Source/Runtime/Windows/D3D11RHI/Private/D3D11RenderTarget.cpp
+++ b/Engine/Source/Runtime/Windows/D3D11RHI/Private/D3D11RenderTarget.cpp
@@ -500,12 +500,10 @@ static uint32 ComputeBytesPerPixel(DXGI_FORMAT Format)
 		case DXGI_FORMAT_R16G16B16A16_UNORM:
 			BytesPerPixel = 8;
 			break;
-#if DEPTH_32_BIT_CONVERSION
 		// Changing Depth Buffers to 32 bit on Dingo as D24S8 is actually implemented as a 32 bit buffer in the hardware
 		case DXGI_FORMAT_R32G8X24_TYPELESS:
 		case DXGI_FORMAT_D32_FLOAT_S8X24_UINT:
 		case DXGI_FORMAT_R32_FLOAT_X8X24_TYPELESS:
-#endif
 			BytesPerPixel = 5;
 			break;
 		case DXGI_FORMAT_R32G32B32A32_FLOAT:
@@ -877,7 +875,6 @@ static void ConvertRAWSurfaceDataToFColor(DXGI_FORMAT Format, uint32 Width, uint
 			}
 		}
 	}
-#if DEPTH_32_BIT_CONVERSION
 	// Changing Depth Buffers to 32 bit on Dingo as D24S8 is actually implemented as a 32 bit buffer in the hardware
 	else if (Format == DXGI_FORMAT_R32G8X24_TYPELESS )
 	{
@@ -900,7 +897,6 @@ static void ConvertRAWSurfaceDataToFColor(DXGI_FORMAT Format, uint32 Width, uint
 			}
 		}
 	}
-#endif
 	else if(Format == DXGI_FORMAT_R16G16B16A16_UNORM)
 	{
 		// Read the data out of the buffer, converting it to FColor.
@@ -1440,7 +1436,6 @@ static void ConvertRAWSurfaceDataToFLinearColor(EPixelFormat Format, uint32 Widt
 			}
 		}
 	}
-#if DEPTH_32_BIT_CONVERSION
 	// Changing Depth Buffers to 32 bit on Dingo as D24S8 is actually implemented as a 32 bit buffer in the hardware
 	else if (Format == PF_DepthStencil)
 	{
@@ -1460,7 +1455,6 @@ static void ConvertRAWSurfaceDataToFLinearColor(EPixelFormat Format, uint32 Widt
 			}
 		}
 	}
-#endif
 	else if (Format == PF_A16B16G16R16)
 	{
 		// Read the data out of the buffer, converting it to FLinearColor.
diff --git a/Engine/Source/Runtime/Windows/D3D11RHI/Private/D3D11Util.cpp b/Engine/Source/Runtime/Windows/D3D11RHI/Private/D3D11Util.cpp
--- a/Engine/Source/Runtime/Windows/D3D11RHI/Private/D3D11Util.cpp
+++ b/Engine/Source/Runtime/Windows/D3D11RHI/Private/D3D11Util.cpp
@@ -96,10 +96,8 @@ const TCHAR* GetD3D11TextureFormatString(DXGI_FORMAT TextureFormat)
 		D3DFORMATCASE(DXGI_FORMAT_R32G32B32A32_FLOAT)
 		D3DFORMATCASE(DXGI_FORMAT_UNKNOWN)
 		D3DFORMATCASE(DXGI_FORMAT_R8_UNORM)
-#if DEPTH_32_BIT_CONVERSION
 		D3DFORMATCASE(DXGI_FORMAT_D32_FLOAT_S8X24_UINT)
 		D3DFORMATCASE(DXGI_FORMAT_R32_FLOAT_X8X24_TYPELESS)
-#endif
 		D3DFORMATCASE(DXGI_FORMAT_R32G8X24_TYPELESS)
 		D3DFORMATCASE(DXGI_FORMAT_D24_UNORM_S8_UINT)
 		D3DFORMATCASE(DXGI_FORMAT_R24_UNORM_X8_TYPELESS)
