More clang fixes.  Thanks Arne

diff --git a/Engine/Source/Runtime/D3D12RHI/Private/D3D12Allocation.h b/Engine/Source/Runtime/D3D12RHI/Private/D3D12Allocation.h
--- a/Engine/Source/Runtime/D3D12RHI/Private/D3D12Allocation.h
+++ b/Engine/Source/Runtime/D3D12RHI/Private/D3D12Allocation.h
@@ -403,7 +403,7 @@ private:
 
 		Count,
 	};
-	FD3D12DefaultBufferPool* DefaultBufferPools[EBufferPool::Count];
+	FD3D12DefaultBufferPool* DefaultBufferPools[(uint32)EBufferPool::Count];
 
 	inline EBufferPool GetBufferPool(D3D12_RESOURCE_FLAGS Flags) const
 	{
diff --git a/Engine/Source/Runtime/D3D12RHI/Private/D3D12RHICommon.h b/Engine/Source/Runtime/D3D12RHI/Private/D3D12RHICommon.h
--- a/Engine/Source/Runtime/D3D12RHI/Private/D3D12RHICommon.h
+++ b/Engine/Source/Runtime/D3D12RHI/Private/D3D12RHICommon.h
@@ -99,13 +99,13 @@ public:
 		, GPUIndex(GPUMask.ToIndex())
 	{}
 
-	FORCEINLINE const uint32 GetGPUIndex() const
+	FORCEINLINE uint32 GetGPUIndex() const
 	{
 		return GPUIndex;
 	}
 
 private:
-	const uint32 GPUIndex;
+	uint32 GPUIndex;
 };
 
 class FD3D12MultiNodeGPUObject : public FD3D12GPUObject
