Final clang fixes from Arne

diff --git a/Engine/Plugins/Runtime/WindowsMixedReality/Source/WindowsMixedRealityHMD/Private/WindowsMixedRealityHMD.cpp b/Engine/Plugins/Runtime/WindowsMixedReality/Source/WindowsMixedRealityHMD/Private/WindowsMixedRealityHMD.cpp
--- a/Engine/Plugins/Runtime/WindowsMixedReality/Source/WindowsMixedRealityHMD/Private/WindowsMixedRealityHMD.cpp
+++ b/Engine/Plugins/Runtime/WindowsMixedReality/Source/WindowsMixedRealityHMD/Private/WindowsMixedRealityHMD.cpp
@@ -9,6 +9,7 @@
 #include "Engine/Engine.h"
 #include "Interfaces/IPluginManager.h"
 #include "IWindowsMixedRealityHMDPlugin.h"
+#include "RHI/Public/PipelineStateCache.h"
 
 #if WITH_EDITOR
 #include "Editor/UnrealEd/Classes/Editor/EditorEngine.h"
diff --git a/Engine/Source/Runtime/D3D12RHI/Private/D3D12View.h b/Engine/Source/Runtime/D3D12RHI/Private/D3D12View.h
--- a/Engine/Source/Runtime/D3D12RHI/Private/D3D12View.h
+++ b/Engine/Source/Runtime/D3D12RHI/Private/D3D12View.h
@@ -620,6 +620,18 @@ inline uint32 CViewSubresourceSubset::MaxSubresource() const
 
 class FD3D12OfflineDescriptorManager;
 
+template<typename TDesc> 
+struct TIsD3D12SRVDescriptorHandle
+{
+	enum { Value = false };
+};
+
+template<>
+struct TIsD3D12SRVDescriptorHandle<D3D12_SHADER_RESOURCE_VIEW_DESC>
+{
+	enum { Value = true };
+};
+
 template <typename TDesc>
 class TD3D12ViewDescriptorHandle : public FD3D12DeviceChild
 {
@@ -656,7 +668,7 @@ public:
 	{
 #if D3D12_RHI_RAYTRACING
 		// NOTE (from D3D Debug runtime): When ViewDimension is D3D12_SRV_DIMENSION_RAYTRACING_ACCELLERATION_STRUCTURE, pResource must be NULL, since the resource location comes from a GPUVA in pDesc
-		if (Desc.ViewDimension == D3D12_SRV_DIMENSION_RAYTRACING_ACCELERATION_STRUCTURE)
+		if (TIsD3D12SRVDescriptorHandle<TDesc>::Value && ((int)Desc.ViewDimension == (int)D3D12_SRV_DIMENSION_RAYTRACING_ACCELERATION_STRUCTURE))
 		{
 			Resource = nullptr;
 		}
