Fix link error with Clang compiler

diff --git a/Engine/Source/Runtime/D3D12RHI/Private/D3D12Texture.cpp b/Engine/Source/Runtime/D3D12RHI/Private/D3D12Texture.cpp
--- a/Engine/Source/Runtime/D3D12RHI/Private/D3D12Texture.cpp
+++ b/Engine/Source/Runtime/D3D12RHI/Private/D3D12Texture.cpp
@@ -48,6 +48,11 @@ template TD3D12Texture2D<FD3D12BaseTextureCube>::TD3D12Texture2D(class FD3D12Dev
 #endif
 	);
 
+template TD3D12Texture2D<FD3D12BaseTexture2D>::~TD3D12Texture2D();
+template TD3D12Texture2D<FD3D12BaseTexture2DArray>::~TD3D12Texture2D();
+template TD3D12Texture2D<FD3D12BaseTextureCube>::~TD3D12Texture2D();
+
+
 /// @cond DOXYGEN_WARNINGS
 
 template void FD3D12TextureStats::D3D12TextureAllocated(TD3D12Texture2D<FD3D12BaseTexture2D>& Texture);
