Removed `-raytracing` command line argument that enables DXR. Only `r.RayTracing=1` CVar is required now, which controls the RHI and high-level features.
Added `-noraytracing` command line argument to easily opt out of all ray tracing features, even if they are enabled by the project.

diff --git a/Engine/Source/Runtime/D3D12RHI/Private/D3D12Adapter.cpp b/Engine/Source/Runtime/D3D12RHI/Private/D3D12Adapter.cpp
--- a/Engine/Source/Runtime/D3D12RHI/Private/D3D12Adapter.cpp
+++ b/Engine/Source/Runtime/D3D12RHI/Private/D3D12Adapter.cpp
@@ -158,7 +158,13 @@ void FD3D12Adapter::CreateRootDevice(bool bWithDebug)
 		}
 	}
 
- 	if (bRayTracingSupported && FParse::Param(FCommandLine::Get(), TEXT("raytracing"))) // Preview DXR functionality requires an explicit opt-in via a command line switch
+	auto GetRayTracingCVarValue = []()
+	{
+		auto CVar = IConsoleManager::Get().FindConsoleVariable(TEXT("r.RayTracing"));
+		return CVar && CVar->GetInt() > 0;
+	};
+
+ 	if (bRayTracingSupported && GetRayTracingCVarValue() && !FParse::Param(FCommandLine::Get(), TEXT("noraytracing")))
 	{
 		RootDevice->QueryInterface(IID_PPV_ARGS(RootRayTracingDevice.GetInitReference()));
 		if (RootRayTracingDevice)
