Undo //UE4/Private-Woodchuck-Staging/Engine/Source/Runtime/D3D12RHI/... changelist 4992708

diff --git a/Engine/Source/Runtime/D3D12RHI/Private/D3D12Allocation.cpp b/Engine/Source/Runtime/D3D12RHI/Private/D3D12Allocation.cpp
--- a/Engine/Source/Runtime/D3D12RHI/Private/D3D12Allocation.cpp
+++ b/Engine/Source/Runtime/D3D12RHI/Private/D3D12Allocation.cpp
@@ -1036,9 +1036,16 @@ FD3D12DefaultBufferAllocator::FD3D12DefaultBufferAllocator(FD3D12Device* InParen
 	FMemory::Memset(DefaultBufferPools, 0);
 }
 
-void FD3D12DefaultBufferAllocator::InitializeAllocator(EBufferPool PoolIndex, D3D12_RESOURCE_FLAGS Flags)
+// Grab a buffer from the available buffers or create a new buffer if none are available
+HRESULT FD3D12DefaultBufferAllocator::AllocDefaultResource(const D3D12_RESOURCE_DESC& Desc, uint32 InUsage, FD3D12ResourceLocation& ResourceLocation, uint32 Alignment)
 {
 	FD3D12Device* Device = GetParentDevice();
+
+	EBufferPool PoolIndex = GetBufferPool(Desc.Flags);
+	check(PoolIndex < EBufferPool::Count);
+
+	if (DefaultBufferPools[(uint32) PoolIndex] == nullptr)
+	{
 	FD3D12AllocatorType* Allocator = nullptr;
 
 #ifdef USE_BUCKET_ALLOCATOR
@@ -1059,7 +1066,7 @@ void FD3D12DefaultBufferAllocator::InitializeAllocator(EBufferPool PoolIndex, D3
 			kPlacedResourceStrategy,
 			D3D12_HEAP_TYPE_DEFAULT,
 			D3D12_HEAP_FLAG_ALLOW_ONLY_BUFFERS,
-			Flags,
+				Desc.Flags,
 			DEFAULT_BUFFER_POOL_MAX_ALLOC_SIZE,
 			ED3D12AllocatorID::DefaultBufferAllocator,
 			DEFAULT_BUFFER_POOL_SIZE,
@@ -1073,7 +1080,7 @@ void FD3D12DefaultBufferAllocator::InitializeAllocator(EBufferPool PoolIndex, D3
 			kManualSubAllocationStrategy,
 			D3D12_HEAP_TYPE_DEFAULT,
 			D3D12_HEAP_FLAG_ALLOW_ONLY_BUFFERS,
-			Flags,
+				Desc.Flags,
 			DEFAULT_BUFFER_POOL_MAX_ALLOC_SIZE,
 			ED3D12AllocatorID::DefaultBufferAllocator,
 			DEFAULT_BUFFER_POOL_SIZE,
@@ -1082,20 +1089,11 @@ void FD3D12DefaultBufferAllocator::InitializeAllocator(EBufferPool PoolIndex, D3
 #endif
 
 	DefaultBufferPools[(uint32) PoolIndex] = new FD3D12DefaultBufferPool(Device, Allocator);
-}
-
-// Grab a buffer from the available buffers or create a new buffer if none are available
-void FD3D12DefaultBufferAllocator::AllocDefaultResource(const D3D12_RESOURCE_DESC& Desc, uint32 InUsage, FD3D12ResourceLocation& ResourceLocation, uint32 Alignment)
-{
-	EBufferPool PoolIndex = GetBufferPool(Desc.Flags);
-	check(PoolIndex < EBufferPool::Count);
-
-	if (DefaultBufferPools[(uint32) PoolIndex] == nullptr)
-	{
-		InitializeAllocator(PoolIndex, Desc.Flags);
 	}
 
 	DefaultBufferPools[(uint32) PoolIndex]->AllocDefaultResource(Desc, InUsage, ResourceLocation, Alignment);
+
+	return S_OK;
 }
 
 void FD3D12DefaultBufferAllocator::FreeDefaultBufferPools()
diff --git a/Engine/Source/Runtime/D3D12RHI/Private/D3D12Allocation.h b/Engine/Source/Runtime/D3D12RHI/Private/D3D12Allocation.h
--- a/Engine/Source/Runtime/D3D12RHI/Private/D3D12Allocation.h
+++ b/Engine/Source/Runtime/D3D12RHI/Private/D3D12Allocation.h
@@ -389,7 +389,7 @@ public:
 	FD3D12DefaultBufferAllocator(FD3D12Device* InParent, FRHIGPUMask VisibleNodes);
 
 	// Grab a buffer from the available buffers or create a new buffer if none are available
-	void AllocDefaultResource(const D3D12_RESOURCE_DESC& pDesc, uint32 InUsage, FD3D12ResourceLocation& ResourceLocation, uint32 Alignment);
+	HRESULT AllocDefaultResource(const D3D12_RESOURCE_DESC& pDesc, uint32 InUsage, FD3D12ResourceLocation& ResourceLocation, uint32 Alignment);
 	void FreeDefaultBufferPools();
 	void CleanupFreeBlocks();
 
@@ -403,8 +403,6 @@ private:
 
 		Count,
 	};
-
-	void InitializeAllocator(EBufferPool PoolIndex, D3D12_RESOURCE_FLAGS Flags);
 	FD3D12DefaultBufferPool* DefaultBufferPools[(uint32)EBufferPool::Count];
 
 	inline EBufferPool GetBufferPool(D3D12_RESOURCE_FLAGS Flags) const
diff --git a/Engine/Source/Runtime/D3D12RHI/Private/D3D12Buffer.cpp b/Engine/Source/Runtime/D3D12RHI/Private/D3D12Buffer.cpp
--- a/Engine/Source/Runtime/D3D12RHI/Private/D3D12Buffer.cpp
+++ b/Engine/Source/Runtime/D3D12RHI/Private/D3D12Buffer.cpp
@@ -105,7 +105,7 @@ void FD3D12Adapter::AllocateBuffer(FD3D12Device* Device,
 	}
 	else
 	{
-		Device->GetDefaultBufferAllocator().AllocDefaultResource(InDesc, InUsage, ResourceLocation, Alignment);
+		VERIFYD3D12RESULT(Device->GetDefaultBufferAllocator().AllocDefaultResource(InDesc, InUsage, ResourceLocation, Alignment));
 		check(ResourceLocation.GetSize() == Size);
 	}
 }
diff --git a/Engine/Source/Runtime/D3D12RHI/Private/D3D12Resources.cpp b/Engine/Source/Runtime/D3D12RHI/Private/D3D12Resources.cpp
--- a/Engine/Source/Runtime/D3D12RHI/Private/D3D12Resources.cpp
+++ b/Engine/Source/Runtime/D3D12RHI/Private/D3D12Resources.cpp
@@ -299,7 +299,8 @@ HRESULT FD3D12Adapter::CreateCommittedResource(const D3D12_RESOURCE_DESC& InDesc
 
 	TRefCountPtr<ID3D12Resource> pResource;
 	const HRESULT hr = RootDevice->CreateCommittedResource(&HeapProps, D3D12_HEAP_FLAG_NONE, &InDesc, InitialUsage, ClearValue, IID_PPV_ARGS(pResource.GetInitReference()));
-	VERIFYD3D12RESULT_EX(hr, RootDevice);
+	checkf(SUCCEEDED(hr), TEXT("HR=0x%x"), hr);
+	check(SUCCEEDED(hr));
 
 	if (SUCCEEDED(hr))
 	{
@@ -328,7 +329,7 @@ HRESULT FD3D12Adapter::CreatePlacedResource(const D3D12_RESOURCE_DESC& InDesc, F
 
 	TRefCountPtr<ID3D12Resource> pResource;
 	const HRESULT hr = RootDevice->CreatePlacedResource(Heap, HeapOffset, &InDesc, InitialUsage, nullptr, IID_PPV_ARGS(pResource.GetInitReference()));
-	VERIFYD3D12RESULT_EX(hr, RootDevice);
+	check(SUCCEEDED(hr));
 
 	if (SUCCEEDED(hr))
 	{
diff --git a/Engine/Source/Runtime/D3D12RHI/Private/D3D12Stats.cpp b/Engine/Source/Runtime/D3D12RHI/Private/D3D12Stats.cpp
--- a/Engine/Source/Runtime/D3D12RHI/Private/D3D12Stats.cpp
+++ b/Engine/Source/Runtime/D3D12RHI/Private/D3D12Stats.cpp
@@ -443,9 +443,9 @@ bool FD3DGPUProfiler::CheckGpuHeartbeat() const
 
 					if (Result == GFSDK_Aftermath_Result_Success)
 					{
-						UE_LOG(LogRHI, Error, TEXT("[Aftermath] Faulting address: 0x%016llx"), FaultInformation.faultingGpuVA);
-						UE_LOG(LogRHI, Error, TEXT("[Aftermath] Faulting resource dims: %d x %d x %d"), FaultInformation.resourceDesc.width, FaultInformation.resourceDesc.height, FaultInformation.resourceDesc.depth);
-						UE_LOG(LogRHI, Error, TEXT("[Aftermath] Faulting result size: 0x%llu"), FaultInformation.resourceDesc.size);
+						UE_LOG(LogRHI, Error, TEXT("[Aftermath] Faulting address: 0x%xI64"), FaultInformation.faultingGpuVA);
+						UE_LOG(LogRHI, Error, TEXT("[Aftermath] Faulting resournce dims: %d x %d x %d"), FaultInformation.resourceDesc.width, FaultInformation.resourceDesc.height, FaultInformation.resourceDesc.depth);
+						UE_LOG(LogRHI, Error, TEXT("[Aftermath] Faulting result size: 0x%dI64"), FaultInformation.resourceDesc.size);
 						UE_LOG(LogRHI, Error, TEXT("[Aftermath] Faulting resource mips: %d"), FaultInformation.resourceDesc.mipLevels);
 						UE_LOG(LogRHI, Error, TEXT("[Aftermath] Faulting resource format: 0x%x"), FaultInformation.resourceDesc.format);
 					}
diff --git a/Engine/Source/Runtime/D3D12RHI/Private/D3D12Texture.cpp b/Engine/Source/Runtime/D3D12RHI/Private/D3D12Texture.cpp
--- a/Engine/Source/Runtime/D3D12RHI/Private/D3D12Texture.cpp
+++ b/Engine/Source/Runtime/D3D12RHI/Private/D3D12Texture.cpp
@@ -602,7 +602,15 @@ void SafeCreateTexture2D(FD3D12Device* pDevice,
 		case D3D12_HEAP_TYPE_READBACK:
 			{
 				FD3D12Resource* Resource = nullptr;
-				VERIFYD3D12CREATETEXTURERESULT(Adapter->CreateBuffer(heapType, pDevice->GetGPUMask(), pDevice->GetVisibilityMask(), MipBytesAligned, &Resource), TextureDesc);
+				VERIFYD3D12CREATETEXTURERESULT(
+					Adapter->CreateBuffer(heapType, pDevice->GetGPUMask(), pDevice->GetVisibilityMask(), MipBytesAligned, &Resource),
+					TextureDesc.Width,
+					TextureDesc.Height,
+					TextureDesc.DepthOrArraySize,
+					TextureDesc.Format,
+					TextureDesc.MipLevels,
+					TextureDesc.Flags
+					);
 				OutTexture2D->AsStandAlone(Resource);
 
 				if (IsCPUWritable(heapType))
@@ -613,7 +621,16 @@ void SafeCreateTexture2D(FD3D12Device* pDevice,
 			break;
 
 		case D3D12_HEAP_TYPE_DEFAULT:
-			VERIFYD3D12CREATETEXTURERESULT(pDevice->GetTextureAllocator().AllocateTexture(TextureDesc, ClearValue, Format, *OutTexture2D, InitialState), TextureDesc);
+			VERIFYD3D12CREATETEXTURERESULT(
+				pDevice->GetTextureAllocator().AllocateTexture(TextureDesc, ClearValue, Format, *OutTexture2D, InitialState),
+				TextureDesc.Width,
+				TextureDesc.Height,
+				TextureDesc.DepthOrArraySize,
+				TextureDesc.Format,
+				TextureDesc.MipLevels,
+				TextureDesc.Flags
+				);
+
 			break;
 
 		default:
@@ -1042,7 +1059,8 @@ FD3D12Texture3D* FD3D12DynamicRHI::CreateD3D12Texture3D(FRHICommandListImmediate
 	{
 		FD3D12Texture3D* Texture3D = new FD3D12Texture3D(Device, SizeX, SizeY, SizeZ, NumMips, Format, Flags, CreateInfo.ClearValueBinding);
 
-		Device->GetTextureAllocator().AllocateTexture(TextureDesc, ClearValuePtr, Format, Texture3D->ResourceLocation, (CreateInfo.BulkData != nullptr) ? D3D12_RESOURCE_STATE_COPY_DEST : DestinationState);
+		HRESULT hr = Device->GetTextureAllocator().AllocateTexture(TextureDesc, ClearValuePtr, Format, Texture3D->ResourceLocation, (CreateInfo.BulkData != nullptr) ? D3D12_RESOURCE_STATE_COPY_DEST : DestinationState);
+		check(SUCCEEDED(hr));
 
 		if (bCreateRTV)
 		{
diff --git a/Engine/Source/Runtime/D3D12RHI/Private/D3D12Util.cpp b/Engine/Source/Runtime/D3D12RHI/Private/D3D12Util.cpp
--- a/Engine/Source/Runtime/D3D12RHI/Private/D3D12Util.cpp
+++ b/Engine/Source/Runtime/D3D12RHI/Private/D3D12Util.cpp
@@ -68,8 +68,7 @@ static FString GetD3D12DeviceHungErrorString(HRESULT ErrorCode)
 		D3DERR(DXGI_ERROR_DEVICE_RESET)
 		D3DERR(DXGI_ERROR_DRIVER_INTERNAL_ERROR)
 		D3DERR(DXGI_ERROR_INVALID_CALL)
-		default:
-			ErrorCodeText = FString::Printf(TEXT("%08X"), (int32)ErrorCode);
+	default: ErrorCodeText = FString::Printf(TEXT("%08X"), (int32)ErrorCode);
 	}
 
 	return ErrorCodeText;
@@ -94,14 +93,13 @@ static FString GetD3D12ErrorString(HRESULT ErrorCode, ID3D12Device* Device)
 		D3DERR(DXGI_ERROR_INVALID_CALL)
 		D3DERR(E_NOINTERFACE)
 		D3DERR(DXGI_ERROR_DEVICE_REMOVED)
-		default:
-			ErrorCodeText = FString::Printf(TEXT("%08X"), (int32)ErrorCode);
+	default: ErrorCodeText = FString::Printf(TEXT("%08X"), (int32)ErrorCode);
 	}
 
 	if (ErrorCode == DXGI_ERROR_DEVICE_REMOVED && Device)
 	{
 		HRESULT hResDeviceRemoved = Device->GetDeviceRemovedReason();
-		ErrorCodeText += FString(TEXT(" with Reason: ")) + GetD3D12DeviceHungErrorString(hResDeviceRemoved);
+		ErrorCodeText += FString(TEXT(" ")) + GetD3D12DeviceHungErrorString(hResDeviceRemoved);
 	}
 
 	return ErrorCodeText;
@@ -247,12 +245,12 @@ namespace D3D12RHI
 		UE_LOG(LogD3D12RHI, Fatal, TEXT("%s failed \n at %s:%u \n with error %s"), ANSI_TO_TCHAR(Code), ANSI_TO_TCHAR(Filename), Line, *ErrorString);
 	}
 
-	void VerifyD3D12CreateTextureResult(HRESULT D3DResult, const ANSICHAR* Code, const ANSICHAR* Filename, uint32 Line, const D3D12_RESOURCE_DESC& TextureDesc)
+	void VerifyD3D12CreateTextureResult(HRESULT D3DResult, const ANSICHAR* Code, const ANSICHAR* Filename, uint32 Line, uint32 SizeX, uint32 SizeY, uint32 SizeZ, uint8 Format, uint32 NumMips, uint32 Flags)
 	{
 		check(FAILED(D3DResult));
 
 		const FString ErrorString = GetD3D12ErrorString(D3DResult, 0);
-		const TCHAR* D3DFormatString = GetD3D12TextureFormatString(TextureDesc.Format);
+		const TCHAR* D3DFormatString = GetD3D12TextureFormatString((DXGI_FORMAT)Format);
 
 		UE_LOG(LogD3D12RHI, Error,
 			TEXT("%s failed \n at %s:%u \n with error %s, \n Size=%ix%ix%i Format=%s(0x%08X), NumMips=%i, Flags=%s"),
@@ -260,13 +258,13 @@ namespace D3D12RHI
 			ANSI_TO_TCHAR(Filename),
 			Line,
 			*ErrorString,
-			TextureDesc.Width,
-			TextureDesc.Height,
-			TextureDesc.DepthOrArraySize,
+			SizeX,
+			SizeY,
+			SizeZ,
 			D3DFormatString,
-			TextureDesc.Format,
-			TextureDesc.MipLevels,
-			*GetD3D12TextureFlagString(TextureDesc.Flags));
+			Format,
+			NumMips,
+			*GetD3D12TextureFlagString(Flags));
 
 		TerminateOnDeviceRemoved(D3DResult);
 		TerminateOnOutOfMemory(D3DResult, true);
@@ -285,13 +283,13 @@ namespace D3D12RHI
 			ANSI_TO_TCHAR(Filename),
 			Line,
 			*ErrorString,
-			TextureDesc.Width,
-			TextureDesc.Height,
-			TextureDesc.DepthOrArraySize,
+			SizeX,
+			SizeY,
+			SizeZ,
 			D3DFormatString,
-			TextureDesc.Format,
-			TextureDesc.MipLevels,
-			*GetD3D12TextureFlagString(TextureDesc.Flags));
+			Format,
+			NumMips,
+			*GetD3D12TextureFlagString(Flags));
 	}
 
 	void VerifyComRefCount(IUnknown* Object, int32 ExpectedRefs, const TCHAR* Code, const TCHAR* Filename, int32 Line)
diff --git a/Engine/Source/Runtime/D3D12RHI/Public/D3D12Util.h b/Engine/Source/Runtime/D3D12RHI/Public/D3D12Util.h
--- a/Engine/Source/Runtime/D3D12RHI/Public/D3D12Util.h
+++ b/Engine/Source/Runtime/D3D12RHI/Public/D3D12Util.h
@@ -42,7 +42,10 @@ namespace D3D12RHI
 	 * @param	Line - The line number of Code within Filename.
 	 */
 	extern void VerifyD3D12Result(HRESULT Result, const ANSICHAR* Code, const ANSICHAR* Filename, uint32 Line, ID3D12Device* Device);
+}
 
+namespace D3D12RHI
+{
 	/**
 	* Checks that the given result isn't a failure.  If it is, the application exits with an appropriate error message.
 	* @param	Result - The result code to check
@@ -50,14 +53,15 @@ namespace D3D12RHI
 	* @param	Filename - The filename of the source file containing Code.
 	* @param	Line - The line number of Code within Filename.
 	*/
-	extern void VerifyD3D12CreateTextureResult(HRESULT D3DResult, const ANSICHAR* Code, const ANSICHAR* Filename, uint32 Line, const D3D12_RESOURCE_DESC& TextureDesc);
+	extern void VerifyD3D12CreateTextureResult(HRESULT D3DResult, const ANSICHAR* Code, const ANSICHAR* Filename, uint32 Line,
+		uint32 SizeX, uint32 SizeY, uint32 SizeZ, uint8 D3DFormat, uint32 NumMips, uint32 Flags);
 
 	/**
 	 * A macro for using VERIFYD3D12RESULT that automatically passes in the code and filename/line.
 	 */
-#define VERIFYD3D12RESULT_EX(x, Device)	{HRESULT hres = x; if (FAILED(hres)) { VerifyD3D12Result(hres, #x, __FILE__, __LINE__, Device); }}
-#define VERIFYD3D12RESULT(x)			{HRESULT hres = x; if (FAILED(hres)) { VerifyD3D12Result(hres, #x, __FILE__, __LINE__, nullptr); }}
-#define VERIFYD3D12CREATETEXTURERESULT(x, Desc) {HRESULT hres = x; if (FAILED(hres)) { VerifyD3D12CreateTextureResult(hres, #x, __FILE__, __LINE__, Desc); }}
+#define VERIFYD3D12RESULT_EX(x, Device)	{HRESULT hr = x; if (FAILED(hr)) { VerifyD3D12Result(hr,#x,__FILE__,__LINE__, Device); }}
+#define VERIFYD3D12RESULT(x)			{HRESULT hr = x; if (FAILED(hr)) { VerifyD3D12Result(hr,#x,__FILE__,__LINE__, 0); }}
+#define VERIFYD3D12CREATETEXTURERESULT(x,SizeX,SizeY,SizeZ,Format,NumMips,Flags) {HRESULT hr = x; if (FAILED(hr)) { VerifyD3D12CreateTextureResult(hr,#x,__FILE__,__LINE__,SizeX,SizeY,SizeZ,Format,NumMips,Flags); }}
 
 	/**
 	 * Checks that a COM object has the expected number of references.
