Cleaned D3D12 after math dependency to do the same done in DevRendering in CL 4992814 to avoid future conflicts.
# fyi emil.persson

diff --git a/Engine/Source/Runtime/D3D12RHI/D3D12RHI.Build.cs b/Engine/Source/Runtime/D3D12RHI/D3D12RHI.Build.cs
--- a/Engine/Source/Runtime/D3D12RHI/D3D12RHI.Build.cs
+++ b/Engine/Source/Runtime/D3D12RHI/D3D12RHI.Build.cs
@@ -39,28 +39,8 @@ public class D3D12RHI : ModuleRules
 			AddEngineThirdPartyPrivateStaticDependencies(Target, "DX11");
 			AddEngineThirdPartyPrivateStaticDependencies(Target, "NVAPI");
 			AddEngineThirdPartyPrivateStaticDependencies(Target, "AMD_AGS");
-			AddEngineThirdPartyPrivateStaticDependencies(Target, "IntelMetricsDiscovery");
+            AddEngineThirdPartyPrivateStaticDependencies(Target, "NVAftermath");
+            AddEngineThirdPartyPrivateStaticDependencies(Target, "IntelMetricsDiscovery");
 		}
-
-        if (Target.Platform == UnrealTargetPlatform.Win64)
-        {
-            string NVAftermathPath = Target.UEThirdPartySourceDirectory + "NVIDIA/NVaftermath/";
-            PublicSystemIncludePaths.Add(NVAftermathPath);
-
-            string NVAftermathLibPath = NVAftermathPath + "amd64/";
-            PublicLibraryPaths.Add(NVAftermathLibPath);
-            PublicAdditionalLibraries.Add("GFSDK_Aftermath_Lib.x64.lib");
-
-            string AftermathDllName = "GFSDK_Aftermath_Lib.x64.dll";
-            string nvDLLPath = "$(EngineDir)/Binaries/ThirdParty/NVIDIA/NVaftermath/Win64/" + AftermathDllName;
-            PublicDelayLoadDLLs.Add(AftermathDllName);
-            RuntimeDependencies.Add(nvDLLPath);
-
-            PublicDefinitions.Add("NV_AFTERMATH=1");
-        }
-        else
-        {
-            PublicDefinitions.Add("NV_AFTERMATH=0");
-        }
     }
 }
