Fixed variable shadowing warning in FD3D12CommandContext::RHICopyBufferRegions

diff --git a/Engine/Source/Runtime/D3D12RHI/Private/D3D12VertexBuffer.cpp b/Engine/Source/Runtime/D3D12RHI/Private/D3D12VertexBuffer.cpp
--- a/Engine/Source/Runtime/D3D12RHI/Private/D3D12VertexBuffer.cpp
+++ b/Engine/Source/Runtime/D3D12RHI/Private/D3D12VertexBuffer.cpp
@@ -210,7 +210,7 @@ void FD3D12CommandContext::RHICopyBufferRegion(FVertexBufferRHIParamRef DestBuff
 
 void FD3D12CommandContext::RHICopyBufferRegions(const TArrayView<const FCopyBufferRegionParams> Params)
 {
-	auto TransitionBuffer = [](FD3D12Resource* pResource, FD3D12CommandListHandle& CommandListHandle, const D3D12_RESOURCE_STATES Desired)
+	auto TransitionBuffer = [](FD3D12Resource* pResource, FD3D12CommandListHandle& InCommandListHandle, const D3D12_RESOURCE_STATES Desired)
 	{
 		bool bUseTracking = pResource->RequiresResourceStateTracking();
 		D3D12_RESOURCE_STATES Current;
@@ -221,12 +221,12 @@ void FD3D12CommandContext::RHICopyBufferRegions(const TArrayView<const FCopyBuff
 			Current = pResource->GetDefaultResourceState();
 			if (Current != Desired)
 			{
-				CommandListHandle.AddTransitionBarrier(pResource, Current, Desired, Subresource);
+				InCommandListHandle.AddTransitionBarrier(pResource, Current, Desired, Subresource);
 			}
 		}
 		else
 		{
-			FD3D12DynamicRHI::TransitionResource(CommandListHandle, pResource, Desired, Subresource);
+			FD3D12DynamicRHI::TransitionResource(InCommandListHandle, pResource, Desired, Subresource);
 		}
 	};
 
