XB1 D3D12 RHI - make up some missed debug names

diff --git a/Engine/Source/Runtime/D3D12RHI/Private/D3D12Allocation.cpp b/Engine/Source/Runtime/D3D12RHI/Private/D3D12Allocation.cpp
--- a/Engine/Source/Runtime/D3D12RHI/Private/D3D12Allocation.cpp
+++ b/Engine/Source/Runtime/D3D12RHI/Private/D3D12Allocation.cpp
@@ -1210,7 +1210,7 @@ FD3D12TextureAllocator::~FD3D12TextureAllocator()
 	FD3D12DynamicRHI::GetD3DRHI()->UpdataTextureMemorySize(-int32(MaxBlockSize / 1024));
 }
 
-HRESULT FD3D12TextureAllocator::AllocateTexture(D3D12_RESOURCE_DESC Desc, const D3D12_CLEAR_VALUE* ClearValue, FD3D12ResourceLocation& TextureLocation, const D3D12_RESOURCE_STATES InitialState)
+HRESULT FD3D12TextureAllocator::AllocateTexture(D3D12_RESOURCE_DESC Desc, const D3D12_CLEAR_VALUE* ClearValue, FD3D12ResourceLocation& TextureLocation, const D3D12_RESOURCE_STATES InitialState, const TCHAR* Name)
 {
 	FD3D12Device* Device = GetParentDevice();
 	FD3D12Adapter* Adapter = Device->GetParentAdapter();
@@ -1229,7 +1229,7 @@ HRESULT FD3D12TextureAllocator::AllocateTexture(D3D12_RESOURCE_DESC Desc, const
 			FD3D12Heap* BackingHeap = ((FD3D12BuddyAllocator*)TextureLocation.GetAllocator())->GetBackingHeap();
 			uint64 HeapOffset = TextureLocation.GetAllocator()->GetAllocationOffsetInBytes(TextureLocation.GetBuddyAllocatorPrivateData());
 
-			hr = Adapter->CreatePlacedResource(Desc, BackingHeap, HeapOffset, InitialState, ClearValue, &NewResource);
+			hr = Adapter->CreatePlacedResource(Desc, BackingHeap, HeapOffset, InitialState, ClearValue, &NewResource, Name);
 	
 			TextureLocation.SetType(FD3D12ResourceLocation::ResourceLocationType::eSubAllocation);
 			TextureLocation.SetResource(NewResource);
@@ -1242,7 +1242,7 @@ HRESULT FD3D12TextureAllocator::AllocateTexture(D3D12_RESOURCE_DESC Desc, const
 	Desc.Alignment = 0;
 	const D3D12_HEAP_PROPERTIES HeapProps = CD3DX12_HEAP_PROPERTIES(D3D12_HEAP_TYPE_DEFAULT, (uint32)GetGPUMask(), (uint32)GetVisibilityMask());
 
-	hr = Adapter->CreateCommittedResource(Desc, HeapProps, InitialState, ClearValue, &NewResource);
+	hr = Adapter->CreateCommittedResource(Desc, HeapProps, InitialState, ClearValue, &NewResource, Name);
 
 	TextureLocation.SetType(FD3D12ResourceLocation::ResourceLocationType::eStandAlone);
 	TextureLocation.SetResource(NewResource);
@@ -1256,7 +1256,7 @@ FD3D12TextureAllocatorPool::FD3D12TextureAllocatorPool(FD3D12Device* Device, FRH
 	ReadOnlyTexturePool(Device, VisibilityNode, FString(L"Small Read-Only Texture allocator"), TEXTURE_POOL_SIZE, D3D12_HEAP_FLAG_ALLOW_ONLY_NON_RT_DS_TEXTURES)
 {};
 
-HRESULT FD3D12TextureAllocatorPool::AllocateTexture(D3D12_RESOURCE_DESC Desc, const D3D12_CLEAR_VALUE* ClearValue, uint8 UEFormat, FD3D12ResourceLocation& TextureLocation, const D3D12_RESOURCE_STATES InitialState)
+HRESULT FD3D12TextureAllocatorPool::AllocateTexture(D3D12_RESOURCE_DESC Desc, const D3D12_CLEAR_VALUE* ClearValue, uint8 UEFormat, FD3D12ResourceLocation& TextureLocation, const D3D12_RESOURCE_STATES InitialState, const TCHAR* Name)
 {
 	// 4KB alignment is only available for read only textures
 	if ((Desc.Flags & D3D12_RESOURCE_FLAG_ALLOW_RENDER_TARGET ||
@@ -1268,7 +1268,7 @@ HRESULT FD3D12TextureAllocatorPool::AllocateTexture(D3D12_RESOURCE_DESC Desc, co
 		if (TextureCanBe4KAligned(Desc, UEFormat))
 		{
 			Desc.Alignment = D3D12_SMALL_RESOURCE_PLACEMENT_ALIGNMENT; // request 4k alignment
-			return ReadOnlyTexturePool.AllocateTexture(Desc, ClearValue, TextureLocation, InitialState);
+			return ReadOnlyTexturePool.AllocateTexture(Desc, ClearValue, TextureLocation, InitialState, Name);
 		}
 	}
 
@@ -1276,7 +1276,7 @@ HRESULT FD3D12TextureAllocatorPool::AllocateTexture(D3D12_RESOURCE_DESC Desc, co
 	FD3D12Resource* Resource = nullptr;
 
 	const D3D12_HEAP_PROPERTIES HeapProps = CD3DX12_HEAP_PROPERTIES(D3D12_HEAP_TYPE_DEFAULT, (uint32)GetGPUMask(), (uint32)GetVisibilityMask());
-	HRESULT hr = Adapter->CreateCommittedResource(Desc, HeapProps, InitialState, ClearValue, &Resource);
+	HRESULT hr = Adapter->CreateCommittedResource(Desc, HeapProps, InitialState, ClearValue, &Resource, Name);
 
 	TextureLocation.SetType(FD3D12ResourceLocation::ResourceLocationType::eStandAlone);
 	TextureLocation.SetResource(Resource);
diff --git a/Engine/Source/Runtime/D3D12RHI/Private/D3D12Allocation.h b/Engine/Source/Runtime/D3D12RHI/Private/D3D12Allocation.h
--- a/Engine/Source/Runtime/D3D12RHI/Private/D3D12Allocation.h
+++ b/Engine/Source/Runtime/D3D12RHI/Private/D3D12Allocation.h
@@ -1051,7 +1051,7 @@ public:
 
 	~FD3D12TextureAllocator();
 
-	HRESULT AllocateTexture(D3D12_RESOURCE_DESC Desc, const D3D12_CLEAR_VALUE* ClearValue, FD3D12ResourceLocation& TextureLocation, const D3D12_RESOURCE_STATES InitialState);
+	HRESULT AllocateTexture(D3D12_RESOURCE_DESC Desc, const D3D12_CLEAR_VALUE* ClearValue, FD3D12ResourceLocation& TextureLocation, const D3D12_RESOURCE_STATES InitialState, const TCHAR* Name);
 };
 
 class FD3D12TextureAllocatorPool : public FD3D12DeviceChild, public FD3D12MultiNodeGPUObject
@@ -1059,7 +1059,7 @@ class FD3D12TextureAllocatorPool : public FD3D12DeviceChild, public FD3D12MultiN
 public:
 	FD3D12TextureAllocatorPool(FD3D12Device* Device, FRHIGPUMask VisibilityNode);
 
-	HRESULT AllocateTexture(D3D12_RESOURCE_DESC Desc, const D3D12_CLEAR_VALUE* ClearValue, uint8 UEFormat, FD3D12ResourceLocation& TextureLocation, const D3D12_RESOURCE_STATES InitialState);
+	HRESULT AllocateTexture(D3D12_RESOURCE_DESC Desc, const D3D12_CLEAR_VALUE* ClearValue, uint8 UEFormat, FD3D12ResourceLocation& TextureLocation, const D3D12_RESOURCE_STATES InitialState, const TCHAR* Name);
 
 	void CleanUpAllocations() { ReadOnlyTexturePool.CleanUpAllocations(); }
 
