D3D12 RHI - Deduct GPU idle time between commandlists from ProfileGPU results

diff --git a/Engine/Source/Runtime/D3D12RHI/Private/D3D12CommandContext.cpp b/Engine/Source/Runtime/D3D12RHI/Private/D3D12CommandContext.cpp
--- a/Engine/Source/Runtime/D3D12RHI/Private/D3D12CommandContext.cpp
+++ b/Engine/Source/Runtime/D3D12RHI/Private/D3D12CommandContext.cpp
@@ -244,17 +244,35 @@ void FD3D12CommandContext::CloseCommandList()
 	CommandListHandle.Close();
 }
 
-FD3D12CommandListHandle FD3D12CommandContext::FlushCommands(bool WaitForCompletion)
+FD3D12CommandListHandle FD3D12CommandContext::FlushCommands(bool WaitForCompletion, EFlushCommandsExtraAction ExtraAction)
 {
 	// We should only be flushing the default context
 	check(IsDefaultContext());
 
+	bool bHasProfileGPUAction = false;
+#if WITH_PROFILEGPU
+	// Only graphics command list supports ID3D12GraphicsCommandList::EndQuery currently
+	if (!bIsAsyncComputeContext)
+	{
+		if (ExtraAction == FCEA_StartProfilingGPU)
+		{
+			GetCommandListManager().StartTrackingCommandListTime();
+		}
+		else if (ExtraAction == FCEA_EndProfilingGPU)
+		{
+			GetCommandListManager().EndTrackingCommandListTime();
+		}
+		bHasProfileGPUAction = true;
+	}
+#endif
+
 	FD3D12Device* Device = GetParentDevice();
 	const bool bHasPendingWork = Device->PendingCommandLists.Num() > 0;
 	const bool bHasDoneWork = HasDoneWork() || bHasPendingWork;
+	const bool bOpenNewCmdList = WaitForCompletion || bHasDoneWork || bHasProfileGPUAction;
 
 	// Only submit a command list if it does meaningful work or the flush is expected to wait for completion.
-	if (WaitForCompletion || bHasDoneWork)
+	if (bOpenNewCmdList)
 	{
 		// Close the current command list
 		CloseCommandList();
diff --git a/Engine/Source/Runtime/D3D12RHI/Private/D3D12CommandContext.h b/Engine/Source/Runtime/D3D12RHI/Private/D3D12CommandContext.h
--- a/Engine/Source/Runtime/D3D12RHI/Private/D3D12CommandContext.h
+++ b/Engine/Source/Runtime/D3D12RHI/Private/D3D12CommandContext.h
@@ -88,6 +88,14 @@ protected:
 class FD3D12CommandContext : public FD3D12CommandContextBase, public FD3D12DeviceChild
 {
 public:
+	enum EFlushCommandsExtraAction
+	{
+		FCEA_None,
+		FCEA_StartProfilingGPU,
+		FCEA_EndProfilingGPU,
+		FCEA_Num
+	};
+
 	FD3D12CommandContext(class FD3D12Device* InParent, FD3D12SubAllocatedOnlineHeap::SubAllocationDesc& SubHeapDesc, bool InIsDefaultContext, bool InIsAsyncComputeContext = false);
 	virtual ~FD3D12CommandContext();
 
@@ -122,7 +130,7 @@ public:
 	void CloseCommandList();
 
 	// Close the D3D command list and execute it.  Optionally wait for the GPU to finish. Returns the handle to the command list so you can wait for it later.
-	FD3D12CommandListHandle FlushCommands(bool WaitForCompletion = false);
+	FD3D12CommandListHandle FlushCommands(bool WaitForCompletion = false, EFlushCommandsExtraAction ExtraAction = FCEA_None);
 
 	void Finish(TArray<FD3D12CommandListHandle>& CommandLists);
 
@@ -429,7 +437,7 @@ public:
 
 	uint32 GetGPUIndex() const { return GPUMask.ToIndex(); }
 
-protected: 
+protected:
 
 	FD3D12CommandContext* GetContext(uint32 InGPUIndex) final override 
 	{  
diff --git a/Engine/Source/Runtime/D3D12RHI/Private/D3D12CommandList.cpp b/Engine/Source/Runtime/D3D12RHI/Private/D3D12CommandList.cpp
--- a/Engine/Source/Runtime/D3D12RHI/Private/D3D12CommandList.cpp
+++ b/Engine/Source/Runtime/D3D12RHI/Private/D3D12CommandList.cpp
@@ -45,8 +45,12 @@ FD3D12CommandListHandle::FD3D12CommandListData::FD3D12CommandListData(FD3D12Devi
 	, CurrentGeneration(1)
 	, LastCompleteGeneration(0)
 	, IsClosed(false)
+	, bShouldTrackStartEndTime(false)
 	, PendingResourceBarriers()
 	, ResidencySet(nullptr)
+#if WITH_PROFILEGPU
+	, StartTimeQueryIdx(INDEX_NONE)
+#endif
 {
 	VERIFYD3D12RESULT(ParentDevice->GetDevice()->CreateCommandList((uint32)GetGPUMask(), CommandListType, CommandAllocator, nullptr, IID_PPV_ARGS(CommandList.GetInitReference())));
 	INC_DWORD_STAT(STAT_D3D12NumCommandLists);
@@ -110,6 +114,84 @@ FD3D12CommandListHandle::FD3D12CommandListData::~FD3D12CommandListData()
 	D3DX12Residency::DestroyResidencySet(GetParentDevice()->GetResidencyManager(), ResidencySet);
 }
 
+void FD3D12CommandListHandle::FD3D12CommandListData::Close()
+{
+	if (!IsClosed)
+	{
+		FlushResourceBarriers();
+		if (bShouldTrackStartEndTime)
+		{
+			FinishTrackingCommandListTime();
+		}
+		VERIFYD3D12RESULT(CommandList->Close());
+
+		D3DX12Residency::Close(ResidencySet);
+		IsClosed = true;
+	}
+}
+
+void FD3D12CommandListHandle::FD3D12CommandListData::Reset(FD3D12CommandAllocator& CommandAllocator, bool bTrackExecTime)
+{
+	VERIFYD3D12RESULT(CommandList->Reset(CommandAllocator, nullptr));
+
+	CurrentCommandAllocator = &CommandAllocator;
+	IsClosed = false;
+
+	// Indicate this command allocator is being used.
+	CurrentCommandAllocator->IncrementPendingCommandLists();
+
+	CleanupActiveGenerations();
+
+	// Remove all pendering barriers from the command list
+	PendingResourceBarriers.Reset();
+
+	// Empty tracked resource state for this command list
+	TrackedResourceState.Empty();
+
+	// If this fails there are too many concurrently open residency sets. Increase the value of MAX_NUM_CONCURRENT_CMD_LISTS
+	// in the residency manager. Beware, this will increase the CPU memory usage of every tracked resource.
+	D3DX12Residency::Open(ResidencySet);
+
+	// If this fails then some previous resource barriers were never submitted.
+	check(ResourceBarrierBatcher.GetBarriers().Num() == 0);
+
+#if DEBUG_RESOURCE_STATES
+	ResourceBarriers.Reset();
+#endif
+
+	if (bTrackExecTime)
+	{
+		StartTrackingCommandListTime();
+	}
+}
+
+int32 FD3D12CommandListHandle::FD3D12CommandListData::CreateAndInsertTimestampQuery()
+{	
+	FD3D12LinearQueryHeap* QueryHeap = GetParentDevice()->GetCmdListExecTimeQueryHeap();
+	check(QueryHeap);
+	return QueryHeap->EndQuery(this);
+}
+
+void FD3D12CommandListHandle::FD3D12CommandListData::StartTrackingCommandListTime()
+{
+#if WITH_PROFILEGPU
+	check(!IsClosed && !bShouldTrackStartEndTime && StartTimeQueryIdx == INDEX_NONE);
+	bShouldTrackStartEndTime = true;
+	StartTimeQueryIdx = CreateAndInsertTimestampQuery();
+#endif
+}
+
+void FD3D12CommandListHandle::FD3D12CommandListData::FinishTrackingCommandListTime()
+{
+#if WITH_PROFILEGPU
+	check(!IsClosed && bShouldTrackStartEndTime && StartTimeQueryIdx != INDEX_NONE);
+	bShouldTrackStartEndTime = false;
+	const int32 EndTimeQueryIdx = CreateAndInsertTimestampQuery();
+	CommandListManager->AddCommandListTimingPair(StartTimeQueryIdx, EndTimeQueryIdx);
+	StartTimeQueryIdx = INDEX_NONE;
+#endif
+}
+
 void inline FD3D12CommandListHandle::FD3D12CommandListData::FCommandListResourceState::ConditionalInitalize(FD3D12Resource* pResource, CResourceState& ResourceState)
 {
 	// If there is no entry, all subresources should be in the resource's TBD state.
diff --git a/Engine/Source/Runtime/D3D12RHI/Private/D3D12CommandList.h b/Engine/Source/Runtime/D3D12RHI/Private/D3D12CommandList.h
--- a/Engine/Source/Runtime/D3D12RHI/Private/D3D12CommandList.h
+++ b/Engine/Source/Runtime/D3D12RHI/Private/D3D12CommandList.h
@@ -58,49 +58,11 @@ private:
 
 		~FD3D12CommandListData();
 
-		void Close()
-		{
-			if (!IsClosed)
-			{
-				FlushResourceBarriers();
-				VERIFYD3D12RESULT(CommandList->Close());
-
-				D3DX12Residency::Close(ResidencySet);
-				IsClosed = true;
-			}
-		}
+		void Close();
 
 		// Reset the command list with a specified command allocator and optional initial state.
 		// Note: Command lists can be reset immediately after they are submitted for execution.
-		void Reset(FD3D12CommandAllocator& CommandAllocator)
-		{
-			VERIFYD3D12RESULT(CommandList->Reset(CommandAllocator, nullptr));
-
-			CurrentCommandAllocator = &CommandAllocator;
-			IsClosed = false;
-
-			// Indicate this command allocator is being used.
-			CurrentCommandAllocator->IncrementPendingCommandLists();
-
-			CleanupActiveGenerations();
-
-			// Remove all pendering barriers from the command list
-			PendingResourceBarriers.Reset();
-
-			// Empty tracked resource state for this command list
-			TrackedResourceState.Empty();
-
-			// If this fails there are too many concurrently open residency sets. Increase the value of MAX_NUM_CONCURRENT_CMD_LISTS
-			// in the residency manager. Beware, this will increase the CPU memory usage of every tracked resource.
-			D3DX12Residency::Open(ResidencySet);
-
-			// If this fails then some previous resource barriers were never submitted.
-			check(ResourceBarrierBatcher.GetBarriers().Num() == 0);
-
-#if DEBUG_RESOURCE_STATES
-			ResourceBarriers.Reset();
-#endif
-		}
+		void Reset(FD3D12CommandAllocator& CommandAllocator, bool bTrackExecTime);
 
 		bool IsComplete(uint64 Generation)
 		{
@@ -254,6 +216,7 @@ private:
 		uint64									CurrentGeneration;
 		uint64									LastCompleteGeneration;
 		bool									IsClosed;
+		bool									bShouldTrackStartEndTime;
 		typedef TPair<uint64, FD3D12SyncPoint>	GenerationSyncPointPair;	// Pair of command list generation to a sync point
 		TQueue<GenerationSyncPointPair>			ActiveGenerations;	// Queue of active command list generations and their sync points. Used to determine what command lists have been completed on the GPU.
 		FCriticalSection						ActiveGenerationsCS;	// While only a single thread can record to a command list at any given time, multiple threads can ask for the state of a given command list. So the associated tracking must be thread-safe.
@@ -289,16 +252,43 @@ private:
 		// Tracks all the resources barriers being issued on this command list in order
 		TArray<D3D12_RESOURCE_BARRIER> ResourceBarriers;
 #endif
+
+	private:
+		/** Helper to allocate and queue a new timestamp query */
+		int32 CreateAndInsertTimestampQuery();
+
+		/** Enable execution time tracking and mark commandlist start time */
+		void StartTrackingCommandListTime();
+
+		/** Disable execution time tracking and mark commandlist end time */
+		void FinishTrackingCommandListTime();
+
+#if WITH_PROFILEGPU
+		int32 StartTimeQueryIdx;
+#endif
 	};
 
 public:
 
 	FD3D12CommandListHandle() : CommandListData(nullptr) {}
 
-	FD3D12CommandListHandle(const FD3D12CommandListHandle& CL) : CommandListData(CL.CommandListData)
+	FD3D12CommandListHandle(const FD3D12CommandListHandle& CL)
+		: FD3D12CommandListHandle(CL.CommandListData)
+	{}
+
+	FD3D12CommandListHandle(FD3D12CommandListData* InData)
+		: CommandListData(InData)
 	{
 		if (CommandListData)
+		{
 			CommandListData->AddRef();
+		}
+	}
+
+	FD3D12CommandListHandle(FD3D12CommandListHandle&& CL)
+		: CommandListData(CL.CommandListData)
+	{
+		CL.CommandListData = nullptr;
 	}
 
 	virtual ~FD3D12CommandListHandle()
@@ -330,6 +320,20 @@ public:
 		return *this;
 	}
 
+	FD3D12CommandListHandle& operator=(FD3D12CommandListHandle&& CL)
+	{
+		if (CommandListData != CL.CommandListData)
+		{
+			if (CommandListData && CommandListData->Release() == 0)
+			{
+				delete CommandListData;
+			}
+			CommandListData = CL.CommandListData;
+			CL.CommandListData = nullptr;
+		}
+		return *this;
+	}
+
 	bool operator!() const
 	{
 		return CommandListData == 0;
@@ -370,7 +374,7 @@ public:
 	void Create(FD3D12Device* ParentDevice, D3D12_COMMAND_LIST_TYPE CommandListType, FD3D12CommandAllocator& CommandAllocator, FD3D12CommandListManager* InCommandListManager);
 
 	void Execute(bool WaitForCompletion = false);
-
+	
 	void Close()
 	{
 		check(CommandListData);
@@ -379,12 +383,12 @@ public:
 
 	// Reset the command list with a specified command allocator and optional initial state.
 	// Note: Command lists can be reset immediately after they are submitted for execution.
-	void Reset(FD3D12CommandAllocator& CommandAllocator)
+	void Reset(FD3D12CommandAllocator& CommandAllocator, bool bTrackExecTime = false)
 	{
 		check(CommandListData);
-		CommandListData->Reset(CommandAllocator);
+		CommandListData->Reset(CommandAllocator, bTrackExecTime);
 	}
-
+	
 	ID3D12CommandList* CommandList() const
 	{
 		check(CommandListData);
@@ -588,6 +592,8 @@ public:
 
 	FD3D12CLSyncPoint(FD3D12CommandListHandle& CL) : CommandList(CL), Generation(CL.CommandList() ? CL.CurrentGeneration() : 0) {}
 
+	FD3D12CLSyncPoint(FD3D12CommandListHandle&& CL) : CommandList(MoveTemp(CL)), Generation(CommandList.CommandList() ? CommandList.CurrentGeneration() : 0) {}
+
 	FD3D12CLSyncPoint(const FD3D12CLSyncPoint& SyncPoint) : CommandList(SyncPoint.CommandList), Generation(SyncPoint.Generation) {}
 
 	FD3D12CLSyncPoint& operator = (FD3D12CommandListHandle& CL)
@@ -598,6 +604,13 @@ public:
 		return *this;
 	}
 
+	FD3D12CLSyncPoint& operator=(FD3D12CommandListHandle&& CL)
+	{
+		Generation = (CL != nullptr) ? CL.CurrentGeneration() : 0;
+		CommandList = MoveTemp(CL);
+		return *this;
+	}
+
 	FD3D12CLSyncPoint& operator = (const FD3D12CLSyncPoint& SyncPoint)
 	{
 		CommandList = SyncPoint.CommandList;
diff --git a/Engine/Source/Runtime/D3D12RHI/Private/D3D12Device.cpp b/Engine/Source/Runtime/D3D12RHI/Private/D3D12Device.cpp
--- a/Engine/Source/Runtime/D3D12RHI/Private/D3D12Device.cpp
+++ b/Engine/Source/Runtime/D3D12RHI/Private/D3D12Device.cpp
@@ -36,6 +36,9 @@ FD3D12Device::FD3D12Device(FRHIGPUMask InGPUMask, FD3D12Adapter* InAdapter) :
 	GlobalViewHeap(this, InGPUMask),
 	OcclusionQueryHeap(this, D3D12_QUERY_HEAP_TYPE_OCCLUSION, 65536, 4 /*frames to keep results */ * 1 /*batches per frame*/),
 	TimestampQueryHeap(this, D3D12_QUERY_HEAP_TYPE_TIMESTAMP, 8192, 4 /*frames to keep results */ * 5 /*batches per frame*/ ),
+#if WITH_PROFILEGPU
+	CmdListExecTimeQueryHeap(this, D3D12_QUERY_HEAP_TYPE_TIMESTAMP, 8192),
+#endif
 	DefaultBufferAllocator(this, InGPUMask), //Note: Cross node buffers are possible 
 	SamplerID(0),
 	DefaultFastAllocator(this, InGPUMask, D3D12_HEAP_TYPE_UPLOAD, 1024 * 1024 * 4),
@@ -75,6 +78,15 @@ ID3D12Device5* FD3D12Device::GetRayTracingDevice()
 }
 #endif // D3D12_RHI_RAYTRACING
 
+FD3D12LinearQueryHeap* FD3D12Device::GetCmdListExecTimeQueryHeap()
+{
+#if WITH_PROFILEGPU
+	return &CmdListExecTimeQueryHeap;
+#else
+	return nullptr;
+#endif
+}
+
 FD3D12DynamicRHI* FD3D12Device::GetOwningRHI()
 { 
 	return GetParentAdapter()->GetOwningRHI();
diff --git a/Engine/Source/Runtime/D3D12RHI/Private/D3D12Device.h b/Engine/Source/Runtime/D3D12RHI/Private/D3D12Device.h
--- a/Engine/Source/Runtime/D3D12RHI/Private/D3D12Device.h
+++ b/Engine/Source/Runtime/D3D12RHI/Private/D3D12Device.h
@@ -53,6 +53,7 @@ public:
 
 	inline FD3D12QueryHeap* GetOcclusionQueryHeap() { return &OcclusionQueryHeap; }
 	inline FD3D12QueryHeap* GetTimestampQueryHeap() { return &TimestampQueryHeap; }
+	FD3D12LinearQueryHeap* GetCmdListExecTimeQueryHeap();
 
 	template <typename TViewDesc> FD3D12OfflineDescriptorManager& GetViewDescriptorAllocator();
 	template<> FD3D12OfflineDescriptorManager& GetViewDescriptorAllocator<D3D12_SHADER_RESOURCE_VIEW_DESC>() { return SRVAllocator; }
@@ -144,6 +145,9 @@ protected:
 
 	FD3D12QueryHeap OcclusionQueryHeap;
 	FD3D12QueryHeap TimestampQueryHeap;
+#if WITH_PROFILEGPU
+	FD3D12LinearQueryHeap CmdListExecTimeQueryHeap;
+#endif
 
 	FD3D12DefaultBufferAllocator DefaultBufferAllocator;
 
diff --git a/Engine/Source/Runtime/D3D12RHI/Private/D3D12DirectCommandListManager.cpp b/Engine/Source/Runtime/D3D12RHI/Private/D3D12DirectCommandListManager.cpp
--- a/Engine/Source/Runtime/D3D12RHI/Private/D3D12DirectCommandListManager.cpp
+++ b/Engine/Source/Runtime/D3D12RHI/Private/D3D12DirectCommandListManager.cpp
@@ -309,6 +309,9 @@ FD3D12CommandListManager::FD3D12CommandListManager(FD3D12Device* InParent, D3D12
 	, CommandListFence(nullptr)
 	, CommandListType(InCommandListType)
 	, QueueType(InQueueType)
+#if WITH_PROFILEGPU
+	, bShouldTrackCmdListTime(false)
+#endif
 {
 }
 
@@ -400,7 +403,7 @@ FD3D12CommandListHandle FD3D12CommandListManager::ObtainCommandList(FD3D12Comman
 	}
 
 	check(List.GetCommandListType() == CommandListType);
-	List.Reset(CommandAllocator);
+	List.Reset(CommandAllocator, ShouldTrackCommandListTime());
 	return List;
 }
 
@@ -626,6 +629,63 @@ void FD3D12CommandListManager::ReleaseResourceBarrierCommandListAllocator()
 	}
 }
 
+void FD3D12CommandListManager::StartTrackingCommandListTime()
+{
+#if WITH_PROFILEGPU
+	check(QueueType == ED3D12CommandQueueType::Default && !bShouldTrackCmdListTime);
+	PendingTimingPairs.Reset();
+	ResolvedTimingPairs.Reset();
+	bShouldTrackCmdListTime = true;
+#endif
+}
+
+void FD3D12CommandListManager::EndTrackingCommandListTime()
+{
+#if WITH_PROFILEGPU
+	check(QueueType == ED3D12CommandQueueType::Default && bShouldTrackCmdListTime);
+	bShouldTrackCmdListTime = false;
+#endif
+}
+
+void FD3D12CommandListManager::GetCommandListTimingResults(TArray<FResolvedCmdListExecTime>& OutTimingPairs)
+{
+#if WITH_PROFILEGPU
+	check(!bShouldTrackCmdListTime && QueueType == ED3D12CommandQueueType::Default);
+	FlushPendingTimingPairs();
+	OutTimingPairs = MoveTemp(ResolvedTimingPairs);
+#endif
+}
+
+void FD3D12CommandListManager::AddCommandListTimingPair(int32 StartTimeQueryIdx, int32 EndTimeQueryIdx)
+{
+#if WITH_PROFILEGPU
+	check(StartTimeQueryIdx >= 0 && EndTimeQueryIdx >= 0);
+	FScopeLock Lock(&CmdListTimingCS);
+	new (PendingTimingPairs) FCmdListExecTime(StartTimeQueryIdx, EndTimeQueryIdx);
+#endif
+}
+
+void FD3D12CommandListManager::FlushPendingTimingPairs()
+{
+#if WITH_PROFILEGPU
+	check(!ResolvedTimingPairs.Num() && !bShouldTrackCmdListTime);
+
+	TArray<uint64> AllTimestamps;
+	GetParentDevice()->GetCmdListExecTimeQueryHeap()->FlushAndGetResults(AllTimestamps);
+
+	const int32 NumPending = PendingTimingPairs.Num();
+	ResolvedTimingPairs.Empty(NumPending);
+	for (int32 Idx = 0; Idx < NumPending; ++Idx)
+	{
+		const FCmdListExecTime& QueryIdxPair = PendingTimingPairs[Idx];
+		const uint64 StartStamp = AllTimestamps[QueryIdxPair.StartTimeQueryIdx];
+		const uint64 EndStamp = AllTimestamps[QueryIdxPair.EndTimeQueryIdx];
+		new (ResolvedTimingPairs) FResolvedCmdListExecTime(StartStamp, EndStamp);
+	}
+	PendingTimingPairs.Reset();
+#endif
+}
+
 uint32 FD3D12CommandListManager::GetResourceBarrierCommandList(FD3D12CommandListHandle& hList, FD3D12CommandListHandle& hResourceBarrierList)
 {
 	TArray<FD3D12PendingResourceBarrier>& PendingResourceBarriers = hList.PendingResourceBarriers();
@@ -750,6 +810,15 @@ FD3D12CommandListHandle FD3D12CommandListManager::CreateCommandListHandle(FD3D12
 	return List;
 }
 
+bool FD3D12CommandListManager::ShouldTrackCommandListTime() const
+{
+#if WITH_PROFILEGPU
+	return bShouldTrackCmdListTime;
+#else
+	return false;
+#endif
+}
+
 FD3D12FenceCore* FD3D12FenceCorePool::ObtainFenceCore(uint32 GPUIndex)
 {
 	{
diff --git a/Engine/Source/Runtime/D3D12RHI/Private/D3D12DirectCommandListManager.h b/Engine/Source/Runtime/D3D12RHI/Private/D3D12DirectCommandListManager.h
--- a/Engine/Source/Runtime/D3D12RHI/Private/D3D12DirectCommandListManager.h
+++ b/Engine/Source/Runtime/D3D12RHI/Private/D3D12DirectCommandListManager.h
@@ -162,6 +162,19 @@ private:
 class FD3D12CommandListManager : public FD3D12DeviceChild, public FD3D12SingleNodeGPUObject
 {
 public:
+	struct FResolvedCmdListExecTime
+	{
+		uint64 StartTimestamp;
+		uint64 EndTimestamp;
+
+		FResolvedCmdListExecTime() = default;
+
+		FResolvedCmdListExecTime(uint64 InStart, uint64 InEnd)
+			: StartTimestamp(InStart)
+			, EndTimestamp(InEnd)
+		{}
+	};
+
 	FD3D12CommandListManager(FD3D12Device* InParent, D3D12_COMMAND_LIST_TYPE InCommandListType, ED3D12CommandQueueType InQueueType);
 	virtual ~FD3D12CommandListManager();
 
@@ -206,11 +219,42 @@ public:
 
 	void ReleaseResourceBarrierCommandListAllocator();
 
+	/** Command lists opened after this returns will track their execution time */
+	void StartTrackingCommandListTime();
+
+	/** Command lists opened after this returns won't track execution time */
+	void EndTrackingCommandListTime();
+
+	/** Get the start/end timestamps of all tracked commandlists obtained from this manager */
+	void GetCommandListTimingResults(TArray<FResolvedCmdListExecTime>& OutTimingPairs);
+
+	/** Called back by commandlists when they are closed */
+	void AddCommandListTimingPair(int32 StartTimeQueryIdx, int32 EndTimeQueryIdx);
+
 protected:
+	struct FCmdListExecTime
+	{
+		int32 StartTimeQueryIdx;
+		int32 EndTimeQueryIdx;
+
+		FCmdListExecTime() = default;
+
+		FCmdListExecTime(int32 InStart, int32 InEnd)
+			: StartTimeQueryIdx(InStart)
+			, EndTimeQueryIdx(InEnd)
+		{}
+	};
+
 	// Returns signaled Fence
 	uint64 ExecuteAndIncrementFence(FD3D12CommandListPayload& Payload, FD3D12Fence &Fence);
 	FD3D12CommandListHandle CreateCommandListHandle(FD3D12CommandAllocator& CommandAllocator);
 
+	/** Should this commandlist track its execution time */
+	bool ShouldTrackCommandListTime() const;
+
+	/** Resolve all commandlist start/end timestamp queries and get results. This method is blocking */
+	void FlushPendingTimingPairs();
+
 	TRefCountPtr<ID3D12CommandQueue>		D3DCommandQueue;
 
 	FThreadsafeQueue<FD3D12CommandListHandle> ReadyLists;
@@ -225,4 +269,11 @@ protected:
 	ED3D12CommandQueueType					QueueType;
 	FCriticalSection						ResourceStateCS;
 	FCriticalSection						FenceCS;
+
+#if WITH_PROFILEGPU
+	bool bShouldTrackCmdListTime;
+	FCriticalSection CmdListTimingCS;
+	TArray<FCmdListExecTime> PendingTimingPairs;
+	TArray<FResolvedCmdListExecTime> ResolvedTimingPairs;
+#endif
 };
diff --git a/Engine/Source/Runtime/D3D12RHI/Private/D3D12Query.cpp b/Engine/Source/Runtime/D3D12RHI/Private/D3D12Query.cpp
--- a/Engine/Source/Runtime/D3D12RHI/Private/D3D12Query.cpp
+++ b/Engine/Source/Runtime/D3D12RHI/Private/D3D12Query.cpp
@@ -436,6 +436,201 @@ void FD3D12QueryHeap::CreateResultBuffer()
 	VERIFYD3D12RESULT(ResultBuffer->GetResource()->Map(0, nullptr, &pResultData));
 }
 
+/*=============================================================================
+* class FD3D12LinearQueryHeap
+*=============================================================================*/
+
+FD3D12LinearQueryHeap::FD3D12LinearQueryHeap(FD3D12Device* InParent, D3D12_QUERY_HEAP_TYPE InHeapType, int32 GrowCount)
+	: FD3D12DeviceChild(InParent)
+	, FD3D12SingleNodeGPUObject(InParent->GetGPUMask())
+	, QueryHeapType(InHeapType)
+	, QueryType(HeapTypeToQueryType(InHeapType))
+	, GrowNumQueries(GrowCount)
+	, SlotToHeapIdxShift(FPlatformMath::CountBits(GrowCount - 1))
+	, HeapState(HS_Open)
+	, NextFreeIdx(0)
+	, CurMaxNumQueries(0)
+	, NextChunkIdx(0)
+{
+	check(GrowCount > 0 && !(GrowCount & (GrowCount - 1)));
+}
+
+FD3D12LinearQueryHeap::~FD3D12LinearQueryHeap()
+{
+	ReleaseResources();
+}
+
+int32 FD3D12LinearQueryHeap::BeginQuery(FD3D12CommandListHandle CmdListHandle)
+{
+	const int32 SlotIdx = AllocateQueryHeapSlot();
+	const int32 HeapIdx = SlotIdx >> SlotToHeapIdxShift;
+	const int32 Offset = SlotIdx & (GrowNumQueries - 1);
+
+	FChunk& Chunk = AllocatedChunks[HeapIdx];
+	CmdListHandle->BeginQuery(Chunk.QueryHeap, QueryType, Offset);
+	CmdListHandle.UpdateResidency(&Chunk.QueryHeapResidencyHandle);
+	FD3D12CommandContext* Context = CmdListHandle.GetCurrentOwningContext();
+	if (Context)
+	{
+		++Context->otherWorkCounter;
+	}
+	return SlotIdx;
+}
+
+int32 FD3D12LinearQueryHeap::EndQuery(FD3D12CommandListHandle CmdListHandle)
+{
+	const int32 SlotIdx = AllocateQueryHeapSlot();
+	const int32 HeapIdx = SlotIdx >> SlotToHeapIdxShift;
+	const int32 Offset = SlotIdx & (GrowNumQueries - 1);
+
+	FChunk& Chunk = AllocatedChunks[HeapIdx];
+	CmdListHandle->EndQuery(Chunk.QueryHeap, QueryType, Offset);
+	CmdListHandle.UpdateResidency(&Chunk.QueryHeapResidencyHandle);
+	FD3D12CommandContext* Context = CmdListHandle.GetCurrentOwningContext();
+	if (Context)
+	{
+		++Context->otherWorkCounter;
+	}
+	return SlotIdx;
+}
+
+void FD3D12LinearQueryHeap::Reset()
+{
+	HeapState = HS_Open;
+	NextFreeIdx = 0;
+}
+
+void FD3D12LinearQueryHeap::FlushAndGetResults(TArray<uint64>& QueryResults, bool bReleaseResources)
+{
+	HeapState = HS_Closed;
+	
+	int32 NumActiveQueries = NextFreeIdx;
+	const uint64 ResultBuffSize = ResultSize * NumActiveQueries;
+	TRefCountPtr<FD3D12Resource> ResultBuff;
+	CreateResultBuffer(ResultBuffSize, ResultBuff.GetInitReference());
+
+	FD3D12CommandContext& Context = GetParentDevice()->GetDefaultCommandContext();
+	++Context.otherWorkCounter;
+	const int32 NumHeaps = (NumActiveQueries + GrowNumQueries - 1) >> SlotToHeapIdxShift;
+	for (int32 HeapIdx = 0; HeapIdx < NumHeaps; ++HeapIdx)
+	{
+		const int32 NumQueriesInHeap = FMath::Min(NumActiveQueries, GrowNumQueries);
+		NumActiveQueries -= GrowNumQueries;
+		FChunk& Chunk = AllocatedChunks[HeapIdx];
+		Context.CommandListHandle->ResolveQueryData(
+			Chunk.QueryHeap,
+			QueryType,
+			0,
+			NumQueriesInHeap,
+			ResultBuff->GetResource(),
+			ResultSize * HeapIdx * GrowNumQueries);
+		Context.CommandListHandle.UpdateResidency(&Chunk.QueryHeapResidencyHandle);
+		Context.CommandListHandle.UpdateResidency(ResultBuff);
+	}
+
+	Context.FlushCommands(true);
+	const int32 NumResults = NextFreeIdx;
+	QueryResults.Empty(NumResults);
+	QueryResults.AddUninitialized(NumResults);
+	void* MappedResult;
+	VERIFYD3D12RESULT(ResultBuff->GetResource()->Map(0, nullptr, &MappedResult));
+	FMemory::Memcpy(QueryResults.GetData(), MappedResult, ResultBuffSize);
+	ResultBuff->GetResource()->Unmap(0, nullptr);
+
+	if (bReleaseResources)
+	{
+		ReleaseResources();
+	}
+	Reset();
+}
+
+D3D12_QUERY_TYPE FD3D12LinearQueryHeap::HeapTypeToQueryType(D3D12_QUERY_HEAP_TYPE HeapType)
+{
+	switch (HeapType)
+	{
+	case D3D12_QUERY_HEAP_TYPE_OCCLUSION:
+		return D3D12_QUERY_TYPE_OCCLUSION;
+	case D3D12_QUERY_HEAP_TYPE_TIMESTAMP:
+		return D3D12_QUERY_TYPE_TIMESTAMP;
+	default:
+		check(false);
+		return (D3D12_QUERY_TYPE)-1;
+	}
+}
+
+int32 FD3D12LinearQueryHeap::AllocateQueryHeapSlot()
+{
+	check(HeapState == HS_Open);
+	const int32 SlotIdx = FPlatformAtomics::InterlockedIncrement(&NextFreeIdx) - 1;
+
+	if (SlotIdx >= CurMaxNumQueries)
+	{
+		FScopeLock Lock(&CS);
+		while (SlotIdx >= CurMaxNumQueries)
+		{
+			Grow();
+		}
+	}
+	return SlotIdx;
+}
+
+void FD3D12LinearQueryHeap::Grow()
+{
+	const int32 ChunkIdx = NextChunkIdx++;
+	checkf(ChunkIdx < MaxNumChunks, TEXT("Running out of chunks, consider increase MaxNumChunks or GrowNumQueries"));
+	FChunk& NewChunk = AllocatedChunks[ChunkIdx];
+	CreateQueryHeap(GrowNumQueries, NewChunk.QueryHeap.GetInitReference(), NewChunk.QueryHeapResidencyHandle);
+	CurMaxNumQueries += GrowNumQueries;
+}
+
+void FD3D12LinearQueryHeap::CreateQueryHeap(int32 NumQueries, ID3D12QueryHeap** OutHeap, FD3D12ResidencyHandle& OutResidencyHandle)
+{
+	D3D12_QUERY_HEAP_DESC Desc;
+	Desc.Type = QueryHeapType;
+	Desc.Count = static_cast<uint32>(NumQueries);
+	Desc.NodeMask = static_cast<uint32>(GetGPUMask());
+	VERIFYD3D12RESULT(GetParentDevice()->GetDevice()->CreateQueryHeap(&Desc, IID_PPV_ARGS(OutHeap)));
+	SetName(*OutHeap, TEXT("FD3D12LinearQueryHeap"));
+
+#if ENABLE_RESIDENCY_MANAGEMENT
+	D3DX12Residency::Initialize(OutResidencyHandle, *OutHeap, ResultSize * Desc.Count);
+	D3DX12Residency::BeginTrackingObject(GetParentDevice()->GetResidencyManager(), OutResidencyHandle);
+#endif
+}
+
+void FD3D12LinearQueryHeap::CreateResultBuffer(uint64 SizeInBytes, FD3D12Resource** OutBuffer)
+{
+	FD3D12Adapter* Adapter = GetParentDevice()->GetParentAdapter();
+	const D3D12_HEAP_PROPERTIES ResultBufferHeapProperties = CD3DX12_HEAP_PROPERTIES(D3D12_HEAP_TYPE_READBACK, (uint32)GetGPUMask(), (uint32)GetVisibilityMask());
+	const D3D12_RESOURCE_DESC ResultBufferDesc = CD3DX12_RESOURCE_DESC::Buffer(SizeInBytes);
+
+	VERIFYD3D12RESULT(Adapter->CreateCommittedResource(
+		ResultBufferDesc,
+		ResultBufferHeapProperties,
+		D3D12_RESOURCE_STATE_COPY_DEST,
+		nullptr,
+		OutBuffer,
+		TEXT("FD3D12LinearQueryHeap Result Buffer")));
+}
+
+void FD3D12LinearQueryHeap::ReleaseResources()
+{
+#if ENABLE_RESIDENCY_MANAGEMENT
+	const int32 NumChunks = NextChunkIdx;
+	for (int32 Idx = 0; Idx < NumChunks; ++Idx)
+	{
+		FChunk& Chunk = AllocatedChunks[Idx];
+		if (D3DX12Residency::IsInitialized(Chunk.QueryHeapResidencyHandle))
+		{
+			D3DX12Residency::EndTrackingObject(GetParentDevice()->GetResidencyManager(), Chunk.QueryHeapResidencyHandle);
+			Chunk.QueryHeapResidencyHandle = {};
+		}
+	}
+#endif
+	NextChunkIdx = 0;
+	CurMaxNumQueries = 0;
+}
+
 /*=============================================================================
  * class FD3D12BufferedGPUTiming
  *=============================================================================*/
@@ -671,7 +866,9 @@ uint64 FD3D12BufferedGPUTiming::GetTiming(bool bGetCurrentResultsAndBlock)
 					
 					if (EndTime > StartTime)
 					{
-						return EndTime - StartTime;
+						const uint64 Bubble = GetParentAdapter()->GetGPUProfiler().CalculateIdleTime(StartTime, EndTime);
+						const uint64 ElapsedTime = EndTime - StartTime;
+						return ElapsedTime >= Bubble ? ElapsedTime - Bubble : 0;
 					}
 				}
 
@@ -719,7 +916,9 @@ uint64 FD3D12BufferedGPUTiming::GetTiming(bool bGetCurrentResultsAndBlock)
 
 			if (EndTime > StartTime)
 			{
-				return EndTime - StartTime;
+				const uint64 Bubble = GetParentAdapter()->GetGPUProfiler().CalculateIdleTime(StartTime, EndTime);
+				const uint64 ElapsedTime = EndTime - StartTime;
+				return ElapsedTime >= Bubble ? ElapsedTime - Bubble : 0;
 			}
 		}
 	}
diff --git a/Engine/Source/Runtime/D3D12RHI/Private/D3D12Query.h b/Engine/Source/Runtime/D3D12RHI/Private/D3D12Query.h
--- a/Engine/Source/Runtime/D3D12RHI/Private/D3D12Query.h
+++ b/Engine/Source/Runtime/D3D12RHI/Private/D3D12Query.h
@@ -151,4 +151,84 @@ private:
 	FD3D12ResidencyHandle QueryHeapResidencyHandle;
 	TRefCountPtr<FD3D12Resource> ResultBuffer;  // The buffer where all query results are stored
 	void* pResultData;
-};
\ No newline at end of file
+};
+
+/**
+ * A simple linear query allocator.
+ * Never resolve or cleanup until results are explicitly requested.
+ * Begin/EndQuery are thread-safe but other methods are not. Make sure no thread may
+ * call Begin/EndQuery before calling FlushAndGetResults.
+ * Only used in ProfileGPU to hold command list start/end timestamp queries currently
+ */
+class FD3D12LinearQueryHeap final : public FD3D12DeviceChild, public FD3D12SingleNodeGPUObject
+{
+public:
+	enum EHeapState
+	{
+		HS_Open,
+		HS_Closed
+	};
+
+	FD3D12LinearQueryHeap(class FD3D12Device* InParent, D3D12_QUERY_HEAP_TYPE InHeapType, int32 GrowCount);
+	~FD3D12LinearQueryHeap();
+
+	/**
+	 * Allocate a slot on query heap and queue a BeginQuery command to the given list
+	 * @param CmdListHandle - a handle to the command list where BeginQuery will be called
+	 * @return index of the allocated query
+	 */
+	int32 BeginQuery(FD3D12CommandListHandle CmdListHandle);
+
+	/**
+	* Allocate a slot on query heap and queue an EndQuery command to the given list
+	* @param CmdListHandle - a handle to the command list where EndQuery will be called
+	* @return index of the allocated query
+	*/
+	int32 EndQuery(FD3D12CommandListHandle CmdListHandle);
+
+	/** Get results of all allocated queries and reset */
+	void FlushAndGetResults(TArray<uint64>& QueryResults, bool bReleaseResources = true);
+
+private:
+	struct FChunk
+	{
+		TRefCountPtr<ID3D12QueryHeap> QueryHeap;
+		FD3D12ResidencyHandle QueryHeapResidencyHandle;
+	};
+
+	static D3D12_QUERY_TYPE HeapTypeToQueryType(D3D12_QUERY_HEAP_TYPE HeapType);
+
+	/** Release all allocated query */
+	void Reset();
+
+	/** Returns an index to the allocated heap slot */
+	int32 AllocateQueryHeapSlot();
+	
+	/** Grow the allocator's backing memory */
+	void Grow();
+
+	/** Helper to create a new query heap */
+	void CreateQueryHeap(int32 NumQueries, ID3D12QueryHeap** OutHeap, FD3D12ResidencyHandle& OutResidencyHandle);
+
+	/** Helper to create a readback buffer used to hold query results */
+	void CreateResultBuffer(uint64 SizeInBytes, FD3D12Resource** OutBuffer);
+
+	/** Release all allocated query heaps and detach them from residency manager */
+	void ReleaseResources();
+
+	/** This allocator can allocate up to (MaxNumChunks * GrowNumQueries) queries before a manual flush is needed */
+	static constexpr int32 MaxNumChunks = 8;
+	/** Size in bytes of a single query result */
+	static constexpr SIZE_T ResultSize = sizeof(uint64);
+
+	const D3D12_QUERY_HEAP_TYPE QueryHeapType;
+	const D3D12_QUERY_TYPE QueryType;
+	const int32 GrowNumQueries;
+	const int32 SlotToHeapIdxShift;
+	EHeapState HeapState;
+	volatile int32 NextFreeIdx;
+	volatile int32 CurMaxNumQueries;
+	volatile int32 NextChunkIdx;
+	FChunk AllocatedChunks[MaxNumChunks];
+	FCriticalSection CS;
+};
diff --git a/Engine/Source/Runtime/D3D12RHI/Private/D3D12Stats.cpp b/Engine/Source/Runtime/D3D12RHI/Private/D3D12Stats.cpp
--- a/Engine/Source/Runtime/D3D12RHI/Private/D3D12Stats.cpp
+++ b/Engine/Source/Runtime/D3D12RHI/Private/D3D12Stats.cpp
@@ -8,74 +8,72 @@ D3D12Stats.cpp:RHI Stats and timing implementation.
 #include "Engine/Engine.h"
 #include "Engine/GameViewportClient.h"
 
-namespace D3D12RHI
+void D3D12RHI::FD3DGPUProfiler::BeginFrame(FD3D12DynamicRHI* InRHI)
 {
-	void FD3DGPUProfiler::BeginFrame(FD3D12DynamicRHI* InRHI)
-	{
-		CurrentEventNode = NULL;
-		check(!bTrackingEvents);
-		check(!CurrentEventNodeFrame); // this should have already been cleaned up and the end of the previous frame
+	CurrentEventNode = NULL;
+	check(!bTrackingEvents);
+	check(!CurrentEventNodeFrame); // this should have already been cleaned up and the end of the previous frame
 
-		// update the crash tracking variables
-		static auto* CrashCollectionEnableCvar = IConsoleManager::Get().FindTConsoleVariableDataInt(TEXT("r.gpucrash.collectionenable"));
-		static auto* CrashCollectionDataDepth = IConsoleManager::Get().FindTConsoleVariableDataInt(TEXT("r.gpucrash.datadepth"));
+	// update the crash tracking variables
+	static auto* CrashCollectionEnableCvar = IConsoleManager::Get().FindTConsoleVariableDataInt(TEXT("r.gpucrash.collectionenable"));
+	static auto* CrashCollectionDataDepth = IConsoleManager::Get().FindTConsoleVariableDataInt(TEXT("r.gpucrash.datadepth"));
 
-		bTrackingGPUCrashData = CrashCollectionEnableCvar ? CrashCollectionEnableCvar->GetValueOnRenderThread() != 0 : false;
-		GPUCrashDataDepth = CrashCollectionDataDepth ? CrashCollectionDataDepth->GetValueOnRenderThread() : -1;
+	bTrackingGPUCrashData = CrashCollectionEnableCvar ? CrashCollectionEnableCvar->GetValueOnRenderThread() != 0 : false;
+	GPUCrashDataDepth = CrashCollectionDataDepth ? CrashCollectionDataDepth->GetValueOnRenderThread() : -1;
 
-									   // latch the bools from the game thread into our private copy
-		bLatchedGProfilingGPU = GTriggerGPUProfile;
-		bLatchedGProfilingGPUHitches = GTriggerGPUHitchProfile;
-		if (bLatchedGProfilingGPUHitches)
-		{
-			bLatchedGProfilingGPU = false; // we do NOT permit an ordinary GPU profile during hitch profiles
-		}
+	// latch the bools from the game thread into our private copy
+	bLatchedGProfilingGPU = GTriggerGPUProfile;
+	bLatchedGProfilingGPUHitches = GTriggerGPUHitchProfile;
+	if (bLatchedGProfilingGPUHitches)
+	{
+		bLatchedGProfilingGPU = false; // we do NOT permit an ordinary GPU profile during hitch profiles
+	}
 
-		// if we are starting a hitch profile or this frame is a gpu profile, then save off the state of the draw events
-		if (bLatchedGProfilingGPU || (!bPreviousLatchedGProfilingGPUHitches && bLatchedGProfilingGPUHitches))
-		{
-			bOriginalGEmitDrawEvents = GetEmitDrawEvents();
-		}
+	// if we are starting a hitch profile or this frame is a gpu profile, then save off the state of the draw events
+	if (bLatchedGProfilingGPU || (!bPreviousLatchedGProfilingGPUHitches && bLatchedGProfilingGPUHitches))
+	{
+		bOriginalGEmitDrawEvents = GetEmitDrawEvents();
+	}
 
-		if (bLatchedGProfilingGPU || bLatchedGProfilingGPUHitches)
+	if (bLatchedGProfilingGPU || bLatchedGProfilingGPUHitches)
+	{
+		if (bLatchedGProfilingGPUHitches && GPUHitchDebounce)
 		{
-			if (bLatchedGProfilingGPUHitches && GPUHitchDebounce)
-			{
-				// if we are doing hitches and we had a recent hitch, wait to recover
-				// the reasoning is that collecting the hitch report may itself hitch the GPU
-				GPUHitchDebounce--;
-			}
-			else
-			{
-				SetEmitDrawEvents(true);  // thwart an attempt to turn this off on the game side
-				bTrackingEvents = true;
-				CurrentEventNodeFrame = new FD3D12EventNodeFrame(GetParentAdapter());
-				CurrentEventNodeFrame->StartFrame();
-			}
+			// if we are doing hitches and we had a recent hitch, wait to recover
+			// the reasoning is that collecting the hitch report may itself hitch the GPU
+			GPUHitchDebounce--;
 		}
-		else if (bPreviousLatchedGProfilingGPUHitches)
+		else
 		{
-			// hitch profiler is turning off, clear history and restore draw events
-			GPUHitchEventNodeFrames.Empty();
-			SetEmitDrawEvents(bOriginalGEmitDrawEvents);
+			SetEmitDrawEvents(true);  // thwart an attempt to turn this off on the game side
+			bTrackingEvents = true;
+			DoPreProfileGPUWork();
+			CurrentEventNodeFrame = new FD3D12EventNodeFrame(GetParentAdapter());
+			CurrentEventNodeFrame->StartFrame();
 		}
-		bPreviousLatchedGProfilingGPUHitches = bLatchedGProfilingGPUHitches;
+	}
+	else if (bPreviousLatchedGProfilingGPUHitches)
+	{
+		// hitch profiler is turning off, clear history and restore draw events
+		GPUHitchEventNodeFrames.Empty();
+		SetEmitDrawEvents(bOriginalGEmitDrawEvents);
+	}
+	bPreviousLatchedGProfilingGPUHitches = bLatchedGProfilingGPUHitches;
 
-		FrameTiming.StartTiming();
+	FrameTiming.StartTiming();
 
-		if (GetEmitDrawEvents())
-		{
+	if (GetEmitDrawEvents())
+	{
 #if NV_AFTERMATH
-			// Assuming that grabbing the device 0 command list here is OK
-			PushEvent(TEXT("FRAME"), FColor(0, 255, 0, 255), InRHI->GetAdapter().GetDevice(0)->GetCommandContext().CommandListHandle.AftermathCommandContext());
+		// Assuming that grabbing the device 0 command list here is OK
+		PushEvent(TEXT("FRAME"), FColor(0, 255, 0, 255), InRHI->GetAdapter().GetDevice(0)->GetCommandContext().CommandListHandle.AftermathCommandContext());
 #else
-			PushEvent(TEXT("FRAME"), FColor(0, 255, 0, 255));
+		PushEvent(TEXT("FRAME"), FColor(0, 255, 0, 255));
 #endif
-		}
 	}
 }
 
-void FD3DGPUProfiler::EndFrame(FD3D12DynamicRHI* InRHI)
+void D3D12RHI::FD3DGPUProfiler::EndFrame(FD3D12DynamicRHI* InRHI)
 {
 	if (GetEmitDrawEvents())
 	{
@@ -115,6 +113,7 @@ void FD3DGPUProfiler::EndFrame(FD3D12DynamicRHI* InRHI)
 		if (bTrackingEvents)
 		{
 			SetEmitDrawEvents(bOriginalGEmitDrawEvents);
+			DoPostProfileGPUWork();
 			UE_LOG(LogD3D12RHI, Log, TEXT(""));
 			UE_LOG(LogD3D12RHI, Log, TEXT(""));
 			CurrentEventNodeFrame->DumpEventTree();
@@ -190,7 +189,7 @@ void FD3DGPUProfiler::EndFrame(FD3D12DynamicRHI* InRHI)
 	CurrentEventNodeFrame = NULL;
 }
 
-void FD3DGPUProfiler::PushEvent(const TCHAR* Name, FColor Color)
+void D3D12RHI::FD3DGPUProfiler::PushEvent(const TCHAR* Name, FColor Color)
 {
 #if WITH_DX_PERF
 	D3DPERF_BeginEvent(Color.DWColor(), Name);
@@ -202,7 +201,7 @@ void FD3DGPUProfiler::PushEvent(const TCHAR* Name, FColor Color)
 static FString EventDeepString(TEXT("EventTooDeep"));
 static const uint32 EventDeepCRC = FCrc::StrCrc32<TCHAR>(*EventDeepString);
 #if NV_AFTERMATH
-void FD3DGPUProfiler::PushEvent(const TCHAR* Name, FColor Color, GFSDK_Aftermath_ContextHandle Context)
+void D3D12RHI::FD3DGPUProfiler::PushEvent(const TCHAR* Name, FColor Color, GFSDK_Aftermath_ContextHandle Context)
 {
 
 	if (GDX12NVAfterMathEnabled && bTrackingGPUCrashData)
@@ -237,7 +236,7 @@ void FD3DGPUProfiler::PushEvent(const TCHAR* Name, FColor Color, GFSDK_Aftermath
 }
 #endif
 
-void FD3DGPUProfiler::PopEvent()
+void D3D12RHI::FD3DGPUProfiler::PopEvent()
 {
 #if WITH_DX_PERF
 	D3DPERF_EndEvent();
@@ -372,14 +371,14 @@ void UpdateBufferStats(FD3D12ResourceLocation* ResourceLocation, bool bAllocatin
 }
 
 #if NV_AFTERMATH
-void FD3DGPUProfiler::RegisterCommandList(GFSDK_Aftermath_ContextHandle context)
+void D3D12RHI::FD3DGPUProfiler::RegisterCommandList(GFSDK_Aftermath_ContextHandle context)
 {
 	FScopeLock Lock(&AftermathLock);
 
 	AftermathContexts.Push(context);
 }
 
-void FD3DGPUProfiler::UnregisterCommandList(GFSDK_Aftermath_ContextHandle context)
+void D3D12RHI::FD3DGPUProfiler::UnregisterCommandList(GFSDK_Aftermath_ContextHandle context)
 {
 	FScopeLock Lock(&AftermathLock);
 
@@ -390,7 +389,7 @@ void FD3DGPUProfiler::UnregisterCommandList(GFSDK_Aftermath_ContextHandle contex
 #endif
 
 extern CORE_API bool GIsGPUCrashed;
-bool FD3DGPUProfiler::CheckGpuHeartbeat() const
+bool D3D12RHI::FD3DGPUProfiler::CheckGpuHeartbeat() const
 {
 #if NV_AFTERMATH
 	if (GDX12NVAfterMathEnabled)
@@ -461,3 +460,82 @@ bool FD3DGPUProfiler::CheckGpuHeartbeat() const
 #endif
 	return true;
 }
+
+static int32 FindCmdListTimingPairIndex(const TArray<uint64>& CmdListStartTimestamps, uint64 Value)
+{
+	int32 Pos = Algo::UpperBound(CmdListStartTimestamps, Value) - 1;
+	return FMath::Max(Pos, 0);
+}
+
+uint64 D3D12RHI::FD3DGPUProfiler::CalculateIdleTime(uint64 StartTime, uint64 EndTime)
+{
+	const int32 NumTimingPairs = CmdListStartTimestamps.Num();
+	check(NumTimingPairs == CmdListEndTimestamps.Num() && NumTimingPairs == IdleTimeCDF.Num());
+	
+	if (!NumTimingPairs)
+	{
+		return 0;
+	}
+
+	const int32 StartIdx = FindCmdListTimingPairIndex(CmdListStartTimestamps, StartTime);
+	const int32 EndIdx = FindCmdListTimingPairIndex(CmdListStartTimestamps, EndTime);
+	return IdleTimeCDF[EndIdx] - IdleTimeCDF[StartIdx];
+}
+
+void D3D12RHI::FD3DGPUProfiler::DoPreProfileGPUWork()
+{
+	typedef typename FD3D12CommandContext::EFlushCommandsExtraAction EFlushCmdsAction;
+	constexpr bool bWaitForCommands = false;
+	constexpr EFlushCmdsAction FlushAction = EFlushCmdsAction::FCEA_StartProfilingGPU;
+
+	FD3D12Adapter* Adapter = GetParentAdapter();
+	for (uint32 GPUIdx : FRHIGPUMask::All())
+	{
+		FD3D12Device* Device = Adapter->GetDevice(GPUIdx);
+		Device->GetDefaultCommandContext().FlushCommands(bWaitForCommands, FlushAction);
+	}
+}
+
+void D3D12RHI::FD3DGPUProfiler::DoPostProfileGPUWork()
+{
+	typedef typename FD3D12CommandContext::EFlushCommandsExtraAction EFlushCmdsAction;
+	constexpr bool bWaitForCommands = false;
+	constexpr EFlushCmdsAction FlushAction = EFlushCmdsAction::FCEA_EndProfilingGPU;
+
+	TArray<FResolvedCmdListExecTime> CmdListExecTimes;
+	FD3D12Adapter* Adapter = GetParentAdapter();
+	for (uint32 GPUIdx : FRHIGPUMask::All())
+	{
+		FD3D12Device* Device = Adapter->GetDevice(GPUIdx);
+		Device->GetDefaultCommandContext().FlushCommands(bWaitForCommands, FlushAction);
+		TArray<FResolvedCmdListExecTime> TimingPairs;
+		Device->GetCommandListManager().GetCommandListTimingResults(TimingPairs);
+		CmdListExecTimes.Append(MoveTemp(TimingPairs));
+	}
+
+	const int32 NumTimingPairs = CmdListExecTimes.Num();
+	CmdListStartTimestamps.Empty(NumTimingPairs);
+	CmdListEndTimestamps.Empty(NumTimingPairs);
+	IdleTimeCDF.Empty(NumTimingPairs);
+
+	if (NumTimingPairs > 0)
+	{
+		Algo::Sort(CmdListExecTimes, [](const FResolvedCmdListExecTime& A, const FResolvedCmdListExecTime& B)
+		{
+			return A.StartTimestamp < B.StartTimestamp;
+		});
+		CmdListStartTimestamps.Add(CmdListExecTimes[0].StartTimestamp);
+		CmdListEndTimestamps.Add(CmdListExecTimes[0].EndTimestamp);
+		IdleTimeCDF.Add(0);
+		for (int32 Idx = 1; Idx < NumTimingPairs; ++Idx)
+		{
+			const FResolvedCmdListExecTime& Prev = CmdListExecTimes[Idx - 1];
+			const FResolvedCmdListExecTime& Cur = CmdListExecTimes[Idx];
+			ensure(Cur.StartTimestamp >= Prev.EndTimestamp);
+			CmdListStartTimestamps.Add(Cur.StartTimestamp);
+			CmdListEndTimestamps.Add(Cur.EndTimestamp);
+			const uint64 Bubble = Cur.StartTimestamp >= Prev.EndTimestamp ? Cur.StartTimestamp - Prev.EndTimestamp : 0;
+			IdleTimeCDF.Add(IdleTimeCDF.Last() + Bubble);
+		}
+	}
+}
diff --git a/Engine/Source/Runtime/D3D12RHI/Private/D3D12Stats.h b/Engine/Source/Runtime/D3D12RHI/Private/D3D12Stats.h
--- a/Engine/Source/Runtime/D3D12RHI/Private/D3D12Stats.h
+++ b/Engine/Source/Runtime/D3D12RHI/Private/D3D12Stats.h
@@ -328,6 +328,14 @@ namespace D3D12RHI
 
 		bool CheckGpuHeartbeat() const;
 
+		/**
+		 * Calculate the amount of GPU idle time between two timestamps
+		 * @param StartTime - start timestamp
+		 * @param EndTime - end timestamp
+		 * @return number of idle GPU clock ticks between or 0 if command list execution time isn't tracked
+		 */
+		uint64 CalculateIdleTime(uint64 StartTime, uint64 EndTime);
+
 #if NV_AFTERMATH
 		virtual void PushEvent(const TCHAR* Name, FColor Color, GFSDK_Aftermath_ContextHandle Context);
 
@@ -343,5 +351,21 @@ namespace D3D12RHI
 
 		/** Used to measure GPU time per frame. */
 		FD3D12BufferedGPUTiming FrameTiming;
+
+	private:
+		/** Flush existing command lists and start command list execution time tracking */
+		void DoPreProfileGPUWork();
+
+		/** Flush existing command lists and obtain timing results of all tracked command lists */
+		void DoPostProfileGPUWork();
+
+		typedef typename FD3D12CommandListManager::FResolvedCmdListExecTime FResolvedCmdListExecTime;
+
+		/** Timstamps marking the beginning of tracked command lists */
+		TArray<uint64> CmdListStartTimestamps;
+		/** Timstamps marking the end of tracked command lists */
+		TArray<uint64> CmdListEndTimestamps;
+		/** Accumulated idle GPU ticks before each corresponding command list */
+		TArray<uint64> IdleTimeCDF;
 	};
 }
\ No newline at end of file
