Fix for crash on launch on DX12

diff --git a/Engine/Source/Runtime/D3D12RHI/Private/Windows/WindowsD3D12Device.cpp b/Engine/Source/Runtime/D3D12RHI/Private/Windows/WindowsD3D12Device.cpp
--- a/Engine/Source/Runtime/D3D12RHI/Private/Windows/WindowsD3D12Device.cpp
+++ b/Engine/Source/Runtime/D3D12RHI/Private/Windows/WindowsD3D12Device.cpp
@@ -15,6 +15,8 @@
 #include "Runtime/HeadMountedDisplay/Public/IHeadMountedDisplayModule.h"
 #include "GenericPlatform/GenericPlatformDriver.h"			// FGPUDriverInfo
 
+#include "ShaderCompiler.h"
+
 #pragma comment(lib, "d3d12.lib")
 
 IMPLEMENT_MODULE(FD3D12DynamicRHIModule, D3D12RHI);
@@ -700,6 +702,17 @@ void FD3D12DynamicRHI::Init()
 
 void FD3D12DynamicRHI::PostInit()
 {
+	if (GRHISupportsRayTracing || GRHISupportsRHIThread)
+	{
+		// Make sure all global shaders are complete at this point
+		extern RENDERCORE_API const int32 GlobalShaderMapId;
+
+		TArray<int32> ShaderMapIds;
+		ShaderMapIds.Add(GlobalShaderMapId);
+
+		GShaderCompilingManager->FinishCompilation(TEXT("Global"), ShaderMapIds);
+	}
+
 	if (GRHISupportsRayTracing)
 	{
 		for (FD3D12Adapter*& Adapter : ChosenAdapters)
diff --git a/Engine/Source/Runtime/Engine/Private/ShaderCompiler/ShaderCompiler.cpp b/Engine/Source/Runtime/Engine/Private/ShaderCompiler/ShaderCompiler.cpp
--- a/Engine/Source/Runtime/Engine/Private/ShaderCompiler/ShaderCompiler.cpp
+++ b/Engine/Source/Runtime/Engine/Private/ShaderCompiler/ShaderCompiler.cpp
@@ -3918,9 +3918,6 @@ void VerifyGlobalShaders(EShaderPlatform Platform, bool bLoadedFromCacheFile)
 			// Metal also needs this when using RHI thread because it uses TOneColorVS very early in RHIPostInit()
 			!IsOpenGLPlatform(GMaxRHIShaderPlatform) && !IsVulkanPlatform(GMaxRHIShaderPlatform) &&
 			!IsMetalPlatform(GMaxRHIShaderPlatform) && !IsSwitchPlatform(GMaxRHIShaderPlatform) &&
-#if RHI_RAYTRACING
-			!GRHISupportsRayTracing && //This is here because DXR is caching its BuiltIn Shaders in PostInit (see FD3D12Device::InitRayTracing)
-#endif // RHI_RAYTRACING
 			GShaderCompilingManager->AllowAsynchronousShaderCompiling();
 
 		if (!bAllowAsynchronousGlobalShaderCompiling)
