Fix data upload to UAV buffer

diff --git a/Engine/Source/Runtime/D3D12RHI/Private/D3D12RHI.cpp b/Engine/Source/Runtime/D3D12RHI/Private/D3D12RHI.cpp
--- a/Engine/Source/Runtime/D3D12RHI/Private/D3D12RHI.cpp
+++ b/Engine/Source/Runtime/D3D12RHI/Private/D3D12RHI.cpp
@@ -356,7 +356,7 @@ void FD3D12DynamicRHI::UpdateBuffer(FD3D12Resource* Dest, uint32 DestOffset, FD3
 	FD3D12CommandContext& DefaultContext = Device->GetDefaultCommandContext();
 	FD3D12CommandListHandle& hCommandList = DefaultContext.CommandListHandle;
 
-	FScopeResourceBarrier ScopeResourceBarrierDest(hCommandList, Dest, Dest->GetDefaultResourceState(), D3D12_RESOURCE_STATE_COPY_DEST, 0);
+	FConditionalScopeResourceBarrier ScopeResourceBarrierDest(hCommandList, Dest, D3D12_RESOURCE_STATE_COPY_DEST, 0);
 	// Don't need to transition upload heaps
 
 	DefaultContext.numCopies++;
