Remove unused member

diff --git a/Engine/Source/Runtime/D3D12RHI/Public/D3D12Resources.h b/Engine/Source/Runtime/D3D12RHI/Public/D3D12Resources.h
--- a/Engine/Source/Runtime/D3D12RHI/Public/D3D12Resources.h
+++ b/Engine/Source/Runtime/D3D12RHI/Public/D3D12Resources.h
@@ -705,9 +705,6 @@ public:
 	}
 
 	virtual ~FD3D12UniformBuffer();
-
-private:
-	class FD3D12DynamicRHI* D3D12RHI;
 };
 
 #if PLATFORM_WINDOWS
