Refix final-release while in-flight fix as previously fixed in CL 4691540

diff --git a/Engine/Source/Runtime/D3D12RHI/Private/D3D12Commands.cpp b/Engine/Source/Runtime/D3D12RHI/Private/D3D12Commands.cpp
--- a/Engine/Source/Runtime/D3D12RHI/Private/D3D12Commands.cpp
+++ b/Engine/Source/Runtime/D3D12RHI/Private/D3D12Commands.cpp
@@ -410,11 +410,9 @@ void FD3D12CommandContext::RHICopyToStagingBuffer(FVertexBufferRHIParamRef Sourc
 	{
 		// @todo-mattc I feel like we should allocate more than NumBytes to handle small reads without blowing tons of space. Need to pool this.
 		// Hopefully d3d12 will do smart pooling out of an internal heap.
-		if (StagingBuffer->StagedRead)
-		{
-			StagingBuffer->StagedRead.SafeRelease();
-		}
-		VERIFYD3D12RESULT(GetParentDevice()->GetParentAdapter()->CreateBuffer(D3D12_HEAP_TYPE_READBACK, GetGPUMask(), GetGPUMask(), NumBytes, StagingBuffer->StagedRead.GetInitReference()));
+		StagingBuffer->SafeRelease();
+
+		VERIFYD3D12RESULT(GetParentDevice()->GetParentAdapter()->CreateBuffer(D3D12_HEAP_TYPE_READBACK, GetGPUMask(), GetGPUMask(), NumBytes, &StagingBuffer->StagedRead));
 		StagingBuffer->ShadowBufferSize = NumBytes;
 	}
 
diff --git a/Engine/Source/Runtime/D3D12RHI/Private/D3D12UAV.cpp b/Engine/Source/Runtime/D3D12RHI/Private/D3D12UAV.cpp
--- a/Engine/Source/Runtime/D3D12RHI/Private/D3D12UAV.cpp
+++ b/Engine/Source/Runtime/D3D12RHI/Private/D3D12UAV.cpp
@@ -266,7 +266,7 @@ FD3D12StagingBuffer::~FD3D12StagingBuffer()
 {
 	if (StagedRead)
 	{
-		StagedRead.SafeRelease();
+		StagedRead->DeferDelete();
 	}
 }
 
@@ -274,13 +274,12 @@ void* FD3D12StagingBuffer::Lock(uint32 Offset, uint32 NumBytes)
 {
 	check(!bIsLocked);
 	bIsLocked = true;
-	FD3D12Resource* pResource = StagedRead.GetReference();
-	if (pResource)
+	if (StagedRead)
 	{
 		D3D12_RANGE ReadRange;
 		ReadRange.Begin = Offset;
 		ReadRange.End = Offset + NumBytes;
-		return reinterpret_cast<uint8*>(pResource->Map(&ReadRange)) + Offset;
+		return reinterpret_cast<uint8*>(StagedRead->Map(&ReadRange)) + Offset;
 	}
 	else
 	{
@@ -292,9 +291,8 @@ void FD3D12StagingBuffer::Unlock()
 {
 	check(bIsLocked);
 	bIsLocked = false;
-	FD3D12Resource* pResource = StagedRead.GetReference();
-	if (pResource)
+	if (StagedRead)
 	{
-		pResource->Unmap();
+		StagedRead->Unmap();
 	}
 }
\ No newline at end of file
diff --git a/Engine/Source/Runtime/D3D12RHI/Public/D3D12Resources.h b/Engine/Source/Runtime/D3D12RHI/Public/D3D12Resources.h
--- a/Engine/Source/Runtime/D3D12RHI/Public/D3D12Resources.h
+++ b/Engine/Source/Runtime/D3D12RHI/Public/D3D12Resources.h
@@ -951,15 +951,25 @@ class FD3D12StagingBuffer : public FRHIStagingBuffer
 public:
 	FD3D12StagingBuffer()
 		: FRHIStagingBuffer()
+		, StagedRead(nullptr)
 		, ShadowBufferSize(0)
 	{}
 	virtual ~FD3D12StagingBuffer() final override;
 
+	void SafeRelease()
+	{
+		if (StagedRead)
+		{
+			StagedRead->Release();
+			StagedRead = nullptr;
+		}
+	}
+
 	virtual void* Lock(uint32 Offset, uint32 NumBytes) final override;
 	virtual void Unlock() final override;
 
 private:
-	TRefCountPtr<FD3D12Resource> StagedRead;
+	FD3D12Resource* StagedRead;
 	uint32 ShadowBufferSize;
 };
 
