Added support for multiple raygen and miss shaders in single ray tracing pipeline state object and multiple hit shaders per primitive (for multiple ray types).
Removed explicit defalt closest hit shader from FRayTracingPipelineStateInitializer. Desired behavior can now be achieved by providing a single shader in the hit group table and setting hit group stride to 0 (effectively disabling shader binding table indexing).
GetRayTracingMaterialLibrary() now ensures that returned shader is contiguous, with no invalid/null entries. Invalid slots are filled with a shader that's explicitly provided by the caller.

diff --git a/Engine/Source/Runtime/D3D12RHI/Private/D3D12CommandContext.h b/Engine/Source/Runtime/D3D12RHI/Private/D3D12CommandContext.h
--- a/Engine/Source/Runtime/D3D12RHI/Private/D3D12CommandContext.h
+++ b/Engine/Source/Runtime/D3D12RHI/Private/D3D12CommandContext.h
@@ -360,15 +360,15 @@ public:
 		FShaderResourceViewRHIParamRef Rays,
 		FUnorderedAccessViewRHIParamRef Output,
 		uint32 NumRays) final override;
-	virtual void RHIRayTraceDispatch(FRayTracingPipelineStateRHIParamRef RayTracingPipelineState,
+	virtual void RHIRayTraceDispatch(FRayTracingPipelineStateRHIParamRef RayTracingPipelineState, uint32 RayGenShaderIndex,
 		const FRayTracingShaderBindings& GlobalResourceBindings,
 		uint32 Width, uint32 Height) final override;
-	virtual void RHIRayTraceDispatch(FRayTracingPipelineStateRHIParamRef RayTracingPipelineState,
+	virtual void RHIRayTraceDispatch(FRayTracingPipelineStateRHIParamRef RayTracingPipelineState, uint32 RayGenShaderIndex,
 		FRayTracingSceneRHIParamRef Scene, // #dxr_todo: replace this with explicit shader table parameter
 		const FRayTracingShaderBindings& GlobalResourceBindings,
 		uint32 Width, uint32 Height) final override;
 	virtual void RHISetRayTracingHitGroup(
-		FRayTracingSceneRHIParamRef Scene, uint32 InstanceIndex, uint32 SegmentIndex,
+		FRayTracingSceneRHIParamRef Scene, uint32 InstanceIndex, uint32 SegmentIndex, uint32 ShaderSlot,
 		FRayTracingPipelineStateRHIParamRef Pipeline, uint32 HitGroupIndex,
 		uint32 NumUniformBuffers, const FUniformBufferRHIParamRef* UniformBuffers) final override;
 #endif // D3D12_RHI_RAYTRACING
diff --git a/Engine/Source/Runtime/D3D12RHI/Private/D3D12RayTracing.cpp b/Engine/Source/Runtime/D3D12RHI/Private/D3D12RayTracing.cpp
--- a/Engine/Source/Runtime/D3D12RHI/Private/D3D12RayTracing.cpp
+++ b/Engine/Source/Runtime/D3D12RHI/Private/D3D12RayTracing.cpp
@@ -40,9 +40,6 @@ struct FHitGroupSystemParameters
 
 struct FD3D12ShaderIdentifier
 {
-	// Shader identifier size is statically defined on RS5, but is dynamic on RS4.
-	// RS5 size is always D3D12_SHADER_IDENTIFIER_SIZE_IN_BYTES (32) and we assume that it will never be bigger than that on RS4.
-
 	uint64 Data[4] = {~0ull, ~0ull, ~0ull, ~0ull};
 
 	bool operator == (const FD3D12ShaderIdentifier& Other) const
@@ -62,6 +59,17 @@ struct FD3D12ShaderIdentifier
 	{
 		return *this != FD3D12ShaderIdentifier();
 	}
+
+	// No shader is executed if a shader binding table record with null identifier is encountered.
+	void SetNull()
+	{
+		Data[3] = Data[2] = Data[1] = Data[0] = 0ull;
+	}
+
+	void SetData(const void* InData)
+	{
+		FMemory::Memcpy(Data, InData, sizeof(Data));
+	}
 };
 
 static_assert(sizeof(FD3D12ShaderIdentifier) == D3D12_SHADER_IDENTIFIER_SIZE_IN_BYTES, "Unexpected shader identifier size");
@@ -516,14 +524,14 @@ private:
 		bIsDirty = true;
 	}
 
-	void WriteHitGroupRecord(uint32 RecordIndex, uint32 OffsetWithinRecord, const void* InData, uint32 InDataSize)
+	void WriteHitRecord(uint32 RecordIndex, uint32 OffsetWithinRecord, const void* InData, uint32 InDataSize)
 	{
 		checkfSlow(OffsetWithinRecord % 4 == 0, TEXT("SBT record parameters must be written on DWORD-aligned boundary"));
 		checkfSlow(InDataSize % 4 == 0, TEXT("SBT record parameters must be DWORD-aligned"));
-		checkfSlow(OffsetWithinRecord + InDataSize <= HitGroupRecordSizeUnaligned, TEXT("SBT record write request is out of bounds"));
-		checkfSlow(RecordIndex < NumHitGroups, TEXT("SBT record write request is out of bounds"));
+		checkfSlow(OffsetWithinRecord + InDataSize <= HitRecordSizeUnaligned, TEXT("SBT record write request is out of bounds"));
+		checkfSlow(RecordIndex < NumHitRecords, TEXT("SBT record write request is out of bounds"));
 
-		const uint32 WriteOffset = HitGroupShaderTableOffset + HitGroupRecordStride * RecordIndex + OffsetWithinRecord;
+		const uint32 WriteOffset = HitGroupShaderTableOffset + HitRecordStride * RecordIndex + OffsetWithinRecord;
 
 		WriteData(WriteOffset, InData, InDataSize);
 	}
@@ -538,13 +546,13 @@ public:
 	{
 	}
 
-	void Init(uint32 InNumRayGenShaders, uint32 InNumMissShaders, uint32 InNumHitGroups, uint32 LocalRootDataSize)
+	void Init(uint32 InNumRayGenShaders, uint32 InNumMissShaders, uint32 InNumHitRecords, uint32 LocalRootDataSize)
 	{
 		checkf(LocalRootDataSize <= 4096, TEXT("The maximum size of a local root signature is 4KB.")); // as per section 4.22.1 of DXR spec v1.0
 		checkf(InNumRayGenShaders >= 1, TEXT("All shader tables must contain at least one raygen shader."));
 
-		HitGroupRecordSizeUnaligned = ShaderIdentifierSize + LocalRootDataSize;
-		HitGroupRecordStride = RoundUpToNextMultiple(HitGroupRecordSizeUnaligned, D3D12_RAYTRACING_SHADER_RECORD_BYTE_ALIGNMENT);
+		HitRecordSizeUnaligned = ShaderIdentifierSize + LocalRootDataSize;
+		HitRecordStride = RoundUpToNextMultiple(HitRecordSizeUnaligned, D3D12_RAYTRACING_SHADER_RECORD_BYTE_ALIGNMENT);
 
 		// Minimum number of descriptors required to support binding global resources (arbitrarily chosen)
 		// #dxr_todo: Remove this when RT descriptors are sub-allocated from the global view descriptor heap.
@@ -553,14 +561,14 @@ public:
 
 		// D3D12 is guaranteed to support 1M (D3D12_MAX_SHADER_VISIBLE_DESCRIPTOR_HEAP_SIZE_TIER_1) descriptors in a CBV/SRV/UAV heap, so clamp the size to this.
 		// https://docs.microsoft.com/en-us/windows/desktop/direct3d12/hardware-support
-		const uint32 NumViewDescriptors = FMath::Max(MinNumViewDescriptors, FMath::Min<uint32>(InNumHitGroups * ApproximateDescriptorsPerRecord, D3D12_MAX_SHADER_VISIBLE_DESCRIPTOR_HEAP_SIZE_TIER_1));
+		const uint32 NumViewDescriptors = FMath::Max(MinNumViewDescriptors, FMath::Min<uint32>(InNumHitRecords * ApproximateDescriptorsPerRecord, D3D12_MAX_SHADER_VISIBLE_DESCRIPTOR_HEAP_SIZE_TIER_1));
 		const uint32 NumSamplerDescriptors = D3D12_MAX_SHADER_VISIBLE_SAMPLER_HEAP_SIZE;
 
 		DescriptorCache.Init(NumViewDescriptors, NumSamplerDescriptors);
 
 		NumRayGenShaders = InNumRayGenShaders;
 		NumMissShaders = InNumMissShaders;
-		NumHitGroups = InNumHitGroups;
+		NumHitRecords = InNumHitRecords;
 
 		uint32 TotalDataSize = 0;
 
@@ -573,7 +581,7 @@ public:
 		TotalDataSize = RoundUpToNextMultiple(TotalDataSize, D3D12_RAYTRACING_SHADER_TABLE_BYTE_ALIGNMENT);
 
 		HitGroupShaderTableOffset = TotalDataSize;
-		TotalDataSize += NumHitGroups * HitGroupRecordStride;
+		TotalDataSize += InNumHitRecords * HitRecordStride;
 		TotalDataSize = RoundUpToNextMultiple(TotalDataSize, D3D12_RAYTRACING_SHADER_TABLE_BYTE_ALIGNMENT);
 
 		DefaultHitGroupShaderTableOffset = TotalDataSize;
@@ -581,36 +589,24 @@ public:
 		TotalDataSize = RoundUpToNextMultiple(TotalDataSize, D3D12_RAYTRACING_SHADER_TABLE_BYTE_ALIGNMENT);
 
 		Data.SetNumZeroed(TotalDataSize);
-
-#if DO_CHECK && DO_GUARD_SLOW
-		WrittenHitGroupRecords.Init(false, NumHitGroups);
-#endif // DO_CHECK && DO_GUARD_SLOW
 	}
 
 	template <typename T>
 	void SetHitGroupParameters(uint32 RecordIndex, uint32 InOffsetWithinRootSignature, const T& Parameters)
 	{
-		WriteHitGroupRecord(RecordIndex, ShaderIdentifierSize + InOffsetWithinRootSignature, &Parameters, sizeof(Parameters));
+		WriteHitRecord(RecordIndex, ShaderIdentifierSize + InOffsetWithinRootSignature, &Parameters, sizeof(Parameters));
 	}
 
 	void SetHitGroupParameters(uint32 RecordIndex, uint32 InOffsetWithinRootSignature, const void* InData, uint32 InDataSize)
 	{
-		WriteHitGroupRecord(RecordIndex, ShaderIdentifierSize + InOffsetWithinRootSignature, InData, InDataSize);
+		WriteHitRecord(RecordIndex, ShaderIdentifierSize + InOffsetWithinRootSignature, InData, InDataSize);
 	}
 
 	void SetHitGroupIdentifier(uint32 RecordIndex, const void* ShaderIdentifierData, uint32 InShaderIdentifierSize)
 	{
 		checkSlow(InShaderIdentifierSize == ShaderIdentifierSize);
 
-		WriteHitGroupRecord(RecordIndex, 0, ShaderIdentifierData, InShaderIdentifierSize);
-
-#if DO_CHECK && DO_GUARD_SLOW
-		if (!WrittenHitGroupRecords[RecordIndex])
-		{
-			WrittenHitGroupRecords[RecordIndex] = true;
-			NumValidHitGroupRecords++;
-		}
-#endif // DO_CHECK && DO_GUARD_SLOW
+		WriteHitRecord(RecordIndex, 0, ShaderIdentifierData, InShaderIdentifierSize);
 	}
 
 	void SetRayGenIdentifier(uint32 RecordIndex, const FD3D12ShaderIdentifier& ShaderIdentifier)
@@ -639,6 +635,24 @@ public:
 		SetHitGroupIdentifier(RecordIndex, ShaderIdentifier.Data, ShaderIdentifierSize);
 	}
 
+	void SetRayGenIdentifiers(const TArrayView<const FD3D12ShaderIdentifier>& Identifiers)
+	{
+		check(Identifiers.Num() == NumRayGenShaders);
+		for (int32 Index = 0; Index < Identifiers.Num(); ++Index)
+		{
+			SetRayGenIdentifier(Index, Identifiers[Index]);
+		}
+	}
+
+	void SetMissIdentifiers(const TArrayView<const FD3D12ShaderIdentifier>& Identifiers)
+	{
+		check(Identifiers.Num() == NumMissShaders);
+		for (int32 Index = 0; Index < Identifiers.Num(); ++Index)
+		{
+			SetMissIdentifier(Index, Identifiers[Index]);
+		}
+	}
+
 	void CopyToGPU()
 	{
 		check(IsInRHIThread() || !IsRunningRHIInSeparateThread());
@@ -646,9 +660,6 @@ public:
 		FD3D12Device* Device = GetParentDevice();
 
 		checkf(Data.Num(), TEXT("Shader table is expected to be initialized before copying to GPU."));
-#if DO_CHECK && DO_GUARD_SLOW
-		checkfSlow(NumValidHitGroupRecords == NumHitGroups, TEXT("Hit group shader table is expected to be fully populated before copying to GPU. Ensure that all shader records have been fille using WriteShaderIdentifier()."));
-#endif
 
 		FD3D12Adapter* Adapter = Device->GetParentAdapter();
 
@@ -672,7 +683,7 @@ public:
 		return Buffer->ResourceLocation.GetGPUVirtualAddress();
 	}
 
-	D3D12_DISPATCH_RAYS_DESC GetDispatchRaysDesc(uint32 RayGenShaderIndex = 0, uint32 MissShaderBaseIndex = 0) const
+	D3D12_DISPATCH_RAYS_DESC GetDispatchRaysDesc(uint32 RayGenShaderIndex, uint32 MissShaderBaseIndex, uint32 HitGroupStride) const
 	{
 		D3D12_GPU_VIRTUAL_ADDRESS ShaderTableAddress = GetShaderTableAddress();
 
@@ -685,11 +696,11 @@ public:
 		Desc.MissShaderTable.StrideInBytes = MissRecordStride;
 		Desc.MissShaderTable.SizeInBytes = MissRecordStride;
 
-		if (NumHitGroups)
+		if (HitGroupStride)
 		{
 			Desc.HitGroupTable.StartAddress = ShaderTableAddress + HitGroupShaderTableOffset;
-			Desc.HitGroupTable.StrideInBytes = HitGroupRecordStride;
-			Desc.HitGroupTable.SizeInBytes = NumHitGroups * HitGroupRecordStride;
+			Desc.HitGroupTable.StrideInBytes = HitRecordStride;
+			Desc.HitGroupTable.SizeInBytes = NumHitRecords * HitRecordStride;
 		}
 		else
 		{
@@ -702,7 +713,7 @@ public:
 	}
 
 	static constexpr uint32 ShaderIdentifierSize = D3D12_SHADER_IDENTIFIER_SIZE_IN_BYTES;
-	uint32 NumHitGroups = 0;
+	uint32 NumHitRecords = 0;
 	uint32 NumRayGenShaders = 0;
 	uint32 NumMissShaders = 0;
 
@@ -716,15 +727,10 @@ public:
 	static constexpr uint32 RayGenRecordStride = D3D12_RAYTRACING_SHADER_TABLE_BYTE_ALIGNMENT;
 	static constexpr uint32 MissRecordStride = D3D12_RAYTRACING_SHADER_TABLE_BYTE_ALIGNMENT;
 
-	uint32 HitGroupRecordSizeUnaligned = 0; // size of the shader identifier + local root parameters, not aligned to SHADER_RECORD_BYTE_ALIGNMENT (used for out-of-bounds access checks)
-	uint32 HitGroupRecordStride = 0; // size of shader identifier + local root parameters, aligned to SHADER_RECORD_BYTE_ALIGNMENT
+	uint32 HitRecordSizeUnaligned = 0; // size of the shader identifier + local root parameters, not aligned to SHADER_RECORD_BYTE_ALIGNMENT (used for out-of-bounds access checks)
+	uint32 HitRecordStride = 0; // size of shader identifier + local root parameters, aligned to SHADER_RECORD_BYTE_ALIGNMENT
 	TResourceArray<uint8, D3D12_RAYTRACING_SHADER_TABLE_BYTE_ALIGNMENT> Data;
 
-#if DO_CHECK && DO_GUARD_SLOW
-	TBitArray<FDefaultBitArrayAllocator> WrittenHitGroupRecords;
-	uint32 NumValidHitGroupRecords = 0;
-#endif // DO_CHECK
-
 	bool bIsDirty = true;
 	TRefCountPtr<FD3D12MemBuffer> Buffer;
 
@@ -756,6 +762,37 @@ void FD3D12Device::DestroyRayTracingDescriptorCache()
 	RayTracingDescriptorHeapCache = nullptr;
 }
 
+struct FD3D12RayTracingShaderLibrary
+{
+	void Reserve(uint32 NumShaders)
+	{
+		Shaders.Reserve(NumShaders);
+		DXILLibraries.Reserve(NumShaders);
+		Identifiers.Reserve(NumShaders);
+	}
+
+	// Single-entry shaders, such as ray-gen, miss and callable
+	void AddSingleShader(FRayTracingShaderRHIParamRef ShaderRHI, LPCWSTR OptExportName = nullptr)
+	{
+		FD3D12RayTracingShader* Shader = FD3D12DynamicRHI::ResourceCast(ShaderRHI);
+		Shaders.Add(Shader);
+
+		LPCWSTR EntryName = *(Shader->EntryPoint);
+		FDXILLibrary& Library = DXILLibraries.AddDefaulted_GetRef();
+		LPCWSTR ExportName = OptExportName ? OptExportName : EntryName;
+		Library.InitFromDXIL(Shader->ShaderBytecode, &EntryName, &ExportName, 1);
+	}
+
+	int32 Num() const
+	{
+		return Identifiers.Num();
+	}
+
+	TArray<TRefCountPtr<FD3D12RayTracingShader>> Shaders;
+	TArray<FDXILLibrary> DXILLibraries;
+	TArray<FD3D12ShaderIdentifier> Identifiers;
+};
+
 class FD3D12RayTracingPipelineState : public FRHIRayTracingPipelineState
 {
 public:
@@ -769,77 +806,67 @@ public:
 		FD3D12Adapter* Adapter = Device->GetParentAdapter();
 		ID3D12Device5* RayTracingDevice = Device->GetRayTracingDevice();
 
-		RayGenShader = FD3D12DynamicRHI::ResourceCast(Initializer.RayGenShaderRHI);
-
-		LPCWSTR RayGenEntryName = *(RayGenShader->EntryPoint);
-		RayGenShaderLibrary.InitFromDXIL(RayGenShader->ShaderBytecode, &RayGenEntryName, &RayGenEntryName, 1);
+		checkf(Initializer.GetRayGenTable().Num() > 0, TEXT("Ray tracing pipelines must have at leat one ray generation shader."));
 
-		const TArrayView<const FRayTracingHitGroupInitializer>& InitializerHitGroups = Initializer.GetHitGroups();
-
-		TArray<ID3D12RootSignature*> LocalRootSignatures;
-		LocalRootSignatures.Reserve(2 + InitializerHitGroups.Num()); // default empty signature + default hit group + one per custom hit group
-
-		// Default empty local root signature
-
-		const uint32 EmptyLocalRootSignatureIndex = LocalRootSignatures.Num();
+		// If no custom hit groups were provided, then disable SBT indexing and force default shader on all primitives
+		HitGroupStride = Initializer.GetHitGroupTable().Num() ? Initializer.HitGroupStride : 0;
 
+		RayGenShaders.Reserve(Initializer.GetRayGenTable().Num());
+		for (FRayTracingShaderRHIParamRef Shader : Initializer.GetRayGenTable())
 		{
-			D3D12_VERSIONED_ROOT_SIGNATURE_DESC LocalRootSignatureDesc = {};
-			LocalRootSignatureDesc.Version = D3D_ROOT_SIGNATURE_VERSION_1_0;
-			LocalRootSignatureDesc.Desc_1_0.Flags |= D3D12_ROOT_SIGNATURE_FLAG_LOCAL_ROOT_SIGNATURE;
-			DefaultLocalRootSignature.Init(LocalRootSignatureDesc);
-			LocalRootSignatures.Add(DefaultLocalRootSignature.GetRootSignature());
+			RayGenShaders.AddSingleShader(Shader);
 		}
 
-		// Miss shader
+		// Add miss shaders (either custom ones provided by user or default one otherwise)
+		if (Initializer.GetMissTable().Num())
+		{
+			MissShaders.Reserve(Initializer.GetMissTable().Num());
+			for (FRayTracingShaderRHIParamRef Shader : Initializer.GetMissTable())
+			{
+				MissShaders.AddSingleShader(Shader);
+			}
+		}
+		else
+		{
+			MissShaders.AddSingleShader(GetBuildInRayTracingShader<FDefaultMainMS>());
+		}
 
-		LPCWSTR MissEntryName = nullptr;
+		// All raygen and miss shaders must share the same root signature, so take the first one and validate the rest
+		GlobalRootSignature = RayGenShaders.Shaders[0]->pRootSignature;
 
-		if (Initializer.MissShaderRHI)
+		for (auto Shader : RayGenShaders.Shaders)
 		{
-			MissShader = FD3D12DynamicRHI::ResourceCast(Initializer.MissShaderRHI);
-			MissEntryName = *(MissShader->EntryPoint);
-			MissShaderLibrary.InitFromDXIL(MissShader->ShaderBytecode, &MissEntryName, &MissEntryName, 1);
+			checkf(Shader->pRootSignature == GlobalRootSignature, TEXT("All raygen and miss shaders must share the same root signature"));
 		}
-		else
+
+		for (auto Shader : MissShaders.Shaders)
 		{
-			GetBuildInShaderLibrary<FDefaultMainMS>(MissShaderLibrary);
-			MissEntryName = *(MissShaderLibrary.EntryNames[0]);
+			checkf(Shader->pRootSignature == GlobalRootSignature, TEXT("All raygen and miss shaders must share the same root signature"));
 		}
 
-		// Default closest hit shader
+		// Use hit shaders from initializer or fall back to default if none were provided
 
-		LPCWSTR ClosestHitEntryName = nullptr;
-		LPCWSTR DefaultHitGroupExportName = TEXT("DefaultHitGroup");
+		FRayTracingShaderRHIParamRef DefaultHitShader = GetBuildInRayTracingShader<FDefaultMainCHS>();
+		TArrayView<const FRayTracingShaderRHIParamRef> DefaultHitGroupTable = { DefaultHitShader };
+
+		const TArrayView<const FRayTracingShaderRHIParamRef>& InitializerHitGroups = Initializer.GetHitGroupTable().Num()
+			? Initializer.GetHitGroupTable()
+			: DefaultHitGroupTable;
+
+		TArray<ID3D12RootSignature*> LocalRootSignatures;
+		LocalRootSignatures.Reserve(1 + InitializerHitGroups.Num()); // default empty signature (for raygen and miss shaders) + one per hit group
 
 		MaxLocalRootSignatureSize = 0;
 
 		const uint32 HitShaderRootSignatureBaseIndex = LocalRootSignatures.Num();
-		if (Initializer.DefaultClosestHitShaderRHI)
-		{
-			ClosestHitShader = FD3D12DynamicRHI::ResourceCast(Initializer.DefaultClosestHitShaderRHI);
-			ClosestHitEntryName = *(ClosestHitShader->EntryPoint);
-			ClosestHitShaderLibrary.InitFromDXIL(ClosestHitShader->ShaderBytecode, &ClosestHitEntryName, &ClosestHitEntryName, 1);
-
-			MaxLocalRootSignatureSize = FMath::Max(MaxLocalRootSignatureSize, ClosestHitShader->pRootSignature->GetTotalRootSignatureSizeInBytes());
-			LocalRootSignatures.Add(ClosestHitShader->pRootSignature->GetRootSignature());
-		}
-		else
-		{
-			GetBuildInShaderLibrary<FDefaultMainCHS>(ClosestHitShaderLibrary);
-			ClosestHitEntryName = *(ClosestHitShaderLibrary.EntryNames[0]);
-		}
 
 		// Initialize hit group shader libraries
-		// #dxr_todo: hit group libraries could come from precompiled pipeline sub-objects
-		HitGroupLibraries.Reserve(InitializerHitGroups.Num()); // #dxr_todo: reserve based on total number of shaders in the hit group list (each hit group can have up to 3)
+		// #dxr_todo: hit group libraries *also* could come from precompiled pipeline sub-objects, when those are supported in the future
+		HitGroupShaders.Reserve(InitializerHitGroups.Num());
 
 		// Each shader within RTPSO must have a unique name, therefore we must rename original shader entry points.
 		TArray<FString> RenamedHitGroupEntryPoints;
-		RenamedHitGroupEntryPoints.Reserve(HitGroupLibraries.Num() * 3); // Up to 3 entry points may exist per hit group
-
-		// Store root signatures per hit group bind shader root parameters later
-		HitGroupShaders.Reserve(InitializerHitGroups.Num());
+		RenamedHitGroupEntryPoints.Reserve(InitializerHitGroups.Num() * 3); // Up to 3 entry points may exist per hit group
 
 		struct FHitGroupEntryIndices
 		{
@@ -851,20 +878,19 @@ public:
 		TArray<FHitGroupEntryIndices> HitGroupEntryIndices;
 		HitGroupEntryIndices.Reserve(InitializerHitGroups.Num());
 
-		for (const FRayTracingHitGroupInitializer& HitGroupInitializer : InitializerHitGroups)
+		for (const FRayTracingShaderRHIParamRef ShaderRHI : InitializerHitGroups)
 		{
-			// #dxr_todo material library can have holes to keep the allocated indices stable and than those holes need to be filled with something, in this case the default.
-			// this might hide bugs where the index was accidentally missing.
-			FD3D12RayTracingShader* Shader = FD3D12DynamicRHI::ResourceCast(HitGroupInitializer.ShaderRHI ? HitGroupInitializer.ShaderRHI : Initializer.DefaultClosestHitShaderRHI);
+			FD3D12RayTracingShader* Shader = FD3D12DynamicRHI::ResourceCast(ShaderRHI);
 
+			checkf(Shader, TEXT("A valid ray tracing hit group shader must be provided for all elements in the FRayTracingPipelineStateInitializer hit group table."));
 			checkf(!Shader->ResourceCounts.bGlobalUniformBufferUsed, TEXT("Global uniform buffers are not implemented for ray tracing shaders"));
 
-			HitGroupShaders.Add(Shader);
+			HitGroupShaders.Shaders.Add(Shader);
 
 			MaxLocalRootSignatureSize = FMath::Max(MaxLocalRootSignatureSize, Shader->pRootSignature->GetTotalRootSignatureSizeInBytes());
 			LocalRootSignatures.Add(Shader->pRootSignature->GetRootSignature());
 
-			FDXILLibrary& Library = HitGroupLibraries.AddDefaulted_GetRef();
+			FDXILLibrary& Library = HitGroupShaders.DXILLibraries.AddDefaulted_GetRef();
 
 			LPCWSTR OriginalGroupEntryPoints[3];
 			LPCWSTR RenamedGroupEntryPoints[3];
@@ -905,198 +931,195 @@ public:
 			Library.InitFromDXIL(Shader->ShaderBytecode, OriginalGroupEntryPoints, RenamedGroupEntryPoints, NumGroupEntryPoints);
 		}
 
-		// Initialize StateObject
+		// Default empty local root signature (#dxr_todo: move this into a single shared place)
+
+		const uint32 EmptyLocalRootSignatureIndex = LocalRootSignatures.Num();
 
 		{
-			TArray<const FDXILLibrary*> Libraries;
-			TArray<LPCWSTR> Exports;
-			TArray<uint32> LocalRootSignatureAssociations;
+			D3D12_VERSIONED_ROOT_SIGNATURE_DESC LocalRootSignatureDesc = {};
+			LocalRootSignatureDesc.Version = D3D_ROOT_SIGNATURE_VERSION_1_0;
+			LocalRootSignatureDesc.Desc_1_0.Flags |= D3D12_ROOT_SIGNATURE_FLAG_LOCAL_ROOT_SIGNATURE;
+			DefaultLocalRootSignature.Init(LocalRootSignatureDesc);
+			LocalRootSignatures.Add(DefaultLocalRootSignature.GetRootSignature());
+		}
 
-			const int32 NumSystemShaders = 3; // Shaders that are guaranteed to be present in the RT pipeline (raygen, miss, default hit)
+		// Initialize StateObject
 
-			// Reserve space for all custom hit groups + system shaders
-			const int32 MaxNumShaders = HitGroupLibraries.Num() + NumSystemShaders;
-			Libraries.Reserve(MaxNumShaders);
+		TArray<const FDXILLibrary*> Libraries;
+		TArray<LPCWSTR> Exports;
+		TArray<uint32> LocalRootSignatureAssociations;
 
-			const int32 MaxNumExports = RenamedHitGroupEntryPoints.Num() + NumSystemShaders;
-			Exports.Reserve(MaxNumExports);
-			LocalRootSignatureAssociations.Reserve(MaxNumExports); // One RS association per export
+		// All shader types other than hit groups
+		const int32 NumSystemShaders = RayGenShaders.DXILLibraries.Num() + MissShaders.DXILLibraries.Num();
 
-			// Ray generation shader
+		// Reserve space for all custom hit groups + system shaders
+		const int32 MaxNumLibraries = HitGroupShaders.DXILLibraries.Num() + NumSystemShaders;
+		Libraries.Reserve(MaxNumLibraries);
 
-			Libraries.Add(&RayGenShaderLibrary);
-			Exports.Add(RayGenEntryName);
-			LocalRootSignatureAssociations.Add(EmptyLocalRootSignatureIndex); // RayGen shaders don't have parameters that come from SBT
+		const int32 MaxNumExports = RenamedHitGroupEntryPoints.Num() + NumSystemShaders;
+		Exports.Reserve(MaxNumExports);
+		LocalRootSignatureAssociations.Reserve(MaxNumExports); // One RS association per export
 
-			// Miss shader
+		// Ray generation shaders
 
-			Libraries.Add(&MissShaderLibrary);
-			Exports.Add(MissEntryName);
-			LocalRootSignatureAssociations.Add(EmptyLocalRootSignatureIndex); // Miss shaders don't have parameters that come from SBT
+		for (const FDXILLibrary& Library : RayGenShaders.DXILLibraries)
+		{
+			Libraries.Add(&Library);
 
-			// Default custom hit shader
+			check(Library.ExportNames.Num() == 1);
+			Exports.Add(*(Library.ExportNames[0]));
 
-			uint32 CurrentHitGroupRootSignatureIndex = HitShaderRootSignatureBaseIndex;
+			// RayGen shaders don't have parameters that come from SBT, so associate with empty local RS
+			LocalRootSignatureAssociations.Add(EmptyLocalRootSignatureIndex);
+		}
+
+		// Miss shaders
 
-			Libraries.Add(&ClosestHitShaderLibrary);
-			Exports.Add(ClosestHitEntryName);
+		for (const FDXILLibrary& Library : MissShaders.DXILLibraries)
+		{
+			Libraries.Add(&Library);
 
-			if (Initializer.DefaultClosestHitShaderRHI)
+			check(Library.ExportNames.Num() == 1);
+			Exports.Add(*(Library.ExportNames[0]));
+
+			// Miss shaders don't have parameters that come from SBT, so associate with empty local RS
+			LocalRootSignatureAssociations.Add(EmptyLocalRootSignatureIndex);
+		}
+
+		// Add hit group shaders
+
+		uint32 CurrentHitGroupRootSignatureIndex = HitShaderRootSignatureBaseIndex;
+		for (const FDXILLibrary& Library : HitGroupShaders.DXILLibraries)
+		{
+			Libraries.Add(&Library);
+			for (const FString& ExportName : Library.ExportNames)
 			{
-				// Add custom closest hit shader root signature association if it is provided
+				Exports.Add(*ExportName);
+
+				// NOTE: Same local root signature is associated with all shaders in a hit group: closest hit, any hit and intersection (if they are present)
 				LocalRootSignatureAssociations.Add(CurrentHitGroupRootSignatureIndex);
-				++CurrentHitGroupRootSignatureIndex;
-			}
-			else
-			{
-				// Use default empty local root signature association if built-in closest hit shader is used
-				LocalRootSignatureAssociations.Add(EmptyLocalRootSignatureIndex);
 			}
+			++CurrentHitGroupRootSignatureIndex;
+		}
 
-			// Additional hit group shaders
+		// Hit groups
 
-			for (const FDXILLibrary& HitGroupLibrary : HitGroupLibraries)
-			{
-				Libraries.Add(&HitGroupLibrary);
-				for (const FString& ExportName : HitGroupLibrary.ExportNames)
-				{
-					Exports.Add(*ExportName);
+		TArray<D3D12_HIT_GROUP_DESC> HitGroups;
+		TArray<FString> HitGroupNames;
 
-					// NOTE: Same local root signature is associated with all shaders in a hit group: closest hit, any hit and intersection (if they are present)
-					LocalRootSignatureAssociations.Add(CurrentHitGroupRootSignatureIndex);
-				}
-				++CurrentHitGroupRootSignatureIndex;
-			}
+		HitGroups.Reserve(InitializerHitGroups.Num());
+		HitGroupNames.Reserve(InitializerHitGroups.Num());
 
-			// Hit groups
+		for (const FRayTracingShaderRHIParamRef ShaderRHI : InitializerHitGroups)
+		{
+			const int32 HitGroupIndex = HitGroupNames.Num(); // #dxr_todo: this would need to be a unique index if we support pipeline sub-object linking
+			FString& HitGroupName = HitGroupNames.AddDefaulted_GetRef();
+			HitGroupName = FString::Printf(TEXT("HitGroup_%d"), HitGroupIndex);
 
-			TArray<D3D12_HIT_GROUP_DESC> HitGroups;
-			TArray<FString> HitGroupNames;
+			D3D12_HIT_GROUP_DESC HitGroup = {};
+			HitGroup.HitGroupExport = *HitGroupName;
 
-			HitGroups.Reserve(1 + HitGroupLibraries.Num()); // N custom hit groups + 1 default (#dxr_todo: deprecate default hit shader and require explicit SBT binding at the high level)
-			HitGroupNames.Reserve(HitGroups.Num());
+			const FHitGroupEntryIndices& EntryIndices = HitGroupEntryIndices[HitGroupIndex];
 
+			HitGroup.ClosestHitShaderImport = *(RenamedHitGroupEntryPoints[EntryIndices.ClosestHit]);
+			if (EntryIndices.AnyHit != INDEX_NONE)
 			{
-				// Add default hit group
-				D3D12_HIT_GROUP_DESC HitGroup = {};
-				HitGroup.HitGroupExport = DefaultHitGroupExportName;
-				HitGroup.ClosestHitShaderImport = ClosestHitEntryName;
-				HitGroups.Add(HitGroup);
+				HitGroup.AnyHitShaderImport = *(RenamedHitGroupEntryPoints[EntryIndices.AnyHit]);
 			}
-
-			for (const FRayTracingHitGroupInitializer& HitGroupInitializer : InitializerHitGroups)
+			if (EntryIndices.Intersection != INDEX_NONE)
 			{
-				const int32 HitGroupIndex = HitGroupNames.Num(); // #dxr_todo: this would need to be a unique index if we support pipeline sub-object linking
-				FString& HitGroupName = HitGroupNames.AddDefaulted_GetRef();
-				HitGroupName = FString::Printf(TEXT("HitGroup_%d"), HitGroupIndex);
-
-				D3D12_HIT_GROUP_DESC HitGroup = {};
-				HitGroup.HitGroupExport = *HitGroupName;
-
-				const FHitGroupEntryIndices& EntryIndices = HitGroupEntryIndices[HitGroupIndex];
-
-				HitGroup.ClosestHitShaderImport = *(RenamedHitGroupEntryPoints[EntryIndices.ClosestHit]);
-				if (EntryIndices.AnyHit != INDEX_NONE)
-				{
-					HitGroup.AnyHitShaderImport = *(RenamedHitGroupEntryPoints[EntryIndices.AnyHit]);
-				}
-				if (EntryIndices.Intersection != INDEX_NONE)
-				{
-					HitGroup.IntersectionShaderImport = *(RenamedHitGroupEntryPoints[EntryIndices.Intersection]);
-				}
-
-				HitGroups.Add(HitGroup);
+				HitGroup.IntersectionShaderImport = *(RenamedHitGroupEntryPoints[EntryIndices.Intersection]);
 			}
 
-			// Create the pipeline
+			HitGroups.Add(HitGroup);
+		}
 
-			check(Libraries.Num() == MaxNumShaders); // Confirm that our memory reservation assumptions hold up
-			check(Exports.Num() == MaxNumExports); // Confirm that our memory reservation assumptions hold up
-			check(Exports.Num() == LocalRootSignatureAssociations.Num()); // Confirm that we have associated local root signatures with all shaders
+		// Create the pipeline
 
-			StateObject = CreateRayTracingStateObject(
-				RayTracingDevice,
-				Libraries,
-				Exports,
-				Initializer.MaxPayloadSizeInBytes,
-				HitGroups,
-				*RayGenShader->pRootSignature,
-				LocalRootSignatures,
-				LocalRootSignatureAssociations);
+		check(Libraries.Num() == MaxNumLibraries); // Confirm that our memory reservation assumptions hold up
+		check(Exports.Num() == MaxNumExports); // Confirm that our memory reservation assumptions hold up
+		check(Exports.Num() == LocalRootSignatureAssociations.Num()); // Confirm that we have associated local root signatures with all shaders
 
-			VERIFYHRESULT(StateObject->QueryInterface(IID_PPV_ARGS(PipelineProperties.GetInitReference())));
+		StateObject = CreateRayTracingStateObject(
+			RayTracingDevice,
+			Libraries,
+			Exports,
+			Initializer.MaxPayloadSizeInBytes,
+			HitGroups,
+			*GlobalRootSignature,
+			LocalRootSignatures,
+			LocalRootSignatureAssociations);
 
-			// Save hit group shader identifiers
+		VERIFYHRESULT(StateObject->QueryInterface(IID_PPV_ARGS(PipelineProperties.GetInitReference())));
 
-			check(HitGroupNames.Num() == InitializerHitGroups.Num());
+		auto GetShaderIdentifier = [PipelineProperties = this->PipelineProperties.GetReference()](LPCWSTR ExportName)->FD3D12ShaderIdentifier
+		{
+			FD3D12ShaderIdentifier Result;
 
-			HitGroupShaderIdentifiers.SetNumUninitialized(ShaderIdentifierSize * InitializerHitGroups.Num());
-			for (int32 HitGroupIndex = 0; HitGroupIndex < HitGroupNames.Num(); ++HitGroupIndex)
+			const void* Data = PipelineProperties->GetShaderIdentifier(ExportName);
+			checkf(Data, TEXT("Couldn't find requested export in the ray tracing shader pipeline"));
+
+			if (Data)
 			{
-				LPCWSTR ExportNameChars = *HitGroupNames[HitGroupIndex];
-				HitGroupShaderIdentifiers[HitGroupIndex] = GetShaderIdentifier(ExportNameChars);
+				Result.SetData(Data);
 			}
-		}
-
-		// Initialize shader identifiers and default ShaderTable
-
-		RayGenShaderIdentifier     = GetShaderIdentifier(RayGenEntryName);
-		MissShaderIdentifier       = GetShaderIdentifier(MissEntryName);
-		DefaultHitGroupIdentifier  = GetShaderIdentifier(DefaultHitGroupExportName);
-
-		const uint32 LocalRootDataSize = 0; // default hit shaders use an empty local root signature
-		DefaultShaderTable.Init(1, 1, 0, LocalRootDataSize);
 
-		DefaultShaderTable.SetRayGenIdentifier(0, RayGenShaderIdentifier);
-		DefaultShaderTable.SetMissIdentifier(0, MissShaderIdentifier);
-		DefaultShaderTable.SetDefaultHitGroupIdentifier(DefaultHitGroupIdentifier);
-	}
+			return Result;
+		};
 
-	FD3D12ShaderIdentifier GetShaderIdentifier(LPCWSTR ExportName)
-	{
-		FD3D12ShaderIdentifier Result;
+		// Query shader identifiers from the pipeline state object
 
-		const void* Data = PipelineProperties->GetShaderIdentifier(ExportName);
-		checkf(Data, TEXT("Couldn't find requested export in the ray tracing shader pipeline"));
+		check(HitGroupNames.Num() == InitializerHitGroups.Num());
 
-		if (Data)
+		HitGroupShaders.Identifiers.SetNumUninitialized(InitializerHitGroups.Num());
+		for (int32 HitGroupIndex = 0; HitGroupIndex < HitGroupNames.Num(); ++HitGroupIndex)
 		{
-			check(sizeof(Result.Data) >= ShaderIdentifierSize);
-			FMemory::Memcpy(Result.Data, Data, ShaderIdentifierSize);
+			LPCWSTR ExportNameChars = *HitGroupNames[HitGroupIndex];
+			HitGroupShaders.Identifiers[HitGroupIndex] = GetShaderIdentifier(ExportNameChars);
 		}
 
-		return Result;
-	}
+		RayGenShaders.Identifiers.SetNumUninitialized(RayGenShaders.DXILLibraries.Num());
+		for (int32 ShaderIndex = 0; ShaderIndex < RayGenShaders.DXILLibraries.Num(); ++ShaderIndex)
+		{
+			LPCWSTR ExportNameChars = *(RayGenShaders.DXILLibraries[ShaderIndex].ExportNames[0]);
+			RayGenShaders.Identifiers[ShaderIndex] = GetShaderIdentifier(ExportNameChars);
+		}
 
-	TRefCountPtr<FD3D12RayTracingShader> RayGenShader;
-	FD3D12ShaderIdentifier RayGenShaderIdentifier;
-	FDXILLibrary RayGenShaderLibrary;
+		MissShaders.Identifiers.SetNumUninitialized(MissShaders.DXILLibraries.Num());
+		for (int32 ShaderIndex = 0; ShaderIndex < MissShaders.DXILLibraries.Num(); ++ShaderIndex)
+		{
+			LPCWSTR ExportNameChars = *(MissShaders.DXILLibraries[ShaderIndex].ExportNames[0]);
+			MissShaders.Identifiers[ShaderIndex] = GetShaderIdentifier(ExportNameChars);
+		}
 
-	// #dxr_todo: deprecate default hit shader and require explicit SBT binding at the high level
-	TRefCountPtr<FD3D12RayTracingShader> ClosestHitShader;
-	FD3D12ShaderIdentifier ClosestHitShaderIdentifier;
-	FDXILLibrary ClosestHitShaderLibrary;
+		// Setup default shader binding table, which simply includes all provided RGS and MS plus a single default closest hit shader.
+		// Hit record indexing and local resources access is disabled when using using this SBT.
 
-	TRefCountPtr<FD3D12RayTracingShader> MissShader;
-	FD3D12ShaderIdentifier MissShaderIdentifier;
-	FDXILLibrary MissShaderLibrary;
+		const uint32 DefaultLocalRootDataSize = 0; // Shaders in default SBT are not allowed to access any local resources
+		DefaultShaderTable.Init(RayGenShaders.Num(), MissShaders.Num(), 0, DefaultLocalRootDataSize);
+		DefaultShaderTable.SetRayGenIdentifiers(RayGenShaders.Identifiers);
+		DefaultShaderTable.SetMissIdentifiers(MissShaders.Identifiers);
+		DefaultShaderTable.SetDefaultHitGroupIdentifier(HitGroupShaders.Identifiers[0]);
+	}
 
-	FD3D12ShaderIdentifier DefaultHitGroupIdentifier;
-	TArray<FDXILLibrary> HitGroupLibraries;
+	FD3D12RayTracingShaderLibrary RayGenShaders;
+	FD3D12RayTracingShaderLibrary MissShaders;
+	FD3D12RayTracingShaderLibrary HitGroupShaders;
 
 	// Shader table that can be used to dispatch ray tracing work that doesn't require real SBT bindings.
 	// This is useful for the case where user only provides default RayGen, Miss and HitGroup shaders.
 	FD3D12RayTracingShaderTable DefaultShaderTable;
 
 	// Default empty root signature used for default hit shaders.
-	FD3D12RootSignature DefaultLocalRootSignature;
+	FD3D12RootSignature DefaultLocalRootSignature; // #dxr_todo: move this into a single shared place
+
+	const FD3D12RootSignature* GlobalRootSignature = nullptr;
 
 	TRefCountPtr<ID3D12StateObject> StateObject;
 	TRefCountPtr<ID3D12StateObjectProperties> PipelineProperties;
 
 	static constexpr uint32 ShaderIdentifierSize = D3D12_SHADER_IDENTIFIER_SIZE_IN_BYTES;
-	TArray<FD3D12ShaderIdentifier> HitGroupShaderIdentifiers;
-	TArray<TRefCountPtr<FD3D12RayTracingShader>> HitGroupShaders;
+	uint32 HitGroupStride = 0; // Number of shader slots per logical SBT record (i.e. per geometry segment)
 
 	uint32 MaxLocalRootSignatureSize = 0;
 };
@@ -1109,18 +1132,38 @@ public:
 
 	FD3D12BasicRayTracingPipeline(FD3D12Device* Device)
 	{
-		FRayTracingPipelineStateInitializer OcclusionInitializer;
-		OcclusionInitializer.RayGenShaderRHI = GetBuildInRayTracingShader<FOcclusionMainRG>();
-		OcclusionInitializer.MissShaderRHI = GetBuildInRayTracingShader<FOcclusionMainMS>();
+		// Occlusion pipeline
+		{
+			FRayTracingPipelineStateInitializer OcclusionInitializer;
 
-		Occlusion = new FD3D12RayTracingPipelineState(Device, OcclusionInitializer);
+			FRayTracingShaderRHIParamRef OcclusionRGSTable[] = { GetBuildInRayTracingShader<FOcclusionMainRG>() };
+			OcclusionInitializer.SetRayGenShaderTable(OcclusionRGSTable);
 
-		FRayTracingPipelineStateInitializer IntersectionInitializer;
-		IntersectionInitializer.RayGenShaderRHI = GetBuildInRayTracingShader<FIntersectionMainRG>();
-		IntersectionInitializer.MissShaderRHI = GetBuildInRayTracingShader<FIntersectionMainMS>();
-		IntersectionInitializer.DefaultClosestHitShaderRHI = GetBuildInRayTracingShader<FIntersectionMainCHS>();
+			FRayTracingShaderRHIParamRef OcclusionMSTable[] = { GetBuildInRayTracingShader<FOcclusionMainMS>() };
+			OcclusionInitializer.SetMissShaderTable(OcclusionMSTable);
 
-		Intersection = new FD3D12RayTracingPipelineState(Device, IntersectionInitializer);
+			OcclusionInitializer.HitGroupStride = 0;
+
+			Occlusion = new FD3D12RayTracingPipelineState(Device, OcclusionInitializer);
+		}
+
+		// Intersection pipeline
+		{
+			FRayTracingPipelineStateInitializer IntersectionInitializer;
+
+			FRayTracingShaderRHIParamRef IntersectionRGSTable[] = { GetBuildInRayTracingShader<FIntersectionMainRG>() };
+			IntersectionInitializer.SetRayGenShaderTable(IntersectionRGSTable);
+
+			FRayTracingShaderRHIParamRef IntersectionMSTable[] = { GetBuildInRayTracingShader<FIntersectionMainMS>() };
+			IntersectionInitializer.SetMissShaderTable(IntersectionMSTable);
+
+			FRayTracingShaderRHIParamRef IntersectionHitTable[] = { GetBuildInRayTracingShader<FIntersectionMainCHS>() };
+			IntersectionInitializer.SetHitGroupTable(IntersectionHitTable);
+
+			IntersectionInitializer.HitGroupStride = 0;
+
+			Intersection = new FD3D12RayTracingPipelineState(Device, IntersectionInitializer);
+		}
 	}
 
 	~FD3D12BasicRayTracingPipeline()
@@ -1600,49 +1643,57 @@ FD3D12RayTracingShaderTable* FD3D12RayTracingScene::FindOrCreateShaderTable(cons
 
 	// #dxr_todo: this needs to take into account multiple ray types (when those are supported)
 	const uint32 NumHitGroupSlots = NumTotalSegments; // one hit group slot per geometry segment
-	const uint32 LocalRootDataSize = Pipeline->MaxLocalRootSignatureSize;
-
-	checkf(LocalRootDataSize >= sizeof(FHitGroupSystemParameters), TEXT("All local root signatures are expected to contain ray tracing system root parameters (2x root buffers + 1x root DWORD)"));
+	const uint32 HitGroupStride = Pipeline->HitGroupStride;
 
-	CreatedShaderTable->Init(1, 1, NumHitGroupSlots, LocalRootDataSize);
+	checkf(Pipeline->MaxLocalRootSignatureSize >= sizeof(FHitGroupSystemParameters), TEXT("All local root signatures are expected to contain ray tracing system root parameters (2x root buffers + 1x root DWORD)"));
 
-	CreatedShaderTable->SetRayGenIdentifier(0, Pipeline->RayGenShaderIdentifier);
-	CreatedShaderTable->SetMissIdentifier(0, Pipeline->MissShaderIdentifier);
-	CreatedShaderTable->SetDefaultHitGroupIdentifier(Pipeline->DefaultHitGroupIdentifier);
+	CreatedShaderTable->Init(
+		Pipeline->RayGenShaders.Num(),
+		Pipeline->MissShaders.Num(),
+		NumHitGroupSlots * HitGroupStride,
+		Pipeline->MaxLocalRootSignatureSize);
 
-	// Bind default hit group, index/vertex buffers and fetch parameters to all SBT entries (all segments of all mesh instances)
+	CreatedShaderTable->SetRayGenIdentifiers(Pipeline->RayGenShaders.Identifiers);
+	CreatedShaderTable->SetMissIdentifiers(Pipeline->MissShaders.Identifiers);
 
-	const uint32 NumInstances = Instances.Num();
-	for (uint32 InstanceIndex = 0; InstanceIndex < NumInstances; ++InstanceIndex)
+	// Bind index/vertex buffers and fetch parameters to all SBT entries (all segments of all mesh instances)
+	// Resource binding is skipped for pipelines that don't use SBT indexing. Such pipelines use the same CHS for all primitives, which can't access any local resources.
+	if (HitGroupStride)
 	{
-		const FRayTracingGeometryInstance& Instance = Instances[InstanceIndex];
+		const uint32 NumInstances = Instances.Num();
+		for (uint32 InstanceIndex = 0; InstanceIndex < NumInstances; ++InstanceIndex)
+		{
+			const FRayTracingGeometryInstance& Instance = Instances[InstanceIndex];
 
-		const FD3D12RayTracingGeometry* Geometry = FD3D12DynamicRHI::ResourceCast(Instance.GeometryRHI);
+			const FD3D12RayTracingGeometry* Geometry = FD3D12DynamicRHI::ResourceCast(Instance.GeometryRHI);
 
-		static constexpr uint32 IndicesPerPrimitive = 3; // Only triangle meshes are supported
+			static constexpr uint32 IndicesPerPrimitive = 3; // Only triangle meshes are supported
 
-		const uint32 IndexStride = Geometry->IndexStride;
-		const D3D12_GPU_VIRTUAL_ADDRESS IndexBufferAddress = Geometry->IndexBuffer ? Geometry->IndexBuffer->ResourceLocation.GetGPUVirtualAddress() : 0;
-		const D3D12_GPU_VIRTUAL_ADDRESS VertexBufferAddress = Geometry->PositionVertexBuffer->ResourceLocation.GetGPUVirtualAddress() + Geometry->VertexOffsetInBytes;
+			const uint32 IndexStride = Geometry->IndexStride;
+			const D3D12_GPU_VIRTUAL_ADDRESS IndexBufferAddress = Geometry->IndexBuffer ? Geometry->IndexBuffer->ResourceLocation.GetGPUVirtualAddress() : 0;
+			const D3D12_GPU_VIRTUAL_ADDRESS VertexBufferAddress = Geometry->PositionVertexBuffer->ResourceLocation.GetGPUVirtualAddress() + Geometry->VertexOffsetInBytes;
 
-		const uint32 NumSegments = Geometry->Segments.Num();
-		for (uint32 SegmentIndex = 0; SegmentIndex < NumSegments; ++SegmentIndex)
-		{
-			const FRayTracingGeometrySegment& Segment = Geometry->Segments[SegmentIndex];
+			const uint32 NumSegments = Geometry->Segments.Num();
+			for (uint32 SegmentIndex = 0; SegmentIndex < NumSegments; ++SegmentIndex)
+			{
+				const FRayTracingGeometrySegment& Segment = Geometry->Segments[SegmentIndex];
 
-			const uint32 RecordIndex = GetHitGroupIndex(InstanceIndex, SegmentIndex);
-			CreatedShaderTable->SetHitGroupIdentifier(RecordIndex, Pipeline->DefaultHitGroupIdentifier);
+				const uint32 RecordBaseIndex = GetHitGroupIndex(InstanceIndex, SegmentIndex) * HitGroupStride;
 
-			FHitGroupSystemParameters RootParameters = {};
-			RootParameters.IndexBuffer = IndexBufferAddress;
-			RootParameters.VertexBuffer = VertexBufferAddress;
+				FHitGroupSystemParameters RootParameters = {};
+				RootParameters.IndexBuffer = IndexBufferAddress;
+				RootParameters.VertexBuffer = VertexBufferAddress;
 
-			// #dxr_todo: support various vertex buffer layouts (fetch/decode based on vertex stride and format)
-			checkf(Geometry->VertexElemType == VET_Float3, TEXT("Only VET_Float3 is currently implemented and tested. Other formats will be supported in the future."));
-			RootParameters.FetchParameters.SetVertexAndIndexStride(Geometry->VertexStrideInBytes, IndexStride);
-			RootParameters.FetchParameters.IndexBufferOffsetInBytes = IndexStride * Segment.FirstPrimitive * IndicesPerPrimitive;
+				// #dxr_todo: support various vertex buffer layouts (fetch/decode based on vertex stride and format)
+				checkf(Geometry->VertexElemType == VET_Float3, TEXT("Only VET_Float3 is currently implemented and tested. Other formats will be supported in the future."));
+				RootParameters.FetchParameters.SetVertexAndIndexStride(Geometry->VertexStrideInBytes, IndexStride);
+				RootParameters.FetchParameters.IndexBufferOffsetInBytes = IndexStride * Segment.FirstPrimitive * IndicesPerPrimitive;
 
-			CreatedShaderTable->SetHitGroupParameters(RecordIndex, 0, RootParameters);
+				for (uint32 SlotIndex = 0; SlotIndex < HitGroupStride; ++SlotIndex)
+				{
+					CreatedShaderTable->SetHitGroupParameters(RecordBaseIndex + SlotIndex, 0, RootParameters);
+				}
+			}
 		}
 	}
 
@@ -2033,6 +2084,7 @@ static void SetRayTracingShaderResources(
 static void DispatchRays(FD3D12CommandContext& CommandContext,
 	const FRayTracingShaderBindings& GlobalBindings,
 	const FD3D12RayTracingPipelineState* Pipeline,
+	uint32 RayGenShaderIndex,
 	FD3D12RayTracingShaderTable* OptShaderTable,
 	const D3D12_DISPATCH_RAYS_DESC& DispatchDesc)
 {
@@ -2053,13 +2105,15 @@ static void DispatchRays(FD3D12CommandContext& CommandContext,
 	// Invalidate state cache to ensure all root parameters for regular shaders are reset when non-RT work is dispatched later.
 	CommandContext.StateCache.TransitionComputeState(D3D12PT_RayTracing);
 
-	CommandList->SetComputeRootSignature(Pipeline->RayGenShader->pRootSignature->GetRootSignature());
+	CommandList->SetComputeRootSignature(Pipeline->GlobalRootSignature->GetRootSignature());
+
+	FD3D12RayTracingShader* RayGenShader = Pipeline->RayGenShaders.Shaders[RayGenShaderIndex];
 
 	if (OptShaderTable)
 	{
 		OptShaderTable->DescriptorCache.SetDescriptorHeaps(CommandContext);
 		FD3D12RayTracingGlobalResourceBinder ResourceBinder(CommandContext);
-		SetRayTracingShaderResources(CommandContext, Pipeline->RayGenShader, GlobalBindings, OptShaderTable->DescriptorCache, ResourceBinder);
+		SetRayTracingShaderResources(CommandContext, RayGenShader, GlobalBindings, OptShaderTable->DescriptorCache, ResourceBinder);
 	}
 	else
 	{
@@ -2067,7 +2121,7 @@ static void DispatchRays(FD3D12CommandContext& CommandContext,
 		TransientDescriptorCache.Init(MAX_SRVS + MAX_UAVS, MAX_SAMPLERS);
 		TransientDescriptorCache.SetDescriptorHeaps(CommandContext);
 		FD3D12RayTracingGlobalResourceBinder ResourceBinder(CommandContext);
-		SetRayTracingShaderResources(CommandContext, Pipeline->RayGenShader, GlobalBindings, TransientDescriptorCache, ResourceBinder);
+		SetRayTracingShaderResources(CommandContext, RayGenShader, GlobalBindings, TransientDescriptorCache, ResourceBinder);
 	}
 
 	CommandContext.CommandListHandle.FlushResourceBarriers();
@@ -2107,7 +2161,7 @@ void FD3D12CommandContext::RHIRayTraceOcclusion(FRayTracingSceneRHIParamRef InSc
 
 	ID3D12GraphicsCommandList4* RayTracingCommandList = CommandListHandle.RayTracingCommandList();
 
-	D3D12_DISPATCH_RAYS_DESC DispatchDesc = ShaderTable.GetDispatchRaysDesc();
+	D3D12_DISPATCH_RAYS_DESC DispatchDesc = ShaderTable.GetDispatchRaysDesc(0, 0, 0);
 
 	DispatchDesc.Width = NumRays;
 	DispatchDesc.Height = 1;
@@ -2118,7 +2172,7 @@ void FD3D12CommandContext::RHIRayTraceOcclusion(FRayTracingSceneRHIParamRef InSc
 	Bindings.SRVs[1] = Rays;
 	Bindings.UAVs[0] = Output;
 
-	DispatchRays(*this, Bindings, Pipeline, nullptr, DispatchDesc);
+	DispatchRays(*this, Bindings, Pipeline, 0, nullptr, DispatchDesc);
 }
 
 void FD3D12CommandContext::RHIRayTraceIntersection(FRayTracingSceneRHIParamRef InScene,
@@ -2142,7 +2196,7 @@ void FD3D12CommandContext::RHIRayTraceIntersection(FRayTracingSceneRHIParamRef I
 
 	ID3D12GraphicsCommandList4* RayTracingCommandList = CommandListHandle.RayTracingCommandList();
 
-	D3D12_DISPATCH_RAYS_DESC DispatchDesc = ShaderTable.GetDispatchRaysDesc();
+	D3D12_DISPATCH_RAYS_DESC DispatchDesc = ShaderTable.GetDispatchRaysDesc(0, 0, 0);
 
 	DispatchDesc.Width = NumRays;
 	DispatchDesc.Height = 1;
@@ -2156,10 +2210,10 @@ void FD3D12CommandContext::RHIRayTraceIntersection(FRayTracingSceneRHIParamRef I
 	Bindings.UAVs[0] = Output;
 	Bindings.UAVs[1] = Output;
 
-	DispatchRays(*this, Bindings, Pipeline, nullptr, DispatchDesc);
+	DispatchRays(*this, Bindings, Pipeline, 0, nullptr, DispatchDesc);
 }
 
-void FD3D12CommandContext::RHIRayTraceDispatch(FRayTracingPipelineStateRHIParamRef InRayTracingPipelineState,
+void FD3D12CommandContext::RHIRayTraceDispatch(FRayTracingPipelineStateRHIParamRef InRayTracingPipelineState, uint32 RayGenShaderIndex,
 	FRayTracingSceneRHIParamRef InScene, // #dxr_todo: replace this with explicit shader table parameter
 	const FRayTracingShaderBindings& GlobalResourceBindings,
 	uint32 Width, uint32 Height)
@@ -2174,16 +2228,16 @@ void FD3D12CommandContext::RHIRayTraceDispatch(FRayTracingPipelineStateRHIParamR
 		ShaderTable->CopyToGPU();
 	}
 
-	D3D12_DISPATCH_RAYS_DESC DispatchDesc = ShaderTable->GetDispatchRaysDesc();
+	D3D12_DISPATCH_RAYS_DESC DispatchDesc = ShaderTable->GetDispatchRaysDesc(RayGenShaderIndex, 0, Pipeline->HitGroupStride);
 
 	DispatchDesc.Width = Width;
 	DispatchDesc.Height = Height;
 	DispatchDesc.Depth = 1;
 
-	DispatchRays(*this, GlobalResourceBindings, Pipeline, ShaderTable, DispatchDesc);
+	DispatchRays(*this, GlobalResourceBindings, Pipeline, RayGenShaderIndex, ShaderTable, DispatchDesc);
 }
 
-void FD3D12CommandContext::RHIRayTraceDispatch(FRayTracingPipelineStateRHIParamRef InRayTracingPipelineState,
+void FD3D12CommandContext::RHIRayTraceDispatch(FRayTracingPipelineStateRHIParamRef InRayTracingPipelineState, uint32 RayGenShaderIndex,
 	const FRayTracingShaderBindings& GlobalResourceBindings,
 	uint32 Width, uint32 Height)
 {
@@ -2195,29 +2249,31 @@ void FD3D12CommandContext::RHIRayTraceDispatch(FRayTracingPipelineStateRHIParamR
 		ShaderTable.CopyToGPU();
 	}
 
-	D3D12_DISPATCH_RAYS_DESC DispatchDesc = ShaderTable.GetDispatchRaysDesc();
+	D3D12_DISPATCH_RAYS_DESC DispatchDesc = ShaderTable.GetDispatchRaysDesc(RayGenShaderIndex, 0, Pipeline->HitGroupStride);
 
 	DispatchDesc.Width = Width;
 	DispatchDesc.Height = Height;
 	DispatchDesc.Depth = 1;
 
-	DispatchRays(*this, GlobalResourceBindings, Pipeline, nullptr, DispatchDesc);
+	DispatchRays(*this, GlobalResourceBindings, Pipeline, RayGenShaderIndex, nullptr, DispatchDesc);
 }
 
 void FD3D12CommandContext::RHISetRayTracingHitGroup(
-	FRayTracingSceneRHIParamRef InScene, uint32 InstanceIndex, uint32 SegmentIndex,
+	FRayTracingSceneRHIParamRef InScene, uint32 InstanceIndex, uint32 SegmentIndex, uint32 ShaderSlot,
 	FRayTracingPipelineStateRHIParamRef InPipeline, uint32 HitGroupIndex,
 	uint32 NumUniformBuffers, const FUniformBufferRHIParamRef* UniformBuffers)
 {
 	FD3D12RayTracingScene* Scene = FD3D12DynamicRHI::ResourceCast(InScene);
 	FD3D12RayTracingPipelineState* Pipeline = FD3D12DynamicRHI::ResourceCast(InPipeline);
-	
+
+	checkf(ShaderSlot < Pipeline->HitGroupStride, TEXT("Shader slot is invalid. Make sure that HitGroupStride is correct on FRayTracingPipelineStateInitializer."));
+
 	FD3D12RayTracingShaderTable* ShaderTable = Scene->FindOrCreateShaderTable(Pipeline);
 
-	const uint32 RecordIndex = Scene->GetHitGroupIndex(InstanceIndex, SegmentIndex);
-	ShaderTable->SetHitGroupIdentifier(RecordIndex, Pipeline->HitGroupShaderIdentifiers[HitGroupIndex]);
+	const uint32 RecordIndex = Scene->GetHitGroupIndex(InstanceIndex, SegmentIndex) * Pipeline->HitGroupStride + ShaderSlot;
+	ShaderTable->SetHitGroupIdentifier(RecordIndex, Pipeline->HitGroupShaders.Identifiers[HitGroupIndex]);
 
-	const FD3D12RayTracingShader* Shader = Pipeline->HitGroupShaders[HitGroupIndex];
+	const FD3D12RayTracingShader* Shader = Pipeline->HitGroupShaders.Shaders[HitGroupIndex];
 
 	FD3D12RayTracingLocalResourceBinder ResourceBinder(*this, ShaderTable, Shader->pRootSignature, RecordIndex);
 	SetRayTracingShaderResources(*this, Shader,
diff --git a/Engine/Source/Runtime/D3D12RHI/Private/D3D12Shaders.cpp b/Engine/Source/Runtime/D3D12RHI/Private/D3D12Shaders.cpp
--- a/Engine/Source/Runtime/D3D12RHI/Private/D3D12Shaders.cpp
+++ b/Engine/Source/Runtime/D3D12RHI/Private/D3D12Shaders.cpp
@@ -308,20 +308,7 @@ FRayTracingShaderRHIRef FD3D12DynamicRHI::RHICreateRayTracingShader(const TArray
 #else // USE_STATIC_ROOT_SIGNATURE
 	const D3D12_RESOURCE_BINDING_TIER Tier = Adapter.GetResourceBindingTier();
 	FD3D12QuantizedBoundShaderState QBSS;
-	QuantizeBoundShaderState(Tier, Shader, QBSS);
-	switch (ShaderFrequency)
-	{
-	case SF_RayGen:
-	case SF_RayMiss:
-		QBSS.RootSignatureType = RS_RayTracingGlobal;
-		break;
-	case SF_RayHitGroup:
-		// Local root signature is used for hit group shaders
-		QBSS.RootSignatureType = RS_RayTracingLocal;
-		break;
-	default:
-		checkNoEntry(); // Unexpected shader target frequency
-	}
+	QuantizeBoundShaderState(ShaderFrequency, Tier, Shader, QBSS);
 	Shader->pRootSignature = Adapter.GetRootSignature(QBSS);
 #endif // USE_STATIC_ROOT_SIGNATURE
 
diff --git a/Engine/Source/Runtime/D3D12RHI/Private/D3D12Util.cpp b/Engine/Source/Runtime/D3D12RHI/Private/D3D12Util.cpp
--- a/Engine/Source/Runtime/D3D12RHI/Private/D3D12Util.cpp
+++ b/Engine/Source/Runtime/D3D12RHI/Private/D3D12Util.cpp
@@ -405,6 +405,7 @@ void QuantizeBoundShaderState(
 
 #if D3D12_RHI_RAYTRACING
 void QuantizeBoundShaderState(
+	EShaderFrequency ShaderFrequency,
 	const D3D12_RESOURCE_BINDING_TIER& ResourceBindingTier,
 	const FD3D12RayTracingShader* const RayTracingShader,
 	FD3D12QuantizedBoundShaderState &OutQBSS
@@ -417,15 +418,42 @@ void QuantizeBoundShaderState(
 	FMemory::Memzero(&OutQBSS, sizeof(OutQBSS));
 	FShaderRegisterCounts& QBSSRegisterCounts = OutQBSS.RegisterCounts[SV_All];
 
-	QBSSRegisterCounts.SamplerCount = Counts.NumSamplers;
-	QBSSRegisterCounts.ShaderResourceCount = Counts.NumSRVs;
-	QBSSRegisterCounts.ConstantBufferCount = Counts.NumCBs;
-	QBSSRegisterCounts.UnorderedAccessCount = Counts.NumUAVs;
+	switch (ShaderFrequency)
+	{
+	case SF_RayGen:
+	case SF_RayMiss:
+
+		// Shared conservative root signature layout is used for all raygen and miss shaders.
+
+		OutQBSS.RootSignatureType = RS_RayTracingGlobal;
+
+		QBSSRegisterCounts.SamplerCount = MAX_SAMPLERS;
+		QBSSRegisterCounts.ShaderResourceCount = MAX_SRVS;
+		QBSSRegisterCounts.ConstantBufferCount = MAX_CBS;
+		QBSSRegisterCounts.UnorderedAccessCount = MAX_UAVS;
+
+		break;
+
+	case SF_RayHitGroup:
 
-	check(QBSSRegisterCounts.SamplerCount < MAX_SAMPLERS);
-	check(QBSSRegisterCounts.ShaderResourceCount < MAX_SRVS);
-	check(QBSSRegisterCounts.ConstantBufferCount < MAX_CBS);
-	check(QBSSRegisterCounts.UnorderedAccessCount < MAX_UAVS);
+		// Local root signature is used for hit group shaders, using the exact number of resources to minimize shader binding table record size.
+
+		OutQBSS.RootSignatureType = RS_RayTracingLocal;
+
+		QBSSRegisterCounts.SamplerCount = Counts.NumSamplers;
+		QBSSRegisterCounts.ShaderResourceCount = Counts.NumSRVs;
+		QBSSRegisterCounts.ConstantBufferCount = Counts.NumCBs;
+		QBSSRegisterCounts.UnorderedAccessCount = Counts.NumUAVs;
+
+		check(QBSSRegisterCounts.SamplerCount <= MAX_SAMPLERS);
+		check(QBSSRegisterCounts.ShaderResourceCount <= MAX_SRVS);
+		check(QBSSRegisterCounts.ConstantBufferCount <= MAX_CBS);
+		check(QBSSRegisterCounts.UnorderedAccessCount <= MAX_UAVS);
+
+		break;
+	default:
+		checkNoEntry(); // Unexpected shader target frequency
+	}
 }
 #endif // D3D12_RHI_RAYTRACING
 
diff --git a/Engine/Source/Runtime/D3D12RHI/Public/D3D12Util.h b/Engine/Source/Runtime/D3D12RHI/Public/D3D12Util.h
--- a/Engine/Source/Runtime/D3D12RHI/Public/D3D12Util.h
+++ b/Engine/Source/Runtime/D3D12RHI/Public/D3D12Util.h
@@ -141,6 +141,7 @@ extern void QuantizeBoundShaderState(
 #if D3D12_RHI_RAYTRACING
 class FD3D12RayTracingShader;
 extern void QuantizeBoundShaderState(
+	EShaderFrequency ShaderFrequency,
 	const D3D12_RESOURCE_BINDING_TIER& ResourceBindingTier,
 	const FD3D12RayTracingShader* const Shader,
 	FD3D12QuantizedBoundShaderState &OutQBSS);
diff --git a/Engine/Source/Runtime/RHI/Private/PipelineStateCache.cpp b/Engine/Source/Runtime/RHI/Private/PipelineStateCache.cpp
--- a/Engine/Source/Runtime/RHI/Private/PipelineStateCache.cpp
+++ b/Engine/Source/Runtime/RHI/Private/PipelineStateCache.cpp
@@ -45,9 +45,9 @@ static inline uint32 GetTypeHash(const FGraphicsPipelineStateInitializer& Initia
 static inline uint32 GetTypeHash(const FRayTracingPipelineStateInitializer& Initializer)
 {
 	return GetTypeHash(Initializer.MaxPayloadSizeInBytes) ^
-		GetTypeHash(Initializer.RayGenShaderRHI) ^
-		GetTypeHash(Initializer.MissShaderRHI) ^
-		GetTypeHash(Initializer.DefaultClosestHitShaderRHI) ^ 
+		GetTypeHash(Initializer.HitGroupStride) ^
+		GetTypeHash(Initializer.GetRayGenHash()) ^
+		GetTypeHash(Initializer.GetRayMissHash()) ^
 		GetTypeHash(Initializer.GetHitGroupHash());
 }
 #endif
diff --git a/Engine/Source/Runtime/RHI/Public/RHICommandList.h b/Engine/Source/Runtime/RHI/Public/RHICommandList.h
--- a/Engine/Source/Runtime/RHI/Public/RHICommandList.h
+++ b/Engine/Source/Runtime/RHI/Public/RHICommandList.h
@@ -1881,23 +1881,24 @@ struct FRHICommandRayTraceDispatch final : public FRHICommand<FRHICommandRayTrac
 	FRayTracingPipelineStateRHIParamRef Pipeline;
 	FRayTracingSceneRHIParamRef Scene;
 	FRayTracingShaderBindings GlobalResourceBindings;
+	uint32 RayGenShaderIndex;
 	uint32 Width;
 	uint32 Height;
 
-	FRHICommandRayTraceDispatch(FRayTracingPipelineStateRHIParamRef InPipeline, const FRayTracingShaderBindings& InGlobalResourceBindings, uint32 InWidth, uint32 InHeight)
+	FRHICommandRayTraceDispatch(FRayTracingPipelineStateRHIParamRef InPipeline, uint32 InRayGenShaderIndex, const FRayTracingShaderBindings& InGlobalResourceBindings, uint32 InWidth, uint32 InHeight)
 		: Pipeline(InPipeline)
 		, Scene(nullptr)
 		, GlobalResourceBindings(InGlobalResourceBindings)
+		, RayGenShaderIndex(InRayGenShaderIndex)
 		, Width(InWidth)
 		, Height(InHeight)
 	{}
 
-	
-
-	FRHICommandRayTraceDispatch(FRayTracingPipelineStateRHIParamRef InPipeline, FRayTracingSceneRHIParamRef InScene, const FRayTracingShaderBindings& InGlobalResourceBindings, uint32 InWidth, uint32 InHeight)
+	FRHICommandRayTraceDispatch(FRayTracingPipelineStateRHIParamRef InPipeline, uint32 InRayGenShaderIndex, FRayTracingSceneRHIParamRef InScene, const FRayTracingShaderBindings& InGlobalResourceBindings, uint32 InWidth, uint32 InHeight)
 		: Pipeline(InPipeline)
 		, Scene(InScene)
 		, GlobalResourceBindings(InGlobalResourceBindings)
+		, RayGenShaderIndex(InRayGenShaderIndex)
 		, Width(InWidth)
 		, Height(InHeight)
 	{}
@@ -1910,15 +1911,17 @@ struct FRHICommandSetRayTracingHitGroup final : public FRHICommand<FRHICommandSe
 	FRayTracingSceneRHIParamRef Scene;
 	uint32 InstanceIndex;
 	uint32 SegmentIndex;
+	uint32 ShaderSlot;
 	FRayTracingPipelineStateRHIParamRef Pipeline;
 	uint32 HitGroupIndex;
 	uint32 NumUniformBuffers;
-	const FUniformBufferRHIParamRef* UniformBuffers; // Pointer to an array of uniform buffers, allocated inline with the command list
+	const FUniformBufferRHIParamRef* UniformBuffers; // Pointer to an array of uniform buffers, allocated inline within the command list
 
-	FRHICommandSetRayTracingHitGroup(FRayTracingSceneRHIParamRef InScene, uint32 InInstanceIndex, uint32 InSegmentIndex, FRayTracingPipelineStateRHIParamRef InPipeline, uint32 InHitGroupIndex, uint32 InNumUniformBuffers, const FUniformBufferRHIParamRef* InUniformBuffers)
+	FRHICommandSetRayTracingHitGroup(FRayTracingSceneRHIParamRef InScene, uint32 InInstanceIndex, uint32 InSegmentIndex, uint32 InShaderSlot, FRayTracingPipelineStateRHIParamRef InPipeline, uint32 InHitGroupIndex, uint32 InNumUniformBuffers, const FUniformBufferRHIParamRef* InUniformBuffers)
 		: Scene(InScene)
 		, InstanceIndex(InInstanceIndex)
 		, SegmentIndex(InSegmentIndex)
+		, ShaderSlot(InShaderSlot)
 		, Pipeline(InPipeline)
 		, HitGroupIndex(InHitGroupIndex)
 		, NumUniformBuffers(InNumUniformBuffers)
@@ -2994,38 +2997,38 @@ public:
 		}
 	}
 
-	FORCEINLINE_DEBUGGABLE void RayTraceDispatch(FRayTracingPipelineStateRHIParamRef Pipeline, const FRayTracingShaderBindings& GlobalResourceBindings, uint32 Width, uint32 Height)
+	FORCEINLINE_DEBUGGABLE void RayTraceDispatch(FRayTracingPipelineStateRHIParamRef Pipeline, uint32 RayGenShaderIndex, const FRayTracingShaderBindings& GlobalResourceBindings, uint32 Width, uint32 Height)
 	{
 		if (Bypass())
 		{
-			GetContext().RHIRayTraceDispatch(Pipeline, GlobalResourceBindings, Width, Height);
+			GetContext().RHIRayTraceDispatch(Pipeline, RayGenShaderIndex, GlobalResourceBindings, Width, Height);
 		}
 		else
 		{
-			ALLOC_COMMAND(FRHICommandRayTraceDispatch)(Pipeline, GlobalResourceBindings, Width, Height);
+			ALLOC_COMMAND(FRHICommandRayTraceDispatch)(Pipeline, RayGenShaderIndex, GlobalResourceBindings, Width, Height);
 		}
 	}
 
-	FORCEINLINE_DEBUGGABLE void RayTraceDispatch(FRayTracingPipelineStateRHIParamRef Pipeline, FRayTracingSceneRHIParamRef Scene, const FRayTracingShaderBindings& GlobalResourceBindings, uint32 Width, uint32 Height)
+	FORCEINLINE_DEBUGGABLE void RayTraceDispatch(FRayTracingPipelineStateRHIParamRef Pipeline, uint32 RayGenShaderIndex, FRayTracingSceneRHIParamRef Scene, const FRayTracingShaderBindings& GlobalResourceBindings, uint32 Width, uint32 Height)
 	{
 		if (Bypass())
 		{
-			GetContext().RHIRayTraceDispatch(Pipeline, Scene, GlobalResourceBindings, Width, Height);
+			GetContext().RHIRayTraceDispatch(Pipeline, RayGenShaderIndex, Scene, GlobalResourceBindings, Width, Height);
 		}
 		else
 		{
-			ALLOC_COMMAND(FRHICommandRayTraceDispatch)(Pipeline, Scene, GlobalResourceBindings, Width, Height);
+			ALLOC_COMMAND(FRHICommandRayTraceDispatch)(Pipeline, RayGenShaderIndex, Scene, GlobalResourceBindings, Width, Height);
 		}
 	}
 
 	FORCEINLINE_DEBUGGABLE void SetRayTracingHitGroup(
-		FRayTracingSceneRHIParamRef Scene, uint32 InstanceIndex, uint32 SegmentIndex,
+		FRayTracingSceneRHIParamRef Scene, uint32 InstanceIndex, uint32 SegmentIndex, uint32 ShaderSlot,
 		FRayTracingPipelineStateRHIParamRef Pipeline, uint32 HitGroupIndex,
 		uint32 NumUniformBuffers, const FUniformBufferRHIParamRef* UniformBuffers)
 	{
 		if (Bypass())
 		{
-			GetContext().RHISetRayTracingHitGroup(Scene, InstanceIndex, SegmentIndex, Pipeline, HitGroupIndex, NumUniformBuffers, UniformBuffers);
+			GetContext().RHISetRayTracingHitGroup(Scene, InstanceIndex, SegmentIndex, ShaderSlot, Pipeline, HitGroupIndex, NumUniformBuffers, UniformBuffers);
 		}
 		else
 		{
@@ -3039,7 +3042,7 @@ public:
 				}
 			}
 
-			ALLOC_COMMAND(FRHICommandSetRayTracingHitGroup)(Scene, InstanceIndex, SegmentIndex, Pipeline, HitGroupIndex, NumUniformBuffers, InlineUniformBuffers);
+			ALLOC_COMMAND(FRHICommandSetRayTracingHitGroup)(Scene, InstanceIndex, SegmentIndex, ShaderSlot, Pipeline, HitGroupIndex, NumUniformBuffers, InlineUniformBuffers);
 		}
 	}
 
diff --git a/Engine/Source/Runtime/RHI/Public/RHICommandListCommandExecutes.inl b/Engine/Source/Runtime/RHI/Public/RHICommandListCommandExecutes.inl
--- a/Engine/Source/Runtime/RHI/Public/RHICommandListCommandExecutes.inl
+++ b/Engine/Source/Runtime/RHI/Public/RHICommandListCommandExecutes.inl
@@ -586,20 +586,20 @@ void FRHICommandRayTraceDispatch::Execute(FRHICommandListBase& CmdList)
 	if (Scene)
 	{
 		// Dispatch rays using SBT associated with the given scene and pipeline.
-		INTERNAL_DECORATOR(RHIRayTraceDispatch)(Pipeline, Scene, GlobalResourceBindings, Width, Height);
+		INTERNAL_DECORATOR(RHIRayTraceDispatch)(Pipeline, RayGenShaderIndex, Scene, GlobalResourceBindings, Width, Height);
 	}
 	else
 	{
 		// Dispatch rays with default shader binding table
 		// #dxr_todo: we should remove this code path entirely and always require an explicit SBT.
-		INTERNAL_DECORATOR(RHIRayTraceDispatch)(Pipeline, GlobalResourceBindings, Width, Height);
+		INTERNAL_DECORATOR(RHIRayTraceDispatch)(Pipeline, RayGenShaderIndex, GlobalResourceBindings, Width, Height);
 	}
 }
 
 void FRHICommandSetRayTracingHitGroup::Execute(FRHICommandListBase& CmdList)
 {
 	RHISTAT(SetRayTracingHitGroup);
-	INTERNAL_DECORATOR(RHISetRayTracingHitGroup)(Scene, InstanceIndex, SegmentIndex, Pipeline, HitGroupIndex, NumUniformBuffers, UniformBuffers);
+	INTERNAL_DECORATOR(RHISetRayTracingHitGroup)(Scene, InstanceIndex, SegmentIndex, ShaderSlot, Pipeline, HitGroupIndex, NumUniformBuffers, UniformBuffers);
 }
 
 #endif // RHI_RAYTRACING
diff --git a/Engine/Source/Runtime/RHI/Public/RHIContext.h b/Engine/Source/Runtime/RHI/Public/RHIContext.h
--- a/Engine/Source/Runtime/RHI/Public/RHIContext.h
+++ b/Engine/Source/Runtime/RHI/Public/RHIContext.h
@@ -631,14 +631,14 @@ public:
 		checkNoEntry();
 	}
 
-	virtual void RHIRayTraceDispatch(FRayTracingPipelineStateRHIParamRef RayTracingPipelineState,
+	virtual void RHIRayTraceDispatch(FRayTracingPipelineStateRHIParamRef RayTracingPipelineState, uint32 RayGenShaderIndex,
 		const FRayTracingShaderBindings& GlobalResourceBindings,
 		uint32 Width, uint32 Height)
 	{
 		checkNoEntry();
 	}
 
-	virtual void RHIRayTraceDispatch(FRayTracingPipelineStateRHIParamRef RayTracingPipelineState,
+	virtual void RHIRayTraceDispatch(FRayTracingPipelineStateRHIParamRef RayTracingPipelineState, uint32 RayGenShaderIndex,
 		FRayTracingSceneRHIParamRef Scene, 
 		const FRayTracingShaderBindings& GlobalResourceBindings,
 		uint32 Width, uint32 Height)
@@ -647,7 +647,7 @@ public:
 	}
 
 	virtual void RHISetRayTracingHitGroup(
-		FRayTracingSceneRHIParamRef Scene, uint32 InstanceIndex, uint32 SegmentIndex,
+		FRayTracingSceneRHIParamRef Scene, uint32 InstanceIndex, uint32 SegmentIndex, uint32 ShaderSlot,
 		FRayTracingPipelineStateRHIParamRef Pipeline, uint32 HitGroupIndex,
 		uint32 NumUniformBuffers, const FUniformBufferRHIParamRef* UniformBuffers)
 	{
diff --git a/Engine/Source/Runtime/RHI/Public/RHIResources.h b/Engine/Source/Runtime/RHI/Public/RHIResources.h
--- a/Engine/Source/Runtime/RHI/Public/RHIResources.h
+++ b/Engine/Source/Runtime/RHI/Public/RHIResources.h
@@ -2001,81 +2001,85 @@ public:
 
 #if RHI_RAYTRACING
 
-struct FRayTracingHitGroupInitializer
-{
-	FRayTracingShaderRHIParamRef ShaderRHI = nullptr;
-
-	bool operator==(const FRayTracingHitGroupInitializer& rhs) const
-	{
-		return ShaderRHI == rhs.ShaderRHI;
-	}
-
-	inline friend uint32 GetTypeHash(const FRayTracingHitGroupInitializer& Initializer)
-	{
-		return PointerHash(Initializer.ShaderRHI); // Only consider the shader when computing the hash (operator== performs full state test)
-	}
-};
-
 class FRayTracingPipelineStateInitializer
 {
 public:
 
-	FRayTracingPipelineStateInitializer() {}
+	FRayTracingPipelineStateInitializer() {};
 
 	// NOTE: GetTypeHash(const FRayTracingPipelineStateInitializer& Initializer) should also be updated when changing this function
 	bool operator==(const FRayTracingPipelineStateInitializer& rhs) const
 	{
 		return MaxPayloadSizeInBytes == rhs.MaxPayloadSizeInBytes
-			&& RayGenShaderRHI == rhs.RayGenShaderRHI
-			&& MissShaderRHI == rhs.MissShaderRHI
-			&& DefaultClosestHitShaderRHI == rhs.DefaultClosestHitShaderRHI
+			&& HitGroupStride == rhs.HitGroupStride
+			&& RayGenHash == rhs.RayGenHash
+			&& MissHash == rhs.MissHash
 			&& HitGroupHash == rhs.HitGroupHash;
 	}
 
-	uint32 MaxPayloadSizeInBytes = 32; // sizeof FDefaultPayload
+	uint32 MaxPayloadSizeInBytes = 32; // sizeof FDefaultPayload declared in RayTracingCommon.ush
 
-	FRayTracingShaderRHIParamRef RayGenShaderRHI = nullptr;
+	// Hit group stride controls how many slots will be allocated in the shader binding table per geometry segment.
+	// Changing this value allows different hit shaders to be used for different effects.
+	// For example, setting this to 2 allows one hit shader for regular material evaluation and a different one for shadows.
+	// Desired hit shader can be selected by providing appropriate RayContributionToHitGroupIndex to TraceRay() function.
+	// Setting hit group stride to 0 effectively disables hit group indexing entirely, forcing the same hit shader to be used for all ray hits.
+	// Use ShaderSlot argument in SetRayTracingHitGroup() to assign shaders and resources for specific part of the shder binding table record.
+	uint32 HitGroupStride = 1;
 
-	// Shader that will be invoked if a ray misses all geometry.
-	// If this is not provided, a default miss shader will be used that sets HitT member of FDefaultPayload to -1.
-	FRayTracingShaderRHIParamRef MissShaderRHI = nullptr;
-
-	// Closest hit shader that will be assigned to all geometry by default.
-	// If this is NULL, then built-in default shader will be used which assumes the following payload:
-	// 	struct FDefaultPayload
-	// 	{
-	// 		float2 Barycentrics;
-	// 		uint   InstanceID;
-	// 		uint   InstanceIndex;
-	// 		uint   PrimitiveIndex;
-	// 		float  HitT;
-	// 		float2 Padding;
-	// 	};
-	FRayTracingShaderRHIParamRef DefaultClosestHitShaderRHI = nullptr;
+	const TArrayView<const FRayTracingShaderRHIParamRef>& GetRayGenTable()   const { return RayGenTable; }
+	const TArrayView<const FRayTracingShaderRHIParamRef>& GetMissTable()     const { return MissTable; }
+	const TArrayView<const FRayTracingShaderRHIParamRef>& GetHitGroupTable() const { return HitGroupTable; }
 
+	// Shaders used as entry point to ray tracing work. At least one RayGen shader must be provided.
+	void SetRayGenShaderTable(const TArrayView<const FRayTracingShaderRHIParamRef>& InRayGenShaders)
+	{
+		RayGenTable = InRayGenShaders;
+		RayGenHash = ComputeShaderTableHash(InRayGenShaders);
+	}
 
-	const TArrayView<const FRayTracingHitGroupInitializer>& GetHitGroups() const { return HitGroups; }
+	// Shaders that will be invoked if a ray misses all geometry.
+	// If this table is empty, then a built-in default miss shader will be used that sets HitT member of FDefaultPayload to -1.
+	// Desired miss shader can be selected by providing MissShaderIndex to TraceRay() function.
+	void SetMissShaderTable(const TArrayView<const FRayTracingShaderRHIParamRef>& InMissShaders)
+	{
+		MissTable = InMissShaders;
+		MissHash = ComputeShaderTableHash(InMissShaders);
+	}
 
-	void SetHitGroups(const TArrayView<const FRayTracingHitGroupInitializer>& InHitGroups)
+	// Shaders that will be invoked when ray intersects geometry.
+	// If this table is empty, then a built-in default shader will be used for all geometry, using FDefaultPayload.
+	void SetHitGroupTable(const TArrayView<const FRayTracingShaderRHIParamRef>& InHitGroups)
 	{
-		HitGroups = InHitGroups;
+		HitGroupTable = InHitGroups;
+		HitGroupHash = ComputeShaderTableHash(HitGroupTable);
+	}
 
-		uint32 CombinedHash = 2085640061; // Large prime number
-		for (const FRayTracingHitGroupInitializer& HitGroup : HitGroups)
+	uint32 GetHitGroupHash() const { return HitGroupHash; }
+	uint32 GetRayGenHash()   const { return RayGenHash; }
+	uint32 GetRayMissHash()  const { return MissHash; }
+
+private:
+
+	uint32 ComputeShaderTableHash(const TArrayView<const FRayTracingShaderRHIParamRef>& ShaderTable, uint32 InitialHash = 2085640061)
+	{
+		uint32 CombinedHash = InitialHash;
+		for (FRayTracingShaderRHIParamRef ShaderRHI : ShaderTable)
 		{
 			// #dxr_todo: some sort of session-unique ID should be used instead of pointers to ensure that unique hash is
 			// produced if the same memory happens to be re-used for a different shader (i.e. delete followed by new).
-			CombinedHash = PointerHash(HitGroup.ShaderRHI, CombinedHash);
+			CombinedHash = PointerHash(ShaderRHI, CombinedHash);
 		}
 
-		HitGroupHash = CombinedHash;
+		return CombinedHash;
 	}
 
-	uint32 GetHitGroupHash() const { return HitGroupHash; }
-
-private:
+	TArrayView<const FRayTracingShaderRHIParamRef> RayGenTable;
+	TArrayView<const FRayTracingShaderRHIParamRef> MissTable;
+	TArrayView<const FRayTracingShaderRHIParamRef> HitGroupTable;
 
-	TArrayView<const FRayTracingHitGroupInitializer> HitGroups;
+	uint32 RayGenHash = 0;
+	uint32 MissHash = 0;
 	uint32 HitGroupHash = 0;
 };
 #endif // RHI_RAYTRACING
diff --git a/Engine/Source/Runtime/RenderCore/Private/Shader.cpp b/Engine/Source/Runtime/RenderCore/Private/Shader.cpp
--- a/Engine/Source/Runtime/RenderCore/Private/Shader.cpp
+++ b/Engine/Source/Runtime/RenderCore/Private/Shader.cpp
@@ -396,7 +396,7 @@ TArray<uint32> FShaderResource::GlobalUnusedIndicies;
 TMap<uint32, FRHIRayTracingShader*> FShaderResource::GlobalRayTracingMaterialLibrary;
 FCriticalSection FShaderResource::GlobalRayTracingMaterialLibraryCS;
 
-void FShaderResource::GetRayTracingMaterialLibrary(TArray<FRayTracingHitGroupInitializer>& RayTracingMaterials)
+void FShaderResource::GetRayTracingMaterialLibrary(TArray<FRayTracingShaderRHIParamRef>& RayTracingMaterials, FRayTracingShaderRHIParamRef DefaultShader)
 {
 	FScopeLock Lock(&GlobalRayTracingMaterialLibraryCS);
 	RayTracingMaterials.Reset();
@@ -404,18 +404,16 @@ void FShaderResource::GetRayTracingMaterialLibrary(TArray<FRayTracingHitGroupIni
 
 	for (const auto Entry : GlobalRayTracingMaterialLibrary)
 	{
-		FRayTracingHitGroupInitializer HitGroupInitializer;
-		HitGroupInitializer.ShaderRHI = Entry.Value;
-		RayTracingMaterials[Entry.Key] = HitGroupInitializer;
+		RayTracingMaterials[Entry.Key] = Entry.Value;
 	}
 
 	for (uint32 Index : GlobalUnusedIndicies)
 	{
-		RayTracingMaterials[Index] = FRayTracingHitGroupInitializer();
+		RayTracingMaterials[Index] = DefaultShader;
 	}
 }
 
-uint32 FShaderResource::AddToRayTracingLibrary(FRHIRayTracingShader* Shader)
+uint32 FShaderResource::AddToRayTracingLibrary(FRayTracingShaderRHIParamRef Shader)
 {
 	FScopeLock Lock(&GlobalRayTracingMaterialLibraryCS);
 
diff --git a/Engine/Source/Runtime/RenderCore/Public/RayGenShaderUtils.h b/Engine/Source/Runtime/RenderCore/Public/RayGenShaderUtils.h
--- a/Engine/Source/Runtime/RenderCore/Public/RayGenShaderUtils.h
+++ b/Engine/Source/Runtime/RenderCore/Public/RayGenShaderUtils.h
@@ -35,10 +35,12 @@ struct RENDERCORE_API FRayGenShaderUtils
 			SetShaderParameters(GlobalResources, RayGenerationShader, *Parameters);
 
 			FRayTracingPipelineStateInitializer Initializer;
-			Initializer.RayGenShaderRHI = RayGenerationShader->GetRayTracingShader();
+			FRayTracingShaderRHIParamRef RayGenShaderTable[] = { RayGenerationShader->GetRayTracingShader() };
+			Initializer.SetRayGenShaderTable(RayGenShaderTable);
 
 			FRHIRayTracingPipelineState* Pipeline = PipelineStateCache::GetAndOrCreateRayTracingPipelineState(Initializer);
-			RHICmdList.RayTraceDispatch(Pipeline, GlobalResources, Resolution.X, Resolution.Y);
+			const uint32 RayGenShaderIndex = 0;
+			RHICmdList.RayTraceDispatch(Pipeline, RayGenShaderIndex, GlobalResources, Resolution.X, Resolution.Y);
 		});
 	}
 };
diff --git a/Engine/Source/Runtime/RenderCore/Public/Shader.h b/Engine/Source/Runtime/RenderCore/Public/Shader.h
--- a/Engine/Source/Runtime/RenderCore/Public/Shader.h
+++ b/Engine/Source/Runtime/RenderCore/Public/Shader.h
@@ -377,15 +377,15 @@ public:
 		return RayTracingMaterialLibraryIndex;
 	}
 
-	RENDERCORE_API static void GetRayTracingMaterialLibrary(TArray<FRayTracingHitGroupInitializer>& RayTracingMaterials);
+	RENDERCORE_API static void GetRayTracingMaterialLibrary(TArray<FRayTracingShaderRHIParamRef>& RayTracingMaterials, FRayTracingShaderRHIParamRef DefaultShader);
 
 private:
-	RENDERCORE_API static uint32 AddToRayTracingLibrary(FRHIRayTracingShader* Shader);
+	RENDERCORE_API static uint32 AddToRayTracingLibrary(FRayTracingShaderRHIParamRef Shader);
 	RENDERCORE_API static void RemoveFromRayTracingLibrary(uint32 Index);
 
 	static uint32 GlobalMaxIndex;
 	static TArray<uint32> GlobalUnusedIndicies;
-	static TMap<uint32, FRHIRayTracingShader*> GlobalRayTracingMaterialLibrary;
+	static TMap<uint32, FRayTracingShaderRHIParamRef> GlobalRayTracingMaterialLibrary;
 	static FCriticalSection GlobalRayTracingMaterialLibraryCS;
 
 public:
diff --git a/Engine/Source/Runtime/Renderer/Private/DeferredShadingRenderer.cpp b/Engine/Source/Runtime/Renderer/Private/DeferredShadingRenderer.cpp
--- a/Engine/Source/Runtime/Renderer/Private/DeferredShadingRenderer.cpp
+++ b/Engine/Source/Runtime/Renderer/Private/DeferredShadingRenderer.cpp
@@ -900,29 +900,54 @@ bool FDeferredShadingSceneRenderer::DispatchRayTracingWorldUpdates(FRHICommandLi
 	return true;
 }
 
-FRHIRayTracingPipelineState* FDeferredShadingSceneRenderer::BindRayTracingPipeline(FRHICommandList& RHICmdList, const FViewInfo& View, FRayTracingShaderRHIParamRef RayGenShader, FRayTracingShaderRHIParamRef MissShader, FRayTracingShaderRHIParamRef DefaultChsShader)
+FRHIRayTracingPipelineState* FDeferredShadingSceneRenderer::BindRayTracingPipeline(FRHICommandList& RHICmdList, const FViewInfo& View, FRayTracingShaderRHIParamRef RayGenShader, FRayTracingShaderRHIParamRef MissShader, FRayTracingShaderRHIParamRef DefaultClosestHitShader)
 {
 	SCOPE_CYCLE_COUNTER(STAT_BindRayTracingPipeline);
 
+	FRHIRayTracingPipelineState* PipelineState = nullptr;
+
 	FRayTracingPipelineStateInitializer Initializer;
-	Initializer.RayGenShaderRHI = RayGenShader;
-	Initializer.DefaultClosestHitShaderRHI = DefaultChsShader;
-	Initializer.MissShaderRHI = MissShader;
 
-	TArray<FRayTracingHitGroupInitializer> RayTracingMaterialLibrary;
-	FShaderResource::GetRayTracingMaterialLibrary(RayTracingMaterialLibrary);
-	Initializer.SetHitGroups(RayTracingMaterialLibrary);
 	Initializer.MaxPayloadSizeInBytes = 152; // #dxr_todo: set Initializer.MaxPayloadSizeInBytes based on shader requirements (it's not obvious how to compute this right now)
 
-	FRHIRayTracingPipelineState* PipelineState = PipelineStateCache::GetAndOrCreateRayTracingPipelineState(Initializer);
+	FRayTracingShaderRHIParamRef RayGenShaderTable[] = { RayGenShader };
+	Initializer.SetRayGenShaderTable(RayGenShaderTable);
+
+	FRayTracingShaderRHIParamRef MissShaderTable[] = { MissShader };
+	Initializer.SetMissShaderTable(MissShaderTable);
 
-	if (GEnableRayTracingMaterials)
+	const bool bEnableMaterials = GEnableRayTracingMaterials;
+
+	if (bEnableMaterials)
 	{
-		for (const FVisibleMeshDrawCommand& VisibleMeshDrawCommand : View.RaytraycingVisibleMeshDrawCommands)
-		{
-			const FMeshDrawCommand& MeshDrawCommand = *VisibleMeshDrawCommand.MeshDrawCommand;
-			MeshDrawCommand.ShaderBindings.SetOnRayTracingStructure(RHICmdList, View.PerViewRayTracingScene.RayTracingSceneRHI, VisibleMeshDrawCommand.RayTracedInstanceIndex, MeshDrawCommand.RayTracedSegmentIndex, PipelineState, MeshDrawCommand.RayTracingMaterialLibraryIndex);
-		}
+		TArray<FRayTracingShaderRHIParamRef> RayTracingMaterialLibrary;
+		FShaderResource::GetRayTracingMaterialLibrary(RayTracingMaterialLibrary, DefaultClosestHitShader);
+		Initializer.SetHitGroupTable(RayTracingMaterialLibrary);
+
+		PipelineState = PipelineStateCache::GetAndOrCreateRayTracingPipelineState(Initializer);
+	}
+	else
+	{
+		FRayTracingShaderRHIParamRef HitGroupTable[] = { DefaultClosestHitShader };
+		Initializer.SetHitGroupTable(HitGroupTable);
+
+		PipelineState = PipelineStateCache::GetAndOrCreateRayTracingPipelineState(Initializer);
+	}
+
+	for (const FVisibleMeshDrawCommand& VisibleMeshDrawCommand : View.RaytraycingVisibleMeshDrawCommands)
+	{
+		const FMeshDrawCommand& MeshDrawCommand = *VisibleMeshDrawCommand.MeshDrawCommand;
+
+		const uint32 HitGroupIndex = bEnableMaterials
+			? MeshDrawCommand.RayTracingMaterialLibraryIndex
+			: 0; // Force the same shader to be used on all geometry
+
+		MeshDrawCommand.ShaderBindings.SetOnRayTracingStructure(RHICmdList,
+			View.PerViewRayTracingScene.RayTracingSceneRHI,
+			VisibleMeshDrawCommand.RayTracedInstanceIndex,
+			MeshDrawCommand.RayTracedSegmentIndex,
+			PipelineState,
+			HitGroupIndex);
 	}
 
 	return PipelineState;
diff --git a/Engine/Source/Runtime/Renderer/Private/DeferredShadingRenderer.h b/Engine/Source/Runtime/Renderer/Private/DeferredShadingRenderer.h
--- a/Engine/Source/Runtime/Renderer/Private/DeferredShadingRenderer.h
+++ b/Engine/Source/Runtime/Renderer/Private/DeferredShadingRenderer.h
@@ -497,7 +497,7 @@ private:
 
 	bool GatherRayTracingWorldInstances(FRHICommandListImmediate& RHICmdList);
 	bool DispatchRayTracingWorldUpdates(FRHICommandListImmediate& RHICmdList);
-	FRHIRayTracingPipelineState* BindRayTracingPipeline(FRHICommandList& RHICmdList, const FViewInfo& View, FRayTracingShaderRHIParamRef RayGenShader, FRayTracingShaderRHIParamRef MissShader, FRayTracingShaderRHIParamRef DefaultChsShader);
+	FRHIRayTracingPipelineState* BindRayTracingPipeline(FRHICommandList& RHICmdList, const FViewInfo& View, FRayTracingShaderRHIParamRef RayGenShader, FRayTracingShaderRHIParamRef MissShader, FRayTracingShaderRHIParamRef DefaultClosestHitShader);
 #endif // RHI_RAYTRACING
 };
 
diff --git a/Engine/Source/Runtime/Renderer/Private/MeshPassProcessor.cpp b/Engine/Source/Runtime/Renderer/Private/MeshPassProcessor.cpp
--- a/Engine/Source/Runtime/Renderer/Private/MeshPassProcessor.cpp
+++ b/Engine/Source/Runtime/Renderer/Private/MeshPassProcessor.cpp
@@ -295,7 +295,8 @@ void FMeshDrawShaderBindings::SetOnRayTracingStructure(FRHICommandList& RHICmdLi
 
 	check(SegmentIndex < 0xFF);
 	uint32 NumUniformBuffersToSet = MaxUniformBufferUsed + 1;
-	RHICmdList.SetRayTracingHitGroup(Scene, InstanceIndex, SegmentIndex, PipelineState, HitGroupIndex, NumUniformBuffersToSet, LocalUniformBuffers);
+	const uint32 ShaderSlot = 0; // Multiple shader slots can be used for different ray types. Slot 0 is the primary material slot.
+	RHICmdList.SetRayTracingHitGroup(Scene, InstanceIndex, SegmentIndex, ShaderSlot, PipelineState, HitGroupIndex, NumUniformBuffersToSet, LocalUniformBuffers);
 }
 #endif // RHI_RAYTRACING
 
diff --git a/Engine/Source/Runtime/Renderer/Private/PathTracing.cpp b/Engine/Source/Runtime/Renderer/Private/PathTracing.cpp
--- a/Engine/Source/Runtime/Renderer/Private/PathTracing.cpp
+++ b/Engine/Source/Runtime/Renderer/Private/PathTracing.cpp
@@ -479,14 +479,18 @@ void FDeferredShadingSceneRenderer::RenderPathTracing(FRHICommandListImmediate&
 	ClearUAV(RHICmdList, RayCountPerPixelRT->GetRenderTargetItem(), FLinearColor::Black);
 
 	auto RayGenShader = GetGlobalShaderMap(FeatureLevel)->GetShader<FPathTracingRG>();
-	auto ClosestHitShader = GetGlobalShaderMap(FeatureLevel)->GetShader<FPathTracingCHS>();
 	auto MissShader = GetGlobalShaderMap(FeatureLevel)->GetShader<FPathTracingMS>();
+	auto ClosestHitShader = GetGlobalShaderMap(FeatureLevel)->GetShader<FPathTracingCHS>();
 
 #if 0
 	FRayTracingPipelineStateInitializer Initializer;
-	Initializer.RayGenShaderRHI = RayGenShader->GetRayTracingShader();
-	Initializer.DefaultClosestHitShaderRHI = ClosestHitShader->GetRayTracingShader();
-	Initializer.MissShaderRHI = MissShader->GetRayTracingShader();
+	FRayTracingShaderRHIParamRef RayGenShaderTable[] = { RayGenShader->GetRayTracingShader() };
+	Initializer.SetRayGenShaderTable(RayGenShaderTable);
+	FRayTracingShaderRHIParamRef MissShaderTable[] = { MissShader->GetRayTracingShader() };
+	Initializer.SetMissShaderTable(MissShaderTable);
+	FRayTracingShaderRHIParamRef HitGroupTable[] = { ClosestHitShader->GetRayTracingShader() };
+	Initializer.SetHitGroupTable(HitGroupTable);
+	Initializer.HitGroupStride = 0; // Use the same hit shader for all geometry in the scene by disabling SBT indexing.
 	FRHIRayTracingPipelineState* Pipeline = PipelineStateCache::GetAndOrCreateRayTracingPipelineState(Initializer); // #dxr_todo: this should be done once at load-time and cached
 #else
 	FRHIRayTracingPipelineState* Pipeline = BindRayTracingPipeline(RHICmdList, View,
@@ -528,7 +532,9 @@ void FDeferredShadingSceneRenderer::RenderPathTracing(FRHICommandListImmediate&
 	);
 
 	FRayTracingSceneRHIParamRef RayTracingSceneRHI = View.PerViewRayTracingScene.RayTracingSceneRHI;
-	RHICmdList.RayTraceDispatch(Pipeline, RayTracingSceneRHI, GlobalResources, View.ViewRect.Size().X, View.ViewRect.Size().Y);
+
+	const uint32 RayGenShaderIndex = 0;
+	RHICmdList.RayTraceDispatch(Pipeline, RayGenShaderIndex, RayTracingSceneRHI, GlobalResources, View.ViewRect.Size().X, View.ViewRect.Size().Y);
 
 	// Save RayTracingIndirect for compositing
 	RHICmdList.CopyToResolveTarget(RadianceRT->GetRenderTargetItem().TargetableTexture, RadianceRT->GetRenderTargetItem().ShaderResourceTexture, FResolveParams());
diff --git a/Engine/Source/Runtime/Renderer/Private/RayTracing/RayTracingAmbientOcclusion.cpp b/Engine/Source/Runtime/Renderer/Private/RayTracing/RayTracingAmbientOcclusion.cpp
--- a/Engine/Source/Runtime/Renderer/Private/RayTracing/RayTracingAmbientOcclusion.cpp
+++ b/Engine/Source/Runtime/Renderer/Private/RayTracing/RayTracingAmbientOcclusion.cpp
@@ -102,7 +102,9 @@ public:
 	)
 	{
 		FRayTracingPipelineStateInitializer Initializer;
-		Initializer.RayGenShaderRHI = GetRayTracingShader();
+
+		FRayTracingShaderRHIParamRef RayGenShaderTable[] = { GetRayTracingShader() };
+		Initializer.SetRayGenShaderTable(RayGenShaderTable);
 
 		FRHIRayTracingPipelineState* Pipeline = PipelineStateCache::GetAndOrCreateRayTracingPipelineState(Initializer); // #dxr_todo: this should be done once at load-time and cached
 
@@ -113,7 +115,8 @@ public:
 		GlobalResources.Set(AmbientOcclusionParameter, AmbientOcclusionUniformBuffer);
 		GlobalResources.Set(OcclusionMaskUAVParameter, OcclusionMaskUAV);
 
-		RHICmdList.RayTraceDispatch(Pipeline, GlobalResources, Width, Height);
+		const uint32 RayGenShaderIndex = 0;
+		RHICmdList.RayTraceDispatch(Pipeline, RayGenShaderIndex, GlobalResources, Width, Height);
 	}
 
 private:
diff --git a/Engine/Source/Runtime/Renderer/Private/RayTracing/RayTracingBarycentrics.cpp b/Engine/Source/Runtime/Renderer/Private/RayTracing/RayTracingBarycentrics.cpp
--- a/Engine/Source/Runtime/Renderer/Private/RayTracing/RayTracingBarycentrics.cpp
+++ b/Engine/Source/Runtime/Renderer/Private/RayTracing/RayTracingBarycentrics.cpp
@@ -85,9 +85,16 @@ void FDeferredShadingSceneRenderer::RenderRayTracedBarycentrics(FRHICommandListI
 	auto MissShader = ShaderMap->GetShader<FRayTracingBarycentricsMS>();
 
 	FRayTracingPipelineStateInitializer Initializer;
-	Initializer.RayGenShaderRHI = RayGenShader->GetRayTracingShader();
-	Initializer.DefaultClosestHitShaderRHI = ClosestHitShader->GetRayTracingShader();
-	Initializer.MissShaderRHI = MissShader->GetRayTracingShader();
+
+	FRayTracingShaderRHIParamRef RayGenShaderTable[] = { RayGenShader->GetRayTracingShader() };
+	Initializer.SetRayGenShaderTable(RayGenShaderTable);
+
+	FRayTracingShaderRHIParamRef MissShaderTable[] = { MissShader->GetRayTracingShader() };
+	Initializer.SetMissShaderTable(MissShaderTable);
+
+	FRayTracingShaderRHIParamRef HitGroupTable[] = { ClosestHitShader->GetRayTracingShader() };
+	Initializer.SetHitGroupTable(HitGroupTable);
+	Initializer.HitGroupStride = 0; // Use the same hit shader for all geometry in the scene by disabling SBT indexing.
 
 	FRHIRayTracingPipelineState* Pipeline = PipelineStateCache::GetAndOrCreateRayTracingPipelineState(Initializer); // #dxr_todo: this should be done once at load-time and cached
 
@@ -111,7 +118,8 @@ void FDeferredShadingSceneRenderer::RenderRayTracedBarycentrics(FRHICommandListI
 		SetShaderParameters(GlobalResources, RayGenShader, *RayGenParameters);
 
 		// Dispatch rays using default shader binding table
-		RHICmdList.RayTraceDispatch(Pipeline, GlobalResources, ViewRect.Size().X, ViewRect.Size().Y);
+		const uint32 RayGenShaderIndex = 0;
+		RHICmdList.RayTraceDispatch(Pipeline, RayGenShaderIndex, GlobalResources, ViewRect.Size().X, ViewRect.Size().Y);
 	});
 
 	GraphBuilder.Execute();
diff --git a/Engine/Source/Runtime/Renderer/Private/RayTracing/RayTracingDebug.cpp b/Engine/Source/Runtime/Renderer/Private/RayTracing/RayTracingDebug.cpp
--- a/Engine/Source/Runtime/Renderer/Private/RayTracing/RayTracingDebug.cpp
+++ b/Engine/Source/Runtime/Renderer/Private/RayTracing/RayTracingDebug.cpp
@@ -110,7 +110,7 @@ void FDeferredShadingSceneRenderer::RenderRayTracedDebug(FRHICommandListImmediat
 
 	FRDGBuilder GraphBuilder(RHICmdList);
 
-	TShaderMap<FGlobalShaderType>* ShaderMap = GetGlobalShaderMap(FeatureLevel);	
+	TShaderMap<FGlobalShaderType>* ShaderMap = GetGlobalShaderMap(FeatureLevel);
 
 	auto RayGenShader = ShaderMap->GetShader<FRayTracingDebugRGS>();
 	auto ClosestHitShader = ShaderMap->GetShader<FRayTracingDebugCHS>();
@@ -144,7 +144,8 @@ void FDeferredShadingSceneRenderer::RenderRayTracedDebug(FRHICommandListImmediat
 		SetShaderParameters(GlobalResources, RayGenShader, *RayGenParameters);
 
 		// Dispatch rays using default shader binding table
-		RHICmdList.RayTraceDispatch(Pipeline, RayTracingSceneRHI, GlobalResources, ViewRect.Size().X, ViewRect.Size().Y);
+		const uint32 RayGenShaderIndex = 0;
+		RHICmdList.RayTraceDispatch(Pipeline, RayGenShaderIndex, RayTracingSceneRHI, GlobalResources, ViewRect.Size().X, ViewRect.Size().Y);
 	});
 
 	GraphBuilder.Execute();
diff --git a/Engine/Source/Runtime/Renderer/Private/RayTracing/RayTracingReflections.cpp b/Engine/Source/Runtime/Renderer/Private/RayTracing/RayTracingReflections.cpp
--- a/Engine/Source/Runtime/Renderer/Private/RayTracing/RayTracingReflections.cpp
+++ b/Engine/Source/Runtime/Renderer/Private/RayTracing/RayTracingReflections.cpp
@@ -277,7 +277,8 @@ void FDeferredShadingSceneRenderer::RayTraceReflections(
 		SetShaderParameters(GlobalResources, RayGenShader, *PassParameters);
 
 		FRayTracingSceneRHIParamRef RayTracingSceneRHI = View.PerViewRayTracingScene.RayTracingSceneRHI;
-		RHICmdList.RayTraceDispatch(Pipeline, RayTracingSceneRHI, GlobalResources, RayTracingResolution.X, RayTracingResolution.Y);
+		const uint32 RayGenShaderIndex = 0;
+		RHICmdList.RayTraceDispatch(Pipeline, RayGenShaderIndex, RayTracingSceneRHI, GlobalResources, RayTracingResolution.X, RayTracingResolution.Y);
 	});
 }
 #else // !RHI_RAYTRACING
diff --git a/Engine/Source/Runtime/Renderer/Private/RayTracing/RayTracingShadowsForRectLight.cpp b/Engine/Source/Runtime/Renderer/Private/RayTracing/RayTracingShadowsForRectLight.cpp
--- a/Engine/Source/Runtime/Renderer/Private/RayTracing/RayTracingShadowsForRectLight.cpp
+++ b/Engine/Source/Runtime/Renderer/Private/RayTracing/RayTracingShadowsForRectLight.cpp
@@ -141,7 +141,9 @@ public:
 	)
 	{
 		FRayTracingPipelineStateInitializer Initializer;
-		Initializer.RayGenShaderRHI = GetRayTracingShader();
+
+		FRayTracingShaderRHIParamRef RayGenShaderTable[] = { GetRayTracingShader() };
+		Initializer.SetRayGenShaderTable(RayGenShaderTable);
 
 		FRHIRayTracingPipelineState* Pipeline = PipelineStateCache::GetAndOrCreateRayTracingPipelineState(Initializer); // #dxr_todo: this should be done once at load-time and cached
 
@@ -153,7 +155,8 @@ public:
 		GlobalResources.Set(OcclusionMaskUAVParameter, OcclusionMaskUAV);
 		GlobalResources.Set(RayDistanceUAVParameter, RayDistanceUAV);
 
-		RHICmdList.RayTraceDispatch(Pipeline, GlobalResources, Width, Height);
+		const uint32 RayGenShaderIndex = 0;
+		RHICmdList.RayTraceDispatch(Pipeline, RayGenShaderIndex, GlobalResources, Width, Height);
 	}
 
 private:
diff --git a/Engine/Source/Runtime/Renderer/Private/RayTracing/RaytracingSkylight.cpp b/Engine/Source/Runtime/Renderer/Private/RayTracing/RaytracingSkylight.cpp
--- a/Engine/Source/Runtime/Renderer/Private/RayTracing/RaytracingSkylight.cpp
+++ b/Engine/Source/Runtime/Renderer/Private/RayTracing/RaytracingSkylight.cpp
@@ -225,7 +225,9 @@ public:
 	)
 	{
 		FRayTracingPipelineStateInitializer Initializer;
-		Initializer.RayGenShaderRHI = GetRayTracingShader();
+
+		FRayTracingShaderRHIParamRef RayGenShaderTable[] = { GetRayTracingShader() };
+		Initializer.SetRayGenShaderTable(RayGenShaderTable);
 
 		FRHIRayTracingPipelineState* Pipeline = PipelineStateCache::GetAndOrCreateRayTracingPipelineState(Initializer); // #dxr_todo: this should be done once at load-time and cached
 
@@ -253,7 +255,8 @@ public:
 			GlobalResources.SetSampler(TransmissionProfilesLinearSamplerParameter.GetBaseIndex(), TStaticSamplerState<SF_Bilinear, AM_Clamp, AM_Clamp, AM_Clamp>::GetRHI());
 		}
 
-		RHICmdList.RayTraceDispatch(Pipeline, GlobalResources, Width, Height);
+		const uint32 RayGenShaderIndex = 0;
+		RHICmdList.RayTraceDispatch(Pipeline, RayGenShaderIndex, GlobalResources, Width, Height);
 	}
 
 private:
