Cherry-pick 4827676: Verify that uniform buffer deference of render graph resources is done on render thread.

diff --git a/Engine/Source/Runtime/D3D12RHI/Private/D3D12UniformBuffer.cpp b/Engine/Source/Runtime/D3D12RHI/Private/D3D12UniformBuffer.cpp
--- a/Engine/Source/Runtime/D3D12RHI/Private/D3D12UniformBuffer.cpp
+++ b/Engine/Source/Runtime/D3D12RHI/Private/D3D12UniformBuffer.cpp
@@ -90,6 +90,7 @@ FUniformBufferRHIRef FD3D12DynamicRHI::RHICreateUniformBuffer(const void* Conten
 				}
 				else if (IsRDGResourceReferenceShaderParameterType(ResourceType))
 				{
+					check(IsInRenderingThread()); // TODO: UE-68018
 					FRHIResource** ResourcePtr = *(FRHIResource***)((uint8*)Contents + Layout.Resources[i].MemberOffset);
 					Resource = ResourcePtr ? *ResourcePtr : nullptr;
 				}
@@ -214,6 +215,7 @@ void FD3D12DynamicRHI::RHIUpdateUniformBuffer(FUniformBufferRHIParamRef UniformB
 			}
 			else if (IsRDGResourceReferenceShaderParameterType(ResourceType))
 			{
+				check(IsInRenderingThread()); // TODO: UE-68018
 				FRHIResource** ResourcePtr = *(FRHIResource***)((uint8*)Contents + Layout.Resources[ResourceIndex].MemberOffset);
 				Resource = ResourcePtr ? *ResourcePtr : nullptr;
 			}
