Aftermath support for DX12 (based on cl 3892738 + tweaks and fixes)

diff --git a/Engine/Source/Runtime/D3D12RHI/D3D12RHI.Build.cs b/Engine/Source/Runtime/D3D12RHI/D3D12RHI.Build.cs
--- a/Engine/Source/Runtime/D3D12RHI/D3D12RHI.Build.cs
+++ b/Engine/Source/Runtime/D3D12RHI/D3D12RHI.Build.cs
@@ -41,5 +41,9 @@ public class D3D12RHI : ModuleRules
 			AddEngineThirdPartyPrivateStaticDependencies(Target, "AMD_AGS");
 			AddEngineThirdPartyPrivateStaticDependencies(Target, "IntelMetricsDiscovery");
 		}
-	}
+        if (Target.Platform == UnrealTargetPlatform.Win64 )
+        {
+            AddEngineThirdPartyPrivateStaticDependencies(Target, "NVAftermath");
+        }
+    }
 }
diff --git a/Engine/Source/Runtime/D3D12RHI/Private/D3D12Adapter.cpp b/Engine/Source/Runtime/D3D12RHI/Private/D3D12Adapter.cpp
--- a/Engine/Source/Runtime/D3D12RHI/Private/D3D12Adapter.cpp
+++ b/Engine/Source/Runtime/D3D12RHI/Private/D3D12Adapter.cpp
@@ -183,6 +183,47 @@ void FD3D12Adapter::CreateRootDevice(bool bWithDebug)
 	}
 #endif // D3D12_RHI_RAYTRACING
 
+#if NV_AFTERMATH
+	// Two ways to enable aftermath, command line or the r.GPUCrashDebugging variable
+	// Note: If intending to change this please alert game teams who use this for user support.
+	if (FParse::Param(FCommandLine::Get(), TEXT("gpucrashdebugging")))
+	{
+		GDX12NVAfterMathEnabled = true;
+	}
+	else
+	{
+		static IConsoleVariable* GPUCrashDebugging = IConsoleManager::Get().FindConsoleVariable(TEXT("r.GPUCrashDebugging"));
+		if (GPUCrashDebugging)
+		{
+			GDX12NVAfterMathEnabled = GPUCrashDebugging->GetInt();
+		}
+	}
+
+	if (GDX12NVAfterMathEnabled)
+	{
+		if (IsRHIDeviceNVIDIA())
+		{
+			GFSDK_Aftermath_Result Result = GFSDK_Aftermath_DX12_Initialize(GFSDK_Aftermath_Version_API, GFSDK_Aftermath_FeatureFlags_Maximum, RootDevice);
+			if (Result == GFSDK_Aftermath_Result_Success)
+			{
+				UE_LOG(LogD3D12RHI, Log, TEXT("[Aftermath] Aftermath enabled and primed"));
+				SetEmitDrawEvents(true);
+			}
+			else
+			{
+				UE_LOG(LogD3D12RHI, Log, TEXT("[Aftermath] Aftermath enabled but failed to initialize (%x)"), Result);
+				GDX12NVAfterMathEnabled = 0;
+			}
+		}
+		else
+		{
+			GDX12NVAfterMathEnabled = 0;
+			UE_LOG(LogD3D12RHI, Warning, TEXT("[Aftermath] Skipping aftermath initialization on non-Nvidia device"));
+		}
+	}
+#endif
+
+
 #if UE_BUILD_DEBUG	&& PLATFORM_WINDOWS
 	//break on debug
 	TRefCountPtr<ID3D12Debug> d3dDebug;
diff --git a/Engine/Source/Runtime/D3D12RHI/Private/D3D12CommandContext.cpp b/Engine/Source/Runtime/D3D12RHI/Private/D3D12CommandContext.cpp
--- a/Engine/Source/Runtime/D3D12RHI/Private/D3D12CommandContext.cpp
+++ b/Engine/Source/Runtime/D3D12RHI/Private/D3D12CommandContext.cpp
@@ -137,9 +137,14 @@ void FD3D12CommandContext::RHIPushEvent(const TCHAR* Name, FColor Color)
 {
 	if (IsDefaultContext())
 	{
+#if NV_AFTERMATH
+		GetParentDevice()->PushGPUEvent(Name, Color, CommandListHandle.AftermathCommandContext());
+#else
 		GetParentDevice()->PushGPUEvent(Name, Color);
+#endif
 	}
 
+
 #if PLATFORM_WINDOWS
 	AGSContext* const AmdAgsContext = OwningRHI.GetAmdAgsContext();
 	if (GEmitRgpFrameMarkers && AmdAgsContext)
diff --git a/Engine/Source/Runtime/D3D12RHI/Private/D3D12CommandList.cpp b/Engine/Source/Runtime/D3D12RHI/Private/D3D12CommandList.cpp
--- a/Engine/Source/Runtime/D3D12RHI/Private/D3D12CommandList.cpp
+++ b/Engine/Source/Runtime/D3D12RHI/Private/D3D12CommandList.cpp
@@ -71,6 +71,16 @@ FD3D12CommandListHandle::FD3D12CommandListData::FD3D12CommandListData(FD3D12Devi
 	SetName(CommandList, Name.GetCharArray().GetData());
 #endif
 
+#if NV_AFTERMATH
+	if (GDX12NVAfterMathEnabled)
+	{
+		GFSDK_Aftermath_Result Result = GFSDK_Aftermath_DX12_CreateContextHandle(CommandList, &AftermathHandle);
+
+		check(Result == GFSDK_Aftermath_Result_Success);
+		ParentDevice->GetParentAdapter()->GetGPUProfiler().RegisterCommandList(AftermathHandle);
+	}
+#endif
+
 	// Initially start with all lists closed.  We'll open them as we allocate them.
 	Close();
 
@@ -81,6 +91,17 @@ FD3D12CommandListHandle::FD3D12CommandListData::FD3D12CommandListData(FD3D12Devi
 
 FD3D12CommandListHandle::FD3D12CommandListData::~FD3D12CommandListData()
 {
+#if NV_AFTERMATH
+	if (GDX12NVAfterMathEnabled && GetParentDevice() && GetParentDevice()->GetParentAdapter())
+	{
+		GetParentDevice()->GetParentAdapter()->GetGPUProfiler().UnregisterCommandList(AftermathHandle);
+
+		GFSDK_Aftermath_Result Result = GFSDK_Aftermath_ReleaseContextHandle(AftermathHandle);
+
+		check(Result == GFSDK_Aftermath_Result_Success);
+	}
+#endif
+
 	CommandList.SafeRelease();
 	DEC_DWORD_STAT(STAT_D3D12NumCommandLists);
 
diff --git a/Engine/Source/Runtime/D3D12RHI/Private/D3D12CommandList.h b/Engine/Source/Runtime/D3D12RHI/Private/D3D12CommandList.h
--- a/Engine/Source/Runtime/D3D12RHI/Private/D3D12CommandList.h
+++ b/Engine/Source/Runtime/D3D12RHI/Private/D3D12CommandList.h
@@ -247,6 +247,9 @@ private:
 #if D3D12_RHI_RAYTRACING
 		TRefCountPtr<ID3D12GraphicsCommandList4> RayTracingCommandList;
 #endif // D3D12_RHI_RAYTRACING
+#if NV_AFTERMATH
+		GFSDK_Aftermath_ContextHandle			AftermathHandle;
+#endif
 		FD3D12CommandAllocator*					CurrentCommandAllocator;	// Command allocator currently being used for recording the command list
 		uint64									CurrentGeneration;
 		uint64									LastCompleteGeneration;
@@ -410,6 +413,14 @@ public:
 	}
 #endif // D3D12_RHI_RAYTRACING
 
+#if NV_AFTERMATH
+	GFSDK_Aftermath_ContextHandle AftermathCommandContext() const
+	{
+		check(CommandListData);
+		return CommandListData->AftermathHandle;
+	}
+#endif
+
 	uint64 CurrentGeneration() const
 	{
 		check(CommandListData);
diff --git a/Engine/Source/Runtime/D3D12RHI/Private/D3D12Device.cpp b/Engine/Source/Runtime/D3D12RHI/Private/D3D12Device.cpp
--- a/Engine/Source/Runtime/D3D12RHI/Private/D3D12Device.cpp
+++ b/Engine/Source/Runtime/D3D12RHI/Private/D3D12Device.cpp
@@ -372,6 +372,13 @@ void FD3D12Device::PushGPUEvent(const TCHAR* Name, FColor Color)
 	GetParentAdapter()->GetGPUProfiler().PushEvent(Name, Color);
 }
 
+#if NV_AFTERMATH
+void FD3D12Device::PushGPUEvent(const TCHAR* Name, FColor Color, GFSDK_Aftermath_ContextHandle Context)
+{
+	GetParentAdapter()->GetGPUProfiler().PushEvent(Name, Color, Context);
+}
+#endif
+
 void FD3D12Device::PopGPUEvent()
 {
 	GetParentAdapter()->GetGPUProfiler().PopEvent();
diff --git a/Engine/Source/Runtime/D3D12RHI/Private/D3D12Device.h b/Engine/Source/Runtime/D3D12RHI/Private/D3D12Device.h
--- a/Engine/Source/Runtime/D3D12RHI/Private/D3D12Device.h
+++ b/Engine/Source/Runtime/D3D12RHI/Private/D3D12Device.h
@@ -110,6 +110,10 @@ public:
 	void PushGPUEvent(const TCHAR* Name, FColor Color);
 	void PopGPUEvent();
 
+#if NV_AFTERMATH
+	void PushGPUEvent(const TCHAR* Name, FColor Color, GFSDK_Aftermath_ContextHandle Context);
+#endif
+
 	FD3D12SamplerState* CreateSampler(const FSamplerStateInitializerRHI& Initializer);
 	void CreateSamplerInternal(const D3D12_SAMPLER_DESC& Desc, D3D12_CPU_DESCRIPTOR_HANDLE Descriptor);
 
diff --git a/Engine/Source/Runtime/D3D12RHI/Private/D3D12RHI.cpp b/Engine/Source/Runtime/D3D12RHI/Private/D3D12RHI.cpp
--- a/Engine/Source/Runtime/D3D12RHI/Private/D3D12RHI.cpp
+++ b/Engine/Source/Runtime/D3D12RHI/Private/D3D12RHI.cpp
@@ -508,3 +508,8 @@ uint32 FD3D12DynamicRHI::GetDebugFlags()
 {
 	return GetAdapter().GetDebugFlags();
 }
+
+bool FD3D12DynamicRHI::CheckGpuHeartbeat() const
+{
+	return ChosenAdapters[0]->GetGPUProfiler().CheckGpuHeartbeat();
+}
diff --git a/Engine/Source/Runtime/D3D12RHI/Private/D3D12RHIPrivate.h b/Engine/Source/Runtime/D3D12RHI/Private/D3D12RHIPrivate.h
--- a/Engine/Source/Runtime/D3D12RHI/Private/D3D12RHIPrivate.h
+++ b/Engine/Source/Runtime/D3D12RHI/Private/D3D12RHIPrivate.h
@@ -48,6 +48,13 @@ DECLARE_LOG_CATEGORY_EXTERN(LogD3D12RHI, Log, All);
 #include "XboxOne/D3D12RHIBasePrivate.h"
 #endif
 
+#if NV_AFTERMATH
+#define GFSDK_Aftermath_WITH_DX12 1
+#include "GFSDK_Aftermath.h"
+#undef GFSDK_Aftermath_WITH_DX12
+extern int32 GDX12NVAfterMathEnabled;
+#endif
+
 #include "D3D12Residency.h"
 
 // D3D RHI public headers.
@@ -466,6 +473,8 @@ public:
 
 #endif //D3D12_RHI_RAYTRACING
 
+	bool CheckGpuHeartbeat() const override;
+
 	bool IsQuadBufferStereoEnabled() const;
 	void DisableQuadBufferStereo();
 
diff --git a/Engine/Source/Runtime/D3D12RHI/Private/D3D12Stats.cpp b/Engine/Source/Runtime/D3D12RHI/Private/D3D12Stats.cpp
--- a/Engine/Source/Runtime/D3D12RHI/Private/D3D12Stats.cpp
+++ b/Engine/Source/Runtime/D3D12RHI/Private/D3D12Stats.cpp
@@ -16,6 +16,13 @@ namespace D3D12RHI
 		check(!bTrackingEvents);
 		check(!CurrentEventNodeFrame); // this should have already been cleaned up and the end of the previous frame
 
+		// update the crash tracking variables
+		static auto* CrashCollectionEnableCvar = IConsoleManager::Get().FindTConsoleVariableDataInt(TEXT("r.gpucrash.collectionenable"));
+		static auto* CrashCollectionDataDepth = IConsoleManager::Get().FindTConsoleVariableDataInt(TEXT("r.gpucrash.datadepth"));
+
+		bTrackingGPUCrashData = CrashCollectionEnableCvar ? CrashCollectionEnableCvar->GetValueOnRenderThread() != 0 : false;
+		GPUCrashDataDepth = CrashCollectionDataDepth ? CrashCollectionDataDepth->GetValueOnRenderThread() : -1;
+
 									   // latch the bools from the game thread into our private copy
 		bLatchedGProfilingGPU = GTriggerGPUProfile;
 		bLatchedGProfilingGPUHitches = GTriggerGPUHitchProfile;
@@ -58,7 +65,12 @@ namespace D3D12RHI
 
 		if (GetEmitDrawEvents())
 		{
+#if NV_AFTERMATH
+			// Assuming that grabbing the device 0 command list here is OK
+			PushEvent(TEXT("FRAME"), FColor(0, 255, 0, 255), InRHI->GetAdapter().GetDevice(0)->GetCommandContext().CommandListHandle.AftermathCommandContext());
+#else
 			PushEvent(TEXT("FRAME"), FColor(0, 255, 0, 255));
+#endif
 		}
 	}
 }
@@ -187,12 +199,61 @@ void FD3DGPUProfiler::PushEvent(const TCHAR* Name, FColor Color)
 	FGPUProfiler::PushEvent(Name, Color);
 }
 
+static FString EventDeepString(TEXT("EventTooDeep"));
+static const uint32 EventDeepCRC = FCrc::StrCrc32<TCHAR>(*EventDeepString);
+#if NV_AFTERMATH
+void FD3DGPUProfiler::PushEvent(const TCHAR* Name, FColor Color, GFSDK_Aftermath_ContextHandle Context)
+{
+
+	if (GDX12NVAfterMathEnabled && bTrackingGPUCrashData)
+	{
+		uint32 CRC = 0;
+		if (GPUCrashDataDepth < 0 || PushPopStack.Num() < GPUCrashDataDepth)
+		{
+			CRC = FCrc::StrCrc32<TCHAR>(Name);
+
+			if (CachedStrings.Num() > 10000)
+			{
+				CachedStrings.Empty(10000);
+				CachedStrings.Emplace(EventDeepCRC, EventDeepString);
+			}
+
+			if (CachedStrings.Find(CRC) == nullptr)
+			{
+				CachedStrings.Emplace(CRC, FString(Name));
+			}
+
+		}
+		else
+		{
+			CRC = EventDeepCRC;
+		}
+		PushPopStack.Push(CRC);
+
+		GFSDK_Aftermath_SetEventMarker(Context, &PushPopStack[0], PushPopStack.Num() * sizeof(uint32));
+	}
+
+	PushEvent(Name, Color);
+}
+#endif
+
 void FD3DGPUProfiler::PopEvent()
 {
 #if WITH_DX_PERF
 	D3DPERF_EndEvent();
 #endif
 
+#if NV_AFTERMATH
+	if (GDX12NVAfterMathEnabled && bTrackingGPUCrashData)
+	{
+		// need to look for unbalanced push/pop
+		if (PushPopStack.Num() > 0)
+		{
+			PushPopStack.Pop(false);
+		}
+	}
+#endif
+
 	FGPUProfiler::PopEvent();
 }
 
@@ -309,3 +370,94 @@ void UpdateBufferStats(FD3D12ResourceLocation* ResourceLocation, bool bAllocatin
 #endif
 	}
 }
+
+#if NV_AFTERMATH
+void FD3DGPUProfiler::RegisterCommandList(GFSDK_Aftermath_ContextHandle context)
+{
+	FScopeLock Lock(&AftermathLock);
+
+	AftermathContexts.Push(context);
+}
+
+void FD3DGPUProfiler::UnregisterCommandList(GFSDK_Aftermath_ContextHandle context)
+{
+	FScopeLock Lock(&AftermathLock);
+
+	int32 Item = AftermathContexts.Find(context);
+
+	AftermathContexts.RemoveAt(Item);
+}
+#endif
+
+extern CORE_API bool GIsGPUCrashed;
+bool FD3DGPUProfiler::CheckGpuHeartbeat() const
+{
+#if NV_AFTERMATH
+	if (GDX12NVAfterMathEnabled)
+	{
+		GFSDK_Aftermath_Device_Status Status;
+		GFSDK_Aftermath_Result Result = GFSDK_Aftermath_GetDeviceStatus(&Status);
+		if (Result == GFSDK_Aftermath_Result_Success)
+		{
+			if (Status != GFSDK_Aftermath_Device_Status_Active)
+			{
+				GIsGPUCrashed = true;
+				const TCHAR* AftermathReason[] = { TEXT("Active"), TEXT("Timeout"), TEXT("OutOfMemory"), TEXT("PageFault"), TEXT("Unknown") };
+				check(Status < ARRAYSIZE(AftermathReason));
+				UE_LOG(LogRHI, Error, TEXT("[Aftermath] Status: %s"), AftermathReason[Status]);
+
+				TArray<GFSDK_Aftermath_ContextData> ContextDataOut;
+				ContextDataOut.AddUninitialized(AftermathContexts.Num());
+				Result = GFSDK_Aftermath_GetData(AftermathContexts.Num(), AftermathContexts.GetData(), ContextDataOut.GetData());
+				if (Result == GFSDK_Aftermath_Result_Success)
+				{
+					UE_LOG(LogRHI, Error, TEXT("[Aftermath] Scanning %d command lists for dumps"), ContextDataOut.Num());
+					for (GFSDK_Aftermath_ContextData& ContextData : ContextDataOut)
+					{
+						if (ContextData.status == GFSDK_Aftermath_Context_Status_Executing)
+						{
+							UE_LOG(LogRHI, Error, TEXT("[Aftermath] GPU Stack Dump"));
+							uint32 NumCRCs = ContextData.markerSize / sizeof(uint32);
+							uint32* Data = (uint32*)ContextData.markerData;
+							for (uint32 i = 0; i < NumCRCs; i++)
+							{
+								const FString* Frame = CachedStrings.Find(Data[i]);
+								if (Frame != nullptr)
+								{
+									UE_LOG(LogRHI, Error, TEXT("[Aftermath] %i: %s"), i, *(*Frame));
+								}
+							}
+							UE_LOG(LogRHI, Error, TEXT("[Aftermath] GPU Stack Dump"));
+						}
+					}
+				}
+				else
+				{
+					UE_LOG(LogRHI, Error, TEXT("[Aftermath] Failed to get Aftermath stack data"));
+				}
+
+				if (Status == GFSDK_Aftermath_Device_Status_PageFault)
+				{
+					GFSDK_Aftermath_PageFaultInformation FaultInformation;
+					Result = GFSDK_Aftermath_GetPageFaultInformation(&FaultInformation);
+
+					if (Result == GFSDK_Aftermath_Result_Success)
+					{
+						UE_LOG(LogRHI, Error, TEXT("[Aftermath] Faulting address: 0x%xI64"), FaultInformation.faultingGpuVA);
+						UE_LOG(LogRHI, Error, TEXT("[Aftermath] Faulting resournce dims: %d x %d x %d"), FaultInformation.resourceDesc.width, FaultInformation.resourceDesc.height, FaultInformation.resourceDesc.depth);
+						UE_LOG(LogRHI, Error, TEXT("[Aftermath] Faulting result size: 0x%dI64"), FaultInformation.resourceDesc.size);
+						UE_LOG(LogRHI, Error, TEXT("[Aftermath] Faulting resource mips: %d"), FaultInformation.resourceDesc.mipLevels);
+						UE_LOG(LogRHI, Error, TEXT("[Aftermath] Faulting resource format: 0x%x"), FaultInformation.resourceDesc.format);
+					}
+					else
+					{
+						UE_LOG(LogRHI, Error, TEXT("[Aftermath] No information on faulting address"));
+					}
+				}
+				return false;
+			}
+		}
+	}
+#endif
+	return true;
+}
diff --git a/Engine/Source/Runtime/D3D12RHI/Private/D3D12Stats.h b/Engine/Source/Runtime/D3D12RHI/Private/D3D12Stats.h
--- a/Engine/Source/Runtime/D3D12RHI/Private/D3D12Stats.h
+++ b/Engine/Source/Runtime/D3D12RHI/Private/D3D12Stats.h
@@ -328,5 +328,21 @@ namespace D3D12RHI
 		void BeginFrame(class FD3D12DynamicRHI* InRHI);
 
 		void EndFrame(class FD3D12DynamicRHI* InRHI);
+
+		bool CheckGpuHeartbeat() const;
+
+#if NV_AFTERMATH
+		virtual void PushEvent(const TCHAR* Name, FColor Color, GFSDK_Aftermath_ContextHandle Context);
+
+		void RegisterCommandList(GFSDK_Aftermath_ContextHandle context);
+		void UnregisterCommandList(GFSDK_Aftermath_ContextHandle context);
+
+		TArray<GFSDK_Aftermath_ContextHandle> AftermathContexts;
+		FCriticalSection AftermathLock;
+
+		TArray<uint32> PushPopStack;
+		TMap<uint32, FString> CachedStrings;
+#endif
+
 	};
 }
\ No newline at end of file
diff --git a/Engine/Source/Runtime/D3D12RHI/Private/D3D12Util.cpp b/Engine/Source/Runtime/D3D12RHI/Private/D3D12Util.cpp
--- a/Engine/Source/Runtime/D3D12RHI/Private/D3D12Util.cpp
+++ b/Engine/Source/Runtime/D3D12RHI/Private/D3D12Util.cpp
@@ -184,6 +184,11 @@ static FString GetD3D12TextureFlagString(uint32 TextureFlags)
 extern CORE_API bool GIsGPUCrashed;
 static void TerminateOnDeviceRemoved(HRESULT D3DResult)
 {
+	if (GDynamicRHI)
+	{
+		GDynamicRHI->CheckGpuHeartbeat();
+	}
+
 	if (D3DResult == DXGI_ERROR_DEVICE_REMOVED)
 	{
 		GIsGPUCrashed = true;
diff --git a/Engine/Source/Runtime/D3D12RHI/Private/Windows/WindowsD3D12Device.cpp b/Engine/Source/Runtime/D3D12RHI/Private/Windows/WindowsD3D12Device.cpp
--- a/Engine/Source/Runtime/D3D12RHI/Private/Windows/WindowsD3D12Device.cpp
+++ b/Engine/Source/Runtime/D3D12RHI/Private/Windows/WindowsD3D12Device.cpp
@@ -34,6 +34,17 @@ static TAutoConsoleVariable<int32> CVarGraphicsAdapter(
 	TEXT("  1: Adapter #1, ..."),
 	ECVF_RenderThreadSafe);
 
+#if NV_AFTERMATH
+// Disabled by default since introduces stalls between render and driver threads
+int32 GDX12NVAfterMathEnabled = 0;
+static FAutoConsoleVariableRef CVarDX12NVAfterMathBufferSize(
+	TEXT("r.DX12NVAfterMathEnabled"),
+	GDX12NVAfterMathEnabled,
+	TEXT("Use NV Aftermath for GPU crash analysis in D3D12"),
+	ECVF_ReadOnly
+);
+#endif
+
 static inline int D3D12RHI_PreferAdapterVendor()
 {
 	if (FParse::Param(FCommandLine::Get(), TEXT("preferAMD")))
@@ -448,6 +459,22 @@ FDynamicRHI* FD3D12DynamicRHIModule::CreateRHI(ERHIFeatureLevel::Type RequestedF
 
 void FD3D12DynamicRHIModule::StartupModule()
 {
+#if NV_AFTERMATH
+	// Note - can't check device type here, we'll check for that before actually initializing Aftermath
+	FString AftermathBinariesRoot = FPaths::EngineDir() / TEXT("Binaries/ThirdParty/NVIDIA/NVaftermath/Win64/");
+	if (LoadLibraryW(*(AftermathBinariesRoot + "GFSDK_Aftermath_Lib.x64.dll")) == nullptr)
+	{
+		UE_LOG(LogD3D12RHI, Warning, TEXT("Failed to load GFSDK_Aftermath_Lib.x64.dll"));
+		GDX12NVAfterMathEnabled = 0;
+		return;
+	}
+	else
+	{
+		UE_LOG(LogD3D12RHI, Log, TEXT("Aftermath initialized"));
+		GDX12NVAfterMathEnabled = 1;
+	}
+#endif
+
 #if USE_PIX
 	static FString WindowsPixDllRelativePath = FPaths::Combine(*FPaths::EngineDir(), TEXT("Binaries/ThirdParty/Windows/DirectX/x64"));
 	static FString WindowsPixDll("WinPixEventRuntime.dll");
diff --git a/Engine/Source/Runtime/Windows/D3D11RHI/Private/D3D11RHI.cpp b/Engine/Source/Runtime/Windows/D3D11RHI/Private/D3D11RHI.cpp
--- a/Engine/Source/Runtime/Windows/D3D11RHI/Private/D3D11RHI.cpp
+++ b/Engine/Source/Runtime/Windows/D3D11RHI/Private/D3D11RHI.cpp
@@ -648,16 +648,16 @@ bool FD3DGPUProfiler::CheckGpuHeartbeat() const
 	if (GDX11NVAfterMathEnabled && bTrackingGPUCrashData)
 	{
 		GFSDK_Aftermath_Device_Status Status;
-		auto Result = GFSDK_Aftermath_GetDeviceStatus(&Status);
+		GFSDK_Aftermath_Result Result = GFSDK_Aftermath_GetDeviceStatus(&Status);
 		if (Result == GFSDK_Aftermath_Result_Success)
 		{
 			if (Status != GFSDK_Aftermath_Device_Status_Active)
 			{
 				GIsGPUCrashed = true;
 				const TCHAR* AftermathReason[] = { TEXT("Active"), TEXT("Timeout"), TEXT("OutOfMemory"), TEXT("PageFault"), TEXT("Unknown") };
-				check(Status < 5);
+				check(Status < ARRAYSIZE(AftermathReason));
 				UE_LOG(LogRHI, Error, TEXT("[Aftermath] Status: %s"), AftermathReason[Status]);
-				auto AftermathContext = D3D11RHI->GetNVAftermathContext();
+				GFSDK_Aftermath_ContextHandle AftermathContext = D3D11RHI->GetNVAftermathContext();
 
 				if (AftermathContext)
 				{
