Fix crash on exit with -gpucrashdebugging (NV aftermath) enabled

diff --git a/Engine/Source/Runtime/D3D12RHI/Private/D3D12CommandList.cpp b/Engine/Source/Runtime/D3D12RHI/Private/D3D12CommandList.cpp
--- a/Engine/Source/Runtime/D3D12RHI/Private/D3D12CommandList.cpp
+++ b/Engine/Source/Runtime/D3D12RHI/Private/D3D12CommandList.cpp
@@ -72,6 +72,8 @@ FD3D12CommandListHandle::FD3D12CommandListData::FD3D12CommandListData(FD3D12Devi
 #endif
 
 #if NV_AFTERMATH
+	AftermathHandle = nullptr;
+
 	if (GDX12NVAfterMathEnabled)
 	{
 		GFSDK_Aftermath_Result Result = GFSDK_Aftermath_DX12_CreateContextHandle(CommandList, &AftermathHandle);
@@ -92,7 +94,7 @@ FD3D12CommandListHandle::FD3D12CommandListData::FD3D12CommandListData(FD3D12Devi
 FD3D12CommandListHandle::FD3D12CommandListData::~FD3D12CommandListData()
 {
 #if NV_AFTERMATH
-	if (GDX12NVAfterMathEnabled && GetParentDevice() && GetParentDevice()->GetParentAdapter())
+	if (AftermathHandle)
 	{
 		GetParentDevice()->GetParentAdapter()->GetGPUProfiler().UnregisterCommandList(AftermathHandle);
 
