Win32 build fix

diff --git a/Engine/Source/Runtime/D3D12RHI/D3D12RHI.Build.cs b/Engine/Source/Runtime/D3D12RHI/D3D12RHI.Build.cs
--- a/Engine/Source/Runtime/D3D12RHI/D3D12RHI.Build.cs
+++ b/Engine/Source/Runtime/D3D12RHI/D3D12RHI.Build.cs
@@ -39,11 +39,8 @@ public class D3D12RHI : ModuleRules
 			AddEngineThirdPartyPrivateStaticDependencies(Target, "DX11");
 			AddEngineThirdPartyPrivateStaticDependencies(Target, "NVAPI");
 			AddEngineThirdPartyPrivateStaticDependencies(Target, "AMD_AGS");
-			AddEngineThirdPartyPrivateStaticDependencies(Target, "IntelMetricsDiscovery");
-		}
-        if (Target.Platform == UnrealTargetPlatform.Win64 )
-        {
             AddEngineThirdPartyPrivateStaticDependencies(Target, "NVAftermath");
+            AddEngineThirdPartyPrivateStaticDependencies(Target, "IntelMetricsDiscovery");
         }
     }
 }
