Cherry-pick 7413442: Ray tracing scene now explicitly holds references to all FD3D12Resource-s that are bound to the SBT to prevent them from being released while some RT commands may still be in flight. Added a command to explicitly clear all SBTs for a scene, which releases resource references and allows resources to be immediately reused.

diff --git a/Engine/Source/Runtime/D3D12RHI/Private/D3D12CommandContext.h b/Engine/Source/Runtime/D3D12RHI/Private/D3D12CommandContext.h
--- a/Engine/Source/Runtime/D3D12RHI/Private/D3D12CommandContext.h
+++ b/Engine/Source/Runtime/D3D12RHI/Private/D3D12CommandContext.h
@@ -338,6 +338,7 @@ public:
 	virtual void RHIUpdateAccelerationStructures(const TArrayView<const FAccelerationStructureUpdateParams> Params) final override;
 	virtual void RHIBuildAccelerationStructures(const TArrayView<const FAccelerationStructureUpdateParams> Params) final override;
 	virtual void RHIBuildAccelerationStructure(FRHIRayTracingScene* Scene) final override;
+	virtual void RHIClearRayTracingBindings(FRHIRayTracingScene* Scene) final override;
 	virtual void RHIRayTraceOcclusion(FRHIRayTracingScene* Scene,
 		FRHIShaderResourceView* Rays,
 		FRHIUnorderedAccessView* Output,
diff --git a/Engine/Source/Runtime/D3D12RHI/Private/D3D12RayTracing.cpp b/Engine/Source/Runtime/D3D12RHI/Private/D3D12RayTracing.cpp
--- a/Engine/Source/Runtime/D3D12RHI/Private/D3D12RayTracing.cpp
+++ b/Engine/Source/Runtime/D3D12RHI/Private/D3D12RayTracing.cpp
@@ -1358,10 +1358,9 @@ public:
 	};
 	TMap<FShaderRecordCacheKey, uint32> ShaderRecordCache;
 
-#if ENABLE_RESIDENCY_MANAGEMENT
 	// A set of all resources referenced by this shader table for the purpose of updating residency before ray tracing work dispatch.
 	// #dxr_todo UE-72159: remove resources from this set when SBT slot entries are replaced
-	TSet<FD3D12Resource*> ReferencedD3D12Resources;
+	TSet<TRefCountPtr<FD3D12Resource>> ReferencedD3D12Resources;
 	TArray<TRefCountPtr<FRHIResource>> ReferencedResources;
 	void AddResourceReference(FD3D12Resource* D3D12Resource, FRHIResource* Resource)
 	{
@@ -1380,10 +1379,6 @@ public:
 		}
 		Buffer->GetResource()->UpdateResidency(CommandContext.CommandListHandle);
 	}
-#else // ENABLE_RESIDENCY_MANAGEMENT
-	void AddResourceReference(FD3D12Resource* D3D12Resource, FRHIResource* Resource) {}
-	FORCEINLINE void UpdateResidency(FD3D12CommandContext&) {}
-#endif // ENABLE_RESIDENCY_MANAGEMENT
 };
 
 
@@ -2609,6 +2604,17 @@ void FD3D12CommandContext::RHIBuildAccelerationStructure(FRHIRayTracingScene* In
 	Scene->BuildAccelerationStructure(*this, D3D12_RAYTRACING_ACCELERATION_STRUCTURE_BUILD_FLAG_PREFER_FAST_TRACE);
 }
 
+void FD3D12CommandContext::RHIClearRayTracingBindings(FRHIRayTracingScene* InScene)
+{
+	FD3D12RayTracingScene* Scene = FD3D12DynamicRHI::ResourceCast(InScene);
+
+	auto& Table = Scene->ShaderTables[GetGPUIndex()];
+	for (auto Item : Table)
+	{
+		delete Item.Value;
+	}
+	Table.Reset();
+}
 
 struct FD3D12RayTracingGlobalResourceBinder
 {
diff --git a/Engine/Source/Runtime/RHI/Public/RHICommandList.h b/Engine/Source/Runtime/RHI/Public/RHICommandList.h
--- a/Engine/Source/Runtime/RHI/Public/RHICommandList.h
+++ b/Engine/Source/Runtime/RHI/Public/RHICommandList.h
@@ -1868,6 +1868,18 @@ struct FRHICommandBuildAccelerationStructure final : public FRHICommand<FRHIComm
 	RHI_API void Execute(FRHICommandListBase& CmdList);
 };
 
+struct FRHICommandClearRayTracingBindings final : public FRHICommand<FRHICommandClearRayTracingBindings>
+{
+	FRHIRayTracingScene* Scene;
+
+	explicit FRHICommandClearRayTracingBindings(FRHIRayTracingScene* InScene)
+		: Scene(InScene)
+	{}
+
+	RHI_API void Execute(FRHICommandListBase& CmdList);
+};
+
+
 struct FRHICommandUpdateAccelerationStructures final : public FRHICommand<FRHICommandUpdateAccelerationStructures>
 {
 	const TArrayView<const FAccelerationStructureUpdateParams> UpdateParams;
@@ -3013,6 +3025,18 @@ public:
 		}
 	}
 
+	FORCEINLINE_DEBUGGABLE void ClearRayTracingBindings(FRHIRayTracingScene* Scene)
+	{
+		if (Bypass())
+		{
+			GetContext().RHIClearRayTracingBindings(Scene);
+		}
+		else
+		{
+			ALLOC_COMMAND(FRHICommandClearRayTracingBindings)(Scene);
+		}
+	}
+
 	/**
 	 * Trace rays from an input buffer of FBasicRayData.
 	 * Binary intersection results are written to output buffer as R32_UINTs.
diff --git a/Engine/Source/Runtime/RHI/Public/RHICommandListCommandExecutes.inl b/Engine/Source/Runtime/RHI/Public/RHICommandListCommandExecutes.inl
--- a/Engine/Source/Runtime/RHI/Public/RHICommandListCommandExecutes.inl
+++ b/Engine/Source/Runtime/RHI/Public/RHICommandListCommandExecutes.inl
@@ -63,6 +63,7 @@ struct FRHICommandTransitionTextures;
 struct FRHICommandTransitionTexturesArray;
 struct FRHICommandUpdateTextureReference;
 struct FRHICommandBuildAccelerationStructure;
+struct FRHICommandClearRayTracingBindings;
 struct FRHICommandRayTraceOcclusion;
 struct FRHICommandRayTraceIntersection;
 struct FRHICommandRayTraceDispatch;
@@ -545,6 +546,12 @@ void FRHICommandBuildAccelerationStructure::Execute(FRHICommandListBase& CmdList
 	}
 }
 
+void FRHICommandClearRayTracingBindings::Execute(FRHICommandListBase& CmdList)
+{
+	RHISTAT(ClearRayTracingBindings);
+	INTERNAL_DECORATOR(RHIClearRayTracingBindings)(Scene);
+}
+
 void FRHICommandUpdateAccelerationStructures::Execute(FRHICommandListBase& CmdList)
 {
 	RHISTAT(UpdateAccelerationStructure);
diff --git a/Engine/Source/Runtime/RHI/Public/RHIContext.h b/Engine/Source/Runtime/RHI/Public/RHIContext.h
--- a/Engine/Source/Runtime/RHI/Public/RHIContext.h
+++ b/Engine/Source/Runtime/RHI/Public/RHIContext.h
@@ -579,6 +579,12 @@ public:
 		checkNoEntry();
 	}
 #endif
+
+	virtual void RHIClearRayTracingBindings(FRHIRayTracingScene* Scene)
+	{
+		checkNoEntry();
+	}
+
 	virtual void RHIBuildAccelerationStructure(FRHIRayTracingGeometry* Geometry)
 	{
 		checkNoEntry();
diff --git a/Engine/Source/Runtime/Renderer/Private/DeferredShadingRenderer.cpp b/Engine/Source/Runtime/Renderer/Private/DeferredShadingRenderer.cpp
--- a/Engine/Source/Runtime/Renderer/Private/DeferredShadingRenderer.cpp
+++ b/Engine/Source/Runtime/Renderer/Private/DeferredShadingRenderer.cpp
@@ -2017,6 +2017,20 @@ void FDeferredShadingSceneRenderer::Render(FRHICommandListImmediate& RHICmdList)
 	// Invalidate the lighting channels
 	SceneContext.LightingChannels.SafeRelease();
 
+
+#if RHI_RAYTRACING
+	// Release resources that were bound to the ray tracing scene to allow them to be immediately recycled.
+	for (int32 ViewIndex = 0; ViewIndex < Views.Num(); ++ViewIndex)
+	{
+		FViewInfo& View = Views[ViewIndex];
+		if (View.RayTracingScene.RayTracingSceneRHI)
+		{
+			RHICmdList.ClearRayTracingBindings(View.RayTracingScene.RayTracingSceneRHI);
+			View.RayTracingScene.RayTracingSceneRHI.SafeRelease();
+		}
+	}
+#endif //  RHI_RAYTRACING
+
 	{
 		SCOPE_CYCLE_COUNTER(STAT_FDeferredShadingSceneRenderer_RenderFinish);
 		RHICmdList.SetCurrentStat(GET_STATID(STAT_CLM_RenderFinish));
