Cherry-pick 7425816: Resurrected CL 4938076 by [REDACTED_PERSON] that uses command list manager fence to release resources instead of frame fence.

diff --git a/Engine/Source/Runtime/D3D12RHI/Private/D3D12Query.cpp b/Engine/Source/Runtime/D3D12RHI/Private/D3D12Query.cpp
--- a/Engine/Source/Runtime/D3D12RHI/Private/D3D12Query.cpp
+++ b/Engine/Source/Runtime/D3D12RHI/Private/D3D12Query.cpp
@@ -452,7 +452,8 @@ void FD3D12QueryHeap::DestroyQueryHeap(bool bDeferDelete)
 	{
 		if (bDeferDelete)
 		{
-			GetParentDevice()->GetParentAdapter()->GetDeferredDeletionQueue().EnqueueResource(QueryHeap);
+			FD3D12Fence& Fence = GetParentDevice()->GetCommandListManager().GetFence();
+			GetParentDevice()->GetParentAdapter()->GetDeferredDeletionQueue().EnqueueResource(QueryHeap, &Fence);
 		}
 		else
 		{
diff --git a/Engine/Source/Runtime/D3D12RHI/Private/D3D12RayTracing.cpp b/Engine/Source/Runtime/D3D12RHI/Private/D3D12RayTracing.cpp
--- a/Engine/Source/Runtime/D3D12RHI/Private/D3D12RayTracing.cpp
+++ b/Engine/Source/Runtime/D3D12RHI/Private/D3D12RayTracing.cpp
@@ -755,8 +755,8 @@ public:
 
 		Entry Result = {};
 
-		FD3D12Adapter* Adapter = GetParentDevice()->GetParentAdapter();
-		const uint64 CompletedFenceValue = Adapter->GetFrameFence().GetLastCompletedFenceFast();
+		const FD3D12Fence& Fence = GetParentDevice()->GetCommandListManager().GetFence();
+		const uint64 CompletedFenceValue = Fence.GetLastCompletedFenceFast();
 
 		for (int32 EntryIndex = 0; EntryIndex < Entries.Num(); ++EntryIndex)
 		{
@@ -840,12 +840,13 @@ public:
 	void Flush()
 	{
 		FD3D12Device* Device = GetParentDevice();
+		FD3D12Fence& Fence = Device->GetCommandListManager().GetFence();
 
 		FScopeLock Lock(&CriticalSection);
 
 		for (const Entry& It : Entries)
 		{
-			Device->GetParentAdapter()->GetDeferredDeletionQueue().EnqueueResource(It.Heap);
+			Device->GetParentAdapter()->GetDeferredDeletionQueue().EnqueueResource(It.Heap, &Fence);
 		}
 		Entries.Empty();
 	}
@@ -920,9 +921,8 @@ struct FD3D12RayTracingDescriptorHeap : public FD3D12DeviceChild
 
 	void UpdateSyncPoint()
 	{
-		FD3D12Adapter* Adapter = GetParentDevice()->GetParentAdapter();
-		const uint64 FrameFenceValue = Adapter->GetFrameFence().GetCurrentFence();
-		HeapCacheEntry.FenceValue = FMath::Max(HeapCacheEntry.FenceValue, FrameFenceValue);
+		const FD3D12Fence& Fence = GetParentDevice()->GetCommandListManager().GetFence();
+		HeapCacheEntry.FenceValue = FMath::Max(HeapCacheEntry.FenceValue, Fence.GetCurrentFence());
 	}
 
 	ID3D12DescriptorHeap* D3D12Heap = nullptr;
diff --git a/Engine/Source/Runtime/D3D12RHI/Private/D3D12Resources.cpp b/Engine/Source/Runtime/D3D12RHI/Private/D3D12Resources.cpp
--- a/Engine/Source/Runtime/D3D12RHI/Private/D3D12Resources.cpp
+++ b/Engine/Source/Runtime/D3D12RHI/Private/D3D12Resources.cpp
@@ -41,31 +41,30 @@ FD3D12DeferredDeletionQueue::~FD3D12DeferredDeletionQueue()
 	}
 }
 
-void FD3D12DeferredDeletionQueue::EnqueueResource(FD3D12Resource* pResource)
+void FD3D12DeferredDeletionQueue::EnqueueResource(FD3D12Resource* pResource, FD3D12Fence* Fence)
 {
 	check(pResource->ShouldDeferDelete());
-	const uint64 CurrentFrameFence = GetParentAdapter()->GetFrameFence().GetCurrentFence();
 
 	// Useful message for identifying when resources are released on the rendering thread.
-	//UE_CLOG(IsInActualRenderingThread(), LogD3D12RHI, Display, TEXT("Rendering Thread: Deleting %#016llx when done with frame fence %llu"), pResource, CurrentFrameFence);
+	//UE_CLOG(IsInActualRenderingThread(), LogD3D12RHI, Display, TEXT("Rendering Thread: Deleting %#016llx when done with frame fence %llu"), pResource, Fence->GetCurrentFence());
 
 	FencedObjectType FencedObject;
 	FencedObject.RHIObject  = pResource;
-	FencedObject.FenceValue = CurrentFrameFence;
+	FencedObject.Fence      = Fence;
+	FencedObject.FenceValue = Fence->GetCurrentFence();
 	FencedObject.Type       = EObjectType::RHI;
 	DeferredReleaseQueue.Enqueue(FencedObject);
 }
 
-void FD3D12DeferredDeletionQueue::EnqueueResource(ID3D12Object* pResource)
+void FD3D12DeferredDeletionQueue::EnqueueResource(ID3D12Object* pResource, FD3D12Fence* Fence)
 {
-	const uint64 CurrentFrameFence = GetParentAdapter()->GetFrameFence().GetCurrentFence();
-
 	// Useful message for identifying when resources are released on the rendering thread.
-	//UE_CLOG(IsInActualRenderingThread(), LogD3D12RHI, Display, TEXT("Rendering Thread: Deleting %#016llx when done with frame fence %llu"), pResource, CurrentFrameFence);
+	//UE_CLOG(IsInActualRenderingThread(), LogD3D12RHI, Display, TEXT("Rendering Thread: Deleting %#016llx when done with frame fence %llu"), pResource, Fence->GetCurrentFence());
 
 	FencedObjectType FencedObject;
 	FencedObject.D3DObject  = pResource;
-	FencedObject.FenceValue = CurrentFrameFence;
+	FencedObject.Fence      = Fence;
+	FencedObject.FenceValue = Fence->GetCurrentFence();
 	FencedObject.Type       = EObjectType::D3D;
 	DeferredReleaseQueue.Enqueue(FencedObject);
 }
@@ -76,24 +75,16 @@ bool FD3D12DeferredDeletionQueue::ReleaseResources(bool DeleteImmediately)
 
 		struct FDequeueFenceObject
 		{
-			FDequeueFenceObject(FD3D12Fence& InFence)
-				: FrameFence(InFence)
-			{
-			}
-
 			bool operator() (FencedObjectType FenceObject) const
 			{
-				return FrameFence.IsFenceComplete(FenceObject.FenceValue);
+				return FenceObject.Fence->IsFenceComplete(FenceObject.FenceValue);
 			}
-
-		private:
-			FD3D12Fence& FrameFence;
 		};
 
 	if (GD3D12AsyncDeferredDeletion)
-		{
+	{
 		if (DeleteImmediately)
-			{
+		{
 			FAsyncTask<FD3D12AsyncDeletionWorker>* DeleteTask = nullptr;
 			// Call back all threads
 			while (DeleteTasks.Peek(DeleteTask))
@@ -103,28 +94,27 @@ bool FD3D12DeferredDeletionQueue::ReleaseResources(bool DeleteImmediately)
 				delete(DeleteTask);
 			}
 		}
-	else
-	{
-		FAsyncTask<FD3D12AsyncDeletionWorker>* DeleteTask = nullptr;
-		while (DeleteTasks.Peek(DeleteTask) && DeleteTask->IsDone())
+		else
 		{
-			DeleteTasks.Dequeue(DeleteTask);
-			delete(DeleteTask);
-		}
+			FAsyncTask<FD3D12AsyncDeletionWorker>* DeleteTask = nullptr;
+			while (DeleteTasks.Peek(DeleteTask) && DeleteTask->IsDone())
+			{
+				DeleteTasks.Dequeue(DeleteTask);
+				delete(DeleteTask);
+			}
 
-		DeleteTask = new FAsyncTask<FD3D12AsyncDeletionWorker>(Adapter, &DeferredReleaseQueue);
+			DeleteTask = new FAsyncTask<FD3D12AsyncDeletionWorker>(Adapter, &DeferredReleaseQueue);
 
-		DeleteTask->StartBackgroundTask();
-		DeleteTasks.Enqueue(DeleteTask);
+			DeleteTask->StartBackgroundTask();
+			DeleteTasks.Enqueue(DeleteTask);
 
-		return false;
-	}
+			return false;
+		}
 	}
 
 	FencedObjectType FenceObject;
-	const FDequeueFenceObject DequeueFenceObject(Adapter->GetFrameFence());
 
-	while (DeferredReleaseQueue.Dequeue(FenceObject, DequeueFenceObject))
+	while (DeferredReleaseQueue.Dequeue(FenceObject, FDequeueFenceObject()))
 	{
 		if (FenceObject.Type == EObjectType::RHI)
 		{
@@ -144,22 +134,13 @@ FD3D12DeferredDeletionQueue::FD3D12AsyncDeletionWorker::FD3D12AsyncDeletionWorke
 {
 	struct FDequeueFenceObject
 	{
-		FDequeueFenceObject(FD3D12Fence& InFence)
-			: FrameFence(InFence)
-		{
-		}
-
 		bool operator() (FencedObjectType FenceObject) const
 		{
-			return FrameFence.IsFenceComplete(FenceObject.FenceValue);
+			return FenceObject.Fence->IsFenceComplete(FenceObject.FenceValue);
 		}
-
-	private:
-		FD3D12Fence& FrameFence;
 	};
 
-	const FDequeueFenceObject DequeueFenceObject(Adapter->GetFrameFence());
-	DeletionQueue->BatchDequeue(&Queue, DequeueFenceObject, 4096);
+	DeletionQueue->BatchDequeue(&Queue, FDequeueFenceObject(), 4096);
 }
 
 void FD3D12DeferredDeletionQueue::FD3D12AsyncDeletionWorker::DoWork()
@@ -267,7 +248,7 @@ void FD3D12Resource::UpdateResidency(FD3D12CommandListHandle& CommandList)
 
 void FD3D12Resource::DeferDelete()
 {
-	GetParentDevice()->GetParentAdapter()->GetDeferredDeletionQueue().EnqueueResource(this);
+	GetParentDevice()->GetParentAdapter()->GetDeferredDeletionQueue().EnqueueResource(this, &GetParentDevice()->GetCommandListManager().GetFence());
 }
 
 /////////////////////////////////////////////////////////////////////
diff --git a/Engine/Source/Runtime/D3D12RHI/Public/D3D12Resources.h b/Engine/Source/Runtime/D3D12RHI/Public/D3D12Resources.h
--- a/Engine/Source/Runtime/D3D12RHI/Public/D3D12Resources.h
+++ b/Engine/Source/Runtime/D3D12RHI/Public/D3D12Resources.h
@@ -590,6 +590,7 @@ class FD3D12DeferredDeletionQueue : public FD3D12AdapterChild
 			FD3D12Resource* RHIObject;
 			ID3D12Object*   D3DObject;
 		};
+		FD3D12Fence* Fence;
 		uint64 FenceValue;
 		EObjectType Type;
 	};
@@ -599,8 +600,8 @@ public:
 
 	inline const uint32 QueueSize() const { return DeferredReleaseQueue.GetSize(); }
 
-	void EnqueueResource(FD3D12Resource* pResource);
-	void EnqueueResource(ID3D12Object* pResource);
+	void EnqueueResource(FD3D12Resource* pResource, FD3D12Fence* Fence);
+	void EnqueueResource(ID3D12Object* pResource, FD3D12Fence* Fence);
 
 	bool ReleaseResources(bool DeleteImmediately = false);
 
diff --git a/Engine/Source/Runtime/D3D12RHI/Public/D3D12Util.h b/Engine/Source/Runtime/D3D12RHI/Public/D3D12Util.h
--- a/Engine/Source/Runtime/D3D12RHI/Public/D3D12Util.h
+++ b/Engine/Source/Runtime/D3D12RHI/Public/D3D12Util.h
@@ -393,7 +393,7 @@ public:
 	}
 
 	template <typename CompareFunc>
-	bool Dequeue(Type& Result, CompareFunc& Func)
+	bool Dequeue(Type& Result, const CompareFunc& Func)
 	{
 		FScopeLock ScopeLock(&SynchronizationObject);
 
@@ -412,7 +412,7 @@ public:
 	}
 
 	template <typename CompareFunc>
-	bool BatchDequeue(TQueue<Type>* Result, CompareFunc& Func, uint32 MaxItems)
+	bool BatchDequeue(TQueue<Type>* Result, const CompareFunc& Func, uint32 MaxItems)
 	{
 		FScopeLock ScopeLock(&SynchronizationObject);
 
