Manually merge CLs 7441030 and 7441039
Integrate D3D12 query refactor CLs from UE4/Main: CLs: 5834472, 5834572, 5958750 (merging with CL 7427737 in this branch)

diff --git a/Engine/Source/Runtime/D3D12RHI/Private/D3D12CommandContext.cpp b/Engine/Source/Runtime/D3D12RHI/Private/D3D12CommandContext.cpp
--- a/Engine/Source/Runtime/D3D12RHI/Private/D3D12CommandContext.cpp
+++ b/Engine/Source/Runtime/D3D12RHI/Private/D3D12CommandContext.cpp
@@ -86,7 +86,7 @@ FD3D12CommandContextBase::FD3D12CommandContextBase(class FD3D12Adapter* InParent
 FD3D12CommandContext::FD3D12CommandContext(FD3D12Device* InParent, FD3D12SubAllocatedOnlineHeap::SubAllocationDesc& SubHeapDesc, bool InIsDefaultContext, bool InIsAsyncComputeContext) :
 	FD3D12CommandContextBase(InParent->GetParentAdapter(), InParent->GetGPUMask(), InIsDefaultContext, InIsAsyncComputeContext),
 	FD3D12DeviceChild(InParent),
-	ConstantsAllocator(InParent, InParent->GetGPUMask(), GetConstantAllocatorSize(InIsAsyncComputeContext, InIsDefaultContext) ),
+	ConstantsAllocator(InParent, InParent->GetGPUMask(), GetConstantAllocatorSize(InIsAsyncComputeContext, InIsDefaultContext)),
 	CommandListHandle(),
 	CommandAllocator(nullptr),
 	CommandAllocatorManager(InParent, InIsAsyncComputeContext ? D3D12_COMMAND_LIST_TYPE_COMPUTE : D3D12_COMMAND_LIST_TYPE_DIRECT),
@@ -97,6 +97,7 @@ FD3D12CommandContext::FD3D12CommandContext(FD3D12Device* InParent, FD3D12SubAllo
 	NumSimultaneousRenderTargets(0),
 	NumUAVs(0),
 	CurrentDSVAccessType(FExclusiveDepthStencil::DepthWrite_StencilWrite),
+	bOuterOcclusionQuerySubmitted(false),
 	bDiscardSharedConstants(false),
 	bUsingTessellation(false),
 	SkipFastClearEliminateState(D3D12_RESOURCE_STATE_PIXEL_SHADER_RESOURCE | D3D12_RESOURCE_STATE_NON_PIXEL_SHADER_RESOURCE),
@@ -414,20 +415,6 @@ void FD3D12CommandContextBase::RHIEndFrame()
 	{
 		FD3D12Device* Device = ParentAdapter->GetDevice(GPUIndex);
 
-		FD3D12CommandContext& DefaultContext = Device->GetDefaultCommandContext();
-		DefaultContext.CommandListHandle.FlushResourceBarriers();
-
-		DefaultContext.ReleaseCommandAllocator();
-		DefaultContext.ClearState();
-		DefaultContext.FlushCommands();
-
-		if (GEnableAsyncCompute)
-		{
-			FD3D12CommandContext& DefaultAsyncComputeContext = Device->GetDefaultAsyncComputeContext();
-			DefaultAsyncComputeContext.ReleaseCommandAllocator();
-			DefaultAsyncComputeContext.ClearState();
-		}
-
 		const uint32 NumContexts = Device->GetNumContexts();
 		for (uint32 i = 0; i < NumContexts; ++i)
 		{
@@ -464,18 +451,8 @@ void FD3D12CommandContextBase::RHIEndFrame()
 		Device->GetCommandListManager().ReleaseResourceBarrierCommandListAllocator();
 	}
 
-	UpdateMemoryStats();
-
-	// Stop Timing at the very last moment
-
-	ParentAdapter->GetGPUProfiler().EndFrame(ParentAdapter->GetOwningRHI());
-
-
-	// Advance frame fence
-
-	FD3D12ManualFence& FrameFence = ParentAdapter->GetFrameFence();
-	FrameFence.Signal(ED3D12CommandQueueType::Default, FrameFence.IncrementCurrentFence());
-}
+		UpdateMemoryStats();
+	}
 
 void FD3D12CommandContextBase::UpdateMemoryStats()
 {
@@ -712,7 +689,86 @@ FD3D12CommandContextRedirector::FD3D12CommandContextRedirector(class FD3D12Adapt
 	FMemory::Memzero(PhysicalContexts, sizeof(PhysicalContexts[0]) * MAX_NUM_GPUS);
 }
 
-void FD3D12CommandContextRedirector::RHITransitionResources(EResourceTransitionAccess TransitionType, EResourceTransitionPipeline TransitionPipeline, FRHIUnorderedAccessView** InUAVs, int32 NumUAVs, FRHIComputeFence* WriteComputeFenceRHI)
+void FD3D12CommandContextRedirector::RHIBeginDrawPrimitiveUP(uint32 NumPrimitives, uint32 NumVertices, uint32 VertexDataStride, void*& OutVertexData)
+{
+	check(!PendingUP);
+
+	PendingUP.NumPrimitives = NumPrimitives;
+	PendingUP.NumVertices = NumVertices;
+	PendingUP.VertexDataStride = VertexDataStride;
+	OutVertexData = PendingUP.VertexData = FMemory::Malloc(NumVertices * VertexDataStride, 16);
+}
+	
+void FD3D12CommandContextRedirector::RHIEndDrawPrimitiveUP()
+{
+	if (PendingUP.VertexData)
+	{
+		for (uint32 GPUIndex : GPUMask)
+		{
+			FD3D12CommandContext* GPUContext = PhysicalContexts[GPUIndex];
+			check(GPUContext);
+
+			void* GPUVertexData = nullptr;
+			GPUContext->RHIBeginDrawPrimitiveUP(PendingUP.NumPrimitives, PendingUP.NumVertices, PendingUP.VertexDataStride, GPUVertexData);
+			if (GPUVertexData)
+			{
+				FMemory::Memcpy(GPUVertexData, PendingUP.VertexData, PendingUP.NumVertices * PendingUP.VertexDataStride);
+			}
+			GPUContext->RHIEndDrawPrimitiveUP();
+		}
+
+		FMemory::Free(PendingUP.VertexData);
+	}
+
+	PendingUP.Reset();
+}
+	
+void FD3D12CommandContextRedirector::RHIBeginDrawIndexedPrimitiveUP(uint32 NumPrimitives, uint32 NumVertices, uint32 VertexDataStride, void*& OutVertexData, uint32 MinVertexIndex, uint32 NumIndices, uint32 IndexDataStride, void*& OutIndexData)
+{
+	check(!PendingUP);
+
+	PendingUP.NumPrimitives = NumPrimitives;
+	PendingUP.NumVertices = NumVertices;
+	PendingUP.VertexDataStride = VertexDataStride;
+	OutVertexData = PendingUP.VertexData = FMemory::Malloc(NumVertices * VertexDataStride);
+
+	PendingUP.MinVertexIndex = MinVertexIndex;
+	PendingUP.NumIndices = NumIndices;
+	PendingUP.IndexDataStride = IndexDataStride;
+	OutIndexData = PendingUP.IndexData = FMemory::Malloc(NumIndices * IndexDataStride);
+}
+	
+void FD3D12CommandContextRedirector::RHIEndDrawIndexedPrimitiveUP()
+{
+	if (PendingUP.VertexData && PendingUP.IndexData)
+	{
+		for (uint32 GPUIndex : GPUMask)
+		{
+			FD3D12CommandContext* GPUContext = PhysicalContexts[GPUIndex];
+			check(GPUContext);
+
+			void* GPUVertexData = nullptr;
+			void* GPUIndexData = nullptr;
+			GPUContext->RHIBeginDrawIndexedPrimitiveUP(PendingUP.NumPrimitives, PendingUP.NumVertices, PendingUP.VertexDataStride, GPUVertexData, PendingUP.MinVertexIndex, PendingUP.NumIndices, PendingUP.IndexDataStride, GPUIndexData);
+			if (GPUVertexData)
+			{
+				FMemory::Memcpy(GPUVertexData, PendingUP.VertexData, PendingUP.NumVertices * PendingUP.VertexDataStride);
+			}
+			if (GPUIndexData)
+			{
+				FMemory::Memcpy(GPUIndexData, PendingUP.IndexData, PendingUP.NumIndices * PendingUP.IndexDataStride);
+			}
+			GPUContext->RHIEndDrawIndexedPrimitiveUP();
+		}
+
+		FMemory::Free(PendingUP.VertexData);
+		FMemory::Free(PendingUP.IndexData);
+	}
+
+	PendingUP.Reset();
+}
+
+void FD3D12CommandContextRedirector::RHITransitionResources(EResourceTransitionAccess TransitionType, EResourceTransitionPipeline TransitionPipeline, FUnorderedAccessViewRHIParamRef* InUAVs, int32 NumUAVs, FComputeFenceRHIParamRef WriteComputeFenceRHI)
 {
 	ContextRedirect(RHITransitionResources(TransitionType, TransitionPipeline, InUAVs, NumUAVs, nullptr));
 
diff --git a/Engine/Source/Runtime/D3D12RHI/Private/D3D12CommandContext.h b/Engine/Source/Runtime/D3D12RHI/Private/D3D12CommandContext.h
--- a/Engine/Source/Runtime/D3D12RHI/Private/D3D12CommandContext.h
+++ b/Engine/Source/Runtime/D3D12RHI/Private/D3D12CommandContext.h
@@ -146,6 +146,10 @@ public:
 	/** Tracks the current depth stencil access type. */
 	FExclusiveDepthStencil CurrentDSVAccessType;
 
+	/** Handle for the dummy outer occlusion query we optionally insert for performance reasons */
+	FRenderQueryRHIRef OuterOcclusionQuery;
+	bool bOuterOcclusionQuerySubmitted;
+
 	/** When a new shader is set, we discard all old constants set for the previous shader. */
 	bool bDiscardSharedConstants;
 
diff --git a/Engine/Source/Runtime/D3D12RHI/Private/D3D12Query.cpp b/Engine/Source/Runtime/D3D12RHI/Private/D3D12Query.cpp
--- a/Engine/Source/Runtime/D3D12RHI/Private/D3D12Query.cpp
+++ b/Engine/Source/Runtime/D3D12RHI/Private/D3D12Query.cpp
@@ -160,10 +160,29 @@ bool FD3D12Device::GetQueryData(FD3D12RenderQuery& Query, bool bWait)
 void FD3D12CommandContext::RHIBeginOcclusionQueryBatch(uint32 NumQueriesInBatch)
 {
 	GetParentDevice()->GetOcclusionQueryHeap()->StartQueryBatch(*this, NumQueriesInBatch);
+	if (RHIConsoleVariables::GInsertOuterOcclusionQuery)
+	{
+		if (!OuterOcclusionQuery.IsValid())
+		{
+			OuterOcclusionQuery = GDynamicRHI->RHICreateRenderQuery(RQT_Occlusion);
+		}
+		
+		FD3D12RenderQuery* OuterOcclusionQueryD3D12 = RetrieveObject<FD3D12RenderQuery>(OuterOcclusionQuery.GetReference());
+		GetParentDevice()->GetOcclusionQueryHeap()->BeginQuery(*this, OuterOcclusionQueryD3D12);
+		bOuterOcclusionQuerySubmitted = true;
+	}
 }
 
 void FD3D12CommandContext::RHIEndOcclusionQueryBatch()
 {
+	if (bOuterOcclusionQuerySubmitted)
+	{
+		check(OuterOcclusionQuery.IsValid());
+		FD3D12RenderQuery* OuterOcclusionQueryD3D12 = RetrieveObject<FD3D12RenderQuery>(OuterOcclusionQuery.GetReference());
+		check(OuterOcclusionQueryD3D12->HeapIndex != INDEX_NONE);
+		GetParentDevice()->GetOcclusionQueryHeap()->EndQuery(*this, OuterOcclusionQueryD3D12);
+		bOuterOcclusionQuerySubmitted = false;
+	}
 	GetParentDevice()->GetOcclusionQueryHeap()->EndQueryBatchAndResolveQueryData(*this);
 
 	// Note: We want to execute this ASAP. The Engine will call RHISubmitCommandHint after this.
@@ -255,19 +274,19 @@ uint32 FD3D12QueryHeap::AllocQuery(FD3D12CommandContext& CmdContext)
 			check(CurrentQueryBatch.bOpen && CurrentQueryBatch.ElementCount == 0);
 		}
 
-		if (CurrentQueryBatch.StartElement > CurrentElement)
-		{
-			// We're in the middle of a batch, but we're at the end of the heap
-			// We need to split the batch in two and resolve the first piece
-			EndQueryBatchAndResolveQueryData(CmdContext);
+	if (CurrentQueryBatch.StartElement > CurrentElement)
+	{
+		// We're in the middle of a batch, but we're at the end of the heap
+		// We need to split the batch in two and resolve the first piece
+		EndQueryBatchAndResolveQueryData(CmdContext);
 		}
 
 		// check for the the batch being closed due to wrap and open a new one
 		if (!CurrentQueryBatch.bOpen)
 		{
 			StartQueryBatch(CmdContext, 256);
-			check(CurrentQueryBatch.bOpen && CurrentQueryBatch.ElementCount == 0);
-		}
+		check(CurrentQueryBatch.bOpen && CurrentQueryBatch.ElementCount == 0);
+	}
 	}
 
 	// Increment the count for the current batch
@@ -344,9 +363,9 @@ void FD3D12QueryHeap::EndQueryBatchAndResolveQueryData(FD3D12CommandContext& Cmd
 	if (CurrentQueryBatch.StartElement + CurrentQueryBatch.ElementCount <= QueryHeapCount)
 	{
 		// Single range
-		CmdContext.CommandListHandle->ResolveQueryData(
-			QueryHeap, QueryType, CurrentQueryBatch.StartElement, CurrentQueryBatch.ElementCount,
-			ResultBuffer->GetResource(), GetResultBufferOffsetForElement(CurrentQueryBatch.StartElement));
+	CmdContext.CommandListHandle->ResolveQueryData(
+		QueryHeap, QueryType, CurrentQueryBatch.StartElement, CurrentQueryBatch.ElementCount,
+		ResultBuffer->GetResource(), GetResultBufferOffsetForElement(CurrentQueryBatch.StartElement));
 	}
 	else
 	{
@@ -390,7 +409,7 @@ void FD3D12QueryHeap::EndQuery(FD3D12CommandContext& CmdContext, FD3D12RenderQue
 
 	if (QueryType == D3D12_QUERY_TYPE_OCCLUSION)
 	{
-		check(CurrentQueryBatch.bOpen);
+	check(CurrentQueryBatch.bOpen);
 	}
 	else
 	{
