Fix a bad merge

diff --git a/Engine/Source/Runtime/D3D12RHI/Private/D3D12CommandContext.cpp b/Engine/Source/Runtime/D3D12RHI/Private/D3D12CommandContext.cpp
--- a/Engine/Source/Runtime/D3D12RHI/Private/D3D12CommandContext.cpp
+++ b/Engine/Source/Runtime/D3D12RHI/Private/D3D12CommandContext.cpp
@@ -415,6 +415,20 @@ void FD3D12CommandContextBase::RHIEndFrame()
 	{
 		FD3D12Device* Device = ParentAdapter->GetDevice(GPUIndex);
 
+		FD3D12CommandContext& DefaultContext = Device->GetDefaultCommandContext();
+		DefaultContext.CommandListHandle.FlushResourceBarriers();
+
+		DefaultContext.ReleaseCommandAllocator();
+		DefaultContext.ClearState();
+		DefaultContext.FlushCommands();
+
+		if (GEnableAsyncCompute)
+		{
+			FD3D12CommandContext& DefaultAsyncComputeContext = Device->GetDefaultAsyncComputeContext();
+			DefaultAsyncComputeContext.ReleaseCommandAllocator();
+			DefaultAsyncComputeContext.ClearState();
+		}
+
 		const uint32 NumContexts = Device->GetNumContexts();
 		for (uint32 i = 0; i < NumContexts; ++i)
 		{
@@ -452,6 +466,16 @@ void FD3D12CommandContextBase::RHIEndFrame()
 	}
 
 		UpdateMemoryStats();
+
+	    // Stop Timing at the very last moment
+    
+	    ParentAdapter->GetGPUProfiler().EndFrame(ParentAdapter->GetOwningRHI());
+    
+    
+	    // Advance frame fence
+    
+	    FD3D12ManualFence& FrameFence = ParentAdapter->GetFrameFence();
+	    FrameFence.Signal(ED3D12CommandQueueType::Default, FrameFence.IncrementCurrentFence());
 	}
 
 void FD3D12CommandContextBase::UpdateMemoryStats()
@@ -689,86 +713,7 @@ FD3D12CommandContextRedirector::FD3D12CommandContextRedirector(class FD3D12Adapt
 	FMemory::Memzero(PhysicalContexts, sizeof(PhysicalContexts[0]) * MAX_NUM_GPUS);
 }
 
-void FD3D12CommandContextRedirector::RHIBeginDrawPrimitiveUP(uint32 NumPrimitives, uint32 NumVertices, uint32 VertexDataStride, void*& OutVertexData)
-{
-	check(!PendingUP);
-
-	PendingUP.NumPrimitives = NumPrimitives;
-	PendingUP.NumVertices = NumVertices;
-	PendingUP.VertexDataStride = VertexDataStride;
-	OutVertexData = PendingUP.VertexData = FMemory::Malloc(NumVertices * VertexDataStride, 16);
-}
-	
-void FD3D12CommandContextRedirector::RHIEndDrawPrimitiveUP()
-{
-	if (PendingUP.VertexData)
-	{
-		for (uint32 GPUIndex : GPUMask)
-		{
-			FD3D12CommandContext* GPUContext = PhysicalContexts[GPUIndex];
-			check(GPUContext);
-
-			void* GPUVertexData = nullptr;
-			GPUContext->RHIBeginDrawPrimitiveUP(PendingUP.NumPrimitives, PendingUP.NumVertices, PendingUP.VertexDataStride, GPUVertexData);
-			if (GPUVertexData)
-			{
-				FMemory::Memcpy(GPUVertexData, PendingUP.VertexData, PendingUP.NumVertices * PendingUP.VertexDataStride);
-			}
-			GPUContext->RHIEndDrawPrimitiveUP();
-		}
-
-		FMemory::Free(PendingUP.VertexData);
-	}
-
-	PendingUP.Reset();
-}
-	
-void FD3D12CommandContextRedirector::RHIBeginDrawIndexedPrimitiveUP(uint32 NumPrimitives, uint32 NumVertices, uint32 VertexDataStride, void*& OutVertexData, uint32 MinVertexIndex, uint32 NumIndices, uint32 IndexDataStride, void*& OutIndexData)
-{
-	check(!PendingUP);
-
-	PendingUP.NumPrimitives = NumPrimitives;
-	PendingUP.NumVertices = NumVertices;
-	PendingUP.VertexDataStride = VertexDataStride;
-	OutVertexData = PendingUP.VertexData = FMemory::Malloc(NumVertices * VertexDataStride);
-
-	PendingUP.MinVertexIndex = MinVertexIndex;
-	PendingUP.NumIndices = NumIndices;
-	PendingUP.IndexDataStride = IndexDataStride;
-	OutIndexData = PendingUP.IndexData = FMemory::Malloc(NumIndices * IndexDataStride);
-}
-	
-void FD3D12CommandContextRedirector::RHIEndDrawIndexedPrimitiveUP()
-{
-	if (PendingUP.VertexData && PendingUP.IndexData)
-	{
-		for (uint32 GPUIndex : GPUMask)
-		{
-			FD3D12CommandContext* GPUContext = PhysicalContexts[GPUIndex];
-			check(GPUContext);
-
-			void* GPUVertexData = nullptr;
-			void* GPUIndexData = nullptr;
-			GPUContext->RHIBeginDrawIndexedPrimitiveUP(PendingUP.NumPrimitives, PendingUP.NumVertices, PendingUP.VertexDataStride, GPUVertexData, PendingUP.MinVertexIndex, PendingUP.NumIndices, PendingUP.IndexDataStride, GPUIndexData);
-			if (GPUVertexData)
-			{
-				FMemory::Memcpy(GPUVertexData, PendingUP.VertexData, PendingUP.NumVertices * PendingUP.VertexDataStride);
-			}
-			if (GPUIndexData)
-			{
-				FMemory::Memcpy(GPUIndexData, PendingUP.IndexData, PendingUP.NumIndices * PendingUP.IndexDataStride);
-			}
-			GPUContext->RHIEndDrawIndexedPrimitiveUP();
-		}
-
-		FMemory::Free(PendingUP.VertexData);
-		FMemory::Free(PendingUP.IndexData);
-	}
-
-	PendingUP.Reset();
-}
-
-void FD3D12CommandContextRedirector::RHITransitionResources(EResourceTransitionAccess TransitionType, EResourceTransitionPipeline TransitionPipeline, FUnorderedAccessViewRHIParamRef* InUAVs, int32 NumUAVs, FComputeFenceRHIParamRef WriteComputeFenceRHI)
+void FD3D12CommandContextRedirector::RHITransitionResources(EResourceTransitionAccess TransitionType, EResourceTransitionPipeline TransitionPipeline, FRHIUnorderedAccessView** InUAVs, int32 NumUAVs, FRHIComputeFence* WriteComputeFenceRHI)
 {
 	ContextRedirect(RHITransitionResources(TransitionType, TransitionPipeline, InUAVs, NumUAVs, nullptr));
 
