Fix bad merge 2

diff --git a/Engine/Source/Runtime/D3D12RHI/Private/D3D12Query.cpp b/Engine/Source/Runtime/D3D12RHI/Private/D3D12Query.cpp
--- a/Engine/Source/Runtime/D3D12RHI/Private/D3D12Query.cpp
+++ b/Engine/Source/Runtime/D3D12RHI/Private/D3D12Query.cpp
@@ -20,6 +20,14 @@ namespace D3D12RHI
 			TEXT("If true, enable stable power state. This increases GPU timing measurement accuracy but may decrease overall GPU clock rate."),
 			ECVF_Default
 			);
+
+		int32 GInsertOuterOcclusionQuery = 0;
+		static FAutoConsoleVariableRef CVarInsertOuterOcclusionQuery(
+			TEXT("D3D12.InsertOuterOcclusionQuery"),
+			GInsertOuterOcclusionQuery,
+			TEXT("If true, enable a dummy outer occlusion query around occlusion query batches. Can help performance on some GPU architectures"),
+			ECVF_Default
+		);
 	}
 }
 using namespace D3D12RHI;
