Added missing DX12 RHI for ReadSurfaceData with FLinearColor Output

diff --git a/Engine/Source/Runtime/D3D12RHI/Private/D3D12RHIPrivate.h b/Engine/Source/Runtime/D3D12RHI/Private/D3D12RHIPrivate.h
--- a/Engine/Source/Runtime/D3D12RHI/Private/D3D12RHIPrivate.h
+++ b/Engine/Source/Runtime/D3D12RHI/Private/D3D12RHIPrivate.h
@@ -337,6 +337,7 @@ public:
 	virtual void RHIUnlockTextureCubeFace(FRHITextureCube* Texture, uint32 FaceIndex, uint32 ArrayIndex, uint32 MipIndex, bool bLockWithinMiptail) final override;
 	virtual void RHIBindDebugLabelName(FRHITexture* Texture, const TCHAR* Name) final override;
 	virtual void RHIReadSurfaceData(FRHITexture* Texture, FIntRect Rect, TArray<FColor>& OutData, FReadSurfaceDataFlags InFlags) final override;
+	virtual void RHIReadSurfaceData(FRHITexture* TextureRHI, FIntRect InRect, TArray<FLinearColor>& OutData, FReadSurfaceDataFlags InFlags) final override;
 	virtual void RHIMapStagingSurface(FRHITexture* Texture, void*& OutData, int32& OutWidth, int32& OutHeight) final override;
 	virtual void RHIUnmapStagingSurface(FRHITexture* Texture) final override;
 	virtual void RHIReadSurfaceFloatData(FRHITexture* Texture, FIntRect Rect, TArray<FFloat16Color>& OutData, ECubeFace CubeFace, int32 ArrayIndex, int32 MipIndex) final override;
diff --git a/Engine/Source/Runtime/D3D12RHI/Private/D3D12RenderTarget.cpp b/Engine/Source/Runtime/D3D12RHI/Private/D3D12RenderTarget.cpp
--- a/Engine/Source/Runtime/D3D12RHI/Private/D3D12RenderTarget.cpp
+++ b/Engine/Source/Runtime/D3D12RHI/Private/D3D12RenderTarget.cpp
@@ -1109,6 +1109,334 @@ static void ConvertRAWSurfaceDataToFColor(DXGI_FORMAT Format, uint32 Width, uint
 	}
 }
 
+
+
+static void ConvertRAWSurfaceDataToFLinearColor(EPixelFormat Format, uint32 Width, uint32 Height, uint8 *In, uint32 SrcPitch, FLinearColor* Out, FReadSurfaceDataFlags InFlags)
+{
+	if (Format == PF_R16F || Format == PF_R16F_FILTER)
+	{
+		// e.g. shadow maps
+		for (uint32 Y = 0; Y < Height; Y++)
+		{
+			uint16* SrcPtr = (uint16*)(In + Y * SrcPitch);
+			FLinearColor* DestPtr = Out + Y * Width;
+
+			for (uint32 X = 0; X < Width; X++)
+			{
+				uint16 Value16 = *SrcPtr;
+				float Value = Value16 / (float)(0xffff);
+
+				*DestPtr = FLinearColor(Value, Value, Value);
+				++SrcPtr;
+				++DestPtr;
+			}
+		}
+	}
+	else if (Format == PF_R8G8B8A8)
+	{
+		// Read the data out of the buffer, converting it from ABGR to ARGB.
+		for (uint32 Y = 0; Y < Height; Y++)
+		{
+			FColor* SrcPtr = (FColor*)(In + Y * SrcPitch);
+			FLinearColor* DestPtr = Out + Y * Width;
+			for (uint32 X = 0; X < Width; X++)
+			{
+				FColor sRGBColor = FColor(SrcPtr->B, SrcPtr->G, SrcPtr->R, SrcPtr->A);
+				*DestPtr = FLinearColor(sRGBColor);
+				++SrcPtr;
+				++DestPtr;
+			}
+		}
+	}
+	else if (Format == PF_B8G8R8A8)
+	{
+		for (uint32 Y = 0; Y < Height; Y++)
+		{
+			FColor* SrcPtr = (FColor*)(In + Y * SrcPitch);
+			FLinearColor* DestPtr = Out + Y * Width;
+			for (uint32 X = 0; X < Width; X++)
+			{
+				FColor sRGBColor = FColor(SrcPtr->R, SrcPtr->G, SrcPtr->B, SrcPtr->A);
+				*DestPtr = FLinearColor(sRGBColor);
+				++SrcPtr;
+				++DestPtr;
+			}
+		}
+	}
+	else if (Format == PF_A2B10G10R10)
+	{
+		// Read the data out of the buffer, converting it from R10G10B10A2 to FLinearColor.
+		for (uint32 Y = 0; Y < Height; Y++)
+		{
+			FD3DR10G10B10A2* SrcPtr = (FD3DR10G10B10A2*)(In + Y * SrcPitch);
+			FLinearColor* DestPtr = Out + Y * Width;
+			for (uint32 X = 0; X < Width; X++)
+			{
+				*DestPtr = FLinearColor(
+					(float)SrcPtr->R / 1023.0f,
+					(float)SrcPtr->G / 1023.0f,
+					(float)SrcPtr->B / 1023.0f,
+					(float)SrcPtr->A / 3.0f
+				);
+				++SrcPtr;
+				++DestPtr;
+			}
+		}
+	}
+	else if (Format == PF_FloatRGBA)
+	{
+		if (InFlags.GetCompressionMode() == RCM_MinMax)
+		{
+			for (uint32 Y = 0; Y < Height; Y++)
+			{
+				FFloat16* SrcPtr = (FFloat16*)(In + Y * SrcPitch);
+				FLinearColor* DestPtr = Out + Y * Width;
+
+				for (uint32 X = 0; X < Width; X++)
+				{
+					*DestPtr = FLinearColor((float)SrcPtr[0], (float)SrcPtr[1], (float)SrcPtr[2], (float)SrcPtr[3]);
+					++DestPtr;
+					SrcPtr += 4;
+				}
+			}
+		}
+		else
+		{
+			FPlane	MinValue(0.0f, 0.0f, 0.0f, 0.0f);
+			FPlane	MaxValue(1.0f, 1.0f, 1.0f, 1.0f);
+
+			check(sizeof(FFloat16) == sizeof(uint16));
+
+			for (uint32 Y = 0; Y < Height; Y++)
+			{
+				FFloat16* SrcPtr = (FFloat16*)(In + Y * SrcPitch);
+
+				for (uint32 X = 0; X < Width; X++)
+				{
+					MinValue.X = FMath::Min<float>(SrcPtr[0], MinValue.X);
+					MinValue.Y = FMath::Min<float>(SrcPtr[1], MinValue.Y);
+					MinValue.Z = FMath::Min<float>(SrcPtr[2], MinValue.Z);
+					MinValue.W = FMath::Min<float>(SrcPtr[3], MinValue.W);
+					MaxValue.X = FMath::Max<float>(SrcPtr[0], MaxValue.X);
+					MaxValue.Y = FMath::Max<float>(SrcPtr[1], MaxValue.Y);
+					MaxValue.Z = FMath::Max<float>(SrcPtr[2], MaxValue.Z);
+					MaxValue.W = FMath::Max<float>(SrcPtr[3], MaxValue.W);
+					SrcPtr += 4;
+				}
+			}
+
+			for (uint32 Y = 0; Y < Height; Y++)
+			{
+				FFloat16* SrcPtr = (FFloat16*)(In + Y * SrcPitch);
+				FLinearColor* DestPtr = Out + Y * Width;
+
+				for (uint32 X = 0; X < Width; X++)
+				{
+					*DestPtr = FLinearColor(
+						(SrcPtr[0] - MinValue.X) / (MaxValue.X - MinValue.X),
+						(SrcPtr[1] - MinValue.Y) / (MaxValue.Y - MinValue.Y),
+						(SrcPtr[2] - MinValue.Z) / (MaxValue.Z - MinValue.Z),
+						(SrcPtr[3] - MinValue.W) / (MaxValue.W - MinValue.W)
+					);
+					++DestPtr;
+					SrcPtr += 4;
+				}
+			}
+		}
+	}
+	else if (Format == PF_FloatRGB || Format == PF_FloatR11G11B10)
+	{
+		check(sizeof(FFloat3Packed) == sizeof(uint32));
+
+		for (uint32 Y = 0; Y < Height; Y++)
+		{
+			FFloat3Packed* SrcPtr = (FFloat3Packed*)(In + Y * SrcPitch);
+			FLinearColor* DestPtr = Out + Y * Width;
+
+			for (uint32 X = 0; X < Width; X++)
+			{
+				*DestPtr = (*SrcPtr).ToLinearColor();
+				++DestPtr;
+				++SrcPtr;
+			}
+		}
+	}
+	else if (Format == PF_A32B32G32R32F)
+	{
+		if (InFlags.GetCompressionMode() == RCM_MinMax)
+		{
+			// Copy data directly, respecting existing min-max values
+			FLinearColor* SrcPtr = (FLinearColor*)In;
+			FLinearColor* DestPtr = (FLinearColor*)Out;
+			const int32 ImageSize = sizeof(FLinearColor) * Height * Width;
+
+			FMemory::Memcpy(DestPtr, SrcPtr, ImageSize);
+		}
+		else
+		{
+			// Normalize data
+			FPlane MinValue(0.0f, 0.0f, 0.0f, 0.0f);
+			FPlane MaxValue(1.0f, 1.0f, 1.0f, 1.0f);
+
+			for (uint32 Y = 0; Y < Height; Y++)
+			{
+				float* SrcPtr = (float*)(In + Y * SrcPitch);
+
+				for (uint32 X = 0; X < Width; X++)
+				{
+					MinValue.X = FMath::Min<float>(SrcPtr[0], MinValue.X);
+					MinValue.Y = FMath::Min<float>(SrcPtr[1], MinValue.Y);
+					MinValue.Z = FMath::Min<float>(SrcPtr[2], MinValue.Z);
+					MinValue.W = FMath::Min<float>(SrcPtr[3], MinValue.W);
+					MaxValue.X = FMath::Max<float>(SrcPtr[0], MaxValue.X);
+					MaxValue.Y = FMath::Max<float>(SrcPtr[1], MaxValue.Y);
+					MaxValue.Z = FMath::Max<float>(SrcPtr[2], MaxValue.Z);
+					MaxValue.W = FMath::Max<float>(SrcPtr[3], MaxValue.W);
+					SrcPtr += 4;
+				}
+			}
+
+			float* SrcPtr = (float*)In;
+
+			for (uint32 Y = 0; Y < Height; Y++)
+			{
+				FLinearColor* DestPtr = Out + Y * Width;
+
+				for (uint32 X = 0; X < Width; X++)
+				{
+					*DestPtr = FLinearColor(
+						(SrcPtr[0] - MinValue.X) / (MaxValue.X - MinValue.X),
+						(SrcPtr[1] - MinValue.Y) / (MaxValue.Y - MinValue.Y),
+						(SrcPtr[2] - MinValue.Z) / (MaxValue.Z - MinValue.Z),
+						(SrcPtr[3] - MinValue.W) / (MaxValue.W - MinValue.W)
+					);
+					++DestPtr;
+					SrcPtr += 4;
+				}
+			}
+		}
+	}
+	else if (Format == PF_DepthStencil || Format == PF_D24)
+	{
+		// Depth stencil
+		for (uint32 Y = 0; Y < Height; Y++)
+		{
+			uint32* SrcPtr = (uint32 *)(In + Y * SrcPitch);
+			FLinearColor* DestPtr = Out + Y * Width;
+
+			for (uint32 X = 0; X < Width; X++)
+			{
+				float DeviceStencil = 0.0f;
+				DeviceStencil = (float)((*SrcPtr & 0xFF000000) >> 24) / 255.0f;
+				float DeviceZ = (*SrcPtr & 0xffffff) / (float)(1 << 24);
+				float LinearValue = FMath::Min(InFlags.ComputeNormalizedDepth(DeviceZ), 1.0f);
+				*DestPtr = FLinearColor(LinearValue, DeviceStencil, 0.0f, 0.0f);
+				++DestPtr;
+				++SrcPtr;
+			}
+		}
+	}
+	// Changing Depth Buffers to 32 bit on Dingo as D24S8 is actually implemented as a 32 bit buffer in the hardware
+	else if (Format == PF_DepthStencil)
+	{
+		// Depth stencil
+		for (uint32 Y = 0; Y < Height; Y++)
+		{
+			uint8* SrcStart = (uint8 *)(In + Y * SrcPitch);
+			FLinearColor* DestPtr = Out + Y * Width;
+
+			for (uint32 X = 0; X < Width; X++)
+			{
+				float DeviceZ = *((float *)(SrcStart));
+				float LinearValue = FMath::Min(InFlags.ComputeNormalizedDepth(DeviceZ), 1.0f);
+				float DeviceStencil = (float)(*(SrcStart + 4)) / 255.0f;
+				*DestPtr = FLinearColor(LinearValue, DeviceStencil, 0.0f, 0.0f);
+				SrcStart += 8; //64 bit format with the last 24 bit ignore
+			}
+		}
+	}
+	else if (Format == PF_A16B16G16R16)
+	{
+		// Read the data out of the buffer, converting it to FLinearColor.
+		for (uint32 Y = 0; Y < Height; Y++)
+		{
+			FD3DRGBA16* SrcPtr = (FD3DRGBA16*)(In + Y * SrcPitch);
+			FLinearColor* DestPtr = Out + Y * Width;
+			for (uint32 X = 0; X < Width; X++)
+			{
+				*DestPtr = FLinearColor(
+					(float)SrcPtr->R / 65535.0f,
+					(float)SrcPtr->G / 65535.0f,
+					(float)SrcPtr->B / 65535.0f,
+					(float)SrcPtr->A / 65535.0f
+				);
+				++SrcPtr;
+				++DestPtr;
+			}
+		}
+	}
+	else if (Format == PF_G16R16)
+	{
+		// Read the data out of the buffer, converting it to FLinearColor.
+		for (uint32 Y = 0; Y < Height; Y++)
+		{
+			FD3DRG16* SrcPtr = (FD3DRG16*)(In + Y * SrcPitch);
+			FLinearColor* DestPtr = Out + Y * Width;
+			for (uint32 X = 0; X < Width; X++)
+			{
+				*DestPtr = FLinearColor(
+					(float)SrcPtr->R / 65535.0f,
+					(float)SrcPtr->G / 65535.0f,
+					0);
+				++SrcPtr;
+				++DestPtr;
+			}
+		}
+	}
+	else
+	{
+		// not supported yet
+		check(0);
+	}
+}
+
+
+void FD3D12DynamicRHI::RHIReadSurfaceData(FRHITexture* TextureRHI, FIntRect InRect, TArray<FLinearColor>& OutData, FReadSurfaceDataFlags InFlags)
+{
+	TArray<uint8> OutDataRaw;
+
+	FD3D12TextureBase* Texture = GetD3D12TextureFromRHITexture(TextureRHI);
+
+	// Check the format of the surface
+	D3D12_RESOURCE_DESC const& TextureDesc = Texture->GetResource()->GetDesc();
+
+	check(TextureDesc.SampleDesc.Count >= 1);
+
+	if (TextureDesc.SampleDesc.Count == 1)
+	{
+		ReadSurfaceDataNoMSAARaw(TextureRHI, InRect, OutDataRaw, InFlags);
+	}
+	else
+	{
+		FRHICommandList_RecursiveHazardous RHICmdList(RHIGetDefaultContext());
+		ReadSurfaceDataMSAARaw(RHICmdList, TextureRHI, InRect, OutDataRaw, InFlags);
+	}
+
+	const uint32 SizeX = InRect.Width() * TextureDesc.SampleDesc.Count;
+	const uint32 SizeY = InRect.Height();
+
+	// Allocate the output buffer.
+	OutData.Empty();
+	OutData.AddUninitialized(SizeX * SizeY);
+
+	uint32 BytesPerPixel = ComputeBytesPerPixel(TextureDesc.Format);
+	uint32 SrcPitch = SizeX * BytesPerPixel;
+	EPixelFormat Format = TextureRHI->GetFormat();
+	if (Format != PF_Unknown)
+	{
+		ConvertRAWSurfaceDataToFLinearColor(Format, SizeX, SizeY, OutDataRaw.GetData(), SrcPitch, OutData.GetData(), InFlags);
+	}
+}
 void FD3D12DynamicRHI::RHIReadSurfaceData(FRHITexture* TextureRHI, FIntRect InRect, TArray<FColor>& OutData, FReadSurfaceDataFlags InFlags)
 {
 	if (!ensure(TextureRHI))
