Use a separate TSet of plain pointers to keep track of referenced resources in FD3D12RayTracingShaderTable. This fixes a CPU performance regression introduced in CL 7413442.

Fixed resource reference counting when D3D12 resource residency tracking is disabled.

(ushell-p4-cherrypick of 7475878 by [REDACTED_PERSON])

diff --git a/Engine/Source/Runtime/D3D12RHI/Private/D3D12RayTracing.cpp b/Engine/Source/Runtime/D3D12RHI/Private/D3D12RayTracing.cpp
--- a/Engine/Source/Runtime/D3D12RHI/Private/D3D12RayTracing.cpp
+++ b/Engine/Source/Runtime/D3D12RHI/Private/D3D12RayTracing.cpp
@@ -1360,15 +1360,20 @@ public:
 
 	// A set of all resources referenced by this shader table for the purpose of updating residency before ray tracing work dispatch.
 	// #dxr_todo UE-72159: remove resources from this set when SBT slot entries are replaced
-	TSet<TRefCountPtr<FD3D12Resource>> ReferencedD3D12Resources;
+	TSet<void*> ReferencedD3D12ResourceSet;
+	TArray<TRefCountPtr<FD3D12Resource>> ReferencedD3D12Resources;
 	TArray<TRefCountPtr<FRHIResource>> ReferencedResources;
 	void AddResourceReference(FD3D12Resource* D3D12Resource, FRHIResource* Resource)
 	{
 		bool bIsAlreadyInSet = false;
-		ReferencedD3D12Resources.Add(D3D12Resource, &bIsAlreadyInSet);
-		if (!bIsAlreadyInSet && Resource)
+		ReferencedD3D12ResourceSet.Add(D3D12Resource, &bIsAlreadyInSet);
+		if (!bIsAlreadyInSet)
 		{
-			ReferencedResources.Add(Resource);
+			ReferencedD3D12Resources.Add(D3D12Resource);
+			if (Resource)
+			{
+				ReferencedResources.Add(Resource);
+			}
 		}
 	}
 	void UpdateResidency(FD3D12CommandContext& CommandContext)
@@ -2991,12 +2996,9 @@ static void SetRayTracingShaderResources(
 		Binder.SetRootDescriptorTable(BindSlot, ResourceDescriptorTableBaseGPU);
 	}
 
-	if (GEnableResidencyManagement)
+	for (const FResourceEntry& Entry : ReferencedResources)
 	{
-		for (const FResourceEntry& Entry : ReferencedResources)
-		{
-			Binder.AddResourceReference(Entry.D3D12Resource, Entry.RHIResource);
-		}
+		Binder.AddResourceReference(Entry.D3D12Resource, Entry.RHIResource);
 	}
 }
 
