Some locations still need further review, but changing the comments from ATG_CHANGE to a more descriptive comment.

diff --git a/Engine/Source/Runtime/Core/Private/HoloLens/HoloLensProcess.cpp b/Engine/Source/Runtime/Core/Private/HoloLens/HoloLensProcess.cpp
--- a/Engine/Source/Runtime/Core/Private/HoloLens/HoloLensProcess.cpp
+++ b/Engine/Source/Runtime/Core/Private/HoloLens/HoloLensProcess.cpp
@@ -26,7 +26,8 @@ const TCHAR* FHoloLensProcess::BaseDir()
 		{
 			CmdLine = FCommandLine::Get();
 		}
-		// @ATG_CHANGE :  - code path that was hitting BaseDir in static initializers was removed
+		// Change from Microsoft, may require further review:
+		// code path that was hitting BaseDir in static initializers was removed
 		//                  and command line parsing was moved to earlier in launch to help avoid having 
 		//                  to have command line parsing code in multiple location.
 		//                  There's still potential for error here until Epic allows platform-subclassing the
diff --git a/Engine/Source/Runtime/D3D12RHI/Private/Windows/WindowsD3D12Viewport.cpp b/Engine/Source/Runtime/D3D12RHI/Private/Windows/WindowsD3D12Viewport.cpp
--- a/Engine/Source/Runtime/D3D12RHI/Private/Windows/WindowsD3D12Viewport.cpp
+++ b/Engine/Source/Runtime/D3D12RHI/Private/Windows/WindowsD3D12Viewport.cpp
@@ -87,7 +87,7 @@ void FD3D12Viewport::Init()
 	const DXGI_MODE_DESC BufferDesc = SetupDXGI_MODE_DESC();
 
 	// Create the swapchain.
-// @ATG_CHANGE : BEGIN HoloLens support
+// Change from Microsoft for HoloLens support, may require further review: : BEGIN HoloLens support
 #if PLATFORM_HOLOLENS
 	{
 		// MSAA Sample count
@@ -121,7 +121,7 @@ void FD3D12Viewport::Init()
 		));
 	}
 #else
-// @ATG_CHANGE : END
+// : END HoloLens support
 	{
 		if (Adapter->GetOwningRHI()->IsQuadBufferStereoEnabled())
 		{
@@ -200,9 +200,9 @@ void FD3D12Viewport::Init()
 			}
 		}
 	}
-// @ATG_CHANGE : BEGIN HoloLens support
+// Change from Microsoft for HoloLens support, may require further review: : BEGIN HoloLens support
 #endif
-// @ATG_CHANGE : END
+// : END HoloLens support
 
 	// Set the DXGI message hook to not change the window behind our back.
 	Adapter->GetDXGIFactory2()->MakeWindowAssociation(WindowHandle, DXGI_MWA_NO_WINDOW_CHANGES);
@@ -210,13 +210,13 @@ void FD3D12Viewport::Init()
 	// Resize to setup mGPU correctly.
 	Resize(BufferDesc.Width, BufferDesc.Height, bIsFullscreen, PixelFormat);
 
-// @ATG_CHANGE : BEGIN HoloLens support
+// Change from Microsoft for HoloLens support, may require further review: : BEGIN HoloLens support
 #if !PLATFORM_HOLOLENS
 	// Tell the window to redraw when they can.
 	// @todo: For Slate viewports, it doesn't make sense to post WM_PAINT messages (we swallow those.)
 	::PostMessage(WindowHandle, WM_PAINT, 0, 0);
 #endif
-// @ATG_CHANGE : END
+	// : END HoloLens support
 }
 
 void FD3D12Viewport::ConditionalResetSwapChain(bool bIgnoreFocus)
