Cherry-pick 8096554: Increase default size of D3D12 transient uniform buffer allocator from 2 MB to 10 MB to fix UE-78143

diff --git a/Engine/Source/Runtime/D3D12RHI/Private/D3D12Adapter.cpp b/Engine/Source/Runtime/D3D12RHI/Private/D3D12Adapter.cpp
--- a/Engine/Source/Runtime/D3D12RHI/Private/D3D12Adapter.cpp
+++ b/Engine/Source/Runtime/D3D12RHI/Private/D3D12Adapter.cpp
@@ -8,7 +8,7 @@ D3D12Adapter.cpp:D3D12 Adapter implementation.
 
 static TAutoConsoleVariable<int32> CVarTransientUniformBufferAllocatorSizeKB(
 	TEXT("D3D12.TransientUniformBufferAllocatorSizeKB"),
-	2 * 1024,
+	10 * 1024,
 	TEXT(""),
 	ECVF_ReadOnly
 );
