DR - Win 10 required for D3D12

diff --git a/Engine/Source/Runtime/D3D12RHI/Private/Windows/WindowsD3D12Device.cpp b/Engine/Source/Runtime/D3D12RHI/Private/Windows/WindowsD3D12Device.cpp
--- a/Engine/Source/Runtime/D3D12RHI/Private/Windows/WindowsD3D12Device.cpp
+++ b/Engine/Source/Runtime/D3D12RHI/Private/Windows/WindowsD3D12Device.cpp
@@ -211,6 +211,11 @@ static bool SupportsHDROutput(FD3D12DynamicRHI* D3DRHI)
 
 bool FD3D12DynamicRHIModule::IsSupported()
 {
+	if (!FWindowsPlatformMisc::VerifyWindowsVersion(10, 0))
+	{
+		return false;
+	}
+
 	// If not computed yet
 	if (ChosenAdapters.Num() == 0)
 	{
