Add RLM_WriteOnly_NoOverwrite to dx12.

diff --git a/Engine/Source/Runtime/D3D12RHI/Private/D3D12Buffer.cpp b/Engine/Source/Runtime/D3D12RHI/Private/D3D12Buffer.cpp
--- a/Engine/Source/Runtime/D3D12RHI/Private/D3D12Buffer.cpp
+++ b/Engine/Source/Runtime/D3D12RHI/Private/D3D12Buffer.cpp
@@ -291,7 +291,7 @@ void* FD3D12DynamicRHI::LockBuffer(FRHICommandListImmediate* RHICmdList, BufferT
 
 	if (bIsDynamic)
 	{
-		check(LockMode == RLM_WriteOnly);
+		check(LockMode == RLM_WriteOnly || LockMode == RLM_WriteOnly_NoOverwrite);
 
 		if (LockedData.bHasNeverBeenLocked)
 		{
@@ -304,7 +304,7 @@ void* FD3D12DynamicRHI::LockBuffer(FRHICommandListImmediate* RHICmdList, BufferT
 			FD3D12Device* Device = Buffer->GetParentDevice();
 
 			// If on the RenderThread, queue up a command on the RHIThread to rename this buffer at the correct time
-			if (ShouldDeferBufferLockOperation(RHICmdList))
+			if (ShouldDeferBufferLockOperation(RHICmdList) && LockMode == RLM_WriteOnly)
 			{
 				FRHICommandRenameUploadBuffer<BufferType>* Command = ALLOC_COMMAND_CL(*RHICmdList, FRHICommandRenameUploadBuffer<BufferType>)(Buffer, Device);
 
