Fixed crash in mGPU mode due to missing FD3D12CommandContextRedirector::RHIClearRayTracingBindings implementation

diff --git a/Engine/Source/Runtime/D3D12RHI/Private/D3D12CommandContext.h b/Engine/Source/Runtime/D3D12RHI/Private/D3D12CommandContext.h
--- a/Engine/Source/Runtime/D3D12RHI/Private/D3D12CommandContext.h
+++ b/Engine/Source/Runtime/D3D12RHI/Private/D3D12CommandContext.h
@@ -823,6 +823,10 @@ public:
 		ContextRedirect(RHISetRayTracingCallableShader(Scene, ShaderSlotInScene, Pipeline, ShaderIndexInPipeline, NumUniformBuffers, UniformBuffers, UserData));
 	}
 
+	virtual void RHIClearRayTracingBindings(FRHIRayTracingScene* Scene) final override
+	{
+		ContextRedirect(RHIClearRayTracingBindings(Scene));
+	}
 
 	FORCEINLINE void SetPhysicalContext(FD3D12CommandContext* Context)
 	{
