Help track down D3D12 uninitialized view

diff --git a/Engine/Source/Runtime/D3D12RHI/Private/D3D12View.h b/Engine/Source/Runtime/D3D12RHI/Private/D3D12View.h
--- a/Engine/Source/Runtime/D3D12RHI/Private/D3D12View.h
+++ b/Engine/Source/Runtime/D3D12RHI/Private/D3D12View.h
@@ -791,13 +791,13 @@ protected:
 public:
 	inline FD3D12Device*					GetParentDevice()			const { return Descriptor.GetParentDevice(); }
 	inline FD3D12Device*					GetParentDevice_Unsafe()	const { return Descriptor.GetParentDevice_Unsafe(); }
-	inline const TDesc&						GetDesc()					const { check(bInitialized); return Desc; }
-	inline CD3DX12_CPU_DESCRIPTOR_HANDLE	GetView()					const { check(bInitialized); return Descriptor.GetHandle(); }
-	inline uint32							GetDescriptorHeapIndex()	const { check(bInitialized); return Descriptor.GetIndex(); }
-	inline FD3D12Resource*					GetResource()				const { check(bInitialized); return ResourceLocation->GetResource(); }
-	inline FD3D12ResourceLocation*			GetResourceLocation()		const { check(bInitialized); return ResourceLocation; }
-	inline FD3D12ResidencyHandle*			GetResidencyHandle()		const { check(bInitialized); return ResidencyHandle; }
-	inline const CViewSubresourceSubset&	GetViewSubresourceSubset()	const { check(bInitialized); return ViewSubresourceSubset; }
+	inline const TDesc&						GetDesc()					const { checkf(bInitialized, TEXT("Uninitialized D3D12View size %d"), (uint32)sizeof(TDesc)); return Desc; }
+	inline CD3DX12_CPU_DESCRIPTOR_HANDLE	GetView()					const { checkf(bInitialized, TEXT("Uninitialized D3D12View size %d"), (uint32)sizeof(TDesc)); return Descriptor.GetHandle(); }
+	inline uint32							GetDescriptorHeapIndex()	const { checkf(bInitialized, TEXT("Uninitialized D3D12View size %d"), (uint32)sizeof(TDesc)); return Descriptor.GetIndex(); }
+	inline FD3D12Resource*					GetResource()				const { checkf(bInitialized, TEXT("Uninitialized D3D12View size %d"), (uint32)sizeof(TDesc)); return ResourceLocation->GetResource(); }
+	inline FD3D12ResourceLocation*			GetResourceLocation()		const { checkf(bInitialized, TEXT("Uninitialized D3D12View size %d"), (uint32)sizeof(TDesc)); return ResourceLocation; }
+	inline FD3D12ResidencyHandle*			GetResidencyHandle()		const { checkf(bInitialized, TEXT("Uninitialized D3D12View size %d"), (uint32)sizeof(TDesc)); return ResidencyHandle; }
+	inline const CViewSubresourceSubset&	GetViewSubresourceSubset()	const { checkf(bInitialized, TEXT("Uninitialized D3D12View size %d"), (uint32)sizeof(TDesc)); return ViewSubresourceSubset; }
 
 	void SetParentDevice(FD3D12Device* InParent)
 	{
