Fix some non-unity and standards compliance errors running with /permissive- on MSVC.

diff --git a/Engine/Plugins/Experimental/DataPrepEditor/Source/DataPrepEditor/Private/SchemaActions/DataprepSchemaActionUtils.h b/Engine/Plugins/Experimental/DataPrepEditor/Source/DataPrepEditor/Private/SchemaActions/DataprepSchemaActionUtils.h
--- a/Engine/Plugins/Experimental/DataPrepEditor/Source/DataPrepEditor/Private/SchemaActions/DataprepSchemaActionUtils.h
+++ b/Engine/Plugins/Experimental/DataPrepEditor/Source/DataPrepEditor/Private/SchemaActions/DataprepSchemaActionUtils.h
@@ -2,6 +2,8 @@
 
 #pragma once
 
+#include "EdGraph/EdGraph.h"
+
 namespace DataprepSchemaActionUtils
 {
 	template<class GraphNodeClass>
diff --git a/Engine/Plugins/Experimental/ProxyLODPlugin/Source/ProxyLOD/Private/ProxyLODMeshSDFConversions.h b/Engine/Plugins/Experimental/ProxyLODPlugin/Source/ProxyLOD/Private/ProxyLODMeshSDFConversions.h
--- a/Engine/Plugins/Experimental/ProxyLODPlugin/Source/ProxyLOD/Private/ProxyLODMeshSDFConversions.h
+++ b/Engine/Plugins/Experimental/ProxyLODPlugin/Source/ProxyLOD/Private/ProxyLODMeshSDFConversions.h
@@ -5,6 +5,7 @@
 
 #include "ProxyLODMeshTypes.h"
 #include "ProxyLODMeshConvertUtils.h"
+#include "ProxyLODMeshUtilities.h"
 
 THIRD_PARTY_INCLUDES_START
 #include <openvdb/openvdb.h>
diff --git a/Engine/Plugins/Runtime/GoogleVR/GoogleVRController/Source/GoogleVRController/Private/GoogleVRControllerPrivate.h b/Engine/Plugins/Runtime/GoogleVR/GoogleVRController/Source/GoogleVRController/Private/GoogleVRControllerPrivate.h
--- a/Engine/Plugins/Runtime/GoogleVR/GoogleVRController/Source/GoogleVRController/Private/GoogleVRControllerPrivate.h
+++ b/Engine/Plugins/Runtime/GoogleVR/GoogleVRController/Source/GoogleVRController/Private/GoogleVRControllerPrivate.h
@@ -1,6 +1,7 @@
 // Copyright 1998-2019 Epic Games, Inc. All Rights Reserved.
 
 #include "CoreMinimal.h"
+#include "Templates/ValueOrError.h"
 #include "IGoogleVRControllerPlugin.h"
 
 #if GOOGLEVRCONTROLLER_SUPPORTED_PLATFORMS
diff --git a/Engine/Plugins/Runtime/nDisplay/Source/MPCDI/Private/MPCDIShader.h b/Engine/Plugins/Runtime/nDisplay/Source/MPCDI/Private/MPCDIShader.h
--- a/Engine/Plugins/Runtime/nDisplay/Source/MPCDI/Private/MPCDIShader.h
+++ b/Engine/Plugins/Runtime/nDisplay/Source/MPCDI/Private/MPCDIShader.h
@@ -5,6 +5,7 @@
 #include "CoreMinimal.h"
 #include "RHI.h"
 #include "RHICommandList.h"
+#include "RHIStaticStates.h"
 #include "RenderResource.h"
 #include "ShaderParameters.h"
 #include "Shader.h"
diff --git a/Engine/Plugins/Runtime/nDisplay/Source/PicpMPCDI/Private/PicpMPCDIShader.h b/Engine/Plugins/Runtime/nDisplay/Source/PicpMPCDI/Private/PicpMPCDIShader.h
--- a/Engine/Plugins/Runtime/nDisplay/Source/PicpMPCDI/Private/PicpMPCDIShader.h
+++ b/Engine/Plugins/Runtime/nDisplay/Source/PicpMPCDI/Private/PicpMPCDIShader.h
@@ -5,6 +5,7 @@
 #include "CoreMinimal.h"
 #include "RHI.h"
 #include "RHICommandList.h"
+#include "RHIStaticStates.h"
 #include "RenderResource.h"
 #include "ShaderParameters.h"
 #include "Shader.h"
diff --git a/Engine/Source/Runtime/D3D12RHI/Private/D3D12Device.h b/Engine/Source/Runtime/D3D12RHI/Private/D3D12Device.h
--- a/Engine/Source/Runtime/D3D12RHI/Private/D3D12Device.h
+++ b/Engine/Source/Runtime/D3D12RHI/Private/D3D12Device.h
@@ -194,4 +194,29 @@ protected:
 	FD3D12RayTracingDescriptorHeapCache* RayTracingDescriptorHeapCache = nullptr;
 	void DestroyRayTracingDescriptorCache();
 #endif
-};
\ No newline at end of file
+};
+
+template <typename TDesc> 
+void TD3D12ViewDescriptorHandle<TDesc>::AllocateDescriptorSlot()
+{
+	if (Parent)
+	{
+		FD3D12Device* Device = GetParentDevice();
+		FD3D12OfflineDescriptorManager& DescriptorAllocator = Device->template GetViewDescriptorAllocator<TDesc>();
+		Handle = DescriptorAllocator.AllocateHeapSlot(Index);
+		check(Handle.ptr != 0);
+	}
+}
+
+template <typename TDesc> 
+void TD3D12ViewDescriptorHandle<TDesc>::FreeDescriptorSlot()
+{
+	if (Parent)
+	{
+		FD3D12Device* Device = GetParentDevice();
+		FD3D12OfflineDescriptorManager& DescriptorAllocator = Device->template GetViewDescriptorAllocator<TDesc>();
+		DescriptorAllocator.FreeHeapSlot(Handle, Index);
+		Handle.ptr = 0;
+	}
+	check(!Handle.ptr);
+}
diff --git a/Engine/Source/Runtime/D3D12RHI/Private/D3D12View.h b/Engine/Source/Runtime/D3D12RHI/Private/D3D12View.h
--- a/Engine/Source/Runtime/D3D12RHI/Private/D3D12View.h
+++ b/Engine/Source/Runtime/D3D12RHI/Private/D3D12View.h
@@ -687,28 +687,9 @@ public:
 	inline uint32 GetIndex() const { return Index; }
 
 private:
-	void AllocateDescriptorSlot()
-	{
-		if (Parent)
-		{
-			FD3D12Device* Device = GetParentDevice();
-			FD3D12OfflineDescriptorManager& DescriptorAllocator = Device->template GetViewDescriptorAllocator<TDesc>();
-			Handle = DescriptorAllocator.AllocateHeapSlot(Index);
-			check(Handle.ptr != 0);
-		}
-	}
-
-	void FreeDescriptorSlot()
-	{
-		if (Parent)
-		{
-			FD3D12Device* Device = GetParentDevice();
-			FD3D12OfflineDescriptorManager& DescriptorAllocator = Device->template GetViewDescriptorAllocator<TDesc>();
-			DescriptorAllocator.FreeHeapSlot(Handle, Index);
-			Handle.ptr = 0;
-		}
-		check(!Handle.ptr);
-	}
+	// Implemented in D3D12Device.h due to dependencies on FD3D12Device
+	inline void AllocateDescriptorSlot();
+	inline void FreeDescriptorSlot();
 };
 
 typedef TD3D12ViewDescriptorHandle<D3D12_SHADER_RESOURCE_VIEW_DESC>		FD3D12DescriptorHandleSRV;
diff --git a/Engine/Source/Runtime/Renderer/Private/WideCustomResolveShaders.h b/Engine/Source/Runtime/Renderer/Private/WideCustomResolveShaders.h
--- a/Engine/Source/Runtime/Renderer/Private/WideCustomResolveShaders.h
+++ b/Engine/Source/Runtime/Renderer/Private/WideCustomResolveShaders.h
@@ -6,6 +6,7 @@
 #include "ShaderParameters.h"
 #include "Shader.h"
 #include "GlobalShader.h"
+#include "ShaderParameterUtils.h"
 
 class FWideCustomResolveVS : public FGlobalShader
 {
