Undo multi-frame fast allocation change. There is still a race condition. To fix, a deferred recycle queue needs to be added which defeats the purpose of fast allocators. Async upload buffers are now allocated through the upload heap allocator.

diff --git a/Engine/Source/Runtime/D3D12RHI/Private/D3D12Allocation.cpp b/Engine/Source/Runtime/D3D12RHI/Private/D3D12Allocation.cpp
--- a/Engine/Source/Runtime/D3D12RHI/Private/D3D12Allocation.cpp
+++ b/Engine/Source/Runtime/D3D12RHI/Private/D3D12Allocation.cpp
@@ -1287,8 +1287,8 @@ HRESULT FD3D12TextureAllocatorPool::AllocateTexture(D3D12_RESOURCE_DESC Desc, co
 //	Fast Allocation
 //-----------------------------------------------------------------------------
 
-template void* FD3D12FastAllocator::Allocate<FD3D12ScopeLock>(uint32 Size, uint32 Alignment, class FD3D12ResourceLocation* ResourceLocation, bool bMultiFrame);
-template void* FD3D12FastAllocator::Allocate<FD3D12ScopeNoLock>(uint32 Size, uint32 Alignment, class FD3D12ResourceLocation* ResourceLocation, bool bMultiFrame);
+template void* FD3D12FastAllocator::Allocate<FD3D12ScopeLock>(uint32 Size, uint32 Alignment, class FD3D12ResourceLocation* ResourceLocation);
+template void* FD3D12FastAllocator::Allocate<FD3D12ScopeNoLock>(uint32 Size, uint32 Alignment, class FD3D12ResourceLocation* ResourceLocation);
 
 template void FD3D12FastAllocator::CleanupPages<FD3D12ScopeLock>(uint64 FrameLag);
 template void FD3D12FastAllocator::CleanupPages<FD3D12ScopeNoLock>(uint64 FrameLag);
@@ -1311,7 +1311,7 @@ FD3D12FastAllocator::FD3D12FastAllocator(FD3D12Device* Parent, FRHIGPUMask Visib
 {}
 
 template<typename LockType>
-void* FD3D12FastAllocator::Allocate(uint32 Size, uint32 Alignment, class FD3D12ResourceLocation* ResourceLocation, bool bMultiFrame)
+void* FD3D12FastAllocator::Allocate(uint32 Size, uint32 Alignment, class FD3D12ResourceLocation* ResourceLocation)
 {
 	LockType Lock(&CS);
 
@@ -1364,8 +1364,7 @@ void* FD3D12FastAllocator::Allocate(uint32 Size, uint32 Alignment, class FD3D12R
 			Size,
 			CurrentAllocatorPage->FastAllocBuffer->GetGPUVirtualAddress(),
 			CurrentAllocatorPage->FastAllocData,
-			CurrentOffset,
-			bMultiFrame);
+			CurrentOffset);
 
 		CurrentAllocatorPage->NextFastAllocOffset = CurrentOffset + Size;
 
diff --git a/Engine/Source/Runtime/D3D12RHI/Private/D3D12Allocation.h b/Engine/Source/Runtime/D3D12RHI/Private/D3D12Allocation.h
--- a/Engine/Source/Runtime/D3D12RHI/Private/D3D12Allocation.h
+++ b/Engine/Source/Runtime/D3D12RHI/Private/D3D12Allocation.h
@@ -493,7 +493,7 @@ public:
 	FD3D12FastAllocator(FD3D12Device* Parent, FRHIGPUMask VisibiltyMask, const D3D12_HEAP_PROPERTIES& InHeapProperties, uint32 PageSize);
 
 	template<typename LockType>
-	void* Allocate(uint32 Size, uint32 Alignment, class FD3D12ResourceLocation* ResourceLocation, bool bMultiFrame = false);
+	void* Allocate(uint32 Size, uint32 Alignment, class FD3D12ResourceLocation* ResourceLocation);
 
 	template<typename LockType>
 	void Destroy();
diff --git a/Engine/Source/Runtime/D3D12RHI/Private/D3D12Buffer.cpp b/Engine/Source/Runtime/D3D12RHI/Private/D3D12Buffer.cpp
--- a/Engine/Source/Runtime/D3D12RHI/Private/D3D12Buffer.cpp
+++ b/Engine/Source/Runtime/D3D12RHI/Private/D3D12Buffer.cpp
@@ -172,7 +172,16 @@ BufferType* FD3D12Adapter::CreateRHIBuffer(FRHICommandListImmediate* RHICmdList,
 
 			// Get an upload heap and initialize data
 			FD3D12ResourceLocation SrcResourceLoc(BufferOut->GetParentDevice());
-			void* pData = SrcResourceLoc.GetParentDevice()->GetDefaultFastAllocator().Allocate<FD3D12ScopeLock>(Size, 4UL, &SrcResourceLoc, bOnAsyncThread);
+			void* pData;
+			if (bOnAsyncThread)
+			{
+				const uint32 GPUIdx = SrcResourceLoc.GetParentDevice()->GetGPUIndex();
+				pData = GetUploadHeapAllocator(GPUIdx).AllocUploadResource(Size, 4u, SrcResourceLoc);
+			}
+			else
+			{
+				pData = SrcResourceLoc.GetParentDevice()->GetDefaultFastAllocator().Allocate<FD3D12ScopeLock>(Size, 4UL, &SrcResourceLoc);
+			}
 			check(pData);
 			FMemory::Memcpy(pData, CreateInfo.ResourceArray->GetResourceData(), Size);
 
diff --git a/Engine/Source/Runtime/D3D12RHI/Private/D3D12Resources.cpp b/Engine/Source/Runtime/D3D12RHI/Private/D3D12Resources.cpp
--- a/Engine/Source/Runtime/D3D12RHI/Private/D3D12Resources.cpp
+++ b/Engine/Source/Runtime/D3D12RHI/Private/D3D12Resources.cpp
@@ -586,12 +586,6 @@ void FD3D12ResourceLocation::ReleaseResource()
 		}
 		break;
 	}
-	case ResourceLocationType::eMultiFrameFastAllocation:
-	{
-		check(UnderlyingResource->GetRefCount() > 1);
-		UnderlyingResource->Release();
-		break;
-	}
 	case ResourceLocationType::eFastAllocation:
 	case ResourceLocationType::eUndefined:
 	default:
diff --git a/Engine/Source/Runtime/D3D12RHI/Public/D3D12Resources.h b/Engine/Source/Runtime/D3D12RHI/Public/D3D12Resources.h
--- a/Engine/Source/Runtime/D3D12RHI/Public/D3D12Resources.h
+++ b/Engine/Source/Runtime/D3D12RHI/Public/D3D12Resources.h
@@ -420,7 +420,6 @@ public:
 		eStandAlone,
 		eSubAllocation,
 		eFastAllocation,
-		eMultiFrameFastAllocation,
 		eAliased, // Oculus is the only API that uses this
 		eNodeReference,
 		eHeapAliased, 
@@ -496,17 +495,9 @@ public:
 	}
 
 
-	inline void AsFastAllocation(FD3D12Resource* Resource, uint32 BufferSize, D3D12_GPU_VIRTUAL_ADDRESS GPUBase, void* CPUBase, uint64 Offset, bool bMultiFrame = false)
+	inline void AsFastAllocation(FD3D12Resource* Resource, uint32 BufferSize, D3D12_GPU_VIRTUAL_ADDRESS GPUBase, void* CPUBase, uint64 Offset)
 	{
-		if (bMultiFrame)
-		{
-			Resource->AddRef();
-			SetType(ResourceLocationType::eMultiFrameFastAllocation);
-		}
-		else
-		{
-			SetType(ResourceLocationType::eFastAllocation);
-		}
+		SetType(ResourceLocationType::eFastAllocation);
 		SetResource(Resource);
 		SetSize(BufferSize);
 		SetOffsetFromBaseOfResource(Offset);
