Create optional Allocation tracking for UE-70143.
Users sometimes crash with leaks or incorrect frees. By enabling this code(disabled by default), they should now get a callstack of where they did the incorrect allocation.

diff --git a/Engine/Source/Runtime/D3D12RHI/Private/D3D12Allocation.cpp b/Engine/Source/Runtime/D3D12RHI/Private/D3D12Allocation.cpp
--- a/Engine/Source/Runtime/D3D12RHI/Private/D3D12Allocation.cpp
+++ b/Engine/Source/Runtime/D3D12RHI/Private/D3D12Allocation.cpp
@@ -8,6 +8,7 @@
 #include "D3D12RHIPrivate.h"
 #include "D3D12Allocation.h"
 #include "Misc/BufferedOutputDevice.h"
+#include "HAL/PlatformStackWalk.h"
 
 #if D3D12RHI_SEGREGATED_TEXTURE_ALLOC
 static int32 GD3D12ReadOnlyTextureAllocatorMinPoolSize = 4 * 1024 * 1024;
@@ -33,6 +34,16 @@ static FAutoConsoleVariableRef CVarD3D12ReadOnlyTextureAllocatorMaxPoolSize(
 	ECVF_ReadOnly);
 #endif
 
+
+#if D3D12RHI_SEGLIST_ALLOC_TRACK_WASTAGE
+static int32 GD3D12SegListTrackLeaks = 0;
+static FAutoConsoleVariableRef CVarD3D12SegListTrackLeaks(
+	TEXT("d3d12.SegListTrackLeaks"),
+	GD3D12SegListTrackLeaks,
+	TEXT("1: Enable leak tracking in d3d12 seglist's"),
+	ECVF_ReadOnly);
+#endif
+
 namespace ED3D12AllocatorID
 {
 	enum Type
@@ -1694,9 +1705,7 @@ void FD3D12SegListAllocator::FreeRetiredBlocks(TArray<TArray<FRetiredBlock, Allo
 			FD3D12SegList* Owner = BackingHeap->OwnerList;
 			check(!!Owner);
 			Owner->FreeBlock(BackingHeap, Block.Offset);
-#if D3D12RHI_SEGLIST_ALLOC_TRACK_WASTAGE
-			TotalBytesRequested -= Block.ResourceSize;
-#endif
+			OnFree(Block.Offset, BackingHeap, Block.ResourceSize);
 		}
 	}
 }
@@ -1741,9 +1750,7 @@ void FD3D12SegListAllocator::Destroy()
 		FreeRetiredBlocks(DeferredDeletionQueue);
 		FenceValues.Empty();
 		DeferredDeletionQueue.Empty();
-#if D3D12RHI_SEGLIST_ALLOC_TRACK_WASTAGE
-		check(!TotalBytesRequested);
-#endif
+		VerifyEmpty();
 	}
 	{
 		FRWScopeLock Lock(SegListsRWLock, SLT_Write);
@@ -1757,3 +1764,73 @@ void FD3D12SegListAllocator::Destroy()
 		SegLists.Empty();
 	}
 }
+#if D3D12RHI_SEGLIST_ALLOC_TRACK_WASTAGE
+void FD3D12SegListAllocator::VerifyEmpty()
+{
+	FScopeLock Lock(&SegListTrackedAllocationCS);
+	if(SegListTrackedAllocations.Num() != 0)
+	{
+		UE_LOG(LogD3D12RHI, Warning, TEXT("Dumping leaked SegListAllocations\n"));
+		for (FD3D12SegListAllocatorLeakTrack& LeakTrack : SegListTrackedAllocations)
+		{
+			DumpStack(LeakTrack);
+		}
+	}
+	check(TotalBytesRequested == 0); //SegList was not properly freed. Run with d3d12.SegListTrackLeaks=1 to print callstacks of offending allocations
+}
+
+
+void FD3D12SegListAllocator::DumpStack(const FD3D12SegListAllocatorLeakTrack& LeakTrack)
+{
+	UE_LOG(LogD3D12RHI, Warning, TEXT("Leaking Allocation Heap %p Offset %d\nStack Dump\n"), LeakTrack.Heap, LeakTrack.Offset);
+	for(uint32 Index = 0; Index < LeakTrack.StackDepth; ++Index)
+	{
+		const size_t STRING_SIZE = 16 * 1024;
+		ANSICHAR StackTrace[STRING_SIZE];
+		StackTrace[0] = 0;
+		FPlatformStackWalk::ProgramCounterToHumanReadableString(Index, LeakTrack.Stack[Index], StackTrace, STRING_SIZE, 0);
+		UE_LOG(LogD3D12RHI, Warning, TEXT("%d %S\n"), Index, StackTrace);
+	}
+}
+
+void FD3D12SegListAllocator::OnAlloc(uint32 Offset, void* Heap, uint32 Size)
+{
+	TotalBytesRequested += Size;
+
+	if(GD3D12SegListTrackLeaks == 0)
+		return;
+	FD3D12SegListAllocatorLeakTrack LeakTrack;
+	LeakTrack.Offset = Offset;
+	LeakTrack.Heap = Heap;
+	LeakTrack.Size = Size;
+	LeakTrack.StackDepth = FPlatformStackWalk::CaptureStackBackTrace(&LeakTrack.Stack[0], D3D12RHI_SEGLIST_ALLOC_TRACK_LEAK_STACK_DEPTH);
+
+	FScopeLock Lock(&SegListTrackedAllocationCS);
+	check(!SegListTrackedAllocations.Contains(LeakTrack));
+	SegListTrackedAllocations.Add(LeakTrack);
+}
+void FD3D12SegListAllocator::OnFree(uint32 Offset, void* Heap, uint32 Size)
+{
+	TotalBytesRequested -= Size;
+	if (GD3D12SegListTrackLeaks == 0)
+		return;
+
+	FD3D12SegListAllocatorLeakTrack LeakTrack;
+	LeakTrack.Offset = Offset;
+	LeakTrack.Heap = Heap;
+	FScopeLock Lock(&SegListTrackedAllocationCS);
+	FD3D12SegListAllocatorLeakTrack* Element = SegListTrackedAllocations.Find(LeakTrack);
+	check(Element); // element being freed was not found.
+	if(Element->Size != Size)
+	{
+		UE_LOG(LogD3D12RHI, Warning, TEXT("Mismatched alloc/free size %d != %d, %p/%08x"), Element->Size, Size, Element->Heap, Element->Offset);
+		DumpStack(*Element);
+		check(0); //element being freed had incorrect size. 
+	}
+	SegListTrackedAllocations.Remove(LeakTrack);
+	check(!SegListTrackedAllocations.Contains(LeakTrack));
+}
+#endif
+
+
+
diff --git a/Engine/Source/Runtime/D3D12RHI/Private/D3D12Allocation.h b/Engine/Source/Runtime/D3D12RHI/Private/D3D12Allocation.h
--- a/Engine/Source/Runtime/D3D12RHI/Private/D3D12Allocation.h
+++ b/Engine/Source/Runtime/D3D12RHI/Private/D3D12Allocation.h
@@ -29,6 +29,36 @@ D3D12Allocation.h: A Collection of allocators
 #define D3D12RHI_SEGREGATED_TEXTURE_ALLOC (PLATFORM_WINDOWS)
 #define D3D12RHI_SEGLIST_ALLOC_TRACK_WASTAGE (!(UE_BUILD_TEST || UE_BUILD_SHIPPING))
 
+#define D3D12RHI_SEGLIST_ALLOC_TRACK_LEAK_STACK_DEPTH 12
+
+struct FD3D12SegListAllocatorLeakTrack
+{
+	uint32 Offset;
+	void*  Heap;
+	uint32 Size;
+	uint32 StackDepth;
+	uint64 Stack[D3D12RHI_SEGLIST_ALLOC_TRACK_LEAK_STACK_DEPTH];
+	bool operator==(const FD3D12SegListAllocatorLeakTrack& Other) const
+	{
+		return Offset == Other.Offset && Heap == Other.Heap;
+	}
+};
+
+FORCEINLINE uint32 GetTypeHash(const FD3D12SegListAllocatorLeakTrack& S)
+{
+	uint32 Prime0 = 0xa6c70167;
+	uint32 Prime1 = 0x5d18b207;
+	uint32 Prime2 = 0xd0a489f9;
+	uint32 Value0 = (uint32)(((uint64)S.Heap) >> 32);
+	uint32 Value1 = (uint32)(uint64)S.Heap;
+	uint32 Value2 = S.Offset;
+	return Value0 * Prime0 + Value1 * Prime1 + Value2 * Prime2;
+}
+
+
+
+
+
 class FD3D12SegList;
 
 class FD3D12ResourceAllocator : public FD3D12DeviceChild, public FD3D12MultiNodeGPUObject
@@ -870,9 +900,7 @@ public:
 		check(!SegLists.Num());
 		check(!FenceValues.Num());
 		check(!DeferredDeletionQueue.Num());
-#if D3D12RHI_SEGLIST_ALLOC_TRACK_WASTAGE
-		check(!TotalBytesRequested);
-#endif
+		VerifyEmpty();
 	}
 
 	FD3D12SegListAllocator(const FD3D12SegListAllocator&) = delete;
@@ -913,9 +941,7 @@ public:
 				HeapFlags,
 				OutHeap);
 			check(Ret != InvalidOffset);
-#if D3D12RHI_SEGLIST_ALLOC_TRACK_WASTAGE
-			TotalBytesRequested += SizeInBytes;
-#endif
+			OnAlloc(Ret, OutHeap.GetReference(), SizeInBytes);
 			return Ret;
 		}
 		OutHeap = nullptr;
@@ -1010,6 +1036,18 @@ private:
 
 #if D3D12RHI_SEGLIST_ALLOC_TRACK_WASTAGE
 	TAtomic<uint64> TotalBytesRequested;
+	FCriticalSection SegListTrackedAllocationCS;
+	TSet<FD3D12SegListAllocatorLeakTrack> SegListTrackedAllocations;
+
+
+	void DumpStack(const FD3D12SegListAllocatorLeakTrack& LeakTrack);
+	void OnAlloc(uint32 Offset, void* Heap, uint32 Size);
+	void OnFree(uint32 Offset, void* Heap, uint32 Size);
+	void VerifyEmpty();
+#else
+	void OnAlloc(uint32 Offset, void* Heap, uint32 Size) {}
+	void OnFree(uint32 Offset, void* Heap, uint32 Size) {}
+	void VerifyEmpty(){}
 #endif
 };
 
