- Handle EReadable transitions on UAVs to avoid a barrier during SetSRVs

diff --git a/Engine/Source/Runtime/D3D12RHI/Private/D3D12Commands.cpp b/Engine/Source/Runtime/D3D12RHI/Private/D3D12Commands.cpp
--- a/Engine/Source/Runtime/D3D12RHI/Private/D3D12Commands.cpp
+++ b/Engine/Source/Runtime/D3D12RHI/Private/D3D12Commands.cpp
@@ -317,7 +317,19 @@ void FD3D12CommandContext::RHITransitionResources(EResourceTransitionAccess Tran
 		}
 		else
 		{
-#if !USE_D3D12RHI_RESOURCE_STATE_TRACKING
+#if USE_D3D12RHI_RESOURCE_STATE_TRACKING
+			if ( TransitionType == EResourceTransitionAccess::EReadable )
+			{
+				for (int32 i = 0; i < InNumUAVs; ++i)
+				{
+					if (InUAVs[i])
+					{
+						FD3D12UnorderedAccessView* const UnorderedAccessView = RetrieveObject<FD3D12UnorderedAccessView>(InUAVs[i]);
+						FD3D12DynamicRHI::TransitionResource(CommandListHandle, UnorderedAccessView, D3D12_RESOURCE_STATE_PIXEL_SHADER_RESOURCE | D3D12_RESOURCE_STATE_NON_PIXEL_SHADER_RESOURCE);
+					}
+				}
+			}
+#else
 			// Determine the direction of the transitions.
 			// Note in this method, the writeable state is always UAV, regardless of the FD3D12Resource's Writeable state.
 			const D3D12_RESOURCE_STATES* pBefore = nullptr;
@@ -375,7 +387,7 @@ void FD3D12CommandContext::RHITransitionResources(EResourceTransitionAccess Tran
 					}
 				}
 			}
-#endif // !USE_D3D12RHI_RESOURCE_STATE_TRACKING
+#endif // USE_D3D12RHI_RESOURCE_STATE_TRACKING
 		}
 	}
 
