Implement ConditionalResetSwapChain(). This is needed for fullscreen mode to work.

diff --git a/Engine/Source/Runtime/D3D12RHI/Private/Windows/WindowsD3D12Viewport.cpp b/Engine/Source/Runtime/D3D12RHI/Private/Windows/WindowsD3D12Viewport.cpp
--- a/Engine/Source/Runtime/D3D12RHI/Private/Windows/WindowsD3D12Viewport.cpp
+++ b/Engine/Source/Runtime/D3D12RHI/Private/Windows/WindowsD3D12Viewport.cpp
@@ -221,6 +221,27 @@ void FD3D12Viewport::Init()
 
 void FD3D12Viewport::ConditionalResetSwapChain(bool bIgnoreFocus)
 {
+	if (!bIsValid)
+	{
+		// Check if the viewport's window is focused before resetting the swap chain's fullscreen state.
+		HWND FocusWindow = ::GetFocus();
+		const bool bIsFocused = FocusWindow == WindowHandle;
+		const bool bIsIconic = !!::IsIconic(WindowHandle);
+		if (bIgnoreFocus || (bIsFocused && !bIsIconic))
+		{
+			FlushRenderingCommands();
+
+			HRESULT Result = SwapChain1->SetFullscreenState(bIsFullscreen, nullptr);
+			if (SUCCEEDED(Result))
+			{
+				bIsValid = true;
+			}
+			else if (Result != DXGI_ERROR_NOT_CURRENTLY_AVAILABLE && Result != DXGI_STATUS_MODE_CHANGE_IN_PROGRESS)
+			{
+				UE_LOG(LogD3D12RHI, Error, TEXT("IDXGISwapChain::SetFullscreenState returned %08x, unknown error status."), Result);
+			}
+		}
+	}
 }
 
 void FD3D12Viewport::ResizeInternal()
