Eliminate most hitches from upload buffer allocations.
 - Keep pages around longer
 - Split large allocations to allow more to fit within a page
 - Avoid holding a CriticalSection in standalone fallback path, allowing other threads to continue to do fast allocations
 - Remove the thread-local allocators, which we had way too many copies of and which frequently fell back to standalone allocations

diff --git a/Engine/Source/Runtime/D3D12RHI/Private/D3D12Adapter.cpp b/Engine/Source/Runtime/D3D12RHI/Private/D3D12Adapter.cpp
--- a/Engine/Source/Runtime/D3D12RHI/Private/D3D12Adapter.cpp
+++ b/Engine/Source/Runtime/D3D12RHI/Private/D3D12Adapter.cpp
@@ -590,13 +590,6 @@ void FD3D12Adapter::Cleanup()
 
 	FRHIResource::FlushPendingDeletes();
 
-	// Clean up the asnyc texture thread allocators
-	for (int32 i = 0; i < GetOwningRHI()->NumThreadDynamicHeapAllocators; i++)
-	{
-		GetOwningRHI()->ThreadDynamicHeapAllocatorArray[i]->Destroy<FD3D12ScopeLock>();
-		delete(GetOwningRHI()->ThreadDynamicHeapAllocatorArray[i]);
-	}
-
 	// Cleanup resources
 	DeferredDeletionQueue.Clear();
 
diff --git a/Engine/Source/Runtime/D3D12RHI/Private/D3D12Allocation.cpp b/Engine/Source/Runtime/D3D12RHI/Private/D3D12Allocation.cpp
--- a/Engine/Source/Runtime/D3D12RHI/Private/D3D12Allocation.cpp
+++ b/Engine/Source/Runtime/D3D12RHI/Private/D3D12Allocation.cpp
@@ -44,6 +44,13 @@ static FAutoConsoleVariableRef CVarD3D12SegListTrackLeaks(
 	ECVF_ReadOnly);
 #endif
 
+static int32 GD3D12FastAllocatorMinPagesToRetain = 5;
+static FAutoConsoleVariableRef CVarD3D12FastAllocatorMinPagesToRetain(
+	TEXT("d3d12.FastAllocator.MinPagesToRetain"),
+	GD3D12FastAllocatorMinPagesToRetain,
+	TEXT("Minimum number of pages to retain. Pages below this limit will never be released. Pages above can be released after being unused for a certain number of frames."),
+	ECVF_Default);
+
 namespace ED3D12AllocatorID
 {
 	enum Type
@@ -1298,15 +1305,6 @@ HRESULT FD3D12TextureAllocatorPool::AllocateTexture(D3D12_RESOURCE_DESC Desc, co
 //	Fast Allocation
 //-----------------------------------------------------------------------------
 
-template void* FD3D12FastAllocator::Allocate<FD3D12ScopeLock>(uint32 Size, uint32 Alignment, class FD3D12ResourceLocation* ResourceLocation);
-template void* FD3D12FastAllocator::Allocate<FD3D12ScopeNoLock>(uint32 Size, uint32 Alignment, class FD3D12ResourceLocation* ResourceLocation);
-
-template void FD3D12FastAllocator::CleanupPages<FD3D12ScopeLock>(uint64 FrameLag);
-template void FD3D12FastAllocator::CleanupPages<FD3D12ScopeNoLock>(uint64 FrameLag);
-
-template void FD3D12FastAllocator::Destroy<FD3D12ScopeLock>();
-template void FD3D12FastAllocator::Destroy<FD3D12ScopeNoLock>();
-
 FD3D12FastAllocator::FD3D12FastAllocator(FD3D12Device* Parent, FRHIGPUMask VisibiltyMask, D3D12_HEAP_TYPE InHeapType, uint32 PageSize)
 	: FD3D12DeviceChild(Parent)
 	, FD3D12MultiNodeGPUObject(Parent->GetGPUMask(), VisibiltyMask)
@@ -1321,11 +1319,8 @@ FD3D12FastAllocator::FD3D12FastAllocator(FD3D12Device* Parent, FRHIGPUMask Visib
 	, CurrentAllocatorPage(nullptr)
 {}
 
-template<typename LockType>
 void* FD3D12FastAllocator::Allocate(uint32 Size, uint32 Alignment, class FD3D12ResourceLocation* ResourceLocation)
 {
-	LockType Lock(&CS);
-
 	// Check to make sure our assumption that we don't need a ResourceLocation->Clear() here is valid.
 	checkf(!ResourceLocation->IsValid(), TEXT("The supplied resource location already has a valid resource. You should Clear() it first or it may leak."));
 
@@ -1353,6 +1348,8 @@ void* FD3D12FastAllocator::Allocate(uint32 Size, uint32 Alignment, class FD3D12R
 	}
 	else
 	{
+		FD3D12ScopeLock Lock(&CS);
+
 		const uint32 Offset = (CurrentAllocatorPage) ? CurrentAllocatorPage->NextFastAllocOffset : 0;
 		uint32 CurrentOffset = AlignArbitrary(Offset, Alignment);
 
@@ -1384,18 +1381,15 @@ void* FD3D12FastAllocator::Allocate(uint32 Size, uint32 Alignment, class FD3D12R
 	}
 }
 
-template<typename LockType>
 void FD3D12FastAllocator::CleanupPages(uint64 FrameLag)
 {
-	LockType Lock(&CS);
+	FD3D12ScopeLock Lock(&CS);
 	PagePool.CleanupPages(FrameLag);
 }
 
-template<typename LockType>
 void FD3D12FastAllocator::Destroy()
 {
-	LockType Lock(&CS);
-
+	FD3D12ScopeLock Lock(&CS);
 	if (CurrentAllocatorPage)
 	{
 		PagePool.ReturnFastAllocatorPage(CurrentAllocatorPage);
@@ -1449,6 +1443,7 @@ FD3D12FastAllocatorPage* FD3D12FastAllocatorPagePool::RequestFastAllocatorPage()
 	VERIFYD3D12RESULT(Adapter->CreateBuffer(HeapProperties, InitialState, PageSize, Page->FastAllocBuffer.GetInitReference(), TEXT("Fast Allocator Page")));
 
 	Page->FastAllocData = Page->FastAllocBuffer->Map();
+
 	return Page;
 }
 
@@ -1464,36 +1459,30 @@ void FD3D12FastAllocatorPagePool::ReturnFastAllocatorPage(FD3D12FastAllocatorPag
 
 void FD3D12FastAllocatorPagePool::CleanupPages(uint64 FrameLag)
 {
+	if (Pool.Num() <= GD3D12FastAllocatorMinPagesToRetain)
+	{
+		return;
+	}
+
 	FD3D12Adapter* Adapter = GetParentDevice()->GetParentAdapter();
 	FD3D12Fence& FrameFence = Adapter->GetFrameFence();
 
 	const uint64 CompletedFence = FrameFence.UpdateLastCompletedFence();
 
-	bool Found = false;
-
-	int32 Index = 0;
-	while (Index < Pool.Num())
+	// Pages get returned to end of list, so we'll look for pages to delete, starting from the LRU
+	for (int32 Index = 0; Index < Pool.Num(); Index++)
 	{
 		//If the GPU is done with it and no-one has a lock on it
 		if (Pool[Index]->FastAllocBuffer->GetRefCount() == 1 &&
 			Pool[Index]->FrameFence + FrameLag <= CompletedFence)
 		{
-			// Always keep one to avoid a recurring delete/create dance
-			if (Found)
-			{
-				FD3D12FastAllocatorPage* Page = Pool[Index];
-				Pool.RemoveAt(Index);
-				delete(Page);
+			FD3D12FastAllocatorPage* Page = Pool[Index];
+			Pool.RemoveAt(Index);
+			delete(Page);
 
-				continue;
-			}
-			else
-			{
-				Found = true;
-			}
+			// Only release at most one page per frame			
+			return;
 		}
-
-		++Index;
 	}
 }
 
diff --git a/Engine/Source/Runtime/D3D12RHI/Private/D3D12Allocation.h b/Engine/Source/Runtime/D3D12RHI/Private/D3D12Allocation.h
--- a/Engine/Source/Runtime/D3D12RHI/Private/D3D12Allocation.h
+++ b/Engine/Source/Runtime/D3D12RHI/Private/D3D12Allocation.h
@@ -522,13 +522,9 @@ public:
 	FD3D12FastAllocator(FD3D12Device* Parent, FRHIGPUMask VisibiltyMask, D3D12_HEAP_TYPE InHeapType, uint32 PageSize);
 	FD3D12FastAllocator(FD3D12Device* Parent, FRHIGPUMask VisibiltyMask, const D3D12_HEAP_PROPERTIES& InHeapProperties, uint32 PageSize);
 
-	template<typename LockType>
 	void* Allocate(uint32 Size, uint32 Alignment, class FD3D12ResourceLocation* ResourceLocation);
-
-	template<typename LockType>
 	void Destroy();
 
-	template<typename LockType>
 	void CleanupPages(uint64 FrameLag);
 
 protected:
diff --git a/Engine/Source/Runtime/D3D12RHI/Private/D3D12Buffer.cpp b/Engine/Source/Runtime/D3D12RHI/Private/D3D12Buffer.cpp
--- a/Engine/Source/Runtime/D3D12RHI/Private/D3D12Buffer.cpp
+++ b/Engine/Source/Runtime/D3D12RHI/Private/D3D12Buffer.cpp
@@ -180,7 +180,7 @@ BufferType* FD3D12Adapter::CreateRHIBuffer(FRHICommandListImmediate* RHICmdList,
 			}
 			else
 			{
-				pData = SrcResourceLoc.GetParentDevice()->GetDefaultFastAllocator().Allocate<FD3D12ScopeLock>(Size, 4UL, &SrcResourceLoc);
+				pData = SrcResourceLoc.GetParentDevice()->GetDefaultFastAllocator().Allocate(Size, 4UL, &SrcResourceLoc);
 			}
 			check(pData);
 			FMemory::Memcpy(pData, CreateInfo.ResourceArray->GetResourceData(), Size);
@@ -391,7 +391,7 @@ void* FD3D12DynamicRHI::LockBuffer(FRHICommandListImmediate* RHICmdList, BufferT
 		else
 		{
 			// If the static buffer is being locked for writing, allocate memory for the contents to be written to.
-			Data = Device->GetDefaultFastAllocator().Allocate<FD3D12ScopeLock>(Size, D3D12_TEXTURE_DATA_PLACEMENT_ALIGNMENT, &LockedData.ResourceLocation);
+			Data = Device->GetDefaultFastAllocator().Allocate(Size, D3D12_TEXTURE_DATA_PLACEMENT_ALIGNMENT, &LockedData.ResourceLocation);
 		}
 	}
 
diff --git a/Engine/Source/Runtime/D3D12RHI/Private/D3D12CommandContext.cpp b/Engine/Source/Runtime/D3D12RHI/Private/D3D12CommandContext.cpp
--- a/Engine/Source/Runtime/D3D12RHI/Private/D3D12CommandContext.cpp
+++ b/Engine/Source/Runtime/D3D12RHI/Private/D3D12CommandContext.cpp
@@ -444,39 +444,27 @@ void FD3D12CommandContextBase::RHIEndFrame()
 		Device->GetTextureAllocator().CleanUpAllocations();
 		Device->GetDefaultBufferAllocator().CleanupFreeBlocks();
 
-		Device->GetDefaultFastAllocator().CleanupPages<FD3D12ScopeLock>(10);
+		Device->GetDefaultFastAllocator().CleanupPages(10);
 	}
 
-		// The Texture streaming threads
-		{
-			for (int32 i = 0; i < FD3D12DynamicRHI::GetD3DRHI()->NumThreadDynamicHeapAllocators; ++i)
-			{
-				FD3D12FastAllocator* TextureStreamingAllocator = FD3D12DynamicRHI::GetD3DRHI()->ThreadDynamicHeapAllocatorArray[i];
-				if (TextureStreamingAllocator)
-				{
-					TextureStreamingAllocator->CleanupPages<FD3D12ScopeLock>(10);
-				}
-			}
-		}
-
 	for (uint32 GPUIndex : GPUMask)
 	{
 		FD3D12Device* Device = ParentAdapter->GetDevice(GPUIndex);
 		Device->GetCommandListManager().ReleaseResourceBarrierCommandListAllocator();
 	}
 
-		UpdateMemoryStats();
+	UpdateMemoryStats();
 
-	    // Stop Timing at the very last moment
+	// Stop Timing at the very last moment
     
-	    ParentAdapter->GetGPUProfiler().EndFrame(ParentAdapter->GetOwningRHI());
+	ParentAdapter->GetGPUProfiler().EndFrame(ParentAdapter->GetOwningRHI());
     
     
-	    // Advance frame fence
+	// Advance frame fence
     
-	    FD3D12ManualFence& FrameFence = ParentAdapter->GetFrameFence();
-	    FrameFence.Signal(ED3D12CommandQueueType::Default, FrameFence.IncrementCurrentFence());
-	}
+	FD3D12ManualFence& FrameFence = ParentAdapter->GetFrameFence();
+	FrameFence.Signal(ED3D12CommandQueueType::Default, FrameFence.IncrementCurrentFence());
+}
 
 void FD3D12CommandContextBase::UpdateMemoryStats()
 {
diff --git a/Engine/Source/Runtime/D3D12RHI/Private/D3D12Device.cpp b/Engine/Source/Runtime/D3D12RHI/Private/D3D12Device.cpp
--- a/Engine/Source/Runtime/D3D12RHI/Private/D3D12Device.cpp
+++ b/Engine/Source/Runtime/D3D12RHI/Private/D3D12Device.cpp
@@ -56,7 +56,7 @@ FD3D12Device::~FD3D12Device()
 	// Cleanup the allocator near the end, as some resources may be returned to the allocator or references are shared by multiple GPUs
 	DefaultBufferAllocator.FreeDefaultBufferPools();
 
-	DefaultFastAllocator.Destroy<FD3D12ScopeLock>();
+	DefaultFastAllocator.Destroy();
 
 	TextureAllocator.CleanUpAllocations();
 	TextureAllocator.Destroy();
diff --git a/Engine/Source/Runtime/D3D12RHI/Private/D3D12RHI.cpp b/Engine/Source/Runtime/D3D12RHI/Private/D3D12RHI.cpp
--- a/Engine/Source/Runtime/D3D12RHI/Private/D3D12RHI.cpp
+++ b/Engine/Source/Runtime/D3D12RHI/Private/D3D12RHI.cpp
@@ -36,18 +36,11 @@ TAutoConsoleVariable<int32> CVarD3D12ZeroBufferSizeInMB(
 	ECVF_ReadOnly
 	);
 
-/// @cond DOXYGEN_WARNINGS
-
-__declspec(thread) FD3D12FastAllocator* FD3D12DynamicRHI::HelperThreadDynamicHeapAllocator = nullptr;
-
-/// @endcond
-
 FD3D12DynamicRHI* FD3D12DynamicRHI::SingleD3DRHI = nullptr;
 
 using namespace D3D12RHI;
 
 FD3D12DynamicRHI::FD3D12DynamicRHI(const TArray<TSharedPtr<FD3D12Adapter>>& ChosenAdaptersIn) :
-	NumThreadDynamicHeapAllocators(0),
 	ChosenAdapters(ChosenAdaptersIn),
 	AmdAgsContext(nullptr),
 	FlipEvent(INVALID_HANDLE_VALUE)
@@ -56,8 +49,6 @@ FD3D12DynamicRHI::FD3D12DynamicRHI(const TArray<TSharedPtr<FD3D12Adapter>>& Chos
 	check(SingleD3DRHI == nullptr);
 	SingleD3DRHI = this;
 
-	ThreadDynamicHeapAllocatorArray.AddZeroed(FPlatformMisc::NumberOfCoresIncludingHyperthreads());
-
 	// This should be called once at the start 
 	check(IsInGameThread());
 	check(!GIsThreadedRendering);
diff --git a/Engine/Source/Runtime/D3D12RHI/Private/D3D12RHIPrivate.h b/Engine/Source/Runtime/D3D12RHI/Private/D3D12RHIPrivate.h
--- a/Engine/Source/Runtime/D3D12RHI/Private/D3D12RHIPrivate.h
+++ b/Engine/Source/Runtime/D3D12RHI/Private/D3D12RHIPrivate.h
@@ -902,31 +902,6 @@ public:
 
 public:
 
-	inline FD3D12FastAllocator& GetHelperThreadDynamicUploadHeapAllocator()
-	{
-		FD3D12Device* Device = GetRHIDevice();
-		const FRHIGPUMask Node = Device->GetGPUMask();
-
-		check(!IsInActualRenderingThread());
-
-		static const uint32 AsyncTexturePoolSize = 1024 * 512;
-
-		if (HelperThreadDynamicHeapAllocator == nullptr)
-		{
-			uint32 NextIndex = FPlatformAtomics::InterlockedIncrement(&NumThreadDynamicHeapAllocators) - 1;
-			check(NextIndex < (uint32)ThreadDynamicHeapAllocatorArray.Num());
-			HelperThreadDynamicHeapAllocator = new FD3D12FastAllocator(Device, Node, D3D12_HEAP_TYPE_UPLOAD, AsyncTexturePoolSize);
-
-			ThreadDynamicHeapAllocatorArray[NextIndex] = HelperThreadDynamicHeapAllocator;
-		}
-
-		return *HelperThreadDynamicHeapAllocator;
-	}
-
-	TArray<FD3D12FastAllocator*> ThreadDynamicHeapAllocatorArray;
-	int32 NumThreadDynamicHeapAllocators;
-	static __declspec(thread) FD3D12FastAllocator* HelperThreadDynamicHeapAllocator;
-
 #if	PLATFORM_SUPPORTS_VIRTUAL_TEXTURES
 	virtual void* CreateVirtualTexture(uint32 Flags, D3D12_RESOURCE_DESC& ResourceDesc, const struct FD3D12TextureLayout& TextureLayout, FD3D12Resource** ppResource, FPlatformMemory::FPlatformVirtualMemoryBlock& RawTextureBlock, D3D12_RESOURCE_STATES InitialUsage = D3D12_RESOURCE_STATE_PIXEL_SHADER_RESOURCE) = 0;
 	virtual void DestroyVirtualTexture(uint32 Flags, void* RawTextureMemory, FPlatformMemory::FPlatformVirtualMemoryBlock& RawTextureBlock, uint64 CommittedTextureSize) = 0;
diff --git a/Engine/Source/Runtime/D3D12RHI/Private/D3D12Texture.cpp b/Engine/Source/Runtime/D3D12RHI/Private/D3D12Texture.cpp
--- a/Engine/Source/Runtime/D3D12RHI/Private/D3D12Texture.cpp
+++ b/Engine/Source/Runtime/D3D12RHI/Private/D3D12Texture.cpp
@@ -1,12 +1,12 @@
 // Copyright 1998-2019 Epic Games, Inc. All Rights Reserved.
 
 /*=============================================================================
-	D3D12VertexBuffer.cpp: D3D texture RHI implementation.
+	D3D12Texture.cpp: D3D texture RHI implementation.
 	=============================================================================*/
 
 #include "D3D12RHIPrivate.h"
 
- 
+
 int64 FD3D12GlobalStats::GDedicatedVideoMemory = 0;
 int64 FD3D12GlobalStats::GDedicatedSystemMemory = 0;
 int64 FD3D12GlobalStats::GSharedSystemMemory = 0;
@@ -1259,10 +1259,31 @@ FTexture2DRHIRef FD3D12DynamicRHI::RHIAsyncCreateTexture2D(uint32 SizeX, uint32
 
 		check((TextureDesc.Flags & D3D12_RESOURCE_FLAG_DENY_SHADER_RESOURCE) == 0);
 
-		const uint64 Size = GetRequiredIntermediateSize(TextureOut->GetResource()->GetResource(), 0, NumMips);
-		FD3D12FastAllocator& FastAllocator = GetHelperThreadDynamicUploadHeapAllocator();
+		FD3D12FastAllocator& FastAllocator = TextureOut->GetParentDevice()->GetDefaultFastAllocator();
+		uint64 Size = GetRequiredIntermediateSize(TextureOut->GetResource()->GetResource(), 0, NumMips);
+		uint64 SizeLowMips;
+
 		FD3D12ResourceLocation TempResourceLocation(FastAllocator.GetParentDevice());
-		FastAllocator.Allocate<FD3D12ScopeLock>(Size, D3D12_TEXTURE_DATA_PLACEMENT_ALIGNMENT, &TempResourceLocation);
+		FD3D12ResourceLocation TempResourceLocationLowMips(FastAllocator.GetParentDevice());
+
+		// The allocator work in pages of 4MB. Increasing page size is undesirable from a hitching point of view because there's a performance cliff above 4MB
+		// where creation time of new pages can increase by an order of magnitude. Most allocations are smaller than 4MB, but a common exception is
+		// 2048x2048 BC3 textures with mips, which takes 5.33MB. To avoid this case falling into the standalone allocations fallback path and risking hitching badly,
+		// we split the top mip into a separate allocation, allowing it to fit within 4MB.
+		const bool bSplitAllocation = (Size > 4 * 1024 * 1024) && (NumMips > 1);
+
+		if (bSplitAllocation)
+		{
+			Size = GetRequiredIntermediateSize(TextureOut->GetResource()->GetResource(), 0, 1);
+			SizeLowMips = GetRequiredIntermediateSize(TextureOut->GetResource()->GetResource(), 1, NumMips - 1);
+
+			FastAllocator.Allocate(Size, D3D12_TEXTURE_DATA_PLACEMENT_ALIGNMENT, &TempResourceLocation);
+			FastAllocator.Allocate(SizeLowMips, D3D12_TEXTURE_DATA_PLACEMENT_ALIGNMENT, &TempResourceLocationLowMips);
+		}
+		else
+		{
+			FastAllocator.Allocate(Size, D3D12_TEXTURE_DATA_PLACEMENT_ALIGNMENT, &TempResourceLocation);
+		}
 
 		FD3D12Texture2D* CurrentTexture = TextureOut;
 		while (CurrentTexture != nullptr)
@@ -1276,13 +1297,35 @@ FTexture2DRHIRef FD3D12DynamicRHI::RHIAsyncCreateTexture2D(uint32 SizeX, uint32
 			hCopyCommandList.SetCurrentOwningContext(&Device->GetDefaultCommandContext());
 
 			hCopyCommandList.GetCurrentOwningContext()->numCopies++;
-			UpdateSubresources(
-				(ID3D12GraphicsCommandList*)hCopyCommandList.CommandList(),
-				Resource->GetResource(),
-				TempResourceLocation.GetResource()->GetResource(),
-				TempResourceLocation.GetOffsetFromBaseOfResource(),
-				0, NumMips,
-				SubResourceData);
+
+			if (bSplitAllocation)
+			{
+				UpdateSubresources(
+					(ID3D12GraphicsCommandList*)hCopyCommandList.CommandList(),
+					Resource->GetResource(),
+					TempResourceLocation.GetResource()->GetResource(),
+					TempResourceLocation.GetOffsetFromBaseOfResource(),
+					0, 1,
+					SubResourceData);
+
+				UpdateSubresources(
+					(ID3D12GraphicsCommandList*)hCopyCommandList.CommandList(),
+					Resource->GetResource(),
+					TempResourceLocationLowMips.GetResource()->GetResource(),
+					TempResourceLocationLowMips.GetOffsetFromBaseOfResource(),
+					1, NumMips - 1,
+					SubResourceData + 1);
+			}
+			else
+			{
+				UpdateSubresources(
+					(ID3D12GraphicsCommandList*)hCopyCommandList.CommandList(),
+					Resource->GetResource(),
+					TempResourceLocation.GetResource()->GetResource(),
+					TempResourceLocation.GetOffsetFromBaseOfResource(),
+					0, NumMips,
+					SubResourceData);
+			}
 
 			hCopyCommandList.UpdateResidency(Resource);
 
@@ -1564,7 +1607,7 @@ void* TD3D12Texture2D<RHIResourceType>::Lock(class FRHICommandListImmediate* RHI
 		//const uint32 bufferSize = (uint32)GetRequiredIntermediateSize(this->GetResource()->GetResource(), Subresource, 1);
 		const uint32 bufferSize = Align(MipBytesAligned, D3D12_TEXTURE_DATA_PLACEMENT_ALIGNMENT);
 
-		void* pData = Device->GetDefaultFastAllocator().Allocate<FD3D12ScopeLock>(bufferSize, D3D12_TEXTURE_DATA_PLACEMENT_ALIGNMENT, &LockedResource->ResourceLocation);
+		void* pData = Device->GetDefaultFastAllocator().Allocate(bufferSize, D3D12_TEXTURE_DATA_PLACEMENT_ALIGNMENT, &LockedResource->ResourceLocation);
 		if (nullptr == pData)
 		{
 			check(false);
@@ -1708,7 +1751,7 @@ void FD3D12TextureBase::InitializeTextureData(FRHICommandListImmediate* RHICmdLi
 	uint64 Size = GetRequiredIntermediateSize(GetResource()->GetResource(), 0, NumSubresources);
 
 	FD3D12ResourceLocation SrcResourceLoc(GetParentDevice());
-	uint8* DstData = (uint8*) SrcResourceLoc.GetParentDevice()->GetDefaultFastAllocator().Allocate<FD3D12ScopeLock>(Size, D3D12_TEXTURE_DATA_PLACEMENT_ALIGNMENT, &SrcResourceLoc);
+	uint8* DstData = (uint8*) SrcResourceLoc.GetParentDevice()->GetDefaultFastAllocator().Allocate(Size, D3D12_TEXTURE_DATA_PLACEMENT_ALIGNMENT, &SrcResourceLoc);
 	uint32 DstDataOffset = 0;
 
 	const uint8* SrcData = (const uint8*) InitData;
@@ -1867,7 +1910,7 @@ void TD3D12Texture2D<RHIResourceType>::UpdateTexture2D(class FRHICommandListImme
 	while (Texture)
 	{
 		FD3D12ResourceLocation UploadHeapResourceLocation(GetParentDevice());
-		void* pData = GetParentDevice()->GetDefaultFastAllocator().template Allocate<FD3D12ScopeLock>(bufferSize, D3D12_TEXTURE_DATA_PLACEMENT_ALIGNMENT, &UploadHeapResourceLocation);
+		void* pData = GetParentDevice()->GetDefaultFastAllocator().Allocate(bufferSize, D3D12_TEXTURE_DATA_PLACEMENT_ALIGNMENT, &UploadHeapResourceLocation);
 		check(nullptr != pData);
 
 		byte* pRowData = (byte*)pData;
@@ -2234,7 +2277,7 @@ FUpdateTexture3DData FD3D12DynamicRHI::BeginUpdateTexture3D_Internal(FRHITexture
 
 		//@TODO Probably need to use the TextureAllocator here to get correct tiling.
 		// Currently the texture are allocated in linear, see hanlding around bVolume in FXboxOneTextureFormat::CompressImage(). 
-		UpdateData.Data = (uint8*)GetRHIDevice()->GetDefaultFastAllocator().Allocate<FD3D12ScopeLock>(BufferSize, D3D12_TEXTURE_DATA_PLACEMENT_ALIGNMENT, UpdateDataD3D12->UploadHeapResourceLocation);
+		UpdateData.Data = (uint8*)GetRHIDevice()->GetDefaultFastAllocator().Allocate(BufferSize, D3D12_TEXTURE_DATA_PLACEMENT_ALIGNMENT, UpdateDataD3D12->UploadHeapResourceLocation);
 
 		check(UpdateData.Data != nullptr);
 	}
