Remove serialization of PSO creation. This substantially reduces the amount of hitching in the empty cache case.

diff --git a/Engine/Source/Runtime/D3D12RHI/Private/D3D12PipelineState.h b/Engine/Source/Runtime/D3D12RHI/Private/D3D12PipelineState.h
--- a/Engine/Source/Runtime/D3D12RHI/Private/D3D12PipelineState.h
+++ b/Engine/Source/Runtime/D3D12RHI/Private/D3D12PipelineState.h
@@ -290,23 +290,36 @@ public:
 	void Create(const GraphicsPipelineCreationArgs& InCreationArgs);
 	void CreateAsync(const GraphicsPipelineCreationArgs& InCreationArgs);
 
-	ID3D12PipelineState* GetPipelineState();
+	FORCEINLINE ID3D12PipelineState* GetPipelineState()
+	{
+		if (CachedPipelineState)
+		{
+			return CachedPipelineState;
+		}
+		else
+		{
+			return InternalGetPipelineState();
+		}
+	}
 
 	FD3D12PipelineState& operator=(const FD3D12PipelineState& other)
 	{
 		checkSlow(GPUMask == other.GPUMask);
 		checkSlow(VisibilityMask == other.VisibilityMask);
-		ensure(PendingWaitOnWorkerCalls == 0);
 
 		PipelineState = other.PipelineState;
 		Worker = other.Worker;
 		return *this;
 	}
 
+private:
+	ID3D12PipelineState* InternalGetPipelineState();
+
 protected:
+	ID3D12PipelineState* CachedPipelineState;
 	TRefCountPtr<ID3D12PipelineState> PipelineState;
 	FAsyncTask<FD3D12PipelineStateWorker>* Worker;
-	volatile int32 PendingWaitOnWorkerCalls;
+	FRWLock GetPipelineStateMutex;
 };
 
 struct FD3D12GraphicsPipelineState : public FRHIGraphicsPipelineState
@@ -426,12 +439,7 @@ protected:
 	mutable FRWLock LowLevelGraphicsPipelineStateCacheMutex;
 	mutable FRWLock ComputePipelineStateCacheMutex;
 
-	FCriticalSection DiskCachesCS;
-
-#if !PLATFORM_WINDOWS
-	FRWLock CS;
-#endif
-
+	FRWLock DiskCachesCS;
 	FDiskCacheInterface DiskCaches[NUM_PSO_CACHE_TYPES];
 
 	void CleanupPipelineStateCaches();
diff --git a/Engine/Source/Runtime/D3D12RHI/Private/D3D12Pipelinestate.cpp b/Engine/Source/Runtime/D3D12RHI/Private/D3D12Pipelinestate.cpp
--- a/Engine/Source/Runtime/D3D12RHI/Private/D3D12Pipelinestate.cpp
+++ b/Engine/Source/Runtime/D3D12RHI/Private/D3D12Pipelinestate.cpp
@@ -116,7 +116,7 @@ FD3D12PipelineStateWorker::FD3D12PipelineStateWorker(FD3D12Adapter* Adapter, con
 
 
 /// @endcond
-#if defined(_M_IX86_FP) && _M_IX86_FP >= 2
+#if defined(_WIN64) || (defined(_M_IX86_FP) && _M_IX86_FP >= 2)
 #ifndef __clang__
 
 FORCEINLINE uint32 SSE4_CRC32(const void* Data, SIZE_T NumBytes)
@@ -163,7 +163,7 @@ uint32 FD3D12PipelineStateCacheBase::HashData(const void* Data, SIZE_T NumBytes)
 #ifdef __clang__
 	return FCrc::MemCrc32(Data, NumBytes);
 #else
-#if defined(_M_IX86_FP) && _M_IX86_FP >= 2
+#if defined(_WIN64) || (defined(_M_IX86_FP) && _M_IX86_FP >= 2)
 	if (GCPUSupportsSSE4)
 	{
 		return SSE4_CRC32(Data, NumBytes);
@@ -219,7 +219,7 @@ SIZE_T FD3D12PipelineStateCacheBase::HashPSODesc(const FD3D12ComputePipelineStat
 FD3D12PipelineStateCacheBase::FD3D12PipelineStateCacheBase(FD3D12Adapter* InParent) :
 	FD3D12AdapterChild(InParent)
 {
-#if defined(_M_IX86_FP) && _M_IX86_FP >= 2
+#if defined(_WIN64) || (defined(_M_IX86_FP) && _M_IX86_FP >= 2)
 	// Check for SSE4 support see: https://msdn.microsoft.com/en-us/library/vstudio/hskdteyh(v=vs.100).aspx
 	{
 		int32 cpui[4];
@@ -240,8 +240,8 @@ FD3D12PipelineStateCacheBase::~FD3D12PipelineStateCacheBase()
 FD3D12PipelineState::FD3D12PipelineState(FD3D12Adapter* Parent)
 	: FD3D12AdapterChild(Parent)
 	, FD3D12MultiNodeGPUObject(FRHIGPUMask::All(), FRHIGPUMask::All()) //Create on all, visible on all
+	, CachedPipelineState(nullptr)
 	, Worker(nullptr)
-	, PendingWaitOnWorkerCalls(0)
 {
 	INC_DWORD_STAT(STAT_D3D12NumPSOs);
 }
@@ -250,7 +250,6 @@ FD3D12PipelineState::~FD3D12PipelineState()
 {
 	if (Worker)
 	{
-		ensure(PendingWaitOnWorkerCalls == 0);
 		Worker->EnsureCompletion(true);
 		delete Worker;
 		Worker = nullptr;
@@ -259,57 +258,42 @@ FD3D12PipelineState::~FD3D12PipelineState()
 	DEC_DWORD_STAT(STAT_D3D12NumPSOs);
 }
 
-ID3D12PipelineState* FD3D12PipelineState::GetPipelineState()
+ID3D12PipelineState* FD3D12PipelineState::InternalGetPipelineState()
 {
+	// Only one thread should finalize any tasks and set the cached pipeline state.
+	// Other threads can wait efficiently behind the lock. Using Sleep() can wake threads 1ms too late.
+	FRWScopeLock Lock(GetPipelineStateMutex, SLT_Write);
+
 	if (Worker)
 	{
-		const bool bIsSyncThread = (FPlatformAtomics::InterlockedIncrement(&PendingWaitOnWorkerCalls) == 1);
+		Worker->EnsureCompletion(true);
+		check(Worker->IsWorkDone());
 
-		// Cache the Worker ptr as the thread with bIsSyncThread cloud clear it at any time.
-		// MemoryBarrier() is required to prevent the compiler from caching Worker.
-		FPlatformMisc::MemoryBarrier();
-		FAsyncTask<FD3D12PipelineStateWorker>* WorkerRef = Worker;
-		if (WorkerRef)
-		{
-			WorkerRef->EnsureCompletion(true);
-			check(WorkerRef->IsWorkDone());		
+		PipelineState = Worker->GetTask().PSO.GetReference();
+		CachedPipelineState = PipelineState.GetReference();
 
-			if (bIsSyncThread)
-			{
-				PipelineState = WorkerRef->GetTask().PSO;
-				
-				// Only set the worker to null after setting the PipelineState because of the initial branching.
-				// Note that only one thread must set the pipeline state as TRefCountPtr is not threadsafe.
-				Worker = nullptr;
-
-				// Decrement but also wait till 0 before deleting the worker as other threads could be referring it.
-				if (FPlatformAtomics::InterlockedDecrement(&PendingWaitOnWorkerCalls) != 0)
-				{
-					do
-					{
-						FPlatformProcess::Sleep(0);
-					}
-					while (PendingWaitOnWorkerCalls != 0);
-				}
-
-				delete WorkerRef;
-			}
-			else
+		// Cleanup the worker.
+		delete Worker;
+		Worker = nullptr;
+	}
+
+	// Busy-wait for the PSO. This avoids giving up our time slice.
+	if (CachedPipelineState == nullptr)
+	{
+		const double StartTime = FPlatformTime::Seconds();
+		static const float BusyWaitTimeoutInMs = CVarPSOStallTimeoutInMs.GetValueOnAnyThread();
+		while (PipelineState.GetReference() == nullptr)
+		{
+			if (((FPlatformTime::Seconds() - StartTime) * 1000) > BusyWaitTimeoutInMs)
 			{
-				// Cache the result before decrementing the counter because after the decrement, 
-				// the worker could be deleted at any time by the thread with bIsSyncThread.
-				// this allows to return immediately without having to wait for PendingWaitOnWorkerCalls to reach 0.
-				ID3D12PipelineState* Result = WorkerRef->GetTask().PSO.GetReference();
-				FPlatformAtomics::InterlockedDecrement(&PendingWaitOnWorkerCalls);
-				return Result;
+				UE_LOG(LogD3D12RHI, Fatal, TEXT("Waiting for PSO creation failed to complete within the timeout interval (%.3f ms)."), BusyWaitTimeoutInMs);
 			}
 		}
-		else // Decrement but don't wait since if Worker is null, PipelineState is valid.
-		{
-			FPlatformAtomics::InterlockedDecrement(&PendingWaitOnWorkerCalls);
-		}
+
+		CachedPipelineState = PipelineState.GetReference();
 	}
-	return PipelineState.GetReference();
+
+	return CachedPipelineState;
 }
 
 FD3D12GraphicsPipelineState::FD3D12GraphicsPipelineState(
@@ -354,10 +338,10 @@ FD3D12ComputePipelineState::~FD3D12ComputePipelineState()
 void FD3D12PipelineStateCacheBase::CleanupPipelineStateCaches()
 {
 #if D3D12RHI_USE_HIGH_LEVEL_PSO_CACHE
+	// The runtime caches manage the lifetime of their FD3D12GraphicsPipelineState and FD3D12ComputePipelineState.
+	// We need to release them.
 	{
 		FRWScopeLock Lock(InitializerToGraphicsPipelineMapMutex, FRWScopeLockType::SLT_Write);
-		// The runtime caches manage the lifetime of their FD3D12GraphicsPipelineState and FD3D12ComputePipelineState.
-		// We need to release them.
 		for (auto Pair : InitializerToGraphicsPipelineMap)
 		{
 			FD3D12GraphicsPipelineState* GraphicsPipelineState = Pair.Value;
@@ -380,8 +364,6 @@ void FD3D12PipelineStateCacheBase::CleanupPipelineStateCaches()
 #endif
 	{
 		FRWScopeLock Lock(LowLevelGraphicsPipelineStateCacheMutex, FRWScopeLockType::SLT_Write);
-		// The low level graphics and compute maps manage the lifetime of their PSOs.
-		// We need to delete each element before we empty it.
 		for (auto Iter = LowLevelGraphicsPipelineStateCache.CreateConstIterator(); Iter; ++Iter)
 		{
 			const FD3D12PipelineState* const PipelineState = Iter.Value();
@@ -470,24 +452,27 @@ void FD3D12PipelineStateCacheBase::AddToLowLevelCache(const FD3D12LowLevelGraphi
 
 	// Double check the desc doesn't already exist while the lock is taken.
 	// This avoids having multiple threads try to create the same PSO.
-	FRWScopeLock Lock(LowLevelGraphicsPipelineStateCacheMutex, FRWScopeLockType::SLT_Write);
-	FD3D12PipelineState** const PipelineState = LowLevelGraphicsPipelineStateCache.Find(Desc);
-	if (PipelineState)
-	{
-		// This desc already exists.
-		*OutPipelineState = *PipelineState;
-	}
-	else
 	{
+		FRWScopeLock Lock(LowLevelGraphicsPipelineStateCacheMutex, FRWScopeLockType::SLT_Write);
+		FD3D12PipelineState** const PipelineState = LowLevelGraphicsPipelineStateCache.Find(Desc);
+		if (PipelineState)
+		{
+			// This desc already exists.
+			*OutPipelineState = *PipelineState;
+			return;
+		}
+
+		// Add the FD3D12PipelineState object to the cache, but don't actually create the underlying PSO yet while the lock is taken.
+		// This allows multiple threads to create different PSOs at the same time.
+		INC_DWORD_STAT(STAT_PSOGraphicsNumLowlevelCacheEntries);
 		FD3D12PipelineState* NewPipelineState = new FD3D12PipelineState(GetParentAdapter());
 		LowLevelGraphicsPipelineStateCache.Add(Desc, NewPipelineState);
 
-		INC_DWORD_STAT(STAT_PSOGraphicsNumLowlevelCacheEntries);
 		*OutPipelineState = NewPipelineState;
-
-		// Do the callback now with the lock still on.
-		PostCreateCallback(OutPipelineState, Desc);
 	}
+
+	// Create the underlying PSO and then perform any other additional tasks like cleaning up/adding to caches, etc.
+	PostCreateCallback(OutPipelineState, Desc);
 }
 
 #if D3D12RHI_USE_HIGH_LEVEL_PSO_CACHE
@@ -548,24 +533,27 @@ void FD3D12PipelineStateCacheBase::AddToLowLevelCache(const FD3D12ComputePipelin
 
 	// Double check the desc doesn't already exist while the lock is taken.
 	// This avoids having multiple threads try to create the same PSO.
-	FRWScopeLock Lock(ComputePipelineStateCacheMutex, FRWScopeLockType::SLT_Write);
-	FD3D12PipelineState** const PipelineState = ComputePipelineStateCache.Find(Desc);
-	if (PipelineState)
-	{
-		// This desc already exists.
-		*OutPipelineState = *PipelineState;
-	}
-	else
 	{
+		FRWScopeLock Lock(ComputePipelineStateCacheMutex, FRWScopeLockType::SLT_Write);
+		FD3D12PipelineState** const PipelineState = ComputePipelineStateCache.Find(Desc);
+		if (PipelineState)
+		{
+			// This desc already exists.
+			*OutPipelineState = *PipelineState;
+			return;
+		}
+
+		// Add the FD3D12PipelineState object to the cache, but don't actually create the underlying PSO yet while the lock is taken.
+		// This allows multiple threads to create different PSOs at the same time.
+		INC_DWORD_STAT(STAT_PSOComputeNumLowlevelCacheEntries);
 		FD3D12PipelineState* NewPipelineState = new FD3D12PipelineState(GetParentAdapter());
 		ComputePipelineStateCache.Add(Desc, NewPipelineState);
 
-		INC_DWORD_STAT(STAT_PSOComputeNumLowlevelCacheEntries);
 		*OutPipelineState = NewPipelineState;
-
-		// Do the callback now with the lock still on.
-		PostCreateCallback(NewPipelineState, Desc);
 	}
+
+	// Create the underlying PSO and then perform any other additional tasks like cleaning up/adding to caches, etc.
+	PostCreateCallback(*OutPipelineState, Desc);
 }
 
 #if D3D12RHI_USE_HIGH_LEVEL_PSO_CACHE
diff --git a/Engine/Source/Runtime/D3D12RHI/Private/Windows/WindowsD3D12PipelineState.cpp b/Engine/Source/Runtime/D3D12RHI/Private/Windows/WindowsD3D12PipelineState.cpp
--- a/Engine/Source/Runtime/D3D12RHI/Private/Windows/WindowsD3D12PipelineState.cpp
+++ b/Engine/Source/Runtime/D3D12RHI/Private/Windows/WindowsD3D12PipelineState.cpp
@@ -121,7 +121,11 @@ void FD3D12PipelineStateCache::OnPSOCreated(FD3D12PipelineState* PipelineState,
 	PipelineState->Create(GraphicsPipelineCreationArgs(&Desc, PipelineLibrary.GetReference()));
 
 	// Save this PSO to disk cache.
-	AddToDiskCache(Desc, PipelineState);
+	if (!DiskCaches[PSO_CACHE_GRAPHICS].IsInErrorState())
+	{
+		FRWScopeLock Lock(DiskCachesCS, SLT_Write);
+		AddToDiskCache(Desc, PipelineState);
+	}
 }
 
 void FD3D12PipelineStateCache::OnPSOCreated(FD3D12PipelineState* PipelineState, const FD3D12ComputePipelineStateDesc& Desc)
@@ -130,12 +134,16 @@ void FD3D12PipelineStateCache::OnPSOCreated(FD3D12PipelineState* PipelineState,
 	PipelineState->Create(ComputePipelineCreationArgs(&Desc, PipelineLibrary.GetReference()));
 
 	// Save this PSO to disk cache.
-	AddToDiskCache(Desc, PipelineState);
+	if (!DiskCaches[PSO_CACHE_COMPUTE].IsInErrorState())
+	{
+		FRWScopeLock Lock(DiskCachesCS, SLT_Write);
+		AddToDiskCache(Desc, PipelineState);
+	}
 }
 
 void FD3D12PipelineStateCache::RebuildFromDiskCache(ID3D12RootSignature* GraphicsRootSignature, ID3D12RootSignature* ComputeRootSignature)
 {
-	FScopeLock Lock(&DiskCachesCS);
+	FRWScopeLock Lock(DiskCachesCS, SLT_Write);
 
 	UNREFERENCED_PARAMETER(GraphicsRootSignature);
 	UNREFERENCED_PARAMETER(ComputeRootSignature);
@@ -315,8 +323,6 @@ void FD3D12PipelineStateCache::RebuildFromDiskCache(ID3D12RootSignature* Graphic
 
 void FD3D12PipelineStateCache::AddToDiskCache(const FD3D12LowLevelGraphicsPipelineStateDesc& Desc, FD3D12PipelineState* PipelineState)
 {
-	FScopeLock Lock(&DiskCachesCS);
-
 	FDiskCacheInterface& DiskCache = DiskCaches[PSO_CACHE_GRAPHICS];
 	const FD3D12_GRAPHICS_PIPELINE_STATE_DESC &PsoDesc = Desc.Desc;
 
@@ -394,8 +400,6 @@ void FD3D12PipelineStateCache::AddToDiskCache(const FD3D12LowLevelGraphicsPipeli
 
 void FD3D12PipelineStateCache::AddToDiskCache(const FD3D12ComputePipelineStateDesc& Desc, FD3D12PipelineState* PipelineState)
 {
-	FScopeLock Lock(&DiskCachesCS);
-
 	FDiskCacheInterface& DiskCache = DiskCaches[PSO_CACHE_COMPUTE];
 	const FD3D12_COMPUTE_PIPELINE_STATE_DESC &PsoDesc = Desc.Desc;
 
@@ -466,7 +470,7 @@ void FD3D12PipelineStateCache::WriteOutShaderBlob(PSO_CACHE_TYPE Cache, ID3D12Pi
 
 void FD3D12PipelineStateCache::Close()
 {
-	FScopeLock Lock(&DiskCachesCS);
+	FRWScopeLock Lock(DiskCachesCS, SLT_Write);
 
 	// Write driver-optimized PSOs to the disk cache.
 	TArray<BYTE> LibraryData;
@@ -499,7 +503,7 @@ void FD3D12PipelineStateCache::Close()
 
 void FD3D12PipelineStateCache::Init(FString& GraphicsCacheFileName, FString& ComputeCacheFileName, FString& DriverBlobFileName)
 {
-	FScopeLock Lock(&DiskCachesCS);
+	FRWScopeLock Lock(DiskCachesCS, SLT_Write);
 
 	const bool bEnableGeneralPipelineStateDiskCaches = CVarPipelineStateDiskCache.GetValueOnAnyThread() != 0;
 	UE_CLOG(!bEnableGeneralPipelineStateDiskCaches, LogD3D12RHI, Display, TEXT("Not using pipeline state disk cache per r.D3D12.PSO.DiskCache=0"));
