FNM - Workaround transitioning to/from same state
#rnx

diff --git a/Engine/Source/Runtime/D3D12RHI/Private/D3D12CommandList.cpp b/Engine/Source/Runtime/D3D12RHI/Private/D3D12CommandList.cpp
--- a/Engine/Source/Runtime/D3D12RHI/Private/D3D12CommandList.cpp
+++ b/Engine/Source/Runtime/D3D12RHI/Private/D3D12CommandList.cpp
@@ -6,10 +6,17 @@
 void FD3D12CommandListHandle::AddTransitionBarrier(FD3D12Resource* pResource, D3D12_RESOURCE_STATES Before, D3D12_RESOURCE_STATES After, uint32 Subresource)
 {
 	check(CommandListData);
-	int32 NumAdded = CommandListData->ResourceBarrierBatcher.AddTransition(pResource->GetResource(), Before, After, Subresource);
-	CommandListData->CurrentOwningContext->numBarriers += NumAdded;
+	if (Before != After)
+	{
+		int32 NumAdded = CommandListData->ResourceBarrierBatcher.AddTransition(pResource->GetResource(), Before, After, Subresource);
+		CommandListData->CurrentOwningContext->numBarriers += NumAdded;
 
-	pResource->UpdateResidency(*this);
+		pResource->UpdateResidency(*this);
+	}
+	else
+	{
+		ensureMsgf(0, TEXT("AddTransitionBarrier: Before == After (%d)"), (uint32)Before);
+	}
 }
 
 void FD3D12CommandListHandle::AddUAVBarrier()
