Several fixes related to texture readback. Also duplicate CL 6971889 and 6971900 to address some asserts triggering.

diff --git a/Engine/Source/Runtime/D3D12RHI/Private/D3D12Allocation.cpp b/Engine/Source/Runtime/D3D12RHI/Private/D3D12Allocation.cpp
--- a/Engine/Source/Runtime/D3D12RHI/Private/D3D12Allocation.cpp
+++ b/Engine/Source/Runtime/D3D12RHI/Private/D3D12Allocation.cpp
@@ -287,6 +287,7 @@ void FD3D12BuddyAllocator::Allocate(uint32 SizeInBytes, uint32 Alignment, FD3D12
 	}
 #endif
 	const uint32 AlignedOffsetFromResourceBase = AllocationBlockOffset + Padding;
+	check((AlignedOffsetFromResourceBase % Alignment) == 0);
 
 	// Setup the info that this allocator
 	FD3D12BuddyAllocatorPrivateData& PrivateData = ResourceLocation.GetBuddyAllocatorPrivateData();
@@ -313,12 +314,6 @@ void FD3D12BuddyAllocator::Allocate(uint32 SizeInBytes, uint32 Alignment, FD3D12
 		// Place resources are intialized elsewhere
 	}
 
-	if (Alignment != 0)
-	{
-		check(uint64(ResourceLocation.GetMappedBaseAddress()) % Alignment == 0);
-		check(uint64(ResourceLocation.GetGPUVirtualAddress()) % Alignment == 0);
-	}
-
 	// track the allocation
 #if !PLATFORM_WINDOWS
 	LLM(FLowLevelMemTracker::Get().OnLowLevelAlloc(ELLMTracker::Default, &ResourceLocation, SizeInBytes));
diff --git a/Engine/Source/Runtime/D3D12RHI/Private/D3D12RHI.cpp b/Engine/Source/Runtime/D3D12RHI/Private/D3D12RHI.cpp
--- a/Engine/Source/Runtime/D3D12RHI/Private/D3D12RHI.cpp
+++ b/Engine/Source/Runtime/D3D12RHI/Private/D3D12RHI.cpp
@@ -163,6 +163,7 @@ FD3D12DynamicRHI::FD3D12DynamicRHI(const TArray<TSharedPtr<FD3D12Adapter>>& Chos
 	GMaxShadowDepthBufferSizeX = GMaxTextureDimensions;
 	GMaxShadowDepthBufferSizeY = GMaxTextureDimensions;
 	GRHISupportsResolveCubemapFaces = true;
+	GRHISupportsCopyToTextureMultipleMips = true;
 
 	GRHISupportsRHIThread = true;
 #if PLATFORM_XBOXONE
diff --git a/Engine/Source/Runtime/D3D12RHI/Private/D3D12RHIPrivate.h b/Engine/Source/Runtime/D3D12RHI/Private/D3D12RHIPrivate.h
--- a/Engine/Source/Runtime/D3D12RHI/Private/D3D12RHIPrivate.h
+++ b/Engine/Source/Runtime/D3D12RHI/Private/D3D12RHIPrivate.h
@@ -1020,7 +1020,7 @@ private:
 	const uint32 Subresource;
 
 public:
-	explicit FScopeResourceBarrier(FD3D12CommandListHandle& hInCommandList, FD3D12Resource* pInResource, const D3D12_RESOURCE_STATES& InCurrent, const D3D12_RESOURCE_STATES& InDesired, const uint32 InSubresource)
+	explicit FScopeResourceBarrier(FD3D12CommandListHandle& hInCommandList, FD3D12Resource* pInResource, const D3D12_RESOURCE_STATES& InCurrent, const D3D12_RESOURCE_STATES InDesired, const uint32 InSubresource)
 		: hCommandList(hInCommandList)
 		, pResource(pInResource)
 		, Current(InCurrent)
@@ -1053,7 +1053,7 @@ private:
 	const bool bUseTracking;
 
 public:
-	explicit FConditionalScopeResourceBarrier(FD3D12CommandListHandle& hInCommandList, FD3D12Resource* pInResource, const D3D12_RESOURCE_STATES& InDesired, const uint32 InSubresource)
+	explicit FConditionalScopeResourceBarrier(FD3D12CommandListHandle& hInCommandList, FD3D12Resource* pInResource, const D3D12_RESOURCE_STATES InDesired, const uint32 InSubresource)
 		: hCommandList(hInCommandList)
 		, pResource(pInResource)
 		, Current(D3D12_RESOURCE_STATE_TBD)
diff --git a/Engine/Source/Runtime/D3D12RHI/Private/D3D12RenderTarget.cpp b/Engine/Source/Runtime/D3D12RHI/Private/D3D12RenderTarget.cpp
--- a/Engine/Source/Runtime/D3D12RHI/Private/D3D12RenderTarget.cpp
+++ b/Engine/Source/Runtime/D3D12RHI/Private/D3D12RenderTarget.cpp
@@ -663,14 +663,7 @@ static uint32 ComputeBytesPerPixel(DXGI_FORMAT Format)
 
 TRefCountPtr<FD3D12Resource> FD3D12DynamicRHI::GetStagingTexture(FRHITexture* TextureRHI, FIntRect InRect, FIntRect& StagingRectOUT, FReadSurfaceDataFlags InFlags, D3D12_PLACED_SUBRESOURCE_FOOTPRINT &readbackHeapDesc)
 {
-	FD3D12Device* Device = GetRHIDevice();
-	FD3D12Adapter* Adapter = Device->GetParentAdapter();
-	const FRHIGPUMask Node = Device->GetGPUMask();
-
-	FD3D12CommandListHandle& hCommandList = Device->GetDefaultCommandContext().CommandListHandle;
 	FD3D12TextureBase* Texture = GetD3D12TextureFromRHITexture(TextureRHI);
-	D3D12_RESOURCE_DESC const& SourceDesc = Texture->GetResource()->GetDesc();
-
 	// Ensure we're dealing with a Texture2D, which the rest of this function already assumes
 	check(TextureRHI->GetTexture2D());
 	FD3D12Texture2D* InTexture2D = static_cast<FD3D12Texture2D*>(Texture);
@@ -684,12 +677,18 @@ TRefCountPtr<FD3D12Resource> FD3D12DynamicRHI::GetStagingTexture(FRHITexture* Te
 
 		// Texture2Ds on the readback heap will have been flattened to 1D, so we need to retrieve pitch
 		// information from the original 2D version to correctly use sub-rects.
-		readbackHeapDesc = InTexture2D->GetReadBackHeapDesc();
+		InTexture2D->GetReadBackHeapDesc(readbackHeapDesc, InFlags.GetMip());
 		StagingRectOUT = InRect;
 
 		return (Texture->GetResource());
 	}
 
+	FD3D12Device* Device = GetRHIDevice();
+	FD3D12Adapter* Adapter = Device->GetParentAdapter();
+	const FRHIGPUMask Node = Device->GetGPUMask();
+
+	D3D12_RESOURCE_DESC const& SourceDesc = Texture->GetResource()->GetDesc();
+
 	// a temporary staging texture is needed.
 	int32 SizeX = InRect.Width();
 	int32 SizeY = InRect.Height();
@@ -749,6 +748,7 @@ TRefCountPtr<FD3D12Resource> FD3D12DynamicRHI::GetStagingTexture(FRHITexture* Te
 	CD3DX12_TEXTURE_COPY_LOCATION DestCopyLocation(TempTexture2D->GetResource(), placedTexture2D);
 	CD3DX12_TEXTURE_COPY_LOCATION SourceCopyLocation(Texture->GetResource()->GetResource(), Subresource);
 
+	FD3D12CommandListHandle& hCommandList = Device->GetDefaultCommandContext().CommandListHandle;
 	TransitionResource(hCommandList, Texture->GetResource(), D3D12_RESOURCE_STATE_COPY_SOURCE, SourceCopyLocation.SubresourceIndex);
 	hCommandList.FlushResourceBarriers();
 	// Upload heap doesn't need to transition
@@ -791,15 +791,12 @@ void FD3D12DynamicRHI::ReadSurfaceDataNoMSAARaw(FRHITexture* TextureRHI, FIntRec
 
 	uint32 BytesPerLine = BytesPerPixel * InRect.Width();
 	const uint32 XBytesAligned = Align((uint32)readBackHeapDesc.Footprint.Width * BytesPerPixel, FD3D12_TEXTURE_DATA_PITCH_ALIGNMENT);
-	uint32 SrcStart = StagingRect.Min.X * BytesPerPixel + StagingRect.Min.Y * XBytesAligned;
-
-	// Determine the subresource index.
-	uint32 Subresource = CalcSubresource(InFlags.GetMip(), 0, TextureRHI->GetNumMips());
+	uint64 SrcStart = readBackHeapDesc.Offset + StagingRect.Min.X * BytesPerPixel + StagingRect.Min.Y * XBytesAligned;
 
 	// Lock the staging resource.
 	void* pData;
 	D3D12_RANGE ReadRange = { SrcStart, SrcStart + XBytesAligned * (SizeY - 1) + BytesPerLine };
-	VERIFYD3D12RESULT(TempTexture2D->GetResource()->Map(Subresource, &ReadRange, &pData));
+	VERIFYD3D12RESULT(TempTexture2D->GetResource()->Map(0, &ReadRange, &pData));
 
 	uint8* DestPtr = OutData.GetData();
 	uint8* SrcPtr = (uint8*)pData + SrcStart;
@@ -1693,13 +1690,14 @@ void FD3D12DynamicRHI::RHIMapStagingSurface(FRHITexture* TextureRHI, void*& OutD
 	{
 		VERIFYD3D12RESULT_EX(Result, GetAdapter().GetD3DDevice());
 
-		const D3D12_PLACED_SUBRESOURCE_FOOTPRINT& readBackHeapDesc = DestTexture2D->GetReadBackHeapDesc();
+		D3D12_PLACED_SUBRESOURCE_FOOTPRINT ReadBackHeapDesc;
+		DestTexture2D->GetReadBackHeapDesc(ReadBackHeapDesc, 0);
 		OutData = pData;
-		OutWidth = readBackHeapDesc.Footprint.RowPitch / BytesPerPixel;
-		OutHeight = readBackHeapDesc.Footprint.Height;
+		OutWidth = ReadBackHeapDesc.Footprint.RowPitch / BytesPerPixel;
+		OutHeight = ReadBackHeapDesc.Footprint.Height;
 
 		// MS: It seems like the second frame in some scenes comes into RHIMapStagingSurface BEFORE the copy to the staging texture, thus the readbackHeapDesc isn't set. This could be bug in UE4.
-		if (readBackHeapDesc.Footprint.Format != DXGI_FORMAT_UNKNOWN)
+		if (ReadBackHeapDesc.Footprint.Format != DXGI_FORMAT_UNKNOWN)
 		{
 			check(OutWidth != 0);
 			check(OutHeight != 0);
diff --git a/Engine/Source/Runtime/D3D12RHI/Private/D3D12Texture.cpp b/Engine/Source/Runtime/D3D12RHI/Private/D3D12Texture.cpp
--- a/Engine/Source/Runtime/D3D12RHI/Private/D3D12Texture.cpp
+++ b/Engine/Source/Runtime/D3D12RHI/Private/D3D12Texture.cpp
@@ -629,26 +629,20 @@ void SafeCreateTexture2D(FD3D12Device* pDevice,
 	{
 #endif // #if GUARDED_TEXTURE_CREATES
 
-		const D3D12_HEAP_TYPE heapType = (Flags & TexCreate_CPUReadback) ? D3D12_HEAP_TYPE_READBACK : D3D12_HEAP_TYPE_DEFAULT;
-		const uint64 BlockSizeX = GPixelFormats[Format].BlockSizeX;
-		const uint64 BlockSizeY = GPixelFormats[Format].BlockSizeY;
-		const uint64 BlockBytes = GPixelFormats[Format].BlockBytes;
-		const uint64 MipSizeX = FMath::Max(TextureDesc.Width, BlockSizeX);
-		const uint64 MipSizeY = FMath::Max((uint64)TextureDesc.Height, BlockSizeY);
-		const uint64 NumBlocksX = (MipSizeX + BlockSizeX - 1) / BlockSizeX;
-		const uint64 NumBlocksY = (MipSizeY + BlockSizeY - 1) / BlockSizeY;
-		const uint64 XBytesAligned = Align(NumBlocksX * BlockBytes, FD3D12_TEXTURE_DATA_PITCH_ALIGNMENT);
-		const uint64 MipBytesAligned = Align(NumBlocksY * XBytesAligned, D3D12_TEXTURE_DATA_PLACEMENT_ALIGNMENT);
-
-		switch (heapType)
+		const D3D12_HEAP_TYPE HeapType = (Flags & TexCreate_CPUReadback) ? D3D12_HEAP_TYPE_READBACK : D3D12_HEAP_TYPE_DEFAULT;
+
+		switch (HeapType)
 		{
 		case D3D12_HEAP_TYPE_READBACK:
 			{
+				uint64 Size = 0;
+				pDevice->GetDevice()->GetCopyableFootprints(&TextureDesc, 0, TextureDesc.MipLevels * TextureDesc.DepthOrArraySize, 0, nullptr, nullptr, nullptr, &Size);
+
 				FD3D12Resource* Resource = nullptr;
-				VERIFYD3D12CREATETEXTURERESULT(Adapter->CreateBuffer(heapType, pDevice->GetGPUMask(), pDevice->GetVisibilityMask(), MipBytesAligned, &Resource, Name), TextureDesc);
+				VERIFYD3D12CREATETEXTURERESULT(Adapter->CreateBuffer(HeapType, pDevice->GetGPUMask(), pDevice->GetVisibilityMask(), Size, &Resource, Name), TextureDesc);
 				OutTexture2D->AsStandAlone(Resource);
 
-				if (IsCPUWritable(heapType))
+				if (IsCPUWritable(HeapType))
 				{
 					OutTexture2D->SetMappedBaseAddress(Resource->Map());
 				}
@@ -951,26 +945,6 @@ TD3D12Texture2D<BaseResourceType>* FD3D12DynamicRHI::CreateD3D12Texture2D(FRHICo
 			}
 		}
 
-		if (Flags & TexCreate_CPUReadback)
-		{
-			const uint32 BlockBytes = GPixelFormats[Format].BlockBytes;
-			const uint32 XBytesAligned = Align((uint32)SizeX * BlockBytes, FD3D12_TEXTURE_DATA_PITCH_ALIGNMENT);
-			D3D12_SUBRESOURCE_FOOTPRINT DestSubresource;
-			DestSubresource.Depth = SizeZ;
-			DestSubresource.Height = SizeY;
-			DestSubresource.Width = SizeX;
-			DestSubresource.Format = PlatformResourceFormat;
-			DestSubresource.RowPitch = XBytesAligned;
-
-			check(DestSubresource.RowPitch % FD3D12_TEXTURE_DATA_PITCH_ALIGNMENT == 0);	// Make sure we align correctly.
-
-			D3D12_PLACED_SUBRESOURCE_FOOTPRINT PlacedTexture2D = { 0 };
-			PlacedTexture2D.Offset = 0;
-			PlacedTexture2D.Footprint = DestSubresource;
-
-			NewTexture->SetReadBackHeapDesc(PlacedTexture2D);
-		}
-
 		// Create a shader resource view for the texture.
 		if (bCreateShaderResource)
 		{
@@ -2632,26 +2606,6 @@ FTexture2DRHIRef FD3D12DynamicRHI::RHICreateTexture2DFromResource(EPixelFormat F
 		}
 	}
 
-	if (TexCreateFlags & TexCreate_CPUReadback)
-	{
-		const uint32 BlockBytes = GPixelFormats[Format].BlockBytes;
-		const uint32 XBytesAligned = Align((uint32)SizeX * BlockBytes, FD3D12_TEXTURE_DATA_PITCH_ALIGNMENT);
-		D3D12_SUBRESOURCE_FOOTPRINT DestSubresource;
-		DestSubresource.Depth = SizeZ;
-		DestSubresource.Height = SizeY;
-		DestSubresource.Width = SizeX;
-		DestSubresource.Format = PlatformResourceFormat;
-		DestSubresource.RowPitch = XBytesAligned;
-
-		check(DestSubresource.RowPitch % FD3D12_TEXTURE_DATA_PITCH_ALIGNMENT == 0);	// Make sure we align correctly.
-
-		D3D12_PLACED_SUBRESOURCE_FOOTPRINT PlacedTexture2D = { 0 };
-		PlacedTexture2D.Offset = 0;
-		PlacedTexture2D.Footprint = DestSubresource;
-
-		Texture2D->SetReadBackHeapDesc(PlacedTexture2D);
-	}
-
 	// Create a shader resource view for the texture.
 	D3D12_SHADER_RESOURCE_VIEW_DESC SRVDesc = {};
 	SRVDesc.Shader4ComponentMapping = D3D12_DEFAULT_SHADER_4_COMPONENT_MAPPING;
@@ -2800,26 +2754,6 @@ FTextureCubeRHIRef FD3D12DynamicRHI::RHICreateTextureCubeFromResource(EPixelForm
 		}
 	}
 
-	if (TexCreateFlags & TexCreate_CPUReadback)
-	{
-		const uint32 BlockBytes = GPixelFormats[Format].BlockBytes;
-		const uint32 XBytesAligned = Align((uint32)SizeX * BlockBytes, FD3D12_TEXTURE_DATA_PITCH_ALIGNMENT);
-		D3D12_SUBRESOURCE_FOOTPRINT DestSubresource;
-		DestSubresource.Depth = SizeZ;
-		DestSubresource.Height = SizeY;
-		DestSubresource.Width = SizeX;
-		DestSubresource.Format = PlatformResourceFormat;
-		DestSubresource.RowPitch = XBytesAligned;
-
-		check(DestSubresource.RowPitch % FD3D12_TEXTURE_DATA_PITCH_ALIGNMENT == 0);	// Make sure we align correctly.
-
-		D3D12_PLACED_SUBRESOURCE_FOOTPRINT PlacedTexture2D = { 0 };
-		PlacedTexture2D.Offset = 0;
-		PlacedTexture2D.Footprint = DestSubresource;
-
-		TextureCube->SetReadBackHeapDesc(PlacedTexture2D);
-	}
-
 	// Create a shader resource view for the texture.
 	D3D12_SHADER_RESOURCE_VIEW_DESC SRVDesc = {};
 	SRVDesc.Shader4ComponentMapping = D3D12_DEFAULT_SHADER_4_COMPONENT_MAPPING;
@@ -2883,16 +2817,15 @@ void FD3D12CommandContext::RHICopyTexture(FRHITexture* SourceTextureRHI, FRHITex
 	FD3D12TextureBase* SourceTexture = GetD3D12TextureFromRHITexture(SourceTextureRHI);
 	FD3D12TextureBase* DestTexture = GetD3D12TextureFromRHITexture(DestTextureRHI);
 
-	CD3DX12_TEXTURE_COPY_LOCATION SourceCopyLocation(SourceTexture->GetResource()->GetResource(), 0);
-	CD3DX12_TEXTURE_COPY_LOCATION DestCopyLocation(DestTexture->GetResource()->GetResource(), 0);
-
 	FConditionalScopeResourceBarrier ConditionalScopeResourceBarrierSource(CommandListHandle, SourceTexture->GetResource(), D3D12_RESOURCE_STATE_COPY_SOURCE, D3D12_RESOURCE_BARRIER_ALL_SUBRESOURCES);
 	FConditionalScopeResourceBarrier ConditionalScopeResourceBarrierDest(CommandListHandle, DestTexture->GetResource(), D3D12_RESOURCE_STATE_COPY_DEST, D3D12_RESOURCE_BARRIER_ALL_SUBRESOURCES);
 
 	numCopies++;
 	CommandListHandle.FlushResourceBarriers();
 
-	if (CopyInfo.Size != FIntVector::ZeroValue)
+	const bool bReadback = (DestTextureRHI->GetFlags() & TexCreate_CPUReadback) != 0;
+
+	if (CopyInfo.Size != FIntVector::ZeroValue || bReadback)
 	{
 		// Copy sub texture regions
 		CD3DX12_BOX SourceBoxD3D(
@@ -2904,29 +2837,64 @@ void FD3D12CommandContext::RHICopyTexture(FRHITexture* SourceTextureRHI, FRHITex
 			CopyInfo.SourcePosition.Z + CopyInfo.Size.Z
 		);
 
+		D3D12_TEXTURE_COPY_LOCATION Dst;
+		Dst.pResource = DestTexture->GetResource()->GetResource();
+		Dst.Type = bReadback? D3D12_TEXTURE_COPY_TYPE_PLACED_FOOTPRINT : D3D12_TEXTURE_COPY_TYPE_SUBRESOURCE_INDEX;
+
+		D3D12_TEXTURE_COPY_LOCATION Src;
+		Src.pResource = SourceTexture->GetResource()->GetResource();
+		Src.Type = D3D12_TEXTURE_COPY_TYPE_SUBRESOURCE_INDEX;
+
+		EPixelFormat Format = SourceTextureRHI->GetFormat();
+
+		uint64 Offset = 0;
 		for (uint32 SliceIndex = 0; SliceIndex < CopyInfo.NumSlices; ++SliceIndex)
 		{
 			uint32 SourceSliceIndex = CopyInfo.SourceSliceIndex + SliceIndex;
-			uint32 DestSliceIndex = CopyInfo.DestSliceIndex + SliceIndex;
+			uint32 DestSliceIndex   = CopyInfo.DestSliceIndex   + SliceIndex;
 
 			for (uint32 MipIndex = 0; MipIndex < CopyInfo.NumMips; ++MipIndex)
 			{
 				uint32 SourceMipIndex = CopyInfo.SourceMipIndex + MipIndex;
-				uint32 DestMipIndex = CopyInfo.DestMipIndex + MipIndex;
+				uint32 DestMipIndex   = CopyInfo.DestMipIndex   + MipIndex;
 
-				SourceCopyLocation.SubresourceIndex = CalcSubresource(SourceMipIndex, SourceSliceIndex, SourceTextureRHI->GetNumMips());
-				DestCopyLocation.SubresourceIndex = CalcSubresource(DestMipIndex, DestSliceIndex, DestTextureRHI->GetNumMips());
+				uint32 SizeX = FMath::Max(CopyInfo.Size.X >> MipIndex, 1);
+				uint32 SizeY = FMath::Max(CopyInfo.Size.Y >> MipIndex, 1);
+				uint32 SizeZ = FMath::Max(CopyInfo.Size.Z >> MipIndex, 1);
 
-				SourceBoxD3D.right = CopyInfo.SourcePosition.X + FMath::Max(CopyInfo.Size.X >> MipIndex, 1);
-				SourceBoxD3D.bottom = CopyInfo.SourcePosition.Y + FMath::Max(CopyInfo.Size.Y >> MipIndex, 1);
-				SourceBoxD3D.back = CopyInfo.SourcePosition.Z + FMath::Max(CopyInfo.Size.Z >> MipIndex, 1);
+				SourceBoxD3D.right  = CopyInfo.SourcePosition.X + SizeX;
+				SourceBoxD3D.bottom = CopyInfo.SourcePosition.Y + SizeY;
+				SourceBoxD3D.back   = CopyInfo.SourcePosition.Z + SizeZ;
+
+				Src.SubresourceIndex = CalcSubresource(SourceMipIndex, SourceSliceIndex, SourceTextureRHI->GetNumMips());
+
+				if (bReadback)
+				{
+					uint32 NumBlocksX = FMath::DivideAndRoundUp<uint32>(SizeX, GPixelFormats[Format].BlockSizeX);
+					uint32 NumBlocksY = FMath::DivideAndRoundUp<uint32>(SizeY, GPixelFormats[Format].BlockSizeY);
+					uint32 NumBlocksZ = FMath::DivideAndRoundUp<uint32>(SizeZ, GPixelFormats[Format].BlockSizeZ);
+					uint32 RowPitch = Align(NumBlocksX * GPixelFormats[Format].BlockBytes, FD3D12_TEXTURE_DATA_PITCH_ALIGNMENT);
+
+					Dst.PlacedFootprint.Offset             = Offset;
+					Dst.PlacedFootprint.Footprint.Format   = (DXGI_FORMAT) GPixelFormats[Format].PlatformFormat;
+					Dst.PlacedFootprint.Footprint.Width    = SizeX;
+					Dst.PlacedFootprint.Footprint.Height   = SizeY;
+					Dst.PlacedFootprint.Footprint.Depth    = SizeZ;
+					Dst.PlacedFootprint.Footprint.RowPitch = RowPitch;
+
+					Offset = Align(Offset + RowPitch * NumBlocksY * NumBlocksZ, D3D12_TEXTURE_DATA_PLACEMENT_ALIGNMENT);
+				}
+				else
+				{
+					Dst.SubresourceIndex = CalcSubresource(DestMipIndex, DestSliceIndex, DestTextureRHI->GetNumMips());
+				}
 
 				CommandListHandle->CopyTextureRegion(
-					&DestCopyLocation,
+					&Dst, 
 					CopyInfo.DestPosition.X,
 					CopyInfo.DestPosition.Y,
 					CopyInfo.DestPosition.Z,
-					&SourceCopyLocation,
+					&Src,
 					&SourceBoxD3D
 				);
 			}
@@ -2935,7 +2903,7 @@ void FD3D12CommandContext::RHICopyTexture(FRHITexture* SourceTextureRHI, FRHITex
 	else
 	{
 		// Copy whole texture
-		CommandListHandle->CopyTextureRegion(&DestCopyLocation, 0, 0, 0, &SourceCopyLocation, nullptr);
+		CommandListHandle->CopyResource(DestTexture->GetResource()->GetResource(), SourceTexture->GetResource()->GetResource());
 	}
 
 	CommandListHandle.UpdateResidency(SourceTexture->GetResource());
diff --git a/Engine/Source/Runtime/D3D12RHI/Private/D3D12Texture.h b/Engine/Source/Runtime/D3D12RHI/Private/D3D12Texture.h
--- a/Engine/Source/Runtime/D3D12RHI/Private/D3D12Texture.h
+++ b/Engine/Source/Runtime/D3D12RHI/Private/D3D12Texture.h
@@ -248,7 +248,6 @@ public:
 		, RawTextureMemory(InRawTextureMemory)
 #endif
 	{
-		FMemory::Memzero(ReadBackHeapDesc);
 		if (InTextureLayout == nullptr)
 		{
 			FMemory::Memzero(TextureLayout);
@@ -275,12 +274,33 @@ public:
 
 	// Accessors.
 	FD3D12Resource* GetResource() const { return (FD3D12Resource*)FD3D12TextureBase::GetResource(); }
-	const D3D12_PLACED_SUBRESOURCE_FOOTPRINT& GetReadBackHeapDesc() const
+
+	void GetReadBackHeapDesc(D3D12_PLACED_SUBRESOURCE_FOOTPRINT& OutFootprint, uint32 Subresource) const
 	{
-		// This should only be called if SetReadBackHeapDesc() was called with actual contents
-		check(ReadBackHeapDesc.Footprint.Width > 0 && ReadBackHeapDesc.Footprint.Height > 0);
+		check((GetFlags() & TexCreate_CPUReadback) != 0);
+
+		FIntVector Size = GetSizeXYZ();
+
+		D3D12_RESOURCE_DESC Desc = {};
+		Desc.Dimension        = D3D12_RESOURCE_DIMENSION_TEXTURE2D;
+		Desc.Width            = Size.X;
+		Desc.Height           = Size.Y;
+		Desc.DepthOrArraySize = Size.Z;
+		Desc.MipLevels        = GetNumMips();
+		Desc.Format           = (DXGI_FORMAT) GPixelFormats[GetFormat()].PlatformFormat;
+		Desc.SampleDesc.Count = GetNumSamples();
+
+		ID3D12Device* Device = GetParentDevice()->GetDevice();
+
+		uint64 Offset = 0;
+		if (Subresource > 0)
+		{
+			Device->GetCopyableFootprints(&Desc, 0, Subresource, 0, nullptr, nullptr, nullptr, &Offset);
+			Offset = Align(Offset, D3D12_TEXTURE_DATA_PLACEMENT_ALIGNMENT);
+		}
+		Device->GetCopyableFootprints(&Desc, Subresource, 1, Offset, &OutFootprint, nullptr, nullptr, nullptr);
 
-		return ReadBackHeapDesc;
+		check(OutFootprint.Footprint.Width > 0 && OutFootprint.Footprint.Height > 0);
 	}
 
 	FD3D12CLSyncPoint GetReadBackSyncPoint() const { return ReadBackSyncPoint; }
@@ -299,7 +319,6 @@ public:
 	}
 
 	// Modifiers.
-	void SetReadBackHeapDesc(const D3D12_PLACED_SUBRESOURCE_FOOTPRINT &newReadBackHeapDesc) { ReadBackHeapDesc = newReadBackHeapDesc; }
 	void SetReadBackListHandle(FD3D12CommandListHandle listToWaitFor) { ReadBackSyncPoint = listToWaitFor; }
 
 	// IRefCountedObject interface.
@@ -337,7 +356,6 @@ private:
 	/** Unlocks a previously locked mip-map. */
 	void UnlockInternal(class FRHICommandListImmediate* RHICmdList, TD3D12Texture2D* Previous, uint32 MipIndex, uint32 ArrayIndex);
 
-	D3D12_PLACED_SUBRESOURCE_FOOTPRINT ReadBackHeapDesc;
 	FD3D12CLSyncPoint ReadBackSyncPoint;
 
 	/** Whether the texture is a cube-map. */
diff --git a/Engine/Source/Runtime/D3D12RHI/Private/D3D12View.h b/Engine/Source/Runtime/D3D12RHI/Private/D3D12View.h
--- a/Engine/Source/Runtime/D3D12RHI/Private/D3D12View.h
+++ b/Engine/Source/Runtime/D3D12RHI/Private/D3D12View.h
@@ -831,7 +831,7 @@ public:
 
 		if (InDesc.ViewDimension == D3D12_SRV_DIMENSION_BUFFER)
 		{
-			check(InResourceLocation.GetOffsetFromBaseOfResource() / Stride == InDesc.Buffer.FirstElement);
+			check(InResourceLocation.GetOffsetFromBaseOfResource() == InDesc.Buffer.FirstElement * Stride);
 		}
 #endif
 
