Build fix

diff --git a/Engine/Source/Runtime/D3D12RHI/Private/D3D12Texture.cpp b/Engine/Source/Runtime/D3D12RHI/Private/D3D12Texture.cpp
--- a/Engine/Source/Runtime/D3D12RHI/Private/D3D12Texture.cpp
+++ b/Engine/Source/Runtime/D3D12RHI/Private/D3D12Texture.cpp
@@ -70,6 +70,10 @@ template void FD3D12TextureStats::D3D12TextureDeleted(TD3D12Texture2D<FD3D12Base
 template void FD3D12TextureStats::D3D12TextureDeleted(TD3D12Texture2D<FD3D12BaseTexture2DArray>& Texture);
 template void FD3D12TextureStats::D3D12TextureDeleted(TD3D12Texture2D<FD3D12BaseTextureCube>& Texture);
 
+template void TD3D12Texture2D<FD3D12BaseTexture2D>::GetReadBackHeapDesc(D3D12_PLACED_SUBRESOURCE_FOOTPRINT& OutFootprint, uint32 Subresource) const;
+template void TD3D12Texture2D<FD3D12BaseTexture2DArray>::GetReadBackHeapDesc(D3D12_PLACED_SUBRESOURCE_FOOTPRINT& OutFootprint, uint32 Subresource) const;
+template void TD3D12Texture2D<FD3D12BaseTextureCube>::GetReadBackHeapDesc(D3D12_PLACED_SUBRESOURCE_FOOTPRINT& OutFootprint, uint32 Subresource) const;
+
 /// @endcond
 
 struct FRHICommandUpdateTexture final : public FRHICommand<FRHICommandUpdateTexture>
@@ -1927,6 +1931,35 @@ void TD3D12Texture2D<RHIResourceType>::UpdateTexture2D(class FRHICommandListImme
 	}
 }
 
+template<typename RHIResourceType>
+void TD3D12Texture2D<RHIResourceType>::GetReadBackHeapDesc(D3D12_PLACED_SUBRESOURCE_FOOTPRINT& OutFootprint, uint32 Subresource) const
+{
+	check((GetFlags() & TexCreate_CPUReadback) != 0);
+
+	FIntVector TextureSize = GetSizeXYZ();
+
+	D3D12_RESOURCE_DESC Desc = {};
+	Desc.Dimension        = D3D12_RESOURCE_DIMENSION_TEXTURE2D;
+	Desc.Width            = TextureSize.X;
+	Desc.Height           = TextureSize.Y;
+	Desc.DepthOrArraySize = TextureSize.Z;
+	Desc.MipLevels        = GetNumMips();
+	Desc.Format           = (DXGI_FORMAT) GPixelFormats[GetFormat()].PlatformFormat;
+	Desc.SampleDesc.Count = GetNumSamples();
+
+	ID3D12Device* Device = GetParentDevice()->GetDevice();
+
+	uint64 Offset = 0;
+	if (Subresource > 0)
+	{
+		Device->GetCopyableFootprints(&Desc, 0, Subresource, 0, nullptr, nullptr, nullptr, &Offset);
+		Offset = Align(Offset, D3D12_TEXTURE_DATA_PLACEMENT_ALIGNMENT);
+	}
+	Device->GetCopyableFootprints(&Desc, Subresource, 1, Offset, &OutFootprint, nullptr, nullptr, nullptr);
+
+	check(OutFootprint.Footprint.Width > 0 && OutFootprint.Footprint.Height > 0);
+}
+
 void* FD3D12DynamicRHI::LockTexture2D_RenderThread(class FRHICommandListImmediate& RHICmdList, FRHITexture2D* TextureRHI, uint32 MipIndex, EResourceLockMode LockMode, uint32& DestStride, bool bLockWithinMiptail, bool bNeedsDefaultRHIFlush)
 {
 	if (CVarD3D12Texture2DRHIFlush.GetValueOnRenderThread() && bNeedsDefaultRHIFlush)
diff --git a/Engine/Source/Runtime/D3D12RHI/Private/D3D12Texture.h b/Engine/Source/Runtime/D3D12RHI/Private/D3D12Texture.h
--- a/Engine/Source/Runtime/D3D12RHI/Private/D3D12Texture.h
+++ b/Engine/Source/Runtime/D3D12RHI/Private/D3D12Texture.h
@@ -275,33 +275,7 @@ public:
 	// Accessors.
 	FD3D12Resource* GetResource() const { return (FD3D12Resource*)FD3D12TextureBase::GetResource(); }
 
-	void GetReadBackHeapDesc(D3D12_PLACED_SUBRESOURCE_FOOTPRINT& OutFootprint, uint32 Subresource) const
-	{
-		check((GetFlags() & TexCreate_CPUReadback) != 0);
-
-		FIntVector Size = GetSizeXYZ();
-
-		D3D12_RESOURCE_DESC Desc = {};
-		Desc.Dimension        = D3D12_RESOURCE_DIMENSION_TEXTURE2D;
-		Desc.Width            = Size.X;
-		Desc.Height           = Size.Y;
-		Desc.DepthOrArraySize = Size.Z;
-		Desc.MipLevels        = GetNumMips();
-		Desc.Format           = (DXGI_FORMAT) GPixelFormats[GetFormat()].PlatformFormat;
-		Desc.SampleDesc.Count = GetNumSamples();
-
-		ID3D12Device* Device = GetParentDevice()->GetDevice();
-
-		uint64 Offset = 0;
-		if (Subresource > 0)
-		{
-			Device->GetCopyableFootprints(&Desc, 0, Subresource, 0, nullptr, nullptr, nullptr, &Offset);
-			Offset = Align(Offset, D3D12_TEXTURE_DATA_PLACEMENT_ALIGNMENT);
-		}
-		Device->GetCopyableFootprints(&Desc, Subresource, 1, Offset, &OutFootprint, nullptr, nullptr, nullptr);
-
-		check(OutFootprint.Footprint.Width > 0 && OutFootprint.Footprint.Height > 0);
-	}
+	void GetReadBackHeapDesc(D3D12_PLACED_SUBRESOURCE_FOOTPRINT& OutFootprint, uint32 Subresource) const;
 
 	FD3D12CLSyncPoint GetReadBackSyncPoint() const { return ReadBackSyncPoint; }
 	bool IsCubemap() const { return bCubemap; }
