#UE4 Non-unity fix

diff --git a/Engine/Source/Runtime/D3D12RHI/Private/D3D12Texture.cpp b/Engine/Source/Runtime/D3D12RHI/Private/D3D12Texture.cpp
--- a/Engine/Source/Runtime/D3D12RHI/Private/D3D12Texture.cpp
+++ b/Engine/Source/Runtime/D3D12RHI/Private/D3D12Texture.cpp
@@ -5,7 +5,7 @@
 	=============================================================================*/
 
 #include "D3D12RHIPrivate.h"
-
+#include "RHIResources.h"
 
 int64 FD3D12GlobalStats::GDedicatedVideoMemory = 0;
 int64 FD3D12GlobalStats::GDedicatedSystemMemory = 0;
