#UE4 Fix for strict mode compilling (wasnt actually a non-unity issue, so removing unncessary include)

diff --git a/Engine/Source/Runtime/D3D12RHI/Private/D3D12Texture.cpp b/Engine/Source/Runtime/D3D12RHI/Private/D3D12Texture.cpp
--- a/Engine/Source/Runtime/D3D12RHI/Private/D3D12Texture.cpp
+++ b/Engine/Source/Runtime/D3D12RHI/Private/D3D12Texture.cpp
@@ -5,7 +5,6 @@
 	=============================================================================*/
 
 #include "D3D12RHIPrivate.h"
-#include "RHIResources.h"
 
 int64 FD3D12GlobalStats::GDedicatedVideoMemory = 0;
 int64 FD3D12GlobalStats::GDedicatedSystemMemory = 0;
@@ -1934,18 +1933,18 @@ void TD3D12Texture2D<RHIResourceType>::UpdateTexture2D(class FRHICommandListImme
 template<typename RHIResourceType>
 void TD3D12Texture2D<RHIResourceType>::GetReadBackHeapDesc(D3D12_PLACED_SUBRESOURCE_FOOTPRINT& OutFootprint, uint32 Subresource) const
 {
-	check((GetFlags() & TexCreate_CPUReadback) != 0);
+	check((RHIResourceType::GetFlags() & TexCreate_CPUReadback) != 0);
 
-	FIntVector TextureSize = GetSizeXYZ();
+	FIntVector TextureSize = RHIResourceType::GetSizeXYZ();
 
 	D3D12_RESOURCE_DESC Desc = {};
 	Desc.Dimension        = D3D12_RESOURCE_DIMENSION_TEXTURE2D;
 	Desc.Width            = TextureSize.X;
 	Desc.Height           = TextureSize.Y;
 	Desc.DepthOrArraySize = TextureSize.Z;
-	Desc.MipLevels        = GetNumMips();
-	Desc.Format           = (DXGI_FORMAT) GPixelFormats[GetFormat()].PlatformFormat;
-	Desc.SampleDesc.Count = GetNumSamples();
+	Desc.MipLevels        = RHIResourceType::GetNumMips();
+	Desc.Format           = (DXGI_FORMAT) GPixelFormats[RHIResourceType::GetFormat()].PlatformFormat;
+	Desc.SampleDesc.Count = RHIResourceType::GetNumSamples();
 
 	ID3D12Device* Device = GetParentDevice()->GetDevice();
 
