[AMD] Fixing DX12 startup crash caused by incorrect resource state transition on ClearTransluscentVolumeLightingAsyncCompute (D3D12_RESOURCE_STATE_PIXEL_SHADER_RESOURCE cannot be used with a Compute CmdList)

diff --git a/Engine/Source/Runtime/D3D12RHI/Private/D3D12Commands.cpp b/Engine/Source/Runtime/D3D12RHI/Private/D3D12Commands.cpp
--- a/Engine/Source/Runtime/D3D12RHI/Private/D3D12Commands.cpp
+++ b/Engine/Source/Runtime/D3D12RHI/Private/D3D12Commands.cpp
@@ -291,12 +291,20 @@ void FD3D12CommandContext::RHITransitionResources(EResourceTransitionAccess Tran
 #if USE_D3D12RHI_RESOURCE_STATE_TRACKING
 			if ( TransitionType == EResourceTransitionAccess::EReadable )
 			{
+				const D3D12_COMMAND_LIST_TYPE CmdListType = CommandListHandle.GetCommandListType();
 				for (int32 i = 0; i < InNumUAVs; ++i)
 				{
 					if (InUAVs[i])
 					{
 						FD3D12UnorderedAccessView* const UnorderedAccessView = RetrieveObject<FD3D12UnorderedAccessView>(InUAVs[i]);
-						FD3D12DynamicRHI::TransitionResource(CommandListHandle, UnorderedAccessView, D3D12_RESOURCE_STATE_PIXEL_SHADER_RESOURCE | D3D12_RESOURCE_STATE_NON_PIXEL_SHADER_RESOURCE);
+						
+						// D3D12_RESOURCE_STATE_PIXEL_SHADER_RESOURCE cannot be used for compute command list so we exclude it here.
+						D3D12_RESOURCE_STATES AfterState = D3D12_RESOURCE_STATE_NON_PIXEL_SHADER_RESOURCE;
+						if (CmdListType != D3D12_COMMAND_LIST_TYPE::D3D12_COMMAND_LIST_TYPE_COMPUTE)
+						{
+							AfterState |= D3D12_RESOURCE_STATE_PIXEL_SHADER_RESOURCE;
+						}
+						FD3D12DynamicRHI::TransitionResource(CommandListHandle, UnorderedAccessView, AfterState);
 					}
 				}
 			}
