Add -RenderOffScreen support to DX12

diff --git a/Engine/Source/Runtime/D3D12RHI/Private/D3D12Viewport.cpp b/Engine/Source/Runtime/D3D12RHI/Private/D3D12Viewport.cpp
--- a/Engine/Source/Runtime/D3D12RHI/Private/D3D12Viewport.cpp
+++ b/Engine/Source/Runtime/D3D12RHI/Private/D3D12Viewport.cpp
@@ -176,16 +176,45 @@ void FD3D12FramePacing::PrePresentQueued(ID3D12CommandQueue* Queue)
 }
 #endif //WITH_MGPU
 
+// TODO: Move this bool into D3D12Viewport.h where it belongs. It's here because it was added as a hotfix for 4.23 and we don't want to touch public headers.
+// Whether to create swap chain and use swap chain's back buffer surface,
+// or don't create swap chain and create an off-screen back buffer surface.
+// Currently used for pixel streaming plugin "windowless" mode to run in the cloud without on screen display.
+bool bNeedSwapChain = true;
+
 /**
  * Creates a FD3D12Surface to represent a swap chain's back buffer.
  */
-FD3D12Texture2D* GetSwapChainSurface(FD3D12Device* Parent, EPixelFormat PixelFormat, IDXGISwapChain* SwapChain, uint32 BackBufferIndex)
+FD3D12Texture2D* GetSwapChainSurface(FD3D12Device* Parent, EPixelFormat PixelFormat, uint32 SizeX, uint32 SizeY, IDXGISwapChain* SwapChain, uint32 BackBufferIndex)
 {
 	FD3D12Adapter* Adapter = Parent->GetParentAdapter();
 
 	// Grab the back buffer
 	TRefCountPtr<ID3D12Resource> BackBufferResource;
-	VERIFYD3D12RESULT_EX(SwapChain->GetBuffer(BackBufferIndex, IID_PPV_ARGS(BackBufferResource.GetInitReference())), Parent->GetDevice());
+	if (SwapChain)
+	{
+		VERIFYD3D12RESULT_EX(SwapChain->GetBuffer(BackBufferIndex, IID_PPV_ARGS(BackBufferResource.GetInitReference())), Parent->GetDevice());
+	}
+	else
+	{
+		const D3D12_HEAP_PROPERTIES HeapProps = CD3DX12_HEAP_PROPERTIES(D3D12_HEAP_TYPE_DEFAULT, (uint32)Parent->GetGPUIndex(), (uint32)Parent->GetGPUMask());
+
+		// Create custom back buffer texture as no swap chain is created in pixel streaming windowless mode
+		D3D12_RESOURCE_DESC TextureDesc;
+		TextureDesc.Dimension = D3D12_RESOURCE_DIMENSION_TEXTURE2D;
+		TextureDesc.Alignment = 0;
+		TextureDesc.Width  = SizeX;
+		TextureDesc.Height = SizeY;
+		TextureDesc.DepthOrArraySize = 1;
+		TextureDesc.MipLevels = 1;
+		TextureDesc.Format = GetRenderTargetFormat(PixelFormat);
+		TextureDesc.SampleDesc.Count = 1;
+		TextureDesc.SampleDesc.Quality = 0;
+		TextureDesc.Layout = D3D12_TEXTURE_LAYOUT_UNKNOWN;
+		TextureDesc.Flags = D3D12_RESOURCE_FLAG_ALLOW_RENDER_TARGET;
+
+		Parent->GetDevice()->CreateCommittedResource(&HeapProps, D3D12_HEAP_FLAG_NONE, &TextureDesc, D3D12_RESOURCE_STATE_PRESENT, nullptr, IID_PPV_ARGS(BackBufferResource.GetInitReference()));
+	}
 
 	D3D12_RESOURCE_DESC BackBufferDesc = BackBufferResource->GetDesc();
 
@@ -432,12 +461,15 @@ void FD3D12Viewport::Resize(uint32 InSizeX, uint32 InSizeY, bool bInIsFullscreen
 		check(SizeX > 0);
 		check(SizeY > 0);
 
-		if (bInIsFullscreen)
+		if (bNeedSwapChain)
 		{
-			const DXGI_MODE_DESC BufferDesc = SetupDXGI_MODE_DESC();
-			if (FAILED(SwapChain1->ResizeTarget(&BufferDesc)))
+			if (bInIsFullscreen)
 			{
-				ConditionalResetSwapChain(true);
+				const DXGI_MODE_DESC BufferDesc = SetupDXGI_MODE_DESC();
+				if (FAILED(SwapChain1->ResizeTarget(&BufferDesc)))
+				{
+					ConditionalResetSwapChain(true);
+				}
 			}
 		}
 	}
@@ -447,9 +479,12 @@ void FD3D12Viewport::Resize(uint32 InSizeX, uint32 InSizeY, bool bInIsFullscreen
 		bIsFullscreen = bInIsFullscreen;
 		bIsValid = false;
 
-		// Use ConditionalResetSwapChain to call SetFullscreenState, to handle the failure case.
-		// Ignore the viewport's focus state; since Resize is called as the result of a user action we assume authority without waiting for Focus.
-		ConditionalResetSwapChain(true);
+		if (bNeedSwapChain)
+		{
+			// Use ConditionalResetSwapChain to call SetFullscreenState, to handle the failure case.
+			// Ignore the viewport's focus state; since Resize is called as the result of a user action we assume authority without waiting for Focus.
+			ConditionalResetSwapChain(true);
+		}
 	}
 
 	ResizeInternal();
diff --git a/Engine/Source/Runtime/D3D12RHI/Private/Windows/WindowsD3D12Viewport.cpp b/Engine/Source/Runtime/D3D12RHI/Private/Windows/WindowsD3D12Viewport.cpp
--- a/Engine/Source/Runtime/D3D12RHI/Private/Windows/WindowsD3D12Viewport.cpp
+++ b/Engine/Source/Runtime/D3D12RHI/Private/Windows/WindowsD3D12Viewport.cpp
@@ -14,7 +14,7 @@
 
 static const uint32 WindowsDefaultNumBackBuffers = 3;
 
-extern FD3D12Texture2D* GetSwapChainSurface(FD3D12Device* Parent, EPixelFormat PixelFormat, IDXGISwapChain* SwapChain, uint32 BackBufferIndex);
+extern FD3D12Texture2D* GetSwapChainSurface(FD3D12Device* Parent, EPixelFormat PixelFormat, uint32 SizeX, uint32 SizeY, IDXGISwapChain* SwapChain, uint32 BackBufferIndex);
 
 FD3D12Viewport::FD3D12Viewport(class FD3D12Adapter* InParent, HWND InWindowHandle, uint32 InSizeX, uint32 InSizeY, bool bInIsFullscreen, EPixelFormat InPreferredPixelFormat) :
 	FD3D12AdapterChild(InParent),
@@ -122,6 +122,11 @@ void FD3D12Viewport::Init()
 	}
 #else
 // : END HoloLens support
+
+	extern bool bNeedSwapChain;
+
+	bNeedSwapChain = !FParse::Param(FCommandLine::Get(), TEXT("RenderOffScreen"));
+	if (bNeedSwapChain)
 	{
 		if (Adapter->GetOwningRHI()->IsQuadBufferStereoEnabled())
 		{
@@ -280,9 +285,12 @@ void FD3D12Viewport::ResizeInternal()
 			NodeMasks.Add(Device->GetGPUMask().GetNative());
 		}
 
-		TRefCountPtr<IDXGISwapChain3> SwapChain3;
-		VERIFYD3D12RESULT(SwapChain1->QueryInterface(IID_PPV_ARGS(SwapChain3.GetInitReference())));
-		VERIFYD3D12RESULT_EX(SwapChain3->ResizeBuffers1(NumBackBuffers, SizeX, SizeY, GetRenderTargetFormat(PixelFormat), SwapChainFlags, NodeMasks.GetData(), (IUnknown**)CommandQueues.GetData()), Adapter->GetD3DDevice());
+		if (SwapChain1)
+		{
+			TRefCountPtr<IDXGISwapChain3> SwapChain3;
+			VERIFYD3D12RESULT(SwapChain1->QueryInterface(IID_PPV_ARGS(SwapChain3.GetInitReference())));
+			VERIFYD3D12RESULT_EX(SwapChain3->ResizeBuffers1(NumBackBuffers, SizeX, SizeY, GetRenderTargetFormat(PixelFormat), SwapChainFlags, NodeMasks.GetData(), (IUnknown**)CommandQueues.GetData()), Adapter->GetD3DDevice());
+		}
 
 		for (uint32 i = 0; i < NumBackBuffers; ++i)
 		{
@@ -290,19 +298,22 @@ void FD3D12Viewport::ResizeInternal()
 			FD3D12Device* Device = Adapter->GetDevice(GPUIndex);
 
 			check(BackBuffers[i].GetReference() == nullptr);
-			BackBuffers[i] = GetSwapChainSurface(Device, PixelFormat, SwapChain1, i);
+			BackBuffers[i] = GetSwapChainSurface(Device, PixelFormat, SizeX, SizeY, SwapChain1, i);
 		}
 	}
 	else
 #endif // WITH_MGPU
 	{
-		VERIFYD3D12RESULT_EX(SwapChain1->ResizeBuffers(NumBackBuffers, SizeX, SizeY, GetRenderTargetFormat(PixelFormat), SwapChainFlags), Adapter->GetD3DDevice());
+		if (SwapChain1)
+		{
+			VERIFYD3D12RESULT_EX(SwapChain1->ResizeBuffers(NumBackBuffers, SizeX, SizeY, GetRenderTargetFormat(PixelFormat), SwapChainFlags), Adapter->GetD3DDevice());
+		}
 
 		FD3D12Device* Device = Adapter->GetDevice(0);
 		for (uint32 i = 0; i < NumBackBuffers; ++i)
 		{
 			check(BackBuffers[i].GetReference() == nullptr);
-			BackBuffers[i] = GetSwapChainSurface(Device, PixelFormat, SwapChain1, i);
+			BackBuffers[i] = GetSwapChainSurface(Device, PixelFormat, SizeX, SizeY, SwapChain1, i);
 		}
 	}
 
@@ -326,7 +337,12 @@ HRESULT FD3D12Viewport::PresentInternal(int32 SyncInterval)
 
 	FThreadHeartBeat::Get().PresentFrame();
 
-	return SwapChain1->Present(SyncInterval, Flags);
+	if (SwapChain1)
+	{
+		return SwapChain1->Present(SyncInterval, Flags);
+	}
+
+	return S_OK;
 }
 
 void FD3D12Viewport::EnableHDR()
