Build fix.

diff --git a/Engine/Source/Runtime/D3D12RHI/Private/D3D12Viewport.cpp b/Engine/Source/Runtime/D3D12RHI/Private/D3D12Viewport.cpp
--- a/Engine/Source/Runtime/D3D12RHI/Private/D3D12Viewport.cpp
+++ b/Engine/Source/Runtime/D3D12RHI/Private/D3D12Viewport.cpp
@@ -197,7 +197,7 @@ FD3D12Texture2D* GetSwapChainSurface(FD3D12Device* Parent, EPixelFormat PixelFor
 	}
 	else
 	{
-		const D3D12_HEAP_PROPERTIES HeapProps = CD3DX12_HEAP_PROPERTIES(D3D12_HEAP_TYPE_DEFAULT, (uint32)Parent->GetGPUIndex(), (uint32)Parent->GetGPUMask());
+		const D3D12_HEAP_PROPERTIES HeapProps = CD3DX12_HEAP_PROPERTIES(D3D12_HEAP_TYPE_DEFAULT, (uint32)Parent->GetGPUIndex(), Parent->GetGPUMask().GetNative());
 
 		// Create custom back buffer texture as no swap chain is created in pixel streaming windowless mode
 		D3D12_RESOURCE_DESC TextureDesc;
