Optimize DX12 draw calls a bit

diff --git a/Engine/Source/Runtime/D3D12RHI/Private/D3D12CommandContext.cpp b/Engine/Source/Runtime/D3D12RHI/Private/D3D12CommandContext.cpp
--- a/Engine/Source/Runtime/D3D12RHI/Private/D3D12CommandContext.cpp
+++ b/Engine/Source/Runtime/D3D12RHI/Private/D3D12CommandContext.cpp
@@ -77,6 +77,7 @@ static uint32 GetConstantAllocatorSize(bool bIsAsyncComputeContext, bool bIsDefa
 FD3D12CommandContextBase::FD3D12CommandContextBase(class FD3D12Adapter* InParentAdapter, FRHIGPUMask InGPUMask, bool InIsDefaultContext, bool InIsAsyncComputeContext)
 	: FD3D12AdapterChild(InParentAdapter)
 	, GPUMask(InGPUMask)
+	, bTrackingEvents(false)
 	, bIsDefaultContext(InIsDefaultContext)
 	, bIsAsyncComputeContext(InIsAsyncComputeContext)
 {
@@ -115,7 +116,7 @@ FD3D12CommandContext::FD3D12CommandContext(FD3D12Device* InParent, FD3D12SubAllo
 {
 	FMemory::Memzero(DirtyUniformBuffers);
 	FMemory::Memzero(BoundUniformBuffers);
-	for (int i = 0; i < UE_ARRAY_COUNT(BoundUniformBufferRefs); i++)
+	for (int i = 0; i < ARRAY_COUNT(BoundUniformBufferRefs); i++)
 	{
 		for (int j = 0; j < UE_ARRAY_COUNT(BoundUniformBufferRefs[i]); j++)
 		{
@@ -232,6 +233,7 @@ void FD3D12CommandContext::OpenCommandList()
 	// Mark state as dirty so next time ApplyState is called, it will set all state on this new command list
 	StateCache.DirtyStateForNewCommandList();
 
+	numPrimitives = 0;
 	numDraws = 0;
 	numDispatches = 0;
 	numClears = 0;
@@ -243,6 +245,10 @@ void FD3D12CommandContext::OpenCommandList()
 void FD3D12CommandContext::CloseCommandList()
 {
 	CommandListHandle.Close();
+
+	INC_DWORD_STAT_BY(STAT_RHIDrawPrimitiveCalls, numDraws);
+	FPlatformAtomics::InterlockedAdd(&GCurrentNumDrawCallsRHI, numDraws);
+	FPlatformAtomics::InterlockedAdd(&GCurrentNumPrimitivesDrawnRHI, numPrimitives);
 }
 
 FD3D12CommandListHandle FD3D12CommandContext::FlushCommands(bool WaitForCompletion, EFlushCommandsExtraAction ExtraAction)
@@ -319,6 +325,8 @@ void FD3D12CommandContext::Finish(TArray<FD3D12CommandListHandle>& CommandLists)
 
 void FD3D12CommandContextBase::RHIBeginFrame()
 {
+	bTrackingEvents = bIsDefaultContext && ParentAdapter->GetGPUProfiler().bTrackingEvents;
+
 	RHIPrivateBeginFrame();
 	for (uint32 GPUIndex : GPUMask)
 	{
diff --git a/Engine/Source/Runtime/D3D12RHI/Private/D3D12CommandContext.h b/Engine/Source/Runtime/D3D12RHI/Private/D3D12CommandContext.h
--- a/Engine/Source/Runtime/D3D12RHI/Private/D3D12CommandContext.h
+++ b/Engine/Source/Runtime/D3D12RHI/Private/D3D12CommandContext.h
@@ -59,6 +59,7 @@ protected:
 
 	FRHIGPUMask GPUMask;
 
+	bool bTrackingEvents;
 	const bool bIsDefaultContext;
 	const bool bIsAsyncComputeContext;
 };
@@ -175,6 +176,7 @@ public:
 	virtual void FlushTextureCache() {};
 #endif
 
+	uint32 numPrimitives;
 	uint32 numDraws;
 	uint32 numDispatches;
 	uint32 numClears;
diff --git a/Engine/Source/Runtime/D3D12RHI/Private/D3D12Commands.cpp b/Engine/Source/Runtime/D3D12RHI/Private/D3D12Commands.cpp
--- a/Engine/Source/Runtime/D3D12RHI/Private/D3D12Commands.cpp
+++ b/Engine/Source/Runtime/D3D12RHI/Private/D3D12Commands.cpp
@@ -466,87 +466,12 @@ void FD3D12CommandContext::RHISetScissorRect(bool bEnable, uint32 MinX, uint32 M
 	}
 }
 
-static inline D3D_PRIMITIVE_TOPOLOGY GetD3D12PrimitiveType(EPrimitiveType PrimitiveType, bool bUsingTessellation)
-{
-	if (bUsingTessellation)
-	{
-		switch (PrimitiveType)
-		{
-		case PT_1_ControlPointPatchList: return D3D_PRIMITIVE_TOPOLOGY_1_CONTROL_POINT_PATCHLIST;
-		case PT_2_ControlPointPatchList: return D3D_PRIMITIVE_TOPOLOGY_2_CONTROL_POINT_PATCHLIST;
-
-			// This is the case for tessellation without AEN or other buffers, so just flip to 3 CPs
-		case PT_TriangleList: return D3D_PRIMITIVE_TOPOLOGY_3_CONTROL_POINT_PATCHLIST;
-
-		case PT_LineList:
-		case PT_TriangleStrip:
-		case PT_QuadList:
-		case PT_RectList:
-		case PT_PointList:
-			UE_LOG(LogD3D12RHI, Fatal, TEXT("Invalid type specified for tessellated render, probably missing a case in FSkeletalMeshSceneProxy::DrawDynamicElementsByMaterial or FStaticMeshSceneProxy::GetMeshElement"));
-			break;
-		default:
-			// Other cases are valid.
-			break;
-		};
-	}
-
-	switch (PrimitiveType)
-	{
-	case PT_TriangleList: return D3D_PRIMITIVE_TOPOLOGY_TRIANGLELIST;
-	case PT_TriangleStrip: return D3D_PRIMITIVE_TOPOLOGY_TRIANGLESTRIP;
-	case PT_LineList: return D3D_PRIMITIVE_TOPOLOGY_LINELIST;
-	case PT_PointList: return D3D_PRIMITIVE_TOPOLOGY_POINTLIST;
-	#if defined(D3D12RHI_PRIMITIVE_TOPOLOGY_RECTLIST)
-		case PT_RectList: return D3D12RHI_PRIMITIVE_TOPOLOGY_RECTLIST;
-	#endif
-
-		// ControlPointPatchList types will pretend to be TRIANGLELISTS with a stride of N 
-		// (where N is the number of control points specified), so we can return them for
-		// tessellation and non-tessellation. This functionality is only used when rendering a 
-		// default material with something that claims to be tessellated, generally because the 
-		// tessellation material failed to compile for some reason.
-	case PT_3_ControlPointPatchList: return D3D_PRIMITIVE_TOPOLOGY_3_CONTROL_POINT_PATCHLIST;
-	case PT_4_ControlPointPatchList: return D3D_PRIMITIVE_TOPOLOGY_4_CONTROL_POINT_PATCHLIST;
-	case PT_5_ControlPointPatchList: return D3D_PRIMITIVE_TOPOLOGY_5_CONTROL_POINT_PATCHLIST;
-	case PT_6_ControlPointPatchList: return D3D_PRIMITIVE_TOPOLOGY_6_CONTROL_POINT_PATCHLIST;
-	case PT_7_ControlPointPatchList: return D3D_PRIMITIVE_TOPOLOGY_7_CONTROL_POINT_PATCHLIST;
-	case PT_8_ControlPointPatchList: return D3D_PRIMITIVE_TOPOLOGY_8_CONTROL_POINT_PATCHLIST;
-	case PT_9_ControlPointPatchList: return D3D_PRIMITIVE_TOPOLOGY_9_CONTROL_POINT_PATCHLIST;
-	case PT_10_ControlPointPatchList: return D3D_PRIMITIVE_TOPOLOGY_10_CONTROL_POINT_PATCHLIST;
-	case PT_11_ControlPointPatchList: return D3D_PRIMITIVE_TOPOLOGY_11_CONTROL_POINT_PATCHLIST;
-	case PT_12_ControlPointPatchList: return D3D_PRIMITIVE_TOPOLOGY_12_CONTROL_POINT_PATCHLIST;
-	case PT_13_ControlPointPatchList: return D3D_PRIMITIVE_TOPOLOGY_13_CONTROL_POINT_PATCHLIST;
-	case PT_14_ControlPointPatchList: return D3D_PRIMITIVE_TOPOLOGY_14_CONTROL_POINT_PATCHLIST;
-	case PT_15_ControlPointPatchList: return D3D_PRIMITIVE_TOPOLOGY_15_CONTROL_POINT_PATCHLIST;
-	case PT_16_ControlPointPatchList: return D3D_PRIMITIVE_TOPOLOGY_16_CONTROL_POINT_PATCHLIST;
-	case PT_17_ControlPointPatchList: return D3D_PRIMITIVE_TOPOLOGY_17_CONTROL_POINT_PATCHLIST;
-	case PT_18_ControlPointPatchList: return D3D_PRIMITIVE_TOPOLOGY_18_CONTROL_POINT_PATCHLIST;
-	case PT_19_ControlPointPatchList: return D3D_PRIMITIVE_TOPOLOGY_19_CONTROL_POINT_PATCHLIST;
-	case PT_20_ControlPointPatchList: return D3D_PRIMITIVE_TOPOLOGY_20_CONTROL_POINT_PATCHLIST;
-	case PT_21_ControlPointPatchList: return D3D_PRIMITIVE_TOPOLOGY_21_CONTROL_POINT_PATCHLIST;
-	case PT_22_ControlPointPatchList: return D3D_PRIMITIVE_TOPOLOGY_22_CONTROL_POINT_PATCHLIST;
-	case PT_23_ControlPointPatchList: return D3D_PRIMITIVE_TOPOLOGY_23_CONTROL_POINT_PATCHLIST;
-	case PT_24_ControlPointPatchList: return D3D_PRIMITIVE_TOPOLOGY_24_CONTROL_POINT_PATCHLIST;
-	case PT_25_ControlPointPatchList: return D3D_PRIMITIVE_TOPOLOGY_25_CONTROL_POINT_PATCHLIST;
-	case PT_26_ControlPointPatchList: return D3D_PRIMITIVE_TOPOLOGY_26_CONTROL_POINT_PATCHLIST;
-	case PT_27_ControlPointPatchList: return D3D_PRIMITIVE_TOPOLOGY_27_CONTROL_POINT_PATCHLIST;
-	case PT_28_ControlPointPatchList: return D3D_PRIMITIVE_TOPOLOGY_28_CONTROL_POINT_PATCHLIST;
-	case PT_29_ControlPointPatchList: return D3D_PRIMITIVE_TOPOLOGY_29_CONTROL_POINT_PATCHLIST;
-	case PT_30_ControlPointPatchList: return D3D_PRIMITIVE_TOPOLOGY_30_CONTROL_POINT_PATCHLIST;
-	case PT_31_ControlPointPatchList: return D3D_PRIMITIVE_TOPOLOGY_31_CONTROL_POINT_PATCHLIST;
-	case PT_32_ControlPointPatchList: return D3D_PRIMITIVE_TOPOLOGY_32_CONTROL_POINT_PATCHLIST;
-	default: UE_LOG(LogD3D12RHI, Fatal, TEXT("Unknown primitive type: %u"), PrimitiveType);
-	};
-
-	return D3D_PRIMITIVE_TOPOLOGY_TRIANGLELIST;
-}
-
 void FD3D12CommandContext::RHISetGraphicsPipelineState(FRHIGraphicsPipelineState* GraphicsState)
 {
 	FD3D12GraphicsPipelineState* GraphicsPipelineState = FD3D12DynamicRHI::ResourceCast(GraphicsState);
 
 	// TODO: [PSO API] Every thing inside this scope is only necessary to keep the PSO shadow in sync while we convert the high level to only use PSOs
+	const bool bWasUsingTessellation = bUsingTessellation;
 	bUsingTessellation = GraphicsPipelineState->GetHullShader() && GraphicsPipelineState->GetDomainShader();
 	// Ensure the command buffers are reset to reduce the amount of data that needs to be versioned.
 	VSConstantBuffer.Reset();
@@ -564,7 +489,7 @@ void FD3D12CommandContext::RHISetGraphicsPipelineState(FRHIGraphicsPipelineState
 		StateCache.SetDepthBounds(0.0f, 1.0f);
 	}
 
-	StateCache.SetGraphicsPipelineState(GraphicsPipelineState);
+	StateCache.SetGraphicsPipelineState(GraphicsPipelineState, bUsingTessellation != bWasUsingTessellation);
 	StateCache.SetStencilRef(0);
 }
 
@@ -1414,37 +1339,34 @@ void FD3D12CommandContext::CommitComputeResourceTables(FD3D12ComputeShader* InCo
 
 void FD3D12CommandContext::RHIDrawPrimitive(uint32 BaseVertexIndex, uint32 NumPrimitives, uint32 NumInstances)
 {
-	RHI_DRAW_CALL_STATS(StateCache.GetGraphicsPipelinePrimitiveType(), NumInstances*NumPrimitives);
-
 	CommitGraphicsResourceTables();
 	CommitNonComputeShaderConstants();
 
 	uint32 VertexCount = GetVertexCountForPrimitiveCount(NumPrimitives, StateCache.GetGraphicsPipelinePrimitiveType());
 
-	if (IsDefaultContext())
+	NumInstances = FMath::Max<uint32>(1, NumInstances);
+	numDraws++;
+	numPrimitives += NumInstances * NumPrimitives;
+	if (bTrackingEvents)
 	{
 		GetParentDevice()->RegisterGPUWork(NumPrimitives * NumInstances, VertexCount * NumInstances);
 	}
 
-	StateCache.SetPrimitiveTopology(GetD3D12PrimitiveType(StateCache.GetGraphicsPipelinePrimitiveType(), bUsingTessellation));
-
 	StateCache.ApplyState<D3D12PT_Graphics>();
-	numDraws++;
-	CommandListHandle->DrawInstanced(VertexCount, FMath::Max<uint32>(1, NumInstances), BaseVertexIndex, 0);
+	CommandListHandle->DrawInstanced(VertexCount, NumInstances, BaseVertexIndex, 0);
+
 #if UE_BUILD_DEBUG	
 	OwningRHI.DrawCount++;
 #endif
 	DEBUG_EXECUTE_COMMAND_LIST(this);
-
 }
 
 void FD3D12CommandContext::RHIDrawPrimitiveIndirect(FRHIVertexBuffer* ArgumentBufferRHI, uint32 ArgumentOffset)
 {
 	FD3D12VertexBuffer* ArgumentBuffer = RetrieveObject<FD3D12VertexBuffer>(ArgumentBufferRHI);
 
-	RHI_DRAW_CALL_INC();
-
-	if (IsDefaultContext())
+	numDraws++;
+	if (bTrackingEvents)
 	{
 		GetParentDevice()->RegisterGPUWork(0);
 	}
@@ -1452,14 +1374,11 @@ void FD3D12CommandContext::RHIDrawPrimitiveIndirect(FRHIVertexBuffer* ArgumentBu
 	CommitGraphicsResourceTables();
 	CommitNonComputeShaderConstants();
 
-	StateCache.SetPrimitiveTopology(GetD3D12PrimitiveType(StateCache.GetGraphicsPipelinePrimitiveType(), bUsingTessellation));
-
 	FD3D12ResourceLocation& Location = ArgumentBuffer->ResourceLocation;
 	FD3D12DynamicRHI::TransitionResource(CommandListHandle, Location.GetResource(), D3D12_RESOURCE_STATE_INDIRECT_ARGUMENT, D3D12_RESOURCE_BARRIER_ALL_SUBRESOURCES);
 
 	StateCache.ApplyState<D3D12PT_Graphics>();
 
-	numDraws++;
 	CommandListHandle->ExecuteIndirect(
 		GetParentDevice()->GetParentAdapter()->GetDrawIndirectCommandSignature(),
 		1,
@@ -1475,7 +1394,6 @@ void FD3D12CommandContext::RHIDrawPrimitiveIndirect(FRHIVertexBuffer* ArgumentBu
 	OwningRHI.DrawCount++;
 #endif
 	DEBUG_EXECUTE_COMMAND_LIST(this);
-
 }
 
 void FD3D12CommandContext::RHIDrawIndexedIndirect(FRHIIndexBuffer* IndexBufferRHI, FRHIStructuredBuffer* ArgumentsBufferRHI, int32 DrawArgumentsIndex, uint32 NumInstances)
@@ -1483,9 +1401,8 @@ void FD3D12CommandContext::RHIDrawIndexedIndirect(FRHIIndexBuffer* IndexBufferRH
 	FD3D12IndexBuffer* IndexBuffer = RetrieveObject<FD3D12IndexBuffer>(IndexBufferRHI);
 	FD3D12StructuredBuffer* ArgumentsBuffer = RetrieveObject<FD3D12StructuredBuffer>(ArgumentsBufferRHI);
 
-	RHI_DRAW_CALL_INC();
-
-	if (IsDefaultContext())
+	numDraws++;
+	if (bTrackingEvents)
 	{
 		GetParentDevice()->RegisterGPUWork(1);
 	}
@@ -1494,18 +1411,15 @@ void FD3D12CommandContext::RHIDrawIndexedIndirect(FRHIIndexBuffer* IndexBufferRH
 	CommitNonComputeShaderConstants();
 
 	// determine 16bit vs 32bit indices
-	uint32 SizeFormat = sizeof(DXGI_FORMAT);
 	const DXGI_FORMAT Format = (IndexBuffer->GetStride() == sizeof(uint16) ? DXGI_FORMAT_R16_UINT : DXGI_FORMAT_R32_UINT);
 
 	StateCache.SetIndexBuffer(&IndexBuffer->ResourceLocation, Format, 0);
-	StateCache.SetPrimitiveTopology(GetD3D12PrimitiveType(StateCache.GetGraphicsPipelinePrimitiveType(), bUsingTessellation));
 
 	FD3D12ResourceLocation& Location = ArgumentsBuffer->ResourceLocation;
 	FD3D12DynamicRHI::TransitionResource(CommandListHandle, Location.GetResource(), D3D12_RESOURCE_STATE_INDIRECT_ARGUMENT, D3D12_RESOURCE_BARRIER_ALL_SUBRESOURCES);
 
 	StateCache.ApplyState<D3D12PT_Graphics>();
 
-	numDraws++;
 	CommandListHandle->ExecuteIndirect(
 		GetParentDevice()->GetParentAdapter()->GetDrawIndexedIndirectCommandSignature(),
 		1,
@@ -1525,43 +1439,39 @@ void FD3D12CommandContext::RHIDrawIndexedIndirect(FRHIIndexBuffer* IndexBufferRH
 
 void FD3D12CommandContext::RHIDrawIndexedPrimitive(FRHIIndexBuffer* IndexBufferRHI, int32 BaseVertexIndex, uint32 FirstInstance, uint32 NumVertices, uint32 StartIndex, uint32 NumPrimitives, uint32 NumInstances)
 {
-	FD3D12IndexBuffer* IndexBuffer = RetrieveObject<FD3D12IndexBuffer>(IndexBufferRHI);
-
 	// called should make sure the input is valid, this avoid hidden bugs
 	ensure(NumPrimitives > 0);
 
-	RHI_DRAW_CALL_STATS(StateCache.GetGraphicsPipelinePrimitiveType(), NumInstances*NumPrimitives);
-
-	if (IsDefaultContext())
+	NumInstances = FMath::Max<uint32>(1, NumInstances);
+	numDraws++;
+	numPrimitives += NumInstances * NumPrimitives;
+	if (bTrackingEvents)
 	{
 		GetParentDevice()->RegisterGPUWork(NumPrimitives * NumInstances, NumVertices * NumInstances);
 	}
+	uint32 IndexCount = GetVertexCountForPrimitiveCount(NumPrimitives, StateCache.GetGraphicsPipelinePrimitiveType());
 
 	CommitGraphicsResourceTables();
 	CommitNonComputeShaderConstants();
 
-	// determine 16bit vs 32bit indices
-	uint32 SizeFormat = sizeof(DXGI_FORMAT);
-	const DXGI_FORMAT Format = (IndexBuffer->GetStride() == sizeof(uint16) ? DXGI_FORMAT_R16_UINT : DXGI_FORMAT_R32_UINT);
-
-	uint32 IndexCount = GetVertexCountForPrimitiveCount(NumPrimitives, StateCache.GetGraphicsPipelinePrimitiveType());
+	FD3D12IndexBuffer* IndexBuffer = RetrieveObject<FD3D12IndexBuffer>(IndexBufferRHI);
 
 	// Verify that we are not trying to read outside the index buffer range
 	// test is an optimized version of: StartIndex + IndexCount <= IndexBuffer->GetSize() / IndexBuffer->GetStride() 
 	checkf((StartIndex + IndexCount) * IndexBuffer->GetStride() <= IndexBuffer->GetSize(),
 		TEXT("Start %u, Count %u, Type %u, Buffer Size %u, Buffer stride %u"), StartIndex, IndexCount, StateCache.GetGraphicsPipelinePrimitiveType(), IndexBuffer->GetSize(), IndexBuffer->GetStride());
 
+	// determine 16bit vs 32bit indices
+	const DXGI_FORMAT Format = (IndexBuffer->GetStride() == sizeof(uint16) ? DXGI_FORMAT_R16_UINT : DXGI_FORMAT_R32_UINT);
 	StateCache.SetIndexBuffer(&IndexBuffer->ResourceLocation, Format, 0);
-	StateCache.SetPrimitiveTopology(GetD3D12PrimitiveType(StateCache.GetGraphicsPipelinePrimitiveType(), bUsingTessellation));
 	StateCache.ApplyState<D3D12PT_Graphics>();
 
-	numDraws++;
-	CommandListHandle->DrawIndexedInstanced(IndexCount, FMath::Max<uint32>(1, NumInstances), StartIndex, BaseVertexIndex, FirstInstance);
+	CommandListHandle->DrawIndexedInstanced(IndexCount, NumInstances, StartIndex, BaseVertexIndex, FirstInstance);
+
 #if UE_BUILD_DEBUG	
 	OwningRHI.DrawCount++;
 #endif
 	DEBUG_EXECUTE_COMMAND_LIST(this);
-
 }
 
 void FD3D12CommandContext::RHIDrawIndexedPrimitiveIndirect(FRHIIndexBuffer* IndexBufferRHI, FRHIVertexBuffer* ArgumentBufferRHI, uint32 ArgumentOffset)
@@ -1569,9 +1479,8 @@ void FD3D12CommandContext::RHIDrawIndexedPrimitiveIndirect(FRHIIndexBuffer* Inde
 	FD3D12IndexBuffer* IndexBuffer = RetrieveObject<FD3D12IndexBuffer>(IndexBufferRHI);
 	FD3D12VertexBuffer* ArgumentBuffer = RetrieveObject<FD3D12VertexBuffer>(ArgumentBufferRHI);
 
-	RHI_DRAW_CALL_INC();
-
-	if (IsDefaultContext())
+	numDraws++;
+	if (bTrackingEvents)
 	{
 		GetParentDevice()->RegisterGPUWork(0);
 	}
@@ -1580,17 +1489,14 @@ void FD3D12CommandContext::RHIDrawIndexedPrimitiveIndirect(FRHIIndexBuffer* Inde
 	CommitNonComputeShaderConstants();
 
 	// Set the index buffer.
-	const uint32 SizeFormat = sizeof(DXGI_FORMAT);
 	const DXGI_FORMAT Format = (IndexBuffer->GetStride() == sizeof(uint16) ? DXGI_FORMAT_R16_UINT : DXGI_FORMAT_R32_UINT);
 	StateCache.SetIndexBuffer(&IndexBuffer->ResourceLocation, Format, 0);
-	StateCache.SetPrimitiveTopology(GetD3D12PrimitiveType(StateCache.GetGraphicsPipelinePrimitiveType(), bUsingTessellation));
 
 	FD3D12ResourceLocation& Location = ArgumentBuffer->ResourceLocation;
 	FD3D12DynamicRHI::TransitionResource(CommandListHandle, Location.GetResource(), D3D12_RESOURCE_STATE_INDIRECT_ARGUMENT, D3D12_RESOURCE_BARRIER_ALL_SUBRESOURCES);
 
 	StateCache.ApplyState<D3D12PT_Graphics>();
 
-	numDraws++;
 	CommandListHandle->ExecuteIndirect(
 		GetParentDevice()->GetParentAdapter()->GetDrawIndexedIndirectCommandSignature(),
 		1,
diff --git a/Engine/Source/Runtime/D3D12RHI/Private/D3D12StateCache.cpp b/Engine/Source/Runtime/D3D12RHI/Private/D3D12StateCache.cpp
--- a/Engine/Source/Runtime/D3D12RHI/Private/D3D12StateCache.cpp
+++ b/Engine/Source/Runtime/D3D12RHI/Private/D3D12StateCache.cpp
@@ -187,6 +187,7 @@ void FD3D12StateCacheBase::ClearState()
 	PipelineState.Graphics.CurrentNumberOfRenderTargets = 0;
 
 	PipelineState.Graphics.CurrentPrimitiveTopology = D3D_PRIMITIVE_TOPOLOGY_UNDEFINED;
+	PipelineState.Graphics.CurrentPrimitiveType = PT_Num;
 
 	PipelineState.Graphics.MinDepth = 0.0f;
 	PipelineState.Graphics.MaxDepth = 1.0f;
@@ -202,6 +203,7 @@ void FD3D12StateCacheBase::DirtyStateForNewCommandList()
 	PipelineState.Common.bNeedSetPSO = true;
 	PipelineState.Compute.bNeedSetRootSignature = true;
 	PipelineState.Graphics.bNeedSetRootSignature = true;
+	bNeedSetPrimitiveTopology = true;
 
 	if (PipelineState.Graphics.VBCache.BoundVBMask) { bNeedSetVB = true; }
 	if (PipelineState.Graphics.IBCache.CurrentIndexBufferLocation) { bNeedSetIB = true; }
@@ -209,7 +211,6 @@ void FD3D12StateCacheBase::DirtyStateForNewCommandList()
 	if (PipelineState.Graphics.CurrentNumberOfRenderTargets || PipelineState.Graphics.CurrentDepthStencilTarget) { bNeedSetRTs = true; }
 	if (PipelineState.Graphics.CurrentNumberOfViewports) { bNeedSetViewports = true; }
 	if (PipelineState.Graphics.CurrentNumberOfScissorRects) { bNeedSetScissorRects = true; }
-	if (PipelineState.Graphics.CurrentPrimitiveTopology != D3D12_IA_DEFAULT_PRIMITIVE_TOPOLOGY) { bNeedSetPrimitiveTopology = true; }
 
 	if (PipelineState.Graphics.CurrentBlendFactor[0] != D3D12_DEFAULT_BLEND_FACTOR_RED ||
 		PipelineState.Graphics.CurrentBlendFactor[1] != D3D12_DEFAULT_BLEND_FACTOR_GREEN ||
@@ -1028,15 +1029,6 @@ void FD3D12StateCacheBase::SetUAVs(uint32 UAVStartSlot, uint32 NumSimultaneousUA
 	}
 }
 
-void FD3D12StateCacheBase::SetPrimitiveTopology(D3D12_PRIMITIVE_TOPOLOGY PrimitiveTopology)
-{
-	if ((PipelineState.Graphics.CurrentPrimitiveTopology != PrimitiveTopology) || GD3D12SkipStateCaching)
-	{
-		PipelineState.Graphics.CurrentPrimitiveTopology = PrimitiveTopology;
-		bNeedSetPrimitiveTopology = true;
-	}
-}
-
 void FD3D12StateCacheBase::SetBlendFactor(const float BlendFactor[4])
 {
 	if (FMemory::Memcmp(PipelineState.Graphics.CurrentBlendFactor, BlendFactor, sizeof(PipelineState.Graphics.CurrentBlendFactor)))
diff --git a/Engine/Source/Runtime/D3D12RHI/Private/D3D12StateCachePrivate.h b/Engine/Source/Runtime/D3D12RHI/Private/D3D12StateCachePrivate.h
--- a/Engine/Source/Runtime/D3D12RHI/Private/D3D12StateCachePrivate.h
+++ b/Engine/Source/Runtime/D3D12RHI/Private/D3D12StateCachePrivate.h
@@ -258,6 +258,70 @@ struct FD3D12SamplerStateCache : public FD3D12ResourceCache<SamplerSlotMask>
 	FD3D12SamplerState* States[SF_NumStandardFrequencies][MAX_SAMPLERS];
 };
 
+
+static inline D3D_PRIMITIVE_TOPOLOGY GetD3D12PrimitiveType(uint32 PrimitiveType, bool bUsingTessellation)
+{
+	static const uint8 D3D12PrimitiveType[] =
+	{
+		D3D_PRIMITIVE_TOPOLOGY_TRIANGLELIST,               // PT_TriangleList
+		D3D_PRIMITIVE_TOPOLOGY_TRIANGLESTRIP,              // PT_TriangleStrip
+		D3D_PRIMITIVE_TOPOLOGY_LINELIST,                   // PT_LineList
+		0,                                                 // PT_QuadList
+		D3D_PRIMITIVE_TOPOLOGY_POINTLIST,                  // PT_PointList
+		D3D_PRIMITIVE_TOPOLOGY_1_CONTROL_POINT_PATCHLIST,  // PT_1_ControlPointPatchList
+		D3D_PRIMITIVE_TOPOLOGY_2_CONTROL_POINT_PATCHLIST,  // PT_2_ControlPointPatchList
+		D3D_PRIMITIVE_TOPOLOGY_3_CONTROL_POINT_PATCHLIST,  // PT_3_ControlPointPatchList
+		D3D_PRIMITIVE_TOPOLOGY_4_CONTROL_POINT_PATCHLIST,  // PT_4_ControlPointPatchList
+		D3D_PRIMITIVE_TOPOLOGY_5_CONTROL_POINT_PATCHLIST,  // PT_5_ControlPointPatchList
+		D3D_PRIMITIVE_TOPOLOGY_6_CONTROL_POINT_PATCHLIST,  // PT_6_ControlPointPatchList
+		D3D_PRIMITIVE_TOPOLOGY_7_CONTROL_POINT_PATCHLIST,  // PT_7_ControlPointPatchList
+		D3D_PRIMITIVE_TOPOLOGY_8_CONTROL_POINT_PATCHLIST,  // PT_8_ControlPointPatchList
+		D3D_PRIMITIVE_TOPOLOGY_9_CONTROL_POINT_PATCHLIST,  // PT_9_ControlPointPatchList
+		D3D_PRIMITIVE_TOPOLOGY_10_CONTROL_POINT_PATCHLIST, // PT_10_ControlPointPatchList
+		D3D_PRIMITIVE_TOPOLOGY_11_CONTROL_POINT_PATCHLIST, // PT_11_ControlPointPatchList
+		D3D_PRIMITIVE_TOPOLOGY_12_CONTROL_POINT_PATCHLIST, // PT_12_ControlPointPatchList
+		D3D_PRIMITIVE_TOPOLOGY_13_CONTROL_POINT_PATCHLIST, // PT_13_ControlPointPatchList
+		D3D_PRIMITIVE_TOPOLOGY_14_CONTROL_POINT_PATCHLIST, // PT_14_ControlPointPatchList
+		D3D_PRIMITIVE_TOPOLOGY_15_CONTROL_POINT_PATCHLIST, // PT_15_ControlPointPatchList
+		D3D_PRIMITIVE_TOPOLOGY_16_CONTROL_POINT_PATCHLIST, // PT_16_ControlPointPatchList
+		D3D_PRIMITIVE_TOPOLOGY_17_CONTROL_POINT_PATCHLIST, // PT_17_ControlPointPatchList
+		D3D_PRIMITIVE_TOPOLOGY_18_CONTROL_POINT_PATCHLIST, // PT_18_ControlPointPatchList
+		D3D_PRIMITIVE_TOPOLOGY_19_CONTROL_POINT_PATCHLIST, // PT_19_ControlPointPatchList
+		D3D_PRIMITIVE_TOPOLOGY_20_CONTROL_POINT_PATCHLIST, // PT_20_ControlPointPatchList
+		D3D_PRIMITIVE_TOPOLOGY_21_CONTROL_POINT_PATCHLIST, // PT_21_ControlPointPatchList
+		D3D_PRIMITIVE_TOPOLOGY_22_CONTROL_POINT_PATCHLIST, // PT_22_ControlPointPatchList
+		D3D_PRIMITIVE_TOPOLOGY_23_CONTROL_POINT_PATCHLIST, // PT_23_ControlPointPatchList
+		D3D_PRIMITIVE_TOPOLOGY_24_CONTROL_POINT_PATCHLIST, // PT_24_ControlPointPatchList
+		D3D_PRIMITIVE_TOPOLOGY_25_CONTROL_POINT_PATCHLIST, // PT_25_ControlPointPatchList
+		D3D_PRIMITIVE_TOPOLOGY_26_CONTROL_POINT_PATCHLIST, // PT_26_ControlPointPatchList
+		D3D_PRIMITIVE_TOPOLOGY_27_CONTROL_POINT_PATCHLIST, // PT_27_ControlPointPatchList
+		D3D_PRIMITIVE_TOPOLOGY_28_CONTROL_POINT_PATCHLIST, // PT_28_ControlPointPatchList
+		D3D_PRIMITIVE_TOPOLOGY_29_CONTROL_POINT_PATCHLIST, // PT_29_ControlPointPatchList
+		D3D_PRIMITIVE_TOPOLOGY_30_CONTROL_POINT_PATCHLIST, // PT_30_ControlPointPatchList
+		D3D_PRIMITIVE_TOPOLOGY_31_CONTROL_POINT_PATCHLIST, // PT_31_ControlPointPatchList
+		D3D_PRIMITIVE_TOPOLOGY_32_CONTROL_POINT_PATCHLIST, // PT_32_ControlPointPatchList
+	};
+
+	if (bUsingTessellation)
+	{
+		if (PrimitiveType == PT_TriangleList)
+		{
+			// This is the case for tessellation without AEN or other buffers, so just flip to 3 CPs
+			return D3D_PRIMITIVE_TOPOLOGY_3_CONTROL_POINT_PATCHLIST;
+		}
+		else/* if (PrimitiveType < PT_1_ControlPointPatchList)*/
+		{
+			//UE_LOG(LogD3D12RHI, Fatal, TEXT("Invalid type specified for tessellated render, probably missing a case in FSkeletalMeshSceneProxy::DrawDynamicElementsByMaterial or FStaticMeshSceneProxy::GetMeshElement"));
+
+			checkf(PrimitiveType >= PT_1_ControlPointPatchList, TEXT("Invalid type specified for tessellated render, probably missing a case in FSkeletalMeshSceneProxy::DrawDynamicElementsByMaterial or FStaticMeshSceneProxy::GetMeshElement"));
+		}
+	}
+
+	D3D_PRIMITIVE_TOPOLOGY D3DType = (D3D_PRIMITIVE_TOPOLOGY) D3D12PrimitiveType[PrimitiveType];
+	checkf(D3DType, TEXT("Unknown primitive type: %u"), PrimitiveType);
+	return D3DType;
+}
+
 //-----------------------------------------------------------------------------
 //	FD3D12StateCache Class Definition
 //-----------------------------------------------------------------------------
@@ -314,6 +378,7 @@ protected:
 			FD3D12IndexBufferCache IBCache;
 
 			// Primitive Topology State
+			EPrimitiveType CurrentPrimitiveType;
 			D3D_PRIMITIVE_TOPOLOGY CurrentPrimitiveTopology;
 
 			// Input Layout State
@@ -329,8 +394,6 @@ protected:
 
 			float MinDepth;
 			float MaxDepth;
-
-			EPrimitiveType PrimitiveType = PT_Num;
 		} Graphics;
 
 		struct
@@ -466,7 +529,7 @@ public:
 
 	inline EPrimitiveType GetGraphicsPipelinePrimitiveType() const
 	{
-		return PipelineState.Graphics.PrimitiveType;
+		return PipelineState.Graphics.CurrentPrimitiveType;
 	}
 
 	const FD3D12RootSignature* GetComputeRootSignature()
@@ -673,7 +736,7 @@ public:
 		GetShader(Shader);
 	}
 
-	D3D12_STATE_CACHE_INLINE void SetGraphicsPipelineState(FD3D12GraphicsPipelineState* GraphicsPipelineState)
+	D3D12_STATE_CACHE_INLINE void SetGraphicsPipelineState(FD3D12GraphicsPipelineState* GraphicsPipelineState, bool bTessellationChanged)
 	{
 		check(GraphicsPipelineState);
 		if (PipelineState.Graphics.CurrentPipelineStateObject != GraphicsPipelineState)
@@ -697,7 +760,15 @@ public:
 			// Save the PSO
 			PipelineState.Common.bNeedSetPSO = true;
 			PipelineState.Graphics.CurrentPipelineStateObject = GraphicsPipelineState;
-			PipelineState.Graphics.PrimitiveType = GraphicsPipelineState->PipelineStateInitializer.PrimitiveType;
+
+			EPrimitiveType PrimitiveType = GraphicsPipelineState->PipelineStateInitializer.PrimitiveType;
+			if (PipelineState.Graphics.CurrentPrimitiveType != PrimitiveType || bTessellationChanged)
+			{
+				const bool bUsingTessellation = GraphicsPipelineState->GetHullShader() && GraphicsPipelineState->GetDomainShader();
+				PipelineState.Graphics.CurrentPrimitiveType = PrimitiveType;
+				PipelineState.Graphics.CurrentPrimitiveTopology = GetD3D12PrimitiveType(PrimitiveType, bUsingTessellation);
+				bNeedSetPrimitiveTopology = true;
+			}
 
 			// Set the PSO
 			InternalSetPipelineState<D3D12PT_Graphics>();
@@ -790,8 +861,6 @@ public:
 		return PipelineState.Graphics.IBCache.CurrentIndexBufferLocation == ResourceLocation;
 	}
 
-	void SetPrimitiveTopology(D3D12_PRIMITIVE_TOPOLOGY PrimitiveTopology);
-	
 	D3D12_STATE_CACHE_INLINE void GetPrimitiveTopology(D3D12_PRIMITIVE_TOPOLOGY* PrimitiveTopology) const
 	{
 		*PrimitiveTopology = PipelineState.Graphics.CurrentPrimitiveTopology;
diff --git a/Engine/Source/Runtime/RHI/Public/RHIUtilities.h b/Engine/Source/Runtime/RHI/Public/RHIUtilities.h
--- a/Engine/Source/Runtime/RHI/Public/RHIUtilities.h
+++ b/Engine/Source/Runtime/RHI/Public/RHIUtilities.h
@@ -830,51 +830,10 @@ inline void RHICreateTargetableShaderResource3D(
  */
 inline uint32 GetVertexCountForPrimitiveCount(uint32 NumPrimitives, uint32 PrimitiveType)
 {
-	uint32 VertexCount = 0;
-	switch(PrimitiveType)
-	{
-	case PT_TriangleList: VertexCount = NumPrimitives*3; break;
-	case PT_TriangleStrip: VertexCount = NumPrimitives+2; break;
-	case PT_LineList: VertexCount = NumPrimitives*2; break;
-	case PT_PointList: VertexCount = NumPrimitives; break;
-	case PT_1_ControlPointPatchList:
-	case PT_2_ControlPointPatchList:
-	case PT_3_ControlPointPatchList: 
-	case PT_4_ControlPointPatchList: 
-	case PT_5_ControlPointPatchList:
-	case PT_6_ControlPointPatchList:
-	case PT_7_ControlPointPatchList: 
-	case PT_8_ControlPointPatchList: 
-	case PT_9_ControlPointPatchList: 
-	case PT_10_ControlPointPatchList: 
-	case PT_11_ControlPointPatchList: 
-	case PT_12_ControlPointPatchList: 
-	case PT_13_ControlPointPatchList: 
-	case PT_14_ControlPointPatchList: 
-	case PT_15_ControlPointPatchList: 
-	case PT_16_ControlPointPatchList: 
-	case PT_17_ControlPointPatchList: 
-	case PT_18_ControlPointPatchList: 
-	case PT_19_ControlPointPatchList: 
-	case PT_20_ControlPointPatchList: 
-	case PT_21_ControlPointPatchList: 
-	case PT_22_ControlPointPatchList: 
-	case PT_23_ControlPointPatchList: 
-	case PT_24_ControlPointPatchList: 
-	case PT_25_ControlPointPatchList: 
-	case PT_26_ControlPointPatchList: 
-	case PT_27_ControlPointPatchList: 
-	case PT_28_ControlPointPatchList: 
-	case PT_29_ControlPointPatchList: 
-	case PT_30_ControlPointPatchList: 
-	case PT_31_ControlPointPatchList: 
-	case PT_32_ControlPointPatchList: 
-		VertexCount = (PrimitiveType - PT_1_ControlPointPatchList + 1) * NumPrimitives;
-		break;
-	default: UE_LOG(LogRHI, Fatal,TEXT("Unknown primitive type: %u"),PrimitiveType);
-	};
+	uint32 Factor = (PrimitiveType == PT_TriangleList)? 3 : (PrimitiveType == PT_LineList)? 2 : (PrimitiveType >= PT_1_ControlPointPatchList)? (PrimitiveType - PT_1_ControlPointPatchList + 1) : 1;
+	uint32 Offset = (PrimitiveType == PT_TriangleStrip)? 2 : 0;
 
-	return VertexCount;
+	return NumPrimitives * Factor + Offset;
 }
 
 inline uint32 ComputeAnisotropyRT(int32 InitializerMaxAnisotropy)
