Move PipelineState hashing to 64-bit to avoid hash collision. Also take vertex declaration into account in the hash.

diff --git a/Engine/Source/Runtime/D3D12RHI/Private/D3D12PipelineState.h b/Engine/Source/Runtime/D3D12RHI/Private/D3D12PipelineState.h
--- a/Engine/Source/Runtime/D3D12RHI/Private/D3D12PipelineState.h
+++ b/Engine/Source/Runtime/D3D12RHI/Private/D3D12PipelineState.h
@@ -18,8 +18,6 @@
 #endif
 #define D3D12_USE_DERIVED_PSO PLATFORM_XBOXONE
 
-static bool GCPUSupportsSSE4;
-
 #if D3D12RHI_USE_HIGH_LEVEL_PSO_CACHE
 DECLARE_DWORD_ACCUMULATOR_STAT(TEXT("Graphics: Num high-level cache entries"), STAT_PSOGraphicsNumHighlevelCacheEntries, STATGROUP_D3D12PipelineState);
 DECLARE_DWORD_COUNTER_STAT(TEXT("Graphics: High-level cache hit"), STAT_PSOGraphicsHighlevelCacheHit, STATGROUP_D3D12PipelineState);
@@ -79,6 +77,7 @@ struct FD3D12LowLevelGraphicsPipelineStateDesc
 	ShaderBytecodeHash DSHash;
 	ShaderBytecodeHash GSHash;
 	ShaderBytecodeHash PSHash;
+	uint32 InputLayoutHash;
 
 	SIZE_T CombinedHash;
 
@@ -493,10 +492,10 @@ public:
 	FD3D12ComputePipelineState* FindInLoadedCache(FD3D12ComputeShader* ComputeShader, FD3D12ComputePipelineStateDesc& OutLowLevelDesc);
 	FD3D12ComputePipelineState* CreateAndAdd(FD3D12ComputeShader* ComputeShader, const FD3D12ComputePipelineStateDesc& LowLevelDesc);
 
-	static SIZE_T HashPSODesc(const FD3D12LowLevelGraphicsPipelineStateDesc& Desc);
-	static SIZE_T HashPSODesc(const FD3D12ComputePipelineStateDesc& Desc);
+	static uint64 HashPSODesc(const FD3D12LowLevelGraphicsPipelineStateDesc& Desc);
+	static uint64 HashPSODesc(const FD3D12ComputePipelineStateDesc& Desc);
 
-	static uint32 HashData(const void* Data, SIZE_T NumBytes);
+	static uint64 HashData(const void* Data, int32 NumBytes);
 
 	FD3D12PipelineStateCacheBase(FD3D12Adapter* InParent);
 	virtual ~FD3D12PipelineStateCacheBase();
diff --git a/Engine/Source/Runtime/D3D12RHI/Private/D3D12Pipelinestate.cpp b/Engine/Source/Runtime/D3D12RHI/Private/D3D12Pipelinestate.cpp
--- a/Engine/Source/Runtime/D3D12RHI/Private/D3D12Pipelinestate.cpp
+++ b/Engine/Source/Runtime/D3D12RHI/Private/D3D12Pipelinestate.cpp
@@ -6,7 +6,7 @@
 //	Include Files
 //-----------------------------------------------------------------------------
 #include "D3D12RHIPrivate.h"
-
+#include "Hash/CityHash.h"
 
 static TAutoConsoleVariable<float> CVarPSOStallWarningThresholdInMs(
 	TEXT("D3D12.PSO.StallWarningThresholdInMs"),
@@ -58,6 +58,7 @@ FD3D12LowLevelGraphicsPipelineStateDesc GetLowLevelGraphicsPipelineStateDesc(con
 	{
 		Desc.Desc.InputLayout.NumElements        = InputLayout->VertexElements.Num();
 		Desc.Desc.InputLayout.pInputElementDescs = InputLayout->VertexElements.GetData();
+		Desc.InputLayoutHash = InputLayout->Hash;
 	}
 
 #define COPY_SHADER(Initial, Name) \
@@ -132,119 +133,115 @@ FD3D12PipelineStateWorker::FD3D12PipelineStateWorker(FD3D12Adapter* Adapter, con
 	CreationArgs.GraphicsArgs.Init(InArgs.Args);
 };
 
-
 /// @endcond
-#if defined(_WIN64) || (defined(_M_IX86_FP) && _M_IX86_FP >= 2)
-#ifndef __clang__
 
-FORCEINLINE uint32 SSE4_CRC32(const void* Data, SIZE_T NumBytes)
+uint64 FD3D12PipelineStateCacheBase::HashData(const void* Data, int32 NumBytes)
 {
-	check(GCPUSupportsSSE4);
-	uint32 Hash = 0;
-#if defined(_WIN64)
-	static const SIZE_T Alignment = 8;//64 Bit
-#elif defined(_WIN32)
-	static const SIZE_T Alignment = 4;//32 Bit
-#else
-	check(0);
-	return 0;
-#endif
+	return CityHash64((const char*) Data, NumBytes);
+}
 
-	const SIZE_T RoundingIterations = (NumBytes & (Alignment - 1));
-	uint8* UnalignedData = (uint8*)Data;
-	for (SIZE_T i = 0; i < RoundingIterations; i++)
+uint64 FD3D12PipelineStateCacheBase::HashPSODesc(const FD3D12LowLevelGraphicsPipelineStateDesc& Desc)
+{
+	struct GraphicsPSOData
 	{
-		Hash = _mm_crc32_u8(Hash, UnalignedData[i]);
-	}
-	UnalignedData += RoundingIterations;
-	NumBytes -= RoundingIterations;
+		ShaderBytecodeHash VSHash;
+		ShaderBytecodeHash HSHash;
+		ShaderBytecodeHash DSHash;
+		ShaderBytecodeHash GSHash;
+		ShaderBytecodeHash PSHash;
+		uint32 InputLayoutHash;
+
+#if !D3D12_USE_DERIVED_PSO
+		uint8 AlphaToCoverageEnable;
+		uint8 IndependentBlendEnable;
+
+		uint32 SampleMask;
+		D3D12_RASTERIZER_DESC RasterizerState;
+		D3D12_DEPTH_STENCIL_DESC1 DepthStencilState;
+#endif // !D3D12_USE_DERIVED_PSO
 
-	SIZE_T* AlignedData = (SIZE_T*)UnalignedData;
-	check((NumBytes % Alignment) == 0);
-	const SIZE_T NumIterations = (NumBytes / Alignment);
-	for (SIZE_T i = 0; i < NumIterations; i++)
+		D3D12_INDEX_BUFFER_STRIP_CUT_VALUE IBStripCutValue;
+		D3D12_PRIMITIVE_TOPOLOGY_TYPE PrimitiveTopologyType;
+		DXGI_FORMAT DSVFormat;
+		DXGI_SAMPLE_DESC SampleDesc;
+		uint32 NodeMask;
+
+		D3D12_PIPELINE_STATE_FLAGS Flags;
+	};
+
+	struct RenderTargetData
 	{
-#ifdef _WIN64
-		Hash = _mm_crc32_u64(Hash, AlignedData[i]);
-#else
-		Hash = _mm_crc32_u32(Hash, AlignedData[i]);
+		DXGI_FORMAT Format;
+#if !D3D12_USE_DERIVED_PSO
+		D3D12_RENDER_TARGET_BLEND_DESC BlendDesc;
 #endif
-	}
+	};
 
-	return Hash;
-}
-#endif
+
+	const int32 NumRenderTargets = Desc.Desc.RTFormatArray.NumRenderTargets;
+
+	const size_t GraphicsPSODataSize = sizeof(GraphicsPSOData);
+	const size_t RenderTargetDataSize = NumRenderTargets * sizeof(RenderTargetData);
+	const size_t TotalDataSize = GraphicsPSODataSize + RenderTargetDataSize;
+
+	char Data[GraphicsPSODataSize + 8 * sizeof(RenderTargetData)];
+	FMemory::Memzero(Data, TotalDataSize);
+
+	GraphicsPSOData* PSOData = (GraphicsPSOData*) Data;
+	RenderTargetData* RTData = (RenderTargetData*) (Data + GraphicsPSODataSize);
+
+	PSOData->VSHash          = Desc.VSHash;
+	PSOData->HSHash          = Desc.HSHash;
+	PSOData->DSHash          = Desc.DSHash;
+	PSOData->GSHash          = Desc.GSHash;
+	PSOData->PSHash          = Desc.PSHash;
+	PSOData->InputLayoutHash = Desc.InputLayoutHash;
+
+#if !D3D12_USE_DERIVED_PSO
+	PSOData->AlphaToCoverageEnable  = Desc.Desc.BlendState.AlphaToCoverageEnable;
+	PSOData->IndependentBlendEnable = Desc.Desc.BlendState.IndependentBlendEnable;
+	PSOData->SampleMask             = Desc.Desc.SampleMask;
+	PSOData->RasterizerState        = Desc.Desc.RasterizerState;
+	PSOData->DepthStencilState      = Desc.Desc.DepthStencilState;
 #endif
+	PSOData->IBStripCutValue        = Desc.Desc.IBStripCutValue;
+	PSOData->PrimitiveTopologyType  = Desc.Desc.PrimitiveTopologyType;
+	PSOData->DSVFormat              = Desc.Desc.DSVFormat;
+	PSOData->SampleDesc             = Desc.Desc.SampleDesc;
+	PSOData->NodeMask               = Desc.Desc.NodeMask;
+	PSOData->Flags                  = Desc.Desc.Flags;
 
-uint32 FD3D12PipelineStateCacheBase::HashData(const void* Data, SIZE_T NumBytes)
-{
-#ifdef __clang__
-	return FCrc::MemCrc32(Data, NumBytes);
-#else
-#if defined(_WIN64) || (defined(_M_IX86_FP) && _M_IX86_FP >= 2)
-	if (GCPUSupportsSSE4)
+	for (int32 RT = 0; RT < NumRenderTargets; RT++)
 	{
-		return SSE4_CRC32(Data, NumBytes);
-	}
-	else
+		RTData[RT].Format    = Desc.Desc.RTFormatArray.RTFormats[RT];
+#if !D3D12_USE_DERIVED_PSO
+		RTData[RT].BlendDesc = Desc.Desc.BlendState.RenderTarget[RT];
 #endif
-	{
-		return FCrc::MemCrc32(Data, NumBytes);
 	}
-#endif
-}
 
-SIZE_T FD3D12PipelineStateCacheBase::HashPSODesc(const FD3D12LowLevelGraphicsPipelineStateDesc& Desc)
-{
-	__declspec(align(32)) FD3D12LowLevelGraphicsPipelineStateDesc Hash;
-	FMemory::Memcpy(&Hash, &Desc, sizeof(Hash)); // Memcpy to avoid introducing garbage due to alignment
-
-	Hash.Desc.VS.pShaderBytecode = nullptr; // null out pointers so stale ones don't ruin the hash
-	Hash.Desc.PS.pShaderBytecode = nullptr;
-	Hash.Desc.HS.pShaderBytecode = nullptr;
-	Hash.Desc.DS.pShaderBytecode = nullptr;
-	Hash.Desc.GS.pShaderBytecode = nullptr;
-	Hash.Desc.InputLayout.pInputElementDescs = nullptr;
-	Hash.Desc.CachedPSO.pCachedBlob = nullptr;
-	Hash.Desc.CachedPSO.CachedBlobSizeInBytes = 0;
-	Hash.CombinedHash = 0;
-	Hash.Desc.pRootSignature = nullptr;
-	Hash.pRootSignature = nullptr;
-
-	return SIZE_T(HashData(&Hash, sizeof(Hash)));
+	return HashData(Data, TotalDataSize);
 }
 
-SIZE_T FD3D12PipelineStateCacheBase::HashPSODesc(const FD3D12ComputePipelineStateDesc& Desc)
+uint64 FD3D12PipelineStateCacheBase::HashPSODesc(const FD3D12ComputePipelineStateDesc& Desc)
 {
-	__declspec(align(32)) FD3D12ComputePipelineStateDesc Hash;
-	FMemory::Memcpy(&Hash, &Desc, sizeof(Hash)); // Memcpy to avoid introducing garbage due to alignment
-
-	Hash.Desc.CS.pShaderBytecode = nullptr;  // null out pointers so stale ones don't ruin the hash
-	Hash.Desc.CachedPSO.pCachedBlob = nullptr;
-	Hash.Desc.CachedPSO.CachedBlobSizeInBytes = 0;
-	Hash.CombinedHash = 0;
-	Hash.Desc.pRootSignature = nullptr;
-	Hash.pRootSignature = nullptr;
-
-	return SIZE_T(HashData(&Hash, sizeof(Hash)));
+	struct ComputePSOData
+	{
+		ShaderBytecodeHash CSHash;
+		UINT NodeMask;
+		D3D12_PIPELINE_STATE_FLAGS Flags;
+	};
+	ComputePSOData Data;
+	FMemory::Memzero(Data);
+	Data.CSHash   = Desc.CSHash;
+	Data.NodeMask = Desc.Desc.NodeMask;
+	Data.Flags    = Desc.Desc.Flags;
+
+	return HashData(&Data, sizeof(Data));
 }
 
-#define SSE4_2     0x100000 
-#define SSE4_CPUID_ARRAY_INDEX 2
-
-FD3D12PipelineStateCacheBase::FD3D12PipelineStateCacheBase(FD3D12Adapter* InParent) :
-	FD3D12AdapterChild(InParent)
+FD3D12PipelineStateCacheBase::FD3D12PipelineStateCacheBase(FD3D12Adapter* InParent)
+	: FD3D12AdapterChild(InParent)
 {
-#if defined(_WIN64) || (defined(_M_IX86_FP) && _M_IX86_FP >= 2)
-	// Check for SSE4 support see: https://msdn.microsoft.com/en-us/library/vstudio/hskdteyh(v=vs.100).aspx
-	{
-		int32 cpui[4];
-		__cpuidex(cpui, 1, 0);
-		GCPUSupportsSSE4 = !!(cpui[SSE4_CPUID_ARRAY_INDEX] & SSE4_2);
-	}
-#else
-	GCPUSupportsSSE4 = false;
-#endif
 }
 
 FD3D12PipelineStateCacheBase::~FD3D12PipelineStateCacheBase()
diff --git a/Engine/Source/Runtime/D3D12RHI/Private/D3D12Shader.h b/Engine/Source/Runtime/D3D12RHI/Private/D3D12Shader.h
--- a/Engine/Source/Runtime/D3D12RHI/Private/D3D12Shader.h
+++ b/Engine/Source/Runtime/D3D12RHI/Private/D3D12Shader.h
@@ -23,10 +23,12 @@ public:
 	FD3D12VertexElements VertexElements;
 
 	uint16 StreamStrides[MaxVertexElementCount];
+	uint32 Hash;
 
 	/** Initialization constructor. */
-	explicit FD3D12VertexDeclaration(const FD3D12VertexElements& InElements, const uint16* InStrides)
+	explicit FD3D12VertexDeclaration(const FD3D12VertexElements& InElements, const uint16* InStrides, const uint32 InHash)
 		: VertexElements(InElements)
+		, Hash(InHash)
 	{
 		FMemory::Memcpy(StreamStrides, InStrides, sizeof(StreamStrides));
 	}
diff --git a/Engine/Source/Runtime/D3D12RHI/Private/D3D12VertexDeclaration.cpp b/Engine/Source/Runtime/D3D12RHI/Private/D3D12VertexDeclaration.cpp
--- a/Engine/Source/Runtime/D3D12RHI/Private/D3D12VertexDeclaration.cpp
+++ b/Engine/Source/Runtime/D3D12RHI/Private/D3D12VertexDeclaration.cpp
@@ -118,14 +118,14 @@ struct FVertexDeclarationCache
 		return Cache.Add(InKey, InValue);
 	}
 
-	FORCEINLINE FVertexDeclarationRHIRef* FindOrAdd(const FD3D12VertexDeclarationKey& InKey, const FVertexDeclarationRHIRef& InValue)
+	FORCEINLINE FVertexDeclarationRHIRef* FindOrAdd(const FD3D12VertexDeclarationKey& InKey)
 	{
 		FScopeLock RWGuard(&LockGuard);
 
 		FVertexDeclarationRHIRef* VertexDeclarationRefPtr = Cache.Find(InKey);
 		if (VertexDeclarationRefPtr == nullptr)
 		{
-			VertexDeclarationRefPtr = &Cache.Add(InKey, InValue);
+			VertexDeclarationRefPtr = &Cache.Add(InKey, new FD3D12VertexDeclaration(InKey.VertexElements, InKey.StreamStrides, InKey.Hash));
 		}
 
 		return VertexDeclarationRefPtr;
@@ -143,7 +143,7 @@ FVertexDeclarationRHIRef FD3D12DynamicRHI::RHICreateVertexDeclaration(const FVer
 	FD3D12VertexDeclarationKey Key(Elements);
 
 	// Check for a cached vertex declaration. Add to the cache if it doesn't exist.
-	FVertexDeclarationRHIRef* VertexDeclarationRefPtr = GVertexDeclarationCache.FindOrAdd(Key, new FD3D12VertexDeclaration(Key.VertexElements, Key.StreamStrides));
+	FVertexDeclarationRHIRef* VertexDeclarationRefPtr = GVertexDeclarationCache.FindOrAdd(Key);
 
 	// The cached declaration must match the input declaration!
 	check(VertexDeclarationRefPtr);
diff --git a/Engine/Source/Runtime/D3D12RHI/Private/Windows/WindowsD3D12PipelineState.cpp b/Engine/Source/Runtime/D3D12RHI/Private/Windows/WindowsD3D12PipelineState.cpp
--- a/Engine/Source/Runtime/D3D12RHI/Private/Windows/WindowsD3D12PipelineState.cpp
+++ b/Engine/Source/Runtime/D3D12RHI/Private/Windows/WindowsD3D12PipelineState.cpp
@@ -644,21 +644,21 @@ void CreatePipelineStateFromStream(ID3D12PipelineState*& PSO, ID3D12Device2* Dev
 	check(PSO);
 }
 
-static inline void FastHashName(wchar_t Name[9], uint32 Hash)
+static inline void FastHashName(wchar_t Name[17], uint64 Hash)
 {
-	for (int32 i = 0; i < 8; i++)
+	for (int32 i = 0; i < 16; i++)
 	{
 		Name[i] = (Hash & 0xF) + 'A';
 		Hash >>= 4;
 	}
-	Name[8] = 0;
+	Name[16] = 0;
 }
 
 template <typename TDesc>
 static void CreatePipelineStateWrapper(ID3D12PipelineState** PSO, FD3D12Adapter* Adapter, const TDesc* CreationArgs)
 {
 	// Get the pipeline state name, currently based on the hash.
-	wchar_t Name[9];
+	wchar_t Name[17];
 	FastHashName(Name, CreationArgs->Desc->CombinedHash);
 
 #if LOG_PSO_CREATES
diff --git a/Engine/Source/Runtime/D3D12RHI/Public/D3D12Util.h b/Engine/Source/Runtime/D3D12RHI/Public/D3D12Util.h
--- a/Engine/Source/Runtime/D3D12RHI/Public/D3D12Util.h
+++ b/Engine/Source/Runtime/D3D12RHI/Public/D3D12Util.h
@@ -303,25 +303,16 @@ private:
 //==================================================================================================================================
 struct ShaderBytecodeHash
 {
-	// 160 bit strong SHA1 hash
-	uint32 SHA1Hash[5];
+	uint64 Hash[2];
 
 	bool operator ==(const ShaderBytecodeHash &b) const
 	{
-		return (SHA1Hash[0] == b.SHA1Hash[0] &&
-			SHA1Hash[1] == b.SHA1Hash[1] &&
-			SHA1Hash[2] == b.SHA1Hash[2] &&
-			SHA1Hash[3] == b.SHA1Hash[3] &&
-			SHA1Hash[4] == b.SHA1Hash[4]);
+		return (Hash[0] == b.Hash[0] && Hash[1] == b.Hash[1]);
 	}
 
 	bool operator !=(const ShaderBytecodeHash &b) const
 	{
-		return (SHA1Hash[0] != b.SHA1Hash[0] ||
-			SHA1Hash[1] != b.SHA1Hash[1] ||
-			SHA1Hash[2] != b.SHA1Hash[2] ||
-			SHA1Hash[3] != b.SHA1Hash[3] ||
-			SHA1Hash[4] != b.SHA1Hash[4]);
+		return (Hash[0] != b.Hash[0] || Hash[1] != b.Hash[1]);
 	}
 };
 
@@ -330,12 +321,12 @@ class FD3D12ShaderBytecode
 public:
 	FD3D12ShaderBytecode()
 	{
-		FMemory::Memzero(&Shader, sizeof(Shader));
-		FMemory::Memset(&Hash, 0, sizeof(Hash));
+		FMemory::Memzero(Shader);
+		FMemory::Memzero(Hash);
 	}
 
-	FD3D12ShaderBytecode(const D3D12_SHADER_BYTECODE &InShader) :
-		Shader(InShader)
+	FD3D12ShaderBytecode(const D3D12_SHADER_BYTECODE &InShader)
+		: Shader(InShader)
 	{
 		HashShader();
 	}
@@ -352,10 +343,15 @@ public:
 private:
 	void HashShader()
 	{
-		FMemory::Memset(&Hash, 0, sizeof(Hash));
 		if (Shader.pShaderBytecode && Shader.BytecodeLength > 0)
 		{
-			FSHA1::HashBuffer(Shader.pShaderBytecode, Shader.BytecodeLength, (uint8*)Hash.SHA1Hash);
+			// D3D shader bytecode contains a 128bit checksum in DWORD 1-4. We can just use that directly instead of hashing the whole shader bytecode ourselves.
+			check(Shader.BytecodeLength >= sizeof(uint32) + sizeof(Hash));
+			FMemory::Memcpy(&Hash, ((uint32*) Shader.pShaderBytecode) + 1, sizeof(Hash));
+		}
+		else
+		{
+			FMemory::Memzero(Hash);
 		}
 	}
 
