Restore lost deprecation fix
#rnx

diff --git a/Engine/Source/Runtime/D3D12RHI/Private/D3D12CommandContext.cpp b/Engine/Source/Runtime/D3D12RHI/Private/D3D12CommandContext.cpp
--- a/Engine/Source/Runtime/D3D12RHI/Private/D3D12CommandContext.cpp
+++ b/Engine/Source/Runtime/D3D12RHI/Private/D3D12CommandContext.cpp
@@ -116,7 +116,7 @@ FD3D12CommandContext::FD3D12CommandContext(FD3D12Device* InParent, FD3D12SubAllo
 {
 	FMemory::Memzero(DirtyUniformBuffers);
 	FMemory::Memzero(BoundUniformBuffers);
-	for (int i = 0; i < ARRAY_COUNT(BoundUniformBufferRefs); i++)
+	for (int i = 0; i < UE_ARRAY_COUNT(BoundUniformBufferRefs); i++)
 	{
 		for (int j = 0; j < UE_ARRAY_COUNT(BoundUniformBufferRefs[i]); j++)
 		{
