Fix AMD GPU hang in DX12

diff --git a/Engine/Source/Runtime/D3D12RHI/Private/D3D12Commands.cpp b/Engine/Source/Runtime/D3D12RHI/Private/D3D12Commands.cpp
--- a/Engine/Source/Runtime/D3D12RHI/Private/D3D12Commands.cpp
+++ b/Engine/Source/Runtime/D3D12RHI/Private/D3D12Commands.cpp
@@ -289,21 +289,21 @@ void FD3D12CommandContext::RHITransitionResources(EResourceTransitionAccess Tran
 		else
 		{
 #if USE_D3D12RHI_RESOURCE_STATE_TRACKING
-			if ( TransitionType == EResourceTransitionAccess::EReadable )
+			if (TransitionType == EResourceTransitionAccess::EReadable)
 			{
 				const D3D12_COMMAND_LIST_TYPE CmdListType = CommandListHandle.GetCommandListType();
+
+				// Compute pipeline can't transition to graphics states such as PIXEL_SHADER_RESOURCE. Best bet is to transition to COMMON given that we're going to consume it on a different queue anyway.
+				// Technically we should be able to transition NON_PIXEL_SHADER_RESOURCE on ComputeToCompute cases, but it appears an AMD driver issue is causing that to hang the GPU, so we're using COMMON
+				// in that case also, which is not ideal, but avoids the hang.
+				D3D12_RESOURCE_STATES AfterState = (CmdListType == D3D12_COMMAND_LIST_TYPE_DIRECT && TransitionPipeline == EResourceTransitionPipeline::EComputeToGfx)?
+					D3D12_RESOURCE_STATE_PIXEL_SHADER_RESOURCE | D3D12_RESOURCE_STATE_NON_PIXEL_SHADER_RESOURCE : D3D12_RESOURCE_STATE_COMMON;
+
 				for (int32 i = 0; i < InNumUAVs; ++i)
 				{
 					if (InUAVs[i])
 					{
 						FD3D12UnorderedAccessView* const UnorderedAccessView = RetrieveObject<FD3D12UnorderedAccessView>(InUAVs[i]);
-						
-						// D3D12_RESOURCE_STATE_PIXEL_SHADER_RESOURCE cannot be used for compute command list so we exclude it here.
-						D3D12_RESOURCE_STATES AfterState = D3D12_RESOURCE_STATE_NON_PIXEL_SHADER_RESOURCE;
-						if (CmdListType != D3D12_COMMAND_LIST_TYPE::D3D12_COMMAND_LIST_TYPE_COMPUTE)
-						{
-							AfterState |= D3D12_RESOURCE_STATE_PIXEL_SHADER_RESOURCE;
-						}
 						FD3D12DynamicRHI::TransitionResource(CommandListHandle, UnorderedAccessView, AfterState);
 					}
 				}
