Fix FrameFence so that RT and RHIT agrees on what frame we're on. This fixes the elusive DX12 mesh corruption bug.

diff --git a/Engine/Source/Runtime/D3D12RHI/Private/D3D12Adapter.h b/Engine/Source/Runtime/D3D12RHI/Private/D3D12Adapter.h
--- a/Engine/Source/Runtime/D3D12RHI/Private/D3D12Adapter.h
+++ b/Engine/Source/Runtime/D3D12RHI/Private/D3D12Adapter.h
@@ -222,9 +222,6 @@ public:
 		FRHIResourceCreateInfo& CreateInfo,
 		FRHIGPUMask GPUMask = FRHIGPUMask::All());
 
-	// Queue up a command to signal the frame fence on the command list. This should only be called from the rendering thread.
-	void SignalFrameFence_RenderThread(FRHICommandListImmediate& RHICmdList);
-
 	template<typename ObjectType, typename CreationCoreFunction>
 	inline ObjectType* CreateLinkedObject(FRHIGPUMask GPUMask, const CreationCoreFunction& pfnCreationCore)
 	{
diff --git a/Engine/Source/Runtime/D3D12RHI/Private/D3D12Allocation.cpp b/Engine/Source/Runtime/D3D12RHI/Private/D3D12Allocation.cpp
--- a/Engine/Source/Runtime/D3D12RHI/Private/D3D12Allocation.cpp
+++ b/Engine/Source/Runtime/D3D12RHI/Private/D3D12Allocation.cpp
@@ -1412,30 +1412,29 @@ FD3D12FastAllocatorPage* FD3D12FastAllocatorPagePool::RequestFastAllocatorPage()
 {
 	FD3D12Device* Device = GetParentDevice();
 	FD3D12Adapter* Adapter = Device->GetParentAdapter();
-	FD3D12Fence& Fence = Adapter->GetFrameFence();
-
-	FD3D12FastAllocatorPage* Page = nullptr;
+	FD3D12ManualFence& Fence = Adapter->GetFrameFence();
 
 	const uint64 CompletedFence = Fence.UpdateLastCompletedFence();
 
 	for (int32 Index = 0; Index < Pool.Num(); Index++)
 	{
+		FD3D12FastAllocatorPage* Page = Pool[Index];
+
 		//If the GPU is done with it and no-one has a lock on it
-		if (Pool[Index]->FastAllocBuffer->GetRefCount() == 1 &&
-			Pool[Index]->FrameFence <= CompletedFence)
+		if (Page->FastAllocBuffer->GetRefCount() == 1 &&
+			Page->FrameFence <= CompletedFence)
 		{
-			Page = Pool[Index];
 			Page->Reset();
 			Pool.RemoveAt(Index);
 			return Page;
 		}
 	}
 
-	check(Page == nullptr);
-	Page = new FD3D12FastAllocatorPage(PageSize);
+	FD3D12FastAllocatorPage* Page = new FD3D12FastAllocatorPage(PageSize);
 
 	const D3D12_RESOURCE_STATES InitialState = DetermineInitialResourceState(HeapProperties.Type, &HeapProperties);
 	VERIFYD3D12RESULT(Adapter->CreateBuffer(HeapProperties, GetGPUMask(), InitialState, PageSize, Page->FastAllocBuffer.GetInitReference(), TEXT("Fast Allocator Page")));
+	Page->FastAllocBuffer->DoNotDeferDelete();
 
 	Page->FastAllocData = Page->FastAllocBuffer->Map();
 
@@ -1445,7 +1444,7 @@ FD3D12FastAllocatorPage* FD3D12FastAllocatorPagePool::RequestFastAllocatorPage()
 void FD3D12FastAllocatorPagePool::ReturnFastAllocatorPage(FD3D12FastAllocatorPage* Page)
 {
 	FD3D12Adapter* Adapter = GetParentDevice()->GetParentAdapter();
-	FD3D12Fence& FrameFence = Adapter->GetFrameFence();
+	FD3D12ManualFence& FrameFence = Adapter->GetFrameFence();
 
 	// Extend the lifetime of these resources when in AFR as other nodes might be relying on this
 	Page->FrameFence = FrameFence.GetCurrentFence();
@@ -1460,18 +1459,19 @@ void FD3D12FastAllocatorPagePool::CleanupPages(uint64 FrameLag)
 	}
 
 	FD3D12Adapter* Adapter = GetParentDevice()->GetParentAdapter();
-	FD3D12Fence& FrameFence = Adapter->GetFrameFence();
+	FD3D12ManualFence& FrameFence = Adapter->GetFrameFence();
 
 	const uint64 CompletedFence = FrameFence.UpdateLastCompletedFence();
 
 	// Pages get returned to end of list, so we'll look for pages to delete, starting from the LRU
 	for (int32 Index = 0; Index < Pool.Num(); Index++)
 	{
+		FD3D12FastAllocatorPage* Page = Pool[Index];
+
 		//If the GPU is done with it and no-one has a lock on it
-		if (Pool[Index]->FastAllocBuffer->GetRefCount() == 1 &&
-			Pool[Index]->FrameFence + FrameLag <= CompletedFence)
+		if (Page->FastAllocBuffer->GetRefCount() == 1 &&
+			Page->FrameFence + FrameLag <= CompletedFence)
 		{
-			FD3D12FastAllocatorPage* Page = Pool[Index];
 			Pool.RemoveAt(Index);
 			delete(Page);
 
diff --git a/Engine/Source/Runtime/D3D12RHI/Private/D3D12CommandContext.cpp b/Engine/Source/Runtime/D3D12RHI/Private/D3D12CommandContext.cpp
--- a/Engine/Source/Runtime/D3D12RHI/Private/D3D12CommandContext.cpp
+++ b/Engine/Source/Runtime/D3D12RHI/Private/D3D12CommandContext.cpp
@@ -466,12 +466,6 @@ void FD3D12CommandContextBase::RHIEndFrame()
 	// Stop Timing at the very last moment
     
 	ParentAdapter->GetGPUProfiler().EndFrame(ParentAdapter->GetOwningRHI());
-    
-    
-	// Advance frame fence
-    
-	FD3D12ManualFence& FrameFence = ParentAdapter->GetFrameFence();
-	FrameFence.Signal(ED3D12CommandQueueType::Default, FrameFence.IncrementCurrentFence());
 }
 
 void FD3D12CommandContextBase::UpdateMemoryStats()
diff --git a/Engine/Source/Runtime/D3D12RHI/Private/D3D12DirectCommandListManager.h b/Engine/Source/Runtime/D3D12RHI/Private/D3D12DirectCommandListManager.h
--- a/Engine/Source/Runtime/D3D12RHI/Private/D3D12DirectCommandListManager.h
+++ b/Engine/Source/Runtime/D3D12RHI/Private/D3D12DirectCommandListManager.h
@@ -93,7 +93,7 @@ public:
 	// Avoids calling GetCompletedValue().
 	bool IsFenceCompleteFast(uint64 FenceValue) const { return FenceValue <= LastCompletedFence; }
 
-	uint64 GetCurrentFence() const { return CurrentFence; }
+	virtual uint64 GetCurrentFence() const { return CurrentFence; }
 	uint64 GetLastSignaledFence() const { return LastSignaledFence; }
 
 	uint64 PeekLastCompletedFence() const;
@@ -119,7 +119,8 @@ protected:
 	FD3D12FenceCore* FenceCores[MAX_NUM_GPUS];
 };
 
-// Fence value must be incremented manually. Useful when you need incrementing and signaling to happen at different times.
+// Fence value must be incremented manually. Useful when you need incrementing and signaling to happen at different times, e.g. for a FrameFence,
+// where the RenderThread increments and the RHIThread signals, so that both are in agreement on which frame they are on.
 class FD3D12ManualFence : public FD3D12Fence
 {
 public:
@@ -127,11 +128,20 @@ public:
 		: FD3D12Fence(InParent, InGPUMask, InName)
 	{}
 
+	uint64 GetCurrentFence() const final override
+	{
+		return IsInRHIThread()? LastSignaledFence + 1 : CurrentFence;
+	}
+
 	// Signals the specified fence value.
 	uint64 Signal(ED3D12CommandQueueType InQueueType, uint64 FenceToSignal);
 
 	// Increments the current fence and returns the previous value.
-	inline uint64 IncrementCurrentFence() { return CurrentFence++; }
+	inline uint64 IncrementCurrentFence()
+	{
+		check(IsInRenderingThread());
+		return CurrentFence++;
+	}
 };
 
 class FD3D12CommandAllocatorManager : public FD3D12DeviceChild
diff --git a/Engine/Source/Runtime/D3D12RHI/Private/D3D12RHIPrivate.h b/Engine/Source/Runtime/D3D12RHI/Private/D3D12RHIPrivate.h
--- a/Engine/Source/Runtime/D3D12RHI/Private/D3D12RHIPrivate.h
+++ b/Engine/Source/Runtime/D3D12RHI/Private/D3D12RHIPrivate.h
@@ -346,6 +346,7 @@ public:
 	virtual uint32 RHIGetViewportNextPresentGPUIndex(FRHIViewport* Viewport) final override;
 	virtual FTexture2DRHIRef RHIGetViewportBackBuffer(FRHIViewport* Viewport) final override;
 	virtual void RHIAliasTextureResources(FRHITexture* DestTexture, FRHITexture* SrcTexture) final override;
+	virtual void RHIAdvanceFrameFence() final override;
 	virtual void RHIAdvanceFrameForGetViewportBackBuffer(FRHIViewport* Viewport) final override;
 	virtual void RHIAcquireThreadOwnership() final override;
 	virtual void RHIReleaseThreadOwnership() final override;
diff --git a/Engine/Source/Runtime/D3D12RHI/Private/D3D12Texture.cpp b/Engine/Source/Runtime/D3D12RHI/Private/D3D12Texture.cpp
--- a/Engine/Source/Runtime/D3D12RHI/Private/D3D12Texture.cpp
+++ b/Engine/Source/Runtime/D3D12RHI/Private/D3D12Texture.cpp
@@ -1256,11 +1256,16 @@ FTexture2DRHIRef FD3D12DynamicRHI::RHIAsyncCreateTexture2D(uint32 SizeX, uint32
 
 			FastAllocator.Allocate(Size, D3D12_TEXTURE_DATA_PLACEMENT_ALIGNMENT, &TempResourceLocation);
 			FastAllocator.Allocate(SizeLowMips, D3D12_TEXTURE_DATA_PLACEMENT_ALIGNMENT, &TempResourceLocationLowMips);
+			TempResourceLocationLowMips.GetResource()->AddRef();
 		}
 		else
 		{
 			FastAllocator.Allocate(Size, D3D12_TEXTURE_DATA_PLACEMENT_ALIGNMENT, &TempResourceLocation);
 		}
+		// We AddRef() the resource here to make sure it doesn't get recycled prematurely. We are likely to be done with it during the frame,
+		// but lifetime of the allocation is not strictly tied to the frame because we're using the copy queue here. Because we're waiting
+		// on the GPU before returning here, this protection is safe, even if we end up straddling frame boundaries.
+		TempResourceLocation.GetResource()->AddRef();
 
 		FD3D12Texture2D* CurrentTexture = TextureOut;
 		while (CurrentTexture != nullptr)
@@ -1316,6 +1321,12 @@ FTexture2DRHIRef FD3D12DynamicRHI::RHIAsyncCreateTexture2D(uint32 SizeX, uint32
 		}
 
 		FD3D12TextureStats::D3D12TextureAllocated(*TextureOut);
+
+		// These are clear to be recycled now because GPU is done with it at this point. We wait on GPU in ExecuteCommandList() above.
+		TempResourceLocation.GetResource()->Release();
+		if (TempResourceLocationLowMips.GetResource())
+			TempResourceLocationLowMips.GetResource()->Release();
+
 	}
 
 	if (TempBufferSize != ZeroBufferSize)
diff --git a/Engine/Source/Runtime/D3D12RHI/Private/D3D12Viewport.cpp b/Engine/Source/Runtime/D3D12RHI/Private/D3D12Viewport.cpp
--- a/Engine/Source/Runtime/D3D12RHI/Private/D3D12Viewport.cpp
+++ b/Engine/Source/Runtime/D3D12RHI/Private/D3D12Viewport.cpp
@@ -913,6 +913,46 @@ void FD3D12CommandContextBase::RHIEndDrawingViewport(FRHIViewport* ViewportRHI,
 	}
 }
 
+struct FRHICommandSignalFrameFence final : public FRHICommand<FRHICommandSignalFrameFence>
+{
+	ED3D12CommandQueueType QueueType;
+	FD3D12ManualFence* const Fence;
+	const uint64 Value;
+	FORCEINLINE_DEBUGGABLE FRHICommandSignalFrameFence(ED3D12CommandQueueType InQueueType, FD3D12ManualFence* InFence, uint64 InValue)
+		: QueueType(InQueueType)
+		, Fence(InFence)
+		, Value(InValue)
+	{ 
+	}
+
+	void Execute(FRHICommandListBase& CmdList)
+	{
+		Fence->Signal(QueueType, Value);
+		check(Fence->GetLastSignaledFence() == Value);
+	}
+};
+
+void FD3D12DynamicRHI::RHIAdvanceFrameFence()
+{
+	check(IsInRenderingThread());
+
+	// Increment the current fence (on render thread timeline).
+	FD3D12ManualFence* FrameFence = &GetAdapter().GetFrameFence();
+	const uint64 PreviousFence = FrameFence->IncrementCurrentFence();
+
+	// Queue a command to signal on RHI thread that the current frame is a complete on the GPU.
+	FRHICommandListImmediate& RHICmdList = FRHICommandListExecutor::GetImmediateCommandList();
+	if (RHICmdList.Bypass() || !IsRunningRHIInSeparateThread())
+	{
+		FRHICommandSignalFrameFence Cmd(ED3D12CommandQueueType::Default, FrameFence, PreviousFence);
+		Cmd.Execute(RHICmdList);
+	}
+	else
+	{
+		ALLOC_COMMAND_CL(RHICmdList, FRHICommandSignalFrameFence)(ED3D12CommandQueueType::Default, FrameFence, PreviousFence);
+	}
+}
+
 void FD3D12DynamicRHI::RHIAdvanceFrameForGetViewportBackBuffer(FRHIViewport* ViewportRHI)
 {
 	check(IsInRenderingThread());
diff --git a/Engine/Source/Runtime/RHI/Private/RHICommandList.cpp b/Engine/Source/Runtime/RHI/Private/RHICommandList.cpp
--- a/Engine/Source/Runtime/RHI/Private/RHICommandList.cpp
+++ b/Engine/Source/Runtime/RHI/Private/RHICommandList.cpp
@@ -1575,7 +1575,11 @@ void FRHICommandList::EndFrame()
 		GetContext().RHIEndFrame();
 		return;
 	}
+
 	ALLOC_COMMAND(FRHICommandEndFrame)();
+
+	FRHICommandListExecutor::GetImmediateCommandList().AdvanceFrameFence();
+
 	if (!IsRunningRHIInSeparateThread())
 	{
 		// if we aren't running an RHIThread, there is no good reason to buffer this frame advance stuff and that complicates state management, so flush everything out now
diff --git a/Engine/Source/Runtime/RHI/Public/DynamicRHI.h b/Engine/Source/Runtime/RHI/Public/DynamicRHI.h
--- a/Engine/Source/Runtime/RHI/Public/DynamicRHI.h
+++ b/Engine/Source/Runtime/RHI/Public/DynamicRHI.h
@@ -945,6 +945,8 @@ public:
 		checkNoEntry();
 	}
 
+	virtual void RHIAdvanceFrameFence(){};
+
 	// Only relevant with an RHI thread, this advances the backbuffer for the purpose of GetViewportBackBuffer
 	// FlushType: Thread safe
 	virtual void RHIAdvanceFrameForGetViewportBackBuffer(FRHIViewport* Viewport) = 0;
@@ -1389,6 +1391,11 @@ FORCEINLINE FTexture2DRHIRef RHIGetFMaskTexture(FRHITexture* SourceTextureRHI)
 	return GDynamicRHI->RHIGetFMaskTexture(SourceTextureRHI);
 }
 
+FORCEINLINE void RHIAdvanceFrameFence()
+{
+	return GDynamicRHI->RHIAdvanceFrameFence();
+}
+
 FORCEINLINE void RHIAdvanceFrameForGetViewportBackBuffer(FRHIViewport* Viewport)
 {
 	return GDynamicRHI->RHIAdvanceFrameForGetViewportBackBuffer(Viewport);
diff --git a/Engine/Source/Runtime/RHI/Public/RHICommandList.h b/Engine/Source/Runtime/RHI/Public/RHICommandList.h
--- a/Engine/Source/Runtime/RHI/Public/RHICommandList.h
+++ b/Engine/Source/Runtime/RHI/Public/RHICommandList.h
@@ -4292,7 +4292,12 @@ public:
 	{
 		return RHIGetViewportBackBuffer(Viewport);
 	}
-	
+
+	FORCEINLINE void AdvanceFrameFence()
+	{
+		RHIAdvanceFrameFence();
+	}
+
 	FORCEINLINE void AdvanceFrameForGetViewportBackBuffer(FRHIViewport* Viewport)
 	{
 		return RHIAdvanceFrameForGetViewportBackBuffer(Viewport);
