Improve FastConstantAllocator. Faster, uses less memory, does not crash GPU.

diff --git a/Engine/Source/Runtime/D3D12RHI/Private/D3D12Adapter.cpp b/Engine/Source/Runtime/D3D12RHI/Private/D3D12Adapter.cpp
--- a/Engine/Source/Runtime/D3D12RHI/Private/D3D12Adapter.cpp
+++ b/Engine/Source/Runtime/D3D12RHI/Private/D3D12Adapter.cpp
@@ -6,13 +6,6 @@ D3D12Adapter.cpp:D3D12 Adapter implementation.
 
 #include "D3D12RHIPrivate.h"
 
-static TAutoConsoleVariable<int32> CVarTransientUniformBufferAllocatorSizeKB(
-	TEXT("D3D12.TransientUniformBufferAllocatorSizeKB"),
-	10 * 1024,
-	TEXT(""),
-	ECVF_ReadOnly
-);
-
 #if ENABLE_RESIDENCY_MANAGEMENT
 bool GEnableResidencyManagement = true;
 static TAutoConsoleVariable<int32> CVarResidencyManagement(
@@ -671,8 +664,7 @@ FD3D12FastConstantAllocator& FD3D12Adapter::GetTransientUniformBufferAllocator()
 	// Multi-GPU support : is using device 0 always appropriate here?
 	return FTransientUniformBufferAllocator::Get([this]() -> FTransientUniformBufferAllocator*
 	{
-		FTransientUniformBufferAllocator* Alloc = new FTransientUniformBufferAllocator(Devices[0], FRHIGPUMask::All(), CVarTransientUniformBufferAllocatorSizeKB.GetValueOnAnyThread() * 1024);
-		Alloc->Init();
+		FTransientUniformBufferAllocator* Alloc = new FTransientUniformBufferAllocator(Devices[0], FRHIGPUMask::All());
 		return Alloc;
 	});
 }
diff --git a/Engine/Source/Runtime/D3D12RHI/Private/D3D12Allocation.cpp b/Engine/Source/Runtime/D3D12RHI/Private/D3D12Allocation.cpp
--- a/Engine/Source/Runtime/D3D12RHI/Private/D3D12Allocation.cpp
+++ b/Engine/Source/Runtime/D3D12RHI/Private/D3D12Allocation.cpp
@@ -1496,44 +1496,16 @@ void FD3D12FastAllocatorPagePool::Destroy()
 	Pool.Empty();
 }
 
-FD3D12FastConstantAllocator::FD3D12FastConstantAllocator(FD3D12Device* Parent, FRHIGPUMask VisibiltyMask, uint32 InPageSize)
+FD3D12FastConstantAllocator::FD3D12FastConstantAllocator(FD3D12Device* Parent, FRHIGPUMask VisibiltyMask)
 	: FD3D12DeviceChild(Parent)
 	, FD3D12MultiNodeGPUObject(Parent->GetGPUMask(), VisibiltyMask)
 	, UnderlyingResource(Parent)
-	, PageSize(InPageSize)
-	, RingBuffer(PageSize / D3D12_CONSTANT_BUFFER_DATA_PLACEMENT_ALIGNMENT)
+	, Offset(64 * 1024) // Initial offset is at end of page so that first Allocate() call triggers a page allocation
+	, PageSize(64 * 1024)
 {
 	check(PageSize % D3D12_CONSTANT_BUFFER_DATA_PLACEMENT_ALIGNMENT == 0);
 }
 
-void FD3D12FastConstantAllocator::Init()
-{
-	FD3D12Device* Device = GetParentDevice();
-	FD3D12Adapter* Adapter = Device->GetParentAdapter();
-	ReallocBuffer();
-
-	RingBuffer.SetFence(&Adapter->GetFrameFence());
-}
-
-void FD3D12FastConstantAllocator::ReallocBuffer()
-{
-	check(PageSize % D3D12_CONSTANT_BUFFER_DATA_PLACEMENT_ALIGNMENT == 0);
-
-	FD3D12Device* Device = GetParentDevice();
-	FD3D12Adapter* Adapter = Device->GetParentAdapter();
-
-	UnderlyingResource.Clear();
-	
-	FD3D12Resource* NewBuffer = nullptr;
-	VERIFYD3D12RESULT(Adapter->CreateBuffer(D3D12_HEAP_TYPE_UPLOAD,
-		GetGPUMask(),
-		GetVisibilityMask(),
-		D3D12_RESOURCE_STATE_GENERIC_READ,
-		PageSize, &NewBuffer, TEXT("FastConstantAllocator")));
-
-	UnderlyingResource.AsStandAlone(NewBuffer, PageSize);
-}
-
 #if USE_STATIC_ROOT_SIGNATURE
 void* FD3D12FastConstantAllocator::Allocate(uint32 Bytes, FD3D12ResourceLocation& OutLocation, FD3D12ConstantBufferView* OutCBView)
 #else
@@ -1542,56 +1514,18 @@ void* FD3D12FastConstantAllocator::Allocate(uint32 Bytes, FD3D12ResourceLocation
 {
 	check(Bytes <= PageSize);
 
-	// Check to make sure our assumption that we don't need a OutLocation.Clear() here is valid.
-	checkf(!OutLocation.IsValid(), TEXT("The supplied resource location already has a valid resource. You should Clear() it first or it may leak."));
-
-	// Align to a constant buffer block size
 	const uint32 AlignedSize = Align(Bytes, D3D12_CONSTANT_BUFFER_DATA_PLACEMENT_ALIGNMENT);
-	const uint64 Location = RingBuffer.Allocate(AlignedSize / D3D12_CONSTANT_BUFFER_DATA_PLACEMENT_ALIGNMENT);
-	if (Location == FD3D12AbstractRingBuffer::FailedReturnValue)
-	{
-		PageSize = Align(PageSize + (PageSize / 2), D3D12_CONSTANT_BUFFER_DATA_PLACEMENT_ALIGNMENT);
-		ReallocBuffer();
-		RingBuffer.Reset(PageSize / D3D12_CONSTANT_BUFFER_DATA_PLACEMENT_ALIGNMENT);
-
-		UE_LOG(LogD3D12RHI, Warning, TEXT("Constant Allocator had to grow! Consider making it larger to begin with. New size: %d bytes"), PageSize);
-
-#if USE_STATIC_ROOT_SIGNATURE
-		return Allocate(Bytes, OutLocation, OutCBView);
-#else
-		return Allocate(Bytes, OutLocation);
-#endif
-	}
-
-	// Useful when trying to tweak initial size
-	//UE_LOG(LogD3D12RHI, Warning, TEXT("Space Left. %d"), RingBuffer.GetSpaceLeft() * D3D12_CONSTANT_BUFFER_DATA_PLACEMENT_ALIGNMENT);
-
-	const uint64 Offset = Location * D3D12_CONSTANT_BUFFER_DATA_PLACEMENT_ALIGNMENT;
 
-#if 0  // used to detect problems with fencing 
-	static FThreadSafeCounter Cnt;
-	if (Cnt.Increment() % 1000 == 0)
+	if (Offset + AlignedSize > PageSize)
 	{
-		uint64 Block1Start = 0;
-		uint64 Block1Size = 0;
-		uint64 Block2Start = 0;
-		uint64 Block2Size = 0;
-		RingBuffer.GetOverwritableBlocks(Block1Start, Block1Size, Block2Start, Block2Size);
-		Block1Start *= D3D12_CONSTANT_BUFFER_DATA_PLACEMENT_ALIGNMENT;
-		Block1Size *= D3D12_CONSTANT_BUFFER_DATA_PLACEMENT_ALIGNMENT;
-		Block2Start *= D3D12_CONSTANT_BUFFER_DATA_PLACEMENT_ALIGNMENT;
-		Block2Size *= D3D12_CONSTANT_BUFFER_DATA_PLACEMENT_ALIGNMENT;
+		Offset = 0;
 
-		if (Block1Size)
-		{
-			FMemory::Memset((uint8*)UnderlyingResource.GetMappedBaseAddress() + Block1Start, 0xff, Block1Size);
-		}
-		if (Block2Size)
-		{
-			FMemory::Memset((uint8*)UnderlyingResource.GetMappedBaseAddress() + Block2Start, 0xff, Block2Size);
-		}
+		FD3D12Device* Device = GetParentDevice();
+		FD3D12Adapter* Adapter = Device->GetParentAdapter();
+
+		FD3D12DynamicHeapAllocator& Allocator = Adapter->GetUploadHeapAllocator(Device->GetGPUIndex());
+		Allocator.AllocUploadResource(PageSize, DEFAULT_CONTEXT_UPLOAD_POOL_ALIGNMENT, UnderlyingResource);
 	}
-#endif
 
 	OutLocation.AsFastAllocation(UnderlyingResource.GetResource(),
 		AlignedSize,
@@ -1605,6 +1539,9 @@ void* FD3D12FastConstantAllocator::Allocate(uint32 Bytes, FD3D12ResourceLocation
 		OutCBView->Create(UnderlyingResource.GetGPUVirtualAddress() + Offset, AlignedSize);
 	}
 #endif
+
+	Offset += AlignedSize;
+
 	return OutLocation.GetMappedBaseAddress();
 }
 
diff --git a/Engine/Source/Runtime/D3D12RHI/Private/D3D12Allocation.h b/Engine/Source/Runtime/D3D12RHI/Private/D3D12Allocation.h
--- a/Engine/Source/Runtime/D3D12RHI/Private/D3D12Allocation.h
+++ b/Engine/Source/Runtime/D3D12RHI/Private/D3D12Allocation.h
@@ -535,138 +535,10 @@ protected:
 	FCriticalSection CS;
 };
 
-class FD3D12AbstractRingBuffer
-{
-public:
-	FD3D12AbstractRingBuffer(uint64 BufferSize)
-		: Fence(nullptr)
-		, Size(BufferSize)
-		, Head(BufferSize)
-		, Tail(0)
-		, LastFence(0)
-	{}
-
-	static const uint64 FailedReturnValue = uint64(-1);
-
-	inline void Reset(uint64 NewSize)
-	{
-		Size = NewSize;
-		Head = Size;
-		Tail = 0;
-		LastFence = 0;
-		OutstandingAllocs.Empty();
-	}
-
-	inline void SetFence(FD3D12Fence* InFence)
-	{
-		Fence = InFence;
-		LastFence = 0;
-	}
-
-	inline const uint64 GetSpaceLeft() const { return Head - Tail; }
-
-#if 0  // used to detect problems with fencing 
-	void GetOverwritableBlocks(uint64& Block1Start, uint64& Block1Size, uint64& Block2Start, uint64& Block2Size) const
-	{
-		check(Head >= Tail);
-		check(Size >= Head - Tail);
-		uint64 Used = Size - (Head - Tail);
-
-		if (Used == Size)
-		{
-			Block1Start = 0;
-			Block1Size = 0;
-
-			Block2Start = 0;
-			Block2Size = 0;
-		}
-		else
-		{
-			uint64 PhysicalTail = Tail % Size;
-
-			if (PhysicalTail <= Used)
-			{
-				// there is only one block, it starts at PhysicalTail
-				Block1Start = 0;
-				Block1Size = 0;
-
-				Block2Start = PhysicalTail;
-				Block2Size = Size - Used;
-			}
-			else
-			{
-				Block1Start = 0;
-				Block1Size = PhysicalTail - Used;
-
-				Block2Start = PhysicalTail;
-				Block2Size = Size - PhysicalTail;
-			}
-		}
-	}
-#endif // used to detect problems with fencing 
-
-	inline uint64 Allocate(uint64 Count)
-	{
-		{
-			const uint64 LastCompletedFence = Fence->GetLastCompletedFenceFast();
-			// If progress has been made since we were here last
-			if (LastCompletedFence > LastFence)
-			{
-				LastFence = LastCompletedFence;
-
-				for (auto It = OutstandingAllocs.CreateIterator(); It; ++It)
-				{
-					if (It.Key() < LastCompletedFence)
-					{
-						Head += It.Value();
-						It.RemoveCurrent();
-					}
-				}
-			}
-		}
-
-		uint64 ReturnValue = FailedReturnValue;
-
-		uint64 PhysicalTail = Tail % Size;
-
-		if (PhysicalTail + Count > Size)
-		{
-			// Force the wrap around by simply allocating the difference
-			uint64 Padding = Allocate(Size - PhysicalTail);
-			if (Padding == FailedReturnValue)
-			{
-				return FailedReturnValue;
-			}
-
-			PhysicalTail = Tail % Size;
-		}
-
-		if (Tail + Count < Head)
-		{
-			ReturnValue = PhysicalTail;
-			Tail += Count;
-			OutstandingAllocs.FindOrAdd(Fence->GetCurrentFence()) += Count;
-		}
-
-		return ReturnValue;
-	}
-
-private:
-	FD3D12Fence* Fence;
-	uint64 Size;
-	uint64 Head;
-	uint64 Tail;
-	uint64 LastFence;
-
-	TMap<uint64, uint64, TInlineSetAllocator<16> > OutstandingAllocs;
-};
-
 class FD3D12FastConstantAllocator : public FD3D12DeviceChild, public FD3D12MultiNodeGPUObject
 {
 public:
-	FD3D12FastConstantAllocator(FD3D12Device* Parent, FRHIGPUMask VisibiltyMask, uint32 InPageSize);
-
-	void Init();
+	FD3D12FastConstantAllocator(FD3D12Device* Parent, FRHIGPUMask VisibiltyMask);
 
 #if USE_STATIC_ROOT_SIGNATURE
 	void* Allocate(uint32 Bytes, class FD3D12ResourceLocation& OutLocation, FD3D12ConstantBufferView* OutCBView);
@@ -676,12 +548,10 @@ public:
 
 
 private:
-	void ReallocBuffer();
-
 	FD3D12ResourceLocation UnderlyingResource;
 
+	uint32 Offset;
 	uint32 PageSize;
-	FD3D12AbstractRingBuffer RingBuffer;
 };
 
 //-----------------------------------------------------------------------------
diff --git a/Engine/Source/Runtime/D3D12RHI/Private/D3D12CommandContext.cpp b/Engine/Source/Runtime/D3D12RHI/Private/D3D12CommandContext.cpp
--- a/Engine/Source/Runtime/D3D12RHI/Private/D3D12CommandContext.cpp
+++ b/Engine/Source/Runtime/D3D12RHI/Private/D3D12CommandContext.cpp
@@ -22,13 +22,6 @@ static FAutoConsoleVariableRef CVarCommandListBatchingMode(
 	ECVF_RenderThreadSafe
 	);
 
-// These can be overridden with the cvars below
-namespace ConstantAllocatorSizesKB
-{
-	int32 DefaultGraphics = 3072; // x1
-	int32 Graphics = 3072; //x4
- 	int32 AsyncCompute = 3072; // x1
-}
 // We don't yet have a way to auto-detect that the Radeon Developer Panel is running
 // with profiling enabled, so for now, we have to manually toggle this console var.
 // It needs to be set before device creation, so it's read only.
@@ -40,40 +33,6 @@ static FAutoConsoleVariableRef CVarEmitRgpFrameMarkers(
 	ECVF_ReadOnly | ECVF_RenderThreadSafe
 	);
 
-static FAutoConsoleVariableRef CVarDefaultGfxCommandContextConstantAllocatorSizeKB(
-	TEXT("D3D12.DefaultGfxCommandContextConstantAllocatorSizeKB"),
-	ConstantAllocatorSizesKB::DefaultGraphics,
-	TEXT(""),
-	ECVF_ReadOnly
-);
-
-static FAutoConsoleVariableRef CVarGfxCommandContextConstantAllocatorSizeKB(
-	TEXT("D3D12.GfxCommandContextConstantAllocatorSizeKB"),
-	ConstantAllocatorSizesKB::Graphics,
-	TEXT(""),
-	ECVF_ReadOnly
-);
-
-static FAutoConsoleVariableRef CVarComputeCommandContextConstantAllocatorSizeKB(
-	TEXT("D3D12.ComputeCommandContextConstantAllocatorSizeKB"),
-	ConstantAllocatorSizesKB::AsyncCompute,
-	TEXT(""),
-	ECVF_ReadOnly
-);
-
-static uint32 GetConstantAllocatorSize(bool bIsAsyncComputeContext, bool bIsDefaultContext)
-{
-	if (bIsAsyncComputeContext)
-	{
-		return ConstantAllocatorSizesKB::AsyncCompute * 1024;
-	}
-	if (bIsDefaultContext)
-	{
-		return ConstantAllocatorSizesKB::DefaultGraphics * 1024;
-	}
-	return ConstantAllocatorSizesKB::Graphics * 1024;
-}
-
 FD3D12CommandContextBase::FD3D12CommandContextBase(class FD3D12Adapter* InParentAdapter, FRHIGPUMask InGPUMask, bool InIsDefaultContext, bool InIsAsyncComputeContext)
 	: FD3D12AdapterChild(InParentAdapter)
 	, GPUMask(InGPUMask)
@@ -87,7 +46,7 @@ FD3D12CommandContextBase::FD3D12CommandContextBase(class FD3D12Adapter* InParent
 FD3D12CommandContext::FD3D12CommandContext(FD3D12Device* InParent, FD3D12SubAllocatedOnlineHeap::SubAllocationDesc& SubHeapDesc, bool InIsDefaultContext, bool InIsAsyncComputeContext) :
 	FD3D12CommandContextBase(InParent->GetParentAdapter(), InParent->GetGPUMask(), InIsDefaultContext, InIsAsyncComputeContext),
 	FD3D12DeviceChild(InParent),
-	ConstantsAllocator(InParent, InParent->GetGPUMask(), GetConstantAllocatorSize(InIsAsyncComputeContext, InIsDefaultContext)),
+	ConstantsAllocator(InParent, InParent->GetGPUMask()),
 	CommandListHandle(),
 	CommandAllocator(nullptr),
 	CommandAllocatorManager(InParent, InIsAsyncComputeContext ? D3D12_COMMAND_LIST_TYPE_COMPUTE : D3D12_COMMAND_LIST_TYPE_DIRECT),
@@ -126,8 +85,6 @@ FD3D12CommandContext::FD3D12CommandContext(FD3D12Device* InParent, FD3D12SubAllo
 	FMemory::Memzero(CurrentRenderTargets);
 	FMemory::Memzero(CurrentUAVs);
 	StateCache.Init(GetParentDevice(), this, nullptr, SubHeapDesc);
-
-	ConstantsAllocator.Init();
 }
 
 FD3D12CommandContext::~FD3D12CommandContext()
