Fix tessellation on DX12

diff --git a/Engine/Source/Runtime/D3D12RHI/Private/D3D12StateCachePrivate.h b/Engine/Source/Runtime/D3D12RHI/Private/D3D12StateCachePrivate.h
--- a/Engine/Source/Runtime/D3D12RHI/Private/D3D12StateCachePrivate.h
+++ b/Engine/Source/Runtime/D3D12RHI/Private/D3D12StateCachePrivate.h
@@ -268,6 +268,11 @@ static inline D3D_PRIMITIVE_TOPOLOGY GetD3D12PrimitiveType(uint32 PrimitiveType,
 		D3D_PRIMITIVE_TOPOLOGY_LINELIST,                   // PT_LineList
 		0,                                                 // PT_QuadList
 		D3D_PRIMITIVE_TOPOLOGY_POINTLIST,                  // PT_PointList
+#if defined(D3D12RHI_PRIMITIVE_TOPOLOGY_RECTLIST)          // PT_RectList
+		D3D_PRIMITIVE_TOPOLOGY_RECTLIST,
+#else
+		0,
+#endif
 		D3D_PRIMITIVE_TOPOLOGY_1_CONTROL_POINT_PATCHLIST,  // PT_1_ControlPointPatchList
 		D3D_PRIMITIVE_TOPOLOGY_2_CONTROL_POINT_PATCHLIST,  // PT_2_ControlPointPatchList
 		D3D_PRIMITIVE_TOPOLOGY_3_CONTROL_POINT_PATCHLIST,  // PT_3_ControlPointPatchList
@@ -301,6 +306,7 @@ static inline D3D_PRIMITIVE_TOPOLOGY GetD3D12PrimitiveType(uint32 PrimitiveType,
 		D3D_PRIMITIVE_TOPOLOGY_31_CONTROL_POINT_PATCHLIST, // PT_31_ControlPointPatchList
 		D3D_PRIMITIVE_TOPOLOGY_32_CONTROL_POINT_PATCHLIST, // PT_32_ControlPointPatchList
 	};
+	static_assert(ARRAY_COUNT(D3D12PrimitiveType) == PT_Num, "Primitive lookup table is wrong size");
 
 	if (bUsingTessellation)
 	{
@@ -311,8 +317,6 @@ static inline D3D_PRIMITIVE_TOPOLOGY GetD3D12PrimitiveType(uint32 PrimitiveType,
 		}
 		else/* if (PrimitiveType < PT_1_ControlPointPatchList)*/
 		{
-			//UE_LOG(LogD3D12RHI, Fatal, TEXT("Invalid type specified for tessellated render, probably missing a case in FSkeletalMeshSceneProxy::DrawDynamicElementsByMaterial or FStaticMeshSceneProxy::GetMeshElement"));
-
 			checkf(PrimitiveType >= PT_1_ControlPointPatchList, TEXT("Invalid type specified for tessellated render, probably missing a case in FSkeletalMeshSceneProxy::DrawDynamicElementsByMaterial or FStaticMeshSceneProxy::GetMeshElement"));
 		}
 	}
diff --git a/Engine/Source/Runtime/RHI/Public/RHIUtilities.h b/Engine/Source/Runtime/RHI/Public/RHIUtilities.h
--- a/Engine/Source/Runtime/RHI/Public/RHIUtilities.h
+++ b/Engine/Source/Runtime/RHI/Public/RHIUtilities.h
@@ -830,7 +830,8 @@ inline void RHICreateTargetableShaderResource3D(
  */
 inline uint32 GetVertexCountForPrimitiveCount(uint32 NumPrimitives, uint32 PrimitiveType)
 {
-	uint32 Factor = (PrimitiveType == PT_TriangleList)? 3 : (PrimitiveType == PT_LineList)? 2 : (PrimitiveType >= PT_1_ControlPointPatchList)? (PrimitiveType - PT_1_ControlPointPatchList + 1) : 1;
+	static_assert(PT_Num == 38, "This function needs to be updated");
+	uint32 Factor = (PrimitiveType == PT_TriangleList)? 3 : (PrimitiveType == PT_LineList)? 2 : (PrimitiveType == PT_RectList)? 3 : (PrimitiveType >= PT_1_ControlPointPatchList)? (PrimitiveType - PT_1_ControlPointPatchList + 1) : 1;
 	uint32 Offset = (PrimitiveType == PT_TriangleStrip)? 2 : 0;
 
 	return NumPrimitives * Factor + Offset;
