Fix ARRAY_COUNT deprecations
#rnx

diff --git a/Engine/Source/Runtime/D3D12RHI/Private/D3D12StateCachePrivate.h b/Engine/Source/Runtime/D3D12RHI/Private/D3D12StateCachePrivate.h
--- a/Engine/Source/Runtime/D3D12RHI/Private/D3D12StateCachePrivate.h
+++ b/Engine/Source/Runtime/D3D12RHI/Private/D3D12StateCachePrivate.h
@@ -306,7 +306,7 @@ static inline D3D_PRIMITIVE_TOPOLOGY GetD3D12PrimitiveType(uint32 PrimitiveType,
 		D3D_PRIMITIVE_TOPOLOGY_31_CONTROL_POINT_PATCHLIST, // PT_31_ControlPointPatchList
 		D3D_PRIMITIVE_TOPOLOGY_32_CONTROL_POINT_PATCHLIST, // PT_32_ControlPointPatchList
 	};
-	static_assert(ARRAY_COUNT(D3D12PrimitiveType) == PT_Num, "Primitive lookup table is wrong size");
+	static_assert(UE_ARRAY_COUNT(D3D12PrimitiveType) == PT_Num, "Primitive lookup table is wrong size");
 
 	if (bUsingTessellation)
 	{
