Add size for R32G32_FLOAT

diff --git a/Engine/Source/Runtime/D3D12RHI/Private/D3D12RenderTarget.cpp b/Engine/Source/Runtime/D3D12RHI/Private/D3D12RenderTarget.cpp
--- a/Engine/Source/Runtime/D3D12RHI/Private/D3D12RenderTarget.cpp
+++ b/Engine/Source/Runtime/D3D12RHI/Private/D3D12RenderTarget.cpp
@@ -684,6 +684,9 @@ static uint32 ComputeBytesPerPixel(DXGI_FORMAT Format)
 	case DXGI_FORMAT_R32G32B32A32_FLOAT:
 		BytesPerPixel = 16;
 		break;
+	case DXGI_FORMAT_R32G32_FLOAT:
+		BytesPerPixel = 8;
+		break;
 	case DXGI_FORMAT_R8_TYPELESS:
 	case DXGI_FORMAT_R8_UNORM:
 	case DXGI_FORMAT_R8_UINT:
