Add missing VET_UInt case to FD3D12VertexDeclaration::GetInitializer()

diff --git a/Engine/Source/Runtime/D3D12RHI/Private/D3D12VertexDeclaration.cpp b/Engine/Source/Runtime/D3D12RHI/Private/D3D12VertexDeclaration.cpp
--- a/Engine/Source/Runtime/D3D12RHI/Private/D3D12VertexDeclaration.cpp
+++ b/Engine/Source/Runtime/D3D12RHI/Private/D3D12VertexDeclaration.cpp
@@ -169,26 +169,28 @@ bool FD3D12VertexDeclaration::GetInitializer(FVertexDeclarationElementList& Init
 		Out.Offset = Elem.AlignedByteOffset;
 		switch (Elem.Format)
 		{
-		case DXGI_FORMAT_R32_FLOAT: Out.Type = VET_Float1; break;
-		case DXGI_FORMAT_R32G32_FLOAT: Out.Type = VET_Float2; break;
-		case DXGI_FORMAT_R32G32B32_FLOAT: Out.Type = VET_Float3; break;
-		case DXGI_FORMAT_R32G32B32A32_FLOAT: Out.Type = VET_Float4; break;
-		case DXGI_FORMAT_R8G8B8A8_SNORM: Out.Type = VET_PackedNormal; break;
-		case DXGI_FORMAT_R8G8B8A8_UINT: Out.Type = VET_UByte4; break;
-		case DXGI_FORMAT_R8G8B8A8_UNORM: Out.Type = VET_UByte4N; break;
-		case DXGI_FORMAT_B8G8R8A8_UNORM: Out.Type = VET_Color; break;
-		case DXGI_FORMAT_R16G16_SINT: Out.Type = VET_Short2; break;
-		case DXGI_FORMAT_R16G16B16A16_SINT: Out.Type = VET_Short4; break;
-		case DXGI_FORMAT_R16G16_SNORM: Out.Type = VET_Short2N; break;
-		case DXGI_FORMAT_R16G16_FLOAT: Out.Type = VET_Half2; break;
-		case DXGI_FORMAT_R16G16B16A16_FLOAT: Out.Type = VET_Half4; break;
-		case DXGI_FORMAT_R16G16B16A16_SNORM: Out.Type = VET_Short4N; break;
-		case DXGI_FORMAT_R16G16_UINT: Out.Type = VET_UShort2; break;
-		case DXGI_FORMAT_R16G16B16A16_UINT: Out.Type = VET_UShort4; break;
-		case DXGI_FORMAT_R16G16_UNORM: Out.Type = VET_UShort2N; break;
-		case DXGI_FORMAT_R16G16B16A16_UNORM: Out.Type = VET_UShort4N; break;
-		case DXGI_FORMAT_R10G10B10A2_UNORM: Out.Type = VET_URGB10A2N; break;
-		default: UE_LOG(LogD3D12RHI, Fatal, TEXT("Unknown D3D vertex element type %u"), Elem.Format);
+		case DXGI_FORMAT_R32_FLOAT:          Out.Type = VET_Float1;       break;
+		case DXGI_FORMAT_R32G32_FLOAT:       Out.Type = VET_Float2;       break;
+		case DXGI_FORMAT_R32G32B32_FLOAT:    Out.Type = VET_Float3;       break;
+		case DXGI_FORMAT_R32G32B32A32_FLOAT: Out.Type = VET_Float4;       break;
+		case DXGI_FORMAT_R8G8B8A8_SNORM:     Out.Type = VET_PackedNormal; break;
+		case DXGI_FORMAT_R8G8B8A8_UINT:      Out.Type = VET_UByte4;       break;
+		case DXGI_FORMAT_R8G8B8A8_UNORM:     Out.Type = VET_UByte4N;      break;
+		case DXGI_FORMAT_B8G8R8A8_UNORM:     Out.Type = VET_Color;        break;
+		case DXGI_FORMAT_R16G16_SINT:        Out.Type = VET_Short2;       break;
+		case DXGI_FORMAT_R16G16B16A16_SINT:  Out.Type = VET_Short4;       break;
+		case DXGI_FORMAT_R16G16_SNORM:       Out.Type = VET_Short2N;      break;
+		case DXGI_FORMAT_R16G16_FLOAT:       Out.Type = VET_Half2;        break;
+		case DXGI_FORMAT_R16G16B16A16_FLOAT: Out.Type = VET_Half4;        break;
+		case DXGI_FORMAT_R16G16B16A16_SNORM: Out.Type = VET_Short4N;      break;
+		case DXGI_FORMAT_R16G16_UINT:        Out.Type = VET_UShort2;      break;
+		case DXGI_FORMAT_R16G16B16A16_UINT:  Out.Type = VET_UShort4;      break;
+		case DXGI_FORMAT_R16G16_UNORM:       Out.Type = VET_UShort2N;     break;
+		case DXGI_FORMAT_R16G16B16A16_UNORM: Out.Type = VET_UShort4N;     break;
+		case DXGI_FORMAT_R10G10B10A2_UNORM:  Out.Type = VET_URGB10A2N;    break;
+		case DXGI_FORMAT_R32_UINT:           Out.Type = VET_UInt;         break;
+		default:
+			UE_LOG(LogD3D12RHI, Fatal, TEXT("Unknown D3D vertex element type %u"), Elem.Format);
 		}
 		Out.AttributeIndex = Elem.SemanticIndex;
 		Out.Stride = StreamStrides[Elem.InputSlot];
