D3D12 root signature size warning is now only reported in verbose mode and only for non-RT root signatures.

diff --git a/Engine/Source/Runtime/D3D12RHI/Private/D3D12RootSignature.cpp b/Engine/Source/Runtime/D3D12RHI/Private/D3D12RootSignature.cpp
--- a/Engine/Source/Runtime/D3D12RHI/Private/D3D12RootSignature.cpp
+++ b/Engine/Source/Runtime/D3D12RHI/Private/D3D12RootSignature.cpp
@@ -258,9 +258,9 @@ FD3D12RootSignatureDesc::FD3D12RootSignatureDesc(const FD3D12QuantizedBoundShade
 	// Init the desc (warn about the size if necessary).
 #if !NO_LOGGING
 	const uint32 SizeWarningThreshold = 12;
-	if (RootParametersSize > SizeWarningThreshold)
+	if (RootParametersSize > SizeWarningThreshold && QBSS.RootSignatureType == RS_Raster)
 	{
-		UE_LOG(LogD3D12RHI, Display, TEXT("Root signature created where the root parameters take up %u DWORDS. Using more than %u DWORDs can negatively impact performance depending on the hardware and root parameter usage."), RootParametersSize, SizeWarningThreshold);
+		UE_LOG(LogD3D12RHI, Verbose, TEXT("Root signature created where the root parameters take up %u DWORDS. Using more than %u DWORDs can negatively impact performance depending on the hardware and root parameter usage."), RootParametersSize, SizeWarningThreshold);
 	}
 #endif
 	RootDesc.Init_1_1(RootParameterCount, TableSlots, 0, nullptr, Flags);
