Linked some of the dxr_todo-s to JIRAs. Removed few irrelevant ones. No functional changes.

diff --git a/Engine/Source/Developer/Windows/ShaderFormatD3D/Private/D3DShaderCompiler.cpp b/Engine/Source/Developer/Windows/ShaderFormatD3D/Private/D3DShaderCompiler.cpp
--- a/Engine/Source/Developer/Windows/ShaderFormatD3D/Private/D3DShaderCompiler.cpp
+++ b/Engine/Source/Developer/Windows/ShaderFormatD3D/Private/D3DShaderCompiler.cpp
@@ -1371,7 +1371,7 @@ static bool CompileAndProcessD3DShader(FString& PreprocessedShaderSource, const
 				// Handy place for a breakpoint for debugging...
 				++GBreakpoint;
 
-				// #dxr_todo: strip DXIL debug and reflection data
+				// #dxr_todo UE-68236: strip DXIL debug and reflection data
 				CompressedData = Shader;
 			}
 			else if (D3DStripShaderFunc)
diff --git a/Engine/Source/Runtime/D3D12RHI/Private/D3D12Adapter.cpp b/Engine/Source/Runtime/D3D12RHI/Private/D3D12Adapter.cpp
--- a/Engine/Source/Runtime/D3D12RHI/Private/D3D12Adapter.cpp
+++ b/Engine/Source/Runtime/D3D12RHI/Private/D3D12Adapter.cpp
@@ -511,7 +511,7 @@ void FD3D12Adapter::InitializeDevices()
 		ID3D12RootSignature* StaticGraphicsRS = (GetStaticGraphicsRootSignature()) ? GetStaticGraphicsRootSignature()->GetRootSignature() : nullptr;
 		ID3D12RootSignature* StaticComputeRS = (GetStaticComputeRootSignature()) ? GetStaticComputeRootSignature()->GetRootSignature() : nullptr;
 
-		// #dxr_todo: verify that disk cache works correctly with DXR
+		// #dxr_todo UE-68235: verify that disk cache works correctly with DXR
 		PipelineStateCache.RebuildFromDiskCache(StaticGraphicsRS, StaticComputeRS);
 	}
 }
diff --git a/Engine/Source/Runtime/D3D12RHI/Private/D3D12Device.cpp b/Engine/Source/Runtime/D3D12RHI/Private/D3D12Device.cpp
--- a/Engine/Source/Runtime/D3D12RHI/Private/D3D12Device.cpp
+++ b/Engine/Source/Runtime/D3D12RHI/Private/D3D12Device.cpp
@@ -50,7 +50,7 @@ FD3D12Device::FD3D12Device(FRHIGPUMask InGPUMask, FD3D12Adapter* InAdapter) :
 FD3D12Device::~FD3D12Device()
 {
 #if D3D12_RHI_RAYTRACING
-	DestroyRayTracingDescriptorCache(); // #dxr_todo: unify RT descriptor cache with main FD3D12DescriptorCache
+	DestroyRayTracingDescriptorCache(); // #dxr_todo UE-72158: unify RT descriptor cache with main FD3D12DescriptorCache
 #endif
 
 	// Cleanup the allocator near the end, as some resources may be returned to the allocator or references are shared by multiple GPUs
diff --git a/Engine/Source/Runtime/D3D12RHI/Private/D3D12Device.h b/Engine/Source/Runtime/D3D12RHI/Private/D3D12Device.h
--- a/Engine/Source/Runtime/D3D12RHI/Private/D3D12Device.h
+++ b/Engine/Source/Runtime/D3D12RHI/Private/D3D12Device.h
@@ -188,7 +188,7 @@ protected:
 #if D3D12_RHI_RAYTRACING
 	FD3D12BasicRayTracingPipeline* BasicRayTracingPipeline = nullptr;
 
-	// #dxr_todo: unify RT descriptor cache with main FD3D12DescriptorCache
+	// #dxr_todo UE-72158: unify RT descriptor cache with main FD3D12DescriptorCache
 	FD3D12RayTracingDescriptorHeapCache* RayTracingDescriptorHeapCache = nullptr;
 	void DestroyRayTracingDescriptorCache();
 #endif
diff --git a/Engine/Source/Runtime/D3D12RHI/Private/D3D12RayTracing.cpp b/Engine/Source/Runtime/D3D12RHI/Private/D3D12RayTracing.cpp
--- a/Engine/Source/Runtime/D3D12RHI/Private/D3D12RayTracing.cpp
+++ b/Engine/Source/Runtime/D3D12RHI/Private/D3D12RayTracing.cpp
@@ -271,7 +271,7 @@ static TRefCountPtr<ID3D12StateObject> CreateRayTracingStateObject(
 	return Result;
 }
 
-// #dxr_todo: FD3D12Device::GlobalViewHeap/GlobalSamplerHeap should be used instead of ad-hoc heaps here.
+// #dxr_todo UE-72158: FD3D12Device::GlobalViewHeap/GlobalSamplerHeap should be used instead of ad-hoc heaps here.
 // Unfortunately, this requires a major refactor of how global heaps work.
 // FD3D12CommandContext-s should not get static chunks of the global heap, but instead should dynamically allocate
 // chunks on as-needed basis and release them when possible.
@@ -348,7 +348,7 @@ public:
 		Desc.Flags = D3D12_DESCRIPTOR_HEAP_FLAG_SHADER_VISIBLE;
 		Desc.Type = Type;
 		Desc.NumDescriptors = NumDescriptors;
-		Desc.NodeMask = 1; // #dxr_todo: handle mGPU
+		Desc.NodeMask = 1; // #dxr_todo UE-72157: handle mGPU
 
 		ID3D12DescriptorHeap* D3D12Heap = nullptr;
 
@@ -578,7 +578,7 @@ public:
 		HitRecordStride = RoundUpToNextMultiple(HitRecordSizeUnaligned, D3D12_RAYTRACING_SHADER_RECORD_BYTE_ALIGNMENT);
 
 		// Minimum number of descriptors required to support binding global resources (arbitrarily chosen)
-		// #dxr_todo: Remove this when RT descriptors are sub-allocated from the global view descriptor heap.
+		// #dxr_todo UE-72158: Remove this when RT descriptors are sub-allocated from the global view descriptor heap.
 		const uint32 MinNumViewDescriptors = 1024;
 		const uint32 ApproximateDescriptorsPerRecord = 32; // #dxr_todo: calculate this based on shader reflection data
 
@@ -694,7 +694,7 @@ public:
 		FRHIResourceCreateInfo CreateInfo;
 		CreateInfo.ResourceArray = &Data;
 
-		checkf(GNumExplicitGPUsForRendering == 1, TEXT("Ray tracing is not implemented for mGPU")); // #dxr_todo: implement mGPU support
+		checkf(GNumExplicitGPUsForRendering == 1, TEXT("Ray tracing is not implemented for mGPU")); // #dxr_todo UE-72157: implement mGPU support
 		Buffer = Adapter->CreateRHIBuffer<FD3D12MemBuffer>(
 			nullptr, BufferDesc, BufferDesc.Alignment,
 			0, BufferDesc.Width, BUF_Static, CreateInfo,
@@ -768,7 +768,7 @@ public:
 
 #if ENABLE_RESIDENCY_MANAGEMENT
 	// A set of all resources referenced by this shader table for the purpose of updating residency before ray tracing work dispatch.
-	// #dxr_todo: remove resources from this set when SBT slot entries are replaced
+	// #dxr_todo UE-72159: remove resources from this set when SBT slot entries are replaced
 	TSet<FD3D12Resource*> ReferencedD3D12Resources;
 	TArray<TRefCountPtr<FRHIResource>> ReferencedResources;
 	void AddResourceReference(FD3D12Resource* D3D12Resource, FRHIResource* Resource)
@@ -921,7 +921,7 @@ public:
 		const uint32 HitShaderRootSignatureBaseIndex = LocalRootSignatures.Num();
 
 		// Initialize hit group shader libraries
-		// #dxr_todo: hit group libraries *also* could come from precompiled pipeline sub-objects, when those are supported in the future
+
 		HitGroupShaders.Reserve(InitializerHitGroups.Num());
 
 		// Each shader within RTPSO must have a unique name, therefore we must rename original shader entry points.
@@ -1267,7 +1267,7 @@ FRayTracingGeometryRHIRef FD3D12DynamicRHI::RHICreateRayTracingGeometry(const FR
 	checkf(Initializer.VertexBufferStride, TEXT("Position vertex buffer is required for ray tracing geometry"));
 	checkf(Initializer.VertexBufferStride % 4 == 0, TEXT("Position vertex buffer stride must be aligned to 4 bytes for ByteAddressBuffer loads to work"));
 
-	// #dxr_todo VET_Half4 (DXGI_FORMAT_R16G16B16A16_FLOAT) is also supported by DXR. Should we support it?
+	// #dxr_todo UE-72160: VET_Half4 (DXGI_FORMAT_R16G16B16A16_FLOAT) is also supported by DXR. Should we support it?
 	check(Initializer.VertexBufferElementType == VET_Float3 || Initializer.VertexBufferElementType == VET_Float2 || Initializer.VertexBufferElementType == VET_Half2);
 	if (Initializer.IndexBuffer)
 	{
@@ -1276,12 +1276,11 @@ FRayTracingGeometryRHIRef FD3D12DynamicRHI::RHICreateRayTracingGeometry(const FR
 
 	checkf(Initializer.PrimitiveType == EPrimitiveType::PT_TriangleList, TEXT("Only TriangleList primitive type is currently supported."));
 
-	// #dxr_todo: temporary constraints on vertex and index buffer formats (this will be relaxed when more flexible vertex/index fetching is implemented)
-	checkf(Initializer.VertexBufferElementType == VET_Float3, TEXT("Only float3 vertex buffers are currently implemented.")); // #dxr_todo: support other vertex buffer formats
-	checkf(Initializer.VertexBufferStride == 12, TEXT("Only deinterleaved float3 position vertex buffers are currently implemented.")); // #dxr_todo: support interleaved vertex buffers
-	checkf(Initializer.BaseVertexIndex == 0, TEXT("BaseVertexIndex is not currently implemented")); // #dxr_todo: implement base vertex index for custom vertex fetch
+	// #dxr_todo UE-72160: temporary constraints on vertex and index buffer formats (this will be relaxed when more flexible vertex/index fetching is implemented)
+	checkf(Initializer.VertexBufferElementType == VET_Float3, TEXT("Only float3 vertex buffers are currently implemented.")); // #dxr_todo UE-72160: support other vertex buffer formats
+	checkf(Initializer.VertexBufferStride == 12, TEXT("Only deinterleaved float3 position vertex buffers are currently implemented.")); // #dxr_todo UE-72160: support interleaved vertex buffers
 
-	checkf(GNumExplicitGPUsForRendering == 1, TEXT("Ray tracing is not implemented for mGPU")); // #dxr_todo: implement mGPU support
+	checkf(GNumExplicitGPUsForRendering == 1, TEXT("Ray tracing is not implemented for mGPU")); // #dxr_todo UE-72157: implement mGPU support
 	FD3D12RayTracingGeometry* Result = GetAdapter().CreateLinkedObject<FD3D12RayTracingGeometry>(FRHIGPUMask::All(), [&](FD3D12Device* Device)
 	{
 		FD3D12RayTracingGeometry* Mesh = new FD3D12RayTracingGeometry(Device);
@@ -1411,7 +1410,7 @@ static void CreateAccelerationStructureBuffers(TRefCountPtr<FD3D12MemBuffer>& Ac
 
 	SetName(AccelerationStructureBuffer->GetResource(), TEXT("Acceleration structure"));
 
-	// #dxr_todo: scratch buffers can be pooled and reused for different scenes and geometries
+	// #dxr_todo UE-72161: scratch buffers can be pooled and reused for different scenes and geometries
 	D3D12_RESOURCE_DESC ScratchBufferDesc = CD3DX12_RESOURCE_DESC::Buffer(
 		FMath::Max(PrebuildInfo.UpdateScratchDataSizeInBytes, PrebuildInfo.ScratchDataSizeInBytes), D3D12_RESOURCE_FLAG_ALLOW_UNORDERED_ACCESS);
 
@@ -1425,7 +1424,7 @@ static void CreateAccelerationStructureBuffers(TRefCountPtr<FD3D12MemBuffer>& Ac
 
 void FD3D12RayTracingGeometry::BuildAccelerationStructure(FD3D12CommandContext& CommandContext, bool bIsUpdate)
 {
-	check(GNumExplicitGPUsForRendering == 1); // #dxr_todo: ensure that mGPU case is handled correctly!
+	check(GNumExplicitGPUsForRendering == 1); // #dxr_todo UE-72157: ensure that mGPU case is handled correctly!
 
 	static constexpr uint32 IndicesPerPrimitive = 3; // Only triangle meshes are supported
 
@@ -1501,9 +1500,9 @@ void FD3D12RayTracingGeometry::BuildAccelerationStructure(FD3D12CommandContext&
 		Descs.Add(Desc);
 	}
 
-	checkf(GNumExplicitGPUsForRendering == 1, TEXT("Ray tracing is not implemented for mGPU")); // #dxr_todo: implement mGPU support
+	checkf(GNumExplicitGPUsForRendering == 1, TEXT("Ray tracing is not implemented for mGPU")); // #dxr_todo UE-72157: implement mGPU support
 
-	const uint32 GPUIndex = CommandContext.GetGPUIndex(); // #dxr_todo: ensure that mGPU case is handled correctly!
+	const uint32 GPUIndex = CommandContext.GetGPUIndex(); // #dxr_todo UE-72157: ensure that mGPU case is handled correctly!
 	FD3D12Adapter* Adapter = CommandContext.GetParentAdapter();
 
 	ID3D12Device5* RayTracingDevice = CommandContext.GetParentDevice()->GetRayTracingDevice();
@@ -1536,7 +1535,7 @@ void FD3D12RayTracingGeometry::BuildAccelerationStructure(FD3D12CommandContext&
 
 		CreateAccelerationStructureBuffers(AccelerationStructureBuffer, ScratchBuffer, Adapter, GPUIndex, PrebuildInfo);
 
-		// #dxr_todo: scratch buffers should be created in UAV state from the start
+		// #dxr_todo UE-72161: scratch buffers should be created in UAV state from the start
 		FD3D12DynamicRHI::TransitionResource(CommandContext.CommandListHandle, ScratchBuffer.GetReference()->GetResource(), D3D12_RESOURCE_STATE_UNORDERED_ACCESS, 0);
 	}
 
@@ -1578,12 +1577,12 @@ FD3D12RayTracingScene::~FD3D12RayTracingScene()
 
 void FD3D12RayTracingScene::BuildAccelerationStructure(FD3D12CommandContext& CommandContext, D3D12_RAYTRACING_ACCELERATION_STRUCTURE_BUILD_FLAGS BuildFlags)
 {
-	checkf(GNumExplicitGPUsForRendering == 1, TEXT("Ray tracing is not implemented for mGPU")); // #dxr_todo: implement mGPU support
+	checkf(GNumExplicitGPUsForRendering == 1, TEXT("Ray tracing is not implemented for mGPU")); // #dxr_todo UE-72157: implement mGPU support
 
 	TRefCountPtr<FD3D12MemBuffer> InstanceBuffer;
 	TRefCountPtr<FD3D12MemBuffer> ScratchBuffer;
 
-	const uint32 GPUIndex = CommandContext.GetGPUIndex(); // #dxr_todo: ensure that mGPU case is handled correctly!
+	const uint32 GPUIndex = CommandContext.GetGPUIndex(); // #dxr_todo UE-72157: ensure that mGPU case is handled correctly!
 	FD3D12Adapter* Adapter = CommandContext.GetParentAdapter();
 	ID3D12Device5* RayTracingDevice = CommandContext.GetParentDevice()->GetRayTracingDevice();
 
@@ -1599,7 +1598,7 @@ void FD3D12RayTracingScene::BuildAccelerationStructure(FD3D12CommandContext& Com
 
 	CreateAccelerationStructureBuffers(AccelerationStructureBuffer, ScratchBuffer, Adapter, GPUIndex, PrebuildInfo);
 
-	// #dxr_todo: scratch buffers should be created in UAV state from the start
+	// #dxr_todo UE-72161: scratch buffers should be created in UAV state from the start
 	FD3D12DynamicRHI::TransitionResource(CommandContext.CommandListHandle, ScratchBuffer.GetReference()->GetResource(), D3D12_RESOURCE_STATE_UNORDERED_ACCESS, 0);
 
 	if (bAccelerationStructureViewInitialized)
@@ -1699,7 +1698,7 @@ void FD3D12RayTracingScene::BuildAccelerationStructure(FD3D12CommandContext& Com
 
 	// Build the actual acceleration structure
 
-	const bool bIsUpdateMode = false; // #dxr_todo: we need an explicit public API to perform a refit/update
+	const bool bIsUpdateMode = false; // Top level acceleration structure is always built from scratch
 
 	AccelerationStructureBuffer->GetResource()->UpdateResidency(CommandContext.CommandListHandle);
 	ScratchBuffer->GetResource()->UpdateResidency(CommandContext.CommandListHandle);
@@ -1808,7 +1807,7 @@ FD3D12RayTracingShaderTable* FD3D12RayTracingScene::FindOrCreateShaderTable(cons
 				SystemParameters.IndexBuffer = IndexBufferAddress;
 				SystemParameters.VertexBuffer = VertexBufferAddress;
 
-				// #dxr_todo: support various vertex buffer layouts (fetch/decode based on vertex stride and format)
+				// #dxr_todo UE-72160: support various vertex buffer layouts (fetch/decode based on vertex stride and format)
 				checkf(Geometry->VertexElemType == VET_Float3, TEXT("Only VET_Float3 is currently implemented and tested. Other formats will be supported in the future."));
 				SystemParameters.RootConstants.SetVertexAndIndexStride(Geometry->VertexStrideInBytes, IndexStride);
 				SystemParameters.RootConstants.IndexBufferOffsetInBytes = IndexStride * Segment.FirstPrimitive * IndicesPerPrimitive;
@@ -2255,7 +2254,7 @@ static void DispatchRays(FD3D12CommandContext& CommandContext,
 {
 	// Setup state for RT dispatch
 	
-	// #dxr_todo: RT and non-RT descriptors should use the same global heap that's dynamically sub-allocated.
+	// #dxr_todo UE-72158: RT and non-RT descriptors should use the same global heap that's dynamically sub-allocated.
 	// This requires a major refactor of descriptor heap management. In the short term, RT work uses a dedicated heap
 	// that's temporarily set for the duration of RT dispatch.
 	ID3D12DescriptorHeap* PreviousHeaps[2] =
@@ -2277,7 +2276,7 @@ static void DispatchRays(FD3D12CommandContext& CommandContext,
 		FD3D12RayTracingGlobalResourceBinder ResourceBinder(CommandContext);
 		SetRayTracingShaderResources(CommandContext, RayGenShader, GlobalBindings, OptShaderTable->DescriptorCache, ResourceBinder);
 
-		// #dxr_todo: avoid updating residency if this scene was already used on the current command list (i.e. multiple ray dispatches are performed back-to-back)
+		// #dxr_todo UE-72159: avoid updating residency if this scene was already used on the current command list (i.e. multiple ray dispatches are performed back-to-back)
 		OptShaderTable->UpdateResidency(CommandContext);
 	}
 	else
@@ -2399,7 +2398,7 @@ void FD3D12CommandContext::RHIRayTraceDispatch(FRayTracingPipelineStateRHIParamR
 		ShaderTable->CopyToGPU();
 	}
 
-	// #dxr_todo: avoid updating residency if this scene was already used on the current command list (i.e. multiple ray dispatches are performed back-to-back)
+	// #dxr_todo UE-72159: avoid updating residency if this scene was already used on the current command list (i.e. multiple ray dispatches are performed back-to-back)
 	Scene->UpdateResidency(*this);
 
 	FD3D12RayTracingShader* RayGenShader = FD3D12DynamicRHI::ResourceCast(RayGenShaderRHI);
diff --git a/Engine/Source/Runtime/D3D12RHI/Private/D3D12RayTracing.h b/Engine/Source/Runtime/D3D12RHI/Private/D3D12RayTracing.h
--- a/Engine/Source/Runtime/D3D12RHI/Private/D3D12RayTracing.h
+++ b/Engine/Source/Runtime/D3D12RHI/Private/D3D12RayTracing.h
@@ -9,7 +9,6 @@
 class FD3D12RayTracingPipelineState;
 class FD3D12RayTracingShaderTable;
 
-// #dxr_todo: perhaps FD3D12VertexBuffer in the core RHI should be generalized
 typedef FD3D12VertexBuffer FD3D12MemBuffer; // Generic GPU memory buffer
 
 class FD3D12RayTracingGeometry : public FRHIRayTracingGeometry, public FD3D12DeviceChild, public FD3D12LinkedAdapterObject<FD3D12RayTracingGeometry>
@@ -56,12 +55,11 @@ public:
 
 	void BuildAccelerationStructure(FD3D12CommandContext& CommandContext, D3D12_RAYTRACING_ACCELERATION_STRUCTURE_BUILD_FLAGS BuildFlags);
 
-	// #dxr_todo: only GPU resources should be managed as LinkedAdapterObjects
+	// #dxr_todo UE-72157: only GPU resources should be managed as LinkedAdapterObjects
 	TRefCountPtr<FD3D12MemBuffer> AccelerationStructureBuffer;
 	TRefCountPtr<FD3D12ShaderResourceView> AccelerationStructureView;
 	bool bAccelerationStructureViewInitialized = false;
 
-	// #dxr_todo: add API to update geometries and instances
 	TArray<FRayTracingGeometryInstance> Instances;
 
 	// Scene keeps track of child acceleration structures to manage their residency
@@ -75,7 +73,7 @@ public:
 	uint32 NumTotalSegments = 0;
 	uint32 GetHitRecordBaseIndex(uint32 InstanceIndex, uint32 SegmentIndex) const { return (SegmentPrefixSum[InstanceIndex] + SegmentIndex) * ShaderSlotsPerGeometrySegment; }
 
-	// #dxr_todo: shader tables should be explicitly registered and unregistered with the scene
+	// #dxr_todo UE-68230: shader tables should be explicitly registered and unregistered with the scene
 	FD3D12RayTracingShaderTable* FindOrCreateShaderTable(const FD3D12RayTracingPipelineState* Pipeline);
 	FD3D12RayTracingShaderTable* FindExistingShaderTable(const FD3D12RayTracingPipelineState* Pipeline) const;
 
diff --git a/Engine/Source/Runtime/Engine/Public/SceneManagement.h b/Engine/Source/Runtime/Engine/Public/SceneManagement.h
--- a/Engine/Source/Runtime/Engine/Public/SceneManagement.h
+++ b/Engine/Source/Runtime/Engine/Public/SceneManagement.h
@@ -2744,8 +2744,8 @@ struct FLODMask
 		return DitheredLODIndices[0] == LODIndex || DitheredLODIndices[1] == LODIndex;
 	}
 
-	//#dxr_todo  We should probably add both LoDs but mask them based on their 
-	//LodFace value within the BVH based on the LodFadeMask in the GBuffer
+	//#dxr_todo UE-72106: We should probably add both LoDs but mask them based on their 
+	//LodFade value within the BVH based on the LodFadeMask in the GBuffer
 	bool ContainsRayTracedLOD(int32 LODIndex) const
 	{
 		return DitheredLODIndices[0] == LODIndex;
diff --git a/Engine/Source/Runtime/RHI/Private/PipelineStateCache.cpp b/Engine/Source/Runtime/RHI/Private/PipelineStateCache.cpp
--- a/Engine/Source/Runtime/RHI/Private/PipelineStateCache.cpp
+++ b/Engine/Source/Runtime/RHI/Private/PipelineStateCache.cpp
@@ -598,7 +598,7 @@ FAutoConsoleTaskPriority CPrio_FCompilePipelineStateTask(
 // This cache uses a single internal lock and therefore is not designed for highly concurrent operations.
 class FRayTracingPipelineCache
 {
-	// #dxr_todo: This needs to support fully asynchronous, non-blocking pipeline creation with explicit completion query mechanism.
+	// #dxr_todo UE-68234: This needs to support fully asynchronous, non-blocking pipeline creation with explicit completion query mechanism.
 	// #dxr_todo: Could move this to a separate cpp file.
 	// #dxr_todo: Should support eviction of stale pipelines.
 	// #dxr_todo: we will likely also need an explicit ray tracing pipeline sub-object cache to hold closest hit shaders
@@ -918,8 +918,8 @@ FRHIRayTracingPipelineState* PipelineStateCache::GetAndOrCreateRayTracingPipelin
 
 	if (bWasFound == false)
 	{
-		// #dxr_todo: RT PSO disk caching
-		// #dxr_todo: asynchronous PSO creation
+		// #dxr_todo UE-68235: RT PSO disk caching
+		// #dxr_todo UE-68234: asynchronous PSO creation
 
 		OutCachedState = new FRayTracingPipelineState();
 
diff --git a/Engine/Source/Runtime/RHI/Public/DynamicRHI.h b/Engine/Source/Runtime/RHI/Public/DynamicRHI.h
--- a/Engine/Source/Runtime/RHI/Public/DynamicRHI.h
+++ b/Engine/Source/Runtime/RHI/Public/DynamicRHI.h
@@ -1127,14 +1127,12 @@ public:
 
 
 #if RHI_RAYTRACING
-	// #dxr_todo: reverted index buffer may be needed to support mirrored meshes
 	virtual FRayTracingGeometryRHIRef RHICreateRayTracingGeometry(const FRayTracingGeometryInitializer& Initializer)
 	{
 		checkNoEntry();
 		return nullptr;
 	}
 
-	// This is a placeholder basic interface for the scene. We will likely need to revise this based on real use cases.
 	virtual FRayTracingSceneRHIRef RHICreateRayTracingScene(const FRayTracingSceneInitializer& Initializer)
 	{
 		checkNoEntry();
@@ -1403,4 +1401,4 @@ FDynamicRHI* PlatformCreateDynamicRHI();
 
 // Name of the RHI module that will be created when PlatformCreateDynamicRHI is called
 // NOTE: This function is very slow when called before RHIInit
-extern RHI_API const TCHAR* GetSelectedDynamicRHIModuleName(bool bCleanup = true);
\ No newline at end of file
+extern RHI_API const TCHAR* GetSelectedDynamicRHIModuleName(bool bCleanup = true);
diff --git a/Engine/Source/Runtime/RHI/Public/RHICommandList.h b/Engine/Source/Runtime/RHI/Public/RHICommandList.h
--- a/Engine/Source/Runtime/RHI/Public/RHICommandList.h
+++ b/Engine/Source/Runtime/RHI/Public/RHICommandList.h
@@ -121,7 +121,6 @@ struct FRayTracingShaderBindings
 };
 
 // C++ counter-part of FBasicRayData declared in RayTracingCommon.ush
-// #dxr_todo: ideally this should be shared with shaders by including a common header
 struct FBasicRayData
 {
 	float Origin[3];
@@ -131,7 +130,6 @@ struct FBasicRayData
 };
 
 // C++ counter-part of FBasicRayIntersectionData declared in RayTracingCommon.ush
-// #dxr_todo: ideally this should be shared with shaders by including a common header
 struct FBasicRayIntersectionData
 {
 	float  Distance;
diff --git a/Engine/Source/Runtime/RHI/Public/ValidationRHI.h b/Engine/Source/Runtime/RHI/Public/ValidationRHI.h
--- a/Engine/Source/Runtime/RHI/Public/ValidationRHI.h
+++ b/Engine/Source/Runtime/RHI/Public/ValidationRHI.h
@@ -1616,13 +1616,11 @@ public:
 
 
 #if RHI_RAYTRACING
-	// #dxr_todo: reverted index buffer may be needed to support mirrored meshes
 	virtual FRayTracingGeometryRHIRef RHICreateRayTracingGeometry(const FRayTracingGeometryInitializer& Initializer) override
 	{
 		return RHI->RHICreateRayTracingGeometry(Initializer);
 	}
 
-	// This is a placeholder basic interface for the scene. We will likely need to revise this based on real use cases.
 	virtual FRayTracingSceneRHIRef RHICreateRayTracingScene(const FRayTracingSceneInitializer& Initializer) override
 	{
 		return RHI->RHICreateRayTracingScene(Initializer);
diff --git a/Engine/Source/Runtime/Renderer/Private/RayTracing/RayTracingMaterialHitShaders.cpp b/Engine/Source/Runtime/Renderer/Private/RayTracing/RayTracingMaterialHitShaders.cpp
--- a/Engine/Source/Runtime/Renderer/Private/RayTracing/RayTracingMaterialHitShaders.cpp
+++ b/Engine/Source/Runtime/Renderer/Private/RayTracing/RayTracingMaterialHitShaders.cpp
@@ -516,7 +516,7 @@ FRHIRayTracingPipelineState* FDeferredShadingSceneRenderer::BindRayTracingMateri
 			else
 			{
 				// Masked materials require full material evaluation with any-hit shader.
-				// #dxr_todo: we need to generate a shadow-specific closest hit shader for this!
+				// #dxr_todo UE-72029: we need to generate a shadow-specific closest hit shader for this!
 				MeshCommand.ShaderBindings.SetRayTracingShaderBindingsForHitGroup(RHICmdList,
 					View.RayTracingScene.RayTracingSceneRHI,
 					VisibleMeshCommand.InstanceIndex,
