SetRayTracingShaderResources() now checks that root signature does not use CBV descriptor tables, as they are not supported. Resolves a minor dxr_todo.

diff --git a/Engine/Source/Runtime/D3D12RHI/Private/D3D12RayTracing.cpp b/Engine/Source/Runtime/D3D12RHI/Private/D3D12RayTracing.cpp
--- a/Engine/Source/Runtime/D3D12RHI/Private/D3D12RayTracing.cpp
+++ b/Engine/Source/Runtime/D3D12RHI/Private/D3D12RayTracing.cpp
@@ -2226,7 +2226,7 @@ static void SetRayTracingShaderResources(
 
 	if (Shader->ResourceCounts.NumCBs)
 	{
-		// #dxr_todo: make sure that root signature only uses root CBVs (this is currently checked in D3D12RootSignature.cpp)
+		checkf(RootSignature->CBVRDTBindSlot(SF_Compute) == 0xFF, TEXT("Root CBV descriptor tables are not implemented for ray tracing shaders."));
 
 		const uint32 BindSlot = RootSignature->CBVRDBaseBindSlot(SF_Compute);
 		check(BindSlot != 0xFF);
