Implemented support for DXR callable shaders

Callable shaders are similar to hit shaders, but they can be invoked from RayGen shaders without tracing any rays. They have records in the shader binding table and allow local parameters to be bound. The total number of callable shader slots/records must be declared when creating a ray tracing scene. Each record must be filled using SetRayTracingCallableShader() command, which is similar to SetRayTracingHitGroup().

diff --git a/Engine/Source/Developer/ShaderCompilerCommon/Private/ShaderCompilerCommon.cpp b/Engine/Source/Developer/ShaderCompilerCommon/Private/ShaderCompilerCommon.cpp
--- a/Engine/Source/Developer/ShaderCompilerCommon/Private/ShaderCompilerCommon.cpp
+++ b/Engine/Source/Developer/ShaderCompilerCommon/Private/ShaderCompilerCommon.cpp
@@ -782,6 +782,7 @@ FString CreateShaderCompilerWorkerDirectCommandLine(const FShaderCompilerInput&
 	case SF_RayGen:			Text += TEXT(" -rgs"); break;
 	case SF_RayMiss:		Text += TEXT(" -rms"); break;
 	case SF_RayHitGroup:	Text += TEXT(" -rhs"); break;
+	case SF_RayCallable:	Text += TEXT(" -rcs"); break;
 #endif // RHI_RAYTRACING
 	default: break;
 	}
@@ -1196,6 +1197,7 @@ namespace CrossCompiler
 		TEXT("RayGen"),
 		TEXT("RayMiss"),
 		TEXT("RayHitGroup"),
+		TEXT("RayCallable"),
 	};
 
 	/** Compile time check to verify that the GL mapping tables are up-to-date. */
diff --git a/Engine/Source/Developer/ShaderFormatOpenGL/Private/OpenGLShaderCompiler.cpp b/Engine/Source/Developer/ShaderFormatOpenGL/Private/OpenGLShaderCompiler.cpp
--- a/Engine/Source/Developer/ShaderFormatOpenGL/Private/OpenGLShaderCompiler.cpp
+++ b/Engine/Source/Developer/ShaderFormatOpenGL/Private/OpenGLShaderCompiler.cpp
@@ -443,6 +443,7 @@ GLenum GLFrequencyTable[] =
 	GLenum(0), // SF_RayGen
 	GLenum(0), // SF_RayMiss
 	GLenum(0), // SF_RayHitGroup (closest hit, any hit, intersection)
+	GLenum(0), // SF_RayCallable
 };
 
 static_assert(ARRAY_COUNT(GLFrequencyTable) == SF_NumFrequencies, "Frequency table size mismatch.");
diff --git a/Engine/Source/Developer/Windows/ShaderFormatD3D/Private/D3DShaderCompiler.cpp b/Engine/Source/Developer/Windows/ShaderFormatD3D/Private/D3DShaderCompiler.cpp
--- a/Engine/Source/Developer/Windows/ShaderFormatD3D/Private/D3DShaderCompiler.cpp
+++ b/Engine/Source/Developer/Windows/ShaderFormatD3D/Private/D3DShaderCompiler.cpp
@@ -92,6 +92,7 @@ static bool IsRayTracingShader(const FShaderTarget& Target)
 	case SF_RayGen:
 	case SF_RayMiss:
 	case SF_RayHitGroup:
+	case SF_RayCallable:
 		return true;
 	default:
 		return false;
@@ -106,6 +107,7 @@ static uint32 GetAutoBindingSpace(const FShaderTarget& Target)
 	case SF_RayMiss:
 		return RAY_TRACING_REGISTER_SPACE_GLOBAL;
 	case SF_RayHitGroup:
+	case SF_RayCallable:
 		return RAY_TRACING_REGISTER_SPACE_LOCAL;
 	default:
 		return 0;
@@ -138,6 +140,7 @@ static const TCHAR* GetShaderProfileName(FShaderTarget Target, bool bForceSM6)
 		case SF_RayGen:
 		case SF_RayMiss:
 		case SF_RayHitGroup:
+		case SF_RayCallable:
 			return TEXT("lib_6_3");
 		}
 	}
diff --git a/Engine/Source/Programs/ShaderCompileWorker/Private/ShaderCompileWorker.cpp b/Engine/Source/Programs/ShaderCompileWorker/Private/ShaderCompileWorker.cpp
--- a/Engine/Source/Programs/ShaderCompileWorker/Private/ShaderCompileWorker.cpp
+++ b/Engine/Source/Programs/ShaderCompileWorker/Private/ShaderCompileWorker.cpp
@@ -772,6 +772,10 @@ static void DirectCompile(const TArray<const class IShaderFormat*>& ShaderFormat
 			{
 				Frequency = SF_RayHitGroup;
 			}
+			else if (!FCString::Strcmp(*Token, TEXT("rcs")))
+			{
+				Frequency = SF_RayCallable;
+			}
 #endif // RHI_RAYTRACING
 			else if (!FCString::Strcmp(*Token, TEXT("pipeline")))
 			{
diff --git a/Engine/Source/Runtime/D3D12RHI/Private/D3D12CommandContext.h b/Engine/Source/Runtime/D3D12RHI/Private/D3D12CommandContext.h
--- a/Engine/Source/Runtime/D3D12RHI/Private/D3D12CommandContext.h
+++ b/Engine/Source/Runtime/D3D12RHI/Private/D3D12CommandContext.h
@@ -354,6 +354,11 @@ public:
 		FRayTracingPipelineStateRHIParamRef Pipeline, uint32 HitGroupIndex,
 		uint32 NumUniformBuffers, const FUniformBufferRHIParamRef* UniformBuffers,
 		uint32 UserData) final override;
+	virtual void RHISetRayTracingCallableShader(
+		FRayTracingSceneRHIParamRef Scene, uint32 ShaderSlotInScene,
+		FRayTracingPipelineStateRHIParamRef Pipeline, uint32 ShaderIndexInPipeline,
+		uint32 NumUniformBuffers, const FUniformBufferRHIParamRef* UniformBuffers,
+		uint32 UserData) final override;
 #endif // D3D12_RHI_RAYTRACING
 
 	template<typename ObjectType, typename RHIType, typename Predicate>
diff --git a/Engine/Source/Runtime/D3D12RHI/Private/D3D12RayTracing.cpp b/Engine/Source/Runtime/D3D12RHI/Private/D3D12RayTracing.cpp
--- a/Engine/Source/Runtime/D3D12RHI/Private/D3D12RayTracing.cpp
+++ b/Engine/Source/Runtime/D3D12RHI/Private/D3D12RayTracing.cpp
@@ -402,6 +402,7 @@ public:
 		RayGen,
 		Miss,
 		HitGroup,
+		Callable,
 	};
 
 	static const TCHAR* GetCollectionTypeName(ECollectionType Type)
@@ -410,12 +411,12 @@ public:
 		{
 		case ECollectionType::RayGen:
 			return TEXT("RayGen");
-			break;
 		case ECollectionType::Miss:
 			return TEXT("Miss");
-			break;
 		case ECollectionType::HitGroup:
 			return TEXT("HitGroup");
+		case ECollectionType::Callable:
+			return TEXT("Callable");
 		default:
 			return TEXT("");
 		}
@@ -432,9 +433,9 @@ public:
 		const uint64 ShaderHash = GetShaderHash64(Shader);
 
 		ID3D12RootSignature* LocalRootSignature = nullptr;
-		if (CollectionType == ECollectionType::HitGroup)
+		if (CollectionType == ECollectionType::HitGroup || CollectionType == ECollectionType::Callable)
 		{
-			// Only hit group shaders have a local root signature
+			// Only hit group and callable shaders have a local root signature
 			LocalRootSignature = Shader->pRootSignature->GetRootSignature();
 		}
 		else
@@ -524,7 +525,7 @@ public:
 			}
 			else
 			{
-				checkf(CollectionType == ECollectionType::Miss || CollectionType == ECollectionType::RayGen, TEXT("Unexpected RT sahder collection type"));
+				checkf(CollectionType == ECollectionType::Miss || CollectionType == ECollectionType::RayGen || CollectionType == ECollectionType::Callable, TEXT("Unexpected RT sahder collection type"));
 
 				OriginalEntryPoints.Add(*(Shader->EntryPoint));
 				RenamedEntryPoints.Add(*PrimaryExportName);
@@ -917,14 +918,14 @@ private:
 		bIsDirty = true;
 	}
 
-	void WriteHitRecord(uint32 RecordIndex, uint32 OffsetWithinRecord, const void* InData, uint32 InDataSize)
+	void WriteLocalShaderRecord(uint32 RecordIndex, uint32 OffsetWithinRecord, const void* InData, uint32 InDataSize)
 	{
 		checkfSlow(OffsetWithinRecord % 4 == 0, TEXT("SBT record parameters must be written on DWORD-aligned boundary"));
 		checkfSlow(InDataSize % 4 == 0, TEXT("SBT record parameters must be DWORD-aligned"));
-		checkfSlow(OffsetWithinRecord + InDataSize <= HitRecordSizeUnaligned, TEXT("SBT record write request is out of bounds"));
-		checkfSlow(RecordIndex < NumHitRecords, TEXT("SBT record write request is out of bounds"));
+		checkfSlow(OffsetWithinRecord + InDataSize <= LocalRecordSizeUnaligned, TEXT("SBT record write request is out of bounds"));
+		checkfSlow(RecordIndex < NumLocalRecords, TEXT("SBT local record write request is out of bounds"));
 
-		const uint32 WriteOffset = HitGroupShaderTableOffset + HitRecordStride * RecordIndex + OffsetWithinRecord;
+		const uint32 WriteOffset = LocalShaderTableOffset + LocalRecordStride * RecordIndex + OffsetWithinRecord;
 
 		WriteData(WriteOffset, InData, InDataSize);
 	}
@@ -933,6 +934,16 @@ public:
 
 	UE_NONCOPYABLE(FD3D12RayTracingShaderTable)
 
+	struct FInitializer
+	{
+		uint32 NumRayGenShaders = 0;
+		uint32 NumMissShaders = 0;
+		uint32 NumHitRecords = 0;
+		uint32 NumCallableRecords = 0;
+		uint32 LocalRootDataSize = 0;
+		uint32 MaxViewDescriptorsPerRecord = 0;
+	};
+
 	FD3D12RayTracingShaderTable(FD3D12Device* Device)
 		: FD3D12DeviceChild(Device)
 	{
@@ -943,17 +954,17 @@ public:
 		delete DescriptorCache;
 	}
 
-	void Init(uint32 InNumRayGenShaders, uint32 InNumMissShaders, uint32 InNumHitRecords, uint32 LocalRootDataSize, uint32 MaxViewDescriptorsPerRecord)
+	void Init(const FInitializer& Initializer)
 	{
-		checkf(LocalRootDataSize <= 4096, TEXT("The maximum size of a local root signature is 4KB.")); // as per section 4.22.1 of DXR spec v1.0
-		checkf(InNumRayGenShaders >= 1, TEXT("All shader tables must contain at least one raygen shader."));
+		checkf(Initializer.LocalRootDataSize <= 4096, TEXT("The maximum size of a local root signature is 4KB.")); // as per section 4.22.1 of DXR spec v1.0
+		checkf(Initializer.NumRayGenShaders >= 1, TEXT("All shader tables must contain at least one raygen shader."));
 
-		HitRecordSizeUnaligned = ShaderIdentifierSize + LocalRootDataSize;
-		HitRecordStride = RoundUpToNextMultiple(HitRecordSizeUnaligned, D3D12_RAYTRACING_SHADER_RECORD_BYTE_ALIGNMENT);
+		LocalRecordSizeUnaligned = ShaderIdentifierSize + Initializer.LocalRootDataSize;
+		LocalRecordStride = RoundUpToNextMultiple(LocalRecordSizeUnaligned, D3D12_RAYTRACING_SHADER_RECORD_BYTE_ALIGNMENT);
 
 		// Custom descriptor cache is only required when local resources may be bound.
 		// If only global resources are used, then transient descriptor cache can be used.
-		const bool bNeedsDescriptorCache = InNumHitRecords * LocalRootDataSize != 0;
+		const bool bNeedsDescriptorCache = (Initializer.NumHitRecords + Initializer.NumCallableRecords) * Initializer.LocalRootDataSize != 0;
 
 		if (bNeedsDescriptorCache)
 		{
@@ -963,16 +974,17 @@ public:
 
 			// D3D12 is guaranteed to support 1M (D3D12_MAX_SHADER_VISIBLE_DESCRIPTOR_HEAP_SIZE_TIER_1) descriptors in a CBV/SRV/UAV heap, so clamp the size to this.
 			// https://docs.microsoft.com/en-us/windows/desktop/direct3d12/hardware-support
-			const uint32 NumViewDescriptors = FMath::Max(MinNumViewDescriptors, FMath::Min<uint32>(InNumHitRecords * MaxViewDescriptorsPerRecord, D3D12_MAX_SHADER_VISIBLE_DESCRIPTOR_HEAP_SIZE_TIER_1));
+			const uint32 NumViewDescriptors = FMath::Max(MinNumViewDescriptors, FMath::Min<uint32>(Initializer.NumHitRecords * Initializer.MaxViewDescriptorsPerRecord, D3D12_MAX_SHADER_VISIBLE_DESCRIPTOR_HEAP_SIZE_TIER_1));
 			const uint32 NumSamplerDescriptors = D3D12_MAX_SHADER_VISIBLE_SAMPLER_HEAP_SIZE;
 
 			DescriptorCache = new FD3D12RayTracingDescriptorCache(GetParentDevice());
 			DescriptorCache->Init(NumViewDescriptors, NumSamplerDescriptors);
 		}
 
-		NumRayGenShaders = InNumRayGenShaders;
-		NumMissShaders = InNumMissShaders;
-		NumHitRecords = InNumHitRecords;
+		NumRayGenShaders = Initializer.NumRayGenShaders;
+		NumMissShaders = Initializer.NumMissShaders;
+		NumHitRecords = Initializer.NumHitRecords;
+		NumCallableRecords = Initializer.NumCallableRecords;
 
 		uint32 TotalDataSize = 0;
 
@@ -984,14 +996,23 @@ public:
 		TotalDataSize += NumMissShaders * MissRecordStride;
 		TotalDataSize = RoundUpToNextMultiple(TotalDataSize, D3D12_RAYTRACING_SHADER_TABLE_BYTE_ALIGNMENT);
 
+		DefaultHitGroupShaderTableOffset = TotalDataSize;
+		TotalDataSize += ShaderIdentifierSize;
+		TotalDataSize = RoundUpToNextMultiple(TotalDataSize, D3D12_RAYTRACING_SHADER_TABLE_BYTE_ALIGNMENT);
+
 		HitGroupShaderTableOffset = TotalDataSize;
-		TotalDataSize += InNumHitRecords * HitRecordStride;
+		TotalDataSize += Initializer.NumHitRecords * LocalRecordStride;
 		TotalDataSize = RoundUpToNextMultiple(TotalDataSize, D3D12_RAYTRACING_SHADER_TABLE_BYTE_ALIGNMENT);
 
-		DefaultHitGroupShaderTableOffset = TotalDataSize;
-		TotalDataSize += ShaderIdentifierSize;
+		CallableShaderTableOffset = TotalDataSize;
+		TotalDataSize += Initializer.NumCallableRecords * LocalRecordStride;
 		TotalDataSize = RoundUpToNextMultiple(TotalDataSize, D3D12_RAYTRACING_SHADER_TABLE_BYTE_ALIGNMENT);
 
+		// Hit groups and callable shaders are stored in a consecutive memory block and are accessed using common local record indices.
+		LocalShaderTableOffset = HitGroupShaderTableOffset;
+		CallableShaderRecordIndexOffset = (CallableShaderTableOffset - LocalShaderTableOffset) / LocalRecordStride;
+		NumLocalRecords = (TotalDataSize - LocalShaderTableOffset) / LocalRecordStride;
+
 		Data.SetNumZeroed(TotalDataSize);
 
 		// Keep CPU-side data after upload
@@ -999,21 +1020,21 @@ public:
 	}
 
 	template <typename T>
-	void SetHitGroupParameters(uint32 RecordIndex, uint32 InOffsetWithinRootSignature, const T& Parameters)
+	void SetLocalShaderParameters(uint32 RecordIndex, uint32 InOffsetWithinRootSignature, const T& Parameters)
 	{
-		WriteHitRecord(RecordIndex, ShaderIdentifierSize + InOffsetWithinRootSignature, &Parameters, sizeof(Parameters));
+		WriteLocalShaderRecord(RecordIndex, ShaderIdentifierSize + InOffsetWithinRootSignature, &Parameters, sizeof(Parameters));
 	}
 
-	void SetHitGroupParameters(uint32 RecordIndex, uint32 InOffsetWithinRootSignature, const void* InData, uint32 InDataSize)
+	void SetLocalShaderParameters(uint32 RecordIndex, uint32 InOffsetWithinRootSignature, const void* InData, uint32 InDataSize)
 	{
-		WriteHitRecord(RecordIndex, ShaderIdentifierSize + InOffsetWithinRootSignature, InData, InDataSize);
+		WriteLocalShaderRecord(RecordIndex, ShaderIdentifierSize + InOffsetWithinRootSignature, InData, InDataSize);
 	}
 
-	void SetHitGroupIdentifier(uint32 RecordIndex, const void* ShaderIdentifierData, uint32 InShaderIdentifierSize)
+	void SetLocalShaderIdentifier(uint32 RecordIndex, const void* ShaderIdentifierData, uint32 InShaderIdentifierSize)
 	{
 		checkSlow(InShaderIdentifierSize == ShaderIdentifierSize);
 
-		WriteHitRecord(RecordIndex, 0, ShaderIdentifierData, InShaderIdentifierSize);
+		WriteLocalShaderRecord(RecordIndex, 0, ShaderIdentifierData, InShaderIdentifierSize);
 	}
 
 	void SetRayGenIdentifier(uint32 RecordIndex, const FD3D12ShaderIdentifier& ShaderIdentifier)
@@ -1034,12 +1055,12 @@ public:
 		WriteData(WriteOffset, ShaderIdentifier.Data, ShaderIdentifierSize);
 	}
 
-	void SetHitGroupIdentifier(uint32 RecordIndex, const FD3D12ShaderIdentifier& ShaderIdentifier)
+	void SetLocalShaderIdentifier(uint32 RecordIndex, const FD3D12ShaderIdentifier& ShaderIdentifier)
 	{
 		checkfSlow(ShaderIdentifier.IsValid(), TEXT("Shader identifier must be initialized FD3D12RayTracingPipelineState::GetShaderIdentifier() before use."));
 		checkSlow(sizeof(ShaderIdentifier.Data) >= ShaderIdentifierSize);
 
-		SetHitGroupIdentifier(RecordIndex, ShaderIdentifier.Data, ShaderIdentifierSize);
+		SetLocalShaderIdentifier(RecordIndex, ShaderIdentifier.Data, ShaderIdentifierSize);
 	}
 
 	void SetRayGenIdentifiers(const TArrayView<const FD3D12ShaderIdentifier>& Identifiers)
@@ -1060,6 +1081,15 @@ public:
 		}
 	}
 
+	void SetLocalShaderIdentifiers(uint32 RecordIndexOffset, const TArrayView<const FD3D12ShaderIdentifier>& Identifiers)
+	{
+		check(Identifiers.Num() == NumCallableRecords);
+		for (int32 Index = 0; Index < Identifiers.Num(); ++Index)
+		{
+			SetLocalShaderIdentifier(RecordIndexOffset + Index, Identifiers[Index]);
+		}
+	}
+
 	void CopyToGPU()
 	{
 		check(IsInRHIThread() || !IsRunningRHIInSeparateThread());
@@ -1105,11 +1135,18 @@ public:
 		Desc.MissShaderTable.StrideInBytes = MissRecordStride;
 		Desc.MissShaderTable.SizeInBytes = MissRecordStride;
 
+		if (NumCallableRecords)
+		{
+			Desc.CallableShaderTable.StartAddress = ShaderTableAddress + CallableShaderTableOffset;
+			Desc.CallableShaderTable.StrideInBytes = LocalRecordStride;
+			Desc.CallableShaderTable.SizeInBytes = NumCallableRecords * LocalRecordStride;
+		}
+
 		if (bAllowHitGroupIndexing)
 		{
 			Desc.HitGroupTable.StartAddress = ShaderTableAddress + HitGroupShaderTableOffset;
-			Desc.HitGroupTable.StrideInBytes = HitRecordStride;
-			Desc.HitGroupTable.SizeInBytes = NumHitRecords * HitRecordStride;
+			Desc.HitGroupTable.StrideInBytes = LocalRecordStride;
+			Desc.HitGroupTable.SizeInBytes = NumHitRecords * LocalRecordStride;
 		}
 		else
 		{
@@ -1125,19 +1162,24 @@ public:
 	uint32 NumHitRecords = 0;
 	uint32 NumRayGenShaders = 0;
 	uint32 NumMissShaders = 0;
+	uint32 NumCallableRecords = 0;
+	uint32 NumLocalRecords = 0;
 
 	uint32 RayGenShaderTableOffset = 0;
 	uint32 MissShaderTableOffset = 0;
-	uint32 HitGroupShaderTableOffset = 0;
 	uint32 DefaultHitGroupShaderTableOffset = 0;
+	uint32 HitGroupShaderTableOffset = 0;
+	uint32 CallableShaderTableOffset = 0;
+	uint32 LocalShaderTableOffset = 0;
+	uint32 CallableShaderRecordIndexOffset = 0;
 
 	// Note: TABLE_BYTE_ALIGNMENT is used instead of RECORD_BYTE_ALIGNMENT to allow arbitrary switching 
 	// between multiple RayGen and Miss shaders within the same underlying table.
 	static constexpr uint32 RayGenRecordStride = D3D12_RAYTRACING_SHADER_TABLE_BYTE_ALIGNMENT;
 	static constexpr uint32 MissRecordStride = D3D12_RAYTRACING_SHADER_TABLE_BYTE_ALIGNMENT;
 
-	uint32 HitRecordSizeUnaligned = 0; // size of the shader identifier + local root parameters, not aligned to SHADER_RECORD_BYTE_ALIGNMENT (used for out-of-bounds access checks)
-	uint32 HitRecordStride = 0; // size of shader identifier + local root parameters, aligned to SHADER_RECORD_BYTE_ALIGNMENT
+	uint32 LocalRecordSizeUnaligned = 0; // size of the shader identifier + local root parameters, not aligned to SHADER_RECORD_BYTE_ALIGNMENT (used for out-of-bounds access checks)
+	uint32 LocalRecordStride = 0; // size of shader identifier + local root parameters, aligned to SHADER_RECORD_BYTE_ALIGNMENT (same for hit groups and callable shaders)
 	TResourceArray<uint8, D3D12_RAYTRACING_SHADER_TABLE_BYTE_ALIGNMENT> Data;
 
 	bool bIsDirty = true;
@@ -1248,8 +1290,9 @@ public:
 			: DefaultMissTable;
 
 		TArrayView<const FRayTracingShaderRHIParamRef> InitializerRayGenShaders = Initializer.GetRayGenTable();
+		TArrayView<const FRayTracingShaderRHIParamRef> InitializerCallableShaders = Initializer.GetCallableTable();
 
-		const uint32 MaxTotalShaders = InitializerRayGenShaders.Num() + InitializerMissShaders.Num() + InitializerHitGroups.Num();
+		const uint32 MaxTotalShaders = InitializerRayGenShaders.Num() + InitializerMissShaders.Num() + InitializerHitGroups.Num() + InitializerCallableShaders.Num();
 
 		// All raygen and miss shaders must share the same global root signature, so take the first one and validate the rest
 
@@ -1312,6 +1355,7 @@ public:
 			checkf(Shader->pRootSignature->GetRootSignature() == GlobalRootSignature, TEXT("All raygen and miss shaders must share the same root signature"));
 
 			D3D12_EXISTING_COLLECTION_DESC CollectionDesc = AddShaderCollection(Shader, FD3D12RayTracingPipelineCache::ECollectionType::RayGen);
+
 			RayGenShaderNames.Add(CollectionDesc.pExports[0].Name);
 			RayGenShaders.Shaders.Add(Shader);
 		}
@@ -1328,6 +1372,7 @@ public:
 			checkf(Shader->pRootSignature->GetRootSignature() == GlobalRootSignature, TEXT("All raygen and miss shaders must share the same root signature"));
 
 			D3D12_EXISTING_COLLECTION_DESC CollectionDesc = AddShaderCollection(Shader, FD3D12RayTracingPipelineCache::ECollectionType::Miss);
+
 			MissShaderNames.Add(CollectionDesc.pExports[0].Name);
 			MissShaders.Shaders.Add(Shader);
 		}
@@ -1359,6 +1404,29 @@ public:
 			HitGroupShaders.Shaders.Add(Shader);
 		}
 
+		// Add callable shaders
+
+		TArray<LPCWSTR> CallableShaderNames;
+		CallableShaders.Reserve(InitializerCallableShaders.Num());
+		CallableShaderNames.Reserve(InitializerCallableShaders.Num());
+
+		for (FRayTracingShaderRHIParamRef ShaderRHI : InitializerCallableShaders)
+		{
+			FD3D12RayTracingShader* Shader = FD3D12DynamicRHI::ResourceCast(ShaderRHI);
+
+			checkf(Shader, TEXT("A valid ray tracing shader must be provided for all elements in the FRayTracingPipelineStateInitializer callable shader table."));
+			checkf(!Shader->ResourceCounts.bGlobalUniformBufferUsed, TEXT("Global uniform buffers are not implemented for ray tracing shaders"));
+
+			const uint32 ShaderViewDescriptors = Shader->ResourceCounts.NumSRVs + Shader->ResourceCounts.NumUAVs;
+			MaxHitGroupViewDescriptors = FMath::Max(MaxHitGroupViewDescriptors, ShaderViewDescriptors);
+			MaxLocalRootSignatureSize = FMath::Max(MaxLocalRootSignatureSize, Shader->pRootSignature->GetTotalRootSignatureSizeInBytes());
+
+			D3D12_EXISTING_COLLECTION_DESC CollectionDesc = AddShaderCollection(Shader, FD3D12RayTracingPipelineCache::ECollectionType::Callable);
+
+			CallableShaderNames.Add(CollectionDesc.pExports[0].Name);
+			CallableShaders.Shaders.Add(Shader);
+		}
+
 		// Link final RTPSO from shader collections
 
 		LinkTime -= FPlatformTime::Cycles64();
@@ -1418,11 +1486,24 @@ public:
 			MissShaders.Identifiers[ShaderIndex] = GetShaderIdentifier(ExportNameChars);
 		}
 
+		CallableShaders.Identifiers.SetNumUninitialized(CallableShaderNames.Num());
+		for (int32 ShaderIndex = 0; ShaderIndex < CallableShaderNames.Num(); ++ShaderIndex)
+		{
+			LPCWSTR ExportNameChars = CallableShaderNames[ShaderIndex];
+			CallableShaders.Identifiers[ShaderIndex] = GetShaderIdentifier(ExportNameChars);
+		}
+
 		// Setup default shader binding table, which simply includes all provided RGS and MS plus a single default closest hit shader.
 		// Hit record indexing and local resources access is disabled when using using this SBT.
 
-		const uint32 DefaultLocalRootDataSize = 0; // Shaders in default SBT are not allowed to access any local resources
-		DefaultShaderTable.Init(RayGenShaders.Identifiers.Num(), MissShaders.Identifiers.Num(), 0, DefaultLocalRootDataSize, 0);
+		FD3D12RayTracingShaderTable::FInitializer SBTInitializer = {};
+		SBTInitializer.NumRayGenShaders = RayGenShaders.Identifiers.Num();
+		SBTInitializer.NumMissShaders = MissShaders.Identifiers.Num();
+		SBTInitializer.NumCallableRecords = 0; // Default SBT does not support callable shaders
+		SBTInitializer.NumHitRecords = 0; // Default SBT does not support indexable hit shaders
+		SBTInitializer.LocalRootDataSize = 0; // Shaders in default SBT are not allowed to access any local resources
+
+		DefaultShaderTable.Init(SBTInitializer);
 		DefaultShaderTable.SetRayGenIdentifiers(RayGenShaders.Identifiers);
 		DefaultShaderTable.SetMissIdentifiers(MissShaders.Identifiers);
 		DefaultShaderTable.SetDefaultHitGroupIdentifier(HitGroupShaders.Identifiers[0]);
@@ -1449,6 +1530,7 @@ public:
 	FD3D12RayTracingShaderLibrary RayGenShaders;
 	FD3D12RayTracingShaderLibrary MissShaders;
 	FD3D12RayTracingShaderLibrary HitGroupShaders;
+	FD3D12RayTracingShaderLibrary CallableShaders;
 
 	// Shader table that can be used to dispatch ray tracing work that doesn't require real SBT bindings.
 	// This is useful for the case where user only provides default RayGen, Miss and HitGroup shaders.
@@ -1666,6 +1748,8 @@ FRayTracingSceneRHIRef FD3D12DynamicRHI::RHICreateRayTracingScene(const FRayTrac
 		}
 		Result->NumTotalSegments = NumTotalSegments;
 
+		Result->NumCallableShaderSlots = Initializer.NumCallableShaderSlots;
+
 		return Result;
 	});
 }
@@ -2117,13 +2201,15 @@ FD3D12RayTracingShaderTable* FD3D12RayTracingScene::FindOrCreateShaderTable(cons
 
 	checkf(Pipeline->MaxLocalRootSignatureSize >= sizeof(FHitGroupSystemParameters), TEXT("All local root signatures are expected to contain ray tracing system root parameters (2x root buffers + 4x root DWORD)"));
 
-	CreatedShaderTable->Init(
-		Pipeline->RayGenShaders.Identifiers.Num(),
-		Pipeline->MissShaders.Identifiers.Num(),
-		NumHitGroupSlots,
-		Pipeline->MaxLocalRootSignatureSize,
-		Pipeline->MaxHitGroupViewDescriptors);
+	FD3D12RayTracingShaderTable::FInitializer SBTInitializer = {};
+	SBTInitializer.NumRayGenShaders = Pipeline->RayGenShaders.Identifiers.Num();
+	SBTInitializer.NumMissShaders = Pipeline->MissShaders.Identifiers.Num();
+	SBTInitializer.NumHitRecords = NumHitGroupSlots;
+	SBTInitializer.NumCallableRecords = NumCallableShaderSlots;
+	SBTInitializer.LocalRootDataSize = Pipeline->MaxLocalRootSignatureSize;
+	SBTInitializer.MaxViewDescriptorsPerRecord = Pipeline->MaxHitGroupViewDescriptors;
 
+	CreatedShaderTable->Init(SBTInitializer);
 	CreatedShaderTable->SetRayGenIdentifiers(Pipeline->RayGenShaders.Identifiers);
 	CreatedShaderTable->SetMissIdentifiers(Pipeline->MissShaders.Identifiers);
 	CreatedShaderTable->SetDefaultHitGroupIdentifier(Pipeline->HitGroupShaders.Identifiers[0]);
@@ -2132,6 +2218,9 @@ FD3D12RayTracingShaderTable* FD3D12RayTracingScene::FindOrCreateShaderTable(cons
 	// Resource binding is skipped for pipelines that don't use SBT indexing. Such pipelines use the same CHS for all primitives, which can't access any local resources.
 	if (NumHitGroupSlots)
 	{
+		checkf(CreatedShaderTable->LocalShaderTableOffset == CreatedShaderTable->HitGroupShaderTableOffset,
+			TEXT("Hit shader records are assumed to be at the beginning of local shader table"));
+
 		const uint32 NumInstances = Instances.Num();
 		for (uint32 InstanceIndex = 0; InstanceIndex < NumInstances; ++InstanceIndex)
 		{
@@ -2170,7 +2259,7 @@ FD3D12RayTracingShaderTable* FD3D12RayTracingScene::FindOrCreateShaderTable(cons
 
 				for (uint32 SlotIndex = 0; SlotIndex < ShaderSlotsPerGeometrySegment; ++SlotIndex)
 				{
-					CreatedShaderTable->SetHitGroupParameters(RecordBaseIndex + SlotIndex, 0, SystemParameters);
+					CreatedShaderTable->SetLocalShaderParameters(RecordBaseIndex + SlotIndex, 0, SystemParameters);
 				}
 			}
 		}
@@ -2287,7 +2376,7 @@ struct FD3D12RayTracingLocalResourceBinder
 		const uint32 BindOffsetBase = RootSignature->GetBindSlotOffsetInBytes(BaseSlotIndex);
 		const uint32 DescriptorSize = uint32(sizeof(D3D12_GPU_VIRTUAL_ADDRESS));
 		const uint32 CurrentOffset = BindOffsetBase + DescriptorIndex * DescriptorSize;
-		ShaderTable->SetHitGroupParameters(RecordIndex, CurrentOffset, Address);
+		ShaderTable->SetLocalShaderParameters(RecordIndex, CurrentOffset, Address);
 	}
 
 	void SetRootCBV(uint32 BaseSlotIndex, uint32 DescriptorIndex, D3D12_GPU_VIRTUAL_ADDRESS Address)
@@ -2303,7 +2392,7 @@ struct FD3D12RayTracingLocalResourceBinder
 	void SetRootDescriptorTable(uint32 SlotIndex, D3D12_GPU_DESCRIPTOR_HANDLE DescriptorTable)
 	{
 		const uint32 BindOffset = RootSignature->GetBindSlotOffsetInBytes(SlotIndex);
-		ShaderTable->SetHitGroupParameters(RecordIndex, BindOffset, DescriptorTable);
+		ShaderTable->SetLocalShaderParameters(RecordIndex, BindOffset, DescriptorTable);
 	}
 
 	void AddResourceReference(FD3D12Resource* D3D12Resource, FRHIResource* RHIResource)
@@ -2784,12 +2873,14 @@ void FD3D12CommandContext::RHISetRayTracingHitGroup(
 	checkf(ShaderSlot < Scene->ShaderSlotsPerGeometrySegment, TEXT("Shader slot is invalid. Make sure that ShaderSlotsPerGeometrySegment is correct on FRayTracingSceneInitializer."));
 
 	FD3D12RayTracingShaderTable* ShaderTable = Scene->FindOrCreateShaderTable(Pipeline);
+	checkf(ShaderTable->LocalShaderTableOffset == ShaderTable->HitGroupShaderTableOffset,
+		TEXT("Hit shader records are assumed to be at the beginning of local shader table"));
 
 	const uint32 RecordIndex = Scene->GetHitRecordBaseIndex(InstanceIndex, SegmentIndex) + ShaderSlot;
-	ShaderTable->SetHitGroupIdentifier(RecordIndex, Pipeline->HitGroupShaders.Identifiers[HitGroupIndex]);
+	ShaderTable->SetLocalShaderIdentifier(RecordIndex, Pipeline->HitGroupShaders.Identifiers[HitGroupIndex]);
 
 	const uint32 UserDataOffset = offsetof(FHitGroupSystemParameters, RootConstants) + offsetof(FHitGroupSystemRootConstants, UserData);
-	ShaderTable->SetHitGroupParameters(RecordIndex, UserDataOffset, UserData);
+	ShaderTable->SetLocalShaderParameters(RecordIndex, UserDataOffset, UserData);
 
 	const FD3D12RayTracingShader* Shader = Pipeline->HitGroupShaders.Shaders[HitGroupIndex];
 
@@ -2804,4 +2895,37 @@ void FD3D12CommandContext::RHISetRayTracingHitGroup(
 		*(ShaderTable->DescriptorCache), ResourceBinder);
 }
 
+
+void FD3D12CommandContext::RHISetRayTracingCallableShader(
+	FRayTracingSceneRHIParamRef InScene, uint32 ShaderSlotInScene,
+	FRayTracingPipelineStateRHIParamRef InPipeline, uint32 ShaderIndexInPipeline,
+	uint32 NumUniformBuffers, const FUniformBufferRHIParamRef* UniformBuffers,
+	uint32 UserData)
+{
+	FD3D12RayTracingScene* Scene = FD3D12DynamicRHI::ResourceCast(InScene);
+	FD3D12RayTracingPipelineState* Pipeline = FD3D12DynamicRHI::ResourceCast(InPipeline);
+
+	checkf(ShaderSlotInScene < Scene->NumCallableShaderSlots, TEXT("Shader slot is invalid. Make sure that NumCallableShaderSlots is correct on FRayTracingSceneInitializer."));
+
+	FD3D12RayTracingShaderTable* ShaderTable = Scene->FindOrCreateShaderTable(Pipeline);
+
+	const uint32 RecordIndex = ShaderTable->CallableShaderRecordIndexOffset + ShaderSlotInScene;
+	ShaderTable->SetLocalShaderIdentifier(RecordIndex, Pipeline->CallableShaders.Identifiers[ShaderIndexInPipeline]);
+
+	const uint32 UserDataOffset = offsetof(FHitGroupSystemParameters, RootConstants) + offsetof(FHitGroupSystemRootConstants, UserData);
+	ShaderTable->SetLocalShaderParameters(RecordIndex, UserDataOffset, UserData);
+
+	const FD3D12RayTracingShader* Shader = Pipeline->CallableShaders.Shaders[ShaderIndexInPipeline];
+
+	FD3D12RayTracingLocalResourceBinder ResourceBinder(*this, ShaderTable, Shader->pRootSignature, RecordIndex);
+	check(ShaderTable->DescriptorCache);
+	SetRayTracingShaderResources(*this, Shader,
+		0, nullptr, // Textures
+		0, nullptr, // SRVs
+		NumUniformBuffers, UniformBuffers,
+		0, nullptr, // Samplers
+		0, nullptr, // UAVs
+		*(ShaderTable->DescriptorCache), ResourceBinder);
+}
+
 #endif // D3D12_RHI_RAYTRACING
diff --git a/Engine/Source/Runtime/D3D12RHI/Private/D3D12RayTracing.h b/Engine/Source/Runtime/D3D12RHI/Private/D3D12RayTracing.h
--- a/Engine/Source/Runtime/D3D12RHI/Private/D3D12RayTracing.h
+++ b/Engine/Source/Runtime/D3D12RHI/Private/D3D12RayTracing.h
@@ -71,6 +71,8 @@ public:
 
 	uint32 TotalPrimitiveCount = 0; // Combined number of primitives in all geometry instances
 
+	uint32 NumCallableShaderSlots = 0;
+
 	// #dxr_todo UE-68230: shader tables should be explicitly registered and unregistered with the scene
 	FD3D12RayTracingShaderTable* FindOrCreateShaderTable(const FD3D12RayTracingPipelineState* Pipeline);
 	FD3D12RayTracingShaderTable* FindExistingShaderTable(const FD3D12RayTracingPipelineState* Pipeline) const;
diff --git a/Engine/Source/Runtime/D3D12RHI/Private/D3D12Util.cpp b/Engine/Source/Runtime/D3D12RHI/Private/D3D12Util.cpp
--- a/Engine/Source/Runtime/D3D12RHI/Private/D3D12Util.cpp
+++ b/Engine/Source/Runtime/D3D12RHI/Private/D3D12Util.cpp
@@ -458,6 +458,7 @@ void QuantizeBoundShaderState(
 		break;
 
 	case SF_RayHitGroup:
+	case SF_RayCallable:
 
 		// Local root signature is used for hit group shaders, using the exact number of resources to minimize shader binding table record size.
 
diff --git a/Engine/Source/Runtime/Engine/Private/ShaderCompiler/ShaderCompiler.cpp b/Engine/Source/Runtime/Engine/Private/ShaderCompiler/ShaderCompiler.cpp
--- a/Engine/Source/Runtime/Engine/Private/ShaderCompiler/ShaderCompiler.cpp
+++ b/Engine/Source/Runtime/Engine/Private/ShaderCompiler/ShaderCompiler.cpp
@@ -3014,6 +3014,7 @@ void GlobalBeginCompileShader(
 		Input.Environment.SetDefine(TEXT("GEOMETRYSHADER"), Target.Frequency == SF_Geometry);
 		Input.Environment.SetDefine(TEXT("COMPUTESHADER"), Target.Frequency == SF_Compute);
 #if RHI_RAYTRACING
+		Input.Environment.SetDefine(TEXT("RAYCALLABLESHADER"), Target.Frequency == SF_RayCallable);
 		Input.Environment.SetDefine(TEXT("RAYHITGROUPSHADER"), Target.Frequency == SF_RayHitGroup);
 		Input.Environment.SetDefine(TEXT("RAYGENSHADER"), Target.Frequency == SF_RayGen);
 		Input.Environment.SetDefine(TEXT("RAYMISSSHADER"), Target.Frequency == SF_RayMiss);
diff --git a/Engine/Source/Runtime/OpenGLDrv/Private/OpenGLShaders.cpp b/Engine/Source/Runtime/OpenGLDrv/Private/OpenGLShaders.cpp
--- a/Engine/Source/Runtime/OpenGLDrv/Private/OpenGLShaders.cpp
+++ b/Engine/Source/Runtime/OpenGLDrv/Private/OpenGLShaders.cpp
@@ -3899,7 +3899,7 @@ void FOpenGLDynamicRHI::BindPendingShaderState( FOpenGLContextState& ContextStat
 	{
 		int32 NextUniformBufferIndex = OGL_FIRST_UNIFORM_BUFFER;
 
-		static_assert(SF_Compute == 5 && SF_NumFrequencies == 9, "Unexpected SF_ ordering");
+		static_assert(SF_Compute == 5 && SF_NumFrequencies == 10, "Unexpected SF_ ordering");
 		static_assert(SF_RayGen > SF_Compute, "SF_Compute must be at the end of the list of frequencies supported in OpenGL");
 
 		int32 NumUniformBuffers[SF_Compute];
diff --git a/Engine/Source/Runtime/RHI/Private/PipelineStateCache.cpp b/Engine/Source/Runtime/RHI/Private/PipelineStateCache.cpp
--- a/Engine/Source/Runtime/RHI/Private/PipelineStateCache.cpp
+++ b/Engine/Source/Runtime/RHI/Private/PipelineStateCache.cpp
@@ -46,17 +46,6 @@ static inline uint32 GetTypeHash(const FGraphicsPipelineStateInitializer& Initia
 		^ Initializer.RenderTargetsEnabled ^ GetTypeHash(Initializer.RasterizerState) ^ GetTypeHash(Initializer.DepthStencilState);
 }
 
-#if RHI_RAYTRACING
-static inline uint32 GetTypeHash(const FRayTracingPipelineStateInitializer& Initializer)
-{
-	return GetTypeHash(Initializer.MaxPayloadSizeInBytes) ^
-		GetTypeHash(Initializer.bAllowHitGroupIndexing) ^
-		GetTypeHash(Initializer.GetRayGenHash()) ^
-		GetTypeHash(Initializer.GetRayMissHash()) ^
-		GetTypeHash(Initializer.GetHitGroupHash());
-}
-#endif
-
 static TAutoConsoleVariable<int32> GCVarAsyncPipelineCompile(
 	TEXT("r.AsyncPipelineCompile"),
 	1,
diff --git a/Engine/Source/Runtime/RHI/Public/DynamicRHI.h b/Engine/Source/Runtime/RHI/Public/DynamicRHI.h
--- a/Engine/Source/Runtime/RHI/Public/DynamicRHI.h
+++ b/Engine/Source/Runtime/RHI/Public/DynamicRHI.h
@@ -130,6 +130,10 @@ struct FRayTracingSceneInitializer
 	// Desired hit shader can be selected by providing appropriate RayContributionToHitGroupIndex to TraceRay() function.
 	// Use ShaderSlot argument in SetRayTracingHitGroup() to assign shaders and resources for specific part of the shder binding table record.
 	uint32 ShaderSlotsPerGeometrySegment = 1;
+
+	// Defines how many different callable shaders with unique resource bindings can be bound to this scene.
+	// Shaders and resources are assigned to slots in the scene using SetRayTracingCallableShader().
+	uint32 NumCallableShaderSlots = 0;
 };
 
 /** The interface which is implemented by the dynamically bound RHI. */
diff --git a/Engine/Source/Runtime/RHI/Public/RHICommandList.h b/Engine/Source/Runtime/RHI/Public/RHICommandList.h
--- a/Engine/Source/Runtime/RHI/Public/RHICommandList.h
+++ b/Engine/Source/Runtime/RHI/Public/RHICommandList.h
@@ -1939,19 +1939,26 @@ struct FRHICommandRayTraceDispatch final : public FRHICommand<FRHICommandRayTrac
 	RHI_API void Execute(FRHICommandListBase& CmdList);
 };
 
-struct FRHICommandSetRayTracingHitGroup final : public FRHICommand<FRHICommandSetRayTracingHitGroup>
+struct FRHICommandSetRayTracingBindings final : public FRHICommand<FRHICommandSetRayTracingBindings>
 {
+	enum EBindingType
+	{
+		EBindingType_HitGroup,
+		EBindingType_CallableShader,
+	};
+
 	FRayTracingSceneRHIParamRef Scene;
 	uint32 InstanceIndex;
 	uint32 SegmentIndex;
 	uint32 ShaderSlot;
 	FRayTracingPipelineStateRHIParamRef Pipeline;
-	uint32 HitGroupIndex;
+	uint32 ShaderIndex;
 	uint32 NumUniformBuffers;
 	const FUniformBufferRHIParamRef* UniformBuffers; // Pointer to an array of uniform buffers, allocated inline within the command list
 	uint32 UserData;
+	EBindingType BindingType;
 
-	FRHICommandSetRayTracingHitGroup(FRayTracingSceneRHIParamRef InScene, uint32 InInstanceIndex, uint32 InSegmentIndex, uint32 InShaderSlot,
+	FRHICommandSetRayTracingBindings(FRayTracingSceneRHIParamRef InScene, uint32 InInstanceIndex, uint32 InSegmentIndex, uint32 InShaderSlot,
 		FRayTracingPipelineStateRHIParamRef InPipeline, uint32 InHitGroupIndex, uint32 InNumUniformBuffers, const FUniformBufferRHIParamRef* InUniformBuffers,
 		uint32 InUserData)
 		: Scene(InScene)
@@ -1959,10 +1966,28 @@ struct FRHICommandSetRayTracingHitGroup final : public FRHICommand<FRHICommandSe
 		, SegmentIndex(InSegmentIndex)
 		, ShaderSlot(InShaderSlot)
 		, Pipeline(InPipeline)
-		, HitGroupIndex(InHitGroupIndex)
+		, ShaderIndex(InHitGroupIndex)
 		, NumUniformBuffers(InNumUniformBuffers)
 		, UniformBuffers(InUniformBuffers)
 		, UserData(InUserData)
+		, BindingType(EBindingType_HitGroup)
+	{
+	}
+
+	FRHICommandSetRayTracingBindings(FRayTracingSceneRHIParamRef InScene, uint32 InShaderSlot,
+		FRayTracingPipelineStateRHIParamRef InPipeline, uint32 InShaderIndex, 
+		uint32 InNumUniformBuffers, const FUniformBufferRHIParamRef* InUniformBuffers,
+		uint32 InUserData)
+		: Scene(InScene)
+		, InstanceIndex(0)
+		, SegmentIndex(0)
+		, ShaderSlot(InShaderSlot)
+		, Pipeline(InPipeline)
+		, ShaderIndex(InShaderIndex)
+		, NumUniformBuffers(InNumUniformBuffers)
+		, UniformBuffers(InUniformBuffers)
+		, UserData(InUserData)
+		, BindingType(EBindingType_CallableShader)
 	{
 	}
 
@@ -3031,7 +3056,33 @@ public:
 				}
 			}
 
-			ALLOC_COMMAND(FRHICommandSetRayTracingHitGroup)(Scene, InstanceIndex, SegmentIndex, ShaderSlot, Pipeline, HitGroupIndex, NumUniformBuffers, InlineUniformBuffers, UserData);
+			ALLOC_COMMAND(FRHICommandSetRayTracingBindings)(Scene, InstanceIndex, SegmentIndex, ShaderSlot, Pipeline, HitGroupIndex, NumUniformBuffers, InlineUniformBuffers, UserData);
+		}
+	}
+
+	FORCEINLINE_DEBUGGABLE void SetRayTracingCallableShader(
+		FRayTracingSceneRHIParamRef Scene, uint32 ShaderSlotInScene,
+		FRayTracingPipelineStateRHIParamRef Pipeline, uint32 ShaderIndexInPipeline,
+		uint32 NumUniformBuffers, const FUniformBufferRHIParamRef* UniformBuffers,
+		uint32 UserData)
+	{
+		if (Bypass())
+		{
+			GetContext().RHISetRayTracingCallableShader(Scene, ShaderSlotInScene, Pipeline, ShaderIndexInPipeline, NumUniformBuffers, UniformBuffers, UserData);
+		}
+		else
+		{
+			FUniformBufferRHIParamRef* InlineUniformBuffers = nullptr;
+			if (NumUniformBuffers)
+			{
+				InlineUniformBuffers = (FUniformBufferRHIParamRef*)Alloc(sizeof(FUniformBufferRHIParamRef) * NumUniformBuffers, alignof(FUniformBufferRHIParamRef));
+				for (uint32 Index = 0; Index < NumUniformBuffers; ++Index)
+				{
+					InlineUniformBuffers[Index] = UniformBuffers[Index];
+				}
+			}
+
+			ALLOC_COMMAND(FRHICommandSetRayTracingBindings)(Scene, ShaderSlotInScene, Pipeline, ShaderIndexInPipeline, NumUniformBuffers, InlineUniformBuffers, UserData);
 		}
 	}
 
diff --git a/Engine/Source/Runtime/RHI/Public/RHICommandListCommandExecutes.inl b/Engine/Source/Runtime/RHI/Public/RHICommandListCommandExecutes.inl
--- a/Engine/Source/Runtime/RHI/Public/RHICommandListCommandExecutes.inl
+++ b/Engine/Source/Runtime/RHI/Public/RHICommandListCommandExecutes.inl
@@ -66,7 +66,7 @@ struct FRHICommandBuildAccelerationStructure;
 struct FRHICommandRayTraceOcclusion;
 struct FRHICommandRayTraceIntersection;
 struct FRHICommandRayTraceDispatch;
-struct FRHICommandSetRayTracingHitGroup;
+struct FRHICommandSetRayTracingBindings;
 
 enum class ECmdList;
 template <typename TShaderRHIParamRef> struct FRHICommandSetLocalUniformBuffer;
@@ -575,10 +575,17 @@ void FRHICommandRayTraceDispatch::Execute(FRHICommandListBase& CmdList)
 	INTERNAL_DECORATOR(RHIRayTraceDispatch)(Pipeline, RayGenShader, Scene, GlobalResourceBindings, Width, Height);
 }
 
-void FRHICommandSetRayTracingHitGroup::Execute(FRHICommandListBase& CmdList)
+void FRHICommandSetRayTracingBindings::Execute(FRHICommandListBase& CmdList)
 {
 	RHISTAT(SetRayTracingHitGroup);
-	INTERNAL_DECORATOR(RHISetRayTracingHitGroup)(Scene, InstanceIndex, SegmentIndex, ShaderSlot, Pipeline, HitGroupIndex, NumUniformBuffers, UniformBuffers, UserData);
+	if (BindingType == EBindingType_HitGroup)
+	{
+		INTERNAL_DECORATOR(RHISetRayTracingHitGroup)(Scene, InstanceIndex, SegmentIndex, ShaderSlot, Pipeline, ShaderIndex, NumUniformBuffers, UniformBuffers, UserData);
+	}
+	else
+	{
+		INTERNAL_DECORATOR(RHISetRayTracingCallableShader)(Scene, ShaderSlot, Pipeline, ShaderIndex, NumUniformBuffers, UniformBuffers, UserData);
+	}
 }
 
 #endif // RHI_RAYTRACING
diff --git a/Engine/Source/Runtime/RHI/Public/RHIContext.h b/Engine/Source/Runtime/RHI/Public/RHIContext.h
--- a/Engine/Source/Runtime/RHI/Public/RHIContext.h
+++ b/Engine/Source/Runtime/RHI/Public/RHIContext.h
@@ -627,6 +627,15 @@ public:
 		checkNoEntry();
 	}
 
+	virtual void RHISetRayTracingCallableShader(
+		FRayTracingSceneRHIParamRef Scene, uint32 ShaderSlotInScene,
+		FRayTracingPipelineStateRHIParamRef Pipeline, uint32 ShaderIndexInPipeline,
+		uint32 NumUniformBuffers, const FUniformBufferRHIParamRef* UniformBuffers,
+		uint32 UserData)
+	{
+		checkNoEntry();
+	}
+
 	protected:
 		FRHIRenderPassInfo RenderPassInfo;
 };
diff --git a/Engine/Source/Runtime/RHI/Public/RHIDefinitions.h b/Engine/Source/Runtime/RHI/Public/RHIDefinitions.h
--- a/Engine/Source/Runtime/RHI/Public/RHIDefinitions.h
+++ b/Engine/Source/Runtime/RHI/Public/RHIDefinitions.h
@@ -34,8 +34,9 @@ enum EShaderFrequency
 	SF_RayGen			= 6,
 	SF_RayMiss			= 7,
 	SF_RayHitGroup		= 8,
+	SF_RayCallable		= 9,
 
-	SF_NumFrequencies	= 9,
+	SF_NumFrequencies	= 10,
 	SF_NumBits			= 4,
 };
 static_assert(SF_NumFrequencies <= (1 << SF_NumBits), "SF_NumFrequencies will not fit on SF_NumBits");
@@ -1204,6 +1205,7 @@ inline const TCHAR* GetShaderFrequencyString(EShaderFrequency Frequency, bool bI
 	case SF_RayGen:			String = TEXT("SF_RayGen"); break;
 	case SF_RayMiss:		String = TEXT("SF_RayMiss"); break;
 	case SF_RayHitGroup:	String = TEXT("SF_RayHitGroup"); break;
+	case SF_RayCallable:	String = TEXT("SF_RayCallable"); break;
 
 	default:
 		checkf(0, TEXT("Unknown ShaderFrequency %d"), (int32)Frequency);
diff --git a/Engine/Source/Runtime/RHI/Public/RHIResources.h b/Engine/Source/Runtime/RHI/Public/RHIResources.h
--- a/Engine/Source/Runtime/RHI/Public/RHIResources.h
+++ b/Engine/Source/Runtime/RHI/Public/RHIResources.h
@@ -2090,7 +2090,18 @@ public:
 			&& bAllowHitGroupIndexing == rhs.bAllowHitGroupIndexing
 			&& RayGenHash == rhs.RayGenHash
 			&& MissHash == rhs.MissHash
-			&& HitGroupHash == rhs.HitGroupHash;
+			&& HitGroupHash == rhs.HitGroupHash
+			&& CallableHash == rhs.CallableHash;
+	}
+
+	static friend uint32 GetTypeHash(const FRayTracingPipelineStateInitializer& Initializer)
+	{
+		return GetTypeHash(Initializer.MaxPayloadSizeInBytes) ^
+			GetTypeHash(Initializer.bAllowHitGroupIndexing) ^
+			GetTypeHash(Initializer.GetRayGenHash()) ^
+			GetTypeHash(Initializer.GetRayMissHash()) ^
+			GetTypeHash(Initializer.GetHitGroupHash()) ^
+			GetTypeHash(Initializer.GetCallableHash());
 	}
 
 	uint32 MaxPayloadSizeInBytes = 32; // sizeof FDefaultPayload declared in RayTracingCommon.ush
@@ -2100,6 +2111,7 @@ public:
 	const TArrayView<const FRayTracingShaderRHIParamRef>& GetRayGenTable()   const { return RayGenTable; }
 	const TArrayView<const FRayTracingShaderRHIParamRef>& GetMissTable()     const { return MissTable; }
 	const TArrayView<const FRayTracingShaderRHIParamRef>& GetHitGroupTable() const { return HitGroupTable; }
+	const TArrayView<const FRayTracingShaderRHIParamRef>& GetCallableTable() const { return CallableTable; }
 
 	// Shaders used as entry point to ray tracing work. At least one RayGen shader must be provided.
 	void SetRayGenShaderTable(const TArrayView<const FRayTracingShaderRHIParamRef>& InRayGenShaders)
@@ -2125,9 +2137,18 @@ public:
 		HitGroupHash = ComputeShaderTableHash(HitGroupTable);
 	}
 
+	// Shaders that can be explicitly invoked from RayGen shaders by their Shader Binding Table (SBT) index.
+	// SetRayTracingCallableShader() command must be used to fill SBT slots before a shader can be called.
+	void SetCallableTable(const TArrayView<const FRayTracingShaderRHIParamRef>& InCallableShaders)
+	{
+		CallableTable = InCallableShaders;
+		CallableHash = ComputeShaderTableHash(CallableTable);
+	}
+
 	uint64 GetHitGroupHash() const { return HitGroupHash; }
 	uint64 GetRayGenHash()   const { return RayGenHash; }
 	uint64 GetRayMissHash()  const { return MissHash; }
+	uint64 GetCallableHash() const { return CallableHash; }
 
 private:
 
@@ -2149,10 +2170,12 @@ private:
 	TArrayView<const FRayTracingShaderRHIParamRef> RayGenTable;
 	TArrayView<const FRayTracingShaderRHIParamRef> MissTable;
 	TArrayView<const FRayTracingShaderRHIParamRef> HitGroupTable;
+	TArrayView<const FRayTracingShaderRHIParamRef> CallableTable;
 
 	uint64 RayGenHash = 0;
 	uint64 MissHash = 0;
 	uint64 HitGroupHash = 0;
+	uint64 CallableHash = 0;
 };
 #endif // RHI_RAYTRACING
 
diff --git a/Engine/Source/Runtime/RenderCore/Private/Shader.cpp b/Engine/Source/Runtime/RenderCore/Private/Shader.cpp
--- a/Engine/Source/Runtime/RenderCore/Private/Shader.cpp
+++ b/Engine/Source/Runtime/RenderCore/Private/Shader.cpp
@@ -926,7 +926,7 @@ void FShaderResource::InitRHI()
 		UE_CLOG((bCodeInSharedLocation && !IsValidRef(Shader)), LogShaders, Fatal, TEXT("FShaderResource::SerializeShaderCode can't find shader code for: [%s]"), *LegacyShaderPlatformToShaderFormat((EShaderPlatform)Target.Platform).ToString());
 	}
 #if RHI_RAYTRACING
-	else if (Target.Frequency == SF_RayGen || Target.Frequency == SF_RayMiss || Target.Frequency == SF_RayHitGroup)
+	else if (Target.Frequency == SF_RayGen || Target.Frequency == SF_RayMiss || Target.Frequency == SF_RayHitGroup || Target.Frequency == SF_RayCallable)
 	{
 		if (GRHISupportsRayTracing)
 		{
diff --git a/Engine/Source/Runtime/RenderCore/Public/Shader.h b/Engine/Source/Runtime/RenderCore/Public/Shader.h
--- a/Engine/Source/Runtime/RenderCore/Public/Shader.h
+++ b/Engine/Source/Runtime/RenderCore/Public/Shader.h
@@ -358,7 +358,8 @@ public:
 	{
 		checkSlow(Target.Frequency == SF_RayGen
 			   || Target.Frequency == SF_RayMiss
-			   || Target.Frequency == SF_RayHitGroup);
+			   || Target.Frequency == SF_RayHitGroup
+			   || Target.Frequency == SF_RayCallable);
 
 		if (!IsInitialized())
 		{
@@ -369,9 +370,7 @@ public:
 
 	inline uint32 GetRayTracingMaterialLibraryIndex()
 	{
-		checkSlow(Target.Frequency == SF_RayGen
-			|| Target.Frequency == SF_RayMiss
-			|| Target.Frequency == SF_RayHitGroup);
+		checkSlow(Target.Frequency == SF_RayHitGroup);
 
 		if (!IsInitialized())
 		{
diff --git a/Engine/Source/Runtime/RenderCore/Public/ShaderCore.h b/Engine/Source/Runtime/RenderCore/Public/ShaderCore.h
--- a/Engine/Source/Runtime/RenderCore/Public/ShaderCore.h
+++ b/Engine/Source/Runtime/RenderCore/Public/ShaderCore.h
@@ -64,7 +64,7 @@ DECLARE_MEMORY_STAT_EXTERN(TEXT("Shader MapMemory"),STAT_Shaders_ShaderMapMemory
 
 inline TStatId GetMemoryStatType(EShaderFrequency ShaderFrequency)
 {
-	static_assert(9 == SF_NumFrequencies, "EShaderFrequency has a bad size.");
+	static_assert(10 == SF_NumFrequencies, "EShaderFrequency has a bad size.");
 
 	switch(ShaderFrequency)
 	{
@@ -73,6 +73,7 @@ inline TStatId GetMemoryStatType(EShaderFrequency ShaderFrequency)
 		case SF_RayGen:				return GET_STATID(STAT_PixelShaderMemory);
 		case SF_RayMiss:			return GET_STATID(STAT_PixelShaderMemory);
 		case SF_RayHitGroup:		return GET_STATID(STAT_PixelShaderMemory);
+		case SF_RayCallable:		return GET_STATID(STAT_PixelShaderMemory);
 	}
 	return GET_STATID(STAT_VertexShaderMemory);
 }
