Added support for procedural ray tracing geometry and intersection shaders

- Added ERayTracingGeometryType to differentiate between triangle and procedural geometries
- Removed bIsDynamic from FRayTracingSceneInitializer, as dynamic refit is not currently supported
- Changed FRayTracingGeometryInstance::Mask default value from 0 to 0xFF, as it's a more sensible default (0 causes geometry to be invisible)

diff --git a/Engine/Plugins/Experimental/GeometryCache/Source/GeometryCache/Private/GeometryCacheSceneProxy.cpp b/Engine/Plugins/Experimental/GeometryCache/Source/GeometryCache/Private/GeometryCacheSceneProxy.cpp
--- a/Engine/Plugins/Experimental/GeometryCache/Source/GeometryCache/Private/GeometryCacheSceneProxy.cpp
+++ b/Engine/Plugins/Experimental/GeometryCache/Source/GeometryCache/Private/GeometryCacheSceneProxy.cpp
@@ -161,7 +161,7 @@ FGeometryCacheSceneProxy::FGeometryCacheSceneProxy(UGeometryCacheComponent* Comp
 							Initializer.VertexBufferByteOffset = 0;
 							Initializer.TotalPrimitiveCount = Section->IndexBuffer.NumIndices / 3;
 							Initializer.VertexBufferElementType = VET_Float3;
-							Initializer.PrimitiveType = PT_TriangleList;
+							Initializer.GeometryType = RTGT_Triangles;
 							Initializer.bFastBuild = false;
 
 							Section->RayTracingGeometry.SetInitializer(Initializer);
diff --git a/Engine/Plugins/FX/Niagara/Source/Niagara/Private/NiagaraRendererSprites.cpp b/Engine/Plugins/FX/Niagara/Source/Niagara/Private/NiagaraRendererSprites.cpp
--- a/Engine/Plugins/FX/Niagara/Source/Niagara/Private/NiagaraRendererSprites.cpp
+++ b/Engine/Plugins/FX/Niagara/Source/Niagara/Private/NiagaraRendererSprites.cpp
@@ -105,7 +105,7 @@ void NiagaraRendererSprites::CreateRenderThreadResources()
 		Initializer.VertexBufferByteOffset = 0;
 		Initializer.TotalPrimitiveCount = 0;
 		Initializer.VertexBufferElementType = VET_Float3;
-		Initializer.PrimitiveType = PT_TriangleList;
+		Initializer.GeometryType = RTGT_Triangles;
 		Initializer.bFastBuild = true;
 		Initializer.bAllowUpdate = false;
 		RayTracingGeometry.SetInitializer(Initializer);
diff --git a/Engine/Source/Runtime/D3D12RHI/Private/D3D12RayTracing.cpp b/Engine/Source/Runtime/D3D12RHI/Private/D3D12RayTracing.cpp
--- a/Engine/Source/Runtime/D3D12RHI/Private/D3D12RayTracing.cpp
+++ b/Engine/Source/Runtime/D3D12RHI/Private/D3D12RayTracing.cpp
@@ -485,6 +485,7 @@ public:
 			if (CollectionType == ECollectionType::HitGroup)
 			{
 				HitGroupDesc.HitGroupExport = *PrimaryExportName;
+				HitGroupDesc.Type = Shader->IntersectionEntryPoint.IsEmpty() ? D3D12_HIT_GROUP_TYPE_TRIANGLES : D3D12_HIT_GROUP_TYPE_PROCEDURAL_PRIMITIVE;
 
 				{
 					FString& ExportName = Entry.ExportNames.Add_GetRef(GenerateShaderName(TEXT("CHS"), ShaderHash));
@@ -510,6 +511,7 @@ public:
 
 					LocalRootSignatureAssociations.Add(0);
 				}
+
 				if (!Shader->IntersectionEntryPoint.IsEmpty())
 				{
 					FString& ExportName = Entry.ExportNames.Add_GetRef(GenerateShaderName(TEXT("IS"), ShaderHash));
@@ -1636,18 +1638,25 @@ FRayTracingGeometryRHIRef FD3D12DynamicRHI::RHICreateRayTracingGeometry(const FR
 	checkf(Initializer.VertexBufferStride, TEXT("Position vertex buffer is required for ray tracing geometry"));
 	checkf(Initializer.VertexBufferStride % 4 == 0, TEXT("Position vertex buffer stride must be aligned to 4 bytes for ByteAddressBuffer loads to work"));
 
-	// #dxr_todo UE-72160: VET_Half4 (DXGI_FORMAT_R16G16B16A16_FLOAT) is also supported by DXR. Should we support it?
-	check(Initializer.VertexBufferElementType == VET_Float3 || Initializer.VertexBufferElementType == VET_Float2 || Initializer.VertexBufferElementType == VET_Half2);
-	if (Initializer.IndexBuffer)
+	if (Initializer.GeometryType == RTGT_Triangles)
 	{
-		checkf(Initializer.IndexBuffer->GetStride() == 2 || Initializer.IndexBuffer->GetStride() == 4, TEXT("Index buffer must be 16 or 32 bit."));
+		// #dxr_todo UE-72160: VET_Half4 (DXGI_FORMAT_R16G16B16A16_FLOAT) is also supported by DXR. Should we support it?
+		check(Initializer.VertexBufferElementType == VET_Float3 || Initializer.VertexBufferElementType == VET_Float2 || Initializer.VertexBufferElementType == VET_Half2);
+
+		// #dxr_todo UE-72160: temporary constraints on vertex and index buffer formats (this will be relaxed when more flexible vertex/index fetching is implemented)
+		checkf(Initializer.VertexBufferElementType == VET_Float3, TEXT("Only float3 vertex buffers are currently implemented.")); // #dxr_todo UE-72160: support other vertex buffer formats
+		checkf(Initializer.VertexBufferStride == 12, TEXT("Only deinterleaved float3 position vertex buffers are currently implemented.")); // #dxr_todo UE-72160: support interleaved vertex buffers
 	}
 
-	checkf(Initializer.PrimitiveType == EPrimitiveType::PT_TriangleList, TEXT("Only TriangleList primitive type is currently supported."));
+	if (Initializer.GeometryType == RTGT_Procedural)
+	{
+		checkf(Initializer.VertexBufferStride >= (2 * sizeof(FVector)), TEXT("Procedural geometry vertex buffer must contain at least 2xFloat3 that defines 3D bounding boxes of primitives."));
+	}
 
-	// #dxr_todo UE-72160: temporary constraints on vertex and index buffer formats (this will be relaxed when more flexible vertex/index fetching is implemented)
-	checkf(Initializer.VertexBufferElementType == VET_Float3, TEXT("Only float3 vertex buffers are currently implemented.")); // #dxr_todo UE-72160: support other vertex buffer formats
-	checkf(Initializer.VertexBufferStride == 12, TEXT("Only deinterleaved float3 position vertex buffers are currently implemented.")); // #dxr_todo UE-72160: support interleaved vertex buffers
+	if (Initializer.IndexBuffer)
+	{
+		checkf(Initializer.IndexBuffer->GetStride() == 2 || Initializer.IndexBuffer->GetStride() == 4, TEXT("Index buffer must be 16 or 32 bit."));
+	}
 
 	checkf(GNumExplicitGPUsForRendering == 1, TEXT("Ray tracing is not implemented for mGPU")); // #dxr_todo UE-72157: implement mGPU support
 	FD3D12RayTracingGeometry* Result = GetAdapter().CreateLinkedObject<FD3D12RayTracingGeometry>(FRHIGPUMask::All(), [&](FD3D12Device* Device)
@@ -1662,6 +1671,18 @@ FRayTracingGeometryRHIRef FD3D12DynamicRHI::RHICreateRayTracingGeometry(const FR
 		Mesh->BaseVertexIndex = Initializer.BaseVertexIndex;
 		Mesh->TotalPrimitiveCount = Initializer.TotalPrimitiveCount;
 
+		switch (Initializer.GeometryType)
+		{
+		case RTGT_Triangles:
+			Mesh->GeometryType = D3D12_RAYTRACING_GEOMETRY_TYPE_TRIANGLES;
+			break;
+		case RTGT_Procedural:
+			Mesh->GeometryType = D3D12_RAYTRACING_GEOMETRY_TYPE_PROCEDURAL_PRIMITIVE_AABBS;
+			break;
+		default:
+			checkf(false, TEXT("Unexpected ray tracing geometry type"));
+		}
+
 		if (Initializer.bFastBuild)
 		{
 			Mesh->BuildFlags = D3D12_RAYTRACING_ACCELERATION_STRUCTURE_BUILD_FLAG_PREFER_FAST_BUILD;
@@ -1831,9 +1852,9 @@ void FD3D12RayTracingGeometry::BuildAccelerationStructure(FD3D12CommandContext&
 	for (const FRayTracingGeometrySegment& Segment : Segments)
 	{
 		D3D12_RAYTRACING_GEOMETRY_DESC Desc = {};
-		Desc.Type = D3D12_RAYTRACING_GEOMETRY_TYPE_TRIANGLES;
 
 		Desc.Flags = D3D12_RAYTRACING_GEOMETRY_FLAG_NONE;
+		Desc.Type = GeometryType;
 
 		if (!Segment.bAllowAnyHitShader)
 		{
@@ -1847,49 +1868,62 @@ void FD3D12RayTracingGeometry::BuildAccelerationStructure(FD3D12CommandContext&
 			Desc.Flags |= D3D12_RAYTRACING_GEOMETRY_FLAG_NO_DUPLICATE_ANYHIT_INVOCATION;
 		}
 
-		switch (VertexElemType)
+		if (GeometryType == D3D12_RAYTRACING_GEOMETRY_TYPE_TRIANGLES)
 		{
-		case VET_Float3:
-			Desc.Triangles.VertexFormat = DXGI_FORMAT_R32G32B32_FLOAT;
-			break;
-		case VET_Float2:
-			Desc.Triangles.VertexFormat = DXGI_FORMAT_R32G32_FLOAT;
-			break;
-		case VET_Half2:
-			Desc.Triangles.VertexFormat = DXGI_FORMAT_R16G16_FLOAT;
-			break;
-		default:
-			checkNoEntry();
-			break;
-		}
+			switch (VertexElemType)
+			{
+			case VET_Float3:
+				Desc.Triangles.VertexFormat = DXGI_FORMAT_R32G32B32_FLOAT;
+				break;
+			case VET_Float2:
+				Desc.Triangles.VertexFormat = DXGI_FORMAT_R32G32_FLOAT;
+				break;
+			case VET_Half2:
+				Desc.Triangles.VertexFormat = DXGI_FORMAT_R16G16_FLOAT;
+				break;
+			default:
+				checkNoEntry();
+				break;
+			}
 
-		Desc.Triangles.Transform3x4 = D3D12_GPU_VIRTUAL_ADDRESS(0);
+			Desc.Triangles.Transform3x4 = D3D12_GPU_VIRTUAL_ADDRESS(0);
 
-		if (IndexBuffer)
-		{
-			Desc.Triangles.IndexFormat = IndexStride == 4 ? DXGI_FORMAT_R32_UINT : DXGI_FORMAT_R16_UINT;
-			Desc.Triangles.IndexCount = Segment.NumPrimitives * IndicesPerPrimitive;
-			Desc.Triangles.IndexBuffer = IndexBuffer->ResourceLocation.GetGPUVirtualAddress() + IndexStride * Segment.FirstPrimitive * IndicesPerPrimitive;
+			if (IndexBuffer)
+			{
+				Desc.Triangles.IndexFormat = IndexStride == 4 ? DXGI_FORMAT_R32_UINT : DXGI_FORMAT_R16_UINT;
+				Desc.Triangles.IndexCount = Segment.NumPrimitives * IndicesPerPrimitive;
+				Desc.Triangles.IndexBuffer = IndexBuffer->ResourceLocation.GetGPUVirtualAddress() + IndexStride * Segment.FirstPrimitive * IndicesPerPrimitive;
 
-			Desc.Triangles.VertexCount = PositionVertexBuffer->ResourceLocation.GetSize() / VertexStrideInBytes;
+				Desc.Triangles.VertexCount = PositionVertexBuffer->ResourceLocation.GetSize() / VertexStrideInBytes;
 
-			IndexBuffer->GetResource()->UpdateResidency(CommandContext.CommandListHandle);
+				IndexBuffer->GetResource()->UpdateResidency(CommandContext.CommandListHandle);
+			}
+			else
+			{
+				// Non-indexed geometry
+				Desc.Triangles.IndexFormat = DXGI_FORMAT_UNKNOWN;
+				Desc.Triangles.IndexCount = 0;
+				Desc.Triangles.IndexBuffer = D3D12_GPU_VIRTUAL_ADDRESS(0);
+
+				checkf(Segments.Num() == 1, TEXT("Non-indexed geometry with multiple segments is not implemented."));
+
+				Desc.Triangles.VertexCount = FMath::Min<uint32>(PositionVertexBuffer->ResourceLocation.GetSize() / VertexStrideInBytes, TotalPrimitiveCount * 3);
+			}
+
+			Desc.Triangles.VertexBuffer.StartAddress = PositionVertexBuffer->ResourceLocation.GetGPUVirtualAddress() + VertexOffsetInBytes;
+			Desc.Triangles.VertexBuffer.StrideInBytes = VertexStrideInBytes;
+		}
+		else if (GeometryType == D3D12_RAYTRACING_GEOMETRY_TYPE_PROCEDURAL_PRIMITIVE_AABBS)
+		{
+			Desc.AABBs.AABBCount = Segment.NumPrimitives;
+			Desc.AABBs.AABBs.StartAddress = PositionVertexBuffer->ResourceLocation.GetGPUVirtualAddress() + VertexOffsetInBytes;
+			Desc.AABBs.AABBs.StrideInBytes = VertexStrideInBytes;
 		}
 		else
 		{
-			// Non-indexed geometry
-			Desc.Triangles.IndexFormat = DXGI_FORMAT_UNKNOWN;
-			Desc.Triangles.IndexCount = 0;
-			Desc.Triangles.IndexBuffer = D3D12_GPU_VIRTUAL_ADDRESS(0);
-
-			checkf(Segments.Num() == 1, TEXT("Non-indexed geometry with multiple segments is not implemented."));
-
-			Desc.Triangles.VertexCount = FMath::Min<uint32>(PositionVertexBuffer->ResourceLocation.GetSize() / VertexStrideInBytes, TotalPrimitiveCount * 3);
+			checkf(false, TEXT("Unexpected ray tracing geometry type"));
 		}
 
-		Desc.Triangles.VertexBuffer.StartAddress = PositionVertexBuffer->ResourceLocation.GetGPUVirtualAddress() + VertexOffsetInBytes;
-		Desc.Triangles.VertexBuffer.StrideInBytes = VertexStrideInBytes;
-
 		PositionVertexBuffer->ResourceLocation.GetResource()->UpdateResidency(CommandContext.CommandListHandle);
 
 		Descs.Add(Desc);
@@ -2075,8 +2109,9 @@ void FD3D12RayTracingScene::BuildAccelerationStructure(FD3D12CommandContext& Com
 		for (const FRayTracingGeometryInstance& Instance : Instances)
 		{
 			FD3D12RayTracingGeometry* Geometry = FD3D12DynamicRHI::ResourceCast(Instance.GeometryRHI);
-			// #patrick todo: temporarily removing check in favor of on-demand BLAS build to diagnose an ordering issue.
-			check(!Geometry->bIsAccelerationStructureDirty); // #dxr_todo: we could probably build BLAS here, if needed (though it may be best to have an explicit build API and just require things to be built at this point)
+
+			checkf(!Geometry->bIsAccelerationStructureDirty, 
+				TEXT("Acceleration structures for all geometries must be built before building the top level acceleration structure for the scene."));
 
 			D3D12_RAYTRACING_INSTANCE_DESC InstanceDesc = {};
 
@@ -2252,8 +2287,12 @@ FD3D12RayTracingShaderTable* FD3D12RayTracingScene::FindOrCreateShaderTable(cons
 				SystemParameters.IndexBuffer = IndexBufferAddress;
 				SystemParameters.VertexBuffer = VertexBufferAddress;
 
-				// #dxr_todo UE-72160: support various vertex buffer layouts (fetch/decode based on vertex stride and format)
-				checkf(Geometry->VertexElemType == VET_Float3, TEXT("Only VET_Float3 is currently implemented and tested. Other formats will be supported in the future."));
+				if (Geometry->GeometryType == D3D12_RAYTRACING_GEOMETRY_TYPE_TRIANGLES)
+				{
+					// #dxr_todo UE-72160: support various vertex buffer layouts (fetch/decode based on vertex stride and format)
+					checkf(Geometry->VertexElemType == VET_Float3, TEXT("Only VET_Float3 is currently implemented and tested. Other formats will be supported in the future."));
+				}
+
 				SystemParameters.RootConstants.SetVertexAndIndexStride(Geometry->VertexStrideInBytes, IndexStride);
 				SystemParameters.RootConstants.IndexBufferOffsetInBytes = IndexStride * Segment.FirstPrimitive * IndicesPerPrimitive;
 
diff --git a/Engine/Source/Runtime/D3D12RHI/Private/D3D12RayTracing.h b/Engine/Source/Runtime/D3D12RHI/Private/D3D12RayTracing.h
--- a/Engine/Source/Runtime/D3D12RHI/Private/D3D12RayTracing.h
+++ b/Engine/Source/Runtime/D3D12RHI/Private/D3D12RayTracing.h
@@ -29,6 +29,8 @@ public:
 	uint32 BaseVertexIndex = 0;
 	uint32 TotalPrimitiveCount = 0; // Combined number of primitives in all mesh segments
 
+	D3D12_RAYTRACING_GEOMETRY_TYPE GeometryType = D3D12_RAYTRACING_GEOMETRY_TYPE_TRIANGLES;
+
 	TArray<FRayTracingGeometrySegment> Segments; // Defines addressable parts of the mesh that can be used for material assignment (one segment = one SBT record)
 	D3D12_RAYTRACING_ACCELERATION_STRUCTURE_BUILD_FLAGS BuildFlags;
 
diff --git a/Engine/Source/Runtime/Engine/Private/SkeletalRenderCPUSkin.cpp b/Engine/Source/Runtime/Engine/Private/SkeletalRenderCPUSkin.cpp
--- a/Engine/Source/Runtime/Engine/Private/SkeletalRenderCPUSkin.cpp
+++ b/Engine/Source/Runtime/Engine/Private/SkeletalRenderCPUSkin.cpp
@@ -360,7 +360,7 @@ void FSkeletalMeshObjectCPUSkin::FSkeletalMeshObjectLOD::InitResources(FSkelMesh
 				Initializer.VertexBufferByteOffset = 0;
 				Initializer.TotalPrimitiveCount = TrianglesCount;
 				Initializer.VertexBufferElementType = VET_Float3;
-				Initializer.PrimitiveType = PT_TriangleList;
+				Initializer.GeometryType = RTGT_Triangles;
 				Initializer.bFastBuild = true;
 
 				TArray<FRayTracingGeometrySegment> GeometrySections;
diff --git a/Engine/Source/Runtime/Engine/Private/SkeletalRenderGPUSkin.cpp b/Engine/Source/Runtime/Engine/Private/SkeletalRenderGPUSkin.cpp
--- a/Engine/Source/Runtime/Engine/Private/SkeletalRenderGPUSkin.cpp
+++ b/Engine/Source/Runtime/Engine/Private/SkeletalRenderGPUSkin.cpp
@@ -363,7 +363,7 @@ void FSkeletalMeshObjectGPUSkin::UpdateDynamicData_RenderThread(FGPUSkinCache* G
 				Initializer.VertexBufferByteOffset = 0;
 				Initializer.TotalPrimitiveCount = TrianglesCount;
 				Initializer.VertexBufferElementType = VET_Float3;
-				Initializer.PrimitiveType = PT_TriangleList;
+				Initializer.GeometryType = RTGT_Triangles;
 				Initializer.bFastBuild = true;
 				Initializer.bAllowUpdate = true;
 
diff --git a/Engine/Source/Runtime/Engine/Private/SkeletalRenderStatic.cpp b/Engine/Source/Runtime/Engine/Private/SkeletalRenderStatic.cpp
--- a/Engine/Source/Runtime/Engine/Private/SkeletalRenderStatic.cpp
+++ b/Engine/Source/Runtime/Engine/Private/SkeletalRenderStatic.cpp
@@ -150,7 +150,7 @@ void FSkeletalMeshObjectStatic::FSkeletalMeshObjectLOD::InitResources(FSkelMeshC
 				Initializer.VertexBufferByteOffset = 0;
 				Initializer.TotalPrimitiveCount = TrianglesCount;
 				Initializer.VertexBufferElementType = VET_Float3;
-				Initializer.PrimitiveType = PT_TriangleList;
+				Initializer.GeometryType = RTGT_Triangles;
 				Initializer.bFastBuild = false;
 
 				TArray<FRayTracingGeometrySegment> GeometrySections;
diff --git a/Engine/Source/Runtime/Engine/Private/StaticMesh.cpp b/Engine/Source/Runtime/Engine/Private/StaticMesh.cpp
--- a/Engine/Source/Runtime/Engine/Private/StaticMesh.cpp
+++ b/Engine/Source/Runtime/Engine/Private/StaticMesh.cpp
@@ -1092,7 +1092,7 @@ void FStaticMeshLODResources::InitResources(UStaticMesh* Parent)
 				Initializer.VertexBufferByteOffset = 0;
 				Initializer.TotalPrimitiveCount = 0; // This is calculated below based on static mesh section data
 				Initializer.VertexBufferElementType = VET_Float3;
-				Initializer.PrimitiveType = PT_TriangleList;
+				Initializer.GeometryType = RTGT_Triangles;
 				Initializer.bFastBuild = false;
 				
 				TArray<FRayTracingGeometrySegment> GeometrySections;
diff --git a/Engine/Source/Runtime/RHI/Public/DynamicRHI.h b/Engine/Source/Runtime/RHI/Public/DynamicRHI.h
--- a/Engine/Source/Runtime/RHI/Public/DynamicRHI.h
+++ b/Engine/Source/Runtime/RHI/Public/DynamicRHI.h
@@ -81,7 +81,7 @@ struct FRayTracingGeometryInstance
 	FRayTracingGeometryRHIParamRef GeometryRHI = nullptr;
 	FMatrix Transform = FMatrix(EForceInit::ForceInitToZero);
 	uint32 UserData = 0;
-	uint8 Mask = 0;
+	uint8 Mask = 0xFF;
 
 	// No any-hit shaders will be invoked for this geometry instance (only closest hit)
 	bool bForceOpaque = false;
@@ -90,6 +90,21 @@ struct FRayTracingGeometryInstance
 	bool bDoubleSided = false;
 };
 
+enum ERayTracingGeometryType
+{
+	// Indexed or non-indexed triangle list with fixed function ray intersection.
+	// Vertex buffer must contain vertex positions as VET_Float3.
+	// Vertex stride must be at least 12 bytes, but may be larger to support custom per-vertex data.
+	// Index buffer may be provided for indexed triangle lists. Implicit triangle list is assumed otherwise.
+	RTGT_Triangles,
+
+	// Custom primitive type that requires an intersection shader.
+	// Vertex buffer for procedural geometry must contain one AABB per primitive as {float3 MinXYZ, float3 MaxXYZ}.
+	// Vertex stride must be at least 24 bytes, but may be larger to support custom per-primitive data.
+	// Index buffers can't be used with procedural geometry.
+	RTGT_Procedural,
+};
+
 struct FRayTracingGeometrySegment
 {
 	uint32 FirstPrimitive = 0;
@@ -112,7 +127,7 @@ struct FRayTracingGeometryInitializer
 	uint32 VertexBufferStride = 0;
 	EVertexElementType VertexBufferElementType = VET_Float3;
 	uint32 BaseVertexIndex = 0;
-	EPrimitiveType PrimitiveType = PT_TriangleList;
+	ERayTracingGeometryType GeometryType = RTGT_Triangles;
 	uint32 TotalPrimitiveCount = 0;
 	TArrayView<FRayTracingGeometrySegment> Segments;
 	bool bFastBuild = false;
@@ -122,7 +137,6 @@ struct FRayTracingGeometryInitializer
 struct FRayTracingSceneInitializer
 {
 	TArrayView<FRayTracingGeometryInstance> Instances;
-	bool bIsDynamic = false;
 
 	// This value controls how many elements will be allocated in the shader binding table per geometry segment.
 	// Changing this value allows different hit shaders to be used for different effects.
diff --git a/Engine/Source/Runtime/Renderer/Private/RayTracing/RayTracingTestbed.cpp b/Engine/Source/Runtime/Renderer/Private/RayTracing/RayTracingTestbed.cpp
--- a/Engine/Source/Runtime/Renderer/Private/RayTracing/RayTracingTestbed.cpp
+++ b/Engine/Source/Runtime/Renderer/Private/RayTracing/RayTracingTestbed.cpp
@@ -91,7 +91,7 @@ void TestBasicRayTracing(bool bValidateResults)
 	GeometryInitializer.VertexBufferStride = sizeof(FVector);
 	GeometryInitializer.VertexBufferElementType = VET_Float3;
 	GeometryInitializer.BaseVertexIndex = 0;
-	GeometryInitializer.PrimitiveType = PT_TriangleList;
+	GeometryInitializer.GeometryType = RTGT_Triangles;
 	GeometryInitializer.TotalPrimitiveCount = 1;
 	GeometryInitializer.bFastBuild = false;
 	FRayTracingGeometryRHIRef Geometry = RHICreateRayTracingGeometry(GeometryInitializer);
@@ -102,7 +102,6 @@ void TestBasicRayTracing(bool bValidateResults)
 
 	FRayTracingSceneInitializer Initializer;
 	Initializer.Instances = Instances;
-	Initializer.bIsDynamic = false;
 	Initializer.ShaderSlotsPerGeometrySegment = RAY_TRACING_NUM_SHADER_SLOTS;
 	FRayTracingSceneRHIParamRef Scene = RHICreateRayTracingScene(Initializer);
 
