Implemented RHICopySubTextureRegion interface function for D3D12 backend.

diff --git a/Engine/Source/Runtime/D3D12RHI/Private/D3D12RHIPrivate.h b/Engine/Source/Runtime/D3D12RHI/Private/D3D12RHIPrivate.h
--- a/Engine/Source/Runtime/D3D12RHI/Private/D3D12RHIPrivate.h
+++ b/Engine/Source/Runtime/D3D12RHI/Private/D3D12RHIPrivate.h
@@ -461,6 +461,8 @@ public:
 		return RHICreateRenderQuery(QueryType);
 	}
 
+	virtual void RHICopySubTextureRegion(FTexture2DRHIParamRef SourceTexture, FTexture2DRHIParamRef DestinationTexture, FBox2D SourceBox, FBox2D DestinationBox) final override;
+
 	void RHICalibrateTimers() override;
 
 #if D3D12_RHI_RAYTRACING
diff --git a/Engine/Source/Runtime/D3D12RHI/Private/D3D12Texture.cpp b/Engine/Source/Runtime/D3D12RHI/Private/D3D12Texture.cpp
--- a/Engine/Source/Runtime/D3D12RHI/Private/D3D12Texture.cpp
+++ b/Engine/Source/Runtime/D3D12RHI/Private/D3D12Texture.cpp
@@ -109,6 +109,35 @@ struct FRHICommandUpdateTexture final : public FRHICommand<FRHICommandUpdateText
 	}
 };
 
+struct FRHICommandCopySubTextureRegion final : public FRHICommand<FRHICommandCopySubTextureRegion>
+{
+	FD3D12TextureBase* DestTexture;
+	uint32 DestX;
+	uint32 DestY;
+	uint32 DestZ;
+	FD3D12TextureBase* SourceTexture;
+	D3D12_BOX SourceBox;
+
+	FORCEINLINE_DEBUGGABLE FRHICommandCopySubTextureRegion(FD3D12TextureBase* InDestTexture, uint32 InDestX, uint32 InDestY, uint32 InDestZ, FD3D12TextureBase* InSourceTexture, const D3D12_BOX& InSourceBox)
+		: DestTexture(InDestTexture)
+		, DestX(InDestX)
+		, DestY(InDestY)
+		, DestZ(InDestZ)
+		, SourceTexture(InSourceTexture)
+		, SourceBox(InSourceBox)
+	{
+	}
+
+	~FRHICommandCopySubTextureRegion()
+	{
+	}
+
+	void Execute(FRHICommandListBase& CmdList)
+	{
+		DestTexture->CopyTextureRegion(DestX, DestY, DestZ, SourceTexture, SourceBox);
+	}
+};
+
 struct FD3D12RHICommandInitializeTexture final : public FRHICommand<FD3D12RHICommandInitializeTexture>
 {
 	FD3D12TextureBase* TextureBase;
@@ -1655,6 +1684,28 @@ void FD3D12TextureBase::UpdateTexture(const D3D12_TEXTURE_COPY_LOCATION& DestCop
 	DEBUG_EXECUTE_COMMAND_CONTEXT(DefaultContext);
 }
 
+void FD3D12TextureBase::CopyTextureRegion(uint32 DestX, uint32 DestY, uint32 DestZ, FD3D12TextureBase* SourceTexture, const D3D12_BOX& SourceBox)
+{
+	FD3D12CommandContext& DefaultContext = GetParentDevice()->GetDefaultCommandContext();
+	FD3D12CommandListHandle& CommandListHandle = DefaultContext.CommandListHandle;
+
+	CD3DX12_TEXTURE_COPY_LOCATION DestCopyLocation(GetResource()->GetResource(), 0);
+	CD3DX12_TEXTURE_COPY_LOCATION SourceCopyLocation(SourceTexture->GetResource()->GetResource(), 0);
+
+	FConditionalScopeResourceBarrier ConditionalScopeResourceBarrierDest(CommandListHandle, GetResource(), D3D12_RESOURCE_STATE_COPY_DEST, DestCopyLocation.SubresourceIndex);
+	FConditionalScopeResourceBarrier ConditionalScopeResourceBarrierSource(CommandListHandle, SourceTexture->GetResource(), D3D12_RESOURCE_STATE_COPY_SOURCE, SourceCopyLocation.SubresourceIndex);
+
+	CommandListHandle.FlushResourceBarriers();
+	CommandListHandle->CopyTextureRegion(
+		&DestCopyLocation,
+		DestX, DestY, DestZ,
+		&SourceCopyLocation,
+		&SourceBox);
+
+	CommandListHandle.UpdateResidency(SourceTexture->GetResource());
+	CommandListHandle.UpdateResidency(GetResource());
+}
+
 void FD3D12TextureBase::InitializeTextureData(FRHICommandListImmediate* RHICmdList, const void* InitData, uint32 InitDataSize, uint32 SizeX, uint32 SizeY, uint32 SizeZ, uint32 NumSlices, uint32 NumMips, EPixelFormat Format, D3D12_RESOURCE_STATES DestinationState)
 {
 	// each mip of each array slice counts as a subresource
@@ -2756,3 +2807,31 @@ void FD3D12DynamicRHI::RHIAliasTextureResources(FTextureRHIParamRef DestTextureR
 		SrcTexture = SrcTexture->GetNextObject();
 	}
 }
+
+void FD3D12DynamicRHI::RHICopySubTextureRegion(FTexture2DRHIParamRef SourceTextureRHI, FTexture2DRHIParamRef DestTextureRHI, FBox2D SourceBox, FBox2D DestinationBox)
+{
+	FD3D12TextureBase* SourceTexture = GetD3D12TextureFromRHITexture(SourceTextureRHI);
+	FD3D12TextureBase* DestTexture = GetD3D12TextureFromRHITexture(DestTextureRHI);
+
+	const uint32 XOffset = (uint32)(DestinationBox.Min.X);
+	const uint32 YOffset = (uint32)(DestinationBox.Min.Y);
+	const uint32 Width = (uint32)(SourceBox.Max.X - SourceBox.Min.X);
+	const uint32 Height = (uint32)(SourceBox.Max.Y - SourceBox.Min.Y);
+
+	const CD3DX12_BOX SourceBoxD3D((LONG)SourceBox.Min.X, (LONG)SourceBox.Min.Y, (LONG)SourceBox.Max.X, (LONG)SourceBox.Max.Y);
+
+	CD3DX12_TEXTURE_COPY_LOCATION DestCopyLocation(DestTexture->GetResource()->GetResource(), 0);
+	CD3DX12_TEXTURE_COPY_LOCATION SourceCopyLocation(SourceTexture->GetResource()->GetResource(), 0);
+
+	FRHICommandListImmediate& RHICmdList = FRHICommandListExecutor::GetImmediateCommandList();
+	if (RHICmdList.Bypass())
+	{
+		FRHICommandCopySubTextureRegion RHICmd(DestTexture, XOffset, YOffset, 0, SourceTexture, SourceBoxD3D);
+		RHICmd.Execute(RHICmdList);
+	}
+	else
+	{
+		ALLOC_COMMAND_CL(RHICmdList, FRHICommandCopySubTextureRegion)(DestTexture, XOffset, YOffset, 0, SourceTexture, SourceBoxD3D);
+	}
+}
+
diff --git a/Engine/Source/Runtime/D3D12RHI/Private/D3D12Texture.h b/Engine/Source/Runtime/D3D12RHI/Private/D3D12Texture.h
--- a/Engine/Source/Runtime/D3D12RHI/Private/D3D12Texture.h
+++ b/Engine/Source/Runtime/D3D12RHI/Private/D3D12Texture.h
@@ -108,6 +108,7 @@ public:
 	}
 
 	void UpdateTexture(const D3D12_TEXTURE_COPY_LOCATION& DestCopyLocation, uint32 DestX, uint32 DestY, uint32 DestZ, const D3D12_TEXTURE_COPY_LOCATION& SourceCopyLocation);
+	void CopyTextureRegion(uint32 DestX, uint32 DestY, uint32 DestZ, FD3D12TextureBase* SourceTexture, const D3D12_BOX& SourceBox);
 	void InitializeTextureData(class FRHICommandListImmediate* RHICmdList, const void* InitData, uint32 InitDataSize, uint32 SizeX, uint32 SizeY, uint32 SizeZ, uint32 NumSlices, uint32 NumMips, EPixelFormat Format, D3D12_RESOURCE_STATES DestinationState);
 
 	/**
