Implemented mGPU support for FD3D12DynamicRHI::RHIUpdateUniformBuffer

diff --git a/Engine/Source/Runtime/D3D12RHI/Private/D3D12UniformBuffer.cpp b/Engine/Source/Runtime/D3D12RHI/Private/D3D12UniformBuffer.cpp
--- a/Engine/Source/Runtime/D3D12RHI/Private/D3D12UniformBuffer.cpp
+++ b/Engine/Source/Runtime/D3D12RHI/Private/D3D12UniformBuffer.cpp
@@ -155,96 +155,101 @@ void FD3D12DynamicRHI::RHIUpdateUniformBuffer(FUniformBufferRHIParamRef UniformB
 	check(IsInRenderingThread());
 	check(UniformBufferRHI);
 
-	checkf(GNumExplicitGPUsForRendering == 1, TEXT("mGPU is support is not implemented for FD3D12DynamicRHI::RHIUpdateUniformBuffer"));
+	const FRHIUniformBufferLayout& Layout = UniformBufferRHI->GetLayout();
+	FRHICommandListImmediate& RHICmdList = FRHICommandListExecutor::GetImmediateCommandList();
 
 	FD3D12UniformBuffer* UniformBuffer = ResourceCast(UniformBufferRHI);
-	const FRHIUniformBufferLayout& Layout = UniformBufferRHI->GetLayout();
 
 	const uint32 NumBytes = Layout.ConstantBufferSize;
 	const int32 NumResources = Layout.Resources.Num();
 
-	check(UniformBuffer->ResourceTable.Num() == NumResources);
-
-	FRHICommandListImmediate& RHICmdList = FRHICommandListExecutor::GetImmediateCommandList();
-
-	const bool bBypass = RHICmdList.Bypass();
-	FD3D12Device* Device = UniformBuffer->GetParentDevice();
-	//FD3D12ResourceLocation is non-copyable, so placement new one on the stack for bypass, or out of the commandlist memory if available. avoids dynamic alloc either way.
-	FD3D12ResourceLocation* UpdatedResourceLocation = bBypass	? new (FMemory_Alloca(sizeof(FD3D12ResourceLocation)))FD3D12ResourceLocation(Device)
-																: new(RHICmdList.Alloc<FD3D12ResourceLocation>()) FD3D12ResourceLocation(Device);
+	// Loop over buffers on all GPUs
+	while (UniformBuffer)
+	{
+		check(UniformBuffer->ResourceTable.Num() == NumResources);
 
-	if (NumBytes > 0)
-	{				
-		void* MappedData = nullptr;
+		const bool bBypass = RHICmdList.Bypass();
+		FD3D12Device* Device = UniformBuffer->GetParentDevice();
+		//FD3D12ResourceLocation is non-copyable, so placement new one on the stack for bypass, or out of the commandlist memory if available. avoids dynamic alloc either way.
+		FD3D12ResourceLocation* UpdatedResourceLocation = bBypass ? new (FMemory_Alloca(sizeof(FD3D12ResourceLocation)))FD3D12ResourceLocation(Device)
+			: new(RHICmdList.Alloc<FD3D12ResourceLocation>()) FD3D12ResourceLocation(Device);
 
-		if (UniformBuffer->UniformBufferUsage == UniformBuffer_MultiFrame)
+		if (NumBytes > 0)
 		{
-			FD3D12DynamicHeapAllocator& Allocator = GetAdapter().GetUploadHeapAllocator(Device->GetGPUIndex());
-			MappedData = Allocator.AllocUploadResource(NumBytes, DEFAULT_CONTEXT_UPLOAD_POOL_ALIGNMENT, *UpdatedResourceLocation);
-		}
-		else
-		{
-			FD3D12FastConstantAllocator& Allocator = GetAdapter().GetTransientUniformBufferAllocator();
+			void* MappedData = nullptr;
+
+			if (UniformBuffer->UniformBufferUsage == UniformBuffer_MultiFrame)
+			{
+				FD3D12DynamicHeapAllocator& Allocator = GetAdapter().GetUploadHeapAllocator(Device->GetGPUIndex());
+				MappedData = Allocator.AllocUploadResource(NumBytes, DEFAULT_CONTEXT_UPLOAD_POOL_ALIGNMENT, *UpdatedResourceLocation);
+			}
+			else
+			{
+				FD3D12FastConstantAllocator& Allocator = GetAdapter().GetTransientUniformBufferAllocator();
 
 #if USE_STATIC_ROOT_SIGNATURE
-			MappedData = Allocator.Allocate(NumBytes, *UpdatedResourceLocation, nullptr);
+				MappedData = Allocator.Allocate(NumBytes, *UpdatedResourceLocation, nullptr);
 #else
-			MappedData = Allocator.Allocate(NumBytes, *UpdatedResourceLocation);
+				MappedData = Allocator.Allocate(NumBytes, *UpdatedResourceLocation);
 #endif
+			}
+
+			check(MappedData != nullptr);
+			FMemory::Memcpy(MappedData, Contents, NumBytes);
 		}
 
-		check(MappedData != nullptr);
-		FMemory::Memcpy(MappedData, Contents, NumBytes);
-	}
 
-	
 
-	FRHIResource** CmdListResources = nullptr;
-	
-	if (NumResources)
-	{
-		CmdListResources = bBypass ? (FRHIResource**)FMemory_Alloca(sizeof(FRHIResource*) * NumResources) : (FRHIResource**)RHICmdList.Alloc(sizeof(FRHIResource*) * NumResources, alignof(FRHIResource*));
-		for (int32 ResourceIndex = 0; ResourceIndex < NumResources; ++ResourceIndex)
-		{
-			EUniformBufferBaseType ResourceType = Layout.Resources[ResourceIndex].MemberType;
+		FRHIResource** CmdListResources = nullptr;
 
-			FRHIResource* Resource;
-			if (IsShaderParameterTypeIgnoredByRHI(ResourceType))
-			{
-				continue;
-			}
-			else if (IsRDGResourceReferenceShaderParameterType(ResourceType))
-			{
-				check(IsInRenderingThread()); // TODO: UE-68018
-				FRHIResource** ResourcePtr = *(FRHIResource***)((uint8*)Contents + Layout.Resources[ResourceIndex].MemberOffset);
-				Resource = ResourcePtr ? *ResourcePtr : nullptr;
-			}
-			else
+		if (NumResources)
+		{
+			CmdListResources = bBypass ? (FRHIResource**)FMemory_Alloca(sizeof(FRHIResource*) * NumResources) : (FRHIResource**)RHICmdList.Alloc(sizeof(FRHIResource*) * NumResources, alignof(FRHIResource*));
+			for (int32 ResourceIndex = 0; ResourceIndex < NumResources; ++ResourceIndex)
 			{
-				Resource = *(FRHIResource**)((uint8*)Contents + Layout.Resources[ResourceIndex].MemberOffset);
+				EUniformBufferBaseType ResourceType = Layout.Resources[ResourceIndex].MemberType;
+
+				FRHIResource* Resource;
+				if (IsShaderParameterTypeIgnoredByRHI(ResourceType))
+				{
+					continue;
+				}
+				else if (IsRDGResourceReferenceShaderParameterType(ResourceType))
+				{
+					check(IsInRenderingThread()); // TODO: UE-68018
+					FRHIResource** ResourcePtr = *(FRHIResource***)((uint8*)Contents + Layout.Resources[ResourceIndex].MemberOffset);
+					Resource = ResourcePtr ? *ResourcePtr : nullptr;
+				}
+				else
+				{
+					Resource = *(FRHIResource**)((uint8*)Contents + Layout.Resources[ResourceIndex].MemberOffset);
+				}
+
+				checkf(Resource, TEXT("Invalid resource entry creating uniform buffer, %s.Resources[%u], ResourceType 0x%x."),
+					*Layout.GetDebugName().ToString(),
+					ResourceIndex,
+					Layout.Resources[ResourceIndex].MemberType);
+
+				CmdListResources[ResourceIndex] = Resource;
 			}
+		}
 
-			checkf(Resource, TEXT("Invalid resource entry creating uniform buffer, %s.Resources[%u], ResourceType 0x%x."),
-				*Layout.GetDebugName().ToString(),
-				ResourceIndex,
-				Layout.Resources[ResourceIndex].MemberType);
+		if (bBypass)
+		{
+			FRHICommandD3D12UpdateUniformBuffer Cmd(UniformBuffer, UpdatedResourceLocation, CmdListResources, NumResources);
+			Cmd.Execute(RHICmdList);
+		}
+		else
+		{
+			new (RHICmdList.AllocCommand<FRHICommandD3D12UpdateUniformBuffer>()) FRHICommandD3D12UpdateUniformBuffer(UniformBuffer, UpdatedResourceLocation, CmdListResources, NumResources);
 
-			CmdListResources[ResourceIndex] = Resource;
+			//fence is required to stop parallel recording threads from recording with the old bad state of the uniformbuffer resource table.  This command MUST execute before dependent recording starts.
+			RHICmdList.RHIThreadFence(true);
 		}
-	}
-	
-	if (bBypass)
-	{
-		FRHICommandD3D12UpdateUniformBuffer Cmd(UniformBuffer, UpdatedResourceLocation, CmdListResources, NumResources);
-		Cmd.Execute(RHICmdList);
-	}
-	else
-	{
-		new (RHICmdList.AllocCommand<FRHICommandD3D12UpdateUniformBuffer>()) FRHICommandD3D12UpdateUniformBuffer(UniformBuffer, UpdatedResourceLocation, CmdListResources, NumResources);
 
-		//fence is required to stop parallel recording threads from recording with the old bad state of the uniformbuffer resource table.  This command MUST execute before dependent recording starts.
-		RHICmdList.RHIThreadFence(true);
+		UniformBuffer = UniformBuffer->GetNextObject();
 	}
+
 }
 
 FD3D12UniformBuffer::~FD3D12UniformBuffer()
