Implemented basic support for mGPU in D3D12RayTracing. Acceleration structures and other resources are now created on all GPUs. Ray tracing commands are handled by FD3D12CommandContextRedirector.

diff --git a/Engine/Source/Runtime/D3D12RHI/Private/D3D12CommandContext.h b/Engine/Source/Runtime/D3D12RHI/Private/D3D12CommandContext.h
--- a/Engine/Source/Runtime/D3D12RHI/Private/D3D12CommandContext.h
+++ b/Engine/Source/Runtime/D3D12RHI/Private/D3D12CommandContext.h
@@ -741,6 +741,81 @@ public:
 		ContextRedirect(RHIEndRenderPass());
 	}
 
+
+#if RHI_RAYTRACING
+	virtual void RHICopyBufferRegion(FVertexBufferRHIParamRef DestBuffer, uint64 DstOffset, FVertexBufferRHIParamRef SourceBuffer, uint64 SrcOffset, uint64 NumBytes) final override
+	{
+		ContextRedirect(RHICopyBufferRegion(DestBuffer, DstOffset, SourceBuffer, SrcOffset, NumBytes));
+	}
+
+	virtual void RHICopyBufferRegions(const TArrayView<const FCopyBufferRegionParams> Params) final override
+	{
+		ContextRedirect(RHICopyBufferRegions(Params));
+	}
+#endif
+	virtual void RHIBuildAccelerationStructure(FRayTracingGeometryRHIParamRef Geometry) final override
+	{
+		ContextRedirect(RHIBuildAccelerationStructure(Geometry));
+	}
+
+	virtual void RHIUpdateAccelerationStructures(const TArrayView<const FAccelerationStructureUpdateParams> Params) final override
+	{
+		ContextRedirect(RHIUpdateAccelerationStructures(Params));
+	}
+
+	virtual void RHIBuildAccelerationStructures(const TArrayView<const FAccelerationStructureUpdateParams> Params) final override
+	{
+		ContextRedirect(RHIBuildAccelerationStructures(Params));
+	}
+
+	virtual void RHIBuildAccelerationStructure(FRayTracingSceneRHIParamRef Scene) final override
+	{
+		ContextRedirect(RHIBuildAccelerationStructure(Scene));
+	}
+
+	virtual void RHIRayTraceOcclusion(FRayTracingSceneRHIParamRef Scene,
+		FShaderResourceViewRHIParamRef Rays,
+		FUnorderedAccessViewRHIParamRef Output,
+		uint32 NumRays) final override
+	{
+		ContextRedirect(RHIRayTraceOcclusion(Scene, Rays, Output, NumRays));
+	}
+
+	virtual void RHIRayTraceIntersection(FRayTracingSceneRHIParamRef Scene,
+		FShaderResourceViewRHIParamRef Rays,
+		FUnorderedAccessViewRHIParamRef Output,
+		uint32 NumRays) final override
+	{
+		ContextRedirect(RHIRayTraceIntersection(Scene, Rays, Output, NumRays));
+	}
+
+	virtual void RHIRayTraceDispatch(FRayTracingPipelineStateRHIParamRef RayTracingPipelineState, FRayTracingShaderRHIParamRef RayGenShader,
+		FRayTracingSceneRHIParamRef Scene,
+		const FRayTracingShaderBindings& GlobalResourceBindings,
+		uint32 Width, uint32 Height) final override
+	{
+		ContextRedirect(RHIRayTraceDispatch(RayTracingPipelineState, RayGenShader, Scene, GlobalResourceBindings, Width, Height));
+	}
+
+	virtual void RHISetRayTracingHitGroup(
+		FRayTracingSceneRHIParamRef Scene, uint32 InstanceIndex, uint32 SegmentIndex, uint32 ShaderSlot,
+		FRayTracingPipelineStateRHIParamRef Pipeline, uint32 HitGroupIndex,
+		uint32 NumUniformBuffers, const FUniformBufferRHIParamRef* UniformBuffers,
+		uint32 UserData) final override
+	{
+		ContextRedirect(RHISetRayTracingHitGroup(Scene, InstanceIndex, SegmentIndex, ShaderSlot, Pipeline, HitGroupIndex, NumUniformBuffers, UniformBuffers, UserData));
+	}
+
+	virtual void RHISetRayTracingCallableShader(
+		FRayTracingSceneRHIParamRef Scene, uint32 ShaderSlotInScene,
+		FRayTracingPipelineStateRHIParamRef Pipeline, uint32 ShaderIndexInPipeline,
+		uint32 NumUniformBuffers, const FUniformBufferRHIParamRef* UniformBuffers,
+		uint32 UserData) final override
+	{
+		ContextRedirect(RHISetRayTracingCallableShader(Scene, ShaderSlotInScene, Pipeline, ShaderIndexInPipeline, NumUniformBuffers, UniformBuffers, UserData));
+	}
+
+
 	FORCEINLINE void SetPhysicalContext(FD3D12CommandContext* Context)
 	{
 		check(Context);
diff --git a/Engine/Source/Runtime/D3D12RHI/Private/D3D12RayTracing.cpp b/Engine/Source/Runtime/D3D12RHI/Private/D3D12RayTracing.cpp
--- a/Engine/Source/Runtime/D3D12RHI/Private/D3D12RayTracing.cpp
+++ b/Engine/Source/Runtime/D3D12RHI/Private/D3D12RayTracing.cpp
@@ -666,7 +666,7 @@ public:
 		Desc.Flags = D3D12_DESCRIPTOR_HEAP_FLAG_SHADER_VISIBLE;
 		Desc.Type = Type;
 		Desc.NumDescriptors = NumDescriptors;
-		Desc.NodeMask = 1; // #dxr_todo UE-72157: handle mGPU
+		Desc.NodeMask = (uint32)GetParentDevice()->GetGPUMask();
 
 		ID3D12DescriptorHeap* D3D12Heap = nullptr;
 
@@ -855,6 +855,9 @@ public:
 			SamplerHeap.D3D12Heap
 		};
 
+		check(ViewHeap.GetParentDevice() == CommandContext.GetParentDevice());
+		check(SamplerHeap.GetParentDevice() == CommandContext.GetParentDevice());
+
 		CommandContext.CommandListHandle.GraphicsCommandList()->SetDescriptorHeaps(2, Heaps);
 	}
 
@@ -904,7 +907,7 @@ public:
 	TMap<uint64, uint32> SamplerDescriptorTableCache;
 };
 
-class FD3D12RayTracingShaderTable : public FD3D12DeviceChild
+class FD3D12RayTracingShaderTable
 {
 private:
 
@@ -946,8 +949,7 @@ public:
 		uint32 MaxViewDescriptorsPerRecord = 0;
 	};
 
-	FD3D12RayTracingShaderTable(FD3D12Device* Device)
-		: FD3D12DeviceChild(Device)
+	FD3D12RayTracingShaderTable()
 	{
 	}
 
@@ -956,7 +958,7 @@ public:
 		delete DescriptorCache;
 	}
 
-	void Init(const FInitializer& Initializer)
+	void Init(const FInitializer& Initializer, FD3D12Device* Device)
 	{
 		checkf(Initializer.LocalRootDataSize <= 4096, TEXT("The maximum size of a local root signature is 4KB.")); // as per section 4.22.1 of DXR spec v1.0
 		checkf(Initializer.NumRayGenShaders >= 1, TEXT("All shader tables must contain at least one raygen shader."));
@@ -979,7 +981,7 @@ public:
 			const uint32 NumViewDescriptors = FMath::Max(MinNumViewDescriptors, FMath::Min<uint32>(Initializer.NumHitRecords * Initializer.MaxViewDescriptorsPerRecord, D3D12_MAX_SHADER_VISIBLE_DESCRIPTOR_HEAP_SIZE_TIER_1));
 			const uint32 NumSamplerDescriptors = D3D12_MAX_SHADER_VISIBLE_SAMPLER_HEAP_SIZE;
 
-			DescriptorCache = new FD3D12RayTracingDescriptorCache(GetParentDevice());
+			DescriptorCache = new FD3D12RayTracingDescriptorCache(Device);
 			DescriptorCache->Init(NumViewDescriptors, NumSamplerDescriptors);
 		}
 
@@ -1092,12 +1094,10 @@ public:
 		}
 	}
 
-	void CopyToGPU()
+	void CopyToGPU(FD3D12Device* Device)
 	{
 		check(IsInRHIThread() || !IsRunningRHIInSeparateThread());
 
-		FD3D12Device* Device = GetParentDevice();
-
 		checkf(Data.Num(), TEXT("Shader table is expected to be initialized before copying to GPU."));
 
 		FD3D12Adapter* Adapter = Device->GetParentAdapter();
@@ -1107,7 +1107,6 @@ public:
 		FRHIResourceCreateInfo CreateInfo;
 		CreateInfo.ResourceArray = &Data;
 
-		checkf(GNumExplicitGPUsForRendering == 1, TEXT("Ray tracing is not implemented for mGPU")); // #dxr_todo UE-72157: implement mGPU support
 		Buffer = Adapter->CreateRHIBuffer<FD3D12MemBuffer>(
 			nullptr, BufferDesc, BufferDesc.Alignment,
 			0, BufferDesc.Width, BUF_Static, CreateInfo,
@@ -1261,7 +1260,6 @@ public:
 	UE_NONCOPYABLE(FD3D12RayTracingPipelineState)
 
 	FD3D12RayTracingPipelineState(FD3D12Device* Device, const FRayTracingPipelineStateInitializer& Initializer)
-		: DefaultShaderTable(Device)
 	{
 		checkf(Initializer.GetRayGenTable().Num() > 0, TEXT("Ray tracing pipelines must have at leat one ray generation shader."));
 
@@ -1272,7 +1270,6 @@ public:
 
 		TotalCreationTime -= FPlatformTime::Cycles64();
 
-		FD3D12Adapter* Adapter = Device->GetParentAdapter();
 		ID3D12Device5* RayTracingDevice = Device->GetRayTracingDevice();
 
 		// Use hit and miss shaders from initializer or fall back to default ones if none were provided
@@ -1505,10 +1502,17 @@ public:
 		SBTInitializer.NumHitRecords = 0; // Default SBT does not support indexable hit shaders
 		SBTInitializer.LocalRootDataSize = 0; // Shaders in default SBT are not allowed to access any local resources
 
-		DefaultShaderTable.Init(SBTInitializer);
-		DefaultShaderTable.SetRayGenIdentifiers(RayGenShaders.Identifiers);
-		DefaultShaderTable.SetMissIdentifiers(MissShaders.Identifiers);
-		DefaultShaderTable.SetDefaultHitGroupIdentifier(HitGroupShaders.Identifiers[0]);
+		// Create default shader tables for every node in the LDA group
+		FD3D12Adapter* Adapter = Device->GetParentAdapter();
+
+		for (uint32 GPUIndex : FRHIGPUMask::All())
+		{
+			FD3D12Device* NodeDevice = Adapter->GetDevice(GPUIndex);
+			DefaultShaderTables[GPUIndex].Init(SBTInitializer, NodeDevice);
+			DefaultShaderTables[GPUIndex].SetRayGenIdentifiers(RayGenShaders.Identifiers);
+			DefaultShaderTables[GPUIndex].SetMissIdentifiers(MissShaders.Identifiers);
+			DefaultShaderTables[GPUIndex].SetDefaultHitGroupIdentifier(HitGroupShaders.Identifiers[0]);
+		}
 
 		TotalCreationTime += FPlatformTime::Cycles64();
 
@@ -1536,7 +1540,7 @@ public:
 
 	// Shader table that can be used to dispatch ray tracing work that doesn't require real SBT bindings.
 	// This is useful for the case where user only provides default RayGen, Miss and HitGroup shaders.
-	FD3D12RayTracingShaderTable DefaultShaderTable;
+	FD3D12RayTracingShaderTable DefaultShaderTables[MAX_NUM_GPUS];
 
 	ID3D12RootSignature* GlobalRootSignature = nullptr;
 
@@ -1627,7 +1631,8 @@ void FD3D12Device::CleanupRayTracing()
 
 FRayTracingPipelineStateRHIRef FD3D12DynamicRHI::RHICreateRayTracingPipelineState(const FRayTracingPipelineStateInitializer& Initializer)
 {
-	FD3D12RayTracingPipelineState* Result = new FD3D12RayTracingPipelineState(GetRHIDevice(), Initializer);
+	FD3D12Device* Device = GetAdapter().GetDevice(0); // All pipelines are created on the first node, as they may be used on any other linked GPU.
+	FD3D12RayTracingPipelineState* Result = new FD3D12RayTracingPipelineState(Device, Initializer);
 
 	return Result;
 }
@@ -1658,139 +1663,132 @@ FRayTracingGeometryRHIRef FD3D12DynamicRHI::RHICreateRayTracingGeometry(const FR
 		checkf(Initializer.IndexBuffer->GetStride() == 2 || Initializer.IndexBuffer->GetStride() == 4, TEXT("Index buffer must be 16 or 32 bit."));
 	}
 
-	checkf(GNumExplicitGPUsForRendering == 1, TEXT("Ray tracing is not implemented for mGPU")); // #dxr_todo UE-72157: implement mGPU support
-	FD3D12RayTracingGeometry* Result = GetAdapter().CreateLinkedObject<FD3D12RayTracingGeometry>(FRHIGPUMask::All(), [&](FD3D12Device* Device)
-	{
-		FD3D12RayTracingGeometry* Mesh = new FD3D12RayTracingGeometry(Device);
-
-		const uint32 GPUIndex = Device->GetGPUIndex();
+	FD3D12RayTracingGeometry* Geometry = new FD3D12RayTracingGeometry();
 
-		Mesh->IndexStride = Initializer.IndexBuffer ? Initializer.IndexBuffer->GetStride() : 0; // stride 0 means implicit triangle list for non-indexed geometry
-		Mesh->VertexOffsetInBytes = (Initializer.BaseVertexIndex * Initializer.VertexBufferStride) + Initializer.VertexBufferByteOffset;
-		Mesh->VertexStrideInBytes = Initializer.VertexBufferStride;
-		Mesh->BaseVertexIndex = Initializer.BaseVertexIndex;
-		Mesh->TotalPrimitiveCount = Initializer.TotalPrimitiveCount;
+	Geometry->IndexStride = Initializer.IndexBuffer ? Initializer.IndexBuffer->GetStride() : 0; // stride 0 means implicit triangle list for non-indexed geometry
+	Geometry->VertexOffsetInBytes = (Initializer.BaseVertexIndex * Initializer.VertexBufferStride) + Initializer.VertexBufferByteOffset;
+	Geometry->VertexStrideInBytes = Initializer.VertexBufferStride;
+	Geometry->BaseVertexIndex = Initializer.BaseVertexIndex;
+	Geometry->TotalPrimitiveCount = Initializer.TotalPrimitiveCount;
 
-		switch (Initializer.GeometryType)
-		{
-		case RTGT_Triangles:
-			Mesh->GeometryType = D3D12_RAYTRACING_GEOMETRY_TYPE_TRIANGLES;
-			break;
-		case RTGT_Procedural:
-			Mesh->GeometryType = D3D12_RAYTRACING_GEOMETRY_TYPE_PROCEDURAL_PRIMITIVE_AABBS;
-			break;
-		default:
-			checkf(false, TEXT("Unexpected ray tracing geometry type"));
-		}
+	switch (Initializer.GeometryType)
+	{
+	case RTGT_Triangles:
+		Geometry->GeometryType = D3D12_RAYTRACING_GEOMETRY_TYPE_TRIANGLES;
+		break;
+	case RTGT_Procedural:
+		Geometry->GeometryType = D3D12_RAYTRACING_GEOMETRY_TYPE_PROCEDURAL_PRIMITIVE_AABBS;
+		break;
+	default:
+		checkf(false, TEXT("Unexpected ray tracing geometry type"));
+	}
 
-		if (Initializer.bFastBuild)
-		{
-			Mesh->BuildFlags = D3D12_RAYTRACING_ACCELERATION_STRUCTURE_BUILD_FLAG_PREFER_FAST_BUILD;
-		}
-		else
-		{
-			Mesh->BuildFlags = D3D12_RAYTRACING_ACCELERATION_STRUCTURE_BUILD_FLAG_PREFER_FAST_TRACE;
-		}
+	if (Initializer.bFastBuild)
+	{
+		Geometry->BuildFlags = D3D12_RAYTRACING_ACCELERATION_STRUCTURE_BUILD_FLAG_PREFER_FAST_BUILD;
+	}
+	else
+	{
+		Geometry->BuildFlags = D3D12_RAYTRACING_ACCELERATION_STRUCTURE_BUILD_FLAG_PREFER_FAST_TRACE;
+	}
 
-		if (Initializer.bAllowUpdate)
-		{
-			Mesh->BuildFlags |= D3D12_RAYTRACING_ACCELERATION_STRUCTURE_BUILD_FLAG_ALLOW_UPDATE;
-		}
+	if (Initializer.bAllowUpdate)
+	{
+		Geometry->BuildFlags |= D3D12_RAYTRACING_ACCELERATION_STRUCTURE_BUILD_FLAG_ALLOW_UPDATE;
+	}
 
-		if (GRayTracingDebugForceBuildMode == 1)
-		{
-			Mesh->BuildFlags |= D3D12_RAYTRACING_ACCELERATION_STRUCTURE_BUILD_FLAG_PREFER_FAST_BUILD;
-			Mesh->BuildFlags &= ~D3D12_RAYTRACING_ACCELERATION_STRUCTURE_BUILD_FLAG_PREFER_FAST_TRACE;
-		}
-		else if (GRayTracingDebugForceBuildMode == 2)
-		{
-			Mesh->BuildFlags |= D3D12_RAYTRACING_ACCELERATION_STRUCTURE_BUILD_FLAG_PREFER_FAST_TRACE;
-			Mesh->BuildFlags &= ~D3D12_RAYTRACING_ACCELERATION_STRUCTURE_BUILD_FLAG_PREFER_FAST_BUILD;
-		}
+	if (GRayTracingDebugForceBuildMode == 1)
+	{
+		Geometry->BuildFlags |= D3D12_RAYTRACING_ACCELERATION_STRUCTURE_BUILD_FLAG_PREFER_FAST_BUILD;
+		Geometry->BuildFlags &= ~D3D12_RAYTRACING_ACCELERATION_STRUCTURE_BUILD_FLAG_PREFER_FAST_TRACE;
+	}
+	else if (GRayTracingDebugForceBuildMode == 2)
+	{
+		Geometry->BuildFlags |= D3D12_RAYTRACING_ACCELERATION_STRUCTURE_BUILD_FLAG_PREFER_FAST_TRACE;
+		Geometry->BuildFlags &= ~D3D12_RAYTRACING_ACCELERATION_STRUCTURE_BUILD_FLAG_PREFER_FAST_BUILD;
+	}
 
-		if (Initializer.Segments.Num())
-		{
-			Mesh->Segments = TArray<FRayTracingGeometrySegment>(Initializer.Segments.GetData(), Initializer.Segments.Num());
-		}
-		else
-		{
-			FRayTracingGeometrySegment DefaultSegment;
-			DefaultSegment.FirstPrimitive = 0;
-			DefaultSegment.NumPrimitives = Initializer.TotalPrimitiveCount;
-			Mesh->Segments.Add(DefaultSegment);
-		}
+	if (Initializer.Segments.Num())
+	{
+		Geometry->Segments = TArray<FRayTracingGeometrySegment>(Initializer.Segments.GetData(), Initializer.Segments.Num());
+	}
+	else
+	{
+		FRayTracingGeometrySegment DefaultSegment;
+		DefaultSegment.FirstPrimitive = 0;
+		DefaultSegment.NumPrimitives = Initializer.TotalPrimitiveCount;
+		Geometry->Segments.Add(DefaultSegment);
+	}
 
 #if DO_CHECK
+	{
+		uint32 ComputedPrimitiveCountForValidation = 0;
+		for (const FRayTracingGeometrySegment& Segment : Geometry->Segments)
 		{
-			uint32 ComputedPrimitiveCountForValidation = 0;
-			for (const FRayTracingGeometrySegment& Segment : Mesh->Segments)
-			{
-				ComputedPrimitiveCountForValidation += Segment.NumPrimitives;
-				check(Segment.FirstPrimitive + Segment.NumPrimitives <= Initializer.TotalPrimitiveCount);
-			}
-			check(ComputedPrimitiveCountForValidation == Initializer.TotalPrimitiveCount);
+			ComputedPrimitiveCountForValidation += Segment.NumPrimitives;
+			check(Segment.FirstPrimitive + Segment.NumPrimitives <= Initializer.TotalPrimitiveCount);
 		}
+		check(ComputedPrimitiveCountForValidation == Initializer.TotalPrimitiveCount);
+	}
 #endif
 
-		Mesh->VertexElemType = Initializer.VertexBufferElementType;
+	Geometry->VertexElemType = Initializer.VertexBufferElementType;
 
-		Mesh->IndexBuffer = Initializer.IndexBuffer ? ResourceCast(Initializer.IndexBuffer.GetReference(), GPUIndex) : nullptr;
-		Mesh->PositionVertexBuffer = ResourceCast(Initializer.PositionVertexBuffer.GetReference(), GPUIndex);
-
-		Mesh->bIsAccelerationStructureDirty = true;
-
-		return Mesh;
-	});
+	Geometry->RHIIndexBuffer = Initializer.IndexBuffer;
+	Geometry->RHIVertexBuffer = Initializer.PositionVertexBuffer;
+	Geometry->SetDirty(FRHIGPUMask::All(), true);
 
 	INC_DWORD_STAT_BY(STAT_D3D12RayTracingTrianglesBLAS, Initializer.TotalPrimitiveCount);
 
-	return Result;
+	return Geometry;
 }
 
 FRayTracingSceneRHIRef FD3D12DynamicRHI::RHICreateRayTracingScene(const FRayTracingSceneInitializer& Initializer)
 {
 	FD3D12Adapter& Adapter = GetAdapter();
 
-	return GetAdapter().CreateLinkedObject<FD3D12RayTracingScene>(FRHIGPUMask::All(), [&](FD3D12Device* Device)
-	{
-		FD3D12RayTracingScene* Result = new FD3D12RayTracingScene(Device);
+	FD3D12RayTracingScene* Result = new FD3D12RayTracingScene(Adapter);
 
-		Result->Instances = TArray<FRayTracingGeometryInstance>(Initializer.Instances.GetData(), Initializer.Instances.Num());
-		Result->ShaderSlotsPerGeometrySegment = Initializer.ShaderSlotsPerGeometrySegment;
+	Result->Instances = TArray<FRayTracingGeometryInstance>(Initializer.Instances.GetData(), Initializer.Instances.Num());
+	Result->ShaderSlotsPerGeometrySegment = Initializer.ShaderSlotsPerGeometrySegment;
 
-		// Compute geometry segment count prefix sum to be later used in GetHitRecordBaseIndex()
-		Result->SegmentPrefixSum.Reserve(Result->Instances.Num());
-		uint32 NumTotalSegments = 0;
-		for (const FRayTracingGeometryInstance& Instance : Result->Instances)
-		{
-			FD3D12RayTracingGeometry* Geometry = FD3D12DynamicRHI::ResourceCast(Instance.GeometryRHI);
-			Result->SegmentPrefixSum.Add(NumTotalSegments);
-			NumTotalSegments += Geometry->Segments.Num();
-		}
-		Result->NumTotalSegments = NumTotalSegments;
+	// Compute geometry segment count prefix sum to be later used in GetHitRecordBaseIndex()
+	Result->SegmentPrefixSum.Reserve(Result->Instances.Num());
+	uint32 NumTotalSegments = 0;
+	for (const FRayTracingGeometryInstance& Instance : Result->Instances)
+	{
+		FD3D12RayTracingGeometry* Geometry = FD3D12DynamicRHI::ResourceCast(Instance.GeometryRHI);
+		Result->SegmentPrefixSum.Add(NumTotalSegments);
+		NumTotalSegments += Geometry->Segments.Num();
+	}
+	Result->NumTotalSegments = NumTotalSegments;
 
-		Result->NumCallableShaderSlots = Initializer.NumCallableShaderSlots;
+	Result->NumCallableShaderSlots = Initializer.NumCallableShaderSlots;
 
-		return Result;
-	});
+	return Result;
 }
 
-FD3D12RayTracingGeometry::FD3D12RayTracingGeometry(FD3D12Device* Device) 
-	: FD3D12DeviceChild(Device) 
+FD3D12RayTracingGeometry::FD3D12RayTracingGeometry()
 {
 	INC_DWORD_STAT(STAT_D3D12RayTracingAllocatedBLAS);
 }
 
 FD3D12RayTracingGeometry::~FD3D12RayTracingGeometry()
 {
-	if (AccelerationStructureBuffer)
+	for (TRefCountPtr<FD3D12MemBuffer>& Buffer : AccelerationStructureBuffers)
 	{
-		DEC_MEMORY_STAT_BY(STAT_D3D12RayTracingUsedVideoMemory, AccelerationStructureBuffer->GetSize());
+		if (Buffer)
+		{
+			DEC_MEMORY_STAT_BY(STAT_D3D12RayTracingUsedVideoMemory, Buffer->GetSize());
+		}
 	}
 
-	if (ScratchBuffer)
+	for (TRefCountPtr<FD3D12MemBuffer>& Buffer : ScratchBuffers)
 	{
-		DEC_MEMORY_STAT_BY(STAT_D3D12RayTracingUsedVideoMemory, ScratchBuffer->GetSize());
+		if (Buffer)
+		{
+			DEC_MEMORY_STAT_BY(STAT_D3D12RayTracingUsedVideoMemory, Buffer->GetSize());
+		}
 	}
 
 	DEC_DWORD_STAT_BY(STAT_D3D12RayTracingTrianglesBLAS, TotalPrimitiveCount);
@@ -1800,13 +1798,19 @@ FD3D12RayTracingGeometry::~FD3D12RayTracingGeometry()
 void FD3D12RayTracingGeometry::TransitionBuffers(FD3D12CommandContext& CommandContext)
 {
 	// Transition vertex and index resources..
-	if (IndexBuffer && IndexBuffer.GetReference()->GetResource()->RequiresResourceStateTracking())
+	if (RHIIndexBuffer)
 	{
-		FD3D12DynamicRHI::TransitionResource(CommandContext.CommandListHandle, IndexBuffer.GetReference()->GetResource(), D3D12_RESOURCE_STATE_NON_PIXEL_SHADER_RESOURCE, 0);
+		FD3D12IndexBuffer* IndexBuffer = CommandContext.RetrieveObject<FD3D12IndexBuffer>(RHIIndexBuffer.GetReference());
+		if (IndexBuffer->GetResource()->RequiresResourceStateTracking())
+		{
+			FD3D12DynamicRHI::TransitionResource(CommandContext.CommandListHandle, IndexBuffer->GetResource(), D3D12_RESOURCE_STATE_NON_PIXEL_SHADER_RESOURCE, 0);
+		}
 	}
-	if (PositionVertexBuffer.GetReference()->GetResource()->RequiresResourceStateTracking())
+
+	FD3D12VertexBuffer* VertexBuffer = CommandContext.RetrieveObject<FD3D12VertexBuffer>(RHIVertexBuffer.GetReference());
+	if (VertexBuffer->GetResource()->RequiresResourceStateTracking())
 	{
-		FD3D12DynamicRHI::TransitionResource(CommandContext.CommandListHandle, PositionVertexBuffer.GetReference()->GetResource(), D3D12_RESOURCE_STATE_NON_PIXEL_SHADER_RESOURCE, 0);
+		FD3D12DynamicRHI::TransitionResource(CommandContext.CommandListHandle, VertexBuffer->GetResource(), D3D12_RESOURCE_STATE_NON_PIXEL_SHADER_RESOURCE, 0);
 	}
 }
 
@@ -1840,8 +1844,6 @@ static void CreateAccelerationStructureBuffers(TRefCountPtr<FD3D12MemBuffer>& Ac
 
 void FD3D12RayTracingGeometry::BuildAccelerationStructure(FD3D12CommandContext& CommandContext, bool bIsUpdate)
 {
-	check(GNumExplicitGPUsForRendering == 1); // #dxr_todo UE-72157: ensure that mGPU case is handled correctly!
-
 	static constexpr uint32 IndicesPerPrimitive = 3; // Only triangle meshes are supported
 
 	// Array of geometry descriptions, one per segment (single-segment geometry is a common case).
@@ -1849,6 +1851,9 @@ void FD3D12RayTracingGeometry::BuildAccelerationStructure(FD3D12CommandContext&
 
 	Descs.Reserve(Segments.Num());
 
+	FD3D12IndexBuffer* IndexBuffer = CommandContext.RetrieveObject<FD3D12IndexBuffer>(RHIIndexBuffer.GetReference());
+	FD3D12VertexBuffer* VertexBuffer = CommandContext.RetrieveObject<FD3D12VertexBuffer>(RHIVertexBuffer.GetReference());
+
 	for (const FRayTracingGeometrySegment& Segment : Segments)
 	{
 		D3D12_RAYTRACING_GEOMETRY_DESC Desc = {};
@@ -1894,7 +1899,7 @@ void FD3D12RayTracingGeometry::BuildAccelerationStructure(FD3D12CommandContext&
 				Desc.Triangles.IndexCount = Segment.NumPrimitives * IndicesPerPrimitive;
 				Desc.Triangles.IndexBuffer = IndexBuffer->ResourceLocation.GetGPUVirtualAddress() + IndexStride * Segment.FirstPrimitive * IndicesPerPrimitive;
 
-				Desc.Triangles.VertexCount = PositionVertexBuffer->ResourceLocation.GetSize() / VertexStrideInBytes;
+				Desc.Triangles.VertexCount = VertexBuffer->ResourceLocation.GetSize() / VertexStrideInBytes;
 
 				IndexBuffer->GetResource()->UpdateResidency(CommandContext.CommandListHandle);
 			}
@@ -1907,16 +1912,16 @@ void FD3D12RayTracingGeometry::BuildAccelerationStructure(FD3D12CommandContext&
 
 				checkf(Segments.Num() == 1, TEXT("Non-indexed geometry with multiple segments is not implemented."));
 
-				Desc.Triangles.VertexCount = FMath::Min<uint32>(PositionVertexBuffer->ResourceLocation.GetSize() / VertexStrideInBytes, TotalPrimitiveCount * 3);
+				Desc.Triangles.VertexCount = FMath::Min<uint32>(VertexBuffer->ResourceLocation.GetSize() / VertexStrideInBytes, TotalPrimitiveCount * 3);
 			}
 
-			Desc.Triangles.VertexBuffer.StartAddress = PositionVertexBuffer->ResourceLocation.GetGPUVirtualAddress() + VertexOffsetInBytes;
+			Desc.Triangles.VertexBuffer.StartAddress = VertexBuffer->ResourceLocation.GetGPUVirtualAddress() + VertexOffsetInBytes;
 			Desc.Triangles.VertexBuffer.StrideInBytes = VertexStrideInBytes;
 		}
 		else if (GeometryType == D3D12_RAYTRACING_GEOMETRY_TYPE_PROCEDURAL_PRIMITIVE_AABBS)
 		{
 			Desc.AABBs.AABBCount = Segment.NumPrimitives;
-			Desc.AABBs.AABBs.StartAddress = PositionVertexBuffer->ResourceLocation.GetGPUVirtualAddress() + VertexOffsetInBytes;
+			Desc.AABBs.AABBs.StartAddress = VertexBuffer->ResourceLocation.GetGPUVirtualAddress() + VertexOffsetInBytes;
 			Desc.AABBs.AABBs.StrideInBytes = VertexStrideInBytes;
 		}
 		else
@@ -1924,14 +1929,12 @@ void FD3D12RayTracingGeometry::BuildAccelerationStructure(FD3D12CommandContext&
 			checkf(false, TEXT("Unexpected ray tracing geometry type"));
 		}
 
-		PositionVertexBuffer->ResourceLocation.GetResource()->UpdateResidency(CommandContext.CommandListHandle);
+		VertexBuffer->ResourceLocation.GetResource()->UpdateResidency(CommandContext.CommandListHandle);
 
 		Descs.Add(Desc);
 	}
 
-	checkf(GNumExplicitGPUsForRendering == 1, TEXT("Ray tracing is not implemented for mGPU")); // #dxr_todo UE-72157: implement mGPU support
-
-	const uint32 GPUIndex = CommandContext.GetGPUIndex(); // #dxr_todo UE-72157: ensure that mGPU case is handled correctly!
+	const uint32 GPUIndex = CommandContext.GetGPUIndex();
 	FD3D12Adapter* Adapter = CommandContext.GetParentAdapter();
 
 	ID3D12Device5* RayTracingDevice = CommandContext.GetParentDevice()->GetRayTracingDevice();
@@ -1954,7 +1957,7 @@ void FD3D12RayTracingGeometry::BuildAccelerationStructure(FD3D12CommandContext&
 	PrebuildDescInputs.pGeometryDescs = Descs.GetData();
 	PrebuildDescInputs.Flags = LocalBuildFlags;
 
-	if (!AccelerationStructureBuffer)
+	if (!AccelerationStructureBuffers[GPUIndex])
 	{
 		check(!bIsUpdate);
 
@@ -1962,34 +1965,34 @@ void FD3D12RayTracingGeometry::BuildAccelerationStructure(FD3D12CommandContext&
 
 		RayTracingDevice->GetRaytracingAccelerationStructurePrebuildInfo(&PrebuildDescInputs, &PrebuildInfo);
 
-		CreateAccelerationStructureBuffers(AccelerationStructureBuffer, ScratchBuffer, Adapter, GPUIndex, PrebuildInfo);
+		CreateAccelerationStructureBuffers(AccelerationStructureBuffers[GPUIndex], ScratchBuffers[GPUIndex], Adapter, GPUIndex, PrebuildInfo);
 
-		INC_MEMORY_STAT_BY(STAT_D3D12RayTracingUsedVideoMemory, AccelerationStructureBuffer->GetSize());
-		INC_MEMORY_STAT_BY(STAT_D3D12RayTracingUsedVideoMemory, ScratchBuffer->GetSize());
+		INC_MEMORY_STAT_BY(STAT_D3D12RayTracingUsedVideoMemory, AccelerationStructureBuffers[GPUIndex]->GetSize());
+		INC_MEMORY_STAT_BY(STAT_D3D12RayTracingUsedVideoMemory, ScratchBuffers[GPUIndex]->GetSize());
 
 		// #dxr_todo UE-72161: scratch buffers should be created in UAV state from the start
-		FD3D12DynamicRHI::TransitionResource(CommandContext.CommandListHandle, ScratchBuffer.GetReference()->GetResource(), D3D12_RESOURCE_STATE_UNORDERED_ACCESS, 0);
+		FD3D12DynamicRHI::TransitionResource(CommandContext.CommandListHandle, ScratchBuffers[GPUIndex].GetReference()->GetResource(), D3D12_RESOURCE_STATE_UNORDERED_ACCESS, 0);
 	}
 
 	TransitionBuffers(CommandContext);
 	CommandContext.CommandListHandle.FlushResourceBarriers();
 
-	if (bIsAccelerationStructureDirty)
+	if (IsDirty(GPUIndex))
 	{
-		ScratchBuffer->GetResource()->UpdateResidency(CommandContext.CommandListHandle);
-		AccelerationStructureBuffer->GetResource()->UpdateResidency(CommandContext.CommandListHandle);
+		ScratchBuffers[GPUIndex]->GetResource()->UpdateResidency(CommandContext.CommandListHandle);
+		AccelerationStructureBuffers[GPUIndex]->GetResource()->UpdateResidency(CommandContext.CommandListHandle);
 
 		D3D12_BUILD_RAYTRACING_ACCELERATION_STRUCTURE_DESC BuildDesc = {};
 		BuildDesc.Inputs = PrebuildDescInputs;
-		BuildDesc.DestAccelerationStructureData = AccelerationStructureBuffer->ResourceLocation.GetGPUVirtualAddress();
-		BuildDesc.ScratchAccelerationStructureData = ScratchBuffer->ResourceLocation.GetGPUVirtualAddress();
+		BuildDesc.DestAccelerationStructureData = AccelerationStructureBuffers[GPUIndex]->ResourceLocation.GetGPUVirtualAddress();
+		BuildDesc.ScratchAccelerationStructureData = ScratchBuffers[GPUIndex]->ResourceLocation.GetGPUVirtualAddress();
 		BuildDesc.SourceAccelerationStructureData = bIsUpdate
-			? AccelerationStructureBuffer->ResourceLocation.GetGPUVirtualAddress()
+			? AccelerationStructureBuffers[GPUIndex]->ResourceLocation.GetGPUVirtualAddress()
 			: D3D12_GPU_VIRTUAL_ADDRESS(0);
 
 		ID3D12GraphicsCommandList4* RayTracingCommandList = CommandContext.CommandListHandle.RayTracingCommandList();
 		RayTracingCommandList->BuildRaytracingAccelerationStructure(&BuildDesc, 0, nullptr);
-		bIsAccelerationStructureDirty = false;
+		SetDirty(CommandContext.GetGPUMask(), false);
 
 		if (bIsUpdate)
 		{
@@ -2004,43 +2007,49 @@ void FD3D12RayTracingGeometry::BuildAccelerationStructure(FD3D12CommandContext&
 	// We don't need to keep a scratch buffer after initial build if acceleration structure is static.
 	if (!(BuildFlags & D3D12_RAYTRACING_ACCELERATION_STRUCTURE_BUILD_FLAG_ALLOW_UPDATE))
 	{
-		DEC_MEMORY_STAT_BY(STAT_D3D12RayTracingUsedVideoMemory, ScratchBuffer->GetSize());
-		ScratchBuffer = nullptr;
+		DEC_MEMORY_STAT_BY(STAT_D3D12RayTracingUsedVideoMemory, ScratchBuffers[GPUIndex]->GetSize());
+		ScratchBuffers[GPUIndex] = nullptr;
 	}
 }
 
-FD3D12RayTracingScene::FD3D12RayTracingScene(FD3D12Device* Device)
-	: FD3D12DeviceChild(Device)
-	, AccelerationStructureView(new FD3D12ShaderResourceView(Device))
+FD3D12RayTracingScene::FD3D12RayTracingScene(FD3D12Adapter& Adapter)
 {
-	INC_DWORD_STAT(STAT_D3D12RayTracingAllocatedTLAS);
+	ShaderResourceView = Adapter.CreateLinkedObject<FD3D12ShaderResourceView>(FRHIGPUMask::All(), [&](FD3D12Device* Device)
+	{
+		return new FD3D12ShaderResourceView(Device);
+	});
 
-	ShaderResourceView = AccelerationStructureView;
+	INC_DWORD_STAT(STAT_D3D12RayTracingAllocatedTLAS);
 };
 
 FD3D12RayTracingScene::~FD3D12RayTracingScene()
 {
-	for (auto Item : ShaderTables)
+	for (auto& Table : ShaderTables)
 	{
-		delete Item.Value;
+		for (auto Item : Table)
+		{
+			delete Item.Value;
+		}
 	}
 
-	if (AccelerationStructureBuffer)
+	for (auto& AccelerationStructureBuffer : AccelerationStructureBuffers)
 	{
-		DEC_MEMORY_STAT_BY(STAT_D3D12RayTracingUsedVideoMemory, AccelerationStructureBuffer->GetSize());
+		if (AccelerationStructureBuffer)
+		{
+			DEC_MEMORY_STAT_BY(STAT_D3D12RayTracingUsedVideoMemory, AccelerationStructureBuffer->GetSize());
+		}
 	}
+	
 
 	DEC_DWORD_STAT(STAT_D3D12RayTracingAllocatedTLAS);
 }
 
 void FD3D12RayTracingScene::BuildAccelerationStructure(FD3D12CommandContext& CommandContext, D3D12_RAYTRACING_ACCELERATION_STRUCTURE_BUILD_FLAGS BuildFlags)
 {
-	checkf(GNumExplicitGPUsForRendering == 1, TEXT("Ray tracing is not implemented for mGPU")); // #dxr_todo UE-72157: implement mGPU support
-
 	TRefCountPtr<FD3D12MemBuffer> InstanceBuffer;
 	TRefCountPtr<FD3D12MemBuffer> ScratchBuffer;
 
-	const uint32 GPUIndex = CommandContext.GetGPUIndex(); // #dxr_todo UE-72157: ensure that mGPU case is handled correctly!
+	const uint32 GPUIndex = CommandContext.GetGPUIndex();
 	FD3D12Adapter* Adapter = CommandContext.GetParentAdapter();
 	ID3D12Device5* RayTracingDevice = CommandContext.GetParentDevice()->GetRayTracingDevice();
 
@@ -2054,6 +2063,8 @@ void FD3D12RayTracingScene::BuildAccelerationStructure(FD3D12CommandContext& Com
 
 	RayTracingDevice->GetRaytracingAccelerationStructurePrebuildInfo(&PrebuildDescInputs, &PrebuildInfo);
 
+	TRefCountPtr<FD3D12MemBuffer>& AccelerationStructureBuffer = AccelerationStructureBuffers[GPUIndex];
+
 	CreateAccelerationStructureBuffers(AccelerationStructureBuffer, ScratchBuffer, Adapter, GPUIndex, PrebuildInfo);
 
 	INC_MEMORY_STAT_BY(STAT_D3D12RayTracingUsedVideoMemory, AccelerationStructureBuffer->GetSize());
@@ -2061,8 +2072,11 @@ void FD3D12RayTracingScene::BuildAccelerationStructure(FD3D12CommandContext& Com
 	// #dxr_todo UE-72161: scratch buffers should be created in UAV state from the start
 	FD3D12DynamicRHI::TransitionResource(CommandContext.CommandListHandle, ScratchBuffer.GetReference()->GetResource(), D3D12_RESOURCE_STATE_UNORDERED_ACCESS, 0);
 
-	if (bAccelerationStructureViewInitialized)
+	FD3D12ShaderResourceView* AccelerationStructureView = CommandContext.RetrieveObject<FD3D12ShaderResourceView>(ShaderResourceView.GetReference());
+
+	if (bAccelerationStructureViewInitialized[GPUIndex])
 	{
+		check(AccelerationStructureView->GetParentDevice() == AccelerationStructureBuffer->GetParentDevice());
 		AccelerationStructureView->Rename(AccelerationStructureBuffer->ResourceLocation);
 	}
 	else
@@ -2073,9 +2087,10 @@ void FD3D12RayTracingScene::BuildAccelerationStructure(FD3D12CommandContext& Com
 		SRVDesc.Shader4ComponentMapping = D3D12_DEFAULT_SHADER_4_COMPONENT_MAPPING;
 		SRVDesc.RaytracingAccelerationStructure.Location = AccelerationStructureBuffer->ResourceLocation.GetGPUVirtualAddress();
 
+		check(AccelerationStructureView->GetParentDevice() == AccelerationStructureBuffer->GetParentDevice());
 		AccelerationStructureView->Initialize(SRVDesc, AccelerationStructureBuffer->ResourceLocation, 4);
 
-		bAccelerationStructureViewInitialized = true;
+		bAccelerationStructureViewInitialized[GPUIndex] = true;
 	}
 
 	// Create and fill instance buffer
@@ -2110,7 +2125,7 @@ void FD3D12RayTracingScene::BuildAccelerationStructure(FD3D12CommandContext& Com
 		{
 			FD3D12RayTracingGeometry* Geometry = FD3D12DynamicRHI::ResourceCast(Instance.GeometryRHI);
 
-			checkf(!Geometry->bIsAccelerationStructureDirty, 
+			checkf(!Geometry->IsDirty(CommandContext.GetGPUIndex()),
 				TEXT("Acceleration structures for all geometries must be built before building the top level acceleration structure for the scene."));
 
 			D3D12_RAYTRACING_INSTANCE_DESC InstanceDesc = {};
@@ -2140,14 +2155,14 @@ void FD3D12RayTracingScene::BuildAccelerationStructure(FD3D12CommandContext& Com
 				InstanceDesc.Flags |= D3D12_RAYTRACING_INSTANCE_FLAG_TRIANGLE_CULL_DISABLE;
 			}
 
-			InstanceDesc.AccelerationStructure = Geometry->AccelerationStructureBuffer->ResourceLocation.GetGPUVirtualAddress();
-			Geometry->AccelerationStructureBuffer->ResourceLocation.GetResource()->UpdateResidency(CommandContext.CommandListHandle);
+			InstanceDesc.AccelerationStructure = Geometry->AccelerationStructureBuffers[GPUIndex]->ResourceLocation.GetGPUVirtualAddress();
+			Geometry->AccelerationStructureBuffers[GPUIndex]->ResourceLocation.GetResource()->UpdateResidency(CommandContext.CommandListHandle);
 
 			bool bAlreadyInSet = false;
 			UniqueGeometries.Add(Geometry, &bAlreadyInSet);
 			if (!bAlreadyInSet)
 			{
-				BottomLevelAccelerationStructureBuffers.Add(Geometry->AccelerationStructureBuffer);
+				BottomLevelAccelerationStructureBuffers[GPUIndex].Add(Geometry->AccelerationStructureBuffers[GPUIndex]);
 			}
 
 			MappedData[InstanceIndex] = InstanceDesc;
@@ -2200,17 +2215,19 @@ void FD3D12RayTracingScene::BuildAccelerationStructure(FD3D12CommandContext& Com
 
 void FD3D12RayTracingScene::UpdateResidency(FD3D12CommandContext& CommandContext)
 {
-	AccelerationStructureBuffer->GetResource()->UpdateResidency(CommandContext.CommandListHandle);
+	const uint32 GPUIndex = CommandContext.GetGPUIndex();
+	AccelerationStructureBuffers[GPUIndex]->GetResource()->UpdateResidency(CommandContext.CommandListHandle);
 
-	for (auto& BLASBuffer : BottomLevelAccelerationStructureBuffers)
+	for (auto& BLASBuffer : BottomLevelAccelerationStructureBuffers[GPUIndex])
 	{
 		BLASBuffer->GetResource()->UpdateResidency(CommandContext.CommandListHandle);
 	}
 }
 
-FD3D12RayTracingShaderTable* FD3D12RayTracingScene::FindExistingShaderTable(const FD3D12RayTracingPipelineState* Pipeline) const
+FD3D12RayTracingShaderTable* FD3D12RayTracingScene::FindExistingShaderTable(const FD3D12RayTracingPipelineState* Pipeline, FD3D12Device* Device) const
 {
-	FD3D12RayTracingShaderTable* const* FoundShaderTable = ShaderTables.Find(Pipeline);
+	const uint32 GPUIndex = Device->GetGPUIndex();
+	FD3D12RayTracingShaderTable* const* FoundShaderTable = ShaderTables[GPUIndex].Find(Pipeline);
 	if (FoundShaderTable)
 	{
 		return *FoundShaderTable;
@@ -2221,16 +2238,18 @@ FD3D12RayTracingShaderTable* FD3D12RayTracingScene::FindExistingShaderTable(cons
 	}
 }
 
-FD3D12RayTracingShaderTable* FD3D12RayTracingScene::FindOrCreateShaderTable(const FD3D12RayTracingPipelineState* Pipeline)
+FD3D12RayTracingShaderTable* FD3D12RayTracingScene::FindOrCreateShaderTable(const FD3D12RayTracingPipelineState* Pipeline, FD3D12Device* Device)
 {
-	FD3D12RayTracingShaderTable* FoundShaderTable = FindExistingShaderTable(Pipeline);
+	FD3D12RayTracingShaderTable* FoundShaderTable = FindExistingShaderTable(Pipeline, Device);
 	if (FoundShaderTable)
 	{
 		return FoundShaderTable;
 	}
 
-	FD3D12RayTracingShaderTable* CreatedShaderTable = new FD3D12RayTracingShaderTable(GetParentDevice());
-	ID3D12Device5* RayTracingDevice = GetParentDevice()->GetRayTracingDevice();
+	FD3D12RayTracingShaderTable* CreatedShaderTable = new FD3D12RayTracingShaderTable();
+	ID3D12Device5* RayTracingDevice = Device->GetRayTracingDevice();
+
+	const uint32 GPUIndex = Device->GetGPUIndex();
 
 	const uint32 NumHitGroupSlots = Pipeline->bAllowHitGroupIndexing ? NumTotalSegments * ShaderSlotsPerGeometrySegment : 0;
 
@@ -2244,7 +2263,7 @@ FD3D12RayTracingShaderTable* FD3D12RayTracingScene::FindOrCreateShaderTable(cons
 	SBTInitializer.LocalRootDataSize = Pipeline->MaxLocalRootSignatureSize;
 	SBTInitializer.MaxViewDescriptorsPerRecord = Pipeline->MaxHitGroupViewDescriptors;
 
-	CreatedShaderTable->Init(SBTInitializer);
+	CreatedShaderTable->Init(SBTInitializer, Device);
 	CreatedShaderTable->SetRayGenIdentifiers(Pipeline->RayGenShaders.Identifiers);
 	CreatedShaderTable->SetMissIdentifiers(Pipeline->MissShaders.Identifiers);
 	CreatedShaderTable->SetDefaultHitGroupIdentifier(Pipeline->HitGroupShaders.Identifiers[0]);
@@ -2266,15 +2285,19 @@ FD3D12RayTracingShaderTable* FD3D12RayTracingScene::FindOrCreateShaderTable(cons
 			static constexpr uint32 IndicesPerPrimitive = 3; // Only triangle meshes are supported
 
 			const uint32 IndexStride = Geometry->IndexStride;
-			const D3D12_GPU_VIRTUAL_ADDRESS IndexBufferAddress = Geometry->IndexBuffer ? Geometry->IndexBuffer->ResourceLocation.GetGPUVirtualAddress() : 0;
-			const D3D12_GPU_VIRTUAL_ADDRESS VertexBufferAddress = Geometry->PositionVertexBuffer->ResourceLocation.GetGPUVirtualAddress() + Geometry->VertexOffsetInBytes;
 
-			if (Geometry->IndexBuffer)
+			FD3D12IndexBuffer* IndexBuffer = FD3D12DynamicRHI::ResourceCast(Geometry->RHIIndexBuffer.GetReference(), GPUIndex);
+			FD3D12VertexBuffer* VertexBuffer = FD3D12DynamicRHI::ResourceCast(Geometry->RHIVertexBuffer.GetReference(), GPUIndex);
+
+			const D3D12_GPU_VIRTUAL_ADDRESS IndexBufferAddress = IndexBuffer ? IndexBuffer->ResourceLocation.GetGPUVirtualAddress() : 0;
+			const D3D12_GPU_VIRTUAL_ADDRESS VertexBufferAddress = VertexBuffer->ResourceLocation.GetGPUVirtualAddress() + Geometry->VertexOffsetInBytes;
+
+			if (IndexBuffer)
 			{
-				CreatedShaderTable->AddResourceReference(Geometry->IndexBuffer->ResourceLocation.GetResource(), Geometry->IndexBuffer);
+				CreatedShaderTable->AddResourceReference(IndexBuffer->ResourceLocation.GetResource(), IndexBuffer);
 			}
 
-			CreatedShaderTable->AddResourceReference(Geometry->PositionVertexBuffer->ResourceLocation.GetResource(), Geometry->PositionVertexBuffer);
+			CreatedShaderTable->AddResourceReference(VertexBuffer->ResourceLocation.GetResource(), VertexBuffer);
 
 			const uint32 NumSegments = Geometry->Segments.Num();
 			for (uint32 SegmentIndex = 0; SegmentIndex < NumSegments; ++SegmentIndex)
@@ -2304,7 +2327,7 @@ FD3D12RayTracingShaderTable* FD3D12RayTracingScene::FindOrCreateShaderTable(cons
 		}
 	}
 
-	ShaderTables.Add(Pipeline, CreatedShaderTable);
+	ShaderTables[GPUIndex].Add(Pipeline, CreatedShaderTable);
 
 	return CreatedShaderTable;
 }
@@ -2325,7 +2348,7 @@ void FD3D12CommandContext::RHIUpdateAccelerationStructures(const TArrayView<cons
 	for (const FAccelerationStructureUpdateParams P : Params)
 	{
 		FD3D12RayTracingGeometry* Geometry = FD3D12DynamicRHI::ResourceCast(P.Geometry);
-		Geometry->PositionVertexBuffer = FD3D12DynamicRHI::ResourceCast(P.VertexBuffer);
+		Geometry->RHIVertexBuffer = P.VertexBuffer;
 		Geometry->TransitionBuffers(*this);
 	}
 	CommandListHandle.FlushResourceBarriers();
@@ -2334,7 +2357,7 @@ void FD3D12CommandContext::RHIUpdateAccelerationStructures(const TArrayView<cons
 	for (const FAccelerationStructureUpdateParams P : Params)
 	{
 		FD3D12RayTracingGeometry* Geometry = FD3D12DynamicRHI::ResourceCast(P.Geometry);
-		Geometry->bIsAccelerationStructureDirty = true;
+		Geometry->SetDirty(GetGPUMask(), true);
 
 		const bool bIsUpdate = !GRayTracingDebugForceFullBuild;
 		Geometry->BuildAccelerationStructure(*this, bIsUpdate);
@@ -2347,7 +2370,7 @@ void FD3D12CommandContext::RHIBuildAccelerationStructures(const TArrayView<const
 	for (const FAccelerationStructureUpdateParams P : Params)
 	{
 		FD3D12RayTracingGeometry* Geometry = FD3D12DynamicRHI::ResourceCast(P.Geometry);
-		Geometry->PositionVertexBuffer = FD3D12DynamicRHI::ResourceCast(P.VertexBuffer);
+		Geometry->RHIVertexBuffer = P.VertexBuffer;
 		Geometry->TransitionBuffers(*this);
 	}
 	CommandListHandle.FlushResourceBarriers();
@@ -2356,7 +2379,7 @@ void FD3D12CommandContext::RHIBuildAccelerationStructures(const TArrayView<const
 	for (const FAccelerationStructureUpdateParams P : Params)
 	{
 		FD3D12RayTracingGeometry* Geometry = FD3D12DynamicRHI::ResourceCast(P.Geometry);
-		Geometry->bIsAccelerationStructureDirty = true;
+		Geometry->SetDirty(GetGPUMask(), true);
 
 		const bool bIsUpdate = false;
 		Geometry->BuildAccelerationStructure(*this, bIsUpdate);
@@ -2495,7 +2518,7 @@ static void SetRayTracingShaderResources(
 		FShaderResourceViewRHIParamRef Resource = SRVs[SRVIndex];
 		if (Resource)
 		{
-			FD3D12ShaderResourceView* SRV = FD3D12DynamicRHI::ResourceCast(Resource);
+			FD3D12ShaderResourceView* SRV = CommandContext.RetrieveObject<FD3D12ShaderResourceView>(Resource);
 			LocalSRVs[SRVIndex] = SRV->GetView();
 			BoundSRVMask |= 1ull << SRVIndex;
 
@@ -2508,7 +2531,7 @@ static void SetRayTracingShaderResources(
 		FUniformBufferRHIParamRef Resource = UniformBuffers[CBVIndex];
 		if (Resource)
 		{
-			FD3D12UniformBuffer* CBV = FD3D12DynamicRHI::ResourceCast(Resource);
+			FD3D12UniformBuffer* CBV = CommandContext.RetrieveObject<FD3D12UniformBuffer>(Resource);
 			LocalCBVs[CBVIndex] = CBV;
 			BoundCBVMask |= 1ull << CBVIndex;
 
@@ -2521,7 +2544,7 @@ static void SetRayTracingShaderResources(
 		FSamplerStateRHIParamRef Resource = Samplers[SamplerIndex];
 		if (Resource)
 		{
-			LocalSamplers[SamplerIndex] = FD3D12DynamicRHI::ResourceCast(Resource)->Descriptor;
+			LocalSamplers[SamplerIndex] = CommandContext.RetrieveObject<FD3D12SamplerState>(Resource)->Descriptor;
 			BoundSamplerMask |= 1ull << SamplerIndex;
 		}
 	}
@@ -2531,7 +2554,7 @@ static void SetRayTracingShaderResources(
 		FUnorderedAccessViewRHIParamRef Resource = UAVs[UAVIndex];
 		if (Resource)
 		{
-			FD3D12UnorderedAccessView* UAV = FD3D12DynamicRHI::ResourceCast(Resource);
+			FD3D12UnorderedAccessView* UAV = CommandContext.RetrieveObject<FD3D12UnorderedAccessView>(Resource);
 			LocalUAVs[UAVIndex] = UAV->GetView();
 			BoundUAVMask |= 1ull << UAVIndex;
 
@@ -2793,22 +2816,20 @@ static void DispatchRays(FD3D12CommandContext& CommandContext,
 
 
 void FD3D12CommandContext::RHIRayTraceOcclusion(FRayTracingSceneRHIParamRef InScene,
-	FShaderResourceViewRHIParamRef InRays,
-	FUnorderedAccessViewRHIParamRef InOutput,
+	FShaderResourceViewRHIParamRef Rays,
+	FUnorderedAccessViewRHIParamRef Output,
 	uint32 NumRays)
 {
 	checkf(GetParentDevice()->GetBasicRayTracingPipeline(), TEXT("Ray tracing support is not initialized for this device. Ensure that InitRayTracing() is called before issuing any ray tracing work."));
 
 	FD3D12RayTracingScene* Scene = FD3D12DynamicRHI::ResourceCast(InScene);
-	FD3D12ShaderResourceView* Rays = FD3D12DynamicRHI::ResourceCast(InRays);
-	FD3D12UnorderedAccessView* Output = FD3D12DynamicRHI::ResourceCast(InOutput);
 
 	FD3D12RayTracingPipelineState* Pipeline = GetParentDevice()->GetBasicRayTracingPipeline()->Occlusion;
-	FD3D12RayTracingShaderTable& ShaderTable = Pipeline->DefaultShaderTable;
+	FD3D12RayTracingShaderTable& ShaderTable = Pipeline->DefaultShaderTables[GetGPUIndex()];
 
 	if (ShaderTable.bIsDirty)
 	{
-		ShaderTable.CopyToGPU();
+		ShaderTable.CopyToGPU(GetParentDevice());
 	}
 
 	Scene->UpdateResidency(*this);
@@ -2820,7 +2841,7 @@ void FD3D12CommandContext::RHIRayTraceOcclusion(FRayTracingSceneRHIParamRef InSc
 	DispatchDesc.Depth = 1;
 
 	FRayTracingShaderBindings Bindings;
-	Bindings.SRVs[0] = Scene->AccelerationStructureView.GetReference();
+	Bindings.SRVs[0] = Scene->GetShaderResourceView();
 	Bindings.SRVs[1] = Rays;
 	Bindings.UAVs[0] = Output;
 
@@ -2841,11 +2862,11 @@ void FD3D12CommandContext::RHIRayTraceIntersection(FRayTracingSceneRHIParamRef I
 	FD3D12UnorderedAccessView* Output = FD3D12DynamicRHI::ResourceCast(InOutput);
 
 	FD3D12RayTracingPipelineState* Pipeline = GetParentDevice()->GetBasicRayTracingPipeline()->Intersection;
-	FD3D12RayTracingShaderTable& ShaderTable = Pipeline->DefaultShaderTable;
+	FD3D12RayTracingShaderTable& ShaderTable = Pipeline->DefaultShaderTables[GetGPUIndex()];
 
 	if (ShaderTable.bIsDirty)
 	{
-		ShaderTable.CopyToGPU();
+		ShaderTable.CopyToGPU(GetParentDevice());
 	}
 
 	Scene->UpdateResidency(*this);
@@ -2857,7 +2878,7 @@ void FD3D12CommandContext::RHIRayTraceIntersection(FRayTracingSceneRHIParamRef I
 	DispatchDesc.Depth = 1;
 
 	FRayTracingShaderBindings Bindings;
-	Bindings.SRVs[0] = Scene->AccelerationStructureView.GetReference();
+	Bindings.SRVs[0] = Scene->GetShaderResourceView();
 	Bindings.SRVs[1] = Rays;
 	// #dxr_todo: intersection and occlusion shaders should be split into separate files to avoid resource slot collisions.
 	// Workaround for now is to bind a valid UAV to slots 0 and 1, even though only slot 1 is referenced.
@@ -2875,13 +2896,14 @@ void FD3D12CommandContext::RHIRayTraceDispatch(FRayTracingPipelineStateRHIParamR
 	uint32 Width, uint32 Height)
 {
 	const FD3D12RayTracingPipelineState* Pipeline = FD3D12DynamicRHI::ResourceCast(InRayTracingPipelineState);
+
 	FD3D12RayTracingScene* Scene = FD3D12DynamicRHI::ResourceCast(InScene);
 
-	FD3D12RayTracingShaderTable* ShaderTable = Scene->FindOrCreateShaderTable(Pipeline);
+	FD3D12RayTracingShaderTable* ShaderTable = Scene->FindOrCreateShaderTable(Pipeline, GetParentDevice());
 
 	if (ShaderTable->bIsDirty)
 	{
-		ShaderTable->CopyToGPU();
+		ShaderTable->CopyToGPU(GetParentDevice());
 	}
 
 	// #dxr_todo UE-72159: avoid updating residency if this scene was already used on the current command list (i.e. multiple ray dispatches are performed back-to-back)
@@ -2911,7 +2933,7 @@ void FD3D12CommandContext::RHISetRayTracingHitGroup(
 
 	checkf(ShaderSlot < Scene->ShaderSlotsPerGeometrySegment, TEXT("Shader slot is invalid. Make sure that ShaderSlotsPerGeometrySegment is correct on FRayTracingSceneInitializer."));
 
-	FD3D12RayTracingShaderTable* ShaderTable = Scene->FindOrCreateShaderTable(Pipeline);
+	FD3D12RayTracingShaderTable* ShaderTable = Scene->FindOrCreateShaderTable(Pipeline, GetParentDevice());
 	checkf(ShaderTable->LocalShaderTableOffset == ShaderTable->HitGroupShaderTableOffset,
 		TEXT("Hit shader records are assumed to be at the beginning of local shader table"));
 
@@ -2946,7 +2968,7 @@ void FD3D12CommandContext::RHISetRayTracingCallableShader(
 
 	checkf(ShaderSlotInScene < Scene->NumCallableShaderSlots, TEXT("Shader slot is invalid. Make sure that NumCallableShaderSlots is correct on FRayTracingSceneInitializer."));
 
-	FD3D12RayTracingShaderTable* ShaderTable = Scene->FindOrCreateShaderTable(Pipeline);
+	FD3D12RayTracingShaderTable* ShaderTable = Scene->FindOrCreateShaderTable(Pipeline, GetParentDevice());
 
 	const uint32 RecordIndex = ShaderTable->CallableShaderRecordIndexOffset + ShaderSlotInScene;
 	ShaderTable->SetLocalShaderIdentifier(RecordIndex, Pipeline->CallableShaders.Identifiers[ShaderIndexInPipeline]);
diff --git a/Engine/Source/Runtime/D3D12RHI/Private/D3D12RayTracing.h b/Engine/Source/Runtime/D3D12RHI/Private/D3D12RayTracing.h
--- a/Engine/Source/Runtime/D3D12RHI/Private/D3D12RayTracing.h
+++ b/Engine/Source/Runtime/D3D12RHI/Private/D3D12RayTracing.h
@@ -11,17 +11,28 @@ class FD3D12RayTracingShaderTable;
 
 typedef FD3D12VertexBuffer FD3D12MemBuffer; // Generic GPU memory buffer
 
-class FD3D12RayTracingGeometry : public FRHIRayTracingGeometry, public FD3D12DeviceChild, public FD3D12LinkedAdapterObject<FD3D12RayTracingGeometry>
+class FD3D12RayTracingGeometry : public FRHIRayTracingGeometry
 {
 public:
 
-	FD3D12RayTracingGeometry(FD3D12Device* Device);
+	FD3D12RayTracingGeometry();
 	~FD3D12RayTracingGeometry();
 
 	void TransitionBuffers(FD3D12CommandContext& CommandContext);
 	void BuildAccelerationStructure(FD3D12CommandContext& CommandContext, bool bIsUpdate);
 
-	bool bIsAccelerationStructureDirty = false;
+	bool bIsAccelerationStructureDirty[MAX_NUM_GPUS] = {};
+	void SetDirty(FRHIGPUMask GPUMask, bool bState)
+	{
+		for (uint32 GPUIndex : GPUMask)
+		{
+			bIsAccelerationStructureDirty[GPUIndex] = bState;
+		}
+	}
+	bool IsDirty(uint32 GPUIndex) const
+	{
+		return bIsAccelerationStructureDirty[GPUIndex];
+	}
 
 	uint32 IndexStride = 0; // 0 for non-indexed / implicit triangle list, 2 for uint16, 4 for uint32
 	uint32 VertexOffsetInBytes = 0;
@@ -36,32 +47,30 @@ public:
 
 	EVertexElementType VertexElemType;
 
-	TRefCountPtr<FD3D12IndexBuffer>  IndexBuffer;
-	TRefCountPtr<FD3D12VertexBuffer> PositionVertexBuffer;
+	FIndexBufferRHIRef  RHIIndexBuffer;
+	FVertexBufferRHIRef RHIVertexBuffer;
 
-	TRefCountPtr<FD3D12MemBuffer> AccelerationStructureBuffer;
-	TRefCountPtr<FD3D12MemBuffer> ScratchBuffer;
+	TRefCountPtr<FD3D12MemBuffer> AccelerationStructureBuffers[MAX_NUM_GPUS];
+	TRefCountPtr<FD3D12MemBuffer> ScratchBuffers[MAX_NUM_GPUS];
 
 };
 
-class FD3D12RayTracingScene : public FRHIRayTracingScene, public FD3D12DeviceChild, public FD3D12LinkedAdapterObject<FD3D12RayTracingScene>
+class FD3D12RayTracingScene : public FRHIRayTracingScene
 {
 public:
 
-	FD3D12RayTracingScene(FD3D12Device* Device);
+	FD3D12RayTracingScene(FD3D12Adapter& Adapter);
 	~FD3D12RayTracingScene();
 
 	void BuildAccelerationStructure(FD3D12CommandContext& CommandContext, D3D12_RAYTRACING_ACCELERATION_STRUCTURE_BUILD_FLAGS BuildFlags);
 
-	// #dxr_todo UE-72157: only GPU resources should be managed as LinkedAdapterObjects
-	TRefCountPtr<FD3D12MemBuffer> AccelerationStructureBuffer;
-	TRefCountPtr<FD3D12ShaderResourceView> AccelerationStructureView;
-	bool bAccelerationStructureViewInitialized = false;
+	TRefCountPtr<FD3D12MemBuffer> AccelerationStructureBuffers[MAX_NUM_GPUS];
+	bool bAccelerationStructureViewInitialized[MAX_NUM_GPUS] = {};
 
 	TArray<FRayTracingGeometryInstance> Instances;
 
 	// Scene keeps track of child acceleration structures to manage their residency
-	TArray<TRefCountPtr<FD3D12MemBuffer>> BottomLevelAccelerationStructureBuffers;
+	TArray<TRefCountPtr<FD3D12MemBuffer>> BottomLevelAccelerationStructureBuffers[MAX_NUM_GPUS];
 	void UpdateResidency(FD3D12CommandContext& CommandContext);
 
 	uint32 ShaderSlotsPerGeometrySegment = 1;
@@ -76,10 +85,10 @@ public:
 	uint32 NumCallableShaderSlots = 0;
 
 	// #dxr_todo UE-68230: shader tables should be explicitly registered and unregistered with the scene
-	FD3D12RayTracingShaderTable* FindOrCreateShaderTable(const FD3D12RayTracingPipelineState* Pipeline);
-	FD3D12RayTracingShaderTable* FindExistingShaderTable(const FD3D12RayTracingPipelineState* Pipeline) const;
+	FD3D12RayTracingShaderTable* FindOrCreateShaderTable(const FD3D12RayTracingPipelineState* Pipeline, FD3D12Device* Device);
+	FD3D12RayTracingShaderTable* FindExistingShaderTable(const FD3D12RayTracingPipelineState* Pipeline, FD3D12Device* Device) const;
 
-	TMap<const FD3D12RayTracingPipelineState*, FD3D12RayTracingShaderTable*> ShaderTables;
+	TMap<const FD3D12RayTracingPipelineState*, FD3D12RayTracingShaderTable*> ShaderTables[MAX_NUM_GPUS];
 };
 
 #endif // D3D12_RHI_RAYTRACING
