Fix typo

diff --git a/Engine/Source/Runtime/D3D12RHI/Private/D3D12RayTracing.cpp b/Engine/Source/Runtime/D3D12RHI/Private/D3D12RayTracing.cpp
--- a/Engine/Source/Runtime/D3D12RHI/Private/D3D12RayTracing.cpp
+++ b/Engine/Source/Runtime/D3D12RHI/Private/D3D12RayTracing.cpp
@@ -2911,7 +2911,7 @@ void FD3D12CommandContext::RHIRayTraceDispatch(FRayTracingPipelineStateRHIParamR
 
 	FD3D12RayTracingShader* RayGenShader = FD3D12DynamicRHI::ResourceCast(RayGenShaderRHI);
 	const int32 RayGenShaderIndex = Pipeline->RayGenShaders.Shaders.Find(RayGenShader);
-	checkf(RayGenShaderIndex != INDEX_NONE, TEXT("ReyGen shader is not present in the given ray tracing pipeline. All RayGen shaders must be declared when creating RTPSO."));
+	checkf(RayGenShaderIndex != INDEX_NONE, TEXT("RayGen shader is not present in the given ray tracing pipeline. All RayGen shaders must be declared when creating RTPSO."));
 
 	D3D12_DISPATCH_RAYS_DESC DispatchDesc = ShaderTable->GetDispatchRaysDesc(RayGenShaderIndex, 0, Pipeline->bAllowHitGroupIndexing);
 
