Fix FAccelerationStructureUpdateParams lifetime issue by changing RHIParamRef raw pointers to RHIRef ref counted pointers

diff --git a/Engine/Source/Runtime/D3D12RHI/Private/D3D12RayTracing.cpp b/Engine/Source/Runtime/D3D12RHI/Private/D3D12RayTracing.cpp
--- a/Engine/Source/Runtime/D3D12RHI/Private/D3D12RayTracing.cpp
+++ b/Engine/Source/Runtime/D3D12RHI/Private/D3D12RayTracing.cpp
@@ -2347,7 +2347,7 @@ void FD3D12CommandContext::RHIUpdateAccelerationStructures(const TArrayView<cons
 	// First batch up all barriers
 	for (const FAccelerationStructureUpdateParams P : Params)
 	{
-		FD3D12RayTracingGeometry* Geometry = FD3D12DynamicRHI::ResourceCast(P.Geometry);
+		FD3D12RayTracingGeometry* Geometry = FD3D12DynamicRHI::ResourceCast(P.Geometry.GetReference());
 		Geometry->RHIVertexBuffer = P.VertexBuffer;
 		Geometry->TransitionBuffers(*this);
 	}
@@ -2356,7 +2356,7 @@ void FD3D12CommandContext::RHIUpdateAccelerationStructures(const TArrayView<cons
 	// Then do all work
 	for (const FAccelerationStructureUpdateParams P : Params)
 	{
-		FD3D12RayTracingGeometry* Geometry = FD3D12DynamicRHI::ResourceCast(P.Geometry);
+		FD3D12RayTracingGeometry* Geometry = FD3D12DynamicRHI::ResourceCast(P.Geometry.GetReference());
 		Geometry->SetDirty(GetGPUMask(), true);
 
 		const bool bIsUpdate = !GRayTracingDebugForceFullBuild;
@@ -2369,7 +2369,7 @@ void FD3D12CommandContext::RHIBuildAccelerationStructures(const TArrayView<const
 	// First batch up all barriers
 	for (const FAccelerationStructureUpdateParams P : Params)
 	{
-		FD3D12RayTracingGeometry* Geometry = FD3D12DynamicRHI::ResourceCast(P.Geometry);
+		FD3D12RayTracingGeometry* Geometry = FD3D12DynamicRHI::ResourceCast(P.Geometry.GetReference());
 		Geometry->RHIVertexBuffer = P.VertexBuffer;
 		Geometry->TransitionBuffers(*this);
 	}
@@ -2378,7 +2378,7 @@ void FD3D12CommandContext::RHIBuildAccelerationStructures(const TArrayView<const
 	// Then do all work
 	for (const FAccelerationStructureUpdateParams P : Params)
 	{
-		FD3D12RayTracingGeometry* Geometry = FD3D12DynamicRHI::ResourceCast(P.Geometry);
+		FD3D12RayTracingGeometry* Geometry = FD3D12DynamicRHI::ResourceCast(P.Geometry.GetReference());
 		Geometry->SetDirty(GetGPUMask(), true);
 
 		const bool bIsUpdate = false;
diff --git a/Engine/Source/Runtime/RHI/Public/RHIContext.h b/Engine/Source/Runtime/RHI/Public/RHIContext.h
--- a/Engine/Source/Runtime/RHI/Public/RHIContext.h
+++ b/Engine/Source/Runtime/RHI/Public/RHIContext.h
@@ -143,8 +143,8 @@ public:
 
 struct FAccelerationStructureUpdateParams
 {
-	FRayTracingGeometryRHIParamRef Geometry;
-	FVertexBufferRHIParamRef VertexBuffer;
+	FRayTracingGeometryRHIRef Geometry;
+	FVertexBufferRHIRef VertexBuffer;
 };
 
 struct FCopyBufferRegionParams
