Fixed implementation of FD3D12CommandContext::RHICopyBufferRegions when using mGPU

- Retrieve and use resources relevant to the parent device of current FD3D12CommandContext
- Fixed transition validation errors due to incorrect handling of duplicate underlying resources in source or destination lists

diff --git a/Engine/Source/Runtime/D3D12RHI/Private/D3D12VertexBuffer.cpp b/Engine/Source/Runtime/D3D12RHI/Private/D3D12VertexBuffer.cpp
--- a/Engine/Source/Runtime/D3D12RHI/Private/D3D12VertexBuffer.cpp
+++ b/Engine/Source/Runtime/D3D12RHI/Private/D3D12VertexBuffer.cpp
@@ -221,146 +221,175 @@ void FD3D12DynamicRHI::RHITransferVertexBufferUnderlyingResource(FVertexBufferRH
 #if D3D12_RHI_RAYTRACING
 void FD3D12CommandContext::RHICopyBufferRegion(FVertexBufferRHIParamRef DestBufferRHI, uint64 DstOffset, FVertexBufferRHIParamRef SourceBufferRHI, uint64 SrcOffset, uint64 NumBytes)
 {
-	FD3D12VertexBuffer*  SourceBuffer = FD3D12DynamicRHI::ResourceCast(SourceBufferRHI);
-	FD3D12VertexBuffer*  DestBuffer = FD3D12DynamicRHI::ResourceCast(DestBufferRHI);
+	FD3D12VertexBuffer*  SourceBuffer = RetrieveObject<FD3D12VertexBuffer>(SourceBufferRHI);
+	FD3D12VertexBuffer*  DestBuffer = RetrieveObject<FD3D12VertexBuffer>(DestBufferRHI);
 
-	while (SourceBuffer && DestBuffer)
-	{
-		FD3D12Device* Device = SourceBuffer->GetParentDevice();
-		check(Device == DestBuffer->GetParentDevice());
+	FD3D12Device* Device = SourceBuffer->GetParentDevice();
+	check(Device == DestBuffer->GetParentDevice());
+	check(Device == GetParentDevice());
 
-		FD3D12Resource* pSourceResource = SourceBuffer->ResourceLocation.GetResource();
-		D3D12_RESOURCE_DESC const& SourceBufferDesc = pSourceResource->GetDesc();
+	FD3D12Resource* pSourceResource = SourceBuffer->ResourceLocation.GetResource();
+	D3D12_RESOURCE_DESC const& SourceBufferDesc = pSourceResource->GetDesc();
 
-		FD3D12Resource* pDestResource = DestBuffer->ResourceLocation.GetResource();
-		D3D12_RESOURCE_DESC const& DestBufferDesc = pDestResource->GetDesc();
+	FD3D12Resource* pDestResource = DestBuffer->ResourceLocation.GetResource();
+	D3D12_RESOURCE_DESC const& DestBufferDesc = pDestResource->GetDesc();
 
-		checkf(pSourceResource != pDestResource, TEXT("CopyBufferRegion cannot be used on the same resource. This can happen when both the source and the dest are suballocated from the same resource."));
+	checkf(pSourceResource != pDestResource, TEXT("CopyBufferRegion cannot be used on the same resource. This can happen when both the source and the dest are suballocated from the same resource."));
 
-		check(DstOffset + NumBytes <= DestBufferDesc.Width);
-		check(SrcOffset + NumBytes <= SourceBufferDesc.Width);
+	check(DstOffset + NumBytes <= DestBufferDesc.Width);
+	check(SrcOffset + NumBytes <= SourceBufferDesc.Width);
 
-		numCopies++;
+	numCopies++;
 
-		FConditionalScopeResourceBarrier ScopeResourceBarrierDest(CommandListHandle, pDestResource, D3D12_RESOURCE_STATE_COPY_DEST, 0);
-		CommandListHandle.FlushResourceBarriers();
-		CommandListHandle->CopyBufferRegion(pDestResource->GetResource(), DestBuffer->ResourceLocation.GetOffsetFromBaseOfResource() + DstOffset, pSourceResource->GetResource(), SourceBuffer->ResourceLocation.GetOffsetFromBaseOfResource() + SrcOffset, NumBytes);
-		CommandListHandle.UpdateResidency(pDestResource);
-		CommandListHandle.UpdateResidency(pSourceResource);
+	FConditionalScopeResourceBarrier ScopeResourceBarrierSource(CommandListHandle, pSourceResource, D3D12_RESOURCE_STATE_COPY_SOURCE, 0);
+	FConditionalScopeResourceBarrier ScopeResourceBarrierDest(CommandListHandle, pDestResource, D3D12_RESOURCE_STATE_COPY_DEST, 0);
+	CommandListHandle.FlushResourceBarriers();
 
-		Device->RegisterGPUWork(1);
+	CommandListHandle->CopyBufferRegion(pDestResource->GetResource(), DestBuffer->ResourceLocation.GetOffsetFromBaseOfResource() + DstOffset, pSourceResource->GetResource(), SourceBuffer->ResourceLocation.GetOffsetFromBaseOfResource() + SrcOffset, NumBytes);
+	CommandListHandle.UpdateResidency(pDestResource);
+	CommandListHandle.UpdateResidency(pSourceResource);
 
-		SourceBuffer = SourceBuffer->GetNextObject();
-		DestBuffer = DestBuffer->GetNextObject();
-	}
+	Device->RegisterGPUWork(1);
 }
 
 void FD3D12CommandContext::RHICopyBufferRegions(const TArrayView<const FCopyBufferRegionParams> Params)
 {
-	auto TransitionBuffer = [](FD3D12Resource* pResource, FD3D12CommandListHandle& InCommandListHandle, const D3D12_RESOURCE_STATES Desired)
-	{
-		bool bUseTracking = pResource->RequiresResourceStateTracking();
-		D3D12_RESOURCE_STATES Current;
-		const uint32 Subresource = 0;
+	// Batched buffer copy finds unique source and destination buffer resources, performs transitions
+	// to copy source / dest state, then performs copies and finally restores original state.
 
-		if (!bUseTracking)
-		{
-			Current = pResource->GetDefaultResourceState();
-			if (Current != Desired)
-			{
-				InCommandListHandle.AddTransitionBarrier(pResource, Current, Desired, Subresource);
-			}
-		}
-		else
-		{
-			FD3D12DynamicRHI::TransitionResource(InCommandListHandle, pResource, Desired, Subresource);
-		}
-	};
+	using FLocalResourceArray = TArray<FD3D12Resource*, TInlineAllocator<16, TMemStackAllocator<>>>;
+	FLocalResourceArray SrcBuffers;
+	FLocalResourceArray DstBuffers;
+
+	SrcBuffers.Reserve(Params.Num());
+	DstBuffers.Reserve(Params.Num());
 
 	// Transition buffers to copy states
 	for (auto& Param : Params)
 	{
-		FD3D12VertexBuffer*  SourceBuffer = FD3D12DynamicRHI::ResourceCast(Param.SourceBuffer);
-		FD3D12VertexBuffer*  DestBuffer = FD3D12DynamicRHI::ResourceCast(Param.DestBuffer);
+		FD3D12VertexBuffer*  SourceBuffer = RetrieveObject<FD3D12VertexBuffer>(Param.SourceBuffer);
+		FD3D12VertexBuffer*  DestBuffer = RetrieveObject<FD3D12VertexBuffer>(Param.DestBuffer);
+		check(SourceBuffer);
+		check(DestBuffer);
 
-		while (SourceBuffer && DestBuffer)
-		{
-			FD3D12Device* Device = SourceBuffer->GetParentDevice();
-			check(Device == DestBuffer->GetParentDevice());
-
-			FD3D12Resource* pSourceResource = SourceBuffer->ResourceLocation.GetResource();
-			FD3D12Resource* pDestResource = DestBuffer->ResourceLocation.GetResource();
+		FD3D12Device* Device = SourceBuffer->GetParentDevice();
+		check(Device == DestBuffer->GetParentDevice());
+		check(Device == GetParentDevice());
 
-			checkf(pSourceResource != pDestResource, TEXT("CopyBufferRegion cannot be used on the same resource. This can happen when both the source and the dest are suballocated from the same resource."));
+		FD3D12Resource* pSourceResource = SourceBuffer->ResourceLocation.GetResource();
+		FD3D12Resource* pDestResource = DestBuffer->ResourceLocation.GetResource();
 
-			TransitionBuffer(pSourceResource, CommandListHandle, D3D12_RESOURCE_STATE_COPY_SOURCE);
-			TransitionBuffer(pDestResource, CommandListHandle, D3D12_RESOURCE_STATE_COPY_DEST);
+		checkf(pSourceResource != pDestResource, TEXT("CopyBufferRegion cannot be used on the same resource. This can happen when both the source and the dest are suballocated from the same resource."));
 
-			SourceBuffer = SourceBuffer->GetNextObject();
-			DestBuffer = DestBuffer->GetNextObject();
-		}
+		SrcBuffers.Add(pSourceResource);
+		DstBuffers.Add(pDestResource);
 	}
 
-	CommandListHandle.FlushResourceBarriers();
+	Algo::Sort(SrcBuffers);
+	Algo::Sort(DstBuffers);
 
-	for (auto& Param : Params)
+	enum class EBatchCopyState
 	{
-		FD3D12VertexBuffer*  SourceBuffer = FD3D12DynamicRHI::ResourceCast(Param.SourceBuffer);
-		FD3D12VertexBuffer*  DestBuffer = FD3D12DynamicRHI::ResourceCast(Param.DestBuffer);
-		uint64 SrcOffset = Param.SrcOffset;
-		uint64 DstOffset = Param.DstOffset;
-		uint64 NumBytes = Param.NumBytes;
+		CopySource,
+		CopyDest,
+		FinalizeSource,
+		FinalizeDest,
+	};
+
+	auto TransitionResources = [](FD3D12CommandListHandle& CommandListHandle, FLocalResourceArray& SortedResources, EBatchCopyState State)
+	{
+		const uint32 Subresource = 0; // Buffers only have one subresource
 
-		while (SourceBuffer && DestBuffer)
+		FD3D12Resource* PrevResource = nullptr;
+		for (FD3D12Resource* Resource : SortedResources)
 		{
-			FD3D12Device* Device = SourceBuffer->GetParentDevice();
-			check(Device == DestBuffer->GetParentDevice());
+			if (Resource == PrevResource) continue; // Skip duplicate resource barriers
 
-			FD3D12Resource* pSourceResource = SourceBuffer->ResourceLocation.GetResource();
-			D3D12_RESOURCE_DESC const& SourceBufferDesc = pSourceResource->GetDesc();
+			const bool bUseDefaultState = !Resource->RequiresResourceStateTracking();
 
-			FD3D12Resource* pDestResource = DestBuffer->ResourceLocation.GetResource();
-			D3D12_RESOURCE_DESC const& DestBufferDesc = pDestResource->GetDesc();
+			D3D12_RESOURCE_STATES DesiredState = D3D12_RESOURCE_STATE_CORRUPT;
+			D3D12_RESOURCE_STATES CurrentState = D3D12_RESOURCE_STATE_CORRUPT;
+			switch (State)
+			{
+			case EBatchCopyState::CopySource:
+				DesiredState = D3D12_RESOURCE_STATE_COPY_SOURCE;
+				CurrentState = bUseDefaultState ? Resource->GetDefaultResourceState() : CurrentState;
+				break;
+			case EBatchCopyState::CopyDest:
+				DesiredState = D3D12_RESOURCE_STATE_COPY_DEST;
+				CurrentState = bUseDefaultState ? Resource->GetDefaultResourceState() : CurrentState;
+				break;
+			case EBatchCopyState::FinalizeSource:
+				CurrentState = D3D12_RESOURCE_STATE_COPY_SOURCE;
+				DesiredState = bUseDefaultState ? Resource->GetDefaultResourceState() : D3D12_RESOURCE_STATE_GENERIC_READ;
+				break;
+			case EBatchCopyState::FinalizeDest:
+				CurrentState = D3D12_RESOURCE_STATE_COPY_DEST;
+				DesiredState = bUseDefaultState ? Resource->GetDefaultResourceState() : D3D12_RESOURCE_STATE_GENERIC_READ;
+				break;
+			default:
+				checkf(false, TEXT("Unexpected batch copy state"));
+				break;
+			}
 
-			check(DstOffset + NumBytes <= DestBufferDesc.Width);
-			check(SrcOffset + NumBytes <= SourceBufferDesc.Width);
+			if (bUseDefaultState)
+			{
+				check(CurrentState != D3D12_RESOURCE_STATE_CORRUPT);
+				CommandListHandle.AddTransitionBarrier(Resource, CurrentState, DesiredState, Subresource);
+			}
+			else
+			{
+				FD3D12DynamicRHI::TransitionResource(CommandListHandle, Resource, DesiredState, Subresource);
+			}
 
-			numCopies++;
+			PrevResource = Resource;
+		}
+	};
 
-			CommandListHandle->CopyBufferRegion(pDestResource->GetResource(), DestBuffer->ResourceLocation.GetOffsetFromBaseOfResource() + DstOffset, pSourceResource->GetResource(), SourceBuffer->ResourceLocation.GetOffsetFromBaseOfResource() + SrcOffset, NumBytes);
-			CommandListHandle.UpdateResidency(pDestResource);
-			CommandListHandle.UpdateResidency(pSourceResource);
+	// Ensure that all previously pending barriers have been processed to avoid incorrect/conflicting transitions for non-tracked resources
+	CommandListHandle.FlushResourceBarriers();
 
-			Device->RegisterGPUWork(1);
+	TransitionResources(CommandListHandle, SrcBuffers, EBatchCopyState::CopySource);
+	TransitionResources(CommandListHandle, DstBuffers, EBatchCopyState::CopyDest);
 
-			SourceBuffer = SourceBuffer->GetNextObject();
-			DestBuffer = DestBuffer->GetNextObject();
-		}
-	}
+	// Issue all copy source/dest barriers before performing actual copies
+	CommandListHandle.FlushResourceBarriers();
 
-	// Transition buffers to generic read
 	for (auto& Param : Params)
 	{
-		FD3D12VertexBuffer*  SourceBuffer = FD3D12DynamicRHI::ResourceCast(Param.SourceBuffer);
-		FD3D12VertexBuffer*  DestBuffer = FD3D12DynamicRHI::ResourceCast(Param.DestBuffer);
+		FD3D12VertexBuffer*  SourceBuffer = RetrieveObject<FD3D12VertexBuffer>(Param.SourceBuffer);
+		FD3D12VertexBuffer*  DestBuffer = RetrieveObject<FD3D12VertexBuffer>(Param.DestBuffer);
+		uint64 SrcOffset = Param.SrcOffset;
+		uint64 DstOffset = Param.DstOffset;
+		uint64 NumBytes = Param.NumBytes;
 
-		while (SourceBuffer && DestBuffer)
-		{
-			FD3D12Device* Device = SourceBuffer->GetParentDevice();
-			check(Device == DestBuffer->GetParentDevice());
+		FD3D12Device* Device = SourceBuffer->GetParentDevice();
+		check(Device == DestBuffer->GetParentDevice());
 
-			FD3D12Resource* pSourceResource = SourceBuffer->ResourceLocation.GetResource();
-			FD3D12Resource* pDestResource = DestBuffer->ResourceLocation.GetResource();
+		FD3D12Resource* pSourceResource = SourceBuffer->ResourceLocation.GetResource();
+		D3D12_RESOURCE_DESC const& SourceBufferDesc = pSourceResource->GetDesc();
 
-			TransitionBuffer(pSourceResource, CommandListHandle, D3D12_RESOURCE_STATE_GENERIC_READ);
-			TransitionBuffer(pDestResource, CommandListHandle, D3D12_RESOURCE_STATE_GENERIC_READ);
+		FD3D12Resource* pDestResource = DestBuffer->ResourceLocation.GetResource();
+		D3D12_RESOURCE_DESC const& DestBufferDesc = pDestResource->GetDesc();
 
-			SourceBuffer = SourceBuffer->GetNextObject();
-			DestBuffer = DestBuffer->GetNextObject();
-		}
+		check(DstOffset + NumBytes <= DestBufferDesc.Width);
+		check(SrcOffset + NumBytes <= SourceBufferDesc.Width);
+
+		numCopies++;
+
+		CommandListHandle->CopyBufferRegion(pDestResource->GetResource(), DestBuffer->ResourceLocation.GetOffsetFromBaseOfResource() + DstOffset, pSourceResource->GetResource(), SourceBuffer->ResourceLocation.GetOffsetFromBaseOfResource() + SrcOffset, NumBytes);
+		CommandListHandle.UpdateResidency(pDestResource);
+		CommandListHandle.UpdateResidency(pSourceResource);
+
+		Device->RegisterGPUWork(1);
 	}
+
+	// Transition buffers back to default readable state
+
+	TransitionResources(CommandListHandle, SrcBuffers, EBatchCopyState::FinalizeSource);
+	TransitionResources(CommandListHandle, DstBuffers, EBatchCopyState::FinalizeDest);
 }
-#endif
+#endif // D3D12_RHI_RAYTRACING
 
 FVertexBufferRHIRef FD3D12DynamicRHI::CreateAndLockVertexBuffer_RenderThread(FRHICommandListImmediate& RHICmdList, uint32 Size, uint32 InUsage, FRHIResourceCreateInfo& CreateInfo, void*& OutDataBuffer)
 {
