Fixed static analysis warning due to alloca() being used in a loop

diff --git a/Engine/Source/Runtime/D3D12RHI/Private/D3D12UniformBuffer.cpp b/Engine/Source/Runtime/D3D12RHI/Private/D3D12UniformBuffer.cpp
--- a/Engine/Source/Runtime/D3D12RHI/Private/D3D12UniformBuffer.cpp
+++ b/Engine/Source/Runtime/D3D12RHI/Private/D3D12UniformBuffer.cpp
@@ -123,15 +123,16 @@ FUniformBufferRHIRef FD3D12DynamicRHI::RHICreateUniformBuffer(const void* Conten
 struct FRHICommandD3D12UpdateUniformBuffer final : public FRHICommand<FRHICommandD3D12UpdateUniformBuffer>
 {
 	FD3D12UniformBuffer* UniformBuffer;
-	FD3D12ResourceLocation* UpdatedLocation;
+	FD3D12ResourceLocation UpdatedLocation;
 	FRHIResource** UpdatedResources;
 	int32 NumResources;
-	FORCEINLINE_DEBUGGABLE FRHICommandD3D12UpdateUniformBuffer(FD3D12UniformBuffer* InUniformBuffer, FD3D12ResourceLocation* InUpdatedLocation, FRHIResource** InUpdatedResources, int32 InNumResources)
+	FORCEINLINE_DEBUGGABLE FRHICommandD3D12UpdateUniformBuffer(FD3D12UniformBuffer* InUniformBuffer, FD3D12ResourceLocation& InUpdatedLocation, FRHIResource** InUpdatedResources, int32 InNumResources)
 		: UniformBuffer(InUniformBuffer)
-		, UpdatedLocation(InUpdatedLocation)
+		, UpdatedLocation(InUpdatedLocation.GetParentDevice())
 		, UpdatedResources(InUpdatedResources)
 		, NumResources(InNumResources)
 	{
+		FD3D12ResourceLocation::TransferOwnership(UpdatedLocation, InUpdatedLocation);
 	}
 
 	void Execute(FRHICommandListBase& CmdList)
@@ -142,7 +143,7 @@ struct FRHICommandD3D12UpdateUniformBuffer final : public FRHICommand<FRHIComman
 			UniformBuffer->ResourceTable[i] = UpdatedResources[i];
 			check(UniformBuffer->ResourceTable[i]);
 		}
-		FD3D12ResourceLocation::TransferOwnership(UniformBuffer->ResourceLocation, *UpdatedLocation);
+		FD3D12ResourceLocation::TransferOwnership(UniformBuffer->ResourceLocation, UpdatedLocation);
 #if USE_STATIC_ROOT_SIGNATURE
 		const uint32 NumBytes = Align(UniformBuffer->GetLayout().ConstantBufferSize, D3D12_CONSTANT_BUFFER_DATA_PLACEMENT_ALIGNMENT);
 		UniformBuffer->View->Create(UniformBuffer->ResourceLocation.GetGPUVirtualAddress(), NumBytes);
@@ -157,22 +158,58 @@ void FD3D12DynamicRHI::RHIUpdateUniformBuffer(FUniformBufferRHIParamRef UniformB
 
 	const FRHIUniformBufferLayout& Layout = UniformBufferRHI->GetLayout();
 	FRHICommandListImmediate& RHICmdList = FRHICommandListExecutor::GetImmediateCommandList();
+	const bool bBypass = RHICmdList.Bypass();
 
 	FD3D12UniformBuffer* UniformBuffer = ResourceCast(UniformBufferRHI);
 
 	const uint32 NumBytes = Layout.ConstantBufferSize;
 	const int32 NumResources = Layout.Resources.Num();
 
-	// Loop over buffers on all GPUs
+	FRHIResource** CmdListResources = nullptr;
+
+	if (NumResources)
+	{
+		CmdListResources = bBypass
+			? (FRHIResource**)FMemory_Alloca(sizeof(FRHIResource*) * NumResources) 
+			: (FRHIResource**)RHICmdList.Alloc(sizeof(FRHIResource*) * NumResources, alignof(FRHIResource*));
+
+		for (int32 ResourceIndex = 0; ResourceIndex < NumResources; ++ResourceIndex)
+		{
+			EUniformBufferBaseType ResourceType = Layout.Resources[ResourceIndex].MemberType;
+
+			FRHIResource* Resource;
+			if (IsShaderParameterTypeIgnoredByRHI(ResourceType))
+			{
+				continue;
+			}
+			else if (IsRDGResourceReferenceShaderParameterType(ResourceType))
+			{
+				check(IsInRenderingThread()); // TODO: UE-68018
+				FRHIResource** ResourcePtr = *(FRHIResource***)((uint8*)Contents + Layout.Resources[ResourceIndex].MemberOffset);
+				Resource = ResourcePtr ? *ResourcePtr : nullptr;
+			}
+			else
+			{
+				Resource = *(FRHIResource**)((uint8*)Contents + Layout.Resources[ResourceIndex].MemberOffset);
+			}
+
+			checkf(Resource, TEXT("Invalid resource entry creating uniform buffer, %s.Resources[%u], ResourceType 0x%x."),
+				*Layout.GetDebugName().ToString(),
+				ResourceIndex,
+				Layout.Resources[ResourceIndex].MemberType);
+
+			CmdListResources[ResourceIndex] = Resource;
+		}
+	}
+
+	// Update buffers on all GPUs by looping over FD3D12LinkedAdapterObject chain
 	while (UniformBuffer)
 	{
 		check(UniformBuffer->ResourceTable.Num() == NumResources);
 
-		const bool bBypass = RHICmdList.Bypass();
 		FD3D12Device* Device = UniformBuffer->GetParentDevice();
-		//FD3D12ResourceLocation is non-copyable, so placement new one on the stack for bypass, or out of the commandlist memory if available. avoids dynamic alloc either way.
-		FD3D12ResourceLocation* UpdatedResourceLocation = bBypass ? new (FMemory_Alloca(sizeof(FD3D12ResourceLocation)))FD3D12ResourceLocation(Device)
-			: new(RHICmdList.Alloc<FD3D12ResourceLocation>()) FD3D12ResourceLocation(Device);
+
+		FD3D12ResourceLocation UpdatedResourceLocation(Device);
 
 		if (NumBytes > 0)
 		{
@@ -181,16 +218,16 @@ void FD3D12DynamicRHI::RHIUpdateUniformBuffer(FUniformBufferRHIParamRef UniformB
 			if (UniformBuffer->UniformBufferUsage == UniformBuffer_MultiFrame)
 			{
 				FD3D12DynamicHeapAllocator& Allocator = GetAdapter().GetUploadHeapAllocator(Device->GetGPUIndex());
-				MappedData = Allocator.AllocUploadResource(NumBytes, DEFAULT_CONTEXT_UPLOAD_POOL_ALIGNMENT, *UpdatedResourceLocation);
+				MappedData = Allocator.AllocUploadResource(NumBytes, DEFAULT_CONTEXT_UPLOAD_POOL_ALIGNMENT, UpdatedResourceLocation);
 			}
 			else
 			{
 				FD3D12FastConstantAllocator& Allocator = GetAdapter().GetTransientUniformBufferAllocator();
 
 #if USE_STATIC_ROOT_SIGNATURE
-				MappedData = Allocator.Allocate(NumBytes, *UpdatedResourceLocation, nullptr);
+				MappedData = Allocator.Allocate(NumBytes, UpdatedResourceLocation, nullptr);
 #else
-				MappedData = Allocator.Allocate(NumBytes, *UpdatedResourceLocation);
+				MappedData = Allocator.Allocate(NumBytes, UpdatedResourceLocation);
 #endif
 			}
 
@@ -198,42 +235,6 @@ void FD3D12DynamicRHI::RHIUpdateUniformBuffer(FUniformBufferRHIParamRef UniformB
 			FMemory::Memcpy(MappedData, Contents, NumBytes);
 		}
 
-
-
-		FRHIResource** CmdListResources = nullptr;
-
-		if (NumResources)
-		{
-			CmdListResources = bBypass ? (FRHIResource**)FMemory_Alloca(sizeof(FRHIResource*) * NumResources) : (FRHIResource**)RHICmdList.Alloc(sizeof(FRHIResource*) * NumResources, alignof(FRHIResource*));
-			for (int32 ResourceIndex = 0; ResourceIndex < NumResources; ++ResourceIndex)
-			{
-				EUniformBufferBaseType ResourceType = Layout.Resources[ResourceIndex].MemberType;
-
-				FRHIResource* Resource;
-				if (IsShaderParameterTypeIgnoredByRHI(ResourceType))
-				{
-					continue;
-				}
-				else if (IsRDGResourceReferenceShaderParameterType(ResourceType))
-				{
-					check(IsInRenderingThread()); // TODO: UE-68018
-					FRHIResource** ResourcePtr = *(FRHIResource***)((uint8*)Contents + Layout.Resources[ResourceIndex].MemberOffset);
-					Resource = ResourcePtr ? *ResourcePtr : nullptr;
-				}
-				else
-				{
-					Resource = *(FRHIResource**)((uint8*)Contents + Layout.Resources[ResourceIndex].MemberOffset);
-				}
-
-				checkf(Resource, TEXT("Invalid resource entry creating uniform buffer, %s.Resources[%u], ResourceType 0x%x."),
-					*Layout.GetDebugName().ToString(),
-					ResourceIndex,
-					Layout.Resources[ResourceIndex].MemberType);
-
-				CmdListResources[ResourceIndex] = Resource;
-			}
-		}
-
 		if (bBypass)
 		{
 			FRHICommandD3D12UpdateUniformBuffer Cmd(UniformBuffer, UpdatedResourceLocation, CmdListResources, NumResources);
