- Fix validation issue for BUF_DrawIndirect resources on D3D12

diff --git a/Engine/Source/Runtime/D3D12RHI/Private/D3D12Allocation.cpp b/Engine/Source/Runtime/D3D12RHI/Private/D3D12Allocation.cpp
--- a/Engine/Source/Runtime/D3D12RHI/Private/D3D12Allocation.cpp
+++ b/Engine/Source/Runtime/D3D12RHI/Private/D3D12Allocation.cpp
@@ -1044,7 +1044,7 @@ void FD3D12DefaultBufferAllocator::InitializeAllocator(EBufferPool PoolIndex, D3
 			Name,
 			kPlacedResourceStrategy,
 			D3D12_HEAP_TYPE_DEFAULT,
-			D3D12_HEAP_FLAG_ALLOW_ONLY_BUFFERS,
+			D3D12_HEAP_FLAG_ALLOW_ONLY_BUFFERS | D3D12RHI_HEAP_FLAG_ALLOW_INDIRECT_BUFFERS,
 			Flags,
 			DEFAULT_BUFFER_POOL_MAX_ALLOC_SIZE,
 			ED3D12AllocatorID::DefaultBufferAllocator,
