Adding Find() method to FD3D12RayTracingShaderLibrary

Uses shader hash instead of pointer lookup when dispatching a RayGen shader.

diff --git a/Engine/Source/Runtime/D3D12RHI/Private/D3D12RayTracing.cpp b/Engine/Source/Runtime/D3D12RHI/Private/D3D12RayTracing.cpp
--- a/Engine/Source/Runtime/D3D12RHI/Private/D3D12RayTracing.cpp
+++ b/Engine/Source/Runtime/D3D12RHI/Private/D3D12RayTracing.cpp
@@ -1249,6 +1249,19 @@ struct FD3D12RayTracingShaderLibrary
 		Identifiers.Reserve(NumShaders);
 	}
 
+	int32 Find(FSHAHash Hash) const
+	{
+		for (int32 Index = 0; Index < Shaders.Num(); ++Index)
+		{
+			if (Hash == Shaders[Index]->GetHash())
+			{
+				return Index;
+			}
+		}
+
+		return INDEX_NONE;
+	}
+
 	TArray<TRefCountPtr<FD3D12RayTracingShader>> Shaders;
 	TArray<FD3D12ShaderIdentifier> Identifiers;
 };
@@ -2910,7 +2923,7 @@ void FD3D12CommandContext::RHIRayTraceDispatch(FRayTracingPipelineStateRHIParamR
 	Scene->UpdateResidency(*this);
 
 	FD3D12RayTracingShader* RayGenShader = FD3D12DynamicRHI::ResourceCast(RayGenShaderRHI);
-	const int32 RayGenShaderIndex = Pipeline->RayGenShaders.Shaders.Find(RayGenShader);
+	const int32 RayGenShaderIndex = Pipeline->RayGenShaders.Find(RayGenShader->GetHash());
 	checkf(RayGenShaderIndex != INDEX_NONE, TEXT("RayGen shader is not present in the given ray tracing pipeline. All RayGen shaders must be declared when creating RTPSO."));
 
 	D3D12_DISPATCH_RAYS_DESC DispatchDesc = ShaderTable->GetDispatchRaysDesc(RayGenShaderIndex, 0, Pipeline->bAllowHitGroupIndexing);
