Show clearer message when running UE4 with -d3ddebug and the Windows Graphics Tools are not installed.

diff --git a/Engine/Source/Runtime/D3D12RHI/Private/D3D12Adapter.cpp b/Engine/Source/Runtime/D3D12RHI/Private/D3D12Adapter.cpp
--- a/Engine/Source/Runtime/D3D12RHI/Private/D3D12Adapter.cpp
+++ b/Engine/Source/Runtime/D3D12RHI/Private/D3D12Adapter.cpp
@@ -96,19 +96,26 @@ void FD3D12Adapter::CreateRootDevice(bool bWithDebug)
 	if (bWithDebug)
 	{
 		TRefCountPtr<ID3D12Debug> DebugController;
-		VERIFYD3D12RESULT(D3D12GetDebugInterface(IID_PPV_ARGS(DebugController.GetInitReference())));
-		DebugController->EnableDebugLayer();
+		if (SUCCEEDED(D3D12GetDebugInterface(IID_PPV_ARGS(DebugController.GetInitReference()))))
+		{
+			DebugController->EnableDebugLayer();
+
+			bool bD3d12gpuvalidation = false;
+			if (FParse::Param(FCommandLine::Get(), TEXT("d3d12gpuvalidation")))
+			{
+				TRefCountPtr<ID3D12Debug1> DebugController1;
+				VERIFYD3D12RESULT(DebugController->QueryInterface(IID_PPV_ARGS(DebugController1.GetInitReference())));
+				DebugController1->SetEnableGPUBasedValidation(true);
+				bD3d12gpuvalidation = true;
+			}
 
-		bool bD3d12gpuvalidation = false;
-		if (FParse::Param(FCommandLine::Get(), TEXT("d3d12gpuvalidation")))
+			UE_LOG(LogD3D12RHI, Log, TEXT("InitD3DDevice: -D3DDebug = %s -D3D12GPUValidation = %s"), bWithDebug ? TEXT("on") : TEXT("off"), bD3d12gpuvalidation ? TEXT("on") : TEXT("off"));
+		}
+		else
 		{
-			TRefCountPtr<ID3D12Debug1> DebugController1;
-			VERIFYD3D12RESULT(DebugController->QueryInterface(IID_PPV_ARGS(DebugController1.GetInitReference())));
-			DebugController1->SetEnableGPUBasedValidation(true);
-			bD3d12gpuvalidation = true;
+			bWithDebug = false;
+			UE_LOG(LogD3D12RHI, Fatal, TEXT("The debug interface requires the D3D12 SDK Layers. Please install the Graphic Tools for Windows. See: https://docs.microsoft.com/en-us/windows/uwp/gaming/use-the-directx-runtime-and-visual-studio-graphics-diagnostic-features"));
 		}
-
-		UE_LOG(LogD3D12RHI, Log, TEXT("InitD3DDevice: -D3DDebug = %s -D3D12GPUValidation = %s"), bWithDebug ? TEXT("on") : TEXT("off"), bD3d12gpuvalidation ? TEXT("on") : TEXT("off") );
 	}
 #endif // PLATFORM_WINDOWS
 
