`gpuvalidation`is also accepted together with `d3d12gpuvalidation` to enable the validation layer to be consistent with the Vulkan RHI

diff --git a/Engine/Source/Runtime/D3D12RHI/Private/D3D12Adapter.cpp b/Engine/Source/Runtime/D3D12RHI/Private/D3D12Adapter.cpp
--- a/Engine/Source/Runtime/D3D12RHI/Private/D3D12Adapter.cpp
+++ b/Engine/Source/Runtime/D3D12RHI/Private/D3D12Adapter.cpp
@@ -101,7 +101,7 @@ void FD3D12Adapter::CreateRootDevice(bool bWithDebug)
 			DebugController->EnableDebugLayer();
 
 			bool bD3d12gpuvalidation = false;
-			if (FParse::Param(FCommandLine::Get(), TEXT("d3d12gpuvalidation")))
+			if (FParse::Param(FCommandLine::Get(), TEXT("d3d12gpuvalidation")) || FParse::Param(FCommandLine::Get(), TEXT("gpuvalidation")))
 			{
 				TRefCountPtr<ID3D12Debug1> DebugController1;
 				VERIFYD3D12RESULT(DebugController->QueryInterface(IID_PPV_ARGS(DebugController1.GetInitReference())));
@@ -114,7 +114,7 @@ void FD3D12Adapter::CreateRootDevice(bool bWithDebug)
 		else
 		{
 			bWithDebug = false;
-			UE_LOG(LogD3D12RHI, Fatal, TEXT("The debug interface requires the D3D12 SDK Layers. Please install the Graphic Tools for Windows. See: https://docs.microsoft.com/en-us/windows/uwp/gaming/use-the-directx-runtime-and-visual-studio-graphics-diagnostic-features"));
+			UE_LOG(LogD3D12RHI, Fatal, TEXT("The debug interface requires the D3D12 SDK Layers. Please install the Graphics Tools for Windows. See: https://docs.microsoft.com/en-us/windows/uwp/gaming/use-the-directx-runtime-and-visual-studio-graphics-diagnostic-features"));
 		}
 	}
 #endif // PLATFORM_WINDOWS
