Fixed implementation of TD3D12Texture2D<RHIResourceType>::UnlockInternal when using mGPU. A common FD3D12ResourceLocation is used for all nodes in the LDA chain, therefore only the last FRHICommandUpdateTexture should take ownership of the resource and release it once it executes.

diff --git a/Engine/Source/Runtime/D3D12RHI/Private/D3D12Texture.cpp b/Engine/Source/Runtime/D3D12RHI/Private/D3D12Texture.cpp
--- a/Engine/Source/Runtime/D3D12RHI/Private/D3D12Texture.cpp
+++ b/Engine/Source/Runtime/D3D12RHI/Private/D3D12Texture.cpp
@@ -75,37 +75,39 @@ template void FD3D12TextureStats::D3D12TextureDeleted(TD3D12Texture2D<FD3D12Base
 struct FRHICommandUpdateTexture final : public FRHICommand<FRHICommandUpdateTexture>
 {
 	FD3D12TextureBase* TextureBase;
-	const D3D12_TEXTURE_COPY_LOCATION DestCopyLocation;
+	D3D12_TEXTURE_COPY_LOCATION DestCopyLocation;
 	uint32 DestX;
 	uint32 DestY;
 	uint32 DestZ;
-	FD3D12ResourceLocation SourceCopyLocation;
-	D3D12_PLACED_SUBRESOURCE_FOOTPRINT SourceFootprint;
+	D3D12_TEXTURE_COPY_LOCATION SourceCopyLocation;
+	FD3D12ResourceLocation Source;
 
-	FORCEINLINE_DEBUGGABLE FRHICommandUpdateTexture(FD3D12TextureBase* InTextureBase, const D3D12_TEXTURE_COPY_LOCATION& InDestCopyLocation, uint32 InDestX, uint32 InDestY, uint32 InDestZ, FD3D12ResourceLocation& Source, D3D12_PLACED_SUBRESOURCE_FOOTPRINT& InSourceFootprint)
+	FORCEINLINE_DEBUGGABLE FRHICommandUpdateTexture(FD3D12TextureBase* InTextureBase,
+		const D3D12_TEXTURE_COPY_LOCATION& InDestCopyLocation, uint32 InDestX, uint32 InDestY, uint32 InDestZ,
+		const D3D12_TEXTURE_COPY_LOCATION& InSourceCopyLocation, FD3D12ResourceLocation* InSource)
 		: TextureBase(InTextureBase)
 		, DestCopyLocation(InDestCopyLocation)
 		, DestX(InDestX)
 		, DestY(InDestY)
 		, DestZ(InDestZ)
-		, SourceCopyLocation(nullptr)
-		, SourceFootprint(InSourceFootprint)
+		, SourceCopyLocation(InSourceCopyLocation)
+		, Source(nullptr)
 	{
 		DestCopyLocation.pResource->AddRef();
-		FD3D12ResourceLocation::TransferOwnership(SourceCopyLocation, Source);
+		if (InSource)
+		{
+			FD3D12ResourceLocation::TransferOwnership(Source, *InSource);
+		}
 	}
 
 	~FRHICommandUpdateTexture()
 	{
 		DestCopyLocation.pResource->Release();
-
 	}
 
 	void Execute(FRHICommandListBase& CmdList)
 	{
-		CD3DX12_TEXTURE_COPY_LOCATION Location(SourceCopyLocation.GetResource()->GetResource(), SourceFootprint);
-
-		TextureBase->UpdateTexture(DestCopyLocation, DestX, DestY, DestZ, Location);
+		TextureBase->UpdateTexture(DestCopyLocation, DestX, DestY, DestZ, SourceCopyLocation);
 	}
 };
 
@@ -1818,18 +1820,19 @@ void TD3D12Texture2D<RHIResourceType>::UnlockInternal(class FRHICommandListImmed
 			PlacedTexture2D.Footprint = BufferPitchDesc;
 
 			CD3DX12_TEXTURE_COPY_LOCATION DestCopyLocation(Resource->GetResource(), Subresource);
+			CD3DX12_TEXTURE_COPY_LOCATION SourceCopyLocation(UploadLocation.GetResource()->GetResource(), PlacedTexture2D);
 
 			FD3D12CommandListHandle& hCommandList = GetParentDevice()->GetDefaultCommandContext().CommandListHandle;
 
 			// If we are on the render thread, queue up the copy on the RHIThread so it happens at the correct time.
 			if (ShouldDeferCmdListOperation(RHICmdList))
 			{
-				ALLOC_COMMAND_CL(*RHICmdList, FRHICommandUpdateTexture)(this, DestCopyLocation, 0, 0, 0, UploadLocation, PlacedTexture2D);
+				// Same FD3D12ResourceLocation is used for all resources in the chain, therefore only the last command must be responsible for releasing it.
+				FD3D12ResourceLocation* Source = GetNextObject() ? nullptr : &UploadLocation;
+				ALLOC_COMMAND_CL(*RHICmdList, FRHICommandUpdateTexture)(this, DestCopyLocation, 0, 0, 0, SourceCopyLocation, Source);
 			}
 			else
 			{
-				CD3DX12_TEXTURE_COPY_LOCATION SourceCopyLocation(UploadLocation.GetResource()->GetResource(), PlacedTexture2D);
-
 				UpdateTexture(DestCopyLocation, 0, 0, 0, SourceCopyLocation);
 			}
 
@@ -1894,15 +1897,15 @@ void TD3D12Texture2D<RHIResourceType>::UpdateTexture2D(class FRHICommandListImme
 		PlacedTexture2D.Footprint = SourceSubresource;
 
 		CD3DX12_TEXTURE_COPY_LOCATION DestCopyLocation(Texture->GetResource()->GetResource(), MipIndex);
-		
+		CD3DX12_TEXTURE_COPY_LOCATION SourceCopyLocation(UploadHeapResourceLocation.GetResource()->GetResource(), PlacedTexture2D);
+
 		// If we are on the render thread, queue up the copy on the RHIThread so it happens at the correct time.
 		if (ShouldDeferCmdListOperation(RHICmdList))
 		{
-			ALLOC_COMMAND_CL(*RHICmdList, FRHICommandUpdateTexture)(this, DestCopyLocation, UpdateRegion.DestX, UpdateRegion.DestY, 0, UploadHeapResourceLocation, PlacedTexture2D);
+			ALLOC_COMMAND_CL(*RHICmdList, FRHICommandUpdateTexture)(this, DestCopyLocation, UpdateRegion.DestX, UpdateRegion.DestY, 0, SourceCopyLocation, &UploadHeapResourceLocation);
 		}
 		else
 		{
-			CD3DX12_TEXTURE_COPY_LOCATION SourceCopyLocation(UploadHeapResourceLocation.GetResource()->GetResource(), PlacedTexture2D);
 			UpdateTexture(DestCopyLocation, UpdateRegion.DestX, UpdateRegion.DestY, 0, SourceCopyLocation);
 		}
 
