Fixed RHICopyTexture implementation in D3D12 RHI.

#rnx

diff --git a/Engine/Source/Runtime/D3D12RHI/Private/D3D12Texture.cpp b/Engine/Source/Runtime/D3D12RHI/Private/D3D12Texture.cpp
--- a/Engine/Source/Runtime/D3D12RHI/Private/D3D12Texture.cpp
+++ b/Engine/Source/Runtime/D3D12RHI/Private/D3D12Texture.cpp
@@ -2847,37 +2847,60 @@ void FD3D12CommandContext::RHICopyTexture(FTextureRHIParamRef SourceTextureRHI,
 	FD3D12TextureBase* SourceTexture = GetD3D12TextureFromRHITexture(SourceTextureRHI);
 	FD3D12TextureBase* DestTexture = GetD3D12TextureFromRHITexture(DestTextureRHI);
 
-	const CD3DX12_BOX SourceBoxD3D(
-		CopyInfo.SourcePosition.X,
-		CopyInfo.SourcePosition.Y,
-		CopyInfo.SourcePosition.Z,
-		CopyInfo.SourcePosition.X + CopyInfo.Size.X,
-		CopyInfo.SourcePosition.Y + CopyInfo.Size.Y,
-		CopyInfo.SourcePosition.Z + CopyInfo.Size.Z
-	);
-
-	CD3DX12_TEXTURE_COPY_LOCATION SourceCopyLocation(
-		SourceTexture->GetResource()->GetResource(),
-		CalcSubresource(CopyInfo.SourceMipIndex, CopyInfo.SourceSliceIndex, SourceTexture->GetResource()->GetMipLevels())
-	);
-	CD3DX12_TEXTURE_COPY_LOCATION DestCopyLocation(
-		DestTexture->GetResource()->GetResource(),
-		CalcSubresource(CopyInfo.DestMipIndex, CopyInfo.DestSliceIndex, DestTexture->GetResource()->GetMipLevels())
-	);
+	CD3DX12_TEXTURE_COPY_LOCATION SourceCopyLocation(SourceTexture->GetResource()->GetResource(), 0);
+	CD3DX12_TEXTURE_COPY_LOCATION DestCopyLocation(DestTexture->GetResource()->GetResource(), 0);
 
 	FConditionalScopeResourceBarrier ConditionalScopeResourceBarrierSource(CommandListHandle, SourceTexture->GetResource(), D3D12_RESOURCE_STATE_COPY_SOURCE, SourceCopyLocation.SubresourceIndex);
 	FConditionalScopeResourceBarrier ConditionalScopeResourceBarrierDest(CommandListHandle, DestTexture->GetResource(), D3D12_RESOURCE_STATE_COPY_DEST, DestCopyLocation.SubresourceIndex);
 
 	numCopies++;
 	CommandListHandle.FlushResourceBarriers();
-	CommandListHandle->CopyTextureRegion(
-		&DestCopyLocation,
-		CopyInfo.DestPosition.X,
-		CopyInfo.DestPosition.Y,
-		CopyInfo.DestPosition.Z,
-		&SourceCopyLocation,
-		&SourceBoxD3D
-	);
+
+	if (CopyInfo.Size != FIntVector::ZeroValue)
+	{
+		// Copy sub texture regions
+		CD3DX12_BOX SourceBoxD3D(
+			CopyInfo.SourcePosition.X,
+			CopyInfo.SourcePosition.Y,
+			CopyInfo.SourcePosition.Z,
+			CopyInfo.SourcePosition.X + CopyInfo.Size.X,
+			CopyInfo.SourcePosition.Y + CopyInfo.Size.Y,
+			CopyInfo.SourcePosition.Z + CopyInfo.Size.Z
+		);
+
+		for (uint32 SliceIndex = 0; SliceIndex < CopyInfo.NumSlices; ++SliceIndex)
+		{
+			uint32 SourceSliceIndex = CopyInfo.SourceSliceIndex + SliceIndex;
+			uint32 DestSliceIndex = CopyInfo.DestSliceIndex + SliceIndex;
+
+			for (uint32 MipIndex = 0; MipIndex < CopyInfo.NumMips; ++MipIndex)
+			{
+				uint32 SourceMipIndex = CopyInfo.SourceMipIndex + MipIndex;
+				uint32 DestMipIndex = CopyInfo.DestMipIndex + MipIndex;
+
+				SourceCopyLocation.SubresourceIndex = CalcSubresource(SourceMipIndex, SourceSliceIndex, SourceTextureRHI->GetNumMips());
+				DestCopyLocation.SubresourceIndex = CalcSubresource(DestMipIndex, DestSliceIndex, DestTextureRHI->GetNumMips());
+
+				SourceBoxD3D.right = CopyInfo.SourcePosition.X + FMath::Max(CopyInfo.Size.X >> MipIndex, 1);
+				SourceBoxD3D.bottom = CopyInfo.SourcePosition.Y + FMath::Max(CopyInfo.Size.Y >> MipIndex, 1);
+				SourceBoxD3D.back = CopyInfo.SourcePosition.Z + FMath::Max(CopyInfo.Size.Z >> MipIndex, 1);
+
+				CommandListHandle->CopyTextureRegion(
+					&DestCopyLocation,
+					CopyInfo.DestPosition.X,
+					CopyInfo.DestPosition.Y,
+					CopyInfo.DestPosition.Z,
+					&SourceCopyLocation,
+					&SourceBoxD3D
+				);
+			}
+		}
+	}
+	else
+	{
+		// Copy whole texture
+		CommandListHandle->CopyTextureRegion(&DestCopyLocation, 0, 0, 0, &SourceCopyLocation, nullptr);
+	}
 
 	CommandListHandle.UpdateResidency(SourceTexture->GetResource());
 	CommandListHandle.UpdateResidency(DestTexture->GetResource());
