Change semantics of 'Width' in DX12 RHIMapStagingSurface to be RowPitch / BytesPerPixel so that it is consistent with DX11. This will also fix virtual texture feedback & HZB readback

diff --git a/Engine/Source/Runtime/D3D12RHI/Private/D3D12RenderTarget.cpp b/Engine/Source/Runtime/D3D12RHI/Private/D3D12RenderTarget.cpp
--- a/Engine/Source/Runtime/D3D12RHI/Private/D3D12RenderTarget.cpp
+++ b/Engine/Source/Runtime/D3D12RHI/Private/D3D12RenderTarget.cpp
@@ -1353,7 +1353,7 @@ void FD3D12DynamicRHI::RHIMapStagingSurface(FTextureRHIParamRef TextureRHI, void
 
 		const D3D12_PLACED_SUBRESOURCE_FOOTPRINT& readBackHeapDesc = DestTexture2D->GetReadBackHeapDesc();
 		OutData = pData;
-		OutWidth = readBackHeapDesc.Footprint.Width;
+		OutWidth = readBackHeapDesc.Footprint.RowPitch / BytesPerPixel;
 		OutHeight = readBackHeapDesc.Footprint.Height;
 
 		// MS: It seems like the second frame in some scenes comes into RHIMapStagingSurface BEFORE the copy to the staging texture, thus the readbackHeapDesc isn't set. This could be bug in UE4.
