Fixed RHICopyTexture in D3D12 RHI, putting all sub resources into memory barrier.

#rnx

diff --git a/Engine/Source/Runtime/D3D12RHI/Private/D3D12Texture.cpp b/Engine/Source/Runtime/D3D12RHI/Private/D3D12Texture.cpp
--- a/Engine/Source/Runtime/D3D12RHI/Private/D3D12Texture.cpp
+++ b/Engine/Source/Runtime/D3D12RHI/Private/D3D12Texture.cpp
@@ -2850,8 +2850,8 @@ void FD3D12CommandContext::RHICopyTexture(FTextureRHIParamRef SourceTextureRHI,
 	CD3DX12_TEXTURE_COPY_LOCATION SourceCopyLocation(SourceTexture->GetResource()->GetResource(), 0);
 	CD3DX12_TEXTURE_COPY_LOCATION DestCopyLocation(DestTexture->GetResource()->GetResource(), 0);
 
-	FConditionalScopeResourceBarrier ConditionalScopeResourceBarrierSource(CommandListHandle, SourceTexture->GetResource(), D3D12_RESOURCE_STATE_COPY_SOURCE, SourceCopyLocation.SubresourceIndex);
-	FConditionalScopeResourceBarrier ConditionalScopeResourceBarrierDest(CommandListHandle, DestTexture->GetResource(), D3D12_RESOURCE_STATE_COPY_DEST, DestCopyLocation.SubresourceIndex);
+	FConditionalScopeResourceBarrier ConditionalScopeResourceBarrierSource(CommandListHandle, SourceTexture->GetResource(), D3D12_RESOURCE_STATE_COPY_SOURCE, D3D12_RESOURCE_BARRIER_ALL_SUBRESOURCES);
+	FConditionalScopeResourceBarrier ConditionalScopeResourceBarrierDest(CommandListHandle, DestTexture->GetResource(), D3D12_RESOURCE_STATE_COPY_DEST, D3D12_RESOURCE_BARRIER_ALL_SUBRESOURCES);
 
 	numCopies++;
 	CommandListHandle.FlushResourceBarriers();
