DR - After RHIExit was moved after the Module unload the the Dx12 Adapter has to be kept alive until the Devices are destroyed. That dependency was sort of manually broken before by marshalling the into a RawPointer Array for Device creation.

diff --git a/Engine/Source/Runtime/D3D12RHI/Private/D3D12RHI.cpp b/Engine/Source/Runtime/D3D12RHI/Private/D3D12RHI.cpp
--- a/Engine/Source/Runtime/D3D12RHI/Private/D3D12RHI.cpp
+++ b/Engine/Source/Runtime/D3D12RHI/Private/D3D12RHI.cpp
@@ -46,7 +46,7 @@ FD3D12DynamicRHI* FD3D12DynamicRHI::SingleD3DRHI = nullptr;
 
 using namespace D3D12RHI;
 
-FD3D12DynamicRHI::FD3D12DynamicRHI(TArray<FD3D12Adapter*>& ChosenAdaptersIn) :
+FD3D12DynamicRHI::FD3D12DynamicRHI(const TArray<TSharedPtr<FD3D12Adapter>>& ChosenAdaptersIn) :
 	NumThreadDynamicHeapAllocators(0),
 	ChosenAdapters(ChosenAdaptersIn),
 	AmdAgsContext(nullptr),
@@ -221,7 +221,7 @@ void FD3D12DynamicRHI::Shutdown()
 	RHIShutdownFlipTracking();
 
 	// Cleanup All of the Adapters
-	for (FD3D12Adapter*& Adapter : ChosenAdapters)
+	for (TSharedPtr<FD3D12Adapter>& Adapter : ChosenAdapters)
 	{
 		// Take a reference on the ID3D12Device so that we can delete the FD3D12Device
 		// and have it's children correctly release ID3D12* objects via RAII
diff --git a/Engine/Source/Runtime/D3D12RHI/Private/D3D12RHIPrivate.h b/Engine/Source/Runtime/D3D12RHI/Private/D3D12RHIPrivate.h
--- a/Engine/Source/Runtime/D3D12RHI/Private/D3D12RHIPrivate.h
+++ b/Engine/Source/Runtime/D3D12RHI/Private/D3D12RHIPrivate.h
@@ -207,7 +207,7 @@ private:
 public:
 
 	/** Initialization constructor. */
-	FD3D12DynamicRHI(TArray<FD3D12Adapter*>& ChosenAdaptersIn);
+	FD3D12DynamicRHI(const TArray<TSharedPtr<FD3D12Adapter>>& ChosenAdaptersIn);
 
 	/** Destructor */
 	virtual ~FD3D12DynamicRHI();
@@ -937,7 +937,7 @@ public:
 
 protected:
 
-	TArray<FD3D12Adapter*> ChosenAdapters;
+	TArray<TSharedPtr<FD3D12Adapter>> ChosenAdapters;
 
 	/** The feature level of the device. */
 	D3D_FEATURE_LEVEL FeatureLevel;
diff --git a/Engine/Source/Runtime/D3D12RHI/Private/D3D12Viewport.cpp b/Engine/Source/Runtime/D3D12RHI/Private/D3D12Viewport.cpp
--- a/Engine/Source/Runtime/D3D12RHI/Private/D3D12Viewport.cpp
+++ b/Engine/Source/Runtime/D3D12RHI/Private/D3D12Viewport.cpp
@@ -912,7 +912,7 @@ void FD3D12DynamicRHI::RHIAdvanceFrameForGetViewportBackBuffer(FViewportRHIParam
 	FRHICommandListImmediate& RHICmdList = FRHICommandListExecutor::GetImmediateCommandList();
 	auto& Adapter = ChosenAdapters[0];
 	{
-		check(Adapter != nullptr);
+		check(Adapter.IsValid());
 		Adapter->SignalFrameFence_RenderThread(RHICmdList);
 	}
 
diff --git a/Engine/Source/Runtime/D3D12RHI/Private/Windows/WindowsD3D12Device.cpp b/Engine/Source/Runtime/D3D12RHI/Private/Windows/WindowsD3D12Device.cpp
--- a/Engine/Source/Runtime/D3D12RHI/Private/Windows/WindowsD3D12Device.cpp
+++ b/Engine/Source/Runtime/D3D12RHI/Private/Windows/WindowsD3D12Device.cpp
@@ -450,13 +450,7 @@ FDynamicRHI* FD3D12DynamicRHIModule::CreateRHI(ERHIFeatureLevel::Type RequestedF
 		GMaxRHIShaderPlatform = SP_PCD3D_SM5;
 	}
 
-	TArray<FD3D12Adapter*> RawPointers;
-	for (int32 i = 0; i < ChosenAdapters.Num(); i++)
-	{
-		RawPointers.Add(ChosenAdapters[i].Get());
-	}
-
-	return new FD3D12DynamicRHI(RawPointers);
+	return new FD3D12DynamicRHI(ChosenAdapters);
 }
 
 void FD3D12DynamicRHIModule::StartupModule()
@@ -505,7 +499,7 @@ void FD3D12DynamicRHIModule::ShutdownModule()
 
 void FD3D12DynamicRHI::Init()
 {
-	for (FD3D12Adapter*& Adapter : ChosenAdapters)
+	for (TSharedPtr<FD3D12Adapter>& Adapter : ChosenAdapters)
 	{
 		Adapter->Initialize(this);
 	}
@@ -533,7 +527,7 @@ void FD3D12DynamicRHI::Init()
 
 	// Create a device chain for each of the adapters we have choosen. This could be a single discrete card,
 	// a set discrete cards linked together (i.e. SLI/Crossfire) an Integrated device or any combination of the above
-	for (FD3D12Adapter*& Adapter : ChosenAdapters)
+	for (TSharedPtr<FD3D12Adapter>& Adapter : ChosenAdapters)
 	{
 		check(Adapter->GetDesc().IsValid());
 		Adapter->InitializeDevices();
@@ -715,7 +709,7 @@ void FD3D12DynamicRHI::PostInit()
 
 	if (GRHISupportsRayTracing)
 	{
-		for (FD3D12Adapter*& Adapter : ChosenAdapters)
+		for (TSharedPtr<FD3D12Adapter>& Adapter : ChosenAdapters)
 		{
 			Adapter->InitializeRayTracing();
 		}
