Fixed race condition when using FD3D12RayTracingPipelineCache. Cache access is now serialized with a critical section until parallel RT shader compilation is implemented as part of UE-68234.

diff --git a/Engine/Source/Runtime/D3D12RHI/Private/D3D12RayTracing.cpp b/Engine/Source/Runtime/D3D12RHI/Private/D3D12RayTracing.cpp
--- a/Engine/Source/Runtime/D3D12RHI/Private/D3D12RayTracing.cpp
+++ b/Engine/Source/Runtime/D3D12RHI/Private/D3D12RayTracing.cpp
@@ -430,6 +430,8 @@ public:
 		ECollectionType CollectionType,
 		bool* bOutCacheHit = nullptr)
 	{
+		FScopeLock Lock(&CriticalSection);
+
 		const uint64 ShaderHash = GetShaderHash64(Shader);
 
 		ID3D12RootSignature* LocalRootSignature = nullptr;
@@ -570,6 +572,7 @@ public:
 
 private:
 
+	FCriticalSection CriticalSection;
 	TMap<FKey, FEntry> Cache;
 	FD3D12RootSignature DefaultLocalRootSignature; // Default empty root signature used for default hit shaders.
 };
