Implemented support for loose parameters (global constant buffer) for ray tracing hit shaders.

diff --git a/Engine/Source/Developer/Windows/ShaderFormatD3D/Private/D3DShaderCompiler.cpp b/Engine/Source/Developer/Windows/ShaderFormatD3D/Private/D3DShaderCompiler.cpp
--- a/Engine/Source/Developer/Windows/ShaderFormatD3D/Private/D3DShaderCompiler.cpp
+++ b/Engine/Source/Developer/Windows/ShaderFormatD3D/Private/D3DShaderCompiler.cpp
@@ -99,6 +99,20 @@ static bool IsRayTracingShader(const FShaderTarget& Target)
 	}
 }
 
+static bool IsGlobalConstantBufferSupported(const FShaderTarget& Target)
+{
+	switch (Target.Frequency)
+	{
+	case SF_RayGen:
+	case SF_RayMiss:
+	case SF_RayCallable:
+		// Global CB is not currently implemented for RayGen, Miss and Callable ray tracing shaders.
+		return false;
+	default:
+		return true;
+	}
+}
+
 static uint32 GetAutoBindingSpace(const FShaderTarget& Target)
 {
 	switch (Target.Frequency)
@@ -1172,9 +1186,12 @@ static bool CompileAndProcessD3DShader(FString& PreprocessedShaderSource, const
 				{
 					Output.bSucceeded = true;
 
-					if (bGlobalUniformBufferUsed && bIsRayTracingShader)
+					bool bGlobalUniformBufferAllowed = false;
+
+					if (bGlobalUniformBufferUsed && !IsGlobalConstantBufferSupported(Input.Target))
 					{
-						FString ErrorString = TEXT("Global constant buffer cannot be used in a ray tracing shader.");
+						const TCHAR* ShaderFrequencyString = GetShaderFrequencyString(Input.Target.GetFrequency(), false);
+						FString ErrorString = FString::Printf(TEXT("Global uniform buffer cannot be used in a %s shader."), ShaderFrequencyString);
 
 						uint32 NumLooseParameters = 0;
 						for (const auto& It : Output.ParameterMap.ParameterMap)
diff --git a/Engine/Source/Runtime/D3D12RHI/Private/D3D12CommandContext.h b/Engine/Source/Runtime/D3D12RHI/Private/D3D12CommandContext.h
--- a/Engine/Source/Runtime/D3D12RHI/Private/D3D12CommandContext.h
+++ b/Engine/Source/Runtime/D3D12RHI/Private/D3D12CommandContext.h
@@ -354,6 +354,7 @@ public:
 		FRayTracingSceneRHIParamRef Scene, uint32 InstanceIndex, uint32 SegmentIndex, uint32 ShaderSlot,
 		FRayTracingPipelineStateRHIParamRef Pipeline, uint32 HitGroupIndex,
 		uint32 NumUniformBuffers, const FUniformBufferRHIParamRef* UniformBuffers,
+		uint32 LooseParameterDataSize, const void* LooseParameterData,
 		uint32 UserData) final override;
 	virtual void RHISetRayTracingCallableShader(
 		FRayTracingSceneRHIParamRef Scene, uint32 ShaderSlotInScene,
@@ -802,9 +803,10 @@ public:
 		FRayTracingSceneRHIParamRef Scene, uint32 InstanceIndex, uint32 SegmentIndex, uint32 ShaderSlot,
 		FRayTracingPipelineStateRHIParamRef Pipeline, uint32 HitGroupIndex,
 		uint32 NumUniformBuffers, const FUniformBufferRHIParamRef* UniformBuffers,
+		uint32 LooseParameterDataSize, const void* LooseParameterData,
 		uint32 UserData) final override
 	{
-		ContextRedirect(RHISetRayTracingHitGroup(Scene, InstanceIndex, SegmentIndex, ShaderSlot, Pipeline, HitGroupIndex, NumUniformBuffers, UniformBuffers, UserData));
+		ContextRedirect(RHISetRayTracingHitGroup(Scene, InstanceIndex, SegmentIndex, ShaderSlot, Pipeline, HitGroupIndex, NumUniformBuffers, UniformBuffers, LooseParameterDataSize, LooseParameterData, UserData));
 	}
 
 	virtual void RHISetRayTracingCallableShader(
diff --git a/Engine/Source/Runtime/D3D12RHI/Private/D3D12RayTracing.cpp b/Engine/Source/Runtime/D3D12RHI/Private/D3D12RayTracing.cpp
--- a/Engine/Source/Runtime/D3D12RHI/Private/D3D12RayTracing.cpp
+++ b/Engine/Source/Runtime/D3D12RHI/Private/D3D12RayTracing.cpp
@@ -1201,7 +1201,7 @@ public:
 	{
 		bool bIsAlreadyInSet = false;
 		ReferencedD3D12Resources.Add(D3D12Resource, &bIsAlreadyInSet);
-		if (!bIsAlreadyInSet)
+		if (!bIsAlreadyInSet && Resource)
 		{
 			ReferencedResources.Add(Resource);
 		}
@@ -1368,6 +1368,7 @@ public:
 		{
 			FD3D12RayTracingShader* Shader = FD3D12DynamicRHI::ResourceCast(ShaderRHI);
 			checkf(Shader->pRootSignature->GetRootSignature() == GlobalRootSignature, TEXT("All raygen and miss shaders must share the same root signature"));
+			checkf(!Shader->ResourceCounts.bGlobalUniformBufferUsed, TEXT("Global uniform buffers are not implemented for ray generation shaders"));
 
 			D3D12_EXISTING_COLLECTION_DESC CollectionDesc = AddShaderCollection(Shader, FD3D12RayTracingPipelineCache::ECollectionType::RayGen);
 
@@ -1385,6 +1386,7 @@ public:
 		{
 			FD3D12RayTracingShader* Shader = FD3D12DynamicRHI::ResourceCast(ShaderRHI);
 			checkf(Shader->pRootSignature->GetRootSignature() == GlobalRootSignature, TEXT("All raygen and miss shaders must share the same root signature"));
+			checkf(!Shader->ResourceCounts.bGlobalUniformBufferUsed, TEXT("Global uniform buffers are not implemented for ray tracing miss shaders"));
 
 			D3D12_EXISTING_COLLECTION_DESC CollectionDesc = AddShaderCollection(Shader, FD3D12RayTracingPipelineCache::ECollectionType::Miss);
 
@@ -1406,7 +1408,6 @@ public:
 			FD3D12RayTracingShader* Shader = FD3D12DynamicRHI::ResourceCast(ShaderRHI);
 
 			checkf(Shader, TEXT("A valid ray tracing hit group shader must be provided for all elements in the FRayTracingPipelineStateInitializer hit group table."));
-			checkf(!Shader->ResourceCounts.bGlobalUniformBufferUsed, TEXT("Global uniform buffers are not implemented for ray tracing shaders"));
 
 			const uint32 ShaderViewDescriptors = Shader->ResourceCounts.NumSRVs + Shader->ResourceCounts.NumUAVs;
 			MaxHitGroupViewDescriptors = FMath::Max(MaxHitGroupViewDescriptors, ShaderViewDescriptors);
@@ -1430,7 +1431,7 @@ public:
 			FD3D12RayTracingShader* Shader = FD3D12DynamicRHI::ResourceCast(ShaderRHI);
 
 			checkf(Shader, TEXT("A valid ray tracing shader must be provided for all elements in the FRayTracingPipelineStateInitializer callable shader table."));
-			checkf(!Shader->ResourceCounts.bGlobalUniformBufferUsed, TEXT("Global uniform buffers are not implemented for ray tracing shaders"));
+			checkf(!Shader->ResourceCounts.bGlobalUniformBufferUsed, TEXT("Global uniform buffers are not implemented for ray tracing callable shaders"));
 
 			const uint32 ShaderViewDescriptors = Shader->ResourceCounts.NumSRVs + Shader->ResourceCounts.NumUAVs;
 			MaxHitGroupViewDescriptors = FMath::Max(MaxHitGroupViewDescriptors, ShaderViewDescriptors);
@@ -1765,6 +1766,10 @@ FRayTracingSceneRHIRef FD3D12DynamicRHI::RHICreateRayTracingScene(const FRayTrac
 
 	FD3D12RayTracingScene* Result = new FD3D12RayTracingScene(Adapter);
 
+	checkf(Initializer.Lifetime == RTSL_SingleFrame, TEXT("Only single-frame ray tracing scenes are currently implemented."));
+
+	Result->Lifetime = Initializer.Lifetime;
+	Result->CreatedFrameFenceValue = Adapter.GetFrameFence().GetCurrentFence();
 	Result->Instances = TArray<FRayTracingGeometryInstance>(Initializer.Instances.GetData(), Initializer.Instances.Num());
 	Result->ShaderSlotsPerGeometrySegment = Initializer.ShaderSlotsPerGeometrySegment;
 
@@ -2431,6 +2436,12 @@ struct FD3D12RayTracingGlobalResourceBinder
 		CommandContext.CommandListHandle->SetComputeRootDescriptorTable(SlotIndex, DescriptorTable);
 	}
 
+	D3D12_GPU_VIRTUAL_ADDRESS CreateTransientConstantBuffer(const void* Data, uint32 DataSize)
+	{
+		checkf(0, TEXT("Loose parameters and transient constant buffers are not implemented for global ray tracing shaders (raygen, miss, callable)"));
+		return (D3D12_GPU_VIRTUAL_ADDRESS)0;
+	}
+
 	void AddResourceReference(FD3D12Resource* D3D12Resource, FRHIResource* RHIResource)
 	{
 		D3D12Resource->UpdateResidency(CommandContext.CommandListHandle);
@@ -2442,7 +2453,8 @@ struct FD3D12RayTracingGlobalResourceBinder
 struct FD3D12RayTracingLocalResourceBinder
 {
 	FD3D12RayTracingLocalResourceBinder(FD3D12CommandContext& InCommandContext, FD3D12RayTracingShaderTable* InShaderTable, const FD3D12RootSignature* InRootSignature, uint32 InRecordIndex)
-		: ShaderTable(InShaderTable)
+		: Device(InCommandContext.GetParentDevice())
+		, ShaderTable(InShaderTable)
 		, RootSignature(InRootSignature)
 		, RecordIndex(InRecordIndex)
 	{
@@ -2473,11 +2485,29 @@ struct FD3D12RayTracingLocalResourceBinder
 		ShaderTable->SetLocalShaderParameters(RecordIndex, BindOffset, DescriptorTable);
 	}
 
+	D3D12_GPU_VIRTUAL_ADDRESS CreateTransientConstantBuffer(const void* Data, uint32 DataSize)
+	{
+		// If we see a significant number of transient allocations coming through this path, we should consider
+		// caching constant buffer blocks inside ShaderTable and linearly sub-allocate from them.
+		// If the amount of data is relatively small, it may also be possible to use root constants and avoid extra allocations entirely.
+
+		FD3D12FastConstantAllocator& Allocator = Device->GetParentAdapter()->GetTransientUniformBufferAllocator();
+		FD3D12ResourceLocation ResourceLocation(Device);
+		void* MappedData = Allocator.Allocate(DataSize, ResourceLocation);
+
+		FMemory::Memcpy(MappedData, Data, DataSize);
+
+		ShaderTable->AddResourceReference(ResourceLocation.GetResource(), nullptr);
+
+		return ResourceLocation.GetGPUVirtualAddress();
+	}
+
 	void AddResourceReference(FD3D12Resource* D3D12Resource, FRHIResource* RHIResource)
 	{
 		ShaderTable->AddResourceReference(D3D12Resource, RHIResource);
 	}
 
+	FD3D12Device* Device = nullptr;
 	FD3D12RayTracingShaderTable* ShaderTable = nullptr;
 	const FD3D12RootSignature* RootSignature = nullptr;
 	uint32 RecordIndex = ~0u;
@@ -2492,6 +2522,7 @@ static void SetRayTracingShaderResources(
 	uint32 InNumUniformBuffers, const FUniformBufferRHIParamRef* UniformBuffers,
 	uint32 InNumSamplers, const FSamplerStateRHIParamRef* Samplers,
 	uint32 InNumUAVs, const FUnorderedAccessViewRHIParamRef* UAVs,
+	uint32 InLooseParameterDataSize, const void* InLooseParameterData,
 	FD3D12RayTracingDescriptorCache& DescriptorCache,
 	ResourceBinderType& Binder)
 {
@@ -2499,7 +2530,7 @@ static void SetRayTracingShaderResources(
 
 	const FD3D12RootSignature* RootSignature = Shader->pRootSignature;
 
-	FD3D12UniformBuffer*        LocalCBVs[MAX_CBS];
+	D3D12_GPU_VIRTUAL_ADDRESS   LocalCBVs[MAX_CBS];
 	D3D12_CPU_DESCRIPTOR_HANDLE LocalSRVs[MAX_SRVS];
 	D3D12_CPU_DESCRIPTOR_HANDLE LocalUAVs[MAX_UAVS];
 	D3D12_CPU_DESCRIPTOR_HANDLE LocalSamplers[MAX_SAMPLERS];
@@ -2548,7 +2579,7 @@ static void SetRayTracingShaderResources(
 		if (Resource)
 		{
 			FD3D12UniformBuffer* CBV = CommandContext.RetrieveObject<FD3D12UniformBuffer>(Resource);
-			LocalCBVs[CBVIndex] = CBV;
+			LocalCBVs[CBVIndex] = CBV->ResourceLocation.GetGPUVirtualAddress();
 			BoundCBVMask |= 1ull << CBVIndex;
 
 			ReferencedResources.Add({ CBV->ResourceLocation.GetResource(), Resource });
@@ -2676,6 +2707,22 @@ static void SetRayTracingShaderResources(
 		}
 	}
 
+	// Bind loose parameters
+
+	if (Shader->ResourceCounts.bGlobalUniformBufferUsed)
+	{
+		checkf(InLooseParameterDataSize && InLooseParameterData, TEXT("Shader uses global uniform buffer, but the required loose parameter data is not provided."));
+	}
+
+	if (InLooseParameterData && Shader->ResourceCounts.bGlobalUniformBufferUsed)
+	{
+		const uint32 CBVIndex = 0; // Global uniform buffer is always assumed to be in slot 0
+		LocalCBVs[CBVIndex] = Binder.CreateTransientConstantBuffer(InLooseParameterData, InLooseParameterDataSize);
+		BoundCBVMask |= 1ull << CBVIndex;
+	}
+
+	// Validate that all resources required by the shader are set
+
 	auto IsCompleteBinding = [](uint32 ExpectedCount, uint64 BoundMask)
 	{
 		if (ExpectedCount > 64) return false; // Bound resource mask can't be represented by uint64
@@ -2723,7 +2770,7 @@ static void SetRayTracingShaderResources(
 		for (uint32 i = 0; i < Shader->ResourceCounts.NumCBs; ++i)
 		{
 			const uint64 SlotMask = (1ull << i);
-			D3D12_GPU_VIRTUAL_ADDRESS BufferAddress = (BoundCBVMask & SlotMask) ? LocalCBVs[i]->ResourceLocation.GetGPUVirtualAddress() : 0;
+			D3D12_GPU_VIRTUAL_ADDRESS BufferAddress = (BoundCBVMask & SlotMask) ? LocalCBVs[i] : 0;
 			Binder.SetRootCBV(BindSlot, i, BufferAddress);
 		}
 	}
@@ -2765,6 +2812,7 @@ static void SetRayTracingShaderResources(
 		ARRAY_COUNT(ResourceBindings.UniformBuffers), ResourceBindings.UniformBuffers,
 		ARRAY_COUNT(ResourceBindings.Samplers), ResourceBindings.Samplers,
 		ARRAY_COUNT(ResourceBindings.UAVs), ResourceBindings.UAVs,
+		0, nullptr, // loose parameters
 		DescriptorCache, Binder);
 }
 
@@ -2942,6 +2990,7 @@ void FD3D12CommandContext::RHISetRayTracingHitGroup(
 	FRayTracingSceneRHIParamRef InScene, uint32 InstanceIndex, uint32 SegmentIndex, uint32 ShaderSlot,
 	FRayTracingPipelineStateRHIParamRef InPipeline, uint32 HitGroupIndex,
 	uint32 NumUniformBuffers, const FUniformBufferRHIParamRef* UniformBuffers,
+	uint32 LooseParameterDataSize, const void* LooseParameterData,
 	uint32 UserData)
 {
 	FD3D12RayTracingScene* Scene = FD3D12DynamicRHI::ResourceCast(InScene);
@@ -2969,6 +3018,7 @@ void FD3D12CommandContext::RHISetRayTracingHitGroup(
 		NumUniformBuffers, UniformBuffers,
 		0, nullptr, // Samplers
 		0, nullptr, // UAVs
+		LooseParameterDataSize, LooseParameterData,
 		*(ShaderTable->DescriptorCache), ResourceBinder);
 }
 
@@ -3002,6 +3052,7 @@ void FD3D12CommandContext::RHISetRayTracingCallableShader(
 		NumUniformBuffers, UniformBuffers,
 		0, nullptr, // Samplers
 		0, nullptr, // UAVs
+		0, nullptr, // Loose parameters
 		*(ShaderTable->DescriptorCache), ResourceBinder);
 }
 
diff --git a/Engine/Source/Runtime/D3D12RHI/Private/D3D12RayTracing.h b/Engine/Source/Runtime/D3D12RHI/Private/D3D12RayTracing.h
--- a/Engine/Source/Runtime/D3D12RHI/Private/D3D12RayTracing.h
+++ b/Engine/Source/Runtime/D3D12RHI/Private/D3D12RayTracing.h
@@ -89,6 +89,9 @@ public:
 	FD3D12RayTracingShaderTable* FindExistingShaderTable(const FD3D12RayTracingPipelineState* Pipeline, FD3D12Device* Device) const;
 
 	TMap<const FD3D12RayTracingPipelineState*, FD3D12RayTracingShaderTable*> ShaderTables[MAX_NUM_GPUS];
+
+	ERayTracingSceneLifetime Lifetime = RTSL_SingleFrame;
+	uint64 CreatedFrameFenceValue = 0;
 };
 
 #endif // D3D12_RHI_RAYTRACING
diff --git a/Engine/Source/Runtime/RHI/Public/DynamicRHI.h b/Engine/Source/Runtime/RHI/Public/DynamicRHI.h
--- a/Engine/Source/Runtime/RHI/Public/DynamicRHI.h
+++ b/Engine/Source/Runtime/RHI/Public/DynamicRHI.h
@@ -134,6 +134,15 @@ struct FRayTracingGeometryInitializer
 	bool bAllowUpdate = false;
 };
 
+enum ERayTracingSceneLifetime
+{
+	// Scene may only be used during the frame when it was created.
+	RTSL_SingleFrame,
+
+	// Scene may be constructed once and used in any number of later frames (not currently implemented).
+	// RTSL_MultiFrame,
+};
+
 struct FRayTracingSceneInitializer
 {
 	TArrayView<FRayTracingGeometryInstance> Instances;
@@ -148,6 +157,10 @@ struct FRayTracingSceneInitializer
 	// Defines how many different callable shaders with unique resource bindings can be bound to this scene.
 	// Shaders and resources are assigned to slots in the scene using SetRayTracingCallableShader().
 	uint32 NumCallableShaderSlots = 0;
+
+	// Defines whether data in this scene should persist between frames.
+	// Currently only single-frame lifetime is supported.
+	ERayTracingSceneLifetime Lifetime = RTSL_SingleFrame;
 };
 
 class FDynamicRHI;
diff --git a/Engine/Source/Runtime/RHI/Public/RHICommandList.h b/Engine/Source/Runtime/RHI/Public/RHICommandList.h
--- a/Engine/Source/Runtime/RHI/Public/RHICommandList.h
+++ b/Engine/Source/Runtime/RHI/Public/RHICommandList.h
@@ -1979,11 +1979,15 @@ struct FRHICommandSetRayTracingBindings final : public FRHICommand<FRHICommandSe
 	uint32 ShaderIndex;
 	uint32 NumUniformBuffers;
 	const FUniformBufferRHIParamRef* UniformBuffers; // Pointer to an array of uniform buffers, allocated inline within the command list
+	uint32 LooseParameterDataSize;
+	const void* LooseParameterData;
 	uint32 UserData;
 	EBindingType BindingType;
 
+	// Hit group bindings
 	FRHICommandSetRayTracingBindings(FRayTracingSceneRHIParamRef InScene, uint32 InInstanceIndex, uint32 InSegmentIndex, uint32 InShaderSlot,
 		FRayTracingPipelineState* InPipeline, uint32 InHitGroupIndex, uint32 InNumUniformBuffers, const FUniformBufferRHIParamRef* InUniformBuffers,
+		uint32 InLooseParameterDataSize, const void* InLooseParameterData,
 		uint32 InUserData)
 		: Scene(InScene)
 		, InstanceIndex(InInstanceIndex)
@@ -1993,11 +1997,14 @@ struct FRHICommandSetRayTracingBindings final : public FRHICommand<FRHICommandSe
 		, ShaderIndex(InHitGroupIndex)
 		, NumUniformBuffers(InNumUniformBuffers)
 		, UniformBuffers(InUniformBuffers)
+		, LooseParameterDataSize(InLooseParameterDataSize)
+		, LooseParameterData(InLooseParameterData)
 		, UserData(InUserData)
 		, BindingType(EBindingType_HitGroup)
 	{
 	}
 
+	// Callable shader bindings
 	FRHICommandSetRayTracingBindings(FRayTracingSceneRHIParamRef InScene, uint32 InShaderSlot,
 		FRayTracingPipelineState* InPipeline, uint32 InShaderIndex,
 		uint32 InNumUniformBuffers, const FUniformBufferRHIParamRef* InUniformBuffers,
@@ -3082,12 +3089,16 @@ public:
 		FRayTracingSceneRHIParamRef Scene, uint32 InstanceIndex, uint32 SegmentIndex, uint32 ShaderSlot,
 		FRayTracingPipelineState* Pipeline, uint32 HitGroupIndex,
 		uint32 NumUniformBuffers, const FUniformBufferRHIParamRef* UniformBuffers,
+		uint32 LooseParameterDataSize, const void* LooseParameterData,
 		uint32 UserData)
 	{
 		if (Bypass())
 		{
 			extern RHI_API FRHIRayTracingPipelineState* GetRHIRayTracingPipelineState(FRayTracingPipelineState*);
-			GetContext().RHISetRayTracingHitGroup(Scene, InstanceIndex, SegmentIndex, ShaderSlot, GetRHIRayTracingPipelineState(Pipeline), HitGroupIndex, NumUniformBuffers, UniformBuffers, UserData);
+			GetContext().RHISetRayTracingHitGroup(Scene, InstanceIndex, SegmentIndex, ShaderSlot, GetRHIRayTracingPipelineState(Pipeline), HitGroupIndex, 
+				NumUniformBuffers, UniformBuffers,
+				LooseParameterDataSize, LooseParameterData,
+				UserData);
 		}
 		else
 		{
@@ -3101,7 +3112,17 @@ public:
 				}
 			}
 
-			ALLOC_COMMAND(FRHICommandSetRayTracingBindings)(Scene, InstanceIndex, SegmentIndex, ShaderSlot, Pipeline, HitGroupIndex, NumUniformBuffers, InlineUniformBuffers, UserData);
+			void* InlineLooseParameterData = nullptr;
+			if (LooseParameterDataSize)
+			{
+				InlineLooseParameterData = Alloc(LooseParameterDataSize, 16);
+				FMemory::Memcpy(InlineLooseParameterData, LooseParameterData, LooseParameterDataSize);
+			}
+
+			ALLOC_COMMAND(FRHICommandSetRayTracingBindings)(Scene, InstanceIndex, SegmentIndex, ShaderSlot, Pipeline, HitGroupIndex, 
+				NumUniformBuffers, InlineUniformBuffers, 
+				LooseParameterDataSize, InlineLooseParameterData,
+				UserData);
 		}
 	}
 
diff --git a/Engine/Source/Runtime/RHI/Public/RHICommandListCommandExecutes.inl b/Engine/Source/Runtime/RHI/Public/RHICommandListCommandExecutes.inl
--- a/Engine/Source/Runtime/RHI/Public/RHICommandListCommandExecutes.inl
+++ b/Engine/Source/Runtime/RHI/Public/RHICommandListCommandExecutes.inl
@@ -588,7 +588,10 @@ void FRHICommandSetRayTracingBindings::Execute(FRHICommandListBase& CmdList)
 	extern RHI_API FRHIRayTracingPipelineState* GetRHIRayTracingPipelineState(FRayTracingPipelineState*);
 	if (BindingType == EBindingType_HitGroup)
 	{
-		INTERNAL_DECORATOR(RHISetRayTracingHitGroup)(Scene, InstanceIndex, SegmentIndex, ShaderSlot, GetRHIRayTracingPipelineState(Pipeline), ShaderIndex, NumUniformBuffers, UniformBuffers, UserData);
+		INTERNAL_DECORATOR(RHISetRayTracingHitGroup)(Scene, InstanceIndex, SegmentIndex, ShaderSlot, GetRHIRayTracingPipelineState(Pipeline), ShaderIndex,
+			NumUniformBuffers, UniformBuffers,
+			LooseParameterDataSize, LooseParameterData,
+			UserData);
 	}
 	else
 	{
diff --git a/Engine/Source/Runtime/RHI/Public/RHIContext.h b/Engine/Source/Runtime/RHI/Public/RHIContext.h
--- a/Engine/Source/Runtime/RHI/Public/RHIContext.h
+++ b/Engine/Source/Runtime/RHI/Public/RHIContext.h
@@ -627,6 +627,7 @@ public:
 		FRayTracingSceneRHIParamRef Scene, uint32 InstanceIndex, uint32 SegmentIndex, uint32 ShaderSlot,
 		FRayTracingPipelineStateRHIParamRef Pipeline, uint32 HitGroupIndex,
 		uint32 NumUniformBuffers, const FUniformBufferRHIParamRef* UniformBuffers,
+		uint32 LooseParameterDataSize, const void* LooseParameterData,
 		uint32 UserData)
 	{
 		checkNoEntry();
diff --git a/Engine/Source/Runtime/RHI/Public/RHIValidationContext.h b/Engine/Source/Runtime/RHI/Public/RHIValidationContext.h
--- a/Engine/Source/Runtime/RHI/Public/RHIValidationContext.h
+++ b/Engine/Source/Runtime/RHI/Public/RHIValidationContext.h
@@ -852,9 +852,10 @@ public:
 		FRayTracingSceneRHIParamRef Scene, uint32 InstanceIndex, uint32 SegmentIndex, uint32 ShaderSlot,
 		FRayTracingPipelineStateRHIParamRef Pipeline, uint32 HitGroupIndex,
 		uint32 NumUniformBuffers, const FUniformBufferRHIParamRef* UniformBuffers,
+		uint32 LooseParameterDataSize, const void* LooseParameterData,
 		uint32 UserData) override final
 	{
-		RHIContext->RHISetRayTracingHitGroup(Scene, InstanceIndex, SegmentIndex, ShaderSlot, Pipeline, HitGroupIndex, NumUniformBuffers, UniformBuffers, UserData);
+		RHIContext->RHISetRayTracingHitGroup(Scene, InstanceIndex, SegmentIndex, ShaderSlot, Pipeline, HitGroupIndex, NumUniformBuffers, UniformBuffers, LooseParameterDataSize, LooseParameterData, UserData);
 	}
 
 	IRHICommandContext* RHIContext;
diff --git a/Engine/Source/Runtime/Renderer/Private/MeshPassProcessor.cpp b/Engine/Source/Runtime/Renderer/Private/MeshPassProcessor.cpp
--- a/Engine/Source/Runtime/Renderer/Private/MeshPassProcessor.cpp
+++ b/Engine/Source/Runtime/Renderer/Private/MeshPassProcessor.cpp
@@ -301,12 +301,25 @@ void FMeshDrawShaderBindings::SetRayTracingShaderBindingsForHitGroup(
 
 	checkf(SingleShaderBindings.ParameterMapInfo.TextureSamplers.Num() == 0, TEXT("Texture sampler parameters are not supported for ray tracing. UniformBuffers must be used for all resource binding."));
 	checkf(SingleShaderBindings.ParameterMapInfo.SRVs.Num() == 0, TEXT("SRV parameters are not supported for ray tracing. UniformBuffers must be used for all resource binding."));
-	checkf(SingleShaderBindings.ParameterMapInfo.LooseParameterBuffers.Num() == 0, TEXT("Loose parameter buffers are not supported for ray tracing. UniformBuffers must be used for all resource binding."));
+
+	const TArray<FShaderLooseParameterBufferInfo>& LooseParameterBuffers = SingleShaderBindings.ParameterMapInfo.LooseParameterBuffers;	
+
+	const void* LooseParameterData = SingleShaderBindings.GetLooseDataStart();
+	uint32 LooseParameterDataSize = 0;
+	check(LooseParameterBuffers.Num() <= 1);
+	if (LooseParameterBuffers.Num())
+	{
+		LooseParameterDataSize = LooseParameterBuffers[0].BufferSize;
+		check(LooseParameterBuffers[0].BufferIndex == 0);
+	}
 
 	check(SegmentIndex < 0xFF);
 	uint32 NumUniformBuffersToSet = MaxUniformBufferUsed + 1;
 	const uint32 UserData = 0; // UserData could be used to store material ID or any other kind of per-material constant. This can be retrieved in hit shaders via GetHitGroupUserData().
-	RHICmdList.SetRayTracingHitGroup(Scene, InstanceIndex, SegmentIndex, ShaderSlot, PipelineState, HitGroupIndex, NumUniformBuffersToSet, LocalUniformBuffers, UserData);
+	RHICmdList.SetRayTracingHitGroup(Scene, InstanceIndex, SegmentIndex, ShaderSlot, PipelineState, HitGroupIndex, 
+		NumUniformBuffersToSet, LocalUniformBuffers,
+		LooseParameterDataSize, LooseParameterData,
+		UserData);
 }
 #endif // RHI_RAYTRACING
 
diff --git a/Engine/Source/Runtime/Renderer/Private/RayTracing/RayTracingDeferredMaterials.cpp b/Engine/Source/Runtime/Renderer/Private/RayTracing/RayTracingDeferredMaterials.cpp
--- a/Engine/Source/Runtime/Renderer/Private/RayTracing/RayTracingDeferredMaterials.cpp
+++ b/Engine/Source/Runtime/Renderer/Private/RayTracing/RayTracingDeferredMaterials.cpp
@@ -79,8 +79,8 @@ FRayTracingPipelineState* FDeferredShadingSceneRenderer::BindRayTracingDeferredM
 			ShaderSlot,
 			PipelineState,
 			HitGroupIndex,
-			0, 
-			nullptr,
+			0, nullptr, // uniform buffers
+			0, nullptr, // loose data
 			MaterialIndexInUserData);
 	}
 
diff --git a/Engine/Source/Runtime/Renderer/Private/RayTracing/RayTracingMaterialHitShaders.cpp b/Engine/Source/Runtime/Renderer/Private/RayTracing/RayTracingMaterialHitShaders.cpp
--- a/Engine/Source/Runtime/Renderer/Private/RayTracing/RayTracingMaterialHitShaders.cpp
+++ b/Engine/Source/Runtime/Renderer/Private/RayTracing/RayTracingMaterialHitShaders.cpp
@@ -510,7 +510,9 @@ FRayTracingPipelineState* FDeferredShadingSceneRenderer::BindRayTracingMaterialP
 					MeshCommand.GeometrySegmentIndex,
 					RAY_TRACING_SHADER_SLOT_SHADOW,
 					PipelineState, OpaqueShadowMaterialIndex,
-					0, nullptr, 0);
+					0, nullptr, // uniform buffers
+					0, nullptr, // loose data
+					0);
 			}
 			else
 			{
@@ -532,7 +534,9 @@ FRayTracingPipelineState* FDeferredShadingSceneRenderer::BindRayTracingMaterialP
 				MeshCommand.GeometrySegmentIndex,
 				RAY_TRACING_SHADER_SLOT_SHADOW,
 				PipelineState, HiddenMaterialIndex,
-				0, nullptr, 0);
+				0, nullptr, // uniform buffers
+				0, nullptr, // loose data
+				0);
 		}
 	}
 
