DR - Deprecate FSamplerStateRHIParamRef

diff --git a/Engine/Plugins/FX/Niagara/Source/Niagara/Private/NiagaraDataInterfaceVectorField.cpp b/Engine/Plugins/FX/Niagara/Source/Niagara/Private/NiagaraDataInterfaceVectorField.cpp
--- a/Engine/Plugins/FX/Niagara/Source/Niagara/Private/NiagaraDataInterfaceVectorField.cpp
+++ b/Engine/Plugins/FX/Niagara/Source/Niagara/Private/NiagaraDataInterfaceVectorField.cpp
@@ -340,7 +340,7 @@ struct FNiagaraDataInterfaceParametersCS_VectorField : public FNiagaraDataInterf
 		//     1st bit: X-axis tiling flag
 		//     2nd bit: Y-axis tiling flag
 		//     3rd bit: Z-axis tiling flag
-		static FSamplerStateRHIParamRef SamplerStates[8] = { nullptr };
+		static FRHISamplerState* SamplerStates[8] = { nullptr };
 		if (SamplerStates[0] == nullptr)
 		{
 			SamplerStates[0] = TStaticSamplerState<SF_Bilinear, AM_Clamp, AM_Clamp, AM_Clamp>::GetRHI();
@@ -365,7 +365,7 @@ struct FNiagaraDataInterfaceParametersCS_VectorField : public FNiagaraDataInterf
 		SetTextureParameter(RHICmdList, ComputeShaderRHI, VectorFieldTexture, VolumeTextureRHI); 
 		
 		// Get and set sampler state
-		FSamplerStateRHIParamRef SamplerState = SamplerStates[int(VFDI->bTileX) + 2 * int(VFDI->bTileY) + 4 * int(VFDI->bTileZ)];
+		FRHISamplerState* SamplerState = SamplerStates[int(VFDI->bTileX) + 2 * int(VFDI->bTileY) + 4 * int(VFDI->bTileZ)];
 		SetSamplerParameter(RHICmdList, ComputeShaderRHI, VectorFieldSampler, SamplerState);
 
 		//
diff --git a/Engine/Plugins/Runtime/Oculus/OculusVR/Source/OculusHMD/Private/OculusHMD_CustomPresent.cpp b/Engine/Plugins/Runtime/Oculus/OculusVR/Source/OculusHMD/Private/OculusHMD_CustomPresent.cpp
--- a/Engine/Plugins/Runtime/Oculus/OculusVR/Source/OculusHMD/Private/OculusHMD_CustomPresent.cpp
+++ b/Engine/Plugins/Runtime/Oculus/OculusVR/Source/OculusHMD/Private/OculusHMD_CustomPresent.cpp
@@ -411,7 +411,7 @@ void FCustomPresent::CopyTexture_RenderThread(FRHICommandListImmediate& RHICmdLi
 			}
 
 			RHICmdList.ApplyCachedRenderTargets(GraphicsPSOInit);
-			FSamplerStateRHIParamRef SamplerState = DstRect.Size() == SrcRect.Size() ? TStaticSamplerState<SF_Point>::GetRHI() : TStaticSamplerState<SF_Bilinear>::GetRHI();
+			FRHISamplerState* SamplerState = DstRect.Size() == SrcRect.Size() ? TStaticSamplerState<SF_Point>::GetRHI() : TStaticSamplerState<SF_Bilinear>::GetRHI();
 
 			if (!sRGBSource)
 			{
@@ -460,7 +460,7 @@ void FCustomPresent::CopyTexture_RenderThread(FRHICommandListImmediate& RHICmdLi
 				TShaderMapRef<FOculusCubemapPS> PixelShader(ShaderMap);
 				GraphicsPSOInit.BoundShaderState.PixelShaderRHI = GETSAFERHISHADER_PIXEL(*PixelShader);
 				SetGraphicsPipelineState(RHICmdList, GraphicsPSOInit);
-				FSamplerStateRHIParamRef SamplerState = DstRect.Size() == SrcRect.Size() ? TStaticSamplerState<SF_Point>::GetRHI() : TStaticSamplerState<SF_Bilinear>::GetRHI();
+				FRHISamplerState* SamplerState = DstRect.Size() == SrcRect.Size() ? TStaticSamplerState<SF_Point>::GetRHI() : TStaticSamplerState<SF_Bilinear>::GetRHI();
 				PixelShader->SetParameters(RHICmdList, SamplerState, SrcTextureRHI, FaceIndex);
 
 				RHICmdList.SetViewport(DstRect.Min.X, DstRect.Min.Y, 0.0f, DstRect.Max.X, DstRect.Max.Y, 1.0f);
diff --git a/Engine/Plugins/Runtime/Steam/SteamVR/Source/SteamVR/Private/SteamVRSpectatorScreen.cpp b/Engine/Plugins/Runtime/Steam/SteamVR/Source/SteamVR/Private/SteamVRSpectatorScreen.cpp
--- a/Engine/Plugins/Runtime/Steam/SteamVR/Source/SteamVR/Private/SteamVRSpectatorScreen.cpp
+++ b/Engine/Plugins/Runtime/Steam/SteamVR/Source/SteamVR/Private/SteamVRSpectatorScreen.cpp
@@ -81,7 +81,7 @@ void FSteamVRHMD::CopyTexture_RenderThread(FRHICommandListImmediate& RHICmdList,
 		GraphicsPSOInit.BoundShaderState.VertexShaderRHI = GETSAFERHISHADER_VERTEX(*VertexShader);
 
 		const bool bSameSize = DstRect.Size() == SrcRect.Size();
-		FSamplerStateRHIParamRef PixelSampler = bSameSize ? TStaticSamplerState<SF_Point>::GetRHI() : TStaticSamplerState<SF_Bilinear>::GetRHI();
+		FRHISamplerState* PixelSampler = bSameSize ? TStaticSamplerState<SF_Point>::GetRHI() : TStaticSamplerState<SF_Bilinear>::GetRHI();
 
 		if ((SrcTexture->GetFlags() & TexCreate_SRGB) != 0)
 		{
diff --git a/Engine/Plugins/Runtime/WindowsMixedReality/Source/WindowsMixedRealityHMD/Private/WindowsMixedRealityHMD.cpp b/Engine/Plugins/Runtime/WindowsMixedReality/Source/WindowsMixedRealityHMD/Private/WindowsMixedRealityHMD.cpp
--- a/Engine/Plugins/Runtime/WindowsMixedReality/Source/WindowsMixedRealityHMD/Private/WindowsMixedRealityHMD.cpp
+++ b/Engine/Plugins/Runtime/WindowsMixedReality/Source/WindowsMixedRealityHMD/Private/WindowsMixedRealityHMD.cpp
@@ -55,7 +55,7 @@ public:
 
 		SetShaderValue(RHICmdList, PixelShaderRHI, FarPlaneDistance, farPlaneDistanceValue);
 
-		FSamplerStateRHIParamRef SamplerStateRHI = TStaticSamplerState<SF_Point>::GetRHI();
+		FRHISamplerState* SamplerStateRHI = TStaticSamplerState<SF_Point>::GetRHI();
 		SetTextureParameter(RHICmdList, PixelShaderRHI, InDepthTexture, InTextureSampler, SamplerStateRHI, DepthTexture);
 	}
 
diff --git a/Engine/Source/Runtime/Apple/MetalRHI/Private/MetalCommands.cpp b/Engine/Source/Runtime/Apple/MetalRHI/Private/MetalCommands.cpp
--- a/Engine/Source/Runtime/Apple/MetalRHI/Private/MetalCommands.cpp
+++ b/Engine/Source/Runtime/Apple/MetalRHI/Private/MetalCommands.cpp
@@ -439,7 +439,7 @@ void FMetalRHICommandContext::RHISetShaderResourceViewParameter(FComputeShaderRH
 }
 
 
-void FMetalRHICommandContext::RHISetShaderSampler(FVertexShaderRHIParamRef VertexShaderRHI, uint32 SamplerIndex, FSamplerStateRHIParamRef NewStateRHI)
+void FMetalRHICommandContext::RHISetShaderSampler(FVertexShaderRHIParamRef VertexShaderRHI, uint32 SamplerIndex, FRHISamplerState* NewStateRHI)
 {
 	@autoreleasepool {
 	FMetalSamplerState* NewState = ResourceCast(NewStateRHI);
@@ -448,7 +448,7 @@ void FMetalRHICommandContext::RHISetShaderSampler(FVertexShaderRHIParamRef Verte
 	}
 }
 
-void FMetalRHICommandContext::RHISetShaderSampler(FHullShaderRHIParamRef HullShader, uint32 SamplerIndex, FSamplerStateRHIParamRef NewStateRHI)
+void FMetalRHICommandContext::RHISetShaderSampler(FHullShaderRHIParamRef HullShader, uint32 SamplerIndex, FRHISamplerState* NewStateRHI)
 {
 #if PLATFORM_SUPPORTS_TESSELLATION_SHADERS
 	@autoreleasepool {
@@ -459,7 +459,7 @@ void FMetalRHICommandContext::RHISetShaderSampler(FHullShaderRHIParamRef HullSha
 #endif
 }
 
-void FMetalRHICommandContext::RHISetShaderSampler(FDomainShaderRHIParamRef DomainShader, uint32 SamplerIndex, FSamplerStateRHIParamRef NewStateRHI)
+void FMetalRHICommandContext::RHISetShaderSampler(FDomainShaderRHIParamRef DomainShader, uint32 SamplerIndex, FRHISamplerState* NewStateRHI)
 {
 #if PLATFORM_SUPPORTS_TESSELLATION_SHADERS
 	@autoreleasepool {
@@ -470,12 +470,12 @@ void FMetalRHICommandContext::RHISetShaderSampler(FDomainShaderRHIParamRef Domai
 #endif
 }
 
-void FMetalRHICommandContext::RHISetShaderSampler(FGeometryShaderRHIParamRef GeometryShader, uint32 SamplerIndex, FSamplerStateRHIParamRef NewStateRHI)
+void FMetalRHICommandContext::RHISetShaderSampler(FGeometryShaderRHIParamRef GeometryShader, uint32 SamplerIndex, FRHISamplerState* NewStateRHI)
 {
 	NOT_SUPPORTED("RHISetSamplerState-Geometry");
 }
 
-void FMetalRHICommandContext::RHISetShaderSampler(FPixelShaderRHIParamRef PixelShader, uint32 SamplerIndex, FSamplerStateRHIParamRef NewStateRHI)
+void FMetalRHICommandContext::RHISetShaderSampler(FPixelShaderRHIParamRef PixelShader, uint32 SamplerIndex, FRHISamplerState* NewStateRHI)
 {
 	@autoreleasepool {
 	FMetalSamplerState* NewState = ResourceCast(NewStateRHI);
@@ -484,7 +484,7 @@ void FMetalRHICommandContext::RHISetShaderSampler(FPixelShaderRHIParamRef PixelS
 	}
 }
 
-void FMetalRHICommandContext::RHISetShaderSampler(FComputeShaderRHIParamRef ComputeShader, uint32 SamplerIndex, FSamplerStateRHIParamRef NewStateRHI)
+void FMetalRHICommandContext::RHISetShaderSampler(FComputeShaderRHIParamRef ComputeShader, uint32 SamplerIndex, FRHISamplerState* NewStateRHI)
 {
 	@autoreleasepool {
 	FMetalSamplerState* NewState = ResourceCast(NewStateRHI);
diff --git a/Engine/Source/Runtime/Apple/MetalRHI/Public/MetalRHIContext.h b/Engine/Source/Runtime/Apple/MetalRHI/Public/MetalRHIContext.h
--- a/Engine/Source/Runtime/Apple/MetalRHI/Public/MetalRHIContext.h
+++ b/Engine/Source/Runtime/Apple/MetalRHI/Public/MetalRHIContext.h
@@ -137,7 +137,7 @@ public:
 	 * @param SamplerIndex		The index of the sampler.
 	 * @param NewState			The new sampler state.
 	 */
-	virtual void RHISetShaderSampler(FComputeShaderRHIParamRef ComputeShader, uint32 SamplerIndex, FSamplerStateRHIParamRef NewState) final override;
+	virtual void RHISetShaderSampler(FComputeShaderRHIParamRef ComputeShader, uint32 SamplerIndex, FRHISamplerState* NewState) final override;
 	
 	/**
 	 * Sets sampler state.
@@ -145,7 +145,7 @@ public:
 	 * @param SamplerIndex		The index of the sampler.
 	 * @param NewState			The new sampler state.
 	 */
-	virtual void RHISetShaderSampler(FVertexShaderRHIParamRef VertexShader, uint32 SamplerIndex, FSamplerStateRHIParamRef NewState) final override;
+	virtual void RHISetShaderSampler(FVertexShaderRHIParamRef VertexShader, uint32 SamplerIndex, FRHISamplerState* NewState) final override;
 	
 	/**
 	 * Sets sampler state.
@@ -153,7 +153,7 @@ public:
 	 * @param SamplerIndex		The index of the sampler.
 	 * @param NewState			The new sampler state.
 	 */
-	virtual void RHISetShaderSampler(FGeometryShaderRHIParamRef GeometryShader, uint32 SamplerIndex, FSamplerStateRHIParamRef NewState) final override;
+	virtual void RHISetShaderSampler(FGeometryShaderRHIParamRef GeometryShader, uint32 SamplerIndex, FRHISamplerState* NewState) final override;
 	
 	/**
 	 * Sets sampler state.
@@ -161,7 +161,7 @@ public:
 	 * @param SamplerIndex		The index of the sampler.
 	 * @param NewState			The new sampler state.
 	 */
-	virtual void RHISetShaderSampler(FDomainShaderRHIParamRef DomainShader, uint32 SamplerIndex, FSamplerStateRHIParamRef NewState) final override;
+	virtual void RHISetShaderSampler(FDomainShaderRHIParamRef DomainShader, uint32 SamplerIndex, FRHISamplerState* NewState) final override;
 	
 	/**
 	 * Sets sampler state.
@@ -169,7 +169,7 @@ public:
 	 * @param SamplerIndex		The index of the sampler.
 	 * @param NewState			The new sampler state.
 	 */
-	virtual void RHISetShaderSampler(FHullShaderRHIParamRef HullShader, uint32 SamplerIndex, FSamplerStateRHIParamRef NewState) final override;
+	virtual void RHISetShaderSampler(FHullShaderRHIParamRef HullShader, uint32 SamplerIndex, FRHISamplerState* NewState) final override;
 	
 	/**
 	 * Sets sampler state.
@@ -177,7 +177,7 @@ public:
 	 * @param SamplerIndex		The index of the sampler.
 	 * @param NewState			The new sampler state.
 	 */
-	virtual void RHISetShaderSampler(FPixelShaderRHIParamRef PixelShader, uint32 SamplerIndex, FSamplerStateRHIParamRef NewState) final override;
+	virtual void RHISetShaderSampler(FPixelShaderRHIParamRef PixelShader, uint32 SamplerIndex, FRHISamplerState* NewState) final override;
 	
 	/** Sets a compute shader UAV parameter. */
 	virtual void RHISetUAVParameter(FComputeShaderRHIParamRef ComputeShader, uint32 UAVIndex, FUnorderedAccessViewRHIParamRef UAV) final override;
diff --git a/Engine/Source/Runtime/D3D12RHI/Private/D3D12CommandContext.h b/Engine/Source/Runtime/D3D12RHI/Private/D3D12CommandContext.h
--- a/Engine/Source/Runtime/D3D12RHI/Private/D3D12CommandContext.h
+++ b/Engine/Source/Runtime/D3D12RHI/Private/D3D12CommandContext.h
@@ -234,7 +234,7 @@ public:
 	virtual void RHIDispatchIndirectComputeShader(FVertexBufferRHIParamRef ArgumentBuffer, uint32 ArgumentOffset) final override;
 	virtual void RHITransitionResources(EResourceTransitionAccess TransitionType, EResourceTransitionPipeline TransitionPipeline, FUnorderedAccessViewRHIParamRef* InUAVs, int32 NumUAVs, FComputeFenceRHIParamRef WriteComputeFenceRHI) final override;
 	virtual void RHISetShaderTexture(FComputeShaderRHIParamRef PixelShader, uint32 TextureIndex, FTextureRHIParamRef NewTexture) final override;
-	virtual void RHISetShaderSampler(FComputeShaderRHIParamRef ComputeShader, uint32 SamplerIndex, FSamplerStateRHIParamRef NewState) final override;
+	virtual void RHISetShaderSampler(FComputeShaderRHIParamRef ComputeShader, uint32 SamplerIndex, FRHISamplerState* NewState) final override;
 	virtual void RHISetUAVParameter(FComputeShaderRHIParamRef ComputeShader, uint32 UAVIndex, FUnorderedAccessViewRHIParamRef UAV) final override;
 	virtual void RHISetUAVParameter(FComputeShaderRHIParamRef ComputeShader, uint32 UAVIndex, FUnorderedAccessViewRHIParamRef UAV, uint32 InitialCount) final override;
 	virtual void RHISetShaderResourceViewParameter(FComputeShaderRHIParamRef ComputeShader, uint32 SamplerIndex, FShaderResourceViewRHIParamRef SRV) final override;
@@ -269,11 +269,11 @@ public:
 	virtual void RHISetShaderTexture(FDomainShaderRHIParamRef DomainShader, uint32 TextureIndex, FTextureRHIParamRef NewTexture) final override;
 	virtual void RHISetShaderTexture(FGeometryShaderRHIParamRef GeometryShader, uint32 TextureIndex, FTextureRHIParamRef NewTexture) final override;
 	virtual void RHISetShaderTexture(FPixelShaderRHIParamRef PixelShader, uint32 TextureIndex, FTextureRHIParamRef NewTexture) final override;
-	virtual void RHISetShaderSampler(FVertexShaderRHIParamRef VertexShader, uint32 SamplerIndex, FSamplerStateRHIParamRef NewState) final override;
-	virtual void RHISetShaderSampler(FGeometryShaderRHIParamRef GeometryShader, uint32 SamplerIndex, FSamplerStateRHIParamRef NewState) final override;
-	virtual void RHISetShaderSampler(FDomainShaderRHIParamRef DomainShader, uint32 SamplerIndex, FSamplerStateRHIParamRef NewState) final override;
-	virtual void RHISetShaderSampler(FHullShaderRHIParamRef HullShader, uint32 SamplerIndex, FSamplerStateRHIParamRef NewState) final override;
-	virtual void RHISetShaderSampler(FPixelShaderRHIParamRef PixelShader, uint32 SamplerIndex, FSamplerStateRHIParamRef NewState) final override;
+	virtual void RHISetShaderSampler(FVertexShaderRHIParamRef VertexShader, uint32 SamplerIndex, FRHISamplerState* NewState) final override;
+	virtual void RHISetShaderSampler(FGeometryShaderRHIParamRef GeometryShader, uint32 SamplerIndex, FRHISamplerState* NewState) final override;
+	virtual void RHISetShaderSampler(FDomainShaderRHIParamRef DomainShader, uint32 SamplerIndex, FRHISamplerState* NewState) final override;
+	virtual void RHISetShaderSampler(FHullShaderRHIParamRef HullShader, uint32 SamplerIndex, FRHISamplerState* NewState) final override;
+	virtual void RHISetShaderSampler(FPixelShaderRHIParamRef PixelShader, uint32 SamplerIndex, FRHISamplerState* NewState) final override;
 	virtual void RHISetShaderResourceViewParameter(FPixelShaderRHIParamRef PixelShader, uint32 SamplerIndex, FShaderResourceViewRHIParamRef SRV) final override;
 	virtual void RHISetShaderResourceViewParameter(FVertexShaderRHIParamRef VertexShader, uint32 SamplerIndex, FShaderResourceViewRHIParamRef SRV) final override;
 	virtual void RHISetShaderResourceViewParameter(FHullShaderRHIParamRef HullShader, uint32 SamplerIndex, FShaderResourceViewRHIParamRef SRV) final override;
@@ -474,7 +474,7 @@ public:
 	{
 		ContextRedirect(RHISetShaderTexture(PixelShader, TextureIndex, NewTexture));
 	}
-	FORCEINLINE virtual void RHISetShaderSampler(FComputeShaderRHIParamRef ComputeShader, uint32 SamplerIndex, FSamplerStateRHIParamRef NewState) final override
+	FORCEINLINE virtual void RHISetShaderSampler(FComputeShaderRHIParamRef ComputeShader, uint32 SamplerIndex, FRHISamplerState* NewState) final override
 	{
 		ContextRedirect(RHISetShaderSampler(ComputeShader, SamplerIndex, NewState));
 	}
@@ -588,23 +588,23 @@ public:
 	{
 		ContextRedirect(RHISetShaderTexture(PixelShader, TextureIndex, NewTexture));
 	}
-	FORCEINLINE virtual void RHISetShaderSampler(FVertexShaderRHIParamRef VertexShader, uint32 SamplerIndex, FSamplerStateRHIParamRef NewState) final override
+	FORCEINLINE virtual void RHISetShaderSampler(FVertexShaderRHIParamRef VertexShader, uint32 SamplerIndex, FRHISamplerState* NewState) final override
 	{
 		ContextRedirect(RHISetShaderSampler(VertexShader, SamplerIndex, NewState));
 	}
-	FORCEINLINE virtual void RHISetShaderSampler(FGeometryShaderRHIParamRef GeometryShader, uint32 SamplerIndex, FSamplerStateRHIParamRef NewState) final override
+	FORCEINLINE virtual void RHISetShaderSampler(FGeometryShaderRHIParamRef GeometryShader, uint32 SamplerIndex, FRHISamplerState* NewState) final override
 	{
 		ContextRedirect(RHISetShaderSampler(GeometryShader, SamplerIndex, NewState));
 	}
-	FORCEINLINE virtual void RHISetShaderSampler(FDomainShaderRHIParamRef DomainShader, uint32 SamplerIndex, FSamplerStateRHIParamRef NewState) final override
+	FORCEINLINE virtual void RHISetShaderSampler(FDomainShaderRHIParamRef DomainShader, uint32 SamplerIndex, FRHISamplerState* NewState) final override
 	{
 		ContextRedirect(RHISetShaderSampler(DomainShader, SamplerIndex, NewState));
 	}
-	FORCEINLINE virtual void RHISetShaderSampler(FHullShaderRHIParamRef HullShader, uint32 SamplerIndex, FSamplerStateRHIParamRef NewState) final override
+	FORCEINLINE virtual void RHISetShaderSampler(FHullShaderRHIParamRef HullShader, uint32 SamplerIndex, FRHISamplerState* NewState) final override
 	{
 		ContextRedirect(RHISetShaderSampler(HullShader, SamplerIndex, NewState));
 	}
-	FORCEINLINE virtual void RHISetShaderSampler(FPixelShaderRHIParamRef PixelShader, uint32 SamplerIndex, FSamplerStateRHIParamRef NewState) final override
+	FORCEINLINE virtual void RHISetShaderSampler(FPixelShaderRHIParamRef PixelShader, uint32 SamplerIndex, FRHISamplerState* NewState) final override
 	{
 		ContextRedirect(RHISetShaderSampler(PixelShader, SamplerIndex, NewState));
 	}
diff --git a/Engine/Source/Runtime/D3D12RHI/Private/D3D12Commands.cpp b/Engine/Source/Runtime/D3D12RHI/Private/D3D12Commands.cpp
--- a/Engine/Source/Runtime/D3D12RHI/Private/D3D12Commands.cpp
+++ b/Engine/Source/Runtime/D3D12RHI/Private/D3D12Commands.cpp
@@ -702,42 +702,42 @@ void FD3D12CommandContext::RHISetShaderResourceViewParameter(FGeometryShaderRHIP
 	StateCache.SetShaderResourceView<SF_Geometry>(SRV, TextureIndex);
 }
 
-void FD3D12CommandContext::RHISetShaderSampler(FVertexShaderRHIParamRef VertexShaderRHI, uint32 SamplerIndex, FSamplerStateRHIParamRef NewStateRHI)
+void FD3D12CommandContext::RHISetShaderSampler(FVertexShaderRHIParamRef VertexShaderRHI, uint32 SamplerIndex, FRHISamplerState* NewStateRHI)
 {
 	VALIDATE_BOUND_SHADER(VertexShaderRHI);
 	FD3D12SamplerState* NewState = RetrieveObject<FD3D12SamplerState>(NewStateRHI);
 	StateCache.SetSamplerState<SF_Vertex>(NewState, SamplerIndex);
 }
 
-void FD3D12CommandContext::RHISetShaderSampler(FHullShaderRHIParamRef HullShaderRHI, uint32 SamplerIndex, FSamplerStateRHIParamRef NewStateRHI)
+void FD3D12CommandContext::RHISetShaderSampler(FHullShaderRHIParamRef HullShaderRHI, uint32 SamplerIndex, FRHISamplerState* NewStateRHI)
 {
 	VALIDATE_BOUND_SHADER(HullShaderRHI);
 	FD3D12SamplerState* NewState = RetrieveObject<FD3D12SamplerState>(NewStateRHI);
 	StateCache.SetSamplerState<SF_Hull>(NewState, SamplerIndex);
 }
 
-void FD3D12CommandContext::RHISetShaderSampler(FDomainShaderRHIParamRef DomainShaderRHI, uint32 SamplerIndex, FSamplerStateRHIParamRef NewStateRHI)
+void FD3D12CommandContext::RHISetShaderSampler(FDomainShaderRHIParamRef DomainShaderRHI, uint32 SamplerIndex, FRHISamplerState* NewStateRHI)
 {
 	VALIDATE_BOUND_SHADER(DomainShaderRHI);
 	FD3D12SamplerState* NewState = RetrieveObject<FD3D12SamplerState>(NewStateRHI);
 	StateCache.SetSamplerState<SF_Domain>(NewState, SamplerIndex);
 }
 
-void FD3D12CommandContext::RHISetShaderSampler(FGeometryShaderRHIParamRef GeometryShaderRHI, uint32 SamplerIndex, FSamplerStateRHIParamRef NewStateRHI)
+void FD3D12CommandContext::RHISetShaderSampler(FGeometryShaderRHIParamRef GeometryShaderRHI, uint32 SamplerIndex, FRHISamplerState* NewStateRHI)
 {
 	VALIDATE_BOUND_SHADER(GeometryShaderRHI);
 	FD3D12SamplerState* NewState = RetrieveObject<FD3D12SamplerState>(NewStateRHI);
 	StateCache.SetSamplerState<SF_Geometry>(NewState, SamplerIndex);
 }
 
-void FD3D12CommandContext::RHISetShaderSampler(FPixelShaderRHIParamRef PixelShaderRHI, uint32 SamplerIndex, FSamplerStateRHIParamRef NewStateRHI)
+void FD3D12CommandContext::RHISetShaderSampler(FPixelShaderRHIParamRef PixelShaderRHI, uint32 SamplerIndex, FRHISamplerState* NewStateRHI)
 {
 	VALIDATE_BOUND_SHADER(PixelShaderRHI);
 	FD3D12SamplerState* NewState = RetrieveObject<FD3D12SamplerState>(NewStateRHI);
 	StateCache.SetSamplerState<SF_Pixel>(NewState, SamplerIndex);
 }
 
-void FD3D12CommandContext::RHISetShaderSampler(FComputeShaderRHIParamRef ComputeShaderRHI, uint32 SamplerIndex, FSamplerStateRHIParamRef NewStateRHI)
+void FD3D12CommandContext::RHISetShaderSampler(FComputeShaderRHIParamRef ComputeShaderRHI, uint32 SamplerIndex, FRHISamplerState* NewStateRHI)
 {
 	//VALIDATE_BOUND_SHADER(ComputeShaderRHI);
 	FD3D12SamplerState* NewState = RetrieveObject<FD3D12SamplerState>(NewStateRHI);
diff --git a/Engine/Source/Runtime/D3D12RHI/Private/D3D12RayTracing.cpp b/Engine/Source/Runtime/D3D12RHI/Private/D3D12RayTracing.cpp
--- a/Engine/Source/Runtime/D3D12RHI/Private/D3D12RayTracing.cpp
+++ b/Engine/Source/Runtime/D3D12RHI/Private/D3D12RayTracing.cpp
@@ -2520,7 +2520,7 @@ static void SetRayTracingShaderResources(
 	uint32 InNumTextures, const FTextureRHIParamRef* Textures,
 	uint32 InNumSRVs, const FShaderResourceViewRHIParamRef* SRVs,
 	uint32 InNumUniformBuffers, const FUniformBufferRHIParamRef* UniformBuffers,
-	uint32 InNumSamplers, const FSamplerStateRHIParamRef* Samplers,
+	uint32 InNumSamplers, FRHISamplerState* const* Samplers,
 	uint32 InNumUAVs, const FUnorderedAccessViewRHIParamRef* UAVs,
 	uint32 InLooseParameterDataSize, const void* InLooseParameterData,
 	FD3D12RayTracingDescriptorCache& DescriptorCache,
@@ -2588,7 +2588,7 @@ static void SetRayTracingShaderResources(
 
 	for (uint32 SamplerIndex = 0; SamplerIndex < InNumSamplers; ++SamplerIndex)
 	{
-		FSamplerStateRHIParamRef Resource = Samplers[SamplerIndex];
+		FRHISamplerState* Resource = Samplers[SamplerIndex];
 		if (Resource)
 		{
 			LocalSamplers[SamplerIndex] = CommandContext.RetrieveObject<FD3D12SamplerState>(Resource)->Descriptor;
diff --git a/Engine/Source/Runtime/EmptyRHI/Private/EmptyCommands.cpp b/Engine/Source/Runtime/EmptyRHI/Private/EmptyCommands.cpp
--- a/Engine/Source/Runtime/EmptyRHI/Private/EmptyCommands.cpp
+++ b/Engine/Source/Runtime/EmptyRHI/Private/EmptyCommands.cpp
@@ -137,38 +137,38 @@ void FEmptyDynamicRHI::RHISetShaderResourceViewParameter(FComputeShaderRHIParamR
 }
 
 
-void FEmptyDynamicRHI::RHISetShaderSampler(FVertexShaderRHIParamRef VertexShaderRHI, uint32 SamplerIndex, FSamplerStateRHIParamRef NewStateRHI)
+void FEmptyDynamicRHI::RHISetShaderSampler(FVertexShaderRHIParamRef VertexShaderRHI, uint32 SamplerIndex, FRHISamplerState* NewStateRHI)
 {
 	FEmptySamplerState* NewState = ResourceCast(NewStateRHI);
 	FEmptyVertexShader* VertexShader = ResourceCast(VertexShaderRHI);
 
 }
 
-void FEmptyDynamicRHI::RHISetShaderSampler(FHullShaderRHIParamRef HullShader, uint32 SamplerIndex, FSamplerStateRHIParamRef NewStateRHI)
+void FEmptyDynamicRHI::RHISetShaderSampler(FHullShaderRHIParamRef HullShader, uint32 SamplerIndex, FRHISamplerState* NewStateRHI)
 {
 	FEmptySamplerState* NewState = ResourceCast(NewStateRHI);
 
 }
 
-void FEmptyDynamicRHI::RHISetShaderSampler(FDomainShaderRHIParamRef DomainShader, uint32 SamplerIndex, FSamplerStateRHIParamRef NewStateRHI)
+void FEmptyDynamicRHI::RHISetShaderSampler(FDomainShaderRHIParamRef DomainShader, uint32 SamplerIndex, FRHISamplerState* NewStateRHI)
 {
 	FEmptySamplerState* NewState = ResourceCast(NewStateRHI);
 
 }
 
-void FEmptyDynamicRHI::RHISetShaderSampler(FGeometryShaderRHIParamRef GeometryShader, uint32 SamplerIndex, FSamplerStateRHIParamRef NewStateRHI)
+void FEmptyDynamicRHI::RHISetShaderSampler(FGeometryShaderRHIParamRef GeometryShader, uint32 SamplerIndex, FRHISamplerState* NewStateRHI)
 {
 	FEmptySamplerState* NewState = ResourceCast(NewStateRHI);
 
 }
 
-void FEmptyDynamicRHI::RHISetShaderSampler(FPixelShaderRHIParamRef PixelShader, uint32 SamplerIndex, FSamplerStateRHIParamRef NewStateRHI)
+void FEmptyDynamicRHI::RHISetShaderSampler(FPixelShaderRHIParamRef PixelShader, uint32 SamplerIndex, FRHISamplerState* NewStateRHI)
 {
 	FEmptySamplerState* NewState = ResourceCast(NewStateRHI);
 
 }
 
-void FEmptyDynamicRHI::RHISetShaderSampler(FComputeShaderRHIParamRef ComputeShader, uint32 SamplerIndex, FSamplerStateRHIParamRef NewStateRHI)
+void FEmptyDynamicRHI::RHISetShaderSampler(FComputeShaderRHIParamRef ComputeShader, uint32 SamplerIndex, FRHISamplerState* NewStateRHI)
 {
 	FEmptySamplerState* NewState = ResourceCast(NewStateRHI);
 
diff --git a/Engine/Source/Runtime/EmptyRHI/Public/EmptyRHI.h b/Engine/Source/Runtime/EmptyRHI/Public/EmptyRHI.h
--- a/Engine/Source/Runtime/EmptyRHI/Public/EmptyRHI.h
+++ b/Engine/Source/Runtime/EmptyRHI/Public/EmptyRHI.h
@@ -161,12 +161,12 @@ public:
 	virtual void RHISetShaderTexture(FGeometryShaderRHIParamRef GeometryShader, uint32 TextureIndex, FTextureRHIParamRef NewTexture) final override;
 	virtual void RHISetShaderTexture(FPixelShaderRHIParamRef PixelShader, uint32 TextureIndex, FTextureRHIParamRef NewTexture) final override;
 	virtual void RHISetShaderTexture(FComputeShaderRHIParamRef PixelShader, uint32 TextureIndex, FTextureRHIParamRef NewTexture) final override;
-	virtual void RHISetShaderSampler(FComputeShaderRHIParamRef ComputeShader, uint32 SamplerIndex, FSamplerStateRHIParamRef NewState) final override;
-	virtual void RHISetShaderSampler(FVertexShaderRHIParamRef VertexShader, uint32 SamplerIndex, FSamplerStateRHIParamRef NewState) final override;
-	virtual void RHISetShaderSampler(FGeometryShaderRHIParamRef GeometryShader, uint32 SamplerIndex, FSamplerStateRHIParamRef NewState) final override;
-	virtual void RHISetShaderSampler(FDomainShaderRHIParamRef DomainShader, uint32 SamplerIndex, FSamplerStateRHIParamRef NewState) final override;
-	virtual void RHISetShaderSampler(FHullShaderRHIParamRef HullShader, uint32 SamplerIndex, FSamplerStateRHIParamRef NewState) final override;
-	virtual void RHISetShaderSampler(FPixelShaderRHIParamRef PixelShader, uint32 SamplerIndex, FSamplerStateRHIParamRef NewState) final override;
+	virtual void RHISetShaderSampler(FComputeShaderRHIParamRef ComputeShader, uint32 SamplerIndex, FRHISamplerState* NewState) final override;
+	virtual void RHISetShaderSampler(FVertexShaderRHIParamRef VertexShader, uint32 SamplerIndex, FRHISamplerState* NewState) final override;
+	virtual void RHISetShaderSampler(FGeometryShaderRHIParamRef GeometryShader, uint32 SamplerIndex, FRHISamplerState* NewState) final override;
+	virtual void RHISetShaderSampler(FDomainShaderRHIParamRef DomainShader, uint32 SamplerIndex, FRHISamplerState* NewState) final override;
+	virtual void RHISetShaderSampler(FHullShaderRHIParamRef HullShader, uint32 SamplerIndex, FRHISamplerState* NewState) final override;
+	virtual void RHISetShaderSampler(FPixelShaderRHIParamRef PixelShader, uint32 SamplerIndex, FRHISamplerState* NewState) final override;
 	virtual void RHISetUAVParameter(FComputeShaderRHIParamRef ComputeShader, uint32 UAVIndex, FUnorderedAccessViewRHIParamRef UAV) final override;
 	virtual void RHISetUAVParameter(FComputeShaderRHIParamRef ComputeShader, uint32 UAVIndex, FUnorderedAccessViewRHIParamRef UAV, uint32 InitialCount) final override;
 	virtual void RHISetShaderResourceViewParameter(FPixelShaderRHIParamRef PixelShader, uint32 SamplerIndex, FShaderResourceViewRHIParamRef SRV) final override;
diff --git a/Engine/Source/Runtime/Engine/Private/Particles/ParticleGpuSimulation.cpp b/Engine/Source/Runtime/Engine/Private/Particles/ParticleGpuSimulation.cpp
--- a/Engine/Source/Runtime/Engine/Private/Particles/ParticleGpuSimulation.cpp
+++ b/Engine/Source/Runtime/Engine/Private/Particles/ParticleGpuSimulation.cpp
@@ -797,8 +797,8 @@ void FGPUSpriteVertexFactoryShaderParametersVS::GetElementShaderBindings(
 	FVertexInputStreamArray& VertexStreams) const 
 {
 	FGPUSpriteVertexFactory* GPUVF = (FGPUSpriteVertexFactory*)VertexFactory;
-	FSamplerStateRHIParamRef SamplerStatePoint = TStaticSamplerState<SF_Point>::GetRHI();
-	FSamplerStateRHIParamRef SamplerStateLinear = TStaticSamplerState<SF_Bilinear>::GetRHI();
+	FRHISamplerState* SamplerStatePoint = TStaticSamplerState<SF_Point>::GetRHI();
+	FRHISamplerState* SamplerStateLinear = TStaticSamplerState<SF_Bilinear>::GetRHI();
 	ShaderBindings.Add(Shader->GetUniformBufferParameter<FGPUSpriteEmitterUniformParameters>(), GPUVF->EmitterUniformBuffer);
 	ShaderBindings.Add(Shader->GetUniformBufferParameter<FGPUSpriteEmitterDynamicUniformParameters>(), GPUVF->EmitterDynamicUniformBuffer);
 	FShaderResourceViewRHIParamRef ParticleIndicesBuffer = GPUVF->ParticleIndicesBuffer->VertexBufferSRV;
@@ -1144,8 +1144,8 @@ public:
 		)
 	{
 		FPixelShaderRHIParamRef PixelShaderRHI = GetPixelShader();
-		FSamplerStateRHIParamRef SamplerStatePoint = TStaticSamplerState<SF_Point>::GetRHI();
-		FSamplerStateRHIParamRef SamplerStateLinear = TStaticSamplerState<SF_Bilinear,AM_Clamp,AM_Clamp,AM_Clamp>::GetRHI();
+		FRHISamplerState* SamplerStatePoint = TStaticSamplerState<SF_Point>::GetRHI();
+		FRHISamplerState* SamplerStateLinear = TStaticSamplerState<SF_Bilinear,AM_Clamp,AM_Clamp,AM_Clamp>::GetRHI();
 		SetTextureParameter(RHICmdList, PixelShaderRHI, PositionTexture, PositionTextureSampler, SamplerStatePoint, TextureResources.PositionTextureRHI);
 		SetTextureParameter(RHICmdList, PixelShaderRHI, VelocityTexture, VelocityTextureSampler, SamplerStatePoint, TextureResources.VelocityTextureRHI);
 		SetTextureParameter(RHICmdList, PixelShaderRHI, AttributesTexture, AttributesTextureSampler, SamplerStatePoint, InAttributesTexture.TextureRHI);
@@ -1191,7 +1191,7 @@ public:
 		FPixelShaderRHIParamRef PixelShaderRHI = GetPixelShader();
 			SetUniformBufferParameter(RHICmdList, PixelShaderRHI, GetUniformBufferParameter<FVectorFieldUniformParameters>(), UniformBuffer);
 		
-		FSamplerStateRHIParamRef SamplerStateLinear = TStaticSamplerState<SF_Bilinear,AM_Clamp,AM_Clamp,AM_Clamp>::GetRHI();
+			FRHISamplerState* SamplerStateLinear = TStaticSamplerState<SF_Bilinear,AM_Clamp,AM_Clamp,AM_Clamp>::GetRHI();
 
 		for (int32 i = 0; i < MAX_VECTOR_FIELDS; ++i)
 		{
@@ -2031,7 +2031,7 @@ public:
 	{
 		FPixelShaderRHIParamRef PixelShader = GetPixelShader();
 		SetShaderValue(RHICmdList, PixelShader, VisualizationMode, InVisualizationMode );
-		FSamplerStateRHIParamRef SamplerStatePoint = TStaticSamplerState<SF_Point>::GetRHI();
+		FRHISamplerState* SamplerStatePoint = TStaticSamplerState<SF_Point>::GetRHI();
 		SetTextureParameter(RHICmdList, PixelShader, PositionTexture, PositionTextureSampler, SamplerStatePoint, PositionTextureRHI );
 		SetTextureParameter(RHICmdList, PixelShader, CurveTexture, CurveTextureSampler, SamplerStatePoint, CurveTextureRHI );
 	}
diff --git a/Engine/Source/Runtime/Engine/Private/VectorField.cpp b/Engine/Source/Runtime/Engine/Private/VectorField.cpp
--- a/Engine/Source/Runtime/Engine/Private/VectorField.cpp
+++ b/Engine/Source/Runtime/Engine/Private/VectorField.cpp
@@ -812,7 +812,7 @@ public:
 		FTextureRHIParamRef NoiseVolumeTextureRHI )
 	{
 		FComputeShaderRHIParamRef ComputeShaderRHI = GetComputeShader();
-		FSamplerStateRHIParamRef SamplerStateLinear = TStaticSamplerState<SF_Bilinear,AM_Clamp,AM_Clamp,AM_Clamp>::GetRHI();
+		FRHISamplerState* SamplerStateLinear = TStaticSamplerState<SF_Bilinear,AM_Clamp,AM_Clamp,AM_Clamp>::GetRHI();
 		SetUniformBufferParameter(RHICmdList, ComputeShaderRHI, GetUniformBufferParameter<FCompositeAnimatedVectorFieldUniformParameters>(), UniformBuffer );
 		SetTextureParameter(RHICmdList, ComputeShaderRHI, AtlasTexture, AtlasTextureSampler, SamplerStateLinear, AtlasTextureRHI );
 		SetTextureParameter(RHICmdList, ComputeShaderRHI, NoiseVolumeTexture, NoiseVolumeTextureSampler, SamplerStateLinear, NoiseVolumeTextureRHI );
diff --git a/Engine/Source/Runtime/Engine/Private/VectorFieldVisualization.cpp b/Engine/Source/Runtime/Engine/Private/VectorFieldVisualization.cpp
--- a/Engine/Source/Runtime/Engine/Private/VectorFieldVisualization.cpp
+++ b/Engine/Source/Runtime/Engine/Private/VectorFieldVisualization.cpp
@@ -183,7 +183,7 @@ void FVectorFieldVisualizationVertexFactoryShaderParameters::GetElementShaderBin
 	FVertexInputStreamArray& VertexStreams) const
 {
 	FVectorFieldVisualizationVertexFactory* VertexFactory = (FVectorFieldVisualizationVertexFactory*)InVertexFactory;
-	FSamplerStateRHIParamRef SamplerStatePoint = TStaticSamplerState<SF_Point>::GetRHI();
+	FRHISamplerState* SamplerStatePoint = TStaticSamplerState<SF_Point>::GetRHI();
 	ShaderBindings.Add(Shader->GetUniformBufferParameter<FVectorFieldVisualizationParameters>(), VertexFactory->UniformBuffer);
 	ShaderBindings.AddTexture(VectorFieldTexture, VectorFieldTextureSampler, SamplerStatePoint, VertexFactory->VectorFieldTextureRHI);
 }
diff --git a/Engine/Source/Runtime/Engine/Public/ScreenRendering.h b/Engine/Source/Runtime/Engine/Public/ScreenRendering.h
--- a/Engine/Source/Runtime/Engine/Public/ScreenRendering.h
+++ b/Engine/Source/Runtime/Engine/Public/ScreenRendering.h
@@ -80,7 +80,7 @@ public:
 		SetTextureParameter(RHICmdList, GetPixelShader(),InTexture,InTextureSampler,Texture);
 	}
 
-	void SetParameters(FRHICommandList& RHICmdList, FSamplerStateRHIParamRef SamplerStateRHI, FTextureRHIParamRef TextureRHI)
+	void SetParameters(FRHICommandList& RHICmdList, FRHISamplerState* SamplerStateRHI, FTextureRHIParamRef TextureRHI)
 	{
 		SetTextureParameter(RHICmdList, GetPixelShader(),InTexture,InTextureSampler,SamplerStateRHI,TextureRHI);
 	}
@@ -121,7 +121,7 @@ public:
 		SetTextureParameter(RHICmdList, GetPixelShader(), InTexture, InTextureSampler, Texture);
 	}
 
-	void SetParameters(FRHICommandList& RHICmdList, FSamplerStateRHIParamRef SamplerStateRHI, FTextureRHIParamRef TextureRHI)
+	void SetParameters(FRHICommandList& RHICmdList, FRHISamplerState* SamplerStateRHI, FTextureRHIParamRef TextureRHI)
 	{
 		SetTextureParameter(RHICmdList, GetPixelShader(), InTexture, InTextureSampler, SamplerStateRHI, TextureRHI);
 	}
@@ -160,7 +160,7 @@ public:
         SetTextureParameter(RHICmdList, GetPixelShader(),InTexture,InTextureSampler,Texture);
     }
 
-    void SetParameters(FRHICommandList& RHICmdList, FSamplerStateRHIParamRef SamplerStateRHI, FTextureRHIParamRef TextureRHI)
+    void SetParameters(FRHICommandList& RHICmdList, FRHISamplerState* SamplerStateRHI, FTextureRHIParamRef TextureRHI)
     {
         SetTextureParameter(RHICmdList, GetPixelShader(),InTexture,InTextureSampler,SamplerStateRHI,TextureRHI);
     }
diff --git a/Engine/Source/Runtime/Engine/Public/StereoLayerRendering.h b/Engine/Source/Runtime/Engine/Public/StereoLayerRendering.h
--- a/Engine/Source/Runtime/Engine/Public/StereoLayerRendering.h
+++ b/Engine/Source/Runtime/Engine/Public/StereoLayerRendering.h
@@ -97,7 +97,7 @@ public:
 	}
 	FStereoLayerPS() {}
 
-	void SetParameters(FRHICommandList& RHICmdList, FSamplerStateRHIParamRef SamplerStateRHI, FTextureRHIParamRef TextureRHI)
+	void SetParameters(FRHICommandList& RHICmdList, FRHISamplerState* SamplerStateRHI, FTextureRHIParamRef TextureRHI)
 	{
 		FPixelShaderRHIParamRef PS = GetPixelShader();
 
diff --git a/Engine/Source/Runtime/GameplayMediaEncoder/Private/Microsoft/Windows/UtilityShaders.h b/Engine/Source/Runtime/GameplayMediaEncoder/Private/Microsoft/Windows/UtilityShaders.h
--- a/Engine/Source/Runtime/GameplayMediaEncoder/Private/Microsoft/Windows/UtilityShaders.h
+++ b/Engine/Source/Runtime/GameplayMediaEncoder/Private/Microsoft/Windows/UtilityShaders.h
@@ -38,7 +38,7 @@ public:
 		SetTextureParameter(RHICmdList, GetPixelShader(), InTexture, InTextureSampler, Texture);
 	}
 
-	void SetParameters(FRHICommandList& RHICmdList, FSamplerStateRHIParamRef SamplerStateRHI, FTextureRHIParamRef TextureRHI)
+	void SetParameters(FRHICommandList& RHICmdList, FRHISamplerState* SamplerStateRHI, FTextureRHIParamRef TextureRHI)
 	{
 		SetTextureParameter(RHICmdList, GetPixelShader(), InTexture, InTextureSampler, SamplerStateRHI, TextureRHI);
 	}
diff --git a/Engine/Source/Runtime/NullDrv/Public/NullRHI.h b/Engine/Source/Runtime/NullDrv/Public/NullRHI.h
--- a/Engine/Source/Runtime/NullDrv/Public/NullRHI.h
+++ b/Engine/Source/Runtime/NullDrv/Public/NullRHI.h
@@ -579,30 +579,30 @@ public:
 
 	}
 
-	virtual void RHISetShaderSampler(FComputeShaderRHIParamRef ComputeShader, uint32 SamplerIndex, FSamplerStateRHIParamRef NewState) final override
+	virtual void RHISetShaderSampler(FComputeShaderRHIParamRef ComputeShader, uint32 SamplerIndex, FRHISamplerState* NewState) final override
 	{
 
 	}
 
-	virtual void RHISetShaderSampler(FVertexShaderRHIParamRef VertexShader, uint32 SamplerIndex, FSamplerStateRHIParamRef NewState) final override
+	virtual void RHISetShaderSampler(FVertexShaderRHIParamRef VertexShader, uint32 SamplerIndex, FRHISamplerState* NewState) final override
 	{
 
 	}
 
-	virtual void RHISetShaderSampler(FGeometryShaderRHIParamRef GeometryShader, uint32 SamplerIndex, FSamplerStateRHIParamRef NewState) final override
+	virtual void RHISetShaderSampler(FGeometryShaderRHIParamRef GeometryShader, uint32 SamplerIndex, FRHISamplerState* NewState) final override
 	{
 
 	}
-	virtual void RHISetShaderSampler(FDomainShaderRHIParamRef DomainShader, uint32 SamplerIndex, FSamplerStateRHIParamRef NewState) final override
+	virtual void RHISetShaderSampler(FDomainShaderRHIParamRef DomainShader, uint32 SamplerIndex, FRHISamplerState* NewState) final override
 	{
 
 	}
-	virtual void RHISetShaderSampler(FHullShaderRHIParamRef HullShader, uint32 SamplerIndex, FSamplerStateRHIParamRef NewState) final override
+	virtual void RHISetShaderSampler(FHullShaderRHIParamRef HullShader, uint32 SamplerIndex, FRHISamplerState* NewState) final override
 	{
 
 	}
 
-	virtual void RHISetShaderSampler(FPixelShaderRHIParamRef PixelShader, uint32 SamplerIndex, FSamplerStateRHIParamRef NewState) final override
+	virtual void RHISetShaderSampler(FPixelShaderRHIParamRef PixelShader, uint32 SamplerIndex, FRHISamplerState* NewState) final override
 	{
 
 	}
diff --git a/Engine/Source/Runtime/OpenGLDrv/Private/OpenGLCommands.cpp b/Engine/Source/Runtime/OpenGLDrv/Private/OpenGLCommands.cpp
--- a/Engine/Source/Runtime/OpenGLDrv/Private/OpenGLCommands.cpp
+++ b/Engine/Source/Runtime/OpenGLDrv/Private/OpenGLCommands.cpp
@@ -1065,7 +1065,7 @@ void FOpenGLDynamicRHI::RHISetShaderTexture(FPixelShaderRHIParamRef PixelShaderR
 	}
 }
 
-void FOpenGLDynamicRHI::RHISetShaderSampler(FVertexShaderRHIParamRef VertexShaderRHI,uint32 SamplerIndex,FSamplerStateRHIParamRef NewStateRHI)
+void FOpenGLDynamicRHI::RHISetShaderSampler(FVertexShaderRHIParamRef VertexShaderRHI,uint32 SamplerIndex, FRHISamplerState* NewStateRHI)
 {
 	VALIDATE_BOUND_SHADER(VertexShaderRHI);
 
@@ -1074,7 +1074,7 @@ void FOpenGLDynamicRHI::RHISetShaderSampler(FVertexShaderRHIParamRef VertexShade
 	InternalSetSamplerStates(FOpenGL::GetFirstVertexTextureUnit() + SamplerIndex, NewState);
 }
 
-void FOpenGLDynamicRHI::RHISetShaderSampler(FHullShaderRHIParamRef HullShaderRHI,uint32 SamplerIndex,FSamplerStateRHIParamRef NewStateRHI)
+void FOpenGLDynamicRHI::RHISetShaderSampler(FHullShaderRHIParamRef HullShaderRHI,uint32 SamplerIndex, FRHISamplerState* NewStateRHI)
 {
 	VALIDATE_BOUND_SHADER(HullShaderRHI);
 
@@ -1083,7 +1083,7 @@ void FOpenGLDynamicRHI::RHISetShaderSampler(FHullShaderRHIParamRef HullShaderRHI
 	FOpenGLSamplerState* NewState = ResourceCast(NewStateRHI);
 	InternalSetSamplerStates(FOpenGL::GetFirstHullTextureUnit() + SamplerIndex, NewState);
 }
-void FOpenGLDynamicRHI::RHISetShaderSampler(FDomainShaderRHIParamRef DomainShaderRHI,uint32 SamplerIndex,FSamplerStateRHIParamRef NewStateRHI)
+void FOpenGLDynamicRHI::RHISetShaderSampler(FDomainShaderRHIParamRef DomainShaderRHI,uint32 SamplerIndex, FRHISamplerState* NewStateRHI)
 {
 	VALIDATE_BOUND_SHADER(DomainShaderRHI);
 
@@ -1093,7 +1093,7 @@ void FOpenGLDynamicRHI::RHISetShaderSampler(FDomainShaderRHIParamRef DomainShade
 	InternalSetSamplerStates(FOpenGL::GetFirstDomainTextureUnit() + SamplerIndex, NewState);
 }
 
-void FOpenGLDynamicRHI::RHISetShaderSampler(FGeometryShaderRHIParamRef GeometryShaderRHI,uint32 SamplerIndex,FSamplerStateRHIParamRef NewStateRHI)
+void FOpenGLDynamicRHI::RHISetShaderSampler(FGeometryShaderRHIParamRef GeometryShaderRHI,uint32 SamplerIndex, FRHISamplerState* NewStateRHI)
 {
 	VALIDATE_BOUND_SHADER(GeometryShaderRHI);
 
@@ -1118,7 +1118,7 @@ void FOpenGLDynamicRHI::RHISetShaderTexture(FComputeShaderRHIParamRef ComputeSha
 	}
 }
 
-void FOpenGLDynamicRHI::RHISetShaderSampler(FPixelShaderRHIParamRef PixelShaderRHI,uint32 SamplerIndex,FSamplerStateRHIParamRef NewStateRHI)
+void FOpenGLDynamicRHI::RHISetShaderSampler(FPixelShaderRHIParamRef PixelShaderRHI,uint32 SamplerIndex, FRHISamplerState* NewStateRHI)
 {
 	VALIDATE_BOUND_SHADER(PixelShaderRHI);
 
@@ -1165,7 +1165,7 @@ void FOpenGLDynamicRHI::RHISetShaderUniformBuffer(FGeometryShaderRHIParamRef Geo
 	PendingState.bAnyDirtyGraphicsUniformBuffers = true;
 }
 
-void FOpenGLDynamicRHI::RHISetShaderSampler(FComputeShaderRHIParamRef ComputeShaderRHI,uint32 SamplerIndex,FSamplerStateRHIParamRef NewStateRHI)
+void FOpenGLDynamicRHI::RHISetShaderSampler(FComputeShaderRHIParamRef ComputeShaderRHI,uint32 SamplerIndex, FRHISamplerState* NewStateRHI)
 {
 	check(GMaxRHIFeatureLevel >= ERHIFeatureLevel::SM5);
 	VERIFY_GL_SCOPE();
diff --git a/Engine/Source/Runtime/OpenGLDrv/Public/OpenGLDrv.h b/Engine/Source/Runtime/OpenGLDrv/Public/OpenGLDrv.h
--- a/Engine/Source/Runtime/OpenGLDrv/Public/OpenGLDrv.h
+++ b/Engine/Source/Runtime/OpenGLDrv/Public/OpenGLDrv.h
@@ -473,12 +473,12 @@ public:
 	virtual void RHISetShaderTexture(FGeometryShaderRHIParamRef GeometryShader, uint32 TextureIndex, FTextureRHIParamRef NewTexture) final override;
 	virtual void RHISetShaderTexture(FPixelShaderRHIParamRef PixelShader, uint32 TextureIndex, FTextureRHIParamRef NewTexture) final override;
 	virtual void RHISetShaderTexture(FComputeShaderRHIParamRef PixelShader, uint32 TextureIndex, FTextureRHIParamRef NewTexture) final override;
-	virtual void RHISetShaderSampler(FComputeShaderRHIParamRef ComputeShader, uint32 SamplerIndex, FSamplerStateRHIParamRef NewState) final override;
-	virtual void RHISetShaderSampler(FVertexShaderRHIParamRef VertexShader, uint32 SamplerIndex, FSamplerStateRHIParamRef NewState) final override;
-	virtual void RHISetShaderSampler(FGeometryShaderRHIParamRef GeometryShader, uint32 SamplerIndex, FSamplerStateRHIParamRef NewState) final override;
-	virtual void RHISetShaderSampler(FDomainShaderRHIParamRef DomainShader, uint32 SamplerIndex, FSamplerStateRHIParamRef NewState) final override;
-	virtual void RHISetShaderSampler(FHullShaderRHIParamRef HullShader, uint32 SamplerIndex, FSamplerStateRHIParamRef NewState) final override;
-	virtual void RHISetShaderSampler(FPixelShaderRHIParamRef PixelShader, uint32 SamplerIndex, FSamplerStateRHIParamRef NewState) final override;
+	virtual void RHISetShaderSampler(FComputeShaderRHIParamRef ComputeShader, uint32 SamplerIndex, FRHISamplerState* NewState) final override;
+	virtual void RHISetShaderSampler(FVertexShaderRHIParamRef VertexShader, uint32 SamplerIndex, FRHISamplerState* NewState) final override;
+	virtual void RHISetShaderSampler(FGeometryShaderRHIParamRef GeometryShader, uint32 SamplerIndex, FRHISamplerState* NewState) final override;
+	virtual void RHISetShaderSampler(FDomainShaderRHIParamRef DomainShader, uint32 SamplerIndex, FRHISamplerState* NewState) final override;
+	virtual void RHISetShaderSampler(FHullShaderRHIParamRef HullShader, uint32 SamplerIndex, FRHISamplerState* NewState) final override;
+	virtual void RHISetShaderSampler(FPixelShaderRHIParamRef PixelShader, uint32 SamplerIndex, FRHISamplerState* NewState) final override;
 	virtual void RHISetUAVParameter(FComputeShaderRHIParamRef ComputeShader, uint32 UAVIndex, FUnorderedAccessViewRHIParamRef UAV) final override;
 	virtual void RHISetUAVParameter(FComputeShaderRHIParamRef ComputeShader, uint32 UAVIndex, FUnorderedAccessViewRHIParamRef UAV, uint32 InitialCount) final override;
 	virtual void RHISetShaderResourceViewParameter(FPixelShaderRHIParamRef PixelShader, uint32 SamplerIndex, FShaderResourceViewRHIParamRef SRV) final override;
diff --git a/Engine/Source/Runtime/RHI/Public/RHICommandList.h b/Engine/Source/Runtime/RHI/Public/RHICommandList.h
--- a/Engine/Source/Runtime/RHI/Public/RHICommandList.h
+++ b/Engine/Source/Runtime/RHI/Public/RHICommandList.h
@@ -115,7 +115,7 @@ struct FRayTracingShaderBindings
 	FTextureRHIParamRef Textures[32] = {};
 	FShaderResourceViewRHIParamRef SRVs[32] = {};
 	FUniformBufferRHIParamRef UniformBuffers[8] = {};
-	FSamplerStateRHIParamRef Samplers[16] = {};
+	FRHISamplerState* Samplers[16] = {};
 	FUnorderedAccessViewRHIParamRef UAVs[8] = {};
 };
 
@@ -830,8 +830,8 @@ struct FRHICommandSetShaderSampler final : public FRHICommand<FRHICommandSetShad
 {
 	TShaderRHIParamRef Shader;
 	uint32 SamplerIndex;
-	FSamplerStateRHIParamRef Sampler;
-	FORCEINLINE_DEBUGGABLE FRHICommandSetShaderSampler(TShaderRHIParamRef InShader, uint32 InSamplerIndex, FSamplerStateRHIParamRef InSampler)
+	FRHISamplerState* Sampler;
+	FORCEINLINE_DEBUGGABLE FRHICommandSetShaderSampler(TShaderRHIParamRef InShader, uint32 InSamplerIndex, FRHISamplerState* InSampler)
 		: Shader(InShader)
 		, SamplerIndex(InSamplerIndex)
 		, Sampler(InSampler)
@@ -2204,7 +2204,7 @@ public:
 	}
 
 	template <typename TShaderRHIParamRef>
-	FORCEINLINE_DEBUGGABLE void SetShaderSampler(TShaderRHIParamRef Shader, uint32 SamplerIndex, FSamplerStateRHIParamRef State)
+	FORCEINLINE_DEBUGGABLE void SetShaderSampler(TShaderRHIParamRef Shader, uint32 SamplerIndex, FRHISamplerState* State)
 	{
 		//check(IsOutsideRenderPass());
 		
@@ -2224,7 +2224,7 @@ public:
 	}
 
 	template <typename TShaderRHI>
-	FORCEINLINE_DEBUGGABLE void SetShaderSampler(TRefCountPtr<TShaderRHI>& Shader, uint32 SamplerIndex, FSamplerStateRHIParamRef State)
+	FORCEINLINE_DEBUGGABLE void SetShaderSampler(TRefCountPtr<TShaderRHI>& Shader, uint32 SamplerIndex, FRHISamplerState* State)
 	{
 		SetShaderSampler(Shader.GetReference(), SamplerIndex, State);
 	}
@@ -3218,7 +3218,7 @@ public:
 		ALLOC_COMMAND(FRHICommandSetShaderResourceViewParameter<FComputeShaderRHIParamRef, ECmdList::ECompute>)(Shader, SamplerIndex, SRV);
 	}
 	
-	FORCEINLINE_DEBUGGABLE void SetShaderSampler(FComputeShaderRHIParamRef Shader, uint32 SamplerIndex, FSamplerStateRHIParamRef State)
+	FORCEINLINE_DEBUGGABLE void SetShaderSampler(FComputeShaderRHIParamRef Shader, uint32 SamplerIndex, FRHISamplerState* State)
 	{
 		// Immutable samplers can't be set dynamically
 		check(!State->IsImmutable());
diff --git a/Engine/Source/Runtime/RHI/Public/RHIContext.h b/Engine/Source/Runtime/RHI/Public/RHIContext.h
--- a/Engine/Source/Runtime/RHI/Public/RHIContext.h
+++ b/Engine/Source/Runtime/RHI/Public/RHIContext.h
@@ -79,7 +79,7 @@ public:
 	* @param SamplerIndex		The index of the sampler.
 	* @param NewState			The new sampler state.
 	*/
-	virtual void RHISetShaderSampler(FComputeShaderRHIParamRef ComputeShader, uint32 SamplerIndex, FSamplerStateRHIParamRef NewState) = 0;
+	virtual void RHISetShaderSampler(FComputeShaderRHIParamRef ComputeShader, uint32 SamplerIndex, FRHISamplerState* NewState) = 0;
 
 	/**
 	* Sets a compute shader UAV parameter.
@@ -352,7 +352,7 @@ public:
 	* @param SamplerIndex		The index of the sampler.
 	* @param NewState			The new sampler state.
 	*/
-	virtual void RHISetShaderSampler(FComputeShaderRHIParamRef ComputeShader, uint32 SamplerIndex, FSamplerStateRHIParamRef NewState) = 0;
+	virtual void RHISetShaderSampler(FComputeShaderRHIParamRef ComputeShader, uint32 SamplerIndex, FRHISamplerState* NewState) = 0;
 
 	/**
 	* Sets sampler state.
@@ -360,7 +360,7 @@ public:
 	* @param SamplerIndex		The index of the sampler.
 	* @param NewState			The new sampler state.
 	*/
-	virtual void RHISetShaderSampler(FVertexShaderRHIParamRef VertexShader, uint32 SamplerIndex, FSamplerStateRHIParamRef NewState) = 0;
+	virtual void RHISetShaderSampler(FVertexShaderRHIParamRef VertexShader, uint32 SamplerIndex, FRHISamplerState* NewState) = 0;
 
 	/**
 	* Sets sampler state.
@@ -368,7 +368,7 @@ public:
 	* @param SamplerIndex		The index of the sampler.
 	* @param NewState			The new sampler state.
 	*/
-	virtual void RHISetShaderSampler(FGeometryShaderRHIParamRef GeometryShader, uint32 SamplerIndex, FSamplerStateRHIParamRef NewState) = 0;
+	virtual void RHISetShaderSampler(FGeometryShaderRHIParamRef GeometryShader, uint32 SamplerIndex, FRHISamplerState* NewState) = 0;
 
 	/**
 	* Sets sampler state.
@@ -376,7 +376,7 @@ public:
 	* @param SamplerIndex		The index of the sampler.
 	* @param NewState			The new sampler state.
 	*/
-	virtual void RHISetShaderSampler(FDomainShaderRHIParamRef DomainShader, uint32 SamplerIndex, FSamplerStateRHIParamRef NewState) = 0;
+	virtual void RHISetShaderSampler(FDomainShaderRHIParamRef DomainShader, uint32 SamplerIndex, FRHISamplerState* NewState) = 0;
 
 	/**
 	* Sets sampler state.
@@ -384,7 +384,7 @@ public:
 	* @param SamplerIndex		The index of the sampler.
 	* @param NewState			The new sampler state.
 	*/
-	virtual void RHISetShaderSampler(FHullShaderRHIParamRef HullShader, uint32 SamplerIndex, FSamplerStateRHIParamRef NewState) = 0;
+	virtual void RHISetShaderSampler(FHullShaderRHIParamRef HullShader, uint32 SamplerIndex, FRHISamplerState* NewState) = 0;
 
 	/**
 	* Sets sampler state.
@@ -392,7 +392,7 @@ public:
 	* @param SamplerIndex		The index of the sampler.
 	* @param NewState			The new sampler state.
 	*/
-	virtual void RHISetShaderSampler(FPixelShaderRHIParamRef PixelShader, uint32 SamplerIndex, FSamplerStateRHIParamRef NewState) = 0;
+	virtual void RHISetShaderSampler(FPixelShaderRHIParamRef PixelShader, uint32 SamplerIndex, FRHISamplerState* NewState) = 0;
 
 	/**
 	* Sets a compute shader UAV parameter.
diff --git a/Engine/Source/Runtime/RHI/Public/RHIResources.h b/Engine/Source/Runtime/RHI/Public/RHIResources.h
--- a/Engine/Source/Runtime/RHI/Public/RHIResources.h
+++ b/Engine/Source/Runtime/RHI/Public/RHIResources.h
@@ -1097,6 +1097,7 @@ class FRHIUnorderedAccessView : public FRHIResource {};
 class FRHIShaderResourceView : public FRHIResource {};
 
 
+UE_DEPRECATED(4.23, "FSamplerStateRHIParamRef typedef is deprecated; please use FRHISamplerState* directly instead.")
 typedef FRHISamplerState*              FSamplerStateRHIParamRef;
 typedef TRefCountPtr<FRHISamplerState> FSamplerStateRHIRef;
 
@@ -1778,7 +1779,7 @@ struct FBoundShaderStateInput
 
 struct FImmutableSamplerState
 {
-	using TImmutableSamplers = TStaticArray<FSamplerStateRHIParamRef, MaxImmutableSamplers>;
+	using TImmutableSamplers = TStaticArray<FRHISamplerState*, MaxImmutableSamplers>;
 
 	FImmutableSamplerState()
 		: ImmutableSamplers(nullptr)
diff --git a/Engine/Source/Runtime/RHI/Public/RHIStaticStates.h b/Engine/Source/Runtime/RHI/Public/RHIStaticStates.h
--- a/Engine/Source/Runtime/RHI/Public/RHIStaticStates.h
+++ b/Engine/Source/Runtime/RHI/Public/RHIStaticStates.h
@@ -154,7 +154,7 @@ template<ESamplerFilter Filter=SF_Point,
 	uint32 BorderColor = 0,
 	/** Only supported in D3D11 */
 	ESamplerCompareFunction SamplerComparisonFunction=SCF_Never>
-class TStaticSamplerState : public TStaticStateRHI<TStaticSamplerState<Filter,AddressU,AddressV,AddressW,MipBias,MaxAnisotropy,BorderColor,SamplerComparisonFunction>,FSamplerStateRHIRef,FSamplerStateRHIParamRef>
+class TStaticSamplerState : public TStaticStateRHI<TStaticSamplerState<Filter,AddressU,AddressV,AddressW,MipBias,MaxAnisotropy,BorderColor,SamplerComparisonFunction>,FSamplerStateRHIRef, FRHISamplerState*>
 {
 public:
 	static FSamplerStateRHIRef CreateRHI()
diff --git a/Engine/Source/Runtime/RHI/Public/RHIValidationContext.h b/Engine/Source/Runtime/RHI/Public/RHIValidationContext.h
--- a/Engine/Source/Runtime/RHI/Public/RHIValidationContext.h
+++ b/Engine/Source/Runtime/RHI/Public/RHIValidationContext.h
@@ -88,7 +88,7 @@ public:
 	* @param SamplerIndex		The index of the sampler.
 	* @param NewState			The new sampler state.
 	*/
-	virtual void RHISetShaderSampler(FComputeShaderRHIParamRef Shader, uint32 SamplerIndex, FSamplerStateRHIParamRef NewState) override final
+	virtual void RHISetShaderSampler(FComputeShaderRHIParamRef Shader, uint32 SamplerIndex, FRHISamplerState* NewState) override final
 	{
 		checkf(State.bComputeShaderSet, TEXT("A Compute shader has to be set to set resources into a shader!"));
 		RHIContext->RHISetShaderSampler(Shader, SamplerIndex, NewState);
@@ -462,7 +462,7 @@ public:
 	* @param SamplerIndex		The index of the sampler.
 	* @param NewState			The new sampler state.
 	*/
-	virtual void RHISetShaderSampler(FComputeShaderRHIParamRef Shader, uint32 SamplerIndex, FSamplerStateRHIParamRef NewState) override final
+	virtual void RHISetShaderSampler(FComputeShaderRHIParamRef Shader, uint32 SamplerIndex, FRHISamplerState* NewState) override final
 	{
 		checkf(State.bComputeShaderSet, TEXT("A Compute shader has to be set to set resources into a shader!"));
 		RHIContext->RHISetShaderSampler(Shader, SamplerIndex, NewState);
@@ -474,7 +474,7 @@ public:
 	* @param SamplerIndex		The index of the sampler.
 	* @param NewState			The new sampler state.
 	*/
-	virtual void RHISetShaderSampler(FVertexShaderRHIParamRef Shader, uint32 SamplerIndex, FSamplerStateRHIParamRef NewState) override final
+	virtual void RHISetShaderSampler(FVertexShaderRHIParamRef Shader, uint32 SamplerIndex, FRHISamplerState* NewState) override final
 	{
 		checkf(State.bGfxPSOSet, TEXT("A Graphics PSO has to be set to set resources into a shader!"));
 		RHIContext->RHISetShaderSampler(Shader, SamplerIndex, NewState);
@@ -487,7 +487,7 @@ public:
 	* @param SamplerIndex		The index of the sampler.
 	* @param NewState			The new sampler state.
 	*/
-	virtual void RHISetShaderSampler(FGeometryShaderRHIParamRef Shader, uint32 SamplerIndex, FSamplerStateRHIParamRef NewState) override final
+	virtual void RHISetShaderSampler(FGeometryShaderRHIParamRef Shader, uint32 SamplerIndex, FRHISamplerState* NewState) override final
 	{
 		checkf(State.bGfxPSOSet, TEXT("A Graphics PSO has to be set to set resources into a shader!"));
 		RHIContext->RHISetShaderSampler(Shader, SamplerIndex, NewState);
@@ -499,7 +499,7 @@ public:
 	* @param SamplerIndex		The index of the sampler.
 	* @param NewState			The new sampler state.
 	*/
-	virtual void RHISetShaderSampler(FDomainShaderRHIParamRef Shader, uint32 SamplerIndex, FSamplerStateRHIParamRef NewState) override final
+	virtual void RHISetShaderSampler(FDomainShaderRHIParamRef Shader, uint32 SamplerIndex, FRHISamplerState* NewState) override final
 	{
 		checkf(State.bGfxPSOSet, TEXT("A Graphics PSO has to be set to set resources into a shader!"));
 		RHIContext->RHISetShaderSampler(Shader, SamplerIndex, NewState);
@@ -511,7 +511,7 @@ public:
 	* @param SamplerIndex		The index of the sampler.
 	* @param NewState			The new sampler state.
 	*/
-	virtual void RHISetShaderSampler(FHullShaderRHIParamRef Shader, uint32 SamplerIndex, FSamplerStateRHIParamRef NewState) override final
+	virtual void RHISetShaderSampler(FHullShaderRHIParamRef Shader, uint32 SamplerIndex, FRHISamplerState* NewState) override final
 	{
 		checkf(State.bGfxPSOSet, TEXT("A Graphics PSO has to be set to set resources into a shader!"));
 		RHIContext->RHISetShaderSampler(Shader, SamplerIndex, NewState);
@@ -523,7 +523,7 @@ public:
 	* @param SamplerIndex		The index of the sampler.
 	* @param NewState			The new sampler state.
 	*/
-	virtual void RHISetShaderSampler(FPixelShaderRHIParamRef Shader, uint32 SamplerIndex, FSamplerStateRHIParamRef NewState) override final
+	virtual void RHISetShaderSampler(FPixelShaderRHIParamRef Shader, uint32 SamplerIndex, FRHISamplerState* NewState) override final
 	{
 		checkf(State.bGfxPSOSet, TEXT("A Graphics PSO has to be set to set resources into a shader!"));
 		RHIContext->RHISetShaderSampler(Shader, SamplerIndex, NewState);
diff --git a/Engine/Source/Runtime/RenderCore/Private/ShaderParameterStruct.cpp b/Engine/Source/Runtime/RenderCore/Private/ShaderParameterStruct.cpp
--- a/Engine/Source/Runtime/RenderCore/Private/ShaderParameterStruct.cpp
+++ b/Engine/Source/Runtime/RenderCore/Private/ShaderParameterStruct.cpp
@@ -420,7 +420,7 @@ void ValidateShaderParameters(const FShader* Shader, const FShaderParametersMeta
 	// Samplers
 	for (const FShaderParameterBindings::FResourceParameter& ParameterBinding : Bindings.Samplers)
 	{
-		auto ShaderParameterRef = *reinterpret_cast<const FSamplerStateRHIParamRef*>(Base + ParameterBinding.ByteOffset);
+		FRHISamplerState* ShaderParameterRef = (FRHISamplerState*)(Base + ParameterBinding.ByteOffset);
 		if (!ShaderParameterRef)
 		{
 			EmitNullShaderParameterFatalError(Shader, ParametersMetadata, ParameterBinding.ByteOffset);
diff --git a/Engine/Source/Runtime/RenderCore/Private/VisualizeTexture.cpp b/Engine/Source/Runtime/RenderCore/Private/VisualizeTexture.cpp
--- a/Engine/Source/Runtime/RenderCore/Private/VisualizeTexture.cpp
+++ b/Engine/Source/Runtime/RenderCore/Private/VisualizeTexture.cpp
@@ -228,7 +228,7 @@ void FVisualizeTexture::CreateContentCapturePass(FRDGBuilder& GraphBuilder, cons
 			PassParameters->VisualizeParam[2] = FVector4(InputValueMapping, 0.0f, SingleChannel);
 		}
 
-		FSamplerStateRHIParamRef PointSampler = TStaticSamplerState<SF_Point, AM_Clamp, AM_Clamp, AM_Clamp>::GetRHI();
+		FRHISamplerState* PointSampler = TStaticSamplerState<SF_Point, AM_Clamp, AM_Clamp, AM_Clamp>::GetRHI();
 
 		PassParameters->VisualizeTexture2D = SrcTexture;
 		PassParameters->VisualizeTexture2DSampler = PointSampler;
diff --git a/Engine/Source/Runtime/RenderCore/Public/ShaderParameterMacros.h b/Engine/Source/Runtime/RenderCore/Public/ShaderParameterMacros.h
--- a/Engine/Source/Runtime/RenderCore/Public/ShaderParameterMacros.h
+++ b/Engine/Source/Runtime/RenderCore/Public/ShaderParameterMacros.h
@@ -861,10 +861,10 @@ extern RENDERCORE_API FShaderParametersMetadata* FindUniformBufferStructByFName(
  *	SHADER_PARAMETER_SAMPLER_ARRAY(SamplerState, MyArrayOfSamplers, [8])
  */
 #define SHADER_PARAMETER_SAMPLER(ShaderType,MemberName) \
-	INTERNAL_SHADER_PARAMETER_EXPLICIT(UBMT_SAMPLER, TShaderResourceParameterTypeInfo<FSamplerStateRHIParamRef>, FSamplerStateRHIParamRef,MemberName,, = nullptr,EShaderPrecisionModifier::Float,TEXT(#ShaderType),false)
+	INTERNAL_SHADER_PARAMETER_EXPLICIT(UBMT_SAMPLER, TShaderResourceParameterTypeInfo<FRHISamplerState*>, FRHISamplerState*,MemberName,, = nullptr,EShaderPrecisionModifier::Float,TEXT(#ShaderType),false)
 
 #define SHADER_PARAMETER_SAMPLER_ARRAY(ShaderType,MemberName, ArrayDecl) \
-	INTERNAL_SHADER_PARAMETER_EXPLICIT(UBMT_SAMPLER, TShaderResourceParameterTypeInfo<FSamplerStateRHIParamRef ArrayDecl>, FSamplerStateRHIParamRef,MemberName,ArrayDecl,,EShaderPrecisionModifier::Float,TEXT(#ShaderType),false)
+	INTERNAL_SHADER_PARAMETER_EXPLICIT(UBMT_SAMPLER, TShaderResourceParameterTypeInfo<FRHISamplerState* ArrayDecl>, FRHISamplerState*,MemberName,ArrayDecl,,EShaderPrecisionModifier::Float,TEXT(#ShaderType),false)
 
 /** Adds a render graph tracked texture.
  *
diff --git a/Engine/Source/Runtime/RenderCore/Public/ShaderParameterStruct.h b/Engine/Source/Runtime/RenderCore/Public/ShaderParameterStruct.h
--- a/Engine/Source/Runtime/RenderCore/Public/ShaderParameterStruct.h
+++ b/Engine/Source/Runtime/RenderCore/Public/ShaderParameterStruct.h
@@ -205,7 +205,7 @@ inline void SetShaderParameters(TRHICmdList& RHICmdList, const TShaderClass* Sha
 	// Samplers
 	for (const FShaderParameterBindings::FResourceParameter& ParameterBinding : Bindings.Samplers)
 	{
-		auto ShaderParameterRef = *reinterpret_cast<const FSamplerStateRHIParamRef*>(Base + ParameterBinding.ByteOffset);
+		FRHISamplerState* ShaderParameterRef = *(FRHISamplerState**)(Base + ParameterBinding.ByteOffset);
 
 		if (DO_CHECK && !ShaderParameterRef)
 		{
@@ -330,7 +330,7 @@ void SetShaderParameters(FRayTracingShaderBindingsWriter& RTBindingsWriter, cons
 	// Samplers
 	for (const FShaderParameterBindings::FResourceParameter& ParameterBinding : Bindings.Samplers)
 	{
-		auto ShaderParameterRef = *reinterpret_cast<const FSamplerStateRHIParamRef*>(Base + ParameterBinding.ByteOffset);
+		FRHISamplerState* ShaderParameterRef = *(FRHISamplerState**)(Base + ParameterBinding.ByteOffset);
 
 		if (DO_CHECK && !ShaderParameterRef)
 		{
diff --git a/Engine/Source/Runtime/RenderCore/Public/ShaderParameterUtils.h b/Engine/Source/Runtime/RenderCore/Public/ShaderParameterUtils.h
--- a/Engine/Source/Runtime/RenderCore/Public/ShaderParameterUtils.h
+++ b/Engine/Source/Runtime/RenderCore/Public/ShaderParameterUtils.h
@@ -234,7 +234,7 @@ FORCEINLINE void SetTextureParameter(
 	ShaderTypeRHIParamRef Shader,
 	const FShaderResourceParameter& TextureParameter,
 	const FShaderResourceParameter& SamplerParameter,
-	FSamplerStateRHIParamRef SamplerStateRHI,
+	FRHISamplerState* SamplerStateRHI,
 	FTextureRHIParamRef TextureRHI,
 	uint32 ElementIndex = 0
 	)
@@ -294,7 +294,7 @@ FORCEINLINE void SetSamplerParameter(
 	TRHICmdList& RHICmdList,
 	ShaderTypeRHIParamRef Shader,
 	const FShaderResourceParameter& Parameter,
-	FSamplerStateRHIParamRef SamplerStateRHI
+	FRHISamplerState* SamplerStateRHI
 	)
 {
 	if(Parameter.IsBound())
diff --git a/Engine/Source/Runtime/RenderCore/Public/ShaderParameters.h b/Engine/Source/Runtime/RenderCore/Public/ShaderParameters.h
--- a/Engine/Source/Runtime/RenderCore/Public/ShaderParameters.h
+++ b/Engine/Source/Runtime/RenderCore/Public/ShaderParameters.h
@@ -275,7 +275,7 @@ struct FRayTracingShaderBindingsWriter : FRayTracingShaderBindings
 		}
 	}
 
-	void Set(const FShaderResourceParameter& Param, FSamplerStateRHIParamRef Value)
+	void Set(const FShaderResourceParameter& Param, FRHISamplerState* Value)
 	{
 		if (Param.IsBound())
 		{
@@ -296,7 +296,7 @@ struct FRayTracingShaderBindingsWriter : FRayTracingShaderBindings
 		SRVs[BaseIndex] = Value;
 	}
 
-	void SetSampler(uint16 BaseIndex, FSamplerStateRHIParamRef Value)
+	void SetSampler(uint16 BaseIndex, FRHISamplerState* Value)
 	{
 		checkSlow(BaseIndex < ARRAY_COUNT(Samplers));
 		Samplers[BaseIndex] = Value;
diff --git a/Engine/Source/Runtime/Renderer/Private/HeightfieldLighting.cpp b/Engine/Source/Runtime/Renderer/Private/HeightfieldLighting.cpp
--- a/Engine/Source/Runtime/Renderer/Private/HeightfieldLighting.cpp
+++ b/Engine/Source/Runtime/Renderer/Private/HeightfieldLighting.cpp
@@ -784,7 +784,7 @@ public:
 		SetShaderValue(RHICmdList, ShaderRHI, CascadeDepthMinMax, FVector2D(ProjectedShadowInfo->CascadeSettings.SplitNear, ProjectedShadowInfo->CascadeSettings.SplitFar));
 
 		FTextureRHIParamRef ShadowDepthTextureValue = ProjectedShadowInfo->RenderTargets.DepthTarget->GetRenderTargetItem().ShaderResourceTexture.GetReference();
-		FSamplerStateRHIParamRef DepthSamplerState = TStaticSamplerState<SF_Point, AM_Clamp, AM_Clamp, AM_Clamp>::GetRHI();
+		FRHISamplerState* DepthSamplerState = TStaticSamplerState<SF_Point, AM_Clamp, AM_Clamp, AM_Clamp>::GetRHI();
 
 		SetTextureParameter(RHICmdList, ShaderRHI, ShadowDepthTexture, ShadowDepthTextureSampler, DepthSamplerState, ShadowDepthTextureValue);
 	}
diff --git a/Engine/Source/Runtime/Renderer/Private/LightPropagationVolume.cpp b/Engine/Source/Runtime/Renderer/Private/LightPropagationVolume.cpp
--- a/Engine/Source/Runtime/Renderer/Private/LightPropagationVolume.cpp
+++ b/Engine/Source/Runtime/Renderer/Private/LightPropagationVolume.cpp
@@ -486,8 +486,8 @@ public:
 		FComputeShaderRHIParamRef ShaderRHI = GetComputeShader();
 		FLpvWriteShaderCSBase::SetParameters(RHICmdList, BaseParams );
 
-		FSamplerStateRHIParamRef SamplerStateLinear  = TStaticSamplerState<SF_Bilinear,AM_Clamp,AM_Clamp,AM_Clamp>::GetRHI();
-		FSamplerStateRHIParamRef SamplerStatePoint   = TStaticSamplerState<SF_Point,AM_Clamp,AM_Clamp,AM_Clamp>::GetRHI();
+		FRHISamplerState* SamplerStateLinear  = TStaticSamplerState<SF_Bilinear,AM_Clamp,AM_Clamp,AM_Clamp>::GetRHI();
+		FRHISamplerState* SamplerStatePoint   = TStaticSamplerState<SF_Point,AM_Clamp,AM_Clamp,AM_Clamp>::GetRHI();
 
 		// FIXME: Why do we have to bind a samplerstate to a sampler here? Presumably this is for legacy reasons... 
 		SetTextureParameter(RHICmdList, ShaderRHI, RsmDiffuseTexture, LinearTextureSampler, SamplerStateLinear, RsmDiffuseTextureRHI );
diff --git a/Engine/Source/Runtime/Renderer/Private/MeshPassProcessor.cpp b/Engine/Source/Runtime/Renderer/Private/MeshPassProcessor.cpp
--- a/Engine/Source/Runtime/Renderer/Private/MeshPassProcessor.cpp
+++ b/Engine/Source/Runtime/Renderer/Private/MeshPassProcessor.cpp
@@ -49,7 +49,7 @@ public:
 	int32 MaxTextureUsed = -1;
 	FTextureRHIParamRef Textures[MAX_SRVs_PER_SHADER_STAGE] = {};
 	int32 MaxSamplerUsed = -1;
-	FSamplerStateRHIParamRef Samplers[MAX_SAMPLERS_PER_SHADER_STAGE] = {};
+	FRHISamplerState* Samplers[MAX_SAMPLERS_PER_SHADER_STAGE] = {};
 };
 
 class FReadOnlyMeshDrawSingleShaderBindings : public FMeshDrawShaderBindingsLayout
@@ -66,10 +66,10 @@ public:
 		return (const FUniformBufferRHIParamRef*)(Data + GetUniformBufferOffset());
 	}
 
-	inline const FSamplerStateRHIParamRef* GetSamplerStart() const
+	inline const FRHISamplerState** GetSamplerStart() const
 	{
 		const uint8* SamplerDataStart = Data + GetSamplerOffset();
-		return (const FSamplerStateRHIParamRef*)SamplerDataStart;
+		return (const FRHISamplerState**)SamplerDataStart;
 	}
 
 	inline const FRHIResource** GetSRVStart() const
@@ -119,7 +119,7 @@ void FMeshDrawShaderBindings::SetShaderBindings(
 		}
 	}
 
-	const FSamplerStateRHIParamRef* RESTRICT SamplerBindings = SingleShaderBindings.GetSamplerStart();
+	const FRHISamplerState** RESTRICT SamplerBindings = SingleShaderBindings.GetSamplerStart();
 	const FShaderParameterInfo* RESTRICT TextureSamplerParameters = SingleShaderBindings.ParameterMapInfo.TextureSamplers.GetData();
 	const int32 NumTextureSamplers = SingleShaderBindings.ParameterMapInfo.TextureSamplers.Num();
 
@@ -127,7 +127,7 @@ void FMeshDrawShaderBindings::SetShaderBindings(
 	{
 		FShaderParameterInfo Parameter = TextureSamplerParameters[SamplerIndex];
 		checkSlow(Parameter.BaseIndex < ARRAY_COUNT(ShaderBindingState.Samplers));
-		FSamplerStateRHIParamRef Sampler = SamplerBindings[SamplerIndex];
+		FRHISamplerState* Sampler = (FRHISamplerState*)SamplerBindings[SamplerIndex];
 
 		if (Sampler != ShaderBindingState.Samplers[Parameter.BaseIndex])
 		{
@@ -211,14 +211,14 @@ void FMeshDrawShaderBindings::SetShaderBindings(
 		RHICmdList.SetShaderUniformBuffer(Shader, Parameter.BaseIndex, UniformBuffer);
 	}
 
-	const FSamplerStateRHIParamRef* RESTRICT SamplerBindings = SingleShaderBindings.GetSamplerStart();
+	const FRHISamplerState** RESTRICT SamplerBindings = SingleShaderBindings.GetSamplerStart();
 	const FShaderParameterInfo* RESTRICT TextureSamplerParameters = SingleShaderBindings.ParameterMapInfo.TextureSamplers.GetData();
 	const int32 NumTextureSamplers = SingleShaderBindings.ParameterMapInfo.TextureSamplers.Num();
 
 	for (int32 SamplerIndex = 0; SamplerIndex < NumTextureSamplers; SamplerIndex++)
 	{
 		FShaderParameterInfo Parameter = TextureSamplerParameters[SamplerIndex];
-		FSamplerStateRHIParamRef Sampler = SamplerBindings[SamplerIndex];
+		FRHISamplerState* Sampler = (FRHISamplerState*)SamplerBindings[SamplerIndex];
 
 		RHICmdList.SetShaderSampler(Shader, Parameter.BaseIndex, Sampler);
 	}
@@ -604,12 +604,12 @@ void FMeshDrawShaderBindings::Finalize(const FMeshProcessorShaders* ShadersForDe
 			}
 		}
 
-		const FSamplerStateRHIParamRef* SamplerBindings = SingleShaderBindings.GetSamplerStart();
+		const FRHISamplerState** SamplerBindings = SingleShaderBindings.GetSamplerStart();
 
 		for (int32 BindingIndex = 0; BindingIndex < ShaderLayout.ParameterMapInfo.TextureSamplers.Num(); BindingIndex++)
 		{
 			FShaderParameterInfo ParameterInfo = ShaderLayout.ParameterMapInfo.TextureSamplers[BindingIndex];
-			FSamplerStateRHIParamRef SamplerValue = SamplerBindings[BindingIndex];
+			const FRHISamplerState* SamplerValue = SamplerBindings[BindingIndex];
 			ensureMsgf(SamplerValue, TEXT("Shader %s with vertex factory %s never set sampler at BaseIndex %u.  This can cause GPU hangs, depending on how the shader uses it."), 
 				Shader->GetType()->GetName(), 
 				Shader->GetVertexFactoryType()->GetName(),
diff --git a/Engine/Source/Runtime/Renderer/Private/MobileSceneCaptureRendering.cpp b/Engine/Source/Runtime/Renderer/Private/MobileSceneCaptureRendering.cpp
--- a/Engine/Source/Runtime/Renderer/Private/MobileSceneCaptureRendering.cpp
+++ b/Engine/Source/Runtime/Renderer/Private/MobileSceneCaptureRendering.cpp
@@ -64,7 +64,7 @@ public:
 		}
 	}
 
-	void SetParameters(FRHICommandList& RHICmdList, const FViewInfo& View, FSamplerStateRHIParamRef SamplerStateRHI, FTextureRHIParamRef TextureRHI)
+	void SetParameters(FRHICommandList& RHICmdList, const FViewInfo& View, FRHISamplerState* SamplerStateRHI, FTextureRHIParamRef TextureRHI)
 	{
 		FGlobalShader::SetParameters<FViewUniformShaderParameters>(RHICmdList, GetPixelShader(), View.ViewUniformBuffer);
 		SetTextureParameter(RHICmdList, GetPixelShader(), InTexture, InTextureSampler, SamplerStateRHI, TextureRHI);
diff --git a/Engine/Source/Runtime/Renderer/Private/PostProcess/PostProcessCombineLUTs.cpp b/Engine/Source/Runtime/Renderer/Private/PostProcess/PostProcessCombineLUTs.cpp
--- a/Engine/Source/Runtime/Renderer/Private/PostProcess/PostProcessCombineLUTs.cpp
+++ b/Engine/Source/Runtime/Renderer/Private/PostProcess/PostProcessCombineLUTs.cpp
@@ -189,7 +189,7 @@ public:
 			if(i != 0)
 			{
 				// don't use texture asset sampler as it might have anisotropic filtering enabled
-				FSamplerStateRHIParamRef Sampler = TStaticSamplerState<SF_Bilinear, AM_Clamp, AM_Clamp, AM_Clamp, 0, 1>::GetRHI();
+				FRHISamplerState* Sampler = TStaticSamplerState<SF_Bilinear, AM_Clamp, AM_Clamp, AM_Clamp, 0, 1>::GetRHI();
 
 				SetTextureParameter(RHICmdList, ShaderRHI, TextureParameter[i], TextureParameterSampler[i], Sampler, Textures[i]->TextureRHI);
 			}
diff --git a/Engine/Source/Runtime/Renderer/Private/PostProcess/PostProcessDownsample.cpp b/Engine/Source/Runtime/Renderer/Private/PostProcess/PostProcessDownsample.cpp
--- a/Engine/Source/Runtime/Renderer/Private/PostProcess/PostProcessDownsample.cpp
+++ b/Engine/Source/Runtime/Renderer/Private/PostProcess/PostProcessDownsample.cpp
@@ -72,7 +72,7 @@ public:
 		SceneTextureParameters.Set(RHICmdList, ShaderRHI, Context.View.FeatureLevel, ESceneTextureSetupMode::All);
 
 		// filter only if needed for better performance
-		FSamplerStateRHIParamRef Filter = (Method == 2) ? 
+		FRHISamplerState* Filter = (Method == 2) ?
 			TStaticSamplerState<SF_Point, AM_Clamp, AM_Clamp, AM_Clamp>::GetRHI():
 			TStaticSamplerState<SF_Bilinear, AM_Clamp, AM_Clamp, AM_Clamp>::GetRHI();
 		
@@ -213,7 +213,7 @@ public:
 		FGlobalShader::SetParameters<FViewUniformShaderParameters>(RHICmdList, ShaderRHI, Context.View.ViewUniformBuffer);
 
 		// filter only if needed for better performance
-		FSamplerStateRHIParamRef Filter = (Method == 2) ? 
+		FRHISamplerState* Filter = (Method == 2) ?
 			TStaticSamplerState<SF_Point, AM_Clamp, AM_Clamp, AM_Clamp>::GetRHI():
 			TStaticSamplerState<SF_Bilinear, AM_Clamp, AM_Clamp, AM_Clamp>::GetRHI();
 
diff --git a/Engine/Source/Runtime/Renderer/Private/PostProcess/PostProcessMotionBlur.cpp b/Engine/Source/Runtime/Renderer/Private/PostProcess/PostProcessMotionBlur.cpp
--- a/Engine/Source/Runtime/Renderer/Private/PostProcess/PostProcessMotionBlur.cpp
+++ b/Engine/Source/Runtime/Renderer/Private/PostProcess/PostProcessMotionBlur.cpp
@@ -150,7 +150,7 @@ EMotionBlurQuality GetMotionBlurQuality()
 	return static_cast<EMotionBlurQuality>(Quality - 1);
 }
 
-FSamplerStateRHIParamRef GetMotionBlurColorSampler()
+FRHISamplerState* GetMotionBlurColorSampler()
 {
 	bool bFiltered = false;
 
@@ -168,7 +168,7 @@ FSamplerStateRHIParamRef GetMotionBlurColorSampler()
 	}
 }
 
-FSamplerStateRHIParamRef GetMotionBlurVelocitySampler()
+FRHISamplerState* GetMotionBlurVelocitySampler()
 {
 	return TStaticSamplerState<SF_Point, AM_Clamp, AM_Clamp, AM_Clamp>::GetRHI();
 }
diff --git a/Engine/Source/Runtime/Renderer/Private/PostProcess/PostProcessSubsurface.cpp b/Engine/Source/Runtime/Renderer/Private/PostProcess/PostProcessSubsurface.cpp
--- a/Engine/Source/Runtime/Renderer/Private/PostProcess/PostProcessSubsurface.cpp
+++ b/Engine/Source/Runtime/Renderer/Private/PostProcess/PostProcessSubsurface.cpp
@@ -337,7 +337,7 @@ public:
 	using FPermutationDomain = TShaderPermutationDomain<FDimensionDirection, FDimensionQuality>;
 
 	// Returns the sampler state based on the requested SSS filter CVar setting.
-	static FSamplerStateRHIParamRef GetSamplerState()
+	static FRHISamplerState* GetSamplerState()
 	{
 		if (CVarSSSFilter.GetValueOnRenderThread())
 		{
@@ -481,8 +481,8 @@ void ComputeSubsurfaceForView(
 	FRDGTextureRef SubsurfaceTextureX = nullptr;
 	FRDGTextureRef SubsurfaceTextureY = nullptr;
 
-	FSamplerStateRHIParamRef PointClampSampler = TStaticSamplerState<SF_Point, AM_Clamp, AM_Clamp, AM_Clamp>::GetRHI();
-	FSamplerStateRHIParamRef BilinearBorderSampler = TStaticSamplerState<SF_Bilinear, AM_Border, AM_Border, AM_Border>::GetRHI();
+	FRHISamplerState* PointClampSampler = TStaticSamplerState<SF_Point, AM_Clamp, AM_Clamp, AM_Clamp>::GetRHI();
+	FRHISamplerState* BilinearBorderSampler = TStaticSamplerState<SF_Bilinear, AM_Border, AM_Border, AM_Border>::GetRHI();
 
 	/**
 	 * When in bypass mode, the setup and convolution passes are skipped, but lighting
@@ -517,7 +517,7 @@ void ComputeSubsurfaceForView(
 		SubsurfaceTextureX = GraphBuilder.CreateTexture(SubsurfaceTextureDescriptor, TEXT("SubsurfaceTextureX"));
 		SubsurfaceTextureY = GraphBuilder.CreateTexture(SubsurfaceTextureDescriptor, TEXT("SubsurfaceTextureY"));
 
-		FSamplerStateRHIParamRef SubsurfaceSamplerState = FSubsurfacePS::GetSamplerState();
+		FRHISamplerState* SubsurfaceSamplerState = FSubsurfacePS::GetSamplerState();
 		const FSubsurfacePS::EQuality SubsurfaceQuality = FSubsurfacePS::GetQuality();
 
 		struct FSubsurfacePassInfo
diff --git a/Engine/Source/Runtime/Renderer/Private/PostProcess/PostProcessTonemap.cpp b/Engine/Source/Runtime/Renderer/Private/PostProcess/PostProcessTonemap.cpp
--- a/Engine/Source/Runtime/Renderer/Private/PostProcess/PostProcessTonemap.cpp
+++ b/Engine/Source/Runtime/Renderer/Private/PostProcess/PostProcessTonemap.cpp
@@ -764,7 +764,7 @@ public:
 
 		{
 			// filtering can cost performance so we use point where possible, we don't want anisotropic sampling
-			FSamplerStateRHIParamRef Filters[] =
+			FRHISamplerState* Filters[] =
 			{
 				TStaticSamplerState<SF_Bilinear, AM_Clamp, AM_Clamp, AM_Clamp, 0, 1>::GetRHI(),		// todo: could be SF_Point if fringe is disabled
 				TStaticSamplerState<SF_Bilinear, AM_Clamp, AM_Clamp, AM_Clamp, 0, 1>::GetRHI(),
@@ -882,7 +882,7 @@ public:
 		// PS params
 		{
 			// filtering can cost performance so we use point where possible, we don't want anisotropic sampling
-			FSamplerStateRHIParamRef Filters[] =
+			FRHISamplerState* Filters[] =
 			{
 				TStaticSamplerState<SF_Bilinear, AM_Clamp, AM_Clamp, AM_Clamp, 0, 1>::GetRHI(),		// todo: could be SF_Point if fringe is disabled
 				TStaticSamplerState<SF_Bilinear, AM_Clamp, AM_Clamp, AM_Clamp, 0, 1>::GetRHI(),
diff --git a/Engine/Source/Runtime/Renderer/Private/PostProcess/PostProcessUpscale.cpp b/Engine/Source/Runtime/Renderer/Private/PostProcess/PostProcessUpscale.cpp
--- a/Engine/Source/Runtime/Renderer/Private/PostProcess/PostProcessUpscale.cpp
+++ b/Engine/Source/Runtime/Renderer/Private/PostProcess/PostProcessUpscale.cpp
@@ -185,7 +185,7 @@ public:
 		
 		FGlobalShader::SetParameters<FViewUniformShaderParameters>(RHICmdList, ShaderRHI, Context.View.ViewUniformBuffer);
 
-		FSamplerStateRHIParamRef FilterTable[2];
+		FRHISamplerState* FilterTable[2];
 		FilterTable[0] = TStaticSamplerState<SF_Bilinear,AM_Clamp,AM_Clamp,AM_Clamp>::GetRHI();
 		FilterTable[1] = TStaticSamplerState<SF_Point,AM_Clamp,AM_Clamp,AM_Clamp>::GetRHI();
 			
diff --git a/Engine/Source/Runtime/Renderer/Private/PostProcess/PostProcessWeightedSampleSum.cpp b/Engine/Source/Runtime/Renderer/Private/PostProcess/PostProcessWeightedSampleSum.cpp
--- a/Engine/Source/Runtime/Renderer/Private/PostProcess/PostProcessWeightedSampleSum.cpp
+++ b/Engine/Source/Runtime/Renderer/Private/PostProcess/PostProcessWeightedSampleSum.cpp
@@ -137,7 +137,7 @@ public:
 
 	/** Sets shader parameter values */
 	void SetParameters(
-		FRHICommandList& RHICmdList, FSamplerStateRHIParamRef SamplerStateRHI, FTextureRHIParamRef FilterTextureRHI, FTextureRHIParamRef AdditiveTextureRHI,
+		FRHICommandList& RHICmdList, FRHISamplerState* SamplerStateRHI, FTextureRHIParamRef FilterTextureRHI, FTextureRHIParamRef AdditiveTextureRHI,
 		const FLinearColor* SampleWeightValues, const FVector2D* SampleOffsetValues, uint32 NumSamples, const FVector4& FilteredBufferUVMinMaxValue, const FVector4& AdditiveBufferUVMinMaxValue )
 	{
 		check(CompileTimeNumSamples == 0 && NumSamples > 0 && NumSamples <= MAX_FILTER_SAMPLES || CompileTimeNumSamples == NumSamples);
@@ -415,7 +415,7 @@ public:
 		// PS params
 		check(CompileTimeNumSamples == 0 && NumSamples > 0 && NumSamples <= MAX_FILTER_SAMPLES || CompileTimeNumSamples == NumSamples);
 
-		static FSamplerStateRHIParamRef SamplerStateRHI = TStaticSamplerState<SF_Bilinear,AM_Border,AM_Border,AM_Clamp>::GetRHI();
+		static FRHISamplerState* SamplerStateRHI = TStaticSamplerState<SF_Bilinear,AM_Border,AM_Border,AM_Clamp>::GetRHI();
 		SetTextureParameter(RHICmdList, ShaderRHI, FilterTexture, FilterTextureSampler, SamplerStateRHI, FilterTextureRHI);
 		SetTextureParameter(RHICmdList, ShaderRHI, AdditiveTexture, AdditiveTextureSampler, SamplerStateRHI, AdditiveTextureRHI);
 		SetShaderValueArray(RHICmdList, ShaderRHI, SampleWeights, SampleWeightValues, NumSamples);
@@ -501,7 +501,7 @@ public:
 void SetFilterShaders(
 	FRHICommandListImmediate& RHICmdList,
 	ERHIFeatureLevel::Type FeatureLevel,
-	FSamplerStateRHIParamRef SamplerStateRHI,
+	FRHISamplerState* SamplerStateRHI,
 	FTextureRHIParamRef FilterTextureRHI,
 	FTextureRHIParamRef AdditiveTextureRHI,
 	uint32 CombineMethodInt,
diff --git a/Engine/Source/Runtime/Renderer/Private/PostProcess/RenderingCompositionGraph.cpp b/Engine/Source/Runtime/Renderer/Private/PostProcess/RenderingCompositionGraph.cpp
--- a/Engine/Source/Runtime/Renderer/Private/PostProcess/RenderingCompositionGraph.cpp
+++ b/Engine/Source/Runtime/Renderer/Private/PostProcess/RenderingCompositionGraph.cpp
@@ -852,16 +852,16 @@ void FPostProcessPassParameters::Bind(const FShaderParameterMap& ParameterMap)
 }
 
 template <typename TRHICmdList>
-void FPostProcessPassParameters::SetPS(TRHICmdList& RHICmdList, const FPixelShaderRHIParamRef& ShaderRHI, const FRenderingCompositePassContext& Context, FSamplerStateRHIParamRef Filter, EFallbackColor FallbackColor, FSamplerStateRHIParamRef* FilterOverrideArray)
+void FPostProcessPassParameters::SetPS(TRHICmdList& RHICmdList, const FPixelShaderRHIParamRef& ShaderRHI, const FRenderingCompositePassContext& Context, FRHISamplerState* Filter, EFallbackColor FallbackColor, FRHISamplerState** FilterOverrideArray)
 {
 	Set(RHICmdList, ShaderRHI, Context, Filter, FallbackColor, FilterOverrideArray);
 }
 
-template void FPostProcessPassParameters::SetPS(FRHICommandList& RHICmdList, const FPixelShaderRHIParamRef& ShaderRHI, const FRenderingCompositePassContext& Context, FSamplerStateRHIParamRef Filter, EFallbackColor FallbackColor, FSamplerStateRHIParamRef* FilterOverrideArray);
-template void FPostProcessPassParameters::SetPS(FRHICommandListImmediate& RHICmdList, const FPixelShaderRHIParamRef& ShaderRHI, const FRenderingCompositePassContext& Context, FSamplerStateRHIParamRef Filter, EFallbackColor FallbackColor, FSamplerStateRHIParamRef* FilterOverrideArray);
+template void FPostProcessPassParameters::SetPS(FRHICommandList& RHICmdList, const FPixelShaderRHIParamRef& ShaderRHI, const FRenderingCompositePassContext& Context, FRHISamplerState* Filter, EFallbackColor FallbackColor, FRHISamplerState** FilterOverrideArray);
+template void FPostProcessPassParameters::SetPS(FRHICommandListImmediate& RHICmdList, const FPixelShaderRHIParamRef& ShaderRHI, const FRenderingCompositePassContext& Context, FRHISamplerState* Filter, EFallbackColor FallbackColor, FRHISamplerState** FilterOverrideArray);
 
 template< typename TRHICmdList >
-void FPostProcessPassParameters::SetCS(const FComputeShaderRHIParamRef& ShaderRHI, const FRenderingCompositePassContext& Context, TRHICmdList& RHICmdList, FSamplerStateRHIParamRef Filter, EFallbackColor FallbackColor, FSamplerStateRHIParamRef* FilterOverrideArray)
+void FPostProcessPassParameters::SetCS(const FComputeShaderRHIParamRef& ShaderRHI, const FRenderingCompositePassContext& Context, TRHICmdList& RHICmdList, FRHISamplerState* Filter, EFallbackColor FallbackColor, FRHISamplerState** FilterOverrideArray)
 {
 	Set(RHICmdList, ShaderRHI, Context, Filter, FallbackColor, FilterOverrideArray);
 }
@@ -869,21 +869,21 @@ template void FPostProcessPassParameters::SetCS< FRHICommandListImmediate >(
 	const FComputeShaderRHIParamRef& ShaderRHI,
 	const FRenderingCompositePassContext& Context,
 	FRHICommandListImmediate& RHICmdList,
-	FSamplerStateRHIParamRef Filter,
+	FRHISamplerState* Filter,
 	EFallbackColor FallbackColor,
-	FSamplerStateRHIParamRef* FilterOverrideArray
+	FRHISamplerState** FilterOverrideArray
 	);
 
 template void FPostProcessPassParameters::SetCS< FRHIAsyncComputeCommandListImmediate >(
 	const FComputeShaderRHIParamRef& ShaderRHI,
 	const FRenderingCompositePassContext& Context,
 	FRHIAsyncComputeCommandListImmediate& RHICmdList,
-	FSamplerStateRHIParamRef Filter,
+	FRHISamplerState* Filter,
 	EFallbackColor FallbackColor,
-	FSamplerStateRHIParamRef* FilterOverrideArray
+	FRHISamplerState** FilterOverrideArray
 	);
 
-void FPostProcessPassParameters::SetVS(const FVertexShaderRHIParamRef& ShaderRHI, const FRenderingCompositePassContext& Context, FSamplerStateRHIParamRef Filter, EFallbackColor FallbackColor, FSamplerStateRHIParamRef* FilterOverrideArray)
+void FPostProcessPassParameters::SetVS(const FVertexShaderRHIParamRef& ShaderRHI, const FRenderingCompositePassContext& Context, FRHISamplerState* Filter, EFallbackColor FallbackColor, FRHISamplerState** FilterOverrideArray)
 {
 	Set(Context.RHICmdList, ShaderRHI, Context, Filter, FallbackColor, FilterOverrideArray);
 }
@@ -893,9 +893,9 @@ void FPostProcessPassParameters::Set(
 	TRHICmdList& RHICmdList,
 	const TShaderRHIParamRef& ShaderRHI,
 	const FRenderingCompositePassContext& Context,
-	FSamplerStateRHIParamRef Filter,
+	FRHISamplerState* Filter,
 	EFallbackColor FallbackColor,
-	FSamplerStateRHIParamRef* FilterOverrideArray)
+	FRHISamplerState** FilterOverrideArray)
 {
 	// assuming all outputs have the same size
 	FRenderingCompositeOutput* Output = Context.Pass->GetOutput(ePId_Output0);
@@ -984,7 +984,7 @@ void FPostProcessPassParameters::Set(
 			InputPooledElement = Input->RequestInput();
 		}
 
-		FSamplerStateRHIParamRef LocalFilter = FilterOverrideArray ? FilterOverrideArray[Id] : Filter;
+		FRHISamplerState* LocalFilter = FilterOverrideArray ? FilterOverrideArray[Id] : Filter;
 
 		if(InputPooledElement)
 		{
@@ -1033,9 +1033,9 @@ void FPostProcessPassParameters::Set(
 		TRHICmdList& RHICmdList,						\
 		const TShaderRHIParamRef& ShaderRHI,			\
 		const FRenderingCompositePassContext& Context,	\
-		FSamplerStateRHIParamRef Filter,				\
+		FRHISamplerState* Filter,				\
 		EFallbackColor FallbackColor,					\
-		FSamplerStateRHIParamRef* FilterOverrideArray	\
+		FRHISamplerState** FilterOverrideArray	\
 	);
 
 IMPLEMENT_POST_PROCESS_PARAM_SET( FVertexShaderRHIParamRef, FRHICommandListImmediate );
diff --git a/Engine/Source/Runtime/Renderer/Private/SceneViewFamilyBlackboard.cpp b/Engine/Source/Runtime/Renderer/Private/SceneViewFamilyBlackboard.cpp
--- a/Engine/Source/Runtime/Renderer/Private/SceneViewFamilyBlackboard.cpp
+++ b/Engine/Source/Runtime/Renderer/Private/SceneViewFamilyBlackboard.cpp
@@ -45,7 +45,7 @@ void SetupSceneViewFamilyBlackboard(
 
 void SetupSceneTextureSamplers(FSceneTextureSamplerParameters* OutSamplers)
 {
-	FSamplerStateRHIParamRef Sampler = TStaticSamplerState<SF_Point>::GetRHI();
+	FRHISamplerState* Sampler = TStaticSamplerState<SF_Point>::GetRHI();
 	OutSamplers->SceneDepthBufferSampler = Sampler;
 	OutSamplers->SceneVelocityBufferSampler = Sampler;
 	OutSamplers->SceneGBufferASampler = Sampler;
diff --git a/Engine/Source/Runtime/Renderer/Private/ShadowRendering.h b/Engine/Source/Runtime/Renderer/Private/ShadowRendering.h
--- a/Engine/Source/Runtime/Renderer/Private/ShadowRendering.h
+++ b/Engine/Source/Runtime/Renderer/Private/ShadowRendering.h
@@ -831,7 +831,7 @@ public:
 			ShadowDepthTextureValue = GSystemTextures.BlackDummy->GetRenderTargetItem().ShaderResourceTexture.GetReference();
 		}
 			
-		FSamplerStateRHIParamRef DepthSamplerState = TStaticSamplerState<SF_Point,AM_Clamp,AM_Clamp,AM_Clamp>::GetRHI();
+		FRHISamplerState* DepthSamplerState = TStaticSamplerState<SF_Point,AM_Clamp,AM_Clamp,AM_Clamp>::GetRHI();
 
 		SetTextureParameter(RHICmdList, ShaderRHI, ShadowDepthTexture, ShadowDepthTextureSampler, DepthSamplerState, ShadowDepthTextureValue);		
 
diff --git a/Engine/Source/Runtime/Renderer/Public/MeshDrawShaderBindings.h b/Engine/Source/Runtime/Renderer/Public/MeshDrawShaderBindings.h
--- a/Engine/Source/Runtime/Renderer/Public/MeshDrawShaderBindings.h
+++ b/Engine/Source/Runtime/Renderer/Public/MeshDrawShaderBindings.h
@@ -73,7 +73,7 @@ protected:
 	inline uint32 GetSRVOffset() const
 	{
 		return GetSamplerOffset() 
-			+ ParameterMapInfo.TextureSamplers.Num() * sizeof(FSamplerStateRHIParamRef);
+			+ ParameterMapInfo.TextureSamplers.Num() * sizeof(FRHISamplerState*);
 	}
 
 	inline uint32 GetSRVTypeOffset() const
@@ -153,7 +153,7 @@ public:
 	void AddTexture(
 		FShaderResourceParameter TextureParameter,
 		FShaderResourceParameter SamplerParameter,
-		FSamplerStateRHIParamRef SamplerStateRHI,
+		FRHISamplerState* SamplerStateRHI,
 		FTextureRHIParamRef TextureRHI)
 	{
 		checkfSlow(TextureParameter.IsInitialized(), TEXT("Parameter was not serialized"));
@@ -223,10 +223,10 @@ private:
 		return (FUniformBufferRHIParamRef*)(Data + GetUniformBufferOffset());
 	}
 
-	inline FSamplerStateRHIParamRef* GetSamplerStart() const
+	inline FRHISamplerState** GetSamplerStart() const
 	{
 		uint8* SamplerDataStart = Data + GetSamplerOffset();
-		return (FSamplerStateRHIParamRef*)SamplerDataStart;
+		return (FRHISamplerState**)SamplerDataStart;
 	}
 
 	inline FRHIResource** GetSRVStart() const
@@ -280,7 +280,7 @@ private:
 		checkfSlow(FoundIndex >= 0, TEXT("Attempted to set a uniform buffer at BaseIndex %u which was never in the shader's parameter map."), BaseIndex);
 	}
 
-	inline void WriteBindingSampler(FSamplerStateRHIParamRef Value, uint32 BaseIndex)
+	inline void WriteBindingSampler(FRHISamplerState* Value, uint32 BaseIndex)
 	{
 		int32 FoundIndex = -1;
 
diff --git a/Engine/Source/Runtime/Renderer/Public/PostProcessParameters.h b/Engine/Source/Runtime/Renderer/Public/PostProcessParameters.h
--- a/Engine/Source/Runtime/Renderer/Public/PostProcessParameters.h
+++ b/Engine/Source/Runtime/Renderer/Public/PostProcessParameters.h
@@ -75,14 +75,14 @@ struct FPostProcessPassParameters
 
 	/** Set the pixel shader parameter values. */
 	template <typename TRHICmdList>
-	void SetPS(TRHICmdList& RHICmdList, const FPixelShaderRHIParamRef& ShaderRHI, const FRenderingCompositePassContext& Context, FSamplerStateRHIParamRef Filter = TStaticSamplerState<>::GetRHI(), EFallbackColor FallbackColor = eFC_0000, FSamplerStateRHIParamRef* FilterOverrideArray = 0);
+	void SetPS(TRHICmdList& RHICmdList, const FPixelShaderRHIParamRef& ShaderRHI, const FRenderingCompositePassContext& Context, FRHISamplerState* Filter = TStaticSamplerState<>::GetRHI(), EFallbackColor FallbackColor = eFC_0000, FRHISamplerState** FilterOverrideArray = 0);
 
 	/** Set the compute shader parameter values. */
 	template< typename TRHICmdList >
-	void SetCS(const FComputeShaderRHIParamRef& ShaderRHI, const FRenderingCompositePassContext& Context, TRHICmdList& RHICmdList, FSamplerStateRHIParamRef Filter = TStaticSamplerState<>::GetRHI(), EFallbackColor FallbackColor = eFC_0000, FSamplerStateRHIParamRef* FilterOverrideArray = 0);
+	void SetCS(const FComputeShaderRHIParamRef& ShaderRHI, const FRenderingCompositePassContext& Context, TRHICmdList& RHICmdList, FRHISamplerState* Filter = TStaticSamplerState<>::GetRHI(), EFallbackColor FallbackColor = eFC_0000, FRHISamplerState** FilterOverrideArray = 0);
 
 	/** Set the vertex shader parameter values. */
-	void SetVS(const FVertexShaderRHIParamRef& ShaderRHI, const FRenderingCompositePassContext& Context, FSamplerStateRHIParamRef Filter = TStaticSamplerState<>::GetRHI(), EFallbackColor FallbackColor = eFC_0000, FSamplerStateRHIParamRef* FilterOverrideArray = 0);
+	void SetVS(const FVertexShaderRHIParamRef& ShaderRHI, const FRenderingCompositePassContext& Context, FRHISamplerState* Filter = TStaticSamplerState<>::GetRHI(), EFallbackColor FallbackColor = eFC_0000, FRHISamplerState** FilterOverrideArray = 0);
 
 	/** Serializer. */
 	friend FArchive& operator<<(FArchive& Ar,FPostProcessPassParameters& P);
@@ -107,8 +107,8 @@ public:
 		TRHICmdList& RHICmdList,
 		const TShaderRHIParamRef& ShaderRHI,
 		const FRenderingCompositePassContext& Context,
-		FSamplerStateRHIParamRef Filter,
+		FRHISamplerState* Filter,
 		EFallbackColor FallbackColor,
-		FSamplerStateRHIParamRef* FilterOverrideArray = 0
+		FRHISamplerState** FilterOverrideArray = 0
 	);
 };
diff --git a/Engine/Source/Runtime/SlateRHIRenderer/Private/SlateRHIRenderingPolicy.cpp b/Engine/Source/Runtime/SlateRHIRenderer/Private/SlateRHIRenderingPolicy.cpp
--- a/Engine/Source/Runtime/SlateRHIRenderer/Private/SlateRHIRenderingPolicy.cpp
+++ b/Engine/Source/Runtime/SlateRHIRenderer/Private/SlateRHIRenderingPolicy.cpp
@@ -879,7 +879,7 @@ void FSlateRHIRenderingPolicy::DrawElements(
 				}
 #endif
 
-				FSamplerStateRHIParamRef SamplerState = BilinearClamp;
+				FRHISamplerState* SamplerState = BilinearClamp;
 				FTextureRHIParamRef TextureRHI = GWhiteTexture->TextureRHI;
 				if (ShaderResource)
 				{
diff --git a/Engine/Source/Runtime/UtilityShaders/Public/OculusShaders.h b/Engine/Source/Runtime/UtilityShaders/Public/OculusShaders.h
--- a/Engine/Source/Runtime/UtilityShaders/Public/OculusShaders.h
+++ b/Engine/Source/Runtime/UtilityShaders/Public/OculusShaders.h
@@ -72,7 +72,7 @@ public:
 		SetTextureParameter(RHICmdList, GetPixelShader(), InTexture, InTextureSampler, Texture);
 	}
 
-	UTILITYSHADERS_API void SetParameters(FRHICommandList& RHICmdList, FSamplerStateRHIParamRef SamplerStateRHI, FTextureRHIParamRef TextureRHI)
+	UTILITYSHADERS_API void SetParameters(FRHICommandList& RHICmdList, FRHISamplerState* SamplerStateRHI, FTextureRHIParamRef TextureRHI)
 	{
 		SetTextureParameter(RHICmdList, GetPixelShader(), InTexture, InTextureSampler, SamplerStateRHI, TextureRHI);
 	}
@@ -116,7 +116,7 @@ public:
 		SetShaderValue(RHICmdList, GetPixelShader(), InFaceIndexParameter, FaceIndex);
 	}
 
-	void SetParameters(FRHICommandList& RHICmdList, FSamplerStateRHIParamRef SamplerStateRHI, FTextureRHIParamRef TextureRHI, int FaceIndex)
+	void SetParameters(FRHICommandList& RHICmdList, FRHISamplerState* SamplerStateRHI, FTextureRHIParamRef TextureRHI, int FaceIndex)
 	{
 		SetTextureParameter(RHICmdList, GetPixelShader(), InTexture, InTextureSampler, SamplerStateRHI, TextureRHI);
 		SetShaderValue(RHICmdList, GetPixelShader(), InFaceIndexParameter, FaceIndex);
diff --git a/Engine/Source/Runtime/VulkanRHI/Private/VulkanCommands.cpp b/Engine/Source/Runtime/VulkanRHI/Private/VulkanCommands.cpp
--- a/Engine/Source/Runtime/VulkanRHI/Private/VulkanCommands.cpp
+++ b/Engine/Source/Runtime/VulkanRHI/Private/VulkanCommands.cpp
@@ -300,14 +300,14 @@ void FVulkanCommandListContext::RHISetShaderResourceViewParameter(FComputeShader
 	PendingComputeState->SetSRVForStage(TextureIndex, SRV);
 }
 
-void FVulkanCommandListContext::RHISetShaderSampler(FVertexShaderRHIParamRef VertexShaderRHI, uint32 SamplerIndex, FSamplerStateRHIParamRef NewStateRHI)
+void FVulkanCommandListContext::RHISetShaderSampler(FVertexShaderRHIParamRef VertexShaderRHI, uint32 SamplerIndex, FRHISamplerState* NewStateRHI)
 {
 	check(PendingGfxState->GetCurrentShaderKey(ShaderStage::Vertex) == GetShaderKey(VertexShaderRHI));
 	FVulkanSamplerState* Sampler = ResourceCast(NewStateRHI);
 	PendingGfxState->SetSamplerStateForStage(ShaderStage::Vertex, SamplerIndex, Sampler);
 }
 
-void FVulkanCommandListContext::RHISetShaderSampler(FHullShaderRHIParamRef HullShaderRHI, uint32 SamplerIndex, FSamplerStateRHIParamRef NewStateRHI)
+void FVulkanCommandListContext::RHISetShaderSampler(FHullShaderRHIParamRef HullShaderRHI, uint32 SamplerIndex, FRHISamplerState* NewStateRHI)
 {
 #if PLATFORM_SUPPORTS_TESSELLATION_SHADERS
 	check(PendingGfxState->GetCurrentShaderKey(ShaderStage::Hull) == GetShaderKey(HullShaderRHI));
@@ -318,7 +318,7 @@ void FVulkanCommandListContext::RHISetShaderSampler(FHullShaderRHIParamRef HullS
 #endif
 }
 
-void FVulkanCommandListContext::RHISetShaderSampler(FDomainShaderRHIParamRef DomainShaderRHI, uint32 SamplerIndex, FSamplerStateRHIParamRef NewStateRHI)
+void FVulkanCommandListContext::RHISetShaderSampler(FDomainShaderRHIParamRef DomainShaderRHI, uint32 SamplerIndex, FRHISamplerState* NewStateRHI)
 {
 #if PLATFORM_SUPPORTS_TESSELLATION_SHADERS
 	check(PendingGfxState->GetCurrentShaderKey(ShaderStage::Domain) == GetShaderKey(DomainShaderRHI));
@@ -329,7 +329,7 @@ void FVulkanCommandListContext::RHISetShaderSampler(FDomainShaderRHIParamRef Dom
 #endif
 }
 
-void FVulkanCommandListContext::RHISetShaderSampler(FGeometryShaderRHIParamRef GeometryShaderRHI, uint32 SamplerIndex, FSamplerStateRHIParamRef NewStateRHI)
+void FVulkanCommandListContext::RHISetShaderSampler(FGeometryShaderRHIParamRef GeometryShaderRHI, uint32 SamplerIndex, FRHISamplerState* NewStateRHI)
 {
 #if VULKAN_SUPPORTS_GEOMETRY_SHADERS
 	check(PendingGfxState->GetCurrentShaderKey(ShaderStage::Geometry) == GetShaderKey(GeometryShaderRHI));
@@ -340,14 +340,14 @@ void FVulkanCommandListContext::RHISetShaderSampler(FGeometryShaderRHIParamRef G
 #endif
 }
 
-void FVulkanCommandListContext::RHISetShaderSampler(FPixelShaderRHIParamRef PixelShaderRHI, uint32 SamplerIndex, FSamplerStateRHIParamRef NewStateRHI)
+void FVulkanCommandListContext::RHISetShaderSampler(FPixelShaderRHIParamRef PixelShaderRHI, uint32 SamplerIndex, FRHISamplerState* NewStateRHI)
 {
 	check(PendingGfxState->GetCurrentShaderKey(ShaderStage::Pixel) == GetShaderKey(PixelShaderRHI));
 	FVulkanSamplerState* Sampler = ResourceCast(NewStateRHI);
 	PendingGfxState->SetSamplerStateForStage(ShaderStage::Pixel, SamplerIndex, Sampler);
 }
 
-void FVulkanCommandListContext::RHISetShaderSampler(FComputeShaderRHIParamRef ComputeShaderRHI, uint32 SamplerIndex, FSamplerStateRHIParamRef NewStateRHI)
+void FVulkanCommandListContext::RHISetShaderSampler(FComputeShaderRHIParamRef ComputeShaderRHI, uint32 SamplerIndex, FRHISamplerState* NewStateRHI)
 {
 	FVulkanComputeShader* ComputeShader = ResourceCast(ComputeShaderRHI);
 	check(PendingComputeState->GetCurrentShader() == ComputeShader);
diff --git a/Engine/Source/Runtime/VulkanRHI/Private/VulkanContext.h b/Engine/Source/Runtime/VulkanRHI/Private/VulkanContext.h
--- a/Engine/Source/Runtime/VulkanRHI/Private/VulkanContext.h
+++ b/Engine/Source/Runtime/VulkanRHI/Private/VulkanContext.h
@@ -179,12 +179,12 @@ public:
 	virtual void RHISetShaderTexture(FGeometryShaderRHIParamRef GeometryShader, uint32 TextureIndex, FTextureRHIParamRef NewTexture) final override;
 	virtual void RHISetShaderTexture(FPixelShaderRHIParamRef PixelShader, uint32 TextureIndex, FTextureRHIParamRef NewTexture) final override;
 	virtual void RHISetShaderTexture(FComputeShaderRHIParamRef PixelShader, uint32 TextureIndex, FTextureRHIParamRef NewTexture) final override;
-	virtual void RHISetShaderSampler(FComputeShaderRHIParamRef ComputeShader, uint32 SamplerIndex, FSamplerStateRHIParamRef NewState) final override;
-	virtual void RHISetShaderSampler(FVertexShaderRHIParamRef VertexShader, uint32 SamplerIndex, FSamplerStateRHIParamRef NewState) final override;
-	virtual void RHISetShaderSampler(FGeometryShaderRHIParamRef GeometryShader, uint32 SamplerIndex, FSamplerStateRHIParamRef NewState) final override;
-	virtual void RHISetShaderSampler(FDomainShaderRHIParamRef DomainShader, uint32 SamplerIndex, FSamplerStateRHIParamRef NewState) final override;
-	virtual void RHISetShaderSampler(FHullShaderRHIParamRef HullShader, uint32 SamplerIndex, FSamplerStateRHIParamRef NewState) final override;
-	virtual void RHISetShaderSampler(FPixelShaderRHIParamRef PixelShader, uint32 SamplerIndex, FSamplerStateRHIParamRef NewState) final override;
+	virtual void RHISetShaderSampler(FComputeShaderRHIParamRef ComputeShader, uint32 SamplerIndex, FRHISamplerState* NewState) final override;
+	virtual void RHISetShaderSampler(FVertexShaderRHIParamRef VertexShader, uint32 SamplerIndex, FRHISamplerState* NewState) final override;
+	virtual void RHISetShaderSampler(FGeometryShaderRHIParamRef GeometryShader, uint32 SamplerIndex, FRHISamplerState* NewState) final override;
+	virtual void RHISetShaderSampler(FDomainShaderRHIParamRef DomainShader, uint32 SamplerIndex, FRHISamplerState* NewState) final override;
+	virtual void RHISetShaderSampler(FHullShaderRHIParamRef HullShader, uint32 SamplerIndex, FRHISamplerState* NewState) final override;
+	virtual void RHISetShaderSampler(FPixelShaderRHIParamRef PixelShader, uint32 SamplerIndex, FRHISamplerState* NewState) final override;
 	virtual void RHISetUAVParameter(FComputeShaderRHIParamRef ComputeShader, uint32 UAVIndex, FUnorderedAccessViewRHIParamRef UAV) final override;
 	virtual void RHISetUAVParameter(FComputeShaderRHIParamRef ComputeShader, uint32 UAVIndex, FUnorderedAccessViewRHIParamRef UAV, uint32 InitialCount) final override;
 	virtual void RHISetShaderResourceViewParameter(FPixelShaderRHIParamRef PixelShader, uint32 SamplerIndex, FShaderResourceViewRHIParamRef SRV) final override;
diff --git a/Engine/Source/Runtime/VulkanRHI/Private/VulkanDescriptorSets.h b/Engine/Source/Runtime/VulkanRHI/Private/VulkanDescriptorSets.h
--- a/Engine/Source/Runtime/VulkanRHI/Private/VulkanDescriptorSets.h
+++ b/Engine/Source/Runtime/VulkanRHI/Private/VulkanDescriptorSets.h
@@ -325,9 +325,9 @@ public:
 	void ProcessBindingsForStage(VkShaderStageFlagBits StageFlags, ShaderStage::EStage DescSetStage, const FVulkanShaderHeader& CodeHeader, FUniformBufferGatherInfo& OutUBGatherInfo) const;
 
 	template<bool bIsCompute>
-	void FinalizeBindings(const FUniformBufferGatherInfo& UBGatherInfo, const TArrayView<const FSamplerStateRHIParamRef>& ImmutableSamplers);
+	void FinalizeBindings(const FUniformBufferGatherInfo& UBGatherInfo, const TArrayView<FRHISamplerState*>& ImmutableSamplers);
 
-	void GenerateHash(const TArrayView<const FSamplerStateRHIParamRef>& ImmutableSamplers);
+	void GenerateHash(const TArrayView<FRHISamplerState*>& ImmutableSamplers);
 
 	friend uint32 GetTypeHash(const FVulkanDescriptorSetsLayoutInfo& In)
 	{
diff --git a/Engine/Source/Runtime/VulkanRHI/Private/VulkanPipeline.cpp b/Engine/Source/Runtime/VulkanRHI/Private/VulkanPipeline.cpp
--- a/Engine/Source/Runtime/VulkanRHI/Private/VulkanPipeline.cpp
+++ b/Engine/Source/Runtime/VulkanRHI/Private/VulkanPipeline.cpp
@@ -1436,7 +1436,7 @@ FVulkanGfxLayout* FVulkanPipelineStateCacheManager::GetOrGenerateGfxLayout(const
 #endif
 	// Second pass
 	const int32 NumImmutableSamplers = PSOInitializer.ImmutableSamplerState.ImmutableSamplers.Num();
-	TArrayView<const FSamplerStateRHIParamRef> ImmutableSamplers(NumImmutableSamplers > 0 ? &PSOInitializer.ImmutableSamplerState.ImmutableSamplers[0] : nullptr, NumImmutableSamplers);
+	TArrayView<FRHISamplerState*> ImmutableSamplers(NumImmutableSamplers > 0 ? &(FRHISamplerState*)PSOInitializer.ImmutableSamplerState.ImmutableSamplers[0] : nullptr, NumImmutableSamplers);
 	DescriptorSetLayoutInfo.FinalizeBindings<false>(UBGatherInfo, ImmutableSamplers);
 
 	FVulkanLayout* Layout = FindOrAddLayout(DescriptorSetLayoutInfo, true);
@@ -1814,7 +1814,7 @@ FVulkanComputePipeline* FVulkanPipelineStateCacheManager::CreateComputePipelineF
 	const FVulkanShaderHeader& CSHeader = Shader->GetCodeHeader();
 	FUniformBufferGatherInfo UBGatherInfo;
 	DescriptorSetLayoutInfo.ProcessBindingsForStage(VK_SHADER_STAGE_COMPUTE_BIT, ShaderStage::Compute, CSHeader, UBGatherInfo);
-	DescriptorSetLayoutInfo.FinalizeBindings<true>(UBGatherInfo, TArrayView<const FSamplerStateRHIParamRef>());
+	DescriptorSetLayoutInfo.FinalizeBindings<true>(UBGatherInfo, TArrayView<FRHISamplerState*>());
 	FVulkanLayout* Layout = FindOrAddLayout(DescriptorSetLayoutInfo, true);
 	FVulkanComputeLayout* ComputeLayout = (FVulkanComputeLayout*)Layout;
 	if (!ComputeLayout->ComputePipelineDescriptorInfo.IsInitialized())
diff --git a/Engine/Source/Runtime/VulkanRHI/Private/VulkanRHI.cpp b/Engine/Source/Runtime/VulkanRHI/Private/VulkanRHI.cpp
--- a/Engine/Source/Runtime/VulkanRHI/Private/VulkanRHI.cpp
+++ b/Engine/Source/Runtime/VulkanRHI/Private/VulkanRHI.cpp
@@ -1201,7 +1201,7 @@ void FVulkanDescriptorSetsLayoutInfo::AddDescriptor(int32 DescriptorSetIndex, co
 	}
 }
 
-void FVulkanDescriptorSetsLayoutInfo::GenerateHash(const TArrayView<const FSamplerStateRHIParamRef>& InImmutableSamplers)
+void FVulkanDescriptorSetsLayoutInfo::GenerateHash(const TArrayView<FRHISamplerState*>& InImmutableSamplers)
 {
 	const int32 LayoutCount = SetLayouts.Num();
 	Hash = FCrc::MemCrc32(&TypesUsageID, sizeof(uint32), LayoutCount);
diff --git a/Engine/Source/Runtime/VulkanRHI/Private/VulkanShaders.cpp b/Engine/Source/Runtime/VulkanRHI/Private/VulkanShaders.cpp
--- a/Engine/Source/Runtime/VulkanRHI/Private/VulkanShaders.cpp
+++ b/Engine/Source/Runtime/VulkanRHI/Private/VulkanShaders.cpp
@@ -403,7 +403,7 @@ void FVulkanDescriptorSetsLayoutInfo::ProcessBindingsForStage(VkShaderStageFlagB
 }
 
 template<bool bIsCompute>
-void FVulkanDescriptorSetsLayoutInfo::FinalizeBindings(const FUniformBufferGatherInfo& UBGatherInfo, const TArrayView<const FSamplerStateRHIParamRef>& ImmutableSamplers)
+void FVulkanDescriptorSetsLayoutInfo::FinalizeBindings(const FUniformBufferGatherInfo& UBGatherInfo, const TArrayView<FRHISamplerState*>& ImmutableSamplers)
 {
 	checkSlow(RemappingInfo.IsEmpty());
 
@@ -532,7 +532,7 @@ void FVulkanDescriptorSetsLayoutInfo::FinalizeBindings(const FUniformBufferGathe
 					{
 						if (CurrentImmutableSampler < ImmutableSamplers.Num())
 						{
-							const FVulkanSamplerState* SamplerState = ResourceCast(ImmutableSamplers[CurrentImmutableSampler]);
+							FVulkanSamplerState* SamplerState = ResourceCast(ImmutableSamplers[CurrentImmutableSampler]);
 							if (SamplerState && SamplerState->Sampler != VK_NULL_HANDLE)
 							{
 								Binding.pImmutableSamplers = &SamplerState->Sampler;
@@ -679,5 +679,5 @@ FBoundShaderStateRHIRef FVulkanDynamicRHI::RHICreateBoundShaderState(
 }
 
 
-template void FVulkanDescriptorSetsLayoutInfo::FinalizeBindings<true>(const FUniformBufferGatherInfo& UBGatherInfo, const TArrayView<const FSamplerStateRHIParamRef>& ImmutableSamplers);
-template void FVulkanDescriptorSetsLayoutInfo::FinalizeBindings<false>(const FUniformBufferGatherInfo& UBGatherInfo, const TArrayView<const FSamplerStateRHIParamRef>& ImmutableSamplers);
+template void FVulkanDescriptorSetsLayoutInfo::FinalizeBindings<true>(const FUniformBufferGatherInfo& UBGatherInfo, const TArrayView<FRHISamplerState*>& ImmutableSamplers);
+template void FVulkanDescriptorSetsLayoutInfo::FinalizeBindings<false>(const FUniformBufferGatherInfo& UBGatherInfo, const TArrayView<FRHISamplerState*>& ImmutableSamplers);
diff --git a/Engine/Source/Runtime/Windows/D3D11RHI/Private/D3D11Commands.cpp b/Engine/Source/Runtime/Windows/D3D11RHI/Private/D3D11Commands.cpp
--- a/Engine/Source/Runtime/Windows/D3D11RHI/Private/D3D11Commands.cpp
+++ b/Engine/Source/Runtime/Windows/D3D11RHI/Private/D3D11Commands.cpp
@@ -591,7 +591,7 @@ void FD3D11DynamicRHI::RHISetShaderResourceViewParameter(FGeometryShaderRHIParam
 	SetShaderResourceView<SF_Geometry>(Resource, D3D11SRV, TextureIndex, NAME_None);
 }
 
-void FD3D11DynamicRHI::RHISetShaderSampler(FVertexShaderRHIParamRef VertexShaderRHI,uint32 SamplerIndex,FSamplerStateRHIParamRef NewStateRHI)
+void FD3D11DynamicRHI::RHISetShaderSampler(FVertexShaderRHIParamRef VertexShaderRHI,uint32 SamplerIndex, FRHISamplerState* NewStateRHI)
 {
 	VALIDATE_BOUND_SHADER(VertexShaderRHI);
 
@@ -602,7 +602,7 @@ void FD3D11DynamicRHI::RHISetShaderSampler(FVertexShaderRHIParamRef VertexShader
 	StateCache.SetSamplerState<SF_Vertex>(StateResource, SamplerIndex);
 }
 
-void FD3D11DynamicRHI::RHISetShaderSampler(FHullShaderRHIParamRef HullShaderRHI,uint32 SamplerIndex,FSamplerStateRHIParamRef NewStateRHI)
+void FD3D11DynamicRHI::RHISetShaderSampler(FHullShaderRHIParamRef HullShaderRHI,uint32 SamplerIndex, FRHISamplerState* NewStateRHI)
 {
 	VALIDATE_BOUND_SHADER(HullShaderRHI);
 
@@ -613,7 +613,7 @@ void FD3D11DynamicRHI::RHISetShaderSampler(FHullShaderRHIParamRef HullShaderRHI,
 	StateCache.SetSamplerState<SF_Hull>(StateResource, SamplerIndex);
 }
 
-void FD3D11DynamicRHI::RHISetShaderSampler(FDomainShaderRHIParamRef DomainShaderRHI,uint32 SamplerIndex,FSamplerStateRHIParamRef NewStateRHI)
+void FD3D11DynamicRHI::RHISetShaderSampler(FDomainShaderRHIParamRef DomainShaderRHI,uint32 SamplerIndex, FRHISamplerState* NewStateRHI)
 {
 	VALIDATE_BOUND_SHADER(DomainShaderRHI);
 
@@ -624,7 +624,7 @@ void FD3D11DynamicRHI::RHISetShaderSampler(FDomainShaderRHIParamRef DomainShader
 	StateCache.SetSamplerState<SF_Domain>(StateResource, SamplerIndex);
 }
 
-void FD3D11DynamicRHI::RHISetShaderSampler(FGeometryShaderRHIParamRef GeometryShaderRHI,uint32 SamplerIndex,FSamplerStateRHIParamRef NewStateRHI)
+void FD3D11DynamicRHI::RHISetShaderSampler(FGeometryShaderRHIParamRef GeometryShaderRHI,uint32 SamplerIndex, FRHISamplerState* NewStateRHI)
 {
 	VALIDATE_BOUND_SHADER(GeometryShaderRHI);
 
@@ -635,7 +635,7 @@ void FD3D11DynamicRHI::RHISetShaderSampler(FGeometryShaderRHIParamRef GeometrySh
 	StateCache.SetSamplerState<SF_Geometry>(StateResource, SamplerIndex);
 }
 
-void FD3D11DynamicRHI::RHISetShaderSampler(FPixelShaderRHIParamRef PixelShaderRHI,uint32 SamplerIndex,FSamplerStateRHIParamRef NewStateRHI)
+void FD3D11DynamicRHI::RHISetShaderSampler(FPixelShaderRHIParamRef PixelShaderRHI,uint32 SamplerIndex, FRHISamplerState* NewStateRHI)
 {
 	VALIDATE_BOUND_SHADER(PixelShaderRHI);
 
@@ -646,7 +646,7 @@ void FD3D11DynamicRHI::RHISetShaderSampler(FPixelShaderRHIParamRef PixelShaderRH
 	StateCache.SetSamplerState<SF_Pixel>(StateResource, SamplerIndex);
 }
 
-void FD3D11DynamicRHI::RHISetShaderSampler(FComputeShaderRHIParamRef ComputeShaderRHI,uint32 SamplerIndex,FSamplerStateRHIParamRef NewStateRHI)
+void FD3D11DynamicRHI::RHISetShaderSampler(FComputeShaderRHIParamRef ComputeShaderRHI,uint32 SamplerIndex, FRHISamplerState* NewStateRHI)
 {
 	//VALIDATE_BOUND_SHADER(ComputeShaderRHI);
 	FD3D11ComputeShader* ComputeShader = ResourceCast(ComputeShaderRHI);
diff --git a/Engine/Source/Runtime/Windows/D3D11RHI/Private/D3D11RHIPrivate.h b/Engine/Source/Runtime/Windows/D3D11RHI/Private/D3D11RHIPrivate.h
--- a/Engine/Source/Runtime/Windows/D3D11RHI/Private/D3D11RHIPrivate.h
+++ b/Engine/Source/Runtime/Windows/D3D11RHI/Private/D3D11RHIPrivate.h
@@ -535,12 +535,12 @@ public:
 	virtual void RHISetShaderTexture(FGeometryShaderRHIParamRef GeometryShader, uint32 TextureIndex, FTextureRHIParamRef NewTexture) final override;
 	virtual void RHISetShaderTexture(FPixelShaderRHIParamRef PixelShader, uint32 TextureIndex, FTextureRHIParamRef NewTexture) final override;
 	virtual void RHISetShaderTexture(FComputeShaderRHIParamRef PixelShader, uint32 TextureIndex, FTextureRHIParamRef NewTexture) final override;
-	virtual void RHISetShaderSampler(FComputeShaderRHIParamRef ComputeShader, uint32 SamplerIndex, FSamplerStateRHIParamRef NewState) final override;
-	virtual void RHISetShaderSampler(FVertexShaderRHIParamRef VertexShader, uint32 SamplerIndex, FSamplerStateRHIParamRef NewState) final override;
-	virtual void RHISetShaderSampler(FGeometryShaderRHIParamRef GeometryShader, uint32 SamplerIndex, FSamplerStateRHIParamRef NewState) final override;
-	virtual void RHISetShaderSampler(FDomainShaderRHIParamRef DomainShader, uint32 SamplerIndex, FSamplerStateRHIParamRef NewState) final override;
-	virtual void RHISetShaderSampler(FHullShaderRHIParamRef HullShader, uint32 SamplerIndex, FSamplerStateRHIParamRef NewState) final override;
-	virtual void RHISetShaderSampler(FPixelShaderRHIParamRef PixelShader, uint32 SamplerIndex, FSamplerStateRHIParamRef NewState) final override;
+	virtual void RHISetShaderSampler(FComputeShaderRHIParamRef ComputeShader, uint32 SamplerIndex, FRHISamplerState* NewState) final override;
+	virtual void RHISetShaderSampler(FVertexShaderRHIParamRef VertexShader, uint32 SamplerIndex, FRHISamplerState* NewState) final override;
+	virtual void RHISetShaderSampler(FGeometryShaderRHIParamRef GeometryShader, uint32 SamplerIndex, FRHISamplerState* NewState) final override;
+	virtual void RHISetShaderSampler(FDomainShaderRHIParamRef DomainShader, uint32 SamplerIndex, FRHISamplerState* NewState) final override;
+	virtual void RHISetShaderSampler(FHullShaderRHIParamRef HullShader, uint32 SamplerIndex, FRHISamplerState* NewState) final override;
+	virtual void RHISetShaderSampler(FPixelShaderRHIParamRef PixelShader, uint32 SamplerIndex, FRHISamplerState* NewState) final override;
 	virtual void RHISetUAVParameter(FComputeShaderRHIParamRef ComputeShader, uint32 UAVIndex, FUnorderedAccessViewRHIParamRef UAV) final override;
 	virtual void RHISetUAVParameter(FComputeShaderRHIParamRef ComputeShader, uint32 UAVIndex, FUnorderedAccessViewRHIParamRef UAV, uint32 InitialCount) final override;
 	virtual void RHISetShaderResourceViewParameter(FPixelShaderRHIParamRef PixelShader, uint32 SamplerIndex, FShaderResourceViewRHIParamRef SRV) final override;
