DR - Deprecated FRasterizerStateRHIParamRef, FTextureReferenceRHIParamRef and FRayTracingPipelineStateRHIParamRef

diff --git a/Engine/Source/Runtime/Apple/MetalRHI/Private/MetalContext.cpp b/Engine/Source/Runtime/Apple/MetalRHI/Private/MetalContext.cpp
--- a/Engine/Source/Runtime/Apple/MetalRHI/Private/MetalContext.cpp
+++ b/Engine/Source/Runtime/Apple/MetalRHI/Private/MetalContext.cpp
@@ -736,7 +736,7 @@ void FMetalDeviceContext::RegisterUB(FMetalUniformBuffer* UB)
 	UniformBuffers.Add(UB);
 }
 
-void FMetalDeviceContext::UpdateIABs(FTextureReferenceRHIParamRef ModifiedRef)
+void FMetalDeviceContext::UpdateIABs(FRHITextureReference* ModifiedRef)
 {
 	if(GIsMetalInitialized)
 	{
diff --git a/Engine/Source/Runtime/Apple/MetalRHI/Private/MetalContext.h b/Engine/Source/Runtime/Apple/MetalRHI/Private/MetalContext.h
--- a/Engine/Source/Runtime/Apple/MetalRHI/Private/MetalContext.h
+++ b/Engine/Source/Runtime/Apple/MetalRHI/Private/MetalContext.h
@@ -188,7 +188,7 @@ public:
 	void ReleaseTexture(FMetalTexture& Texture);
 	void ReleaseFence(FMetalFence* Fence);
 	void RegisterUB(FMetalUniformBuffer* UB);
-	void UpdateIABs(FTextureReferenceRHIParamRef ModifiedRef);
+	void UpdateIABs(FRHITextureReference* ModifiedRef);
 	void UnregisterUB(FMetalUniformBuffer* UB);
 	
 	void BeginFrame();
diff --git a/Engine/Source/Runtime/Apple/MetalRHI/Private/MetalTexture.cpp b/Engine/Source/Runtime/Apple/MetalRHI/Private/MetalTexture.cpp
--- a/Engine/Source/Runtime/Apple/MetalRHI/Private/MetalTexture.cpp
+++ b/Engine/Source/Runtime/Apple/MetalRHI/Private/MetalTexture.cpp
@@ -2362,7 +2362,7 @@ FTextureReferenceRHIRef FMetalDynamicRHI::RHICreateTextureReference(FLastRenderT
 	}
 }
 
-void FMetalRHICommandContext::RHIUpdateTextureReference(FTextureReferenceRHIParamRef TextureRefRHI, FTextureRHIParamRef NewTextureRHI)
+void FMetalRHICommandContext::RHIUpdateTextureReference(FRHITextureReference* TextureRefRHI, FTextureRHIParamRef NewTextureRHI)
 {
 	@autoreleasepool {
 		FMetalTextureReference* TextureRef = (FMetalTextureReference*)TextureRefRHI;
diff --git a/Engine/Source/Runtime/Apple/MetalRHI/Public/MetalRHIContext.h b/Engine/Source/Runtime/Apple/MetalRHI/Public/MetalRHIContext.h
--- a/Engine/Source/Runtime/Apple/MetalRHI/Public/MetalRHIContext.h
+++ b/Engine/Source/Runtime/Apple/MetalRHI/Public/MetalRHIContext.h
@@ -252,7 +252,7 @@ public:
 	
 	virtual void RHIPopEvent() final override;
 	
-	virtual void RHIUpdateTextureReference(FTextureReferenceRHIParamRef TextureRef, FTextureRHIParamRef NewTexture) final override;
+	virtual void RHIUpdateTextureReference(FRHITextureReference* TextureRef, FTextureRHIParamRef NewTexture) final override;
 	
 	/**
 	 * Explicitly transition a UAV from readable -> writable by the GPU or vice versa.
diff --git a/Engine/Source/Runtime/Apple/MetalRHI/Public/MetalResources.h b/Engine/Source/Runtime/Apple/MetalRHI/Public/MetalResources.h
--- a/Engine/Source/Runtime/Apple/MetalRHI/Public/MetalResources.h
+++ b/Engine/Source/Runtime/Apple/MetalRHI/Public/MetalResources.h
@@ -863,7 +863,7 @@ public:
 	/** Resource table containing RHI references. */
 	TArray<TRefCountPtr<FRHIResource> > ResourceTable;
 	
-	TSet<FTextureReferenceRHIParamRef> TextureReferences;
+	TSet<FRHITextureReference*> TextureReferences;
 
 	struct Argument
 	{
diff --git a/Engine/Source/Runtime/D3D12RHI/Private/D3D12CommandContext.h b/Engine/Source/Runtime/D3D12RHI/Private/D3D12CommandContext.h
--- a/Engine/Source/Runtime/D3D12RHI/Private/D3D12CommandContext.h
+++ b/Engine/Source/Runtime/D3D12RHI/Private/D3D12CommandContext.h
@@ -300,7 +300,7 @@ public:
 	virtual void RHIDrawIndexedPrimitive(FIndexBufferRHIParamRef IndexBuffer, int32 BaseVertexIndex, uint32 FirstInstance, uint32 NumVertices, uint32 StartIndex, uint32 NumPrimitives, uint32 NumInstances) final override;
 	virtual void RHIDrawIndexedPrimitiveIndirect(FIndexBufferRHIParamRef IndexBuffer, FVertexBufferRHIParamRef ArgumentBuffer, uint32 ArgumentOffset) final override;
 	virtual void RHISetDepthBounds(float MinDepth, float MaxDepth) final override;
-	virtual void RHIUpdateTextureReference(FTextureReferenceRHIParamRef TextureRef, FTextureRHIParamRef NewTexture) final override;
+	virtual void RHIUpdateTextureReference(FRHITextureReference* TextureRef, FTextureRHIParamRef NewTexture) final override;
 
 	virtual void RHIClearMRTImpl(bool bClearColor, int32 NumClearColors, const FLinearColor* ColorArray, bool bClearDepth, float Depth, bool bClearStencil, uint32 Stencil);
 
@@ -346,19 +346,19 @@ public:
 		FShaderResourceViewRHIParamRef Rays,
 		FUnorderedAccessViewRHIParamRef Output,
 		uint32 NumRays) final override;
-	virtual void RHIRayTraceDispatch(FRayTracingPipelineStateRHIParamRef RayTracingPipelineState, FRayTracingShaderRHIParamRef RayGenShader,
+	virtual void RHIRayTraceDispatch(FRHIRayTracingPipelineState* RayTracingPipelineState, FRayTracingShaderRHIParamRef RayGenShader,
 		FRayTracingSceneRHIParamRef Scene,
 		const FRayTracingShaderBindings& GlobalResourceBindings,
 		uint32 Width, uint32 Height) final override;
 	virtual void RHISetRayTracingHitGroup(
 		FRayTracingSceneRHIParamRef Scene, uint32 InstanceIndex, uint32 SegmentIndex, uint32 ShaderSlot,
-		FRayTracingPipelineStateRHIParamRef Pipeline, uint32 HitGroupIndex,
+		FRHIRayTracingPipelineState* Pipeline, uint32 HitGroupIndex,
 		uint32 NumUniformBuffers, const FUniformBufferRHIParamRef* UniformBuffers,
 		uint32 LooseParameterDataSize, const void* LooseParameterData,
 		uint32 UserData) final override;
 	virtual void RHISetRayTracingCallableShader(
 		FRayTracingSceneRHIParamRef Scene, uint32 ShaderSlotInScene,
-		FRayTracingPipelineStateRHIParamRef Pipeline, uint32 ShaderIndexInPipeline,
+		FRHIRayTracingPipelineState* Pipeline, uint32 ShaderIndexInPipeline,
 		uint32 NumUniformBuffers, const FUniformBufferRHIParamRef* UniformBuffers,
 		uint32 UserData) final override;
 #endif // D3D12_RHI_RAYTRACING
@@ -714,7 +714,7 @@ public:
 		ContextRedirect(RHISetDepthBounds(MinDepth, MaxDepth));
 	}
 
-	FORCEINLINE virtual void RHIUpdateTextureReference(FTextureReferenceRHIParamRef TextureRef, FTextureRHIParamRef NewTexture) final override
+	FORCEINLINE virtual void RHIUpdateTextureReference(FRHITextureReference* TextureRef, FTextureRHIParamRef NewTexture) final override
 	{
 		ContextRedirect(RHIUpdateTextureReference(TextureRef, NewTexture));
 	}
@@ -791,7 +791,7 @@ public:
 		ContextRedirect(RHIRayTraceIntersection(Scene, Rays, Output, NumRays));
 	}
 
-	virtual void RHIRayTraceDispatch(FRayTracingPipelineStateRHIParamRef RayTracingPipelineState, FRayTracingShaderRHIParamRef RayGenShader,
+	virtual void RHIRayTraceDispatch(FRHIRayTracingPipelineState* RayTracingPipelineState, FRayTracingShaderRHIParamRef RayGenShader,
 		FRayTracingSceneRHIParamRef Scene,
 		const FRayTracingShaderBindings& GlobalResourceBindings,
 		uint32 Width, uint32 Height) final override
@@ -801,7 +801,7 @@ public:
 
 	virtual void RHISetRayTracingHitGroup(
 		FRayTracingSceneRHIParamRef Scene, uint32 InstanceIndex, uint32 SegmentIndex, uint32 ShaderSlot,
-		FRayTracingPipelineStateRHIParamRef Pipeline, uint32 HitGroupIndex,
+		FRHIRayTracingPipelineState* Pipeline, uint32 HitGroupIndex,
 		uint32 NumUniformBuffers, const FUniformBufferRHIParamRef* UniformBuffers,
 		uint32 LooseParameterDataSize, const void* LooseParameterData,
 		uint32 UserData) final override
@@ -811,7 +811,7 @@ public:
 
 	virtual void RHISetRayTracingCallableShader(
 		FRayTracingSceneRHIParamRef Scene, uint32 ShaderSlotInScene,
-		FRayTracingPipelineStateRHIParamRef Pipeline, uint32 ShaderIndexInPipeline,
+		FRHIRayTracingPipelineState* Pipeline, uint32 ShaderIndexInPipeline,
 		uint32 NumUniformBuffers, const FUniformBufferRHIParamRef* UniformBuffers,
 		uint32 UserData) final override
 	{
diff --git a/Engine/Source/Runtime/D3D12RHI/Private/D3D12RayTracing.cpp b/Engine/Source/Runtime/D3D12RHI/Private/D3D12RayTracing.cpp
--- a/Engine/Source/Runtime/D3D12RHI/Private/D3D12RayTracing.cpp
+++ b/Engine/Source/Runtime/D3D12RHI/Private/D3D12RayTracing.cpp
@@ -2954,7 +2954,7 @@ void FD3D12CommandContext::RHIRayTraceIntersection(FRayTracingSceneRHIParamRef I
 	DispatchRays(*this, Bindings, Pipeline, 0, nullptr, DispatchDesc);
 }
 
-void FD3D12CommandContext::RHIRayTraceDispatch(FRayTracingPipelineStateRHIParamRef InRayTracingPipelineState, FRayTracingShaderRHIParamRef RayGenShaderRHI,
+void FD3D12CommandContext::RHIRayTraceDispatch(FRHIRayTracingPipelineState* InRayTracingPipelineState, FRayTracingShaderRHIParamRef RayGenShaderRHI,
 	FRayTracingSceneRHIParamRef InScene,
 	const FRayTracingShaderBindings& GlobalResourceBindings,
 	uint32 Width, uint32 Height)
@@ -2988,7 +2988,7 @@ void FD3D12CommandContext::RHIRayTraceDispatch(FRayTracingPipelineStateRHIParamR
 
 void FD3D12CommandContext::RHISetRayTracingHitGroup(
 	FRayTracingSceneRHIParamRef InScene, uint32 InstanceIndex, uint32 SegmentIndex, uint32 ShaderSlot,
-	FRayTracingPipelineStateRHIParamRef InPipeline, uint32 HitGroupIndex,
+	FRHIRayTracingPipelineState* InPipeline, uint32 HitGroupIndex,
 	uint32 NumUniformBuffers, const FUniformBufferRHIParamRef* UniformBuffers,
 	uint32 LooseParameterDataSize, const void* LooseParameterData,
 	uint32 UserData)
@@ -3025,7 +3025,7 @@ void FD3D12CommandContext::RHISetRayTracingHitGroup(
 
 void FD3D12CommandContext::RHISetRayTracingCallableShader(
 	FRayTracingSceneRHIParamRef InScene, uint32 ShaderSlotInScene,
-	FRayTracingPipelineStateRHIParamRef InPipeline, uint32 ShaderIndexInPipeline,
+	FRHIRayTracingPipelineState* InPipeline, uint32 ShaderIndexInPipeline,
 	uint32 NumUniformBuffers, const FUniformBufferRHIParamRef* UniformBuffers,
 	uint32 UserData)
 {
diff --git a/Engine/Source/Runtime/D3D12RHI/Private/D3D12Texture.cpp b/Engine/Source/Runtime/D3D12RHI/Private/D3D12Texture.cpp
--- a/Engine/Source/Runtime/D3D12RHI/Private/D3D12Texture.cpp
+++ b/Engine/Source/Runtime/D3D12RHI/Private/D3D12Texture.cpp
@@ -2443,7 +2443,7 @@ FTextureReferenceRHIRef FD3D12DynamicRHI::RHICreateTextureReference(FLastRenderT
 	return new FD3D12TextureReference(GetRHIDevice(), LastRenderTime);
 }
 
-void FD3D12CommandContext::RHIUpdateTextureReference(FTextureReferenceRHIParamRef TextureRefRHI, FTextureRHIParamRef NewTextureRHI)
+void FD3D12CommandContext::RHIUpdateTextureReference(FRHITextureReference* TextureRefRHI, FTextureRHIParamRef NewTextureRHI)
 {
 #if 0
 	// Updating texture references is disallowed while the RHI could be caching them in referenced resource tables.
diff --git a/Engine/Source/Runtime/EmptyRHI/Private/EmptyCommands.cpp b/Engine/Source/Runtime/EmptyRHI/Private/EmptyCommands.cpp
--- a/Engine/Source/Runtime/EmptyRHI/Private/EmptyCommands.cpp
+++ b/Engine/Source/Runtime/EmptyRHI/Private/EmptyCommands.cpp
@@ -17,7 +17,7 @@ void FEmptyDynamicRHI::RHISetStreamOutTargets(uint32 NumTargets, const FVertexBu
 
 }
 
-void FEmptyDynamicRHI::RHISetRasterizerState(FRasterizerStateRHIParamRef NewStateRHI)
+void FEmptyDynamicRHI::RHISetRasterizerState(FRHIRasterizerState* NewStateRHI)
 {
 	FEmptyRasterizerState* NewState = ResourceCast(NewStateRHI);
 
diff --git a/Engine/Source/Runtime/EmptyRHI/Private/EmptyTexture.cpp b/Engine/Source/Runtime/EmptyRHI/Private/EmptyTexture.cpp
--- a/Engine/Source/Runtime/EmptyRHI/Private/EmptyTexture.cpp
+++ b/Engine/Source/Runtime/EmptyRHI/Private/EmptyTexture.cpp
@@ -262,6 +262,6 @@ FTextureReferenceRHIRef FEmptyDynamicRHI::RHICreateTextureReference(FLastRenderT
 	return nullptr;
 }
 
-void FEmptyDynamicRHI::RHIUpdateTextureReference(FTextureReferenceRHIParamRef TextureRef, FTextureRHIParamRef NewTexture)
+void FEmptyDynamicRHI::RHIUpdateTextureReference(FRHITextureReference* TextureRef, FTextureRHIParamRef NewTexture)
 {
 }
diff --git a/Engine/Source/Runtime/EmptyRHI/Public/EmptyRHI.h b/Engine/Source/Runtime/EmptyRHI/Public/EmptyRHI.h
--- a/Engine/Source/Runtime/EmptyRHI/Public/EmptyRHI.h
+++ b/Engine/Source/Runtime/EmptyRHI/Public/EmptyRHI.h
@@ -75,7 +75,7 @@ public:
 	virtual void RHIGetTextureMemoryStats(FTextureMemoryStats& OutStats) final override;
 	virtual bool RHIGetTextureMemoryVisualizeData(FColor* TextureData,int32 SizeX,int32 SizeY,int32 Pitch,int32 PixelSize) final override;
 	virtual FTextureReferenceRHIRef RHICreateTextureReference(FLastRenderTimeContainer* LastRenderTime) final override;
-	virtual void RHIUpdateTextureReference(FTextureReferenceRHIParamRef TextureRef, FTextureRHIParamRef NewTexture) final override;
+	virtual void RHIUpdateTextureReference(FRHITextureReference* TextureRef, FTextureRHIParamRef NewTexture) final override;
 	virtual FTexture2DRHIRef RHICreateTexture2D(uint32 SizeX, uint32 SizeY, uint8 Format, uint32 NumMips, uint32 NumSamples, uint32 Flags, FRHIResourceCreateInfo& CreateInfo) final override;
 	virtual FTexture2DRHIRef RHIAsyncCreateTexture2D(uint32 SizeX, uint32 SizeY, uint8 Format, uint32 NumMips, uint32 Flags, void** InitialMipData, uint32 NumInitialMips) final override;
 	virtual void RHICopySharedMips(FTexture2DRHIParamRef DestTexture2D, FTexture2DRHIParamRef SrcTexture2D) final override;
@@ -151,7 +151,7 @@ public:
 	virtual void RHIBeginScene() final override;
 	virtual void RHIEndScene() final override;
 	virtual void RHISetStreamSource(uint32 StreamIndex, FVertexBufferRHIParamRef VertexBuffer, uint32 Offset) final override;
-	virtual void RHISetRasterizerState(FRasterizerStateRHIParamRef NewState) final override;
+	virtual void RHISetRasterizerState(FRHIRasterizerState* NewState) final override;
 	virtual void RHISetViewport(uint32 MinX, uint32 MinY, float MinZ, uint32 MaxX, uint32 MaxY, float MaxZ) final override;
 	virtual void RHISetScissorRect(bool bEnable, uint32 MinX, uint32 MinY, uint32 MaxX, uint32 MaxY) final override;
 	virtual void RHISetBoundShaderState(FBoundShaderStateRHIParamRef BoundShaderState) final override;
diff --git a/Engine/Source/Runtime/NullDrv/Public/NullRHI.h b/Engine/Source/Runtime/NullDrv/Public/NullRHI.h
--- a/Engine/Source/Runtime/NullDrv/Public/NullRHI.h
+++ b/Engine/Source/Runtime/NullDrv/Public/NullRHI.h
@@ -271,7 +271,7 @@ public:
 		return new FRHITextureReferenceNullImpl(); 
 	}
 
-	virtual void RHIUpdateTextureReference(FTextureReferenceRHIParamRef TextureRef, FTextureRHIParamRef NewTexture) final override
+	virtual void RHIUpdateTextureReference(FRHITextureReference* TextureRef, FTextureRHIParamRef NewTexture) final override
 	{ 
 		if(TextureRef) 
 		{ 
@@ -518,7 +518,7 @@ public:
 	{
 
 	}
-	virtual void RHISetRasterizerState(FRasterizerStateRHIParamRef NewState) final override
+	virtual void RHISetRasterizerState(FRHIRasterizerState* NewState) final override
 	{
 
 	}
diff --git a/Engine/Source/Runtime/OpenGLDrv/Private/OpenGLCommands.cpp b/Engine/Source/Runtime/OpenGLDrv/Private/OpenGLCommands.cpp
--- a/Engine/Source/Runtime/OpenGLDrv/Private/OpenGLCommands.cpp
+++ b/Engine/Source/Runtime/OpenGLDrv/Private/OpenGLCommands.cpp
@@ -288,7 +288,7 @@ void FOpenGLDynamicRHI::RHISetStreamOutTargets(uint32 NumTargets, const FVertexB
 }
 
 // Rasterizer state.
-void FOpenGLDynamicRHI::RHISetRasterizerState(FRasterizerStateRHIParamRef NewStateRHI)
+void FOpenGLDynamicRHI::RHISetRasterizerState(FRHIRasterizerState* NewStateRHI)
 {
 	VERIFY_GL_SCOPE();
 	FOpenGLRasterizerState* NewState = ResourceCast(NewStateRHI);
diff --git a/Engine/Source/Runtime/OpenGLDrv/Private/OpenGLTexture.cpp b/Engine/Source/Runtime/OpenGLDrv/Private/OpenGLTexture.cpp
--- a/Engine/Source/Runtime/OpenGLDrv/Private/OpenGLTexture.cpp
+++ b/Engine/Source/Runtime/OpenGLDrv/Private/OpenGLTexture.cpp
@@ -2550,7 +2550,7 @@ void FOpenGLTextureReference::SetReferencedTexture(FRHITexture* InTexture)
 	TexturePtr = GetOpenGLTextureFromRHITexture(InTexture);
 }
 
-void FOpenGLDynamicRHI::RHIUpdateTextureReference(FTextureReferenceRHIParamRef TextureRefRHI, FTextureRHIParamRef NewTextureRHI)
+void FOpenGLDynamicRHI::RHIUpdateTextureReference(FRHITextureReference* TextureRefRHI, FTextureRHIParamRef NewTextureRHI)
 {
 	auto* TextureRef = (FOpenGLTextureReference*)TextureRefRHI;
 	if (TextureRef)
diff --git a/Engine/Source/Runtime/OpenGLDrv/Public/OpenGLDrv.h b/Engine/Source/Runtime/OpenGLDrv/Public/OpenGLDrv.h
--- a/Engine/Source/Runtime/OpenGLDrv/Public/OpenGLDrv.h
+++ b/Engine/Source/Runtime/OpenGLDrv/Public/OpenGLDrv.h
@@ -387,7 +387,7 @@ public:
 	virtual void RHIGetTextureMemoryStats(FTextureMemoryStats& OutStats) final override;
 	virtual bool RHIGetTextureMemoryVisualizeData(FColor* TextureData,int32 SizeX,int32 SizeY,int32 Pitch,int32 PixelSize) final override;
 	virtual FTextureReferenceRHIRef RHICreateTextureReference(FLastRenderTimeContainer* LastRenderTime) final override;
-	virtual void RHIUpdateTextureReference(FTextureReferenceRHIParamRef TextureRef, FTextureRHIParamRef NewTexture) final override;
+	virtual void RHIUpdateTextureReference(FRHITextureReference* TextureRef, FTextureRHIParamRef NewTexture) final override;
 	virtual FTexture2DRHIRef RHICreateTexture2D(uint32 SizeX, uint32 SizeY, uint8 Format, uint32 NumMips, uint32 NumSamples, uint32 Flags, FRHIResourceCreateInfo& CreateInfo) final override;
 	virtual FTexture2DRHIRef RHICreateTextureExternal2D(uint32 SizeX, uint32 SizeY, uint8 Format, uint32 NumMips, uint32 NumSamples, uint32 Flags, FRHIResourceCreateInfo& CreateInfo) final override;
 	virtual FTexture2DRHIRef RHIAsyncCreateTexture2D(uint32 SizeX, uint32 SizeY, uint8 Format, uint32 NumMips, uint32 Flags, void** InitialMipData, uint32 NumInitialMips) final override;
@@ -463,7 +463,7 @@ public:
 	virtual void RHIBeginScene() final override;
 	virtual void RHIEndScene() final override;
 	virtual void RHISetStreamSource(uint32 StreamIndex, FVertexBufferRHIParamRef VertexBuffer, uint32 Offset) final override;
-	virtual void RHISetRasterizerState(FRasterizerStateRHIParamRef NewState) final override;
+	virtual void RHISetRasterizerState(FRHIRasterizerState* NewState) final override;
 	virtual void RHISetViewport(uint32 MinX, uint32 MinY, float MinZ, uint32 MaxX, uint32 MaxY, float MaxZ) final override;
 	virtual void RHISetScissorRect(bool bEnable, uint32 MinX, uint32 MinY, uint32 MaxX, uint32 MaxY) final override;
 	virtual void RHISetBoundShaderState(FBoundShaderStateRHIParamRef BoundShaderState) final override;
diff --git a/Engine/Source/Runtime/RHI/Private/RHICommandList.cpp b/Engine/Source/Runtime/RHI/Private/RHICommandList.cpp
--- a/Engine/Source/Runtime/RHI/Private/RHICommandList.cpp
+++ b/Engine/Source/Runtime/RHI/Private/RHICommandList.cpp
@@ -2471,7 +2471,7 @@ void FDynamicRHI::RHIReadSurfaceFloatData_RenderThread(class FRHICommandListImme
 }
 
 
-void FRHICommandListImmediate::UpdateTextureReference(FTextureReferenceRHIParamRef TextureRef, FTextureRHIParamRef NewTexture)
+void FRHICommandListImmediate::UpdateTextureReference(FRHITextureReference* TextureRef, FTextureRHIParamRef NewTexture)
 {
 	if (Bypass() || !IsRunningRHIInSeparateThread() || CVarRHICmdFlushUpdateTextureReference.GetValueOnRenderThread() > 0)
 	{
diff --git a/Engine/Source/Runtime/RHI/Public/RHICommandList.h b/Engine/Source/Runtime/RHI/Public/RHICommandList.h
--- a/Engine/Source/Runtime/RHI/Public/RHICommandList.h
+++ b/Engine/Source/Runtime/RHI/Public/RHICommandList.h
@@ -1757,9 +1757,9 @@ struct FRHICommandDebugBreak final : public FRHICommand<FRHICommandDebugBreak>
 
 struct FRHICommandUpdateTextureReference final : public FRHICommand<FRHICommandUpdateTextureReference>
 {
-	FTextureReferenceRHIParamRef TextureRef;
+	FRHITextureReference* TextureRef;
 	FTextureRHIParamRef NewTexture;
-	FORCEINLINE_DEBUGGABLE FRHICommandUpdateTextureReference(FTextureReferenceRHIParamRef InTextureRef, FTextureRHIParamRef InNewTexture)
+	FORCEINLINE_DEBUGGABLE FRHICommandUpdateTextureReference(FRHITextureReference* InTextureRef, FTextureRHIParamRef InNewTexture)
 		: TextureRef(InTextureRef)
 		, NewTexture(InNewTexture)
 	{
@@ -4326,7 +4326,7 @@ public:
 	{
 		return RHIGetCommandContextContainer(Index, Num, GetGPUMask());
 	}
-	void UpdateTextureReference(FTextureReferenceRHIParamRef TextureRef, FTextureRHIParamRef NewTexture);
+	void UpdateTextureReference(FRHITextureReference* TextureRef, FTextureRHIParamRef NewTexture);
 
 	FORCEINLINE void PollRenderQueryResults()
 	{
@@ -4687,7 +4687,7 @@ FORCEINLINE FTextureReferenceRHIRef RHICreateTextureReference(FLastRenderTimeCon
 	return FRHICommandListExecutor::GetImmediateCommandList().CreateTextureReference(LastRenderTime);
 }
 
-FORCEINLINE void RHIUpdateTextureReference(FTextureReferenceRHIParamRef TextureRef, FTextureRHIParamRef NewTexture)
+FORCEINLINE void RHIUpdateTextureReference(FRHITextureReference* TextureRef, FTextureRHIParamRef NewTexture)
 {
 	 FRHICommandListExecutor::GetImmediateCommandList().UpdateTextureReference(TextureRef, NewTexture);
 }
diff --git a/Engine/Source/Runtime/RHI/Public/RHIContext.h b/Engine/Source/Runtime/RHI/Public/RHIContext.h
--- a/Engine/Source/Runtime/RHI/Public/RHIContext.h
+++ b/Engine/Source/Runtime/RHI/Public/RHIContext.h
@@ -479,7 +479,7 @@ public:
 	*/
 	virtual void RHISetDepthBounds(float MinDepth, float MaxDepth) = 0;
 
-	virtual void RHIUpdateTextureReference(FTextureReferenceRHIParamRef TextureRef, FTextureRHIParamRef NewTexture) = 0;
+	virtual void RHIUpdateTextureReference(FRHITextureReference* TextureRef, FTextureRHIParamRef NewTexture) = 0;
 
 	virtual void RHIBeginRenderPass(const FRHIRenderPassInfo& InInfo, const TCHAR* InName)
 	{
@@ -615,7 +615,7 @@ public:
 		checkNoEntry();
 	}
 
-	virtual void RHIRayTraceDispatch(FRayTracingPipelineStateRHIParamRef RayTracingPipelineState, FRayTracingShaderRHIParamRef RayGenShader,
+	virtual void RHIRayTraceDispatch(FRHIRayTracingPipelineState* RayTracingPipelineState, FRayTracingShaderRHIParamRef RayGenShader,
 		FRayTracingSceneRHIParamRef Scene, 
 		const FRayTracingShaderBindings& GlobalResourceBindings,
 		uint32 Width, uint32 Height)
@@ -625,7 +625,7 @@ public:
 
 	virtual void RHISetRayTracingHitGroup(
 		FRayTracingSceneRHIParamRef Scene, uint32 InstanceIndex, uint32 SegmentIndex, uint32 ShaderSlot,
-		FRayTracingPipelineStateRHIParamRef Pipeline, uint32 HitGroupIndex,
+		FRHIRayTracingPipelineState* Pipeline, uint32 HitGroupIndex,
 		uint32 NumUniformBuffers, const FUniformBufferRHIParamRef* UniformBuffers,
 		uint32 LooseParameterDataSize, const void* LooseParameterData,
 		uint32 UserData)
@@ -635,7 +635,7 @@ public:
 
 	virtual void RHISetRayTracingCallableShader(
 		FRayTracingSceneRHIParamRef Scene, uint32 ShaderSlotInScene,
-		FRayTracingPipelineStateRHIParamRef Pipeline, uint32 ShaderIndexInPipeline,
+		FRHIRayTracingPipelineState* Pipeline, uint32 ShaderIndexInPipeline,
 		uint32 NumUniformBuffers, const FUniformBufferRHIParamRef* UniformBuffers,
 		uint32 UserData)
 	{
@@ -670,7 +670,7 @@ public:
 
 	virtual void RHISetDepthStencilState(FDepthStencilStateRHIParamRef NewState, uint32 StencilRef) = 0;
 
-	virtual void RHISetRasterizerState(FRasterizerStateRHIParamRef NewState) = 0;
+	virtual void RHISetRasterizerState(FRHIRasterizerState* NewState) = 0;
 
 	virtual void RHISetBlendState(FBlendStateRHIParamRef NewState, const FLinearColor& BlendFactor) = 0;
 
diff --git a/Engine/Source/Runtime/RHI/Public/RHIResources.h b/Engine/Source/Runtime/RHI/Public/RHIResources.h
--- a/Engine/Source/Runtime/RHI/Public/RHIResources.h
+++ b/Engine/Source/Runtime/RHI/Public/RHIResources.h
@@ -1101,6 +1101,7 @@ UE_DEPRECATED(4.23, "FSamplerStateRHIParamRef typedef is deprecated; please use
 typedef FRHISamplerState*              FSamplerStateRHIParamRef;
 typedef TRefCountPtr<FRHISamplerState> FSamplerStateRHIRef;
 
+UE_DEPRECATED(4.23, "FRasterizerStateRHIParamRef typedef is deprecated; please use FRHIRasterizerState* directly instead.")
 typedef FRHIRasterizerState*              FRasterizerStateRHIParamRef;
 typedef TRefCountPtr<FRHIRasterizerState> FRasterizerStateRHIRef;
 
@@ -1167,6 +1168,7 @@ typedef TRefCountPtr<FRHITexture3D> FTexture3DRHIRef;
 typedef FRHITextureCube*              FTextureCubeRHIParamRef;
 typedef TRefCountPtr<FRHITextureCube> FTextureCubeRHIRef;
 
+UE_DEPRECATED(4.23, "FTextureReferenceRHIParamRef typedef is deprecated; please use FRHITextureReference* directly instead.")
 typedef FRHITextureReference*              FTextureReferenceRHIParamRef;
 typedef TRefCountPtr<FRHITextureReference> FTextureReferenceRHIRef;
 
@@ -1191,6 +1193,7 @@ typedef TRefCountPtr<FRHIShaderResourceView> FShaderResourceViewRHIRef;
 typedef FRHIGraphicsPipelineState*              FGraphicsPipelineStateRHIParamRef;
 typedef TRefCountPtr<FRHIGraphicsPipelineState> FGraphicsPipelineStateRHIRef;
 
+UE_DEPRECATED(4.23, "FRayTracingPipelineStateRHIParamRef typedef is deprecated; please use FRHIRayTracingPipelineState* directly instead.")
 typedef FRHIRayTracingPipelineState*              FRayTracingPipelineStateRHIParamRef;
 typedef TRefCountPtr<FRHIRayTracingPipelineState> FRayTracingPipelineStateRHIRef;
 
@@ -1831,7 +1834,7 @@ public:
 	FGraphicsMinimalPipelineStateInitializer(
 		FBoundShaderStateInput				InBoundShaderState,
 		FBlendStateRHIParamRef				InBlendState,
-		FRasterizerStateRHIParamRef			InRasterizerState,
+		FRHIRasterizerState*				InRasterizerState,
 		FDepthStencilStateRHIParamRef		InDepthStencilState,
 		FImmutableSamplerState				InImmutableSamplerState,
 		EPrimitiveType						InPrimitiveType
@@ -1961,7 +1964,7 @@ public:
 	// When we do that work, move this to RHI.h as its more appropriate there, but here for now since dependent typdefs are here.
 	FBoundShaderStateInput			BoundShaderState;
 	FBlendStateRHIParamRef			BlendState;
-	FRasterizerStateRHIParamRef		RasterizerState;
+	FRHIRasterizerState*			RasterizerState;
 	FDepthStencilStateRHIParamRef	DepthStencilState;
 	FImmutableSamplerState			ImmutableSamplerState;
 
@@ -2031,7 +2034,7 @@ public:
 	FGraphicsPipelineStateInitializer(
 		FBoundShaderStateInput				InBoundShaderState,
 		FBlendStateRHIParamRef				InBlendState,
-		FRasterizerStateRHIParamRef			InRasterizerState,
+		FRHIRasterizerState*				InRasterizerState,
 		FDepthStencilStateRHIParamRef		InDepthStencilState,
 		FImmutableSamplerState				InImmutableSamplerState,
 		EPrimitiveType						InPrimitiveType,
diff --git a/Engine/Source/Runtime/RHI/Public/RHIStaticStates.h b/Engine/Source/Runtime/RHI/Public/RHIStaticStates.h
--- a/Engine/Source/Runtime/RHI/Public/RHIStaticStates.h
+++ b/Engine/Source/Runtime/RHI/Public/RHIStaticStates.h
@@ -171,7 +171,7 @@ public:
  * Should only be used from the rendering thread.
  */
 template<ERasterizerFillMode FillMode=FM_Solid,ERasterizerCullMode CullMode=CM_None,bool bEnableLineAA=false,bool bEnableMSAA=true>
-class TStaticRasterizerState : public TStaticStateRHI<TStaticRasterizerState<FillMode,CullMode,bEnableLineAA>,FRasterizerStateRHIRef,FRasterizerStateRHIParamRef>
+class TStaticRasterizerState : public TStaticStateRHI<TStaticRasterizerState<FillMode,CullMode,bEnableLineAA>,FRasterizerStateRHIRef, FRHIRasterizerState*>
 {
 public:
 	FORCEINLINE_DEBUGGABLE static FRasterizerStateRHIRef CreateRHI()
@@ -183,7 +183,7 @@ public:
 
 /** Given a fill and cull mode, returns a static rasterizer state. */
 template<bool bEnableMSAA>
-FORCEINLINE_DEBUGGABLE FRasterizerStateRHIParamRef GetStaticRasterizerState(ERasterizerFillMode FillMode,ERasterizerCullMode CullMode)
+FORCEINLINE_DEBUGGABLE FRHIRasterizerState* GetStaticRasterizerState(ERasterizerFillMode FillMode,ERasterizerCullMode CullMode)
 {
 	switch(FillMode)
 	{
diff --git a/Engine/Source/Runtime/RHI/Public/RHIValidationContext.h b/Engine/Source/Runtime/RHI/Public/RHIValidationContext.h
--- a/Engine/Source/Runtime/RHI/Public/RHIValidationContext.h
+++ b/Engine/Source/Runtime/RHI/Public/RHIValidationContext.h
@@ -745,7 +745,7 @@ public:
 		RHIContext->RHIPopEvent();
 	}
 
-	virtual void RHIUpdateTextureReference(FTextureReferenceRHIParamRef TextureRef, FTextureRHIParamRef NewTexture) override final
+	virtual void RHIUpdateTextureReference(FRHITextureReference* TextureRef, FTextureRHIParamRef NewTexture) override final
 	{
 		RHIContext->RHIUpdateTextureReference(TextureRef, NewTexture);
 	}
@@ -840,7 +840,7 @@ public:
 		RHIContext->RHIRayTraceIntersection(Scene, Rays, Output, NumRays);
 	}
 
-	virtual void RHIRayTraceDispatch(FRayTracingPipelineStateRHIParamRef RayTracingPipelineState, FRHIRayTracingShader* RayGenShader,
+	virtual void RHIRayTraceDispatch(FRHIRayTracingPipelineState* RayTracingPipelineState, FRHIRayTracingShader* RayGenShader,
 		FRayTracingSceneRHIParamRef Scene,
 		const FRayTracingShaderBindings& GlobalResourceBindings,
 		uint32 Width, uint32 Height) override final
@@ -850,7 +850,7 @@ public:
 
 	virtual void RHISetRayTracingHitGroup(
 		FRayTracingSceneRHIParamRef Scene, uint32 InstanceIndex, uint32 SegmentIndex, uint32 ShaderSlot,
-		FRayTracingPipelineStateRHIParamRef Pipeline, uint32 HitGroupIndex,
+		FRHIRayTracingPipelineState* Pipeline, uint32 HitGroupIndex,
 		uint32 NumUniformBuffers, const FUniformBufferRHIParamRef* UniformBuffers,
 		uint32 LooseParameterDataSize, const void* LooseParameterData,
 		uint32 UserData) override final
diff --git a/Engine/Source/Runtime/Renderer/Private/CompositionLighting/PostProcessDeferredDecals.cpp b/Engine/Source/Runtime/Renderer/Private/CompositionLighting/PostProcessDeferredDecals.cpp
--- a/Engine/Source/Runtime/Renderer/Private/CompositionLighting/PostProcessDeferredDecals.cpp
+++ b/Engine/Source/Runtime/Renderer/Private/CompositionLighting/PostProcessDeferredDecals.cpp
@@ -485,7 +485,7 @@ static FDepthStencilStateRHIParamRef GetDecalDepthState(uint32& StencilRef, FDec
 	}
 }
 
-FRasterizerStateRHIParamRef GetDecalRasterizerState(EDecalRasterizerState DecalRasterizerState)
+FRHIRasterizerState* GetDecalRasterizerState(EDecalRasterizerState DecalRasterizerState)
 {
 	switch (DecalRasterizerState)
 	{
diff --git a/Engine/Source/Runtime/Renderer/Private/MobileDecalRendering.cpp b/Engine/Source/Runtime/Renderer/Private/MobileDecalRendering.cpp
--- a/Engine/Source/Runtime/Renderer/Private/MobileDecalRendering.cpp
+++ b/Engine/Source/Runtime/Renderer/Private/MobileDecalRendering.cpp
@@ -16,7 +16,7 @@
 #include "DecalRenderingShared.h"
 
 
-extern FRasterizerStateRHIParamRef GetDecalRasterizerState(EDecalRasterizerState DecalRasterizerState);
+extern FRHIRasterizerState* GetDecalRasterizerState(EDecalRasterizerState DecalRasterizerState);
 
 void FMobileSceneRenderer::RenderDecals(FRHICommandListImmediate& RHICmdList)
 {
diff --git a/Engine/Source/Runtime/VulkanRHI/Private/VulkanContext.h b/Engine/Source/Runtime/VulkanRHI/Private/VulkanContext.h
--- a/Engine/Source/Runtime/VulkanRHI/Private/VulkanContext.h
+++ b/Engine/Source/Runtime/VulkanRHI/Private/VulkanContext.h
@@ -238,7 +238,7 @@ public:
 	virtual void RHIBeginRenderQuery(FRenderQueryRHIParamRef RenderQuery) final override;
 	virtual void RHIEndRenderQuery(FRenderQueryRHIParamRef RenderQuery) final override;
 
-	virtual void RHIUpdateTextureReference(FTextureReferenceRHIParamRef TextureRef, FTextureRHIParamRef NewTexture) final override;
+	virtual void RHIUpdateTextureReference(FRHITextureReference* TextureRef, FTextureRHIParamRef NewTexture) final override;
 
 	virtual void RHISubmitCommandsHint() final override;
 
diff --git a/Engine/Source/Runtime/VulkanRHI/Private/VulkanTexture.cpp b/Engine/Source/Runtime/VulkanRHI/Private/VulkanTexture.cpp
--- a/Engine/Source/Runtime/VulkanRHI/Private/VulkanTexture.cpp
+++ b/Engine/Source/Runtime/VulkanRHI/Private/VulkanTexture.cpp
@@ -2200,7 +2200,7 @@ FTextureReferenceRHIRef FVulkanDynamicRHI::RHICreateTextureReference(FLastRender
 	return new FVulkanTextureReference(*Device, LastRenderTime);
 }
 
-void FVulkanCommandListContext::RHIUpdateTextureReference(FTextureReferenceRHIParamRef TextureRef, FTextureRHIParamRef NewTexture)
+void FVulkanCommandListContext::RHIUpdateTextureReference(FRHITextureReference* TextureRef, FTextureRHIParamRef NewTexture)
 {
 	//#todo-rco: Implementation needs to be verified
 	FVulkanTextureReference* VulkanTextureRef = (FVulkanTextureReference*)TextureRef;
diff --git a/Engine/Source/Runtime/Windows/D3D11RHI/Private/D3D11Commands.cpp b/Engine/Source/Runtime/Windows/D3D11RHI/Private/D3D11Commands.cpp
--- a/Engine/Source/Runtime/Windows/D3D11RHI/Private/D3D11Commands.cpp
+++ b/Engine/Source/Runtime/Windows/D3D11RHI/Private/D3D11Commands.cpp
@@ -179,7 +179,7 @@ void FD3D11DynamicRHI::RHISetStreamOutTargets(uint32 NumTargets, const FVertexBu
 }
 
 // Rasterizer state.
-void FD3D11DynamicRHI::RHISetRasterizerState(FRasterizerStateRHIParamRef NewStateRHI)
+void FD3D11DynamicRHI::RHISetRasterizerState(FRHIRasterizerState* NewStateRHI)
 {
 	FD3D11RasterizerState* NewState = ResourceCast(NewStateRHI);
 	StateCache.SetRasterizerState(NewState->Resource);
diff --git a/Engine/Source/Runtime/Windows/D3D11RHI/Private/D3D11RHIPrivate.h b/Engine/Source/Runtime/Windows/D3D11RHI/Private/D3D11RHIPrivate.h
--- a/Engine/Source/Runtime/Windows/D3D11RHI/Private/D3D11RHIPrivate.h
+++ b/Engine/Source/Runtime/Windows/D3D11RHI/Private/D3D11RHIPrivate.h
@@ -439,7 +439,7 @@ public:
 	virtual void RHIGetTextureMemoryStats(FTextureMemoryStats& OutStats) final override;
 	virtual bool RHIGetTextureMemoryVisualizeData(FColor* TextureData,int32 SizeX,int32 SizeY,int32 Pitch,int32 PixelSize) final override;
 	virtual FTextureReferenceRHIRef RHICreateTextureReference(FLastRenderTimeContainer* LastRenderTime) final override;
-	virtual void RHIUpdateTextureReference(FTextureReferenceRHIParamRef TextureRef, FTextureRHIParamRef NewTexture) final override;
+	virtual void RHIUpdateTextureReference(FRHITextureReference* TextureRef, FTextureRHIParamRef NewTexture) final override;
 	virtual FTexture2DRHIRef RHICreateTexture2D(uint32 SizeX, uint32 SizeY, uint8 Format, uint32 NumMips, uint32 NumSamples, uint32 Flags, FRHIResourceCreateInfo& CreateInfo) final override;
 	virtual FTexture2DRHIRef RHIAsyncCreateTexture2D(uint32 SizeX, uint32 SizeY, uint8 Format, uint32 NumMips, uint32 Flags, void** InitialMipData, uint32 NumInitialMips) final override;
 	virtual void RHICopySharedMips(FTexture2DRHIParamRef DestTexture2D, FTexture2DRHIParamRef SrcTexture2D) final override;
@@ -517,7 +517,7 @@ public:
 	virtual void RHIBeginScene() final override;
 	virtual void RHIEndScene() final override;
 	virtual void RHISetStreamSource(uint32 StreamIndex, FVertexBufferRHIParamRef VertexBuffer, uint32 Offset) final override;
-	virtual void RHISetRasterizerState(FRasterizerStateRHIParamRef NewState) final override;
+	virtual void RHISetRasterizerState(FRHIRasterizerState* NewState) final override;
 	virtual void RHISetViewport(uint32 MinX, uint32 MinY, float MinZ, uint32 MaxX, uint32 MaxY, float MaxZ) final override;
 	virtual void RHISetScissorRect(bool bEnable, uint32 MinX, uint32 MinY, uint32 MaxX, uint32 MaxY) final override;
 	virtual void RHISetBoundShaderState(FBoundShaderStateRHIParamRef BoundShaderState) final override;
diff --git a/Engine/Source/Runtime/Windows/D3D11RHI/Private/D3D11Texture.cpp b/Engine/Source/Runtime/Windows/D3D11RHI/Private/D3D11Texture.cpp
--- a/Engine/Source/Runtime/Windows/D3D11RHI/Private/D3D11Texture.cpp
+++ b/Engine/Source/Runtime/Windows/D3D11RHI/Private/D3D11Texture.cpp
@@ -2141,7 +2141,7 @@ void FD3D11DynamicRHI::RHICopySubTextureRegion_RenderThread(
 	}
 }
 
-void FD3D11DynamicRHI::RHIUpdateTextureReference(FTextureReferenceRHIParamRef TextureRefRHI, FTextureRHIParamRef NewTextureRHI)
+void FD3D11DynamicRHI::RHIUpdateTextureReference(FRHITextureReference* TextureRefRHI, FTextureRHIParamRef NewTextureRHI)
 {
 	// Updating texture references is disallowed while the RHI could be caching them in referenced resource tables.
 	check(ResourceTableFrameCounter == INDEX_NONE);
