DR - Deprecate FShaderResourceViewRHIParamRef and FUnorderedAccessViewRHIParamRef

diff --git a/Engine/Plugins/FX/Niagara/Source/Niagara/Private/NiagaraEmitterInstanceBatcher.cpp b/Engine/Plugins/FX/Niagara/Source/Niagara/Private/NiagaraEmitterInstanceBatcher.cpp
--- a/Engine/Plugins/FX/Niagara/Source/Niagara/Private/NiagaraEmitterInstanceBatcher.cpp
+++ b/Engine/Plugins/FX/Niagara/Source/Niagara/Private/NiagaraEmitterInstanceBatcher.cpp
@@ -239,7 +239,7 @@ void NiagaraEmitterInstanceBatcher::SortGPUParticles(FRHICommandListImmediate& R
 		SET_DWORD_STAT(STAT_NiagaraGPUSortedBuffers, ParticleSortBuffers.GetSize());
 
 		// Make sure our outputs are safe to write to.
-		FUnorderedAccessViewRHIParamRef OutputUAVs[2] = { ParticleSortBuffers.GetKeyBufferUAV(), ParticleSortBuffers.GetVertexBufferUAV() };
+		FRHIUnorderedAccessView* OutputUAVs[2] = { ParticleSortBuffers.GetKeyBufferUAV(), ParticleSortBuffers.GetVertexBufferUAV() };
 		RHICmdList.TransitionResources(EResourceTransitionAccess::ERWBarrier, EResourceTransitionPipeline::EGfxToCompute, OutputUAVs, 2);
 
 		{
@@ -335,7 +335,7 @@ void NiagaraEmitterInstanceBatcher::ResolveParticleSortBuffers(FRHICommandListIm
 
 	for (int32 Index = 0; Index < SortedVertexBuffers.Num(); Index += NIAGARA_COPY_BUFFER_BUFFER_COUNT)
 	{
-		FUnorderedAccessViewRHIParamRef UAVs[NIAGARA_COPY_BUFFER_BUFFER_COUNT] = {};
+		FRHIUnorderedAccessView* UAVs[NIAGARA_COPY_BUFFER_BUFFER_COUNT] = {};
 		int32 UsedIndexCounts[NIAGARA_COPY_BUFFER_BUFFER_COUNT] = {};
 
 		const int32 NumBuffers = FMath::Min<int32>(NIAGARA_COPY_BUFFER_BUFFER_COUNT, SortedVertexBuffers.Num() - Index);
diff --git a/Engine/Plugins/FX/Niagara/Source/NiagaraVertexFactories/Private/NiagaraMeshVertexFactory.cpp b/Engine/Plugins/FX/Niagara/Source/NiagaraVertexFactories/Private/NiagaraMeshVertexFactory.cpp
--- a/Engine/Plugins/FX/Niagara/Source/NiagaraVertexFactories/Private/NiagaraMeshVertexFactory.cpp
+++ b/Engine/Plugins/FX/Niagara/Source/NiagaraVertexFactories/Private/NiagaraMeshVertexFactory.cpp
@@ -190,7 +190,7 @@ void FNiagaraMeshVertexFactory::UnlockPreviousTransformBuffer()
 	PrevTransformBuffer.Unlock();
 }
 
-FShaderResourceViewRHIParamRef FNiagaraMeshVertexFactory::GetPreviousTransformBufferSRV() const
+FRHIShaderResourceView* FNiagaraMeshVertexFactory::GetPreviousTransformBufferSRV() const
 {
 	return PrevTransformBuffer.SRV;
 }
diff --git a/Engine/Plugins/FX/Niagara/Source/NiagaraVertexFactories/Private/NiagaraSortingGPU.cpp b/Engine/Plugins/FX/Niagara/Source/NiagaraVertexFactories/Private/NiagaraSortingGPU.cpp
--- a/Engine/Plugins/FX/Niagara/Source/NiagaraVertexFactories/Private/NiagaraSortingGPU.cpp
+++ b/Engine/Plugins/FX/Niagara/Source/NiagaraVertexFactories/Private/NiagaraSortingGPU.cpp
@@ -92,7 +92,7 @@ bool FNiagaraSortKeyGenCS::Serialize(FArchive& Ar)
 	return bShaderHasOutdatedParameters;
 }
 
-void FNiagaraSortKeyGenCS::SetOutput(FRHICommandList& RHICmdList, FUnorderedAccessViewRHIParamRef OutKeysUAV, FUnorderedAccessViewRHIParamRef OutIndicesUAV)
+void FNiagaraSortKeyGenCS::SetOutput(FRHICommandList& RHICmdList, FRHIUnorderedAccessView* OutKeysUAV, FRHIUnorderedAccessView* OutIndicesUAV)
 {
 	FComputeShaderRHIParamRef ComputeShaderRHI = GetComputeShader();
 	if (OutKeys.IsBound())
@@ -130,19 +130,19 @@ void FNiagaraSortKeyGenCS::UnbindBuffers(FRHICommandList& RHICmdList)
 	FComputeShaderRHIParamRef ComputeShaderRHI = GetComputeShader();
 	if (NiagaraParticleDataFloat.IsBound())
 	{
-		RHICmdList.SetShaderResourceViewParameter(ComputeShaderRHI, NiagaraParticleDataFloat.GetBaseIndex(), FShaderResourceViewRHIParamRef());
+		RHICmdList.SetShaderResourceViewParameter(ComputeShaderRHI, NiagaraParticleDataFloat.GetBaseIndex(), nullptr);
 	}
 	if (GPUParticleCountBuffer.IsBound())
 	{
-		RHICmdList.SetShaderResourceViewParameter(ComputeShaderRHI, GPUParticleCountBuffer.GetBaseIndex(), FShaderResourceViewRHIParamRef());
+		RHICmdList.SetShaderResourceViewParameter(ComputeShaderRHI, GPUParticleCountBuffer.GetBaseIndex(), nullptr);
 	}
 	if (OutKeys.IsBound())
 	{
-		RHICmdList.SetUAVParameter(ComputeShaderRHI, OutKeys.GetBaseIndex(), FUnorderedAccessViewRHIParamRef());
+		RHICmdList.SetUAVParameter(ComputeShaderRHI, OutKeys.GetBaseIndex(), nullptr);
 	}
 	if (OutParticleIndices.IsBound())
 	{
-		RHICmdList.SetUAVParameter(ComputeShaderRHI, OutParticleIndices.GetBaseIndex(), FUnorderedAccessViewRHIParamRef());
+		RHICmdList.SetUAVParameter(ComputeShaderRHI, OutParticleIndices.GetBaseIndex(), nullptr);
 	}
 }
 
@@ -173,8 +173,8 @@ bool FNiagaraCopyIntBufferRegionCS::Serialize(FArchive& Ar)
 
 void FNiagaraCopyIntBufferRegionCS::SetParameters(
 	FRHICommandList& RHICmdList,
-	const FShaderResourceViewRHIParamRef InSourceData,
-	const FUnorderedAccessViewRHIParamRef* InDestDatas,
+	FRHIShaderResourceView* InSourceData,
+	FRHIUnorderedAccessView* const* InDestDatas,
 	const int32* InUsedIndexCounts,
 	int32 StartingIndex,
 	int32 DestCount)
@@ -205,13 +205,13 @@ void FNiagaraCopyIntBufferRegionCS::UnbindBuffers(FRHICommandList& RHICmdList)
 	FComputeShaderRHIParamRef ComputeShaderRHI = GetComputeShader();
 	if (SourceData.IsBound())
 	{
-		RHICmdList.SetShaderResourceViewParameter(ComputeShaderRHI, SourceData.GetBaseIndex(), FShaderResourceViewRHIParamRef());
+		RHICmdList.SetShaderResourceViewParameter(ComputeShaderRHI, SourceData.GetBaseIndex(), nullptr);
 	}
 	for (int32 Index = 0; Index < NIAGARA_COPY_BUFFER_BUFFER_COUNT; ++Index)
 	{
 		if (DestData[Index].IsBound())
 		{
-			RHICmdList.SetUAVParameter(ComputeShaderRHI, DestData[Index].GetBaseIndex(), FUnorderedAccessViewRHIParamRef());
+			RHICmdList.SetUAVParameter(ComputeShaderRHI, DestData[Index].GetBaseIndex(), nullptr);
 		}
 	}
 }
diff --git a/Engine/Plugins/FX/Niagara/Source/NiagaraVertexFactories/Private/NiagaraSpriteVertexFactory.cpp b/Engine/Plugins/FX/Niagara/Source/NiagaraVertexFactories/Private/NiagaraSpriteVertexFactory.cpp
--- a/Engine/Plugins/FX/Niagara/Source/NiagaraVertexFactories/Private/NiagaraSpriteVertexFactory.cpp
+++ b/Engine/Plugins/FX/Niagara/Source/NiagaraVertexFactories/Private/NiagaraSpriteVertexFactory.cpp
@@ -91,7 +91,7 @@ public:
 		ShaderBindings.Add(Shader->GetUniformBufferParameter<FNiagaraSpriteVFLooseParameters>(), SpriteVF->LooseParameterUniformBuffer);
 		
 		ShaderBindings.Add(NumCutoutVerticesPerFrame, SpriteVF->GetNumCutoutVerticesPerFrame());
-		FShaderResourceViewRHIParamRef NullSRV = GFNiagaraNullCutoutVertexBuffer.VertexBufferSRV;
+		FRHIShaderResourceView* NullSRV = GFNiagaraNullCutoutVertexBuffer.VertexBufferSRV;
 		ShaderBindings.Add(CutoutGeometry, SpriteVF->GetCutoutGeometrySRV() ? SpriteVF->GetCutoutGeometrySRV() : NullSRV);
 
 		ShaderBindings.Add(ParticleAlignmentMode, SpriteVF->GetAlignmentMode());
diff --git a/Engine/Plugins/FX/Niagara/Source/NiagaraVertexFactories/Public/NiagaraMeshVertexFactory.h b/Engine/Plugins/FX/Niagara/Source/NiagaraVertexFactories/Public/NiagaraMeshVertexFactory.h
--- a/Engine/Plugins/FX/Niagara/Source/NiagaraVertexFactories/Public/NiagaraMeshVertexFactory.h
+++ b/Engine/Plugins/FX/Niagara/Source/NiagaraVertexFactories/Public/NiagaraMeshVertexFactory.h
@@ -162,7 +162,7 @@ public:
 	
 	//uint8* LockPreviousTransformBuffer(uint32 ParticleCount);
 	//void UnlockPreviousTransformBuffer();
-	//FShaderResourceViewRHIParamRef GetPreviousTransformBufferSRV() const;
+	//FRHIShaderResourceView* GetPreviousTransformBufferSRV() const;
 
 	/**
 	* Copy the data from another vertex factory
diff --git a/Engine/Plugins/FX/Niagara/Source/NiagaraVertexFactories/Public/NiagaraSortingGPU.h b/Engine/Plugins/FX/Niagara/Source/NiagaraVertexFactories/Public/NiagaraSortingGPU.h
--- a/Engine/Plugins/FX/Niagara/Source/NiagaraVertexFactories/Public/NiagaraSortingGPU.h
+++ b/Engine/Plugins/FX/Niagara/Source/NiagaraVertexFactories/Public/NiagaraSortingGPU.h
@@ -49,7 +49,7 @@ public:
 	/**
 	 * Set output buffers for this shader.
 	 */
-	void SetOutput(FRHICommandList& RHICmdList, FUnorderedAccessViewRHIParamRef OutKeysUAV, FUnorderedAccessViewRHIParamRef OutIndicesUAV);
+	void SetOutput(FRHICommandList& RHICmdList, FRHIUnorderedAccessView* OutKeysUAV, FRHIUnorderedAccessView* OutIndicesUAV);
 
 	/**
 	 * Set input parameters.
@@ -115,8 +115,8 @@ public:
 	 */
 	void SetParameters(
 		FRHICommandList& RHICmdList,	
-		const FShaderResourceViewRHIParamRef InSourceData, 
-		const FUnorderedAccessViewRHIParamRef* InDestDatas, 
+		FRHIShaderResourceView* InSourceData,
+		FRHIUnorderedAccessView* const* InDestDatas,
 		const int32* InUsedIndexCounts, 
 		int32 StartingIndex,
 		int32 DestCount);
diff --git a/Engine/Plugins/FX/Niagara/Source/NiagaraVertexFactories/Public/NiagaraSpriteVertexFactory.h b/Engine/Plugins/FX/Niagara/Source/NiagaraVertexFactories/Public/NiagaraSpriteVertexFactory.h
--- a/Engine/Plugins/FX/Niagara/Source/NiagaraVertexFactories/Public/NiagaraSpriteVertexFactory.h
+++ b/Engine/Plugins/FX/Niagara/Source/NiagaraVertexFactories/Public/NiagaraSpriteVertexFactory.h
@@ -139,14 +139,14 @@ public:
 		return SpriteUniformBuffer;
 	}
 
-	void SetCutoutParameters(int32 InNumCutoutVerticesPerFrame, FShaderResourceViewRHIParamRef InCutoutGeometrySRV)
+	void SetCutoutParameters(int32 InNumCutoutVerticesPerFrame, FRHIShaderResourceView* InCutoutGeometrySRV)
 	{
 		NumCutoutVerticesPerFrame = InNumCutoutVerticesPerFrame;
 		CutoutGeometrySRV = InCutoutGeometrySRV;
 	}
 
 	inline int32 GetNumCutoutVerticesPerFrame() const { return NumCutoutVerticesPerFrame; }
-	inline FShaderResourceViewRHIParamRef GetCutoutGeometrySRV() const { return CutoutGeometrySRV; }
+	inline FRHIShaderResourceView* GetCutoutGeometrySRV() const { return CutoutGeometrySRV; }
 
 	void SetParticleData(const FShaderResourceViewRHIRef& InParticleDataFloatSRV, uint32 InFloatDataOffset, uint32 InFloatDataStride)
 	{
diff --git a/Engine/Plugins/Runtime/Oculus/OculusVR/Source/OculusHMD/Private/OculusHMD_StressTester.cpp b/Engine/Plugins/Runtime/Oculus/OculusVR/Source/OculusHMD/Private/OculusHMD_StressTester.cpp
--- a/Engine/Plugins/Runtime/Oculus/OculusVR/Source/OculusHMD/Private/OculusHMD_StressTester.cpp
+++ b/Engine/Plugins/Runtime/Oculus/OculusVR/Source/OculusHMD/Private/OculusHMD_StressTester.cpp
@@ -170,7 +170,7 @@ public:
 
 		if (TextureParameter.IsBound())
 		{
-			RHICmdList.SetShaderResourceViewParameter(PixelShaderRHI, TextureParameter.GetBaseIndex(), FShaderResourceViewRHIParamRef());
+			RHICmdList.SetShaderResourceViewParameter(PixelShaderRHI, TextureParameter.GetBaseIndex(), nullptr);
 		}
 	}
 
diff --git a/Engine/Source/Runtime/Apple/MetalRHI/Private/MetalCommands.cpp b/Engine/Source/Runtime/Apple/MetalRHI/Private/MetalCommands.cpp
--- a/Engine/Source/Runtime/Apple/MetalRHI/Private/MetalCommands.cpp
+++ b/Engine/Source/Runtime/Apple/MetalRHI/Private/MetalCommands.cpp
@@ -246,7 +246,7 @@ void FMetalRHICommandContext::RHISetGraphicsPipelineState(FRHIGraphicsPipelineSt
 	}
 }
 
-void FMetalRHICommandContext::RHISetUAVParameter(FComputeShaderRHIParamRef ComputeShaderRHI, uint32 UAVIndex, FUnorderedAccessViewRHIParamRef UAVRHI)
+void FMetalRHICommandContext::RHISetUAVParameter(FComputeShaderRHIParamRef ComputeShaderRHI, uint32 UAVIndex, FRHIUnorderedAccessView* UAVRHI)
 {
 	@autoreleasepool {
 	FMetalUnorderedAccessView* UAV = ResourceCast(UAVRHI);
@@ -254,7 +254,7 @@ void FMetalRHICommandContext::RHISetUAVParameter(FComputeShaderRHIParamRef Compu
 	}
 }
 
-void FMetalRHICommandContext::RHISetUAVParameter(FComputeShaderRHIParamRef ComputeShaderRHI,uint32 UAVIndex,FUnorderedAccessViewRHIParamRef UAVRHI, uint32 InitialCount)
+void FMetalRHICommandContext::RHISetUAVParameter(FComputeShaderRHIParamRef ComputeShaderRHI,uint32 UAVIndex, FRHIUnorderedAccessView* UAVRHI, uint32 InitialCount)
 {
 	@autoreleasepool {
 		FMetalUnorderedAccessView* UAV = ResourceCast(UAVRHI);
@@ -389,7 +389,7 @@ void FMetalRHICommandContext::RHISetShaderTexture(FComputeShaderRHIParamRef Comp
 }
 
 
-void FMetalRHICommandContext::RHISetShaderResourceViewParameter(FVertexShaderRHIParamRef VertexShaderRHI, uint32 TextureIndex, FShaderResourceViewRHIParamRef SRVRHI)
+void FMetalRHICommandContext::RHISetShaderResourceViewParameter(FVertexShaderRHIParamRef VertexShaderRHI, uint32 TextureIndex, FRHIShaderResourceView* SRVRHI)
 {
 	@autoreleasepool {
 	FMetalShaderResourceView* SRV = ResourceCast(SRVRHI);
@@ -397,7 +397,7 @@ void FMetalRHICommandContext::RHISetShaderResourceViewParameter(FVertexShaderRHI
 	}
 }
 
-void FMetalRHICommandContext::RHISetShaderResourceViewParameter(FHullShaderRHIParamRef HullShaderRHI,uint32 TextureIndex,FShaderResourceViewRHIParamRef SRVRHI)
+void FMetalRHICommandContext::RHISetShaderResourceViewParameter(FHullShaderRHIParamRef HullShaderRHI,uint32 TextureIndex, FRHIShaderResourceView* SRVRHI)
 {
 #if PLATFORM_SUPPORTS_TESSELLATION_SHADERS
 	@autoreleasepool {
@@ -407,7 +407,7 @@ void FMetalRHICommandContext::RHISetShaderResourceViewParameter(FHullShaderRHIPa
 #endif
 }
 
-void FMetalRHICommandContext::RHISetShaderResourceViewParameter(FDomainShaderRHIParamRef DomainShaderRHI,uint32 TextureIndex,FShaderResourceViewRHIParamRef SRVRHI)
+void FMetalRHICommandContext::RHISetShaderResourceViewParameter(FDomainShaderRHIParamRef DomainShaderRHI,uint32 TextureIndex, FRHIShaderResourceView* SRVRHI)
 {
 #if PLATFORM_SUPPORTS_TESSELLATION_SHADERS
 	@autoreleasepool {
@@ -417,12 +417,12 @@ void FMetalRHICommandContext::RHISetShaderResourceViewParameter(FDomainShaderRHI
 #endif
 }
 
-void FMetalRHICommandContext::RHISetShaderResourceViewParameter(FGeometryShaderRHIParamRef GeometryShaderRHI,uint32 TextureIndex,FShaderResourceViewRHIParamRef SRVRHI)
+void FMetalRHICommandContext::RHISetShaderResourceViewParameter(FGeometryShaderRHIParamRef GeometryShaderRHI,uint32 TextureIndex, FRHIShaderResourceView* SRVRHI)
 {
 	NOT_SUPPORTED("RHISetShaderResourceViewParameter");
 }
 
-void FMetalRHICommandContext::RHISetShaderResourceViewParameter(FPixelShaderRHIParamRef PixelShaderRHI,uint32 TextureIndex,FShaderResourceViewRHIParamRef SRVRHI)
+void FMetalRHICommandContext::RHISetShaderResourceViewParameter(FPixelShaderRHIParamRef PixelShaderRHI,uint32 TextureIndex, FRHIShaderResourceView* SRVRHI)
 {
 	@autoreleasepool {
 	FMetalShaderResourceView* SRV = ResourceCast(SRVRHI);
@@ -430,7 +430,7 @@ void FMetalRHICommandContext::RHISetShaderResourceViewParameter(FPixelShaderRHIP
 	}
 }
 
-void FMetalRHICommandContext::RHISetShaderResourceViewParameter(FComputeShaderRHIParamRef ComputeShaderRHI,uint32 TextureIndex,FShaderResourceViewRHIParamRef SRVRHI)
+void FMetalRHICommandContext::RHISetShaderResourceViewParameter(FComputeShaderRHIParamRef ComputeShaderRHI,uint32 TextureIndex, FRHIShaderResourceView* SRVRHI)
 {
 	@autoreleasepool {
 	FMetalShaderResourceView* SRV = ResourceCast(SRVRHI);
@@ -635,7 +635,7 @@ void FMetalRHICommandContext::RHISetBlendFactor(const FLinearColor& BlendFactor)
 }
 
 void FMetalRHICommandContext::RHISetRenderTargets(uint32 NumSimultaneousRenderTargets, const FRHIRenderTargetView* NewRenderTargets,
-	const FRHIDepthRenderTargetView* NewDepthStencilTargetRHI, uint32 NumUAVs, const FUnorderedAccessViewRHIParamRef* UAVs)
+	const FRHIDepthRenderTargetView* NewDepthStencilTargetRHI, uint32 NumUAVs, FRHIUnorderedAccessView* const* UAVs)
 {
 	@autoreleasepool {
 	FMetalContext* Manager = Context;
diff --git a/Engine/Source/Runtime/Apple/MetalRHI/Private/MetalContext.cpp b/Engine/Source/Runtime/Apple/MetalRHI/Private/MetalContext.cpp
--- a/Engine/Source/Runtime/Apple/MetalRHI/Private/MetalContext.cpp
+++ b/Engine/Source/Runtime/Apple/MetalRHI/Private/MetalContext.cpp
@@ -1106,7 +1106,7 @@ void FMetalContext::FinishFrame()
 #endif
 }
 
-void FMetalContext::TransitionResources(FUnorderedAccessViewRHIParamRef* InUAVs, int32 NumUAVs)
+void FMetalContext::TransitionResources(FRHIUnorderedAccessView** InUAVs, int32 NumUAVs)
 {
 	for (uint32 i = 0; i < NumUAVs; i++)
 	{
diff --git a/Engine/Source/Runtime/Apple/MetalRHI/Private/MetalContext.h b/Engine/Source/Runtime/Apple/MetalRHI/Private/MetalContext.h
--- a/Engine/Source/Runtime/Apple/MetalRHI/Private/MetalContext.h
+++ b/Engine/Source/Runtime/Apple/MetalRHI/Private/MetalContext.h
@@ -124,7 +124,7 @@ public:
 	void FinishFrame();
 
 	// Track Write->Read transitions for TBDR Fragment->Verex fencing
-	void TransitionResources(FUnorderedAccessViewRHIParamRef* InUAVs, int32 NumUAVs);
+	void TransitionResources(FRHIUnorderedAccessView** InUAVs, int32 NumUAVs);
 	void TransitionResources(FTextureRHIParamRef* InTextures, int32 NumTextures);
 	
 protected:
diff --git a/Engine/Source/Runtime/Apple/MetalRHI/Private/MetalUAV.cpp b/Engine/Source/Runtime/Apple/MetalRHI/Private/MetalUAV.cpp
--- a/Engine/Source/Runtime/Apple/MetalRHI/Private/MetalUAV.cpp
+++ b/Engine/Source/Runtime/Apple/MetalRHI/Private/MetalUAV.cpp
@@ -339,7 +339,7 @@ FShaderResourceViewRHIRef FMetalDynamicRHI::RHICreateShaderResourceView(FIndexBu
 	}
 }
 
-void FMetalDynamicRHI::RHIUpdateShaderResourceView(FShaderResourceViewRHIParamRef SRVRHI, FVertexBufferRHIParamRef VertexBufferRHI, uint32 Stride, uint8 Format)
+void FMetalDynamicRHI::RHIUpdateShaderResourceView(FRHIShaderResourceView* SRVRHI, FVertexBufferRHIParamRef VertexBufferRHI, uint32 Stride, uint8 Format)
 {
 	check(SRVRHI);
 	FMetalShaderResourceView* SRV = ResourceCast(SRVRHI);
@@ -364,7 +364,7 @@ void FMetalDynamicRHI::RHIUpdateShaderResourceView(FShaderResourceViewRHIParamRe
 	}
 }
 
-void FMetalDynamicRHI::RHIUpdateShaderResourceView(FShaderResourceViewRHIParamRef SRVRHI, FIndexBufferRHIParamRef IndexBufferRHI)
+void FMetalDynamicRHI::RHIUpdateShaderResourceView(FRHIShaderResourceView* SRVRHI, FIndexBufferRHIParamRef IndexBufferRHI)
 {
 	check(SRVRHI);
 	FMetalShaderResourceView* SRV = ResourceCast(SRVRHI);
@@ -389,7 +389,7 @@ void FMetalDynamicRHI::RHIUpdateShaderResourceView(FShaderResourceViewRHIParamRe
 	}
 }
 
-void FMetalRHICommandContext::RHIClearTinyUAV(FUnorderedAccessViewRHIParamRef UnorderedAccessViewRHI, const uint32* Values)
+void FMetalRHICommandContext::RHIClearTinyUAV(FRHIUnorderedAccessView* UnorderedAccessViewRHI, const uint32* Values)
 {
 	@autoreleasepool {
 	FMetalUnorderedAccessView* UnorderedAccessView = ResourceCast(UnorderedAccessViewRHI);
@@ -715,7 +715,7 @@ void FMetalComputeFence::Reset()
 	Fence = nullptr;
 }
 
-void FMetalRHICommandContext::RHITransitionResources(EResourceTransitionAccess TransitionType, EResourceTransitionPipeline TransitionPipeline, FUnorderedAccessViewRHIParamRef* InUAVs, int32 NumUAVs, FRHIComputeFence* WriteComputeFence)
+void FMetalRHICommandContext::RHITransitionResources(EResourceTransitionAccess TransitionType, EResourceTransitionPipeline TransitionPipeline, FRHIUnorderedAccessView** InUAVs, int32 NumUAVs, FRHIComputeFence* WriteComputeFence)
 {
 	@autoreleasepool
 	{
diff --git a/Engine/Source/Runtime/Apple/MetalRHI/Public/MetalRHI.h b/Engine/Source/Runtime/Apple/MetalRHI/Public/MetalRHI.h
--- a/Engine/Source/Runtime/Apple/MetalRHI/Public/MetalRHI.h
+++ b/Engine/Source/Runtime/Apple/MetalRHI/Public/MetalRHI.h
@@ -83,8 +83,8 @@ public:
 	virtual FShaderResourceViewRHIRef RHICreateShaderResourceView(FTextureRHIParamRef Texture2DRHI, const FRHITextureSRVCreateInfo& CreateInfo) final override;
 	virtual FShaderResourceViewRHIRef RHICreateShaderResourceView(FVertexBufferRHIParamRef VertexBuffer, uint32 Stride, uint8 Format) final override;
 	virtual FShaderResourceViewRHIRef RHICreateShaderResourceView(FIndexBufferRHIParamRef Buffer) final override;
-	virtual void RHIUpdateShaderResourceView(FShaderResourceViewRHIParamRef SRV, FVertexBufferRHIParamRef VertexBuffer, uint32 Stride, uint8 Format) final override;
-	virtual void RHIUpdateShaderResourceView(FShaderResourceViewRHIParamRef SRV, FIndexBufferRHIParamRef IndexBuffer) final override;
+	virtual void RHIUpdateShaderResourceView(FRHIShaderResourceView* SRV, FVertexBufferRHIParamRef VertexBuffer, uint32 Stride, uint8 Format) final override;
+	virtual void RHIUpdateShaderResourceView(FRHIShaderResourceView* SRV, FIndexBufferRHIParamRef IndexBuffer) final override;
 	virtual uint64 RHICalcTexture2DPlatformSize(uint32 SizeX, uint32 SizeY, uint8 Format, uint32 NumMips, uint32 NumSamples, uint32 Flags, uint32& OutAlign) final override;
 	virtual uint64 RHICalcTexture3DPlatformSize(uint32 SizeX, uint32 SizeY, uint32 SizeZ, uint8 Format, uint32 NumMips, uint32 Flags, uint32& OutAlign) final override;
 	virtual uint64 RHICalcTextureCubePlatformSize(uint32 Size, uint8 Format, uint32 NumMips, uint32 Flags, uint32& OutAlign) final override;
diff --git a/Engine/Source/Runtime/Apple/MetalRHI/Public/MetalRHIContext.h b/Engine/Source/Runtime/Apple/MetalRHI/Public/MetalRHIContext.h
--- a/Engine/Source/Runtime/Apple/MetalRHI/Public/MetalRHIContext.h
+++ b/Engine/Source/Runtime/Apple/MetalRHI/Public/MetalRHIContext.h
@@ -45,7 +45,7 @@ public:
 	virtual void RHISetMultipleViewports(uint32 Count, const FViewportBounds* Data) final override;
 	
 	/** Clears a UAV to the multi-component value provided. */
-	virtual void RHIClearTinyUAV(FUnorderedAccessViewRHIParamRef UnorderedAccessViewRHI, const uint32* Values) final override;
+	virtual void RHIClearTinyUAV(FRHIUnorderedAccessView* UnorderedAccessViewRHI, const uint32* Values) final override;
 	
 	virtual void RHICopyTexture(FTextureRHIParamRef SourceTextureRHI, FTextureRHIParamRef DestTextureRHI, const FRHICopyTextureInfo& CopyInfo) final override;
 	
@@ -180,22 +180,22 @@ public:
 	virtual void RHISetShaderSampler(FPixelShaderRHIParamRef PixelShader, uint32 SamplerIndex, FRHISamplerState* NewState) final override;
 	
 	/** Sets a compute shader UAV parameter. */
-	virtual void RHISetUAVParameter(FComputeShaderRHIParamRef ComputeShader, uint32 UAVIndex, FUnorderedAccessViewRHIParamRef UAV) final override;
+	virtual void RHISetUAVParameter(FComputeShaderRHIParamRef ComputeShader, uint32 UAVIndex, FRHIUnorderedAccessView* UAV) final override;
 	
 	/** Sets a compute shader UAV parameter and initial count */
-	virtual void RHISetUAVParameter(FComputeShaderRHIParamRef ComputeShader, uint32 UAVIndex, FUnorderedAccessViewRHIParamRef UAV, uint32 InitialCount) final override;
+	virtual void RHISetUAVParameter(FComputeShaderRHIParamRef ComputeShader, uint32 UAVIndex, FRHIUnorderedAccessView* UAV, uint32 InitialCount) final override;
 	
-	virtual void RHISetShaderResourceViewParameter(FPixelShaderRHIParamRef PixelShader, uint32 SamplerIndex, FShaderResourceViewRHIParamRef SRV) final override;
+	virtual void RHISetShaderResourceViewParameter(FPixelShaderRHIParamRef PixelShader, uint32 SamplerIndex, FRHIShaderResourceView* SRV) final override;
 	
-	virtual void RHISetShaderResourceViewParameter(FVertexShaderRHIParamRef VertexShader, uint32 SamplerIndex, FShaderResourceViewRHIParamRef SRV) final override;
+	virtual void RHISetShaderResourceViewParameter(FVertexShaderRHIParamRef VertexShader, uint32 SamplerIndex, FRHIShaderResourceView* SRV) final override;
 	
-	virtual void RHISetShaderResourceViewParameter(FComputeShaderRHIParamRef ComputeShader, uint32 SamplerIndex, FShaderResourceViewRHIParamRef SRV) final override;
+	virtual void RHISetShaderResourceViewParameter(FComputeShaderRHIParamRef ComputeShader, uint32 SamplerIndex, FRHIShaderResourceView* SRV) final override;
 	
-	virtual void RHISetShaderResourceViewParameter(FHullShaderRHIParamRef HullShader, uint32 SamplerIndex, FShaderResourceViewRHIParamRef SRV) final override;
+	virtual void RHISetShaderResourceViewParameter(FHullShaderRHIParamRef HullShader, uint32 SamplerIndex, FRHIShaderResourceView* SRV) final override;
 	
-	virtual void RHISetShaderResourceViewParameter(FDomainShaderRHIParamRef DomainShader, uint32 SamplerIndex, FShaderResourceViewRHIParamRef SRV) final override;
+	virtual void RHISetShaderResourceViewParameter(FDomainShaderRHIParamRef DomainShader, uint32 SamplerIndex, FRHIShaderResourceView* SRV) final override;
 	
-	virtual void RHISetShaderResourceViewParameter(FGeometryShaderRHIParamRef GeometryShader, uint32 SamplerIndex, FShaderResourceViewRHIParamRef SRV) final override;
+	virtual void RHISetShaderResourceViewParameter(FGeometryShaderRHIParamRef GeometryShader, uint32 SamplerIndex, FRHIShaderResourceView* SRV) final override;
 	
 	virtual void RHISetShaderUniformBuffer(FVertexShaderRHIParamRef VertexShader, uint32 BufferIndex, FUniformBufferRHIParamRef Buffer) final override;
 	
@@ -225,7 +225,7 @@ public:
 
 	virtual void RHISetBlendFactor(const FLinearColor& BlendFactor) final override;
 
-	virtual void RHISetRenderTargets(uint32 NumSimultaneousRenderTargets, const FRHIRenderTargetView* NewRenderTargets, const FRHIDepthRenderTargetView* NewDepthStencilTarget, uint32 NumUAVs, const FUnorderedAccessViewRHIParamRef* UAVs) final override;
+	virtual void RHISetRenderTargets(uint32 NumSimultaneousRenderTargets, const FRHIRenderTargetView* NewRenderTargets, const FRHIDepthRenderTargetView* NewDepthStencilTarget, uint32 NumUAVs, FRHIUnorderedAccessView* const* UAVs) final override;
 	
 	virtual void RHISetRenderTargetsAndClear(const FRHISetRenderTargetsInfo& RenderTargetsInfo) final override;
 	
@@ -265,7 +265,7 @@ public:
 	 * @param NumUAVs - number of UAVs to transition
 	 * @param WriteComputeFence - Optional ComputeFence to write as part of this transition
 	 */
-	virtual void RHITransitionResources(EResourceTransitionAccess TransitionType, EResourceTransitionPipeline TransitionPipeline, FUnorderedAccessViewRHIParamRef* InUAVs, int32 NumUAVs, FRHIComputeFence* WriteComputeFence) final override;
+	virtual void RHITransitionResources(EResourceTransitionAccess TransitionType, EResourceTransitionPipeline TransitionPipeline, FRHIUnorderedAccessView** InUAVs, int32 NumUAVs, FRHIComputeFence* WriteComputeFence) final override;
 	
 	virtual void RHICopyToStagingBuffer(FVertexBufferRHIParamRef SourceBufferRHI, FStagingBufferRHIParamRef DestinationStagingBufferRHI, uint32 Offset, uint32 NumBytes) final override;
 	virtual void RHIWriteGPUFence(FGPUFenceRHIParamRef FenceRHI) final override;
diff --git a/Engine/Source/Runtime/D3D12RHI/Private/D3D12CommandContext.cpp b/Engine/Source/Runtime/D3D12RHI/Private/D3D12CommandContext.cpp
--- a/Engine/Source/Runtime/D3D12RHI/Private/D3D12CommandContext.cpp
+++ b/Engine/Source/Runtime/D3D12RHI/Private/D3D12CommandContext.cpp
@@ -712,7 +712,7 @@ FD3D12CommandContextRedirector::FD3D12CommandContextRedirector(class FD3D12Adapt
 	FMemory::Memzero(PhysicalContexts, sizeof(PhysicalContexts[0]) * MAX_NUM_GPUS);
 }
 
-void FD3D12CommandContextRedirector::RHITransitionResources(EResourceTransitionAccess TransitionType, EResourceTransitionPipeline TransitionPipeline, FUnorderedAccessViewRHIParamRef* InUAVs, int32 NumUAVs, FRHIComputeFence* WriteComputeFenceRHI)
+void FD3D12CommandContextRedirector::RHITransitionResources(EResourceTransitionAccess TransitionType, EResourceTransitionPipeline TransitionPipeline, FRHIUnorderedAccessView** InUAVs, int32 NumUAVs, FRHIComputeFence* WriteComputeFenceRHI)
 {
 	ContextRedirect(RHITransitionResources(TransitionType, TransitionPipeline, InUAVs, NumUAVs, nullptr));
 
diff --git a/Engine/Source/Runtime/D3D12RHI/Private/D3D12CommandContext.h b/Engine/Source/Runtime/D3D12RHI/Private/D3D12CommandContext.h
--- a/Engine/Source/Runtime/D3D12RHI/Private/D3D12CommandContext.h
+++ b/Engine/Source/Runtime/D3D12RHI/Private/D3D12CommandContext.h
@@ -232,12 +232,12 @@ public:
 	virtual void RHISetComputePipelineState(FRHIComputePipelineState* ComputePipelineState) final override;
 	virtual void RHIDispatchComputeShader(uint32 ThreadGroupCountX, uint32 ThreadGroupCountY, uint32 ThreadGroupCountZ) final override;
 	virtual void RHIDispatchIndirectComputeShader(FVertexBufferRHIParamRef ArgumentBuffer, uint32 ArgumentOffset) final override;
-	virtual void RHITransitionResources(EResourceTransitionAccess TransitionType, EResourceTransitionPipeline TransitionPipeline, FUnorderedAccessViewRHIParamRef* InUAVs, int32 NumUAVs, FRHIComputeFence* WriteComputeFenceRHI) final override;
+	virtual void RHITransitionResources(EResourceTransitionAccess TransitionType, EResourceTransitionPipeline TransitionPipeline, FRHIUnorderedAccessView** InUAVs, int32 NumUAVs, FRHIComputeFence* WriteComputeFenceRHI) final override;
 	virtual void RHISetShaderTexture(FComputeShaderRHIParamRef PixelShader, uint32 TextureIndex, FTextureRHIParamRef NewTexture) final override;
 	virtual void RHISetShaderSampler(FComputeShaderRHIParamRef ComputeShader, uint32 SamplerIndex, FRHISamplerState* NewState) final override;
-	virtual void RHISetUAVParameter(FComputeShaderRHIParamRef ComputeShader, uint32 UAVIndex, FUnorderedAccessViewRHIParamRef UAV) final override;
-	virtual void RHISetUAVParameter(FComputeShaderRHIParamRef ComputeShader, uint32 UAVIndex, FUnorderedAccessViewRHIParamRef UAV, uint32 InitialCount) final override;
-	virtual void RHISetShaderResourceViewParameter(FComputeShaderRHIParamRef ComputeShader, uint32 SamplerIndex, FShaderResourceViewRHIParamRef SRV) final override;
+	virtual void RHISetUAVParameter(FComputeShaderRHIParamRef ComputeShader, uint32 UAVIndex, FRHIUnorderedAccessView* UAV) final override;
+	virtual void RHISetUAVParameter(FComputeShaderRHIParamRef ComputeShader, uint32 UAVIndex, FRHIUnorderedAccessView* UAV, uint32 InitialCount) final override;
+	virtual void RHISetShaderResourceViewParameter(FComputeShaderRHIParamRef ComputeShader, uint32 SamplerIndex, FRHIShaderResourceView* SRV) final override;
 	virtual void RHISetShaderUniformBuffer(FComputeShaderRHIParamRef ComputeShader, uint32 BufferIndex, FUniformBufferRHIParamRef Buffer) final override;
 	virtual void RHISetShaderParameter(FComputeShaderRHIParamRef ComputeShader, uint32 BufferIndex, uint32 BaseIndex, uint32 NumBytes, const void* NewValue) final override;
 	virtual void RHIPushEvent(const TCHAR* Name, FColor Color) final override;
@@ -248,7 +248,7 @@ public:
 	virtual void RHIAutomaticCacheFlushAfterComputeShader(bool bEnable) final override;
 	virtual void RHIFlushComputeShaderCache() final override;
 	virtual void RHISetMultipleViewports(uint32 Count, const FViewportBounds* Data) final override;
-	virtual void RHIClearTinyUAV(FUnorderedAccessViewRHIParamRef UnorderedAccessViewRHI, const uint32* Values) final override;
+	virtual void RHIClearTinyUAV(FRHIUnorderedAccessView* UnorderedAccessViewRHI, const uint32* Values) final override;
 	virtual void RHICopyToResolveTarget(FTextureRHIParamRef SourceTexture, FTextureRHIParamRef DestTexture, const FResolveParams& ResolveParams) override;
 	virtual void RHICopyTexture(FTextureRHIParamRef SourceTexture, FTextureRHIParamRef DestTexture, const FRHICopyTextureInfo& CopyInfo) final override;
 	virtual void RHITransitionResources(EResourceTransitionAccess TransitionType, FTextureRHIParamRef* InTextures, int32 NumTextures) final override;
@@ -274,11 +274,11 @@ public:
 	virtual void RHISetShaderSampler(FDomainShaderRHIParamRef DomainShader, uint32 SamplerIndex, FRHISamplerState* NewState) final override;
 	virtual void RHISetShaderSampler(FHullShaderRHIParamRef HullShader, uint32 SamplerIndex, FRHISamplerState* NewState) final override;
 	virtual void RHISetShaderSampler(FPixelShaderRHIParamRef PixelShader, uint32 SamplerIndex, FRHISamplerState* NewState) final override;
-	virtual void RHISetShaderResourceViewParameter(FPixelShaderRHIParamRef PixelShader, uint32 SamplerIndex, FShaderResourceViewRHIParamRef SRV) final override;
-	virtual void RHISetShaderResourceViewParameter(FVertexShaderRHIParamRef VertexShader, uint32 SamplerIndex, FShaderResourceViewRHIParamRef SRV) final override;
-	virtual void RHISetShaderResourceViewParameter(FHullShaderRHIParamRef HullShader, uint32 SamplerIndex, FShaderResourceViewRHIParamRef SRV) final override;
-	virtual void RHISetShaderResourceViewParameter(FDomainShaderRHIParamRef DomainShader, uint32 SamplerIndex, FShaderResourceViewRHIParamRef SRV) final override;
-	virtual void RHISetShaderResourceViewParameter(FGeometryShaderRHIParamRef GeometryShader, uint32 SamplerIndex, FShaderResourceViewRHIParamRef SRV) final override;
+	virtual void RHISetShaderResourceViewParameter(FPixelShaderRHIParamRef PixelShader, uint32 SamplerIndex, FRHIShaderResourceView* SRV) final override;
+	virtual void RHISetShaderResourceViewParameter(FVertexShaderRHIParamRef VertexShader, uint32 SamplerIndex, FRHIShaderResourceView* SRV) final override;
+	virtual void RHISetShaderResourceViewParameter(FHullShaderRHIParamRef HullShader, uint32 SamplerIndex, FRHIShaderResourceView* SRV) final override;
+	virtual void RHISetShaderResourceViewParameter(FDomainShaderRHIParamRef DomainShader, uint32 SamplerIndex, FRHIShaderResourceView* SRV) final override;
+	virtual void RHISetShaderResourceViewParameter(FGeometryShaderRHIParamRef GeometryShader, uint32 SamplerIndex, FRHIShaderResourceView* SRV) final override;
 	virtual void RHISetShaderUniformBuffer(FVertexShaderRHIParamRef VertexShader, uint32 BufferIndex, FUniformBufferRHIParamRef Buffer) final override;
 	virtual void RHISetShaderUniformBuffer(FHullShaderRHIParamRef HullShader, uint32 BufferIndex, FUniformBufferRHIParamRef Buffer) final override;
 	virtual void RHISetShaderUniformBuffer(FDomainShaderRHIParamRef DomainShader, uint32 BufferIndex, FUniformBufferRHIParamRef Buffer) final override;
@@ -291,7 +291,7 @@ public:
 	virtual void RHISetShaderParameter(FGeometryShaderRHIParamRef GeometryShader, uint32 BufferIndex, uint32 BaseIndex, uint32 NumBytes, const void* NewValue) final override;
 	virtual void RHISetStencilRef(uint32 StencilRef) final override;
 	virtual void RHISetBlendFactor(const FLinearColor& BlendFactor) final override;
-	virtual void RHISetRenderTargets(uint32 NumSimultaneousRenderTargets, const FRHIRenderTargetView* NewRenderTargets, const FRHIDepthRenderTargetView* NewDepthStencilTarget, uint32 NumUAVs, const FUnorderedAccessViewRHIParamRef* UAVs) final override;
+	virtual void RHISetRenderTargets(uint32 NumSimultaneousRenderTargets, const FRHIRenderTargetView* NewRenderTargets, const FRHIDepthRenderTargetView* NewDepthStencilTarget, uint32 NumUAVs, FRHIUnorderedAccessView* const* UAVs) final override;
 	virtual void RHISetRenderTargetsAndClear(const FRHISetRenderTargetsInfo& RenderTargetsInfo) final override;
 	virtual void RHIBindClearMRTValues(bool bClearColor, bool bClearDepth, bool bClearStencil) final override;
 	virtual void RHIDrawPrimitive(uint32 BaseVertexIndex, uint32 NumPrimitives, uint32 NumInstances) final override;
@@ -339,12 +339,12 @@ public:
 	virtual void RHIBuildAccelerationStructures(const TArrayView<const FAccelerationStructureUpdateParams> Params) final override;
 	virtual void RHIBuildAccelerationStructure(FRayTracingSceneRHIParamRef Scene) final override;
 	virtual void RHIRayTraceOcclusion(FRayTracingSceneRHIParamRef Scene,
-		FShaderResourceViewRHIParamRef Rays,
-		FUnorderedAccessViewRHIParamRef Output,
+		FRHIShaderResourceView* Rays,
+		FRHIUnorderedAccessView* Output,
 		uint32 NumRays) final override;
 	virtual void RHIRayTraceIntersection(FRayTracingSceneRHIParamRef Scene,
-		FShaderResourceViewRHIParamRef Rays,
-		FUnorderedAccessViewRHIParamRef Output,
+		FRHIShaderResourceView* Rays,
+		FRHIUnorderedAccessView* Output,
 		uint32 NumRays) final override;
 	virtual void RHIRayTraceDispatch(FRHIRayTracingPipelineState* RayTracingPipelineState, FRayTracingShaderRHIParamRef RayGenShader,
 		FRayTracingSceneRHIParamRef Scene,
@@ -460,7 +460,7 @@ public:
 		ContextRedirect(RHIDispatchIndirectComputeShader(ArgumentBuffer, ArgumentOffset));
 	}
 	// Special implementation that only signal the fence once.
-	virtual void RHITransitionResources(EResourceTransitionAccess TransitionType, EResourceTransitionPipeline TransitionPipeline, FUnorderedAccessViewRHIParamRef* InUAVs, int32 NumUAVs, FRHIComputeFence* WriteComputeFenceRHI) final override;
+	virtual void RHITransitionResources(EResourceTransitionAccess TransitionType, EResourceTransitionPipeline TransitionPipeline, FRHIUnorderedAccessView** InUAVs, int32 NumUAVs, FRHIComputeFence* WriteComputeFenceRHI) final override;
 
 	FORCEINLINE virtual void RHICopyToStagingBuffer(FVertexBufferRHIParamRef SourceBuffer, FStagingBufferRHIParamRef DestinationStagingBuffer, uint32 Offset, uint32 NumBytes) final override
 	{
@@ -478,15 +478,15 @@ public:
 	{
 		ContextRedirect(RHISetShaderSampler(ComputeShader, SamplerIndex, NewState));
 	}
-	FORCEINLINE virtual void RHISetUAVParameter(FComputeShaderRHIParamRef ComputeShader, uint32 UAVIndex, FUnorderedAccessViewRHIParamRef UAV) final override
+	FORCEINLINE virtual void RHISetUAVParameter(FComputeShaderRHIParamRef ComputeShader, uint32 UAVIndex, FRHIUnorderedAccessView* UAV) final override
 	{
 		ContextRedirect(RHISetUAVParameter(ComputeShader, UAVIndex, UAV));
 	}
-	FORCEINLINE virtual void RHISetUAVParameter(FComputeShaderRHIParamRef ComputeShader, uint32 UAVIndex, FUnorderedAccessViewRHIParamRef UAV, uint32 InitialCount) final override
+	FORCEINLINE virtual void RHISetUAVParameter(FComputeShaderRHIParamRef ComputeShader, uint32 UAVIndex, FRHIUnorderedAccessView* UAV, uint32 InitialCount) final override
 	{
 		ContextRedirect(RHISetUAVParameter(ComputeShader, UAVIndex, UAV, InitialCount));
 	}
-	FORCEINLINE virtual void RHISetShaderResourceViewParameter(FComputeShaderRHIParamRef ComputeShader, uint32 SamplerIndex, FShaderResourceViewRHIParamRef SRV) final override
+	FORCEINLINE virtual void RHISetShaderResourceViewParameter(FComputeShaderRHIParamRef ComputeShader, uint32 SamplerIndex, FRHIShaderResourceView* SRV) final override
 	{
 		ContextRedirect(RHISetShaderResourceViewParameter(ComputeShader, SamplerIndex, SRV));
 	}
@@ -524,7 +524,7 @@ public:
 	{
 		ContextRedirect(RHISetMultipleViewports(Count, Data));
 	}
-	FORCEINLINE virtual void RHIClearTinyUAV(FUnorderedAccessViewRHIParamRef UnorderedAccessViewRHI, const uint32* Values) final override
+	FORCEINLINE virtual void RHIClearTinyUAV(FRHIUnorderedAccessView* UnorderedAccessViewRHI, const uint32* Values) final override
 	{
 		ContextRedirect(RHIClearTinyUAV(UnorderedAccessViewRHI, Values));
 	}
@@ -608,23 +608,23 @@ public:
 	{
 		ContextRedirect(RHISetShaderSampler(PixelShader, SamplerIndex, NewState));
 	}
-	FORCEINLINE virtual void RHISetShaderResourceViewParameter(FPixelShaderRHIParamRef PixelShader, uint32 SamplerIndex, FShaderResourceViewRHIParamRef SRV) final override
+	FORCEINLINE virtual void RHISetShaderResourceViewParameter(FPixelShaderRHIParamRef PixelShader, uint32 SamplerIndex, FRHIShaderResourceView* SRV) final override
 	{
 		ContextRedirect(RHISetShaderResourceViewParameter(PixelShader, SamplerIndex, SRV));
 	}
-	FORCEINLINE virtual void RHISetShaderResourceViewParameter(FVertexShaderRHIParamRef VertexShader, uint32 SamplerIndex, FShaderResourceViewRHIParamRef SRV) final override
+	FORCEINLINE virtual void RHISetShaderResourceViewParameter(FVertexShaderRHIParamRef VertexShader, uint32 SamplerIndex, FRHIShaderResourceView* SRV) final override
 	{
 		ContextRedirect(RHISetShaderResourceViewParameter(VertexShader, SamplerIndex, SRV));
 	}
-	FORCEINLINE virtual void RHISetShaderResourceViewParameter(FHullShaderRHIParamRef HullShader, uint32 SamplerIndex, FShaderResourceViewRHIParamRef SRV) final override
+	FORCEINLINE virtual void RHISetShaderResourceViewParameter(FHullShaderRHIParamRef HullShader, uint32 SamplerIndex, FRHIShaderResourceView* SRV) final override
 	{
 		ContextRedirect(RHISetShaderResourceViewParameter(HullShader, SamplerIndex, SRV));
 	}
-	FORCEINLINE virtual void RHISetShaderResourceViewParameter(FDomainShaderRHIParamRef DomainShader, uint32 SamplerIndex, FShaderResourceViewRHIParamRef SRV) final override
+	FORCEINLINE virtual void RHISetShaderResourceViewParameter(FDomainShaderRHIParamRef DomainShader, uint32 SamplerIndex, FRHIShaderResourceView* SRV) final override
 	{
 		ContextRedirect(RHISetShaderResourceViewParameter(DomainShader, SamplerIndex, SRV));
 	}
-	FORCEINLINE virtual void RHISetShaderResourceViewParameter(FGeometryShaderRHIParamRef GeometryShader, uint32 SamplerIndex, FShaderResourceViewRHIParamRef SRV) final override
+	FORCEINLINE virtual void RHISetShaderResourceViewParameter(FGeometryShaderRHIParamRef GeometryShader, uint32 SamplerIndex, FRHIShaderResourceView* SRV) final override
 	{
 		ContextRedirect(RHISetShaderResourceViewParameter(GeometryShader, SamplerIndex, SRV));
 	}
@@ -676,7 +676,7 @@ public:
 	{
 		ContextRedirect(RHISetBlendFactor(BlendFactor));
 	}
-	FORCEINLINE virtual void RHISetRenderTargets(uint32 NumSimultaneousRenderTargets, const FRHIRenderTargetView* NewRenderTargets, const FRHIDepthRenderTargetView* NewDepthStencilTarget, uint32 NumUAVs, const FUnorderedAccessViewRHIParamRef* UAVs) final override
+	FORCEINLINE virtual void RHISetRenderTargets(uint32 NumSimultaneousRenderTargets, const FRHIRenderTargetView* NewRenderTargets, const FRHIDepthRenderTargetView* NewDepthStencilTarget, uint32 NumUAVs, FRHIUnorderedAccessView* const* UAVs) final override
 	{
 		ContextRedirect(RHISetRenderTargets(NumSimultaneousRenderTargets, NewRenderTargets, NewDepthStencilTarget, NumUAVs, UAVs));
 	}
@@ -776,16 +776,16 @@ public:
 	}
 
 	virtual void RHIRayTraceOcclusion(FRayTracingSceneRHIParamRef Scene,
-		FShaderResourceViewRHIParamRef Rays,
-		FUnorderedAccessViewRHIParamRef Output,
+		FRHIShaderResourceView* Rays,
+		FRHIUnorderedAccessView* Output,
 		uint32 NumRays) final override
 	{
 		ContextRedirect(RHIRayTraceOcclusion(Scene, Rays, Output, NumRays));
 	}
 
 	virtual void RHIRayTraceIntersection(FRayTracingSceneRHIParamRef Scene,
-		FShaderResourceViewRHIParamRef Rays,
-		FUnorderedAccessViewRHIParamRef Output,
+		FRHIShaderResourceView* Rays,
+		FRHIUnorderedAccessView* Output,
 		uint32 NumRays) final override
 	{
 		ContextRedirect(RHIRayTraceIntersection(Scene, Rays, Output, NumRays));
diff --git a/Engine/Source/Runtime/D3D12RHI/Private/D3D12Commands.cpp b/Engine/Source/Runtime/D3D12RHI/Private/D3D12Commands.cpp
--- a/Engine/Source/Runtime/D3D12RHI/Private/D3D12Commands.cpp
+++ b/Engine/Source/Runtime/D3D12RHI/Private/D3D12Commands.cpp
@@ -278,7 +278,7 @@ void FD3D12CommandContext::RHITransitionResources(EResourceTransitionAccess Tran
 }
 
 
-void FD3D12CommandContext::RHITransitionResources(EResourceTransitionAccess TransitionType, EResourceTransitionPipeline TransitionPipeline, FUnorderedAccessViewRHIParamRef* InUAVs, int32 InNumUAVs, FRHIComputeFence* WriteComputeFenceRHI)
+void FD3D12CommandContext::RHITransitionResources(EResourceTransitionAccess TransitionType, EResourceTransitionPipeline TransitionPipeline, FRHIUnorderedAccessView** InUAVs, int32 InNumUAVs, FRHIComputeFence* WriteComputeFenceRHI)
 {
 	static IConsoleVariable* CVarShowTransitions = IConsoleManager::Get().FindConsoleVariable(TEXT("r.ProfileGPU.ShowTransitions"));
 	const bool bShowTransitionEvents = CVarShowTransitions->GetInt() != 0;
@@ -629,7 +629,7 @@ void FD3D12CommandContext::RHISetShaderTexture(FComputeShaderRHIParamRef Compute
 	StateCache.SetShaderResourceView<SF_Compute>(NewTexture ? NewTexture->GetShaderResourceView() : nullptr, TextureIndex);
 }
 
-void FD3D12CommandContext::RHISetUAVParameter(FComputeShaderRHIParamRef ComputeShaderRHI, uint32 UAVIndex, FUnorderedAccessViewRHIParamRef UAVRHI)
+void FD3D12CommandContext::RHISetUAVParameter(FComputeShaderRHIParamRef ComputeShaderRHI, uint32 UAVIndex, FRHIUnorderedAccessView* UAVRHI)
 {
 	//VALIDATE_BOUND_SHADER(ComputeShaderRHI);
 
@@ -646,7 +646,7 @@ void FD3D12CommandContext::RHISetUAVParameter(FComputeShaderRHIParamRef ComputeS
 	StateCache.SetUAVs<SF_Compute>(UAVIndex, 1, &UAV, &InitialCount);
 }
 
-void FD3D12CommandContext::RHISetUAVParameter(FComputeShaderRHIParamRef ComputeShaderRHI, uint32 UAVIndex, FUnorderedAccessViewRHIParamRef UAVRHI, uint32 InitialCount)
+void FD3D12CommandContext::RHISetUAVParameter(FComputeShaderRHIParamRef ComputeShaderRHI, uint32 UAVIndex, FRHIUnorderedAccessView* UAVRHI, uint32 InitialCount)
 {
 	//VALIDATE_BOUND_SHADER(ComputeShaderRHI);
 
@@ -660,42 +660,42 @@ void FD3D12CommandContext::RHISetUAVParameter(FComputeShaderRHIParamRef ComputeS
 	StateCache.SetUAVs<SF_Compute>(UAVIndex, 1, &UAV, &InitialCount);
 }
 
-void FD3D12CommandContext::RHISetShaderResourceViewParameter(FPixelShaderRHIParamRef PixelShaderRHI, uint32 TextureIndex, FShaderResourceViewRHIParamRef SRVRHI)
+void FD3D12CommandContext::RHISetShaderResourceViewParameter(FPixelShaderRHIParamRef PixelShaderRHI, uint32 TextureIndex, FRHIShaderResourceView* SRVRHI)
 {
 	VALIDATE_BOUND_SHADER(PixelShaderRHI);
 	FD3D12ShaderResourceView* const SRV = RetrieveObject<FD3D12ShaderResourceView>(SRVRHI);
 	StateCache.SetShaderResourceView<SF_Pixel>(SRV, TextureIndex);
 }
 
-void FD3D12CommandContext::RHISetShaderResourceViewParameter(FVertexShaderRHIParamRef VertexShaderRHI, uint32 TextureIndex, FShaderResourceViewRHIParamRef SRVRHI)
+void FD3D12CommandContext::RHISetShaderResourceViewParameter(FVertexShaderRHIParamRef VertexShaderRHI, uint32 TextureIndex, FRHIShaderResourceView* SRVRHI)
 {
 	VALIDATE_BOUND_SHADER(VertexShaderRHI);
 	FD3D12ShaderResourceView* const SRV = RetrieveObject<FD3D12ShaderResourceView>(SRVRHI);
 	StateCache.SetShaderResourceView<SF_Vertex>(SRV, TextureIndex);
 }
 
-void FD3D12CommandContext::RHISetShaderResourceViewParameter(FComputeShaderRHIParamRef ComputeShaderRHI, uint32 TextureIndex, FShaderResourceViewRHIParamRef SRVRHI)
+void FD3D12CommandContext::RHISetShaderResourceViewParameter(FComputeShaderRHIParamRef ComputeShaderRHI, uint32 TextureIndex, FRHIShaderResourceView* SRVRHI)
 {
 	//VALIDATE_BOUND_SHADER(ComputeShaderRHI);
 	FD3D12ShaderResourceView* const SRV = RetrieveObject<FD3D12ShaderResourceView>(SRVRHI);
 	StateCache.SetShaderResourceView<SF_Compute>(SRV, TextureIndex);
 }
 
-void FD3D12CommandContext::RHISetShaderResourceViewParameter(FHullShaderRHIParamRef HullShaderRHI, uint32 TextureIndex, FShaderResourceViewRHIParamRef SRVRHI)
+void FD3D12CommandContext::RHISetShaderResourceViewParameter(FHullShaderRHIParamRef HullShaderRHI, uint32 TextureIndex, FRHIShaderResourceView* SRVRHI)
 {
 	VALIDATE_BOUND_SHADER(HullShaderRHI);
 	FD3D12ShaderResourceView* const SRV = RetrieveObject<FD3D12ShaderResourceView>(SRVRHI);
 	StateCache.SetShaderResourceView<SF_Hull>(SRV, TextureIndex);
 }
 
-void FD3D12CommandContext::RHISetShaderResourceViewParameter(FDomainShaderRHIParamRef DomainShaderRHI, uint32 TextureIndex, FShaderResourceViewRHIParamRef SRVRHI)
+void FD3D12CommandContext::RHISetShaderResourceViewParameter(FDomainShaderRHIParamRef DomainShaderRHI, uint32 TextureIndex, FRHIShaderResourceView* SRVRHI)
 {
 	VALIDATE_BOUND_SHADER(DomainShaderRHI);
 	FD3D12ShaderResourceView* const SRV = RetrieveObject<FD3D12ShaderResourceView>(SRVRHI);
 	StateCache.SetShaderResourceView<SF_Domain>(SRV, TextureIndex);
 }
 
-void FD3D12CommandContext::RHISetShaderResourceViewParameter(FGeometryShaderRHIParamRef GeometryShaderRHI, uint32 TextureIndex, FShaderResourceViewRHIParamRef SRVRHI)
+void FD3D12CommandContext::RHISetShaderResourceViewParameter(FGeometryShaderRHIParamRef GeometryShaderRHI, uint32 TextureIndex, FRHIShaderResourceView* SRVRHI)
 {
 	VALIDATE_BOUND_SHADER(GeometryShaderRHI);
 	FD3D12ShaderResourceView* const SRV = RetrieveObject<FD3D12ShaderResourceView>(SRVRHI);
@@ -996,7 +996,7 @@ void FD3D12CommandContext::RHISetRenderTargets(
 	const FRHIRenderTargetView* NewRenderTargetsRHI,
 	const FRHIDepthRenderTargetView* NewDepthStencilTargetRHI,
 	uint32 NewNumUAVs,
-	const FUnorderedAccessViewRHIParamRef* UAVs
+	FRHIUnorderedAccessView* const* UAVs
 	)
 {
 	FD3D12TextureBase* NewDepthStencilTarget = RetrieveTextureBase(NewDepthStencilTargetRHI ? NewDepthStencilTargetRHI->Texture : nullptr);
@@ -1106,7 +1106,7 @@ void FD3D12CommandContext::RHISetRenderTargets(
 
 void FD3D12CommandContext::RHISetRenderTargetsAndClear(const FRHISetRenderTargetsInfo& RenderTargetsInfo)
 {
-	FUnorderedAccessViewRHIParamRef UAVs[MaxSimultaneousUAVs] = {};
+	FRHIUnorderedAccessView* UAVs[MaxSimultaneousUAVs] = {};
 	check(RenderTargetsInfo.NumUAVs <= MaxSimultaneousUAVs);
 	for (int32 UAVIndex = 0; UAVIndex < RenderTargetsInfo.NumUAVs; ++UAVIndex)
 	{
diff --git a/Engine/Source/Runtime/D3D12RHI/Private/D3D12RHIPrivate.h b/Engine/Source/Runtime/D3D12RHI/Private/D3D12RHIPrivate.h
--- a/Engine/Source/Runtime/D3D12RHI/Private/D3D12RHIPrivate.h
+++ b/Engine/Source/Runtime/D3D12RHI/Private/D3D12RHIPrivate.h
@@ -299,8 +299,8 @@ public:
 	virtual FShaderResourceViewRHIRef RHICreateShaderResourceView(FStructuredBufferRHIParamRef StructuredBuffer) final override;
 	virtual FShaderResourceViewRHIRef RHICreateShaderResourceView(FVertexBufferRHIParamRef VertexBuffer, uint32 Stride, uint8 Format) final override;
 	virtual FShaderResourceViewRHIRef RHICreateShaderResourceView(FIndexBufferRHIParamRef Buffer) final override;
-	virtual void RHIUpdateShaderResourceView(FShaderResourceViewRHIParamRef SRV, FVertexBufferRHIParamRef VertexBuffer, uint32 Stride, uint8 Format) final override;
-	virtual void RHIUpdateShaderResourceView(FShaderResourceViewRHIParamRef SRV, FIndexBufferRHIParamRef IndexBuffer) final override;
+	virtual void RHIUpdateShaderResourceView(FRHIShaderResourceView* SRV, FVertexBufferRHIParamRef VertexBuffer, uint32 Stride, uint8 Format) final override;
+	virtual void RHIUpdateShaderResourceView(FRHIShaderResourceView* SRV, FIndexBufferRHIParamRef IndexBuffer) final override;
 	virtual uint64 RHICalcTexture2DPlatformSize(uint32 SizeX, uint32 SizeY, uint8 Format, uint32 NumMips, uint32 NumSamples, uint32 Flags, uint32& OutAlign) override;
 	virtual uint64 RHICalcTexture3DPlatformSize(uint32 SizeX, uint32 SizeY, uint32 SizeZ, uint8 Format, uint32 NumMips, uint32 Flags, uint32& OutAlign) final override;
 	virtual uint64 RHICalcTextureCubePlatformSize(uint32 Size, uint8 Format, uint32 NumMips, uint32 Flags, uint32& OutAlign) final override;
diff --git a/Engine/Source/Runtime/D3D12RHI/Private/D3D12RayTracing.cpp b/Engine/Source/Runtime/D3D12RHI/Private/D3D12RayTracing.cpp
--- a/Engine/Source/Runtime/D3D12RHI/Private/D3D12RayTracing.cpp
+++ b/Engine/Source/Runtime/D3D12RHI/Private/D3D12RayTracing.cpp
@@ -2515,10 +2515,10 @@ static void SetRayTracingShaderResources(
 	FD3D12CommandContext& CommandContext,
 	const FD3D12RayTracingShader* Shader,
 	uint32 InNumTextures, const FTextureRHIParamRef* Textures,
-	uint32 InNumSRVs, const FShaderResourceViewRHIParamRef* SRVs,
+	uint32 InNumSRVs, FRHIShaderResourceView* const* SRVs,
 	uint32 InNumUniformBuffers, const FUniformBufferRHIParamRef* UniformBuffers,
 	uint32 InNumSamplers, FRHISamplerState* const* Samplers,
-	uint32 InNumUAVs, const FUnorderedAccessViewRHIParamRef* UAVs,
+	uint32 InNumUAVs, FRHIUnorderedAccessView* const* UAVs,
 	uint32 InLooseParameterDataSize, const void* InLooseParameterData,
 	FD3D12RayTracingDescriptorCache& DescriptorCache,
 	ResourceBinderType& Binder)
@@ -2559,7 +2559,7 @@ static void SetRayTracingShaderResources(
 
 	for (uint32 SRVIndex = 0; SRVIndex < InNumSRVs; ++SRVIndex)
 	{
-		FShaderResourceViewRHIParamRef Resource = SRVs[SRVIndex];
+		FRHIShaderResourceView* Resource = SRVs[SRVIndex];
 		if (Resource)
 		{
 			FD3D12ShaderResourceView* SRV = CommandContext.RetrieveObject<FD3D12ShaderResourceView>(Resource);
@@ -2595,7 +2595,7 @@ static void SetRayTracingShaderResources(
 
 	for (uint32 UAVIndex = 0; UAVIndex < InNumUAVs; ++UAVIndex)
 	{
-		FUnorderedAccessViewRHIParamRef Resource = UAVs[UAVIndex];
+		FRHIUnorderedAccessView* Resource = UAVs[UAVIndex];
 		if (Resource)
 		{
 			FD3D12UnorderedAccessView* UAV = CommandContext.RetrieveObject<FD3D12UnorderedAccessView>(Resource);
@@ -2877,8 +2877,8 @@ static void DispatchRays(FD3D12CommandContext& CommandContext,
 
 
 void FD3D12CommandContext::RHIRayTraceOcclusion(FRayTracingSceneRHIParamRef InScene,
-	FShaderResourceViewRHIParamRef Rays,
-	FUnorderedAccessViewRHIParamRef Output,
+	FRHIShaderResourceView* Rays,
+	FRHIUnorderedAccessView* Output,
 	uint32 NumRays)
 {
 	checkf(GetParentDevice()->GetBasicRayTracingPipeline(), TEXT("Ray tracing support is not initialized for this device. Ensure that InitRayTracing() is called before issuing any ray tracing work."));
@@ -2912,8 +2912,8 @@ void FD3D12CommandContext::RHIRayTraceOcclusion(FRayTracingSceneRHIParamRef InSc
 }
 
 void FD3D12CommandContext::RHIRayTraceIntersection(FRayTracingSceneRHIParamRef InScene,
-	FShaderResourceViewRHIParamRef InRays,
-	FUnorderedAccessViewRHIParamRef InOutput,
+	FRHIShaderResourceView* InRays,
+	FRHIUnorderedAccessView* InOutput,
 	uint32 NumRays)
 {
 	checkf(GetParentDevice()->GetBasicRayTracingPipeline(), TEXT("Ray tracing support is not initialized for this device. Ensure that InitRayTracing() is called before issuing any ray tracing work."));
diff --git a/Engine/Source/Runtime/D3D12RHI/Private/D3D12UAV.cpp b/Engine/Source/Runtime/D3D12RHI/Private/D3D12UAV.cpp
--- a/Engine/Source/Runtime/D3D12RHI/Private/D3D12UAV.cpp
+++ b/Engine/Source/Runtime/D3D12RHI/Private/D3D12UAV.cpp
@@ -202,7 +202,7 @@ FUnorderedAccessViewRHIRef FD3D12DynamicRHI::RHICreateUnorderedAccessView(FIndex
 	return CreateUAV(UAVDesc, IndexBuffer, false);
 }
 
-void FD3D12CommandContext::RHIClearTinyUAV(FUnorderedAccessViewRHIParamRef UnorderedAccessViewRHI, const uint32* Values)
+void FD3D12CommandContext::RHIClearTinyUAV(FRHIUnorderedAccessView* UnorderedAccessViewRHI, const uint32* Values)
 {
 	FD3D12UnorderedAccessView* UnorderedAccessView = RetrieveObject<FD3D12UnorderedAccessView>(UnorderedAccessViewRHI);
 
diff --git a/Engine/Source/Runtime/D3D12RHI/Private/D3D12View.cpp b/Engine/Source/Runtime/D3D12RHI/Private/D3D12View.cpp
--- a/Engine/Source/Runtime/D3D12RHI/Private/D3D12View.cpp
+++ b/Engine/Source/Runtime/D3D12RHI/Private/D3D12View.cpp
@@ -318,7 +318,7 @@ FShaderResourceViewRHIRef FD3D12DynamicRHI::RHICreateShaderResourceView(FIndexBu
 	});
 }
 
-void FD3D12DynamicRHI::RHIUpdateShaderResourceView(FShaderResourceViewRHIParamRef SRV, FVertexBufferRHIParamRef VertexBuffer, uint32 Stride, uint8 Format)
+void FD3D12DynamicRHI::RHIUpdateShaderResourceView(FRHIShaderResourceView* SRV, FVertexBufferRHIParamRef VertexBuffer, uint32 Stride, uint8 Format)
 {
 	check(SRV);
 	if (VertexBuffer)
@@ -343,7 +343,7 @@ void FD3D12DynamicRHI::RHIUpdateShaderResourceView(FShaderResourceViewRHIParamRe
 	}
 }
 
-void FD3D12DynamicRHI::RHIUpdateShaderResourceView(FShaderResourceViewRHIParamRef SRV, FIndexBufferRHIParamRef IndexBuffer)
+void FD3D12DynamicRHI::RHIUpdateShaderResourceView(FRHIShaderResourceView* SRV, FIndexBufferRHIParamRef IndexBuffer)
 {
 	check(SRV);
 	if (IndexBuffer)
diff --git a/Engine/Source/Runtime/EmptyRHI/Private/EmptyCommands.cpp b/Engine/Source/Runtime/EmptyRHI/Private/EmptyCommands.cpp
--- a/Engine/Source/Runtime/EmptyRHI/Private/EmptyCommands.cpp
+++ b/Engine/Source/Runtime/EmptyRHI/Private/EmptyCommands.cpp
@@ -62,13 +62,13 @@ void FEmptyDynamicRHI::RHISetBoundShaderState(FRHIBoundShaderState* BoundShaderS
 }
 
 
-void FEmptyDynamicRHI::RHISetUAVParameter(FComputeShaderRHIParamRef ComputeShaderRHI, uint32 UAVIndex, FUnorderedAccessViewRHIParamRef UAVRHI)
+void FEmptyDynamicRHI::RHISetUAVParameter(FComputeShaderRHIParamRef ComputeShaderRHI, uint32 UAVIndex, FRHIUnorderedAccessView* UAVRHI)
 {
 	FEmptyUnorderedAccessView* UAV = ResourceCast(UAVRHI);
 
 }
 
-void FEmptyDynamicRHI::RHISetUAVParameter(FComputeShaderRHIParamRef ComputeShaderRHI,uint32 UAVIndex,FUnorderedAccessViewRHIParamRef UAVRHI, uint32 InitialCount)
+void FEmptyDynamicRHI::RHISetUAVParameter(FComputeShaderRHIParamRef ComputeShaderRHI,uint32 UAVIndex, FRHIUnorderedAccessView* UAVRHI, uint32 InitialCount)
 {
 	FEmptyUnorderedAccessView* UAV = ResourceCast(UAVRHI);
 
@@ -106,32 +106,32 @@ void FEmptyDynamicRHI::RHISetShaderTexture(FComputeShaderRHIParamRef ComputeShad
 }
 
 
-void FEmptyDynamicRHI::RHISetShaderResourceViewParameter(FVertexShaderRHIParamRef VertexShaderRHI, uint32 TextureIndex, FShaderResourceViewRHIParamRef SRVRHI)
+void FEmptyDynamicRHI::RHISetShaderResourceViewParameter(FVertexShaderRHIParamRef VertexShaderRHI, uint32 TextureIndex, FRHIShaderResourceView* SRVRHI)
 {
 
 }
 
-void FEmptyDynamicRHI::RHISetShaderResourceViewParameter(FHullShaderRHIParamRef HullShaderRHI,uint32 TextureIndex,FShaderResourceViewRHIParamRef SRVRHI)
+void FEmptyDynamicRHI::RHISetShaderResourceViewParameter(FHullShaderRHIParamRef HullShaderRHI,uint32 TextureIndex, FRHIShaderResourceView* SRVRHI)
 {
 
 }
 
-void FEmptyDynamicRHI::RHISetShaderResourceViewParameter(FDomainShaderRHIParamRef DomainShaderRHI,uint32 TextureIndex,FShaderResourceViewRHIParamRef SRVRHI)
+void FEmptyDynamicRHI::RHISetShaderResourceViewParameter(FDomainShaderRHIParamRef DomainShaderRHI,uint32 TextureIndex, FRHIShaderResourceView* SRVRHI)
 {
 
 }
 
-void FEmptyDynamicRHI::RHISetShaderResourceViewParameter(FGeometryShaderRHIParamRef GeometryShaderRHI,uint32 TextureIndex,FShaderResourceViewRHIParamRef SRVRHI)
+void FEmptyDynamicRHI::RHISetShaderResourceViewParameter(FGeometryShaderRHIParamRef GeometryShaderRHI,uint32 TextureIndex, FRHIShaderResourceView* SRVRHI)
 {
 
 }
 
-void FEmptyDynamicRHI::RHISetShaderResourceViewParameter(FPixelShaderRHIParamRef PixelShaderRHI,uint32 TextureIndex,FShaderResourceViewRHIParamRef SRVRHI)
+void FEmptyDynamicRHI::RHISetShaderResourceViewParameter(FPixelShaderRHIParamRef PixelShaderRHI,uint32 TextureIndex, FRHIShaderResourceView* SRVRHI)
 {
 
 }
 
-void FEmptyDynamicRHI::RHISetShaderResourceViewParameter(FComputeShaderRHIParamRef ComputeShaderRHI,uint32 TextureIndex,FShaderResourceViewRHIParamRef SRVRHI)
+void FEmptyDynamicRHI::RHISetShaderResourceViewParameter(FComputeShaderRHIParamRef ComputeShaderRHI,uint32 TextureIndex, FRHIShaderResourceView* SRVRHI)
 {
 
 }
@@ -252,7 +252,7 @@ void FEmptyDynamicRHI::RHISetBlendState(FRHIBlendState* NewStateRHI, const FLine
 
 
 void FEmptyDynamicRHI::RHISetRenderTargets(uint32 NumSimultaneousRenderTargets, const FRHIRenderTargetView* NewRenderTargets, 
-	FTextureRHIParamRef NewDepthStencilTargetRHI, uint32 NumUAVs, const FUnorderedAccessViewRHIParamRef* UAVs)
+	FTextureRHIParamRef NewDepthStencilTargetRHI, uint32 NumUAVs, FRHIUnorderedAccessView* const* UAVs)
 {
 
 }
diff --git a/Engine/Source/Runtime/EmptyRHI/Private/EmptyUAV.cpp b/Engine/Source/Runtime/EmptyRHI/Private/EmptyUAV.cpp
--- a/Engine/Source/Runtime/EmptyRHI/Private/EmptyUAV.cpp
+++ b/Engine/Source/Runtime/EmptyRHI/Private/EmptyUAV.cpp
@@ -113,7 +113,7 @@ FShaderResourceViewRHIRef FEmptyDynamicRHI::RHICreateShaderResourceView(FTexture
 	return SRV;
 }
 
-void FEmptyDynamicRHI::RHIClearTinyUAV(FUnorderedAccessViewRHIParamRef UnorderedAccessViewRHI, const uint32* Values)
+void FEmptyDynamicRHI::RHIClearTinyUAV(FRHIUnorderedAccessView* UnorderedAccessViewRHI, const uint32* Values)
 {
 	FEmptyUnorderedAccessView* UnorderedAccessView = ResourceCast(UnorderedAccessViewRHI);
 
diff --git a/Engine/Source/Runtime/EmptyRHI/Public/EmptyRHI.h b/Engine/Source/Runtime/EmptyRHI/Public/EmptyRHI.h
--- a/Engine/Source/Runtime/EmptyRHI/Public/EmptyRHI.h
+++ b/Engine/Source/Runtime/EmptyRHI/Public/EmptyRHI.h
@@ -138,7 +138,7 @@ public:
 	virtual void RHIAutomaticCacheFlushAfterComputeShader(bool bEnable) final override;
 	virtual void RHIFlushComputeShaderCache() final override;
 	virtual void RHISetMultipleViewports(uint32 Count, const FViewportBounds* Data) final override;
-	virtual void RHIClearTinyUAV(FUnorderedAccessViewRHIParamRef UnorderedAccessViewRHI, const uint32* Values) final override;
+	virtual void RHIClearTinyUAV(FRHIUnorderedAccessView* UnorderedAccessViewRHI, const uint32* Values) final override;
 	virtual void RHICopyToResolveTarget(FTextureRHIParamRef SourceTexture, FTextureRHIParamRef DestTexture, const FResolveParams& ResolveParams) final override;
 	virtual void RHICopyTexture(FTextureRHIParamRef SourceTexture, FTextureRHIParamRef DestTexture, const FRHICopyTextureInfo& CopyInfo) final override;
 	virtual void RHIBeginRenderQuery(FRenderQueryRHIParamRef RenderQuery) final override;
@@ -167,14 +167,14 @@ public:
 	virtual void RHISetShaderSampler(FDomainShaderRHIParamRef DomainShader, uint32 SamplerIndex, FRHISamplerState* NewState) final override;
 	virtual void RHISetShaderSampler(FHullShaderRHIParamRef HullShader, uint32 SamplerIndex, FRHISamplerState* NewState) final override;
 	virtual void RHISetShaderSampler(FPixelShaderRHIParamRef PixelShader, uint32 SamplerIndex, FRHISamplerState* NewState) final override;
-	virtual void RHISetUAVParameter(FComputeShaderRHIParamRef ComputeShader, uint32 UAVIndex, FUnorderedAccessViewRHIParamRef UAV) final override;
-	virtual void RHISetUAVParameter(FComputeShaderRHIParamRef ComputeShader, uint32 UAVIndex, FUnorderedAccessViewRHIParamRef UAV, uint32 InitialCount) final override;
-	virtual void RHISetShaderResourceViewParameter(FPixelShaderRHIParamRef PixelShader, uint32 SamplerIndex, FShaderResourceViewRHIParamRef SRV) final override;
-	virtual void RHISetShaderResourceViewParameter(FVertexShaderRHIParamRef VertexShader, uint32 SamplerIndex, FShaderResourceViewRHIParamRef SRV) final override;
-	virtual void RHISetShaderResourceViewParameter(FComputeShaderRHIParamRef ComputeShader, uint32 SamplerIndex, FShaderResourceViewRHIParamRef SRV) final override;
-	virtual void RHISetShaderResourceViewParameter(FHullShaderRHIParamRef HullShader, uint32 SamplerIndex, FShaderResourceViewRHIParamRef SRV) final override;
-	virtual void RHISetShaderResourceViewParameter(FDomainShaderRHIParamRef DomainShader, uint32 SamplerIndex, FShaderResourceViewRHIParamRef SRV) final override;
-	virtual void RHISetShaderResourceViewParameter(FGeometryShaderRHIParamRef GeometryShader, uint32 SamplerIndex, FShaderResourceViewRHIParamRef SRV) final override;
+	virtual void RHISetUAVParameter(FComputeShaderRHIParamRef ComputeShader, uint32 UAVIndex, FRHIUnorderedAccessView* UAV) final override;
+	virtual void RHISetUAVParameter(FComputeShaderRHIParamRef ComputeShader, uint32 UAVIndex, FRHIUnorderedAccessView* UAV, uint32 InitialCount) final override;
+	virtual void RHISetShaderResourceViewParameter(FPixelShaderRHIParamRef PixelShader, uint32 SamplerIndex, FRHIShaderResourceView* SRV) final override;
+	virtual void RHISetShaderResourceViewParameter(FVertexShaderRHIParamRef VertexShader, uint32 SamplerIndex, FRHIShaderResourceView* SRV) final override;
+	virtual void RHISetShaderResourceViewParameter(FComputeShaderRHIParamRef ComputeShader, uint32 SamplerIndex, FRHIShaderResourceView* SRV) final override;
+	virtual void RHISetShaderResourceViewParameter(FHullShaderRHIParamRef HullShader, uint32 SamplerIndex, FRHIShaderResourceView* SRV) final override;
+	virtual void RHISetShaderResourceViewParameter(FDomainShaderRHIParamRef DomainShader, uint32 SamplerIndex, FRHIShaderResourceView* SRV) final override;
+	virtual void RHISetShaderResourceViewParameter(FGeometryShaderRHIParamRef GeometryShader, uint32 SamplerIndex, FRHIShaderResourceView* SRV) final override;
 	virtual void RHISetShaderUniformBuffer(FVertexShaderRHIParamRef VertexShader, uint32 BufferIndex, FUniformBufferRHIParamRef Buffer) final override;
 	virtual void RHISetShaderUniformBuffer(FHullShaderRHIParamRef HullShader, uint32 BufferIndex, FUniformBufferRHIParamRef Buffer) final override;
 	virtual void RHISetShaderUniformBuffer(FDomainShaderRHIParamRef DomainShader, uint32 BufferIndex, FUniformBufferRHIParamRef Buffer) final override;
@@ -189,7 +189,7 @@ public:
 	virtual void RHISetShaderParameter(FComputeShaderRHIParamRef ComputeShader, uint32 BufferIndex, uint32 BaseIndex, uint32 NumBytes, const void* NewValue) final override;
 	virtual void RHISetDepthStencilState(FRHIDepthStencilState* NewState, uint32 StencilRef) final override;
 	virtual void RHISetBlendState(FRHIBlendState* NewState, const FLinearColor& BlendFactor) final override;
-	virtual void RHISetRenderTargets(uint32 NumSimultaneousRenderTargets, const FRHIRenderTargetView* NewRenderTargets, const FRHIDepthRenderTargetView* NewDepthStencilTarget, uint32 NumUAVs, const FUnorderedAccessViewRHIParamRef* UAVs) final override;
+	virtual void RHISetRenderTargets(uint32 NumSimultaneousRenderTargets, const FRHIRenderTargetView* NewRenderTargets, const FRHIDepthRenderTargetView* NewDepthStencilTarget, uint32 NumUAVs, FRHIUnorderedAccessView* const* UAVs) final override;
 	virtual void RHISetRenderTargetsAndClear(const FRHISetRenderTargetsInfo& RenderTargetsInfo) final override;
 	virtual void RHIDrawPrimitive(uint32 PrimitiveType, uint32 BaseVertexIndex, uint32 NumPrimitives, uint32 NumInstances) final override;
 	virtual void RHIDrawPrimitiveIndirect(uint32 PrimitiveType, FVertexBufferRHIParamRef ArgumentBuffer, uint32 ArgumentOffset) final override;
diff --git a/Engine/Source/Runtime/Engine/Classes/Particles/ParticleModuleRequired.h b/Engine/Source/Runtime/Engine/Classes/Particles/ParticleModuleRequired.h
--- a/Engine/Source/Runtime/Engine/Classes/Particles/ParticleModuleRequired.h
+++ b/Engine/Source/Runtime/Engine/Classes/Particles/ParticleModuleRequired.h
@@ -88,7 +88,7 @@ struct FParticleRequiredModule
 	uint32 NumBoundingTriangles;
 	float AlphaThreshold;
 	TArray<FVector2D> FrameData;
-	FShaderResourceViewRHIParamRef BoundingGeometryBufferSRV;
+	FRHIShaderResourceView* BoundingGeometryBufferSRV;
 };
 
 
@@ -424,7 +424,7 @@ class UParticleModuleRequired : public UParticleModule
 		return CutoutTexture != nullptr;
 	}
 
-	inline FShaderResourceViewRHIParamRef GetBoundingGeometrySRV() const
+	inline FRHIShaderResourceView* GetBoundingGeometrySRV() const
 	{
 		return BoundingGeometryBuffer->ShaderResourceView;
 	}
diff --git a/Engine/Source/Runtime/Engine/Classes/Particles/SubUVAnimation.h b/Engine/Source/Runtime/Engine/Classes/Particles/SubUVAnimation.h
--- a/Engine/Source/Runtime/Engine/Classes/Particles/SubUVAnimation.h
+++ b/Engine/Source/Runtime/Engine/Classes/Particles/SubUVAnimation.h
@@ -193,7 +193,7 @@ public:
 		return &DerivedData.BoundingGeometry[FrameIndex * GetNumBoundingVertices()];
 	}
 
-	inline FShaderResourceViewRHIParamRef GetBoundingGeometrySRV() const
+	inline FRHIShaderResourceView* GetBoundingGeometrySRV() const
 	{
 		return BoundingGeometryBuffer->ShaderResourceView;
 	}
diff --git a/Engine/Source/Runtime/Engine/Private/GPUSkinCache.cpp b/Engine/Source/Runtime/Engine/Private/GPUSkinCache.cpp
--- a/Engine/Source/Runtime/Engine/Private/GPUSkinCache.cpp
+++ b/Engine/Source/Runtime/Engine/Private/GPUSkinCache.cpp
@@ -175,7 +175,7 @@ public:
 		FGPUSkinPassthroughVertexFactory* TargetVertexFactory = nullptr;
 
 		// triangle index buffer (input for the RecomputeSkinTangents, might need special index buffer unique to position and normal, not considering UV/vertex color)
-		FShaderResourceViewRHIParamRef IndexBuffer = nullptr;
+		FRHIShaderResourceView* IndexBuffer = nullptr;
 
 		const FSkelMeshRenderSection* Section = nullptr;
 
@@ -372,7 +372,7 @@ protected:
 	FSkeletalMeshObjectGPUSkin* GPUSkin;
 	uint32 InputWeightStride;
 	FShaderResourceViewRHIRef InputWeightStreamSRV;
-	FShaderResourceViewRHIParamRef MorphBuffer;
+	FRHIShaderResourceView* MorphBuffer;
 	FShaderResourceViewRHIRef ClothBuffer;
 	FShaderResourceViewRHIRef ClothPositionsAndNormalsBuffer;
 	int32 LOD;
@@ -422,7 +422,7 @@ public:
 	void SetParameters(FRHICommandListImmediate& RHICmdList, const FVertexBufferAndSRV& BoneBuffer,
 		FGPUSkinCacheEntry* Entry,
 		FGPUSkinCacheEntry::FSectionDispatchData& DispatchData,
-		FUnorderedAccessViewRHIParamRef PositionUAV, FUnorderedAccessViewRHIParamRef TangentUAV)
+		FRHIUnorderedAccessView* PositionUAV, FRHIUnorderedAccessView* TangentUAV)
 	{
 		FComputeShaderRHIParamRef ShaderRHI = GetComputeShader();
 
@@ -631,7 +631,7 @@ void FGPUSkinCache::CommitRayTracingGeometryUpdates(FRHICommandList& RHICmdList)
 void FGPUSkinCache::TransitionToWriteable(FRHICommandList& RHICmdList)
 {
 	int32 BufferIndex = InternalUpdateCount % GPUSKINCACHE_FRAMES;
-	FUnorderedAccessViewRHIParamRef OutUAVs[] = { SkinCacheBuffer[BufferIndex].UAV };
+	FRHIUnorderedAccessView* OutUAVs[] = { SkinCacheBuffer[BufferIndex].UAV };
 	RHICmdList.TransitionResources(EResourceTransitionAccess::ERWNoBarrier, EResourceTransitionPipeline::EGfxToCompute, OutUAVs, ARRAY_COUNT(OutUAVs));
 }
 
@@ -639,7 +639,7 @@ void FGPUSkinCache::TransitionAllToWriteable(FRHICommandList& RHICmdList)
 {
 	if (bInitialized)
 	{
-		FUnorderedAccessViewRHIParamRef OutUAVs[GPUSKINCACHE_FRAMES];
+		FRHIUnorderedAccessView* OutUAVs[GPUSKINCACHE_FRAMES];
 
 		for (int32 Index = 0; Index < GPUSKINCACHE_FRAMES; ++Index)
 		{
@@ -1311,7 +1311,7 @@ void FGPUSkinCache::DispatchUpdateSkinning(FRHICommandListImmediate& RHICmdList,
 	check(DispatchData.PreviousPositionBuffer != DispatchData.PositionBuffer);
 }
 
-void FGPUSkinCache::FRWBuffersAllocation::RemoveAllFromTransitionArray(TArray<FUnorderedAccessViewRHIParamRef>& InBuffersToTransition)
+void FGPUSkinCache::FRWBuffersAllocation::RemoveAllFromTransitionArray(TArray<FRHIUnorderedAccessView*>& InBuffersToTransition)
 {
 	for (uint32 i = 0; i < NUM_BUFFERS; i++)
 	{
diff --git a/Engine/Source/Runtime/Engine/Private/GPUSkinVertexFactory.cpp b/Engine/Source/Runtime/Engine/Private/GPUSkinVertexFactory.cpp
--- a/Engine/Source/Runtime/Engine/Private/GPUSkinVertexFactory.cpp
+++ b/Engine/Source/Runtime/Engine/Private/GPUSkinVertexFactory.cpp
@@ -542,7 +542,7 @@ public:
 		{
 			if (BoneMatrices.IsBound())
 			{
-				FShaderResourceViewRHIParamRef CurrentData = ShaderData.GetBoneBufferForReading(false).VertexBufferSRV;
+				FRHIShaderResourceView* CurrentData = ShaderData.GetBoneBufferForReading(false).VertexBufferSRV;
 				ShaderBindings.Add(BoneMatrices, CurrentData);
 			}
 
@@ -551,7 +551,7 @@ public:
 				// todo: Maybe a check for PreviousData!=CurrentData would save some performance (when objects don't have velocty yet) but removing the bool also might save performance
 				bLocalPerBoneMotionBlur = true;
 
-				FShaderResourceViewRHIParamRef PreviousData = ShaderData.GetBoneBufferForReading(true).VertexBufferSRV;
+				FRHIShaderResourceView* PreviousData = ShaderData.GetBoneBufferForReading(true).VertexBufferSRV;
 				ShaderBindings.Add(PreviousBoneMatrices, PreviousData);
 			}
 		}
diff --git a/Engine/Source/Runtime/Engine/Private/GPUSort.cpp b/Engine/Source/Runtime/Engine/Private/GPUSort.cpp
--- a/Engine/Source/Runtime/Engine/Private/GPUSort.cpp
+++ b/Engine/Source/Runtime/Engine/Private/GPUSort.cpp
@@ -267,7 +267,7 @@ public:
 	/**
 	 * Set output buffer for this shader.
 	 */
-	void SetOutput(FRHICommandList& RHICmdList, FUnorderedAccessViewRHIParamRef OutOffsetsUAV)
+	void SetOutput(FRHICommandList& RHICmdList, FRHIUnorderedAccessView* OutOffsetsUAV)
 	{
 		FComputeShaderRHIParamRef ComputeShaderRHI = GetComputeShader();
 		if ( OutOffsets.IsBound() )
@@ -284,7 +284,7 @@ public:
 		FComputeShaderRHIParamRef ComputeShaderRHI = GetComputeShader();
 		if ( OutOffsets.IsBound() )
 		{
-			RHICmdList.SetUAVParameter(ComputeShaderRHI, OutOffsets.GetBaseIndex(), FUnorderedAccessViewRHIParamRef());
+			RHICmdList.SetUAVParameter(ComputeShaderRHI, OutOffsets.GetBaseIndex(), nullptr);
 		}
 	}
 
@@ -355,7 +355,7 @@ public:
 	/**
 	 * Set parameters for this shader.
 	 */
-	void SetParameters(FRHICommandList& RHICmdList, FShaderResourceViewRHIParamRef InKeysSRV, FRadixSortUniformBufferRef& RadixSortUniformBuffer, FShaderResourceViewRHIParamRef RadixSortParameterBufferSRV)
+	void SetParameters(FRHICommandList& RHICmdList, FRHIShaderResourceView* InKeysSRV, FRadixSortUniformBufferRef& RadixSortUniformBuffer, FRHIShaderResourceView* RadixSortParameterBufferSRV)
 	{
 		FComputeShaderRHIParamRef ComputeShaderRHI = GetComputeShader();
 		SetUniformBufferParameter(RHICmdList, ComputeShaderRHI, GetUniformBufferParameter<FRadixSortParameters>(), RadixSortUniformBuffer );
@@ -372,7 +372,7 @@ public:
 	/**
 	 * Set output buffer for this shader.
 	 */
-	void SetOutput(FRHICommandList& RHICmdList, FUnorderedAccessViewRHIParamRef OutOffsetsUAV)
+	void SetOutput(FRHICommandList& RHICmdList, FRHIUnorderedAccessView* OutOffsetsUAV)
 	{
 		FComputeShaderRHIParamRef ComputeShaderRHI = GetComputeShader();
 		if ( OutOffsets.IsBound() )
@@ -389,15 +389,15 @@ public:
 		FComputeShaderRHIParamRef ComputeShaderRHI = GetComputeShader();
 		if ( RadixSortParameterBuffer.IsBound() )
 		{
-			RHICmdList.SetShaderResourceViewParameter(ComputeShaderRHI, RadixSortParameterBuffer.GetBaseIndex(), FShaderResourceViewRHIParamRef());
+			RHICmdList.SetShaderResourceViewParameter(ComputeShaderRHI, RadixSortParameterBuffer.GetBaseIndex(), nullptr);
 		}
 		if ( InKeys.IsBound() )
 		{
-			RHICmdList.SetShaderResourceViewParameter(ComputeShaderRHI, InKeys.GetBaseIndex(), FShaderResourceViewRHIParamRef());
+			RHICmdList.SetShaderResourceViewParameter(ComputeShaderRHI, InKeys.GetBaseIndex(), nullptr);
 		}
 		if ( OutOffsets.IsBound() )
 		{
-			RHICmdList.SetUAVParameter(ComputeShaderRHI, OutOffsets.GetBaseIndex(), FUnorderedAccessViewRHIParamRef());
+			RHICmdList.SetUAVParameter(ComputeShaderRHI, OutOffsets.GetBaseIndex(), nullptr);
 		}
 	}
 
@@ -461,7 +461,7 @@ public:
 	/**
 	 * Set parameters for this shader.
 	 */
-	void SetParameters(FRHICommandList& RHICmdList, FShaderResourceViewRHIParamRef InOffsetsSRV)
+	void SetParameters(FRHICommandList& RHICmdList, FRHIShaderResourceView* InOffsetsSRV)
 	{
 		FComputeShaderRHIParamRef ComputeShaderRHI = GetComputeShader();
 		if ( InOffsets.IsBound() )
@@ -473,7 +473,7 @@ public:
 	/**
 	 * Set output buffer for this shader.
 	 */
-	void SetOutput(FRHICommandList& RHICmdList, FUnorderedAccessViewRHIParamRef OutOffsetsUAV)
+	void SetOutput(FRHICommandList& RHICmdList, FRHIUnorderedAccessView* OutOffsetsUAV)
 	{
 		FComputeShaderRHIParamRef ComputeShaderRHI = GetComputeShader();
 		if ( OutOffsets.IsBound() )
@@ -490,11 +490,11 @@ public:
 		FComputeShaderRHIParamRef ComputeShaderRHI = GetComputeShader();
 		if ( InOffsets.IsBound() )
 		{
-			RHICmdList.SetShaderResourceViewParameter(ComputeShaderRHI, InOffsets.GetBaseIndex(), FShaderResourceViewRHIParamRef());
+			RHICmdList.SetShaderResourceViewParameter(ComputeShaderRHI, InOffsets.GetBaseIndex(), nullptr);
 		}
 		if ( OutOffsets.IsBound() )
 		{
-			RHICmdList.SetUAVParameter(ComputeShaderRHI, OutOffsets.GetBaseIndex(), FUnorderedAccessViewRHIParamRef());
+			RHICmdList.SetUAVParameter(ComputeShaderRHI, OutOffsets.GetBaseIndex(), nullptr);
 		}
 	}
 
@@ -577,11 +577,11 @@ public:
 	 */
 	void SetParameters(
 		FRHICommandList& RHICmdList,
-		FShaderResourceViewRHIParamRef InKeysSRV,
-		FShaderResourceViewRHIParamRef InValuesSRV,
-		FShaderResourceViewRHIParamRef InOffsetsSRV,
+		FRHIShaderResourceView* InKeysSRV,
+		FRHIShaderResourceView* InValuesSRV,
+		FRHIShaderResourceView* InOffsetsSRV,
 		FRadixSortUniformBufferRef& RadixSortUniformBuffer,
-		FShaderResourceViewRHIParamRef RadixSortParameterBufferSRV )
+		FRHIShaderResourceView* RadixSortParameterBufferSRV )
 	{
 		FComputeShaderRHIParamRef ComputeShaderRHI = GetComputeShader();
 		SetUniformBufferParameter(RHICmdList, ComputeShaderRHI, GetUniformBufferParameter<FRadixSortParameters>(), RadixSortUniformBuffer );
@@ -606,7 +606,7 @@ public:
 	/**
 	 * Set output buffer for this shader.
 	 */
-	void SetOutput(FRHICommandList& RHICmdList, FUnorderedAccessViewRHIParamRef OutKeysUAV, FUnorderedAccessViewRHIParamRef OutValuesUAV)
+	void SetOutput(FRHICommandList& RHICmdList, FRHIUnorderedAccessView* OutKeysUAV, FRHIUnorderedAccessView* OutValuesUAV)
 	{
 		FComputeShaderRHIParamRef ComputeShaderRHI = GetComputeShader();
 		if ( OutKeys.IsBound() )
@@ -627,27 +627,27 @@ public:
 		FComputeShaderRHIParamRef ComputeShaderRHI = GetComputeShader();
 		if ( RadixSortParameterBuffer.IsBound() )
 		{
-			RHICmdList.SetShaderResourceViewParameter(ComputeShaderRHI, RadixSortParameterBuffer.GetBaseIndex(), FShaderResourceViewRHIParamRef());
+			RHICmdList.SetShaderResourceViewParameter(ComputeShaderRHI, RadixSortParameterBuffer.GetBaseIndex(), nullptr);
 		}
 		if ( InKeys.IsBound() )
 		{
-			RHICmdList.SetShaderResourceViewParameter(ComputeShaderRHI, InKeys.GetBaseIndex(), FShaderResourceViewRHIParamRef());
+			RHICmdList.SetShaderResourceViewParameter(ComputeShaderRHI, InKeys.GetBaseIndex(), nullptr);
 		}
 		if ( InValues.IsBound() )
 		{
-			RHICmdList.SetShaderResourceViewParameter(ComputeShaderRHI, InValues.GetBaseIndex(), FShaderResourceViewRHIParamRef());
+			RHICmdList.SetShaderResourceViewParameter(ComputeShaderRHI, InValues.GetBaseIndex(), nullptr);
 		}
 		if ( InOffsets.IsBound() )
 		{
-			RHICmdList.SetShaderResourceViewParameter(ComputeShaderRHI, InOffsets.GetBaseIndex(), FShaderResourceViewRHIParamRef());
+			RHICmdList.SetShaderResourceViewParameter(ComputeShaderRHI, InOffsets.GetBaseIndex(), nullptr);
 		}
 		if ( OutKeys.IsBound() )
 		{
-			RHICmdList.SetUAVParameter(ComputeShaderRHI, OutKeys.GetBaseIndex(), FUnorderedAccessViewRHIParamRef());
+			RHICmdList.SetUAVParameter(ComputeShaderRHI, OutKeys.GetBaseIndex(), nullptr);
 		}
 		if ( OutValues.IsBound() )
 		{
-			RHICmdList.SetUAVParameter(ComputeShaderRHI, OutValues.GetBaseIndex(), FUnorderedAccessViewRHIParamRef());
+			RHICmdList.SetUAVParameter(ComputeShaderRHI, OutValues.GetBaseIndex(), nullptr);
 		}
 	}
 
@@ -778,7 +778,7 @@ int32 SortGPUBuffers(FRHICommandListImmediate& RHICmdList, FGPUSortBuffers SortB
 			UpsweepCS->UnbindBuffers(RHICmdList);
 
 			//barrier both UAVS since for next step.
-			FUnorderedAccessViewRHIParamRef PrePhase2BarrierUAVS[2];
+			FRHIUnorderedAccessView* PrePhase2BarrierUAVS[2];
 			PrePhase2BarrierUAVS[0] = GSortOffsetBuffers.BufferUAVs[0];
 			PrePhase2BarrierUAVS[1] = GSortOffsetBuffers.BufferUAVs[1];
 
@@ -806,7 +806,7 @@ int32 SortGPUBuffers(FRHICommandListImmediate& RHICmdList, FGPUSortBuffers SortB
 			//UAV is going to SRV, so transition to Readable.
 			RHICmdList.TransitionResource(EResourceTransitionAccess::EReadable, EResourceTransitionPipeline::EComputeToCompute, GSortOffsetBuffers.BufferUAVs[1]);
 
-			FUnorderedAccessViewRHIParamRef PrePhase3BarrierUAVS[2];
+			FRHIUnorderedAccessView* PrePhase3BarrierUAVS[2];
 			PrePhase3BarrierUAVS[0] = SortBuffers.RemoteKeyUAVs[BufferIndex ^ 0x1];
 			PrePhase3BarrierUAVS[1] = SortBuffers.RemoteValueUAVs[BufferIndex ^ 0x1];
 			RHICmdList.TransitionResources(EResourceTransitionAccess::ERWBarrier, EResourceTransitionPipeline::EComputeToCompute, PrePhase3BarrierUAVS, 2);
diff --git a/Engine/Source/Runtime/Engine/Private/GPUSort.h b/Engine/Source/Runtime/Engine/Private/GPUSort.h
--- a/Engine/Source/Runtime/Engine/Private/GPUSort.h
+++ b/Engine/Source/Runtime/Engine/Private/GPUSort.h
@@ -15,14 +15,14 @@
 struct FGPUSortBuffers
 {
 	/** Shader resource views for vertex buffers containing the keys. */
-	FShaderResourceViewRHIParamRef RemoteKeySRVs[2];
+	FRHIShaderResourceView* RemoteKeySRVs[2];
 	/** Unordered access views for vertex buffers containing the keys. */
-	FUnorderedAccessViewRHIParamRef RemoteKeyUAVs[2];
+	FRHIUnorderedAccessView* RemoteKeyUAVs[2];
 
 	/** Shader resource views for vertex buffers containing the values. */
-	FShaderResourceViewRHIParamRef RemoteValueSRVs[2];
+	FRHIShaderResourceView* RemoteValueSRVs[2];
 	/** Unordered access views for vertex buffers containing the values. */
-	FUnorderedAccessViewRHIParamRef RemoteValueUAVs[2];
+	FRHIUnorderedAccessView* RemoteValueUAVs[2];
 
 	/** Default constructor. */
 	FGPUSortBuffers()
diff --git a/Engine/Source/Runtime/Engine/Private/InstancedStaticMesh.h b/Engine/Source/Runtime/Engine/Private/InstancedStaticMesh.h
--- a/Engine/Source/Runtime/Engine/Private/InstancedStaticMesh.h
+++ b/Engine/Source/Runtime/Engine/Private/InstancedStaticMesh.h
@@ -166,9 +166,9 @@ struct FInstancedStaticMeshDataType
 	/** The stream to read the Lightmap Bias and Random instance ID from. */
 	FVertexStreamComponent InstanceLightmapAndShadowMapUVBiasComponent;
 
-	FShaderResourceViewRHIParamRef InstanceOriginSRV = nullptr;
-	FShaderResourceViewRHIParamRef InstanceTransformSRV = nullptr;
-	FShaderResourceViewRHIParamRef InstanceLightmapSRV = nullptr;
+	FRHIShaderResourceView* InstanceOriginSRV = nullptr;
+	FRHIShaderResourceView* InstanceTransformSRV = nullptr;
+	FRHIShaderResourceView* InstanceLightmapSRV = nullptr;
 };
 
 /**
@@ -258,17 +258,17 @@ public:
 	virtual bool SupportsNullPixelShader() const override { return false; }
 #endif
 
-	inline const FShaderResourceViewRHIParamRef GetInstanceOriginSRV() const
+	inline FRHIShaderResourceView* GetInstanceOriginSRV() const
 	{
 		return Data.InstanceOriginSRV;
 	}
 
-	inline const FShaderResourceViewRHIParamRef GetInstanceTransformSRV() const
+	inline FRHIShaderResourceView* GetInstanceTransformSRV() const
 	{
 		return Data.InstanceTransformSRV;
 	}
 
-	inline const FShaderResourceViewRHIParamRef GetInstanceLightmapSRV() const
+	inline FRHIShaderResourceView* GetInstanceLightmapSRV() const
 	{
 		return Data.InstanceLightmapSRV;
 	}
diff --git a/Engine/Source/Runtime/Engine/Private/MeshParticleVertexFactory.cpp b/Engine/Source/Runtime/Engine/Private/MeshParticleVertexFactory.cpp
--- a/Engine/Source/Runtime/Engine/Private/MeshParticleVertexFactory.cpp
+++ b/Engine/Source/Runtime/Engine/Private/MeshParticleVertexFactory.cpp
@@ -142,7 +142,7 @@ public:
 		return VB;
 	}
 
-	inline FShaderResourceViewRHIParamRef GetSRV() const
+	inline FRHIShaderResourceView* GetSRV() const
 	{
 		return SRV;
 	}
@@ -321,7 +321,7 @@ void FMeshParticleVertexFactory::UnlockPreviousTransformBuffer()
 	PrevTransformBuffer.Unlock();
 }
 
-FShaderResourceViewRHIParamRef FMeshParticleVertexFactory::GetPreviousTransformBufferSRV() const
+FRHIShaderResourceView* FMeshParticleVertexFactory::GetPreviousTransformBufferSRV() const
 {
 	return PrevTransformBuffer.SRV;
 }
diff --git a/Engine/Source/Runtime/Engine/Private/Particles/ParticleGpuSimulation.cpp b/Engine/Source/Runtime/Engine/Private/Particles/ParticleGpuSimulation.cpp
--- a/Engine/Source/Runtime/Engine/Private/Particles/ParticleGpuSimulation.cpp
+++ b/Engine/Source/Runtime/Engine/Private/Particles/ParticleGpuSimulation.cpp
@@ -801,8 +801,8 @@ void FGPUSpriteVertexFactoryShaderParametersVS::GetElementShaderBindings(
 	FRHISamplerState* SamplerStateLinear = TStaticSamplerState<SF_Bilinear>::GetRHI();
 	ShaderBindings.Add(Shader->GetUniformBufferParameter<FGPUSpriteEmitterUniformParameters>(), GPUVF->EmitterUniformBuffer);
 	ShaderBindings.Add(Shader->GetUniformBufferParameter<FGPUSpriteEmitterDynamicUniformParameters>(), GPUVF->EmitterDynamicUniformBuffer);
-	FShaderResourceViewRHIParamRef ParticleIndicesBuffer = GPUVF->ParticleIndicesBuffer->VertexBufferSRV;
-	ShaderBindings.Add(ParticleIndices, ParticleIndicesBuffer ? ParticleIndicesBuffer : (FShaderResourceViewRHIParamRef)GNullColorVertexBuffer.VertexBufferSRV);
+	FRHIShaderResourceView* ParticleIndicesBuffer = GPUVF->ParticleIndicesBuffer->VertexBufferSRV;
+	ShaderBindings.Add(ParticleIndices, ParticleIndicesBuffer ? ParticleIndicesBuffer : (FRHIShaderResourceView*)GNullColorVertexBuffer.VertexBufferSRV);
 	ShaderBindings.Add(ParticleIndicesOffset, GPUVF->ParticleIndicesOffset);
 
 	ShaderBindings.AddTexture(PositionTexture, PositionTextureSampler, SamplerStatePoint, GPUVF->PositionTextureRHI);
@@ -1216,7 +1216,7 @@ public:
 	void UnbindBuffers(FRHICommandList& RHICmdList)
 	{
 		FPixelShaderRHIParamRef PixelShaderRHI = GetPixelShader();
-		FShaderResourceViewRHIParamRef NullSRV = FShaderResourceViewRHIParamRef();
+		FRHIShaderResourceView* NullSRV = nullptr;
 		for (int32 i = 0; i < MAX_VECTOR_FIELDS; ++i)
 		{
 			if (VectorFieldTextures[i].IsBound())
@@ -2269,7 +2269,7 @@ public:
 	/**
 	 * Set output buffers for this shader.
 	 */
-	void SetOutput(FRHICommandList& RHICmdList, FUnorderedAccessViewRHIParamRef OutBoundsUAV )
+	void SetOutput(FRHICommandList& RHICmdList, FRHIUnorderedAccessView* OutBoundsUAV )
 	{
 		FComputeShaderRHIParamRef ComputeShaderRHI = GetComputeShader();
 		if ( OutBounds.IsBound() )
@@ -2284,7 +2284,7 @@ public:
 	void SetParameters(
 		FRHICommandList& RHICmdList,
 		FParticleBoundsUniformBufferRef& UniformBuffer,
-		FShaderResourceViewRHIParamRef InIndicesSRV,
+		FRHIShaderResourceView* InIndicesSRV,
 		FTexture2DRHIParamRef PositionTextureRHI
 		)
 	{
@@ -2311,11 +2311,11 @@ public:
 		FComputeShaderRHIParamRef ComputeShaderRHI = GetComputeShader();
 		if ( InParticleIndices.IsBound() )
 		{
-			RHICmdList.SetShaderResourceViewParameter(ComputeShaderRHI, InParticleIndices.GetBaseIndex(), FShaderResourceViewRHIParamRef());
+			RHICmdList.SetShaderResourceViewParameter(ComputeShaderRHI, InParticleIndices.GetBaseIndex(), nullptr);
 		}
 		if ( OutBounds.IsBound() )
 		{
-			RHICmdList.SetUAVParameter(ComputeShaderRHI, OutBounds.GetBaseIndex(), FUnorderedAccessViewRHIParamRef());
+			RHICmdList.SetUAVParameter(ComputeShaderRHI, OutBounds.GetBaseIndex(), nullptr);
 		}
 	}
 
@@ -2350,7 +2350,7 @@ static bool AreBoundsValid( const FVector& Mins, const FVector& Maxs )
  */
 static FBox ComputeParticleBounds(
 	FRHICommandListImmediate& RHICmdList,
-	FShaderResourceViewRHIParamRef VertexBufferSRV,
+	FRHIShaderResourceView* VertexBufferSRV,
 	FTexture2DRHIParamRef PositionTextureRHI,
 	int32 ParticleCount )
 {
diff --git a/Engine/Source/Runtime/Engine/Private/Particles/ParticleSortingGPU.cpp b/Engine/Source/Runtime/Engine/Private/Particles/ParticleSortingGPU.cpp
--- a/Engine/Source/Runtime/Engine/Private/Particles/ParticleSortingGPU.cpp
+++ b/Engine/Source/Runtime/Engine/Private/Particles/ParticleSortingGPU.cpp
@@ -91,7 +91,7 @@ public:
 	/**
 	 * Set output buffers for this shader.
 	 */
-	void SetOutput(FRHICommandList& RHICmdList, FUnorderedAccessViewRHIParamRef OutKeysUAV, FUnorderedAccessViewRHIParamRef OutIndicesUAV )
+	void SetOutput(FRHICommandList& RHICmdList, FRHIUnorderedAccessView* OutKeysUAV, FRHIUnorderedAccessView* OutIndicesUAV )
 	{
 		FComputeShaderRHIParamRef ComputeShaderRHI = GetComputeShader();
 		if ( OutKeys.IsBound() )
@@ -110,7 +110,7 @@ public:
 	void SetParameters(
 		FRHICommandList& RHICmdList,
 		FParticleKeyGenUniformBufferRef& UniformBuffer,
-		FShaderResourceViewRHIParamRef InIndicesSRV
+		FRHIShaderResourceView* InIndicesSRV
 		)
 	{
 		FComputeShaderRHIParamRef ComputeShaderRHI = GetComputeShader();
@@ -141,15 +141,15 @@ public:
 		FComputeShaderRHIParamRef ComputeShaderRHI = GetComputeShader();
 		if ( InParticleIndices.IsBound() )
 		{
-			RHICmdList.SetShaderResourceViewParameter(ComputeShaderRHI, InParticleIndices.GetBaseIndex(), FShaderResourceViewRHIParamRef());
+			RHICmdList.SetShaderResourceViewParameter(ComputeShaderRHI, InParticleIndices.GetBaseIndex(), nullptr);
 		}
 		if ( OutKeys.IsBound() )
 		{
-			RHICmdList.SetUAVParameter(ComputeShaderRHI, OutKeys.GetBaseIndex(), FUnorderedAccessViewRHIParamRef());
+			RHICmdList.SetUAVParameter(ComputeShaderRHI, OutKeys.GetBaseIndex(), nullptr);
 		}
 		if ( OutParticleIndices.IsBound() )
 		{
-			RHICmdList.SetUAVParameter(ComputeShaderRHI, OutParticleIndices.GetBaseIndex(), FUnorderedAccessViewRHIParamRef());
+			RHICmdList.SetUAVParameter(ComputeShaderRHI, OutParticleIndices.GetBaseIndex(), nullptr);
 		}
 	}
 
@@ -177,8 +177,8 @@ IMPLEMENT_SHADER_TYPE(,FParticleSortKeyGenCS,TEXT("/Engine/Private/ParticleSortK
  */
 static int32 GenerateParticleSortKeys(
 	FRHICommandListImmediate& RHICmdList,
-	FUnorderedAccessViewRHIParamRef KeyBufferUAV,
-	FUnorderedAccessViewRHIParamRef SortedVertexBufferUAV,
+	FRHIUnorderedAccessView* KeyBufferUAV,
+	FRHIUnorderedAccessView* SortedVertexBufferUAV,
 	FTexture2DRHIParamRef PositionTextureRHI,
 	const TArray<FParticleSimulationSortInfo>& SimulationsToSort,
 	ERHIFeatureLevel::Type FeatureLevel
@@ -192,7 +192,7 @@ static int32 GenerateParticleSortKeys(
 	const uint32 MaxGroupCount = 128;
 	int32 TotalParticleCount = 0;
 
-	FUnorderedAccessViewRHIParamRef OutputUAVs[2];
+	FRHIUnorderedAccessView* OutputUAVs[2];
 	OutputUAVs[0] = KeyBufferUAV;
 	OutputUAVs[1] = SortedVertexBufferUAV;
 
diff --git a/Engine/Source/Runtime/Engine/Private/Particles/ParticleSortingGPU.h b/Engine/Source/Runtime/Engine/Private/Particles/ParticleSortingGPU.h
--- a/Engine/Source/Runtime/Engine/Private/Particles/ParticleSortingGPU.h
+++ b/Engine/Source/Runtime/Engine/Private/Particles/ParticleSortingGPU.h
@@ -43,7 +43,7 @@ public:
 	/**
 	 * Retrieve the UAV for writing particle sort keys.
 	 */
-	FUnorderedAccessViewRHIParamRef GetKeyBufferUAV()
+	FRHIUnorderedAccessView* GetKeyBufferUAV()
 	{
 		return KeyBufferUAVs[0];
 	}
@@ -52,7 +52,7 @@ public:
 	 * Retrieve the UAV for writing particle vertices.
 	 * bAsUint : whether to return a G16R16 view or a Uint32 view.
 	 */
-	FORCEINLINE FUnorderedAccessViewRHIParamRef GetVertexBufferUAV()
+	FORCEINLINE FRHIUnorderedAccessView* GetVertexBufferUAV()
 	{
 		return VertexBufferUAVs[0];
 	}
@@ -74,7 +74,7 @@ public:
 	/**
 	 * Retrieve the SRV for the sorted vertex buffer at the given index.
 	 */
-	FShaderResourceViewRHIParamRef GetSortedVertexBufferSRV(int32 BufferIndex)
+	FRHIShaderResourceView* GetSortedVertexBufferSRV(int32 BufferIndex)
 	{
 		check((BufferIndex & 0xFFFFFFFE) == 0);
 		return VertexBufferSRVs[BufferIndex];
@@ -83,7 +83,7 @@ public:
 	/**
 	 * Retrieve the UAV for the sorted vertex buffer at the given index.
 	 */
-	FUnorderedAccessViewRHIParamRef GetSortedVertexBufferUAV(int32 BufferIndex)
+	FRHIUnorderedAccessView* GetSortedVertexBufferUAV(int32 BufferIndex)
 	{
 		check((BufferIndex & 0xFFFFFFFE) == 0);
 		return VertexBufferUAVs[BufferIndex];
@@ -127,7 +127,7 @@ private:
 struct FParticleSimulationSortInfo
 {
 	/** Vertex buffer containing indices in to the particle state texture. */
-	FShaderResourceViewRHIParamRef VertexBufferSRV;
+	FRHIShaderResourceView* VertexBufferSRV;
 	/** World space position from which to sort. */
 	FVector ViewOrigin;
 	/** The number of particles in the simulation. */
diff --git a/Engine/Source/Runtime/Engine/Private/Particles/ParticleVertexFactory.cpp b/Engine/Source/Runtime/Engine/Private/Particles/ParticleVertexFactory.cpp
--- a/Engine/Source/Runtime/Engine/Private/Particles/ParticleVertexFactory.cpp
+++ b/Engine/Source/Runtime/Engine/Private/Particles/ParticleVertexFactory.cpp
@@ -92,7 +92,7 @@ public:
 		ShaderBindings.Add(Shader->GetUniformBufferParameter<FParticleSpriteUniformParameters>(), SpriteVF->GetSpriteUniformBuffer() );
 		
 		ShaderBindings.Add(NumCutoutVerticesPerFrame, SpriteVF->GetNumCutoutVerticesPerFrame());
-		FShaderResourceViewRHIParamRef NullSRV = GFNullSubUVCutoutVertexBuffer.VertexBufferSRV;
+		FRHIShaderResourceView* NullSRV = GFNullSubUVCutoutVertexBuffer.VertexBufferSRV;
 		ShaderBindings.Add(CutoutGeometry, SpriteVF->GetCutoutGeometrySRV() ? SpriteVF->GetCutoutGeometrySRV() : NullSRV);
 	}
 
diff --git a/Engine/Source/Runtime/Engine/Private/SkeletalRenderGPUSkin.h b/Engine/Source/Runtime/Engine/Private/SkeletalRenderGPUSkin.h
--- a/Engine/Source/Runtime/Engine/Private/SkeletalRenderGPUSkin.h
+++ b/Engine/Source/Runtime/Engine/Private/SkeletalRenderGPUSkin.h
@@ -222,7 +222,7 @@ public:
 	bool bNeedsInitialClear;
 
 	// @param guaranteed only to be valid if the vertex buffer is valid
-	FShaderResourceViewRHIParamRef GetSRV() const
+	FRHIShaderResourceView* GetSRV() const
 	{
 		return SRVValue;
 	}
diff --git a/Engine/Source/Runtime/Engine/Private/VectorField.cpp b/Engine/Source/Runtime/Engine/Private/VectorField.cpp
--- a/Engine/Source/Runtime/Engine/Private/VectorField.cpp
+++ b/Engine/Source/Runtime/Engine/Private/VectorField.cpp
@@ -821,7 +821,7 @@ public:
 	/**
 	 * Set output buffer for this shader.
 	 */
-	void SetOutput(FRHICommandList& RHICmdList, FUnorderedAccessViewRHIParamRef VolumeTextureUAV)
+	void SetOutput(FRHICommandList& RHICmdList, FRHIUnorderedAccessView* VolumeTextureUAV)
 	{
 		FComputeShaderRHIParamRef ComputeShaderRHI = GetComputeShader();
 		if ( OutVolumeTexture.IsBound() )
@@ -838,7 +838,7 @@ public:
 		FComputeShaderRHIParamRef ComputeShaderRHI = GetComputeShader();
 		if ( OutVolumeTexture.IsBound() )
 		{
-			RHICmdList.SetUAVParameter(ComputeShaderRHI, OutVolumeTexture.GetBaseIndex(), FUnorderedAccessViewRHIParamRef());
+			RHICmdList.SetUAVParameter(ComputeShaderRHI, OutVolumeTexture.GetBaseIndex(), nullptr);
 		}
 	}
 
diff --git a/Engine/Source/Runtime/Engine/Public/Components.h b/Engine/Source/Runtime/Engine/Public/Components.h
--- a/Engine/Source/Runtime/Engine/Public/Components.h
+++ b/Engine/Source/Runtime/Engine/Public/Components.h
@@ -45,15 +45,15 @@ struct FStaticMeshDataType
 	/** The stream to read the vertex color from. */
 	FVertexStreamComponent ColorComponent;
 
-	FShaderResourceViewRHIParamRef PositionComponentSRV = nullptr;
+	FRHIShaderResourceView* PositionComponentSRV = nullptr;
 
-	FShaderResourceViewRHIParamRef TangentsSRV = nullptr;
+	FRHIShaderResourceView* TangentsSRV = nullptr;
 
 	/** A SRV to manually bind and load TextureCoordinates in the Vertexshader. */
-	FShaderResourceViewRHIParamRef TextureCoordinatesSRV = nullptr;
+	FRHIShaderResourceView* TextureCoordinatesSRV = nullptr;
 
 	/** A SRV to manually bind and load Colors in the Vertexshader. */
-	FShaderResourceViewRHIParamRef ColorComponentsSRV = nullptr;
+	FRHIShaderResourceView* ColorComponentsSRV = nullptr;
 
 	int LightMapCoordinateIndex = -1;
 	int NumTexCoords = -1;
diff --git a/Engine/Source/Runtime/Engine/Public/GPUSkinCache.h b/Engine/Source/Runtime/Engine/Public/GPUSkinCache.h
--- a/Engine/Source/Runtime/Engine/Public/GPUSkinCache.h
+++ b/Engine/Source/Runtime/Engine/Public/GPUSkinCache.h
@@ -195,7 +195,7 @@ public:
 			return WithTangents ? &Tangents : nullptr;
 		}
 
-		void RemoveAllFromTransitionArray(TArray<FUnorderedAccessViewRHIParamRef>& BuffersToTransition);
+		void RemoveAllFromTransitionArray(TArray<FRHIUnorderedAccessView*>& BuffersToTransition);
 
 	private:
 		// Output of the GPU skinning (ie Pos, Normals)
@@ -296,7 +296,7 @@ public:
 #endif // RHI_RAYTRACING
 
 protected:
-	TArray<FUnorderedAccessViewRHIParamRef> BuffersToTransition;
+	TArray<FRHIUnorderedAccessView*> BuffersToTransition;
 #if RHI_RAYTRACING
 	TSet<FRayTracingGeometry*> RayTracingGeometriesToUpdate;
 #endif // RHI_RAYTRACING
diff --git a/Engine/Source/Runtime/Engine/Public/LocalVertexFactory.h b/Engine/Source/Runtime/Engine/Public/LocalVertexFactory.h
--- a/Engine/Source/Runtime/Engine/Public/LocalVertexFactory.h
+++ b/Engine/Source/Runtime/Engine/Public/LocalVertexFactory.h
@@ -96,22 +96,22 @@ public:
 		VertexStreams.Add(FVertexInputStream(ColorStreamIndex, 0, ColorVertexBuffer->VertexBufferRHI));
 	}
 
-	inline const FShaderResourceViewRHIParamRef GetPositionsSRV() const
+	inline FRHIShaderResourceView* GetPositionsSRV() const
 	{
 		return Data.PositionComponentSRV;
 	}
 
-	inline const FShaderResourceViewRHIParamRef GetTangentsSRV() const
+	inline FRHIShaderResourceView* GetTangentsSRV() const
 	{
 		return Data.TangentsSRV;
 	}
 
-	inline const FShaderResourceViewRHIParamRef GetTextureCoordinatesSRV() const
+	inline FRHIShaderResourceView* GetTextureCoordinatesSRV() const
 	{
 		return Data.TextureCoordinatesSRV;
 	}
 
-	inline const FShaderResourceViewRHIParamRef GetColorComponentsSRV() const
+	inline FRHIShaderResourceView* GetColorComponentsSRV() const
 	{
 		return Data.ColorComponentsSRV;
 	}
diff --git a/Engine/Source/Runtime/Engine/Public/MeshParticleVertexFactory.h b/Engine/Source/Runtime/Engine/Public/MeshParticleVertexFactory.h
--- a/Engine/Source/Runtime/Engine/Public/MeshParticleVertexFactory.h
+++ b/Engine/Source/Runtime/Engine/Public/MeshParticleVertexFactory.h
@@ -166,7 +166,7 @@ public:
 
 	uint8* LockPreviousTransformBuffer(uint32 ParticleCount);
 	void UnlockPreviousTransformBuffer();
-	FShaderResourceViewRHIParamRef GetPreviousTransformBufferSRV() const;
+	FRHIShaderResourceView* GetPreviousTransformBufferSRV() const;
 
 	/**
 	* Copy the data from another vertex factory
diff --git a/Engine/Source/Runtime/Engine/Public/ParticleResources.h b/Engine/Source/Runtime/Engine/Public/ParticleResources.h
--- a/Engine/Source/Runtime/Engine/Public/ParticleResources.h
+++ b/Engine/Source/Runtime/Engine/Public/ParticleResources.h
@@ -56,7 +56,7 @@ public:
 ENGINE_API extern TGlobalResource<FParticleIndexBuffer> GParticleIndexBuffer;
 ENGINE_API extern TGlobalResource<FSixTriangleParticleIndexBuffer> GSixTriangleParticleIndexBuffer;
 
-typedef FShaderResourceViewRHIParamRef FParticleShaderParamRef;
+typedef FRHIShaderResourceView* FParticleShaderParamRef;
 typedef FVertexBufferRHIParamRef FParticleBufferParamRef;
 
 /**
diff --git a/Engine/Source/Runtime/Engine/Public/ParticleVertexFactory.h b/Engine/Source/Runtime/Engine/Public/ParticleVertexFactory.h
--- a/Engine/Source/Runtime/Engine/Public/ParticleVertexFactory.h
+++ b/Engine/Source/Runtime/Engine/Public/ParticleVertexFactory.h
@@ -214,14 +214,14 @@ public:
 		return SpriteUniformBuffer;
 	}
 
-	void SetCutoutParameters(int32 InNumCutoutVerticesPerFrame, FShaderResourceViewRHIParamRef InCutoutGeometrySRV)
+	void SetCutoutParameters(int32 InNumCutoutVerticesPerFrame, FRHIShaderResourceView* InCutoutGeometrySRV)
 	{
 		NumCutoutVerticesPerFrame = InNumCutoutVerticesPerFrame;
 		CutoutGeometrySRV = InCutoutGeometrySRV;
 	}
 
 	inline int32 GetNumCutoutVerticesPerFrame() const { return NumCutoutVerticesPerFrame; }
-	inline FShaderResourceViewRHIParamRef GetCutoutGeometrySRV() const { return CutoutGeometrySRV; }
+	inline FRHIShaderResourceView* GetCutoutGeometrySRV() const { return CutoutGeometrySRV; }
 
 	void SetCustomAlignment(bool bAlign)
 	{
@@ -250,7 +250,7 @@ private:
 	FUniformBufferRHIRef SpriteUniformBuffer;
 
 	int32 NumCutoutVerticesPerFrame;
-	FShaderResourceViewRHIParamRef CutoutGeometrySRV;
+	FRHIShaderResourceView* CutoutGeometrySRV;
 	bool bCustomAlignment;
 	bool bUsesDynamicParameter;
 	uint32 DynamicParameterStride;
diff --git a/Engine/Source/Runtime/Engine/Public/RawIndexBuffer.h b/Engine/Source/Runtime/Engine/Public/RawIndexBuffer.h
--- a/Engine/Source/Runtime/Engine/Public/RawIndexBuffer.h
+++ b/Engine/Source/Runtime/Engine/Public/RawIndexBuffer.h
@@ -324,7 +324,7 @@ public:
 	virtual int32 GetResourceDataSize() const = 0;
 
 	// @param guaranteed only to be valid if the vertex buffer is valid and the buffer was created with the SRV flags
-	FShaderResourceViewRHIParamRef GetSRV() const
+	FRHIShaderResourceView* GetSRV() const
 	{
 		return SRVValue;
 	}
diff --git a/Engine/Source/Runtime/Engine/Public/Rendering/ColorVertexBuffer.h b/Engine/Source/Runtime/Engine/Public/Rendering/ColorVertexBuffer.h
--- a/Engine/Source/Runtime/Engine/Public/Rendering/ColorVertexBuffer.h
+++ b/Engine/Source/Runtime/Engine/Public/Rendering/ColorVertexBuffer.h
@@ -173,7 +173,7 @@ public:
 	ENGINE_API void BindColorVertexBuffer(const class FVertexFactory* VertexFactory, struct FStaticMeshDataType& StaticMeshData) const;
 	ENGINE_API static void BindDefaultColorVertexBuffer(const class FVertexFactory* VertexFactory, struct FStaticMeshDataType& StaticMeshData, NullBindStride BindStride);
 
-	FORCEINLINE const FShaderResourceViewRHIParamRef GetColorComponentsSRV() const
+	FORCEINLINE FRHIShaderResourceView* GetColorComponentsSRV() const
 	{
 		return ColorComponentsSRV;
 	}
diff --git a/Engine/Source/Runtime/Landscape/Private/LandscapeEditLayers.cpp b/Engine/Source/Runtime/Landscape/Private/LandscapeEditLayers.cpp
--- a/Engine/Source/Runtime/Landscape/Private/LandscapeEditLayers.cpp
+++ b/Engine/Source/Runtime/Landscape/Private/LandscapeEditLayers.cpp
@@ -731,7 +731,7 @@ public:
 
 	void UnsetParameters(FRHICommandList& RHICmdList)
 	{
-		SetUAVParameter(RHICmdList, GetComputeShader(), AtlasPaintListsParam, FUnorderedAccessViewRHIParamRef());
+		SetUAVParameter(RHICmdList, GetComputeShader(), AtlasPaintListsParam, nullptr);
 	}
 
 	virtual bool Serialize(FArchive& Ar) override
@@ -928,7 +928,7 @@ public:
 
 	void UnsetParameters(FRHICommandList& RHICmdList)
 	{
-		SetUAVParameter(RHICmdList, GetComputeShader(), ComponentWeightmapParam, FUnorderedAccessViewRHIParamRef());
+		SetUAVParameter(RHICmdList, GetComputeShader(), ComponentWeightmapParam, nullptr);
 	}
 
 	virtual bool Serialize(FArchive& Ar) override
