Fixed merge error introduced in CL 6621195. RHIUpdateAccelerationStructures calls BuildAccelerationStructure(), as it was before merge.

diff --git a/Engine/Source/Runtime/D3D12RHI/Private/D3D12RayTracing.cpp b/Engine/Source/Runtime/D3D12RHI/Private/D3D12RayTracing.cpp
--- a/Engine/Source/Runtime/D3D12RHI/Private/D3D12RayTracing.cpp
+++ b/Engine/Source/Runtime/D3D12RHI/Private/D3D12RayTracing.cpp
@@ -2379,6 +2379,9 @@ void FD3D12CommandContext::RHIUpdateAccelerationStructures(const TArrayView<cons
 	{
 		FD3D12RayTracingGeometry* Geometry = FD3D12DynamicRHI::ResourceCast(P.Geometry.GetReference());
 		Geometry->SetDirty(GetGPUMask(), true);
+
+		const bool bIsUpdate = !GRayTracingDebugForceFullBuild;
+		Geometry->BuildAccelerationStructure(*this, bIsUpdate);
 	}
 }
 
