DR - Deprecate FRayTracingShaderRHIParamRef and FRenderQueryRHIParamRef

diff --git a/Engine/Source/Runtime/Apple/MetalRHI/Private/MetalQuery.cpp b/Engine/Source/Runtime/Apple/MetalRHI/Private/MetalQuery.cpp
--- a/Engine/Source/Runtime/Apple/MetalRHI/Private/MetalQuery.cpp
+++ b/Engine/Source/Runtime/Apple/MetalRHI/Private/MetalQuery.cpp
@@ -305,7 +305,7 @@ FRenderQueryRHIRef FMetalDynamicRHI::RHICreateRenderQuery(ERenderQueryType Query
 	}
 }
 
-bool FMetalDynamicRHI::RHIGetRenderQueryResult(FRenderQueryRHIParamRef QueryRHI,uint64& OutNumPixels,bool bWait)
+bool FMetalDynamicRHI::RHIGetRenderQueryResult(FRHIRenderQuery* QueryRHI,uint64& OutNumPixels,bool bWait)
 {
 	@autoreleasepool {
 	check(IsInRenderingThread());
@@ -375,7 +375,7 @@ bool FMetalDynamicRHI::RHIGetRenderQueryResult(FRenderQueryRHIParamRef QueryRHI,
 }
 
 // Occlusion/Timer queries.
-void FMetalRHICommandContext::RHIBeginRenderQuery(FRenderQueryRHIParamRef QueryRHI)
+void FMetalRHICommandContext::RHIBeginRenderQuery(FRHIRenderQuery* QueryRHI)
 {
 	@autoreleasepool {
 	FMetalRenderQuery* Query = ResourceCast(QueryRHI);
@@ -384,7 +384,7 @@ void FMetalRHICommandContext::RHIBeginRenderQuery(FRenderQueryRHIParamRef QueryR
 	}
 }
 
-void FMetalRHICommandContext::RHIEndRenderQuery(FRenderQueryRHIParamRef QueryRHI)
+void FMetalRHICommandContext::RHIEndRenderQuery(FRHIRenderQuery* QueryRHI)
 {
 	@autoreleasepool {
 	FMetalRenderQuery* Query = ResourceCast(QueryRHI);
diff --git a/Engine/Source/Runtime/Apple/MetalRHI/Private/MetalUAV.cpp b/Engine/Source/Runtime/Apple/MetalRHI/Private/MetalUAV.cpp
--- a/Engine/Source/Runtime/Apple/MetalRHI/Private/MetalUAV.cpp
+++ b/Engine/Source/Runtime/Apple/MetalRHI/Private/MetalUAV.cpp
@@ -806,7 +806,7 @@ void FMetalRHICommandContext::RHICopyToStagingBuffer(FVertexBufferRHIParamRef So
 	}
 }
 
-void FMetalRHICommandContext::RHIWriteGPUFence(FGPUFenceRHIParamRef FenceRHI)
+void FMetalRHICommandContext::RHIWriteGPUFence(FRHIGPUFence* FenceRHI)
 {
 	@autoreleasepool {
 		check(FenceRHI);
diff --git a/Engine/Source/Runtime/Apple/MetalRHI/Public/MetalRHI.h b/Engine/Source/Runtime/Apple/MetalRHI/Public/MetalRHI.h
--- a/Engine/Source/Runtime/Apple/MetalRHI/Public/MetalRHI.h
+++ b/Engine/Source/Runtime/Apple/MetalRHI/Public/MetalRHI.h
@@ -121,7 +121,7 @@ public:
 	virtual void RHIReadSurfaceFloatData(FTextureRHIParamRef Texture,FIntRect Rect,TArray<FFloat16Color>& OutData,ECubeFace CubeFace,int32 ArrayIndex,int32 MipIndex) final override;
 	virtual void RHIRead3DSurfaceFloatData(FTextureRHIParamRef Texture,FIntRect Rect,FIntPoint ZMinMax,TArray<FFloat16Color>& OutData) final override;
 	virtual FRenderQueryRHIRef RHICreateRenderQuery(ERenderQueryType QueryType) final override;
-	virtual bool RHIGetRenderQueryResult(FRenderQueryRHIParamRef RenderQuery, uint64& OutResult, bool bWait) final override;
+	virtual bool RHIGetRenderQueryResult(FRHIRenderQuery* RenderQuery, uint64& OutResult, bool bWait) final override;
 	virtual FTexture2DRHIRef RHIGetViewportBackBuffer(FViewportRHIParamRef Viewport) final override;
 	virtual void RHIAdvanceFrameForGetViewportBackBuffer(FViewportRHIParamRef Viewport) final override;
 	virtual void RHIAcquireThreadOwnership() final override;
diff --git a/Engine/Source/Runtime/Apple/MetalRHI/Public/MetalRHIContext.h b/Engine/Source/Runtime/Apple/MetalRHI/Public/MetalRHIContext.h
--- a/Engine/Source/Runtime/Apple/MetalRHI/Public/MetalRHIContext.h
+++ b/Engine/Source/Runtime/Apple/MetalRHI/Public/MetalRHIContext.h
@@ -57,9 +57,9 @@ public:
 	 */
 	virtual void RHICopyToResolveTarget(FTextureRHIParamRef SourceTexture, FTextureRHIParamRef DestTexture, const FResolveParams& ResolveParams) final override;
 	
-	virtual void RHIBeginRenderQuery(FRenderQueryRHIParamRef RenderQuery) final override;
+	virtual void RHIBeginRenderQuery(FRHIRenderQuery* RenderQuery) final override;
 	
-	virtual void RHIEndRenderQuery(FRenderQueryRHIParamRef RenderQuery) final override;
+	virtual void RHIEndRenderQuery(FRHIRenderQuery* RenderQuery) final override;
 	
 	void RHIBeginOcclusionQueryBatch(uint32 NumQueriesInBatch);
 	
@@ -268,7 +268,7 @@ public:
 	virtual void RHITransitionResources(EResourceTransitionAccess TransitionType, EResourceTransitionPipeline TransitionPipeline, FRHIUnorderedAccessView** InUAVs, int32 NumUAVs, FRHIComputeFence* WriteComputeFence) final override;
 	
 	virtual void RHICopyToStagingBuffer(FVertexBufferRHIParamRef SourceBufferRHI, FStagingBufferRHIParamRef DestinationStagingBufferRHI, uint32 Offset, uint32 NumBytes) final override;
-	virtual void RHIWriteGPUFence(FGPUFenceRHIParamRef FenceRHI) final override;
+	virtual void RHIWriteGPUFence(FRHIGPUFence* FenceRHI) final override;
 
 	/**
 	 * Explicitly transition a texture resource from readable -> writable by the GPU or vice versa.
diff --git a/Engine/Source/Runtime/D3D12RHI/Private/D3D12CommandContext.h b/Engine/Source/Runtime/D3D12RHI/Private/D3D12CommandContext.h
--- a/Engine/Source/Runtime/D3D12RHI/Private/D3D12CommandContext.h
+++ b/Engine/Source/Runtime/D3D12RHI/Private/D3D12CommandContext.h
@@ -253,9 +253,9 @@ public:
 	virtual void RHICopyTexture(FTextureRHIParamRef SourceTexture, FTextureRHIParamRef DestTexture, const FRHICopyTextureInfo& CopyInfo) final override;
 	virtual void RHITransitionResources(EResourceTransitionAccess TransitionType, FTextureRHIParamRef* InTextures, int32 NumTextures) final override;
 	virtual void RHICopyToStagingBuffer(FVertexBufferRHIParamRef SourceBuffer, FStagingBufferRHIParamRef DestinationStagingBuffer, uint32 Offset, uint32 NumBytes) final override;
-	virtual void RHIWriteGPUFence(FGPUFenceRHIParamRef Fence) final override;
-	virtual void RHIBeginRenderQuery(FRenderQueryRHIParamRef RenderQuery) final override;
-	virtual void RHIEndRenderQuery(FRenderQueryRHIParamRef RenderQuery) final override;
+	virtual void RHIWriteGPUFence(FRHIGPUFence* Fence) final override;
+	virtual void RHIBeginRenderQuery(FRHIRenderQuery* RenderQuery) final override;
+	virtual void RHIEndRenderQuery(FRHIRenderQuery* RenderQuery) final override;
 	void RHIBeginOcclusionQueryBatch(uint32 NumQueriesInBatch);
 	void RHIEndOcclusionQueryBatch();
 	virtual void RHIBeginScene() final override;
@@ -346,7 +346,7 @@ public:
 		FRHIShaderResourceView* Rays,
 		FRHIUnorderedAccessView* Output,
 		uint32 NumRays) final override;
-	virtual void RHIRayTraceDispatch(FRHIRayTracingPipelineState* RayTracingPipelineState, FRayTracingShaderRHIParamRef RayGenShader,
+	virtual void RHIRayTraceDispatch(FRHIRayTracingPipelineState* RayTracingPipelineState, FRHIRayTracingShader* RayGenShader,
 		FRayTracingSceneRHIParamRef Scene,
 		const FRayTracingShaderBindings& GlobalResourceBindings,
 		uint32 Width, uint32 Height) final override;
@@ -466,7 +466,7 @@ public:
 	{
 		ContextRedirect(RHICopyToStagingBuffer(SourceBuffer, DestinationStagingBuffer, Offset, NumBytes));
 	}
-	FORCEINLINE virtual void RHIWriteGPUFence(FGPUFenceRHIParamRef Fence) final override
+	FORCEINLINE virtual void RHIWriteGPUFence(FRHIGPUFence* Fence) final override
 	{
 		ContextRedirect(RHIWriteGPUFence(Fence));
 	}
@@ -536,11 +536,11 @@ public:
 	{
 		ContextRedirect(RHITransitionResources(TransitionType, InTextures, NumTextures));
 	}
-	FORCEINLINE virtual void RHIBeginRenderQuery(FRenderQueryRHIParamRef RenderQuery) final override
+	FORCEINLINE virtual void RHIBeginRenderQuery(FRHIRenderQuery* RenderQuery) final override
 	{
 		ContextRedirect(RHIBeginRenderQuery(RenderQuery));
 	}
-	FORCEINLINE virtual void RHIEndRenderQuery(FRenderQueryRHIParamRef RenderQuery) final override
+	FORCEINLINE virtual void RHIEndRenderQuery(FRHIRenderQuery* RenderQuery) final override
 	{
 		ContextRedirect(RHIEndRenderQuery(RenderQuery));
 	}
@@ -791,7 +791,7 @@ public:
 		ContextRedirect(RHIRayTraceIntersection(Scene, Rays, Output, NumRays));
 	}
 
-	virtual void RHIRayTraceDispatch(FRHIRayTracingPipelineState* RayTracingPipelineState, FRayTracingShaderRHIParamRef RayGenShader,
+	virtual void RHIRayTraceDispatch(FRHIRayTracingPipelineState* RayTracingPipelineState, FRHIRayTracingShader* RayGenShader,
 		FRayTracingSceneRHIParamRef Scene,
 		const FRayTracingShaderBindings& GlobalResourceBindings,
 		uint32 Width, uint32 Height) final override
diff --git a/Engine/Source/Runtime/D3D12RHI/Private/D3D12Commands.cpp b/Engine/Source/Runtime/D3D12RHI/Private/D3D12Commands.cpp
--- a/Engine/Source/Runtime/D3D12RHI/Private/D3D12Commands.cpp
+++ b/Engine/Source/Runtime/D3D12RHI/Private/D3D12Commands.cpp
@@ -432,7 +432,7 @@ void FD3D12CommandContext::RHICopyToStagingBuffer(FVertexBufferRHIParamRef Sourc
 	}
 }
 
-void FD3D12CommandContext::RHIWriteGPUFence(FGPUFenceRHIParamRef FenceRHI)
+void FD3D12CommandContext::RHIWriteGPUFence(FRHIGPUFence* FenceRHI)
 {
 	check(FenceRHI)
 
@@ -1153,7 +1153,7 @@ void FD3D12CommandContext::RHISetRenderTargetsAndClear(const FRHISetRenderTarget
 }
 
 // Occlusion/Timer queries.
-void FD3D12CommandContext::RHIBeginRenderQuery(FRenderQueryRHIParamRef QueryRHI)
+void FD3D12CommandContext::RHIBeginRenderQuery(FRHIRenderQuery* QueryRHI)
 {
 	FD3D12RenderQuery* Query = RetrieveObject<FD3D12RenderQuery>(QueryRHI);
 	check(IsDefaultContext());
@@ -1166,7 +1166,7 @@ void FD3D12CommandContext::RHIBeginRenderQuery(FRenderQueryRHIParamRef QueryRHI)
 #endif
 }
 
-void FD3D12CommandContext::RHIEndRenderQuery(FRenderQueryRHIParamRef QueryRHI)
+void FD3D12CommandContext::RHIEndRenderQuery(FRHIRenderQuery* QueryRHI)
 {
 	FD3D12RenderQuery* Query = RetrieveObject<FD3D12RenderQuery>(QueryRHI);
 	check(IsDefaultContext());
diff --git a/Engine/Source/Runtime/D3D12RHI/Private/D3D12Query.cpp b/Engine/Source/Runtime/D3D12RHI/Private/D3D12Query.cpp
--- a/Engine/Source/Runtime/D3D12RHI/Private/D3D12Query.cpp
+++ b/Engine/Source/Runtime/D3D12RHI/Private/D3D12Query.cpp
@@ -36,7 +36,7 @@ FRenderQueryRHIRef FD3D12DynamicRHI::RHICreateRenderQuery(ERenderQueryType Query
 	});
 }
 
-bool FD3D12DynamicRHI::RHIGetRenderQueryResult(FRenderQueryRHIParamRef QueryRHI, uint64& OutResult, bool bWait)
+bool FD3D12DynamicRHI::RHIGetRenderQueryResult(FRHIRenderQuery* QueryRHI, uint64& OutResult, bool bWait)
 {
 	check(IsInRenderingThread());
 	FD3D12Adapter& Adapter = GetAdapter();
diff --git a/Engine/Source/Runtime/D3D12RHI/Private/D3D12RHIPrivate.h b/Engine/Source/Runtime/D3D12RHI/Private/D3D12RHIPrivate.h
--- a/Engine/Source/Runtime/D3D12RHI/Private/D3D12RHIPrivate.h
+++ b/Engine/Source/Runtime/D3D12RHI/Private/D3D12RHIPrivate.h
@@ -341,7 +341,7 @@ public:
 	virtual void RHIReadSurfaceFloatData(FTextureRHIParamRef Texture, FIntRect Rect, TArray<FFloat16Color>& OutData, ECubeFace CubeFace, int32 ArrayIndex, int32 MipIndex) final override;
 	virtual void RHIRead3DSurfaceFloatData(FTextureRHIParamRef Texture, FIntRect Rect, FIntPoint ZMinMax, TArray<FFloat16Color>& OutData) final override;
 	virtual FRenderQueryRHIRef RHICreateRenderQuery(ERenderQueryType QueryType) final override;
-	virtual bool RHIGetRenderQueryResult(FRenderQueryRHIParamRef RenderQuery, uint64& OutResult, bool bWait) final override;
+	virtual bool RHIGetRenderQueryResult(FRHIRenderQuery* RenderQuery, uint64& OutResult, bool bWait) final override;
 	virtual uint32 RHIGetViewportNextPresentGPUIndex(FViewportRHIParamRef Viewport) final override;
 	virtual FTexture2DRHIRef RHIGetViewportBackBuffer(FViewportRHIParamRef Viewport) final override;
 	virtual void RHIAdvanceFrameForGetViewportBackBuffer(FViewportRHIParamRef Viewport) final override;
diff --git a/Engine/Source/Runtime/D3D12RHI/Private/D3D12RayTracing.cpp b/Engine/Source/Runtime/D3D12RHI/Private/D3D12RayTracing.cpp
--- a/Engine/Source/Runtime/D3D12RHI/Private/D3D12RayTracing.cpp
+++ b/Engine/Source/Runtime/D3D12RHI/Private/D3D12RayTracing.cpp
@@ -312,7 +312,7 @@ static TRefCountPtr<ID3D12StateObject> CreateRayTracingStateObject(
 	return Result;
 }
 
-inline uint64 GetShaderHash64(FRayTracingShaderRHIParamRef ShaderRHI)
+inline uint64 GetShaderHash64(FRHIRayTracingShader* ShaderRHI)
 {
 	uint64 ShaderHash; // 64 bits from the shader SHA1
 	FMemory::Memcpy(&ShaderHash, ShaderRHI->GetHash().Hash, sizeof(ShaderHash));
@@ -326,7 +326,7 @@ inline FString GenerateShaderName(const TCHAR* Prefix, uint64 Hash)
 	return FString::Printf(TEXT("%s_%016llx"), Prefix, Hash);
 }
 
-inline FString GenerateShaderName(FRayTracingShaderRHIParamRef ShaderRHI)
+inline FString GenerateShaderName(FRHIRayTracingShader* ShaderRHI)
 {
 	const FD3D12RayTracingShader* Shader = FD3D12DynamicRHI::ResourceCast(ShaderRHI);
 	uint64 ShaderHash = GetShaderHash64(ShaderRHI);
@@ -1290,22 +1290,22 @@ public:
 
 		// Use hit and miss shaders from initializer or fall back to default ones if none were provided
 
-		FRayTracingShaderRHIParamRef DefaultHitShader = GetBuildInRayTracingShader<FDefaultMainCHS>();
-		FRayTracingShaderRHIParamRef DefaultHitGroupTable[] = { DefaultHitShader };
+		FRHIRayTracingShader* DefaultHitShader = GetBuildInRayTracingShader<FDefaultMainCHS>();
+		FRHIRayTracingShader* DefaultHitGroupTable[] = { DefaultHitShader };
 
-		TArrayView<const FRayTracingShaderRHIParamRef> InitializerHitGroups = Initializer.GetHitGroupTable().Num()
+		TArrayView<FRHIRayTracingShader*> InitializerHitGroups = Initializer.GetHitGroupTable().Num()
 			? Initializer.GetHitGroupTable()
 			: DefaultHitGroupTable;
 
-		FRayTracingShaderRHIParamRef DefaultMissShader = GetBuildInRayTracingShader<FDefaultMainMS>();
-		FRayTracingShaderRHIParamRef DefaultMissTable[] = { DefaultMissShader };
+		FRHIRayTracingShader* DefaultMissShader = GetBuildInRayTracingShader<FDefaultMainMS>();
+		FRHIRayTracingShader* DefaultMissTable[] = { DefaultMissShader };
 
-		TArrayView<const FRayTracingShaderRHIParamRef> InitializerMissShaders = Initializer.GetMissTable().Num()
+		TArrayView<FRHIRayTracingShader*> InitializerMissShaders = Initializer.GetMissTable().Num()
 			? Initializer.GetMissTable()
 			: DefaultMissTable;
 
-		TArrayView<const FRayTracingShaderRHIParamRef> InitializerRayGenShaders = Initializer.GetRayGenTable();
-		TArrayView<const FRayTracingShaderRHIParamRef> InitializerCallableShaders = Initializer.GetCallableTable();
+		TArrayView<FRHIRayTracingShader*> InitializerRayGenShaders = Initializer.GetRayGenTable();
+		TArrayView<FRHIRayTracingShader*> InitializerCallableShaders = Initializer.GetCallableTable();
 
 		const uint32 MaxTotalShaders = InitializerRayGenShaders.Num() + InitializerMissShaders.Num() + InitializerHitGroups.Num() + InitializerCallableShaders.Num();
 
@@ -1364,7 +1364,7 @@ public:
 		RayGenShaders.Reserve(InitializerRayGenShaders.Num());
 		RayGenShaderNames.Reserve(InitializerRayGenShaders.Num());
 
-		for (FRayTracingShaderRHIParamRef ShaderRHI : InitializerRayGenShaders)
+		for (FRHIRayTracingShader* ShaderRHI : InitializerRayGenShaders)
 		{
 			FD3D12RayTracingShader* Shader = FD3D12DynamicRHI::ResourceCast(ShaderRHI);
 			checkf(Shader->pRootSignature->GetRootSignature() == GlobalRootSignature, TEXT("All raygen and miss shaders must share the same root signature"));
@@ -1382,7 +1382,7 @@ public:
 		MissShaders.Reserve(InitializerMissShaders.Num());
 		MissShaderNames.Reserve(InitializerMissShaders.Num());
 
-		for (FRayTracingShaderRHIParamRef ShaderRHI : InitializerMissShaders)
+		for (FRHIRayTracingShader* ShaderRHI : InitializerMissShaders)
 		{
 			FD3D12RayTracingShader* Shader = FD3D12DynamicRHI::ResourceCast(ShaderRHI);
 			checkf(Shader->pRootSignature->GetRootSignature() == GlobalRootSignature, TEXT("All raygen and miss shaders must share the same root signature"));
@@ -1403,7 +1403,7 @@ public:
 		HitGroupShaders.Reserve(InitializerHitGroups.Num());
 		HitGroupNames.Reserve(InitializerHitGroups.Num());
 
-		for (const FRayTracingShaderRHIParamRef ShaderRHI : InitializerHitGroups)
+		for (FRHIRayTracingShader* ShaderRHI : InitializerHitGroups)
 		{
 			FD3D12RayTracingShader* Shader = FD3D12DynamicRHI::ResourceCast(ShaderRHI);
 
@@ -1426,7 +1426,7 @@ public:
 		CallableShaders.Reserve(InitializerCallableShaders.Num());
 		CallableShaderNames.Reserve(InitializerCallableShaders.Num());
 
-		for (FRayTracingShaderRHIParamRef ShaderRHI : InitializerCallableShaders)
+		for (FRHIRayTracingShader* ShaderRHI : InitializerCallableShaders)
 		{
 			FD3D12RayTracingShader* Shader = FD3D12DynamicRHI::ResourceCast(ShaderRHI);
 
@@ -1583,10 +1583,10 @@ public:
 		{
 			FRayTracingPipelineStateInitializer OcclusionInitializer;
 
-			FRayTracingShaderRHIParamRef OcclusionRGSTable[] = { GetBuildInRayTracingShader<FOcclusionMainRG>() };
+			FRHIRayTracingShader* OcclusionRGSTable[] = { GetBuildInRayTracingShader<FOcclusionMainRG>() };
 			OcclusionInitializer.SetRayGenShaderTable(OcclusionRGSTable);
 
-			FRayTracingShaderRHIParamRef OcclusionMSTable[] = { GetBuildInRayTracingShader<FOcclusionMainMS>() };
+			FRHIRayTracingShader* OcclusionMSTable[] = { GetBuildInRayTracingShader<FOcclusionMainMS>() };
 			OcclusionInitializer.SetMissShaderTable(OcclusionMSTable);
 
 			OcclusionInitializer.bAllowHitGroupIndexing = false;
@@ -1598,13 +1598,13 @@ public:
 		{
 			FRayTracingPipelineStateInitializer IntersectionInitializer;
 
-			FRayTracingShaderRHIParamRef IntersectionRGSTable[] = { GetBuildInRayTracingShader<FIntersectionMainRG>() };
+			FRHIRayTracingShader* IntersectionRGSTable[] = { GetBuildInRayTracingShader<FIntersectionMainRG>() };
 			IntersectionInitializer.SetRayGenShaderTable(IntersectionRGSTable);
 
-			FRayTracingShaderRHIParamRef IntersectionMSTable[] = { GetBuildInRayTracingShader<FIntersectionMainMS>() };
+			FRHIRayTracingShader* IntersectionMSTable[] = { GetBuildInRayTracingShader<FIntersectionMainMS>() };
 			IntersectionInitializer.SetMissShaderTable(IntersectionMSTable);
 
-			FRayTracingShaderRHIParamRef IntersectionHitTable[] = { GetBuildInRayTracingShader<FIntersectionMainCHS>() };
+			FRHIRayTracingShader* IntersectionHitTable[] = { GetBuildInRayTracingShader<FIntersectionMainCHS>() };
 			IntersectionInitializer.SetHitGroupTable(IntersectionHitTable);
 
 			IntersectionInitializer.bAllowHitGroupIndexing = false;
@@ -2954,7 +2954,7 @@ void FD3D12CommandContext::RHIRayTraceIntersection(FRayTracingSceneRHIParamRef I
 	DispatchRays(*this, Bindings, Pipeline, 0, nullptr, DispatchDesc);
 }
 
-void FD3D12CommandContext::RHIRayTraceDispatch(FRHIRayTracingPipelineState* InRayTracingPipelineState, FRayTracingShaderRHIParamRef RayGenShaderRHI,
+void FD3D12CommandContext::RHIRayTraceDispatch(FRHIRayTracingPipelineState* InRayTracingPipelineState, FRHIRayTracingShader* RayGenShaderRHI,
 	FRayTracingSceneRHIParamRef InScene,
 	const FRayTracingShaderBindings& GlobalResourceBindings,
 	uint32 Width, uint32 Height)
diff --git a/Engine/Source/Runtime/EmptyRHI/Private/EmptyCommands.cpp b/Engine/Source/Runtime/EmptyRHI/Private/EmptyCommands.cpp
--- a/Engine/Source/Runtime/EmptyRHI/Private/EmptyCommands.cpp
+++ b/Engine/Source/Runtime/EmptyRHI/Private/EmptyCommands.cpp
@@ -262,14 +262,14 @@ void FEmptyDynamicRHI::RHISetRenderTargetsAndClear(const FRHISetRenderTargetsInf
 }
 
 // Occlusion/Timer queries.
-void FEmptyDynamicRHI::RHIBeginRenderQuery(FRenderQueryRHIParamRef QueryRHI)
+void FEmptyDynamicRHI::RHIBeginRenderQuery(FRHIRenderQuery* QueryRHI)
 {
 	FEmptyRenderQuery* Query = ResourceCast(QueryRHI);
 
 	Query->Begin();
 }
 
-void FEmptyDynamicRHI::RHIEndRenderQuery(FRenderQueryRHIParamRef QueryRHI)
+void FEmptyDynamicRHI::RHIEndRenderQuery(FRHIRenderQuery* QueryRHI)
 {
 	FEmptyRenderQuery* Query = ResourceCast(QueryRHI);
 
diff --git a/Engine/Source/Runtime/EmptyRHI/Private/EmptyQuery.cpp b/Engine/Source/Runtime/EmptyRHI/Private/EmptyQuery.cpp
--- a/Engine/Source/Runtime/EmptyRHI/Private/EmptyQuery.cpp
+++ b/Engine/Source/Runtime/EmptyRHI/Private/EmptyQuery.cpp
@@ -37,7 +37,7 @@ FRenderQueryRHIRef FEmptyDynamicRHI::RHICreateRenderQuery(ERenderQueryType Query
 	return new FEmptyRenderQuery(QueryType);
 }
 
-bool FEmptyDynamicRHI::RHIGetRenderQueryResult(FRenderQueryRHIParamRef QueryRHI,uint64& OutNumPixels,bool bWait)
+bool FEmptyDynamicRHI::RHIGetRenderQueryResult(FRHIRenderQuery* QueryRHI,uint64& OutNumPixels,bool bWait)
 {
 	check(IsInRenderingThread());
 
diff --git a/Engine/Source/Runtime/EmptyRHI/Public/EmptyRHI.h b/Engine/Source/Runtime/EmptyRHI/Public/EmptyRHI.h
--- a/Engine/Source/Runtime/EmptyRHI/Public/EmptyRHI.h
+++ b/Engine/Source/Runtime/EmptyRHI/Public/EmptyRHI.h
@@ -109,7 +109,7 @@ public:
 	virtual void RHIReadSurfaceFloatData(FTextureRHIParamRef Texture,FIntRect Rect,TArray<FFloat16Color>& OutData,ECubeFace CubeFace,int32 ArrayIndex,int32 MipIndex) final override;
 	virtual void RHIRead3DSurfaceFloatData(FTextureRHIParamRef Texture,FIntRect Rect,FIntPoint ZMinMax,TArray<FFloat16Color>& OutData) final override;
 	virtual FRenderQueryRHIRef RHICreateRenderQuery(ERenderQueryType QueryType) final override;
-	virtual bool RHIGetRenderQueryResult(FRenderQueryRHIParamRef RenderQuery, uint64& OutResult, bool bWait) final override;
+	virtual bool RHIGetRenderQueryResult(FRHIRenderQuery* RenderQuery, uint64& OutResult, bool bWait) final override;
 	virtual FTexture2DRHIRef RHIGetViewportBackBuffer(FViewportRHIParamRef Viewport) final override;
 	virtual void RHIAdvanceFrameForGetViewportBackBuffer(FViewportRHIParamRef Viewport) final override;
 	virtual void RHIAcquireThreadOwnership() final override;
@@ -141,8 +141,8 @@ public:
 	virtual void RHIClearTinyUAV(FRHIUnorderedAccessView* UnorderedAccessViewRHI, const uint32* Values) final override;
 	virtual void RHICopyToResolveTarget(FTextureRHIParamRef SourceTexture, FTextureRHIParamRef DestTexture, const FResolveParams& ResolveParams) final override;
 	virtual void RHICopyTexture(FTextureRHIParamRef SourceTexture, FTextureRHIParamRef DestTexture, const FRHICopyTextureInfo& CopyInfo) final override;
-	virtual void RHIBeginRenderQuery(FRenderQueryRHIParamRef RenderQuery) final override;
-	virtual void RHIEndRenderQuery(FRenderQueryRHIParamRef RenderQuery) final override;
+	virtual void RHIBeginRenderQuery(FRHIRenderQuery* RenderQuery) final override;
+	virtual void RHIEndRenderQuery(FRHIRenderQuery* RenderQuery) final override;
 	virtual void RHISubmitCommandsHint() final override;
 	virtual void RHIBeginDrawingViewport(FViewportRHIParamRef Viewport, FTextureRHIParamRef RenderTargetRHI) final override;
 	virtual void RHIEndDrawingViewport(FViewportRHIParamRef Viewport, bool bPresent, bool bLockToVsync) final override;
diff --git a/Engine/Source/Runtime/NullDrv/Public/NullRHI.h b/Engine/Source/Runtime/NullDrv/Public/NullRHI.h
--- a/Engine/Source/Runtime/NullDrv/Public/NullRHI.h
+++ b/Engine/Source/Runtime/NullDrv/Public/NullRHI.h
@@ -422,17 +422,17 @@ public:
 		return new FRHIRenderQuery(); 
 	}
 
-	virtual void RHIBeginRenderQuery(FRenderQueryRHIParamRef RenderQuery) final override
+	virtual void RHIBeginRenderQuery(FRHIRenderQuery* RenderQuery) final override
 	{
 
 	}
-	virtual void RHIEndRenderQuery(FRenderQueryRHIParamRef RenderQuery) final override
+	virtual void RHIEndRenderQuery(FRHIRenderQuery* RenderQuery) final override
 	{
 
 	}
 
 
-	virtual bool RHIGetRenderQueryResult(FRenderQueryRHIParamRef RenderQuery, uint64& OutResult, bool bWait) final override
+	virtual bool RHIGetRenderQueryResult(FRHIRenderQuery* RenderQuery, uint64& OutResult, bool bWait) final override
 	{ 
 		return true; 
 	}
diff --git a/Engine/Source/Runtime/OpenGLDrv/Private/OpenGLCommands.cpp b/Engine/Source/Runtime/OpenGLDrv/Private/OpenGLCommands.cpp
--- a/Engine/Source/Runtime/OpenGLDrv/Private/OpenGLCommands.cpp
+++ b/Engine/Source/Runtime/OpenGLDrv/Private/OpenGLCommands.cpp
@@ -3350,7 +3350,7 @@ void FOpenGLDynamicRHI::RHICopyToStagingBuffer(FVertexBufferRHIParamRef SourceBu
 #endif
 }
 
-void FOpenGLDynamicRHI::RHIWriteGPUFence(FGPUFenceRHIParamRef FenceRHI)
+void FOpenGLDynamicRHI::RHIWriteGPUFence(FRHIGPUFence* FenceRHI)
 {
 	check(FenceRHI);
 	FOpenGLGPUFence* CopyFence = ResourceCast(FenceRHI);
diff --git a/Engine/Source/Runtime/OpenGLDrv/Private/OpenGLQuery.cpp b/Engine/Source/Runtime/OpenGLDrv/Private/OpenGLQuery.cpp
--- a/Engine/Source/Runtime/OpenGLDrv/Private/OpenGLQuery.cpp
+++ b/Engine/Source/Runtime/OpenGLDrv/Private/OpenGLQuery.cpp
@@ -20,10 +20,10 @@ static FAutoConsoleVariableRef CVarOpenGLPollRenderQueryResult(
 
 struct FQueryItem
 {
-	FRenderQueryRHIParamRef Query;
+	FRHIRenderQuery* Query;
 	int32 BeginSequence;
 
-	FQueryItem(FRenderQueryRHIParamRef InQueryRHI)
+	FQueryItem(FRHIRenderQuery* InQueryRHI)
 		: Query(InQueryRHI)
 	{
 		FOpenGLRenderQuery* InQuery = FOpenGLDynamicRHI::ResourceCast(InQueryRHI);
@@ -57,7 +57,7 @@ struct FGLQueryBatcher
 	{
 	}
 
-	void Add(FRenderQueryRHIParamRef Query)
+	void Add(FRHIRenderQuery* Query)
 	{
 		if (NewBatch && NewBatch->FrameNumberRenderThread)
 		{
@@ -72,7 +72,7 @@ struct FGLQueryBatcher
 			Batch->bHasFlushedSinceLastWait = false;
 		}
 	}
-	void Flush(FOpenGLDynamicRHI& RHI, FRenderQueryRHIParamRef TargetQueryRHI)
+	void Flush(FOpenGLDynamicRHI& RHI, FRHIRenderQuery* TargetQueryRHI)
 	{
 		QUICK_SCOPE_CYCLE_COUNTER(STAT_FGLQueryBatcher_FlushScan);
 		bool bFoundQuery = false;
@@ -88,7 +88,7 @@ struct FGLQueryBatcher
 			for (int32 IndexInner = 0; IndexInner < Batch->BatchContents.Num(); IndexInner++)
 			{
 				FQueryItem& Item = Batch->BatchContents[IndexInner];
-				FRenderQueryRHIParamRef QueryRHI = Item.Query;
+				FRHIRenderQuery* QueryRHI = Item.Query;
 				FOpenGLRenderQuery* Query = FOpenGLDynamicRHI::ResourceCast(QueryRHI);
 				if (TargetQueryRHI == QueryRHI)
 				{
@@ -147,7 +147,7 @@ struct FGLQueryBatcher
 			for (int32 IndexInner = 0; IndexInner < Batch->BatchContents.Num(); IndexInner++)
 			{
 				FQueryItem& Item = Batch->BatchContents[IndexInner];
-				FRenderQueryRHIParamRef QueryRHI = Item.Query;
+				FRHIRenderQuery* QueryRHI = Item.Query;
 				FOpenGLRenderQuery* Query = FOpenGLDynamicRHI::ResourceCast(QueryRHI);
 
 				int32 Begins = Query->TotalBegins.GetValue();
@@ -253,7 +253,7 @@ FRenderQueryRHIRef FOpenGLDynamicRHI::RHICreateRenderQuery(ERenderQueryType Quer
 	return new FOpenGLRenderQuery(QueryType);
 }
 
-void FOpenGLDynamicRHI::RHIBeginRenderQuery(FRenderQueryRHIParamRef QueryRHI)
+void FOpenGLDynamicRHI::RHIBeginRenderQuery(FRHIRenderQuery* QueryRHI)
 {
 	VERIFY_GL_SCOPE();
 
@@ -266,7 +266,7 @@ void FOpenGLDynamicRHI::RHIBeginRenderQuery(FRenderQueryRHIParamRef QueryRHI)
 	}
 }
 
-void FOpenGLDynamicRHI::RHIEndRenderQuery(FRenderQueryRHIParamRef QueryRHI)
+void FOpenGLDynamicRHI::RHIEndRenderQuery(FRHIRenderQuery* QueryRHI)
 {
 	VERIFY_GL_SCOPE();
 
@@ -527,7 +527,7 @@ public:
 };
 
 
-bool FOpenGLDynamicRHI::RHIGetRenderQueryResult(FRenderQueryRHIParamRef QueryRHI,uint64& OutResult,bool bWait)
+bool FOpenGLDynamicRHI::RHIGetRenderQueryResult(FRHIRenderQuery* QueryRHI,uint64& OutResult,bool bWait)
 {
 	check(IsInRenderingThread() || IsInRHIThread());
 
diff --git a/Engine/Source/Runtime/OpenGLDrv/Public/OpenGLDrv.h b/Engine/Source/Runtime/OpenGLDrv/Public/OpenGLDrv.h
--- a/Engine/Source/Runtime/OpenGLDrv/Public/OpenGLDrv.h
+++ b/Engine/Source/Runtime/OpenGLDrv/Public/OpenGLDrv.h
@@ -419,7 +419,7 @@ public:
 	virtual void RHIReadSurfaceFloatData(FTextureRHIParamRef Texture,FIntRect Rect,TArray<FFloat16Color>& OutData,ECubeFace CubeFace,int32 ArrayIndex,int32 MipIndex) final override;
 	virtual void RHIRead3DSurfaceFloatData(FTextureRHIParamRef Texture,FIntRect Rect,FIntPoint ZMinMax,TArray<FFloat16Color>& OutData) final override;
 	virtual FRenderQueryRHIRef RHICreateRenderQuery(ERenderQueryType QueryType) final override;
-	virtual bool RHIGetRenderQueryResult(FRenderQueryRHIParamRef RenderQuery, uint64& OutResult, bool bWait) final override;
+	virtual bool RHIGetRenderQueryResult(FRHIRenderQuery* RenderQuery, uint64& OutResult, bool bWait) final override;
 	virtual FTexture2DRHIRef RHIGetViewportBackBuffer(FViewportRHIParamRef Viewport) final override;
 	virtual void RHIAdvanceFrameForGetViewportBackBuffer(FViewportRHIParamRef Viewport) final override;
 	virtual void RHIAcquireThreadOwnership() final override;
@@ -453,8 +453,8 @@ public:
 	virtual void RHISetMultipleViewports(uint32 Count, const FViewportBounds* Data) final override;
 	virtual void RHIClearTinyUAV(FRHIUnorderedAccessView* UnorderedAccessViewRHI, const uint32* Values) final override;
 	virtual void RHICopyToResolveTarget(FTextureRHIParamRef SourceTexture, FTextureRHIParamRef DestTexture, const FResolveParams& ResolveParams) final override;
-	virtual void RHIBeginRenderQuery(FRenderQueryRHIParamRef RenderQuery) final override;
-	virtual void RHIEndRenderQuery(FRenderQueryRHIParamRef RenderQuery) final override;
+	virtual void RHIBeginRenderQuery(FRHIRenderQuery* RenderQuery) final override;
+	virtual void RHIEndRenderQuery(FRHIRenderQuery* RenderQuery) final override;
 	virtual void RHISubmitCommandsHint() final override;
 	virtual void RHIBeginDrawingViewport(FViewportRHIParamRef Viewport, FTextureRHIParamRef RenderTargetRHI) final override;
 	virtual void RHIEndDrawingViewport(FViewportRHIParamRef Viewport, bool bPresent, bool bLockToVsync) final override;
@@ -531,7 +531,7 @@ public:
 
 	// Inline copy
 	virtual void RHICopyToStagingBuffer(FVertexBufferRHIParamRef SourceBufferRHI, FStagingBufferRHIParamRef DestinationStagingBufferRHI, uint32 InOffset, uint32 InNumBytes) final override;
-	virtual void RHIWriteGPUFence(FGPUFenceRHIParamRef FenceRHI) final override;
+	virtual void RHIWriteGPUFence(FRHIGPUFence* FenceRHI) final override;
 	virtual void* RHILockStagingBuffer(FStagingBufferRHIParamRef StagingBuffer, uint32 Offset, uint32 SizeRHI) final override;
 	virtual void RHIUnlockStagingBuffer(FStagingBufferRHIParamRef StagingBuffer) final override;
 	virtual void* LockStagingBuffer_RenderThread(class FRHICommandListImmediate& RHICmdList, FStagingBufferRHIParamRef StagingBuffer, uint32 Offset, uint32 SizeRHI) final override;
diff --git a/Engine/Source/Runtime/RHI/Private/PipelineStateCache.cpp b/Engine/Source/Runtime/RHI/Private/PipelineStateCache.cpp
--- a/Engine/Source/Runtime/RHI/Private/PipelineStateCache.cpp
+++ b/Engine/Source/Runtime/RHI/Private/PipelineStateCache.cpp
@@ -963,33 +963,33 @@ public:
 
 private:
 
-	void AddRefShaders(TArray<FRayTracingShaderRHIParamRef>& ShaderTable)
+	void AddRefShaders(TArray<FRHIRayTracingShader*>& ShaderTable)
 	{
-		for (FRayTracingShaderRHIParamRef Ptr : ShaderTable)
+		for (FRHIRayTracingShader* Ptr : ShaderTable)
 		{
 			Ptr->AddRef();
 		}
 	}
 
-	void ReleaseShaders(TArray<FRayTracingShaderRHIParamRef>& ShaderTable)
+	void ReleaseShaders(TArray<FRHIRayTracingShader*>& ShaderTable)
 	{
-		for (FRayTracingShaderRHIParamRef Ptr : ShaderTable)
+		for (FRHIRayTracingShader* Ptr : ShaderTable)
 		{
 			Ptr->Release();
 		}
 	}
 
-	TArray<FRayTracingShaderRHIParamRef> CopyShaderTable(const TArrayView<const FRayTracingShaderRHIParamRef>& Source)
+	TArray<FRHIRayTracingShader*> CopyShaderTable(const TArrayView<FRHIRayTracingShader*>& Source)
 	{
-		TArray<FRayTracingShaderRHIParamRef> Result(Source.GetData(), Source.Num());
+		TArray<FRHIRayTracingShader*> Result(Source.GetData(), Source.Num());
 		AddRefShaders(Result);
 		return Result;
 	}
 
-	TArray<FRayTracingShaderRHIParamRef> RayGenTable;
-	TArray<FRayTracingShaderRHIParamRef> MissTable;
-	TArray<FRayTracingShaderRHIParamRef> HitGroupTable;
-	TArray<FRayTracingShaderRHIParamRef> CallableTable;
+	TArray<FRHIRayTracingShader*> RayGenTable;
+	TArray<FRHIRayTracingShader*> MissTable;
+	TArray<FRHIRayTracingShader*> HitGroupTable;
+	TArray<FRHIRayTracingShader*> CallableTable;
 };
 
 FRayTracingPipelineState* PipelineStateCache::GetAndOrCreateRayTracingPipelineState(FRHICommandList& RHICmdList, const FRayTracingPipelineStateInitializer& Initializer)
diff --git a/Engine/Source/Runtime/RHI/Public/DynamicRHI.h b/Engine/Source/Runtime/RHI/Public/DynamicRHI.h
--- a/Engine/Source/Runtime/RHI/Public/DynamicRHI.h
+++ b/Engine/Source/Runtime/RHI/Public/DynamicRHI.h
@@ -913,7 +913,7 @@ public:
 
 	// CAUTION: Even though this is marked as threadsafe, it is only valid to call from the render thread. It is need not be threadsafe on platforms that do not support or aren't using an RHIThread
 	// FlushType: Thread safe, but varies by RHI
-	virtual bool RHIGetRenderQueryResult(FRenderQueryRHIParamRef RenderQuery, uint64& OutResult, bool bWait) = 0;
+	virtual bool RHIGetRenderQueryResult(FRHIRenderQuery* RenderQuery, uint64& OutResult, bool bWait) = 0;
 
 	// FlushType: Thread safe
 	virtual uint32 RHIGetViewportNextPresentGPUIndex(FViewportRHIParamRef Viewport)
@@ -1320,7 +1320,7 @@ FORCEINLINE void RHIBindDebugLabelName(FRHIUnorderedAccessView* UnorderedAccessV
 	GDynamicRHI->RHIBindDebugLabelName(UnorderedAccessViewRHI, Name);
 }
 
-FORCEINLINE bool RHIGetRenderQueryResult(FRenderQueryRHIParamRef RenderQuery, uint64& OutResult, bool bWait)
+FORCEINLINE bool RHIGetRenderQueryResult(FRHIRenderQuery* RenderQuery, uint64& OutResult, bool bWait)
 {
 	return GDynamicRHI->RHIGetRenderQueryResult(RenderQuery, OutResult, bWait);
 }
diff --git a/Engine/Source/Runtime/RHI/Public/RHI.h b/Engine/Source/Runtime/RHI/Public/RHI.h
--- a/Engine/Source/Runtime/RHI/Public/RHI.h
+++ b/Engine/Source/Runtime/RHI/Public/RHI.h
@@ -1628,7 +1628,7 @@ extern RHI_API void RHIExit();
 #define GETSAFERHISHADER_DOMAIN(Shader) ((Shader) ? (Shader)->GetDomainShader() : (FDomainShaderRHIParamRef)FDomainShaderRHIRef())
 #define GETSAFERHISHADER_GEOMETRY(Shader) ((Shader) ? (Shader)->GetGeometryShader() : (FGeometryShaderRHIParamRef)FGeometryShaderRHIRef())
 #define GETSAFERHISHADER_COMPUTE(Shader) ((Shader) ? (Shader)->GetComputeShader() : (FComputeShaderRHIParamRef)FComputeShaderRHIRef())
-#define GETSAFERHISHADER_RAYTRACING(Shader) ((Shader) ? (Shader)->GetRayTracingShader() : (FRayTracingShaderRHIParamRef)FRayTracingShaderRHIRef())
+#define GETSAFERHISHADER_RAYTRACING(Shader) ((Shader) ? (Shader)->GetRayTracingShader() : (FRHIRayTracingShader*)FRayTracingShaderRHIRef())
 
 
 // Panic delegate is called when when a fatal condition is encountered within RHI function.
diff --git a/Engine/Source/Runtime/RHI/Public/RHICommandList.h b/Engine/Source/Runtime/RHI/Public/RHICommandList.h
--- a/Engine/Source/Runtime/RHI/Public/RHICommandList.h
+++ b/Engine/Source/Runtime/RHI/Public/RHICommandList.h
@@ -1430,9 +1430,9 @@ struct FRHICommandCopyToStagingBuffer final : public FRHICommand<FRHICommandCopy
 template<ECmdList CmdListType>
 struct FRHICommandWriteGPUFence final : public FRHICommand<FRHICommandWriteGPUFence<CmdListType>>
 {
-	FGPUFenceRHIParamRef Fence;
+	FRHIGPUFence* Fence;
 
-	FORCEINLINE_DEBUGGABLE FRHICommandWriteGPUFence(FGPUFenceRHIParamRef InFence)
+	FORCEINLINE_DEBUGGABLE FRHICommandWriteGPUFence(FRHIGPUFence* InFence)
 		: Fence(InFence)
 	{
 	}
@@ -1599,9 +1599,9 @@ struct FRHICommandSetLocalUniformBuffer final : public FRHICommand<FRHICommandSe
 
 struct FRHICommandBeginRenderQuery final : public FRHICommand<FRHICommandBeginRenderQuery>
 {
-	FRenderQueryRHIParamRef RenderQuery;
+	FRHIRenderQuery* RenderQuery;
 
-	FORCEINLINE_DEBUGGABLE FRHICommandBeginRenderQuery(FRenderQueryRHIParamRef InRenderQuery)
+	FORCEINLINE_DEBUGGABLE FRHICommandBeginRenderQuery(FRHIRenderQuery* InRenderQuery)
 		: RenderQuery(InRenderQuery)
 	{
 	}
@@ -1610,9 +1610,9 @@ struct FRHICommandBeginRenderQuery final : public FRHICommand<FRHICommandBeginRe
 
 struct FRHICommandEndRenderQuery final : public FRHICommand<FRHICommandEndRenderQuery>
 {
-	FRenderQueryRHIParamRef RenderQuery;
+	FRHIRenderQuery* RenderQuery;
 
-	FORCEINLINE_DEBUGGABLE FRHICommandEndRenderQuery(FRenderQueryRHIParamRef InRenderQuery)
+	FORCEINLINE_DEBUGGABLE FRHICommandEndRenderQuery(FRHIRenderQuery* InRenderQuery)
 		: RenderQuery(InRenderQuery)
 	{
 	}
@@ -1947,11 +1947,11 @@ struct FRHICommandRayTraceDispatch final : public FRHICommand<FRHICommandRayTrac
 	FRayTracingPipelineState* Pipeline;
 	FRayTracingSceneRHIParamRef Scene;
 	FRayTracingShaderBindings GlobalResourceBindings;
-	FRayTracingShaderRHIParamRef RayGenShader;
+	FRHIRayTracingShader* RayGenShader;
 	uint32 Width;
 	uint32 Height;
 
-	FRHICommandRayTraceDispatch(FRayTracingPipelineState* InPipeline, FRayTracingShaderRHIParamRef InRayGenShader, FRayTracingSceneRHIParamRef InScene, const FRayTracingShaderBindings& InGlobalResourceBindings, uint32 InWidth, uint32 InHeight)
+	FRHICommandRayTraceDispatch(FRayTracingPipelineState* InPipeline, FRHIRayTracingShader* InRayGenShader, FRayTracingSceneRHIParamRef InScene, const FRayTracingShaderBindings& InGlobalResourceBindings, uint32 InWidth, uint32 InHeight)
 		: Pipeline(InPipeline)
 		, Scene(InScene)
 		, GlobalResourceBindings(InGlobalResourceBindings)
@@ -2641,7 +2641,7 @@ public:
 		ALLOC_COMMAND(FRHICommandClearTinyUAV)(UnorderedAccessViewRHI, Values);
 	}
 
-	FORCEINLINE_DEBUGGABLE void BeginRenderQuery(FRenderQueryRHIParamRef RenderQuery)
+	FORCEINLINE_DEBUGGABLE void BeginRenderQuery(FRHIRenderQuery* RenderQuery)
 	{
 		if (Bypass())
 		{
@@ -2650,7 +2650,7 @@ public:
 		}
 		ALLOC_COMMAND(FRHICommandBeginRenderQuery)(RenderQuery);
 	}
-	FORCEINLINE_DEBUGGABLE void EndRenderQuery(FRenderQueryRHIParamRef RenderQuery)
+	FORCEINLINE_DEBUGGABLE void EndRenderQuery(FRHIRenderQuery* RenderQuery)
 	{
 		if (Bypass())
 		{
@@ -2788,7 +2788,7 @@ public:
 		ALLOC_COMMAND(FRHICommandCopyToStagingBuffer<ECmdList::EGfx>)(SourceBuffer, DestinationStagingBuffer, Offset, NumBytes);
 	}
 
-	FORCEINLINE_DEBUGGABLE void WriteGPUFence(FGPUFenceRHIParamRef Fence)
+	FORCEINLINE_DEBUGGABLE void WriteGPUFence(FRHIGPUFence* Fence)
 	{
 		if (Bypass())
 		{
@@ -3072,7 +3072,7 @@ public:
 		}
 	}
 
-	FORCEINLINE_DEBUGGABLE void RayTraceDispatch(FRayTracingPipelineState* Pipeline, FRayTracingShaderRHIParamRef RayGenShader, FRayTracingSceneRHIParamRef Scene, const FRayTracingShaderBindings& GlobalResourceBindings, uint32 Width, uint32 Height)
+	FORCEINLINE_DEBUGGABLE void RayTraceDispatch(FRayTracingPipelineState* Pipeline, FRHIRayTracingShader* RayGenShader, FRayTracingSceneRHIParamRef Scene, const FRayTracingShaderBindings& GlobalResourceBindings, uint32 Width, uint32 Height)
 	{
 		if (Bypass())
 		{
@@ -3417,7 +3417,7 @@ public:
 		ALLOC_COMMAND(FRHICommandCopyToStagingBuffer<ECmdList::ECompute>)(SourceBuffer, DestinationStagingBuffer, Offset, NumBytes);
 	}
 
-	FORCEINLINE_DEBUGGABLE void WriteGPUFence(FGPUFenceRHIParamRef Fence)
+	FORCEINLINE_DEBUGGABLE void WriteGPUFence(FRHIGPUFence* Fence)
 	{
 		if (Bypass())
 		{
@@ -4162,7 +4162,7 @@ public:
 		}
 	}
 
-	FORCEINLINE bool GetRenderQueryResult(FRenderQueryRHIParamRef RenderQuery, uint64& OutResult, bool bWait)
+	FORCEINLINE bool GetRenderQueryResult(FRHIRenderQuery* RenderQuery, uint64& OutResult, bool bWait)
 	{
 		return RHIGetRenderQueryResult(RenderQuery, OutResult, bWait);
 	}
diff --git a/Engine/Source/Runtime/RHI/Public/RHIContext.h b/Engine/Source/Runtime/RHI/Public/RHIContext.h
--- a/Engine/Source/Runtime/RHI/Public/RHIContext.h
+++ b/Engine/Source/Runtime/RHI/Public/RHIContext.h
@@ -135,7 +135,7 @@ public:
 	 * Write the fence in the GPU timeline. The fence can then be tested on the CPU to know if the previous GPU commands are completed.
 	 * @param Fence 
 	 */
-	virtual void RHIWriteGPUFence(FGPUFenceRHIParamRef FenceRHI)
+	virtual void RHIWriteGPUFence(FRHIGPUFence* FenceRHI)
 	{
 		check(false);
 	}
@@ -245,9 +245,9 @@ public:
 		RHITransitionResources(TransitionType, TransitionPipeline, InUAVs, NumUAVs, nullptr);
 	}
 
-	virtual void RHIBeginRenderQuery(FRenderQueryRHIParamRef RenderQuery) = 0;
+	virtual void RHIBeginRenderQuery(FRHIRenderQuery* RenderQuery) = 0;
 
-	virtual void RHIEndRenderQuery(FRenderQueryRHIParamRef RenderQuery) = 0;
+	virtual void RHIEndRenderQuery(FRHIRenderQuery* RenderQuery) = 0;
 
 	// Used for OpenGL to check and see if any occlusion queries can be read back on the RHI thread. If they aren't ready when we need them, then we end up stalling.
 	virtual void RHIPollOcclusionQueries()
@@ -615,7 +615,7 @@ public:
 		checkNoEntry();
 	}
 
-	virtual void RHIRayTraceDispatch(FRHIRayTracingPipelineState* RayTracingPipelineState, FRayTracingShaderRHIParamRef RayGenShader,
+	virtual void RHIRayTraceDispatch(FRHIRayTracingPipelineState* RayTracingPipelineState, FRHIRayTracingShader* RayGenShader,
 		FRayTracingSceneRHIParamRef Scene, 
 		const FRayTracingShaderBindings& GlobalResourceBindings,
 		uint32 Width, uint32 Height)
diff --git a/Engine/Source/Runtime/RHI/Public/RHIResources.h b/Engine/Source/Runtime/RHI/Public/RHIResources.h
--- a/Engine/Source/Runtime/RHI/Public/RHIResources.h
+++ b/Engine/Source/Runtime/RHI/Public/RHIResources.h
@@ -1135,6 +1135,7 @@ typedef TRefCountPtr<FRHIGeometryShader> FGeometryShaderRHIRef;
 typedef FRHIComputeShader*              FComputeShaderRHIParamRef;
 typedef TRefCountPtr<FRHIComputeShader> FComputeShaderRHIRef;
 
+UE_DEPRECATED(4.23, "FRayTracingShaderRHIParamRef typedef is deprecated; please use FRHIRayTracingShader* directly instead.")
 typedef FRHIRayTracingShader*                       FRayTracingShaderRHIParamRef;
 typedef TRefCountPtr<FRHIRayTracingShader>          FRayTracingShaderRHIRef;
 
@@ -1177,12 +1178,13 @@ UE_DEPRECATED(4.23, "FTextureReferenceRHIParamRef typedef is deprecated; please
 typedef FRHITextureReference*              FTextureReferenceRHIParamRef;
 typedef TRefCountPtr<FRHITextureReference> FTextureReferenceRHIRef;
 
+UE_DEPRECATED(4.23, "FRenderQueryRHIParamRef typedef is deprecated; please use FRHIRenderQuery* directly instead.")
 typedef FRHIRenderQuery*              FRenderQueryRHIParamRef;
 typedef TRefCountPtr<FRHIRenderQuery> FRenderQueryRHIRef;
 
-typedef FRHIRenderQueryPool*              FRenderQueryPoolRHIParamRef;
 typedef TRefCountPtr<FRHIRenderQueryPool> FRenderQueryPoolRHIRef;
 
+UE_DEPRECATED(4.23, "FGPUFenceRHIParamRef typedef is deprecated; please use FRHIGPUFence* directly instead.")
 typedef FRHIGPUFence*				FGPUFenceRHIParamRef;
 typedef TRefCountPtr<FRHIGPUFence>	FGPUFenceRHIRef;
 
@@ -2201,13 +2203,13 @@ public:
 
 	bool bAllowHitGroupIndexing = true;
 
-	const TArrayView<const FRayTracingShaderRHIParamRef>& GetRayGenTable()   const { return RayGenTable; }
-	const TArrayView<const FRayTracingShaderRHIParamRef>& GetMissTable()     const { return MissTable; }
-	const TArrayView<const FRayTracingShaderRHIParamRef>& GetHitGroupTable() const { return HitGroupTable; }
-	const TArrayView<const FRayTracingShaderRHIParamRef>& GetCallableTable() const { return CallableTable; }
+	const TArrayView<FRHIRayTracingShader*>& GetRayGenTable()   const { return RayGenTable; }
+	const TArrayView<FRHIRayTracingShader*>& GetMissTable()     const { return MissTable; }
+	const TArrayView<FRHIRayTracingShader*>& GetHitGroupTable() const { return HitGroupTable; }
+	const TArrayView<FRHIRayTracingShader*>& GetCallableTable() const { return CallableTable; }
 
 	// Shaders used as entry point to ray tracing work. At least one RayGen shader must be provided.
-	void SetRayGenShaderTable(const TArrayView<const FRayTracingShaderRHIParamRef>& InRayGenShaders, uint64 Hash = 0)
+	void SetRayGenShaderTable(const TArrayView<FRHIRayTracingShader*>& InRayGenShaders, uint64 Hash = 0)
 	{
 		RayGenTable = InRayGenShaders;
 		RayGenHash = Hash ? Hash : ComputeShaderTableHash(InRayGenShaders);
@@ -2216,7 +2218,7 @@ public:
 	// Shaders that will be invoked if a ray misses all geometry.
 	// If this table is empty, then a built-in default miss shader will be used that sets HitT member of FDefaultPayload to -1.
 	// Desired miss shader can be selected by providing MissShaderIndex to TraceRay() function.
-	void SetMissShaderTable(const TArrayView<const FRayTracingShaderRHIParamRef>& InMissShaders, uint64 Hash = 0)
+	void SetMissShaderTable(const TArrayView<FRHIRayTracingShader*>& InMissShaders, uint64 Hash = 0)
 	{
 		MissTable = InMissShaders;
 		MissHash = Hash ? Hash : ComputeShaderTableHash(InMissShaders);
@@ -2224,7 +2226,7 @@ public:
 
 	// Shaders that will be invoked when ray intersects geometry.
 	// If this table is empty, then a built-in default shader will be used for all geometry, using FDefaultPayload.
-	void SetHitGroupTable(const TArrayView<const FRayTracingShaderRHIParamRef>& InHitGroups, uint64 Hash = 0)
+	void SetHitGroupTable(const TArrayView<FRHIRayTracingShader*>& InHitGroups, uint64 Hash = 0)
 	{
 		HitGroupTable = InHitGroups;
 		HitGroupHash = Hash ? Hash : ComputeShaderTableHash(HitGroupTable);
@@ -2232,7 +2234,7 @@ public:
 
 	// Shaders that can be explicitly invoked from RayGen shaders by their Shader Binding Table (SBT) index.
 	// SetRayTracingCallableShader() command must be used to fill SBT slots before a shader can be called.
-	void SetCallableTable(const TArrayView<const FRayTracingShaderRHIParamRef>& InCallableShaders, uint64 Hash = 0)
+	void SetCallableTable(const TArrayView<FRHIRayTracingShader*>& InCallableShaders, uint64 Hash = 0)
 	{
 		CallableTable = InCallableShaders;
 		CallableHash = Hash ? Hash : ComputeShaderTableHash(CallableTable);
@@ -2245,10 +2247,10 @@ public:
 
 private:
 
-	uint64 ComputeShaderTableHash(const TArrayView<const FRayTracingShaderRHIParamRef>& ShaderTable, uint64 InitialHash = 5699878132332235837ull)
+	uint64 ComputeShaderTableHash(const TArrayView<FRHIRayTracingShader*>& ShaderTable, uint64 InitialHash = 5699878132332235837ull)
 	{
 		uint64 CombinedHash = InitialHash;
-		for (FRayTracingShaderRHIParamRef ShaderRHI : ShaderTable)
+		for (FRHIRayTracingShader* ShaderRHI : ShaderTable)
 		{
 			uint64 ShaderHash; // 64 bits from the shader SHA1
 			FMemory::Memcpy(&ShaderHash, ShaderRHI->GetHash().Hash, sizeof(ShaderHash));
@@ -2260,10 +2262,10 @@ private:
 		return CombinedHash;
 	}
 
-	TArrayView<const FRayTracingShaderRHIParamRef> RayGenTable;
-	TArrayView<const FRayTracingShaderRHIParamRef> MissTable;
-	TArrayView<const FRayTracingShaderRHIParamRef> HitGroupTable;
-	TArrayView<const FRayTracingShaderRHIParamRef> CallableTable;
+	TArrayView<FRHIRayTracingShader*> RayGenTable;
+	TArrayView<FRHIRayTracingShader*> MissTable;
+	TArrayView<FRHIRayTracingShader*> HitGroupTable;
+	TArrayView<FRHIRayTracingShader*> CallableTable;
 
 	uint64 RayGenHash = 0;
 	uint64 MissHash = 0;
diff --git a/Engine/Source/Runtime/RHI/Public/RHIValidationContext.h b/Engine/Source/Runtime/RHI/Public/RHIValidationContext.h
--- a/Engine/Source/Runtime/RHI/Public/RHIValidationContext.h
+++ b/Engine/Source/Runtime/RHI/Public/RHIValidationContext.h
@@ -147,7 +147,7 @@ public:
 		RHIContext->RHIPopEvent();
 	}
 
-	virtual void RHIWriteGPUFence(FGPUFenceRHIParamRef FenceRHI) override final
+	virtual void RHIWriteGPUFence(FRHIGPUFence* FenceRHI) override final
 	{
 		RHIContext->RHIWriteGPUFence(FenceRHI);
 	}
@@ -278,12 +278,12 @@ public:
 		RHIContext->RHITransitionResources(TransitionType, TransitionPipeline, InUAVs, NumUAVs, WriteComputeFence);
 	}
 
-	virtual void RHIBeginRenderQuery(FRenderQueryRHIParamRef RenderQuery) override final
+	virtual void RHIBeginRenderQuery(FRHIRenderQuery* RenderQuery) override final
 	{
 		RHIContext->RHIBeginRenderQuery(RenderQuery);
 	}
 
-	virtual void RHIEndRenderQuery(FRenderQueryRHIParamRef RenderQuery) override final
+	virtual void RHIEndRenderQuery(FRHIRenderQuery* RenderQuery) override final
 	{
 		RHIContext->RHIEndRenderQuery(RenderQuery);
 	}
@@ -787,7 +787,7 @@ public:
 		State.bInsideComputePass = false;
 	}
 
-	virtual void RHIWriteGPUFence(FGPUFenceRHIParamRef FenceRHI) override final
+	virtual void RHIWriteGPUFence(FRHIGPUFence* FenceRHI) override final
 	{
 		RHIContext->RHIWriteGPUFence(FenceRHI);
 	}
diff --git a/Engine/Source/Runtime/RenderCore/Private/Shader.cpp b/Engine/Source/Runtime/RenderCore/Private/Shader.cpp
--- a/Engine/Source/Runtime/RenderCore/Private/Shader.cpp
+++ b/Engine/Source/Runtime/RenderCore/Private/Shader.cpp
@@ -398,7 +398,7 @@ TArray<uint32> FShaderResource::GlobalUnusedIndicies;
 TArray<FRHIRayTracingShader*> FShaderResource::GlobalRayTracingMaterialLibrary;
 FCriticalSection FShaderResource::GlobalRayTracingMaterialLibraryCS;
 
-void FShaderResource::GetRayTracingMaterialLibrary(TArray<FRayTracingShaderRHIParamRef>& RayTracingMaterials, FRayTracingShaderRHIParamRef DefaultShader)
+void FShaderResource::GetRayTracingMaterialLibrary(TArray<FRHIRayTracingShader*>& RayTracingMaterials, FRHIRayTracingShader* DefaultShader)
 {
 	FScopeLock Lock(&GlobalRayTracingMaterialLibraryCS);
 	RayTracingMaterials = GlobalRayTracingMaterialLibrary;
@@ -409,7 +409,7 @@ void FShaderResource::GetRayTracingMaterialLibrary(TArray<FRayTracingShaderRHIPa
 	}
 }
 
-uint32 FShaderResource::AddToRayTracingLibrary(FRayTracingShaderRHIParamRef Shader)
+uint32 FShaderResource::AddToRayTracingLibrary(FRHIRayTracingShader* Shader)
 {
 	FScopeLock Lock(&GlobalRayTracingMaterialLibraryCS);
 
diff --git a/Engine/Source/Runtime/RenderCore/Public/RayGenShaderUtils.h b/Engine/Source/Runtime/RenderCore/Public/RayGenShaderUtils.h
--- a/Engine/Source/Runtime/RenderCore/Public/RayGenShaderUtils.h
+++ b/Engine/Source/Runtime/RenderCore/Public/RayGenShaderUtils.h
@@ -35,7 +35,7 @@ struct RENDERCORE_API FRayGenShaderUtils
 			SetShaderParameters(GlobalResources, RayGenerationShader, *Parameters);
 
 			FRayTracingPipelineStateInitializer Initializer;
-			FRayTracingShaderRHIParamRef RayGenShaderTable[] = { RayGenerationShader->GetRayTracingShader() };
+			FRHIRayTracingShader* RayGenShaderTable[] = { RayGenerationShader->GetRayTracingShader() };
 			Initializer.SetRayGenShaderTable(RayGenShaderTable);
 
 			FRHIRayTracingPipelineState* Pipeline = PipelineStateCache::GetAndOrCreateRayTracingPipelineState(Initializer);
diff --git a/Engine/Source/Runtime/RenderCore/Public/Shader.h b/Engine/Source/Runtime/RenderCore/Public/Shader.h
--- a/Engine/Source/Runtime/RenderCore/Public/Shader.h
+++ b/Engine/Source/Runtime/RenderCore/Public/Shader.h
@@ -371,7 +371,7 @@ public:
 	}
 
 #if RHI_RAYTRACING
-	inline const FRayTracingShaderRHIParamRef GetRayTracingShader()
+	inline FRHIRayTracingShader* GetRayTracingShader()
 	{
 		checkSlow(Target.Frequency == SF_RayGen
 			   || Target.Frequency == SF_RayMiss
@@ -396,10 +396,10 @@ public:
 		return RayTracingMaterialLibraryIndex;
 	}
 
-	RENDERCORE_API static void GetRayTracingMaterialLibrary(TArray<FRayTracingShaderRHIParamRef>& RayTracingMaterials, FRayTracingShaderRHIParamRef DefaultShader);
+	RENDERCORE_API static void GetRayTracingMaterialLibrary(TArray<FRHIRayTracingShader*>& RayTracingMaterials, FRHIRayTracingShader* DefaultShader);
 
 private:
-	RENDERCORE_API static uint32 AddToRayTracingLibrary(FRayTracingShaderRHIParamRef Shader);
+	RENDERCORE_API static uint32 AddToRayTracingLibrary(FRHIRayTracingShader* Shader);
 	RENDERCORE_API static void RemoveFromRayTracingLibrary(uint32 Index);
 
 	static TArray<uint32> GlobalUnusedIndicies;
@@ -936,7 +936,7 @@ public:
 	}
 
 #if RHI_RAYTRACING
-	inline const FRayTracingShaderRHIParamRef GetRayTracingShader() const
+	inline FRHIRayTracingShader* GetRayTracingShader() const
 	{
 		return Resource->GetRayTracingShader();
 	}
diff --git a/Engine/Source/Runtime/Renderer/Private/DeferredShadingRenderer.cpp b/Engine/Source/Runtime/Renderer/Private/DeferredShadingRenderer.cpp
--- a/Engine/Source/Runtime/Renderer/Private/DeferredShadingRenderer.cpp
+++ b/Engine/Source/Runtime/Renderer/Private/DeferredShadingRenderer.cpp
@@ -778,7 +778,7 @@ bool FDeferredShadingSceneRenderer::DispatchRayTracingWorldUpdates(FRHICommandLi
 		RHICmdList.BuildAccelerationStructure(View.RayTracingScene.RayTracingSceneRHI);
 
 		// #dxr_todo: UE-72565: refactor ray tracing effects to not be member functions of DeferredShadingRenderer. register each effect at startup and just loop over them automatically to gather all required shaders
-		TArray<FRayTracingShaderRHIParamRef> RayGenShaders;
+		TArray<FRHIRayTracingShader*> RayGenShaders;
 		PrepareRayTracingReflections(View, RayGenShaders);
 		PrepareRayTracingShadows(View, RayGenShaders);
 		PrepareRayTracingAmbientOcclusion(View, RayGenShaders);
diff --git a/Engine/Source/Runtime/Renderer/Private/DeferredShadingRenderer.h b/Engine/Source/Runtime/Renderer/Private/DeferredShadingRenderer.h
--- a/Engine/Source/Runtime/Renderer/Private/DeferredShadingRenderer.h
+++ b/Engine/Source/Runtime/Renderer/Private/DeferredShadingRenderer.h
@@ -561,19 +561,19 @@ private:
 
 	bool GatherRayTracingWorldInstances(FRHICommandListImmediate& RHICmdList);
 	bool DispatchRayTracingWorldUpdates(FRHICommandListImmediate& RHICmdList);
-	FRayTracingPipelineState* BindRayTracingMaterialPipeline(FRHICommandList& RHICmdList, const FViewInfo& View, const TArrayView<const FRayTracingShaderRHIParamRef>& RayGenShaderTable, FRayTracingShaderRHIParamRef MissShader, FRayTracingShaderRHIParamRef DefaultClosestHitShader);
-	FRayTracingPipelineState* BindRayTracingDeferredMaterialGatherPipeline(FRHICommandList& RHICmdList, const FViewInfo& View, FRayTracingShaderRHIParamRef RayGenShader);
+	FRayTracingPipelineState* BindRayTracingMaterialPipeline(FRHICommandList& RHICmdList, const FViewInfo& View, const TArrayView<FRHIRayTracingShader*>& RayGenShaderTable, FRHIRayTracingShader* MissShader, FRHIRayTracingShader* DefaultClosestHitShader);
+	FRayTracingPipelineState* BindRayTracingDeferredMaterialGatherPipeline(FRHICommandList& RHICmdList, const FViewInfo& View, FRHIRayTracingShader* RayGenShader);
 
 	// #dxr_todo: UE-72565: refactor ray tracing effects to not be member functions of DeferredShadingRenderer. Register each effect at startup and just loop over them automatically
-	static void PrepareRayTracingReflections(const FViewInfo& View, TArray<FRayTracingShaderRHIParamRef>& OutRayGenShaders);
-	static void PrepareRayTracingShadows(const FViewInfo& View, TArray<FRayTracingShaderRHIParamRef>& OutRayGenShaders);
-	static void PrepareRayTracingAmbientOcclusion(const FViewInfo& View, TArray<FRayTracingShaderRHIParamRef>& OutRayGenShaders);
-	static void PrepareRayTracingSkyLight(const FViewInfo& View, TArray<FRayTracingShaderRHIParamRef>& OutRayGenShaders);
-	static void PrepareRayTracingRectLight(const FViewInfo& View, TArray<FRayTracingShaderRHIParamRef>& OutRayGenShaders);
-	static void PrepareRayTracingGlobalIllumination(const FViewInfo& View, TArray<FRayTracingShaderRHIParamRef>& OutRayGenShaders);
-	static void PrepareRayTracingTranslucency(const FViewInfo& View, TArray<FRayTracingShaderRHIParamRef>& OutRayGenShaders);
-	static void PrepareRayTracingDebug(const FViewInfo& View, TArray<FRayTracingShaderRHIParamRef>& OutRayGenShaders);
-	static void PreparePathTracing(const FViewInfo& View, TArray<FRayTracingShaderRHIParamRef>& OutRayGenShaders);
+	static void PrepareRayTracingReflections(const FViewInfo& View, TArray<FRHIRayTracingShader*>& OutRayGenShaders);
+	static void PrepareRayTracingShadows(const FViewInfo& View, TArray<FRHIRayTracingShader*>& OutRayGenShaders);
+	static void PrepareRayTracingAmbientOcclusion(const FViewInfo& View, TArray<FRHIRayTracingShader*>& OutRayGenShaders);
+	static void PrepareRayTracingSkyLight(const FViewInfo& View, TArray<FRHIRayTracingShader*>& OutRayGenShaders);
+	static void PrepareRayTracingRectLight(const FViewInfo& View, TArray<FRHIRayTracingShader*>& OutRayGenShaders);
+	static void PrepareRayTracingGlobalIllumination(const FViewInfo& View, TArray<FRHIRayTracingShader*>& OutRayGenShaders);
+	static void PrepareRayTracingTranslucency(const FViewInfo& View, TArray<FRHIRayTracingShader*>& OutRayGenShaders);
+	static void PrepareRayTracingDebug(const FViewInfo& View, TArray<FRHIRayTracingShader*>& OutRayGenShaders);
+	static void PreparePathTracing(const FViewInfo& View, TArray<FRHIRayTracingShader*>& OutRayGenShaders);
 
 #endif // RHI_RAYTRACING
 
diff --git a/Engine/Source/Runtime/Renderer/Private/PathTracing.cpp b/Engine/Source/Runtime/Renderer/Private/PathTracing.cpp
--- a/Engine/Source/Runtime/Renderer/Private/PathTracing.cpp
+++ b/Engine/Source/Runtime/Renderer/Private/PathTracing.cpp
@@ -435,7 +435,7 @@ public:
 
 IMPLEMENT_SHADER_TYPE(, FPathTracingCompositorPS, TEXT("/Engine/Private/PathTracing/PathTracingCompositingPixelShader.usf"), TEXT("CompositeMain"), SF_Pixel);
 
-void FDeferredShadingSceneRenderer::PreparePathTracing(const FViewInfo& View, TArray<FRayTracingShaderRHIParamRef>& OutRayGenShaders)
+void FDeferredShadingSceneRenderer::PreparePathTracing(const FViewInfo& View, TArray<FRHIRayTracingShader*>& OutRayGenShaders)
 {
 	// Declare all RayGen shaders that require material closest hit shaders to be bound
 	auto RayGenShader = View.ShaderMap->GetShader<FPathTracingRG>();
diff --git a/Engine/Source/Runtime/Renderer/Private/RayTracing/RayTracingAmbientOcclusion.cpp b/Engine/Source/Runtime/Renderer/Private/RayTracing/RayTracingAmbientOcclusion.cpp
--- a/Engine/Source/Runtime/Renderer/Private/RayTracing/RayTracingAmbientOcclusion.cpp
+++ b/Engine/Source/Runtime/Renderer/Private/RayTracing/RayTracingAmbientOcclusion.cpp
@@ -98,7 +98,7 @@ class FRayTracingAmbientOcclusionRGS : public FGlobalShader
 
 IMPLEMENT_GLOBAL_SHADER(FRayTracingAmbientOcclusionRGS, "/Engine/Private/RayTracing/RayTracingAmbientOcclusionRGS.usf", "AmbientOcclusionRGS", SF_RayGen);
 
-void FDeferredShadingSceneRenderer::PrepareRayTracingAmbientOcclusion(const FViewInfo& View, TArray<FRayTracingShaderRHIParamRef>& OutRayGenShaders)
+void FDeferredShadingSceneRenderer::PrepareRayTracingAmbientOcclusion(const FViewInfo& View, TArray<FRHIRayTracingShader*>& OutRayGenShaders)
 {
 	// Declare all RayGen shaders that require material closest hit shaders to be bound
 	FRayTracingAmbientOcclusionRGS::FPermutationDomain PermutationVector;
@@ -214,13 +214,13 @@ void FDeferredShadingSceneRenderer::RenderRayTracingAmbientOcclusion(
 			// Declare default pipeline
 			FRayTracingPipelineStateInitializer Initializer;
 			Initializer.MaxPayloadSizeInBytes = 52; // sizeof(FPackedMaterialClosestHitPayload)
-			FRayTracingShaderRHIParamRef RayGenShaderTable[] = { RayGenerationShader->GetRayTracingShader() };
+			FRHIRayTracingShader* RayGenShaderTable[] = { RayGenerationShader->GetRayTracingShader() };
 			Initializer.SetRayGenShaderTable(RayGenShaderTable);
 
-			FRayTracingShaderRHIParamRef MissShaderTable[] = { View.ShaderMap->GetShader<FDefaultMaterialMS>()->GetRayTracingShader() };
+			FRHIRayTracingShader* MissShaderTable[] = { View.ShaderMap->GetShader<FDefaultMaterialMS>()->GetRayTracingShader() };
 			Initializer.SetMissShaderTable(MissShaderTable);
 
-			FRayTracingShaderRHIParamRef HitGroupTable[] = { View.ShaderMap->GetShader<FOpaqueShadowHitGroup>()->GetRayTracingShader() };
+			FRHIRayTracingShader* HitGroupTable[] = { View.ShaderMap->GetShader<FOpaqueShadowHitGroup>()->GetRayTracingShader() };
 			Initializer.SetHitGroupTable(HitGroupTable);
 			Initializer.bAllowHitGroupIndexing = false; // Use the same hit shader for all geometry in the scene by disabling SBT indexing.
 
diff --git a/Engine/Source/Runtime/Renderer/Private/RayTracing/RayTracingBarycentrics.cpp b/Engine/Source/Runtime/Renderer/Private/RayTracing/RayTracingBarycentrics.cpp
--- a/Engine/Source/Runtime/Renderer/Private/RayTracing/RayTracingBarycentrics.cpp
+++ b/Engine/Source/Runtime/Renderer/Private/RayTracing/RayTracingBarycentrics.cpp
@@ -86,13 +86,13 @@ void FDeferredShadingSceneRenderer::RenderRayTracingBarycentrics(FRHICommandList
 
 	FRayTracingPipelineStateInitializer Initializer;
 
-	FRayTracingShaderRHIParamRef RayGenShaderTable[] = { RayGenShader->GetRayTracingShader() };
+	FRHIRayTracingShader* RayGenShaderTable[] = { RayGenShader->GetRayTracingShader() };
 	Initializer.SetRayGenShaderTable(RayGenShaderTable);
 
-	FRayTracingShaderRHIParamRef MissShaderTable[] = { MissShader->GetRayTracingShader() };
+	FRHIRayTracingShader* MissShaderTable[] = { MissShader->GetRayTracingShader() };
 	Initializer.SetMissShaderTable(MissShaderTable);
 
-	FRayTracingShaderRHIParamRef HitGroupTable[] = { ClosestHitShader->GetRayTracingShader() };
+	FRHIRayTracingShader* HitGroupTable[] = { ClosestHitShader->GetRayTracingShader() };
 	Initializer.SetHitGroupTable(HitGroupTable);
 	Initializer.bAllowHitGroupIndexing = false; // Use the same hit shader for all geometry in the scene by disabling SBT indexing.
 
diff --git a/Engine/Source/Runtime/Renderer/Private/RayTracing/RayTracingDebug.cpp b/Engine/Source/Runtime/Renderer/Private/RayTracing/RayTracingDebug.cpp
--- a/Engine/Source/Runtime/Renderer/Private/RayTracing/RayTracingDebug.cpp
+++ b/Engine/Source/Runtime/Renderer/Private/RayTracing/RayTracingDebug.cpp
@@ -73,7 +73,7 @@ public:
 // Dummy shader permutations to test hit group API
 IMPLEMENT_SHADER_TYPE(, FRayTracingDebugCHS, TEXT("/Engine/Private/RayTracing/RayTracingDebug.usf"), TEXT("RayTracingDebugMainCHS"), SF_RayHitGroup);
 
-void FDeferredShadingSceneRenderer::PrepareRayTracingDebug(const FViewInfo& View, TArray<FRayTracingShaderRHIParamRef>& OutRayGenShaders)
+void FDeferredShadingSceneRenderer::PrepareRayTracingDebug(const FViewInfo& View, TArray<FRHIRayTracingShader*>& OutRayGenShaders)
 {
 	// Declare all RayGen shaders that require material closest hit shaders to be bound
 
diff --git a/Engine/Source/Runtime/Renderer/Private/RayTracing/RayTracingDeferredMaterials.cpp b/Engine/Source/Runtime/Renderer/Private/RayTracing/RayTracingDeferredMaterials.cpp
--- a/Engine/Source/Runtime/Renderer/Private/RayTracing/RayTracingDeferredMaterials.cpp
+++ b/Engine/Source/Runtime/Renderer/Private/RayTracing/RayTracingDeferredMaterials.cpp
@@ -40,24 +40,24 @@ class FRayTracingDeferredMaterialMS : public FGlobalShader
 IMPLEMENT_GLOBAL_SHADER(FRayTracingDeferredMaterialCHS, "/Engine/Private/RayTracing/RayTracingDeferredMaterials.usf", "DeferredMaterialCHS", SF_RayHitGroup);
 IMPLEMENT_GLOBAL_SHADER(FRayTracingDeferredMaterialMS,  "/Engine/Private/RayTracing/RayTracingDeferredMaterials.usf", "DeferredMaterialMS",  SF_RayMiss);
 
-FRayTracingPipelineState* FDeferredShadingSceneRenderer::BindRayTracingDeferredMaterialGatherPipeline(FRHICommandList& RHICmdList, const FViewInfo& View, FRayTracingShaderRHIParamRef RayGenShader)
+FRayTracingPipelineState* FDeferredShadingSceneRenderer::BindRayTracingDeferredMaterialGatherPipeline(FRHICommandList& RHICmdList, const FViewInfo& View, FRHIRayTracingShader* RayGenShader)
 {
 	SCOPE_CYCLE_COUNTER(STAT_BindRayTracingPipeline);
 
 	FRayTracingPipelineStateInitializer Initializer;
 
-	FRayTracingShaderRHIParamRef RayGenShaderTable[] = { RayGenShader };
+	FRHIRayTracingShader* RayGenShaderTable[] = { RayGenShader };
 	Initializer.SetRayGenShaderTable(RayGenShaderTable);
 
-	auto MissShader = View.ShaderMap->GetShader<FRayTracingDeferredMaterialMS>();
-	FRayTracingShaderRHIParamRef MissShaderTable[] = { MissShader->GetRayTracingShader() };
+	auto* MissShader = View.ShaderMap->GetShader<FRayTracingDeferredMaterialMS>();
+	FRHIRayTracingShader* MissShaderTable[] = { MissShader->GetRayTracingShader() };
 	Initializer.SetMissShaderTable(MissShaderTable);
 
 	Initializer.MaxPayloadSizeInBytes = 12; // sizeof FDeferredMaterialPayload
 
 	// Get the ray tracing materials
-	auto ClosestHitShader = View.ShaderMap->GetShader<FRayTracingDeferredMaterialCHS>();
-	FRayTracingShaderRHIParamRef HitShaderTable[] = { ClosestHitShader->GetRayTracingShader() };
+	auto* ClosestHitShader = View.ShaderMap->GetShader<FRayTracingDeferredMaterialCHS>();
+	FRHIRayTracingShader* HitShaderTable[] = { ClosestHitShader->GetRayTracingShader() };
 	Initializer.SetHitGroupTable(HitShaderTable);
 
 	FRayTracingPipelineState* PipelineState = PipelineStateCache::GetAndOrCreateRayTracingPipelineState(RHICmdList, Initializer);
diff --git a/Engine/Source/Runtime/Renderer/Private/RayTracing/RayTracingGlobalIllumination.cpp b/Engine/Source/Runtime/Renderer/Private/RayTracing/RayTracingGlobalIllumination.cpp
--- a/Engine/Source/Runtime/Renderer/Private/RayTracing/RayTracingGlobalIllumination.cpp
+++ b/Engine/Source/Runtime/Renderer/Private/RayTracing/RayTracingGlobalIllumination.cpp
@@ -313,7 +313,7 @@ IMPLEMENT_GLOBAL_SHADER(FRayTracingGlobalIlluminationMS, "/Engine/Private/RayTra
 IMPLEMENT_GLOBAL_SHADER(FRayTracingGlobalIlluminationCompositePS, "/Engine/Private/RayTracing/RayTracingGlobalIlluminationCompositePS.usf", "GlobalIlluminationCompositePS", SF_Pixel);
 IMPLEMENT_GLOBAL_SHADER(FRayTracingGlobalIlluminationSceneColorCompositePS, "/Engine/Private/RayTracing/RayTracingGlobalIlluminationCompositePS.usf", "GlobalIlluminationSceneColorCompositePS", SF_Pixel);
 
-void FDeferredShadingSceneRenderer::PrepareRayTracingGlobalIllumination(const FViewInfo& View, TArray<FRayTracingShaderRHIParamRef>& OutRayGenShaders)
+void FDeferredShadingSceneRenderer::PrepareRayTracingGlobalIllumination(const FViewInfo& View, TArray<FRHIRayTracingShader*>& OutRayGenShaders)
 {
 	// Declare all RayGen shaders that require material closest hit shaders to be bound
 	for (int UseAttenuationTerm = 0; UseAttenuationTerm < 2; ++UseAttenuationTerm)
diff --git a/Engine/Source/Runtime/Renderer/Private/RayTracing/RayTracingMaterialHitShaders.cpp b/Engine/Source/Runtime/Renderer/Private/RayTracing/RayTracingMaterialHitShaders.cpp
--- a/Engine/Source/Runtime/Renderer/Private/RayTracing/RayTracingMaterialHitShaders.cpp
+++ b/Engine/Source/Runtime/Renderer/Private/RayTracing/RayTracingMaterialHitShaders.cpp
@@ -440,9 +440,9 @@ void FRayTracingMeshProcessor::AddMeshBatch(const FMeshBatch& RESTRICT MeshBatch
 FRayTracingPipelineState* FDeferredShadingSceneRenderer::BindRayTracingMaterialPipeline(
 	FRHICommandList& RHICmdList,
 	const FViewInfo& View,
-	const TArrayView<const FRayTracingShaderRHIParamRef>& RayGenShaderTable,
-	FRayTracingShaderRHIParamRef MissShader,
-	FRayTracingShaderRHIParamRef DefaultClosestHitShader
+	const TArrayView<FRHIRayTracingShader*>& RayGenShaderTable,
+	FRHIRayTracingShader* MissShader,
+	FRHIRayTracingShader* DefaultClosestHitShader
 )
 {
 	SCOPE_CYCLE_COUNTER(STAT_BindRayTracingPipeline);
@@ -456,12 +456,12 @@ FRayTracingPipelineState* FDeferredShadingSceneRenderer::BindRayTracingMaterialP
 
 	Initializer.SetRayGenShaderTable(RayGenShaderTable);
 
-	FRayTracingShaderRHIParamRef MissShaderTable[] = { MissShader };
+	FRHIRayTracingShader* MissShaderTable[] = { MissShader };
 	Initializer.SetMissShaderTable(MissShaderTable);
 
 	const bool bEnableMaterials = GEnableRayTracingMaterials != 0;
 
-	TArray<FRayTracingShaderRHIParamRef> RayTracingMaterialLibrary;
+	TArray<FRHIRayTracingShader*> RayTracingMaterialLibrary;
 
 	if (bEnableMaterials)
 	{
diff --git a/Engine/Source/Runtime/Renderer/Private/RayTracing/RayTracingRectLight.cpp b/Engine/Source/Runtime/Renderer/Private/RayTracing/RayTracingRectLight.cpp
--- a/Engine/Source/Runtime/Renderer/Private/RayTracing/RayTracingRectLight.cpp
+++ b/Engine/Source/Runtime/Renderer/Private/RayTracing/RayTracingRectLight.cpp
@@ -269,7 +269,7 @@ public:
 	{
 		FRayTracingPipelineStateInitializer Initializer;
 
-		FRayTracingShaderRHIParamRef RayGenShaderTable[] = { GetRayTracingShader() };
+		FRHIRayTracingShader* RayGenShaderTable[] = { GetRayTracingShader() };
 		Initializer.SetRayGenShaderTable(RayGenShaderTable);
 
 		FRayTracingPipelineState* Pipeline = PipelineStateCache::GetAndOrCreateRayTracingPipelineState(RHICmdList, Initializer);
@@ -443,7 +443,7 @@ void FDeferredShadingSceneRenderer::VisualizeRectLightMipTree(
 	RHICmdList.TransitionResource(EResourceTransitionAccess::ERWBarrier, EResourceTransitionPipeline::EGfxToCompute, RectLightMipTree.UAV);
 }
 
-void FDeferredShadingSceneRenderer::PrepareRayTracingRectLight(const FViewInfo& View, TArray<FRayTracingShaderRHIParamRef>& OutRayGenShaders)
+void FDeferredShadingSceneRenderer::PrepareRayTracingRectLight(const FViewInfo& View, TArray<FRHIRayTracingShader*>& OutRayGenShaders)
 {
 	// Declare all RayGen shaders that require material closest hit shaders to be bound
 
diff --git a/Engine/Source/Runtime/Renderer/Private/RayTracing/RayTracingReflections.cpp b/Engine/Source/Runtime/Renderer/Private/RayTracing/RayTracingReflections.cpp
--- a/Engine/Source/Runtime/Renderer/Private/RayTracing/RayTracingReflections.cpp
+++ b/Engine/Source/Runtime/Renderer/Private/RayTracing/RayTracingReflections.cpp
@@ -192,7 +192,7 @@ IMPLEMENT_GLOBAL_SHADER(FRayTracingReflectionsRGS, "/Engine/Private/RayTracing/R
 IMPLEMENT_GLOBAL_SHADER(FRayTracingReflectionsCHS, "/Engine/Private/RayTracing/RayTracingReflections.usf", "RayTracingReflectionsMainCHS", SF_RayHitGroup);
 IMPLEMENT_GLOBAL_SHADER(FRayTracingReflectionsMS, "/Engine/Private/RayTracing/RayTracingReflections.usf", "RayTracingReflectionsMainMS", SF_RayMiss);
 
-void FDeferredShadingSceneRenderer::PrepareRayTracingReflections(const FViewInfo& View, TArray<FRayTracingShaderRHIParamRef>& OutRayGenShaders)
+void FDeferredShadingSceneRenderer::PrepareRayTracingReflections(const FViewInfo& View, TArray<FRHIRayTracingShader*>& OutRayGenShaders)
 {
 	// Declare all RayGen shaders that require material closest hit shaders to be bound
 
diff --git a/Engine/Source/Runtime/Renderer/Private/RayTracing/RayTracingShadows.cpp b/Engine/Source/Runtime/Renderer/Private/RayTracing/RayTracingShadows.cpp
--- a/Engine/Source/Runtime/Renderer/Private/RayTracing/RayTracingShadows.cpp
+++ b/Engine/Source/Runtime/Renderer/Private/RayTracing/RayTracingShadows.cpp
@@ -80,7 +80,7 @@ float GetRaytracingMaxNormalBias()
 	return FMath::Max(0.01f, GRayTracingMaxNormalBias);
 }
 
-void FDeferredShadingSceneRenderer::PrepareRayTracingShadows(const FViewInfo& View, TArray<FRayTracingShaderRHIParamRef>& OutRayGenShaders)
+void FDeferredShadingSceneRenderer::PrepareRayTracingShadows(const FViewInfo& View, TArray<FRHIRayTracingShader*>& OutRayGenShaders)
 {
 	// Declare all RayGen shaders that require material closest hit shaders to be bound
 
@@ -223,13 +223,13 @@ void FDeferredShadingSceneRenderer::RenderRayTracingShadows(
 
 				Initializer.MaxPayloadSizeInBytes = 52; // sizeof(FPackedMaterialClosestHitPayload)
 
-				FRayTracingShaderRHIParamRef RayGenShaderTable[] = { RayGenerationShader->GetRayTracingShader() };
+				FRHIRayTracingShader* RayGenShaderTable[] = { RayGenerationShader->GetRayTracingShader() };
 				Initializer.SetRayGenShaderTable(RayGenShaderTable);
 
-				FRayTracingShaderRHIParamRef MissShaderTable[] = { View.ShaderMap->GetShader<FDefaultMaterialMS>()->GetRayTracingShader() };
+				FRHIRayTracingShader* MissShaderTable[] = { View.ShaderMap->GetShader<FDefaultMaterialMS>()->GetRayTracingShader() };
 				Initializer.SetMissShaderTable(MissShaderTable);
 
-				FRayTracingShaderRHIParamRef HitGroupTable[] = { View.ShaderMap->GetShader<FOpaqueShadowHitGroup>()->GetRayTracingShader() };
+				FRHIRayTracingShader* HitGroupTable[] = { View.ShaderMap->GetShader<FOpaqueShadowHitGroup>()->GetRayTracingShader() };
 				Initializer.SetHitGroupTable(HitGroupTable);
 				Initializer.bAllowHitGroupIndexing = false; // Use the same hit shader for all geometry in the scene by disabling SBT indexing.
 
diff --git a/Engine/Source/Runtime/Renderer/Private/RayTracing/RayTracingTranslucency.cpp b/Engine/Source/Runtime/Renderer/Private/RayTracing/RayTracingTranslucency.cpp
--- a/Engine/Source/Runtime/Renderer/Private/RayTracing/RayTracingTranslucency.cpp
+++ b/Engine/Source/Runtime/Renderer/Private/RayTracing/RayTracingTranslucency.cpp
@@ -254,7 +254,7 @@ private:
 
 IMPLEMENT_SHADER_TYPE(, FCompositeTranslucencyPS, TEXT("/Engine/Private/RayTracing/CompositeTranslucencyPS.usf"), TEXT("CompositeTranslucencyPS"), SF_Pixel)
 
-void FDeferredShadingSceneRenderer::PrepareRayTracingTranslucency(const FViewInfo& View, TArray<FRayTracingShaderRHIParamRef>& OutRayGenShaders)
+void FDeferredShadingSceneRenderer::PrepareRayTracingTranslucency(const FViewInfo& View, TArray<FRHIRayTracingShader*>& OutRayGenShaders)
 {
 	// Declare all RayGen shaders that require material closest hit shaders to be bound
 
diff --git a/Engine/Source/Runtime/Renderer/Private/RayTracing/RaytracingSkylight.cpp b/Engine/Source/Runtime/Renderer/Private/RayTracing/RaytracingSkylight.cpp
--- a/Engine/Source/Runtime/Renderer/Private/RayTracing/RaytracingSkylight.cpp
+++ b/Engine/Source/Runtime/Renderer/Private/RayTracing/RaytracingSkylight.cpp
@@ -758,7 +758,7 @@ void FDeferredShadingSceneRenderer::VisualizeSkyLightMipTree(
 	RHICmdList.TransitionResource(EResourceTransitionAccess::ERWBarrier, EResourceTransitionPipeline::EGfxToCompute, SkyLightMipTreeNegZ.UAV);
 }
 
-void FDeferredShadingSceneRenderer::PrepareRayTracingSkyLight(const FViewInfo& View, TArray<FRayTracingShaderRHIParamRef>& OutRayGenShaders)
+void FDeferredShadingSceneRenderer::PrepareRayTracingSkyLight(const FViewInfo& View, TArray<FRHIRayTracingShader*>& OutRayGenShaders)
 {
 	// Declare all RayGen shaders that require material closest hit shaders to be bound
 	FRayTracingSkyLightRGS::FPermutationDomain PermutationVector;
@@ -883,13 +883,13 @@ void FDeferredShadingSceneRenderer::RenderRayTracingSkyLight(
 				// Declare default pipeline
 				FRayTracingPipelineStateInitializer Initializer;
 				Initializer.MaxPayloadSizeInBytes = 52; // sizeof(FPackedMaterialClosestHitPayload)
-				FRayTracingShaderRHIParamRef RayGenShaderTable[] = { RayGenerationShader->GetRayTracingShader() };
+				FRHIRayTracingShader* RayGenShaderTable[] = { RayGenerationShader->GetRayTracingShader() };
 				Initializer.SetRayGenShaderTable(RayGenShaderTable);
 
-				FRayTracingShaderRHIParamRef MissShaderTable[] = { View.ShaderMap->GetShader<FDefaultMaterialMS>()->GetRayTracingShader() };
+				FRHIRayTracingShader* MissShaderTable[] = { View.ShaderMap->GetShader<FDefaultMaterialMS>()->GetRayTracingShader() };
 				Initializer.SetMissShaderTable(MissShaderTable);
 
-				FRayTracingShaderRHIParamRef HitGroupTable[] = { View.ShaderMap->GetShader<FOpaqueShadowHitGroup>()->GetRayTracingShader() };
+				FRHIRayTracingShader* HitGroupTable[] = { View.ShaderMap->GetShader<FOpaqueShadowHitGroup>()->GetRayTracingShader() };
 				Initializer.SetHitGroupTable(HitGroupTable);
 				Initializer.bAllowHitGroupIndexing = false; // Use the same hit shader for all geometry in the scene by disabling SBT indexing.
 
diff --git a/Engine/Source/Runtime/VulkanRHI/Private/VulkanCommands.cpp b/Engine/Source/Runtime/VulkanRHI/Private/VulkanCommands.cpp
--- a/Engine/Source/Runtime/VulkanRHI/Private/VulkanCommands.cpp
+++ b/Engine/Source/Runtime/VulkanRHI/Private/VulkanCommands.cpp
@@ -1135,7 +1135,7 @@ void FVulkanCommandListContext::RHICopyToStagingBuffer(FVertexBufferRHIParamRef
 	VulkanRHI::vkCmdCopyBuffer(CmdBuffer->GetHandle(), VertexBuffer->GetHandle(), StagingBuffer->StagingBuffer->GetHandle(), 1, &Region);
 }
 
-void FVulkanCommandListContext::RHIWriteGPUFence(FGPUFenceRHIParamRef FenceRHI)
+void FVulkanCommandListContext::RHIWriteGPUFence(FRHIGPUFence* FenceRHI)
 {
 	FVulkanCmdBuffer* CmdBuffer = CommandBufferManager->GetActiveCmdBuffer();
 	FVulkanGPUFence* Fence = ResourceCast(FenceRHI);
diff --git a/Engine/Source/Runtime/VulkanRHI/Private/VulkanContext.h b/Engine/Source/Runtime/VulkanRHI/Private/VulkanContext.h
--- a/Engine/Source/Runtime/VulkanRHI/Private/VulkanContext.h
+++ b/Engine/Source/Runtime/VulkanRHI/Private/VulkanContext.h
@@ -232,11 +232,11 @@ public:
 	virtual void RHITransitionResources(EResourceTransitionAccess TransitionType, FTextureRHIParamRef* InRenderTargets, int32 NumTextures) final override;
 	virtual void RHITransitionResources(EResourceTransitionAccess TransitionType, EResourceTransitionPipeline TransitionPipeline, FRHIUnorderedAccessView** InUAVs, int32 NumUAVs, FRHIComputeFence* WriteComputeFence) final override;
 	virtual void RHICopyToStagingBuffer(FVertexBufferRHIParamRef SourceBuffer, FStagingBufferRHIParamRef DestinationStagingBuffer, uint32 Offset, uint32 NumBytes) final override;
-	virtual void RHIWriteGPUFence(FGPUFenceRHIParamRef Fence) final override;
+	virtual void RHIWriteGPUFence(FRHIGPUFence* Fence) final override;
 
 	// Render time measurement
-	virtual void RHIBeginRenderQuery(FRenderQueryRHIParamRef RenderQuery) final override;
-	virtual void RHIEndRenderQuery(FRenderQueryRHIParamRef RenderQuery) final override;
+	virtual void RHIBeginRenderQuery(FRHIRenderQuery* RenderQuery) final override;
+	virtual void RHIEndRenderQuery(FRHIRenderQuery* RenderQuery) final override;
 
 	virtual void RHIUpdateTextureReference(FRHITextureReference* TextureRef, FTextureRHIParamRef NewTexture) final override;
 
diff --git a/Engine/Source/Runtime/VulkanRHI/Private/VulkanQuery.cpp b/Engine/Source/Runtime/VulkanRHI/Private/VulkanQuery.cpp
--- a/Engine/Source/Runtime/VulkanRHI/Private/VulkanQuery.cpp
+++ b/Engine/Source/Runtime/VulkanRHI/Private/VulkanQuery.cpp
@@ -426,7 +426,7 @@ FRenderQueryRHIRef FVulkanDynamicRHI::RHICreateRenderQuery(ERenderQueryType Quer
 	return Query;
 }
 
-bool FVulkanDynamicRHI::RHIGetRenderQueryResult(FRenderQueryRHIParamRef QueryRHI, uint64& OutNumPixels, bool bWait)
+bool FVulkanDynamicRHI::RHIGetRenderQueryResult(FRHIRenderQuery* QueryRHI, uint64& OutNumPixels, bool bWait)
 {
 	auto ToMicroseconds = [](uint64 Timestamp)
 	{
@@ -527,7 +527,7 @@ bool FVulkanDynamicRHI::RHIGetRenderQueryResult(FRenderQueryRHIParamRef QueryRHI
 	return false;
 }
 
-void FVulkanCommandListContext::RHIBeginRenderQuery(FRenderQueryRHIParamRef QueryRHI)
+void FVulkanCommandListContext::RHIBeginRenderQuery(FRHIRenderQuery* QueryRHI)
 {
 	FVulkanRenderQuery* BaseQuery = ResourceCast(QueryRHI);
 	if (BaseQuery->QueryType == RQT_Occlusion)
@@ -564,7 +564,7 @@ void FVulkanCommandListContext::RHIBeginRenderQuery(FRenderQueryRHIParamRef Quer
 	}
 }
 
-void FVulkanCommandListContext::RHIEndRenderQuery(FRenderQueryRHIParamRef QueryRHI)
+void FVulkanCommandListContext::RHIEndRenderQuery(FRHIRenderQuery* QueryRHI)
 {
 	FVulkanRenderQuery* BaseQuery = ResourceCast(QueryRHI);
 	if (BaseQuery->QueryType == RQT_Occlusion)
diff --git a/Engine/Source/Runtime/VulkanRHI/Public/VulkanDynamicRHI.h b/Engine/Source/Runtime/VulkanRHI/Public/VulkanDynamicRHI.h
--- a/Engine/Source/Runtime/VulkanRHI/Public/VulkanDynamicRHI.h
+++ b/Engine/Source/Runtime/VulkanRHI/Public/VulkanDynamicRHI.h
@@ -135,7 +135,7 @@ public:
 	virtual void RHIReadSurfaceFloatData(FTextureRHIParamRef Texture, FIntRect Rect, TArray<FFloat16Color>& OutData, ECubeFace CubeFace, int32 ArrayIndex, int32 MipIndex) final override;
 	virtual void RHIRead3DSurfaceFloatData(FTextureRHIParamRef Texture, FIntRect Rect, FIntPoint ZMinMax, TArray<FFloat16Color>& OutData) final override;
 	virtual FRenderQueryRHIRef RHICreateRenderQuery(ERenderQueryType QueryType) final override;
-	virtual bool RHIGetRenderQueryResult(FRenderQueryRHIParamRef RenderQuery, uint64& OutResult, bool bWait) final override;
+	virtual bool RHIGetRenderQueryResult(FRHIRenderQuery* RenderQuery, uint64& OutResult, bool bWait) final override;
 	virtual FTexture2DRHIRef RHIGetViewportBackBuffer(FViewportRHIParamRef Viewport) final override;
 	virtual void RHIAdvanceFrameForGetViewportBackBuffer(FViewportRHIParamRef Viewport) final override;
 	virtual void RHIAcquireThreadOwnership() final override;
diff --git a/Engine/Source/Runtime/Windows/D3D11RHI/Private/D3D11Commands.cpp b/Engine/Source/Runtime/Windows/D3D11RHI/Private/D3D11Commands.cpp
--- a/Engine/Source/Runtime/Windows/D3D11RHI/Private/D3D11Commands.cpp
+++ b/Engine/Source/Runtime/Windows/D3D11RHI/Private/D3D11Commands.cpp
@@ -2090,7 +2090,7 @@ void FD3D11DynamicRHI::RHICopyToStagingBuffer(FVertexBufferRHIParamRef SourceBuf
 	}
 }
 
-void FD3D11DynamicRHI::RHIWriteGPUFence(FGPUFenceRHIParamRef FenceRHI)
+void FD3D11DynamicRHI::RHIWriteGPUFence(FRHIGPUFence* FenceRHI)
 {
 	// @todo-staging Implement real fences for D3D11
 	// D3D11 only has the generic fence for now.
diff --git a/Engine/Source/Runtime/Windows/D3D11RHI/Private/D3D11Query.cpp b/Engine/Source/Runtime/Windows/D3D11RHI/Private/D3D11Query.cpp
--- a/Engine/Source/Runtime/Windows/D3D11RHI/Private/D3D11Query.cpp
+++ b/Engine/Source/Runtime/Windows/D3D11RHI/Private/D3D11Query.cpp
@@ -203,7 +203,7 @@ FRenderQueryRHIRef FD3D11DynamicRHI::RHICreateRenderQuery(ERenderQueryType Query
 	return new FD3D11RenderQuery(Query, QueryType);
 }
 
-bool FD3D11DynamicRHI::RHIGetRenderQueryResult(FRenderQueryRHIParamRef QueryRHI,uint64& OutResult,bool bWait)
+bool FD3D11DynamicRHI::RHIGetRenderQueryResult(FRHIRenderQuery* QueryRHI,uint64& OutResult,bool bWait)
 {
 	check(IsInRenderingThread());
 	FD3D11RenderQuery* Query = ResourceCast(QueryRHI);
@@ -250,7 +250,7 @@ bool FD3D11DynamicRHI::RHIGetRenderQueryResult(FRenderQueryRHIParamRef QueryRHI,
 
 
 // Occlusion/Timer queries.
-void FD3D11DynamicRHI::RHIBeginRenderQuery(FRenderQueryRHIParamRef QueryRHI)
+void FD3D11DynamicRHI::RHIBeginRenderQuery(FRHIRenderQuery* QueryRHI)
 {
 	FD3D11RenderQuery* Query = ResourceCast(QueryRHI);
 
@@ -268,7 +268,7 @@ void FD3D11DynamicRHI::RHIBeginRenderQuery(FRenderQueryRHIParamRef QueryRHI)
 	}
 }
 
-void FD3D11DynamicRHI::RHIEndRenderQuery(FRenderQueryRHIParamRef QueryRHI)
+void FD3D11DynamicRHI::RHIEndRenderQuery(FRHIRenderQuery* QueryRHI)
 {
 	FD3D11RenderQuery* Query = ResourceCast(QueryRHI);
 	Query->bResultIsCached = false; // for occlusion queries, this is redundant with the one in begin
diff --git a/Engine/Source/Runtime/Windows/D3D11RHI/Private/D3D11RHIPrivate.h b/Engine/Source/Runtime/Windows/D3D11RHI/Private/D3D11RHIPrivate.h
--- a/Engine/Source/Runtime/Windows/D3D11RHI/Private/D3D11RHIPrivate.h
+++ b/Engine/Source/Runtime/Windows/D3D11RHI/Private/D3D11RHIPrivate.h
@@ -406,7 +406,7 @@ public:
 	virtual FComputeShaderRHIRef RHICreateComputeShader(const TArray<uint8>& Code) final override;
 	virtual FStagingBufferRHIRef RHICreateStagingBuffer() final override;
 	virtual void RHICopyToStagingBuffer(FVertexBufferRHIParamRef SourceBuffer, FStagingBufferRHIParamRef DestinationStagingBuffer, uint32 Offset, uint32 NumBytes) final override;
-	virtual void RHIWriteGPUFence(FGPUFenceRHIParamRef Fence) final override;
+	virtual void RHIWriteGPUFence(FRHIGPUFence* Fence) final override;
 	virtual void* RHILockStagingBuffer(FStagingBufferRHIParamRef StagingBuffer, uint32 Offset, uint32 SizeRHI) final override;
     virtual void RHIUnlockStagingBuffer(FStagingBufferRHIParamRef StagingBuffer) final override;
 	virtual FBoundShaderStateRHIRef RHICreateBoundShaderState(FRHIVertexDeclaration* VertexDeclaration, FVertexShaderRHIParamRef VertexShader, FHullShaderRHIParamRef HullShader, FDomainShaderRHIParamRef DomainShader, FPixelShaderRHIParamRef PixelShader, FGeometryShaderRHIParamRef GeometryShader) final override;
@@ -472,7 +472,7 @@ public:
 	virtual void RHIReadSurfaceFloatData(FTextureRHIParamRef Texture,FIntRect Rect,TArray<FFloat16Color>& OutData,ECubeFace CubeFace,int32 ArrayIndex,int32 MipIndex) final override;
 	virtual void RHIRead3DSurfaceFloatData(FTextureRHIParamRef Texture,FIntRect Rect,FIntPoint ZMinMax,TArray<FFloat16Color>& OutData) final override;
 	virtual FRenderQueryRHIRef RHICreateRenderQuery(ERenderQueryType QueryType) final override;
-	virtual bool RHIGetRenderQueryResult(FRenderQueryRHIParamRef RenderQuery, uint64& OutResult, bool bWait) final override;
+	virtual bool RHIGetRenderQueryResult(FRHIRenderQuery* RenderQuery, uint64& OutResult, bool bWait) final override;
 	virtual FTexture2DRHIRef RHIGetViewportBackBuffer(FViewportRHIParamRef Viewport) final override;
 	virtual void RHIAdvanceFrameForGetViewportBackBuffer(FViewportRHIParamRef Viewport) final override;
 	virtual void RHIAcquireThreadOwnership() final override;
@@ -505,8 +505,8 @@ public:
 	virtual void RHICopyTexture(FTextureRHIParamRef SourceTexture, FTextureRHIParamRef DestTexture, const FRHICopyTextureInfo& CopyInfo) final override;
 	virtual void RHITransitionResources(EResourceTransitionAccess TransitionType, FTextureRHIParamRef* InTextures, int32 NumTextures) final override;
 	virtual void RHITransitionResources(EResourceTransitionAccess TransitionType, EResourceTransitionPipeline TransitionPipeline, FRHIUnorderedAccessView** InUAVs, int32 NumUAVs, FRHIComputeFence* WriteFence) final override;
-	virtual void RHIBeginRenderQuery(FRenderQueryRHIParamRef RenderQuery) final override;
-	virtual void RHIEndRenderQuery(FRenderQueryRHIParamRef RenderQuery) final override;
+	virtual void RHIBeginRenderQuery(FRHIRenderQuery* RenderQuery) final override;
+	virtual void RHIEndRenderQuery(FRHIRenderQuery* RenderQuery) final override;
 	void RHIBeginOcclusionQueryBatch(uint32 NumQueriesInBatch);
 	void RHIEndOcclusionQueryBatch();
 	virtual void RHISubmitCommandsHint() final override;
