Eliminate several frequent D3D warnings:

ID3D12Resource1::ID3D12Resource::Map: pReadRange is NULL and the heap page property is WRITE_BACK. This can be an indicator of inefficiencies that will result on mobile systems, as the range should be as tight as possible. But, it is also possible the range must encompass the whole subresource

ID3D12CommandQueue::ExecuteCommandLists: Readback Resource still has mapped subresorces when executing a command list that performs a copy operation to the resource.This may be ok if any data read from the readback resources was flushed by calling Unmap() after the resourcecopy operation completed.

ID3D12CommandList::ResourceBarrier: Called on the same subresource(s) of Resource in separate Barrier Descs which is inefficient and likely unintentional.

diff --git a/Engine/Source/Runtime/D3D12RHI/Private/D3D12CommandList.cpp b/Engine/Source/Runtime/D3D12RHI/Private/D3D12CommandList.cpp
--- a/Engine/Source/Runtime/D3D12RHI/Private/D3D12CommandList.cpp
+++ b/Engine/Source/Runtime/D3D12RHI/Private/D3D12CommandList.cpp
@@ -6,8 +6,8 @@
 void FD3D12CommandListHandle::AddTransitionBarrier(FD3D12Resource* pResource, D3D12_RESOURCE_STATES Before, D3D12_RESOURCE_STATES After, uint32 Subresource)
 {
 	check(CommandListData);
-	CommandListData->ResourceBarrierBatcher.AddTransition(pResource->GetResource(), Before, After, Subresource);
-	CommandListData->CurrentOwningContext->numBarriers++;
+	int32 NumAdded = CommandListData->ResourceBarrierBatcher.AddTransition(pResource->GetResource(), Before, After, Subresource);
+	CommandListData->CurrentOwningContext->numBarriers += NumAdded;
 
 	pResource->UpdateResidency(*this);
 }
diff --git a/Engine/Source/Runtime/D3D12RHI/Private/D3D12RenderTarget.cpp b/Engine/Source/Runtime/D3D12RHI/Private/D3D12RenderTarget.cpp
--- a/Engine/Source/Runtime/D3D12RHI/Private/D3D12RenderTarget.cpp
+++ b/Engine/Source/Runtime/D3D12RHI/Private/D3D12RenderTarget.cpp
@@ -785,16 +785,17 @@ void FD3D12DynamicRHI::ReadSurfaceDataNoMSAARaw(FTextureRHIParamRef TextureRHI,
 	OutData.Empty();
 	OutData.AddUninitialized(SizeX * SizeY * BytesPerPixel);
 
-	// Lock the staging resource.
-	void* pData;
-	VERIFYD3D12RESULT(TempTexture2D->GetResource()->Map(0, nullptr, &pData));
-
 	uint32 BytesPerLine = BytesPerPixel * InRect.Width();
-
 	const uint32 XBytesAligned = Align((uint32)readBackHeapDesc.Footprint.Width * BytesPerPixel, FD3D12_TEXTURE_DATA_PITCH_ALIGNMENT);
+	uint32 SrcStart = StagingRect.Min.X * BytesPerPixel + StagingRect.Min.Y * XBytesAligned;
+
+	// Lock the staging resource.
+	void* pData;
+	D3D12_RANGE ReadRange = { SrcStart, SrcStart + XBytesAligned * (SizeY - 1) + BytesPerLine };
+	VERIFYD3D12RESULT(TempTexture2D->GetResource()->Map(0, &ReadRange, &pData));
 
 	uint8* DestPtr = OutData.GetData();
-	uint8* SrcPtr = (uint8*)pData + StagingRect.Min.X * BytesPerPixel + StagingRect.Min.Y * XBytesAligned;
+	uint8* SrcPtr = (uint8*)pData + SrcStart;
 	for (uint32 Y = 0; Y < SizeY; Y++)
 	{
 		memcpy(DestPtr, SrcPtr, BytesPerLine);
@@ -1334,7 +1335,8 @@ void FD3D12DynamicRHI::RHIMapStagingSurface(FTextureRHIParamRef TextureRHI, void
 		GetRHIDevice()->GetCommandListManager().WaitForCompletion(SyncPoint);
 
 	void* pData;
-	HRESULT Result = Texture->GetResource()->Map(0, nullptr, &pData);
+	D3D12_RANGE ReadRange = { 0, Texture->GetDesc().Width };
+	HRESULT Result = Texture->GetResource()->Map(0, &ReadRange, &pData);
 	if (Result == DXGI_ERROR_DEVICE_REMOVED)
 	{
 		// When reading back to the CPU, we have to watch out for DXGI_ERROR_DEVICE_REMOVED
@@ -1479,7 +1481,8 @@ void FD3D12DynamicRHI::RHIReadSurfaceFloatData(FTextureRHIParamRef TextureRHI, F
 
 	// Lock the staging resource.
 	void* pData;
-	VERIFYD3D12RESULT(TempTexture2D->GetResource()->Map(0, nullptr, &pData));
+	D3D12_RANGE Range = { 0, MipBytesAligned };
+	VERIFYD3D12RESULT(TempTexture2D->GetResource()->Map(0, &Range, &pData));
 
 	// Presize the array
 	int32 TotalCount = SizeX * SizeY;
diff --git a/Engine/Source/Runtime/D3D12RHI/Public/D3D12Resources.h b/Engine/Source/Runtime/D3D12RHI/Public/D3D12Resources.h
--- a/Engine/Source/Runtime/D3D12RHI/Public/D3D12Resources.h
+++ b/Engine/Source/Runtime/D3D12RHI/Public/D3D12Resources.h
@@ -474,7 +474,7 @@ public:
 		SetResource(Resource);
 		SetSize(BufferSize);
 
-		if (!IsCPUInaccessible(Resource->GetHeapType()))
+		if (IsCPUWritable(Resource->GetHeapType()))
 		{
 			SetMappedBaseAddress(Resource->Map());
 		}
@@ -916,10 +916,29 @@ public:
 		Barrier.UAV.pResource = nullptr;	// Ignore the resource ptr for now. HW doesn't do anything with it.
 	}
 
-	// Add a transition resource barrier to the batch.
-	void AddTransition(ID3D12Resource* pResource, D3D12_RESOURCE_STATES Before, D3D12_RESOURCE_STATES After, uint32 Subresource)
+	// Add a transition resource barrier to the batch. Returns the number of barriers added, which may be negative if an existing barrier was cancelled.
+	int32 AddTransition(ID3D12Resource* pResource, D3D12_RESOURCE_STATES Before, D3D12_RESOURCE_STATES After, uint32 Subresource)
 	{
 		check(Before != After);
+
+		if (Barriers.Num())
+		{
+			// Check if we are simply reverting the last transition. In that case, we can just remove both transitions.
+			// This happens fairly frequently due to resource pooling since different RHI buffers can point to the same underlying D3D buffer.
+			// Instead of ping-ponging that underlying resource between COPY_DEST and GENERIC_READ, several copies can happen without a ResourceBarrier() in between.
+			// Doing this check also eliminates a D3D debug layer warning about multiple transitions of the same subresource.
+			const D3D12_RESOURCE_BARRIER& Last = Barriers.Last();
+			if (pResource == Last.Transition.pResource &&
+				Subresource == Last.Transition.Subresource &&
+				Before == Last.Transition.StateAfter &&
+				After  == Last.Transition.StateBefore &&
+				Last.Type == D3D12_RESOURCE_BARRIER_TYPE_TRANSITION)
+			{
+				Barriers.RemoveAt(Barriers.Num() - 1);
+				return -1;
+			}
+		}
+
 		Barriers.AddUninitialized();
 		D3D12_RESOURCE_BARRIER& Barrier = Barriers.Last();
 		Barrier.Type = D3D12_RESOURCE_BARRIER_TYPE_TRANSITION;
@@ -928,6 +947,7 @@ public:
 		Barrier.Transition.StateAfter = After;
 		Barrier.Transition.Subresource = Subresource;
 		Barrier.Transition.pResource = pResource;
+		return 1;
 	}
 
 	void AddAliasingBarrier(ID3D12Resource* pResource)
