Implemented parallel compilation of shaders when creating the ray tracing pipeline state object.

diff --git a/Engine/Source/Runtime/D3D12RHI/Private/D3D12RayTracing.cpp b/Engine/Source/Runtime/D3D12RHI/Private/D3D12RayTracing.cpp
--- a/Engine/Source/Runtime/D3D12RHI/Private/D3D12RayTracing.cpp
+++ b/Engine/Source/Runtime/D3D12RHI/Private/D3D12RayTracing.cpp
@@ -75,6 +75,9 @@ DECLARE_DWORD_COUNTER_STAT(TEXT("Updated TL AS (per frame)"), STAT_D3D12RayTraci
 
 DECLARE_MEMORY_STAT(TEXT("Used Video Memory"), STAT_D3D12RayTracingUsedVideoMemory, STATGROUP_D3D12RayTracing);
 
+DECLARE_CYCLE_STAT(TEXT("RTPSO Compile Shader"), STAT_RTPSO_CompileShader, STATGROUP_D3D12RayTracing);
+DECLARE_CYCLE_STAT(TEXT("RTPSO Create Pipeline"), STAT_RTPSO_CreatePipeline, STATGROUP_D3D12RayTracing);
+
 // Built-in local root parameters that are always bound to all hit shaders
 struct FHitGroupSystemParameters
 {
@@ -172,9 +175,10 @@ struct FDXILLibrary
 
 	// NOTE: typical DXIL library may contain up to 3 entry points (i.e. hit groups with closest hit, any hit and intersection shaders)
 	// Typical case is 1 (RGS, MS or CHS only) or 2 (CHS + AHS for shaders with alpha masking)
-	TArray<D3D12_EXPORT_DESC, TInlineAllocator<3>> ExportDesc;
-	TArray<FString, TInlineAllocator<3>> EntryNames;
-	TArray<FString, TInlineAllocator<3>> ExportNames;
+	static constexpr uint32 ExpectedEntryPoints = 3;
+	TArray<D3D12_EXPORT_DESC, TInlineAllocator<ExpectedEntryPoints>> ExportDesc;
+	TArray<FString, TInlineAllocator<ExpectedEntryPoints>> EntryNames;
+	TArray<FString, TInlineAllocator<ExpectedEntryPoints>> ExportNames;
 
 	D3D12_DXIL_LIBRARY_DESC Desc = {};
 };
@@ -187,12 +191,14 @@ static TRefCountPtr<ID3D12StateObject> CreateRayTracingStateObject(
 	const TArrayView<const D3D12_HIT_GROUP_DESC>& HitGroups,
 	const ID3D12RootSignature* GlobalRootSignature,
 	const TArrayView<ID3D12RootSignature*>& LocalRootSignatures,
-	const TArrayView<uint32>& LocalRootSignatureAssociations, // indices into LocalRootSignatures, one per export
+	const TArrayView<uint32>& LocalRootSignatureAssociations, // indices into LocalRootSignatures, one per export (may be empty, which assumes single root signature used for everything)
 	const TArrayView<D3D12_EXISTING_COLLECTION_DESC>& ExistingCollections,
 	D3D12_STATE_OBJECT_TYPE StateObjectType // Full RTPSO or a Collection
 )
 {
-	checkf(LocalRootSignatureAssociations.Num() == Exports.Num(), TEXT("There must be exactly one local root signature association per export."));
+	checkf((LocalRootSignatureAssociations.Num() == 0 && LocalRootSignatures.Num() == 1)
+		|| (LocalRootSignatureAssociations.Num() == Exports.Num()),
+		TEXT("There must be exactly one local root signature association per export."));
 
 	TRefCountPtr<ID3D12StateObject> Result;
 
@@ -208,12 +214,12 @@ static TRefCountPtr<ID3D12StateObject> CreateRayTracingStateObject(
 		+ ShaderLibraries.Num()
 		+ HitGroups.Num()
 		+ LocalRootSignatures.Num()
-		+ LocalRootSignatureAssociations.Num()
+		+ Exports.Num()
 		+ ExistingCollections.Num()
 	);
 
 	TArray<D3D12_SUBOBJECT_TO_EXPORTS_ASSOCIATION> ExportAssociations;
-	ExportAssociations.SetNumUninitialized(LocalRootSignatureAssociations.Num());
+	ExportAssociations.SetNumUninitialized(Exports.Num());
 
 	uint32 Index = 0;
 
@@ -273,7 +279,10 @@ static TRefCountPtr<ID3D12StateObject> CreateRayTracingStateObject(
 
 	for (int32 ExportIndex = 0; ExportIndex < Exports.Num(); ++ExportIndex)
 	{
-		const int32 LocalRootSignatureIndex = LocalRootSignatureAssociations[ExportIndex];
+		// If custom LocalRootSignatureAssociations data is not provided, then assume same default local RS association.
+		const int32 LocalRootSignatureIndex = LocalRootSignatureAssociations.Num() != 0
+			? LocalRootSignatureAssociations[ExportIndex]
+			: 0;
 
 		D3D12_SUBOBJECT_TO_EXPORTS_ASSOCIATION& Association = ExportAssociations[ExportIndex];
 		Association = D3D12_SUBOBJECT_TO_EXPORTS_ASSOCIATION{};
@@ -338,6 +347,8 @@ class FD3D12RayTracingPipelineCache
 {
 public:
 
+	UE_NONCOPYABLE(FD3D12RayTracingPipelineCache)
+
 	FD3D12RayTracingPipelineCache(FD3D12Device* Device)
 		: DefaultLocalRootSignature(Device->GetParentAdapter())
 	{
@@ -349,12 +360,17 @@ public:
 		DefaultLocalRootSignature.Init(LocalRootSignatureDesc);
 	}
 
+	~FD3D12RayTracingPipelineCache()
+	{
+		Reset();
+	}
+
 	struct FKey
 	{
 		uint64 ShaderHash = 0;
 		uint32 MaxPayloadSizeInBytes = 0;
-		const ID3D12RootSignature* GlobalRootSignature = nullptr;
-		const ID3D12RootSignature* LocalRootSignature = nullptr;
+		ID3D12RootSignature* GlobalRootSignature = nullptr;
+		ID3D12RootSignature* LocalRootSignature = nullptr;
 
 		bool operator == (const FKey& Other) const
 		{
@@ -382,19 +398,31 @@ public:
 
 		D3D12_EXISTING_COLLECTION_DESC GetCollectionDesc()
 		{
+			check(CompileEvent.IsValid() && CompileEvent->IsComplete());
+			check(StateObject);
+
 			D3D12_EXISTING_COLLECTION_DESC Result = {};
 			Result.pExistingCollection = StateObject;
 			Result.pExports = ExportDescs.GetData();
 			Result.NumExports = ExportDescs.Num();
+
 			return Result;
 		}
 
+		const TCHAR* GetPrimaryExportNameChars()
+		{
+			checkf(ExportNames.Num()!=0, TEXT("This ray tracing shader collection does not export any symbols."));
+			return *(ExportNames[0]);
+		}
+
 		TRefCountPtr<FD3D12RayTracingShader> Shader;
 
 		TRefCountPtr<ID3D12StateObject> StateObject;
+		FGraphEventRef CompileEvent;
 
-		TArray<FString> ExportNames;
-		TArray<D3D12_EXPORT_DESC> ExportDescs;
+		static constexpr uint32 MaxExports = 4;
+		TArray<FString, TFixedAllocator<MaxExports>> ExportNames;
+		TArray<D3D12_EXPORT_DESC, TFixedAllocator<MaxExports>> ExportDescs;
 	};
 
 	enum class ECollectionType
@@ -422,109 +450,77 @@ public:
 		}
 	}
 
-	D3D12_EXISTING_COLLECTION_DESC GetOrCompileShader(
-		ID3D12Device5* RayTracingDevice,
-		FD3D12RayTracingShader* Shader,
-		ID3D12RootSignature* GlobalRootSignature,
-		uint32 MaxPayloadSizeInBytes,
-		ECollectionType CollectionType,
-		bool* bOutCacheHit = nullptr)
+	class FShaderCompileTask
 	{
-		FScopeLock Lock(&CriticalSection);
+	public:
 
-		const uint64 ShaderHash = GetShaderHash64(Shader);
+		UE_NONCOPYABLE(FShaderCompileTask)
 
-		ID3D12RootSignature* LocalRootSignature = nullptr;
-		if (CollectionType == ECollectionType::HitGroup || CollectionType == ECollectionType::Callable)
+		FShaderCompileTask(
+				FEntry& InEntry,
+				FKey InCacheKey,
+				ID3D12Device5* InRayTracingDevice,
+				ECollectionType InCollectionType)
+			: Entry(InEntry)
+			, CacheKey(InCacheKey)
+			, RayTracingDevice(InRayTracingDevice)
+			, CollectionType(InCollectionType)
 		{
-			// Only hit group and callable shaders have a local root signature
-			LocalRootSignature = Shader->pRootSignature->GetRootSignature();
-		}
-		else
-		{
-			// ... all other shaders share a default empty local root signature
-			LocalRootSignature = DefaultLocalRootSignature.GetRootSignature();
 		}
 
-		FKey CacheKey;
-		CacheKey.ShaderHash = ShaderHash;
-		CacheKey.MaxPayloadSizeInBytes = MaxPayloadSizeInBytes;
-		CacheKey.GlobalRootSignature = GlobalRootSignature;
-		CacheKey.LocalRootSignature = LocalRootSignature;
-
-		FEntry& Entry = Cache.FindOrAdd(CacheKey);
+		static ESubsequentsMode::Type GetSubsequentsMode() { return ESubsequentsMode::TrackSubsequents; }
 
-		if (Entry.StateObject)
+		void DoTask(ENamedThreads::Type CurrentThread, const FGraphEventRef& MyCompletionGraphEvent)
 		{
-			if (bOutCacheHit) *bOutCacheHit = true;
-		}
-		else
-		{
-			if (bOutCacheHit) *bOutCacheHit = false;
-
-			Entry.Shader = Shader;
+			SCOPE_CYCLE_COUNTER(STAT_RTPSO_CompileShader);
 
-			LPCWSTR EntryName = *(Shader->EntryPoint);
+			FD3D12RayTracingShader* Shader = Entry.Shader;
 
-			TArray<LPCWSTR, TInlineAllocator<4>> ExportNames;
-			TArray<LPCWSTR, TInlineAllocator<3>> OriginalEntryPoints;
-			TArray<LPCWSTR, TInlineAllocator<3>> RenamedEntryPoints;
-			TArray<uint32, TInlineAllocator<3>> LocalRootSignatureAssociations;
-
-			// Typical collection can have up to 4 exports: HitGroup/RGS/MS, CHS, AHS, IS
-			const uint32 MaxExports = 4;
-			Entry.ExportNames.Reserve(MaxExports);
-			Entry.ExportDescs.Reserve(MaxExports);
+			static constexpr uint32 MaxEntryPoints = 3; // CHS+AHS+IS for HitGroup or just a single entry point for other collection types
+			TArray<LPCWSTR, TFixedAllocator<MaxEntryPoints>> OriginalEntryPoints;
+			TArray<LPCWSTR, TFixedAllocator<MaxEntryPoints>> RenamedEntryPoints;
 
 			const uint32 NumHitGroups = CollectionType == ECollectionType::HitGroup ? 1 : 0;
-
-			FString& PrimaryExportName = Entry.ExportNames.Add_GetRef(GenerateShaderName(GetCollectionTypeName(CollectionType), ShaderHash));
+			const uint64 ShaderHash = CacheKey.ShaderHash;
+			ID3D12RootSignature* GlobalRootSignature = CacheKey.GlobalRootSignature;
+			ID3D12RootSignature* LocalRootSignature = CacheKey.LocalRootSignature;
+			const uint32 DefaultLocalRootSignatureIndex = 0;
+			uint32 MaxPayloadSizeInBytes = CacheKey.MaxPayloadSizeInBytes;
 
 			D3D12_HIT_GROUP_DESC HitGroupDesc = {};
 
-			ExportNames.Add(*PrimaryExportName);
-
 			if (CollectionType == ECollectionType::HitGroup)
 			{
-				HitGroupDesc.HitGroupExport = *PrimaryExportName;
+				HitGroupDesc.HitGroupExport = Entry.GetPrimaryExportNameChars();
 				HitGroupDesc.Type = Shader->IntersectionEntryPoint.IsEmpty() ? D3D12_HIT_GROUP_TYPE_TRIANGLES : D3D12_HIT_GROUP_TYPE_PROCEDURAL_PRIMITIVE;
 
 				{
-					FString& ExportName = Entry.ExportNames.Add_GetRef(GenerateShaderName(TEXT("CHS"), ShaderHash));
+					const FString& ExportName = Entry.ExportNames.Add_GetRef(GenerateShaderName(TEXT("CHS"), ShaderHash));
 
 					HitGroupDesc.ClosestHitShaderImport = *ExportName;
 
 					OriginalEntryPoints.Add(*(Shader->EntryPoint));
 					RenamedEntryPoints.Add(*ExportName);
-					ExportNames.Add(*ExportName);
-
-					LocalRootSignatureAssociations.Add(0);
 				}
 
 				if (!Shader->AnyHitEntryPoint.IsEmpty())
 				{
-					FString& ExportName = Entry.ExportNames.Add_GetRef(GenerateShaderName(TEXT("AHS"), ShaderHash));
+					const FString& ExportName = Entry.ExportNames.Add_GetRef(GenerateShaderName(TEXT("AHS"), ShaderHash));
 
 					HitGroupDesc.AnyHitShaderImport = *ExportName;
 
 					OriginalEntryPoints.Add(*(Shader->AnyHitEntryPoint));
 					RenamedEntryPoints.Add(*ExportName);
-					ExportNames.Add(*ExportName);
-
-					LocalRootSignatureAssociations.Add(0);
 				}
 
 				if (!Shader->IntersectionEntryPoint.IsEmpty())
 				{
-					FString& ExportName = Entry.ExportNames.Add_GetRef(GenerateShaderName(TEXT("IS"), ShaderHash));
+					const FString& ExportName = Entry.ExportNames.Add_GetRef(GenerateShaderName(TEXT("IS"), ShaderHash));
 
 					HitGroupDesc.IntersectionShaderImport = *ExportName;
 
 					OriginalEntryPoints.Add(*(Shader->IntersectionEntryPoint));
 					RenamedEntryPoints.Add(*ExportName);
-					ExportNames.Add(*ExportName);
-
-					LocalRootSignatureAssociations.Add(0);
 				}
 			}
 			else
@@ -532,11 +528,24 @@ public:
 				checkf(CollectionType == ECollectionType::Miss || CollectionType == ECollectionType::RayGen || CollectionType == ECollectionType::Callable, TEXT("Unexpected RT sahder collection type"));
 
 				OriginalEntryPoints.Add(*(Shader->EntryPoint));
-				RenamedEntryPoints.Add(*PrimaryExportName);
+				RenamedEntryPoints.Add(Entry.GetPrimaryExportNameChars());
+			}
 
-				LocalRootSignatureAssociations.Add(0);
+
+			for (const FString& ExportName : Entry.ExportNames)
+			{
+				D3D12_EXPORT_DESC ExportDesc = {};
+				ExportDesc.Name = *ExportName;
+				ExportDesc.Flags = D3D12_EXPORT_FLAG_NONE;
+				Entry.ExportDescs.Add(ExportDesc);
 			}
 
+			// Validate that memory reservation was correct
+
+			check(Entry.ExportNames.Num() <= Entry.MaxExports);
+			check(Entry.ExportDescs.Num() <= Entry.MaxExports);
+			check(Entry.ExportDescs.Num() != 0);
+
 			FDXILLibrary Library;
 			Library.InitFromDXIL(Shader->ShaderBytecode, OriginalEntryPoints.GetData(), RenamedEntryPoints.GetData(), OriginalEntryPoints.Num());
 
@@ -550,30 +559,123 @@ public:
 				MakeArrayView(&HitGroupDesc, NumHitGroups),
 				GlobalRootSignature,
 				MakeArrayView(&LocalRootSignature, 1),
-				LocalRootSignatureAssociations,
-				{},
+				{}, // LocalRootSignatureAssociations (single RS will be used for all exports since this is null)
+				{}, // ExistingCollections
 				D3D12_STATE_OBJECT_TYPE_COLLECTION);
+		}
+
+		FORCEINLINE TStatId GetStatId() const
+		{
+			return GET_STATID(STAT_RTPSO_CompileShader);
+		}
+
+		ENamedThreads::Type GetDesiredThread()
+		{
+			return ENamedThreads::AnyHiPriThreadHiPriTask;
+		}
 
-			for (LPCWSTR ExportName : ExportNames)
+		FEntry& Entry;
+		FKey CacheKey;
+		ID3D12Device5* RayTracingDevice;
+		ECollectionType CollectionType;
+
+	};
+
+	FEntry* GetOrCompileShader(
+		ID3D12Device5* RayTracingDevice,
+		FD3D12RayTracingShader* Shader,
+		ID3D12RootSignature* GlobalRootSignature,
+		uint32 MaxPayloadSizeInBytes,
+		ECollectionType CollectionType,
+		FGraphEventArray& CompletionList,
+		bool* bOutCacheHit = nullptr)
+	{
+		FScopeLock Lock(&CriticalSection);
+
+		const uint64 ShaderHash = GetShaderHash64(Shader);
+
+		ID3D12RootSignature* LocalRootSignature = nullptr;
+		if (CollectionType == ECollectionType::HitGroup || CollectionType == ECollectionType::Callable)
+		{
+			// Only hit group and callable shaders have a local root signature
+			LocalRootSignature = Shader->pRootSignature->GetRootSignature();
+		}
+		else
+		{
+			// ... all other shaders share a default empty local root signature
+			LocalRootSignature = DefaultLocalRootSignature.GetRootSignature();
+		}
+
+		FKey CacheKey;
+		CacheKey.ShaderHash = ShaderHash;
+		CacheKey.MaxPayloadSizeInBytes = MaxPayloadSizeInBytes;
+		CacheKey.GlobalRootSignature = GlobalRootSignature;
+		CacheKey.LocalRootSignature = LocalRootSignature;
+
+		FEntry*& FindResult = Cache.FindOrAdd(CacheKey);
+
+		if (FindResult)
+		{
+			if (bOutCacheHit) *bOutCacheHit = true;
+		}
+		else
+		{
+			if (bOutCacheHit) *bOutCacheHit = false;
+
+			if (FindResult == nullptr)
 			{
-				D3D12_EXPORT_DESC ExportDesc = {};
-				ExportDesc.Name = ExportName;
-				ExportDesc.Flags = D3D12_EXPORT_FLAG_NONE;
-				Entry.ExportDescs.Add(ExportDesc);
+				FindResult = new FEntry;
 			}
 
-			// Validate that memory reservation was correct
-			check(Entry.ExportNames.Num() <= MaxExports);
-			check(Entry.ExportDescs.Num() <= MaxExports);
+			FEntry& Entry = *FindResult;
+
+			Entry.Shader = Shader;
+
+			// Generate primary export name, which is immediately required on the PSO creation thread.
+			Entry.ExportNames.Add(GenerateShaderName(GetCollectionTypeName(CollectionType), ShaderHash));
+			checkf(Entry.ExportNames.Num() == 1, TEXT("Primary export name must always be first."));
+
+			// Defer actual compilation to another task, as there may be many shaders that may be compiled in parallel.
+			// Result of the compilation (the collection PSO) is not needed until final RT PSO is linked.
+			Entry.CompileEvent = TGraphTask<FShaderCompileTask>::CreateTask().ConstructAndDispatchWhenReady(
+				Entry,
+				CacheKey,
+				RayTracingDevice,
+				CollectionType
+			);
+		}
+
+		if (FindResult->CompileEvent.IsValid())
+		{
+			if (!FindResult->CompileEvent->IsComplete())
+			{
+				CompletionList.Add(FindResult->CompileEvent);
+			}
+		}
+		else
+		{
+			check(FindResult->StateObject != nullptr);
 		}
 
-		return Entry.GetCollectionDesc();
+		return FindResult;
+	}
+
+	void Reset()
+	{
+		FScopeLock Lock(&CriticalSection);
+
+		for (auto It : Cache)
+		{
+			delete It.Value;
+		}
+
+		Cache.Reset();
 	}
 
 private:
 
 	FCriticalSection CriticalSection;
-	TMap<FKey, FEntry> Cache;
+	TMap<FKey, FEntry*> Cache;
 	FD3D12RootSignature DefaultLocalRootSignature; // Default empty root signature used for default hit shaders.
 };
 
@@ -1277,6 +1379,8 @@ public:
 
 	FD3D12RayTracingPipelineState(FD3D12Device* Device, const FRayTracingPipelineStateInitializer& Initializer)
 	{
+		SCOPE_CYCLE_COUNTER(STAT_RTPSO_CreatePipeline);
+
 		checkf(Initializer.GetRayGenTable().Num() > 0, TEXT("Ray tracing pipelines must have at leat one ray generation shader."));
 
 		uint64 TotalCreationTime = 0;
@@ -1318,39 +1422,43 @@ public:
 		TSet<uint64> UniqueShaderHashes;
 		UniqueShaderHashes.Reserve(MaxTotalShaders);
 
-		TArray<D3D12_EXISTING_COLLECTION_DESC> UniqueShaderCollections;
+		TArray<FD3D12RayTracingPipelineCache::FEntry*> UniqueShaderCollections;
 		UniqueShaderCollections.Reserve(MaxTotalShaders);
 
-		auto AddShaderCollection = [Device, RayTracingDevice, GlobalRootSignature = this->GlobalRootSignature,
-										&UniqueShaderHashes, &UniqueShaderCollections, &Initializer, &NumCacheHits, &CompileTime]
+		FGraphEventArray CompileCompletionList;
+		CompileCompletionList.Reserve(MaxTotalShaders);
+
+		FD3D12RayTracingPipelineCache* PipelineCache = Device->GetRayTracingPipelineCache();
+
+		auto AddShaderCollection = [Device, RayTracingDevice, GlobalRootSignature = this->GlobalRootSignature, PipelineCache,
+										&UniqueShaderHashes, &UniqueShaderCollections, &Initializer, &NumCacheHits, &CompileTime,
+										&CompileCompletionList]
 			(FD3D12RayTracingShader* Shader, FD3D12RayTracingPipelineCache::ECollectionType CollectionType)
 		{
 			bool bIsAlreadyInSet = false;
 			const uint64 ShaderHash = GetShaderHash64(Shader);
 			UniqueShaderHashes.Add(ShaderHash, &bIsAlreadyInSet);
 
-			FD3D12RayTracingPipelineCache* PipelineCache = Device->GetRayTracingPipelineCache();
-
 			bool bCacheHit = false;
 
 			CompileTime -= FPlatformTime::Cycles64();
 
-			D3D12_EXISTING_COLLECTION_DESC CollectionDesc = PipelineCache->GetOrCompileShader(
+			FD3D12RayTracingPipelineCache::FEntry* ShaderCacheEntry = PipelineCache->GetOrCompileShader(
 				RayTracingDevice, Shader, GlobalRootSignature,
 				Initializer.MaxPayloadSizeInBytes,
-				CollectionType,
+				CollectionType, CompileCompletionList,
 				&bCacheHit);
 
 			CompileTime += FPlatformTime::Cycles64();
 
 			if (!bIsAlreadyInSet)
 			{
-				UniqueShaderCollections.Add(CollectionDesc);
+				UniqueShaderCollections.Add(ShaderCacheEntry);
 
 				if (bCacheHit) NumCacheHits++;
 			}
 
-			return CollectionDesc;
+			return ShaderCacheEntry;
 		};
 
 		// If no custom hit groups were provided, then disable SBT indexing and force default shader on all primitives
@@ -1370,9 +1478,9 @@ public:
 			checkf(Shader->pRootSignature->GetRootSignature() == GlobalRootSignature, TEXT("All raygen and miss shaders must share the same root signature"));
 			checkf(!Shader->ResourceCounts.bGlobalUniformBufferUsed, TEXT("Global uniform buffers are not implemented for ray generation shaders"));
 
-			D3D12_EXISTING_COLLECTION_DESC CollectionDesc = AddShaderCollection(Shader, FD3D12RayTracingPipelineCache::ECollectionType::RayGen);
+			FD3D12RayTracingPipelineCache::FEntry* ShaderCacheEntry = AddShaderCollection(Shader, FD3D12RayTracingPipelineCache::ECollectionType::RayGen);
 
-			RayGenShaderNames.Add(CollectionDesc.pExports[0].Name);
+			RayGenShaderNames.Add(ShaderCacheEntry->GetPrimaryExportNameChars());
 			RayGenShaders.Shaders.Add(Shader);
 		}
 
@@ -1388,9 +1496,9 @@ public:
 			checkf(Shader->pRootSignature->GetRootSignature() == GlobalRootSignature, TEXT("All raygen and miss shaders must share the same root signature"));
 			checkf(!Shader->ResourceCounts.bGlobalUniformBufferUsed, TEXT("Global uniform buffers are not implemented for ray tracing miss shaders"));
 
-			D3D12_EXISTING_COLLECTION_DESC CollectionDesc = AddShaderCollection(Shader, FD3D12RayTracingPipelineCache::ECollectionType::Miss);
+			FD3D12RayTracingPipelineCache::FEntry* ShaderCacheEntry = AddShaderCollection(Shader, FD3D12RayTracingPipelineCache::ECollectionType::Miss);
 
-			MissShaderNames.Add(CollectionDesc.pExports[0].Name);
+			MissShaderNames.Add(ShaderCacheEntry->GetPrimaryExportNameChars());
 			MissShaders.Shaders.Add(Shader);
 		}
 
@@ -1413,10 +1521,9 @@ public:
 			MaxHitGroupViewDescriptors = FMath::Max(MaxHitGroupViewDescriptors, ShaderViewDescriptors);
 			MaxLocalRootSignatureSize = FMath::Max(MaxLocalRootSignatureSize, Shader->pRootSignature->GetTotalRootSignatureSizeInBytes());
 
-			D3D12_EXISTING_COLLECTION_DESC CollectionDesc = AddShaderCollection(Shader, FD3D12RayTracingPipelineCache::ECollectionType::HitGroup);
+			FD3D12RayTracingPipelineCache::FEntry* ShaderCacheEntry = AddShaderCollection(Shader, FD3D12RayTracingPipelineCache::ECollectionType::HitGroup);
 
-			check(CollectionDesc.NumExports != 0);
-			HitGroupNames.Add(CollectionDesc.pExports[0].Name);
+			HitGroupNames.Add(ShaderCacheEntry->GetPrimaryExportNameChars());
 			HitGroupShaders.Shaders.Add(Shader);
 		}
 
@@ -1437,12 +1544,29 @@ public:
 			MaxHitGroupViewDescriptors = FMath::Max(MaxHitGroupViewDescriptors, ShaderViewDescriptors);
 			MaxLocalRootSignatureSize = FMath::Max(MaxLocalRootSignatureSize, Shader->pRootSignature->GetTotalRootSignatureSizeInBytes());
 
-			D3D12_EXISTING_COLLECTION_DESC CollectionDesc = AddShaderCollection(Shader, FD3D12RayTracingPipelineCache::ECollectionType::Callable);
+			FD3D12RayTracingPipelineCache::FEntry* ShaderCacheEntry = AddShaderCollection(Shader, FD3D12RayTracingPipelineCache::ECollectionType::Callable);
 
-			CallableShaderNames.Add(CollectionDesc.pExports[0].Name);
+			CallableShaderNames.Add(ShaderCacheEntry->GetPrimaryExportNameChars());
 			CallableShaders.Shaders.Add(Shader);
 		}
 
+		// Wait for all compilation tasks to be complete and then gather the compiled collection descriptors
+
+		CompileTime -= FPlatformTime::Cycles64();
+
+		FGraphEventRef CompileCompletionFence = FFunctionGraphTask::CreateAndDispatchWhenReady([]() {}, TStatId(), &CompileCompletionList, ENamedThreads::AnyHiPriThreadHiPriTask);
+		FTaskGraphInterface::Get().WaitUntilTaskCompletes(CompileCompletionFence);
+
+		CompileTime += FPlatformTime::Cycles64();
+
+
+		TArray<D3D12_EXISTING_COLLECTION_DESC> UniqueShaderCollectionDescs;
+		UniqueShaderCollectionDescs.Reserve(MaxTotalShaders);
+		for (FD3D12RayTracingPipelineCache::FEntry* Entry : UniqueShaderCollections)
+		{
+			UniqueShaderCollectionDescs.Add(Entry->GetCollectionDesc());
+		}
+
 		// Link final RTPSO from shader collections
 
 		LinkTime -= FPlatformTime::Cycles64();
@@ -1455,7 +1579,7 @@ public:
 			GlobalRootSignature,
 			{}, // LocalRootSignatures
 			{}, // LocalRootSignatureAssociations,
-			UniqueShaderCollections,
+			UniqueShaderCollectionDescs,
 			D3D12_STATE_OBJECT_TYPE_RAYTRACING_PIPELINE);
 		LinkTime += FPlatformTime::Cycles64();
 
