DR - Deprecate FVertexShaderRHIParamRef, FHullShaderRHIParamRef, FDomainShaderRHIParamRef, FPixelShaderRHIParamRef, FGeometryShaderRHIParamRef and FComputeShaderRHIParamRef

diff --git a/Engine/Plugins/Experimental/RemoteSession/Source/RemoteSession/Private/Channels/RemoteSessionARCameraChannel.cpp b/Engine/Plugins/Experimental/RemoteSession/Source/RemoteSession/Private/Channels/RemoteSessionARCameraChannel.cpp
--- a/Engine/Plugins/Experimental/RemoteSession/Source/RemoteSession/Private/Channels/RemoteSessionARCameraChannel.cpp
+++ b/Engine/Plugins/Experimental/RemoteSession/Source/RemoteSession/Private/Channels/RemoteSessionARCameraChannel.cpp
@@ -89,7 +89,7 @@ public:
 
 	void SetParameters(FRHICommandList& RHICmdList, const FSceneView& View)
 	{
-		const FVertexShaderRHIParamRef ShaderRHI = GetVertexShader();
+		FRHIVertexShader* ShaderRHI = GetVertexShader();
 		FMaterialShader::SetViewParameters(RHICmdList, ShaderRHI, View, View.ViewUniformBuffer);
 	}
 
@@ -135,7 +135,7 @@ public:
 
 	void SetParameters(FRHICommandList& RHICmdList, const FSceneView& View, const FMaterialRenderProxy* Material)
 	{
-		const FPixelShaderRHIParamRef ShaderRHI = GetPixelShader();
+		FRHIPixelShader* ShaderRHI = GetPixelShader();
 		FMaterialShader::SetParameters(RHICmdList, ShaderRHI, Material, *Material->GetMaterial(View.GetFeatureLevel()), View, View.ViewUniformBuffer, ESceneTextureSetupMode::None);
 
 		for (uint32 InputIter = 0; InputIter < ePId_Input_MAX; ++InputIter)
diff --git a/Engine/Plugins/FX/HoudiniNiagara/Source/HoudiniNiagara/Private/NiagaraDataInterfaceHoudiniCSV.cpp b/Engine/Plugins/FX/HoudiniNiagara/Source/HoudiniNiagara/Private/NiagaraDataInterfaceHoudiniCSV.cpp
--- a/Engine/Plugins/FX/HoudiniNiagara/Source/HoudiniNiagara/Private/NiagaraDataInterfaceHoudiniCSV.cpp
+++ b/Engine/Plugins/FX/HoudiniNiagara/Source/HoudiniNiagara/Private/NiagaraDataInterfaceHoudiniCSV.cpp
@@ -2069,10 +2069,12 @@ struct FNiagaraDataInterfaceParametersCS_HoudiniCSV : public FNiagaraDataInterfa
 	{
 		check( IsInRenderingThread() );
 
-		const FComputeShaderRHIParamRef ComputeShaderRHI = Shader->GetComputeShader();
+		FRHIComputeShader* ComputeShaderRHI = Shader->GetComputeShader();
 		UNiagaraDataInterfaceHoudiniCSV* HoudiniDI = CastChecked<UNiagaraDataInterfaceHoudiniCSV>( DataInterface );
 		if ( !HoudiniDI )
+		{
 			return;
+		}
 
 		SetShaderValue(RHICmdList, ComputeShaderRHI, NumberOfRows, HoudiniDI->GetNumberOfRows());
 		SetShaderValue(RHICmdList, ComputeShaderRHI, NumberOfColumns, HoudiniDI->GetNumberOfColumns());
diff --git a/Engine/Plugins/FX/Niagara/Source/Niagara/Private/NiagaraDataInterface.cpp b/Engine/Plugins/FX/Niagara/Source/Niagara/Private/NiagaraDataInterface.cpp
--- a/Engine/Plugins/FX/Niagara/Source/Niagara/Private/NiagaraDataInterface.cpp
+++ b/Engine/Plugins/FX/Niagara/Source/Niagara/Private/NiagaraDataInterface.cpp
@@ -191,7 +191,7 @@ struct FNiagaraDataInterfaceParametersCS_Curve : public FNiagaraDataInterfacePar
 	{
 		check(IsInRenderingThread());
 
-		const FComputeShaderRHIParamRef ComputeShaderRHI = Shader->GetComputeShader();
+		FRHIComputeShader* ComputeShaderRHI = Shader->GetComputeShader();
 		UNiagaraDataInterfaceCurveBase* CurveDI = CastChecked<UNiagaraDataInterfaceCurveBase>(DataInterface);
 		FReadBuffer& CurveLUTBuffer = CurveDI->GetCurveLUTGPUBuffer();
 
diff --git a/Engine/Plugins/FX/Niagara/Source/Niagara/Private/NiagaraDataInterfaceCollisionQuery.cpp b/Engine/Plugins/FX/Niagara/Source/Niagara/Private/NiagaraDataInterfaceCollisionQuery.cpp
--- a/Engine/Plugins/FX/Niagara/Source/Niagara/Private/NiagaraDataInterfaceCollisionQuery.cpp
+++ b/Engine/Plugins/FX/Niagara/Source/Niagara/Private/NiagaraDataInterfaceCollisionQuery.cpp
@@ -855,7 +855,7 @@ struct FNiagaraDataInterfaceParametersCS_CollisionQuery : public FNiagaraDataInt
 	{
 		check(IsInRenderingThread());
 
-		const FComputeShaderRHIParamRef ComputeShaderRHI = Shader->GetComputeShader();
+		FRHIComputeShader* ComputeShaderRHI = Shader->GetComputeShader();
 		
 		TUniformBufferRef<FSceneTexturesUniformParameters> SceneTextureUniformParams = GNiagaraViewDataManager.GetSceneTextureUniformParameters();
 		SetUniformBufferParameter(RHICmdList, ComputeShaderRHI, PassUniformBuffer/*Shader->GetUniformBufferParameter(SceneTexturesUniformBufferStruct)*/, SceneTextureUniformParams);
diff --git a/Engine/Plugins/FX/Niagara/Source/Niagara/Private/NiagaraDataInterfaceTexture.cpp b/Engine/Plugins/FX/Niagara/Source/Niagara/Private/NiagaraDataInterfaceTexture.cpp
--- a/Engine/Plugins/FX/Niagara/Source/Niagara/Private/NiagaraDataInterfaceTexture.cpp
+++ b/Engine/Plugins/FX/Niagara/Source/Niagara/Private/NiagaraDataInterfaceTexture.cpp
@@ -377,7 +377,7 @@ struct FNiagaraDataInterfaceParametersCS_Texture : public FNiagaraDataInterfaceP
 	{
 		check(IsInRenderingThread());
 
-		FComputeShaderRHIParamRef ComputeShaderRHI = Shader->GetComputeShader();
+		FRHIComputeShader* ComputeShaderRHI = Shader->GetComputeShader();
 		UNiagaraDataInterfaceTexture* TextureDI = CastChecked<UNiagaraDataInterfaceTexture>(DataInterface);
 		UTexture *Texture = TextureDI->Texture;
 		float TexDims[2];
diff --git a/Engine/Plugins/FX/Niagara/Source/NiagaraVertexFactories/Private/NiagaraSortingGPU.cpp b/Engine/Plugins/FX/Niagara/Source/NiagaraVertexFactories/Private/NiagaraSortingGPU.cpp
--- a/Engine/Plugins/FX/Niagara/Source/NiagaraVertexFactories/Private/NiagaraSortingGPU.cpp
+++ b/Engine/Plugins/FX/Niagara/Source/NiagaraVertexFactories/Private/NiagaraSortingGPU.cpp
@@ -94,7 +94,7 @@ bool FNiagaraSortKeyGenCS::Serialize(FArchive& Ar)
 
 void FNiagaraSortKeyGenCS::SetOutput(FRHICommandList& RHICmdList, FRHIUnorderedAccessView* OutKeysUAV, FRHIUnorderedAccessView* OutIndicesUAV)
 {
-	FComputeShaderRHIParamRef ComputeShaderRHI = GetComputeShader();
+	FRHIComputeShader* ComputeShaderRHI = GetComputeShader();
 	if (OutKeys.IsBound())
 	{
 		RHICmdList.SetUAVParameter(ComputeShaderRHI, OutKeys.GetBaseIndex(), OutKeysUAV);
@@ -107,7 +107,7 @@ void FNiagaraSortKeyGenCS::SetOutput(FRHICommandList& RHICmdList, FRHIUnorderedA
 
 void FNiagaraSortKeyGenCS::SetParameters(FRHICommandList& RHICmdList, const FNiagaraGPUSortInfo& SortInfo, int32 EmitterIndex, int32 OutputOffset)
 {
-	FComputeShaderRHIParamRef ComputeShaderRHI = GetComputeShader();
+	FRHIComputeShader* ComputeShaderRHI = GetComputeShader();
 
 	RHICmdList.SetShaderResourceViewParameter(ComputeShaderRHI, NiagaraParticleDataFloat.GetBaseIndex(), SortInfo.ParticleDataFloatSRV);
 	RHICmdList.SetShaderParameter(ComputeShaderRHI, FloatDataOffset.GetBufferIndex(), FloatDataOffset.GetBaseIndex(), FloatDataOffset.GetNumBytes(), &SortInfo.FloatDataOffset);
@@ -127,7 +127,7 @@ void FNiagaraSortKeyGenCS::SetParameters(FRHICommandList& RHICmdList, const FNia
 
 void FNiagaraSortKeyGenCS::UnbindBuffers(FRHICommandList& RHICmdList)
 {
-	FComputeShaderRHIParamRef ComputeShaderRHI = GetComputeShader();
+	FRHIComputeShader* ComputeShaderRHI = GetComputeShader();
 	if (NiagaraParticleDataFloat.IsBound())
 	{
 		RHICmdList.SetShaderResourceViewParameter(ComputeShaderRHI, NiagaraParticleDataFloat.GetBaseIndex(), nullptr);
@@ -179,7 +179,7 @@ void FNiagaraCopyIntBufferRegionCS::SetParameters(
 	int32 StartingIndex,
 	int32 DestCount)
 {
-	FComputeShaderRHIParamRef ComputeShaderRHI = GetComputeShader();
+	FRHIComputeShader* ComputeShaderRHI = GetComputeShader();
 	check(DestCount > 0 && DestCount <= NIAGARA_COPY_BUFFER_BUFFER_COUNT);
 
 	RHICmdList.SetShaderResourceViewParameter(ComputeShaderRHI, SourceData.GetBaseIndex(), InSourceData);
@@ -202,7 +202,7 @@ void FNiagaraCopyIntBufferRegionCS::SetParameters(
 
 void FNiagaraCopyIntBufferRegionCS::UnbindBuffers(FRHICommandList& RHICmdList)
 {
-	FComputeShaderRHIParamRef ComputeShaderRHI = GetComputeShader();
+	FRHIComputeShader* ComputeShaderRHI = GetComputeShader();
 	if (SourceData.IsBound())
 	{
 		RHICmdList.SetShaderResourceViewParameter(ComputeShaderRHI, SourceData.GetBaseIndex(), nullptr);
diff --git a/Engine/Plugins/Runtime/AR/Apple/AppleARKit/Source/AppleARKit/Private/AppleARKitVideoOverlay.cpp b/Engine/Plugins/Runtime/AR/Apple/AppleARKit/Source/AppleARKit/Private/AppleARKitVideoOverlay.cpp
--- a/Engine/Plugins/Runtime/AR/Apple/AppleARKit/Source/AppleARKit/Private/AppleARKitVideoOverlay.cpp
+++ b/Engine/Plugins/Runtime/AR/Apple/AppleARKit/Source/AppleARKit/Private/AppleARKitVideoOverlay.cpp
@@ -290,7 +290,7 @@ public:
 
 	void SetParameters(FRHICommandList& RHICmdList, const FSceneView View)
 	{
-		const FVertexShaderRHIParamRef ShaderRHI = GetVertexShader();
+		FRHIVertexShader* ShaderRHI = GetVertexShader();
 		FMaterialShader::SetViewParameters(RHICmdList, ShaderRHI, View, View.ViewUniformBuffer);
 	}
 
@@ -343,7 +343,7 @@ public:
 
 	void SetParameters(FRHICommandList& RHICmdList, const FSceneView View, const FMaterialRenderProxy* Material)
 	{
-		const FPixelShaderRHIParamRef ShaderRHI = GetPixelShader();
+		FRHIPixelShader* ShaderRHI = GetPixelShader();
 		FMaterialShader::SetParameters(RHICmdList, ShaderRHI, Material, *Material->GetMaterial(View.GetFeatureLevel()), View, View.ViewUniformBuffer, ESceneTextureSetupMode::None);
 
 		for (uint32 InputIter = 0; InputIter < ePId_Input_MAX; ++InputIter)
diff --git a/Engine/Plugins/Runtime/Oculus/OculusVR/Source/OculusHMD/Private/OculusHMD_StressTester.cpp b/Engine/Plugins/Runtime/Oculus/OculusVR/Source/OculusHMD/Private/OculusHMD_StressTester.cpp
--- a/Engine/Plugins/Runtime/Oculus/OculusVR/Source/OculusHMD/Private/OculusHMD_StressTester.cpp
+++ b/Engine/Plugins/Runtime/Oculus/OculusVR/Source/OculusHMD/Private/OculusHMD_StressTester.cpp
@@ -142,7 +142,7 @@ public:
 	//This function is required to let us bind our runtime surface to the shader using an SRV.
 	void SetSurfaces(FRHICommandList& RHICmdList, FShaderResourceViewRHIRef TextureParameterSRV)
 	{
-		FPixelShaderRHIParamRef PixelShaderRHI = GetPixelShader();
+		FRHIPixelShader* PixelShaderRHI = GetPixelShader();
 
 		if (TextureParameter.IsBound())
 		{
@@ -166,7 +166,7 @@ public:
 	//This is used to clean up the buffer binds after each invocation to let them be changed and used elsewhere if needed.
 	void UnbindBuffers(FRHICommandList& RHICmdList)
 	{
-		FPixelShaderRHIParamRef PixelShaderRHI = GetPixelShader();
+		FRHIPixelShader* PixelShaderRHI = GetPixelShader();
 
 		if (TextureParameter.IsBound())
 		{
diff --git a/Engine/Plugins/Runtime/WindowsMixedReality/Source/WindowsMixedRealityHMD/Private/WindowsMixedRealityHMD.cpp b/Engine/Plugins/Runtime/WindowsMixedReality/Source/WindowsMixedRealityHMD/Private/WindowsMixedRealityHMD.cpp
--- a/Engine/Plugins/Runtime/WindowsMixedReality/Source/WindowsMixedRealityHMD/Private/WindowsMixedRealityHMD.cpp
+++ b/Engine/Plugins/Runtime/WindowsMixedReality/Source/WindowsMixedRealityHMD/Private/WindowsMixedRealityHMD.cpp
@@ -51,7 +51,7 @@ public:
 
 	void SetParameters(FRHICommandList& RHICmdList, float farPlaneDistanceValue, FTextureRHIParamRef DepthTexture)
 	{
-		FPixelShaderRHIParamRef PixelShaderRHI = GetPixelShader();
+		FRHIPixelShader* PixelShaderRHI = GetPixelShader();
 
 		SetShaderValue(RHICmdList, PixelShaderRHI, FarPlaneDistance, farPlaneDistanceValue);
 
diff --git a/Engine/Source/Editor/UnrealEd/Private/MeshPaintRendering.cpp b/Engine/Source/Editor/UnrealEd/Private/MeshPaintRendering.cpp
--- a/Engine/Source/Editor/UnrealEd/Private/MeshPaintRendering.cpp
+++ b/Engine/Source/Editor/UnrealEd/Private/MeshPaintRendering.cpp
@@ -111,7 +111,7 @@ namespace MeshPaintRendering
 
 		void SetParameters(FRHICommandList& RHICmdList, const float InGamma, const FMeshPaintShaderParameters& InShaderParams )
 		{
-			const FPixelShaderRHIParamRef ShaderRHI = GetPixelShader();
+			FRHIPixelShader* ShaderRHI = GetPixelShader();
 
 			SetTextureParameter(
 				RHICmdList, 
@@ -274,7 +274,7 @@ namespace MeshPaintRendering
 
 		void SetParameters(FRHICommandList& RHICmdList, const float InGamma, const FMeshPaintDilateShaderParameters& InShaderParams )
 		{
-			const FPixelShaderRHIParamRef ShaderRHI = GetPixelShader();
+			FRHIPixelShader* ShaderRHI = GetPixelShader();
 
 			SetTextureParameter(
 				RHICmdList, 
diff --git a/Engine/Source/Editor/UnrealEd/Private/NormalMapPreview.cpp b/Engine/Source/Editor/UnrealEd/Private/NormalMapPreview.cpp
--- a/Engine/Source/Editor/UnrealEd/Private/NormalMapPreview.cpp
+++ b/Engine/Source/Editor/UnrealEd/Private/NormalMapPreview.cpp
@@ -49,7 +49,7 @@ public:
 	 */
 	void SetParameters(FRHICommandList& RHICmdList, const FTexture* NormalMapTexture)
 	{
-		FPixelShaderRHIParamRef PixelShaderRHI = GetPixelShader();
+		FRHIPixelShader* PixelShaderRHI = GetPixelShader();
 		SetTextureParameter(RHICmdList, PixelShaderRHI,Texture,TextureSampler,NormalMapTexture);
 	}
 
diff --git a/Engine/Source/Runtime/Apple/MetalRHI/Private/MetalCommands.cpp b/Engine/Source/Runtime/Apple/MetalRHI/Private/MetalCommands.cpp
--- a/Engine/Source/Runtime/Apple/MetalRHI/Private/MetalCommands.cpp
+++ b/Engine/Source/Runtime/Apple/MetalRHI/Private/MetalCommands.cpp
@@ -108,7 +108,7 @@ void FMetalDynamicRHI::RHISetStreamOutTargets(uint32 NumTargets, const FVertexBu
 
 }
 
-void FMetalRHICommandContext::RHISetComputeShader(FComputeShaderRHIParamRef ComputeShaderRHI)
+void FMetalRHICommandContext::RHISetComputeShader(FRHIComputeShader* ComputeShaderRHI)
 {
 	@autoreleasepool {
 	FMetalComputeShader* ComputeShader = ResourceCast(ComputeShaderRHI);
@@ -246,7 +246,7 @@ void FMetalRHICommandContext::RHISetGraphicsPipelineState(FRHIGraphicsPipelineSt
 	}
 }
 
-void FMetalRHICommandContext::RHISetUAVParameter(FComputeShaderRHIParamRef ComputeShaderRHI, uint32 UAVIndex, FRHIUnorderedAccessView* UAVRHI)
+void FMetalRHICommandContext::RHISetUAVParameter(FRHIComputeShader* ComputeShaderRHI, uint32 UAVIndex, FRHIUnorderedAccessView* UAVRHI)
 {
 	@autoreleasepool {
 	FMetalUnorderedAccessView* UAV = ResourceCast(UAVRHI);
@@ -254,7 +254,7 @@ void FMetalRHICommandContext::RHISetUAVParameter(FComputeShaderRHIParamRef Compu
 	}
 }
 
-void FMetalRHICommandContext::RHISetUAVParameter(FComputeShaderRHIParamRef ComputeShaderRHI,uint32 UAVIndex, FRHIUnorderedAccessView* UAVRHI, uint32 InitialCount)
+void FMetalRHICommandContext::RHISetUAVParameter(FRHIComputeShader* ComputeShaderRHI,uint32 UAVIndex, FRHIUnorderedAccessView* UAVRHI, uint32 InitialCount)
 {
 	@autoreleasepool {
 		FMetalUnorderedAccessView* UAV = ResourceCast(UAVRHI);
@@ -263,7 +263,7 @@ void FMetalRHICommandContext::RHISetUAVParameter(FComputeShaderRHIParamRef Compu
 }
 
 
-void FMetalRHICommandContext::RHISetShaderTexture(FVertexShaderRHIParamRef VertexShaderRHI, uint32 TextureIndex, FTextureRHIParamRef NewTextureRHI)
+void FMetalRHICommandContext::RHISetShaderTexture(FRHIVertexShader* VertexShaderRHI, uint32 TextureIndex, FTextureRHIParamRef NewTextureRHI)
 {
 	@autoreleasepool {
 	FMetalSurface* Surface = GetMetalSurfaceFromRHITexture(NewTextureRHI);
@@ -286,7 +286,7 @@ void FMetalRHICommandContext::RHISetShaderTexture(FVertexShaderRHIParamRef Verte
 	}
 }
 
-void FMetalRHICommandContext::RHISetShaderTexture(FHullShaderRHIParamRef HullShader, uint32 TextureIndex, FTextureRHIParamRef NewTextureRHI)
+void FMetalRHICommandContext::RHISetShaderTexture(FRHIHullShader* HullShader, uint32 TextureIndex, FTextureRHIParamRef NewTextureRHI)
 {
 #if PLATFORM_SUPPORTS_TESSELLATION_SHADERS
 	@autoreleasepool {
@@ -311,7 +311,7 @@ void FMetalRHICommandContext::RHISetShaderTexture(FHullShaderRHIParamRef HullSha
 #endif
 }
 
-void FMetalRHICommandContext::RHISetShaderTexture(FDomainShaderRHIParamRef DomainShader, uint32 TextureIndex, FTextureRHIParamRef NewTextureRHI)
+void FMetalRHICommandContext::RHISetShaderTexture(FRHIDomainShader* DomainShader, uint32 TextureIndex, FTextureRHIParamRef NewTextureRHI)
 {
 #if PLATFORM_SUPPORTS_TESSELLATION_SHADERS
 	@autoreleasepool {
@@ -336,13 +336,13 @@ void FMetalRHICommandContext::RHISetShaderTexture(FDomainShaderRHIParamRef Domai
 #endif
 }
 
-void FMetalRHICommandContext::RHISetShaderTexture(FGeometryShaderRHIParamRef GeometryShader, uint32 TextureIndex, FTextureRHIParamRef NewTextureRHI)
+void FMetalRHICommandContext::RHISetShaderTexture(FRHIGeometryShader* GeometryShader, uint32 TextureIndex, FTextureRHIParamRef NewTextureRHI)
 {
 	NOT_SUPPORTED("RHISetShaderTexture-Geometry");
 
 }
 
-void FMetalRHICommandContext::RHISetShaderTexture(FPixelShaderRHIParamRef PixelShader, uint32 TextureIndex, FTextureRHIParamRef NewTextureRHI)
+void FMetalRHICommandContext::RHISetShaderTexture(FRHIPixelShader* PixelShader, uint32 TextureIndex, FTextureRHIParamRef NewTextureRHI)
 {
 	@autoreleasepool {
 	FMetalSurface* Surface = GetMetalSurfaceFromRHITexture(NewTextureRHI);
@@ -365,7 +365,7 @@ void FMetalRHICommandContext::RHISetShaderTexture(FPixelShaderRHIParamRef PixelS
 	}
 }
 
-void FMetalRHICommandContext::RHISetShaderTexture(FComputeShaderRHIParamRef ComputeShader, uint32 TextureIndex, FTextureRHIParamRef NewTextureRHI)
+void FMetalRHICommandContext::RHISetShaderTexture(FRHIComputeShader* ComputeShader, uint32 TextureIndex, FTextureRHIParamRef NewTextureRHI)
 {
 	@autoreleasepool {
 	FMetalSurface* Surface = GetMetalSurfaceFromRHITexture(NewTextureRHI);
@@ -389,7 +389,7 @@ void FMetalRHICommandContext::RHISetShaderTexture(FComputeShaderRHIParamRef Comp
 }
 
 
-void FMetalRHICommandContext::RHISetShaderResourceViewParameter(FVertexShaderRHIParamRef VertexShaderRHI, uint32 TextureIndex, FRHIShaderResourceView* SRVRHI)
+void FMetalRHICommandContext::RHISetShaderResourceViewParameter(FRHIVertexShader* VertexShaderRHI, uint32 TextureIndex, FRHIShaderResourceView* SRVRHI)
 {
 	@autoreleasepool {
 	FMetalShaderResourceView* SRV = ResourceCast(SRVRHI);
@@ -397,7 +397,7 @@ void FMetalRHICommandContext::RHISetShaderResourceViewParameter(FVertexShaderRHI
 	}
 }
 
-void FMetalRHICommandContext::RHISetShaderResourceViewParameter(FHullShaderRHIParamRef HullShaderRHI,uint32 TextureIndex, FRHIShaderResourceView* SRVRHI)
+void FMetalRHICommandContext::RHISetShaderResourceViewParameter(FRHIHullShader* HullShaderRHI,uint32 TextureIndex, FRHIShaderResourceView* SRVRHI)
 {
 #if PLATFORM_SUPPORTS_TESSELLATION_SHADERS
 	@autoreleasepool {
@@ -407,7 +407,7 @@ void FMetalRHICommandContext::RHISetShaderResourceViewParameter(FHullShaderRHIPa
 #endif
 }
 
-void FMetalRHICommandContext::RHISetShaderResourceViewParameter(FDomainShaderRHIParamRef DomainShaderRHI,uint32 TextureIndex, FRHIShaderResourceView* SRVRHI)
+void FMetalRHICommandContext::RHISetShaderResourceViewParameter(FRHIDomainShader* DomainShaderRHI,uint32 TextureIndex, FRHIShaderResourceView* SRVRHI)
 {
 #if PLATFORM_SUPPORTS_TESSELLATION_SHADERS
 	@autoreleasepool {
@@ -417,12 +417,12 @@ void FMetalRHICommandContext::RHISetShaderResourceViewParameter(FDomainShaderRHI
 #endif
 }
 
-void FMetalRHICommandContext::RHISetShaderResourceViewParameter(FGeometryShaderRHIParamRef GeometryShaderRHI,uint32 TextureIndex, FRHIShaderResourceView* SRVRHI)
+void FMetalRHICommandContext::RHISetShaderResourceViewParameter(FRHIGeometryShader* GeometryShaderRHI,uint32 TextureIndex, FRHIShaderResourceView* SRVRHI)
 {
 	NOT_SUPPORTED("RHISetShaderResourceViewParameter");
 }
 
-void FMetalRHICommandContext::RHISetShaderResourceViewParameter(FPixelShaderRHIParamRef PixelShaderRHI,uint32 TextureIndex, FRHIShaderResourceView* SRVRHI)
+void FMetalRHICommandContext::RHISetShaderResourceViewParameter(FRHIPixelShader* PixelShaderRHI,uint32 TextureIndex, FRHIShaderResourceView* SRVRHI)
 {
 	@autoreleasepool {
 	FMetalShaderResourceView* SRV = ResourceCast(SRVRHI);
@@ -430,7 +430,7 @@ void FMetalRHICommandContext::RHISetShaderResourceViewParameter(FPixelShaderRHIP
 	}
 }
 
-void FMetalRHICommandContext::RHISetShaderResourceViewParameter(FComputeShaderRHIParamRef ComputeShaderRHI,uint32 TextureIndex, FRHIShaderResourceView* SRVRHI)
+void FMetalRHICommandContext::RHISetShaderResourceViewParameter(FRHIComputeShader* ComputeShaderRHI,uint32 TextureIndex, FRHIShaderResourceView* SRVRHI)
 {
 	@autoreleasepool {
 	FMetalShaderResourceView* SRV = ResourceCast(SRVRHI);
@@ -439,7 +439,7 @@ void FMetalRHICommandContext::RHISetShaderResourceViewParameter(FComputeShaderRH
 }
 
 
-void FMetalRHICommandContext::RHISetShaderSampler(FVertexShaderRHIParamRef VertexShaderRHI, uint32 SamplerIndex, FRHISamplerState* NewStateRHI)
+void FMetalRHICommandContext::RHISetShaderSampler(FRHIVertexShader* VertexShaderRHI, uint32 SamplerIndex, FRHISamplerState* NewStateRHI)
 {
 	@autoreleasepool {
 	FMetalSamplerState* NewState = ResourceCast(NewStateRHI);
@@ -448,7 +448,7 @@ void FMetalRHICommandContext::RHISetShaderSampler(FVertexShaderRHIParamRef Verte
 	}
 }
 
-void FMetalRHICommandContext::RHISetShaderSampler(FHullShaderRHIParamRef HullShader, uint32 SamplerIndex, FRHISamplerState* NewStateRHI)
+void FMetalRHICommandContext::RHISetShaderSampler(FRHIHullShader* HullShader, uint32 SamplerIndex, FRHISamplerState* NewStateRHI)
 {
 #if PLATFORM_SUPPORTS_TESSELLATION_SHADERS
 	@autoreleasepool {
@@ -459,7 +459,7 @@ void FMetalRHICommandContext::RHISetShaderSampler(FHullShaderRHIParamRef HullSha
 #endif
 }
 
-void FMetalRHICommandContext::RHISetShaderSampler(FDomainShaderRHIParamRef DomainShader, uint32 SamplerIndex, FRHISamplerState* NewStateRHI)
+void FMetalRHICommandContext::RHISetShaderSampler(FRHIDomainShader* DomainShader, uint32 SamplerIndex, FRHISamplerState* NewStateRHI)
 {
 #if PLATFORM_SUPPORTS_TESSELLATION_SHADERS
 	@autoreleasepool {
@@ -470,12 +470,12 @@ void FMetalRHICommandContext::RHISetShaderSampler(FDomainShaderRHIParamRef Domai
 #endif
 }
 
-void FMetalRHICommandContext::RHISetShaderSampler(FGeometryShaderRHIParamRef GeometryShader, uint32 SamplerIndex, FRHISamplerState* NewStateRHI)
+void FMetalRHICommandContext::RHISetShaderSampler(FRHIGeometryShader* GeometryShader, uint32 SamplerIndex, FRHISamplerState* NewStateRHI)
 {
 	NOT_SUPPORTED("RHISetSamplerState-Geometry");
 }
 
-void FMetalRHICommandContext::RHISetShaderSampler(FPixelShaderRHIParamRef PixelShader, uint32 SamplerIndex, FRHISamplerState* NewStateRHI)
+void FMetalRHICommandContext::RHISetShaderSampler(FRHIPixelShader* PixelShader, uint32 SamplerIndex, FRHISamplerState* NewStateRHI)
 {
 	@autoreleasepool {
 	FMetalSamplerState* NewState = ResourceCast(NewStateRHI);
@@ -484,7 +484,7 @@ void FMetalRHICommandContext::RHISetShaderSampler(FPixelShaderRHIParamRef PixelS
 	}
 }
 
-void FMetalRHICommandContext::RHISetShaderSampler(FComputeShaderRHIParamRef ComputeShader, uint32 SamplerIndex, FRHISamplerState* NewStateRHI)
+void FMetalRHICommandContext::RHISetShaderSampler(FRHIComputeShader* ComputeShader, uint32 SamplerIndex, FRHISamplerState* NewStateRHI)
 {
 	@autoreleasepool {
 	FMetalSamplerState* NewState = ResourceCast(NewStateRHI);
@@ -493,27 +493,27 @@ void FMetalRHICommandContext::RHISetShaderSampler(FComputeShaderRHIParamRef Comp
 	}
 }
 
-void FMetalRHICommandContext::RHISetShaderParameter(FVertexShaderRHIParamRef VertexShaderRHI, uint32 BufferIndex, uint32 BaseIndex, uint32 NumBytes, const void* NewValue)
+void FMetalRHICommandContext::RHISetShaderParameter(FRHIVertexShader* VertexShaderRHI, uint32 BufferIndex, uint32 BaseIndex, uint32 NumBytes, const void* NewValue)
 {
 	@autoreleasepool {
 	Context->GetCurrentState().GetShaderParameters(EMetalShaderStages::Vertex).Set(BufferIndex, BaseIndex, NumBytes, NewValue);
 	}
 }
 
-void FMetalRHICommandContext::RHISetShaderParameter(FHullShaderRHIParamRef HullShaderRHI, uint32 BufferIndex, uint32 BaseIndex, uint32 NumBytes, const void* NewValue)
+void FMetalRHICommandContext::RHISetShaderParameter(FRHIHullShader* HullShaderRHI, uint32 BufferIndex, uint32 BaseIndex, uint32 NumBytes, const void* NewValue)
 {
 	// Just ignore Hull shader parameter sets - none of our Hull shaders have any loose parameters to bind.
 	// @todo Whenever we do put a shader parameter into a hull shader we'll need to map it into the vertex-shader parameter buffer so that it can be set on the device.
 }
 
-void FMetalRHICommandContext::RHISetShaderParameter(FPixelShaderRHIParamRef PixelShaderRHI, uint32 BufferIndex, uint32 BaseIndex, uint32 NumBytes, const void* NewValue)
+void FMetalRHICommandContext::RHISetShaderParameter(FRHIPixelShader* PixelShaderRHI, uint32 BufferIndex, uint32 BaseIndex, uint32 NumBytes, const void* NewValue)
 {
 	@autoreleasepool {
 	Context->GetCurrentState().GetShaderParameters(EMetalShaderStages::Pixel).Set(BufferIndex, BaseIndex, NumBytes, NewValue);
 	}
 }
 
-void FMetalRHICommandContext::RHISetShaderParameter(FDomainShaderRHIParamRef DomainShaderRHI, uint32 BufferIndex, uint32 BaseIndex, uint32 NumBytes, const void* NewValue)
+void FMetalRHICommandContext::RHISetShaderParameter(FRHIDomainShader* DomainShaderRHI, uint32 BufferIndex, uint32 BaseIndex, uint32 NumBytes, const void* NewValue)
 {
 #if PLATFORM_SUPPORTS_TESSELLATION_SHADERS
 	@autoreleasepool {
@@ -522,19 +522,19 @@ void FMetalRHICommandContext::RHISetShaderParameter(FDomainShaderRHIParamRef Dom
 #endif
 }
 
-void FMetalRHICommandContext::RHISetShaderParameter(FGeometryShaderRHIParamRef GeometryShaderRHI, uint32 BufferIndex, uint32 BaseIndex, uint32 NumBytes, const void* NewValue)
+void FMetalRHICommandContext::RHISetShaderParameter(FRHIGeometryShader* GeometryShaderRHI, uint32 BufferIndex, uint32 BaseIndex, uint32 NumBytes, const void* NewValue)
 {
 	NOT_SUPPORTED("RHISetShaderParameter-Geometry");
 }
 
-void FMetalRHICommandContext::RHISetShaderParameter(FComputeShaderRHIParamRef ComputeShaderRHI,uint32 BufferIndex, uint32 BaseIndex, uint32 NumBytes, const void* NewValue)
+void FMetalRHICommandContext::RHISetShaderParameter(FRHIComputeShader* ComputeShaderRHI,uint32 BufferIndex, uint32 BaseIndex, uint32 NumBytes, const void* NewValue)
 {
 	@autoreleasepool {
 	Context->GetCurrentState().GetShaderParameters(EMetalShaderStages::Compute).Set(BufferIndex, BaseIndex, NumBytes, NewValue);
 	}
 }
 
-void FMetalRHICommandContext::RHISetShaderUniformBuffer(FVertexShaderRHIParamRef VertexShaderRHI, uint32 BufferIndex, FUniformBufferRHIParamRef BufferRHI)
+void FMetalRHICommandContext::RHISetShaderUniformBuffer(FRHIVertexShader* VertexShaderRHI, uint32 BufferIndex, FUniformBufferRHIParamRef BufferRHI)
 {
 	@autoreleasepool {
 	FMetalVertexShader* VertexShader = ResourceCast(VertexShaderRHI);
@@ -550,7 +550,7 @@ void FMetalRHICommandContext::RHISetShaderUniformBuffer(FVertexShaderRHIParamRef
 	}
 }
 
-void FMetalRHICommandContext::RHISetShaderUniformBuffer(FHullShaderRHIParamRef HullShaderRHI, uint32 BufferIndex, FUniformBufferRHIParamRef BufferRHI)
+void FMetalRHICommandContext::RHISetShaderUniformBuffer(FRHIHullShader* HullShaderRHI, uint32 BufferIndex, FUniformBufferRHIParamRef BufferRHI)
 {
 #if PLATFORM_SUPPORTS_TESSELLATION_SHADERS
 	@autoreleasepool {
@@ -568,7 +568,7 @@ void FMetalRHICommandContext::RHISetShaderUniformBuffer(FHullShaderRHIParamRef H
 #endif
 }
 
-void FMetalRHICommandContext::RHISetShaderUniformBuffer(FDomainShaderRHIParamRef DomainShaderRHI, uint32 BufferIndex, FUniformBufferRHIParamRef BufferRHI)
+void FMetalRHICommandContext::RHISetShaderUniformBuffer(FRHIDomainShader* DomainShaderRHI, uint32 BufferIndex, FUniformBufferRHIParamRef BufferRHI)
 {
 #if PLATFORM_SUPPORTS_TESSELLATION_SHADERS
 	@autoreleasepool {
@@ -586,12 +586,12 @@ void FMetalRHICommandContext::RHISetShaderUniformBuffer(FDomainShaderRHIParamRef
 #endif
 }
 
-void FMetalRHICommandContext::RHISetShaderUniformBuffer(FGeometryShaderRHIParamRef GeometryShader, uint32 BufferIndex, FUniformBufferRHIParamRef BufferRHI)
+void FMetalRHICommandContext::RHISetShaderUniformBuffer(FRHIGeometryShader* GeometryShader, uint32 BufferIndex, FUniformBufferRHIParamRef BufferRHI)
 {
 	NOT_SUPPORTED("RHISetShaderUniformBuffer-Geometry");
 }
 
-void FMetalRHICommandContext::RHISetShaderUniformBuffer(FPixelShaderRHIParamRef PixelShaderRHI, uint32 BufferIndex, FUniformBufferRHIParamRef BufferRHI)
+void FMetalRHICommandContext::RHISetShaderUniformBuffer(FRHIPixelShader* PixelShaderRHI, uint32 BufferIndex, FUniformBufferRHIParamRef BufferRHI)
 {
 	@autoreleasepool {
 	FMetalPixelShader* PixelShader = ResourceCast(PixelShaderRHI);
@@ -607,7 +607,7 @@ void FMetalRHICommandContext::RHISetShaderUniformBuffer(FPixelShaderRHIParamRef
 	}
 }
 
-void FMetalRHICommandContext::RHISetShaderUniformBuffer(FComputeShaderRHIParamRef ComputeShaderRHI, uint32 BufferIndex, FUniformBufferRHIParamRef BufferRHI)
+void FMetalRHICommandContext::RHISetShaderUniformBuffer(FRHIComputeShader* ComputeShaderRHI, uint32 BufferIndex, FUniformBufferRHIParamRef BufferRHI)
 {
 	@autoreleasepool {
 	FMetalComputeShader* ComputeShader = ResourceCast(ComputeShaderRHI);
diff --git a/Engine/Source/Runtime/Apple/MetalRHI/Private/MetalRHIContext.cpp b/Engine/Source/Runtime/Apple/MetalRHI/Private/MetalRHIContext.cpp
--- a/Engine/Source/Runtime/Apple/MetalRHI/Private/MetalRHIContext.cpp
+++ b/Engine/Source/Runtime/Apple/MetalRHI/Private/MetalRHIContext.cpp
@@ -102,7 +102,7 @@ void FMetalRHIComputeContext::RHISetAsyncComputeBudget(EAsyncComputeBudget Budge
 	FMetalRHICommandContext::RHISetAsyncComputeBudget(Budget);
 }
 
-void FMetalRHIComputeContext::RHISetComputeShader(FComputeShaderRHIParamRef ComputeShader)
+void FMetalRHIComputeContext::RHISetComputeShader(FRHIComputeShader* ComputeShader)
 {
 	if (!Context->GetCurrentCommandBuffer())
 	{
diff --git a/Engine/Source/Runtime/Apple/MetalRHI/Private/MetalShaders.cpp b/Engine/Source/Runtime/Apple/MetalRHI/Private/MetalShaders.cpp
--- a/Engine/Source/Runtime/Apple/MetalRHI/Private/MetalShaders.cpp
+++ b/Engine/Source/Runtime/Apple/MetalRHI/Private/MetalShaders.cpp
@@ -1585,11 +1585,11 @@ FRHIShaderLibraryRef FMetalDynamicRHI::RHICreateShaderLibrary(EShaderPlatform Pl
 
 FBoundShaderStateRHIRef FMetalDynamicRHI::RHICreateBoundShaderState(
 	FRHIVertexDeclaration* VertexDeclarationRHI,
-	FVertexShaderRHIParamRef VertexShaderRHI, 
-	FHullShaderRHIParamRef HullShaderRHI, 
-	FDomainShaderRHIParamRef DomainShaderRHI, 
-	FPixelShaderRHIParamRef PixelShaderRHI,
-	FGeometryShaderRHIParamRef GeometryShaderRHI
+	FRHIVertexShader* VertexShaderRHI,
+	FRHIHullShader* HullShaderRHI,
+	FRHIDomainShader* DomainShaderRHI,
+	FRHIPixelShader* PixelShaderRHI,
+	FRHIGeometryShader* GeometryShaderRHI
 	)
 {
 	NOT_SUPPORTED("RHICreateBoundShaderState");
diff --git a/Engine/Source/Runtime/Apple/MetalRHI/Public/MetalRHI.h b/Engine/Source/Runtime/Apple/MetalRHI/Public/MetalRHI.h
--- a/Engine/Source/Runtime/Apple/MetalRHI/Public/MetalRHI.h
+++ b/Engine/Source/Runtime/Apple/MetalRHI/Public/MetalRHI.h
@@ -59,7 +59,7 @@ public:
 	virtual FGeometryShaderRHIRef RHICreateGeometryShaderWithStreamOutput(const FStreamOutElementList& ElementList, uint32 NumStrides, const uint32* Strides, int32 RasterizedStream, FRHIShaderLibrary* Library, FSHAHash Hash) final override;
 	virtual FComputeShaderRHIRef RHICreateComputeShader(const TArray<uint8>& Code) final override;
 	virtual FComputeShaderRHIRef RHICreateComputeShader(FRHIShaderLibrary* Library, FSHAHash Hash) final override;
-	virtual FBoundShaderStateRHIRef RHICreateBoundShaderState(FRHIVertexDeclaration* VertexDeclaration, FVertexShaderRHIParamRef VertexShader, FHullShaderRHIParamRef HullShader, FDomainShaderRHIParamRef DomainShader, FPixelShaderRHIParamRef PixelShader, FGeometryShaderRHIParamRef GeometryShader) final override;
+	virtual FBoundShaderStateRHIRef RHICreateBoundShaderState(FRHIVertexDeclaration* VertexDeclaration, FRHIVertexShader* VertexShader, FRHIHullShader* HullShader, FRHIDomainShader* DomainShader, FRHIPixelShader* PixelShader, FRHIGeometryShader* GeometryShader) final override;
 	virtual FGraphicsPipelineStateRHIRef RHICreateGraphicsPipelineState(const FGraphicsPipelineStateInitializer& Initializer) final override;
 	virtual TRefCountPtr<FRHIComputePipelineState> RHICreateComputePipelineState(FRHIComputeShader* ComputeShader) final override;
 	virtual FUniformBufferRHIRef RHICreateUniformBuffer(const void* Contents, const FRHIUniformBufferLayout& Layout, EUniformBufferUsage Usage, EUniformBufferValidation Validation) final override;
diff --git a/Engine/Source/Runtime/Apple/MetalRHI/Public/MetalRHIContext.h b/Engine/Source/Runtime/Apple/MetalRHI/Public/MetalRHIContext.h
--- a/Engine/Source/Runtime/Apple/MetalRHI/Public/MetalRHIContext.h
+++ b/Engine/Source/Runtime/Apple/MetalRHI/Public/MetalRHIContext.h
@@ -27,7 +27,7 @@ public:
 	 *Sets the current compute shader.  Mostly for compliance with platforms
 	 *that require shader setting before resource binding.
 	 */
-	virtual void RHISetComputeShader(FComputeShaderRHIParamRef ComputeShader) override;
+	virtual void RHISetComputeShader(FRHIComputeShader* ComputeShader) override;
 	
 	virtual void RHISetComputePipelineState(FRHIComputePipelineState* ComputePipelineState) override;
 	
@@ -114,22 +114,22 @@ public:
 	virtual void RHISetGraphicsPipelineState(FRHIGraphicsPipelineState* GraphicsState) final override;
 	
 	/** Set the shader resource view of a surface.  This is used for binding TextureMS parameter types that need a multi sampled view. */
-	virtual void RHISetShaderTexture(FVertexShaderRHIParamRef VertexShader, uint32 TextureIndex, FTextureRHIParamRef NewTexture) final override;
+	virtual void RHISetShaderTexture(FRHIVertexShader* VertexShader, uint32 TextureIndex, FTextureRHIParamRef NewTexture) final override;
 	
 	/** Set the shader resource view of a surface.  This is used for binding TextureMS parameter types that need a multi sampled view. */
-	virtual void RHISetShaderTexture(FHullShaderRHIParamRef HullShader, uint32 TextureIndex, FTextureRHIParamRef NewTexture) final override;
+	virtual void RHISetShaderTexture(FRHIHullShader* HullShader, uint32 TextureIndex, FTextureRHIParamRef NewTexture) final override;
 	
 	/** Set the shader resource view of a surface.  This is used for binding TextureMS parameter types that need a multi sampled view. */
-	virtual void RHISetShaderTexture(FDomainShaderRHIParamRef DomainShader, uint32 TextureIndex, FTextureRHIParamRef NewTexture) final override;
+	virtual void RHISetShaderTexture(FRHIDomainShader* DomainShader, uint32 TextureIndex, FTextureRHIParamRef NewTexture) final override;
 	
 	/** Set the shader resource view of a surface.  This is used for binding TextureMS parameter types that need a multi sampled view. */
-	virtual void RHISetShaderTexture(FGeometryShaderRHIParamRef GeometryShader, uint32 TextureIndex, FTextureRHIParamRef NewTexture) final override;
+	virtual void RHISetShaderTexture(FRHIGeometryShader* GeometryShader, uint32 TextureIndex, FTextureRHIParamRef NewTexture) final override;
 	
 	/** Set the shader resource view of a surface.  This is used for binding TextureMS parameter types that need a multi sampled view. */
-	virtual void RHISetShaderTexture(FPixelShaderRHIParamRef PixelShader, uint32 TextureIndex, FTextureRHIParamRef NewTexture) final override;
+	virtual void RHISetShaderTexture(FRHIPixelShader* PixelShader, uint32 TextureIndex, FTextureRHIParamRef NewTexture) final override;
 	
 	/** Set the shader resource view of a surface.  This is used for binding TextureMS parameter types that need a multi sampled view. */
-	virtual void RHISetShaderTexture(FComputeShaderRHIParamRef PixelShader, uint32 TextureIndex, FTextureRHIParamRef NewTexture) final override;
+	virtual void RHISetShaderTexture(FRHIComputeShader* PixelShader, uint32 TextureIndex, FTextureRHIParamRef NewTexture) final override;
 	
 	/**
 	 * Sets sampler state.
@@ -137,7 +137,7 @@ public:
 	 * @param SamplerIndex		The index of the sampler.
 	 * @param NewState			The new sampler state.
 	 */
-	virtual void RHISetShaderSampler(FComputeShaderRHIParamRef ComputeShader, uint32 SamplerIndex, FRHISamplerState* NewState) final override;
+	virtual void RHISetShaderSampler(FRHIComputeShader* ComputeShader, uint32 SamplerIndex, FRHISamplerState* NewState) final override;
 	
 	/**
 	 * Sets sampler state.
@@ -145,7 +145,7 @@ public:
 	 * @param SamplerIndex		The index of the sampler.
 	 * @param NewState			The new sampler state.
 	 */
-	virtual void RHISetShaderSampler(FVertexShaderRHIParamRef VertexShader, uint32 SamplerIndex, FRHISamplerState* NewState) final override;
+	virtual void RHISetShaderSampler(FRHIVertexShader* VertexShader, uint32 SamplerIndex, FRHISamplerState* NewState) final override;
 	
 	/**
 	 * Sets sampler state.
@@ -153,7 +153,7 @@ public:
 	 * @param SamplerIndex		The index of the sampler.
 	 * @param NewState			The new sampler state.
 	 */
-	virtual void RHISetShaderSampler(FGeometryShaderRHIParamRef GeometryShader, uint32 SamplerIndex, FRHISamplerState* NewState) final override;
+	virtual void RHISetShaderSampler(FRHIGeometryShader* GeometryShader, uint32 SamplerIndex, FRHISamplerState* NewState) final override;
 	
 	/**
 	 * Sets sampler state.
@@ -161,7 +161,7 @@ public:
 	 * @param SamplerIndex		The index of the sampler.
 	 * @param NewState			The new sampler state.
 	 */
-	virtual void RHISetShaderSampler(FDomainShaderRHIParamRef DomainShader, uint32 SamplerIndex, FRHISamplerState* NewState) final override;
+	virtual void RHISetShaderSampler(FRHIDomainShader* DomainShader, uint32 SamplerIndex, FRHISamplerState* NewState) final override;
 	
 	/**
 	 * Sets sampler state.
@@ -169,7 +169,7 @@ public:
 	 * @param SamplerIndex		The index of the sampler.
 	 * @param NewState			The new sampler state.
 	 */
-	virtual void RHISetShaderSampler(FHullShaderRHIParamRef HullShader, uint32 SamplerIndex, FRHISamplerState* NewState) final override;
+	virtual void RHISetShaderSampler(FRHIHullShader* HullShader, uint32 SamplerIndex, FRHISamplerState* NewState) final override;
 	
 	/**
 	 * Sets sampler state.
@@ -177,49 +177,49 @@ public:
 	 * @param SamplerIndex		The index of the sampler.
 	 * @param NewState			The new sampler state.
 	 */
-	virtual void RHISetShaderSampler(FPixelShaderRHIParamRef PixelShader, uint32 SamplerIndex, FRHISamplerState* NewState) final override;
+	virtual void RHISetShaderSampler(FRHIPixelShader* PixelShader, uint32 SamplerIndex, FRHISamplerState* NewState) final override;
 	
 	/** Sets a compute shader UAV parameter. */
-	virtual void RHISetUAVParameter(FComputeShaderRHIParamRef ComputeShader, uint32 UAVIndex, FRHIUnorderedAccessView* UAV) final override;
+	virtual void RHISetUAVParameter(FRHIComputeShader* ComputeShader, uint32 UAVIndex, FRHIUnorderedAccessView* UAV) final override;
 	
 	/** Sets a compute shader UAV parameter and initial count */
-	virtual void RHISetUAVParameter(FComputeShaderRHIParamRef ComputeShader, uint32 UAVIndex, FRHIUnorderedAccessView* UAV, uint32 InitialCount) final override;
+	virtual void RHISetUAVParameter(FRHIComputeShader* ComputeShader, uint32 UAVIndex, FRHIUnorderedAccessView* UAV, uint32 InitialCount) final override;
 	
-	virtual void RHISetShaderResourceViewParameter(FPixelShaderRHIParamRef PixelShader, uint32 SamplerIndex, FRHIShaderResourceView* SRV) final override;
+	virtual void RHISetShaderResourceViewParameter(FRHIPixelShader* PixelShader, uint32 SamplerIndex, FRHIShaderResourceView* SRV) final override;
 	
-	virtual void RHISetShaderResourceViewParameter(FVertexShaderRHIParamRef VertexShader, uint32 SamplerIndex, FRHIShaderResourceView* SRV) final override;
+	virtual void RHISetShaderResourceViewParameter(FRHIVertexShader* VertexShader, uint32 SamplerIndex, FRHIShaderResourceView* SRV) final override;
 	
-	virtual void RHISetShaderResourceViewParameter(FComputeShaderRHIParamRef ComputeShader, uint32 SamplerIndex, FRHIShaderResourceView* SRV) final override;
+	virtual void RHISetShaderResourceViewParameter(FRHIComputeShader* ComputeShader, uint32 SamplerIndex, FRHIShaderResourceView* SRV) final override;
 	
-	virtual void RHISetShaderResourceViewParameter(FHullShaderRHIParamRef HullShader, uint32 SamplerIndex, FRHIShaderResourceView* SRV) final override;
+	virtual void RHISetShaderResourceViewParameter(FRHIHullShader* HullShader, uint32 SamplerIndex, FRHIShaderResourceView* SRV) final override;
 	
-	virtual void RHISetShaderResourceViewParameter(FDomainShaderRHIParamRef DomainShader, uint32 SamplerIndex, FRHIShaderResourceView* SRV) final override;
+	virtual void RHISetShaderResourceViewParameter(FRHIDomainShader* DomainShader, uint32 SamplerIndex, FRHIShaderResourceView* SRV) final override;
 	
-	virtual void RHISetShaderResourceViewParameter(FGeometryShaderRHIParamRef GeometryShader, uint32 SamplerIndex, FRHIShaderResourceView* SRV) final override;
+	virtual void RHISetShaderResourceViewParameter(FRHIGeometryShader* GeometryShader, uint32 SamplerIndex, FRHIShaderResourceView* SRV) final override;
 	
-	virtual void RHISetShaderUniformBuffer(FVertexShaderRHIParamRef VertexShader, uint32 BufferIndex, FUniformBufferRHIParamRef Buffer) final override;
+	virtual void RHISetShaderUniformBuffer(FRHIVertexShader* VertexShader, uint32 BufferIndex, FUniformBufferRHIParamRef Buffer) final override;
 	
-	virtual void RHISetShaderUniformBuffer(FHullShaderRHIParamRef HullShader, uint32 BufferIndex, FUniformBufferRHIParamRef Buffer) final override;
+	virtual void RHISetShaderUniformBuffer(FRHIHullShader* HullShader, uint32 BufferIndex, FUniformBufferRHIParamRef Buffer) final override;
 	
-	virtual void RHISetShaderUniformBuffer(FDomainShaderRHIParamRef DomainShader, uint32 BufferIndex, FUniformBufferRHIParamRef Buffer) final override;
+	virtual void RHISetShaderUniformBuffer(FRHIDomainShader* DomainShader, uint32 BufferIndex, FUniformBufferRHIParamRef Buffer) final override;
 	
-	virtual void RHISetShaderUniformBuffer(FGeometryShaderRHIParamRef GeometryShader, uint32 BufferIndex, FUniformBufferRHIParamRef Buffer) final override;
+	virtual void RHISetShaderUniformBuffer(FRHIGeometryShader* GeometryShader, uint32 BufferIndex, FUniformBufferRHIParamRef Buffer) final override;
 	
-	virtual void RHISetShaderUniformBuffer(FPixelShaderRHIParamRef PixelShader, uint32 BufferIndex, FUniformBufferRHIParamRef Buffer) final override;
+	virtual void RHISetShaderUniformBuffer(FRHIPixelShader* PixelShader, uint32 BufferIndex, FUniformBufferRHIParamRef Buffer) final override;
 	
-	virtual void RHISetShaderUniformBuffer(FComputeShaderRHIParamRef ComputeShader, uint32 BufferIndex, FUniformBufferRHIParamRef Buffer) final override;
+	virtual void RHISetShaderUniformBuffer(FRHIComputeShader* ComputeShader, uint32 BufferIndex, FUniformBufferRHIParamRef Buffer) final override;
 	
-	virtual void RHISetShaderParameter(FVertexShaderRHIParamRef VertexShader, uint32 BufferIndex, uint32 BaseIndex, uint32 NumBytes, const void* NewValue) final override;
+	virtual void RHISetShaderParameter(FRHIVertexShader* VertexShader, uint32 BufferIndex, uint32 BaseIndex, uint32 NumBytes, const void* NewValue) final override;
 	
-	virtual void RHISetShaderParameter(FPixelShaderRHIParamRef PixelShader, uint32 BufferIndex, uint32 BaseIndex, uint32 NumBytes, const void* NewValue) final override;
+	virtual void RHISetShaderParameter(FRHIPixelShader* PixelShader, uint32 BufferIndex, uint32 BaseIndex, uint32 NumBytes, const void* NewValue) final override;
 	
-	virtual void RHISetShaderParameter(FHullShaderRHIParamRef HullShader, uint32 BufferIndex, uint32 BaseIndex, uint32 NumBytes, const void* NewValue) final override;
+	virtual void RHISetShaderParameter(FRHIHullShader* HullShader, uint32 BufferIndex, uint32 BaseIndex, uint32 NumBytes, const void* NewValue) final override;
 	
-	virtual void RHISetShaderParameter(FDomainShaderRHIParamRef DomainShader, uint32 BufferIndex, uint32 BaseIndex, uint32 NumBytes, const void* NewValue) final override;
+	virtual void RHISetShaderParameter(FRHIDomainShader* DomainShader, uint32 BufferIndex, uint32 BaseIndex, uint32 NumBytes, const void* NewValue) final override;
 	
-	virtual void RHISetShaderParameter(FGeometryShaderRHIParamRef GeometryShader, uint32 BufferIndex, uint32 BaseIndex, uint32 NumBytes, const void* NewValue) final override;
+	virtual void RHISetShaderParameter(FRHIGeometryShader* GeometryShader, uint32 BufferIndex, uint32 BaseIndex, uint32 NumBytes, const void* NewValue) final override;
 	
-	virtual void RHISetShaderParameter(FComputeShaderRHIParamRef ComputeShader, uint32 BufferIndex, uint32 BaseIndex, uint32 NumBytes, const void* NewValue) final override;
+	virtual void RHISetShaderParameter(FRHIComputeShader* ComputeShader, uint32 BufferIndex, uint32 BaseIndex, uint32 NumBytes, const void* NewValue) final override;
 	
 	virtual void RHISetStencilRef(uint32 StencilRef) final override;
 
@@ -327,7 +327,7 @@ public:
 	virtual ~FMetalRHIComputeContext();
 	
 	virtual void RHISetAsyncComputeBudget(EAsyncComputeBudget Budget) final override;
-	virtual void RHISetComputeShader(FComputeShaderRHIParamRef ComputeShader) final override;
+	virtual void RHISetComputeShader(FRHIComputeShader* ComputeShader) final override;
 	virtual void RHISetComputePipelineState(FRHIComputePipelineState* ComputePipelineState) final override;
 	virtual void RHISubmitCommandsHint() final override;
 };
diff --git a/Engine/Source/Runtime/D3D12RHI/Private/D3D12CommandContext.h b/Engine/Source/Runtime/D3D12RHI/Private/D3D12CommandContext.h
--- a/Engine/Source/Runtime/D3D12RHI/Private/D3D12CommandContext.h
+++ b/Engine/Source/Runtime/D3D12RHI/Private/D3D12CommandContext.h
@@ -228,18 +228,18 @@ public:
 	virtual void SetAsyncComputeBudgetInternal(EAsyncComputeBudget Budget) {}
 
 	// IRHIComputeContext interface
-	virtual void RHISetComputeShader(FComputeShaderRHIParamRef ComputeShader) final override;
+	virtual void RHISetComputeShader(FRHIComputeShader* ComputeShader) final override;
 	virtual void RHISetComputePipelineState(FRHIComputePipelineState* ComputePipelineState) final override;
 	virtual void RHIDispatchComputeShader(uint32 ThreadGroupCountX, uint32 ThreadGroupCountY, uint32 ThreadGroupCountZ) final override;
 	virtual void RHIDispatchIndirectComputeShader(FVertexBufferRHIParamRef ArgumentBuffer, uint32 ArgumentOffset) final override;
 	virtual void RHITransitionResources(EResourceTransitionAccess TransitionType, EResourceTransitionPipeline TransitionPipeline, FRHIUnorderedAccessView** InUAVs, int32 NumUAVs, FRHIComputeFence* WriteComputeFenceRHI) final override;
-	virtual void RHISetShaderTexture(FComputeShaderRHIParamRef PixelShader, uint32 TextureIndex, FTextureRHIParamRef NewTexture) final override;
-	virtual void RHISetShaderSampler(FComputeShaderRHIParamRef ComputeShader, uint32 SamplerIndex, FRHISamplerState* NewState) final override;
-	virtual void RHISetUAVParameter(FComputeShaderRHIParamRef ComputeShader, uint32 UAVIndex, FRHIUnorderedAccessView* UAV) final override;
-	virtual void RHISetUAVParameter(FComputeShaderRHIParamRef ComputeShader, uint32 UAVIndex, FRHIUnorderedAccessView* UAV, uint32 InitialCount) final override;
-	virtual void RHISetShaderResourceViewParameter(FComputeShaderRHIParamRef ComputeShader, uint32 SamplerIndex, FRHIShaderResourceView* SRV) final override;
-	virtual void RHISetShaderUniformBuffer(FComputeShaderRHIParamRef ComputeShader, uint32 BufferIndex, FUniformBufferRHIParamRef Buffer) final override;
-	virtual void RHISetShaderParameter(FComputeShaderRHIParamRef ComputeShader, uint32 BufferIndex, uint32 BaseIndex, uint32 NumBytes, const void* NewValue) final override;
+	virtual void RHISetShaderTexture(FRHIComputeShader* PixelShader, uint32 TextureIndex, FTextureRHIParamRef NewTexture) final override;
+	virtual void RHISetShaderSampler(FRHIComputeShader* ComputeShader, uint32 SamplerIndex, FRHISamplerState* NewState) final override;
+	virtual void RHISetUAVParameter(FRHIComputeShader* ComputeShader, uint32 UAVIndex, FRHIUnorderedAccessView* UAV) final override;
+	virtual void RHISetUAVParameter(FRHIComputeShader* ComputeShader, uint32 UAVIndex, FRHIUnorderedAccessView* UAV, uint32 InitialCount) final override;
+	virtual void RHISetShaderResourceViewParameter(FRHIComputeShader* ComputeShader, uint32 SamplerIndex, FRHIShaderResourceView* SRV) final override;
+	virtual void RHISetShaderUniformBuffer(FRHIComputeShader* ComputeShader, uint32 BufferIndex, FUniformBufferRHIParamRef Buffer) final override;
+	virtual void RHISetShaderParameter(FRHIComputeShader* ComputeShader, uint32 BufferIndex, uint32 BaseIndex, uint32 NumBytes, const void* NewValue) final override;
 	virtual void RHIPushEvent(const TCHAR* Name, FColor Color) final override;
 	virtual void RHIPopEvent() final override;
 	virtual void RHISubmitCommandsHint() final override;
@@ -264,31 +264,31 @@ public:
 	virtual void RHISetViewport(uint32 MinX, uint32 MinY, float MinZ, uint32 MaxX, uint32 MaxY, float MaxZ) final override;
 	virtual void RHISetScissorRect(bool bEnable, uint32 MinX, uint32 MinY, uint32 MaxX, uint32 MaxY) final override;
 	virtual void RHISetGraphicsPipelineState(FRHIGraphicsPipelineState* GraphicsPipelineState) final override;
-	virtual void RHISetShaderTexture(FVertexShaderRHIParamRef VertexShader, uint32 TextureIndex, FTextureRHIParamRef NewTexture) final override;
-	virtual void RHISetShaderTexture(FHullShaderRHIParamRef HullShader, uint32 TextureIndex, FTextureRHIParamRef NewTexture) final override;
-	virtual void RHISetShaderTexture(FDomainShaderRHIParamRef DomainShader, uint32 TextureIndex, FTextureRHIParamRef NewTexture) final override;
-	virtual void RHISetShaderTexture(FGeometryShaderRHIParamRef GeometryShader, uint32 TextureIndex, FTextureRHIParamRef NewTexture) final override;
-	virtual void RHISetShaderTexture(FPixelShaderRHIParamRef PixelShader, uint32 TextureIndex, FTextureRHIParamRef NewTexture) final override;
-	virtual void RHISetShaderSampler(FVertexShaderRHIParamRef VertexShader, uint32 SamplerIndex, FRHISamplerState* NewState) final override;
-	virtual void RHISetShaderSampler(FGeometryShaderRHIParamRef GeometryShader, uint32 SamplerIndex, FRHISamplerState* NewState) final override;
-	virtual void RHISetShaderSampler(FDomainShaderRHIParamRef DomainShader, uint32 SamplerIndex, FRHISamplerState* NewState) final override;
-	virtual void RHISetShaderSampler(FHullShaderRHIParamRef HullShader, uint32 SamplerIndex, FRHISamplerState* NewState) final override;
-	virtual void RHISetShaderSampler(FPixelShaderRHIParamRef PixelShader, uint32 SamplerIndex, FRHISamplerState* NewState) final override;
-	virtual void RHISetShaderResourceViewParameter(FPixelShaderRHIParamRef PixelShader, uint32 SamplerIndex, FRHIShaderResourceView* SRV) final override;
-	virtual void RHISetShaderResourceViewParameter(FVertexShaderRHIParamRef VertexShader, uint32 SamplerIndex, FRHIShaderResourceView* SRV) final override;
-	virtual void RHISetShaderResourceViewParameter(FHullShaderRHIParamRef HullShader, uint32 SamplerIndex, FRHIShaderResourceView* SRV) final override;
-	virtual void RHISetShaderResourceViewParameter(FDomainShaderRHIParamRef DomainShader, uint32 SamplerIndex, FRHIShaderResourceView* SRV) final override;
-	virtual void RHISetShaderResourceViewParameter(FGeometryShaderRHIParamRef GeometryShader, uint32 SamplerIndex, FRHIShaderResourceView* SRV) final override;
-	virtual void RHISetShaderUniformBuffer(FVertexShaderRHIParamRef VertexShader, uint32 BufferIndex, FUniformBufferRHIParamRef Buffer) final override;
-	virtual void RHISetShaderUniformBuffer(FHullShaderRHIParamRef HullShader, uint32 BufferIndex, FUniformBufferRHIParamRef Buffer) final override;
-	virtual void RHISetShaderUniformBuffer(FDomainShaderRHIParamRef DomainShader, uint32 BufferIndex, FUniformBufferRHIParamRef Buffer) final override;
-	virtual void RHISetShaderUniformBuffer(FGeometryShaderRHIParamRef GeometryShader, uint32 BufferIndex, FUniformBufferRHIParamRef Buffer) final override;
-	virtual void RHISetShaderUniformBuffer(FPixelShaderRHIParamRef PixelShader, uint32 BufferIndex, FUniformBufferRHIParamRef Buffer) final override;
-	virtual void RHISetShaderParameter(FVertexShaderRHIParamRef VertexShader, uint32 BufferIndex, uint32 BaseIndex, uint32 NumBytes, const void* NewValue) final override;
-	virtual void RHISetShaderParameter(FPixelShaderRHIParamRef PixelShader, uint32 BufferIndex, uint32 BaseIndex, uint32 NumBytes, const void* NewValue) final override;
-	virtual void RHISetShaderParameter(FHullShaderRHIParamRef HullShader, uint32 BufferIndex, uint32 BaseIndex, uint32 NumBytes, const void* NewValue) final override;
-	virtual void RHISetShaderParameter(FDomainShaderRHIParamRef DomainShader, uint32 BufferIndex, uint32 BaseIndex, uint32 NumBytes, const void* NewValue) final override;
-	virtual void RHISetShaderParameter(FGeometryShaderRHIParamRef GeometryShader, uint32 BufferIndex, uint32 BaseIndex, uint32 NumBytes, const void* NewValue) final override;
+	virtual void RHISetShaderTexture(FRHIVertexShader* VertexShader, uint32 TextureIndex, FTextureRHIParamRef NewTexture) final override;
+	virtual void RHISetShaderTexture(FRHIHullShader* HullShader, uint32 TextureIndex, FTextureRHIParamRef NewTexture) final override;
+	virtual void RHISetShaderTexture(FRHIDomainShader* DomainShader, uint32 TextureIndex, FTextureRHIParamRef NewTexture) final override;
+	virtual void RHISetShaderTexture(FRHIGeometryShader* GeometryShader, uint32 TextureIndex, FTextureRHIParamRef NewTexture) final override;
+	virtual void RHISetShaderTexture(FRHIPixelShader* PixelShader, uint32 TextureIndex, FTextureRHIParamRef NewTexture) final override;
+	virtual void RHISetShaderSampler(FRHIVertexShader* VertexShader, uint32 SamplerIndex, FRHISamplerState* NewState) final override;
+	virtual void RHISetShaderSampler(FRHIGeometryShader* GeometryShader, uint32 SamplerIndex, FRHISamplerState* NewState) final override;
+	virtual void RHISetShaderSampler(FRHIDomainShader* DomainShader, uint32 SamplerIndex, FRHISamplerState* NewState) final override;
+	virtual void RHISetShaderSampler(FRHIHullShader* HullShader, uint32 SamplerIndex, FRHISamplerState* NewState) final override;
+	virtual void RHISetShaderSampler(FRHIPixelShader* PixelShader, uint32 SamplerIndex, FRHISamplerState* NewState) final override;
+	virtual void RHISetShaderResourceViewParameter(FRHIPixelShader* PixelShader, uint32 SamplerIndex, FRHIShaderResourceView* SRV) final override;
+	virtual void RHISetShaderResourceViewParameter(FRHIVertexShader* VertexShader, uint32 SamplerIndex, FRHIShaderResourceView* SRV) final override;
+	virtual void RHISetShaderResourceViewParameter(FRHIHullShader* HullShader, uint32 SamplerIndex, FRHIShaderResourceView* SRV) final override;
+	virtual void RHISetShaderResourceViewParameter(FRHIDomainShader* DomainShader, uint32 SamplerIndex, FRHIShaderResourceView* SRV) final override;
+	virtual void RHISetShaderResourceViewParameter(FRHIGeometryShader* GeometryShader, uint32 SamplerIndex, FRHIShaderResourceView* SRV) final override;
+	virtual void RHISetShaderUniformBuffer(FRHIVertexShader* VertexShader, uint32 BufferIndex, FUniformBufferRHIParamRef Buffer) final override;
+	virtual void RHISetShaderUniformBuffer(FRHIHullShader* HullShader, uint32 BufferIndex, FUniformBufferRHIParamRef Buffer) final override;
+	virtual void RHISetShaderUniformBuffer(FRHIDomainShader* DomainShader, uint32 BufferIndex, FUniformBufferRHIParamRef Buffer) final override;
+	virtual void RHISetShaderUniformBuffer(FRHIGeometryShader* GeometryShader, uint32 BufferIndex, FUniformBufferRHIParamRef Buffer) final override;
+	virtual void RHISetShaderUniformBuffer(FRHIPixelShader* PixelShader, uint32 BufferIndex, FUniformBufferRHIParamRef Buffer) final override;
+	virtual void RHISetShaderParameter(FRHIVertexShader* VertexShader, uint32 BufferIndex, uint32 BaseIndex, uint32 NumBytes, const void* NewValue) final override;
+	virtual void RHISetShaderParameter(FRHIPixelShader* PixelShader, uint32 BufferIndex, uint32 BaseIndex, uint32 NumBytes, const void* NewValue) final override;
+	virtual void RHISetShaderParameter(FRHIHullShader* HullShader, uint32 BufferIndex, uint32 BaseIndex, uint32 NumBytes, const void* NewValue) final override;
+	virtual void RHISetShaderParameter(FRHIDomainShader* DomainShader, uint32 BufferIndex, uint32 BaseIndex, uint32 NumBytes, const void* NewValue) final override;
+	virtual void RHISetShaderParameter(FRHIGeometryShader* GeometryShader, uint32 BufferIndex, uint32 BaseIndex, uint32 NumBytes, const void* NewValue) final override;
 	virtual void RHISetStencilRef(uint32 StencilRef) final override;
 	virtual void RHISetBlendFactor(const FLinearColor& BlendFactor) final override;
 	virtual void RHISetRenderTargets(uint32 NumSimultaneousRenderTargets, const FRHIRenderTargetView* NewRenderTargets, const FRHIDepthRenderTargetView* NewDepthStencilTarget, uint32 NumUAVs, FRHIUnorderedAccessView* const* UAVs) final override;
@@ -443,7 +443,7 @@ public:
 #define ContextRedirect(Call) { for (uint32 GPUIndex : GPUMask) PhysicalContexts[GPUIndex]->##Call; }
 #define ContextGPU0(Call) { PhysicalContexts[0]->##Call; }
 
-	FORCEINLINE virtual void RHISetComputeShader(FComputeShaderRHIParamRef ComputeShader) final override
+	FORCEINLINE virtual void RHISetComputeShader(FRHIComputeShader* ComputeShader) final override
 	{
 		ContextRedirect(RHISetComputeShader(ComputeShader));
 	}
@@ -470,31 +470,31 @@ public:
 	{
 		ContextRedirect(RHIWriteGPUFence(Fence));
 	}
-	FORCEINLINE virtual void RHISetShaderTexture(FComputeShaderRHIParamRef PixelShader, uint32 TextureIndex, FTextureRHIParamRef NewTexture) final override
+	FORCEINLINE virtual void RHISetShaderTexture(FRHIComputeShader* PixelShader, uint32 TextureIndex, FTextureRHIParamRef NewTexture) final override
 	{
 		ContextRedirect(RHISetShaderTexture(PixelShader, TextureIndex, NewTexture));
 	}
-	FORCEINLINE virtual void RHISetShaderSampler(FComputeShaderRHIParamRef ComputeShader, uint32 SamplerIndex, FRHISamplerState* NewState) final override
+	FORCEINLINE virtual void RHISetShaderSampler(FRHIComputeShader* ComputeShader, uint32 SamplerIndex, FRHISamplerState* NewState) final override
 	{
 		ContextRedirect(RHISetShaderSampler(ComputeShader, SamplerIndex, NewState));
 	}
-	FORCEINLINE virtual void RHISetUAVParameter(FComputeShaderRHIParamRef ComputeShader, uint32 UAVIndex, FRHIUnorderedAccessView* UAV) final override
+	FORCEINLINE virtual void RHISetUAVParameter(FRHIComputeShader* ComputeShader, uint32 UAVIndex, FRHIUnorderedAccessView* UAV) final override
 	{
 		ContextRedirect(RHISetUAVParameter(ComputeShader, UAVIndex, UAV));
 	}
-	FORCEINLINE virtual void RHISetUAVParameter(FComputeShaderRHIParamRef ComputeShader, uint32 UAVIndex, FRHIUnorderedAccessView* UAV, uint32 InitialCount) final override
+	FORCEINLINE virtual void RHISetUAVParameter(FRHIComputeShader* ComputeShader, uint32 UAVIndex, FRHIUnorderedAccessView* UAV, uint32 InitialCount) final override
 	{
 		ContextRedirect(RHISetUAVParameter(ComputeShader, UAVIndex, UAV, InitialCount));
 	}
-	FORCEINLINE virtual void RHISetShaderResourceViewParameter(FComputeShaderRHIParamRef ComputeShader, uint32 SamplerIndex, FRHIShaderResourceView* SRV) final override
+	FORCEINLINE virtual void RHISetShaderResourceViewParameter(FRHIComputeShader* ComputeShader, uint32 SamplerIndex, FRHIShaderResourceView* SRV) final override
 	{
 		ContextRedirect(RHISetShaderResourceViewParameter(ComputeShader, SamplerIndex, SRV));
 	}
-	FORCEINLINE virtual void RHISetShaderUniformBuffer(FComputeShaderRHIParamRef ComputeShader, uint32 BufferIndex, FUniformBufferRHIParamRef Buffer) final override
+	FORCEINLINE virtual void RHISetShaderUniformBuffer(FRHIComputeShader* ComputeShader, uint32 BufferIndex, FUniformBufferRHIParamRef Buffer) final override
 	{
 		ContextRedirect(RHISetShaderUniformBuffer(ComputeShader, BufferIndex, Buffer));
 	}
-	FORCEINLINE virtual void RHISetShaderParameter(FComputeShaderRHIParamRef ComputeShader, uint32 BufferIndex, uint32 BaseIndex, uint32 NumBytes, const void* NewValue) final override
+	FORCEINLINE virtual void RHISetShaderParameter(FRHIComputeShader* ComputeShader, uint32 BufferIndex, uint32 BaseIndex, uint32 NumBytes, const void* NewValue) final override
 	{
 		ContextRedirect(RHISetShaderParameter(ComputeShader, BufferIndex, BaseIndex, NumBytes, NewValue));
 	}
@@ -568,103 +568,103 @@ public:
 	{
 		ContextRedirect(RHISetGraphicsPipelineState(GraphicsPipelineState));
 	}
-	FORCEINLINE virtual void RHISetShaderTexture(FVertexShaderRHIParamRef VertexShader, uint32 TextureIndex, FTextureRHIParamRef NewTexture) final override
+	FORCEINLINE virtual void RHISetShaderTexture(FRHIVertexShader* VertexShader, uint32 TextureIndex, FTextureRHIParamRef NewTexture) final override
 	{
 		ContextRedirect(RHISetShaderTexture(VertexShader, TextureIndex, NewTexture));
 	}
-	FORCEINLINE virtual void RHISetShaderTexture(FHullShaderRHIParamRef HullShader, uint32 TextureIndex, FTextureRHIParamRef NewTexture) final override
+	FORCEINLINE virtual void RHISetShaderTexture(FRHIHullShader* HullShader, uint32 TextureIndex, FTextureRHIParamRef NewTexture) final override
 	{
 		ContextRedirect(RHISetShaderTexture(HullShader, TextureIndex, NewTexture));
 	}
-	FORCEINLINE virtual void RHISetShaderTexture(FDomainShaderRHIParamRef DomainShader, uint32 TextureIndex, FTextureRHIParamRef NewTexture) final override
+	FORCEINLINE virtual void RHISetShaderTexture(FRHIDomainShader* DomainShader, uint32 TextureIndex, FTextureRHIParamRef NewTexture) final override
 	{
 		ContextRedirect(RHISetShaderTexture(DomainShader, TextureIndex, NewTexture));
 	}
-	FORCEINLINE virtual void RHISetShaderTexture(FGeometryShaderRHIParamRef GeometryShader, uint32 TextureIndex, FTextureRHIParamRef NewTexture) final override
+	FORCEINLINE virtual void RHISetShaderTexture(FRHIGeometryShader* GeometryShader, uint32 TextureIndex, FTextureRHIParamRef NewTexture) final override
 	{
 		ContextRedirect(RHISetShaderTexture(GeometryShader, TextureIndex, NewTexture));
 	}
-	FORCEINLINE virtual void RHISetShaderTexture(FPixelShaderRHIParamRef PixelShader, uint32 TextureIndex, FTextureRHIParamRef NewTexture) final override
+	FORCEINLINE virtual void RHISetShaderTexture(FRHIPixelShader* PixelShader, uint32 TextureIndex, FTextureRHIParamRef NewTexture) final override
 	{
 		ContextRedirect(RHISetShaderTexture(PixelShader, TextureIndex, NewTexture));
 	}
-	FORCEINLINE virtual void RHISetShaderSampler(FVertexShaderRHIParamRef VertexShader, uint32 SamplerIndex, FRHISamplerState* NewState) final override
+	FORCEINLINE virtual void RHISetShaderSampler(FRHIVertexShader* VertexShader, uint32 SamplerIndex, FRHISamplerState* NewState) final override
 	{
 		ContextRedirect(RHISetShaderSampler(VertexShader, SamplerIndex, NewState));
 	}
-	FORCEINLINE virtual void RHISetShaderSampler(FGeometryShaderRHIParamRef GeometryShader, uint32 SamplerIndex, FRHISamplerState* NewState) final override
+	FORCEINLINE virtual void RHISetShaderSampler(FRHIGeometryShader* GeometryShader, uint32 SamplerIndex, FRHISamplerState* NewState) final override
 	{
 		ContextRedirect(RHISetShaderSampler(GeometryShader, SamplerIndex, NewState));
 	}
-	FORCEINLINE virtual void RHISetShaderSampler(FDomainShaderRHIParamRef DomainShader, uint32 SamplerIndex, FRHISamplerState* NewState) final override
+	FORCEINLINE virtual void RHISetShaderSampler(FRHIDomainShader* DomainShader, uint32 SamplerIndex, FRHISamplerState* NewState) final override
 	{
 		ContextRedirect(RHISetShaderSampler(DomainShader, SamplerIndex, NewState));
 	}
-	FORCEINLINE virtual void RHISetShaderSampler(FHullShaderRHIParamRef HullShader, uint32 SamplerIndex, FRHISamplerState* NewState) final override
+	FORCEINLINE virtual void RHISetShaderSampler(FRHIHullShader* HullShader, uint32 SamplerIndex, FRHISamplerState* NewState) final override
 	{
 		ContextRedirect(RHISetShaderSampler(HullShader, SamplerIndex, NewState));
 	}
-	FORCEINLINE virtual void RHISetShaderSampler(FPixelShaderRHIParamRef PixelShader, uint32 SamplerIndex, FRHISamplerState* NewState) final override
+	FORCEINLINE virtual void RHISetShaderSampler(FRHIPixelShader* PixelShader, uint32 SamplerIndex, FRHISamplerState* NewState) final override
 	{
 		ContextRedirect(RHISetShaderSampler(PixelShader, SamplerIndex, NewState));
 	}
-	FORCEINLINE virtual void RHISetShaderResourceViewParameter(FPixelShaderRHIParamRef PixelShader, uint32 SamplerIndex, FRHIShaderResourceView* SRV) final override
+	FORCEINLINE virtual void RHISetShaderResourceViewParameter(FRHIPixelShader* PixelShader, uint32 SamplerIndex, FRHIShaderResourceView* SRV) final override
 	{
 		ContextRedirect(RHISetShaderResourceViewParameter(PixelShader, SamplerIndex, SRV));
 	}
-	FORCEINLINE virtual void RHISetShaderResourceViewParameter(FVertexShaderRHIParamRef VertexShader, uint32 SamplerIndex, FRHIShaderResourceView* SRV) final override
+	FORCEINLINE virtual void RHISetShaderResourceViewParameter(FRHIVertexShader* VertexShader, uint32 SamplerIndex, FRHIShaderResourceView* SRV) final override
 	{
 		ContextRedirect(RHISetShaderResourceViewParameter(VertexShader, SamplerIndex, SRV));
 	}
-	FORCEINLINE virtual void RHISetShaderResourceViewParameter(FHullShaderRHIParamRef HullShader, uint32 SamplerIndex, FRHIShaderResourceView* SRV) final override
+	FORCEINLINE virtual void RHISetShaderResourceViewParameter(FRHIHullShader* HullShader, uint32 SamplerIndex, FRHIShaderResourceView* SRV) final override
 	{
 		ContextRedirect(RHISetShaderResourceViewParameter(HullShader, SamplerIndex, SRV));
 	}
-	FORCEINLINE virtual void RHISetShaderResourceViewParameter(FDomainShaderRHIParamRef DomainShader, uint32 SamplerIndex, FRHIShaderResourceView* SRV) final override
+	FORCEINLINE virtual void RHISetShaderResourceViewParameter(FRHIDomainShader* DomainShader, uint32 SamplerIndex, FRHIShaderResourceView* SRV) final override
 	{
 		ContextRedirect(RHISetShaderResourceViewParameter(DomainShader, SamplerIndex, SRV));
 	}
-	FORCEINLINE virtual void RHISetShaderResourceViewParameter(FGeometryShaderRHIParamRef GeometryShader, uint32 SamplerIndex, FRHIShaderResourceView* SRV) final override
+	FORCEINLINE virtual void RHISetShaderResourceViewParameter(FRHIGeometryShader* GeometryShader, uint32 SamplerIndex, FRHIShaderResourceView* SRV) final override
 	{
 		ContextRedirect(RHISetShaderResourceViewParameter(GeometryShader, SamplerIndex, SRV));
 	}
-	FORCEINLINE virtual void RHISetShaderUniformBuffer(FVertexShaderRHIParamRef VertexShader, uint32 BufferIndex, FUniformBufferRHIParamRef Buffer) final override
+	FORCEINLINE virtual void RHISetShaderUniformBuffer(FRHIVertexShader* VertexShader, uint32 BufferIndex, FUniformBufferRHIParamRef Buffer) final override
 	{
 		ContextRedirect(RHISetShaderUniformBuffer(VertexShader, BufferIndex, Buffer));
 	}
-	FORCEINLINE virtual void RHISetShaderUniformBuffer(FHullShaderRHIParamRef HullShader, uint32 BufferIndex, FUniformBufferRHIParamRef Buffer) final override
+	FORCEINLINE virtual void RHISetShaderUniformBuffer(FRHIHullShader* HullShader, uint32 BufferIndex, FUniformBufferRHIParamRef Buffer) final override
 	{
 		ContextRedirect(RHISetShaderUniformBuffer(HullShader, BufferIndex, Buffer));
 	}
-	FORCEINLINE virtual void RHISetShaderUniformBuffer(FDomainShaderRHIParamRef DomainShader, uint32 BufferIndex, FUniformBufferRHIParamRef Buffer) final override
+	FORCEINLINE virtual void RHISetShaderUniformBuffer(FRHIDomainShader* DomainShader, uint32 BufferIndex, FUniformBufferRHIParamRef Buffer) final override
 	{
 		ContextRedirect(RHISetShaderUniformBuffer(DomainShader, BufferIndex, Buffer));
 	}
-	FORCEINLINE virtual void RHISetShaderUniformBuffer(FGeometryShaderRHIParamRef GeometryShader, uint32 BufferIndex, FUniformBufferRHIParamRef Buffer) final override
+	FORCEINLINE virtual void RHISetShaderUniformBuffer(FRHIGeometryShader* GeometryShader, uint32 BufferIndex, FUniformBufferRHIParamRef Buffer) final override
 	{
 		ContextRedirect(RHISetShaderUniformBuffer(GeometryShader, BufferIndex, Buffer));
 	}
-	FORCEINLINE virtual void RHISetShaderUniformBuffer(FPixelShaderRHIParamRef PixelShader, uint32 BufferIndex, FUniformBufferRHIParamRef Buffer) final override
+	FORCEINLINE virtual void RHISetShaderUniformBuffer(FRHIPixelShader* PixelShader, uint32 BufferIndex, FUniformBufferRHIParamRef Buffer) final override
 	{
 		ContextRedirect(RHISetShaderUniformBuffer(PixelShader, BufferIndex, Buffer));
 	}
-	FORCEINLINE virtual void RHISetShaderParameter(FVertexShaderRHIParamRef VertexShader, uint32 BufferIndex, uint32 BaseIndex, uint32 NumBytes, const void* NewValue) final override
+	FORCEINLINE virtual void RHISetShaderParameter(FRHIVertexShader* VertexShader, uint32 BufferIndex, uint32 BaseIndex, uint32 NumBytes, const void* NewValue) final override
 	{
 		ContextRedirect(RHISetShaderParameter(VertexShader, BufferIndex, BaseIndex, NumBytes, NewValue));
 	}
-	FORCEINLINE virtual void RHISetShaderParameter(FPixelShaderRHIParamRef PixelShader, uint32 BufferIndex, uint32 BaseIndex, uint32 NumBytes, const void* NewValue) final override
+	FORCEINLINE virtual void RHISetShaderParameter(FRHIPixelShader* PixelShader, uint32 BufferIndex, uint32 BaseIndex, uint32 NumBytes, const void* NewValue) final override
 	{
 		ContextRedirect(RHISetShaderParameter(PixelShader, BufferIndex, BaseIndex, NumBytes, NewValue));
 	}
-	FORCEINLINE virtual void RHISetShaderParameter(FHullShaderRHIParamRef HullShader, uint32 BufferIndex, uint32 BaseIndex, uint32 NumBytes, const void* NewValue) final override
+	FORCEINLINE virtual void RHISetShaderParameter(FRHIHullShader* HullShader, uint32 BufferIndex, uint32 BaseIndex, uint32 NumBytes, const void* NewValue) final override
 	{
 		ContextRedirect(RHISetShaderParameter(HullShader, BufferIndex, BaseIndex, NumBytes, NewValue));
 	}
-	FORCEINLINE virtual void RHISetShaderParameter(FDomainShaderRHIParamRef DomainShader, uint32 BufferIndex, uint32 BaseIndex, uint32 NumBytes, const void* NewValue) final override
+	FORCEINLINE virtual void RHISetShaderParameter(FRHIDomainShader* DomainShader, uint32 BufferIndex, uint32 BaseIndex, uint32 NumBytes, const void* NewValue) final override
 	{
 		ContextRedirect(RHISetShaderParameter(DomainShader, BufferIndex, BaseIndex, NumBytes, NewValue));
 	}
-	FORCEINLINE virtual void RHISetShaderParameter(FGeometryShaderRHIParamRef GeometryShader, uint32 BufferIndex, uint32 BaseIndex, uint32 NumBytes, const void* NewValue) final override
+	FORCEINLINE virtual void RHISetShaderParameter(FRHIGeometryShader* GeometryShader, uint32 BufferIndex, uint32 BaseIndex, uint32 NumBytes, const void* NewValue) final override
 	{
 		ContextRedirect(RHISetShaderParameter(GeometryShader, BufferIndex, BaseIndex, NumBytes, NewValue));
 	}
diff --git a/Engine/Source/Runtime/D3D12RHI/Private/D3D12Commands.cpp b/Engine/Source/Runtime/D3D12RHI/Private/D3D12Commands.cpp
--- a/Engine/Source/Runtime/D3D12RHI/Private/D3D12Commands.cpp
+++ b/Engine/Source/Runtime/D3D12RHI/Private/D3D12Commands.cpp
@@ -26,7 +26,7 @@ static FAutoConsoleVariableRef CVarSyncTemporalResources(
 
 using namespace D3D12RHI;
 
-#define DECLARE_ISBOUNDSHADER(ShaderType) inline void ValidateBoundShader(FD3D12StateCache& InStateCache, F##ShaderType##RHIParamRef ShaderType##RHI) \
+#define DECLARE_ISBOUNDSHADER(ShaderType) inline void ValidateBoundShader(FD3D12StateCache& InStateCache, FRHI##ShaderType* ShaderType##RHI) \
 { \
 	FD3D12##ShaderType* CachedShader; \
 	InStateCache.Get##ShaderType(&CachedShader); \
@@ -122,7 +122,7 @@ void FD3D12DynamicRHI::RHISetStreamOutTargets(uint32 NumTargets, const FVertexBu
 	CmdContext.StateCache.SetStreamOutTargets(NumTargets, D3DVertexBuffers, D3DOffsets);
 }
 
-void FD3D12CommandContext::RHISetComputeShader(FComputeShaderRHIParamRef ComputeShaderRHI)
+void FD3D12CommandContext::RHISetComputeShader(FRHIComputeShader* ComputeShaderRHI)
 {
 #if D3D12_RHI_RAYTRACING
 	StateCache.TransitionComputeState(D3D12PT_Compute);
@@ -587,49 +587,49 @@ void FD3D12CommandContext::RHISetComputePipelineState(FRHIComputePipelineState*
 	StateCache.SetComputePipelineState(ComputePipelineState);
 }
 
-void FD3D12CommandContext::RHISetShaderTexture(FVertexShaderRHIParamRef VertexShaderRHI, uint32 TextureIndex, FTextureRHIParamRef NewTextureRHI)
+void FD3D12CommandContext::RHISetShaderTexture(FRHIVertexShader* VertexShaderRHI, uint32 TextureIndex, FTextureRHIParamRef NewTextureRHI)
 {
 	VALIDATE_BOUND_SHADER(VertexShaderRHI);
 	FD3D12TextureBase* const NewTexture = RetrieveTextureBase(NewTextureRHI);
 	StateCache.SetShaderResourceView<SF_Vertex>(NewTexture ? NewTexture->GetShaderResourceView() : nullptr, TextureIndex);
 }
 
-void FD3D12CommandContext::RHISetShaderTexture(FHullShaderRHIParamRef HullShaderRHI, uint32 TextureIndex, FTextureRHIParamRef NewTextureRHI)
+void FD3D12CommandContext::RHISetShaderTexture(FRHIHullShader* HullShaderRHI, uint32 TextureIndex, FTextureRHIParamRef NewTextureRHI)
 {
 	VALIDATE_BOUND_SHADER(HullShaderRHI);
 	FD3D12TextureBase* const NewTexture = RetrieveTextureBase(NewTextureRHI);
 	StateCache.SetShaderResourceView<SF_Hull>(NewTexture ? NewTexture->GetShaderResourceView() : nullptr, TextureIndex);
 }
 
-void FD3D12CommandContext::RHISetShaderTexture(FDomainShaderRHIParamRef DomainShaderRHI, uint32 TextureIndex, FTextureRHIParamRef NewTextureRHI)
+void FD3D12CommandContext::RHISetShaderTexture(FRHIDomainShader* DomainShaderRHI, uint32 TextureIndex, FTextureRHIParamRef NewTextureRHI)
 {
 	VALIDATE_BOUND_SHADER(DomainShaderRHI);
 	FD3D12TextureBase* const NewTexture = RetrieveTextureBase(NewTextureRHI);
 	StateCache.SetShaderResourceView<SF_Domain>(NewTexture ? NewTexture->GetShaderResourceView() : nullptr, TextureIndex);
 }
 
-void FD3D12CommandContext::RHISetShaderTexture(FGeometryShaderRHIParamRef GeometryShaderRHI, uint32 TextureIndex, FTextureRHIParamRef NewTextureRHI)
+void FD3D12CommandContext::RHISetShaderTexture(FRHIGeometryShader* GeometryShaderRHI, uint32 TextureIndex, FTextureRHIParamRef NewTextureRHI)
 {
 	VALIDATE_BOUND_SHADER(GeometryShaderRHI);
 	FD3D12TextureBase* const NewTexture = RetrieveTextureBase(NewTextureRHI);
 	StateCache.SetShaderResourceView<SF_Geometry>(NewTexture ? NewTexture->GetShaderResourceView() : nullptr, TextureIndex);
 }
 
-void FD3D12CommandContext::RHISetShaderTexture(FPixelShaderRHIParamRef PixelShaderRHI, uint32 TextureIndex, FTextureRHIParamRef NewTextureRHI)
+void FD3D12CommandContext::RHISetShaderTexture(FRHIPixelShader* PixelShaderRHI, uint32 TextureIndex, FTextureRHIParamRef NewTextureRHI)
 {
 	VALIDATE_BOUND_SHADER(PixelShaderRHI);
 	FD3D12TextureBase* const NewTexture = RetrieveTextureBase(NewTextureRHI);
 	StateCache.SetShaderResourceView<SF_Pixel>(NewTexture ? NewTexture->GetShaderResourceView() : nullptr, TextureIndex);
 }
 
-void FD3D12CommandContext::RHISetShaderTexture(FComputeShaderRHIParamRef ComputeShaderRHI, uint32 TextureIndex, FTextureRHIParamRef NewTextureRHI)
+void FD3D12CommandContext::RHISetShaderTexture(FRHIComputeShader* ComputeShaderRHI, uint32 TextureIndex, FTextureRHIParamRef NewTextureRHI)
 {
 	//VALIDATE_BOUND_SHADER(ComputeShaderRHI);
 	FD3D12TextureBase* const NewTexture = RetrieveTextureBase(NewTextureRHI);
 	StateCache.SetShaderResourceView<SF_Compute>(NewTexture ? NewTexture->GetShaderResourceView() : nullptr, TextureIndex);
 }
 
-void FD3D12CommandContext::RHISetUAVParameter(FComputeShaderRHIParamRef ComputeShaderRHI, uint32 UAVIndex, FRHIUnorderedAccessView* UAVRHI)
+void FD3D12CommandContext::RHISetUAVParameter(FRHIComputeShader* ComputeShaderRHI, uint32 UAVIndex, FRHIUnorderedAccessView* UAVRHI)
 {
 	//VALIDATE_BOUND_SHADER(ComputeShaderRHI);
 
@@ -646,7 +646,7 @@ void FD3D12CommandContext::RHISetUAVParameter(FComputeShaderRHIParamRef ComputeS
 	StateCache.SetUAVs<SF_Compute>(UAVIndex, 1, &UAV, &InitialCount);
 }
 
-void FD3D12CommandContext::RHISetUAVParameter(FComputeShaderRHIParamRef ComputeShaderRHI, uint32 UAVIndex, FRHIUnorderedAccessView* UAVRHI, uint32 InitialCount)
+void FD3D12CommandContext::RHISetUAVParameter(FRHIComputeShader* ComputeShaderRHI, uint32 UAVIndex, FRHIUnorderedAccessView* UAVRHI, uint32 InitialCount)
 {
 	//VALIDATE_BOUND_SHADER(ComputeShaderRHI);
 
@@ -660,91 +660,91 @@ void FD3D12CommandContext::RHISetUAVParameter(FComputeShaderRHIParamRef ComputeS
 	StateCache.SetUAVs<SF_Compute>(UAVIndex, 1, &UAV, &InitialCount);
 }
 
-void FD3D12CommandContext::RHISetShaderResourceViewParameter(FPixelShaderRHIParamRef PixelShaderRHI, uint32 TextureIndex, FRHIShaderResourceView* SRVRHI)
+void FD3D12CommandContext::RHISetShaderResourceViewParameter(FRHIPixelShader* PixelShaderRHI, uint32 TextureIndex, FRHIShaderResourceView* SRVRHI)
 {
 	VALIDATE_BOUND_SHADER(PixelShaderRHI);
 	FD3D12ShaderResourceView* const SRV = RetrieveObject<FD3D12ShaderResourceView>(SRVRHI);
 	StateCache.SetShaderResourceView<SF_Pixel>(SRV, TextureIndex);
 }
 
-void FD3D12CommandContext::RHISetShaderResourceViewParameter(FVertexShaderRHIParamRef VertexShaderRHI, uint32 TextureIndex, FRHIShaderResourceView* SRVRHI)
+void FD3D12CommandContext::RHISetShaderResourceViewParameter(FRHIVertexShader* VertexShaderRHI, uint32 TextureIndex, FRHIShaderResourceView* SRVRHI)
 {
 	VALIDATE_BOUND_SHADER(VertexShaderRHI);
 	FD3D12ShaderResourceView* const SRV = RetrieveObject<FD3D12ShaderResourceView>(SRVRHI);
 	StateCache.SetShaderResourceView<SF_Vertex>(SRV, TextureIndex);
 }
 
-void FD3D12CommandContext::RHISetShaderResourceViewParameter(FComputeShaderRHIParamRef ComputeShaderRHI, uint32 TextureIndex, FRHIShaderResourceView* SRVRHI)
+void FD3D12CommandContext::RHISetShaderResourceViewParameter(FRHIComputeShader* ComputeShaderRHI, uint32 TextureIndex, FRHIShaderResourceView* SRVRHI)
 {
 	//VALIDATE_BOUND_SHADER(ComputeShaderRHI);
 	FD3D12ShaderResourceView* const SRV = RetrieveObject<FD3D12ShaderResourceView>(SRVRHI);
 	StateCache.SetShaderResourceView<SF_Compute>(SRV, TextureIndex);
 }
 
-void FD3D12CommandContext::RHISetShaderResourceViewParameter(FHullShaderRHIParamRef HullShaderRHI, uint32 TextureIndex, FRHIShaderResourceView* SRVRHI)
+void FD3D12CommandContext::RHISetShaderResourceViewParameter(FRHIHullShader* HullShaderRHI, uint32 TextureIndex, FRHIShaderResourceView* SRVRHI)
 {
 	VALIDATE_BOUND_SHADER(HullShaderRHI);
 	FD3D12ShaderResourceView* const SRV = RetrieveObject<FD3D12ShaderResourceView>(SRVRHI);
 	StateCache.SetShaderResourceView<SF_Hull>(SRV, TextureIndex);
 }
 
-void FD3D12CommandContext::RHISetShaderResourceViewParameter(FDomainShaderRHIParamRef DomainShaderRHI, uint32 TextureIndex, FRHIShaderResourceView* SRVRHI)
+void FD3D12CommandContext::RHISetShaderResourceViewParameter(FRHIDomainShader* DomainShaderRHI, uint32 TextureIndex, FRHIShaderResourceView* SRVRHI)
 {
 	VALIDATE_BOUND_SHADER(DomainShaderRHI);
 	FD3D12ShaderResourceView* const SRV = RetrieveObject<FD3D12ShaderResourceView>(SRVRHI);
 	StateCache.SetShaderResourceView<SF_Domain>(SRV, TextureIndex);
 }
 
-void FD3D12CommandContext::RHISetShaderResourceViewParameter(FGeometryShaderRHIParamRef GeometryShaderRHI, uint32 TextureIndex, FRHIShaderResourceView* SRVRHI)
+void FD3D12CommandContext::RHISetShaderResourceViewParameter(FRHIGeometryShader* GeometryShaderRHI, uint32 TextureIndex, FRHIShaderResourceView* SRVRHI)
 {
 	VALIDATE_BOUND_SHADER(GeometryShaderRHI);
 	FD3D12ShaderResourceView* const SRV = RetrieveObject<FD3D12ShaderResourceView>(SRVRHI);
 	StateCache.SetShaderResourceView<SF_Geometry>(SRV, TextureIndex);
 }
 
-void FD3D12CommandContext::RHISetShaderSampler(FVertexShaderRHIParamRef VertexShaderRHI, uint32 SamplerIndex, FRHISamplerState* NewStateRHI)
+void FD3D12CommandContext::RHISetShaderSampler(FRHIVertexShader* VertexShaderRHI, uint32 SamplerIndex, FRHISamplerState* NewStateRHI)
 {
 	VALIDATE_BOUND_SHADER(VertexShaderRHI);
 	FD3D12SamplerState* NewState = RetrieveObject<FD3D12SamplerState>(NewStateRHI);
 	StateCache.SetSamplerState<SF_Vertex>(NewState, SamplerIndex);
 }
 
-void FD3D12CommandContext::RHISetShaderSampler(FHullShaderRHIParamRef HullShaderRHI, uint32 SamplerIndex, FRHISamplerState* NewStateRHI)
+void FD3D12CommandContext::RHISetShaderSampler(FRHIHullShader* HullShaderRHI, uint32 SamplerIndex, FRHISamplerState* NewStateRHI)
 {
 	VALIDATE_BOUND_SHADER(HullShaderRHI);
 	FD3D12SamplerState* NewState = RetrieveObject<FD3D12SamplerState>(NewStateRHI);
 	StateCache.SetSamplerState<SF_Hull>(NewState, SamplerIndex);
 }
 
-void FD3D12CommandContext::RHISetShaderSampler(FDomainShaderRHIParamRef DomainShaderRHI, uint32 SamplerIndex, FRHISamplerState* NewStateRHI)
+void FD3D12CommandContext::RHISetShaderSampler(FRHIDomainShader* DomainShaderRHI, uint32 SamplerIndex, FRHISamplerState* NewStateRHI)
 {
 	VALIDATE_BOUND_SHADER(DomainShaderRHI);
 	FD3D12SamplerState* NewState = RetrieveObject<FD3D12SamplerState>(NewStateRHI);
 	StateCache.SetSamplerState<SF_Domain>(NewState, SamplerIndex);
 }
 
-void FD3D12CommandContext::RHISetShaderSampler(FGeometryShaderRHIParamRef GeometryShaderRHI, uint32 SamplerIndex, FRHISamplerState* NewStateRHI)
+void FD3D12CommandContext::RHISetShaderSampler(FRHIGeometryShader* GeometryShaderRHI, uint32 SamplerIndex, FRHISamplerState* NewStateRHI)
 {
 	VALIDATE_BOUND_SHADER(GeometryShaderRHI);
 	FD3D12SamplerState* NewState = RetrieveObject<FD3D12SamplerState>(NewStateRHI);
 	StateCache.SetSamplerState<SF_Geometry>(NewState, SamplerIndex);
 }
 
-void FD3D12CommandContext::RHISetShaderSampler(FPixelShaderRHIParamRef PixelShaderRHI, uint32 SamplerIndex, FRHISamplerState* NewStateRHI)
+void FD3D12CommandContext::RHISetShaderSampler(FRHIPixelShader* PixelShaderRHI, uint32 SamplerIndex, FRHISamplerState* NewStateRHI)
 {
 	VALIDATE_BOUND_SHADER(PixelShaderRHI);
 	FD3D12SamplerState* NewState = RetrieveObject<FD3D12SamplerState>(NewStateRHI);
 	StateCache.SetSamplerState<SF_Pixel>(NewState, SamplerIndex);
 }
 
-void FD3D12CommandContext::RHISetShaderSampler(FComputeShaderRHIParamRef ComputeShaderRHI, uint32 SamplerIndex, FRHISamplerState* NewStateRHI)
+void FD3D12CommandContext::RHISetShaderSampler(FRHIComputeShader* ComputeShaderRHI, uint32 SamplerIndex, FRHISamplerState* NewStateRHI)
 {
 	//VALIDATE_BOUND_SHADER(ComputeShaderRHI);
 	FD3D12SamplerState* NewState = RetrieveObject<FD3D12SamplerState>(NewStateRHI);
 	StateCache.SetSamplerState<SF_Compute>(NewState, SamplerIndex);
 }
 
-void FD3D12CommandContext::RHISetShaderUniformBuffer(FVertexShaderRHIParamRef VertexShader, uint32 BufferIndex, FUniformBufferRHIParamRef BufferRHI)
+void FD3D12CommandContext::RHISetShaderUniformBuffer(FRHIVertexShader* VertexShader, uint32 BufferIndex, FUniformBufferRHIParamRef BufferRHI)
 {
 	//SCOPE_CYCLE_COUNTER(STAT_D3D12SetShaderUniformBuffer);
 	VALIDATE_BOUND_SHADER(VertexShader);
@@ -761,7 +761,7 @@ void FD3D12CommandContext::RHISetShaderUniformBuffer(FVertexShaderRHIParamRef Ve
 	DirtyUniformBuffers[SF_Vertex] |= (1 << BufferIndex);
 }
 
-void FD3D12CommandContext::RHISetShaderUniformBuffer(FHullShaderRHIParamRef HullShader, uint32 BufferIndex, FUniformBufferRHIParamRef BufferRHI)
+void FD3D12CommandContext::RHISetShaderUniformBuffer(FRHIHullShader* HullShader, uint32 BufferIndex, FUniformBufferRHIParamRef BufferRHI)
 {
 	//SCOPE_CYCLE_COUNTER(STAT_D3D12SetShaderUniformBuffer);
 	VALIDATE_BOUND_SHADER(HullShader);
@@ -778,7 +778,7 @@ void FD3D12CommandContext::RHISetShaderUniformBuffer(FHullShaderRHIParamRef Hull
 	DirtyUniformBuffers[SF_Hull] |= (1 << BufferIndex);
 }
 
-void FD3D12CommandContext::RHISetShaderUniformBuffer(FDomainShaderRHIParamRef DomainShader, uint32 BufferIndex, FUniformBufferRHIParamRef BufferRHI)
+void FD3D12CommandContext::RHISetShaderUniformBuffer(FRHIDomainShader* DomainShader, uint32 BufferIndex, FUniformBufferRHIParamRef BufferRHI)
 {
 	//SCOPE_CYCLE_COUNTER(STAT_D3D12SetShaderUniformBuffer);
 	VALIDATE_BOUND_SHADER(DomainShader);
@@ -795,7 +795,7 @@ void FD3D12CommandContext::RHISetShaderUniformBuffer(FDomainShaderRHIParamRef Do
 	DirtyUniformBuffers[SF_Domain] |= (1 << BufferIndex);
 }
 
-void FD3D12CommandContext::RHISetShaderUniformBuffer(FGeometryShaderRHIParamRef GeometryShader, uint32 BufferIndex, FUniformBufferRHIParamRef BufferRHI)
+void FD3D12CommandContext::RHISetShaderUniformBuffer(FRHIGeometryShader* GeometryShader, uint32 BufferIndex, FUniformBufferRHIParamRef BufferRHI)
 {
 	//SCOPE_CYCLE_COUNTER(STAT_D3D12SetShaderUniformBuffer);
 	VALIDATE_BOUND_SHADER(GeometryShader);
@@ -812,7 +812,7 @@ void FD3D12CommandContext::RHISetShaderUniformBuffer(FGeometryShaderRHIParamRef
 	DirtyUniformBuffers[SF_Geometry] |= (1 << BufferIndex);
 }
 
-void FD3D12CommandContext::RHISetShaderUniformBuffer(FPixelShaderRHIParamRef PixelShader, uint32 BufferIndex, FUniformBufferRHIParamRef BufferRHI)
+void FD3D12CommandContext::RHISetShaderUniformBuffer(FRHIPixelShader* PixelShader, uint32 BufferIndex, FUniformBufferRHIParamRef BufferRHI)
 {
 	//SCOPE_CYCLE_COUNTER(STAT_D3D12SetShaderUniformBuffer);
 	VALIDATE_BOUND_SHADER(PixelShader);
@@ -829,7 +829,7 @@ void FD3D12CommandContext::RHISetShaderUniformBuffer(FPixelShaderRHIParamRef Pix
 	DirtyUniformBuffers[SF_Pixel] |= (1 << BufferIndex);
 }
 
-void FD3D12CommandContext::RHISetShaderUniformBuffer(FComputeShaderRHIParamRef ComputeShader, uint32 BufferIndex, FUniformBufferRHIParamRef BufferRHI)
+void FD3D12CommandContext::RHISetShaderUniformBuffer(FRHIComputeShader* ComputeShader, uint32 BufferIndex, FUniformBufferRHIParamRef BufferRHI)
 {
 	//SCOPE_CYCLE_COUNTER(STAT_D3D12SetShaderUniformBuffer);
 	//VALIDATE_BOUND_SHADER(ComputeShader);
@@ -846,42 +846,42 @@ void FD3D12CommandContext::RHISetShaderUniformBuffer(FComputeShaderRHIParamRef C
 	DirtyUniformBuffers[SF_Compute] |= (1 << BufferIndex);
 }
 
-void FD3D12CommandContext::RHISetShaderParameter(FHullShaderRHIParamRef HullShaderRHI, uint32 BufferIndex, uint32 BaseIndex, uint32 NumBytes, const void* NewValue)
+void FD3D12CommandContext::RHISetShaderParameter(FRHIHullShader* HullShaderRHI, uint32 BufferIndex, uint32 BaseIndex, uint32 NumBytes, const void* NewValue)
 {
 	VALIDATE_BOUND_SHADER(HullShaderRHI);
 	checkSlow(BufferIndex == 0);
 	HSConstantBuffer.UpdateConstant((const uint8*)NewValue, BaseIndex, NumBytes);
 }
 
-void FD3D12CommandContext::RHISetShaderParameter(FDomainShaderRHIParamRef DomainShaderRHI, uint32 BufferIndex, uint32 BaseIndex, uint32 NumBytes, const void* NewValue)
+void FD3D12CommandContext::RHISetShaderParameter(FRHIDomainShader* DomainShaderRHI, uint32 BufferIndex, uint32 BaseIndex, uint32 NumBytes, const void* NewValue)
 {
 	VALIDATE_BOUND_SHADER(DomainShaderRHI);
 	checkSlow(BufferIndex == 0);
 	DSConstantBuffer.UpdateConstant((const uint8*)NewValue, BaseIndex, NumBytes);
 }
 
-void FD3D12CommandContext::RHISetShaderParameter(FVertexShaderRHIParamRef VertexShaderRHI, uint32 BufferIndex, uint32 BaseIndex, uint32 NumBytes, const void* NewValue)
+void FD3D12CommandContext::RHISetShaderParameter(FRHIVertexShader* VertexShaderRHI, uint32 BufferIndex, uint32 BaseIndex, uint32 NumBytes, const void* NewValue)
 {
 	VALIDATE_BOUND_SHADER(VertexShaderRHI);
 	checkSlow(BufferIndex == 0);
 	VSConstantBuffer.UpdateConstant((const uint8*)NewValue, BaseIndex, NumBytes);
 }
 
-void FD3D12CommandContext::RHISetShaderParameter(FPixelShaderRHIParamRef PixelShaderRHI, uint32 BufferIndex, uint32 BaseIndex, uint32 NumBytes, const void* NewValue)
+void FD3D12CommandContext::RHISetShaderParameter(FRHIPixelShader* PixelShaderRHI, uint32 BufferIndex, uint32 BaseIndex, uint32 NumBytes, const void* NewValue)
 {
 	VALIDATE_BOUND_SHADER(PixelShaderRHI);
 	checkSlow(BufferIndex == 0);
 	PSConstantBuffer.UpdateConstant((const uint8*)NewValue, BaseIndex, NumBytes);
 }
 
-void FD3D12CommandContext::RHISetShaderParameter(FGeometryShaderRHIParamRef GeometryShaderRHI, uint32 BufferIndex, uint32 BaseIndex, uint32 NumBytes, const void* NewValue)
+void FD3D12CommandContext::RHISetShaderParameter(FRHIGeometryShader* GeometryShaderRHI, uint32 BufferIndex, uint32 BaseIndex, uint32 NumBytes, const void* NewValue)
 {
 	VALIDATE_BOUND_SHADER(GeometryShaderRHI);
 	checkSlow(BufferIndex == 0);
 	GSConstantBuffer.UpdateConstant((const uint8*)NewValue, BaseIndex, NumBytes);
 }
 
-void FD3D12CommandContext::RHISetShaderParameter(FComputeShaderRHIParamRef ComputeShaderRHI, uint32 BufferIndex, uint32 BaseIndex, uint32 NumBytes, const void* NewValue)
+void FD3D12CommandContext::RHISetShaderParameter(FRHIComputeShader* ComputeShaderRHI, uint32 BufferIndex, uint32 BaseIndex, uint32 NumBytes, const void* NewValue)
 {
 	//VALIDATE_BOUND_SHADER(ComputeShaderRHI);
 	checkSlow(BufferIndex == 0);
diff --git a/Engine/Source/Runtime/D3D12RHI/Private/D3D12RHIPrivate.h b/Engine/Source/Runtime/D3D12RHI/Private/D3D12RHIPrivate.h
--- a/Engine/Source/Runtime/D3D12RHI/Private/D3D12RHIPrivate.h
+++ b/Engine/Source/Runtime/D3D12RHI/Private/D3D12RHIPrivate.h
@@ -275,7 +275,7 @@ public:
 	virtual FStagingBufferRHIRef RHICreateStagingBuffer() final override;
 	virtual void* RHILockStagingBuffer(FRHIStagingBuffer* StagingBuffer, uint32 Offset, uint32 SizeRHI) final override;
     virtual void RHIUnlockStagingBuffer(FRHIStagingBuffer* StagingBuffer) final override;
-	virtual FBoundShaderStateRHIRef RHICreateBoundShaderState(FRHIVertexDeclaration* VertexDeclaration, FVertexShaderRHIParamRef VertexShader, FHullShaderRHIParamRef HullShader, FDomainShaderRHIParamRef DomainShader, FPixelShaderRHIParamRef PixelShader, FGeometryShaderRHIParamRef GeometryShader) final override;
+	virtual FBoundShaderStateRHIRef RHICreateBoundShaderState(FRHIVertexDeclaration* VertexDeclaration, FRHIVertexShader* VertexShader, FRHIHullShader* HullShader, FRHIDomainShader* DomainShader, FRHIPixelShader* PixelShader, FRHIGeometryShader* GeometryShader) final override;
 	virtual FGraphicsPipelineStateRHIRef RHICreateGraphicsPipelineState(const FGraphicsPipelineStateInitializer& Initializer) final override;
 	virtual TRefCountPtr<FRHIComputePipelineState> RHICreateComputePipelineState(FRHIComputeShader* ComputeShader) final override;
 	virtual FUniformBufferRHIRef RHICreateUniformBuffer(const void* Contents, const FRHIUniformBufferLayout& Layout, EUniformBufferUsage Usage, EUniformBufferValidation Validation) final override;
diff --git a/Engine/Source/Runtime/D3D12RHI/Private/D3D12Shader.h b/Engine/Source/Runtime/D3D12RHI/Private/D3D12Shader.h
--- a/Engine/Source/Runtime/D3D12RHI/Private/D3D12Shader.h
+++ b/Engine/Source/Runtime/D3D12RHI/Private/D3D12Shader.h
@@ -184,11 +184,11 @@ public:
 	/** Initialization constructor. */
 	FD3D12BoundShaderState(
 		FRHIVertexDeclaration* InVertexDeclarationRHI,
-		FVertexShaderRHIParamRef InVertexShaderRHI,
-		FPixelShaderRHIParamRef InPixelShaderRHI,
-		FHullShaderRHIParamRef InHullShaderRHI,
-		FDomainShaderRHIParamRef InDomainShaderRHI,
-		FGeometryShaderRHIParamRef InGeometryShaderRHI,
+		FRHIVertexShader* InVertexShaderRHI,
+		FRHIPixelShader* InPixelShaderRHI,
+		FRHIHullShader* InHullShaderRHI,
+		FRHIDomainShader* InDomainShaderRHI,
+		FRHIGeometryShader* InGeometryShaderRHI,
 		FD3D12Device* InDevice
 		);
 
diff --git a/Engine/Source/Runtime/D3D12RHI/Private/D3D12Shaders.cpp b/Engine/Source/Runtime/D3D12RHI/Private/D3D12Shaders.cpp
--- a/Engine/Source/Runtime/D3D12RHI/Private/D3D12Shaders.cpp
+++ b/Engine/Source/Runtime/D3D12RHI/Private/D3D12Shaders.cpp
@@ -325,11 +325,11 @@ void FD3D12CommandContext::RHISetMultipleViewports(uint32 Count, const FViewport
 
 FD3D12BoundShaderState::FD3D12BoundShaderState(
 	FRHIVertexDeclaration* InVertexDeclarationRHI,
-	FVertexShaderRHIParamRef InVertexShaderRHI,
-	FPixelShaderRHIParamRef InPixelShaderRHI,
-	FHullShaderRHIParamRef InHullShaderRHI,
-	FDomainShaderRHIParamRef InDomainShaderRHI,
-	FGeometryShaderRHIParamRef InGeometryShaderRHI,
+	FRHIVertexShader* InVertexShaderRHI,
+	FRHIPixelShader* InPixelShaderRHI,
+	FRHIHullShader* InHullShaderRHI,
+	FRHIDomainShader* InDomainShaderRHI,
+	FRHIGeometryShader* InGeometryShaderRHI,
 	FD3D12Device* InDevice
 	) :
 	CacheLink(InVertexDeclarationRHI, InVertexShaderRHI, InPixelShaderRHI, InHullShaderRHI, InDomainShaderRHI, InGeometryShaderRHI, this)
@@ -372,11 +372,11 @@ FD3D12BoundShaderState::~FD3D12BoundShaderState()
 */
 FBoundShaderStateRHIRef FD3D12DynamicRHI::RHICreateBoundShaderState(
 	FRHIVertexDeclaration* VertexDeclarationRHI,
-	FVertexShaderRHIParamRef VertexShaderRHI,
-	FHullShaderRHIParamRef HullShaderRHI,
-	FDomainShaderRHIParamRef DomainShaderRHI,
-	FPixelShaderRHIParamRef PixelShaderRHI,
-	FGeometryShaderRHIParamRef GeometryShaderRHI
+	FRHIVertexShader* VertexShaderRHI,
+	FRHIHullShader* HullShaderRHI,
+	FRHIDomainShader* DomainShaderRHI,
+	FRHIPixelShader* PixelShaderRHI,
+	FRHIGeometryShader* GeometryShaderRHI
 	)
 {
 	//SCOPE_CYCLE_COUNTER(STAT_D3D12CreateBoundShaderStateTime);
