Make DX12 scissor behave like DX11

diff --git a/Engine/Source/Runtime/D3D12RHI/Private/D3D12Commands.cpp b/Engine/Source/Runtime/D3D12RHI/Private/D3D12Commands.cpp
--- a/Engine/Source/Runtime/D3D12RHI/Private/D3D12Commands.cpp
+++ b/Engine/Source/Runtime/D3D12RHI/Private/D3D12Commands.cpp
@@ -456,6 +456,7 @@ void FD3D12CommandContext::RHISetViewport(uint32 MinX, uint32 MinY, float MinZ,
 	{
 		// Setting a viewport will also set the scissor rect appropriately.
 		StateCache.SetViewport(Viewport);
+		RHISetScissorRect(true, MinX, MinY, MaxX, MaxY);
 	}
 }
 
@@ -468,7 +469,7 @@ void FD3D12CommandContext::RHISetScissorRect(bool bEnable, uint32 MinX, uint32 M
 	}
 	else
 	{
-		const CD3DX12_RECT ScissorRect(0, 0, GetMax2DTextureDimension(), GetMax2DTextureDimension());
+		const CD3DX12_RECT ScissorRect(0, 0, (LONG) StateCache.GetViewport().Width, (LONG) StateCache.GetViewport().Height);
 		StateCache.SetScissorRect(ScissorRect);
 	}
 }
@@ -1625,12 +1626,8 @@ void FD3D12CommandContext::RHIClearMRTImpl(bool bClearColor, int32 NumClearColor
 {
 	SCOPE_CYCLE_COUNTER(STAT_D3D12ClearMRT);
 
-	uint32 NumViews = 1;
-	D3D12_VIEWPORT Viewport;
-	StateCache.GetViewports(&NumViews, &Viewport);
-
-	D3D12_RECT ScissorRect;
-	StateCache.GetScissorRect(&ScissorRect);
+	const D3D12_VIEWPORT& Viewport = StateCache.GetViewport();
+	const D3D12_RECT& ScissorRect = StateCache.GetScissorRect();
 
 	if (ScissorRect.left >= ScissorRect.right || ScissorRect.top >= ScissorRect.bottom)
 	{
diff --git a/Engine/Source/Runtime/D3D12RHI/Private/D3D12StateCache.cpp b/Engine/Source/Runtime/D3D12RHI/Private/D3D12StateCache.cpp
--- a/Engine/Source/Runtime/D3D12RHI/Private/D3D12StateCache.cpp
+++ b/Engine/Source/Runtime/D3D12RHI/Private/D3D12StateCache.cpp
@@ -167,6 +167,9 @@ void FD3D12StateCacheBase::ClearState()
 	FMemory::Memzero(PipelineState.Graphics.CurrentViewport, sizeof(PipelineState.Graphics.CurrentViewport));
 	PipelineState.Graphics.CurrentNumberOfViewports = 0;
 
+	FMemory::Memzero(PipelineState.Graphics.CurrentScissorRects, sizeof(PipelineState.Graphics.CurrentScissorRects));
+	PipelineState.Graphics.CurrentNumberOfScissorRects = 0;
+
 	PipelineState.Compute.ComputeBudget = EAsyncComputeBudget::EAll_4;
 	PipelineState.Graphics.CurrentPipelineStateObject = nullptr;
 	PipelineState.Compute.CurrentPipelineStateObject = nullptr;
@@ -175,9 +178,6 @@ void FD3D12StateCacheBase::ClearState()
 	FMemory::Memzero(PipelineState.Graphics.CurrentStreamOutTargets, sizeof(PipelineState.Graphics.CurrentStreamOutTargets));
 	FMemory::Memzero(PipelineState.Graphics.CurrentSOOffsets, sizeof(PipelineState.Graphics.CurrentSOOffsets));
 
-	const CD3DX12_RECT ScissorRect(0, 0, GetMax2DTextureDimension(), GetMax2DTextureDimension());
-	SetScissorRect(ScissorRect);
-
 	PipelineState.Graphics.VBCache.Clear();
 	PipelineState.Graphics.IBCache.Clear();
 
@@ -279,7 +279,6 @@ void FD3D12StateCacheBase::SetViewport(const D3D12_VIEWPORT& Viewport)
 		FMemory::Memcpy(&PipelineState.Graphics.CurrentViewport[0], &Viewport, sizeof(D3D12_VIEWPORT));
 		PipelineState.Graphics.CurrentNumberOfViewports = 1;
 		bNeedSetViewports = true;
-		UpdateViewportScissorRects();
 	}
 }
 
@@ -291,49 +290,44 @@ void FD3D12StateCacheBase::SetViewports(uint32 Count, const D3D12_VIEWPORT* cons
 		FMemory::Memcpy(&PipelineState.Graphics.CurrentViewport[0], Viewports, sizeof(D3D12_VIEWPORT) * Count);
 		PipelineState.Graphics.CurrentNumberOfViewports = Count;
 		bNeedSetViewports = true;
-		UpdateViewportScissorRects();
 	}
 }
 
-void FD3D12StateCacheBase::UpdateViewportScissorRects()
+static void ValidateScissorRect(const D3D12_VIEWPORT& Viewport, const D3D12_RECT& ScissorRect)
 {
-	for (uint32 i = 0; i < PipelineState.Graphics.CurrentNumberOfScissorRects; ++i)
-	{
-		const D3D12_VIEWPORT& Viewport = PipelineState.Graphics.CurrentViewport[FMath::Min(i, PipelineState.Graphics.CurrentNumberOfViewports)];
-		const D3D12_RECT& ScissorRect = PipelineState.Graphics.CurrentScissorRects[i];
-		D3D12_RECT& ViewportScissorRect = PipelineState.Graphics.CurrentViewportScissorRects[i];
-
-		ViewportScissorRect.top = FMath::Max(ScissorRect.top, (LONG)Viewport.TopLeftY);
-		ViewportScissorRect.left = FMath::Max(ScissorRect.left, (LONG)Viewport.TopLeftX);
-		ViewportScissorRect.bottom = FMath::Min(ScissorRect.bottom, (LONG)Viewport.TopLeftY + (LONG)Viewport.Height);
-		ViewportScissorRect.right = FMath::Min(ScissorRect.right, (LONG)Viewport.TopLeftX + (LONG)Viewport.Width);
-
-		const bool ViewportEmpty = Viewport.Width <= 0 || Viewport.Height <= 0;
-		const bool ScissorEmpty = ViewportScissorRect.right <= ViewportScissorRect.left || ViewportScissorRect.bottom <= ViewportScissorRect.top;
-		check(!ViewportEmpty || ScissorEmpty);
-	}
-
-	bNeedSetScissorRects = true;
+	ensure(ScissorRect.left   >= (LONG)Viewport.TopLeftX);
+	ensure(ScissorRect.top    >= (LONG)Viewport.TopLeftY);
+	ensure(ScissorRect.right  <= (LONG)Viewport.TopLeftX + (LONG)Viewport.Width);
+	ensure(ScissorRect.bottom <= (LONG)Viewport.TopLeftY + (LONG)Viewport.Height);
+	ensure(ScissorRect.left <= ScissorRect.right && ScissorRect.top <= ScissorRect.bottom);
 }
 
 void FD3D12StateCacheBase::SetScissorRect(const D3D12_RECT& ScissorRect)
 {
+	ValidateScissorRect(PipelineState.Graphics.CurrentViewport[0], ScissorRect);
+
 	if ((PipelineState.Graphics.CurrentNumberOfScissorRects != 1 || FMemory::Memcmp(&PipelineState.Graphics.CurrentScissorRects[0], &ScissorRect, sizeof(D3D12_RECT))) || GD3D12SkipStateCaching)
 	{
 		FMemory::Memcpy(&PipelineState.Graphics.CurrentScissorRects[0], &ScissorRect, sizeof(D3D12_RECT));
 		PipelineState.Graphics.CurrentNumberOfScissorRects = 1;
-		UpdateViewportScissorRects();
+		bNeedSetScissorRects = true;
 	}
 }
 
 void FD3D12StateCacheBase::SetScissorRects(uint32 Count, const D3D12_RECT* const ScissorRects)
 {
 	check(Count < ARRAY_COUNT(PipelineState.Graphics.CurrentScissorRects));
+
+	for (uint32 Rect = 0; Rect < Count; ++Rect)
+	{
+		ValidateScissorRect(PipelineState.Graphics.CurrentViewport[Rect], ScissorRects[Rect]);
+	}
+
 	if ((PipelineState.Graphics.CurrentNumberOfScissorRects != Count || FMemory::Memcmp(&PipelineState.Graphics.CurrentScissorRects[0], ScissorRects, sizeof(D3D12_RECT) * Count)) || GD3D12SkipStateCaching)
 	{
 		FMemory::Memcpy(&PipelineState.Graphics.CurrentScissorRects[0], ScissorRects, sizeof(D3D12_RECT) * Count);
 		PipelineState.Graphics.CurrentNumberOfScissorRects = Count;
-		UpdateViewportScissorRects();
+		bNeedSetScissorRects = true;
 	}
 }
 
@@ -429,7 +423,7 @@ void FD3D12StateCacheBase::ApplyState()
 		}
 		if (bNeedSetScissorRects)
 		{
-			CommandList->RSSetScissorRects(PipelineState.Graphics.CurrentNumberOfScissorRects, PipelineState.Graphics.CurrentViewportScissorRects);
+			CommandList->RSSetScissorRects(PipelineState.Graphics.CurrentNumberOfScissorRects, PipelineState.Graphics.CurrentScissorRects);
 			bNeedSetScissorRects = false;
 		}
 		if (bNeedSetPrimitiveTopology)
diff --git a/Engine/Source/Runtime/D3D12RHI/Private/D3D12StateCachePrivate.h b/Engine/Source/Runtime/D3D12RHI/Private/D3D12StateCachePrivate.h
--- a/Engine/Source/Runtime/D3D12RHI/Private/D3D12StateCachePrivate.h
+++ b/Engine/Source/Runtime/D3D12RHI/Private/D3D12StateCachePrivate.h
@@ -318,7 +318,6 @@ protected:
 
 			// Input Layout State
 			D3D12_RECT CurrentScissorRects[D3D12_VIEWPORT_AND_SCISSORRECT_OBJECT_COUNT_PER_PIPELINE];
-			D3D12_RECT CurrentViewportScissorRects[D3D12_VIEWPORT_AND_SCISSORRECT_OBJECT_COUNT_PER_PIPELINE];
 			uint32 CurrentNumberOfScissorRects;
 
 			uint16 StreamStrides[MaxVertexElementCount];
@@ -518,14 +517,12 @@ public:
 		}
 	}
 
-	void UpdateViewportScissorRects();
 	void SetScissorRects(uint32 Count, const D3D12_RECT* const ScissorRects);
 	void SetScissorRect(const D3D12_RECT& ScissorRect);
 
-	D3D12_STATE_CACHE_INLINE void GetScissorRect(D3D12_RECT* ScissorRect) const
+	D3D12_STATE_CACHE_INLINE const D3D12_RECT& GetScissorRect(int32 Index = 0) const
 	{
-		check(ScissorRect);
-		FMemory::Memcpy(ScissorRect, &PipelineState.Graphics.CurrentScissorRects, sizeof(D3D12_RECT));
+		return PipelineState.Graphics.CurrentScissorRects[Index];
 	}
 
 	void SetViewport(const D3D12_VIEWPORT& Viewport);
@@ -536,10 +533,9 @@ public:
 		return PipelineState.Graphics.CurrentNumberOfViewports;
 	}
 
-	D3D12_STATE_CACHE_INLINE void GetViewport(D3D12_VIEWPORT* Viewport) const
+	D3D12_STATE_CACHE_INLINE const D3D12_VIEWPORT& GetViewport(int32 Index = 0) const
 	{
-		check(Viewport);
-		FMemory::Memcpy(Viewport, &PipelineState.Graphics.CurrentViewport, sizeof(D3D12_VIEWPORT));
+		return PipelineState.Graphics.CurrentViewport[Index];
 	}
 
 	D3D12_STATE_CACHE_INLINE void GetViewports(uint32* Count, D3D12_VIEWPORT* Viewports) const
diff --git a/Engine/Source/Runtime/Windows/D3D11RHI/Private/D3D11Commands.cpp b/Engine/Source/Runtime/Windows/D3D11RHI/Private/D3D11Commands.cpp
--- a/Engine/Source/Runtime/Windows/D3D11RHI/Private/D3D11Commands.cpp
+++ b/Engine/Source/Runtime/Windows/D3D11RHI/Private/D3D11Commands.cpp
@@ -248,7 +248,7 @@ void FD3D11DynamicRHI::RHISetViewport(uint32 MinX,uint32 MinY,float MinZ,uint32
 	if (Viewport.Width > 0 && Viewport.Height > 0)
 	{
 		StateCache.SetViewport(Viewport);
-		SetScissorRectIfRequiredWhenSettingViewport(MinX, MinY, MaxX, MaxY);
+		RHISetScissorRect(true, MinX, MinY, MaxX, MaxY);
 	}
 }
 
diff --git a/Engine/Source/Runtime/Windows/D3D11RHI/Private/D3D11RHIPrivate.h b/Engine/Source/Runtime/Windows/D3D11RHI/Private/D3D11RHIPrivate.h
--- a/Engine/Source/Runtime/Windows/D3D11RHI/Private/D3D11RHIPrivate.h
+++ b/Engine/Source/Runtime/Windows/D3D11RHI/Private/D3D11RHIPrivate.h
@@ -1003,12 +1003,6 @@ protected:
 			typename TPixelShader::FParameter PixelShaderParameter
 		);
 
-	// Some platforms might want to override this
-	virtual void SetScissorRectIfRequiredWhenSettingViewport(uint32 MinX, uint32 MinY, uint32 MaxX, uint32 MaxY)
-	{
-		RHISetScissorRect(true, MinX, MinY, MaxX, MaxY);
-	}
-
 	/**
 	* Returns a pointer to a texture resource that can be used for CPU reads.
 	* Note: the returned resource could be the original texture or a new temporary texture.
