Fixed post merge compilation errors and warnings

diff --git a/Engine/Plugins/FX/Niagara/Source/Niagara/Private/NiagaraRendererSprites.cpp b/Engine/Plugins/FX/Niagara/Source/Niagara/Private/NiagaraRendererSprites.cpp
--- a/Engine/Plugins/FX/Niagara/Source/Niagara/Private/NiagaraRendererSprites.cpp
+++ b/Engine/Plugins/FX/Niagara/Source/Niagara/Private/NiagaraRendererSprites.cpp
@@ -347,7 +347,7 @@ void FNiagaraRendererSprites::SetVertexFactoryParticleData(
 
 		int32 NumInstances = SourceParticleData->GetNumInstances();
 		FNiagaraGPUSortInfo SortInfo;
-		if (View && SortMode != ENiagaraSortMode::None && (BlendMode == BLEND_AlphaComposite || BlendMode == BLEND_AlphaHoldout|| BlendMode == BLEND_Translucent || !Properties->bSortOnlyWhenTranslucent))
+		if (View && SortMode != ENiagaraSortMode::None && (BlendMode == BLEND_AlphaComposite || BlendMode == BLEND_AlphaHoldout|| BlendMode == BLEND_Translucent || !bSortOnlyWhenTranslucent))
 		{
 			SortInfo.ParticleCount = NumInstances;
 			SortInfo.SortMode = SortMode;
diff --git a/Engine/Source/Runtime/D3D12RHI/Private/D3D12RHIPrivate.h b/Engine/Source/Runtime/D3D12RHI/Private/D3D12RHIPrivate.h
--- a/Engine/Source/Runtime/D3D12RHI/Private/D3D12RHIPrivate.h
+++ b/Engine/Source/Runtime/D3D12RHI/Private/D3D12RHIPrivate.h
@@ -316,7 +316,6 @@ public:
 	virtual FTexture3DRHIRef RHICreateTexture3D(uint32 SizeX, uint32 SizeY, uint32 SizeZ, uint8 Format, uint32 NumMips, uint32 Flags, FRHIResourceCreateInfo& CreateInfo) override;
 	virtual void RHIGetResourceInfo(FTextureRHIParamRef Ref, FRHIResourceInfo& OutInfo) override;
 	virtual FShaderResourceViewRHIRef RHICreateShaderResourceView(FTextureRHIParamRef Texture, const FRHITextureSRVCreateInfo& CreateInfo) final override;
-	virtual void RHIGenerateMips(FTextureRHIParamRef Texture) final override;
 	virtual uint32 RHIComputeMemorySize(FTextureRHIParamRef TextureRHI) final override;
 	virtual FTexture2DRHIRef RHIAsyncReallocateTexture2D(FTexture2DRHIParamRef Texture2D, int32 NewMipCount, int32 NewSizeX, int32 NewSizeY, FThreadSafeCounter* RequestStatus) override;
 	virtual ETextureReallocationStatus RHIFinalizeAsyncReallocateTexture2D(FTexture2DRHIParamRef Texture2D, bool bBlockUntilCompleted) final override;
diff --git a/Engine/Source/Runtime/Engine/Public/Rendering/PositionVertexBuffer.h b/Engine/Source/Runtime/Engine/Public/Rendering/PositionVertexBuffer.h
--- a/Engine/Source/Runtime/Engine/Public/Rendering/PositionVertexBuffer.h
+++ b/Engine/Source/Runtime/Engine/Public/Rendering/PositionVertexBuffer.h
@@ -136,7 +136,7 @@ public:
 		return Data;
 	}
 
-	const FShaderResourceViewRHIParamRef GetSRV() const { return PositionComponentSRV; }
+	FRHIShaderResourceView* GetSRV() const { return PositionComponentSRV; }
 
 private:
 
diff --git a/Engine/Source/Runtime/Engine/Public/Rendering/StaticMeshVertexBuffer.h b/Engine/Source/Runtime/Engine/Public/Rendering/StaticMeshVertexBuffer.h
--- a/Engine/Source/Runtime/Engine/Public/Rendering/StaticMeshVertexBuffer.h
+++ b/Engine/Source/Runtime/Engine/Public/Rendering/StaticMeshVertexBuffer.h
@@ -493,8 +493,8 @@ public:
 		return IsValidRef(TangentsVertexBuffer.VertexBufferRHI) && IsValidRef(TexCoordVertexBuffer.VertexBufferRHI);
 	}
 
-	const FShaderResourceViewRHIParamRef GetTangentsSRV() const { return TangentsSRV; }
-	const FShaderResourceViewRHIParamRef GetTexCoordsSRV() const { return TextureCoordinatesSRV; }
+	FRHIShaderResourceView* GetTangentsSRV() const { return TangentsSRV; }
+	FRHIShaderResourceView* GetTexCoordsSRV() const { return TextureCoordinatesSRV; }
 private:
 
 	/** The vertex data storage type */
