DR - Merge fixes
#rnx

diff --git a/Engine/Source/Runtime/D3D12RHI/Private/D3D12View.cpp b/Engine/Source/Runtime/D3D12RHI/Private/D3D12View.cpp
--- a/Engine/Source/Runtime/D3D12RHI/Private/D3D12View.cpp
+++ b/Engine/Source/Runtime/D3D12RHI/Private/D3D12View.cpp
@@ -153,41 +153,6 @@ FShaderResourceViewRHIRef FD3D12DynamicRHI::RHICreateShaderResourceView(FTexture
 	return CreateSRV(BaseTexture, SRVDesc);
 }
 
-FShaderResourceViewRHIRef FD3D12DynamicRHI::RHICreateShaderResourceView(FStructuredBufferRHIParamRef StructuredBufferRHI)
-{
-	struct FD3D12InitializeStructuredBufferSRVRHICommand final : public FRHICommand<FD3D12InitializeStructuredBufferSRVRHICommand>
-	{
-		FD3D12StructuredBuffer* Buffer;
-		FD3D12ShaderResourceView* SRV;
-
-		FD3D12InitializeStructuredBufferSRVRHICommand(FD3D12StructuredBuffer* InBuffer, FD3D12ShaderResourceView* InSRV)
-			: Buffer(InBuffer)
-			, SRV(InSRV)
-		{}
-
-		void Execute(FRHICommandListBase& RHICmdList)
-		{
-			FD3D12ResourceLocation& Location = Buffer->ResourceLocation;
-			const uint64 Offset = Location.GetOffsetFromBaseOfResource();
-			const D3D12_RESOURCE_DESC& BufferDesc = Location.GetResource()->GetDesc();
-
-			const uint32 BufferUsage = Buffer->GetUsage();
-			const bool bByteAccessBuffer = (BufferUsage & BUF_ByteAddressBuffer) != 0;
-			const bool bUINT8Access = (BufferUsage & BUF_UINT8) != 0;
-
-			// Create a Shader Resource View
-			D3D12_SHADER_RESOURCE_VIEW_DESC SRVDesc = {};
-			SRVDesc.Shader4ComponentMapping = D3D12_DEFAULT_SHADER_4_COMPONENT_MAPPING;
-
-			// BufferDesc.StructureByteStride  is not getting patched through the D3D resource DESC structs, so use the RHI version as a hack
-			uint32 Stride = Buffer->GetStride();
-
-
-	SRVDesc.Format = PlatformShaderResourceFormat;
-
-	return CreateSRV(Texture3D, SRVDesc);
-}
-
 FShaderResourceViewRHIRef FD3D12DynamicRHI::RHICreateShaderResourceView(FStructuredBufferRHIParamRef StructuredBufferRHI)
 {
 	struct FD3D12InitializeStructuredBufferSRVRHICommand final : public FRHICommand<FD3D12InitializeStructuredBufferSRVRHICommand>
@@ -211,7 +176,7 @@ FShaderResourceViewRHIRef FD3D12DynamicRHI::RHICreateShaderResourceView(FStructu
 			const bool bByteAccessBuffer = (BufferUsage & BUF_ByteAddressBuffer) != 0;
 			const bool bUINT8Access = (BufferUsage & BUF_UINT8) != 0;
 			// Create a Shader Resource View
-			D3D12_SHADER_RESOURCE_VIEW_DESC SRVDesc = {};
+			D3D12_SHADER_RESOURCE_VIEW_DESC SRVDesc ={};
 			SRVDesc.Shader4ComponentMapping = D3D12_DEFAULT_SHADER_4_COMPONENT_MAPPING;
 
 			// BufferDesc.StructureByteStride  is not getting patched through the D3D resource DESC structs, so use the RHI version as a hack
diff --git a/Engine/Source/Runtime/RHI/Public/RHIValidation.h b/Engine/Source/Runtime/RHI/Public/RHIValidation.h
--- a/Engine/Source/Runtime/RHI/Public/RHIValidation.h
+++ b/Engine/Source/Runtime/RHI/Public/RHIValidation.h
@@ -643,10 +643,17 @@ public:
 	* @param Flags - ETextureCreateFlags creation flags
 	*/
 	// FlushType: Wait RHI Thread
+	PRAGMA_DISABLE_DEPRECATION_WARNINGS
 	virtual FTexture2DArrayRHIRef RHICreateTexture2DArray(uint32 SizeX, uint32 SizeY, uint32 SizeZ, uint8 Format, uint32 NumMips, uint32 Flags, FRHIResourceCreateInfo& CreateInfo) override final
 	{
 		return RHI->RHICreateTexture2DArray(SizeX, SizeY, SizeZ, Format, NumMips, Flags, CreateInfo);
 	}
+	PRAGMA_ENABLE_DEPRECATION_WARNINGS
+
+	virtual FTexture2DArrayRHIRef RHICreateTexture2DArray(uint32 SizeX, uint32 SizeY, uint32 SizeZ, uint8 Format, uint32 NumMips, uint32 NumSamples, uint32 Flags, FRHIResourceCreateInfo& CreateInfo)
+	{
+		return RHI->RHICreateTexture2DArray(SizeX, SizeY, SizeZ, Format, NumMips, NumSamples, Flags, CreateInfo);
+	}
 
 	/**
 	* Creates a 3d RHI texture resource
@@ -1412,10 +1419,12 @@ public:
 		return RHI->RHICreateTextureExternal2D_RenderThread(RHICmdList, SizeX, SizeY, Format, NumMips, NumSamples, Flags, CreateInfo);
 	}
 
+	PRAGMA_DISABLE_DEPRECATION_WARNINGS
 	virtual FTexture2DArrayRHIRef RHICreateTexture2DArray_RenderThread(class FRHICommandListImmediate& RHICmdList, uint32 SizeX, uint32 SizeY, uint32 SizeZ, uint8 Format, uint32 NumMips, uint32 Flags, FRHIResourceCreateInfo& CreateInfo) override final
 	{
 		return RHI->RHICreateTexture2DArray_RenderThread(RHICmdList, SizeX, SizeY, SizeZ, Format, NumMips, Flags, CreateInfo);
 	}
+	PRAGMA_ENABLE_DEPRECATION_WARNINGS
 
 	virtual FTexture3DRHIRef RHICreateTexture3D_RenderThread(class FRHICommandListImmediate& RHICmdList, uint32 SizeX, uint32 SizeY, uint32 SizeZ, uint8 Format, uint32 NumMips, uint32 Flags, FRHIResourceCreateInfo& CreateInfo) override final
 	{
