Proper alignment check - as discussed in UDN:
https://udn.unrealengine.com/questions/501803/alignment-issues-in-fd3d12buddyallocatorallocate.html

diff --git a/Engine/Source/Runtime/D3D12RHI/Private/D3D12Allocation.cpp b/Engine/Source/Runtime/D3D12RHI/Private/D3D12Allocation.cpp
--- a/Engine/Source/Runtime/D3D12RHI/Private/D3D12Allocation.cpp
+++ b/Engine/Source/Runtime/D3D12RHI/Private/D3D12Allocation.cpp
@@ -269,6 +269,7 @@ void FD3D12BuddyAllocator::Allocate(uint32 SizeInBytes, uint32 Alignment, FD3D12
 	}
 #endif
 	const uint32 AlignedOffsetFromResourceBase = AllocationBlockOffset + Padding;
+	check((AlignedOffsetFromResourceBase % Alignment) == 0);
 
 	// Setup the info that this allocator
 	FD3D12BuddyAllocatorPrivateData& PrivateData = ResourceLocation.GetBuddyAllocatorPrivateData();
@@ -295,12 +296,6 @@ void FD3D12BuddyAllocator::Allocate(uint32 SizeInBytes, uint32 Alignment, FD3D12
 		// Place resources are intialized elsewhere
 	}
 
-	if (Alignment != 0)
-	{
-		check(uint64(ResourceLocation.GetMappedBaseAddress()) % Alignment == 0);
-		check(uint64(ResourceLocation.GetGPUVirtualAddress()) % Alignment == 0);
-	}
-
 	// track the allocation
 #if !PLATFORM_WINDOWS
 	LLM(FLowLevelMemTracker::Get().OnLowLevelAlloc(ELLMTracker::Default, &ResourceLocation, SizeInBytes));
