Better alignment test, catches non-integer offset case

diff --git a/Engine/Source/Runtime/D3D12RHI/Private/D3D12View.h b/Engine/Source/Runtime/D3D12RHI/Private/D3D12View.h
--- a/Engine/Source/Runtime/D3D12RHI/Private/D3D12View.h
+++ b/Engine/Source/Runtime/D3D12RHI/Private/D3D12View.h
@@ -850,7 +850,7 @@ public:
 
 		if (InDesc.ViewDimension == D3D12_SRV_DIMENSION_BUFFER)
 		{
-			check(InResourceLocation.GetOffsetFromBaseOfResource() / Stride == InDesc.Buffer.FirstElement);
+			check(InResourceLocation.GetOffsetFromBaseOfResource() == InDesc.Buffer.FirstElement * Stride);
 		}
 #endif
 
