Remove unused variable

diff --git a/Engine/Source/Runtime/D3D12RHI/Private/D3D12DescriptorCache.cpp b/Engine/Source/Runtime/D3D12RHI/Private/D3D12DescriptorCache.cpp
--- a/Engine/Source/Runtime/D3D12RHI/Private/D3D12DescriptorCache.cpp
+++ b/Engine/Source/Runtime/D3D12RHI/Private/D3D12DescriptorCache.cpp
@@ -734,42 +734,39 @@ void FD3D12DescriptorCache::SetConstantBuffers(const FD3D12RootSignature* RootSi
 	FMsg::Logf(__FILE__, __LINE__, TEXT("DescriptorCache"), ELogVerbosity::Log, TEXT("SetShaderResourceViewTable [STAGE %d] to slots %d - %d"), (int32)ShaderStage, FirstSlotIndex, FirstSlotIndex + SlotsNeeded - 1);
 #endif
 #else
-	auto& CBVs = Cache.CurrentGPUVirtualAddress[ShaderStage];
+
+	// Set root descriptors.
+	// At least one needed root descriptor is dirty.
+	const uint32 BaseIndex = RootSignature->CBVRDBaseBindSlot(ShaderStage);
+	ensure(BaseIndex != 255);
+	const uint32 RDCBVsNeeded = FMath::FloorLog2(RDCBVSlotsNeededMask) + 1;	// Get the index of the most significant bit that's set.
+	check(RDCBVsNeeded <= MAX_ROOT_CBVS);
+	for (uint32 SlotIndex = 0; SlotIndex < RDCBVsNeeded; SlotIndex++)
 	{
-		// Set root descriptors.
-		// At least one needed root descriptor is dirty.
-
-		const uint32 BaseIndex = RootSignature->CBVRDBaseBindSlot(ShaderStage);
-		ensure(BaseIndex != 255);
-		const uint32 RDCBVsNeeded = FMath::FloorLog2(RDCBVSlotsNeededMask) + 1;	// Get the index of the most significant bit that's set.
-		check(RDCBVsNeeded <= MAX_ROOT_CBVS);
-		for (uint32 SlotIndex = 0; SlotIndex < RDCBVsNeeded; SlotIndex++)
+		// Only set the root descriptor if it's dirty and we need to set it (it can be used by the shader).
+		if (FD3D12ConstantBufferCache::IsSlotDirty(RDCBVSlotsNeededMask, SlotIndex))
 		{
-			// Only set the root descriptor if it's dirty and we need to set it (it can be used by the shader).
-			if (FD3D12ConstantBufferCache::IsSlotDirty(RDCBVSlotsNeededMask, SlotIndex))
+			const D3D12_GPU_VIRTUAL_ADDRESS& CurrentGPUVirtualAddress = Cache.CurrentGPUVirtualAddress[ShaderStage][SlotIndex];
+			check(CurrentGPUVirtualAddress != 0);
+			if (ShaderStage == SF_Compute)
 			{
-				const D3D12_GPU_VIRTUAL_ADDRESS& CurrentGPUVirtualAddress = Cache.CurrentGPUVirtualAddress[ShaderStage][SlotIndex];
-				check(CurrentGPUVirtualAddress != 0);
-				if (ShaderStage == SF_Compute)
-				{
-					CommandList->SetComputeRootConstantBufferView(BaseIndex + SlotIndex, CurrentGPUVirtualAddress);
-				}
-				else
-				{
-					CommandList->SetGraphicsRootConstantBufferView(BaseIndex + SlotIndex, CurrentGPUVirtualAddress);
-				}
+				CommandList->SetComputeRootConstantBufferView(BaseIndex + SlotIndex, CurrentGPUVirtualAddress);
+			}
+			else
+			{
+				CommandList->SetGraphicsRootConstantBufferView(BaseIndex + SlotIndex, CurrentGPUVirtualAddress);
+			}
 
-				// Update residency.
-				CommandList.UpdateResidency(Cache.ResidencyHandles[ShaderStage][SlotIndex]);
+			// Update residency.
+			CommandList.UpdateResidency(Cache.ResidencyHandles[ShaderStage][SlotIndex]);
 
-				// Clear the dirty bit.
-				FD3D12ConstantBufferCache::CleanSlot(CurrentDirtySlotMask, SlotIndex);
-			}
+			// Clear the dirty bit.
+			FD3D12ConstantBufferCache::CleanSlot(CurrentDirtySlotMask, SlotIndex);
 		}
-		check((CurrentDirtySlotMask & RDCBVSlotsNeededMask) == 0);	// Check all slots that needed to be set, were set.
-
-		static_assert(GDescriptorTableCBVSlotMask == 0, "FD3D12DescriptorCache::SetConstantBuffers needs to be updated to handle descriptor tables.");	// Check that all CBVs slots are controlled by root descriptors.
 	}
+	check((CurrentDirtySlotMask & RDCBVSlotsNeededMask) == 0);	// Check all slots that needed to be set, were set.
+
+	static_assert(GDescriptorTableCBVSlotMask == 0, "FD3D12DescriptorCache::SetConstantBuffers needs to be updated to handle descriptor tables.");	// Check that all CBVs slots are controlled by root descriptors.
 #endif	
 }
 
