Removed FRHIGPUMask uint32 cast operator and made uint32 constructor private. Added explicit member function to retrieve native mask.

diff --git a/Engine/Source/Runtime/D3D12RHI/Private/D3D12Adapter.cpp b/Engine/Source/Runtime/D3D12RHI/Private/D3D12Adapter.cpp
--- a/Engine/Source/Runtime/D3D12RHI/Private/D3D12Adapter.cpp
+++ b/Engine/Source/Runtime/D3D12RHI/Private/D3D12Adapter.cpp
@@ -530,7 +530,7 @@ void FD3D12Adapter::CreateSignatures()
 	D3D12_COMMAND_SIGNATURE_DESC commandSignatureDesc = {};
 	commandSignatureDesc.NumArgumentDescs = 1;
 	commandSignatureDesc.ByteStride = 20;
-	commandSignatureDesc.NodeMask = (uint32)FRHIGPUMask::All();
+	commandSignatureDesc.NodeMask = FRHIGPUMask::All().GetNative();
 
 	D3D12_INDIRECT_ARGUMENT_DESC indirectParameterDesc[1] = {};
 	commandSignatureDesc.pArgumentDescs = indirectParameterDesc;
diff --git a/Engine/Source/Runtime/D3D12RHI/Private/D3D12Adapter.h b/Engine/Source/Runtime/D3D12RHI/Private/D3D12Adapter.h
--- a/Engine/Source/Runtime/D3D12RHI/Private/D3D12Adapter.h
+++ b/Engine/Source/Runtime/D3D12RHI/Private/D3D12Adapter.h
@@ -181,6 +181,7 @@ public:
 
 	// Resource Creation
 	HRESULT CreateCommittedResource(const D3D12_RESOURCE_DESC& Desc,
+		FRHIGPUMask CreationNode,
 		const D3D12_HEAP_PROPERTIES& HeapProps,
 		const D3D12_RESOURCE_STATES& InitialUsage,
 		const D3D12_CLEAR_VALUE* ClearValue,
@@ -213,6 +214,7 @@ public:
 		D3D12_RESOURCE_FLAGS Flags = D3D12_RESOURCE_FLAG_NONE);
 
 	HRESULT CreateBuffer(const D3D12_HEAP_PROPERTIES& HeapProps,
+		FRHIGPUMask CreationNode,
 		D3D12_RESOURCE_STATES InitialState,
 		uint64 HeapSize,
 		FD3D12Resource** ppOutResource,
diff --git a/Engine/Source/Runtime/D3D12RHI/Private/D3D12Allocation.cpp b/Engine/Source/Runtime/D3D12RHI/Private/D3D12Allocation.cpp
--- a/Engine/Source/Runtime/D3D12RHI/Private/D3D12Allocation.cpp
+++ b/Engine/Source/Runtime/D3D12RHI/Private/D3D12Allocation.cpp
@@ -117,8 +117,8 @@ void FD3D12BuddyAllocator::Initialize()
 	if (AllocationStrategy == eBuddyAllocationStrategy::kPlacedResourceStrategy)
 	{
 		D3D12_HEAP_PROPERTIES HeapProps = CD3DX12_HEAP_PROPERTIES(HeapType);
-		HeapProps.CreationNodeMask = (uint32)GetGPUMask();
-		HeapProps.VisibleNodeMask = (uint32)GetVisibilityMask();
+		HeapProps.CreationNodeMask = GetGPUMask().GetNative();
+		HeapProps.VisibleNodeMask = GetVisibilityMask().GetNative();
 
 		D3D12_HEAP_DESC Desc = {};
 		Desc.SizeInBytes = MaxBlockSize;
@@ -1168,9 +1168,9 @@ HRESULT FD3D12TextureAllocatorPool::AllocateTexture(
 		}
 	}
 
-	const D3D12_HEAP_PROPERTIES HeapProps = CD3DX12_HEAP_PROPERTIES(D3D12_HEAP_TYPE_DEFAULT, (uint32)GetGPUMask(), (uint32)GetVisibilityMask());
+	const D3D12_HEAP_PROPERTIES HeapProps = CD3DX12_HEAP_PROPERTIES(D3D12_HEAP_TYPE_DEFAULT, GetGPUMask().GetNative(), GetVisibilityMask().GetNative());
 	Desc.Alignment = 0;
-	RetCode = Adapter->CreateCommittedResource(Desc, HeapProps, InitialState, ClearValue, &NewResource, Name);
+	RetCode = Adapter->CreateCommittedResource(Desc, GetGPUMask(), HeapProps, InitialState, ClearValue, &NewResource, Name);
 
 	TextureLocation.SetType(FD3D12ResourceLocation::ResourceLocationType::eStandAlone);
 	TextureLocation.SetResource(NewResource);
@@ -1394,7 +1394,7 @@ FD3D12FastAllocatorPagePool::FD3D12FastAllocatorPagePool(FD3D12Device* Parent, F
 	: FD3D12DeviceChild(Parent)
 	, FD3D12MultiNodeGPUObject(Parent->GetGPUMask(), VisibiltyMask)
 	, PageSize(Size)
-	, HeapProperties(CD3DX12_HEAP_PROPERTIES(InHeapType, (uint32)Parent->GetGPUMask(), (uint32)VisibiltyMask))
+	, HeapProperties(CD3DX12_HEAP_PROPERTIES(InHeapType, Parent->GetGPUMask().GetNative(), VisibiltyMask.GetNative()))
 {};
 
 FD3D12FastAllocatorPagePool::FD3D12FastAllocatorPagePool(FD3D12Device* Parent, FRHIGPUMask VisibiltyMask, const D3D12_HEAP_PROPERTIES& InHeapProperties, uint32 Size)
@@ -1431,7 +1431,7 @@ FD3D12FastAllocatorPage* FD3D12FastAllocatorPagePool::RequestFastAllocatorPage()
 	Page = new FD3D12FastAllocatorPage(PageSize);
 
 	const D3D12_RESOURCE_STATES InitialState = DetermineInitialResourceState(HeapProperties.Type, &HeapProperties);
-	VERIFYD3D12RESULT(Adapter->CreateBuffer(HeapProperties, InitialState, PageSize, Page->FastAllocBuffer.GetInitReference(), TEXT("Fast Allocator Page")));
+	VERIFYD3D12RESULT(Adapter->CreateBuffer(HeapProperties, GetGPUMask(), InitialState, PageSize, Page->FastAllocBuffer.GetInitReference(), TEXT("Fast Allocator Page")));
 
 	Page->FastAllocData = Page->FastAllocBuffer->Map();
 	return Page;
@@ -1620,7 +1620,7 @@ FD3D12SegHeap* FD3D12SegList::CreateBackingHeap(
 	ID3D12Heap* D3DHeap;
 	D3D12_HEAP_DESC Desc = {};
 	Desc.SizeInBytes = HeapSize;
-	Desc.Properties = CD3DX12_HEAP_PROPERTIES(HeapType, (uint32)Parent->GetGPUMask(), (uint32)VisibleNodeMask);
+	Desc.Properties = CD3DX12_HEAP_PROPERTIES(HeapType, Parent->GetGPUMask().GetNative(), VisibleNodeMask.GetNative());
 	Desc.Flags = HeapFlags;
 
 	VERIFYD3D12RESULT(Parent->GetDevice()->CreateHeap(&Desc, IID_PPV_ARGS(&D3DHeap)));
diff --git a/Engine/Source/Runtime/D3D12RHI/Private/D3D12CommandList.cpp b/Engine/Source/Runtime/D3D12RHI/Private/D3D12CommandList.cpp
--- a/Engine/Source/Runtime/D3D12RHI/Private/D3D12CommandList.cpp
+++ b/Engine/Source/Runtime/D3D12RHI/Private/D3D12CommandList.cpp
@@ -52,7 +52,7 @@ FD3D12CommandListHandle::FD3D12CommandListData::FD3D12CommandListData(FD3D12Devi
 	, StartTimeQueryIdx(INDEX_NONE)
 #endif
 {
-	VERIFYD3D12RESULT(ParentDevice->GetDevice()->CreateCommandList((uint32)GetGPUMask(), CommandListType, CommandAllocator, nullptr, IID_PPV_ARGS(CommandList.GetInitReference())));
+	VERIFYD3D12RESULT(ParentDevice->GetDevice()->CreateCommandList(GetGPUMask().GetNative(), CommandListType, CommandAllocator, nullptr, IID_PPV_ARGS(CommandList.GetInitReference())));
 	INC_DWORD_STAT(STAT_D3D12NumCommandLists);
 
 #if PLATFORM_WINDOWS
diff --git a/Engine/Source/Runtime/D3D12RHI/Private/D3D12DescriptorCache.cpp b/Engine/Source/Runtime/D3D12RHI/Private/D3D12DescriptorCache.cpp
--- a/Engine/Source/Runtime/D3D12RHI/Private/D3D12DescriptorCache.cpp
+++ b/Engine/Source/Runtime/D3D12RHI/Private/D3D12DescriptorCache.cpp
@@ -909,7 +909,7 @@ void FD3D12GlobalOnlineHeap::Init(uint32 TotalSize, D3D12_DESCRIPTOR_HEAP_TYPE T
 	Desc.Flags = HeapFlags;
 	Desc.Type = Type;
 	Desc.NumDescriptors = TotalSize;
-	Desc.NodeMask = (uint32)GetGPUMask();
+	Desc.NodeMask = GetGPUMask().GetNative();
 
 	//LLM_SCOPE(ELLMTag::DescriptorCache);
 
@@ -1068,7 +1068,7 @@ void FD3D12ThreadLocalOnlineHeap::Init(uint32 NumDescriptors, D3D12_DESCRIPTOR_H
 	Desc.Flags = D3D12_DESCRIPTOR_HEAP_FLAG_SHADER_VISIBLE;
 	Desc.Type = Type;
 	Desc.NumDescriptors = NumDescriptors;
-	Desc.NodeMask = (uint32)GetGPUMask();
+	Desc.NodeMask = GetGPUMask().GetNative();
 
 	//LLM_SCOPE(ELLMTag::DescriptorCache);
 	VERIFYD3D12RESULT(GetParentDevice()->GetDevice()->CreateDescriptorHeap(&Desc, IID_PPV_ARGS(Heap.GetInitReference())));
diff --git a/Engine/Source/Runtime/D3D12RHI/Private/D3D12DescriptorCache.h b/Engine/Source/Runtime/D3D12RHI/Private/D3D12DescriptorCache.h
--- a/Engine/Source/Runtime/D3D12RHI/Private/D3D12DescriptorCache.h
+++ b/Engine/Source/Runtime/D3D12RHI/Private/D3D12DescriptorCache.h
@@ -186,7 +186,7 @@ private: // Types
 		Desc.Type = Type;
 		Desc.NumDescriptors = NumDescriptorsPerHeap;
 		Desc.Flags = D3D12_DESCRIPTOR_HEAP_FLAG_NONE;// None as this heap is offline
-		Desc.NodeMask = (uint32)Node;
+		Desc.NodeMask = Node.GetNative();
 
 		return Desc;
 	}
diff --git a/Engine/Source/Runtime/D3D12RHI/Private/D3D12DirectCommandListManager.cpp b/Engine/Source/Runtime/D3D12RHI/Private/D3D12DirectCommandListManager.cpp
--- a/Engine/Source/Runtime/D3D12RHI/Private/D3D12DirectCommandListManager.cpp
+++ b/Engine/Source/Runtime/D3D12RHI/Private/D3D12DirectCommandListManager.cpp
@@ -354,7 +354,7 @@ void FD3D12CommandListManager::Create(const TCHAR* Name, uint32 NumCommandLists,
 
 	D3D12_COMMAND_QUEUE_DESC CommandQueueDesc = {};
 	CommandQueueDesc.Flags = D3D12_COMMAND_QUEUE_FLAG_NONE;
-	CommandQueueDesc.NodeMask = (uint32)GetGPUMask();
+	CommandQueueDesc.NodeMask = GetGPUMask().GetNative();
 	CommandQueueDesc.Priority = Priority;
 	CommandQueueDesc.Type = CommandListType;
 	D3DCommandQueue = Adapter->GetOwningRHI()->CreateCommandQueue(Device, CommandQueueDesc);
diff --git a/Engine/Source/Runtime/D3D12RHI/Private/D3D12Pipelinestate.cpp b/Engine/Source/Runtime/D3D12RHI/Private/D3D12Pipelinestate.cpp
--- a/Engine/Source/Runtime/D3D12RHI/Private/D3D12Pipelinestate.cpp
+++ b/Engine/Source/Runtime/D3D12RHI/Private/D3D12Pipelinestate.cpp
@@ -598,7 +598,7 @@ FD3D12GraphicsPipelineState* FD3D12PipelineStateCacheBase::FindInLoadedCache(
 {
 	// TODO: For now PSOs will be created on every node of the LDA chain.
 	OutLowLevelDesc = GetLowLevelGraphicsPipelineStateDesc(Initializer, RootSignature);
-	OutLowLevelDesc.Desc.NodeMask = (uint32)FRHIGPUMask::All();
+	OutLowLevelDesc.Desc.NodeMask = FRHIGPUMask::All().GetNative();
 
 	OutLowLevelDesc.CombinedHash = FD3D12PipelineStateCacheBase::HashPSODesc(OutLowLevelDesc);
 
@@ -677,7 +677,7 @@ FD3D12ComputePipelineState* FD3D12PipelineStateCacheBase::FindInLoadedCache(FD3D
 {
 	// TODO: For now PSOs will be created on every node of the LDA chain.
 	OutLowLevelDesc = GetComputePipelineStateDesc(ComputeShader);
-	OutLowLevelDesc.Desc.NodeMask = (uint32)FRHIGPUMask::All();
+	OutLowLevelDesc.Desc.NodeMask = FRHIGPUMask::All().GetNative();
 	OutLowLevelDesc.CombinedHash = FD3D12PipelineStateCacheBase::HashPSODesc(OutLowLevelDesc);
 
 	// First try to find the PSO in the low level cache that can be populated from disk.
diff --git a/Engine/Source/Runtime/D3D12RHI/Private/D3D12Query.cpp b/Engine/Source/Runtime/D3D12RHI/Private/D3D12Query.cpp
--- a/Engine/Source/Runtime/D3D12RHI/Private/D3D12Query.cpp
+++ b/Engine/Source/Runtime/D3D12RHI/Private/D3D12Query.cpp
@@ -413,7 +413,7 @@ void FD3D12QueryHeap::CreateQueryHeap()
 	D3D12_QUERY_HEAP_DESC QueryHeapDesc;
 	QueryHeapDesc.Type = (QueryType == D3D12_QUERY_TYPE_OCCLUSION)? D3D12_QUERY_HEAP_TYPE_OCCLUSION : D3D12_QUERY_HEAP_TYPE_TIMESTAMP;
 	QueryHeapDesc.Count = QueryHeapCount;
-	QueryHeapDesc.NodeMask = (uint32)GetGPUMask();
+	QueryHeapDesc.NodeMask = GetGPUMask().GetNative();
 
 	// Create the upload heap
 	VERIFYD3D12RESULT(GetParentDevice()->GetDevice()->CreateQueryHeap(&QueryHeapDesc, IID_PPV_ARGS(&QueryHeap)));
@@ -426,12 +426,13 @@ void FD3D12QueryHeap::CreateQueryHeap()
 
 	FD3D12Adapter* Adapter = GetParentDevice()->GetParentAdapter();
 
-	const D3D12_HEAP_PROPERTIES ResultBufferHeapProperties = CD3DX12_HEAP_PROPERTIES(D3D12_HEAP_TYPE_READBACK, (uint32)GetGPUMask(), (uint32)GetVisibilityMask());
+	const D3D12_HEAP_PROPERTIES ResultBufferHeapProperties = CD3DX12_HEAP_PROPERTIES(D3D12_HEAP_TYPE_READBACK, GetGPUMask().GetNative(), GetVisibilityMask().GetNative());
 	const D3D12_RESOURCE_DESC ResultBufferDesc = CD3DX12_RESOURCE_DESC::Buffer(ResultSize * QueryHeapDesc.Count); // Each query's result occupies ResultSize bytes.
 
 	// Create the readback heap
 	VERIFYD3D12RESULT(Adapter->CreateCommittedResource(
 		ResultBufferDesc,
+		GetGPUMask(),
 		ResultBufferHeapProperties,
 		D3D12_RESOURCE_STATE_COPY_DEST,
 		nullptr,
@@ -626,7 +627,7 @@ void FD3D12LinearQueryHeap::CreateQueryHeap(int32 NumQueries, ID3D12QueryHeap**
 	D3D12_QUERY_HEAP_DESC Desc;
 	Desc.Type = QueryHeapType;
 	Desc.Count = static_cast<uint32>(NumQueries);
-	Desc.NodeMask = static_cast<uint32>(GetGPUMask());
+	Desc.NodeMask = GetGPUMask().GetNative();
 	VERIFYD3D12RESULT(GetParentDevice()->GetDevice()->CreateQueryHeap(&Desc, IID_PPV_ARGS(OutHeap)));
 	SetName(*OutHeap, TEXT("FD3D12LinearQueryHeap"));
 
@@ -639,11 +640,12 @@ void FD3D12LinearQueryHeap::CreateQueryHeap(int32 NumQueries, ID3D12QueryHeap**
 void FD3D12LinearQueryHeap::CreateResultBuffer(uint64 SizeInBytes, FD3D12Resource** OutBuffer)
 {
 	FD3D12Adapter* Adapter = GetParentDevice()->GetParentAdapter();
-	const D3D12_HEAP_PROPERTIES ResultBufferHeapProperties = CD3DX12_HEAP_PROPERTIES(D3D12_HEAP_TYPE_READBACK, (uint32)GetGPUMask(), (uint32)GetVisibilityMask());
+	const D3D12_HEAP_PROPERTIES ResultBufferHeapProperties = CD3DX12_HEAP_PROPERTIES(D3D12_HEAP_TYPE_READBACK, GetGPUMask().GetNative(), GetVisibilityMask().GetNative());
 	const D3D12_RESOURCE_DESC ResultBufferDesc = CD3DX12_RESOURCE_DESC::Buffer(SizeInBytes);
 
 	VERIFYD3D12RESULT(Adapter->CreateCommittedResource(
 		ResultBufferDesc,
+		GetGPUMask(),
 		ResultBufferHeapProperties,
 		D3D12_RESOURCE_STATE_COPY_DEST,
 		nullptr,
@@ -748,7 +750,7 @@ void FD3D12BufferedGPUTiming::InitDynamicRHI()
 		TimestampQueryHeap = Adapter->CreateLinkedObject<QueryHeap>(FRHIGPUMask::All(), [&] (FD3D12Device* Device)
 		{
 			QueryHeap* NewHeap = new QueryHeap(Device);
-			QueryHeapDesc.NodeMask = (uint32)Device->GetGPUMask();
+			QueryHeapDesc.NodeMask = Device->GetGPUMask().GetNative();
 			VERIFYD3D12RESULT(D3DDevice->CreateQueryHeap(&QueryHeapDesc, IID_PPV_ARGS(NewHeap->Heap.GetInitReference())));
 			SetName(NewHeap->Heap, L"FD3D12BufferedGPUTiming: Timestamp Query Heap");
 
diff --git a/Engine/Source/Runtime/D3D12RHI/Private/D3D12RayTracing.cpp b/Engine/Source/Runtime/D3D12RHI/Private/D3D12RayTracing.cpp
--- a/Engine/Source/Runtime/D3D12RHI/Private/D3D12RayTracing.cpp
+++ b/Engine/Source/Runtime/D3D12RHI/Private/D3D12RayTracing.cpp
@@ -773,7 +773,7 @@ public:
 		Desc.Flags = D3D12_DESCRIPTOR_HEAP_FLAG_SHADER_VISIBLE;
 		Desc.Type = Type;
 		Desc.NumDescriptors = NumDescriptors;
-		Desc.NodeMask = (uint32)GetParentDevice()->GetGPUMask();
+		Desc.NodeMask = GetParentDevice()->GetGPUMask().GetNative();
 
 		ID3D12DescriptorHeap* D3D12Heap = nullptr;
 
diff --git a/Engine/Source/Runtime/D3D12RHI/Private/D3D12RenderTarget.cpp b/Engine/Source/Runtime/D3D12RHI/Private/D3D12RenderTarget.cpp
--- a/Engine/Source/Runtime/D3D12RHI/Private/D3D12RenderTarget.cpp
+++ b/Engine/Source/Runtime/D3D12RHI/Private/D3D12RenderTarget.cpp
@@ -1181,8 +1181,8 @@ void FD3D12DynamicRHI::ReadSurfaceDataMSAARaw(FRHICommandList_RecursiveHazardous
 	NonMSAADesc.Flags = D3D12_RESOURCE_FLAG_ALLOW_RENDER_TARGET;
 	TRefCountPtr<FD3D12Resource> NonMSAATexture2D;
 
-	const D3D12_HEAP_PROPERTIES HeapProps = CD3DX12_HEAP_PROPERTIES(D3D12_HEAP_TYPE_DEFAULT, (uint32)NodeMask, (uint32)NodeMask);
-	VERIFYD3D12RESULT(Adapter->CreateCommittedResource(NonMSAADesc, HeapProps, D3D12_RESOURCE_STATE_RENDER_TARGET, nullptr, NonMSAATexture2D.GetInitReference(), nullptr));
+	const D3D12_HEAP_PROPERTIES HeapProps = CD3DX12_HEAP_PROPERTIES(D3D12_HEAP_TYPE_DEFAULT, NodeMask.GetNative(), NodeMask.GetNative());
+	VERIFYD3D12RESULT(Adapter->CreateCommittedResource(NonMSAADesc, NodeMask, HeapProps, D3D12_RESOURCE_STATE_RENDER_TARGET, nullptr, NonMSAATexture2D.GetInitReference(), nullptr));
 
 	FD3D12ResourceLocation ResourceLocation(Device);
 	ResourceLocation.AsStandAlone(NonMSAATexture2D);
diff --git a/Engine/Source/Runtime/D3D12RHI/Private/D3D12Resources.cpp b/Engine/Source/Runtime/D3D12RHI/Private/D3D12Resources.cpp
--- a/Engine/Source/Runtime/D3D12RHI/Private/D3D12Resources.cpp
+++ b/Engine/Source/Runtime/D3D12RHI/Private/D3D12Resources.cpp
@@ -318,7 +318,7 @@ void FD3D12Heap::BeginTrackingResidency(uint64 Size)
 //	FD3D12 Adapter
 /////////////////////////////////////////////////////////////////////
 
-HRESULT FD3D12Adapter::CreateCommittedResource(const D3D12_RESOURCE_DESC& InDesc, const D3D12_HEAP_PROPERTIES& HeapProps, const D3D12_RESOURCE_STATES& InitialUsage, const D3D12_CLEAR_VALUE* ClearValue, FD3D12Resource** ppOutResource, const TCHAR* Name)
+HRESULT FD3D12Adapter::CreateCommittedResource(const D3D12_RESOURCE_DESC& InDesc, FRHIGPUMask CreationNode, const D3D12_HEAP_PROPERTIES& HeapProps, const D3D12_RESOURCE_STATES& InitialUsage, const D3D12_CLEAR_VALUE* ClearValue, FD3D12Resource** ppOutResource, const TCHAR* Name)
 {
 	if (!ppOutResource)
 	{
@@ -337,7 +337,7 @@ HRESULT FD3D12Adapter::CreateCommittedResource(const D3D12_RESOURCE_DESC& InDesc
 		SetName(pResource, Name);
 
 		// Set the output pointer
-		*ppOutResource = new FD3D12Resource(GetDevice(FRHIGPUMask(HeapProps.CreationNodeMask).ToIndex()), FRHIGPUMask(HeapProps.VisibleNodeMask), pResource, InitialUsage, InDesc, nullptr, HeapProps.Type);
+		*ppOutResource = new FD3D12Resource(GetDevice(CreationNode.ToIndex()), CreationNode, pResource, InitialUsage, InDesc, nullptr, HeapProps.Type);
 		(*ppOutResource)->AddRef();
 
 		// Only track resources that cannot be accessed on the CPU.
@@ -388,18 +388,19 @@ HRESULT FD3D12Adapter::CreatePlacedResource(const D3D12_RESOURCE_DESC& InDesc, F
 
 HRESULT FD3D12Adapter::CreateBuffer(D3D12_HEAP_TYPE HeapType, FRHIGPUMask CreationNode, FRHIGPUMask VisibleNodes, uint64 HeapSize, FD3D12Resource** ppOutResource, const TCHAR* Name, D3D12_RESOURCE_FLAGS Flags)
 {
-	const D3D12_HEAP_PROPERTIES HeapProps = CD3DX12_HEAP_PROPERTIES(HeapType, (uint32)CreationNode, (uint32)VisibleNodes);
+	const D3D12_HEAP_PROPERTIES HeapProps = CD3DX12_HEAP_PROPERTIES(HeapType, CreationNode.GetNative(), VisibleNodes.GetNative());
 	const D3D12_RESOURCE_STATES InitialState = DetermineInitialResourceState(HeapProps.Type, &HeapProps);
-	return CreateBuffer(HeapProps, InitialState, HeapSize, ppOutResource, Name, Flags);
+	return CreateBuffer(HeapProps, CreationNode, InitialState, HeapSize, ppOutResource, Name, Flags);
 }
 
 HRESULT FD3D12Adapter::CreateBuffer(D3D12_HEAP_TYPE HeapType, FRHIGPUMask CreationNode, FRHIGPUMask VisibleNodes, D3D12_RESOURCE_STATES InitialState, uint64 HeapSize, FD3D12Resource** ppOutResource, const TCHAR* Name, D3D12_RESOURCE_FLAGS Flags)
 {
-	const D3D12_HEAP_PROPERTIES HeapProps = CD3DX12_HEAP_PROPERTIES(HeapType, (uint32)CreationNode, (uint32)VisibleNodes);
-	return CreateBuffer(HeapProps, InitialState, HeapSize, ppOutResource, Name, Flags);
+	const D3D12_HEAP_PROPERTIES HeapProps = CD3DX12_HEAP_PROPERTIES(HeapType, CreationNode.GetNative(), VisibleNodes.GetNative());
+	return CreateBuffer(HeapProps, CreationNode, InitialState, HeapSize, ppOutResource, Name, Flags);
 }
 
 HRESULT FD3D12Adapter::CreateBuffer(const D3D12_HEAP_PROPERTIES& HeapProps,
+	FRHIGPUMask CreationNode,
 	D3D12_RESOURCE_STATES InitialState,
 	uint64 HeapSize,
 	FD3D12Resource** ppOutResource,
@@ -413,6 +414,7 @@ HRESULT FD3D12Adapter::CreateBuffer(const D3D12_HEAP_PROPERTIES& HeapProps,
 
 	const D3D12_RESOURCE_DESC BufDesc = CD3DX12_RESOURCE_DESC::Buffer(HeapSize, Flags);
 	return CreateCommittedResource(BufDesc,
+		CreationNode,
 		HeapProps,
 		InitialState,
 		nullptr,
diff --git a/Engine/Source/Runtime/D3D12RHI/Private/D3D12RootSignature.cpp b/Engine/Source/Runtime/D3D12RHI/Private/D3D12RootSignature.cpp
--- a/Engine/Source/Runtime/D3D12RHI/Private/D3D12RootSignature.cpp
+++ b/Engine/Source/Runtime/D3D12RHI/Private/D3D12RootSignature.cpp
@@ -378,7 +378,7 @@ void FD3D12RootSignature::Init(const D3D12_VERSIONED_ROOT_SIGNATURE_DESC& InDesc
 	VERIFYD3D12RESULT(SerializeHR);
 
 	// Create and analyze the root signature.
-	VERIFYD3D12RESULT(Device->CreateRootSignature((uint32)FRHIGPUMask::All(),
+	VERIFYD3D12RESULT(Device->CreateRootSignature(FRHIGPUMask::All().GetNative(),
 		RootSignatureBlob->GetBufferPointer(),
 		RootSignatureBlob->GetBufferSize(),
 		IID_PPV_ARGS(RootSignature.GetInitReference())));
@@ -398,7 +398,7 @@ void FD3D12RootSignature::Init(ID3DBlob* const InBlob, uint32 BindingSpace)
 	VERIFYD3D12RESULT(D3D12CreateVersionedRootSignatureDeserializer(RootSignatureBlob->GetBufferPointer(), RootSignatureBlob->GetBufferSize(), IID_PPV_ARGS(Deserializer.GetInitReference())));
 
 	// Create and analyze the root signature.
-	VERIFYD3D12RESULT(Device->CreateRootSignature((uint32)FRHIGPUMask::All(),
+	VERIFYD3D12RESULT(Device->CreateRootSignature(FRHIGPUMask::All().GetNative(),
 		RootSignatureBlob->GetBufferPointer(),
 		RootSignatureBlob->GetBufferSize(),
 		IID_PPV_ARGS(RootSignature.GetInitReference())));
diff --git a/Engine/Source/Runtime/D3D12RHI/Private/Windows/WindowsD3D12PipelineState.cpp b/Engine/Source/Runtime/D3D12RHI/Private/Windows/WindowsD3D12PipelineState.cpp
--- a/Engine/Source/Runtime/D3D12RHI/Private/Windows/WindowsD3D12PipelineState.cpp
+++ b/Engine/Source/Runtime/D3D12RHI/Private/Windows/WindowsD3D12PipelineState.cpp
@@ -233,7 +233,7 @@ void FD3D12PipelineStateCache::RebuildFromDiskCache(ID3D12RootSignature* Graphic
 		if (!DiskCaches[PSO_CACHE_GRAPHICS].IsInErrorState())
 		{
 			// Only reload PSOs that match the LDA mask, or else it will fail.
-			if ((uint32)FRHIGPUMask::All() == Desc->Desc.NodeMask)
+			if (FRHIGPUMask::All().GetNative() == Desc->Desc.NodeMask)
 			{
 				// Add PSO to low level cache.
 				FD3D12PipelineState* PipelineState = nullptr;
@@ -287,7 +287,7 @@ void FD3D12PipelineStateCache::RebuildFromDiskCache(ID3D12RootSignature* Graphic
 		if (!DiskCaches[PSO_CACHE_COMPUTE].IsInErrorState())
 		{
 			// Only reload PSOs that match the LDA mask, or else it will fail.
-			if ((uint32)FRHIGPUMask::All() == Desc->Desc.NodeMask)
+			if (FRHIGPUMask::All().GetNative() == Desc->Desc.NodeMask)
 			{
 				Desc->CombinedHash = FD3D12PipelineStateCache::HashPSODesc(*Desc);
 
diff --git a/Engine/Source/Runtime/RHI/Public/MultiGPU.h b/Engine/Source/Runtime/RHI/Public/MultiGPU.h
--- a/Engine/Source/Runtime/RHI/Public/MultiGPU.h
+++ b/Engine/Source/Runtime/RHI/Public/MultiGPU.h
@@ -29,9 +29,9 @@ extern RHI_API uint32 GVirtualMGPU;
 struct FRHIGPUMask
 {
 private:
+
 	uint32 GPUMask;
 
-public:
 	FORCEINLINE explicit FRHIGPUMask(uint32 InGPUMask) : GPUMask(InGPUMask)
 	{
 #if WITH_MGPU
@@ -41,7 +41,9 @@ public:
 #endif
 	}
 
-	FORCEINLINE FRHIGPUMask() : GPUMask(FRHIGPUMask::GPU0())
+public:
+
+	FORCEINLINE FRHIGPUMask() : FRHIGPUMask(FRHIGPUMask::GPU0())
 	{
 	}
 
@@ -93,7 +95,7 @@ public:
 	void operator |=(const FRHIGPUMask& Rhs) { GPUMask |= Rhs.GPUMask; }
 	void operator &=(const FRHIGPUMask& Rhs) { GPUMask &= Rhs.GPUMask; }
 
-	FORCEINLINE explicit operator uint32() const { return GVirtualMGPU ? 1 : GPUMask; }
+	FORCEINLINE uint32 GetNative() const { return GVirtualMGPU ? 1 : GPUMask; }
 
 	FORCEINLINE FRHIGPUMask operator &(const FRHIGPUMask& Rhs) const
 	{
