Fixed compile error due to missing file in CL 7061991.

diff --git a/Engine/Source/Runtime/D3D12RHI/Private/Windows/WindowsD3D12Viewport.cpp b/Engine/Source/Runtime/D3D12RHI/Private/Windows/WindowsD3D12Viewport.cpp
--- a/Engine/Source/Runtime/D3D12RHI/Private/Windows/WindowsD3D12Viewport.cpp
+++ b/Engine/Source/Runtime/D3D12RHI/Private/Windows/WindowsD3D12Viewport.cpp
@@ -256,7 +256,7 @@ void FD3D12Viewport::ResizeInternal()
 			FD3D12Device* Device = Adapter->GetDevice(GPUIndex);
 
 			CommandQueues.Add(Device->GetD3DCommandQueue());
-			NodeMasks.Add((uint32)Device->GetGPUMask());
+			NodeMasks.Add(Device->GetGPUMask().GetNative());
 		}
 
 		TRefCountPtr<IDXGISwapChain3> SwapChain3;
