DR - Fix deprecations from merge from main
#rnx

diff --git a/Engine/Source/Runtime/Apple/MetalRHI/Public/MetalRHI.h b/Engine/Source/Runtime/Apple/MetalRHI/Public/MetalRHI.h
--- a/Engine/Source/Runtime/Apple/MetalRHI/Public/MetalRHI.h
+++ b/Engine/Source/Runtime/Apple/MetalRHI/Public/MetalRHI.h
@@ -70,7 +70,7 @@ public:
 	virtual FVertexBufferRHIRef RHICreateVertexBuffer(uint32 Size, uint32 InUsage, FRHIResourceCreateInfo& CreateInfo) final override;
 	virtual void* LockVertexBuffer_BottomOfPipe(FRHICommandListImmediate& RHICmdList, FRHIVertexBuffer* VertexBuffer, uint32 Offset, uint32 SizeRHI, EResourceLockMode LockMode) final override;
 	virtual void UnlockVertexBuffer_BottomOfPipe(FRHICommandListImmediate& RHICmdList, FRHIVertexBuffer* VertexBuffer) final override;
-	virtual void RHICopyVertexBuffer(FVertexBufferRHIParamRef SourceBuffer, FRHIVertexBuffer* DestBuffer) final override;
+	virtual void RHICopyVertexBuffer(FRHIVertexBuffer* SourceBuffer, FRHIVertexBuffer* DestBuffer) final override;
 	virtual void RHITransferVertexBufferUnderlyingResource(FRHIVertexBuffer* DestVertexBuffer, FRHIVertexBuffer* SrcVertexBuffer) final override;
 	virtual FStructuredBufferRHIRef RHICreateStructuredBuffer(uint32 Stride, uint32 Size, uint32 InUsage, FRHIResourceCreateInfo& CreateInfo) final override;
 	virtual void* LockStructuredBuffer_BottomOfPipe(FRHICommandListImmediate& RHICmdList, FRHIStructuredBuffer* StructuredBuffer, uint32 Offset, uint32 SizeRHI, EResourceLockMode LockMode) final override;
diff --git a/Engine/Source/Runtime/D3D12RHI/Private/D3D12RHIPrivate.h b/Engine/Source/Runtime/D3D12RHI/Private/D3D12RHIPrivate.h
--- a/Engine/Source/Runtime/D3D12RHI/Private/D3D12RHIPrivate.h
+++ b/Engine/Source/Runtime/D3D12RHI/Private/D3D12RHIPrivate.h
@@ -283,26 +283,26 @@ public:
 	virtual FUniformBufferRHIRef RHICreateUniformBuffer(const void* Contents, const FRHIUniformBufferLayout& Layout, EUniformBufferUsage Usage, EUniformBufferValidation Validation) final override;
 	virtual void RHIUpdateUniformBuffer(FRHIUniformBuffer* UniformBufferRHI, const void* Contents) final override;
 	virtual FIndexBufferRHIRef RHICreateIndexBuffer(uint32 Stride, uint32 Size, uint32 InUsage, FRHIResourceCreateInfo& CreateInfo) final override;
-	virtual void* RHILockIndexBuffer(FRHICommandListImmediate& RHICmdList, FIndexBufferRHIParamRef IndexBuffer, uint32 Offset, uint32 Size, EResourceLockMode LockMode) final override;
-	virtual void RHIUnlockIndexBuffer(FRHICommandListImmediate& RHICmdList, FIndexBufferRHIParamRef IndexBuffer) final override;
-	virtual void RHITransferIndexBufferUnderlyingResource(FIndexBufferRHIParamRef DestIndexBuffer, FIndexBufferRHIParamRef SrcIndexBuffer) final override;
+	virtual void* RHILockIndexBuffer(FRHICommandListImmediate& RHICmdList, FRHIIndexBuffer* IndexBuffer, uint32 Offset, uint32 Size, EResourceLockMode LockMode) final override;
+	virtual void RHIUnlockIndexBuffer(FRHICommandListImmediate& RHICmdList, FRHIIndexBuffer* IndexBuffer) final override;
+	virtual void RHITransferIndexBufferUnderlyingResource(FRHIIndexBuffer* DestIndexBuffer, FRHIIndexBuffer* SrcIndexBuffer) final override;
 	virtual FVertexBufferRHIRef RHICreateVertexBuffer(uint32 Size, uint32 InUsage, FRHIResourceCreateInfo& CreateInfo) final override;
-	virtual void* RHILockVertexBuffer(FRHICommandListImmediate& RHICmdList, FVertexBufferRHIParamRef VertexBuffer, uint32 Offset, uint32 SizeRHI, EResourceLockMode LockMode) final override;
-	virtual void RHIUnlockVertexBuffer(FRHICommandListImmediate& RHICmdList, FVertexBufferRHIParamRef VertexBuffer) final override;
-	virtual void RHICopyVertexBuffer(FVertexBufferRHIParamRef SourceBuffer, FVertexBufferRHIParamRef DestBuffer) final override;
-	virtual void RHITransferVertexBufferUnderlyingResource(FVertexBufferRHIParamRef DestVertexBuffer, FVertexBufferRHIParamRef SrcVertexBuffer) final override;
+	virtual void* RHILockVertexBuffer(FRHICommandListImmediate& RHICmdList, FRHIVertexBuffer* VertexBuffer, uint32 Offset, uint32 SizeRHI, EResourceLockMode LockMode) final override;
+	virtual void RHIUnlockVertexBuffer(FRHICommandListImmediate& RHICmdList, FRHIVertexBuffer* VertexBuffer) final override;
+	virtual void RHICopyVertexBuffer(FRHIVertexBuffer* SourceBuffer, FRHIVertexBuffer* DestBuffer) final override;
+	virtual void RHITransferVertexBufferUnderlyingResource(FRHIVertexBuffer* DestVertexBuffer, FRHIVertexBuffer* SrcVertexBuffer) final override;
 	virtual FStructuredBufferRHIRef RHICreateStructuredBuffer(uint32 Stride, uint32 Size, uint32 InUsage, FRHIResourceCreateInfo& CreateInfo) final override;
-	virtual void* RHILockStructuredBuffer(FRHICommandListImmediate& RHICmdList, FStructuredBufferRHIParamRef StructuredBuffer, uint32 Offset, uint32 SizeRHI, EResourceLockMode LockMode) final override;
-	virtual void RHIUnlockStructuredBuffer(FRHICommandListImmediate& RHICmdList, FStructuredBufferRHIParamRef StructuredBuffer) final override;
-	virtual FUnorderedAccessViewRHIRef RHICreateUnorderedAccessView(FStructuredBufferRHIParamRef StructuredBuffer, bool bUseUAVCounter, bool bAppendBuffer) final override;
-	virtual FUnorderedAccessViewRHIRef RHICreateUnorderedAccessView(FTextureRHIParamRef Texture, uint32 MipLevel) final override;
-	virtual FUnorderedAccessViewRHIRef RHICreateUnorderedAccessView(FVertexBufferRHIParamRef VertexBuffer, uint8 Format) final override;
-	virtual FUnorderedAccessViewRHIRef RHICreateUnorderedAccessView(FIndexBufferRHIParamRef IndexBuffer, uint8 Format) final override;
-	virtual FShaderResourceViewRHIRef RHICreateShaderResourceView(FStructuredBufferRHIParamRef StructuredBuffer) final override;
-	virtual FShaderResourceViewRHIRef RHICreateShaderResourceView(FVertexBufferRHIParamRef VertexBuffer, uint32 Stride, uint8 Format) final override;
-	virtual FShaderResourceViewRHIRef RHICreateShaderResourceView(FIndexBufferRHIParamRef Buffer) final override;
-	virtual void RHIUpdateShaderResourceView(FRHIShaderResourceView* SRV, FVertexBufferRHIParamRef VertexBuffer, uint32 Stride, uint8 Format) final override;
-	virtual void RHIUpdateShaderResourceView(FRHIShaderResourceView* SRV, FIndexBufferRHIParamRef IndexBuffer) final override;
+	virtual void* RHILockStructuredBuffer(FRHICommandListImmediate& RHICmdList, FRHIStructuredBuffer* StructuredBuffer, uint32 Offset, uint32 SizeRHI, EResourceLockMode LockMode) final override;
+	virtual void RHIUnlockStructuredBuffer(FRHICommandListImmediate& RHICmdList, FRHIStructuredBuffer* StructuredBuffer) final override;
+	virtual FUnorderedAccessViewRHIRef RHICreateUnorderedAccessView(FRHIStructuredBuffer* StructuredBuffer, bool bUseUAVCounter, bool bAppendBuffer) final override;
+	virtual FUnorderedAccessViewRHIRef RHICreateUnorderedAccessView(FRHITexture* Texture, uint32 MipLevel) final override;
+	virtual FUnorderedAccessViewRHIRef RHICreateUnorderedAccessView(FRHIVertexBuffer* VertexBuffer, uint8 Format) final override;
+	virtual FUnorderedAccessViewRHIRef RHICreateUnorderedAccessView(FRHIIndexBuffer* IndexBuffer, uint8 Format) final override;
+	virtual FShaderResourceViewRHIRef RHICreateShaderResourceView(FRHIStructuredBuffer* StructuredBuffer) final override;
+	virtual FShaderResourceViewRHIRef RHICreateShaderResourceView(FRHIVertexBuffer* VertexBuffer, uint32 Stride, uint8 Format) final override;
+	virtual FShaderResourceViewRHIRef RHICreateShaderResourceView(FRHIIndexBuffer* Buffer) final override;
+	virtual void RHIUpdateShaderResourceView(FRHIShaderResourceView* SRV, FRHIVertexBuffer* VertexBuffer, uint32 Stride, uint8 Format) final override;
+	virtual void RHIUpdateShaderResourceView(FRHIShaderResourceView* SRV, FRHIIndexBuffer* IndexBuffer) final override;
 	virtual uint64 RHICalcTexture2DPlatformSize(uint32 SizeX, uint32 SizeY, uint8 Format, uint32 NumMips, uint32 NumSamples, uint32 Flags, uint32& OutAlign) override;
 	virtual uint64 RHICalcTexture3DPlatformSize(uint32 SizeX, uint32 SizeY, uint32 SizeZ, uint8 Format, uint32 NumMips, uint32 Flags, uint32& OutAlign) final override;
 	virtual uint64 RHICalcTextureCubePlatformSize(uint32 Size, uint8 Format, uint32 NumMips, uint32 Flags, uint32& OutAlign) final override;
@@ -428,9 +428,9 @@ public:
 	// virtual ETextureReallocationStatus FinalizeAsyncReallocateTexture2D_RenderThread(class FRHICommandListImmediate& RHICmdList, FRHITexture2D* Texture2D, bool bBlockUntilCompleted);
 	// virtual ETextureReallocationStatus CancelAsyncReallocateTexture2D_RenderThread(class FRHICommandListImmediate& RHICmdList, FRHITexture2D* Texture2D, bool bBlockUntilCompleted);
 	virtual FIndexBufferRHIRef CreateIndexBuffer_RenderThread(class FRHICommandListImmediate& RHICmdList, uint32 Stride, uint32 Size, uint32 InUsage, FRHIResourceCreateInfo& CreateInfo);
-	virtual void UpdateTexture2D_RenderThread(class FRHICommandListImmediate& RHICmdList, FTexture2DRHIParamRef Texture, uint32 MipIndex, const struct FUpdateTextureRegion2D& UpdateRegion, uint32 SourcePitch, const uint8* SourceData);
-	virtual void* LockTexture2D_RenderThread(class FRHICommandListImmediate& RHICmdList, FTexture2DRHIParamRef Texture, uint32 MipIndex, EResourceLockMode LockMode, uint32& DestStride, bool bLockWithinMiptail, bool bNeedsDefaultRHIFlush = true);
-	virtual void UnlockTexture2D_RenderThread(class FRHICommandListImmediate& RHICmdList, FTexture2DRHIParamRef Texture, uint32 MipIndex, bool bLockWithinMiptail, bool bNeedsDefaultRHIFlush = true);
+	virtual void UpdateTexture2D_RenderThread(class FRHICommandListImmediate& RHICmdList, FRHITexture2D* Texture, uint32 MipIndex, const struct FUpdateTextureRegion2D& UpdateRegion, uint32 SourcePitch, const uint8* SourceData);
+	virtual void* LockTexture2D_RenderThread(class FRHICommandListImmediate& RHICmdList, FRHITexture2D* Texture, uint32 MipIndex, EResourceLockMode LockMode, uint32& DestStride, bool bLockWithinMiptail, bool bNeedsDefaultRHIFlush = true);
+	virtual void UnlockTexture2D_RenderThread(class FRHICommandListImmediate& RHICmdList, FRHITexture2D* Texture, uint32 MipIndex, bool bLockWithinMiptail, bool bNeedsDefaultRHIFlush = true);
 	// 
 	virtual FTexture2DRHIRef RHICreateTexture2D_RenderThread(class FRHICommandListImmediate& RHICmdList, uint32 SizeX, uint32 SizeY, uint8 Format, uint32 NumMips, uint32 NumSamples, uint32 Flags, FRHIResourceCreateInfo& CreateInfo);
 	virtual FTexture2DArrayRHIRef RHICreateTexture2DArray_RenderThread(class FRHICommandListImmediate& RHICmdList, uint32 SizeX, uint32 SizeY, uint32 SizeZ, uint8 Format, uint32 NumMips, uint32 NumSamples, uint32 Flags, FRHIResourceCreateInfo& CreateInfo);
diff --git a/Engine/Source/Runtime/D3D12RHI/Private/D3D12RayTracing.cpp b/Engine/Source/Runtime/D3D12RHI/Private/D3D12RayTracing.cpp
--- a/Engine/Source/Runtime/D3D12RHI/Private/D3D12RayTracing.cpp
+++ b/Engine/Source/Runtime/D3D12RHI/Private/D3D12RayTracing.cpp
@@ -1323,17 +1323,17 @@ public:
 	struct FShaderRecordCacheKey
 	{
 		static constexpr uint32 MaxUniformBuffers = 4;
-		const FUniformBufferRHIParamRef* UniformBuffers[MaxUniformBuffers];
+		FRHIUniformBuffer* const* UniformBuffers[MaxUniformBuffers];
 		uint64 Hash;
 		uint32 NumUniformBuffers;
 
 		FShaderRecordCacheKey() = default;
-		FShaderRecordCacheKey(uint32 InNumUniformBuffers, const FUniformBufferRHIParamRef* InUniformBuffers)
+		FShaderRecordCacheKey(uint32 InNumUniformBuffers, FRHIUniformBuffer* const* InUniformBuffers)
 		{
 			check(InNumUniformBuffers <= MaxUniformBuffers);
 			NumUniformBuffers = FMath::Min(MaxUniformBuffers, InNumUniformBuffers);
 
-			const uint64 DataSizeInBytes = sizeof(FUniformBufferRHIParamRef) * NumUniformBuffers;
+			const uint64 DataSizeInBytes = sizeof(FRHIUniformBuffer*) * NumUniformBuffers;
 			FMemory::Memcpy(UniformBuffers, InUniformBuffers, DataSizeInBytes);
 			Hash = CityHash64(reinterpret_cast<const char*>(UniformBuffers), DataSizeInBytes);
 		}
diff --git a/Engine/Source/Runtime/EmptyRHI/Public/EmptyRHI.h b/Engine/Source/Runtime/EmptyRHI/Public/EmptyRHI.h
--- a/Engine/Source/Runtime/EmptyRHI/Public/EmptyRHI.h
+++ b/Engine/Source/Runtime/EmptyRHI/Public/EmptyRHI.h
@@ -51,24 +51,24 @@ public:
 	virtual FBoundShaderStateRHIRef RHICreateBoundShaderState(FRHIVertexDeclaration* VertexDeclaration, FRHIVertexShader* VertexShader, FRHIHullShader* HullShader, FRHIDomainShaderFRHIDomainShader*, FRHIPixelShader* PixelShader, FRHIGeometryShader* GeometryShader) final override;
 	virtual FUniformBufferRHIRef RHICreateUniformBuffer(const void* Contents, const FRHIUniformBufferLayout& Layout, EUniformBufferUsage Usage, EUniformBufferValidation Validation) final override;
 	virtual FIndexBufferRHIRef RHICreateIndexBuffer(uint32 Stride, uint32 Size, uint32 InUsage, FRHIResourceCreateInfo& CreateInfo) final override;
-	virtual void* LockIndexBuffer_BottomOfPipe(FRHICommandListImmediate& RHICmdList, FIndexBufferRHIParamRef IndexBuffer, uint32 Offset, uint32 Size, EResourceLockMode LockMode) final override;
-	virtual void UnlockIndexBuffer_BottomOfPipe(FRHICommandListImmediate& RHICmdList, FIndexBufferRHIParamRef IndexBuffer) final override;
-	virtual void RHITransferIndexBufferUnderlyingResource(FIndexBufferRHIParamRef DestIndexBuffer, FIndexBufferRHIParamRef SrcIndexBuffer) final override;
+	virtual void* LockIndexBuffer_BottomOfPipe(FRHICommandListImmediate& RHICmdList, FRHIIndexBuffer* IndexBuffer, uint32 Offset, uint32 Size, EResourceLockMode LockMode) final override;
+	virtual void UnlockIndexBuffer_BottomOfPipe(FRHICommandListImmediate& RHICmdList, FRHIIndexBuffer* IndexBuffer) final override;
+	virtual void RHITransferIndexBufferUnderlyingResource(FRHIIndexBuffer* DestIndexBuffer, FRHIIndexBuffer* SrcIndexBuffer) final override;
 	virtual FVertexBufferRHIRef RHICreateVertexBuffer(uint32 Size, uint32 InUsage, FRHIResourceCreateInfo& CreateInfo) final override;
-	virtual void* LockVertexBuffer_BottomOfPipe(FRHICommandListImmediate& RHICmdList, FVertexBufferRHIParamRef VertexBuffer, uint32 Offset, uint32 SizeRHI, EResourceLockMode LockMode) final override;
-	virtual void UnlockVertexBuffer_BottomOfPipe(FRHICommandListImmediate& RHICmdList, FVertexBufferRHIParamRef VertexBuffer) final override;
-	virtual void RHICopyVertexBuffer(FVertexBufferRHIParamRef SourceBuffer,FVertexBufferRHIParamRef DestBuffer) final override;
-	virtual void RHITransferVertexBufferUnderlyingResource(FVertexBufferRHIParamRef DestVertexBuffer, FVertexBufferRHIParamRef SrcVertexBuffer) final override;
+	virtual void* LockVertexBuffer_BottomOfPipe(FRHICommandListImmediate& RHICmdList, FRHIVertexBuffer* VertexBuffer, uint32 Offset, uint32 SizeRHI, EResourceLockMode LockMode) final override;
+	virtual void UnlockVertexBuffer_BottomOfPipe(FRHICommandListImmediate& RHICmdList, FRHIVertexBuffer* VertexBuffer) final override;
+	virtual void RHICopyVertexBuffer(FRHIVertexBuffer* SourceBuffer, FRHIVertexBuffer* DestBuffer) final override;
+	virtual void RHITransferVertexBufferUnderlyingResource(FRHIVertexBuffer* DestVertexBuffer, FRHIVertexBuffer* SrcVertexBuffer) final override;
 	virtual FStructuredBufferRHIRef RHICreateStructuredBuffer(uint32 Stride, uint32 Size, uint32 InUsage, FRHIResourceCreateInfo& CreateInfo) final override;
-	virtual void* LockStructuredBuffer_BottomOfPipe(FRHICommandListImmediate& RHICmdList, FStructuredBufferRHIParamRef StructuredBuffer, uint32 Offset, uint32 SizeRHI, EResourceLockMode LockMode) final override;
-	virtual void UnlockStructuredBuffer_BottomOfPipe(FRHICommandListImmediate& RHICmdList, FStructuredBufferRHIParamRef StructuredBuffer) final override;
-	virtual FUnorderedAccessViewRHIRef RHICreateUnorderedAccessView(FStructuredBufferRHIParamRef StructuredBuffer, bool bUseUAVCounter, bool bAppendBuffer) final override;
-	virtual FUnorderedAccessViewRHIRef RHICreateUnorderedAccessView(FTextureRHIParamRef Texture, uint32 MipLevel) final override;
-	virtual FUnorderedAccessViewRHIRef RHICreateUnorderedAccessView(FVertexBufferRHIParamRef VertexBuffer, uint8 Format) final override;
-	virtual FUnorderedAccessViewRHIRef RHICreateUnorderedAccessView(FIndexBufferRHIParamRef IndexBuffer, uint8 Format) final override;
-	virtual FShaderResourceViewRHIRef RHICreateShaderResourceView(FStructuredBufferRHIParamRef StructuredBuffer) final override;
-	virtual FShaderResourceViewRHIRef RHICreateShaderResourceView(FVertexBufferRHIParamRef VertexBuffer, uint32 Stride, uint8 Format) final override;
-	virtual FShaderResourceViewRHIRef RHICreateShaderResourceView(FIndexBufferRHIParamRef Buffer) final override;
+	virtual void* LockStructuredBuffer_BottomOfPipe(FRHICommandListImmediate& RHICmdList, FRHIStructuredBuffer* StructuredBuffer, uint32 Offset, uint32 SizeRHI, EResourceLockMode LockMode) final override;
+	virtual void UnlockStructuredBuffer_BottomOfPipe(FRHICommandListImmediate& RHICmdList, FRHIStructuredBuffer* StructuredBuffer) final override;
+	virtual FUnorderedAccessViewRHIRef RHICreateUnorderedAccessView(FRHIStructuredBuffer* StructuredBuffer, bool bUseUAVCounter, bool bAppendBuffer) final override;
+	virtual FUnorderedAccessViewRHIRef RHICreateUnorderedAccessView(FRHITexture* Texture, uint32 MipLevel) final override;
+	virtual FUnorderedAccessViewRHIRef RHICreateUnorderedAccessView(FRHIVertexBuffer* VertexBuffer, uint8 Format) final override;
+	virtual FUnorderedAccessViewRHIRef RHICreateUnorderedAccessView(FRHIIndexBuffer* IndexBuffer, uint8 Format) final override;
+	virtual FShaderResourceViewRHIRef RHICreateShaderResourceView(FRHIStructuredBuffer* StructuredBuffer) final override;
+	virtual FShaderResourceViewRHIRef RHICreateShaderResourceView(FRHIVertexBuffer* VertexBuffer, uint32 Stride, uint8 Format) final override;
+	virtual FShaderResourceViewRHIRef RHICreateShaderResourceView(FRHIIndexBuffer* Buffer) final override;
 	virtual uint64 RHICalcTexture2DPlatformSize(uint32 SizeX, uint32 SizeY, uint8 Format, uint32 NumMips, uint32 NumSamples, uint32 Flags, uint32& OutAlign) final override;
 	virtual uint64 RHICalcTexture3DPlatformSize(uint32 SizeX, uint32 SizeY, uint32 SizeZ, uint8 Format, uint32 NumMips, uint32 Flags, uint32& OutAlign) final override;
 	virtual uint64 RHICalcTextureCubePlatformSize(uint32 Size, uint8 Format, uint32 NumMips, uint32 Flags, uint32& OutAlign) final override;
diff --git a/Engine/Source/Runtime/Engine/Private/GPUSkinVertexFactory.cpp b/Engine/Source/Runtime/Engine/Private/GPUSkinVertexFactory.cpp
--- a/Engine/Source/Runtime/Engine/Private/GPUSkinVertexFactory.cpp
+++ b/Engine/Source/Runtime/Engine/Private/GPUSkinVertexFactory.cpp
@@ -235,7 +235,7 @@ bool FGPUBaseSkinVertexFactory::FShaderDataType::UpdateBoneData(FRHICommandListI
 		{
 			if (!bUseSkinCache && DeferSkeletalLockAndFillToRHIThread())
 			{
-				FVertexBufferRHIParamRef VertexBuffer = CurrentBoneBuffer->VertexBufferRHI;
+				FRHIVertexBuffer* VertexBuffer = CurrentBoneBuffer->VertexBufferRHI;
 				RHICmdList.EnqueueLambda([VertexBuffer, VectorArraySize, &ReferenceToLocalMatrices, &BoneMap](FRHICommandListImmediate& RHICmdList)
 				{
 					QUICK_SCOPE_CYCLE_COUNTER(STAT_FRHICommandUpdateBoneBuffer_Execute);
@@ -919,7 +919,7 @@ bool FGPUBaseSkinAPEXClothVertexFactory::ClothShaderType::UpdateClothSimulData(F
 		{
 			if (DeferSkeletalLockAndFillToRHIThread())
 			{
-				FVertexBufferRHIParamRef VertexBuffer = CurrentClothBuffer->VertexBufferRHI;
+				FRHIVertexBuffer* VertexBuffer = CurrentClothBuffer->VertexBufferRHI;
 				RHICmdList.EnqueueLambda([VertexBuffer, VectorArraySize, &InSimulPositions, &InSimulNormals](FRHICommandListImmediate& RHICmdList)
 				{
 					QUICK_SCOPE_CYCLE_COUNTER(STAT_FRHICommandUpdateBoneBuffer_Execute);
diff --git a/Engine/Source/Runtime/NullDrv/Public/NullRHI.h b/Engine/Source/Runtime/NullDrv/Public/NullRHI.h
--- a/Engine/Source/Runtime/NullDrv/Public/NullRHI.h
+++ b/Engine/Source/Runtime/NullDrv/Public/NullRHI.h
@@ -132,11 +132,11 @@ public:
 		return new FRHIIndexBuffer(Stride,Size,InUsage); 
 	}
 
-	virtual void* LockIndexBuffer_BottomOfPipe(FRHICommandListImmediate& RHICmdList, FIndexBufferRHIParamRef IndexBuffer, uint32 Offset, uint32 Size, EResourceLockMode LockMode) final override
+	virtual void* LockIndexBuffer_BottomOfPipe(FRHICommandListImmediate& RHICmdList, FRHIIndexBuffer* IndexBuffer, uint32 Offset, uint32 Size, EResourceLockMode LockMode) final override
 	{ 
 		return GetStaticBuffer(); 
 	}
-	virtual void UnlockIndexBuffer_BottomOfPipe(FRHICommandListImmediate& RHICmdList, FIndexBufferRHIParamRef IndexBuffer) final override
+	virtual void UnlockIndexBuffer_BottomOfPipe(FRHICommandListImmediate& RHICmdList, FRHIIndexBuffer* IndexBuffer) final override
 	{
 
 	}
@@ -155,11 +155,12 @@ public:
 		return new FRHIVertexBuffer(Size,InUsage); 
 	}
 
-	virtual void* LockVertexBuffer_BottomOfPipe(FRHICommandListImmediate& RHICmdList, FVertexBufferRHIParamRef VertexBuffer, uint32 Offset, uint32 SizeRHI, EResourceLockMode LockMode) final override
+	virtual void* LockVertexBuffer_BottomOfPipe(FRHICommandListImmediate& RHICmdList, FRHIVertexBuffer* VertexBuffer, uint32 Offset, uint32 SizeRHI, EResourceLockMode LockMode) final override
 	{ 
 		return GetStaticBuffer(); 
 	}
-	virtual void UnlockVertexBuffer_BottomOfPipe(FRHICommandListImmediate& RHICmdList, FVertexBufferRHIParamRef VertexBuffer) final override
+
+	virtual void UnlockVertexBuffer_BottomOfPipe(FRHICommandListImmediate& RHICmdList, FRHIVertexBuffer* VertexBuffer) final override
 	{
 
 	}
@@ -184,11 +185,11 @@ public:
 		return new FRHIStructuredBuffer(Stride,Size,InUsage); 
 	}
 
-	virtual void* LockStructuredBuffer_BottomOfPipe(FRHICommandListImmediate& RHICmdList, FStructuredBufferRHIParamRef StructuredBuffer, uint32 Offset, uint32 SizeRHI, EResourceLockMode LockMode) final override
+	virtual void* LockStructuredBuffer_BottomOfPipe(FRHICommandListImmediate& RHICmdList, FRHIStructuredBuffer* StructuredBuffer, uint32 Offset, uint32 SizeRHI, EResourceLockMode LockMode) final override
 	{ 
 		return GetStaticBuffer(); 
 	}
-	virtual void UnlockStructuredBuffer_BottomOfPipe(FRHICommandListImmediate& RHICmdList, FStructuredBufferRHIParamRef StructuredBuffer) final override
+	virtual void UnlockStructuredBuffer_BottomOfPipe(FRHICommandListImmediate& RHICmdList, FRHIStructuredBuffer* StructuredBuffer) final override
 	{
 
 	}
diff --git a/Engine/Source/Runtime/OpenGLDrv/Public/OpenGLDrv.h b/Engine/Source/Runtime/OpenGLDrv/Public/OpenGLDrv.h
--- a/Engine/Source/Runtime/OpenGLDrv/Public/OpenGLDrv.h
+++ b/Engine/Source/Runtime/OpenGLDrv/Public/OpenGLDrv.h
@@ -364,23 +364,23 @@ public:
 	virtual FUniformBufferRHIRef RHICreateUniformBuffer(const void* Contents, const FRHIUniformBufferLayout& Layout, EUniformBufferUsage Usage, EUniformBufferValidation Validation) final override;
 	virtual void RHIUpdateUniformBuffer(FRHIUniformBuffer* UniformBufferRHI, const void* Contents) final override;
 	virtual FIndexBufferRHIRef RHICreateIndexBuffer(uint32 Stride, uint32 Size, uint32 InUsage, FRHIResourceCreateInfo& CreateInfo) final override;
-	virtual void* LockIndexBuffer_BottomOfPipe(FRHICommandListImmediate& RHICmdList, FIndexBufferRHIParamRef IndexBuffer, uint32 Offset, uint32 Size, EResourceLockMode LockMode) final override;
-	virtual void UnlockIndexBuffer_BottomOfPipe(FRHICommandListImmediate& RHICmdList, FIndexBufferRHIParamRef IndexBuffer) final override;
+	virtual void* LockIndexBuffer_BottomOfPipe(FRHICommandListImmediate& RHICmdList, FRHIIndexBuffer* IndexBuffer, uint32 Offset, uint32 Size, EResourceLockMode LockMode) final override;
+	virtual void UnlockIndexBuffer_BottomOfPipe(FRHICommandListImmediate& RHICmdList, FRHIIndexBuffer* IndexBuffer) final override;
 	virtual FVertexBufferRHIRef RHICreateVertexBuffer(uint32 Size, uint32 InUsage, FRHIResourceCreateInfo& CreateInfo) final override;
-	virtual void* LockVertexBuffer_BottomOfPipe(FRHICommandListImmediate& RHICmdList, FVertexBufferRHIParamRef VertexBuffer, uint32 Offset, uint32 SizeRHI, EResourceLockMode LockMode) final override;
-	virtual void UnlockVertexBuffer_BottomOfPipe(FRHICommandListImmediate& RHICmdList, FVertexBufferRHIParamRef VertexBuffer) final override;
-	virtual void RHICopyVertexBuffer(FVertexBufferRHIParamRef SourceBuffer,FVertexBufferRHIParamRef DestBuffer) final override;
+	virtual void* LockVertexBuffer_BottomOfPipe(FRHICommandListImmediate& RHICmdList, FRHIVertexBuffer* VertexBuffer, uint32 Offset, uint32 SizeRHI, EResourceLockMode LockMode) final override;
+	virtual void UnlockVertexBuffer_BottomOfPipe(FRHICommandListImmediate& RHICmdList, FRHIVertexBuffer* VertexBuffer) final override;
+	virtual void RHICopyVertexBuffer(FRHIVertexBuffer* SourceBuffer, FRHIVertexBuffer* DestBuffer) final override;
 	virtual FStructuredBufferRHIRef RHICreateStructuredBuffer(uint32 Stride, uint32 Size, uint32 InUsage, FRHIResourceCreateInfo& CreateInfo) final override;
-	virtual void* LockStructuredBuffer_BottomOfPipe(FRHICommandListImmediate& RHICmdList, FStructuredBufferRHIParamRef StructuredBuffer, uint32 Offset, uint32 SizeRHI, EResourceLockMode LockMode) final override;
-	virtual void UnlockStructuredBuffer_BottomOfPipe(FRHICommandListImmediate& RHICmdList, FStructuredBufferRHIParamRef StructuredBuffer) final override;
-	virtual FUnorderedAccessViewRHIRef RHICreateUnorderedAccessView(FStructuredBufferRHIParamRef StructuredBuffer, bool bUseUAVCounter, bool bAppendBuffer) final override;
-	virtual FUnorderedAccessViewRHIRef RHICreateUnorderedAccessView(FTextureRHIParamRef Texture, uint32 MipLevel) final override;
-	virtual FUnorderedAccessViewRHIRef RHICreateUnorderedAccessView(FVertexBufferRHIParamRef VertexBuffer, uint8 Format) final override;
-	virtual FUnorderedAccessViewRHIRef RHICreateUnorderedAccessView(FIndexBufferRHIParamRef IndexBuffer, uint8 Format) final override;
-	virtual FShaderResourceViewRHIRef RHICreateShaderResourceView(FStructuredBufferRHIParamRef StructuredBuffer) final override;
-	virtual FShaderResourceViewRHIRef RHICreateShaderResourceView(FVertexBufferRHIParamRef VertexBuffer, uint32 Stride, uint8 Format) final override;
-	virtual FShaderResourceViewRHIRef RHICreateShaderResourceView(FIndexBufferRHIParamRef Buffer) final override;
-	virtual void RHIUpdateShaderResourceView(FRHIShaderResourceView* SRV, FVertexBufferRHIParamRef VertexBuffer, uint32 Stride, uint8 Format) final override;
+	virtual void* LockStructuredBuffer_BottomOfPipe(FRHICommandListImmediate& RHICmdList, FRHIStructuredBuffer* StructuredBuffer, uint32 Offset, uint32 SizeRHI, EResourceLockMode LockMode) final override;
+	virtual void UnlockStructuredBuffer_BottomOfPipe(FRHICommandListImmediate& RHICmdList, FRHIStructuredBuffer* StructuredBuffer) final override;
+	virtual FUnorderedAccessViewRHIRef RHICreateUnorderedAccessView(FRHIStructuredBuffer* StructuredBuffer, bool bUseUAVCounter, bool bAppendBuffer) final override;
+	virtual FUnorderedAccessViewRHIRef RHICreateUnorderedAccessView(FRHITexture* Texture, uint32 MipLevel) final override;
+	virtual FUnorderedAccessViewRHIRef RHICreateUnorderedAccessView(FRHIVertexBuffer* VertexBuffer, uint8 Format) final override;
+	virtual FUnorderedAccessViewRHIRef RHICreateUnorderedAccessView(FRHIIndexBuffer* IndexBuffer, uint8 Format) final override;
+	virtual FShaderResourceViewRHIRef RHICreateShaderResourceView(FRHIStructuredBuffer* StructuredBuffer) final override;
+	virtual FShaderResourceViewRHIRef RHICreateShaderResourceView(FRHIVertexBuffer* VertexBuffer, uint32 Stride, uint8 Format) final override;
+	virtual FShaderResourceViewRHIRef RHICreateShaderResourceView(FRHIIndexBuffer* Buffer) final override;
+	virtual void RHIUpdateShaderResourceView(FRHIShaderResourceView* SRV, FRHIVertexBuffer* VertexBuffer, uint32 Stride, uint8 Format) final override;
 	virtual uint64 RHICalcTexture2DPlatformSize(uint32 SizeX, uint32 SizeY, uint8 Format, uint32 NumMips, uint32 NumSamples, uint32 Flags, uint32& OutAlign) final override;
 	virtual uint64 RHICalcTexture3DPlatformSize(uint32 SizeX, uint32 SizeY, uint32 SizeZ, uint8 Format, uint32 NumMips, uint32 Flags, uint32& OutAlign) final override;
 	virtual uint64 RHICalcTextureCubePlatformSize(uint32 Size, uint8 Format, uint32 NumMips, uint32 Flags, uint32& OutAlign) final override;
diff --git a/Engine/Source/Runtime/RHI/Private/RHICommandList.cpp b/Engine/Source/Runtime/RHI/Private/RHICommandList.cpp
--- a/Engine/Source/Runtime/RHI/Private/RHICommandList.cpp
+++ b/Engine/Source/Runtime/RHI/Private/RHICommandList.cpp
@@ -1920,7 +1920,7 @@ FShaderResourceViewRHIRef FDynamicRHI::CreateShaderResourceView_RenderThread(cla
 
 static FLockTracker GLockTracker;
 
-void* FDynamicRHI::RHILockVertexBuffer(class FRHICommandListImmediate& RHICmdList, FVertexBufferRHIParamRef VertexBuffer, uint32 Offset, uint32 SizeRHI, EResourceLockMode LockMode)
+void* FDynamicRHI::RHILockVertexBuffer(class FRHICommandListImmediate& RHICmdList, FRHIVertexBuffer* VertexBuffer, uint32 Offset, uint32 SizeRHI, EResourceLockMode LockMode)
 {
 	QUICK_SCOPE_CYCLE_COUNTER(STAT_FDynamicRHI_LockVertexBuffer);
 
@@ -1954,7 +1954,7 @@ void* FDynamicRHI::RHILockVertexBuffer(class FRHICommandListImmediate& RHICmdLis
 	return Result;	
 }
 
-void FDynamicRHI::RHIUnlockVertexBuffer(class FRHICommandListImmediate& RHICmdList, FVertexBufferRHIParamRef VertexBuffer)
+void FDynamicRHI::RHIUnlockVertexBuffer(class FRHICommandListImmediate& RHICmdList, FRHIVertexBuffer* VertexBuffer)
 {
 	QUICK_SCOPE_CYCLE_COUNTER(STAT_FDynamicRHI_UnlockVertexBuffer_RenderThread);
 
@@ -1998,7 +1998,7 @@ void FDynamicRHI::RHIUnlockVertexBuffer(class FRHICommandListImmediate& RHICmdLi
 	}
 }
 
-void* FDynamicRHI::RHILockIndexBuffer(class FRHICommandListImmediate& RHICmdList, FIndexBufferRHIParamRef IndexBuffer, uint32 Offset, uint32 SizeRHI, EResourceLockMode LockMode)
+void* FDynamicRHI::RHILockIndexBuffer(class FRHICommandListImmediate& RHICmdList, FRHIIndexBuffer* IndexBuffer, uint32 Offset, uint32 SizeRHI, EResourceLockMode LockMode)
 {
 	QUICK_SCOPE_CYCLE_COUNTER(STAT_FDynamicRHI_LockIndexBuffer);
 
@@ -2032,7 +2032,7 @@ void* FDynamicRHI::RHILockIndexBuffer(class FRHICommandListImmediate& RHICmdList
 	return Result;
 }
 
-void FDynamicRHI::RHIUnlockIndexBuffer(class FRHICommandListImmediate& RHICmdList, FIndexBufferRHIParamRef IndexBuffer)
+void FDynamicRHI::RHIUnlockIndexBuffer(class FRHICommandListImmediate& RHICmdList, FRHIIndexBuffer* IndexBuffer)
 {
 	QUICK_SCOPE_CYCLE_COUNTER(STAT_FDynamicRHI_UnlockIndexBuffer_RenderThread);
 
@@ -2076,7 +2076,7 @@ void FDynamicRHI::RHIUnlockIndexBuffer(class FRHICommandListImmediate& RHICmdLis
 	}
 }
 
-void* FDynamicRHI::RHILockStructuredBuffer(class FRHICommandListImmediate& RHICmdList, FStructuredBufferRHIParamRef StructuredBuffer, uint32 Offset, uint32 SizeRHI, EResourceLockMode LockMode)
+void* FDynamicRHI::RHILockStructuredBuffer(class FRHICommandListImmediate& RHICmdList, FRHIStructuredBuffer* StructuredBuffer, uint32 Offset, uint32 SizeRHI, EResourceLockMode LockMode)
 {
 	QUICK_SCOPE_CYCLE_COUNTER(STAT_FDynamicRHI_LockStructuredBuffer);
 
@@ -2110,7 +2110,7 @@ void* FDynamicRHI::RHILockStructuredBuffer(class FRHICommandListImmediate& RHICm
 	return Result;
 }
 
-void FDynamicRHI::RHIUnlockStructuredBuffer(class FRHICommandListImmediate& RHICmdList, FStructuredBufferRHIParamRef StructuredBuffer)
+void FDynamicRHI::RHIUnlockStructuredBuffer(class FRHICommandListImmediate& RHICmdList, FRHIStructuredBuffer* StructuredBuffer)
 {
 	QUICK_SCOPE_CYCLE_COUNTER(STAT_FDynamicRHI_UnlockStructuredBuffer_RenderThread);
 
diff --git a/Engine/Source/Runtime/RHI/Public/DynamicRHI.h b/Engine/Source/Runtime/RHI/Public/DynamicRHI.h
--- a/Engine/Source/Runtime/RHI/Public/DynamicRHI.h
+++ b/Engine/Source/Runtime/RHI/Public/DynamicRHI.h
@@ -1117,11 +1117,11 @@ public:
 
 	virtual FVertexBufferRHIRef CreateVertexBuffer_RenderThread(class FRHICommandListImmediate& RHICmdList, uint32 Size, uint32 InUsage, FRHIResourceCreateInfo& CreateInfo);
 	virtual FStructuredBufferRHIRef CreateStructuredBuffer_RenderThread(class FRHICommandListImmediate& RHICmdList, uint32 Stride, uint32 Size, uint32 InUsage, FRHIResourceCreateInfo& CreateInfo);
-	virtual FShaderResourceViewRHIRef CreateShaderResourceView_RenderThread(class FRHICommandListImmediate& RHICmdList, FVertexBufferRHIParamRef VertexBuffer, uint32 Stride, uint8 Format);
-	virtual FShaderResourceViewRHIRef CreateShaderResourceView_RenderThread(class FRHICommandListImmediate& RHICmdList, FIndexBufferRHIParamRef Buffer);
-	virtual FTexture2DRHIRef AsyncReallocateTexture2D_RenderThread(class FRHICommandListImmediate& RHICmdList, FTexture2DRHIParamRef Texture2D, int32 NewMipCount, int32 NewSizeX, int32 NewSizeY, FThreadSafeCounter* RequestStatus);
-	virtual ETextureReallocationStatus FinalizeAsyncReallocateTexture2D_RenderThread(class FRHICommandListImmediate& RHICmdList, FTexture2DRHIParamRef Texture2D, bool bBlockUntilCompleted);
-	virtual ETextureReallocationStatus CancelAsyncReallocateTexture2D_RenderThread(class FRHICommandListImmediate& RHICmdList, FTexture2DRHIParamRef Texture2D, bool bBlockUntilCompleted);
+	virtual FShaderResourceViewRHIRef CreateShaderResourceView_RenderThread(class FRHICommandListImmediate& RHICmdList, FRHIVertexBuffer* VertexBuffer, uint32 Stride, uint8 Format);
+	virtual FShaderResourceViewRHIRef CreateShaderResourceView_RenderThread(class FRHICommandListImmediate& RHICmdList, FRHIIndexBuffer* Buffer);
+	virtual FTexture2DRHIRef AsyncReallocateTexture2D_RenderThread(class FRHICommandListImmediate& RHICmdList, FRHITexture2D* Texture2D, int32 NewMipCount, int32 NewSizeX, int32 NewSizeY, FThreadSafeCounter* RequestStatus);
+	virtual ETextureReallocationStatus FinalizeAsyncReallocateTexture2D_RenderThread(class FRHICommandListImmediate& RHICmdList, FRHITexture2D* Texture2D, bool bBlockUntilCompleted);
+	virtual ETextureReallocationStatus CancelAsyncReallocateTexture2D_RenderThread(class FRHICommandListImmediate& RHICmdList, FRHITexture2D* Texture2D, bool bBlockUntilCompleted);
 	virtual FIndexBufferRHIRef CreateIndexBuffer_RenderThread(class FRHICommandListImmediate& RHICmdList, uint32 Stride, uint32 Size, uint32 InUsage, FRHIResourceCreateInfo& CreateInfo);
 	virtual FVertexShaderRHIRef CreateVertexShader_RenderThread(class FRHICommandListImmediate& RHICmdList, const TArray<uint8>& Code);
 	virtual FVertexShaderRHIRef CreateVertexShader_RenderThread(class FRHICommandListImmediate& RHICmdList, FRHIShaderLibrary* Library, FSHAHash Hash);
@@ -1191,36 +1191,36 @@ public:
 	virtual void RHIReadSurfaceFloatData_RenderThread(class FRHICommandListImmediate& RHICmdList, FRHITexture* Texture, FIntRect Rect, TArray<FFloat16Color>& OutData, ECubeFace CubeFace, int32 ArrayIndex, int32 MipIndex);
 
 	// Buffer Lock/Unlock
-	virtual void* LockStructuredBuffer_BottomOfPipe(class FRHICommandListImmediate& RHICmdList, FStructuredBufferRHIParamRef StructuredBuffer, uint32 Offset, uint32 SizeRHI, EResourceLockMode LockMode)
+	virtual void* LockStructuredBuffer_BottomOfPipe(class FRHICommandListImmediate& RHICmdList, FRHIStructuredBuffer* StructuredBuffer, uint32 Offset, uint32 SizeRHI, EResourceLockMode LockMode)
 	{
 		// Either this function or RHILockStructuredBuffer must be implemented by the platform RHI.
 		checkNoEntry();
 		return nullptr;
 	}
-	virtual void* LockVertexBuffer_BottomOfPipe(class FRHICommandListImmediate& RHICmdList, FVertexBufferRHIParamRef VertexBuffer, uint32 Offset, uint32 SizeRHI, EResourceLockMode LockMode)
+	virtual void* LockVertexBuffer_BottomOfPipe(class FRHICommandListImmediate& RHICmdList, FRHIVertexBuffer* VertexBuffer, uint32 Offset, uint32 SizeRHI, EResourceLockMode LockMode)
 	{
 		// Either this function or RHILockVertexBuffer must be implemented by the platform RHI.
 		checkNoEntry();
 		return nullptr;
 	}
-	virtual void* LockIndexBuffer_BottomOfPipe(class FRHICommandListImmediate& RHICmdList, FIndexBufferRHIParamRef IndexBuffer, uint32 Offset, uint32 SizeRHI, EResourceLockMode LockMode)
+	virtual void* LockIndexBuffer_BottomOfPipe(class FRHICommandListImmediate& RHICmdList, FRHIIndexBuffer* IndexBuffer, uint32 Offset, uint32 SizeRHI, EResourceLockMode LockMode)
 	{
 		// Either this function or RHILockIndexBuffer must be implemented by the platform RHI.
 		checkNoEntry();
 		return nullptr;
 	}
 
-	virtual void UnlockStructuredBuffer_BottomOfPipe(class FRHICommandListImmediate& RHICmdList, FStructuredBufferRHIParamRef StructuredBuffer)
+	virtual void UnlockStructuredBuffer_BottomOfPipe(class FRHICommandListImmediate& RHICmdList, FRHIStructuredBuffer* StructuredBuffer)
 	{
 		// Either this function or RHIUnlockStructuredBuffer must be implemented by the platform RHI.
 		checkNoEntry();
 	}
-	virtual void UnlockVertexBuffer_BottomOfPipe(class FRHICommandListImmediate& RHICmdList, FVertexBufferRHIParamRef VertexBuffer)
+	virtual void UnlockVertexBuffer_BottomOfPipe(class FRHICommandListImmediate& RHICmdList, FRHIVertexBuffer* VertexBuffer)
 	{
 		// Either this function or RHIUnlockVertexBuffer must be implemented by the platform RHI.
 		checkNoEntry();
 	}
-	virtual void UnlockIndexBuffer_BottomOfPipe(class FRHICommandListImmediate& RHICmdList, FIndexBufferRHIParamRef IndexBuffer)
+	virtual void UnlockIndexBuffer_BottomOfPipe(class FRHICommandListImmediate& RHICmdList, FRHIIndexBuffer* IndexBuffer)
 	{
 		// Either this function or RHIUnlockIndexBuffer must be implemented by the platform RHI.
 		checkNoEntry();
diff --git a/Engine/Source/Runtime/SlateRHIRenderer/Private/SlateRHIRenderingPolicy.cpp b/Engine/Source/Runtime/SlateRHIRenderer/Private/SlateRHIRenderingPolicy.cpp
--- a/Engine/Source/Runtime/SlateRHIRenderer/Private/SlateRHIRenderingPolicy.cpp
+++ b/Engine/Source/Runtime/SlateRHIRenderer/Private/SlateRHIRenderingPolicy.cpp
@@ -191,8 +191,8 @@ void FSlateRHIRenderingPolicy::UpdateVertexAndIndexBuffers(FRHICommandListImmedi
 		}
 		else
 		{
-			FVertexBufferRHIParamRef VertexBufferRHI = VertexBuffer.VertexBufferRHI;
-			FIndexBufferRHIParamRef IndexBufferRHI = IndexBuffer.IndexBufferRHI;
+			FRHIVertexBuffer* VertexBufferRHI = VertexBuffer.VertexBufferRHI;
+			FRHIIndexBuffer* IndexBufferRHI = IndexBuffer.IndexBufferRHI;
 
 			RHICmdList.EnqueueLambda([VertexBufferRHI, IndexBufferRHI, &InBatchData, bAbsoluteIndices](FRHICommandListImmediate& RHICmdList)
 			{
diff --git a/Engine/Source/Runtime/SlateRHIRenderer/Private/SlateUpdatableBuffer.cpp b/Engine/Source/Runtime/SlateRHIRenderer/Private/SlateUpdatableBuffer.cpp
--- a/Engine/Source/Runtime/SlateRHIRenderer/Private/SlateUpdatableBuffer.cpp
+++ b/Engine/Source/Runtime/SlateRHIRenderer/Private/SlateUpdatableBuffer.cpp
@@ -77,7 +77,7 @@ void FSlateUpdatableInstanceBuffer::UpdateRenderingData_RenderThread(FRHICommand
 	}
 	else
 	{
-		FVertexBufferRHIParamRef VertexBufferRHI = InstanceBufferResource.VertexBufferRHI;
+		FRHIVertexBuffer* VertexBufferRHI = InstanceBufferResource.VertexBufferRHI;
 		RHICmdList.EnqueueLambda([VertexBufferRHI, &InstanceData = RenderThreadBufferData](FRHICommandListImmediate& RHICmdList)
 		{
 			SCOPE_CYCLE_COUNTER(STAT_SlateUpdateInstanceBuffer);
diff --git a/Engine/Source/Runtime/VulkanRHI/Public/VulkanDynamicRHI.h b/Engine/Source/Runtime/VulkanRHI/Public/VulkanDynamicRHI.h
--- a/Engine/Source/Runtime/VulkanRHI/Public/VulkanDynamicRHI.h
+++ b/Engine/Source/Runtime/VulkanRHI/Public/VulkanDynamicRHI.h
@@ -60,26 +60,26 @@ public:
 	virtual FUniformBufferRHIRef RHICreateUniformBuffer(const void* Contents, const FRHIUniformBufferLayout& Layout, EUniformBufferUsage Usage, EUniformBufferValidation Validation) final override;
 	virtual void RHIUpdateUniformBuffer(FRHIUniformBuffer* UniformBufferRHI, const void* Contents) final override;
 	virtual FIndexBufferRHIRef RHICreateIndexBuffer(uint32 Stride, uint32 Size, uint32 InUsage, FRHIResourceCreateInfo& CreateInfo) final override;
-	virtual void* LockIndexBuffer_BottomOfPipe(FRHICommandListImmediate& RHICmdList, FIndexBufferRHIParamRef IndexBuffer, uint32 Offset, uint32 Size, EResourceLockMode LockMode) final override;
-	virtual void UnlockIndexBuffer_BottomOfPipe(FRHICommandListImmediate& RHICmdList, FIndexBufferRHIParamRef IndexBuffer) final override;
-	virtual void RHITransferIndexBufferUnderlyingResource(FIndexBufferRHIParamRef DestIndexBuffer, FIndexBufferRHIParamRef SrcIndexBuffer) final override;
+	virtual void* LockIndexBuffer_BottomOfPipe(FRHICommandListImmediate& RHICmdList, FRHIIndexBuffer* IndexBuffer, uint32 Offset, uint32 Size, EResourceLockMode LockMode) final override;
+	virtual void UnlockIndexBuffer_BottomOfPipe(FRHICommandListImmediate& RHICmdList, FRHIIndexBuffer* IndexBuffer) final override;
+	virtual void RHITransferIndexBufferUnderlyingResource(FRHIIndexBuffer* DestIndexBuffer, FRHIIndexBuffer* SrcIndexBuffer) final override;
 	virtual FVertexBufferRHIRef RHICreateVertexBuffer(uint32 Size, uint32 InUsage, FRHIResourceCreateInfo& CreateInfo) final override;
-	virtual void* LockVertexBuffer_BottomOfPipe(FRHICommandListImmediate& RHICmdList, FVertexBufferRHIParamRef VertexBuffer, uint32 Offset, uint32 SizeRHI, EResourceLockMode LockMode) final override;
-	virtual void UnlockVertexBuffer_BottomOfPipe(FRHICommandListImmediate& RHICmdList, FVertexBufferRHIParamRef VertexBuffer) final override;
-	virtual void RHICopyVertexBuffer(FVertexBufferRHIParamRef SourceBuffer, FVertexBufferRHIParamRef DestBuffer) final override;
-	virtual void RHITransferVertexBufferUnderlyingResource(FVertexBufferRHIParamRef DestVertexBuffer, FVertexBufferRHIParamRef SrcVertexBuffer) final override;
+	virtual void* LockVertexBuffer_BottomOfPipe(FRHICommandListImmediate& RHICmdList, FRHIVertexBuffer* VertexBuffer, uint32 Offset, uint32 SizeRHI, EResourceLockMode LockMode) final override;
+	virtual void UnlockVertexBuffer_BottomOfPipe(FRHICommandListImmediate& RHICmdList, FRHIVertexBuffer* VertexBuffer) final override;
+	virtual void RHICopyVertexBuffer(FRHIVertexBuffer* SourceBuffer, FRHIVertexBuffer* DestBuffer) final override;
+	virtual void RHITransferVertexBufferUnderlyingResource(FRHIVertexBuffer* DestVertexBuffer, FRHIVertexBuffer* SrcVertexBuffer) final override;
 	virtual FStructuredBufferRHIRef RHICreateStructuredBuffer(uint32 Stride, uint32 Size, uint32 InUsage, FRHIResourceCreateInfo& CreateInfo) final override;
-	virtual void* LockStructuredBuffer_BottomOfPipe(FRHICommandListImmediate& RHICmdList, FStructuredBufferRHIParamRef StructuredBuffer, uint32 Offset, uint32 SizeRHI, EResourceLockMode LockMode) final override;
-	virtual void UnlockStructuredBuffer_BottomOfPipe(FRHICommandListImmediate& RHICmdList, FStructuredBufferRHIParamRef StructuredBuffer) final override;
-	virtual FUnorderedAccessViewRHIRef RHICreateUnorderedAccessView(FStructuredBufferRHIParamRef StructuredBuffer, bool bUseUAVCounter, bool bAppendBuffer) final override;
-	virtual FUnorderedAccessViewRHIRef RHICreateUnorderedAccessView(FTextureRHIParamRef Texture, uint32 MipLevel) final override;
-	virtual FUnorderedAccessViewRHIRef RHICreateUnorderedAccessView(FVertexBufferRHIParamRef VertexBuffer, uint8 Format) final override;
-	virtual FUnorderedAccessViewRHIRef RHICreateUnorderedAccessView(FIndexBufferRHIParamRef IndexBuffer, uint8 Format) final override;
-	virtual FShaderResourceViewRHIRef RHICreateShaderResourceView(FStructuredBufferRHIParamRef StructuredBuffer) final override;
-	virtual FShaderResourceViewRHIRef RHICreateShaderResourceView(FVertexBufferRHIParamRef VertexBuffer, uint32 Stride, uint8 Format) final override;
-	virtual FShaderResourceViewRHIRef RHICreateShaderResourceView(FIndexBufferRHIParamRef Buffer) final override;
-	virtual void RHIUpdateShaderResourceView(FRHIShaderResourceView* SRV, FVertexBufferRHIParamRef VertexBuffer, uint32 Stride, uint8 Format) final override;
-	virtual void RHIUpdateShaderResourceView(FRHIShaderResourceView* SRV, FIndexBufferRHIParamRef IndexBuffer) final override;
+	virtual void* LockStructuredBuffer_BottomOfPipe(FRHICommandListImmediate& RHICmdList, FRHIStructuredBuffer* StructuredBuffer, uint32 Offset, uint32 SizeRHI, EResourceLockMode LockMode) final override;
+	virtual void UnlockStructuredBuffer_BottomOfPipe(FRHICommandListImmediate& RHICmdList, FRHIStructuredBuffer* StructuredBuffer) final override;
+	virtual FUnorderedAccessViewRHIRef RHICreateUnorderedAccessView(FRHIStructuredBuffer* StructuredBuffer, bool bUseUAVCounter, bool bAppendBuffer) final override;
+	virtual FUnorderedAccessViewRHIRef RHICreateUnorderedAccessView(FRHITexture* Texture, uint32 MipLevel) final override;
+	virtual FUnorderedAccessViewRHIRef RHICreateUnorderedAccessView(FRHIVertexBuffer* VertexBuffer, uint8 Format) final override;
+	virtual FUnorderedAccessViewRHIRef RHICreateUnorderedAccessView(FRHIIndexBuffer* IndexBuffer, uint8 Format) final override;
+	virtual FShaderResourceViewRHIRef RHICreateShaderResourceView(FRHIStructuredBuffer* StructuredBuffer) final override;
+	virtual FShaderResourceViewRHIRef RHICreateShaderResourceView(FRHIVertexBuffer* VertexBuffer, uint32 Stride, uint8 Format) final override;
+	virtual FShaderResourceViewRHIRef RHICreateShaderResourceView(FRHIIndexBuffer* Buffer) final override;
+	virtual void RHIUpdateShaderResourceView(FRHIShaderResourceView* SRV, FRHIVertexBuffer* VertexBuffer, uint32 Stride, uint8 Format) final override;
+	virtual void RHIUpdateShaderResourceView(FRHIShaderResourceView* SRV, FRHIIndexBuffer* IndexBuffer) final override;
 	virtual uint64 RHICalcTexture2DPlatformSize(uint32 SizeX, uint32 SizeY, uint8 Format, uint32 NumMips, uint32 NumSamples, uint32 Flags, uint32& OutAlign) final override;
 	virtual uint64 RHICalcTexture3DPlatformSize(uint32 SizeX, uint32 SizeY, uint32 SizeZ, uint8 Format, uint32 NumMips, uint32 Flags, uint32& OutAlign) final override;
 	virtual uint64 RHICalcTextureCubePlatformSize(uint32 Size, uint8 Format, uint32 NumMips, uint32 Flags, uint32& OutAlign) final override;
diff --git a/Engine/Source/Runtime/Windows/D3D11RHI/Private/D3D11RHIPrivate.h b/Engine/Source/Runtime/Windows/D3D11RHI/Private/D3D11RHIPrivate.h
--- a/Engine/Source/Runtime/Windows/D3D11RHI/Private/D3D11RHIPrivate.h
+++ b/Engine/Source/Runtime/Windows/D3D11RHI/Private/D3D11RHIPrivate.h
@@ -416,26 +416,26 @@ public:
 	virtual FUniformBufferRHIRef RHICreateUniformBuffer(const void* Contents, const FRHIUniformBufferLayout& Layout, EUniformBufferUsage Usage, EUniformBufferValidation Validation) final override;
 	virtual void RHIUpdateUniformBuffer(FRHIUniformBuffer* UniformBufferRHI, const void* Contents) final override;
 	virtual FIndexBufferRHIRef RHICreateIndexBuffer(uint32 Stride, uint32 Size, uint32 InUsage, FRHIResourceCreateInfo& CreateInfo) final override;
-	virtual void* LockIndexBuffer_BottomOfPipe(FRHICommandListImmediate& RHICmdList, FIndexBufferRHIParamRef IndexBuffer, uint32 Offset, uint32 Size, EResourceLockMode LockMode) final override;
-	virtual void UnlockIndexBuffer_BottomOfPipe(FRHICommandListImmediate& RHICmdList, FIndexBufferRHIParamRef IndexBuffer) final override;
-	virtual void RHITransferIndexBufferUnderlyingResource(FIndexBufferRHIParamRef DestIndexBuffer, FIndexBufferRHIParamRef SrcIndexBuffer) final override;
+	virtual void* LockIndexBuffer_BottomOfPipe(FRHICommandListImmediate& RHICmdList, FRHIIndexBuffer* IndexBuffer, uint32 Offset, uint32 Size, EResourceLockMode LockMode) final override;
+	virtual void UnlockIndexBuffer_BottomOfPipe(FRHICommandListImmediate& RHICmdList, FRHIIndexBuffer* IndexBuffer) final override;
+	virtual void RHITransferIndexBufferUnderlyingResource(FRHIIndexBuffer* DestIndexBuffer, FRHIIndexBuffer* SrcIndexBuffer) final override;
 	virtual FVertexBufferRHIRef RHICreateVertexBuffer(uint32 Size, uint32 InUsage, FRHIResourceCreateInfo& CreateInfo) final override;
-	virtual void* LockVertexBuffer_BottomOfPipe(FRHICommandListImmediate& RHICmdList, FVertexBufferRHIParamRef VertexBuffer, uint32 Offset, uint32 SizeRHI, EResourceLockMode LockMode) final override;
-	virtual void UnlockVertexBuffer_BottomOfPipe(FRHICommandListImmediate& RHICmdList, FVertexBufferRHIParamRef VertexBuffer) final override;
-	virtual void RHICopyVertexBuffer(FVertexBufferRHIParamRef SourceBuffer,FVertexBufferRHIParamRef DestBuffer) final override;
-	virtual void RHITransferVertexBufferUnderlyingResource(FVertexBufferRHIParamRef DestVertexBuffer, FVertexBufferRHIParamRef SrcVertexBuffer) final override;
+	virtual void* LockVertexBuffer_BottomOfPipe(FRHICommandListImmediate& RHICmdList, FRHIVertexBuffer* VertexBuffer, uint32 Offset, uint32 SizeRHI, EResourceLockMode LockMode) final override;
+	virtual void UnlockVertexBuffer_BottomOfPipe(FRHICommandListImmediate& RHICmdList, FRHIVertexBuffer* VertexBuffer) final override;
+	virtual void RHICopyVertexBuffer(FRHIVertexBuffer* SourceBuffer, FRHIVertexBuffer* DestBuffer) final override;
+	virtual void RHITransferVertexBufferUnderlyingResource(FRHIVertexBuffer* DestVertexBuffer, FRHIVertexBuffer* SrcVertexBuffer) final override;
 	virtual FStructuredBufferRHIRef RHICreateStructuredBuffer(uint32 Stride, uint32 Size, uint32 InUsage, FRHIResourceCreateInfo& CreateInfo) final override;
-	virtual void* LockStructuredBuffer_BottomOfPipe(FRHICommandListImmediate& RHICmdList, FStructuredBufferRHIParamRef StructuredBuffer, uint32 Offset, uint32 SizeRHI, EResourceLockMode LockMode) final override;
-	virtual void UnlockStructuredBuffer_BottomOfPipe(FRHICommandListImmediate& RHICmdList, FStructuredBufferRHIParamRef StructuredBuffer) final override;
-	virtual FUnorderedAccessViewRHIRef RHICreateUnorderedAccessView(FStructuredBufferRHIParamRef StructuredBuffer, bool bUseUAVCounter, bool bAppendBuffer) final override;
-	virtual FUnorderedAccessViewRHIRef RHICreateUnorderedAccessView(FTextureRHIParamRef Texture, uint32 MipLevel) final override;
-	virtual FUnorderedAccessViewRHIRef RHICreateUnorderedAccessView(FVertexBufferRHIParamRef VertexBuffer, uint8 Format) final override;
-	virtual FUnorderedAccessViewRHIRef RHICreateUnorderedAccessView(FIndexBufferRHIParamRef IndexBuffer, uint8 Format) final override;
-	virtual FShaderResourceViewRHIRef RHICreateShaderResourceView(FStructuredBufferRHIParamRef StructuredBuffer) final override;
-	virtual FShaderResourceViewRHIRef RHICreateShaderResourceView(FVertexBufferRHIParamRef VertexBuffer, uint32 Stride, uint8 Format) final override;
-	virtual FShaderResourceViewRHIRef RHICreateShaderResourceView(FIndexBufferRHIParamRef Buffer) final override;
-	virtual void RHIUpdateShaderResourceView(FRHIShaderResourceView* SRV, FVertexBufferRHIParamRef VertexBuffer, uint32 Stride, uint8 Format) final override;
-	virtual void RHIUpdateShaderResourceView(FRHIShaderResourceView* SRV, FIndexBufferRHIParamRef IndexBuffer) final override;
+	virtual void* LockStructuredBuffer_BottomOfPipe(FRHICommandListImmediate& RHICmdList, FRHIStructuredBuffer* StructuredBuffer, uint32 Offset, uint32 SizeRHI, EResourceLockMode LockMode) final override;
+	virtual void UnlockStructuredBuffer_BottomOfPipe(FRHICommandListImmediate& RHICmdList, FRHIStructuredBuffer* StructuredBuffer) final override;
+	virtual FUnorderedAccessViewRHIRef RHICreateUnorderedAccessView(FRHIStructuredBuffer* StructuredBuffer, bool bUseUAVCounter, bool bAppendBuffer) final override;
+	virtual FUnorderedAccessViewRHIRef RHICreateUnorderedAccessView(FRHITexture* Texture, uint32 MipLevel) final override;
+	virtual FUnorderedAccessViewRHIRef RHICreateUnorderedAccessView(FRHIVertexBuffer* VertexBuffer, uint8 Format) final override;
+	virtual FUnorderedAccessViewRHIRef RHICreateUnorderedAccessView(FRHIIndexBuffer* IndexBuffer, uint8 Format) final override;
+	virtual FShaderResourceViewRHIRef RHICreateShaderResourceView(FRHIStructuredBuffer* StructuredBuffer) final override;
+	virtual FShaderResourceViewRHIRef RHICreateShaderResourceView(FRHIVertexBuffer* VertexBuffer, uint32 Stride, uint8 Format) final override;
+	virtual FShaderResourceViewRHIRef RHICreateShaderResourceView(FRHIIndexBuffer* Buffer) final override;
+	virtual void RHIUpdateShaderResourceView(FRHIShaderResourceView* SRV, FRHIVertexBuffer* VertexBuffer, uint32 Stride, uint8 Format) final override;
+	virtual void RHIUpdateShaderResourceView(FRHIShaderResourceView* SRV, FRHIIndexBuffer* IndexBuffer) final override;
 	virtual uint64 RHICalcTexture2DPlatformSize(uint32 SizeX, uint32 SizeY, uint8 Format, uint32 NumMips, uint32 NumSamples, uint32 Flags, uint32& OutAlign) final override;
 	virtual uint64 RHICalcTexture3DPlatformSize(uint32 SizeX, uint32 SizeY, uint32 SizeZ, uint8 Format, uint32 NumMips, uint32 Flags, uint32& OutAlign) final override;
 	virtual uint64 RHICalcTextureCubePlatformSize(uint32 Size, uint8 Format, uint32 NumMips, uint32 Flags, uint32& OutAlign) final override;
