Fixes compilation failure due to deletion of the cast of FRHIGPUMask to uint32

diff --git a/Engine/Source/Runtime/D3D12RHI/Private/D3D12Allocation.cpp b/Engine/Source/Runtime/D3D12RHI/Private/D3D12Allocation.cpp
--- a/Engine/Source/Runtime/D3D12RHI/Private/D3D12Allocation.cpp
+++ b/Engine/Source/Runtime/D3D12RHI/Private/D3D12Allocation.cpp
@@ -1233,7 +1233,7 @@ HRESULT FD3D12TextureAllocator::AllocateTexture(D3D12_RESOURCE_DESC Desc, const
 
 	// Request default alignment for stand alone textures
 	Desc.Alignment = 0;
-	const D3D12_HEAP_PROPERTIES HeapProps = CD3DX12_HEAP_PROPERTIES(D3D12_HEAP_TYPE_DEFAULT, (uint32)GetGPUMask(), (uint32)GetVisibilityMask());
+	const D3D12_HEAP_PROPERTIES HeapProps = CD3DX12_HEAP_PROPERTIES(D3D12_HEAP_TYPE_DEFAULT, GetGPUMask().GetNative(), (uint32)GetVisibilityMask().GetNative());
 
 	hr = Adapter->CreateCommittedResource(Desc, HeapProps, InitialState, ClearValue, &NewResource, Name);
 
@@ -1268,7 +1268,7 @@ HRESULT FD3D12TextureAllocatorPool::AllocateTexture(D3D12_RESOURCE_DESC Desc, co
 	FD3D12Adapter* Adapter = GetParentDevice()->GetParentAdapter();
 	FD3D12Resource* Resource = nullptr;
 
-	const D3D12_HEAP_PROPERTIES HeapProps = CD3DX12_HEAP_PROPERTIES(D3D12_HEAP_TYPE_DEFAULT, (uint32)GetGPUMask(), (uint32)GetVisibilityMask());
+	const D3D12_HEAP_PROPERTIES HeapProps = CD3DX12_HEAP_PROPERTIES(D3D12_HEAP_TYPE_DEFAULT, GetGPUMask().GetNative(), GetVisibilityMask().GetNative());
 	HRESULT hr = Adapter->CreateCommittedResource(Desc, HeapProps, InitialState, ClearValue, &Resource, Name);
 
 	TextureLocation.SetType(FD3D12ResourceLocation::ResourceLocationType::eStandAlone);
