Fix compile error on DX12. I dont have a XB1 so was not able to 100% confirm the fix, but this is very likely fine.

diff --git a/Engine/Source/Runtime/D3D12RHI/Private/D3D12Allocation.cpp b/Engine/Source/Runtime/D3D12RHI/Private/D3D12Allocation.cpp
--- a/Engine/Source/Runtime/D3D12RHI/Private/D3D12Allocation.cpp
+++ b/Engine/Source/Runtime/D3D12RHI/Private/D3D12Allocation.cpp
@@ -1235,7 +1235,7 @@ HRESULT FD3D12TextureAllocator::AllocateTexture(D3D12_RESOURCE_DESC Desc, const
 	Desc.Alignment = 0;
 	const D3D12_HEAP_PROPERTIES HeapProps = CD3DX12_HEAP_PROPERTIES(D3D12_HEAP_TYPE_DEFAULT, GetGPUMask().GetNative(), (uint32)GetVisibilityMask().GetNative());
 
-	hr = Adapter->CreateCommittedResource(Desc, HeapProps, InitialState, ClearValue, &NewResource, Name);
+	hr = Adapter->CreateCommittedResource(Desc, GetGPUMask(), HeapProps, InitialState, ClearValue, &NewResource, Name);
 
 	TextureLocation.SetType(FD3D12ResourceLocation::ResourceLocationType::eStandAlone);
 	TextureLocation.SetResource(NewResource);
@@ -1269,7 +1269,7 @@ HRESULT FD3D12TextureAllocatorPool::AllocateTexture(D3D12_RESOURCE_DESC Desc, co
 	FD3D12Resource* Resource = nullptr;
 
 	const D3D12_HEAP_PROPERTIES HeapProps = CD3DX12_HEAP_PROPERTIES(D3D12_HEAP_TYPE_DEFAULT, GetGPUMask().GetNative(), GetVisibilityMask().GetNative());
-	HRESULT hr = Adapter->CreateCommittedResource(Desc, HeapProps, InitialState, ClearValue, &Resource, Name);
+	HRESULT hr = Adapter->CreateCommittedResource(Desc, GetGPUMask(), HeapProps, InitialState, ClearValue, &Resource, Name);
 
 	TextureLocation.SetType(FD3D12ResourceLocation::ResourceLocationType::eStandAlone);
 	TextureLocation.SetResource(Resource);
