DR - Remove RHISetStreamOutTargets from actual RHIs

diff --git a/Engine/Source/Runtime/Apple/MetalRHI/Private/MetalCommands.cpp b/Engine/Source/Runtime/Apple/MetalRHI/Private/MetalCommands.cpp
--- a/Engine/Source/Runtime/Apple/MetalRHI/Private/MetalCommands.cpp
+++ b/Engine/Source/Runtime/Apple/MetalRHI/Private/MetalCommands.cpp
@@ -96,12 +96,6 @@ void FMetalRHICommandContext::RHISetStreamSource(uint32 StreamIndex, FRHIVertexB
 	}
 }
 
-void FMetalDynamicRHI::RHISetStreamOutTargets(uint32 NumTargets, FRHIVertexBuffer* const * VertexBuffers, const uint32* Offsets)
-{
-	NOT_SUPPORTED("RHISetStreamOutTargets");
-
-}
-
 void FMetalRHICommandContext::RHISetComputeShader(FRHIComputeShader* ComputeShaderRHI)
 {
 	@autoreleasepool {
diff --git a/Engine/Source/Runtime/Apple/MetalRHI/Public/MetalRHI.h b/Engine/Source/Runtime/Apple/MetalRHI/Public/MetalRHI.h
--- a/Engine/Source/Runtime/Apple/MetalRHI/Public/MetalRHI.h
+++ b/Engine/Source/Runtime/Apple/MetalRHI/Public/MetalRHI.h
@@ -130,7 +130,6 @@ public:
 	virtual void RHIResizeViewport(FRHIViewport* Viewport, uint32 SizeX, uint32 SizeY, bool bIsFullscreen) final override;
 	virtual void RHIResizeViewport(FRHIViewport* Viewport, uint32 SizeX, uint32 SizeY, bool bIsFullscreen, EPixelFormat PreferredPixelFormat) final override;
 	virtual void RHITick(float DeltaTime) final override;
-	virtual void RHISetStreamOutTargets(uint32 NumTargets,FRHIVertexBuffer* const * VertexBuffers,const uint32* Offsets) final override;
 	virtual void RHIBlockUntilGPUIdle() final override;
 	virtual bool RHIGetAvailableResolutions(FScreenResolutionArray& Resolutions, bool bIgnoreRefreshRate) final override;
 	virtual void RHIGetSupportedResolution(uint32& Width, uint32& Height) final override;
diff --git a/Engine/Source/Runtime/D3D12RHI/Private/D3D12Commands.cpp b/Engine/Source/Runtime/D3D12RHI/Private/D3D12Commands.cpp
--- a/Engine/Source/Runtime/D3D12RHI/Private/D3D12Commands.cpp
+++ b/Engine/Source/Runtime/D3D12RHI/Private/D3D12Commands.cpp
@@ -93,35 +93,6 @@ void FD3D12CommandContext::RHISetStreamSource(uint32 StreamIndex, FRHIVertexBuff
 	StateCache.SetStreamSource(VertexBuffer ? &VertexBuffer->ResourceLocation : nullptr, StreamIndex, Offset);
 }
 
-// Stream-Out state.
-void FD3D12DynamicRHI::RHISetStreamOutTargets(uint32 NumTargets, FRHIVertexBuffer* const* VertexBuffers, const uint32* Offsets)
-{
-	// Multi-GPU support: this might need a node mask parameter or broadcast to all GPUs.
-	FD3D12CommandContext& CmdContext = GetRHIDevice()->GetDefaultCommandContext();
-	FD3D12Resource* D3DVertexBuffers[D3D12_SO_BUFFER_SLOT_COUNT] = { 0 };
-	uint32 D3DOffsets[D3D12_SO_BUFFER_SLOT_COUNT] = { 0 };
-
-	if (VertexBuffers)
-	{
-		for (uint32 BufferIndex = 0; BufferIndex < NumTargets; BufferIndex++)
-		{
-			const FD3D12VertexBuffer* VB = ((FD3D12VertexBuffer*)VertexBuffers[BufferIndex]);
-			if (VB != nullptr)
-			{
-				D3DVertexBuffers[BufferIndex] = VB->ResourceLocation.GetResource();
-				D3DOffsets[BufferIndex] = VB->ResourceLocation.GetOffsetFromBaseOfResource();
-			}
-			else
-			{
-				D3DVertexBuffers[BufferIndex] = nullptr;
-				D3DOffsets[BufferIndex] = 0;
-			}
-		}
-	}
-
-	CmdContext.StateCache.SetStreamOutTargets(NumTargets, D3DVertexBuffers, D3DOffsets);
-}
-
 void FD3D12CommandContext::RHISetComputeShader(FRHIComputeShader* ComputeShaderRHI)
 {
 #if D3D12_RHI_RAYTRACING
diff --git a/Engine/Source/Runtime/D3D12RHI/Private/D3D12PipelineState.h b/Engine/Source/Runtime/D3D12RHI/Private/D3D12PipelineState.h
--- a/Engine/Source/Runtime/D3D12RHI/Private/D3D12PipelineState.h
+++ b/Engine/Source/Runtime/D3D12RHI/Private/D3D12PipelineState.h
@@ -48,8 +48,6 @@ struct FD3D12_GRAPHICS_PIPELINE_STATE_DESC
 	D3D12_SHADER_BYTECODE DS;
 	D3D12_SHADER_BYTECODE HS;
 	D3D12_SHADER_BYTECODE GS;
-	//#todo-RemoveStreamOut
-	D3D12_STREAM_OUTPUT_DESC StreamOutput;
 #if !D3D12_USE_DERIVED_PSO
 	D3D12_BLEND_DESC BlendState;
 	uint32 SampleMask;
@@ -160,9 +158,6 @@ template <> struct equality_pipeline_state_desc<FD3D12LowLevelGraphicsPipelineSt
 #endif
 		PSO_IF_NOT_EQUAL_RETURN_FALSE(Desc.IBStripCutValue)
 		PSO_IF_NOT_EQUAL_RETURN_FALSE(Desc.NodeMask)
-		PSO_IF_NOT_EQUAL_RETURN_FALSE(Desc.StreamOutput.RasterizedStream)
-		PSO_IF_NOT_EQUAL_RETURN_FALSE(Desc.StreamOutput.NumEntries)
-		PSO_IF_NOT_EQUAL_RETURN_FALSE(Desc.StreamOutput.NumStrides)
 		PSO_IF_NOT_EQUAL_RETURN_FALSE(Desc.SampleDesc.Count)
 		PSO_IF_NOT_EQUAL_RETURN_FALSE(Desc.SampleDesc.Quality)
 
@@ -180,29 +175,6 @@ template <> struct equality_pipeline_state_desc<FD3D12LowLevelGraphicsPipelineSt
 		PSO_IF_NOT_EQUAL_RETURN_FALSE(HSHash)
 		PSO_IF_NOT_EQUAL_RETURN_FALSE(DSHash)
 
-		if (lhs.Desc.StreamOutput.pSODeclaration != rhs.Desc.StreamOutput.pSODeclaration &&
-			lhs.Desc.StreamOutput.NumEntries)
-		{
-			for (uint32 i = 0; i < lhs.Desc.StreamOutput.NumEntries; i++)
-			{
-				PSO_IF_NOT_EQUAL_RETURN_FALSE(Desc.StreamOutput.pSODeclaration[i].Stream)
-				PSO_IF_NOT_EQUAL_RETURN_FALSE(Desc.StreamOutput.pSODeclaration[i].SemanticIndex)
-				PSO_IF_NOT_EQUAL_RETURN_FALSE(Desc.StreamOutput.pSODeclaration[i].StartComponent)
-				PSO_IF_NOT_EQUAL_RETURN_FALSE(Desc.StreamOutput.pSODeclaration[i].ComponentCount)
-				PSO_IF_NOT_EQUAL_RETURN_FALSE(Desc.StreamOutput.pSODeclaration[i].OutputSlot)
-				PSO_IF_STRING_COMPARE_FAILS_RETURN_FALSE(Desc.StreamOutput.pSODeclaration[i].SemanticName)
-			}
-		}
-
-		if (lhs.Desc.StreamOutput.pBufferStrides != rhs.Desc.StreamOutput.pBufferStrides &&
-			lhs.Desc.StreamOutput.NumStrides)
-		{
-			for (uint32 i = 0; i < lhs.Desc.StreamOutput.NumStrides; i++)
-			{
-				PSO_IF_NOT_EQUAL_RETURN_FALSE(Desc.StreamOutput.pBufferStrides[i])
-			}
-		}
-
 		if (lhs.Desc.InputLayout.pInputElementDescs != rhs.Desc.InputLayout.pInputElementDescs &&
 			lhs.Desc.InputLayout.NumElements)
 		{
diff --git a/Engine/Source/Runtime/D3D12RHI/Private/D3D12Pipelinestate.cpp b/Engine/Source/Runtime/D3D12RHI/Private/D3D12Pipelinestate.cpp
--- a/Engine/Source/Runtime/D3D12RHI/Private/D3D12Pipelinestate.cpp
+++ b/Engine/Source/Runtime/D3D12RHI/Private/D3D12Pipelinestate.cpp
@@ -60,12 +60,6 @@ FD3D12LowLevelGraphicsPipelineStateDesc GetLowLevelGraphicsPipelineStateDesc(con
 		Desc.Desc.InputLayout.pInputElementDescs = InputLayout->VertexElements.GetData();
 	}
 
-	if (Initializer.BoundShaderState.GeometryShaderRHI)
-	{
-		//#todo-RemoveStreamOut
-		Desc.Desc.StreamOutput = ((FD3D12GeometryShader*) Initializer.BoundShaderState.GeometryShaderRHI)->StreamOutput;
-	}
-
 #define COPY_SHADER(Initial, Name) \
 	if (FD3D12##Name##Shader* Shader = (FD3D12##Name##Shader*) Initializer.BoundShaderState.##Name##ShaderRHI) \
 	{ \
@@ -188,8 +182,6 @@ SIZE_T FD3D12PipelineStateCacheBase::HashPSODesc(const FD3D12LowLevelGraphicsPip
 	Hash.Desc.DS.pShaderBytecode = nullptr;
 	Hash.Desc.GS.pShaderBytecode = nullptr;
 	Hash.Desc.InputLayout.pInputElementDescs = nullptr;
-	Hash.Desc.StreamOutput.pBufferStrides = nullptr;
-	Hash.Desc.StreamOutput.pSODeclaration = nullptr;
 	Hash.Desc.CachedPSO.pCachedBlob = nullptr;
 	Hash.Desc.CachedPSO.CachedBlobSizeInBytes = 0;
 	Hash.CombinedHash = 0;
diff --git a/Engine/Source/Runtime/D3D12RHI/Private/D3D12RHIPrivate.h b/Engine/Source/Runtime/D3D12RHI/Private/D3D12RHIPrivate.h
--- a/Engine/Source/Runtime/D3D12RHI/Private/D3D12RHIPrivate.h
+++ b/Engine/Source/Runtime/D3D12RHI/Private/D3D12RHIPrivate.h
@@ -355,7 +355,6 @@ public:
 	virtual void RHIResizeViewport(FRHIViewport* Viewport, uint32 SizeX, uint32 SizeY, bool bIsFullscreen) final override;
 	virtual void RHIResizeViewport(FRHIViewport* ViewportRHI, uint32 SizeX, uint32 SizeY, bool bIsFullscreen, EPixelFormat PreferredPixelFormat) final override;
 	virtual void RHITick(float DeltaTime) final override;
-	virtual void RHISetStreamOutTargets(uint32 NumTargets, FRHIVertexBuffer* const* VertexBuffers, const uint32* Offsets) final override;
 	virtual void RHIBlockUntilGPUIdle() final override;
 	virtual void RHISubmitCommandsAndFlushGPU() final override;
 	virtual bool RHIGetAvailableResolutions(FScreenResolutionArray& Resolutions, bool bIgnoreRefreshRate) final override;
diff --git a/Engine/Source/Runtime/D3D12RHI/Private/D3D12Shader.h b/Engine/Source/Runtime/D3D12RHI/Private/D3D12Shader.h
--- a/Engine/Source/Runtime/D3D12RHI/Private/D3D12Shader.h
+++ b/Engine/Source/Runtime/D3D12RHI/Private/D3D12Shader.h
@@ -67,41 +67,7 @@ public:
 	/** The shader's bytecode, with custom data in the last byte. */
 	TArray<uint8> Code;
 
-	/** The shader's stream output description. */
-	//#todo-RemoveStreamOut
-	D3D12_STREAM_OUTPUT_DESC StreamOutput;
-	//#todo-RemoveStreamOut
-	D3D12_SO_DECLARATION_ENTRY* pStreamOutEntries;
-	//#todo-RemoveStreamOut
-	uint32* pStreamOutStrides;
-
-	//#todo-RemoveStreamOut
-	bool bShaderNeedsStreamOutput;
-
 	FShaderCodePackedResourceCounts ResourceCounts;
-
-	FD3D12GeometryShader()
-		: pStreamOutEntries(nullptr)
-		, pStreamOutStrides(nullptr)
-		, bShaderNeedsStreamOutput(false)
-	{
-		FMemory::Memzero(&StreamOutput, sizeof(StreamOutput));
-	}
-
-	virtual ~FD3D12GeometryShader()
-	{
-		if (pStreamOutEntries)
-		{
-			delete[] pStreamOutEntries;
-			pStreamOutEntries = nullptr;
-		}
-
-		if (pStreamOutStrides)
-		{
-			delete[] pStreamOutStrides;
-			pStreamOutStrides = nullptr;
-		}
-	}
 };
 
 class FD3D12HullShader : public FRHIHullShader
diff --git a/Engine/Source/Runtime/D3D12RHI/Private/D3D12StateCache.cpp b/Engine/Source/Runtime/D3D12RHI/Private/D3D12StateCache.cpp
--- a/Engine/Source/Runtime/D3D12RHI/Private/D3D12StateCache.cpp
+++ b/Engine/Source/Runtime/D3D12RHI/Private/D3D12StateCache.cpp
@@ -1239,15 +1239,3 @@ void FD3D12StateCacheBase::SetRenderTargets(uint32 NumSimultaneousRenderTargets,
 	}
 	PipelineState.Graphics.CurrentNumberOfRenderTargets = ActiveNumSimultaneousRenderTargets;
 }
-
-void FD3D12StateCacheBase::SetStreamOutTargets(uint32 NumSimultaneousStreamOutTargets, FD3D12Resource** SOArray, const uint32* SOOffsets)
-{
-	PipelineState.Graphics.CurrentNumberOfStreamOutTargets = NumSimultaneousStreamOutTargets;
-	if (PipelineState.Graphics.CurrentNumberOfStreamOutTargets > 0)
-	{
-		FMemory::Memcpy(PipelineState.Graphics.CurrentStreamOutTargets, SOArray, sizeof(FD3D12Resource*)* NumSimultaneousStreamOutTargets);
-		FMemory::Memcpy(PipelineState.Graphics.CurrentSOOffsets, SOOffsets, sizeof(uint32*)* NumSimultaneousStreamOutTargets);
-
-		bNeedSetSOs = true;
-	}
-}
diff --git a/Engine/Source/Runtime/D3D12RHI/Private/D3D12StateCachePrivate.h b/Engine/Source/Runtime/D3D12RHI/Private/D3D12StateCachePrivate.h
--- a/Engine/Source/Runtime/D3D12RHI/Private/D3D12StateCachePrivate.h
+++ b/Engine/Source/Runtime/D3D12RHI/Private/D3D12StateCachePrivate.h
@@ -845,8 +845,6 @@ public:
 		}
 	}
 
-	void SetStreamOutTargets(uint32 NumSimultaneousStreamOutTargets, FD3D12Resource** SOArray, const uint32* SOOffsets);
-
 	template <EShaderFrequency ShaderStage>
 	void SetUAVs(uint32 UAVStartSlot, uint32 NumSimultaneousUAVs, FD3D12UnorderedAccessView** UAVArray, uint32* UAVInitialCountArray);
 
diff --git a/Engine/Source/Runtime/D3D12RHI/Private/Windows/WindowsD3D12PipelineState.cpp b/Engine/Source/Runtime/D3D12RHI/Private/Windows/WindowsD3D12PipelineState.cpp
--- a/Engine/Source/Runtime/D3D12RHI/Private/Windows/WindowsD3D12PipelineState.cpp
+++ b/Engine/Source/Runtime/D3D12RHI/Private/Windows/WindowsD3D12PipelineState.cpp
@@ -51,7 +51,6 @@ FD3D12_GRAPHICS_PIPELINE_STATE_STREAM FD3D12_GRAPHICS_PIPELINE_STATE_DESC::Pipel
 	Stream.PrimitiveTopologyType = this->PrimitiveTopologyType;
 	Stream.VS = this->VS;
 	Stream.GS = this->GS;
-	Stream.StreamOutput = this->StreamOutput;
 	Stream.HS = this->HS;
 	Stream.DS = this->DS;
 	Stream.PS = this->PS;
@@ -77,7 +76,6 @@ D3D12_GRAPHICS_PIPELINE_STATE_DESC FD3D12_GRAPHICS_PIPELINE_STATE_DESC::Graphics
 	D.PrimitiveTopologyType = this->PrimitiveTopologyType;
 	D.VS = this->VS;
 	D.GS = this->GS;
-	D.StreamOutput = this->StreamOutput;
 	D.HS = this->HS;
 	D.DS = this->DS;
 	D.PS = this->PS;
@@ -192,21 +190,6 @@ void FD3D12PipelineStateCache::RebuildFromDiskCache(ID3D12RootSignature* Graphic
 				DiskCaches[PSO_CACHE_GRAPHICS].SetPointerAndAdvanceFilePosition((void**)&PSODesc->InputLayout.pInputElementDescs[j].SemanticName, *stringLength, true);
 			}
 		}
-		if (PSODesc->StreamOutput.NumEntries)
-		{
-			DiskCaches[PSO_CACHE_GRAPHICS].SetPointerAndAdvanceFilePosition((void**)&PSODesc->StreamOutput.pSODeclaration, PSODesc->StreamOutput.NumEntries * sizeof(D3D12_SO_DECLARATION_ENTRY), true);
-			for (uint32 j = 0; j < PSODesc->StreamOutput.NumEntries; j++)
-			{
-				//Get the Sematic name string
-				uint32* stringLength = nullptr;
-				DiskCaches[PSO_CACHE_GRAPHICS].SetPointerAndAdvanceFilePosition((void**)&stringLength, sizeof(uint32));
-				DiskCaches[PSO_CACHE_GRAPHICS].SetPointerAndAdvanceFilePosition((void**)&PSODesc->StreamOutput.pSODeclaration[j].SemanticName, *stringLength, true);
-			}
-		}
-		if (PSODesc->StreamOutput.NumStrides)
-		{
-			DiskCaches[PSO_CACHE_GRAPHICS].SetPointerAndAdvanceFilePosition((void**)&PSODesc->StreamOutput.pBufferStrides, PSODesc->StreamOutput.NumStrides * sizeof(uint32), true);
-		}
 		if (PSODesc->VS.BytecodeLength)
 		{
 			DiskCaches[PSO_CACHE_GRAPHICS].SetPointerAndAdvanceFilePosition((void**)&PSODesc->VS.pShaderBytecode, PSODesc->VS.BytecodeLength, bBackShadersWithSystemMemory);
@@ -349,22 +332,6 @@ void FD3D12PipelineStateCache::AddToDiskCache(const FD3D12LowLevelGraphicsPipeli
 				DiskCache.AppendData((void*)PsoDesc.InputLayout.pInputElementDescs[i].SemanticName, stringLength);
 			}
 		}
-		if (PsoDesc.StreamOutput.NumEntries)
-		{
-			DiskCache.AppendData((void*)&PsoDesc.StreamOutput.pSODeclaration, PsoDesc.StreamOutput.NumEntries * sizeof(D3D12_SO_DECLARATION_ENTRY));
-			for (uint32 i = 0; i < PsoDesc.StreamOutput.NumEntries; i++)
-			{
-				//Save the Sematic name string
-				uint32 stringLength = (uint32)strnlen_s(PsoDesc.StreamOutput.pSODeclaration[i].SemanticName, IL_MAX_SEMANTIC_NAME);
-				stringLength++; // include the NULL char
-				DiskCache.AppendData((void*)&stringLength, sizeof(stringLength));
-				DiskCache.AppendData((void*)PsoDesc.StreamOutput.pSODeclaration[i].SemanticName, stringLength);
-			}
-		}
-		if (PsoDesc.StreamOutput.NumStrides)
-		{
-			DiskCache.AppendData((void*)&PsoDesc.StreamOutput.pBufferStrides, PsoDesc.StreamOutput.NumStrides * sizeof(uint32));
-		}
 		if (PsoDesc.VS.BytecodeLength)
 		{
 			DiskCache.AppendData((void*)PsoDesc.VS.pShaderBytecode, PsoDesc.VS.BytecodeLength);
diff --git a/Engine/Source/Runtime/EmptyRHI/Private/EmptyCommands.cpp b/Engine/Source/Runtime/EmptyRHI/Private/EmptyCommands.cpp
--- a/Engine/Source/Runtime/EmptyRHI/Private/EmptyCommands.cpp
+++ b/Engine/Source/Runtime/EmptyRHI/Private/EmptyCommands.cpp
@@ -12,11 +12,6 @@ void FEmptyDynamicRHI::RHISetStreamSource(uint32 StreamIndex, FRHIVertexBuffer*
 
 }
 
-void FEmptyDynamicRHI::RHISetStreamOutTargets(uint32 NumTargets, const FRHIVertexBuffer** VertexBuffers, const uint32* Offsets)
-{
-
-}
-
 void FEmptyDynamicRHI::RHISetRasterizerState(FRHIRasterizerState* NewStateRHI)
 {
 	FEmptyRasterizerState* NewState = ResourceCast(NewStateRHI);
diff --git a/Engine/Source/Runtime/EmptyRHI/Public/EmptyRHI.h b/Engine/Source/Runtime/EmptyRHI/Public/EmptyRHI.h
--- a/Engine/Source/Runtime/EmptyRHI/Public/EmptyRHI.h
+++ b/Engine/Source/Runtime/EmptyRHI/Public/EmptyRHI.h
@@ -118,7 +118,6 @@ public:
 	virtual FViewportRHIRef RHICreateViewport(void* WindowHandle, uint32 SizeX, uint32 SizeY, bool bIsFullscreen, EPixelFormat PreferredPixelFormat) final override;
 	virtual void RHIResizeViewport(FRHIViewport* Viewport, uint32 SizeX, uint32 SizeY, bool bIsFullscreen) final override;
 	virtual void RHITick(float DeltaTime) final override;
-	virtual void RHISetStreamOutTargets(uint32 NumTargets,const FRHIVertexBuffer** VertexBuffers,const uint32* Offsets) final override;
 	virtual void RHIBlockUntilGPUIdle() final override;
 	virtual bool RHIGetAvailableResolutions(FScreenResolutionArray& Resolutions, bool bIgnoreRefreshRate) final override;
 	virtual void RHIGetSupportedResolution(uint32& Width, uint32& Height) final override;
diff --git a/Engine/Source/Runtime/NullDrv/Public/NullRHI.h b/Engine/Source/Runtime/NullDrv/Public/NullRHI.h
--- a/Engine/Source/Runtime/NullDrv/Public/NullRHI.h
+++ b/Engine/Source/Runtime/NullDrv/Public/NullRHI.h
@@ -513,10 +513,6 @@ public:
 	{
 	}
 
-	virtual void RHISetStreamOutTargets(uint32 NumTargets, FRHIVertexBuffer* const* VertexBuffers,const uint32* Offsets) final override
-	{
-
-	}
 	virtual void RHISetRasterizerState(FRHIRasterizerState* NewState) final override
 	{
 
diff --git a/Engine/Source/Runtime/OpenGLDrv/Private/OpenGLCommands.cpp b/Engine/Source/Runtime/OpenGLDrv/Private/OpenGLCommands.cpp
--- a/Engine/Source/Runtime/OpenGLDrv/Private/OpenGLCommands.cpp
+++ b/Engine/Source/Runtime/OpenGLDrv/Private/OpenGLCommands.cpp
@@ -282,11 +282,6 @@ void FOpenGLDynamicRHI::RHISetStreamSource(uint32 StreamIndex, FRHIVertexBuffer*
 	PendingState.Streams[StreamIndex].Offset = Offset;
 }
 
-void FOpenGLDynamicRHI::RHISetStreamOutTargets(uint32 NumTargets, FRHIVertexBuffer* const* VertexBuffers, const uint32* Offsets)
-{
-	check(0);
-}
-
 // Rasterizer state.
 void FOpenGLDynamicRHI::RHISetRasterizerState(FRHIRasterizerState* NewStateRHI)
 {
diff --git a/Engine/Source/Runtime/OpenGLDrv/Public/OpenGLDrv.h b/Engine/Source/Runtime/OpenGLDrv/Public/OpenGLDrv.h
--- a/Engine/Source/Runtime/OpenGLDrv/Public/OpenGLDrv.h
+++ b/Engine/Source/Runtime/OpenGLDrv/Public/OpenGLDrv.h
@@ -428,7 +428,6 @@ public:
 	virtual FViewportRHIRef RHICreateViewport(void* WindowHandle, uint32 SizeX, uint32 SizeY, bool bIsFullscreen, EPixelFormat PreferredPixelFormat) final override;
 	virtual void RHIResizeViewport(FRHIViewport* Viewport, uint32 SizeX, uint32 SizeY, bool bIsFullscreen) final override;
 	virtual void RHITick(float DeltaTime) final override;
-	virtual void RHISetStreamOutTargets(uint32 NumTargets, FRHIVertexBuffer* const* VertexBuffers,const uint32* Offsets) final override;
 	virtual void RHIBlockUntilGPUIdle() final override;
 	virtual void RHISubmitCommandsAndFlushGPU() final override;
 	virtual void RHIPollOcclusionQueries() final override;
diff --git a/Engine/Source/Runtime/RHI/Public/DynamicRHI.h b/Engine/Source/Runtime/RHI/Public/DynamicRHI.h
--- a/Engine/Source/Runtime/RHI/Public/DynamicRHI.h
+++ b/Engine/Source/Runtime/RHI/Public/DynamicRHI.h
@@ -987,7 +987,11 @@ public:
 	/** Sets stream output targets, for use with a geometry shader created with RHICreateGeometryShaderWithStreamOutput. */
 	//@todo this should be a CMDLIST method
 	// FlushType: Flush Immediate (seems wrong)
-	virtual void RHISetStreamOutTargets(uint32 NumTargets, FRHIVertexBuffer* const* VertexBuffers, const uint32* Offsets) = 0;
+	//#todo-RemoveStreamOut
+	virtual void RHISetStreamOutTargets(uint32 NumTargets, FRHIVertexBuffer* const* VertexBuffers, const uint32* Offsets)
+	{
+		check(0);
+	}
 
 	// Blocks the CPU until the GPU catches up and goes idle.
 	// FlushType: Flush Immediate (seems wrong)
diff --git a/Engine/Source/Runtime/RHI/Public/RHI.h b/Engine/Source/Runtime/RHI/Public/RHI.h
--- a/Engine/Source/Runtime/RHI/Public/RHI.h
+++ b/Engine/Source/Runtime/RHI/Public/RHI.h
@@ -736,6 +736,7 @@ struct FVertexElement
 typedef TArray<FVertexElement,TFixedAllocator<MaxVertexElementCount> > FVertexDeclarationElementList;
 
 /** RHI representation of a single stream out element. */
+//#todo-RemoveStreamOut
 struct FStreamOutElement
 {
 	/** Index of the output stream from the geometry shader. */
@@ -767,6 +768,7 @@ struct FStreamOutElement
 	{}
 };
 
+//#todo-RemoveStreamOut
 typedef TArray<FStreamOutElement,TFixedAllocator<MaxVertexElementCount> > FStreamOutElementList;
 
 struct FSamplerStateInitializerRHI
diff --git a/Engine/Source/Runtime/RHI/Public/RHICommandList.h b/Engine/Source/Runtime/RHI/Public/RHICommandList.h
--- a/Engine/Source/Runtime/RHI/Public/RHICommandList.h
+++ b/Engine/Source/Runtime/RHI/Public/RHICommandList.h
@@ -4235,6 +4235,7 @@ public:
 		RHITick(DeltaTime);
 	}
 	
+	//#todo-RemoveStreamOut
 	FORCEINLINE void SetStreamOutTargets(uint32 NumTargets, FRHIVertexBuffer* const* VertexBuffers,const uint32* Offsets)
 	{
 		QUICK_SCOPE_CYCLE_COUNTER(STAT_RHIMETHOD_SetStreamOutTargets_Flush);
diff --git a/Engine/Source/Runtime/RHI/Public/RHIValidation.h b/Engine/Source/Runtime/RHI/Public/RHIValidation.h
--- a/Engine/Source/Runtime/RHI/Public/RHIValidation.h
+++ b/Engine/Source/Runtime/RHI/Public/RHIValidation.h
@@ -1029,14 +1029,6 @@ public:
 		RHI->RHITick(DeltaTime);
 	}
 
-	/** Sets stream output targets, for use with a geometry shader created with RHICreateGeometryShaderWithStreamOutput. */
-	//@todo this should be a CMDLIST method
-	// FlushType: Flush Immediate (seems wrong)
-	virtual void RHISetStreamOutTargets(uint32 NumTargets, FRHIVertexBuffer* const* VertexBuffers, const uint32* Offsets) override final
-	{
-		RHI->RHISetStreamOutTargets(NumTargets, VertexBuffers, Offsets);
-	}
-
 	// Blocks the CPU until the GPU catches up and goes idle.
 	// FlushType: Flush Immediate (seems wrong)
 	virtual void RHIBlockUntilGPUIdle() override final
diff --git a/Engine/Source/Runtime/Renderer/Private/ConvertToUniformMesh.cpp b/Engine/Source/Runtime/Renderer/Private/ConvertToUniformMesh.cpp
--- a/Engine/Source/Runtime/Renderer/Private/ConvertToUniformMesh.cpp
+++ b/Engine/Source/Runtime/Renderer/Private/ConvertToUniformMesh.cpp
@@ -264,6 +264,7 @@ int32 FUniformMeshConverter::Convert(
 
 		uint32 Offsets[1] = {0};
 		FRHIVertexBuffer* const StreamOutTargets[1] = {GUniformMeshTemporaryBuffers.TriangleData.GetReference()};
+		//#todo-RemoveStreamOut
 		RHICmdList.SetStreamOutTargets(1, StreamOutTargets, Offsets);
 
 		for (int32 MeshIndex = 0; MeshIndex < MeshElements.Num(); MeshIndex++)
@@ -290,6 +291,7 @@ int32 FUniformMeshConverter::Convert(
 			}
 		}
 
+		//#todo-RemoveStreamOut
 		RHICmdList.SetStreamOutTargets(1, nullptr, Offsets);
 	}
 
diff --git a/Engine/Source/Runtime/VulkanRHI/Private/VulkanCommands.cpp b/Engine/Source/Runtime/VulkanRHI/Private/VulkanCommands.cpp
--- a/Engine/Source/Runtime/VulkanRHI/Private/VulkanCommands.cpp
+++ b/Engine/Source/Runtime/VulkanRHI/Private/VulkanCommands.cpp
@@ -57,11 +57,6 @@ void FVulkanCommandListContext::RHISetStreamSource(uint32 StreamIndex, FRHIVerte
 	}
 }
 
-void FVulkanDynamicRHI::RHISetStreamOutTargets(uint32 NumTargets, FRHIVertexBuffer* const* VertexBuffers, const uint32* Offsets)
-{
-	VULKAN_SIGNAL_UNIMPLEMENTED();
-}
-
 void FVulkanCommandListContext::RHISetComputeShader(FRHIComputeShader* ComputeShaderRHI)
 {
 	FVulkanComputeShader* ComputeShader = ResourceCast(ComputeShaderRHI);
diff --git a/Engine/Source/Runtime/VulkanRHI/Public/VulkanDynamicRHI.h b/Engine/Source/Runtime/VulkanRHI/Public/VulkanDynamicRHI.h
--- a/Engine/Source/Runtime/VulkanRHI/Public/VulkanDynamicRHI.h
+++ b/Engine/Source/Runtime/VulkanRHI/Public/VulkanDynamicRHI.h
@@ -145,7 +145,6 @@ public:
 	virtual void RHIResizeViewport(FRHIViewport* Viewport, uint32 SizeX, uint32 SizeY, bool bIsFullscreen) final override;
 	virtual void RHIResizeViewport(FRHIViewport* Viewport, uint32 SizeX, uint32 SizeY, bool bIsFullscreen, EPixelFormat PreferredPixelFormat) final override;
 	virtual void RHITick(float DeltaTime) final override;
-	virtual void RHISetStreamOutTargets(uint32 NumTargets, FRHIVertexBuffer* const* VertexBuffers, const uint32* Offsets) final override;
 	virtual void RHIBlockUntilGPUIdle() final override;
 	virtual void RHISubmitCommandsAndFlushGPU() final override;
 	virtual void RHISuspendRendering() final override;
diff --git a/Engine/Source/Runtime/Windows/D3D11RHI/Private/D3D11Commands.cpp b/Engine/Source/Runtime/Windows/D3D11RHI/Private/D3D11Commands.cpp
--- a/Engine/Source/Runtime/Windows/D3D11RHI/Private/D3D11Commands.cpp
+++ b/Engine/Source/Runtime/Windows/D3D11RHI/Private/D3D11Commands.cpp
@@ -165,21 +165,6 @@ void FD3D11DynamicRHI::RHISetStreamSource(uint32 StreamIndex, FRHIVertexBuffer*
 	StateCache.SetStreamSource(D3DBuffer, StreamIndex, Offset);
 }
 
-void FD3D11DynamicRHI::RHISetStreamOutTargets(uint32 NumTargets, FRHIVertexBuffer* const* VertexBuffers, const uint32* Offsets)
-{
-	ID3D11Buffer* D3DVertexBuffers[D3D11_SO_BUFFER_SLOT_COUNT] = {0};
-
-	if (VertexBuffers)
-	{
-		for (uint32 BufferIndex = 0; BufferIndex < NumTargets; BufferIndex++)
-		{
-			D3DVertexBuffers[BufferIndex] = VertexBuffers[BufferIndex] ? ((FD3D11VertexBuffer*)VertexBuffers[BufferIndex])->Resource : NULL;
-		}
-	}
-
-	Direct3DDeviceIMContext->SOSetTargets(NumTargets, D3DVertexBuffers, Offsets);
-}
-
 // Rasterizer state.
 void FD3D11DynamicRHI::RHISetRasterizerState(FRHIRasterizerState* NewStateRHI)
 {
diff --git a/Engine/Source/Runtime/Windows/D3D11RHI/Private/D3D11RHIPrivate.h b/Engine/Source/Runtime/Windows/D3D11RHI/Private/D3D11RHIPrivate.h
--- a/Engine/Source/Runtime/Windows/D3D11RHI/Private/D3D11RHIPrivate.h
+++ b/Engine/Source/Runtime/Windows/D3D11RHI/Private/D3D11RHIPrivate.h
@@ -486,7 +486,6 @@ public:
 	virtual void RHIResizeViewport(FRHIViewport* Viewport, uint32 SizeX, uint32 SizeY, bool bIsFullscreen) final override;
 	virtual void RHIResizeViewport(FRHIViewport* Viewport, uint32 SizeX, uint32 SizeY, bool bIsFullscreen, EPixelFormat PreferredPixelFormat) final override;
 	virtual void RHITick(float DeltaTime) final override;
-	virtual void RHISetStreamOutTargets(uint32 NumTargets, FRHIVertexBuffer* const* VertexBuffers,const uint32* Offsets) final override;
 	virtual void RHIBlockUntilGPUIdle() final override;
 	virtual bool RHIGetAvailableResolutions(FScreenResolutionArray& Resolutions, bool bIgnoreRefreshRate) final override;
 	virtual void RHIGetSupportedResolution(uint32& Width, uint32& Height) final override;
