DR - Fix d3d12 streamout

diff --git a/Engine/Source/Runtime/D3D12RHI/Private/HoloLens/HoloLensD3D12DiskCache.h b/Engine/Source/Runtime/D3D12RHI/Private/HoloLens/HoloLensD3D12DiskCache.h
--- a/Engine/Source/Runtime/D3D12RHI/Private/HoloLens/HoloLensD3D12DiskCache.h
+++ b/Engine/Source/Runtime/D3D12RHI/Private/HoloLens/HoloLensD3D12DiskCache.h
@@ -10,7 +10,7 @@ class FDiskCacheInterface
 {
 	// Increment this if changes are made to the
 	// disk caches so stale caches get updated correctly
-	static const uint32 mCurrentHeaderVersion = 8;
+	static const uint32 mCurrentHeaderVersion = 9;
 	struct FDiskCacheHeader
 	{
 		uint32 mHeaderVersion;
diff --git a/Engine/Source/Runtime/D3D12RHI/Private/HoloLens/HoloLensD3D12PipelineState.h b/Engine/Source/Runtime/D3D12RHI/Private/HoloLens/HoloLensD3D12PipelineState.h
--- a/Engine/Source/Runtime/D3D12RHI/Private/HoloLens/HoloLensD3D12PipelineState.h
+++ b/Engine/Source/Runtime/D3D12RHI/Private/HoloLens/HoloLensD3D12PipelineState.h
@@ -23,7 +23,6 @@ struct FD3D12_GRAPHICS_PIPELINE_STATE_STREAM
 	CD3DX12_PIPELINE_STATE_STREAM_PRIMITIVE_TOPOLOGY PrimitiveTopologyType;
 	CD3DX12_PIPELINE_STATE_STREAM_VS VS;
 	CD3DX12_PIPELINE_STATE_STREAM_GS GS;
-	CD3DX12_PIPELINE_STATE_STREAM_STREAM_OUTPUT StreamOutput;
 	CD3DX12_PIPELINE_STATE_STREAM_HS HS;
 	CD3DX12_PIPELINE_STATE_STREAM_DS DS;
 	CD3DX12_PIPELINE_STATE_STREAM_PS PS;
diff --git a/Engine/Source/Runtime/D3D12RHI/Private/Windows/WindowsD3D12DiskCache.h b/Engine/Source/Runtime/D3D12RHI/Private/Windows/WindowsD3D12DiskCache.h
--- a/Engine/Source/Runtime/D3D12RHI/Private/Windows/WindowsD3D12DiskCache.h
+++ b/Engine/Source/Runtime/D3D12RHI/Private/Windows/WindowsD3D12DiskCache.h
@@ -10,7 +10,7 @@ class FDiskCacheInterface
 {
 	// Increment this if changes are made to the
 	// disk caches so stale caches get updated correctly
-	static const uint32 mCurrentHeaderVersion = 8;
+	static const uint32 mCurrentHeaderVersion = 9;
 	struct FDiskCacheHeader
 	{
 		uint32 mHeaderVersion;
diff --git a/Engine/Source/Runtime/D3D12RHI/Private/Windows/WindowsD3D12PipelineState.h b/Engine/Source/Runtime/D3D12RHI/Private/Windows/WindowsD3D12PipelineState.h
--- a/Engine/Source/Runtime/D3D12RHI/Private/Windows/WindowsD3D12PipelineState.h
+++ b/Engine/Source/Runtime/D3D12RHI/Private/Windows/WindowsD3D12PipelineState.h
@@ -23,7 +23,6 @@ struct FD3D12_GRAPHICS_PIPELINE_STATE_STREAM
 	CD3DX12_PIPELINE_STATE_STREAM_PRIMITIVE_TOPOLOGY PrimitiveTopologyType;
 	CD3DX12_PIPELINE_STATE_STREAM_VS VS;
 	CD3DX12_PIPELINE_STATE_STREAM_GS GS;
-	CD3DX12_PIPELINE_STATE_STREAM_STREAM_OUTPUT StreamOutput;
 	CD3DX12_PIPELINE_STATE_STREAM_HS HS;
 	CD3DX12_PIPELINE_STATE_STREAM_DS DS;
 	CD3DX12_PIPELINE_STATE_STREAM_PS PS;
