Makes static analsysis happy.

diff --git a/Engine/Source/Runtime/D3D12RHI/Private/D3D12RenderTarget.cpp b/Engine/Source/Runtime/D3D12RHI/Private/D3D12RenderTarget.cpp
--- a/Engine/Source/Runtime/D3D12RHI/Private/D3D12RenderTarget.cpp
+++ b/Engine/Source/Runtime/D3D12RHI/Private/D3D12RenderTarget.cpp
@@ -1421,6 +1421,7 @@ void FD3D12DynamicRHI::RHIMapStagingSurface(FRHITexture* TextureRHI, void*& OutD
 	}
 #endif
 
+	check(DestTexture2D);
 	FD3D12Resource* Texture = DestTexture2D->GetResource();
 
 	DXGI_FORMAT Format = (DXGI_FORMAT)GPixelFormats[DestTexture2D->GetFormat()].PlatformFormat;
@@ -1487,6 +1488,7 @@ void FD3D12DynamicRHI::RHIUnmapStagingSurface(FRHITexture* TextureRHI, uint32 GP
 	}
 #endif
 
+	check(DestTexture2D);
 	ID3D12Resource* Texture = DestTexture2D->GetResource()->GetResource();
 
 	Texture->Unmap(0, nullptr);
diff --git a/Engine/Source/Runtime/Renderer/Private/PostProcess/PostProcessUpscale.cpp b/Engine/Source/Runtime/Renderer/Private/PostProcess/PostProcessUpscale.cpp
--- a/Engine/Source/Runtime/Renderer/Private/PostProcess/PostProcessUpscale.cpp
+++ b/Engine/Source/Runtime/Renderer/Private/PostProcess/PostProcessUpscale.cpp
@@ -247,6 +247,7 @@ FScreenPassTexture AddUpscalePass(FRDGBuilder& GraphBuilder, const FViewInfo& Vi
 			SetScreenPassPipelineState(RHICmdList, FScreenPassPipelineState(*TypedVertexShader, *PixelShader));
 			VertexShader = *TypedVertexShader;
 		}
+		check(VertexShader != nullptr);
 
 		SetShaderParameters(RHICmdList, *PixelShader, PixelShader->GetPixelShader(), *PassParameters);
 
diff --git a/Engine/Source/Runtime/Renderer/Private/ScreenSpaceDenoise.cpp b/Engine/Source/Runtime/Renderer/Private/ScreenSpaceDenoise.cpp
--- a/Engine/Source/Runtime/Renderer/Private/ScreenSpaceDenoise.cpp
+++ b/Engine/Source/Runtime/Renderer/Private/ScreenSpaceDenoise.cpp
@@ -1913,7 +1913,7 @@ static void DenoiseSignalAtConstantPixelDensity(
 
 				if (CompressedMetadataLayout == ECompressedMetadataLayout::DepthAndViewNormal)
 				{
-					if (i == 0)
+					// if (i == 0)
 					{
 						Dest = &View.ViewState->PrevFrameViewInfo.CompressedDepthViewNormal;
 					}
