Implemented native instancing for ray tracing geometry. Same shader bindings can now be used with many copies of the geometry that have unique transforms.
This results in a significant CPU time savings for GatherRayTracingWorldInstances(), BindRayTracingMaterialPipeline() and SetRayTracingHitGroup().

(ushell-p4-cherrypick of 9567809 by [REDACTED_PERSON])

(ushell-p4-cherrypick of 9568441 by [REDACTED_PERSON])

diff --git a/Engine/Source/Runtime/D3D12RHI/Private/D3D12RayTracing.cpp b/Engine/Source/Runtime/D3D12RHI/Private/D3D12RayTracing.cpp
--- a/Engine/Source/Runtime/D3D12RHI/Private/D3D12RayTracing.cpp
+++ b/Engine/Source/Runtime/D3D12RHI/Private/D3D12RayTracing.cpp
@@ -2291,12 +2291,20 @@ void FD3D12RayTracingScene::BuildAccelerationStructure(FD3D12CommandContext& Com
 	FD3D12Adapter* Adapter = CommandContext.GetParentAdapter();
 	ID3D12Device5* RayTracingDevice = CommandContext.GetParentDevice()->GetRayTracingDevice();
 
+	const uint32 NumSceneInstances = Instances.Num();
+
+	uint32 NumDxrInstances = 0;
+	for (uint32 InstanceIndex = 0; InstanceIndex < NumSceneInstances; ++InstanceIndex)
+	{
+		NumDxrInstances += Instances[InstanceIndex].Transforms.Num();
+	}
+
 	D3D12_BUILD_RAYTRACING_ACCELERATION_STRUCTURE_INPUTS PrebuildDescInputs = {};
 
 	D3D12_RAYTRACING_ACCELERATION_STRUCTURE_PREBUILD_INFO PrebuildInfo = {};
 	PrebuildDescInputs.Type = D3D12_RAYTRACING_ACCELERATION_STRUCTURE_TYPE_TOP_LEVEL;
 	PrebuildDescInputs.DescsLayout = D3D12_ELEMENTS_LAYOUT_ARRAY;
-	PrebuildDescInputs.NumDescs = Instances.Num();
+	PrebuildDescInputs.NumDescs = NumDxrInstances;
 	PrebuildDescInputs.Flags = BuildFlags;
 
 	RayTracingDevice->GetRaytracingAccelerationStructurePrebuildInfo(&PrebuildDescInputs, &PrebuildInfo);
@@ -2336,13 +2344,13 @@ void FD3D12RayTracingScene::BuildAccelerationStructure(FD3D12CommandContext& Com
 
 	TotalPrimitiveCount = 0;
 
-	if (Instances.Num())
+	if (NumSceneInstances)
 	{
 		FRHIResourceCreateInfo CreateInfo;
 		CreateInfo.GPUMask = FRHIGPUMask::FromIndex(GPUIndex);
 
 		D3D12_RESOURCE_DESC InstanceBufferDesc = CD3DX12_RESOURCE_DESC::Buffer(
-			sizeof(D3D12_RAYTRACING_INSTANCE_DESC) * Instances.Num(),
+			sizeof(D3D12_RAYTRACING_INSTANCE_DESC) * PrebuildDescInputs.NumDescs,
 			D3D12_RESOURCE_FLAG_NONE, D3D12_RAYTRACING_INSTANCE_DESCS_BYTE_ALIGNMENT);
 
 		// Create a temporary (volatile) buffer to hold instance data that we're about to upload.
@@ -2357,12 +2365,12 @@ void FD3D12RayTracingScene::BuildAccelerationStructure(FD3D12CommandContext& Com
 
 		check(MappedData);
 
-		uint32 InstanceIndex = 0;
-
 		TSet<FD3D12RayTracingGeometry*> UniqueGeometries;
+		uint32 DxrInstanceIndex = 0;
 
-		for (const FRayTracingGeometryInstance& Instance : Instances)
+		for (uint32 InstanceIndex = 0; InstanceIndex < NumSceneInstances; ++InstanceIndex)
 		{
+			const FRayTracingGeometryInstance& Instance = Instances[InstanceIndex];
 			FD3D12RayTracingGeometry* Geometry = FD3D12DynamicRHI::ResourceCast(Instance.GeometryRHI);
 
 			checkf(!Geometry->IsDirty(CommandContext.GetGPUIndex()),
@@ -2370,16 +2378,6 @@ void FD3D12RayTracingScene::BuildAccelerationStructure(FD3D12CommandContext& Com
 
 			D3D12_RAYTRACING_INSTANCE_DESC InstanceDesc = {};
 
-			FMatrix TransformTransposed = Instance.Transform.GetTransposed();
-
-			// Ensure the last row of the original Transform is <0,0,0,1>
-			check((TransformTransposed.M[3][0] == 0)
-				&& (TransformTransposed.M[3][1] == 0)
-				&& (TransformTransposed.M[3][2] == 0)
-				&& (TransformTransposed.M[3][3] == 1));
-
-			FMemory::Memcpy(&InstanceDesc.Transform, &TransformTransposed.M[0][0], sizeof(InstanceDesc.Transform));
-
 			InstanceDesc.InstanceID = Instance.UserData;
 			InstanceDesc.InstanceMask = Instance.Mask;
 			InstanceDesc.InstanceContributionToHitGroupIndex = SegmentPrefixSum[InstanceIndex] * ShaderSlotsPerGeometrySegment;
@@ -2405,10 +2403,22 @@ void FD3D12RayTracingScene::BuildAccelerationStructure(FD3D12CommandContext& Com
 				BottomLevelAccelerationStructureBuffers[GPUIndex].Add(Geometry->AccelerationStructureBuffers[GPUIndex]);
 			}
 
-			MappedData[InstanceIndex] = InstanceDesc;
-			++InstanceIndex;
+			const uint32 NumTransforms = Instance.Transforms.Num();
+			for (uint32 TransformIndex = 0; TransformIndex < NumTransforms; ++TransformIndex)
+			{
+				FMatrix TransformTransposed = Instance.Transforms[TransformIndex].GetTransposed();
+
+				// Ensure the last row of the original Transform is <0,0,0,1>
+				check((TransformTransposed.M[3][0] == 0)
+					&& (TransformTransposed.M[3][1] == 0)
+					&& (TransformTransposed.M[3][2] == 0)
+					&& (TransformTransposed.M[3][3] == 1));
+
+				FMemory::Memcpy(&InstanceDesc.Transform, &TransformTransposed.M[0][0], sizeof(InstanceDesc.Transform));
+				MappedData[DxrInstanceIndex++] = InstanceDesc;
+			}
 
-			TotalPrimitiveCount += Geometry->TotalPrimitiveCount;
+			TotalPrimitiveCount += Geometry->TotalPrimitiveCount * Instance.Transforms.Num();
 		}
 
 		Adapter->GetOwningRHI()->UnlockBuffer(nullptr, InstanceBuffer.GetReference());
diff --git a/Engine/Source/Runtime/RHI/Public/DynamicRHI.h b/Engine/Source/Runtime/RHI/Public/DynamicRHI.h
--- a/Engine/Source/Runtime/RHI/Public/DynamicRHI.h
+++ b/Engine/Source/Runtime/RHI/Public/DynamicRHI.h
@@ -79,7 +79,7 @@ struct FRHIFlipDetails
 struct FRayTracingGeometryInstance
 {
 	FRHIRayTracingGeometry* GeometryRHI = nullptr;
-	FMatrix Transform = FMatrix(EForceInit::ForceInitToZero);
+	TArray<FMatrix, TInlineAllocator<1>> Transforms;
 	uint32 UserData = 0;
 	uint8 Mask = 0xFF;
 
diff --git a/Engine/Source/Runtime/Renderer/Private/DeferredShadingRenderer.cpp b/Engine/Source/Runtime/Renderer/Private/DeferredShadingRenderer.cpp
--- a/Engine/Source/Runtime/Renderer/Private/DeferredShadingRenderer.cpp
+++ b/Engine/Source/Runtime/Renderer/Private/DeferredShadingRenderer.cpp
@@ -760,7 +760,7 @@ bool FDeferredShadingSceneRenderer::GatherRayTracingWorldInstances(FRHICommandLi
 
 							// When no cached command is found, NewInstanceMask == 0 and the instance is effectively filtered out
 							FRayTracingGeometryInstance RayTracingInstance = { RayTracingGeometryInstance };
-							RayTracingInstance.Transform = Scene->PrimitiveTransforms[PrimitiveIndex];
+							RayTracingInstance.Transforms.Add(Scene->PrimitiveTransforms[PrimitiveIndex]);
 							RayTracingInstance.UserData = (uint32)PrimitiveIndex;
 							RayTracingInstance.Mask = NewInstanceMask;
 							RayTracingInstance.bForceOpaque = bAllSegmentsOpaque;
@@ -806,25 +806,24 @@ bool FDeferredShadingSceneRenderer::GatherRayTracingWorldInstances(FRHICommandLi
 
 						check(Instance.Materials.Num() == Instance.Geometry->Initializer.Segments.Num() || (Instance.Geometry->Initializer.Segments.Num() == 0 && Instance.Materials.Num() == 1));
 
-						for (const FMatrix& InstanceTransform : Instance.InstanceTransforms)
-						{
-							RayTracingInstance.Transform = InstanceTransform;
+						RayTracingInstance.Transforms.SetNumUninitialized(Instance.InstanceTransforms.Num());
+						FMemory::Memcpy(RayTracingInstance.Transforms.GetData(), Instance.InstanceTransforms.GetData(), Instance.InstanceTransforms.Num() * sizeof(RayTracingInstance.Transforms[0]));
+						static_assert(TIsSame<decltype(RayTracingInstance.Transforms[0]), decltype(Instance.InstanceTransforms[0])>::Value, "Unexpected transform type");
 
-							uint32 InstanceIndex = ReferenceView.RayTracingGeometryInstances.Add(RayTracingInstance);
+						uint32 InstanceIndex = ReferenceView.RayTracingGeometryInstances.Add(RayTracingInstance);
 
-							for (int32 ViewIndex = 1; ViewIndex < Views.Num(); ViewIndex++)
-							{
-								Views[ViewIndex].RayTracingGeometryInstances.Add(RayTracingInstance);
-							}
+						for (int32 ViewIndex = 1; ViewIndex < Views.Num(); ViewIndex++)
+						{
+							Views[ViewIndex].RayTracingGeometryInstances.Add(RayTracingInstance);
+						}
 
-							for (int32 SegmentIndex = 0; SegmentIndex < Instance.Materials.Num(); SegmentIndex++)
-							{
-								FMeshBatch& MeshBatch = Instance.Materials[SegmentIndex];
-								FDynamicRayTracingMeshCommandContext CommandContext(ReferenceView.DynamicRayTracingMeshCommandStorage, ReferenceView.VisibleRayTracingMeshCommands, SegmentIndex, InstanceIndex);
-								FRayTracingMeshProcessor RayTracingMeshProcessor(&CommandContext, Scene, &ReferenceView);
+						for (int32 SegmentIndex = 0; SegmentIndex < Instance.Materials.Num(); SegmentIndex++)
+						{
+							FMeshBatch& MeshBatch = Instance.Materials[SegmentIndex];
+							FDynamicRayTracingMeshCommandContext CommandContext(ReferenceView.DynamicRayTracingMeshCommandStorage, ReferenceView.VisibleRayTracingMeshCommands, SegmentIndex, InstanceIndex);
+							FRayTracingMeshProcessor RayTracingMeshProcessor(&CommandContext, Scene, &ReferenceView);
 
-								RayTracingMeshProcessor.AddMeshBatch(MeshBatch, 1, SceneProxy);
-							}
+							RayTracingMeshProcessor.AddMeshBatch(MeshBatch, 1, SceneProxy);
 						}
 					}
 				}
diff --git a/Engine/Source/Runtime/Renderer/Private/Tests/RayTracingTestbed.cpp b/Engine/Source/Runtime/Renderer/Private/Tests/RayTracingTestbed.cpp
--- a/Engine/Source/Runtime/Renderer/Private/Tests/RayTracingTestbed.cpp
+++ b/Engine/Source/Runtime/Renderer/Private/Tests/RayTracingTestbed.cpp
@@ -105,7 +105,7 @@ bool RunRayTracingTestbed_RenderThread(const FString& Parameters)
 	FRayTracingGeometryRHIRef Geometry = RHICreateRayTracingGeometry(GeometryInitializer);
 
 	FRayTracingGeometryInstance Instances[] = {
-		FRayTracingGeometryInstance { Geometry, FMatrix::Identity, 0, 0xFF }
+		FRayTracingGeometryInstance { Geometry, {FMatrix::Identity}, 0, 0xFF }
 	};
 
 	FRayTracingSceneInitializer Initializer;
