Implemented support for binding different vertex buffers per ray tracing geometry segment. Removed the skin cache vertex buffer merging step, improving GPU performance. Deprecated RHICopyBufferRegion/s, as it was exclusively required for skin cache VB merge. Unified bottom level acceleration structure build and update APIs.

(ushell-p4-cherrypick of 9535449 by [REDACTED_PERSON])

diff --git a/Engine/Plugins/Experimental/GeometryCache/Source/GeometryCache/Private/GeometryCacheSceneProxy.cpp b/Engine/Plugins/Experimental/GeometryCache/Source/GeometryCache/Private/GeometryCacheSceneProxy.cpp
--- a/Engine/Plugins/Experimental/GeometryCache/Source/GeometryCache/Private/GeometryCacheSceneProxy.cpp
+++ b/Engine/Plugins/Experimental/GeometryCache/Source/GeometryCache/Private/GeometryCacheSceneProxy.cpp
@@ -155,19 +155,19 @@ FGeometryCacheSceneProxy::FGeometryCacheSceneProxy(UGeometryCacheComponent* Comp
 						{
 							FRayTracingGeometryInitializer Initializer;
 							const int PositionBufferIndex = Section->CurrentPositionBufferIndex != -1 ? Section->CurrentPositionBufferIndex % 2 : 0;
-							Initializer.PositionVertexBuffer = Section->PositionBuffers[PositionBufferIndex].VertexBufferRHI;
 							Initializer.IndexBuffer = Section->IndexBuffer.IndexBufferRHI;
-							Initializer.VertexBufferStride = sizeof(FVector);
-							Initializer.VertexBufferByteOffset = 0;
 							Initializer.TotalPrimitiveCount = 0;
-							Initializer.VertexBufferElementType = VET_Float3;
 							Initializer.GeometryType = RTGT_Triangles;
 							Initializer.bFastBuild = false;
 
 							TArray<FRayTracingGeometrySegment> Segments;
 							for (FGeometryCacheMeshBatchInfo& BatchInfo : Section->MeshData->BatchesInfo)
 							{
-								Segments.Add(FRayTracingGeometrySegment { BatchInfo.StartIndex / 3, BatchInfo.NumTriangles });
+								FRayTracingGeometrySegment Segment;
+								Segment.FirstPrimitive = BatchInfo.StartIndex / 3;
+								Segment.NumPrimitives = BatchInfo.NumTriangles;
+								Segment.VertexBuffer = Section->PositionBuffers[PositionBufferIndex].VertexBufferRHI;
+								Segments.Add(Segment);
 								Initializer.TotalPrimitiveCount += BatchInfo.NumTriangles;
 							}
 
@@ -557,14 +557,19 @@ void FGeometryCacheSceneProxy::UpdateAnimation(float NewTime, bool bNewLooping,
 			{
 				const int PositionBufferIndex = Section->CurrentPositionBufferIndex != -1 ? Section->CurrentPositionBufferIndex % 2 : 0;
 
-				Section->RayTracingGeometry.Initializer.PositionVertexBuffer = Section->PositionBuffers[PositionBufferIndex].VertexBufferRHI;
 				Section->RayTracingGeometry.Initializer.IndexBuffer = Section->IndexBuffer.IndexBufferRHI;
 				Section->RayTracingGeometry.Initializer.TotalPrimitiveCount = 0;
 				
-				TArray<FRayTracingGeometrySegment> Segments;
+				TArray<FRayTracingGeometrySegment>& Segments = Section->RayTracingGeometry.Initializer.Segments;
+
 				for (FGeometryCacheMeshBatchInfo& BatchInfo : Section->MeshData->BatchesInfo)
 				{
-					Segments.Add(FRayTracingGeometrySegment { BatchInfo.StartIndex / 3, BatchInfo.NumTriangles });
+					FRayTracingGeometrySegment Segment;
+					Segment.FirstPrimitive = BatchInfo.StartIndex / 3;
+					Segment.NumPrimitives = BatchInfo.NumTriangles;
+					Segment.VertexBuffer = Section->PositionBuffers[PositionBufferIndex].VertexBufferRHI;
+
+					Segments.Add(Segment);
 					Section->RayTracingGeometry.Initializer.TotalPrimitiveCount += BatchInfo.NumTriangles;
 				}
 
diff --git a/Engine/Plugins/Experimental/HairStrands/Source/HairStrandsCore/Private/GroomComponent.cpp b/Engine/Plugins/Experimental/HairStrands/Source/HairStrandsCore/Private/GroomComponent.cpp
--- a/Engine/Plugins/Experimental/HairStrands/Source/HairStrandsCore/Private/GroomComponent.cpp
+++ b/Engine/Plugins/Experimental/HairStrands/Source/HairStrandsCore/Private/GroomComponent.cpp
@@ -183,7 +183,11 @@ public:
 		if (!IsRayTracingEnabled() || !RayTracingGeometry || !RayTracingGeometry->RayTracingGeometryRHI.IsValid())
 			return;
 
-		check(RayTracingGeometry->Initializer.PositionVertexBuffer.IsValid());
+		for (const FRayTracingGeometrySegment& Segment : RayTracingGeometry->Initializer.Segments)
+		{
+			check(Segment.VertexBuffer.IsValid());
+		}
+
 		AddOpaqueRaytracingInstance(GetLocalToWorld(), RayTracingGeometry, RaytracingInstanceMask_ThinShadow, OutRayTracingInstances);
 	}
 #endif
diff --git a/Engine/Plugins/Experimental/HairStrands/Source/HairStrandsCore/Private/HairStrandsRendering.cpp b/Engine/Plugins/Experimental/HairStrands/Source/HairStrandsCore/Private/HairStrandsRendering.cpp
--- a/Engine/Plugins/Experimental/HairStrands/Source/HairStrandsCore/Private/HairStrandsRendering.cpp
+++ b/Engine/Plugins/Experimental/HairStrands/Source/HairStrandsCore/Private/HairStrandsRendering.cpp
@@ -335,29 +335,31 @@ static void UpdateHairAccelerationStructure(FRHICommandList& RHICmdList, FRayTra
 {
 	SCOPED_DRAW_EVENT(RHICmdList, CommitHairRayTracingGeometryUpdates);
 
-	TArray<FAccelerationStructureUpdateParams> Updates;
-	FAccelerationStructureUpdateParams Params;
+	FAccelerationStructureBuildParams Params;
+	Params.BuildMode = EAccelerationStructureBuildMode::Update;
 	Params.Geometry = RayTracingGeometry->RayTracingGeometryRHI;
-	Params.VertexBuffer = RayTracingGeometry->Initializer.PositionVertexBuffer;
-	Updates.Add(Params);
+	Params.Segments = RayTracingGeometry->Initializer.Segments;
 
-	RHICmdList.UpdateAccelerationStructures(Updates);
+	RHICmdList.BuildAccelerationStructures(MakeArrayView(&Params, 1));
 }
 
 static void BuildHairAccelerationStructure(FRHICommandList& RHICmdList, uint32 RaytracingVertexCount, FVertexBufferRHIRef& PositionBuffer, FRayTracingGeometry* OutRayTracingGeometry)
 {
 	FRayTracingGeometryInitializer Initializer;
-	Initializer.PositionVertexBuffer = PositionBuffer;
-	Initializer.VertexBufferByteOffset = 0;
-	Initializer.VertexBufferStride = FHairStrandsRaytracingFormat::SizeInByte;
-	Initializer.VertexBufferElementType = FHairStrandsRaytracingFormat::VertexElementType;
 	Initializer.IndexBuffer = nullptr;
-	Initializer.IndexBufferByteOffset = 0;
+	Initializer.IndexBufferOffset = 0;
 	Initializer.GeometryType = RTGT_Triangles;
 	Initializer.TotalPrimitiveCount = RaytracingVertexCount;
 	Initializer.bFastBuild = true;
 	Initializer.bAllowUpdate = true;
 
+	FRayTracingGeometrySegment Segment;
+	Segment.VertexBuffer = PositionBuffer;
+	Segment.VertexBufferStride = FHairStrandsRaytracingFormat::SizeInByte;
+	Segment.VertexBufferElementType = FHairStrandsRaytracingFormat::VertexElementType;
+	Segment.NumPrimitives = RaytracingVertexCount;
+	Initializer.Segments.Add(Segment);
+
 	OutRayTracingGeometry->SetInitializer(Initializer);
 	OutRayTracingGeometry->RayTracingGeometryRHI = RHICreateRayTracingGeometry(Initializer);
 	RHICmdList.BuildAccelerationStructure(OutRayTracingGeometry->RayTracingGeometryRHI);
diff --git a/Engine/Plugins/FX/Niagara/Source/Niagara/Private/NiagaraRendererRibbons.cpp b/Engine/Plugins/FX/Niagara/Source/Niagara/Private/NiagaraRendererRibbons.cpp
--- a/Engine/Plugins/FX/Niagara/Source/Niagara/Private/NiagaraRendererRibbons.cpp
+++ b/Engine/Plugins/FX/Niagara/Source/Niagara/Private/NiagaraRendererRibbons.cpp
@@ -230,12 +230,8 @@ void FNiagaraRendererRibbons::CreateRenderThreadResources(NiagaraEmitterInstance
 		RayTracingDynamicVertexBuffer.Initialize(4, 256, PF_R32_FLOAT, BUF_UnorderedAccess | BUF_ShaderResource, TEXT("RayTracingDynamicVertexBuffer"));
 
 		FRayTracingGeometryInitializer Initializer;
-		Initializer.PositionVertexBuffer = nullptr;
 		Initializer.IndexBuffer = nullptr;
-		Initializer.VertexBufferStride = 12;
-		Initializer.VertexBufferByteOffset = 0;
 		Initializer.TotalPrimitiveCount = 0;
-		Initializer.VertexBufferElementType = VET_Float3;
 		Initializer.GeometryType = RTGT_Triangles;
 		Initializer.bFastBuild = true;
 		Initializer.bAllowUpdate = false;
@@ -1047,7 +1043,7 @@ void FNiagaraRendererRibbons::GetDynamicRayTracingInstances(FRayTracingMaterialG
 			VertexCount, NumPrimitives, CollectorResources.UniformBuffer, DynamicIndexAllocation);
 
 		RayTracingGeometry.Initializer.IndexBuffer = DynamicIndexAllocation.IndexBuffer->IndexBufferRHI;
-		RayTracingGeometry.Initializer.IndexBufferByteOffset = DynamicIndexAllocation.FirstIndex * sizeof(uint16);
+		RayTracingGeometry.Initializer.IndexBufferOffset = DynamicIndexAllocation.FirstIndex * sizeof(uint16);
 
 
 		FMeshBatch MeshBatch;
diff --git a/Engine/Plugins/FX/Niagara/Source/Niagara/Private/NiagaraRendererSprites.cpp b/Engine/Plugins/FX/Niagara/Source/Niagara/Private/NiagaraRendererSprites.cpp
--- a/Engine/Plugins/FX/Niagara/Source/Niagara/Private/NiagaraRendererSprites.cpp
+++ b/Engine/Plugins/FX/Niagara/Source/Niagara/Private/NiagaraRendererSprites.cpp
@@ -180,12 +180,7 @@ void FNiagaraRendererSprites::CreateRenderThreadResources(NiagaraEmitterInstance
 		RayTracingDynamicVertexBuffer.Initialize(4, 256, PF_R32_FLOAT, BUF_UnorderedAccess | BUF_ShaderResource, TEXT("RayTracingDynamicVertexBuffer"));
 
 		FRayTracingGeometryInitializer Initializer;
-		Initializer.PositionVertexBuffer = nullptr;
 		Initializer.IndexBuffer = nullptr;
-		Initializer.VertexBufferStride = 12;
-		Initializer.VertexBufferByteOffset = 0;
-		Initializer.TotalPrimitiveCount = 0;
-		Initializer.VertexBufferElementType = VET_Float3;
 		Initializer.GeometryType = RTGT_Triangles;
 		Initializer.bFastBuild = true;
 		Initializer.bAllowUpdate = false;
diff --git a/Engine/Plugins/Runtime/ProceduralMeshComponent/Source/ProceduralMeshComponent/Private/ProceduralMeshComponent.cpp b/Engine/Plugins/Runtime/ProceduralMeshComponent/Source/ProceduralMeshComponent/Private/ProceduralMeshComponent.cpp
--- a/Engine/Plugins/Runtime/ProceduralMeshComponent/Source/ProceduralMeshComponent/Private/ProceduralMeshComponent.cpp
+++ b/Engine/Plugins/Runtime/ProceduralMeshComponent/Source/ProceduralMeshComponent/Private/ProceduralMeshComponent.cpp
@@ -176,22 +176,23 @@ public:
 						[this, NewSection/*, VertexBufferRHI, IndexBufferRHI, VertexBufferStride, TrianglesCount, RenderSections*/](FRHICommandListImmediate& RHICmdList)
 					{
 						FRayTracingGeometryInitializer Initializer;
-						Initializer.PositionVertexBuffer = nullptr;
 						Initializer.IndexBuffer = nullptr;
-						Initializer.VertexBufferStride = 12;
-						Initializer.VertexBufferByteOffset = 0;
 						Initializer.TotalPrimitiveCount = 0;
-						Initializer.VertexBufferElementType = VET_Float3;
 						Initializer.GeometryType = RTGT_Triangles;
 						Initializer.bFastBuild = true;
 						Initializer.bAllowUpdate = false;
+
 						NewSection->RayTracingGeometry.SetInitializer(Initializer);
 						NewSection->RayTracingGeometry.InitResource();
 
-						NewSection->RayTracingGeometry.Initializer.PositionVertexBuffer = NewSection->VertexBuffers.PositionVertexBuffer.VertexBufferRHI;
 						NewSection->RayTracingGeometry.Initializer.IndexBuffer = NewSection->IndexBuffer.IndexBufferRHI;
 						NewSection->RayTracingGeometry.Initializer.TotalPrimitiveCount = NewSection->IndexBuffer.Indices.Num() / 3;
 
+						FRayTracingGeometrySegment Segment;
+						Segment.VertexBuffer = NewSection->VertexBuffers.PositionVertexBuffer.VertexBufferRHI;
+						Segment.NumPrimitives = NewSection->RayTracingGeometry.Initializer.TotalPrimitiveCount;
+						NewSection->RayTracingGeometry.Initializer.Segments.Add(Segment);
+
 						//#dxr_todo: add support for segments?
 						
 						NewSection->RayTracingGeometry.UpdateRHI();
@@ -434,7 +435,6 @@ public:
 				
 				if (Section->RayTracingGeometry.RayTracingGeometryRHI.IsValid())
 				{
-					check(Section->RayTracingGeometry.Initializer.PositionVertexBuffer.IsValid());
 					check(Section->RayTracingGeometry.Initializer.IndexBuffer.IsValid());
 
 					FRayTracingInstance RayTracingInstance;
diff --git a/Engine/Source/Runtime/D3D12RHI/Private/D3D12CommandContext.h b/Engine/Source/Runtime/D3D12RHI/Private/D3D12CommandContext.h
--- a/Engine/Source/Runtime/D3D12RHI/Private/D3D12CommandContext.h
+++ b/Engine/Source/Runtime/D3D12RHI/Private/D3D12CommandContext.h
@@ -336,9 +336,7 @@ public:
 #if D3D12_RHI_RAYTRACING
 	virtual void RHICopyBufferRegion(FRHIVertexBuffer* DestBuffer, uint64 DstOffset, FRHIVertexBuffer* SourceBuffer, uint64 SrcOffset, uint64 NumBytes) final override;
 	virtual void RHICopyBufferRegions(const TArrayView<const FCopyBufferRegionParams> Params) final override;
-	virtual void RHIBuildAccelerationStructure(FRHIRayTracingGeometry* Geometry) final override;
-	virtual void RHIUpdateAccelerationStructures(const TArrayView<const FAccelerationStructureUpdateParams> Params) final override;
-	virtual void RHIBuildAccelerationStructures(const TArrayView<const FAccelerationStructureUpdateParams> Params) final override;
+	virtual void RHIBuildAccelerationStructures(const TArrayView<const FAccelerationStructureBuildParams> Params) final override;
 	virtual void RHIBuildAccelerationStructure(FRHIRayTracingScene* Scene) final override;
 	virtual void RHIClearRayTracingBindings(FRHIRayTracingScene* Scene) final override;
 	virtual void RHIRayTraceOcclusion(FRHIRayTracingScene* Scene,
@@ -753,29 +751,7 @@ public:
 		ContextRedirect(RHIEndRenderPass());
 	}
 
-
-#if RHI_RAYTRACING
-	virtual void RHICopyBufferRegion(FRHIVertexBuffer* DestBuffer, uint64 DstOffset, FRHIVertexBuffer* SourceBuffer, uint64 SrcOffset, uint64 NumBytes) final override
-	{
-		ContextRedirect(RHICopyBufferRegion(DestBuffer, DstOffset, SourceBuffer, SrcOffset, NumBytes));
-	}
-
-	virtual void RHICopyBufferRegions(const TArrayView<const FCopyBufferRegionParams> Params) final override
-	{
-		ContextRedirect(RHICopyBufferRegions(Params));
-	}
-#endif
-	virtual void RHIBuildAccelerationStructure(FRHIRayTracingGeometry* Geometry) final override
-	{
-		ContextRedirect(RHIBuildAccelerationStructure(Geometry));
-	}
-
-	virtual void RHIUpdateAccelerationStructures(const TArrayView<const FAccelerationStructureUpdateParams> Params) final override
-	{
-		ContextRedirect(RHIUpdateAccelerationStructures(Params));
-	}
-
-	virtual void RHIBuildAccelerationStructures(const TArrayView<const FAccelerationStructureUpdateParams> Params) final override
+	virtual void RHIBuildAccelerationStructures(const TArrayView<const FAccelerationStructureBuildParams> Params) final override
 	{
 		ContextRedirect(RHIBuildAccelerationStructures(Params));
 	}
diff --git a/Engine/Source/Runtime/D3D12RHI/Private/D3D12RayTracing.cpp b/Engine/Source/Runtime/D3D12RHI/Private/D3D12RayTracing.cpp
--- a/Engine/Source/Runtime/D3D12RHI/Private/D3D12RayTracing.cpp
+++ b/Engine/Source/Runtime/D3D12RHI/Private/D3D12RayTracing.cpp
@@ -1844,29 +1844,6 @@ FRayTracingPipelineStateRHIRef FD3D12DynamicRHI::RHICreateRayTracingPipelineStat
 
 FRayTracingGeometryRHIRef FD3D12DynamicRHI::RHICreateRayTracingGeometry(const FRayTracingGeometryInitializer& Initializer)
 {
-	checkf(Initializer.PositionVertexBuffer, TEXT("Position vertex buffer is required for ray tracing geometry"));
-	checkf(Initializer.VertexBufferStride, TEXT("Position vertex buffer is required for ray tracing geometry"));
-	checkf(Initializer.VertexBufferStride % 4 == 0, TEXT("Position vertex buffer stride must be aligned to 4 bytes for ByteAddressBuffer loads to work"));
-
-	if (Initializer.GeometryType == RTGT_Triangles)
-	{
-		// #dxr_todo UE-72160: VET_Half4 (DXGI_FORMAT_R16G16B16A16_FLOAT) is also supported by DXR. Should we support it?
-		check(Initializer.VertexBufferElementType == VET_Float4 || Initializer.VertexBufferElementType == VET_Float3 || Initializer.VertexBufferElementType == VET_Float2 || Initializer.VertexBufferElementType == VET_Half2);
-
-		// #dxr_todo UE-72160: temporary constraints on vertex and index buffer formats (this will be relaxed when more flexible vertex/index fetching is implemented)
-		checkf(Initializer.VertexBufferElementType == VET_Float3 || Initializer.VertexBufferElementType == VET_Float4, TEXT("Only float3 vertex buffers are currently implemented.")); // #dxr_todo UE-72160: support other vertex buffer formats
-		if (Initializer.VertexBufferElementType == VET_Float3)
-			checkf(Initializer.VertexBufferStride >= 12, TEXT("Only deinterleaved float3 position vertex buffers are currently implemented.")); // #dxr_todo UE-72160: support interleaved vertex buffers
-
-		if (Initializer.VertexBufferElementType == VET_Float4)
-			checkf(Initializer.VertexBufferStride >= 16, TEXT("Only deinterleaved float3 position vertex buffers are currently implemented.")); // #dxr_todo UE-72160: support interleaved vertex buffers
-	}
-
-	if (Initializer.GeometryType == RTGT_Procedural)
-	{
-		checkf(Initializer.VertexBufferStride >= (2 * sizeof(FVector)), TEXT("Procedural geometry vertex buffer must contain at least 2xFloat3 that defines 3D bounding boxes of primitives."));
-	}
-
 	if (Initializer.IndexBuffer)
 	{
 		checkf(Initializer.IndexBuffer->GetStride() == 2 || Initializer.IndexBuffer->GetStride() == 4, TEXT("Index buffer must be 16 or 32 bit."));
@@ -1876,9 +1853,7 @@ FRayTracingGeometryRHIRef FD3D12DynamicRHI::RHICreateRayTracingGeometry(const FR
 
 	uint32 IndexStride = Initializer.IndexBuffer ? Initializer.IndexBuffer->GetStride() : 0; // stride 0 means implicit triangle list for non-indexed geometry
 	Geometry->IndexStride = IndexStride;
-	Geometry->IndexOffsetInBytes = Initializer.IndexBufferByteOffset;
-	Geometry->VertexOffsetInBytes = Initializer.VertexBufferByteOffset;
-	Geometry->VertexStrideInBytes = Initializer.VertexBufferStride;
+	Geometry->IndexOffsetInBytes = Initializer.IndexBufferOffset;
 	Geometry->TotalPrimitiveCount = Initializer.TotalPrimitiveCount;
 
 	switch (Initializer.GeometryType)
@@ -1918,17 +1893,8 @@ FRayTracingGeometryRHIRef FD3D12DynamicRHI::RHICreateRayTracingGeometry(const FR
 		Geometry->BuildFlags &= ~D3D12_RAYTRACING_ACCELERATION_STRUCTURE_BUILD_FLAG_PREFER_FAST_BUILD;
 	}
 
-	if (Initializer.Segments.Num())
-	{
-		Geometry->Segments = TArray<FRayTracingGeometrySegment>(Initializer.Segments.GetData(), Initializer.Segments.Num());
-	}
-	else
-	{
-		FRayTracingGeometrySegment DefaultSegment;
-		DefaultSegment.FirstPrimitive = 0;
-		DefaultSegment.NumPrimitives = Initializer.TotalPrimitiveCount;
-		Geometry->Segments.Add(DefaultSegment);
-	}
+	checkf(Initializer.Segments.Num() > 0, TEXT("Ray tracing geometry must be initialized with at least one segment."));
+	Geometry->Segments = TArray<FRayTracingGeometrySegment>(Initializer.Segments.GetData(), Initializer.Segments.Num());
 
 #if DO_CHECK
 	static constexpr uint32 IndicesPerPrimitive = 3; // Only triangle meshes are supported
@@ -1937,6 +1903,22 @@ FRayTracingGeometryRHIRef FD3D12DynamicRHI::RHICreateRayTracingGeometry(const FR
 		uint32 ComputedPrimitiveCountForValidation = 0;
 		for (const FRayTracingGeometrySegment& Segment : Geometry->Segments)
 		{
+			checkf(Segment.VertexBuffer, TEXT("Position vertex buffer is required for ray tracing geometry"));
+			checkf(Segment.VertexBufferStride, TEXT("Position vertex buffer is required for ray tracing geometry"));
+			checkf(Segment.VertexBufferStride % 4 == 0, TEXT("Position vertex buffer stride must be aligned to 4 bytes for ByteAddressBuffer loads to work"));
+
+			if (Initializer.GeometryType == RTGT_Triangles)
+			{
+				// #dxr_todo UE-72160: temporary constraints on vertex and index buffer formats (this will be relaxed when more flexible vertex/index fetching is implemented)
+				checkf(Segment.VertexBufferElementType == VET_Float3 
+					|| Segment.VertexBufferElementType == VET_Float4, TEXT("Only float3/4 vertex buffers are currently implemented.")); // #dxr_todo UE-72160: support other vertex buffer formats
+				checkf(Segment.VertexBufferStride >= 12, TEXT("Only deinterleaved float3 position vertex buffers are currently implemented.")); // #dxr_todo UE-72160: support interleaved vertex buffers
+			}
+			else if (Initializer.GeometryType == RTGT_Procedural)
+			{
+				checkf(Segment.VertexBufferStride >= (2 * sizeof(FVector)), TEXT("Procedural geometry vertex buffer must contain at least 2xFloat3 that defines 3D bounding boxes of primitives."));
+			}
+
 			ComputedPrimitiveCountForValidation += Segment.NumPrimitives;
 			check(Segment.FirstPrimitive + Segment.NumPrimitives <= Initializer.TotalPrimitiveCount);
 			
@@ -1949,12 +1931,9 @@ FRayTracingGeometryRHIRef FD3D12DynamicRHI::RHICreateRayTracingGeometry(const FR
 		check(ComputedPrimitiveCountForValidation == Initializer.TotalPrimitiveCount);
 
 	}
-#endif
-
-	Geometry->VertexElemType = Initializer.VertexBufferElementType;
+#endif // DO_CHECK
 
 	Geometry->RHIIndexBuffer = Initializer.IndexBuffer;
-	Geometry->RHIVertexBuffer = Initializer.PositionVertexBuffer;
 	Geometry->SetDirty(FRHIGPUMask::All(), true);
 
 	INC_DWORD_STAT_BY(STAT_D3D12RayTracingTrianglesBLAS, Initializer.TotalPrimitiveCount);
@@ -2032,10 +2011,14 @@ void FD3D12RayTracingGeometry::TransitionBuffers(FD3D12CommandContext& CommandCo
 		}
 	}
 
-	FD3D12VertexBuffer* VertexBuffer = CommandContext.RetrieveObject<FD3D12VertexBuffer>(RHIVertexBuffer.GetReference());
-	if (VertexBuffer->GetResource()->RequiresResourceStateTracking())
+	for (const FRayTracingGeometrySegment& Segment : Segments)
 	{
-		FD3D12DynamicRHI::TransitionResource(CommandContext.CommandListHandle, VertexBuffer->GetResource(), D3D12_RESOURCE_STATE_NON_PIXEL_SHADER_RESOURCE, 0);
+		const FVertexBufferRHIRef& RHIVertexBuffer = Segment.VertexBuffer;
+		FD3D12VertexBuffer* VertexBuffer = CommandContext.RetrieveObject<FD3D12VertexBuffer>(RHIVertexBuffer.GetReference());
+		if (VertexBuffer->GetResource()->RequiresResourceStateTracking())
+		{
+			FD3D12DynamicRHI::TransitionResource(CommandContext.CommandListHandle, VertexBuffer->GetResource(), D3D12_RESOURCE_STATE_NON_PIXEL_SHADER_RESOURCE, 0);
+		}
 	}
 }
 
@@ -2069,8 +2052,9 @@ static void CreateAccelerationStructureBuffers(TRefCountPtr<FD3D12MemBuffer>& Ac
 	SetName(ScratchBuffer->GetResource(), TEXT("Acceleration structure scratch"));
 }
 
-void FD3D12RayTracingGeometry::BuildAccelerationStructure(FD3D12CommandContext& CommandContext, bool bIsUpdate)
+void FD3D12RayTracingGeometry::BuildAccelerationStructure(FD3D12CommandContext& CommandContext, EAccelerationStructureBuildMode BuildMode)
 {
+	const bool bIsUpdate = BuildMode == EAccelerationStructureBuildMode::Update;
 	static constexpr uint32 IndicesPerPrimitive = 3; // Only triangle meshes are supported
 
 	// Array of geometry descriptions, one per segment (single-segment geometry is a common case).
@@ -2079,7 +2063,6 @@ void FD3D12RayTracingGeometry::BuildAccelerationStructure(FD3D12CommandContext&
 	Descs.Reserve(Segments.Num());
 
 	FD3D12IndexBuffer* IndexBuffer = CommandContext.RetrieveObject<FD3D12IndexBuffer>(RHIIndexBuffer.GetReference());
-	FD3D12VertexBuffer* VertexBuffer = CommandContext.RetrieveObject<FD3D12VertexBuffer>(RHIVertexBuffer.GetReference());
 
 	for (const FRayTracingGeometrySegment& Segment : Segments)
 	{
@@ -2100,9 +2083,11 @@ void FD3D12RayTracingGeometry::BuildAccelerationStructure(FD3D12CommandContext&
 			Desc.Flags |= D3D12_RAYTRACING_GEOMETRY_FLAG_NO_DUPLICATE_ANYHIT_INVOCATION;
 		}
 
+		FD3D12VertexBuffer* VertexBuffer = CommandContext.RetrieveObject<FD3D12VertexBuffer>(Segment.VertexBuffer.GetReference());
+
 		if (GeometryType == D3D12_RAYTRACING_GEOMETRY_TYPE_TRIANGLES)
 		{
-			switch (VertexElemType)
+			switch (Segment.VertexBufferElementType)
 			{
 			case VET_Float4: 
 				// While the DXGI_FORMAT_R32G32B32A32_FLOAT format is not supported by DXR, since we manually load vertex 
@@ -2132,7 +2117,7 @@ void FD3D12RayTracingGeometry::BuildAccelerationStructure(FD3D12CommandContext&
 											 IndexOffsetInBytes +
 											 IndexStride * Segment.FirstPrimitive * IndicesPerPrimitive;
 
-				Desc.Triangles.VertexCount = VertexBuffer->ResourceLocation.GetSize() / VertexStrideInBytes;
+				Desc.Triangles.VertexCount = VertexBuffer->ResourceLocation.GetSize() / Segment.VertexBufferStride;
 
 				IndexBuffer->GetResource()->UpdateResidency(CommandContext.CommandListHandle);
 			}
@@ -2145,17 +2130,17 @@ void FD3D12RayTracingGeometry::BuildAccelerationStructure(FD3D12CommandContext&
 
 				checkf(Segments.Num() == 1, TEXT("Non-indexed geometry with multiple segments is not implemented."));
 
-				Desc.Triangles.VertexCount = FMath::Min<uint32>(VertexBuffer->ResourceLocation.GetSize() / VertexStrideInBytes, TotalPrimitiveCount * 3);
+				Desc.Triangles.VertexCount = FMath::Min<uint32>(VertexBuffer->ResourceLocation.GetSize() / Segment.VertexBufferStride, TotalPrimitiveCount * 3);
 			}
 
-			Desc.Triangles.VertexBuffer.StartAddress = VertexBuffer->ResourceLocation.GetGPUVirtualAddress() + VertexOffsetInBytes;
-			Desc.Triangles.VertexBuffer.StrideInBytes = VertexStrideInBytes;
+			Desc.Triangles.VertexBuffer.StartAddress = VertexBuffer->ResourceLocation.GetGPUVirtualAddress() + Segment.VertexBufferOffset;
+			Desc.Triangles.VertexBuffer.StrideInBytes = Segment.VertexBufferStride;
 		}
 		else if (GeometryType == D3D12_RAYTRACING_GEOMETRY_TYPE_PROCEDURAL_PRIMITIVE_AABBS)
 		{
 			Desc.AABBs.AABBCount = Segment.NumPrimitives;
-			Desc.AABBs.AABBs.StartAddress = VertexBuffer->ResourceLocation.GetGPUVirtualAddress() + VertexOffsetInBytes;
-			Desc.AABBs.AABBs.StrideInBytes = VertexStrideInBytes;
+			Desc.AABBs.AABBs.StartAddress = VertexBuffer->ResourceLocation.GetGPUVirtualAddress() + Segment.VertexBufferOffset;
+			Desc.AABBs.AABBs.StrideInBytes = Segment.VertexBufferStride;
 		}
 		else
 		{
@@ -2537,7 +2522,6 @@ FD3D12RayTracingShaderTable* FD3D12RayTracingScene::FindOrCreateShaderTable(cons
 			const uint32 IndexStride = Geometry->IndexStride;
 
 			FD3D12IndexBuffer* IndexBuffer = FD3D12DynamicRHI::ResourceCast(Geometry->RHIIndexBuffer.GetReference(), GPUIndex);
-			FD3D12VertexBuffer* VertexBuffer = FD3D12DynamicRHI::ResourceCast(Geometry->RHIVertexBuffer.GetReference(), GPUIndex);
 
 			// Here, we directly apply the vertex offset to the address but set the offset for index buffer using a RootConstant instead.
 			// Index Buffer has to be 4 bytes aligned.
@@ -2545,20 +2529,22 @@ FD3D12RayTracingShaderTable* FD3D12RayTracingScene::FindOrCreateShaderTable(cons
 			// ByteAddressBuffer loads must be aligned to DWORD boundary.
 			// which means the alignment can be wrong for index buffer if we directly apply the offset to the address when binding it.
 			const D3D12_GPU_VIRTUAL_ADDRESS IndexBufferAddress = IndexBuffer ? IndexBuffer->ResourceLocation.GetGPUVirtualAddress() : 0;
-			const D3D12_GPU_VIRTUAL_ADDRESS VertexBufferAddress = VertexBuffer->ResourceLocation.GetGPUVirtualAddress() + Geometry->VertexOffsetInBytes;
 
 			if (IndexBuffer)
 			{
 				CreatedShaderTable->AddResourceReference(IndexBuffer->ResourceLocation.GetResource(), IndexBuffer);
 			}
 
-			CreatedShaderTable->AddResourceReference(VertexBuffer->ResourceLocation.GetResource(), VertexBuffer);
-
 			const uint32 NumSegments = Geometry->Segments.Num();
 			for (uint32 SegmentIndex = 0; SegmentIndex < NumSegments; ++SegmentIndex)
 			{
 				const FRayTracingGeometrySegment& Segment = Geometry->Segments[SegmentIndex];
 
+				FD3D12VertexBuffer* VertexBuffer = FD3D12DynamicRHI::ResourceCast(Segment.VertexBuffer.GetReference(), GPUIndex);
+				CreatedShaderTable->AddResourceReference(VertexBuffer->ResourceLocation.GetResource(), VertexBuffer);
+
+				const D3D12_GPU_VIRTUAL_ADDRESS VertexBufferAddress = VertexBuffer->ResourceLocation.GetGPUVirtualAddress() + Segment.VertexBufferOffset;
+
 				const uint32 RecordBaseIndex = GetHitRecordBaseIndex(InstanceIndex, SegmentIndex);
 
 				FHitGroupSystemParameters SystemParameters = {};
@@ -2568,10 +2554,10 @@ FD3D12RayTracingShaderTable* FD3D12RayTracingScene::FindOrCreateShaderTable(cons
 				if (Geometry->GeometryType == D3D12_RAYTRACING_GEOMETRY_TYPE_TRIANGLES)
 				{
 					// #dxr_todo UE-72160: support various vertex buffer layouts (fetch/decode based on vertex stride and format)
-					checkf(Geometry->VertexElemType == VET_Float3 || Geometry->VertexElemType == VET_Float4, TEXT("Only VET_Float3 is currently implemented and tested. Other formats will be supported in the future."));
+					checkf(Segment.VertexBufferElementType == VET_Float3 || Segment.VertexBufferElementType == VET_Float4, TEXT("Only VET_Float3 and Float4 are currently implemented and tested. Other formats will be supported in the future."));
 				}
 
-				SystemParameters.RootConstants.SetVertexAndIndexStride(Geometry->VertexStrideInBytes, IndexStride);
+				SystemParameters.RootConstants.SetVertexAndIndexStride(Segment.VertexBufferStride, IndexStride);
 				SystemParameters.RootConstants.IndexBufferOffsetInBytes = IndexStride * Segment.FirstPrimitive * IndicesPerPrimitive + Geometry->IndexOffsetInBytes;
 
 				for (uint32 SlotIndex = 0; SlotIndex < ShaderSlotsPerGeometrySegment; ++SlotIndex)
@@ -2587,57 +2573,47 @@ FD3D12RayTracingShaderTable* FD3D12RayTracingScene::FindOrCreateShaderTable(cons
 	return CreatedShaderTable;
 }
 
-void FD3D12CommandContext::RHIBuildAccelerationStructure(FRHIRayTracingGeometry* InGeometry)
+void FD3D12CommandContext::RHIBuildAccelerationStructures(const TArrayView<const FAccelerationStructureBuildParams> Params)
 {
-	FD3D12RayTracingGeometry* Geometry = FD3D12DynamicRHI::ResourceCast(InGeometry);
-	Geometry->TransitionBuffers(*this);
-	CommandListHandle.FlushResourceBarriers();
-
-	const bool bIsUpdate = false;
-	Geometry->BuildAccelerationStructure(*this, bIsUpdate);
-}
+	// Update geometry vertex buffers
 
-void FD3D12CommandContext::RHIUpdateAccelerationStructures(const TArrayView<const FAccelerationStructureUpdateParams> Params)
-{
-	// First batch up all barriers
-	for (const FAccelerationStructureUpdateParams P : Params)
+	for (const FAccelerationStructureBuildParams& P : Params)
 	{
 		FD3D12RayTracingGeometry* Geometry = FD3D12DynamicRHI::ResourceCast(P.Geometry.GetReference());
-		Geometry->RHIVertexBuffer = P.VertexBuffer;
-		Geometry->TransitionBuffers(*this);
-	}
-	CommandListHandle.FlushResourceBarriers();
 
-	// Then do all work
-	for (const FAccelerationStructureUpdateParams P : Params)
-	{
-		FD3D12RayTracingGeometry* Geometry = FD3D12DynamicRHI::ResourceCast(P.Geometry.GetReference());
-		Geometry->SetDirty(GetGPUMask(), true);
+		if (P.Segments.Num())
+		{
+			checkf(P.Segments.Num() == Geometry->Segments.Num(),
+				TEXT("If updated segments are provided, they must exactly match existing geometry segments. Only vertex buffer bindings may change."));
 
-		const bool bIsUpdate = !GRayTracingDebugForceFullBuild;
-		Geometry->BuildAccelerationStructure(*this, bIsUpdate);
+			for (int32 i = 0; i < P.Segments.Num(); ++i)
+			{
+				Geometry->Segments[i].VertexBuffer            = P.Segments[i].VertexBuffer;
+				Geometry->Segments[i].VertexBufferElementType = P.Segments[i].VertexBufferElementType;
+				Geometry->Segments[i].VertexBufferStride      = P.Segments[i].VertexBufferStride;
+				Geometry->Segments[i].VertexBufferOffset      = P.Segments[i].VertexBufferOffset;
+			}
+		}
 	}
-}
 
-void FD3D12CommandContext::RHIBuildAccelerationStructures(const TArrayView<const FAccelerationStructureUpdateParams> Params)
-{
-	// First batch up all barriers
-	for (const FAccelerationStructureUpdateParams P : Params)
+	// Transition all VBs and IBs to readable state
+
+	for (const FAccelerationStructureBuildParams& P : Params)
 	{
 		FD3D12RayTracingGeometry* Geometry = FD3D12DynamicRHI::ResourceCast(P.Geometry.GetReference());
-		Geometry->RHIVertexBuffer = P.VertexBuffer;
 		Geometry->TransitionBuffers(*this);
 	}
+
 	CommandListHandle.FlushResourceBarriers();
 
 	// Then do all work
-	for (const FAccelerationStructureUpdateParams P : Params)
+
+	for (const FAccelerationStructureBuildParams& P : Params)
 	{
 		FD3D12RayTracingGeometry* Geometry = FD3D12DynamicRHI::ResourceCast(P.Geometry.GetReference());
 		Geometry->SetDirty(GetGPUMask(), true);
 
-		const bool bIsUpdate = false;
-		Geometry->BuildAccelerationStructure(*this, bIsUpdate);
+		Geometry->BuildAccelerationStructure(*this, P.BuildMode);
 	}
 }
 
diff --git a/Engine/Source/Runtime/D3D12RHI/Private/D3D12RayTracing.h b/Engine/Source/Runtime/D3D12RHI/Private/D3D12RayTracing.h
--- a/Engine/Source/Runtime/D3D12RHI/Private/D3D12RayTracing.h
+++ b/Engine/Source/Runtime/D3D12RHI/Private/D3D12RayTracing.h
@@ -19,7 +19,7 @@ public:
 	~FD3D12RayTracingGeometry();
 
 	void TransitionBuffers(FD3D12CommandContext& CommandContext);
-	void BuildAccelerationStructure(FD3D12CommandContext& CommandContext, bool bIsUpdate);
+	void BuildAccelerationStructure(FD3D12CommandContext& CommandContext, EAccelerationStructureBuildMode BuildMode);
 
 	bool bIsAccelerationStructureDirty[MAX_NUM_GPUS] = {};
 	void SetDirty(FRHIGPUMask GPUMask, bool bState)
@@ -36,8 +36,6 @@ public:
 
 	uint32 IndexStride = 0; // 0 for non-indexed / implicit triangle list, 2 for uint16, 4 for uint32
 	uint32 IndexOffsetInBytes = 0;
-	uint32 VertexOffsetInBytes = 0;
-	uint32 VertexStrideInBytes = 0;
 	uint32 TotalPrimitiveCount = 0; // Combined number of primitives in all mesh segments
 
 	D3D12_RAYTRACING_GEOMETRY_TYPE GeometryType = D3D12_RAYTRACING_GEOMETRY_TYPE_TRIANGLES;
@@ -45,10 +43,7 @@ public:
 	TArray<FRayTracingGeometrySegment> Segments; // Defines addressable parts of the mesh that can be used for material assignment (one segment = one SBT record)
 	D3D12_RAYTRACING_ACCELERATION_STRUCTURE_BUILD_FLAGS BuildFlags;
 
-	EVertexElementType VertexElemType;
-
 	FIndexBufferRHIRef  RHIIndexBuffer;
-	FVertexBufferRHIRef RHIVertexBuffer;
 
 	TRefCountPtr<FD3D12MemBuffer> AccelerationStructureBuffers[MAX_NUM_GPUS];
 	TRefCountPtr<FD3D12MemBuffer> ScratchBuffers[MAX_NUM_GPUS];
diff --git a/Engine/Source/Runtime/Engine/Private/GPUSkinCache.cpp b/Engine/Source/Runtime/Engine/Private/GPUSkinCache.cpp
--- a/Engine/Source/Runtime/Engine/Private/GPUSkinCache.cpp
+++ b/Engine/Source/Runtime/Engine/Private/GPUSkinCache.cpp
@@ -361,25 +361,23 @@ public:
 		}
 	}
 
-	void CreateMergedPositionVertexBuffer(FRHICommandList& RHICmdList, FVertexBufferRHIRef OutVertexBuffer)
+#if RHI_RAYTRACING
+	void GetRayTracingSegmentVertexBuffers(TArrayView<FRayTracingGeometrySegment> OutSegments) const
 	{
-	#if RHI_RAYTRACING
-		TArray<FCopyBufferRegionParams> CopyParams;
+		check(OutSegments.Num() == DispatchData.Num());
 
 		for (int32 SectionIdx = 0; SectionIdx < DispatchData.Num(); SectionIdx++)
 		{
-			CopyParams.Add(FCopyBufferRegionParams {
-				OutVertexBuffer,
-				DispatchData[SectionIdx].OutputStreamStart * sizeof(FVector),
-				DispatchData[SectionIdx].PositionBuffer->Buffer,
-				DispatchData[SectionIdx].OutputStreamStart * sizeof(FVector),
-				DispatchData[SectionIdx].NumVertices * sizeof(FVector) }
-			);
-		}
+			const FSectionDispatchData& SectionData = DispatchData[SectionIdx];
+			FRayTracingGeometrySegment& Segment = OutSegments[SectionIdx];
 
-		RHICmdList.CopyBufferRegions(CopyParams);
-	#endif
+			Segment.VertexBuffer = SectionData.PositionBuffer->Buffer;
+			Segment.VertexBufferOffset = 0;
+
+			check(SectionData.Section->NumTriangles == Segment.NumPrimitives);
+		}
 	}
+#endif // RHI_RAYTRACING
 
 protected:
 	FGPUSkinCache::FRWBuffersAllocation* PositionAllocation;
@@ -629,16 +627,17 @@ void FGPUSkinCache::CommitRayTracingGeometryUpdates(FRHICommandList& RHICmdList)
 
 	if (RayTracingGeometriesToUpdate.Num())
 	{
-		TArray<FAccelerationStructureUpdateParams> Updates;
-		for (const auto& RayTracingGeometry : RayTracingGeometriesToUpdate)
+		TArray<FAccelerationStructureBuildParams> Updates;
+		for (const FRayTracingGeometry* RayTracingGeometry : RayTracingGeometriesToUpdate)
 		{
-			FAccelerationStructureUpdateParams Params;
+			FAccelerationStructureBuildParams Params;
+			Params.BuildMode = EAccelerationStructureBuildMode::Update;
 			Params.Geometry = RayTracingGeometry->RayTracingGeometryRHI;
-			Params.VertexBuffer = RayTracingGeometry->Initializer.PositionVertexBuffer;
+			Params.Segments = RayTracingGeometry->Initializer.Segments;
 			Updates.Add(Params);
 		}
 
-		RHICmdList.UpdateAccelerationStructures(Updates);
+		RHICmdList.BuildAccelerationStructures(Updates);
 		RayTracingGeometriesToUpdate.Reset();
 	}
 }
@@ -1379,10 +1378,12 @@ void FGPUSkinCache::ReleaseSkinCacheEntry(FGPUSkinCacheEntry* SkinCacheEntry)
 	delete SkinCacheEntry;
 }
 
-void FGPUSkinCache::CreateMergedPositionVertexBuffer(FRHICommandList& RHICmdList, FGPUSkinCacheEntry* SkinCacheEntry, FVertexBufferRHIRef OutVertexBuffer)
+#if RHI_RAYTRACING
+void FGPUSkinCache::GetRayTracingSegmentVertexBuffers(const FGPUSkinCacheEntry& SkinCacheEntry, TArrayView<FRayTracingGeometrySegment> OutSegments)
 {
-	SkinCacheEntry->CreateMergedPositionVertexBuffer(RHICmdList, OutVertexBuffer);
+	SkinCacheEntry.GetRayTracingSegmentVertexBuffers(OutSegments);
 }
+#endif // RHI_RAYTRACING
 
 bool FGPUSkinCache::IsEntryValid(FGPUSkinCacheEntry* SkinCacheEntry, int32 Section)
 {
diff --git a/Engine/Source/Runtime/Engine/Private/SkeletalMesh.cpp b/Engine/Source/Runtime/Engine/Private/SkeletalMesh.cpp
--- a/Engine/Source/Runtime/Engine/Private/SkeletalMesh.cpp
+++ b/Engine/Source/Runtime/Engine/Private/SkeletalMesh.cpp
@@ -4781,7 +4781,6 @@ void FSkeletalMeshSceneProxy::GetDynamicRayTracingInstances(FRayTracingMaterialG
 		// #dxr: the only case where RayTracingGeometryRHI is invalid is the very first frame - if that's not the case we have a bug somewhere else
 		if (MeshObject->GetRayTracingGeometry()->RayTracingGeometryRHI.IsValid())
 		{
-			check(MeshObject->GetRayTracingGeometry()->Initializer.PositionVertexBuffer.IsValid());
 			check(MeshObject->GetRayTracingGeometry()->Initializer.IndexBuffer.IsValid());
 			
 			FRayTracingInstance RayTracingInstance;
diff --git a/Engine/Source/Runtime/Engine/Private/SkeletalRenderCPUSkin.cpp b/Engine/Source/Runtime/Engine/Private/SkeletalRenderCPUSkin.cpp
--- a/Engine/Source/Runtime/Engine/Private/SkeletalRenderCPUSkin.cpp
+++ b/Engine/Source/Runtime/Engine/Private/SkeletalRenderCPUSkin.cpp
@@ -380,12 +380,8 @@ void FSkeletalMeshObjectCPUSkin::FSkeletalMeshObjectLOD::InitResources(FSkelMesh
 			[this, VertexBufferRHI, IndexBufferRHI, VertexBufferStride, TrianglesCount, RenderSections](FRHICommandListImmediate& RHICmdList)
 			{
 				FRayTracingGeometryInitializer Initializer;
-				Initializer.PositionVertexBuffer = VertexBufferRHI;
 				Initializer.IndexBuffer = IndexBufferRHI;
-				Initializer.VertexBufferStride = VertexBufferStride;
-				Initializer.VertexBufferByteOffset = 0;
 				Initializer.TotalPrimitiveCount = TrianglesCount;
-				Initializer.VertexBufferElementType = VET_Float3;
 				Initializer.GeometryType = RTGT_Triangles;
 				Initializer.bFastBuild = true;
 
@@ -394,6 +390,10 @@ void FSkeletalMeshObjectCPUSkin::FSkeletalMeshObjectLOD::InitResources(FSkelMesh
 				for (const FSkelMeshRenderSection& Section : *RenderSections)
 				{
 					FRayTracingGeometrySegment Segment;
+					Segment.VertexBuffer = VertexBufferRHI;
+					Segment.VertexBufferStride = VertexBufferStride;
+					Segment.VertexBufferOffset = 0;
+					Segment.VertexBufferElementType = VET_Float3;
 					Segment.FirstPrimitive = Section.BaseIndex / 3;
 					Segment.NumPrimitives = Section.NumTriangles;
 					GeometrySections.Add(Segment);
diff --git a/Engine/Source/Runtime/Engine/Private/SkeletalRenderGPUSkin.cpp b/Engine/Source/Runtime/Engine/Private/SkeletalRenderGPUSkin.cpp
--- a/Engine/Source/Runtime/Engine/Private/SkeletalRenderGPUSkin.cpp
+++ b/Engine/Source/Runtime/Engine/Private/SkeletalRenderGPUSkin.cpp
@@ -357,10 +357,21 @@ void FSkeletalMeshObjectGPUSkin::UpdateDynamicData_RenderThread(FGPUSkinCache* G
 		: true);
 
 #if RHI_RAYTRACING
-	bRequireRecreatingRayTracingGeometry = (DynamicData == nullptr || // Newly created
-		RayTracingGeometry.Initializer.PositionVertexBuffer == nullptr ||
+	bRequireRecreatingRayTracingGeometry = (DynamicData == nullptr || RayTracingGeometry.Initializer.Segments.Num() == 0 || // Newly created
 		(DynamicData != nullptr && DynamicData->LODIndex != InDynamicData->LODIndex)); // LOD level changed
-#endif
+	
+	if (!bRequireRecreatingRayTracingGeometry)
+	{
+		for (FRayTracingGeometrySegment& Segment : RayTracingGeometry.Initializer.Segments)
+		{
+			if (Segment.VertexBuffer == nullptr)
+			{
+				bRequireRecreatingRayTracingGeometry = true;
+				break;
+			}
+		}
+	}
+#endif // RHI_RAYTRACING
 
 	WaitForRHIThreadFenceForDynamicData();
 	if (DynamicData)
@@ -403,27 +414,27 @@ void FSkeletalMeshObjectGPUSkin::UpdateDynamicData_RenderThread(FGPUSkinCache* G
 				FRayTracingGeometryInitializer Initializer;
 
 				FRHIResourceCreateInfo CreateInfo;
-				Initializer.PositionVertexBuffer = RHICreateVertexBuffer(LODModel.GetNumVertices() * sizeof(FVector), BUF_ShaderResource, CreateInfo);
-				FGPUSkinCache::CreateMergedPositionVertexBuffer(RHICmdList, SkinCacheEntry, Initializer.PositionVertexBuffer);
+
 				Initializer.IndexBuffer = IndexBufferRHI;
-				Initializer.VertexBufferStride = VertexBufferStride;
-				Initializer.VertexBufferByteOffset = 0;
 				Initializer.TotalPrimitiveCount = TrianglesCount;
-				Initializer.VertexBufferElementType = VET_Float3;
 				Initializer.GeometryType = RTGT_Triangles;
 				Initializer.bFastBuild = true;
 				Initializer.bAllowUpdate = true;
 
-				TArray<FRayTracingGeometrySegment> GeometrySections;
-				GeometrySections.Reserve(LODModel.RenderSections.Num());
+				Initializer.Segments.Reserve(LODModel.RenderSections.Num());
 				for (const FSkelMeshRenderSection& Section : LODModel.RenderSections)
 				{
 					FRayTracingGeometrySegment Segment;
+					Segment.VertexBuffer = nullptr;
+					Segment.VertexBufferElementType = VET_Float3;
+					Segment.VertexBufferStride = VertexBufferStride;
+					Segment.VertexBufferOffset = 0;
 					Segment.FirstPrimitive = Section.BaseIndex / 3;
 					Segment.NumPrimitives = Section.NumTriangles;
-					GeometrySections.Add(Segment);
+					Initializer.Segments.Add(Segment);
 				}
-				Initializer.Segments = GeometrySections;
+
+				FGPUSkinCache::GetRayTracingSegmentVertexBuffers(*SkinCacheEntry, Initializer.Segments);
 
 				RayTracingGeometry.SetInitializer(Initializer);
 				RayTracingGeometry.UpdateRHI();
@@ -434,7 +445,7 @@ void FSkeletalMeshObjectGPUSkin::UpdateDynamicData_RenderThread(FGPUSkinCache* G
 				if (!DynamicData->bAnySegmentUsesWorldPositionOffset)
 				{
 					// Refit BLAS with new vertex buffer data
-					FGPUSkinCache::CreateMergedPositionVertexBuffer(RHICmdList, SkinCacheEntry, RayTracingGeometry.Initializer.PositionVertexBuffer);
+					FGPUSkinCache::GetRayTracingSegmentVertexBuffers(*SkinCacheEntry, RayTracingGeometry.Initializer.Segments);
 					GPUSkinCache->AddRayTracingGeometryToUpdate(&RayTracingGeometry);
 				}
 				else
diff --git a/Engine/Source/Runtime/Engine/Private/SkeletalRenderStatic.cpp b/Engine/Source/Runtime/Engine/Private/SkeletalRenderStatic.cpp
--- a/Engine/Source/Runtime/Engine/Private/SkeletalRenderStatic.cpp
+++ b/Engine/Source/Runtime/Engine/Private/SkeletalRenderStatic.cpp
@@ -143,12 +143,8 @@ void FSkeletalMeshObjectStatic::FSkeletalMeshObjectLOD::InitResources(FSkelMeshC
 			[this, VertexBufferRHI, IndexBufferRHI, VertexBufferStride, TrianglesCount, RenderSections](FRHICommandListImmediate& RHICmdList)
 			{
 				FRayTracingGeometryInitializer Initializer;
-				Initializer.PositionVertexBuffer = VertexBufferRHI;
 				Initializer.IndexBuffer = IndexBufferRHI;
-				Initializer.VertexBufferStride = VertexBufferStride;
-				Initializer.VertexBufferByteOffset = 0;
 				Initializer.TotalPrimitiveCount = TrianglesCount;
-				Initializer.VertexBufferElementType = VET_Float3;
 				Initializer.GeometryType = RTGT_Triangles;
 				Initializer.bFastBuild = false;
 
@@ -157,6 +153,10 @@ void FSkeletalMeshObjectStatic::FSkeletalMeshObjectLOD::InitResources(FSkelMeshC
 				for (const FSkelMeshRenderSection& Section : *RenderSections)
 				{
 					FRayTracingGeometrySegment Segment;
+					Segment.VertexBuffer = VertexBufferRHI;
+					Segment.VertexBufferElementType = VET_Float3;
+					Segment.VertexBufferOffset = 0;
+					Segment.VertexBufferStride = VertexBufferStride;
 					Segment.FirstPrimitive = Section.BaseIndex / 3;
 					Segment.NumPrimitives = Section.NumTriangles;
 					GeometrySections.Add(Segment);
diff --git a/Engine/Source/Runtime/Engine/Private/StaticMesh.cpp b/Engine/Source/Runtime/Engine/Private/StaticMesh.cpp
--- a/Engine/Source/Runtime/Engine/Private/StaticMesh.cpp
+++ b/Engine/Source/Runtime/Engine/Private/StaticMesh.cpp
@@ -1179,12 +1179,9 @@ void FStaticMeshLODResources::InitResources(UStaticMesh* Parent)
 			[this](FRHICommandListImmediate& RHICmdList)
 			{
 				FRayTracingGeometryInitializer Initializer;
-				Initializer.PositionVertexBuffer = VertexBuffers.PositionVertexBuffer.VertexBufferRHI;
+				
 				Initializer.IndexBuffer = IndexBuffer.IndexBufferRHI;
-				Initializer.VertexBufferStride = VertexBuffers.PositionVertexBuffer.GetStride();
-				Initializer.VertexBufferByteOffset = 0;
 				Initializer.TotalPrimitiveCount = 0; // This is calculated below based on static mesh section data
-				Initializer.VertexBufferElementType = VET_Float3;
 				Initializer.GeometryType = RTGT_Triangles;
 				Initializer.bFastBuild = false;
 				
@@ -1193,6 +1190,10 @@ void FStaticMeshLODResources::InitResources(UStaticMesh* Parent)
 				for (const FStaticMeshSection& Section : Sections)
 				{
 					FRayTracingGeometrySegment Segment;
+					Segment.VertexBuffer = VertexBuffers.PositionVertexBuffer.VertexBufferRHI;
+					Segment.VertexBufferElementType = VET_Float3;
+					Segment.VertexBufferStride = VertexBuffers.PositionVertexBuffer.GetStride();
+					Segment.VertexBufferOffset = 0;
 					Segment.FirstPrimitive = Section.FirstIndex / 3;
 					Segment.NumPrimitives = Section.NumTriangles;
 					GeometrySections.Add(Segment);
diff --git a/Engine/Source/Runtime/Engine/Private/StaticMeshRender.cpp b/Engine/Source/Runtime/Engine/Private/StaticMeshRender.cpp
--- a/Engine/Source/Runtime/Engine/Private/StaticMeshRender.cpp
+++ b/Engine/Source/Runtime/Engine/Private/StaticMeshRender.cpp
@@ -204,7 +204,10 @@ FStaticMeshSceneProxy::FStaticMeshSceneProxy(UStaticMeshComponent* InComponent,
 			{
 				auto& Initializer = DynamicRayTracingGeometries[LODIndex].Initializer;
 				Initializer = RenderData->LODResources[LODIndex].RayTracingGeometry.Initializer;
-				Initializer.PositionVertexBuffer = nullptr;
+				for (FRayTracingGeometrySegment& Segment : Initializer.Segments)
+				{
+					Segment.VertexBuffer = nullptr;
+				}
 				Initializer.bAllowUpdate = true;
 				Initializer.bFastBuild = true;
 			}
diff --git a/Engine/Source/Runtime/Engine/Public/GPUSkinCache.h b/Engine/Source/Runtime/Engine/Public/GPUSkinCache.h
--- a/Engine/Source/Runtime/Engine/Public/GPUSkinCache.h
+++ b/Engine/Source/Runtime/Engine/Public/GPUSkinCache.h
@@ -49,6 +49,7 @@ class FSkeletalMeshVertexClothBuffer;
 struct FClothSimulData;
 struct FSkelMeshRenderSection;
 struct FVertexBufferAndSRV;
+struct FRayTracingGeometrySegment;
 
 // Can the skin cache be used (ie shaders added, etc)
 extern ENGINE_API bool IsGPUSkinCacheAvailable(EShaderPlatform Platform);
@@ -136,7 +137,9 @@ public:
 		return nullptr;
 	}
 
-	static void CreateMergedPositionVertexBuffer(FRHICommandList& RHICmdList, FGPUSkinCacheEntry* SkinCacheEntry, FVertexBufferRHIRef OutVertexBuffer);
+#if RHI_RAYTRACING
+	static void GetRayTracingSegmentVertexBuffers(const FGPUSkinCacheEntry& SkinCacheEntry, TArrayView<FRayTracingGeometrySegment> OutSegments);
+#endif // RHI_RAYTRACING
 
 	static bool IsEntryValid(FGPUSkinCacheEntry* SkinCacheEntry, int32 Section);
 
diff --git a/Engine/Source/Runtime/Landscape/Private/LandscapeRender.cpp b/Engine/Source/Runtime/Landscape/Private/LandscapeRender.cpp
--- a/Engine/Source/Runtime/Landscape/Private/LandscapeRender.cpp
+++ b/Engine/Source/Runtime/Landscape/Private/LandscapeRender.cpp
@@ -1465,15 +1465,17 @@ void FLandscapeComponentSceneProxy::CreateRenderThreadResources()
 
 				FRayTracingGeometryInitializer Initializer;
 				FRHIResourceCreateInfo CreateInfo;
-				Initializer.PositionVertexBuffer = RHICreateVertexBuffer(sizeof(FVector4) * NumPrimitives * 3, BUF_UnorderedAccess | BUF_ShaderResource, CreateInfo);
 				Initializer.IndexBuffer = nullptr;
-				Initializer.VertexBufferStride = sizeof(FVector);
 				Initializer.TotalPrimitiveCount = NumPrimitives;
-				Initializer.VertexBufferElementType = VET_Float3;
 				Initializer.GeometryType = RTGT_Triangles;
 				Initializer.bFastBuild = true;
 				Initializer.bAllowUpdate = true;
-
+				FRayTracingGeometrySegment Segment;
+				Segment.VertexBuffer = RHICreateVertexBuffer(sizeof(FVector4) * NumPrimitives * 3, BUF_UnorderedAccess | BUF_ShaderResource, CreateInfo);
+				Segment.VertexBufferStride = sizeof(FVector);
+				Segment.VertexBufferElementType = VET_Float3;
+				Segment.NumPrimitives = NumPrimitives;
+				Initializer.Segments.Add(Segment);
 				SectionRayTracingStates[SubSectionIdx].Geometry.SetInitializer(Initializer);
 				SectionRayTracingStates[SubSectionIdx].Geometry.InitResource();
 			}
diff --git a/Engine/Source/Runtime/RHI/Public/DynamicRHI.h b/Engine/Source/Runtime/RHI/Public/DynamicRHI.h
--- a/Engine/Source/Runtime/RHI/Public/DynamicRHI.h
+++ b/Engine/Source/Runtime/RHI/Public/DynamicRHI.h
@@ -107,6 +107,17 @@ enum ERayTracingGeometryType
 
 struct FRayTracingGeometrySegment
 {
+	FVertexBufferRHIRef VertexBuffer = nullptr;
+	EVertexElementType VertexBufferElementType = VET_Float3;
+
+	// Offset in bytes from the base address of the vertex buffer.
+	uint32 VertexBufferOffset = 0;
+
+	// Number of bytes between elements of the vertex buffer (sizeof VET_Float3 by default).
+	// Must be equal or greater than the size of the position vector.
+	uint32 VertexBufferStride = 12;
+
+	// Primitive range for this segment.
 	uint32 FirstPrimitive = 0;
 	uint32 NumPrimitives = 0;
 
@@ -121,17 +132,20 @@ struct FRayTracingGeometrySegment
 
 struct FRayTracingGeometryInitializer
 {
-	FVertexBufferRHIRef PositionVertexBuffer = nullptr;
-	uint32 VertexBufferByteOffset = 0;
-	uint32 VertexBufferStride = 0;
-	EVertexElementType VertexBufferElementType = VET_Float3;
-
 	FIndexBufferRHIRef IndexBuffer = nullptr;
-	uint32 IndexBufferByteOffset = 0;
+
+	// Offset in bytes from the base address of the index buffer.
+	uint32 IndexBufferOffset = 0;
 
 	ERayTracingGeometryType GeometryType = RTGT_Triangles;
+
+	// Total number of primitives in all segments of the geometry. Only used for validation.
 	uint32 TotalPrimitiveCount = 0;
+
+	// Partitions of geometry to allow different shader and resource bindings.
+	// All ray tracing geometries must have at least one segment.
 	TArray<FRayTracingGeometrySegment> Segments;
+
 	bool bFastBuild = false;
 	bool bAllowUpdate = false;
 };
diff --git a/Engine/Source/Runtime/RHI/Public/RHICommandList.h b/Engine/Source/Runtime/RHI/Public/RHICommandList.h
--- a/Engine/Source/Runtime/RHI/Public/RHICommandList.h
+++ b/Engine/Source/Runtime/RHI/Public/RHICommandList.h
@@ -2045,17 +2045,10 @@ FRHICOMMAND_MACRO(FRHICommandCopyBufferRegions)
 template <ECmdList CmdListType>
 struct FRHICommandBuildAccelerationStructure final : public FRHICommand<FRHICommandBuildAccelerationStructure<CmdListType>>
 {
-	FRHIRayTracingGeometry* Geometry;
 	FRHIRayTracingScene* Scene;
 
-	explicit FRHICommandBuildAccelerationStructure(FRHIRayTracingGeometry* InGeometry)
-		: Geometry(InGeometry)
-		, Scene(nullptr)
-	{}
-
 	explicit FRHICommandBuildAccelerationStructure(FRHIRayTracingScene* InScene)
-		: Geometry(nullptr)
-		, Scene(InScene)
+		: Scene(InScene)
 	{}
 
 	RHI_API void Execute(FRHICommandListBase& CmdList);
@@ -2072,25 +2065,13 @@ FRHICOMMAND_MACRO(FRHICommandClearRayTracingBindings)
 	RHI_API void Execute(FRHICommandListBase& CmdList);
 };
 
-template <ECmdList CmdListType>
-struct FRHICommandUpdateAccelerationStructures final : public FRHICommand<FRHICommandUpdateAccelerationStructures<CmdListType>>
-{
-	const TArrayView<const FAccelerationStructureUpdateParams> UpdateParams;
-
-	explicit FRHICommandUpdateAccelerationStructures(const TArrayView<const FAccelerationStructureUpdateParams> InParams)
-		: UpdateParams(InParams)
-	{}
-
-	RHI_API void Execute(FRHICommandListBase& CmdList);
-};
-
 template <ECmdList CmdListType>
 struct FRHICommandBuildAccelerationStructures final : public FRHICommand<FRHICommandBuildAccelerationStructures<CmdListType>>
 {
-	const TArrayView<const FAccelerationStructureUpdateParams> UpdateParams;
+	const TArrayView<const FAccelerationStructureBuildParams> Params;
 
-	explicit FRHICommandBuildAccelerationStructures(const TArrayView<const FAccelerationStructureUpdateParams> InParams)
-		: UpdateParams(InParams)
+	explicit FRHICommandBuildAccelerationStructures(const TArrayView<const FAccelerationStructureBuildParams> InParams)
+		: Params(InParams)
 	{}
 
 	RHI_API void Execute(FRHICommandListBase& CmdList);
@@ -3194,6 +3175,7 @@ public:
 
 #if RHI_RAYTRACING
 	// Ray tracing API
+	UE_DEPRECATED(4.25, "CopyBufferRegion API is deprecated. Use an explicit compute shader copy dispatch instead.")
 	FORCEINLINE_DEBUGGABLE void CopyBufferRegion(FRHIVertexBuffer* DestBuffer, uint64 DstOffset, FRHIVertexBuffer* SourceBuffer, uint64 SrcOffset, uint64 NumBytes)
 	{
 		// No copy/DMA operation inside render passes
@@ -3209,6 +3191,7 @@ public:
 		}
 	}
 
+	UE_DEPRECATED(4.25, "CopyBufferRegions API is deprecated. Use an explicit compute shader copy dispatch instead.")
 	FORCEINLINE_DEBUGGABLE void CopyBufferRegions(const TArrayView<const FCopyBufferRegionParams> Params)
 	{
 		// No copy/DMA operation inside render passes
@@ -3224,40 +3207,24 @@ public:
 		}
 	}
 
-	FORCEINLINE_DEBUGGABLE void BuildAccelerationStructure(FRHIRayTracingGeometry* Geometry)
-	{
-		if (Bypass())
-		{
-			GetContext().RHIBuildAccelerationStructure(Geometry);
-		}
-		else
-		{
-			ALLOC_COMMAND(FRHICommandBuildAccelerationStructure<ECmdList::EGfx>)(Geometry);
-		}
-	}
-
-	FORCEINLINE_DEBUGGABLE void UpdateAccelerationStructures(const TArrayView<const FAccelerationStructureUpdateParams> Params)
+	FORCEINLINE_DEBUGGABLE void BuildAccelerationStructures(const TArrayView<const FAccelerationStructureBuildParams> Params)
 	{
 		if (Bypass())
 		{
-			GetContext().RHIUpdateAccelerationStructures(Params);
+			GetContext().RHIBuildAccelerationStructures(Params);
 		}
 		else
 		{
-			ALLOC_COMMAND(FRHICommandUpdateAccelerationStructures<ECmdList::EGfx>)(AllocArray(Params));
+			ALLOC_COMMAND(FRHICommandBuildAccelerationStructures<ECmdList::EGfx>)(AllocArray(Params));
 		}
 	}
 
-	FORCEINLINE_DEBUGGABLE void BuildAccelerationStructures(const TArrayView<const FAccelerationStructureUpdateParams> Params)
+	FORCEINLINE_DEBUGGABLE void BuildAccelerationStructure(FRHIRayTracingGeometry* Geometry)
 	{
-		if (Bypass())
-		{
-			GetContext().RHIBuildAccelerationStructures(Params);
-		}
-		else
-		{
-			ALLOC_COMMAND(FRHICommandBuildAccelerationStructures<ECmdList::EGfx>)(AllocArray(Params));
-		}
+		FAccelerationStructureBuildParams Params;
+		Params.Geometry = Geometry;
+		Params.BuildMode = EAccelerationStructureBuildMode::Build;
+		BuildAccelerationStructures(MakeArrayView(&Params, 1));
 	}
 
 	FORCEINLINE_DEBUGGABLE void BuildAccelerationStructure(FRHIRayTracingScene* Scene)
@@ -3713,29 +3680,13 @@ public:
 #if RHI_RAYTRACING
 	FORCEINLINE_DEBUGGABLE void BuildAccelerationStructure(FRHIRayTracingGeometry* Geometry)
 	{
-		if (Bypass())
-		{
-			GetComputeContext().RHIBuildAccelerationStructure(Geometry);
-		}
-		else
-		{
-			ALLOC_COMMAND(FRHICommandBuildAccelerationStructure<ECmdList::ECompute>)(Geometry);
-		}
-	}
-
-	FORCEINLINE_DEBUGGABLE void UpdateAccelerationStructures(const TArrayView<const FAccelerationStructureUpdateParams> Params)
-	{
-		if (Bypass())
-		{
-			GetComputeContext().RHIUpdateAccelerationStructures(Params);
-		}
-		else
-		{
-			ALLOC_COMMAND(FRHICommandUpdateAccelerationStructures<ECmdList::ECompute>)(AllocArray(Params));
-		}
+		FAccelerationStructureBuildParams Params;
+		Params.Geometry = Geometry;
+		Params.BuildMode = EAccelerationStructureBuildMode::Build;
+		BuildAccelerationStructures(MakeArrayView(&Params, 1));
 	}
 
-	FORCEINLINE_DEBUGGABLE void BuildAccelerationStructures(const TArrayView<const FAccelerationStructureUpdateParams> Params)
+	FORCEINLINE_DEBUGGABLE void BuildAccelerationStructures(const TArrayView<const FAccelerationStructureBuildParams> Params)
 	{
 		if (Bypass())
 		{
diff --git a/Engine/Source/Runtime/RHI/Public/RHICommandListCommandExecutes.inl b/Engine/Source/Runtime/RHI/Public/RHICommandListCommandExecutes.inl
--- a/Engine/Source/Runtime/RHI/Public/RHICommandListCommandExecutes.inl
+++ b/Engine/Source/Runtime/RHI/Public/RHICommandListCommandExecutes.inl
@@ -568,14 +568,7 @@ template<ECmdList CmdListType>
 void FRHICommandBuildAccelerationStructure<CmdListType>::Execute(FRHICommandListBase& CmdList)
 {
 	RHISTAT(BuildAccelerationStructure);
-	if (Geometry)
-	{
-		INTERNAL_DECORATOR_CONTEXT(RHIBuildAccelerationStructure)(Geometry);
-	}
-	else
-	{
-		INTERNAL_DECORATOR_CONTEXT(RHIBuildAccelerationStructure)(Scene);
-	}
+	INTERNAL_DECORATOR_CONTEXT(RHIBuildAccelerationStructure)(Scene);
 }
 
 template struct FRHICommandBuildAccelerationStructure<ECmdList::EGfx>;
@@ -587,21 +580,11 @@ void FRHICommandClearRayTracingBindings::Execute(FRHICommandListBase& CmdList)
 	INTERNAL_DECORATOR(RHIClearRayTracingBindings)(Scene);
 }
 
-template<ECmdList CmdListType>
-void FRHICommandUpdateAccelerationStructures<CmdListType>::Execute(FRHICommandListBase& CmdList)
-{
-	RHISTAT(UpdateAccelerationStructure);
-	INTERNAL_DECORATOR_CONTEXT(RHIUpdateAccelerationStructures)(UpdateParams);
-}
-
-template struct FRHICommandUpdateAccelerationStructures<ECmdList::EGfx>;
-template struct FRHICommandUpdateAccelerationStructures<ECmdList::ECompute>;
-
 template<ECmdList CmdListType>
 void FRHICommandBuildAccelerationStructures<CmdListType>::Execute(FRHICommandListBase& CmdList)
 {
 	RHISTAT(BuildAccelerationStructure);
-	INTERNAL_DECORATOR_CONTEXT(RHIBuildAccelerationStructures)(UpdateParams);
+	INTERNAL_DECORATOR_CONTEXT(RHIBuildAccelerationStructures)(Params);
 }
 
 template struct FRHICommandBuildAccelerationStructures<ECmdList::EGfx>;
diff --git a/Engine/Source/Runtime/RHI/Public/RHIContext.h b/Engine/Source/Runtime/RHI/Public/RHIContext.h
--- a/Engine/Source/Runtime/RHI/Public/RHIContext.h
+++ b/Engine/Source/Runtime/RHI/Public/RHIContext.h
@@ -24,7 +24,8 @@ struct FResolveParams;
 struct FViewportBounds;
 struct FRayTracingGeometryInstance;
 struct FRayTracingShaderBindings;
-struct FAccelerationStructureUpdateParams;
+struct FRayTracingGeometrySegment;
+struct FAccelerationStructureBuildParams;
 enum class EAsyncComputeBudget;
 enum class EResourceTransitionAccess;
 enum class EResourceTransitionPipeline;
@@ -163,12 +164,7 @@ public:
 		checkNoEntry();
 	}
 
-	virtual void RHIUpdateAccelerationStructures(const TArrayView<const FAccelerationStructureUpdateParams> Params)
-	{
-		checkNoEntry();
-	}
-
-	virtual void RHIBuildAccelerationStructures(const TArrayView<const FAccelerationStructureUpdateParams> Params)
+	virtual void RHIBuildAccelerationStructures(const TArrayView<const FAccelerationStructureBuildParams> Params)
 	{
 		checkNoEntry();
 	}
@@ -179,10 +175,25 @@ public:
 	}
 };
 
-struct FAccelerationStructureUpdateParams
+enum class EAccelerationStructureBuildMode
+{
+	// Perform a full acceleration structure build.
+	Build,
+
+	// Update existing acceleration structure, based on new vertex positions.
+	// Index buffer must not change between initial build and update operations.
+	// Only valid when geometry was created with FRayTracingGeometryInitializer::bAllowUpdate = true.
+	Update,
+};
+
+struct FAccelerationStructureBuildParams
 {
 	FRayTracingGeometryRHIRef Geometry;
-	FVertexBufferRHIRef VertexBuffer;
+	EAccelerationStructureBuildMode BuildMode = EAccelerationStructureBuildMode::Build;
+
+	// Optional array of geometry segments that can be used to change per-segment vertex buffers.
+	// Only fields related to vertex buffer are used. If empty, then geometry vertex buffers are not changed.
+	TArrayView<const FRayTracingGeometrySegment> Segments;
 };
 
 struct FCopyBufferRegionParams
@@ -649,19 +660,18 @@ public:
 		checkNoEntry();
 	}
 
-	virtual void RHIBuildAccelerationStructure(FRHIRayTracingGeometry* Geometry)
+	virtual void RHIBuildAccelerationStructures(const TArrayView<const FAccelerationStructureBuildParams> Params)
 	{
 		checkNoEntry();
 	}
 
-	virtual void RHIUpdateAccelerationStructures(const TArrayView<const FAccelerationStructureUpdateParams> Params)
+	virtual void RHIBuildAccelerationStructure(FRHIRayTracingGeometry* Geometry) final override
 	{
-		checkNoEntry();
-	}
+		FAccelerationStructureBuildParams Params;
+		Params.Geometry = Geometry;
+		Params.BuildMode = EAccelerationStructureBuildMode::Build;
 
-	virtual void RHIBuildAccelerationStructures(const TArrayView<const FAccelerationStructureUpdateParams> Params)
-	{
-		checkNoEntry();
+		RHIBuildAccelerationStructures(MakeArrayView(&Params, 1));
 	}
 
 	virtual void RHIBuildAccelerationStructure(FRHIRayTracingScene* Scene)
diff --git a/Engine/Source/Runtime/RHI/Public/RHIValidationContext.h b/Engine/Source/Runtime/RHI/Public/RHIValidationContext.h
--- a/Engine/Source/Runtime/RHI/Public/RHIValidationContext.h
+++ b/Engine/Source/Runtime/RHI/Public/RHIValidationContext.h
@@ -804,17 +804,7 @@ public:
 		RHIContext->RHICopyTexture(SourceTexture, DestTexture, CopyInfo);
 	}
 
-	virtual void RHIBuildAccelerationStructure(FRHIRayTracingGeometry* Geometry) override final
-	{
-		RHIContext->RHIBuildAccelerationStructure(Geometry);
-	}
-
-	virtual void RHIUpdateAccelerationStructures(const TArrayView<const FAccelerationStructureUpdateParams> Params) override final
-	{
-		RHIContext->RHIUpdateAccelerationStructures(Params);
-	}
-
-	virtual void RHIBuildAccelerationStructures(const TArrayView<const FAccelerationStructureUpdateParams> Params) override final
+	virtual void RHIBuildAccelerationStructures(const TArrayView<const FAccelerationStructureBuildParams> Params) override final
 	{
 		RHIContext->RHIBuildAccelerationStructures(Params);
 	}
diff --git a/Engine/Source/Runtime/RenderCore/Public/RenderResource.h b/Engine/Source/Runtime/RenderCore/Public/RenderResource.h
--- a/Engine/Source/Runtime/RenderCore/Public/RenderResource.h
+++ b/Engine/Source/Runtime/RenderCore/Public/RenderResource.h
@@ -597,7 +597,7 @@ public:
 
 	virtual void InitRHI() override
 	{
-		if (Initializer.IndexBuffer && Initializer.PositionVertexBuffer && IsRayTracingEnabled())
+		if (Initializer.IndexBuffer && IsRayTracingEnabled())
 		{
 			RayTracingGeometryRHI = RHICreateRayTracingGeometry(Initializer);
 			FRHICommandListExecutor::GetImmediateCommandList().BuildAccelerationStructure(RayTracingGeometryRHI);
diff --git a/Engine/Source/Runtime/Renderer/Private/DeferredShadingRenderer.cpp b/Engine/Source/Runtime/Renderer/Private/DeferredShadingRenderer.cpp
--- a/Engine/Source/Runtime/Renderer/Private/DeferredShadingRenderer.cpp
+++ b/Engine/Source/Runtime/Renderer/Private/DeferredShadingRenderer.cpp
@@ -849,9 +849,6 @@ bool FDeferredShadingSceneRenderer::DispatchRayTracingWorldUpdates(FRHICommandLi
 	bool bAsyncUpdateGeometry = (CVarRayTracingAsyncBuild.GetValueOnRenderThread() != 0)
 							  && GSupportAsyncComputeRaytracingBuildBVH;
 
-	TArrayView<FAccelerationStructureUpdateParams> BuildParams, RefitParams;
-
-
 	if (!bAsyncUpdateGeometry)
 	{
 		SCOPED_GPU_STAT(RHICmdList, RayTracingAS);
diff --git a/Engine/Source/Runtime/Renderer/Private/RayTracing/RayTracingDynamicGeometry.cpp b/Engine/Source/Runtime/Renderer/Private/RayTracing/RayTracingDynamicGeometry.cpp
--- a/Engine/Source/Runtime/Renderer/Private/RayTracing/RayTracingDynamicGeometry.cpp
+++ b/Engine/Source/Runtime/Renderer/Private/RayTracing/RayTracingDynamicGeometry.cpp
@@ -182,18 +182,34 @@ void FRayTracingDynamicGeometryCollection::AddDynamicMeshBatchForGeometryUpdate(
 	}
 
 	check(Geometry.IsInitialized());
-	Geometry.Initializer.PositionVertexBuffer = Buffer.Buffer;
-	Geometry.Initializer.TotalPrimitiveCount = UpdateParams.NumTriangles;
 
-	if (!bRefit)
+	if (Geometry.Initializer.TotalPrimitiveCount != UpdateParams.NumTriangles)
 	{
-		Geometry.RayTracingGeometryRHI = RHICreateRayTracingGeometry(Geometry.Initializer);
-		BuildParams.Add(FAccelerationStructureUpdateParams{ Geometry.RayTracingGeometryRHI, Buffer.Buffer });
+		check(Geometry.Initializer.Segments.Num() <= 1);
+		Geometry.Initializer.TotalPrimitiveCount = UpdateParams.NumTriangles;
+		Geometry.Initializer.Segments.Empty();
+		FRayTracingGeometrySegment Segment;
+		Segment.NumPrimitives = UpdateParams.NumTriangles;
+		Geometry.Initializer.Segments.Add(Segment);
+		bRefit = false;
 	}
-	else
+
+	for (FRayTracingGeometrySegment& Segment : Geometry.Initializer.Segments)
 	{
-		RefitParams.Add(FAccelerationStructureUpdateParams{ Geometry.RayTracingGeometryRHI, Buffer.Buffer });
+		Segment.VertexBuffer = Buffer.Buffer;
 	}
+
+	if (!bRefit)
+	{
+		Geometry.RayTracingGeometryRHI = RHICreateRayTracingGeometry(Geometry.Initializer);
+	}
+
+	FAccelerationStructureBuildParams Params;
+	Params.Geometry = Geometry.RayTracingGeometryRHI;
+	Params.BuildMode = bRefit
+		? EAccelerationStructureBuildMode::Update
+		: EAccelerationStructureBuildMode::Build;
+	BuildParams.Add(Params);
 }
 
 
@@ -254,15 +270,8 @@ void FRayTracingDynamicGeometryCollection::DispatchUpdates(CmdListType& RHICmdLi
 			}
 		}
 
-		{
-			SCOPED_DRAW_OR_COMPUTE_EVENT(RHICmdList, Build)
-			RHICmdList.BuildAccelerationStructures(BuildParams);
-		}
-
-		{
-			SCOPED_DRAW_OR_COMPUTE_EVENT(RHICmdList, Refit)
-			RHICmdList.UpdateAccelerationStructures(RefitParams);
-		}
+		SCOPED_DRAW_OR_COMPUTE_EVENT(RHICmdList, Build);
+		RHICmdList.BuildAccelerationStructures(BuildParams);
 
 		Clear();
 	}
@@ -276,7 +285,6 @@ void FRayTracingDynamicGeometryCollection::Clear()
 {
 	DispatchCommands->Empty();
 	BuildParams.Empty();
-	RefitParams.Empty();
 }
 
 #endif // RHI_RAYTRACING
diff --git a/Engine/Source/Runtime/Renderer/Private/Tests/RayTracingTestbed.cpp b/Engine/Source/Runtime/Renderer/Private/Tests/RayTracingTestbed.cpp
--- a/Engine/Source/Runtime/Renderer/Private/Tests/RayTracingTestbed.cpp
+++ b/Engine/Source/Runtime/Renderer/Private/Tests/RayTracingTestbed.cpp
@@ -95,13 +95,12 @@ bool RunRayTracingTestbed_RenderThread(const FString& Parameters)
 
 	FRayTracingGeometryInitializer GeometryInitializer;
 	GeometryInitializer.IndexBuffer = IndexBuffer;
-	GeometryInitializer.PositionVertexBuffer = VertexBuffer;
-	GeometryInitializer.VertexBufferByteOffset = 0;
-	GeometryInitializer.VertexBufferStride = sizeof(FVector);
-	GeometryInitializer.VertexBufferElementType = VET_Float3;
 	GeometryInitializer.GeometryType = RTGT_Triangles;
-	GeometryInitializer.TotalPrimitiveCount = 1;
 	GeometryInitializer.bFastBuild = false;
+	FRayTracingGeometrySegment Segment;
+	Segment.VertexBuffer = VertexBuffer;
+	Segment.NumPrimitives = 1;
+	GeometryInitializer.Segments.Add(Segment);
 	FRayTracingGeometryRHIRef Geometry = RHICreateRayTracingGeometry(GeometryInitializer);
 
 	FRayTracingGeometryInstance Instances[] = {
diff --git a/Engine/Source/Runtime/Renderer/Public/RayTracingDynamicGeometryCollection.h b/Engine/Source/Runtime/Renderer/Public/RayTracingDynamicGeometryCollection.h
--- a/Engine/Source/Runtime/Renderer/Public/RayTracingDynamicGeometryCollection.h
+++ b/Engine/Source/Runtime/Renderer/Public/RayTracingDynamicGeometryCollection.h
@@ -24,8 +24,7 @@ public:
 
 private:
 	TUniquePtr<TArray<struct FMeshComputeDispatchCommand>> DispatchCommands;
-	TArray<FAccelerationStructureUpdateParams> BuildParams;
-	TArray<FAccelerationStructureUpdateParams> RefitParams;
+	TArray<FAccelerationStructureBuildParams> BuildParams;
 };
 
 #endif // RHI_RAYTRACING
