Optimize IndexBuffer state

diff --git a/Engine/Source/Runtime/D3D12RHI/Private/D3D12Commands.cpp b/Engine/Source/Runtime/D3D12RHI/Private/D3D12Commands.cpp
--- a/Engine/Source/Runtime/D3D12RHI/Private/D3D12Commands.cpp
+++ b/Engine/Source/Runtime/D3D12RHI/Private/D3D12Commands.cpp
@@ -1413,7 +1413,7 @@ void FD3D12CommandContext::RHIDrawIndexedIndirect(FRHIIndexBuffer* IndexBufferRH
 	// determine 16bit vs 32bit indices
 	const DXGI_FORMAT Format = (IndexBuffer->GetStride() == sizeof(uint16) ? DXGI_FORMAT_R16_UINT : DXGI_FORMAT_R32_UINT);
 
-	StateCache.SetIndexBuffer(&IndexBuffer->ResourceLocation, Format, 0);
+	StateCache.SetIndexBuffer(IndexBuffer->ResourceLocation, Format, 0);
 
 	FD3D12ResourceLocation& Location = ArgumentsBuffer->ResourceLocation;
 	FD3D12DynamicRHI::TransitionResource(CommandListHandle, Location.GetResource(), D3D12_RESOURCE_STATE_INDIRECT_ARGUMENT, D3D12_RESOURCE_BARRIER_ALL_SUBRESOURCES);
@@ -1463,7 +1463,7 @@ void FD3D12CommandContext::RHIDrawIndexedPrimitive(FRHIIndexBuffer* IndexBufferR
 
 	// determine 16bit vs 32bit indices
 	const DXGI_FORMAT Format = (IndexBuffer->GetStride() == sizeof(uint16) ? DXGI_FORMAT_R16_UINT : DXGI_FORMAT_R32_UINT);
-	StateCache.SetIndexBuffer(&IndexBuffer->ResourceLocation, Format, 0);
+	StateCache.SetIndexBuffer(IndexBuffer->ResourceLocation, Format, 0);
 	StateCache.ApplyState<D3D12PT_Graphics>();
 
 	CommandListHandle->DrawIndexedInstanced(IndexCount, NumInstances, StartIndex, BaseVertexIndex, FirstInstance);
@@ -1490,7 +1490,7 @@ void FD3D12CommandContext::RHIDrawIndexedPrimitiveIndirect(FRHIIndexBuffer* Inde
 
 	// Set the index buffer.
 	const DXGI_FORMAT Format = (IndexBuffer->GetStride() == sizeof(uint16) ? DXGI_FORMAT_R16_UINT : DXGI_FORMAT_R32_UINT);
-	StateCache.SetIndexBuffer(&IndexBuffer->ResourceLocation, Format, 0);
+	StateCache.SetIndexBuffer(IndexBuffer->ResourceLocation, Format, 0);
 
 	FD3D12ResourceLocation& Location = ArgumentBuffer->ResourceLocation;
 	FD3D12DynamicRHI::TransitionResource(CommandListHandle, Location.GetResource(), D3D12_RESOURCE_STATE_INDIRECT_ARGUMENT, D3D12_RESOURCE_BARRIER_ALL_SUBRESOURCES);
diff --git a/Engine/Source/Runtime/D3D12RHI/Private/D3D12DescriptorCache.cpp b/Engine/Source/Runtime/D3D12RHI/Private/D3D12DescriptorCache.cpp
--- a/Engine/Source/Runtime/D3D12RHI/Private/D3D12DescriptorCache.cpp
+++ b/Engine/Source/Runtime/D3D12RHI/Private/D3D12DescriptorCache.cpp
@@ -279,12 +279,6 @@ void FD3D12DescriptorCache::NotifyCurrentCommandList(const FD3D12CommandListHand
 	LocalSamplerHeap.NotifyCurrentCommandList(CommandListHandle);
 }
 
-void FD3D12DescriptorCache::SetIndexBuffer(FD3D12IndexBufferCache& Cache)
-{
-	CmdContext->CommandListHandle.UpdateResidency(Cache.ResidencyHandle);
-	CmdContext->CommandListHandle->IASetIndexBuffer(&Cache.CurrentIndexBufferView);
-}
-
 void FD3D12DescriptorCache::SetVertexBuffers(FD3D12VertexBufferCache& Cache)
 {
 	const uint32 Count = Cache.MaxBoundVertexBufferIndex + 1;
diff --git a/Engine/Source/Runtime/D3D12RHI/Private/D3D12DescriptorCache.h b/Engine/Source/Runtime/D3D12RHI/Private/D3D12DescriptorCache.h
--- a/Engine/Source/Runtime/D3D12RHI/Private/D3D12DescriptorCache.h
+++ b/Engine/Source/Runtime/D3D12RHI/Private/D3D12DescriptorCache.h
@@ -561,7 +561,6 @@ public:
 #endif
 	TRefCountPtr<FD3D12SamplerState> pDefaultSampler;
 
-	void SetIndexBuffer(FD3D12IndexBufferCache& Cache);
 	void SetVertexBuffers(FD3D12VertexBufferCache& Cache);
 	void SetRenderTargets(FD3D12RenderTargetView** RenderTargetViewArray, uint32 Count, FD3D12DepthStencilView* DepthStencilTarget);
 
diff --git a/Engine/Source/Runtime/D3D12RHI/Private/D3D12StateCache.cpp b/Engine/Source/Runtime/D3D12RHI/Private/D3D12StateCache.cpp
--- a/Engine/Source/Runtime/D3D12RHI/Private/D3D12StateCache.cpp
+++ b/Engine/Source/Runtime/D3D12RHI/Private/D3D12StateCache.cpp
@@ -206,7 +206,10 @@ void FD3D12StateCacheBase::DirtyStateForNewCommandList()
 	bNeedSetPrimitiveTopology = true;
 
 	if (PipelineState.Graphics.VBCache.BoundVBMask) { bNeedSetVB = true; }
-	if (PipelineState.Graphics.IBCache.CurrentIndexBufferLocation) { bNeedSetIB = true; }
+
+	// IndexBuffers are set in DrawIndexed*() calls, so there's no way to depend on previously set IndexBuffers without making a new DrawIndexed*() call.
+	PipelineState.Graphics.IBCache.Clear();
+
 	if (PipelineState.Graphics.CurrentNumberOfStreamOutTargets) { bNeedSetSOs = true; }
 	if (PipelineState.Graphics.CurrentNumberOfRenderTargets || PipelineState.Graphics.CurrentDepthStencilTarget) { bNeedSetRTs = true; }
 	if (PipelineState.Graphics.CurrentNumberOfViewports) { bNeedSetViewports = true; }
@@ -242,7 +245,6 @@ void FD3D12StateCacheBase::DirtyState()
 	PipelineState.Compute.bNeedSetRootSignature = true;
 	PipelineState.Graphics.bNeedSetRootSignature = true;
 	bNeedSetVB = true;
-	bNeedSetIB = true;
 	bNeedSetSOs = true;
 	bNeedSetRTs = true;
 	bNeedSetViewports = true;
@@ -402,57 +404,49 @@ void FD3D12StateCacheBase::ApplyState()
 		// Setup non-heap bindings
 		if (bNeedSetVB)
 		{
+			bNeedSetVB = false;
 			//SCOPE_CYCLE_COUNTER(STAT_D3D12ApplyStateSetVertexBufferTime);
 			DescriptorCache.SetVertexBuffers(PipelineState.Graphics.VBCache);
-			bNeedSetVB = false;
-		}
-		if (bNeedSetIB)
-		{
-			if (PipelineState.Graphics.IBCache.CurrentIndexBufferLocation != nullptr)
-			{
-				DescriptorCache.SetIndexBuffer(PipelineState.Graphics.IBCache);
-			}
-			bNeedSetIB = false;
 		}
 		if (bNeedSetSOs)
 		{
-			DescriptorCache.SetStreamOutTargets(PipelineState.Graphics.CurrentStreamOutTargets, PipelineState.Graphics.CurrentNumberOfStreamOutTargets, PipelineState.Graphics.CurrentSOOffsets);
 			bNeedSetSOs = false;
+			DescriptorCache.SetStreamOutTargets(PipelineState.Graphics.CurrentStreamOutTargets, PipelineState.Graphics.CurrentNumberOfStreamOutTargets, PipelineState.Graphics.CurrentSOOffsets);
 		}
 		if (bNeedSetViewports)
 		{
-			CommandList->RSSetViewports(PipelineState.Graphics.CurrentNumberOfViewports, PipelineState.Graphics.CurrentViewport);
 			bNeedSetViewports = false;
+			CommandList->RSSetViewports(PipelineState.Graphics.CurrentNumberOfViewports, PipelineState.Graphics.CurrentViewport);
 		}
 		if (bNeedSetScissorRects)
 		{
-			CommandList->RSSetScissorRects(PipelineState.Graphics.CurrentNumberOfScissorRects, PipelineState.Graphics.CurrentScissorRects);
 			bNeedSetScissorRects = false;
+			CommandList->RSSetScissorRects(PipelineState.Graphics.CurrentNumberOfScissorRects, PipelineState.Graphics.CurrentScissorRects);
 		}
 		if (bNeedSetPrimitiveTopology)
 		{
-			CommandList->IASetPrimitiveTopology(PipelineState.Graphics.CurrentPrimitiveTopology);
 			bNeedSetPrimitiveTopology = false;
+			CommandList->IASetPrimitiveTopology(PipelineState.Graphics.CurrentPrimitiveTopology);
 		}
 		if (bNeedSetBlendFactor)
 		{
-			CommandList->OMSetBlendFactor(PipelineState.Graphics.CurrentBlendFactor);
 			bNeedSetBlendFactor = false;
+			CommandList->OMSetBlendFactor(PipelineState.Graphics.CurrentBlendFactor);
 		}
 		if (bNeedSetStencilRef)
 		{
-			CommandList->OMSetStencilRef(PipelineState.Graphics.CurrentReferenceStencil);
 			bNeedSetStencilRef = false;
+			CommandList->OMSetStencilRef(PipelineState.Graphics.CurrentReferenceStencil);
 		}
 		if (bNeedSetRTs)
 		{
-			DescriptorCache.SetRenderTargets(PipelineState.Graphics.RenderTargetArray, PipelineState.Graphics.CurrentNumberOfRenderTargets, PipelineState.Graphics.CurrentDepthStencilTarget);
 			bNeedSetRTs = false;
+			DescriptorCache.SetRenderTargets(PipelineState.Graphics.RenderTargetArray, PipelineState.Graphics.CurrentNumberOfRenderTargets, PipelineState.Graphics.CurrentDepthStencilTarget);
 		}
 		if (bNeedSetDepthBounds)
 		{
-			CmdContext->SetDepthBounds(PipelineState.Graphics.MinDepth, PipelineState.Graphics.MaxDepth);
 			bNeedSetDepthBounds = false;
+			CmdContext->SetDepthBounds(PipelineState.Graphics.MinDepth, PipelineState.Graphics.MaxDepth);
 		}
 	}
 
@@ -1071,46 +1065,17 @@ void FD3D12StateCacheBase::SetComputeShader(FD3D12ComputeShader* Shader)
 	}
 }
 
-void FD3D12StateCacheBase::InternalSetIndexBuffer(FD3D12ResourceLocation* IndexBufferLocation, DXGI_FORMAT Format, uint32 Offset)
+void FD3D12StateCacheBase::InternalSetIndexBuffer(FD3D12Resource* Resource)
 {
-	__declspec(align(16)) D3D12_INDEX_BUFFER_VIEW NewView;
-	NewView.BufferLocation = (IndexBufferLocation) ? IndexBufferLocation->GetGPUVirtualAddress() + Offset : 0;
-	NewView.SizeInBytes = (IndexBufferLocation) ? IndexBufferLocation->GetSize() - Offset : 0;
-	NewView.Format = Format;
-
-	D3D12_INDEX_BUFFER_VIEW& CurrentView = PipelineState.Graphics.IBCache.CurrentIndexBufferView;
-
-	if (NewView.BufferLocation != CurrentView.BufferLocation ||
-		NewView.SizeInBytes != CurrentView.SizeInBytes ||
-		NewView.Format != CurrentView.Format ||
-		GD3D12SkipStateCaching)
-	{
-		bNeedSetIB = true;
-		PipelineState.Graphics.IBCache.CurrentIndexBufferLocation = IndexBufferLocation;
-
-		if (IndexBufferLocation != nullptr)
-		{
-			PipelineState.Graphics.IBCache.ResidencyHandle = IndexBufferLocation->GetResource()->GetResidencyHandle();
-			FMemory::Memcpy(CurrentView, NewView);
-		}
-		else
-		{
-			FMemory::Memzero(&CurrentView, sizeof(CurrentView));
-			PipelineState.Graphics.IBCache.CurrentIndexBufferLocation = nullptr;
-			PipelineState.Graphics.IBCache.ResidencyHandle = nullptr;
-		}
-	}
+	FD3D12CommandListHandle& CommandListHandle = CmdContext->CommandListHandle;
+	CommandListHandle.UpdateResidency(Resource->GetResidencyHandle());
+	CommandListHandle->IASetIndexBuffer(&PipelineState.Graphics.IBCache.CurrentIndexBufferView);
 
-	if (IndexBufferLocation)
+	if (Resource->RequiresResourceStateTracking())
 	{
-		FD3D12Resource* const pResource = IndexBufferLocation->GetResource();
-		if (pResource && pResource->RequiresResourceStateTracking())
-		{
-			check(pResource->GetSubresourceCount() == 1);
-			FD3D12DynamicRHI::TransitionResource(CmdContext->CommandListHandle, pResource, D3D12_RESOURCE_STATE_INDEX_BUFFER, D3D12_RESOURCE_BARRIER_ALL_SUBRESOURCES);
-		}
+		check(Resource->GetSubresourceCount() == 1);
+		FD3D12DynamicRHI::TransitionResource(CommandListHandle, Resource, D3D12_RESOURCE_STATE_INDEX_BUFFER, D3D12_RESOURCE_BARRIER_ALL_SUBRESOURCES);
 	}
-
 }
 
 void FD3D12StateCacheBase::InternalSetStreamSource(FD3D12ResourceLocation* VertexBufferLocation, uint32 StreamIndex, uint32 Stride, uint32 Offset)
diff --git a/Engine/Source/Runtime/D3D12RHI/Private/D3D12StateCachePrivate.h b/Engine/Source/Runtime/D3D12RHI/Private/D3D12StateCachePrivate.h
--- a/Engine/Source/Runtime/D3D12RHI/Private/D3D12StateCachePrivate.h
+++ b/Engine/Source/Runtime/D3D12RHI/Private/D3D12StateCachePrivate.h
@@ -96,13 +96,9 @@ struct FD3D12IndexBufferCache
 	inline void Clear()
 	{
 		FMemory::Memzero(&CurrentIndexBufferView, sizeof(CurrentIndexBufferView));
-		CurrentIndexBufferLocation = nullptr;
-		ResidencyHandle = nullptr;
 	}
 
 	D3D12_INDEX_BUFFER_VIEW CurrentIndexBufferView;
-	FD3D12ResourceLocation* CurrentIndexBufferLocation;
-	FD3D12ResidencyHandle* ResidencyHandle;
 };
 
 template<typename ResourceSlotMask>
@@ -337,7 +333,6 @@ protected:
 	FD3D12CommandContext* CmdContext;
 
 	bool bNeedSetVB;
-	bool bNeedSetIB;
 	bool bNeedSetRTs;
 	bool bNeedSetSOs;
 	bool bSRVSCleared;
@@ -432,7 +427,7 @@ protected:
 
 	FD3D12DescriptorCache DescriptorCache;
 
-	void InternalSetIndexBuffer(FD3D12ResourceLocation *IndexBufferLocation, DXGI_FORMAT Format, uint32 Offset);
+	void InternalSetIndexBuffer(FD3D12Resource* Resource);
 
 	void InternalSetStreamSource(FD3D12ResourceLocation* VertexBufferLocation, uint32 StreamIndex, uint32 Stride, uint32 Offset);
 
@@ -855,14 +850,24 @@ public:
 
 public:
 
-	D3D12_STATE_CACHE_INLINE void SetIndexBuffer(FD3D12ResourceLocation* IndexBufferLocation, DXGI_FORMAT Format, uint32 Offset)
+	D3D12_STATE_CACHE_INLINE void SetIndexBuffer(const FD3D12ResourceLocation& IndexBufferLocation, DXGI_FORMAT Format, uint32 Offset)
 	{
-		InternalSetIndexBuffer(IndexBufferLocation, Format, Offset);
-	}
+		D3D12_GPU_VIRTUAL_ADDRESS BufferLocation = IndexBufferLocation.GetGPUVirtualAddress() + Offset;
+		UINT SizeInBytes = IndexBufferLocation.GetSize() - Offset;
 
-	D3D12_STATE_CACHE_INLINE bool IsIndexBuffer(const FD3D12ResourceLocation *ResourceLocation) const
-	{
-		return PipelineState.Graphics.IBCache.CurrentIndexBufferLocation == ResourceLocation;
+		D3D12_INDEX_BUFFER_VIEW& CurrentView = PipelineState.Graphics.IBCache.CurrentIndexBufferView;
+
+		if (BufferLocation != CurrentView.BufferLocation ||
+			SizeInBytes != CurrentView.SizeInBytes ||
+			Format != CurrentView.Format ||
+			GD3D12SkipStateCaching)
+		{
+			CurrentView.BufferLocation = BufferLocation;
+			CurrentView.SizeInBytes = SizeInBytes;
+			CurrentView.Format = Format;
+
+			InternalSetIndexBuffer(IndexBufferLocation.GetResource());
+		}
 	}
 
 	D3D12_STATE_CACHE_INLINE void GetPrimitiveTopology(D3D12_PRIMITIVE_TOPOLOGY* PrimitiveTopology) const
@@ -959,7 +964,6 @@ public:
 	void ForceSetGraphicsRootSignature() { PipelineState.Graphics.bNeedSetRootSignature = true; }
 	void ForceSetComputeRootSignature() { PipelineState.Compute.bNeedSetRootSignature = true; }
 	void ForceSetVB() { bNeedSetVB = true; }
-	void ForceSetIB() { bNeedSetIB = true; }
 	void ForceSetRTs() { bNeedSetRTs = true; }
 	void ForceSetSOs() { bNeedSetSOs = true; }
 	void ForceSetSamplersPerShaderStage(uint32 Frequency) { PipelineState.Common.SamplerCache.Dirty((EShaderFrequency)Frequency); }
@@ -971,7 +975,6 @@ public:
 	void ForceSetStencilRef() { bNeedSetStencilRef = true; }
 
 	bool GetForceSetVB() const { return bNeedSetVB; }
-	bool GetForceSetIB() const { return bNeedSetIB; }
 	bool GetForceSetRTs() const { return bNeedSetRTs; }
 	bool GetForceSetSOs() const { return bNeedSetSOs; }
 	bool GetForceSetSamplersPerShaderStage(uint32 Frequency) const { return PipelineState.Common.SamplerCache.DirtySlotMask[Frequency] != 0; }
