Allow unique custom user data to be provided per instance when native ray tracing instancing is used. Previously only transforms could be unique.

(ushell-p4-cherrypick of 9622434 by [REDACTED_PERSON])

diff --git a/Engine/Source/Runtime/D3D12RHI/Private/D3D12RayTracing.cpp b/Engine/Source/Runtime/D3D12RHI/Private/D3D12RayTracing.cpp
--- a/Engine/Source/Runtime/D3D12RHI/Private/D3D12RayTracing.cpp
+++ b/Engine/Source/Runtime/D3D12RHI/Private/D3D12RayTracing.cpp
@@ -2363,7 +2363,6 @@ void FD3D12RayTracingScene::BuildAccelerationStructure(FD3D12CommandContext& Com
 
 			D3D12_RAYTRACING_INSTANCE_DESC InstanceDesc = {};
 
-			InstanceDesc.InstanceID = Instance.UserData;
 			InstanceDesc.InstanceMask = Instance.Mask;
 			InstanceDesc.InstanceContributionToHitGroupIndex = SegmentPrefixSum[InstanceIndex] * ShaderSlotsPerGeometrySegment;
 			InstanceDesc.Flags = D3D12_RAYTRACING_INSTANCE_FLAG_TRIANGLE_FRONT_COUNTERCLOCKWISE; // #dxr_todo: convert cull mode based on instance mirroring or double-sidedness
@@ -2389,8 +2388,17 @@ void FD3D12RayTracingScene::BuildAccelerationStructure(FD3D12CommandContext& Com
 			}
 
 			const uint32 NumTransforms = Instance.Transforms.Num();
+
+			checkf(Instance.UserData.Num() == 0 || Instance.UserData.Num() == 1 || Instance.UserData.Num() == NumTransforms, 
+				TEXT("User data array must be either be empty (implicit 0 is assumed), contain a single entry (same value applied to all instances) or contain one entry per entry in Transforms array."));
+
+			const bool bUseUniqueUserData = Instance.UserData.Num() > 1;
+			const uint32 CommonUserData = Instance.UserData.Num() == 1 ? Instance.UserData[0] : 0;
+
 			for (uint32 TransformIndex = 0; TransformIndex < NumTransforms; ++TransformIndex)
 			{
+				InstanceDesc.InstanceID = bUseUniqueUserData ? Instance.UserData[TransformIndex] : CommonUserData;
+
 				FMatrix TransformTransposed = Instance.Transforms[TransformIndex].GetTransposed();
 
 				// Ensure the last row of the original Transform is <0,0,0,1>
diff --git a/Engine/Source/Runtime/RHI/Public/DynamicRHI.h b/Engine/Source/Runtime/RHI/Public/DynamicRHI.h
--- a/Engine/Source/Runtime/RHI/Public/DynamicRHI.h
+++ b/Engine/Source/Runtime/RHI/Public/DynamicRHI.h
@@ -79,8 +79,20 @@ struct FRHIFlipDetails
 struct FRayTracingGeometryInstance
 {
 	FRHIRayTracingGeometry* GeometryRHI = nullptr;
+
+	// A physical FRayTracingGeometryInstance may be duplicated many times in the scene with different transforms and user data.
+	// All copies share the same shader binding table entries and therefore will have the same material and shader resources.
 	TArray<FMatrix, TInlineAllocator<1>> Transforms;
-	uint32 UserData = 0;
+
+	// Each geometry copy can receive a user-provided integer, which can be used to retrieve extra shader parameters or customize appearance.
+	// This data can be retrieved using GetInstanceUserData() in closest/any hit shaders.
+	// If UserData is empty, then 0 will be implicitly used for all entries.
+	// If UserData contains a single entry, then it will be applied to all instances/copies.
+	// Otherwise one UserData entry must be provided per entry in Transforms array.
+	TArray<uint32, TInlineAllocator<1>> UserData;
+
+	// Mask that will be tested against one provided to TraceRay() in shader code.
+	// If binary AND of instance mask with ray mask is zero, then the instance is considered not intersected / invisible.
 	uint8 Mask = 0xFF;
 
 	// No any-hit shaders will be invoked for this geometry instance (only closest hit)
diff --git a/Engine/Source/Runtime/Renderer/Private/DeferredShadingRenderer.cpp b/Engine/Source/Runtime/Renderer/Private/DeferredShadingRenderer.cpp
--- a/Engine/Source/Runtime/Renderer/Private/DeferredShadingRenderer.cpp
+++ b/Engine/Source/Runtime/Renderer/Private/DeferredShadingRenderer.cpp
@@ -770,7 +770,7 @@ bool FDeferredShadingSceneRenderer::GatherRayTracingWorldInstances(FRHICommandLi
 							// When no cached command is found, NewInstanceMask == 0 and the instance is effectively filtered out
 							FRayTracingGeometryInstance RayTracingInstance = { RayTracingGeometryInstance };
 							RayTracingInstance.Transforms.Add(Scene->PrimitiveTransforms[PrimitiveIndex]);
-							RayTracingInstance.UserData = (uint32)PrimitiveIndex;
+							RayTracingInstance.UserData.Add((uint32)PrimitiveIndex);
 							RayTracingInstance.Mask = NewInstanceMask;
 							RayTracingInstance.bForceOpaque = bAllSegmentsOpaque;
 							View.RayTracingGeometryInstances.Add(RayTracingInstance);
@@ -809,7 +809,7 @@ bool FDeferredShadingSceneRenderer::GatherRayTracingWorldInstances(FRHICommandLi
 					for (FRayTracingInstance& Instance : RayTracingInstances)
 					{
 						FRayTracingGeometryInstance RayTracingInstance = { Instance.Geometry->RayTracingGeometryRHI };
-						RayTracingInstance.UserData = (uint32)PrimitiveIndex;
+						RayTracingInstance.UserData.Add((uint32)PrimitiveIndex);
 						RayTracingInstance.Mask = Instance.Mask;
 						RayTracingInstance.bForceOpaque = Instance.bForceOpaque;
 
diff --git a/Engine/Source/Runtime/Renderer/Private/Tests/RayTracingTestbed.cpp b/Engine/Source/Runtime/Renderer/Private/Tests/RayTracingTestbed.cpp
--- a/Engine/Source/Runtime/Renderer/Private/Tests/RayTracingTestbed.cpp
+++ b/Engine/Source/Runtime/Renderer/Private/Tests/RayTracingTestbed.cpp
@@ -104,7 +104,7 @@ bool RunRayTracingTestbed_RenderThread(const FString& Parameters)
 	FRayTracingGeometryRHIRef Geometry = RHICreateRayTracingGeometry(GeometryInitializer);
 
 	FRayTracingGeometryInstance Instances[] = {
-		FRayTracingGeometryInstance { Geometry, {FMatrix::Identity}, 0, 0xFF }
+		FRayTracingGeometryInstance { Geometry, {FMatrix::Identity}, {0}, 0xFF }
 	};
 
 	FRayTracingSceneInitializer Initializer;
