Fix FrameFence for the RHI thread disabled case

diff --git a/Engine/Source/Runtime/D3D12RHI/Private/D3D12Viewport.cpp b/Engine/Source/Runtime/D3D12RHI/Private/D3D12Viewport.cpp
--- a/Engine/Source/Runtime/D3D12RHI/Private/D3D12Viewport.cpp
+++ b/Engine/Source/Runtime/D3D12RHI/Private/D3D12Viewport.cpp
@@ -938,15 +938,17 @@ void FD3D12DynamicRHI::RHIAdvanceFrameFence()
 	FD3D12ManualFence* FrameFence = &GetAdapter().GetFrameFence();
 	const uint64 PreviousFence = FrameFence->IncrementCurrentFence();
 
-	// Queue a command to signal on RHI thread that the current frame is a complete on the GPU.
 	FRHICommandListImmediate& RHICmdList = FRHICommandListExecutor::GetImmediateCommandList();
-	if (RHICmdList.Bypass() || !IsRunningRHIInSeparateThread())
+	if (RHICmdList.Bypass())
 	{
+		// In bypass mode, we should execute this directly
 		FRHICommandSignalFrameFence Cmd(ED3D12CommandQueueType::Default, FrameFence, PreviousFence);
 		Cmd.Execute(RHICmdList);
 	}
 	else
 	{
+		// Queue a command to signal on RHI thread that the current frame is a complete on the GPU.
+		// This must be done in a deferred way even if RHI thread is disabled, just for correct ordering of operations.
 		ALLOC_COMMAND_CL(RHICmdList, FRHICommandSignalFrameFence)(ED3D12CommandQueueType::Default, FrameFence, PreviousFence);
 	}
 }
