Added a more explicit description to memory leak check in FD3D12SegListAllocator::VerifyEmpty()

(ushell-p4-cherrypick of 9631112 by [REDACTED_PERSON])

diff --git a/Engine/Source/Runtime/D3D12RHI/Private/D3D12Allocation.cpp b/Engine/Source/Runtime/D3D12RHI/Private/D3D12Allocation.cpp
--- a/Engine/Source/Runtime/D3D12RHI/Private/D3D12Allocation.cpp
+++ b/Engine/Source/Runtime/D3D12RHI/Private/D3D12Allocation.cpp
@@ -1697,7 +1697,10 @@ void FD3D12SegListAllocator::VerifyEmpty()
 			DumpStack(LeakTrack);
 		}
 	}
-	check(TotalBytesRequested == 0); //SegList was not properly freed. Run with d3d12.SegListTrackLeaks=1 to print callstacks of offending allocations
+
+	checkf(TotalBytesRequested == 0,
+		TEXT("FD3D12SegListAllocator contains %lld allocated bytes but is expected to be empty. This likely means a memory leak. Use d3d12.SegListTrackLeaks=1 CVar to print allocations to the log."),
+		(uint64)TotalBytesRequested);
 }
 
 
