Batch descriptor copies

diff --git a/Engine/Source/Runtime/D3D12RHI/Private/D3D12DescriptorCache.cpp b/Engine/Source/Runtime/D3D12RHI/Private/D3D12DescriptorCache.cpp
--- a/Engine/Source/Runtime/D3D12RHI/Private/D3D12DescriptorCache.cpp
+++ b/Engine/Source/Runtime/D3D12RHI/Private/D3D12DescriptorCache.cpp
@@ -329,9 +329,8 @@ void FD3D12DescriptorCache::SetUAVs(const FD3D12RootSignature* RootSignature, FD
 			FD3D12DynamicRHI::TransitionResource(CommandList, UAVs[SlotIndex], D3D12_RESOURCE_STATE_UNORDERED_ACCESS);
 			CommandList.UpdateResidency(Cache.ResidencyHandles[ShaderStage][SlotIndex]);
 		}
-
-		FD3D12UnorderedAccessViewCache::CleanSlot(CurrentDirtySlotMask, SlotIndex);
 	}
+	FD3D12UnorderedAccessViewCache::CleanSlots(CurrentDirtySlotMask, SlotsNeeded);
 
 	check((CurrentDirtySlotMask & SlotsNeededMask) == 0);	// Check all slots that needed to be set, were set.
 
@@ -480,9 +479,8 @@ void FD3D12DescriptorCache::SetSamplers(const FD3D12RootSignature* RootSignature
 		for (uint32 SlotIndex = 0; SlotIndex < SlotsNeeded; SlotIndex++)
 		{
 			Desc.SamplerID[SlotIndex] = Samplers[SlotIndex] ? Samplers[SlotIndex]->ID : 0;
-
-			FD3D12SamplerStateCache::CleanSlot(CacheDirtySlotMask, SlotIndex);
 		}
+		FD3D12SamplerStateCache::CleanSlots(CacheDirtySlotMask, SlotsNeeded);
 
 		// The hash uses all of the bits
 		for (uint32 SlotIndex = SlotsNeeded; SlotIndex < UE_ARRAY_COUNT(Desc.SamplerID); SlotIndex++)
@@ -522,9 +520,8 @@ void FD3D12DescriptorCache::SetSamplers(const FD3D12RootSignature* RootSignature
 			{
 				SrcDescriptors[SlotIndex] = pDefaultSampler->Descriptor;
 			}
-
-			FD3D12SamplerStateCache::CleanSlot(CurrentDirtySlotMask, SlotIndex);
 		}
+		FD3D12SamplerStateCache::CleanSlots(CurrentDirtySlotMask, SlotsNeeded);
 
 		GetParentDevice()->GetDevice()->CopyDescriptors(
 			1, &DestDescriptor, &SlotsNeeded,
@@ -574,7 +571,6 @@ void FD3D12DescriptorCache::SetSRVs(const FD3D12RootSignature* RootSignature, FD
 	check(SlotsNeededMask != 0);		// All dirty slots for the current shader stage AND used by the current shader stage.
 	check(SlotsNeeded != 0);
 
-	ID3D12Device* Device = GetParentDevice()->GetDevice();
 	FD3D12CommandListHandle& CommandList = CmdContext->CommandListHandle;
 
 	auto& SRVs = Cache.Views[ShaderStage];
@@ -584,17 +580,13 @@ void FD3D12DescriptorCache::SetSRVs(const FD3D12RootSignature* RootSignature, FD
 	HeapSlot += SlotsNeeded;
 
 	D3D12_CPU_DESCRIPTOR_HANDLE DestDescriptor = CurrentViewHeap->GetCPUSlotHandle(FirstSlotIndex);
-	const uint64 DescriptorSize = CurrentViewHeap->GetDescriptorSize();
+	D3D12_CPU_DESCRIPTOR_HANDLE SrcDescriptors[MAX_SRVS];
 
 	for (uint32 SlotIndex = 0; SlotIndex < SlotsNeeded; SlotIndex++)
 	{
-		CD3DX12_CPU_DESCRIPTOR_HANDLE SrcDescriptor;
-
 		if (SRVs[SlotIndex] != nullptr)
 		{
-			SrcDescriptor = SRVs[SlotIndex]->GetView();
-
-			CommandList.UpdateResidency(Cache.ResidencyHandles[ShaderStage][SlotIndex]);
+			SrcDescriptors[SlotIndex] = SRVs[SlotIndex]->GetView();
 
 			if (SRVs[SlotIndex]->IsDepthStencilResource())
 			{
@@ -608,19 +600,20 @@ void FD3D12DescriptorCache::SetSRVs(const FD3D12RootSignature* RootSignature, FD
 			{
 				FD3D12DynamicRHI::TransitionResource(CommandList, SRVs[SlotIndex], D3D12_RESOURCE_STATE_PIXEL_SHADER_RESOURCE | D3D12_RESOURCE_STATE_NON_PIXEL_SHADER_RESOURCE);
 			}
+
+			CommandList.UpdateResidency(Cache.ResidencyHandles[ShaderStage][SlotIndex]);
 		}
 		else
 		{
-			SrcDescriptor = pNullSRV->GetHandle();
+			SrcDescriptors[SlotIndex] = pNullSRV->GetHandle();
 		}
-		check(SrcDescriptor.ptr != 0);
-
-		Device->CopyDescriptorsSimple(1, DestDescriptor, SrcDescriptor, D3D12_DESCRIPTOR_HEAP_TYPE_CBV_SRV_UAV);
+		check(SrcDescriptors[SlotIndex].ptr != 0);
+	}
+	FD3D12ShaderResourceViewCache::CleanSlots(CurrentDirtySlotMask, SlotsNeeded);
 
-		DestDescriptor.ptr += DescriptorSize;
+	ID3D12Device* Device = GetParentDevice()->GetDevice();
+	Device->CopyDescriptors(1, &DestDescriptor, &SlotsNeeded, SlotsNeeded, SrcDescriptors, nullptr, D3D12_DESCRIPTOR_HEAP_TYPE_CBV_SRV_UAV);
 
-		FD3D12ShaderResourceViewCache::CleanSlot(CurrentDirtySlotMask, SlotIndex);
-	}
 
 	check((CurrentDirtySlotMask & SlotsNeededMask) == 0);	// Check all slots that needed to be set, were set.
 
diff --git a/Engine/Source/Runtime/D3D12RHI/Private/D3D12StateCachePrivate.h b/Engine/Source/Runtime/D3D12RHI/Private/D3D12StateCachePrivate.h
--- a/Engine/Source/Runtime/D3D12RHI/Private/D3D12StateCachePrivate.h
+++ b/Engine/Source/Runtime/D3D12RHI/Private/D3D12StateCachePrivate.h
@@ -109,6 +109,11 @@ struct FD3D12ResourceCache
 		SlotMask &= ~((ResourceSlotMask)1 << SlotIndex);
 	}
 
+	static inline void CleanSlots(ResourceSlotMask& SlotMask, uint32 NumSlots)
+	{
+		SlotMask &= ~(((ResourceSlotMask)1 << NumSlots) - 1);
+	}
+
 	static inline void DirtySlot(ResourceSlotMask& SlotMask, uint32 SlotIndex)
 	{
 		SlotMask |= ((ResourceSlotMask)1 << SlotIndex);
