Fix for FConditionalScopeResourceBarrier-related assert : don't trigger a resource transition for resources that don't have state tracking enabled when the source and destination state are identical (in the case of the landscape edit layers : CPU readback resource, whose state is always D3D12_RESOURCE_STATE_COPY_DEST)
This could be handled on the client side by not using a FConditionalScopeResourceBarrier object in the first place but this would require the client to do pretty much what the scope object does internally.

diff --git a/Engine/Source/Runtime/D3D12RHI/Private/D3D12RHIPrivate.h b/Engine/Source/Runtime/D3D12RHI/Private/D3D12RHIPrivate.h
--- a/Engine/Source/Runtime/D3D12RHI/Private/D3D12RHIPrivate.h
+++ b/Engine/Source/Runtime/D3D12RHI/Private/D3D12RHIPrivate.h
@@ -1040,7 +1040,7 @@ private:
 	D3D12_RESOURCE_STATES Current;
 	const D3D12_RESOURCE_STATES Desired;
 	const uint32 Subresource;
-	const bool bUseTracking;
+	bool bRestoreDefaultState;
 
 public:
 	explicit FConditionalScopeResourceBarrier(FD3D12CommandListHandle& hInCommandList, FD3D12Resource* pInResource, const D3D12_RESOURCE_STATES InDesired, const uint32 InSubresource)
@@ -1049,12 +1049,18 @@ public:
 		, Current(D3D12_RESOURCE_STATE_TBD)
 		, Desired(InDesired)
 		, Subresource(InSubresource)
-		, bUseTracking(pResource->RequiresResourceStateTracking())
+		, bRestoreDefaultState(false)
 	{
-		if (!bUseTracking)
+		// when we don't use resource state tracking, transition the resource (only if necessary)
+		if (!pResource->RequiresResourceStateTracking())
 		{
 			Current = pResource->GetDefaultResourceState();
-			hCommandList.AddTransitionBarrier(pResource, Current, Desired, Subresource);
+			if (Current != Desired)
+			{
+				// we will add a transition, we need to transition back to the default state when the scoped object dies : 
+				bRestoreDefaultState = true;
+				hCommandList.AddTransitionBarrier(pResource, Current, Desired, Subresource);
+			}
 		}
 		else
 		{
@@ -1064,8 +1070,8 @@ public:
 
 	~FConditionalScopeResourceBarrier()
 	{
-		// Return the resource to it's default state if it doesn't use tracking
-		if (!bUseTracking)
+		// Return the resource to it's default state if necessary : 
+		if (bRestoreDefaultState)
 		{
 			hCommandList.AddTransitionBarrier(pResource, Desired, Current, Subresource);
 		}
