F - Don't fatal if can't query D3D12 resolutions over remote desktop
#rnx

diff --git a/Engine/Source/Runtime/D3D12RHI/Private/Windows/WindowsD3D12Device.cpp b/Engine/Source/Runtime/D3D12RHI/Private/Windows/WindowsD3D12Device.cpp
--- a/Engine/Source/Runtime/D3D12RHI/Private/Windows/WindowsD3D12Device.cpp
+++ b/Engine/Source/Runtime/D3D12RHI/Private/Windows/WindowsD3D12Device.cpp
@@ -890,8 +890,8 @@ bool FD3D12DynamicRHI::RHIGetAvailableResolutions(FScreenResolutionArray& Resolu
 		}
 		else if (HResult == DXGI_ERROR_NOT_CURRENTLY_AVAILABLE)
 		{
-			UE_LOG(LogD3D12RHI, Fatal,
-				TEXT("This application cannot be run over a remote desktop configuration")
+			UE_LOG(LogD3D12RHI, Warning,
+				TEXT("RHIGetAvailableResolutions() can not be used over a remote desktop configuration")
 				);
 			return false;
 		}
