UE4R - Bubble up create texture error on D3D12 as it was eaten deep inside and had no relevant info
#rnx

diff --git a/Engine/Source/Runtime/D3D12RHI/Private/D3D12Adapter.h b/Engine/Source/Runtime/D3D12RHI/Private/D3D12Adapter.h
--- a/Engine/Source/Runtime/D3D12RHI/Private/D3D12Adapter.h
+++ b/Engine/Source/Runtime/D3D12RHI/Private/D3D12Adapter.h
@@ -180,7 +180,8 @@ public:
 		const D3D12_RESOURCE_STATES& InitialUsage,
 		const D3D12_CLEAR_VALUE* ClearValue,
 		FD3D12Resource** ppOutResource,
-		const TCHAR* Name);
+		const TCHAR* Name,
+		bool bVerifyHResult = true);
 
 	HRESULT CreatePlacedResource(const D3D12_RESOURCE_DESC& Desc,
 		FD3D12Heap* BackingHeap,
@@ -188,7 +189,8 @@ public:
 		const D3D12_RESOURCE_STATES& InitialUsage,
 		const D3D12_CLEAR_VALUE* ClearValue,
 		FD3D12Resource** ppOutResource,
-		const TCHAR* Name);
+		const TCHAR* Name,
+		bool bVerifyHResult = true);
 
 	HRESULT CreateBuffer(D3D12_HEAP_TYPE HeapType,
 		FRHIGPUMask CreationNode,
diff --git a/Engine/Source/Runtime/D3D12RHI/Private/D3D12Allocation.cpp b/Engine/Source/Runtime/D3D12RHI/Private/D3D12Allocation.cpp
--- a/Engine/Source/Runtime/D3D12RHI/Private/D3D12Allocation.cpp
+++ b/Engine/Source/Runtime/D3D12RHI/Private/D3D12Allocation.cpp
@@ -1172,23 +1172,25 @@ HRESULT FD3D12TextureAllocatorPool::AllocateTexture(
 
 		if (Offset != FD3D12SegListAllocator::InvalidOffset)
 		{
-			RetCode = Adapter->CreatePlacedResource(Desc, BackingHeap.GetReference(), Offset, InitialState, ClearValue, &NewResource, Name);
-
-			FD3D12SegListAllocatorPrivateData& PrivateData = TextureLocation.GetSegListAllocatorPrivateData();
-			PrivateData.Offset = Offset;
-
-			TextureLocation.SetType(FD3D12ResourceLocation::ResourceLocationType::eSubAllocation);
-			TextureLocation.SetSegListAllocator(&ReadOnlyTexturePool);
-			TextureLocation.SetSize(Info.SizeInBytes);
-			TextureLocation.SetOffsetFromBaseOfResource(Offset);
-			TextureLocation.SetResource(NewResource);
+			RetCode = Adapter->CreatePlacedResource(Desc, BackingHeap.GetReference(), Offset, InitialState, ClearValue, &NewResource, Name, false);
+			if (SUCCEEDED(RetCode))
+			{
+				FD3D12SegListAllocatorPrivateData& PrivateData = TextureLocation.GetSegListAllocatorPrivateData();
+				PrivateData.Offset = Offset;
+
+				TextureLocation.SetType(FD3D12ResourceLocation::ResourceLocationType::eSubAllocation);
+				TextureLocation.SetSegListAllocator(&ReadOnlyTexturePool);
+				TextureLocation.SetSize(Info.SizeInBytes);
+				TextureLocation.SetOffsetFromBaseOfResource(Offset);
+				TextureLocation.SetResource(NewResource);
+			}
 			return RetCode;
 		}
 	}
 
 	const D3D12_HEAP_PROPERTIES HeapProps = CD3DX12_HEAP_PROPERTIES(D3D12_HEAP_TYPE_DEFAULT, GetGPUMask().GetNative(), GetVisibilityMask().GetNative());
 	Desc.Alignment = 0;
-	RetCode = Adapter->CreateCommittedResource(Desc, GetGPUMask(), HeapProps, InitialState, ClearValue, &NewResource, Name);
+	RetCode = Adapter->CreateCommittedResource(Desc, GetGPUMask(), HeapProps, InitialState, ClearValue, &NewResource, Name, false);
 
 	TextureLocation.SetType(FD3D12ResourceLocation::ResourceLocationType::eStandAlone);
 	TextureLocation.SetResource(NewResource);
@@ -1240,7 +1242,7 @@ HRESULT FD3D12TextureAllocator::AllocateTexture(D3D12_RESOURCE_DESC Desc, const
 			FD3D12Heap* BackingHeap = ((FD3D12BuddyAllocator*)TextureLocation.GetAllocator())->GetBackingHeap();
 			uint64 HeapOffset = TextureLocation.GetAllocator()->GetAllocationOffsetInBytes(TextureLocation.GetBuddyAllocatorPrivateData());
 
-			hr = Adapter->CreatePlacedResource(Desc, BackingHeap, HeapOffset, InitialState, ClearValue, &NewResource, Name);
+			hr = Adapter->CreatePlacedResource(Desc, BackingHeap, HeapOffset, InitialState, ClearValue, &NewResource, Name, false);
 	
 			TextureLocation.SetType(FD3D12ResourceLocation::ResourceLocationType::eSubAllocation);
 			TextureLocation.SetResource(NewResource);
@@ -1253,7 +1255,7 @@ HRESULT FD3D12TextureAllocator::AllocateTexture(D3D12_RESOURCE_DESC Desc, const
 	Desc.Alignment = 0;
 	const D3D12_HEAP_PROPERTIES HeapProps = CD3DX12_HEAP_PROPERTIES(D3D12_HEAP_TYPE_DEFAULT, GetGPUMask().GetNative(), (uint32)GetVisibilityMask().GetNative());
 
-	hr = Adapter->CreateCommittedResource(Desc, GetGPUMask(), HeapProps, InitialState, ClearValue, &NewResource, Name);
+	hr = Adapter->CreateCommittedResource(Desc, GetGPUMask(), HeapProps, InitialState, ClearValue, &NewResource, Name, false);
 
 	TextureLocation.SetType(FD3D12ResourceLocation::ResourceLocationType::eStandAlone);
 	TextureLocation.SetResource(NewResource);
@@ -1287,7 +1289,7 @@ HRESULT FD3D12TextureAllocatorPool::AllocateTexture(D3D12_RESOURCE_DESC Desc, co
 	FD3D12Resource* Resource = nullptr;
 
 	const D3D12_HEAP_PROPERTIES HeapProps = CD3DX12_HEAP_PROPERTIES(D3D12_HEAP_TYPE_DEFAULT, GetGPUMask().GetNative(), GetVisibilityMask().GetNative());
-	HRESULT hr = Adapter->CreateCommittedResource(Desc, GetGPUMask(), HeapProps, InitialState, ClearValue, &Resource, Name);
+	HRESULT hr = Adapter->CreateCommittedResource(Desc, GetGPUMask(), HeapProps, InitialState, ClearValue, &Resource, Name, false);
 
 	TextureLocation.SetType(FD3D12ResourceLocation::ResourceLocationType::eStandAlone);
 	TextureLocation.SetResource(Resource);
diff --git a/Engine/Source/Runtime/D3D12RHI/Private/D3D12Resources.cpp b/Engine/Source/Runtime/D3D12RHI/Private/D3D12Resources.cpp
--- a/Engine/Source/Runtime/D3D12RHI/Private/D3D12Resources.cpp
+++ b/Engine/Source/Runtime/D3D12RHI/Private/D3D12Resources.cpp
@@ -299,7 +299,7 @@ void FD3D12Heap::BeginTrackingResidency(uint64 Size)
 //	FD3D12 Adapter
 /////////////////////////////////////////////////////////////////////
 
-HRESULT FD3D12Adapter::CreateCommittedResource(const D3D12_RESOURCE_DESC& InDesc, FRHIGPUMask CreationNode, const D3D12_HEAP_PROPERTIES& HeapProps, const D3D12_RESOURCE_STATES& InitialUsage, const D3D12_CLEAR_VALUE* ClearValue, FD3D12Resource** ppOutResource, const TCHAR* Name)
+HRESULT FD3D12Adapter::CreateCommittedResource(const D3D12_RESOURCE_DESC& InDesc, FRHIGPUMask CreationNode, const D3D12_HEAP_PROPERTIES& HeapProps, const D3D12_RESOURCE_STATES& InitialUsage, const D3D12_CLEAR_VALUE* ClearValue, FD3D12Resource** ppOutResource, const TCHAR* Name, bool bVerifyHResult)
 {
 	if (!ppOutResource)
 	{
@@ -316,7 +316,10 @@ HRESULT FD3D12Adapter::CreateCommittedResource(const D3D12_RESOURCE_DESC& InDesc
 	}
 
 	const HRESULT hr = RootDevice->CreateCommittedResource(&HeapProps, HeapFlags, &InDesc, InitialUsage, ClearValue, IID_PPV_ARGS(pResource.GetInitReference()));
-	VERIFYD3D12RESULT_EX(hr, RootDevice);
+	if (bVerifyHResult)
+	{
+		VERIFYD3D12RESULT_EX(hr, RootDevice);
+	}
 
 	if (SUCCEEDED(hr))
 	{
@@ -337,7 +340,7 @@ HRESULT FD3D12Adapter::CreateCommittedResource(const D3D12_RESOURCE_DESC& InDesc
 	return hr;
 }
 
-HRESULT FD3D12Adapter::CreatePlacedResource(const D3D12_RESOURCE_DESC& InDesc, FD3D12Heap* BackingHeap, uint64 HeapOffset, const D3D12_RESOURCE_STATES& InitialUsage, const D3D12_CLEAR_VALUE* ClearValue, FD3D12Resource** ppOutResource, const TCHAR* Name)
+HRESULT FD3D12Adapter::CreatePlacedResource(const D3D12_RESOURCE_DESC& InDesc, FD3D12Heap* BackingHeap, uint64 HeapOffset, const D3D12_RESOURCE_STATES& InitialUsage, const D3D12_CLEAR_VALUE* ClearValue, FD3D12Resource** ppOutResource, const TCHAR* Name, bool bVerifyHResult)
 {
 	if (!ppOutResource)
 	{
@@ -348,7 +351,11 @@ HRESULT FD3D12Adapter::CreatePlacedResource(const D3D12_RESOURCE_DESC& InDesc, F
 
 	TRefCountPtr<ID3D12Resource> pResource;
 	const HRESULT hr = RootDevice->CreatePlacedResource(Heap, HeapOffset, &InDesc, InitialUsage, nullptr, IID_PPV_ARGS(pResource.GetInitReference()));
-	VERIFYD3D12RESULT_EX(hr, RootDevice);
+
+	if (bVerifyHResult)
+	{
+		VERIFYD3D12RESULT_EX(hr, RootDevice);
+	}
 
 	if (SUCCEEDED(hr))
 	{
diff --git a/Engine/Source/Runtime/D3D12RHI/Private/D3D12Texture.cpp b/Engine/Source/Runtime/D3D12RHI/Private/D3D12Texture.cpp
--- a/Engine/Source/Runtime/D3D12RHI/Private/D3D12Texture.cpp
+++ b/Engine/Source/Runtime/D3D12RHI/Private/D3D12Texture.cpp
@@ -1075,7 +1075,7 @@ FD3D12Texture3D* FD3D12DynamicRHI::CreateD3D12Texture3D(FRHICommandListImmediate
 	{
 		FD3D12Texture3D* Texture3D = new FD3D12Texture3D(Device, SizeX, SizeY, SizeZ, NumMips, Format, Flags, CreateInfo.ClearValueBinding);
 
-		Device->GetTextureAllocator().AllocateTexture(TextureDesc, ClearValuePtr, Format, Texture3D->ResourceLocation, (CreateInfo.BulkData != nullptr) ? D3D12_RESOURCE_STATE_COPY_DEST : DestinationState, CreateInfo.DebugName);
+		VERIFYD3D12CREATETEXTURERESULT(Device->GetTextureAllocator().AllocateTexture(TextureDesc, ClearValuePtr, Format, Texture3D->ResourceLocation, (CreateInfo.BulkData != nullptr) ? D3D12_RESOURCE_STATE_COPY_DEST : DestinationState, CreateInfo.DebugName), TextureDesc);
 
 		if (bCreateRTV)
 		{
