DX12 fixes :
* Fixed potentially erroneous transition of when using ReadSurfaceData/GetStagingTexture
* Fixed ReadSurfaceData of DXGI_FORMAT_R8_UNORM render targets

#rnx

diff --git a/Engine/Source/Runtime/D3D12RHI/Private/D3D12RenderTarget.cpp b/Engine/Source/Runtime/D3D12RHI/Private/D3D12RenderTarget.cpp
--- a/Engine/Source/Runtime/D3D12RHI/Private/D3D12RenderTarget.cpp
+++ b/Engine/Source/Runtime/D3D12RHI/Private/D3D12RenderTarget.cpp
@@ -709,7 +709,7 @@ TRefCountPtr<FD3D12Resource> FD3D12DynamicRHI::GetStagingTexture(FRHITexture* Te
 	CD3DX12_TEXTURE_COPY_LOCATION DestCopyLocation(TempTexture2D->GetResource(), placedTexture2D);
 	CD3DX12_TEXTURE_COPY_LOCATION SourceCopyLocation(Texture->GetResource()->GetResource(), Subresource);
 
-	TransitionResource(hCommandList, Texture->GetResource(), D3D12_RESOURCE_STATE_COPY_SOURCE, SourceCopyLocation.SubresourceIndex);
+	FConditionalScopeResourceBarrier ScopeResourceBarrierSource(hCommandList, Texture->GetResource(), D3D12_RESOURCE_STATE_COPY_SOURCE, SourceCopyLocation.SubresourceIndex);
 	hCommandList.FlushResourceBarriers();
 	// Upload heap doesn't need to transition
 
@@ -1232,7 +1232,7 @@ void FD3D12DynamicRHI::RHIReadSurfaceData(FRHITexture* TextureRHI, FIntRect InRe
 	FPixelFormatInfo FormatInfo = GPixelFormats[TextureRHI->GetFormat()];
 	uint32 BytesPerPixel = FormatInfo.BlockBytes;
 	uint32 SrcPitch = SizeX * BytesPerPixel;
-	
+
 	ConvertDXGIToFColor((DXGI_FORMAT)FormatInfo.PlatformFormat, SizeX, SizeY, OutDataRaw.GetData(), SrcPitch, OutData.GetData(), InFlags);
 }
 
@@ -1355,7 +1355,7 @@ void FD3D12DynamicRHI::ReadSurfaceDataMSAARaw(FRHICommandList_RecursiveHazardous
 			SampleIndex
 			);
 
-		TransitionResource(hCommandList, NonMSAATexture2D, D3D12_RESOURCE_STATE_COPY_SOURCE, SourceCopyLocation.SubresourceIndex);
+		FConditionalScopeResourceBarrier ScopeResourceBarrierSource(hCommandList, NonMSAATexture2D, D3D12_RESOURCE_STATE_COPY_SOURCE, SourceCopyLocation.SubresourceIndex);
 		// Upload heap doesn't need to transition
 
 		DefaultContext.numCopies++;
