Async PSO creation to hide latency

Resolved conflict in d FD3D12PipelineStateWorker::DoWork where delete was added and function renamed
#rnx

diff --git a/Engine/Source/Runtime/D3D12RHI/Private/D3D12PipelineState.h b/Engine/Source/Runtime/D3D12RHI/Private/D3D12PipelineState.h
--- a/Engine/Source/Runtime/D3D12RHI/Private/D3D12PipelineState.h
+++ b/Engine/Source/Runtime/D3D12RHI/Private/D3D12PipelineState.h
@@ -281,8 +281,8 @@ struct FD3D12PipelineStateWorker : public FD3D12AdapterChild, public FNonAbandon
 
 	union PipelineCreationArgs
 	{
-		ComputePipelineCreationArgs_POD ComputeArgs;
-		GraphicsPipelineCreationArgs_POD GraphicsArgs;
+		ComputePipelineCreationArgs_POD* ComputeArgs;
+		GraphicsPipelineCreationArgs_POD* GraphicsArgs;
 	} CreationArgs;
 
 	const bool bIsGraphics;
diff --git a/Engine/Source/Runtime/D3D12RHI/Private/D3D12Pipelinestate.cpp b/Engine/Source/Runtime/D3D12RHI/Private/D3D12Pipelinestate.cpp
--- a/Engine/Source/Runtime/D3D12RHI/Private/D3D12Pipelinestate.cpp
+++ b/Engine/Source/Runtime/D3D12RHI/Private/D3D12Pipelinestate.cpp
@@ -123,14 +123,16 @@ FD3D12PipelineStateWorker::FD3D12PipelineStateWorker(FD3D12Adapter* Adapter, con
 	: FD3D12AdapterChild(Adapter)
 	, bIsGraphics(false)
 {
-	CreationArgs.ComputeArgs.Init(InArgs.Args);
+	CreationArgs.ComputeArgs = new ComputePipelineCreationArgs_POD();
+	CreationArgs.ComputeArgs->Init(InArgs.Args);
 };
 
 FD3D12PipelineStateWorker::FD3D12PipelineStateWorker(FD3D12Adapter* Adapter, const GraphicsPipelineCreationArgs& InArgs)
 	: FD3D12AdapterChild(Adapter)
 	, bIsGraphics(true)
 {
-	CreationArgs.GraphicsArgs.Init(InArgs.Args);
+	CreationArgs.GraphicsArgs = new GraphicsPipelineCreationArgs_POD();
+	CreationArgs.GraphicsArgs->Init(InArgs.Args);
 };
 
 /// @endcond
@@ -442,18 +444,6 @@ FD3D12PipelineState* FD3D12PipelineStateCacheBase::CreateAndAddToLowLevelCache(c
 	AddToLowLevelCache(Desc, &PipelineState, [this](FD3D12PipelineState** PipelineState, const FD3D12LowLevelGraphicsPipelineStateDesc& Desc)
 	{ 
 		OnPSOCreated(*PipelineState, Desc);
-
-		// The lock will be held at this point so we can modify the cache.
-		// Clean ourselves up if the compilation failed.
-		// Note: This check is called here instead of in AddToLowLevelCache
-		// because GetPipelineState will force a synchronization. This
-		// path is always synchronous anyway.
-		if ((*PipelineState)->GetPipelineState() == nullptr)
-		{
-			this->LowLevelGraphicsPipelineStateCache.Remove(Desc);
-			delete *PipelineState;
-			*PipelineState = nullptr;
-		}
 	});
 
 	return PipelineState;
diff --git a/Engine/Source/Runtime/D3D12RHI/Private/Windows/WindowsD3D12PipelineState.cpp b/Engine/Source/Runtime/D3D12RHI/Private/Windows/WindowsD3D12PipelineState.cpp
--- a/Engine/Source/Runtime/D3D12RHI/Private/Windows/WindowsD3D12PipelineState.cpp
+++ b/Engine/Source/Runtime/D3D12RHI/Private/Windows/WindowsD3D12PipelineState.cpp
@@ -117,7 +117,7 @@ D3D12_COMPUTE_PIPELINE_STATE_DESC FD3D12_COMPUTE_PIPELINE_STATE_DESC::ComputeDes
 void FD3D12PipelineStateCache::OnPSOCreated(FD3D12PipelineState* PipelineState, const FD3D12LowLevelGraphicsPipelineStateDesc& Desc)
 {
 	// Actually create the PSO.
-	PipelineState->Create(GraphicsPipelineCreationArgs(&Desc, PipelineLibrary.GetReference()));
+	PipelineState->CreateAsync(GraphicsPipelineCreationArgs(&Desc, PipelineLibrary.GetReference()));
 
 	// Save this PSO to disk cache.
 	if (!DiskCaches[PSO_CACHE_GRAPHICS].IsInErrorState())
@@ -130,7 +130,7 @@ void FD3D12PipelineStateCache::OnPSOCreated(FD3D12PipelineState* PipelineState,
 void FD3D12PipelineStateCache::OnPSOCreated(FD3D12PipelineState* PipelineState, const FD3D12ComputePipelineStateDesc& Desc)
 {
 	// Actually create the PSO.
-	PipelineState->Create(ComputePipelineCreationArgs(&Desc, PipelineLibrary.GetReference()));
+	PipelineState->CreateAsync(ComputePipelineCreationArgs(&Desc, PipelineLibrary.GetReference()));
 
 	// Save this PSO to disk cache.
 	if (!DiskCaches[PSO_CACHE_COMPUTE].IsInErrorState())
@@ -659,7 +659,7 @@ static void CreatePipelineStateWrapper(ID3D12PipelineState** PSO, FD3D12Adapter*
 {
 	// Get the pipeline state name, currently based on the hash.
 	wchar_t Name[17];
-	FastHashName(Name, CreationArgs->Desc->CombinedHash);
+	FastHashName(Name, CreationArgs->Desc.CombinedHash);
 
 #if LOG_PSO_CREATES
 	const FString CreatePipelineStateMessage = FString::Printf(TEXT("CreatePipelineState (%s, Hash = %s)"), TPSOStreamFunctionMap<TDesc>::GetString(), Name);
@@ -670,14 +670,14 @@ static void CreatePipelineStateWrapper(ID3D12PipelineState** PSO, FD3D12Adapter*
 	ID3D12Device2* const pDevice2 = Adapter->GetD3DDevice2();
 	if (pDevice2)
 	{
-		typename TPSOStreamFunctionMap<TDesc>::D3D12PipelineStateStreamType Stream = CreationArgs->Desc->Desc.PipelineStateStream();
+		typename TPSOStreamFunctionMap<TDesc>::D3D12PipelineStateStreamType Stream = CreationArgs->Desc.Desc.PipelineStateStream();
 		const D3D12_PIPELINE_STATE_STREAM_DESC StreamDesc = { sizeof(Stream), &Stream };
 		CreatePipelineStateFromStream(*PSO, pDevice2, &StreamDesc, static_cast<ID3D12PipelineLibrary1*>(CreationArgs->Library), Name);	// Static cast to ID3D12PipelineLibrary1 since we already checked for ID3D12Device2.
 	}
 	else
 	{
 		// Use the older pipeline descs.
-		const typename TPSOStreamFunctionMap<TDesc>::D3D12PipelineStateDescV0Type Desc = (CreationArgs->Desc->Desc.*TPSOStreamFunctionMap<TDesc>::GetPipelineStateDescV0())();
+		const typename TPSOStreamFunctionMap<TDesc>::D3D12PipelineStateDescV0Type Desc = (CreationArgs->Desc.Desc.*TPSOStreamFunctionMap<TDesc>::GetPipelineStateDescV0())();
 		CreatePipelineState(*PSO, Adapter->GetD3DDevice(), &Desc, CreationArgs->Library, Name);
 	}
 }
@@ -844,10 +844,12 @@ void FD3D12PipelineStateWorker::DoWork()
 {
 	if (bIsGraphics)
 	{
-		CreateGraphicsPipelineState(PSO.GetInitReference(), GetParentAdapter(), &CreationArgs.GraphicsArgs);
+		CreateGraphicsPipelineState(PSO.GetInitReference(), GetParentAdapter(), CreationArgs.GraphicsArgs);
+		delete CreationArgs.GraphicsArgs;
 	}
 	else
 	{
-		CreateComputePipelineState(PSO.GetInitReference(), GetParentAdapter(), &CreationArgs.ComputeArgs);
+		CreateGraphicsPipelineState(PSO.GetInitReference(), GetParentAdapter(), CreationArgs.ComputeArgs);
+		delete CreationArgs.ComputeArgs;
 	}
 }
diff --git a/Engine/Source/Runtime/D3D12RHI/Private/Windows/WindowsD3D12PipelineState.h b/Engine/Source/Runtime/D3D12RHI/Private/Windows/WindowsD3D12PipelineState.h
--- a/Engine/Source/Runtime/D3D12RHI/Private/Windows/WindowsD3D12PipelineState.h
+++ b/Engine/Source/Runtime/D3D12RHI/Private/Windows/WindowsD3D12PipelineState.h
@@ -50,12 +50,17 @@ struct FD3D12_COMPUTE_PIPELINE_STATE_STREAM
 	CD3DX12_PIPELINE_STATE_STREAM_CACHED_PSO CachedPSO;
 };
 
+struct ComputePipelineCreationArgs_POD;
+struct GraphicsPipelineCreationArgs_POD;
+
+#include "D3D12PipelineState.h"
+
 struct ComputePipelineCreationArgs_POD
 {
-	const FD3D12ComputePipelineStateDesc* Desc;
+	FD3D12ComputePipelineStateDesc Desc;
 	ID3D12PipelineLibrary* Library;
 
-	void Init(ComputePipelineCreationArgs_POD InArgs)
+	void Init(const ComputePipelineCreationArgs_POD& InArgs)
 	{
 		Desc = InArgs.Desc;
 		Library = InArgs.Library;
@@ -64,79 +69,43 @@ struct ComputePipelineCreationArgs_POD
 
 struct GraphicsPipelineCreationArgs_POD
 {
-	const FD3D12LowLevelGraphicsPipelineStateDesc* Desc;
+	FD3D12LowLevelGraphicsPipelineStateDesc Desc;
 	ID3D12PipelineLibrary* Library;
 
-	void Init(GraphicsPipelineCreationArgs_POD InArgs)
+	void Init(const GraphicsPipelineCreationArgs_POD& InArgs)
 	{
 		Desc = InArgs.Desc;
 		Library = InArgs.Library;
 	}
 };
 
-#include "D3D12PipelineState.h"
-
 struct ComputePipelineCreationArgs
 {
-	ComputePipelineCreationArgs()
-	{
-		Args.Desc = nullptr;
-		Args.Library = nullptr;
-	}
-
 	ComputePipelineCreationArgs(const FD3D12ComputePipelineStateDesc* InDesc, ID3D12PipelineLibrary* InLibrary)
 	{
-		Args.Desc = InDesc;
+		Args.Desc = *InDesc;
 		Args.Library = InLibrary;
 	}
 
 	ComputePipelineCreationArgs(const ComputePipelineCreationArgs& InArgs)
-		: ComputePipelineCreationArgs(InArgs.Args.Desc, InArgs.Args.Library)
+		: ComputePipelineCreationArgs(&InArgs.Args.Desc, InArgs.Args.Library)
 	{}
 
-	ComputePipelineCreationArgs& operator=(const ComputePipelineCreationArgs& Other)
-	{
-		if (this != &Other)
-		{
-			Args.Desc = Other.Args.Desc;
-			Args.Library = Other.Args.Library;
-		}
-
-		return *this;
-	}
-
 	ComputePipelineCreationArgs_POD Args;
 };
 
 struct GraphicsPipelineCreationArgs
 {
-	GraphicsPipelineCreationArgs()
-	{
-		Args.Desc = nullptr;
-		Args.Library = nullptr;
-	}
-
 	GraphicsPipelineCreationArgs(const FD3D12LowLevelGraphicsPipelineStateDesc* InDesc, ID3D12PipelineLibrary* InLibrary)
 	{
-		Args.Desc = InDesc;
+		Args.Desc = *InDesc;
 		Args.Library = InLibrary;
 	}
 
 	GraphicsPipelineCreationArgs(const GraphicsPipelineCreationArgs& InArgs)
-		: GraphicsPipelineCreationArgs(InArgs.Args.Desc, InArgs.Args.Library)
+		: GraphicsPipelineCreationArgs(&InArgs.Args.Desc, InArgs.Args.Library)
 	{}
 
-	GraphicsPipelineCreationArgs& operator=(const GraphicsPipelineCreationArgs& Other)
-	{
-		if (this != &Other)
-		{
-			Args.Desc = Other.Args.Desc;
-			Args.Library = Other.Args.Library;
-		}
-
-		return *this;
-	}
-
 	GraphicsPipelineCreationArgs_POD Args;
 };
 
