Try to fix up code based on changes that came in from CL# 10086999
#rnx

diff --git a/Engine/Source/Runtime/D3D12RHI/Private/Windows/WindowsD3D12PipelineState.cpp b/Engine/Source/Runtime/D3D12RHI/Private/Windows/WindowsD3D12PipelineState.cpp
--- a/Engine/Source/Runtime/D3D12RHI/Private/Windows/WindowsD3D12PipelineState.cpp
+++ b/Engine/Source/Runtime/D3D12RHI/Private/Windows/WindowsD3D12PipelineState.cpp
@@ -698,15 +698,15 @@ static FORCEINLINE NVAPI_D3D12_PSO_SET_SHADER_EXTENSION_SLOT_DESC GetNVShaderExt
 
 static void CreateGraphicsPipelineState(ID3D12PipelineState** PSO, FD3D12Adapter* Adapter, const GraphicsPipelineCreationArgs_POD* CreationArgs)
 {
-	if (CreationArgs->Desc->HasVendorExtensions())
+	if (CreationArgs->Desc.HasVendorExtensions())
 	{
 		// Need to merge extensions across all stages for a single PSO
 		TArray<FShaderCodeVendorExtension, TInlineAllocator<2 /* VS + PS */>> MergedExtensions;
 
 	#define MERGE_EXT(Initial) \
-		if (CreationArgs->Desc->##Initial##SExtensions != nullptr) \
+		if (CreationArgs->Desc.##Initial##SExtensions != nullptr) \
 		{ \
-			const TArray<FShaderCodeVendorExtension>& Extensions = *CreationArgs->Desc->##Initial##SExtensions; \
+			const TArray<FShaderCodeVendorExtension>& Extensions = *CreationArgs->Desc.##Initial##SExtensions; \
 			for (int32 ExtIndex = 0; ExtIndex < Extensions.Num(); ++ExtIndex) \
 			{ \
 				MergedExtensions.AddUnique(Extensions[ExtIndex]); \
@@ -728,7 +728,7 @@ static void CreateGraphicsPipelineState(ID3D12PipelineState** PSO, FD3D12Adapter
 				if (Extension.Parameter.Type == EShaderParameterType::UAV)
 				{
 					const typename TPSOStreamFunctionMap<GraphicsPipelineCreationArgs_POD>::D3D12PipelineStateDescV0Type PsoDesc
-						= (CreationArgs->Desc->Desc.*TPSOStreamFunctionMap<GraphicsPipelineCreationArgs_POD>::GetPipelineStateDescV0())();
+						= (CreationArgs->Desc.Desc.*TPSOStreamFunctionMap<GraphicsPipelineCreationArgs_POD>::GetPipelineStateDescV0())();
 
 					const NVAPI_D3D12_PSO_SET_SHADER_EXTENSION_SLOT_DESC ShdExtensionDesc = GetNVShaderExtensionDesc(Extension.Parameter.BaseIndex);
 					const NVAPI_D3D12_PSO_EXTENSION_DESC* NvExtensions[] = { &ShdExtensionDesc };
@@ -757,7 +757,7 @@ static void CreateGraphicsPipelineState(ID3D12PipelineState** PSO, FD3D12Adapter
 
 static void CreateComputePipelineState(ID3D12PipelineState** PSO, FD3D12Adapter* Adapter, const ComputePipelineCreationArgs_POD* CreationArgs)
 {
-	if (CreationArgs->Desc->HasVendorExtensions())
+	if (CreationArgs->Desc.HasVendorExtensions())
 	{
 		const TArray<FShaderCodeVendorExtension>& Extensions = *CreationArgs->Desc->Extensions;
 		for (int32 ExtIndex = 0; ExtIndex < Extensions.Num(); ++ExtIndex)
