Fix additional bad merge errors
#rnx

diff --git a/Engine/Source/Runtime/D3D12RHI/Private/Windows/WindowsD3D12PipelineState.cpp b/Engine/Source/Runtime/D3D12RHI/Private/Windows/WindowsD3D12PipelineState.cpp
--- a/Engine/Source/Runtime/D3D12RHI/Private/Windows/WindowsD3D12PipelineState.cpp
+++ b/Engine/Source/Runtime/D3D12RHI/Private/Windows/WindowsD3D12PipelineState.cpp
@@ -759,7 +759,7 @@ static void CreateComputePipelineState(ID3D12PipelineState** PSO, FD3D12Adapter*
 {
 	if (CreationArgs->Desc.HasVendorExtensions())
 	{
-		const TArray<FShaderCodeVendorExtension>& Extensions = *CreationArgs->Desc->Extensions;
+		const TArray<FShaderCodeVendorExtension>& Extensions = *CreationArgs->Desc.Extensions;
 		for (int32 ExtIndex = 0; ExtIndex < Extensions.Num(); ++ExtIndex)
 		{
 			const FShaderCodeVendorExtension& Extension = Extensions[ExtIndex];
@@ -768,7 +768,7 @@ static void CreateComputePipelineState(ID3D12PipelineState** PSO, FD3D12Adapter*
 				if (Extension.Parameter.Type == EShaderParameterType::UAV)
 				{
 					const typename TPSOStreamFunctionMap<ComputePipelineCreationArgs_POD>::D3D12PipelineStateDescV0Type PsoDesc
-						= (CreationArgs->Desc->Desc.*TPSOStreamFunctionMap<ComputePipelineCreationArgs_POD>::GetPipelineStateDescV0())();
+						= (CreationArgs->Desc.Desc.*TPSOStreamFunctionMap<ComputePipelineCreationArgs_POD>::GetPipelineStateDescV0())();
 
 					const NVAPI_D3D12_PSO_SET_SHADER_EXTENSION_SLOT_DESC ShdExtensionDesc = GetNVShaderExtensionDesc(Extension.Parameter.BaseIndex);
 					const NVAPI_D3D12_PSO_EXTENSION_DESC* NvExtensions[] = { &ShdExtensionDesc };
@@ -849,7 +849,7 @@ void FD3D12PipelineStateWorker::DoWork()
 	}
 	else
 	{
-		CreateGraphicsPipelineState(PSO.GetInitReference(), GetParentAdapter(), CreationArgs.ComputeArgs);
+		CreateComputePipelineState(PSO.GetInitReference(), GetParentAdapter(), CreationArgs.ComputeArgs);
 		delete CreationArgs.ComputeArgs;
 	}
 }
