UE4R - Support for -vulkansdk to use validation layers from the SDK instead of the ones we provide
- Support for D3D12 with -rhivalidation
#rnx

diff --git a/Engine/Source/Runtime/D3D12RHI/Private/D3D12Device.cpp b/Engine/Source/Runtime/D3D12RHI/Private/D3D12Device.cpp
--- a/Engine/Source/Runtime/D3D12RHI/Private/D3D12Device.cpp
+++ b/Engine/Source/Runtime/D3D12RHI/Private/D3D12Device.cpp
@@ -4,6 +4,7 @@
 D3D12Device.cpp: D3D device RHI implementation.
 =============================================================================*/
 #include "D3D12RHIPrivate.h"
+#include "RHIValidation.h"
 
 
 namespace D3D12RHI
@@ -235,7 +236,7 @@ void FD3D12Device::SetupAfterDeviceCreation()
 
 	if(bUnderGPUCapture)
 	{
-		GDynamicRHI->EnableIdealGPUCaptureOptions(true);
+		GetDynamicRHI<FD3D12DynamicRHI>()->EnableIdealGPUCaptureOptions(true);
 	}
 #endif
 
diff --git a/Engine/Source/Runtime/D3D12RHI/Private/Windows/WindowsD3D12Device.cpp b/Engine/Source/Runtime/D3D12RHI/Private/Windows/WindowsD3D12Device.cpp
--- a/Engine/Source/Runtime/D3D12RHI/Private/Windows/WindowsD3D12Device.cpp
+++ b/Engine/Source/Runtime/D3D12RHI/Private/Windows/WindowsD3D12Device.cpp
@@ -23,6 +23,7 @@
 #include "HardwareInfo.h"
 #include "Runtime/HeadMountedDisplay/Public/IHeadMountedDisplayModule.h"
 #include "GenericPlatform/GenericPlatformDriver.h"			// FGPUDriverInfo
+#include "RHIValidation.h"
 
 #include "ShaderCompiler.h"
 
@@ -35,6 +36,8 @@ extern bool D3D12RHI_ShouldCreateWithWarp();
 extern bool D3D12RHI_ShouldAllowAsyncResourceCreation();
 extern bool D3D12RHI_ShouldForceCompatibility();
 
+FD3D12DynamicRHI* GD3D12RHI = nullptr;
+
 #if NV_AFTERMATH
 // Disabled by default since introduces stalls between render and driver threads
 int32 GDX12NVAfterMathEnabled = 0;
@@ -462,7 +465,21 @@ FDynamicRHI* FD3D12DynamicRHIModule::CreateRHI(ERHIFeatureLevel::Type RequestedF
 		GMaxRHIShaderPlatform = SP_PCD3D_SM5;
 	}
 
-	return new FD3D12DynamicRHI(ChosenAdapters);
+	GD3D12RHI = new FD3D12DynamicRHI(ChosenAdapters);
+#if ENABLE_RHI_VALIDATION
+	if (FParse::Param(FCommandLine::Get(), TEXT("RHIValidation")))
+	{
+		GValidationRHI = new FValidationRHI(GD3D12RHI);
+	}
+	else
+	{
+		check(!GValidationRHI);
+	}
+
+	return GValidationRHI ? (FDynamicRHI*)GValidationRHI : (FDynamicRHI*)GD3D12RHI;
+#else
+	return GD3D12RHI;
+#endif
 }
 
 void FD3D12DynamicRHIModule::StartupModule()
@@ -692,7 +709,7 @@ void FD3D12DynamicRHI::Init()
 
 	GSupportsDepthBoundsTest = SupportsDepthBoundsTest(this);
 
-	GRHICommandList.GetImmediateCommandList().SetContext(GDynamicRHI->RHIGetDefaultContext());
+	GRHICommandList.GetImmediateCommandList().SetContext(GetDynamicRHI<FD3D12DynamicRHI>()->RHIGetDefaultContext());
 	GRHICommandList.GetImmediateCommandList().SetGPUMask(FRHIGPUMask::All());
 	GRHICommandList.GetImmediateAsyncComputeCommandList().SetComputeContext(GDynamicRHI->RHIGetDefaultAsyncComputeContext());
 	GRHICommandList.GetImmediateAsyncComputeCommandList().SetGPUMask(FRHIGPUMask::All());
diff --git a/Engine/Source/Runtime/RHI/Private/RHIValidation.cpp b/Engine/Source/Runtime/RHI/Private/RHIValidation.cpp
--- a/Engine/Source/Runtime/RHI/Private/RHIValidation.cpp
+++ b/Engine/Source/Runtime/RHI/Private/RHIValidation.cpp
@@ -19,6 +19,8 @@ FValidationRHI::FValidationRHI(FDynamicRHI* InRHI)
 
 	Context = new FValidationContext(this);
 	AsyncComputeContext = new FValidationComputeContext(this);
+
+	UE_LOG(LogRHI, Warning, TEXT("FValidationRHI on, intercepting %s RHI!"), InRHI && InRHI->GetName() ? InRHI->GetName() : TEXT("<NULL>"));
 }
 
 FValidationRHI::~FValidationRHI()
diff --git a/Engine/Source/Runtime/VulkanRHI/Private/Windows/VulkanWindowsPlatform.cpp b/Engine/Source/Runtime/VulkanRHI/Private/Windows/VulkanWindowsPlatform.cpp
--- a/Engine/Source/Runtime/VulkanRHI/Private/Windows/VulkanWindowsPlatform.cpp
+++ b/Engine/Source/Runtime/VulkanRHI/Private/Windows/VulkanWindowsPlatform.cpp
@@ -28,17 +28,21 @@ bool FVulkanWindowsPlatform::LoadVulkanLibrary()
 #if VULKAN_HAS_DEBUGGING_ENABLED
 	if (GValidationCvar->GetInt() > 0)
 	{
-		const FString PreviousEnvVar = FPlatformMisc::GetEnvironmentVariable(TEXT("VK_LAYER_PATH"));
-		if (PreviousEnvVar.IsEmpty())
+		const bool bUseSDK = FParse::Param(FCommandLine::Get(), TEXT("vulkansdk"));
+		if (!bUseSDK)
 		{
-			// Change behavior of loading Vulkan layers by setting environment variable "VK_LAYER_PATH" to UE4 specific directory
-			FString VulkanLayerPath = FPaths::EngineDir();
+			const FString PreviousEnvVar = FPlatformMisc::GetEnvironmentVariable(TEXT("VK_LAYER_PATH"));
+			if (PreviousEnvVar.IsEmpty())
+			{
+				// Change behavior of loading Vulkan layers by setting environment variable "VK_LAYER_PATH" to UE4 specific directory
+				FString VulkanLayerPath = FPaths::EngineDir();
 #if PLATFORM_64BITS
-			VulkanLayerPath.Append(TEXT("Binaries/ThirdParty/Windows/Vulkan/Win64"));
+				VulkanLayerPath.Append(TEXT("Binaries/ThirdParty/Windows/Vulkan/Win64"));
 #else
-			VulkanLayerPath.Append(TEXT("Binaries/ThirdParty/Windows/Vulkan/Win32"));
+				VulkanLayerPath.Append(TEXT("Binaries/ThirdParty/Windows/Vulkan/Win32"));
 #endif
-			FPlatformMisc::SetEnvironmentVar(TEXT("VK_LAYER_PATH"), *VulkanLayerPath);
+				FPlatformMisc::SetEnvironmentVar(TEXT("VK_LAYER_PATH"), *VulkanLayerPath);
+			}
 		}
 	}
 #endif // VULKAN_HAS_DEBUGGING_ENABLED
