Fix DX12 Async PSO creation crash

diff --git a/Engine/Source/Runtime/D3D12RHI/Private/D3D12PipelineState.h b/Engine/Source/Runtime/D3D12RHI/Private/D3D12PipelineState.h
--- a/Engine/Source/Runtime/D3D12RHI/Private/D3D12PipelineState.h
+++ b/Engine/Source/Runtime/D3D12RHI/Private/D3D12PipelineState.h
@@ -78,6 +78,7 @@ struct FD3D12LowLevelGraphicsPipelineStateDesc
 	ShaderBytecodeHash GSHash;
 	ShaderBytecodeHash PSHash;
 	uint32 InputLayoutHash;
+	bool bFromPSOFileCache;
 
 	SIZE_T CombinedHash;
 
diff --git a/Engine/Source/Runtime/D3D12RHI/Private/D3D12Pipelinestate.cpp b/Engine/Source/Runtime/D3D12RHI/Private/D3D12Pipelinestate.cpp
--- a/Engine/Source/Runtime/D3D12RHI/Private/D3D12Pipelinestate.cpp
+++ b/Engine/Source/Runtime/D3D12RHI/Private/D3D12Pipelinestate.cpp
@@ -96,6 +96,8 @@ FD3D12LowLevelGraphicsPipelineStateDesc GetLowLevelGraphicsPipelineStateDesc(con
 	Desc.Desc.DepthStencilState.DepthBoundsTestEnable = GSupportsDepthBoundsTest && Initializer.bDepthBounds;
 #endif
 
+	Desc.bFromPSOFileCache = Initializer.bFromPSOFileCache;
+
 	return Desc;
 }
 
diff --git a/Engine/Source/Runtime/D3D12RHI/Private/Windows/WindowsD3D12PipelineState.cpp b/Engine/Source/Runtime/D3D12RHI/Private/Windows/WindowsD3D12PipelineState.cpp
--- a/Engine/Source/Runtime/D3D12RHI/Private/Windows/WindowsD3D12PipelineState.cpp
+++ b/Engine/Source/Runtime/D3D12RHI/Private/Windows/WindowsD3D12PipelineState.cpp
@@ -116,8 +116,18 @@ D3D12_COMPUTE_PIPELINE_STATE_DESC FD3D12_COMPUTE_PIPELINE_STATE_DESC::ComputeDes
 
 void FD3D12PipelineStateCache::OnPSOCreated(FD3D12PipelineState* PipelineState, const FD3D12LowLevelGraphicsPipelineStateDesc& Desc)
 {
+	const bool bAsync = !Desc.bFromPSOFileCache;
+
 	// Actually create the PSO.
-	PipelineState->CreateAsync(GraphicsPipelineCreationArgs(&Desc, PipelineLibrary.GetReference()));
+	GraphicsPipelineCreationArgs Args(&Desc, PipelineLibrary.GetReference());
+	if (bAsync)
+	{
+		PipelineState->CreateAsync(Args);
+	}
+	else
+	{
+		PipelineState->Create(Args);
+	}
 
 	// Save this PSO to disk cache.
 	if (!DiskCaches[PSO_CACHE_GRAPHICS].IsInErrorState())
@@ -129,8 +139,18 @@ void FD3D12PipelineStateCache::OnPSOCreated(FD3D12PipelineState* PipelineState,
 
 void FD3D12PipelineStateCache::OnPSOCreated(FD3D12PipelineState* PipelineState, const FD3D12ComputePipelineStateDesc& Desc)
 {
+	const bool bAsync = false;
+
 	// Actually create the PSO.
-	PipelineState->CreateAsync(ComputePipelineCreationArgs(&Desc, PipelineLibrary.GetReference()));
+	ComputePipelineCreationArgs Args(&Desc, PipelineLibrary.GetReference());
+	if (bAsync)
+	{
+		PipelineState->CreateAsync(Args);
+	}
+	else
+	{
+		PipelineState->Create(Args);
+	}
 
 	// Save this PSO to disk cache.
 	if (!DiskCaches[PSO_CACHE_COMPUTE].IsInErrorState())
@@ -845,11 +865,13 @@ void FD3D12PipelineStateWorker::DoWork()
 	if (bIsGraphics)
 	{
 		CreateGraphicsPipelineState(PSO.GetInitReference(), GetParentAdapter(), CreationArgs.GraphicsArgs);
+		CreationArgs.GraphicsArgs->Destroy();
 		delete CreationArgs.GraphicsArgs;
 	}
 	else
 	{
 		CreateComputePipelineState(PSO.GetInitReference(), GetParentAdapter(), CreationArgs.ComputeArgs);
+		CreationArgs.ComputeArgs->Destroy();
 		delete CreationArgs.ComputeArgs;
 	}
 }
diff --git a/Engine/Source/Runtime/D3D12RHI/Private/Windows/WindowsD3D12PipelineState.h b/Engine/Source/Runtime/D3D12RHI/Private/Windows/WindowsD3D12PipelineState.h
--- a/Engine/Source/Runtime/D3D12RHI/Private/Windows/WindowsD3D12PipelineState.h
+++ b/Engine/Source/Runtime/D3D12RHI/Private/Windows/WindowsD3D12PipelineState.h
@@ -55,6 +55,16 @@ struct GraphicsPipelineCreationArgs_POD;
 
 #include "D3D12PipelineState.h"
 
+static void SaveByteCode(D3D12_SHADER_BYTECODE& ByteCode)
+{
+	if (ByteCode.pShaderBytecode)
+	{
+		char* NewBytes = new char[ByteCode.BytecodeLength];
+		memcpy(NewBytes, ByteCode.pShaderBytecode, ByteCode.BytecodeLength);
+		ByteCode.pShaderBytecode = NewBytes;
+	}
+}
+
 struct ComputePipelineCreationArgs_POD
 {
 	FD3D12ComputePipelineStateDesc Desc;
@@ -64,6 +74,13 @@ struct ComputePipelineCreationArgs_POD
 	{
 		Desc = InArgs.Desc;
 		Library = InArgs.Library;
+
+		SaveByteCode(Desc.Desc.CS);
+	}
+
+	void Destroy()
+	{
+		delete[] Desc.Desc.CS.pShaderBytecode;
 	}
 };
 
@@ -76,6 +93,21 @@ struct GraphicsPipelineCreationArgs_POD
 	{
 		Desc = InArgs.Desc;
 		Library = InArgs.Library;
+
+		SaveByteCode(Desc.Desc.VS);
+		SaveByteCode(Desc.Desc.PS);
+		SaveByteCode(Desc.Desc.GS);
+		SaveByteCode(Desc.Desc.HS);
+		SaveByteCode(Desc.Desc.DS);
+	}
+
+	void Destroy()
+	{
+		delete[] Desc.Desc.VS.pShaderBytecode;
+		delete[] Desc.Desc.PS.pShaderBytecode;
+		delete[] Desc.Desc.GS.pShaderBytecode;
+		delete[] Desc.Desc.HS.pShaderBytecode;
+		delete[] Desc.Desc.DS.pShaderBytecode;
 	}
 };
 
