Use FMemory::Malloc, Memcpy, and Free to consistently manage memory with other code that interacts with shader byte code and fix static analysis warnings/compile errors
#rnx

diff --git a/Engine/Source/Runtime/D3D12RHI/Private/Windows/WindowsD3D12PipelineState.cpp b/Engine/Source/Runtime/D3D12RHI/Private/Windows/WindowsD3D12PipelineState.cpp
--- a/Engine/Source/Runtime/D3D12RHI/Private/Windows/WindowsD3D12PipelineState.cpp
+++ b/Engine/Source/Runtime/D3D12RHI/Private/Windows/WindowsD3D12PipelineState.cpp
@@ -114,6 +114,30 @@ D3D12_COMPUTE_PIPELINE_STATE_DESC FD3D12_COMPUTE_PIPELINE_STATE_DESC::ComputeDes
 	return D;
 }
 
+void SaveByteCode(D3D12_SHADER_BYTECODE& ByteCode);
+{
+	if (ByteCode.pShaderBytecode)
+	{
+		void* NewBytes = FMemory::Malloc(ByteCode.BytecodeLength);
+		FMemory::Memcpy(NewBytes, ByteCode.pShaderBytecode, ByteCode.BytecodeLength);
+		ByteCode.pShaderBytecode = NewBytes;
+	}
+}
+
+void ComputePipelineCreationArgs_POD::Destroy()
+{
+	FMemory::Free((void*)Desc.Desc.CS.pShaderBytecode);
+}
+
+void GraphicsPipelineCreationArgs_POD::Destroy()
+{
+	FMemory::Free((void*)Desc.Desc.VS.pShaderBytecode);
+	FMemory::Free((void*)Desc.Desc.PS.pShaderBytecode);
+	FMemory::Free((void*)Desc.Desc.GS.pShaderBytecode);
+	FMemory::Free((void*)Desc.Desc.HS.pShaderBytecode);
+	FMemory::Free((void*)Desc.Desc.DS.pShaderBytecode);
+}
+
 void FD3D12PipelineStateCache::OnPSOCreated(FD3D12PipelineState* PipelineState, const FD3D12LowLevelGraphicsPipelineStateDesc& Desc)
 {
 	const bool bAsync = !Desc.bFromPSOFileCache;
diff --git a/Engine/Source/Runtime/D3D12RHI/Private/Windows/WindowsD3D12PipelineState.h b/Engine/Source/Runtime/D3D12RHI/Private/Windows/WindowsD3D12PipelineState.h
--- a/Engine/Source/Runtime/D3D12RHI/Private/Windows/WindowsD3D12PipelineState.h
+++ b/Engine/Source/Runtime/D3D12RHI/Private/Windows/WindowsD3D12PipelineState.h
@@ -55,15 +55,7 @@ struct GraphicsPipelineCreationArgs_POD;
 
 #include "D3D12PipelineState.h"
 
-static void SaveByteCode(D3D12_SHADER_BYTECODE& ByteCode)
-{
-	if (ByteCode.pShaderBytecode)
-	{
-		char* NewBytes = new char[ByteCode.BytecodeLength];
-		memcpy(NewBytes, ByteCode.pShaderBytecode, ByteCode.BytecodeLength);
-		ByteCode.pShaderBytecode = NewBytes;
-	}
-}
+static void SaveByteCode(D3D12_SHADER_BYTECODE& ByteCode);
 
 struct ComputePipelineCreationArgs_POD
 {
@@ -78,10 +70,7 @@ struct ComputePipelineCreationArgs_POD
 		SaveByteCode(Desc.Desc.CS);
 	}
 
-	void Destroy()
-	{
-		delete[] Desc.Desc.CS.pShaderBytecode;
-	}
+	void Destroy();
 };
 
 struct GraphicsPipelineCreationArgs_POD
@@ -101,14 +90,7 @@ struct GraphicsPipelineCreationArgs_POD
 		SaveByteCode(Desc.Desc.DS);
 	}
 
-	void Destroy()
-	{
-		delete[] Desc.Desc.VS.pShaderBytecode;
-		delete[] Desc.Desc.PS.pShaderBytecode;
-		delete[] Desc.Desc.GS.pShaderBytecode;
-		delete[] Desc.Desc.HS.pShaderBytecode;
-		delete[] Desc.Desc.DS.pShaderBytecode;
-	}
+	void Destroy();
 };
 
 struct ComputePipelineCreationArgs
