Fix extra semi-colon that snuck in to change
#rnx

diff --git a/Engine/Source/Runtime/D3D12RHI/Private/Windows/WindowsD3D12PipelineState.cpp b/Engine/Source/Runtime/D3D12RHI/Private/Windows/WindowsD3D12PipelineState.cpp
--- a/Engine/Source/Runtime/D3D12RHI/Private/Windows/WindowsD3D12PipelineState.cpp
+++ b/Engine/Source/Runtime/D3D12RHI/Private/Windows/WindowsD3D12PipelineState.cpp
@@ -114,7 +114,7 @@ D3D12_COMPUTE_PIPELINE_STATE_DESC FD3D12_COMPUTE_PIPELINE_STATE_DESC::ComputeDes
 	return D;
 }
 
-void SaveByteCode(D3D12_SHADER_BYTECODE& ByteCode);
+void SaveByteCode(D3D12_SHADER_BYTECODE& ByteCode)
 {
 	if (ByteCode.pShaderBytecode)
 	{
