Fixed static function declared but not defined
#rnx

diff --git a/Engine/Source/Runtime/D3D12RHI/Private/Windows/WindowsD3D12PipelineState.h b/Engine/Source/Runtime/D3D12RHI/Private/Windows/WindowsD3D12PipelineState.h
--- a/Engine/Source/Runtime/D3D12RHI/Private/Windows/WindowsD3D12PipelineState.h
+++ b/Engine/Source/Runtime/D3D12RHI/Private/Windows/WindowsD3D12PipelineState.h
@@ -55,7 +55,7 @@ struct GraphicsPipelineCreationArgs_POD;
 
 #include "D3D12PipelineState.h"
 
-static void SaveByteCode(D3D12_SHADER_BYTECODE& ByteCode);
+void SaveByteCode(D3D12_SHADER_BYTECODE& ByteCode);
 
 struct ComputePipelineCreationArgs_POD
 {
