Fix for some DX12 crashes on fullscreen mode switches

diff --git a/Engine/Source/Runtime/D3D12RHI/Private/D3D12Viewport.cpp b/Engine/Source/Runtime/D3D12RHI/Private/D3D12Viewport.cpp
--- a/Engine/Source/Runtime/D3D12RHI/Private/D3D12Viewport.cpp
+++ b/Engine/Source/Runtime/D3D12RHI/Private/D3D12Viewport.cpp
@@ -216,6 +216,9 @@ FD3D12Texture2D* GetSwapChainSurface(FD3D12Device* Parent, EPixelFormat PixelFor
 		Parent->GetDevice()->CreateCommittedResource(&HeapProps, D3D12_HEAP_FLAG_NONE, &TextureDesc, D3D12_RESOURCE_STATE_PRESENT, nullptr, IID_PPV_ARGS(BackBufferResource.GetInitReference()));
 	}
 
+	FString Name = FString::Printf(TEXT("BackBuffer%d"), BackBufferIndex);
+	SetName(BackBufferResource, *Name);
+
 	D3D12_RESOURCE_DESC BackBufferDesc = BackBufferResource->GetDesc();
 
 	FD3D12Texture2D* SwapChainTexture = Adapter->CreateLinkedObject<FD3D12Texture2D>(FRHIGPUMask::All(), [&](FD3D12Device* Device)
@@ -513,6 +516,28 @@ static bool IsCompositionEnabled()
 /** Presents the swap chain checking the return result. */
 bool FD3D12Viewport::PresentChecked(int32 SyncInterval)
 {
+#if PLATFORM_WINDOWS
+	// We can't call Present if !bIsValid, as it waits a window message to be processed, but the main thread may not be pumping the message handler.
+	if (bIsValid && SwapChain1.IsValid())
+	{
+		// Check if the viewport's swap chain has been invalidated by DXGI.
+		BOOL bSwapChainFullscreenState;
+		TRefCountPtr<IDXGIOutput> SwapChainOutput;
+		SwapChain1->GetFullscreenState(&bSwapChainFullscreenState, SwapChainOutput.GetInitReference());
+		// Can't compare BOOL with bool...
+		if ( (!!bSwapChainFullscreenState)  != bIsFullscreen )
+		{
+			bFullscreenLost = true;
+			bIsValid = false;
+		}
+	}
+
+	if (!bIsValid)
+	{
+		return false;
+	}
+#endif
+
 	HRESULT Result = S_OK;
 	bool bNeedNativePresent = true;
 
diff --git a/Engine/Source/Runtime/D3D12RHI/Private/Windows/WindowsD3D12Viewport.cpp b/Engine/Source/Runtime/D3D12RHI/Private/Windows/WindowsD3D12Viewport.cpp
--- a/Engine/Source/Runtime/D3D12RHI/Private/Windows/WindowsD3D12Viewport.cpp
+++ b/Engine/Source/Runtime/D3D12RHI/Private/Windows/WindowsD3D12Viewport.cpp
@@ -28,6 +28,7 @@ FD3D12Viewport::FD3D12Viewport(class FD3D12Adapter* InParent, HWND InWindowHandl
 	SizeX(InSizeX),
 	SizeY(InSizeY),
 	bIsFullscreen(bInIsFullscreen),
+	bFullscreenLost(false),
 	PixelFormat(InPreferredPixelFormat),
 	bIsValid(true),
 	bHDRMetaDataSet(false),
@@ -228,22 +229,31 @@ void FD3D12Viewport::ConditionalResetSwapChain(bool bIgnoreFocus)
 {
 	if (!bIsValid)
 	{
-		// Check if the viewport's window is focused before resetting the swap chain's fullscreen state.
-		HWND FocusWindow = ::GetFocus();
-		const bool bIsFocused = FocusWindow == WindowHandle;
-		const bool bIsIconic = !!::IsIconic(WindowHandle);
-		if (bIgnoreFocus || (bIsFocused && !bIsIconic))
+		if (bFullscreenLost)
 		{
 			FlushRenderingCommands();
-
-			HRESULT Result = SwapChain1->SetFullscreenState(bIsFullscreen, nullptr);
-			if (SUCCEEDED(Result))
-			{
-				bIsValid = true;
-			}
-			else if (Result != DXGI_ERROR_NOT_CURRENTLY_AVAILABLE && Result != DXGI_STATUS_MODE_CHANGE_IN_PROGRESS)
+			bFullscreenLost = false;
+			Resize(SizeX, SizeY, false, PixelFormat);
+		}
+		else
+		{
+			// Check if the viewport's window is focused before resetting the swap chain's fullscreen state.
+			HWND FocusWindow = ::GetFocus();
+			const bool bIsFocused = FocusWindow == WindowHandle;
+			const bool bIsIconic = !!::IsIconic(WindowHandle);
+			if (bIgnoreFocus || (bIsFocused && !bIsIconic))
 			{
-				UE_LOG(LogD3D12RHI, Error, TEXT("IDXGISwapChain::SetFullscreenState returned %08x, unknown error status."), Result);
+				FlushRenderingCommands();
+
+				HRESULT Result = SwapChain1->SetFullscreenState(bIsFullscreen, nullptr);
+				if (SUCCEEDED(Result))
+				{
+					bIsValid = true;
+				}
+				else if (Result != DXGI_ERROR_NOT_CURRENTLY_AVAILABLE && Result != DXGI_STATUS_MODE_CHANGE_IN_PROGRESS)
+				{
+					UE_LOG(LogD3D12RHI, Error, TEXT("IDXGISwapChain::SetFullscreenState returned %08x, unknown error status."), Result);
+				}
 			}
 		}
 	}
diff --git a/Engine/Source/Runtime/D3D12RHI/Public/D3D12Viewport.h b/Engine/Source/Runtime/D3D12RHI/Public/D3D12Viewport.h
--- a/Engine/Source/Runtime/D3D12RHI/Public/D3D12Viewport.h
+++ b/Engine/Source/Runtime/D3D12RHI/Public/D3D12Viewport.h
@@ -165,6 +165,7 @@ private:
 	uint32 SizeX;
 	uint32 SizeY;
 	bool bIsFullscreen;
+	bool bFullscreenLost;
 	EPixelFormat PixelFormat;
 	bool bIsValid;
 	bool bAllowTearing;
