ResourceBarrier fix for texture initial data

diff --git a/Engine/Source/Runtime/D3D12RHI/Private/D3D12Texture.cpp b/Engine/Source/Runtime/D3D12RHI/Private/D3D12Texture.cpp
--- a/Engine/Source/Runtime/D3D12RHI/Private/D3D12Texture.cpp
+++ b/Engine/Source/Runtime/D3D12RHI/Private/D3D12Texture.cpp
@@ -190,6 +190,8 @@ struct FD3D12RHICommandInitializeTexture final : public FRHICommand<FD3D12RHICom
 			FD3D12CommandListHandle& hCommandList = Device->GetDefaultCommandContext().CommandListHandle;
 			hCommandList.GetCurrentOwningContext()->numCopies += NumSlices * NumMips;
 
+			FConditionalScopeResourceBarrier ConditionalScopeResourceBarrier(hCommandList, Resource, D3D12_RESOURCE_STATE_COPY_DEST, D3D12_RESOURCE_BARRIER_ALL_SUBRESOURCES);
+
 			ID3D12Device* D3DDevice = Device->GetDevice();
 			ID3D12GraphicsCommandList* CmdList = hCommandList.GraphicsCommandList();
 
@@ -228,15 +230,6 @@ struct FD3D12RHICommandInitializeTexture final : public FRHICommand<FD3D12RHICom
 
 			hCommandList.UpdateResidency(Resource);
 
-			if (!Resource->RequiresResourceStateTracking())
-			{
-				hCommandList.AddTransitionBarrier(Resource, D3D12_RESOURCE_STATE_COPY_DEST, DestinationState, D3D12_RESOURCE_BARRIER_ALL_SUBRESOURCES);
-			}
-			else
-			{
-				FD3D12DynamicRHI::TransitionResource(hCommandList, Resource, DestinationState, D3D12_RESOURCE_BARRIER_ALL_SUBRESOURCES);
-			}
-
 			CurrentTexture = CurrentTexture->GetNextObject();
 		}
 
