F - More info to track down D3D12 crash
#rnx

diff --git a/Engine/Source/Runtime/D3D12RHI/Private/D3D12Util.cpp b/Engine/Source/Runtime/D3D12RHI/Private/D3D12Util.cpp
--- a/Engine/Source/Runtime/D3D12RHI/Private/D3D12Util.cpp
+++ b/Engine/Source/Runtime/D3D12RHI/Private/D3D12Util.cpp
@@ -249,18 +249,17 @@ static void TerminateOnOutOfMemory(HRESULT D3DResult, bool bCreatingTextures)
 
 namespace D3D12RHI
 {
-	void VerifyD3D12Result(HRESULT D3DResult, const ANSICHAR* Code, const ANSICHAR* Filename, uint32 Line, ID3D12Device* Device)
+	void VerifyD3D12Result(HRESULT D3DResult, const ANSICHAR* Code, const ANSICHAR* Filename, uint32 Line, ID3D12Device* Device, FString Message)
 	{
 		check(FAILED(D3DResult));
 
 		const FString& ErrorString = GetD3D12ErrorString(D3DResult, Device);
-
-		UE_LOG(LogD3D12RHI, Error, TEXT("%s failed \n at %s:%u \n with error %s"), ANSI_TO_TCHAR(Code), ANSI_TO_TCHAR(Filename), Line, *ErrorString);
+		UE_LOG(LogD3D12RHI, Error, TEXT("%s failed \n at %s:%u \n with error %s\n%s"), ANSI_TO_TCHAR(Code), ANSI_TO_TCHAR(Filename), Line, *ErrorString, *Message);
 
 		TerminateOnDeviceRemoved(D3DResult);
 		TerminateOnOutOfMemory(D3DResult, false);
 
-		UE_LOG(LogD3D12RHI, Fatal, TEXT("%s failed \n at %s:%u \n with error %s"), ANSI_TO_TCHAR(Code), ANSI_TO_TCHAR(Filename), Line, *ErrorString);
+		UE_LOG(LogD3D12RHI, Fatal, TEXT("%s failed \n at %s:%u \n with error %s\n%s"), ANSI_TO_TCHAR(Code), ANSI_TO_TCHAR(Filename), Line, *ErrorString, *Message);
 	}
 
 	void VerifyD3D12CreateTextureResult(HRESULT D3DResult, const ANSICHAR* Code, const ANSICHAR* Filename, uint32 Line, const D3D12_RESOURCE_DESC& TextureDesc)
diff --git a/Engine/Source/Runtime/D3D12RHI/Private/Windows/WindowsD3D12Viewport.cpp b/Engine/Source/Runtime/D3D12RHI/Private/Windows/WindowsD3D12Viewport.cpp
--- a/Engine/Source/Runtime/D3D12RHI/Private/Windows/WindowsD3D12Viewport.cpp
+++ b/Engine/Source/Runtime/D3D12RHI/Private/Windows/WindowsD3D12Viewport.cpp
@@ -316,7 +316,12 @@ void FD3D12Viewport::ResizeInternal()
 	{
 		if (SwapChain1)
 		{
-			VERIFYD3D12RESULT_EX(SwapChain1->ResizeBuffers(NumBackBuffers, SizeX, SizeY, GetRenderTargetFormat(PixelFormat), SwapChainFlags), Adapter->GetD3DDevice());
+			auto Lambda = [=]() -> FString
+			{
+				return FString::Printf(TEXT("Num=%d, Size=(%d,%d), PF=%d, DXGIFormat=0x%x, Flags=0x%x"), NumBackBuffers, SizeX, SizeY, (int32)PixelFormat, (int32)GetRenderTargetFormat(PixelFormat), SwapChainFlags);
+			};
+
+			VERIFYD3D12RESULT_LAMBDA(SwapChain1->ResizeBuffers(NumBackBuffers, SizeX, SizeY, GetRenderTargetFormat(PixelFormat), SwapChainFlags), Adapter->GetD3DDevice(), Lambda);
 		}
 
 		FD3D12Device* Device = Adapter->GetDevice(0);
diff --git a/Engine/Source/Runtime/D3D12RHI/Public/D3D12Util.h b/Engine/Source/Runtime/D3D12RHI/Public/D3D12Util.h
--- a/Engine/Source/Runtime/D3D12RHI/Public/D3D12Util.h
+++ b/Engine/Source/Runtime/D3D12RHI/Public/D3D12Util.h
@@ -41,7 +41,7 @@ namespace D3D12RHI
 	 * @param	Filename - The filename of the source file containing Code.
 	 * @param	Line - The line number of Code within Filename.
 	 */
-	extern void VerifyD3D12Result(HRESULT Result, const ANSICHAR* Code, const ANSICHAR* Filename, uint32 Line, ID3D12Device* Device);
+	extern void VerifyD3D12Result(HRESULT Result, const ANSICHAR* Code, const ANSICHAR* Filename, uint32 Line, ID3D12Device* Device, FString Message = FString());
 
 	/**
 	* Checks that the given result isn't a failure.  If it is, the application exits with an appropriate error message.
@@ -55,6 +55,7 @@ namespace D3D12RHI
 	/**
 	 * A macro for using VERIFYD3D12RESULT that automatically passes in the code and filename/line.
 	 */
+#define VERIFYD3D12RESULT_LAMBDA(x, Device, Lambda)	{HRESULT hres = x; if (FAILED(hres)) { VerifyD3D12Result(hres, #x, __FILE__, __LINE__, Device, Lambda()); }}
 #define VERIFYD3D12RESULT_EX(x, Device)	{HRESULT hres = x; if (FAILED(hres)) { VerifyD3D12Result(hres, #x, __FILE__, __LINE__, Device); }}
 #define VERIFYD3D12RESULT(x)			{HRESULT hres = x; if (FAILED(hres)) { VerifyD3D12Result(hres, #x, __FILE__, __LINE__, nullptr); }}
 #define VERIFYD3D12CREATETEXTURERESULT(x, Desc) {HRESULT hres = x; if (FAILED(hres)) { VerifyD3D12CreateTextureResult(hres, #x, __FILE__, __LINE__, Desc); }}
