UE4R - Fix for local variable used after leaving stack
#rnx

diff --git a/Engine/Source/Runtime/D3D12RHI/Private/D3D12View.h b/Engine/Source/Runtime/D3D12RHI/Private/D3D12View.h
--- a/Engine/Source/Runtime/D3D12RHI/Private/D3D12View.h
+++ b/Engine/Source/Runtime/D3D12RHI/Private/D3D12View.h
@@ -878,6 +878,23 @@ public:
 	FORCEINLINE bool GetSkipFastClearFinalize()	const { return bSkipFastClearFinalize; }
 };
 
+class FD3D12ShaderResourceViewWithLocation : public FD3D12ShaderResourceView
+{
+public:
+	FD3D12ShaderResourceViewWithLocation(FD3D12Device* InParent)
+		: FD3D12ShaderResourceView(InParent)
+		, ViewLocation(InParent)
+	{}
+
+	FD3D12ShaderResourceViewWithLocation(FD3D12Device* InParent, D3D12_SHADER_RESOURCE_VIEW_DESC& InDesc, FD3D12ResourceLocation& InResourceLocation, uint32 InStride = -1, bool InSkipFastClearFinalize = false)
+		: FD3D12ShaderResourceView(InParent, InDesc, ViewLocation, InStride, InSkipFastClearFinalize)
+		, ViewLocation(InParent)
+	{
+	}
+
+	FD3D12ResourceLocation ViewLocation;
+};
+
 class FD3D12UnorderedAccessView : public FRHIUnorderedAccessView, public FD3D12View < D3D12_UNORDERED_ACCESS_VIEW_DESC >, public FD3D12LinkedAdapterObject<FD3D12UnorderedAccessView>
 {
 public:
