Fix one D3D resource transition error (more remain)

diff --git a/Engine/Source/Runtime/D3D12RHI/Private/D3D12Texture.cpp b/Engine/Source/Runtime/D3D12RHI/Private/D3D12Texture.cpp
--- a/Engine/Source/Runtime/D3D12RHI/Private/D3D12Texture.cpp
+++ b/Engine/Source/Runtime/D3D12RHI/Private/D3D12Texture.cpp
@@ -2126,10 +2126,9 @@ public:
 
 			CD3DX12_TEXTURE_COPY_LOCATION DestCopyLocation(TextureLink->GetResource()->GetResource(), MipIdx);
 
-			FScopeResourceBarrier ScopeResourceBarrierDest(
+			FConditionalScopeResourceBarrier ScopeResourceBarrierDest(
 				NativeCmdList,
 				TextureLink->GetResource(),
-				TextureLink->GetResource()->GetDefaultResourceState(),
 				D3D12_RESOURCE_STATE_COPY_DEST,
 				DestCopyLocation.SubresourceIndex);
 
