Fix for crash in UMaterialExpressionGetMaterialAttributes::PostLoad when no material is set.

Fix for D3D12 PSOs failing to be created due to an uninitialized member.

diff --git a/Engine/Source/Runtime/D3D12RHI/Private/Windows/WindowsD3D12PipelineState.cpp b/Engine/Source/Runtime/D3D12RHI/Private/Windows/WindowsD3D12PipelineState.cpp
--- a/Engine/Source/Runtime/D3D12RHI/Private/Windows/WindowsD3D12PipelineState.cpp
+++ b/Engine/Source/Runtime/D3D12RHI/Private/Windows/WindowsD3D12PipelineState.cpp
@@ -87,6 +87,7 @@ D3D12_GRAPHICS_PIPELINE_STATE_DESC FD3D12_GRAPHICS_PIPELINE_STATE_DESC::Graphics
 	D.RasterizerState = this->RasterizerState;
 	D.NumRenderTargets = this->RTFormatArray.NumRenderTargets;
 	FMemory::Memcpy(D.RTVFormats, this->RTFormatArray.RTFormats, sizeof(D.RTVFormats));
+	FMemory::Memzero(&D.StreamOutput, sizeof(D.StreamOutput));
 	D.SampleDesc = this->SampleDesc;
 	D.SampleMask = this->SampleMask;
 	D.CachedPSO = this->CachedPSO;
