Fix calculation of footprint offset for DX12 texture copy.
Previous logic was only correct if copying a mip or slice range starting at zero.
Hand picked integrate from dev-rendering

diff --git a/Engine/Source/Runtime/D3D12RHI/Private/D3D12Texture.cpp b/Engine/Source/Runtime/D3D12RHI/Private/D3D12Texture.cpp
--- a/Engine/Source/Runtime/D3D12RHI/Private/D3D12Texture.cpp
+++ b/Engine/Source/Runtime/D3D12RHI/Private/D3D12Texture.cpp
@@ -1946,8 +1946,21 @@ void TD3D12Texture2D<RHIResourceType>::UpdateTexture2D(class FRHICommandListImme
 	}
 }
 
+static void GetReadBackHeapDescImpl(D3D12_PLACED_SUBRESOURCE_FOOTPRINT& OutFootprint, ID3D12Device* InDevice, D3D12_RESOURCE_DESC const& InResourceDesc, uint32 InSubresource)
+{
+	uint64 Offset = 0;
+	if (InSubresource > 0)
+	{
+		InDevice->GetCopyableFootprints(&InResourceDesc, 0, InSubresource, 0, nullptr, nullptr, nullptr, &Offset);
+		Offset = Align(Offset, D3D12_TEXTURE_DATA_PLACEMENT_ALIGNMENT);
+	}
+	InDevice->GetCopyableFootprints(&InResourceDesc, InSubresource, 1, Offset, &OutFootprint, nullptr, nullptr, nullptr);
+
+	check(OutFootprint.Footprint.Width > 0 && OutFootprint.Footprint.Height > 0);
+}
+
 template<typename RHIResourceType>
-void TD3D12Texture2D<RHIResourceType>::GetReadBackHeapDesc(D3D12_PLACED_SUBRESOURCE_FOOTPRINT& OutFootprint, uint32 Subresource) const
+void TD3D12Texture2D<RHIResourceType>::GetReadBackHeapDesc(D3D12_PLACED_SUBRESOURCE_FOOTPRINT& OutFootprint, uint32 InSubresource) const
 {
 	check((RHIResourceType::GetFlags() & TexCreate_CPUReadback) != 0);
 
@@ -1962,17 +1975,7 @@ void TD3D12Texture2D<RHIResourceType>::GetReadBackHeapDesc(D3D12_PLACED_SUBRESOU
 	Desc.Format           = (DXGI_FORMAT) GPixelFormats[RHIResourceType::GetFormat()].PlatformFormat;
 	Desc.SampleDesc.Count = RHIResourceType::GetNumSamples();
 
-	ID3D12Device* Device = GetParentDevice()->GetDevice();
-
-	uint64 Offset = 0;
-	if (Subresource > 0)
-	{
-		Device->GetCopyableFootprints(&Desc, 0, Subresource, 0, nullptr, nullptr, nullptr, &Offset);
-		Offset = Align(Offset, D3D12_TEXTURE_DATA_PLACEMENT_ALIGNMENT);
-	}
-	Device->GetCopyableFootprints(&Desc, Subresource, 1, Offset, &OutFootprint, nullptr, nullptr, nullptr);
-
-	check(OutFootprint.Footprint.Width > 0 && OutFootprint.Footprint.Height > 0);
+	GetReadBackHeapDescImpl(OutFootprint, GetParentDevice()->GetDevice(), Desc, InSubresource);
 }
 
 void* FD3D12DynamicRHI::LockTexture2D_RenderThread(class FRHICommandListImmediate& RHICmdList, FRHITexture2D* TextureRHI, uint32 MipIndex, EResourceLockMode LockMode, uint32& DestStride, bool bLockWithinMiptail, bool bNeedsDefaultRHIFlush)
@@ -2977,17 +2980,24 @@ void FD3D12CommandContext::RHICopyTexture(FRHITexture* SourceTextureRHI, FRHITex
 			CopyInfo.SourcePosition.Z + CopyInfo.Size.Z
 		);
 
-		D3D12_TEXTURE_COPY_LOCATION Dst;
-		Dst.pResource = DestTexture->GetResource()->GetResource();
-		Dst.Type = bReadback? D3D12_TEXTURE_COPY_TYPE_PLACED_FOOTPRINT : D3D12_TEXTURE_COPY_TYPE_SUBRESOURCE_INDEX;
-
 		D3D12_TEXTURE_COPY_LOCATION Src;
 		Src.pResource = SourceTexture->GetResource()->GetResource();
 		Src.Type = D3D12_TEXTURE_COPY_TYPE_SUBRESOURCE_INDEX;
 
-		EPixelFormat Format = SourceTextureRHI->GetFormat();
+		D3D12_TEXTURE_COPY_LOCATION Dst;
+		Dst.pResource = DestTexture->GetResource()->GetResource();
+		Dst.Type = bReadback ? D3D12_TEXTURE_COPY_TYPE_PLACED_FOOTPRINT : D3D12_TEXTURE_COPY_TYPE_SUBRESOURCE_INDEX;
+
+		D3D12_RESOURCE_DESC DstDesc = {};
+		FIntVector TextureSize = DestTextureRHI->GetSizeXYZ();
+		DstDesc.Dimension = DestTextureRHI->GetTexture3D() ? D3D12_RESOURCE_DIMENSION_TEXTURE3D : D3D12_RESOURCE_DIMENSION_TEXTURE2D; 
+		DstDesc.Width = TextureSize.X;
+		DstDesc.Height = TextureSize.Y;
+		DstDesc.DepthOrArraySize = TextureSize.Z;
+		DstDesc.MipLevels = DestTextureRHI->GetNumMips();
+		DstDesc.Format = (DXGI_FORMAT)GPixelFormats[DestTextureRHI->GetFormat()].PlatformFormat;
+		DstDesc.SampleDesc.Count = DestTextureRHI->GetNumSamples();
 
-		uint64 Offset = 0;
 		for (uint32 SliceIndex = 0; SliceIndex < CopyInfo.NumSlices; ++SliceIndex)
 		{
 			uint32 SourceSliceIndex = CopyInfo.SourceSliceIndex + SliceIndex;
@@ -3007,26 +3017,11 @@ void FD3D12CommandContext::RHICopyTexture(FRHITexture* SourceTextureRHI, FRHITex
 				SourceBoxD3D.back   = CopyInfo.SourcePosition.Z + SizeZ;
 
 				Src.SubresourceIndex = CalcSubresource(SourceMipIndex, SourceSliceIndex, SourceTextureRHI->GetNumMips());
+				Dst.SubresourceIndex = CalcSubresource(DestMipIndex, DestSliceIndex, DestTextureRHI->GetNumMips());
 
 				if (bReadback)
 				{
-					uint32 NumBlocksX = FMath::DivideAndRoundUp<uint32>(SizeX, GPixelFormats[Format].BlockSizeX);
-					uint32 NumBlocksY = FMath::DivideAndRoundUp<uint32>(SizeY, GPixelFormats[Format].BlockSizeY);
-					uint32 NumBlocksZ = FMath::DivideAndRoundUp<uint32>(SizeZ, GPixelFormats[Format].BlockSizeZ);
-					uint32 RowPitch = Align(NumBlocksX * GPixelFormats[Format].BlockBytes, FD3D12_TEXTURE_DATA_PITCH_ALIGNMENT);
-
-					Dst.PlacedFootprint.Offset             = Offset;
-					Dst.PlacedFootprint.Footprint.Format   = (DXGI_FORMAT) GPixelFormats[Format].PlatformFormat;
-					Dst.PlacedFootprint.Footprint.Width    = SizeX;
-					Dst.PlacedFootprint.Footprint.Height   = SizeY;
-					Dst.PlacedFootprint.Footprint.Depth    = SizeZ;
-					Dst.PlacedFootprint.Footprint.RowPitch = RowPitch;
-
-					Offset = Align(Offset + RowPitch * NumBlocksY * NumBlocksZ, D3D12_TEXTURE_DATA_PLACEMENT_ALIGNMENT);
-				}
-				else
-				{
-					Dst.SubresourceIndex = CalcSubresource(DestMipIndex, DestSliceIndex, DestTextureRHI->GetNumMips());
+					GetReadBackHeapDescImpl(Dst.PlacedFootprint, GetParentDevice()->GetDevice(), DstDesc, Dst.SubresourceIndex);
 				}
 
 				CommandListHandle->CopyTextureRegion(
