UE4R - Show more info when D3D12 fails (including on viewport operations)

diff --git a/Engine/Source/Runtime/D3D12RHI/Private/D3D12RHIPrivate.h b/Engine/Source/Runtime/D3D12RHI/Private/D3D12RHIPrivate.h
--- a/Engine/Source/Runtime/D3D12RHI/Private/D3D12RHIPrivate.h
+++ b/Engine/Source/Runtime/D3D12RHI/Private/D3D12RHIPrivate.h
@@ -1166,4 +1166,25 @@ private:
 
 #endif
 
+
+#define EMBED_DXGI_ERROR_LIST(PerEntry, Terminator)	\
+	PerEntry(DXGI_ERROR_UNSUPPORTED) Terminator \
+	PerEntry(DXGI_ERROR_NOT_CURRENT) Terminator \
+	PerEntry(DXGI_ERROR_MORE_DATA) Terminator \
+	PerEntry(DXGI_ERROR_MODE_CHANGE_IN_PROGRESS) Terminator \
+	PerEntry(DXGI_ERROR_ALREADY_EXISTS) Terminator \
+	PerEntry(DXGI_ERROR_SESSION_DISCONNECTED) Terminator \
+	PerEntry(DXGI_ERROR_ACCESS_DENIED) Terminator \
+	PerEntry(DXGI_ERROR_NON_COMPOSITED_UI) Terminator \
+	PerEntry(DXGI_ERROR_CACHE_FULL) Terminator \
+	PerEntry(DXGI_ERROR_NOT_CURRENTLY_AVAILABLE) Terminator \
+	PerEntry(DXGI_ERROR_CACHE_CORRUPT) Terminator \
+	PerEntry(DXGI_ERROR_WAIT_TIMEOUT) Terminator \
+	PerEntry(DXGI_ERROR_FRAME_STATISTICS_DISJOINT) Terminator \
+	PerEntry(DXGI_ERROR_DYNAMIC_CODE_POLICY_VIOLATION) Terminator \
+	PerEntry(DXGI_ERROR_REMOTE_OUTOFMEMORY) Terminator \
+	PerEntry(DXGI_ERROR_ACCESS_LOST) Terminator
+
+
+
 #endif //!PLATFORM_XBOXONE
diff --git a/Engine/Source/Runtime/D3D12RHI/Private/D3D12Util.cpp b/Engine/Source/Runtime/D3D12RHI/Private/D3D12Util.cpp
--- a/Engine/Source/Runtime/D3D12RHI/Private/D3D12Util.cpp
+++ b/Engine/Source/Runtime/D3D12RHI/Private/D3D12Util.cpp
@@ -110,6 +110,7 @@ static FString GetD3D12ErrorString(HRESULT ErrorCode, ID3D12Device* Device)
 		D3DERR(DXGI_ERROR_INVALID_CALL)
 		D3DERR(E_NOINTERFACE)
 		D3DERR(DXGI_ERROR_DEVICE_REMOVED)
+		EMBED_DXGI_ERROR_LIST(D3DERR, )
 		default:
 			ErrorCodeText = FString::Printf(TEXT("%08X"), (int32)ErrorCode);
 	}
diff --git a/Engine/Source/Runtime/D3D12RHI/Private/Windows/WindowsD3D12Viewport.cpp b/Engine/Source/Runtime/D3D12RHI/Private/Windows/WindowsD3D12Viewport.cpp
--- a/Engine/Source/Runtime/D3D12RHI/Private/Windows/WindowsD3D12Viewport.cpp
+++ b/Engine/Source/Runtime/D3D12RHI/Private/Windows/WindowsD3D12Viewport.cpp
@@ -252,7 +252,17 @@ void FD3D12Viewport::ConditionalResetSwapChain(bool bIgnoreFocus)
 				}
 				else if (Result != DXGI_ERROR_NOT_CURRENTLY_AVAILABLE && Result != DXGI_STATUS_MODE_CHANGE_IN_PROGRESS)
 				{
-					UE_LOG(LogD3D12RHI, Error, TEXT("IDXGISwapChain::SetFullscreenState returned %08x, unknown error status."), Result);
+					const TCHAR* Name = nullptr;
+					switch (Result)
+					{
+#define CASE_ERROR_NAME(x)	case x: Name = TEXT(#x); break
+						EMBED_DXGI_ERROR_LIST(CASE_ERROR_NAME, ;)
+#undef CASE_ERROR_NAME
+					default:
+						Name = TEXT("unknown error status");
+						break;
+					}
+					UE_LOG(LogD3D12RHI, Error, TEXT("IDXGISwapChain::SetFullscreenState returned 0x%08x, %s."), Result, Name);
 				}
 			}
 		}
