UE4R - Only show this message in windows

diff --git a/Engine/Source/Runtime/D3D12RHI/Private/D3D12Util.cpp b/Engine/Source/Runtime/D3D12RHI/Private/D3D12Util.cpp
--- a/Engine/Source/Runtime/D3D12RHI/Private/D3D12Util.cpp
+++ b/Engine/Source/Runtime/D3D12RHI/Private/D3D12Util.cpp
@@ -110,7 +110,9 @@ static FString GetD3D12ErrorString(HRESULT ErrorCode, ID3D12Device* Device)
 		D3DERR(DXGI_ERROR_INVALID_CALL)
 		D3DERR(E_NOINTERFACE)
 		D3DERR(DXGI_ERROR_DEVICE_REMOVED)
+#if PLATFORM_WINDOWS
 		EMBED_DXGI_ERROR_LIST(D3DERR, )
+#endif
 		default:
 			ErrorCodeText = FString::Printf(TEXT("%08X"), (int32)ErrorCode);
 	}
