Added shader index to FD3D12RayTracingShaderTable::FShaderRecordCacheKey. This fixes a problem where same resources may be bound at the high level to records that have potentially different local root parameter layouts.
Increased maximum number of uniform buffers in FShaderRecordCacheKey from 4 to 6, as there is content that commonly uses 5 or 6 slots.

diff --git a/Engine/Source/Runtime/D3D12RHI/Private/D3D12RayTracing.cpp b/Engine/Source/Runtime/D3D12RHI/Private/D3D12RayTracing.cpp
--- a/Engine/Source/Runtime/D3D12RHI/Private/D3D12RayTracing.cpp
+++ b/Engine/Source/Runtime/D3D12RHI/Private/D3D12RayTracing.cpp
@@ -1322,14 +1322,17 @@ public:
 
 	struct FShaderRecordCacheKey
 	{
-		static constexpr uint32 MaxUniformBuffers = 4;
+		static constexpr uint32 MaxUniformBuffers = 6;
 		FRHIUniformBuffer* const* UniformBuffers[MaxUniformBuffers];
 		uint64 Hash;
 		uint32 NumUniformBuffers;
+		uint32 ShaderIndex;
 
 		FShaderRecordCacheKey() = default;
-		FShaderRecordCacheKey(uint32 InNumUniformBuffers, FRHIUniformBuffer* const* InUniformBuffers)
+		FShaderRecordCacheKey(uint32 InNumUniformBuffers, FRHIUniformBuffer* const* InUniformBuffers, uint32 InShaderIndex)
 		{
+			ShaderIndex = InShaderIndex;
+
 			check(InNumUniformBuffers <= MaxUniformBuffers);
 			NumUniformBuffers = FMath::Min(MaxUniformBuffers, InNumUniformBuffers);
 
@@ -1341,6 +1344,7 @@ public:
 		bool operator == (const FShaderRecordCacheKey& Other) const
 		{
 			if (Hash != Other.Hash) return false;
+			if (ShaderIndex != Other.ShaderIndex) return false;
 			if (NumUniformBuffers != Other.NumUniformBuffers) return false;
 
 			for (uint32 BufferIndex = 0; BufferIndex < NumUniformBuffers; ++BufferIndex)
@@ -3240,13 +3244,13 @@ void FD3D12CommandContext::RHISetRayTracingHitGroup(
 
 	const bool bCanUseRecordCache = GRayTracingCacheShaderRecords
 		&& Scene->Lifetime == RTSL_SingleFrame
-		&& LooseParameterDataSize == 0
-		&& NumUniformBuffers > 0
+		&& LooseParameterDataSize == 0 // loose parameters end up in unique constant buffers, so SBT records can't be shared
+		&& NumUniformBuffers > 0 // there is no benefit from cache if no resources are being bound
 		&& NumUniformBuffers <= CacheKey.MaxUniformBuffers;
 
 	if (bCanUseRecordCache)
 	{
-		CacheKey = FD3D12RayTracingShaderTable::FShaderRecordCacheKey(NumUniformBuffers, UniformBuffers);
+		CacheKey = FD3D12RayTracingShaderTable::FShaderRecordCacheKey(NumUniformBuffers, UniformBuffers, HitGroupIndex);
 
 		uint32* ExistingRecordIndex = ShaderTable->ShaderRecordCache.Find(CacheKey);
 		if (ExistingRecordIndex)
