Fixed resource state tracking and transitions for ray tracing resources

diff --git a/Engine/Source/Runtime/D3D12RHI/Private/D3D12RayTracing.cpp b/Engine/Source/Runtime/D3D12RHI/Private/D3D12RayTracing.cpp
--- a/Engine/Source/Runtime/D3D12RHI/Private/D3D12RayTracing.cpp
+++ b/Engine/Source/Runtime/D3D12RHI/Private/D3D12RayTracing.cpp
@@ -1388,6 +1388,35 @@ public:
 		}
 		Buffer->GetResource()->UpdateResidency(CommandContext.CommandListHandle);
 	}
+
+	// Some resources referenced in SBT may be dynamic (written on GPU timeline) and may require transition barriers.
+	// We save such resources while we fill the SBT and issue transitions before the SBT is used.
+
+	TSet<FD3D12ShaderResourceView*> TransitionSRVs;
+	TSet<FD3D12UnorderedAccessView*> TransitionUAVs;
+
+	void AddResourceTransition(FD3D12ShaderResourceView* SRV)
+	{
+		TransitionSRVs.Add(SRV);
+	}
+
+	void AddResourceTransition(FD3D12UnorderedAccessView* UAV)
+	{
+		TransitionUAVs.Add(UAV);
+	}
+
+	void TransitionResources(FD3D12CommandContext& CommandContext)
+	{
+		for (FD3D12ShaderResourceView* SRV : TransitionSRVs)
+		{
+			FD3D12DynamicRHI::TransitionResource(CommandContext.CommandListHandle, SRV, D3D12_RESOURCE_STATE_NON_PIXEL_SHADER_RESOURCE);
+		}
+
+		for (FD3D12UnorderedAccessView* UAV : TransitionUAVs)
+		{
+			FD3D12DynamicRHI::TransitionResource(CommandContext.CommandListHandle, UAV, D3D12_RESOURCE_STATE_UNORDERED_ACCESS);
+		}
+	}
 };
 
 
@@ -2680,6 +2709,16 @@ struct FD3D12RayTracingGlobalResourceBinder
 		D3D12Resource->UpdateResidency(CommandContext.CommandListHandle);
 	}
 
+	void AddResourceTransition(FD3D12ShaderResourceView* SRV)
+	{
+		FD3D12DynamicRHI::TransitionResource(CommandContext.CommandListHandle, SRV, D3D12_RESOURCE_STATE_NON_PIXEL_SHADER_RESOURCE);
+	}
+
+	void AddResourceTransition(FD3D12UnorderedAccessView* UAV)
+	{
+		FD3D12DynamicRHI::TransitionResource(CommandContext.CommandListHandle, UAV, D3D12_RESOURCE_STATE_UNORDERED_ACCESS);
+	}
+
 	FD3D12CommandContext& CommandContext;
 };
 
@@ -2740,6 +2779,22 @@ struct FD3D12RayTracingLocalResourceBinder
 		ShaderTable->AddResourceReference(D3D12Resource, RHIResource);
 	}
 
+	void AddResourceTransition(FD3D12ShaderResourceView* SRV)
+	{
+		if (SRV->GetResource()->RequiresResourceStateTracking())
+		{
+			ShaderTable->AddResourceTransition(SRV);
+		}
+	}
+
+	void AddResourceTransition(FD3D12UnorderedAccessView* UAV)
+	{
+		if (UAV->GetResource()->RequiresResourceStateTracking())
+		{
+			ShaderTable->AddResourceTransition(UAV);
+		}
+	}
+
 	FD3D12Device* Device = nullptr;
 	FD3D12RayTracingShaderTable* ShaderTable = nullptr;
 	const FD3D12RootSignature* RootSignature = nullptr;
@@ -2839,6 +2894,7 @@ static void SetRayTracingShaderResources(
 			BoundUAVMask |= 1ull << UAVIndex;
 
 			ReferencedResources.Add({ UAV->GetResource(), Resource });
+			Binder.AddResourceTransition(UAV);
 		}
 	}
 
@@ -2881,6 +2937,7 @@ static void SetRayTracingShaderResources(
 					BoundSRVMask |= 1ull << BindIndex;
 
 					ReferencedResources.Add({ SRV->GetResource(), SRV });
+					Binder.AddResourceTransition(SRV);
 
 					ResourceInfo = *ResourceInfos++;
 				} while (FRHIResourceTableEntry::GetUniformBufferIndex(ResourceInfo) == BufferIndex);
@@ -2908,6 +2965,7 @@ static void SetRayTracingShaderResources(
 					BoundSRVMask |= 1ull << BindIndex;
 
 					ReferencedResources.Add({ SRV->GetResource(), SRV });
+					Binder.AddResourceTransition(SRV);
 
 					ResourceInfo = *ResourceInfos++;
 				} while (FRHIResourceTableEntry::GetUniformBufferIndex(ResourceInfo) == BufferIndex);
@@ -3091,6 +3149,11 @@ static void DispatchRays(FD3D12CommandContext& CommandContext,
 		SetRayTracingShaderResources(CommandContext, RayGenShader, GlobalBindings, TransientDescriptorCache, ResourceBinder);
 	}
 
+	if (OptShaderTable)
+	{
+		OptShaderTable->TransitionResources(CommandContext);
+	}
+
 	CommandContext.CommandListHandle.FlushResourceBarriers();
 
 	ID3D12StateObject* RayTracingStateObject = Pipeline->StateObject.GetReference();
