Refactored global FRenderResource creation to occur after FDynamicRHI::Init. This removes the need for every backend to manually initialize the immediate context and removes copy-pasta.

This change is in response to a bug in the current implementation where only the D3D12 backend correctly initializes the immediate command list context prior to initializing global render resources (which use the immediate command list). All other backends only happen to work due to the bypass being enabled in the editor.

In shipping builds this causes a crash due to the command list context being null.

diff --git a/Engine/Source/Runtime/Apple/MetalRHI/Private/MetalRHI.cpp b/Engine/Source/Runtime/Apple/MetalRHI/Private/MetalRHI.cpp
--- a/Engine/Source/Runtime/Apple/MetalRHI/Private/MetalRHI.cpp
+++ b/Engine/Source/Runtime/Apple/MetalRHI/Private/MetalRHI.cpp
@@ -838,10 +838,7 @@ FMetalDynamicRHI::FMetalDynamicRHI(ERHIFeatureLevel::Type RequestedFeatureLevel)
 	if (ImmediateContext.Profiler)
 		ImmediateContext.Profiler->BeginFrame();
 #endif
-		
-	// Notify all initialized FRenderResources that there's a valid RHI device to create their RHI resources for now.
-	FRenderResource::InitPreRHIResources();	
-	
+
 	AsyncComputeContext = GSupportsEfficientAsyncCompute ? new FMetalRHIComputeContext(ImmediateContext.Profiler, new FMetalContext(ImmediateContext.Context->GetDevice(), ImmediateContext.Context->GetCommandQueue(), true)) : nullptr;
 
 #if ENABLE_METAL_GPUPROFILE
@@ -903,7 +900,6 @@ uint64 FMetalDynamicRHI::RHICalcTextureCubePlatformSize(uint32 Size, uint8 Forma
 
 void FMetalDynamicRHI::Init()
 {
-	GIsRHIInitialized = true;
 }
 
 void FMetalRHIImmediateCommandContext::RHIBeginFrame()
diff --git a/Engine/Source/Runtime/D3D12RHI/Private/Windows/WindowsD3D12Device.cpp b/Engine/Source/Runtime/D3D12RHI/Private/Windows/WindowsD3D12Device.cpp
--- a/Engine/Source/Runtime/D3D12RHI/Private/Windows/WindowsD3D12Device.cpp
+++ b/Engine/Source/Runtime/D3D12RHI/Private/Windows/WindowsD3D12Device.cpp
@@ -709,22 +709,6 @@ void FD3D12DynamicRHI::Init()
 
 	GSupportsDepthBoundsTest = SupportsDepthBoundsTest(this);
 
-	GRHICommandList.GetImmediateCommandList().SetContext(GetDynamicRHI<FD3D12DynamicRHI>()->RHIGetDefaultContext());
-	GRHICommandList.GetImmediateCommandList().SetGPUMask(FRHIGPUMask::All());
-	GRHICommandList.GetImmediateAsyncComputeCommandList().SetComputeContext(GDynamicRHI->RHIGetDefaultAsyncComputeContext());
-	GRHICommandList.GetImmediateAsyncComputeCommandList().SetGPUMask(FRHIGPUMask::All());
-
-	// Notify all initialized FRenderResources that there's a valid RHI device to create their RHI resources for now.
-	for (TLinkedList<FRenderResource*>::TIterator ResourceIt(FRenderResource::GetResourceList()); ResourceIt; ResourceIt.Next())
-	{
-		ResourceIt->InitRHI();
-	}
-	// Dynamic resources can have dependencies on static resources (with uniform buffers) and must initialized last!
-	for (TLinkedList<FRenderResource*>::TIterator ResourceIt(FRenderResource::GetResourceList()); ResourceIt; ResourceIt.Next())
-	{
-		ResourceIt->InitDynamicRHI();
-	}
-
 	{
 		GRHISupportsHDROutput = SupportsHDROutput(this);
 
@@ -760,9 +744,6 @@ void FD3D12DynamicRHI::Init()
 #if D3D12_RHI_RAYTRACING
 	GRHISupportsRayTracing = GetAdapter().GetD3DRayTracingDevice() != nullptr;
 #endif
-
-	// Set the RHI initialized flag.
-	GIsRHIInitialized = true;
 }
 
 void FD3D12DynamicRHI::PostInit()
diff --git a/Engine/Source/Runtime/EmptyRHI/Private/EmptyRHI.cpp b/Engine/Source/Runtime/EmptyRHI/Private/EmptyRHI.cpp
--- a/Engine/Source/Runtime/EmptyRHI/Private/EmptyRHI.cpp
+++ b/Engine/Source/Runtime/EmptyRHI/Private/EmptyRHI.cpp
@@ -94,19 +94,6 @@ FEmptyDynamicRHI::FEmptyDynamicRHI()
 //	GPixelFormats[PF_R16G16B16A16_SNORM ].PlatformFormat    =
 
 	GDynamicRHI = this;
-
-	// Notify all initialized FRenderResources that there's a valid RHI device to create their RHI resources for now.
-	for(TLinkedList<FRenderResource*>::TIterator ResourceIt(FRenderResource::GetResourceList());ResourceIt;ResourceIt.Next())
-	{
-		ResourceIt->InitRHI();
-	}
-	// Dynamic resources can have dependencies on static resources (with uniform buffers) and must initialized last!
-	for(TLinkedList<FRenderResource*>::TIterator ResourceIt(FRenderResource::GetResourceList());ResourceIt;ResourceIt.Next())
-	{
-		ResourceIt->InitDynamicRHI();
-	}
-
-	GIsRHIInitialized = true;
 }
 
 void FEmptyDynamicRHI::Init()
diff --git a/Engine/Source/Runtime/NullDrv/Private/NullRHI.cpp b/Engine/Source/Runtime/NullDrv/Private/NullRHI.cpp
--- a/Engine/Source/Runtime/NullDrv/Private/NullRHI.cpp
+++ b/Engine/Source/Runtime/NullDrv/Private/NullRHI.cpp
@@ -34,25 +34,6 @@ void FNullDynamicRHI::Init()
 	GRHIVendorId = 1;
 
 	check(!GIsRHIInitialized);
-
-	// do not do this at least on dedicated server; clients with -NullRHI may need additional consideration
-#if !WITH_EDITOR	
-	if (!IsRunningDedicatedServer())
-#endif
-	{
-		// Notify all initialized FRenderResources that there's a valid RHI device to create their RHI resources for now.
-		for(TLinkedList<FRenderResource*>::TIterator ResourceIt(FRenderResource::GetResourceList());ResourceIt;ResourceIt.Next())
-		{
-			ResourceIt->InitRHI();
-		}
-		// Dynamic resources can have dependencies on static resources (with uniform buffers) and must initialized last!
-		for(TLinkedList<FRenderResource*>::TIterator ResourceIt(FRenderResource::GetResourceList());ResourceIt;ResourceIt.Next())
-		{
-			ResourceIt->InitDynamicRHI();
-		}
-	}
-
-	GIsRHIInitialized = true;
 }
 
 
diff --git a/Engine/Source/Runtime/OpenGLDrv/Private/OpenGLDevice.cpp b/Engine/Source/Runtime/OpenGLDrv/Private/OpenGLDevice.cpp
--- a/Engine/Source/Runtime/OpenGLDrv/Private/OpenGLDevice.cpp
+++ b/Engine/Source/Runtime/OpenGLDrv/Private/OpenGLDevice.cpp
@@ -1671,17 +1671,6 @@ void FOpenGLDynamicRHI::Init()
 		DynamicIndexBuffers.Init(CalcDynamicBufferSize(1));
 	}
 
-	// Notify all initialized FRenderResources that there's a valid RHI device to create their RHI resources for now.
-	for(TLinkedList<FRenderResource*>::TIterator ResourceIt(FRenderResource::GetResourceList());ResourceIt;ResourceIt.Next())
-	{
-		ResourceIt->InitRHI();
-	}
-	// Dynamic resources can have dependencies on static resources (with uniform buffers) and must initialized last!
-	for(TLinkedList<FRenderResource*>::TIterator ResourceIt(FRenderResource::GetResourceList());ResourceIt;ResourceIt.Next())
-	{
-		ResourceIt->InitDynamicRHI();
-	}
-
 #if PLATFORM_WINDOWS || PLATFORM_LINUX
 
 	extern int64 GOpenGLDedicatedVideoMemory;
@@ -1723,9 +1712,6 @@ void FOpenGLDynamicRHI::Init()
 
 	FHardwareInfo::RegisterHardwareInfo( NAME_RHI, TEXT( "OpenGL" ) );
 
-	// Set the RHI initialized flag.
-	GIsRHIInitialized = true;
-
 	CheckTextureCubeLodSupport();
 	CheckVaryingLimit();
 	CheckRoundFunction();
diff --git a/Engine/Source/Runtime/RHI/Private/DynamicRHI.cpp b/Engine/Source/Runtime/RHI/Private/DynamicRHI.cpp
--- a/Engine/Source/Runtime/RHI/Private/DynamicRHI.cpp
+++ b/Engine/Source/Runtime/RHI/Private/DynamicRHI.cpp
@@ -14,6 +14,7 @@
 #include "GenericPlatform/GenericPlatformDriver.h"
 #include "GenericPlatform/GenericPlatformCrashContext.h"
 #include "PipelineStateCache.h"
+#include "RenderResource.h"
 
 #ifndef PLATFORM_ALLOW_NULL_RHI
 	#define PLATFORM_ALLOW_NULL_RHI		0
@@ -52,6 +53,15 @@ void InitNullRHI()
 	GDynamicRHI->Init();
 	GRHICommandList.GetImmediateCommandList().SetContext(GDynamicRHI->RHIGetDefaultContext());
 	GRHICommandList.GetImmediateAsyncComputeCommandList().SetComputeContext(GDynamicRHI->RHIGetDefaultAsyncComputeContext());
+
+	// do not do this at least on dedicated server; clients with -NullRHI may need additional consideration
+#if !WITH_EDITOR
+	if (!IsRunningDedicatedServer())
+#endif
+	{
+		FRenderResource::InitPreRHIResources();
+	}
+
 	GUsingNullRHI = true;
 	GRHISupportsTextureStreaming = false;
 
@@ -239,6 +249,8 @@ void RHIInit(bool bHasEditorToken)
 				GRHICommandList.GetImmediateCommandList().SetContext(GDynamicRHI->RHIGetDefaultContext());
 				GRHICommandList.GetImmediateAsyncComputeCommandList().SetComputeContext(GDynamicRHI->RHIGetDefaultAsyncComputeContext());
 
+				FRenderResource::InitPreRHIResources();
+
 				FString FeatureLevelString;
 				GetFeatureLevelName(GMaxRHIFeatureLevel, FeatureLevelString);
 
@@ -268,6 +280,7 @@ void RHIInit(bool bHasEditorToken)
 		}
 
 		check(GDynamicRHI);
+		GIsRHIInitialized = true;
 	}
 
 #if PLATFORM_WINDOWS || PLATFORM_MAC
diff --git a/Engine/Source/Runtime/RHI/RHI.Build.cs b/Engine/Source/Runtime/RHI/RHI.Build.cs
--- a/Engine/Source/Runtime/RHI/RHI.Build.cs
+++ b/Engine/Source/Runtime/RHI/RHI.Build.cs
@@ -9,8 +9,9 @@ public class RHI : ModuleRules
 	{
 		PrivateDependencyModuleNames.Add("Core");
 		PrivateDependencyModuleNames.Add("ApplicationCore");
+        PrivateDependencyModuleNames.Add("RenderCore");
 
-		if (Target.bCompileAgainstEngine)
+        if (Target.bCompileAgainstEngine)
 		{
 			DynamicallyLoadedModuleNames.Add("NullDrv");
 
diff --git a/Engine/Source/Runtime/RenderCore/Private/RenderResource.cpp b/Engine/Source/Runtime/RenderCore/Private/RenderResource.cpp
--- a/Engine/Source/Runtime/RenderCore/Private/RenderResource.cpp
+++ b/Engine/Source/Runtime/RenderCore/Private/RenderResource.cpp
@@ -34,12 +34,15 @@ TLinkedList<FRenderResource*>*& FRenderResource::GetResourceList()
 
 /** Initialize all resources initialized before the RHI was initialized */
 void FRenderResource::InitPreRHIResources()
-{	
+{
+	check(!GIsRHIInitialized);
+
 	// Notify all initialized FRenderResources that there's a valid RHI device to create their RHI resources for now.
 	for (TLinkedList<FRenderResource*>::TIterator ResourceIt(FRenderResource::GetResourceList()); ResourceIt; ResourceIt.Next())
 	{
 		ResourceIt->InitRHI();
 	}
+
 	// Dynamic resources can have dependencies on static resources (with uniform buffers) and must initialized last!
 	for (TLinkedList<FRenderResource*>::TIterator ResourceIt(FRenderResource::GetResourceList()); ResourceIt; ResourceIt.Next())
 	{
diff --git a/Engine/Source/Runtime/VulkanRHI/Private/VulkanRHI.cpp b/Engine/Source/Runtime/VulkanRHI/Private/VulkanRHI.cpp
--- a/Engine/Source/Runtime/VulkanRHI/Private/VulkanRHI.cpp
+++ b/Engine/Source/Runtime/VulkanRHI/Private/VulkanRHI.cpp
@@ -713,23 +713,10 @@ void FVulkanDynamicRHI::InitInstance()
 
 		GUseTexture3DBulkDataRHI = true;
 
-		// Notify all initialized FRenderResources that there's a valid RHI device to create their RHI resources for now.
-		for (TLinkedList<FRenderResource*>::TIterator ResourceIt(FRenderResource::GetResourceList()); ResourceIt; ResourceIt.Next())
-		{
-			ResourceIt->InitRHI();
-		}
-		// Dynamic resources can have dependencies on static resources (with uniform buffers) and must initialized last!
-		for (TLinkedList<FRenderResource*>::TIterator ResourceIt(FRenderResource::GetResourceList()); ResourceIt; ResourceIt.Next())
-		{
-			ResourceIt->InitDynamicRHI();
-		}
-
 		FHardwareInfo::RegisterHardwareInfo(NAME_RHI, TEXT("Vulkan"));
 
 		GProjectionSignY = 1.0f;
 
-		GIsRHIInitialized = true;
-
 		SavePipelineCacheCmd = IConsoleManager::Get().RegisterConsoleCommand(
 			TEXT("r.Vulkan.SavePipelineCache"),
 			TEXT("Save pipeline cache."),
diff --git a/Engine/Source/Runtime/Windows/D3D11RHI/Private/Windows/WindowsD3D11Device.cpp b/Engine/Source/Runtime/Windows/D3D11RHI/Private/Windows/WindowsD3D11Device.cpp
--- a/Engine/Source/Runtime/Windows/D3D11RHI/Private/Windows/WindowsD3D11Device.cpp
+++ b/Engine/Source/Runtime/Windows/D3D11RHI/Private/Windows/WindowsD3D11Device.cpp
@@ -1753,17 +1753,6 @@ void FD3D11DynamicRHI::InitD3DDevice()
 
 		SetupAfterDeviceCreation();
 
-		// Notify all initialized FRenderResources that there's a valid RHI device to create their RHI resources for now.
-		for(TLinkedList<FRenderResource*>::TIterator ResourceIt(FRenderResource::GetResourceList());ResourceIt;ResourceIt.Next())
-		{
-			ResourceIt->InitRHI();
-		}
-		// Dynamic resources can have dependencies on static resources (with uniform buffers) and must initialized last!
-		for(TLinkedList<FRenderResource*>::TIterator ResourceIt(FRenderResource::GetResourceList());ResourceIt;ResourceIt.Next())
-		{
-			ResourceIt->InitDynamicRHI();
-		}
-
 #if !(UE_BUILD_SHIPPING && WITH_EDITOR)
 		// Add some filter outs for known debug spew messages (that we don't care about)
 		if(DeviceFlags & D3D11_CREATE_DEVICE_DEBUG)
@@ -1843,8 +1832,6 @@ void FD3D11DynamicRHI::InitD3DDevice()
 		GRHISupportsTextureStreaming = true;
 		GRHISupportsFirstInstance = true;
 		GRHINeedsExtraDeletionLatency = false;
-		// Set the RHI initialized flag.
-		GIsRHIInitialized = true;
 	}
 }
 
