Second attempted fix for null context assert in RHI init. This time without a cyclical module dependency.

diff --git a/Engine/Source/Runtime/Apple/MetalRHI/Private/MetalRHI.cpp b/Engine/Source/Runtime/Apple/MetalRHI/Private/MetalRHI.cpp
--- a/Engine/Source/Runtime/Apple/MetalRHI/Private/MetalRHI.cpp
+++ b/Engine/Source/Runtime/Apple/MetalRHI/Private/MetalRHI.cpp
@@ -838,10 +838,6 @@ FMetalDynamicRHI::FMetalDynamicRHI(ERHIFeatureLevel::Type RequestedFeatureLevel)
 	if (ImmediateContext.Profiler)
 		ImmediateContext.Profiler->BeginFrame();
 #endif
-		
-	// Notify all initialized FRenderResources that there's a valid RHI device to create their RHI resources for now.
-	FRenderResource::InitPreRHIResources();	
-	
 	AsyncComputeContext = GSupportsEfficientAsyncCompute ? new FMetalRHIComputeContext(ImmediateContext.Profiler, new FMetalContext(ImmediateContext.Context->GetDevice(), ImmediateContext.Context->GetCommandQueue(), true)) : nullptr;
 
 #if ENABLE_METAL_GPUPROFILE
@@ -903,6 +899,9 @@ uint64 FMetalDynamicRHI::RHICalcTextureCubePlatformSize(uint32 Size, uint8 Forma
 
 void FMetalDynamicRHI::Init()
 {
+	GRHICommandList.GetImmediateCommandList().SetContext(RHIGetDefaultContext());
+	GRHICommandList.GetImmediateAsyncComputeCommandList().SetComputeContext(RHIGetDefaultAsyncComputeContext());
+	FRenderResource::InitPreRHIResources();
 	GIsRHIInitialized = true;
 }
 
diff --git a/Engine/Source/Runtime/D3D12RHI/Private/Windows/WindowsD3D12Device.cpp b/Engine/Source/Runtime/D3D12RHI/Private/Windows/WindowsD3D12Device.cpp
--- a/Engine/Source/Runtime/D3D12RHI/Private/Windows/WindowsD3D12Device.cpp
+++ b/Engine/Source/Runtime/D3D12RHI/Private/Windows/WindowsD3D12Device.cpp
@@ -709,22 +709,6 @@ void FD3D12DynamicRHI::Init()
 
 	GSupportsDepthBoundsTest = SupportsDepthBoundsTest(this);
 
-	GRHICommandList.GetImmediateCommandList().SetContext(GetDynamicRHI<FD3D12DynamicRHI>()->RHIGetDefaultContext());
-	GRHICommandList.GetImmediateCommandList().SetGPUMask(FRHIGPUMask::All());
-	GRHICommandList.GetImmediateAsyncComputeCommandList().SetComputeContext(GDynamicRHI->RHIGetDefaultAsyncComputeContext());
-	GRHICommandList.GetImmediateAsyncComputeCommandList().SetGPUMask(FRHIGPUMask::All());
-
-	// Notify all initialized FRenderResources that there's a valid RHI device to create their RHI resources for now.
-	for (TLinkedList<FRenderResource*>::TIterator ResourceIt(FRenderResource::GetResourceList()); ResourceIt; ResourceIt.Next())
-	{
-		ResourceIt->InitRHI();
-	}
-	// Dynamic resources can have dependencies on static resources (with uniform buffers) and must initialized last!
-	for (TLinkedList<FRenderResource*>::TIterator ResourceIt(FRenderResource::GetResourceList()); ResourceIt; ResourceIt.Next())
-	{
-		ResourceIt->InitDynamicRHI();
-	}
-
 	{
 		GRHISupportsHDROutput = SupportsHDROutput(this);
 
@@ -761,7 +745,9 @@ void FD3D12DynamicRHI::Init()
 	GRHISupportsRayTracing = GetAdapter().GetD3DRayTracingDevice() != nullptr;
 #endif
 
-	// Set the RHI initialized flag.
+	GRHICommandList.GetImmediateCommandList().SetContext(RHIGetDefaultContext());
+	GRHICommandList.GetImmediateAsyncComputeCommandList().SetComputeContext(RHIGetDefaultAsyncComputeContext());
+	FRenderResource::InitPreRHIResources();
 	GIsRHIInitialized = true;
 }
 
diff --git a/Engine/Source/Runtime/EmptyRHI/Private/EmptyRHI.cpp b/Engine/Source/Runtime/EmptyRHI/Private/EmptyRHI.cpp
--- a/Engine/Source/Runtime/EmptyRHI/Private/EmptyRHI.cpp
+++ b/Engine/Source/Runtime/EmptyRHI/Private/EmptyRHI.cpp
@@ -95,17 +95,9 @@ FEmptyDynamicRHI::FEmptyDynamicRHI()
 
 	GDynamicRHI = this;
 
-	// Notify all initialized FRenderResources that there's a valid RHI device to create their RHI resources for now.
-	for(TLinkedList<FRenderResource*>::TIterator ResourceIt(FRenderResource::GetResourceList());ResourceIt;ResourceIt.Next())
-	{
-		ResourceIt->InitRHI();
-	}
-	// Dynamic resources can have dependencies on static resources (with uniform buffers) and must initialized last!
-	for(TLinkedList<FRenderResource*>::TIterator ResourceIt(FRenderResource::GetResourceList());ResourceIt;ResourceIt.Next())
-	{
-		ResourceIt->InitDynamicRHI();
-	}
-
+	GRHICommandList.GetImmediateCommandList().SetContext(RHIGetDefaultContext());
+	GRHICommandList.GetImmediateAsyncComputeCommandList().SetComputeContext(RHIGetDefaultAsyncComputeContext());
+	FRenderResource::InitPreRHIResources();
 	GIsRHIInitialized = true;
 }
 
diff --git a/Engine/Source/Runtime/NullDrv/Private/NullRHI.cpp b/Engine/Source/Runtime/NullDrv/Private/NullRHI.cpp
--- a/Engine/Source/Runtime/NullDrv/Private/NullRHI.cpp
+++ b/Engine/Source/Runtime/NullDrv/Private/NullRHI.cpp
@@ -40,16 +40,9 @@ void FNullDynamicRHI::Init()
 	if (!IsRunningDedicatedServer())
 #endif
 	{
-		// Notify all initialized FRenderResources that there's a valid RHI device to create their RHI resources for now.
-		for(TLinkedList<FRenderResource*>::TIterator ResourceIt(FRenderResource::GetResourceList());ResourceIt;ResourceIt.Next())
-		{
-			ResourceIt->InitRHI();
-		}
-		// Dynamic resources can have dependencies on static resources (with uniform buffers) and must initialized last!
-		for(TLinkedList<FRenderResource*>::TIterator ResourceIt(FRenderResource::GetResourceList());ResourceIt;ResourceIt.Next())
-		{
-			ResourceIt->InitDynamicRHI();
-		}
+		GRHICommandList.GetImmediateCommandList().SetContext(RHIGetDefaultContext());
+		GRHICommandList.GetImmediateAsyncComputeCommandList().SetComputeContext(RHIGetDefaultAsyncComputeContext());
+		FRenderResource::InitPreRHIResources();
 	}
 
 	GIsRHIInitialized = true;
diff --git a/Engine/Source/Runtime/OpenGLDrv/Private/OpenGLDevice.cpp b/Engine/Source/Runtime/OpenGLDrv/Private/OpenGLDevice.cpp
--- a/Engine/Source/Runtime/OpenGLDrv/Private/OpenGLDevice.cpp
+++ b/Engine/Source/Runtime/OpenGLDrv/Private/OpenGLDevice.cpp
@@ -1671,17 +1671,6 @@ void FOpenGLDynamicRHI::Init()
 		DynamicIndexBuffers.Init(CalcDynamicBufferSize(1));
 	}
 
-	// Notify all initialized FRenderResources that there's a valid RHI device to create their RHI resources for now.
-	for(TLinkedList<FRenderResource*>::TIterator ResourceIt(FRenderResource::GetResourceList());ResourceIt;ResourceIt.Next())
-	{
-		ResourceIt->InitRHI();
-	}
-	// Dynamic resources can have dependencies on static resources (with uniform buffers) and must initialized last!
-	for(TLinkedList<FRenderResource*>::TIterator ResourceIt(FRenderResource::GetResourceList());ResourceIt;ResourceIt.Next())
-	{
-		ResourceIt->InitDynamicRHI();
-	}
-
 #if PLATFORM_WINDOWS || PLATFORM_LINUX
 
 	extern int64 GOpenGLDedicatedVideoMemory;
@@ -1723,12 +1712,14 @@ void FOpenGLDynamicRHI::Init()
 
 	FHardwareInfo::RegisterHardwareInfo( NAME_RHI, TEXT( "OpenGL" ) );
 
-	// Set the RHI initialized flag.
-	GIsRHIInitialized = true;
-
 	CheckTextureCubeLodSupport();
 	CheckVaryingLimit();
 	CheckRoundFunction();
+
+	GRHICommandList.GetImmediateCommandList().SetContext(RHIGetDefaultContext());
+	GRHICommandList.GetImmediateAsyncComputeCommandList().SetComputeContext(RHIGetDefaultAsyncComputeContext());
+	FRenderResource::InitPreRHIResources();
+	GIsRHIInitialized = true;
 }
 
 void FOpenGLDynamicRHI::PostInit()
diff --git a/Engine/Source/Runtime/RHI/Private/DynamicRHI.cpp b/Engine/Source/Runtime/RHI/Private/DynamicRHI.cpp
--- a/Engine/Source/Runtime/RHI/Private/DynamicRHI.cpp
+++ b/Engine/Source/Runtime/RHI/Private/DynamicRHI.cpp
@@ -236,8 +236,10 @@ void RHIInit(bool bHasEditorToken)
 				AFRUtils::StaticInitialize();
 #endif
 
-				GRHICommandList.GetImmediateCommandList().SetContext(GDynamicRHI->RHIGetDefaultContext());
-				GRHICommandList.GetImmediateAsyncComputeCommandList().SetComputeContext(GDynamicRHI->RHIGetDefaultAsyncComputeContext());
+				// Validation of contexts.
+				GRHICommandList.GetImmediateCommandList().GetContext();
+				GRHICommandList.GetImmediateAsyncComputeCommandList().GetComputeContext();
+				check(GIsRHIInitialized);
 
 				FString FeatureLevelString;
 				GetFeatureLevelName(GMaxRHIFeatureLevel, FeatureLevelString);
diff --git a/Engine/Source/Runtime/VulkanRHI/Private/VulkanRHI.cpp b/Engine/Source/Runtime/VulkanRHI/Private/VulkanRHI.cpp
--- a/Engine/Source/Runtime/VulkanRHI/Private/VulkanRHI.cpp
+++ b/Engine/Source/Runtime/VulkanRHI/Private/VulkanRHI.cpp
@@ -713,23 +713,10 @@ void FVulkanDynamicRHI::InitInstance()
 
 		GUseTexture3DBulkDataRHI = true;
 
-		// Notify all initialized FRenderResources that there's a valid RHI device to create their RHI resources for now.
-		for (TLinkedList<FRenderResource*>::TIterator ResourceIt(FRenderResource::GetResourceList()); ResourceIt; ResourceIt.Next())
-		{
-			ResourceIt->InitRHI();
-		}
-		// Dynamic resources can have dependencies on static resources (with uniform buffers) and must initialized last!
-		for (TLinkedList<FRenderResource*>::TIterator ResourceIt(FRenderResource::GetResourceList()); ResourceIt; ResourceIt.Next())
-		{
-			ResourceIt->InitDynamicRHI();
-		}
-
 		FHardwareInfo::RegisterHardwareInfo(NAME_RHI, TEXT("Vulkan"));
 
 		GProjectionSignY = 1.0f;
 
-		GIsRHIInitialized = true;
-
 		SavePipelineCacheCmd = IConsoleManager::Get().RegisterConsoleCommand(
 			TEXT("r.Vulkan.SavePipelineCache"),
 			TEXT("Save pipeline cache."),
@@ -779,6 +766,11 @@ void FVulkanDynamicRHI::InitInstance()
 		);
 
 #endif
+
+		GRHICommandList.GetImmediateCommandList().SetContext(RHIGetDefaultContext());
+		GRHICommandList.GetImmediateAsyncComputeCommandList().SetComputeContext(RHIGetDefaultAsyncComputeContext());
+		FRenderResource::InitPreRHIResources();
+		GIsRHIInitialized = true;
 	}
 }
 
diff --git a/Engine/Source/Runtime/Windows/D3D11RHI/Private/Windows/WindowsD3D11Device.cpp b/Engine/Source/Runtime/Windows/D3D11RHI/Private/Windows/WindowsD3D11Device.cpp
--- a/Engine/Source/Runtime/Windows/D3D11RHI/Private/Windows/WindowsD3D11Device.cpp
+++ b/Engine/Source/Runtime/Windows/D3D11RHI/Private/Windows/WindowsD3D11Device.cpp
@@ -1753,17 +1753,6 @@ void FD3D11DynamicRHI::InitD3DDevice()
 
 		SetupAfterDeviceCreation();
 
-		// Notify all initialized FRenderResources that there's a valid RHI device to create their RHI resources for now.
-		for(TLinkedList<FRenderResource*>::TIterator ResourceIt(FRenderResource::GetResourceList());ResourceIt;ResourceIt.Next())
-		{
-			ResourceIt->InitRHI();
-		}
-		// Dynamic resources can have dependencies on static resources (with uniform buffers) and must initialized last!
-		for(TLinkedList<FRenderResource*>::TIterator ResourceIt(FRenderResource::GetResourceList());ResourceIt;ResourceIt.Next())
-		{
-			ResourceIt->InitDynamicRHI();
-		}
-
 #if !(UE_BUILD_SHIPPING && WITH_EDITOR)
 		// Add some filter outs for known debug spew messages (that we don't care about)
 		if(DeviceFlags & D3D11_CREATE_DEVICE_DEBUG)
@@ -1843,7 +1832,10 @@ void FD3D11DynamicRHI::InitD3DDevice()
 		GRHISupportsTextureStreaming = true;
 		GRHISupportsFirstInstance = true;
 		GRHINeedsExtraDeletionLatency = false;
-		// Set the RHI initialized flag.
+
+		GRHICommandList.GetImmediateCommandList().SetContext(RHIGetDefaultContext());
+		GRHICommandList.GetImmediateAsyncComputeCommandList().SetComputeContext(RHIGetDefaultAsyncComputeContext());
+		FRenderResource::InitPreRHIResources();
 		GIsRHIInitialized = true;
 	}
 }
