Added support for disabled sections for dynamic meshes in ray tracing

diff --git a/Engine/Source/Runtime/D3D12RHI/Private/D3D12RayTracing.cpp b/Engine/Source/Runtime/D3D12RHI/Private/D3D12RayTracing.cpp
--- a/Engine/Source/Runtime/D3D12RHI/Private/D3D12RayTracing.cpp
+++ b/Engine/Source/Runtime/D3D12RHI/Private/D3D12RayTracing.cpp
@@ -2140,66 +2140,70 @@ void FD3D12RayTracingGeometry::BuildAccelerationStructure(FD3D12CommandContext&
 
 		FD3D12VertexBuffer* VertexBuffer = CommandContext.RetrieveObject<FD3D12VertexBuffer>(Segment.VertexBuffer.GetReference());
 
-		if (GeometryType == D3D12_RAYTRACING_GEOMETRY_TYPE_TRIANGLES)
+		if (Segment.bEnabled)
 		{
-			switch (Segment.VertexBufferElementType)
+			if (GeometryType == D3D12_RAYTRACING_GEOMETRY_TYPE_TRIANGLES)
 			{
-			case VET_Float4: 
-				// While the DXGI_FORMAT_R32G32B32A32_FLOAT format is not supported by DXR, since we manually load vertex 
-				// data when we are building the BLAS, we can just rely on the vertex stride to offset the read index, 
-				// and read only the 3 vertex components, and so use the DXGI_FORMAT_R32G32B32_FLOAT vertex format
-			case VET_Float3:
-				Desc.Triangles.VertexFormat = DXGI_FORMAT_R32G32B32_FLOAT;
-				break;
-			case VET_Float2:
-				Desc.Triangles.VertexFormat = DXGI_FORMAT_R32G32_FLOAT;
-				break;
-			case VET_Half2:
-				Desc.Triangles.VertexFormat = DXGI_FORMAT_R16G16_FLOAT;
-				break;
-			default:
-				checkNoEntry();
-				break;
-			}
+				switch (Segment.VertexBufferElementType)
+				{
+				case VET_Float4:
+					// While the DXGI_FORMAT_R32G32B32A32_FLOAT format is not supported by DXR, since we manually load vertex 
+					// data when we are building the BLAS, we can just rely on the vertex stride to offset the read index, 
+					// and read only the 3 vertex components, and so use the DXGI_FORMAT_R32G32B32_FLOAT vertex format
+				case VET_Float3:
+					Desc.Triangles.VertexFormat = DXGI_FORMAT_R32G32B32_FLOAT;
+					break;
+				case VET_Float2:
+					Desc.Triangles.VertexFormat = DXGI_FORMAT_R32G32_FLOAT;
+					break;
+				case VET_Half2:
+					Desc.Triangles.VertexFormat = DXGI_FORMAT_R16G16_FLOAT;
+					break;
+				default:
+					checkNoEntry();
+					break;
+				}
 
-			Desc.Triangles.Transform3x4 = D3D12_GPU_VIRTUAL_ADDRESS(0);
+				Desc.Triangles.Transform3x4 = D3D12_GPU_VIRTUAL_ADDRESS(0);
 
-			if (IndexBuffer)
-			{
-				Desc.Triangles.IndexFormat = IndexStride == 4 ? DXGI_FORMAT_R32_UINT : DXGI_FORMAT_R16_UINT;
-				Desc.Triangles.IndexCount = Segment.NumPrimitives * IndicesPerPrimitive;
-				Desc.Triangles.IndexBuffer = IndexBuffer->ResourceLocation.GetGPUVirtualAddress() +
-											 IndexOffsetInBytes +
-											 IndexStride * Segment.FirstPrimitive * IndicesPerPrimitive;
+				if (IndexBuffer)
+				{
+					Desc.Triangles.IndexFormat = IndexStride == 4 ? DXGI_FORMAT_R32_UINT : DXGI_FORMAT_R16_UINT;
+					Desc.Triangles.IndexCount = Segment.NumPrimitives * IndicesPerPrimitive;
+					Desc.Triangles.IndexBuffer = IndexBuffer->ResourceLocation.GetGPUVirtualAddress() +
+						IndexOffsetInBytes +
+						IndexStride * Segment.FirstPrimitive * IndicesPerPrimitive;
+
+					Desc.Triangles.VertexCount = VertexBuffer->ResourceLocation.GetSize() / Segment.VertexBufferStride;
+
+					IndexBuffer->GetResource()->UpdateResidency(CommandContext.CommandListHandle);
+				}
+				else
+				{
+					// Non-indexed geometry
+					Desc.Triangles.IndexFormat = DXGI_FORMAT_UNKNOWN;
+					Desc.Triangles.IndexCount = 0;
+					Desc.Triangles.IndexBuffer = D3D12_GPU_VIRTUAL_ADDRESS(0);
+
+					checkf(Segments.Num() == 1, TEXT("Non-indexed geometry with multiple segments is not implemented."));
 
-				Desc.Triangles.VertexCount = VertexBuffer->ResourceLocation.GetSize() / Segment.VertexBufferStride;
+					Desc.Triangles.VertexCount = FMath::Min<uint32>(VertexBuffer->ResourceLocation.GetSize() / Segment.VertexBufferStride, TotalPrimitiveCount * 3);
+				}
+
+				Desc.Triangles.VertexBuffer.StartAddress = VertexBuffer->ResourceLocation.GetGPUVirtualAddress() + Segment.VertexBufferOffset;
+				Desc.Triangles.VertexBuffer.StrideInBytes = Segment.VertexBufferStride;
 
-				IndexBuffer->GetResource()->UpdateResidency(CommandContext.CommandListHandle);
+			}
+			else if (GeometryType == D3D12_RAYTRACING_GEOMETRY_TYPE_PROCEDURAL_PRIMITIVE_AABBS)
+			{
+				Desc.AABBs.AABBCount = Segment.NumPrimitives;
+				Desc.AABBs.AABBs.StartAddress = VertexBuffer->ResourceLocation.GetGPUVirtualAddress() + Segment.VertexBufferOffset;
+				Desc.AABBs.AABBs.StrideInBytes = Segment.VertexBufferStride;
 			}
 			else
 			{
-				// Non-indexed geometry
-				Desc.Triangles.IndexFormat = DXGI_FORMAT_UNKNOWN;
-				Desc.Triangles.IndexCount = 0;
-				Desc.Triangles.IndexBuffer = D3D12_GPU_VIRTUAL_ADDRESS(0);
-
-				checkf(Segments.Num() == 1, TEXT("Non-indexed geometry with multiple segments is not implemented."));
-
-				Desc.Triangles.VertexCount = FMath::Min<uint32>(VertexBuffer->ResourceLocation.GetSize() / Segment.VertexBufferStride, TotalPrimitiveCount * 3);
+				checkf(false, TEXT("Unexpected ray tracing geometry type"));
 			}
-
-			Desc.Triangles.VertexBuffer.StartAddress = VertexBuffer->ResourceLocation.GetGPUVirtualAddress() + Segment.VertexBufferOffset;
-			Desc.Triangles.VertexBuffer.StrideInBytes = Segment.VertexBufferStride;
-		}
-		else if (GeometryType == D3D12_RAYTRACING_GEOMETRY_TYPE_PROCEDURAL_PRIMITIVE_AABBS)
-		{
-			Desc.AABBs.AABBCount = Segment.NumPrimitives;
-			Desc.AABBs.AABBs.StartAddress = VertexBuffer->ResourceLocation.GetGPUVirtualAddress() + Segment.VertexBufferOffset;
-			Desc.AABBs.AABBs.StrideInBytes = Segment.VertexBufferStride;
-		}
-		else
-		{
-			checkf(false, TEXT("Unexpected ray tracing geometry type"));
 		}
 
 		VertexBuffer->ResourceLocation.GetResource()->UpdateResidency(CommandContext.CommandListHandle);
diff --git a/Engine/Source/Runtime/Engine/Private/SkeletalMesh.cpp b/Engine/Source/Runtime/Engine/Private/SkeletalMesh.cpp
--- a/Engine/Source/Runtime/Engine/Private/SkeletalMesh.cpp
+++ b/Engine/Source/Runtime/Engine/Private/SkeletalMesh.cpp
@@ -5062,6 +5062,7 @@ void FSkeletalMeshSceneProxy::GetDynamicRayTracingInstances(FRayTracingMaterialG
 						FRayTracingGeometrySegment Segment;
 						Segment.FirstPrimitive = Section.BaseIndex / 3;
 						Segment.NumPrimitives = Section.NumTriangles;
+						Segment.bEnabled = !Section.bDisabled;
 						GeometrySections.Add(Segment);
 					}
 					MeshObject->GetRayTracingGeometry()->Initializer.Segments = GeometrySections;
diff --git a/Engine/Source/Runtime/Engine/Private/SkeletalRenderCPUSkin.cpp b/Engine/Source/Runtime/Engine/Private/SkeletalRenderCPUSkin.cpp
--- a/Engine/Source/Runtime/Engine/Private/SkeletalRenderCPUSkin.cpp
+++ b/Engine/Source/Runtime/Engine/Private/SkeletalRenderCPUSkin.cpp
@@ -402,6 +402,7 @@ void FSkeletalMeshObjectCPUSkin::FSkeletalMeshObjectLOD::InitResources(FSkelMesh
 					Segment.VertexBufferElementType = VET_Float3;
 					Segment.FirstPrimitive = Section.BaseIndex / 3;
 					Segment.NumPrimitives = Section.NumTriangles;
+					Segment.bEnabled = !Section.bDisabled;
 					GeometrySections.Add(Segment);
 				}
 				Initializer.Segments = GeometrySections;
diff --git a/Engine/Source/Runtime/Engine/Private/SkeletalRenderGPUSkin.cpp b/Engine/Source/Runtime/Engine/Private/SkeletalRenderGPUSkin.cpp
--- a/Engine/Source/Runtime/Engine/Private/SkeletalRenderGPUSkin.cpp
+++ b/Engine/Source/Runtime/Engine/Private/SkeletalRenderGPUSkin.cpp
@@ -433,6 +433,7 @@ void FSkeletalMeshObjectGPUSkin::UpdateDynamicData_RenderThread(FGPUSkinCache* G
 					Segment.VertexBufferOffset = 0;
 					Segment.FirstPrimitive = Section.BaseIndex / 3;
 					Segment.NumPrimitives = Section.NumTriangles;
+					Segment.bEnabled = !Section.bDisabled;
 					Initializer.Segments.Add(Segment);
 				}
 
diff --git a/Engine/Source/Runtime/Engine/Private/SkeletalRenderStatic.cpp b/Engine/Source/Runtime/Engine/Private/SkeletalRenderStatic.cpp
--- a/Engine/Source/Runtime/Engine/Private/SkeletalRenderStatic.cpp
+++ b/Engine/Source/Runtime/Engine/Private/SkeletalRenderStatic.cpp
@@ -159,6 +159,7 @@ void FSkeletalMeshObjectStatic::FSkeletalMeshObjectLOD::InitResources(FSkelMeshC
 					Segment.VertexBufferStride = VertexBufferStride;
 					Segment.FirstPrimitive = Section.BaseIndex / 3;
 					Segment.NumPrimitives = Section.NumTriangles;
+					Segment.bEnabled = !Section.bDisabled;
 					GeometrySections.Add(Segment);
 				}
 				Initializer.Segments = GeometrySections;
diff --git a/Engine/Source/Runtime/RHI/Public/DynamicRHI.h b/Engine/Source/Runtime/RHI/Public/DynamicRHI.h
--- a/Engine/Source/Runtime/RHI/Public/DynamicRHI.h
+++ b/Engine/Source/Runtime/RHI/Public/DynamicRHI.h
@@ -140,6 +140,9 @@ struct FRayTracingGeometrySegment
 	// Any-hit shader may be invoked multiple times for the same primitive during ray traversal.
 	// Setting this to `false` guarantees that only a single instance of any-hit shader will run per primitive, at some performance cost.
 	bool bAllowDuplicateAnyHitShaderInvocation = true;
+
+	// Indicates whether this section is enabled and should be taken into account during acceleration structure creation
+	bool bEnabled = true;
 };
 
 struct FRayTracingGeometryInitializer
