More robust system to discard hidden sections of skeletal meshes in ray tracing

diff --git a/Engine/Source/Runtime/D3D12RHI/Private/D3D12RayTracing.cpp b/Engine/Source/Runtime/D3D12RHI/Private/D3D12RayTracing.cpp
--- a/Engine/Source/Runtime/D3D12RHI/Private/D3D12RayTracing.cpp
+++ b/Engine/Source/Runtime/D3D12RHI/Private/D3D12RayTracing.cpp
@@ -2025,9 +2025,22 @@ FRayTracingSceneRHIRef FD3D12DynamicRHI::RHICreateRayTracingScene(const FRayTrac
 	return Result;
 }
 
+FVertexBufferRHIRef FD3D12RayTracingGeometry::NullTransformBuffer;
+
 FD3D12RayTracingGeometry::FD3D12RayTracingGeometry()
 {
 	INC_DWORD_STAT(STAT_D3D12RayTracingAllocatedBLAS);
+
+	if(!FD3D12RayTracingGeometry::NullTransformBuffer.IsValid())
+	{
+		TResourceArray<float> NullTransformData;
+		NullTransformData.SetNumZeroed(12);
+
+		FRHIResourceCreateInfo CreateInfo;
+		CreateInfo.ResourceArray = &NullTransformData;
+
+		FD3D12RayTracingGeometry::NullTransformBuffer = RHICreateVertexBuffer(NullTransformData.GetResourceDataSize(), BUF_Static, CreateInfo);
+	}
 }
 
 FD3D12RayTracingGeometry::~FD3D12RayTracingGeometry()
@@ -2118,6 +2131,7 @@ void FD3D12RayTracingGeometry::BuildAccelerationStructure(FD3D12CommandContext&
 	Descs.Reserve(Segments.Num());
 
 	FD3D12IndexBuffer* IndexBuffer = CommandContext.RetrieveObject<FD3D12IndexBuffer>(RHIIndexBuffer.GetReference());
+	FD3D12VertexBuffer* NullTransformBufferD3D12 = CommandContext.RetrieveObject<FD3D12VertexBuffer>(NullTransformBuffer.GetReference());
 
 	for (const FRayTracingGeometrySegment& Segment : Segments)
 	{
@@ -2139,71 +2153,74 @@ void FD3D12RayTracingGeometry::BuildAccelerationStructure(FD3D12CommandContext&
 		}
 
 		FD3D12VertexBuffer* VertexBuffer = CommandContext.RetrieveObject<FD3D12VertexBuffer>(Segment.VertexBuffer.GetReference());
-
-		if (Segment.bEnabled)
+			   
+		if (GeometryType == D3D12_RAYTRACING_GEOMETRY_TYPE_TRIANGLES)
 		{
-			if (GeometryType == D3D12_RAYTRACING_GEOMETRY_TYPE_TRIANGLES)
+			switch (Segment.VertexBufferElementType)
 			{
-				switch (Segment.VertexBufferElementType)
-				{
-				case VET_Float4:
-					// While the DXGI_FORMAT_R32G32B32A32_FLOAT format is not supported by DXR, since we manually load vertex 
-					// data when we are building the BLAS, we can just rely on the vertex stride to offset the read index, 
-					// and read only the 3 vertex components, and so use the DXGI_FORMAT_R32G32B32_FLOAT vertex format
-				case VET_Float3:
-					Desc.Triangles.VertexFormat = DXGI_FORMAT_R32G32B32_FLOAT;
-					break;
-				case VET_Float2:
-					Desc.Triangles.VertexFormat = DXGI_FORMAT_R32G32_FLOAT;
-					break;
-				case VET_Half2:
-					Desc.Triangles.VertexFormat = DXGI_FORMAT_R16G16_FLOAT;
-					break;
-				default:
-					checkNoEntry();
-					break;
-				}
+			case VET_Float4:
+				// While the DXGI_FORMAT_R32G32B32A32_FLOAT format is not supported by DXR, since we manually load vertex 
+				// data when we are building the BLAS, we can just rely on the vertex stride to offset the read index, 
+				// and read only the 3 vertex components, and so use the DXGI_FORMAT_R32G32B32_FLOAT vertex format
+			case VET_Float3:
+				Desc.Triangles.VertexFormat = DXGI_FORMAT_R32G32B32_FLOAT;
+				break;
+			case VET_Float2:
+				Desc.Triangles.VertexFormat = DXGI_FORMAT_R32G32_FLOAT;
+				break;
+			case VET_Half2:
+				Desc.Triangles.VertexFormat = DXGI_FORMAT_R16G16_FLOAT;
+				break;
+			default:
+				checkNoEntry();
+				break;
+			}
 
+			if (Segment.bEnabled)
+			{
 				Desc.Triangles.Transform3x4 = D3D12_GPU_VIRTUAL_ADDRESS(0);
+			}
+			else
+			{
+				Desc.Triangles.Transform3x4 = NullTransformBufferD3D12->ResourceLocation.GetGPUVirtualAddress();
+			}
 
-				if (IndexBuffer)
-				{
-					Desc.Triangles.IndexFormat = IndexStride == 4 ? DXGI_FORMAT_R32_UINT : DXGI_FORMAT_R16_UINT;
-					Desc.Triangles.IndexCount = Segment.NumPrimitives * IndicesPerPrimitive;
-					Desc.Triangles.IndexBuffer = IndexBuffer->ResourceLocation.GetGPUVirtualAddress() +
-						IndexOffsetInBytes +
-						IndexStride * Segment.FirstPrimitive * IndicesPerPrimitive;
-
-					Desc.Triangles.VertexCount = VertexBuffer->ResourceLocation.GetSize() / Segment.VertexBufferStride;
-
-					IndexBuffer->GetResource()->UpdateResidency(CommandContext.CommandListHandle);
-				}
-				else
-				{
-					// Non-indexed geometry
-					Desc.Triangles.IndexFormat = DXGI_FORMAT_UNKNOWN;
-					Desc.Triangles.IndexCount = 0;
-					Desc.Triangles.IndexBuffer = D3D12_GPU_VIRTUAL_ADDRESS(0);
-
-					checkf(Segments.Num() == 1, TEXT("Non-indexed geometry with multiple segments is not implemented."));
-
-					Desc.Triangles.VertexCount = FMath::Min<uint32>(VertexBuffer->ResourceLocation.GetSize() / Segment.VertexBufferStride, TotalPrimitiveCount * 3);
-				}
+			if (IndexBuffer)
+			{
+				Desc.Triangles.IndexFormat = IndexStride == 4 ? DXGI_FORMAT_R32_UINT : DXGI_FORMAT_R16_UINT;
+				Desc.Triangles.IndexCount = Segment.NumPrimitives * IndicesPerPrimitive;
+				Desc.Triangles.IndexBuffer = IndexBuffer->ResourceLocation.GetGPUVirtualAddress() +
+					IndexOffsetInBytes +
+					IndexStride * Segment.FirstPrimitive * IndicesPerPrimitive;
 
-				Desc.Triangles.VertexBuffer.StartAddress = VertexBuffer->ResourceLocation.GetGPUVirtualAddress() + Segment.VertexBufferOffset;
-				Desc.Triangles.VertexBuffer.StrideInBytes = Segment.VertexBufferStride;
+				Desc.Triangles.VertexCount = VertexBuffer->ResourceLocation.GetSize() / Segment.VertexBufferStride;
 
-			}
-			else if (GeometryType == D3D12_RAYTRACING_GEOMETRY_TYPE_PROCEDURAL_PRIMITIVE_AABBS)
-			{
-				Desc.AABBs.AABBCount = Segment.NumPrimitives;
-				Desc.AABBs.AABBs.StartAddress = VertexBuffer->ResourceLocation.GetGPUVirtualAddress() + Segment.VertexBufferOffset;
-				Desc.AABBs.AABBs.StrideInBytes = Segment.VertexBufferStride;
+				IndexBuffer->GetResource()->UpdateResidency(CommandContext.CommandListHandle);
 			}
 			else
 			{
-				checkf(false, TEXT("Unexpected ray tracing geometry type"));
+				// Non-indexed geometry
+				Desc.Triangles.IndexFormat = DXGI_FORMAT_UNKNOWN;
+				Desc.Triangles.IndexCount = 0;
+				Desc.Triangles.IndexBuffer = D3D12_GPU_VIRTUAL_ADDRESS(0);
+
+				checkf(Segments.Num() == 1, TEXT("Non-indexed geometry with multiple segments is not implemented."));
+
+				Desc.Triangles.VertexCount = FMath::Min<uint32>(VertexBuffer->ResourceLocation.GetSize() / Segment.VertexBufferStride, TotalPrimitiveCount * 3);
 			}
+
+			Desc.Triangles.VertexBuffer.StartAddress = VertexBuffer->ResourceLocation.GetGPUVirtualAddress() + Segment.VertexBufferOffset;
+			Desc.Triangles.VertexBuffer.StrideInBytes = Segment.VertexBufferStride;
+		}
+		else if (GeometryType == D3D12_RAYTRACING_GEOMETRY_TYPE_PROCEDURAL_PRIMITIVE_AABBS)
+		{
+			Desc.AABBs.AABBCount = Segment.NumPrimitives;
+			Desc.AABBs.AABBs.StartAddress = VertexBuffer->ResourceLocation.GetGPUVirtualAddress() + Segment.VertexBufferOffset;
+			Desc.AABBs.AABBs.StrideInBytes = Segment.VertexBufferStride;
+		}
+		else
+		{
+			checkf(false, TEXT("Unexpected ray tracing geometry type"));
 		}
 
 		VertexBuffer->ResourceLocation.GetResource()->UpdateResidency(CommandContext.CommandListHandle);
diff --git a/Engine/Source/Runtime/D3D12RHI/Private/D3D12RayTracing.h b/Engine/Source/Runtime/D3D12RHI/Private/D3D12RayTracing.h
--- a/Engine/Source/Runtime/D3D12RHI/Private/D3D12RayTracing.h
+++ b/Engine/Source/Runtime/D3D12RHI/Private/D3D12RayTracing.h
@@ -44,6 +44,7 @@ public:
 	D3D12_RAYTRACING_ACCELERATION_STRUCTURE_BUILD_FLAGS BuildFlags;
 
 	FIndexBufferRHIRef  RHIIndexBuffer;
+	static FVertexBufferRHIRef NullTransformBuffer; // Null transform for hidden sections
 
 	TRefCountPtr<FD3D12MemBuffer> AccelerationStructureBuffers[MAX_NUM_GPUS];
 	TRefCountPtr<FD3D12MemBuffer> ScratchBuffers[MAX_NUM_GPUS];
