Fix missing primitive stats for DX12 and incorrect primitive stats for other RHIs

diff --git a/Engine/Source/Runtime/D3D12RHI/Private/D3D12CommandContext.cpp b/Engine/Source/Runtime/D3D12RHI/Private/D3D12CommandContext.cpp
--- a/Engine/Source/Runtime/D3D12RHI/Private/D3D12CommandContext.cpp
+++ b/Engine/Source/Runtime/D3D12RHI/Private/D3D12CommandContext.cpp
@@ -187,7 +187,6 @@ void FD3D12CommandContext::OpenCommandList()
 	// Mark state as dirty so next time ApplyState is called, it will set all state on this new command list
 	StateCache.DirtyStateForNewCommandList();
 
-	numPrimitives = 0;
 	numDraws = 0;
 	numDispatches = 0;
 	numClears = 0;
@@ -200,9 +199,17 @@ void FD3D12CommandContext::CloseCommandList()
 {
 	CommandListHandle.Close();
 
+	uint32 NumTriangles = StateCache.GetNumTrianglesStat();
+	uint32 NumLines     = StateCache.GetNumLinesStat();
+
+#if STATS
 	INC_DWORD_STAT_BY(STAT_RHIDrawPrimitiveCalls, numDraws);
+	INC_DWORD_STAT_BY(STAT_RHILines, NumLines);
+	INC_DWORD_STAT_BY(STAT_RHITriangles, NumTriangles);
+#endif
+
 	FPlatformAtomics::InterlockedAdd(&GCurrentNumDrawCallsRHI, numDraws);
-	FPlatformAtomics::InterlockedAdd(&GCurrentNumPrimitivesDrawnRHI, numPrimitives);
+	FPlatformAtomics::InterlockedAdd(&GCurrentNumPrimitivesDrawnRHI, NumLines + NumTriangles);
 }
 
 FD3D12CommandListHandle FD3D12CommandContext::FlushCommands(bool WaitForCompletion, EFlushCommandsExtraAction ExtraAction)
diff --git a/Engine/Source/Runtime/D3D12RHI/Private/D3D12CommandContext.h b/Engine/Source/Runtime/D3D12RHI/Private/D3D12CommandContext.h
--- a/Engine/Source/Runtime/D3D12RHI/Private/D3D12CommandContext.h
+++ b/Engine/Source/Runtime/D3D12RHI/Private/D3D12CommandContext.h
@@ -174,7 +174,6 @@ public:
 	virtual void FlushTextureCache() {};
 #endif
 
-	uint32 numPrimitives;
 	uint32 numDraws;
 	uint32 numDispatches;
 	uint32 numClears;
diff --git a/Engine/Source/Runtime/D3D12RHI/Private/D3D12Commands.cpp b/Engine/Source/Runtime/D3D12RHI/Private/D3D12Commands.cpp
--- a/Engine/Source/Runtime/D3D12RHI/Private/D3D12Commands.cpp
+++ b/Engine/Source/Runtime/D3D12RHI/Private/D3D12Commands.cpp
@@ -1424,11 +1424,9 @@ void FD3D12CommandContext::RHIDrawPrimitive(uint32 BaseVertexIndex, uint32 NumPr
 	CommitGraphicsResourceTables();
 	CommitNonComputeShaderConstants();
 
-	uint32 VertexCount = GetVertexCountForPrimitiveCount(NumPrimitives, StateCache.GetGraphicsPipelinePrimitiveType());
-
+	uint32 VertexCount = StateCache.GetVertexCountAndIncrementStat(NumPrimitives);
 	NumInstances = FMath::Max<uint32>(1, NumInstances);
 	numDraws++;
-	numPrimitives += NumInstances * NumPrimitives;
 	if (bTrackingEvents)
 	{
 		GetParentDevice()->RegisterGPUWork(NumPrimitives * NumInstances, VertexCount * NumInstances);
@@ -1526,16 +1524,15 @@ void FD3D12CommandContext::RHIDrawIndexedPrimitive(FRHIIndexBuffer* IndexBufferR
 
 	NumInstances = FMath::Max<uint32>(1, NumInstances);
 	numDraws++;
-	numPrimitives += NumInstances * NumPrimitives;
 	if (bTrackingEvents)
 	{
 		GetParentDevice()->RegisterGPUWork(NumPrimitives * NumInstances, NumVertices * NumInstances);
 	}
-	uint32 IndexCount = GetVertexCountForPrimitiveCount(NumPrimitives, StateCache.GetGraphicsPipelinePrimitiveType());
 
 	CommitGraphicsResourceTables();
 	CommitNonComputeShaderConstants();
 
+	uint32 IndexCount = StateCache.GetVertexCountAndIncrementStat(NumPrimitives);
 	FD3D12IndexBuffer* IndexBuffer = RetrieveObject<FD3D12IndexBuffer>(IndexBufferRHI);
 
 	// Verify that we are not trying to read outside the index buffer range
diff --git a/Engine/Source/Runtime/D3D12RHI/Private/D3D12StateCache.cpp b/Engine/Source/Runtime/D3D12RHI/Private/D3D12StateCache.cpp
--- a/Engine/Source/Runtime/D3D12RHI/Private/D3D12StateCache.cpp
+++ b/Engine/Source/Runtime/D3D12RHI/Private/D3D12StateCache.cpp
@@ -208,6 +208,9 @@ void FD3D12StateCacheBase::DirtyStateForNewCommandList()
 	PipelineState.Graphics.bNeedSetRootSignature = true;
 	bNeedSetPrimitiveTopology = true;
 
+	PipelineState.Graphics.NumLines = 0;
+	PipelineState.Graphics.NumTriangles = 0;
+
 	if (PipelineState.Graphics.VBCache.BoundVBMask) { bNeedSetVB = true; }
 
 	// IndexBuffers are set in DrawIndexed*() calls, so there's no way to depend on previously set IndexBuffers without making a new DrawIndexed*() call.
diff --git a/Engine/Source/Runtime/D3D12RHI/Private/D3D12StateCachePrivate.h b/Engine/Source/Runtime/D3D12RHI/Private/D3D12StateCachePrivate.h
--- a/Engine/Source/Runtime/D3D12RHI/Private/D3D12StateCachePrivate.h
+++ b/Engine/Source/Runtime/D3D12RHI/Private/D3D12StateCachePrivate.h
@@ -384,6 +384,11 @@ protected:
 			// Primitive Topology State
 			EPrimitiveType CurrentPrimitiveType;
 			D3D_PRIMITIVE_TOPOLOGY CurrentPrimitiveTopology;
+			uint32 PrimitiveTypeFactor;
+			uint32 PrimitiveTypeOffset;
+			uint32* CurrentPrimitiveStat;
+			uint32 NumTriangles;
+			uint32 NumLines;
 
 			// Input Layout State
 			D3D12_RECT CurrentScissorRects[D3D12_VIEWPORT_AND_SCISSORRECT_OBJECT_COUNT_PER_PIPELINE];
@@ -526,7 +531,7 @@ public:
 		return PipelineState.Graphics.CurrentPipelineStateObject;
 	}
 
-	const FD3D12RootSignature* GetGraphicsRootSignature()
+	const FD3D12RootSignature* GetGraphicsRootSignature() const
 	{
 		return PipelineState.Graphics.CurrentPipelineStateObject ? PipelineState.Graphics.CurrentPipelineStateObject->RootSignature : nullptr;
 	}
@@ -536,7 +541,16 @@ public:
 		return PipelineState.Graphics.CurrentPrimitiveType;
 	}
 
-	const FD3D12RootSignature* GetComputeRootSignature()
+	inline uint32 GetVertexCountAndIncrementStat(uint32 NumPrimitives)
+	{
+		*PipelineState.Graphics.CurrentPrimitiveStat += NumPrimitives;
+		return PipelineState.Graphics.PrimitiveTypeFactor * NumPrimitives + PipelineState.Graphics.PrimitiveTypeOffset;
+	}
+
+	inline uint32 GetNumTrianglesStat() const { return PipelineState.Graphics.NumTriangles; }
+	inline uint32 GetNumLinesStat() const { return PipelineState.Graphics.NumLines; }
+
+	const FD3D12RootSignature* GetComputeRootSignature() const
 	{
 		return PipelineState.Compute.CurrentPipelineStateObject ? PipelineState.Compute.CurrentPipelineStateObject->ComputeShader->pRootSignature : nullptr;
 	}
@@ -772,6 +786,11 @@ public:
 				PipelineState.Graphics.CurrentPrimitiveType = PrimitiveType;
 				PipelineState.Graphics.CurrentPrimitiveTopology = GetD3D12PrimitiveType(PrimitiveType, bUsingTessellation);
 				bNeedSetPrimitiveTopology = true;
+
+				static_assert(PT_Num == 38, "This computation needs to be updated, matching that of GetVertexCountForPrimitiveCount()");
+				PipelineState.Graphics.PrimitiveTypeFactor = (PrimitiveType == PT_TriangleList)? 3 : (PrimitiveType == PT_LineList)? 2 : (PrimitiveType == PT_RectList)? 3 : (PrimitiveType >= PT_1_ControlPointPatchList)? (PrimitiveType - PT_1_ControlPointPatchList + 1) : 1;
+				PipelineState.Graphics.PrimitiveTypeOffset = (PrimitiveType == PT_TriangleStrip)? 2 : 0;
+				PipelineState.Graphics.CurrentPrimitiveStat = (PrimitiveType == PT_LineList)? &PipelineState.Graphics.NumLines : &PipelineState.Graphics.NumTriangles;
 			}
 
 			// Set the PSO
diff --git a/Engine/Source/Runtime/Windows/D3D11RHI/Private/D3D11Commands.cpp b/Engine/Source/Runtime/Windows/D3D11RHI/Private/D3D11Commands.cpp
--- a/Engine/Source/Runtime/Windows/D3D11RHI/Private/D3D11Commands.cpp
+++ b/Engine/Source/Runtime/Windows/D3D11RHI/Private/D3D11Commands.cpp
@@ -1661,7 +1661,7 @@ void FD3D11DynamicRHI::CommitComputeResourceTables(FD3D11ComputeShader* InComput
 
 void FD3D11DynamicRHI::RHIDrawPrimitive(uint32 BaseVertexIndex,uint32 NumPrimitives,uint32 NumInstances)
 {
-	RHI_DRAW_CALL_STATS(PrimitiveType,NumInstances*NumPrimitives);
+	RHI_DRAW_CALL_STATS(PrimitiveType, FMath::Max(NumInstances, 1U) * NumPrimitives);
 
 	CommitGraphicsResourceTables();
 	CommitNonComputeShaderConstants();
@@ -1727,13 +1727,13 @@ void FD3D11DynamicRHI::RHIDrawIndexedIndirect(FRHIIndexBuffer* IndexBufferRHI, F
 
 void FD3D11DynamicRHI::RHIDrawIndexedPrimitive(FRHIIndexBuffer* IndexBufferRHI,int32 BaseVertexIndex,uint32 FirstInstance,uint32 NumVertices,uint32 StartIndex,uint32 NumPrimitives,uint32 NumInstances)
 {
+	RHI_DRAW_CALL_STATS(PrimitiveType, FMath::Max(NumInstances, 1U) * NumPrimitives);
+
 	FD3D11IndexBuffer* IndexBuffer = ResourceCast(IndexBufferRHI);
 
 	// called should make sure the input is valid, this avoid hidden bugs
 	ensure(NumPrimitives > 0);
 
-	RHI_DRAW_CALL_STATS(PrimitiveType,NumInstances*NumPrimitives);
-
 	GPUProfilingData.RegisterGPUWork(NumPrimitives * NumInstances, NumVertices * NumInstances);
 
 	CommitGraphicsResourceTables();
