Revert to windowed mode before shutting down. This fixes DX12 crash on shutdown when in fullscreen.

diff --git a/Engine/Source/Runtime/D3D12RHI/Private/D3D12Viewport.cpp b/Engine/Source/Runtime/D3D12RHI/Private/D3D12Viewport.cpp
--- a/Engine/Source/Runtime/D3D12RHI/Private/D3D12Viewport.cpp
+++ b/Engine/Source/Runtime/D3D12RHI/Private/D3D12Viewport.cpp
@@ -217,9 +217,6 @@ FD3D12Texture2D* GetSwapChainSurface(FD3D12Device* Parent, EPixelFormat PixelFor
 		Parent->GetDevice()->CreateCommittedResource(&HeapProps, D3D12_HEAP_FLAG_NONE, &TextureDesc, D3D12_RESOURCE_STATE_PRESENT, nullptr, IID_PPV_ARGS(BackBufferResource.GetInitReference()));
 	}
 
-	FString Name = FString::Printf(TEXT("BackBuffer%d"), BackBufferIndex);
-	SetName(BackBufferResource, *Name);
-
 	D3D12_RESOURCE_DESC BackBufferDesc = BackBufferResource->GetDesc();
 
 	FD3D12Texture2D* SwapChainTexture = Adapter->CreateLinkedObject<FD3D12Texture2D>(FRHIGPUMask::All(), [&](FD3D12Device* Device)
@@ -322,6 +319,9 @@ FD3D12Texture2D* GetSwapChainSurface(FD3D12Device* Parent, EPixelFormat PixelFor
 		return NewTexture;
 	});
 
+	FString Name = FString::Printf(TEXT("BackBuffer%d"), BackBufferIndex);
+	SetName(SwapChainTexture->GetResource(), *Name);
+
 	FD3D12TextureStats::D3D12TextureAllocated2D(*SwapChainTexture);
 	return SwapChainTexture;
 }
@@ -330,6 +330,15 @@ FD3D12Viewport::~FD3D12Viewport()
 {
 	check(IsInRenderingThread());
 
+	// If the swap chain was in fullscreen mode, switch back to windowed before releasing the swap chain.
+	// DXGI throws an error otherwise.
+#if !PLATFORM_HOLOLENS
+	if (SwapChain1)
+	{
+		SwapChain1->SetFullscreenState(0, nullptr);
+	}
+#endif
+
 	GetParentAdapter()->GetViewports().Remove(this);
 
 #if WITH_MGPU
