Recover from failed SetFullscreenState()

diff --git a/Engine/Source/Runtime/D3D12RHI/Private/Windows/WindowsD3D12Viewport.cpp b/Engine/Source/Runtime/D3D12RHI/Private/Windows/WindowsD3D12Viewport.cpp
--- a/Engine/Source/Runtime/D3D12RHI/Private/Windows/WindowsD3D12Viewport.cpp
+++ b/Engine/Source/Runtime/D3D12RHI/Private/Windows/WindowsD3D12Viewport.cpp
@@ -263,6 +263,17 @@ void FD3D12Viewport::ConditionalResetSwapChain(bool bIgnoreFocus)
 						break;
 					}
 					UE_LOG(LogD3D12RHI, Error, TEXT("IDXGISwapChain::SetFullscreenState returned 0x%08x, %s."), Result, Name);
+
+					if (bIsFullscreen)
+					{
+						// Something went wrong, attempt to proceed in windowed mode.
+						Result = SwapChain1->SetFullscreenState(FALSE, nullptr);
+						if (SUCCEEDED(Result))
+						{
+							bIsValid = true;
+							bIsFullscreen = false;
+						}
+					}
 				}
 			}
 		}
