Reapply Emil's fix merge error CL9354108

diff --git a/Engine/Source/Runtime/D3D12RHI/Private/D3D12CommandContext.cpp b/Engine/Source/Runtime/D3D12RHI/Private/D3D12CommandContext.cpp
--- a/Engine/Source/Runtime/D3D12RHI/Private/D3D12CommandContext.cpp
+++ b/Engine/Source/Runtime/D3D12RHI/Private/D3D12CommandContext.cpp
@@ -421,12 +421,6 @@ void FD3D12CommandContextBase::RHIEndFrame()
 	    // Stop Timing at the very last moment
     
 	    ParentAdapter->GetGPUProfiler().EndFrame(ParentAdapter->GetOwningRHI());
-    
-    
-	    // Advance frame fence
-    
-	    FD3D12ManualFence& FrameFence = ParentAdapter->GetFrameFence();
-	    FrameFence.Signal(ED3D12CommandQueueType::Default, FrameFence.IncrementCurrentFence());
 	}
 
 void FD3D12CommandContextBase::UpdateMemoryStats()
