Fixed merging errors

diff --git a/Engine/Plugins/FX/Niagara/Source/NiagaraEditor/Private/NiagaraCompiler.cpp b/Engine/Plugins/FX/Niagara/Source/NiagaraEditor/Private/NiagaraCompiler.cpp
--- a/Engine/Plugins/FX/Niagara/Source/NiagaraEditor/Private/NiagaraCompiler.cpp
+++ b/Engine/Plugins/FX/Niagara/Source/NiagaraEditor/Private/NiagaraCompiler.cpp
@@ -1065,7 +1065,8 @@ TOptional<FNiagaraCompileResults> FHlslNiagaraCompiler::GetCompileResult(int32 J
 				int32 NewBindingIdx = Results.Data->CalledVMExternalFunctions.AddDefaulted();
 				Results.Data->CalledVMExternalFunctions[NewBindingIdx].Name = *Sig->GetName();
 				Results.Data->CalledVMExternalFunctions[NewBindingIdx].OwnerName = Sig->OwnerName;
-				Results.Data->CalledVMExternalFunctions[NewBindingIdx].Specifiers = Sig->FunctionSpecifiers;
+				//#todo: fix me
+				//Results.Data->CalledVMExternalFunctions[NewBindingIdx].Specifiers = Sig->FunctionSpecifiers;
 
 				Results.Data->CalledVMExternalFunctions[NewBindingIdx].InputParamLocations = FuncInfo.InputParamLocations;
 				Results.Data->CalledVMExternalFunctions[NewBindingIdx].NumOutputs = FuncInfo.NumOutputs;
diff --git a/Engine/Plugins/FX/Niagara/Source/NiagaraShader/Private/NiagaraShader.cpp b/Engine/Plugins/FX/Niagara/Source/NiagaraShader/Private/NiagaraShader.cpp
--- a/Engine/Plugins/FX/Niagara/Source/NiagaraShader/Private/NiagaraShader.cpp
+++ b/Engine/Plugins/FX/Niagara/Source/NiagaraShader/Private/NiagaraShader.cpp
@@ -907,14 +907,14 @@ void FNiagaraShaderMap::Compile(
 			// Mark as not having been compiled
 			bCompiledSuccessfully = false;
   
-			FNiagaraShaderCompilationManager::Get().AddJobs(NewJobs);
+			GNiagaraShaderCompilationManager.AddJobs(NewJobs);
   
 			// Compile the shaders for this shader map now if not deferring and deferred compiles are not enabled globally
 			if (bSynchronousCompile)
 			{
 				TArray<int32> CurrentShaderMapId;
 				CurrentShaderMapId.Add(CompilingId);
-				FNiagaraShaderCompilationManager::Get().FinishCompilation(*FriendlyName, CurrentShaderMapId);
+				GNiagaraShaderCompilationManager.FinishCompilation(*FriendlyName, CurrentShaderMapId);
 			}
 		}
 	}
diff --git a/Engine/Plugins/FX/Niagara/Source/NiagaraShader/Private/NiagaraShared.cpp b/Engine/Plugins/FX/Niagara/Source/NiagaraShader/Private/NiagaraShared.cpp
--- a/Engine/Plugins/FX/Niagara/Source/NiagaraShader/Private/NiagaraShared.cpp
+++ b/Engine/Plugins/FX/Niagara/Source/NiagaraShader/Private/NiagaraShared.cpp
@@ -378,7 +378,7 @@ void FNiagaraShaderScript::FinishCompilation()
 		}
 		// Block until the shader maps that we will save have finished being compiled
 		// NIAGARATODO: implement when async compile works
-		FNiagaraShaderCompilationManager::Get().FinishCompilation(*GetFriendlyName(), ShaderMapIdsToFinish);
+		GNiagaraShaderCompilationManager.FinishCompilation(*GetFriendlyName(), ShaderMapIdsToFinish);
 
 		// Shouldn't have anything left to do...
 		TArray<int32> ShaderMapIdsToFinish2;
diff --git a/Engine/Plugins/FX/Niagara/Source/NiagaraShader/Public/NiagaraShaderCompilationManager.h b/Engine/Plugins/FX/Niagara/Source/NiagaraShader/Public/NiagaraShaderCompilationManager.h
--- a/Engine/Plugins/FX/Niagara/Source/NiagaraShader/Public/NiagaraShaderCompilationManager.h
+++ b/Engine/Plugins/FX/Niagara/Source/NiagaraShader/Public/NiagaraShaderCompilationManager.h
@@ -42,7 +42,6 @@ public:
 
 	NIAGARASHADER_API void AddJobs(TArray<class FShaderCommonCompileJob*> InNewJobs);
 	NIAGARASHADER_API void ProcessAsyncResults();
-	NIAGARASHADER_API static FNiagaraShaderCompilationManager& Get();
 
 	void FinishCompilation(const TCHAR* ScriptName, const TArray<int32>& ShaderMapIdsToFinishCompiling);
 private:
@@ -55,4 +54,7 @@ private:
 
 	/** Map from shader map id to results being finalized.  Used to track shader finalizations over multiple frames. */
 	TMap<int32, FNiagaraShaderMapFinalizeResults> PendingFinalizeNiagaraShaderMaps;
-};
\ No newline at end of file
+};
+
+
+extern NIAGARASHADER_API FNiagaraShaderCompilationManager GNiagaraShaderCompilationManager;
diff --git a/Engine/Source/Runtime/D3D12RHI/Private/Windows/WindowsD3D12PipelineState.cpp b/Engine/Source/Runtime/D3D12RHI/Private/Windows/WindowsD3D12PipelineState.cpp
--- a/Engine/Source/Runtime/D3D12RHI/Private/Windows/WindowsD3D12PipelineState.cpp
+++ b/Engine/Source/Runtime/D3D12RHI/Private/Windows/WindowsD3D12PipelineState.cpp
@@ -693,7 +693,7 @@ static inline void FastHashName(wchar_t Name[17], uint64 Hash)
 	Name[16] = 0;
 }
 
-static void CreatePipelineStateWrapper(ID3D12PipelineState** PSO, FD3D12Adapter* Adapter, const GraphicsPipelineCreationArgs_POD* CreationArgs)
+static void CreatePipelineStateWrapper(ID3D12PipelineState** PSO, FD3D12Adapter* Adapter, const GraphicsPipelineCreationArgs_POD* CreationArgs, bool bUseStream)
 {
 	// Get the pipeline state name, currently based on the hash.
 	wchar_t Name[17];
@@ -706,7 +706,7 @@ static void CreatePipelineStateWrapper(ID3D12PipelineState** PSO, FD3D12Adapter*
 
 	// Use pipeline streams if the system supports it.
 	ID3D12Device2* const pDevice2 = Adapter->GetD3DDevice2();
-	if (pDevice2)
+	if (pDevice2 && bUseStream)
 	{
 		FD3D12_GRAPHICS_PIPELINE_STATE_STREAM Stream = CreationArgs->Desc.Desc.PipelineStateStream();
 		const D3D12_PIPELINE_STATE_STREAM_DESC StreamDesc = { sizeof(Stream), &Stream };
@@ -739,7 +739,7 @@ static void CreatePipelineStateWrapper(ID3D12PipelineState** PSO, FD3D12Adapter*
 	}
 }
 
-static void CreatePipelineStateWrapper(ID3D12PipelineState** PSO, FD3D12Adapter* Adapter, const ComputePipelineCreationArgs_POD* CreationArgs)
+static void CreatePipelineStateWrapper(ID3D12PipelineState** PSO, FD3D12Adapter* Adapter, const ComputePipelineCreationArgs_POD* CreationArgs, bool bUseStream)
 {
 	// Get the pipeline state name, currently based on the hash.
 	wchar_t Name[17];
diff --git a/Engine/Source/Runtime/RHI/Public/RHIDefinitions.h b/Engine/Source/Runtime/RHI/Public/RHIDefinitions.h
--- a/Engine/Source/Runtime/RHI/Public/RHIDefinitions.h
+++ b/Engine/Source/Runtime/RHI/Public/RHIDefinitions.h
@@ -1171,10 +1171,10 @@ inline bool IsHlslccShaderPlatform(const EShaderPlatform Platform)
 
 inline bool IsDeprecatedShaderPlatform(const EShaderPlatform ShaderPlatform)
 {
-	return ShaderPlatform == SP_OPENGL_SM5 || ShaderPlatform == SP_PCD3D_SM4_DEPRECATED || ShaderPlatform == SP_OPENGL_ES2_IOS_DEPRECATED ||
-		ShaderPlatform == SP_PCD3D_ES2 || ShaderPlatform == SP_METAL_MACES2 || ShaderPlatform == SP_OPENGL_PCES2 ||
+	return ShaderPlatform == SP_OPENGL_SM5 || ShaderPlatform == SP_PCD3D_SM4_REMOVED || ShaderPlatform == SP_OPENGL_ES2_IOS_REMOVED ||
+		ShaderPlatform == SP_PCD3D_ES2_DEPRECATED || ShaderPlatform == SP_METAL_MACES2_DEPRECATED || ShaderPlatform == SP_OPENGL_PCES2_DEPRECATED ||
 		ShaderPlatform == SP_OPENGL_ES2_ANDROID || ShaderPlatform == SP_OPENGL_ES2_WEBGL ||
-		ShaderPlatform == SP_VULKAN_SM4_DEPRECATED;
+		ShaderPlatform == SP_VULKAN_SM4_REMOVED;
 }
 
 inline ERHIFeatureLevel::Type GetMaxSupportedFeatureLevel(EShaderPlatform InShaderPlatform)
