Fix PSO creation crash on AMD by disabling currently unused shader extension.

diff --git a/Engine/Source/Runtime/D3D12RHI/Private/D3D12Adapter.cpp b/Engine/Source/Runtime/D3D12RHI/Private/D3D12Adapter.cpp
--- a/Engine/Source/Runtime/D3D12RHI/Private/D3D12Adapter.cpp
+++ b/Engine/Source/Runtime/D3D12RHI/Private/D3D12Adapter.cpp
@@ -157,8 +157,12 @@ void FD3D12Adapter::CreateRootDevice(bool bWithDebug)
 		AmdExtensionParams.pAppName = bDisableAppRegistration ? TEXT("") : FApp::GetProjectName();
 		AmdExtensionParams.appVersion = AGS_UNSPECIFIED_VERSION;
 
+		// UE-88560 - Temporarily disable this AMD shader extension for now until AMD releases fixed drivers.		
+		// As of 2020-02-19, this causes PSO creation failures and device loss on unrelated shaders, preventing AMD users from launching the editor.
+#if 0
 		// Specify custom UAV bind point for the special UAV (custom slot will always assume space0 in the root signature).
 		AmdExtensionParams.uavSlot = 7;
+#endif
 
 		AGSDX12ReturnedParams DeviceCreationReturnedParams;
 		AGSReturnCode DeviceCreation = agsDriverExtensionsDX12_CreateDevice(OwningRHI->GetAmdAgsContext(), &AmdDeviceCreationParams, &AmdExtensionParams, &DeviceCreationReturnedParams);
