Removed unnecesary barrier that was triggering an assert for not being in the rendering thread.

diff --git a/Engine/Source/Runtime/D3D12RHI/Private/D3D12RayTracing.cpp b/Engine/Source/Runtime/D3D12RHI/Private/D3D12RayTracing.cpp
--- a/Engine/Source/Runtime/D3D12RHI/Private/D3D12RayTracing.cpp
+++ b/Engine/Source/Runtime/D3D12RHI/Private/D3D12RayTracing.cpp
@@ -2545,7 +2545,6 @@ void FD3D12RayTracingScene::BuildAccelerationStructure(FD3D12CommandContext& Com
 			}
 
 			RHICmdList.EndUAVOverlap();	
-			RHICmdList.TransitionResource(EResourceTransitionAccess::EReadable, EResourceTransitionPipeline::EGfxToCompute, InstancesDescUAV);
 		}
 
 		FD3D12DynamicRHI::TransitionResource(CommandContext.CommandListHandle, InstanceBuffer.GetReference()->GetResource(), D3D12_RESOURCE_STATE_NON_PIXEL_SHADER_RESOURCE, 0);
