SRV limit in DX12 increased from 48 to 64 so a shader with 50 SRV fits.

diff --git a/Engine/Source/Runtime/D3D12RHI/Public/D3D12RHI.h b/Engine/Source/Runtime/D3D12RHI/Public/D3D12RHI.h
--- a/Engine/Source/Runtime/D3D12RHI/Public/D3D12RHI.h
+++ b/Engine/Source/Runtime/D3D12RHI/Public/D3D12RHI.h
@@ -18,7 +18,7 @@
 
 // This is a value that should be tweaked to fit the app, lower numbers will have better performance
 // Titles using many terrain layers may want to set MAX_SRVS to 64 to avoid shader compilation errors. This will have a small performance hit of around 0.1%
-#define MAX_SRVS		48
+#define MAX_SRVS		64
 #define MAX_SAMPLERS	16
 #define MAX_UAVS		16
 #define MAX_CBS			16
diff --git a/Engine/Source/Runtime/RHI/Public/RHICommandList.h b/Engine/Source/Runtime/RHI/Public/RHICommandList.h
--- a/Engine/Source/Runtime/RHI/Public/RHICommandList.h
+++ b/Engine/Source/Runtime/RHI/Public/RHICommandList.h
@@ -124,7 +124,7 @@ extern RHI_API TAutoConsoleVariable<int32> CVarRHICmdFlushRenderThreadTasks;
 struct FRayTracingShaderBindings
 {
 	FRHITexture* Textures[32] = {};
-	FRHIShaderResourceView* SRVs[48] = {};
+	FRHIShaderResourceView* SRVs[64] = {};
 	FRHIUniformBuffer* UniformBuffers[16] = {};
 	FRHISamplerState* Samplers[16] = {};
 	FRHIUnorderedAccessView* UAVs[16] = {};
