Fixed double delete of GPUTransformsSRV.

diff --git a/Engine/Source/Runtime/D3D12RHI/Private/D3D12RayTracing.cpp b/Engine/Source/Runtime/D3D12RHI/Private/D3D12RayTracing.cpp
--- a/Engine/Source/Runtime/D3D12RHI/Private/D3D12RayTracing.cpp
+++ b/Engine/Source/Runtime/D3D12RHI/Private/D3D12RayTracing.cpp
@@ -2537,8 +2537,6 @@ void FD3D12RayTracingScene::BuildAccelerationStructure(FD3D12CommandContext& Com
 					FRHIShaderResourceView* TransformsSRV = Instance.GPUTransformsSRV.GetReference();
 
 					CopyRayTracingGPUInstances(RHICmdList, NumTransforms, DescOffset, TransformsSRV, InstancesDescUAV);
-
-					Instance.GPUTransformsSRV.GetReference()->Release();					
 				}
 
 				DescOffset += NumTransforms;
