Always indicate -d3ddebug status in LOG. Also, clean out some unused DX11 code fragments left around here.

diff --git a/Engine/Source/Runtime/D3D12RHI/Private/D3D12Adapter.cpp b/Engine/Source/Runtime/D3D12RHI/Private/D3D12Adapter.cpp
--- a/Engine/Source/Runtime/D3D12RHI/Private/D3D12Adapter.cpp
+++ b/Engine/Source/Runtime/D3D12RHI/Private/D3D12Adapter.cpp
@@ -83,12 +83,9 @@ void FD3D12Adapter::CreateRootDevice(bool bWithDebug)
 	DxgiFactory->EnumAdapters(Desc.AdapterIndex, TempAdapter.GetInitReference());
 	VERIFYD3D12RESULT(TempAdapter->QueryInterface(IID_PPV_ARGS(DxgiAdapter.GetInitReference())));
 
-	// In Direct3D 11, if you are trying to create a hardware or a software device, set pAdapter != NULL which constrains the other inputs to be:
-	//		DriverType must be D3D_DRIVER_TYPE_UNKNOWN 
-	//		Software must be NULL. 
-	D3D_DRIVER_TYPE DriverType = D3D_DRIVER_TYPE_UNKNOWN;
-
 #if PLATFORM_WINDOWS || (PLATFORM_HOLOLENS && !UE_BUILD_SHIPPING && D3D12_PROFILING_ENABLED)
+	bool bD3d12gpuvalidation = false;
+
 	if (bWithDebug)
 	{
 		TRefCountPtr<ID3D12Debug> DebugController;
@@ -96,7 +93,6 @@ void FD3D12Adapter::CreateRootDevice(bool bWithDebug)
 		{
 			DebugController->EnableDebugLayer();
 
-			bool bD3d12gpuvalidation = false;
 			if (FParse::Param(FCommandLine::Get(), TEXT("d3d12gpuvalidation")) || FParse::Param(FCommandLine::Get(), TEXT("gpuvalidation")))
 			{
 				TRefCountPtr<ID3D12Debug1> DebugController1;
@@ -104,8 +100,6 @@ void FD3D12Adapter::CreateRootDevice(bool bWithDebug)
 				DebugController1->SetEnableGPUBasedValidation(true);
 				bD3d12gpuvalidation = true;
 			}
-
-			UE_LOG(LogD3D12RHI, Log, TEXT("InitD3DDevice: -D3DDebug = %s -D3D12GPUValidation = %s"), bWithDebug ? TEXT("on") : TEXT("off"), bD3d12gpuvalidation ? TEXT("on") : TEXT("off"));
 		}
 		else
 		{
@@ -113,18 +107,9 @@ void FD3D12Adapter::CreateRootDevice(bool bWithDebug)
 			UE_LOG(LogD3D12RHI, Fatal, TEXT("The debug interface requires the D3D12 SDK Layers. Please install the Graphics Tools for Windows. See: https://docs.microsoft.com/en-us/windows/uwp/gaming/use-the-directx-runtime-and-visual-studio-graphics-diagnostic-features"));
 		}
 	}
-#endif // PLATFORM_WINDOWS || (PLATFORM_HOLOLENS && !UE_BUILD_SHIPPING && D3D12_PROFILING_ENABLED)
 
-#if USE_PIX
-	UE_LOG(LogD3D12RHI, Log, TEXT("Emitting draw events for PIX profiling."));
-	SetEmitDrawEvents(true);
-#endif
-	const bool bIsPerfHUD = !FCString::Stricmp(GetD3DAdapterDesc().Description, TEXT("NVIDIA PerfHUD"));
-
-	if (bIsPerfHUD)
-	{
-		DriverType = D3D_DRIVER_TYPE_REFERENCE;
-	}
+	UE_LOG(LogD3D12RHI, Log, TEXT("InitD3DDevice: -D3DDebug = %s -D3D12GPUValidation = %s"), bWithDebug ? TEXT("on") : TEXT("off"), bD3d12gpuvalidation ? TEXT("on") : TEXT("off"));
+#endif // PLATFORM_WINDOWS || (PLATFORM_HOLOLENS && !UE_BUILD_SHIPPING && D3D12_PROFILING_ENABLED)
 
 	bool bDeviceCreated = false;
 #if !PLATFORM_CPU_ARM_FAMILY && (PLATFORM_WINDOWS || PLATFORM_HOLOLENS)
