Make PSO hashes consistent across sessions so we can catch the specific PSO from a crash log for debugging.

diff --git a/Engine/Source/Runtime/D3D12RHI/Private/D3D12VertexDeclaration.cpp b/Engine/Source/Runtime/D3D12RHI/Private/D3D12VertexDeclaration.cpp
--- a/Engine/Source/Runtime/D3D12RHI/Private/D3D12VertexDeclaration.cpp
+++ b/Engine/Source/Runtime/D3D12RHI/Private/D3D12VertexDeclaration.cpp
@@ -54,7 +54,9 @@ struct FD3D12VertexDeclarationKey
 			case VET_UInt:			D3DElement.Format = DXGI_FORMAT_R32_UINT; break;
 			default: UE_LOG(LogD3D12RHI, Fatal, TEXT("Unknown RHI vertex element type %u"), (uint8)InElements[ElementIndex].Type);
 			};
-			D3DElement.SemanticName = "ATTRIBUTE";
+
+			// We don't assign D3DElement.SemanticName here, because it's a constant string and we don't want to hash pointers. For best debugging experience,
+			// we want to get a consistent hash value across sessions. Therefore it's assigned below, after hashing.
 			D3DElement.SemanticIndex = Element.AttributeIndex;
 			D3DElement.InputSlotClass = Element.bUseInstanceIndex ? D3D12_INPUT_CLASSIFICATION_PER_INSTANCE_DATA : D3D12_INPUT_CLASSIFICATION_PER_VERTEX_DATA;
 
@@ -87,6 +89,12 @@ struct FD3D12VertexDeclarationKey
 		// Hash once.
 		Hash = FCrc::MemCrc_DEPRECATED(VertexElements.GetData(), VertexElements.Num()*sizeof(D3D12_INPUT_ELEMENT_DESC));
 		Hash = FCrc::MemCrc_DEPRECATED(StreamStrides, sizeof(StreamStrides), Hash);
+
+		// Assign all the SemanticName after hashing. It's a constant string, always the same, so no need to hash the data.
+		for (int32 ElementIndex = 0; ElementIndex < VertexElements.Num(); ElementIndex++)
+		{
+			VertexElements[ElementIndex].SemanticName = "ATTRIBUTE";
+		}
 	}
 };
 
diff --git a/Engine/Source/Runtime/D3D12RHI/Private/Windows/WindowsD3D12PipelineState.cpp b/Engine/Source/Runtime/D3D12RHI/Private/Windows/WindowsD3D12PipelineState.cpp
--- a/Engine/Source/Runtime/D3D12RHI/Private/Windows/WindowsD3D12PipelineState.cpp
+++ b/Engine/Source/Runtime/D3D12RHI/Private/Windows/WindowsD3D12PipelineState.cpp
@@ -704,8 +704,9 @@ static inline void FastHashName(wchar_t Name[17], uint64 Hash)
 {
 	for (int32 i = 0; i < 16; i++)
 	{
-		Name[i] = (Hash & 0xF) + 'A';
-		Hash >>= 4;
+		Hash = (Hash << 4) | (Hash >> 60); // rol
+		auto Ch = (Hash & 0xF);
+		Name[i] = Ch + ((Ch < 10)? '0' : 'A' - 10);
 	}
 	Name[16] = 0;
 }
@@ -730,7 +731,7 @@ static void CreatePipelineStateWrapper(ID3D12PipelineState** PSO, FD3D12Adapter*
 		HRESULT hr = CreatePipelineStateFromStream(*PSO, pDevice2, &StreamDesc, static_cast<ID3D12PipelineLibrary1*>(CreationArgs->Library), Name);	// Static cast to ID3D12PipelineLibrary1 since we already checked for ID3D12Device2.
 		if (FAILED(hr))
 		{
-			UE_LOG(LogD3D12RHI, Warning, TEXT("Failed to create Graphics PSO with hash %s"), Name);
+			UE_LOG(LogD3D12RHI, Warning, TEXT("Failed to create Graphics PSO with hash 0x%s"), Name);
 
 			DumpShaderAsm(Stream.VS, TEXT("VS"));
 			DumpShaderAsm(Stream.GS, TEXT("GS"));
@@ -745,7 +746,7 @@ static void CreatePipelineStateWrapper(ID3D12PipelineState** PSO, FD3D12Adapter*
 		HRESULT hr = CreatePipelineState(*PSO, Adapter->GetD3DDevice(), &Desc, CreationArgs->Library, Name);
 		if (FAILED(hr))
 		{
-			UE_LOG(LogD3D12RHI, Warning, TEXT("Failed to create Graphics PSO with hash %s"), Name);
+			UE_LOG(LogD3D12RHI, Warning, TEXT("Failed to create Graphics PSO with hash 0x%s"), Name);
 
 			DumpShaderAsm(Desc.VS, TEXT("VS"));
 			DumpShaderAsm(Desc.GS, TEXT("GS"));
@@ -776,7 +777,7 @@ static void CreatePipelineStateWrapper(ID3D12PipelineState** PSO, FD3D12Adapter*
 		HRESULT hr = CreatePipelineStateFromStream(*PSO, pDevice2, &StreamDesc, static_cast<ID3D12PipelineLibrary1*>(CreationArgs->Library), Name);	// Static cast to ID3D12PipelineLibrary1 since we already checked for ID3D12Device2.
 		if (FAILED(hr))
 		{
-			UE_LOG(LogD3D12RHI, Warning, TEXT("Failed to create Compute PSO with hash %s"), Name);
+			UE_LOG(LogD3D12RHI, Warning, TEXT("Failed to create Compute PSO with hash 0x%s"), Name);
 
 			DumpShaderAsm(Stream.CS, TEXT("CS"));
 		}
@@ -787,7 +788,7 @@ static void CreatePipelineStateWrapper(ID3D12PipelineState** PSO, FD3D12Adapter*
 		HRESULT hr = CreatePipelineState(*PSO, Adapter->GetD3DDevice(), &Desc, CreationArgs->Library, Name);
 		if (FAILED(hr))
 		{
-			UE_LOG(LogD3D12RHI, Warning, TEXT("Failed to create Compute PSO with hash %s"), Name);
+			UE_LOG(LogD3D12RHI, Warning, TEXT("Failed to create Compute PSO with hash 0x%s"), Name);
 
 			DumpShaderAsm(Desc.CS, TEXT("CS"));
 		}
