Build fix with missing include

diff --git a/Engine/Source/Runtime/D3D12RHI/Private/D3D12Util.cpp b/Engine/Source/Runtime/D3D12RHI/Private/D3D12Util.cpp
--- a/Engine/Source/Runtime/D3D12RHI/Private/D3D12Util.cpp
+++ b/Engine/Source/Runtime/D3D12RHI/Private/D3D12Util.cpp
@@ -7,6 +7,8 @@ D3D12Util.h: D3D RHI utility implementation.
 #include "D3D12RHIPrivate.h"
 #include "EngineModule.h"
 #include "RendererInterface.h"
+#include "CoreGlobals.h"
+#include "Misc/OutputDeviceRedirector.h"
 
 #define D3DERR(x) case x: ErrorCodeText = TEXT(#x); break;
 #define LOCTEXT_NAMESPACE "Developer.MessageLog"
