UE4.254 - Compile fix
#rnx

diff --git a/Engine/Source/Runtime/D3D12RHI/Private/D3D12Util.cpp b/Engine/Source/Runtime/D3D12RHI/Private/D3D12Util.cpp
--- a/Engine/Source/Runtime/D3D12RHI/Private/D3D12Util.cpp
+++ b/Engine/Source/Runtime/D3D12RHI/Private/D3D12Util.cpp
@@ -204,6 +204,9 @@ static FString GetD3D12TextureFlagString(uint32 TextureFlags)
 	return TextureFormatText;
 }
 
+
+#if PLATFORM_WINDOWS
+
 static void LogDREDData(ID3D12Device* Device)
 {
 	// Should match all values from D3D12_AUTO_BREADCRUMB_OP
@@ -357,6 +360,8 @@ static void LogDREDData(ID3D12Device* Device)
 	}
 }
 
+#endif  // PLATFORM_WINDOWS
+
 extern CORE_API bool GIsGPUCrashed;
 static void TerminateOnDeviceRemoved(HRESULT D3DResult, ID3D12Device* Device)
 {
@@ -377,10 +382,12 @@ static void TerminateOnDeviceRemoved(HRESULT D3DResult, ID3D12Device* Device)
 		GIsCriticalError = true;
 		GIsGPUCrashed = true;
 
+#if PLATFORM_WINDOWS
 		if (Device)
 		{
 			LogDREDData(Device);
 		}
+#endif  // PLATFORM_WINDOWS
 
 		if (!FApp::IsUnattended())
 		{
