Removed all dynamic SRVs when releasing underlying IB and VB resources

#rnx

diff --git a/Engine/Source/Runtime/D3D12RHI/Private/D3D12IndexBuffer.cpp b/Engine/Source/Runtime/D3D12RHI/Private/D3D12IndexBuffer.cpp
--- a/Engine/Source/Runtime/D3D12RHI/Private/D3D12IndexBuffer.cpp
+++ b/Engine/Source/Runtime/D3D12RHI/Private/D3D12IndexBuffer.cpp
@@ -77,6 +77,8 @@ void FD3D12IndexBuffer::ReleaseUnderlyingResource()
 	UpdateBufferStats<FD3D12IndexBuffer>(&ResourceLocation, false);
 	ResourceLocation.Clear();
 	FRHIIndexBuffer::ReleaseUnderlyingResource();
+	RemoveAllDynamicSRVs();
+
 	FD3D12IndexBuffer* NextIB = GetNextObject();
 	if (NextIB)
 	{
diff --git a/Engine/Source/Runtime/D3D12RHI/Private/D3D12VertexBuffer.cpp b/Engine/Source/Runtime/D3D12RHI/Private/D3D12VertexBuffer.cpp
--- a/Engine/Source/Runtime/D3D12RHI/Private/D3D12VertexBuffer.cpp
+++ b/Engine/Source/Runtime/D3D12RHI/Private/D3D12VertexBuffer.cpp
@@ -103,6 +103,8 @@ void FD3D12VertexBuffer::ReleaseUnderlyingResource()
 	UpdateBufferStats<FD3D12VertexBuffer>(&ResourceLocation, false);
 	ResourceLocation.Clear();
 	FRHIVertexBuffer::ReleaseUnderlyingResource();
+	RemoveAllDynamicSRVs();
+
 	FD3D12VertexBuffer* NextVB = GetNextObject();
 	if (NextVB)
 	{
diff --git a/Engine/Source/Runtime/D3D12RHI/Public/D3D12Resources.h b/Engine/Source/Runtime/D3D12RHI/Public/D3D12Resources.h
--- a/Engine/Source/Runtime/D3D12RHI/Public/D3D12Resources.h
+++ b/Engine/Source/Runtime/D3D12RHI/Public/D3D12Resources.h
@@ -685,11 +685,6 @@ public:
 		check(InSRV->DynamicResource == nullptr);
 		InSRV->DynamicResource = this;
 		DynamicSRVs.Add(InSRV);
-		if (DynamicSRVs.Num() == 4)
-		{
-			static int dbg = 0;
-			dbg++;
-		}
 	}
 
 	void RemoveDynamicSRV(FD3D12BaseShaderResourceView* InSRV)
@@ -701,6 +696,19 @@ public:
 		check(Removed == 1);
 	}
 
+	void RemoveAllDynamicSRVs()
+	{
+		FScopeLock Lock(&DynamicSRVsCS);
+		for (FD3D12BaseShaderResourceView* DynamicSRV : DynamicSRVs)
+		{
+			if (DynamicSRV != nullptr)
+			{
+				DynamicSRV->DynamicResource = nullptr;
+			}
+		}
+		DynamicSRVs.Reset();
+	}
+
 	void Swap(FD3D12BaseShaderResource& Other)
 	{
 		::Swap(Parent, Other.Parent);
