RL - xbox instrumented driver can cause a cmdBuffer corruption under high load of Barreir transitions.

diff --git a/Engine/Source/Runtime/D3D12RHI/Private/D3D12DirectCommandListManager.cpp b/Engine/Source/Runtime/D3D12RHI/Private/D3D12DirectCommandListManager.cpp
--- a/Engine/Source/Runtime/D3D12RHI/Private/D3D12DirectCommandListManager.cpp
+++ b/Engine/Source/Runtime/D3D12RHI/Private/D3D12DirectCommandListManager.cpp
@@ -769,7 +769,25 @@ uint32 FD3D12CommandListManager::GetResourceBarrierCommandList(FD3D12CommandList
 			LogResourceBarriers(BarrierDescs.Num(), BarrierDescs.GetData(), hResourceBarrierList.CommandList());
 #endif
 
-			hResourceBarrierList->ResourceBarrier(BarrierDescs.Num(), BarrierDescs.GetData());
+#if USE_PIX && PLATFORM_XBOXONE 
+			//there was a bug in the instrumented driver that corrupts the cmdBuffer if more than 2000 Barrieres are submitted at once
+			if (BarrierDescs.Num() > 1900)
+			{
+				int Num = BarrierDescs.Num();
+				D3D12_RESOURCE_BARRIER* Ptr = BarrierDescs.GetData();
+				while (Num > 0)
+				{
+					int DispatchNum = FMath::Min(Num, 1900);
+					hResourceBarrierList->ResourceBarrier(DispatchNum, Ptr);
+					Ptr += 1900;
+					Num -= 1900;
+				}
+			}
+			else
+#endif
+			{
+				hResourceBarrierList->ResourceBarrier(BarrierDescs.Num(), BarrierDescs.GetData());
+			}
 		}
 
 		return BarrierDescs.Num();
diff --git a/Engine/Source/Runtime/D3D12RHI/Public/D3D12Resources.h b/Engine/Source/Runtime/D3D12RHI/Public/D3D12Resources.h
--- a/Engine/Source/Runtime/D3D12RHI/Public/D3D12Resources.h
+++ b/Engine/Source/Runtime/D3D12RHI/Public/D3D12Resources.h
@@ -1016,7 +1016,25 @@ public:
 		if (Barriers.Num())
 		{
 			check(pCommandList);
-			pCommandList->ResourceBarrier(Barriers.Num(), Barriers.GetData());
+#if USE_PIX && PLATFORM_XBOXONE 
+			//there was a bug in the instrumented driver that corrupts the cmdBuffer if more than 2000 Barrieres are submitted at once
+			if (Barriers.Num() > 1900)
+			{
+				int Num = Barriers.Num();
+				D3D12_RESOURCE_BARRIER* Ptr = Barriers.GetData();
+				while (Num > 0)
+				{
+					int DispatchNum = FMath::Min(Num, 1900);
+					pCommandList->ResourceBarrier(DispatchNum, Ptr);
+					Ptr += 1900;
+					Num -= 1900;
+				}
+			}
+			else
+#endif
+			{
+				pCommandList->ResourceBarrier(Barriers.Num(), Barriers.GetData());
+			}
 			Reset();
 		}
 	}
