More PSO debug output

diff --git a/Engine/Source/Runtime/D3D12RHI/Private/Windows/WindowsD3D12PipelineState.cpp b/Engine/Source/Runtime/D3D12RHI/Private/Windows/WindowsD3D12PipelineState.cpp
--- a/Engine/Source/Runtime/D3D12RHI/Private/Windows/WindowsD3D12PipelineState.cpp
+++ b/Engine/Source/Runtime/D3D12RHI/Private/Windows/WindowsD3D12PipelineState.cpp
@@ -620,6 +620,86 @@ static void DumpShaderAsm(const D3D12_SHADER_BYTECODE& Shader, const TCHAR* Stag
 #endif
 }
 
+static void DumpPSO(const FD3D12_GRAPHICS_PIPELINE_STATE_DESC& Desc)
+{
+	FString String;
+	String.Appendf(TEXT("AlphaToCoverageEnable = %u\n"),  Desc.BlendState.AlphaToCoverageEnable);
+	String.Appendf(TEXT("IndependentBlendEnable = %u\n"), Desc.BlendState.IndependentBlendEnable);
+
+	uint32 NumBlendRT = Desc.BlendState.IndependentBlendEnable? Desc.RTFormatArray.NumRenderTargets : 1;
+	for (uint32 Index = 0; Index < NumBlendRT; ++Index)
+	{
+		const D3D12_RENDER_TARGET_BLEND_DESC& BlendDesc = Desc.BlendState.RenderTarget[Index];
+
+		String.Appendf(TEXT("RenderTarget[%u] = { %u, %u, 0x%X, 0x%X, 0x%X, 0x%X, 0x%X, 0x%X, 0x%X, 0x%X }\n"),
+			Index,
+			BlendDesc.BlendEnable,
+			BlendDesc.LogicOpEnable,
+			BlendDesc.SrcBlend,
+			BlendDesc.DestBlend,
+			BlendDesc.BlendOp,
+			BlendDesc.SrcBlendAlpha,
+			BlendDesc.DestBlendAlpha,
+			BlendDesc.BlendOpAlpha,
+			BlendDesc.LogicOp,
+			BlendDesc.RenderTargetWriteMask);
+	}
+
+	String.Appendf(TEXT("SampleMask = 0x%X\n"),          Desc.SampleMask);
+
+	const D3D12_RASTERIZER_DESC& RSState = Desc.RasterizerState;
+	String.Appendf(TEXT("FillMode = %u\n"),              RSState.FillMode);
+	String.Appendf(TEXT("CullMode = %u\n"),              RSState.CullMode);
+	String.Appendf(TEXT("FrontCounterClockwise = %u\n"), RSState.FrontCounterClockwise);
+	String.Appendf(TEXT("DepthBias = %d\n"),             RSState.DepthBias);
+	String.Appendf(TEXT("DepthBiasClamp = %f\n"),        RSState.DepthBiasClamp);
+	String.Appendf(TEXT("SlopeScaledDepthBias = %f\n"),  RSState.SlopeScaledDepthBias);
+	String.Appendf(TEXT("DepthClipEnable = %u\n"),       RSState.DepthClipEnable);
+	String.Appendf(TEXT("MultisampleEnable = %u\n"),     RSState.MultisampleEnable);
+	String.Appendf(TEXT("AntialiasedLineEnable = %u\n"), RSState.AntialiasedLineEnable);
+	String.Appendf(TEXT("ForcedSampleCount = %u\n"),     RSState.ForcedSampleCount);
+	String.Appendf(TEXT("ConservativeRaster = %u\n"),    RSState.ConservativeRaster);
+
+	const D3D12_DEPTH_STENCIL_DESC1& DSState = Desc.DepthStencilState;
+	String.Appendf(TEXT("DepthEnable = %u\n"),               DSState.DepthEnable);
+	String.Appendf(TEXT("DepthWriteMask = %u\n"),            DSState.DepthWriteMask);
+	String.Appendf(TEXT("DepthFunc = u\n"),                  DSState.DepthFunc);
+	String.Appendf(TEXT("StencilEnable = %u\n"),             DSState.StencilEnable);
+	String.Appendf(TEXT("StencilReadMask = 0x%X\n"),         DSState.StencilReadMask);
+	String.Appendf(TEXT("StencilWriteMask = 0x%X\n"),        DSState.StencilWriteMask);
+	String.Appendf(TEXT("FrontFace = { %u, %u, %u, %u }\n"), DSState.FrontFace.StencilFailOp, DSState.FrontFace.StencilDepthFailOp, DSState.FrontFace.StencilFailOp, DSState.FrontFace.StencilFunc);
+	String.Appendf(TEXT("BackFace  = { %u, %u, %u, %u }\n"), DSState.BackFace.StencilFailOp,  DSState.BackFace.StencilDepthFailOp,  DSState.BackFace.StencilFailOp,  DSState.BackFace.StencilFunc);
+
+	for (uint32 Index = 0; Index < Desc.InputLayout.NumElements; ++Index)
+	{
+		const D3D12_INPUT_ELEMENT_DESC& ILDesc = Desc.InputLayout.pInputElementDescs[Index];
+
+		String.Appendf(TEXT("InputLayout[%u] = { \"%S\", %u, 0x%X, %u, %u, 0x%X, %u }\n"),
+			Index, ILDesc.SemanticName, ILDesc.SemanticIndex, ILDesc.Format, ILDesc.InputSlot, ILDesc.AlignedByteOffset, ILDesc.InputSlotClass, ILDesc.InstanceDataStepRate);
+	}
+
+	String.Appendf(TEXT("IBStripCutValue = 0x%X\n"), Desc.IBStripCutValue);
+	String.Appendf(TEXT("PrimitiveTopologyType = 0x%X\n"), Desc.PrimitiveTopologyType);
+	String.Appendf(TEXT("NumRenderTargets = %u\n"), Desc.RTFormatArray.NumRenderTargets);
+	for (uint32 Index = 0; Index < Desc.RTFormatArray.NumRenderTargets; ++Index)
+	{
+		String.Appendf(TEXT("RTVFormats[%u] = 0x%X\n"), Index, Desc.RTFormatArray.RTFormats[Index]);
+	}
+	String.Appendf(TEXT("DSVFormat = 0x%X\n"), Desc.DSVFormat);
+	String.Appendf(TEXT("SampleDesc = { %u, %u }\n"), Desc.SampleDesc.Count, Desc.SampleDesc.Quality);
+	String.Appendf(TEXT("NodeMask = 0x%X\n"), Desc.NodeMask);
+	String.Appendf(TEXT("Flags = 0x%X\n"), Desc.Flags);
+
+	UE_LOG(LogD3D12RHI, Display, TEXT("DumpPSO:\n%s"), *String);
+
+	DumpShaderAsm(Desc.VS, TEXT("VS"));
+	DumpShaderAsm(Desc.GS, TEXT("GS"));
+	DumpShaderAsm(Desc.HS, TEXT("HS"));
+	DumpShaderAsm(Desc.DS, TEXT("DS"));
+	DumpShaderAsm(Desc.PS, TEXT("PS"));
+}
+
+
 // Thread-safe create graphics/compute pipeline state. Conditionally load/store the PSO using a Pipeline Library.
 static HRESULT CreatePipelineState(ID3D12PipelineState*&PSO, ID3D12Device* Device, const D3D12_GRAPHICS_PIPELINE_STATE_DESC* Desc, ID3D12PipelineLibrary* Library, const TCHAR* Name)
 {
@@ -733,11 +813,7 @@ static void CreatePipelineStateWrapper(ID3D12PipelineState** PSO, FD3D12Adapter*
 		{
 			UE_LOG(LogD3D12RHI, Warning, TEXT("Failed to create Graphics PSO with hash 0x%s"), Name);
 
-			DumpShaderAsm(Stream.VS, TEXT("VS"));
-			DumpShaderAsm(Stream.GS, TEXT("GS"));
-			DumpShaderAsm(Stream.HS, TEXT("HS"));
-			DumpShaderAsm(Stream.DS, TEXT("DS"));
-			DumpShaderAsm(Stream.PS, TEXT("PS"));
+			DumpPSO(CreationArgs->Desc.Desc);
 		}
 	}
 	else
@@ -748,11 +824,7 @@ static void CreatePipelineStateWrapper(ID3D12PipelineState** PSO, FD3D12Adapter*
 		{
 			UE_LOG(LogD3D12RHI, Warning, TEXT("Failed to create Graphics PSO with hash 0x%s"), Name);
 
-			DumpShaderAsm(Desc.VS, TEXT("VS"));
-			DumpShaderAsm(Desc.GS, TEXT("GS"));
-			DumpShaderAsm(Desc.HS, TEXT("HS"));
-			DumpShaderAsm(Desc.DS, TEXT("DS"));
-			DumpShaderAsm(Desc.PS, TEXT("PS"));
+			DumpPSO(CreationArgs->Desc.Desc);
 		}
 	}
 }
