Fallback to generic full copy of resource when copy rect is the same the the texture size

diff --git a/Engine/Source/Runtime/D3D12RHI/Private/D3D12RenderTarget.cpp b/Engine/Source/Runtime/D3D12RHI/Private/D3D12RenderTarget.cpp
--- a/Engine/Source/Runtime/D3D12RHI/Private/D3D12RenderTarget.cpp
+++ b/Engine/Source/Runtime/D3D12RHI/Private/D3D12RenderTarget.cpp
@@ -270,10 +270,17 @@ void FD3D12CommandContext::RHICopyToResolveTarget(FRHITexture* SourceTextureRHI,
 				}
 				else
 				{
-					if (ResolveParams.Rect.IsValid()
+					D3D12_RESOURCE_DESC const& srcDesc = SourceTexture2D->GetResource()->GetDesc();
+					D3D12_RESOURCE_DESC const& ResolveTargetDesc = DestTexture2D->GetResource()->GetDesc();
+					bool bCopySubRect = ResolveParams.Rect.IsValid() && (ResolveParams.Rect.X1 != 0 || ResolveParams.Rect.Y1 != 0 || ResolveParams.Rect.X2 != srcDesc.Width || ResolveParams.Rect.Y2 != srcDesc.Height);
+
+					if (bCopySubRect
 						&& !SourceTextureRHI->IsMultisampled()
 						&& !DestTexture2D->GetDepthStencilView(FExclusiveDepthStencil::DepthWrite_StencilWrite))
 					{
+						// currently no support for readback buffers
+						check(ResolveTargetDesc.Dimension != D3D12_RESOURCE_DIMENSION_BUFFER);
+
 						D3D12_BOX SrcBox;
 
 						SrcBox.left = ResolveParams.Rect.X1;
@@ -307,11 +314,9 @@ void FD3D12CommandContext::RHICopyToResolveTarget(FRHITexture* SourceTextureRHI,
 						FConditionalScopeResourceBarrier ConditionalScopeResourceBarrierSource(CommandListHandle, SourceTexture2D->GetResource(), D3D12_RESOURCE_STATE_COPY_SOURCE, ResolveParams.SourceArrayIndex);
 
 						// Resolve to a buffer.
-						D3D12_RESOURCE_DESC const& ResolveTargetDesc = DestTexture2D->GetResource()->GetDesc();
 						if (ResolveTargetDesc.Dimension == D3D12_RESOURCE_DIMENSION_BUFFER)
 						{
 							check(IsDefaultContext());
-							D3D12_RESOURCE_DESC const& srcDesc = SourceTexture2D->GetResource()->GetDesc();
 
 							const uint32 BlockBytes = GPixelFormats[SourceTexture2D->GetFormat()].BlockBytes;
 							const uint32 XBytes = (uint32)srcDesc.Width * BlockBytes;
