UE4.25 - Change SRV offset to bytes instead of elements to fix some issues/mismatch
#rnx

diff --git a/Engine/Source/Runtime/Apple/MetalRHI/Private/MetalUAV.cpp b/Engine/Source/Runtime/Apple/MetalRHI/Private/MetalUAV.cpp
--- a/Engine/Source/Runtime/Apple/MetalRHI/Private/MetalUAV.cpp
+++ b/Engine/Source/Runtime/Apple/MetalRHI/Private/MetalUAV.cpp
@@ -312,7 +312,7 @@ FShaderResourceViewRHIRef FMetalDynamicRHI::RHICreateShaderResourceView(FRHIVert
 
 	@autoreleasepool
 	{
-		return RHICreateShaderResourceView(FShaderResourceViewInitializer(VertexBufferRHI, Format));
+		return RHICreateShaderResourceView(FShaderResourceViewInitializer(VertexBufferRHI, EPixelFormat(Format)));
 	}
 }
 
@@ -355,9 +355,9 @@ FShaderResourceViewRHIRef FMetalDynamicRHI::RHICreateShaderResourceView(const FS
 					uint32 Stride = GPixelFormats[Desc.Format].BlockBytes;
 					
 					SRV->Stride = Stride;
-					SRV->Offset = Desc.StartElement * Stride;
+					SRV->Offset = Desc.StartOffsetBytes;
 					
-					FMetalLinearTextureDescriptor LinearTextureDesc(Desc.StartElement, Desc.NumElements, Stride);
+					FMetalLinearTextureDescriptor LinearTextureDesc(Desc.StartOffsetBytes, Desc.NumElements, Stride);
 					SRV->InitLinearTextureDescriptor(LinearTextureDesc);
 					
 					VertexBuffer->CreateLinearTexture((EPixelFormat)Desc.Format, VertexBuffer, &LinearTextureDesc);
@@ -378,7 +378,7 @@ FShaderResourceViewRHIRef FMetalDynamicRHI::RHICreateShaderResourceView(const FS
 				SRV->TextureView = nullptr;
 				SRV->SourceStructuredBuffer = StructuredBuffer;
 				
-				SRV->Offset = Desc.StartElement * StructuredBuffer->GetStride();
+				SRV->Offset = Desc.StartOffsetBytes;
 				SRV->Format = 0;
 				SRV->Stride = StructuredBuffer->GetStride();
 				
@@ -409,10 +409,10 @@ FShaderResourceViewRHIRef FMetalDynamicRHI::RHICreateShaderResourceView(const FS
 					const uint32 Stride = Desc.IndexBuffer->GetStride();
 					check(Stride == ((SRV->Format == PF_R16_UINT) ? 2 : 4));
 					
-					SRV->Offset = Desc.StartElement * Stride;
+					SRV->Offset = Desc.StartOffsetBytes;
 					SRV->Stride = Stride;
 					
-					FMetalLinearTextureDescriptor LinearTextureDesc(Desc.StartElement, Desc.NumElements, Stride);
+					FMetalLinearTextureDescriptor LinearTextureDesc(Desc.StartOffsetBytes, Desc.NumElements, Stride);
 					SRV->InitLinearTextureDescriptor(LinearTextureDesc);
 					
 					IndexBuffer->CreateLinearTexture((EPixelFormat)SRV->Format, IndexBuffer, &LinearTextureDesc);
diff --git a/Engine/Source/Runtime/Apple/MetalRHI/Private/MetalVertexBuffer.cpp b/Engine/Source/Runtime/Apple/MetalRHI/Private/MetalVertexBuffer.cpp
--- a/Engine/Source/Runtime/Apple/MetalRHI/Private/MetalVertexBuffer.cpp
+++ b/Engine/Source/Runtime/Apple/MetalRHI/Private/MetalVertexBuffer.cpp
@@ -304,7 +304,7 @@ FMetalTexture FMetalRHIBuffer::AllocLinearTexture(EPixelFormat InFormat, const F
 		const uint32 MinimumByteAlignment = GetMetalDeviceContext().GetDevice().GetMinimumLinearTextureAlignmentForPixelFormat((mtlpp::PixelFormat)GMetalBufferFormats[InFormat].LinearTextureFormat);
 		const uint32 MinimumElementAlignment = MinimumByteAlignment / BytesPerElement;
 
-		uint32 Offset = LinearTextureDesc.StartElement * BytesPerElement;
+		uint32 Offset = LinearTextureDesc.StartOffsetBytes;
 		check(Offset % MinimumByteAlignment == 0);
 
 		uint32 NumElements = (UINT_MAX == LinearTextureDesc.NumElements) ? ((Size - Offset) / BytesPerElement) : LinearTextureDesc.NumElements;
diff --git a/Engine/Source/Runtime/Apple/MetalRHI/Public/MetalResources.h b/Engine/Source/Runtime/Apple/MetalRHI/Public/MetalResources.h
--- a/Engine/Source/Runtime/Apple/MetalRHI/Public/MetalResources.h
+++ b/Engine/Source/Runtime/Apple/MetalRHI/Public/MetalResources.h
@@ -771,15 +771,15 @@ class FMetalLinearTextureDescriptor
 {
 public:
 	FMetalLinearTextureDescriptor()
-		: StartElement(0)
+		: StartOffsetBytes(0)
 		, NumElements(UINT_MAX)
 		, BytesPerElement(0)
 	{
 		// void
 	}
 
-	FMetalLinearTextureDescriptor(uint32 InStartElement, uint32 InNumElements, uint32 InBytesPerElement)
-		: StartElement(InStartElement)
+	FMetalLinearTextureDescriptor(uint32 InStartOffsetBytes, uint32 InNumElements, uint32 InBytesPerElement)
+		: StartOffsetBytes(InStartOffsetBytes)
 		, NumElements(InNumElements)
 		, BytesPerElement(InBytesPerElement)
 	{
@@ -787,7 +787,7 @@ public:
 	}
 	
 	FMetalLinearTextureDescriptor(const FMetalLinearTextureDescriptor& Other)
-		: StartElement(Other.StartElement)
+		: StartOffsetBytes(Other.StartOffsetBytes)
 		, NumElements(Other.NumElements)
 		, BytesPerElement(Other.BytesPerElement)
 	{
@@ -801,7 +801,7 @@ public:
 
 	friend uint32 GetTypeHash(FMetalLinearTextureDescriptor const& Key)
 	{
-		uint32 Hash = GetTypeHash((uint64)Key.StartElement);
+		uint32 Hash = GetTypeHash((uint64)Key.StartOffsetBytes);
 		Hash = HashCombine(Hash, GetTypeHash((uint64)Key.NumElements));
 		Hash = HashCombine(Hash, GetTypeHash((uint64)Key.BytesPerElement));
 		return Hash;
@@ -809,12 +809,12 @@ public:
 
 	bool operator==(FMetalLinearTextureDescriptor const& Other) const
 	{
-		return    StartElement    == Other.StartElement
-		       && NumElements     == Other.NumElements
-		       && BytesPerElement == Other.BytesPerElement;
+		return    StartOffsetBytes == Other.StartOffsetBytes
+		       && NumElements      == Other.NumElements
+		       && BytesPerElement  == Other.BytesPerElement;
 	}
 
-	uint32 StartElement;
+	uint32 StartOffsetBytes;
 	uint32 NumElements;
 	uint32 BytesPerElement;
 };
diff --git a/Engine/Source/Runtime/D3D12RHI/Private/D3D12RHIPrivate.h b/Engine/Source/Runtime/D3D12RHI/Private/D3D12RHIPrivate.h
--- a/Engine/Source/Runtime/D3D12RHI/Private/D3D12RHIPrivate.h
+++ b/Engine/Source/Runtime/D3D12RHI/Private/D3D12RHIPrivate.h
@@ -301,6 +301,7 @@ public:
 	virtual uint64 RHICalcTexture2DPlatformSize(uint32 SizeX, uint32 SizeY, uint8 Format, uint32 NumMips, uint32 NumSamples, uint32 Flags, const FRHIResourceCreateInfo& CreateInfo, uint32& OutAlign) override;
 	virtual uint64 RHICalcTexture3DPlatformSize(uint32 SizeX, uint32 SizeY, uint32 SizeZ, uint8 Format, uint32 NumMips, uint32 Flags, const FRHIResourceCreateInfo& CreateInfo, uint32& OutAlign) final override;
 	virtual uint64 RHICalcTextureCubePlatformSize(uint32 Size, uint8 Format, uint32 NumMips, uint32 Flags, const FRHIResourceCreateInfo& CreateInfo, uint32& OutAlign) final override;
+	virtual uint64 RHIGetMinimumAlignmentForBufferBackedSRV(EPixelFormat Format) final override;
 	virtual void RHIGetTextureMemoryStats(FTextureMemoryStats& OutStats) final override;
 	virtual bool RHIGetTextureMemoryVisualizeData(FColor* TextureData, int32 SizeX, int32 SizeY, int32 Pitch, int32 PixelSize) final override;
 	virtual FTextureReferenceRHIRef RHICreateTextureReference(FLastRenderTimeContainer* LastRenderTime) final override;
diff --git a/Engine/Source/Runtime/D3D12RHI/Private/D3D12View.cpp b/Engine/Source/Runtime/D3D12RHI/Private/D3D12View.cpp
--- a/Engine/Source/Runtime/D3D12RHI/Private/D3D12View.cpp
+++ b/Engine/Source/Runtime/D3D12RHI/Private/D3D12View.cpp
@@ -6,12 +6,13 @@
 
 #include "D3D12RHIPrivate.h"
 
-static FORCEINLINE D3D12_SHADER_RESOURCE_VIEW_DESC GetVertexBufferSRVDesc(FD3D12VertexBuffer* VertexBuffer, uint32& CreationStride, uint8 Format, uint32 StartElement, uint32 NumElements)
+static FORCEINLINE D3D12_SHADER_RESOURCE_VIEW_DESC GetVertexBufferSRVDesc(FD3D12VertexBuffer* VertexBuffer, uint32& CreationStride, uint8 Format, uint32 StartOffsetBytes, uint32 NumElements)
 {
 	D3D12_SHADER_RESOURCE_VIEW_DESC SRVDesc = {};
 
 	uint32 MaxElements = VertexBuffer->GetSize() / CreationStride;
-	StartElement = FMath::Min(StartElement, MaxElements);
+	StartOffsetBytes = FMath::Min(StartOffsetBytes, VertexBuffer->GetSize());
+	uint32 StartElement = StartOffsetBytes / CreationStride;
 
 	SRVDesc.Shader4ComponentMapping = D3D12_DEFAULT_SHADER_4_COMPONENT_MAPPING;
 	SRVDesc.ViewDimension = D3D12_SRV_DIMENSION_BUFFER;
@@ -31,13 +32,14 @@ static FORCEINLINE D3D12_SHADER_RESOURCE_VIEW_DESC GetVertexBufferSRVDesc(FD3D12
 	return SRVDesc;
 }
 
-static FORCEINLINE D3D12_SHADER_RESOURCE_VIEW_DESC GetIndexBufferSRVDesc(FD3D12IndexBuffer* IndexBuffer, uint32 StartElement, uint32 NumElements)
+static FORCEINLINE D3D12_SHADER_RESOURCE_VIEW_DESC GetIndexBufferSRVDesc(FD3D12IndexBuffer* IndexBuffer, uint32 StartOffsetBytes, uint32 NumElements)
 {
 	const uint32 Usage = IndexBuffer->GetUsage();
 	const uint32 Width = IndexBuffer->GetSize();
 	const uint32 CreationStride = IndexBuffer->GetStride();
 	uint32 MaxElements = Width / CreationStride;
-	StartElement = FMath::Min(StartElement, MaxElements);
+	StartOffsetBytes = FMath::Min(StartOffsetBytes, Width);
+	uint32 StartElement = StartOffsetBytes / CreationStride;
 	const FD3D12ResourceLocation& Location = IndexBuffer->ResourceLocation;
 
 	D3D12_SHADER_RESOURCE_VIEW_DESC SRVDesc = {};
@@ -167,7 +169,12 @@ FShaderResourceViewRHIRef FD3D12DynamicRHI::RHICreateShaderResourceView(FRHIStru
 FShaderResourceViewRHIRef FD3D12DynamicRHI::RHICreateShaderResourceView(FRHIVertexBuffer* VertexBufferRHI, uint32 Stride, uint8 Format)
 {
 	ensureMsgf(Stride == GPixelFormats[Format].BlockBytes, TEXT("provided stride: %i was not consitent with Pixelformat: %s"), Stride, GPixelFormats[Format].Name);
-	return FD3D12DynamicRHI::RHICreateShaderResourceView(FShaderResourceViewInitializer(VertexBufferRHI, Format));
+	return FD3D12DynamicRHI::RHICreateShaderResourceView(FShaderResourceViewInitializer(VertexBufferRHI, EPixelFormat(Format)));
+}
+
+uint64 FD3D12DynamicRHI::RHIGetMinimumAlignmentForBufferBackedSRV(EPixelFormat Format)
+{
+	return GPixelFormats[Format].BlockBytes;
 }
 
 FShaderResourceViewRHIRef FD3D12DynamicRHI::RHICreateShaderResourceView(const FShaderResourceViewInitializer& Initializer)
@@ -183,14 +190,14 @@ FShaderResourceViewRHIRef FD3D12DynamicRHI::RHICreateShaderResourceView(const FS
 				FD3D12VertexBuffer* VertexBuffer;
 				FD3D12ShaderResourceView* SRV;
 
-				uint32 StartElement;
+				uint32 StartOffsetBytes;
 				uint32 NumElements;
 				uint8 Format;
 
-				FD3D12InitializeVertexBufferSRVRHICommand(FD3D12VertexBuffer* InVertexBuffer, FD3D12ShaderResourceView* InSRV, uint32 InStartElement, uint32 InNumElements, uint8 InFormat)
+				FD3D12InitializeVertexBufferSRVRHICommand(FD3D12VertexBuffer* InVertexBuffer, FD3D12ShaderResourceView* InSRV, uint32 InStartOffsetBytes, uint32 InNumElements, uint8 InFormat)
 					: VertexBuffer(InVertexBuffer)
 					, SRV(InSRV)
-					, StartElement(InStartElement)
+					, StartOffsetBytes(InStartOffsetBytes)
 					, NumElements(InNumElements)
 					, Format(InFormat)
 				{}
@@ -198,7 +205,7 @@ FShaderResourceViewRHIRef FD3D12DynamicRHI::RHICreateShaderResourceView(const FS
 				void Execute(FRHICommandListBase& RHICmdList)
 				{
 					uint32 CreationStride = GPixelFormats[Format].BlockBytes;
-					D3D12_SHADER_RESOURCE_VIEW_DESC SRVDesc = GetVertexBufferSRVDesc(VertexBuffer, CreationStride, Format, StartElement, NumElements);
+					D3D12_SHADER_RESOURCE_VIEW_DESC SRVDesc = GetVertexBufferSRVDesc(VertexBuffer, CreationStride, Format, StartOffsetBytes, NumElements);
 
 					// Create a Shader Resource View
 					SRV->Initialize(SRVDesc, VertexBuffer->ResourceLocation, CreationStride);
@@ -225,13 +232,13 @@ FShaderResourceViewRHIRef FD3D12DynamicRHI::RHICreateShaderResourceView(const FS
 					{
 						// We have to defer the SRV initialization to the RHI thread if the buffer is dynamic (and RHI threading is enabled), as dynamic buffers can be renamed.
 						// Also insert an RHI thread fence to prevent parallel translate tasks running until this command has completed.
-						ALLOC_COMMAND_CL(RHICmdList, FD3D12InitializeVertexBufferSRVRHICommand)(VertexBuffer, ShaderResourceView, Desc.StartElement, Desc.NumElements, Desc.Format);
+						ALLOC_COMMAND_CL(RHICmdList, FD3D12InitializeVertexBufferSRVRHICommand)(VertexBuffer, ShaderResourceView, Desc.StartOffsetBytes, Desc.NumElements, Desc.Format);
 						RHICmdList.RHIThreadFence(true);
 					}
 					else
 					{
 						// Run the command directly if we're bypassing RHI command list recording, or the buffer is not dynamic.
-						FD3D12InitializeVertexBufferSRVRHICommand Command(VertexBuffer, ShaderResourceView, Desc.StartElement, Desc.NumElements, Desc.Format);
+						FD3D12InitializeVertexBufferSRVRHICommand Command(VertexBuffer, ShaderResourceView, Desc.StartOffsetBytes, Desc.NumElements, Desc.Format);
 						Command.Execute(RHICmdList);
 					}
 
@@ -247,13 +254,13 @@ FShaderResourceViewRHIRef FD3D12DynamicRHI::RHICreateShaderResourceView(const FS
 			{
 				FD3D12StructuredBuffer* StructuredBuffer;
 				FD3D12ShaderResourceView* SRV;
-				uint32 StartElement;
+				uint32 StartOffsetBytes;
 				uint32 NumElements;
 
-				FD3D12InitializeStructuredBufferSRVRHICommand(FD3D12StructuredBuffer* InStructuredBuffer, FD3D12ShaderResourceView* InSRV, uint32 InStartElement, uint32 InNumElements)
+				FD3D12InitializeStructuredBufferSRVRHICommand(FD3D12StructuredBuffer* InStructuredBuffer, FD3D12ShaderResourceView* InSRV, uint32 InStartOffsetBytes, uint32 InNumElements)
 					: StructuredBuffer(InStructuredBuffer)
 					, SRV(InSRV)
-					, StartElement(InStartElement)
+					, StartOffsetBytes(InStartOffsetBytes)
 					, NumElements(InNumElements)
 				{}
 
@@ -273,7 +280,8 @@ FShaderResourceViewRHIRef FD3D12DynamicRHI::RHICreateShaderResourceView(const FS
 					// BufferDesc.StructureByteStride  is not getting patched through the D3D resource DESC structs, so use the RHI version as a hack
 					uint32 Stride = StructuredBuffer->GetStride();
 					uint32 MaxElements = Location.GetSize() / Stride;
-					StartElement = FMath::Min(StartElement, MaxElements);
+					StartOffsetBytes = FMath::Min<uint32>(StartOffsetBytes, Location.GetSize());
+					uint32 StartElement = StartOffsetBytes / Stride;
 
 					if (bByteAccessBuffer)
 					{
@@ -311,13 +319,13 @@ FShaderResourceViewRHIRef FD3D12DynamicRHI::RHICreateShaderResourceView(const FS
 						{
 							// We have to defer the SRV initialization to the RHI thread if the buffer is dynamic (and RHI threading is enabled), as dynamic buffers can be renamed.
 							// Also insert an RHI thread fence to prevent parallel translate tasks running until this command has completed.
-							ALLOC_COMMAND_CL(RHICmdList, FD3D12InitializeStructuredBufferSRVRHICommand)(StructuredBuffer, ShaderResourceView, Desc.StartElement, Desc.NumElements);
+							ALLOC_COMMAND_CL(RHICmdList, FD3D12InitializeStructuredBufferSRVRHICommand)(StructuredBuffer, ShaderResourceView, Desc.StartOffsetBytes, Desc.NumElements);
 							RHICmdList.RHIThreadFence(true);
 						}
 						else
 						{
 							// Run the command directly if we're bypassing RHI command list recording, or the buffer is not dynamic.
-							FD3D12InitializeStructuredBufferSRVRHICommand Command(StructuredBuffer, ShaderResourceView, Desc.StartElement, Desc.NumElements);
+							FD3D12InitializeStructuredBufferSRVRHICommand Command(StructuredBuffer, ShaderResourceView, Desc.StartOffsetBytes, Desc.NumElements);
 							Command.Execute(RHICmdList);
 						}
 
@@ -342,7 +350,7 @@ FShaderResourceViewRHIRef FD3D12DynamicRHI::RHICreateShaderResourceView(const FS
 
 				FD3D12ResourceLocation& Location = IndexBuffer->ResourceLocation;
 				const uint32 CreationStride = IndexBuffer->GetStride();
-				D3D12_SHADER_RESOURCE_VIEW_DESC SRVDesc = GetIndexBufferSRVDesc(IndexBuffer, Desc.StartElement, Desc.NumElements);
+				D3D12_SHADER_RESOURCE_VIEW_DESC SRVDesc = GetIndexBufferSRVDesc(IndexBuffer, Desc.StartOffsetBytes, Desc.NumElements);
 
 				FD3D12ShaderResourceView* ShaderResourceView = new FD3D12ShaderResourceView(IndexBuffer->GetParentDevice(), SRVDesc, Location, CreationStride);
 				return ShaderResourceView;
diff --git a/Engine/Source/Runtime/OpenGLDrv/Private/OpenGLUAV.cpp b/Engine/Source/Runtime/OpenGLDrv/Private/OpenGLUAV.cpp
--- a/Engine/Source/Runtime/OpenGLDrv/Private/OpenGLUAV.cpp
+++ b/Engine/Source/Runtime/OpenGLDrv/Private/OpenGLUAV.cpp
@@ -9,20 +9,20 @@
 FShaderResourceViewRHIRef FOpenGLDynamicRHI::RHICreateShaderResourceView(FRHIVertexBuffer* VertexBufferRHI, uint32 Stride, uint8 Format)
 {
 	ensureMsgf(Stride == GPixelFormats[Format].BlockBytes, TEXT("provided stride: %i was not consitent with Pixelformat: %s"), Stride, GPixelFormats[Format].Name);
-	return FOpenGLDynamicRHI::RHICreateShaderResourceView(FShaderResourceViewInitializer(VertexBufferRHI, Format));
+	return FOpenGLDynamicRHI::RHICreateShaderResourceView(FShaderResourceViewInitializer(VertexBufferRHI, EPixelFormat(Format)));
 }
 
 // Binds the specified buffer range to a texture resource and selects glTexBuffer or glTexBufferRange 
-static void BindGLTexBufferRange(GLenum Target, GLenum InternalFormat, GLuint Buffer, uint32 StartElement, uint32 NumElements, uint32 Stride)
+static void BindGLTexBufferRange(GLenum Target, GLenum InternalFormat, GLuint Buffer, uint32 StartOffsetBytes, uint32 NumElements, uint32 Stride)
 {
-	if (StartElement == 0 && NumElements == UINT32_MAX)
+	if (StartOffsetBytes == 0 && NumElements == UINT32_MAX)
 	{
 		FOpenGL::TexBuffer(Target, InternalFormat, Buffer);
 	}
 	else
 	{
 		// Validate buffer offset is a multiple of buffer offset alignment
-		GLintptr Offset = StartElement * Stride;
+		GLintptr Offset = StartOffsetBytes;
 		GLsizeiptr Size = NumElements * Stride;
 
 #if DO_CHECK
@@ -63,7 +63,7 @@ FShaderResourceViewRHIRef FOpenGLDynamicRHI::RHICreateShaderResourceView(const F
 					if (VertexBuffer)
 					{
 						OGLRHI->CachedSetupTextureStage(OGLRHI->GetContextStateForCurrentContext(), FOpenGL::GetMaxCombinedTextureImageUnits() - 1, GL_TEXTURE_BUFFER, TextureID, -1, 1);
-						BindGLTexBufferRange(GL_TEXTURE_BUFFER, GLFormat.InternalFormat[0], VertexBuffer->Resource, Desc.StartElement, Desc.NumElements, FormatBPP);
+						BindGLTexBufferRange(GL_TEXTURE_BUFFER, GLFormat.InternalFormat[0], VertexBuffer->Resource, Desc.StartOffsetBytes, Desc.NumElements, FormatBPP);
 					}
 				}
 
@@ -91,7 +91,7 @@ FShaderResourceViewRHIRef FOpenGLDynamicRHI::RHICreateShaderResourceView(const F
 					FOpenGLStructuredBuffer* StructuredBuffer = ResourceCast(StructuredBufferRHI);
 					FOpenGL::GenTextures(1, &TextureID);
 					CachedSetupTextureStage(GetContextStateForCurrentContext(), FOpenGL::GetMaxCombinedTextureImageUnits() - 1, GL_TEXTURE_BUFFER, TextureID, -1, 1);
-					BindGLTexBufferRange(GL_TEXTURE_BUFFER, GL_RGBA32F, StructuredBuffer->Resource, Desc.StartElement, Desc.NumElements, 16);
+					BindGLTexBufferRange(GL_TEXTURE_BUFFER, GL_RGBA32F, StructuredBuffer->Resource, Desc.StartOffsetBytes, Desc.NumElements, 16);
 				}
 
 				return new FOpenGLShaderResourceView(this, TextureID, GL_TEXTURE_BUFFER);
@@ -114,7 +114,7 @@ FShaderResourceViewRHIRef FOpenGLDynamicRHI::RHICreateShaderResourceView(const F
 					CachedSetupTextureStage(GetContextStateForCurrentContext(), FOpenGL::GetMaxCombinedTextureImageUnits() - 1, GL_TEXTURE_BUFFER, TextureID, -1, 1);
 					uint32 Stride = IndexBufferRHI->GetStride();
 					GLenum Format = (Stride == 2) ? GL_R16UI : GL_R32UI;
-					BindGLTexBufferRange(GL_TEXTURE_BUFFER, Format, IndexBuffer->Resource, Desc.StartElement, Desc.NumElements, Stride);
+					BindGLTexBufferRange(GL_TEXTURE_BUFFER, Format, IndexBuffer->Resource, Desc.StartOffsetBytes, Desc.NumElements, Stride);
 				}
 
 				return new FOpenGLShaderResourceView(this, TextureID, GL_TEXTURE_BUFFER);
diff --git a/Engine/Source/Runtime/RHI/Private/DynamicRHI.cpp b/Engine/Source/Runtime/RHI/Private/DynamicRHI.cpp
--- a/Engine/Source/Runtime/RHI/Private/DynamicRHI.cpp
+++ b/Engine/Source/Runtime/RHI/Private/DynamicRHI.cpp
@@ -371,7 +371,7 @@ void FDynamicRHI::RHIUpdateShaderResourceView(FRHIShaderResourceView* SRV, FRHII
 
 uint64 FDynamicRHI::RHIGetMinimumAlignmentForBufferBackedSRV(EPixelFormat Format)
 {
-	return 16; // not-Metal requires a minimum alignment of of 16-bytes.
+	return 1;
 }
 
 uint64 FDynamicRHI::RHICalcVMTexture2DPlatformSize(uint32 Mip0Width, uint32 Mip0Height, uint8 Format, uint32 NumMips, uint32 FirstMipIdx, uint32 NumSamples, uint32 Flags, uint32& OutAlign)
@@ -455,28 +455,30 @@ void FDynamicRHI::RHICheckViewportHDRStatus(FRHIViewport* Viewport)
 }
 
 
-FShaderResourceViewInitializer::FShaderResourceViewInitializer(FRHIVertexBuffer* InVertexBuffer, uint8 InFormat, uint32 InStartElement, uint32 InNumElements)
-	: VertexBufferInitializer({ InVertexBuffer, InStartElement, InNumElements, InFormat }), Type(EType::VertexBufferSRV)
+FShaderResourceViewInitializer::FShaderResourceViewInitializer(FRHIVertexBuffer* InVertexBuffer, EPixelFormat InFormat, uint32 InStartOffsetBytes, uint32 InNumElements)
+	: VertexBufferInitializer({ InVertexBuffer, InStartOffsetBytes, InNumElements, InFormat }), Type(EType::VertexBufferSRV)
 {
+	check(InStartOffsetBytes % RHIGetMinimumAlignmentForBufferBackedSRV(InFormat) == 0);
 	/*if (!VertexBufferInitializer.IsWholeResource())
 	{
 		const uint32 Stride = GPixelFormats[InFormat].BlockBytes;
-		check((VertexBufferInitializer.NumElements + VertexBufferInitializer.StartElement) * Stride <= VertexBufferInitializer.VertexBuffer->GetSize());
+		check((VertexBufferInitializer.NumElements * Stride + VertexBufferInitializer.StartOffsetBytes)  <= VertexBufferInitializer.VertexBuffer->GetSize());
 	}*/
 }
 
-FShaderResourceViewInitializer::FShaderResourceViewInitializer(FRHIVertexBuffer* InVertexBuffer, uint8 InFormat)
+FShaderResourceViewInitializer::FShaderResourceViewInitializer(FRHIVertexBuffer* InVertexBuffer, EPixelFormat InFormat)
 	: VertexBufferInitializer({ InVertexBuffer, 0, UINT32_MAX, InFormat }), Type(EType::VertexBufferSRV) 
 {
 }
 
-FShaderResourceViewInitializer::FShaderResourceViewInitializer(FRHIStructuredBuffer* InStructuredBuffer, uint32 InStartElement, uint32 InNumElements)
-	: StructuredBufferInitializer(FStructuredBufferShaderResourceViewInitializer{ InStructuredBuffer, InStartElement, InNumElements }), Type(EType::StructuredBufferSRV) 
+FShaderResourceViewInitializer::FShaderResourceViewInitializer(FRHIStructuredBuffer* InStructuredBuffer, uint32 InStartOffsetBytes, uint32 InNumElements)
+	: StructuredBufferInitializer(FStructuredBufferShaderResourceViewInitializer{ InStructuredBuffer, InStartOffsetBytes, InNumElements }), Type(EType::StructuredBufferSRV)
 {
+	check(InStartOffsetBytes % InStructuredBuffer->GetStride() == 0);
 	if (!StructuredBufferInitializer.IsWholeResource())
 	{
 		const uint32 Stride = StructuredBufferInitializer.StructuredBuffer->GetStride();
-		check((StructuredBufferInitializer.NumElements + StructuredBufferInitializer.StartElement) * Stride <= StructuredBufferInitializer.StructuredBuffer->GetSize());
+		check((StructuredBufferInitializer.NumElements * Stride + StructuredBufferInitializer.StartOffsetBytes)  <= StructuredBufferInitializer.StructuredBuffer->GetSize());
 	}
 }
 
@@ -485,13 +487,14 @@ FShaderResourceViewInitializer::FShaderResourceViewInitializer(FRHIStructuredBuf
 {
 }
 
-FShaderResourceViewInitializer::FShaderResourceViewInitializer(FRHIIndexBuffer* InIndexBuffer, uint32 InStartElement, uint32 InNumElements)
-	: IndexBufferInitializer(FIndexBufferShaderResourceViewInitializer{ InIndexBuffer, InStartElement, InNumElements }), Type(EType::IndexBufferSRV) 
+FShaderResourceViewInitializer::FShaderResourceViewInitializer(FRHIIndexBuffer* InIndexBuffer, uint32 InStartOffsetBytes, uint32 InNumElements)
+	: IndexBufferInitializer(FIndexBufferShaderResourceViewInitializer{ InIndexBuffer, InStartOffsetBytes, InNumElements }), Type(EType::IndexBufferSRV)
 {
+	check(InStartOffsetBytes % RHIGetMinimumAlignmentForBufferBackedSRV(InIndexBuffer->GetStride() == 2 ? PF_R16_UINT : PF_R32_UINT) == 0);
 	if (!IndexBufferInitializer.IsWholeResource())
 	{
 		const uint32 Stride = IndexBufferInitializer.IndexBuffer->GetStride();
-		check((IndexBufferInitializer.NumElements + IndexBufferInitializer.StartElement) * Stride <= IndexBufferInitializer.IndexBuffer->GetSize());
+		check((IndexBufferInitializer.NumElements * Stride + IndexBufferInitializer.StartOffsetBytes) <= IndexBufferInitializer.IndexBuffer->GetSize());
 	}
 }
 
diff --git a/Engine/Source/Runtime/RHI/Public/DynamicRHI.h b/Engine/Source/Runtime/RHI/Public/DynamicRHI.h
--- a/Engine/Source/Runtime/RHI/Public/DynamicRHI.h
+++ b/Engine/Source/Runtime/RHI/Public/DynamicRHI.h
@@ -219,18 +219,18 @@ struct FShaderResourceViewInitializer
 	struct FVertexBufferShaderResourceViewInitializer
 	{
 		FRHIVertexBuffer* VertexBuffer;
-		uint32 StartElement;
+		uint32 StartOffsetBytes;
 		uint32 NumElements;
-		uint8 Format;
+		EPixelFormat Format;
 
 		inline bool IsWholeResource() const 
 		{ 
-			return StartElement == 0 && NumElements == UINT32_MAX; 
+			return StartOffsetBytes == 0 && NumElements == UINT32_MAX;
 		}
 	};
 
-	RHI_API FShaderResourceViewInitializer(FRHIVertexBuffer* InVertexBuffer, uint8 InFormat, uint32 InStartElement, uint32 InNumElements);
-	RHI_API FShaderResourceViewInitializer(FRHIVertexBuffer* InVertexBuffer, uint8 InFormat);
+	RHI_API FShaderResourceViewInitializer(FRHIVertexBuffer* InVertexBuffer, EPixelFormat InFormat, uint32 InStartOffsetBytes, uint32 InNumElements);
+	RHI_API FShaderResourceViewInitializer(FRHIVertexBuffer* InVertexBuffer, EPixelFormat InFormat);
 
 	const FVertexBufferShaderResourceViewInitializer& AsVertexBufferSRV() const
 	{
@@ -241,16 +241,16 @@ struct FShaderResourceViewInitializer
 	struct FStructuredBufferShaderResourceViewInitializer
 	{
 		FRHIStructuredBuffer* StructuredBuffer;
-		uint32 StartElement;
+		uint32 StartOffsetBytes;
 		uint32 NumElements;
 
 		inline bool IsWholeResource() const 
 		{
-			return StartElement == 0 && NumElements == UINT32_MAX; 
+			return StartOffsetBytes == 0 && NumElements == UINT32_MAX;
 		}
 	};
 
-	RHI_API FShaderResourceViewInitializer(FRHIStructuredBuffer* InStructuredBuffer, uint32 InStartElement, uint32 InNumElements);
+	RHI_API FShaderResourceViewInitializer(FRHIStructuredBuffer* InStructuredBuffer, uint32 InStartOffsetBytes, uint32 InNumElements);
 	RHI_API FShaderResourceViewInitializer(FRHIStructuredBuffer* InStructuredBuffer);
 
 	const FStructuredBufferShaderResourceViewInitializer& AsStructuredBufferSRV() const
@@ -262,16 +262,16 @@ struct FShaderResourceViewInitializer
 	struct FIndexBufferShaderResourceViewInitializer
 	{
 		FRHIIndexBuffer* IndexBuffer;
-		uint32 StartElement;
+		uint32 StartOffsetBytes;
 		uint32 NumElements;
 
 		inline bool IsWholeResource() const 
 		{ 
-			return StartElement == 0 && NumElements == UINT32_MAX; 
+			return StartOffsetBytes == 0 && NumElements == UINT32_MAX;
 		}
 	};
 
-	RHI_API FShaderResourceViewInitializer(FRHIIndexBuffer* InIndexBuffer, uint32 InStartElement, uint32 InNumElements);
+	RHI_API FShaderResourceViewInitializer(FRHIIndexBuffer* InIndexBuffer, uint32 InStartOffsetBytes, uint32 InNumElements);
 	RHI_API FShaderResourceViewInitializer(FRHIIndexBuffer* InIndexBuffer);
 
 	const FIndexBufferShaderResourceViewInitializer& AsIndexBufferSRV() const
diff --git a/Engine/Source/Runtime/RenderCore/Private/DynamicBufferAllocator.cpp b/Engine/Source/Runtime/RenderCore/Private/DynamicBufferAllocator.cpp
--- a/Engine/Source/Runtime/RenderCore/Private/DynamicBufferAllocator.cpp
+++ b/Engine/Source/Runtime/RenderCore/Private/DynamicBufferAllocator.cpp
@@ -98,9 +98,6 @@ FGlobalDynamicReadBuffer::FAllocation FGlobalDynamicReadBuffer::AllocateFloat(ui
 	FScopeLock ScopeLock(&FloatBufferPool->CriticalSection);
 	FAllocation Allocation;
 
-	// The codepath using FShaderResourceViewInitializer, requires the SRV to be aligned on some platforms.
-	Num = Align(Num, RHIGetMinimumAlignmentForBufferBackedSRV(PF_R32_FLOAT) / sizeof(float));
-
 	TotalAllocatedSinceLastCommit += Num;
 	if (IsRenderAlarmLoggingEnabled())
 	{
@@ -108,14 +105,19 @@ FGlobalDynamicReadBuffer::FAllocation FGlobalDynamicReadBuffer::AllocateFloat(ui
 	}
 	uint32 SizeInBytes = sizeof(float) * Num;
 	FDynamicAllocReadBuffer* Buffer = FloatBufferPool->CurrentBuffer;
-	if (Buffer == NULL || Buffer->AllocatedByteCount + SizeInBytes > Buffer->NumBytes)
+
+	uint32 BufferAlignment = RHIGetMinimumAlignmentForBufferBackedSRV(PF_R32_FLOAT);
+	uint32 ByteOffset = Buffer == NULL ? 0 : Align(Buffer->AllocatedByteCount, BufferAlignment);
+
+	if (Buffer == NULL || ByteOffset + SizeInBytes > Buffer->NumBytes)
 	{
 		// Find a buffer in the pool big enough to service the request.
 		Buffer = NULL;
 		for (int32 BufferIndex = 0, NumBuffers = FloatBufferPool->Buffers.Num(); BufferIndex < NumBuffers; ++BufferIndex)
 		{
 			FDynamicAllocReadBuffer& BufferToCheck = FloatBufferPool->Buffers[BufferIndex];
-			if (BufferToCheck.AllocatedByteCount + SizeInBytes <= BufferToCheck.NumBytes)
+			uint32 ByteOffsetToCheck = Align(BufferToCheck.AllocatedByteCount, BufferAlignment);
+			if (ByteOffsetToCheck + SizeInBytes <= BufferToCheck.NumBytes)
 			{
 				Buffer = &BufferToCheck;
 				break;
@@ -141,12 +143,13 @@ FGlobalDynamicReadBuffer::FAllocation FGlobalDynamicReadBuffer::AllocateFloat(ui
 		// Remember this buffer, we'll try to allocate out of it in the future.
 		FloatBufferPool->CurrentBuffer = Buffer;
 	}
+	Buffer->AllocatedByteCount = Align(Buffer->AllocatedByteCount, BufferAlignment);
 
 	check(Buffer != NULL);
 	checkf(Buffer->AllocatedByteCount + SizeInBytes <= Buffer->NumBytes, TEXT("Global dynamic read buffer float buffer allocation failed: BufferSize=%d AllocatedByteCount=%d SizeInBytes=%d"), Buffer->NumBytes, Buffer->AllocatedByteCount, SizeInBytes);
 	Allocation.Buffer = Buffer->MappedBuffer + Buffer->AllocatedByteCount;
 	Allocation.ReadBuffer = Buffer;
-	Buffer->SubAllocations.Emplace(RHICreateShaderResourceView(FShaderResourceViewInitializer(Buffer->Buffer, PF_R32_FLOAT, Buffer->AllocatedByteCount / sizeof(float), Num)));
+	Buffer->SubAllocations.Emplace(RHICreateShaderResourceView(FShaderResourceViewInitializer(Buffer->Buffer, PF_R32_FLOAT, Buffer->AllocatedByteCount, Num)));
 	Allocation.SRV = Buffer->SubAllocations.Last();
 	Buffer->AllocatedByteCount += SizeInBytes;
 
@@ -158,9 +161,6 @@ FGlobalDynamicReadBuffer::FAllocation FGlobalDynamicReadBuffer::AllocateInt32(ui
 	FScopeLock ScopeLock(&Int32BufferPool->CriticalSection);
 	FAllocation Allocation;
 
-	// The codepath using FShaderResourceViewInitializer, requires the SRV to be aligned on 16 bytes on some platforms.
-	Num = Align(Num, RHIGetMinimumAlignmentForBufferBackedSRV(PF_R32_SINT) / sizeof(int32));
-
 	TotalAllocatedSinceLastCommit += Num;
 	if (IsRenderAlarmLoggingEnabled())
 	{
@@ -168,14 +168,19 @@ FGlobalDynamicReadBuffer::FAllocation FGlobalDynamicReadBuffer::AllocateInt32(ui
 	}
 	uint32 SizeInBytes = sizeof(int32) * Num;
 	FDynamicAllocReadBuffer* Buffer = Int32BufferPool->CurrentBuffer;
-	if (Buffer == NULL || Buffer->AllocatedByteCount + SizeInBytes > Buffer->NumBytes)
+
+	uint32 BufferAlignment = RHIGetMinimumAlignmentForBufferBackedSRV(PF_R32_SINT);
+	uint32 ByteOffset = Buffer == NULL ? 0 : Align(Buffer->AllocatedByteCount, BufferAlignment);
+
+	if (Buffer == NULL || ByteOffset + SizeInBytes > Buffer->NumBytes)
 	{
 		// Find a buffer in the pool big enough to service the request.
 		Buffer = NULL;
 		for (int32 BufferIndex = 0, NumBuffers = Int32BufferPool->Buffers.Num(); BufferIndex < NumBuffers; ++BufferIndex)
 		{
 			FDynamicAllocReadBuffer& BufferToCheck = Int32BufferPool->Buffers[BufferIndex];
-			if (BufferToCheck.AllocatedByteCount + SizeInBytes <= BufferToCheck.NumBytes)
+			uint32 ByteOffsetToCheck = Align(BufferToCheck.AllocatedByteCount, BufferAlignment);
+			if (ByteOffsetToCheck + SizeInBytes <= BufferToCheck.NumBytes)
 			{
 				Buffer = &BufferToCheck;
 				break;
@@ -201,12 +206,13 @@ FGlobalDynamicReadBuffer::FAllocation FGlobalDynamicReadBuffer::AllocateInt32(ui
 		// Remember this buffer, we'll try to allocate out of it in the future.
 		Int32BufferPool->CurrentBuffer = Buffer;
 	}
+	Buffer->AllocatedByteCount = Align(Buffer->AllocatedByteCount, BufferAlignment);
 
 	check(Buffer != NULL);
 	checkf(Buffer->AllocatedByteCount + SizeInBytes <= Buffer->NumBytes, TEXT("Global dynamic read buffer int32 buffer allocation failed: BufferSize=%d AllocatedByteCount=%d SizeInBytes=%d"), Buffer->NumBytes, Buffer->AllocatedByteCount, SizeInBytes);
 	Allocation.Buffer = Buffer->MappedBuffer + Buffer->AllocatedByteCount;
 	Allocation.ReadBuffer = Buffer;
-	Buffer->SubAllocations.Emplace(RHICreateShaderResourceView(FShaderResourceViewInitializer(Buffer->Buffer, PF_R32_SINT, Buffer->AllocatedByteCount / sizeof(int), Num)));
+	Buffer->SubAllocations.Emplace(RHICreateShaderResourceView(FShaderResourceViewInitializer(Buffer->Buffer, PF_R32_SINT, Buffer->AllocatedByteCount, Num)));
 	Allocation.SRV = Buffer->SubAllocations.Last();
 	Buffer->AllocatedByteCount += SizeInBytes;
 
diff --git a/Engine/Source/Runtime/VulkanRHI/Private/VulkanUAV.cpp b/Engine/Source/Runtime/VulkanRHI/Private/VulkanUAV.cpp
--- a/Engine/Source/Runtime/VulkanRHI/Private/VulkanUAV.cpp
+++ b/Engine/Source/Runtime/VulkanRHI/Private/VulkanUAV.cpp
@@ -314,13 +314,13 @@ FShaderResourceViewRHIRef FVulkanDynamicRHI::RHICreateShaderResourceView(const F
 			const FShaderResourceViewInitializer::FVertexBufferShaderResourceViewInitializer Desc = Initializer.AsVertexBufferSRV();
 			const uint32 Stride = GPixelFormats[Desc.Format].BlockBytes;
 			FVulkanVertexBuffer* VertexBuffer = ResourceCast(Desc.VertexBuffer);
-			return new FVulkanShaderResourceView(Device, Desc.VertexBuffer, VertexBuffer, Desc.NumElements * Stride, (EPixelFormat)Desc.Format, Desc.StartElement * Stride);
+			return new FVulkanShaderResourceView(Device, Desc.VertexBuffer, VertexBuffer, Desc.NumElements * Stride, (EPixelFormat)Desc.Format, Desc.StartOffsetBytes);
 		}
 		case FShaderResourceViewInitializer::EType::StructuredBufferSRV:
 		{
 			const FShaderResourceViewInitializer::FStructuredBufferShaderResourceViewInitializer Desc = Initializer.AsStructuredBufferSRV();
 			FVulkanStructuredBuffer* StructuredBuffer = ResourceCast(Desc.StructuredBuffer);
-			return new FVulkanShaderResourceView(Device, StructuredBuffer, StructuredBuffer->GetStride() * Desc.StartElement);
+			return new FVulkanShaderResourceView(Device, StructuredBuffer, Desc.StartOffsetBytes);
 		}			
 		case FShaderResourceViewInitializer::EType::IndexBufferSRV:
 		{
@@ -329,7 +329,7 @@ FShaderResourceViewRHIRef FVulkanDynamicRHI::RHICreateShaderResourceView(const F
 			const uint32 Stride = Desc.IndexBuffer->GetStride();
 			check(Stride == 2 || Stride == 4);
 			EPixelFormat Format = (Stride == 4) ? PF_R32_UINT : PF_R16_UINT;
-			return new FVulkanShaderResourceView(Device, Desc.IndexBuffer, IndexBuffer, Desc.NumElements * Stride, Format, Desc.StartElement * Stride);
+			return new FVulkanShaderResourceView(Device, Desc.IndexBuffer, IndexBuffer, Desc.NumElements * Stride, Format, Desc.StartOffsetBytes);
 		}
 	}
 	checkNoEntry();
diff --git a/Engine/Source/Runtime/Windows/D3D11RHI/Private/D3D11RHIPrivate.h b/Engine/Source/Runtime/Windows/D3D11RHI/Private/D3D11RHIPrivate.h
--- a/Engine/Source/Runtime/Windows/D3D11RHI/Private/D3D11RHIPrivate.h
+++ b/Engine/Source/Runtime/Windows/D3D11RHI/Private/D3D11RHIPrivate.h
@@ -455,6 +455,7 @@ public:
 	virtual uint64 RHICalcTexture2DPlatformSize(uint32 SizeX, uint32 SizeY, uint8 Format, uint32 NumMips, uint32 NumSamples, uint32 Flags, const FRHIResourceCreateInfo& CreateInfo, uint32& OutAlign) final override;
 	virtual uint64 RHICalcTexture3DPlatformSize(uint32 SizeX, uint32 SizeY, uint32 SizeZ, uint8 Format, uint32 NumMips, uint32 Flags, const FRHIResourceCreateInfo& CreateInfo, uint32& OutAlign) final override;
 	virtual uint64 RHICalcTextureCubePlatformSize(uint32 Size, uint8 Format, uint32 NumMips, uint32 Flags, const FRHIResourceCreateInfo& CreateInfo, uint32& OutAlign) final override;
+	virtual uint64 RHIGetMinimumAlignmentForBufferBackedSRV(EPixelFormat Format) final override;
 	virtual void RHIGetTextureMemoryStats(FTextureMemoryStats& OutStats) final override;
 	virtual bool RHIGetTextureMemoryVisualizeData(FColor* TextureData,int32 SizeX,int32 SizeY,int32 Pitch,int32 PixelSize) final override;
 	virtual FTextureReferenceRHIRef RHICreateTextureReference(FLastRenderTimeContainer* LastRenderTime) final override;
diff --git a/Engine/Source/Runtime/Windows/D3D11RHI/Private/D3D11UAV.cpp b/Engine/Source/Runtime/Windows/D3D11RHI/Private/D3D11UAV.cpp
--- a/Engine/Source/Runtime/Windows/D3D11RHI/Private/D3D11UAV.cpp
+++ b/Engine/Source/Runtime/Windows/D3D11RHI/Private/D3D11UAV.cpp
@@ -200,7 +200,12 @@ FShaderResourceViewRHIRef FD3D11DynamicRHI::RHICreateShaderResourceView_RenderTh
 	return RHICreateShaderResourceView(StructuredBuffer);
 }
 
-static void CreateD3D11ShaderResourceViewOnBuffer(ID3D11Device* Direct3DDevice, ID3D11Buffer* Buffer, uint32 StartElement, uint32 NumElements, uint8 Format, ID3D11ShaderResourceView** OutSRV)
+uint64 FD3D11DynamicRHI::RHIGetMinimumAlignmentForBufferBackedSRV(EPixelFormat Format)
+{
+	return GPixelFormats[Format].BlockBytes;
+}
+
+static void CreateD3D11ShaderResourceViewOnBuffer(ID3D11Device* Direct3DDevice, ID3D11Buffer* Buffer, uint32 StartOffsetBytes, uint32 NumElements, uint8 Format, ID3D11ShaderResourceView** OutSRV)
 {
 	D3D11_BUFFER_DESC BufferDesc;
 	Buffer->GetDesc(&BufferDesc);
@@ -212,7 +217,8 @@ static void CreateD3D11ShaderResourceViewOnBuffer(ID3D11Device* Direct3DDevice,
 
 	uint32 Stride = GPixelFormats[Format].BlockBytes;
 	uint32 MaxElements = BufferDesc.ByteWidth / Stride;
-	StartElement = FMath::Min(StartElement, MaxElements);
+	StartOffsetBytes = FMath::Min(StartOffsetBytes, BufferDesc.ByteWidth);
+	uint32 StartElement = StartOffsetBytes / Stride;
 
 	SRVDesc.Buffer.FirstElement = StartElement;
 	SRVDesc.Format = FindShaderResourceDXGIFormat((DXGI_FORMAT)GPixelFormats[Format].PlatformFormat, false);
@@ -237,7 +243,7 @@ static void CreateD3D11ShaderResourceViewOnBuffer(ID3D11Device* Direct3DDevice,
 FShaderResourceViewRHIRef FD3D11DynamicRHI::RHICreateShaderResourceView(FRHIVertexBuffer* VertexBufferRHI, uint32 Stride, uint8 Format)
 {
 	ensureMsgf(Stride == GPixelFormats[Format].BlockBytes, TEXT("provided stride: %i was not consitent with Pixelformat: %s"), Stride, GPixelFormats[Format].Name);
-	return FD3D11DynamicRHI::RHICreateShaderResourceView(FShaderResourceViewInitializer(VertexBufferRHI, Format));
+	return FD3D11DynamicRHI::RHICreateShaderResourceView(FShaderResourceViewInitializer(VertexBufferRHI, EPixelFormat(Format)));
 }
 
 FShaderResourceViewRHIRef FD3D11DynamicRHI::RHICreateShaderResourceView(const FShaderResourceViewInitializer& Initializer)
@@ -255,7 +261,7 @@ FShaderResourceViewRHIRef FD3D11DynamicRHI::RHICreateShaderResourceView(const FS
 			}
 
 			TRefCountPtr<ID3D11ShaderResourceView> ShaderResourceView;
-			CreateD3D11ShaderResourceViewOnBuffer(Direct3DDevice, VertexBuffer->Resource, Desc.StartElement, Desc.NumElements, Desc.Format, ShaderResourceView.GetInitReference());
+			CreateD3D11ShaderResourceViewOnBuffer(Direct3DDevice, VertexBuffer->Resource, Desc.StartOffsetBytes, Desc.NumElements, Desc.Format, ShaderResourceView.GetInitReference());
 
 			return new FD3D11ShaderResourceView(ShaderResourceView, VertexBuffer);
 		}
@@ -277,19 +283,23 @@ FShaderResourceViewRHIRef FD3D11DynamicRHI::RHICreateShaderResourceView(const FS
 			if (bByteAccessBuffer)
 			{
 				uint32 MaxElements = BufferDesc.ByteWidth / 4;
-				uint32 StartElement = FMath::Min(Desc.StartElement, MaxElements);
+				uint32 StartOffsetBytes = FMath::Min(Desc.StartOffsetBytes, BufferDesc.ByteWidth);
+				uint32 StartElement = StartOffsetBytes / 4;
+
 				SRVDesc.ViewDimension = D3D11_SRV_DIMENSION_BUFFEREX;
 				SRVDesc.BufferEx.NumElements = FMath::Min(Desc.NumElements, MaxElements - StartElement);
-				SRVDesc.BufferEx.FirstElement = Desc.StartElement;
+				SRVDesc.BufferEx.FirstElement = StartElement;
 				SRVDesc.BufferEx.Flags = D3D11_BUFFEREX_SRV_FLAG_RAW;
 				SRVDesc.Format = DXGI_FORMAT_R32_TYPELESS;
 			}
 			else
 			{
 				uint32 MaxElements = BufferDesc.ByteWidth / BufferDesc.StructureByteStride;
-				uint32 StartElement = FMath::Min(Desc.StartElement, BufferDesc.ByteWidth / BufferDesc.StructureByteStride);
+				uint32 StartOffsetBytes = FMath::Min(Desc.StartOffsetBytes, BufferDesc.ByteWidth);
+				uint32 StartElement = StartOffsetBytes / BufferDesc.StructureByteStride;
+
 				SRVDesc.ViewDimension = D3D11_SRV_DIMENSION_BUFFER;
-				SRVDesc.Buffer.FirstElement = Desc.StartElement;
+				SRVDesc.Buffer.FirstElement = StartElement;
 				SRVDesc.Buffer.NumElements = FMath::Min(Desc.NumElements, MaxElements - StartElement);
 				SRVDesc.Format = DXGI_FORMAT_UNKNOWN;
 			}
@@ -317,7 +327,7 @@ FShaderResourceViewRHIRef FD3D11DynamicRHI::RHICreateShaderResourceView(const FS
 			check(Stride == 2 || Stride == 4);
 			const EPixelFormat Format = (Stride == 2) ? PF_R16_UINT : PF_R32_UINT;
 			TRefCountPtr<ID3D11ShaderResourceView> ShaderResourceView;
-			CreateD3D11ShaderResourceViewOnBuffer(Direct3DDevice, Buffer->Resource, Desc.StartElement, Desc.NumElements, Format, ShaderResourceView.GetInitReference());
+			CreateD3D11ShaderResourceViewOnBuffer(Direct3DDevice, Buffer->Resource, Desc.StartOffsetBytes, Desc.NumElements, Format, ShaderResourceView.GetInitReference());
 
 			return new FD3D11ShaderResourceView(ShaderResourceView, Buffer);
 		}
