Fix computation of miss shader table size

diff --git a/Engine/Source/Runtime/D3D12RHI/Private/D3D12RayTracing.cpp b/Engine/Source/Runtime/D3D12RHI/Private/D3D12RayTracing.cpp
--- a/Engine/Source/Runtime/D3D12RHI/Private/D3D12RayTracing.cpp
+++ b/Engine/Source/Runtime/D3D12RHI/Private/D3D12RayTracing.cpp
@@ -1281,7 +1281,7 @@ public:
 
 		Desc.MissShaderTable.StartAddress = ShaderTableAddress + MissShaderTableOffset;
 		Desc.MissShaderTable.StrideInBytes = LocalRecordStride;
-		Desc.MissShaderTable.SizeInBytes = LocalRecordStride * NumMissShaders;
+		Desc.MissShaderTable.SizeInBytes = LocalRecordStride * NumMissRecords;
 
 
 		if (NumCallableRecords)
