Fixed D3D12 crashes when running in virtual mGPU mode.

diff --git a/Engine/Source/Runtime/D3D12RHI/Private/D3D12RenderTarget.cpp b/Engine/Source/Runtime/D3D12RHI/Private/D3D12RenderTarget.cpp
--- a/Engine/Source/Runtime/D3D12RHI/Private/D3D12RenderTarget.cpp
+++ b/Engine/Source/Runtime/D3D12RHI/Private/D3D12RenderTarget.cpp
@@ -13,6 +13,7 @@
 #include "PipelineStateCache.h"
 #include "Math/PackedVector.h"
 #include "RHISurfaceDataConversion.h"
+#include "CommonRenderResources.h"
 
 static inline DXGI_FORMAT ConvertTypelessToUnorm(DXGI_FORMAT Format)
 {
@@ -159,6 +160,7 @@ void FD3D12CommandContext::ResolveTextureUsingShader(
 
 	GraphicsPSOInit.BoundShaderState.VertexShaderRHI = ResolveVertexShader.GetVertexShader();
 	GraphicsPSOInit.BoundShaderState.PixelShaderRHI = ResolvePixelShader.GetPixelShader();
+	GraphicsPSOInit.BoundShaderState.VertexDeclarationRHI = GFilterVertexDeclaration.VertexDeclarationRHI;
 	GraphicsPSOInit.PrimitiveType = PT_TriangleStrip;
 
 	SetGraphicsPipelineState(RHICmdList, GraphicsPSOInit, EApplyRendertargetOption::DoNothing);
diff --git a/Engine/Source/Runtime/D3D12RHI/Private/D3D12UAV.cpp b/Engine/Source/Runtime/D3D12RHI/Private/D3D12UAV.cpp
--- a/Engine/Source/Runtime/D3D12RHI/Private/D3D12UAV.cpp
+++ b/Engine/Source/Runtime/D3D12RHI/Private/D3D12UAV.cpp
@@ -336,13 +336,13 @@ void FD3D12CommandContext::ClearUAV(TRHICommandList_RecursiveHazardous<FD3D12Com
 
 void FD3D12CommandContext::RHIClearUAVFloat(FRHIUnorderedAccessView* UnorderedAccessViewRHI, const FVector4& Values)
 {
-	TRHICommandList_RecursiveHazardous<FD3D12CommandContext> RHICmdList(this);
+	TRHICommandList_RecursiveHazardous<FD3D12CommandContext> RHICmdList(this, GetGPUMask());
 	ClearUAV(RHICmdList, ResourceCast(UnorderedAccessViewRHI), &Values, true);
 }
 
 void FD3D12CommandContext::RHIClearUAVUint(FRHIUnorderedAccessView* UnorderedAccessViewRHI, const FUintVector4& Values)
 {
-	TRHICommandList_RecursiveHazardous<FD3D12CommandContext> RHICmdList(this);
+	TRHICommandList_RecursiveHazardous<FD3D12CommandContext> RHICmdList(this, GetGPUMask());
 	ClearUAV(RHICmdList, ResourceCast(UnorderedAccessViewRHI), &Values, false);
 }
 
