Fix bug in D3D12 RHICopyTexture where zero size box isn't interpreted as full resource when copying to staging textures.
Cherry pick integration from //UE4/Dev-Rendering/...

diff --git a/Engine/Source/Runtime/D3D12RHI/Private/D3D12Texture.cpp b/Engine/Source/Runtime/D3D12RHI/Private/D3D12Texture.cpp
--- a/Engine/Source/Runtime/D3D12RHI/Private/D3D12Texture.cpp
+++ b/Engine/Source/Runtime/D3D12RHI/Private/D3D12Texture.cpp
@@ -3017,14 +3017,17 @@ void FD3D12CommandContext::RHICopyTexture(FRHITexture* SourceTextureRHI, FRHITex
 
 	if (CopyInfo.Size != FIntVector::ZeroValue || bReadback)
 	{
+		// Interpret zero size as source size
+		const FIntVector CopySize = CopyInfo.Size == FIntVector::ZeroValue ? SourceTextureRHI->GetSizeXYZ() : CopyInfo.Size;
+
 		// Copy sub texture regions
 		CD3DX12_BOX SourceBoxD3D(
 			CopyInfo.SourcePosition.X,
 			CopyInfo.SourcePosition.Y,
 			CopyInfo.SourcePosition.Z,
-			CopyInfo.SourcePosition.X + CopyInfo.Size.X,
-			CopyInfo.SourcePosition.Y + CopyInfo.Size.Y,
-			CopyInfo.SourcePosition.Z + CopyInfo.Size.Z
+			CopyInfo.SourcePosition.X + CopySize.X,
+			CopyInfo.SourcePosition.Y + CopySize.Y,
+			CopyInfo.SourcePosition.Z + CopySize.Z
 		);
 
 		D3D12_TEXTURE_COPY_LOCATION Src;
@@ -3055,9 +3058,9 @@ void FD3D12CommandContext::RHICopyTexture(FRHITexture* SourceTextureRHI, FRHITex
 				uint32 SourceMipIndex = CopyInfo.SourceMipIndex + MipIndex;
 				uint32 DestMipIndex   = CopyInfo.DestMipIndex   + MipIndex;
 
-				uint32 SizeX = FMath::Max(CopyInfo.Size.X >> MipIndex, 1);
-				uint32 SizeY = FMath::Max(CopyInfo.Size.Y >> MipIndex, 1);
-				uint32 SizeZ = FMath::Max(CopyInfo.Size.Z >> MipIndex, 1);
+				uint32 SizeX = FMath::Max(CopySize.X >> MipIndex, 1);
+				uint32 SizeY = FMath::Max(CopySize.Y >> MipIndex, 1);
+				uint32 SizeZ = FMath::Max(CopySize.Z >> MipIndex, 1);
 
 				SourceBoxD3D.right  = CopyInfo.SourcePosition.X + SizeX;
 				SourceBoxD3D.bottom = CopyInfo.SourcePosition.Y + SizeY;
