Fix initial resource barrier state for textures created via InitialData
Cherry pick integration from //UE4/Dev-Rendering/...

diff --git a/Engine/Source/Runtime/D3D12RHI/Private/D3D12Texture.cpp b/Engine/Source/Runtime/D3D12RHI/Private/D3D12Texture.cpp
--- a/Engine/Source/Runtime/D3D12RHI/Private/D3D12Texture.cpp
+++ b/Engine/Source/Runtime/D3D12RHI/Private/D3D12Texture.cpp
@@ -205,9 +205,9 @@ struct FD3D12RHICommandInitializeTexture final : public FRHICommand<FD3D12RHICom
 
 			FD3D12CommandListHandle& hCommandList = Device->GetDefaultCommandContext().CommandListHandle;
 			hCommandList.GetCurrentOwningContext()->numCopies += NumSubresources;
-
-			FConditionalScopeResourceBarrier ConditionalScopeResourceBarrier(hCommandList, Resource, D3D12_RESOURCE_STATE_COPY_DEST, D3D12_RESOURCE_BARRIER_ALL_SUBRESOURCES);
-
+			
+			// resource should be in copy dest already, because it's created like that, so no transition required here
+			
 			ID3D12GraphicsCommandList* CmdList = hCommandList.GraphicsCommandList();
 
 			D3D12_TEXTURE_COPY_LOCATION Dst;
@@ -221,7 +221,18 @@ struct FD3D12RHICommandInitializeTexture final : public FRHICommand<FD3D12RHICom
 				CmdList->CopyTextureRegion(&Dst, 0, 0, 0, &Src, nullptr);
 			}
 
-			hCommandList.UpdateResidency(Resource);
+			// Update the resource state after the copy has been done (will take care of updating the residency as well)
+			if (Resource->RequiresResourceStateTracking())
+			{
+				// record the dummy copy_dest to copy_dest transition in the command list to make sure we have proper tracking of the resource
+				// mostly needed to make sure we have correct storage of end resource state when all the pending buffer transitions are flushed 
+				// (otherwise it's untracked and not updated)
+				FD3D12DynamicRHI::TransitionResource(hCommandList, Resource, D3D12_RESOURCE_STATE_COPY_DEST, D3D12_RESOURCE_BARRIER_ALL_SUBRESOURCES);
+			}
+			else
+			{
+				hCommandList.AddTransitionBarrier(Resource, D3D12_RESOURCE_STATE_COPY_DEST, Resource->GetDefaultResourceState(), D3D12_RESOURCE_BARRIER_ALL_SUBRESOURCES);
+			}
 
 			CurrentTexture = CurrentTexture->GetNextObject();
 		}
