Fix the case where we do automatic resource transition based on bindings, but bind to several slots.

diff --git a/Engine/Source/Runtime/D3D12RHI/Private/D3D12RHIPrivate.h b/Engine/Source/Runtime/D3D12RHI/Private/D3D12RHIPrivate.h
--- a/Engine/Source/Runtime/D3D12RHI/Private/D3D12RHIPrivate.h
+++ b/Engine/Source/Runtime/D3D12RHI/Private/D3D12RHIPrivate.h
@@ -542,23 +542,34 @@ public:
 		return !pLeftView->DoesNotOverlap(*pRightView);
 	}
 
-	static inline bool IsTransitionNeeded(uint32 before, uint32 after)
+	static inline bool IsTransitionNeeded(D3D12_RESOURCE_STATES Before, D3D12_RESOURCE_STATES& After)
 	{
-		check(before != D3D12_RESOURCE_STATE_CORRUPT && after != D3D12_RESOURCE_STATE_CORRUPT);
-		check(before != D3D12_RESOURCE_STATE_TBD && after != D3D12_RESOURCE_STATE_TBD);
+		check(Before != D3D12_RESOURCE_STATE_CORRUPT && After != D3D12_RESOURCE_STATE_CORRUPT);
+		check(Before != D3D12_RESOURCE_STATE_TBD && After != D3D12_RESOURCE_STATE_TBD);
 
 		// Depth write is actually a suitable for read operations as a "normal" depth buffer.
-		if ((before == D3D12_RESOURCE_STATE_DEPTH_WRITE) && (after == D3D12_RESOURCE_STATE_DEPTH_READ))
+		if ((Before == D3D12_RESOURCE_STATE_DEPTH_WRITE) && (After == D3D12_RESOURCE_STATE_DEPTH_READ))
 		{
 			return false;
 		}
 
-		// If 'after' is a subset of 'before', then there's no need for a transition
-		// Note: COMMON is an oddball state that doesn't follow the RESOURE_STATE pattern of 
+		// COMMON is an oddball state that doesn't follow the RESOURE_STATE pattern of 
 		// having exactly one bit set so we need to special case these
-		return before != after &&
-			((before | after) != before ||
-				after == D3D12_RESOURCE_STATE_COMMON);
+		if (After == D3D12_RESOURCE_STATE_COMMON)
+		{
+			return true;
+		}
+
+		// We should avoid doing read-to-read state transitions. But when we do, we should avoid turning off already transitioned bits,
+		// e.g. VERTEX_BUFFER -> SHADER_RESOURCE is turned into VERTEX_BUFFER -> VERTEX_BUFFER | SHADER_RESOURCE.
+		// This reduces the number of resource transitions and ensures automatic states from resource bindings get properly combined.
+		D3D12_RESOURCE_STATES Combined = Before | After;
+		if ((Combined & (D3D12_RESOURCE_STATE_GENERIC_READ | D3D12_RESOURCE_STATE_INDIRECT_ARGUMENT)) == Combined)
+		{
+			After = Combined;
+		}
+
+		return Before != After;
 	}
 
 	/** Transition a resource's state based on a Render target view */
