Panic flush log when GPU crash has happened to make sure we have all the needed data in the log before force exit is called

diff --git a/Engine/Source/Runtime/D3D12RHI/Private/D3D12Util.cpp b/Engine/Source/Runtime/D3D12RHI/Private/D3D12Util.cpp
--- a/Engine/Source/Runtime/D3D12RHI/Private/D3D12Util.cpp
+++ b/Engine/Source/Runtime/D3D12RHI/Private/D3D12Util.cpp
@@ -355,8 +355,6 @@ static void LogDREDData(ID3D12Device* Device)
 				}
 			}
 		}
-
-		GLog->Flush();
 	}
 }
 
@@ -388,6 +386,10 @@ static void TerminateOnDeviceRemoved(HRESULT D3DResult, ID3D12Device* Device)
 			LogDREDData(Device);
 		}
 #endif  // PLATFORM_WINDOWS
+		
+		// Make sure the log is flushed!
+		GLog->PanicFlushThreadedLogs();
+		GLog->Flush();
 
 		if (!FApp::IsUnattended())
 		{
