Fix ensure() in DX12 when trying to transition a resource that was already in the COMMON state to readable.

IsTransitionNeeded() wasn't handling the COMMON -> COMMON case correctly.

#rnx

diff --git a/Engine/Source/Runtime/D3D12RHI/Private/D3D12RHIPrivate.h b/Engine/Source/Runtime/D3D12RHI/Private/D3D12RHIPrivate.h
--- a/Engine/Source/Runtime/D3D12RHI/Private/D3D12RHIPrivate.h
+++ b/Engine/Source/Runtime/D3D12RHI/Private/D3D12RHIPrivate.h
@@ -557,7 +557,11 @@ public:
 		// having exactly one bit set so we need to special case these
 		if (After == D3D12_RESOURCE_STATE_COMMON)
 		{
-			return true;
+			// The resource state tracking code in FD3D12CommandContext::RHITransitionResources forces all EReadable transitions
+			// to go through the COMMON state right now, so we can end up with some COMMON -> COMMON transitions which can be
+			// skipped. Once that is fixed or removed, we shouldn't get here anymore if we're already in the COMMON state,
+			// so we can simply return true and let the ensure in FD3D12CommandListHandle::AddTransitionBarrier catch bad usage.
+			return (Before != D3D12_RESOURCE_STATE_COMMON);
 		}
 
 		// We should avoid doing read-to-read state transitions. But when we do, we should avoid turning off already transitioned bits,
