Log swap chain desc when failed to create swapchain to get more information about failure

diff --git a/Engine/Source/Runtime/D3D12RHI/Private/Windows/WindowsD3D12Viewport.cpp b/Engine/Source/Runtime/D3D12RHI/Private/Windows/WindowsD3D12Viewport.cpp
--- a/Engine/Source/Runtime/D3D12RHI/Private/Windows/WindowsD3D12Viewport.cpp
+++ b/Engine/Source/Runtime/D3D12RHI/Private/Windows/WindowsD3D12Viewport.cpp
@@ -209,7 +209,19 @@ void FD3D12Viewport::Init()
 				ID3D12CommandQueue* pCommandQueue = Adapter->GetDevice(0)->GetD3DCommandQueue();
 
 				TRefCountPtr<IDXGISwapChain> SwapChain;
-				VERIFYD3D12RESULT(Adapter->GetDXGIFactory2()->CreateSwapChain(pCommandQueue, &SwapChainDesc, SwapChain.GetInitReference()));
+				HRESULT hr = Adapter->GetDXGIFactory2()->CreateSwapChain(pCommandQueue, &SwapChainDesc, SwapChain.GetInitReference());
+				if (FAILED(hr))
+				{
+					UE_LOG(LogD3D12RHI, Warning, TEXT("Failed to create swapchain with the following parameters:"));
+					UE_LOG(LogD3D12RHI, Warning, TEXT("\tDXGI_MODE_DESC: width: %d height: %d DXGI format: %d"), SwapChainDesc.BufferDesc.Width, SwapChainDesc.BufferDesc.Height, SwapChainDesc.BufferDesc.Format);
+					UE_LOG(LogD3D12RHI, Warning, TEXT("\tBack buffer count: %d"), NumBackBuffers);
+					UE_LOG(LogD3D12RHI, Warning, TEXT("\tWindows handle: 0x%x (IsWindow: %s)"), WindowHandle, IsWindow(WindowHandle) ? TEXT("true") : TEXT("false"));
+					UE_LOG(LogD3D12RHI, Warning, TEXT("\tFullscreen: %s"), bIsFullscreen ? TEXT("true") : TEXT("false"));
+					UE_LOG(LogD3D12RHI, Warning, TEXT("\tSwapchain flags: %d"), SwapChainFlags);
+
+					VERIFYD3D12RESULT(hr);
+				}
+				
 				VERIFYD3D12RESULT(SwapChain->QueryInterface(IID_PPV_ARGS(SwapChain1.GetInitReference())));
 
 				// Get a SwapChain4 if supported.
