Make sure all the D3D errors and warnings are forwarded to the log file (use vectorized exception handler to catch events)

diff --git a/Engine/Source/Runtime/D3D12RHI/Private/D3D12Adapter.cpp b/Engine/Source/Runtime/D3D12RHI/Private/D3D12Adapter.cpp
--- a/Engine/Source/Runtime/D3D12RHI/Private/D3D12Adapter.cpp
+++ b/Engine/Source/Runtime/D3D12RHI/Private/D3D12Adapter.cpp
@@ -23,6 +23,78 @@ static TAutoConsoleVariable<int32> CVarResidencyManagement(
 );
 #endif // ENABLE_RESIDENCY_MANAGEMENT
 
+#if PLATFORM_WINDOWS || PLATFORM_HOLOLENS
+
+/** Handle d3d messages and write them to the log file **/
+static LONG __stdcall D3DVectoredExceptionHandler(EXCEPTION_POINTERS* InInfo)
+{
+	// Only handle D3D error codes here
+	if (InInfo->ExceptionRecord->ExceptionCode == _FACDXGI)
+	{
+		TRefCountPtr<ID3D12Debug> d3dDebug;
+		if (SUCCEEDED(D3D12GetDebugInterface(__uuidof(ID3D12Debug), (void**)d3dDebug.GetInitReference())))
+		{
+			FD3D12DynamicRHI* D3D12RHI = (FD3D12DynamicRHI*)GDynamicRHI;
+			TRefCountPtr<ID3D12InfoQueue> d3dInfoQueue;
+			if (SUCCEEDED(D3D12RHI->GetAdapter().GetD3DDevice()->QueryInterface(__uuidof(ID3D12InfoQueue), (void**)d3dInfoQueue.GetInitReference())))
+			{
+				D3D12_MESSAGE* d3dMessage = nullptr;
+				SIZE_T AllocateSize = 0;
+
+				int StoredMessageCount = d3dInfoQueue->GetNumStoredMessagesAllowedByRetrievalFilter();
+				for (int MessageIndex = 0; MessageIndex < StoredMessageCount; MessageIndex++)
+				{
+					SIZE_T MessageLength = 0;
+					HRESULT hr = d3dInfoQueue->GetMessage(MessageIndex, nullptr, &MessageLength);
+
+					// Ideally the exception handler should not allocate any memory because it could fail
+					// and can cause another exception to be triggered and possible even cause a deadlock.
+					// But for these D3D error message it should be fine right now because they are requested
+					// exceptions when making an error against the API.
+					// Not allocating memory for the messages is easy (cache memory in Adapter), but ANSI_TO_TCHAR
+					// and UE_LOG will also allocate memory and aren't that easy to fix.
+
+					// realloc the message
+					if (MessageLength > AllocateSize)
+					{
+						if (d3dMessage)
+						{
+							FMemory::Free(d3dMessage);
+							d3dMessage = nullptr;
+							AllocateSize = 0;
+						}
+
+						d3dMessage = (D3D12_MESSAGE*) FMemory::Malloc(MessageLength);
+						AllocateSize = MessageLength;
+					}
+
+					// get the actual message data from the queue
+					hr = d3dInfoQueue->GetMessage(MessageIndex, d3dMessage, &MessageLength);
+
+					UE_LOG(LogD3D12RHI, Error, TEXT("%s"), ANSI_TO_TCHAR(d3dMessage->pDescription));
+
+					// when we get here, then it means that BreakOnSeverity was set for this error message, so request the debug break here as well
+					UE_DEBUG_BREAK();
+				}
+
+				if (AllocateSize > 0)
+				{
+					FMemory::Free(d3dMessage);
+				}
+			}
+		}
+
+		// Handles the exception
+		return EXCEPTION_CONTINUE_EXECUTION;
+	}
+
+	// continue searching
+	return EXCEPTION_CONTINUE_SEARCH;
+}
+
+#endif // #if PLATFORM_WINDOWS || PLATFORM_HOLOLENS
+
+
 FD3D12Adapter::FD3D12Adapter(FD3D12AdapterDesc& DescIn)
 	: OwningRHI(nullptr)
 	, bDepthBoundsTestSupported(false)
@@ -292,6 +364,13 @@ void FD3D12Adapter::CreateRootDevice(bool bWithDebug)
 	}
 #endif
 
+#if (PLATFORM_WINDOWS || PLATFORM_HOLOLENS)
+	if (bWithDebug)
+	{
+		// add vectored exception handler to write the debug device warning & error messages to the log
+		ExceptionHandlerHandle = AddVectoredExceptionHandler(1, D3DVectoredExceptionHandler);
+	}
+#endif //  (PLATFORM_WINDOWS || PLATFORM_HOLOLENS)
 
 #if UE_BUILD_DEBUG	&& (PLATFORM_WINDOWS || PLATFORM_HOLOLENS)
 	//break on debug
@@ -681,6 +760,13 @@ void FD3D12Adapter::Cleanup()
 	}
 #endif
 
+#if (PLATFORM_WINDOWS || PLATFORM_HOLOLENS)
+	if (ExceptionHandlerHandle != INVALID_HANDLE_VALUE)
+	{
+		RemoveVectoredExceptionHandler(ExceptionHandlerHandle);
+	}
+#endif //  (PLATFORM_WINDOWS || PLATFORM_HOLOLENS)
+
 	// Ask all initialized FRenderResources to release their RHI resources.
 	FRenderResource::ReleaseRHIForAllResources();
 
diff --git a/Engine/Source/Runtime/D3D12RHI/Private/D3D12Adapter.h b/Engine/Source/Runtime/D3D12RHI/Private/D3D12Adapter.h
--- a/Engine/Source/Runtime/D3D12RHI/Private/D3D12Adapter.h
+++ b/Engine/Source/Runtime/D3D12RHI/Private/D3D12Adapter.h
@@ -308,6 +308,8 @@ protected:
 	TRefCountPtr<ID3D12Device1> RootDevice1;
 #if PLATFORM_WINDOWS || PLATFORM_HOLOLENS
 	TRefCountPtr<ID3D12Device2> RootDevice2;
+
+	HANDLE ExceptionHandlerHandle = INVALID_HANDLE_VALUE;
 #endif
 #if D3D12_RHI_RAYTRACING
 	TRefCountPtr<ID3D12Device5> RootRayTracingDevice;
