Integrate Dev-Rendering fix:

Don't need to align the allocation size because the allocators will perform the alignment itself.

Remove the allignment during construction because RHIUpdateUniformBuffer doesn't perform the aligment, and then there can be a different the allocation size in the destructor (used for tracking stats)

diff --git a/Engine/Source/Runtime/D3D12RHI/Private/D3D12UniformBuffer.cpp b/Engine/Source/Runtime/D3D12RHI/Private/D3D12UniformBuffer.cpp
--- a/Engine/Source/Runtime/D3D12RHI/Private/D3D12UniformBuffer.cpp
+++ b/Engine/Source/Runtime/D3D12RHI/Private/D3D12UniformBuffer.cpp
@@ -24,11 +24,10 @@ FUniformBufferRHIRef FD3D12DynamicRHI::RHICreateUniformBuffer(const void* Conten
 		const uint32 NumBytesActualData = Layout.ConstantBufferSize;
 		if (NumBytesActualData > 0)
 		{
-			// Constant buffers must also be 16-byte aligned.
-			const uint32 NumBytes = Align(NumBytesActualData, D3D12_CONSTANT_BUFFER_DATA_PLACEMENT_ALIGNMENT);	// Allocate a size that is big enough for a multiple of 256
-			check(Align(NumBytes, 16) == NumBytes);
+			// Is this check really needed?
+			check(Align(NumBytesActualData, 16) == NumBytesActualData);
 			check(Align(Contents, 16) == Contents);
-			check(NumBytes <= D3D12_REQ_CONSTANT_BUFFER_ELEMENT_COUNT * 16);
+			check(NumBytesActualData <= D3D12_REQ_CONSTANT_BUFFER_ELEMENT_COUNT * 16);
 
 #if USE_STATIC_ROOT_SIGNATURE
 			// Create an offline CBV descriptor
@@ -39,20 +38,19 @@ FUniformBufferRHIRef FD3D12DynamicRHI::RHICreateUniformBuffer(const void* Conten
 			{
 				// Uniform buffers that live for multiple frames must use the more expensive and persistent allocation path
 				FD3D12DynamicHeapAllocator& Allocator = GetAdapter().GetUploadHeapAllocator(Device->GetGPUIndex());
-				MappedData = Allocator.AllocUploadResource(NumBytes, DEFAULT_CONTEXT_UPLOAD_POOL_ALIGNMENT, NewUniformBuffer->ResourceLocation);
+				MappedData = Allocator.AllocUploadResource(NumBytesActualData, DEFAULT_CONTEXT_UPLOAD_POOL_ALIGNMENT, NewUniformBuffer->ResourceLocation);
 			}
 			else
 			{
 				// Uniform buffers which will live for 1 frame at the max can be allocated very efficiently from a ring buffer
 				FD3D12FastConstantAllocator& Allocator = GetAdapter().GetTransientUniformBufferAllocator();
 #if USE_STATIC_ROOT_SIGNATURE
-				MappedData = Allocator.Allocate(NumBytes, NewUniformBuffer->ResourceLocation, nullptr);
+				MappedData = Allocator.Allocate(NumBytesActualData, NewUniformBuffer->ResourceLocation, nullptr);
 #else
-				MappedData = Allocator.Allocate(NumBytes, NewUniformBuffer->ResourceLocation);
+				MappedData = Allocator.Allocate(NumBytesActualData, NewUniformBuffer->ResourceLocation);
 #endif
 			}
 			check(NewUniformBuffer->ResourceLocation.GetOffsetFromBaseOfResource() % 16 == 0);
-			check(NewUniformBuffer->ResourceLocation.GetSize() == NumBytes);
 
 			// Copy the data to the upload heap
 			check(MappedData != nullptr);
