Fix compile issues for static root signature code path & static code analysis

diff --git a/Engine/Source/Runtime/D3D12RHI/Private/D3D12Adapter.cpp b/Engine/Source/Runtime/D3D12RHI/Private/D3D12Adapter.cpp
--- a/Engine/Source/Runtime/D3D12RHI/Private/D3D12Adapter.cpp
+++ b/Engine/Source/Runtime/D3D12RHI/Private/D3D12Adapter.cpp
@@ -68,10 +68,13 @@ static LONG __stdcall D3DVectoredExceptionHandler(EXCEPTION_POINTERS* InInfo)
 						AllocateSize = MessageLength;
 					}
 
-					// get the actual message data from the queue
-					hr = d3dInfoQueue->GetMessage(MessageIndex, d3dMessage, &MessageLength);
+					if (d3dMessage)
+					{
+						// get the actual message data from the queue
+						hr = d3dInfoQueue->GetMessage(MessageIndex, d3dMessage, &MessageLength);
 
-					UE_LOG(LogD3D12RHI, Error, TEXT("%s"), ANSI_TO_TCHAR(d3dMessage->pDescription));
+						UE_LOG(LogD3D12RHI, Error, TEXT("%s"), ANSI_TO_TCHAR(d3dMessage->pDescription));
+					}
 
 					// when we get here, then it means that BreakOnSeverity was set for this error message, so request the debug break here as well
 					UE_DEBUG_BREAK();
diff --git a/Engine/Source/Runtime/D3D12RHI/Private/D3D12UniformBuffer.cpp b/Engine/Source/Runtime/D3D12RHI/Private/D3D12UniformBuffer.cpp
--- a/Engine/Source/Runtime/D3D12RHI/Private/D3D12UniformBuffer.cpp
+++ b/Engine/Source/Runtime/D3D12RHI/Private/D3D12UniformBuffer.cpp
@@ -57,7 +57,7 @@ FUniformBufferRHIRef FD3D12DynamicRHI::RHICreateUniformBuffer(const void* Conten
 			FMemory::Memcpy(MappedData, Contents, NumBytesActualData);
 
 #if USE_STATIC_ROOT_SIGNATURE
-			NewUniformBuffer->View->Create(NewUniformBuffer->ResourceLocation.GetGPUVirtualAddress(), NumBytes);
+			NewUniformBuffer->View->Create(NewUniformBuffer->ResourceLocation.GetGPUVirtualAddress(), NumBytesActualData);
 #endif
 		}
 
