Editgrate CL: 12195357 from Dev-Rendering

- Refactor gpu event stack tracking to store event stack for each command list (instead of only for the default command context)
- add support for manual GPU breadcrumb tracking which should work on all IHVs

diff --git a/Engine/Source/Runtime/D3D12RHI/Private/D3D12Adapter.cpp b/Engine/Source/Runtime/D3D12RHI/Private/D3D12Adapter.cpp
--- a/Engine/Source/Runtime/D3D12RHI/Private/D3D12Adapter.cpp
+++ b/Engine/Source/Runtime/D3D12RHI/Private/D3D12Adapter.cpp
@@ -101,6 +101,7 @@ static LONG __stdcall D3DVectoredExceptionHandler(EXCEPTION_POINTERS* InInfo)
 FD3D12Adapter::FD3D12Adapter(FD3D12AdapterDesc& DescIn)
 	: OwningRHI(nullptr)
 	, bDepthBoundsTestSupported(false)
+	, bGPUCrashDebugging(false)
 	, bDeviceRemoved(false)
 	, Desc(DescIn)
 	, RootSignatureManager(this)
@@ -185,7 +186,6 @@ void FD3D12Adapter::CreateRootDevice(bool bWithDebug)
 	// Two ways to enable GPU crash debugging, command line or the r.GPUCrashDebugging variable
 	// Note: If intending to change this please alert game teams who use this for user support.
 	// GPU crash debugging will enable DRED and Aftermath if available
-	bool bGPUCrashDebugging = false;
 	if (FParse::Param(FCommandLine::Get(), TEXT("gpucrashdebugging")))
 	{
 		bGPUCrashDebugging = true;
diff --git a/Engine/Source/Runtime/D3D12RHI/Private/D3D12Adapter.h b/Engine/Source/Runtime/D3D12RHI/Private/D3D12Adapter.h
--- a/Engine/Source/Runtime/D3D12RHI/Private/D3D12Adapter.h
+++ b/Engine/Source/Runtime/D3D12RHI/Private/D3D12Adapter.h
@@ -93,6 +93,7 @@ public:
 #endif // D3D12_RHI_RAYTRACING
 	FORCEINLINE void SetDeviceRemoved(bool value) { bDeviceRemoved = value; }
 	FORCEINLINE const bool IsDeviceRemoved() const { return bDeviceRemoved; }
+	FORCEINLINE const bool IsGPUCrashDebugging() const { return bGPUCrashDebugging; }
 	FORCEINLINE FD3D12DynamicRHI* GetOwningRHI() { return OwningRHI; }
 	FORCEINLINE const D3D12_RESOURCE_HEAP_TIER GetResourceHeapTier() const { return ResourceHeapTier; }
 	FORCEINLINE const D3D12_RESOURCE_BINDING_TIER GetResourceBindingTier() const { return ResourceBindingTier; }
@@ -319,6 +320,9 @@ protected:
 	D3D_ROOT_SIGNATURE_VERSION RootSignatureVersion;
 	bool bDepthBoundsTestSupported;
 
+	/** GPU Crash debugging requested */
+	bool bGPUCrashDebugging;
+
 	/** True if the device being used has been removed. */
 	bool bDeviceRemoved;
 
diff --git a/Engine/Source/Runtime/D3D12RHI/Private/D3D12CommandContext.cpp b/Engine/Source/Runtime/D3D12RHI/Private/D3D12CommandContext.cpp
--- a/Engine/Source/Runtime/D3D12RHI/Private/D3D12CommandContext.cpp
+++ b/Engine/Source/Runtime/D3D12RHI/Private/D3D12CommandContext.cpp
@@ -89,17 +89,86 @@ FD3D12CommandContext::~FD3D12CommandContext()
 	ClearState();
 }
 
+
+/** Write out the event stack to the bread crumb resource if available */
+void FD3D12CommandContext::WriteGPUEventStackToBreadCrumbData(bool bBeginEvent)
+{
+	// Only in Windows for now, could be made available on Xbox as well
+#if PLATFORM_WINDOWS
+	// Write directly to command list if breadcrumb resource is available
+	FD3D12Resource* BreadCrumbResource = CommandListHandle.GetCommandListManager()->GetBreadCrumbResource();
+	ID3D12GraphicsCommandList2* CommandList2 = CommandListHandle.GraphicsCommandList2();
+	if (BreadCrumbResource && CommandList2)
+	{
+		// Find the max parameter count from the resource
+		int MaxParameterCount = BreadCrumbResource->GetDesc().Width / sizeof(uint32);
+
+		// allocate the parameters on the stack if smaller than 4K
+		int ParameterCount = GPUEventStack.Num() < (MaxParameterCount - 2) ? GPUEventStack.Num() + 2 : MaxParameterCount;
+		size_t MemSize = ParameterCount * (sizeof(D3D12_WRITEBUFFERIMMEDIATE_PARAMETER) + sizeof(D3D12_WRITEBUFFERIMMEDIATE_PARAMETER));
+		const bool bAllocateOnStack = (MemSize < 4096);
+		void* Mem = bAllocateOnStack ? FMemory_Alloca(MemSize) : FMemory::Malloc(MemSize);
+
+		D3D12_WRITEBUFFERIMMEDIATE_PARAMETER* Parameters = (D3D12_WRITEBUFFERIMMEDIATE_PARAMETER*)Mem;
+		D3D12_WRITEBUFFERIMMEDIATE_MODE* Modes = (D3D12_WRITEBUFFERIMMEDIATE_MODE*)(Parameters + ParameterCount);
+		for (int i = 0; i < ParameterCount; ++i)
+		{
+			Parameters[i].Dest = BreadCrumbResource->GetGPUVirtualAddress() + 4 * i;
+			Modes[i] = D3D12_WRITEBUFFERIMMEDIATE_MODE_MARKER_IN;
+
+			// Write event stack count first
+			if (i == 0)
+			{
+				Parameters[i].Value = GPUEventStack.Num();
+			}
+			// Then if it's the begin or end event
+			else if (i == 1)
+			{
+				Parameters[i].Value = bBeginEvent ? 1 : 0;
+			}
+			// Otherwise the actual stack value
+			else
+			{
+				Parameters[i].Value = GPUEventStack[i - 2];
+			}
+		}
+		CommandList2->WriteBufferImmediate(ParameterCount, Parameters, Modes);
+
+		if (!bAllocateOnStack)
+		{
+			FMemory::Free(Mem);
+		}
+	}
+#endif // PLATFORM_WINDOWS
+}
+
+
 void FD3D12CommandContext::RHIPushEvent(const TCHAR* Name, FColor Color)
 {
+	D3D12RHI::FD3DGPUProfiler& GPUProfiler = GetParentDevice()->GetParentAdapter()->GetGPUProfiler();
+
+	// forward event to profiler if it's the default context
 	if (IsDefaultContext())
 	{
-#if NV_AFTERMATH
-		GetParentDevice()->PushGPUEvent(Name, Color, CommandListHandle.AftermathCommandContext());
-#else
-		GetParentDevice()->PushGPUEvent(Name, Color);
-#endif
+		GPUProfiler.PushEvent(Name, Color);
 	}
 
+	// If we are tracking GPU crashes then retrieve the hash of the name and track in the command list somewhere
+	if (GPUProfiler.bTrackingGPUCrashData)
+	{
+		// Get the CRC of the event (handle case when depth is too big)
+		const TCHAR* EventName = (GPUProfiler.GPUCrashDataDepth < 0 || GPUEventStack.Num() < GPUProfiler.GPUCrashDataDepth) ? Name : *D3D12RHI::FD3DGPUProfiler::EventDeepString;
+		uint32 CRC = GPUProfiler.GetOrAddEventStringHash(Name);
+
+		GPUEventStack.Push(CRC);
+		WriteGPUEventStackToBreadCrumbData(true);
+
+#if NV_AFTERMATH
+		// Only track aftermath for default context?
+		if (IsDefaultContext() && GDX12NVAfterMathEnabled)
+			GFSDK_Aftermath_SetEventMarker(CommandListHandle.AftermathCommandContext(), &GPUEventStack[0], GPUEventStack.Num() * sizeof(uint32));
+#endif // NV_AFTERMATH		
+	}
 
 #if PLATFORM_WINDOWS
 	AGSContext* const AmdAgsContext = OwningRHI.GetAmdAgsContext();
@@ -111,14 +180,27 @@ void FD3D12CommandContext::RHIPushEvent(const TCHAR* Name, FColor Color)
 
 #if USE_PIX
 	PIXBeginEvent(CommandListHandle.GraphicsCommandList(), PIX_COLOR(Color.R, Color.G, Color.B), Name);
-#endif
+#endif // USE_PIX
 }
 
 void FD3D12CommandContext::RHIPopEvent()
 {
+	D3D12RHI::FD3DGPUProfiler& GPUProfiler = GetParentDevice()->GetParentAdapter()->GetGPUProfiler();
+	
 	if (IsDefaultContext())
 	{
-		GetParentDevice()->PopGPUEvent();
+		GPUProfiler.PopEvent();
+	}
+		
+	if (GPUProfiler.bTrackingGPUCrashData)
+	{
+		WriteGPUEventStackToBreadCrumbData(false);
+
+		// need to look for unbalanced push/pop
+		if (GPUEventStack.Num() > 0)
+		{
+			GPUEventStack.Pop(false);
+		}
 	}
 
 #if PLATFORM_WINDOWS
diff --git a/Engine/Source/Runtime/D3D12RHI/Private/D3D12CommandContext.h b/Engine/Source/Runtime/D3D12RHI/Private/D3D12CommandContext.h
--- a/Engine/Source/Runtime/D3D12RHI/Private/D3D12CommandContext.h
+++ b/Engine/Source/Runtime/D3D12RHI/Private/D3D12CommandContext.h
@@ -123,6 +123,9 @@ public:
 	FD3D12CommandListHandle CommandListHandle;
 	FD3D12CommandAllocator* CommandAllocator;
 	FD3D12CommandAllocatorManager CommandAllocatorManager;
+	
+	// Current GPU event stack
+	TArray<uint32> GPUEventStack;
 
 	FD3D12StateCache StateCache;
 
@@ -423,6 +426,8 @@ protected:
 	{  
 		return InGPUIndex == GetGPUIndex() ? this : nullptr; 
 	}
+	
+	void WriteGPUEventStackToBreadCrumbData(bool bBeginEvent);
 
 private:
 	void RHIClearMRT(bool bClearColor, int32 NumClearColors, const FLinearColor* ColorArray, bool bClearDepth, float Depth, bool bClearStencil, uint32 Stencil);
diff --git a/Engine/Source/Runtime/D3D12RHI/Private/D3D12CommandList.cpp b/Engine/Source/Runtime/D3D12RHI/Private/D3D12CommandList.cpp
--- a/Engine/Source/Runtime/D3D12RHI/Private/D3D12CommandList.cpp
+++ b/Engine/Source/Runtime/D3D12RHI/Private/D3D12CommandList.cpp
@@ -63,8 +63,9 @@ FD3D12CommandListHandle::FD3D12CommandListData::FD3D12CommandListData(FD3D12Devi
 	INC_DWORD_STAT(STAT_D3D12NumCommandLists);
 
 #if PLATFORM_WINDOWS
-	// Optionally obtain the ID3D12GraphicsCommandList1 interface, we don't check the HRESULT.
+	// Optionally obtain the ID3D12GraphicsCommandList1 & ID3D12GraphicsCommandList2 interface, we don't check the HRESULT.
 	CommandList->QueryInterface(IID_PPV_ARGS(CommandList1.GetInitReference()));
+	CommandList->QueryInterface(IID_PPV_ARGS(CommandList2.GetInitReference()));
 #endif
 
 #if D3D12_RHI_RAYTRACING
diff --git a/Engine/Source/Runtime/D3D12RHI/Private/D3D12CommandList.h b/Engine/Source/Runtime/D3D12RHI/Private/D3D12CommandList.h
--- a/Engine/Source/Runtime/D3D12RHI/Private/D3D12CommandList.h
+++ b/Engine/Source/Runtime/D3D12RHI/Private/D3D12CommandList.h
@@ -205,6 +205,7 @@ private:
 		TRefCountPtr<ID3D12GraphicsCommandList>	CommandList;		// Raw D3D command list pointer
 #if PLATFORM_WINDOWS
 		TRefCountPtr<ID3D12GraphicsCommandList1> CommandList1;
+		TRefCountPtr<ID3D12GraphicsCommandList2> CommandList2;
 #endif
 #if D3D12_RHI_RAYTRACING
 		TRefCountPtr<ID3D12GraphicsCommandList4> RayTracingCommandList;
@@ -407,6 +408,17 @@ public:
 		check(CommandListData && (CommandListData->CommandListType == D3D12_COMMAND_LIST_TYPE_DIRECT || CommandListData->CommandListType == D3D12_COMMAND_LIST_TYPE_COMPUTE));
 		return CommandListData->CommandList1.GetReference();
 	}
+
+	ID3D12GraphicsCommandList2* GraphicsCommandList2() const
+	{
+		check(CommandListData && (CommandListData->CommandListType == D3D12_COMMAND_LIST_TYPE_DIRECT || CommandListData->CommandListType == D3D12_COMMAND_LIST_TYPE_COMPUTE));
+		return CommandListData->CommandList2.GetReference();
+	}
+
+	FD3D12CommandListManager* GetCommandListManager() const
+	{
+		return CommandListData->CommandListManager;
+	}
 #endif
 
 #if D3D12_RHI_RAYTRACING
diff --git a/Engine/Source/Runtime/D3D12RHI/Private/D3D12Device.cpp b/Engine/Source/Runtime/D3D12RHI/Private/D3D12Device.cpp
--- a/Engine/Source/Runtime/D3D12RHI/Private/D3D12Device.cpp
+++ b/Engine/Source/Runtime/D3D12RHI/Private/D3D12Device.cpp
@@ -370,19 +370,19 @@ void FD3D12Device::Cleanup()
 	D3DX12Residency::DestroyResidencyManager(ResidencyManager);
 }
 
-ID3D12CommandQueue* FD3D12Device::GetD3DCommandQueue(ED3D12CommandQueueType InQueueType) const
+FD3D12CommandListManager* FD3D12Device::GetCommandListManager(ED3D12CommandQueueType InQueueType) const
 {
 	switch (InQueueType)
 	{
 	case ED3D12CommandQueueType::Default:
 		check(CommandListManager->GetQueueType() == InQueueType);
-		return CommandListManager->GetD3DCommandQueue();
+		return CommandListManager;
 	case ED3D12CommandQueueType::Async:
 		check(AsyncCommandListManager->GetQueueType() == InQueueType);
-		return AsyncCommandListManager->GetD3DCommandQueue();
+		return AsyncCommandListManager;
 	case ED3D12CommandQueueType::Copy:
 		check(CopyCommandListManager->GetQueueType() == InQueueType);
-		return CopyCommandListManager->GetD3DCommandQueue();
+		return CopyCommandListManager;
 	default:
 		check(false);
 		return nullptr;
@@ -399,23 +399,6 @@ void FD3D12Device::RegisterGPUDispatch(FIntVector GroupCount)
 	GetParentAdapter()->GetGPUProfiler().RegisterGPUDispatch(GroupCount);
 }
 
-void FD3D12Device::PushGPUEvent(const TCHAR* Name, FColor Color)
-{
-	GetParentAdapter()->GetGPUProfiler().PushEvent(Name, Color);
-}
-
-#if NV_AFTERMATH
-void FD3D12Device::PushGPUEvent(const TCHAR* Name, FColor Color, GFSDK_Aftermath_ContextHandle Context)
-{
-	GetParentAdapter()->GetGPUProfiler().PushEvent(Name, Color, Context);
-}
-#endif
-
-void FD3D12Device::PopGPUEvent()
-{
-	GetParentAdapter()->GetGPUProfiler().PopEvent();
-}
-
 void FD3D12Device::BlockUntilIdle()
 {
 	GetDefaultCommandContext().FlushCommands();
diff --git a/Engine/Source/Runtime/D3D12RHI/Private/D3D12Device.h b/Engine/Source/Runtime/D3D12RHI/Private/D3D12Device.h
--- a/Engine/Source/Runtime/D3D12RHI/Private/D3D12Device.h
+++ b/Engine/Source/Runtime/D3D12RHI/Private/D3D12Device.h
@@ -99,7 +99,8 @@ public:
 		FreeCommandContexts.Add(CmdContext);
 	}
 
-	ID3D12CommandQueue* GetD3DCommandQueue(ED3D12CommandQueueType InQueueType = ED3D12CommandQueueType::Default) const;
+	FD3D12CommandListManager* GetCommandListManager(ED3D12CommandQueueType inQueueType) const;
+	ID3D12CommandQueue* GetD3DCommandQueue(ED3D12CommandQueueType InQueueType = ED3D12CommandQueueType::Default) { return GetCommandListManager(InQueueType)->GetD3DCommandQueue(); }
 
 	inline FD3D12CommandContext& GetDefaultCommandContext() const { return GetCommandContext(0); }
 	inline FD3D12CommandContext& GetDefaultAsyncComputeContext() const { return GetAsyncComputeContext(0); }
@@ -111,13 +112,7 @@ public:
 
 	void RegisterGPUWork(uint32 NumPrimitives = 0, uint32 NumVertices = 0);
 	void RegisterGPUDispatch(FIntVector GroupCount);
-	void PushGPUEvent(const TCHAR* Name, FColor Color);
-	void PopGPUEvent();
-
-#if NV_AFTERMATH
-	void PushGPUEvent(const TCHAR* Name, FColor Color, GFSDK_Aftermath_ContextHandle Context);
-#endif
-
+	
 	FD3D12SamplerState* CreateSampler(const FSamplerStateInitializerRHI& Initializer);
 	void CreateSamplerInternal(const D3D12_SAMPLER_DESC& Desc, D3D12_CPU_DESCRIPTOR_HANDLE Descriptor);
 
diff --git a/Engine/Source/Runtime/D3D12RHI/Private/D3D12DirectCommandListManager.cpp b/Engine/Source/Runtime/D3D12RHI/Private/D3D12DirectCommandListManager.cpp
--- a/Engine/Source/Runtime/D3D12RHI/Private/D3D12DirectCommandListManager.cpp
+++ b/Engine/Source/Runtime/D3D12RHI/Private/D3D12DirectCommandListManager.cpp
@@ -320,6 +320,7 @@ FD3D12CommandListManager::FD3D12CommandListManager(FD3D12Device* InParent, D3D12
 	, CommandListFence(nullptr)
 	, CommandListType(InCommandListType)
 	, QueueType(InQueueType)
+	, BreadCrumbResourceAddress(nullptr)
 #if WITH_PROFILEGPU
 	, bShouldTrackCmdListTime(false)
 #endif
@@ -349,6 +350,15 @@ void FD3D12CommandListManager::Destroy()
 		CommandListFence->Destroy();
 		CommandListFence.SafeRelease();
 	}
+
+	if (BreadCrumbResource)
+	{
+		BreadCrumbResource.SafeRelease();
+		BreadCrumbHeap.SafeRelease();
+
+		VirtualFree(BreadCrumbResourceAddress, 0, MEM_RELEASE);
+		BreadCrumbResourceAddress = nullptr;
+	}
 }
 
 void FD3D12CommandListManager::Create(const TCHAR* Name, uint32 NumCommandLists, uint32 Priority)
@@ -381,6 +391,62 @@ void FD3D12CommandListManager::Create(const TCHAR* Name, uint32 NumCommandLists,
 			ReadyLists.Enqueue(hList);
 		}
 	}
+
+	// setup the bread crumb data to track GPU progress on this command queue when GPU crash debugging is enabled
+	if (Adapter->IsGPUCrashDebugging())
+	{		
+		// QI for the ID3DDevice3 - manual buffer write from command line only supported on 1709+
+		TRefCountPtr<ID3D12Device3> D3D12Device3;
+		HRESULT hr = Device->GetDevice()->QueryInterface(IID_PPV_ARGS(D3D12Device3.GetInitReference()));
+		if (SUCCEEDED(hr))
+		{
+			// find out how many entries we can much push in a single event (limit to MAX_GPU_BREADCRUMB_DEPTH)
+			int32 GPUCrashDataDepth = GetParentDevice()->GetParentAdapter()->GetGPUProfiler().GPUCrashDataDepth;
+			int32 MaxEventCount = GPUCrashDataDepth > 0 ? FMath::Min(GPUCrashDataDepth, MAX_GPU_BREADCRUMB_DEPTH) : MAX_GPU_BREADCRUMB_DEPTH;			
+
+			// Allocate persistent CPU reabable memory which will still be valid after a device lost and wrap this data in a placed resource
+			// so the GPU command list can write to it
+			BreadCrumbResourceAddress = VirtualAlloc(nullptr, MaxEventCount, MEM_COMMIT, PAGE_READWRITE);
+
+			TRefCountPtr<ID3D12Heap> D3D12Heap;
+			hr = D3D12Device3->OpenExistingHeapFromAddress(BreadCrumbResourceAddress, IID_PPV_ARGS(D3D12Heap.GetInitReference()));
+			if (SUCCEEDED(hr))
+			{
+				BreadCrumbHeap = new FD3D12Heap(Device, GetVisibilityMask());
+				BreadCrumbHeap->SetHeap(D3D12Heap);							
+
+				TCHAR TempStr[MAX_SPRINTF] = TEXT("");
+				FCString::Sprintf(TempStr, TEXT("BreadCrumbResource_%s"), Name);
+
+				const D3D12_RESOURCE_DESC BufferDesc = CD3DX12_RESOURCE_DESC::Buffer(MaxEventCount * sizeof(uint32), D3D12_RESOURCE_FLAG_ALLOW_CROSS_ADAPTER);
+				hr = Adapter->CreatePlacedResource(BufferDesc, BreadCrumbHeap.GetReference(), 0, D3D12_RESOURCE_STATE_COPY_DEST, nullptr, BreadCrumbResource.GetInitReference(), TempStr, false);
+				if (SUCCEEDED(hr))
+				{
+					UE_LOG(LogD3D12RHI, Log, TEXT("[GPUBreadCrumb] Successfully setup breadcrumb resource for %s"), Name);
+				}
+				else
+				{
+					BreadCrumbHeap.SafeRelease();
+
+					VirtualFree(BreadCrumbResourceAddress, 0, MEM_RELEASE);
+					BreadCrumbResourceAddress = nullptr;
+
+					UE_LOG(LogD3D12RHI, Warning, TEXT("[GPUBreadCrumb] Failed to CreatePlacedResource, error: %x"), hr);
+				}
+			}
+			else
+			{
+				VirtualFree(BreadCrumbResourceAddress, 0, MEM_RELEASE);
+				BreadCrumbResourceAddress = nullptr;
+
+				UE_LOG(LogD3D12RHI, Warning, TEXT("[GPUBreadCrumb] Failed to OpenExistingHeapFromAddress, error: %x"), hr);
+			}
+		}
+		else
+		{
+			UE_LOG(LogD3D12RHI, Warning, TEXT("[GPUBreadCrumb] ID3D12Device3 not available (only available on Windows 10 1709+), error: %x"), hr);
+		}
+	}
 }
 
 FGPUTimingCalibrationTimestamp FD3D12CommandListManager::GetCalibrationTimestamp()
diff --git a/Engine/Source/Runtime/D3D12RHI/Private/D3D12DirectCommandListManager.h b/Engine/Source/Runtime/D3D12RHI/Private/D3D12DirectCommandListManager.h
--- a/Engine/Source/Runtime/D3D12RHI/Private/D3D12DirectCommandListManager.h
+++ b/Engine/Source/Runtime/D3D12RHI/Private/D3D12DirectCommandListManager.h
@@ -226,6 +226,12 @@ public:
 
 	FORCEINLINE FD3D12Fence& GetFence() { check(CommandListFence); return *CommandListFence; }
 
+	/** Get the breadcrumb resource which is written during command context recording */
+	FD3D12Resource* GetBreadCrumbResource() { return BreadCrumbResource.GetReference(); }
+
+	/** Get the CPU readable data from the breadcrumb data - this data is still valid after the Device is Lost */
+	const void* GetBreadCrumbResourceAddress() const { return BreadCrumbResourceAddress; }
+
 	void WaitForCommandQueueFlush();
 
 	void ReleaseResourceBarrierCommandListAllocator();
@@ -281,6 +287,11 @@ protected:
 	FCriticalSection						ResourceStateCS;
 	FCriticalSection						FenceCS;
 
+	// Helper data used to track GPU progress on this command queue
+	void* BreadCrumbResourceAddress;
+	TRefCountPtr<FD3D12Heap> BreadCrumbHeap;
+	TRefCountPtr<FD3D12Resource> BreadCrumbResource;
+	
 #if WITH_PROFILEGPU
 	bool bShouldTrackCmdListTime;
 	FCriticalSection CmdListTimingCS;
diff --git a/Engine/Source/Runtime/D3D12RHI/Private/D3D12RHIPrivate.h b/Engine/Source/Runtime/D3D12RHI/Private/D3D12RHIPrivate.h
--- a/Engine/Source/Runtime/D3D12RHI/Private/D3D12RHIPrivate.h
+++ b/Engine/Source/Runtime/D3D12RHI/Private/D3D12RHIPrivate.h
@@ -104,6 +104,8 @@ typedef FD3D12StateCacheBase FD3D12StateCache;
 #define DEFAULT_CONTEXT_UPLOAD_POOL_ALIGNMENT (D3D12_CONSTANT_BUFFER_DATA_PLACEMENT_ALIGNMENT)
 #define TEXTURE_POOL_SIZE (8 * 1024 * 1024)
 
+#define MAX_GPU_BREADCRUMB_DEPTH 1024
+
 #if DEBUG_RESOURCE_STATES
 #define LOG_EXECUTE_COMMAND_LISTS 1
 #define ASSERT_RESOURCE_STATES 0	// Disabled for now.
@@ -922,6 +924,23 @@ public:
 	FD3D12Adapter& GetAdapter(uint32_t Index = 0) { return *ChosenAdapters[Index]; }
 	int32 GetNumAdapters() const { return ChosenAdapters.Num(); }
 
+	template<typename PerDeviceFunction>
+	void ForEachDevice(ID3D12Device* inDevice, const PerDeviceFunction& pfPerDeviceFunction)
+	{
+		for (int AdapterIndex = 0; AdapterIndex < GetNumAdapters(); ++AdapterIndex)
+		{
+			FD3D12Adapter& D3D12Adapter = GetAdapter(AdapterIndex);
+			for (uint32 GPUIndex : FRHIGPUMask::All())
+			{
+				FD3D12Device* D3D12Device = D3D12Adapter.GetDevice(GPUIndex);
+				if (inDevice == nullptr || D3D12Device->GetDevice() == inDevice)
+				{
+					pfPerDeviceFunction(D3D12Device);
+				}
+			}
+		}
+	}
+
 	AGSContext* GetAmdAgsContext() { return AmdAgsContext; }
 	void SetAmdSupportedExtensionFlags(uint32 Flags) { AmdSupportedExtensionFlags = Flags; }
 
diff --git a/Engine/Source/Runtime/D3D12RHI/Private/D3D12Stats.cpp b/Engine/Source/Runtime/D3D12RHI/Private/D3D12Stats.cpp
--- a/Engine/Source/Runtime/D3D12RHI/Private/D3D12Stats.cpp
+++ b/Engine/Source/Runtime/D3D12RHI/Private/D3D12Stats.cpp
@@ -18,7 +18,7 @@ void D3D12RHI::FD3DGPUProfiler::BeginFrame(FD3D12DynamicRHI* InRHI)
 	static auto* CrashCollectionEnableCvar = IConsoleManager::Get().FindTConsoleVariableDataInt(TEXT("r.gpucrash.collectionenable"));
 	static auto* CrashCollectionDataDepth = IConsoleManager::Get().FindTConsoleVariableDataInt(TEXT("r.gpucrash.datadepth"));
 
-	bTrackingGPUCrashData = CrashCollectionEnableCvar ? CrashCollectionEnableCvar->GetValueOnRenderThread() != 0 : false;
+	bTrackingGPUCrashData = CrashCollectionEnableCvar ? (InRHI->GetAdapter().IsGPUCrashDebugging() && CrashCollectionEnableCvar->GetValueOnRenderThread() != 0) : false;
 	GPUCrashDataDepth = CrashCollectionDataDepth ? CrashCollectionDataDepth->GetValueOnRenderThread() : -1;
 
 	// latch the bools from the game thread into our private copy
@@ -64,12 +64,7 @@ void D3D12RHI::FD3DGPUProfiler::BeginFrame(FD3D12DynamicRHI* InRHI)
 
 	if (GetEmitDrawEvents())
 	{
-#if NV_AFTERMATH
-		// Assuming that grabbing the device 0 command list here is OK
-		PushEvent(TEXT("FRAME"), FColor(0, 255, 0, 255), InRHI->GetAdapter().GetDevice(0)->GetCommandContext().CommandListHandle.AftermathCommandContext());
-#else
 		PushEvent(TEXT("FRAME"), FColor(0, 255, 0, 255));
-#endif
 	}
 }
 
@@ -189,63 +184,45 @@ void D3D12RHI::FD3DGPUProfiler::EndFrame(FD3D12DynamicRHI* InRHI)
 	CurrentEventNodeFrame = NULL;
 }
 
-void D3D12RHI::FD3DGPUProfiler::PushEvent(const TCHAR* Name, FColor Color)
-{
-	FGPUProfiler::PushEvent(Name, Color);
-}
+FString D3D12RHI::FD3DGPUProfiler::EventDeepString(TEXT("EventTooDeep"));
+const uint32 D3D12RHI::FD3DGPUProfiler::EventDeepCRC = FCrc::StrCrc32<TCHAR>(*EventDeepString);
 
-static FString EventDeepString(TEXT("EventTooDeep"));
-static const uint32 EventDeepCRC = FCrc::StrCrc32<TCHAR>(*EventDeepString);
-#if NV_AFTERMATH
-void D3D12RHI::FD3DGPUProfiler::PushEvent(const TCHAR* Name, FColor Color, GFSDK_Aftermath_ContextHandle Context)
+/** Get the CRC of the given event Name and cache the lookup internally so it can be retrieved again later */
+uint32 D3D12RHI::FD3DGPUProfiler::GetOrAddEventStringHash(const TCHAR* Name)
 {
-
-	if (GDX12NVAfterMathEnabled && bTrackingGPUCrashData)
+	if (bTrackingGPUCrashData)
 	{
-		uint32 CRC = 0;
-		if (GPUCrashDataDepth < 0 || PushPopStack.Num() < GPUCrashDataDepth)
+		uint32 CRC = FCrc::StrCrc32<TCHAR>(Name);
+
+		// make sure the Name is cached
+		FRWScopeLock RWScopeLock(CacheEventStringsRWLock, SLT_ReadOnly);
+		if (!CachedEventStrings.Contains(CRC))
 		{
-			CRC = FCrc::StrCrc32<TCHAR>(Name);
+			RWScopeLock.ReleaseReadOnlyLockAndAcquireWriteLock_USE_WITH_CAUTION();
 
-			if (CachedStrings.Num() > 10000)
+			if (CachedEventStrings.Num() > 10000)
 			{
-				CachedStrings.Empty(10000);
-				CachedStrings.Emplace(EventDeepCRC, EventDeepString);
+				CachedEventStrings.Empty(10000);
+				CachedEventStrings.Emplace(EventDeepCRC, EventDeepString);
 			}
 
-			if (CachedStrings.Find(CRC) == nullptr)
+			if (CachedEventStrings.Find(CRC) == nullptr)
 			{
-				CachedStrings.Emplace(CRC, FString(Name));
+				CachedEventStrings.Emplace(CRC, FString(Name));
 			}
-
-		}
-		else
-		{
-			CRC = EventDeepCRC;
 		}
-		PushPopStack.Push(CRC);
 
-		GFSDK_Aftermath_SetEventMarker(Context, &PushPopStack[0], PushPopStack.Num() * sizeof(uint32));
+		return CRC;
 	}
-
-	PushEvent(Name, Color);
+	else
+		return 0;
 }
-#endif
 
-void D3D12RHI::FD3DGPUProfiler::PopEvent()
+/** Try and find the cached event string for given CRC */
+const FString* D3D12RHI::FD3DGPUProfiler::FindEventString(uint32 CRC)
 {
-#if NV_AFTERMATH
-	if (GDX12NVAfterMathEnabled && bTrackingGPUCrashData)
-	{
-		// need to look for unbalanced push/pop
-		if (PushPopStack.Num() > 0)
-		{
-			PushPopStack.Pop(false);
-		}
-	}
-#endif
-
-	FGPUProfiler::PopEvent();
+	FReadScopeLock ReadScopeLock(CacheEventStringsRWLock);
+	return CachedEventStrings.Find(CRC);
 }
 
 /** Start this frame of per tracking */
@@ -412,7 +389,7 @@ bool D3D12RHI::FD3DGPUProfiler::CheckGpuHeartbeat() const
 							uint32* Data = (uint32*)ContextData.markerData;
 							for (uint32 i = 0; i < NumCRCs; i++)
 							{
-								const FString* Frame = CachedStrings.Find(Data[i]);
+								const FString* Frame = CachedEventStrings.Find(Data[i]);
 								if (Frame != nullptr)
 								{
 									UE_LOG(LogRHI, Error, TEXT("[Aftermath] %i: %s"), i, *(*Frame));
diff --git a/Engine/Source/Runtime/D3D12RHI/Private/D3D12Stats.h b/Engine/Source/Runtime/D3D12RHI/Private/D3D12Stats.h
--- a/Engine/Source/Runtime/D3D12RHI/Private/D3D12Stats.h
+++ b/Engine/Source/Runtime/D3D12RHI/Private/D3D12Stats.h
@@ -319,14 +319,16 @@ namespace D3D12RHI
 			return EventNode;
 		}
 
-		virtual void PushEvent(const TCHAR* Name, FColor Color) override;
-		virtual void PopEvent() override;
-
 		void BeginFrame(class FD3D12DynamicRHI* InRHI);
-
 		void EndFrame(class FD3D12DynamicRHI* InRHI);
 
 		bool CheckGpuHeartbeat() const;
+		
+		static FString EventDeepString;
+		static const uint32 EventDeepCRC;
+
+		uint32 GetOrAddEventStringHash(const TCHAR* Name);
+		const FString* FindEventString(uint32 CRC);
 
 		/**
 		 * Calculate the amount of GPU idle time between two timestamps
@@ -337,16 +339,11 @@ namespace D3D12RHI
 		uint64 CalculateIdleTime(uint64 StartTime, uint64 EndTime);
 
 #if NV_AFTERMATH
-		virtual void PushEvent(const TCHAR* Name, FColor Color, GFSDK_Aftermath_ContextHandle Context);
-
 		void RegisterCommandList(GFSDK_Aftermath_ContextHandle context);
 		void UnregisterCommandList(GFSDK_Aftermath_ContextHandle context);
 
 		TArray<GFSDK_Aftermath_ContextHandle> AftermathContexts;
 		FCriticalSection AftermathLock;
-
-		TArray<uint32> PushPopStack;
-		TMap<uint32, FString> CachedStrings;
 #endif
 
 		/** Used to measure GPU time per frame. */
@@ -367,5 +364,9 @@ namespace D3D12RHI
 		TArray<uint64> CmdListEndTimestamps;
 		/** Accumulated idle GPU ticks before each corresponding command list */
 		TArray<uint64> IdleTimeCDF;
+
+		/** Map containing all the currently hashed event strings */
+		FRWLock	CacheEventStringsRWLock;
+		TMap<uint32, FString> CachedEventStrings;
 	};
 }
\ No newline at end of file
diff --git a/Engine/Source/Runtime/D3D12RHI/Private/D3D12Util.cpp b/Engine/Source/Runtime/D3D12RHI/Private/D3D12Util.cpp
--- a/Engine/Source/Runtime/D3D12RHI/Private/D3D12Util.cpp
+++ b/Engine/Source/Runtime/D3D12RHI/Private/D3D12Util.cpp
@@ -128,6 +128,19 @@ static FString GetD3D12ErrorString(HRESULT ErrorCode, ID3D12Device* Device)
 	return ErrorCodeText;
 }
 
+/** Build string name of command queue type */
+static const TCHAR* GetD3DCommandQueueTypeName(ED3D12CommandQueueType QueueType)
+{
+	switch (QueueType)
+	{
+	case ED3D12CommandQueueType::Default:	 return TEXT("3D");
+	case ED3D12CommandQueueType::Async:		 return TEXT("Compute");
+	case ED3D12CommandQueueType::Copy:		 return TEXT("Copy");
+	}
+
+	return nullptr;
+}
+
 #undef D3DERR
 
 namespace D3D12RHI
@@ -204,9 +217,72 @@ static FString GetD3D12TextureFlagString(uint32 TextureFlags)
 	return TextureFormatText;
 }
 
+/** Log the GPU progress of the given CommandListManager to the Error log if breadcrumb data is available */
+static bool LogBreadcrumbData(D3D12RHI::FD3DGPUProfiler& GPUProfiler, FD3D12CommandListManager& CommandListManager)
+{
+	uint32* BreadCrumbData = (uint32*)CommandListManager.GetBreadCrumbResourceAddress();
+	if (BreadCrumbData == nullptr)
+	{
+		return false;
+	}
+
+	uint32 EventCount = BreadCrumbData[0];
+	bool bBeginEvent = BreadCrumbData[1] > 0;
+	check(EventCount >= 0 && EventCount < (MAX_GPU_BREADCRUMB_DEPTH - 2));
+
+	FString gpu_progress = FString::Printf(TEXT("[GPUBreadCrumb]\t%s Queue %d - %s"), GetD3DCommandQueueTypeName(CommandListManager.GetQueueType()), 
+		CommandListManager.GetGPUIndex(), EventCount == 0 ? TEXT("No Data") : (bBeginEvent ? TEXT("Begin: ") : TEXT("End: ")));
+	for (uint32 EventIndex = 0; EventIndex < EventCount; ++EventIndex)
+	{
+		if (EventIndex > 0)
+		{
+			gpu_progress.Append(TEXT(" - "));
+		}
+
+		// get the crc and try and translate back into a string
+		uint32 event_crc = BreadCrumbData[EventIndex + 2];
+		const FString* event_name = GPUProfiler.FindEventString(event_crc);
+		if (event_name)
+		{
+			gpu_progress.Append(*event_name);
+		}
+		else
+		{
+			gpu_progress.Append(TEXT("Unknown Event"));
+		}
+	}
+
+	UE_LOG(LogD3D12RHI, Error, TEXT("%s"), *gpu_progress);
+
+	return true;
+}
+
+/** Log the GPU progress of the given Device to the Error log if breadcrumb data is available */
+static void LogBreadcrumbData(ID3D12Device* Device)
+{
+	UE_LOG(LogD3D12RHI, Error, TEXT("[GPUBreadCrumb] Last tracked GPU operations:"));
+
+	bool bValidData = true;
+
+	// Check all the devices
+	FD3D12DynamicRHI* D3D12RHI = (FD3D12DynamicRHI*)GDynamicRHI;
+	D3D12RHI->ForEachDevice(Device, [&](FD3D12Device* Device)
+		{
+			bValidData = bValidData && LogBreadcrumbData(Device->GetParentAdapter()->GetGPUProfiler(), Device->GetCommandListManager());
+			bValidData = bValidData && LogBreadcrumbData(Device->GetParentAdapter()->GetGPUProfiler(), Device->GetAsyncCommandListManager());
+			bValidData = bValidData && LogBreadcrumbData(Device->GetParentAdapter()->GetGPUProfiler(), Device->GetCopyCommandListManager());
+		});
+
+	if (!bValidData)
+	{
+		UE_LOG(LogD3D12RHI, Error, TEXT("No Valid GPU Breadcrumb data found. Use -gpucrashdebugging to collect GPU progress when debugging GPU crashes."));
+	}
+}
 
 #if PLATFORM_WINDOWS
 
+
+/** Log the DRED data to Error log if available */
 static void LogDREDData(ID3D12Device* Device)
 {
 	// Should match all values from D3D12_AUTO_BREADCRUMB_OP
@@ -380,6 +456,8 @@ static void TerminateOnDeviceRemoved(HRESULT D3DResult, ID3D12Device* Device)
 		GIsCriticalError = true;
 		GIsGPUCrashed = true;
 
+		LogBreadcrumbData(Device);
+
 #if PLATFORM_WINDOWS
 		if (Device)
 		{
