Editgrate Dev-Rendering code:

- Refactor gpu event stack tracking to store event stack for each command list (instead of only for the default command context)
- add support for manual GPU breadcrumb tracking which should work on all IHVs

diff --git a/Engine/Source/Runtime/D3D12RHI/Private/D3D12Adapter.cpp b/Engine/Source/Runtime/D3D12RHI/Private/D3D12Adapter.cpp
--- a/Engine/Source/Runtime/D3D12RHI/Private/D3D12Adapter.cpp
+++ b/Engine/Source/Runtime/D3D12RHI/Private/D3D12Adapter.cpp
@@ -243,6 +243,22 @@ void FD3D12Adapter::CreateRootDevice(bool bWithDebug)
 			UE_LOG(LogD3D12RHI, Fatal, TEXT("The debug interface requires the D3D12 SDK Layers. Please install the Graphics Tools for Windows. See: https://docs.microsoft.com/en-us/windows/uwp/gaming/use-the-directx-runtime-and-visual-studio-graphics-diagnostic-features"));
 		}
 	}
+
+	// Two ways to enable GPU crash debugging, command line or the r.GPUCrashDebugging variable
+	// Note: If intending to change this please alert game teams who use this for user support.
+	// GPU crash debugging will enable DRED and Aftermath if available
+	if (FParse::Param(FCommandLine::Get(), TEXT("gpucrashdebugging")))
+	{
+		bGPUCrashDebugging = true;
+	}
+	else
+	{
+		static IConsoleVariable* GPUCrashDebugging = IConsoleManager::Get().FindConsoleVariable(TEXT("r.GPUCrashDebugging"));
+		if (GPUCrashDebugging)
+		{
+			bGPUCrashDebugging = GPUCrashDebugging->GetInt() > 0;
+		}
+	}
 	
 	// Setup DRED if requested
 	if (bGPUCrashDebugging || FParse::Param(FCommandLine::Get(), TEXT("dred")))
diff --git a/Engine/Source/Runtime/D3D12RHI/Private/D3D12CommandContext.cpp b/Engine/Source/Runtime/D3D12RHI/Private/D3D12CommandContext.cpp
--- a/Engine/Source/Runtime/D3D12RHI/Private/D3D12CommandContext.cpp
+++ b/Engine/Source/Runtime/D3D12RHI/Private/D3D12CommandContext.cpp
@@ -165,8 +165,10 @@ void FD3D12CommandContext::RHIPushEvent(const TCHAR* Name, FColor Color)
 
 #if NV_AFTERMATH
 		// Only track aftermath for default context?
-		if (IsDefaultContext() && GDX12NVAfterMathEnabled && GDX12NVAfterMathMarkers)
+		if (IsDefaultContext() && GDX12NVAfterMathEnabled)
+		{
 			GFSDK_Aftermath_SetEventMarker(CommandListHandle.AftermathCommandContext(), &GPUEventStack[0], GPUEventStack.Num() * sizeof(uint32));
+		}
 #endif // NV_AFTERMATH		
 	}
 
diff --git a/Engine/Source/Runtime/D3D12RHI/Private/D3D12Util.cpp b/Engine/Source/Runtime/D3D12RHI/Private/D3D12Util.cpp
--- a/Engine/Source/Runtime/D3D12RHI/Private/D3D12Util.cpp
+++ b/Engine/Source/Runtime/D3D12RHI/Private/D3D12Util.cpp
@@ -474,14 +474,7 @@ namespace D3D12RHI
 		GIsCriticalError = true;
 		GIsGPUCrashed = true;
 
-		// Check GPU heartbeat - will trace Aftermath state
-		if (GDynamicRHI)
-		{
-			GDynamicRHI->CheckGpuHeartbeat();
-		}
-
-		// Log RHI independent breadcrumbing data
-		LogBreadcrumbData(InDevice);
+		LogBreadcrumbData(Device);
 
 		FD3D12DynamicRHI* D3D12RHI = (FD3D12DynamicRHI*)GDynamicRHI;
 #if PLATFORM_WINDOWS
