Compile fix

diff --git a/Engine/Source/Runtime/D3D12RHI/Private/D3D12Util.cpp b/Engine/Source/Runtime/D3D12RHI/Private/D3D12Util.cpp
--- a/Engine/Source/Runtime/D3D12RHI/Private/D3D12Util.cpp
+++ b/Engine/Source/Runtime/D3D12RHI/Private/D3D12Util.cpp
@@ -531,6 +531,7 @@ namespace D3D12RHI
 			UE_LOG(LogD3D12RHI, Error, TEXT("%s"), *ErrorMessage.ToText().ToString());
 		}
 
+#if PLATFORM_WINDOWS || PLATFORM_HOLOLENS
 		// If we have crash dump data then dump to disc
 		if (InGPUCrashDump != nullptr)
 		{
@@ -549,6 +550,7 @@ namespace D3D12RHI
 			ReportGPUCrash(TEXT("Aftermath GPU Crash dump Triggered"), 0);
 		}
 		else
+#endif // PLATFORM_WINDOWS || PLATFORM_HOLOLENS
 		{
 			// Make sure the log is flushed!
 			GLog->PanicFlushThreadedLogs();
