UE4.25 - Compile fix
#rnx

diff --git a/Engine/Source/Runtime/D3D12RHI/Private/D3D12Adapter.cpp b/Engine/Source/Runtime/D3D12RHI/Private/D3D12Adapter.cpp
--- a/Engine/Source/Runtime/D3D12RHI/Private/D3D12Adapter.cpp
+++ b/Engine/Source/Runtime/D3D12RHI/Private/D3D12Adapter.cpp
@@ -269,6 +269,11 @@ void FD3D12Adapter::CreateRootDevice(bool bWithDebug)
 
 #endif // PLATFORM_WINDOWS || (PLATFORM_HOLOLENS && !UE_BUILD_SHIPPING && D3D12_PROFILING_ENABLED)
 
+#if USE_PIX
+	UE_LOG(LogD3D12RHI, Log, TEXT("Emitting draw events for PIX profiling."));
+	SetEmitDrawEvents(true);
+#endif
+
 	bool bDeviceCreated = false;
 #if !PLATFORM_CPU_ARM_FAMILY && (PLATFORM_WINDOWS || PLATFORM_HOLOLENS)
 	if (IsRHIDeviceAMD() && OwningRHI->GetAmdAgsContext())
