Static code analysis fixes

diff --git a/Engine/Source/Runtime/D3D12RHI/Private/D3D12CommandContext.cpp b/Engine/Source/Runtime/D3D12RHI/Private/D3D12CommandContext.cpp
--- a/Engine/Source/Runtime/D3D12RHI/Private/D3D12CommandContext.cpp
+++ b/Engine/Source/Runtime/D3D12RHI/Private/D3D12CommandContext.cpp
@@ -109,30 +109,33 @@ void FD3D12CommandContext::WriteGPUEventStackToBreadCrumbData(bool bBeginEvent)
 		const bool bAllocateOnStack = (MemSize < 4096);
 		void* Mem = bAllocateOnStack ? FMemory_Alloca(MemSize) : FMemory::Malloc(MemSize);
 
-		D3D12_WRITEBUFFERIMMEDIATE_PARAMETER* Parameters = (D3D12_WRITEBUFFERIMMEDIATE_PARAMETER*)Mem;
-		D3D12_WRITEBUFFERIMMEDIATE_MODE* Modes = (D3D12_WRITEBUFFERIMMEDIATE_MODE*)(Parameters + ParameterCount);
-		for (int i = 0; i < ParameterCount; ++i)
+		if (Mem)
 		{
-			Parameters[i].Dest = BreadCrumbResource->GetGPUVirtualAddress() + 4 * i;
-			Modes[i] = D3D12_WRITEBUFFERIMMEDIATE_MODE_MARKER_IN;
-
-			// Write event stack count first
-			if (i == 0)
-			{
-				Parameters[i].Value = GPUEventStack.Num();
-			}
-			// Then if it's the begin or end event
-			else if (i == 1)
+			D3D12_WRITEBUFFERIMMEDIATE_PARAMETER* Parameters = (D3D12_WRITEBUFFERIMMEDIATE_PARAMETER*)Mem;
+			D3D12_WRITEBUFFERIMMEDIATE_MODE* Modes = (D3D12_WRITEBUFFERIMMEDIATE_MODE*)(Parameters + ParameterCount);
+			for (int i = 0; i < ParameterCount; ++i)
 			{
-				Parameters[i].Value = bBeginEvent ? 1 : 0;
-			}
-			// Otherwise the actual stack value
-			else
-			{
-				Parameters[i].Value = GPUEventStack[i - 2];
+				Parameters[i].Dest = BreadCrumbResource->GetGPUVirtualAddress() + 4 * i;
+				Modes[i] = D3D12_WRITEBUFFERIMMEDIATE_MODE_MARKER_IN;
+
+				// Write event stack count first
+				if (i == 0)
+				{
+					Parameters[i].Value = GPUEventStack.Num();
+				}
+				// Then if it's the begin or end event
+				else if (i == 1)
+				{
+					Parameters[i].Value = bBeginEvent ? 1 : 0;
+				}
+				// Otherwise the actual stack value
+				else
+				{
+					Parameters[i].Value = GPUEventStack[i - 2];
+				}
 			}
+			CommandList2->WriteBufferImmediate(Paramet	erCount, Parameters, Modes);
 		}
-		CommandList2->WriteBufferImmediate(ParameterCount, Parameters, Modes);
 
 		if (!bAllocateOnStack)
 		{
diff --git a/Engine/Source/Runtime/D3D12RHI/Private/D3D12DirectCommandListManager.cpp b/Engine/Source/Runtime/D3D12RHI/Private/D3D12DirectCommandListManager.cpp
--- a/Engine/Source/Runtime/D3D12RHI/Private/D3D12DirectCommandListManager.cpp
+++ b/Engine/Source/Runtime/D3D12RHI/Private/D3D12DirectCommandListManager.cpp
@@ -407,40 +407,42 @@ void FD3D12CommandListManager::Create(const TCHAR* Name, uint32 NumCommandLists,
 			// Allocate persistent CPU reabable memory which will still be valid after a device lost and wrap this data in a placed resource
 			// so the GPU command list can write to it
 			BreadCrumbResourceAddress = VirtualAlloc(nullptr, MaxEventCount, MEM_COMMIT, PAGE_READWRITE);
-
-			TRefCountPtr<ID3D12Heap> D3D12Heap;
-			hr = D3D12Device3->OpenExistingHeapFromAddress(BreadCrumbResourceAddress, IID_PPV_ARGS(D3D12Heap.GetInitReference()));
-			if (SUCCEEDED(hr))
+			if (BreadCrumbResourceAddress)
 			{
-				BreadCrumbHeap = new FD3D12Heap(Device, GetVisibilityMask());
-				BreadCrumbHeap->SetHeap(D3D12Heap);							
-
-				TCHAR TempStr[MAX_SPRINTF] = TEXT("");
-				FCString::Sprintf(TempStr, TEXT("BreadCrumbResource_%s"), Name);
-
-				const D3D12_RESOURCE_DESC BufferDesc = CD3DX12_RESOURCE_DESC::Buffer(MaxEventCount * sizeof(uint32), D3D12_RESOURCE_FLAG_ALLOW_CROSS_ADAPTER);
-				hr = Adapter->CreatePlacedResource(BufferDesc, BreadCrumbHeap.GetReference(), 0, D3D12_RESOURCE_STATE_COPY_DEST, nullptr, BreadCrumbResource.GetInitReference(), TempStr, false);
+				TRefCountPtr<ID3D12Heap> D3D12Heap;
+				hr = D3D12Device3->OpenExistingHeapFromAddress(BreadCrumbResourceAddress, IID_PPV_ARGS(D3D12Heap.GetInitReference()));
 				if (SUCCEEDED(hr))
 				{
-					UE_LOG(LogD3D12RHI, Log, TEXT("[GPUBreadCrumb] Successfully setup breadcrumb resource for %s"), Name);
+					BreadCrumbHeap = new FD3D12Heap(Device, GetVisibilityMask());
+					BreadCrumbHeap->SetHeap(D3D12Heap);
+
+					TCHAR TempStr[MAX_SPRINTF] = TEXT("");
+					FCString::Sprintf(TempStr, TEXT("BreadCrumbResource_%s"), Name);
+
+					const D3D12_RESOURCE_DESC BufferDesc = CD3DX12_RESOURCE_DESC::Buffer(MaxEventCount * sizeof(uint32), D3D12_RESOURCE_FLAG_ALLOW_CROSS_ADAPTER);
+					hr = Adapter->CreatePlacedResource(BufferDesc, BreadCrumbHeap.GetReference(), 0, D3D12_RESOURCE_STATE_COPY_DEST, nullptr, BreadCrumbResource.GetInitReference(), TempStr, false);
+					if (SUCCEEDED(hr))
+					{
+						UE_LOG(LogD3D12RHI, Log, TEXT("[GPUBreadCrumb] Successfully setup breadcrumb resource for %s"), Name);
+					}
+					else
+					{
+						BreadCrumbHeap.SafeRelease();
+
+						VirtualFree(BreadCrumbResourceAddress, 0, MEM_RELEASE);
+						BreadCrumbResourceAddress = nullptr;
+
+						UE_LOG(LogD3D12RHI, Warning, TEXT("[GPUBreadCrumb] Failed to CreatePlacedResource, error: %x"), hr);
+					}
 				}
 				else
 				{
-					BreadCrumbHeap.SafeRelease();
-
 					VirtualFree(BreadCrumbResourceAddress, 0, MEM_RELEASE);
 					BreadCrumbResourceAddress = nullptr;
 
-					UE_LOG(LogD3D12RHI, Warning, TEXT("[GPUBreadCrumb] Failed to CreatePlacedResource, error: %x"), hr);
+					UE_LOG(LogD3D12RHI, Warning, TEXT("[GPUBreadCrumb] Failed to OpenExistingHeapFromAddress, error: %x"), hr);
 				}
 			}
-			else
-			{
-				VirtualFree(BreadCrumbResourceAddress, 0, MEM_RELEASE);
-				BreadCrumbResourceAddress = nullptr;
-
-				UE_LOG(LogD3D12RHI, Warning, TEXT("[GPUBreadCrumb] Failed to OpenExistingHeapFromAddress, error: %x"), hr);
-			}
 		}
 		else
 		{
