Added support for static samplers for D3D12 shaders compiled with DXC for PC.

diff --git a/Engine/Shaders/Public/Platform/D3D/D3DCommon.ush b/Engine/Shaders/Public/Platform/D3D/D3DCommon.ush
--- a/Engine/Shaders/Public/Platform/D3D/D3DCommon.ush
+++ b/Engine/Shaders/Public/Platform/D3D/D3DCommon.ush
@@ -9,3 +9,17 @@ precise float2 MakePrecise(precise float2 v) { return v; }
 precise float3 MakePrecise(precise float3 v) { return v; }
 precise float4 MakePrecise(precise float4 v) { return v; }
 #define INVARIANT(X) X
+
+#if PLATFORM_SUPPORTS_STATIC_SAMPLERS
+
+// Static sampler table must match D3D12RootSignature.cpp
+SamplerState D3DStaticPointWrappedSampler       : register(s1000);
+SamplerState D3DStaticPointClampedSampler       : register(s1001);
+SamplerState D3DStaticBilinearWrappedSampler    : register(s1002);
+SamplerState D3DStaticBilinearClampedSampler    : register(s1003);
+SamplerState D3DStaticTrilinearWrappedSampler   : register(s1004);
+SamplerState D3DStaticTrilinearClampedSampler   : register(s1005);
+
+#define GetGlobalSampler(Filter, WrapMode) D3DStatic##Filter##WrapMode##Sampler
+
+#endif // PLATFORM_SUPPORTS_STATIC_SAMPLERS
diff --git a/Engine/Source/Developer/Windows/ShaderFormatD3D/Private/D3DShaderCompiler.cpp b/Engine/Source/Developer/Windows/ShaderFormatD3D/Private/D3DShaderCompiler.cpp
--- a/Engine/Source/Developer/Windows/ShaderFormatD3D/Private/D3DShaderCompiler.cpp
+++ b/Engine/Source/Developer/Windows/ShaderFormatD3D/Private/D3DShaderCompiler.cpp
@@ -1670,6 +1670,7 @@ void CompileD3DShader(const FShaderCompilerInput& Input, FShaderCompilerOutput&
 	if (bUseDXC)
 	{
 		AdditionalDefines.SetDefine(TEXT("PLATFORM_SUPPORTS_SM6_0_WAVE_OPERATIONS"), 1);
+		AdditionalDefines.SetDefine(TEXT("PLATFORM_SUPPORTS_STATIC_SAMPLERS"), 1);
 	}
 
 	if (Input.bSkipPreprocessedCache)
diff --git a/Engine/Source/Runtime/D3D12RHI/Private/D3D12RootSignature.cpp b/Engine/Source/Runtime/D3D12RHI/Private/D3D12RootSignature.cpp
--- a/Engine/Source/Runtime/D3D12RHI/Private/D3D12RootSignature.cpp
+++ b/Engine/Source/Runtime/D3D12RHI/Private/D3D12RootSignature.cpp
@@ -17,6 +17,38 @@ namespace
 	static const uint32 RootDescriptorCost = 2; // Root descriptor is 64-bit GPU virtual address, 2 DWORDs
 }
 
+static D3D12_STATIC_SAMPLER_DESC MakeStaticSampler(D3D12_FILTER Filter, D3D12_TEXTURE_ADDRESS_MODE WrapMode, uint32 Register)
+{
+	D3D12_STATIC_SAMPLER_DESC Result = {};
+	
+	Result.Filter           = Filter;
+	Result.AddressU         = WrapMode;
+	Result.AddressV         = WrapMode;
+	Result.AddressW         = WrapMode;
+	Result.MipLODBias       = 0.0f;
+	Result.MaxAnisotropy    = 1;
+	Result.ComparisonFunc   = D3D12_COMPARISON_FUNC_NEVER;
+	Result.BorderColor      = D3D12_STATIC_BORDER_COLOR_TRANSPARENT_BLACK;
+	Result.MinLOD           = 0.0f;
+	Result.MaxLOD           = D3D12_FLOAT32_MAX;
+	Result.ShaderRegister   = Register;
+	Result.RegisterSpace    = 0;
+	Result.ShaderVisibility = D3D12_SHADER_VISIBILITY_ALL;
+
+	return Result;
+}
+
+// Static sampler table must match D3DCommon.ush
+static const D3D12_STATIC_SAMPLER_DESC StaticSamplerDescs[] =
+{
+	MakeStaticSampler(D3D12_FILTER_MIN_MAG_MIP_POINT,        D3D12_TEXTURE_ADDRESS_MODE_WRAP,  1000),
+	MakeStaticSampler(D3D12_FILTER_MIN_MAG_MIP_POINT,        D3D12_TEXTURE_ADDRESS_MODE_CLAMP, 1001),
+	MakeStaticSampler(D3D12_FILTER_MIN_MAG_LINEAR_MIP_POINT, D3D12_TEXTURE_ADDRESS_MODE_WRAP,  1002),
+	MakeStaticSampler(D3D12_FILTER_MIN_MAG_LINEAR_MIP_POINT, D3D12_TEXTURE_ADDRESS_MODE_CLAMP, 1003),
+	MakeStaticSampler(D3D12_FILTER_MIN_MAG_MIP_LINEAR,       D3D12_TEXTURE_ADDRESS_MODE_WRAP,  1004),
+	MakeStaticSampler(D3D12_FILTER_MIN_MAG_MIP_LINEAR,       D3D12_TEXTURE_ADDRESS_MODE_CLAMP, 1005),
+};
+
 FORCEINLINE D3D12_SHADER_VISIBILITY GetD3D12ShaderVisibility(EShaderVisibility Visibility)
 {
 	switch (Visibility)
@@ -261,7 +293,20 @@ FD3D12RootSignatureDesc::FD3D12RootSignatureDesc(const FD3D12QuantizedBoundShade
 		UE_LOG(LogD3D12RHI, Verbose, TEXT("Root signature created where the root parameters take up %u DWORDS. Using more than %u DWORDs can negatively impact performance depending on the hardware and root parameter usage."), RootParametersSize, SizeWarningThreshold);
 	}
 #endif
-	RootDesc.Init_1_1(RootParameterCount, TableSlots, 0, nullptr, Flags);
+
+#if D3D12_RHI_RAYTRACING
+	if (QBSS.RootSignatureType == RS_RayTracingLocal)
+	{
+		// Local root signatures don't need to provide static samplers as they are provided by global RS already.
+		// Providing static sampler bindings in global and local RS simultaneously is invalid due to overlapping register ranges.
+		RootDesc.Init_1_1(RootParameterCount, TableSlots, 0, nullptr, Flags);
+	}
+	else 
+#endif
+	{
+		RootDesc.Init_1_1(RootParameterCount, TableSlots, UE_ARRAY_COUNT(StaticSamplerDescs), StaticSamplerDescs, Flags);
+	}
+
 }
 
 D3D12_VERSIONED_ROOT_SIGNATURE_DESC& FD3D12RootSignatureDesc::GetStaticGraphicsRootSignatureDesc()
