Add ED3D12GPUCrashDebugginMode which allows for minimal impact gpu state tracking

diff --git a/Engine/Source/Runtime/D3D12RHI/Private/D3D12Adapter.cpp b/Engine/Source/Runtime/D3D12RHI/Private/D3D12Adapter.cpp
--- a/Engine/Source/Runtime/D3D12RHI/Private/D3D12Adapter.cpp
+++ b/Engine/Source/Runtime/D3D12RHI/Private/D3D12Adapter.cpp
@@ -25,6 +25,18 @@ static TAutoConsoleVariable<int32> CVarResidencyManagement(
 
 #if PLATFORM_WINDOWS || PLATFORM_HOLOLENS
 
+static int32 GD3D12GPUCrashDebuggingMode = 0;
+static TAutoConsoleVariable<int32> CVarD3D12GPUCrashDebuggingMode(
+	TEXT("r.D3D12.GPUCrashDebuggingMode"),
+	GD3D12GPUCrashDebuggingMode,
+	TEXT("Enable GPU crash debugging: tracks the current GPU state and logs information what operations the GPU executed last.\n")
+	TEXT("Optionally generate a GPU crash dump as well (on nVidia hardware only)):\n")
+	TEXT(" 0: GPU crash debugging enabled (default)\n")
+	TEXT(" 1: Minimal overhead GPU crash debugging\n")
+	TEXT(" 2: Enable all available GPU crash debugging options (DRED, Aftermath, ...)\n"),
+	ECVF_RenderThreadSafe | ECVF_ReadOnly
+);
+
 /** Handle d3d messages and write them to the log file **/
 static LONG __stdcall D3DVectoredExceptionHandler(EXCEPTION_POINTERS* InInfo)
 {
@@ -102,7 +114,7 @@ FD3D12Adapter::FD3D12Adapter(FD3D12AdapterDesc& DescIn)
 	: OwningRHI(nullptr)
 	, bDepthBoundsTestSupported(false)
 	, bDebugDevice(false)
-	, bGPUCrashDebugging(false)
+	, GPUCrashDebuggingMode(ED3D12GPUCrashDebugginMode::Disabled)
 	, bDeviceRemoved(false)
 	, Desc(DescIn)
 	, RootSignatureManager(this)
@@ -176,15 +188,23 @@ void FD3D12Adapter::CreateRootDevice(bool bWithDebug)
 	// GPU crash debugging will enable DRED and Aftermath if available
 	if (FParse::Param(FCommandLine::Get(), TEXT("gpucrashdebugging")))
 	{
-		bGPUCrashDebugging = true;
+		GPUCrashDebuggingMode = ED3D12GPUCrashDebugginMode::Full;
 	}
 	else
 	{
 		static IConsoleVariable* GPUCrashDebugging = IConsoleManager::Get().FindConsoleVariable(TEXT("r.GPUCrashDebugging"));
 		if (GPUCrashDebugging)
 		{
+			GPUCrashDebuggingMode = (GPUCrashDebugging->GetInt() > 0) ? ED3D12GPUCrashDebugginMode::Full : ED3D12GPUCrashDebugginMode::Disabled;
+		}
 
-			bGPUCrashDebugging = GPUCrashDebugging->GetInt() > 0;
+		// Still disabled then check the D3D specific cvar for minimal tracking
+		if (GPUCrashDebuggingMode == ED3D12GPUCrashDebugginMode::Disabled)
+		{
+			auto* GPUCrashDebuggingModeVar = IConsoleManager::Get().FindTConsoleVariableDataInt(TEXT("r.D3D12.GPUCrashDebuggingMode"));
+			int32 GPUCrashDebuggingModeValue = GPUCrashDebuggingModeVar ? GPUCrashDebuggingModeVar->GetValueOnAnyThread() : -1;
+			if (GPUCrashDebuggingModeValue >= 0 && GPUCrashDebuggingModeValue <= (int)ED3D12GPUCrashDebugginMode::Full)
+				GPUCrashDebuggingMode = (ED3D12GPUCrashDebugginMode)GPUCrashDebuggingModeValue;
 		}
 	}
 
@@ -193,7 +213,7 @@ void FD3D12Adapter::CreateRootDevice(bool bWithDebug)
 	{
 		// GPUcrash dump handler must be attached prior to device creation
 		static IConsoleVariable* GPUCrashDump = IConsoleManager::Get().FindConsoleVariable(TEXT("r.GPUCrashDump"));
-		if (bGPUCrashDebugging || FParse::Param(FCommandLine::Get(), TEXT("gpucrashdump")) || (GPUCrashDump && GPUCrashDump->GetInt()))
+		if (GPUCrashDebuggingMode == ED3D12GPUCrashDebugginMode::Full || FParse::Param(FCommandLine::Get(), TEXT("gpucrashdump")) || (GPUCrashDump && GPUCrashDump->GetInt()))
 		{
 			HANDLE CurrentThread = ::GetCurrentThread();
 
@@ -215,6 +235,8 @@ void FD3D12Adapter::CreateRootDevice(bool bWithDebug)
 			else
 			{
 				UE_LOG(LogD3D12RHI, Log, TEXT("[Aftermath] Aftermath crash dumping failed to initialize (%x)"), Result);
+
+				GDX12NVAfterMathEnabled = 0;
 			}
 		}
 	}
@@ -245,7 +267,7 @@ void FD3D12Adapter::CreateRootDevice(bool bWithDebug)
 	}
 	
 	// Setup DRED if requested
-	if (bGPUCrashDebugging || FParse::Param(FCommandLine::Get(), TEXT("dred")))
+	if (GPUCrashDebuggingMode == ED3D12GPUCrashDebugginMode::Full || FParse::Param(FCommandLine::Get(), TEXT("dred")))
 	{
 		ID3D12DeviceRemovedExtendedDataSettings* DredSettings = nullptr;
 		HRESULT hr = D3D12GetDebugInterface(IID_PPV_ARGS(&DredSettings));
@@ -392,7 +414,7 @@ void FD3D12Adapter::CreateRootDevice(bool bWithDebug)
 
 #if NV_AFTERMATH
 	// Enable aftermath when GPU crash debugging is enabled
-	if (bGPUCrashDebugging || GDX12NVAfterMathEnabled)
+	if (GPUCrashDebuggingMode == ED3D12GPUCrashDebugginMode::Full && GDX12NVAfterMathEnabled)
 	{
 		if (IsRHIDeviceNVIDIA() && bAllowVendorDevice)
 		{
@@ -442,6 +464,10 @@ void FD3D12Adapter::CreateRootDevice(bool bWithDebug)
 			UE_LOG(LogD3D12RHI, Warning, TEXT("[Aftermath] Skipping aftermath initialization on non-Nvidia device"));
 		}
 	}
+	else
+	{
+		GDX12NVAfterMathEnabled = 0;
+	}
 #endif
 
 #if (PLATFORM_WINDOWS || PLATFORM_HOLOLENS)
diff --git a/Engine/Source/Runtime/D3D12RHI/Private/D3D12Adapter.h b/Engine/Source/Runtime/D3D12RHI/Private/D3D12Adapter.h
--- a/Engine/Source/Runtime/D3D12RHI/Private/D3D12Adapter.h
+++ b/Engine/Source/Runtime/D3D12RHI/Private/D3D12Adapter.h
@@ -66,6 +66,13 @@ struct FD3D12AdapterDesc
 	uint32 NumDeviceNodes;
 };
 
+enum class ED3D12GPUCrashDebugginMode
+{
+	Disabled,
+	Minimal,
+	Full
+};
+
 // Represents a set of linked D3D12 device nodes (LDA i.e 1 or more identical GPUs). In most cases there will be only 1 node, however if the system supports
 // SLI/Crossfire and the app enables it an Adapter will have 2 or more nodes. This class will own anything that can be shared
 // across LDA including: System Pool Memory,.Pipeline State Objects, Root Signatures etc.
@@ -94,7 +101,7 @@ public:
 	FORCEINLINE void SetDeviceRemoved(bool value) { bDeviceRemoved = value; }
 	FORCEINLINE const bool IsDeviceRemoved() const { return bDeviceRemoved; }
 	FORCEINLINE const bool IsDebugDevice() const { return bDebugDevice; }
-	FORCEINLINE const bool IsGPUCrashDebugging() const { return bGPUCrashDebugging; }
+	FORCEINLINE const ED3D12GPUCrashDebugginMode GetGPUCrashDebuggingMode() const { return GPUCrashDebuggingMode; }
 	FORCEINLINE FD3D12DynamicRHI* GetOwningRHI() { return OwningRHI; }
 	FORCEINLINE const D3D12_RESOURCE_HEAP_TIER GetResourceHeapTier() const { return ResourceHeapTier; }
 	FORCEINLINE const D3D12_RESOURCE_BINDING_TIER GetResourceBindingTier() const { return ResourceBindingTier; }
@@ -324,8 +331,8 @@ protected:
 	/** Running with debug device */
 	bool bDebugDevice;
 
-	/** GPU Crash debugging requested */
-	bool bGPUCrashDebugging;
+	/** GPU Crash debugging mode */
+	ED3D12GPUCrashDebugginMode GPUCrashDebuggingMode;
 
 	/** True if the device being used has been removed. */
 	bool bDeviceRemoved;
diff --git a/Engine/Source/Runtime/D3D12RHI/Private/D3D12DirectCommandListManager.cpp b/Engine/Source/Runtime/D3D12RHI/Private/D3D12DirectCommandListManager.cpp
--- a/Engine/Source/Runtime/D3D12RHI/Private/D3D12DirectCommandListManager.cpp
+++ b/Engine/Source/Runtime/D3D12RHI/Private/D3D12DirectCommandListManager.cpp
@@ -393,7 +393,7 @@ void FD3D12CommandListManager::Create(const TCHAR* Name, uint32 NumCommandLists,
 	}
 
 	// setup the bread crumb data to track GPU progress on this command queue when GPU crash debugging is enabled
-	if (Adapter->IsGPUCrashDebugging())
+	if (Adapter->GetGPUCrashDebuggingMode() != ED3D12GPUCrashDebugginMode::Disabled)
 	{		
 		// QI for the ID3DDevice3 - manual buffer write from command line only supported on 1709+
 		TRefCountPtr<ID3D12Device3> D3D12Device3;
diff --git a/Engine/Source/Runtime/D3D12RHI/Private/D3D12Stats.cpp b/Engine/Source/Runtime/D3D12RHI/Private/D3D12Stats.cpp
--- a/Engine/Source/Runtime/D3D12RHI/Private/D3D12Stats.cpp
+++ b/Engine/Source/Runtime/D3D12RHI/Private/D3D12Stats.cpp
@@ -18,7 +18,7 @@ void D3D12RHI::FD3DGPUProfiler::BeginFrame(FD3D12DynamicRHI* InRHI)
 	static auto* CrashCollectionEnableCvar = IConsoleManager::Get().FindTConsoleVariableDataInt(TEXT("r.gpucrash.collectionenable"));
 	static auto* CrashCollectionDataDepth = IConsoleManager::Get().FindTConsoleVariableDataInt(TEXT("r.gpucrash.datadepth"));
 
-	bTrackingGPUCrashData = CrashCollectionEnableCvar ? (InRHI->GetAdapter().IsGPUCrashDebugging() && CrashCollectionEnableCvar->GetValueOnRenderThread() != 0) : false;
+	bTrackingGPUCrashData = CrashCollectionEnableCvar ? (InRHI->GetAdapter().GetGPUCrashDebuggingMode() != ED3D12GPUCrashDebugginMode::Disabled && CrashCollectionEnableCvar->GetValueOnRenderThread() != 0) : false;
 	GPUCrashDataDepth = CrashCollectionDataDepth ? CrashCollectionDataDepth->GetValueOnRenderThread() : -1;
 
 	// latch the bools from the game thread into our private copy
diff --git a/Engine/Source/Runtime/D3D12RHI/Private/D3D12Util.cpp b/Engine/Source/Runtime/D3D12RHI/Private/D3D12Util.cpp
--- a/Engine/Source/Runtime/D3D12RHI/Private/D3D12Util.cpp
+++ b/Engine/Source/Runtime/D3D12RHI/Private/D3D12Util.cpp
@@ -502,7 +502,7 @@ namespace D3D12RHI
 		{
 			ErrorMessage.AppendLine(LOCTEXT("D3D Debug Device", "Use -d3ddebug to enable the D3D debug device."));
 		}
-		if (D3D12RHI->GetAdapter().IsGPUCrashDebugging())
+		if (D3D12RHI->GetAdapter().GetGPUCrashDebuggingMode() != ED3D12GPUCrashDebugginMode::Disabled)
 		{
 			ErrorMessage.AppendLine(LOCTEXT("GPU Crash Debugging enabled", "Check log for GPU state information."));
 		}
@@ -579,6 +579,10 @@ namespace D3D12RHI
 		}
 
 		UE_LOG(LogD3D12RHI, Fatal, TEXT("%s failed \n at %s:%u \n with error %s\n%s"), ANSI_TO_TCHAR(Code), ANSI_TO_TCHAR(Filename), Line, *ErrorString, *Message);
+
+		// Make sure the log is flushed!
+		GLog->PanicFlushThreadedLogs();
+		GLog->Flush();
 	}
 
 	void VerifyD3D12CreateTextureResult(HRESULT D3DResult, const ANSICHAR* Code, const ANSICHAR* Filename, uint32 Line, const D3D12_RESOURCE_DESC& TextureDesc, ID3D12Device* Device)
@@ -629,6 +633,10 @@ namespace D3D12RHI
 			TextureDesc.Format,
 			TextureDesc.MipLevels,
 			*GetD3D12TextureFlagString(TextureDesc.Flags));
+
+		// Make sure the log is flushed!
+		GLog->PanicFlushThreadedLogs();
+		GLog->Flush();
 	}
 
 	void VerifyComRefCount(IUnknown* Object, int32 ExpectedRefs, const TCHAR* Code, const TCHAR* Filename, int32 Line)
diff --git a/Engine/Source/Runtime/Engine/Classes/Engine/RendererSettings.h b/Engine/Source/Runtime/Engine/Classes/Engine/RendererSettings.h
--- a/Engine/Source/Runtime/Engine/Classes/Engine/RendererSettings.h
+++ b/Engine/Source/Runtime/Engine/Classes/Engine/RendererSettings.h
@@ -593,7 +593,7 @@ class ENGINE_API URendererSettings : public UDeveloperSettings
 
 	UPROPERTY(config, EditAnywhere, Category = Debugging, meta = (
 		ConsoleVariable = "r.GPUCrashDebugging", DisplayName = "Enable vendor specific GPU crash analysis tools",
-		ToolTip = "Enables vendor specific GPU crash analysis tools.  Currently only supports NVIDIA Aftermath on DX11.",
+		ToolTip = "Enables vendor specific GPU crash analysis tools.",
 		ConfigRestartRequired = true))
 		uint32 bNvidiaAftermathEnabled : 1;
 
