Enable minimal GPU trackingon DX12 by default in development builds

diff --git a/Engine/Source/Runtime/D3D12RHI/Private/D3D12Adapter.cpp b/Engine/Source/Runtime/D3D12RHI/Private/D3D12Adapter.cpp
--- a/Engine/Source/Runtime/D3D12RHI/Private/D3D12Adapter.cpp
+++ b/Engine/Source/Runtime/D3D12RHI/Private/D3D12Adapter.cpp
@@ -25,14 +25,20 @@ static TAutoConsoleVariable<int32> CVarResidencyManagement(
 
 #if PLATFORM_WINDOWS || PLATFORM_HOLOLENS
 
+// Enabled in debug and development mode while sorting out D3D12 stability issues
+#if UE_BUILD_SHIPPING || UE_BUILD_TEST
 static int32 GD3D12GPUCrashDebuggingMode = 0;
+#else
+static int32 GD3D12GPUCrashDebuggingMode = 1;
+#endif // UE_BUILD_SHIPPING || UE_BUILD_TEST
+
 static TAutoConsoleVariable<int32> CVarD3D12GPUCrashDebuggingMode(
 	TEXT("r.D3D12.GPUCrashDebuggingMode"),
 	GD3D12GPUCrashDebuggingMode,
 	TEXT("Enable GPU crash debugging: tracks the current GPU state and logs information what operations the GPU executed last.\n")
 	TEXT("Optionally generate a GPU crash dump as well (on nVidia hardware only)):\n")
-	TEXT(" 0: GPU crash debugging enabled (default)\n")
-	TEXT(" 1: Minimal overhead GPU crash debugging\n")
+	TEXT(" 0: GPU crash debugging disabled (default in shipping and test builds)\n")
+	TEXT(" 1: Minimal overhead GPU crash debugging (default in development builds)\n")
 	TEXT(" 2: Enable all available GPU crash debugging options (DRED, Aftermath, ...)\n"),
 	ECVF_RenderThreadSafe | ECVF_ReadOnly
 );
