F - Copy from 12762312
Added r.D3D12.GPUTimeout to help on really long GPU frames
#rnx

diff --git a/Engine/Source/Runtime/D3D12RHI/Private/D3D12DirectCommandListManager.cpp b/Engine/Source/Runtime/D3D12RHI/Private/D3D12DirectCommandListManager.cpp
--- a/Engine/Source/Runtime/D3D12RHI/Private/D3D12DirectCommandListManager.cpp
+++ b/Engine/Source/Runtime/D3D12RHI/Private/D3D12DirectCommandListManager.cpp
@@ -4,6 +4,14 @@
 #include "Windows.h"
 
 
+static TAutoConsoleVariable<int32> CVarD3D12GPUTimeout(
+	TEXT("r.D3D12.GPUTimeout"),
+	1,
+	TEXT("0: Disable GPU Timeout; use with care as it could freeze your PC!\n")
+	TEXT("1: Enable GPU Timeout; operation taking long on the GPU will fail(default)\n"),
+	ECVF_ReadOnly
+);
+
 extern bool D3D12RHI_ShouldCreateWithD3DDebug();
 
 FComputeFenceRHIRef FD3D12DynamicRHI::RHICreateComputeFence(const FName& Name)
@@ -374,7 +382,8 @@ void FD3D12CommandListManager::Create(const TCHAR* Name, uint32 NumCommandLists,
 	checkf(NumCommandLists <= 0xffff, TEXT("Exceeded maximum supported command lists"));
 
 	D3D12_COMMAND_QUEUE_DESC CommandQueueDesc = {};
-	CommandQueueDesc.Flags = D3D12_COMMAND_QUEUE_FLAG_NONE;
+	CommandQueueDesc.Flags = (Adapter->IsGPUCrashDebugging() || CVarD3D12GPUTimeout.GetValueOnAnyThread() == 0) 
+		? D3D12_COMMAND_QUEUE_FLAG_DISABLE_GPU_TIMEOUT : D3D12_COMMAND_QUEUE_FLAG_NONE;
 	CommandQueueDesc.NodeMask = GetGPUMask().GetNative();
 	CommandQueueDesc.Priority = Priority;
 	CommandQueueDesc.Type = CommandListType;
