UE4 - Compile fix from merges
#rnx

diff --git a/Engine/Source/Runtime/D3D12RHI/Private/D3D12DirectCommandListManager.cpp b/Engine/Source/Runtime/D3D12RHI/Private/D3D12DirectCommandListManager.cpp
--- a/Engine/Source/Runtime/D3D12RHI/Private/D3D12DirectCommandListManager.cpp
+++ b/Engine/Source/Runtime/D3D12RHI/Private/D3D12DirectCommandListManager.cpp
@@ -381,8 +381,10 @@ void FD3D12CommandListManager::Create(const TCHAR* Name, uint32 NumCommandLists,
 	check(ReadyLists.IsEmpty());
 	checkf(NumCommandLists <= 0xffff, TEXT("Exceeded maximum supported command lists"));
 
+	bool bFullGPUCrashDebugging = (Adapter->GetGPUCrashDebuggingMode() == ED3D12GPUCrashDebugginMode::Full);
+
 	D3D12_COMMAND_QUEUE_DESC CommandQueueDesc = {};
-	CommandQueueDesc.Flags = (Adapter->IsGPUCrashDebugging() || CVarD3D12GPUTimeout.GetValueOnAnyThread() == 0) 
+	CommandQueueDesc.Flags = (bFullGPUCrashDebugging || CVarD3D12GPUTimeout.GetValueOnAnyThread() == 0) 
 		? D3D12_COMMAND_QUEUE_FLAG_DISABLE_GPU_TIMEOUT : D3D12_COMMAND_QUEUE_FLAG_NONE;
 	CommandQueueDesc.NodeMask = GetGPUMask().GetNative();
 	CommandQueueDesc.Priority = Priority;
