Add validation code to make sure the D3D12 command queue hasn't been destroyed yet

diff --git a/Engine/Source/Runtime/D3D12RHI/Private/D3D12Device.cpp b/Engine/Source/Runtime/D3D12RHI/Private/D3D12Device.cpp
--- a/Engine/Source/Runtime/D3D12RHI/Private/D3D12Device.cpp
+++ b/Engine/Source/Runtime/D3D12RHI/Private/D3D12Device.cpp
@@ -329,6 +329,30 @@ void FD3D12Device::UpdateMSAASettings()
 
 void FD3D12Device::Cleanup()
 {
+	const auto ValidateCommandQueue = [this](ED3D12CommandQueueType QueueType, const TCHAR* Name)
+	{
+		ID3D12CommandQueue* CommandQueue = GetD3DCommandQueue(QueueType);
+		if (CommandQueue)
+		{
+			CommandQueue->AddRef();
+			int RefCount = CommandQueue->Release();
+			if (RefCount < 0)
+			{
+				UE_LOG(LogD3D12RHI, Error, TEXT("%s CommandQueue is already destroyed!"), Name);
+			}
+			else if (RefCount > 2)
+			{
+				UE_LOG(LogD3D12RHI, Warning, TEXT("%s CommandQueue is leaking (Refcount %d)"), Name, RefCount);
+			}
+			check(RefCount >= 1);
+		}
+	};
+
+	// Validate that all the D3D command queues are still valid (temp code to check for a shutdown crash)
+	ValidateCommandQueue(ED3D12CommandQueueType::Default, TEXT("Direct"));
+	ValidateCommandQueue(ED3D12CommandQueueType::Copy, TEXT("Copy"));
+	ValidateCommandQueue(ED3D12CommandQueueType::Async, TEXT("Async"));
+
 	// Wait for the command queues to flush
 	CommandListManager->WaitForCommandQueueFlush();
 	CopyCommandListManager->WaitForCommandQueueFlush();
