Variable rate shading support

diff --git a/Engine/Source/Runtime/Core/Public/HAL/Platform.h b/Engine/Source/Runtime/Core/Public/HAL/Platform.h
--- a/Engine/Source/Runtime/Core/Public/HAL/Platform.h
+++ b/Engine/Source/Runtime/Core/Public/HAL/Platform.h
@@ -282,6 +282,9 @@
 #ifndef PLATFORM_SUPPORTS_VIRTUAL_TEXTURES
 	#define PLATFORM_SUPPORTS_VIRTUAL_TEXTURES		0
 #endif
+#ifndef PLATFORM_SUPPORTS_VARIABLE_RATE_SHADING
+	#define PLATFORM_SUPPORTS_VARIABLE_RATE_SHADING		0
+#endif
 #ifndef PLATFORM_REQUIRES_FILESERVER
 	#define PLATFORM_REQUIRES_FILESERVER		0
 #endif
diff --git a/Engine/Source/Runtime/D3D12RHI/Private/D3D12CommandContext.cpp b/Engine/Source/Runtime/D3D12RHI/Private/D3D12CommandContext.cpp
--- a/Engine/Source/Runtime/D3D12RHI/Private/D3D12CommandContext.cpp
+++ b/Engine/Source/Runtime/D3D12RHI/Private/D3D12CommandContext.cpp
@@ -59,6 +59,10 @@ FD3D12CommandContext::FD3D12CommandContext(FD3D12Device* InParent, FD3D12SubAllo
 	bOuterOcclusionQuerySubmitted(false),
 	bDiscardSharedConstants(false),
 	bUsingTessellation(false),
+#if PLATFORM_SUPPORTS_VARIABLE_RATE_SHADING
+	VRSCombiners{ D3D12_SHADING_RATE_COMBINER_PASSTHROUGH, D3D12_SHADING_RATE_COMBINER_PASSTHROUGH } ,
+	VRSShadingRate(D3D12_SHADING_RATE_1X1),
+#endif
 	SkipFastClearEliminateState(D3D12_RESOURCE_STATE_PIXEL_SHADER_RESOURCE | D3D12_RESOURCE_STATE_NON_PIXEL_SHADER_RESOURCE),
 #if PLATFORM_SUPPORTS_VIRTUAL_TEXTURES
 	bNeedFlushTextureCache(false),
diff --git a/Engine/Source/Runtime/D3D12RHI/Private/D3D12CommandContext.h b/Engine/Source/Runtime/D3D12RHI/Private/D3D12CommandContext.h
--- a/Engine/Source/Runtime/D3D12RHI/Private/D3D12CommandContext.h
+++ b/Engine/Source/Runtime/D3D12RHI/Private/D3D12CommandContext.h
@@ -158,6 +158,12 @@ public:
 	/** Set to true when the current shading setup uses tessellation */
 	bool bUsingTessellation;
 
+	/** Used by variable rate shading to cache the current state of the combiners and the constant shading rate*/
+#if PLATFORM_SUPPORTS_VARIABLE_RATE_SHADING
+	D3D12_SHADING_RATE_COMBINER		VRSCombiners[D3D12_RS_SET_SHADING_RATE_COMBINER_COUNT];
+	D3D12_SHADING_RATE				VRSShadingRate;
+#endif
+
 	virtual void FlushMetadata(FRHITexture** InTextures, int32 NumTextures) {};
 
 	D3D12_RESOURCE_STATES SkipFastClearEliminateState;
@@ -227,6 +233,7 @@ public:
 		);
 
 	virtual void SetDepthBounds(float MinDepth, float MaxDepth);
+	virtual void SetShadingRate(EVRSShadingRate ShadingRate, EVRSRateCombiner Combiner);
 
 	virtual void SetAsyncComputeBudgetInternal(EAsyncComputeBudget Budget) {}
 
@@ -287,6 +294,8 @@ public:
 	virtual void RHIDrawIndexedPrimitive(FRHIIndexBuffer* IndexBuffer, int32 BaseVertexIndex, uint32 FirstInstance, uint32 NumVertices, uint32 StartIndex, uint32 NumPrimitives, uint32 NumInstances) final override;
 	virtual void RHIDrawIndexedPrimitiveIndirect(FRHIIndexBuffer* IndexBuffer, FRHIVertexBuffer* ArgumentBuffer, uint32 ArgumentOffset) final override;
 	virtual void RHISetDepthBounds(float MinDepth, float MaxDepth) final override;
+    virtual void RHISetShadingRate(EVRSShadingRate ShadingRate, EVRSRateCombiner Combiner) final override;
+    virtual void RHISetShadingRateImage(FRHITexture* RateImageTexture, EVRSRateCombiner Combiner) final override;
 	virtual void RHIUpdateTextureReference(FRHITextureReference* TextureRef, FRHITexture* NewTexture) final override;
 
 	virtual void RHIClearMRTImpl(bool bClearColor, int32 NumClearColors, const FLinearColor* ColorArray, bool bClearDepth, float Depth, bool bClearStencil, uint32 Stencil);
@@ -665,6 +674,16 @@ public:
 	{
 		ContextRedirect(RHISetDepthBounds(MinDepth, MaxDepth));
 	}
+	
+	FORCEINLINE virtual void RHISetShadingRate(EVRSShadingRate ShadingRate, EVRSRateCombiner Combiner) final override
+	{
+		ContextRedirect(RHISetShadingRate(ShadingRate, Combiner));
+	}
+
+	FORCEINLINE virtual void RHISetShadingRateImage(FRHITexture* RateImageTexture, EVRSRateCombiner Combiner) final override
+	{
+		ContextRedirect(RHISetShadingRateImage(RateImageTexture, Combiner));
+	}
 
 	FORCEINLINE virtual void RHIUpdateTextureReference(FRHITextureReference* TextureRef, FRHITexture* NewTexture) final override
 	{
diff --git a/Engine/Source/Runtime/D3D12RHI/Private/D3D12CommandList.cpp b/Engine/Source/Runtime/D3D12RHI/Private/D3D12CommandList.cpp
--- a/Engine/Source/Runtime/D3D12RHI/Private/D3D12CommandList.cpp
+++ b/Engine/Source/Runtime/D3D12RHI/Private/D3D12CommandList.cpp
@@ -68,6 +68,10 @@ FD3D12CommandListHandle::FD3D12CommandListData::FD3D12CommandListData(FD3D12Devi
 	CommandList->QueryInterface(IID_PPV_ARGS(CommandList2.GetInitReference()));
 #endif
 
+#if PLATFORM_SUPPORTS_VARIABLE_RATE_SHADING
+	CommandList->QueryInterface(IID_PPV_ARGS(CommandList5.GetInitReference()));
+#endif
+
 #if D3D12_RHI_RAYTRACING
 	// Obtain ID3D12CommandListRaytracingPrototype if parent device supports ray tracing and this is a compatible command list type (compute or graphics).
 	if (ParentDevice->GetRayTracingDevice() && (InCommandListType == D3D12_COMMAND_LIST_TYPE_DIRECT || InCommandListType == D3D12_COMMAND_LIST_TYPE_COMPUTE))
diff --git a/Engine/Source/Runtime/D3D12RHI/Private/D3D12CommandList.h b/Engine/Source/Runtime/D3D12RHI/Private/D3D12CommandList.h
--- a/Engine/Source/Runtime/D3D12RHI/Private/D3D12CommandList.h
+++ b/Engine/Source/Runtime/D3D12RHI/Private/D3D12CommandList.h
@@ -207,6 +207,11 @@ private:
 		TRefCountPtr<ID3D12GraphicsCommandList1> CommandList1;
 		TRefCountPtr<ID3D12GraphicsCommandList2> CommandList2;
 #endif
+
+#if PLATFORM_SUPPORTS_VARIABLE_RATE_SHADING
+		TRefCountPtr<ID3D12GraphicsCommandList5> CommandList5;
+#endif
+
 #if D3D12_RHI_RAYTRACING
 		TRefCountPtr<ID3D12GraphicsCommandList4> RayTracingCommandList;
 #endif // D3D12_RHI_RAYTRACING
@@ -421,6 +426,14 @@ public:
 	}
 #endif
 
+#if PLATFORM_SUPPORTS_VARIABLE_RATE_SHADING
+	ID3D12GraphicsCommandList5* GraphicsCommandList5() const
+	{
+		check(CommandListData && (CommandListData->CommandListType == D3D12_COMMAND_LIST_TYPE_DIRECT));
+		return CommandListData->CommandList5.GetReference();
+	}
+#endif
+
 #if D3D12_RHI_RAYTRACING
 	ID3D12GraphicsCommandList4* RayTracingCommandList() const
 	{
diff --git a/Engine/Source/Runtime/D3D12RHI/Private/D3D12Commands.cpp b/Engine/Source/Runtime/D3D12RHI/Private/D3D12Commands.cpp
--- a/Engine/Source/Runtime/D3D12RHI/Private/D3D12Commands.cpp
+++ b/Engine/Source/Runtime/D3D12RHI/Private/D3D12Commands.cpp
@@ -504,6 +504,11 @@ void FD3D12CommandContext::RHISetGraphicsPipelineState(FRHIGraphicsPipelineState
 		StateCache.SetDepthBounds(0.0f, 1.0f);
 	}
 
+	if (GRHISupportsVariableRateShading  )
+	{
+		StateCache.SetShadingRate(GraphicsPipelineState->PipelineStateInitializer.ShadingRate, VRSRB_Passthrough);
+	}
+
 	StateCache.SetGraphicsPipelineState(GraphicsPipelineState, bUsingTessellation != bWasUsingTessellation);
 	StateCache.SetStencilRef(0);
 
@@ -1836,6 +1841,67 @@ void FD3D12CommandContext::SetDepthBounds(float MinDepth, float MaxDepth)
 #endif
 }
 
+
+#if PLATFORM_SUPPORTS_VARIABLE_RATE_SHADING
+static D3D12_SHADING_RATE_COMBINER ConvertShadingRateCombiner(EVRSRateCombiner InCombiner)
+{
+	switch (InCombiner)
+	{
+	case VRSRB_Override:
+		return D3D12_SHADING_RATE_COMBINER_OVERRIDE;
+	case VRSRB_Min:
+		return D3D12_SHADING_RATE_COMBINER_MIN;
+	case VRSRB_Max:
+		return D3D12_SHADING_RATE_COMBINER_MAX;
+	case VRSRB_Sum:
+		return D3D12_SHADING_RATE_COMBINER_SUM;
+	case VRSRB_Passthrough:
+	default:
+		return D3D12_SHADING_RATE_COMBINER_PASSTHROUGH;
+	}
+	return D3D12_SHADING_RATE_COMBINER_PASSTHROUGH;
+}
+#endif
+
+void FD3D12CommandContext::RHISetShadingRate(EVRSShadingRate ShadingRate, EVRSRateCombiner Combiner)
+{
+#if PLATFORM_SUPPORTS_VARIABLE_RATE_SHADING
+	StateCache.SetShadingRate(ShadingRate, Combiner);
+#endif
+}
+
+ void FD3D12CommandContext::SetShadingRate(EVRSShadingRate ShadingRate, EVRSRateCombiner Combiner)
+ {
+ #if PLATFORM_SUPPORTS_VARIABLE_RATE_SHADING
+ 	if (GRHISupportsVariableRateShading && CommandListHandle.GraphicsCommandList5())
+ 	{
+ 		VRSCombiners[0] = ConvertShadingRateCombiner(Combiner);	// Combiner 0 is used to mix per draw and per VS/GS rates
+ 		VRSShadingRate = static_cast<D3D12_SHADING_RATE>(ShadingRate);
+ 		CommandListHandle.GraphicsCommandList5()->RSSetShadingRate(VRSShadingRate, VRSCombiners);
+ 	}
+ #endif
+}
+
+void FD3D12CommandContext::RHISetShadingRateImage(FRHITexture* RateImageTexture, EVRSRateCombiner Combiner)
+{
+#if PLATFORM_SUPPORTS_VARIABLE_RATE_SHADING
+	if (GRHISupportsVariableRateShading && CommandListHandle.GraphicsCommandList5())
+	{
+		VRSCombiners[1] = ConvertShadingRateCombiner(Combiner); // Combiner 1 is used to mix rates from a texture and the previous combiner
+		if (RateImageTexture)
+		{
+			FD3D12Resource* Resource = RetrieveTextureBase(RateImageTexture)->GetResource();
+			CommandListHandle.GraphicsCommandList5()->RSSetShadingRateImage(Resource->GetResource());
+		}
+		else
+		{
+			CommandListHandle.GraphicsCommandList5()->RSSetShadingRateImage(nullptr);
+		}
+		CommandListHandle.GraphicsCommandList5()->RSSetShadingRate(VRSShadingRate, VRSCombiners);
+	}
+#endif
+}
+
 void FD3D12CommandContext::RHISubmitCommandsHint()
 {
 	// Resolve any timestamp queries so far (if any).
diff --git a/Engine/Source/Runtime/D3D12RHI/Private/D3D12StateCache.cpp b/Engine/Source/Runtime/D3D12RHI/Private/D3D12StateCache.cpp
--- a/Engine/Source/Runtime/D3D12RHI/Private/D3D12StateCache.cpp
+++ b/Engine/Source/Runtime/D3D12RHI/Private/D3D12StateCache.cpp
@@ -263,6 +263,7 @@ void FD3D12StateCacheBase::DirtyState()
 	bNeedSetBlendFactor = true;
 	bNeedSetStencilRef = true;
 	bNeedSetDepthBounds = GSupportsDepthBoundsTest;
+	bNeedSetShadingRate = GRHISupportsVariableRateShading;
 	PipelineState.Common.SRVCache.DirtyAll();
 	PipelineState.Common.UAVCache.DirtyAll();
 	PipelineState.Common.CBVCache.DirtyAll();
@@ -458,6 +459,12 @@ void FD3D12StateCacheBase::ApplyState()
 			bNeedSetDepthBounds = false;
 			CmdContext->SetDepthBounds(PipelineState.Graphics.MinDepth, PipelineState.Graphics.MaxDepth);
 		}
+		
+		if (bNeedSetShadingRate)
+		{
+			bNeedSetShadingRate = false;
+			CmdContext->SetShadingRate(PipelineState.Graphics.DrawShadingRate, PipelineState.Graphics.Combiner);
+		}
 	}
 
 	// Note that ray tracing pipeline shares state with compute
diff --git a/Engine/Source/Runtime/D3D12RHI/Private/D3D12StateCachePrivate.h b/Engine/Source/Runtime/D3D12RHI/Private/D3D12StateCachePrivate.h
--- a/Engine/Source/Runtime/D3D12RHI/Private/D3D12StateCachePrivate.h
+++ b/Engine/Source/Runtime/D3D12RHI/Private/D3D12StateCachePrivate.h
@@ -347,6 +347,7 @@ protected:
 	bool bNeedSetBlendFactor;
 	bool bNeedSetStencilRef;
 	bool bNeedSetDepthBounds;
+	bool bNeedSetShadingRate;
 	bool bAutoFlushComputeShaderCache;
 	D3D12_RESOURCE_BINDING_TIER ResourceBindingTier;
 
@@ -403,6 +404,9 @@ protected:
 
 			float MinDepth;
 			float MaxDepth;
+
+			EVRSShadingRate  DrawShadingRate;
+			EVRSRateCombiner Combiner;
 		} Graphics;
 
 		struct
@@ -964,6 +968,21 @@ public:
 		}
 	}
 
+	void SetShadingRate(EVRSShadingRate ShadingRate, EVRSRateCombiner Combiner)
+	{
+		if (PipelineState.Graphics.DrawShadingRate != ShadingRate )
+		{
+			PipelineState.Graphics.DrawShadingRate = ShadingRate;
+			bNeedSetShadingRate = GRHISupportsVariableRateShading;
+		}
+
+		if (PipelineState.Graphics.Combiner != Combiner)
+		{
+			PipelineState.Graphics.Combiner = Combiner;
+			bNeedSetShadingRate = GRHISupportsVariableRateShading;
+		}
+	}
+
 	void SetComputeBudget(EAsyncComputeBudget ComputeBudget)
 	{
 		PipelineState.Compute.ComputeBudget = ComputeBudget;
diff --git a/Engine/Source/Runtime/EmptyRHI/Public/EmptyRHI.h b/Engine/Source/Runtime/EmptyRHI/Public/EmptyRHI.h
--- a/Engine/Source/Runtime/EmptyRHI/Public/EmptyRHI.h
+++ b/Engine/Source/Runtime/EmptyRHI/Public/EmptyRHI.h
@@ -180,6 +180,8 @@ public:
 	virtual void RHIDrawIndexedPrimitiveIndirect(uint32 PrimitiveType, FRHIIndexBuffer* IndexBuffer, FRHIVertexBuffer* ArgumentBuffer, uint32 ArgumentOffset) final override;
 	virtual void RHIEnableDepthBoundsTest(bool bEnable) final override;
 	virtual void RHISetDepthBounds(float MinDepth, float MaxDepth) final override;
+	virtual void RHISetShadingRate(EVRSShadingRate ShadingRate, EVRSRateCombiner Combiner) final override;
+	virtual void RHISetShadingRateImage(FRHITexture* RateImageTexture, EVRSRateCombiner Combiner)  final override;
 	virtual void RHIPushEvent(const TCHAR* Name, FColor Color) final override;
 	virtual void RHIPopEvent() final override;
 };
diff --git a/Engine/Source/Runtime/Engine/Classes/Engine/EngineTypes.h b/Engine/Source/Runtime/Engine/Classes/Engine/EngineTypes.h
--- a/Engine/Source/Runtime/Engine/Classes/Engine/EngineTypes.h
+++ b/Engine/Source/Runtime/Engine/Classes/Engine/EngineTypes.h
@@ -589,6 +589,17 @@ enum EMaterialStencilCompare
 	MSC_Count			UMETA(Hidden),
 };
 
+UENUM()
+enum EMaterialShadingRate
+{
+	MSR_1x1				UMETA(DisplayName = "1x1"),
+	MSR_2x1				UMETA(DisplayName = "2x1"),
+	MSR_1x2				UMETA(DisplayName = "1x2"),
+	MSR_2x2				UMETA(DisplayName = "2x2"),
+	MSR_Count			UMETA(Hidden),
+};
+
+
 /**	Lighting build quality enumeration */
 UENUM()
 enum ELightingBuildQuality
diff --git a/Engine/Source/Runtime/Engine/Classes/Materials/Material.h b/Engine/Source/Runtime/Engine/Classes/Materials/Material.h
--- a/Engine/Source/Runtime/Engine/Classes/Materials/Material.h
+++ b/Engine/Source/Runtime/Engine/Classes/Materials/Material.h
@@ -795,6 +795,10 @@ public:
 	UPROPERTY(EditAnywhere, Category=Material, AdvancedDisplay)
 	uint8 Wireframe : 1;
 
+	/** Select what shading rate to apply for platforms that have variable rate shading */
+	UPROPERTY(EditAnywhere, Category = Material, AdvancedDisplay)
+	TEnumAsByte<EMaterialShadingRate> ShadingRate;
+
 #if WITH_EDITORONLY_DATA
 	UPROPERTY()
 	int32 EditorX;
diff --git a/Engine/Source/Runtime/Engine/Private/Materials/MaterialShared.cpp b/Engine/Source/Runtime/Engine/Private/Materials/MaterialShared.cpp
--- a/Engine/Source/Runtime/Engine/Private/Materials/MaterialShared.cpp
+++ b/Engine/Source/Runtime/Engine/Private/Materials/MaterialShared.cpp
@@ -1040,6 +1040,7 @@ bool FMaterialResource::IsSpecialEngineMaterial() const { return Material->bUsed
 bool FMaterialResource::HasVertexPositionOffsetConnected() const { return HasMaterialAttributesConnected() || (!Material->bUseMaterialAttributes && Material->WorldPositionOffset.IsConnected()); }
 bool FMaterialResource::HasPixelDepthOffsetConnected() const { return HasMaterialAttributesConnected() || (!Material->bUseMaterialAttributes && Material->PixelDepthOffset.IsConnected()); }
 bool FMaterialResource::HasMaterialAttributesConnected() const { return Material->bUseMaterialAttributes && Material->MaterialAttributes.IsConnected(); }
+EMaterialShadingRate FMaterialResource::GetShadingRate() const { return Material->ShadingRate; }
 FString FMaterialResource::GetBaseMaterialPathName() const { return Material->GetPathName(); }
 FString FMaterialResource::GetDebugName() const
 {
diff --git a/Engine/Source/Runtime/Engine/Public/MaterialShared.h b/Engine/Source/Runtime/Engine/Public/MaterialShared.h
--- a/Engine/Source/Runtime/Engine/Public/MaterialShared.h
+++ b/Engine/Source/Runtime/Engine/Public/MaterialShared.h
@@ -1612,7 +1612,7 @@ public:
 	virtual uint32 GetStencilCompare() const { return 0; }
 	virtual bool HasRuntimeVirtualTextureOutput() const { return false; }
 	virtual bool CastsRayTracedShadows() const { return true; }
-
+	virtual EMaterialShadingRate GetShadingRate() const { return MSR_1x1; }
 	/**
 	 * Should shaders compiled for this material be saved to disk?
 	 */
@@ -2341,6 +2341,7 @@ public:
 	ENGINE_API virtual bool IsDeferredDecal() const override;
 	ENGINE_API virtual bool IsVolumetricPrimitive() const override;
 	ENGINE_API virtual bool IsWireframe() const override;
+	ENGINE_API virtual EMaterialShadingRate  GetShadingRate() const override;
 	ENGINE_API virtual bool IsUIMaterial() const override;
 	ENGINE_API virtual bool IsSpecialEngineMaterial() const override;
 	ENGINE_API virtual bool IsUsedWithSkeletalMesh() const override;
diff --git a/Engine/Source/Runtime/NullDrv/Public/NullRHI.h b/Engine/Source/Runtime/NullDrv/Public/NullRHI.h
--- a/Engine/Source/Runtime/NullDrv/Public/NullRHI.h
+++ b/Engine/Source/Runtime/NullDrv/Public/NullRHI.h
@@ -694,6 +694,12 @@ public:
 	virtual void RHISetDepthBounds(float MinDepth, float MaxDepth) final override
 	{
 	}
+	virtual void RHISetShadingRate(EVRSShadingRate ShadingRate, EVRSRateCombiner Combiner) final override
+	{
+	}
+	virtual void RHISetShadingRateImage(FRHITexture* RateImageTexture, EVRSRateCombiner Combiner)  final override
+	{
+	}
 	virtual void* RHIGetNativeDevice() final override
 	{ 
 		return 0; 
diff --git a/Engine/Source/Runtime/RHI/Private/RHI.cpp b/Engine/Source/Runtime/RHI/Private/RHI.cpp
--- a/Engine/Source/Runtime/RHI/Private/RHI.cpp
+++ b/Engine/Source/Runtime/RHI/Private/RHI.cpp
@@ -751,6 +751,11 @@ bool GRHISupportsBackBufferWithCustomDepthStencil = true;
 
 bool GRHIIsHDREnabled = false;
 bool GRHISupportsHDROutput = false;
+
+bool GRHISupportsVariableRateShading = false;
+int32 GVariableRateShadingImageTileSize = 0;
+int32 GVariableRateShadingTier = 0;
+
 EPixelFormat GRHIHDRDisplayOutputFormat = PF_FloatRGBA;
 
 uint64 GRHIPresentCounter = 1;
diff --git a/Engine/Source/Runtime/RHI/Public/RHI.h b/Engine/Source/Runtime/RHI/Public/RHI.h
--- a/Engine/Source/Runtime/RHI/Public/RHI.h
+++ b/Engine/Source/Runtime/RHI/Public/RHI.h
@@ -572,6 +572,15 @@ extern RHI_API bool GRHIIsHDREnabled;
 /** Whether the present adapter/display offers HDR output capabilities. */
 extern RHI_API bool GRHISupportsHDROutput;
 
+/** Whether or not the RHI can support Variable Rate Shading. */
+extern RHI_API bool GRHISupportsVariableRateShading;
+
+/** Size of the tiles in a screen space texture that can be used to drive Variable Rate Shading. */
+extern RHI_API int32 GVariableRateShadingImageTileSize;
+
+/** What tier of Variable Rate Shading the HW supports */
+extern RHI_API int32 GVariableRateShadingTier;
+
 /** Format used for the backbuffer when outputting to a HDR display. */
 extern RHI_API EPixelFormat GRHIHDRDisplayOutputFormat;
 
diff --git a/Engine/Source/Runtime/RHI/Public/RHICommandList.h b/Engine/Source/Runtime/RHI/Public/RHICommandList.h
--- a/Engine/Source/Runtime/RHI/Public/RHICommandList.h
+++ b/Engine/Source/Runtime/RHI/Public/RHICommandList.h
@@ -1443,6 +1443,32 @@ FRHICOMMAND_MACRO(FRHICommandSetDepthBounds)
 	RHI_API void Execute(FRHICommandListBase& CmdList);
 };
 
+FRHICOMMAND_MACRO(FRHICommandSetShadingRate)
+{
+	EVRSShadingRate   ShadingRate;
+	EVRSRateCombiner  Combiner;
+
+	FORCEINLINE_DEBUGGABLE FRHICommandSetShadingRate(EVRSShadingRate InShadingRate, EVRSRateCombiner InCombiner)
+		: ShadingRate(InShadingRate),
+		Combiner(InCombiner)
+	{
+	}
+	RHI_API void Execute(FRHICommandListBase& CmdList);
+};
+
+FRHICOMMAND_MACRO(FRHICommandSetShadingRateImage)
+{
+	FRHITexture* RateImageTexture;
+	EVRSRateCombiner  Combiner;
+
+	FORCEINLINE_DEBUGGABLE FRHICommandSetShadingRateImage(FRHITexture* InRateImageTexture, EVRSRateCombiner InCombiner)
+		: RateImageTexture(InRateImageTexture),
+		Combiner(InCombiner)
+	{
+	}
+	RHI_API void Execute(FRHICommandListBase& CmdList);
+};
+
 FRHICOMMAND_MACRO(FRHICommandClearUAVFloat)
 {
 	FRHIUnorderedAccessView* UnorderedAccessViewRHI;
@@ -3273,6 +3299,30 @@ public:
 		}
 		ALLOC_COMMAND(FRHICommandSetDepthBounds)(MinDepth, MaxDepth);
 	}
+	
+	FORCEINLINE_DEBUGGABLE void SetShadingRate(EVRSShadingRate ShadingRate, EVRSRateCombiner Combiner)
+	{
+#if PLATFORM_SUPPORTS_VARIABLE_RATE_SHADING
+		if (Bypass())
+		{
+			GetContext().RHISetShadingRate(ShadingRate, Combiner);
+			return;
+		}
+		ALLOC_COMMAND(FRHICommandSetShadingRate)(ShadingRate, Combiner);
+#endif
+	}
+
+	FORCEINLINE_DEBUGGABLE void SetShadingRateImage(FRHITexture* RateImageTexture, EVRSRateCombiner Combiner)
+	{
+#if PLATFORM_SUPPORTS_VARIABLE_RATE_SHADING
+		if (Bypass())
+		{
+			GetContext().RHISetShadingRateImage(RateImageTexture, Combiner);
+			return;
+		}
+		ALLOC_COMMAND(FRHICommandSetShadingRateImage)(RateImageTexture, Combiner);
+#endif
+	}
 
 	FORCEINLINE_DEBUGGABLE void CopyToResolveTarget(FRHITexture* SourceTextureRHI, FRHITexture* DestTextureRHI, const FResolveParams& ResolveParams)
 	{
diff --git a/Engine/Source/Runtime/RHI/Public/RHICommandListCommandExecutes.inl b/Engine/Source/Runtime/RHI/Public/RHICommandListCommandExecutes.inl
--- a/Engine/Source/Runtime/RHI/Public/RHICommandListCommandExecutes.inl
+++ b/Engine/Source/Runtime/RHI/Public/RHICommandListCommandExecutes.inl
@@ -360,6 +360,18 @@ void FRHICommandDrawIndexedPrimitiveIndirect::Execute(FRHICommandListBase& CmdLi
 	INTERNAL_DECORATOR(RHIDrawIndexedPrimitiveIndirect)(IndexBuffer, ArgumentsBuffer, ArgumentOffset);
 }
 
+void FRHICommandSetShadingRate::Execute(FRHICommandListBase& CmdList)
+{
+	RHISTAT(SetShadingRate);
+	INTERNAL_DECORATOR(RHISetShadingRate)(ShadingRate, Combiner);
+}
+
+void FRHICommandSetShadingRateImage::Execute(FRHICommandListBase& CmdList)
+{
+	RHISTAT(SetShadingRateImage);
+	INTERNAL_DECORATOR(RHISetShadingRateImage)(RateImageTexture, Combiner);
+}
+
 void FRHICommandSetDepthBounds::Execute(FRHICommandListBase& CmdList)
 {
 	RHISTAT(EnableDepthBoundsTest);
diff --git a/Engine/Source/Runtime/RHI/Public/RHIContext.h b/Engine/Source/Runtime/RHI/Public/RHIContext.h
--- a/Engine/Source/Runtime/RHI/Public/RHIContext.h
+++ b/Engine/Source/Runtime/RHI/Public/RHIContext.h
@@ -608,6 +608,16 @@ public:
 	*/
 	virtual void RHISetDepthBounds(float MinDepth, float MaxDepth) = 0;
 
+	virtual void RHISetShadingRate(EVRSShadingRate ShadingRate, EVRSRateCombiner Combiner)
+	{
+		/* empty default implementation */
+	} 
+	
+	virtual void RHISetShadingRateImage(FRHITexture* RateImageTexture, EVRSRateCombiner Combiner)
+	{
+		/* empty default implementation */
+	}
+
 	virtual void RHIUpdateTextureReference(FRHITextureReference* TextureRef, FRHITexture* NewTexture) = 0;
 
 	virtual void RHIBeginRenderPass(const FRHIRenderPassInfo& InInfo, const TCHAR* InName)
diff --git a/Engine/Source/Runtime/RHI/Public/RHIDefinitions.h b/Engine/Source/Runtime/RHI/Public/RHIDefinitions.h
--- a/Engine/Source/Runtime/RHI/Public/RHIDefinitions.h
+++ b/Engine/Source/Runtime/RHI/Public/RHIDefinitions.h
@@ -1063,6 +1063,33 @@ enum EPrimitiveType
 static_assert(PT_Num <= (1 << 8), "EPrimitiveType doesn't fit in a byte");
 static_assert(PT_Num <= (1 << PT_NumBits), "PT_NumBits is too small");
 
+enum EVRSAxisShadingRate : uint8
+{
+	VRSASR_1X = 0x0,
+	VRSASR_2X = 0x1,
+	VRSASR_4X = 0x2,
+};
+
+enum EVRSShadingRate : uint8
+{
+	VRSSR_1x1  = (VRSASR_1X << 2) + VRSASR_1X,
+	VRSSR_1x2  = (VRSASR_1X << 2) + VRSASR_2X,
+	VRSSR_2x1  = (VRSASR_2X << 2) + VRSASR_1X,
+	VRSSR_2x2  = (VRSASR_2X << 2) + VRSASR_2X,
+	VRSSR_2x4  = (VRSASR_2X << 2) + VRSASR_4X,
+	VRSSR_4x2  = (VRSASR_4X << 2) + VRSASR_2X,
+	VRSSR_4x4  = (VRSASR_4X << 2) + VRSASR_4X,
+};
+
+enum EVRSRateCombiner : uint8
+{
+	VRSRB_Passthrough,
+	VRSRB_Override,
+	VRSRB_Min,
+	VRSRB_Max,
+	VRSRB_Sum,
+};
+
 /**
  *	Resource usage flags - for vertex and index buffers.
  */
diff --git a/Engine/Source/Runtime/RHI/Public/RHIResources.h b/Engine/Source/Runtime/RHI/Public/RHIResources.h
--- a/Engine/Source/Runtime/RHI/Public/RHIResources.h
+++ b/Engine/Source/Runtime/RHI/Public/RHIResources.h
@@ -1973,6 +1973,7 @@ public:
 		, bDepthBounds(false)
 		, bMultiView(false)
 		, bHasFragmentDensityAttachment(false)
+		, ShadingRate(EVRSShadingRate::VRSSR_1x1)
 		, Flags(0)
 	{
 		static_assert(sizeof(EPixelFormat) != sizeof(uint8), "Change TRenderTargetFormats's uint8 to EPixelFormat");
@@ -2002,8 +2003,8 @@ public:
 		uint16						InFlags,
 		bool						bInDepthBounds,
 		bool						bInMultiView,
-		bool						bHasFragmentDensityAttachment
-		)
+		bool						bHasFragmentDensityAttachment,
+		EVRSShadingRate				InShadingRate)
 		: BoundShaderState(InBoundShaderState)
 		, BlendState(InBlendState)
 		, RasterizerState(InRasterizerState)
@@ -2026,6 +2027,7 @@ public:
 		, bDepthBounds(bInDepthBounds)
 		, bMultiView(bInMultiView)
 		, bHasFragmentDensityAttachment(bHasFragmentDensityAttachment)
+		, ShadingRate(InShadingRate)
 		, Flags(InFlags)
 	{
 	}
@@ -2112,6 +2114,7 @@ public:
 	bool							bDepthBounds;
 	bool							bMultiView;
 	bool							bHasFragmentDensityAttachment;
+	EVRSShadingRate					ShadingRate;
 	
 	// Note: these flags do NOT affect compilation of this PSO.
 	// The resulting object is invariant with respect to whatever is set here, they are
diff --git a/Engine/Source/Runtime/RHI/Public/RHIValidationContext.h b/Engine/Source/Runtime/RHI/Public/RHIValidationContext.h
--- a/Engine/Source/Runtime/RHI/Public/RHIValidationContext.h
+++ b/Engine/Source/Runtime/RHI/Public/RHIValidationContext.h
@@ -620,6 +620,16 @@ public:
 		RHIContext->RHISetDepthBounds(MinDepth, MaxDepth);
 	}
 
+	virtual void RHISetShadingRate(EVRSShadingRate ShadingRate, EVRSRateCombiner Combiner) override final
+	{
+		RHIContext->RHISetShadingRate(ShadingRate, Combiner);
+	}
+
+	virtual void RHISetShadingRateImage(FRHITexture* RateImageTexture, EVRSRateCombiner Combiner) override final
+	{
+		RHIContext->RHISetShadingRateImage(RateImageTexture, Combiner);
+	}
+
 	virtual void RHIPushEvent(const TCHAR* Name, FColor Color) override final
 	{
 		RHIContext->RHIPushEvent(Name, Color);
diff --git a/Engine/Source/Runtime/Renderer/Public/MeshPassProcessor.h b/Engine/Source/Runtime/Renderer/Public/MeshPassProcessor.h
--- a/Engine/Source/Runtime/Renderer/Public/MeshPassProcessor.h
+++ b/Engine/Source/Runtime/Renderer/Public/MeshPassProcessor.h
@@ -242,6 +242,7 @@ public:
 		, DepthStencilState(InMinimalState.DepthStencilState)
 		, ImmutableSamplerState(InMinimalState.ImmutableSamplerState)
 		, bDepthBounds(InMinimalState.bDepthBounds)
+		, DrawShadingRate(InMinimalState.DrawShadingRate)
 		, PrimitiveType(InMinimalState.PrimitiveType)
 	{
 	}
@@ -272,6 +273,7 @@ public:
 			, bDepthBounds
 			, bMultiView
 			, bHasFragmentDensityAttachment
+			, DrawShadingRate
 		);
 	}
 
@@ -299,6 +301,7 @@ public:
 			bDepthBounds != rhs.bDepthBounds ||
 			bMultiView != rhs.bMultiView ||
 			bHasFragmentDensityAttachment != rhs.bHasFragmentDensityAttachment ||
+			DrawShadingRate != rhs.DrawShadingRate ||
 			PrimitiveType != rhs.PrimitiveType)
 		{
 			return false;
@@ -366,6 +369,7 @@ public:
 			COMPARE_FIELD(bDepthBounds)
 			COMPARE_FIELD(bMultiView)
 			COMPARE_FIELD(bHasFragmentDensityAttachment)
+			COMPARE_FIELD(DrawShadingRate)
 			COMPARE_FIELD(PrimitiveType)
 		COMPARE_FIELD_END;
 
@@ -397,6 +401,7 @@ public:
 			COMPARE_FIELD(bDepthBounds)
 			COMPARE_FIELD(bMultiView)
 			COMPARE_FIELD(bHasFragmentDensityAttachment)
+			COMPARE_FIELD(DrawShadingRate)
 			COMPARE_FIELD(PrimitiveType)
 			COMPARE_FIELD_END;
 
@@ -421,7 +426,7 @@ public:
 	bool							bDepthBounds = false;
 	bool							bMultiView = false;
 	bool							bHasFragmentDensityAttachment = false;
-	uint8							Padding[1] = {};
+	EVRSShadingRate					DrawShadingRate  = EVRSShadingRate::VRSSR_1x1;
 
 	EPrimitiveType			PrimitiveType;
 };
diff --git a/Engine/Source/Runtime/Renderer/Public/MeshPassProcessor.inl b/Engine/Source/Runtime/Renderer/Public/MeshPassProcessor.inl
--- a/Engine/Source/Runtime/Renderer/Public/MeshPassProcessor.inl
+++ b/Engine/Source/Runtime/Renderer/Public/MeshPassProcessor.inl
@@ -6,6 +6,23 @@ MeshPassProcessor.inl:
 
 #pragma once
 
+static EVRSShadingRate GetShadingRateFromMaterial(EMaterialShadingRate MaterialShadingRate)
+{
+	if (GRHISupportsVariableRateShading)
+	{
+		switch (MaterialShadingRate)
+		{
+		case MSR_1x2:
+			return EVRSShadingRate::VRSSR_1x2;
+		case MSR_2x1:
+			return EVRSShadingRate::VRSSR_2x1;
+		case MSR_2x2:
+			return EVRSShadingRate::VRSSR_2x2;
+		}
+	}
+	return EVRSShadingRate::VRSSR_1x1;
+}
+
 template<typename PassShadersType, typename ShaderElementDataType>
 void FMeshPassProcessor::BuildMeshDrawCommands(
 	const FMeshBatch& RESTRICT MeshBatch,
@@ -50,6 +67,7 @@ void FMeshPassProcessor::BuildMeshDrawCommands(
 
 	PipelineState.BlendState = DrawRenderState.GetBlendState();
 	PipelineState.DepthStencilState = DrawRenderState.GetDepthStencilState();
+	PipelineState.DrawShadingRate = GetShadingRateFromMaterial(MaterialResource.GetShadingRate());
 
 	check(VertexFactory && VertexFactory->IsInitialized());
 	VertexFactory->GetStreams(FeatureLevel, InputStreamType, SharedMeshDrawCommand.VertexStreams);
