Implement ray tracing bottom-level acceleration structure compaction for static (!bFastBuild && !bAllowUpdate) meshes

diff --git a/Engine/Source/Runtime/D3D12RHI/Private/D3D12RayTracing.cpp b/Engine/Source/Runtime/D3D12RHI/Private/D3D12RayTracing.cpp
--- a/Engine/Source/Runtime/D3D12RHI/Private/D3D12RayTracing.cpp
+++ b/Engine/Source/Runtime/D3D12RHI/Private/D3D12RayTracing.cpp
@@ -1949,6 +1949,11 @@ FRayTracingGeometryRHIRef FD3D12DynamicRHI::RHICreateRayTracingGeometry(const FR
 		Geometry->BuildFlags |= D3D12_RAYTRACING_ACCELERATION_STRUCTURE_BUILD_FLAG_ALLOW_UPDATE;
 	}
 
+	if (!Initializer.bFastBuild && !Initializer.bAllowUpdate)
+	{
+		Geometry->BuildFlags |= D3D12_RAYTRACING_ACCELERATION_STRUCTURE_BUILD_FLAG_ALLOW_COMPACTION;
+	}
+
 	if (GRayTracingDebugForceBuildMode == 1)
 	{
 		Geometry->BuildFlags |= D3D12_RAYTRACING_ACCELERATION_STRUCTURE_BUILD_FLAG_PREFER_FAST_BUILD;
@@ -2054,6 +2059,8 @@ FD3D12RayTracingGeometry::FD3D12RayTracingGeometry()
 
 		FD3D12RayTracingGeometry::NullTransformBuffer = RHICreateVertexBuffer(NullTransformData.GetResourceDataSize(), BUF_Static, CreateInfo);
 	}
+
+	FMemory::Memzero(PostBuildInfoBufferReadbackFences);
 }
 
 FD3D12RayTracingGeometry::~FD3D12RayTracingGeometry()
@@ -2284,6 +2291,28 @@ void FD3D12RayTracingGeometry::BuildAccelerationStructure(FD3D12CommandContext&
 		FD3D12DynamicRHI::TransitionResource(CommandContext.CommandListHandle, ScratchBuffers[GPUIndex].GetReference()->GetResource(), D3D12_RESOURCE_STATE_UNORDERED_ACCESS, 0);
 	}
 
+	// Compact BLAS if it is static
+	bool bShouldCompactAfterBuild = BuildFlags & D3D12_RAYTRACING_ACCELERATION_STRUCTURE_BUILD_FLAG_ALLOW_COMPACTION && BuildFlags & D3D12_RAYTRACING_ACCELERATION_STRUCTURE_BUILD_FLAG_PREFER_FAST_TRACE && !(BuildFlags & D3D12_RAYTRACING_ACCELERATION_STRUCTURE_BUILD_FLAG_ALLOW_UPDATE);
+
+	if (bShouldCompactAfterBuild)
+	{
+		D3D12_RESOURCE_DESC PostBuildInfoBufferDesc = CD3DX12_RESOURCE_DESC::Buffer(sizeof(uint64), D3D12_RESOURCE_FLAG_ALLOW_UNORDERED_ACCESS);
+
+		FRHIResourceCreateInfo CreateInfo;
+
+		CreateInfo.GPUMask = FRHIGPUMask::FromIndex(GPUIndex);
+		CreateInfo.DebugName = TEXT("PostBuildInfoBuffer");
+		PostBuildInfoBuffers[GPUIndex] = Adapter->CreateRHIBuffer<FD3D12MemBuffer>(
+			nullptr, PostBuildInfoBufferDesc, 8,
+			0, PostBuildInfoBufferDesc.Width, BUF_UnorderedAccess | BUF_SourceCopy, CreateInfo);
+
+		SetName(PostBuildInfoBuffers[GPUIndex]->GetResource(), TEXT("PostBuildInfoBuffer"));
+
+		FD3D12DynamicRHI::TransitionResource(CommandContext.CommandListHandle, PostBuildInfoBuffers[GPUIndex].GetReference()->GetResource(), D3D12_RESOURCE_STATE_UNORDERED_ACCESS, 0);
+
+		PostBuildInfoStagingBuffers[GPUIndex] = RHICreateStagingBuffer();
+	}
+
 	TransitionBuffers(CommandContext);
 	CommandContext.CommandListHandle.FlushResourceBarriers();
 
@@ -2301,7 +2330,27 @@ void FD3D12RayTracingGeometry::BuildAccelerationStructure(FD3D12CommandContext&
 			: D3D12_GPU_VIRTUAL_ADDRESS(0);
 
 		ID3D12GraphicsCommandList4* RayTracingCommandList = CommandContext.CommandListHandle.RayTracingCommandList();
-		RayTracingCommandList->BuildRaytracingAccelerationStructure(&BuildDesc, 0, nullptr);
+
+		if (!bShouldCompactAfterBuild)
+		{
+			// Build
+			RayTracingCommandList->BuildRaytracingAccelerationStructure(&BuildDesc, 0, nullptr);
+		}
+		else
+		{
+			D3D12_RAYTRACING_ACCELERATION_STRUCTURE_POSTBUILD_INFO_DESC PostBuildInfoDesc = {};
+			PostBuildInfoDesc.DestBuffer = PostBuildInfoBuffers[GPUIndex]->ResourceLocation.GetGPUVirtualAddress();
+			PostBuildInfoDesc.InfoType = D3D12_RAYTRACING_ACCELERATION_STRUCTURE_POSTBUILD_INFO_COMPACTED_SIZE;
+
+			// Build with post build info
+			RayTracingCommandList->BuildRaytracingAccelerationStructure(&BuildDesc, 1, &PostBuildInfoDesc);
+
+			CommandContext.RHICopyToStagingBuffer(PostBuildInfoBuffers[GPUIndex], PostBuildInfoStagingBuffers[GPUIndex], 0, sizeof(uint64));
+
+			const FD3D12Fence& Fence = CommandContext.GetParentDevice()->GetCommandListManager().GetFence();
+			PostBuildInfoBufferReadbackFences[GPUIndex] = FMath::Max(PostBuildInfoBufferReadbackFences[GPUIndex], Fence.GetCurrentFence());
+		}
+
 		SetDirty(CommandContext.GetGPUMask(), false);
 
 		if (bIsUpdate)
@@ -2323,6 +2372,70 @@ void FD3D12RayTracingGeometry::BuildAccelerationStructure(FD3D12CommandContext&
 	}
 }
 
+void FD3D12RayTracingGeometry::ConditionalCompactAccelerationStructure(FD3D12CommandContext& CommandContext)
+{
+	bool bShouldCompactAfterBuild = BuildFlags & D3D12_RAYTRACING_ACCELERATION_STRUCTURE_BUILD_FLAG_ALLOW_COMPACTION && BuildFlags & D3D12_RAYTRACING_ACCELERATION_STRUCTURE_BUILD_FLAG_PREFER_FAST_TRACE && !(BuildFlags & D3D12_RAYTRACING_ACCELERATION_STRUCTURE_BUILD_FLAG_ALLOW_UPDATE);
+
+	if (!bShouldCompactAfterBuild)
+	{
+		return;
+	}
+
+	const uint32 GPUIndex = CommandContext.GetGPUIndex();
+	FD3D12Adapter* Adapter = CommandContext.GetParentAdapter();
+
+	const FD3D12Fence& Fence = CommandContext.GetParentDevice()->GetCommandListManager().GetFence();
+
+	// Ensure that our builds & copies have finished on GPU
+	if (PostBuildInfoBufferReadbackFences[GPUIndex] >= Fence.GetCurrentFence())
+	{
+		return;
+	}
+
+	uint64 SizeAfterCompaction = 0;
+
+	{
+		// Read size and release readback buffers.
+		SizeAfterCompaction = *(uint64*)PostBuildInfoStagingBuffers[GPUIndex]->Lock(0, sizeof(uint64));
+
+		PostBuildInfoStagingBuffers[GPUIndex]->Unlock();
+		PostBuildInfoBuffers[GPUIndex] = nullptr;
+		PostBuildInfoStagingBuffers[GPUIndex] = nullptr;
+
+		// Set fence so we will not compact the same BLAS repeatedly
+		PostBuildInfoBufferReadbackFences[GPUIndex] = MAX_uint64;
+	}
+
+	DEC_MEMORY_STAT_BY(STAT_D3D12RayTracingUsedVideoMemory, AccelerationStructureBuffers[GPUIndex]->GetSize());
+	DEC_MEMORY_STAT_BY(STAT_D3D12RayTracingBLASMemory, AccelerationStructureBuffers[GPUIndex]->GetSize());
+
+	// Move old AS into this temporary variable which gets released when this function returns
+	TRefCountPtr<FD3D12MemBuffer> OldAccelerationStructure = MoveTemp(AccelerationStructureBuffers[GPUIndex]);
+
+	{
+		FRHIResourceCreateInfo CreateInfo;
+
+		D3D12_RESOURCE_DESC AccelerationStructureBufferDesc = CD3DX12_RESOURCE_DESC::Buffer(SizeAfterCompaction, D3D12_RESOURCE_FLAG_ALLOW_UNORDERED_ACCESS);
+
+		CreateInfo.GPUMask = FRHIGPUMask::FromIndex(GPUIndex);
+		CreateInfo.DebugName = TEXT("AccelerationStructureBuffer");
+		AccelerationStructureBuffers[GPUIndex] = Adapter->CreateRHIBuffer<FD3D12MemBuffer>(
+			nullptr, AccelerationStructureBufferDesc, D3D12_RAYTRACING_ACCELERATION_STRUCTURE_BYTE_ALIGNMENT,
+			0, AccelerationStructureBufferDesc.Width, BUF_AccelerationStructure, CreateInfo);
+
+		INC_MEMORY_STAT_BY(STAT_D3D12RayTracingUsedVideoMemory, AccelerationStructureBuffers[GPUIndex]->GetSize());
+		INC_MEMORY_STAT_BY(STAT_D3D12RayTracingBLASMemory, AccelerationStructureBuffers[GPUIndex]->GetSize());
+
+		SetName(AccelerationStructureBuffers[GPUIndex]->GetResource(), TEXT("Acceleration structure"));
+
+		ID3D12GraphicsCommandList4* RayTracingCommandList = CommandContext.CommandListHandle.RayTracingCommandList();
+		RayTracingCommandList->CopyRaytracingAccelerationStructure(
+			AccelerationStructureBuffers[GPUIndex]->ResourceLocation.GetGPUVirtualAddress(),
+			OldAccelerationStructure->ResourceLocation.GetGPUVirtualAddress(),
+			D3D12_RAYTRACING_ACCELERATION_STRUCTURE_COPY_MODE_COMPACT);
+	}
+}
+
 FD3D12RayTracingScene::FD3D12RayTracingScene(FD3D12Adapter& Adapter)
 {
 	ShaderResourceView = Adapter.CreateLinkedObject<FD3D12ShaderResourceView>(FRHIGPUMask::All(), [&](FD3D12Device* Device)
@@ -2447,6 +2560,8 @@ void FD3D12RayTracingScene::BuildAccelerationStructure(FD3D12CommandContext& Com
 			const FRayTracingGeometryInstance& Instance = Instances[InstanceIndex];
 			FD3D12RayTracingGeometry* Geometry = FD3D12DynamicRHI::ResourceCast(Instance.GeometryRHI);
 
+			Geometry->ConditionalCompactAccelerationStructure(CommandContext);
+
 			checkf(!Geometry->IsDirty(CommandContext.GetGPUIndex()),
 				TEXT("Acceleration structures for all geometries must be built before building the top level acceleration structure for the scene."));
 
diff --git a/Engine/Source/Runtime/D3D12RHI/Private/D3D12RayTracing.h b/Engine/Source/Runtime/D3D12RHI/Private/D3D12RayTracing.h
--- a/Engine/Source/Runtime/D3D12RHI/Private/D3D12RayTracing.h
+++ b/Engine/Source/Runtime/D3D12RHI/Private/D3D12RayTracing.h
@@ -20,6 +20,7 @@ public:
 
 	void TransitionBuffers(FD3D12CommandContext& CommandContext);
 	void BuildAccelerationStructure(FD3D12CommandContext& CommandContext, EAccelerationStructureBuildMode BuildMode);
+	void ConditionalCompactAccelerationStructure(FD3D12CommandContext& CommandContext);
 
 	bool bIsAccelerationStructureDirty[MAX_NUM_GPUS] = {};
 	void SetDirty(FRHIGPUMask GPUMask, bool bState)
@@ -49,6 +50,10 @@ public:
 	TRefCountPtr<FD3D12MemBuffer> AccelerationStructureBuffers[MAX_NUM_GPUS];
 	TRefCountPtr<FD3D12MemBuffer> ScratchBuffers[MAX_NUM_GPUS];
 
+	uint64 PostBuildInfoBufferReadbackFences[MAX_NUM_GPUS];
+	TRefCountPtr<FD3D12MemBuffer> PostBuildInfoBuffers[MAX_NUM_GPUS];
+	FStagingBufferRHIRef PostBuildInfoStagingBuffers[MAX_NUM_GPUS];
+
 };
 
 class FD3D12RayTracingScene : public FRHIRayTracingScene
