Added CPU profiling scopes to several known expensive functions related to ray tracing scene management

diff --git a/Engine/Source/Runtime/D3D12RHI/Private/D3D12RayTracing.cpp b/Engine/Source/Runtime/D3D12RHI/Private/D3D12RayTracing.cpp
--- a/Engine/Source/Runtime/D3D12RHI/Private/D3D12RayTracing.cpp
+++ b/Engine/Source/Runtime/D3D12RHI/Private/D3D12RayTracing.cpp
@@ -2015,6 +2015,8 @@ FRayTracingGeometryRHIRef FD3D12DynamicRHI::RHICreateRayTracingGeometry(const FR
 
 FRayTracingSceneRHIRef FD3D12DynamicRHI::RHICreateRayTracingScene(const FRayTracingSceneInitializer& Initializer)
 {
+	TRACE_CPUPROFILER_EVENT_SCOPE(CreateRayTracingScene);
+
 	FD3D12Adapter& Adapter = GetAdapter();
 
 	FD3D12RayTracingScene* Result = new FD3D12RayTracingScene(Adapter);
@@ -2142,6 +2144,8 @@ static void CreateAccelerationStructureBuffers(TRefCountPtr<FD3D12MemBuffer>& Ac
 
 void FD3D12RayTracingGeometry::BuildAccelerationStructure(FD3D12CommandContext& CommandContext, EAccelerationStructureBuildMode BuildMode)
 {
+	TRACE_CPUPROFILER_EVENT_SCOPE(BuildAccelerationStructure_BottomLevel);
+
 	const bool bIsUpdate = BuildMode == EAccelerationStructureBuildMode::Update;
 	static constexpr uint32 IndicesPerPrimitive = 3; // Only triangle meshes are supported
 
@@ -2471,6 +2475,8 @@ FD3D12RayTracingScene::~FD3D12RayTracingScene()
 
 void FD3D12RayTracingScene::BuildAccelerationStructure(FD3D12CommandContext& CommandContext, D3D12_RAYTRACING_ACCELERATION_STRUCTURE_BUILD_FLAGS BuildFlags)
 {
+	TRACE_CPUPROFILER_EVENT_SCOPE(BuildAccelerationStructure_TopLevel);
+
 	TRefCountPtr<FD3D12StructuredBuffer> InstanceBuffer;
 	TRefCountPtr<FD3D12MemBuffer> ScratchBuffer;
 
@@ -2734,6 +2740,8 @@ FD3D12RayTracingShaderTable* FD3D12RayTracingScene::FindOrCreateShaderTable(cons
 		return FoundShaderTable;
 	}
 
+	TRACE_CPUPROFILER_EVENT_SCOPE(FindOrCreateShaderTable);
+
 	FD3D12RayTracingShaderTable* CreatedShaderTable = new FD3D12RayTracingShaderTable();
 	ID3D12Device5* RayTracingDevice = Device->GetRayTracingDevice();
 
@@ -3574,6 +3582,8 @@ void FD3D12CommandContext::RHISetRayTracingHitGroups(
 	FRHIRayTracingScene* InScene, FRHIRayTracingPipelineState* InPipeline,
 	uint32 NumBindings, const FRayTracingLocalShaderBindings* Bindings)
 {
+	TRACE_CPUPROFILER_EVENT_SCOPE(SetRayTracingHitGroups);
+
 	FD3D12RayTracingScene* Scene = FD3D12DynamicRHI::ResourceCast(InScene);
 	FD3D12RayTracingPipelineState* Pipeline = FD3D12DynamicRHI::ResourceCast(InPipeline);
 
