Fix ray tracing acceleration compacted size readback. Use GetLastCompletedFenceFast() instead of GetCurrentFence() to detect GPU work completion.

diff --git a/Engine/Source/Runtime/D3D12RHI/Private/D3D12RayTracing.cpp b/Engine/Source/Runtime/D3D12RHI/Private/D3D12RayTracing.cpp
--- a/Engine/Source/Runtime/D3D12RHI/Private/D3D12RayTracing.cpp
+++ b/Engine/Source/Runtime/D3D12RHI/Private/D3D12RayTracing.cpp
@@ -2378,7 +2378,7 @@ void FD3D12RayTracingGeometry::BuildAccelerationStructure(FD3D12CommandContext&
 
 void FD3D12RayTracingGeometry::ConditionalCompactAccelerationStructure(FD3D12CommandContext& CommandContext)
 {
-	bool bShouldCompactAfterBuild = BuildFlags & D3D12_RAYTRACING_ACCELERATION_STRUCTURE_BUILD_FLAG_ALLOW_COMPACTION && BuildFlags & D3D12_RAYTRACING_ACCELERATION_STRUCTURE_BUILD_FLAG_PREFER_FAST_TRACE && !(BuildFlags & D3D12_RAYTRACING_ACCELERATION_STRUCTURE_BUILD_FLAG_ALLOW_UPDATE);
+	const bool bShouldCompactAfterBuild = BuildFlags& D3D12_RAYTRACING_ACCELERATION_STRUCTURE_BUILD_FLAG_ALLOW_COMPACTION&& BuildFlags& D3D12_RAYTRACING_ACCELERATION_STRUCTURE_BUILD_FLAG_PREFER_FAST_TRACE && !(BuildFlags & D3D12_RAYTRACING_ACCELERATION_STRUCTURE_BUILD_FLAG_ALLOW_UPDATE);
 
 	if (!bShouldCompactAfterBuild)
 	{
@@ -2391,7 +2391,7 @@ void FD3D12RayTracingGeometry::ConditionalCompactAccelerationStructure(FD3D12Com
 	const FD3D12Fence& Fence = CommandContext.GetParentDevice()->GetCommandListManager().GetFence();
 
 	// Ensure that our builds & copies have finished on GPU
-	if (PostBuildInfoBufferReadbackFences[GPUIndex] >= Fence.GetCurrentFence())
+	if (PostBuildInfoBufferReadbackFences[GPUIndex] > Fence.GetLastCompletedFenceFast())
 	{
 		return;
 	}
@@ -2410,6 +2410,13 @@ void FD3D12RayTracingGeometry::ConditionalCompactAccelerationStructure(FD3D12Com
 		PostBuildInfoBufferReadbackFences[GPUIndex] = MAX_uint64;
 	}
 
+	ensureMsgf(SizeAfterCompaction, TEXT("Compacted acceleration structure size is expected to be non-zero. This error suggests that GPU readback synchronization is broken."));
+
+	if (SizeAfterCompaction == 0)
+	{
+		return;
+	}
+
 	DEC_MEMORY_STAT_BY(STAT_D3D12RayTracingUsedVideoMemory, AccelerationStructureBuffers[GPUIndex]->GetSize());
 	DEC_MEMORY_STAT_BY(STAT_D3D12RayTracingBLASMemory, AccelerationStructureBuffers[GPUIndex]->GetSize());
 
