DX12 buffer refactor, unifying the core of buffer to a single class

diff --git a/Engine/Source/Runtime/D3D12RHI/Private/D3D12Buffer.cpp b/Engine/Source/Runtime/D3D12RHI/Private/D3D12Buffer.cpp
--- a/Engine/Source/Runtime/D3D12RHI/Private/D3D12Buffer.cpp
+++ b/Engine/Source/Runtime/D3D12RHI/Private/D3D12Buffer.cpp
@@ -7,14 +7,6 @@ D3D12Buffer.cpp: D3D Common code for buffers.
 #include "D3D12RHIPrivate.h"
 
 // Forward declarations are required for the template functions
-template void* FD3D12DynamicRHI::LockBuffer<FD3D12VertexBuffer>(FRHICommandListImmediate* RHICmdList, FD3D12VertexBuffer* Buffer, uint32 Offset, uint32 Size, EResourceLockMode LockMode);
-template void* FD3D12DynamicRHI::LockBuffer<FD3D12IndexBuffer>(FRHICommandListImmediate* RHICmdList, FD3D12IndexBuffer* Buffer, uint32 Offset, uint32 Size, EResourceLockMode LockMode);
-template void* FD3D12DynamicRHI::LockBuffer<FD3D12StructuredBuffer>(FRHICommandListImmediate* RHICmdList, FD3D12StructuredBuffer* Buffer, uint32 Offset, uint32 Size, EResourceLockMode LockMode);
-
-template void FD3D12DynamicRHI::UnlockBuffer<FD3D12VertexBuffer>(FRHICommandListImmediate* RHICmdList, FD3D12VertexBuffer* Buffer);
-template void FD3D12DynamicRHI::UnlockBuffer<FD3D12IndexBuffer>(FRHICommandListImmediate* RHICmdList, FD3D12IndexBuffer* Buffer);
-template void FD3D12DynamicRHI::UnlockBuffer<FD3D12StructuredBuffer>(FRHICommandListImmediate* RHICmdList, FD3D12StructuredBuffer* Buffer);
-
 template FD3D12VertexBuffer* FD3D12Adapter::CreateRHIBuffer<FD3D12VertexBuffer>(FRHICommandListImmediate* RHICmdList,
 	const D3D12_RESOURCE_DESC& Desc,
 	uint32 Alignment, uint32 Stride, uint32 Size, uint32 InUsage,
@@ -63,13 +55,12 @@ struct FRHICommandRenameUploadBufferString
 {
 	static const TCHAR* TStr() { return TEXT("FRHICommandRenameUploadBuffer"); }
 };
-template<typename ResourceType>
-struct FRHICommandRenameUploadBuffer final : public FRHICommand<FRHICommandRenameUploadBuffer<ResourceType>, FRHICommandRenameUploadBufferString>
+struct FRHICommandRenameUploadBuffer final : public FRHICommand<FRHICommandRenameUploadBuffer, FRHICommandRenameUploadBufferString>
 {
-	ResourceType* Resource;
+	FD3D12Buffer* Resource;
 	FD3D12ResourceLocation NewLocation;
 
-	FORCEINLINE_DEBUGGABLE FRHICommandRenameUploadBuffer(ResourceType* InResource, FD3D12Device* Device)
+	FORCEINLINE_DEBUGGABLE FRHICommandRenameUploadBuffer(FD3D12Buffer* InResource, FD3D12Device* Device)
 		: Resource(InResource)
 		, NewLocation(Device) 
 	{}
@@ -80,6 +71,59 @@ struct FRHICommandRenameUploadBuffer final : public FRHICommand<FRHICommandRenam
 	}
 };
 
+struct FD3D12RHICommandInitializeBufferString
+{
+	static const TCHAR* TStr() { return TEXT("FD3D12RHICommandInitializeBuffer"); }
+};
+struct FD3D12RHICommandInitializeBuffer final : public FRHICommand<FD3D12RHICommandInitializeBuffer, FD3D12RHICommandInitializeBufferString>
+{
+	FD3D12Buffer* CurrentBuffer;
+	FD3D12ResourceLocation SrcResourceLoc;
+	uint32 Size;
+
+	FORCEINLINE_DEBUGGABLE FD3D12RHICommandInitializeBuffer(FD3D12Buffer* InCurrentBuffer, FD3D12ResourceLocation& InSrcResourceLoc, uint32 InSize)
+		: CurrentBuffer(InCurrentBuffer)
+		, SrcResourceLoc(InSrcResourceLoc.GetParentDevice())
+		, Size(InSize)
+	{
+		FD3D12ResourceLocation::TransferOwnership(SrcResourceLoc, InSrcResourceLoc);
+	}
+
+	void Execute(FRHICommandListBase& /* unused */)
+	{
+		ExecuteNoCmdList();
+	}
+
+	void ExecuteNoCmdList()
+	{
+		while (CurrentBuffer != nullptr)
+		{
+			FD3D12Resource* Destination = CurrentBuffer->ResourceLocation.GetResource();
+			FD3D12Device* Device = Destination->GetParentDevice();
+
+			FD3D12CommandListHandle& hCommandList = Device->GetDefaultCommandContext().CommandListHandle;
+			// Copy from the temporary upload heap to the default resource
+			{
+				// Writable structured buffers are sometimes initialized with initial data which means they sometimes need tracking.
+				FConditionalScopeResourceBarrier ConditionalScopeResourceBarrier(hCommandList, Destination, D3D12_RESOURCE_STATE_COPY_DEST, 0);
+
+				Device->GetDefaultCommandContext().numCopies++;
+				hCommandList.FlushResourceBarriers();
+				hCommandList->CopyBufferRegion(
+					Destination->GetResource(),
+					CurrentBuffer->ResourceLocation.GetOffsetFromBaseOfResource(),
+					SrcResourceLoc.GetResource()->GetResource(),
+					SrcResourceLoc.GetOffsetFromBaseOfResource(), Size);
+
+				hCommandList.UpdateResidency(Destination);
+				hCommandList.UpdateResidency(SrcResourceLoc.GetResource());
+			}
+
+			CurrentBuffer = CurrentBuffer->GetNextObject();
+		}
+	}
+};
+
 void FD3D12Adapter::AllocateBuffer(FD3D12Device* Device,
 	const D3D12_RESOURCE_DESC& InDesc,
 	uint32 Size,
@@ -188,59 +232,6 @@ BufferType* FD3D12Adapter::CreateRHIBuffer(FRHICommandListImmediate* RHICmdList,
 			}
 			check(pData);
 			FMemory::Memcpy(pData, CreateInfo.ResourceArray->GetResourceData(), Size);
-
-			struct FD3D12RHICommandInitializeBufferString
-			{
-				static const TCHAR* TStr() { return TEXT("FD3D12RHICommandInitializeBuffer"); }
-			};
-			struct FD3D12RHICommandInitializeBuffer final : public FRHICommand<FD3D12RHICommandInitializeBuffer, FD3D12RHICommandInitializeBufferString>
-			{
-				BufferType* CurrentBuffer;
-				FD3D12ResourceLocation SrcResourceLoc;
-				uint32 Size;
-
-				FORCEINLINE_DEBUGGABLE FD3D12RHICommandInitializeBuffer(BufferType* InCurrentBuffer, FD3D12ResourceLocation& InSrcResourceLoc, uint32 InSize)
-					: CurrentBuffer(InCurrentBuffer)
-					, SrcResourceLoc(InSrcResourceLoc.GetParentDevice())
-					, Size(InSize)
-				{
-					FD3D12ResourceLocation::TransferOwnership(SrcResourceLoc, InSrcResourceLoc);
-				}
-
-				void Execute(FRHICommandListBase& /* unused */)
-				{
-					ExecuteNoCmdList();
-				}
-
-				void ExecuteNoCmdList()
-				{
-					while (CurrentBuffer != nullptr)
-					{
-						FD3D12Resource* Destination = CurrentBuffer->ResourceLocation.GetResource();
-						FD3D12Device* Device = Destination->GetParentDevice();
-
-						FD3D12CommandListHandle& hCommandList = Device->GetDefaultCommandContext().CommandListHandle;
-						// Copy from the temporary upload heap to the default resource
-						{
-							// Writable structured buffers are sometimes initialized with initial data which means they sometimes need tracking.
-							FConditionalScopeResourceBarrier ConditionalScopeResourceBarrier(hCommandList, Destination, D3D12_RESOURCE_STATE_COPY_DEST, 0);
-
-							Device->GetDefaultCommandContext().numCopies++;
-							hCommandList.FlushResourceBarriers();
-							hCommandList->CopyBufferRegion(
-								Destination->GetResource(),
-								CurrentBuffer->ResourceLocation.GetOffsetFromBaseOfResource(),
-								SrcResourceLoc.GetResource()->GetResource(),
-								SrcResourceLoc.GetOffsetFromBaseOfResource(), Size);
-
-							hCommandList.UpdateResidency(Destination);
-							hCommandList.UpdateResidency(SrcResourceLoc.GetResource());
-						}
-
-						CurrentBuffer = CurrentBuffer->GetNextObject();
-					}
-				}
-			};
 			
 			if (bOnAsyncThread)
 			{
@@ -292,17 +283,76 @@ BufferType* FD3D12Adapter::CreateRHIBuffer(FRHICommandListImmediate* RHICmdList,
 	return BufferOut;
 }
 
-template<class BufferType>
-void* FD3D12DynamicRHI::LockBuffer(FRHICommandListImmediate* RHICmdList, BufferType* Buffer, uint32 Offset, uint32 Size, EResourceLockMode LockMode)
+void FD3D12Buffer::Rename(FD3D12ResourceLocation& NewLocation)
+{
+	FD3D12ResourceLocation::TransferOwnership(ResourceLocation, NewLocation);
+
+	FScopeLock Lock(&DynamicSRVsCS);
+	for (FD3D12BaseShaderResourceView* DynamicSRVBase : DynamicSRVs)
+	{
+		FD3D12ShaderResourceView* DynamicSRV = static_cast<FD3D12ShaderResourceView*>(DynamicSRVBase);
+		if (DynamicSRV->IsValid())
+		{
+			DynamicSRV->Rename(ResourceLocation);
+		}
+	}
+}
+
+void FD3D12Buffer::RenameLDAChain(FD3D12ResourceLocation& NewLocation)
+{
+	// Dynamic buffers use cross-node resources.
+	//ensure(GetUsage() & BUF_AnyDynamic);
+	Rename(NewLocation);
+
+	if (GNumExplicitGPUsForRendering > 1)
+	{
+		// This currently crashes at exit time because NewLocation isn't tracked in the right allocator.
+		ensure(IsHeadLink());
+		ensure(GetParentDevice() == NewLocation.GetParentDevice());
+
+		// Update all of the resources in the LDA chain to reference this cross-node resource
+		for (FD3D12Buffer* NextBuffer = GetNextObject(); NextBuffer; NextBuffer = NextBuffer->GetNextObject())
+		{
+			FD3D12ResourceLocation::ReferenceNode(NextBuffer->GetParentDevice(), NextBuffer->ResourceLocation, ResourceLocation);
+
+			FScopeLock Lock(&NextBuffer->DynamicSRVsCS);
+			for (FD3D12BaseShaderResourceView* DynamicSRVBase : NextBuffer->DynamicSRVs)
+			{
+				FD3D12ShaderResourceView* DynamicSRV = static_cast<FD3D12ShaderResourceView*>(DynamicSRVBase);
+				if (DynamicSRV->IsValid())
+				{
+					DynamicSRV->Rename(NextBuffer->ResourceLocation);
+				}
+			}
+		}
+	}
+}
+
+void FD3D12Buffer::ReleaseUnderlyingResource()
+{
+	check(!LockedData.bLocked && ResourceLocation.IsValid());
+	ResourceLocation.Clear();
+	RemoveAllDynamicSRVs();
+
+	FD3D12Buffer* NextBuffer = GetNextObject();
+	if (NextBuffer)
+	{
+		NextBuffer->ReleaseUnderlyingResource();
+	}
+}
+
+void* FD3D12DynamicRHI::LockBuffer(FRHICommandListImmediate* RHICmdList, FD3D12Buffer* Buffer, uint32 BufferSize, uint32 BufferUsage, uint32 Offset, uint32 Size, EResourceLockMode LockMode)
 {
 	SCOPE_CYCLE_COUNTER(STAT_D3D12LockBufferTime);
 
+	check(Size <= BufferSize);
+
 	FD3D12LockedResource& LockedData = Buffer->LockedData;
 	check(LockedData.bLocked == false);
 	FD3D12Adapter& Adapter = GetAdapter();
 
 	// Determine whether the buffer is dynamic or not.
-	const bool bIsDynamic = (Buffer->GetUsage() & BUF_AnyDynamic) ? true : false;
+	const bool bIsDynamic = (BufferUsage & BUF_AnyDynamic) ? true : false;
 
 	void* Data = nullptr;
 
@@ -323,15 +373,15 @@ void* FD3D12DynamicRHI::LockBuffer(FRHICommandListImmediate* RHICmdList, BufferT
 			// If on the RenderThread, queue up a command on the RHIThread to rename this buffer at the correct time
 			if (ShouldDeferBufferLockOperation(RHICmdList) && LockMode == RLM_WriteOnly)
 			{
-				FRHICommandRenameUploadBuffer<BufferType>* Command = ALLOC_COMMAND_CL(*RHICmdList, FRHICommandRenameUploadBuffer<BufferType>)(Buffer, Device);
+				FRHICommandRenameUploadBuffer* Command = ALLOC_COMMAND_CL(*RHICmdList, FRHICommandRenameUploadBuffer)(Buffer, Device);
 
-				Data = Adapter.GetUploadHeapAllocator(Device->GetGPUIndex()).AllocUploadResource(Buffer->GetSize(), Buffer->BufferAlignment, Command->NewLocation);
+				Data = Adapter.GetUploadHeapAllocator(Device->GetGPUIndex()).AllocUploadResource(BufferSize, Buffer->BufferAlignment, Command->NewLocation);
 				RHICmdList->RHIThreadFence(true);
 			}
 			else
 			{
 				FD3D12ResourceLocation Location(Buffer->GetParentDevice());
-				Data = Adapter.GetUploadHeapAllocator(Device->GetGPUIndex()).AllocUploadResource(Buffer->GetSize(), Buffer->BufferAlignment, Location);
+				Data = Adapter.GetUploadHeapAllocator(Device->GetGPUIndex()).AllocUploadResource(BufferSize, Buffer->BufferAlignment, Location);
 				Buffer->RenameLDAChain(Location);
 			}
 		}
@@ -413,8 +463,7 @@ void* FD3D12DynamicRHI::LockBuffer(FRHICommandListImmediate* RHICmdList, BufferT
 	return Data;
 }
 
-template<class BufferType>
-void FD3D12DynamicRHI::UnlockBuffer(FRHICommandListImmediate* RHICmdList, BufferType* Buffer)
+void FD3D12DynamicRHI::UnlockBuffer(FRHICommandListImmediate* RHICmdList, FD3D12Buffer* Buffer, uint32 BufferUsage)
 {
 	SCOPE_CYCLE_COUNTER(STAT_D3D12UnlockBufferTime);
 
@@ -422,7 +471,7 @@ void FD3D12DynamicRHI::UnlockBuffer(FRHICommandListImmediate* RHICmdList, Buffer
 	check(LockedData.bLocked == true);
 
 	// Determine whether the buffer is dynamic or not.
-	const bool bIsDynamic = (Buffer->GetUsage() & BUF_AnyDynamic) != 0;
+	const bool bIsDynamic = (BufferUsage & BUF_AnyDynamic) != 0;
 
 	if (bIsDynamic)
 	{
@@ -437,7 +486,7 @@ void FD3D12DynamicRHI::UnlockBuffer(FRHICommandListImmediate* RHICmdList, Buffer
 		else
 		{
 			// Update all of the resources in the LDA chain
-			for (BufferType* CurrentBuffer = Buffer; CurrentBuffer; CurrentBuffer = CurrentBuffer->GetNextObject())
+			for (FD3D12Buffer* CurrentBuffer = Buffer; CurrentBuffer; CurrentBuffer = CurrentBuffer->GetNextObject())
 			{
 				// If we are on the render thread, queue up the copy on the RHIThread so it happens at the correct time.
 				if (ShouldDeferBufferLockOperation(RHICmdList))
diff --git a/Engine/Source/Runtime/D3D12RHI/Private/D3D12CommandContext.h b/Engine/Source/Runtime/D3D12RHI/Private/D3D12CommandContext.h
--- a/Engine/Source/Runtime/D3D12RHI/Private/D3D12CommandContext.h
+++ b/Engine/Source/Runtime/D3D12RHI/Private/D3D12CommandContext.h
@@ -384,14 +384,20 @@ public:
 	}
 
 	template<typename ObjectType, typename RHIType>
-	FORCEINLINE_DEBUGGABLE ObjectType* RetrieveObject(RHIType RHIObject)
+	static FORCEINLINE_DEBUGGABLE ObjectType* RetrieveObject(RHIType RHIObject, FD3D12Device* Device)
 	{
 		return RetrieveObject<ObjectType, RHIType>(RHIObject, [&](ObjectType* Object)
 		{
-			return Object->GetParentDevice() == GetParentDevice();
+			return Object->GetParentDevice() == Device;
 		});
 	}
 
+	template<typename ObjectType, typename RHIType>
+	FORCEINLINE_DEBUGGABLE ObjectType* RetrieveObject(RHIType RHIObject)
+	{
+		return RetrieveObject<ObjectType, RHIType>(RHIObject, GetParentDevice());
+	}
+
 	template<typename Predicate>
 	static inline FD3D12TextureBase* RetrieveTextureBase(FRHITexture* Texture, Predicate Func)
 	{
diff --git a/Engine/Source/Runtime/D3D12RHI/Private/D3D12Commands.cpp b/Engine/Source/Runtime/D3D12RHI/Private/D3D12Commands.cpp
--- a/Engine/Source/Runtime/D3D12RHI/Private/D3D12Commands.cpp
+++ b/Engine/Source/Runtime/D3D12RHI/Private/D3D12Commands.cpp
@@ -88,7 +88,7 @@ void FD3D12DynamicRHI::SetupRecursiveResources()
 // Vertex state.
 void FD3D12CommandContext::RHISetStreamSource(uint32 StreamIndex, FRHIVertexBuffer* VertexBufferRHI, uint32 Offset)
 {
-	FD3D12VertexBuffer* VertexBuffer = RetrieveObject<FD3D12VertexBuffer>(VertexBufferRHI);
+	FD3D12Buffer* VertexBuffer = RetrieveObject<FD3D12Buffer>(VertexBufferRHI);
 
 	StateCache.SetStreamSource(VertexBuffer ? &VertexBuffer->ResourceLocation : nullptr, StreamIndex, Offset);
 }
@@ -1465,7 +1465,7 @@ void FD3D12CommandContext::RHIDrawPrimitive(uint32 BaseVertexIndex, uint32 NumPr
 
 void FD3D12CommandContext::RHIDrawPrimitiveIndirect(FRHIVertexBuffer* ArgumentBufferRHI, uint32 ArgumentOffset)
 {
-	FD3D12VertexBuffer* ArgumentBuffer = RetrieveObject<FD3D12VertexBuffer>(ArgumentBufferRHI);
+	FD3D12Buffer* ArgumentBuffer = RetrieveObject<FD3D12Buffer>(ArgumentBufferRHI);
 
 	numDraws++;
 	if (bTrackingEvents)
@@ -1504,8 +1504,11 @@ void FD3D12CommandContext::RHIDrawPrimitiveIndirect(FRHIVertexBuffer* ArgumentBu
 
 void FD3D12CommandContext::RHIDrawIndexedIndirect(FRHIIndexBuffer* IndexBufferRHI, FRHIStructuredBuffer* ArgumentsBufferRHI, int32 DrawArgumentsIndex, uint32 NumInstances)
 {
-	FD3D12IndexBuffer* IndexBuffer = RetrieveObject<FD3D12IndexBuffer>(IndexBufferRHI);
-	FD3D12StructuredBuffer* ArgumentsBuffer = RetrieveObject<FD3D12StructuredBuffer>(ArgumentsBufferRHI);
+	const uint32 IndexBufferStride = FD3D12DynamicRHI::ResourceCast(IndexBufferRHI)->GetStride();
+	const uint32 ArgumentsBufferStride = FD3D12DynamicRHI::ResourceCast(ArgumentsBufferRHI)->GetStride();
+
+	FD3D12Buffer* IndexBuffer = RetrieveObject<FD3D12Buffer>(IndexBufferRHI);
+	FD3D12Buffer* ArgumentsBuffer = RetrieveObject<FD3D12Buffer>(ArgumentsBufferRHI);
 
 	numDraws++;
 	if (bTrackingEvents)
@@ -1517,7 +1520,7 @@ void FD3D12CommandContext::RHIDrawIndexedIndirect(FRHIIndexBuffer* IndexBufferRH
 	CommitNonComputeShaderConstants();
 
 	// determine 16bit vs 32bit indices
-	const DXGI_FORMAT Format = (IndexBuffer->GetStride() == sizeof(uint16) ? DXGI_FORMAT_R16_UINT : DXGI_FORMAT_R32_UINT);
+	const DXGI_FORMAT Format = (IndexBufferStride == sizeof(uint16) ? DXGI_FORMAT_R16_UINT : DXGI_FORMAT_R32_UINT);
 
 	StateCache.SetIndexBuffer(IndexBuffer->ResourceLocation, Format, 0);
 
@@ -1534,7 +1537,7 @@ void FD3D12CommandContext::RHIDrawIndexedIndirect(FRHIIndexBuffer* IndexBufferRH
 		GetParentDevice()->GetParentAdapter()->GetDrawIndexedIndirectCommandSignature(),
 		1,
 		Location.GetResource()->GetResource(),
-		Location.GetOffsetFromBaseOfResource() + DrawArgumentsIndex * ArgumentsBuffer->GetStride(),
+		Location.GetOffsetFromBaseOfResource() + DrawArgumentsIndex * ArgumentsBufferStride,
 		NULL,
 		0
 		);
@@ -1563,15 +1566,17 @@ void FD3D12CommandContext::RHIDrawIndexedPrimitive(FRHIIndexBuffer* IndexBufferR
 	CommitNonComputeShaderConstants();
 
 	uint32 IndexCount = StateCache.GetVertexCountAndIncrementStat(NumPrimitives);
-	FD3D12IndexBuffer* IndexBuffer = RetrieveObject<FD3D12IndexBuffer>(IndexBufferRHI);
+
+	FD3D12IndexBuffer* RHIIndexBuffer = FD3D12DynamicRHI::ResourceCast(IndexBufferRHI);
+	FD3D12Buffer* IndexBuffer = RetrieveObject<FD3D12Buffer>(IndexBufferRHI);
 
 	// Verify that we are not trying to read outside the index buffer range
 	// test is an optimized version of: StartIndex + IndexCount <= IndexBuffer->GetSize() / IndexBuffer->GetStride() 
-	checkf((StartIndex + IndexCount) * IndexBuffer->GetStride() <= IndexBuffer->GetSize(),
-		TEXT("Start %u, Count %u, Type %u, Buffer Size %u, Buffer stride %u"), StartIndex, IndexCount, StateCache.GetGraphicsPipelinePrimitiveType(), IndexBuffer->GetSize(), IndexBuffer->GetStride());
+	checkf((StartIndex + IndexCount) * RHIIndexBuffer->GetStride() <= RHIIndexBuffer->GetSize(),
+		TEXT("Start %u, Count %u, Type %u, Buffer Size %u, Buffer stride %u"), StartIndex, IndexCount, StateCache.GetGraphicsPipelinePrimitiveType(), RHIIndexBuffer->GetSize(), RHIIndexBuffer->GetStride());
 
 	// determine 16bit vs 32bit indices
-	const DXGI_FORMAT Format = (IndexBuffer->GetStride() == sizeof(uint16) ? DXGI_FORMAT_R16_UINT : DXGI_FORMAT_R32_UINT);
+	const DXGI_FORMAT Format = (RHIIndexBuffer->GetStride() == sizeof(uint16) ? DXGI_FORMAT_R16_UINT : DXGI_FORMAT_R32_UINT);
 	StateCache.SetIndexBuffer(IndexBuffer->ResourceLocation, Format, 0);
 	StateCache.ApplyState<D3D12PT_Graphics>();
 
@@ -1585,8 +1590,9 @@ void FD3D12CommandContext::RHIDrawIndexedPrimitive(FRHIIndexBuffer* IndexBufferR
 
 void FD3D12CommandContext::RHIDrawIndexedPrimitiveIndirect(FRHIIndexBuffer* IndexBufferRHI, FRHIVertexBuffer* ArgumentBufferRHI, uint32 ArgumentOffset)
 {
-	FD3D12IndexBuffer* IndexBuffer = RetrieveObject<FD3D12IndexBuffer>(IndexBufferRHI);
-	FD3D12VertexBuffer* ArgumentBuffer = RetrieveObject<FD3D12VertexBuffer>(ArgumentBufferRHI);
+	const uint32 IndexBufferStride = FD3D12DynamicRHI::ResourceCast(IndexBufferRHI)->GetStride();
+	FD3D12Buffer* IndexBuffer = RetrieveObject<FD3D12Buffer>(IndexBufferRHI);
+	FD3D12Buffer* ArgumentBuffer = RetrieveObject<FD3D12Buffer>(ArgumentBufferRHI);
 
 	numDraws++;
 	if (bTrackingEvents)
@@ -1598,7 +1604,7 @@ void FD3D12CommandContext::RHIDrawIndexedPrimitiveIndirect(FRHIIndexBuffer* Inde
 	CommitNonComputeShaderConstants();
 
 	// Set the index buffer.
-	const DXGI_FORMAT Format = (IndexBuffer->GetStride() == sizeof(uint16) ? DXGI_FORMAT_R16_UINT : DXGI_FORMAT_R32_UINT);
+	const DXGI_FORMAT Format = (IndexBufferStride == sizeof(uint16) ? DXGI_FORMAT_R16_UINT : DXGI_FORMAT_R32_UINT);
 	StateCache.SetIndexBuffer(IndexBuffer->ResourceLocation, Format, 0);
 
 	FD3D12ResourceLocation& Location = ArgumentBuffer->ResourceLocation;
diff --git a/Engine/Source/Runtime/D3D12RHI/Private/D3D12IndexBuffer.cpp b/Engine/Source/Runtime/D3D12RHI/Private/D3D12IndexBuffer.cpp
--- a/Engine/Source/Runtime/D3D12RHI/Private/D3D12IndexBuffer.cpp
+++ b/Engine/Source/Runtime/D3D12RHI/Private/D3D12IndexBuffer.cpp
@@ -37,53 +37,20 @@ FD3D12IndexBuffer::~FD3D12IndexBuffer()
 	}
 }
 
-void FD3D12IndexBuffer::Rename(FD3D12ResourceLocation& NewLocation)
-{
-	FD3D12ResourceLocation::TransferOwnership(ResourceLocation, NewLocation);
-}
-
-void FD3D12IndexBuffer::RenameLDAChain(FD3D12ResourceLocation& NewLocation)
-{
-	// Dynamic buffers use cross-node resources.
-	ensure(GetUsage() & BUF_AnyDynamic);
-	Rename(NewLocation);
-
-	if (GNumExplicitGPUsForRendering > 1)
-	{
-		// This currently crashes at exit time because NewLocation isn't tracked in the right allocator.
-		ensure(IsHeadLink());
-		ensure(GetParentDevice() == NewLocation.GetParentDevice());
-
-		// Update all of the resources in the LDA chain to reference this cross-node resource
-		for (FD3D12IndexBuffer* NextBuffer = GetNextObject(); NextBuffer; NextBuffer = NextBuffer->GetNextObject())
-		{
-			FD3D12ResourceLocation::ReferenceNode(NextBuffer->GetParentDevice(), NextBuffer->ResourceLocation, ResourceLocation);
-		}
-	}
-}
-
 void FD3D12IndexBuffer::Swap(FD3D12IndexBuffer& Other)
 {
 	check(!LockedData.bLocked && !Other.LockedData.bLocked);
 	FRHIIndexBuffer::Swap(Other);
 	FD3D12BaseShaderResource::Swap(Other);
 	FD3D12TransientResource::Swap(Other);
-	FD3D12LinkedAdapterObject<FD3D12IndexBuffer>::Swap(Other);
+	FD3D12LinkedAdapterObject<FD3D12Buffer>::Swap(Other);
 }
 
 void FD3D12IndexBuffer::ReleaseUnderlyingResource()
 {
-	check(!LockedData.bLocked && ResourceLocation.IsValid());
 	UpdateBufferStats<FD3D12IndexBuffer>(&ResourceLocation, false);
-	ResourceLocation.Clear();
 	FRHIIndexBuffer::ReleaseUnderlyingResource();
-	RemoveAllDynamicSRVs();
-
-	FD3D12IndexBuffer* NextIB = GetNextObject();
-	if (NextIB)
-	{
-		NextIB->ReleaseUnderlyingResource();
-	}
+	FD3D12Buffer::ReleaseUnderlyingResource();
 }
 
 FIndexBufferRHIRef FD3D12DynamicRHI::RHICreateIndexBuffer(uint32 Stride, uint32 Size, uint32 InUsage, FRHIResourceCreateInfo& CreateInfo)
@@ -108,12 +75,14 @@ FIndexBufferRHIRef FD3D12DynamicRHI::RHICreateIndexBuffer(uint32 Stride, uint32
 
 void* FD3D12DynamicRHI::RHILockIndexBuffer(FRHICommandListImmediate& RHICmdList, FRHIIndexBuffer* IndexBufferRHI, uint32 Offset, uint32 Size, EResourceLockMode LockMode)
 {
-	return LockBuffer(&RHICmdList, FD3D12DynamicRHI::ResourceCast(IndexBufferRHI), Offset, Size, LockMode);
+	FD3D12IndexBuffer* Buffer = FD3D12DynamicRHI::ResourceCast(IndexBufferRHI);
+	return LockBuffer(&RHICmdList, Buffer, Buffer->GetSize(), Buffer->GetUsage(), Offset, Size, LockMode);
 }
 
 void FD3D12DynamicRHI::RHIUnlockIndexBuffer(FRHICommandListImmediate& RHICmdList, FRHIIndexBuffer* IndexBufferRHI)
 {
-	UnlockBuffer(&RHICmdList, FD3D12DynamicRHI::ResourceCast(IndexBufferRHI));
+	FD3D12IndexBuffer* Buffer = FD3D12DynamicRHI::ResourceCast(IndexBufferRHI);
+	UnlockBuffer(&RHICmdList, Buffer, Buffer->GetUsage());
 }
 
 void FD3D12DynamicRHI::RHITransferIndexBufferUnderlyingResource(FRHIIndexBuffer* DestIndexBuffer, FRHIIndexBuffer* SrcIndexBuffer)
@@ -163,7 +132,7 @@ FIndexBufferRHIRef FD3D12DynamicRHI::CreateAndLockIndexBuffer_RenderThread(class
 		Buffer->SetCommitted(false);
 	}
 
-	OutDataBuffer = LockBuffer(&RHICmdList, Buffer, 0, Size, RLM_WriteOnly);
+	OutDataBuffer = LockBuffer(&RHICmdList, Buffer, Buffer->GetSize(), Buffer->GetUsage(), 0, Size, RLM_WriteOnly);
 
 	return Buffer;
 }
diff --git a/Engine/Source/Runtime/D3D12RHI/Private/D3D12RHIPrivate.h b/Engine/Source/Runtime/D3D12RHI/Private/D3D12RHIPrivate.h
--- a/Engine/Source/Runtime/D3D12RHI/Private/D3D12RHIPrivate.h
+++ b/Engine/Source/Runtime/D3D12RHI/Private/D3D12RHIPrivate.h
@@ -225,31 +225,6 @@ public:
 	{
 		return static_cast<typename TD3D12ResourceTraits<TRHIType>::TConcreteType*>(Resource);
 	}
-
-	template<typename TRHIType>
-	static FORCEINLINE_DEBUGGABLE typename TD3D12ResourceTraits<TRHIType>::TConcreteType* ResourceCast(TRHIType* Resource, uint32 GPUIndex)
-	{
-#if !WITH_MGPU
-		return ResourceCast(Resource);
-#else
-		typename TD3D12ResourceTraits<TRHIType>::TConcreteType* Object = ResourceCast(Resource);
-		if (GNumExplicitGPUsForRendering > 1)
-		{
-			if (!Object)
-			{
-				return nullptr;
-			}
-
-			while (Object && Object->GetParentDevice()->GetGPUIndex() != GPUIndex)
-			{
-				Object = Object->GetNextObject();
-			}
-
-			check(Object);
-		}
-		return Object;
-#endif
-	}
 	
 	virtual FD3D12CommandContext* CreateCommandContext(FD3D12Device* InParent, FD3D12SubAllocatedOnlineHeap::SubAllocationDesc& SubHeapDesc, bool InIsDefaultContext, bool InIsAsyncComputeContext = false);
 	virtual ID3D12CommandQueue* CreateCommandQueue(FD3D12Device* Device, const D3D12_COMMAND_QUEUE_DESC& Desc);
@@ -475,11 +450,8 @@ public:
 	bool IsQuadBufferStereoEnabled() const;
 	void DisableQuadBufferStereo();
 
-	template<class BufferType>
-	void* LockBuffer(FRHICommandListImmediate* RHICmdList, BufferType* Buffer, uint32 Offset, uint32 Size, EResourceLockMode LockMode);
-
-	template<class BufferType>
-	void UnlockBuffer(FRHICommandListImmediate* RHICmdList, BufferType* Buffer);
+	void* LockBuffer(FRHICommandListImmediate* RHICmdList, FD3D12Buffer* Buffer, uint32 BufferSize, uint32 BufferUsage, uint32 Offset, uint32 Size, EResourceLockMode LockMode);
+	void UnlockBuffer(FRHICommandListImmediate* RHICmdList, FD3D12Buffer* Buffer, uint32 BufferUsage);
 
 	static inline bool ShouldDeferBufferLockOperation(FRHICommandListImmediate* RHICmdList)
 	{
diff --git a/Engine/Source/Runtime/D3D12RHI/Private/D3D12RayTracing.cpp b/Engine/Source/Runtime/D3D12RHI/Private/D3D12RayTracing.cpp
--- a/Engine/Source/Runtime/D3D12RHI/Private/D3D12RayTracing.cpp
+++ b/Engine/Source/Runtime/D3D12RHI/Private/D3D12RayTracing.cpp
@@ -1384,18 +1384,13 @@ public:
 	// #dxr_todo UE-72159: remove resources from this set when SBT slot entries are replaced
 	Experimental::TSherwoodSet<void*> ReferencedD3D12ResourceSet;
 	TArray<TRefCountPtr<FD3D12Resource>> ReferencedD3D12Resources;
-	TArray<TRefCountPtr<FRHIResource>> ReferencedResources;
-	void AddResourceReference(FD3D12Resource* D3D12Resource, FRHIResource* Resource)
+	void AddResourceReference(FD3D12Resource* D3D12Resource)
 	{
 		bool bIsAlreadyInSet = false;
 		ReferencedD3D12ResourceSet.Add(D3D12Resource, &bIsAlreadyInSet);
 		if (!bIsAlreadyInSet)
 		{
 			ReferencedD3D12Resources.Add(D3D12Resource);
-			if (Resource)
-			{
-				ReferencedResources.Add(Resource);
-			}
 		}
 	}
 	void UpdateResidency(FD3D12CommandContext& CommandContext)
@@ -2094,7 +2089,7 @@ void FD3D12RayTracingGeometry::TransitionBuffers(FD3D12CommandContext& CommandCo
 	// Transition vertex and index resources..
 	if (RHIIndexBuffer)
 	{
-		FD3D12IndexBuffer* IndexBuffer = CommandContext.RetrieveObject<FD3D12IndexBuffer>(RHIIndexBuffer.GetReference());
+		FD3D12Buffer* IndexBuffer = CommandContext.RetrieveObject<FD3D12Buffer>(RHIIndexBuffer.GetReference());
 		if (IndexBuffer->GetResource()->RequiresResourceStateTracking())
 		{
 			FD3D12DynamicRHI::TransitionResource(CommandContext.CommandListHandle, IndexBuffer->GetResource(), D3D12_RESOURCE_STATE_NON_PIXEL_SHADER_RESOURCE, 0);
@@ -2104,7 +2099,7 @@ void FD3D12RayTracingGeometry::TransitionBuffers(FD3D12CommandContext& CommandCo
 	for (const FRayTracingGeometrySegment& Segment : Segments)
 	{
 		const FVertexBufferRHIRef& RHIVertexBuffer = Segment.VertexBuffer;
-		FD3D12VertexBuffer* VertexBuffer = CommandContext.RetrieveObject<FD3D12VertexBuffer>(RHIVertexBuffer.GetReference());
+		FD3D12Buffer* VertexBuffer = CommandContext.RetrieveObject<FD3D12Buffer>(RHIVertexBuffer.GetReference());
 		if (VertexBuffer->GetResource()->RequiresResourceStateTracking())
 		{
 			FD3D12DynamicRHI::TransitionResource(CommandContext.CommandListHandle, VertexBuffer->GetResource(), D3D12_RESOURCE_STATE_NON_PIXEL_SHADER_RESOURCE, 0);
@@ -2154,8 +2149,8 @@ void FD3D12RayTracingGeometry::BuildAccelerationStructure(FD3D12CommandContext&
 
 	Descs.Reserve(Segments.Num());
 
-	FD3D12IndexBuffer* IndexBuffer = CommandContext.RetrieveObject<FD3D12IndexBuffer>(RHIIndexBuffer.GetReference());
-	FD3D12VertexBuffer* NullTransformBufferD3D12 = CommandContext.RetrieveObject<FD3D12VertexBuffer>(NullTransformBuffer.GetReference());
+	FD3D12Buffer* IndexBuffer = CommandContext.RetrieveObject<FD3D12Buffer>(RHIIndexBuffer.GetReference());
+	FD3D12Buffer* NullTransformBufferD3D12 = CommandContext.RetrieveObject<FD3D12Buffer>(NullTransformBuffer.GetReference());
 
 	for (const FRayTracingGeometrySegment& Segment : Segments)
 	{
@@ -2176,7 +2171,7 @@ void FD3D12RayTracingGeometry::BuildAccelerationStructure(FD3D12CommandContext&
 			Desc.Flags |= D3D12_RAYTRACING_GEOMETRY_FLAG_NO_DUPLICATE_ANYHIT_INVOCATION;
 		}
 
-		FD3D12VertexBuffer* VertexBuffer = CommandContext.RetrieveObject<FD3D12VertexBuffer>(Segment.VertexBuffer.GetReference());
+		FD3D12Buffer* VertexBuffer = CommandContext.RetrieveObject<FD3D12Buffer>(Segment.VertexBuffer.GetReference());
 			   
 		if (GeometryType == D3D12_RAYTRACING_GEOMETRY_TYPE_TRIANGLES)
 		{
@@ -2559,7 +2554,7 @@ void FD3D12RayTracingScene::BuildAccelerationStructure(FD3D12CommandContext& Com
 			sizeof(D3D12_RAYTRACING_INSTANCE_DESC), InstanceBufferDesc.Width, BUF_UnorderedAccess, CreateInfo);
 
 		D3D12_RAYTRACING_INSTANCE_DESC* MappedData = (D3D12_RAYTRACING_INSTANCE_DESC*)Adapter->GetOwningRHI()->LockBuffer(
-			nullptr, InstanceBuffer.GetReference(), 0, InstanceBufferDesc.Width, RLM_WriteOnly);
+			nullptr, InstanceBuffer.GetReference(), InstanceBuffer->GetSize(), InstanceBuffer->GetUsage(), 0, InstanceBufferDesc.Width, RLM_WriteOnly);
 
 		check(MappedData);
 
@@ -2639,7 +2634,7 @@ void FD3D12RayTracingScene::BuildAccelerationStructure(FD3D12CommandContext& Com
 			TotalPrimitiveCount += Geometry->TotalPrimitiveCount * NumTransforms;
 		}
 
-		Adapter->GetOwningRHI()->UnlockBuffer(nullptr, InstanceBuffer.GetReference());
+		Adapter->GetOwningRHI()->UnlockBuffer(nullptr, InstanceBuffer.GetReference(), InstanceBuffer->GetUsage());
 
 		InstanceBuffer->GetResource()->UpdateResidency(CommandContext.CommandListHandle);
 
@@ -2752,8 +2747,6 @@ FD3D12RayTracingShaderTable* FD3D12RayTracingScene::FindOrCreateShaderTable(cons
 	FD3D12RayTracingShaderTable* CreatedShaderTable = new FD3D12RayTracingShaderTable();
 	ID3D12Device5* RayTracingDevice = Device->GetRayTracingDevice();
 
-	const uint32 GPUIndex = Device->GetGPUIndex();
-
 	const uint32 NumHitGroupSlots = Pipeline->bAllowHitGroupIndexing ? NumTotalSegments * ShaderSlotsPerGeometrySegment : 0;
 
 	checkf(Pipeline->MaxLocalRootSignatureSize >= sizeof(FHitGroupSystemParameters), TEXT("All local root signatures are expected to contain ray tracing system root parameters (2x root buffers + 4x root DWORD)"));
@@ -2794,7 +2787,7 @@ FD3D12RayTracingShaderTable* FD3D12RayTracingScene::FindOrCreateShaderTable(cons
 
 			const uint32 IndexStride = Geometry->IndexStride;
 
-			FD3D12IndexBuffer* IndexBuffer = FD3D12DynamicRHI::ResourceCast(Geometry->RHIIndexBuffer.GetReference(), GPUIndex);
+			FD3D12Buffer* IndexBuffer = FD3D12CommandContext::RetrieveObject<FD3D12Buffer>(Geometry->RHIIndexBuffer.GetReference(), Device);
 
 			// Here, we directly apply the vertex offset to the address but set the offset for index buffer using a RootConstant instead.
 			// Index Buffer has to be 4 bytes aligned.
@@ -2805,7 +2798,7 @@ FD3D12RayTracingShaderTable* FD3D12RayTracingScene::FindOrCreateShaderTable(cons
 
 			if (IndexBuffer)
 			{
-				CreatedShaderTable->AddResourceReference(IndexBuffer->ResourceLocation.GetResource(), IndexBuffer);
+				CreatedShaderTable->AddResourceReference(IndexBuffer->ResourceLocation.GetResource());
 			}
 
 			const uint32 NumSegments = Geometry->Segments.Num();
@@ -2813,8 +2806,8 @@ FD3D12RayTracingShaderTable* FD3D12RayTracingScene::FindOrCreateShaderTable(cons
 			{
 				const FRayTracingGeometrySegment& Segment = Geometry->Segments[SegmentIndex];
 
-				FD3D12VertexBuffer* VertexBuffer = FD3D12DynamicRHI::ResourceCast(Segment.VertexBuffer.GetReference(), GPUIndex);
-				CreatedShaderTable->AddResourceReference(VertexBuffer->ResourceLocation.GetResource(), VertexBuffer);
+				FD3D12Buffer* VertexBuffer = FD3D12CommandContext::RetrieveObject<FD3D12Buffer>(Segment.VertexBuffer.GetReference(), Device);
+				CreatedShaderTable->AddResourceReference(VertexBuffer->ResourceLocation.GetResource());
 
 				const D3D12_GPU_VIRTUAL_ADDRESS VertexBufferAddress = VertexBuffer->ResourceLocation.GetGPUVirtualAddress() + Segment.VertexBufferOffset;
 
@@ -2844,6 +2837,8 @@ FD3D12RayTracingShaderTable* FD3D12RayTracingScene::FindOrCreateShaderTable(cons
 		}
 	}
 
+	const uint32 GPUIndex = Device->GetGPUIndex();
+
 	ShaderTables[GPUIndex].Add(Pipeline, CreatedShaderTable);
 
 	return CreatedShaderTable;
@@ -3004,14 +2999,14 @@ struct FD3D12RayTracingLocalResourceBinder
 
 		FMemory::Memcpy(MappedData, Data, DataSize);
 
-		ShaderTable->AddResourceReference(ResourceLocation.GetResource(), nullptr);
+		ShaderTable->AddResourceReference(ResourceLocation.GetResource());
 
 		return ResourceLocation.GetGPUVirtualAddress();
 	}
 
 	void AddResourceReference(FD3D12Resource* D3D12Resource, FRHIResource* RHIResource)
 	{
-		ShaderTable->AddResourceReference(D3D12Resource, RHIResource);
+		ShaderTable->AddResourceReference(D3D12Resource);
 	}
 
 	void AddResourceTransition(FD3D12ShaderResourceView* SRV)
diff --git a/Engine/Source/Runtime/D3D12RHI/Private/D3D12StructuredBuffer.cpp b/Engine/Source/Runtime/D3D12RHI/Private/D3D12StructuredBuffer.cpp
--- a/Engine/Source/Runtime/D3D12RHI/Private/D3D12StructuredBuffer.cpp
+++ b/Engine/Source/Runtime/D3D12RHI/Private/D3D12StructuredBuffer.cpp
@@ -74,57 +74,14 @@ FD3D12StructuredBuffer::~FD3D12StructuredBuffer()
 	UpdateBufferStats<FD3D12StructuredBuffer>(&ResourceLocation, false);
 }
 
-void FD3D12StructuredBuffer::Rename(FD3D12ResourceLocation& NewLocation)
-{
-	FD3D12ResourceLocation::TransferOwnership(ResourceLocation, NewLocation);
-
-	FScopeLock Lock(&DynamicSRVsCS);
-	for(FD3D12BaseShaderResourceView* DynamicSRVBase : DynamicSRVs)
-	{
-		FD3D12ShaderResourceView* DynamicSRV = static_cast<FD3D12ShaderResourceView*>(DynamicSRVBase);
-		if (DynamicSRV->IsValid())
-		{
-			DynamicSRV->Rename(ResourceLocation);
-		}
-	}
-}
-
-void FD3D12StructuredBuffer::RenameLDAChain(FD3D12ResourceLocation& NewLocation)
-{
-	// Dynamic buffers use cross-node resources.
-	ensure(GetUsage() & BUF_AnyDynamic);
-	Rename(NewLocation);
-
-	if (GNumExplicitGPUsForRendering > 1)
-	{
-		// This currently crashes at exit time because NewLocation isn't tracked in the right allocator.
-		ensure(IsHeadLink());
-		ensure(GetParentDevice() == NewLocation.GetParentDevice());
-
-		// Update all of the resources in the LDA chain to reference this cross-node resource
-		for (FD3D12StructuredBuffer* NextBuffer = GetNextObject(); NextBuffer; NextBuffer = NextBuffer->GetNextObject())
-		{
-			FD3D12ResourceLocation::ReferenceNode(NextBuffer->GetParentDevice(), NextBuffer->ResourceLocation, ResourceLocation);
-
-			FScopeLock Lock(&NextBuffer->DynamicSRVsCS);
-			for (FD3D12BaseShaderResourceView* DynamicSRVBase : NextBuffer->DynamicSRVs)
-			{
-				FD3D12ShaderResourceView* DynamicSRV = static_cast<FD3D12ShaderResourceView*>(DynamicSRVBase);
-				if (DynamicSRV->IsValid())
-				{
-					DynamicSRV->Rename(NextBuffer->ResourceLocation);
-				}
-			}
-		}
-	}
-}
-
 void* FD3D12DynamicRHI::RHILockStructuredBuffer(FRHICommandListImmediate& RHICmdList, FRHIStructuredBuffer* StructuredBufferRHI, uint32 Offset, uint32 Size, EResourceLockMode LockMode)
 {
-	return LockBuffer(&RHICmdList, FD3D12DynamicRHI::ResourceCast(StructuredBufferRHI), Offset, Size, LockMode);
+	FD3D12StructuredBuffer* Buffer = FD3D12DynamicRHI::ResourceCast(StructuredBufferRHI);
+	return LockBuffer(&RHICmdList, Buffer, Buffer->GetSize(), Buffer->GetUsage(), Offset, Size, LockMode);
 }
 
 void FD3D12DynamicRHI::RHIUnlockStructuredBuffer(FRHICommandListImmediate& RHICmdList, FRHIStructuredBuffer* StructuredBufferRHI)
 {
-	UnlockBuffer(&RHICmdList, FD3D12DynamicRHI::ResourceCast(StructuredBufferRHI));
+	FD3D12StructuredBuffer* Buffer = FD3D12DynamicRHI::ResourceCast(StructuredBufferRHI);
+	UnlockBuffer(&RHICmdList, Buffer, Buffer->GetUsage());
 }
diff --git a/Engine/Source/Runtime/D3D12RHI/Private/D3D12VertexBuffer.cpp b/Engine/Source/Runtime/D3D12RHI/Private/D3D12VertexBuffer.cpp
--- a/Engine/Source/Runtime/D3D12RHI/Private/D3D12VertexBuffer.cpp
+++ b/Engine/Source/Runtime/D3D12RHI/Private/D3D12VertexBuffer.cpp
@@ -44,72 +44,20 @@ FD3D12VertexBuffer::~FD3D12VertexBuffer()
 	}
 }
 
-void FD3D12VertexBuffer::Rename(FD3D12ResourceLocation& NewLocation)
-{
-	FD3D12ResourceLocation::TransferOwnership(ResourceLocation, NewLocation);
-
-	FScopeLock Lock(&DynamicSRVsCS);
-	for (FD3D12BaseShaderResourceView* DynamicSRVBase : DynamicSRVs)
-	{
-		FD3D12ShaderResourceView* DynamicSRV = static_cast<FD3D12ShaderResourceView*>(DynamicSRVBase);
-		if (DynamicSRV->IsValid())
-		{
-			DynamicSRV->Rename(ResourceLocation);
-		}
-	}
-}
-
-void FD3D12VertexBuffer::RenameLDAChain(FD3D12ResourceLocation& NewLocation)
-{
-	ensure(GetUsage() & BUF_AnyDynamic);
-	Rename(NewLocation);
-
-	if (GNumExplicitGPUsForRendering > 1)
-	{
-		// Mutli-GPU support : renaming the LDA only works if we start we the head link. Otherwise Rename() must be used per GPU.
-		ensure(IsHeadLink());
-		ensure(GetParentDevice() == NewLocation.GetParentDevice());
-
-		// Update all of the resources in the LDA chain to reference this cross-node resource
-		for (FD3D12VertexBuffer* NextBuffer = GetNextObject(); NextBuffer; NextBuffer = NextBuffer->GetNextObject())
-		{
-			FD3D12ResourceLocation::ReferenceNode(NextBuffer->GetParentDevice(), NextBuffer->ResourceLocation, ResourceLocation);
-
-			FScopeLock Lock(&NextBuffer->DynamicSRVsCS);
-			for (FD3D12BaseShaderResourceView* DynamicSRVBase : NextBuffer->DynamicSRVs)
-			{
-				FD3D12ShaderResourceView* DynamicSRV = static_cast<FD3D12ShaderResourceView*>(DynamicSRVBase);
-				if (DynamicSRV->IsValid())
-				{
-					DynamicSRV->Rename(NextBuffer->ResourceLocation);
-				}
-			}
-		}
-	}
-}
-
 void FD3D12VertexBuffer::Swap(FD3D12VertexBuffer& Other)
 {
 	check(!LockedData.bLocked && !Other.LockedData.bLocked);
 	FRHIVertexBuffer::Swap(Other);
 	FD3D12BaseShaderResource::Swap(Other);
 	FD3D12TransientResource::Swap(Other);
-	FD3D12LinkedAdapterObject<FD3D12VertexBuffer>::Swap(Other);
+	FD3D12LinkedAdapterObject<FD3D12Buffer>::Swap(Other);
 }
 
 void FD3D12VertexBuffer::ReleaseUnderlyingResource()
 {
-	check(!LockedData.bLocked && ResourceLocation.IsValid());
 	UpdateBufferStats<FD3D12VertexBuffer>(&ResourceLocation, false);
-	ResourceLocation.Clear();
 	FRHIVertexBuffer::ReleaseUnderlyingResource();
-	RemoveAllDynamicSRVs();
-
-	FD3D12VertexBuffer* NextVB = GetNextObject();
-	if (NextVB)
-	{
-		NextVB->ReleaseUnderlyingResource();
-	}
+	FD3D12Buffer::ReleaseUnderlyingResource();
 }
 
 FVertexBufferRHIRef FD3D12DynamicRHI::RHICreateVertexBuffer(uint32 Size, uint32 InUsage, FRHIResourceCreateInfo& CreateInfo)
@@ -134,12 +82,14 @@ FVertexBufferRHIRef FD3D12DynamicRHI::RHICreateVertexBuffer(uint32 Size, uint32
 
 void* FD3D12DynamicRHI::RHILockVertexBuffer(FRHICommandListImmediate& RHICmdList, FRHIVertexBuffer* VertexBufferRHI, uint32 Offset, uint32 Size, EResourceLockMode LockMode)
 {
-	return LockBuffer(&RHICmdList, FD3D12DynamicRHI::ResourceCast(VertexBufferRHI), Offset, Size, LockMode);
+	FD3D12VertexBuffer* Buffer = FD3D12DynamicRHI::ResourceCast(VertexBufferRHI);
+	return LockBuffer(&RHICmdList, Buffer, Buffer->GetSize(), Buffer->GetUsage(), Offset, Size, LockMode);
 }
 
 void FD3D12DynamicRHI::RHIUnlockVertexBuffer(FRHICommandListImmediate& RHICmdList, FRHIVertexBuffer* VertexBufferRHI)
 {
-	UnlockBuffer(&RHICmdList, FD3D12DynamicRHI::ResourceCast(VertexBufferRHI));
+	FD3D12VertexBuffer* Buffer = FD3D12DynamicRHI::ResourceCast(VertexBufferRHI);
+	UnlockBuffer(&RHICmdList, Buffer, Buffer->GetUsage());
 }
 
 FVertexBufferRHIRef FD3D12DynamicRHI::CreateVertexBuffer_RenderThread(FRHICommandListImmediate& RHICmdList, uint32 Size, uint32 InUsage, FRHIResourceCreateInfo& CreateInfo)
@@ -164,8 +114,12 @@ FVertexBufferRHIRef FD3D12DynamicRHI::CreateVertexBuffer_RenderThread(FRHIComman
 
 void FD3D12DynamicRHI::RHICopyVertexBuffer(FRHIVertexBuffer* SourceBufferRHI, FRHIVertexBuffer* DestBufferRHI)
 {
-	FD3D12VertexBuffer*  SourceBuffer = FD3D12DynamicRHI::ResourceCast(SourceBufferRHI);
-	FD3D12VertexBuffer*  DestBuffer = FD3D12DynamicRHI::ResourceCast(DestBufferRHI);
+	FD3D12VertexBuffer* SrcBuffer = FD3D12DynamicRHI::ResourceCast(SourceBufferRHI);
+	FD3D12VertexBuffer* DstBuffer = FD3D12DynamicRHI::ResourceCast(DestBufferRHI);
+	check(SrcBuffer->GetSize() == DstBuffer->GetSize());
+
+	FD3D12Buffer* SourceBuffer = SrcBuffer;
+	FD3D12Buffer* DestBuffer = DstBuffer;
 
 	while (SourceBuffer && DestBuffer)
 	{
@@ -179,7 +133,6 @@ void FD3D12DynamicRHI::RHICopyVertexBuffer(FRHIVertexBuffer* SourceBufferRHI, FR
 		D3D12_RESOURCE_DESC const& DestBufferDesc = pDestResource->GetDesc();
 
 		check(SourceBufferDesc.Width == DestBufferDesc.Width);
-		check(SourceBuffer->GetSize() == DestBuffer->GetSize());
 
 		FD3D12CommandContext& Context = Device->GetDefaultCommandContext();
 		Context.numCopies++;
@@ -214,8 +167,8 @@ void FD3D12DynamicRHI::RHITransferVertexBufferUnderlyingResource(FRHIVertexBuffe
 #if D3D12_RHI_RAYTRACING
 void FD3D12CommandContext::RHICopyBufferRegion(FRHIVertexBuffer* DestBufferRHI, uint64 DstOffset, FRHIVertexBuffer* SourceBufferRHI, uint64 SrcOffset, uint64 NumBytes)
 {
-	FD3D12VertexBuffer*  SourceBuffer = RetrieveObject<FD3D12VertexBuffer>(SourceBufferRHI);
-	FD3D12VertexBuffer*  DestBuffer = RetrieveObject<FD3D12VertexBuffer>(DestBufferRHI);
+	FD3D12Buffer* SourceBuffer = RetrieveObject<FD3D12Buffer>(SourceBufferRHI);
+	FD3D12Buffer* DestBuffer = RetrieveObject<FD3D12Buffer>(DestBufferRHI);
 
 	FD3D12Device* Device = SourceBuffer->GetParentDevice();
 	check(Device == DestBuffer->GetParentDevice());
@@ -260,8 +213,8 @@ void FD3D12CommandContext::RHICopyBufferRegions(const TArrayView<const FCopyBuff
 	// Transition buffers to copy states
 	for (auto& Param : Params)
 	{
-		FD3D12VertexBuffer*  SourceBuffer = RetrieveObject<FD3D12VertexBuffer>(Param.SourceBuffer);
-		FD3D12VertexBuffer*  DestBuffer = RetrieveObject<FD3D12VertexBuffer>(Param.DestBuffer);
+		FD3D12Buffer* SourceBuffer = RetrieveObject<FD3D12Buffer>(Param.SourceBuffer);
+		FD3D12Buffer* DestBuffer = RetrieveObject<FD3D12Buffer>(Param.DestBuffer);
 		check(SourceBuffer);
 		check(DestBuffer);
 
@@ -350,8 +303,8 @@ void FD3D12CommandContext::RHICopyBufferRegions(const TArrayView<const FCopyBuff
 
 	for (auto& Param : Params)
 	{
-		FD3D12VertexBuffer*  SourceBuffer = RetrieveObject<FD3D12VertexBuffer>(Param.SourceBuffer);
-		FD3D12VertexBuffer*  DestBuffer = RetrieveObject<FD3D12VertexBuffer>(Param.DestBuffer);
+		FD3D12Buffer* SourceBuffer = RetrieveObject<FD3D12Buffer>(Param.SourceBuffer);
+		FD3D12Buffer* DestBuffer = RetrieveObject<FD3D12Buffer>(Param.DestBuffer);
 		uint64 SrcOffset = Param.SrcOffset;
 		uint64 DstOffset = Param.DstOffset;
 		uint64 NumBytes = Param.NumBytes;
@@ -395,7 +348,7 @@ FVertexBufferRHIRef FD3D12DynamicRHI::CreateAndLockVertexBuffer_RenderThread(FRH
 		// TODO: this should ideally be set in platform-independent code, since this tracking is for the high level
 		Buffer->SetCommitted(false);
 	}
-	OutDataBuffer = LockBuffer(&RHICmdList, Buffer, 0, Size, RLM_WriteOnly);
+	OutDataBuffer = LockBuffer(&RHICmdList, Buffer, Buffer->GetSize(), Buffer->GetUsage(), 0, Size, RLM_WriteOnly);
 
 	return Buffer;
 }
diff --git a/Engine/Source/Runtime/D3D12RHI/Private/D3D12View.cpp b/Engine/Source/Runtime/D3D12RHI/Private/D3D12View.cpp
--- a/Engine/Source/Runtime/D3D12RHI/Private/D3D12View.cpp
+++ b/Engine/Source/Runtime/D3D12RHI/Private/D3D12View.cpp
@@ -377,15 +377,16 @@ void FD3D12DynamicRHI::RHIUpdateShaderResourceView(FRHIShaderResourceView* SRV,
 		D3D12_SHADER_RESOURCE_VIEW_DESC SRVDesc = GetVertexBufferSRVDesc(VBD3D12, Stride, Format, 0, UINT32_MAX);
 
 		// Rename the SRV to view on the new vertex buffer
-		while (VBD3D12)
+		FD3D12Buffer* Buffer = VBD3D12;
+		while (Buffer)
 		{
-			FD3D12Device* ParentDevice = VBD3D12->GetParentDevice();
-			SRVD3D12->Initialize(ParentDevice, SRVDesc, VBD3D12->ResourceLocation, Stride);
-			VBD3D12->AddDynamicSRV(SRVD3D12);
-			VBD3D12 = VBD3D12->GetNextObject();
-			if (VBD3D12 && !SRVD3D12->GetNextObject())
+			FD3D12Device* ParentDevice = Buffer->GetParentDevice();
+			SRVD3D12->Initialize(ParentDevice, SRVDesc, Buffer->ResourceLocation, Stride);
+			Buffer->AddDynamicSRV(SRVD3D12);
+			Buffer = Buffer->GetNextObject();
+			if (Buffer && !SRVD3D12->GetNextObject())
 			{
-				SRVD3D12->SetNextObject(new FD3D12ShaderResourceView(VBD3D12->GetParentDevice()));
+				SRVD3D12->SetNextObject(new FD3D12ShaderResourceView(Buffer->GetParentDevice()));
 			}
 			SRVD3D12 = SRVD3D12->GetNextObject();
 		}
@@ -403,14 +404,15 @@ void FD3D12DynamicRHI::RHIUpdateShaderResourceView(FRHIShaderResourceView* SRV,
 		const uint32 Stride = IBD3D12->GetStride();
 
 		// Rename the SRV to view on the new index buffer
-		while (IBD3D12)
+		FD3D12Buffer* Buffer = IBD3D12;
+		while (Buffer)
 		{
-			FD3D12Device* ParentDevice = IBD3D12->GetParentDevice();
-			SRVD3D12->Initialize(ParentDevice, SRVDesc, IBD3D12->ResourceLocation, Stride);
-			IBD3D12 = IBD3D12->GetNextObject();
-			if (IBD3D12 && !SRVD3D12->GetNextObject())
+			FD3D12Device* ParentDevice = Buffer->GetParentDevice();
+			SRVD3D12->Initialize(ParentDevice, SRVDesc, Buffer->ResourceLocation, Stride);
+			Buffer = Buffer->GetNextObject();
+			if (Buffer && !SRVD3D12->GetNextObject())
 			{
-				SRVD3D12->SetNextObject(new FD3D12ShaderResourceView(IBD3D12->GetParentDevice()));
+				SRVD3D12->SetNextObject(new FD3D12ShaderResourceView(Buffer->GetParentDevice()));
 			}
 			SRVD3D12 = SRVD3D12->GetNextObject();
 		}
diff --git a/Engine/Source/Runtime/D3D12RHI/Public/D3D12Resources.h b/Engine/Source/Runtime/D3D12RHI/Public/D3D12Resources.h
--- a/Engine/Source/Runtime/D3D12RHI/Public/D3D12Resources.h
+++ b/Engine/Source/Runtime/D3D12RHI/Public/D3D12Resources.h
@@ -806,26 +806,43 @@ public:
 };
 #endif
 
+class FD3D12Buffer : public FD3D12BaseShaderResource, public FD3D12TransientResource, public FD3D12LinkedAdapterObject<FD3D12Buffer>
+{
+public:
+	FD3D12Buffer(FD3D12Device* InParent)
+		: FD3D12BaseShaderResource(InParent)
+		, LockedData(InParent)
+	{
+	}
+	virtual ~FD3D12Buffer()
+	{
+	}
+
+	void Rename(FD3D12ResourceLocation& NewLocation);
+	void RenameLDAChain(FD3D12ResourceLocation& NewLocation);
+
+	void ReleaseUnderlyingResource();
+
+	FD3D12LockedResource LockedData;
+};
+
 /** Index buffer resource class that stores stride information. */
-class FD3D12IndexBuffer : public FRHIIndexBuffer, public FD3D12BaseShaderResource, public FD3D12TransientResource, public FD3D12LinkedAdapterObject<FD3D12IndexBuffer>
+class FD3D12IndexBuffer : public FRHIIndexBuffer, public FD3D12Buffer
 {
 public:
 	FD3D12IndexBuffer()
-		: FD3D12BaseShaderResource(nullptr)
-		, LockedData(nullptr)
-	{}
+		: FD3D12Buffer(nullptr)
+	{
+	}
 
 	FD3D12IndexBuffer(FD3D12Device* InParent, uint32 InStride, uint32 InSize, uint32 InUsage)
 		: FRHIIndexBuffer(InStride, InSize, InUsage)
-		, FD3D12BaseShaderResource(InParent)
-		, LockedData(InParent)
-	{}
+		, FD3D12Buffer(InParent)
+	{
+	}
 
 	virtual ~FD3D12IndexBuffer();
 
-	void Rename(FD3D12ResourceLocation& NewLocation);
-	void RenameLDAChain(FD3D12ResourceLocation& NewLocation);
-
 	void Swap(FD3D12IndexBuffer& Other);
 
 	void ReleaseUnderlyingResource();
@@ -843,28 +860,20 @@ public:
 	{
 		return FRHIResource::GetRefCount();
 	}
-
-	FD3D12LockedResource LockedData;
 };
 
 class FD3D12ShaderResourceView;
 
 /** Structured buffer resource class. */
-class FD3D12StructuredBuffer : public FRHIStructuredBuffer, public FD3D12BaseShaderResource, public FD3D12TransientResource, public FD3D12LinkedAdapterObject<FD3D12StructuredBuffer>
+class FD3D12StructuredBuffer : public FRHIStructuredBuffer, public FD3D12Buffer
 {
 public:
-	// Current SRV
-
 	FD3D12StructuredBuffer(FD3D12Device* InParent, uint32 InStride, uint32 InSize, uint32 InUsage)
 		: FRHIStructuredBuffer(InStride, InSize, InUsage)
-		, FD3D12BaseShaderResource(InParent)
-		, LockedData(InParent)
+		, FD3D12Buffer(InParent)
 	{
 	}
 
-	void Rename(FD3D12ResourceLocation& NewLocation);
-	void RenameLDAChain(FD3D12ResourceLocation& NewLocation);
-
 	virtual ~FD3D12StructuredBuffer();
 
 	// IRefCountedObject interface.
@@ -880,32 +889,25 @@ public:
 	{
 		return FRHIResource::GetRefCount();
 	}
-
-	FD3D12LockedResource LockedData;
 };
 
 /** Vertex buffer resource class. */
-class FD3D12VertexBuffer : public FRHIVertexBuffer, public FD3D12BaseShaderResource, public FD3D12TransientResource, public FD3D12LinkedAdapterObject<FD3D12VertexBuffer>
+class FD3D12VertexBuffer : public FRHIVertexBuffer, public FD3D12Buffer
 {
 public:
 	FD3D12VertexBuffer()
-		: FD3D12BaseShaderResource(nullptr)
-		, LockedData(nullptr)
+		: FD3D12Buffer(nullptr)
 	{}
 
 	FD3D12VertexBuffer(FD3D12Device* InParent, uint32 InStride, uint32 InSize, uint32 InUsage)
 		: FRHIVertexBuffer(InSize, InUsage)
-		, FD3D12BaseShaderResource(InParent)
-		, LockedData(InParent)
+		, FD3D12Buffer(InParent)
 	{
 		UNREFERENCED_PARAMETER(InStride);
 	}
 
 	virtual ~FD3D12VertexBuffer();
 
-	void Rename(FD3D12ResourceLocation& NewLocation);
-	void RenameLDAChain(FD3D12ResourceLocation& NewLocation);
-
 	void Swap(FD3D12VertexBuffer& Other);
 
 	void ReleaseUnderlyingResource();
@@ -923,8 +925,6 @@ public:
 	{
 		return FRHIResource::GetRefCount();
 	}
-
-	FD3D12LockedResource LockedData;
 };
 
 template<class BufferType>
