Moved ray tracing vertex and index buffer tracking from SBT to the scene. This saves 2/3 of the time spent in FD3D12RayTracingScene::FindOrCreateShaderTable().

diff --git a/Engine/Source/Runtime/D3D12RHI/Private/D3D12RayTracing.cpp b/Engine/Source/Runtime/D3D12RHI/Private/D3D12RayTracing.cpp
--- a/Engine/Source/Runtime/D3D12RHI/Private/D3D12RayTracing.cpp
+++ b/Engine/Source/Runtime/D3D12RHI/Private/D3D12RayTracing.cpp
@@ -8,7 +8,6 @@
 #include "Containers/DynamicRHIResourceArray.h"
 #include "Experimental/Containers/SherwoodHashTable.h"
 #include "BuiltInRayTracingShaders.h"
-#include "RayTracingBuiltInResources.h"
 #include "Hash/CityHash.h"
 #include "HAL/CriticalSection.h"
 #include "HAL/IConsoleManager.h"
@@ -89,14 +88,6 @@ DECLARE_MEMORY_STAT(TEXT("Total Used Video Memory"), STAT_D3D12RayTracingUsedVid
 DECLARE_CYCLE_STAT(TEXT("RTPSO Compile Shader"), STAT_RTPSO_CompileShader, STATGROUP_D3D12RayTracing);
 DECLARE_CYCLE_STAT(TEXT("RTPSO Create Pipeline"), STAT_RTPSO_CreatePipeline, STATGROUP_D3D12RayTracing);
 
-// Built-in local root parameters that are always bound to all hit shaders
-struct FHitGroupSystemParameters
-{
-	D3D12_GPU_VIRTUAL_ADDRESS IndexBuffer;
-	D3D12_GPU_VIRTUAL_ADDRESS VertexBuffer;
-	FHitGroupSystemRootConstants RootConstants;
-};
-
 struct FD3D12ShaderIdentifier
 {
 	uint64 Data[4] = {~0ull, ~0ull, ~0ull, ~0ull};
@@ -1906,25 +1897,53 @@ FRayTracingPipelineStateRHIRef FD3D12DynamicRHI::RHICreateRayTracingPipelineStat
 
 FRayTracingGeometryRHIRef FD3D12DynamicRHI::RHICreateRayTracingGeometry(const FRayTracingGeometryInitializer& Initializer)
 {
+	return new FD3D12RayTracingGeometry(Initializer);
+}
+
+FRayTracingSceneRHIRef FD3D12DynamicRHI::RHICreateRayTracingScene(const FRayTracingSceneInitializer& Initializer)
+{
+	TRACE_CPUPROFILER_EVENT_SCOPE(CreateRayTracingScene);
+
+	FD3D12Adapter& Adapter = GetAdapter();
+
+	return new FD3D12RayTracingScene(&Adapter, Initializer);
+}
+
+FVertexBufferRHIRef FD3D12RayTracingGeometry::NullTransformBuffer;
+
+FD3D12RayTracingGeometry::FD3D12RayTracingGeometry(const FRayTracingGeometryInitializer& Initializer)
+{
+	INC_DWORD_STAT(STAT_D3D12RayTracingAllocatedBLAS);
+
+	if(!FD3D12RayTracingGeometry::NullTransformBuffer.IsValid())
+	{
+		TResourceArray<float> NullTransformData;
+		NullTransformData.SetNumZeroed(12);
+
+		FRHIResourceCreateInfo CreateInfo;
+		CreateInfo.ResourceArray = &NullTransformData;
+
+		FD3D12RayTracingGeometry::NullTransformBuffer = RHICreateVertexBuffer(NullTransformData.GetResourceDataSize(), BUF_Static, CreateInfo);
+	}
+
+	FMemory::Memzero(PostBuildInfoBufferReadbackFences);
+
 	if (Initializer.IndexBuffer)
 	{
 		checkf(Initializer.IndexBuffer->GetStride() == 2 || Initializer.IndexBuffer->GetStride() == 4, TEXT("Index buffer must be 16 or 32 bit."));
 	}
 
-	FD3D12RayTracingGeometry* Geometry = new FD3D12RayTracingGeometry();
-
-	uint32 IndexStride = Initializer.IndexBuffer ? Initializer.IndexBuffer->GetStride() : 0; // stride 0 means implicit triangle list for non-indexed geometry
-	Geometry->IndexStride = IndexStride;
-	Geometry->IndexOffsetInBytes = Initializer.IndexBufferOffset;
-	Geometry->TotalPrimitiveCount = Initializer.TotalPrimitiveCount;
+	IndexStride = Initializer.IndexBuffer ? Initializer.IndexBuffer->GetStride() : 0; // stride 0 means implicit triangle list for non-indexed geometry
+	IndexOffsetInBytes = Initializer.IndexBufferOffset;
+	TotalPrimitiveCount = Initializer.TotalPrimitiveCount;
 
 	switch (Initializer.GeometryType)
 	{
 	case RTGT_Triangles:
-		Geometry->GeometryType = D3D12_RAYTRACING_GEOMETRY_TYPE_TRIANGLES;
+		GeometryType = D3D12_RAYTRACING_GEOMETRY_TYPE_TRIANGLES;
 		break;
 	case RTGT_Procedural:
-		Geometry->GeometryType = D3D12_RAYTRACING_GEOMETRY_TYPE_PROCEDURAL_PRIMITIVE_AABBS;
+		GeometryType = D3D12_RAYTRACING_GEOMETRY_TYPE_PROCEDURAL_PRIMITIVE_AABBS;
 		break;
 	default:
 		checkf(false, TEXT("Unexpected ray tracing geometry type"));
@@ -1932,43 +1951,44 @@ FRayTracingGeometryRHIRef FD3D12DynamicRHI::RHICreateRayTracingGeometry(const FR
 
 	if (Initializer.bFastBuild)
 	{
-		Geometry->BuildFlags = D3D12_RAYTRACING_ACCELERATION_STRUCTURE_BUILD_FLAG_PREFER_FAST_BUILD;
+		BuildFlags = D3D12_RAYTRACING_ACCELERATION_STRUCTURE_BUILD_FLAG_PREFER_FAST_BUILD;
 	}
 	else
 	{
-		Geometry->BuildFlags = D3D12_RAYTRACING_ACCELERATION_STRUCTURE_BUILD_FLAG_PREFER_FAST_TRACE;
+		BuildFlags = D3D12_RAYTRACING_ACCELERATION_STRUCTURE_BUILD_FLAG_PREFER_FAST_TRACE;
 	}
 
 	if (Initializer.bAllowUpdate)
 	{
-		Geometry->BuildFlags |= D3D12_RAYTRACING_ACCELERATION_STRUCTURE_BUILD_FLAG_ALLOW_UPDATE;
+		BuildFlags |= D3D12_RAYTRACING_ACCELERATION_STRUCTURE_BUILD_FLAG_ALLOW_UPDATE;
 	}
 
 	if (!Initializer.bFastBuild && !Initializer.bAllowUpdate)
 	{
-		Geometry->BuildFlags |= D3D12_RAYTRACING_ACCELERATION_STRUCTURE_BUILD_FLAG_ALLOW_COMPACTION;
+		BuildFlags |= D3D12_RAYTRACING_ACCELERATION_STRUCTURE_BUILD_FLAG_ALLOW_COMPACTION;
 	}
 
 	if (GRayTracingDebugForceBuildMode == 1)
 	{
-		Geometry->BuildFlags |= D3D12_RAYTRACING_ACCELERATION_STRUCTURE_BUILD_FLAG_PREFER_FAST_BUILD;
-		Geometry->BuildFlags &= ~D3D12_RAYTRACING_ACCELERATION_STRUCTURE_BUILD_FLAG_PREFER_FAST_TRACE;
+		BuildFlags |= D3D12_RAYTRACING_ACCELERATION_STRUCTURE_BUILD_FLAG_PREFER_FAST_BUILD;
+		BuildFlags &= ~D3D12_RAYTRACING_ACCELERATION_STRUCTURE_BUILD_FLAG_PREFER_FAST_TRACE;
 	}
 	else if (GRayTracingDebugForceBuildMode == 2)
 	{
-		Geometry->BuildFlags |= D3D12_RAYTRACING_ACCELERATION_STRUCTURE_BUILD_FLAG_PREFER_FAST_TRACE;
-		Geometry->BuildFlags &= ~D3D12_RAYTRACING_ACCELERATION_STRUCTURE_BUILD_FLAG_PREFER_FAST_BUILD;
+		BuildFlags |= D3D12_RAYTRACING_ACCELERATION_STRUCTURE_BUILD_FLAG_PREFER_FAST_TRACE;
+		BuildFlags &= ~D3D12_RAYTRACING_ACCELERATION_STRUCTURE_BUILD_FLAG_PREFER_FAST_BUILD;
 	}
 
 	checkf(Initializer.Segments.Num() > 0, TEXT("Ray tracing geometry must be initialized with at least one segment."));
-	Geometry->Segments = TArray<FRayTracingGeometrySegment>(Initializer.Segments.GetData(), Initializer.Segments.Num());
+	Segments = TArray<FRayTracingGeometrySegment>(Initializer.Segments.GetData(), Initializer.Segments.Num());
 
 #if DO_CHECK
 	static constexpr uint32 IndicesPerPrimitive = 3; // Only triangle meshes are supported
 
 	{
 		uint32 ComputedPrimitiveCountForValidation = 0;
-		for (const FRayTracingGeometrySegment& Segment : Geometry->Segments)
+
+		for (const FRayTracingGeometrySegment& Segment : Segments)
 		{
 			checkf(Segment.VertexBuffer, TEXT("Position vertex buffer is required for ray tracing geometry"));
 			checkf(Segment.VertexBufferStride, TEXT("Position vertex buffer is required for ray tracing geometry"));
@@ -1977,7 +1997,7 @@ FRayTracingGeometryRHIRef FD3D12DynamicRHI::RHICreateRayTracingGeometry(const FR
 			if (Initializer.GeometryType == RTGT_Triangles)
 			{
 				// #dxr_todo UE-72160: temporary constraints on vertex and index buffer formats (this will be relaxed when more flexible vertex/index fetching is implemented)
-				checkf(Segment.VertexBufferElementType == VET_Float3 
+				checkf(Segment.VertexBufferElementType == VET_Float3
 					|| Segment.VertexBufferElementType == VET_Float4, TEXT("Only float3/4 vertex buffers are currently implemented.")); // #dxr_todo UE-72160: support other vertex buffer formats
 				checkf(Segment.VertexBufferStride >= 12, TEXT("Only deinterleaved float3 position vertex buffers are currently implemented.")); // #dxr_todo UE-72160: support interleaved vertex buffers
 			}
@@ -1988,76 +2008,21 @@ FRayTracingGeometryRHIRef FD3D12DynamicRHI::RHICreateRayTracingGeometry(const FR
 
 			ComputedPrimitiveCountForValidation += Segment.NumPrimitives;
 			check(Segment.FirstPrimitive + Segment.NumPrimitives <= Initializer.TotalPrimitiveCount);
-			
-			if(Initializer.IndexBuffer)
+
+			if (Initializer.IndexBuffer)
 			{
 				check(Initializer.IndexBuffer->GetSize() >=
-					  (Segment.FirstPrimitive + Segment.NumPrimitives) * IndicesPerPrimitive * IndexStride + Geometry->IndexOffsetInBytes);
+					(Segment.FirstPrimitive + Segment.NumPrimitives) * IndicesPerPrimitive * IndexStride + IndexOffsetInBytes);
 			}
 		}
 		check(ComputedPrimitiveCountForValidation == Initializer.TotalPrimitiveCount);
-
 	}
 #endif // DO_CHECK
 
-	Geometry->RHIIndexBuffer = Initializer.IndexBuffer;
-	Geometry->SetDirty(FRHIGPUMask::All(), true);
+	RHIIndexBuffer = Initializer.IndexBuffer;
+	SetDirty(FRHIGPUMask::All(), true);
 
 	INC_DWORD_STAT_BY(STAT_D3D12RayTracingTrianglesBLAS, Initializer.TotalPrimitiveCount);
-
-	return Geometry;
-}
-
-FRayTracingSceneRHIRef FD3D12DynamicRHI::RHICreateRayTracingScene(const FRayTracingSceneInitializer& Initializer)
-{
-	TRACE_CPUPROFILER_EVENT_SCOPE(CreateRayTracingScene);
-
-	FD3D12Adapter& Adapter = GetAdapter();
-
-	FD3D12RayTracingScene* Result = new FD3D12RayTracingScene(Adapter);
-
-	checkf(Initializer.Lifetime == RTSL_SingleFrame, TEXT("Only single-frame ray tracing scenes are currently implemented."));
-
-	Result->Lifetime = Initializer.Lifetime;
-	Result->CreatedFrameFenceValue = Adapter.GetFrameFence().GetCurrentFence();
-	Result->Instances = TArray<FRayTracingGeometryInstance>(Initializer.Instances.GetData(), Initializer.Instances.Num());
-	Result->ShaderSlotsPerGeometrySegment = Initializer.ShaderSlotsPerGeometrySegment;
-
-	// Compute geometry segment count prefix sum to be later used in GetHitRecordBaseIndex()
-	Result->SegmentPrefixSum.Reserve(Result->Instances.Num());
-	uint32 NumTotalSegments = 0;
-	for (const FRayTracingGeometryInstance& Instance : Result->Instances)
-	{
-		FD3D12RayTracingGeometry* Geometry = FD3D12DynamicRHI::ResourceCast(Instance.GeometryRHI);
-		Result->SegmentPrefixSum.Add(NumTotalSegments);
-		NumTotalSegments += Geometry->Segments.Num();
-	}
-	Result->NumTotalSegments = NumTotalSegments;
-
-	Result->NumCallableShaderSlots = Initializer.NumCallableShaderSlots;
-	Result->NumMissShaderSlots = FMath::Max<uint32>(1, Initializer.NumMissShaderSlots);
-
-	return Result;
-}
-
-FVertexBufferRHIRef FD3D12RayTracingGeometry::NullTransformBuffer;
-
-FD3D12RayTracingGeometry::FD3D12RayTracingGeometry()
-{
-	INC_DWORD_STAT(STAT_D3D12RayTracingAllocatedBLAS);
-
-	if(!FD3D12RayTracingGeometry::NullTransformBuffer.IsValid())
-	{
-		TResourceArray<float> NullTransformData;
-		NullTransformData.SetNumZeroed(12);
-
-		FRHIResourceCreateInfo CreateInfo;
-		CreateInfo.ResourceArray = &NullTransformData;
-
-		FD3D12RayTracingGeometry::NullTransformBuffer = RHICreateVertexBuffer(NullTransformData.GetResourceDataSize(), BUF_Static, CreateInfo);
-	}
-
-	FMemory::Memzero(PostBuildInfoBufferReadbackFences);
 }
 
 FD3D12RayTracingGeometry::~FD3D12RayTracingGeometry()
@@ -2152,6 +2117,9 @@ void FD3D12RayTracingGeometry::BuildAccelerationStructure(FD3D12CommandContext&
 	FD3D12Buffer* IndexBuffer = CommandContext.RetrieveObject<FD3D12Buffer>(RHIIndexBuffer.GetReference());
 	FD3D12Buffer* NullTransformBufferD3D12 = CommandContext.RetrieveObject<FD3D12Buffer>(NullTransformBuffer.GetReference());
 
+	TArray<FHitGroupSystemParameters>& HitGroupSystemParametersForThisGPU = HitGroupSystemParameters[CommandContext.GetGPUIndex()];
+	HitGroupSystemParametersForThisGPU.Reset(Segments.Num());
+
 	for (const FRayTracingGeometrySegment& Segment : Segments)
 	{
 		D3D12_RAYTRACING_GEOMETRY_DESC Desc = {};
@@ -2172,7 +2140,11 @@ void FD3D12RayTracingGeometry::BuildAccelerationStructure(FD3D12CommandContext&
 		}
 
 		FD3D12Buffer* VertexBuffer = CommandContext.RetrieveObject<FD3D12Buffer>(Segment.VertexBuffer.GetReference());
-			   
+		
+		FHitGroupSystemParameters SystemParameters = {};
+		SystemParameters.RootConstants.SetVertexAndIndexStride(Segment.VertexBufferStride, IndexStride);
+		SystemParameters.VertexBuffer = VertexBuffer->ResourceLocation.GetGPUVirtualAddress() + Segment.VertexBufferOffset;
+
 		if (GeometryType == D3D12_RAYTRACING_GEOMETRY_TYPE_TRIANGLES)
 		{
 			switch (Segment.VertexBufferElementType)
@@ -2206,11 +2178,12 @@ void FD3D12RayTracingGeometry::BuildAccelerationStructure(FD3D12CommandContext&
 
 			if (IndexBuffer)
 			{
+				SystemParameters.IndexBuffer = IndexBuffer->ResourceLocation.GetGPUVirtualAddress();
+				SystemParameters.RootConstants.IndexBufferOffsetInBytes = IndexOffsetInBytes + IndexStride * Segment.FirstPrimitive * IndicesPerPrimitive;
+
 				Desc.Triangles.IndexFormat = IndexStride == 4 ? DXGI_FORMAT_R32_UINT : DXGI_FORMAT_R16_UINT;
 				Desc.Triangles.IndexCount = Segment.NumPrimitives * IndicesPerPrimitive;
-				Desc.Triangles.IndexBuffer = IndexBuffer->ResourceLocation.GetGPUVirtualAddress() +
-					IndexOffsetInBytes +
-					IndexStride * Segment.FirstPrimitive * IndicesPerPrimitive;
+				Desc.Triangles.IndexBuffer = SystemParameters.IndexBuffer + SystemParameters.RootConstants.IndexBufferOffsetInBytes;
 
 				Desc.Triangles.VertexCount = VertexBuffer->ResourceLocation.GetSize() / Segment.VertexBufferStride;
 
@@ -2228,13 +2201,13 @@ void FD3D12RayTracingGeometry::BuildAccelerationStructure(FD3D12CommandContext&
 				Desc.Triangles.VertexCount = FMath::Min<uint32>(VertexBuffer->ResourceLocation.GetSize() / Segment.VertexBufferStride, TotalPrimitiveCount * 3);
 			}
 
-			Desc.Triangles.VertexBuffer.StartAddress = VertexBuffer->ResourceLocation.GetGPUVirtualAddress() + Segment.VertexBufferOffset;
+			Desc.Triangles.VertexBuffer.StartAddress = SystemParameters.VertexBuffer;
 			Desc.Triangles.VertexBuffer.StrideInBytes = Segment.VertexBufferStride;
 		}
 		else if (GeometryType == D3D12_RAYTRACING_GEOMETRY_TYPE_PROCEDURAL_PRIMITIVE_AABBS)
 		{
 			Desc.AABBs.AABBCount = Segment.NumPrimitives;
-			Desc.AABBs.AABBs.StartAddress = VertexBuffer->ResourceLocation.GetGPUVirtualAddress() + Segment.VertexBufferOffset;
+			Desc.AABBs.AABBs.StartAddress = SystemParameters.VertexBuffer;
 			Desc.AABBs.AABBs.StrideInBytes = Segment.VertexBufferStride;
 		}
 		else
@@ -2245,6 +2218,14 @@ void FD3D12RayTracingGeometry::BuildAccelerationStructure(FD3D12CommandContext&
 		VertexBuffer->ResourceLocation.GetResource()->UpdateResidency(CommandContext.CommandListHandle);
 
 		Descs.Add(Desc);
+
+		if (GeometryType == D3D12_RAYTRACING_GEOMETRY_TYPE_TRIANGLES)
+		{
+			// #dxr_todo UE-72160: support various vertex buffer layouts (fetch/decode based on vertex stride and format)
+			checkf(Segment.VertexBufferElementType == VET_Float3 || Segment.VertexBufferElementType == VET_Float4, TEXT("Only VET_Float3 and Float4 are currently implemented and tested. Other formats will be supported in the future."));
+		}
+
+		HitGroupSystemParametersForThisGPU.Add(SystemParameters);
 	}
 
 	const uint32 GPUIndex = CommandContext.GetGPUIndex();
@@ -2442,14 +2423,46 @@ void FD3D12RayTracingGeometry::ConditionalCompactAccelerationStructure(FD3D12Com
 	}
 }
 
-FD3D12RayTracingScene::FD3D12RayTracingScene(FD3D12Adapter& Adapter)
+FD3D12RayTracingScene::FD3D12RayTracingScene(FD3D12Adapter* Adapter, const FRayTracingSceneInitializer& Initializer)
+	: FD3D12AdapterChild(Adapter)
 {
-	ShaderResourceView = Adapter.CreateLinkedObject<FD3D12ShaderResourceView>(FRHIGPUMask::All(), [&](FD3D12Device* Device)
+	INC_DWORD_STAT(STAT_D3D12RayTracingAllocatedTLAS);
+
+	ShaderResourceView = Adapter->CreateLinkedObject<FD3D12ShaderResourceView>(FRHIGPUMask::All(), [&](FD3D12Device* Device)
 	{
 		return new FD3D12ShaderResourceView(Device);
 	});
 
-	INC_DWORD_STAT(STAT_D3D12RayTracingAllocatedTLAS);
+	checkf(Initializer.Lifetime == RTSL_SingleFrame, TEXT("Only single-frame ray tracing scenes are currently implemented."));
+
+	Lifetime = Initializer.Lifetime;
+	CreatedFrameFenceValue = Adapter->GetFrameFence().GetCurrentFence();
+	Instances = TArray<FRayTracingGeometryInstance>(Initializer.Instances.GetData(), Initializer.Instances.Num());
+	ShaderSlotsPerGeometrySegment = Initializer.ShaderSlotsPerGeometrySegment;
+
+	// Compute geometry segment count prefix sum to be later used in GetHitRecordBaseIndex()
+	SegmentPrefixSum.Reserve(Instances.Num());
+	NumTotalSegments = 0;
+
+	Experimental::TSherwoodSet<FD3D12RayTracingGeometry*> UniqueGeometries;
+
+	for (const FRayTracingGeometryInstance& Instance : Instances)
+	{
+		FD3D12RayTracingGeometry* Geometry = FD3D12DynamicRHI::ResourceCast(Instance.GeometryRHI);
+		SegmentPrefixSum.Add(NumTotalSegments);
+		NumTotalSegments += Geometry->Segments.Num();
+
+		bool bIsAlreadyInSet = false;
+		UniqueGeometries.Add(Geometry, &bIsAlreadyInSet);
+		if (!bIsAlreadyInSet)
+		{
+			Geometries.Add(Geometry);
+		}
+	}
+	NumTotalSegments = NumTotalSegments;
+
+	NumCallableShaderSlots = Initializer.NumCallableShaderSlots;
+	NumMissShaderSlots = FMath::Max<uint32>(1, Initializer.NumMissShaderSlots);
 };
 
 FD3D12RayTracingScene::~FD3D12RayTracingScene()
@@ -2558,21 +2571,64 @@ void FD3D12RayTracingScene::BuildAccelerationStructure(FD3D12CommandContext& Com
 
 		check(MappedData);
 
-		TSet<FD3D12RayTracingGeometry*> UniqueGeometries;
 		uint32 DxrInstanceIndex = 0;
 
 		bool bShouldCopyIndirectInstances = false;
 
+		{
+			TArray<FD3D12ResidencyHandle*>& GeometryResidencyHandlesForThisGPU = GeometryResidencyHandles[GPUIndex];
+
+			GeometryResidencyHandlesForThisGPU.Reset(0);
+
+			Experimental::TSherwoodSet<FD3D12ResidencyHandle*> UniqueResidencyHandles;
+
+			auto AddResidencyHandleForResource = [&UniqueResidencyHandles, &GeometryResidencyHandlesForThisGPU] (FD3D12Resource* Resource)
+			{
+			#if ENABLE_RESIDENCY_MANAGEMENT
+				FD3D12ResidencyHandle* ResidencyHandle = Resource->GetResidencyHandle();
+				if (D3DX12Residency::IsInitialized(ResidencyHandle))
+				{
+					bool bIsAlreadyInSet = false;
+					UniqueResidencyHandles.Add(ResidencyHandle, &bIsAlreadyInSet);
+					if (!bIsAlreadyInSet)
+					{
+						GeometryResidencyHandlesForThisGPU.Add(ResidencyHandle);
+					}
+				}
+			#endif // ENABLE_RESIDENCY_MANAGEMENT
+			};
+
+			for (TRefCountPtr<FD3D12RayTracingGeometry>& Geometry : Geometries)
+			{
+				checkf(!Geometry->IsDirty(CommandContext.GetGPUIndex()),
+					TEXT("Acceleration structures for all geometries must be built before building the top level acceleration structure for the scene."));
+
+				Geometry->ConditionalCompactAccelerationStructure(CommandContext);
+
+				AddResidencyHandleForResource(Geometry->AccelerationStructureBuffers[GPUIndex]->GetResource());
+
+				if (Geometry->RHIIndexBuffer)
+				{
+					FD3D12Buffer* IndexBuffer = CommandContext.RetrieveObject<FD3D12Buffer>(Geometry->RHIIndexBuffer.GetReference());
+					AddResidencyHandleForResource(IndexBuffer->GetResource());
+				}
+
+				for (const FRayTracingGeometrySegment& Segment : Geometry->Segments)
+				{
+					if (Segment.VertexBuffer)
+					{
+						FD3D12Buffer* VertexBuffer = CommandContext.RetrieveObject<FD3D12Buffer>(Segment.VertexBuffer.GetReference());
+						AddResidencyHandleForResource(VertexBuffer->GetResource());
+					}
+				}
+			}
+		}
+
 		for (uint32 InstanceIndex = 0; InstanceIndex < NumSceneInstances; ++InstanceIndex)
 		{
 			const FRayTracingGeometryInstance& Instance = Instances[InstanceIndex];
 			FD3D12RayTracingGeometry* Geometry = FD3D12DynamicRHI::ResourceCast(Instance.GeometryRHI);
 
-			Geometry->ConditionalCompactAccelerationStructure(CommandContext);
-
-			checkf(!Geometry->IsDirty(CommandContext.GetGPUIndex()),
-				TEXT("Acceleration structures for all geometries must be built before building the top level acceleration structure for the scene."));
-
 			D3D12_RAYTRACING_INSTANCE_DESC InstanceDesc = {};
 
 			InstanceDesc.InstanceMask = Instance.Mask;
@@ -2592,13 +2648,6 @@ void FD3D12RayTracingScene::BuildAccelerationStructure(FD3D12CommandContext& Com
 			InstanceDesc.AccelerationStructure = Geometry->AccelerationStructureBuffers[GPUIndex]->ResourceLocation.GetGPUVirtualAddress();
 			Geometry->AccelerationStructureBuffers[GPUIndex]->ResourceLocation.GetResource()->UpdateResidency(CommandContext.CommandListHandle);
 
-			bool bAlreadyInSet = false;
-			UniqueGeometries.Add(Geometry, &bAlreadyInSet);
-			if (!bAlreadyInSet)
-			{
-				BottomLevelAccelerationStructureBuffers[GPUIndex].Add(Geometry->AccelerationStructureBuffers[GPUIndex]);
-			}
-
 			bool bHasGPUTransforms = Instance.GPUTransformsSRV.IsValid();
 
 			const uint32 NumTransforms = Instance.NumTransforms;
@@ -2711,13 +2760,17 @@ void FD3D12RayTracingScene::BuildAccelerationStructure(FD3D12CommandContext& Com
 
 void FD3D12RayTracingScene::UpdateResidency(FD3D12CommandContext& CommandContext)
 {
+#if ENABLE_RESIDENCY_MANAGEMENT
 	const uint32 GPUIndex = CommandContext.GetGPUIndex();
 	AccelerationStructureBuffers[GPUIndex]->GetResource()->UpdateResidency(CommandContext.CommandListHandle);
 
-	for (auto& BLASBuffer : BottomLevelAccelerationStructureBuffers[GPUIndex])
+	FD3D12ResidencySet& ResidencySet = CommandContext.CommandListHandle.GetResidencySet();
+
+	for (FD3D12ResidencyHandle* ResidencyHandle : GeometryResidencyHandles[GPUIndex])
 	{
-		BLASBuffer->GetResource()->UpdateResidency(CommandContext.CommandListHandle);
+		D3DX12Residency::Insert(ResidencySet, ResidencyHandle);
 	}
+#endif // ENABLE_RESIDENCY_MANAGEMENT
 }
 
 FD3D12RayTracingShaderTable* FD3D12RayTracingScene::FindExistingShaderTable(const FD3D12RayTracingPipelineState* Pipeline, FD3D12Device* Device) const
@@ -2746,6 +2799,7 @@ FD3D12RayTracingShaderTable* FD3D12RayTracingScene::FindOrCreateShaderTable(cons
 
 	FD3D12RayTracingShaderTable* CreatedShaderTable = new FD3D12RayTracingShaderTable();
 	ID3D12Device5* RayTracingDevice = Device->GetRayTracingDevice();
+	const uint32 GPUIndex = Device->GetGPUIndex();
 
 	const uint32 NumHitGroupSlots = Pipeline->bAllowHitGroupIndexing ? NumTotalSegments * ShaderSlotsPerGeometrySegment : 0;
 
@@ -2782,51 +2836,17 @@ FD3D12RayTracingShaderTable* FD3D12RayTracingScene::FindOrCreateShaderTable(cons
 			const FRayTracingGeometryInstance& Instance = Instances[InstanceIndex];
 
 			const FD3D12RayTracingGeometry* Geometry = FD3D12DynamicRHI::ResourceCast(Instance.GeometryRHI);
+			checkf(!Geometry->IsDirty(GPUIndex), TEXT("Ray tracing acceleration structure must be built before it is used in a shader binding table."));
 
-			static constexpr uint32 IndicesPerPrimitive = 3; // Only triangle meshes are supported
-
-			const uint32 IndexStride = Geometry->IndexStride;
-
-			FD3D12Buffer* IndexBuffer = FD3D12CommandContext::RetrieveObject<FD3D12Buffer>(Geometry->RHIIndexBuffer.GetReference(), Device);
-
-			// Here, we directly apply the vertex offset to the address but set the offset for index buffer using a RootConstant instead.
-			// Index Buffer has to be 4 bytes aligned.
-			// cite comments from `RayTracingHitGroupComment.ush`:
-			// ByteAddressBuffer loads must be aligned to DWORD boundary.
-			// which means the alignment can be wrong for index buffer if we directly apply the offset to the address when binding it.
-			const D3D12_GPU_VIRTUAL_ADDRESS IndexBufferAddress = IndexBuffer ? IndexBuffer->ResourceLocation.GetGPUVirtualAddress() : 0;
-
-			if (IndexBuffer)
-			{
-				CreatedShaderTable->AddResourceReference(IndexBuffer->ResourceLocation.GetResource());
-			}
+			const TArray<FHitGroupSystemParameters>& HitGroupSystemParametersForThisGPU = Geometry->HitGroupSystemParameters[GPUIndex];
 
 			const uint32 NumSegments = Geometry->Segments.Num();
 			for (uint32 SegmentIndex = 0; SegmentIndex < NumSegments; ++SegmentIndex)
 			{
-				const FRayTracingGeometrySegment& Segment = Geometry->Segments[SegmentIndex];
-
-				FD3D12Buffer* VertexBuffer = FD3D12CommandContext::RetrieveObject<FD3D12Buffer>(Segment.VertexBuffer.GetReference(), Device);
-				CreatedShaderTable->AddResourceReference(VertexBuffer->ResourceLocation.GetResource());
-
-				const D3D12_GPU_VIRTUAL_ADDRESS VertexBufferAddress = VertexBuffer->ResourceLocation.GetGPUVirtualAddress() + Segment.VertexBufferOffset;
-
 				const uint32 RecordBaseIndex = GetHitRecordBaseIndex(InstanceIndex, SegmentIndex);
 
-				FHitGroupSystemParameters SystemParameters = {};
-				SystemParameters.IndexBuffer = IndexBufferAddress;
-				SystemParameters.VertexBuffer = VertexBufferAddress;
-
-				if (Geometry->GeometryType == D3D12_RAYTRACING_GEOMETRY_TYPE_TRIANGLES)
-				{
-					// #dxr_todo UE-72160: support various vertex buffer layouts (fetch/decode based on vertex stride and format)
-					checkf(Segment.VertexBufferElementType == VET_Float3 || Segment.VertexBufferElementType == VET_Float4, TEXT("Only VET_Float3 and Float4 are currently implemented and tested. Other formats will be supported in the future."));
-				}
-
-				SystemParameters.RootConstants.SetVertexAndIndexStride(Segment.VertexBufferStride, IndexStride);
-				SystemParameters.RootConstants.IndexBufferOffsetInBytes = IndexStride * Segment.FirstPrimitive * IndicesPerPrimitive + Geometry->IndexOffsetInBytes;
+				FHitGroupSystemParameters SystemParameters = HitGroupSystemParametersForThisGPU[SegmentIndex];
 				SystemParameters.RootConstants.BaseInstanceIndex = BaseInstanceIndex;
-
 				for (uint32 SlotIndex = 0; SlotIndex < ShaderSlotsPerGeometrySegment; ++SlotIndex)
 				{
 					CreatedShaderTable->SetLocalShaderParameters(RecordBaseIndex + SlotIndex, 0, SystemParameters);
@@ -2837,8 +2857,6 @@ FD3D12RayTracingShaderTable* FD3D12RayTracingScene::FindOrCreateShaderTable(cons
 		}
 	}
 
-	const uint32 GPUIndex = Device->GetGPUIndex();
-
 	ShaderTables[GPUIndex].Add(Pipeline, CreatedShaderTable);
 
 	return CreatedShaderTable;
diff --git a/Engine/Source/Runtime/D3D12RHI/Private/D3D12RayTracing.h b/Engine/Source/Runtime/D3D12RHI/Private/D3D12RayTracing.h
--- a/Engine/Source/Runtime/D3D12RHI/Private/D3D12RayTracing.h
+++ b/Engine/Source/Runtime/D3D12RHI/Private/D3D12RayTracing.h
@@ -6,16 +6,26 @@
 
 #if D3D12_RHI_RAYTRACING
 
+#include "RayTracingBuiltInResources.h"
+
 class FD3D12RayTracingPipelineState;
 class FD3D12RayTracingShaderTable;
 
 typedef FD3D12VertexBuffer FD3D12MemBuffer; // Generic GPU memory buffer
 
+// Built-in local root parameters that are always bound to all hit shaders
+struct FHitGroupSystemParameters
+{
+	D3D12_GPU_VIRTUAL_ADDRESS IndexBuffer;
+	D3D12_GPU_VIRTUAL_ADDRESS VertexBuffer;
+	FHitGroupSystemRootConstants RootConstants;
+};
+
 class FD3D12RayTracingGeometry : public FRHIRayTracingGeometry
 {
 public:
 
-	FD3D12RayTracingGeometry();
+	FD3D12RayTracingGeometry(const FRayTracingGeometryInitializer& Initializer);
 	~FD3D12RayTracingGeometry();
 
 	void TransitionBuffers(FD3D12CommandContext& CommandContext);
@@ -54,13 +64,15 @@ public:
 	TRefCountPtr<FD3D12MemBuffer> PostBuildInfoBuffers[MAX_NUM_GPUS];
 	FStagingBufferRHIRef PostBuildInfoStagingBuffers[MAX_NUM_GPUS];
 
+	// Hit shader parameters per geometry segment
+	TArray<FHitGroupSystemParameters> HitGroupSystemParameters[MAX_NUM_GPUS];
 };
 
-class FD3D12RayTracingScene : public FRHIRayTracingScene
+class FD3D12RayTracingScene : public FRHIRayTracingScene, public FD3D12AdapterChild
 {
 public:
 
-	FD3D12RayTracingScene(FD3D12Adapter& Adapter);
+	FD3D12RayTracingScene(FD3D12Adapter* Adapter, const FRayTracingSceneInitializer& Initializer);
 	~FD3D12RayTracingScene();
 
 	void BuildAccelerationStructure(FD3D12CommandContext& CommandContext, D3D12_RAYTRACING_ACCELERATION_STRUCTURE_BUILD_FLAGS BuildFlags);
@@ -70,8 +82,14 @@ public:
 
 	TArray<FRayTracingGeometryInstance> Instances;
 
-	// Scene keeps track of child acceleration structures to manage their residency
-	TArray<TRefCountPtr<FD3D12MemBuffer>> BottomLevelAccelerationStructureBuffers[MAX_NUM_GPUS];
+	// Unique list of geometries referenced by all instances in this scene.
+	// Any referenced geometry is kept alive while the scene is alive.
+	TArray<TRefCountPtr<FD3D12RayTracingGeometry>> Geometries;
+
+	// Scene keeps track of child acceleration structure buffers to ensure
+	// they are resident when any ray tracing work is dispatched.
+	TArray<FD3D12ResidencyHandle*> GeometryResidencyHandles[MAX_NUM_GPUS];
+
 	void UpdateResidency(FD3D12CommandContext& CommandContext);
 
 	uint32 ShaderSlotsPerGeometrySegment = 1;
