Fix trying to compact BLAS repeatedly because GetLastCompletedFenceFast() returning MAX_uint64

diff --git a/Engine/Source/Runtime/D3D12RHI/Private/D3D12RayTracing.cpp b/Engine/Source/Runtime/D3D12RHI/Private/D3D12RayTracing.cpp
--- a/Engine/Source/Runtime/D3D12RHI/Private/D3D12RayTracing.cpp
+++ b/Engine/Source/Runtime/D3D12RHI/Private/D3D12RayTracing.cpp
@@ -2362,6 +2362,13 @@ void FD3D12RayTracingGeometry::ConditionalCompactAccelerationStructure(FD3D12Com
 	}
 
 	const uint32 GPUIndex = CommandContext.GetGPUIndex();
+
+	// return the buffers are null (mostly because the BLAS has been compacted)
+	if (PostBuildInfoBuffers[GPUIndex] == nullptr)
+	{
+		return;
+	}
+
 	FD3D12Adapter* Adapter = CommandContext.GetParentAdapter();
 
 	const FD3D12Fence& Fence = CommandContext.GetParentDevice()->GetCommandListManager().GetFence();
@@ -2381,8 +2388,6 @@ void FD3D12RayTracingGeometry::ConditionalCompactAccelerationStructure(FD3D12Com
 		PostBuildInfoStagingBuffers[GPUIndex]->Unlock();
 		PostBuildInfoBuffers[GPUIndex] = nullptr;
 		PostBuildInfoStagingBuffers[GPUIndex] = nullptr;
-
-		// Set fence so we will not compact the same BLAS repeatedly
 		PostBuildInfoBufferReadbackFences[GPUIndex] = MAX_uint64;
 	}
 
