FD3D12CommandContext::RHISetRayTracingHitGroups is now able to internally parallelize the work to improve RHI thread critical path time.

The maximum number of worker threads is limited to 5 (RHIT + 4 task threads). Each thread gets a dedicated descriptor cache as well as separate referenced resource tracking sets to avoid contention.

diff --git a/Engine/Source/Runtime/D3D12RHI/Private/D3D12RayTracing.cpp b/Engine/Source/Runtime/D3D12RHI/Private/D3D12RayTracing.cpp
--- a/Engine/Source/Runtime/D3D12RHI/Private/D3D12RayTracing.cpp
+++ b/Engine/Source/Runtime/D3D12RHI/Private/D3D12RayTracing.cpp
@@ -883,18 +883,16 @@ struct FD3D12RayTracingDescriptorHeap : public FD3D12DeviceChild
 		DescriptorSize = GetParentDevice()->GetDevice()->GetDescriptorHandleIncrementSize(Type);
 	}
 
-	bool CanAllocate(uint32 InNumDescriptors) const
-	{
-		return NumAllocatedDescriptors + InNumDescriptors <= MaxNumDescriptors;
-	}
-
+	// Atomic linear allocation.
 	uint32 Allocate(uint32 InNumDescriptors)
 	{
-		check(CanAllocate(InNumDescriptors));
+		int32 Result = FPlatformAtomics::InterlockedAdd(&NumAllocatedDescriptors, InNumDescriptors);
 
-		uint32 Result = NumAllocatedDescriptors;
-		NumAllocatedDescriptors += InNumDescriptors;
-		return Result;
+		checkf(Result + InNumDescriptors <= MaxNumDescriptors,
+			TEXT("Ray tracing descriptor heap overflow. Heap created with %d descriptors."),
+			MaxNumDescriptors);
+
+		return uint32(Result);
 	}
 
 	D3D12_CPU_DESCRIPTOR_HANDLE GetDescriptorCPU(uint32 Index) const
@@ -919,7 +917,8 @@ struct FD3D12RayTracingDescriptorHeap : public FD3D12DeviceChild
 
 	ID3D12DescriptorHeap* D3D12Heap = nullptr;
 	uint32 MaxNumDescriptors = 0;
-	uint32 NumAllocatedDescriptors = 0;
+
+	int32 NumAllocatedDescriptors = 0;
 
 	uint32 DescriptorSize = 0;
 	D3D12_CPU_DESCRIPTOR_HANDLE CPUBase = {};
@@ -948,6 +947,7 @@ public:
 
 	void UpdateSyncPoint()
 	{
+		check(IsInRHIThread() || !IsRunningRHIInSeparateThread());
 		ViewHeap.UpdateSyncPoint();
 		SamplerHeap.UpdateSyncPoint();
 	}
@@ -968,12 +968,12 @@ public:
 		CommandContext.CommandListHandle.GraphicsCommandList()->SetDescriptorHeaps(2, Heaps);
 	}
 
-	uint32 GetDescriptorTableBaseIndex(const D3D12_CPU_DESCRIPTOR_HANDLE* Descriptors, uint32 NumDescriptors, D3D12_DESCRIPTOR_HEAP_TYPE Type)
+	uint32 AllocateDescriptorTable(const D3D12_CPU_DESCRIPTOR_HANDLE* Descriptors, uint32 NumDescriptors, D3D12_DESCRIPTOR_HEAP_TYPE Type, uint32 WorkerIndex)
 	{
 		checkSlow(Type == D3D12_DESCRIPTOR_HEAP_TYPE_CBV_SRV_UAV || Type == D3D12_DESCRIPTOR_HEAP_TYPE_SAMPLER);
 
 		FD3D12RayTracingDescriptorHeap& Heap = (Type == D3D12_DESCRIPTOR_HEAP_TYPE_CBV_SRV_UAV) ? ViewHeap : SamplerHeap;
-		TDescriptorHashMap& Map = (Type == D3D12_DESCRIPTOR_HEAP_TYPE_CBV_SRV_UAV) ? ViewDescriptorTableCache : SamplerDescriptorTableCache;
+		TDescriptorHashMap& Map = (Type == D3D12_DESCRIPTOR_HEAP_TYPE_CBV_SRV_UAV) ? ViewDescriptorTableCache[WorkerIndex]: SamplerDescriptorTableCache[WorkerIndex];
 
 		const uint64 Key = CityHash64((const char*)Descriptors, sizeof(Descriptors[0]) * NumDescriptors);
 
@@ -989,7 +989,6 @@ public:
 			DescriptorTableBaseIndex = Heap.Allocate(NumDescriptors);
 
 			D3D12_CPU_DESCRIPTOR_HANDLE DestDescriptor = Heap.GetDescriptorCPU(DescriptorTableBaseIndex);
-			checkf(Heap.CPUBase.ptr, TEXT("Ray tracing descriptor heap of type %d assigned to descriptor cache is invalid."), Type);
 			GetParentDevice()->GetDevice()->CopyDescriptors(1, &DestDescriptor, &NumDescriptors, NumDescriptors, Descriptors, nullptr, Type);
 
 			if (Type == D3D12_DESCRIPTOR_HEAP_TYPE_CBV_SRV_UAV)
@@ -1021,15 +1020,15 @@ public:
 		}
 	};
 
+	static constexpr uint32 MaxBindingWorkers = FD3D12RayTracingScene::MaxBindingWorkers;
 	using TDescriptorHashMap = Experimental::TSherwoodMap<uint64, uint32, TIdentityHash<uint64>>;
-	TDescriptorHashMap ViewDescriptorTableCache;
-	TDescriptorHashMap SamplerDescriptorTableCache;
+	TDescriptorHashMap ViewDescriptorTableCache[MaxBindingWorkers];
+	TDescriptorHashMap SamplerDescriptorTableCache[MaxBindingWorkers];
 };
 
 class FD3D12RayTracingShaderTable
 {
 private:
-
 	void WriteData(uint32 WriteOffset, const void* InData, uint32 InDataSize)
 	{
 #if DO_CHECK && DO_GUARD_SLOW
@@ -1058,6 +1057,8 @@ public:
 
 	UE_NONCOPYABLE(FD3D12RayTracingShaderTable)
 
+	static constexpr uint32 MaxBindingWorkers = FD3D12RayTracingScene::MaxBindingWorkers;
+
 	struct FInitializer
 	{
 		uint32 NumRayGenShaders = 0;
@@ -1369,54 +1370,85 @@ public:
 			return uint32(Key.Hash);
 		}
 	};
-	TMap<FShaderRecordCacheKey, uint32> ShaderRecordCache;
+	TMap<FShaderRecordCacheKey, uint32> ShaderRecordCache[MaxBindingWorkers];
 
 	// A set of all resources referenced by this shader table for the purpose of updating residency before ray tracing work dispatch.
 	// #dxr_todo UE-72159: remove resources from this set when SBT slot entries are replaced
-	Experimental::TSherwoodSet<void*> ReferencedD3D12ResourceSet;
-	TArray<TRefCountPtr<FD3D12Resource>> ReferencedD3D12Resources;
-	void AddResourceReference(FD3D12Resource* D3D12Resource)
+	Experimental::TSherwoodSet<void*> ReferencedD3D12ResourceSet[MaxBindingWorkers];
+	TArray<TRefCountPtr<FD3D12Resource>> ReferencedD3D12Resources[MaxBindingWorkers];
+	void AddResourceReference(FD3D12Resource* D3D12Resource, uint32 WorkerIndex)
 	{
 		bool bIsAlreadyInSet = false;
-		ReferencedD3D12ResourceSet.Add(D3D12Resource, &bIsAlreadyInSet);
+		ReferencedD3D12ResourceSet[WorkerIndex].Add(D3D12Resource, &bIsAlreadyInSet);
 		if (!bIsAlreadyInSet)
 		{
-			ReferencedD3D12Resources.Add(D3D12Resource);
+			ReferencedD3D12Resources[WorkerIndex].Add(D3D12Resource);
 		}
 	}
+
 	void UpdateResidency(FD3D12CommandContext& CommandContext)
 	{
-		for (FD3D12Resource* Resource : ReferencedD3D12Resources)
+		// Merge all data from worker threads into the main set
+
+		for (uint32 WorkerIndex = 1; WorkerIndex < MaxBindingWorkers; ++WorkerIndex)
+		{
+			for (FD3D12Resource* Resource : ReferencedD3D12Resources[WorkerIndex])
+			{
+				AddResourceReference(Resource, 0);
+			}
+		}
+
+		// Use the main (merged) set data to actually update resource residency
+
+		for (FD3D12Resource* Resource : ReferencedD3D12Resources[0])
 		{
 			Resource->UpdateResidency(CommandContext.CommandListHandle);
 		}
+
 		Buffer->GetResource()->UpdateResidency(CommandContext.CommandListHandle);
 	}
 
 	// Some resources referenced in SBT may be dynamic (written on GPU timeline) and may require transition barriers.
 	// We save such resources while we fill the SBT and issue transitions before the SBT is used.
 
-	TSet<FD3D12ShaderResourceView*> TransitionSRVs;
-	TSet<FD3D12UnorderedAccessView*> TransitionUAVs;
+	TSet<FD3D12ShaderResourceView*> TransitionSRVs[MaxBindingWorkers];
+	TSet<FD3D12UnorderedAccessView*> TransitionUAVs[MaxBindingWorkers];
 
-	void AddResourceTransition(FD3D12ShaderResourceView* SRV)
+	void AddResourceTransition(FD3D12ShaderResourceView* SRV, uint32 WorkerIndex)
 	{
-		TransitionSRVs.Add(SRV);
+		TransitionSRVs[WorkerIndex].Add(SRV);
 	}
 
-	void AddResourceTransition(FD3D12UnorderedAccessView* UAV)
+	void AddResourceTransition(FD3D12UnorderedAccessView* UAV, uint32 WorkerIndex)
 	{
-		TransitionUAVs.Add(UAV);
+		TransitionUAVs[WorkerIndex].Add(UAV);
 	}
 
 	void TransitionResources(FD3D12CommandContext& CommandContext)
 	{
-		for (FD3D12ShaderResourceView* SRV : TransitionSRVs)
+		// Merge all data from worker threads into the main set
+
+		for (uint32 WorkerIndex = 1; WorkerIndex < MaxBindingWorkers; ++WorkerIndex)
+		{
+			for (FD3D12ShaderResourceView* SRV : TransitionSRVs[WorkerIndex])
+			{
+				AddResourceTransition(SRV, 0);
+			}
+
+			for (FD3D12UnorderedAccessView* UAV : TransitionUAVs[WorkerIndex])
+			{
+				AddResourceTransition(UAV, 0);
+			}
+		}
+
+		// Use the main (merged) set data to perform resource transitions
+
+		for (FD3D12ShaderResourceView* SRV : TransitionSRVs[0])
 		{
 			FD3D12DynamicRHI::TransitionResource(CommandContext.CommandListHandle, SRV, D3D12_RESOURCE_STATE_NON_PIXEL_SHADER_RESOURCE);
 		}
 
-		for (FD3D12UnorderedAccessView* UAV : TransitionUAVs)
+		for (FD3D12UnorderedAccessView* UAV : TransitionUAVs[0])
 		{
 			FD3D12DynamicRHI::TransitionResource(CommandContext.CommandListHandle, UAV, D3D12_RESOURCE_STATE_UNORDERED_ACCESS);
 		}
@@ -2931,9 +2963,11 @@ void FD3D12CommandContext::RHIClearRayTracingBindings(FRHIRayTracingScene* InSce
 
 struct FD3D12RayTracingGlobalResourceBinder
 {
-	FD3D12RayTracingGlobalResourceBinder(FD3D12CommandContext& InCommandContext)
+	FD3D12RayTracingGlobalResourceBinder(FD3D12CommandContext& InCommandContext, FD3D12RayTracingDescriptorCache& InDescriptorCache)
 		: CommandContext(InCommandContext)
+		, DescriptorCache(InDescriptorCache)
 	{
+		check(IsInRHIThread() || !IsRunningRHIInSeparateThread());
 	}
 
 	void SetRootCBV(uint32 BaseSlotIndex, uint32 DescriptorIndex, D3D12_GPU_VIRTUAL_ADDRESS Address)
@@ -2972,26 +3006,38 @@ struct FD3D12RayTracingGlobalResourceBinder
 		FD3D12DynamicRHI::TransitionResource(CommandContext.CommandListHandle, UAV, D3D12_RESOURCE_STATE_UNORDERED_ACCESS);
 	}
 
+	FD3D12Device* GetDevice()
+	{
+		return CommandContext.GetParentDevice();
+	}
+
 	FD3D12CommandContext& CommandContext;
+	FD3D12RayTracingDescriptorCache& DescriptorCache;
+	static constexpr uint32 WorkerIndex = 0;
 };
 
 struct FD3D12RayTracingLocalResourceBinder
 {
-	FD3D12RayTracingLocalResourceBinder(FD3D12CommandContext& InCommandContext, FD3D12RayTracingShaderTable* InShaderTable, const FD3D12RootSignature* InRootSignature, uint32 InRecordIndex)
-		: Device(InCommandContext.GetParentDevice())
+	FD3D12RayTracingLocalResourceBinder(FD3D12Device& InDevice, FD3D12RayTracingShaderTable& InShaderTable, const FD3D12RootSignature& InRootSignature, uint32 InRecordIndex, uint32 InWorkerIndex)
+		: Device(InDevice)
 		, ShaderTable(InShaderTable)
+		, DescriptorCache(*InShaderTable.DescriptorCache)
 		, RootSignature(InRootSignature)
 		, RecordIndex(InRecordIndex)
+		, WorkerIndex(InWorkerIndex)
 	{
+		check(InShaderTable.DescriptorCache != nullptr);
+		check(WorkerIndex < InShaderTable.MaxBindingWorkers);
+		check(WorkerIndex < DescriptorCache.MaxBindingWorkers);
 		check(RecordIndex != ~0u);
 	}
 
 	void SetRootDescriptor(uint32 BaseSlotIndex, uint32 DescriptorIndex, D3D12_GPU_VIRTUAL_ADDRESS Address)
 	{
-		const uint32 BindOffsetBase = RootSignature->GetBindSlotOffsetInBytes(BaseSlotIndex);
+		const uint32 BindOffsetBase = RootSignature.GetBindSlotOffsetInBytes(BaseSlotIndex);
 		const uint32 DescriptorSize = uint32(sizeof(D3D12_GPU_VIRTUAL_ADDRESS));
 		const uint32 CurrentOffset = BindOffsetBase + DescriptorIndex * DescriptorSize;
-		ShaderTable->SetLocalShaderParameters(RecordIndex, CurrentOffset, Address);
+		ShaderTable.SetLocalShaderParameters(RecordIndex, CurrentOffset, Address);
 	}
 
 	void SetRootCBV(uint32 BaseSlotIndex, uint32 DescriptorIndex, D3D12_GPU_VIRTUAL_ADDRESS Address)
@@ -3006,8 +3052,8 @@ struct FD3D12RayTracingLocalResourceBinder
 
 	void SetRootDescriptorTable(uint32 SlotIndex, D3D12_GPU_DESCRIPTOR_HANDLE DescriptorTable)
 	{
-		const uint32 BindOffset = RootSignature->GetBindSlotOffsetInBytes(SlotIndex);
-		ShaderTable->SetLocalShaderParameters(RecordIndex, BindOffset, DescriptorTable);
+		const uint32 BindOffset = RootSignature.GetBindSlotOffsetInBytes(SlotIndex);
+		ShaderTable.SetLocalShaderParameters(RecordIndex, BindOffset, DescriptorTable);
 	}
 
 	D3D12_GPU_VIRTUAL_ADDRESS CreateTransientConstantBuffer(const void* Data, uint32 DataSize)
@@ -3016,27 +3062,27 @@ struct FD3D12RayTracingLocalResourceBinder
 		// caching constant buffer blocks inside ShaderTable and linearly sub-allocate from them.
 		// If the amount of data is relatively small, it may also be possible to use root constants and avoid extra allocations entirely.
 
-		FD3D12FastConstantAllocator& Allocator = Device->GetParentAdapter()->GetTransientUniformBufferAllocator();
-		FD3D12ResourceLocation ResourceLocation(Device);
+		FD3D12FastConstantAllocator& Allocator = Device.GetParentAdapter()->GetTransientUniformBufferAllocator();
+		FD3D12ResourceLocation ResourceLocation(&Device);
 		void* MappedData = Allocator.Allocate(DataSize, ResourceLocation);
 
 		FMemory::Memcpy(MappedData, Data, DataSize);
 
-		ShaderTable->AddResourceReference(ResourceLocation.GetResource());
+		ShaderTable.AddResourceReference(ResourceLocation.GetResource(), WorkerIndex);
 
 		return ResourceLocation.GetGPUVirtualAddress();
 	}
 
 	void AddResourceReference(FD3D12Resource* D3D12Resource, FRHIResource* RHIResource)
 	{
-		ShaderTable->AddResourceReference(D3D12Resource);
+		ShaderTable.AddResourceReference(D3D12Resource, WorkerIndex);
 	}
 
 	void AddResourceTransition(FD3D12ShaderResourceView* SRV)
 	{
 		if (SRV->GetResource()->RequiresResourceStateTracking())
 		{
-			ShaderTable->AddResourceTransition(SRV);
+			ShaderTable.AddResourceTransition(SRV, WorkerIndex);
 		}
 	}
 
@@ -3044,19 +3090,25 @@ struct FD3D12RayTracingLocalResourceBinder
 	{
 		if (UAV->GetResource()->RequiresResourceStateTracking())
 		{
-			ShaderTable->AddResourceTransition(UAV);
+			ShaderTable.AddResourceTransition(UAV, WorkerIndex);
 		}
 	}
 
-	FD3D12Device* Device = nullptr;
-	FD3D12RayTracingShaderTable* ShaderTable = nullptr;
-	const FD3D12RootSignature* RootSignature = nullptr;
+	FD3D12Device* GetDevice()
+	{
+		return &Device;
+	}
+
+	FD3D12Device& Device;
+	FD3D12RayTracingShaderTable& ShaderTable;
+	FD3D12RayTracingDescriptorCache& DescriptorCache;
+	const FD3D12RootSignature& RootSignature;
 	uint32 RecordIndex = ~0u;
+	uint32 WorkerIndex = 0;
 };
 
 template <typename ResourceBinderType>
 static void SetRayTracingShaderResources(
-	FD3D12CommandContext& CommandContext,
 	const FD3D12RayTracingShader* Shader,
 	uint32 InNumTextures, FRHITexture* const* Textures,
 	uint32 InNumSRVs, FRHIShaderResourceView* const* SRVs,
@@ -3064,11 +3116,8 @@ static void SetRayTracingShaderResources(
 	uint32 InNumSamplers, FRHISamplerState* const* Samplers,
 	uint32 InNumUAVs, FRHIUnorderedAccessView* const* UAVs,
 	uint32 InLooseParameterDataSize, const void* InLooseParameterData,
-	FD3D12RayTracingDescriptorCache& DescriptorCache,
 	ResourceBinderType& Binder)
 {
-	ID3D12Device* Device = CommandContext.GetParentDevice()->GetDevice();
-
 	const FD3D12RootSignature* RootSignature = Shader->pRootSignature;
 
 	D3D12_GPU_VIRTUAL_ADDRESS   LocalCBVs[MAX_CBS];
@@ -3088,12 +3137,18 @@ static void SetRayTracingShaderResources(
 	uint64 BoundUAVMask = 0;
 	uint64 BoundSamplerMask = 0;
 
+	FD3D12Device* Device = Binder.GetDevice();
+	auto DevicePredicate = [Device](FD3D12Device* CandidateDevice)
+	{
+		return Device == CandidateDevice;
+	};
+
 	for (uint32 SRVIndex = 0; SRVIndex < InNumTextures; ++SRVIndex)
 	{
 		FRHITexture* Resource = Textures[SRVIndex];
 		if (Resource)
 		{
-			FD3D12TextureBase* Texture = CommandContext.RetrieveTextureBase(Resource);
+			FD3D12TextureBase* Texture = FD3D12CommandContext::RetrieveTextureBase(Resource, DevicePredicate);
 			LocalSRVs[SRVIndex] = Texture->GetShaderResourceView()->GetView();
 			BoundSRVMask |= 1ull << SRVIndex;
 
@@ -3106,7 +3161,7 @@ static void SetRayTracingShaderResources(
 		FRHIShaderResourceView* Resource = SRVs[SRVIndex];
 		if (Resource)
 		{
-			FD3D12ShaderResourceView* SRV = CommandContext.RetrieveObject<FD3D12ShaderResourceView>(Resource);
+			FD3D12ShaderResourceView* SRV = FD3D12CommandContext::RetrieveObject<FD3D12ShaderResourceView>(Resource, Device);
 			LocalSRVs[SRVIndex] = SRV->GetView();
 			BoundSRVMask |= 1ull << SRVIndex;
 
@@ -3119,7 +3174,7 @@ static void SetRayTracingShaderResources(
 		FRHIUniformBuffer* Resource = UniformBuffers[CBVIndex];
 		if (Resource)
 		{
-			FD3D12UniformBuffer* CBV = CommandContext.RetrieveObject<FD3D12UniformBuffer>(Resource);
+			FD3D12UniformBuffer* CBV = FD3D12CommandContext::RetrieveObject<FD3D12UniformBuffer>(Resource, Device);
 			LocalCBVs[CBVIndex] = CBV->ResourceLocation.GetGPUVirtualAddress();
 			BoundCBVMask |= 1ull << CBVIndex;
 
@@ -3132,7 +3187,7 @@ static void SetRayTracingShaderResources(
 		FRHISamplerState* Resource = Samplers[SamplerIndex];
 		if (Resource)
 		{
-			LocalSamplers[SamplerIndex] = CommandContext.RetrieveObject<FD3D12SamplerState>(Resource)->Descriptor;
+			LocalSamplers[SamplerIndex] = FD3D12CommandContext::RetrieveObject<FD3D12SamplerState>(Resource, Device)->Descriptor;
 			BoundSamplerMask |= 1ull << SamplerIndex;
 		}
 	}
@@ -3142,7 +3197,7 @@ static void SetRayTracingShaderResources(
 		FRHIUnorderedAccessView* Resource = UAVs[UAVIndex];
 		if (Resource)
 		{
-			FD3D12UnorderedAccessView* UAV = CommandContext.RetrieveObject<FD3D12UnorderedAccessView>(Resource);
+			FD3D12UnorderedAccessView* UAV = FD3D12CommandContext::RetrieveObject<FD3D12UnorderedAccessView>(Resource, Device);
 			LocalUAVs[UAVIndex] = UAV->GetView();
 			BoundUAVMask |= 1ull << UAVIndex;
 
@@ -3185,7 +3240,8 @@ static void SetRayTracingShaderResources(
 					const uint16 ResourceIndex = FRHIResourceTableEntry::GetResourceIndex(ResourceInfo);
 					const uint8 BindIndex = FRHIResourceTableEntry::GetBindIndex(ResourceInfo);
 
-					FD3D12ShaderResourceView* SRV = CommandContext.RetrieveTextureBase((FRHITexture*)Resources[ResourceIndex].GetReference())->GetShaderResourceView();
+					FD3D12TextureBase* TextureBase = FD3D12CommandContext::RetrieveTextureBase((FRHITexture*)Resources[ResourceIndex].GetReference(), DevicePredicate);
+					FD3D12ShaderResourceView* SRV = TextureBase->GetShaderResourceView();
 					LocalSRVs[BindIndex] = SRV->GetView();
 					BoundSRVMask |= 1ull << BindIndex;
 
@@ -3213,7 +3269,7 @@ static void SetRayTracingShaderResources(
 					const uint16 ResourceIndex = FRHIResourceTableEntry::GetResourceIndex(ResourceInfo);
 					const uint8 BindIndex = FRHIResourceTableEntry::GetBindIndex(ResourceInfo);
 
-					FD3D12ShaderResourceView* SRV = CommandContext.RetrieveObject<FD3D12ShaderResourceView>((FRHIShaderResourceView*)(Resources[ResourceIndex].GetReference()));
+					FD3D12ShaderResourceView* SRV = FD3D12CommandContext::RetrieveObject<FD3D12ShaderResourceView>((FRHIShaderResourceView*)(Resources[ResourceIndex].GetReference()), Device);
 					LocalSRVs[BindIndex] = SRV->GetView();
 					BoundSRVMask |= 1ull << BindIndex;
 
@@ -3241,7 +3297,7 @@ static void SetRayTracingShaderResources(
 					const uint16 ResourceIndex = FRHIResourceTableEntry::GetResourceIndex(ResourceInfo);
 					const uint8 BindIndex = FRHIResourceTableEntry::GetBindIndex(ResourceInfo);
 
-					FD3D12SamplerState* Sampler = CommandContext.RetrieveObject<FD3D12SamplerState>((FRHISamplerState*)(Resources[ResourceIndex].GetReference()));
+					FD3D12SamplerState* Sampler = FD3D12CommandContext::RetrieveObject<FD3D12SamplerState>((FRHISamplerState*)(Resources[ResourceIndex].GetReference()), Device);
 					LocalSamplers[BindIndex] = Sampler->Descriptor;
 					BoundSamplerMask |= 1ull << BindIndex;
 
@@ -3280,10 +3336,13 @@ static void SetRayTracingShaderResources(
 	check(IsCompleteBinding(Shader->ResourceCounts.NumCBs, BoundCBVMask));
 	check(IsCompleteBinding(Shader->ResourceCounts.NumSamplers, BoundSamplerMask));
 
+	FD3D12RayTracingDescriptorCache& DescriptorCache = Binder.DescriptorCache;
+	const uint32 WorkerIndex = Binder.WorkerIndex;
+
 	const uint32 NumSRVs = Shader->ResourceCounts.NumSRVs;
 	if (NumSRVs)
 	{
-		const uint32 DescriptorTableBaseIndex = DescriptorCache.GetDescriptorTableBaseIndex(LocalSRVs, NumSRVs, D3D12_DESCRIPTOR_HEAP_TYPE_CBV_SRV_UAV);
+		const uint32 DescriptorTableBaseIndex = DescriptorCache.AllocateDescriptorTable(LocalSRVs, NumSRVs, D3D12_DESCRIPTOR_HEAP_TYPE_CBV_SRV_UAV, WorkerIndex);
 
 		const uint32 BindSlot = RootSignature->SRVRDTBindSlot(SF_Compute);
 		check(BindSlot != 0xFF);
@@ -3295,7 +3354,7 @@ static void SetRayTracingShaderResources(
 	const uint32 NumUAVs = Shader->ResourceCounts.NumUAVs;
 	if (NumUAVs)
 	{
-		const uint32 DescriptorTableBaseIndex = DescriptorCache.GetDescriptorTableBaseIndex(LocalUAVs, NumUAVs, D3D12_DESCRIPTOR_HEAP_TYPE_CBV_SRV_UAV);
+		const uint32 DescriptorTableBaseIndex = DescriptorCache.AllocateDescriptorTable(LocalUAVs, NumUAVs, D3D12_DESCRIPTOR_HEAP_TYPE_CBV_SRV_UAV, WorkerIndex);
 
 		const uint32 BindSlot = RootSignature->UAVRDTBindSlot(SF_Compute);
 		check(BindSlot != 0xFF);
@@ -3324,7 +3383,7 @@ static void SetRayTracingShaderResources(
 	const uint32 NumSamplers = Shader->ResourceCounts.NumSamplers;
 	if (NumSamplers)
 	{
-		const uint32 DescriptorTableBaseIndex = DescriptorCache.GetDescriptorTableBaseIndex(LocalSamplers, NumSamplers, D3D12_DESCRIPTOR_HEAP_TYPE_SAMPLER);
+		const uint32 DescriptorTableBaseIndex = DescriptorCache.AllocateDescriptorTable(LocalSamplers, NumSamplers, D3D12_DESCRIPTOR_HEAP_TYPE_SAMPLER, WorkerIndex);
 
 		const uint32 BindSlot = RootSignature->SamplerRDTBindSlot(SF_Compute);
 		check(BindSlot != 0xFF);
@@ -3341,10 +3400,8 @@ static void SetRayTracingShaderResources(
 
 template <typename ResourceBinderType>
 static void SetRayTracingShaderResources(
-	FD3D12CommandContext& CommandContext,
 	const FD3D12RayTracingShader* Shader,
 	const FRayTracingShaderBindings& ResourceBindings,
-	FD3D12RayTracingDescriptorCache& DescriptorCache,
 	ResourceBinderType& Binder)
 {
 	static_assert(
@@ -3360,14 +3417,14 @@ static void SetRayTracingShaderResources(
 		sizeof(ResourceBindings.UAVs) / sizeof(*ResourceBindings.UAVs) == MAX_UAVS,
 		"Ray Tracing Shader Bindings UAV array size must match D3D12 RHI Limit");
 
-	SetRayTracingShaderResources(CommandContext, Shader,
+	SetRayTracingShaderResources(Shader,
 		UE_ARRAY_COUNT(ResourceBindings.Textures), ResourceBindings.Textures,
 		UE_ARRAY_COUNT(ResourceBindings.SRVs), ResourceBindings.SRVs,
 		UE_ARRAY_COUNT(ResourceBindings.UniformBuffers), ResourceBindings.UniformBuffers,
 		UE_ARRAY_COUNT(ResourceBindings.Samplers), ResourceBindings.Samplers,
 		UE_ARRAY_COUNT(ResourceBindings.UAVs), ResourceBindings.UAVs,
 		0, nullptr, // loose parameters
-		DescriptorCache, Binder);
+		Binder);
 }
 
 static void DispatchRays(FD3D12CommandContext& CommandContext,
@@ -3398,10 +3455,11 @@ static void DispatchRays(FD3D12CommandContext& CommandContext,
 	if (OptShaderTable && OptShaderTable->DescriptorCache)
 	{
 		FD3D12RayTracingDescriptorCache* DescriptorCache = OptShaderTable->DescriptorCache;
+		check(DescriptorCache != nullptr);
 
 		DescriptorCache->SetDescriptorHeaps(CommandContext);
-		FD3D12RayTracingGlobalResourceBinder ResourceBinder(CommandContext);
-		SetRayTracingShaderResources(CommandContext, RayGenShader, GlobalBindings, *DescriptorCache, ResourceBinder);
+		FD3D12RayTracingGlobalResourceBinder ResourceBinder(CommandContext, *DescriptorCache);
+		SetRayTracingShaderResources(RayGenShader, GlobalBindings, ResourceBinder);
 
 		// #dxr_todo UE-72159: avoid updating residency if this scene was already used on the current command list (i.e. multiple ray dispatches are performed back-to-back)
 		OptShaderTable->UpdateResidency(CommandContext);
@@ -3411,8 +3469,8 @@ static void DispatchRays(FD3D12CommandContext& CommandContext,
 		FD3D12RayTracingDescriptorCache TransientDescriptorCache(CommandContext.GetParentDevice());
 		TransientDescriptorCache.Init(MAX_SRVS + MAX_UAVS, MAX_SAMPLERS);
 		TransientDescriptorCache.SetDescriptorHeaps(CommandContext);
-		FD3D12RayTracingGlobalResourceBinder ResourceBinder(CommandContext);
-		SetRayTracingShaderResources(CommandContext, RayGenShader, GlobalBindings, TransientDescriptorCache, ResourceBinder);
+		FD3D12RayTracingGlobalResourceBinder ResourceBinder(CommandContext, TransientDescriptorCache);
+		SetRayTracingShaderResources(RayGenShader, GlobalBindings, ResourceBinder);
 	}
 
 	if (OptShaderTable)
@@ -3546,14 +3604,15 @@ void FD3D12CommandContext::RHIRayTraceDispatch(FRHIRayTracingPipelineState* InRa
 }
 
 static void SetRayTracingHitGroup(
-	FD3D12CommandContext& CommandContext,
+	FD3D12Device* Device,
 	FD3D12RayTracingShaderTable* ShaderTable,
 	FD3D12RayTracingScene* Scene, 
 	FD3D12RayTracingPipelineState* Pipeline,
 	uint32 InstanceIndex, uint32 SegmentIndex, uint32 ShaderSlot, uint32 HitGroupIndex,
 	uint32 NumUniformBuffers, FRHIUniformBuffer* const* UniformBuffers,
 	uint32 LooseParameterDataSize, const void* LooseParameterData,
-	uint32 UserData)
+	uint32 UserData,
+	uint32 WorkerIndex)
 {
 	checkf(ShaderSlot < Scene->ShaderSlotsPerGeometrySegment, TEXT("Shader slot is invalid. Make sure that ShaderSlotsPerGeometrySegment is correct on FRayTracingSceneInitializer."));
 
@@ -3577,7 +3636,7 @@ static void SetRayTracingHitGroup(
 	{
 		CacheKey = FD3D12RayTracingShaderTable::FShaderRecordCacheKey(NumUniformBuffers, UniformBuffers, HitGroupIndex);
 
-		uint32* ExistingRecordIndex = ShaderTable->ShaderRecordCache.Find(CacheKey);
+		uint32* ExistingRecordIndex = ShaderTable->ShaderRecordCache[WorkerIndex].Find(CacheKey);
 		if (ExistingRecordIndex)
 		{
 			const uint32 OffsetFromRootSignatureStart = sizeof(FHitGroupSystemParameters);
@@ -3586,20 +3645,19 @@ static void SetRayTracingHitGroup(
 		}
 	}
 
-	FD3D12RayTracingLocalResourceBinder ResourceBinder(CommandContext, ShaderTable, Shader->pRootSignature, RecordIndex);
-	check(ShaderTable->DescriptorCache);
-	SetRayTracingShaderResources(CommandContext, Shader,
+	FD3D12RayTracingLocalResourceBinder ResourceBinder(*Device, *ShaderTable, *(Shader->pRootSignature), RecordIndex, WorkerIndex);
+	SetRayTracingShaderResources(Shader,
 		0, nullptr, // Textures
 		0, nullptr, // SRVs
 		NumUniformBuffers, UniformBuffers,
 		0, nullptr, // Samplers
 		0, nullptr, // UAVs
 		LooseParameterDataSize, LooseParameterData,
-		*(ShaderTable->DescriptorCache), ResourceBinder);
+		ResourceBinder);
 
 	if (bCanUseRecordCache)
 	{
-		ShaderTable->ShaderRecordCache.Add(CacheKey, RecordIndex);
+		ShaderTable->ShaderRecordCache[WorkerIndex].Add(CacheKey, RecordIndex);
 	}
 }
 
@@ -3611,25 +3669,65 @@ void FD3D12CommandContext::RHISetRayTracingHitGroups(
 
 	FD3D12RayTracingScene* Scene = FD3D12DynamicRHI::ResourceCast(InScene);
 	FD3D12RayTracingPipelineState* Pipeline = FD3D12DynamicRHI::ResourceCast(InPipeline);
+	FD3D12Device* Device = GetParentDevice();
 
 	FD3D12RayTracingShaderTable* ShaderTable = Scene->FindOrCreateShaderTable(Pipeline, GetParentDevice());
 	checkf(ShaderTable->LocalShaderTableOffset == ShaderTable->HitGroupShaderTableOffset,
 		TEXT("Hit shader records are assumed to be at the beginning of local shader table"));
 
-	for (uint32 BindingIndex = 0; BindingIndex < NumBindings; ++ BindingIndex)
+	FGraphEventArray TaskList;
+
+	const uint32 NumWorkerThreads = FTaskGraphInterface::Get().GetNumWorkerThreads();
+	const uint32 MaxTasks = FApp::ShouldUseThreadingForPerformance() 
+		? FMath::Min<uint32>(NumWorkerThreads, FD3D12RayTracingScene::MaxBindingWorkers)
+		: 1;
+
+	const uint32 MinItemsPerTask = 1024u;
+	const uint32 NumTasks = FMath::Min(MaxTasks, FMath::Max(1u, FMath::DivideAndRoundUp(NumBindings, MinItemsPerTask)));
+	const uint32 ItemsPerTask = FMath::DivideAndRoundUp(NumBindings, NumTasks);
+
+	auto BindingTask = [Bindings, Device, ShaderTable, Scene, Pipeline]
+	(uint32 FirstItemIndex, uint32 NumItems, uint32 WorkerIndex)
+	{
+		TRACE_CPUPROFILER_EVENT_SCOPE(SetRayTracingHitGroups_Task);
+
+		for (uint32 i = 0; i < NumItems; ++i)
+		{
+			const FRayTracingLocalShaderBindings& Binding = Bindings[FirstItemIndex + i];
+			SetRayTracingHitGroup(Device, ShaderTable, Scene, Pipeline,
+				Binding.InstanceIndex,
+				Binding.SegmentIndex,
+				Binding.ShaderSlot,
+				Binding.ShaderIndexInPipeline,
+				Binding.NumUniformBuffers,
+				Binding.UniformBuffers,
+				Binding.LooseParameterDataSize,
+				Binding.LooseParameterData,
+				Binding.UserData,
+				WorkerIndex);
+		}
+	};
+
+	TaskList.Reserve(NumTasks);
+
+	// Create and kick any parallel tasks first and run task 0 on this thread
+	for (uint32 TaskIndex = 1; TaskIndex < NumTasks; ++TaskIndex)
 	{
-		const FRayTracingLocalShaderBindings& Binding = Bindings[BindingIndex];
-		SetRayTracingHitGroup(*this, ShaderTable, Scene, Pipeline,
-			Binding.InstanceIndex,
-			Binding.SegmentIndex,
-			Binding.ShaderSlot,
-			Binding.ShaderIndexInPipeline,
-			Binding.NumUniformBuffers,
-			Binding.UniformBuffers,
-			Binding.LooseParameterDataSize,
-			Binding.LooseParameterData,
-			Binding.UserData);
+		const uint32 FirstItemIndex = TaskIndex * ItemsPerTask;
+		const uint32 NumItems = FMath::Min(ItemsPerTask, NumBindings - FirstItemIndex);
+		const FRayTracingLocalShaderBindings* Items = Bindings + FirstItemIndex;
+
+		TaskList.Add(FFunctionGraphTask::CreateAndDispatchWhenReady(
+			[TaskIndex, FirstItemIndex, NumItems, &BindingTask]()
+		{
+			BindingTask(FirstItemIndex, NumItems, TaskIndex);
+		}, TStatId(), nullptr, ENamedThreads::AnyThread));
 	}
+
+	// Run task 0 on this thread, in parallel with any workers
+	BindingTask(0, FMath::Min(ItemsPerTask, NumBindings), 0);
+
+	FTaskGraphInterface::Get().WaitUntilTasksComplete(TaskList, ENamedThreads::AnyThread);
 }
 
 void FD3D12CommandContext::RHISetRayTracingHitGroup(
@@ -3648,7 +3746,8 @@ void FD3D12CommandContext::RHISetRayTracingHitGroup(
 	checkf(ShaderTable->LocalShaderTableOffset == ShaderTable->HitGroupShaderTableOffset,
 		TEXT("Hit shader records are assumed to be at the beginning of local shader table"));
 
-	SetRayTracingHitGroup(*this, ShaderTable, Scene, Pipeline,
+	const uint32 WorkerIndex = 0;
+	SetRayTracingHitGroup(GetParentDevice(), ShaderTable, Scene, Pipeline,
 		InstanceIndex,
 		SegmentIndex,
 		ShaderSlot,
@@ -3657,7 +3756,8 @@ void FD3D12CommandContext::RHISetRayTracingHitGroup(
 		UniformBuffers,
 		LooseParameterDataSize,
 		LooseParameterData,
-		UserData);
+		UserData,
+		WorkerIndex);
 }
 
 void FD3D12CommandContext::RHISetRayTracingCallableShader(
@@ -3681,16 +3781,16 @@ void FD3D12CommandContext::RHISetRayTracingCallableShader(
 
 	const FD3D12RayTracingShader* Shader = Pipeline->CallableShaders.Shaders[ShaderIndexInPipeline];
 
-	FD3D12RayTracingLocalResourceBinder ResourceBinder(*this, ShaderTable, Shader->pRootSignature, RecordIndex);
-	check(ShaderTable->DescriptorCache);
-	SetRayTracingShaderResources(*this, Shader,
+	const uint32 WorkerIndex = 0;
+	FD3D12RayTracingLocalResourceBinder ResourceBinder(*GetParentDevice(), *ShaderTable, *(Shader->pRootSignature), RecordIndex, WorkerIndex);
+	SetRayTracingShaderResources(Shader,
 		0, nullptr, // Textures
 		0, nullptr, // SRVs
 		NumUniformBuffers, UniformBuffers,
 		0, nullptr, // Samplers
 		0, nullptr, // UAVs
 		0, nullptr, // Loose parameters
-		*(ShaderTable->DescriptorCache), ResourceBinder);
+		ResourceBinder);
 }
 
 void FD3D12CommandContext::RHISetRayTracingMissShader(
@@ -3715,15 +3815,15 @@ void FD3D12CommandContext::RHISetRayTracingMissShader(
 
 	const FD3D12RayTracingShader* Shader = Pipeline->MissShaders.Shaders[ShaderIndexInPipeline];
 
-	FD3D12RayTracingLocalResourceBinder ResourceBinder(*this, ShaderTable, Shader->pRootSignature, RecordIndex);
-	check(ShaderTable->DescriptorCache);
-	SetRayTracingShaderResources(*this, Shader,
+	uint32 WorkerIndex = 0;
+	FD3D12RayTracingLocalResourceBinder ResourceBinder(*GetParentDevice(), *ShaderTable, *(Shader->pRootSignature), RecordIndex, WorkerIndex);
+	SetRayTracingShaderResources(Shader,
 		0, nullptr, // Textures
 		0, nullptr, // SRVs
 		NumUniformBuffers, UniformBuffers,
 		0, nullptr, // Samplers
 		0, nullptr, // UAVs
 		0, nullptr, // Loose parameters
-		*(ShaderTable->DescriptorCache), ResourceBinder);
+		ResourceBinder);
 }
 #endif // D3D12_RHI_RAYTRACING
diff --git a/Engine/Source/Runtime/D3D12RHI/Private/D3D12RayTracing.h b/Engine/Source/Runtime/D3D12RHI/Private/D3D12RayTracing.h
--- a/Engine/Source/Runtime/D3D12RHI/Private/D3D12RayTracing.h
+++ b/Engine/Source/Runtime/D3D12RHI/Private/D3D12RayTracing.h
@@ -72,6 +72,11 @@ class FD3D12RayTracingScene : public FRHIRayTracingScene, public FD3D12AdapterCh
 {
 public:
 
+	// Ray tracing shader bindings can be processed in parallel.
+	// Each concurrent worker gets its own dedicated descriptor cache instance to avoid contention or locking.
+	// Scaling beyond 5 total threads does not yield any speedup in practice.
+	static constexpr uint32 MaxBindingWorkers = 5; // RHI thread + 4 parallel workers.
+
 	FD3D12RayTracingScene(FD3D12Adapter* Adapter, const FRayTracingSceneInitializer& Initializer);
 	~FD3D12RayTracingScene();
 
diff --git a/Engine/Source/Runtime/Renderer/Private/DeferredShadingRenderer.cpp b/Engine/Source/Runtime/Renderer/Private/DeferredShadingRenderer.cpp
--- a/Engine/Source/Runtime/Renderer/Private/DeferredShadingRenderer.cpp
+++ b/Engine/Source/Runtime/Renderer/Private/DeferredShadingRenderer.cpp
@@ -1287,18 +1287,47 @@ void FDeferredShadingSceneRenderer::WaitForRayTracingScene(FRHICommandListImmedi
 		{
 			FTaskGraphInterface::Get().WaitUntilTaskCompletes(View.RayTracingMaterialBindingsTask, ENamedThreads::GetRenderThread_Local());
 
+			// Gather bindings from all chunks and submit them all as a single batch to allow RHI to bind all shader parameters in parallel.
+
+			uint32 NumTotalBindings = 0;
+
 			for (FRayTracingLocalShaderBindingWriter* BindingWriter : View.RayTracingMaterialBindings)
 			{
-				// Data is kept alive at the high level and explicitly deleted on RHI timeline,
-				// so we can avoid copying parameters to the command list and simply pass raw pointers around.
-				const bool bCopyDataToInlineStorage = false;
-				BindingWriter->Commit(
-					RHICmdList,
-					View.RayTracingScene.RayTracingSceneRHI,
-					View.RayTracingMaterialPipeline,
-					bCopyDataToInlineStorage);
+				const FRayTracingLocalShaderBindingWriter::FChunk* Chunk = BindingWriter->GetFirstChunk();
+				while (Chunk)
+				{
+					NumTotalBindings += Chunk->Num;
+					Chunk = Chunk->Next;
+				}
 			}
 
+			FRayTracingLocalShaderBindings* MergedBindings = (FRayTracingLocalShaderBindings*)RHICmdList.Alloc(
+				sizeof(FRayTracingLocalShaderBindings) * NumTotalBindings,
+				alignof(FRayTracingLocalShaderBindings));
+
+			uint32 MergedBindingIndex = 0;
+			for (FRayTracingLocalShaderBindingWriter* BindingWriter : View.RayTracingMaterialBindings)
+			{
+				const FRayTracingLocalShaderBindingWriter::FChunk* Chunk = BindingWriter->GetFirstChunk();
+				while (Chunk)
+				{
+					const uint32 Num = Chunk->Num;
+					for (uint32_t i = 0; i < Num; ++i)
+					{
+						MergedBindings[MergedBindingIndex] = Chunk->Bindings[i];
+						MergedBindingIndex++;
+					}
+					Chunk = Chunk->Next;
+				}
+			}
+
+			const bool bCopyDataToInlineStorage = false; // Storage is already allocated from RHICmdList, no extra copy necessary
+			RHICmdList.SetRayTracingHitGroups(
+				View.RayTracingScene.RayTracingSceneRHI, 
+				View.RayTracingMaterialPipeline,
+				NumTotalBindings, MergedBindings,
+				bCopyDataToInlineStorage);
+
 			// Move the ray tracing binding container ownership to the command list, so that memory will be
 			// released on the RHI thread timeline, after the commands that reference it are processed.
 			RHICmdList.EnqueueLambda([Ptrs = MoveTemp(View.RayTracingMaterialBindings)](FRHICommandListImmediate&)
@@ -1324,7 +1353,6 @@ void FDeferredShadingSceneRenderer::WaitForRayTracingScene(FRHICommandListImmedi
 	}
 }
 
-
 #endif // RHI_RAYTRACING
 
 extern bool IsLpvIndirectPassRequired(const FViewInfo& View);
diff --git a/Engine/Source/Runtime/Renderer/Private/RayTracing/RayTracingMaterialHitShaders.h b/Engine/Source/Runtime/Renderer/Private/RayTracing/RayTracingMaterialHitShaders.h
--- a/Engine/Source/Runtime/Renderer/Private/RayTracing/RayTracingMaterialHitShaders.h
+++ b/Engine/Source/Runtime/Renderer/Private/RayTracing/RayTracingMaterialHitShaders.h
@@ -127,8 +127,6 @@ public:
 		}
 	}
 
-private:
-
 	struct FChunk
 	{
 		static constexpr uint32 MaxNum = 1024;
@@ -140,6 +138,13 @@ private:
 		uint32 Num;
 	};
 
+	const FChunk* GetFirstChunk() const
+	{
+		return FirstChunk;
+	}
+
+private:
+
 	FChunk* FirstChunk = nullptr;
 	FChunk* CurrentChunk = nullptr;
 
