Fixed static analisys warning

diff --git a/Engine/Source/Runtime/D3D12RHI/Private/D3D12RayTracing.cpp b/Engine/Source/Runtime/D3D12RHI/Private/D3D12RayTracing.cpp
--- a/Engine/Source/Runtime/D3D12RHI/Private/D3D12RayTracing.cpp
+++ b/Engine/Source/Runtime/D3D12RHI/Private/D3D12RayTracing.cpp
@@ -2496,7 +2496,6 @@ FD3D12RayTracingScene::FD3D12RayTracingScene(FD3D12Adapter* Adapter, const FRayT
 			Geometries.Add(Geometry);
 		}
 	}
-	NumTotalSegments = NumTotalSegments;
 
 	NumCallableShaderSlots = Initializer.NumCallableShaderSlots;
 	NumMissShaderSlots = FMath::Max<uint32>(1, Initializer.NumMissShaderSlots);
