Don't call PIXEvent when the DLL fails to load

diff --git a/Engine/Source/Runtime/D3D12RHI/Private/D3D12CommandContext.cpp b/Engine/Source/Runtime/D3D12RHI/Private/D3D12CommandContext.cpp
--- a/Engine/Source/Runtime/D3D12RHI/Private/D3D12CommandContext.cpp
+++ b/Engine/Source/Runtime/D3D12RHI/Private/D3D12CommandContext.cpp
@@ -186,7 +186,10 @@ void FD3D12CommandContext::RHIPushEvent(const TCHAR* Name, FColor Color)
 #endif
 
 #if USE_PIX
-	PIXBeginEvent(CommandListHandle.GraphicsCommandList(), PIX_COLOR(Color.R, Color.G, Color.B), Name);
+	if (FD3D12DynamicRHI::GetD3DRHI()->IsPixEventEnabled())
+	{
+		PIXBeginEvent(CommandListHandle.GraphicsCommandList(), PIX_COLOR(Color.R, Color.G, Color.B), Name);
+	}
 #endif // USE_PIX
 }
 
@@ -219,7 +222,10 @@ void FD3D12CommandContext::RHIPopEvent()
 #endif
 
 #if USE_PIX
-	PIXEndEvent(CommandListHandle.GraphicsCommandList());
+	if (FD3D12DynamicRHI::GetD3DRHI()->IsPixEventEnabled())
+	{
+		PIXEndEvent(CommandListHandle.GraphicsCommandList());
+	}
 #endif
 }
 
diff --git a/Engine/Source/Runtime/D3D12RHI/Private/D3D12RHI.cpp b/Engine/Source/Runtime/D3D12RHI/Private/D3D12RHI.cpp
--- a/Engine/Source/Runtime/D3D12RHI/Private/D3D12RHI.cpp
+++ b/Engine/Source/Runtime/D3D12RHI/Private/D3D12RHI.cpp
@@ -44,8 +44,9 @@ FD3D12DynamicRHI* FD3D12DynamicRHI::SingleD3DRHI = nullptr;
 
 using namespace D3D12RHI;
 
-FD3D12DynamicRHI::FD3D12DynamicRHI(const TArray<TSharedPtr<FD3D12Adapter>>& ChosenAdaptersIn) :
+FD3D12DynamicRHI::FD3D12DynamicRHI(const TArray<TSharedPtr<FD3D12Adapter>>& ChosenAdaptersIn, bool bInPixEventEnabled) :
 	ChosenAdapters(ChosenAdaptersIn),
+	bPixEventEnabled(bInPixEventEnabled),
 	AmdAgsContext(nullptr),
 	AmdSupportedExtensionFlags(0),
 	FlipEvent(INVALID_HANDLE_VALUE),
diff --git a/Engine/Source/Runtime/D3D12RHI/Private/D3D12RHIPrivate.h b/Engine/Source/Runtime/D3D12RHI/Private/D3D12RHIPrivate.h
--- a/Engine/Source/Runtime/D3D12RHI/Private/D3D12RHIPrivate.h
+++ b/Engine/Source/Runtime/D3D12RHI/Private/D3D12RHIPrivate.h
@@ -209,7 +209,7 @@ private:
 public:
 
 	/** Initialization constructor. */
-	FD3D12DynamicRHI(const TArray<TSharedPtr<FD3D12Adapter>>& ChosenAdaptersIn);
+	FD3D12DynamicRHI(const TArray<TSharedPtr<FD3D12Adapter>>& ChosenAdaptersIn, bool bInPixEventEnabled);
 
 	/** Destructor */
 	virtual ~FD3D12DynamicRHI();
@@ -899,6 +899,8 @@ public:
 	FD3D12Adapter& GetAdapter(uint32_t Index = 0) { return *ChosenAdapters[Index]; }
 	int32 GetNumAdapters() const { return ChosenAdapters.Num(); }
 
+	bool IsPixEventEnabled() const { return bPixEventEnabled; }
+
 	template<typename PerDeviceFunction>
 	void ForEachDevice(ID3D12Device* inDevice, const PerDeviceFunction& pfPerDeviceFunction)
 	{
@@ -923,6 +925,9 @@ protected:
 
 	TArray<TSharedPtr<FD3D12Adapter>> ChosenAdapters;
 
+	/** Can pix events be used */
+	bool bPixEventEnabled = false;
+
 	/** The feature level of the device. */
 	D3D_FEATURE_LEVEL FeatureLevel;
 
@@ -993,7 +998,6 @@ class FD3D12DynamicRHIModule : public IDynamicRHIModule
 public:
 
 	FD3D12DynamicRHIModule()
-		: WindowsPixDllHandle(nullptr)
 	{
 	}
 
@@ -1011,7 +1015,11 @@ public:
 	virtual FDynamicRHI* CreateRHI(ERHIFeatureLevel::Type RequestedFeatureLevel = ERHIFeatureLevel::Num) override;
 
 private:
-	void* WindowsPixDllHandle;
+
+#if USE_PIX && (PLATFORM_WINDOWS || PLATFORM_HOLOLENS)
+	void* WindowsPixDllHandle = nullptr;
+#endif // USE_PIX && (PLATFORM_WINDOWS || PLATFORM_HOLOLENS)
+
 	TArray<TSharedPtr<FD3D12Adapter>> ChosenAdapters;
 
 	// set MaxSupportedFeatureLevel and ChosenAdapter
diff --git a/Engine/Source/Runtime/D3D12RHI/Private/D3D12Texture.cpp b/Engine/Source/Runtime/D3D12RHI/Private/D3D12Texture.cpp
--- a/Engine/Source/Runtime/D3D12RHI/Private/D3D12Texture.cpp
+++ b/Engine/Source/Runtime/D3D12RHI/Private/D3D12Texture.cpp
@@ -2159,7 +2159,10 @@ public:
 				FD3D12Resource* UploadBuffer = UpdateInfo.SrcResourceLocation->GetResource();
 				CD3DX12_TEXTURE_COPY_LOCATION SourceCopyLocation(UploadBuffer->GetResource(), UpdateInfo.PlacedSubresourceFootprint);
 #if USE_PIX
-				PIXBeginEvent(NativeCmdList.GraphicsCommandList(), PIX_COLOR(255, 255, 255), TEXT("EndMultiUpdateTexture3D"));
+				if (FD3D12DynamicRHI::GetD3DRHI()->IsPixEventEnabled())
+				{
+					PIXBeginEvent(NativeCmdList.GraphicsCommandList(), PIX_COLOR(255, 255, 255), TEXT("EndMultiUpdateTexture3D"));
+				}
 #endif
 				NativeCmdList->CopyTextureRegion(
 					&DestCopyLocation,
@@ -2173,7 +2176,10 @@ public:
 
 				DEBUG_EXECUTE_COMMAND_CONTEXT(Device->GetDefaultCommandContext());
 #if USE_PIX
-				PIXEndEvent(NativeCmdList.GraphicsCommandList());
+				if (FD3D12DynamicRHI::GetD3DRHI()->IsPixEventEnabled())
+				{
+					PIXEndEvent(NativeCmdList.GraphicsCommandList());
+				}
 #endif
 			}
 		}
@@ -2375,7 +2381,10 @@ public:
 			FD3D12Device* Device = TextureLink->GetParentDevice();
 			FD3D12CommandListHandle& NativeCmdList = Device->GetDefaultCommandContext().CommandListHandle;
 #if USE_PIX
-			PIXBeginEvent(NativeCmdList.GraphicsCommandList(), PIX_COLOR(255, 255, 255), TEXT("EndUpdateTexture3D"));
+			if (FD3D12DynamicRHI::GetD3DRHI()->IsPixEventEnabled())
+			{
+				PIXBeginEvent(NativeCmdList.GraphicsCommandList(), PIX_COLOR(255, 255, 255), TEXT("EndUpdateTexture3D"));
+			}
 #endif
 			CD3DX12_TEXTURE_COPY_LOCATION DestCopyLocation(TextureLink->GetResource()->GetResource(), MipIdx);
 			CD3DX12_TEXTURE_COPY_LOCATION SourceCopyLocation(UploadBuffer->GetResource(), PlacedSubresourceFootprint);
@@ -2400,7 +2409,10 @@ public:
 
 			DEBUG_EXECUTE_COMMAND_CONTEXT(Device->GetDefaultCommandContext());
 #if USE_PIX
-			PIXEndEvent(NativeCmdList.GraphicsCommandList());
+			if (FD3D12DynamicRHI::GetD3DRHI()->IsPixEventEnabled())
+			{
+				PIXEndEvent(NativeCmdList.GraphicsCommandList());
+			}
 #endif
 		}
 
diff --git a/Engine/Source/Runtime/D3D12RHI/Private/Windows/WindowsD3D12Device.cpp b/Engine/Source/Runtime/D3D12RHI/Private/Windows/WindowsD3D12Device.cpp
--- a/Engine/Source/Runtime/D3D12RHI/Private/Windows/WindowsD3D12Device.cpp
+++ b/Engine/Source/Runtime/D3D12RHI/Private/Windows/WindowsD3D12Device.cpp
@@ -469,7 +469,13 @@ FDynamicRHI* FD3D12DynamicRHIModule::CreateRHI(ERHIFeatureLevel::Type RequestedF
 		GMaxRHIShaderPlatform = SP_PCD3D_SM5;
 	}
 
-	GD3D12RHI = new FD3D12DynamicRHI(ChosenAdapters);
+#if USE_PIX
+	bool bPixEventEnabled = (WindowsPixDllHandle != nullptr);
+#else
+	bool bPixEventEnabled = false;
+#endif // USE_PIX
+
+	GD3D12RHI = new FD3D12DynamicRHI(ChosenAdapters, bPixEventEnabled);
 #if ENABLE_RHI_VALIDATION
 	if (FParse::Param(FCommandLine::Get(), TEXT("RHIValidation")))
 	{
