Added a UAV barrier at the end of each BLAS build batch. This fixes GPU when running with ray tracing due to build and update commands overlapping on GPU and corrupting memory.

diff --git a/Engine/Source/Runtime/D3D12RHI/Private/D3D12RayTracing.cpp b/Engine/Source/Runtime/D3D12RHI/Private/D3D12RayTracing.cpp
--- a/Engine/Source/Runtime/D3D12RHI/Private/D3D12RayTracing.cpp
+++ b/Engine/Source/Runtime/D3D12RHI/Private/D3D12RayTracing.cpp
@@ -2940,6 +2940,13 @@ void FD3D12CommandContext::RHIBuildAccelerationStructures(const TArrayView<const
 
 		Geometry->BuildAccelerationStructure(*this, P.BuildMode);
 	}
+
+	// Add a UAV barrier after each acceleration structure build batch.
+	// This is required because there are currently no explicit read/write barriers
+	// for acceleration structures, but we need to ensure that all commands
+	// are complete before BLAS is used again on the GPU.
+
+	CommandListHandle.AddUAVBarrier();
 }
 
 void FD3D12CommandContext::RHIBuildAccelerationStructure(FRHIRayTracingScene* InScene)
