Add a per-platform limit on resource barrier batch sizes

- D3D12 platforms only

#rnx

diff --git a/Engine/Source/Runtime/D3D12RHI/Private/D3D12DirectCommandListManager.cpp b/Engine/Source/Runtime/D3D12RHI/Private/D3D12DirectCommandListManager.cpp
--- a/Engine/Source/Runtime/D3D12RHI/Private/D3D12DirectCommandListManager.cpp
+++ b/Engine/Source/Runtime/D3D12RHI/Private/D3D12DirectCommandListManager.cpp
@@ -847,23 +847,20 @@ uint32 FD3D12CommandListManager::GetResourceBarrierCommandList(FD3D12CommandList
 #if DEBUG_RESOURCE_STATES
 			LogResourceBarriers(BarrierDescs.Num(), BarrierDescs.GetData(), hResourceBarrierList.CommandList());
 #endif
-
-#if USE_PIX && PLATFORM_XBOXONE 
-			//there was a bug in the instrumented driver that corrupts the cmdBuffer if more than 2000 Barrieres are submitted at once
-			if (BarrierDescs.Num() > 1900)
+			const int32 BarrierBatchMax = FD3D12DynamicRHI::GetResourceBarrierBatchSizeLimit();
+			if (BarrierDescs.Num() > BarrierBatchMax)
 			{
 				int Num = BarrierDescs.Num();
 				D3D12_RESOURCE_BARRIER* Ptr = BarrierDescs.GetData();
 				while (Num > 0)
 				{
-					int DispatchNum = FMath::Min(Num, 1900);
+					const int DispatchNum = FMath::Min(Num, BarrierBatchMax);
 					hResourceBarrierList->ResourceBarrier(DispatchNum, Ptr);
-					Ptr += 1900;
-					Num -= 1900;
+					Ptr += BarrierBatchMax;
+					Num -= BarrierBatchMax;
 				}
 			}
 			else
-#endif
 			{
 				hResourceBarrierList->ResourceBarrier(BarrierDescs.Num(), BarrierDescs.GetData());
 			}
diff --git a/Engine/Source/Runtime/D3D12RHI/Private/D3D12RHIPrivate.h b/Engine/Source/Runtime/D3D12RHI/Private/D3D12RHIPrivate.h
--- a/Engine/Source/Runtime/D3D12RHI/Private/D3D12RHIPrivate.h
+++ b/Engine/Source/Runtime/D3D12RHI/Private/D3D12RHIPrivate.h
@@ -450,6 +450,8 @@ public:
 	bool IsQuadBufferStereoEnabled() const;
 	void DisableQuadBufferStereo();
 
+	static int32 GetResourceBarrierBatchSizeLimit();
+
 	void* LockBuffer(FRHICommandListImmediate* RHICmdList, FD3D12Buffer* Buffer, uint32 BufferSize, uint32 BufferUsage, uint32 Offset, uint32 Size, EResourceLockMode LockMode);
 	void UnlockBuffer(FRHICommandListImmediate* RHICmdList, FD3D12Buffer* Buffer, uint32 BufferUsage);
 
diff --git a/Engine/Source/Runtime/D3D12RHI/Private/Windows/WindowsD3D12Device.cpp b/Engine/Source/Runtime/D3D12RHI/Private/Windows/WindowsD3D12Device.cpp
--- a/Engine/Source/Runtime/D3D12RHI/Private/Windows/WindowsD3D12Device.cpp
+++ b/Engine/Source/Runtime/D3D12RHI/Private/Windows/WindowsD3D12Device.cpp
@@ -797,6 +797,11 @@ void FD3D12DynamicRHI::DisableQuadBufferStereo()
 	bIsQuadBufferStereoEnabled = false;
 }
 
+int32 FD3D12DynamicRHI::GetResourceBarrierBatchSizeLimit()
+{
+	return INT32_MAX;
+}
+
 void FD3D12Device::Initialize()
 {
 	check(IsInGameThread());
