Only return a WARP device if explicitly requested.

diff --git a/Engine/Source/Runtime/D3D12RHI/Private/Windows/WindowsD3D12Device.cpp b/Engine/Source/Runtime/D3D12RHI/Private/Windows/WindowsD3D12Device.cpp
--- a/Engine/Source/Runtime/D3D12RHI/Private/Windows/WindowsD3D12Device.cpp
+++ b/Engine/Source/Runtime/D3D12RHI/Private/Windows/WindowsD3D12Device.cpp
@@ -370,8 +370,8 @@ void FD3D12DynamicRHIModule::FindAdapter()
 
 				FD3D12AdapterDesc CurrentAdapter(AdapterDesc, AdapterIndex, MaxSupportedFeatureLevel, NumNodes);
 				
-				// Requested WARP, reject all other adapters.
-				const bool bSkipRequestedWARP = bRequestedWARP && !bIsWARP;
+				// If requested WARP, then reject all other adapters. If WARP not requested, then reject the WARP device.
+				const bool bSkipRequestedWARP = (bRequestedWARP && !bIsWARP) || (!bRequestedWARP && bIsWARP);
 
 				// we don't allow the PerfHUD adapter
 				const bool bSkipPerfHUDAdapter = bIsPerfHUD && !bAllowPerfHUD;
