Workaround for ray tracing with mesh LOD streaming by deferring index buffer initialization until BLAS build. This fixes the race condition where the buffer is moved by mesh streaming using an RHI command while the RT geometry is created on render thread.

diff --git a/Engine/Source/Runtime/D3D12RHI/Private/D3D12RayTracing.cpp b/Engine/Source/Runtime/D3D12RHI/Private/D3D12RayTracing.cpp
--- a/Engine/Source/Runtime/D3D12RHI/Private/D3D12RayTracing.cpp
+++ b/Engine/Source/Runtime/D3D12RHI/Private/D3D12RayTracing.cpp
@@ -1960,12 +1960,6 @@ FD3D12RayTracingGeometry::FD3D12RayTracingGeometry(const FRayTracingGeometryInit
 
 	FMemory::Memzero(PostBuildInfoBufferReadbackFences);
 
-	if (Initializer.IndexBuffer)
-	{
-		checkf(Initializer.IndexBuffer->GetStride() == 2 || Initializer.IndexBuffer->GetStride() == 4, TEXT("Index buffer must be 16 or 32 bit."));
-	}
-
-	IndexStride = Initializer.IndexBuffer ? Initializer.IndexBuffer->GetStride() : 0; // stride 0 means implicit triangle list for non-indexed geometry
 	IndexOffsetInBytes = Initializer.IndexBufferOffset;
 	TotalPrimitiveCount = Initializer.TotalPrimitiveCount;
 
@@ -2138,6 +2132,13 @@ void FD3D12RayTracingGeometry::BuildAccelerationStructure(FD3D12CommandContext&
 {
 	TRACE_CPUPROFILER_EVENT_SCOPE(BuildAccelerationStructure_BottomLevel);
 
+	if (RHIIndexBuffer)
+	{
+		checkf(RHIIndexBuffer->GetStride() == 2 || RHIIndexBuffer->GetStride() == 4, TEXT("Index buffer must be 16 or 32 bit."));
+	}
+
+	IndexStride = RHIIndexBuffer ? RHIIndexBuffer->GetStride() : 0; // stride 0 means implicit triangle list for non-indexed geometry
+
 	const bool bIsUpdate = BuildMode == EAccelerationStructureBuildMode::Update;
 	static constexpr uint32 IndicesPerPrimitive = 3; // Only triangle meshes are supported
 
