Exclude use of IDXGISwapChain::GetBuffer when D3D12_VIEWPORT_EXPOSES_SWAP_CHAIN is not enabled

- D3D12_VIEWPORT_EXPOSES_SWAP_CHAIN is currently utilized as a workaround but this should be revisited (tracked via MTHRL-169)

#rnx

diff --git a/Engine/Source/Runtime/D3D12RHI/Private/D3D12Viewport.cpp b/Engine/Source/Runtime/D3D12RHI/Private/D3D12Viewport.cpp
--- a/Engine/Source/Runtime/D3D12RHI/Private/D3D12Viewport.cpp
+++ b/Engine/Source/Runtime/D3D12RHI/Private/D3D12Viewport.cpp
@@ -202,13 +202,19 @@ bool bNeedSwapChain = true;
  */
 FD3D12Texture2D* GetSwapChainSurface(FD3D12Device* Parent, EPixelFormat PixelFormat, uint32 SizeX, uint32 SizeY, IDXGISwapChain* SwapChain, uint32 BackBufferIndex, TRefCountPtr<ID3D12Resource> BackBufferResourceOverride)
 {
+	verify(D3D12_VIEWPORT_EXPOSES_SWAP_CHAIN || SwapChain == nullptr);
+
 	FD3D12Adapter* Adapter = Parent->GetParentAdapter();
 
 	// Grab the back buffer
 	TRefCountPtr<ID3D12Resource> BackBufferResource;
 	if (SwapChain)
 	{
+#if D3D12_VIEWPORT_EXPOSES_SWAP_CHAIN
 		VERIFYD3D12RESULT_EX(SwapChain->GetBuffer(BackBufferIndex, IID_PPV_ARGS(BackBufferResource.GetInitReference())), Parent->GetDevice());
+#else // #if D3D12_VIEWPORT_EXPOSES_SWAP_CHAIN
+		return nullptr;
+#endif // #if D3D12_VIEWPORT_EXPOSES_SWAP_CHAIN
 	}
 	else if (BackBufferResourceOverride.IsValid())
 	{
