Support FD3D12Viewport without swap chain references

- enabled via disabling D3D12_VIEWPORT_EXPOSES_SWAP_CHAIN
  - D3D12_VIEWPORT_EXPOSES_SWAP_CHAIN is enabled by default

#rnx

diff --git a/Engine/Source/Runtime/D3D12RHI/Private/D3D12Viewport.cpp b/Engine/Source/Runtime/D3D12RHI/Private/D3D12Viewport.cpp
--- a/Engine/Source/Runtime/D3D12RHI/Private/D3D12Viewport.cpp
+++ b/Engine/Source/Runtime/D3D12RHI/Private/D3D12Viewport.cpp
@@ -354,14 +354,14 @@ FD3D12Viewport::~FD3D12Viewport()
 {
 	check(IsInRenderingThread());
 
+#if !PLATFORM_HOLOLENS && D3D12_VIEWPORT_EXPOSES_SWAP_CHAIN
 	// If the swap chain was in fullscreen mode, switch back to windowed before releasing the swap chain.
 	// DXGI throws an error otherwise.
-#if !PLATFORM_HOLOLENS
 	if (SwapChain1)
 	{
 		SwapChain1->SetFullscreenState(0, nullptr);
 	}
-#endif
+#endif // !PLATFORM_HOLOLENS && D3D12_VIEWPORT_EXPOSES_SWAP_CHAIN
 
 	GetParentAdapter()->GetViewports().Remove(this);
 
@@ -509,7 +509,7 @@ void FD3D12Viewport::Resize(uint32 InSizeX, uint32 InSizeY, bool bInIsFullscreen
 
 		check(SizeX > 0);
 		check(SizeY > 0);
-
+#if D3D12_VIEWPORT_EXPOSES_SWAP_CHAIN
 		if (bNeedSwapChain)
 		{
 			if (bInIsFullscreen)
@@ -521,6 +521,7 @@ void FD3D12Viewport::Resize(uint32 InSizeX, uint32 InSizeY, bool bInIsFullscreen
 				}
 			}
 		}
+#endif // D3D12_VIEWPORT_EXPOSES_SWAP_CHAIN
 	}
 
 	if (bIsFullscreen != bInIsFullscreen)
diff --git a/Engine/Source/Runtime/D3D12RHI/Public/D3D12Viewport.h b/Engine/Source/Runtime/D3D12RHI/Public/D3D12Viewport.h
--- a/Engine/Source/Runtime/D3D12RHI/Public/D3D12Viewport.h
+++ b/Engine/Source/Runtime/D3D12RHI/Public/D3D12Viewport.h
@@ -11,6 +11,10 @@
 
 #define ALTERNATE_TIMESTAMP_METRIC 0
 
+#ifndef D3D12_VIEWPORT_EXPOSES_SWAP_CHAIN
+#define D3D12_VIEWPORT_EXPOSES_SWAP_CHAIN 1
+#endif // D3D12_VIEWPORT_EXPOSES_SWAP_CHAIN
+
 static DXGI_FORMAT GetRenderTargetFormat(EPixelFormat PixelFormat)
 {
 	DXGI_FORMAT	DXFormat = (DXGI_FORMAT)GPixelFormats[PixelFormat].PlatformFormat;
@@ -105,10 +109,12 @@ public:
 	void WaitForFrameEventCompletion();
 	void IssueFrameEvent();
 
+#if D3D12_VIEWPORT_EXPOSES_SWAP_CHAIN
 	IDXGISwapChain1* GetSwapChain() const { return SwapChain1; }
 	IDXGISwapChain1* GetSDRSwapChain() const { return (PixelFormat == SDRPixelFormat) ? GetSwapChain() : (IDXGISwapChain1*)SDRSwapChain1; }
-
 	virtual void* GetNativeSwapChain() const override { return GetSwapChain(); }
+#endif // #if D3D12_VIEWPORT_EXPOSES_SWAP_CHAIN
+
 	virtual void* GetNativeBackBufferTexture() const override { return GetBackBuffer_RHIThread()->GetResource(); }
 	virtual void* GetNativeBackBufferRT() const override { return GetBackBuffer_RHIThread()->GetRenderTargetView(0, 0); }
 
@@ -170,15 +176,16 @@ private:
 	EPixelFormat PixelFormat;
 	bool bIsValid;
 	bool bAllowTearing;
-	TRefCountPtr<IDXGISwapChain1> SwapChain1;
 
+#if D3D12_VIEWPORT_EXPOSES_SWAP_CHAIN
+	TRefCountPtr<IDXGISwapChain1> SwapChain1;
 #if PLATFORM_WINDOWS || PLATFORM_HOLOLENS
 	bool bHDRMetaDataSet;
 	DXGI_COLOR_SPACE_TYPE ColorSpace;
 	TRefCountPtr<IDXGISwapChain4> SwapChain4;
 #endif
-
 	TRefCountPtr<IDXGISwapChain1> SDRSwapChain1;
+#endif // D3D12_VIEWPORT_EXPOSES_SWAP_CHAIN
 
 	TArray<TRefCountPtr<FD3D12Texture2D>> BackBuffers;
 	TArray<uint32> BackBufferGPUIndices;
