Reinstate dispatch L2 indirect argument policy for DX12 async compute command signature. Removed during merge of Reverb into Main due to versioning issues. Original CL 13030723.

diff --git a/Engine/Source/Runtime/D3D12RHI/Private/D3D12Adapter.cpp b/Engine/Source/Runtime/D3D12RHI/Private/D3D12Adapter.cpp
--- a/Engine/Source/Runtime/D3D12RHI/Private/D3D12Adapter.cpp
+++ b/Engine/Source/Runtime/D3D12RHI/Private/D3D12Adapter.cpp
@@ -953,6 +953,9 @@ void FD3D12Adapter::CreateSignatures()
 	commandSignatureDesc.ByteStride = sizeof(D3D12_DISPATCH_ARGUMENTS);
 	VERIFYD3D12RESULT(Device->CreateCommandSignature(&commandSignatureDesc, nullptr, IID_PPV_ARGS(DispatchIndirectGraphicsCommandSignature.GetInitReference())));
 
+#if !PLATFORM_WINDOWS && !PLATFORM_HOLOLENS
+	indirectParameterDesc[0].Type = D3D12XBOX_INDIRECT_ARGUMENT_TYPE_DISPATCH_L2; // D3D12XBOX_INDIRECT_ARGUMENT_TYPE_DISPATCH_SINGLE
+#endif
 	VERIFYD3D12RESULT(Device->CreateCommandSignature(&commandSignatureDesc, nullptr, IID_PPV_ARGS(DispatchIndirectComputeCommandSignature.GetInitReference())));
 }
 
