Re-enable DEBUG_BARRIER_DISPATCHES on Mithril. It seems like compute UAV barriers are still incorrect wrt. indirect argument generation in the RHI. Nanite will GPU crash 100% without this, due to corrupt arguments between passes.

diff --git a/Engine/Source/Runtime/D3D12RHI/Private/D3D12RHIPrivate.h b/Engine/Source/Runtime/D3D12RHI/Private/D3D12RHIPrivate.h
--- a/Engine/Source/Runtime/D3D12RHI/Private/D3D12RHIPrivate.h
+++ b/Engine/Source/Runtime/D3D12RHI/Private/D3D12RHIPrivate.h
@@ -85,7 +85,7 @@ typedef FD3D12StateCacheBase FD3D12StateCache;
 #define USE_D3D12RHI_RESOURCE_STATE_TRACKING 1	// Fully relying on the engine's resource barriers is a work in progress. For now, continue to use the D3D12 RHI's resource state tracking.
 
 #define EXECUTE_DEBUG_COMMAND_LISTS 0
-#define DEBUG_BARRIER_DISPATCHES 0
+#define DEBUG_BARRIER_DISPATCHES 1
 #define ENABLE_PLACED_RESOURCES 0 // Disabled due to a couple of NVidia bugs related to placed resources. Works fine on Intel
 #define NAME_OBJECTS !(UE_BUILD_SHIPPING || UE_BUILD_TEST)	// Name objects in all builds except shipping
 #define LOG_PSO_CREATES (0 && STATS)	// Logs Create Pipeline State timings (also requires STATS)
