Fixed compile errors in D3D12 RHI due to accessing FD3D12CommandContext that is not defined yet when using VS2019 16.7.0

Access to FD3D12CommandContext is now done through helper functions that are delayed-instantiated.

#[fyi] Emil.Persson, Jonas.Meyer

diff --git a/Engine/Source/Runtime/D3D12RHI/Private/D3D12StateCachePrivate.h b/Engine/Source/Runtime/D3D12RHI/Private/D3D12StateCachePrivate.h
--- a/Engine/Source/Runtime/D3D12RHI/Private/D3D12StateCachePrivate.h
+++ b/Engine/Source/Runtime/D3D12RHI/Private/D3D12StateCachePrivate.h
@@ -481,7 +481,7 @@ protected:
 			PipelineState.Common.CurrentShaderUAVCounts[Traits::Frequency]     = (Shader) ? Shader->ResourceCounts.NumUAVs     : 0;
 		
 			// Shader changed so its resource table is dirty
-			this->CmdContext->DirtyUniformBuffers[Traits::Frequency] = 0xffff;
+			SetDirtyUniformBuffers(this->CmdContext, Traits::Frequency);
 		}
 	}
 
@@ -523,11 +523,28 @@ protected:
 #endif // #if USE_D3D12RHI_PRE_PSO_SET_WORKAROUND
 
 			check(CurrentPSO);
-			this->CmdContext->CommandListHandle->SetPipelineState(CurrentPSO);
+			SetPipelineState(this->CmdContext, CurrentPSO);
 			PipelineState.Common.bNeedSetPSO = false;
 		}
 	}
 
+private:
+
+	// SetDirtyUniformBuffers and SetPipelineState helper functions are required
+	// to allow using FD3D12CommandContext type which is not defined at this point.
+	// Making ContextType a template parameter delays instantiation of these functions.
+
+	template <typename ContextType>
+	static void SetDirtyUniformBuffers(ContextType* Context, EShaderFrequency Frequency)
+	{
+		Context->DirtyUniformBuffers[Frequency] = 0xffff;
+	}
+
+	template <typename ContextType>
+	static void SetPipelineState(ContextType* Context, ID3D12PipelineState* State)
+	{
+		Context->CommandListHandle->SetPipelineState(State);
+	}
 
 public:
 
