Xbox:  resolve check in CreateD3D12Texture2D on Xbox platforms

- duplicate AsyncReallocateTexture2D_RenderThread() in the derived Xbox RHI to workaround the templated member call
  - this is consistent with AsyncReallocateTexture2D()

- move FRHICommandD3D12AsyncReallocateTexture2D to the D3D12Texture.h so it has visibility in derived Xbox RHI but keep the Execute() call in the cpp so that it can call the static functions declared there

#rnx

diff --git a/Engine/Source/Runtime/D3D12RHI/Private/D3D12Texture.cpp b/Engine/Source/Runtime/D3D12RHI/Private/D3D12Texture.cpp
--- a/Engine/Source/Runtime/D3D12RHI/Private/D3D12Texture.cpp
+++ b/Engine/Source/Runtime/D3D12RHI/Private/D3D12Texture.cpp
@@ -1507,30 +1507,10 @@ static void DoAsyncReallocateTexture2D(FD3D12Texture2D* Texture2D, FD3D12Texture
 }
 
 
-struct FRHICommandD3D12AsyncReallocateTexture2D final : public FRHICommand<FRHICommandD3D12AsyncReallocateTexture2D>
+void FRHICommandD3D12AsyncReallocateTexture2D::Execute(FRHICommandListBase& RHICmdList)
 {
-	FD3D12Texture2D* OldTexture;
-	FD3D12Texture2D* NewTexture;
-	int32 NewMipCount;
-	int32 NewSizeX;
-	int32 NewSizeY;
-	FThreadSafeCounter* RequestStatus;
-
-	FORCEINLINE_DEBUGGABLE FRHICommandD3D12AsyncReallocateTexture2D(FD3D12Texture2D* InOldTexture, FD3D12Texture2D* InNewTexture, int32 InNewMipCount, int32 InNewSizeX, int32 InNewSizeY, FThreadSafeCounter* InRequestStatus)
-		: OldTexture(InOldTexture)
-		, NewTexture(InNewTexture)
-		, NewMipCount(InNewMipCount)
-		, NewSizeX(InNewSizeX)
-		, NewSizeY(InNewSizeY)
-		, RequestStatus(InRequestStatus)
-	{
-	}
-
-	void Execute(FRHICommandListBase& RHICmdList)
-	{
-		DoAsyncReallocateTexture2D(OldTexture, NewTexture, NewMipCount, NewSizeX, NewSizeY, RequestStatus);
-	}
-};
+	DoAsyncReallocateTexture2D(OldTexture, NewTexture, NewMipCount, NewSizeX, NewSizeY, RequestStatus);
+}
 
 
 FTexture2DRHIRef FD3D12DynamicRHI::AsyncReallocateTexture2D_RenderThread(class FRHICommandListImmediate& RHICmdList, FRHITexture2D* Texture2DRHI, int32 NewMipCount, int32 NewSizeX, int32 NewSizeY, FThreadSafeCounter* RequestStatus)
diff --git a/Engine/Source/Runtime/D3D12RHI/Private/D3D12Texture.h b/Engine/Source/Runtime/D3D12RHI/Private/D3D12Texture.h
--- a/Engine/Source/Runtime/D3D12RHI/Private/D3D12Texture.h
+++ b/Engine/Source/Runtime/D3D12RHI/Private/D3D12Texture.h
@@ -566,4 +566,26 @@ template<>
 struct TD3D12ResourceTraits<FRHITextureReference>
 {
 	typedef FD3D12TextureReference TConcreteType;
+};
+
+struct FRHICommandD3D12AsyncReallocateTexture2D final : public FRHICommand<FRHICommandD3D12AsyncReallocateTexture2D>
+{
+	FD3D12Texture2D* OldTexture;
+	FD3D12Texture2D* NewTexture;
+	int32 NewMipCount;
+	int32 NewSizeX;
+	int32 NewSizeY;
+	FThreadSafeCounter* RequestStatus;
+
+	FORCEINLINE_DEBUGGABLE FRHICommandD3D12AsyncReallocateTexture2D(FD3D12Texture2D* InOldTexture, FD3D12Texture2D* InNewTexture, int32 InNewMipCount, int32 InNewSizeX, int32 InNewSizeY, FThreadSafeCounter* InRequestStatus)
+		: OldTexture(InOldTexture)
+		, NewTexture(InNewTexture)
+		, NewMipCount(InNewMipCount)
+		, NewSizeX(InNewSizeX)
+		, NewSizeY(InNewSizeY)
+		, RequestStatus(InRequestStatus)
+	{
+	}
+
+	void Execute(FRHICommandListBase& RHICmdList);
 };
\ No newline at end of file
