Temporary workaround for some platforms using GPU uncached memory when D3D12RHI_HEAP_FLAG_ALLOW_INDIRECT_BUFFERS is specified on a heap. Saves ~6.3ms in Reverb.

diff --git a/Engine/Source/Runtime/D3D12RHI/Private/D3D12Allocation.cpp b/Engine/Source/Runtime/D3D12RHI/Private/D3D12Allocation.cpp
--- a/Engine/Source/Runtime/D3D12RHI/Private/D3D12Allocation.cpp
+++ b/Engine/Source/Runtime/D3D12RHI/Private/D3D12Allocation.cpp
@@ -51,6 +51,14 @@ static FAutoConsoleVariableRef CVarD3D12FastAllocatorMinPagesToRetain(
 	TEXT("Minimum number of pages to retain. Pages below this limit will never be released. Pages above can be released after being unused for a certain number of frames."),
 	ECVF_Default);
 
+// TODO: TEMP HACK - all D3D12 platforms should pass this flag. Incoming update will fix the problem.
+static int32 GD3D12SkipIndirectArgsHeap = 0;
+static FAutoConsoleVariableRef CVarD3D12SkipIndirectArgsHeap(
+	TEXT("d3d12.SkipIndirectArgsHeap"),
+	GD3D12SkipIndirectArgsHeap,
+	TEXT("1: Enable temporary hack to skip IA flag on heaps."),
+	ECVF_ReadOnly);
+
 namespace ED3D12AllocatorID
 {
 	enum Type
@@ -1009,7 +1017,11 @@ FD3D12ResourceAllocator::FInitConfig FD3D12DefaultBufferPool::GetResourceAllocat
 	if (EnumHasAnyFlags(InBufferUsage, BUF_DrawIndirect))
 	{
 		check(InResourceFlags & D3D12_RESOURCE_FLAG_ALLOW_UNORDERED_ACCESS);
-		InitConfig.HeapFlags |= D3D12RHI_HEAP_FLAG_ALLOW_INDIRECT_BUFFERS;
+		if (GD3D12SkipIndirectArgsHeap == 0) // TEMP HACK
+		{
+			// TODO: TEMP HACK - all D3D12 platforms should pass this flag. Incoming update will fix the problem.
+			InitConfig.HeapFlags |= D3D12RHI_HEAP_FLAG_ALLOW_INDIRECT_BUFFERS;
+		}
 	}
 
 	return InitConfig;
