Compilation fixes for scrict compilation mode

#rnx

diff --git a/Engine/Source/Runtime/D3D12RHI/Private/D3D12TimedIntervalQuery.cpp b/Engine/Source/Runtime/D3D12RHI/Private/D3D12TimedIntervalQuery.cpp
--- a/Engine/Source/Runtime/D3D12RHI/Private/D3D12TimedIntervalQuery.cpp
+++ b/Engine/Source/Runtime/D3D12RHI/Private/D3D12TimedIntervalQuery.cpp
@@ -108,13 +108,13 @@ public:
 							D3D12_RANGE ReadRange;
 							ReadRange.Begin = Batch.StartIndex * sizeof(uint64);
 							ReadRange.End = ReadRange.Begin + Batch.Size * sizeof(uint64);
-							VERIFYD3D12RESULT(ResolveBuffer->Map(0, &ReadRange, &static_cast<void*>(MappedResults)));
+							VERIFYD3D12RESULT(ResolveBuffer->Map(0, &ReadRange, reinterpret_cast<void**>(&MappedResults)));
 							FMemory::Memcpy(QueryResultsOut.GetData(), MappedResults + Batch.StartIndex, (PoolSize - Batch.StartIndex) * sizeof(uint64));
 							ResolveBuffer->Unmap(0, nullptr);
 
 							ReadRange.Begin = 0;
 							ReadRange.End = ReadRange.Begin + (Batch.Size - (PoolSize - Batch.StartIndex)) * sizeof(uint64);
-							VERIFYD3D12RESULT(ResolveBuffer->Map(0, &ReadRange, &static_cast<void*>(MappedResults)));
+							VERIFYD3D12RESULT(ResolveBuffer->Map(0, &ReadRange, reinterpret_cast<void**>(&MappedResults)));
 							FMemory::Memcpy(QueryResultsOut.GetData() + (PoolSize - Batch.StartIndex), MappedResults, (Batch.Size - (PoolSize - Batch.StartIndex)) * sizeof(uint64));
 							ResolveBuffer->Unmap(0, nullptr);
 						}
@@ -123,7 +123,7 @@ public:
 							D3D12_RANGE ReadRange;
 							ReadRange.Begin = Batch.StartIndex * sizeof(uint64);
 							ReadRange.End = ReadRange.Begin + Batch.Size * sizeof(uint64);
-							VERIFYD3D12RESULT(ResolveBuffer->Map(0, &ReadRange, &static_cast<void*>(MappedResults)));
+							VERIFYD3D12RESULT(ResolveBuffer->Map(0, &ReadRange, reinterpret_cast<void**>(&MappedResults)));
 							FMemory::Memcpy(QueryResultsOut.GetData(), MappedResults + Batch.StartIndex, Batch.Size * sizeof(uint64));
 							ResolveBuffer->Unmap(0, nullptr);
 						}
