Fix disabling of CopyDescriptor warning for older windows versions

diff --git a/Engine/Source/Runtime/D3D12RHI/Private/D3D12Adapter.cpp b/Engine/Source/Runtime/D3D12RHI/Private/D3D12Adapter.cpp
--- a/Engine/Source/Runtime/D3D12RHI/Private/D3D12Adapter.cpp
+++ b/Engine/Source/Runtime/D3D12RHI/Private/D3D12Adapter.cpp
@@ -558,13 +558,13 @@ void FD3D12Adapter::CreateRootDevice(bool bWithDebug)
 #endif
 			};
 
-#if D3D12_RHI_RAYTRACING
-			if (GRHISupportsRayTracing && !FWindowsPlatformMisc::VerifyWindowsVersion(10, 0, 18363))
+#if PLATFORM_DESKTOP
+			if (!FWindowsPlatformMisc::VerifyWindowsVersion(10, 0, 18363))
 			{
-				// Ignore a known false positive error due to a bug in validation layer in certain Windows versions on DXR-capable hardware.
+				// Ignore a known false positive error due to a bug in validation layer in certain older Windows versions
 				DenyIds.Add(D3D12_MESSAGE_ID_COPY_DESCRIPTORS_INVALID_RANGES);
 			}
-#endif // D3D12_RHI_RAYTRACING
+#endif // PLATFORM_DESKTOP
 
 			NewFilter.DenyList.NumIDs = DenyIds.Num();
 			NewFilter.DenyList.pIDList = DenyIds.GetData();
