Change check to ensure for leaking textures and heap in SegListAllocator because it's caused by leaking high level objects and shouldn't be fatal here

diff --git a/Engine/Source/Runtime/D3D12RHI/Private/D3D12Allocation.cpp b/Engine/Source/Runtime/D3D12RHI/Private/D3D12Allocation.cpp
--- a/Engine/Source/Runtime/D3D12RHI/Private/D3D12Allocation.cpp
+++ b/Engine/Source/Runtime/D3D12RHI/Private/D3D12Allocation.cpp
@@ -1889,7 +1889,7 @@ void FD3D12SegListAllocator::VerifyEmpty()
 		}
 	}
 
-	checkf(TotalBytesRequested == 0,
+	ensureMsgf(TotalBytesRequested == 0,
 		TEXT("FD3D12SegListAllocator contains %lld allocated bytes but is expected to be empty. This likely means a memory leak. Use d3d12.SegListTrackLeaks=1 CVar to print allocations to the log."),
 		(uint64)TotalBytesRequested);
 }
diff --git a/Engine/Source/Runtime/D3D12RHI/Private/D3D12Allocation.h b/Engine/Source/Runtime/D3D12RHI/Private/D3D12Allocation.h
--- a/Engine/Source/Runtime/D3D12RHI/Private/D3D12Allocation.h
+++ b/Engine/Source/Runtime/D3D12RHI/Private/D3D12Allocation.h
@@ -633,7 +633,7 @@ private:
 		check(!!HeapSize);
 		for (const auto& Heap : FreeHeaps)
 		{
-			check(Heap->GetRefCount() == 1);
+			ensure(Heap->GetRefCount() == 1);
 		}
 	}
 
