Reenable GD3D12AsyncDeferredDeletion on nVidia and Intel for Windows version 1903+ - problem should have been fixed by MS

diff --git a/Engine/Source/Runtime/D3D12RHI/Private/Windows/WindowsD3D12Device.cpp b/Engine/Source/Runtime/D3D12RHI/Private/Windows/WindowsD3D12Device.cpp
--- a/Engine/Source/Runtime/D3D12RHI/Private/Windows/WindowsD3D12Device.cpp
+++ b/Engine/Source/Runtime/D3D12RHI/Private/Windows/WindowsD3D12Device.cpp
@@ -400,15 +400,16 @@ void FD3D12DynamicRHIModule::FindAdapter()
 		}
 	}
 
-#if WITH_EDITOR
-	if (bIsAnyIntel || bIsAnyNVIDIA)
+#if PLATFORM_DESKTOP
+	// Problem is fixed in Windows 1903+
+	if (!FWindowsPlatformMisc::VerifyWindowsVersion(10, 0, 18362))
 	{
-		UE_LOG(LogD3D12RHI, Log, TEXT("Forcing D3D12.AsyncDeferredDeletion=0 as a workaround for a deadlock."));
+		UE_LOG(LogD3D12RHI, Log, TEXT("Forcing D3D12.AsyncDeferredDeletion=0 as a workaround for a deadlock on older versions of Windows."));
 
 		extern D3D12RHI_API int32 GD3D12AsyncDeferredDeletion;
 		GD3D12AsyncDeferredDeletion = 0;
 	}
-#endif
+#endif // PLATFORM_DESKTOP
 
 	TSharedPtr<FD3D12Adapter> NewAdapter;
 	if (bFavorNonIntegrated && (bIsAnyAMD || bIsAnyNVIDIA))
