Ported split-barrier implementation from 4.26 DX12.

diff --git a/Engine/Source/Runtime/D3D12RHI/Private/D3D12CommandContext.cpp b/Engine/Source/Runtime/D3D12RHI/Private/D3D12CommandContext.cpp
--- a/Engine/Source/Runtime/D3D12RHI/Private/D3D12CommandContext.cpp
+++ b/Engine/Source/Runtime/D3D12RHI/Private/D3D12CommandContext.cpp
@@ -571,6 +571,44 @@ void FD3D12CommandContextBase::RHIEndFrame()
 	UpdateMemoryStats();
 }
 
+void FD3D12CommandContextBase::SignalTransitionFences(TArrayView<const FRHITransition*> Transitions)
+{
+	bool bSubmitted = false;
+	for (const FRHITransition* Transition : Transitions)
+	{
+		const auto* Data = Transition->GetPrivateData<FD3D12TransitionData>();
+		const auto& Fence = Data->Fence;
+		if (Fence)
+		{
+			if (!bSubmitted)
+			{
+				RHISubmitCommandsHint();
+				bSubmitted = true;
+			}
+			Fence->Signal(bIsAsyncComputeContext ? ED3D12CommandQueueType::Async : ED3D12CommandQueueType::Default);
+		}
+	}
+}
+
+void FD3D12CommandContextBase::WaitForTransitionFences(TArrayView<const FRHITransition*> Transitions)
+{
+	bool bSubmitted = false;
+	for (const FRHITransition* Transition : Transitions)
+	{
+		const auto* Data = Transition->GetPrivateData<FD3D12TransitionData>();
+		const auto& Fence = Data->Fence;
+		if (Fence)
+		{
+			if (!bSubmitted)
+			{
+				RHISubmitCommandsHint();
+				bSubmitted = true;
+			}
+			Fence->GpuWait(bIsAsyncComputeContext ? ED3D12CommandQueueType::Async : ED3D12CommandQueueType::Default, Fence->GetLastSignaledFence());
+		}
+	}
+}
+
 void FD3D12CommandContextBase::UpdateMemoryStats()
 {
 #if PLATFORM_WINDOWS && STATS
@@ -803,6 +841,62 @@ IRHICommandContextContainer* FD3D12DynamicRHI::RHIGetCommandContextContainer(int
 
 #endif // D3D12_SUPPORTS_PARALLEL_RHI_EXECUTE
 
+void FD3D12DynamicRHI::RHICreateResourceTransition(FRHITransition* Transition, EResourceTransitionPipeline Pipeline, EResourceTransitionPipelineFlags CreateFlags, TArrayView<const FRHITransitionInfo> Infos)
+{
+	ERHIPipeline SrcPipelines = ERHIPipeline::Graphics;
+	ERHIPipeline DstPipelines = ERHIPipeline::Graphics;
+	switch (Pipeline)
+	{
+	case EResourceTransitionPipeline::Graphics_To_Graphics:
+		SrcPipelines = ERHIPipeline::Graphics;
+		DstPipelines = ERHIPipeline::Graphics;
+		break;
+
+	case EResourceTransitionPipeline::Graphics_To_AsyncCompute:
+		SrcPipelines = ERHIPipeline::Graphics;
+		DstPipelines = ERHIPipeline::AsyncCompute;
+		break;
+
+	case EResourceTransitionPipeline::AsyncCompute_To_Graphics:
+		SrcPipelines = ERHIPipeline::AsyncCompute;
+		DstPipelines = ERHIPipeline::Graphics;
+		break;
+
+	case EResourceTransitionPipeline::AsyncCompute_To_AsyncCompute:
+		SrcPipelines = ERHIPipeline::AsyncCompute;
+		DstPipelines = ERHIPipeline::AsyncCompute;
+		break;
+	}
+
+	checkf(FMath::IsPowerOfTwo(uint32(SrcPipelines)) && FMath::IsPowerOfTwo(uint32(DstPipelines)), TEXT("Support for multi-pipe resources is not yet implemented."));
+
+	// Construct the data in-place on the transition instance
+	FD3D12TransitionData* Data = new (Transition->GetPrivateData<FD3D12TransitionData>()) FD3D12TransitionData;
+
+	Data->SrcPipelines = SrcPipelines;
+	Data->DstPipelines = DstPipelines;
+	Data->CreateFlags = CreateFlags;
+
+	const bool bCrossPipeline = SrcPipelines != DstPipelines;
+
+	if (bCrossPipeline && !EnumHasAnyFlags(Data->CreateFlags, EResourceTransitionPipelineFlags::NoFence))
+	{
+		const FName Name = SrcPipelines == ERHIPipeline::Graphics ? L"<Graphics To AsyncCompute>" : L"<AsyncCompute To Graphics>";
+
+		Data->Fence = new FD3D12Fence(&GetAdapter(), FRHIGPUMask::All(), Name);
+		Data->Fence->CreateFence();
+	}
+
+	Data->bCrossPipeline = bCrossPipeline;
+	Data->Infos.Append(Infos.GetData(), Infos.Num());
+}
+
+void FD3D12DynamicRHI::RHIReleaseResourceTransition(FRHITransition* Transition)
+{
+	// Destruct the transition data
+	Transition->GetPrivateData<FD3D12TransitionData>()->~FD3D12TransitionData();
+}
+
 
 //////////////////////////////////////////////////////////////////////////////////////////////////////////
 //
@@ -832,6 +926,17 @@ void FD3D12CommandContextRedirector::RHITransitionResources(EResourceTransitionA
 	}	
 }
 
+void FD3D12CommandContextRedirector::RHIBeginResourceTransitions(TArrayView<const FRHITransition*> Transitions)
+{
+	ContextRedirect(RHIBeginTransitionsWithoutFencing(Transitions));
+	SignalTransitionFences(Transitions);
+}
+
+void FD3D12CommandContextRedirector::RHIEndResourceTransitions(TArrayView<const FRHITransition*> Transitions)
+{
+	ContextRedirect(RHIEndResourceTransitions(Transitions));
+}
+
 //////////////////////////////////////////////////////////////////////////////////////////////////////////
 //
 // FD3D12TemporalEffect
diff --git a/Engine/Source/Runtime/D3D12RHI/Private/D3D12CommandContext.h b/Engine/Source/Runtime/D3D12RHI/Private/D3D12CommandContext.h
--- a/Engine/Source/Runtime/D3D12RHI/Private/D3D12CommandContext.h
+++ b/Engine/Source/Runtime/D3D12RHI/Private/D3D12CommandContext.h
@@ -57,7 +57,8 @@ protected:
 
 	virtual FD3D12CommandContext* GetContext(uint32 InGPUIndex) = 0;
 
-protected:
+	void SignalTransitionFences(TArrayView<const FRHITransition*> Transitions);
+	void WaitForTransitionFences(TArrayView<const FRHITransition*> Transitions);
 
 	FRHIGPUMask GPUMask;
 
@@ -242,6 +243,10 @@ public:
 
 	virtual void SetAsyncComputeBudgetInternal(EAsyncComputeBudget Budget) {}
 
+	void RHIBeginTransitionsWithoutFencing(TArrayView<const FRHITransition*> Transitions);
+	virtual void RHIBeginResourceTransitions(TArrayView<const FRHITransition*> Transitions) final override;
+	virtual void RHIEndResourceTransitions(TArrayView<const FRHITransition*> Transitions) final override;
+
 	// IRHIComputeContext interface
 	virtual void RHISetComputeShader(FRHIComputeShader* ComputeShader) final override;
 	virtual void RHISetComputePipelineState(FRHIComputePipelineState* ComputePipelineState) final override;
@@ -512,6 +517,10 @@ public:
 	// Special implementation that only signal the fence once.
 	virtual void RHITransitionResources(EResourceTransitionAccess TransitionType, EResourceTransitionPipeline TransitionPipeline, FRHIUnorderedAccessView** InUAVs, int32 NumUAVs, FRHIComputeFence* WriteComputeFenceRHI) final override;
 
+	// Special implementation that only signal the fence once.
+	virtual void RHIBeginResourceTransitions(TArrayView<const FRHITransition*> Transitions) final override;
+	virtual void RHIEndResourceTransitions(TArrayView<const FRHITransition*> Transitions) final override;
+
 	FORCEINLINE virtual void RHICopyToStagingBuffer(FRHIVertexBuffer* SourceBuffer, FRHIStagingBuffer* DestinationStagingBuffer, uint32 Offset, uint32 NumBytes) final override
 	{
 		ContextRedirect(RHICopyToStagingBuffer(SourceBuffer, DestinationStagingBuffer, Offset, NumBytes));
@@ -893,3 +902,14 @@ private:
 		return EffectFences.FindByPredicate([GPUIndex](const auto& Other) { return Other.GPUMask.Contains(GPUIndex); });
 	}
 };
+
+struct FD3D12TransitionData
+{
+	ERHIPipeline SrcPipelines, DstPipelines;
+	EResourceTransitionPipelineFlags CreateFlags = EResourceTransitionPipelineFlags::None;
+
+	TArray<FRHITransitionInfo, TInlineAllocator<4>> Infos;
+	TRefCountPtr<FD3D12Fence> Fence;
+
+	bool bCrossPipeline = false;
+};
\ No newline at end of file
diff --git a/Engine/Source/Runtime/D3D12RHI/Private/D3D12Commands.cpp b/Engine/Source/Runtime/D3D12RHI/Private/D3D12Commands.cpp
--- a/Engine/Source/Runtime/D3D12RHI/Private/D3D12Commands.cpp
+++ b/Engine/Source/Runtime/D3D12RHI/Private/D3D12Commands.cpp
@@ -201,6 +201,263 @@ void FD3D12CommandContext::RHIDispatchIndirectComputeShader(FRHIVertexBuffer* Ar
 }
 
 
+template <typename FunctionType>
+void EnumerateSubresources(FD3D12Resource* Resource, const FRHITransitionInfo& Info, FunctionType Function)
+{
+	uint32 FirstMipSlice = 0;
+	uint32 FirstArraySlice = 0;
+	uint32 FirstPlaneSlice = 0;
+
+	uint32 MipCount = Resource->GetMipLevels();
+	uint32 ArraySize = Resource->GetArraySize();
+	uint32 PlaneCount = Resource->GetPlaneCount();
+
+	if (!Info.IsAllMips())
+	{
+		FirstMipSlice = Info.MipIndex;
+		MipCount = 1;
+	}
+
+	if (!Info.IsAllArraySlices())
+	{
+		FirstArraySlice = Info.ArraySlice;
+		ArraySize = 1;
+	}
+
+	if (!Info.IsAllPlaneSlices())
+	{
+		FirstPlaneSlice = Info.PlaneSlice;
+		PlaneCount = 1;
+	}
+
+	for (uint32 PlaneSlice = FirstPlaneSlice; PlaneSlice < FirstPlaneSlice + PlaneCount; ++PlaneSlice)
+	{
+		for (uint32 ArraySlice = FirstArraySlice; ArraySlice < FirstArraySlice + ArraySize; ++ArraySlice)
+		{
+			for (uint32 MipSlice = FirstMipSlice; MipSlice < FirstMipSlice + MipCount; ++MipSlice)
+			{
+				const uint32 Subresource = D3D12CalcSubresource(FirstMipSlice, FirstArraySlice, FirstPlaneSlice, MipCount, ArraySize);
+				Function(Subresource);
+			}
+		}
+	}
+}
+
+template <typename FunctionType>
+void ProcessResource(FD3D12CommandContext& Context, const FRHITransitionInfo& Info, FunctionType Function)
+{
+	switch (Info.Type)
+	{
+	case FRHITransitionInfo::EType::UAV:
+	{
+		FD3D12UnorderedAccessView* UAV = Context.RetrieveObject<FD3D12UnorderedAccessView>(Info.UAV);
+		check(UAV);
+
+		FRHITransitionInfo LocalInfo = Info;
+		LocalInfo.MipIndex = UAV->GetViewSubresourceSubset().MostDetailedMip();
+		Function(LocalInfo, UAV->GetResource());
+	}
+	break;
+	case FRHITransitionInfo::EType::VertexBuffer:
+	{
+		FD3D12VertexBuffer* VertexBuffer = Context.RetrieveObject<FD3D12VertexBuffer>(Info.VertexBuffer);
+		check(VertexBuffer);
+		Function(Info, VertexBuffer->GetResource());
+	}
+	break;
+	case FRHITransitionInfo::EType::IndexBuffer:
+	{
+		FD3D12IndexBuffer* IndexBuffer = Context.RetrieveObject<FD3D12IndexBuffer>(Info.IndexBuffer);
+		check(IndexBuffer);
+		Function(Info, IndexBuffer->GetResource());
+	}
+	break;
+	case FRHITransitionInfo::EType::StructuredBuffer:
+	{
+		FD3D12StructuredBuffer* StructuredBuffer = Context.RetrieveObject<FD3D12StructuredBuffer>(Info.StructuredBuffer);
+		check(StructuredBuffer);
+		Function(Info, StructuredBuffer->GetResource());
+	}
+	break;
+	case FRHITransitionInfo::EType::Texture:
+	{
+		FD3D12TextureBase* Texture = Context.RetrieveTextureBase(Info.Texture);
+		check(Texture);
+		Function(Info, Texture->GetResource());
+	}
+	break;
+	default:
+		checkNoEntry();
+		break;
+	}
+}
+
+void FD3D12CommandContext::RHIBeginTransitionsWithoutFencing(TArrayView<const FRHITransition*> Transitions)
+{
+	static IConsoleVariable* CVarShowTransitions = IConsoleManager::Get().FindConsoleVariable(TEXT("r.ProfileGPU.ShowTransitions"));
+	const bool bShowTransitionEvents = CVarShowTransitions->GetInt() != 0;
+	SCOPED_RHI_CONDITIONAL_DRAW_EVENTF(*this, RHIBeginTransitions, bShowTransitionEvents, TEXT("RHIBeginTransitions"));
+
+	for (const FRHITransition* Transition : Transitions)
+	{
+		const FD3D12TransitionData* Data = Transition->GetPrivateData<FD3D12TransitionData>();
+
+		// Same pipe transitions or transitions onto the graphics pipe are handled in End.
+		if (Data->SrcPipelines == Data->DstPipelines || Data->DstPipelines == ERHIPipeline::Graphics)
+		{
+			continue;
+		}
+
+		for (const FRHITransitionInfo& Info : Data->Infos)
+		{
+			if (!Info.Resource)
+			{
+				continue;
+			}
+
+			ProcessResource(*this, Info, [&](const FRHITransitionInfo& Info, FD3D12Resource* Resource)
+			{
+				if (!Resource->RequiresResourceStateTracking())
+				{
+					return;
+				}
+
+				D3D12_RESOURCE_STATES State = D3D12_RESOURCE_STATE_COMMON;
+
+				if (EnumHasAnyFlags(Info.AccessAfter, EResourceTransitionAccess::EWritable))
+				{
+					State |= D3D12_RESOURCE_STATE_UNORDERED_ACCESS;
+				}
+				else if (EnumHasAnyFlags(Info.AccessAfter, EResourceTransitionAccess::EReadable))
+				{
+					State |= D3D12_RESOURCE_STATE_NON_PIXEL_SHADER_RESOURCE;
+				}
+
+				if (State == D3D12_RESOURCE_STATE_COMMON)
+				{
+					return;
+				}
+
+				if (Info.IsWholeResource() || Resource->GetSubresourceCount() == 1)
+				{
+					FD3D12DynamicRHI::TransitionResource(CommandListHandle, Resource, State, D3D12_RESOURCE_BARRIER_ALL_SUBRESOURCES);
+				}
+				else
+				{
+					EnumerateSubresources(Resource, Info, [&](uint32 Subresource)
+					{
+						FD3D12DynamicRHI::TransitionResource(CommandListHandle, Resource, State, Subresource);
+					});
+				}
+			});
+		}
+	}
+}
+
+void FD3D12CommandContext::RHIBeginResourceTransitions(TArrayView<const FRHITransition*> Transitions)
+{
+	RHIBeginTransitionsWithoutFencing(Transitions);
+	SignalTransitionFences(Transitions);
+}
+
+void FD3D12CommandContext::RHIEndResourceTransitions(TArrayView<const FRHITransition*> Transitions)
+{
+	WaitForTransitionFences(Transitions);
+
+	static IConsoleVariable* CVarShowTransitions = IConsoleManager::Get().FindConsoleVariable(TEXT("r.ProfileGPU.ShowTransitions"));
+	const bool bShowTransitionEvents = CVarShowTransitions->GetInt() != 0;
+	SCOPED_RHI_CONDITIONAL_DRAW_EVENTF(*this, RHIEndTransitions, bShowTransitionEvents, TEXT("RHIEndTransitions"));
+
+	static_assert(USE_D3D12RHI_RESOURCE_STATE_TRACKING, "RHIEndTransitions is not implemented properly for this to be disabled.");
+
+	bool bUAVBarrier = false;
+
+	for (const FRHITransition* Transition : Transitions)
+	{
+		const FD3D12TransitionData* Data = Transition->GetPrivateData<FD3D12TransitionData>();
+
+		const bool bSamePipeline = !Data->bCrossPipeline;
+
+		for (const FRHITransitionInfo& Info : Data->Infos)
+		{
+			// Sometimes we could still have barriers with resources, invalid but can still happen
+			if (bSamePipeline)
+			{
+				bUAVBarrier |= Info.AccessAfter == EResourceTransitionAccess::ERWBarrier;
+				bUAVBarrier |= EnumHasAnyFlags(Info.AccessBefore, EResourceTransitionAccess::UAVMask) && EnumHasAnyFlags(Info.AccessAfter, EResourceTransitionAccess::UAVMask);
+			}
+
+			if (!Info.Resource)
+			{
+				continue;
+			}
+
+			ProcessResource(*this, Info, [&](const FRHITransitionInfo& Info, FD3D12Resource* Resource)
+			{
+				if (!Resource->RequiresResourceStateTracking())
+				{
+					return;
+				}
+
+				D3D12_RESOURCE_STATES State = D3D12_RESOURCE_STATE_COMMON;
+				if (EnumHasAnyFlags(Info.Flags, EResourceTransitionFlags::MaintainCompression))
+				{
+					State |= SkipFastClearEliminateState;
+				}
+
+				if (Info.AccessAfter != EResourceTransitionAccess::ERWBarrier)
+				{
+					if (EnumHasAnyFlags(Info.AccessAfter, EResourceTransitionAccess::EWritable))
+					{
+						if (bIsAsyncComputeContext)
+						{
+							State |= D3D12_RESOURCE_STATE_UNORDERED_ACCESS;
+						}
+						else
+						{
+							State |= Resource->GetWritableState();
+						}
+					}
+					else if (EnumHasAnyFlags(Info.AccessAfter, EResourceTransitionAccess::EReadable))
+					{
+						if (bIsAsyncComputeContext)
+						{
+							State |= D3D12_RESOURCE_STATE_NON_PIXEL_SHADER_RESOURCE;
+						}
+						else
+						{
+							State |= Resource->GetReadableState();
+						}
+					}
+				}
+
+				if (State == D3D12_RESOURCE_STATE_COMMON)
+				{
+					return;
+				}
+
+				if (Info.IsWholeResource() || Resource->GetSubresourceCount() == 1)
+				{
+					FD3D12DynamicRHI::TransitionResource(CommandListHandle, Resource, State, D3D12_RESOURCE_BARRIER_ALL_SUBRESOURCES);
+				}
+				else
+				{
+					EnumerateSubresources(Resource, Info, [&](uint32 Subresource)
+					{
+						FD3D12DynamicRHI::TransitionResource(CommandListHandle, Resource, State, Subresource);
+					});
+				}
+			});
+		}
+	}
+
+	if (bUAVBarrier)
+	{
+		StateCache.FlushComputeShaderCache(true);
+	}
+}
+
+
 void FD3D12CommandContext::RHITransitionResources(EResourceTransitionAccess TransitionType, FRHITexture** InTextures, int32 NumTextures)
 {
 #if !USE_D3D12RHI_RESOURCE_STATE_TRACKING
diff --git a/Engine/Source/Runtime/D3D12RHI/Private/D3D12RHI.cpp b/Engine/Source/Runtime/D3D12RHI/Private/D3D12RHI.cpp
--- a/Engine/Source/Runtime/D3D12RHI/Private/D3D12RHI.cpp
+++ b/Engine/Source/Runtime/D3D12RHI/Private/D3D12RHI.cpp
@@ -85,6 +85,9 @@ FD3D12DynamicRHI::FD3D12DynamicRHI(const TArray<TSharedPtr<FD3D12Adapter>>& Chos
 	GTexturePoolSize = 0;
 	GConfig->GetInt(TEXT("TextureStreaming"), TEXT("PoolSizeVRAMPercentage"), GPoolSizeVRAMPercentage, GEngineIni);
 
+	GRHITransitionPrivateData_SizeInBytes = sizeof(FD3D12TransitionData);
+	GRHITransitionPrivateData_AlignInBytes = alignof(FD3D12TransitionData);
+
 	// Initialize the platform pixel format map.
 	GPixelFormats[PF_Unknown		].PlatformFormat = DXGI_FORMAT_UNKNOWN;
 	GPixelFormats[PF_A32B32G32R32F	].PlatformFormat = DXGI_FORMAT_R32G32B32A32_FLOAT;
diff --git a/Engine/Source/Runtime/D3D12RHI/Private/D3D12RHIPrivate.h b/Engine/Source/Runtime/D3D12RHI/Private/D3D12RHIPrivate.h
--- a/Engine/Source/Runtime/D3D12RHI/Private/D3D12RHIPrivate.h
+++ b/Engine/Source/Runtime/D3D12RHI/Private/D3D12RHIPrivate.h
@@ -262,6 +262,8 @@ public:
 	virtual FBoundShaderStateRHIRef RHICreateBoundShaderState(FRHIVertexDeclaration* VertexDeclaration, FRHIVertexShader* VertexShader, FRHIHullShader* HullShader, FRHIDomainShader* DomainShader, FRHIPixelShader* PixelShader, FRHIGeometryShader* GeometryShader) final override;
 	virtual FGraphicsPipelineStateRHIRef RHICreateGraphicsPipelineState(const FGraphicsPipelineStateInitializer& Initializer) final override;
 	virtual TRefCountPtr<FRHIComputePipelineState> RHICreateComputePipelineState(FRHIComputeShader* ComputeShader) final override;
+	virtual void RHICreateResourceTransition(FRHITransition* Transition, EResourceTransitionPipeline Pipeline, EResourceTransitionPipelineFlags CreateFlags, TArrayView<const FRHITransitionInfo> Infos) final override;
+	virtual void RHIReleaseResourceTransition(FRHITransition* Transition) final override;
 	virtual FUniformBufferRHIRef RHICreateUniformBuffer(const void* Contents, const FRHIUniformBufferLayout& Layout, EUniformBufferUsage Usage, EUniformBufferValidation Validation) final override;
 	virtual void RHIUpdateUniformBuffer(FRHIUniformBuffer* UniformBufferRHI, const void* Contents) final override;
 	virtual FIndexBufferRHIRef RHICreateIndexBuffer(uint32 Stride, uint32 Size, uint32 InUsage, FRHIResourceCreateInfo& CreateInfo) final override;
