Fix some merge shenanigans

diff --git a/Engine/Source/Runtime/D3D12RHI/Private/D3D12CommandContext.cpp b/Engine/Source/Runtime/D3D12RHI/Private/D3D12CommandContext.cpp
--- a/Engine/Source/Runtime/D3D12RHI/Private/D3D12CommandContext.cpp
+++ b/Engine/Source/Runtime/D3D12RHI/Private/D3D12CommandContext.cpp
@@ -8,7 +8,7 @@ D3D12CommandContext.cpp: RHI  Command Context implementation.
 
 #if PLATFORM_WINDOWS
 #include "Windows/AllowWindowsPlatformTypes.h"
-	#include "amd_ags.h"
+#include "amd_ags.h"
 #include "Windows/HideWindowsPlatformTypes.h"
 #endif
 
@@ -23,7 +23,7 @@ static FAutoConsoleVariableRef CVarCommandListBatchingMode(
 	GCommandListBatchingMode,
 	TEXT("Changes how command lists are batched and submitted to the GPU."),
 	ECVF_RenderThreadSafe
-	);
+);
 
 // We don't yet have a way to auto-detect that the Radeon Developer Panel is running
 // with profiling enabled, so for now, we have to manually toggle this console var.
@@ -34,7 +34,7 @@ static FAutoConsoleVariableRef CVarEmitRgpFrameMarkers(
 	GEmitRgpFrameMarkers,
 	TEXT("Enables/Disables frame markers for AMD's RGP tool."),
 	ECVF_ReadOnly | ECVF_RenderThreadSafe
-	);
+);
 
 FD3D12CommandContextBase::FD3D12CommandContextBase(class FD3D12Adapter* InParentAdapter, FRHIGPUMask InGPUMask, bool InIsDefaultContext, bool InIsAsyncComputeContext)
 	: FD3D12AdapterChild(InParentAdapter)
@@ -50,7 +50,7 @@ static D3D12_RESOURCE_STATES GetValidResourceStates(D3D12_COMMAND_LIST_TYPE Comm
 	// For reasons, we can't just list the allowed states, we have to list the disallowed states.
 	// For reference on allowed/disallowed states, see:
 	//    https://microsoft.github.io/DirectX-Specs/d3d/CPUEfficiency.html#state-support-by-command-list-type
-	
+
 	const D3D12_RESOURCE_STATES DisallowedDirectStates =
 		static_cast<D3D12_RESOURCE_STATES>(0);
 
@@ -104,7 +104,7 @@ FD3D12CommandContext::FD3D12CommandContext(FD3D12Device* InParent, bool InIsDefa
 	bDiscardSharedComputeConstants(false),
 	bUsingTessellation(false),
 #if PLATFORM_SUPPORTS_VARIABLE_RATE_SHADING
-	VRSCombiners{ D3D12_SHADING_RATE_COMBINER_PASSTHROUGH, D3D12_SHADING_RATE_COMBINER_PASSTHROUGH } ,
+	VRSCombiners{ D3D12_SHADING_RATE_COMBINER_PASSTHROUGH, D3D12_SHADING_RATE_COMBINER_PASSTHROUGH },
 	VRSShadingRate(D3D12_SHADING_RATE_1X1),
 #endif
 	SkipFastClearEliminateState(D3D12_RESOURCE_STATE_PIXEL_SHADER_RESOURCE | D3D12_RESOURCE_STATE_NON_PIXEL_SHADER_RESOURCE),
@@ -241,12 +241,12 @@ void FD3D12CommandContext::RHIPushEvent(const TCHAR* Name, FColor Color)
 void FD3D12CommandContext::RHIPopEvent()
 {
 	D3D12RHI::FD3DGPUProfiler& GPUProfiler = GetParentDevice()->GetGPUProfiler();
-	
+
 	if (IsDefaultContext())
 	{
 		GPUProfiler.PopEvent();
 	}
-		
+
 	if (GPUProfiler.bTrackingGPUCrashData)
 	{
 		WriteGPUEventStackToBreadCrumbData(false);
@@ -372,7 +372,7 @@ FD3D12CommandListHandle FD3D12CommandContext::FlushCommands(bool WaitForCompleti
 			// Submit all pending command lists and the current command list if it was still open
 			if (bIsCommandListOpen)
 			{
-				Device->PendingCommandLists.Add(CommandListHandle);				
+				Device->PendingCommandLists.Add(CommandListHandle);
 			}
 			else
 			{
@@ -688,7 +688,7 @@ void FD3D12CommandContextBase::UpdateMemoryStats()
 #endif
 
 		Device->GetDefaultBufferAllocator().UpdateMemoryStats();
-	}		
+	}
 #endif
 }
 
@@ -925,47 +925,6 @@ void FD3D12DynamicRHI::RHIReleaseTransition(FRHITransition* Transition)
 	Transition->GetPrivateData<FD3D12TransitionData>()->~FD3D12TransitionData();
 }
 
-	case EResourceTransitionPipeline::AsyncCompute_To_Graphics:
-		SrcPipelines = ERHIPipeline::AsyncCompute;
-		DstPipelines = ERHIPipeline::Graphics;
-		break;
-
-	case EResourceTransitionPipeline::AsyncCompute_To_AsyncCompute:
-		SrcPipelines = ERHIPipeline::AsyncCompute;
-		DstPipelines = ERHIPipeline::AsyncCompute;
-		break;
-	}
-
-	checkf(FMath::IsPowerOfTwo(uint32(SrcPipelines)) && FMath::IsPowerOfTwo(uint32(DstPipelines)), TEXT("Support for multi-pipe resources is not yet implemented."));
-
-	// Construct the data in-place on the transition instance
-	FD3D12TransitionData* Data = new (Transition->GetPrivateData<FD3D12TransitionData>()) FD3D12TransitionData;
-
-	Data->SrcPipelines = SrcPipelines;
-	Data->DstPipelines = DstPipelines;
-	Data->CreateFlags = CreateFlags;
-
-	const bool bCrossPipeline = SrcPipelines != DstPipelines;
-
-	if (bCrossPipeline && !EnumHasAnyFlags(Data->CreateFlags, EResourceTransitionPipelineFlags::NoFence))
-	{
-		const FName Name = SrcPipelines == ERHIPipeline::Graphics ? L"<Graphics To AsyncCompute>" : L"<AsyncCompute To Graphics>";
-
-		Data->Fence = new FD3D12Fence(&GetAdapter(), FRHIGPUMask::All(), Name);
-		Data->Fence->CreateFence();
-	}
-
-	Data->bCrossPipeline = bCrossPipeline;
-	Data->Infos.Append(Infos.GetData(), Infos.Num());
-}
-
-void FD3D12DynamicRHI::RHIReleaseResourceTransition(FRHITransition* Transition)
-{
-	// Destruct the transition data
-	Transition->GetPrivateData<FD3D12TransitionData>()->~FD3D12TransitionData();
-}
-
-
 //////////////////////////////////////////////////////////////////////////////////////////////////////////
 //
 // FD3D12CommandContextRedirector
@@ -989,17 +948,6 @@ void FD3D12CommandContextRedirector::RHIEndTransitions(TArrayView<const FRHITran
 	ContextRedirect(RHIEndTransitions(Transitions));
 }
 
-void FD3D12CommandContextRedirector::RHIBeginResourceTransitions(TArrayView<const FRHITransition*> Transitions)
-{
-	ContextRedirect(RHIBeginTransitionsWithoutFencing(Transitions));
-	SignalTransitionFences(Transitions);
-}
-
-void FD3D12CommandContextRedirector::RHIEndResourceTransitions(TArrayView<const FRHITransition*> Transitions)
-{
-	ContextRedirect(RHIEndResourceTransitions(Transitions));
-}
-
 //////////////////////////////////////////////////////////////////////////////////////////////////////////
 //
 // FD3D12TemporalEffect
@@ -1069,4 +1017,4 @@ void FD3D12TemporalEffect::SignalSyncComplete(uint32 GPUIndex, ED3D12CommandQueu
 	VERIFYD3D12RESULT(CommandQueue->Signal(CrossGPUFence->FenceCore->GetFence(), ++CrossGPUFence->LastSignaledFence));
 }
 
-#endif // WITH_MGPU
\ No newline at end of file
+#endif // WITH_MGPU
diff --git a/Engine/Source/Runtime/D3D12RHI/Private/D3D12CommandContext.h b/Engine/Source/Runtime/D3D12RHI/Private/D3D12CommandContext.h
--- a/Engine/Source/Runtime/D3D12RHI/Private/D3D12CommandContext.h
+++ b/Engine/Source/Runtime/D3D12RHI/Private/D3D12CommandContext.h
@@ -487,11 +487,7 @@ public:
 
 	// Special implementation that only signal the fence once.
 	virtual void RHIBeginTransitions(TArrayView<const FRHITransition*> Transitions) final override;
-	virtual void RHIEndTransitions(TArrayView<const FRHITransition*> Transitions) final override;
-
-	// Special implementation that only signal the fence once.
-	virtual void RHIBeginResourceTransitions(TArrayView<const FRHITransition*> Transitions) final override;
-	virtual void RHIEndResourceTransitions(TArrayView<const FRHITransition*> Transitions) final override;
+	virtual void RHIEndTransitions(TArrayView<const FRHITransition*> Transitions) final override;	
 
 	FORCEINLINE virtual void RHICopyToStagingBuffer(FRHIVertexBuffer* SourceBuffer, FRHIStagingBuffer* DestinationStagingBuffer, uint32 Offset, uint32 NumBytes) final override
 	{
