Adding the RenderTargetable flag to the DX12 backbuffer textures.

diff --git a/Engine/Source/Runtime/D3D12RHI/Private/D3D12Texture.h b/Engine/Source/Runtime/D3D12RHI/Private/D3D12Texture.h
--- a/Engine/Source/Runtime/D3D12RHI/Private/D3D12Texture.h
+++ b/Engine/Source/Runtime/D3D12RHI/Private/D3D12Texture.h
@@ -528,7 +528,7 @@ public:
 		uint32 InSizeX,
 		uint32 InSizeY,
 		EPixelFormat InFormat) :
-		FD3D12Texture2D(InDevice, InSizeX, InSizeY, 1, 1, 1, InFormat, false, TexCreate_Presentable, FClearValueBinding()),
+		FD3D12Texture2D(InDevice, InSizeX, InSizeY, 1, 1, 1, InFormat, false, TexCreate_RenderTargetable | TexCreate_Presentable, FClearValueBinding()),
 		Viewport(InViewPort), bIsSDR(bInIsSDR)
 	{
 	}
