Fixed compile error

diff --git a/Engine/Source/Runtime/D3D12RHI/Private/Windows/WindowsD3D12Device.cpp b/Engine/Source/Runtime/D3D12RHI/Private/Windows/WindowsD3D12Device.cpp
--- a/Engine/Source/Runtime/D3D12RHI/Private/Windows/WindowsD3D12Device.cpp
+++ b/Engine/Source/Runtime/D3D12RHI/Private/Windows/WindowsD3D12Device.cpp
@@ -633,6 +633,7 @@ void FD3D12DynamicRHI::Init()
 	}
 #endif
 
+#if !PLATFORM_HOLOLENS
 	// Disable ray tracing for Windows build versions
 	if (GRHISupportsRayTracing
 		&& GMinimumWindowsBuildVersionForRayTracing > 0
@@ -642,6 +643,7 @@ void FD3D12DynamicRHI::Init()
 
 		UE_LOG(LogD3D12RHI, Warning, TEXT("Ray tracing is disabled because it requires Windows 10 version %u"), (uint32)GMinimumWindowsBuildVersionForRayTracing);
 	}
+#endif
 
 #if NV_API_ENABLE
 	if (IsRHIDeviceNVIDIA() && bAllowVendorDevice)
