Also rename SRVs which are marked for delete in this frame because they could still be used - fixes a GPU fault

diff --git a/Engine/Source/Runtime/D3D12RHI/Private/D3D12Buffer.cpp b/Engine/Source/Runtime/D3D12RHI/Private/D3D12Buffer.cpp
--- a/Engine/Source/Runtime/D3D12RHI/Private/D3D12Buffer.cpp
+++ b/Engine/Source/Runtime/D3D12RHI/Private/D3D12Buffer.cpp
@@ -295,10 +295,7 @@ void FD3D12Buffer::Rename(FD3D12ResourceLocation& NewLocation)
 	for (FD3D12BaseShaderResourceView* DynamicSRVBase : DynamicSRVs)
 	{
 		FD3D12ShaderResourceView* DynamicSRV = static_cast<FD3D12ShaderResourceView*>(DynamicSRVBase);
-		if (DynamicSRV->IsValid())
-		{
-			DynamicSRV->Rename(ResourceLocation);
-		}
+		DynamicSRV->Rename(ResourceLocation);
 	}
 }
 
@@ -323,10 +320,7 @@ void FD3D12Buffer::RenameLDAChain(FD3D12ResourceLocation& NewLocation)
 			for (FD3D12BaseShaderResourceView* DynamicSRVBase : NextBuffer->DynamicSRVs)
 			{
 				FD3D12ShaderResourceView* DynamicSRV = static_cast<FD3D12ShaderResourceView*>(DynamicSRVBase);
-				if (DynamicSRV->IsValid())
-				{
-					DynamicSRV->Rename(NextBuffer->ResourceLocation);
-				}
+				DynamicSRV->Rename(NextBuffer->ResourceLocation);
 			}
 		}
 	}
