Validate uniform buffer hash during RHISetShaderUniformBuffer on D3D12

diff --git a/Engine/Source/Runtime/D3D12RHI/Private/D3D12Commands.cpp b/Engine/Source/Runtime/D3D12RHI/Private/D3D12Commands.cpp
--- a/Engine/Source/Runtime/D3D12RHI/Private/D3D12Commands.cpp
+++ b/Engine/Source/Runtime/D3D12RHI/Private/D3D12Commands.cpp
@@ -42,14 +42,34 @@ DECLARE_ISBOUNDSHADER(HullShader)
 DECLARE_ISBOUNDSHADER(DomainShader)
 DECLARE_ISBOUNDSHADER(ComputeShader)
 
+#define DECLARE_ISVALIDUNIFORMBUFFERHASH(ShaderType) inline void ValidateBoundUniformBuffer(FD3D12UniformBuffer* InUniformBuffer, FRHI##ShaderType* In##RHI, uint32 InBufferIndex) \
+{ \
+	FD3D12##ShaderType* D3D12Shader = FD3D12DynamicRHI::ResourceCast(In##RHI); \
+	if (InBufferIndex < (uint32)D3D12Shader->ShaderResourceTable.ResourceTableLayoutHashes.Num()) \
+	{ \
+		uint32 UniformBufferHash = InUniformBuffer->GetLayout().GetHash(); \
+		uint32 ShaderTableHash = D3D12Shader->ShaderResourceTable.ResourceTableLayoutHashes[InBufferIndex]; \
+		ensureMsgf(ShaderTableHash == 0 || UniformBufferHash == ShaderTableHash, TEXT("Invalid uniform buffer %s bound on shader %s at index %d."), *(InUniformBuffer->GetLayout().GetDebugName()), TEXT( #ShaderType ), InBufferIndex); \
+	} \
+}
+
+DECLARE_ISVALIDUNIFORMBUFFERHASH(VertexShader)
+DECLARE_ISVALIDUNIFORMBUFFERHASH(PixelShader)
+DECLARE_ISVALIDUNIFORMBUFFERHASH(GeometryShader)
+DECLARE_ISVALIDUNIFORMBUFFERHASH(HullShader)
+DECLARE_ISVALIDUNIFORMBUFFERHASH(DomainShader)
+DECLARE_ISVALIDUNIFORMBUFFERHASH(ComputeShader)
+
 #if EXECUTE_DEBUG_COMMAND_LISTS
 bool GIsDoingQuery = false;
 #endif
 
 #if DO_CHECK
 #define VALIDATE_BOUND_SHADER(s) ValidateBoundShader(StateCache, s)
+#define VALIDATE_BOUND_UNIFORMBUFFER_HASH(ub, s, i) ValidateBoundUniformBuffer(ub, s, i)
 #else
 #define VALIDATE_BOUND_SHADER(s)
+#define VALIDATE_BOUND_UNIFORMBUFFER_HASH(ub, s, i)
 #endif
 
 void FD3D12DynamicRHI::SetupRecursiveResources()
@@ -822,6 +842,7 @@ void FD3D12CommandContext::RHISetShaderUniformBuffer(FRHIGraphicsShader* ShaderR
 	{
 		FRHIVertexShader* VertexShaderRHI = static_cast<FRHIVertexShader*>(ShaderRHI);
 		VALIDATE_BOUND_SHADER(VertexShaderRHI);
+		VALIDATE_BOUND_UNIFORMBUFFER_HASH(Buffer, VertexShaderRHI, BufferIndex);
 		StateCache.SetConstantsFromUniformBuffer<SF_Vertex>(BufferIndex, Buffer);
 		Stage = SF_Vertex;
 	}
@@ -830,6 +851,7 @@ void FD3D12CommandContext::RHISetShaderUniformBuffer(FRHIGraphicsShader* ShaderR
 	{
 		FRHIHullShader* HullShaderRHI = static_cast<FRHIHullShader*>(ShaderRHI);
 		VALIDATE_BOUND_SHADER(HullShaderRHI);
+		VALIDATE_BOUND_UNIFORMBUFFER_HASH(Buffer, HullShaderRHI, BufferIndex);
 		StateCache.SetConstantsFromUniformBuffer<SF_Hull>(BufferIndex, Buffer);
 		Stage = SF_Hull;
 	}
@@ -838,6 +860,7 @@ void FD3D12CommandContext::RHISetShaderUniformBuffer(FRHIGraphicsShader* ShaderR
 	{
 		FRHIDomainShader* DomainShaderRHI = static_cast<FRHIDomainShader*>(ShaderRHI);
 		VALIDATE_BOUND_SHADER(DomainShaderRHI);
+		VALIDATE_BOUND_UNIFORMBUFFER_HASH(Buffer, DomainShaderRHI, BufferIndex);
 		StateCache.SetConstantsFromUniformBuffer<SF_Domain>(BufferIndex, Buffer);
 		Stage = SF_Domain;
 	}
@@ -846,6 +869,7 @@ void FD3D12CommandContext::RHISetShaderUniformBuffer(FRHIGraphicsShader* ShaderR
 	{
 		FRHIGeometryShader* GeometryShaderRHI = static_cast<FRHIGeometryShader*>(ShaderRHI);
 		VALIDATE_BOUND_SHADER(GeometryShaderRHI);
+		VALIDATE_BOUND_UNIFORMBUFFER_HASH(Buffer, GeometryShaderRHI, BufferIndex);
 		StateCache.SetConstantsFromUniformBuffer<SF_Geometry>(BufferIndex, Buffer);
 		Stage = SF_Geometry;
 	}
@@ -854,6 +878,7 @@ void FD3D12CommandContext::RHISetShaderUniformBuffer(FRHIGraphicsShader* ShaderR
 	{
 		FRHIPixelShader* PixelShaderRHI = static_cast<FRHIPixelShader*>(ShaderRHI);
 		VALIDATE_BOUND_SHADER(PixelShaderRHI);
+		VALIDATE_BOUND_UNIFORMBUFFER_HASH(Buffer, PixelShaderRHI, BufferIndex);
 		StateCache.SetConstantsFromUniformBuffer<SF_Pixel>(BufferIndex, Buffer);
 		Stage = SF_Pixel;
 	}
