Fix for missing D3D12 UAV barrier.

diff --git a/Engine/Source/Runtime/D3D12RHI/Private/D3D12Commands.cpp b/Engine/Source/Runtime/D3D12RHI/Private/D3D12Commands.cpp
--- a/Engine/Source/Runtime/D3D12RHI/Private/D3D12Commands.cpp
+++ b/Engine/Source/Runtime/D3D12RHI/Private/D3D12Commands.cpp
@@ -384,10 +384,12 @@ void FD3D12CommandContext::RHIEndTransitions(TArrayView<const FRHITransition*> T
 
 		for (const FRHITransitionInfo& Info : Data->Infos)
 		{
+			const bool bUAVAccess = EnumHasAnyFlags(Info.AccessAfter, ERHIAccess::UAVMask);
+
 			// Sometimes we could still have barriers with resources, invalid but can still happen
 			if (bSamePipeline)
 			{
-				bUAVBarrier |= EnumHasAnyFlags(Info.AccessBefore, ERHIAccess::UAVMask) && EnumHasAnyFlags(Info.AccessAfter, ERHIAccess::UAVMask);
+				bUAVBarrier |= bUAVAccess;
 			}
 
 			if (!Info.Resource)
@@ -410,7 +412,7 @@ void FD3D12CommandContext::RHIEndTransitions(TArrayView<const FRHITransition*> T
 
 				if (EnumHasAnyFlags(Info.AccessAfter, ERHIAccess::WritableMask))
 				{
-					if (bIsAsyncComputeContext)
+					if (bUAVAccess || bIsAsyncComputeContext)
 					{
 						State |= D3D12_RESOURCE_STATE_UNORDERED_ACCESS;
 					}
