Handle missing texture data on reference texture the same as other RHI's on D3D12

diff --git a/Engine/Source/Runtime/D3D12RHI/Private/D3D12Commands.cpp b/Engine/Source/Runtime/D3D12RHI/Private/D3D12Commands.cpp
--- a/Engine/Source/Runtime/D3D12RHI/Private/D3D12Commands.cpp
+++ b/Engine/Source/Runtime/D3D12RHI/Private/D3D12Commands.cpp
@@ -1337,13 +1337,11 @@ inline int32 SetShaderResourcesFromBuffer_Surface(FD3D12CommandContext& CmdConte
 
 			FD3D12TextureBase* TextureD3D12 = CmdContext.RetrieveTextureBase(TextureRHI);
 			FD3D12ShaderResourceView* D3D12Resource = TextureD3D12->GetShaderResourceView();
-			if (!ensure(D3D12Resource))
+			if (D3D12Resource == nullptr)
 			{
-				D3D12Resource = CmdContext.RetrieveTextureBase(GBlackTexture->TextureRHI)->GetShaderResourceView();
+				D3D12Resource = CmdContext.RetrieveTextureBase(GWhiteTexture->TextureRHI)->GetShaderResourceView();
 			}
 
-			check(D3D12Resource != nullptr);
-
 			SetResource<ShaderFrequency>(CmdContext, BindIndex, D3D12Resource);
 			NumSetCalls++;
 			ResourceInfo = *ResourceInfos++;
