Fix D3D12 Query heap with dangling resource on outstanding queries after heap resize and fix outstanding open query count to track if it's from the current active heap

diff --git a/Engine/Source/Runtime/D3D12RHI/Private/D3D12Query.cpp b/Engine/Source/Runtime/D3D12RHI/Private/D3D12Query.cpp
--- a/Engine/Source/Runtime/D3D12RHI/Private/D3D12Query.cpp
+++ b/Engine/Source/Runtime/D3D12RHI/Private/D3D12Query.cpp
@@ -240,8 +240,8 @@ FD3D12QueryHeap::FD3D12QueryHeap(FD3D12Device* InParent, const D3D12_QUERY_TYPE
 	, LastAllocatedElement(InQueryHeapCount - 1)
 	, QueryType(InQueryType)
 	, QueryHeapCount(InQueryHeapCount)
-	, QueryHeap(nullptr)
-	, ResultBuffer(nullptr)
+	, ActiveQueryHeap(nullptr)
+	, ActiveResultBuffer(nullptr)
 {
 	check(QueryType == D3D12_QUERY_TYPE_OCCLUSION || QueryType == D3D12_QUERY_TYPE_TIMESTAMP);
 
@@ -263,7 +263,7 @@ void FD3D12QueryHeap::Init()
 
 void FD3D12QueryHeap::Destroy()
 {
-	DestroyQueryHeap(false);
+	DestroyQueryHeap();
 }
 
 uint32 FD3D12QueryHeap::GetNextElement(uint32 InElement)
@@ -313,19 +313,19 @@ uint32 FD3D12QueryHeap::AllocQuery(FD3D12CommandContext& CmdContext)
 			check(CurrentQueryBatch.bOpen && CurrentQueryBatch.ElementCount == 0);
 		}
 
-	if (CurrentQueryBatch.StartElement > CurrentElement)
-	{
-		// We're in the middle of a batch, but we're at the end of the heap
-		// We need to split the batch in two and resolve the first piece
-		EndQueryBatchAndResolveQueryData(CmdContext);
+		if (CurrentQueryBatch.StartElement > CurrentElement)
+		{
+			// We're in the middle of a batch, but we're at the end of the heap
+			// We need to split the batch in two and resolve the first piece
+			EndQueryBatchAndResolveQueryData(CmdContext);
 		}
 
 		// check for the the batch being closed due to wrap and open a new one
 		if (!CurrentQueryBatch.bOpen)
 		{
 			StartQueryBatch(CmdContext, 256);
-		check(CurrentQueryBatch.bOpen && CurrentQueryBatch.ElementCount == 0);
-	}
+			check(CurrentQueryBatch.bOpen && CurrentQueryBatch.ElementCount == 0);
+		}
 	}
 
 	// Increment the count for the current batch
@@ -340,30 +340,28 @@ void FD3D12QueryHeap::StartQueryBatch(FD3D12CommandContext& CmdContext, uint32 N
 {
 	check(!CurrentQueryBatch.bOpen);
 
+	// Clear the current batch
+	CurrentQueryBatch.Clear();
 
-	if (!CurrentQueryBatch.bOpen)
+	if (ActiveAllocatedElementCount + NumQueriesInBatch > QueryHeapCount)
 	{
-		// Clear the current batch
-		CurrentQueryBatch.Clear();
-
-		if (ActiveAllocatedElementCount + NumQueriesInBatch > QueryHeapCount)
-		{
-			DestroyQueryHeap(true);
+		DestroyQueryHeap();
 
-			QueryHeapCount = Align(NumQueriesInBatch + QueryHeapCount, 65536 / ResultSize);
+		QueryHeapCount = Align(NumQueriesInBatch + QueryHeapCount, 65536 / ResultSize);
 
-			CreateQueryHeap();
+		CreateQueryHeap();
 
-			UE_LOG(LogD3D12RHI, Display, TEXT("QueryHeapCount is now %d elements"), QueryHeapCount);
+		UE_LOG(LogD3D12RHI, Display, TEXT("QueryHeapCount is now %d elements"), QueryHeapCount);
 
-			ActiveAllocatedElementCount = 0;
-			LastAllocatedElement = QueryHeapCount - 1;
-		}
-
-		// Start a new batch
-		CurrentQueryBatch.StartElement = GetNextElement(LastAllocatedElement);
-		CurrentQueryBatch.bOpen = true;
+		ActiveAllocatedElementCount = 0;
+		LastAllocatedElement = QueryHeapCount - 1;
 	}
+
+	// Start a new batch
+	CurrentQueryBatch.StartElement = GetNextElement(LastAllocatedElement);
+	CurrentQueryBatch.UsedQueryHeap = ActiveQueryHeap;
+	CurrentQueryBatch.UsedResultBuffer = ActiveResultBuffer;
+	CurrentQueryBatch.bOpen = true;
 }
 
 void FD3D12QueryHeap::EndQueryBatchAndResolveQueryData(FD3D12CommandContext& CmdContext)
@@ -393,38 +391,43 @@ void FD3D12QueryHeap::EndQueryBatchAndResolveQueryData(FD3D12CommandContext& Cmd
 	// Track the current active batches (application is using the data)
 	LastBatch = GetNextBatchElement(LastBatch);
 	ActiveQueryBatches[LastBatch] = CurrentQueryBatch;
-
-	// Update the head
+	
+	// Update the active element count if still part of this query heap
 	QueryBatch& OldestBatch = ActiveQueryBatches[GetNextBatchElement(LastBatch)];
-	ActiveAllocatedElementCount -= OldestBatch.ElementCount;
+	if (OldestBatch.UsedQueryHeap == ActiveQueryHeap)
+	{
+		check(ActiveAllocatedElementCount >= OldestBatch.ElementCount);
+		ActiveAllocatedElementCount -= OldestBatch.ElementCount;
+	}
 
 	CmdContext.otherWorkCounter++;
 	if (CurrentQueryBatch.StartElement + CurrentQueryBatch.ElementCount <= QueryHeapCount)
 	{
 		// Single range
-	CmdContext.CommandListHandle->ResolveQueryData(
-		QueryHeap, QueryType, CurrentQueryBatch.StartElement, CurrentQueryBatch.ElementCount,
-		ResultBuffer->GetResource(), GetResultBufferOffsetForElement(CurrentQueryBatch.StartElement));
+		CmdContext.CommandListHandle->ResolveQueryData(
+			ActiveQueryHeap, QueryType, CurrentQueryBatch.StartElement, CurrentQueryBatch.ElementCount,
+			ActiveResultBuffer->GetResource(), GetResultBufferOffsetForElement(CurrentQueryBatch.StartElement));
 	}
 	else
 	{
 		// Wrapping around heap border, need two resolves for end of heap and beginning of new range
 		CmdContext.CommandListHandle->ResolveQueryData(
-			QueryHeap, QueryType, CurrentQueryBatch.StartElement, QueryHeapCount - CurrentQueryBatch.StartElement,
-			ResultBuffer->GetResource(), GetResultBufferOffsetForElement(CurrentQueryBatch.StartElement));
+			ActiveQueryHeap, QueryType, CurrentQueryBatch.StartElement, QueryHeapCount - CurrentQueryBatch.StartElement,
+			ActiveResultBuffer->GetResource(), GetResultBufferOffsetForElement(CurrentQueryBatch.StartElement));
 		CmdContext.CommandListHandle->ResolveQueryData(
-			QueryHeap, QueryType, 0, CurrentQueryBatch.ElementCount - (QueryHeapCount - CurrentQueryBatch.StartElement),
-			ResultBuffer->GetResource(), 0);
+			ActiveQueryHeap, QueryType, 0, CurrentQueryBatch.ElementCount - (QueryHeapCount - CurrentQueryBatch.StartElement),
+			ActiveResultBuffer->GetResource(), 0);
 	}
 
-	CmdContext.CommandListHandle.UpdateResidency(&QueryHeapResidencyHandle);
-	CmdContext.CommandListHandle.UpdateResidency(ResultBuffer);
+	CmdContext.CommandListHandle.UpdateResidency(&ActiveQueryHeapResidencyHandle);
+	CmdContext.CommandListHandle.UpdateResidency(ActiveResultBuffer);
 
 	// For each render query used in this batch, update the command list
-	// so we know what sync point to wait for. The query's data isn't ready to read until the above ResolveQueryData completes on the GPU.
+	// so we know what sync point to wait for. The query's data isn't ready to read until the above ResolveQueryData completes on the GPU.	
+	check(CurrentQueryBatch.UsedResultBuffer == ActiveResultBuffer);
 	for (int32 i = 0; i < CurrentQueryBatch.RenderQueries.Num(); i++)
 	{
-		CurrentQueryBatch.RenderQueries[i]->MarkResolved(CmdContext.CommandListHandle, ResultBuffer);
+		CurrentQueryBatch.RenderQueries[i]->MarkResolved(CmdContext.CommandListHandle, ActiveResultBuffer);
 	}
 }
 
@@ -437,9 +440,9 @@ void FD3D12QueryHeap::BeginQuery(FD3D12CommandContext& CmdContext, FD3D12RenderQ
 	RenderQuery->HeapIndex = AllocQuery(CmdContext);
 
 	CmdContext.otherWorkCounter++;
-	CmdContext.CommandListHandle->BeginQuery(QueryHeap, QueryType, RenderQuery->HeapIndex);
+	CmdContext.CommandListHandle->BeginQuery(ActiveQueryHeap, QueryType, RenderQuery->HeapIndex);
 
-	CmdContext.CommandListHandle.UpdateResidency(&QueryHeapResidencyHandle);
+	CmdContext.CommandListHandle.UpdateResidency(&ActiveQueryHeapResidencyHandle);
 }
 
 void FD3D12QueryHeap::EndQuery(FD3D12CommandContext& CmdContext, FD3D12RenderQuery* RenderQuery)
@@ -448,7 +451,7 @@ void FD3D12QueryHeap::EndQuery(FD3D12CommandContext& CmdContext, FD3D12RenderQue
 
 	if (QueryType == D3D12_QUERY_TYPE_OCCLUSION)
 	{
-	check(CurrentQueryBatch.bOpen);
+		check(CurrentQueryBatch.bOpen);
 	}
 	else
 	{
@@ -466,9 +469,9 @@ void FD3D12QueryHeap::EndQuery(FD3D12CommandContext& CmdContext, FD3D12RenderQue
 	}
 
 	CmdContext.otherWorkCounter++;
-	CmdContext.CommandListHandle->EndQuery(QueryHeap, QueryType, RenderQuery->HeapIndex);
+	CmdContext.CommandListHandle->EndQuery(ActiveQueryHeap, QueryType, RenderQuery->HeapIndex);
 
-	CmdContext.CommandListHandle.UpdateResidency(&QueryHeapResidencyHandle);
+	CmdContext.CommandListHandle.UpdateResidency(&ActiveQueryHeapResidencyHandle);
 
 	// Track which render queries are used in this batch.
 	CurrentQueryBatch.RenderQueries.Push(RenderQuery);
@@ -483,12 +486,12 @@ void FD3D12QueryHeap::CreateQueryHeap()
 	QueryHeapDesc.NodeMask = GetGPUMask().GetNative();
 
 	// Create the upload heap
-	VERIFYD3D12RESULT(GetParentDevice()->GetDevice()->CreateQueryHeap(&QueryHeapDesc, IID_PPV_ARGS(&QueryHeap)));
-	SetName(QueryHeap, L"Query Heap");
+	VERIFYD3D12RESULT(GetParentDevice()->GetDevice()->CreateQueryHeap(&QueryHeapDesc, IID_PPV_ARGS(ActiveQueryHeap.GetInitReference())));
+	SetName(ActiveQueryHeap, L"Query Heap");
 
 #if ENABLE_RESIDENCY_MANAGEMENT
-	D3DX12Residency::Initialize(QueryHeapResidencyHandle, QueryHeap, ResultSize * QueryHeapDesc.Count);
-	D3DX12Residency::BeginTrackingObject(GetParentDevice()->GetResidencyManager(), QueryHeapResidencyHandle);
+	D3DX12Residency::Initialize(ActiveQueryHeapResidencyHandle, ActiveQueryHeap, ResultSize * QueryHeapDesc.Count);
+	D3DX12Residency::BeginTrackingObject(GetParentDevice()->GetResidencyManager(), ActiveQueryHeapResidencyHandle);
 #endif
 
 	FD3D12Adapter* Adapter = GetParentDevice()->GetParentAdapter();
@@ -503,44 +506,22 @@ void FD3D12QueryHeap::CreateQueryHeap()
 		ResultBufferHeapProperties,
 		D3D12_RESOURCE_STATE_COPY_DEST,
 		nullptr,
-		&ResultBuffer,
+		ActiveResultBuffer.GetInitReference(),
 		TEXT("Query Heap Result Buffer")));
 }
 
-void FD3D12QueryHeap::DestroyQueryHeap(bool bDeferDelete)
+void FD3D12QueryHeap::DestroyQueryHeap()
 {
 #if ENABLE_RESIDENCY_MANAGEMENT
-	if (D3DX12Residency::IsInitialized(QueryHeapResidencyHandle))
+	if (D3DX12Residency::IsInitialized(ActiveQueryHeapResidencyHandle))
 	{
-		D3DX12Residency::EndTrackingObject(GetParentDevice()->GetResidencyManager(), QueryHeapResidencyHandle);
-		QueryHeapResidencyHandle = {};
+		D3DX12Residency::EndTrackingObject(GetParentDevice()->GetResidencyManager(), ActiveQueryHeapResidencyHandle);
+		ActiveQueryHeapResidencyHandle = {};
 	}
 #endif
-	if (QueryHeap)
-	{
-		if (bDeferDelete)
-		{
-			FD3D12Fence& Fence = GetParentDevice()->GetCommandListManager().GetFence();
-			GetParentDevice()->GetParentAdapter()->GetDeferredDeletionQueue().EnqueueResource(QueryHeap, &Fence);
-		}
-		else
-		{
-			QueryHeap->Release();
-		}
-		QueryHeap = nullptr;
-	}
-	if (ResultBuffer)
-	{
-		if (bDeferDelete)
-		{
-			ResultBuffer->DeferDelete();
-		}
-		else
-		{
-			ResultBuffer->Release();
-		}
-		ResultBuffer = nullptr;
-	}
+
+	ActiveQueryHeap = nullptr;
+	ActiveResultBuffer = nullptr;
 }
 
 /*=============================================================================
diff --git a/Engine/Source/Runtime/D3D12RHI/Private/D3D12Query.h b/Engine/Source/Runtime/D3D12RHI/Private/D3D12Query.h
--- a/Engine/Source/Runtime/D3D12RHI/Private/D3D12Query.h
+++ b/Engine/Source/Runtime/D3D12RHI/Private/D3D12Query.h
@@ -9,7 +9,8 @@ D3D12Query.h: Implementation of D3D12 Query
 class FD3D12RenderQuery : public FRHIRenderQuery, public FD3D12DeviceChild, public FD3D12LinkedAdapterObject<FD3D12RenderQuery>
 {
 public:
-	FD3D12Resource* ResultBuffer;
+	/** Refcounted storage to result buffer because it could be holding the last valid reference if we have more than n active batches per frame */
+	TRefCountPtr<FD3D12Resource> ResultBuffer;
 
 	/** The cached query result. */
 	uint64 Result;
@@ -45,9 +46,10 @@ public:
 
 	inline void Reset()
 	{
+		ResultBuffer = nullptr;
 		HeapIndex = INDEX_NONE;
 		bResultIsCached = false;
-		bResolved = false;
+		bResolved = false;		
 		FrameSubmitted = -1;
 	}
 
@@ -84,9 +86,12 @@ private:
 	struct QueryBatch
 	{
 	public:
-		uint32 StartElement;    // The first element in the batch (inclusive)
-		uint32 ElementCount;    // The number of elements in the batch
-		bool bOpen;             // Is the batch still open for more begin/end queries?
+		uint32 StartElement;    							// The first element in the batch (inclusive)
+		uint32 ElementCount;    							// The number of elements in the batch
+		bool bOpen;             							// Is the batch still open for more begin/end queries?
+
+		TRefCountPtr<ID3D12QueryHeap> UsedQueryHeap;		// The query heap where all elements reside
+		TRefCountPtr<FD3D12Resource> UsedResultBuffer;		// The buffer where all query results are stored
 		
 		// A list of all FD3D12RenderQuery objects used in the batch. 
 		// This is used to set when each queries' result is ready to be read.
@@ -104,6 +109,9 @@ private:
 			ElementCount = 0;
 			bOpen = false;
 			RenderQueries.Reset();
+
+			UsedQueryHeap = nullptr;
+			UsedResultBuffer = nullptr;
 		}
 	};
 
@@ -128,7 +136,7 @@ private:
 	uint32 GetNextBatchElement(uint32 InBatchElement);
 
 	void CreateQueryHeap();
-	void DestroyQueryHeap(bool bDeferDelete);
+	void DestroyQueryHeap();
 
 	uint64 GetResultBufferOffsetForElement(uint32 InElement) const { return ResultSize * InElement; };
 
@@ -138,16 +146,16 @@ private:
 	TArray<QueryBatch> ActiveQueryBatches;              // List of active query batches. The data for these is in use.
 	uint32 LastBatch;                                   // The index of the newest batch.
 
-	uint32 ActiveAllocatedElementCount;         // The number of elements that are in use (Active). Between the head and the tail.
+	uint32 ActiveAllocatedElementCount;					// The number of elements that are in use (Active). Between the head and the tail.
 
-	uint32 LastAllocatedElement;                // The last element that was allocated for BeginQuery
+	uint32 LastAllocatedElement;						// The last element that was allocated for BeginQuery
 	const D3D12_QUERY_TYPE QueryType;
 	uint32 QueryHeapCount;
-	ID3D12QueryHeap* QueryHeap;                 // The query heap where all elements reside
-	FD3D12ResidencyHandle QueryHeapResidencyHandle;
-	FD3D12Resource* ResultBuffer;               // The buffer where all query results are stored
+	TRefCountPtr<ID3D12QueryHeap> ActiveQueryHeap;		// The query heap where all elements reside
+	FD3D12ResidencyHandle ActiveQueryHeapResidencyHandle;
+	TRefCountPtr<FD3D12Resource> ActiveResultBuffer;	// The buffer where all query results are stored
 
-	static const uint32 ResultSize = 8;         // The byte size of a result for a single query
+	static const uint32 ResultSize = 8;					// The byte size of a result for a single query
 };
 
 /**
