- Make sure previously bound SRV resource is removed from cmd list binding cache before rename - force rebind of SRV handle
- Take start offset into account during rename

diff --git a/Engine/Source/Runtime/D3D12RHI/Private/D3D12Buffer.cpp b/Engine/Source/Runtime/D3D12RHI/Private/D3D12Buffer.cpp
--- a/Engine/Source/Runtime/D3D12RHI/Private/D3D12Buffer.cpp
+++ b/Engine/Source/Runtime/D3D12RHI/Private/D3D12Buffer.cpp
@@ -70,6 +70,10 @@ struct FRHICommandRenameUploadBuffer final : public FRHICommand<FRHICommandRenam
 
 	void Execute(FRHICommandListBase& CmdList)
 	{
+		// Clear the resource if still bound to make sure the SRVs are rebound again on next operation
+		FD3D12CommandContext& Context = (FD3D12CommandContext&)(CmdList.IsImmediateAsyncCompute() ? CmdList.GetComputeContext() : CmdList.GetContext());
+		Context.ConditionalClearShaderResource(&Resource->ResourceLocation);
+
 		Resource->RenameLDAChain(NewLocation);
 	}
 };
@@ -376,9 +380,9 @@ void* FD3D12DynamicRHI::LockBuffer(FRHICommandListImmediate* RHICmdList, FD3D12B
 			}
 			else
 			{
-				FD3D12ResourceLocation Location(Buffer->GetParentDevice());
-				Data = Adapter.GetUploadHeapAllocator(Device->GetGPUIndex()).AllocUploadResource(BufferSize, Buffer->BufferAlignment, Location);
-				Buffer->RenameLDAChain(Location);
+				FRHICommandRenameUploadBuffer Command(Buffer, Device);
+				Data = Adapter.GetUploadHeapAllocator(Device->GetGPUIndex()).AllocUploadResource(BufferSize, Buffer->BufferAlignment, Command.NewLocation);
+				Command.Execute(*RHICmdList);
 			}
 		}
 	}
diff --git a/Engine/Source/Runtime/D3D12RHI/Private/D3D12View.cpp b/Engine/Source/Runtime/D3D12RHI/Private/D3D12View.cpp
--- a/Engine/Source/Runtime/D3D12RHI/Private/D3D12View.cpp
+++ b/Engine/Source/Runtime/D3D12RHI/Private/D3D12View.cpp
@@ -215,7 +215,7 @@ FShaderResourceViewRHIRef FD3D12DynamicRHI::RHICreateShaderResourceView(const FS
 					D3D12_SHADER_RESOURCE_VIEW_DESC SRVDesc = GetVertexBufferSRVDesc(VertexBuffer, CreationStride, Format, StartOffsetBytes, NumElements);
 
 					// Create a Shader Resource View
-					SRV->Initialize(SRVDesc, VertexBuffer->ResourceLocation, CreationStride);
+					SRV->Initialize(SRVDesc, VertexBuffer->ResourceLocation, CreationStride, StartOffsetBytes);
 					VertexBuffer->AddDynamicSRV(SRV);
 				}
 			};
@@ -310,7 +310,7 @@ FShaderResourceViewRHIRef FD3D12DynamicRHI::RHICreateShaderResourceView(const FS
 					SRVDesc.Buffer.FirstElement = Offset / Stride + StartElement;
 
 					// Create a Shader Resource View
-					SRV->Initialize(SRVDesc, StructuredBuffer->ResourceLocation, Stride);
+					SRV->Initialize(SRVDesc, StructuredBuffer->ResourceLocation, Stride, StartOffsetBytes);
 					StructuredBuffer->AddDynamicSRV(SRV);
 				}
 			};
diff --git a/Engine/Source/Runtime/D3D12RHI/Private/D3D12View.h b/Engine/Source/Runtime/D3D12RHI/Private/D3D12View.h
--- a/Engine/Source/Runtime/D3D12RHI/Private/D3D12View.h
+++ b/Engine/Source/Runtime/D3D12RHI/Private/D3D12View.h
@@ -801,6 +801,7 @@ class FD3D12ShaderResourceView : public FD3D12BaseShaderResourceView, public FRH
 	bool bSkipFastClearFinalize;
 	bool bRequiresResourceStateTracking;
 	uint32 Stride;
+	uint32 StartOffsetBytes;
 
 public:
 	// Used for dynamic buffer SRVs, which can be renamed. Must be explicitly Initialize()'d before it can be used.
@@ -810,10 +811,10 @@ public:
 
 	// Used for all other SRV resource types. Initialization is immediate on the calling thread.
 	// Should not be used for dynamic resources which can be renamed.
-	FD3D12ShaderResourceView(FD3D12Device* InParent, D3D12_SHADER_RESOURCE_VIEW_DESC& InDesc, FD3D12ResourceLocation& InResourceLocation, uint32 InStride = -1, bool InSkipFastClearFinalize = false)
+	FD3D12ShaderResourceView(FD3D12Device* InParent, D3D12_SHADER_RESOURCE_VIEW_DESC& InDesc, FD3D12ResourceLocation& InResourceLocation, uint32 InStride = -1, uint32 InStartOffsetBytes = 0, bool InSkipFastClearFinalize = false)
 		: FD3D12ShaderResourceView(InParent)
 	{
-		Initialize(InDesc, InResourceLocation, InStride, InSkipFastClearFinalize);
+		Initialize(InDesc, InResourceLocation, InStride, InStartOffsetBytes, InSkipFastClearFinalize);
 	}
 
 	~FD3D12ShaderResourceView()
@@ -823,9 +824,10 @@ public:
 		FD3D12BaseShaderResourceView::Remove();
 	}
 
-	void Initialize(D3D12_SHADER_RESOURCE_VIEW_DESC& InDesc, FD3D12ResourceLocation& InResourceLocation, uint32 InStride, bool InSkipFastClearFinalize = false)
+	void Initialize(D3D12_SHADER_RESOURCE_VIEW_DESC& InDesc, FD3D12ResourceLocation& InResourceLocation, uint32 InStride, uint32 InStartOffsetBytes, bool InSkipFastClearFinalize = false)
 	{
 		Stride = InStride;
+		StartOffsetBytes = InStartOffsetBytes;
 		bContainsDepthPlane   = InResourceLocation.GetResource()->IsDepthStencilResource() && GetPlaneSliceFromViewFormat(InResourceLocation.GetResource()->GetDesc().Format, InDesc.Format) == 0;
 		bContainsStencilPlane = InResourceLocation.GetResource()->IsDepthStencilResource() && GetPlaneSliceFromViewFormat(InResourceLocation.GetResource()->GetDesc().Format, InDesc.Format) == 1;
 		bRequiresResourceStateTracking = InResourceLocation.GetResource()->RequiresResourceStateTracking();
@@ -865,10 +867,11 @@ public:
 		// Update the first element index, then reinitialize the SRV
 		if (Desc.ViewDimension == D3D12_SRV_DIMENSION_BUFFER)
 		{
-			Desc.Buffer.FirstElement = InResourceLocation.GetOffsetFromBaseOfResource() / Stride;
+			uint32 StartElement = StartOffsetBytes / Stride;
+			Desc.Buffer.FirstElement = InResourceLocation.GetOffsetFromBaseOfResource() / Stride + StartElement;
 		}
 
-		Initialize(Desc, InResourceLocation, Stride);
+		Initialize(Desc, InResourceLocation, Stride, StartOffsetBytes);
 	}
 
 	void Rename(float ResourceMinLODClamp)
diff --git a/Engine/Source/Runtime/RenderCore/Private/UnifiedBuffer.cpp b/Engine/Source/Runtime/RenderCore/Private/UnifiedBuffer.cpp
--- a/Engine/Source/Runtime/RenderCore/Private/UnifiedBuffer.cpp
+++ b/Engine/Source/Runtime/RenderCore/Private/UnifiedBuffer.cpp
@@ -473,11 +473,6 @@ void FScatterUploadBuffer::Init( uint32 NumElements, uint32 InNumBytesPerElement
 
 	ScatterData = (uint32*)RHILockStructuredBuffer( ScatterBuffer.Buffer, 0, ScatterBytes, RLM_WriteOnly );
 	UploadData = (uint8*)RHILockStructuredBuffer( UploadBuffer.Buffer, 0, UploadBytes, RLM_WriteOnly );
-
-	// Recreate SRVs as D3D12RHI's dynamic SRV tracking is bugged and existing SRVs get invalidated if the buffer is locked
-	// TODO: fix dynamic SRV tracking in D3D12
-	ScatterBuffer.SRV = RHICreateShaderResourceView(ScatterBuffer.Buffer);
-	UploadBuffer.SRV = RHICreateShaderResourceView(UploadBuffer.Buffer);
 }
 
 template<typename ResourceType>
