Submit missing file - build fix

diff --git a/Engine/Source/Runtime/D3D12RHI/Private/D3D12RayTracing.cpp b/Engine/Source/Runtime/D3D12RHI/Private/D3D12RayTracing.cpp
--- a/Engine/Source/Runtime/D3D12RHI/Private/D3D12RayTracing.cpp
+++ b/Engine/Source/Runtime/D3D12RHI/Private/D3D12RayTracing.cpp
@@ -2846,7 +2846,7 @@ void FD3D12RayTracingScene::BuildAccelerationStructure(FD3D12CommandContext& Com
 		SRVDesc.RaytracingAccelerationStructure.Location = AccelerationStructureBuffer->ResourceLocation.GetGPUVirtualAddress();
 
 		check(AccelerationStructureView->GetParentDevice() == AccelerationStructureBuffer->GetParentDevice());
-		AccelerationStructureView->Initialize(SRVDesc, AccelerationStructureBuffer->ResourceLocation, 4);
+		AccelerationStructureView->Initialize(SRVDesc, AccelerationStructureBuffer->ResourceLocation, 4, 0);
 
 		bAccelerationStructureViewInitialized[GPUIndex] = true;
 	}
