Fix for when RHIValidation is enabled - need lowest level context then

diff --git a/Engine/Source/Runtime/D3D12RHI/Private/D3D12Buffer.cpp b/Engine/Source/Runtime/D3D12RHI/Private/D3D12Buffer.cpp
--- a/Engine/Source/Runtime/D3D12RHI/Private/D3D12Buffer.cpp
+++ b/Engine/Source/Runtime/D3D12RHI/Private/D3D12Buffer.cpp
@@ -71,7 +71,7 @@ struct FRHICommandRenameUploadBuffer final : public FRHICommand<FRHICommandRenam
 	void Execute(FRHICommandListBase& CmdList)
 	{
 		// Clear the resource if still bound to make sure the SRVs are rebound again on next operation
-		FD3D12CommandContext& Context = (FD3D12CommandContext&)(CmdList.IsImmediateAsyncCompute() ? CmdList.GetComputeContext() : CmdList.GetContext());
+		FD3D12CommandContext& Context = (FD3D12CommandContext&)(CmdList.IsImmediateAsyncCompute() ? CmdList.GetComputeContext().GetLowestLevelContext() : CmdList.GetContext().GetLowestLevelContext());
 		Context.ConditionalClearShaderResource(&Resource->ResourceLocation);
 
 		Resource->RenameLDAChain(NewLocation);
