Only flush pending command lists when kicking direct queue command lists

diff --git a/Engine/Source/Runtime/D3D12RHI/Private/D3D12CommandContext.cpp b/Engine/Source/Runtime/D3D12RHI/Private/D3D12CommandContext.cpp
--- a/Engine/Source/Runtime/D3D12RHI/Private/D3D12CommandContext.cpp
+++ b/Engine/Source/Runtime/D3D12RHI/Private/D3D12CommandContext.cpp
@@ -354,7 +354,7 @@ FD3D12CommandListHandle FD3D12CommandContext::FlushCommands(bool WaitForCompleti
 
 	FD3D12Device* Device = GetParentDevice();
 	const bool bIsCommandListOpen = !CommandListHandle.IsClosed();
-	const bool bHasPendingWork = Device->PendingCommandLists.Num() > 0;
+	const bool bHasPendingWork = Device->PendingCommandLists.Num() > 0 && (CommandListHandle.GetCommandListType() == D3D12_COMMAND_LIST_TYPE_DIRECT);
 	const bool bHasDoneWork = HasDoneWork() || bHasPendingWork;
 	const bool bOpenNewCmdList = WaitForCompletion || bHasDoneWork || bHasProfileGPUAction;
 
@@ -837,6 +837,7 @@ public:
 		{
 			FD3D12Device* Device = Adapter->GetDevice(CommandList.GetGPUIndex());
 			check(Device);
+			check(CommandList.GetCommandListType() == D3D12_COMMAND_LIST_TYPE_DIRECT);
 
 			Device->PendingCommandLists.Add(CommandList);
 		}
