Buffer refactor merging Vertex, Index and Structured Buffer types into a single FRHIBuffer. No changes to APIs in this changelist other than those that naturally folded together when types merged. A f

diff --git a/Engine/Source/Runtime/D3D12RHI/Private/D3D12VertexBuffer.cpp b/Engine/Source/Runtime/D3D12RHI/Private/D3D12VertexBuffer.cpp
--- a/Engine/Source/Runtime/D3D12RHI/Private/D3D12VertexBuffer.cpp
+++ b/Engine/Source/Runtime/D3D12RHI/Private/D3D12VertexBuffer.cpp
@@ -36,44 +36,20 @@ D3D12_RESOURCE_DESC CreateVertexBufferResourceDesc(uint32 Size, uint32 InUsage)
 	return Desc;
 }
 
-FD3D12VertexBuffer::~FD3D12VertexBuffer()
-{
-	if (ResourceLocation.GetResource() != nullptr)
-	{
-		UpdateBufferStats<FD3D12VertexBuffer>(&ResourceLocation, false);
-	}
-}
-
-void FD3D12VertexBuffer::Swap(FD3D12VertexBuffer& Other)
-{
-	check(!LockedData.bLocked && !Other.LockedData.bLocked);
-	FRHIVertexBuffer::Swap(Other);
-	FD3D12BaseShaderResource::Swap(Other);
-	FD3D12TransientResource::Swap(Other);
-	FD3D12LinkedAdapterObject<FD3D12Buffer>::Swap(Other);
-}
-
-void FD3D12VertexBuffer::ReleaseUnderlyingResource()
-{
-	UpdateBufferStats<FD3D12VertexBuffer>(&ResourceLocation, false);
-	FRHIVertexBuffer::ReleaseUnderlyingResource();
-	FD3D12Buffer::ReleaseUnderlyingResource();
-}
-
 FVertexBufferRHIRef FD3D12DynamicRHI::RHICreateVertexBuffer(uint32 Size, uint32 InUsage, ERHIAccess InResourceState, FRHIResourceCreateInfo& CreateInfo)
 {
 	if (CreateInfo.bWithoutNativeResource)
 	{
-		return GetAdapter().CreateLinkedObject<FD3D12VertexBuffer>(CreateInfo.GPUMask, [](FD3D12Device* Device)
+		return GetAdapter().CreateLinkedObject<FD3D12Buffer>(CreateInfo.GPUMask, [](FD3D12Device* Device)
 			{
-				return new FD3D12VertexBuffer();
+				return new FD3D12Buffer();
 			});
 	}
 
 	const D3D12_RESOURCE_DESC Desc = CreateVertexBufferResourceDesc(Size, InUsage);
 	const uint32 Alignment = 4;
 
-	FD3D12VertexBuffer* Buffer = GetAdapter().CreateRHIBuffer<FD3D12VertexBuffer>(nullptr, Desc, Alignment, 0, Size, InUsage, ED3D12ResourceStateMode::Default, CreateInfo);
+	FD3D12Buffer* Buffer = GetAdapter().CreateRHIBuffer(nullptr, Desc, Alignment, 0, Size, InUsage | BUF_VertexBuffer, ED3D12ResourceStateMode::Default, CreateInfo);
 	if (Buffer->ResourceLocation.IsTransient() )
 	{
 		// TODO: this should ideally be set in platform-independent code, since this tracking is for the high level
@@ -85,13 +61,13 @@ FVertexBufferRHIRef FD3D12DynamicRHI::RHICreateVertexBuffer(uint32 Size, uint32
 
 void* FD3D12DynamicRHI::RHILockVertexBuffer(FRHICommandListImmediate& RHICmdList, FRHIVertexBuffer* VertexBufferRHI, uint32 Offset, uint32 Size, EResourceLockMode LockMode)
 {
-	FD3D12VertexBuffer* Buffer = FD3D12DynamicRHI::ResourceCast(VertexBufferRHI);
+	FD3D12Buffer* Buffer = FD3D12DynamicRHI::ResourceCast(VertexBufferRHI);
 	return LockBuffer(&RHICmdList, Buffer, Buffer->GetSize(), Buffer->GetUsage(), Offset, Size, LockMode);
 }
 
 void FD3D12DynamicRHI::RHIUnlockVertexBuffer(FRHICommandListImmediate& RHICmdList, FRHIVertexBuffer* VertexBufferRHI)
 {
-	FD3D12VertexBuffer* Buffer = FD3D12DynamicRHI::ResourceCast(VertexBufferRHI);
+	FD3D12Buffer* Buffer = FD3D12DynamicRHI::ResourceCast(VertexBufferRHI);
 	UnlockBuffer(&RHICmdList, Buffer, Buffer->GetUsage());
 }
 
@@ -99,16 +75,16 @@ FVertexBufferRHIRef FD3D12DynamicRHI::CreateVertexBuffer_RenderThread(FRHIComman
 {	
 	if (CreateInfo.bWithoutNativeResource)
 	{
-		return GetAdapter().CreateLinkedObject<FD3D12VertexBuffer>(CreateInfo.GPUMask, [](FD3D12Device* Device)
+		return GetAdapter().CreateLinkedObject<FD3D12Buffer>(CreateInfo.GPUMask, [](FD3D12Device* Device)
 			{
-				return new FD3D12VertexBuffer();
+				return new FD3D12Buffer();
 			});
 	}
 
 	const D3D12_RESOURCE_DESC Desc = CreateVertexBufferResourceDesc(Size, InUsage);
 	const uint32 Alignment = 4;
 
-	FD3D12VertexBuffer* Buffer = GetAdapter().CreateRHIBuffer<FD3D12VertexBuffer>(&RHICmdList, Desc, Alignment, 0, Size, InUsage, ED3D12ResourceStateMode::Default, CreateInfo);
+	FD3D12Buffer* Buffer = GetAdapter().CreateRHIBuffer(&RHICmdList, Desc, Alignment, 0, Size, InUsage | BUF_VertexBuffer, ED3D12ResourceStateMode::Default, CreateInfo);
 	if (Buffer->ResourceLocation.IsTransient())
 	{
 		// TODO: this should ideally be set in platform-independent code, since this tracking is for the high level
@@ -120,8 +96,8 @@ FVertexBufferRHIRef FD3D12DynamicRHI::CreateVertexBuffer_RenderThread(FRHIComman
 
 void FD3D12DynamicRHI::RHICopyVertexBuffer(FRHIVertexBuffer* SourceBufferRHI, FRHIVertexBuffer* DestBufferRHI)
 {
-	FD3D12VertexBuffer* SrcBuffer = FD3D12DynamicRHI::ResourceCast(SourceBufferRHI);
-	FD3D12VertexBuffer* DstBuffer = FD3D12DynamicRHI::ResourceCast(DestBufferRHI);
+	FD3D12Buffer* SrcBuffer = FD3D12DynamicRHI::ResourceCast(SourceBufferRHI);
+	FD3D12Buffer* DstBuffer = FD3D12DynamicRHI::ResourceCast(DestBufferRHI);
 	check(SrcBuffer->GetSize() == DstBuffer->GetSize());
 
 	FD3D12Buffer* SourceBuffer = SrcBuffer;
@@ -155,21 +131,6 @@ void FD3D12DynamicRHI::RHICopyVertexBuffer(FRHIVertexBuffer* SourceBufferRHI, FR
 	}
 }
 
-void FD3D12DynamicRHI::RHITransferVertexBufferUnderlyingResource(FRHIVertexBuffer* DestVertexBuffer, FRHIVertexBuffer* SrcVertexBuffer)
-{
-	check(DestVertexBuffer);
-	FD3D12VertexBuffer* Dest = ResourceCast(DestVertexBuffer);
-	if (!SrcVertexBuffer)
-	{
-		Dest->ReleaseUnderlyingResource();
-	}
-	else
-	{
-		FD3D12VertexBuffer* Src = ResourceCast(SrcVertexBuffer);
-		Dest->Swap(*Src);
-	}
-}
-
 #if D3D12_RHI_RAYTRACING
 void FD3D12CommandContext::RHICopyBufferRegion(FRHIVertexBuffer* DestBufferRHI, uint64 DstOffset, FRHIVertexBuffer* SourceBufferRHI, uint64 SrcOffset, uint64 NumBytes)
 {
@@ -348,7 +309,7 @@ FVertexBufferRHIRef FD3D12DynamicRHI::CreateAndLockVertexBuffer_RenderThread(FRH
 	const D3D12_RESOURCE_DESC Desc = CreateVertexBufferResourceDesc(Size, InUsage);
 	const uint32 Alignment = 4;
 
-	FD3D12VertexBuffer* Buffer = GetAdapter().CreateRHIBuffer<FD3D12VertexBuffer>(nullptr, Desc, Alignment, 0, Size, InUsage, ED3D12ResourceStateMode::Default, CreateInfo);
+	FD3D12Buffer* Buffer = GetAdapter().CreateRHIBuffer(nullptr, Desc, Alignment, 0, Size, InUsage | BUF_VertexBuffer, ED3D12ResourceStateMode::Default, CreateInfo);
 	if (Buffer->ResourceLocation.IsTransient())
 	{
 		// TODO: this should ideally be set in platform-independent code, since this tracking is for the high level
diff --git a/Engine/Source/Runtime/D3D12RHI/Private/D3D12View.cpp b/Engine/Source/Runtime/D3D12RHI/Private/D3D12View.cpp
--- a/Engine/Source/Runtime/D3D12RHI/Private/D3D12View.cpp
+++ b/Engine/Source/Runtime/D3D12RHI/Private/D3D12View.cpp
@@ -6,7 +6,7 @@
 
 #include "D3D12RHIPrivate.h"
 
-static FORCEINLINE D3D12_SHADER_RESOURCE_VIEW_DESC GetVertexBufferSRVDesc(FD3D12VertexBuffer* VertexBuffer, uint32& CreationStride, uint8 Format, uint32 StartOffsetBytes, uint32 NumElements)
+static FORCEINLINE D3D12_SHADER_RESOURCE_VIEW_DESC GetVertexBufferSRVDesc(FD3D12Buffer* VertexBuffer, uint32& CreationStride, uint8 Format, uint32 StartOffsetBytes, uint32 NumElements)
 {
 	const uint32 BufferSize = VertexBuffer->GetSize();
 	const uint64 BufferOffset = VertexBuffer->ResourceLocation.GetOffsetFromBaseOfResource();
@@ -40,7 +40,7 @@ static FORCEINLINE D3D12_SHADER_RESOURCE_VIEW_DESC GetVertexBufferSRVDesc(FD3D12
 	return SRVDesc;
 }
 
-static FORCEINLINE D3D12_SHADER_RESOURCE_VIEW_DESC GetIndexBufferSRVDesc(FD3D12IndexBuffer* IndexBuffer, uint32 StartOffsetBytes, uint32 NumElements)
+static FORCEINLINE D3D12_SHADER_RESOURCE_VIEW_DESC GetIndexBufferSRVDesc(FD3D12Buffer* IndexBuffer, uint32 StartOffsetBytes, uint32 NumElements)
 {
 	const uint32 Usage = IndexBuffer->GetUsage();
 	const uint32 Width = IndexBuffer->GetSize();
@@ -186,22 +186,32 @@ uint64 FD3D12DynamicRHI::RHIGetMinimumAlignmentForBufferBackedSRV(EPixelFormat F
 
 FShaderResourceViewRHIRef FD3D12DynamicRHI::RHICreateShaderResourceView(const FShaderResourceViewInitializer& Initializer)
 {
+	FShaderResourceViewInitializer::FBufferShaderResourceViewInitializer Desc = Initializer.AsBufferSRV();
+
+	if (!Desc.Buffer)
+	{
+		return GetAdapter().CreateLinkedObject<FD3D12ShaderResourceView>(FRHIGPUMask::All(), [](FD3D12Device* Device)
+			{
+				return new FD3D12ShaderResourceView(nullptr);
+			});
+	}
+
+	FD3D12Buffer* Buffer = FD3D12DynamicRHI::ResourceCast(Desc.Buffer);
+
 	switch (Initializer.GetType())
 	{
 		case FShaderResourceViewInitializer::EType::VertexBufferSRV:
 		{
-			FShaderResourceViewInitializer::FVertexBufferShaderResourceViewInitializer Desc = Initializer.AsVertexBufferSRV();
-
 			struct FD3D12InitializeVertexBufferSRVRHICommand final : public FRHICommand<FD3D12InitializeVertexBufferSRVRHICommand>
 			{
-				FD3D12VertexBuffer* VertexBuffer;
+				FD3D12Buffer* VertexBuffer;
 				FD3D12ShaderResourceView* SRV;
 
 				uint32 StartOffsetBytes;
 				uint32 NumElements;
 				uint8 Format;
 
-				FD3D12InitializeVertexBufferSRVRHICommand(FD3D12VertexBuffer* InVertexBuffer, FD3D12ShaderResourceView* InSRV, uint32 InStartOffsetBytes, uint32 InNumElements, uint8 InFormat)
+				FD3D12InitializeVertexBufferSRVRHICommand(FD3D12Buffer* InVertexBuffer, FD3D12ShaderResourceView* InSRV, uint32 InStartOffsetBytes, uint32 InNumElements, uint8 InFormat)
 					: VertexBuffer(InVertexBuffer)
 					, SRV(InSRV)
 					, StartOffsetBytes(InStartOffsetBytes)
@@ -220,35 +230,26 @@ FShaderResourceViewRHIRef FD3D12DynamicRHI::RHICreateShaderResourceView(const FS
 				}
 			};
 
-			if (!Desc.VertexBuffer)
-			{
-				return GetAdapter().CreateLinkedObject<FD3D12ShaderResourceView>(FRHIGPUMask::All(), [](FD3D12Device* Device)
-					{
-						return new FD3D12ShaderResourceView(nullptr);
-					});
-			}
-
-			FD3D12VertexBuffer* VertexBuffer = FD3D12DynamicRHI::ResourceCast(Desc.VertexBuffer);
-			return GetAdapter().CreateLinkedViews<FD3D12VertexBuffer, FD3D12ShaderResourceView>(VertexBuffer,
-				[Desc](FD3D12VertexBuffer* VertexBuffer)
+			return GetAdapter().CreateLinkedViews<FD3D12Buffer, FD3D12ShaderResourceView>(Buffer,
+				[Desc](FD3D12Buffer* Buffer)
 				{
-					check(VertexBuffer);
+					check(Buffer);
 
-					FD3D12ShaderResourceView* ShaderResourceView = new FD3D12ShaderResourceView(VertexBuffer->GetParentDevice());
+					FD3D12ShaderResourceView* ShaderResourceView = new FD3D12ShaderResourceView(Buffer->GetParentDevice());
 					
 					uint32 Stride = GPixelFormats[Desc.Format].BlockBytes;
 					FRHICommandListImmediate& RHICmdList = FRHICommandListExecutor::GetImmediateCommandList();
-					if (ShouldDeferBufferLockOperation(&RHICmdList) && (VertexBuffer->GetUsage() & BUF_AnyDynamic))
+					if (ShouldDeferBufferLockOperation(&RHICmdList) && (Buffer->GetUsage() & BUF_AnyDynamic))
 					{
 						// We have to defer the SRV initialization to the RHI thread if the buffer is dynamic (and RHI threading is enabled), as dynamic buffers can be renamed.
 						// Also insert an RHI thread fence to prevent parallel translate tasks running until this command has completed.
-						ALLOC_COMMAND_CL(RHICmdList, FD3D12InitializeVertexBufferSRVRHICommand)(VertexBuffer, ShaderResourceView, Desc.StartOffsetBytes, Desc.NumElements, Desc.Format);
+						ALLOC_COMMAND_CL(RHICmdList, FD3D12InitializeVertexBufferSRVRHICommand)(Buffer, ShaderResourceView, Desc.StartOffsetBytes, Desc.NumElements, Desc.Format);
 						RHICmdList.RHIThreadFence(true);
 					}
 					else
 					{
 						// Run the command directly if we're bypassing RHI command list recording, or the buffer is not dynamic.
-						FD3D12InitializeVertexBufferSRVRHICommand Command(VertexBuffer, ShaderResourceView, Desc.StartOffsetBytes, Desc.NumElements, Desc.Format);
+						FD3D12InitializeVertexBufferSRVRHICommand Command(Buffer, ShaderResourceView, Desc.StartOffsetBytes, Desc.NumElements, Desc.Format);
 						Command.Execute(RHICmdList);
 					}
 
@@ -258,16 +259,14 @@ FShaderResourceViewRHIRef FD3D12DynamicRHI::RHICreateShaderResourceView(const FS
 
 		case FShaderResourceViewInitializer::EType::StructuredBufferSRV:
 		{
-			FShaderResourceViewInitializer::FStructuredBufferShaderResourceViewInitializer Desc = Initializer.AsStructuredBufferSRV();
-
 			struct FD3D12InitializeStructuredBufferSRVRHICommand final : public FRHICommand<FD3D12InitializeStructuredBufferSRVRHICommand>
 			{
-				FD3D12StructuredBuffer* StructuredBuffer;
+				FD3D12Buffer* StructuredBuffer;
 				FD3D12ShaderResourceView* SRV;
 				uint32 StartOffsetBytes;
 				uint32 NumElements;
 
-				FD3D12InitializeStructuredBufferSRVRHICommand(FD3D12StructuredBuffer* InStructuredBuffer, FD3D12ShaderResourceView* InSRV, uint32 InStartOffsetBytes, uint32 InNumElements)
+				FD3D12InitializeStructuredBufferSRVRHICommand(FD3D12Buffer* InStructuredBuffer, FD3D12ShaderResourceView* InSRV, uint32 InStartOffsetBytes, uint32 InNumElements)
 					: StructuredBuffer(InStructuredBuffer)
 					, SRV(InSRV)
 					, StartOffsetBytes(InStartOffsetBytes)
@@ -315,27 +314,25 @@ FShaderResourceViewRHIRef FD3D12DynamicRHI::RHICreateShaderResourceView(const FS
 				}
 			};
 
-			FD3D12StructuredBuffer*  StructuredBuffer = FD3D12DynamicRHI::ResourceCast(Desc.StructuredBuffer);
-
-			return GetAdapter().CreateLinkedViews<FD3D12StructuredBuffer,
-				FD3D12ShaderResourceView>(StructuredBuffer, [Desc](FD3D12StructuredBuffer* StructuredBuffer)
+			return GetAdapter().CreateLinkedViews<FD3D12Buffer,
+				FD3D12ShaderResourceView>(Buffer, [Desc](FD3D12Buffer* Buffer)
 					{
-						check(StructuredBuffer);
+						check(Buffer);
 
-						FD3D12ShaderResourceView* ShaderResourceView = new FD3D12ShaderResourceView(StructuredBuffer->GetParentDevice());
+						FD3D12ShaderResourceView* ShaderResourceView = new FD3D12ShaderResourceView(Buffer->GetParentDevice());
 
 						FRHICommandListImmediate& RHICmdList = FRHICommandListExecutor::GetImmediateCommandList();
-						if (ShouldDeferBufferLockOperation(&RHICmdList) && (StructuredBuffer->GetUsage() & BUF_AnyDynamic))
+						if (ShouldDeferBufferLockOperation(&RHICmdList) && (Buffer->GetUsage() & BUF_AnyDynamic))
 						{
 							// We have to defer the SRV initialization to the RHI thread if the buffer is dynamic (and RHI threading is enabled), as dynamic buffers can be renamed.
 							// Also insert an RHI thread fence to prevent parallel translate tasks running until this command has completed.
-							ALLOC_COMMAND_CL(RHICmdList, FD3D12InitializeStructuredBufferSRVRHICommand)(StructuredBuffer, ShaderResourceView, Desc.StartOffsetBytes, Desc.NumElements);
+							ALLOC_COMMAND_CL(RHICmdList, FD3D12InitializeStructuredBufferSRVRHICommand)(Buffer, ShaderResourceView, Desc.StartOffsetBytes, Desc.NumElements);
 							RHICmdList.RHIThreadFence(true);
 						}
 						else
 						{
 							// Run the command directly if we're bypassing RHI command list recording, or the buffer is not dynamic.
-							FD3D12InitializeStructuredBufferSRVRHICommand Command(StructuredBuffer, ShaderResourceView, Desc.StartOffsetBytes, Desc.NumElements);
+							FD3D12InitializeStructuredBufferSRVRHICommand Command(Buffer, ShaderResourceView, Desc.StartOffsetBytes, Desc.NumElements);
 							Command.Execute(RHICmdList);
 						}
 
@@ -345,27 +342,16 @@ FShaderResourceViewRHIRef FD3D12DynamicRHI::RHICreateShaderResourceView(const FS
 
 		case FShaderResourceViewInitializer::EType::IndexBufferSRV:
 		{
-			FShaderResourceViewInitializer::FIndexBufferShaderResourceViewInitializer Desc = Initializer.AsIndexBufferSRV();
-
-			if (!Desc.IndexBuffer)
+			return GetAdapter().CreateLinkedViews<FD3D12Buffer, FD3D12ShaderResourceView>(Buffer,
+			[Desc](FD3D12Buffer* Buffer)
 			{
-				return GetAdapter().CreateLinkedObject<FD3D12ShaderResourceView>(FRHIGPUMask::All(), [](FD3D12Device* Device)
-					{
-						return new FD3D12ShaderResourceView(nullptr);
-					});
-			}
+				check(Buffer);
 
-			FD3D12IndexBuffer* IndexBuffer = FD3D12DynamicRHI::ResourceCast(Desc.IndexBuffer);
-			return GetAdapter().CreateLinkedViews<FD3D12IndexBuffer, FD3D12ShaderResourceView>(IndexBuffer,
-			[Desc](FD3D12IndexBuffer* IndexBuffer)
-			{
-				check(IndexBuffer);
-
-				FD3D12ResourceLocation& Location = IndexBuffer->ResourceLocation;
-				const uint32 CreationStride = IndexBuffer->GetStride();
-				D3D12_SHADER_RESOURCE_VIEW_DESC SRVDesc = GetIndexBufferSRVDesc(IndexBuffer, Desc.StartOffsetBytes, Desc.NumElements);
+				FD3D12ResourceLocation& Location = Buffer->ResourceLocation;
+				const uint32 CreationStride = Buffer->GetStride();
+				D3D12_SHADER_RESOURCE_VIEW_DESC SRVDesc = GetIndexBufferSRVDesc(Buffer, Desc.StartOffsetBytes, Desc.NumElements);
 
-				FD3D12ShaderResourceView* ShaderResourceView = new FD3D12ShaderResourceView(IndexBuffer->GetParentDevice(), SRVDesc, Location, CreationStride);
+				FD3D12ShaderResourceView* ShaderResourceView = new FD3D12ShaderResourceView(Buffer->GetParentDevice(), SRVDesc, Location, CreationStride);
 				return ShaderResourceView;
 			});
 		}
@@ -375,23 +361,17 @@ FShaderResourceViewRHIRef FD3D12DynamicRHI::RHICreateShaderResourceView(const FS
 	return nullptr;
 }
 
-FShaderResourceViewRHIRef FD3D12DynamicRHI::RHICreateShaderResourceView(FRHIIndexBuffer* BufferRHI)
-{
-	return FD3D12DynamicRHI::RHICreateShaderResourceView(FShaderResourceViewInitializer(BufferRHI));
-}
-
-void FD3D12DynamicRHI::RHIUpdateShaderResourceView(FRHIShaderResourceView* SRV, FRHIVertexBuffer* VertexBuffer, uint32 Stride, uint8 Format)
+void FD3D12DynamicRHI::RHIUpdateShaderResourceView(FRHIShaderResourceView* SRV, FRHIBuffer* BufferRHI, uint32 Stride, uint8 Format)
 {
 	check(SRV);
-	if (VertexBuffer)
+	if (BufferRHI)
 	{
-		FD3D12VertexBuffer* VBD3D12 = ResourceCast(VertexBuffer);
+		FD3D12Buffer* Buffer = ResourceCast(BufferRHI);
 		FD3D12ShaderResourceView* SRVD3D12 = ResourceCast(SRV);
-		D3D12_SHADER_RESOURCE_VIEW_DESC SRVDesc = GetVertexBufferSRVDesc(VBD3D12, Stride, Format, 0, UINT32_MAX);
+		D3D12_SHADER_RESOURCE_VIEW_DESC SRVDesc = GetVertexBufferSRVDesc(Buffer, Stride, Format, 0, UINT32_MAX);
 
 		// Rename the SRV to view on the new vertex buffer
-		FD3D12Buffer* Buffer = VBD3D12;
-		for (auto It = MakeDualLinkedObjectIterator(VBD3D12, SRVD3D12); It; ++It)
+		for (auto It = MakeDualLinkedObjectIterator(Buffer, SRVD3D12); It; ++It)
 		{
 			Buffer = It.GetFirst();
 			SRVD3D12 = It.GetSecond();
@@ -403,19 +383,18 @@ void FD3D12DynamicRHI::RHIUpdateShaderResourceView(FRHIShaderResourceView* SRV,
 	}
 }
 
-void FD3D12DynamicRHI::RHIUpdateShaderResourceView(FRHIShaderResourceView* SRV, FRHIIndexBuffer* IndexBuffer)
+void FD3D12DynamicRHI::RHIUpdateShaderResourceView(FRHIShaderResourceView* SRV, FRHIBuffer* BufferRHI)
 {
 	check(SRV);
-	if (IndexBuffer)
+	if (BufferRHI)
 	{
-		FD3D12IndexBuffer* IBD3D12 = ResourceCast(IndexBuffer);
+		FD3D12Buffer* Buffer = ResourceCast(BufferRHI);
 		FD3D12ShaderResourceView* SRVD3D12 = ResourceCast(SRV);
-		D3D12_SHADER_RESOURCE_VIEW_DESC SRVDesc = GetIndexBufferSRVDesc(IBD3D12, 0, UINT32_MAX);
-		const uint32 Stride = IBD3D12->GetStride();
+		D3D12_SHADER_RESOURCE_VIEW_DESC SRVDesc = GetIndexBufferSRVDesc(Buffer, 0, UINT32_MAX);
+		const uint32 Stride = Buffer->GetStride();
 
 		// Rename the SRV to view on the new index buffer
-		FD3D12Buffer* Buffer = IBD3D12;
-		for (auto It = MakeDualLinkedObjectIterator(IBD3D12, SRVD3D12); It; ++It)
+		for (auto It = MakeDualLinkedObjectIterator(Buffer, SRVD3D12); It; ++It)
 		{
 			Buffer = It.GetFirst();
 			SRVD3D12 = It.GetSecond();
@@ -431,9 +410,9 @@ FShaderResourceViewRHIRef FD3D12DynamicRHI::RHICreateShaderResourceView_RenderTh
 	return RHICreateShaderResourceView(Texture, CreateInfo);
 }
 
-FShaderResourceViewRHIRef FD3D12DynamicRHI::RHICreateShaderResourceView_RenderThread(FRHICommandListImmediate& RHICmdList, FRHIVertexBuffer* VertexBufferRHI, uint32 Stride, uint8 Format)
+FShaderResourceViewRHIRef FD3D12DynamicRHI::RHICreateShaderResourceView_RenderThread(FRHICommandListImmediate& RHICmdList, FRHIBuffer* BufferRHI, uint32 Stride, uint8 Format)
 {
-	return RHICreateShaderResourceView(VertexBufferRHI, Stride, Format);
+	return RHICreateShaderResourceView(BufferRHI, Stride, Format);
 }
 
 FShaderResourceViewRHIRef FD3D12DynamicRHI::RHICreateShaderResourceView_RenderThread(FRHICommandListImmediate& RHICmdList, const FShaderResourceViewInitializer& Initializer)
@@ -441,9 +420,9 @@ FShaderResourceViewRHIRef FD3D12DynamicRHI::RHICreateShaderResourceView_RenderTh
 	return RHICreateShaderResourceView(Initializer);
 }
 
-FShaderResourceViewRHIRef FD3D12DynamicRHI::CreateShaderResourceView_RenderThread(FRHICommandListImmediate& RHICmdList, FRHIVertexBuffer* VertexBufferRHI, uint32 Stride, uint8 Format)
+FShaderResourceViewRHIRef FD3D12DynamicRHI::CreateShaderResourceView_RenderThread(FRHICommandListImmediate& RHICmdList, FRHIBuffer* BufferRHI, uint32 Stride, uint8 Format)
 {
-	return RHICreateShaderResourceView_RenderThread(RHICmdList, VertexBufferRHI, Stride, Format);
+	return RHICreateShaderResourceView_RenderThread(RHICmdList, BufferRHI, Stride, Format);
 }
 
 FShaderResourceViewRHIRef FD3D12DynamicRHI::CreateShaderResourceView_RenderThread(FRHICommandListImmediate& RHICmdList, FRHIIndexBuffer* Buffer)
@@ -456,9 +435,9 @@ FShaderResourceViewRHIRef FD3D12DynamicRHI::CreateShaderResourceView_RenderThrea
 	return RHICreateShaderResourceView_RenderThread(RHICmdList, Initializer);
 }
 
-FShaderResourceViewRHIRef FD3D12DynamicRHI::RHICreateShaderResourceView_RenderThread(FRHICommandListImmediate& RHICmdList, FRHIStructuredBuffer* StructuredBufferRHI)
+FShaderResourceViewRHIRef FD3D12DynamicRHI::RHICreateShaderResourceView_RenderThread(FRHICommandListImmediate& RHICmdList, FRHIBuffer* BufferRHI)
 {
-	return RHICreateShaderResourceView(StructuredBufferRHI);
+	return RHICreateShaderResourceView(BufferRHI);
 }
 
 FShaderResourceViewRHIRef FD3D12DynamicRHI::RHICreateShaderResourceViewWriteMask_RenderThread(class FRHICommandListImmediate& RHICmdList, FRHITexture2D* Texture2D)
diff --git a/Engine/Source/Runtime/D3D12RHI/Public/D3D12Resources.h b/Engine/Source/Runtime/D3D12RHI/Public/D3D12Resources.h
--- a/Engine/Source/Runtime/D3D12RHI/Public/D3D12Resources.h
+++ b/Engine/Source/Runtime/D3D12RHI/Public/D3D12Resources.h
@@ -812,12 +812,23 @@ inline void FD3D12BaseShaderResourceView::Remove()
 	}
 }
 
-/** Updates tracked stats for a buffer. */
-#define D3D12_BUFFER_TYPE_CONSTANT   1
-#define D3D12_BUFFER_TYPE_INDEX      2
-#define D3D12_BUFFER_TYPE_VERTEX     3
-#define D3D12_BUFFER_TYPE_STRUCTURED 4
-extern void UpdateBufferStats(FD3D12ResourceLocation* ResourceLocation, bool bAllocating, uint32 BufferType);
+extern void UpdateBufferStats(FName Name, int64 RequestedSize);
+
+inline FName GetBufferStats(uint32 Usage)
+{
+	if (Usage & BUF_VertexBuffer)
+	{
+		return GET_STATFNAME(STAT_VertexBufferMemory);
+	}
+	else if (Usage & BUF_IndexBuffer)
+	{
+		return GET_STATFNAME(STAT_IndexBufferMemory);
+	}
+	else
+	{
+		return GET_STATFNAME(STAT_StructuredBufferMemory);
+	}
+}
 
 /** Uniform buffer resource class. */
 class FD3D12UniformBuffer : public FRHIUniformBuffer, public FD3D12DeviceChild, public FD3D12LinkedAdapterObject<FD3D12UniformBuffer>
@@ -868,44 +879,32 @@ public:
 };
 #endif
 
-class FD3D12Buffer : public FD3D12BaseShaderResource, public FD3D12TransientResource, public FD3D12LinkedAdapterObject<FD3D12Buffer>
+class FD3D12Buffer : public FRHIBuffer, public FD3D12BaseShaderResource, public FD3D12TransientResource, public FD3D12LinkedAdapterObject<FD3D12Buffer>
 {
 public:
-	FD3D12Buffer(FD3D12Device* InParent)
-		: FD3D12BaseShaderResource(InParent)
+	FD3D12Buffer()
+		: FRHIBuffer(0, 0, 0)
+		, FD3D12BaseShaderResource(nullptr)
+		, LockedData(nullptr)
+	{
+	}
+
+	FD3D12Buffer(FD3D12Device* InParent, uint32 InSize, uint32 InUsage, uint32 InStride)
+		: FRHIBuffer(InSize, InUsage, InStride)
+		, FD3D12BaseShaderResource(InParent)
 		, LockedData(InParent)
 	{
 	}
 	virtual ~FD3D12Buffer()
 	{
+		int64 BufferSize = ResourceLocation.GetSize();
+		UpdateBufferStats(GetBufferStats(GetUsage()), -BufferSize);
 	}
 
 	void Rename(FD3D12ResourceLocation& NewLocation);
 	void RenameLDAChain(FD3D12ResourceLocation& NewLocation);
 
-	void ReleaseUnderlyingResource();
-
-	FD3D12LockedResource LockedData;
-};
-
-/** Index buffer resource class that stores stride information. */
-class FD3D12IndexBuffer : public FRHIIndexBuffer, public FD3D12Buffer
-{
-public:
-	FD3D12IndexBuffer()
-		: FD3D12Buffer(nullptr)
-	{
-	}
-
-	FD3D12IndexBuffer(FD3D12Device* InParent, uint32 InStride, uint32 InSize, uint32 InUsage)
-		: FRHIIndexBuffer(InStride, InSize, InUsage)
-		, FD3D12Buffer(InParent)
-	{
-	}
-
-	virtual ~FD3D12IndexBuffer();
-
-	void Swap(FD3D12IndexBuffer& Other);
+	void Swap(FD3D12Buffer& Other);
 
 	void ReleaseUnderlyingResource();
 
@@ -922,99 +921,12 @@ public:
 	{
 		return FRHIResource::GetRefCount();
 	}
-};
-
-class FD3D12ShaderResourceView;
-
-/** Structured buffer resource class. */
-class FD3D12StructuredBuffer : public FRHIStructuredBuffer, public FD3D12Buffer
-{
-public:
-	FD3D12StructuredBuffer(FD3D12Device* InParent, uint32 InStride, uint32 InSize, uint32 InUsage)
-		: FRHIStructuredBuffer(InStride, InSize, InUsage)
-		, FD3D12Buffer(InParent)
-	{
-	}
 
-	virtual ~FD3D12StructuredBuffer();
-
-	// IRefCountedObject interface.
-	virtual uint32 AddRef() const
-	{
-		return FRHIResource::AddRef();
-	}
-	virtual uint32 Release() const
-	{
-		return FRHIResource::Release();
-	}
-	virtual uint32 GetRefCount() const
-	{
-		return FRHIResource::GetRefCount();
-	}
-};
 
-/** Vertex buffer resource class. */
-class FD3D12VertexBuffer : public FRHIVertexBuffer, public FD3D12Buffer
-{
-public:
-	FD3D12VertexBuffer()
-		: FD3D12Buffer(nullptr)
-	{}
-
-	FD3D12VertexBuffer(FD3D12Device* InParent, uint32 InStride, uint32 InSize, uint32 InUsage)
-		: FRHIVertexBuffer(InSize, InUsage)
-		, FD3D12Buffer(InParent)
-	{
-		UNREFERENCED_PARAMETER(InStride);
-	}
-
-	virtual ~FD3D12VertexBuffer();
-
-	void Swap(FD3D12VertexBuffer& Other);
-
-	void ReleaseUnderlyingResource();
-
-	// IRefCountedObject interface.
-	virtual uint32 AddRef() const
-	{
-		return FRHIResource::AddRef();
-	}
-	virtual uint32 Release() const
-	{
-		return FRHIResource::Release();
-	}
-	virtual uint32 GetRefCount() const
-	{
-		return FRHIResource::GetRefCount();
-	}
+	FD3D12LockedResource LockedData;
 };
 
-template<class BufferType>
-inline void UpdateBufferStats(FD3D12ResourceLocation* ResourceLocation, bool bAllocating);
-
-template<>
-inline void UpdateBufferStats<FD3D12UniformBuffer>(FD3D12ResourceLocation* ResourceLocation, bool bAllocating)
-{
-	UpdateBufferStats(ResourceLocation, bAllocating, D3D12_BUFFER_TYPE_CONSTANT);
-}
-
-template<>
-inline void UpdateBufferStats<FD3D12VertexBuffer>(FD3D12ResourceLocation* ResourceLocation, bool bAllocating)
-{
-	UpdateBufferStats(ResourceLocation, bAllocating, D3D12_BUFFER_TYPE_VERTEX);
-}
-
-template<>
-inline void UpdateBufferStats<FD3D12IndexBuffer>(FD3D12ResourceLocation* ResourceLocation, bool bAllocating)
-{
-	UpdateBufferStats(ResourceLocation, bAllocating, D3D12_BUFFER_TYPE_INDEX);
-}
-
-template<>
-inline void UpdateBufferStats<FD3D12StructuredBuffer>(FD3D12ResourceLocation* ResourceLocation, bool bAllocating)
-{
-	UpdateBufferStats(ResourceLocation, bAllocating, D3D12_BUFFER_TYPE_STRUCTURED);
-}
+class FD3D12ShaderResourceView;
 
 class FD3D12ResourceBarrierBatcher : public FNoncopyable
 {
@@ -1180,19 +1092,9 @@ struct TD3D12ResourceTraits<FRHIUniformBuffer>
 	typedef FD3D12UniformBuffer TConcreteType;
 };
 template<>
-struct TD3D12ResourceTraits<FRHIIndexBuffer>
-{
-	typedef FD3D12IndexBuffer TConcreteType;
-};
-template<>
-struct TD3D12ResourceTraits<FRHIStructuredBuffer>
-{
-	typedef FD3D12StructuredBuffer TConcreteType;
-};
-template<>
-struct TD3D12ResourceTraits<FRHIVertexBuffer>
+struct TD3D12ResourceTraits<FRHIBuffer>
 {
-	typedef FD3D12VertexBuffer TConcreteType;
+	typedef FD3D12Buffer TConcreteType;
 };
 template<>
 struct TD3D12ResourceTraits<FRHISamplerState>
diff --git a/Engine/Source/Runtime/EmptyRHI/Private/EmptyCommands.cpp b/Engine/Source/Runtime/EmptyRHI/Private/EmptyCommands.cpp
--- a/Engine/Source/Runtime/EmptyRHI/Private/EmptyCommands.cpp
+++ b/Engine/Source/Runtime/EmptyRHI/Private/EmptyCommands.cpp
@@ -29,7 +29,7 @@ void FEmptyDynamicRHI::RHIDispatchComputeShader(uint32 ThreadGroupCountX, uint32
 
 }
 
-void FEmptyDynamicRHI::RHIDispatchIndirectComputeShader(FRHIVertexBuffer* ArgumentBufferRHI, uint32 ArgumentOffset)
+void FEmptyDynamicRHI::RHIDispatchIndirectComputeShader(FRHIBuffer* ArgumentBufferRHI, uint32 ArgumentOffset)
 { 
 	FEmptyVertexBuffer* ArgumentBuffer = ResourceCast(ArgumentBufferRHI);
 
@@ -225,93 +225,34 @@ void FEmptyDynamicRHI::RHIDrawPrimitive(uint32 PrimitiveType, uint32 BaseVertexI
 {
 }
 
-void FEmptyDynamicRHI::RHIDrawPrimitiveIndirect(uint32 PrimitiveType, FRHIVertexBuffer* ArgumentBufferRHI, uint32 ArgumentOffset)
+void FEmptyDynamicRHI::RHIDrawPrimitiveIndirect(uint32 PrimitiveType, FRHIBuffer* ArgumentBufferRHI, uint32 ArgumentOffset)
 {
-	FEmptyVertexBuffer* ArgumentBuffer = ResourceCast(ArgumentBufferRHI);
+	FEmptyBuffer* ArgumentBuffer = ResourceCast(ArgumentBufferRHI);
 
 }
 
 
-void FEmptyDynamicRHI::RHIDrawIndexedPrimitive(FRHIIndexBuffer* IndexBufferRHI, uint32 PrimitiveType, int32 BaseVertexIndex, uint32 FirstInstance,
+void FEmptyDynamicRHI::RHIDrawIndexedPrimitive(FRHIBuffer* IndexBufferRHI, uint32 PrimitiveType, int32 BaseVertexIndex, uint32 FirstInstance,
 	uint32 NumVertices, uint32 StartIndex, uint32 NumPrimitives, uint32 NumInstances)
 {
-	FEmptyIndexBuffer* IndexBuffer = ResourceCast(IndexBufferRHI);
-
-}
-
-void FEmptyDynamicRHI::RHIDrawIndexedIndirect(FRHIIndexBuffer* IndexBufferRHI, uint32 PrimitiveType, FRHIStructuredBuffer* ArgumentsBufferRHI, int32 DrawArgumentsIndex, uint32 NumInstances)
-{
-	FEmptyIndexBuffer* IndexBuffer = ResourceCast(IndexBufferRHI);
-	FEmptyStructuredBuffer* ArgumentsBuffer = ResourceCast(ArgumentsBufferRHI);
-
-}
-
-void FEmptyDynamicRHI::RHIDrawIndexedPrimitiveIndirect(uint32 PrimitiveType, FRHIIndexBuffer* IndexBufferRHI, FRHIVertexBuffer* ArgumentBufferRHI,uint32 ArgumentOffset)
-{
-	FEmptyIndexBuffer* IndexBuffer = ResourceCast(IndexBufferRHI);
-	FEmptyVertexBuffer* ArgumentBuffer = ResourceCast(ArgumentBufferRHI);
-
-}
-
-
-/** Some locally global variables to track the pending primitive information uised in RHIEnd*UP functions */
-static void *GPendingDrawPrimitiveUPVertexData = NULL;
-static uint32 GPendingNumVertices;
-static uint32 GPendingVertexDataStride;
-
-static void *GPendingDrawPrimitiveUPIndexData = NULL;
-static uint32 GPendingPrimitiveType;
-static uint32 GPendingNumPrimitives;
-static uint32 GPendingMinVertexIndex;
-static uint32 GPendingIndexDataStride;
-
-void FEmptyDynamicRHI::RHIBeginDrawPrimitiveUP( uint32 PrimitiveType, uint32 NumPrimitives, uint32 NumVertices, uint32 VertexDataStride, void*& OutVertexData)
-{
-	checkSlow(GPendingDrawPrimitiveUPVertexData == NULL);
-//	GPendingDrawPrimitiveUPVertexData = 
-	OutVertexData = GPendingDrawPrimitiveUPVertexData;
+	FEmptyBuffer* IndexBuffer = ResourceCast(IndexBufferRHI);
 
-	GPendingPrimitiveType = PrimitiveType;
-	GPendingNumPrimitives = NumPrimitives;
-	GPendingNumVertices = NumVertices;
-	GPendingVertexDataStride = VertexDataStride;
 }
 
-
-void FEmptyDynamicRHI::RHIEndDrawPrimitiveUP()
+void FEmptyDynamicRHI::RHIDrawIndexedIndirect(FRHIBuffer* IndexBufferRHI, uint32 PrimitiveType, FRHIBuffer* ArgumentsBufferRHI, int32 DrawArgumentsIndex, uint32 NumInstances)
 {
+	FEmptyBuffer* IndexBuffer = ResourceCast(IndexBufferRHI);
+	FEmptyBuffer* ArgumentsBuffer = ResourceCast(ArgumentsBufferRHI);
 
-	// free used mem
-	GPendingDrawPrimitiveUPVertexData = NULL;
 }
 
-void FEmptyDynamicRHI::RHIBeginDrawIndexedPrimitiveUP( uint32 PrimitiveType, uint32 NumPrimitives, uint32 NumVertices, uint32 VertexDataStride, void*& OutVertexData, uint32 MinVertexIndex, uint32 NumIndices, uint32 IndexDataStride, void*& OutIndexData)
+void FEmptyDynamicRHI::RHIDrawIndexedPrimitiveIndirect(uint32 PrimitiveType, FRHIBuffer* IndexBufferRHI, FRHIBuffer* ArgumentBufferRHI,uint32 ArgumentOffset)
 {
-	checkSlow(GPendingDrawPrimitiveUPVertexData == NULL);
-	checkSlow(GPendingDrawPrimitiveUPIndexData == NULL);
+	FEmptyBuffer* IndexBuffer = ResourceCast(IndexBufferRHI);
+	FEmptyBuffer* ArgumentBuffer = ResourceCast(ArgumentBufferRHI);
 
-//	GPendingDrawPrimitiveUPVertexData = 
-	OutVertexData = GPendingDrawPrimitiveUPVertexData;
-
-//	GPendingDrawPrimitiveUPIndexData = 
-	OutIndexData = GPendingDrawPrimitiveUPIndexData;
-
-	GPendingPrimitiveType = PrimitiveType;
-	GPendingNumPrimitives = NumPrimitives;
-	GPendingMinVertexIndex = MinVertexIndex;
-	GPendingIndexDataStride = IndexDataStride;
-
-	GPendingNumVertices = NumVertices;
-	GPendingVertexDataStride = VertexDataStride;
 }
 
-void FEmptyDynamicRHI::RHIEndDrawIndexedPrimitiveUP()
-{
-
-	// free used mem
-	GPendingDrawPrimitiveUPIndexData = NULL;
-	GPendingDrawPrimitiveUPVertexData = NULL;
-}
 
 void FEmptyDynamicRHI::RHIClearMRT(bool bClearColor,int32 NumClearColors,const FLinearColor* ClearColorArray,bool bClearDepth,float Depth,bool bClearStencil,uint32 Stencil)
 {
diff --git a/Engine/Source/Runtime/EmptyRHI/Public/EmptyRHI.h b/Engine/Source/Runtime/EmptyRHI/Public/EmptyRHI.h
--- a/Engine/Source/Runtime/EmptyRHI/Public/EmptyRHI.h
+++ b/Engine/Source/Runtime/EmptyRHI/Public/EmptyRHI.h
@@ -134,7 +134,7 @@ public:
 
 	virtual void RHISetComputeShader(FRHIComputeShader* ComputeShader) final override;
 	virtual void RHIDispatchComputeShader(uint32 ThreadGroupCountX, uint32 ThreadGroupCountY, uint32 ThreadGroupCountZ) final override;
-	virtual void RHIDispatchIndirectComputeShader(FRHIVertexBuffer* ArgumentBuffer, uint32 ArgumentOffset) final override;
+	virtual void RHIDispatchIndirectComputeShader(FRHIBuffer* ArgumentBuffer, uint32 ArgumentOffset) final override;
 	virtual void RHISetMultipleViewports(uint32 Count, const FViewportBounds* Data) final override;
 	virtual void RHIClearUAVFloat(FRHIUnorderedAccessView* UnorderedAccessViewRHI, const FVector4& Values) final override;
 	virtual void RHIClearUAVUint(FRHIUnorderedAccessView* UnorderedAccessViewRHI, const FUintVector4& Values) final override;
@@ -172,10 +172,10 @@ public:
 	virtual void RHIBeginRenderPass(const FRHIRenderPassInfo& InInfo, const TCHAR* InName) final override;
 	virtual void RHIEndRenderPass() final override;
 	virtual void RHIDrawPrimitive(uint32 PrimitiveType, uint32 BaseVertexIndex, uint32 NumPrimitives, uint32 NumInstances) final override;
-	virtual void RHIDrawPrimitiveIndirect(uint32 PrimitiveType, FRHIVertexBuffer* ArgumentBuffer, uint32 ArgumentOffset) final override;
-	virtual void RHIDrawIndexedIndirect(FRHIIndexBuffer* IndexBufferRHI, uint32 PrimitiveType, FRHIStructuredBuffer* ArgumentsBufferRHI, int32 DrawArgumentsIndex, uint32 NumInstances) final override;
-	virtual void RHIDrawIndexedPrimitive(FRHIIndexBuffer* IndexBuffer, uint32 PrimitiveType, int32 BaseVertexIndex, uint32 FirstInstance, uint32 NumVertices, uint32 StartIndex, uint32 NumPrimitives, uint32 NumInstances) final override;
-	virtual void RHIDrawIndexedPrimitiveIndirect(uint32 PrimitiveType, FRHIIndexBuffer* IndexBuffer, FRHIVertexBuffer* ArgumentBuffer, uint32 ArgumentOffset) final override;
+	virtual void RHIDrawPrimitiveIndirect(uint32 PrimitiveType, FRHIBuffer* ArgumentBuffer, uint32 ArgumentOffset) final override;
+	virtual void RHIDrawIndexedIndirect(FRHIBuffer* IndexBufferRHI, uint32 PrimitiveType, FRHIBuffer* ArgumentsBufferRHI, int32 DrawArgumentsIndex, uint32 NumInstances) final override;
+	virtual void RHIDrawIndexedPrimitive(FRHIBuffer* IndexBuffer, uint32 PrimitiveType, int32 BaseVertexIndex, uint32 FirstInstance, uint32 NumVertices, uint32 StartIndex, uint32 NumPrimitives, uint32 NumInstances) final override;
+	virtual void RHIDrawIndexedPrimitiveIndirect(uint32 PrimitiveType, FRHIBuffer* IndexBuffer, FRHIBuffer* ArgumentBuffer, uint32 ArgumentOffset) final override;
 	virtual void RHIEnableDepthBoundsTest(bool bEnable) final override;
 	virtual void RHISetDepthBounds(float MinDepth, float MaxDepth) final override;
 	virtual void RHISetShadingRate(EVRSShadingRate ShadingRate, EVRSRateCombiner Combiner) final override;
diff --git a/Engine/Source/Runtime/EmptyRHI/Public/EmptyResources.h b/Engine/Source/Runtime/EmptyRHI/Public/EmptyResources.h
--- a/Engine/Source/Runtime/EmptyRHI/Public/EmptyResources.h
+++ b/Engine/Source/Runtime/EmptyRHI/Public/EmptyResources.h
@@ -179,13 +179,13 @@ public:
 };
 
 /** Index buffer resource class that stores stride information. */
-class FEmptyIndexBuffer : public FRHIIndexBuffer
+class FEmptyBuffer : public FRHIBuffer
 {
 public:
-	FEmptyIndexBuffer() = default;
+	FEmptyBuffer() = default;
 
 	/** Constructor */
-	FEmptyIndexBuffer(uint32 InStride, uint32 InSize, uint32 InUsage);
+	FEmptyBuffer(uint32 InStride, uint32 InSize, uint32 InUsage);
 
 	/**
 	 * Prepare a CPU accessible buffer for uploading to GPU memory
@@ -199,26 +199,6 @@ public:
 };
 
 
-/** Vertex buffer resource class that stores usage type. */
-class FEmptyVertexBuffer : public FRHIVertexBuffer
-{
-public:
-	FEmptyVertexBuffer() = default;
-
-	/** Constructor */
-	FEmptyVertexBuffer(uint32 InSize, uint32 InUsage);
-
-	/**
-	 * Prepare a CPU accessible buffer for uploading to GPU memory
-	 */
-	void* Lock(EResourceLockMode LockMode, uint32 Size=0);
-
-	/**
-	 * Prepare a CPU accessible buffer for uploading to GPU memory
-	 */
-	void Unlock();
-};
-
 class FEmptyUniformBuffer : public FRHIUniformBuffer
 {
 public:
@@ -231,27 +211,13 @@ public:
 };
 
 
-class FEmptyStructuredBuffer : public FRHIStructuredBuffer
-{
-public:
-	// Constructor
-	FEmptyStructuredBuffer(uint32 Stride, uint32 Size, FResourceArrayInterface* ResourceArray, uint32 InUsage);
-
-	// Destructor
-	~FEmptyStructuredBuffer();
-
-};
-
-
 
 class FEmptyUnorderedAccessView : public FRHIUnorderedAccessView
 {
 public:
 
 	// the potential resources to refer to with the UAV object
-	TRefCountPtr<FEmptyStructuredBuffer> SourceStructuredBuffer;
-	TRefCountPtr<FEmptyVertexBuffer> SourceVertexBuffer;
-	TRefCountPtr<FEmptyIndexBuffer> SourceIndexBuffer;
+	TRefCountPtr<FEmptyBuffer> SourceBuffer;
 	TRefCountPtr<FRHITexture> SourceTexture;
 };
 
@@ -260,8 +226,8 @@ class FEmptyShaderResourceView : public FRHIShaderResourceView
 {
 public:
 
-	// The vertex buffer this SRV comes from (can be null)
-	TRefCountPtr<FEmptyVertexBuffer> SourceVertexBuffer;
+	// The buffer this SRV comes from (can be null)
+	TRefCountPtr<FEmptyBuffer> SourceBuffer;
 
 	// The texture that this SRV come from
 	TRefCountPtr<FRHITexture> SourceTexture;
@@ -349,19 +315,9 @@ struct TEmptyResourceTraits<FRHIUniformBuffer>
 	typedef FEmptyUniformBuffer TConcreteType;
 };
 template<>
-struct TEmptyResourceTraits<FRHIIndexBuffer>
-{
-	typedef FEmptyIndexBuffer TConcreteType;
-};
-template<>
-struct TEmptyResourceTraits<FRHIStructuredBuffer>
-{
-	typedef FEmptyStructuredBuffer TConcreteType;
-};
-template<>
-struct TEmptyResourceTraits<FRHIVertexBuffer>
+struct TEmptyResourceTraits<FRHIBuffer>
 {
-	typedef FEmptyVertexBuffer TConcreteType;
+	typedef FEmptyBuffer TConcreteType;
 };
 template<>
 struct TEmptyResourceTraits<FRHIShaderResourceView>
diff --git a/Engine/Source/Runtime/Engine/Private/VT/VirtualTextureUploadCache.h b/Engine/Source/Runtime/Engine/Private/VT/VirtualTextureUploadCache.h
--- a/Engine/Source/Runtime/Engine/Private/VT/VirtualTextureUploadCache.h
+++ b/Engine/Source/Runtime/Engine/Private/VT/VirtualTextureUploadCache.h
@@ -9,7 +9,7 @@
 #include "PixelFormat.h"
 
 class FRHITexture2D;
-class FRHIStructuredBuffer;
+class FRHIBuffer;
 
 struct FVTUploadTileBuffer
 {
diff --git a/Engine/Source/Runtime/NullDrv/Public/NullRHI.h b/Engine/Source/Runtime/NullDrv/Public/NullRHI.h
--- a/Engine/Source/Runtime/NullDrv/Public/NullRHI.h
+++ b/Engine/Source/Runtime/NullDrv/Public/NullRHI.h
@@ -88,7 +88,7 @@ public:
 
 	}
 
-	virtual void RHIDispatchIndirectComputeShader(FRHIVertexBuffer* ArgumentBuffer, uint32 ArgumentOffset) final override
+	virtual void RHIDispatchIndirectComputeShader(FRHIBuffer* ArgumentBuffer, uint32 ArgumentOffset) final override
 	{
 
 	}
@@ -142,7 +142,7 @@ public:
 
 	}
 
-	virtual void RHITransferIndexBufferUnderlyingResource(FRHIIndexBuffer* DestIndexBuffer, FRHIIndexBuffer* SrcIndexBuffer) final override
+	virtual void RHITransferBufferUnderlyingResource(FRHIBuffer* DestBuffer, FRHIBuffer* SrcBuffer) final override
 	{
 
 	}
@@ -153,7 +153,7 @@ public:
 		{ 
 			CreateInfo.ResourceArray->Discard(); 
 		} 
-		return new FRHIVertexBuffer(Size,InUsage); 
+		return new FRHIVertexBuffer(Size, InUsage, 0); 
 	}
 
 	virtual void* LockVertexBuffer_BottomOfPipe(FRHICommandListImmediate& RHICmdList, FRHIVertexBuffer* VertexBuffer, uint32 Offset, uint32 SizeRHI, EResourceLockMode LockMode) final override
@@ -166,12 +166,6 @@ public:
 
 	}
 
-	virtual void RHITransferVertexBufferUnderlyingResource(FRHIVertexBuffer* DestVertexBuffer, FRHIVertexBuffer* SrcVertexBuffer) final override
-	{
-
-	}
-
-
 	virtual void RHICopyVertexBuffer(FRHIVertexBuffer* SourceBuffer, FRHIVertexBuffer* DestBuffer) final override
 	{
 
@@ -196,7 +190,7 @@ public:
 	}
 
 
-	virtual FUnorderedAccessViewRHIRef RHICreateUnorderedAccessView(FRHIStructuredBuffer* StructuredBuffer, bool bUseUAVCounter, bool bAppendBuffer) final override
+	virtual FUnorderedAccessViewRHIRef RHICreateUnorderedAccessView(FRHIBuffer* StructuredBuffer, bool bUseUAVCounter, bool bAppendBuffer) final override
 	{ 
 		return new FRHIUnorderedAccessView(); 
 	}
@@ -208,23 +202,18 @@ public:
 	}
 
 
-	virtual FUnorderedAccessViewRHIRef RHICreateUnorderedAccessView(FRHIVertexBuffer* VertexBuffer, uint8 Format) final override
+	virtual FUnorderedAccessViewRHIRef RHICreateUnorderedAccessView(FRHIBuffer* Buffer, uint8 Format) final override
 	{ 
 		return new FRHIUnorderedAccessView(); 
 	}
 
-	virtual FUnorderedAccessViewRHIRef RHICreateUnorderedAccessView(FRHIIndexBuffer* IndexBuffer, uint8 Format) final override
-	{
-		return new FRHIUnorderedAccessView();
-	}
-
-	virtual FShaderResourceViewRHIRef RHICreateShaderResourceView(FRHIStructuredBuffer* StructuredBuffer) final override
+	virtual FShaderResourceViewRHIRef RHICreateShaderResourceView(FRHIBuffer* Buffer) final override
 	{ 
 		return new FRHIShaderResourceView(); 
 	}
 
 
-	virtual FShaderResourceViewRHIRef RHICreateShaderResourceView(FRHIVertexBuffer* VertexBuffer, uint32 Stride, uint8 Format) final override
+	virtual FShaderResourceViewRHIRef RHICreateShaderResourceView(FRHIBuffer* Buffer, uint32 Stride, uint8 Format) final override
 	{ 
 		return new FRHIShaderResourceView(); 
 	}
@@ -234,11 +223,6 @@ public:
 		return new FRHIShaderResourceView();
 	}
 
-	virtual FShaderResourceViewRHIRef RHICreateShaderResourceView(FRHIIndexBuffer* Buffer) final override
-	{ 
-		return new FRHIShaderResourceView(); 
-	}
-
 	virtual void RHIClearUAVFloat(FRHIUnorderedAccessView* UnorderedAccessViewRHI, const FVector4& Values) final override
 	{
 
@@ -641,22 +625,22 @@ public:
 	{
 
 	}
-	virtual void RHIDrawPrimitiveIndirect(FRHIVertexBuffer* ArgumentBuffer, uint32 ArgumentOffset) final override
+	virtual void RHIDrawPrimitiveIndirect(FRHIBuffer* ArgumentBuffer, uint32 ArgumentOffset) final override
 	{
 
 	}
 
-	virtual void RHIDrawIndexedIndirect(FRHIIndexBuffer* IndexBufferRHI, FRHIStructuredBuffer* ArgumentsBufferRHI, int32 DrawArgumentsIndex, uint32 NumInstances) final override
+	virtual void RHIDrawIndexedIndirect(FRHIBuffer* IndexBufferRHI, FRHIBuffer* ArgumentsBufferRHI, int32 DrawArgumentsIndex, uint32 NumInstances) final override
 	{
 
 	}
 
 
-	virtual void RHIDrawIndexedPrimitive(FRHIIndexBuffer* IndexBuffer, int32 BaseVertexIndex, uint32 FirstInstance, uint32 NumVertices, uint32 StartIndex, uint32 NumPrimitives, uint32 NumInstances) final override
+	virtual void RHIDrawIndexedPrimitive(FRHIBuffer* IndexBuffer, int32 BaseVertexIndex, uint32 FirstInstance, uint32 NumVertices, uint32 StartIndex, uint32 NumPrimitives, uint32 NumInstances) final override
 	{
 
 	}
-	virtual void RHIDrawIndexedPrimitiveIndirect(FRHIIndexBuffer* IndexBuffer, FRHIVertexBuffer* ArgumentBuffer, uint32 ArgumentOffset) final override
+	virtual void RHIDrawIndexedPrimitiveIndirect(FRHIBuffer* IndexBuffer, FRHIBuffer* ArgumentBuffer, uint32 ArgumentOffset) final override
 	{
 
 	}
diff --git a/Engine/Source/Runtime/OpenGLDrv/Private/OpenGLCommands.cpp b/Engine/Source/Runtime/OpenGLDrv/Private/OpenGLCommands.cpp
--- a/Engine/Source/Runtime/OpenGLDrv/Private/OpenGLCommands.cpp
+++ b/Engine/Source/Runtime/OpenGLDrv/Private/OpenGLCommands.cpp
@@ -136,13 +136,13 @@ struct FPendingSamplerDataValue
 
 struct FVertexBufferPair
 {
-	FOpenGLVertexBuffer*				Source;
-	TRefCountPtr<FOpenGLVertexBuffer>	Dest;
+	FOpenGLBuffer*				Source;
+	TRefCountPtr<FOpenGLBuffer>	Dest;
 };
 static TArray<FVertexBufferPair> ZeroStrideExpandedBuffersList;
 
 
-static int FindVertexBuffer(FOpenGLVertexBuffer* Source)
+static int FindVertexBuffer(FOpenGLBuffer* Source)
 {
 	for (int32 Index = 0; Index < ZeroStrideExpandedBuffersList.Num(); ++Index)
 	{
@@ -154,14 +154,14 @@ static int FindVertexBuffer(FOpenGLVertexBuffer* Source)
 	return -1;
 }
 
-static FOpenGLVertexBuffer* FindExpandedZeroStrideBuffer(FOpenGLVertexBuffer* ZeroStrideVertexBuffer, uint32 Stride, uint32 NumVertices, const FOpenGLVertexElement& VertexElement)
+static FOpenGLBuffer* FindExpandedZeroStrideBuffer(FOpenGLBuffer* ZeroStrideVertexBuffer, uint32 Stride, uint32 NumVertices, const FOpenGLVertexElement& VertexElement)
 {
 	uint32 Size = NumVertices * Stride;
 	int32 FoundExpandedVBIndex = FindVertexBuffer(ZeroStrideVertexBuffer);
 	if (FoundExpandedVBIndex != -1)
 	{
 		// Check if the current size is big enough
-		FOpenGLVertexBuffer* ExpandedVB = ZeroStrideExpandedBuffersList[FoundExpandedVBIndex].Dest;
+		FOpenGLBuffer* ExpandedVB = ZeroStrideExpandedBuffersList[FoundExpandedVBIndex].Dest;
 		if (Size <= ExpandedVB->GetSize())
 		{
 			return ExpandedVB;
@@ -205,7 +205,7 @@ static FOpenGLVertexBuffer* FindExpandedZeroStrideBuffer(FOpenGLVertexBuffer* Ze
 	const int32 SizeToFill = VertexElementSize * VertexTypeSize;
 	void* RESTRICT SourceData = ZeroStrideVertexBuffer->GetZeroStrideBuffer();
 	check(SourceData);
-	TRefCountPtr<FOpenGLVertexBuffer> ExpandedVB = new FOpenGLVertexBuffer(0, Size, BUF_Static, NULL);
+	TRefCountPtr<FOpenGLBuffer> ExpandedVB = new FOpenGLBuffer(GL_ARRAY_BUFFER, 0, Size, BUF_Static, nullptr);
 	uint8* RESTRICT Data = ExpandedVB->Lock(0, Size, false, true);
 
 	switch (SizeToFill)
@@ -390,7 +390,7 @@ static FORCEINLINE void GetShaderStageIndexAndMaxUnits(FRHIGraphicsShader* Shade
 void FOpenGLDynamicRHI::RHISetStreamSource(uint32 StreamIndex, FRHIVertexBuffer* VertexBufferRHI, uint32 Offset)
 {
 	VERIFY_GL_SCOPE();
-	FOpenGLVertexBuffer* VertexBuffer = ResourceCast(VertexBufferRHI);
+	FOpenGLBuffer* VertexBuffer = ResourceCast(VertexBufferRHI);
 	PendingState.Streams[StreamIndex].VertexBuffer = VertexBuffer;
 	PendingState.Streams[StreamIndex].Stride = PendingState.BoundShaderState ? PendingState.BoundShaderState->StreamStrides[StreamIndex] : 0;
 	PendingState.Streams[StreamIndex].Offset = Offset;
@@ -969,7 +969,7 @@ void FOpenGLDynamicRHI::UpdateSRV(FOpenGLShaderResourceView* SRV)
 		uint32 SizeY = Texture2D->GetSizeY();
 		
 		uint32 MipBytes = SizeX * SizeY;
-		TRefCountPtr<FOpenGLPixelBuffer> PixelBuffer = new FOpenGLPixelBuffer(0, MipBytes, BUF_Dynamic);
+		TRefCountPtr<FOpenGLPixelBuffer> PixelBuffer = new FOpenGLPixelBuffer(GL_PIXEL_UNPACK_BUFFER, 0, MipBytes, BUF_Dynamic, nullptr);
 		
 		glBindBuffer( GL_PIXEL_PACK_BUFFER, 0 );
 		glBindBuffer( GL_PIXEL_PACK_BUFFER, PixelBuffer->Resource );
@@ -1845,10 +1845,10 @@ void FOpenGLDynamicRHI::EnableVertexElementCached(
 	}
 }
 
-FORCEINLINE void FOpenGLDynamicRHI::EnableVertexElementCachedZeroStride(FOpenGLContextState& ContextState, GLuint AttributeIndex, const FOpenGLVertexElement& VertexElement, uint32 NumVertices, FOpenGLVertexBuffer* ZeroStrideVertexBuffer)
+FORCEINLINE void FOpenGLDynamicRHI::EnableVertexElementCachedZeroStride(FOpenGLContextState& ContextState, GLuint AttributeIndex, const FOpenGLVertexElement& VertexElement, uint32 NumVertices, FOpenGLBuffer* ZeroStrideVertexBuffer)
 {
 	uint32 Stride = ZeroStrideVertexBuffer->GetSize();
-	FOpenGLVertexBuffer* ExpandedVertexBuffer = FindExpandedZeroStrideBuffer(ZeroStrideVertexBuffer, Stride, NumVertices, VertexElement);
+	FOpenGLBuffer* ExpandedVertexBuffer = FindExpandedZeroStrideBuffer(ZeroStrideVertexBuffer, Stride, NumVertices, VertexElement);
 	EnableVertexElementCached(ContextState, AttributeIndex, VertexElement, Stride, 0, ExpandedVertexBuffer->Resource);
 }
 
@@ -2189,15 +2189,16 @@ void FOpenGLDynamicRHI::OnProgramDeletion( GLint ProgramResource )
 	}
 }
 
-void FOpenGLDynamicRHI::OnVertexBufferDeletion( GLuint VertexBufferResource )
+void FOpenGLDynamicRHI::OnBufferDeletion( GLuint BufferResource )
 {
 	VERIFY_GL_SCOPE();
-	if (SharedContextState.ArrayBufferBound == VertexBufferResource)
+
+	if (SharedContextState.ArrayBufferBound == BufferResource)
 	{
 		SharedContextState.ArrayBufferBound = -1;	// will force refresh
 	}
 
-	if (RenderingContextState.ArrayBufferBound == VertexBufferResource)
+	if (RenderingContextState.ArrayBufferBound == BufferResource)
 	{
 		RenderingContextState.ArrayBufferBound = -1;	// will force refresh
 	}
@@ -2211,7 +2212,7 @@ void FOpenGLDynamicRHI::OnVertexBufferDeletion( GLuint VertexBufferResource )
 			FOpenGLStream& CachedStream = SharedContextState.VertexStreams[StreamIndex];
 			if ((ActiveStreamMask & 0x1) && 
 				CachedStream.VertexBuffer && 
-				CachedStream.VertexBuffer->Resource == VertexBufferResource)
+				CachedStream.VertexBuffer->Resource == BufferResource)
 			{
 				FOpenGL::BindVertexBuffer(StreamIndex, 0, 0, 0); // brianh@nvidia: work around driver bug 1809000
 				CachedStream.VertexBuffer = nullptr;
@@ -2228,7 +2229,7 @@ void FOpenGLDynamicRHI::OnVertexBufferDeletion( GLuint VertexBufferResource )
 			FOpenGLStream& CachedStream = RenderingContextState.VertexStreams[StreamIndex];
 			if ((ActiveStreamMask & 0x1) && 
 				CachedStream.VertexBuffer && 
-				CachedStream.VertexBuffer->Resource == VertexBufferResource)
+				CachedStream.VertexBuffer->Resource == BufferResource)
 			{
 				FOpenGL::BindVertexBuffer(StreamIndex, 0, 0, 0); // brianh@nvidia: work around driver bug 1809000
 				CachedStream.VertexBuffer = nullptr;
@@ -2242,7 +2243,7 @@ void FOpenGLDynamicRHI::OnVertexBufferDeletion( GLuint VertexBufferResource )
 	{
 		for (GLuint AttribIndex = 0; AttribIndex < NUM_OPENGL_VERTEX_STREAMS; AttribIndex++)
 		{
-			if (SharedContextState.VertexAttrs[AttribIndex].Buffer == VertexBufferResource)
+			if (SharedContextState.VertexAttrs[AttribIndex].Buffer == BufferResource)
 			{
 				SharedContextState.VertexAttrs[AttribIndex].Pointer = FOpenGLCachedAttr_Invalid;	// that'll enforce state update on next cache test		}
 			}
@@ -2250,23 +2251,19 @@ void FOpenGLDynamicRHI::OnVertexBufferDeletion( GLuint VertexBufferResource )
 	
 		for (GLuint AttribIndex = 0; AttribIndex < NUM_OPENGL_VERTEX_STREAMS; AttribIndex++)
 		{
-			if (RenderingContextState.VertexAttrs[AttribIndex].Buffer == VertexBufferResource)
+			if (RenderingContextState.VertexAttrs[AttribIndex].Buffer == BufferResource)
 			{
 				RenderingContextState.VertexAttrs[AttribIndex].Pointer = FOpenGLCachedAttr_Invalid;	// that'll enforce state update on next cache test		}
 			}
 		}
 	}
-}
 
-void FOpenGLDynamicRHI::OnIndexBufferDeletion( GLuint IndexBufferResource )
-{
-	VERIFY_GL_SCOPE();
-	if (SharedContextState.ElementArrayBufferBound == IndexBufferResource)
+	if (SharedContextState.ElementArrayBufferBound == BufferResource)
 	{
 		SharedContextState.ElementArrayBufferBound = -1;	// will force refresh
 	}
 
-	if (RenderingContextState.ElementArrayBufferBound == IndexBufferResource)
+	if (RenderingContextState.ElementArrayBufferBound == BufferResource)
 	{
 		RenderingContextState.ElementArrayBufferBound = -1;	// will force refresh
 	}
@@ -2675,7 +2672,7 @@ void FOpenGLDynamicRHI::RHIDrawPrimitive(uint32 BaseVertexIndex,uint32 NumPrimit
 	}
 }
 
-void FOpenGLDynamicRHI::RHIDrawPrimitiveIndirect(FRHIVertexBuffer* ArgumentBufferRHI,uint32 ArgumentOffset)
+void FOpenGLDynamicRHI::RHIDrawPrimitiveIndirect(FRHIBuffer* ArgumentBufferRHI, uint32 ArgumentOffset)
 {
 	if (FOpenGL::SupportsDrawIndirect())
 	{
@@ -2711,7 +2708,7 @@ void FOpenGLDynamicRHI::RHIDrawPrimitiveIndirect(FRHIVertexBuffer* ArgumentBuffe
 			FOpenGL::PatchParameteri(GL_PATCH_VERTICES, PatchSize);
 		} 
 
-		FOpenGLVertexBuffer* ArgumentBuffer = ResourceCast(ArgumentBufferRHI);
+		FOpenGLBuffer* ArgumentBuffer = ResourceCast(ArgumentBufferRHI);
 
 
 		glBindBuffer( GL_DRAW_INDIRECT_BUFFER, ArgumentBuffer->Resource);
@@ -2728,14 +2725,14 @@ void FOpenGLDynamicRHI::RHIDrawPrimitiveIndirect(FRHIVertexBuffer* ArgumentBuffe
 
 }
 
-void FOpenGLDynamicRHI::RHIDrawIndexedIndirect(FRHIIndexBuffer* IndexBufferRHI, FRHIStructuredBuffer* ArgumentsBufferRHI, int32 DrawArgumentsIndex, uint32 NumInstances)
+void FOpenGLDynamicRHI::RHIDrawIndexedIndirect(FRHIBuffer* IndexBufferRHI, FRHIBuffer* ArgumentsBufferRHI, int32 DrawArgumentsIndex, uint32 NumInstances)
 {
 	if (FOpenGL::SupportsDrawIndirect())
 	{
 		VERIFY_GL_SCOPE();
 
-		FOpenGLIndexBuffer* IndexBuffer = ResourceCast(IndexBufferRHI);
-	GPUProfilingData.RegisterGPUWork(1);
+		FOpenGLBuffer* IndexBuffer = ResourceCast(IndexBufferRHI);
+		GPUProfilingData.RegisterGPUWork(1);
 
 		check(ArgumentsBufferRHI);
 
@@ -2771,7 +2768,7 @@ void FOpenGLDynamicRHI::RHIDrawIndexedIndirect(FRHIIndexBuffer* IndexBufferRHI,
 
 		GLenum IndexType = IndexBuffer->GetStride() == sizeof(uint32) ? GL_UNSIGNED_INT : GL_UNSIGNED_SHORT;
 
-		FOpenGLStructuredBuffer* ArgumentsBuffer = ResourceCast(ArgumentsBufferRHI);
+		FOpenGLBuffer* ArgumentsBuffer = ResourceCast(ArgumentsBufferRHI);
 
 
 		glBindBuffer( GL_DRAW_INDIRECT_BUFFER, ArgumentsBuffer->Resource);
@@ -2789,12 +2786,12 @@ void FOpenGLDynamicRHI::RHIDrawIndexedIndirect(FRHIIndexBuffer* IndexBufferRHI,
 	}
 }
 
-void FOpenGLDynamicRHI::RHIDrawIndexedPrimitive(FRHIIndexBuffer* IndexBufferRHI,int32 BaseVertexIndex,uint32 FirstInstance,uint32 NumVertices,uint32 StartIndex,uint32 NumPrimitives,uint32 NumInstances)
+void FOpenGLDynamicRHI::RHIDrawIndexedPrimitive(FRHIBuffer* IndexBufferRHI, int32 BaseVertexIndex, uint32 FirstInstance, uint32 NumVertices, uint32 StartIndex, uint32 NumPrimitives, uint32 NumInstances)
 {
 	SCOPE_CYCLE_COUNTER_DETAILED(STAT_OpenGLDrawPrimitiveTime);
 	VERIFY_GL_SCOPE();
 
-	FOpenGLIndexBuffer* IndexBuffer = ResourceCast(IndexBufferRHI);
+	FOpenGLBuffer* IndexBuffer = ResourceCast(IndexBufferRHI);
 
 	RHI_DRAW_CALL_STATS(PrimitiveType,NumPrimitives*NumInstances);
 
@@ -2894,13 +2891,13 @@ void FOpenGLDynamicRHI::RHIDrawIndexedPrimitive(FRHIIndexBuffer* IndexBufferRHI,
 	}
 }
 
-void FOpenGLDynamicRHI::RHIDrawIndexedPrimitiveIndirect(FRHIIndexBuffer* IndexBufferRHI, FRHIVertexBuffer* ArgumentBufferRHI,uint32 ArgumentOffset)
+void FOpenGLDynamicRHI::RHIDrawIndexedPrimitiveIndirect(FRHIBuffer* IndexBufferRHI, FRHIBuffer* ArgumentBufferRHI, uint32 ArgumentOffset)
 {
 	if (FOpenGL::SupportsDrawIndirect())
 	{
 		VERIFY_GL_SCOPE();
 
-		FOpenGLIndexBuffer* IndexBuffer = ResourceCast(IndexBufferRHI);
+		FOpenGLBuffer* IndexBuffer = ResourceCast(IndexBufferRHI);
 		GPUProfilingData.RegisterGPUWork(1);
 
 		check(ArgumentBufferRHI);
@@ -2934,7 +2931,7 @@ void FOpenGLDynamicRHI::RHIDrawIndexedPrimitiveIndirect(FRHIIndexBuffer* IndexBu
 
 		GLenum IndexType = IndexBuffer->GetStride() == sizeof(uint32) ? GL_UNSIGNED_INT : GL_UNSIGNED_SHORT;
 
-		FOpenGLVertexBuffer* ArgumentBuffer = ResourceCast(ArgumentBufferRHI);
+		FOpenGLBuffer* ArgumentBuffer = ResourceCast(ArgumentBufferRHI);
 
 
 		glBindBuffer( GL_DRAW_INDIRECT_BUFFER, ArgumentBuffer->Resource);
@@ -3240,7 +3237,7 @@ void FOpenGLDynamicRHI::RHIDispatchComputeShader(uint32 ThreadGroupCountX, uint3
 	}
 }
 
-void FOpenGLDynamicRHI::RHIDispatchIndirectComputeShader(FRHIVertexBuffer* ArgumentBufferRHI,uint32 ArgumentOffset)
+void FOpenGLDynamicRHI::RHIDispatchIndirectComputeShader(FRHIBuffer* ArgumentBufferRHI, uint32 ArgumentOffset)
 {
 	if (FOpenGL::SupportsComputeShaders())
 	{
@@ -3255,7 +3252,7 @@ void FOpenGLDynamicRHI::RHIDispatchIndirectComputeShader(FRHIVertexBuffer* Argum
 			ComputeShader->LinkedProgram = GetLinkedComputeProgram(ComputeShaderRHI);
 		}
 
-		FOpenGLVertexBuffer* ArgumentBuffer = ResourceCast(ArgumentBufferRHI);
+		FOpenGLBuffer* ArgumentBuffer = ResourceCast(ArgumentBufferRHI);
 
 		FOpenGLContextState& ContextState = GetContextStateForCurrentContext();
 
@@ -3346,7 +3343,7 @@ void FOpenGLDynamicRHI::RHIInvalidateCachedState()
 void FOpenGLDynamicRHI::RHICopyToStagingBuffer(FRHIVertexBuffer* SourceBufferRHI, FRHIStagingBuffer* DestinationStagingBufferRHI, uint32 InOffset, uint32 InNumBytes)
 {
 	VERIFY_GL_SCOPE();
-	FOpenGLVertexBuffer* SourceBuffer = ResourceCast(SourceBufferRHI);
+	FOpenGLBuffer* SourceBuffer = ResourceCast(SourceBufferRHI);
 	FOpenGLStagingBuffer* DestinationBuffer = ResourceCast(DestinationStagingBufferRHI);
 
 	check(DestinationBuffer->ShadowBuffer != 0);
diff --git a/Engine/Source/Runtime/OpenGLDrv/Private/OpenGLDevice.cpp b/Engine/Source/Runtime/OpenGLDrv/Private/OpenGLDevice.cpp
--- a/Engine/Source/Runtime/OpenGLDrv/Private/OpenGLDevice.cpp
+++ b/Engine/Source/Runtime/OpenGLDrv/Private/OpenGLDevice.cpp
@@ -82,16 +82,10 @@ void OnProgramDeletion( GLint ProgramResource )
 	PrivateOpenGLDevicePtr->OnProgramDeletion( ProgramResource );
 }
 
-void OnVertexBufferDeletion( GLuint VertexBufferResource )
+void OnBufferDeletion( GLuint BufferResource )
 {
 	check(PrivateOpenGLDevicePtr);
-	PrivateOpenGLDevicePtr->OnVertexBufferDeletion( VertexBufferResource );
-}
-
-void OnIndexBufferDeletion( GLuint IndexBufferResource )
-{
-	check(PrivateOpenGLDevicePtr);
-	PrivateOpenGLDevicePtr->OnIndexBufferDeletion( IndexBufferResource );
+	PrivateOpenGLDevicePtr->OnBufferDeletion( BufferResource );
 }
 
 void OnPixelBufferDeletion( GLuint PixelBufferResource )
@@ -106,19 +100,24 @@ void OnUniformBufferDeletion( GLuint UniformBufferResource, uint32 AllocatedSize
 	PrivateOpenGLDevicePtr->OnUniformBufferDeletion( UniformBufferResource, AllocatedSize, bStreamDraw );
 }
 
-void CachedBindArrayBuffer( GLuint Buffer )
-{
-	check(PrivateOpenGLDevicePtr);
-	PrivateOpenGLDevicePtr->CachedBindArrayBuffer(PrivateOpenGLDevicePtr->GetContextStateForCurrentContext(),Buffer);
-}
-
-void CachedBindElementArrayBuffer( GLuint Buffer )
+void CachedBindBuffer( GLenum Type, GLuint Buffer )
 {
 	check(PrivateOpenGLDevicePtr);
-	PrivateOpenGLDevicePtr->CachedBindElementArrayBuffer(PrivateOpenGLDevicePtr->GetContextStateForCurrentContext(),Buffer);
+	if (Type == GL_ARRAY_BUFFER)
+	{
+		PrivateOpenGLDevicePtr->CachedBindArrayBuffer(PrivateOpenGLDevicePtr->GetContextStateForCurrentContext(), Buffer);
+	}
+	else if (Type == GL_ELEMENT_ARRAY_BUFFER)
+	{
+		PrivateOpenGLDevicePtr->CachedBindElementArrayBuffer(PrivateOpenGLDevicePtr->GetContextStateForCurrentContext(), Buffer);
+	}
+	else
+	{
+		checkNoEntry();
+	}
 }
 
-void CachedBindPixelUnpackBuffer( GLuint Buffer )
+void CachedBindPixelUnpackBuffer( GLenum Type, GLuint Buffer )
 {
 	check(PrivateOpenGLDevicePtr);
 	PrivateOpenGLDevicePtr->CachedBindPixelUnpackBuffer(PrivateOpenGLDevicePtr->GetContextStateForCurrentContext(),Buffer);
@@ -1111,13 +1110,6 @@ void FOpenGLDynamicRHI::Init()
 	FSamplerStateInitializerRHI PointSamplerStateParams(SF_Point,AM_Clamp,AM_Clamp,AM_Clamp);
 	PointSamplerState = this->RHICreateSamplerState(PointSamplerStateParams);
 
-	if (FOpenGL::SupportsFastBufferData())
-	{
-		// Allocate vertex and index buffers for DrawPrimitiveUP calls.
-		DynamicVertexBuffers.Init(CalcDynamicBufferSize(1));
-		DynamicIndexBuffers.Init(CalcDynamicBufferSize(1));
-	}
-
 #if PLATFORM_WINDOWS || PLATFORM_LINUX
 
 	extern int64 GOpenGLDedicatedVideoMemory;
@@ -1204,10 +1196,6 @@ void FOpenGLDynamicRHI::Cleanup()
 		FRenderResource::ReleaseRHIForAllResources();
 	}
 
-	// Release dynamic vertex and index buffers.
-	DynamicVertexBuffers.Cleanup();
-	DynamicIndexBuffers.Cleanup();
-
 	FreeZeroStrideBuffers();
 
 	// Release the point sampler state.
diff --git a/Engine/Source/Runtime/OpenGLDrv/Private/OpenGLIndexBuffer.cpp b/Engine/Source/Runtime/OpenGLDrv/Private/OpenGLIndexBuffer.cpp
--- a/Engine/Source/Runtime/OpenGLDrv/Private/OpenGLIndexBuffer.cpp
+++ b/Engine/Source/Runtime/OpenGLDrv/Private/OpenGLIndexBuffer.cpp
@@ -12,7 +12,7 @@ FIndexBufferRHIRef FOpenGLDynamicRHI::RHICreateIndexBuffer(uint32 Stride,uint32
 {
 	if (CreateInfo.bWithoutNativeResource)
 	{
-		return new FOpenGLIndexBuffer();
+		return new FOpenGLBuffer();
 	}
 
 	const void *Data = NULL;
@@ -24,7 +24,7 @@ FIndexBufferRHIRef FOpenGLDynamicRHI::RHICreateIndexBuffer(uint32 Stride,uint32
 		Data = CreateInfo.ResourceArray->GetResourceData();
 	}
 
-	TRefCountPtr<FOpenGLIndexBuffer> IndexBuffer = new FOpenGLIndexBuffer(Stride, Size, InUsage, Data);
+	TRefCountPtr<FOpenGLBuffer> IndexBuffer = new FOpenGLBuffer(GL_ELEMENT_ARRAY_BUFFER, Stride, Size, InUsage | BUF_IndexBuffer, Data);
 
 	if (CreateInfo.ResourceArray)
 	{
@@ -43,7 +43,7 @@ void* FOpenGLDynamicRHI::LockIndexBuffer_BottomOfPipe(FRHICommandListImmediate&
 {
 	RHITHREAD_GLCOMMAND_PROLOGUE();
 	VERIFY_GL_SCOPE();
-	FOpenGLIndexBuffer* IndexBuffer = ResourceCast(IndexBufferRHI);
+	FOpenGLBuffer* IndexBuffer = ResourceCast(IndexBufferRHI);
 	return IndexBuffer->Lock(Offset, Size, LockMode == RLM_ReadOnly, IndexBuffer->IsDynamic());
 	RHITHREAD_GLCOMMAND_EPILOGUE_RETURN(void*);
 }
@@ -52,24 +52,24 @@ void FOpenGLDynamicRHI::UnlockIndexBuffer_BottomOfPipe(FRHICommandListImmediate&
 {
 	RHITHREAD_GLCOMMAND_PROLOGUE();
 	VERIFY_GL_SCOPE();
-	FOpenGLIndexBuffer* IndexBuffer = ResourceCast(IndexBufferRHI);
+	FOpenGLBuffer* IndexBuffer = ResourceCast(IndexBufferRHI);
 	IndexBuffer->Unlock();
 	RHITHREAD_GLCOMMAND_EPILOGUE();
 }
 
-void FOpenGLDynamicRHI::RHITransferIndexBufferUnderlyingResource(FRHIIndexBuffer* DestIndexBuffer, FRHIIndexBuffer* SrcIndexBuffer)
+void FOpenGLDynamicRHI::RHITransferBufferUnderlyingResource(FRHIBuffer* DestBuffer, FRHIBuffer* SrcBuffer)
 {
 	VERIFY_GL_SCOPE();
-	check(DestIndexBuffer);
-	FOpenGLIndexBuffer* Dest = ResourceCast(DestIndexBuffer);
-	if (!SrcIndexBuffer)
+	check(DestBuffer);
+	FOpenGLBuffer* Dest = ResourceCast(DestBuffer);
+	if (!SrcBuffer)
 	{
-		TRefCountPtr<FOpenGLIndexBuffer> Src = new FOpenGLIndexBuffer();
+		TRefCountPtr<FOpenGLBuffer> Src = new FOpenGLBuffer();
 		Dest->Swap(*Src);
 	}
 	else
 	{
-		FOpenGLIndexBuffer* Src = ResourceCast(SrcIndexBuffer);
+		FOpenGLBuffer* Src = ResourceCast(SrcBuffer);
 		Dest->Swap(*Src);
 	}
 }
diff --git a/Engine/Source/Runtime/OpenGLDrv/Private/OpenGLStructuredBuffer.cpp b/Engine/Source/Runtime/OpenGLDrv/Private/OpenGLStructuredBuffer.cpp
--- a/Engine/Source/Runtime/OpenGLDrv/Private/OpenGLStructuredBuffer.cpp
+++ b/Engine/Source/Runtime/OpenGLDrv/Private/OpenGLStructuredBuffer.cpp
@@ -21,20 +21,20 @@ FStructuredBufferRHIRef FOpenGLDynamicRHI::RHICreateStructuredBuffer(uint32 Stri
 		Data = CreateInfo.ResourceArray->GetResourceData();
 	}
 
-	TRefCountPtr<FOpenGLStructuredBuffer> StructuredBuffer = new FOpenGLStructuredBuffer(Stride, Size, InUsage & BUF_AnyDynamic, Data);
+	TRefCountPtr<FOpenGLBuffer> StructuredBuffer = new FOpenGLBuffer(GL_ARRAY_BUFFER, Stride, Size, InUsage | BUF_StructuredBuffer, Data);
 	return StructuredBuffer.GetReference();
 }
 
 void* FOpenGLDynamicRHI::LockStructuredBuffer_BottomOfPipe(FRHICommandListImmediate& RHICmdList, FRHIStructuredBuffer* StructuredBufferRHI, uint32 Offset, uint32 Size, EResourceLockMode LockMode)
 {
 	VERIFY_GL_SCOPE();
-	FOpenGLStructuredBuffer* StructuredBuffer = ResourceCast(StructuredBufferRHI);
+	FOpenGLBuffer* StructuredBuffer = ResourceCast(StructuredBufferRHI);
 	return StructuredBuffer->Lock(Offset, Size, LockMode == RLM_ReadOnly, StructuredBuffer->IsDynamic());
 }
 
 void FOpenGLDynamicRHI::UnlockStructuredBuffer_BottomOfPipe(FRHICommandListImmediate& RHICmdList, FRHIStructuredBuffer* StructuredBufferRHI)
 {
 	VERIFY_GL_SCOPE();
-	FOpenGLStructuredBuffer* StructuredBuffer = ResourceCast(StructuredBufferRHI);
+	FOpenGLBuffer* StructuredBuffer = ResourceCast(StructuredBufferRHI);
 	StructuredBuffer->Unlock();
 }
diff --git a/Engine/Source/Runtime/OpenGLDrv/Private/OpenGLTexture.cpp b/Engine/Source/Runtime/OpenGLDrv/Private/OpenGLTexture.cpp
--- a/Engine/Source/Runtime/OpenGLDrv/Private/OpenGLTexture.cpp
+++ b/Engine/Source/Runtime/OpenGLDrv/Private/OpenGLTexture.cpp
@@ -668,7 +668,7 @@ void TOpenGLTexture<RHIResourceType>::Resolve(uint32 MipIndex,uint32 ArrayIndex)
 	// Standard path with a PBO mirroring ever slice of a texture to allow multiple simulataneous maps
 	if (!IsValidRef(PixelBuffers[BufferIndex]))
 	{
-		PixelBuffers[BufferIndex] = new FOpenGLPixelBuffer(0, MipBytes, BUF_Dynamic);
+		PixelBuffers[BufferIndex] = new FOpenGLPixelBuffer(GL_PIXEL_UNPACK_BUFFER, 0, MipBytes, BUF_Dynamic, nullptr);
 	}
 	
 	TRefCountPtr<FOpenGLPixelBuffer> PixelBuffer = PixelBuffers[BufferIndex];
@@ -771,7 +771,7 @@ void* TOpenGLTexture<RHIResourceType>::Lock(uint32 InMipIndex,uint32 ArrayIndex,
 		if (!IsValidRef(PixelBuffers[BufferIndex]))
 		{
 			bBufferExists = false;
-			PixelBuffers[BufferIndex] = new FOpenGLPixelBuffer(0, MipBytes, BUF_Dynamic);
+			PixelBuffers[BufferIndex] = new FOpenGLPixelBuffer(GL_PIXEL_UNPACK_BUFFER, 0, MipBytes, BUF_Dynamic, nullptr);
 		}
 
 		TRefCountPtr<FOpenGLPixelBuffer> PixelBuffer = PixelBuffers[BufferIndex];
@@ -871,7 +871,7 @@ void TOpenGLTexture<RHIResourceType>::Unlock(uint32 MipIndex,uint32 ArrayIndex)
 		FOpenGLContextState& ContextState = OpenGLRHI->GetContextStateForCurrentContext();
 		OpenGLRHI->CachedSetupTextureStage(ContextState, FOpenGL::GetMaxCombinedTextureImageUnits() - 1, Target, Resource, -1, this->GetNumMips());
 
-		CachedBindPixelUnpackBuffer(0);
+		CachedBindPixelUnpackBuffer(GL_PIXEL_UNPACK_BUFFER, 0);
 
 		// get the source data and size
 		uint16* floatData = (uint16*)PixelBuffer->GetLockedBuffer();
@@ -914,7 +914,7 @@ void TOpenGLTexture<RHIResourceType>::Unlock(uint32 MipIndex,uint32 ArrayIndex)
 		// and the next draw will take care of cleaning it up; or
 		// next operation that needs the stage will switch something else in on it.
 
-		CachedBindPixelUnpackBuffer(0);
+		CachedBindPixelUnpackBuffer(GL_PIXEL_UNPACK_BUFFER, 0);
 
 		return;
 	}
@@ -1028,7 +1028,7 @@ void TOpenGLTexture<RHIResourceType>::Unlock(uint32 MipIndex,uint32 ArrayIndex)
 	// and the next draw will take care of cleaning it up; or
 	// next operation that needs the stage will switch something else in on it.
 
-	CachedBindPixelUnpackBuffer(0);
+	CachedBindPixelUnpackBuffer(GL_PIXEL_UNPACK_BUFFER, 0);
 }
 
 template<typename RHIResourceType>
@@ -1107,7 +1107,7 @@ void TOpenGLTexture<RHIResourceType>::CloneViaPBO( TOpenGLTexture<RHIResourceTyp
 			// Standard path with a PBO mirroring ever slice of a texture to allow multiple simulataneous maps
 			if (!IsValidRef(PixelBuffers[BufferIndex]))
 			{
-				PixelBuffers[BufferIndex] = new FOpenGLPixelBuffer(0, MipBytes, BUF_Dynamic);
+				PixelBuffers[BufferIndex] = new FOpenGLPixelBuffer(GL_PIXEL_UNPACK_BUFFER, 0, MipBytes, BUF_Dynamic, nullptr);
 			}
 			
 			TRefCountPtr<FOpenGLPixelBuffer> PixelBuffer = PixelBuffers[BufferIndex];
@@ -1141,7 +1141,7 @@ void TOpenGLTexture<RHIResourceType>::CloneViaPBO( TOpenGLTexture<RHIResourceTyp
 			// copy the texture data
 			// Upload directly into Dst to avoid out-of-band synchronization caused by glMapBuffer!
 			{
-				CachedBindPixelUnpackBuffer( PixelBuffer->Resource );
+				CachedBindPixelUnpackBuffer( GL_PIXEL_UNPACK_BUFFER, PixelBuffer->Resource );
 				
 				// Use a texture stage that's not likely to be used for draws, to avoid waiting
 				OpenGLRHI->CachedSetupTextureStage(ContextState, FOpenGL::GetMaxCombinedTextureImageUnits() - 1, Target, Resource, -1, this->GetNumMips());
@@ -1241,7 +1241,7 @@ void TOpenGLTexture<RHIResourceType>::CloneViaPBO( TOpenGLTexture<RHIResourceTyp
 	
 	// Reset the buffer bindings on exit only
 	glBindBuffer( GL_PIXEL_PACK_BUFFER, 0 );
-	CachedBindPixelUnpackBuffer(0);
+	CachedBindPixelUnpackBuffer(GL_PIXEL_UNPACK_BUFFER, 0);
 }
 
 
diff --git a/Engine/Source/Runtime/OpenGLDrv/Private/OpenGLUAV.cpp b/Engine/Source/Runtime/OpenGLDrv/Private/OpenGLUAV.cpp
--- a/Engine/Source/Runtime/OpenGLDrv/Private/OpenGLUAV.cpp
+++ b/Engine/Source/Runtime/OpenGLDrv/Private/OpenGLUAV.cpp
@@ -7,10 +7,10 @@
 #include "OpenGLDrvPrivate.h"
 #include "ClearReplacementShaders.h"
 
-FShaderResourceViewRHIRef FOpenGLDynamicRHI::RHICreateShaderResourceView(FRHIVertexBuffer* VertexBufferRHI, uint32 Stride, uint8 Format)
+FShaderResourceViewRHIRef FOpenGLDynamicRHI::RHICreateShaderResourceView(FRHIBuffer* BufferRHI, uint32 Stride, uint8 Format)
 {
 	ensureMsgf(Stride == GPixelFormats[Format].BlockBytes, TEXT("provided stride: %i was not consitent with Pixelformat: %s"), Stride, GPixelFormats[Format].Name);
-	return FOpenGLDynamicRHI::RHICreateShaderResourceView(FShaderResourceViewInitializer(VertexBufferRHI, EPixelFormat(Format)));
+	return FOpenGLDynamicRHI::RHICreateShaderResourceView(FShaderResourceViewInitializer(BufferRHI, EPixelFormat(Format)));
 }
 
 // Binds the specified buffer range to a texture resource and selects glTexBuffer or glTexBufferRange 
@@ -37,12 +37,13 @@ static void BindGLTexBufferRange(GLenum Target, GLenum InternalFormat, GLuint Bu
 
 FShaderResourceViewRHIRef FOpenGLDynamicRHI::RHICreateShaderResourceView(const FShaderResourceViewInitializer& Initializer)
 {
+	FShaderResourceViewInitializer::FBufferShaderResourceViewInitializer Desc = Initializer.AsBufferSRV();
+	FRHIBuffer* BufferRHI = Desc.Buffer;
+
 	switch (Initializer.GetType())
 	{
 		case FShaderResourceViewInitializer::EType::VertexBufferSRV:
 		{
-			FShaderResourceViewInitializer::FVertexBufferShaderResourceViewInitializer Desc = Initializer.AsVertexBufferSRV();
-			FRHIVertexBuffer* VertexBufferRHI = Desc.VertexBuffer;
 			const uint8 Format = Desc.Format;
 
 			FShaderResourceViewRHIRef Result = new FOpenGLShaderResourceViewProxy([=, OGLRHI = this](FRHIShaderResourceView* OwnerRHI)
@@ -53,9 +54,9 @@ FShaderResourceViewRHIRef FOpenGLDynamicRHI::RHICreateShaderResourceView(const F
 				{
 					FOpenGL::GenTextures(1, &TextureID);
 					UE_CLOG(!GPixelFormats[Format].Supported, LogRHI, Error, TEXT("Unsupported EPixelFormat %d"), Format);
-					if (VertexBufferRHI)
+					if (BufferRHI)
 					{
-						FOpenGLVertexBuffer* VertexBuffer = FOpenGLDynamicRHI::ResourceCast(VertexBufferRHI);
+						FOpenGLBuffer* VertexBuffer = FOpenGLDynamicRHI::ResourceCast(BufferRHI);
 
 						const uint32 FormatBPP = GPixelFormats[Format].BlockBytes;
 
@@ -70,7 +71,7 @@ FShaderResourceViewRHIRef FOpenGLDynamicRHI::RHICreateShaderResourceView(const F
 				// and the next draw will take care of cleaning it up; or
 				// next operation that needs the stage will switch something else in on it.
 
-				return new FOpenGLShaderResourceView(OGLRHI, TextureID, GL_TEXTURE_BUFFER, VertexBufferRHI, Format);
+				return new FOpenGLShaderResourceView(OGLRHI, TextureID, GL_TEXTURE_BUFFER, BufferRHI, Format);
 			});
 
 			return Result;
@@ -78,16 +79,13 @@ FShaderResourceViewRHIRef FOpenGLDynamicRHI::RHICreateShaderResourceView(const F
 
 		case FShaderResourceViewInitializer::EType::StructuredBufferSRV:
 		{
-			FShaderResourceViewInitializer::FStructuredBufferShaderResourceViewInitializer Desc = Initializer.AsStructuredBufferSRV();
-			FRHIStructuredBuffer* StructuredBufferRHI = Desc.StructuredBuffer;
-
 			return new FOpenGLShaderResourceViewProxy([=](FRHIShaderResourceView* OwnerRHI)
 			{
 				VERIFY_GL_SCOPE();
 				GLuint TextureID = 0;
 				if (FOpenGL::SupportsResourceView())
 				{
-					FOpenGLStructuredBuffer* StructuredBuffer = ResourceCast(StructuredBufferRHI);
+					FOpenGLBuffer* StructuredBuffer = ResourceCast(BufferRHI);
 					FOpenGL::GenTextures(1, &TextureID);
 					CachedSetupTextureStage(GetContextStateForCurrentContext(), FOpenGL::GetMaxCombinedTextureImageUnits() - 1, GL_TEXTURE_BUFFER, TextureID, -1, 1);
 					uint32 Stride = StructuredBuffer->GetStride();
@@ -101,9 +99,6 @@ FShaderResourceViewRHIRef FOpenGLDynamicRHI::RHICreateShaderResourceView(const F
 
 		case FShaderResourceViewInitializer::EType::IndexBufferSRV:
 		{
-			FShaderResourceViewInitializer::FIndexBufferShaderResourceViewInitializer Desc = Initializer.AsIndexBufferSRV();
-			FRHIIndexBuffer* IndexBufferRHI = Desc.IndexBuffer;
-
 			return new FOpenGLShaderResourceViewProxy([=](FRHIShaderResourceView* OwnerRHI)
 			{
 				VERIFY_GL_SCOPE();
@@ -111,16 +106,16 @@ FShaderResourceViewRHIRef FOpenGLDynamicRHI::RHICreateShaderResourceView(const F
 				if (FOpenGL::SupportsResourceView())
 				{
 					FOpenGL::GenTextures(1, &TextureID);
-					if (IndexBufferRHI)
+					if (BufferRHI)
 					{
-						FOpenGLIndexBuffer* IndexBuffer = ResourceCast(IndexBufferRHI);
+						FOpenGLBuffer* IndexBuffer = ResourceCast(BufferRHI);
 						CachedSetupTextureStage(GetContextStateForCurrentContext(), FOpenGL::GetMaxCombinedTextureImageUnits() - 1, GL_TEXTURE_BUFFER, TextureID, -1, 1);
-						uint32 Stride = IndexBufferRHI->GetStride();
+						uint32 Stride = BufferRHI->GetStride();
 						GLenum Format = (Stride == 2) ? GL_R16UI : GL_R32UI;
 						BindGLTexBufferRange(GL_TEXTURE_BUFFER, Format, IndexBuffer->Resource, Desc.StartOffsetBytes, Desc.NumElements, Stride);
 					}
 				}
-				return new FOpenGLShaderResourceView(this, TextureID, GL_TEXTURE_BUFFER, IndexBufferRHI);
+				return new FOpenGLShaderResourceView(this, TextureID, GL_TEXTURE_BUFFER, BufferRHI);
 			});
 		}
 
@@ -141,7 +136,7 @@ void FOpenGLDynamicRHI::RHIUpdateShaderResourceView(FRHIShaderResourceView* SRV,
 	VERIFY_GL_SCOPE();
 
 	FOpenGLShaderResourceView* SRVGL = FOpenGLDynamicRHI::ResourceCast(SRV);
-	FOpenGLIndexBuffer* IBGL = FOpenGLDynamicRHI::ResourceCast(IndexBuffer);
+	FOpenGLBuffer* IBGL = FOpenGLDynamicRHI::ResourceCast(IndexBuffer);
 
 	check(SRVGL);
 	check(!SRVGL->VertexBuffer);
@@ -176,7 +171,7 @@ void FOpenGLDynamicRHI::RHIUpdateShaderResourceView(FRHIShaderResourceView* SRV,
 	VERIFY_GL_SCOPE();
 
 	FOpenGLShaderResourceView* SRVGL = FOpenGLDynamicRHI::ResourceCast(SRV);
-	FOpenGLVertexBuffer* VBGL = FOpenGLDynamicRHI::ResourceCast(VertexBuffer);
+	FOpenGLBuffer* VBGL = FOpenGLDynamicRHI::ResourceCast(VertexBuffer);
 	const FOpenGLTextureFormat& GLFormat = GOpenGLTextureFormats[Format];
 	
 	check(SRVGL);
@@ -199,7 +194,7 @@ void FOpenGLDynamicRHI::RHIUpdateShaderResourceView(FRHIShaderResourceView* SRV,
 	}
 }
 
-FShaderResourceViewRHIRef FOpenGLDynamicRHI::RHICreateShaderResourceView(FRHIIndexBuffer* BufferRHI)
+FShaderResourceViewRHIRef FOpenGLDynamicRHI::RHICreateShaderResourceView(FRHIBuffer* BufferRHI)
 {
 	return FOpenGLDynamicRHI::RHICreateShaderResourceView(FShaderResourceViewInitializer(BufferRHI));
 }
@@ -217,22 +212,22 @@ FOpenGLShaderResourceView::~FOpenGLShaderResourceView()
 	}
 }
 
-FUnorderedAccessViewRHIRef FOpenGLDynamicRHI::RHICreateUnorderedAccessView(FRHIStructuredBuffer* StructuredBufferRHI, bool bUseUAVCounter, bool bAppendBuffer)
+FUnorderedAccessViewRHIRef FOpenGLDynamicRHI::RHICreateUnorderedAccessView(FRHIBuffer* BufferRHI, bool bUseUAVCounter, bool bAppendBuffer)
 {
-	FOpenGLStructuredBuffer* StructuredBuffer = ResourceCast(StructuredBufferRHI);
+	FOpenGLBuffer* Buffer = ResourceCast(BufferRHI);
 	// emulate structured buffer of specific size as typed buffer
 	if (GMaxRHIFeatureLevel == ERHIFeatureLevel::ES3_1)
 	{
 		// ES3.1 cross-compiler converts StructuredBuffer<type4> into Buffer<type4> and StructuredBuffer<type> into Buffer<type>
 		// type can be float, int and uint
-		check(StructuredBuffer->GetStride() == 16 || StructuredBuffer->GetStride() == 4);
-		if (StructuredBuffer->GetStride() == 16)
+		check(Buffer->GetStride() == 16 || Buffer->GetStride() == 4);
+		if (Buffer->GetStride() == 16)
 		{
-			return new FOpenGLStructuredBufferUnorderedAccessView(this, StructuredBufferRHI, PF_A32B32G32R32F);
+			return new FOpenGLStructuredBufferUnorderedAccessView(this, BufferRHI, PF_A32B32G32R32F);
 		}
-		else if (StructuredBuffer->GetStride() == 4)
+		else if (Buffer->GetStride() == 4)
 		{
-			return new FOpenGLStructuredBufferUnorderedAccessView(this, StructuredBufferRHI, PF_R32_FLOAT);
+			return new FOpenGLStructuredBufferUnorderedAccessView(this, BufferRHI, PF_R32_FLOAT);
 		}
 	}
 	
@@ -268,7 +263,7 @@ FOpenGLVertexBufferUnorderedAccessView::FOpenGLVertexBufferUnorderedAccessView(
 	OpenGLRHI(InOpenGLRHI)
 {
 	VERIFY_GL_SCOPE();
-	FOpenGLVertexBuffer* InVertexBuffer = FOpenGLDynamicRHI::ResourceCast(InVertexBufferRHI);
+	FOpenGLBuffer* InVertexBuffer = FOpenGLDynamicRHI::ResourceCast(InVertexBufferRHI);
 
 
 	const FOpenGLTextureFormat& GLFormat = GOpenGLTextureFormats[Format];
@@ -293,7 +288,7 @@ FOpenGLVertexBufferUnorderedAccessView::FOpenGLVertexBufferUnorderedAccessView(
 
 uint32 FOpenGLVertexBufferUnorderedAccessView::GetBufferSize()
 {
-	FOpenGLVertexBuffer* VertexBuffer = FOpenGLDynamicRHI::ResourceCast(VertexBufferRHI.GetReference());
+	FOpenGLBuffer* VertexBuffer = FOpenGLDynamicRHI::ResourceCast(VertexBufferRHI.GetReference());
 	return VertexBufferRHI->GetSize();
 }
 
@@ -311,21 +306,9 @@ FOpenGLVertexBufferUnorderedAccessView::~FOpenGLVertexBufferUnorderedAccessView(
 }
 
 
-FUnorderedAccessViewRHIRef FOpenGLDynamicRHI::RHICreateUnorderedAccessView(FRHIVertexBuffer* VertexBufferRHI,uint8 Format)
-{
-	FOpenGLVertexBuffer* VertexBuffer = ResourceCast(VertexBufferRHI);
-	return new FOpenGLVertexBufferUnorderedAccessView(this, VertexBufferRHI, Format);
-}
-
-FUnorderedAccessViewRHIRef FOpenGLDynamicRHI::RHICreateUnorderedAccessView(FRHIIndexBuffer* IndexBufferRHI, uint8 Format)
-{
-	checkf(0, TEXT("Not implemented!"));
-	return nullptr;
-}
-
-FShaderResourceViewRHIRef FOpenGLDynamicRHI::RHICreateShaderResourceView(FRHIStructuredBuffer* StructuredBufferRHI)
+FUnorderedAccessViewRHIRef FOpenGLDynamicRHI::RHICreateUnorderedAccessView(FRHIBuffer* BufferRHI, uint8 Format)
 {
-	return FOpenGLDynamicRHI::RHICreateShaderResourceView(FShaderResourceViewInitializer(StructuredBufferRHI));
+	return new FOpenGLVertexBufferUnorderedAccessView(this, BufferRHI, Format);
 }
 
 void FOpenGLDynamicRHI::RHIClearUAVFloat(FRHIUnorderedAccessView* UnorderedAccessViewRHI, const FVector4& Values)
@@ -449,7 +432,7 @@ FOpenGLStructuredBufferUnorderedAccessView::FOpenGLStructuredBufferUnorderedAcce
 	, OpenGLRHI(InOpenGLRHI)
 {
 	VERIFY_GL_SCOPE();
-	FOpenGLStructuredBuffer* InStructuredBuffer = FOpenGLDynamicRHI::ResourceCast(InStructuredBufferRHI);
+	FOpenGLBuffer* InStructuredBuffer = FOpenGLDynamicRHI::ResourceCast(InStructuredBufferRHI);
 	const FOpenGLTextureFormat& GLFormat = GOpenGLTextureFormats[InFormat];
 
 	GLuint TextureID = 0;
@@ -471,7 +454,7 @@ FOpenGLStructuredBufferUnorderedAccessView::FOpenGLStructuredBufferUnorderedAcce
 
 uint32 FOpenGLStructuredBufferUnorderedAccessView::GetBufferSize()
 {
-	FOpenGLStructuredBuffer* StructuredBuffer = FOpenGLDynamicRHI::ResourceCast(StructuredBufferRHI.GetReference());
+	FOpenGLBuffer* StructuredBuffer = FOpenGLDynamicRHI::ResourceCast(StructuredBufferRHI.GetReference());
 	return StructuredBuffer->GetSize();
 }
 
diff --git a/Engine/Source/Runtime/OpenGLDrv/Private/OpenGLVertexBuffer.cpp b/Engine/Source/Runtime/OpenGLDrv/Private/OpenGLVertexBuffer.cpp
--- a/Engine/Source/Runtime/OpenGLDrv/Private/OpenGLVertexBuffer.cpp
+++ b/Engine/Source/Runtime/OpenGLDrv/Private/OpenGLVertexBuffer.cpp
@@ -129,7 +129,7 @@ void* GetAllocation( void* Target, uint32 Size, uint32 Offset, uint32 Alignment
 	return nullptr;
 }
 
-bool RetireAllocation( FOpenGLVertexBuffer* Target)
+bool RetireAllocation( FOpenGLBuffer* Target)
 {
 	if (FOpenGL::SupportsBufferStorage() && OpenGLConsoleVariables::bUseStagingBuffer)
 	{
@@ -180,7 +180,7 @@ FVertexBufferRHIRef FOpenGLDynamicRHI::RHICreateVertexBuffer(uint32 Size, uint32
 {
 	if (CreateInfo.bWithoutNativeResource)
 	{
-		return new FOpenGLVertexBuffer();
+		return new FOpenGLBuffer();
 	}
 
 	const void *Data = NULL;
@@ -192,7 +192,7 @@ FVertexBufferRHIRef FOpenGLDynamicRHI::RHICreateVertexBuffer(uint32 Size, uint32
 		Data = CreateInfo.ResourceArray->GetResourceData();
 	}
 
-	TRefCountPtr<FOpenGLVertexBuffer> VertexBuffer = new FOpenGLVertexBuffer(0, Size, InUsage, Data);
+	TRefCountPtr<FOpenGLBuffer> VertexBuffer = new FOpenGLBuffer(GL_ARRAY_BUFFER, 0, Size, InUsage | BUF_VertexBuffer, Data);
 	
 	if (CreateInfo.ResourceArray)
 	{
@@ -208,7 +208,7 @@ void* FOpenGLDynamicRHI::LockVertexBuffer_BottomOfPipe(FRHICommandListImmediate&
 	RHITHREAD_GLCOMMAND_PROLOGUE();
 
 	VERIFY_GL_SCOPE();
-	FOpenGLVertexBuffer* VertexBuffer = ResourceCast(VertexBufferRHI);
+	FOpenGLBuffer* VertexBuffer = ResourceCast(VertexBufferRHI);
 	if( !(FOpenGL::SupportsVertexAttribBinding() && OpenGLConsoleVariables::bUseVAB) && VertexBuffer->GetUsage() & BUF_ZeroStride )
 	{
 		check( Offset + Size <= VertexBuffer->GetSize() );
@@ -235,7 +235,7 @@ void FOpenGLDynamicRHI::UnlockVertexBuffer_BottomOfPipe(FRHICommandListImmediate
 {
 	RHITHREAD_GLCOMMAND_PROLOGUE();
 	VERIFY_GL_SCOPE();
-	FOpenGLVertexBuffer* VertexBuffer = ResourceCast(VertexBufferRHI);
+	FOpenGLBuffer* VertexBuffer = ResourceCast(VertexBufferRHI);
 	if( (FOpenGL::SupportsVertexAttribBinding() && OpenGLConsoleVariables::bUseVAB) || !( VertexBuffer->GetUsage() & BUF_ZeroStride ) )
 	{
 		if (!RetireAllocation(VertexBuffer))
@@ -252,28 +252,11 @@ void FOpenGLDynamicRHI::RHICopyVertexBuffer(FRHIVertexBuffer* SourceBufferRHI, F
 	RHICopyBufferRegion(DestBufferRHI, 0, SourceBufferRHI, 0, SourceBufferRHI->GetSize());
 }
 
-void FOpenGLDynamicRHI::RHITransferVertexBufferUnderlyingResource(FRHIVertexBuffer* DestVertexBuffer, FRHIVertexBuffer* SrcVertexBuffer)
-{
-	VERIFY_GL_SCOPE();
-	check(DestVertexBuffer);
-	FOpenGLVertexBuffer* Dest = ResourceCast(DestVertexBuffer);
-	if (!SrcVertexBuffer)
-	{
-		TRefCountPtr<FOpenGLVertexBuffer> Src = new FOpenGLVertexBuffer();
-		Dest->Swap(*Src);
-	}
-	else
-	{
-		FOpenGLVertexBuffer* Src = ResourceCast(SrcVertexBuffer);
-		Dest->Swap(*Src);
-	}
-}
-
 void FOpenGLDynamicRHI::RHICopyBufferRegion(FRHIVertexBuffer* DestBufferRHI, uint64 DstOffset, FRHIVertexBuffer* SourceBufferRHI, uint64 SrcOffset, uint64 NumBytes)
 {
 	VERIFY_GL_SCOPE();
-	FOpenGLVertexBuffer* SourceBuffer = ResourceCast(SourceBufferRHI);
-	FOpenGLVertexBuffer* DestBuffer = ResourceCast(DestBufferRHI);
+	FOpenGLBuffer* SourceBuffer = ResourceCast(SourceBufferRHI);
+	FOpenGLBuffer* DestBuffer = ResourceCast(DestBufferRHI);
 
 	glBindBuffer(GL_COPY_READ_BUFFER, SourceBuffer->Resource);
 	glBindBuffer(GL_COPY_WRITE_BUFFER, DestBuffer->Resource);
diff --git a/Engine/Source/Runtime/OpenGLDrv/Public/OpenGLDrv.h b/Engine/Source/Runtime/OpenGLDrv/Public/OpenGLDrv.h
--- a/Engine/Source/Runtime/OpenGLDrv/Public/OpenGLDrv.h
+++ b/Engine/Source/Runtime/OpenGLDrv/Public/OpenGLDrv.h
@@ -353,19 +353,16 @@ public:
 	virtual void* LockVertexBuffer_BottomOfPipe(FRHICommandListImmediate& RHICmdList, FRHIVertexBuffer* VertexBuffer, uint32 Offset, uint32 SizeRHI, EResourceLockMode LockMode) final override;
 	virtual void UnlockVertexBuffer_BottomOfPipe(FRHICommandListImmediate& RHICmdList, FRHIVertexBuffer* VertexBuffer) final override;
 	virtual void RHICopyVertexBuffer(FRHIVertexBuffer* SourceBuffer, FRHIVertexBuffer* DestBuffer) final override;
-	virtual void RHITransferVertexBufferUnderlyingResource(FRHIVertexBuffer* DestVertexBuffer, FRHIVertexBuffer* SrcVertexBuffer) final override;
-	virtual void RHITransferIndexBufferUnderlyingResource(FRHIIndexBuffer* DestIndexBuffer, FRHIIndexBuffer* SrcIndexBuffer) final override;
+	virtual void RHITransferBufferUnderlyingResource(FRHIBuffer* DestBuffer, FRHIBuffer* SrcBuffer) final override;
 	virtual FStructuredBufferRHIRef RHICreateStructuredBuffer(uint32 Stride, uint32 Size, uint32 InUsage, ERHIAccess InResourceState, FRHIResourceCreateInfo& CreateInfo) final override;
 	virtual void* LockStructuredBuffer_BottomOfPipe(FRHICommandListImmediate& RHICmdList, FRHIStructuredBuffer* StructuredBuffer, uint32 Offset, uint32 SizeRHI, EResourceLockMode LockMode) final override;
 	virtual void UnlockStructuredBuffer_BottomOfPipe(FRHICommandListImmediate& RHICmdList, FRHIStructuredBuffer* StructuredBuffer) final override;
-	virtual FUnorderedAccessViewRHIRef RHICreateUnorderedAccessView(FRHIStructuredBuffer* StructuredBuffer, bool bUseUAVCounter, bool bAppendBuffer) final override;
+	virtual FUnorderedAccessViewRHIRef RHICreateUnorderedAccessView(FRHIBuffer* Buffer, bool bUseUAVCounter, bool bAppendBuffer) final override;
 	virtual FUnorderedAccessViewRHIRef RHICreateUnorderedAccessView(FRHITexture* Texture, uint32 MipLevel) final override;
-	virtual FUnorderedAccessViewRHIRef RHICreateUnorderedAccessView(FRHIVertexBuffer* VertexBuffer, uint8 Format) final override;
-	virtual FUnorderedAccessViewRHIRef RHICreateUnorderedAccessView(FRHIIndexBuffer* IndexBuffer, uint8 Format) final override;
-	virtual FShaderResourceViewRHIRef RHICreateShaderResourceView(FRHIStructuredBuffer* StructuredBuffer) final override;
-	virtual FShaderResourceViewRHIRef RHICreateShaderResourceView(FRHIVertexBuffer* VertexBuffer, uint32 Stride, uint8 Format) final override;
+	virtual FUnorderedAccessViewRHIRef RHICreateUnorderedAccessView(FRHIBuffer* Buffer, uint8 Format) final override;
+	virtual FShaderResourceViewRHIRef RHICreateShaderResourceView(FRHIBuffer* Buffer) final override;
+	virtual FShaderResourceViewRHIRef RHICreateShaderResourceView(FRHIBuffer* Buffer, uint32 Stride, uint8 Format) final override;
 	virtual FShaderResourceViewRHIRef RHICreateShaderResourceView(const FShaderResourceViewInitializer& Initializer) final override;
-	virtual FShaderResourceViewRHIRef RHICreateShaderResourceView(FRHIIndexBuffer* Buffer) final override;
 	virtual void RHIUpdateShaderResourceView(FRHIShaderResourceView* SRV, FRHIVertexBuffer* VertexBuffer, uint32 Stride, uint8 Format) final override;
 	virtual void RHIUpdateShaderResourceView(FRHIShaderResourceView* SRV, FRHIIndexBuffer* IndexBuffer) final override;
 	virtual uint64 RHICalcTexture2DPlatformSize(uint32 SizeX, uint32 SizeY, uint8 Format, uint32 NumMips, uint32 NumSamples, ETextureCreateFlags Flags, const FRHIResourceCreateInfo& CreateInfo, uint32& OutAlign) final override;
@@ -445,7 +442,7 @@ public:
 
 	virtual void RHISetComputeShader(FRHIComputeShader* ComputeShader) final override;
 	virtual void RHIDispatchComputeShader(uint32 ThreadGroupCountX, uint32 ThreadGroupCountY, uint32 ThreadGroupCountZ) final override;
-	virtual void RHIDispatchIndirectComputeShader(FRHIVertexBuffer* ArgumentBuffer, uint32 ArgumentOffset) final override;
+	virtual void RHIDispatchIndirectComputeShader(FRHIBuffer* ArgumentBuffer, uint32 ArgumentOffset) final override;
 	virtual void RHISetMultipleViewports(uint32 Count, const FViewportBounds* Data) final override;
 	virtual void RHIClearUAVFloat(FRHIUnorderedAccessView* UnorderedAccessViewRHI, const FVector4& Values) final override;
 	virtual void RHIClearUAVUint(FRHIUnorderedAccessView* UnorderedAccessViewRHI, const FUintVector4& Values) final override;
@@ -490,10 +487,10 @@ public:
 	void SetRenderTargetsAndClear(const FRHISetRenderTargetsInfo& RenderTargetsInfo);
 
 	virtual void RHIDrawPrimitive(uint32 BaseVertexIndex, uint32 NumPrimitives, uint32 NumInstances) final override;
-	virtual void RHIDrawPrimitiveIndirect(FRHIVertexBuffer* ArgumentBuffer, uint32 ArgumentOffset) final override;
-	virtual void RHIDrawIndexedIndirect(FRHIIndexBuffer* IndexBufferRHI, FRHIStructuredBuffer* ArgumentsBufferRHI, int32 DrawArgumentsIndex, uint32 NumInstances) final override;
-	virtual void RHIDrawIndexedPrimitive(FRHIIndexBuffer* IndexBuffer, int32 BaseVertexIndex, uint32 FirstInstance, uint32 NumVertices, uint32 StartIndex, uint32 NumPrimitives, uint32 NumInstances) final override;
-	virtual void RHIDrawIndexedPrimitiveIndirect(FRHIIndexBuffer* IndexBuffer, FRHIVertexBuffer* ArgumentBuffer, uint32 ArgumentOffset) final override;
+	virtual void RHIDrawPrimitiveIndirect(FRHIBuffer* ArgumentBuffer, uint32 ArgumentOffset) final override;
+	virtual void RHIDrawIndexedIndirect(FRHIBuffer* IndexBufferRHI, FRHIBuffer* ArgumentsBufferRHI, int32 DrawArgumentsIndex, uint32 NumInstances) final override;
+	virtual void RHIDrawIndexedPrimitive(FRHIBuffer* IndexBuffer, int32 BaseVertexIndex, uint32 FirstInstance, uint32 NumVertices, uint32 StartIndex, uint32 NumPrimitives, uint32 NumInstances) final override;
+	virtual void RHIDrawIndexedPrimitiveIndirect(FRHIBuffer* IndexBuffer, FRHIBuffer* ArgumentBuffer, uint32 ArgumentOffset) final override;
 	virtual void RHIEnableDepthBoundsTest(bool bEnable) final override;
 	virtual void RHISetDepthBounds(float MinDepth, float MaxDepth) final override;
 	virtual void RHIPushEvent(const TCHAR* Name, FColor Color) final override;
@@ -532,8 +529,7 @@ public:
 	void Cleanup();
 
 	void PurgeFramebufferFromCaches(GLuint Framebuffer);
-	void OnVertexBufferDeletion(GLuint VertexBufferResource);
-	void OnIndexBufferDeletion(GLuint IndexBufferResource);
+	void OnBufferDeletion(GLuint VertexBufferResource);
 	void OnPixelBufferDeletion(GLuint PixelBufferResource);
 	void OnUniformBufferDeletion(GLuint UniformBufferResource,uint32 AllocatedSize,bool bStreamDraw);
 	void OnProgramDeletion(GLint ProgramResource);
@@ -841,10 +837,10 @@ public:
 		RHITHREAD_GLCOMMAND_EPILOGUE_RETURN(FTexture3DRHIRef);
 	}
 
-	virtual FUnorderedAccessViewRHIRef RHICreateUnorderedAccessView_RenderThread(class FRHICommandListImmediate& RHICmdList, FRHIStructuredBuffer* StructuredBuffer, bool bUseUAVCounter, bool bAppendBuffer) final override
+	virtual FUnorderedAccessViewRHIRef RHICreateUnorderedAccessView_RenderThread(class FRHICommandListImmediate& RHICmdList, FRHIBuffer* Buffer, bool bUseUAVCounter, bool bAppendBuffer) final override
 	{
 		RHITHREAD_GLCOMMAND_PROLOGUE();
-		return this->RHICreateUnorderedAccessView(StructuredBuffer, bUseUAVCounter, bAppendBuffer);
+		return this->RHICreateUnorderedAccessView(Buffer, bUseUAVCounter, bAppendBuffer);
 		RHITHREAD_GLCOMMAND_EPILOGUE_RETURN(FUnorderedAccessViewRHIRef);
 	}
 
@@ -855,10 +851,10 @@ public:
 		RHITHREAD_GLCOMMAND_EPILOGUE_RETURN(FUnorderedAccessViewRHIRef);
 	}
 
-	virtual FUnorderedAccessViewRHIRef RHICreateUnorderedAccessView_RenderThread(class FRHICommandListImmediate& RHICmdList, FRHIVertexBuffer* VertexBuffer, uint8 Format) final override
+	virtual FUnorderedAccessViewRHIRef RHICreateUnorderedAccessView_RenderThread(class FRHICommandListImmediate& RHICmdList, FRHIBuffer* Buffer, uint8 Format) final override
 	{
 		RHITHREAD_GLCOMMAND_PROLOGUE();
-		return this->RHICreateUnorderedAccessView(VertexBuffer, Format);
+		return this->RHICreateUnorderedAccessView(Buffer, Format);
 		RHITHREAD_GLCOMMAND_EPILOGUE_RETURN(FUnorderedAccessViewRHIRef);
 	}
 
@@ -867,21 +863,16 @@ public:
 		return this->RHICreateShaderResourceView(Texture, CreateInfo);
 	}
 
-	virtual FShaderResourceViewRHIRef RHICreateShaderResourceView_RenderThread(class FRHICommandListImmediate& RHICmdList, FRHIVertexBuffer* VertexBuffer, uint32 Stride, uint8 Format) final override
+	virtual FShaderResourceViewRHIRef RHICreateShaderResourceView_RenderThread(class FRHICommandListImmediate& RHICmdList, FRHIBuffer* Buffer, uint32 Stride, uint8 Format) final override
 	{
-		return this->RHICreateShaderResourceView(VertexBuffer, Stride, Format);
+		return this->RHICreateShaderResourceView(Buffer, Stride, Format);
 	}
 
-	virtual FShaderResourceViewRHIRef RHICreateShaderResourceView_RenderThread(class FRHICommandListImmediate& RHICmdList, FRHIIndexBuffer* Buffer) final override
+	virtual FShaderResourceViewRHIRef RHICreateShaderResourceView_RenderThread(class FRHICommandListImmediate& RHICmdList, FRHIBuffer* Buffer) final override
 	{
 		return this->RHICreateShaderResourceView(Buffer);
 	}
 
-	virtual FShaderResourceViewRHIRef RHICreateShaderResourceView_RenderThread(class FRHICommandListImmediate& RHICmdList, FRHIStructuredBuffer* StructuredBuffer) final override
-	{
-		return this->RHICreateShaderResourceView(StructuredBuffer);
-	}
-
 	virtual FTextureCubeRHIRef RHICreateTextureCube_RenderThread(class FRHICommandListImmediate& RHICmdList, uint32 Size, uint8 Format, uint32 NumMips, ETextureCreateFlags Flags, ERHIAccess InResourceState, FRHIResourceCreateInfo& CreateInfo) final override
 	{
 		const bool bCubeTexture = true;
@@ -1064,8 +1055,6 @@ private:
 
 	/** RHI device state, independent of underlying OpenGL context used */
 	FOpenGLRHIState						PendingState;
-	FOpenGLStreamedVertexBufferArray	DynamicVertexBuffers;
-	FOpenGLStreamedIndexBufferArray		DynamicIndexBuffers;
 	FSamplerStateRHIRef					PointSamplerState;
 
 	/** A list of all viewport RHIs that have been created. */
@@ -1126,7 +1115,7 @@ private:
 	
 
 	void EnableVertexElementCached(FOpenGLContextState& ContextCache, GLuint AttributeIndex, const FOpenGLVertexElement &VertexElement, GLsizei Stride, void *Pointer, GLuint Buffer);
-	void EnableVertexElementCachedZeroStride(FOpenGLContextState& ContextCache, GLuint AttributeIndex, const FOpenGLVertexElement &VertexElement, uint32 NumVertices, FOpenGLVertexBuffer* VertexBuffer);
+	void EnableVertexElementCachedZeroStride(FOpenGLContextState& ContextCache, GLuint AttributeIndex, const FOpenGLVertexElement &VertexElement, uint32 NumVertices, FOpenGLBuffer* VertexBuffer);
 	void SetupVertexArrays(FOpenGLContextState& ContextCache, uint32 BaseVertexIndex, FOpenGLStream* Streams, uint32 NumStreams, uint32 MaxVertices);
 	void SetupVertexArraysVAB(FOpenGLContextState& ContextCache, uint32 BaseVertexIndex, FOpenGLStream* Streams, uint32 NumStreams, uint32 MaxVertices);
 	void SetupVertexArraysUP(FOpenGLContextState& ContextState, void* Buffer, uint32 Stride);
