- Add option to use RHITransitions to drive all resource transitions and use internal transitions for validations only (disabled by default on all platforms - still dependent on high level fixes first)
- Initialize textures and data buffers using the correct initial resource state
- Validate bound resources in uniform buffers when rhi validation is enabled

diff --git a/Engine/Source/Runtime/D3D12RHI/Private/D3D12Adapter.h b/Engine/Source/Runtime/D3D12RHI/Private/D3D12Adapter.h
--- a/Engine/Source/Runtime/D3D12RHI/Private/D3D12Adapter.h
+++ b/Engine/Source/Runtime/D3D12RHI/Private/D3D12Adapter.h
@@ -262,6 +262,7 @@ public:
 		const D3D12_RESOURCE_DESC& Desc,
 		uint32 Alignment, uint32 Stride, uint32 Size, uint32 InUsage,
 		ED3D12ResourceStateMode InResourceStateMode,
+		ERHIAccess InResourceState,
 		FRHIResourceCreateInfo& CreateInfo);
 
 	template <typename ObjectType, typename CreationCoreFunction>
@@ -312,6 +313,7 @@ protected:
 		uint32 Size,
 		uint32 InUsage,
 		ED3D12ResourceStateMode InResourceStateMode,
+		D3D12_RESOURCE_STATES InCreateState,
 		FRHIResourceCreateInfo& CreateInfo,
 		uint32 Alignment,
 		FD3D12TransientResource& TransientResource,
diff --git a/Engine/Source/Runtime/D3D12RHI/Private/D3D12Allocation.cpp b/Engine/Source/Runtime/D3D12RHI/Private/D3D12Allocation.cpp
--- a/Engine/Source/Runtime/D3D12RHI/Private/D3D12Allocation.cpp
+++ b/Engine/Source/Runtime/D3D12RHI/Private/D3D12Allocation.cpp
@@ -1071,7 +1071,7 @@ void FD3D12DefaultBufferPool::CleanUpAllocations(uint64 FrameLag)
 }
 
 // Grab a buffer from the available buffers or create a new buffer if none are available
-void FD3D12DefaultBufferPool::AllocDefaultResource(D3D12_HEAP_TYPE InHeapType, const D3D12_RESOURCE_DESC& Desc, EBufferUsageFlags InUsage, ED3D12ResourceStateMode InResourceStateMode, FD3D12ResourceLocation& ResourceLocation, uint32 Alignment, const TCHAR* Name)
+void FD3D12DefaultBufferPool::AllocDefaultResource(D3D12_HEAP_TYPE InHeapType, const D3D12_RESOURCE_DESC& Desc, EBufferUsageFlags InUsage, ED3D12ResourceStateMode InResourceStateMode, D3D12_RESOURCE_STATES InCreateState, FD3D12ResourceLocation& ResourceLocation, uint32 Alignment, const TCHAR* Name)
 {
 	FD3D12Device* Device = GetParentDevice();
 	FD3D12Adapter* Adapter = Device->GetParentAdapter();
@@ -1084,17 +1084,23 @@ void FD3D12DefaultBufferPool::AllocDefaultResource(D3D12_HEAP_TYPE InHeapType, c
 		return;
 	}
 
-#if D3D12_RHI_RAYTRACING
-	// RayTracing acceleration structures must be created in a particular state and may never transition out of it.
-	D3D12_RESOURCE_STATES InitialState = (InUsage & BUF_AccelerationStructure) ? D3D12_RESOURCE_STATE_RAYTRACING_ACCELERATION_STRUCTURE : D3D12_RESOURCE_STATE_GENERIC_READ;
-#else
-	D3D12_RESOURCE_STATES InitialState = D3D12_RESOURCE_STATE_GENERIC_READ;
-#endif
-
+#if DO_CHECK
+	// Validate the create state
 	if (InHeapType == D3D12_HEAP_TYPE_READBACK)
 	{
-		InitialState = D3D12_RESOURCE_STATE_COPY_DEST;
+		check(InCreateState == D3D12_RESOURCE_STATE_COPY_DEST);
+	}
+	else if (InHeapType == D3D12_HEAP_TYPE_UPLOAD)
+	{
+		check(InCreateState == D3D12_RESOURCE_STATE_GENERIC_READ);
+	}
+#if D3D12_RHI_RAYTRACING
+	else if (InUsage & BUF_AccelerationStructure)
+	{
+		check(InCreateState == D3D12_RESOURCE_STATE_RAYTRACING_ACCELERATION_STRUCTURE);
 	}
+#endif // D3D12_RHI_RAYTRACING
+#endif  // DO_CHECK
 
 	const bool PoolResource = Desc.Width < Allocator->GetMaximumAllocationSizeForPooling()/* && ((Desc.Width % (1024 * 64)) != 0)*/;
 
@@ -1128,7 +1134,7 @@ void FD3D12DefaultBufferPool::AllocDefaultResource(D3D12_HEAP_TYPE InHeapType, c
 				uint64 HeapOffset = ResourceLocation.GetAllocator()->GetAllocationOffsetInBytes(ResourceLocation.GetBuddyAllocatorPrivateData());
 
 				FD3D12Resource* NewResource = nullptr;
-				VERIFYD3D12RESULT(Adapter->CreatePlacedResource(Desc, BackingHeap, HeapOffset, InitialState, ED3D12ResourceStateMode::MultiState, D3D12_RESOURCE_STATE_TBD, nullptr, &NewResource, Name));
+				VERIFYD3D12RESULT(Adapter->CreatePlacedResource(Desc, BackingHeap, HeapOffset, InCreateState, ED3D12ResourceStateMode::MultiState, D3D12_RESOURCE_STATE_TBD, nullptr, &NewResource, Name));
 
 				ResourceLocation.SetResource(NewResource);
 			}
@@ -1145,7 +1151,7 @@ void FD3D12DefaultBufferPool::AllocDefaultResource(D3D12_HEAP_TYPE InHeapType, c
 	// Allocate Standalone
 	// Todo: track stand alone allocations and see how much memory we use by this and how many we have
 	FD3D12Resource* NewResource = nullptr;
-	VERIFYD3D12RESULT(Adapter->CreateBuffer(InHeapType, GetGPUMask(), GetVisibilityMask(), InitialState, Desc.Width, &NewResource, Name, Desc.Flags));
+	VERIFYD3D12RESULT(Adapter->CreateBuffer(InHeapType, GetGPUMask(), GetVisibilityMask(), InCreateState, Desc.Width, &NewResource, Name, Desc.Flags));
 
 	ResourceLocation.AsStandAlone(NewResource, Desc.Width);
 }
@@ -1208,8 +1214,44 @@ FD3D12DefaultBufferPool* FD3D12DefaultBufferAllocator::CreateBufferPool(D3D12_HE
 	return NewPool;
 }
 
+
+bool FD3D12DefaultBufferAllocator::IsPlacedResource(D3D12_RESOURCE_FLAGS InResourceFlags, ED3D12ResourceStateMode InResourceStateMode)
+{
+#ifdef USE_BUCKET_ALLOCATOR
+	return false;
+#else
+	FD3D12BuddyAllocator::EAllocationStrategy AllocationStrategy = FD3D12DefaultBufferPool::GetBuddyAllocatorAllocationStrategy(InResourceFlags, InResourceStateMode);
+	return (AllocationStrategy == FD3D12BuddyAllocator::EAllocationStrategy::kPlacedResource);
+#endif
+}
+
+
+D3D12_RESOURCE_STATES FD3D12DefaultBufferAllocator::GetDefaultInitialResourceState(D3D12_HEAP_TYPE InHeapType, EBufferUsageFlags InBufferFlags)
+{
+	// Validate the create state
+	if (InHeapType == D3D12_HEAP_TYPE_READBACK)
+	{
+		return D3D12_RESOURCE_STATE_COPY_DEST;
+	}
+	else if (InHeapType == D3D12_HEAP_TYPE_UPLOAD)
+	{
+		return D3D12_RESOURCE_STATE_GENERIC_READ;
+	}
+#if D3D12_RHI_RAYTRACING
+	else if (InBufferFlags & BUF_AccelerationStructure)
+	{
+		return D3D12_RESOURCE_STATE_RAYTRACING_ACCELERATION_STRUCTURE;
+	}
+#endif // D3D12_RHI_RAYTRACING
+	else
+	{
+		return D3D12_RESOURCE_STATE_GENERIC_READ;
+	}
+}
+
+
 // Grab a buffer from the available buffers or create a new buffer if none are available
-void FD3D12DefaultBufferAllocator::AllocDefaultResource(D3D12_HEAP_TYPE InHeapType, const D3D12_RESOURCE_DESC& InResourceDesc, EBufferUsageFlags InBufferUsage, ED3D12ResourceStateMode InResourceStateMode, FD3D12ResourceLocation& ResourceLocation, uint32 Alignment, const TCHAR* Name)
+void FD3D12DefaultBufferAllocator::AllocDefaultResource(D3D12_HEAP_TYPE InHeapType, const D3D12_RESOURCE_DESC& InResourceDesc, EBufferUsageFlags InBufferUsage, ED3D12ResourceStateMode InResourceStateMode, D3D12_RESOURCE_STATES InCreateState, FD3D12ResourceLocation& ResourceLocation, uint32 Alignment, const TCHAR* Name)
 {
 	// Patch out deny shader resource because it doesn't add anything for buffers and allows more pool sharing
 	// TODO: check if this is different on Xbox?
@@ -1234,7 +1276,7 @@ void FD3D12DefaultBufferAllocator::AllocDefaultResource(D3D12_HEAP_TYPE InHeapTy
 	}
 
 	// Perform actual allocation
-	BufferPool->AllocDefaultResource(InHeapType, ResourceDesc, InBufferUsage, InResourceStateMode, ResourceLocation, Alignment, Name);
+	BufferPool->AllocDefaultResource(InHeapType, ResourceDesc, InBufferUsage, InResourceStateMode, InCreateState, ResourceLocation, Alignment, Name);
 }
 
 
diff --git a/Engine/Source/Runtime/D3D12RHI/Private/D3D12Allocation.h b/Engine/Source/Runtime/D3D12RHI/Private/D3D12Allocation.h
--- a/Engine/Source/Runtime/D3D12RHI/Private/D3D12Allocation.h
+++ b/Engine/Source/Runtime/D3D12RHI/Private/D3D12Allocation.h
@@ -416,7 +416,7 @@ public:
 	~FD3D12DefaultBufferPool() { delete Allocator; }
 
 	bool SupportsAllocation(D3D12_HEAP_TYPE InHeapType, D3D12_RESOURCE_FLAGS InResourceFlags, EBufferUsageFlags InBufferUsage, ED3D12ResourceStateMode InResourceStateMode) const;
-	void AllocDefaultResource(D3D12_HEAP_TYPE InHeapType, const D3D12_RESOURCE_DESC& Desc, EBufferUsageFlags InBufferUsage, ED3D12ResourceStateMode InResourceStateMode, FD3D12ResourceLocation& ResourceLocation, uint32 Alignment, const TCHAR* Name);
+	void AllocDefaultResource(D3D12_HEAP_TYPE InHeapType, const D3D12_RESOURCE_DESC& Desc, EBufferUsageFlags InBufferUsage, ED3D12ResourceStateMode InResourceStateMode, D3D12_RESOURCE_STATES InCreateState, FD3D12ResourceLocation& ResourceLocation, uint32 Alignment, const TCHAR* Name);
 	void CleanUpAllocations(uint64 InFrameLag);
 	void UpdateMemoryStats(uint32& IOMemoryAllocated, uint32& IOMemoryUsed, uint32& IOMemoryFree, uint32& IOAlignmentWaste, uint32& IOAllocatedPageCount, uint32& IOFullPageCount);
 
@@ -435,11 +435,14 @@ public:
 	FD3D12DefaultBufferAllocator(FD3D12Device* InParent, FRHIGPUMask VisibleNodes);
 
 	// Grab a buffer from the available buffers or create a new buffer if none are available
-	void AllocDefaultResource(D3D12_HEAP_TYPE InHeapType, const D3D12_RESOURCE_DESC& pDesc, EBufferUsageFlags InBufferUsage,	ED3D12ResourceStateMode InResourceStateMode, FD3D12ResourceLocation& ResourceLocation, uint32 Alignment, const TCHAR* Name);
+	void AllocDefaultResource(D3D12_HEAP_TYPE InHeapType, const D3D12_RESOURCE_DESC& pDesc, EBufferUsageFlags InBufferUsage, ED3D12ResourceStateMode InResourceStateMode, D3D12_RESOURCE_STATES InCreateState, FD3D12ResourceLocation& ResourceLocation, uint32 Alignment, const TCHAR* Name);
 	void FreeDefaultBufferPools();
 	void CleanupFreeBlocks(uint64 InFrameLag);
 	void UpdateMemoryStats();
 
+	static bool IsPlacedResource(D3D12_RESOURCE_FLAGS InResourceFlags, ED3D12ResourceStateMode InResourceStateMode);
+	static D3D12_RESOURCE_STATES GetDefaultInitialResourceState(D3D12_HEAP_TYPE InHeapType, EBufferUsageFlags InBufferFlags);
+
 private:
 
 	FD3D12DefaultBufferPool* CreateBufferPool(D3D12_HEAP_TYPE InHeapType, D3D12_RESOURCE_FLAGS InResourceFlags, EBufferUsageFlags InBufferUsage, ED3D12ResourceStateMode InResourceStateMode);
diff --git a/Engine/Source/Runtime/D3D12RHI/Private/D3D12Buffer.cpp b/Engine/Source/Runtime/D3D12RHI/Private/D3D12Buffer.cpp
--- a/Engine/Source/Runtime/D3D12RHI/Private/D3D12Buffer.cpp
+++ b/Engine/Source/Runtime/D3D12RHI/Private/D3D12Buffer.cpp
@@ -68,11 +68,13 @@ struct FD3D12RHICommandInitializeBuffer final : public FRHICommand<FD3D12RHIComm
 	FD3D12Buffer* Buffer;
 	FD3D12ResourceLocation SrcResourceLoc;
 	uint32 Size;
+	D3D12_RESOURCE_STATES DestinationState;
 
-	FORCEINLINE_DEBUGGABLE FD3D12RHICommandInitializeBuffer(FD3D12Buffer* InBuffer, FD3D12ResourceLocation& InSrcResourceLoc, uint32 InSize)
+	FORCEINLINE_DEBUGGABLE FD3D12RHICommandInitializeBuffer(FD3D12Buffer* InBuffer, FD3D12ResourceLocation& InSrcResourceLoc, uint32 InSize, D3D12_RESOURCE_STATES InDestinationState)
 		: Buffer(InBuffer)
 		, SrcResourceLoc(InSrcResourceLoc.GetParentDevice())
 		, Size(InSize)
+		, DestinationState(InDestinationState)
 	{
 		FD3D12ResourceLocation::TransferOwnership(SrcResourceLoc, InSrcResourceLoc);
 	}
@@ -92,8 +94,11 @@ struct FD3D12RHICommandInitializeBuffer final : public FRHICommand<FD3D12RHIComm
 			FD3D12CommandListHandle& hCommandList = Device->GetDefaultCommandContext().CommandListHandle;
 			// Copy from the temporary upload heap to the default resource
 			{
-				// Writable structured buffers are sometimes initialized with initial data which means they sometimes need tracking.
-				FConditionalScopeResourceBarrier ConditionalScopeResourceBarrier(hCommandList, Destination, D3D12_RESOURCE_STATE_COPY_DEST, 0);
+				// if resource doesn't require state tracking then transition to copy dest here (could have been suballocated from shared resource) - not very optimal and should be batched
+				if (!Destination->RequiresResourceStateTracking())
+				{
+					hCommandList.AddTransitionBarrier(Destination, Destination->GetDefaultResourceState(), D3D12_RESOURCE_STATE_COPY_DEST, D3D12_RESOURCE_BARRIER_ALL_SUBRESOURCES);
+				}
 
 				Device->GetDefaultCommandContext().numCopies++;
 				hCommandList.FlushResourceBarriers();
@@ -103,7 +108,24 @@ struct FD3D12RHICommandInitializeBuffer final : public FRHICommand<FD3D12RHIComm
 					SrcResourceLoc.GetResource()->GetResource(),
 					SrcResourceLoc.GetOffsetFromBaseOfResource(), Size);
 
-				hCommandList.UpdateResidency(Destination);
+				// Update the resource state after the copy has been done (will take care of updating the residency as well)
+				hCommandList.AddTransitionBarrier(Destination, D3D12_RESOURCE_STATE_COPY_DEST, DestinationState, D3D12_RESOURCE_BARRIER_ALL_SUBRESOURCES);
+
+				if (Destination->RequiresResourceStateTracking())
+				{
+					// Update the tracked resource state of this resource in the command list
+					CResourceState& ResourceState = hCommandList.GetResourceState(Destination);
+					ResourceState.SetResourceState(DestinationState);
+					Destination->GetResourceState().SetResourceState(DestinationState);
+
+					// Add dummy pending barrier, because the end state needs to be updated after execture command list with tracked state in the command list
+					hCommandList.AddPendingResourceBarrier(Destination, D3D12_RESOURCE_STATE_TBD, D3D12_RESOURCE_BARRIER_ALL_SUBRESOURCES);
+				}
+				else
+				{
+					check(Destination->GetDefaultResourceState() == DestinationState);
+				}
+
 				hCommandList.UpdateResidency(SrcResourceLoc.GetResource());
 			}
 		}
@@ -115,6 +137,7 @@ void FD3D12Adapter::AllocateBuffer(FD3D12Device* Device,
 	uint32 Size,
 	uint32 InUsage,
 	ED3D12ResourceStateMode InResourceStateMode,
+	D3D12_RESOURCE_STATES InCreateState,
 	FRHIResourceCreateInfo& CreateInfo,
 	uint32 Alignment,
 	FD3D12TransientResource& TransientResource,
@@ -128,6 +151,7 @@ void FD3D12Adapter::AllocateBuffer(FD3D12Device* Device,
 	if (bIsDynamic)
 	{
 		check(InResourceStateMode != ED3D12ResourceStateMode::MultiState);
+		check(InCreateState == D3D12_RESOURCE_STATE_GENERIC_READ);
 		void* pData = GetUploadHeapAllocator(Device->GetGPUIndex()).AllocUploadResource(Size, Alignment, ResourceLocation);
 		check(ResourceLocation.GetSize() == Size);
 
@@ -142,7 +166,7 @@ void FD3D12Adapter::AllocateBuffer(FD3D12Device* Device,
 	}
 	else
 	{
-		Device->GetDefaultBufferAllocator().AllocDefaultResource(D3D12_HEAP_TYPE_DEFAULT, InDesc, (EBufferUsageFlags)InUsage, InResourceStateMode, ResourceLocation, Alignment, CreateInfo.DebugName);
+		Device->GetDefaultBufferAllocator().AllocDefaultResource(D3D12_HEAP_TYPE_DEFAULT, InDesc, (EBufferUsageFlags)InUsage, InResourceStateMode, InCreateState, ResourceLocation, Alignment, CreateInfo.DebugName);
 		check(ResourceLocation.GetSize() == Size);
 	}
 }
@@ -154,6 +178,7 @@ FD3D12Buffer* FD3D12Adapter::CreateRHIBuffer(FRHICommandListImmediate* RHICmdLis
 	uint32 Size,
 	uint32 InUsage,
 	ED3D12ResourceStateMode InResourceStateMode,
+	ERHIAccess InResourceState,
 	FRHIResourceCreateInfo& CreateInfo)
 {
 	SCOPE_CYCLE_COUNTER(STAT_D3D12CreateBufferTime);
@@ -161,6 +186,15 @@ FD3D12Buffer* FD3D12Adapter::CreateRHIBuffer(FRHICommandListImmediate* RHICmdLis
 	const bool bIsDynamic = (InUsage & BUF_AnyDynamic) ? true : false;
 	const uint32 FirstGPUIndex = CreateInfo.GPUMask.GetFirstIndex();
 
+	// Does this resource support tracking?
+	const bool bSupportResourceStateTracking = !bIsDynamic && FD3D12DefaultBufferAllocator::IsPlacedResource(InDesc.Flags, InResourceStateMode);
+
+	// Initial state is derived from the InResourceState if it's supports tracking
+	D3D12_HEAP_TYPE HeapType = bIsDynamic ? D3D12_HEAP_TYPE_UPLOAD : D3D12_HEAP_TYPE_DEFAULT;
+	const FD3D12Resource::FD3D12ResourceTypeHelper Type(InDesc, HeapType);
+	const D3D12_RESOURCE_STATES InitialState = bSupportResourceStateTracking ? Type.GetOptimalInitialState(InResourceState, false) : 
+		FD3D12DefaultBufferAllocator::GetDefaultInitialResourceState(HeapType, (EBufferUsageFlags)InUsage);
+
 	FD3D12Buffer* BufferOut = nullptr;
 	if (bIsDynamic)
 	{
@@ -172,7 +206,7 @@ FD3D12Buffer* FD3D12Adapter::CreateRHIBuffer(FRHICommandListImmediate* RHICmdLis
 
 			if (Device->GetGPUIndex() == FirstGPUIndex)
 			{
-				AllocateBuffer(Device, InDesc, Size, InUsage, InResourceStateMode, CreateInfo, Alignment, *NewBuffer, NewBuffer->ResourceLocation);
+				AllocateBuffer(Device, InDesc, Size, InUsage, InResourceStateMode, InitialState, CreateInfo, Alignment, *NewBuffer, NewBuffer->ResourceLocation);
 				NewBuffer0 = NewBuffer;
 			}
 			else
@@ -186,12 +220,15 @@ FD3D12Buffer* FD3D12Adapter::CreateRHIBuffer(FRHICommandListImmediate* RHICmdLis
 	}
 	else
 	{
+		// Setup the state at which the resource needs to be created - copy dest only supported for placed resources
+		const D3D12_RESOURCE_STATES CreateState = (CreateInfo.ResourceArray && bSupportResourceStateTracking) ? D3D12_RESOURCE_STATE_COPY_DEST : InitialState;
+
 		BufferOut = CreateLinkedObject<FD3D12Buffer>(CreateInfo.GPUMask, [&](FD3D12Device* Device)
 		{
 			FD3D12Buffer* NewBuffer = new FD3D12Buffer(Device, Size, InUsage, Stride);
 			NewBuffer->BufferAlignment = Alignment;
 
-			AllocateBuffer(Device, InDesc, Size, InUsage, InResourceStateMode, CreateInfo, Alignment, *NewBuffer, NewBuffer->ResourceLocation);
+			AllocateBuffer(Device, InDesc, Size, InUsage, InResourceStateMode, CreateState, CreateInfo, Alignment, *NewBuffer, NewBuffer->ResourceLocation);
 
 			return NewBuffer;
 		});
@@ -234,16 +271,16 @@ FD3D12Buffer* FD3D12Adapter::CreateRHIBuffer(FRHICommandListImmediate* RHICmdLis
 				FD3D12ResourceLocation* SrcResourceLoc_Heap = new FD3D12ResourceLocation(SrcResourceLoc.GetParentDevice());
 				FD3D12ResourceLocation::TransferOwnership(*SrcResourceLoc_Heap, SrcResourceLoc);
 				ENQUEUE_RENDER_COMMAND(CmdD3D12InitializeBuffer)(
-					[BufferOut, SrcResourceLoc_Heap, Size](FRHICommandListImmediate& RHICmdList)
+					[BufferOut, SrcResourceLoc_Heap, Size, InitialState](FRHICommandListImmediate& RHICmdList)
 				{
 					if (RHICmdList.Bypass())
 					{
-						FD3D12RHICommandInitializeBuffer Command(BufferOut, *SrcResourceLoc_Heap, Size);
+						FD3D12RHICommandInitializeBuffer Command(BufferOut, *SrcResourceLoc_Heap, Size, InitialState);
 						Command.ExecuteNoCmdList();
 					}
 					else
 					{
-						new (RHICmdList.AllocCommand<FD3D12RHICommandInitializeBuffer>()) FD3D12RHICommandInitializeBuffer(BufferOut, *SrcResourceLoc_Heap, Size);
+						new (RHICmdList.AllocCommand<FD3D12RHICommandInitializeBuffer>()) FD3D12RHICommandInitializeBuffer(BufferOut, *SrcResourceLoc_Heap, Size, InitialState);
 					}
 					delete SrcResourceLoc_Heap;
 				});
@@ -251,13 +288,13 @@ FD3D12Buffer* FD3D12Adapter::CreateRHIBuffer(FRHICommandListImmediate* RHICmdLis
 			else if (!RHICmdList || RHICmdList->Bypass())
 			{
 				// On RHIT or RT (when bypassing), we can access immediate context directly
-				FD3D12RHICommandInitializeBuffer Command(BufferOut, SrcResourceLoc, Size);
+				FD3D12RHICommandInitializeBuffer Command(BufferOut, SrcResourceLoc, Size, InitialState);
 				Command.ExecuteNoCmdList();
 			}
 			else
 			{
 				// On RT but not bypassing
-				new (RHICmdList->AllocCommand<FD3D12RHICommandInitializeBuffer>()) FD3D12RHICommandInitializeBuffer(BufferOut, SrcResourceLoc, Size);
+				new (RHICmdList->AllocCommand<FD3D12RHICommandInitializeBuffer>()) FD3D12RHICommandInitializeBuffer(BufferOut, SrcResourceLoc, Size, InitialState);
 			}
 		}
 
@@ -397,7 +434,7 @@ void* FD3D12DynamicRHI::LockBuffer(FRHICommandListImmediate* RHICmdList, FD3D12B
 					FD3D12CommandContext& DefaultContext = Device->GetDefaultCommandContext();
 
 					FD3D12CommandListHandle& hCommandList = DefaultContext.CommandListHandle;
-					FConditionalScopeResourceBarrier ScopeResourceBarrierSource(hCommandList, pResource, D3D12_RESOURCE_STATE_COPY_SOURCE, 0);
+					FScopedResourceBarrier ScopeResourceBarrierSource(hCommandList, pResource, D3D12_RESOURCE_STATE_COPY_SOURCE, 0, FD3D12DynamicRHI::ETransitionMode::Apply);
 					// Don't need to transition upload heaps
 
 					uint64 SubAllocOffset = Buffer->ResourceLocation.GetOffsetFromBaseOfResource();
diff --git a/Engine/Source/Runtime/D3D12RHI/Private/D3D12Commands.cpp b/Engine/Source/Runtime/D3D12RHI/Private/D3D12Commands.cpp
--- a/Engine/Source/Runtime/D3D12RHI/Private/D3D12Commands.cpp
+++ b/Engine/Source/Runtime/D3D12RHI/Private/D3D12Commands.cpp
@@ -177,7 +177,7 @@ void FD3D12CommandContext::RHIDispatchIndirectComputeShader(FRHIBuffer* Argument
 
 	// Indirect args buffer can be a previously pending UAV, which becomes PS\Non-PS read. ApplyState will flush pending transitions, so enqueue the indirect
 	// arg transition and flush here.
-	FD3D12DynamicRHI::TransitionResource(CommandListHandle, Location.GetResource(), D3D12_RESOURCE_STATE_INDIRECT_ARGUMENT | D3D12_RESOURCE_STATE_NON_PIXEL_SHADER_RESOURCE, D3D12_RESOURCE_BARRIER_ALL_SUBRESOURCES);
+	FD3D12DynamicRHI::TransitionResource(CommandListHandle, Location.GetResource(), D3D12_RESOURCE_STATE_TBD, D3D12_RESOURCE_STATE_INDIRECT_ARGUMENT | D3D12_RESOURCE_STATE_NON_PIXEL_SHADER_RESOURCE, D3D12_RESOURCE_BARRIER_ALL_SUBRESOURCES, FD3D12DynamicRHI::ETransitionMode::Validate);
 	CommandListHandle.FlushResourceBarriers();	// Must flush so the desired state is actually set.
 
 	FD3D12Adapter* Adapter = GetParentDevice()->GetParentAdapter();
@@ -216,31 +216,35 @@ void EnumerateSubresources(FD3D12Resource* Resource, const FRHITransitionInfo& I
 	uint32 ArraySize = Resource->GetArraySize();
 	uint32 PlaneCount = Resource->GetPlaneCount();
 
+	uint32 IterationMipCount = MipCount;
+	uint32 IterationArraySize = ArraySize;
+	uint32 IterationPlaneCount = PlaneCount;
+
 	if (!Info.IsAllMips())
 	{
 		FirstMipSlice = Info.MipIndex;
-		MipCount = 1;
+		IterationMipCount = 1;
 	}
 
 	if (!Info.IsAllArraySlices())
 	{
 		FirstArraySlice = Info.ArraySlice;
-		ArraySize = 1;
+		IterationArraySize = 1;
 	}
 
 	if (!Info.IsAllPlaneSlices())
 	{
 		FirstPlaneSlice = Info.PlaneSlice;
-		PlaneCount = 1;
+		IterationPlaneCount = 1;
 	}
 
-	for (uint32 PlaneSlice = FirstPlaneSlice; PlaneSlice < FirstPlaneSlice + PlaneCount; ++PlaneSlice)
+	for (uint32 PlaneSlice = FirstPlaneSlice; PlaneSlice < FirstPlaneSlice + IterationPlaneCount; ++PlaneSlice)
 	{
-		for (uint32 ArraySlice = FirstArraySlice; ArraySlice < FirstArraySlice + ArraySize; ++ArraySlice)
+		for (uint32 ArraySlice = FirstArraySlice; ArraySlice < FirstArraySlice + IterationArraySize; ++ArraySlice)
 		{
-			for (uint32 MipSlice = FirstMipSlice; MipSlice < FirstMipSlice + MipCount; ++MipSlice)
+			for (uint32 MipSlice = FirstMipSlice; MipSlice < FirstMipSlice + IterationMipCount; ++MipSlice)
 			{
-				const uint32 Subresource = D3D12CalcSubresource(FirstMipSlice, FirstArraySlice, FirstPlaneSlice, MipCount, ArraySize);
+				const uint32 Subresource = D3D12CalcSubresource(MipSlice, ArraySlice, PlaneSlice, MipCount, ArraySize);
 				Function(Subresource);
 			}
 		}
@@ -282,49 +286,77 @@ void ProcessResource(FD3D12CommandContext& Context, const FRHITransitionInfo& In
 	}
 }
 
-static D3D12_RESOURCE_STATES GetEndTransitionState(const FRHITransitionInfo& Info, FD3D12Resource* Resource, ERHIPipeline Pipelines, D3D12_RESOURCE_STATES SkipFastClearEliminateState)
+static void HandleResourceTransitions(FD3D12CommandContext& Context, const FD3D12TransitionData* TransitionData, D3D12_RESOURCE_STATES SkipFastClearEliminateState, bool& bUAVBarrier)
 {
-	D3D12_RESOURCE_STATES State = D3D12_RESOURCE_STATE_COMMON;
+	const bool bDstAllPipelines = EnumHasAllFlags(TransitionData->DstPipelines, ERHIPipeline::All);
 
-	if (EnumHasAnyFlags(Info.Flags, EResourceTransitionFlags::MaintainCompression) && IsReadOnlyAccess(Info.AccessAfter))
+	for (const FRHITransitionInfo& Info : TransitionData->Infos)
 	{
-		State |= SkipFastClearEliminateState;
-	}
-
-	const bool bIsAsyncComputeContext = EnumHasAnyFlags(Pipelines, ERHIPipeline::AsyncCompute);
+		const bool bUAVAccess = EnumHasAnyFlags(Info.AccessAfter, ERHIAccess::UAVMask);
 
-	if (Info.AccessAfter == ERHIAccess::ResolveSrc)
-	{
-		State |= D3D12_RESOURCE_STATE_RESOLVE_SOURCE;
-	}
-	else if (Info.AccessAfter == ERHIAccess::ResolveDst)
-	{
-		State |= D3D12_RESOURCE_STATE_RESOLVE_DEST;
-	}
-	else if (EnumHasAnyFlags(Info.AccessAfter, ERHIAccess::WritableMask))
-	{
-		if (EnumHasAnyFlags(Info.AccessAfter, ERHIAccess::UAVMask) || bIsAsyncComputeContext)
-		{
-			State |= D3D12_RESOURCE_STATE_UNORDERED_ACCESS;
-		}
-		else
+		// If targeting all pipelines and UAV access then add UAV barrier even with RHI based transitions
+		if (bDstAllPipelines)
 		{
-			State |= Resource->GetWritableState();
+			bUAVBarrier |= bUAVAccess;
 		}
-	}
-	else if (EnumHasAnyFlags(Info.AccessAfter, ERHIAccess::ReadableMask))
-	{
-		if (bIsAsyncComputeContext)
+
+		// Use the engine defined transitions?
+		if (GUseInternalTransitions)
 		{
-			State |= D3D12_RESOURCE_STATE_NON_PIXEL_SHADER_RESOURCE;
+			// Only need to check for UAV barriers here, all other transitions are handled inside the RHI
+			bUAVBarrier |= bUAVAccess;
 		}
 		else
 		{
-			State |= Resource->GetReadableState();
+			// Sometimes we could still have barriers without resources, invalid but can still happen
+			if (!Info.Resource)
+			{
+				// Add a UAV barrier in case of no resource and after state is UAV - user probably requested a force UAV barrier here				
+				bUAVBarrier |= bUAVAccess;
+			}
+			else
+			{
+				ProcessResource(Context, Info, [&](const FRHITransitionInfo& Info, FD3D12Resource* Resource)
+					{
+						if (!Resource->RequiresResourceStateTracking())
+						{
+							return;
+						}
+
+						const bool bIsAsyncCompute = EnumHasAnyFlags(TransitionData->DstPipelines, ERHIPipeline::AsyncCompute);
+
+						// Derive or get the correct before & after state
+						D3D12_RESOURCE_STATES BeforeState = (Info.AccessBefore != ERHIAccess::Unknown) ? GetD3D12ResourceState(Info.AccessBefore, bIsAsyncCompute) : D3D12_RESOURCE_STATE_TBD;
+						D3D12_RESOURCE_STATES AfterState = GetD3D12ResourceState(Info.AccessAfter, bIsAsyncCompute);
+
+						// Add the compression flags if needed
+						if (EnumHasAnyFlags(Info.Flags, EResourceTransitionFlags::MaintainCompression))
+						{
+							AfterState |= SkipFastClearEliminateState;
+						}
+
+						// enqueue the correct transitions
+						if (Info.IsWholeResource() || Resource->GetSubresourceCount() == 1)
+						{
+							bUAVBarrier = bUAVBarrier | FD3D12DynamicRHI::TransitionResource(Context.CommandListHandle, Resource, BeforeState, AfterState, D3D12_RESOURCE_BARRIER_ALL_SUBRESOURCES, FD3D12DynamicRHI::ETransitionMode::Apply);
+						}
+						else
+						{
+							EnumerateSubresources(Resource, Info, [&](uint32 Subresource)
+								{
+									bUAVBarrier = bUAVBarrier | FD3D12DynamicRHI::TransitionResource(Context.CommandListHandle, Resource, BeforeState, AfterState, Subresource, FD3D12DynamicRHI::ETransitionMode::Apply);
+								});
+						}
+					});
+			}
 		}
 	}
+}
 
-	return State;
+static bool ProcessTransitionDuringBegin(const FD3D12TransitionData* Data)
+{
+	// Pipe changes which are not ending with graphics or targeting all pipelines are handle during begin
+	return ((Data->SrcPipelines != Data->DstPipelines && Data->DstPipelines != ERHIPipeline::Graphics) || EnumHasAllFlags(Data->DstPipelines, ERHIPipeline::All));
 }
 
 void FD3D12CommandContext::RHIBeginTransitionsWithoutFencing(TArrayView<const FRHITransition*> Transitions)
@@ -339,69 +371,10 @@ void FD3D12CommandContext::RHIBeginTransitionsWithoutFencing(TArrayView<const FR
 	{
 		const FD3D12TransitionData* Data = Transition->GetPrivateData<FD3D12TransitionData>();
 
-		// Same pipe transitions or transitions onto the graphics pipe are handled in End.
-		if (Data->SrcPipelines == Data->DstPipelines || Data->DstPipelines == ERHIPipeline::Graphics)
-		{
-			continue;
-		}
-
-		const bool bDstAllPipelines = EnumHasAllFlags(Data->DstPipelines, ERHIPipeline::All);
-
-		for (const FRHITransitionInfo& Info : Data->Infos)
+		// Handle transition during BeginTransitions?
+		if (ProcessTransitionDuringBegin(Data))
 		{
-			if (bDstAllPipelines)
-			{
-				bUAVBarrier |= EnumHasAnyFlags(Info.AccessAfter, ERHIAccess::UAVMask);
-			}
-
-			if (!Info.Resource)
-			{
-				continue;
-			}
-
-			ProcessResource(*this, Info, [&](const FRHITransitionInfo& Info, FD3D12Resource* Resource)
-			{
-				if (!Resource->RequiresResourceStateTracking())
-				{
-					return;
-				}
-
-				D3D12_RESOURCE_STATES State = D3D12_RESOURCE_STATE_COMMON;
-
-				// When forking from one pipeline to all, we perform the end transitions immediately.
-				if (bDstAllPipelines)
-				{
-					State = GetEndTransitionState(Info, Resource, Data->DstPipelines, SkipFastClearEliminateState);
-				}
-				else
-				{
-					if (EnumHasAnyFlags(Info.AccessAfter, ERHIAccess::WritableMask))
-					{
-						State |= D3D12_RESOURCE_STATE_UNORDERED_ACCESS;
-					}
-					else if (EnumHasAnyFlags(Info.AccessAfter, ERHIAccess::ReadableMask))
-					{
-						State |= D3D12_RESOURCE_STATE_NON_PIXEL_SHADER_RESOURCE;
-					}
-				}
-
-				if (State == D3D12_RESOURCE_STATE_COMMON)
-				{
-					return;
-				}
-
-				if (Info.IsWholeResource() || Resource->GetSubresourceCount() == 1)
-				{
-					FD3D12DynamicRHI::TransitionResource(CommandListHandle, Resource, State, D3D12_RESOURCE_BARRIER_ALL_SUBRESOURCES);
-				}
-				else
-				{
-					EnumerateSubresources(Resource, Info, [&](uint32 Subresource)
-					{
-						FD3D12DynamicRHI::TransitionResource(CommandListHandle, Resource, State, Subresource);
-					});
-				}
-			});
+			HandleResourceTransitions(*this, Data, SkipFastClearEliminateState, bUAVBarrier);
 		}
 	}
 
@@ -413,6 +386,10 @@ void FD3D12CommandContext::RHIBeginTransitionsWithoutFencing(TArrayView<const FR
 
 void FD3D12CommandContext::RHIBeginTransitions(TArrayView<const FRHITransition*> Transitions)
 {
+	// Can't correctly handle RHITransitions from multiple threads at the samae time right now - only internal transitions work correctly then
+	// (will be fixed if unknown before state is not allowed anymore)
+	check(bIsDefaultContext || bIsAsyncComputeContext || GUseInternalTransitions);
+
 	RHIBeginTransitionsWithoutFencing(Transitions);
 	SignalTransitionFences(Transitions);
 }
@@ -425,62 +402,16 @@ void FD3D12CommandContext::RHIEndTransitions(TArrayView<const FRHITransition*> T
 	const bool bShowTransitionEvents = CVarShowTransitions->GetInt() != 0;
 	SCOPED_RHI_CONDITIONAL_DRAW_EVENTF(*this, RHIEndTransitions, bShowTransitionEvents, TEXT("RHIEndTransitions"));
 
-	static_assert(USE_D3D12RHI_RESOURCE_STATE_TRACKING, "RHIEndTransitions is not implemented properly for this to be disabled.");
-
 	bool bUAVBarrier = false;
 
 	for (const FRHITransition* Transition : Transitions)
 	{
 		const FD3D12TransitionData* Data = Transition->GetPrivateData<FD3D12TransitionData>();
 
-		const bool bSamePipeline = !Data->bCrossPipeline;
-		const bool bDstAllPipelines = EnumHasAllFlags(Data->DstPipelines, ERHIPipeline::All);
-
-		for (const FRHITransitionInfo& Info : Data->Infos)
+		// Handle transition during EndTransitions?
+		if (!ProcessTransitionDuringBegin(Data))
 		{
-			// Sometimes we could still have barriers with resources, invalid but can still happen
-			if (bSamePipeline && !bDstAllPipelines)
-			{
-				bUAVBarrier |= EnumHasAnyFlags(Info.AccessAfter, ERHIAccess::UAVMask);
-			}
-
-			if (!Info.Resource)
-			{
-				continue;
-			}
-
-			ProcessResource(*this, Info, [&](const FRHITransitionInfo& Info, FD3D12Resource* Resource)
-			{
-				if (!Resource->RequiresResourceStateTracking())
-				{
-					return;
-				}
-
-				D3D12_RESOURCE_STATES State = D3D12_RESOURCE_STATE_COMMON;
-
-				// When forking from one pipeline to all, we perform the end transitions in begin.
-				if (!bDstAllPipelines)
-				{
-					State = GetEndTransitionState(Info, Resource, Data->DstPipelines, SkipFastClearEliminateState);
-				}
-
-				if (State == D3D12_RESOURCE_STATE_COMMON)
-				{
-					return;
-				}
-
-				if (Info.IsWholeResource() || Resource->GetSubresourceCount() == 1)
-				{
-					FD3D12DynamicRHI::TransitionResource(CommandListHandle, Resource, State, D3D12_RESOURCE_BARRIER_ALL_SUBRESOURCES);
-				}
-				else
-				{
-					EnumerateSubresources(Resource, Info, [&](uint32 Subresource)
-					{
-						FD3D12DynamicRHI::TransitionResource(CommandListHandle, Resource, State, Subresource);
-					});
-				}
-			});
+			HandleResourceTransitions(*this, Data, SkipFastClearEliminateState, bUAVBarrier);
 		}
 	}
 
@@ -522,7 +453,7 @@ void FD3D12CommandContext::RHICopyToStagingBuffer(FRHIVertexBuffer* SourceBuffer
 		// Unknown aligment requirement for sub allocated read back buffer data
 		uint32 AllocationAlignment = 16;
 		const D3D12_RESOURCE_DESC BufferDesc = CD3DX12_RESOURCE_DESC::Buffer(NumBytes, D3D12_RESOURCE_FLAG_NONE);		
-		GetParentDevice()->GetDefaultBufferAllocator().AllocDefaultResource(D3D12_HEAP_TYPE_READBACK, BufferDesc, (EBufferUsageFlags)BUF_None, ED3D12ResourceStateMode::SingleState, StagingBuffer->ResourceLocation, AllocationAlignment, TEXT("StagedRead"));
+		GetParentDevice()->GetDefaultBufferAllocator().AllocDefaultResource(D3D12_HEAP_TYPE_READBACK, BufferDesc, (EBufferUsageFlags)BUF_None, ED3D12ResourceStateMode::SingleState, D3D12_RESOURCE_STATE_COPY_DEST, StagingBuffer->ResourceLocation, AllocationAlignment, TEXT("StagedRead"));
 		check(StagingBuffer->ResourceLocation.GetSize() == NumBytes);
 		StagingBuffer->ShadowBufferSize = NumBytes;
 	}
@@ -540,7 +471,7 @@ void FD3D12CommandContext::RHICopyToStagingBuffer(FRHIVertexBuffer* SourceBuffer
 
 		if (pSourceResource->RequiresResourceStateTracking())
 		{
-			FD3D12DynamicRHI::TransitionResource(CommandListHandle, pSourceResource, D3D12_RESOURCE_STATE_COPY_SOURCE, 0);
+			FD3D12DynamicRHI::TransitionResource(CommandListHandle, pSourceResource, D3D12_RESOURCE_STATE_TBD, D3D12_RESOURCE_STATE_COPY_SOURCE, 0, FD3D12DynamicRHI::ETransitionMode::Validate);
 			CommandListHandle.FlushResourceBarriers();	// Must flush so the desired state is actually set.
 		}
 
@@ -1183,6 +1114,54 @@ void FD3D12CommandContext::SetRenderTargetsAndClear(const FRHISetRenderTargetsIn
 		RenderTargetsInfo.ColorRenderTarget,
 		&RenderTargetsInfo.DepthStencilRenderTarget);
 
+	FD3D12RenderTargetView* RenderTargetViews[MaxSimultaneousRenderTargets];
+	FD3D12DepthStencilView* DSView = nullptr;
+	uint32 NumSimultaneousRTs = 0;
+	StateCache.GetRenderTargets(RenderTargetViews, &NumSimultaneousRTs, &DSView);
+	FD3D12BoundRenderTargets BoundRenderTargets(RenderTargetViews, NumSimultaneousRTs, DSView);
+	FD3D12DepthStencilView* DepthStencilView = BoundRenderTargets.GetDepthStencilView();
+
+	// perform the correct barrier transitions to RTV
+	{
+		for (int32 TargetIndex = 0; TargetIndex < BoundRenderTargets.GetNumActiveTargets(); TargetIndex++)
+		{
+			FD3D12RenderTargetView* RTView = BoundRenderTargets.GetRenderTargetView(TargetIndex);
+			if (RTView != nullptr)
+			{
+				FD3D12DynamicRHI::TransitionResource(CommandListHandle, RTView, D3D12_RESOURCE_STATE_RENDER_TARGET, FD3D12DynamicRHI::ETransitionMode::Validate);
+			}
+		}
+	}
+
+	// perform the optimal barrier transition for depth stencil as well
+	if (DepthStencilView)
+	{
+		// Find out the desired state for both the depth & the stencil
+		FExclusiveDepthStencil DepthStencilAccess = RenderTargetsInfo.DepthStencilRenderTarget.GetDepthStencilAccess();
+		D3D12_RESOURCE_STATES DepthState = DepthStencilAccess.IsDepthWrite() ? D3D12_RESOURCE_STATE_DEPTH_WRITE : (DepthStencilAccess.IsDepthRead() ? D3D12_RESOURCE_STATE_DEPTH_READ : D3D12_RESOURCE_STATE_COMMON);
+		D3D12_RESOURCE_STATES StencilState = DepthStencilAccess.IsStencilWrite() ? D3D12_RESOURCE_STATE_DEPTH_WRITE : (DepthStencilAccess.IsStencilRead() ? D3D12_RESOURCE_STATE_DEPTH_READ : D3D12_RESOURCE_STATE_COMMON);
+
+		// Same target state for both depth & stencil, then it can be done with single barrier
+		if (DepthState == StencilState)
+		{
+			if (DepthState != D3D12_RESOURCE_STATE_COMMON)
+			{
+				FD3D12DynamicRHI::TransitionResource(CommandListHandle, DepthStencilView, DepthState, FD3D12DynamicRHI::ETransitionMode::Validate);
+			}
+		}
+		else
+		{
+			if (DepthState != D3D12_RESOURCE_STATE_COMMON)
+			{
+				FD3D12DynamicRHI::TransitionResource(CommandListHandle, DepthStencilView->GetResource(), D3D12_RESOURCE_STATE_TBD, DepthState, DepthStencilView->GetDepthOnlyViewSubresourceSubset(), FD3D12DynamicRHI::ETransitionMode::Validate);
+			}
+			if (StencilState != D3D12_RESOURCE_STATE_COMMON)
+			{
+				FD3D12DynamicRHI::TransitionResource(CommandListHandle, DepthStencilView->GetResource(), D3D12_RESOURCE_STATE_TBD, StencilState, DepthStencilView->GetStencilOnlyViewSubresourceSubset(), FD3D12DynamicRHI::ETransitionMode::Validate);
+			}
+		}
+	}
+
 	if (RenderTargetsInfo.bClearColor || RenderTargetsInfo.bClearStencil || RenderTargetsInfo.bClearDepth)
 	{
 		FLinearColor ClearColors[MaxSimultaneousRenderTargets];
@@ -1349,6 +1328,10 @@ FORCEINLINE void SetResource(FD3D12CommandContext& CmdContext, uint32 BindIndex,
 template <EShaderFrequency ShaderFrequency>
 inline int32 SetShaderResourcesFromBuffer_Surface(FD3D12CommandContext& CmdContext, FD3D12UniformBuffer* RESTRICT Buffer, const uint32 * RESTRICT ResourceMap, int32 BufferIndex)
 {
+#if ENABLE_RHI_VALIDATION
+	constexpr ERHIAccess SRVAccess = (ShaderFrequency == SF_Compute) ? ERHIAccess::SRVCompute : ERHIAccess::SRVGraphics;
+#endif
+
 	const TRefCountPtr<FRHIResource>* RESTRICT Resources = Buffer->ResourceTable.GetData();
 	const float CurrentTime = FApp::GetCurrentTime();
 	int32 NumSetCalls = 0;
@@ -1373,6 +1356,13 @@ inline int32 SetShaderResourcesFromBuffer_Surface(FD3D12CommandContext& CmdConte
 				D3D12Resource = CmdContext.RetrieveTextureBase(GWhiteTexture->TextureRHI)->GetShaderResourceView();
 			}
 
+#if ENABLE_RHI_VALIDATION
+			if (CmdContext.Tracker)
+			{
+				CmdContext.Tracker->Assert(TextureRHI->GetWholeResourceIdentitySRV(), SRVAccess);
+			}
+#endif
+
 			SetResource<ShaderFrequency>(CmdContext, BindIndex, D3D12Resource);
 			NumSetCalls++;
 			ResourceInfo = *ResourceInfos++;
@@ -1386,6 +1376,10 @@ inline int32 SetShaderResourcesFromBuffer_Surface(FD3D12CommandContext& CmdConte
 template <EShaderFrequency ShaderFrequency>
 inline int32 SetShaderResourcesFromBuffer_SRV(FD3D12CommandContext& CmdContext, FD3D12UniformBuffer* RESTRICT Buffer, const uint32 * RESTRICT ResourceMap, int32 BufferIndex)
 {
+#if ENABLE_RHI_VALIDATION
+	constexpr ERHIAccess SRVAccess = (ShaderFrequency == SF_Compute) ? ERHIAccess::SRVCompute : ERHIAccess::SRVGraphics;
+#endif
+
 	const TRefCountPtr<FRHIResource>* RESTRICT Resources = Buffer->ResourceTable.GetData();
 	int32 NumSetCalls = 0;
 	const uint32 BufferOffset = ResourceMap[BufferIndex];
@@ -1401,6 +1395,13 @@ inline int32 SetShaderResourcesFromBuffer_SRV(FD3D12CommandContext& CmdContext,
 
 			FD3D12ShaderResourceView* D3D12Resource = CmdContext.RetrieveObject<FD3D12ShaderResourceView>((FRHIShaderResourceView*)(Resources[ResourceIndex].GetReference()));
 
+#if ENABLE_RHI_VALIDATION
+			if (CmdContext.Tracker)
+			{
+				CmdContext.Tracker->Assert(D3D12Resource->ViewIdentity, SRVAccess);
+			}
+#endif
+
 			SetResource<ShaderFrequency>(CmdContext, BindIndex, D3D12Resource);
 			NumSetCalls++;
 			ResourceInfo = *ResourceInfos++;
@@ -1471,6 +1472,10 @@ void FD3D12CommandContext::SetResourcesFromTables(const ShaderType* RESTRICT Sha
 template <EShaderFrequency ShaderFrequency>
 inline int32 SetShaderResourcesFromBuffer_UAVPS(FD3D12CommandContext& CmdContext, FD3D12UniformBuffer* RESTRICT Buffer, const uint32 * RESTRICT ResourceMap, int32 BufferIndex)
 {
+#if ENABLE_RHI_VALIDATION
+	constexpr ERHIAccess UAVAccess = (ShaderFrequency == SF_Compute) ? ERHIAccess::UAVCompute : ERHIAccess::UAVGraphics;
+#endif
+
 	const TRefCountPtr<FRHIResource>* RESTRICT Resources = Buffer->ResourceTable.GetData();
 	int32 NumSetCalls = 0;
 	const uint32 BufferOffset = ResourceMap[BufferIndex];
@@ -1487,6 +1492,14 @@ inline int32 SetShaderResourcesFromBuffer_UAVPS(FD3D12CommandContext& CmdContext
 			FRHIUnorderedAccessView* RHIUAV = (FRHIUnorderedAccessView*)(Resources[ResourceIndex].GetReference());
 
 			FD3D12UnorderedAccessView* D3D12Resource = CmdContext.RetrieveObject<FD3D12UnorderedAccessView>(RHIUAV);
+
+#if ENABLE_RHI_VALIDATION
+			if (CmdContext.Tracker)
+			{
+				CmdContext.Tracker->Assert(D3D12Resource->ViewIdentity, UAVAccess);
+			}
+#endif
+
 			SetResource<ShaderFrequency>(CmdContext, BindIndex, D3D12Resource);
 
 			NumSetCalls++;
@@ -1614,7 +1627,7 @@ void FD3D12CommandContext::RHIDrawPrimitiveIndirect(FRHIBuffer* ArgumentBufferRH
 
 	// Indirect args buffer can be a previously pending UAV, which becomes PS\Non-PS read. ApplyState will flush pending transitions, so enqueue the indirect
 	// arg transition and flush here.
-	FD3D12DynamicRHI::TransitionResource(CommandListHandle, Location.GetResource(), D3D12_RESOURCE_STATE_INDIRECT_ARGUMENT | D3D12_RESOURCE_STATE_PIXEL_SHADER_RESOURCE | D3D12_RESOURCE_STATE_NON_PIXEL_SHADER_RESOURCE, D3D12_RESOURCE_BARRIER_ALL_SUBRESOURCES);
+	FD3D12DynamicRHI::TransitionResource(CommandListHandle, Location.GetResource(), D3D12_RESOURCE_STATE_TBD, D3D12_RESOURCE_STATE_INDIRECT_ARGUMENT, D3D12_RESOURCE_BARRIER_ALL_SUBRESOURCES, FD3D12DynamicRHI::ETransitionMode::Validate);
 	CommandListHandle.FlushResourceBarriers();	// Must flush so the desired state is actually set.
 
 	CommandListHandle->ExecuteIndirect(
@@ -1663,7 +1676,7 @@ void FD3D12CommandContext::RHIDrawIndexedIndirect(FRHIBuffer* IndexBufferRHI, FR
 
 	// Indirect args buffer can be a previously pending UAV, which becomes PS\Non-PS read. ApplyState will flush pending transitions, so enqueue the indirect
 	// arg transition and flush here.
-	FD3D12DynamicRHI::TransitionResource(CommandListHandle, Location.GetResource(), D3D12_RESOURCE_STATE_INDIRECT_ARGUMENT | D3D12_RESOURCE_STATE_PIXEL_SHADER_RESOURCE | D3D12_RESOURCE_STATE_NON_PIXEL_SHADER_RESOURCE, D3D12_RESOURCE_BARRIER_ALL_SUBRESOURCES);
+	FD3D12DynamicRHI::TransitionResource(CommandListHandle, Location.GetResource(), D3D12_RESOURCE_STATE_TBD, D3D12_RESOURCE_STATE_INDIRECT_ARGUMENT, D3D12_RESOURCE_BARRIER_ALL_SUBRESOURCES, FD3D12DynamicRHI::ETransitionMode::Validate);
 	CommandListHandle.FlushResourceBarriers();	// Must flush so the desired state is actually set.
 
 	CommandListHandle->ExecuteIndirect(
@@ -1747,7 +1760,7 @@ void FD3D12CommandContext::RHIDrawIndexedPrimitiveIndirect(FRHIBuffer* IndexBuff
 
 	// Indirect args buffer can be a previously pending UAV, which becomes PS\Non-PS read. ApplyState will flush pending transitions, so enqueue the indirect
 	// arg transition and flush here.
-	FD3D12DynamicRHI::TransitionResource(CommandListHandle, Location.GetResource(), D3D12_RESOURCE_STATE_INDIRECT_ARGUMENT | D3D12_RESOURCE_STATE_PIXEL_SHADER_RESOURCE | D3D12_RESOURCE_STATE_NON_PIXEL_SHADER_RESOURCE, D3D12_RESOURCE_BARRIER_ALL_SUBRESOURCES);
+	FD3D12DynamicRHI::TransitionResource(CommandListHandle, Location.GetResource(), D3D12_RESOURCE_STATE_TBD, D3D12_RESOURCE_STATE_INDIRECT_ARGUMENT, D3D12_RESOURCE_BARRIER_ALL_SUBRESOURCES, FD3D12DynamicRHI::ETransitionMode::Validate);
 	CommandListHandle.FlushResourceBarriers();	// Must flush so the desired state is actually set.
 
 	CommandListHandle->ExecuteIndirect(
@@ -1835,19 +1848,6 @@ void FD3D12CommandContext::RHIClearMRTImpl(bool bClearColor, int32 NumClearColor
 	const bool ClearRTV = bClearColor && BoundRenderTargets.GetNumActiveTargets() > 0;
 	const bool ClearDSV = (bClearDepth || bClearStencil) && DepthStencilView;
 
-	if (ClearRTV)
-	{
-		for (int32 TargetIndex = 0; TargetIndex < BoundRenderTargets.GetNumActiveTargets(); TargetIndex++)
-		{
-			FD3D12RenderTargetView* RTView = BoundRenderTargets.GetRenderTargetView(TargetIndex);
-
-			if (RTView != nullptr)
-			{
-				FD3D12DynamicRHI::TransitionResource(CommandListHandle, RTView, D3D12_RESOURCE_STATE_RENDER_TARGET);
-			}
-		}
-	}
-
 	uint32 ClearFlags = 0;
 	if (ClearDSV)
 	{
@@ -1868,28 +1868,6 @@ void FD3D12CommandContext::RHIClearMRTImpl(bool bClearColor, int32 NumClearColor
 		{
 			UE_LOG(LogD3D12RHI, Warning, TEXT("RHIClearMRTImpl: Asking to clear a DSV that does not store stencil."));
 		}
-
-		if (bClearDepth && (!DepthStencilView->HasStencil() || bClearStencil))
-		{
-			// Transition the entire view (Both depth and stencil planes if applicable)
-			// Some DSVs don't have stencil bits.
-			FD3D12DynamicRHI::TransitionResource(CommandListHandle, DepthStencilView, D3D12_RESOURCE_STATE_DEPTH_WRITE);
-		}
-		else
-		{
-			if (bClearDepth)
-			{
-				// Transition just the depth plane
-				check(bClearDepth && !bClearStencil);
-				FD3D12DynamicRHI::TransitionResource(CommandListHandle, DepthStencilView->GetResource(), D3D12_RESOURCE_STATE_DEPTH_WRITE, DepthStencilView->GetDepthOnlyViewSubresourceSubset());
-			}
-			else
-			{
-				// Transition just the stencil plane
-				check(!bClearDepth && bClearStencil);
-				FD3D12DynamicRHI::TransitionResource(CommandListHandle, DepthStencilView->GetResource(), D3D12_RESOURCE_STATE_DEPTH_WRITE, DepthStencilView->GetStencilOnlyViewSubresourceSubset());
-			}
-		}
 	}
 
 	if (ClearRTV || ClearDSV)
@@ -2105,8 +2083,8 @@ void FD3D12CommandContext::RHIBroadcastTemporalEffect(const FName& InEffectName,
 	for (int32 i = 0; i < NumTextures; i++)
 	{
 		// Resources must be in the COMMON state before using on the copy queue.
-		FD3D12DynamicRHI::TransitionResource(CommandListHandle, SrcTextures[i]->GetResource(), D3D12_RESOURCE_STATE_COMMON, D3D12_RESOURCE_BARRIER_ALL_SUBRESOURCES);
-		FD3D12DynamicRHI::TransitionResource(CommandListHandle, DstTextures[i]->GetResource(), D3D12_RESOURCE_STATE_COMMON, D3D12_RESOURCE_BARRIER_ALL_SUBRESOURCES);
+		FD3D12DynamicRHI::TransitionResource(CommandListHandle, SrcTextures[i]->GetResource(), D3D12_RESOURCE_STATE_TBD, D3D12_RESOURCE_STATE_COMMON, D3D12_RESOURCE_BARRIER_ALL_SUBRESOURCES, FD3D12DynamicRHI::ETransitionMode::Apply);
+		FD3D12DynamicRHI::TransitionResource(CommandListHandle, DstTextures[i]->GetResource(), D3D12_RESOURCE_STATE_TBD, D3D12_RESOURCE_STATE_COMMON, D3D12_RESOURCE_BARRIER_ALL_SUBRESOURCES, FD3D12DynamicRHI::ETransitionMode::Apply);
 	}
 	CommandListHandle.FlushResourceBarriers();
 
@@ -2140,8 +2118,8 @@ void FD3D12CommandContext::RHIBroadcastTemporalEffect(const FName& InEffectName,
 
 	for (int32 i = 0; i < NumTextures; i++)
 	{
-		FD3D12DynamicRHI::TransitionResource(CommandListHandle, SrcTextures[i]->GetResource(), D3D12_RESOURCE_STATE_COPY_SOURCE, D3D12_RESOURCE_BARRIER_ALL_SUBRESOURCES);
-		FD3D12DynamicRHI::TransitionResource(CommandListHandle, DstTextures[i]->GetResource(), D3D12_RESOURCE_STATE_COPY_DEST, D3D12_RESOURCE_BARRIER_ALL_SUBRESOURCES);
+		FD3D12DynamicRHI::TransitionResource(CommandListHandle, SrcTextures[i]->GetResource(), D3D12_RESOURCE_STATE_TBD, D3D12_RESOURCE_STATE_COPY_SOURCE, D3D12_RESOURCE_BARRIER_ALL_SUBRESOURCES);
+		FD3D12DynamicRHI::TransitionResource(CommandListHandle, DstTextures[i]->GetResource(), D3D12_RESOURCE_STATE_TBD, D3D12_RESOURCE_STATE_COPY_DEST, D3D12_RESOURCE_BARRIER_ALL_SUBRESOURCES);
 	}
 	CommandListHandle.FlushResourceBarriers();
 
diff --git a/Engine/Source/Runtime/D3D12RHI/Private/D3D12DescriptorCache.cpp b/Engine/Source/Runtime/D3D12RHI/Private/D3D12DescriptorCache.cpp
--- a/Engine/Source/Runtime/D3D12RHI/Private/D3D12DescriptorCache.cpp
+++ b/Engine/Source/Runtime/D3D12RHI/Private/D3D12DescriptorCache.cpp
@@ -332,7 +332,7 @@ void FD3D12DescriptorCache::SetUAVs(const FD3D12RootSignature* RootSignature, FD
 		{
 			SrcDescriptors[SlotIndex] = UAVs[SlotIndex]->GetView();
 
-			FD3D12DynamicRHI::TransitionResource(CommandList, UAVs[SlotIndex], D3D12_RESOURCE_STATE_UNORDERED_ACCESS);
+			FD3D12DynamicRHI::TransitionResource(CommandList, UAVs[SlotIndex], D3D12_RESOURCE_STATE_UNORDERED_ACCESS, FD3D12DynamicRHI::ETransitionMode::Validate);
 			CommandList.UpdateResidency(Cache.ResidencyHandles[ShaderStage][SlotIndex]);
 		}
 	}
@@ -378,15 +378,15 @@ void FD3D12DescriptorCache::SetRenderTargets(FD3D12RenderTargetView** RenderTarg
 
 	FD3D12CommandListHandle& CommandList = CmdContext->CommandListHandle;
 
+	// RTV & DS stace should already be in the correct state. It is transitioned in RHISetRenderTargetsAndClear coming from BeginPass because
+	// then we know the correct depth & stencil read & write flags.
+
 	// Fill heap slots
 	for (uint32 i = 0; i < Count; i++)
 	{
 		if (RenderTargetViewArray[i] != NULL)
 		{
-			// RTV should already be in the correct state. It is transitioned in RHISetRenderTargets.
-			FD3D12DynamicRHI::TransitionResource(CommandList, RenderTargetViewArray[i], D3D12_RESOURCE_STATE_RENDER_TARGET);
 			RTVDescriptors[i] = RenderTargetViewArray[i]->GetView();
-
 			CommandList.UpdateResidency(RenderTargetViewArray[i]->GetResource());
 		}
 		else
@@ -397,8 +397,6 @@ void FD3D12DescriptorCache::SetRenderTargets(FD3D12RenderTargetView** RenderTarg
 
 	if (DepthStencilTarget != nullptr)
 	{
-		FD3D12DynamicRHI::TransitionResource(CommandList, DepthStencilTarget);
-
 		const D3D12_CPU_DESCRIPTOR_HANDLE DSVDescriptor = DepthStencilTarget->GetView();
 		CommandList->OMSetRenderTargets(Count, RTVDescriptors, 0, &DSVDescriptor);
 		CommandList.UpdateResidency(DepthStencilTarget->GetResource());
@@ -412,49 +410,7 @@ void FD3D12DescriptorCache::SetRenderTargets(FD3D12RenderTargetView** RenderTarg
 
 void FD3D12DescriptorCache::SetStreamOutTargets(FD3D12Resource** Buffers, uint32 Count, const uint32* Offsets)
 {
-	// Determine how many slots are really needed, since the Count passed in is a pre-defined maximum
-	uint32 SlotsNeeded = 0;
-	for (int32 i = Count - 1; i >= 0; i--)
-	{
-		if (Buffers[i] != NULL)
-		{
-			SlotsNeeded = i + 1;
-			break;
-		}
-	}
-
-	if (0 == SlotsNeeded)
-	{
-		return; // No-op
-	}
-
-	D3D12_STREAM_OUTPUT_BUFFER_VIEW SOViews[D3D12_SO_BUFFER_SLOT_COUNT] = { };
-
-	FD3D12CommandListHandle& CommandList = CmdContext->CommandListHandle;
-
-	// Fill heap slots
-	for (uint32 i = 0; i < SlotsNeeded; i++)
-	{
-		if (Buffers[i])
-		{
-			CommandList.UpdateResidency(Buffers[i]);
-		}
-
-		D3D12_STREAM_OUTPUT_BUFFER_VIEW &currentView = SOViews[i];
-		currentView.BufferLocation = (Buffers[i] != nullptr) ? Buffers[i]->GetGPUVirtualAddress() : 0;
-
-		// MS - The following view members are not correct
-		check(0);
-		currentView.BufferFilledSizeLocation = 0;
-		currentView.SizeInBytes = -1;
-
-		if (Buffers[i] != nullptr)
-		{
-			FD3D12DynamicRHI::TransitionResource(CommandList, Buffers[i], D3D12_RESOURCE_STATE_STREAM_OUT, D3D12_RESOURCE_BARRIER_ALL_SUBRESOURCES);
-		}
-	}
-
-	CommandList->SOSetTargets(0, SlotsNeeded, SOViews);
+	check(Count == 0);
 }
 
 template <EShaderFrequency ShaderStage>
@@ -596,19 +552,19 @@ void FD3D12DescriptorCache::SetSRVs(const FD3D12RootSignature* RootSignature, FD
 		{
 			SrcDescriptors[SlotIndex] = SRVs[SlotIndex]->GetView();
 
+			D3D12_RESOURCE_STATES State = (ShaderStage == SF_Compute) ? D3D12_RESOURCE_STATE_NON_PIXEL_SHADER_RESOURCE : D3D12_RESOURCE_STATE_NON_PIXEL_SHADER_RESOURCE | D3D12_RESOURCE_STATE_PIXEL_SHADER_RESOURCE;
 			if (SRVs[SlotIndex]->IsDepthStencilResource())
 			{
-				FD3D12DynamicRHI::TransitionResource(CommandList, SRVs[SlotIndex], (D3D12_RESOURCE_STATE_PIXEL_SHADER_RESOURCE | D3D12_RESOURCE_STATE_NON_PIXEL_SHADER_RESOURCE | D3D12_RESOURCE_STATE_DEPTH_READ) & ValidResourceStates);
+				if (GUseInternalTransitions)
+				{
+					State |= D3D12_RESOURCE_STATE_DEPTH_READ;
+				}
 			}
 			else if (SRVs[SlotIndex]->GetSkipFastClearFinalize())
 			{
-				D3D12_RESOURCE_STATES State = D3D12_RESOURCE_STATE_PIXEL_SHADER_RESOURCE | D3D12_RESOURCE_STATE_NON_PIXEL_SHADER_RESOURCE | CmdContext->SkipFastClearEliminateState;
-				FD3D12DynamicRHI::TransitionResource(CommandList, SRVs[SlotIndex], State & ValidResourceStates);
-			}
-			else
-			{
-				FD3D12DynamicRHI::TransitionResource(CommandList, SRVs[SlotIndex], (D3D12_RESOURCE_STATE_PIXEL_SHADER_RESOURCE | D3D12_RESOURCE_STATE_NON_PIXEL_SHADER_RESOURCE) & ValidResourceStates);
+				State |= CmdContext->SkipFastClearEliminateState;
 			}
+			FD3D12DynamicRHI::TransitionResource(CommandList, SRVs[SlotIndex], State & ValidResourceStates, FD3D12DynamicRHI::ETransitionMode::Validate);
 
 			CommandList.UpdateResidency(Cache.ResidencyHandles[ShaderStage][SlotIndex]);
 		}
diff --git a/Engine/Source/Runtime/D3D12RHI/Private/D3D12DirectCommandListManager.cpp b/Engine/Source/Runtime/D3D12RHI/Private/D3D12DirectCommandListManager.cpp
--- a/Engine/Source/Runtime/D3D12RHI/Private/D3D12DirectCommandListManager.cpp
+++ b/Engine/Source/Runtime/D3D12RHI/Private/D3D12DirectCommandListManager.cpp
@@ -663,11 +663,6 @@ void FD3D12CommandListManager::ExecuteCommandListInteral(TArray<FD3D12CommandLis
 	FD3D12CommandListHandle BarrierCommandList[128];
 	if (NeedsResourceBarriers)
 	{
-#if !USE_D3D12RHI_RESOURCE_STATE_TRACKING
-		// If we're using the engine's resource state tracking and barriers, then we should never have pending resource barriers.
-		check(false);
-#endif // !USE_D3D12RHI_RESOURCE_STATE_TRACKING
-
 #if UE_BUILD_DEBUG	
 		if (ResourceStateCS.TryLock())
 		{
@@ -910,7 +905,10 @@ uint32 FD3D12CommandListManager::GetResourceBarrierCommandList(FD3D12CommandList
 
 			Desc.Transition.Subresource = PRB.SubResource;
 			const D3D12_RESOURCE_STATES Before = ResourceState.GetSubresourceState(Desc.Transition.Subresource);
-			const D3D12_RESOURCE_STATES After = PRB.State;
+
+			// If state unknown then we don't enqueue a transition - only want to update the end state on the resource
+			// and not really enqueue a transition
+			const D3D12_RESOURCE_STATES After = (PRB.State != D3D12_RESOURCE_STATE_TBD) ? PRB.State : Before;
 
 			check(Before != D3D12_RESOURCE_STATE_TBD && Before != D3D12_RESOURCE_STATE_CORRUPT);
 			if (Before != After)
diff --git a/Engine/Source/Runtime/D3D12RHI/Private/D3D12IndexBuffer.cpp b/Engine/Source/Runtime/D3D12RHI/Private/D3D12IndexBuffer.cpp
--- a/Engine/Source/Runtime/D3D12RHI/Private/D3D12IndexBuffer.cpp
+++ b/Engine/Source/Runtime/D3D12RHI/Private/D3D12IndexBuffer.cpp
@@ -42,7 +42,7 @@ FIndexBufferRHIRef FD3D12DynamicRHI::RHICreateIndexBuffer(uint32 Stride, uint32
 	const D3D12_RESOURCE_DESC Desc = CreateIndexBufferResourceDesc(Size, InUsage);
 	const uint32 Alignment = 4;
 
-	FD3D12Buffer* Buffer = GetAdapter().CreateRHIBuffer(nullptr, Desc, Alignment, Stride, Size, InUsage | BUF_IndexBuffer, ED3D12ResourceStateMode::Default, CreateInfo);
+	FD3D12Buffer* Buffer = GetAdapter().CreateRHIBuffer(nullptr, Desc, Alignment, Stride, Size, InUsage | BUF_IndexBuffer, ED3D12ResourceStateMode::Default, InResourceState, CreateInfo);
 	if (Buffer->ResourceLocation.IsTransient())
 	{
 		// TODO: this should ideally be set in platform-independent code, since this tracking is for the high level
@@ -92,7 +92,7 @@ FIndexBufferRHIRef FD3D12DynamicRHI::CreateIndexBuffer_RenderThread(class FRHICo
 	const D3D12_RESOURCE_DESC Desc = CreateIndexBufferResourceDesc(Size, InUsage);
 	const uint32 Alignment = 4;
 
-	FD3D12Buffer* Buffer = GetAdapter().CreateRHIBuffer(&RHICmdList, Desc, Alignment, Stride, Size, InUsage | BUF_IndexBuffer, ED3D12ResourceStateMode::Default, CreateInfo);
+	FD3D12Buffer* Buffer = GetAdapter().CreateRHIBuffer(&RHICmdList, Desc, Alignment, Stride, Size, InUsage | BUF_IndexBuffer, ED3D12ResourceStateMode::Default, InResourceState, CreateInfo);
 	if (Buffer->ResourceLocation.IsTransient())
 	{
 		// TODO: this should ideally be set in platform-independent code, since this tracking is for the high level
@@ -107,7 +107,7 @@ FIndexBufferRHIRef FD3D12DynamicRHI::CreateAndLockIndexBuffer_RenderThread(class
 	const D3D12_RESOURCE_DESC Desc = CreateIndexBufferResourceDesc(Size, InUsage);
 	const uint32 Alignment = 4;
 
-	FD3D12Buffer* Buffer = GetAdapter().CreateRHIBuffer(&RHICmdList, Desc, Alignment, Stride, Size, InUsage | BUF_IndexBuffer, ED3D12ResourceStateMode::Default, CreateInfo);
+	FD3D12Buffer* Buffer = GetAdapter().CreateRHIBuffer(&RHICmdList, Desc, Alignment, Stride, Size, InUsage, ED3D12ResourceStateMode::Default, InResourceState, CreateInfo);
 	if (Buffer->ResourceLocation.IsTransient())
 	{
 		// TODO: this should ideally be set in platform-independent code, since this tracking is for the high level
diff --git a/Engine/Source/Runtime/D3D12RHI/Private/D3D12RHI.cpp b/Engine/Source/Runtime/D3D12RHI/Private/D3D12RHI.cpp
--- a/Engine/Source/Runtime/D3D12RHI/Private/D3D12RHI.cpp
+++ b/Engine/Source/Runtime/D3D12RHI/Private/D3D12RHI.cpp
@@ -329,7 +329,7 @@ void FD3D12DynamicRHI::UpdateBuffer(FD3D12Resource* Dest, uint32 DestOffset, FD3
 	FD3D12CommandContext& DefaultContext = Device->GetDefaultCommandContext();
 	FD3D12CommandListHandle& hCommandList = DefaultContext.CommandListHandle;
 
-	FConditionalScopeResourceBarrier ScopeResourceBarrierDest(hCommandList, Dest, D3D12_RESOURCE_STATE_COPY_DEST, 0);
+	FScopedResourceBarrier ScopeResourceBarrierDest(hCommandList, Dest, D3D12_RESOURCE_STATE_COPY_DEST, 0, FD3D12DynamicRHI::ETransitionMode::Apply);
 	// Don't need to transition upload heaps
 
 	DefaultContext.numCopies++;
@@ -575,91 +575,98 @@ uint64 FD3D12SubmissionGapRecorder::SubmitSubmissionTimestampsForFrame(uint32 Fr
 #endif
 		}
 
-		// Store the timestamp values
-		for (int i = 0; i < PrevFrameBeginSubmissionTimestamps.Num() - 1; i++)
+		if (PrevFrameBeginSubmissionTimestamps.Num() > 0)
 		{
-			FGapSpan GapSpan;
+			// Store the timestamp values
+			for (int i = 0; i < PrevFrameBeginSubmissionTimestamps.Num() - 1; i++)
+			{
+				FGapSpan GapSpan;
 
-			uint64 BeginTimestampPtr = PrevFrameEndSubmissionTimestamps[i];
-			uint64 EndTimestampPtr = PrevFrameBeginSubmissionTimestamps[i + 1];
+				uint64 BeginTimestampPtr = PrevFrameEndSubmissionTimestamps[i];
+				uint64 EndTimestampPtr = PrevFrameBeginSubmissionTimestamps[i + 1];
 
-			GapSpan.BeginCycles = BeginTimestampPtr;
-			uint64 EndCycles = EndTimestampPtr;
+				GapSpan.BeginCycles = BeginTimestampPtr;
+				uint64 EndCycles = EndTimestampPtr;
 
-			// Check begin/end is contiguous
-			if (EndCycles < GapSpan.BeginCycles)
-			{
+				// Check begin/end is contiguous
+				if (EndCycles < GapSpan.BeginCycles)
+				{
 #if D3D12_SUBMISSION_GAP_RECORDER_DEBUG_INFO
-				UE_LOG(LogD3D12GapRecorder, Verbose, TEXT("SubmitSubmissionTimestampsForFrame EndCycles occurs before BeginCycles not valid"));
+					UE_LOG(LogD3D12GapRecorder, Verbose, TEXT("SubmitSubmissionTimestampsForFrame EndCycles occurs before BeginCycles not valid"));
 #endif
-				bValid = false;
-				break;
-			}
-			GapSpan.DurationCycles = EndCycles - GapSpan.BeginCycles;
+					bValid = false;
+					break;
+				}
+				GapSpan.DurationCycles = EndCycles - GapSpan.BeginCycles;
 
-			UE_LOG(LogD3D12GapRecorder, Verbose, TEXT("GapSpan Begin %lu End %lu Duration %lu"),GapSpan.BeginCycles,EndCycles,GapSpan.DurationCycles);
+				UE_LOG(LogD3D12GapRecorder, Verbose, TEXT("GapSpan Begin %lu End %lu Duration %lu"), GapSpan.BeginCycles, EndCycles, GapSpan.DurationCycles);
 
-			// Check gap spans are contiguous (TODO: we might want to modify this to support async compute submissions which overlap)
-			if (i > 0)
-			{
-				const FGapSpan& PrevGap = Frame.GapSpans[i - 1];
-				uint64 PrevGapEndCycles = PrevGap.BeginCycles + PrevGap.DurationCycles;
-				if (GapSpan.BeginCycles < PrevGapEndCycles)
+				// Check gap spans are contiguous (TODO: we might want to modify this to support async compute submissions which overlap)
+				if (i > 0)
 				{
-					UE_LOG(LogD3D12GapRecorder, Verbose, TEXT("SubmitSubmissionTimestampsForFrame Gap Span Begin Cycle is later than Prev Gap Cycle End not valid"));
-					bValid = false;
-					break;
+					const FGapSpan& PrevGap = Frame.GapSpans[i - 1];
+					uint64 PrevGapEndCycles = PrevGap.BeginCycles + PrevGap.DurationCycles;
+					if (GapSpan.BeginCycles < PrevGapEndCycles)
+					{
+						UE_LOG(LogD3D12GapRecorder, Verbose, TEXT("SubmitSubmissionTimestampsForFrame Gap Span Begin Cycle is later than Prev Gap Cycle End not valid"));
+						bValid = false;
+						break;
+					}
 				}
-			}
 
-			TotalWaitCycles += GapSpan.DurationCycles;
+				TotalWaitCycles += GapSpan.DurationCycles;
 
-			Frame.GapSpans.Add(GapSpan);
-		}
+				Frame.GapSpans.Add(GapSpan);
+			}
 
 #if D3D12_SUBMISSION_GAP_RECORDER_DEBUG_INFO
-		float Timing = (float)FGPUTiming::GetTimingFrequency();
+			float Timing = (float)FGPUTiming::GetTimingFrequency();
 
-		uint64 CurrSpan = 0;
-		uint64 TotalDuration = 0;
+			uint64 CurrSpan = 0;
+			uint64 TotalDuration = 0;
 
-		for (int i = 0; i < PrevFrameBeginSubmissionTimestamps.Num(); i++)
-		{
-			CurrSpan = PrevFrameEndSubmissionTimestamps[i] - PrevFrameBeginSubmissionTimestamps[i];
+			for (int i = 0; i < PrevFrameBeginSubmissionTimestamps.Num(); i++)
+			{
+				CurrSpan = PrevFrameEndSubmissionTimestamps[i] - PrevFrameBeginSubmissionTimestamps[i];
 
-			double CurrSpanSeconds = (CurrSpan / Timing);
-			double CurrSpanOutputTime = FMath::TruncToInt(CurrSpanSeconds / FPlatformTime::GetSecondsPerCycle());
+				double CurrSpanSeconds = (CurrSpan / Timing);
+				double CurrSpanOutputTime = FMath::TruncToInt(CurrSpanSeconds / FPlatformTime::GetSecondsPerCycle());
 
-			UE_LOG(LogD3D12GapRecorder, Verbose, TEXT("Total GPU Duration for span Begin %lu End %lu Duration %lu Seconds %f"),
-				PrevFrameBeginSubmissionTimestamps[i],
-				PrevFrameEndSubmissionTimestamps[i],
-				CurrSpan,
-				(CurrSpanSeconds * 1000.0f));
-			TotalDuration += CurrSpan;
-		}
+				UE_LOG(LogD3D12GapRecorder, Verbose, TEXT("Total GPU Duration for span Begin %lu End %lu Duration %lu Seconds %f"),
+					PrevFrameBeginSubmissionTimestamps[i],
+					PrevFrameEndSubmissionTimestamps[i],
+					CurrSpan,
+					(CurrSpanSeconds * 1000.0f));
+				TotalDuration += CurrSpan;
+			}
 
-		int32 len = PrevFrameEndSubmissionTimestamps.Num() - 1;
-		uint64 tbegin = PrevFrameBeginSubmissionTimestamps[0];
-		uint64 tend = PrevFrameEndSubmissionTimestamps[len];
-		uint64 duration = tend - tbegin;
-		double seconds = (duration / Timing);
-		double TotalDurationSeconds = (TotalDuration / Timing);
-
-		UE_LOG(LogD3D12GapRecorder, Verbose, TEXT("Total GPU Duration for all Timestamps for Frame %u Cycles %lu Timing %f Milliseconds %f"),
-			FrameNumber,
-			TotalDuration,
-			Timing,
-			TotalDurationSeconds);
-
-		UE_LOG(LogD3D12GapRecorder, Verbose, TEXT("Total GPU Duration from StartTimestamp %lu to EndTimestamp %lu Duration %lu MilliSeconds %f Timing %f"),
-			tbegin, 
-			tend, 
-			duration,
-			seconds,
-			Timing);
-
-		CSV_CUSTOM_STAT_GLOBAL(GPUTimestamps, float(TotalDurationSeconds * 1000.0f), ECsvCustomStatOp::Set);
+			int32 len = PrevFrameEndSubmissionTimestamps.Num() - 1;
+			uint64 tbegin = PrevFrameBeginSubmissionTimestamps[0];
+			uint64 tend = PrevFrameEndSubmissionTimestamps[len];
+			uint64 duration = tend - tbegin;
+			double seconds = (duration / Timing);
+			double TotalDurationSeconds = (TotalDuration / Timing);
+
+			UE_LOG(LogD3D12GapRecorder, Verbose, TEXT("Total GPU Duration for all Timestamps for Frame %u Cycles %lu Timing %f Milliseconds %f"),
+				FrameNumber,
+				TotalDuration,
+				Timing,
+				TotalDurationSeconds);
+
+			UE_LOG(LogD3D12GapRecorder, Verbose, TEXT("Total GPU Duration from StartTimestamp %lu to EndTimestamp %lu Duration %lu MilliSeconds %f Timing %f"),
+				tbegin,
+				tend,
+				duration,
+				seconds,
+				Timing);
+
+			CSV_CUSTOM_STAT_GLOBAL(GPUTimestamps, float(TotalDurationSeconds * 1000.0f), ECsvCustomStatOp::Set);
 #endif
+		}
+		else
+		{
+			bValid = false;
+		}
 	}
 
 	UE_LOG(LogD3D12GapRecorder, Verbose, TEXT("SubmitSubmissionTimestampsForFrame Frame %u FN %u TotalWaitCycles %lu"), FrameCounter, FrameNumber, TotalWaitCycles);
diff --git a/Engine/Source/Runtime/D3D12RHI/Private/D3D12RHIPrivate.h b/Engine/Source/Runtime/D3D12RHI/Private/D3D12RHIPrivate.h
--- a/Engine/Source/Runtime/D3D12RHI/Private/D3D12RHIPrivate.h
+++ b/Engine/Source/Runtime/D3D12RHI/Private/D3D12RHIPrivate.h
@@ -86,8 +86,6 @@ typedef FD3D12StateCacheBase FD3D12StateCache;
 #include "D3D12Device.h"
 #include "D3D12Adapter.h"
 
-// Definitions.
-#define USE_D3D12RHI_RESOURCE_STATE_TRACKING 1	// Fully relying on the engine's resource barriers is a work in progress. For now, continue to use the D3D12 RHI's resource state tracking.
 
 #define EXECUTE_DEBUG_COMMAND_LISTS 0
 #define DEBUG_BARRIER_DISPATCHES 1
@@ -149,6 +147,11 @@ extern bool GIsDoingQuery;
 #define DEBUG_RHI_EXECUTE_COMMAND_LIST(scope) 
 #endif
 
+// Use the D3D12 RHI internal transitions to drive all resource transitions
+extern bool GUseInternalTransitions;
+// Use the D3D12 RHI internal transitions to validate the engine pushed RHI transitions
+extern bool GValidateInternalTransitions;
+
 template< typename t_A, typename t_B >
 inline t_A RoundUpToNextMultiple(const t_A& a, const t_B& b)
 {
@@ -552,7 +555,7 @@ public:
 		return !pLeftView->DoesNotOverlap(*pRightView);
 	}
 
-	static inline bool IsTransitionNeeded(D3D12_RESOURCE_STATES Before, D3D12_RESOURCE_STATES& After)
+	static inline bool IsTransitionNeeded(bool bInAllowStateMerging, D3D12_RESOURCE_STATES Before, D3D12_RESOURCE_STATES& After)
 	{
 		check(Before != D3D12_RESOURCE_STATE_CORRUPT && After != D3D12_RESOURCE_STATE_CORRUPT);
 		check(Before != D3D12_RESOURCE_STATE_TBD && After != D3D12_RESOURCE_STATE_TBD);
@@ -567,29 +570,40 @@ public:
 		// having exactly one bit set so we need to special case these
 		if (After == D3D12_RESOURCE_STATE_COMMON)
 		{
-			// The resource state tracking code in FD3D12CommandContext::RHITransitionResources forces all EReadable transitions
-			// to go through the COMMON state right now, so we can end up with some COMMON -> COMMON transitions which can be
-			// skipped. Once that is fixed or removed, we shouldn't get here anymore if we're already in the COMMON state,
-			// so we can simply return true and let the ensure in FD3D12CommandListHandle::AddTransitionBarrier catch bad usage.
-			return (Before != D3D12_RESOURCE_STATE_COMMON);
+			// Before state should not have the common state otherwise it's invalid transition
+			check(Before != D3D12_RESOURCE_STATE_COMMON);
+			return true;
 		}
 
-		// We should avoid doing read-to-read state transitions. But when we do, we should avoid turning off already transitioned bits,
-		// e.g. VERTEX_BUFFER -> SHADER_RESOURCE is turned into VERTEX_BUFFER -> VERTEX_BUFFER | SHADER_RESOURCE.
-		// This reduces the number of resource transitions and ensures automatic states from resource bindings get properly combined.
-		D3D12_RESOURCE_STATES Combined = Before | After;
-		if ((Combined & (D3D12_RESOURCE_STATE_GENERIC_READ | D3D12_RESOURCE_STATE_INDIRECT_ARGUMENT)) == Combined)
+		if (bInAllowStateMerging)
 		{
-			After = Combined;
+			// We should avoid doing read-to-read state transitions. But when we do, we should avoid turning off already transitioned bits,
+			// e.g. VERTEX_BUFFER -> SHADER_RESOURCE is turned into VERTEX_BUFFER -> VERTEX_BUFFER | SHADER_RESOURCE.
+			// This reduces the number of resource transitions and ensures automatic states from resource bindings get properly combined.
+			D3D12_RESOURCE_STATES Combined = Before | After;
+			if ((Combined & (D3D12_RESOURCE_STATE_GENERIC_READ | D3D12_RESOURCE_STATE_INDIRECT_ARGUMENT)) == Combined)
+			{
+				After = Combined;
+			}
 		}
 
 		return Before != After;
 	}
 
+	enum class ETransitionMode
+	{
+		Apply,
+		Validate
+	};
+
 	/** Transition a resource's state based on a Render target view */
-	static inline void TransitionResource(FD3D12CommandListHandle& hCommandList, FD3D12RenderTargetView* pView, D3D12_RESOURCE_STATES after)
+	static inline void TransitionResource(FD3D12CommandListHandle& hCommandList, FD3D12RenderTargetView* pView, D3D12_RESOURCE_STATES after, ETransitionMode InMode)
 	{
-#if USE_D3D12RHI_RESOURCE_STATE_TRACKING
+		// Early out if we are not using engine transitions and not validating them
+		check(InMode == ETransitionMode::Validate);
+		if (!GUseInternalTransitions && !GValidateInternalTransitions)
+			return;
+
 		FD3D12Resource* pResource = pView->GetResource();
 
 		const D3D12_RENDER_TARGET_VIEW_DESC &desc = pView->GetDesc();
@@ -601,13 +615,13 @@ public:
 		case D3D12_RTV_DIMENSION_TEXTURE2D:
 		case D3D12_RTV_DIMENSION_TEXTURE2DMS:
 			// Only one subresource to transition
-			TransitionResource(hCommandList, pResource, after, desc.Texture2D.MipSlice);
+			TransitionResource(hCommandList, pResource, D3D12_RESOURCE_STATE_TBD, after, desc.Texture2D.MipSlice, InMode);
 			break;
 
 		case D3D12_RTV_DIMENSION_TEXTURE2DARRAY:
 		{
 			// Multiple subresources to transition
-			TransitionResource(hCommandList, pResource, after, pView->GetViewSubresourceSubset());
+			TransitionResource(hCommandList, pResource, D3D12_RESOURCE_STATE_TBD, after, pView->GetViewSubresourceSubset(), InMode);
 			break;
 		}
 
@@ -615,13 +629,16 @@ public:
 			check(false);	// Need to update this code to include the view type
 			break;
 		}
-#endif // USE_D3D12RHI_RESOURCE_STATE_TRACKING
 	}
 
 	/** Transition a resource's state based on a Depth stencil view's desc flags */
-	static inline void TransitionResource(FD3D12CommandListHandle& hCommandList, FD3D12DepthStencilView* pView)
+	static inline void TransitionResource(FD3D12CommandListHandle& hCommandList, FD3D12DepthStencilView* pView, ETransitionMode InMode)
 	{
-#if USE_D3D12RHI_RESOURCE_STATE_TRACKING
+		// Early out if we are not using engine transitions and not validating them
+		check(InMode == ETransitionMode::Validate);
+		if (!GUseInternalTransitions && !GValidateInternalTransitions)
+			return;
+
 		// Determine the required subresource states from the view desc
 		const D3D12_DEPTH_STENCIL_VIEW_DESC& DSVDesc = pView->GetDesc();
 		const bool bDSVDepthIsWritable = (DSVDesc.Flags & D3D12_DSV_FLAG_READ_ONLY_DEPTH) == 0;
@@ -639,20 +656,23 @@ public:
 		FD3D12Resource* pResource = pView->GetResource();
 		if (bDepthIsWritable)
 		{
-			TransitionResource(hCommandList, pResource, D3D12_RESOURCE_STATE_DEPTH_WRITE, pView->GetDepthOnlyViewSubresourceSubset());
+			TransitionResource(hCommandList, pResource, D3D12_RESOURCE_STATE_TBD, D3D12_RESOURCE_STATE_DEPTH_WRITE, pView->GetDepthOnlyViewSubresourceSubset(), InMode);
 		}
 
 		if (bStencilIsWritable)
 		{
-			TransitionResource(hCommandList, pResource, D3D12_RESOURCE_STATE_DEPTH_WRITE, pView->GetStencilOnlyViewSubresourceSubset());
+			TransitionResource(hCommandList, pResource, D3D12_RESOURCE_STATE_TBD, D3D12_RESOURCE_STATE_DEPTH_WRITE, pView->GetStencilOnlyViewSubresourceSubset(), InMode);
 		}
-#endif // USE_D3D12RHI_RESOURCE_STATE_TRACKING
 	}
 
 	/** Transition a resource's state based on a Depth stencil view */
-	static inline void TransitionResource(FD3D12CommandListHandle& hCommandList, FD3D12DepthStencilView* pView, D3D12_RESOURCE_STATES after)
+	static inline void TransitionResource(FD3D12CommandListHandle& hCommandList, FD3D12DepthStencilView* pView, D3D12_RESOURCE_STATES after, ETransitionMode InMode)
 	{
-#if USE_D3D12RHI_RESOURCE_STATE_TRACKING
+		// Early out if we are not using engine transitions and not validating them
+		check(InMode == ETransitionMode::Validate);
+		if (!GUseInternalTransitions && !GValidateInternalTransitions)
+			return;
+
 		FD3D12Resource* pResource = pView->GetResource();
 
 		const D3D12_DEPTH_STENCIL_VIEW_DESC &desc = pView->GetDesc();
@@ -663,14 +683,14 @@ public:
 			if (pResource->GetPlaneCount() > 1)
 			{
 				// Multiple subresources to transtion
-				TransitionResource(hCommandList, pResource, after, pView->GetViewSubresourceSubset());
+				TransitionResource(hCommandList, pResource, D3D12_RESOURCE_STATE_TBD, after, pView->GetViewSubresourceSubset(), InMode);
 				break;
 			}
 			else
 			{
 				// Only one subresource to transition
 				check(pResource->GetPlaneCount() == 1);
-				TransitionResource(hCommandList, pResource, after, desc.Texture2D.MipSlice);
+				TransitionResource(hCommandList, pResource, D3D12_RESOURCE_STATE_TBD, after, desc.Texture2D.MipSlice, InMode);
 			}
 			break;
 
@@ -678,7 +698,7 @@ public:
 		case D3D12_DSV_DIMENSION_TEXTURE2DMSARRAY:
 		{
 			// Multiple subresources to transtion
-			TransitionResource(hCommandList, pResource, after, pView->GetViewSubresourceSubset());
+			TransitionResource(hCommandList, pResource, D3D12_RESOURCE_STATE_TBD, after, pView->GetViewSubresourceSubset(), InMode);
 			break;
 		}
 
@@ -686,37 +706,39 @@ public:
 			check(false);	// Need to update this code to include the view type
 			break;
 		}
-#endif // USE_D3D12RHI_RESOURCE_STATE_TRACKING
 	}
 
 	/** Transition a resource's state based on a Unordered access view */
-	static inline void TransitionResource(FD3D12CommandListHandle& hCommandList, FD3D12UnorderedAccessView* pView, D3D12_RESOURCE_STATES after)
+	static inline void TransitionResource(FD3D12CommandListHandle& hCommandList, FD3D12UnorderedAccessView* pView, D3D12_RESOURCE_STATES after, ETransitionMode InMode)
 	{
-#if USE_D3D12RHI_RESOURCE_STATE_TRACKING
+		// Early out if we are not using engine transitions and not validating them
+		if (!GUseInternalTransitions && !GValidateInternalTransitions)
+			return;
+
 		FD3D12Resource* pResource = pView->GetResource();
 
 		const D3D12_UNORDERED_ACCESS_VIEW_DESC &desc = pView->GetDesc();
 		switch (desc.ViewDimension)
 		{
 		case D3D12_UAV_DIMENSION_BUFFER:
-			TransitionResource(hCommandList, pResource, after, 0);
+			TransitionResource(hCommandList, pResource, D3D12_RESOURCE_STATE_TBD, after, 0, InMode);
 			break;
 
 		case D3D12_UAV_DIMENSION_TEXTURE2D:
 			// Only one subresource to transition
-			TransitionResource(hCommandList, pResource, after, desc.Texture2D.MipSlice);
+			TransitionResource(hCommandList, pResource, D3D12_RESOURCE_STATE_TBD, after, desc.Texture2D.MipSlice, InMode);
 			break;
 
 		case D3D12_UAV_DIMENSION_TEXTURE2DARRAY:
 		{
 			// Multiple subresources to transtion
-			TransitionResource(hCommandList, pResource, after, pView->GetViewSubresourceSubset());
+			TransitionResource(hCommandList, pResource, D3D12_RESOURCE_STATE_TBD, after, pView->GetViewSubresourceSubset(), InMode);
 			break;
 		}
 		case D3D12_UAV_DIMENSION_TEXTURE3D:
 		{
 			// Multiple subresources to transtion
-			TransitionResource(hCommandList, pResource, after, pView->GetViewSubresourceSubset());
+			TransitionResource(hCommandList, pResource, D3D12_RESOURCE_STATE_TBD, after, pView->GetViewSubresourceSubset(), InMode);
 			break;
 		}
 
@@ -724,13 +746,15 @@ public:
 			check(false);	// Need to update this code to include the view type
 			break;
 		}
-#endif // USE_D3D12RHI_RESOURCE_STATE_TRACKING
 	}
 
 	/** Transition a resource's state based on a Shader resource view */
-	static inline void TransitionResource(FD3D12CommandListHandle& hCommandList, FD3D12ShaderResourceView* pView, D3D12_RESOURCE_STATES after)
+	static inline void TransitionResource(FD3D12CommandListHandle& hCommandList, FD3D12ShaderResourceView* pView, D3D12_RESOURCE_STATES after, ETransitionMode InMode)
 	{
-#if USE_D3D12RHI_RESOURCE_STATE_TRACKING
+		// Early out if we are not using engine transitions and not validating them
+		if (!GUseInternalTransitions && !GValidateInternalTransitions)
+			return;
+
 		FD3D12Resource* pResource = pView->GetResource();
 
 		if (!pResource || !pResource->RequiresResourceStateTracking())
@@ -748,7 +772,7 @@ public:
 		default:
 		{
 			// Transition the resource
-			TransitionResource(hCommandList, pResource, after, subresourceSubset);
+			TransitionResource(hCommandList, pResource, D3D12_RESOURCE_STATE_TBD, after, subresourceSubset, InMode);
 			break;
 		}
 
@@ -757,37 +781,40 @@ public:
 			if (pResource->GetHeapType() == D3D12_HEAP_TYPE_DEFAULT)
 			{
 				// Transition the resource
-				TransitionResource(hCommandList, pResource, after, subresourceSubset);
+				TransitionResource(hCommandList, pResource, D3D12_RESOURCE_STATE_TBD, after, subresourceSubset, InMode);
 			}
 			break;
 		}
 		}
-#endif // USE_D3D12RHI_RESOURCE_STATE_TRACKING
 	}
 
 	// Transition a specific subresource to the after state.
-	static inline void TransitionResource(FD3D12CommandListHandle& hCommandList, FD3D12Resource* pResource, D3D12_RESOURCE_STATES after, uint32 subresource)
+	// Return true if UAV barrier is required
+	static inline bool TransitionResource(FD3D12CommandListHandle& hCommandList, FD3D12Resource* pResource, D3D12_RESOURCE_STATES before, D3D12_RESOURCE_STATES after, uint32 subresource, ETransitionMode InMode)
 	{
-#if USE_D3D12RHI_RESOURCE_STATE_TRACKING
-		TransitionResourceWithTracking(hCommandList, pResource, after, subresource);
-#endif // USE_D3D12RHI_RESOURCE_STATE_TRACKING
+		// Early out if we are not using engine transitions and not validating them
+		if (InMode == ETransitionMode::Validate && !GUseInternalTransitions && !GValidateInternalTransitions)
+			return false;
+
+		return TransitionResourceWithTracking(hCommandList, pResource, before, after, subresource, InMode);
 	}
 
 	// Transition a subset of subresources to the after state.
-	static inline void TransitionResource(FD3D12CommandListHandle& hCommandList, FD3D12Resource* pResource, D3D12_RESOURCE_STATES after, const CViewSubresourceSubset& subresourceSubset)
+	// Return true if UAV barrier is required
+	static inline bool TransitionResource(FD3D12CommandListHandle& hCommandList, FD3D12Resource* pResource, D3D12_RESOURCE_STATES before, D3D12_RESOURCE_STATES after, const CViewSubresourceSubset& subresourceSubset, ETransitionMode InMode)
 	{
-#if USE_D3D12RHI_RESOURCE_STATE_TRACKING
-		TransitionResourceWithTracking(hCommandList, pResource, after, subresourceSubset);
-#endif // USE_D3D12RHI_RESOURCE_STATE_TRACKING
+		// Early out if we are not using engine transitions and not validating them
+		if (InMode == ETransitionMode::Validate && !GUseInternalTransitions && !GValidateInternalTransitions)
+			return false;
+
+		return TransitionResourceWithTracking(hCommandList, pResource, before, after, subresourceSubset, InMode);
 	}
 
 	// Transition a subresource from current to a new state, using resource state tracking.
-	static void TransitionResourceWithTracking(FD3D12CommandListHandle& hCommandList, FD3D12Resource* pResource, D3D12_RESOURCE_STATES after, uint32 subresource)
+	static bool TransitionResourceWithTracking(FD3D12CommandListHandle& hCommandList, FD3D12Resource* pResource, D3D12_RESOURCE_STATES before, D3D12_RESOURCE_STATES after, uint32 subresource, ETransitionMode InMode)
 	{
-#if USE_D3D12RHI_RESOURCE_STATE_TRACKING
 		check(pResource);
 		check(pResource->RequiresResourceStateTracking());
-
 		check(!((after & (D3D12_RESOURCE_STATE_PIXEL_SHADER_RESOURCE | D3D12_RESOURCE_STATE_NON_PIXEL_SHADER_RESOURCE)) && (pResource->GetDesc().Flags & D3D12_RESOURCE_FLAG_DENY_SHADER_RESOURCE)));
 
 #ifdef PLATFORM_SUPPORTS_RESOURCE_COMPRESSION
@@ -796,6 +823,8 @@ public:
 
 		hCommandList.UpdateResidency(pResource);
 
+		bool bRequireUAVBarrier = false;
+
 		CResourceState& ResourceState = hCommandList.GetResourceState(pResource);
 		if (subresource == D3D12_RESOURCE_BARRIER_ALL_SUBRESOURCES && !ResourceState.AreAllSubresourcesSame())
 		{
@@ -804,51 +833,28 @@ public:
 			const uint8 SubresourceCount = pResource->GetSubresourceCount();
 			for (uint32 SubresourceIndex = 0; SubresourceIndex < SubresourceCount; SubresourceIndex++)
 			{
-				const D3D12_RESOURCE_STATES before = ResourceState.GetSubresourceState(SubresourceIndex);
-				if (before == D3D12_RESOURCE_STATE_TBD)
-				{
-					// We need a pending resource barrier so we can setup the state before this command list executes
-					hCommandList.AddPendingResourceBarrier(pResource, after, SubresourceIndex);
-					ResourceState.SetSubresourceState(SubresourceIndex, after);
-				}
-				// We're not using IsTransitionNeeded() because we do want to transition even if 'after' is a subset of 'before'
-				// This is so that we can ensure all subresources are in the same state, simplifying future barriers
-				else if (before != after)
-				{
-					hCommandList.AddTransitionBarrier(pResource, before, after, SubresourceIndex);
-					ResourceState.SetSubresourceState(SubresourceIndex, after);
-				}
+				bool bForceInAfterState = true;
+				bRequireUAVBarrier |= ValidateAndSetResourceState(hCommandList, pResource, ResourceState, SubresourceIndex, before, after, bForceInAfterState, InMode);
 			}
 
 			// The entire resource should now be in the after state on this command list (even if all barriers are pending)
-			check(ResourceState.CheckResourceState(after));
-			ResourceState.SetResourceState(after);
+			verify(ResourceState.CheckAllSubresourceSame());
+			check(EnumHasAllFlags(ResourceState.GetSubresourceState(D3D12_RESOURCE_BARRIER_ALL_SUBRESOURCES), after));
 		}
 		else
 		{
-			const D3D12_RESOURCE_STATES before = ResourceState.GetSubresourceState(subresource);
-			if (before == D3D12_RESOURCE_STATE_TBD)
-			{
-				// We need a pending resource barrier so we can setup the state before this command list executes
-				hCommandList.AddPendingResourceBarrier(pResource, after, subresource);
-				ResourceState.SetSubresourceState(subresource, after);
-			}
-			else if (IsTransitionNeeded(before, after))
-			{
-				hCommandList.AddTransitionBarrier(pResource, before, after, subresource);
-				ResourceState.SetSubresourceState(subresource, after);
-			}
+			bool bForceInAfterState = false;
+			bRequireUAVBarrier = ValidateAndSetResourceState(hCommandList, pResource, ResourceState, subresource, before, after, bForceInAfterState, InMode);
 		}
-#endif // USE_D3D12RHI_RESOURCE_STATE_TRACKING
+
+		return bRequireUAVBarrier;
 	}
 
 	// Transition subresources from current to a new state, using resource state tracking.
-	static void TransitionResourceWithTracking(FD3D12CommandListHandle& hCommandList, FD3D12Resource* pResource, D3D12_RESOURCE_STATES after, const CViewSubresourceSubset& subresourceSubset)
+	static bool TransitionResourceWithTracking(FD3D12CommandListHandle& hCommandList, FD3D12Resource* pResource, D3D12_RESOURCE_STATES before, D3D12_RESOURCE_STATES after, const CViewSubresourceSubset& subresourceSubset, ETransitionMode InMode)
 	{
-#if USE_D3D12RHI_RESOURCE_STATE_TRACKING
 		check(pResource);
 		check(pResource->RequiresResourceStateTracking());
-
 		check(!((after & (D3D12_RESOURCE_STATE_PIXEL_SHADER_RESOURCE | D3D12_RESOURCE_STATE_NON_PIXEL_SHADER_RESOURCE)) && (pResource->GetDesc().Flags & D3D12_RESOURCE_FLAG_DENY_SHADER_RESOURCE)));
 
 #ifdef PLATFORM_SUPPORTS_RESOURCE_COMPRESSION
@@ -856,24 +862,16 @@ public:
 #endif
 
 		hCommandList.UpdateResidency(pResource);
-		D3D12_RESOURCE_STATES before;
 		const bool bIsWholeResource = subresourceSubset.IsWholeResource();
 		CResourceState& ResourceState = hCommandList.GetResourceState(pResource);
+
+		bool bRequireUAVBarrier = false;
+
 		if (bIsWholeResource && ResourceState.AreAllSubresourcesSame())
 		{
 			// Fast path. Transition the entire resource from one state to another.
-			before = ResourceState.GetSubresourceState(D3D12_RESOURCE_BARRIER_ALL_SUBRESOURCES);
-			if (before == D3D12_RESOURCE_STATE_TBD)
-			{
-				// We need a pending resource barrier so we can setup the state before this command list executes
-				hCommandList.AddPendingResourceBarrier(pResource, after, D3D12_RESOURCE_BARRIER_ALL_SUBRESOURCES);
-				ResourceState.SetResourceState(after);
-			}
-			else if (IsTransitionNeeded(before, after))
-			{
-				hCommandList.AddTransitionBarrier(pResource, before, after, D3D12_RESOURCE_BARRIER_ALL_SUBRESOURCES);
-				ResourceState.SetResourceState(after);
-			}
+			bool bForceInAfterState = false;
+			bRequireUAVBarrier = ValidateAndSetResourceState(hCommandList, pResource, ResourceState, D3D12_RESOURCE_BARRIER_ALL_SUBRESOURCES, before, after, bForceInAfterState, InMode);
 		}
 		else
 		{
@@ -885,26 +883,12 @@ public:
 			{
 				for (uint32 SubresourceIndex = it.StartSubresource(); SubresourceIndex < it.EndSubresource(); SubresourceIndex++)
 				{
-					before = ResourceState.GetSubresourceState(SubresourceIndex);
-					if (before == D3D12_RESOURCE_STATE_TBD)
-					{
-						// We need a pending resource barrier so we can setup the state before this command list executes
-						hCommandList.AddPendingResourceBarrier(pResource, after, SubresourceIndex);
-						ResourceState.SetSubresourceState(SubresourceIndex, after);
-					}
-					else if (IsTransitionNeeded(before, after))
-					{
-						hCommandList.AddTransitionBarrier(pResource, before, after, SubresourceIndex);
-						ResourceState.SetSubresourceState(SubresourceIndex, after);
-					}
-					else
-					{
-						// Didn't need to transition the subresource.
-						if (before != after)
-						{
-							bWholeResourceWasTransitionedToSameState = false;
-						}
-					}
+					bool bForceInAfterState = false;
+					bRequireUAVBarrier |= ValidateAndSetResourceState(hCommandList, pResource, ResourceState, SubresourceIndex, before, after, bForceInAfterState, InMode);
+
+					// Subresource not in the same state, then whole resource is not in the same state anymore
+					if (ResourceState.GetSubresourceState(SubresourceIndex) != after)
+						bWholeResourceWasTransitionedToSameState = false;
 				}
 			}
 
@@ -912,12 +896,96 @@ public:
 			if (bWholeResourceWasTransitionedToSameState)
 			{
 				// Sanity check to make sure all subresources are really in the 'after' state
-				check(ResourceState.CheckResourceState(after));
+				verify(ResourceState.CheckAllSubresourceSame());
+				check(EnumHasAllFlags(ResourceState.GetSubresourceState(D3D12_RESOURCE_BARRIER_ALL_SUBRESOURCES), after));
+			}
+		}
+
+		return bRequireUAVBarrier;
+	}
+
+	static bool ValidateAndSetResourceState(FD3D12CommandListHandle& InCommandList, FD3D12Resource* InResource, CResourceState& InResourceState, uint32 InSubresourceIndex, D3D12_RESOURCE_STATES InBeforeState, D3D12_RESOURCE_STATES InAfterState, bool bInForceAfterState, ETransitionMode InMode)
+	{
+		// Only validate the current state?
+		bool bValidateState = !GUseInternalTransitions && (InMode == ETransitionMode::Validate);
+
+		// Try and get the correct D3D before state for the transition
+		D3D12_RESOURCE_STATES BeforeState = InResourceState.GetSubresourceState(InSubresourceIndex);
+
+		// Still untracked in this command list, then try and find out a before state to use
+		if (BeforeState == D3D12_RESOURCE_STATE_TBD)
+		{
+			if (bValidateState)
+			{
+				// Can't correctly validate on parallel command list because command list with final state which
+				// updates the resource state might not have been executed yet (on RHI Thread)
+				// Unless it's a transition on the default context and all transitions happen on the default context (validated somewhere else)
+				if (GRHICommandList.Bypass() || InCommandList.GetCurrentOwningContext()->IsDefaultContext())
+				{
+					BeforeState = InResource->GetResourceState().GetSubresourceState(InSubresourceIndex);
+				}
+			}
+			else if (GUseInternalTransitions)
+			{
+				// We need a pending resource barrier so we can setup the state before this command list executes
+				InCommandList.AddPendingResourceBarrier(InResource, InAfterState, InSubresourceIndex);
+				InResourceState.SetSubresourceState(InSubresourceIndex, InAfterState);
+			}
+			else
+			{
+				// We have enqueue the transition right now in the command list and can't add it to the pending list because
+				// this resource can already have been used in the current state in the command list
+				// so changing that state before this command list is invalid.				
+				BeforeState = InBeforeState;
+				if (BeforeState == D3D12_RESOURCE_STATE_TBD)
+				{
+					// If we don't have a valid before state, then we have to use the actual last stored state of the
+					// resource. We can sadly enough only correctly do this when parallel command lists don't perform
+					// any resource transition because then the current stored state might be invalid
+					// (Currently validated during begin/end transition in D3D12Commands)
+					BeforeState = InResource->GetResourceState().GetSubresourceState(InSubresourceIndex);
+				}
 
-				ResourceState.SetResourceState(after);
+				// Add dummy pending barrier, because the end state needs to be updated during execute
+				InCommandList.AddPendingResourceBarrier(InResource, D3D12_RESOURCE_STATE_TBD, InSubresourceIndex);
 			}
 		}
-#endif // USE_D3D12RHI_RESOURCE_STATE_TRACKING
+
+		bool bRequireUAVBarrier = false;
+
+		// Have a valid state now?
+		check(BeforeState != D3D12_RESOURCE_STATE_TBD || GUseInternalTransitions || bValidateState);
+		if (BeforeState != D3D12_RESOURCE_STATE_TBD)
+		{
+			// Make sure the before states match up or are unknown
+			check(InBeforeState == D3D12_RESOURCE_STATE_TBD || BeforeState == InBeforeState);
+
+			if (bValidateState)
+			{
+				// Check if all after states are valid and special case for DepthRead because then DepthWrite is also valid
+				check(EnumHasAllFlags(BeforeState, InAfterState) || (BeforeState == D3D12_RESOURCE_STATE_DEPTH_WRITE && InAfterState == D3D12_RESOURCE_STATE_DEPTH_READ));
+			}
+			else
+			{
+				// Require UAV barrier when before and after are UAV
+				if (BeforeState == D3D12_RESOURCE_STATE_UNORDERED_ACCESS && InAfterState == D3D12_RESOURCE_STATE_UNORDERED_ACCESS)
+				{
+					bRequireUAVBarrier = true;
+				}
+
+				// We're not using IsTransitionNeeded() when bInForceAfterState is set because we do want to transition even if 'after' is a subset of 'before'
+				// This is so that we can ensure all subresources are in the same state, simplifying future barriers
+				// No state merging when using engine transitions - otherwise next before state might not match up anymore)
+				bool bAllowStateMerging = GUseInternalTransitions;
+				if ((bInForceAfterState && BeforeState != InAfterState) || IsTransitionNeeded(bAllowStateMerging, BeforeState, InAfterState))
+				{
+					InCommandList.AddTransitionBarrier(InResource, BeforeState, InAfterState, InSubresourceIndex);
+					InResourceState.SetSubresourceState(InSubresourceIndex, InAfterState);
+				}				
+			}
+		}
+
+		return bRequireUAVBarrier;
 	}
 
 public:
@@ -979,9 +1047,9 @@ protected:
 
 	template<typename BaseResourceType>
 	TD3D12Texture2D<BaseResourceType>* CreateD3D12Texture2D(class FRHICommandListImmediate* RHICmdList, uint32 SizeX, uint32 SizeY, uint32 SizeZ, bool bTextureArray, bool CubeTexture, EPixelFormat Format,
-		uint32 NumMips, uint32 NumSamples, ETextureCreateFlags InFlags, FRHIResourceCreateInfo& CreateInfo);
+		uint32 NumMips, uint32 NumSamples, ETextureCreateFlags InFlags, ERHIAccess InResourceState, FRHIResourceCreateInfo& CreateInfo);
 
-	FD3D12Texture3D* CreateD3D12Texture3D(class FRHICommandListImmediate* RHICmdList, uint32 SizeX, uint32 SizeY, uint32 SizeZ, EPixelFormat Format, uint32 NumMips, ETextureCreateFlags InFlags, FRHIResourceCreateInfo& CreateInfo);
+	FD3D12Texture3D* CreateD3D12Texture3D(class FRHICommandListImmediate* RHICmdList, uint32 SizeX, uint32 SizeY, uint32 SizeZ, EPixelFormat Format, uint32 NumMips, ETextureCreateFlags InFlags, ERHIAccess InResourceState, FRHIResourceCreateInfo& CreateInfo);
 
 	template<typename BaseResourceType>
 	TD3D12Texture2D<BaseResourceType>* CreateTextureFromResource(bool bTextureArray, bool bCubeTexture, EPixelFormat Format, ETextureCreateFlags TexCreateFlags, const FClearValueBinding& ClearValueBinding, ID3D12Resource* Resource);
@@ -1095,52 +1163,77 @@ public:
 };
 
 /**
-*	Class of a scoped resource barrier.
-*	This class conditionally uses resource state tracking.
-*	This should only be used with the Editor.
+*	Class of a scoped resource barrier - handles both tracked and untracked resources
 */
-class FConditionalScopeResourceBarrier
+class FScopedResourceBarrier
 {
 private:
 	FD3D12CommandListHandle& hCommandList;
 	FD3D12Resource* const pResource;
-	D3D12_RESOURCE_STATES Current;
-	const D3D12_RESOURCE_STATES Desired;
+	D3D12_RESOURCE_STATES CurrentState;
+	const D3D12_RESOURCE_STATES DesiredState;
 	const uint32 Subresource;
-	bool bRestoreDefaultState;
+	FD3D12DynamicRHI::ETransitionMode TransitionMode;
+
+	bool bRestoreState;
 
 public:
-	explicit FConditionalScopeResourceBarrier(FD3D12CommandListHandle& hInCommandList, FD3D12Resource* pInResource, const D3D12_RESOURCE_STATES InDesired, const uint32 InSubresource)
+	FScopedResourceBarrier(FD3D12CommandListHandle& hInCommandList, FD3D12Resource* pInResource, const D3D12_RESOURCE_STATES InDesiredState, const uint32 InSubresource, FD3D12DynamicRHI::ETransitionMode InTransitionMode)
 		: hCommandList(hInCommandList)
 		, pResource(pInResource)
-		, Current(D3D12_RESOURCE_STATE_TBD)
-		, Desired(InDesired)
+		, CurrentState(D3D12_RESOURCE_STATE_TBD)
+		, DesiredState(InDesiredState)
 		, Subresource(InSubresource)
-		, bRestoreDefaultState(false)
+		, TransitionMode(InTransitionMode)
+		, bRestoreState(false)
 	{
 		// when we don't use resource state tracking, transition the resource (only if necessary)
 		if (!pResource->RequiresResourceStateTracking())
 		{
-			Current = pResource->GetDefaultResourceState();
-			if (Current != Desired)
+			CurrentState = pResource->GetDefaultResourceState();
+			if (CurrentState != DesiredState)
 			{
 				// we will add a transition, we need to transition back to the default state when the scoped object dies : 
-				bRestoreDefaultState = true;
-				hCommandList.AddTransitionBarrier(pResource, Current, Desired, Subresource);
+				bRestoreState = true;
+				hCommandList.AddTransitionBarrier(pResource, CurrentState, DesiredState, Subresource);
 			}
 		}
 		else
 		{
-			FD3D12DynamicRHI::TransitionResource(hCommandList, pResource, Desired, Subresource);
+			// If we are not using the internal transitions and need to apply the state change, then store the current state of restore
+			if (!GUseInternalTransitions && InTransitionMode == FD3D12DynamicRHI::ETransitionMode::Apply)
+			{
+				// try tracked state in command list
+				CResourceState& ResourceState = hCommandList.GetResourceState(pInResource);
+				CurrentState = ResourceState.GetSubresourceState(InSubresource);
+
+				// if still unknown then use the stored state (not valid when transitions happen in parallel command lists)
+				if (CurrentState == D3D12_RESOURCE_STATE_TBD)
+				{
+					CurrentState = pInResource->GetResourceState().GetSubresourceState(InSubresource);
+				}
+
+				// Restore state to current state when done
+				bRestoreState = true;
+			}
+
+			FD3D12DynamicRHI::TransitionResource(hCommandList, pResource, CurrentState, DesiredState, Subresource, TransitionMode);
 		}
 	}
 
-	~FConditionalScopeResourceBarrier()
+	~FScopedResourceBarrier()
 	{
-		// Return the resource to it's default state if necessary : 
-		if (bRestoreDefaultState)
+		// Return the resource to the original state if requested
+		if (bRestoreState)
 		{
-			hCommandList.AddTransitionBarrier(pResource, Desired, Current, Subresource);
+			if (!pResource->RequiresResourceStateTracking())
+			{
+				hCommandList.AddTransitionBarrier(pResource, DesiredState, CurrentState, Subresource);
+			}
+			else
+			{
+				FD3D12DynamicRHI::TransitionResource(hCommandList, pResource, DesiredState, CurrentState, Subresource, TransitionMode);
+			}
 		}
 	}
 };
diff --git a/Engine/Source/Runtime/D3D12RHI/Private/D3D12RayTracing.cpp b/Engine/Source/Runtime/D3D12RHI/Private/D3D12RayTracing.cpp
--- a/Engine/Source/Runtime/D3D12RHI/Private/D3D12RayTracing.cpp
+++ b/Engine/Source/Runtime/D3D12RHI/Private/D3D12RayTracing.cpp
@@ -1381,7 +1381,7 @@ public:
 
 		Buffer = Adapter->CreateRHIBuffer(
 			nullptr, BufferDesc, BufferDesc.Alignment,
-			0, BufferDesc.Width, BUF_Static, ED3D12ResourceStateMode::SingleState, CreateInfo);
+			0, BufferDesc.Width, BUF_Static, ED3D12ResourceStateMode::SingleState, ERHIAccess::SRVMask, CreateInfo);
 
 		SetName(Buffer->GetResource(), TEXT("Shader binding table"));
 
@@ -1591,12 +1591,12 @@ public:
 
 		for (FD3D12ShaderResourceView* SRV : TransitionSRVs[0])
 		{
-			FD3D12DynamicRHI::TransitionResource(CommandContext.CommandListHandle, SRV, D3D12_RESOURCE_STATE_NON_PIXEL_SHADER_RESOURCE);
+			FD3D12DynamicRHI::TransitionResource(CommandContext.CommandListHandle, SRV, D3D12_RESOURCE_STATE_PIXEL_SHADER_RESOURCE | D3D12_RESOURCE_STATE_NON_PIXEL_SHADER_RESOURCE, FD3D12DynamicRHI::ETransitionMode::Apply);
 		}
 
 		for (FD3D12UnorderedAccessView* UAV : TransitionUAVs[0])
 		{
-			FD3D12DynamicRHI::TransitionResource(CommandContext.CommandListHandle, UAV, D3D12_RESOURCE_STATE_UNORDERED_ACCESS);
+			FD3D12DynamicRHI::TransitionResource(CommandContext.CommandListHandle, UAV, D3D12_RESOURCE_STATE_UNORDERED_ACCESS, FD3D12DynamicRHI::ETransitionMode::Apply);
 		}
 	}
 };
@@ -2321,7 +2321,7 @@ void FD3D12RayTracingGeometry::TransitionBuffers(FD3D12CommandContext& CommandCo
 		FD3D12Buffer* IndexBuffer = CommandContext.RetrieveObject<FD3D12Buffer>(RHIIndexBuffer.GetReference());
 		if (IndexBuffer->GetResource()->RequiresResourceStateTracking())
 		{
-			FD3D12DynamicRHI::TransitionResource(CommandContext.CommandListHandle, IndexBuffer->GetResource(), D3D12_RESOURCE_STATE_NON_PIXEL_SHADER_RESOURCE, 0);
+			FD3D12DynamicRHI::TransitionResource(CommandContext.CommandListHandle, IndexBuffer->GetResource(), D3D12_RESOURCE_STATE_TBD, D3D12_RESOURCE_STATE_PIXEL_SHADER_RESOURCE | D3D12_RESOURCE_STATE_NON_PIXEL_SHADER_RESOURCE, 0, FD3D12DynamicRHI::ETransitionMode::Apply);
 		}
 	}
 
@@ -2331,7 +2331,7 @@ void FD3D12RayTracingGeometry::TransitionBuffers(FD3D12CommandContext& CommandCo
 		FD3D12Buffer* VertexBuffer = CommandContext.RetrieveObject<FD3D12Buffer>(RHIVertexBuffer.GetReference());
 		if (VertexBuffer->GetResource()->RequiresResourceStateTracking())
 		{
-			FD3D12DynamicRHI::TransitionResource(CommandContext.CommandListHandle, VertexBuffer->GetResource(), D3D12_RESOURCE_STATE_NON_PIXEL_SHADER_RESOURCE, 0);
+			FD3D12DynamicRHI::TransitionResource(CommandContext.CommandListHandle, VertexBuffer->GetResource(), D3D12_RESOURCE_STATE_TBD, D3D12_RESOURCE_STATE_PIXEL_SHADER_RESOURCE | D3D12_RESOURCE_STATE_NON_PIXEL_SHADER_RESOURCE, 0, FD3D12DynamicRHI::ETransitionMode::Apply);
 		}
 	}
 }
@@ -2352,7 +2352,7 @@ static void CreateAccelerationStructureBuffers(
 	CreateInfo.DebugName = DebugName;
 	AccelerationStructureBuffer = Adapter->CreateRHIBuffer(
 		nullptr, AccelerationStructureBufferDesc, D3D12_RAYTRACING_ACCELERATION_STRUCTURE_BYTE_ALIGNMENT,
-		0, AccelerationStructureBufferDesc.Width, BUF_AccelerationStructure, ED3D12ResourceStateMode::SingleState, CreateInfo);
+		0, AccelerationStructureBufferDesc.Width, BUF_AccelerationStructure, ED3D12ResourceStateMode::SingleState, ERHIAccess::RTAccelerationStructure, CreateInfo);
 
 	SetName(AccelerationStructureBuffer->GetResource(), DebugName);
 
@@ -2364,7 +2364,7 @@ static void CreateAccelerationStructureBuffers(
 	CreateInfo.DebugName = TEXT("ScratchBuffer");
 	ScratchBuffer = Adapter->CreateRHIBuffer(
 		nullptr, ScratchBufferDesc, D3D12_RAYTRACING_ACCELERATION_STRUCTURE_BYTE_ALIGNMENT,
-		0, ScratchBufferDesc.Width, BUF_UnorderedAccess, ED3D12ResourceStateMode::SingleState, CreateInfo);
+		0, ScratchBufferDesc.Width, BUF_UnorderedAccess, ED3D12ResourceStateMode::SingleState, ERHIAccess::UAVMask, CreateInfo);
 
 	SetName(ScratchBuffer->GetResource(), TEXT("Acceleration structure scratch"));
 }
@@ -2563,12 +2563,10 @@ void FD3D12RayTracingGeometry::BuildAccelerationStructure(FD3D12CommandContext&
 		CreateInfo.DebugName = TEXT("PostBuildInfoBuffer");
 		PostBuildInfoBuffers[GPUIndex] = Adapter->CreateRHIBuffer(
 			nullptr, PostBuildInfoBufferDesc, 8,
-			0, PostBuildInfoBufferDesc.Width, BUF_UnorderedAccess | BUF_SourceCopy, ED3D12ResourceStateMode::MultiState, CreateInfo);
+			0, PostBuildInfoBufferDesc.Width, BUF_UnorderedAccess | BUF_SourceCopy, ED3D12ResourceStateMode::MultiState, ERHIAccess::UAVMask, CreateInfo);
 
 		SetName(PostBuildInfoBuffers[GPUIndex]->GetResource(), TEXT("PostBuildInfoBuffer"));
 
-		FD3D12DynamicRHI::TransitionResource(CommandContext.CommandListHandle, PostBuildInfoBuffers[GPUIndex].GetReference()->GetResource(), D3D12_RESOURCE_STATE_UNORDERED_ACCESS, 0);
-
 		PostBuildInfoStagingBuffers[GPUIndex] = RHICreateStagingBuffer();
 
 		// Do not defer delete the staging buffer data because when readback is done we don't need
@@ -2609,6 +2607,8 @@ void FD3D12RayTracingGeometry::BuildAccelerationStructure(FD3D12CommandContext&
 			// Build with post build info
 			RayTracingCommandList->BuildRaytracingAccelerationStructure(&BuildDesc, 1, &PostBuildInfoDesc);
 
+			// Transition to copy source and perform the copy to readback
+			FD3D12DynamicRHI::TransitionResource(CommandContext.CommandListHandle, PostBuildInfoBuffers[GPUIndex].GetReference()->GetResource(), D3D12_RESOURCE_STATE_UNORDERED_ACCESS, D3D12_RESOURCE_STATE_COPY_SOURCE, 0, FD3D12DynamicRHI::ETransitionMode::Apply);
 			CommandContext.RHICopyToStagingBuffer(PostBuildInfoBuffers[GPUIndex], PostBuildInfoStagingBuffers[GPUIndex], 0, sizeof(uint64));
 
 			const FD3D12Fence& Fence = CommandContext.GetParentDevice()->GetCommandListManager().GetFence();
@@ -2697,7 +2697,7 @@ void FD3D12RayTracingGeometry::ConditionalCompactAccelerationStructure(FD3D12Com
 		CreateInfo.DebugName = TEXT("AccelerationStructureBuffer");
 		AccelerationStructureBuffers[GPUIndex] = Adapter->CreateRHIBuffer(
 			nullptr, AccelerationStructureBufferDesc, D3D12_RAYTRACING_ACCELERATION_STRUCTURE_BYTE_ALIGNMENT,
-			0, AccelerationStructureBufferDesc.Width, BUF_AccelerationStructure, ED3D12ResourceStateMode::SingleState, CreateInfo);
+			0, AccelerationStructureBufferDesc.Width, BUF_AccelerationStructure, ED3D12ResourceStateMode::SingleState, ERHIAccess::RTAccelerationStructure, CreateInfo);
 
 		INC_MEMORY_STAT_BY(STAT_D3D12RayTracingUsedVideoMemory, AccelerationStructureBuffers[GPUIndex]->GetSize());
 		INC_MEMORY_STAT_BY(STAT_D3D12RayTracingBLASMemory, AccelerationStructureBuffers[GPUIndex]->GetSize());
@@ -2867,7 +2867,7 @@ void FD3D12RayTracingScene::BuildAccelerationStructure(FD3D12CommandContext& Com
 		// #dxr_todo multi state only when bShouldCopyIndirectInstances is set - will still need transition to copy_src for lock behind but this can be done on the complete pool in theory (have to check cost)
 		InstanceBuffer = Adapter->CreateRHIBuffer(
 			nullptr, InstanceBufferDesc, D3D12_RAYTRACING_INSTANCE_DESCS_BYTE_ALIGNMENT,
-			sizeof(D3D12_RAYTRACING_INSTANCE_DESC), InstanceBufferDesc.Width, BUF_UnorderedAccess, ED3D12ResourceStateMode::MultiState, CreateInfo);
+			sizeof(D3D12_RAYTRACING_INSTANCE_DESC), InstanceBufferDesc.Width, BUF_UnorderedAccess, ED3D12ResourceStateMode::MultiState, ERHIAccess::UAVMask, CreateInfo);
 
 		D3D12_RAYTRACING_INSTANCE_DESC* MappedData = (D3D12_RAYTRACING_INSTANCE_DESC*)Adapter->GetOwningRHI()->LockBuffer(
 			nullptr, InstanceBuffer.GetReference(), InstanceBuffer->GetSize(), InstanceBuffer->GetUsage(), 0, InstanceBufferDesc.Width, RLM_WriteOnly);
@@ -2999,7 +2999,7 @@ void FD3D12RayTracingScene::BuildAccelerationStructure(FD3D12CommandContext& Com
 			TRHICommandList_RecursiveHazardous<FD3D12CommandContext> RHICmdList(&CommandContext);
 			FUnorderedAccessViewRHIRef InstancesDescUAV = RHICreateUnorderedAccessView(InstanceBuffer, false, false);
 
-			FD3D12DynamicRHI::TransitionResource(CommandContext.CommandListHandle, InstanceBuffer.GetReference()->GetResource(), D3D12_RESOURCE_STATE_UNORDERED_ACCESS, 0);
+			FD3D12DynamicRHI::TransitionResource(CommandContext.CommandListHandle, InstanceBuffer.GetReference()->GetResource(), D3D12_RESOURCE_STATE_TBD, D3D12_RESOURCE_STATE_UNORDERED_ACCESS, 0, FD3D12DynamicRHI::ETransitionMode::Apply);
 
 			RHICmdList.BeginUAVOverlap(InstancesDescUAV);
 
@@ -3024,7 +3024,7 @@ void FD3D12RayTracingScene::BuildAccelerationStructure(FD3D12CommandContext& Com
 			RHICmdList.EndUAVOverlap(InstancesDescUAV);
 		}
 
-		FD3D12DynamicRHI::TransitionResource(CommandContext.CommandListHandle, InstanceBuffer.GetReference()->GetResource(), D3D12_RESOURCE_STATE_NON_PIXEL_SHADER_RESOURCE, 0);
+		FD3D12DynamicRHI::TransitionResource(CommandContext.CommandListHandle, InstanceBuffer.GetReference()->GetResource(), D3D12_RESOURCE_STATE_TBD, D3D12_RESOURCE_STATE_PIXEL_SHADER_RESOURCE | D3D12_RESOURCE_STATE_NON_PIXEL_SHADER_RESOURCE, 0, FD3D12DynamicRHI::ETransitionMode::Apply);
 	}
 	
 	
@@ -3248,12 +3248,12 @@ struct FD3D12RayTracingGlobalResourceBinder
 
 	void AddResourceTransition(FD3D12ShaderResourceView* SRV)
 	{
-		FD3D12DynamicRHI::TransitionResource(CommandContext.CommandListHandle, SRV, D3D12_RESOURCE_STATE_NON_PIXEL_SHADER_RESOURCE);
+		FD3D12DynamicRHI::TransitionResource(CommandContext.CommandListHandle, SRV, D3D12_RESOURCE_STATE_PIXEL_SHADER_RESOURCE | D3D12_RESOURCE_STATE_NON_PIXEL_SHADER_RESOURCE, FD3D12DynamicRHI::ETransitionMode::Apply);
 	}
 
 	void AddResourceTransition(FD3D12UnorderedAccessView* UAV)
 	{
-		FD3D12DynamicRHI::TransitionResource(CommandContext.CommandListHandle, UAV, D3D12_RESOURCE_STATE_UNORDERED_ACCESS);
+		FD3D12DynamicRHI::TransitionResource(CommandContext.CommandListHandle, UAV, D3D12_RESOURCE_STATE_UNORDERED_ACCESS, FD3D12DynamicRHI::ETransitionMode::Apply);
 	}
 
 	FD3D12Device* GetDevice()
diff --git a/Engine/Source/Runtime/D3D12RHI/Private/D3D12RenderTarget.cpp b/Engine/Source/Runtime/D3D12RHI/Private/D3D12RenderTarget.cpp
--- a/Engine/Source/Runtime/D3D12RHI/Private/D3D12RenderTarget.cpp
+++ b/Engine/Source/Runtime/D3D12RHI/Private/D3D12RenderTarget.cpp
@@ -95,7 +95,7 @@ void FD3D12CommandContext::ResolveTextureUsingShader(
 				GetParentDevice()->RegisterGPUWork(0);
 			}
 
-			FD3D12DynamicRHI::TransitionResource(CommandListHandle, DestTextureDSV, D3D12_RESOURCE_STATE_DEPTH_WRITE);
+			FD3D12DynamicRHI::TransitionResource(CommandListHandle, DestTextureDSV, D3D12_RESOURCE_STATE_DEPTH_WRITE, FD3D12DynamicRHI::ETransitionMode::Apply);
 
 			CommandListHandle.FlushResourceBarriers();
 
@@ -127,7 +127,7 @@ void FD3D12CommandContext::ResolveTextureUsingShader(
 				GetParentDevice()->RegisterGPUWork(0);
 			}
 
-			FD3D12DynamicRHI::TransitionResource(CommandListHandle, DestTextureRTV, D3D12_RESOURCE_STATE_RENDER_TARGET);
+			FD3D12DynamicRHI::TransitionResource(CommandListHandle, DestTextureRTV, D3D12_RESOURCE_STATE_RENDER_TARGET, FD3D12DynamicRHI::ETransitionMode::Apply);
 
 			CommandListHandle.FlushResourceBarriers();
 
@@ -267,8 +267,8 @@ void FD3D12CommandContext::RHICopyToResolveTarget(FRHITexture* SourceTextureRHI,
 				// Determine whether a MSAA resolve is needed, or just a copy.
 				if (SourceTextureRHI->IsMultisampled() && !DestTexture2D->IsMultisampled())
 				{
-					FConditionalScopeResourceBarrier ConditionalScopeResourceBarrierDest(CommandListHandle, DestTexture2D->GetResource(), D3D12_RESOURCE_STATE_RESOLVE_DEST, ResolveParams.DestArrayIndex);
-					FConditionalScopeResourceBarrier ConditionalScopeResourceBarrierSource(CommandListHandle, SourceTexture2D->GetResource(), D3D12_RESOURCE_STATE_RESOLVE_SOURCE, ResolveParams.SourceArrayIndex);
+					FScopedResourceBarrier ConditionalScopeResourceBarrierDest(CommandListHandle, DestTexture2D->GetResource(), D3D12_RESOURCE_STATE_RESOLVE_DEST, ResolveParams.DestArrayIndex, FD3D12DynamicRHI::ETransitionMode::Validate);
+					FScopedResourceBarrier ConditionalScopeResourceBarrierSource(CommandListHandle, SourceTexture2D->GetResource(), D3D12_RESOURCE_STATE_RESOLVE_SOURCE, ResolveParams.SourceArrayIndex, FD3D12DynamicRHI::ETransitionMode::Validate);
 
 					otherWorkCounter++;
 					CommandListHandle.FlushResourceBarriers();
@@ -309,8 +309,8 @@ void FD3D12CommandContext::RHICopyToResolveTarget(FRHITexture* SourceTextureRHI,
 
 						const FResolveRect& DestRect = ResolveParams.DestRect.IsValid() ? ResolveParams.DestRect : SrcRect;
 
-						FConditionalScopeResourceBarrier ConditionalScopeResourceBarrierDest(CommandListHandle, DestTexture2D->GetResource(), D3D12_RESOURCE_STATE_COPY_DEST, ResolveParams.DestArrayIndex);
-						FConditionalScopeResourceBarrier ConditionalScopeResourceBarrierSource(CommandListHandle, SourceTexture2D->GetResource(), D3D12_RESOURCE_STATE_COPY_SOURCE, ResolveParams.SourceArrayIndex);
+						FScopedResourceBarrier ConditionalScopeResourceBarrierDest(CommandListHandle, DestTexture2D->GetResource(), D3D12_RESOURCE_STATE_COPY_DEST, ResolveParams.DestArrayIndex, FD3D12DynamicRHI::ETransitionMode::Validate);
+						FScopedResourceBarrier ConditionalScopeResourceBarrierSource(CommandListHandle, SourceTexture2D->GetResource(), D3D12_RESOURCE_STATE_COPY_SOURCE, ResolveParams.SourceArrayIndex, FD3D12DynamicRHI::ETransitionMode::Validate);
 
 						CD3DX12_TEXTURE_COPY_LOCATION DestCopyLocation(DestTexture2D->GetResource()->GetResource(), ResolveParams.DestArrayIndex);
 						CD3DX12_TEXTURE_COPY_LOCATION SourceCopyLocation(SourceTexture2D->GetResource()->GetResource(), ResolveParams.SourceArrayIndex);
@@ -328,7 +328,7 @@ void FD3D12CommandContext::RHICopyToResolveTarget(FRHITexture* SourceTextureRHI,
 					}
 					else
 					{
-						FConditionalScopeResourceBarrier ConditionalScopeResourceBarrierSource(CommandListHandle, SourceTexture2D->GetResource(), D3D12_RESOURCE_STATE_COPY_SOURCE, ResolveParams.SourceArrayIndex);
+						FScopedResourceBarrier ConditionalScopeResourceBarrierSource(CommandListHandle, SourceTexture2D->GetResource(), D3D12_RESOURCE_STATE_COPY_SOURCE, ResolveParams.SourceArrayIndex, FD3D12DynamicRHI::ETransitionMode::Validate);
 
 						// Resolve to a buffer.
 						if (ResolveTargetDesc.Dimension == D3D12_RESOURCE_DIMENSION_BUFFER)
@@ -374,7 +374,7 @@ void FD3D12CommandContext::RHICopyToResolveTarget(FRHITexture* SourceTextureRHI,
 						else
 						{
 							// Transition to the copy dest state.
-							FConditionalScopeResourceBarrier ConditionalScopeResourceBarrierDest(CommandListHandle, DestTexture2D->GetResource(), D3D12_RESOURCE_STATE_COPY_DEST, 0);
+							FScopedResourceBarrier ConditionalScopeResourceBarrierDest(CommandListHandle, DestTexture2D->GetResource(), D3D12_RESOURCE_STATE_COPY_DEST, 0, FD3D12DynamicRHI::ETransitionMode::Validate);
 
 							CD3DX12_TEXTURE_COPY_LOCATION DestCopyLocation(DestTexture2D->GetResource()->GetResource(), ResolveParams.DestArrayIndex);
 							CD3DX12_TEXTURE_COPY_LOCATION SourceCopyLocation(SourceTexture2D->GetResource()->GetResource(), ResolveParams.SourceArrayIndex);
@@ -414,8 +414,8 @@ void FD3D12CommandContext::RHICopyToResolveTarget(FRHITexture* SourceTextureRHI,
 			// Determine whether a MSAA resolve is needed, or just a copy.
 			if (SourceTextureRHI->IsMultisampled() && !DestTextureCube->IsMultisampled())
 			{
-				FConditionalScopeResourceBarrier ConditionalScopeResourceBarrierDest(CommandListHandle, DestTextureCube->GetResource(), D3D12_RESOURCE_STATE_RESOLVE_DEST, DestSubresource);
-				FConditionalScopeResourceBarrier ConditionalScopeResourceBarrierSource(CommandListHandle, SourceTextureCube->GetResource(), D3D12_RESOURCE_STATE_RESOLVE_SOURCE, SourceSubresource);
+				FScopedResourceBarrier ConditionalScopeResourceBarrierDest(CommandListHandle, DestTextureCube->GetResource(), D3D12_RESOURCE_STATE_RESOLVE_DEST, DestSubresource, FD3D12DynamicRHI::ETransitionMode::Validate);
+				FScopedResourceBarrier ConditionalScopeResourceBarrierSource(CommandListHandle, SourceTextureCube->GetResource(), D3D12_RESOURCE_STATE_RESOLVE_SOURCE, SourceSubresource, FD3D12DynamicRHI::ETransitionMode::Validate);
 
 				otherWorkCounter++;
 				CommandListHandle.FlushResourceBarriers();
@@ -435,8 +435,8 @@ void FD3D12CommandContext::RHICopyToResolveTarget(FRHITexture* SourceTextureRHI,
 				CD3DX12_TEXTURE_COPY_LOCATION DestCopyLocation(DestTextureCube->GetResource()->GetResource(), DestSubresource);
 				CD3DX12_TEXTURE_COPY_LOCATION SourceCopyLocation(SourceTextureCube->GetResource()->GetResource(), SourceSubresource);
 
-				FConditionalScopeResourceBarrier ConditionalScopeResourceBarrierDest(CommandListHandle, DestTextureCube->GetResource(), D3D12_RESOURCE_STATE_COPY_DEST, DestCopyLocation.SubresourceIndex);
-				FConditionalScopeResourceBarrier ConditionalScopeResourceBarrierSource(CommandListHandle, SourceTextureCube->GetResource(), D3D12_RESOURCE_STATE_COPY_SOURCE, SourceCopyLocation.SubresourceIndex);
+				FScopedResourceBarrier ConditionalScopeResourceBarrierDest(CommandListHandle, DestTextureCube->GetResource(), D3D12_RESOURCE_STATE_COPY_DEST, DestCopyLocation.SubresourceIndex, FD3D12DynamicRHI::ETransitionMode::Validate);
+				FScopedResourceBarrier ConditionalScopeResourceBarrierSource(CommandListHandle, SourceTextureCube->GetResource(), D3D12_RESOURCE_STATE_COPY_SOURCE, SourceCopyLocation.SubresourceIndex, FD3D12DynamicRHI::ETransitionMode::Validate);
 
 				numCopies++;
 				CommandListHandle.FlushResourceBarriers();
@@ -461,8 +461,8 @@ void FD3D12CommandContext::RHICopyToResolveTarget(FRHITexture* SourceTextureRHI,
 		CD3DX12_TEXTURE_COPY_LOCATION DestCopyLocation(DestTextureCube->GetResource()->GetResource(), Subresource);
 		CD3DX12_TEXTURE_COPY_LOCATION SourceCopyLocation(SourceTexture2D->GetResource()->GetResource(), 0);
 
-		FConditionalScopeResourceBarrier ConditionalScopeResourceBarrierDest(CommandListHandle, DestTextureCube->GetResource(), D3D12_RESOURCE_STATE_COPY_DEST, DestCopyLocation.SubresourceIndex);
-		FConditionalScopeResourceBarrier ConditionalScopeResourceBarrierSource(CommandListHandle, SourceTexture2D->GetResource(), D3D12_RESOURCE_STATE_COPY_SOURCE, SourceCopyLocation.SubresourceIndex);
+		FScopedResourceBarrier ConditionalScopeResourceBarrierDest(CommandListHandle, DestTextureCube->GetResource(), D3D12_RESOURCE_STATE_COPY_DEST, DestCopyLocation.SubresourceIndex, FD3D12DynamicRHI::ETransitionMode::Validate);
+		FScopedResourceBarrier ConditionalScopeResourceBarrierSource(CommandListHandle, SourceTexture2D->GetResource(), D3D12_RESOURCE_STATE_COPY_SOURCE, SourceCopyLocation.SubresourceIndex, FD3D12DynamicRHI::ETransitionMode::Validate);
 
 		numCopies++;
 		CommandListHandle.FlushResourceBarriers();
@@ -482,6 +482,16 @@ void FD3D12CommandContext::RHICopyToResolveTarget(FRHITexture* SourceTextureRHI,
 		check(SourceTexture3D == DestTexture3D);
 	}
 
+	// Force transtion to readable since RHI user expects both source and dest to be in readable state after calling RHICopyToResolveTarget
+	if (SourceTextureRHI && ResolveParams.SourceAccessFinal != ERHIAccess::Unknown)
+	{
+		RHICmdList.TransitionResource(ResolveParams.SourceAccessFinal, SourceTextureRHI);
+	}
+	if (DestTextureRHI && SourceTextureRHI != DestTextureRHI && ResolveParams.DestAccessFinal != ERHIAccess::Unknown)
+	{
+		RHICmdList.TransitionResource(ResolveParams.DestAccessFinal, DestTextureRHI);
+	}
+
 	DEBUG_EXECUTE_COMMAND_LIST(this);
 }
 
@@ -521,10 +531,10 @@ void FD3D12DynamicRHI::RHITransferTexture(FRHITexture2D* TextureRHI, FIntRect Re
 
 	{
 		FD3D12CommandContext& SrcContext = Adapter.GetDevice(SrcGPUIndex)->GetDefaultCommandContext();
-		FD3D12DynamicRHI::TransitionResource(SrcContext.CommandListHandle, SrcTexture2D->GetResource(), D3D12_RESOURCE_STATE_COPY_SOURCE, 0);
+		FD3D12DynamicRHI::TransitionResource(SrcContext.CommandListHandle, SrcTexture2D->GetResource(), D3D12_RESOURCE_STATE_TBD, D3D12_RESOURCE_STATE_COPY_SOURCE, 0, FD3D12DynamicRHI::ETransitionMode::Apply);
 
 		FD3D12CommandContext& DestContext = Adapter.GetDevice(DestGPUIndex)->GetDefaultCommandContext();
-		FD3D12DynamicRHI::TransitionResource(DestContext.CommandListHandle, DestTexture2D->GetResource(), D3D12_RESOURCE_STATE_COPY_DEST, 0);
+		FD3D12DynamicRHI::TransitionResource(DestContext.CommandListHandle, DestTexture2D->GetResource(), D3D12_RESOURCE_STATE_TBD, D3D12_RESOURCE_STATE_COPY_DEST, 0, FD3D12DynamicRHI::ETransitionMode::Apply);
 	}
 
 	RHIMultiGPULockstep(SrcAndDestMask);
@@ -573,10 +583,10 @@ void FD3D12DynamicRHI::RHITransferTextures(const TArrayView<const FTransferTextu
 
 		{
 			FD3D12CommandContext& SrcContext = Adapter.GetDevice(SrcGPUIndex)->GetDefaultCommandContext();
-			FD3D12DynamicRHI::TransitionResource(SrcContext.CommandListHandle, SrcTexture->GetResource(), D3D12_RESOURCE_STATE_COPY_SOURCE, 0);
+			FD3D12DynamicRHI::TransitionResource(SrcContext.CommandListHandle, SrcTexture->GetResource(), D3D12_RESOURCE_STATE_TBD, D3D12_RESOURCE_STATE_COPY_SOURCE, 0, FD3D12DynamicRHI::ETransitionMode::Apply);
 
 			FD3D12CommandContext& DestContext = Adapter.GetDevice(DestGPUIndex)->GetDefaultCommandContext();
-			FD3D12DynamicRHI::TransitionResource(DestContext.CommandListHandle, DestTexture->GetResource(), D3D12_RESOURCE_STATE_COPY_DEST, 0);
+			FD3D12DynamicRHI::TransitionResource(DestContext.CommandListHandle, DestTexture->GetResource(), D3D12_RESOURCE_STATE_TBD, D3D12_RESOURCE_STATE_COPY_DEST, 0, FD3D12DynamicRHI::ETransitionMode::Apply);
 		}
 	}
 
@@ -798,7 +808,7 @@ TRefCountPtr<FD3D12Resource> FD3D12DynamicRHI::GetStagingTexture(FRHITexture* Te
 	CD3DX12_TEXTURE_COPY_LOCATION DestCopyLocation(TempTexture2D->GetResource(), placedTexture2D);
 	CD3DX12_TEXTURE_COPY_LOCATION SourceCopyLocation(Texture->GetResource()->GetResource(), Subresource);
 
-	FConditionalScopeResourceBarrier ScopeResourceBarrierSource(hCommandList, Texture->GetResource(), D3D12_RESOURCE_STATE_COPY_SOURCE, SourceCopyLocation.SubresourceIndex);
+	FScopedResourceBarrier ScopeResourceBarrierSource(hCommandList, Texture->GetResource(), D3D12_RESOURCE_STATE_COPY_SOURCE, SourceCopyLocation.SubresourceIndex, FD3D12DynamicRHI::ETransitionMode::Apply);
 	hCommandList.FlushResourceBarriers();
 	// Upload heap doesn't need to transition
 
@@ -1457,7 +1467,7 @@ void FD3D12DynamicRHI::ReadSurfaceDataMSAARaw(FRHICommandList_RecursiveHazardous
 			SampleIndex
 			);
 
-		FConditionalScopeResourceBarrier ScopeResourceBarrierSource(hCommandList, NonMSAATexture2D, D3D12_RESOURCE_STATE_COPY_SOURCE, SourceCopyLocation.SubresourceIndex);
+		FScopedResourceBarrier ScopeResourceBarrierSource(hCommandList, NonMSAATexture2D, D3D12_RESOURCE_STATE_COPY_SOURCE, SourceCopyLocation.SubresourceIndex, FD3D12DynamicRHI::ETransitionMode::Apply);
 		// Upload heap doesn't need to transition
 
 		DefaultContext.numCopies++;
@@ -1668,7 +1678,7 @@ void FD3D12DynamicRHI::RHIReadSurfaceFloatData(FRHITexture* TextureRHI, FIntRect
 	CD3DX12_TEXTURE_COPY_LOCATION SourceCopyLocation(Texture->GetResource()->GetResource(), Subresource);
 
 	{
-		FConditionalScopeResourceBarrier ConditionalScopeResourceBarrier(hCommandList, Texture->GetResource(), D3D12_RESOURCE_STATE_COPY_SOURCE, SourceCopyLocation.SubresourceIndex);
+		FScopedResourceBarrier ConditionalScopeResourceBarrier(hCommandList, Texture->GetResource(), D3D12_RESOURCE_STATE_COPY_SOURCE, SourceCopyLocation.SubresourceIndex, FD3D12DynamicRHI::ETransitionMode::Apply);
 		// Don't need to transition upload heaps
 
 		DefaultContext.numCopies++;
@@ -1776,7 +1786,7 @@ void FD3D12DynamicRHI::RHIRead3DSurfaceFloatData(FRHITexture* TextureRHI, FIntRe
 	CD3DX12_TEXTURE_COPY_LOCATION SourceCopyLocation(Texture->GetResource()->GetResource(), Subresource);
 
 	{
-		FConditionalScopeResourceBarrier ConditionalScopeResourceBarrier(hCommandList, Texture->GetResource(), D3D12_RESOURCE_STATE_COPY_SOURCE, SourceCopyLocation.SubresourceIndex);
+		FScopedResourceBarrier ConditionalScopeResourceBarrier(hCommandList, Texture->GetResource(), D3D12_RESOURCE_STATE_COPY_SOURCE, SourceCopyLocation.SubresourceIndex, FD3D12DynamicRHI::ETransitionMode::Apply);
 		// Don't need to transition upload heaps
 
 		DefaultContext.numCopies++;
diff --git a/Engine/Source/Runtime/D3D12RHI/Private/D3D12StateCache.cpp b/Engine/Source/Runtime/D3D12RHI/Private/D3D12StateCache.cpp
--- a/Engine/Source/Runtime/D3D12RHI/Private/D3D12StateCache.cpp
+++ b/Engine/Source/Runtime/D3D12RHI/Private/D3D12StateCache.cpp
@@ -1109,7 +1109,7 @@ void FD3D12StateCacheBase::InternalSetIndexBuffer(FD3D12Resource* Resource)
 	if (Resource->RequiresResourceStateTracking())
 	{
 		check(Resource->GetSubresourceCount() == 1);
-		FD3D12DynamicRHI::TransitionResource(CommandListHandle, Resource, D3D12_RESOURCE_STATE_INDEX_BUFFER, D3D12_RESOURCE_BARRIER_ALL_SUBRESOURCES);
+		FD3D12DynamicRHI::TransitionResource(CommandListHandle, Resource, D3D12_RESOURCE_STATE_TBD, D3D12_RESOURCE_STATE_INDEX_BUFFER, D3D12_RESOURCE_BARRIER_ALL_SUBRESOURCES, FD3D12DynamicRHI::ETransitionMode::Validate);
 	}
 }
 
@@ -1166,7 +1166,7 @@ void FD3D12StateCacheBase::InternalSetStreamSource(FD3D12ResourceLocation* Verte
 		if (pResource && pResource->RequiresResourceStateTracking())
 		{
 			check(pResource->GetSubresourceCount() == 1);
-			FD3D12DynamicRHI::TransitionResource(CmdContext->CommandListHandle, pResource, D3D12_RESOURCE_STATE_VERTEX_AND_CONSTANT_BUFFER, D3D12_RESOURCE_BARRIER_ALL_SUBRESOURCES);
+			FD3D12DynamicRHI::TransitionResource(CmdContext->CommandListHandle, pResource, D3D12_RESOURCE_STATE_TBD, D3D12_RESOURCE_STATE_VERTEX_AND_CONSTANT_BUFFER, D3D12_RESOURCE_BARRIER_ALL_SUBRESOURCES, FD3D12DynamicRHI::ETransitionMode::Validate);
 		}
 	}
 
diff --git a/Engine/Source/Runtime/D3D12RHI/Private/D3D12StructuredBuffer.cpp b/Engine/Source/Runtime/D3D12RHI/Private/D3D12StructuredBuffer.cpp
--- a/Engine/Source/Runtime/D3D12RHI/Private/D3D12StructuredBuffer.cpp
+++ b/Engine/Source/Runtime/D3D12RHI/Private/D3D12StructuredBuffer.cpp
@@ -38,7 +38,7 @@ FStructuredBufferRHIRef FD3D12DynamicRHI::CreateStructuredBuffer_RenderThread(FR
 	// can be addressed correctly with element based offsets.
 	const uint32 Alignment = ((InUsage & (BUF_ByteAddressBuffer | BUF_DrawIndirect)) == 0) ? Stride : 4;
 
-	FD3D12Buffer* NewBuffer = GetAdapter().CreateRHIBuffer(&RHICmdList, Desc, Alignment, Stride, Size, InUsage | BUF_StructuredBuffer, ED3D12ResourceStateMode::Default, CreateInfo);
+	FD3D12Buffer* NewBuffer = GetAdapter().CreateRHIBuffer(&RHICmdList, Desc, Alignment, Stride, Size, InUsage | BUF_StructuredBuffer, ED3D12ResourceStateMode::Default, InResourceState, CreateInfo);
 	if (NewBuffer->ResourceLocation.IsTransient())
 	{
 		// TODO: this should ideally be set in platform-independent code, since this tracking is for the high level
@@ -59,7 +59,7 @@ FStructuredBufferRHIRef FD3D12DynamicRHI::RHICreateStructuredBuffer(uint32 Strid
 	// can be addressed correctly with element based offsets.
 	const uint32 Alignment = ((InUsage & (BUF_ByteAddressBuffer | BUF_DrawIndirect)) == 0) ? Stride : 4;
 
-	FD3D12Buffer* NewBuffer = GetAdapter().CreateRHIBuffer(nullptr, Desc, Alignment, Stride, Size, InUsage | BUF_StructuredBuffer, ED3D12ResourceStateMode::Default, CreateInfo);
+	FD3D12Buffer* NewBuffer = GetAdapter().CreateRHIBuffer(nullptr, Desc, Alignment, Stride, Size, InUsage | BUF_StructuredBuffer, ED3D12ResourceStateMode::Default, InResourceState, CreateInfo);
 	if (NewBuffer->ResourceLocation.IsTransient())
 	{
 		// TODO: this should ideally be set in platform-independent code, since this tracking is for the high level
diff --git a/Engine/Source/Runtime/D3D12RHI/Private/D3D12Texture.cpp b/Engine/Source/Runtime/D3D12RHI/Private/D3D12Texture.cpp
--- a/Engine/Source/Runtime/D3D12RHI/Private/D3D12Texture.cpp
+++ b/Engine/Source/Runtime/D3D12RHI/Private/D3D12Texture.cpp
@@ -231,20 +231,24 @@ struct FD3D12RHICommandInitializeTexture final : public FRHICommand<FD3D12RHICom
 				Src.PlacedFootprint = Footprints[Subresource];
 				CmdList->CopyTextureRegion(&Dst, 0, 0, 0, &Src, nullptr);
 			}			
-			
+
 			// Update the resource state after the copy has been done (will take care of updating the residency as well)
+			hCommandList.AddTransitionBarrier(Resource, D3D12_RESOURCE_STATE_COPY_DEST, DestinationState, D3D12_RESOURCE_BARRIER_ALL_SUBRESOURCES);
+
 			if (Resource->RequiresResourceStateTracking())
 			{
-				// record the dummy copy_dest to copy_dest transition in the command list to make sure we have proper tracking of the resource
-				// mostly needed to make sure we have correct storage of end resource state when all the pending buffer transitions are flushed 
-				// (otherwise it's untracked and not updated)
-				FD3D12DynamicRHI::TransitionResource(hCommandList, Resource, D3D12_RESOURCE_STATE_COPY_DEST, D3D12_RESOURCE_BARRIER_ALL_SUBRESOURCES);
+				// Update the tracked resource state of this resource in the command list
+				CResourceState& ResourceState = hCommandList.GetResourceState(Resource);
+				ResourceState.SetResourceState(DestinationState);
+				Resource->GetResourceState().SetResourceState(DestinationState);
+
+				// Add dummy pending barrier, because the end state needs to be updated after execture command list with tracked state in the command list
+				hCommandList.AddPendingResourceBarrier(Resource, D3D12_RESOURCE_STATE_TBD, D3D12_RESOURCE_BARRIER_ALL_SUBRESOURCES);
 			}
 			else
 			{
-				hCommandList.AddTransitionBarrier(Resource, D3D12_RESOURCE_STATE_COPY_DEST, Resource->GetDefaultResourceState(), D3D12_RESOURCE_BARRIER_ALL_SUBRESOURCES);
+				check(Resource->GetDefaultResourceState() == DestinationState);
 			}
-
 		}
 
 		if (!bAllocateOnStack)
@@ -699,7 +703,7 @@ void SafeCreateTexture2D(FD3D12Device* pDevice,
 
 template<typename BaseResourceType>
 TD3D12Texture2D<BaseResourceType>* FD3D12DynamicRHI::CreateD3D12Texture2D(FRHICommandListImmediate* RHICmdList, uint32 SizeX, uint32 SizeY, uint32 SizeZ, bool bTextureArray, bool bCubeTexture, EPixelFormat Format,
-	uint32 NumMips, uint32 NumSamples, ETextureCreateFlags Flags, FRHIResourceCreateInfo& CreateInfo)
+	uint32 NumMips, uint32 NumSamples, ETextureCreateFlags Flags, ERHIAccess InResourceState, FRHIResourceCreateInfo& CreateInfo)
 {
 #if PLATFORM_WINDOWS || PLATFORM_HOLOLENS
 	check(SizeX > 0 && SizeY > 0 && NumMips > 0);
@@ -842,7 +846,7 @@ TD3D12Texture2D<BaseResourceType>* FD3D12DynamicRHI::CreateD3D12Texture2D(FRHICo
 
 	// The state this resource will be in when it leaves this function
 	const FD3D12Resource::FD3D12ResourceTypeHelper Type(TextureDesc, D3D12_HEAP_TYPE_DEFAULT);
-	const D3D12_RESOURCE_STATES InitialState = Type.GetOptimalInitialState(false);
+	const D3D12_RESOURCE_STATES InitialState = Type.GetOptimalInitialState(InResourceState, false);
 
 	TD3D12Texture2D<BaseResourceType>* D3D12TextureOut = Adapter->CreateLinkedObject<TD3D12Texture2D<BaseResourceType>>(CreateInfo.GPUMask, [&](FD3D12Device* Device)
 	{
@@ -1070,7 +1074,7 @@ TD3D12Texture2D<BaseResourceType>* FD3D12DynamicRHI::CreateD3D12Texture2D(FRHICo
 #endif // PLATFORM_WINDOWS || PLATFORM_HOLOLENS
 }
 
-FD3D12Texture3D* FD3D12DynamicRHI::CreateD3D12Texture3D(FRHICommandListImmediate* RHICmdList, uint32 SizeX, uint32 SizeY, uint32 SizeZ, EPixelFormat Format, uint32 NumMips, ETextureCreateFlags Flags, FRHIResourceCreateInfo& CreateInfo)
+FD3D12Texture3D* FD3D12DynamicRHI::CreateD3D12Texture3D(FRHICommandListImmediate* RHICmdList, uint32 SizeX, uint32 SizeY, uint32 SizeZ, EPixelFormat Format, uint32 NumMips, ETextureCreateFlags Flags, ERHIAccess InResourceState, FRHIResourceCreateInfo& CreateInfo)
 {
 #if PLATFORM_WINDOWS || PLATFORM_HOLOLENS
 	SCOPE_CYCLE_COUNTER(STAT_D3D12CreateTextureTime);
@@ -1117,7 +1121,7 @@ FD3D12Texture3D* FD3D12DynamicRHI::CreateD3D12Texture3D(FRHICommandListImmediate
 
 	// The state this resource will be in when it leaves this function
 	const FD3D12Resource::FD3D12ResourceTypeHelper Type(TextureDesc, D3D12_HEAP_TYPE_DEFAULT);
-	const D3D12_RESOURCE_STATES InitialState = Type.GetOptimalInitialState(false);
+	const D3D12_RESOURCE_STATES InitialState = Type.GetOptimalInitialState(InResourceState, false);
 
 	FD3D12Adapter* Adapter = &GetAdapter();
 	FD3D12Texture3D* D3D12TextureOut = Adapter->CreateLinkedObject<FD3D12Texture3D>(CreateInfo.GPUMask, [&](FD3D12Device* Device)
@@ -1182,12 +1186,12 @@ FD3D12Texture3D* FD3D12DynamicRHI::CreateD3D12Texture3D(FRHICommandListImmediate
 
 FTexture2DRHIRef FD3D12DynamicRHI::RHICreateTexture2D_RenderThread(class FRHICommandListImmediate& RHICmdList, uint32 SizeX, uint32 SizeY, uint8 Format, uint32 NumMips, uint32 NumSamples, ETextureCreateFlags Flags, ERHIAccess InResourceState, FRHIResourceCreateInfo& CreateInfo)
 {
-	return CreateD3D12Texture2D<FD3D12BaseTexture2D>(&RHICmdList, SizeX, SizeY, 1, false, false, (EPixelFormat) Format, NumMips, NumSamples, Flags, CreateInfo);
+	return CreateD3D12Texture2D<FD3D12BaseTexture2D>(&RHICmdList, SizeX, SizeY, 1, false, false, (EPixelFormat) Format, NumMips, NumSamples, Flags, InResourceState, CreateInfo);
 }
 
 FTexture2DRHIRef FD3D12DynamicRHI::RHICreateTexture2D(uint32 SizeX, uint32 SizeY, uint8 Format, uint32 NumMips, uint32 NumSamples, ETextureCreateFlags Flags, ERHIAccess InResourceState, FRHIResourceCreateInfo& CreateInfo)
 {
-	return CreateD3D12Texture2D<FD3D12BaseTexture2D>(nullptr, SizeX, SizeY, 1, false, false, (EPixelFormat) Format, NumMips, NumSamples, Flags, CreateInfo);
+	return CreateD3D12Texture2D<FD3D12BaseTexture2D>(nullptr, SizeX, SizeY, 1, false, false, (EPixelFormat) Format, NumMips, NumSamples, Flags, InResourceState, CreateInfo);
 }
 
 FTexture2DRHIRef FD3D12DynamicRHI::RHIAsyncCreateTexture2D(uint32 SizeX, uint32 SizeY, uint8 Format, uint32 NumMips, ETextureCreateFlags Flags, ERHIAccess InResourceState, void** InitialMipData, uint32 NumInitialMips)
@@ -1452,25 +1456,25 @@ void FD3D12DynamicRHI::RHICopySharedMips(FRHITexture2D* DestTexture2DRHI, FRHITe
 FTexture2DArrayRHIRef FD3D12DynamicRHI::RHICreateTexture2DArray_RenderThread(class FRHICommandListImmediate& RHICmdList, uint32 SizeX, uint32 SizeY, uint32 SizeZ, uint8 Format, uint32 NumMips, uint32 NumSamples, ETextureCreateFlags Flags, ERHIAccess InResourceState, FRHIResourceCreateInfo& CreateInfo)
 {
 	check(SizeZ >= 1);
-	return CreateD3D12Texture2D<FD3D12BaseTexture2DArray>(&RHICmdList, SizeX, SizeY, SizeZ, true, false, (EPixelFormat) Format, NumMips, NumSamples, Flags, CreateInfo);
+	return CreateD3D12Texture2D<FD3D12BaseTexture2DArray>(&RHICmdList, SizeX, SizeY, SizeZ, true, false, (EPixelFormat) Format, NumMips, NumSamples, Flags, InResourceState, CreateInfo);
 }
 
 FTexture2DArrayRHIRef FD3D12DynamicRHI::RHICreateTexture2DArray(uint32 SizeX, uint32 SizeY, uint32 SizeZ, uint8 Format, uint32 NumMips, uint32 NumSamples, ETextureCreateFlags Flags, ERHIAccess InResourceState, FRHIResourceCreateInfo& CreateInfo)
 {
 	check(SizeZ >= 1);
-	return CreateD3D12Texture2D<FD3D12BaseTexture2DArray>(nullptr, SizeX, SizeY, SizeZ, true, false, (EPixelFormat) Format, NumMips, NumSamples, Flags, CreateInfo);
+	return CreateD3D12Texture2D<FD3D12BaseTexture2DArray>(nullptr, SizeX, SizeY, SizeZ, true, false, (EPixelFormat) Format, NumMips, NumSamples, Flags, InResourceState, CreateInfo);
 }
 
 FTexture3DRHIRef FD3D12DynamicRHI::RHICreateTexture3D_RenderThread(class FRHICommandListImmediate& RHICmdList, uint32 SizeX, uint32 SizeY, uint32 SizeZ, uint8 Format, uint32 NumMips, ETextureCreateFlags Flags, ERHIAccess InResourceState, FRHIResourceCreateInfo& CreateInfo)
 {
-	return CreateD3D12Texture3D(&RHICmdList, SizeX, SizeY, SizeZ, (EPixelFormat) Format, NumMips, Flags, CreateInfo);
+	return CreateD3D12Texture3D(&RHICmdList, SizeX, SizeY, SizeZ, (EPixelFormat) Format, NumMips, Flags, InResourceState, CreateInfo);
 }
 
 FTexture3DRHIRef FD3D12DynamicRHI::RHICreateTexture3D(uint32 SizeX, uint32 SizeY, uint32 SizeZ, uint8 Format, uint32 NumMips, ETextureCreateFlags Flags, ERHIAccess InResourceState, FRHIResourceCreateInfo& CreateInfo)
 {
 	check(SizeZ >= 1);
 #if PLATFORM_WINDOWS || PLATFORM_HOLOLENS
-	return CreateD3D12Texture3D(nullptr, SizeX, SizeY, SizeZ, (EPixelFormat) Format, NumMips, Flags, CreateInfo);
+	return CreateD3D12Texture3D(nullptr, SizeX, SizeY, SizeZ, (EPixelFormat) Format, NumMips, Flags, InResourceState, CreateInfo);
 #else
 	checkf(false, TEXT("XBOX_CODE_MERGE : Removed. The Xbox platform version should be used."));
 	return nullptr;
@@ -1571,10 +1575,11 @@ FTexture2DRHIRef FD3D12DynamicRHI::AsyncReallocateTexture2D_RenderThread(class F
 	}
 
 	FD3D12Texture2D* Texture2D = FD3D12DynamicRHI::ResourceCast(Texture2DRHI);
-
+	
 	// Allocate a new texture.
 	FRHIResourceCreateInfo CreateInfo;
-	FD3D12Texture2D* NewTexture2D = CreateD3D12Texture2D<FD3D12BaseTexture2D>(nullptr, NewSizeX, NewSizeY, 1, false, false, Texture2DRHI->GetFormat(), NewMipCount, 1, Texture2DRHI->GetFlags(), CreateInfo);
+	ERHIAccess RHIAccess = ERHIAccess::Unknown;
+	FD3D12Texture2D* NewTexture2D = CreateD3D12Texture2D<FD3D12BaseTexture2D>(nullptr, NewSizeX, NewSizeY, 1, false, false, Texture2DRHI->GetFormat(), NewMipCount, 1, Texture2DRHI->GetFlags(), RHIAccess, CreateInfo);
 	
 	ALLOC_COMMAND_CL(RHICmdList, FRHICommandD3D12AsyncReallocateTexture2D)(Texture2D, NewTexture2D, NewMipCount, NewSizeX, NewSizeY, RequestStatus);
 
@@ -1603,7 +1608,8 @@ FTexture2DRHIRef FD3D12DynamicRHI::RHIAsyncReallocateTexture2D(FRHITexture2D* Te
 
 	// Allocate a new texture.
 	FRHIResourceCreateInfo CreateInfo;
-	FD3D12Texture2D* NewTexture2D = CreateD3D12Texture2D<FD3D12BaseTexture2D>(nullptr, NewSizeX, NewSizeY, 1, false, false, Texture2DRHI->GetFormat(), NewMipCount, 1, Texture2DRHI->GetFlags(), CreateInfo);
+	ERHIAccess RHIAccess = ERHIAccess::Unknown;
+	FD3D12Texture2D* NewTexture2D = CreateD3D12Texture2D<FD3D12BaseTexture2D>(nullptr, NewSizeX, NewSizeY, 1, false, false, Texture2DRHI->GetFormat(), NewMipCount, 1, Texture2DRHI->GetFlags(), RHIAccess, CreateInfo);
 	
 	DoAsyncReallocateTexture2D(Texture2D, NewTexture2D, NewMipCount, NewSizeX, NewSizeY, RequestStatus);
 
@@ -1783,7 +1789,7 @@ void FD3D12TextureBase::UpdateTexture(const D3D12_TEXTURE_COPY_LOCATION& DestCop
 	FD3D12CommandContext& DefaultContext = GetParentDevice()->GetDefaultCommandContext();
 	FD3D12CommandListHandle& hCommandList = DefaultContext.CommandListHandle;
 
-	FConditionalScopeResourceBarrier ScopeResourceBarrierDest(hCommandList, GetResource(), D3D12_RESOURCE_STATE_COPY_DEST, DestCopyLocation.SubresourceIndex);
+	FScopedResourceBarrier ScopeResourceBarrierDest(hCommandList, GetResource(), D3D12_RESOURCE_STATE_COPY_DEST, DestCopyLocation.SubresourceIndex, FD3D12DynamicRHI::ETransitionMode::Apply);
 	// Don't need to transition upload heaps
 
 	DefaultContext.numCopies++;
@@ -1807,8 +1813,8 @@ void FD3D12TextureBase::CopyTextureRegion(uint32 DestX, uint32 DestY, uint32 Des
 	CD3DX12_TEXTURE_COPY_LOCATION DestCopyLocation(GetResource()->GetResource(), 0);
 	CD3DX12_TEXTURE_COPY_LOCATION SourceCopyLocation(SourceTexture->GetResource()->GetResource(), 0);
 
-	FConditionalScopeResourceBarrier ConditionalScopeResourceBarrierDest(CommandListHandle, GetResource(), D3D12_RESOURCE_STATE_COPY_DEST, DestCopyLocation.SubresourceIndex);
-	FConditionalScopeResourceBarrier ConditionalScopeResourceBarrierSource(CommandListHandle, SourceTexture->GetResource(), D3D12_RESOURCE_STATE_COPY_SOURCE, SourceCopyLocation.SubresourceIndex);
+	FScopedResourceBarrier ConditionalScopeResourceBarrierDest(CommandListHandle, GetResource(), D3D12_RESOURCE_STATE_COPY_DEST, DestCopyLocation.SubresourceIndex, FD3D12DynamicRHI::ETransitionMode::Apply);
+	FScopedResourceBarrier ConditionalScopeResourceBarrierSource(CommandListHandle, SourceTexture->GetResource(), D3D12_RESOURCE_STATE_COPY_SOURCE, SourceCopyLocation.SubresourceIndex, FD3D12DynamicRHI::ETransitionMode::Apply);
 
 	CommandListHandle.FlushResourceBarriers();
 	CommandListHandle->CopyTextureRegion(
@@ -2243,11 +2249,12 @@ public:
 
 			CD3DX12_TEXTURE_COPY_LOCATION DestCopyLocation(TextureLink.GetResource()->GetResource(), MipIdx);
 
-			FConditionalScopeResourceBarrier ScopeResourceBarrierDest(
+			FScopedResourceBarrier ScopeResourceBarrierDest(
 				NativeCmdList,
 				TextureLink.GetResource(),
 				D3D12_RESOURCE_STATE_COPY_DEST,
-				DestCopyLocation.SubresourceIndex);
+				DestCopyLocation.SubresourceIndex,
+				FD3D12DynamicRHI::ETransitionMode::Apply);
 
 			NativeCmdList.FlushResourceBarriers();
 			Device->GetDefaultCommandContext().numCopies += UpdateInfos.Num();
@@ -2489,11 +2496,12 @@ public:
 			CD3DX12_TEXTURE_COPY_LOCATION DestCopyLocation(TextureLink.GetResource()->GetResource(), MipIdx);
 			CD3DX12_TEXTURE_COPY_LOCATION SourceCopyLocation(UploadBuffer->GetResource(), PlacedSubresourceFootprint);
 
-			FConditionalScopeResourceBarrier ScopeResourceBarrierDest(
+			FScopedResourceBarrier ScopeResourceBarrierDest(
 				NativeCmdList,
 				TextureLink.GetResource(),
 				D3D12_RESOURCE_STATE_COPY_DEST,
-				DestCopyLocation.SubresourceIndex);
+				DestCopyLocation.SubresourceIndex,
+				FD3D12DynamicRHI::ETransitionMode::Apply);
 
 			Device->GetDefaultCommandContext().numCopies++;
 			NativeCmdList.FlushResourceBarriers();
@@ -2562,22 +2570,22 @@ void FD3D12DynamicRHI::EndUpdateTexture3D_Internal(FUpdateTexture3DData& UpdateD
 	-----------------------------------------------------------------------------*/
 FTextureCubeRHIRef FD3D12DynamicRHI::RHICreateTextureCube_RenderThread(class FRHICommandListImmediate& RHICmdList, uint32 Size, uint8 Format, uint32 NumMips, ETextureCreateFlags Flags, ERHIAccess InResourceState, FRHIResourceCreateInfo& CreateInfo)
 {
-	return CreateD3D12Texture2D<FD3D12BaseTextureCube>(&RHICmdList, Size, Size, 6, false, true, (EPixelFormat) Format, NumMips, 1, Flags, CreateInfo);
+	return CreateD3D12Texture2D<FD3D12BaseTextureCube>(&RHICmdList, Size, Size, 6, false, true, (EPixelFormat) Format, NumMips, 1, Flags, InResourceState, CreateInfo);
 }
 
 FTextureCubeRHIRef FD3D12DynamicRHI::RHICreateTextureCube(uint32 Size, uint8 Format, uint32 NumMips, ETextureCreateFlags Flags, ERHIAccess InResourceState, FRHIResourceCreateInfo& CreateInfo)
 {
-	return CreateD3D12Texture2D<FD3D12BaseTextureCube>(nullptr, Size, Size, 6, false, true, (EPixelFormat) Format, NumMips, 1, Flags, CreateInfo);
+	return CreateD3D12Texture2D<FD3D12BaseTextureCube>(nullptr, Size, Size, 6, false, true, (EPixelFormat) Format, NumMips, 1, Flags, InResourceState, CreateInfo);
 }
 
 FTextureCubeRHIRef FD3D12DynamicRHI::RHICreateTextureCubeArray_RenderThread(class FRHICommandListImmediate& RHICmdList, uint32 Size, uint32 ArraySize, uint8 Format, uint32 NumMips, ETextureCreateFlags Flags, ERHIAccess InResourceState, FRHIResourceCreateInfo& CreateInfo)
 {
-	return CreateD3D12Texture2D<FD3D12BaseTextureCube>(&RHICmdList, Size, Size, 6 * ArraySize, true, true, (EPixelFormat) Format, NumMips, 1, Flags, CreateInfo);
+	return CreateD3D12Texture2D<FD3D12BaseTextureCube>(&RHICmdList, Size, Size, 6 * ArraySize, true, true, (EPixelFormat) Format, NumMips, 1, Flags, InResourceState, CreateInfo);
 }
 
 FTextureCubeRHIRef FD3D12DynamicRHI::RHICreateTextureCubeArray(uint32 Size, uint32 ArraySize, uint8 Format, uint32 NumMips, ETextureCreateFlags Flags, ERHIAccess InResourceState, FRHIResourceCreateInfo& CreateInfo)
 {
-	return CreateD3D12Texture2D<FD3D12BaseTextureCube>(nullptr, Size, Size, 6 * ArraySize, true, true, (EPixelFormat) Format, NumMips, 1, Flags, CreateInfo);
+	return CreateD3D12Texture2D<FD3D12BaseTextureCube>(nullptr, Size, Size, 6 * ArraySize, true, true, (EPixelFormat) Format, NumMips, 1, Flags, InResourceState, CreateInfo);
 }
 
 void* FD3D12DynamicRHI::RHILockTextureCubeFace(FRHITextureCube* TextureCubeRHI, uint32 FaceIndex, uint32 ArrayIndex, uint32 MipIndex, EResourceLockMode LockMode, uint32& DestStride, bool bLockWithinMiptail)
@@ -2712,7 +2720,7 @@ TD3D12Texture2D<BaseResourceType>* FD3D12DynamicRHI::CreateTextureFromResource(b
 
 	// The state this resource will be in when it leaves this function
 	const FD3D12Resource::FD3D12ResourceTypeHelper Type(TextureDesc, D3D12_HEAP_TYPE_DEFAULT);
-	const D3D12_RESOURCE_STATES DestinationState = Type.GetOptimalInitialState(!(TexCreateFlags & TexCreate_Shared));
+	const D3D12_RESOURCE_STATES DestinationState = Type.GetOptimalInitialState(ERHIAccess::Unknown, !(TexCreateFlags & TexCreate_Shared));
 
 	FD3D12Device* Device = Adapter->GetDevice(0);
 	FD3D12Resource* TextureResource = new FD3D12Resource(Device, Device->GetGPUMask(), Resource, DestinationState, TextureDesc);
@@ -3090,8 +3098,8 @@ void FD3D12CommandContext::RHICopyTexture(FRHITexture* SourceTextureRHI, FRHITex
 	FD3D12TextureBase* SourceTexture = RetrieveTextureBase(SourceTextureRHI);
 	FD3D12TextureBase* DestTexture = RetrieveTextureBase(DestTextureRHI);
 
-	FConditionalScopeResourceBarrier ConditionalScopeResourceBarrierSource(CommandListHandle, SourceTexture->GetResource(), D3D12_RESOURCE_STATE_COPY_SOURCE, D3D12_RESOURCE_BARRIER_ALL_SUBRESOURCES);
-	FConditionalScopeResourceBarrier ConditionalScopeResourceBarrierDest(CommandListHandle, DestTexture->GetResource(), D3D12_RESOURCE_STATE_COPY_DEST, D3D12_RESOURCE_BARRIER_ALL_SUBRESOURCES);
+	FScopedResourceBarrier ConditionalScopeResourceBarrierSource(CommandListHandle, SourceTexture->GetResource(), D3D12_RESOURCE_STATE_COPY_SOURCE, D3D12_RESOURCE_BARRIER_ALL_SUBRESOURCES, FD3D12DynamicRHI::ETransitionMode::Validate);
+	FScopedResourceBarrier ConditionalScopeResourceBarrierDest(CommandListHandle, DestTexture->GetResource(), D3D12_RESOURCE_STATE_COPY_DEST, D3D12_RESOURCE_BARRIER_ALL_SUBRESOURCES, FD3D12DynamicRHI::ETransitionMode::Validate);
 
 	numCopies++;
 	CommandListHandle.FlushResourceBarriers();
diff --git a/Engine/Source/Runtime/D3D12RHI/Private/D3D12UAV.cpp b/Engine/Source/Runtime/D3D12RHI/Private/D3D12UAV.cpp
--- a/Engine/Source/Runtime/D3D12RHI/Private/D3D12UAV.cpp
+++ b/Engine/Source/Runtime/D3D12RHI/Private/D3D12UAV.cpp
@@ -238,7 +238,7 @@ void FD3D12CommandContext::ClearUAV(TRHICommandList_RecursiveHazardous<FD3D12Com
 				D3D12_GPU_DESCRIPTOR_HANDLE GPUHandle = Context.StateCache.GetDescriptorCache()->GetCurrentViewHeap()->GetGPUSlotHandle(ReservedSlot);
 
 				Device->CopyDescriptorsSimple(1, DestSlot, CPUHandle, D3D12_DESCRIPTOR_HEAP_TYPE_CBV_SRV_UAV);
-				FD3D12DynamicRHI::TransitionResource(Context.CommandListHandle, UnorderedAccessView, D3D12_RESOURCE_STATE_UNORDERED_ACCESS);
+				FD3D12DynamicRHI::TransitionResource(Context.CommandListHandle, UnorderedAccessView, D3D12_RESOURCE_STATE_UNORDERED_ACCESS, FD3D12DynamicRHI::ETransitionMode::Validate);
 				Context.numClears++;
 
 				Context.CommandListHandle.FlushResourceBarriers();
diff --git a/Engine/Source/Runtime/D3D12RHI/Private/D3D12Util.cpp b/Engine/Source/Runtime/D3D12RHI/Private/D3D12Util.cpp
--- a/Engine/Source/Runtime/D3D12RHI/Private/D3D12Util.cpp
+++ b/Engine/Source/Runtime/D3D12RHI/Private/D3D12Util.cpp
@@ -1073,6 +1073,81 @@ void LogResourceBarriers(uint32 NumBarriers, D3D12_RESOURCE_BARRIER* pBarriers,
 }
 
 
+D3D12_RESOURCE_STATES GetD3D12ResourceState(ERHIAccess InRHIAccess, bool InIsAsyncCompute)
+{
+	// Add switch for common states (should cover all writeable states)
+	switch (InRHIAccess)
+	{
+	// all single write states
+	case ERHIAccess::RTV:					return D3D12_RESOURCE_STATE_RENDER_TARGET;
+	case ERHIAccess::UAVMask:		
+	case ERHIAccess::UAVCompute:	
+	case ERHIAccess::UAVGraphics:			return D3D12_RESOURCE_STATE_UNORDERED_ACCESS;
+	case ERHIAccess::DSVWrite:				return D3D12_RESOURCE_STATE_DEPTH_WRITE;
+	case ERHIAccess::CopyDest:				return D3D12_RESOURCE_STATE_COPY_DEST;
+	case ERHIAccess::ResolveDst:			return D3D12_RESOURCE_STATE_RESOLVE_DEST;
+
+	// Generic read for mask read states
+	case ERHIAccess::ReadOnlyMask:	
+	case ERHIAccess::ReadOnlyExclusiveMask:	return D3D12_RESOURCE_STATE_GENERIC_READ;
+	default:
+	{
+		// Special case for DSV read & write (Depth write allows depth read as well in D3D)
+		if (InRHIAccess == ERHIAccess(ERHIAccess::DSVRead | ERHIAccess::DSVWrite))
+		{
+			return D3D12_RESOURCE_STATE_DEPTH_WRITE;
+		}
+		else
+		{
+			// Should be combination from read only flags (write flags covered above)
+			check(!(EnumHasAnyFlags(InRHIAccess, ERHIAccess::WritableMask)));
+			check(EnumHasAnyFlags(InRHIAccess, ERHIAccess::ReadOnlyMask));
+
+			D3D12_RESOURCE_STATES State = D3D12_RESOURCE_STATE_COMMON;
+
+			// Translate the requested after state to a D3D state
+			if (EnumHasAnyFlags(InRHIAccess, ERHIAccess::SRVGraphics) && !InIsAsyncCompute)
+			{
+				State |= D3D12_RESOURCE_STATE_PIXEL_SHADER_RESOURCE | D3D12_RESOURCE_STATE_NON_PIXEL_SHADER_RESOURCE;
+			}
+			if (EnumHasAnyFlags(InRHIAccess, ERHIAccess::SRVCompute))
+			{
+				State |= D3D12_RESOURCE_STATE_NON_PIXEL_SHADER_RESOURCE;
+			}
+			if (EnumHasAnyFlags(InRHIAccess, ERHIAccess::VertexOrIndexBuffer))
+			{
+				State |= D3D12_RESOURCE_STATE_VERTEX_AND_CONSTANT_BUFFER | D3D12_RESOURCE_STATE_INDEX_BUFFER;
+			}
+			if (EnumHasAnyFlags(InRHIAccess, ERHIAccess::CopySrc))
+			{
+				State |= D3D12_RESOURCE_STATE_COPY_SOURCE;
+			}
+			if (EnumHasAnyFlags(InRHIAccess, ERHIAccess::IndirectArgs))
+			{
+				State |= D3D12_RESOURCE_STATE_INDIRECT_ARGUMENT;
+			}
+			if (EnumHasAnyFlags(InRHIAccess, ERHIAccess::ResolveSrc))
+			{
+				State |= D3D12_RESOURCE_STATE_RESOLVE_SOURCE;
+			}
+			if (EnumHasAnyFlags(InRHIAccess, ERHIAccess::DSVRead))
+			{
+				State |= D3D12_RESOURCE_STATE_DEPTH_READ;
+			}
+
+			// Should have at least one valid state
+			check(State != D3D12_RESOURCE_STATE_COMMON);
+
+			return State;
+		}
+	}
+	}
+
+	// unreachable code
+	return D3D12_RESOURCE_STATE_COMMON;
+}
+
+
 //==================================================================================================================================
 // CResourceState
 // Tracking of per-resource or per-subresource state
@@ -1140,6 +1215,35 @@ D3D12_RESOURCE_STATES CResourceState::GetSubresourceState(uint32 SubresourceInde
 	}
 }
 
+//----------------------------------------------------------------------------------------------------------------------------------
+
+bool CResourceState::CheckAllSubresourceSame()
+{
+	// already marked same?
+	if (m_AllSubresourcesSame)
+	{
+		return true;
+	}
+	else
+	{
+		D3D12_RESOURCE_STATES State = m_SubresourceState[0];
+
+		// All subresources must be individually checked
+		const uint32 numSubresourceStates = m_SubresourceState.Num();
+		for (uint32 i = 1; i < numSubresourceStates; i++)
+		{
+			if (m_SubresourceState[i] != State)
+			{
+				return false;
+			}
+		}
+
+		SetResourceState(State);
+
+		return true;
+	}
+}
+
 //----------------------------------------------------------------------------------------------------------------------------------
 void CResourceState::SetResourceState(D3D12_RESOURCE_STATES State)
 {
diff --git a/Engine/Source/Runtime/D3D12RHI/Private/D3D12VertexBuffer.cpp b/Engine/Source/Runtime/D3D12RHI/Private/D3D12VertexBuffer.cpp
--- a/Engine/Source/Runtime/D3D12RHI/Private/D3D12VertexBuffer.cpp
+++ b/Engine/Source/Runtime/D3D12RHI/Private/D3D12VertexBuffer.cpp
@@ -49,7 +49,7 @@ FVertexBufferRHIRef FD3D12DynamicRHI::RHICreateVertexBuffer(uint32 Size, uint32
 	const D3D12_RESOURCE_DESC Desc = CreateVertexBufferResourceDesc(Size, InUsage);
 	const uint32 Alignment = 4;
 
-	FD3D12Buffer* Buffer = GetAdapter().CreateRHIBuffer(nullptr, Desc, Alignment, 0, Size, InUsage | BUF_VertexBuffer, ED3D12ResourceStateMode::Default, CreateInfo);
+	FD3D12Buffer* Buffer = GetAdapter().CreateRHIBuffer(nullptr, Desc, Alignment, 0, Size, InUsage | BUF_VertexBuffer, ED3D12ResourceStateMode::Default, InResourceState, CreateInfo);
 	if (Buffer->ResourceLocation.IsTransient() )
 	{
 		// TODO: this should ideally be set in platform-independent code, since this tracking is for the high level
@@ -84,7 +84,7 @@ FVertexBufferRHIRef FD3D12DynamicRHI::CreateVertexBuffer_RenderThread(FRHIComman
 	const D3D12_RESOURCE_DESC Desc = CreateVertexBufferResourceDesc(Size, InUsage);
 	const uint32 Alignment = 4;
 
-	FD3D12Buffer* Buffer = GetAdapter().CreateRHIBuffer(&RHICmdList, Desc, Alignment, 0, Size, InUsage | BUF_VertexBuffer, ED3D12ResourceStateMode::Default, CreateInfo);
+	FD3D12Buffer* Buffer = GetAdapter().CreateRHIBuffer(&RHICmdList, Desc, Alignment, 0, Size, InUsage | BUF_VertexBuffer, ED3D12ResourceStateMode::Default, InResourceState, CreateInfo);
 	if (Buffer->ResourceLocation.IsTransient())
 	{
 		// TODO: this should ideally be set in platform-independent code, since this tracking is for the high level
@@ -154,8 +154,8 @@ void FD3D12CommandContext::RHICopyBufferRegion(FRHIVertexBuffer* DestBufferRHI,
 
 	numCopies++;
 
-	FConditionalScopeResourceBarrier ScopeResourceBarrierSource(CommandListHandle, pSourceResource, D3D12_RESOURCE_STATE_COPY_SOURCE, 0);
-	FConditionalScopeResourceBarrier ScopeResourceBarrierDest(CommandListHandle, pDestResource, D3D12_RESOURCE_STATE_COPY_DEST, 0);
+	FScopedResourceBarrier ScopeResourceBarrierSource(CommandListHandle, pSourceResource, D3D12_RESOURCE_STATE_COPY_SOURCE, 0, FD3D12DynamicRHI::ETransitionMode::Validate);
+	FScopedResourceBarrier ScopeResourceBarrierDest(CommandListHandle, pDestResource, D3D12_RESOURCE_STATE_COPY_DEST, 0, FD3D12DynamicRHI::ETransitionMode::Validate);
 	CommandListHandle.FlushResourceBarriers();
 
 	CommandListHandle->CopyBufferRegion(pDestResource->GetResource(), DestBuffer->ResourceLocation.GetOffsetFromBaseOfResource() + DstOffset, pSourceResource->GetResource(), SourceBuffer->ResourceLocation.GetOffsetFromBaseOfResource() + SrcOffset, NumBytes);
@@ -252,7 +252,7 @@ void FD3D12CommandContext::RHICopyBufferRegions(const TArrayView<const FCopyBuff
 			}
 			else
 			{
-				FD3D12DynamicRHI::TransitionResource(CommandListHandle, Resource, DesiredState, Subresource);
+				FD3D12DynamicRHI::TransitionResource(CommandListHandle, Resource, D3D12_RESOURCE_STATE_TBD, DesiredState, Subresource, FD3D12DynamicRHI::ETransitionMode::Validate);
 			}
 
 			PrevResource = Resource;
@@ -309,7 +309,7 @@ FVertexBufferRHIRef FD3D12DynamicRHI::CreateAndLockVertexBuffer_RenderThread(FRH
 	const D3D12_RESOURCE_DESC Desc = CreateVertexBufferResourceDesc(Size, InUsage);
 	const uint32 Alignment = 4;
 
-	FD3D12Buffer* Buffer = GetAdapter().CreateRHIBuffer(nullptr, Desc, Alignment, 0, Size, InUsage | BUF_VertexBuffer, ED3D12ResourceStateMode::Default, CreateInfo);
+	FD3D12Buffer* Buffer = GetAdapter().CreateRHIBuffer(nullptr, Desc, Alignment, 0, Size, InUsage | BUF_VertexBuffer, ED3D12ResourceStateMode::Default, InResourceState, CreateInfo);
 	if (Buffer->ResourceLocation.IsTransient())
 	{
 		// TODO: this should ideally be set in platform-independent code, since this tracking is for the high level
diff --git a/Engine/Source/Runtime/D3D12RHI/Private/D3D12Viewport.cpp b/Engine/Source/Runtime/D3D12RHI/Private/D3D12Viewport.cpp
--- a/Engine/Source/Runtime/D3D12RHI/Private/D3D12Viewport.cpp
+++ b/Engine/Source/Runtime/D3D12RHI/Private/D3D12Viewport.cpp
@@ -824,10 +824,10 @@ bool FD3D12Viewport::Present(bool bLockToVsync)
 		FD3D12Texture2D* DeviceBackBuffer = DefaultContext.RetrieveObject<FD3D12Texture2D, FRHITexture2D*>(GetBackBuffer_RHIThread());
 		FD3D12Texture2D* DeviceSDRBackBuffer = DefaultContext.RetrieveObject<FD3D12Texture2D, FRHITexture2D*>(GetSDRBackBuffer_RHIThread());
 
-		FD3D12DynamicRHI::TransitionResource(DefaultContext.CommandListHandle, DeviceBackBuffer->GetShaderResourceView(), D3D12_RESOURCE_STATE_PRESENT);
+		FD3D12DynamicRHI::TransitionResource(DefaultContext.CommandListHandle, DeviceBackBuffer->GetShaderResourceView()->GetResource(), D3D12_RESOURCE_STATE_TBD, D3D12_RESOURCE_STATE_PRESENT, 0, FD3D12DynamicRHI::ETransitionMode::Apply);
 		if (SDRBackBuffer_RHIThread != nullptr)
 		{
-			FD3D12DynamicRHI::TransitionResource(DefaultContext.CommandListHandle, DeviceSDRBackBuffer->GetShaderResourceView(), D3D12_RESOURCE_STATE_PRESENT);
+			FD3D12DynamicRHI::TransitionResource(DefaultContext.CommandListHandle, DeviceSDRBackBuffer->GetShaderResourceView()->GetResource(), D3D12_RESOURCE_STATE_TBD, D3D12_RESOURCE_STATE_PRESENT, 0, FD3D12DynamicRHI::ETransitionMode::Apply);
 		}
 		DefaultContext.CommandListHandle.FlushResourceBarriers();
 		DefaultContext.FlushCommands();
diff --git a/Engine/Source/Runtime/D3D12RHI/Private/Windows/WindowsD3D12Device.cpp b/Engine/Source/Runtime/D3D12RHI/Private/Windows/WindowsD3D12Device.cpp
--- a/Engine/Source/Runtime/D3D12RHI/Private/Windows/WindowsD3D12Device.cpp
+++ b/Engine/Source/Runtime/D3D12RHI/Private/Windows/WindowsD3D12Device.cpp
@@ -47,6 +47,25 @@ extern bool D3D12RHI_ShouldForceCompatibility();
 
 FD3D12DynamicRHI* GD3D12RHI = nullptr;
 
+bool GUseInternalTransitions = true;
+static FAutoConsoleVariableRef CVarUseInternalTransitions(
+	TEXT("r.D3D12.UseInternalTransitions"),
+	GUseInternalTransitions,
+	TEXT("Use the D3D12 RHI internal transitions to drive all resource transitions"),
+	ECVF_ReadOnly
+);
+
+#if UE_BUILD_SHIPPING || UE_BUILD_TEST
+bool GValidateInternalTransitions = false;
+#else
+bool GValidateInternalTransitions = true;
+#endif
+static FAutoConsoleVariableRef CVarValidateInternalTransitions(
+	TEXT("r.D3D12.ValidateInternalTransitions"),
+	GValidateInternalTransitions,
+	TEXT("Use the D3D12 RHI internal transitions to validate the engine pushed RHI transitions")
+);
+
 #if NV_AFTERMATH
 // Disabled by default since introduces stalls between render and driver threads
 int32 GDX12NVAfterMathEnabled = 0;
diff --git a/Engine/Source/Runtime/D3D12RHI/Public/D3D12Resources.h b/Engine/Source/Runtime/D3D12RHI/Public/D3D12Resources.h
--- a/Engine/Source/Runtime/D3D12RHI/Public/D3D12Resources.h
+++ b/Engine/Source/Runtime/D3D12RHI/Public/D3D12Resources.h
@@ -247,7 +247,7 @@ public:
 
 	struct FD3D12ResourceTypeHelper
 	{
-		FD3D12ResourceTypeHelper(D3D12_RESOURCE_DESC& Desc, D3D12_HEAP_TYPE HeapType) :
+		FD3D12ResourceTypeHelper(const D3D12_RESOURCE_DESC& Desc, D3D12_HEAP_TYPE HeapType) :
 			bSRV((Desc.Flags & D3D12_RESOURCE_FLAG_DENY_SHADER_RESOURCE) == 0),
 			bDSV((Desc.Flags & D3D12_RESOURCE_FLAG_ALLOW_DEPTH_STENCIL) != 0),
 			bRTV((Desc.Flags & D3D12_RESOURCE_FLAG_ALLOW_RENDER_TARGET) != 0),
@@ -258,41 +258,51 @@ public:
 			bReadBackResource(HeapType == D3D12_HEAP_TYPE_READBACK)
 		{}
 
-		const D3D12_RESOURCE_STATES GetOptimalInitialState(bool bAccurateWriteableStates) const
+		const D3D12_RESOURCE_STATES GetOptimalInitialState(ERHIAccess InResourceState, bool bAccurateWriteableStates) const
 		{
-			if (bSRVOnly)
+			// Ignore the requested resource state for non tracked resource because RHI will assume it's always in default resource 
+			// state then when a transition is required (will transition via scoped push/pop to requested state)
+			if (!bSRVOnly && InResourceState != ERHIAccess::Unknown)
 			{
-				return D3D12_RESOURCE_STATE_NON_PIXEL_SHADER_RESOURCE | D3D12_RESOURCE_STATE_PIXEL_SHADER_RESOURCE;
+				bool bAsyncCompute = false;
+				return GetD3D12ResourceState(InResourceState, bAsyncCompute);
 			}
-			else if (bBuffer && !bUAV)
-			{
-				return (bReadBackResource) ? D3D12_RESOURCE_STATE_COPY_DEST : D3D12_RESOURCE_STATE_GENERIC_READ;
-			}
-			else if (bWritable)
+			else
 			{
-				if (bAccurateWriteableStates)
+				if (bSRVOnly)
 				{
-					if (bDSV)
-					{
-						return D3D12_RESOURCE_STATE_DEPTH_WRITE;
-					}
-					else if (bRTV)
+					return D3D12_RESOURCE_STATE_NON_PIXEL_SHADER_RESOURCE | D3D12_RESOURCE_STATE_PIXEL_SHADER_RESOURCE;
+				}
+				else if (bBuffer && !bUAV)
+				{
+					return (bReadBackResource) ? D3D12_RESOURCE_STATE_COPY_DEST : D3D12_RESOURCE_STATE_GENERIC_READ;
+				}
+				else if (bWritable)
+				{
+					if (bAccurateWriteableStates)
 					{
-						return D3D12_RESOURCE_STATE_RENDER_TARGET;
+						if (bDSV)
+						{
+							return D3D12_RESOURCE_STATE_DEPTH_WRITE;
+						}
+						else if (bRTV)
+						{
+							return D3D12_RESOURCE_STATE_RENDER_TARGET;
+						}
+						else if (bUAV)
+						{
+							return D3D12_RESOURCE_STATE_UNORDERED_ACCESS;
+						}
 					}
-					else if (bUAV)
+					else
 					{
-						return D3D12_RESOURCE_STATE_UNORDERED_ACCESS;
+						// This things require tracking anyway
+						return D3D12_RESOURCE_STATE_COMMON;
 					}
 				}
-				else
-				{
-					// This things require tracking anyway
-					return D3D12_RESOURCE_STATE_COMMON;
-				}
-			}
 
-			return D3D12_RESOURCE_STATE_COMMON;
+				return D3D12_RESOURCE_STATE_COMMON;
+			}
 		}
 
 		const uint32 bSRV : 1;
diff --git a/Engine/Source/Runtime/D3D12RHI/Public/D3D12Util.h b/Engine/Source/Runtime/D3D12RHI/Public/D3D12Util.h
--- a/Engine/Source/Runtime/D3D12RHI/Public/D3D12Util.h
+++ b/Engine/Source/Runtime/D3D12RHI/Public/D3D12Util.h
@@ -281,6 +281,8 @@ static bool IsValidD3D12ResourceState(D3D12_RESOURCE_STATES InState)
 	return (InState != D3D12_RESOURCE_STATE_TBD && InState != D3D12_RESOURCE_STATE_CORRUPT);
 }
 
+D3D12_RESOURCE_STATES GetD3D12ResourceState(ERHIAccess InRHIAccess, bool InIsAsyncCompute);
+
 //==================================================================================================================================
 // CResourceState
 // Tracking of per-resource or per-subresource state
@@ -294,6 +296,7 @@ public:
 	bool CheckResourceState(D3D12_RESOURCE_STATES State) const;
 	bool CheckResourceStateInitalized() const;
 	D3D12_RESOURCE_STATES GetSubresourceState(uint32 SubresourceIndex) const;
+	bool CheckAllSubresourceSame();
 	void SetResourceState(D3D12_RESOURCE_STATES State);
 	void SetSubresourceState(uint32 SubresourceIndex, D3D12_RESOURCE_STATES State);
 
