- Create a time stamp query heap per command queue type to fix issues where multiple queues access the same readback resource at the same time
- Rename ED3D12CommandQueueType::Default to ED3D12CommandQueueType::Direct

diff --git a/Engine/Source/Runtime/D3D12RHI/Private/D3D12Adapter.cpp b/Engine/Source/Runtime/D3D12RHI/Private/D3D12Adapter.cpp
--- a/Engine/Source/Runtime/D3D12RHI/Private/D3D12Adapter.cpp
+++ b/Engine/Source/Runtime/D3D12RHI/Private/D3D12Adapter.cpp
@@ -176,8 +176,8 @@ FD3D12Adapter::FD3D12Adapter(FD3D12AdapterDesc& DescIn)
 	, PipelineStateCache(this)
 	, FenceCorePool(this)
 	, DeferredDeletionQueue(this)
-	, DefaultContextRedirector(this, true, false)
-	, DefaultAsyncComputeContextRedirector(this, true, true)
+	, DefaultContextRedirector(this, ED3D12CommandQueueType::Direct, true)
+	, DefaultAsyncComputeContextRedirector(this, ED3D12CommandQueueType::Async, true)
 	, FrameCounter(0)
 	, DebugFlags(0)
 {
diff --git a/Engine/Source/Runtime/D3D12RHI/Private/D3D12CommandContext.cpp b/Engine/Source/Runtime/D3D12RHI/Private/D3D12CommandContext.cpp
--- a/Engine/Source/Runtime/D3D12RHI/Private/D3D12CommandContext.cpp
+++ b/Engine/Source/Runtime/D3D12RHI/Private/D3D12CommandContext.cpp
@@ -36,16 +36,16 @@ static FAutoConsoleVariableRef CVarEmitRgpFrameMarkers(
 	ECVF_ReadOnly | ECVF_RenderThreadSafe
 );
 
-FD3D12CommandContextBase::FD3D12CommandContextBase(class FD3D12Adapter* InParentAdapter, FRHIGPUMask InGPUMask, bool InIsDefaultContext, bool InIsAsyncComputeContext)
+FD3D12CommandContextBase::FD3D12CommandContextBase(class FD3D12Adapter* InParentAdapter, FRHIGPUMask InGPUMask, ED3D12CommandQueueType InCommandQueueType, bool InIsDefaultContext)
 	: FD3D12AdapterChild(InParentAdapter)
 	, GPUMask(InGPUMask)
 	, bTrackingEvents(false)
+	, CommandQueueType(InCommandQueueType)
 	, bIsDefaultContext(InIsDefaultContext)
-	, bIsAsyncComputeContext(InIsAsyncComputeContext)
 {
 }
 
-static D3D12_RESOURCE_STATES GetValidResourceStates(D3D12_COMMAND_LIST_TYPE CommandListType)
+static D3D12_RESOURCE_STATES GetValidResourceStates(ED3D12CommandQueueType CommandListType)
 {
 	// For reasons, we can't just list the allowed states, we have to list the disallowed states.
 	// For reference on allowed/disallowed states, see:
@@ -73,12 +73,12 @@ static D3D12_RESOURCE_STATES GetValidResourceStates(D3D12_COMMAND_LIST_TYPE Comm
 		D3D12_RESOURCE_STATE_NON_PIXEL_SHADER_RESOURCE;
 
 
-	if (CommandListType == D3D12_COMMAND_LIST_TYPE_COPY)
+	if (CommandListType == ED3D12CommandQueueType::Copy)
 	{
 		return ~DisallowedCopyStates;
 	}
 
-	if (CommandListType == D3D12_COMMAND_LIST_TYPE_COMPUTE)
+	if (CommandListType == ED3D12CommandQueueType::Async)
 	{
 		return ~DisallowedComputeStates;
 	}
@@ -86,13 +86,13 @@ static D3D12_RESOURCE_STATES GetValidResourceStates(D3D12_COMMAND_LIST_TYPE Comm
 	return ~DisallowedDirectStates;
 }
 
-FD3D12CommandContext::FD3D12CommandContext(FD3D12Device* InParent, bool InIsDefaultContext, bool InIsAsyncComputeContext) :
-	FD3D12CommandContextBase(InParent->GetParentAdapter(), InParent->GetGPUMask(), InIsDefaultContext, InIsAsyncComputeContext),
+FD3D12CommandContext::FD3D12CommandContext(FD3D12Device* InParent, ED3D12CommandQueueType InCommandQueueType, bool InIsDefaultContext) :
+	FD3D12CommandContextBase(InParent->GetParentAdapter(), InParent->GetGPUMask(), InCommandQueueType, InIsDefaultContext),
 	FD3D12DeviceChild(InParent),
 	ConstantsAllocator(InParent, InParent->GetGPUMask()),
 	CommandListHandle(),
 	CommandAllocator(nullptr),
-	CommandAllocatorManager(InParent, InIsAsyncComputeContext ? D3D12_COMMAND_LIST_TYPE_COMPUTE : D3D12_COMMAND_LIST_TYPE_DIRECT),
+	CommandAllocatorManager(InParent, (InCommandQueueType == ED3D12CommandQueueType::Async) ? D3D12_COMMAND_LIST_TYPE_COMPUTE : D3D12_COMMAND_LIST_TYPE_DIRECT),
 	StateCache(InParent->GetGPUMask()),
 	OwningRHI(*InParent->GetOwningRHI()),
 	CurrentDepthStencilTarget(nullptr),
@@ -108,7 +108,7 @@ FD3D12CommandContext::FD3D12CommandContext(FD3D12Device* InParent, bool InIsDefa
 	VRSShadingRate(D3D12_SHADING_RATE_1X1),
 #endif
 	SkipFastClearEliminateState(D3D12_RESOURCE_STATES(0)),
-	ValidResourceStates(GetValidResourceStates(InIsAsyncComputeContext ? D3D12_COMMAND_LIST_TYPE_COMPUTE : D3D12_COMMAND_LIST_TYPE_DIRECT)),
+	ValidResourceStates(GetValidResourceStates(InCommandQueueType)),
 #if PLATFORM_SUPPORTS_VIRTUAL_TEXTURES
 	bNeedFlushTextureCache(false),
 #endif
@@ -276,7 +276,7 @@ void FD3D12CommandContext::RHIPopEvent()
 
 FD3D12CommandListManager& FD3D12CommandContext::GetCommandListManager()
 {
-	return bIsAsyncComputeContext ? GetParentDevice()->GetAsyncCommandListManager() : GetParentDevice()->GetCommandListManager();
+	return *GetParentDevice()->GetCommandListManager(CommandQueueType);
 }
 
 void FD3D12CommandContext::ConditionalObtainCommandAllocator()
@@ -338,7 +338,7 @@ FD3D12CommandListHandle FD3D12CommandContext::FlushCommands(bool WaitForCompleti
 	bool bHasProfileGPUAction = false;
 #if WITH_PROFILEGPU || D3D12_SUBMISSION_GAP_RECORDER
 	// Only graphics command list supports ID3D12GraphicsCommandList::EndQuery currently
-	if (!bIsAsyncComputeContext)
+	if (!IsAsyncComputeContext())
 	{
 		if (ExtraAction == FCEA_StartProfilingGPU)
 		{
@@ -457,10 +457,10 @@ void FD3D12CommandContextBase::RHIBeginFrame()
 		FD3D12CommandContext* ContextAtIndex = GetContext(GPUIndex);
 		if (ensure(ContextAtIndex))
 		{
-			Device->GetTimestampQueryHeap()->EndQueryBatchAndResolveQueryData(*ContextAtIndex);
+			Device->GetTimestampQueryHeap(CommandQueueType)->EndQueryBatchAndResolveQueryData(*ContextAtIndex);
 #if PLATFORM_USE_BACKBUFFER_WRITE_TRANSITION_TRACKING
 			uint64 TimeStampFrequency = 0;
-			VERIFYD3D12RESULT(Device->GetCommandListManager().GetTimestampFrequency(&TimeStampFrequency));
+			VERIFYD3D12RESULT(Device->GetCommandListManager(CommandQueueType)->GetTimestampFrequency(&TimeStampFrequency));
 			Device->GetBackBufferWriteBarrierTracker()->ResolveBatches(TimeStampFrequency, false);
 #endif // #if PLATFORM_USE_BACKBUFFER_WRITE_TRANSITION_TRACKING
 		}
@@ -510,7 +510,7 @@ void FD3D12CommandContext::ClearState()
 		}
 	}
 
-	if (!bIsAsyncComputeContext)
+	if (!IsAsyncComputeContext())
 	{
 		FMemory::Memzero(CurrentRenderTargets, sizeof(CurrentRenderTargets));
 		NumSimultaneousRenderTargets = 0;
@@ -638,7 +638,7 @@ void FD3D12CommandContextBase::SignalTransitionFences(TArrayView<const FRHITrans
 				RHISubmitCommandsHint();
 				bSubmitted = true;
 			}
-			Fence->Signal(bIsAsyncComputeContext ? ED3D12CommandQueueType::Async : ED3D12CommandQueueType::Default);
+			Fence->Signal(IsAsyncComputeContext() ? ED3D12CommandQueueType::Async : ED3D12CommandQueueType::Direct);
 		}
 	}
 }
@@ -657,7 +657,7 @@ void FD3D12CommandContextBase::WaitForTransitionFences(TArrayView<const FRHITran
 				RHISubmitCommandsHint();
 				bSubmitted = true;
 			}
-			Fence->GpuWait(bIsAsyncComputeContext ? ED3D12CommandQueueType::Async : ED3D12CommandQueueType::Default, Fence->GetLastSignaledFence());
+			Fence->GpuWait(IsAsyncComputeContext() ? ED3D12CommandQueueType::Async : ED3D12CommandQueueType::Direct, Fence->GetLastSignaledFence());
 		}
 	}
 }
@@ -763,7 +763,7 @@ public:
 		}
 		else
 		{
-			CmdContextRedirector = new FD3D12CommandContextRedirector(Adapter, false, false);
+			CmdContextRedirector = new FD3D12CommandContextRedirector(Adapter, ED3D12CommandQueueType::Direct, false);
 			CmdContextRedirector->SetPhysicalGPUMask(GPUMask);
 
 			for (uint32 GPUIndex : GPUMask)
@@ -930,8 +930,8 @@ void FD3D12DynamicRHI::RHIReleaseTransition(FRHITransition* Transition)
 //
 //////////////////////////////////////////////////////////////////////////////////////////////////////////
 
-FD3D12CommandContextRedirector::FD3D12CommandContextRedirector(class FD3D12Adapter* InParent, bool InIsDefaultContext, bool InIsAsyncComputeContext)
-	: FD3D12CommandContextBase(InParent, FRHIGPUMask::All(), InIsDefaultContext, InIsAsyncComputeContext)
+FD3D12CommandContextRedirector::FD3D12CommandContextRedirector(class FD3D12Adapter* InParent, ED3D12CommandQueueType InCommandQueueType, bool InIsDefaultContext)
+	: FD3D12CommandContextBase(InParent, FRHIGPUMask::All(), InCommandQueueType, InIsDefaultContext)
 {
 	FMemory::Memzero(PhysicalContexts, sizeof(PhysicalContexts[0]) * MAX_NUM_GPUS);
 }
diff --git a/Engine/Source/Runtime/D3D12RHI/Private/D3D12CommandContext.h b/Engine/Source/Runtime/D3D12RHI/Private/D3D12CommandContext.h
--- a/Engine/Source/Runtime/D3D12RHI/Private/D3D12CommandContext.h
+++ b/Engine/Source/Runtime/D3D12RHI/Private/D3D12CommandContext.h
@@ -33,7 +33,7 @@ class FD3D12CommandContextBase : public IRHICommandContext, public FD3D12Adapter
 {
 public:
 
-	FD3D12CommandContextBase(class FD3D12Adapter* InParent, FRHIGPUMask InGPUMask, bool InIsDefaultContext, bool InIsAsyncComputeContext);
+	FD3D12CommandContextBase(class FD3D12Adapter* InParent, FRHIGPUMask InGPUMask, ED3D12CommandQueueType InCommandQueueType, bool InIsDefaultContext);
 
 	void RHIBeginDrawingViewport(FRHIViewport* Viewport, FRHITexture* RenderTargetRHI) final override;
 	void RHIEndDrawingViewport(FRHIViewport* Viewport, bool bPresent, bool bLockToVsync) final override;
@@ -47,7 +47,9 @@ public:
 
 	FRHIGPUMask GetGPUMask() const { return GPUMask; }
 
+	ED3D12CommandQueueType GetCommandQueueType() const { return CommandQueueType; }
 	bool IsDefaultContext() const { return bIsDefaultContext; }
+	bool IsAsyncComputeContext() const { return (CommandQueueType == ED3D12CommandQueueType::Async); }
 
 	virtual void RHISetAsyncComputeBudget(EAsyncComputeBudget Budget) {}
 
@@ -60,8 +62,8 @@ protected:
 	FRHIGPUMask GPUMask;
 
 	bool bTrackingEvents;
+	const ED3D12CommandQueueType CommandQueueType;
 	const bool bIsDefaultContext;
-	const bool bIsAsyncComputeContext;
 };
 
 class FD3D12CommandContext : public FD3D12CommandContextBase, public FD3D12DeviceChild
@@ -75,7 +77,7 @@ public:
 		FCEA_Num
 	};
 
-	FD3D12CommandContext(class FD3D12Device* InParent, bool InIsDefaultContext, bool InIsAsyncComputeContext);
+	FD3D12CommandContext(class FD3D12Device* InParent, ED3D12CommandQueueType InCommandQueueType, bool InIsDefaultContext);
 	virtual ~FD3D12CommandContext();
 
 	FD3D12CommandListManager& GetCommandListManager();
@@ -463,7 +465,7 @@ private:
 class FD3D12CommandContextRedirector final : public FD3D12CommandContextBase
 {
 public:
-	FD3D12CommandContextRedirector(class FD3D12Adapter* InParent, bool InIsDefaultContext, bool InIsAsyncComputeContext);
+	FD3D12CommandContextRedirector(class FD3D12Adapter* InParent, ED3D12CommandQueueType InCommandQueueType, bool InIsDefaultContext);
 
 #define ContextRedirect(Call) { for (uint32 GPUIndex : GPUMask) PhysicalContexts[GPUIndex]->##Call; }
 #define ContextGPU0(Call) { PhysicalContexts[0]->##Call; }
diff --git a/Engine/Source/Runtime/D3D12RHI/Private/D3D12Commands.cpp b/Engine/Source/Runtime/D3D12RHI/Private/D3D12Commands.cpp
--- a/Engine/Source/Runtime/D3D12RHI/Private/D3D12Commands.cpp
+++ b/Engine/Source/Runtime/D3D12RHI/Private/D3D12Commands.cpp
@@ -181,7 +181,7 @@ void FD3D12CommandContext::RHIDispatchIndirectComputeShader(FRHIBuffer* Argument
 	CommandListHandle.FlushResourceBarriers();	// Must flush so the desired state is actually set.
 
 	FD3D12Adapter* Adapter = GetParentDevice()->GetParentAdapter();
-	ID3D12CommandSignature* CommandSignature = bIsAsyncComputeContext ? Adapter->GetDispatchIndirectComputeCommandSignature() : Adapter->GetDispatchIndirectGraphicsCommandSignature();
+	ID3D12CommandSignature* CommandSignature = IsAsyncComputeContext() ? Adapter->GetDispatchIndirectComputeCommandSignature() : Adapter->GetDispatchIndirectGraphicsCommandSignature();
 	
 	numDispatches++;
 	CommandListHandle->ExecuteIndirect(
@@ -388,7 +388,7 @@ void FD3D12CommandContext::RHIBeginTransitions(TArrayView<const FRHITransition*>
 {
 	// Can't correctly handle RHITransitions from multiple threads at the samae time right now - only internal transitions work correctly then
 	// (will be fixed if unknown before state is not allowed anymore)
-	check(bIsDefaultContext || bIsAsyncComputeContext || GUseInternalTransitions);
+	check(bIsDefaultContext || GUseInternalTransitions);
 
 	RHIBeginTransitionsWithoutFencing(Transitions);
 	SignalTransitionFences(Transitions);
@@ -490,7 +490,7 @@ void FD3D12CommandContext::RHIWriteGPUFence(FRHIGPUFence* FenceRHI)
 	// @todo-mattc we don't want to flush here. That should be the caller's responsibility.
 	RHISubmitCommandsHint();
 	FD3D12GPUFence* Fence = FD3D12DynamicRHI::ResourceCast(FenceRHI);
-	Fence->WriteInternal(ED3D12CommandQueueType::Default);
+	Fence->WriteInternal(ED3D12CommandQueueType::Direct);
 }
 
 void FD3D12CommandContext::RHISetViewport(float MinX, float MinY, float MinZ, float MaxX, float MaxY, float MaxZ)
@@ -1202,7 +1202,7 @@ void FD3D12CommandContext::RHIBeginRenderQuery(FRHIRenderQuery* QueryRHI)
 	check(IsDefaultContext());
 	check(Query->Type == RQT_Occlusion);
 
-	GetParentDevice()->GetOcclusionQueryHeap()->BeginQuery(*this, Query);
+	GetParentDevice()->GetOcclusionQueryHeap(CommandQueueType)->BeginQuery(*this, Query);
 
 #if EXECUTE_DEBUG_COMMAND_LISTS
 	GIsDoingQuery = true;
@@ -1219,11 +1219,11 @@ void FD3D12CommandContext::RHIEndRenderQuery(FRHIRenderQuery* QueryRHI)
 	switch (Query->Type)
 	{
 	case RQT_Occlusion:
-		QueryHeap = GetParentDevice()->GetOcclusionQueryHeap();
+		QueryHeap = GetParentDevice()->GetOcclusionQueryHeap(CommandQueueType);
 		break;
 
 	case RQT_AbsoluteTime:
-		QueryHeap = GetParentDevice()->GetTimestampQueryHeap();
+		QueryHeap = GetParentDevice()->GetTimestampQueryHeap(CommandQueueType);
 		break;
 
 	default:
@@ -2018,7 +2018,7 @@ void FD3D12CommandContext::RHISetShadingRateImage(FRHITexture* RateImageTexture,
 void FD3D12CommandContext::RHISubmitCommandsHint()
 {
 	// Resolve any timestamp queries so far (if any).
-	GetParentDevice()->GetTimestampQueryHeap()->EndQueryBatchAndResolveQueryData(*this);
+	GetParentDevice()->GetTimestampQueryHeap(CommandQueueType)->EndQueryBatchAndResolveQueryData(*this);
 
 	// Submit the work we have so far, and start a new command list.
 	FlushCommands();
@@ -2049,7 +2049,7 @@ void FD3D12CommandContext::RHIWaitForTemporalEffect(const FName& InEffectName)
 		// Execute the current command list so we can have a point to insert a wait
 		FlushCommands();
 
-		Effect->WaitForPrevious(GPUIndex, bIsAsyncComputeContext ? ED3D12CommandQueueType::Async : ED3D12CommandQueueType::Default);
+		Effect->WaitForPrevious(GPUIndex, CommandQueueType);
 	}
 #endif // WITH_MGPU
 }
@@ -2092,7 +2092,7 @@ void FD3D12CommandContext::RHIBroadcastTemporalEffect(const FName& InEffectName,
 	FlushCommands();
 
 	// Tell the copy queue to wait for the current queue to finish rendering before starting the copy.
-	Effect->SignalSyncComplete(GPUIndex, bIsAsyncComputeContext ? ED3D12CommandQueueType::Async : ED3D12CommandQueueType::Default);
+	Effect->SignalSyncComplete(GPUIndex, CommandQueueType);
 	Effect->WaitForPrevious(GPUIndex, ED3D12CommandQueueType::Copy);
 
 	FD3D12CommandAllocatorManager& CopyCommandAllocatorManager = Device->GetTextureStreamingCommandAllocatorManager();
@@ -2131,7 +2131,7 @@ void FD3D12CommandContext::RHIBroadcastTemporalEffect(const FName& InEffectName,
 
 	FlushCommands();
 
-	Effect->SignalSyncComplete(GPUIndex, bIsAsyncComputeContext ? ED3D12CommandQueueType::Async : ED3D12CommandQueueType::Default);
+	Effect->SignalSyncComplete(GPUIndex, CommandQueueType);
 
 #endif // USE_COPY_QUEUE_FOR_RESOURCE_SYNC
 #endif // WITH_MGPU
diff --git a/Engine/Source/Runtime/D3D12RHI/Private/D3D12Device.cpp b/Engine/Source/Runtime/D3D12RHI/Private/D3D12Device.cpp
--- a/Engine/Source/Runtime/D3D12RHI/Private/D3D12Device.cpp
+++ b/Engine/Source/Runtime/D3D12RHI/Private/D3D12Device.cpp
@@ -34,7 +34,6 @@ FD3D12Device::FD3D12Device(FRHIGPUMask InGPUMask, FD3D12Adapter* InAdapter) :
 	GlobalSamplerHeap(this, InGPUMask),
 	GlobalViewHeap(this, InGPUMask),
 	OcclusionQueryHeap(this, D3D12_QUERY_TYPE_OCCLUSION, 65536, 4 /*frames to keep results */ * 1 /*batches per frame*/),
-	TimestampQueryHeap(this, D3D12_QUERY_TYPE_TIMESTAMP, 8192, 4 /*frames to keep results */ * 5 /*batches per frame*/ ),
 #if WITH_PROFILEGPU || D3D12_SUBMISSION_GAP_RECORDER
 	CmdListExecTimeQueryHeap(this, D3D12_QUERY_HEAP_TYPE_TIMESTAMP, 8192),
 #endif
@@ -44,6 +43,11 @@ FD3D12Device::FD3D12Device(FRHIGPUMask InGPUMask, FD3D12Adapter* InAdapter) :
 	TextureAllocator(this, FRHIGPUMask::All()),
 	GPUProfilingData(this)
 {
+	for (uint32 QueueType = 0; QueueType < (uint32)ED3D12CommandQueueType::Count; ++QueueType)
+	{
+		TimestampQueryHeaps[QueueType] = new FD3D12QueryHeap(this, D3D12_QUERY_TYPE_TIMESTAMP, 8192, 4 /*frames to keep results */ * 5 /*batches per frame*/);
+	}
+
 	InitPlatformSpecific();
 }
 
@@ -116,8 +120,8 @@ void FD3D12Device::CreateCommandContexts()
 	for (uint32 i = 0; i < NumContexts; ++i)
 	{	
 		const bool bIsDefaultContext = (i == 0);
-		const bool bIsAsyncComputeContext = false;
-		FD3D12CommandContext* NewCmdContext = GetOwningRHI()->CreateCommandContext(this, bIsDefaultContext, bIsAsyncComputeContext);
+		const ED3D12CommandQueueType CommandQueueType = ED3D12CommandQueueType::Direct;
+		FD3D12CommandContext* NewCmdContext = GetOwningRHI()->CreateCommandContext(this, CommandQueueType, bIsDefaultContext);
 
 		// without that the first RHIClear would get a scissor rect of (0,0)-(0,0) which means we get a draw call clear 
 		NewCmdContext->RHISetScissorRect(false, 0, 0, 0, 0);
@@ -134,8 +138,8 @@ void FD3D12Device::CreateCommandContexts()
 	for (uint32 i = 0; i < NumAsyncComputeContexts; ++i)
 	{		
 		const bool bIsDefaultContext = (i == 0); //-V547
-		const bool bIsAsyncComputeContext = true;
-		FD3D12CommandContext* NewCmdContext = GetOwningRHI()->CreateCommandContext(this, bIsDefaultContext, bIsAsyncComputeContext);
+		const ED3D12CommandQueueType CommandQueueType = ED3D12CommandQueueType::Async;
+		FD3D12CommandContext* NewCmdContext = GetOwningRHI()->CreateCommandContext(this, CommandQueueType, bIsDefaultContext);
 	
 		AsyncComputeContextArray.Add(NewCmdContext);
 	}
@@ -287,7 +291,10 @@ void FD3D12Device::SetupAfterDeviceCreation()
 
 	// Init the occlusion and timestamp query heaps
 	OcclusionQueryHeap.Init();
-	TimestampQueryHeap.Init();
+	for (uint32 QueueType = 0; QueueType < (uint32)ED3D12CommandQueueType::Count; ++QueueType)
+	{
+		TimestampQueryHeaps[QueueType]->Init();
+	}
 
 	CommandListManager->Create(*FString::Printf(TEXT("3D Queue %d"), GetGPUIndex()));
 	gD3D12CommandQueue = CommandListManager->GetD3DCommandQueue();
@@ -360,7 +367,7 @@ void FD3D12Device::Cleanup()
 	};
 
 	// Validate that all the D3D command queues are still valid (temp code to check for a shutdown crash)
-	ValidateCommandQueue(ED3D12CommandQueueType::Default, TEXT("Direct"));
+	ValidateCommandQueue(ED3D12CommandQueueType::Direct, TEXT("Direct"));
 	ValidateCommandQueue(ED3D12CommandQueueType::Copy, TEXT("Copy"));
 	ValidateCommandQueue(ED3D12CommandQueueType::Async, TEXT("Async"));
 
@@ -408,7 +415,12 @@ void FD3D12Device::Cleanup()
 	AsyncCommandListManager->Destroy();
 
 	OcclusionQueryHeap.Destroy();
-	TimestampQueryHeap.Destroy();
+	for (uint32 QueueType = 0; QueueType < (uint32)ED3D12CommandQueueType::Count; ++QueueType)
+	{
+		TimestampQueryHeaps[QueueType]->Destroy();
+		delete TimestampQueryHeaps[QueueType];
+		TimestampQueryHeaps[QueueType] = nullptr;
+	}
 
 	D3DX12Residency::DestroyResidencyManager(ResidencyManager);
 
@@ -420,7 +432,7 @@ FD3D12CommandListManager* FD3D12Device::GetCommandListManager(ED3D12CommandQueue
 {
 	switch (InQueueType)
 	{
-	case ED3D12CommandQueueType::Default:
+	case ED3D12CommandQueueType::Direct:
 		check(CommandListManager->GetQueueType() == InQueueType);
 		return CommandListManager;
 	case ED3D12CommandQueueType::Async:
diff --git a/Engine/Source/Runtime/D3D12RHI/Private/D3D12Device.h b/Engine/Source/Runtime/D3D12RHI/Private/D3D12Device.h
--- a/Engine/Source/Runtime/D3D12RHI/Private/D3D12Device.h
+++ b/Engine/Source/Runtime/D3D12RHI/Private/D3D12Device.h
@@ -55,8 +55,8 @@ public:
 
 	FD3D12DynamicRHI* GetOwningRHI();
 
-	inline FD3D12QueryHeap* GetOcclusionQueryHeap() { return &OcclusionQueryHeap; }
-	inline FD3D12QueryHeap* GetTimestampQueryHeap() { return &TimestampQueryHeap; }
+	inline FD3D12QueryHeap* GetOcclusionQueryHeap(ED3D12CommandQueueType inQueueType) { check(inQueueType == ED3D12CommandQueueType::Direct); return &OcclusionQueryHeap; }
+	inline FD3D12QueryHeap* GetTimestampQueryHeap(ED3D12CommandQueueType inQueueType) { return TimestampQueryHeaps[(uint32)inQueueType]; }
 	FD3D12LinearQueryHeap* GetCmdListExecTimeQueryHeap();
 
 	template <typename TViewDesc> FD3D12OfflineDescriptorManager& GetViewDescriptorAllocator();
@@ -102,7 +102,7 @@ public:
 	}
 
 	FD3D12CommandListManager* GetCommandListManager(ED3D12CommandQueueType inQueueType) const;
-	ID3D12CommandQueue* GetD3DCommandQueue(ED3D12CommandQueueType InQueueType = ED3D12CommandQueueType::Default) { return GetCommandListManager(InQueueType)->GetD3DCommandQueue(); }
+	ID3D12CommandQueue* GetD3DCommandQueue(ED3D12CommandQueueType InQueueType = ED3D12CommandQueueType::Direct) { return GetCommandListManager(InQueueType)->GetD3DCommandQueue(); }
 
 	inline FD3D12CommandContext& GetDefaultCommandContext() const { return GetCommandContext(0); }
 	inline FD3D12CommandContext& GetDefaultAsyncComputeContext() const { return GetAsyncComputeContext(0); }
@@ -149,7 +149,7 @@ protected:
 	FD3D12GlobalHeap GlobalViewHeap;
 
 	FD3D12QueryHeap OcclusionQueryHeap;
-	FD3D12QueryHeap TimestampQueryHeap;
+	FD3D12QueryHeap* TimestampQueryHeaps[(uint32)ED3D12CommandQueueType::Count];
 #if WITH_PROFILEGPU || D3D12_SUBMISSION_GAP_RECORDER
 	FD3D12LinearQueryHeap CmdListExecTimeQueryHeap;
 #endif
diff --git a/Engine/Source/Runtime/D3D12RHI/Private/D3D12DirectCommandListManager.cpp b/Engine/Source/Runtime/D3D12RHI/Private/D3D12DirectCommandListManager.cpp
--- a/Engine/Source/Runtime/D3D12RHI/Private/D3D12DirectCommandListManager.cpp
+++ b/Engine/Source/Runtime/D3D12RHI/Private/D3D12DirectCommandListManager.cpp
@@ -783,7 +783,7 @@ void FD3D12CommandListManager::ReleaseResourceBarrierCommandListAllocator()
 void FD3D12CommandListManager::StartTrackingCommandListTime()
 {
 #if WITH_PROFILEGPU || D3D12_SUBMISSION_GAP_RECORDER
-	check(QueueType == ED3D12CommandQueueType::Default && !GetShouldTrackCmdListTime());
+	check(QueueType == ED3D12CommandQueueType::Direct && !GetShouldTrackCmdListTime());
 	ResolvedTimingPairs.Reset();
 	SetShouldTrackCmdListTime(true);
 #endif
@@ -792,7 +792,7 @@ void FD3D12CommandListManager::StartTrackingCommandListTime()
 void FD3D12CommandListManager::EndTrackingCommandListTime()
 {
 #if WITH_PROFILEGPU || D3D12_SUBMISSION_GAP_RECORDER
-	check(QueueType == ED3D12CommandQueueType::Default && GetShouldTrackCmdListTime());
+	check(QueueType == ED3D12CommandQueueType::Direct && GetShouldTrackCmdListTime());
 	SetShouldTrackCmdListTime(false);
 #endif
 }
@@ -800,7 +800,7 @@ void FD3D12CommandListManager::EndTrackingCommandListTime()
 void FD3D12CommandListManager::GetCommandListTimingResults(TArray<FResolvedCmdListExecTime>& OutTimingPairs, bool bUseBlockingCall)
 {
 #if WITH_PROFILEGPU || D3D12_SUBMISSION_GAP_RECORDER
-	check(!GetShouldTrackCmdListTime() && QueueType == ED3D12CommandQueueType::Default);
+	check(!GetShouldTrackCmdListTime() && QueueType == ED3D12CommandQueueType::Direct);
 	FlushPendingTimingPairs(bUseBlockingCall);
 	if (bUseBlockingCall)
 	{
diff --git a/Engine/Source/Runtime/D3D12RHI/Private/D3D12Query.cpp b/Engine/Source/Runtime/D3D12RHI/Private/D3D12Query.cpp
--- a/Engine/Source/Runtime/D3D12RHI/Private/D3D12Query.cpp
+++ b/Engine/Source/Runtime/D3D12RHI/Private/D3D12Query.cpp
@@ -198,7 +198,7 @@ bool FD3D12Device::GetQueryData(FD3D12RenderQuery& Query, bool bWait)
 
 void FD3D12CommandContext::RHIBeginOcclusionQueryBatch(uint32 NumQueriesInBatch)
 {
-	GetParentDevice()->GetOcclusionQueryHeap()->StartQueryBatch(*this, NumQueriesInBatch);
+	GetParentDevice()->GetOcclusionQueryHeap(CommandQueueType)->StartQueryBatch(*this, NumQueriesInBatch);
 	if (RHIConsoleVariables::GInsertOuterOcclusionQuery)
 	{
 		if (!OuterOcclusionQuery.IsValid())
@@ -207,7 +207,7 @@ void FD3D12CommandContext::RHIBeginOcclusionQueryBatch(uint32 NumQueriesInBatch)
 		}
 		
 		FD3D12RenderQuery* OuterOcclusionQueryD3D12 = RetrieveObject<FD3D12RenderQuery>(OuterOcclusionQuery.GetReference());
-		GetParentDevice()->GetOcclusionQueryHeap()->BeginQuery(*this, OuterOcclusionQueryD3D12);
+		GetParentDevice()->GetOcclusionQueryHeap(CommandQueueType)->BeginQuery(*this, OuterOcclusionQueryD3D12);
 		bOuterOcclusionQuerySubmitted = true;
 	}
 }
@@ -219,10 +219,10 @@ void FD3D12CommandContext::RHIEndOcclusionQueryBatch()
 		check(OuterOcclusionQuery.IsValid());
 		FD3D12RenderQuery* OuterOcclusionQueryD3D12 = RetrieveObject<FD3D12RenderQuery>(OuterOcclusionQuery.GetReference());
 		check(OuterOcclusionQueryD3D12->HeapIndex != INDEX_NONE);
-		GetParentDevice()->GetOcclusionQueryHeap()->EndQuery(*this, OuterOcclusionQueryD3D12);
+		GetParentDevice()->GetOcclusionQueryHeap(CommandQueueType)->EndQuery(*this, OuterOcclusionQueryD3D12);
 		bOuterOcclusionQuerySubmitted = false;
 	}
-	GetParentDevice()->GetOcclusionQueryHeap()->EndQueryBatchAndResolveQueryData(*this);
+	GetParentDevice()->GetOcclusionQueryHeap(CommandQueueType)->EndQueryBatchAndResolveQueryData(*this);
 
 	// Note: We want to execute this ASAP. The Engine will call RHISubmitCommandHint after this.
 	// We'll break up the command list there so that the wait on the previous frame's results don't block.
@@ -313,6 +313,8 @@ uint32 FD3D12QueryHeap::AllocQuery(FD3D12CommandContext& CmdContext)
 			check(CurrentQueryBatch.bOpen && CurrentQueryBatch.ElementCount == 0);
 		}
 
+		// No support for more than 256 timestamp queries?
+		check(CurrentQueryBatch.ElementCount < 256);
 		if (CurrentQueryBatch.StartElement > CurrentElement)
 		{
 			// We're in the middle of a batch, but we're at the end of the heap
@@ -487,7 +489,7 @@ void FD3D12QueryHeap::CreateQueryHeap()
 
 	// Create the upload heap
 	VERIFYD3D12RESULT(GetParentDevice()->GetDevice()->CreateQueryHeap(&QueryHeapDesc, IID_PPV_ARGS(ActiveQueryHeap.GetInitReference())));
-	SetName(ActiveQueryHeap, L"Query Heap");
+	SetName(ActiveQueryHeap, (QueryType == D3D12_QUERY_TYPE_OCCLUSION) ? TEXT("Occlusion Query Heap") : TEXT("Timestamp Query Heap"));
 
 #if ENABLE_RESIDENCY_MANAGEMENT
 	D3DX12Residency::Initialize(ActiveQueryHeapResidencyHandle, ActiveQueryHeap, ResultSize * QueryHeapDesc.Count);
@@ -507,7 +509,7 @@ void FD3D12QueryHeap::CreateQueryHeap()
 		D3D12_RESOURCE_STATE_COPY_DEST,
 		nullptr,
 		ActiveResultBuffer.GetInitReference(),
-		TEXT("Query Heap Result Buffer")));
+		(QueryType == D3D12_QUERY_TYPE_OCCLUSION) ? TEXT(" Occlusion Query Heap Result Buffer") : TEXT("Timestamp Query Heap Result Buffer")));
 }
 
 void FD3D12QueryHeap::DestroyQueryHeap()
diff --git a/Engine/Source/Runtime/D3D12RHI/Private/D3D12RHICommon.h b/Engine/Source/Runtime/D3D12RHI/Private/D3D12RHICommon.h
--- a/Engine/Source/Runtime/D3D12RHI/Private/D3D12RHICommon.h
+++ b/Engine/Source/Runtime/D3D12RHI/Private/D3D12RHICommon.h
@@ -20,9 +20,11 @@ class FD3D12Device;
 // Defines a unique command queue type within a FD3D12Device (owner by the command list managers).
 enum class ED3D12CommandQueueType
 {
-	Default,
+	Direct = 0,
 	Copy,
-	Async
+	Async,
+
+	Count,
 };
 
 class FD3D12AdapterChild
diff --git a/Engine/Source/Runtime/D3D12RHI/Private/D3D12RHIPrivate.h b/Engine/Source/Runtime/D3D12RHI/Private/D3D12RHIPrivate.h
--- a/Engine/Source/Runtime/D3D12RHI/Private/D3D12RHIPrivate.h
+++ b/Engine/Source/Runtime/D3D12RHI/Private/D3D12RHIPrivate.h
@@ -247,7 +247,7 @@ public:
 		return Object ? static_cast<ReturnType*>(Object->GetLinkedObject(GPUIndex)) : nullptr;
 	}
 
-	virtual FD3D12CommandContext* CreateCommandContext(FD3D12Device* InParent, bool InIsDefaultContext, bool InIsAsyncComputeContext);
+	virtual FD3D12CommandContext* CreateCommandContext(FD3D12Device* InParent, ED3D12CommandQueueType InQueueType, bool InIsDefaultContext);
 	virtual void CreateCommandQueue(FD3D12Device* Device, const D3D12_COMMAND_QUEUE_DESC& Desc, TRefCountPtr<ID3D12CommandQueue>& OutCommandQueue);
 
 	virtual bool GetHardwareGPUFrameTime(double& OutGPUFrameTime) const
diff --git a/Engine/Source/Runtime/D3D12RHI/Private/D3D12RenderTarget.cpp b/Engine/Source/Runtime/D3D12RHI/Private/D3D12RenderTarget.cpp
--- a/Engine/Source/Runtime/D3D12RHI/Private/D3D12RenderTarget.cpp
+++ b/Engine/Source/Runtime/D3D12RHI/Private/D3D12RenderTarget.cpp
@@ -514,7 +514,7 @@ void FD3D12DynamicRHI::RHIMultiGPULockstep(FRHIGPUMask GPUMask)
 		{
 			if (GPUIndex != GPUIndex2)
 			{
-				Fence.GpuWait(GPUIndex2, ED3D12CommandQueueType::Default, Fence.GetLastSignaledFence(), GPUIndex);
+				Fence.GpuWait(GPUIndex2, ED3D12CommandQueueType::Direct, Fence.GetLastSignaledFence(), GPUIndex);
 			}
 		}
 	}
diff --git a/Engine/Source/Runtime/D3D12RHI/Private/D3D12Util.cpp b/Engine/Source/Runtime/D3D12RHI/Private/D3D12Util.cpp
--- a/Engine/Source/Runtime/D3D12RHI/Private/D3D12Util.cpp
+++ b/Engine/Source/Runtime/D3D12RHI/Private/D3D12Util.cpp
@@ -135,7 +135,7 @@ static const TCHAR* GetD3DCommandQueueTypeName(ED3D12CommandQueueType QueueType)
 {
 	switch (QueueType)
 	{
-	case ED3D12CommandQueueType::Default:	 return TEXT("3D");
+	case ED3D12CommandQueueType::Direct:	 return TEXT("3D");
 	case ED3D12CommandQueueType::Async:		 return TEXT("Compute");
 	case ED3D12CommandQueueType::Copy:		 return TEXT("Copy");
 	}
diff --git a/Engine/Source/Runtime/D3D12RHI/Private/D3D12Viewport.cpp b/Engine/Source/Runtime/D3D12RHI/Private/D3D12Viewport.cpp
--- a/Engine/Source/Runtime/D3D12RHI/Private/D3D12Viewport.cpp
+++ b/Engine/Source/Runtime/D3D12RHI/Private/D3D12Viewport.cpp
@@ -844,7 +844,7 @@ bool FD3D12Viewport::Present(bool bLockToVsync)
 		// system present queue out of order.
 		const uint32 PresentGPUIndex = BackBufferGPUIndices[CurrentBackBufferIndex_RHIThread];
 		const uint32 LastGPUIndex = BackBufferGPUIndices[(CurrentBackBufferIndex_RHIThread + NumBackBuffers - 1) % NumBackBuffers];
-		Fence.GpuWait(PresentGPUIndex, ED3D12CommandQueueType::Default, LastSignaledValue, LastGPUIndex);
+		Fence.GpuWait(PresentGPUIndex, ED3D12CommandQueueType::Direct, LastSignaledValue, LastGPUIndex);
 	}
 
 #if 0 // Multi-GPU support : figure out what kind of synchronization is needed.
@@ -899,7 +899,7 @@ void FD3D12Viewport::WaitForFrameEventCompletion()
 void FD3D12Viewport::IssueFrameEvent()
 {
 	// Signal the fence.
-	LastSignaledValue = Fence.Signal(ED3D12CommandQueueType::Default);
+	LastSignaledValue = Fence.Signal(ED3D12CommandQueueType::Direct);
 }
 
 bool FD3D12Viewport::CheckHDRSupport()
@@ -1092,14 +1092,14 @@ void FD3D12DynamicRHI::RHIAdvanceFrameFence()
 	if (RHICmdList.Bypass())
 	{
 		// In bypass mode, we should execute this directly
-		FRHICommandSignalFrameFence Cmd(ED3D12CommandQueueType::Default, FrameFence, PreviousFence);
+		FRHICommandSignalFrameFence Cmd(ED3D12CommandQueueType::Direct, FrameFence, PreviousFence);
 		Cmd.Execute(RHICmdList);
 	}
 	else
 	{
 		// Queue a command to signal on RHI thread that the current frame is a complete on the GPU.
 		// This must be done in a deferred way even if RHI thread is disabled, just for correct ordering of operations.
-		ALLOC_COMMAND_CL(RHICmdList, FRHICommandSignalFrameFence)(ED3D12CommandQueueType::Default, FrameFence, PreviousFence);
+		ALLOC_COMMAND_CL(RHICmdList, FRHICommandSignalFrameFence)(ED3D12CommandQueueType::Direct, FrameFence, PreviousFence);
 	}
 #if D3D12_SUBMISSION_GAP_RECORDER
 	FD3D12Adapter* Adapter = &GetAdapter();
diff --git a/Engine/Source/Runtime/D3D12RHI/Private/Windows/WindowsD3D12Device.cpp b/Engine/Source/Runtime/D3D12RHI/Private/Windows/WindowsD3D12Device.cpp
--- a/Engine/Source/Runtime/D3D12RHI/Private/Windows/WindowsD3D12Device.cpp
+++ b/Engine/Source/Runtime/D3D12RHI/Private/Windows/WindowsD3D12Device.cpp
@@ -968,7 +968,7 @@ void FD3D12Device::Initialize()
 
 void FD3D12Device::InitPlatformSpecific()
 {
-	CommandListManager = new FD3D12CommandListManager(this, D3D12_COMMAND_LIST_TYPE_DIRECT, ED3D12CommandQueueType::Default);
+	CommandListManager = new FD3D12CommandListManager(this, D3D12_COMMAND_LIST_TYPE_DIRECT, ED3D12CommandQueueType::Direct);
 	CopyCommandListManager = new FD3D12CommandListManager(this, D3D12_COMMAND_LIST_TYPE_COPY, ED3D12CommandQueueType::Copy);
 	AsyncCommandListManager = new FD3D12CommandListManager(this, D3D12_COMMAND_LIST_TYPE_COMPUTE, ED3D12CommandQueueType::Async);
 }
