Check in missing file

diff --git a/Engine/Source/Runtime/D3D12RHI/Private/D3D12RHI.cpp b/Engine/Source/Runtime/D3D12RHI/Private/D3D12RHI.cpp
--- a/Engine/Source/Runtime/D3D12RHI/Private/D3D12RHI.cpp
+++ b/Engine/Source/Runtime/D3D12RHI/Private/D3D12RHI.cpp
@@ -269,9 +269,9 @@ void FD3D12DynamicRHI::Shutdown()
 	ZeroBufferSize = 0;
 }
 
-FD3D12CommandContext* FD3D12DynamicRHI::CreateCommandContext(FD3D12Device* InParent, bool InIsDefaultContext, bool InIsAsyncComputeContext)
+FD3D12CommandContext* FD3D12DynamicRHI::CreateCommandContext(FD3D12Device* InParent, ED3D12CommandQueueType InQueueType, bool InIsDefaultContext)
 {
-	FD3D12CommandContext* NewContext = new FD3D12CommandContext(InParent, InIsDefaultContext, InIsAsyncComputeContext);
+	FD3D12CommandContext* NewContext = new FD3D12CommandContext(InParent, InQueueType, InIsDefaultContext);
 	return NewContext;
 }
 
