Fix static analysis bug fix

diff --git a/Engine/Source/Runtime/D3D12RHI/Private/D3D12Commands.cpp b/Engine/Source/Runtime/D3D12RHI/Private/D3D12Commands.cpp
--- a/Engine/Source/Runtime/D3D12RHI/Private/D3D12Commands.cpp
+++ b/Engine/Source/Runtime/D3D12RHI/Private/D3D12Commands.cpp
@@ -338,13 +338,20 @@ static void HandleResourceTransitions(FD3D12CommandContext& Context, const FD3D1
 						// enqueue the correct transitions
 						if (Info.IsWholeResource() || Resource->GetSubresourceCount() == 1)
 						{
-							bUAVBarrier = bUAVBarrier | FD3D12DynamicRHI::TransitionResource(Context.CommandListHandle, Resource, BeforeState, AfterState, D3D12_RESOURCE_BARRIER_ALL_SUBRESOURCES, FD3D12DynamicRHI::ETransitionMode::Apply);
+							if (FD3D12DynamicRHI::TransitionResource(Context.CommandListHandle, Resource, BeforeState, AfterState, D3D12_RESOURCE_BARRIER_ALL_SUBRESOURCES, FD3D12DynamicRHI::ETransitionMode::Apply))
+							{
+								bUAVBarrier = true;
+							}
+
 						}
 						else
 						{
 							EnumerateSubresources(Resource, Info, [&](uint32 Subresource)
 								{
-									bUAVBarrier = bUAVBarrier | FD3D12DynamicRHI::TransitionResource(Context.CommandListHandle, Resource, BeforeState, AfterState, Subresource, FD3D12DynamicRHI::ETransitionMode::Apply);
+									if (FD3D12DynamicRHI::TransitionResource(Context.CommandListHandle, Resource, BeforeState, AfterState, Subresource, FD3D12DynamicRHI::ETransitionMode::Apply))
+									{
+										bUAVBarrier = true;
+									}
 								});
 						}
 					});
