Stall and flush rhi thread when sync point is still open during RHIMapStagingSurface

diff --git a/Engine/Source/Runtime/D3D12RHI/Private/D3D12Query.cpp b/Engine/Source/Runtime/D3D12RHI/Private/D3D12Query.cpp
--- a/Engine/Source/Runtime/D3D12RHI/Private/D3D12Query.cpp
+++ b/Engine/Source/Runtime/D3D12RHI/Private/D3D12Query.cpp
@@ -166,15 +166,7 @@ bool FD3D12Device::GetQueryData(FD3D12RenderQuery& Query, bool bWait)
 		{
 			// We should really try to avoid this!
 			UE_LOG(LogD3D12RHI, Verbose, TEXT("Stalling the RHI thread and flushing GPU commands to wait for a RenderQuery that hasn't been submitted to the GPU yet."));
-
-			// The query is on a command list that hasn't been submitted yet.
-			// We need to flush, but the RHI thread may be using the default command list...so stall it first.
-			check(IsInRenderingThread());
-			FScopedRHIThreadStaller StallRHIThread(FRHICommandListExecutor::GetImmediateCommandList());
-			GetDefaultCommandContext().FlushCommands();	// Don't wait yet, since we're stalling the RHI thread.
-
-			// We have to make sure all command lists have actually flush and executed here
-			CommandListManager->WaitOnExecuteTask();
+			StallRHIThreadAndForceFlush(this);
 		}
 
 		SyncPoint.WaitForCompletion();
diff --git a/Engine/Source/Runtime/D3D12RHI/Private/D3D12RenderTarget.cpp b/Engine/Source/Runtime/D3D12RHI/Private/D3D12RenderTarget.cpp
--- a/Engine/Source/Runtime/D3D12RHI/Private/D3D12RenderTarget.cpp
+++ b/Engine/Source/Runtime/D3D12RHI/Private/D3D12RenderTarget.cpp
@@ -1309,6 +1309,7 @@ void FD3D12DynamicRHI::RHIReadSurfaceData(FRHITexture* InRHITexture, FIntRect In
 		CommandListState ListState = GetRHIDevice(GPUIndex)->GetCommandListManager().GetCommandListState(SyncPoint);
 		if (ListState == CommandListState::kOpen)
 		{
+			// This should not happen because RHI thread has been flushed before a call to RHIReadSurfaceData
 			CommandContext.FlushCommands(true);
 		}
 		else
@@ -1525,7 +1526,9 @@ void FD3D12DynamicRHI::RHIMapStagingSurface(FRHITexture* TextureRHI, FRHIGPUFenc
 	CommandListState listState = GetRHIDevice(GPUIndex)->GetCommandListManager().GetCommandListState(SyncPoint);
 	if (listState == CommandListState::kOpen)
 	{
-		GetRHIDevice(GPUIndex)->GetDefaultCommandContext().FlushCommands(true);
+		// We should really try to avoid this!
+		UE_LOG(LogD3D12RHI, Verbose, TEXT("Stalling the RHI thread and flushing GPU commands to wait for a open fence that hasn't been submitted to the GPU yet."));
+		StallRHIThreadAndForceFlush(GetRHIDevice(GPUIndex));
 	}
 	else
 	{
diff --git a/Engine/Source/Runtime/D3D12RHI/Private/D3D12Util.cpp b/Engine/Source/Runtime/D3D12RHI/Private/D3D12Util.cpp
--- a/Engine/Source/Runtime/D3D12RHI/Private/D3D12Util.cpp
+++ b/Engine/Source/Runtime/D3D12RHI/Private/D3D12Util.cpp
@@ -1147,6 +1147,16 @@ D3D12_RESOURCE_STATES GetD3D12ResourceState(ERHIAccess InRHIAccess, bool InIsAsy
 	return D3D12_RESOURCE_STATE_COMMON;
 }
 
+void StallRHIThreadAndForceFlush(FD3D12Device* InDevice)
+{
+	// We need to flush, but the RHI thread may be using the default command list...so stall it first.
+	check(IsInRenderingThread());
+	FScopedRHIThreadStaller StallRHIThread(FRHICommandListExecutor::GetImmediateCommandList());
+	InDevice->GetDefaultCommandContext().FlushCommands();	// Don't wait yet, since we're stalling the RHI thread.
+
+	// We have to make sure all command lists have actually flush and executed here
+	InDevice->GetCommandListManager().WaitOnExecuteTask();
+}
 
 //==================================================================================================================================
 // CResourceState
diff --git a/Engine/Source/Runtime/D3D12RHI/Public/D3D12Util.h b/Engine/Source/Runtime/D3D12RHI/Public/D3D12Util.h
--- a/Engine/Source/Runtime/D3D12RHI/Public/D3D12Util.h
+++ b/Engine/Source/Runtime/D3D12RHI/Public/D3D12Util.h
@@ -270,6 +270,7 @@ void LogExecuteCommandLists(uint32 NumCommandLists, ID3D12CommandList *const *pp
 FString ConvertToResourceStateString(uint32 ResourceState);
 void LogResourceBarriers(uint32 NumBarriers, D3D12_RESOURCE_BARRIER *pBarriers, ID3D12CommandList *const pCommandList);
 
+void StallRHIThreadAndForceFlush(FD3D12Device* InDevice);
 
 // Custom resource states
 // To Be Determined (TBD) means we need to fill out a resource barrier before the command list is executed.
