Remove 'incorrect' check and added comments why and what can still possibly go wrong

diff --git a/Engine/Source/Runtime/D3D12RHI/Private/D3D12Query.cpp b/Engine/Source/Runtime/D3D12RHI/Private/D3D12Query.cpp
--- a/Engine/Source/Runtime/D3D12RHI/Private/D3D12Query.cpp
+++ b/Engine/Source/Runtime/D3D12RHI/Private/D3D12Query.cpp
@@ -299,14 +299,18 @@ uint32 FD3D12QueryHeap::AllocQuery(FD3D12CommandContext& CmdContext)
 	}
 	else
 	{
+		// Queries per batch gives a hint for resizing the query heap when starting a batch
+		// The code will still 'work' if it uses more than this value, but it might 'overwrite'
+		// previous batches (this isn't checked anywhere). If a single batch uses more queries
+		// than available in the heap than resolve will happen in the middle (see below) 
+		static const uint32 DefaultMaxTimeStampQueriesPerBat = 256;
+
 		if (!CurrentQueryBatch.bOpen)
 		{
-			StartQueryBatch(CmdContext, 256);
+			StartQueryBatch(CmdContext, DefaultMaxTimeStampQueriesPerBat);
 			check(CurrentQueryBatch.bOpen && CurrentQueryBatch.ElementCount == 0);
 		}
 
-		// No support for more than 256 timestamp queries?
-		check(CurrentQueryBatch.ElementCount < 256);
 		if (CurrentQueryBatch.StartElement > CurrentElement)
 		{
 			// We're in the middle of a batch, but we're at the end of the heap
@@ -317,7 +321,7 @@ uint32 FD3D12QueryHeap::AllocQuery(FD3D12CommandContext& CmdContext)
 		// check for the the batch being closed due to wrap and open a new one
 		if (!CurrentQueryBatch.bOpen)
 		{
-			StartQueryBatch(CmdContext, 256);
+			StartQueryBatch(CmdContext, DefaultMaxTimeStampQueriesPerBat);
 			check(CurrentQueryBatch.bOpen && CurrentQueryBatch.ElementCount == 0);
 		}
 	}
