Barrier fix raytrace bost BVH build readback data

diff --git a/Engine/Source/Runtime/D3D12RHI/Private/D3D12RayTracing.cpp b/Engine/Source/Runtime/D3D12RHI/Private/D3D12RayTracing.cpp
--- a/Engine/Source/Runtime/D3D12RHI/Private/D3D12RayTracing.cpp
+++ b/Engine/Source/Runtime/D3D12RHI/Private/D3D12RayTracing.cpp
@@ -2658,6 +2658,8 @@ void FD3D12RayTracingGeometry::BuildAccelerationStructure(FD3D12CommandContext&
 			PostBuildInfoDesc.DestBuffer = PostBuildInfoBuffers[GPUIndex]->ResourceLocation.GetGPUVirtualAddress();
 			PostBuildInfoDesc.InfoType = D3D12_RAYTRACING_ACCELERATION_STRUCTURE_POSTBUILD_INFO_COMPACTED_SIZE;
 
+			FD3D12DynamicRHI::TransitionResource(CommandContext.CommandListHandle, PostBuildInfoBuffers[GPUIndex].GetReference()->GetResource(), D3D12_RESOURCE_STATE_TBD, D3D12_RESOURCE_STATE_UNORDERED_ACCESS, 0, FD3D12DynamicRHI::ETransitionMode::Apply);
+
 			// Build with post build info
 			RayTracingCommandList->BuildRaytracingAccelerationStructure(&BuildDesc, 1, &PostBuildInfoDesc);
 
