Changing alignment requirement in FD3D12CommandContext::RHICopyTexture() to take into consideration each texture format.

diff --git a/Engine/Source/Runtime/D3D12RHI/Private/D3D12Texture.cpp b/Engine/Source/Runtime/D3D12RHI/Private/D3D12Texture.cpp
--- a/Engine/Source/Runtime/D3D12RHI/Private/D3D12Texture.cpp
+++ b/Engine/Source/Runtime/D3D12RHI/Private/D3D12Texture.cpp
@@ -3147,7 +3147,8 @@ void FD3D12CommandContext::RHICopyTexture(FRHITexture* SourceTextureRHI, FRHITex
 		Dst.pResource = DestTexture->GetResource()->GetResource();
 		Dst.Type = bReadback ? D3D12_TEXTURE_COPY_TYPE_PLACED_FOOTPRINT : D3D12_TEXTURE_COPY_TYPE_SUBRESOURCE_INDEX;
 
-		const FPixelFormatInfo& PixelFormatInfo = GPixelFormats[DestTextureRHI->GetFormat()];
+		const FPixelFormatInfo& SourcePixelFormatInfo = GPixelFormats[SourceTextureRHI->GetFormat()];
+		const FPixelFormatInfo& DestPixelFormatInfo = GPixelFormats[DestTextureRHI->GetFormat()];
 
 		D3D12_RESOURCE_DESC DstDesc = {};
 		FIntVector TextureSize = DestTextureRHI->GetSizeXYZ();
@@ -3156,7 +3157,7 @@ void FD3D12CommandContext::RHICopyTexture(FRHITexture* SourceTextureRHI, FRHITex
 		DstDesc.Height = TextureSize.Y;
 		DstDesc.DepthOrArraySize = TextureSize.Z;
 		DstDesc.MipLevels = DestTextureRHI->GetNumMips();
-		DstDesc.Format = (DXGI_FORMAT)PixelFormatInfo.PlatformFormat;
+		DstDesc.Format = (DXGI_FORMAT)DestPixelFormatInfo.PlatformFormat;
 		DstDesc.SampleDesc.Count = DestTextureRHI->GetNumSamples();
 
 		for (uint32 SliceIndex = 0; SliceIndex < CopyInfo.NumSlices; ++SliceIndex)
@@ -3175,9 +3176,9 @@ void FD3D12CommandContext::RHICopyTexture(FRHITexture* SourceTextureRHI, FRHITex
 					SourceBoxD3D.front >> MipIndex,
 					// Align to block size to pad the copy when processing the last surface texels.
 					// This will give inconsistent results otherwise between different RHI.
-					AlignArbitrary<uint32>(FMath::Max<uint32>(SourceBoxD3D.right  >> MipIndex, 1), PixelFormatInfo.BlockSizeX),
-					AlignArbitrary<uint32>(FMath::Max<uint32>(SourceBoxD3D.bottom >> MipIndex, 1), PixelFormatInfo.BlockSizeY),
-					AlignArbitrary<uint32>(FMath::Max<uint32>(SourceBoxD3D.back   >> MipIndex, 1), PixelFormatInfo.BlockSizeZ)
+					AlignArbitrary<uint32>(FMath::Max<uint32>(SourceBoxD3D.right  >> MipIndex, 1), SourcePixelFormatInfo.BlockSizeX),
+					AlignArbitrary<uint32>(FMath::Max<uint32>(SourceBoxD3D.bottom >> MipIndex, 1), SourcePixelFormatInfo.BlockSizeY),
+					AlignArbitrary<uint32>(FMath::Max<uint32>(SourceBoxD3D.back   >> MipIndex, 1), SourcePixelFormatInfo.BlockSizeZ)
 					);
 
 				const uint32 DestX = CopyInfo.DestPosition.X >> MipIndex;
@@ -3185,8 +3186,8 @@ void FD3D12CommandContext::RHICopyTexture(FRHITexture* SourceTextureRHI, FRHITex
 				const uint32 DestZ = CopyInfo.DestPosition.Z >> MipIndex;
 
 				// RHICopyTexture is allowed to copy mip regions only if are aligned on the block size to prevent unexpected / inconsistent results.
-				ensure(MipSourceBoxD3D.left % PixelFormatInfo.BlockSizeX == 0 && MipSourceBoxD3D.top % PixelFormatInfo.BlockSizeY == 0 && MipSourceBoxD3D.front % PixelFormatInfo.BlockSizeZ == 0);
-				ensure(DestX % PixelFormatInfo.BlockSizeX == 0 && DestY % PixelFormatInfo.BlockSizeY == 0 && DestZ % PixelFormatInfo.BlockSizeZ == 0);
+				ensure(MipSourceBoxD3D.left % SourcePixelFormatInfo.BlockSizeX == 0 && MipSourceBoxD3D.top % SourcePixelFormatInfo.BlockSizeY == 0 && MipSourceBoxD3D.front % SourcePixelFormatInfo.BlockSizeZ == 0);
+				ensure(DestX % DestPixelFormatInfo.BlockSizeX == 0 && DestY % DestPixelFormatInfo.BlockSizeY == 0 && DestZ % DestPixelFormatInfo.BlockSizeZ == 0);
 
 				Src.SubresourceIndex = CalcSubresource(SourceMipIndex, SourceSliceIndex, SourceTextureRHI->GetNumMips());
 				Dst.SubresourceIndex = CalcSubresource(DestMipIndex, DestSliceIndex, DestTextureRHI->GetNumMips());
