- Update D3D12 memory stats to track better all open allocation and try and match the DXGI used memory better
- Fix incorrect tracking of texture memory in D3D12
- Take 4K alignment into account when requesting texture sizes on d3d12

diff --git a/Engine/Source/Runtime/D3D12RHI/Private/D3D12Allocation.cpp b/Engine/Source/Runtime/D3D12RHI/Private/D3D12Allocation.cpp
--- a/Engine/Source/Runtime/D3D12RHI/Private/D3D12Allocation.cpp
+++ b/Engine/Source/Runtime/D3D12RHI/Private/D3D12Allocation.cpp
@@ -1396,6 +1396,8 @@ HRESULT FD3D12TextureAllocatorPool::AllocateTexture(
 				TextureLocation.SetSize(Info.SizeInBytes);
 				TextureLocation.SetOffsetFromBaseOfResource(Offset);
 				TextureLocation.SetResource(NewResource);
+
+				INC_DWORD_STAT(STAT_D3D12TextureAllocatorCount);
 			}
 			return RetCode;
 		}
@@ -1405,8 +1407,7 @@ HRESULT FD3D12TextureAllocatorPool::AllocateTexture(
 	Desc.Alignment = 0;
 	VERIFYD3D12RESULT(RetCode = Adapter->CreateCommittedResource(Desc, GetGPUMask(), HeapProps, InitialState, ClearValue, &NewResource, Name, false));
 
-	TextureLocation.SetType(FD3D12ResourceLocation::ResourceLocationType::eStandAlone);
-	TextureLocation.SetResource(NewResource);
+	TextureLocation.AsStandAlone(NewResource);
 	return RetCode;
 }
 #else
@@ -1430,13 +1431,10 @@ FD3D12TextureAllocator::FD3D12TextureAllocator(FD3D12Device* Device,
 		HeapSize,
 		D3D12_SMALL_RESOURCE_PLACEMENT_ALIGNMENT)
 {
-	// Inform the texture streaming system of this heap so that it correctly accounts for placed textures
-	FD3D12DynamicRHI::GetD3DRHI()->UpdataTextureMemorySize((DefaultPoolSize / 1024));
 }
 
 FD3D12TextureAllocator::~FD3D12TextureAllocator()
 {
-	FD3D12DynamicRHI::GetD3DRHI()->UpdataTextureMemorySize(-int32(DefaultPoolSize / 1024));
 }
 
 HRESULT FD3D12TextureAllocator::AllocateTexture(D3D12_RESOURCE_DESC Desc, const D3D12_CLEAR_VALUE* ClearValue, FD3D12ResourceLocation& TextureLocation, const D3D12_RESOURCE_STATES InitialState, const TCHAR* Name)
@@ -1473,8 +1471,7 @@ HRESULT FD3D12TextureAllocator::AllocateTexture(D3D12_RESOURCE_DESC Desc, const
 
 	hr = Adapter->CreateCommittedResource(Desc, GetGPUMask(), HeapProps, InitialState, ClearValue, &NewResource, Name, false);
 
-	TextureLocation.SetType(FD3D12ResourceLocation::ResourceLocationType::eStandAlone);
-	TextureLocation.SetResource(NewResource);
+	TextureLocation.AsStandAlone(NewResource, Info.SizeInBytes);
 
 	return hr;
 }
@@ -1851,6 +1848,8 @@ void FD3D12SegListAllocator::Deallocate(
 		}
 		new (DeferredDeletionQueue[LastIdx]) FRetiredBlock(PlacedResource, Offset, SizeInBytes);
 	}
+
+	DEC_DWORD_STAT(STAT_D3D12TextureAllocatorCount);
 }
 
 template <typename AllocX, typename AllocY>
diff --git a/Engine/Source/Runtime/D3D12RHI/Private/D3D12CommandContext.cpp b/Engine/Source/Runtime/D3D12RHI/Private/D3D12CommandContext.cpp
--- a/Engine/Source/Runtime/D3D12RHI/Private/D3D12CommandContext.cpp
+++ b/Engine/Source/Runtime/D3D12RHI/Private/D3D12CommandContext.cpp
@@ -686,7 +686,8 @@ void FD3D12CommandContextBase::UpdateMemoryStats()
 		uint64 TotalUnused;
 		Device->GetTextureAllocator().GetMemoryStats(TotalAllocated, TotalUnused);
 		MaxTexAllocWastage = FMath::Max(MaxTexAllocWastage, TotalUnused);
-		SET_MEMORY_STAT(STAT_D3D12TextureAllocatorWastage, MaxTexAllocWastage);
+		SET_MEMORY_STAT(STAT_D3D12TextureAllocatorAllocated, TotalAllocated);
+		SET_MEMORY_STAT(STAT_D3D12TextureAllocatorUnused, TotalUnused);
 #endif
 
 		Device->GetDefaultBufferAllocator().UpdateMemoryStats();
diff --git a/Engine/Source/Runtime/D3D12RHI/Private/D3D12RHICommon.h b/Engine/Source/Runtime/D3D12RHI/Private/D3D12RHICommon.h
--- a/Engine/Source/Runtime/D3D12RHI/Private/D3D12RHICommon.h
+++ b/Engine/Source/Runtime/D3D12RHI/Private/D3D12RHICommon.h
@@ -8,6 +8,7 @@ D3D12RHICommon.h: Common D3D12 RHI definitions for Windows.
 
 DECLARE_STATS_GROUP(TEXT("D3D12RHI"), STATGROUP_D3D12RHI, STATCAT_Advanced);
 DECLARE_STATS_GROUP(TEXT("D3D12RHI: Memory"), STATGROUP_D3D12Memory, STATCAT_Advanced);
+DECLARE_STATS_GROUP(TEXT("D3D12RHI: Memory (Buffers)"), STATGROUP_D3D12MemoryBuffer, STATCAT_Advanced);
 DECLARE_STATS_GROUP(TEXT("D3D12RHI: Pipeline State (PSO)"), STATGROUP_D3D12PipelineState, STATCAT_Advanced);
 DECLARE_STATS_GROUP(TEXT("D3D12RHI: Descriptor Heap (GPU visible)"), STATGROUP_D3D12DescriptorHeap, STATCAT_Advanced);
 
diff --git a/Engine/Source/Runtime/D3D12RHI/Private/D3D12RHIPrivate.h b/Engine/Source/Runtime/D3D12RHI/Private/D3D12RHIPrivate.h
--- a/Engine/Source/Runtime/D3D12RHI/Private/D3D12RHIPrivate.h
+++ b/Engine/Source/Runtime/D3D12RHI/Private/D3D12RHIPrivate.h
@@ -526,8 +526,6 @@ public:
 	uint64 PresentCount;
 #endif
 
-	inline void UpdataTextureMemorySize(int64 TextureSizeInKiloBytes) { FPlatformAtomics::InterlockedAdd(&GCurrentTextureMemorySize, TextureSizeInKiloBytes); }
-
 	/** Determine if an two views intersect */
 	template <class LeftT, class RightT>
 	static inline bool ResourceViewsIntersect(FD3D12View<LeftT>* pLeftView, FD3D12View<RightT>* pRightView)
diff --git a/Engine/Source/Runtime/D3D12RHI/Private/D3D12Resources.cpp b/Engine/Source/Runtime/D3D12RHI/Private/D3D12Resources.cpp
--- a/Engine/Source/Runtime/D3D12RHI/Private/D3D12Resources.cpp
+++ b/Engine/Source/Runtime/D3D12RHI/Private/D3D12Resources.cpp
@@ -659,6 +659,9 @@ void FD3D12ResourceLocation::ReleaseResource()
 	{
 	case ResourceLocationType::eStandAlone:
 	{
+		bool bIncrement = false;
+		UpdateStandAloneStats(bIncrement);
+
 		// Multi-GPU support : because of references, several GPU nodes can refrence the same stand-alone resource.
 		check(UnderlyingResource->GetRefCount() == 1 || GNumExplicitGPUsForRendering > 1);
 		
@@ -722,6 +725,73 @@ void FD3D12ResourceLocation::ReleaseResource()
 	}
 }
 
+void FD3D12ResourceLocation::UpdateStandAloneStats(bool bIncrement)
+{
+	if (UnderlyingResource->GetHeapType() == D3D12_HEAP_TYPE_DEFAULT)
+	{
+		D3D12_RESOURCE_DESC Desc = UnderlyingResource->GetDesc();
+		bool bIsBuffer = (Desc.Dimension == D3D12_RESOURCE_DIMENSION_BUFFER);
+		bool bIsRenderTarget = (Desc.Flags & D3D12_RESOURCE_FLAG_ALLOW_RENDER_TARGET || Desc.Flags & D3D12_RESOURCE_FLAG_ALLOW_DEPTH_STENCIL);
+		bool bIsUAV = (Desc.Flags & D3D12_RESOURCE_FLAG_ALLOW_UNORDERED_ACCESS) > 0;
+
+		// Get the desired size and allocated size for stand alone resources - allocated are very slow anyway
+		D3D12_RESOURCE_ALLOCATION_INFO Info = UnderlyingResource->GetParentDevice()->GetDevice()->GetResourceAllocationInfo(0, 1, &Desc);
+
+		if (bIsBuffer)
+		{
+			if (bIncrement)
+			{
+				INC_DWORD_STAT(STAT_D3D12BufferStandAloneCount);
+				INC_MEMORY_STAT_BY(STAT_D3D12BufferStandAloneAllocated, Info.SizeInBytes);
+			}
+			else
+			{
+				DEC_DWORD_STAT(STAT_D3D12BufferStandAloneCount);
+				DEC_MEMORY_STAT_BY(STAT_D3D12BufferStandAloneAllocated, Info.SizeInBytes);
+			}
+		}
+		else if (bIsRenderTarget)
+		{
+			if (bIncrement)
+			{
+				INC_DWORD_STAT(STAT_D3D12RenderTargetStandAloneCount);
+				INC_MEMORY_STAT_BY(STAT_D3D12RenderTargetStandAloneAllocated, Info.SizeInBytes);
+			}
+			else
+			{
+				DEC_DWORD_STAT(STAT_D3D12RenderTargetStandAloneCount);
+				DEC_MEMORY_STAT_BY(STAT_D3D12RenderTargetStandAloneAllocated, Info.SizeInBytes);
+			}
+		}
+		else if (bIsUAV)
+		{
+			if (bIncrement)
+			{
+				INC_DWORD_STAT(STAT_D3D12UAVTextureStandAloneCount);
+				INC_MEMORY_STAT_BY(STAT_D3D12UAVTextureStandAloneAllocated, Info.SizeInBytes);
+			}
+			else
+			{
+				DEC_DWORD_STAT(STAT_D3D12UAVTextureStandAloneCount);
+				DEC_MEMORY_STAT_BY(STAT_D3D12UAVTextureStandAloneAllocated, Info.SizeInBytes);
+			}
+		}
+		else
+		{
+			if (bIncrement)
+			{
+				INC_DWORD_STAT(STAT_D3D12TextureStandAloneCount);
+				INC_MEMORY_STAT_BY(STAT_D3D12TextureStandAloneAllocated, Info.SizeInBytes);
+			}
+			else
+			{
+				DEC_DWORD_STAT(STAT_D3D12TextureStandAloneCount);
+				DEC_MEMORY_STAT_BY(STAT_D3D12TextureStandAloneAllocated, Info.SizeInBytes);
+			}
+		}
+	}
+}
+
 void FD3D12ResourceLocation::SetResource(FD3D12Resource* Value)
 {
 	check(UnderlyingResource == nullptr);
@@ -733,6 +803,25 @@ void FD3D12ResourceLocation::SetResource(FD3D12Resource* Value)
 	ResidencyHandle = UnderlyingResource->GetResidencyHandle();
 }
 
+
+void FD3D12ResourceLocation::AsStandAlone(FD3D12Resource* Resource, uint32 BufferSize, bool bInIsTransient)
+{
+	SetType(FD3D12ResourceLocation::ResourceLocationType::eStandAlone);
+	SetResource(Resource);
+	SetSize(BufferSize);
+
+	if (!IsCPUInaccessible(Resource->GetHeapType()))
+	{
+		D3D12_RANGE range = { 0, IsCPUWritable(Resource->GetHeapType()) ? 0 : BufferSize };
+		SetMappedBaseAddress(Resource->Map(&range));
+	}
+	SetGPUVirtualAddress(Resource->GetGPUVirtualAddress());
+	SetTransient(bInIsTransient);
+
+	bool bIncrement = true;
+	UpdateStandAloneStats(bIncrement);
+}
+
 /////////////////////////////////////////////////////////////////////
 //	FD3D12 Resource Barrier Batcher
 /////////////////////////////////////////////////////////////////////
diff --git a/Engine/Source/Runtime/D3D12RHI/Private/D3D12Stats.h b/Engine/Source/Runtime/D3D12RHI/Private/D3D12Stats.h
--- a/Engine/Source/Runtime/D3D12RHI/Private/D3D12Stats.h
+++ b/Engine/Source/Runtime/D3D12RHI/Private/D3D12Stats.h
@@ -58,19 +58,35 @@ DECLARE_CYCLE_STAT_EXTERN(TEXT("Clear MRT time"), STAT_D3D12ClearMRT, STATGROUP_
 DECLARE_CYCLE_STAT_EXTERN(TEXT("ExecuteCommandList time"), STAT_D3D12ExecuteCommandListTime, STATGROUP_D3D12RHI, );
 DECLARE_CYCLE_STAT_EXTERN(TEXT("WaitForFence time"), STAT_D3D12WaitForFenceTime, STATGROUP_D3D12RHI, );
 
-DECLARE_MEMORY_STAT_EXTERN(TEXT("Used Video Memory"), STAT_D3D12UsedVideoMemory, STATGROUP_D3D12RHI, );
-DECLARE_MEMORY_STAT_EXTERN(TEXT("Available Video Memory"), STAT_D3D12AvailableVideoMemory, STATGROUP_D3D12RHI, );
-DECLARE_MEMORY_STAT_EXTERN(TEXT("Total Video Memory"), STAT_D3D12TotalVideoMemory, STATGROUP_D3D12RHI, );
-DECLARE_MEMORY_STAT_EXTERN(TEXT("Texture allocator wastage"), STAT_D3D12TextureAllocatorWastage, STATGROUP_D3D12RHI, );
+DECLARE_MEMORY_STAT_EXTERN(TEXT("Used Video Memory"), STAT_D3D12UsedVideoMemory, STATGROUP_D3D12Memory, );
+DECLARE_MEMORY_STAT_EXTERN(TEXT("Available Video Memory"), STAT_D3D12AvailableVideoMemory, STATGROUP_D3D12Memory, );
+DECLARE_MEMORY_STAT_EXTERN(TEXT("Total Video Memory"), STAT_D3D12TotalVideoMemory, STATGROUP_D3D12Memory, );
 
+DECLARE_MEMORY_STAT_EXTERN(TEXT("RenderTarget StandAlone Allocated"), STAT_D3D12RenderTargetStandAloneAllocated, STATGROUP_D3D12Memory, );
+DECLARE_MEMORY_STAT_EXTERN(TEXT("UAVTexture StandAlone Allocated"), STAT_D3D12UAVTextureStandAloneAllocated, STATGROUP_D3D12Memory, );
+DECLARE_MEMORY_STAT_EXTERN(TEXT("Texture StandAlone Allocated"), STAT_D3D12TextureStandAloneAllocated, STATGROUP_D3D12Memory, );
+DECLARE_MEMORY_STAT_EXTERN(TEXT("Buffer StandAlone Allocated"), STAT_D3D12BufferStandAloneAllocated, STATGROUP_D3D12Memory, );
+
+DECLARE_DWORD_ACCUMULATOR_STAT_EXTERN(TEXT("RenderTarget StandAlone Allocated"), STAT_D3D12RenderTargetStandAloneCount, STATGROUP_D3D12Memory, );
+DECLARE_DWORD_ACCUMULATOR_STAT_EXTERN(TEXT("UAVTexture StandAlone Allocated"), STAT_D3D12UAVTextureStandAloneCount, STATGROUP_D3D12Memory, );
+DECLARE_DWORD_ACCUMULATOR_STAT_EXTERN(TEXT("Texture StandAlone Allocated"), STAT_D3D12TextureStandAloneCount, STATGROUP_D3D12Memory, );
+DECLARE_DWORD_ACCUMULATOR_STAT_EXTERN(TEXT("Buffer StandAlone Allocated"), STAT_D3D12BufferStandAloneCount, STATGROUP_D3D12Memory, );
+
+DECLARE_MEMORY_STAT_EXTERN(TEXT("Texture Allocator Allocated"), STAT_D3D12TextureAllocatorAllocated, STATGROUP_D3D12Memory, );
+DECLARE_MEMORY_STAT_EXTERN(TEXT("Texture Allocator Unused"), STAT_D3D12TextureAllocatorUnused, STATGROUP_D3D12Memory, );
+DECLARE_DWORD_ACCUMULATOR_STAT_EXTERN(TEXT("Texture Allocator Allocated"), STAT_D3D12TextureAllocatorCount, STATGROUP_D3D12Memory, );
+
+DECLARE_MEMORY_STAT_POOL_EXTERN(TEXT("View DescriptorHeap Memory (SRV, CBV, UAV)"), STAT_ViewOnlineDescriptorHeapMemory, STATGROUP_D3D12Memory, FPlatformMemory::MCR_GPUSystem, );
+DECLARE_MEMORY_STAT_POOL_EXTERN(TEXT("Sampler DescriptorHeap Memory"), STAT_SamplerOnlineDescriptorHeapMemory, STATGROUP_D3D12Memory, FPlatformMemory::MCR_GPUSystem, );
 
 DECLARE_MEMORY_STAT_EXTERN(TEXT("BufferPool Memory Allocated"), STAT_D3D12BufferPoolMemoryAllocated, STATGROUP_D3D12Memory, );
 DECLARE_MEMORY_STAT_EXTERN(TEXT("BufferPool Memory Used"), STAT_D3D12BufferPoolMemoryUsed, STATGROUP_D3D12Memory, );
-DECLARE_MEMORY_STAT_EXTERN(TEXT("BufferPool Memory Free"), STAT_D3D12BufferPoolMemoryFree, STATGROUP_D3D12Memory, );
-DECLARE_MEMORY_STAT_EXTERN(TEXT("BufferPool Memory Alignment Waste"), STAT_D3D12BufferPoolAlignmentWaste, STATGROUP_D3D12Memory, );
-DECLARE_DWORD_ACCUMULATOR_STAT_EXTERN(TEXT("BufferPool Page Count"), STAT_D3D12BufferPoolPageCount, STATGROUP_D3D12Memory, );
-DECLARE_DWORD_ACCUMULATOR_STAT_EXTERN(TEXT("BufferPool Full Pages"), STAT_D3D12BufferPoolFullPages, STATGROUP_D3D12Memory, );
-DECLARE_MEMORY_STAT_EXTERN(TEXT("Buffer StandAlone Memory Used"), STAT_D3D12BufferStandAloneUsedMemory, STATGROUP_D3D12Memory, );
+
+DECLARE_MEMORY_STAT_EXTERN(TEXT("BufferPool Memory Free"), STAT_D3D12BufferPoolMemoryFree, STATGROUP_D3D12MemoryBuffer, ); 
+DECLARE_MEMORY_STAT_EXTERN(TEXT("BufferPool Memory Alignment Waste"), STAT_D3D12BufferPoolAlignmentWaste, STATGROUP_D3D12MemoryBuffer, );
+DECLARE_DWORD_ACCUMULATOR_STAT_EXTERN(TEXT("BufferPool Page Count"), STAT_D3D12BufferPoolPageCount, STATGROUP_D3D12MemoryBuffer, );
+DECLARE_DWORD_ACCUMULATOR_STAT_EXTERN(TEXT("BufferPool Full Pages"), STAT_D3D12BufferPoolFullPages, STATGROUP_D3D12MemoryBuffer, );
+
 
 /**
 * Detailed Descriptor heap stats
@@ -88,8 +104,6 @@ DECLARE_DWORD_COUNTER_STAT_EXTERN(TEXT("View: Num reserved descriptors"), STAT_N
 DECLARE_DWORD_COUNTER_STAT_EXTERN(TEXT("Sampler: Num reserved descriptors"), STAT_NumReservedSamplerOnlineDescriptors, STATGROUP_D3D12DescriptorHeap, );
 DECLARE_DWORD_COUNTER_STAT_EXTERN(TEXT("Sampler: Num reused descriptors"), STAT_NumReusedSamplerOnlineDescriptors, STATGROUP_D3D12DescriptorHeap, );
 
-DECLARE_MEMORY_STAT_POOL_EXTERN(TEXT("View: Total descriptor heap memory (SRV, CBV, UAV)"), STAT_ViewOnlineDescriptorHeapMemory, STATGROUP_D3D12DescriptorHeap, FPlatformMemory::MCR_GPUSystem, );
-DECLARE_MEMORY_STAT_POOL_EXTERN(TEXT("Sampler: Total descriptor heap memory"), STAT_SamplerOnlineDescriptorHeapMemory, STATGROUP_D3D12DescriptorHeap, FPlatformMemory::MCR_GPUSystem, );
 
 DECLARE_DWORD_ACCUMULATOR_STAT_EXTERN(TEXT("View Global: Free Descriptors"), STAT_GlobalViewHeapFreeDescriptors, STATGROUP_D3D12DescriptorHeap, );
 DECLARE_DWORD_ACCUMULATOR_STAT_EXTERN(TEXT("View Global: Reserved Descriptors"), STAT_GlobalViewHeapReservedDescriptors, STATGROUP_D3D12DescriptorHeap, );
diff --git a/Engine/Source/Runtime/D3D12RHI/Private/D3D12Texture.cpp b/Engine/Source/Runtime/D3D12RHI/Private/D3D12Texture.cpp
--- a/Engine/Source/Runtime/D3D12RHI/Private/D3D12Texture.cpp
+++ b/Engine/Source/Runtime/D3D12RHI/Private/D3D12Texture.cpp
@@ -350,8 +350,7 @@ void FD3D12TextureStats::D3D12TextureAllocated(TD3D12Texture2D<BaseResourceType>
 {
 	FD3D12Resource* D3D12Texture2D = Texture.GetResource();
 
-	// Ignore placed textures as their memory is already allocated and accounted for
-	if (D3D12Texture2D && D3D12Texture2D->IsPlacedResource() == false)
+	if (D3D12Texture2D)
 	{
 		if ((Texture.Flags & TexCreate_Virtual) != TexCreate_Virtual)
 		{
@@ -382,8 +381,7 @@ void FD3D12TextureStats::D3D12TextureDeleted(TD3D12Texture2D<BaseResourceType>&
 {
 	FD3D12Resource* D3D12Texture2D = Texture.GetResource();
 
-	// Ignore placed textures as their memory is already allocated and accounted for
-	if (D3D12Texture2D && D3D12Texture2D->IsPlacedResource() == false)
+	if (D3D12Texture2D)
 	{
 		const D3D12_RESOURCE_DESC& Desc = D3D12Texture2D->GetDesc();
 		const int64 TextureSize = Texture.GetMemorySize();
@@ -486,6 +484,9 @@ uint64 FD3D12DynamicRHI::RHICalcTexture2DPlatformSize(uint32 SizeX, uint32 SizeY
 	Desc.SampleDesc.Count = NumSamples;
 	Desc.Width = SizeX;
 
+	// Check if the 4K aligment is possible
+	Desc.Alignment = TextureCanBe4KAligned(Desc, Format) ? D3D12_SMALL_RESOURCE_PLACEMENT_ALIGNMENT : 0;
+
 	const D3D12_RESOURCE_ALLOCATION_INFO AllocationInfo = GetAdapter().GetD3DDevice()->GetResourceAllocationInfo(0, 1, &Desc);
 	OutAlign = static_cast<uint32>(AllocationInfo.Alignment);
 
@@ -504,6 +505,9 @@ uint64 FD3D12DynamicRHI::RHICalcTexture2DArrayPlatformSize(uint32 SizeX, uint32
 	Desc.SampleDesc.Count = NumSamples;
 	Desc.Width = SizeX;
 
+	// Check if the 4K aligment is possible
+	Desc.Alignment = TextureCanBe4KAligned(Desc, Format) ? D3D12_SMALL_RESOURCE_PLACEMENT_ALIGNMENT : 0;
+
 	const D3D12_RESOURCE_ALLOCATION_INFO AllocationInfo = GetAdapter().GetD3DDevice()->GetResourceAllocationInfo(0, 1, &Desc);
 	OutAlign = static_cast<uint32>(AllocationInfo.Alignment);
 
@@ -522,6 +526,9 @@ uint64 FD3D12DynamicRHI::RHICalcTexture3DPlatformSize(uint32 SizeX, uint32 SizeY
 	Desc.SampleDesc.Count = 1;
 	Desc.Width = SizeX;
 
+	// Check if the 4K aligment is possible
+	Desc.Alignment = TextureCanBe4KAligned(Desc, Format) ? D3D12_SMALL_RESOURCE_PLACEMENT_ALIGNMENT : 0;
+
 	const D3D12_RESOURCE_ALLOCATION_INFO AllocationInfo = GetAdapter().GetD3DDevice()->GetResourceAllocationInfo(0, 1, &Desc);
 	OutAlign = static_cast<uint32>(AllocationInfo.Alignment);
 
@@ -540,6 +547,9 @@ uint64 FD3D12DynamicRHI::RHICalcTextureCubePlatformSize(uint32 Size, uint8 Forma
 	Desc.SampleDesc.Count = 1;
 	Desc.Width = Size;
 
+	// Check if the 4K aligment is possible
+	Desc.Alignment = TextureCanBe4KAligned(Desc, Format) ? D3D12_SMALL_RESOURCE_PLACEMENT_ALIGNMENT : 0;
+
 	const D3D12_RESOURCE_ALLOCATION_INFO AllocationInfo = GetAdapter().GetD3DDevice()->GetResourceAllocationInfo(0, 1, &Desc);
 	OutAlign = static_cast<uint32>(AllocationInfo.Alignment);
 
diff --git a/Engine/Source/Runtime/D3D12RHI/Private/D3D12Util.cpp b/Engine/Source/Runtime/D3D12RHI/Private/D3D12Util.cpp
--- a/Engine/Source/Runtime/D3D12RHI/Private/D3D12Util.cpp
+++ b/Engine/Source/Runtime/D3D12RHI/Private/D3D12Util.cpp
@@ -1453,7 +1453,20 @@ DEFINE_STAT(STAT_D3D12WaitForFenceTime);
 DEFINE_STAT(STAT_D3D12UsedVideoMemory);
 DEFINE_STAT(STAT_D3D12AvailableVideoMemory);
 DEFINE_STAT(STAT_D3D12TotalVideoMemory);
-DEFINE_STAT(STAT_D3D12TextureAllocatorWastage);
+
+DEFINE_STAT(STAT_D3D12RenderTargetStandAloneAllocated);
+DEFINE_STAT(STAT_D3D12UAVTextureStandAloneAllocated);
+DEFINE_STAT(STAT_D3D12TextureStandAloneAllocated);
+DEFINE_STAT(STAT_D3D12BufferStandAloneAllocated);
+
+DEFINE_STAT(STAT_D3D12RenderTargetStandAloneCount);
+DEFINE_STAT(STAT_D3D12UAVTextureStandAloneCount);
+DEFINE_STAT(STAT_D3D12TextureStandAloneCount);
+DEFINE_STAT(STAT_D3D12BufferStandAloneCount);
+
+DEFINE_STAT(STAT_D3D12TextureAllocatorAllocated);
+DEFINE_STAT(STAT_D3D12TextureAllocatorUnused);
+DEFINE_STAT(STAT_D3D12TextureAllocatorCount);
 
 DEFINE_STAT(STAT_D3D12BufferPoolMemoryAllocated);
 DEFINE_STAT(STAT_D3D12BufferPoolMemoryUsed);
@@ -1461,7 +1474,6 @@ DEFINE_STAT(STAT_D3D12BufferPoolMemoryFree);
 DEFINE_STAT(STAT_D3D12BufferPoolAlignmentWaste);
 DEFINE_STAT(STAT_D3D12BufferPoolPageCount);
 DEFINE_STAT(STAT_D3D12BufferPoolFullPages);
-DEFINE_STAT(STAT_D3D12BufferStandAloneUsedMemory);
 
 DEFINE_STAT(STAT_UniqueSamplers);
 
diff --git a/Engine/Source/Runtime/D3D12RHI/Public/D3D12Resources.h b/Engine/Source/Runtime/D3D12RHI/Public/D3D12Resources.h
--- a/Engine/Source/Runtime/D3D12RHI/Public/D3D12Resources.h
+++ b/Engine/Source/Runtime/D3D12RHI/Public/D3D12Resources.h
@@ -534,20 +534,7 @@ public:
 
 	const inline bool IsValid() const { return Type != ResourceLocationType::eUndefined; }
 
-	inline void AsStandAlone(FD3D12Resource* Resource, uint32 BufferSize = 0, bool bInIsTransient = false )
-	{
-		SetType(FD3D12ResourceLocation::ResourceLocationType::eStandAlone);
-		SetResource(Resource);
-		SetSize(BufferSize);
-
-		if (!IsCPUInaccessible(Resource->GetHeapType()))
-		{
-			D3D12_RANGE range = { 0, IsCPUWritable(Resource->GetHeapType())? 0 : BufferSize };
-			SetMappedBaseAddress(Resource->Map(&range));
-		}
-		SetGPUVirtualAddress(Resource->GetGPUVirtualAddress());
-		SetTransient(bInIsTransient);
-	}
+	void AsStandAlone(FD3D12Resource* Resource, uint32 BufferSize = 0, bool bInIsTransient = false);
 
 	inline void AsHeapAliased(FD3D12Resource* Resource)
 	{
@@ -615,6 +602,7 @@ private:
 	void InternalClear();
 
 	void ReleaseResource();
+	void UpdateStandAloneStats(bool bIncrement);
 
 	ResourceLocationType Type;
 
diff --git a/Engine/Source/Runtime/D3D12RHI/Public/D3D12Util.h b/Engine/Source/Runtime/D3D12RHI/Public/D3D12Util.h
--- a/Engine/Source/Runtime/D3D12RHI/Public/D3D12Util.h
+++ b/Engine/Source/Runtime/D3D12RHI/Public/D3D12Util.h
@@ -782,12 +782,23 @@ static void Get4KTileShape(D3D12_TILE_SHAPE* pTileShape, DXGI_FORMAT Format, uin
 
 static bool TextureCanBe4KAligned(D3D12_RESOURCE_DESC& Desc, uint8 UEFormat)
 {
-	D3D12_TILE_SHAPE Tile = {};
-	Get4KTileShape(&Tile, Desc.Format, UEFormat, Desc.Dimension, Desc.SampleDesc.Count);
+	// 4KB alignment is only available for read only textures
+	if (!(Desc.Flags & D3D12_RESOURCE_FLAG_ALLOW_RENDER_TARGET ||
+		Desc.Flags & D3D12_RESOURCE_FLAG_ALLOW_DEPTH_STENCIL ||
+		Desc.Flags & D3D12_RESOURCE_FLAG_ALLOW_UNORDERED_ACCESS) &&
+		Desc.SampleDesc.Count == 1)
+	{
+		D3D12_TILE_SHAPE Tile = {};
+		Get4KTileShape(&Tile, Desc.Format, UEFormat, Desc.Dimension, Desc.SampleDesc.Count);
 
-	uint32 TilesNeeded = GetTilesNeeded(Desc.Width, Desc.Height, Desc.DepthOrArraySize, Tile);
+		uint32 TilesNeeded = GetTilesNeeded(Desc.Width, Desc.Height, Desc.DepthOrArraySize, Tile);
 
-	return TilesNeeded <= NUM_4K_BLOCKS_PER_64K_PAGE;
+		return TilesNeeded <= NUM_4K_BLOCKS_PER_64K_PAGE;
+	}
+	else
+	{
+		return false;
+	}
 }
 
 template <class TView>
