- Split UAV buffer sizes into different stat for d3d12
- Simplify stat tracking code

diff --git a/Engine/Source/Runtime/D3D12RHI/Private/D3D12Resources.cpp b/Engine/Source/Runtime/D3D12RHI/Private/D3D12Resources.cpp
--- a/Engine/Source/Runtime/D3D12RHI/Private/D3D12Resources.cpp
+++ b/Engine/Source/Runtime/D3D12RHI/Private/D3D12Resources.cpp
@@ -737,56 +737,38 @@ void FD3D12ResourceLocation::UpdateStandAloneStats(bool bIncrement)
 		// Get the desired size and allocated size for stand alone resources - allocated are very slow anyway
 		D3D12_RESOURCE_ALLOCATION_INFO Info = UnderlyingResource->GetParentDevice()->GetDevice()->GetResourceAllocationInfo(0, 1, &Desc);
 
+		int64 SizeInBytes = bIncrement ? Info.SizeInBytes : -(int64)Info.SizeInBytes;
+		int32 Count = bIncrement ? 1 : -1;
+
 		if (bIsBuffer)
 		{
-			if (bIncrement)
-			{
-				INC_DWORD_STAT(STAT_D3D12BufferStandAloneCount);
-				INC_MEMORY_STAT_BY(STAT_D3D12BufferStandAloneAllocated, Info.SizeInBytes);
-			}
-			else
-			{
-				DEC_DWORD_STAT(STAT_D3D12BufferStandAloneCount);
-				DEC_MEMORY_STAT_BY(STAT_D3D12BufferStandAloneAllocated, Info.SizeInBytes);
-			}
-		}
-		else if (bIsRenderTarget)
-		{
-			if (bIncrement)
+			if (bIsUAV)
 			{
-				INC_DWORD_STAT(STAT_D3D12RenderTargetStandAloneCount);
-				INC_MEMORY_STAT_BY(STAT_D3D12RenderTargetStandAloneAllocated, Info.SizeInBytes);
+				INC_DWORD_STAT_BY(STAT_D3D12UAVBufferStandAloneCount, Count);
+				INC_MEMORY_STAT_BY(STAT_D3D12UAVBufferStandAloneAllocated, SizeInBytes);
 			}
 			else
 			{
-				DEC_DWORD_STAT(STAT_D3D12RenderTargetStandAloneCount);
-				DEC_MEMORY_STAT_BY(STAT_D3D12RenderTargetStandAloneAllocated, Info.SizeInBytes);
+				INC_DWORD_STAT_BY(STAT_D3D12BufferStandAloneCount, Count);
+				INC_MEMORY_STAT_BY(STAT_D3D12BufferStandAloneAllocated, SizeInBytes);
 			}
 		}
-		else if (bIsUAV)
+		else
 		{
-			if (bIncrement)
-			{
-				INC_DWORD_STAT(STAT_D3D12UAVTextureStandAloneCount);
-				INC_MEMORY_STAT_BY(STAT_D3D12UAVTextureStandAloneAllocated, Info.SizeInBytes);
-			}
-			else
+			if (bIsRenderTarget)
 			{
-				DEC_DWORD_STAT(STAT_D3D12UAVTextureStandAloneCount);
-				DEC_MEMORY_STAT_BY(STAT_D3D12UAVTextureStandAloneAllocated, Info.SizeInBytes);
+				INC_DWORD_STAT_BY(STAT_D3D12RenderTargetStandAloneCount, Count);
+				INC_MEMORY_STAT_BY(STAT_D3D12RenderTargetStandAloneAllocated, SizeInBytes);
 			}
-		}
-		else
-		{
-			if (bIncrement)
+			else if (bIsUAV)
 			{
-				INC_DWORD_STAT(STAT_D3D12TextureStandAloneCount);
-				INC_MEMORY_STAT_BY(STAT_D3D12TextureStandAloneAllocated, Info.SizeInBytes);
+				INC_DWORD_STAT_BY(STAT_D3D12UAVTextureStandAloneCount, Count);
+				INC_MEMORY_STAT_BY(STAT_D3D12UAVTextureStandAloneAllocated, SizeInBytes);
 			}
 			else
 			{
-				DEC_DWORD_STAT(STAT_D3D12TextureStandAloneCount);
-				DEC_MEMORY_STAT_BY(STAT_D3D12TextureStandAloneAllocated, Info.SizeInBytes);
+				INC_DWORD_STAT_BY(STAT_D3D12TextureStandAloneCount, Count);
+				INC_MEMORY_STAT_BY(STAT_D3D12TextureStandAloneAllocated, SizeInBytes);
 			}
 		}
 	}
diff --git a/Engine/Source/Runtime/D3D12RHI/Private/D3D12Stats.h b/Engine/Source/Runtime/D3D12RHI/Private/D3D12Stats.h
--- a/Engine/Source/Runtime/D3D12RHI/Private/D3D12Stats.h
+++ b/Engine/Source/Runtime/D3D12RHI/Private/D3D12Stats.h
@@ -65,11 +65,13 @@ DECLARE_MEMORY_STAT_EXTERN(TEXT("Total Video Memory"), STAT_D3D12TotalVideoMemor
 DECLARE_MEMORY_STAT_EXTERN(TEXT("RenderTarget StandAlone Allocated"), STAT_D3D12RenderTargetStandAloneAllocated, STATGROUP_D3D12Memory, );
 DECLARE_MEMORY_STAT_EXTERN(TEXT("UAVTexture StandAlone Allocated"), STAT_D3D12UAVTextureStandAloneAllocated, STATGROUP_D3D12Memory, );
 DECLARE_MEMORY_STAT_EXTERN(TEXT("Texture StandAlone Allocated"), STAT_D3D12TextureStandAloneAllocated, STATGROUP_D3D12Memory, );
+DECLARE_MEMORY_STAT_EXTERN(TEXT("UAVBuffer StandAlone Allocated"), STAT_D3D12UAVBufferStandAloneAllocated, STATGROUP_D3D12Memory, );
 DECLARE_MEMORY_STAT_EXTERN(TEXT("Buffer StandAlone Allocated"), STAT_D3D12BufferStandAloneAllocated, STATGROUP_D3D12Memory, );
 
 DECLARE_DWORD_ACCUMULATOR_STAT_EXTERN(TEXT("RenderTarget StandAlone Allocated"), STAT_D3D12RenderTargetStandAloneCount, STATGROUP_D3D12Memory, );
 DECLARE_DWORD_ACCUMULATOR_STAT_EXTERN(TEXT("UAVTexture StandAlone Allocated"), STAT_D3D12UAVTextureStandAloneCount, STATGROUP_D3D12Memory, );
 DECLARE_DWORD_ACCUMULATOR_STAT_EXTERN(TEXT("Texture StandAlone Allocated"), STAT_D3D12TextureStandAloneCount, STATGROUP_D3D12Memory, );
+DECLARE_DWORD_ACCUMULATOR_STAT_EXTERN(TEXT("UAVBuffer StandAlone Allocated"), STAT_D3D12UAVBufferStandAloneCount, STATGROUP_D3D12Memory, );
 DECLARE_DWORD_ACCUMULATOR_STAT_EXTERN(TEXT("Buffer StandAlone Allocated"), STAT_D3D12BufferStandAloneCount, STATGROUP_D3D12Memory, );
 
 DECLARE_MEMORY_STAT_EXTERN(TEXT("Texture Allocator Allocated"), STAT_D3D12TextureAllocatorAllocated, STATGROUP_D3D12Memory, );
diff --git a/Engine/Source/Runtime/D3D12RHI/Private/D3D12Util.cpp b/Engine/Source/Runtime/D3D12RHI/Private/D3D12Util.cpp
--- a/Engine/Source/Runtime/D3D12RHI/Private/D3D12Util.cpp
+++ b/Engine/Source/Runtime/D3D12RHI/Private/D3D12Util.cpp
@@ -1457,11 +1457,13 @@ DEFINE_STAT(STAT_D3D12TotalVideoMemory);
 DEFINE_STAT(STAT_D3D12RenderTargetStandAloneAllocated);
 DEFINE_STAT(STAT_D3D12UAVTextureStandAloneAllocated);
 DEFINE_STAT(STAT_D3D12TextureStandAloneAllocated);
+DEFINE_STAT(STAT_D3D12UAVBufferStandAloneAllocated);
 DEFINE_STAT(STAT_D3D12BufferStandAloneAllocated);
 
 DEFINE_STAT(STAT_D3D12RenderTargetStandAloneCount);
 DEFINE_STAT(STAT_D3D12UAVTextureStandAloneCount);
 DEFINE_STAT(STAT_D3D12TextureStandAloneCount);
+DEFINE_STAT(STAT_D3D12UAVBufferStandAloneCount);
 DEFINE_STAT(STAT_D3D12BufferStandAloneCount);
 
 DEFINE_STAT(STAT_D3D12TextureAllocatorAllocated);
