Removing d3d12 message filter that doesn't exist in later Windows SDKs since it's never emitted by anything. This is a remnant of copy pasting dx11 into the start of dx12.

diff --git a/Engine/Source/Runtime/D3D12RHI/Private/D3D12Adapter.cpp b/Engine/Source/Runtime/D3D12RHI/Private/D3D12Adapter.cpp
--- a/Engine/Source/Runtime/D3D12RHI/Private/D3D12Adapter.cpp
+++ b/Engine/Source/Runtime/D3D12RHI/Private/D3D12Adapter.cpp
@@ -576,13 +576,6 @@ void FD3D12Adapter::CreateRootDevice(bool bWithDebug)
 				// This typically happens when a non-depth-only pixel shader is used for depth-only rendering.
 				D3D12_MESSAGE_ID_CREATEGRAPHICSPIPELINESTATE_RENDERTARGETVIEW_NOT_SET,
 
-#if PLATFORM_DESKTOP || PLATFORM_HOLOLENS
-				// OMSETRENDERTARGETS_INVALIDVIEW - d3d will complain if depth and color targets don't have the exact same dimensions, but actually
-				//	if the color target is smaller then things are ok.  So turn off this error.  There is a manual check in FD3D12DynamicRHI::SetRenderTarget
-				//	that tests for depth smaller than color and MSAA settings to match.
-				D3D12_MESSAGE_ID_OMSETRENDERTARGETS_INVALIDVIEW,
-#endif
-
 				// QUERY_BEGIN_ABANDONING_PREVIOUS_RESULTS - The RHI exposes the interface to make and issue queries and a separate interface to use that data.
 				//		Currently there is a situation where queries are issued and the results may be ignored on purpose.  Filtering out this message so it doesn't
 				//		swarm the debug spew and mask other important warnings
