Fix for crash in D3D12 with static uniform buffers not being cleared properly.

diff --git a/Engine/Source/Runtime/D3D12RHI/Private/D3D12CommandContext.cpp b/Engine/Source/Runtime/D3D12RHI/Private/D3D12CommandContext.cpp
--- a/Engine/Source/Runtime/D3D12RHI/Private/D3D12CommandContext.cpp
+++ b/Engine/Source/Runtime/D3D12RHI/Private/D3D12CommandContext.cpp
@@ -501,6 +501,7 @@ void FD3D12CommandContext::ClearState()
 
 	FMemory::Memzero(BoundUniformBuffers, sizeof(BoundUniformBuffers));
 	FMemory::Memzero(DirtyUniformBuffers, sizeof(DirtyUniformBuffers));
+	FMemory::Memzero(GlobalUniformBuffers.GetData(), GlobalUniformBuffers.Num() * sizeof(FRHIUniformBuffer*));
 
 	for (int i = 0; i < UE_ARRAY_COUNT(BoundUniformBufferRefs); i++)
 	{
