- Reenable check(CommandListType != D3D12_COMMAND_LIST_TYPE_COMPUTE);
- Fix typo in function name

diff --git a/Engine/Source/Runtime/D3D12RHI/Private/D3D12DirectCommandListManager.cpp b/Engine/Source/Runtime/D3D12RHI/Private/D3D12DirectCommandListManager.cpp
--- a/Engine/Source/Runtime/D3D12RHI/Private/D3D12DirectCommandListManager.cpp
+++ b/Engine/Source/Runtime/D3D12RHI/Private/D3D12DirectCommandListManager.cpp
@@ -615,18 +615,18 @@ void FD3D12CommandListManager::ExecuteCommandLists(TArray<FD3D12CommandListHandl
 		ExecuteTask = FFunctionGraphTask::CreateAndDispatchWhenReady(
 			[this]()
 			{	
-				ExecuteCommandListInteral(ExecuteCommandListHandles, false);
+				ExecuteCommandListInternal(ExecuteCommandListHandles, false);
 				ExecuteCommandListHandles.Reset();
 			}, TStatId(), nullptr, ENamedThreads::AnyThread);
 	}
 	else
 	{
-		ExecuteCommandListInteral(Lists, WaitForCompletion);
+		ExecuteCommandListInternal(Lists, WaitForCompletion);
 	}
 }
 
 
-void FD3D12CommandListManager::ExecuteCommandListInteral(TArray<FD3D12CommandListHandle>& Lists, bool WaitForCompletion)
+void FD3D12CommandListManager::ExecuteCommandListInternal(TArray<FD3D12CommandListHandle>& Lists, bool WaitForCompletion)
 {
 	SCOPE_CYCLE_COUNTER(STAT_D3D12ExecuteCommandListTime);
 	check(CommandListFence);
@@ -736,7 +736,7 @@ void FD3D12CommandListManager::ExecuteCommandListInteral(TArray<FD3D12CommandLis
 			Lists[i].LogResourceBarriers();
 		}
 		SignaledFenceValue = ExecuteAndIncrementFence(CurrentCommandListPayload, *CommandListFence);
-		//check(CommandListType != D3D12_COMMAND_LIST_TYPE_COMPUTE);
+		check(CommandListType != D3D12_COMMAND_LIST_TYPE_COMPUTE);
 		SyncPoint = FD3D12SyncPoint(CommandListFence, SignaledFenceValue);
 		BarrierSyncPoint = SyncPoint;
 	}
diff --git a/Engine/Source/Runtime/D3D12RHI/Private/D3D12DirectCommandListManager.h b/Engine/Source/Runtime/D3D12RHI/Private/D3D12DirectCommandListManager.h
--- a/Engine/Source/Runtime/D3D12RHI/Private/D3D12DirectCommandListManager.h
+++ b/Engine/Source/Runtime/D3D12RHI/Private/D3D12DirectCommandListManager.h
@@ -308,7 +308,7 @@ protected:
 		{}
 	};
 
-	void ExecuteCommandListInteral(TArray<FD3D12CommandListHandle>& Lists, bool WaitForCompletion);
+	void ExecuteCommandListInternal(TArray<FD3D12CommandListHandle>& Lists, bool WaitForCompletion);
 	uint32 GetResourceBarrierCommandList(FD3D12CommandListHandle& hList, FD3D12CommandListHandle& hResourceBarrierList);
 
 	// Returns signaled Fence
