- Add RHI indepent RHIPoolAllocator which manages a set of pools which contain a link list entry (supports defrag)
- Add D3D12PoolAllocator which is a d3d12 specific implementation of the shared pool - handles D3D12 resource management & operations (eg copies & resource creation)
- Also add UAVs as views to resources
- D3D12ResourceLocation now supports pool allocated data and handles the defrag op
- Defrag still disabled by default because RT Geometry object don't get recreated yet on defrag of Index/Vertex buffers
- Pool allocator only enabled by default on PC until all possible initial issues as resolved

diff --git a/Engine/Source/Runtime/D3D12RHI/Private/D3D12Adapter.cpp b/Engine/Source/Runtime/D3D12RHI/Private/D3D12Adapter.cpp
--- a/Engine/Source/Runtime/D3D12RHI/Private/D3D12Adapter.cpp
+++ b/Engine/Source/Runtime/D3D12RHI/Private/D3D12Adapter.cpp
@@ -869,7 +869,7 @@ void FD3D12Adapter::InitializeDevices()
 			UploadHeapAllocator[GPUIndex] = new FD3D12DynamicHeapAllocator(this,
 				Devices[GPUIndex],
 				Name,
-				FD3D12BuddyAllocator::EAllocationStrategy::kManualSubAllocation,
+				EResourceAllocationStrategy::kManualSubAllocation,
 				DEFAULT_CONTEXT_UPLOAD_POOL_MAX_ALLOC_SIZE,
 				DEFAULT_CONTEXT_UPLOAD_POOL_SIZE,
 				DEFAULT_CONTEXT_UPLOAD_POOL_ALIGNMENT);
diff --git a/Engine/Source/Runtime/D3D12RHI/Private/D3D12Adapter.h b/Engine/Source/Runtime/D3D12RHI/Private/D3D12Adapter.h
--- a/Engine/Source/Runtime/D3D12RHI/Private/D3D12Adapter.h
+++ b/Engine/Source/Runtime/D3D12RHI/Private/D3D12Adapter.h
@@ -329,7 +329,7 @@ protected:
 		D3D12_RESOURCE_STATES InCreateState,
 		FRHIResourceCreateInfo& CreateInfo,
 		uint32 Alignment,
-		FD3D12TransientResource& TransientResource,
+		FD3D12Buffer* Buffer,
 		FD3D12ResourceLocation& ResourceLocation);
 
 	// Creates default root and execute indirect signatures
diff --git a/Engine/Source/Runtime/D3D12RHI/Private/D3D12Allocation.cpp b/Engine/Source/Runtime/D3D12RHI/Private/D3D12Allocation.cpp
--- a/Engine/Source/Runtime/D3D12RHI/Private/D3D12Allocation.cpp
+++ b/Engine/Source/Runtime/D3D12RHI/Private/D3D12Allocation.cpp
@@ -38,6 +38,37 @@ static FAutoConsoleVariableRef CVarD3D12ReadOnlyTextureAllocatorMaxPoolSize(
 	ECVF_ReadOnly);
 #endif
 
+#if USE_POOL_ALLOCATOR
+static int32 GD3D12PoolAllocatorVRAMPoolSize = 32 * 1024 * 1024;
+static FAutoConsoleVariableRef CVarD3D12PoolAllocatorVRAMPoolSize(
+	TEXT("d3d12.PoolAllocator.VRAMPoolSize"),
+	GD3D12PoolAllocatorVRAMPoolSize,
+	TEXT("Pool size of a single VRAM memory pool (default 32MB)"),
+	ECVF_ReadOnly);
+
+static int32 GD3D12PoolAllocatorVRAMMaxAllocationSize = 16 * 1024 * 1024;
+static FAutoConsoleVariableRef CVarD3D12PoolAllocatorMaxAllocationSize(
+	TEXT("d3d12.PoolAllocator.MaxAllocationSize"),
+	GD3D12PoolAllocatorVRAMMaxAllocationSize,
+	TEXT("Maximum size of a single allocation in the VRAM pool allocator (default 16MB)"),
+	ECVF_ReadOnly);
+#endif // USE_POOL_ALLOCATOR
+
+// NOTE: buffer pool defrag is disabled by default because FD3D12RayTracingGeometry is currently not notified
+//       yet of defragged index and vertex buffers used by the the BVH data. Without RT it works fine.
+static int32 GD3D12VRAMBufferPoolDefrag = 0;
+static FAutoConsoleVariableRef CVarD3D12D3D12VRAMBufferPoolDefrag(
+	TEXT("d3d12.VRAMBufferPoolDefrag"),
+	GD3D12VRAMBufferPoolDefrag,
+	TEXT("Defrag the VRAM buffer pool (disabled by default)"),
+	ECVF_ReadOnly);
+
+static int32 GD3D12VRAMBufferPoolDefragMaxCopySizePerFrame = 32 * 1024 * 1024;
+static FAutoConsoleVariableRef CVarD3D12D3D12VRAMBufferPoolDefragMaxCopySizePerFrame(
+	TEXT("d3d12.VRAMBufferPoolDefrag.MaxCopySizePerFrame"),
+	GD3D12VRAMBufferPoolDefragMaxCopySizePerFrame,
+	TEXT("Max amount of data to copy during defragmentation in a single frame (default 32MB)"),
+	ECVF_ReadOnly);
 
 #if D3D12RHI_SEGLIST_ALLOC_TRACK_WASTAGE
 static int32 GD3D12SegListTrackLeaks = 0;
@@ -71,7 +102,7 @@ namespace ED3D12AllocatorID
 //-----------------------------------------------------------------------------
 FD3D12ResourceAllocator::FD3D12ResourceAllocator(FD3D12Device* ParentDevice,
 	FRHIGPUMask VisibleNodes,
-	const FInitConfig& InInitConfig,
+	const FD3D12ResourceInitConfig& InInitConfig,
 	const FString& Name,
 	uint32 MaxSizeForPooling)
 	: FD3D12DeviceChild(ParentDevice)
@@ -100,9 +131,9 @@ FD3D12ResourceAllocator::~FD3D12ResourceAllocator()
 
 FD3D12BuddyAllocator::FD3D12BuddyAllocator(FD3D12Device* ParentDevice,
 	FRHIGPUMask VisibleNodes,
-	const FInitConfig& InInitConfig,
+	const FD3D12ResourceInitConfig& InInitConfig,
 	const FString& Name,
-	EAllocationStrategy InAllocationStrategy,
+	EResourceAllocationStrategy InAllocationStrategy,
 	uint32 MaxSizeForPooling,
 	uint32 InMaxBlockSize,
 	uint32 InMinBlockSize)
@@ -130,7 +161,7 @@ void FD3D12BuddyAllocator::Initialize()
 	FD3D12Device* Device = GetParentDevice();
 	FD3D12Adapter* Adapter = Device->GetParentAdapter();
 
-	if (AllocationStrategy == EAllocationStrategy::kPlacedResource)
+	if (AllocationStrategy == EResourceAllocationStrategy::kPlacedResource)
 	{
 		D3D12_HEAP_PROPERTIES HeapProps = CD3DX12_HEAP_PROPERTIES(InitConfig.HeapType);
 		HeapProps.CreationNodeMask = GetGPUMask().GetNative();
@@ -303,7 +334,7 @@ void FD3D12BuddyAllocator::Allocate(uint32 SizeInBytes, uint32 Alignment, FD3D12
 	ResourceLocation.SetAllocator((FD3D12BaseAllocatorType*)this);
 	ResourceLocation.SetSize(SizeInBytes);
 
-	if (AllocationStrategy == EAllocationStrategy::kManualSubAllocation)
+	if (AllocationStrategy == EResourceAllocationStrategy::kManualSubAllocation)
 	{
 		ResourceLocation.SetOffsetFromBaseOfResource(AlignedOffsetFromResourceBase);
 		ResourceLocation.SetResource(BackingResource);
@@ -399,7 +430,7 @@ void FD3D12BuddyAllocator::DeallocateInternal(RetiredBlock& Block)
 
 	TotalSizeUsed -= Size;
 
-	if (AllocationStrategy == EAllocationStrategy::kPlacedResource)
+	if (AllocationStrategy == EResourceAllocationStrategy::kPlacedResource)
 	{
 		// Release the resource
 		check(Block.PlacedResource != nullptr);
@@ -553,9 +584,9 @@ void FD3D12BuddyAllocator::Reset()
 
 FD3D12MultiBuddyAllocator::FD3D12MultiBuddyAllocator(FD3D12Device* ParentDevice,
 	FRHIGPUMask VisibleNodes,
-	const FInitConfig& InInitConfig,
+	const FD3D12ResourceInitConfig& InInitConfig,
 	const FString& Name,
-	FD3D12BuddyAllocator::EAllocationStrategy InAllocationStrategy,
+	EResourceAllocationStrategy InAllocationStrategy,
 	uint32 InMaxAllocationSize,
 	uint32 InDefaultPoolSize,
 	uint32 InMinBlockSize)
@@ -676,7 +707,7 @@ void FD3D12MultiBuddyAllocator::Reset()
 //-----------------------------------------------------------------------------
 FD3D12BucketAllocator::FD3D12BucketAllocator(FD3D12Device* ParentDevice,
 	FRHIGPUMask VisibleNodes,
-	const FInitConfig& InInitConfig,
+	const FD3D12ResourceInitConfig& InInitConfig,
 	const FString& Name,
 	uint64 InBlockRetentionFrameCount) :
 	FD3D12ResourceAllocator(ParentDevice, VisibleNodes, InitConfig, Name, 32 * 1024 * 1024),
@@ -888,16 +919,16 @@ void FD3D12BucketAllocator::Reset()
 //	Dynamic Buffer Allocator
 //-----------------------------------------------------------------------------
 
-FD3D12DynamicHeapAllocator::FD3D12DynamicHeapAllocator(FD3D12Adapter* InParent, FD3D12Device* InParentDevice, const FString& InName, FD3D12BuddyAllocator::EAllocationStrategy InAllocationStrategy,
+FD3D12DynamicHeapAllocator::FD3D12DynamicHeapAllocator(FD3D12Adapter* InParent, FD3D12Device* InParentDevice, const FString& InName, EResourceAllocationStrategy InAllocationStrategy,
 	uint32 InMaxSizeForPooling,
 	uint32 InMaxBlockSize,
 	uint32 InMinBlockSize)
 	: FD3D12AdapterChild(InParent)
 	, FD3D12MultiNodeGPUObject(InParentDevice->GetGPUMask(), FRHIGPUMask::All()), // Dynamic heaps are upload memory, thus they can be trivially visibile to all GPUs
-#ifdef USE_BUCKET_ALLOCATOR
+#if USE_BUCKET_ALLOCATOR
 	Allocator(InParentDevice,
 		GetVisibilityMask(),
-		FD3D12ResourceAllocator::FInitConfig
+		FD3D12ResourceInitConfig
 		{
 			D3D12_HEAP_TYPE_UPLOAD,
 			D3D12_HEAP_FLAG_ALLOW_ONLY_BUFFERS,
@@ -909,7 +940,7 @@ FD3D12DynamicHeapAllocator::FD3D12DynamicHeapAllocator(FD3D12Adapter* InParent,
 #else
 	Allocator(InParentDevice,
 		GetVisibilityMask(),
-		FD3D12ResourceAllocator::FInitConfig
+		FD3D12ResourceInitConfig
 		{
 			D3D12_HEAP_TYPE_UPLOAD,
 			D3D12_HEAP_FLAG_ALLOW_ONLY_BUFFERS,
@@ -982,9 +1013,9 @@ void FD3D12DynamicHeapAllocator::Destroy()
 //	Default Buffer Allocator
 //-----------------------------------------------------------------------------
 
-FD3D12ResourceAllocator::FInitConfig FD3D12DefaultBufferPool::GetResourceAllocatorInitConfig(D3D12_HEAP_TYPE InHeapType, D3D12_RESOURCE_FLAGS InResourceFlags, EBufferUsageFlags InBufferUsage)
+FD3D12ResourceInitConfig FD3D12DefaultBufferPool::GetResourceAllocatorInitConfig(D3D12_HEAP_TYPE InHeapType, D3D12_RESOURCE_FLAGS InResourceFlags, EBufferUsageFlags InBufferUsage)
 {
-	FD3D12ResourceAllocator::FInitConfig InitConfig;
+	FD3D12ResourceInitConfig InitConfig;
 	InitConfig.HeapType = InHeapType;
 	InitConfig.ResourceFlags = InResourceFlags;
 
@@ -1025,7 +1056,7 @@ FD3D12ResourceAllocator::FInitConfig FD3D12DefaultBufferPool::GetResourceAllocat
 }
 
 
-FD3D12BuddyAllocator::EAllocationStrategy FD3D12DefaultBufferPool::GetBuddyAllocatorAllocationStrategy(D3D12_RESOURCE_FLAGS InResourceFlags, ED3D12ResourceStateMode InResourceStateMode)
+EResourceAllocationStrategy FD3D12DefaultBufferPool::GetResourceAllocationStrategy(D3D12_RESOURCE_FLAGS InResourceFlags, ED3D12ResourceStateMode InResourceStateMode)
 {
 	// Does the resource need state tracking and transitions
 	ED3D12ResourceStateMode ResourceStateMode = InResourceStateMode;
@@ -1035,7 +1066,7 @@ FD3D12BuddyAllocator::EAllocationStrategy FD3D12DefaultBufferPool::GetBuddyAlloc
 	}
 
 	// multi state resource need to placed because each allocation can be in a different state
-	return (ResourceStateMode == ED3D12ResourceStateMode::MultiState) ?	FD3D12BuddyAllocator::EAllocationStrategy::kPlacedResource : FD3D12BuddyAllocator::EAllocationStrategy::kManualSubAllocation;
+	return (ResourceStateMode == ED3D12ResourceStateMode::MultiState) ? EResourceAllocationStrategy::kPlacedResource : EResourceAllocationStrategy::kManualSubAllocation;
 }
 
 
@@ -1054,12 +1085,12 @@ FD3D12DefaultBufferPool::FD3D12DefaultBufferPool(FD3D12Device* InParent, FD3D12A
 
 bool FD3D12DefaultBufferPool::SupportsAllocation(D3D12_HEAP_TYPE InHeapType, D3D12_RESOURCE_FLAGS InResourceFlags, EBufferUsageFlags InBufferUsage, ED3D12ResourceStateMode InResourceStateMode) const
 {
-	FD3D12ResourceAllocator::FInitConfig InitConfig = GetResourceAllocatorInitConfig(InHeapType, InResourceFlags, InBufferUsage);
+	FD3D12ResourceInitConfig InitConfig = GetResourceAllocatorInitConfig(InHeapType, InResourceFlags, InBufferUsage);
 
-#ifdef USE_BUCKET_ALLOCATOR
+#if USE_BUCKET_ALLOCATOR
 	return Allocator->GetInitConfig() == InitConfig;
 #else
-	FD3D12BuddyAllocator::EAllocationStrategy AllocationStrategy = GetBuddyAllocatorAllocationStrategy(InResourceFlags, InResourceStateMode);
+	EResourceAllocationStrategy AllocationStrategy = GetResourceAllocationStrategy(InResourceFlags, InResourceStateMode);
 	return (Allocator->GetInitConfig() == InitConfig && Allocator->GetAllocationStrategy() == AllocationStrategy);
 #endif
 }
@@ -1071,7 +1102,8 @@ void FD3D12DefaultBufferPool::CleanUpAllocations(uint64 FrameLag)
 }
 
 // Grab a buffer from the available buffers or create a new buffer if none are available
-void FD3D12DefaultBufferPool::AllocDefaultResource(D3D12_HEAP_TYPE InHeapType, const D3D12_RESOURCE_DESC& Desc, EBufferUsageFlags InUsage, ED3D12ResourceStateMode InResourceStateMode, D3D12_RESOURCE_STATES InCreateState, FD3D12ResourceLocation& ResourceLocation, uint32 Alignment, const TCHAR* Name)
+void FD3D12DefaultBufferPool::AllocDefaultResource(D3D12_HEAP_TYPE InHeapType, const D3D12_RESOURCE_DESC& Desc, EBufferUsageFlags InUsage, ED3D12ResourceStateMode InResourceStateMode,
+	D3D12_RESOURCE_STATES InCreateState, uint32 Alignment, const TCHAR* Name, FD3D12ResourceLocation& ResourceLocation)
 {
 	FD3D12Device* Device = GetParentDevice();
 	FD3D12Adapter* Adapter = Device->GetParentAdapter();
@@ -1108,10 +1140,10 @@ void FD3D12DefaultBufferPool::AllocDefaultResource(D3D12_HEAP_TYPE InHeapType, c
 
 	if (PoolResource)
 	{
-#ifdef USE_BUCKET_ALLOCATOR
+#if USE_BUCKET_ALLOCATOR
 		const bool bPlacedResource = false;
 #else
-		const bool bPlacedResource = (Allocator->GetAllocationStrategy() == FD3D12BuddyAllocator::EAllocationStrategy::kPlacedResource);
+		const bool bPlacedResource = (Allocator->GetAllocationStrategy() == EResourceAllocationStrategy::kPlacedResource);
 #endif // USE_BUCKET_ALLOCATOR
 
 		// Ensure we're allocating from the correct pool
@@ -1178,14 +1210,27 @@ FD3D12DefaultBufferAllocator::FD3D12DefaultBufferAllocator(FD3D12Device* InParen
 	FMemory::Memset(DefaultBufferPools, 0);
 }
 
-FD3D12DefaultBufferPool* FD3D12DefaultBufferAllocator::CreateBufferPool(D3D12_HEAP_TYPE InHeapType, D3D12_RESOURCE_FLAGS InResourceFlags, EBufferUsageFlags InBufferUsage, ED3D12ResourceStateMode InResourceStateMode)
+FD3D12BufferPool* FD3D12DefaultBufferAllocator::CreateBufferPool(D3D12_HEAP_TYPE InHeapType, D3D12_RESOURCE_FLAGS InResourceFlags, EBufferUsageFlags InBufferUsage, ED3D12ResourceStateMode InResourceStateMode)
 {
 	FD3D12Device* Device = GetParentDevice();
 	FD3D12AllocatorType* Allocator = nullptr;
 
-	FD3D12ResourceAllocator::FInitConfig InitConfig = FD3D12DefaultBufferPool::GetResourceAllocatorInitConfig(InHeapType, InResourceFlags, InBufferUsage);
+	FD3D12ResourceInitConfig InitConfig = FD3D12BufferPool::GetResourceAllocatorInitConfig(InHeapType, InResourceFlags, InBufferUsage);
+
+#if USE_POOL_ALLOCATOR
+
+	const FString Name(L"D3D12 Pool Allocator");
+	EResourceAllocationStrategy AllocationStrategy = FD3D12PoolAllocator::GetResourceAllocationStrategy(InResourceFlags, InResourceStateMode);
+	uint64 PoolSize = InHeapType == D3D12_HEAP_TYPE_READBACK ? READBACK_BUFFER_POOL_DEFAULT_POOL_SIZE : GD3D12PoolAllocatorVRAMPoolSize;
+	uint64 PoolAlignment = (AllocationStrategy == EResourceAllocationStrategy::kPlacedResource) ? MIN_PLACED_RESOURCE_SIZE : 256;
+	uint64 MaxAllocationSize = InHeapType == D3D12_HEAP_TYPE_READBACK ? READBACK_BUFFER_POOL_MAX_ALLOC_SIZE : GD3D12PoolAllocatorVRAMMaxAllocationSize;
+
+	FD3D12BufferPool* NewPool = new FD3D12PoolAllocator(Device, GetVisibilityMask(), InitConfig, Name, AllocationStrategy, PoolSize, PoolAlignment, MaxAllocationSize);
+
+#else // USE_POOL_ALLOCATOR
+
+#if USE_BUCKET_ALLOCATOR
 
-#ifdef USE_BUCKET_ALLOCATOR
 	const FString Name(L"Default Buffer Bucket Allocator");
 	Allocator = new FD3D12BucketAllocator(Device,
 		GetVisibilityMask(),
@@ -1194,10 +1239,10 @@ FD3D12DefaultBufferPool* FD3D12DefaultBufferAllocator::CreateBufferPool(D3D12_HE
 		5);
 #else
 
-	FD3D12BuddyAllocator::EAllocationStrategy AllocationStrategy = FD3D12DefaultBufferPool::GetBuddyAllocatorAllocationStrategy(InResourceFlags, InResourceStateMode);
+	EResourceAllocationStrategy AllocationStrategy = FD3D12DefaultBufferPool::GetResourceAllocationStrategy(InResourceFlags, InResourceStateMode);
 
 	// if placed then 64KB alignment required :(
-	uint32 MinBlockSize = (AllocationStrategy == FD3D12BuddyAllocator::EAllocationStrategy::kPlacedResource) ? MIN_PLACED_BUFFER_SIZE : 16;
+	uint32 MinBlockSize = (AllocationStrategy == EResourceAllocationStrategy::kPlacedResource) ? MIN_PLACED_RESOURCE_SIZE : 16;
 
 	const FString Name(L"Default Buffer Multi Buddy Allocator");
 	Allocator = new FD3D12MultiBuddyAllocator(Device,
@@ -1212,6 +1257,9 @@ FD3D12DefaultBufferPool* FD3D12DefaultBufferAllocator::CreateBufferPool(D3D12_HE
 #endif
 
 	FD3D12DefaultBufferPool* NewPool = new FD3D12DefaultBufferPool(Device, Allocator);
+
+#endif // USE_POOL_ALLOCATOR
+
 	DefaultBufferPools.Add(NewPool);
 	return NewPool;
 }
@@ -1219,11 +1267,11 @@ FD3D12DefaultBufferPool* FD3D12DefaultBufferAllocator::CreateBufferPool(D3D12_HE
 
 bool FD3D12DefaultBufferAllocator::IsPlacedResource(D3D12_RESOURCE_FLAGS InResourceFlags, ED3D12ResourceStateMode InResourceStateMode)
 {
-#ifdef USE_BUCKET_ALLOCATOR
+#if USE_BUCKET_ALLOCATOR
 	return false;
 #else
-	FD3D12BuddyAllocator::EAllocationStrategy AllocationStrategy = FD3D12DefaultBufferPool::GetBuddyAllocatorAllocationStrategy(InResourceFlags, InResourceStateMode);
-	return (AllocationStrategy == FD3D12BuddyAllocator::EAllocationStrategy::kPlacedResource);
+	EResourceAllocationStrategy AllocationStrategy = FD3D12BufferPool::GetResourceAllocationStrategy(InResourceFlags, InResourceStateMode);
+	return (AllocationStrategy == EResourceAllocationStrategy::kPlacedResource);
 #endif
 }
 
@@ -1261,8 +1309,8 @@ void FD3D12DefaultBufferAllocator::AllocDefaultResource(D3D12_HEAP_TYPE InHeapTy
 	ResourceDesc.Flags = ResourceDesc.Flags & (~D3D12_RESOURCE_FLAG_DENY_SHADER_RESOURCE);
 
 	// Do we already have a default pool which support this allocation?
-	FD3D12DefaultBufferPool* BufferPool = nullptr;
-	for (FD3D12DefaultBufferPool* Pool : DefaultBufferPools)
+	FD3D12BufferPool* BufferPool = nullptr;
+	for (FD3D12BufferPool* Pool : DefaultBufferPools)
 	{
 		if (Pool->SupportsAllocation(InHeapType, ResourceDesc.Flags, InBufferUsage, InResourceStateMode))
 		{
@@ -1278,13 +1326,13 @@ void FD3D12DefaultBufferAllocator::AllocDefaultResource(D3D12_HEAP_TYPE InHeapTy
 	}
 
 	// Perform actual allocation
-	BufferPool->AllocDefaultResource(InHeapType, ResourceDesc, InBufferUsage, InResourceStateMode, InCreateState, ResourceLocation, Alignment, Name);
+	BufferPool->AllocDefaultResource(InHeapType, ResourceDesc, InBufferUsage, InResourceStateMode, InCreateState, Alignment, Name, ResourceLocation);
 }
 
 
 void FD3D12DefaultBufferAllocator::FreeDefaultBufferPools()
 {
-	for (FD3D12DefaultBufferPool*& DefaultBufferPool : DefaultBufferPools)
+	for (FD3D12BufferPool*& DefaultBufferPool : DefaultBufferPools)
 	{
 		if (DefaultBufferPool)
 		{
@@ -1297,9 +1345,38 @@ void FD3D12DefaultBufferAllocator::FreeDefaultBufferPools()
 	}
 }
 
+
+void FD3D12DefaultBufferAllocator::BeginFrame()
+{
+#if USE_POOL_ALLOCATOR
+	if (GD3D12VRAMBufferPoolDefrag > 0 && GD3D12VRAMBufferPoolDefragMaxCopySizePerFrame > 0)
+	{
+		uint32 MaxCopySize = GD3D12VRAMBufferPoolDefragMaxCopySizePerFrame;
+		uint32 CopySize = 0;
+		for (FD3D12BufferPool* DefaultBufferPool : DefaultBufferPools)
+		{
+			if (DefaultBufferPool)
+			{
+				DefaultBufferPool->Defrag(MaxCopySize, CopySize);
+			}
+		}
+	}
+
+	FD3D12CommandContext& CommandContext = GetParentDevice()->GetDefaultCommandContext();
+	for (FD3D12BufferPool* DefaultBufferPool : DefaultBufferPools)
+	{
+		if (DefaultBufferPool)
+		{
+			DefaultBufferPool->FlushPendingCopyOps(CommandContext);
+		}
+	}
+#endif // USE_POOL_ALLOCATOR
+}
+
+
 void FD3D12DefaultBufferAllocator::CleanupFreeBlocks(uint64 InFrameLag)
 {
-	for (FD3D12DefaultBufferPool* DefaultBufferPool : DefaultBufferPools)
+	for (FD3D12BufferPool* DefaultBufferPool : DefaultBufferPools)
 	{
 		if (DefaultBufferPool)
 		{
@@ -1318,7 +1395,7 @@ void FD3D12DefaultBufferAllocator::UpdateMemoryStats()
 	uint32 FullPageCount = 0;
 
 #if defined(D3D12RHI_TRACK_DETAILED_STATS)
-	for (FD3D12DefaultBufferPool* DefaultBufferPool : DefaultBufferPools)
+	for (FD3D12BufferPool* DefaultBufferPool : DefaultBufferPools)
 	{
 		if (DefaultBufferPool)
 		{
@@ -1418,7 +1495,7 @@ FD3D12TextureAllocator::FD3D12TextureAllocator(FD3D12Device* Device,
 	D3D12_HEAP_FLAGS Flags) :
 	FD3D12MultiBuddyAllocator(Device,
 		VisibleNodes,
-		FD3D12ResourceAllocator::FInitConfig
+		FD3D12ResourceInitConfig
 		{
 			D3D12_HEAP_TYPE_DEFAULT,
 			Flags | D3D12_HEAP_FLAG_DENY_BUFFERS,
@@ -1426,7 +1503,7 @@ FD3D12TextureAllocator::FD3D12TextureAllocator(FD3D12Device* Device,
 			D3D12_RESOURCE_STATE_GENERIC_READ
 		},
 		Name,
-		FD3D12BuddyAllocator::EAllocationStrategy::kPlacedResource,
+		EResourceAllocationStrategy::kPlacedResource,
 		D3D12_DEFAULT_RESOURCE_PLACEMENT_ALIGNMENT,
 		HeapSize,
 		D3D12_SMALL_RESOURCE_PLACEMENT_ALIGNMENT)
diff --git a/Engine/Source/Runtime/D3D12RHI/Private/D3D12Allocation.h b/Engine/Source/Runtime/D3D12RHI/Private/D3D12Allocation.h
--- a/Engine/Source/Runtime/D3D12RHI/Private/D3D12Allocation.h
+++ b/Engine/Source/Runtime/D3D12RHI/Private/D3D12Allocation.h
@@ -6,6 +6,12 @@ D3D12Allocation.h: A Collection of allocators
 
 #pragma once
 #include "D3D12Resources.h"
+#include "D3D12PoolAllocator.h"
+
+// Enable pool allocator by default only on Windows until all initial issues have been resolved.
+// Xbox is currently using the buddy allocator with a lot smaller pool sizes so the alignment waste is a lot less
+// It also has less of a problem with committed resource allocations
+#define USE_POOL_ALLOCATOR (PLATFORM_WINDOWS)
 
 // Segregated free list texture allocator
 // Description:
@@ -66,36 +72,20 @@ class FD3D12SegList;
 class FD3D12ResourceAllocator : public FD3D12DeviceChild, public FD3D12MultiNodeGPUObject
 {
 public:
-
-	struct FInitConfig
-	{
-		bool operator==(const FInitConfig& InOther) const
-		{
-			return HeapType == InOther.HeapType && 
-				HeapFlags && InOther.HeapFlags && 
-				ResourceFlags == InOther.ResourceFlags && 
-				InitialResourceState == InOther.InitialResourceState;
-		}
-
-		D3D12_HEAP_TYPE HeapType;
-		D3D12_HEAP_FLAGS HeapFlags;
-		D3D12_RESOURCE_FLAGS ResourceFlags;
-		D3D12_RESOURCE_STATES InitialResourceState;
-	};
-
+	
 	FD3D12ResourceAllocator(FD3D12Device* ParentDevice,
 		FRHIGPUMask VisibleNodes,
-		const FInitConfig& InInitConfig,
+		const FD3D12ResourceInitConfig& InInitConfig,
 		const FString& Name,
 		uint32 MaxSizeForPooling);
 	~FD3D12ResourceAllocator();
 
-	const FInitConfig& GetInitConfig() const { return InitConfig; }
+	const FD3D12ResourceInitConfig& GetInitConfig() const { return InitConfig; }
 	const uint32 GetMaximumAllocationSizeForPooling() const { return MaximumAllocationSizeForPooling; }
 
 protected:
 
-	const FInitConfig InitConfig;
+	const FD3D12ResourceInitConfig InitConfig;
 	const FString DebugName;
 	bool Initialized;
 
@@ -125,7 +115,7 @@ protected:
 // Based on reference implementation by MSFT: billkris
 
 // Unfortunately the api restricts the minimum size of a placed buffer resource to 64k
-#define MIN_PLACED_BUFFER_SIZE (64 * 1024)
+#define MIN_PLACED_RESOURCE_SIZE (64 * 1024)
 #define D3D_BUFFER_ALIGNMENT (64 * 1024)
 
 #if defined(D3D12RHI_TRACK_DETAILED_STATS)
@@ -139,27 +129,13 @@ protected:
 
 class FD3D12BuddyAllocator : public FD3D12ResourceAllocator
 {
-public:
-
-	enum class EAllocationStrategy
-	{
-		// This strategy uses Placed Resources to sub-allocate a buffer out of an underlying ID3D12Heap.
-		// The benefit of this is that each buffer can have it's own resource state and can be treated
-		// as any other buffer. The downside of this strategy is the API limitation which enforces
-		// the minimum buffer size to 64k leading to large internal fragmentation in the allocator
-		kPlacedResource,
-		// The alternative is to manually sub-allocate out of a single large buffer which allows block
-		// allocation granularity down to 1 byte. However, this strategy is only really valid for buffers which
-		// will be treated as read-only after their creation (i.e. most Index and Vertex buffers). This 
-		// is because the underlying resource can only have one state at a time.
-		kManualSubAllocation
-	};
+public:	
 
 	FD3D12BuddyAllocator(FD3D12Device* ParentDevice,
 		FRHIGPUMask VisibleNodes,
-		const FInitConfig& InInitConfig,
+		const FD3D12ResourceInitConfig& InInitConfig,
 		const FString& Name,
-		EAllocationStrategy InAllocationStrategy,
+		EResourceAllocationStrategy InAllocationStrategy,
 		uint32 MaxSizeForPooling,
 		uint32 InMaxBlockSize,
 		uint32 InMinBlockSize);
@@ -190,7 +166,7 @@ public:
 	inline uint32 GetTotalSizeUsed() const { return TotalSizeUsed; }
 	inline uint64 GetAllocationOffsetInBytes(const FD3D12BuddyAllocatorPrivateData& AllocatorPrivateData) const { return uint64(AllocatorPrivateData.Offset) * MinBlockSize; }
 
-	inline FD3D12Heap* GetBackingHeap() { check(AllocationStrategy == EAllocationStrategy::kPlacedResource); return BackingHeap.GetReference(); }
+	inline FD3D12Heap* GetBackingHeap() { check(AllocationStrategy == EResourceAllocationStrategy::kPlacedResource); return BackingHeap.GetReference(); }
 
 	inline bool IsOwner(FD3D12ResourceLocation& ResourceLocation)
 	{
@@ -201,7 +177,7 @@ protected:
 
 	const uint32 MaxBlockSize;
 	const uint32 MinBlockSize;
-	const FD3D12BuddyAllocator::EAllocationStrategy AllocationStrategy;
+	const EResourceAllocationStrategy AllocationStrategy;
 
 	TRefCountPtr<FD3D12Resource> BackingResource;
 	TRefCountPtr<FD3D12Heap> BackingHeap;
@@ -267,15 +243,15 @@ public:
 
 	FD3D12MultiBuddyAllocator(FD3D12Device* ParentDevice,
 		FRHIGPUMask VisibleNodes,
-		const FInitConfig& InInitConfig,
+		const FD3D12ResourceInitConfig& InInitConfig,
 		const FString& Name,
-		FD3D12BuddyAllocator::EAllocationStrategy InAllocationStrategy,
+		EResourceAllocationStrategy InAllocationStrategy,
 		uint32 InMaxAllocationSize,
 		uint32 InDefaultPoolSize,
 		uint32 InMinBlockSize);
 	~FD3D12MultiBuddyAllocator();
 
-	const FD3D12BuddyAllocator::EAllocationStrategy GetAllocationStrategy() const { return AllocationStrategy; }
+	const EResourceAllocationStrategy GetAllocationStrategy() const { return AllocationStrategy; }
 
 	bool TryAllocate(uint32 SizeInBytes, uint32 Alignment, FD3D12ResourceLocation& ResourceLocation);
 
@@ -298,7 +274,7 @@ protected:
 
 	FD3D12BuddyAllocator* CreateNewAllocator(uint32 InMinSizeInBytes);
 
-	const FD3D12BuddyAllocator::EAllocationStrategy AllocationStrategy;
+	const EResourceAllocationStrategy AllocationStrategy;
 	const uint32 MinBlockSize;
 	const uint32 DefaultPoolSize;
 
@@ -316,7 +292,7 @@ public:
 
 	FD3D12BucketAllocator(FD3D12Device* ParentDevice,
 		FRHIGPUMask VisibleNodes,
-		const FInitConfig& InInitConfig,
+		const FD3D12ResourceInitConfig& InInitConfig,
 		const FString& Name,
 		uint64 InBlockRetentionFrameCount);
 
@@ -372,7 +348,7 @@ private:
 	const uint64 BlockRetentionFrameCount;
 };
 
-#ifdef USE_BUCKET_ALLOCATOR
+#if USE_BUCKET_ALLOCATOR
 typedef FD3D12BucketAllocator FD3D12AllocatorType;
 #else
 typedef FD3D12MultiBuddyAllocator FD3D12AllocatorType;
@@ -387,7 +363,7 @@ class FD3D12DynamicHeapAllocator : public FD3D12AdapterChild, public FD3D12Multi
 {
 public:
 
-	FD3D12DynamicHeapAllocator(FD3D12Adapter* InParent, FD3D12Device* InParentDevice, const FString& InName, FD3D12BuddyAllocator::EAllocationStrategy InAllocationStrategy,
+	FD3D12DynamicHeapAllocator(FD3D12Adapter* InParent, FD3D12Device* InParentDevice, const FString& InName, EResourceAllocationStrategy InAllocationStrategy,
 		uint32 InMaxSizeForPooling,
 		uint32 InMaxBlockSize,
 		uint32 InMinBlockSize);
@@ -416,17 +392,24 @@ public:
 	~FD3D12DefaultBufferPool() { delete Allocator; }
 
 	bool SupportsAllocation(D3D12_HEAP_TYPE InHeapType, D3D12_RESOURCE_FLAGS InResourceFlags, EBufferUsageFlags InBufferUsage, ED3D12ResourceStateMode InResourceStateMode) const;
-	void AllocDefaultResource(D3D12_HEAP_TYPE InHeapType, const D3D12_RESOURCE_DESC& Desc, EBufferUsageFlags InBufferUsage, ED3D12ResourceStateMode InResourceStateMode, D3D12_RESOURCE_STATES InCreateState, FD3D12ResourceLocation& ResourceLocation, uint32 Alignment, const TCHAR* Name);
+	void AllocDefaultResource(D3D12_HEAP_TYPE InHeapType, const D3D12_RESOURCE_DESC& InDesc, EBufferUsageFlags InBufferUsage, ED3D12ResourceStateMode InResourceStateMode,
+		D3D12_RESOURCE_STATES InCreateState, uint32 InAlignment, const TCHAR* InName, FD3D12ResourceLocation& ResourceLocation);
 	void CleanUpAllocations(uint64 InFrameLag);
 	void UpdateMemoryStats(uint32& IOMemoryAllocated, uint32& IOMemoryUsed, uint32& IOMemoryFree, uint32& IOAlignmentWaste, uint32& IOAllocatedPageCount, uint32& IOFullPageCount);
 
-	static FD3D12ResourceAllocator::FInitConfig GetResourceAllocatorInitConfig(D3D12_HEAP_TYPE InHeapType, D3D12_RESOURCE_FLAGS InResourceFlags, EBufferUsageFlags InBufferUsage);
-	static FD3D12BuddyAllocator::EAllocationStrategy GetBuddyAllocatorAllocationStrategy(D3D12_RESOURCE_FLAGS InResourceFlags, ED3D12ResourceStateMode InResourceStateMode);
+	static FD3D12ResourceInitConfig GetResourceAllocatorInitConfig(D3D12_HEAP_TYPE InHeapType, D3D12_RESOURCE_FLAGS InResourceFlags, EBufferUsageFlags InBufferUsage);
+	static EResourceAllocationStrategy GetResourceAllocationStrategy(D3D12_RESOURCE_FLAGS InResourceFlags, ED3D12ResourceStateMode InResourceStateMode);
 
 private:
 	FD3D12AllocatorType* Allocator;
 };
 
+#if USE_POOL_ALLOCATOR
+typedef FD3D12PoolAllocator FD3D12BufferPool;
+#else
+typedef FD3D12DefaultBufferPool FD3D12BufferPool;
+#endif
+
 // FD3D12DefaultBufferAllocator
 //
 class FD3D12DefaultBufferAllocator : public FD3D12DeviceChild, public FD3D12MultiNodeGPUObject
@@ -437,7 +420,8 @@ public:
 	// Grab a buffer from the available buffers or create a new buffer if none are available
 	void AllocDefaultResource(D3D12_HEAP_TYPE InHeapType, const D3D12_RESOURCE_DESC& pDesc, EBufferUsageFlags InBufferUsage, ED3D12ResourceStateMode InResourceStateMode, D3D12_RESOURCE_STATES InCreateState, FD3D12ResourceLocation& ResourceLocation, uint32 Alignment, const TCHAR* Name);
 	void FreeDefaultBufferPools();
-	void CleanupFreeBlocks(uint64 InFrameLag);
+	void BeginFrame();
+	void CleanupFreeBlocks(uint64 InFrameLag);	
 	void UpdateMemoryStats();
 
 	static bool IsPlacedResource(D3D12_RESOURCE_FLAGS InResourceFlags, ED3D12ResourceStateMode InResourceStateMode);
@@ -445,9 +429,9 @@ public:
 
 private:
 
-	FD3D12DefaultBufferPool* CreateBufferPool(D3D12_HEAP_TYPE InHeapType, D3D12_RESOURCE_FLAGS InResourceFlags, EBufferUsageFlags InBufferUsage, ED3D12ResourceStateMode InResourceStateMode);
+	FD3D12BufferPool* CreateBufferPool(D3D12_HEAP_TYPE InHeapType, D3D12_RESOURCE_FLAGS InResourceFlags, EBufferUsageFlags InBufferUsage, ED3D12ResourceStateMode InResourceStateMode);
 
-	TArray<FD3D12DefaultBufferPool*> DefaultBufferPools;
+	TArray<FD3D12BufferPool*> DefaultBufferPools;
 };
 
 //-----------------------------------------------------------------------------
diff --git a/Engine/Source/Runtime/D3D12RHI/Private/D3D12Buffer.cpp b/Engine/Source/Runtime/D3D12RHI/Private/D3D12Buffer.cpp
--- a/Engine/Source/Runtime/D3D12RHI/Private/D3D12Buffer.cpp
+++ b/Engine/Source/Runtime/D3D12RHI/Private/D3D12Buffer.cpp
@@ -128,6 +128,9 @@ struct FD3D12RHICommandInitializeBuffer final : public FRHICommand<FD3D12RHIComm
 
 				hCommandList.UpdateResidency(SrcResourceLoc.GetResource());
 			}
+
+			// Buffer is now written and ready, so unlock the block (locked after creation and can be defragmented if needed)
+			CurrentBuffer->ResourceLocation.UnlockPoolData();
 		}
 	}
 };
@@ -140,7 +143,7 @@ void FD3D12Adapter::AllocateBuffer(FD3D12Device* Device,
 	D3D12_RESOURCE_STATES InCreateState,
 	FRHIResourceCreateInfo& CreateInfo,
 	uint32 Alignment,
-	FD3D12TransientResource& TransientResource,
+	FD3D12Buffer* Buffer,
 	FD3D12ResourceLocation& ResourceLocation)
 {
 	// Explicitly check that the size is nonzero before allowing CreateBuffer to opaquely fail.
@@ -167,6 +170,7 @@ void FD3D12Adapter::AllocateBuffer(FD3D12Device* Device,
 	else
 	{
 		Device->GetDefaultBufferAllocator().AllocDefaultResource(D3D12_HEAP_TYPE_DEFAULT, InDesc, (EBufferUsageFlags)InUsage, InResourceStateMode, InCreateState, ResourceLocation, Alignment, CreateInfo.DebugName);
+		ResourceLocation.SetOwner(Buffer);
 		check(ResourceLocation.GetSize() == Size);
 	}
 }
@@ -206,7 +210,7 @@ FD3D12Buffer* FD3D12Adapter::CreateRHIBuffer(FRHICommandListImmediate* RHICmdLis
 
 			if (Device->GetGPUIndex() == FirstGPUIndex)
 			{
-				AllocateBuffer(Device, InDesc, Size, InUsage, InResourceStateMode, InitialState, CreateInfo, Alignment, *NewBuffer, NewBuffer->ResourceLocation);
+				AllocateBuffer(Device, InDesc, Size, InUsage, InResourceStateMode, InitialState, CreateInfo, Alignment, NewBuffer, NewBuffer->ResourceLocation);
 				NewBuffer0 = NewBuffer;
 			}
 			else
@@ -228,7 +232,13 @@ FD3D12Buffer* FD3D12Adapter::CreateRHIBuffer(FRHICommandListImmediate* RHICmdLis
 			FD3D12Buffer* NewBuffer = new FD3D12Buffer(Device, Size, InUsage, Stride);
 			NewBuffer->BufferAlignment = Alignment;
 
-			AllocateBuffer(Device, InDesc, Size, InUsage, InResourceStateMode, CreateState, CreateInfo, Alignment, *NewBuffer, NewBuffer->ResourceLocation);
+			AllocateBuffer(Device, InDesc, Size, InUsage, InResourceStateMode, CreateState, CreateInfo, Alignment, NewBuffer, NewBuffer->ResourceLocation);
+
+			// Unlock immediately if no initial data
+			if (CreateInfo.ResourceArray == nullptr)
+			{
+				NewBuffer->ResourceLocation.UnlockPoolData();
+			}
 
 			return NewBuffer;
 		});
@@ -310,13 +320,7 @@ FD3D12Buffer* FD3D12Adapter::CreateRHIBuffer(FRHICommandListImmediate* RHICmdLis
 void FD3D12Buffer::Rename(FD3D12ResourceLocation& NewLocation)
 {
 	FD3D12ResourceLocation::TransferOwnership(ResourceLocation, NewLocation);
-
-	FScopeLock Lock(&DynamicSRVsCS);
-	for (FD3D12BaseShaderResourceView* DynamicSRVBase : DynamicSRVs)
-	{
-		FD3D12ShaderResourceView* DynamicSRV = static_cast<FD3D12ShaderResourceView*>(DynamicSRVBase);
-		DynamicSRV->Rename(ResourceLocation);
-	}
+	RecreateViews(ResourceLocation);
 }
 
 void FD3D12Buffer::RenameLDAChain(FD3D12ResourceLocation& NewLocation)
@@ -335,13 +339,7 @@ void FD3D12Buffer::RenameLDAChain(FD3D12ResourceLocation& NewLocation)
 		for (auto NextBuffer = ++FLinkedObjectIterator(this); NextBuffer; ++NextBuffer)
 		{
 			FD3D12ResourceLocation::ReferenceNode(NextBuffer->GetParentDevice(), NextBuffer->ResourceLocation, ResourceLocation);
-
-			FScopeLock Lock(&NextBuffer->DynamicSRVsCS);
-			for (FD3D12BaseShaderResourceView* DynamicSRVBase : NextBuffer->DynamicSRVs)
-			{
-				FD3D12ShaderResourceView* DynamicSRV = static_cast<FD3D12ShaderResourceView*>(DynamicSRVBase);
-				DynamicSRV->Rename(NextBuffer->ResourceLocation);
-			}
+			NextBuffer->RecreateViews(NextBuffer->ResourceLocation);
 		}
 	}
 }
@@ -362,7 +360,7 @@ void FD3D12Buffer::ReleaseUnderlyingResource()
 	{
 		check(!NextBuffer->LockedData.bLocked && NextBuffer->ResourceLocation.IsValid());
 		NextBuffer->ResourceLocation.Clear();
-		NextBuffer->RemoveAllDynamicSRVs();
+		NextBuffer->RemoveAllViews();
 	}
 }
 
diff --git a/Engine/Source/Runtime/D3D12RHI/Private/D3D12CommandContext.cpp b/Engine/Source/Runtime/D3D12RHI/Private/D3D12CommandContext.cpp
--- a/Engine/Source/Runtime/D3D12RHI/Private/D3D12CommandContext.cpp
+++ b/Engine/Source/Runtime/D3D12RHI/Private/D3D12CommandContext.cpp
@@ -489,6 +489,8 @@ void FD3D12CommandContextBase::RHIBeginFrame()
 		SamplerHeap.ToggleDescriptorTablesDirtyFlag(false);
 
 		Device->GetGPUProfiler().BeginFrame(ParentAdapter->GetOwningRHI());
+
+		Device->GetDefaultBufferAllocator().BeginFrame();
 	}
 }
 
diff --git a/Engine/Source/Runtime/D3D12RHI/Private/D3D12PoolAllocator.cpp b/Engine/Source/Runtime/D3D12RHI/Private/D3D12PoolAllocator.cpp
new file mode 100644
--- /dev/null
+++ b/Engine/Source/Runtime/D3D12RHI/Private/D3D12PoolAllocator.cpp
@@ -0,0 +1,559 @@
+// Copyright Epic Games, Inc. All Rights Reserved.
+
+#include "D3D12RHIPrivate.h"
+#include "D3D12PoolAllocator.h"
+
+#ifndef NEEDS_D3D12_INDIRECT_ARGUMENT_HEAP_WORKAROUND
+#define NEEDS_D3D12_INDIRECT_ARGUMENT_HEAP_WORKAROUND 0
+#endif
+
+//-----------------------------------------------------------------------------
+//	FD3D12MemoryPool
+//-----------------------------------------------------------------------------
+
+FD3D12MemoryPool::FD3D12MemoryPool(FD3D12Device* ParentDevice, FRHIGPUMask VisibleNodes, const FD3D12ResourceInitConfig& InInitConfig, const FString& InName,
+	EResourceAllocationStrategy InAllocationStrategy, int16 InPoolIndex, uint64 InPoolSize, uint32 InPoolAlignment)
+	: FRHIMemoryPool(InPoolIndex, InPoolSize, InPoolAlignment), FD3D12DeviceChild(ParentDevice), FD3D12MultiNodeGPUObject(ParentDevice->GetGPUMask(), VisibleNodes)
+	, InitConfig(InInitConfig), Name(InName), AllocationStrategy(InAllocationStrategy), LastUsedFrameFence(0)
+{
+}
+
+
+FD3D12MemoryPool::~FD3D12MemoryPool()
+{
+	Destroy();
+}
+
+
+void FD3D12MemoryPool::Init()
+{
+	FD3D12Device* Device = GetParentDevice();
+	FD3D12Adapter* Adapter = Device->GetParentAdapter();
+
+	if (AllocationStrategy == EResourceAllocationStrategy::kPlacedResource)
+	{
+		// Alignment should be either 4K or 64K for places resources
+		check(PoolAlignment == D3D12_SMALL_RESOURCE_PLACEMENT_ALIGNMENT || PoolAlignment == D3D12_DEFAULT_RESOURCE_PLACEMENT_ALIGNMENT);
+
+		D3D12_HEAP_PROPERTIES HeapProps = CD3DX12_HEAP_PROPERTIES(InitConfig.HeapType);
+		HeapProps.CreationNodeMask = GetGPUMask().GetNative();
+		HeapProps.VisibleNodeMask = GetVisibilityMask().GetNative();
+
+		D3D12_HEAP_DESC Desc = {};
+		Desc.SizeInBytes = PoolSize;
+		Desc.Properties = HeapProps;
+		Desc.Alignment = 0;
+		Desc.Flags = InitConfig.HeapFlags;
+#if PLATFORM_WINDOWS
+		if (Adapter->IsHeapNotZeroedSupported())
+		{
+			Desc.Flags |= D3D12_HEAP_FLAG_CREATE_NOT_ZEROED;
+		}
+#endif
+
+		ID3D12Heap* Heap = nullptr;
+		{
+			LLM_PLATFORM_SCOPE(ELLMTag::GraphicsPlatform);
+
+			// we are tracking allocations ourselves, so don't let XMemAlloc track these as well
+			LLM_SCOPED_PAUSE_TRACKING_FOR_TRACKER(ELLMTracker::Default, ELLMAllocType::System);
+			VERIFYD3D12RESULT(Adapter->GetD3DDevice()->CreateHeap(&Desc, IID_PPV_ARGS(&Heap)));
+		}
+		SetName(Heap, L"LinkListAllocator Backing Heap");
+
+		BackingHeap = new FD3D12Heap(GetParentDevice(), GetVisibilityMask());
+		BackingHeap->SetHeap(Heap);
+
+		// Only track resources that cannot be accessed on the CPU.
+		if (IsCPUInaccessible(InitConfig.HeapType))
+		{
+			BackingHeap->BeginTrackingResidency(Desc.SizeInBytes);
+		}
+	}
+	else
+	{
+		{
+			LLM_SCOPED_PAUSE_TRACKING_FOR_TRACKER(ELLMTracker::Default, ELLMAllocType::System);
+			const D3D12_HEAP_PROPERTIES HeapProps = CD3DX12_HEAP_PROPERTIES(InitConfig.HeapType, GetGPUMask().GetNative(), GetVisibilityMask().GetNative());
+			VERIFYD3D12RESULT(Adapter->CreateBuffer(HeapProps, GetGPUMask(), InitConfig.InitialResourceState, ED3D12ResourceStateMode::SingleState, InitConfig.InitialResourceState, PoolSize, BackingResource.GetInitReference(), TEXT("Resource Allocator Underlying Buffer"), InitConfig.ResourceFlags));
+		}
+
+		if (IsCPUWritable(InitConfig.HeapType))
+		{
+			BackingResource->Map();
+		}
+	}	
+
+	FRHIMemoryPool::Init();
+}
+
+
+void FD3D12MemoryPool::Destroy()
+{
+	LLM_SCOPED_PAUSE_TRACKING_FOR_TRACKER(ELLMTracker::Default, ELLMAllocType::System);
+
+	FRHIMemoryPool::Destroy();
+
+	if (BackingResource)
+	{
+		ensure(BackingResource->GetRefCount() == 1 || GNumExplicitGPUsForRendering > 1);
+		BackingResource = nullptr;
+	}
+
+	if (BackingHeap)
+	{
+		BackingHeap->Destroy();
+	}
+}
+
+
+//-----------------------------------------------------------------------------
+//	FD3D12PoolAllocator
+//-----------------------------------------------------------------------------
+
+
+FD3D12ResourceInitConfig FD3D12PoolAllocator::GetResourceAllocatorInitConfig(D3D12_HEAP_TYPE InHeapType, D3D12_RESOURCE_FLAGS InResourceFlags, EBufferUsageFlags InBufferUsage)
+{
+	FD3D12ResourceInitConfig InitConfig;
+	InitConfig.HeapType = InHeapType;
+	InitConfig.ResourceFlags = InResourceFlags;
+
+#if D3D12_RHI_RAYTRACING
+	// Setup initial resource state depending on the requested buffer flags
+	if (EnumHasAnyFlags(InBufferUsage, BUF_AccelerationStructure))
+	{
+		// should only have this flag and no other flags
+		check(InBufferUsage == BUF_AccelerationStructure);
+		InitConfig.InitialResourceState = D3D12_RESOURCE_STATE_RAYTRACING_ACCELERATION_STRUCTURE;
+	}
+	else
+#endif // D3D12_RHI_RAYTRACING
+		if (InitConfig.HeapType == D3D12_HEAP_TYPE_READBACK)
+		{
+			InitConfig.InitialResourceState = D3D12_RESOURCE_STATE_COPY_DEST;
+		}
+		else if (EnumHasAnyFlags(InBufferUsage, BUF_UnorderedAccess))
+		{
+			check(InResourceFlags & D3D12_RESOURCE_FLAG_ALLOW_UNORDERED_ACCESS);
+			InitConfig.InitialResourceState = D3D12_RESOURCE_STATE_UNORDERED_ACCESS;
+		}
+		else
+		{
+			InitConfig.InitialResourceState = D3D12_RESOURCE_STATE_GENERIC_READ;
+		}
+
+	InitConfig.HeapFlags = D3D12_HEAP_FLAG_ALLOW_ONLY_BUFFERS;
+	if (EnumHasAnyFlags(InBufferUsage, BUF_DrawIndirect))
+	{
+		check(InResourceFlags & D3D12_RESOURCE_FLAG_ALLOW_UNORDERED_ACCESS);
+#if !NEEDS_D3D12_INDIRECT_ARGUMENT_HEAP_WORKAROUND
+		InitConfig.HeapFlags |= D3D12RHI_HEAP_FLAG_ALLOW_INDIRECT_BUFFERS;
+#endif
+	}
+
+	return InitConfig;
+}
+
+
+EResourceAllocationStrategy FD3D12PoolAllocator::GetResourceAllocationStrategy(D3D12_RESOURCE_FLAGS InResourceFlags, ED3D12ResourceStateMode InResourceStateMode)
+{
+	// Does the resource need state tracking and transitions
+	ED3D12ResourceStateMode ResourceStateMode = InResourceStateMode;
+	if (ResourceStateMode == ED3D12ResourceStateMode::Default)
+	{
+		ResourceStateMode = (InResourceFlags & D3D12_RESOURCE_FLAG_ALLOW_UNORDERED_ACCESS) ? ED3D12ResourceStateMode::MultiState : ED3D12ResourceStateMode::SingleState;
+	}
+
+	// multi state resource need to placed because each allocation can be in a different state
+	return (ResourceStateMode == ED3D12ResourceStateMode::MultiState) ? EResourceAllocationStrategy::kPlacedResource : EResourceAllocationStrategy::kManualSubAllocation;
+}
+
+
+FD3D12PoolAllocator::FD3D12PoolAllocator(FD3D12Device* ParentDevice, FRHIGPUMask VisibleNodes, const FD3D12ResourceInitConfig& InInitConfig, const FString& InName,
+	EResourceAllocationStrategy InAllocationStrategy, uint64 InPoolSize, uint32 InPoolAlignment, uint32 InMaxAllocationSize) :
+	FRHIPoolAllocator(InPoolSize, InPoolAlignment, InMaxAllocationSize), 
+	FD3D12DeviceChild(ParentDevice), 
+	FD3D12MultiNodeGPUObject(ParentDevice->GetGPUMask(), VisibleNodes), 
+	InitConfig(InInitConfig), 
+	Name(InName), 
+	AllocationStrategy(InAllocationStrategy)
+{}
+
+
+FD3D12PoolAllocator::~FD3D12PoolAllocator()
+{
+	Destroy();
+}
+
+
+bool FD3D12PoolAllocator::SupportsAllocation(D3D12_HEAP_TYPE InHeapType, D3D12_RESOURCE_FLAGS InResourceFlags, EBufferUsageFlags InBufferUsage, ED3D12ResourceStateMode InResourceStateMode) const
+{
+	FD3D12ResourceInitConfig InInitConfig = GetResourceAllocatorInitConfig(InHeapType, InResourceFlags, InBufferUsage);
+	EResourceAllocationStrategy InAllocationStrategy = GetResourceAllocationStrategy(InResourceFlags, InResourceStateMode);
+	return (InitConfig == InInitConfig && AllocationStrategy == InAllocationStrategy);
+}
+
+
+void FD3D12PoolAllocator::AllocDefaultResource(D3D12_HEAP_TYPE InHeapType, const D3D12_RESOURCE_DESC& InDesc, EBufferUsageFlags InBufferUsage, ED3D12ResourceStateMode InResourceStateMode,
+	D3D12_RESOURCE_STATES InCreateState, uint32 InAllocationAlignment, const TCHAR* InName, FD3D12ResourceLocation& ResourceLocation)
+{
+	FD3D12Device* Device = GetParentDevice();
+	FD3D12Adapter* Adapter = Device->GetParentAdapter();
+
+	// If the resource location owns a block, this will deallocate it.
+	ResourceLocation.Clear();
+	if (InDesc.Width == 0)
+	{
+		return;
+	}
+
+#if DO_CHECK
+	// Validate the create state
+	if (InHeapType == D3D12_HEAP_TYPE_READBACK)
+	{
+		check(InCreateState == D3D12_RESOURCE_STATE_COPY_DEST);
+	}
+	else if (InHeapType == D3D12_HEAP_TYPE_UPLOAD)
+	{
+		check(InCreateState == D3D12_RESOURCE_STATE_GENERIC_READ);
+	}
+#if D3D12_RHI_RAYTRACING
+	else if (InBufferUsage & BUF_AccelerationStructure)
+	{
+		// RayTracing acceleration structures must be created in a particular state and may never transition out of it.
+		check(InResourceStateMode == ED3D12ResourceStateMode::SingleState);
+		check(InCreateState == D3D12_RESOURCE_STATE_RAYTRACING_ACCELERATION_STRUCTURE);
+	}
+#endif // D3D12_RHI_RAYTRACING
+#endif  // DO_CHECK
+
+	const bool PoolResource = InDesc.Width <= MaxAllocationSize;
+	if (PoolResource)
+	{
+		const bool bPlacedResource = (AllocationStrategy == EResourceAllocationStrategy::kPlacedResource);
+
+		// Ensure we're allocating from the correct pool
+		if (bPlacedResource)
+		{
+			// Writeable resources get separate ID3D12Resource* with their own resource state by using placed resources. Just make sure it's UAV, other flags are free to differ.
+			check((InDesc.Flags & InitConfig.ResourceFlags & D3D12_RESOURCE_FLAG_ALLOW_UNORDERED_ACCESS) != 0 || InHeapType == D3D12_HEAP_TYPE_READBACK);
+		}
+		else
+		{
+			// Read-only resources get suballocated from big resources, thus share ID3D12Resource* and resource state with other resources. Ensure it's suballocated from a resource with identical flags.
+			check(InDesc.Flags == InitConfig.ResourceFlags);
+		}
+
+		// Try to allocate in one of the pools
+		FRHIPoolAllocationData& AllocationData = ResourceLocation.GetPoolAllocatorPrivateData().PoolData;
+		if (TryAllocateInternal(InDesc.Width, InAllocationAlignment, AllocationData))
+		{
+			// Setup the resource location
+			ResourceLocation.SetType(FD3D12ResourceLocation::ResourceLocationType::eSubAllocation);
+			ResourceLocation.SetPoolAllocator(this);
+			ResourceLocation.SetSize(InDesc.Width);
+
+			AllocationData.SetOwner(&ResourceLocation);
+
+			if (AllocationStrategy == EResourceAllocationStrategy::kManualSubAllocation)
+			{
+				FD3D12Resource* BackingResource = GetBackingResource(ResourceLocation);
+
+				ResourceLocation.SetOffsetFromBaseOfResource(AllocationData.GetOffset());
+				ResourceLocation.SetResource(BackingResource);
+				ResourceLocation.SetGPUVirtualAddress(BackingResource->GetGPUVirtualAddress() + AllocationData.GetOffset());
+
+				if (IsCPUWritable(InitConfig.HeapType))
+				{
+					ResourceLocation.SetMappedBaseAddress((uint8*)BackingResource->GetResourceBaseAddress() + AllocationData.GetOffset());
+				}
+			}
+			else
+			{
+				check(ResourceLocation.GetResource() == nullptr);
+
+				FD3D12HeapAndOffset HeapAndOffset = GetBackingHeapAndAllocationOffsetInBytes(ResourceLocation);
+
+				FD3D12Resource* NewResource = nullptr;
+				VERIFYD3D12RESULT(Adapter->CreatePlacedResource(InDesc, HeapAndOffset.Heap, HeapAndOffset.Offset, InCreateState, ED3D12ResourceStateMode::MultiState, D3D12_RESOURCE_STATE_TBD, nullptr, &NewResource, InName));
+
+				ResourceLocation.SetResource(NewResource);
+			}
+
+			// Successfully sub-allocated
+			return;
+		}
+
+		// shouldn't fail
+		check(false);
+	}
+
+	// Allocate Standalone - move to owner of resource because this allocator should only manage pooled allocations (needed for now to do the same as FD3D12DefaultBufferPool)
+	FD3D12Resource* NewResource = nullptr;
+	VERIFYD3D12RESULT(Adapter->CreateBuffer(InHeapType, GetGPUMask(), GetVisibilityMask(), InCreateState, InResourceStateMode, InDesc.Width, &NewResource, InName, InDesc.Flags));
+	ResourceLocation.AsStandAlone(NewResource, InDesc.Width);
+}
+
+
+void FD3D12PoolAllocator::Deallocate(FD3D12ResourceLocation& ResourceLocation)
+{
+	check(IsOwner(ResourceLocation));
+
+	FScopeLock Lock(&CS);
+
+	// Mark allocation data as free
+	FRHIPoolAllocationData& AllocationData = ResourceLocation.GetPoolAllocatorPrivateData().PoolData;	
+	check(AllocationData.IsAllocated());
+
+	// If locked then assume still initial setup or in defragmentation unlock request
+	// Mark as nop because block will be deleted anyway
+	// TODO: optimize to not perform the full iteration here!
+	if (AllocationData.IsLocked())
+	{
+		for (FrameFencedAllocationData& Operation : FrameFencedOperations)
+		{
+			if (Operation.AllocationData == &AllocationData)
+			{
+				check(Operation.Operation == FrameFencedAllocationData::EOperation::Unlock);
+				Operation.Operation = FrameFencedAllocationData::EOperation::Nop;
+				Operation.AllocationData = nullptr;
+				break;
+			}
+		}
+	}
+
+	int16 PoolIndex = AllocationData.GetPoolIndex();
+	FRHIPoolAllocationData* ReleasedAllocationData = (AllocationDataPool.Num() > 0) ? AllocationDataPool.Pop(false) : new FRHIPoolAllocationData();;
+	bool bLocked = true;
+	ReleasedAllocationData->MoveFrom(AllocationData, bLocked);
+
+	// Clear the allocator data
+	ResourceLocation.ClearAllocator();
+
+	// Store fence when last used so we know when to unlock the free data
+	FD3D12Adapter* Adapter = GetParentDevice()->GetParentAdapter();
+	FD3D12Fence& FrameFence = Adapter->GetFrameFence();
+
+	FrameFencedOperations.AddUninitialized();
+	FrameFencedAllocationData& DeleteRequest = FrameFencedOperations.Last();
+	DeleteRequest.Operation = FrameFencedAllocationData::EOperation::Deallocate;
+	DeleteRequest.FrameFence = FrameFence.GetCurrentFence();
+	DeleteRequest.AllocationData = ReleasedAllocationData;
+
+	// Update the last used frame fence (used during garbage collection)
+	((FD3D12MemoryPool*)Pools[PoolIndex])->UpdateLastUsedFrameFence(DeleteRequest.FrameFence);
+
+	// Also store the placed resource so it can be correctly freed when fence is done
+	if (ResourceLocation.GetResource()->IsPlacedResource())
+	{
+		DeleteRequest.PlacedResource = ResourceLocation.GetResource();
+	}
+	else
+	{
+		DeleteRequest.PlacedResource = nullptr;
+	}
+}
+
+
+FRHIMemoryPool* FD3D12PoolAllocator::CreateNewPool(int16 InPoolIndex)
+{
+	FD3D12MemoryPool* NewPool = new FD3D12MemoryPool(GetParentDevice(),	GetVisibilityMask(), InitConfig,
+		Name, AllocationStrategy, InPoolIndex, PoolSize, PoolAlignment);
+	NewPool->Init();
+	return NewPool;
+}
+
+
+bool FD3D12PoolAllocator::HandleDefragRequest(FRHIPoolAllocationData* InSourceBlock, FRHIPoolAllocationData& InTmpTargetBlock)
+{
+	// Cache source copy data
+	FD3D12ResourceLocation* Owner = (FD3D12ResourceLocation*)InSourceBlock->GetOwner();
+	uint64 CurrentOffset = Owner->GetOffsetFromBaseOfResource();
+	FD3D12Resource* CurrentResource = Owner->GetResource();
+
+	// Release the current allocation (will only be freed on the next frame fence)
+	Deallocate(*Owner);
+
+	// Move temp allocation data to allocation data of the owner (part of different allocator now)		
+	bool bLocked = true;
+	InSourceBlock->MoveFrom(InTmpTargetBlock, bLocked);
+	InSourceBlock->SetOwner(Owner);
+	Owner->SetPoolAllocator(this);
+
+	// Notify owner of moved allocation data (recreated resources and SRVs if needed)
+	Owner->OnAllocationMoved(InSourceBlock);
+
+	// Add request to unlock the source block on the next fence value (copy operation should have been done by then)
+	FD3D12Adapter* Adapter = GetParentDevice()->GetParentAdapter();
+	FD3D12Fence& FrameFence = Adapter->GetFrameFence();
+	FrameFencedOperations.AddUninitialized();
+	FrameFencedAllocationData& UnlockRequest = FrameFencedOperations.Last();
+	UnlockRequest.Operation = FrameFencedAllocationData::EOperation::Unlock;
+	UnlockRequest.FrameFence = FrameFence.GetCurrentFence();
+	UnlockRequest.AllocationData = InSourceBlock;
+
+	// Schedule a copy operation of the actual data
+	FD3D12VRAMCopyOperation CopyOp;
+	CopyOp.SourceResource	= CurrentResource;
+	CopyOp.SourceOffset		= CurrentOffset;
+	CopyOp.DestResource		= Owner->GetResource();
+	CopyOp.DestOffset		= Owner->GetOffsetFromBaseOfResource();
+	CopyOp.Size				= InSourceBlock->GetSize();
+	PendingCopyOps.Add(CopyOp);
+
+	return true;
+}
+
+
+void FD3D12PoolAllocator::CleanUpAllocations(uint64 InFrameLag)
+{
+	FScopeLock Lock(&CS);
+
+	FD3D12Adapter* Adapter = GetParentDevice()->GetParentAdapter();
+	FD3D12Fence& FrameFence = Adapter->GetFrameFence();
+
+	uint32 PopCount = 0;
+	for (int32 i = 0; i < FrameFencedOperations.Num(); i++)
+	{
+		FrameFencedAllocationData& Operation = FrameFencedOperations[i];
+		if (FrameFence.IsFenceComplete(Operation.FrameFence))
+		{
+			switch (Operation.Operation)
+			{
+			case FrameFencedAllocationData::EOperation::Deallocate:
+			{
+				// Deallocate the locked block (actually free now)
+				DeallocateInternal(*Operation.AllocationData);
+				Operation.AllocationData->Reset();
+				AllocationDataPool.Add(Operation.AllocationData);
+
+				// Free placed resource if created
+				if (AllocationStrategy == EResourceAllocationStrategy::kPlacedResource)
+				{
+					// Release the resource
+					check(Operation.PlacedResource != nullptr);
+					Operation.PlacedResource->Release();
+					Operation.PlacedResource = nullptr;
+				}
+				else
+				{
+					check(Operation.PlacedResource == nullptr);
+				}
+				break;
+			}
+			case FrameFencedAllocationData::EOperation::Unlock:
+			{
+				Operation.AllocationData->Unlock();
+				break;
+			}
+			case FrameFencedAllocationData::EOperation::Nop:
+			{
+				break;
+			}
+			default: check(false);
+			}
+
+			PopCount = i + 1;
+		}
+		else
+		{
+			break;
+		}
+	}
+
+	if (PopCount)
+	{
+		// clear out all of the released blocks, don't allow the array to shrink
+		FrameFencedOperations.RemoveAt(0, PopCount, false);
+	}
+
+	// Trim empty allocators if not used in last n frames
+	const uint64 CompletedFence = FrameFence.UpdateLastCompletedFence();
+	for (int32 PoolIndex = 0; PoolIndex < Pools.Num(); ++PoolIndex)
+	{
+		FD3D12MemoryPool* MemoryPool = (FD3D12MemoryPool*) Pools[PoolIndex];
+		if (MemoryPool != nullptr && MemoryPool->IsEmpty() && (MemoryPool->GetLastUsedFrameFence() + InFrameLag <= CompletedFence))
+		{
+			MemoryPool->Destroy();
+			delete(MemoryPool);
+			Pools[PoolIndex] = nullptr;
+		}
+	}
+}
+
+
+void FD3D12PoolAllocator::TransferOwnership(FD3D12ResourceLocation& InSource, FD3D12ResourceLocation& InDest)
+{
+	FScopeLock Lock(&CS);
+
+	check(IsOwner(InSource));
+
+	// Don't need to lock - ownership simply changed
+	bool bLocked = false;
+	FRHIPoolAllocationData& DestinationPoolData = InDest.GetPoolAllocatorPrivateData().PoolData;
+	DestinationPoolData.MoveFrom(InSource.GetPoolAllocatorPrivateData().PoolData, bLocked);
+	DestinationPoolData.SetOwner(&InDest);
+}
+
+
+FD3D12Resource* FD3D12PoolAllocator::GetBackingResource(FD3D12ResourceLocation& InResourceLocation) const
+{
+	check(IsOwner(InResourceLocation));
+	FRHIPoolAllocationData& AllocationData = InResourceLocation.GetPoolAllocatorPrivateData().PoolData;
+	return ((FD3D12MemoryPool*)Pools[AllocationData.GetPoolIndex()])->GetBackingResource();
+}
+
+
+FD3D12HeapAndOffset FD3D12PoolAllocator::GetBackingHeapAndAllocationOffsetInBytes(FD3D12ResourceLocation& InResourceLocation) const
+{
+	check(IsOwner(InResourceLocation));
+
+	FD3D12HeapAndOffset HeapAndOffset;
+	FRHIPoolAllocationData& AllocationData = InResourceLocation.GetPoolAllocatorPrivateData().PoolData;
+	HeapAndOffset.Heap = ((FD3D12MemoryPool*)Pools[AllocationData.GetPoolIndex()])->GetBackingHeap();
+	HeapAndOffset.Offset = uint64(AlignDown(AllocationData.GetOffset(), PoolAlignment));
+	return HeapAndOffset;
+}
+
+
+void FD3D12PoolAllocator::FlushPendingCopyOps(FD3D12CommandContext& InCommandContext)
+{
+	FScopeLock Lock(&CS);
+
+	// TODO: sort the copy ops to reduce amount of transitions!!
+
+	FD3D12CommandListHandle& CommandListHandle = InCommandContext.CommandListHandle;
+
+	for (FD3D12VRAMCopyOperation& CopyOperation : PendingCopyOps)
+	{
+		if (CopyOperation.SourceResource->RequiresResourceStateTracking())
+		{
+			check(CopyOperation.DestResource->RequiresResourceStateTracking());
+			FD3D12DynamicRHI::TransitionResource(CommandListHandle, CopyOperation.SourceResource, D3D12_RESOURCE_STATE_TBD, D3D12_RESOURCE_STATE_COPY_SOURCE, D3D12_RESOURCE_BARRIER_ALL_SUBRESOURCES, FD3D12DynamicRHI::ETransitionMode::Apply);
+			FD3D12DynamicRHI::TransitionResource(CommandListHandle, CopyOperation.DestResource, D3D12_RESOURCE_STATE_TBD, D3D12_RESOURCE_STATE_COPY_DEST, D3D12_RESOURCE_BARRIER_ALL_SUBRESOURCES, FD3D12DynamicRHI::ETransitionMode::Apply);
+		}
+		else
+		{
+			check(!CopyOperation.DestResource->RequiresResourceStateTracking());
+			CommandListHandle.AddTransitionBarrier(CopyOperation.SourceResource, CopyOperation.SourceResource->GetDefaultResourceState(), D3D12_RESOURCE_STATE_COPY_SOURCE, D3D12_RESOURCE_BARRIER_ALL_SUBRESOURCES);
+			CommandListHandle.AddTransitionBarrier(CopyOperation.DestResource, CopyOperation.DestResource->GetDefaultResourceState(), D3D12_RESOURCE_STATE_COPY_DEST, D3D12_RESOURCE_BARRIER_ALL_SUBRESOURCES);
+		}
+
+		InCommandContext.numCopies++;
+		CommandListHandle.FlushResourceBarriers();
+		CommandListHandle->CopyBufferRegion(CopyOperation.DestResource->GetResource(), CopyOperation.DestOffset,
+			CopyOperation.SourceResource->GetResource(), CopyOperation.SourceOffset, CopyOperation.Size);
+		
+		CommandListHandle.UpdateResidency(CopyOperation.SourceResource);
+		CommandListHandle.UpdateResidency(CopyOperation.DestResource);
+
+		if (!CopyOperation.SourceResource->RequiresResourceStateTracking())
+		{
+			CommandListHandle.AddTransitionBarrier(CopyOperation.SourceResource, D3D12_RESOURCE_STATE_COPY_SOURCE, CopyOperation.SourceResource->GetDefaultResourceState(), D3D12_RESOURCE_BARRIER_ALL_SUBRESOURCES);
+			CommandListHandle.AddTransitionBarrier(CopyOperation.DestResource, D3D12_RESOURCE_STATE_COPY_DEST, CopyOperation.DestResource->GetDefaultResourceState(), D3D12_RESOURCE_BARRIER_ALL_SUBRESOURCES);
+		}
+	}
+
+	PendingCopyOps.Empty(PendingCopyOps.Num());
+}
+
diff --git a/Engine/Source/Runtime/D3D12RHI/Private/D3D12PoolAllocator.h b/Engine/Source/Runtime/D3D12RHI/Private/D3D12PoolAllocator.h
new file mode 100644
--- /dev/null
+++ b/Engine/Source/Runtime/D3D12RHI/Private/D3D12PoolAllocator.h
@@ -0,0 +1,171 @@
+// Copyright Epic Games, Inc. All Rights Reserved.
+#pragma once
+
+#include "RHIPoolAllocator.h"
+#include "D3D12Resources.h"
+
+enum class EResourceAllocationStrategy
+{
+	// This strategy uses Placed Resources to sub-allocate a buffer out of an underlying ID3D12Heap.
+	// The benefit of this is that each buffer can have it's own resource state and can be treated
+	// as any other buffer. The downside of this strategy is the API limitation which enforces
+	// the minimum buffer size to 64k leading to large internal fragmentation in the allocator
+	kPlacedResource,
+	// The alternative is to manually sub-allocate out of a single large buffer which allows block
+	// allocation granularity down to 1 byte. However, this strategy is only really valid for buffers which
+	// will be treated as read-only after their creation (i.e. most Index and Vertex buffers). This 
+	// is because the underlying resource can only have one state at a time.
+	kManualSubAllocation
+};
+
+
+// All data required to create a pool
+struct FD3D12ResourceInitConfig
+{
+	bool operator==(const FD3D12ResourceInitConfig& InOther) const
+	{
+		return HeapType == InOther.HeapType &&
+			HeapFlags && InOther.HeapFlags &&
+			ResourceFlags == InOther.ResourceFlags &&
+			InitialResourceState == InOther.InitialResourceState;
+	}
+
+	D3D12_HEAP_TYPE HeapType;
+	D3D12_HEAP_FLAGS HeapFlags;
+	D3D12_RESOURCE_FLAGS ResourceFlags;
+	D3D12_RESOURCE_STATES InitialResourceState;
+};
+
+
+/**
+@brief Stores all required data for a VRAM copy operation - doesn't own the resources
+**/
+struct FD3D12VRAMCopyOperation
+{
+	FD3D12Resource* SourceResource;
+	uint32 SourceOffset;
+	FD3D12Resource* DestResource;
+	uint32 DestOffset;
+	uint32 Size;
+};
+
+// Heap and offset combined to specific location of a resource
+struct FD3D12HeapAndOffset
+{
+	FD3D12Heap* Heap;
+	uint64 Offset;
+};
+
+
+//////////////////////////////////////////////////////////////////////////
+// FD3D12Pool
+//////////////////////////////////////////////////////////////////////////
+
+/**
+@brief D3D12 specific implementation of the FRHIMemoryPool
+**/
+class FD3D12MemoryPool : public FRHIMemoryPool, public FD3D12DeviceChild, public FD3D12MultiNodeGPUObject
+{
+public:
+
+	// Constructor
+	FD3D12MemoryPool(FD3D12Device* ParentDevice, FRHIGPUMask VisibleNodes, const FD3D12ResourceInitConfig& InInitConfig, const FString& Name,
+		EResourceAllocationStrategy InAllocationStrategy, int16 InPoolIndex, uint64 InPoolSize, uint32 InPoolAlignment);
+	~FD3D12MemoryPool();
+
+	// Setup/Shutdown
+	virtual void Init() override;
+	virtual void Destroy() override;
+
+	FD3D12Resource* GetBackingResource() { check(AllocationStrategy == EResourceAllocationStrategy::kManualSubAllocation); return BackingResource.GetReference(); }
+	FD3D12Heap* GetBackingHeap() { check(AllocationStrategy == EResourceAllocationStrategy::kPlacedResource); return BackingHeap.GetReference(); }
+
+	uint64 GetLastUsedFrameFence() const { return LastUsedFrameFence; }
+	void UpdateLastUsedFrameFence(uint64 InFrameFence) { LastUsedFrameFence = FMath::Max(LastUsedFrameFence, InFrameFence); }
+
+protected:
+
+	const FD3D12ResourceInitConfig InitConfig;
+	const FString Name;
+	EResourceAllocationStrategy AllocationStrategy;
+
+	// Actual D3D resource - either resource or heap (see EResourceAllocationStrategy)
+	TRefCountPtr<FD3D12Resource> BackingResource;
+	TRefCountPtr<FD3D12Heap> BackingHeap;
+
+	uint64 LastUsedFrameFence;
+};
+
+
+//////////////////////////////////////////////////////////////////////////
+// FD3D12PoolAllocator
+//////////////////////////////////////////////////////////////////////////
+
+/**
+@brief D3D12 specific implementation of the FRHIPoolAllocator
+**/
+class FD3D12PoolAllocator : public FRHIPoolAllocator, public FD3D12DeviceChild, public FD3D12MultiNodeGPUObject
+{
+public:
+
+	// Constructor
+	FD3D12PoolAllocator(FD3D12Device* ParentDevice, FRHIGPUMask VisibleNodes, const FD3D12ResourceInitConfig& InInitConfig, const FString& InName,
+		EResourceAllocationStrategy InAllocationStrategy, uint64 InPoolSize, uint32 InPoolAlignment, uint32 InMaxAllocationSize);
+	~FD3D12PoolAllocator();
+	
+	// Function names currently have to match with FD3D12DefaultBufferPool until we can make full replacement of the allocator
+	bool SupportsAllocation(D3D12_HEAP_TYPE InHeapType, D3D12_RESOURCE_FLAGS InResourceFlags, EBufferUsageFlags InBufferUsage, ED3D12ResourceStateMode InResourceStateMode) const;
+	void AllocDefaultResource(D3D12_HEAP_TYPE InHeapType, const D3D12_RESOURCE_DESC& InDesc, EBufferUsageFlags InBufferUsage, ED3D12ResourceStateMode InResourceStateMode,
+		D3D12_RESOURCE_STATES InCreateState, uint32 InAlignment, const TCHAR* InName, FD3D12ResourceLocation& ResourceLocation);
+	void Deallocate(FD3D12ResourceLocation& ResourceLocation);
+
+	void CleanUpAllocations(uint64 InFrameLag);
+	void FlushPendingCopyOps(FD3D12CommandContext& InCommandContext);
+
+	void TransferOwnership(FD3D12ResourceLocation& InSource, FD3D12ResourceLocation& InDest);
+	bool IsOwner(FD3D12ResourceLocation& ResourceLocation) const { return ResourceLocation.GetPoolAllocator() == this; }
+
+	EResourceAllocationStrategy GetAllocationStrategy() const { return AllocationStrategy; }
+	FD3D12Resource* GetBackingResource(FD3D12ResourceLocation& InResourceLocation) const;
+	FD3D12HeapAndOffset GetBackingHeapAndAllocationOffsetInBytes(FD3D12ResourceLocation& InResourceLocation) const;
+
+	static FD3D12ResourceInitConfig GetResourceAllocatorInitConfig(D3D12_HEAP_TYPE InHeapType, D3D12_RESOURCE_FLAGS InResourceFlags, EBufferUsageFlags InBufferUsage);
+	static EResourceAllocationStrategy GetResourceAllocationStrategy(D3D12_RESOURCE_FLAGS InResourceFlags, ED3D12ResourceStateMode InResourceStateMode);
+	
+protected:
+
+	// Implementation of FRHIPoolAllocator pure virtuals
+	virtual FRHIMemoryPool* CreateNewPool(int16 InPoolIndex) override;
+	virtual bool HandleDefragRequest(FRHIPoolAllocationData* InSourceBlock, FRHIPoolAllocationData& InTmpTargetBlock) override;
+	
+	// Locked allocation data which needs to do something specific at certain frame fence value	
+	struct FrameFencedAllocationData
+	{
+		enum class EOperation
+		{
+			Invalid,
+			Deallocate,
+			Unlock,
+			Nop,
+		};
+
+		EOperation Operation = EOperation::Invalid;
+		FD3D12Resource* PlacedResource = nullptr;
+		uint64 FrameFence = 0;
+		FRHIPoolAllocationData* AllocationData = nullptr;
+	};
+
+	const FD3D12ResourceInitConfig InitConfig;
+	const FString Name;
+	EResourceAllocationStrategy AllocationStrategy;
+
+	// All operations which need to happen on specific frame fences
+	TArray<FrameFencedAllocationData> FrameFencedOperations;
+
+	// Pending copy ops which need to be flush this frame
+	TArray<FD3D12VRAMCopyOperation> PendingCopyOps;
+
+	// Allocation data pool used to reduce allocation overhead
+	TArray<FRHIPoolAllocationData*> AllocationDataPool;
+};
+
diff --git a/Engine/Source/Runtime/D3D12RHI/Private/D3D12RHIPrivate.h b/Engine/Source/Runtime/D3D12RHI/Private/D3D12RHIPrivate.h
--- a/Engine/Source/Runtime/D3D12RHI/Private/D3D12RHIPrivate.h
+++ b/Engine/Source/Runtime/D3D12RHI/Private/D3D12RHIPrivate.h
@@ -61,6 +61,8 @@ extern int32 GDX12NVAfterMathTrackResources;
 extern int32 GDX12NVAfterMathMarkers;
 #endif
 
+#define USE_BUCKET_ALLOCATOR 0
+
 #include "D3D12Residency.h"
 
 // D3D RHI public headers.
diff --git a/Engine/Source/Runtime/D3D12RHI/Private/D3D12Resources.cpp b/Engine/Source/Runtime/D3D12RHI/Private/D3D12Resources.cpp
--- a/Engine/Source/Runtime/D3D12RHI/Private/D3D12Resources.cpp
+++ b/Engine/Source/Runtime/D3D12RHI/Private/D3D12Resources.cpp
@@ -535,6 +535,7 @@ HRESULT FD3D12Adapter::CreateBuffer(const D3D12_HEAP_PROPERTIES& HeapProps,
 FD3D12ResourceLocation::FD3D12ResourceLocation(FD3D12Device* Parent)
 	: FD3D12DeviceChild(Parent)
 	, Type(ResourceLocationType::eUndefined)
+	, Owner(nullptr)
 	, UnderlyingResource(nullptr)
 	, ResidencyHandle(nullptr)
 	, Allocator(nullptr)
@@ -590,6 +591,11 @@ void FD3D12ResourceLocation::TransferOwnership(FD3D12ResourceLocation& Destinati
 
 	FMemory::Memmove(&Destination, &Source, sizeof(FD3D12ResourceLocation));
 
+	if (Source.GetAllocatorType() == FD3D12ResourceLocation::AT_Pool)
+	{
+		Source.GetPoolAllocator()->TransferOwnership(Source, Destination);
+	}
+
 	// update tracked allocation
 #if !PLATFORM_WINDOWS && ENABLE_LOW_LEVEL_MEM_TRACKER
 	if (Source.GetType() == ResourceLocationType::eSubAllocation && Source.AllocatorType != AT_SegList)
@@ -623,11 +629,18 @@ void FD3D12ResourceLocation::Swap(FD3D12ResourceLocation& Other)
 	}
 #endif
 
+	// Should not be linked list allocated - otherwise internal linked list data needs to be updated as well in a threadsafe way
+	check(Other.GetAllocatorType() != FD3D12ResourceLocation::AT_Pool);
+	check(GetAllocatorType() != FD3D12ResourceLocation::AT_Pool);
+
 	::Swap(*this, Other);
 }
 
 void FD3D12ResourceLocation::Alias(FD3D12ResourceLocation & Destination, FD3D12ResourceLocation & Source)
 {
+	// Should not be linked list allocated - otherwise internal linked list data needs to be updated as well in a threadsafe way
+	check(Source.GetAllocatorType() != FD3D12ResourceLocation::AT_Pool);
+
 	check(Source.GetResource() != nullptr);
 	Destination.Clear();
 
@@ -641,6 +654,8 @@ void FD3D12ResourceLocation::Alias(FD3D12ResourceLocation & Destination, FD3D12R
 
 void FD3D12ResourceLocation::ReferenceNode(FD3D12Device* DestinationDevice, FD3D12ResourceLocation& Destination, FD3D12ResourceLocation& Source)
 {
+	// Should not be linked list allocated - otherwise internal linked list data needs to be updated as well in a threadsafe way
+	check(Source.GetAllocatorType() != FD3D12ResourceLocation::AT_Pool);
 	check(Source.GetResource() != nullptr);
 	Destination.Clear();
 
@@ -685,6 +700,10 @@ void FD3D12ResourceLocation::ReleaseResource()
 				GetSegListAllocatorPrivateData().Offset,
 				GetSize());
 		}
+		else if (AllocatorType == AT_Pool)
+		{
+			PoolAllocator->Deallocate(*this);
+		}
 		else
 		{
 			Allocator->Deallocate(*this);
@@ -804,6 +823,67 @@ void FD3D12ResourceLocation::AsStandAlone(FD3D12Resource* Resource, uint32 Buffe
 	UpdateStandAloneStats(bIncrement);
 }
 
+
+bool FD3D12ResourceLocation::OnAllocationMoved(FRHIPoolAllocationData* InNewData)
+{
+	// Assume linked list allocated for now - only defragging allocator
+	FRHIPoolAllocationData& AllocationData = GetPoolAllocatorPrivateData().PoolData;
+	check(InNewData == &AllocationData);
+	check(AllocationData.IsAllocated()); // Should be allocated
+	check(AllocationData.GetSize() == Size); // Same size
+	check(Type == ResourceLocationType::eSubAllocation); // Suballocated
+	check(GetMappedBaseAddress() == nullptr); // And VRAM only
+	
+	// Get the resource and the actual new allocator
+	FD3D12Resource* CurrentResource = GetResource();
+	FD3D12PoolAllocator* NewAllocator = GetPoolAllocator();
+
+	// If sub allocated and not placed only update the internal data
+	if (NewAllocator->GetAllocationStrategy() == EResourceAllocationStrategy::kManualSubAllocation)
+	{
+		check(!CurrentResource->IsPlacedResource());
+
+		OffsetFromBaseOfResource = AllocationData.GetOffset();
+		UnderlyingResource = NewAllocator->GetBackingResource(*this);
+	}
+	else
+	{
+		check(CurrentResource->IsPlacedResource());
+		check(OffsetFromBaseOfResource == 0);
+
+		// recreate the placed resource (ownership of current resource is already handled during the internal move)
+		FD3D12HeapAndOffset HeapAndOffset = NewAllocator->GetBackingHeapAndAllocationOffsetInBytes(*this);
+
+		CResourceState& ResourceState = CurrentResource->GetResourceState();
+		check(ResourceState.AreAllSubresourcesSame());
+		D3D12_RESOURCE_STATES CreateState = ResourceState.GetSubresourceState(0);
+
+		FD3D12Resource* NewResource = nullptr;
+		VERIFYD3D12RESULT(CurrentResource->GetParentDevice()->GetParentAdapter()->CreatePlacedResource(CurrentResource->GetDesc(), HeapAndOffset.Heap, HeapAndOffset.Offset, CreateState, 
+			ED3D12ResourceStateMode::MultiState, D3D12_RESOURCE_STATE_TBD, nullptr, &NewResource, nullptr)); // CurrentResource->GetName()
+
+		UnderlyingResource = NewResource;
+	}
+
+	GPUVirtualAddress = UnderlyingResource->GetGPUVirtualAddress() + OffsetFromBaseOfResource;
+	ResidencyHandle = UnderlyingResource->GetResidencyHandle();
+
+	// Recreate all the views (SRVs and UAVs)
+	Owner->RecreateViews(*this);
+
+	return true;
+}
+
+
+void FD3D12ResourceLocation::UnlockPoolData()
+{
+	if (AllocatorType == AT_Pool)
+	{
+		GetPoolAllocatorPrivateData().PoolData.Unlock();
+	}
+}
+
+
 /////////////////////////////////////////////////////////////////////
 //	FD3D12 Resource Barrier Batcher
 /////////////////////////////////////////////////////////////////////
diff --git a/Engine/Source/Runtime/D3D12RHI/Private/D3D12UAV.cpp b/Engine/Source/Runtime/D3D12RHI/Private/D3D12UAV.cpp
--- a/Engine/Source/Runtime/D3D12RHI/Private/D3D12UAV.cpp
+++ b/Engine/Source/Runtime/D3D12RHI/Private/D3D12UAV.cpp
@@ -24,7 +24,9 @@ inline FD3D12UnorderedAccessView* CreateUAV(D3D12_UNORDERED_ACCESS_VIEW_DESC& De
 			Device->GetParentAdapter()->CreateBuffer(D3D12_HEAP_TYPE_DEFAULT, Node, Node, 4, &CounterResource,  TEXT("Counter"), D3D12_RESOURCE_FLAG_ALLOW_UNORDERED_ACCESS);
 		}
 
-		return new FD3D12UnorderedAccessView(Device, Desc, Resource->ResourceLocation, CounterResource);
+		FD3D12UnorderedAccessView* UAVView = new FD3D12UnorderedAccessView(Device, Desc, Resource->ResourceLocation, CounterResource);
+		Resource->AddView(UAVView);
+		return UAVView;
 	});
 }
 
diff --git a/Engine/Source/Runtime/D3D12RHI/Private/D3D12View.cpp b/Engine/Source/Runtime/D3D12RHI/Private/D3D12View.cpp
--- a/Engine/Source/Runtime/D3D12RHI/Private/D3D12View.cpp
+++ b/Engine/Source/Runtime/D3D12RHI/Private/D3D12View.cpp
@@ -226,7 +226,7 @@ FShaderResourceViewRHIRef FD3D12DynamicRHI::RHICreateShaderResourceView(const FS
 
 					// Create a Shader Resource View
 					SRV->Initialize(SRVDesc, VertexBuffer->ResourceLocation, CreationStride, StartOffsetBytes);
-					VertexBuffer->AddDynamicSRV(SRV);
+					VertexBuffer->AddView(SRV);
 				}
 			};
 
@@ -310,7 +310,7 @@ FShaderResourceViewRHIRef FD3D12DynamicRHI::RHICreateShaderResourceView(const FS
 
 					// Create a Shader Resource View
 					SRV->Initialize(SRVDesc, StructuredBuffer->ResourceLocation, Stride, StartOffsetBytes);
-					StructuredBuffer->AddDynamicSRV(SRV);
+					StructuredBuffer->AddView(SRV);
 				}
 			};
 
@@ -378,7 +378,7 @@ void FD3D12DynamicRHI::RHIUpdateShaderResourceView(FRHIShaderResourceView* SRV,
 
 			FD3D12Device* ParentDevice = Buffer->GetParentDevice();
 			SRVD3D12->Initialize(ParentDevice, SRVDesc, Buffer->ResourceLocation, Stride);
-			Buffer->AddDynamicSRV(SRVD3D12);
+			Buffer->AddView(SRVD3D12);
 		}
 	}
 }
diff --git a/Engine/Source/Runtime/D3D12RHI/Private/D3D12View.h b/Engine/Source/Runtime/D3D12RHI/Private/D3D12View.h
--- a/Engine/Source/Runtime/D3D12RHI/Private/D3D12View.h
+++ b/Engine/Source/Runtime/D3D12RHI/Private/D3D12View.h
@@ -885,6 +885,11 @@ public:
 		CreateView(Desc, *ResourceLocation);
 	}
 
+	virtual void Recreate(FD3D12ResourceLocation& InResourceLocation) override
+	{
+		Rename(InResourceLocation);
+	}
+
 	FORCEINLINE bool IsDepthStencilResource()	const { return bContainsDepthPlane || bContainsStencilPlane; }
 	FORCEINLINE bool IsDepthPlaneResource()		const { return bContainsDepthPlane; }
 	FORCEINLINE bool IsStencilPlaneResource()	const { return bContainsStencilPlane; }
@@ -909,7 +914,7 @@ public:
 	FD3D12ResourceLocation ViewLocation;
 };
 
-class FD3D12UnorderedAccessView : public FRHIUnorderedAccessView, public FD3D12View < D3D12_UNORDERED_ACCESS_VIEW_DESC >, public FD3D12LinkedAdapterObject<FD3D12UnorderedAccessView>
+class FD3D12UnorderedAccessView : public FD3D12BaseShaderResourceView, public FRHIUnorderedAccessView, public FD3D12View < D3D12_UNORDERED_ACCESS_VIEW_DESC >, public FD3D12LinkedAdapterObject<FD3D12UnorderedAccessView>
 {
 public:
 	TRefCountPtr<FD3D12Resource> CounterResource;
@@ -928,10 +933,24 @@ public:
 		CreateViewWithCounter(InDesc, InResourceLocation, InCounterResource);
 	}
 
+	~FD3D12UnorderedAccessView()
+	{
+		FD3D12BaseShaderResourceView::Remove();
+	}
+
 	void Initialize(D3D12_UNORDERED_ACCESS_VIEW_DESC& InDesc, FD3D12ResourceLocation& InResourceLocation)
 	{
 		CreateViewWithCounter(InDesc, InResourceLocation, nullptr);
 	}
+
+	virtual void Recreate(FD3D12ResourceLocation& InResourceLocation) override
+	{
+		check(ResourceLocation == &InResourceLocation);
+		check(CounterResource == nullptr);
+		check(InResourceLocation.GetOffsetFromBaseOfResource() == 0);
+
+		Initialize(Desc, InResourceLocation);
+	}
 };
 
 class FD3D12UnorderedAccessViewWithLocation : public FD3D12UnorderedAccessView
diff --git a/Engine/Source/Runtime/D3D12RHI/Public/D3D12Resources.h b/Engine/Source/Runtime/D3D12RHI/Public/D3D12Resources.h
--- a/Engine/Source/Runtime/D3D12RHI/Public/D3D12Resources.h
+++ b/Engine/Source/Runtime/D3D12RHI/Public/D3D12Resources.h
@@ -8,6 +8,7 @@
 
 #include "BoundShaderStateCache.h"
 #include "D3D12ShaderResources.h"
+#include "RHIPoolAllocator.h"
 
 #if PLATFORM_USE_BACKBUFFER_WRITE_TRANSITION_TRACKING
 constexpr D3D12_RESOURCE_STATES BackBufferBarrierWriteTransitionTargets = D3D12_RESOURCE_STATE_RENDER_TARGET | D3D12_RESOURCE_STATE_UNORDERED_ACCESS |
@@ -21,6 +22,7 @@ class FD3D12CommandListManager;
 class FD3D12CommandContext;
 class FD3D12CommandListHandle;
 class FD3D12SegListAllocator;
+class FD3D12PoolAllocator;
 struct FD3D12GraphicsPipelineState;
 typedef FD3D12StateCacheBase FD3D12StateCache;
 
@@ -430,7 +432,7 @@ private:
 	}
 };
 
-#ifdef USE_BUCKET_ALLOCATOR
+#if USE_BUCKET_ALLOCATOR
 typedef class FD3D12BucketAllocator FD3D12BaseAllocatorType;
 #else
 typedef class FD3D12BuddyAllocator FD3D12BaseAllocatorType;
@@ -474,9 +476,21 @@ struct FD3D12SegListAllocatorPrivateData
 	}
 };
 
+struct FD3D12PoolAllocatorPrivateData
+{
+	FRHIPoolAllocationData PoolData;
+
+	void Init()
+	{
+		PoolData.Reset();
+	}
+};
+
 class FD3D12ResourceAllocator;
+class FD3D12BaseShaderResource;
+
 // A very light-weight and cache friendly way of accessing a GPU resource
-class FD3D12ResourceLocation : public FD3D12DeviceChild, public FNoncopyable
+class FD3D12ResourceLocation : public FRHIPoolResource, public FD3D12DeviceChild, public FNoncopyable
 {
 public:
 
@@ -496,6 +510,7 @@ public:
 	{
 		AT_Default, // FD3D12BaseAllocatorType
 		AT_SegList, // FD3D12SegListAllocator
+		AT_Pool,	// FD3D12PoolAllocator
 		AT_Unknown = 0xff
 	};
 
@@ -508,11 +523,15 @@ public:
 	static void TransferOwnership(FD3D12ResourceLocation& Destination, FD3D12ResourceLocation& Source);
 
 	// Setters
+	inline void SetOwner(FD3D12BaseShaderResource* InOwner) { Owner = InOwner; }
 	void SetResource(FD3D12Resource* Value);
 	inline void SetType(ResourceLocationType Value) { Type = Value;}
 
 	inline void SetAllocator(FD3D12BaseAllocatorType* Value) { Allocator = Value; AllocatorType = AT_Default; }
 	inline void SetSegListAllocator(FD3D12SegListAllocator* Value) { SegListAllocator = Value; AllocatorType = AT_SegList; }
+	inline void SetPoolAllocator(FD3D12PoolAllocator* Value) { PoolAllocator = Value; AllocatorType = AT_Pool; }
+	inline void ClearAllocator() { Allocator = nullptr; AllocatorType = AT_Unknown; }
+
 	inline void SetMappedBaseAddress(void* Value) { MappedBaseAddress = Value; }
 	inline void SetGPUVirtualAddress(D3D12_GPU_VIRTUAL_ADDRESS Value) { GPUVirtualAddress = Value; }
 	inline void SetOffsetFromBaseOfResource(uint64 Value) { OffsetFromBaseOfResource = Value; }
@@ -520,8 +539,10 @@ public:
 
 	// Getters
 	inline ResourceLocationType GetType() const { return Type; }
+	inline EAllocatorType GetAllocatorType() const { return AllocatorType; }
 	inline FD3D12BaseAllocatorType* GetAllocator() { check(AT_Default == AllocatorType); return Allocator; }
 	inline FD3D12SegListAllocator* GetSegListAllocator() { check(AT_SegList == AllocatorType); return SegListAllocator; }
+	inline FD3D12PoolAllocator* GetPoolAllocator() { check(AT_Pool == AllocatorType); return PoolAllocator; }
 	inline FD3D12Resource* GetResource() const { return UnderlyingResource; }
 	inline void* GetMappedBaseAddress() const { return MappedBaseAddress; }
 	inline D3D12_GPU_VIRTUAL_ADDRESS GetGPUVirtualAddress() const { return GPUVirtualAddress; }
@@ -531,6 +552,11 @@ public:
 	inline FD3D12BuddyAllocatorPrivateData& GetBuddyAllocatorPrivateData() { return AllocatorData.BuddyAllocatorPrivateData; }
 	inline FD3D12BlockAllocatorPrivateData& GetBlockAllocatorPrivateData() { return AllocatorData.BlockAllocatorPrivateData; }
 	inline FD3D12SegListAllocatorPrivateData& GetSegListAllocatorPrivateData() { return AllocatorData.SegListAllocatorPrivateData; }
+	inline FD3D12PoolAllocatorPrivateData& GetPoolAllocatorPrivateData() { return AllocatorData.PoolAllocatorPrivateData; }
+
+	// Pool allocation specific functions
+	bool OnAllocationMoved(FRHIPoolAllocationData* InNewData);
+	void UnlockPoolData();
 
 	const inline bool IsValid() const { return Type != ResourceLocationType::eUndefined; }
 
@@ -606,6 +632,7 @@ private:
 
 	ResourceLocationType Type;
 
+	FD3D12BaseShaderResource* Owner;
 	FD3D12Resource* UnderlyingResource;
 	FD3D12ResidencyHandle* ResidencyHandle;
 
@@ -614,6 +641,7 @@ private:
 	{
 		FD3D12BaseAllocatorType* Allocator;
 		FD3D12SegListAllocator* SegListAllocator;
+		FD3D12PoolAllocator* PoolAllocator;
 	};
 
 	// Union to save memory
@@ -622,6 +650,7 @@ private:
 		FD3D12BuddyAllocatorPrivateData BuddyAllocatorPrivateData;
 		FD3D12BlockAllocatorPrivateData BlockAllocatorPrivateData;
 		FD3D12SegListAllocatorPrivateData SegListAllocatorPrivateData;
+		FD3D12PoolAllocatorPrivateData PoolAllocatorPrivateData;
 	} AllocatorData;
 
 	// Note: These values refer to the start of this location including any padding *NOT* the start of the underlying resource
@@ -727,50 +756,53 @@ struct FD3D12LockedResource : public FD3D12DeviceChild
 class FD3D12BaseShaderResourceView
 {
 protected:
+	virtual ~FD3D12BaseShaderResourceView() { check(BaseShaderResource == nullptr); }
+
 	void Remove();
+	virtual void Recreate(FD3D12ResourceLocation& InResourceLocation) = 0;
 
 	friend class FD3D12BaseShaderResource;
-	FD3D12BaseShaderResource* DynamicResource = nullptr;
+	FD3D12BaseShaderResource* BaseShaderResource = nullptr;
 };
 
 /** The base class of resources that may be bound as shader resources. */
 class FD3D12BaseShaderResource : public FD3D12DeviceChild, public IRefCountedObject
 {
 protected:
-	FCriticalSection DynamicSRVsCS;
-	TArray<class FD3D12BaseShaderResourceView*> DynamicSRVs;
+	FCriticalSection ViewsCS;
+	TArray<class FD3D12BaseShaderResourceView*> Views;
 
 public:
 	FD3D12Resource* GetResource() const { return ResourceLocation.GetResource(); }
 
-	void AddDynamicSRV(FD3D12BaseShaderResourceView* InSRV)
+	void AddView(FD3D12BaseShaderResourceView* InView)
 	{
-		FScopeLock Lock(&DynamicSRVsCS);
-		check(InSRV->DynamicResource == nullptr);
-		InSRV->DynamicResource = this;
-		DynamicSRVs.Add(InSRV);
+		FScopeLock Lock(&ViewsCS);
+		check(InView->BaseShaderResource == nullptr);
+		InView->BaseShaderResource = this;
+		Views.Add(InView);
 	}
 
-	void RemoveDynamicSRV(FD3D12BaseShaderResourceView* InSRV)
+	void RemoveView(FD3D12BaseShaderResourceView* InSRV)
 	{
-		FScopeLock Lock(&DynamicSRVsCS);
-		check(InSRV->DynamicResource == this);
-		InSRV->DynamicResource = nullptr;
-		uint32 Removed = DynamicSRVs.Remove(InSRV);
+		FScopeLock Lock(&ViewsCS);
+		check(InSRV->BaseShaderResource == this);
+		InSRV->BaseShaderResource = nullptr;
+		uint32 Removed = Views.Remove(InSRV);
 		check(Removed == 1);
 	}
 
-	void RemoveAllDynamicSRVs()
+	void RemoveAllViews()
 	{
-		FScopeLock Lock(&DynamicSRVsCS);
-		for (FD3D12BaseShaderResourceView* DynamicSRV : DynamicSRVs)
+		FScopeLock Lock(&ViewsCS);
+		for (FD3D12BaseShaderResourceView* View : Views)
 		{
-			if (DynamicSRV != nullptr)
+			if (View != nullptr)
 			{
-				DynamicSRV->DynamicResource = nullptr;
+				View->BaseShaderResource = nullptr;
 			}
 		}
-		DynamicSRVs.Reset();
+		Views.Reset();
 	}
 
 	void Swap(FD3D12BaseShaderResource& Other)
@@ -778,7 +810,19 @@ public:
 		::Swap(Parent, Other.Parent);
 		ResourceLocation.Swap(Other.ResourceLocation);
 		::Swap(BufferAlignment, Other.BufferAlignment);
-		::Swap(DynamicSRVs, Other.DynamicSRVs);
+		::Swap(Views, Other.Views);
+	}
+
+	void RecreateViews(FD3D12ResourceLocation& InResourceLocation)
+	{
+		FScopeLock Lock(&ViewsCS);
+		for (FD3D12BaseShaderResourceView* View : Views)
+		{
+			if (View != nullptr)
+			{
+				View->Recreate(InResourceLocation);
+			}
+		}
 	}
 
 	FD3D12ResourceLocation ResourceLocation;
@@ -794,19 +838,19 @@ public:
 
 	~FD3D12BaseShaderResource()
 	{
-		for (FD3D12BaseShaderResourceView* DynamicSRV : DynamicSRVs)
+		for (FD3D12BaseShaderResourceView* View : Views)
 		{
-			check(DynamicSRV->DynamicResource == this);
-			DynamicSRV->DynamicResource = nullptr;
+			check(View->BaseShaderResource == this);
+			View->BaseShaderResource = nullptr;
 		}
 	}
 };
 
 inline void FD3D12BaseShaderResourceView::Remove()
 {
-	if (DynamicResource)
+	if (BaseShaderResource)
 	{
-		DynamicResource->RemoveDynamicSRV(this);
+		BaseShaderResource->RemoveView(this);
 	}
 }
 
diff --git a/Engine/Source/Runtime/RHICore/Private/RHIPoolAllocator.cpp b/Engine/Source/Runtime/RHICore/Private/RHIPoolAllocator.cpp
new file mode 100644
--- /dev/null
+++ b/Engine/Source/Runtime/RHICore/Private/RHIPoolAllocator.cpp
@@ -0,0 +1,659 @@
+// Copyright Epic Games, Inc. All Rights Reserved.
+
+#include "RHIPoolAllocator.h"
+
+static int32 GRHIPoolAllocatorValidateLinkList = 0;
+static FAutoConsoleVariableRef CVarD3D12FastAllocatorMinPagesToRetain(
+	TEXT("RHIPoolAllocator.ValidateLinkedList"),
+	GRHIPoolAllocatorValidateLinkList,
+	TEXT("Validate all the internal linked list data of all the RHIPoolAllocators during every operation."),
+	ECVF_Default);
+
+//-----------------------------------------------------------------------------
+//	LinkedList Allocator PrivateData
+//-----------------------------------------------------------------------------
+
+void FRHIPoolAllocationData::Reset()
+{
+	Size = 0;
+	Alignment = 0;
+	SetAllocationType(EAllocationType::Unknown);
+	PoolIndex = -1;
+	Offset = 0;
+	Locked = false;
+
+	Owner = nullptr;
+	PreviousAllocation = nullptr;
+	NextAllocation = nullptr;
+}
+
+void FRHIPoolAllocationData::InitAsHead(int16 InPoolIndex)
+{
+	Reset();
+
+	SetAllocationType(EAllocationType::Head);
+	NextAllocation = this;
+	PreviousAllocation = this;
+	PoolIndex = InPoolIndex;
+}
+
+void FRHIPoolAllocationData::InitAsFree(int16 InPoolIndex, uint32 InSize, uint32 InAlignment, uint32 InOffset)
+{
+	Reset();
+
+	Size = InSize;
+	Alignment = InAlignment;
+	SetAllocationType(EAllocationType::Free);
+	Offset = InOffset;
+	PoolIndex = InPoolIndex;
+}
+
+
+void FRHIPoolAllocationData::InitAsAllocated(uint32 InSize, uint32 InAlignment, FRHIPoolAllocationData* InFree)
+{
+	check(InFree->IsFree());
+
+	Reset();
+	
+	Size = InSize;
+	Alignment = InAlignment;
+	SetAllocationType(EAllocationType::Allocated);
+	Offset = FRHIMemoryPool::GetAlignedOffset(InFree->Offset, InFree->Alignment, InAlignment);
+	PoolIndex = InFree->PoolIndex;
+	Locked = true;
+
+	uint32 AlignedSize = AlignArbitrary(InSize, InFree->Alignment);
+	InFree->Size -= AlignedSize;
+	InFree->Offset += AlignedSize;
+	InFree->AddBefore(this);
+}
+
+
+void FRHIPoolAllocationData::MoveFrom(FRHIPoolAllocationData& InAllocated, bool InLocked)
+{
+	check(InAllocated.IsAllocated());
+	
+	Reset();
+
+	Size = InAllocated.Size;
+	Alignment = InAllocated.Alignment;
+	Type = InAllocated.Type;
+	Offset = InAllocated.Offset;
+	PoolIndex = InAllocated.PoolIndex;
+	Locked = InLocked;
+	
+	// Update linked list
+	InAllocated.PreviousAllocation->NextAllocation = this;
+	InAllocated.NextAllocation->PreviousAllocation = this;
+	PreviousAllocation = InAllocated.PreviousAllocation;
+	NextAllocation = InAllocated.NextAllocation;
+
+	InAllocated.Reset();
+}
+
+
+void FRHIPoolAllocationData::MarkFree(uint32 InAlignment)
+{
+	check(IsAllocated());
+
+	// Mark as free and update the size, offset and alignment according to the new requested alignment
+	SetAllocationType(EAllocationType::Free);
+	Size = Align(Size, InAlignment);
+	Offset = AlignDown(Offset, InAlignment);
+	Alignment = InAlignment;
+}
+
+
+void FRHIPoolAllocationData::Merge(FRHIPoolAllocationData* InOther)
+{
+	check(IsFree() && InOther->IsFree());
+	check((Offset + Size) == InOther->Offset);
+	check(PoolIndex == InOther->GetPoolIndex())
+
+	Size += InOther->Size;
+
+	InOther->RemoveFromLinkedList();
+	InOther->Reset();
+}
+
+
+void FRHIPoolAllocationData::RemoveFromLinkedList()
+{
+	check(IsFree());
+	PreviousAllocation->NextAllocation = NextAllocation;
+	NextAllocation->PreviousAllocation = PreviousAllocation;
+}
+
+
+void FRHIPoolAllocationData::AddBefore(FRHIPoolAllocationData* InOther)
+{
+	PreviousAllocation->NextAllocation = InOther;
+	InOther->PreviousAllocation = PreviousAllocation;
+
+	PreviousAllocation = InOther;
+	InOther->NextAllocation = this;
+}
+
+
+void FRHIPoolAllocationData::AddAfter(FRHIPoolAllocationData* InOther)
+{
+	NextAllocation->PreviousAllocation = InOther;
+	InOther->NextAllocation = NextAllocation;
+
+	NextAllocation = InOther;
+	InOther->PreviousAllocation = this;
+}
+
+
+/**
+ * Sort predicate that sorts size of the allocation data
+ */
+static bool LinkedListSortBySizePredicate(const FRHIPoolAllocationData* InLeft, const FRHIPoolAllocationData* InRight)
+{
+	return InLeft->GetSize() < InRight->GetSize();
+}
+
+
+//-----------------------------------------------------------------------------
+//	LinkedList Allocator
+//-----------------------------------------------------------------------------
+
+static int32 DesiredAllocationPoolSize = 32;
+
+uint32 FRHIMemoryPool::GetAlignedSize(uint32 InSizeInBytes, uint32 InPoolAlignment, uint32 InAllocationAlignment)
+{
+	check(InAllocationAlignment <= InPoolAlignment);
+
+	// Compute the size, taking the pool and allocation alignmend into account (conservative size)
+	uint32 Size = ((InPoolAlignment % InAllocationAlignment) != 0) ? InSizeInBytes + InAllocationAlignment : InSizeInBytes;
+	return AlignArbitrary(Size, InPoolAlignment);
+}
+
+
+uint32 FRHIMemoryPool::GetAlignedOffset(uint32 InOffset, uint32 InPoolAlignment, uint32 InAllocationAlignment)
+{
+	uint32 AlignedOffset = InOffset;
+
+	// fix the offset with the requested alignment if needed
+	if ((InPoolAlignment % InAllocationAlignment) != 0)
+	{
+		uint32 AlignmentRest = AlignedOffset % InAllocationAlignment;
+		if (AlignmentRest > 0)
+		{
+			AlignedOffset += (InAllocationAlignment - AlignmentRest);
+		}
+	}
+
+	return AlignedOffset;
+}
+
+
+FRHIMemoryPool::FRHIMemoryPool(int16 InPoolIndex, uint64 InPoolSize, uint32 InPoolAlignment) :
+	PoolIndex(InPoolIndex), 
+	PoolSize(InPoolSize), 
+	PoolAlignment(InPoolAlignment), 
+	FreeSize(0), 
+	AligmnentWaste(0), 
+	AllocatedBlocks(0)
+{	
+}
+
+
+FRHIMemoryPool::~FRHIMemoryPool()
+{
+	// Can't currently validate because Engine is still leaking too much during shutdown sadly enough
+	/*
+	check(AllocatedBlocks == 0);
+	check(AligmnentWaste == 0);
+	check(FreeSize == PoolSize);
+	check(FreeBlocks.Num() == 1);
+	check(FreeBlocks[0]->GetSize() == PoolSize);
+	*/
+
+	Validate();
+}
+
+
+void FRHIMemoryPool::Init()
+{
+	FreeSize = PoolSize;
+
+	// Setup a few working free blocks	
+	AllocationDataPool.Reserve(DesiredAllocationPoolSize);
+	for (int32 Index = 0; Index < DesiredAllocationPoolSize; ++Index)
+	{
+		FRHIPoolAllocationData* AllocationData = new FRHIPoolAllocationData();
+		AllocationData->Reset();
+		AllocationDataPool.Add(AllocationData);
+	}
+
+	// Setup the special start and free block
+	HeadBlock.InitAsHead(PoolIndex);
+	FRHIPoolAllocationData* FreeBlock = GetNewAllocationData();
+	FreeBlock->InitAsFree(PoolIndex, PoolSize, PoolAlignment, 0);
+	HeadBlock.AddAfter(FreeBlock);
+	FreeBlocks.Add(FreeBlock);
+
+	Validate();
+}
+
+
+void FRHIMemoryPool::Destroy()
+{
+	// Validate state?
+}
+
+
+bool FRHIMemoryPool::TryAllocate(uint32 InSizeInBytes, uint32 InAllocationAlignment, FRHIPoolAllocationData& AllocationData)
+{
+	int32 FreeBlockIndex = FindFreeBlock(InSizeInBytes, InAllocationAlignment);
+	if (FreeBlockIndex != INDEX_NONE)
+	{
+		uint32 AlignedSize = GetAlignedSize(InSizeInBytes, PoolAlignment, InAllocationAlignment);
+
+		FRHIPoolAllocationData* FreeBlock = FreeBlocks[FreeBlockIndex];
+		check(FreeBlock->GetSize() >= AlignedSize);
+
+		// Remove from the free blocks because size will change and need sorted insert again then
+		FreeBlocks.RemoveAt(FreeBlockIndex);
+
+		// update private allocator data of new and free block
+		AllocationData.InitAsAllocated(InSizeInBytes, InAllocationAlignment, FreeBlock);
+		check((AllocationData.GetOffset() % InAllocationAlignment) == 0);
+
+		// Update working stats
+		check(FreeSize >= AlignedSize);
+		FreeSize -= AlignedSize;
+		AligmnentWaste += (AlignedSize - InSizeInBytes);
+		AllocatedBlocks++;
+
+		// Free block is empty then release otherwise sorted reinsert 
+		if (FreeBlock->GetSize() == 0)
+		{
+			FreeBlock->RemoveFromLinkedList();
+			ReleaseAllocationData(FreeBlock);
+
+			Validate();
+		}
+		else
+		{
+			AddToFreeBlocks(FreeBlock);
+		}
+		
+		return true;
+	}
+	else
+	{
+		return false;
+	}
+}
+
+
+void FRHIMemoryPool::Deallocate(FRHIPoolAllocationData& AllocationData)
+{
+	check(AllocationData.IsLocked());
+	check(PoolIndex == AllocationData.GetPoolIndex());
+
+	// Free block should not be locked anymore - can be reused immediatly when we get to actual pool deallocate
+	bool bLocked = false;
+	uint64 AllocationSize = AllocationData.GetSize();
+
+	FRHIPoolAllocationData* FreeBlock = GetNewAllocationData();
+	FreeBlock->MoveFrom(AllocationData, bLocked);
+	FreeBlock->MarkFree(PoolAlignment);
+		
+	// Update working stats
+	FreeSize += FreeBlock->GetSize();
+	AligmnentWaste -= FreeBlock->GetSize() - AllocationSize;
+	AllocatedBlocks--;
+
+	AddToFreeBlocks(FreeBlock);
+}
+
+
+void FRHIMemoryPool::TryClear(FRHIPoolAllocator* InAllocator, uint32 InMaxCopySize, uint32& CopySize, const TArray<FRHIMemoryPool*>& InTargetPools)
+{
+	FRHIPoolAllocationData* BlockToMove = HeadBlock.GetPrev();
+	while (BlockToMove != &HeadBlock && CopySize < InMaxCopySize)
+	{
+		FRHIPoolAllocationData* NextBlockToMove = BlockToMove->GetPrev();
+
+		// Only try to move if allocated (not free) and not 'locked' by another operation
+		if (BlockToMove->IsAllocated() && !BlockToMove->IsLocked() && BlockToMove->GetOwner())
+		{
+			uint64 SizeToAllocate = BlockToMove->GetSize();
+			uint32 AllocationAlignment = BlockToMove->GetAlignment();
+
+			FRHIPoolAllocationData TempTargetAllocation;
+
+			// Try find pool to move requested allocation (in desired order already)
+			for (FRHIMemoryPool* TargetPool : InTargetPools)
+			{
+				if (TargetPool->TryAllocate(SizeToAllocate, AllocationAlignment, TempTargetAllocation))
+				{
+					// RHI specific handling of the actual defrag request
+					InAllocator->HandleDefragRequest(BlockToMove, TempTargetAllocation);
+
+					// Increment the working copy size
+					CopySize += BlockToMove->GetSize();
+
+					break;
+				}
+			}
+		}
+
+		BlockToMove = NextBlockToMove;
+	}
+
+	Validate();	
+}
+
+
+int32 FRHIMemoryPool::FindFreeBlock(uint32 InSizeInBytes, uint32 InAllocationAlignment) const
+{
+	// Early out if total free size doesn't fit
+	if (FreeSize < InSizeInBytes)
+	{
+		return INDEX_NONE;
+	}
+
+	// TODO: lowerbound search
+	uint32 AlignedSize = GetAlignedSize(InSizeInBytes, PoolAlignment, InAllocationAlignment);
+
+	// Check all the free blocks (sorted by size, so take best fit)
+	for (int32 FreeBlockIndex = 0; FreeBlockIndex < FreeBlocks.Num(); ++FreeBlockIndex)
+	{
+		if (FreeBlocks[FreeBlockIndex]->GetSize() >= AlignedSize)
+		{
+			return FreeBlockIndex;
+		}
+	}
+
+	return INDEX_NONE;
+}
+
+
+void FRHIMemoryPool::RemoveFromFreeBlocks(FRHIPoolAllocationData* InFreeBlock)
+{
+	for (int32 FreeBlockIndex = 0; FreeBlockIndex < FreeBlocks.Num(); ++FreeBlockIndex)
+	{
+		if (FreeBlocks[FreeBlockIndex] == InFreeBlock)
+		{
+			FreeBlocks.RemoveAt(FreeBlockIndex, 1, false);
+			break;
+		}
+	}	
+}
+
+
+FRHIPoolAllocationData* FRHIMemoryPool::AddToFreeBlocks(FRHIPoolAllocationData* InFreeBlock)
+{
+	check(InFreeBlock->IsFree());
+	check(IsAligned(InFreeBlock->GetSize(), PoolAlignment));
+	check(IsAligned(InFreeBlock->GetOffset(), PoolAlignment));
+
+	FRHIPoolAllocationData* FreeBlock = InFreeBlock;
+
+	// Coalesce with previous?
+	if (FreeBlock->GetPrev()->IsFree() && !FreeBlock->GetPrev()->IsLocked())
+	{
+		FRHIPoolAllocationData* PrevFree = FreeBlock->GetPrev();
+		PrevFree->Merge(FreeBlock);
+		RemoveFromFreeBlocks(PrevFree);
+		ReleaseAllocationData(FreeBlock);
+
+		FreeBlock = PrevFree;
+	}
+
+	// Coalesce with next?
+	if (FreeBlock->GetNext()->IsFree() && !FreeBlock->GetNext()->IsLocked())
+	{
+		FRHIPoolAllocationData* NextFree = FreeBlock->GetNext();
+		FreeBlock->Merge(NextFree);
+		RemoveFromFreeBlocks(NextFree);
+		ReleaseAllocationData(NextFree);
+	}
+
+	// Sorted insert
+	int32 InsertIndex = Algo::LowerBound(FreeBlocks, FreeBlock, LinkedListSortBySizePredicate);
+	FreeBlocks.Insert(FreeBlock, InsertIndex);
+
+	Validate();
+
+	return FreeBlock;
+}
+
+
+FRHIPoolAllocationData* FRHIMemoryPool::GetNewAllocationData()
+{
+	return (AllocationDataPool.Num() > 0) ? AllocationDataPool.Pop(false) : new FRHIPoolAllocationData();
+}
+
+
+void FRHIMemoryPool::ReleaseAllocationData(FRHIPoolAllocationData* InData)
+{
+	InData->Reset();
+	if (AllocationDataPool.Num() >= DesiredAllocationPoolSize)
+	{
+		delete InData;
+	}
+	else
+	{
+		AllocationDataPool.Add(InData);
+	}
+}
+
+
+void FRHIMemoryPool::Validate()
+{
+	if (GRHIPoolAllocatorValidateLinkList <= 0)
+		return;
+
+	uint64 TotalFreeSize = 0;
+	uint64 TotalAllocatedSize = 0;
+	uint64 CurrentOffset = 0;
+
+	// Validate the linked list
+	FRHIPoolAllocationData* CurrentBlock = HeadBlock.GetNext();
+	while (CurrentBlock != &HeadBlock)
+	{
+		uint32 AlignedSize = GetAlignedSize(CurrentBlock->GetSize(), PoolAlignment, CurrentBlock->GetAlignment());
+		uint32 AlignedOffset = AlignDown(CurrentBlock->GetOffset(), PoolAlignment);
+
+		check(AlignedOffset == CurrentOffset);
+
+		if (CurrentBlock->IsAllocated())
+		{
+			TotalAllocatedSize += AlignedSize;
+		}
+		else if (CurrentBlock->IsFree())
+		{
+			TotalFreeSize += AlignedSize;
+		}
+		else
+		{
+			check(false);
+		}
+
+		CurrentOffset += AlignedSize;
+		CurrentBlock = CurrentBlock->GetNext();
+	}
+
+	check(TotalFreeSize + TotalAllocatedSize == PoolSize);
+	check(TotalFreeSize == FreeSize);
+	check(TotalAllocatedSize == GetUsedSize());
+
+	// Validate the free blocks
+	TotalFreeSize = 0;
+	for (int32 FreeBlockIndex = 0; FreeBlockIndex < FreeBlocks.Num(); ++FreeBlockIndex)
+	{
+		FRHIPoolAllocationData* FreeBlock = FreeBlocks[FreeBlockIndex];
+		check(!FreeBlock->GetPrev()->IsFree() || FreeBlock->GetPrev()->IsLocked());
+		check(!FreeBlock->GetNext()->IsFree() || FreeBlock->GetNext()->IsLocked());
+		check(FreeBlockIndex == (FreeBlocks.Num() - 1) || FreeBlock->GetSize() <= FreeBlocks[FreeBlockIndex + 1]->GetSize());
+		check(IsAligned(FreeBlock->GetOffset(), PoolAlignment));
+		TotalFreeSize += FreeBlock->GetSize();
+	}
+	check(TotalFreeSize == FreeSize);
+}
+
+
+//-----------------------------------------------------------------------------
+//	Pool Allocator
+//-----------------------------------------------------------------------------
+
+FRHIPoolAllocator::FRHIPoolAllocator(uint64 InPoolSize, uint32 InPoolAlignment, uint32 InMaxAllocationSize)	: 
+	PoolSize(InPoolSize), 
+	PoolAlignment(InPoolAlignment), 
+	MaxAllocationSize(InMaxAllocationSize),
+	TotalAllocatedBlocks(0)
+{}
+
+
+FRHIPoolAllocator::~FRHIPoolAllocator()
+{
+	Destroy();
+}
+
+
+void FRHIPoolAllocator::Initialize()
+{
+	Pools.Add(CreateNewPool(0));
+}
+
+
+void FRHIPoolAllocator::Destroy()
+{
+	for (int32 i = (Pools.Num() - 1); i >= 0; i--)
+	{
+		if (Pools[i] != nullptr)
+		{
+			Pools[i]->Destroy();
+			delete(Pools[i]);
+		}
+	}
+
+	Pools.Empty();
+}
+
+
+bool FRHIPoolAllocator::TryAllocateInternal(uint32 InSizeInBytes, uint32 InAllocationAlignment, FRHIPoolAllocationData& AllocationData)
+{
+	FScopeLock Lock(&CS);
+
+	for (int32 PoolIndex = 0; PoolIndex < Pools.Num(); PoolIndex++)
+	{
+		if (Pools[PoolIndex] && Pools[PoolIndex]->TryAllocate(InSizeInBytes, InAllocationAlignment, AllocationData))
+		{
+			return true;
+		}
+	}
+
+	// Find a free pool index
+	int16 NewPoolIndex = -1;
+	for (int32 PoolIndex = 0; PoolIndex < Pools.Num(); PoolIndex++)
+	{
+		if (Pools[PoolIndex] == nullptr)
+		{
+			NewPoolIndex = PoolIndex;
+			break;
+		}
+	}
+
+	if (NewPoolIndex >= 0)
+	{
+		Pools[NewPoolIndex] = CreateNewPool(NewPoolIndex);
+	}
+	else
+	{
+		NewPoolIndex = Pools.Num();
+		Pools.Add(CreateNewPool(NewPoolIndex));
+	}
+
+	return Pools[NewPoolIndex]->TryAllocate(InSizeInBytes, InAllocationAlignment, AllocationData);
+}
+
+
+void FRHIPoolAllocator::DeallocateInternal(FRHIPoolAllocationData& AllocationData)
+{
+	FScopeLock Lock(&CS);
+
+	check(AllocationData.IsAllocated());
+	check(AllocationData.GetPoolIndex() < Pools.Num());
+	Pools[AllocationData.GetPoolIndex()]->Deallocate(AllocationData);
+}
+
+
+void FRHIPoolAllocator::Defrag(uint32 InMaxCopySize, uint32& CurrentCopySize)
+{
+	FScopeLock Lock(&CS);
+
+	TArray<FRHIMemoryPool*> SortedPools;
+	SortedPools.Reserve(Pools.Num());
+	for (FRHIMemoryPool* Pool : Pools)
+	{
+		if (Pool && !Pool->IsFull())
+		{
+			SortedPools.Add(Pool);
+		}
+	}
+
+	// Need more than 1 pool at least
+	if (SortedPools.Num() < 2)
+	{
+		return;
+	}
+
+	Algo::Sort(SortedPools, [](const FRHIMemoryPool* InLHS, const FRHIMemoryPool* InRHS)
+		{
+			return InLHS->GetUsedSize() < InRHS->GetUsedSize();
+		});
+
+	// Build list all target pools (added reverse order because we want to move to the fullest pool first)
+	// Skip index 0 because that's the pool we want to clear to the target pools
+	TArray<FRHIMemoryPool*> TargetPools;
+	TargetPools.Reserve(SortedPools.Num());
+	for (int32 PoolIndex = SortedPools.Num() - 1; PoolIndex > 0; --PoolIndex)
+	{
+		TargetPools.Add(SortedPools[PoolIndex]);
+	}
+
+	for (int32 PoolIndex = 0; PoolIndex < SortedPools.Num() - 1; ++PoolIndex)
+	{
+		FRHIMemoryPool* PoolToClear = SortedPools[PoolIndex];
+		PoolToClear->TryClear(this, InMaxCopySize, CurrentCopySize, TargetPools);
+
+		// Remove last allocator since we will try and clear that one next		
+		TargetPools.Pop();
+
+		if (CurrentCopySize >= InMaxCopySize)
+		{
+			break;
+		}
+	}
+}
+
+
+void FRHIPoolAllocator::UpdateMemoryStats(uint32& IOMemoryAllocated, uint32& IOMemoryUsed, uint32& IOMemoryFree, uint32& IOAlignmentWaste, uint32& IOAllocatedPageCount, uint32& IOFullPageCount)
+{
+	FScopeLock Lock(&CS);
+
+	TotalAllocatedBlocks = 0;
+	for (FRHIMemoryPool* Pool : Pools)
+	{
+		if (Pool)
+		{
+			IOMemoryAllocated += Pool->GetPoolSize();
+			IOMemoryUsed += Pool->GetUsedSize();
+			IOMemoryFree += Pool->GetFreeSize();
+			IOAlignmentWaste += Pool->GetAlignmentWaste();
+			IOAllocatedPageCount++;
+			if (Pool->IsFull())
+			{
+				IOFullPageCount++;
+			}
+
+			TotalAllocatedBlocks += Pool->GetAllocatedBlocks();
+		}
+	}
+}
diff --git a/Engine/Source/Runtime/RHICore/Public/RHIPoolAllocator.h b/Engine/Source/Runtime/RHICore/Public/RHIPoolAllocator.h
new file mode 100644
--- /dev/null
+++ b/Engine/Source/Runtime/RHICore/Public/RHIPoolAllocator.h
@@ -0,0 +1,230 @@
+// Copyright Epic Games, Inc. All Rights Reserved.
+#pragma once
+
+#include "RHI.h"
+#include "RHIResources.h"
+
+// Pre declares
+class FRHIPoolAllocator;
+
+/**
+@brief Owner of the pool allocation data which needs to handle memory move operations
+**/
+class FRHIPoolResource
+{
+};
+
+
+/**
+@brief Pool allocator internal data
+**/
+struct RHICORE_API FRHIPoolAllocationData
+{
+public:
+
+	// Reset to unknown state
+	void Reset();
+
+	// Initialize options
+	void InitAsHead(int16 InPoolIndex);
+	void InitAsFree(int16 InPoolIndex, uint32 InSize, uint32 InAlignment, uint32 InOffset);
+	void InitAsAllocated(uint32 InSize, uint32 InAlignment, FRHIPoolAllocationData* InFree);
+	void MoveFrom(FRHIPoolAllocationData& InAllocated, bool InLocked);
+
+	// Free block operation
+	void MarkFree(uint32 InAlignment);
+	void Merge(FRHIPoolAllocationData* InOther);
+
+	// Linked list operation
+	void RemoveFromLinkedList();
+	void AddBefore(FRHIPoolAllocationData* InOther);
+	void AddAfter(FRHIPoolAllocationData* InOther);
+
+	// Type getters
+	bool IsHead() const { return GetAllocationType() == EAllocationType::Head; }
+	bool IsFree() const { return GetAllocationType() == EAllocationType::Free; }
+	bool IsAllocated() const { return GetAllocationType() == EAllocationType::Allocated; }
+
+	// Allocation data getters
+	uint32 GetSize() const { return Size; }
+	uint32 GetAlignment() const { return Alignment; }
+	uint32 GetOffset() const { return Offset; }
+	int16 GetPoolIndex() const { return PoolIndex; }
+
+	// Linked list getters
+	FRHIPoolAllocationData* GetNext() const { return NextAllocation; }
+	FRHIPoolAllocationData* GetPrev() const { return PreviousAllocation; }
+
+	// Owner getter/setter - opaque data
+	void SetOwner(FRHIPoolResource* InOwner) { check(Owner == nullptr); Owner = InOwner; }
+	FRHIPoolResource* GetOwner() const { return Owner; }
+
+	// Unlock can be public - no lock required
+	bool IsLocked() const { return (Locked == 1); }
+	void Unlock() { check(Locked == 1); Locked = 0; }
+
+private:
+
+	// Unlock data block again - actual lock on owned allocator should be taken then
+	friend class FRHIMemoryPool;
+	friend class FRHIPoolAllocator;
+	void Lock() { check(Locked == 0); Locked = 1; }
+
+	enum class EAllocationType : uint8
+	{
+		Unknown,
+		Free,
+		Allocated,
+		Head,
+	};
+	EAllocationType GetAllocationType() const { return (EAllocationType) Type; }
+	void SetAllocationType(EAllocationType InType) { Type = (uint32)InType; }
+
+	uint32 Size;
+	uint32 Alignment;
+	uint32 Offset;
+	uint32 PoolIndex : 16;
+	uint32 Type : 4;
+	// Allocation data is locked after creation until initial data is filled in to make sure it's not being
+	// defragmented then. If a block is locked for read/write from CPU side then it also needs be locked.
+	// While it's defragmented then the block is also locked so make sure to wait for the GPU before destroying it
+	uint32 Locked : 1;		
+	uint32 Unused : 11;
+
+	// TODO: Link list data could be stored as index into a fixed array managed by the memory pool. Then only a uin16 or uint32 is needed here to get the next and previous index
+	FRHIPoolAllocationData* PreviousAllocation;
+	FRHIPoolAllocationData* NextAllocation;
+
+	FRHIPoolResource* Owner;
+};
+
+
+//////////////////////////////////////////////////////////////////////////
+// FRHIMemoryPool
+//////////////////////////////////////////////////////////////////////////
+
+/**
+@brief Pool which stores for each allocation the previous and next allocation.
+	   Each block is either free or allocated.
+**/
+class RHICORE_API FRHIMemoryPool
+{
+public:
+
+	// Constructor
+	FRHIMemoryPool(int16 InPoolIndex, uint64 InPoolSize, uint32 InPoolAlignment);
+	virtual ~FRHIMemoryPool();
+
+	// Setup/Shutdown
+	virtual void Init();
+	virtual void Destroy();
+
+	// AllocationData functions: allocate, free, move, ...
+	bool TryAllocate(uint32 InSizeInBytes, uint32 InAllocationAlignment, FRHIPoolAllocationData& AllocationData);
+	void Deallocate(FRHIPoolAllocationData& AllocationData);
+
+	// Bookkeeping and clearing
+	void TryClear(FRHIPoolAllocator* InAllocator, uint32 InMaxCopySize, uint32& CopySize, const TArray<FRHIMemoryPool*>& InTargetPools);
+
+	// Getters
+	int16 GetPoolIndex() const { return PoolIndex; }
+	uint64 GetPoolSize() const { return PoolSize; }
+	uint64 GetFreeSize() const { return FreeSize; }
+	uint64 GetUsedSize() const { return (PoolSize - FreeSize); }
+	uint32 GetAlignmentWaste() const { return AligmnentWaste; }
+	uint32 GetAllocatedBlocks() const { return AllocatedBlocks; }
+	bool IsEmpty() const { return GetUsedSize() == 0; }
+	bool IsFull() const { return FreeSize == 0; }
+
+	// Static public helper functions
+	static uint32 GetAlignedSize(uint32 InSizeInBytes, uint32 InPoolAlignment, uint32 InAllocationAlignment);
+	static uint32 GetAlignedOffset(uint32 InOffset, uint32 InPoolAlignment, uint32 InAllocationAlignment);
+
+protected:
+		
+	// Free block management helper function
+	int32 FindFreeBlock(uint32 InSizeInBytes, uint32 InAllocationAlignment) const;
+	FRHIPoolAllocationData* AddToFreeBlocks(FRHIPoolAllocationData* InFreeBlock);
+	void RemoveFromFreeBlocks(FRHIPoolAllocationData* InFreeBlock);
+
+	// Get/Release allocation data blocks
+	FRHIPoolAllocationData* GetNewAllocationData();
+	void ReleaseAllocationData(FRHIPoolAllocationData* InData);
+
+	// Validate the linked list data - slow
+	void Validate();
+
+	// Const creation members
+	const int16 PoolIndex;
+	const uint64 PoolSize;
+	const uint32 PoolAlignment;
+
+	// Stats
+	uint64 FreeSize;
+	uint64 AligmnentWaste;
+	uint32 AllocatedBlocks;
+
+	// Special start block to mark the beginning and the end of the linked list
+	FRHIPoolAllocationData HeadBlock;
+
+	// Actual Free blocks
+	TArray<FRHIPoolAllocationData*> FreeBlocks;
+
+	// Allocation data pool used to reduce allocation overhead
+	TArray<FRHIPoolAllocationData*> AllocationDataPool;
+};
+
+
+//////////////////////////////////////////////////////////////////////////
+// FRHIPoolAllocator
+//////////////////////////////////////////////////////////////////////////
+
+/**
+@brief Manages an array of FRHIMemoryPool and supports defragmentation between multiple pools
+**/
+class RHICORE_API FRHIPoolAllocator
+{
+public:
+
+	// Constructor
+	FRHIPoolAllocator(uint64 InPoolSize, uint32 InPoolAlignment, uint32 InMaxAllocationSize);
+	virtual ~FRHIPoolAllocator();
+
+	// Setup/Shutdown
+	void Initialize();
+	void Destroy();
+	
+	// Defrag & cleanup operation
+	void Defrag(uint32 InMaxCopySize, uint32& CurrentCopySize);
+
+	// Stats
+	void UpdateMemoryStats(uint32& IOMemoryAllocated, uint32& IOMemoryUsed, uint32& IOMemoryFree, uint32& IOAlignmentWaste, uint32& IOAllocatedPageCount, uint32& IOFullPageCount);
+
+protected:
+
+	// Supported allocator operations
+	bool TryAllocateInternal(uint32 InSizeInBytes, uint32 InAllocationAlignment, FRHIPoolAllocationData& AllocationData);
+	void DeallocateInternal(FRHIPoolAllocationData& AllocationData);
+	
+	// Helper function to create a new platform specific pool
+	virtual FRHIMemoryPool* CreateNewPool(int16 InPoolIndex) = 0;
+
+	// Handle a rhi specific defrag op
+	friend class FRHIMemoryPool;
+	virtual bool HandleDefragRequest(FRHIPoolAllocationData* InSourceBlock, FRHIPoolAllocationData& InTmpTargetBlock) = 0;
+
+	// Const creation members - used to create new pools
+	const uint64 PoolSize;
+	const uint32 PoolAlignment;
+	const uint64 MaxAllocationSize;
+
+	// Critical section to lock access to the pools
+	FCriticalSection CS;
+
+	// Actual pools managing each a linked list of allocations
+	TArray<FRHIMemoryPool*> Pools;
+
+	// Stats
+	uint32 TotalAllocatedBlocks;
+};
+
