Fixing PVS 7.7 Issues:  warning V1050: The uninitialized class member 'XYZ' is used when initializing the base class 'FXYZ'.

In this specific case, ViewLocation is being improperly accessed prior to initialization.  And the argument list to the base class constructor is wrong.

diff --git a/Engine/Source/Runtime/D3D12RHI/Private/D3D12View.h b/Engine/Source/Runtime/D3D12RHI/Private/D3D12View.h
--- a/Engine/Source/Runtime/D3D12RHI/Private/D3D12View.h
+++ b/Engine/Source/Runtime/D3D12RHI/Private/D3D12View.h
@@ -906,9 +906,10 @@ public:
 	{}
 
 	FD3D12ShaderResourceViewWithLocation(FD3D12Device* InParent, D3D12_SHADER_RESOURCE_VIEW_DESC& InDesc, FD3D12ResourceLocation& InResourceLocation, uint32 InStride = -1, bool InSkipFastClearFinalize = false)
-		: FD3D12ShaderResourceView(InParent, InDesc, ViewLocation, InStride, InSkipFastClearFinalize)
+		: FD3D12ShaderResourceView(InParent)
 		, ViewLocation(InParent)
 	{
+		Initialize(InDesc, ViewLocation, InStride, 0, InSkipFastClearFinalize);
 	}
 
 	FD3D12ResourceLocation ViewLocation;
