Skip redundant transition barrier when creating buffers with CopyDest initial state.

diff --git a/Engine/Source/Runtime/D3D12RHI/Private/D3D12Buffer.cpp b/Engine/Source/Runtime/D3D12RHI/Private/D3D12Buffer.cpp
--- a/Engine/Source/Runtime/D3D12RHI/Private/D3D12Buffer.cpp
+++ b/Engine/Source/Runtime/D3D12RHI/Private/D3D12Buffer.cpp
@@ -109,7 +109,10 @@ struct FD3D12RHICommandInitializeBuffer final : public FRHICommand<FD3D12RHIComm
 					SrcResourceLoc.GetOffsetFromBaseOfResource(), Size);
 
 				// Update the resource state after the copy has been done (will take care of updating the residency as well)
-				hCommandList.AddTransitionBarrier(Destination, D3D12_RESOURCE_STATE_COPY_DEST, DestinationState, D3D12_RESOURCE_BARRIER_ALL_SUBRESOURCES);
+				if (DestinationState != D3D12_RESOURCE_STATE_COPY_DEST)
+				{
+					hCommandList.AddTransitionBarrier(Destination, D3D12_RESOURCE_STATE_COPY_DEST, DestinationState, D3D12_RESOURCE_BARRIER_ALL_SUBRESOURCES);
+				}
 
 				if (Destination->RequiresResourceStateTracking())
 				{
