Implemented support for indirect ray tracing dispatch in D3D12 / DXR 1.1.

diff --git a/Engine/Shaders/Private/RayTracing/RayTracingDispatchDesc.usf b/Engine/Shaders/Private/RayTracing/RayTracingDispatchDesc.usf
new file mode 100644
--- /dev/null
+++ b/Engine/Shaders/Private/RayTracing/RayTracingDispatchDesc.usf
@@ -0,0 +1,38 @@
+// Copyright Epic Games, Inc. All Rights Reserved.
+
+/*==========================================================================
+	RayTracingDispatchDesc.usf: Generate indirect ray tracing dispatch args 
+	using header from CPU (uploaded as constants) and dimensions from GPU.
+==========================================================================*/
+
+#include "../Common.ush"
+
+#ifndef DISPATCH_DESC_MAX_SIZE_DWORDS
+#define DISPATCH_DESC_MAX_SIZE_DWORDS 32
+#endif // DISPATCH_DESC_MAX_SIZE_DWORDS
+
+#define DISPATCH_DESC_MAX_SIZE_UINT4S (DISPATCH_DESC_MAX_SIZE_DWORDS/4)
+
+uint4 DispatchDescInput[DISPATCH_DESC_MAX_SIZE_UINT4S];
+uint DispatchDescSizeDwords;
+uint DispatchDescDimensionsOffsetDwords;
+uint DimensionsBufferOffsetDwords;
+Buffer<uint> DispatchDimensions;
+RWBuffer<uint> DispatchDescOutput;
+
+[numthreads(DISPATCH_DESC_MAX_SIZE_DWORDS, 1, 1)]
+void RayTracingDispatchDescCS(uint i : SV_DispatchThreadID)
+{
+	if (i < DispatchDescSizeDwords)
+	{
+		if (i >= DispatchDescDimensionsOffsetDwords && i < DispatchDescDimensionsOffsetDwords + 3)
+		{
+			uint DimensionIndex = i - DispatchDescDimensionsOffsetDwords;
+			DispatchDescOutput[i] = DispatchDimensions[DimensionsBufferOffsetDwords + DimensionIndex];
+		}
+		else
+		{
+			DispatchDescOutput[i] = DispatchDescInput[i/4][i%4];
+		}
+	}
+}
diff --git a/Engine/Source/Runtime/D3D12RHI/Private/D3D12Adapter.cpp b/Engine/Source/Runtime/D3D12RHI/Private/D3D12Adapter.cpp
--- a/Engine/Source/Runtime/D3D12RHI/Private/D3D12Adapter.cpp
+++ b/Engine/Source/Runtime/D3D12RHI/Private/D3D12Adapter.cpp
@@ -806,6 +806,7 @@ void FD3D12Adapter::InitializeDevices()
 						UE_LOG(LogD3D12RHI, Log, TEXT("D3D12 ray tracing 1.1 is supported."));
 
 						GRHISupportsRayTracingPSOAdditions = true;
+						GRHISupportsRayTracingDispatchIndirect = true;
 					}
 				}
 			}
@@ -936,6 +937,22 @@ void FD3D12Adapter::CreateSignatures()
 	indirectParameterDesc[0].Type = D3D12XBOX_INDIRECT_ARGUMENT_TYPE_DISPATCH_L2; // D3D12XBOX_INDIRECT_ARGUMENT_TYPE_DISPATCH_SINGLE
 #endif
 	VERIFYD3D12RESULT(Device->CreateCommandSignature(&commandSignatureDesc, nullptr, IID_PPV_ARGS(DispatchIndirectComputeCommandSignature.GetInitReference())));
+
+#if D3D12_RHI_RAYTRACING
+	if (GRHISupportsRayTracingDispatchIndirect)
+	{
+		D3D12_COMMAND_SIGNATURE_DESC SignatureDesc = {};
+		SignatureDesc.NumArgumentDescs = 1;
+		SignatureDesc.ByteStride = sizeof(D3D12_DISPATCH_RAYS_DESC);
+		SignatureDesc.NodeMask = FRHIGPUMask::All().GetNative();
+
+		D3D12_INDIRECT_ARGUMENT_DESC ArgumentDesc[1] = {};
+		ArgumentDesc[0].Type = D3D12_INDIRECT_ARGUMENT_TYPE_DISPATCH_RAYS;
+		SignatureDesc.pArgumentDescs = ArgumentDesc;
+
+		VERIFYD3D12RESULT(Device->CreateCommandSignature(&SignatureDesc, nullptr, IID_PPV_ARGS(DispatchRaysIndirectCommandSignature.GetInitReference())));
+	}
+#endif // D3D12_RHI_RAYTRACING
 }
 
 
@@ -1017,6 +1034,7 @@ void FD3D12Adapter::Cleanup()
 	DrawIndexedIndirectCommandSignature.SafeRelease();
 	DispatchIndirectGraphicsCommandSignature.SafeRelease();
 	DispatchIndirectComputeCommandSignature.SafeRelease();
+	DispatchRaysIndirectCommandSignature.SafeRelease();
 
 	FenceCorePool.Destroy();
 
diff --git a/Engine/Source/Runtime/D3D12RHI/Private/D3D12Adapter.h b/Engine/Source/Runtime/D3D12RHI/Private/D3D12Adapter.h
--- a/Engine/Source/Runtime/D3D12RHI/Private/D3D12Adapter.h
+++ b/Engine/Source/Runtime/D3D12RHI/Private/D3D12Adapter.h
@@ -126,6 +126,7 @@ public:
 	FORCEINLINE ID3D12CommandSignature* GetDrawIndexedIndirectCommandSignature() { return DrawIndexedIndirectCommandSignature; }
 	FORCEINLINE ID3D12CommandSignature* GetDispatchIndirectGraphicsCommandSignature() { return DispatchIndirectGraphicsCommandSignature; }
 	FORCEINLINE ID3D12CommandSignature* GetDispatchIndirectComputeCommandSignature() { return DispatchIndirectComputeCommandSignature; }
+	FORCEINLINE ID3D12CommandSignature* GetDispatchRaysIndirectCommandSignature() { return DispatchRaysIndirectCommandSignature; }
 
 	FORCEINLINE FD3D12PipelineStateCache& GetPSOCache() { return PipelineStateCache; }
 
@@ -376,6 +377,7 @@ protected:
 	TRefCountPtr<ID3D12CommandSignature> DrawIndexedIndirectCommandSignature;
 	TRefCountPtr<ID3D12CommandSignature> DispatchIndirectGraphicsCommandSignature;
 	TRefCountPtr<ID3D12CommandSignature> DispatchIndirectComputeCommandSignature;
+	TRefCountPtr<ID3D12CommandSignature> DispatchRaysIndirectCommandSignature;
 
 	FD3D12FenceCorePool FenceCorePool;
 
diff --git a/Engine/Source/Runtime/D3D12RHI/Private/D3D12CommandContext.h b/Engine/Source/Runtime/D3D12RHI/Private/D3D12CommandContext.h
--- a/Engine/Source/Runtime/D3D12RHI/Private/D3D12CommandContext.h
+++ b/Engine/Source/Runtime/D3D12RHI/Private/D3D12CommandContext.h
@@ -380,6 +380,10 @@ public:
 		FRHIRayTracingScene* Scene,
 		const FRayTracingShaderBindings& GlobalResourceBindings,
 		uint32 Width, uint32 Height) final override;
+	virtual void RHIRayTraceDispatchIndirect(FRHIRayTracingPipelineState* RayTracingPipelineState, FRHIRayTracingShader* RayGenShader,
+		FRHIRayTracingScene* Scene,
+		const FRayTracingShaderBindings& GlobalResourceBindings,
+		FRHIBuffer* ArgumentBuffer, uint32 ArgumentOffset) final override;
 	virtual void RHISetRayTracingHitGroups(
 		FRHIRayTracingScene* Scene, FRHIRayTracingPipelineState* Pipeline,
 		uint32 NumBindings, const FRayTracingLocalShaderBindings* Bindings) final override;
@@ -738,6 +742,14 @@ public:
 		ContextRedirect(RHIRayTraceDispatch(RayTracingPipelineState, RayGenShader, Scene, GlobalResourceBindings, Width, Height));
 	}
 
+	virtual void RHIRayTraceDispatchIndirect(FRHIRayTracingPipelineState* RayTracingPipelineState, FRHIRayTracingShader* RayGenShader,
+		FRHIRayTracingScene* Scene,
+		const FRayTracingShaderBindings& GlobalResourceBindings,
+		FRHIBuffer* ArgumentBuffer, uint32 ArgumentOffset) final override
+	{
+		ContextRedirect(RHIRayTraceDispatchIndirect(RayTracingPipelineState, RayGenShader, Scene, GlobalResourceBindings, ArgumentBuffer, ArgumentOffset));
+	}
+
 	virtual void RHISetRayTracingHitGroup(
 		FRHIRayTracingScene* Scene, uint32 InstanceIndex, uint32 SegmentIndex, uint32 ShaderSlot,
 		FRHIRayTracingPipelineState* Pipeline, uint32 HitGroupIndex,
diff --git a/Engine/Source/Runtime/D3D12RHI/Private/D3D12RayTracing.cpp b/Engine/Source/Runtime/D3D12RHI/Private/D3D12RayTracing.cpp
--- a/Engine/Source/Runtime/D3D12RHI/Private/D3D12RayTracing.cpp
+++ b/Engine/Source/Runtime/D3D12RHI/Private/D3D12RayTracing.cpp
@@ -13,7 +13,7 @@
 #include "HAL/CriticalSection.h"
 #include "HAL/IConsoleManager.h"
 #include "Misc/ScopeLock.h"
-#include "RayTracingInstanceCopyShader.h" 
+#include "RayTracingInstanceCopyShader.h"
 
 static int32 GRayTracingDebugForceOpaque = 0;
 static FAutoConsoleVariableRef CVarRayTracingDebugForceOpaque(
@@ -3862,9 +3862,49 @@ static void DispatchRays(FD3D12CommandContext& CommandContext,
 	const FD3D12RayTracingPipelineState* Pipeline,
 	uint32 RayGenShaderIndex,
 	FD3D12RayTracingShaderTable* OptShaderTable,
-	const D3D12_DISPATCH_RAYS_DESC& DispatchDesc)
+	const D3D12_DISPATCH_RAYS_DESC& DispatchDesc,
+	FD3D12Buffer* ArgumentBuffer = nullptr, uint32 ArgumentOffset = 0)
 {
-	 SCOPE_CYCLE_COUNTER(STAT_D3D12DispatchRays);
+	SCOPE_CYCLE_COUNTER(STAT_D3D12DispatchRays);
+
+	FD3D12Adapter* Adapter = CommandContext.GetParentDevice()->GetParentAdapter();
+
+	TRefCountPtr<FD3D12Buffer> DispatchRaysDescBuffer;
+
+	if (ArgumentBuffer)
+	{
+		FD3D12CommandListHandle& CommandListHandle = CommandContext.CommandListHandle;
+
+		// Source indirect argument buffer only contains the dispatch dimensions, however D3D12 requires a full D3D12_DISPATCH_RAYS_DESC structure.
+		// We create a new buffer, fill the SBT pointers on CPU and copy the dispatch dimensions into the right place.
+
+		D3D12_RESOURCE_DESC BufferDesc = CD3DX12_RESOURCE_DESC::Buffer(
+			sizeof(D3D12_DISPATCH_RAYS_DESC),
+			D3D12_RESOURCE_FLAG_ALLOW_UNORDERED_ACCESS,
+			D3D12_DEFAULT_RESOURCE_PLACEMENT_ALIGNMENT);
+
+		FRHIResourceCreateInfo DispatchRaysDescBufferCreateInfo;
+		DispatchRaysDescBufferCreateInfo.GPUMask = FRHIGPUMask::FromIndex(CommandContext.GetGPUIndex());
+
+		DispatchRaysDescBuffer = Adapter->CreateRHIBuffer(
+			nullptr, BufferDesc, BufferDesc.Alignment,
+			4, BufferDesc.Width, BUF_Static | BUF_DrawIndirect | BUF_UnorderedAccess,
+			ED3D12ResourceStateMode::MultiState, ERHIAccess::UAVCompute,
+			DispatchRaysDescBufferCreateInfo);
+
+		FD3D12Resource* DispatchRaysDescBufferResource = DispatchRaysDescBuffer->GetResource();
+
+		SetName(DispatchRaysDescBufferResource, TEXT("DispatchRaysDescBuffer"));
+
+		{
+			TRHICommandList_RecursiveHazardous<FD3D12CommandContext> RHICmdList(&CommandContext);
+			FShaderResourceViewRHIRef SRV = RHICreateShaderResourceView(ArgumentBuffer, 4, PF_R32_UINT);
+			FUnorderedAccessViewRHIRef UAV = RHICreateUnorderedAccessView(DispatchRaysDescBuffer, PF_R32_UINT);
+			FRayTracingDispatchDescCS::Dispatch(RHICmdList, &DispatchDesc, sizeof(DispatchDesc), offsetof(D3D12_DISPATCH_RAYS_DESC, Width), SRV, ArgumentOffset, UAV);
+		}
+
+		CommandListHandle.AddTransitionBarrier(DispatchRaysDescBufferResource, D3D12_RESOURCE_STATE_UNORDERED_ACCESS, D3D12_RESOURCE_STATE_INDIRECT_ARGUMENT, 0);
+	}
 
 	// Setup state for RT dispatch
 	
@@ -3918,7 +3958,23 @@ static void DispatchRays(FD3D12CommandContext& CommandContext,
 
 		ID3D12GraphicsCommandList4* RayTracingCommandList = CommandContext.CommandListHandle.RayTracingCommandList();
 		RayTracingCommandList->SetPipelineState1(RayTracingStateObject);
-		RayTracingCommandList->DispatchRays(&DispatchDesc);
+
+		if (DispatchRaysDescBuffer)
+		{
+			ID3D12CommandSignature* CommandSignature = Adapter->GetDispatchRaysIndirectCommandSignature();
+			RayTracingCommandList->ExecuteIndirect(
+				CommandSignature,
+				1,
+				DispatchRaysDescBuffer->ResourceLocation.GetResource()->GetResource(),
+				DispatchRaysDescBuffer->ResourceLocation.GetOffsetFromBaseOfResource(),
+				nullptr,
+				0
+			);
+		}
+		else
+		{
+			RayTracingCommandList->DispatchRays(&DispatchDesc);
+		}
 
 		if (CommandContext.IsDefaultContext())
 		{
@@ -4040,6 +4096,34 @@ void FD3D12CommandContext::RHIRayTraceDispatch(FRHIRayTracingPipelineState* InRa
 	DispatchRays(*this, GlobalResourceBindings, Pipeline, RayGenShaderIndex, ShaderTable, DispatchDesc);
 }
 
+void FD3D12CommandContext::RHIRayTraceDispatchIndirect(FRHIRayTracingPipelineState* InRayTracingPipelineState, FRHIRayTracingShader* RayGenShaderRHI,
+	FRHIRayTracingScene* InScene,
+	const FRayTracingShaderBindings& GlobalResourceBindings,
+	FRHIBuffer* ArgumentBuffer, uint32 ArgumentOffset)
+{
+	checkf(GRHISupportsRayTracingDispatchIndirect, TEXT("RHIRayTraceDispatchIndirect may not be used because DXR 1.1 is not supported on this machine."));
+
+	const FD3D12RayTracingPipelineState* Pipeline = FD3D12DynamicRHI::ResourceCast(InRayTracingPipelineState);
+
+	FD3D12RayTracingScene* Scene = FD3D12DynamicRHI::ResourceCast(InScene);
+
+	FD3D12RayTracingShaderTable* ShaderTable = Scene->FindOrCreateShaderTable(Pipeline, GetParentDevice());
+
+	if (ShaderTable->bIsDirty)
+	{
+		ShaderTable->CopyToGPU(GetParentDevice());
+	}
+
+	Scene->UpdateResidency(*this);
+
+	FD3D12RayTracingShader* RayGenShader = FD3D12DynamicRHI::ResourceCast(RayGenShaderRHI);
+	const int32 RayGenShaderIndex = Pipeline->RayGenShaders.Find(RayGenShader->GetHash());
+	checkf(RayGenShaderIndex != INDEX_NONE, TEXT("RayGen shader is not present in the given ray tracing pipeline. All RayGen shaders must be declared when creating RTPSO."));
+
+	D3D12_DISPATCH_RAYS_DESC DispatchDesc = ShaderTable->GetDispatchRaysDesc(RayGenShaderIndex, 0, Pipeline->bAllowHitGroupIndexing);
+	DispatchRays(*this, GlobalResourceBindings, Pipeline, RayGenShaderIndex, ShaderTable, DispatchDesc, RetrieveObject<FD3D12Buffer>(ArgumentBuffer), ArgumentOffset);
+}
+
 static void SetRayTracingHitGroup(
 	FD3D12Device* Device,
 	FD3D12RayTracingShaderTable* ShaderTable,
diff --git a/Engine/Source/Runtime/RHI/Private/RHI.cpp b/Engine/Source/Runtime/RHI/Private/RHI.cpp
--- a/Engine/Source/Runtime/RHI/Private/RHI.cpp
+++ b/Engine/Source/Runtime/RHI/Private/RHI.cpp
@@ -843,6 +843,7 @@ bool GRHISupportsFirstInstance = false;
 bool GRHISupportsDynamicResolution = false;
 bool GRHISupportsRayTracing = false;
 bool GRHISupportsRayTracingPSOAdditions = false;
+bool GRHISupportsRayTracingDispatchIndirect = false;
 bool GRHISupportsRayTracingAsyncBuildAccelerationStructure = false;
 bool GRHISupportsRayTracingAMDHitToken = false;
 bool GRHISupportsWaveOperations = false;
diff --git a/Engine/Source/Runtime/RHI/Public/RHI.h b/Engine/Source/Runtime/RHI/Public/RHI.h
--- a/Engine/Source/Runtime/RHI/Public/RHI.h
+++ b/Engine/Source/Runtime/RHI/Public/RHI.h
@@ -567,6 +567,9 @@ extern RHI_API bool GRHISupportsRayTracing;
 /** Whether or not the RHI supports adding new shaders to an existing RT PSO. */
 extern RHI_API bool GRHISupportsRayTracingPSOAdditions;
 
+/** Whether or not the RHI supports indirect ray tracing dispatch commands. */
+extern RHI_API bool GRHISupportsRayTracingDispatchIndirect;
+
 /** Whether or not the RHI supports async building ray tracing acceleration structures. */
 extern RHI_API bool GRHISupportsRayTracingAsyncBuildAccelerationStructure;
 
diff --git a/Engine/Source/Runtime/RHI/Public/RHICommandList.h b/Engine/Source/Runtime/RHI/Public/RHICommandList.h
--- a/Engine/Source/Runtime/RHI/Public/RHICommandList.h
+++ b/Engine/Source/Runtime/RHI/Public/RHICommandList.h
@@ -2173,6 +2173,8 @@ FRHICOMMAND_MACRO(FRHICommandRayTraceDispatch)
 	FRHIRayTracingScene* Scene;
 	FRayTracingShaderBindings GlobalResourceBindings;
 	FRHIRayTracingShader* RayGenShader;
+	FRHIBuffer* ArgumentBuffer;
+	uint32 ArgumentOffset;
 	uint32 Width;
 	uint32 Height;
 
@@ -2181,10 +2183,23 @@ FRHICOMMAND_MACRO(FRHICommandRayTraceDispatch)
 		, Scene(InScene)
 		, GlobalResourceBindings(InGlobalResourceBindings)
 		, RayGenShader(InRayGenShader)
+		, ArgumentBuffer(nullptr)
+		, ArgumentOffset(0)
 		, Width(InWidth)
 		, Height(InHeight)
 	{}
 
+	FRHICommandRayTraceDispatch(FRayTracingPipelineState* InPipeline, FRHIRayTracingShader* InRayGenShader, FRHIRayTracingScene* InScene, const FRayTracingShaderBindings& InGlobalResourceBindings, FRHIBuffer* InArgumentBuffer, uint32 InArgumentOffset)
+		: Pipeline(InPipeline)
+		, Scene(InScene)
+		, GlobalResourceBindings(InGlobalResourceBindings)
+		, RayGenShader(InRayGenShader)
+		, ArgumentBuffer(InArgumentBuffer)
+		, ArgumentOffset(InArgumentOffset)
+		, Width(0)
+		, Height(0)
+	{}
+
 	RHI_API void Execute(FRHICommandListBase& CmdList);
 };
 
@@ -3678,6 +3693,24 @@ public:
 		}
 	}
 
+	/**
+	 * Trace rays using dimensions from a GPU buffer containing uint[3], interpreted as number of rays in X, Y and Z dimensions.
+	 * ArgumentBuffer must be in ERHIAccess::ReadOnlyExclusiveMask state rather than simply IndirectArgs since it's read as a regular GPU resource on some platforms.
+	 */
+	FORCEINLINE_DEBUGGABLE void RayTraceDispatchIndirect(FRayTracingPipelineState* Pipeline, FRHIRayTracingShader* RayGenShader, FRHIRayTracingScene* Scene, const FRayTracingShaderBindings& GlobalResourceBindings, FRHIBuffer* ArgumentBuffer, uint32 ArgumentOffset)
+	{
+		checkf(GRHISupportsRayTracingDispatchIndirect, TEXT("Indirect ray tracing is not supported on this machine."));
+		if (Bypass())
+		{
+			extern RHI_API FRHIRayTracingPipelineState* GetRHIRayTracingPipelineState(FRayTracingPipelineState*);
+			GetContext().RHIRayTraceDispatchIndirect(GetRHIRayTracingPipelineState(Pipeline), RayGenShader, Scene, GlobalResourceBindings, ArgumentBuffer, ArgumentOffset);
+		}
+		else
+		{
+			ALLOC_COMMAND(FRHICommandRayTraceDispatch)(Pipeline, RayGenShader, Scene, GlobalResourceBindings, ArgumentBuffer, ArgumentOffset);
+		}
+	}
+
 	FORCEINLINE_DEBUGGABLE void SetRayTracingHitGroups(
 		FRHIRayTracingScene* Scene, FRayTracingPipelineState* Pipeline,
 		uint32 NumBindings, const FRayTracingLocalShaderBindings* Bindings,
diff --git a/Engine/Source/Runtime/RHI/Public/RHICommandListCommandExecutes.inl b/Engine/Source/Runtime/RHI/Public/RHICommandListCommandExecutes.inl
--- a/Engine/Source/Runtime/RHI/Public/RHICommandListCommandExecutes.inl
+++ b/Engine/Source/Runtime/RHI/Public/RHICommandListCommandExecutes.inl
@@ -545,7 +545,14 @@ void FRHICommandRayTraceDispatch::Execute(FRHICommandListBase& CmdList)
 {
 	RHISTAT(RayTraceDispatch);
 	extern RHI_API FRHIRayTracingPipelineState* GetRHIRayTracingPipelineState(FRayTracingPipelineState*);
-	INTERNAL_DECORATOR(RHIRayTraceDispatch)(GetRHIRayTracingPipelineState(Pipeline), RayGenShader, Scene, GlobalResourceBindings, Width, Height);
+	if (ArgumentBuffer)
+	{
+		INTERNAL_DECORATOR(RHIRayTraceDispatchIndirect)(GetRHIRayTracingPipelineState(Pipeline), RayGenShader, Scene, GlobalResourceBindings, ArgumentBuffer, ArgumentOffset);
+	}
+	else
+	{
+		INTERNAL_DECORATOR(RHIRayTraceDispatch)(GetRHIRayTracingPipelineState(Pipeline), RayGenShader, Scene, GlobalResourceBindings, Width, Height);
+	}
 }
 
 void FRHICommandSetRayTracingBindings::Execute(FRHICommandListBase& CmdList)
diff --git a/Engine/Source/Runtime/RHI/Public/RHIContext.h b/Engine/Source/Runtime/RHI/Public/RHIContext.h
--- a/Engine/Source/Runtime/RHI/Public/RHIContext.h
+++ b/Engine/Source/Runtime/RHI/Public/RHIContext.h
@@ -668,6 +668,14 @@ public:
 		checkNoEntry();
 	}
 
+	virtual void RHIRayTraceDispatchIndirect(FRHIRayTracingPipelineState* RayTracingPipelineState, FRHIRayTracingShader* RayGenShader,
+		FRHIRayTracingScene* Scene,
+		const FRayTracingShaderBindings& GlobalResourceBindings,
+		FRHIBuffer* ArgumentBuffer, uint32 ArgumentOffset)
+	{
+		checkNoEntry();
+	}
+
 	virtual void RHISetRayTracingHitGroups(FRHIRayTracingScene* Scene, FRHIRayTracingPipelineState* Pipeline, uint32 NumBindings, const FRayTracingLocalShaderBindings* Bindings)
 	{
 		checkNoEntry();
diff --git a/Engine/Source/Runtime/RHI/Public/RHIValidationContext.h b/Engine/Source/Runtime/RHI/Public/RHIValidationContext.h
--- a/Engine/Source/Runtime/RHI/Public/RHIValidationContext.h
+++ b/Engine/Source/Runtime/RHI/Public/RHIValidationContext.h
@@ -45,7 +45,7 @@ public:
 
 	virtual void RHIDispatchIndirectComputeShader(FRHIBuffer* ArgumentBuffer, uint32 ArgumentOffset) override final
 	{
-		Tracker->Assert(ArgumentBuffer->GetWholeResourceIdentity(), ERHIAccess::IndirectArgs);
+		Tracker->Assert(ArgumentBuffer->GetWholeResourceIdentity(), ERHIAccess::ReadOnlyExclusiveMask);
 
 		RHIContext->RHIDispatchIndirectComputeShader(ArgumentBuffer, ArgumentOffset);
 		Tracker->Dispatch();
@@ -921,6 +921,15 @@ public:
 		RHIContext->RHIRayTraceDispatch(RayTracingPipelineState, RayGenShader, Scene, GlobalResourceBindings, Width, Height);
 	}
 
+	virtual void RHIRayTraceDispatchIndirect(FRHIRayTracingPipelineState* RayTracingPipelineState, FRHIRayTracingShader* RayGenShader,
+		FRHIRayTracingScene* Scene,
+		const FRayTracingShaderBindings& GlobalResourceBindings,
+		FRHIBuffer* ArgumentBuffer, uint32 ArgumentOffset) override final
+	{
+		Tracker->Assert(ArgumentBuffer->GetWholeResourceIdentity(), ERHIAccess::IndirectArgs);
+		RHIContext->RHIRayTraceDispatchIndirect(RayTracingPipelineState, RayGenShader, Scene, GlobalResourceBindings, ArgumentBuffer, ArgumentOffset);
+	}
+
 	virtual void RHISetRayTracingHitGroups(FRHIRayTracingScene* Scene, FRHIRayTracingPipelineState* Pipeline, uint32 NumBindings, const FRayTracingLocalShaderBindings* Bindings) override final
 	{
 		RHIContext->RHISetRayTracingHitGroups(Scene, Pipeline, NumBindings, Bindings);
diff --git a/Engine/Source/Runtime/RenderCore/Private/BuiltInRayTracingShaders.cpp b/Engine/Source/Runtime/RenderCore/Private/BuiltInRayTracingShaders.cpp
--- a/Engine/Source/Runtime/RenderCore/Private/BuiltInRayTracingShaders.cpp
+++ b/Engine/Source/Runtime/RenderCore/Private/BuiltInRayTracingShaders.cpp
@@ -1,6 +1,7 @@
 // Copyright Epic Games, Inc. All Rights Reserved.
 
 #include "BuiltInRayTracingShaders.h"
+#include "ShaderParameterUtils.h"
 
 #if RHI_RAYTRACING
 
@@ -12,4 +13,43 @@ IMPLEMENT_SHADER_TYPE(, FDefaultMainCHSOpaqueAHS, TEXT("/Engine/Private/RayTraci
 IMPLEMENT_SHADER_TYPE(, FDefaultPayloadMS,		TEXT("/Engine/Private/RayTracing/RayTracingBuiltInShaders.usf"), TEXT("DefaultPayloadMS"),			SF_RayMiss);
 IMPLEMENT_SHADER_TYPE(, FPackedMaterialClosestHitPayloadMS, TEXT("/Engine/Private/RayTracing/RayTracingBuiltInShaders.usf"), TEXT("PackedMaterialClosestHitPayloadMS"), SF_RayMiss);
 
+IMPLEMENT_GLOBAL_SHADER(FRayTracingDispatchDescCS, "/Engine/Private/RayTracing/RayTracingDispatchDesc.usf", "RayTracingDispatchDescCS", SF_Compute);
+
+void FRayTracingDispatchDescCS::Dispatch(FRHICommandList& RHICmdList, 
+	const void* DispatchDescInput, uint32 DispatchDescSize, uint32 DispatchDescDimensionsOffset,
+	FRHIShaderResourceView* DispatchDimensionsSRV, uint32 DimensionsBufferOffset,
+	FRHIUnorderedAccessView* DispatchDescOutputUAV)
+{
+	const uint32 DispatchDescSizeDwords = DispatchDescSize / 4;
+	const uint32 DispatchDescDimensionsOffsetDwords = DispatchDescDimensionsOffset / 4;
+
+	checkf(DimensionsBufferOffset % 4 == 0, TEXT("Dispatch dimensions buffer offset must be DWORD-aligned"));
+	const uint32 DimensionsBufferOffsetDwords = DimensionsBufferOffset / 4;
+
+	check(DispatchDescSizeDwords <= DispatchDescMaxSizeDwords);
+
+	TShaderMapRef<FRayTracingDispatchDescCS> ComputeShader(GetGlobalShaderMap(GMaxRHIFeatureLevel));
+	FRHIComputeShader* ShaderRHI = ComputeShader.GetComputeShader();
+	RHICmdList.SetComputeShader(ShaderRHI);
+
+	static_assert(DispatchDescMaxSizeDwords % 4 == 0, "DispatchDescMaxSizeDwords must be a multiple of 4");
+	static constexpr uint32 DispatchDescMaxSizeUint4s = DispatchDescMaxSizeDwords / 4;
+
+	FUintVector4 DispatchDescData[DispatchDescMaxSizeUint4s] = {};
+	FMemory::Memcpy(DispatchDescData, DispatchDescInput, DispatchDescSize);
+
+	SetShaderValueArray(RHICmdList, ShaderRHI, ComputeShader->DispatchDescInputParam, DispatchDescData, DispatchDescMaxSizeUint4s);
+	SetShaderValue(RHICmdList, ShaderRHI, ComputeShader->DispatchDescSizeDwordsParam, DispatchDescSizeDwords);
+	SetShaderValue(RHICmdList, ShaderRHI, ComputeShader->DispatchDescDimensionsOffsetDwordsParam, DispatchDescDimensionsOffsetDwords);
+	SetShaderValue(RHICmdList, ShaderRHI, ComputeShader->DimensionsBufferOffsetDwordsParam, DimensionsBufferOffsetDwords);
+
+	SetSRVParameter(RHICmdList, ShaderRHI, ComputeShader->DispatchDimensionsParam, DispatchDimensionsSRV);
+	SetUAVParameter(RHICmdList, ShaderRHI, ComputeShader->DispatchDescOutputParam, DispatchDescOutputUAV);
+
+	RHICmdList.DispatchComputeShader(1, 1, 1);
+
+	SetSRVParameter(RHICmdList, ShaderRHI, ComputeShader->DispatchDimensionsParam, nullptr);
+	SetUAVParameter(RHICmdList, ShaderRHI, ComputeShader->DispatchDescOutputParam, nullptr);
+}
+
 #endif // RHI_RAYTRACING
diff --git a/Engine/Source/Runtime/RenderCore/Public/BuiltInRayTracingShaders.h b/Engine/Source/Runtime/RenderCore/Public/BuiltInRayTracingShaders.h
--- a/Engine/Source/Runtime/RenderCore/Public/BuiltInRayTracingShaders.h
+++ b/Engine/Source/Runtime/RenderCore/Public/BuiltInRayTracingShaders.h
@@ -106,5 +106,43 @@ public:
 	{}
 };
 
+class RENDERCORE_API FRayTracingDispatchDescCS : public FBuiltInRayTracingShader
+{
+	DECLARE_GLOBAL_SHADER(FRayTracingDispatchDescCS);
+
+public:
+	FRayTracingDispatchDescCS() = default;
+	FRayTracingDispatchDescCS(const ShaderMetaType::CompiledShaderInitializerType& Initializer)
+		: FBuiltInRayTracingShader(Initializer)
+	{
+		DispatchDescInputParam.Bind(Initializer.ParameterMap, TEXT("DispatchDescInput"), SPF_Mandatory);
+		DispatchDescSizeDwordsParam.Bind(Initializer.ParameterMap, TEXT("DispatchDescSizeDwords"), SPF_Mandatory);
+		DispatchDescDimensionsOffsetDwordsParam.Bind(Initializer.ParameterMap, TEXT("DispatchDescDimensionsOffsetDwords"), SPF_Mandatory);
+		DimensionsBufferOffsetDwordsParam.Bind(Initializer.ParameterMap, TEXT("DimensionsBufferOffsetDwords"), SPF_Mandatory);
+		DispatchDimensionsParam.Bind(Initializer.ParameterMap, TEXT("DispatchDimensions"), SPF_Mandatory);
+		DispatchDescOutputParam.Bind(Initializer.ParameterMap, TEXT("DispatchDescOutput"), SPF_Mandatory);
+	}
+
+	static inline void ModifyCompilationEnvironment(const FGlobalShaderPermutationParameters& Parameters, FShaderCompilerEnvironment& OutEnvironment)
+	{
+		OutEnvironment.SetDefine(TEXT("DISPATCH_DESC_MAX_SIZE_DWORDS"), DispatchDescMaxSizeDwords);
+		FGlobalShader::ModifyCompilationEnvironment(Parameters, OutEnvironment);
+	}
+
+	static void Dispatch(FRHICommandList& RHICmdList,
+		const void* DispatchDescInput, uint32 DispatchDescSize, uint32 DispatchDescDimensionsOffset,
+		FRHIShaderResourceView* DispatchDimensionsSRV, uint32 DimensionsBufferOffset,
+		FRHIUnorderedAccessView* DispatchDescOutputUAV);
+
+	static constexpr uint32 DispatchDescMaxSizeDwords = 32;
+
+	LAYOUT_FIELD(FShaderParameter, DispatchDescInputParam);
+	LAYOUT_FIELD(FShaderParameter, DispatchDescSizeDwordsParam);
+	LAYOUT_FIELD(FShaderParameter, DispatchDescDimensionsOffsetDwordsParam);
+	LAYOUT_FIELD(FShaderParameter, DimensionsBufferOffsetDwordsParam);
+	LAYOUT_FIELD(FShaderResourceParameter, DispatchDimensionsParam);
+	LAYOUT_FIELD(FShaderResourceParameter, DispatchDescOutputParam);
+};
+
 #endif // RHI_RAYTRACING
 
