Fixed validation layer bug introduced in 15036968. Indirect argument buffer state in RHIDispatchIndirectComputeShader should be IndirectArgs. RHIRayTraceDispatchIndirect should check IndirectArgs and SRVCompute. Fixed resource tracking initialization for DispatchRaysDescBuffer.

diff --git a/Engine/Source/Runtime/D3D12RHI/Private/D3D12RayTracing.cpp b/Engine/Source/Runtime/D3D12RHI/Private/D3D12RayTracing.cpp
--- a/Engine/Source/Runtime/D3D12RHI/Private/D3D12RayTracing.cpp
+++ b/Engine/Source/Runtime/D3D12RHI/Private/D3D12RayTracing.cpp
@@ -3869,7 +3869,8 @@ static void DispatchRays(FD3D12CommandContext& CommandContext,
 
 	FD3D12Adapter* Adapter = CommandContext.GetParentDevice()->GetParentAdapter();
 
-	TRefCountPtr<FD3D12Buffer> DispatchRaysDescBuffer;
+	FVertexBufferRHIRef DispatchRaysDescBufferRHI;
+	FD3D12Buffer* DispatchRaysDescBuffer = nullptr;
 
 	if (ArgumentBuffer)
 	{
@@ -3885,17 +3886,13 @@ static void DispatchRays(FD3D12CommandContext& CommandContext,
 
 		FRHIResourceCreateInfo DispatchRaysDescBufferCreateInfo;
 		DispatchRaysDescBufferCreateInfo.GPUMask = FRHIGPUMask::FromIndex(CommandContext.GetGPUIndex());
+		DispatchRaysDescBufferCreateInfo.DebugName = TEXT("DispatchRaysDescBuffer");
 
-		DispatchRaysDescBuffer = Adapter->CreateRHIBuffer(
-			nullptr, BufferDesc, BufferDesc.Alignment,
-			4, BufferDesc.Width, BUF_Static | BUF_DrawIndirect | BUF_UnorderedAccess,
-			ED3D12ResourceStateMode::MultiState, ERHIAccess::UAVCompute,
-			DispatchRaysDescBufferCreateInfo);
+		DispatchRaysDescBufferRHI = ::RHICreateVertexBuffer(sizeof(D3D12_DISPATCH_RAYS_DESC), BUF_UnorderedAccess | BUF_DrawIndirect, DispatchRaysDescBufferCreateInfo);
+		DispatchRaysDescBuffer = FD3D12DynamicRHI::ResourceCast(DispatchRaysDescBufferRHI.GetReference());
 
 		FD3D12Resource* DispatchRaysDescBufferResource = DispatchRaysDescBuffer->GetResource();
 
-		SetName(DispatchRaysDescBufferResource, TEXT("DispatchRaysDescBuffer"));
-
 		{
 			TRHICommandList_RecursiveHazardous<FD3D12CommandContext> RHICmdList(&CommandContext);
 			FShaderResourceViewRHIRef SRV = RHICreateShaderResourceView(ArgumentBuffer, 4, PF_R32_UINT);
diff --git a/Engine/Source/Runtime/RHI/Public/RHICommandList.h b/Engine/Source/Runtime/RHI/Public/RHICommandList.h
--- a/Engine/Source/Runtime/RHI/Public/RHICommandList.h
+++ b/Engine/Source/Runtime/RHI/Public/RHICommandList.h
@@ -3695,7 +3695,7 @@ public:
 
 	/**
 	 * Trace rays using dimensions from a GPU buffer containing uint[3], interpreted as number of rays in X, Y and Z dimensions.
-	 * ArgumentBuffer must be in ERHIAccess::ReadOnlyExclusiveMask state rather than simply IndirectArgs since it's read as a regular GPU resource on some platforms.
+	 * ArgumentBuffer must be in IndirectArgs|SRVCompute state.
 	 */
 	FORCEINLINE_DEBUGGABLE void RayTraceDispatchIndirect(FRayTracingPipelineState* Pipeline, FRHIRayTracingShader* RayGenShader, FRHIRayTracingScene* Scene, const FRayTracingShaderBindings& GlobalResourceBindings, FRHIBuffer* ArgumentBuffer, uint32 ArgumentOffset)
 	{
diff --git a/Engine/Source/Runtime/RHI/Public/RHIValidationContext.h b/Engine/Source/Runtime/RHI/Public/RHIValidationContext.h
--- a/Engine/Source/Runtime/RHI/Public/RHIValidationContext.h
+++ b/Engine/Source/Runtime/RHI/Public/RHIValidationContext.h
@@ -45,7 +45,7 @@ public:
 
 	virtual void RHIDispatchIndirectComputeShader(FRHIBuffer* ArgumentBuffer, uint32 ArgumentOffset) override final
 	{
-		Tracker->Assert(ArgumentBuffer->GetWholeResourceIdentity(), ERHIAccess::ReadOnlyExclusiveMask);
+		Tracker->Assert(ArgumentBuffer->GetWholeResourceIdentity(), ERHIAccess::IndirectArgs);
 
 		RHIContext->RHIDispatchIndirectComputeShader(ArgumentBuffer, ArgumentOffset);
 		Tracker->Dispatch();
@@ -927,6 +927,8 @@ public:
 		FRHIBuffer* ArgumentBuffer, uint32 ArgumentOffset) override final
 	{
 		Tracker->Assert(ArgumentBuffer->GetWholeResourceIdentity(), ERHIAccess::IndirectArgs);
+		Tracker->Assert(ArgumentBuffer->GetWholeResourceIdentity(), ERHIAccess::SRVCompute);
+
 		RHIContext->RHIRayTraceDispatchIndirect(RayTracingPipelineState, RayGenShader, Scene, GlobalResourceBindings, ArgumentBuffer, ArgumentOffset);
 	}
 
