Define D3D12_HEAP_FLAG_CREATE_NOT_ZEROED as D3D12_HEAP_FLAG_NONE on Xbox so Windows defines can be removed all over the code

diff --git a/Engine/Source/Runtime/D3D12RHI/Private/D3D12Allocation.cpp b/Engine/Source/Runtime/D3D12RHI/Private/D3D12Allocation.cpp
--- a/Engine/Source/Runtime/D3D12RHI/Private/D3D12Allocation.cpp
+++ b/Engine/Source/Runtime/D3D12RHI/Private/D3D12Allocation.cpp
@@ -172,12 +172,10 @@ void FD3D12BuddyAllocator::Initialize()
 		Desc.Properties = HeapProps;
 		Desc.Alignment = 0;
 		Desc.Flags = InitConfig.HeapFlags;
-#if PLATFORM_WINDOWS
 		if (Adapter->IsHeapNotZeroedSupported())
 		{
-			Desc.Flags |= D3D12_HEAP_FLAG_CREATE_NOT_ZEROED;
+			Desc.Flags |= FD3D12_HEAP_FLAG_CREATE_NOT_ZEROED;
 		}
-#endif
 
 		ID3D12Heap* Heap = nullptr;
 		{
@@ -1868,12 +1866,10 @@ FD3D12SegHeap* FD3D12SegList::CreateBackingHeap(
 	Desc.SizeInBytes = HeapSize;
 	Desc.Properties = CD3DX12_HEAP_PROPERTIES(HeapType, Parent->GetGPUMask().GetNative(), VisibleNodeMask.GetNative());
 	Desc.Flags = HeapFlags;
-#if PLATFORM_WINDOWS
 	if (Parent->GetParentAdapter()->IsHeapNotZeroedSupported())
 	{
-		Desc.Flags |= D3D12_HEAP_FLAG_CREATE_NOT_ZEROED;
+		Desc.Flags |= FD3D12_HEAP_FLAG_CREATE_NOT_ZEROED;
 	}
-#endif
 
 	VERIFYD3D12RESULT(Parent->GetDevice()->CreateHeap(&Desc, IID_PPV_ARGS(&D3DHeap)));
 
diff --git a/Engine/Source/Runtime/D3D12RHI/Private/D3D12PoolAllocator.cpp b/Engine/Source/Runtime/D3D12RHI/Private/D3D12PoolAllocator.cpp
--- a/Engine/Source/Runtime/D3D12RHI/Private/D3D12PoolAllocator.cpp
+++ b/Engine/Source/Runtime/D3D12RHI/Private/D3D12PoolAllocator.cpp
@@ -44,12 +44,10 @@ void FD3D12MemoryPool::Init()
 		Desc.Properties = HeapProps;
 		Desc.Alignment = 0;
 		Desc.Flags = InitConfig.HeapFlags;
-#if PLATFORM_WINDOWS
 		if (Adapter->IsHeapNotZeroedSupported())
 		{
-			Desc.Flags |= D3D12_HEAP_FLAG_CREATE_NOT_ZEROED;
+			Desc.Flags |= FD3D12_HEAP_FLAG_CREATE_NOT_ZEROED;
 		}
-#endif
 
 		ID3D12Heap* Heap = nullptr;
 		{
diff --git a/Engine/Source/Runtime/D3D12RHI/Private/D3D12RHIPrivate.h b/Engine/Source/Runtime/D3D12RHI/Private/D3D12RHIPrivate.h
--- a/Engine/Source/Runtime/D3D12RHI/Private/D3D12RHIPrivate.h
+++ b/Engine/Source/Runtime/D3D12RHI/Private/D3D12RHIPrivate.h
@@ -119,6 +119,10 @@ typedef FD3D12StateCacheBase FD3D12StateCache;
 
 #define MAX_GPU_BREADCRUMB_DEPTH 1024
 
+#ifndef FD3D12_HEAP_FLAG_CREATE_NOT_ZEROED
+#define FD3D12_HEAP_FLAG_CREATE_NOT_ZEROED D3D12_HEAP_FLAG_CREATE_NOT_ZEROED
+#endif
+
 #if DEBUG_RESOURCE_STATES
 #define LOG_EXECUTE_COMMAND_LISTS 1
 #define ASSERT_RESOURCE_STATES 0	// Disabled for now.
diff --git a/Engine/Source/Runtime/D3D12RHI/Private/D3D12Resources.cpp b/Engine/Source/Runtime/D3D12RHI/Private/D3D12Resources.cpp
--- a/Engine/Source/Runtime/D3D12RHI/Private/D3D12Resources.cpp
+++ b/Engine/Source/Runtime/D3D12RHI/Private/D3D12Resources.cpp
@@ -410,12 +410,7 @@ HRESULT FD3D12Adapter::CreateCommittedResource(const D3D12_RESOURCE_DESC& InDesc
 	LLM_PLATFORM_SCOPE(ELLMTag::GraphicsPlatform);
 
 	TRefCountPtr<ID3D12Resource> pResource;
-#if PLATFORM_WINDOWS
-	D3D12_HEAP_FLAGS HeapFlags = bHeapNotZeroedSupported? D3D12_HEAP_FLAG_CREATE_NOT_ZEROED : D3D12_HEAP_FLAG_NONE;
-#else
-	check(!bHeapNotZeroedSupported);
-	D3D12_HEAP_FLAGS HeapFlags = D3D12_HEAP_FLAG_NONE;
-#endif
+	D3D12_HEAP_FLAGS HeapFlags = bHeapNotZeroedSupported ? FD3D12_HEAP_FLAG_CREATE_NOT_ZEROED : D3D12_HEAP_FLAG_NONE;
 	if (InDesc.Flags & D3D12_RESOURCE_FLAG_ALLOW_SIMULTANEOUS_ACCESS)
 	{
 		HeapFlags |= D3D12_HEAP_FLAG_SHARED;
diff --git a/Engine/Source/Runtime/D3D12RHI/Private/HoloLens/D3D12RHIBasePrivate.h b/Engine/Source/Runtime/D3D12RHI/Private/HoloLens/D3D12RHIBasePrivate.h
--- a/Engine/Source/Runtime/D3D12RHI/Private/HoloLens/D3D12RHIBasePrivate.h
+++ b/Engine/Source/Runtime/D3D12RHI/Private/HoloLens/D3D12RHIBasePrivate.h
@@ -32,6 +32,10 @@
 #define D3D12RHI_RESOURCE_FLAG_ALLOW_INDIRECT_BUFFER D3D12_RESOURCE_FLAG_NONE
 #define D3D12RHI_HEAP_FLAG_ALLOW_INDIRECT_BUFFERS		D3D12_HEAP_FLAG_NONE
 
+// Heap create not zeroed flag is not available on Hololens so use internal define to disable the feature
+// but make code path shared when it becomes available
+#define FD3D12_HEAP_FLAG_CREATE_NOT_ZEROED D3D12_HEAP_FLAG_NONE
+
 #include "../Public/D3D12Util.h"
 
 // Windows desktop version is used.
