Batch request and process BLAS compaction requests to reduce memory waste of pending compaction buffers (64KB - 8 per request otherwise)

diff --git a/Engine/Source/Runtime/D3D12RHI/Private/D3D12CommandContext.cpp b/Engine/Source/Runtime/D3D12RHI/Private/D3D12CommandContext.cpp
--- a/Engine/Source/Runtime/D3D12RHI/Private/D3D12CommandContext.cpp
+++ b/Engine/Source/Runtime/D3D12RHI/Private/D3D12CommandContext.cpp
@@ -12,6 +12,8 @@ D3D12CommandContext.cpp: RHI  Command Context implementation.
 #include "Windows/HideWindowsPlatformTypes.h"
 #endif
 
+#include "D3D12RayTracing.h"
+
 // Aggressive batching saves ~0.1ms on the RHI thread, reduces executecommandlist calls by around 25%
 int32 GCommandListBatchingMode = CLB_AggressiveBatching;
 
@@ -491,6 +493,10 @@ void FD3D12CommandContextBase::RHIBeginFrame()
 		Device->GetGPUProfiler().BeginFrame(ParentAdapter->GetOwningRHI());
 
 		Device->GetDefaultBufferAllocator().BeginFrame();
+
+#if D3D12_RHI_RAYTRACING
+		Device->GetRayTracingCompactionRequestHandler()->Update(*ContextAtIndex);
+#endif // D3D12_RHI_RAYTRACING
 	}
 }
 
diff --git a/Engine/Source/Runtime/D3D12RHI/Private/D3D12Device.cpp b/Engine/Source/Runtime/D3D12RHI/Private/D3D12Device.cpp
--- a/Engine/Source/Runtime/D3D12RHI/Private/D3D12Device.cpp
+++ b/Engine/Source/Runtime/D3D12RHI/Private/D3D12Device.cpp
@@ -4,6 +4,7 @@
 D3D12Device.cpp: D3D device RHI implementation.
 =============================================================================*/
 #include "D3D12RHIPrivate.h"
+#include "D3D12RayTracing.h"
 
 namespace D3D12RHI
 {
@@ -54,6 +55,9 @@ FD3D12Device::FD3D12Device(FRHIGPUMask InGPUMask, FD3D12Adapter* InAdapter) :
 FD3D12Device::~FD3D12Device()
 {
 #if D3D12_RHI_RAYTRACING
+	delete RayTracingCompactionRequestHandler;
+	RayTracingCompactionRequestHandler = nullptr;
+
 	DestroyRayTracingDescriptorCache(); // #dxr_todo UE-72158: unify RT descriptor cache with main FD3D12DescriptorCache
 #endif
 
@@ -308,6 +312,11 @@ void FD3D12Device::SetupAfterDeviceCreation()
 
 	UpdateMSAASettings();
 
+#if D3D12_RHI_RAYTRACING
+	check(RayTracingCompactionRequestHandler == nullptr);
+	RayTracingCompactionRequestHandler = new FD3D12RayTracingCompactionRequestHandler(this);
+#endif // D3D12_RHI_RAYTRACING
+
 	GPUProfilingData.Init();
 }
 
diff --git a/Engine/Source/Runtime/D3D12RHI/Private/D3D12Device.h b/Engine/Source/Runtime/D3D12RHI/Private/D3D12Device.h
--- a/Engine/Source/Runtime/D3D12RHI/Private/D3D12Device.h
+++ b/Engine/Source/Runtime/D3D12RHI/Private/D3D12Device.h
@@ -12,6 +12,7 @@ class FD3D12DynamicRHI;
 class FD3D12BasicRayTracingPipeline;
 class FD3D12RayTracingDescriptorHeapCache;
 class FD3D12RayTracingPipelineCache;
+class FD3D12RayTracingCompactionRequestHandler;
 class FD3D12TimedIntervalQueryTracker;
 class FD3D12LinearQueryHeap;
 
@@ -52,6 +53,7 @@ public:
 	const FD3D12BasicRayTracingPipeline*	GetBasicRayTracingPipeline() const { return BasicRayTracingPipeline; }
 	FD3D12RayTracingDescriptorHeapCache*	GetRayTracingDescriptorHeapCache() { return RayTracingDescriptorHeapCache; }
 	FD3D12RayTracingPipelineCache*			GetRayTracingPipelineCache() { return RayTracingPipelineCache; }
+	FD3D12RayTracingCompactionRequestHandler* GetRayTracingCompactionRequestHandler() { return RayTracingCompactionRequestHandler; }
 	TRefCountPtr<ID3D12StateObject>			DeserializeRayTracingStateObject(D3D12_SHADER_BYTECODE Bytecode, ID3D12RootSignature* RootSignature);
 #endif // D3D12_RHI_RAYTRACING
 
@@ -195,6 +197,7 @@ protected:
 #if D3D12_RHI_RAYTRACING
 	FD3D12BasicRayTracingPipeline* BasicRayTracingPipeline = nullptr;
 	FD3D12RayTracingPipelineCache* RayTracingPipelineCache = nullptr;
+	FD3D12RayTracingCompactionRequestHandler* RayTracingCompactionRequestHandler = nullptr;
 	// #dxr_todo UE-72158: unify RT descriptor cache with main FD3D12DescriptorCache
 	FD3D12RayTracingDescriptorHeapCache* RayTracingDescriptorHeapCache = nullptr;
 	void DestroyRayTracingDescriptorCache();
diff --git a/Engine/Source/Runtime/D3D12RHI/Private/D3D12RayTracing.cpp b/Engine/Source/Runtime/D3D12RHI/Private/D3D12RayTracing.cpp
--- a/Engine/Source/Runtime/D3D12RHI/Private/D3D12RayTracing.cpp
+++ b/Engine/Source/Runtime/D3D12RHI/Private/D3D12RayTracing.cpp
@@ -75,6 +75,14 @@ static FAutoConsoleVariableRef CVarD3D12RayTracingAllowCompaction(
 	ECVF_ReadOnly
 );
 
+static int32 GD3D12RayTracingMaxBatchedCompaction = 64;
+static FAutoConsoleVariableRef CVarD3D12RayTracingMaxBatchedCompaction(
+	TEXT("r.D3D12.RayTracing.MaxBatchedCompaction"),
+	GD3D12RayTracingAllowCompaction,
+	TEXT("Maximum of amount of compaction requests and rebuilds per frame. (default = 64)\n"),
+	ECVF_ReadOnly
+);
+
 // Ray tracing stat counters
 
 DECLARE_STATS_GROUP(TEXT("D3D12RHI: Ray Tracing"), STATGROUP_D3D12RayTracing, STATCAT_Advanced);
@@ -405,6 +413,131 @@ static FD3D12ShaderIdentifier GetShaderIdentifier(ID3D12StateObject* StateObject
 	return GetShaderIdentifier(PipelineProperties, ExportName);
 }
 
+
+FD3D12RayTracingCompactionRequestHandler::FD3D12RayTracingCompactionRequestHandler(FD3D12Device* Device) : FD3D12DeviceChild(Device), PostBuildInfoBufferReadbackFence(0)
+{
+	D3D12_RESOURCE_DESC PostBuildInfoBufferDesc = CD3DX12_RESOURCE_DESC::Buffer(GD3D12RayTracingMaxBatchedCompaction * sizeof(uint64), D3D12_RESOURCE_FLAG_ALLOW_UNORDERED_ACCESS);
+
+	FRHIResourceCreateInfo CreateInfo;
+	CreateInfo.GPUMask = FRHIGPUMask::FromIndex(GetParentDevice()->GetGPUIndex());
+	CreateInfo.DebugName = TEXT("PostBuildInfoBuffer");
+	PostBuildInfoBuffer = GetParentDevice()->GetParentAdapter()->CreateRHIBuffer(nullptr, PostBuildInfoBufferDesc, 8,
+		0, PostBuildInfoBufferDesc.Width, BUF_UnorderedAccess | BUF_SourceCopy, ED3D12ResourceStateMode::MultiState, ERHIAccess::UAVMask, CreateInfo);
+	SetName(PostBuildInfoBuffer->GetResource(), TEXT("PostBuildInfoBuffer"));
+
+	PostBuildInfoStagingBuffer = RHICreateStagingBuffer();
+}
+
+void FD3D12RayTracingCompactionRequestHandler::RequestCompact(FD3D12RayTracingGeometry* InRTGeometry)
+{
+	uint32 GPUIndex = GetParentDevice()->GetGPUIndex();
+	check(InRTGeometry->AccelerationStructureBuffers[GPUIndex]);
+	check((InRTGeometry->BuildFlags & D3D12_RAYTRACING_ACCELERATION_STRUCTURE_BUILD_FLAG_ALLOW_COMPACTION) &&
+		(InRTGeometry->BuildFlags & D3D12_RAYTRACING_ACCELERATION_STRUCTURE_BUILD_FLAG_PREFER_FAST_TRACE) &&
+		!(InRTGeometry->BuildFlags & D3D12_RAYTRACING_ACCELERATION_STRUCTURE_BUILD_FLAG_ALLOW_UPDATE));
+
+	FScopeLock Lock(&CS);
+	PendingRequests.Add(InRTGeometry);
+}
+
+bool FD3D12RayTracingCompactionRequestHandler::ReleaseRequest(FD3D12RayTracingGeometry* InRTGeometry)
+{
+	FScopeLock Lock(&CS);
+
+	// Remove from pending list, not found then try active requests
+	if (PendingRequests.Remove(InRTGeometry) <= 0)
+	{
+		// If currently enqueued, then clear pointer to not handle the compaction request anymore			
+		for (int32 BLASIndex = 0; BLASIndex < ActiveBLASGPUAddresses.Num(); ++BLASIndex)
+		{
+			if (ActiveRequests[BLASIndex] == InRTGeometry)
+			{
+				ActiveRequests[BLASIndex] = nullptr;
+				return true;
+			}
+		}
+
+		return false;
+	}
+	else
+	{
+		return true;
+	}
+}
+
+void FD3D12RayTracingCompactionRequestHandler::Update(FD3D12CommandContext& InCommandContext)
+{
+	FScopeLock Lock(&CS);
+
+	// process previous build request data retrieval
+	uint32 GPUIndex = GetParentDevice()->GetGPUIndex();
+	FD3D12Adapter* Adapter = InCommandContext.GetParentAdapter();
+	const FD3D12Fence& Fence = InCommandContext.GetParentDevice()->GetCommandListManager().GetFence();
+
+	if (ActiveBLASGPUAddresses.Num() > 0)
+	{
+		// Ensure that our builds & copies have finished on GPU when enqueued - if still busy then wait until done
+		if (PostBuildInfoBufferReadbackFence > Fence.GetLastCompletedFenceFast())
+		{
+			return;
+		}
+
+		// Readback the sizes from the readback buffer and schedule new builds ops on the RTGeometry objects
+		uint64* SizesAfterCompaction = (uint64*)PostBuildInfoStagingBuffer->Lock(0, ActiveBLASGPUAddresses.Num() * sizeof(uint64));
+		for (int32 BLASIndex = 0; BLASIndex < ActiveBLASGPUAddresses.Num(); ++BLASIndex)
+		{
+			if (ActiveRequests[BLASIndex] != nullptr)
+			{
+				ActiveRequests[BLASIndex]->CompactAccelerationStructure(InCommandContext, GPUIndex, SizesAfterCompaction[BLASIndex]);
+			}
+		}
+		PostBuildInfoStagingBuffer->Unlock();
+
+		// reset working values
+		PostBuildInfoBufferReadbackFence = 0;
+		ActiveRequests.Empty(ActiveRequests.Num());
+		ActiveBLASGPUAddresses.Empty(ActiveBLASGPUAddresses.Num());
+	}
+
+	// build a new set of build requests to extract the build data	
+	for (FD3D12RayTracingGeometry* RTGeomtry : PendingRequests)
+	{
+		ActiveRequests.Add(RTGeomtry);
+		ActiveBLASGPUAddresses.Add(RTGeomtry->AccelerationStructureBuffers[GPUIndex].GetReference()->ResourceLocation.GetGPUVirtualAddress());
+
+		// enqueued enough requests for this update round
+		if (ActiveRequests.Num() >= GD3D12RayTracingMaxBatchedCompaction)
+		{
+			break;
+		}
+	}
+
+	// Do we have requests?
+	if (ActiveRequests.Num() > 0)
+	{
+		// clear out all of the pending requests, don't allow the array to shrink
+		PendingRequests.RemoveAt(0, ActiveRequests.Num(), false);
+
+		D3D12_RAYTRACING_ACCELERATION_STRUCTURE_POSTBUILD_INFO_DESC PostBuildInfoDesc = {};
+		PostBuildInfoDesc.DestBuffer = PostBuildInfoBuffer->ResourceLocation.GetGPUVirtualAddress();
+		PostBuildInfoDesc.InfoType = D3D12_RAYTRACING_ACCELERATION_STRUCTURE_POSTBUILD_INFO_COMPACTED_SIZE;
+
+		FD3D12DynamicRHI::TransitionResource(InCommandContext.CommandListHandle, PostBuildInfoBuffer->GetResource(), D3D12_RESOURCE_STATE_TBD, D3D12_RESOURCE_STATE_UNORDERED_ACCESS, 0, FD3D12DynamicRHI::ETransitionMode::Apply);
+
+		// Emit the RT post build info from the selected requests
+		ID3D12GraphicsCommandList4* CommandList4 = InCommandContext.CommandListHandle.RayTracingCommandList();
+		CommandList4->EmitRaytracingAccelerationStructurePostbuildInfo(&PostBuildInfoDesc, ActiveBLASGPUAddresses.Num(), ActiveBLASGPUAddresses.GetData());
+
+		// Transition to copy source and perform the copy to readback
+		FD3D12DynamicRHI::TransitionResource(InCommandContext.CommandListHandle, PostBuildInfoBuffer.GetReference()->GetResource(), D3D12_RESOURCE_STATE_UNORDERED_ACCESS, D3D12_RESOURCE_STATE_COPY_SOURCE, 0, FD3D12DynamicRHI::ETransitionMode::Apply);
+		InCommandContext.RHICopyToStagingBuffer(PostBuildInfoBuffer, PostBuildInfoStagingBuffer, 0, sizeof(uint64) * ActiveBLASGPUAddresses.Num());
+
+		// Update the fence value
+		PostBuildInfoBufferReadbackFence = FMath::Max(PostBuildInfoBufferReadbackFence, Fence.GetCurrentFence());
+	}
+}
+
+
 // Cache for ray tracing pipeline collection objects, containing single shaders that can be linked into full pipelines.
 class FD3D12RayTracingPipelineCache : FD3D12DeviceChild
 {
@@ -2238,6 +2371,7 @@ FD3D12RayTracingGeometry::FD3D12RayTracingGeometry(const FRayTracingGeometryInit
 	INC_DWORD_STAT(STAT_D3D12RayTracingAllocatedBLAS);
 
 	DebugName = Initializer.DebugName;
+	FMemory::Memzero(bHasPendingCompactionRequests);
 
 	if(!FD3D12RayTracingGeometry::NullTransformBuffer.IsValid())
 	{
@@ -2250,8 +2384,6 @@ FD3D12RayTracingGeometry::FD3D12RayTracingGeometry(const FRayTracingGeometryInit
 		FD3D12RayTracingGeometry::NullTransformBuffer = RHICreateVertexBuffer(NullTransformData.GetResourceDataSize(), BUF_Static, CreateInfo);
 	}
 
-	FMemory::Memzero(PostBuildInfoBufferReadbackFences);
-
 	IndexOffsetInBytes = Initializer.IndexBufferOffset;
 	TotalPrimitiveCount = Initializer.TotalPrimitiveCount;
 
@@ -2345,6 +2477,18 @@ FD3D12RayTracingGeometry::FD3D12RayTracingGeometry(const FRayTracingGeometryInit
 
 FD3D12RayTracingGeometry::~FD3D12RayTracingGeometry()
 {
+	for (uint32 GPUIndex = 0; GPUIndex < MAX_NUM_GPUS; ++GPUIndex)
+	{
+		if (bHasPendingCompactionRequests[GPUIndex])
+		{
+			check(AccelerationStructureBuffers[GPUIndex]);
+			FD3D12Device* Device = AccelerationStructureBuffers[GPUIndex].GetReference()->GetParentDevice();
+			bool bRequestFound = Device->GetRayTracingCompactionRequestHandler()->ReleaseRequest(this);
+			check(bRequestFound);
+			bHasPendingCompactionRequests[GPUIndex] = false;
+		}
+	}
+
 	for (TRefCountPtr<FD3D12Buffer>& Buffer : AccelerationStructureBuffers)
 	{
 		if (Buffer)
@@ -2606,31 +2750,6 @@ void FD3D12RayTracingGeometry::BuildAccelerationStructure(FD3D12CommandContext&
 		INC_MEMORY_STAT_BY(STAT_D3D12RayTracingBLASMemory, ScratchBuffers[GPUIndex]->GetSize());
 	}
 
-	// Compact BLAS if it is static
-	bool bShouldCompactAfterBuild = BuildFlags & D3D12_RAYTRACING_ACCELERATION_STRUCTURE_BUILD_FLAG_ALLOW_COMPACTION && BuildFlags & D3D12_RAYTRACING_ACCELERATION_STRUCTURE_BUILD_FLAG_PREFER_FAST_TRACE && !(BuildFlags & D3D12_RAYTRACING_ACCELERATION_STRUCTURE_BUILD_FLAG_ALLOW_UPDATE);
-
-	if (bShouldCompactAfterBuild)
-	{
-		D3D12_RESOURCE_DESC PostBuildInfoBufferDesc = CD3DX12_RESOURCE_DESC::Buffer(sizeof(uint64), D3D12_RESOURCE_FLAG_ALLOW_UNORDERED_ACCESS);
-
-		FRHIResourceCreateInfo CreateInfo;
-
-		CreateInfo.GPUMask = FRHIGPUMask::FromIndex(GPUIndex);
-		CreateInfo.DebugName = TEXT("PostBuildInfoBuffer");
-		PostBuildInfoBuffers[GPUIndex] = Adapter->CreateRHIBuffer(
-			nullptr, PostBuildInfoBufferDesc, 8,
-			0, PostBuildInfoBufferDesc.Width, BUF_UnorderedAccess | BUF_SourceCopy, ED3D12ResourceStateMode::MultiState, ERHIAccess::UAVMask, CreateInfo);
-
-		SetName(PostBuildInfoBuffers[GPUIndex]->GetResource(), TEXT("PostBuildInfoBuffer"));
-
-		PostBuildInfoStagingBuffers[GPUIndex] = RHICreateStagingBuffer();
-
-		// Do not defer delete the staging buffer data because when readback is done we don't need
-		// to keep it alive any longer - shouldn't even wait on frame fence
-		PostBuildInfoBuffers[GPUIndex]->DoNoDeferDelete();
-		PostBuildInfoStagingBuffers[GPUIndex]->DoNoDeferDelete();
-	}
-
 	TransitionBuffers(CommandContext);
 	CommandContext.CommandListHandle.FlushResourceBarriers();
 
@@ -2647,30 +2766,19 @@ void FD3D12RayTracingGeometry::BuildAccelerationStructure(FD3D12CommandContext&
 			? AccelerationStructureBuffers[GPUIndex]->ResourceLocation.GetGPUVirtualAddress()
 			: D3D12_GPU_VIRTUAL_ADDRESS(0);
 
+		// Build
 		ID3D12GraphicsCommandList4* RayTracingCommandList = CommandContext.CommandListHandle.RayTracingCommandList();
-
-		if (!bShouldCompactAfterBuild)
+		RayTracingCommandList->BuildRaytracingAccelerationStructure(&BuildDesc, 0, nullptr);
+
+		// Request compact BLAS if it is static
+		bool bShouldCompactAfterBuild = 
+			(BuildFlags & D3D12_RAYTRACING_ACCELERATION_STRUCTURE_BUILD_FLAG_ALLOW_COMPACTION) && 
+			(BuildFlags & D3D12_RAYTRACING_ACCELERATION_STRUCTURE_BUILD_FLAG_PREFER_FAST_TRACE) && 
+			!(BuildFlags & D3D12_RAYTRACING_ACCELERATION_STRUCTURE_BUILD_FLAG_ALLOW_UPDATE);		
+		if (bShouldCompactAfterBuild)
 		{
-			// Build
-			RayTracingCommandList->BuildRaytracingAccelerationStructure(&BuildDesc, 0, nullptr);
-		}
-		else
-		{
-			D3D12_RAYTRACING_ACCELERATION_STRUCTURE_POSTBUILD_INFO_DESC PostBuildInfoDesc = {};
-			PostBuildInfoDesc.DestBuffer = PostBuildInfoBuffers[GPUIndex]->ResourceLocation.GetGPUVirtualAddress();
-			PostBuildInfoDesc.InfoType = D3D12_RAYTRACING_ACCELERATION_STRUCTURE_POSTBUILD_INFO_COMPACTED_SIZE;
-
-			FD3D12DynamicRHI::TransitionResource(CommandContext.CommandListHandle, PostBuildInfoBuffers[GPUIndex].GetReference()->GetResource(), D3D12_RESOURCE_STATE_TBD, D3D12_RESOURCE_STATE_UNORDERED_ACCESS, 0, FD3D12DynamicRHI::ETransitionMode::Apply);
-
-			// Build with post build info
-			RayTracingCommandList->BuildRaytracingAccelerationStructure(&BuildDesc, 1, &PostBuildInfoDesc);
-
-			// Transition to copy source and perform the copy to readback
-			FD3D12DynamicRHI::TransitionResource(CommandContext.CommandListHandle, PostBuildInfoBuffers[GPUIndex].GetReference()->GetResource(), D3D12_RESOURCE_STATE_UNORDERED_ACCESS, D3D12_RESOURCE_STATE_COPY_SOURCE, 0, FD3D12DynamicRHI::ETransitionMode::Apply);
-			CommandContext.RHICopyToStagingBuffer(PostBuildInfoBuffers[GPUIndex], PostBuildInfoStagingBuffers[GPUIndex], 0, sizeof(uint64));
-
-			const FD3D12Fence& Fence = CommandContext.GetParentDevice()->GetCommandListManager().GetFence();
-			PostBuildInfoBufferReadbackFences[GPUIndex] = FMath::Max(PostBuildInfoBufferReadbackFences[GPUIndex], Fence.GetCurrentFence());
+			CommandContext.GetParentDevice()->GetRayTracingCompactionRequestHandler()->RequestCompact(this);
+			bHasPendingCompactionRequests[GPUIndex] = true;
 		}
 
 		SetDirty(CommandContext.GetGPUMask(), false);
@@ -2694,81 +2802,47 @@ void FD3D12RayTracingGeometry::BuildAccelerationStructure(FD3D12CommandContext&
 	}
 }
 
-void FD3D12RayTracingGeometry::ConditionalCompactAccelerationStructure(FD3D12CommandContext& CommandContext)
+void FD3D12RayTracingGeometry::CompactAccelerationStructure(FD3D12CommandContext& CommandContext, uint32 InGPUIndex, uint64 InSizeAfterCompaction)
 {
-	const bool bShouldCompactAfterBuild = BuildFlags& D3D12_RAYTRACING_ACCELERATION_STRUCTURE_BUILD_FLAG_ALLOW_COMPACTION&& BuildFlags& D3D12_RAYTRACING_ACCELERATION_STRUCTURE_BUILD_FLAG_PREFER_FAST_TRACE && !(BuildFlags & D3D12_RAYTRACING_ACCELERATION_STRUCTURE_BUILD_FLAG_ALLOW_UPDATE);
+	// Should have a pending request
+	check(bHasPendingCompactionRequests[InGPUIndex]);
+	bHasPendingCompactionRequests[InGPUIndex] = false;
 
-	if (!bShouldCompactAfterBuild)
+	ensureMsgf(InSizeAfterCompaction > 0, TEXT("Compacted acceleration structure size is expected to be non-zero. This error suggests that GPU readback synchronization is broken."));
+	if (InSizeAfterCompaction == 0)
 	{
 		return;
 	}
 
-	const uint32 GPUIndex = CommandContext.GetGPUIndex();
-
-	// return the buffers are null (mostly because the BLAS has been compacted)
-	if (PostBuildInfoBuffers[GPUIndex] == nullptr)
-	{
-		return;
-	}
-
-	FD3D12Adapter* Adapter = CommandContext.GetParentAdapter();
-
-	const FD3D12Fence& Fence = CommandContext.GetParentDevice()->GetCommandListManager().GetFence();
-
-	// Ensure that our builds & copies have finished on GPU
-	if (PostBuildInfoBufferReadbackFences[GPUIndex] > Fence.GetLastCompletedFenceFast())
-	{
-		return;
-	}
-
-	uint64 SizeAfterCompaction = 0;
-
-	{
-		// Read size and release readback buffers.
-		SizeAfterCompaction = *(uint64*)PostBuildInfoStagingBuffers[GPUIndex]->Lock(0, sizeof(uint64));
-
-		PostBuildInfoStagingBuffers[GPUIndex]->Unlock();
-		PostBuildInfoBuffers[GPUIndex] = nullptr;
-		PostBuildInfoStagingBuffers[GPUIndex] = nullptr;
-		PostBuildInfoBufferReadbackFences[GPUIndex] = MAX_uint64;
-	}
-
-	ensureMsgf(SizeAfterCompaction, TEXT("Compacted acceleration structure size is expected to be non-zero. This error suggests that GPU readback synchronization is broken."));
-
-	if (SizeAfterCompaction == 0)
-	{
-		return;
-	}
-
-	DEC_MEMORY_STAT_BY(STAT_D3D12RayTracingUsedVideoMemory, AccelerationStructureBuffers[GPUIndex]->GetSize());
-	DEC_MEMORY_STAT_BY(STAT_D3D12RayTracingBLASMemory, AccelerationStructureBuffers[GPUIndex]->GetSize());
+	DEC_MEMORY_STAT_BY(STAT_D3D12RayTracingUsedVideoMemory, AccelerationStructureBuffers[InGPUIndex]->GetSize());
+	DEC_MEMORY_STAT_BY(STAT_D3D12RayTracingBLASMemory, AccelerationStructureBuffers[InGPUIndex]->GetSize());
 
 	// Move old AS into this temporary variable which gets released when this function returns
-	TRefCountPtr<FD3D12Buffer> OldAccelerationStructure = MoveTemp(AccelerationStructureBuffers[GPUIndex]);
+	TRefCountPtr<FD3D12Buffer> OldAccelerationStructure = MoveTemp(AccelerationStructureBuffers[InGPUIndex]);
 
 	{
 		FRHIResourceCreateInfo CreateInfo;
 
-		D3D12_RESOURCE_DESC AccelerationStructureBufferDesc = CD3DX12_RESOURCE_DESC::Buffer(SizeAfterCompaction, D3D12_RESOURCE_FLAG_ALLOW_UNORDERED_ACCESS);
+		D3D12_RESOURCE_DESC AccelerationStructureBufferDesc = CD3DX12_RESOURCE_DESC::Buffer(InSizeAfterCompaction, D3D12_RESOURCE_FLAG_ALLOW_UNORDERED_ACCESS);
 
-		CreateInfo.GPUMask = FRHIGPUMask::FromIndex(GPUIndex);
+		CreateInfo.GPUMask = FRHIGPUMask::FromIndex(InGPUIndex);
 		CreateInfo.DebugName = TEXT("AccelerationStructureBuffer");
-		AccelerationStructureBuffers[GPUIndex] = Adapter->CreateRHIBuffer(
+		AccelerationStructureBuffers[InGPUIndex] = CommandContext.GetParentAdapter()->CreateRHIBuffer(
 			nullptr, AccelerationStructureBufferDesc, D3D12_RAYTRACING_ACCELERATION_STRUCTURE_BYTE_ALIGNMENT,
 			0, AccelerationStructureBufferDesc.Width, BUF_AccelerationStructure, ED3D12ResourceStateMode::SingleState, ERHIAccess::RTAccelerationStructure, CreateInfo);
 
-		INC_MEMORY_STAT_BY(STAT_D3D12RayTracingUsedVideoMemory, AccelerationStructureBuffers[GPUIndex]->GetSize());
-		INC_MEMORY_STAT_BY(STAT_D3D12RayTracingBLASMemory, AccelerationStructureBuffers[GPUIndex]->GetSize());
+		INC_MEMORY_STAT_BY(STAT_D3D12RayTracingUsedVideoMemory, AccelerationStructureBuffers[InGPUIndex]->GetSize());
+		INC_MEMORY_STAT_BY(STAT_D3D12RayTracingBLASMemory, AccelerationStructureBuffers[InGPUIndex]->GetSize());
 
-		SetName(AccelerationStructureBuffers[GPUIndex]->GetResource(),
+		SetName(AccelerationStructureBuffers[InGPUIndex]->GetResource(),
 			DebugName.IsValid() ? *DebugName.ToString() : TEXT("BLAS"));
 
 		OldAccelerationStructure->GetResource()->UpdateResidency(CommandContext.CommandListHandle);
-		AccelerationStructureBuffers[GPUIndex]->GetResource()->UpdateResidency(CommandContext.CommandListHandle);
+		AccelerationStructureBuffers[InGPUIndex]->GetResource()->UpdateResidency(CommandContext.CommandListHandle);
 
 		ID3D12GraphicsCommandList4* RayTracingCommandList = CommandContext.CommandListHandle.RayTracingCommandList();
 		RayTracingCommandList->CopyRaytracingAccelerationStructure(
-			AccelerationStructureBuffers[GPUIndex]->ResourceLocation.GetGPUVirtualAddress(),
+			AccelerationStructureBuffers[InGPUIndex]->ResourceLocation.GetGPUVirtualAddress(),
 			OldAccelerationStructure->ResourceLocation.GetGPUVirtualAddress(),
 			D3D12_RAYTRACING_ACCELERATION_STRUCTURE_COPY_MODE_COMPACT);
 	}
@@ -2964,8 +3038,6 @@ void FD3D12RayTracingScene::BuildAccelerationStructure(FD3D12CommandContext& Com
 				checkf(!Geometry->IsDirty(CommandContext.GetGPUIndex()),
 					TEXT("Acceleration structures for all geometries must be built before building the top level acceleration structure for the scene."));
 
-				Geometry->ConditionalCompactAccelerationStructure(CommandContext);
-
 				AddResidencyHandleForResource(Geometry->AccelerationStructureBuffers[GPUIndex]->GetResource());
 
 				if (Geometry->RHIIndexBuffer)
diff --git a/Engine/Source/Runtime/D3D12RHI/Private/D3D12RayTracing.h b/Engine/Source/Runtime/D3D12RHI/Private/D3D12RayTracing.h
--- a/Engine/Source/Runtime/D3D12RHI/Private/D3D12RayTracing.h
+++ b/Engine/Source/Runtime/D3D12RHI/Private/D3D12RayTracing.h
@@ -28,7 +28,7 @@ public:
 
 	void TransitionBuffers(FD3D12CommandContext& CommandContext);
 	void BuildAccelerationStructure(FD3D12CommandContext& CommandContext, EAccelerationStructureBuildMode BuildMode);
-	void ConditionalCompactAccelerationStructure(FD3D12CommandContext& CommandContext);
+	void CompactAccelerationStructure(FD3D12CommandContext& CommandContext, uint32 InGPUIndex, uint64 InSizeAfterCompaction);
 
 	bool bIsAccelerationStructureDirty[MAX_NUM_GPUS] = {};
 	void SetDirty(FRHIGPUMask GPUMask, bool bState)
@@ -58,9 +58,7 @@ public:
 	TRefCountPtr<FD3D12Buffer> AccelerationStructureBuffers[MAX_NUM_GPUS];
 	TRefCountPtr<FD3D12Buffer> ScratchBuffers[MAX_NUM_GPUS];
 
-	uint64 PostBuildInfoBufferReadbackFences[MAX_NUM_GPUS];
-	TRefCountPtr<FD3D12Buffer> PostBuildInfoBuffers[MAX_NUM_GPUS];
-	FStagingBufferRHIRef PostBuildInfoStagingBuffers[MAX_NUM_GPUS];
+	bool bHasPendingCompactionRequests[MAX_NUM_GPUS];
 
 	// Hit shader parameters per geometry segment
 	TArray<FHitGroupSystemParameters> HitGroupSystemParameters[MAX_NUM_GPUS];
@@ -128,4 +126,34 @@ public:
 	FName DebugName;
 };
 
+// Manages all the pending BLAS compaction requests
+class FD3D12RayTracingCompactionRequestHandler : FD3D12DeviceChild
+{
+public:
+
+	UE_NONCOPYABLE(FD3D12RayTracingCompactionRequestHandler)
+
+	FD3D12RayTracingCompactionRequestHandler(FD3D12Device* Device);
+	~FD3D12RayTracingCompactionRequestHandler()
+	{
+		check(PendingRequests.IsEmpty());
+	}
+
+	void RequestCompact(FD3D12RayTracingGeometry* InRTGeometry);
+	bool ReleaseRequest(FD3D12RayTracingGeometry* InRTGeometry);
+
+	void Update(FD3D12CommandContext& InCommandContext);
+
+private:
+
+	FCriticalSection CS;
+	TArray<FD3D12RayTracingGeometry*> PendingRequests;
+	TArray<FD3D12RayTracingGeometry*> ActiveRequests;
+	TArray<D3D12_GPU_VIRTUAL_ADDRESS> ActiveBLASGPUAddresses;
+
+	TRefCountPtr<FD3D12Buffer> PostBuildInfoBuffer;
+	FStagingBufferRHIRef PostBuildInfoStagingBuffer;
+	uint64 PostBuildInfoBufferReadbackFence;
+};
+
 #endif // D3D12_RHI_RAYTRACING
