Added CVar r.D3D12.RayTracing.ExecuteIndirectWorkaround (default 0) that adds a command list flush before ExecuteIndirect to avoid NVIDIA driver bug that causes some indirect RT commands to get lost.

diff --git a/Engine/Source/Runtime/D3D12RHI/Private/D3D12RayTracing.cpp b/Engine/Source/Runtime/D3D12RHI/Private/D3D12RayTracing.cpp
--- a/Engine/Source/Runtime/D3D12RHI/Private/D3D12RayTracing.cpp
+++ b/Engine/Source/Runtime/D3D12RHI/Private/D3D12RayTracing.cpp
@@ -83,6 +83,13 @@ static FAutoConsoleVariableRef CVarD3D12RayTracingMaxBatchedCompaction(
 	ECVF_ReadOnly
 );
 
+static int32 GD3D12RayTracingExecuteIndirectWorkaround = 0;
+static FAutoConsoleVariableRef CVarD3D12RayTracingExecuteIndirectWorkaround(
+	TEXT("r.D3D12.RayTracing.ExecuteIndirectWorkaround"),
+	GD3D12RayTracingExecuteIndirectWorkaround,
+	TEXT("Adds a command list flush before ExecuteIndirect to avoid NVIDIA driver bug that causes some indirect RT commands to get lost. (default = 0)\n")
+);
+
 // Ray tracing stat counters
 
 DECLARE_STATS_GROUP(TEXT("D3D12RHI: Ray Tracing"), STATGROUP_D3D12RayTracing, STATCAT_Advanced);
@@ -3975,6 +3982,11 @@ static void DispatchRays(FD3D12CommandContext& CommandContext,
 		}
 
 		CommandListHandle.AddTransitionBarrier(DispatchRaysDescBufferResource, D3D12_RESOURCE_STATE_UNORDERED_ACCESS, D3D12_RESOURCE_STATE_INDIRECT_ARGUMENT, 0);
+
+		if (GD3D12RayTracingExecuteIndirectWorkaround)
+		{
+			CommandContext.RHISubmitCommandsHint();
+		}
 	}
 
 	// Setup state for RT dispatch
