- Generalize d3d12 pool allocator code to work for textures as well
- Use pool allocator for non RT/DS textures on PC (with defrag enabled for non 4K aligned textures)
- Support for defrag of RT/DS textures
- Don't extract the D3D12Resource for a copy op from RenderThread because it could still be defragged before operation is executed (issue with Texture Unlock)

diff --git a/Engine/Source/Runtime/D3D12RHI/Private/D3D12Allocation.cpp b/Engine/Source/Runtime/D3D12RHI/Private/D3D12Allocation.cpp
--- a/Engine/Source/Runtime/D3D12RHI/Private/D3D12Allocation.cpp
+++ b/Engine/Source/Runtime/D3D12RHI/Private/D3D12Allocation.cpp
@@ -39,20 +39,61 @@ static FAutoConsoleVariableRef CVarD3D12ReadOnlyTextureAllocatorMaxPoolSize(
 #endif
 
 #if USE_POOL_ALLOCATOR
-static int32 GD3D12PoolAllocatorVRAMPoolSize = 32 * 1024 * 1024;
-static FAutoConsoleVariableRef CVarD3D12PoolAllocatorVRAMPoolSize(
-	TEXT("d3d12.PoolAllocator.VRAMPoolSize"),
-	GD3D12PoolAllocatorVRAMPoolSize,
-	TEXT("Pool size of a single VRAM memory pool (default 32MB)"),
+static int32 GD3D12PoolAllocatorBufferVRAMPoolSize = 32 * 1024 * 1024;
+static FAutoConsoleVariableRef CVarD3D12PoolAllocatorBufferVRAMPoolSize(
+	TEXT("d3d12.PoolAllocator.BufferVRAMPoolSize"),
+	GD3D12PoolAllocatorBufferVRAMPoolSize,
+	TEXT("Pool size of a single VRAM DataBuffer memory pool (default 32MB)"),
 	ECVF_ReadOnly);
 
-static int32 GD3D12PoolAllocatorVRAMMaxAllocationSize = 16 * 1024 * 1024;
-static FAutoConsoleVariableRef CVarD3D12PoolAllocatorMaxAllocationSize(
-	TEXT("d3d12.PoolAllocator.MaxAllocationSize"),
-	GD3D12PoolAllocatorVRAMMaxAllocationSize,
-	TEXT("Maximum size of a single allocation in the VRAM pool allocator (default 16MB)"),
+static int32 GD3D12PoolAllocatorBufferVRAMMaxAllocationSize = 16 * 1024 * 1024;
+static FAutoConsoleVariableRef CVarD3D12PoolAllocatorBufferMaxAllocationSize(
+	TEXT("d3d12.PoolAllocator.BufferMaxAllocationSize"),
+	GD3D12PoolAllocatorBufferVRAMMaxAllocationSize,
+	TEXT("Maximum size of a single allocation in the VRAM DataBuffer pool allocator (default 16MB)"),
+	ECVF_ReadOnly);
+
+static int32 GD3D12PoolAllocatorReadOnlyTextureVRAMPoolSize = 64 * 1024 * 1024;
+static FAutoConsoleVariableRef CVarD3D12PoolAllocatorReadOnlyTextureVRAMPoolSize(
+	TEXT("d3d12.PoolAllocator.ReadOnlyTextureVRAMPoolSize"),
+	GD3D12PoolAllocatorReadOnlyTextureVRAMPoolSize,
+	TEXT("Pool size of a single VRAM ReadOnly Texture memory pool (default 64MB)"),
+	ECVF_ReadOnly);
+
+static int32 GD3D12PoolAllocatorReadOnlyTextureVRAMMaxAllocationSize = 64 * 1024 * 1024;
+static FAutoConsoleVariableRef CVarD3D12PoolAllocatorReadOnlyTextureMaxAllocationSize(
+	TEXT("d3d12.PoolAllocator.ReadOnlyTextureMaxAllocationSize"),
+	GD3D12PoolAllocatorReadOnlyTextureVRAMMaxAllocationSize,
+	TEXT("Maximum size of a single allocation in the VRAM ReadOnly Texture pool allocator (default 64MB)"),
+	ECVF_ReadOnly);
+
+static int32 GD3D12PoolAllocatorRTUAVTextureVRAMPoolSize = 0 * 1024 * 1024;
+static FAutoConsoleVariableRef CVarD3D12PoolAllocatorRTUAVTextureVRAMPoolSize(
+	TEXT("d3d12.PoolAllocator.RTUAVTextureVRAMPoolSize"),
+	GD3D12PoolAllocatorRTUAVTextureVRAMPoolSize,
+	TEXT("Pool size of a single VRAM RTUAV Texture memory pool (default 0MB - disabled)"),
+	ECVF_ReadOnly);
+
+static int32 GD3D12PoolAllocatorRTUAVTextureVRAMMaxAllocationSize = 0 * 1024 * 1024;
+static FAutoConsoleVariableRef CVarD3D12PoolAllocatorRTUAVTextureMaxAllocationSize(
+	TEXT("d3d12.PoolAllocator.RTUAVTextureMaxAllocationSize"),
+	GD3D12PoolAllocatorRTUAVTextureVRAMMaxAllocationSize,
+	TEXT("Maximum size of a single allocation in the VRAM RTUAV Texture pool allocator (default 0MB - disabled)"),
+	ECVF_ReadOnly);
+
+static int32 GD3D12VRAMTexturePoolDefrag = 1;
+static FAutoConsoleVariableRef CVarD3D12D3D12VRAMTexturePoolDefrag(
+	TEXT("d3d12.VRAMTexturePoolDefrag"),
+	GD3D12VRAMTexturePoolDefrag,
+	TEXT("Defrag the VRAM Texture pool (enabled by default)"),
+	ECVF_ReadOnly);
+
+static int32 GD3D12VRAMTexturePoolDefragMaxCopySizePerFrame = 32 * 1024 * 1024;
+static FAutoConsoleVariableRef CVarD3D12D3D12VRAMTexturePoolDefragMaxCopySizePerFrame(
+	TEXT("d3d12.VRAMTexturePoolDefrag.MaxCopySizePerFrame"),
+	GD3D12VRAMTexturePoolDefragMaxCopySizePerFrame,
+	TEXT("Max amount of data to copy during defragmentation in a single frame (default 32MB)"),
 	ECVF_ReadOnly);
-#endif // USE_POOL_ALLOCATOR
 
 // NOTE: buffer pool defrag is disabled by default because FD3D12RayTracingGeometry is currently not notified
 //       yet of defragged index and vertex buffers used by the the BVH data. Without RT it works fine.
@@ -70,6 +111,8 @@ static FAutoConsoleVariableRef CVarD3D12D3D12VRAMBufferPoolDefragMaxCopySizePerF
 	TEXT("Max amount of data to copy during defragmentation in a single frame (default 32MB)"),
 	ECVF_ReadOnly);
 
+#endif // USE_POOL_ALLOCATOR
+
 #if D3D12RHI_SEGLIST_ALLOC_TRACK_WASTAGE
 static int32 GD3D12SegListTrackLeaks = 0;
 static FAutoConsoleVariableRef CVarD3D12SegListTrackLeaks(
@@ -1219,9 +1262,9 @@ FD3D12BufferPool* FD3D12DefaultBufferAllocator::CreateBufferPool(D3D12_HEAP_TYPE
 
 	const FString Name(L"D3D12 Pool Allocator");
 	EResourceAllocationStrategy AllocationStrategy = FD3D12PoolAllocator::GetResourceAllocationStrategy(InResourceFlags, InResourceStateMode);
-	uint64 PoolSize = InHeapType == D3D12_HEAP_TYPE_READBACK ? READBACK_BUFFER_POOL_DEFAULT_POOL_SIZE : GD3D12PoolAllocatorVRAMPoolSize;
+	uint64 PoolSize = InHeapType == D3D12_HEAP_TYPE_READBACK ? READBACK_BUFFER_POOL_DEFAULT_POOL_SIZE : GD3D12PoolAllocatorBufferVRAMPoolSize;
 	uint64 PoolAlignment = (AllocationStrategy == EResourceAllocationStrategy::kPlacedResource) ? MIN_PLACED_RESOURCE_SIZE : 256;
-	uint64 MaxAllocationSize = InHeapType == D3D12_HEAP_TYPE_READBACK ? READBACK_BUFFER_POOL_MAX_ALLOC_SIZE : GD3D12PoolAllocatorVRAMMaxAllocationSize;
+	uint64 MaxAllocationSize = InHeapType == D3D12_HEAP_TYPE_READBACK ? READBACK_BUFFER_POOL_MAX_ALLOC_SIZE : GD3D12PoolAllocatorBufferVRAMMaxAllocationSize;
 
 	FD3D12BufferPool* NewPool = new FD3D12PoolAllocator(Device, GetVisibilityMask(), InitConfig, Name, AllocationStrategy, PoolSize, PoolAlignment, MaxAllocationSize);
 
@@ -1414,7 +1457,162 @@ void FD3D12DefaultBufferAllocator::UpdateMemoryStats()
 //	Texture Allocator
 //-----------------------------------------------------------------------------
 
-#if D3D12RHI_SEGREGATED_TEXTURE_ALLOC
+#if USE_POOL_ALLOCATOR
+FD3D12TextureAllocatorPool::FD3D12TextureAllocatorPool(FD3D12Device* Device, FRHIGPUMask VisibilityNode) :
+	FD3D12DeviceChild(Device),
+	FD3D12MultiNodeGPUObject(Device->GetGPUMask(), VisibilityNode)
+{	
+	FD3D12ResourceInitConfig SharedInitConfig;
+
+	// texture only interesting in VRAM for now
+	SharedInitConfig.HeapType = D3D12_HEAP_TYPE_DEFAULT;
+
+	// unused for textures because placed and not suballocated
+	SharedInitConfig.ResourceFlags = D3D12_RESOURCE_FLAG_NONE;
+	SharedInitConfig.InitialResourceState = D3D12_RESOURCE_STATE_COMMON;
+
+	EResourceAllocationStrategy AllocationStrategy = EResourceAllocationStrategy::kPlacedResource;
+	
+	{
+		FD3D12ResourceInitConfig InitConfig = SharedInitConfig;
+		InitConfig.HeapFlags = D3D12_HEAP_FLAG_ALLOW_ONLY_NON_RT_DS_TEXTURES;
+
+		const FString Name(L"D3D12 ReadOnly4K Texture Pool Allocator");
+		uint64 PoolSize = 4 * 1024 * 1024;
+		uint64 MaxAllocationSize = PoolSize;
+		FD3D12PoolAllocator* ReadOnly4KPool = new FD3D12PoolAllocator(Device, GetVisibilityMask(), InitConfig, Name, AllocationStrategy, PoolSize, D3D12_SMALL_RESOURCE_PLACEMENT_ALIGNMENT, MaxAllocationSize);
+		ReadOnly4KPool->Initialize();
+
+		PoolAllocators[(int)EPoolType::ReadOnly4K] = ReadOnly4KPool;
+	}
+
+	{
+		FD3D12ResourceInitConfig InitConfig = SharedInitConfig;
+		InitConfig.HeapFlags = D3D12_HEAP_FLAG_ALLOW_ONLY_NON_RT_DS_TEXTURES;
+
+		const FString Name(L"D3D12 ReadOnly Texture Pool Allocator");
+		uint64 PoolSize = GD3D12PoolAllocatorReadOnlyTextureVRAMPoolSize;
+		uint64 MaxAllocationSize = GD3D12PoolAllocatorReadOnlyTextureVRAMMaxAllocationSize;
+		FD3D12PoolAllocator* ReadOnlyPool = new FD3D12PoolAllocator(Device, GetVisibilityMask(), InitConfig, Name, AllocationStrategy, PoolSize, D3D12_DEFAULT_RESOURCE_PLACEMENT_ALIGNMENT, MaxAllocationSize);
+		ReadOnlyPool->Initialize();
+
+		PoolAllocators[(int)EPoolType::ReadOnly] = ReadOnlyPool;
+	}
+
+	{
+		FD3D12ResourceInitConfig InitConfig = SharedInitConfig;
+		InitConfig.HeapFlags = D3D12_HEAP_FLAG_ALLOW_ONLY_RT_DS_TEXTURES;
+
+		const FString Name(L"D3D12 RT Texture Pool Allocator");
+		uint64 PoolSize = GD3D12PoolAllocatorRTUAVTextureVRAMPoolSize;
+		uint64 MaxAllocationSize = GD3D12PoolAllocatorRTUAVTextureVRAMMaxAllocationSize;
+		FD3D12PoolAllocator* RTPool = new FD3D12PoolAllocator(Device, GetVisibilityMask(), InitConfig, Name, AllocationStrategy, PoolSize, D3D12_DEFAULT_RESOURCE_PLACEMENT_ALIGNMENT, MaxAllocationSize);
+		RTPool->Initialize();
+
+		PoolAllocators[(int)EPoolType::RenderTarget] = RTPool;
+	}
+
+	{
+		FD3D12ResourceInitConfig InitConfig = SharedInitConfig;
+		InitConfig.HeapFlags = D3D12_HEAP_FLAG_ALLOW_ONLY_NON_RT_DS_TEXTURES;
+
+		const FString Name(L"D3D12 UAV Texture Pool Allocator");
+		uint64 PoolSize = GD3D12PoolAllocatorRTUAVTextureVRAMPoolSize;
+		uint64 MaxAllocationSize = GD3D12PoolAllocatorRTUAVTextureVRAMMaxAllocationSize;
+		FD3D12PoolAllocator* UAVPool = new FD3D12PoolAllocator(Device, GetVisibilityMask(), InitConfig, Name, AllocationStrategy, PoolSize, D3D12_DEFAULT_RESOURCE_PLACEMENT_ALIGNMENT, MaxAllocationSize);
+		UAVPool->Initialize();
+
+		PoolAllocators[(int)EPoolType::UAV] = UAVPool;
+	}
+}
+
+void FD3D12TextureAllocatorPool::Destroy()
+{
+	for (uint32 PoolIndex = 0; PoolIndex < (uint32)EPoolType::Count; ++PoolIndex)
+	{
+		PoolAllocators[PoolIndex]->CleanUpAllocations(0);
+		delete PoolAllocators[PoolIndex];
+		PoolAllocators[PoolIndex] = nullptr;
+	}
+}
+
+HRESULT FD3D12TextureAllocatorPool::AllocateTexture(
+	D3D12_RESOURCE_DESC Desc,
+	const D3D12_CLEAR_VALUE* ClearValue,
+	uint8 UEFormat,
+	FD3D12ResourceLocation& TextureLocation,
+	const D3D12_RESOURCE_STATES InitialState,
+	const TCHAR* Name)
+{
+	// The top mip level must be less than 64 KB to use 4 KB alignment
+	bool b4KAligment = TextureCanBe4KAligned(Desc, UEFormat);
+	Desc.Alignment = b4KAligment ?	D3D12_SMALL_RESOURCE_PLACEMENT_ALIGNMENT : D3D12_DEFAULT_RESOURCE_PLACEMENT_ALIGNMENT;
+	const D3D12_RESOURCE_ALLOCATION_INFO Info = GetParentDevice()->GetDevice()->GetResourceAllocationInfo(0, 1, &Desc);
+
+	bool bIsReadOnly = !(Desc.Flags & D3D12_RESOURCE_FLAG_ALLOW_RENDER_TARGET ||
+						 Desc.Flags & D3D12_RESOURCE_FLAG_ALLOW_DEPTH_STENCIL ||
+						 Desc.Flags & D3D12_RESOURCE_FLAG_ALLOW_UNORDERED_ACCESS);
+	bool bIsRenderTarget = (Desc.Flags & D3D12_RESOURCE_FLAG_ALLOW_RENDER_TARGET ||
+						 Desc.Flags & D3D12_RESOURCE_FLAG_ALLOW_DEPTH_STENCIL);
+	ED3D12ResourceStateMode ResourceStateMode = bIsReadOnly ? ED3D12ResourceStateMode::Default : ED3D12ResourceStateMode::MultiState;
+	EPoolType PoolType = bIsReadOnly ? (b4KAligment ? EPoolType::ReadOnly4K : EPoolType::ReadOnly) : (bIsRenderTarget ? EPoolType::RenderTarget : EPoolType::UAV);
+	PoolAllocators[(int)PoolType]->AllocResource(D3D12_HEAP_TYPE_DEFAULT, Desc, Info.SizeInBytes, Info.Alignment, ResourceStateMode, InitialState, ClearValue, Name, TextureLocation);
+
+	return S_OK;
+}
+
+
+void FD3D12TextureAllocatorPool::BeginFrame()
+{
+	FD3D12CommandContext& CommandContext = GetParentDevice()->GetDefaultCommandContext();
+
+	if (GD3D12VRAMTexturePoolDefrag > 0 && GD3D12VRAMTexturePoolDefragMaxCopySizePerFrame > 0)
+	{
+		uint32 MaxCopySize = GD3D12VRAMTexturePoolDefragMaxCopySizePerFrame;
+		uint32 CopySize = 0;
+		// Don't defrag the 4K pool - not worth it because it shouldn't really fragment
+		for (uint32 PoolIndex = 1; PoolIndex < (uint32)EPoolType::Count; ++PoolIndex)
+		{
+			PoolAllocators[PoolIndex]->Defrag(MaxCopySize, CopySize);
+		}
+	}
+
+	for (uint32 PoolIndex = 0; PoolIndex < (uint32)EPoolType::Count; ++PoolIndex)
+	{
+		PoolAllocators[PoolIndex]->FlushPendingCopyOps(CommandContext);
+	}
+}
+
+void FD3D12TextureAllocatorPool::CleanUpAllocations()
+{
+	for (uint32 PoolIndex = 0; PoolIndex < (uint32)EPoolType::Count; ++PoolIndex)
+	{
+		PoolAllocators[PoolIndex]->CleanUpAllocations(2);
+	}
+}
+
+
+bool FD3D12TextureAllocatorPool::GetMemoryStats(uint64& OutTotalAllocated, uint64& OutTotalUnused) const
+{
+	uint32 MemoryAllocated = 0;
+	uint32 MemoryUsed = 0;
+	uint32 FreeMemory = 0;
+	uint32 AlignmentWaste = 0;
+	uint32 AllocatedPageCount = 0;
+	uint32 FullPageCount = 0;
+
+	for (uint32 PoolIndex = 0; PoolIndex < (uint32)EPoolType::Count; ++PoolIndex)
+	{
+		PoolAllocators[PoolIndex]->UpdateMemoryStats(MemoryAllocated, MemoryUsed, FreeMemory, AlignmentWaste, AllocatedPageCount, FullPageCount);
+	}
+
+	OutTotalAllocated = MemoryAllocated;
+	OutTotalUnused = FreeMemory;
+
+	return true;
+}
+
+#elif D3D12RHI_SEGREGATED_TEXTURE_ALLOC
 FD3D12TextureAllocatorPool::FD3D12TextureAllocatorPool(FD3D12Device* Device, FRHIGPUMask VisibilityNode) :
 	FD3D12DeviceChild(Device),
 	FD3D12MultiNodeGPUObject(Device->GetGPUMask(), VisibilityNode),
diff --git a/Engine/Source/Runtime/D3D12RHI/Private/D3D12Allocation.h b/Engine/Source/Runtime/D3D12RHI/Private/D3D12Allocation.h
--- a/Engine/Source/Runtime/D3D12RHI/Private/D3D12Allocation.h
+++ b/Engine/Source/Runtime/D3D12RHI/Private/D3D12Allocation.h
@@ -898,7 +898,33 @@ private:
 //	FD3D12TextureAllocator
 //-----------------------------------------------------------------------------
 
-#if D3D12RHI_SEGREGATED_TEXTURE_ALLOC
+#if USE_POOL_ALLOCATOR
+class FD3D12TextureAllocatorPool : public FD3D12DeviceChild, public FD3D12MultiNodeGPUObject
+{
+public:
+	FD3D12TextureAllocatorPool(FD3D12Device* Device, FRHIGPUMask VisibilityNode);
+
+	HRESULT AllocateTexture(D3D12_RESOURCE_DESC Desc, const D3D12_CLEAR_VALUE* ClearValue, uint8 UEFormat, FD3D12ResourceLocation& TextureLocation, const D3D12_RESOURCE_STATES InitialState, const TCHAR* Name);
+	
+	void BeginFrame();
+	void CleanUpAllocations();
+	void Destroy();
+	bool GetMemoryStats(uint64& OutTotalAllocated, uint64& OutTotalUnused) const;
+
+private:
+	
+	enum class EPoolType
+	{
+		ReadOnly4K,
+		ReadOnly,
+		RenderTarget,
+		UAV,
+		Count,
+	};
+
+	FD3D12PoolAllocator* PoolAllocators[(int)EPoolType::Count];
+};
+#elif D3D12RHI_SEGREGATED_TEXTURE_ALLOC
 class FD3D12TextureAllocatorPool : public FD3D12DeviceChild, public FD3D12MultiNodeGPUObject
 {
 public:
@@ -906,6 +932,7 @@ public:
 
 	HRESULT AllocateTexture(D3D12_RESOURCE_DESC Desc, const D3D12_CLEAR_VALUE* ClearValue, uint8 UEFormat, FD3D12ResourceLocation& TextureLocation, const D3D12_RESOURCE_STATES InitialState, const TCHAR* Name);
 
+	void BeginFrame() {}
 	void CleanUpAllocations()
 	{
 		ReadOnlyTexturePool.CleanUpAllocations();
@@ -941,7 +968,8 @@ public:
 	FD3D12TextureAllocatorPool(FD3D12Device* Device, FRHIGPUMask VisibilityNode);
 
 	HRESULT AllocateTexture(D3D12_RESOURCE_DESC Desc, const D3D12_CLEAR_VALUE* ClearValue, uint8 UEFormat, FD3D12ResourceLocation& TextureLocation, const D3D12_RESOURCE_STATES InitialState, const TCHAR* Name);
-
+	
+	void BeginFrame() {}
 	void CleanUpAllocations() { ReadOnlyTexturePool.CleanUpAllocations(0); }
 
 	void Destroy() { ReadOnlyTexturePool.Destroy(); }
diff --git a/Engine/Source/Runtime/D3D12RHI/Private/D3D12CommandContext.cpp b/Engine/Source/Runtime/D3D12RHI/Private/D3D12CommandContext.cpp
--- a/Engine/Source/Runtime/D3D12RHI/Private/D3D12CommandContext.cpp
+++ b/Engine/Source/Runtime/D3D12RHI/Private/D3D12CommandContext.cpp
@@ -493,6 +493,7 @@ void FD3D12CommandContextBase::RHIBeginFrame()
 		Device->GetGPUProfiler().BeginFrame(ParentAdapter->GetOwningRHI());
 
 		Device->GetDefaultBufferAllocator().BeginFrame();
+		Device->GetTextureAllocator().BeginFrame();
 
 #if D3D12_RHI_RAYTRACING
 		Device->GetRayTracingCompactionRequestHandler()->Update(*ContextAtIndex);
diff --git a/Engine/Source/Runtime/D3D12RHI/Private/D3D12PoolAllocator.cpp b/Engine/Source/Runtime/D3D12RHI/Private/D3D12PoolAllocator.cpp
--- a/Engine/Source/Runtime/D3D12RHI/Private/D3D12PoolAllocator.cpp
+++ b/Engine/Source/Runtime/D3D12RHI/Private/D3D12PoolAllocator.cpp
@@ -27,6 +27,11 @@ FD3D12MemoryPool::~FD3D12MemoryPool()
 
 void FD3D12MemoryPool::Init()
 {
+	if (PoolSize == 0)
+	{
+		return;
+	}
+
 	FD3D12Device* Device = GetParentDevice();
 	FD3D12Adapter* Adapter = Device->GetParentAdapter();
 
@@ -195,16 +200,6 @@ bool FD3D12PoolAllocator::SupportsAllocation(D3D12_HEAP_TYPE InHeapType, D3D12_R
 void FD3D12PoolAllocator::AllocDefaultResource(D3D12_HEAP_TYPE InHeapType, const D3D12_RESOURCE_DESC& InDesc, EBufferUsageFlags InBufferUsage, ED3D12ResourceStateMode InResourceStateMode,
 	D3D12_RESOURCE_STATES InCreateState, uint32 InAllocationAlignment, const TCHAR* InName, FD3D12ResourceLocation& ResourceLocation)
 {
-	FD3D12Device* Device = GetParentDevice();
-	FD3D12Adapter* Adapter = Device->GetParentAdapter();
-
-	// If the resource location owns a block, this will deallocate it.
-	ResourceLocation.Clear();
-	if (InDesc.Width == 0)
-	{
-		return;
-	}
-
 #if DO_CHECK
 	// Validate the create state
 	if (InHeapType == D3D12_HEAP_TYPE_READBACK)
@@ -225,7 +220,24 @@ void FD3D12PoolAllocator::AllocDefaultResource(D3D12_HEAP_TYPE InHeapType, const
 #endif // D3D12_RHI_RAYTRACING
 #endif  // DO_CHECK
 
-	const bool PoolResource = InDesc.Width <= MaxAllocationSize;
+	AllocResource(InHeapType, InDesc, InDesc.Width, InAllocationAlignment, InResourceStateMode, InCreateState, nullptr, InName, ResourceLocation);
+}
+
+
+void FD3D12PoolAllocator::AllocResource(D3D12_HEAP_TYPE InHeapType, const D3D12_RESOURCE_DESC& InDesc, uint64 InSize, uint32 InAllocationAlignment, ED3D12ResourceStateMode InResourceStateMode,
+	D3D12_RESOURCE_STATES InCreateState, const D3D12_CLEAR_VALUE* InClearValue, const TCHAR* InName, FD3D12ResourceLocation& ResourceLocation)
+{
+	// If the resource location owns a block, this will deallocate it.
+	ResourceLocation.Clear();
+	if (InSize == 0)
+	{
+		return;
+	}
+
+	FD3D12Device* Device = GetParentDevice();
+	FD3D12Adapter* Adapter = Device->GetParentAdapter();
+
+	const bool PoolResource = InSize <= MaxAllocationSize;
 	if (PoolResource)
 	{
 		const bool bPlacedResource = (AllocationStrategy == EResourceAllocationStrategy::kPlacedResource);
@@ -234,7 +246,7 @@ void FD3D12PoolAllocator::AllocDefaultResource(D3D12_HEAP_TYPE InHeapType, const
 		if (bPlacedResource)
 		{
 			// Writeable resources get separate ID3D12Resource* with their own resource state by using placed resources. Just make sure it's UAV, other flags are free to differ.
-			check((InDesc.Flags & InitConfig.ResourceFlags & D3D12_RESOURCE_FLAG_ALLOW_UNORDERED_ACCESS) != 0 || InHeapType == D3D12_HEAP_TYPE_READBACK);
+			check(InDesc.Dimension != D3D12_RESOURCE_DIMENSION_BUFFER || (InDesc.Flags & InitConfig.ResourceFlags & D3D12_RESOURCE_FLAG_ALLOW_UNORDERED_ACCESS) != 0 || InHeapType == D3D12_HEAP_TYPE_READBACK);
 		}
 		else
 		{
@@ -244,12 +256,12 @@ void FD3D12PoolAllocator::AllocDefaultResource(D3D12_HEAP_TYPE InHeapType, const
 
 		// Try to allocate in one of the pools
 		FRHIPoolAllocationData& AllocationData = ResourceLocation.GetPoolAllocatorPrivateData().PoolData;
-		if (TryAllocateInternal(InDesc.Width, InAllocationAlignment, AllocationData))
+		if (TryAllocateInternal(InSize, InAllocationAlignment, AllocationData))
 		{
 			// Setup the resource location
 			ResourceLocation.SetType(FD3D12ResourceLocation::ResourceLocationType::eSubAllocation);
 			ResourceLocation.SetPoolAllocator(this);
-			ResourceLocation.SetSize(InDesc.Width);
+			ResourceLocation.SetSize(InSize);
 
 			AllocationData.SetOwner(&ResourceLocation);
 
@@ -273,7 +285,7 @@ void FD3D12PoolAllocator::AllocDefaultResource(D3D12_HEAP_TYPE InHeapType, const
 				FD3D12HeapAndOffset HeapAndOffset = GetBackingHeapAndAllocationOffsetInBytes(ResourceLocation);
 
 				FD3D12Resource* NewResource = nullptr;
-				VERIFYD3D12RESULT(Adapter->CreatePlacedResource(InDesc, HeapAndOffset.Heap, HeapAndOffset.Offset, InCreateState, ED3D12ResourceStateMode::MultiState, D3D12_RESOURCE_STATE_TBD, nullptr, &NewResource, InName));
+				VERIFYD3D12RESULT(Adapter->CreatePlacedResource(InDesc, HeapAndOffset.Heap, HeapAndOffset.Offset, InCreateState, InResourceStateMode, D3D12_RESOURCE_STATE_TBD, InClearValue, &NewResource, InName));
 
 				ResourceLocation.SetResource(NewResource);
 			}
@@ -288,8 +300,12 @@ void FD3D12PoolAllocator::AllocDefaultResource(D3D12_HEAP_TYPE InHeapType, const
 
 	// Allocate Standalone - move to owner of resource because this allocator should only manage pooled allocations (needed for now to do the same as FD3D12DefaultBufferPool)
 	FD3D12Resource* NewResource = nullptr;
-	VERIFYD3D12RESULT(Adapter->CreateBuffer(InHeapType, GetGPUMask(), GetVisibilityMask(), InCreateState, InResourceStateMode, InDesc.Width, &NewResource, InName, InDesc.Flags));
-	ResourceLocation.AsStandAlone(NewResource, InDesc.Width);
+	const D3D12_HEAP_PROPERTIES HeapProps = CD3DX12_HEAP_PROPERTIES(InHeapType, GetGPUMask().GetNative(), GetVisibilityMask().GetNative());
+	D3D12_RESOURCE_DESC Desc = InDesc;
+	Desc.Alignment = 0;
+	VERIFYD3D12RESULT(Adapter->CreateCommittedResource(Desc, GetGPUMask(), HeapProps, InCreateState, InClearValue, &NewResource, InName, false));
+	
+	ResourceLocation.AsStandAlone(NewResource, InSize);
 }
 
 
@@ -397,8 +413,11 @@ bool FD3D12PoolAllocator::HandleDefragRequest(FRHIPoolAllocationData* InSourceBl
 	CopyOp.DestResource		= Owner->GetResource();
 	CopyOp.DestOffset		= Owner->GetOffsetFromBaseOfResource();
 	CopyOp.Size				= InSourceBlock->GetSize();
+	CopyOp.CopyType			= AllocationStrategy == EResourceAllocationStrategy::kManualSubAllocation ? FD3D12VRAMCopyOperation::ECopyType::BufferRegion : FD3D12VRAMCopyOperation::ECopyType::Resource;
 	PendingCopyOps.Add(CopyOp);
 
+	// TODO: Using aliasing buffer on whole heap for copies to reduce flushes and resource transitions
+
 	return true;
 }
 
@@ -522,14 +541,26 @@ void FD3D12PoolAllocator::FlushPendingCopyOps(FD3D12CommandContext& InCommandCon
 
 	FD3D12CommandListHandle& CommandListHandle = InCommandContext.CommandListHandle;
 
+	FD3D12Adapter* Adapter = GetParentDevice()->GetParentAdapter();
+	FD3D12Fence& FrameFence = Adapter->GetFrameFence();
+
 	for (FD3D12VRAMCopyOperation& CopyOperation : PendingCopyOps)
-	{
+	{		
+		bool bRTAccelerationStructure = false;
 		if (CopyOperation.SourceResource->RequiresResourceStateTracking())
 		{
 			check(CopyOperation.DestResource->RequiresResourceStateTracking());
 			FD3D12DynamicRHI::TransitionResource(CommandListHandle, CopyOperation.SourceResource, D3D12_RESOURCE_STATE_TBD, D3D12_RESOURCE_STATE_COPY_SOURCE, D3D12_RESOURCE_BARRIER_ALL_SUBRESOURCES, FD3D12DynamicRHI::ETransitionMode::Apply);
 			FD3D12DynamicRHI::TransitionResource(CommandListHandle, CopyOperation.DestResource, D3D12_RESOURCE_STATE_TBD, D3D12_RESOURCE_STATE_COPY_DEST, D3D12_RESOURCE_BARRIER_ALL_SUBRESOURCES, FD3D12DynamicRHI::ETransitionMode::Apply);
 		}
+#if D3D12_RHI_RAYTRACING
+		else if (CopyOperation.SourceResource->GetDefaultResourceState() == D3D12_RESOURCE_STATE_RAYTRACING_ACCELERATION_STRUCTURE)
+		{
+			// can't make state changes to RT resources
+			bRTAccelerationStructure = true;
+			check(CopyOperation.DestResource->GetDefaultResourceState() == D3D12_RESOURCE_STATE_RAYTRACING_ACCELERATION_STRUCTURE);
+		}
+#endif // D3D12_RHI_RAYTRACING
 		else
 		{
 			check(!CopyOperation.DestResource->RequiresResourceStateTracking());
@@ -537,15 +568,41 @@ void FD3D12PoolAllocator::FlushPendingCopyOps(FD3D12CommandContext& InCommandCon
 			CommandListHandle.AddTransitionBarrier(CopyOperation.DestResource, CopyOperation.DestResource->GetDefaultResourceState(), D3D12_RESOURCE_STATE_COPY_DEST, D3D12_RESOURCE_BARRIER_ALL_SUBRESOURCES);
 		}
 
+		// Enable to log all defrag copy ops
+		//UE_LOG(LogD3D12RHI, Log, TEXT("Running defrag copy op for: %s at Frame: %d (Old: %#016llx New: %#016llx)"), *CopyOperation.SourceResource->GetName().ToString(), FrameFence.GetCurrentFence(), CopyOperation.SourceResource, CopyOperation.DestResource);
+
 		InCommandContext.numCopies++;
 		CommandListHandle.FlushResourceBarriers();
-		CommandListHandle->CopyBufferRegion(CopyOperation.DestResource->GetResource(), CopyOperation.DestOffset,
-			CopyOperation.SourceResource->GetResource(), CopyOperation.SourceOffset, CopyOperation.Size);
-		
+
+#if D3D12_RHI_RAYTRACING
+		if (bRTAccelerationStructure)
+		{
+			CommandListHandle.GraphicsCommandList5()->CopyRaytracingAccelerationStructure(CopyOperation.DestResource->GetResource()->GetGPUVirtualAddress(),
+				CopyOperation.SourceResource->GetResource()->GetGPUVirtualAddress(), D3D12_RAYTRACING_ACCELERATION_STRUCTURE_COPY_MODE_CLONE);
+		}
+		else
+#endif // D3D12_RHI_RAYTRACING
+		{
+			switch (CopyOperation.CopyType)
+			{
+			case FD3D12VRAMCopyOperation::BufferRegion:
+			{
+				CommandListHandle->CopyBufferRegion(CopyOperation.DestResource->GetResource(), CopyOperation.DestOffset,
+					CopyOperation.SourceResource->GetResource(), CopyOperation.SourceOffset, CopyOperation.Size);
+				break;
+			}
+			case FD3D12VRAMCopyOperation::Resource:
+			{
+				CommandListHandle->CopyResource(CopyOperation.DestResource->GetResource(), CopyOperation.SourceResource->GetResource());
+				break;
+			}
+			}
+		}
+
 		CommandListHandle.UpdateResidency(CopyOperation.SourceResource);
 		CommandListHandle.UpdateResidency(CopyOperation.DestResource);
 
-		if (!CopyOperation.SourceResource->RequiresResourceStateTracking())
+		if (!bRTAccelerationStructure && !CopyOperation.SourceResource->RequiresResourceStateTracking())
 		{
 			CommandListHandle.AddTransitionBarrier(CopyOperation.SourceResource, D3D12_RESOURCE_STATE_COPY_SOURCE, CopyOperation.SourceResource->GetDefaultResourceState(), D3D12_RESOURCE_BARRIER_ALL_SUBRESOURCES);
 			CommandListHandle.AddTransitionBarrier(CopyOperation.DestResource, D3D12_RESOURCE_STATE_COPY_DEST, CopyOperation.DestResource->GetDefaultResourceState(), D3D12_RESOURCE_BARRIER_ALL_SUBRESOURCES);
diff --git a/Engine/Source/Runtime/D3D12RHI/Private/D3D12PoolAllocator.h b/Engine/Source/Runtime/D3D12RHI/Private/D3D12PoolAllocator.h
--- a/Engine/Source/Runtime/D3D12RHI/Private/D3D12PoolAllocator.h
+++ b/Engine/Source/Runtime/D3D12RHI/Private/D3D12PoolAllocator.h
@@ -42,11 +42,18 @@ struct FD3D12ResourceInitConfig
 **/
 struct FD3D12VRAMCopyOperation
 {
+	enum ECopyType
+	{
+		BufferRegion,
+		Resource,
+	};
+
 	FD3D12Resource* SourceResource;
 	uint32 SourceOffset;
 	FD3D12Resource* DestResource;
 	uint32 DestOffset;
 	uint32 Size;
+	ECopyType CopyType;
 };
 
 // Heap and offset combined to specific location of a resource
@@ -119,6 +126,9 @@ public:
 		D3D12_RESOURCE_STATES InCreateState, uint32 InAlignment, const TCHAR* InName, FD3D12ResourceLocation& ResourceLocation);
 	void Deallocate(FD3D12ResourceLocation& ResourceLocation);
 
+	void AllocResource(D3D12_HEAP_TYPE InHeapType, const D3D12_RESOURCE_DESC& InDesc, uint64 InSize, uint32 InAllocationAlignment, ED3D12ResourceStateMode InResourceStateMode,
+		D3D12_RESOURCE_STATES InCreateState, const D3D12_CLEAR_VALUE* InClearValue, const TCHAR* InName, FD3D12ResourceLocation& ResourceLocation);
+
 	void CleanUpAllocations(uint64 InFrameLag);
 	void FlushPendingCopyOps(FD3D12CommandContext& InCommandContext);
 
diff --git a/Engine/Source/Runtime/D3D12RHI/Private/D3D12Resources.cpp b/Engine/Source/Runtime/D3D12RHI/Private/D3D12Resources.cpp
--- a/Engine/Source/Runtime/D3D12RHI/Private/D3D12Resources.cpp
+++ b/Engine/Source/Runtime/D3D12RHI/Private/D3D12Resources.cpp
@@ -461,9 +461,6 @@ HRESULT FD3D12Adapter::CreatePlacedResource(const D3D12_RESOURCE_DESC& InDesc, F
 
 	if (SUCCEEDED(hr))
 	{
-		// Set a default name (can override later).
-		SetName(pResource, Name);
-
 		FD3D12Device* Device = BackingHeap->GetParentDevice();
 		const D3D12_HEAP_DESC HeapDesc = Heap->GetDesc();
 
@@ -478,6 +475,9 @@ HRESULT FD3D12Adapter::CreatePlacedResource(const D3D12_RESOURCE_DESC& InDesc, F
 			BackingHeap,
 			HeapDesc.Properties.Type);
 
+		// Set a default name (can override later).
+		SetName(*ppOutResource, Name);
+
 		(*ppOutResource)->AddRef();
 	}
 
@@ -800,15 +800,15 @@ void FD3D12ResourceLocation::SetResource(FD3D12Resource* Value)
 }
 
 
-void FD3D12ResourceLocation::AsStandAlone(FD3D12Resource* Resource, uint32 BufferSize, bool bInIsTransient)
+void FD3D12ResourceLocation::AsStandAlone(FD3D12Resource* Resource, uint64 InSize, bool bInIsTransient)
 {
 	SetType(FD3D12ResourceLocation::ResourceLocationType::eStandAlone);
 	SetResource(Resource);
-	SetSize(BufferSize);
+	SetSize(InSize);
 
 	if (!IsCPUInaccessible(Resource->GetHeapType()))
 	{
-		D3D12_RANGE range = { 0, IsCPUWritable(Resource->GetHeapType()) ? 0 : BufferSize };
+		D3D12_RANGE range = { 0, IsCPUWritable(Resource->GetHeapType()) ? 0 : InSize };
 		SetMappedBaseAddress(Resource->Map(&range));
 	}
 	SetGPUVirtualAddress(Resource->GetGPUVirtualAddress());
@@ -849,13 +849,38 @@ bool FD3D12ResourceLocation::OnAllocationMoved(FRHIPoolAllocationData* InNewData
 		// recreate the placed resource (ownership of current resource is already handled during the internal move)
 		FD3D12HeapAndOffset HeapAndOffset = NewAllocator->GetBackingHeapAndAllocationOffsetInBytes(*this);
 
-		CResourceState& ResourceState = CurrentResource->GetResourceState();
-		check(ResourceState.AreAllSubresourcesSame());
-		D3D12_RESOURCE_STATES CreateState = ResourceState.GetSubresourceState(0);
+		D3D12_RESOURCE_STATES CreateState;
+		ED3D12ResourceStateMode ResourceStateMode;
+		if (CurrentResource->RequiresResourceStateTracking())
+		{			
+			CResourceState& ResourceState = CurrentResource->GetResourceState();
+			if (ResourceState.AreAllSubresourcesSame())
+			{
+				// All resource states the same so we can just create the resource at that state and know everything is fine
+				CreateState = ResourceState.GetSubresourceState(0);
+			}
+			else
+			{
+				// Force in the readable state when there are different states (or use just state of subresource 0?)
+				// ideally restore all subresources in correct state? needed for when not using GUseInternalTransitions anymore because then it needs to match the set state from the engine
+				check(GUseInternalTransitions);
+				CreateState = CurrentResource->GetReadableState();
+			}
+			ResourceStateMode = ED3D12ResourceStateMode::MultiState;
+		}
+		else
+		{
+			CreateState = CurrentResource->GetDefaultResourceState();
+			ResourceStateMode = ED3D12ResourceStateMode::Default;
+		}
+
+		// TODO: fix retrieval of ClearValue from owner (currently not a problem because not defragging RT/DS resource yet)
+		D3D12_CLEAR_VALUE* ClearValue = nullptr;
+
+		FName Name = CurrentResource->GetName();
 
 		FD3D12Resource* NewResource = nullptr;
-		VERIFYD3D12RESULT(CurrentResource->GetParentDevice()->GetParentAdapter()->CreatePlacedResource(CurrentResource->GetDesc(), HeapAndOffset.Heap, HeapAndOffset.Offset, CreateState, 
-			ED3D12ResourceStateMode::MultiState, D3D12_RESOURCE_STATE_TBD, nullptr, &NewResource, nullptr)); // CurrentResource->GetName()
+		VERIFYD3D12RESULT(CurrentResource->GetParentDevice()->GetParentAdapter()->CreatePlacedResource(CurrentResource->GetDesc(), HeapAndOffset.Heap, HeapAndOffset.Offset, CreateState, ResourceStateMode, D3D12_RESOURCE_STATE_TBD, ClearValue, &NewResource, *Name.ToString()));
 
 		UnderlyingResource = NewResource;
 	}
diff --git a/Engine/Source/Runtime/D3D12RHI/Private/D3D12Texture.cpp b/Engine/Source/Runtime/D3D12RHI/Private/D3D12Texture.cpp
--- a/Engine/Source/Runtime/D3D12RHI/Private/D3D12Texture.cpp
+++ b/Engine/Source/Runtime/D3D12RHI/Private/D3D12Texture.cpp
@@ -95,7 +95,7 @@ struct FRHICommandUpdateTextureString
 struct FRHICommandUpdateTexture final : public FRHICommand<FRHICommandUpdateTexture, FRHICommandUpdateTextureString>
 {
 	FD3D12TextureBase* TextureBase;
-	D3D12_TEXTURE_COPY_LOCATION DestCopyLocation;
+	uint32 MipIndex;
 	uint32 DestX;
 	uint32 DestY;
 	uint32 DestZ;
@@ -103,17 +103,16 @@ struct FRHICommandUpdateTexture final : public FRHICommand<FRHICommandUpdateText
 	FD3D12ResourceLocation Source;
 
 	FORCEINLINE_DEBUGGABLE FRHICommandUpdateTexture(FD3D12TextureBase* InTextureBase,
-		const D3D12_TEXTURE_COPY_LOCATION& InDestCopyLocation, uint32 InDestX, uint32 InDestY, uint32 InDestZ,
+		uint32 InMipIndex, uint32 InDestX, uint32 InDestY, uint32 InDestZ,
 		const D3D12_TEXTURE_COPY_LOCATION& InSourceCopyLocation, FD3D12ResourceLocation* InSource)
 		: TextureBase(InTextureBase)
-		, DestCopyLocation(InDestCopyLocation)
+		, MipIndex(InMipIndex)
 		, DestX(InDestX)
 		, DestY(InDestY)
 		, DestZ(InDestZ)
 		, SourceCopyLocation(InSourceCopyLocation)
 		, Source(nullptr)
 	{
-		DestCopyLocation.pResource->AddRef();
 		if (InSource)
 		{
 			FD3D12ResourceLocation::TransferOwnership(Source, *InSource);
@@ -122,12 +121,11 @@ struct FRHICommandUpdateTexture final : public FRHICommand<FRHICommandUpdateText
 
 	~FRHICommandUpdateTexture()
 	{
-		DestCopyLocation.pResource->Release();
 	}
 
 	void Execute(FRHICommandListBase& CmdList)
 	{
-		TextureBase->UpdateTexture(DestCopyLocation, DestX, DestY, DestZ, SourceCopyLocation);
+		TextureBase->UpdateTexture(MipIndex, DestX, DestY, DestZ, SourceCopyLocation);
 	}
 };
 
@@ -249,6 +247,9 @@ struct FD3D12RHICommandInitializeTexture final : public FRHICommand<FD3D12RHICom
 			{
 				check(Resource->GetDefaultResourceState() == DestinationState);
 			}
+
+			// Texture is now written and ready, so unlock the block (locked after creation and can be defragmented if needed)
+			CurrentTexture.ResourceLocation.UnlockPoolData();
 		}
 
 		if (!bAllocateOnStack)
@@ -668,6 +669,7 @@ void SafeCreateTexture2D(FD3D12Device* pDevice,
 	const D3D12_RESOURCE_DESC& TextureDesc,
 	const D3D12_CLEAR_VALUE* ClearValue, 
 	FD3D12ResourceLocation* OutTexture2D, 
+	FD3D12BaseShaderResource* Owner,
 	uint8 Format, 
 	ETextureCreateFlags Flags,
 	D3D12_RESOURCE_STATES InitialState,
@@ -701,8 +703,11 @@ void SafeCreateTexture2D(FD3D12Device* pDevice,
 			break;
 
 		case D3D12_HEAP_TYPE_DEFAULT:
+		{
 			VERIFYD3D12CREATETEXTURERESULT(pDevice->GetTextureAllocator().AllocateTexture(TextureDesc, ClearValue, Format, *OutTexture2D, InitialState, Name), TextureDesc, pDevice->GetDevice());
+			OutTexture2D->SetOwner(Owner);
 			break;
+		}
 
 		default:
 			check(false);	// Need to create a resource here
@@ -893,10 +898,17 @@ TD3D12Texture2D<BaseResourceType>* FD3D12DynamicRHI::CreateD3D12Texture2D(FRHICo
 			TextureDesc,
 			ClearValuePtr,
 			&Location,
+			NewTexture,
 			Format,
 			Flags,
 			(CreateInfo.BulkData != nullptr) ? D3D12_RESOURCE_STATE_COPY_DEST : InitialState,
 			CreateInfo.DebugName);
+		
+		// Unlock immediately if no initial data
+		if (CreateInfo.BulkData == nullptr)
+		{
+			Location.UnlockPoolData();
+		}
 
 		uint32 RTVIndex = 0;
 
@@ -1154,6 +1166,13 @@ FD3D12Texture3D* FD3D12DynamicRHI::CreateD3D12Texture3D(FRHICommandListImmediate
 		FD3D12Texture3D* Texture3D = new FD3D12Texture3D(Device, SizeX, SizeY, SizeZ, NumMips, Format, Flags, CreateInfo.ClearValueBinding);
 
 		VERIFYD3D12CREATETEXTURERESULT(Device->GetTextureAllocator().AllocateTexture(TextureDesc, ClearValuePtr, Format, Texture3D->ResourceLocation, (CreateInfo.BulkData != nullptr) ? D3D12_RESOURCE_STATE_COPY_DEST : InitialState, CreateInfo.DebugName), TextureDesc, Device->GetDevice());
+		Texture3D->ResourceLocation.SetOwner(Texture3D);
+
+		// Unlock immediately if no initial data
+		if (CreateInfo.BulkData == nullptr)
+		{
+			Texture3D->ResourceLocation.UnlockPoolData();
+		}
 
 		if (bCreateRTV)
 		{
@@ -1294,6 +1313,7 @@ FTexture2DRHIRef FD3D12DynamicRHI::RHIAsyncCreateTexture2D(uint32 SizeX, uint32
 			TextureDesc,
 			nullptr,
 			&NewTexture->ResourceLocation,
+			NewTexture,
 			Format,
 			Flags,
 			InitialState,
@@ -1474,6 +1494,9 @@ void FD3D12DynamicRHI::RHICopySharedMips(FRHITexture2D* DestTexture2DRHI, FRHITe
 			}
 		}
 
+		// unlock the pool allocated resource because all data has been written
+		DestTexture2D->ResourceLocation.UnlockPoolData();
+
 		DEBUG_EXECUTE_COMMAND_CONTEXT(Device->GetDefaultCommandContext());
 	}
 }
@@ -1809,14 +1832,16 @@ void* TD3D12Texture2D<RHIResourceType>::Lock(class FRHICommandListImmediate* RHI
 	return Data;
 }
 
-void FD3D12TextureBase::UpdateTexture(const D3D12_TEXTURE_COPY_LOCATION& DestCopyLocation, uint32 DestX, uint32 DestY, uint32 DestZ, const D3D12_TEXTURE_COPY_LOCATION& SourceCopyLocation)
+void FD3D12TextureBase::UpdateTexture(uint32 MipIndex, uint32 DestX, uint32 DestY, uint32 DestZ, const D3D12_TEXTURE_COPY_LOCATION& SourceCopyLocation)
 {
 	FD3D12CommandContext& DefaultContext = GetParentDevice()->GetDefaultCommandContext();
 	FD3D12CommandListHandle& hCommandList = DefaultContext.CommandListHandle;
 
-	FScopedResourceBarrier ScopeResourceBarrierDest(hCommandList, GetResource(), D3D12_RESOURCE_STATE_COPY_DEST, DestCopyLocation.SubresourceIndex, FD3D12DynamicRHI::ETransitionMode::Apply);
+	FScopedResourceBarrier ScopeResourceBarrierDest(hCommandList, GetResource(), D3D12_RESOURCE_STATE_COPY_DEST, MipIndex, FD3D12DynamicRHI::ETransitionMode::Apply);
 	// Don't need to transition upload heaps
 
+	CD3DX12_TEXTURE_COPY_LOCATION DestCopyLocation(GetResource()->GetResource(), MipIndex);
+
 	DefaultContext.numCopies++;
 	hCommandList.FlushResourceBarriers();
 	hCommandList->CopyTextureRegion(
@@ -1992,7 +2017,6 @@ void TD3D12Texture2D<RHIResourceType>::UnlockInternal(class FRHICommandListImmed
 			PlacedTexture2D.Offset = UploadLocation.GetOffsetFromBaseOfResource();
 			PlacedTexture2D.Footprint = BufferPitchDesc;
 
-			CD3DX12_TEXTURE_COPY_LOCATION DestCopyLocation(Resource->GetResource(), Subresource);
 			CD3DX12_TEXTURE_COPY_LOCATION SourceCopyLocation(UploadLocation.GetResource()->GetResource(), PlacedTexture2D);
 
 			FD3D12CommandListHandle& hCommandList = GetParentDevice()->GetDefaultCommandContext().CommandListHandle;
@@ -2002,11 +2026,11 @@ void TD3D12Texture2D<RHIResourceType>::UnlockInternal(class FRHICommandListImmed
 			{
 				// Same FD3D12ResourceLocation is used for all resources in the chain, therefore only the last command must be responsible for releasing it.
 				FD3D12ResourceLocation* Source = NextObject ? nullptr : &UploadLocation;
-				ALLOC_COMMAND_CL(*RHICmdList, FRHICommandUpdateTexture)(this, DestCopyLocation, 0, 0, 0, SourceCopyLocation, Source);
+				ALLOC_COMMAND_CL(*RHICmdList, FRHICommandUpdateTexture)(this, Subresource, 0, 0, 0, SourceCopyLocation, Source);
 			}
 			else
 			{
-				UpdateTexture(DestCopyLocation, 0, 0, 0, SourceCopyLocation);
+				UpdateTexture(Subresource, 0, 0, 0, SourceCopyLocation);
 			}
 
 			// Recurse to update all of the resources in the LDA chain
@@ -2073,17 +2097,16 @@ void TD3D12Texture2D<RHIResourceType>::UpdateTexture2D(class FRHICommandListImme
 		PlacedTexture2D.Offset = UploadHeapResourceLocation.GetOffsetFromBaseOfResource();
 		PlacedTexture2D.Footprint = SourceSubresource;
 
-		CD3DX12_TEXTURE_COPY_LOCATION DestCopyLocation(Texture.GetResource()->GetResource(), MipIndex);
 		CD3DX12_TEXTURE_COPY_LOCATION SourceCopyLocation(UploadHeapResourceLocation.GetResource()->GetResource(), PlacedTexture2D);
 
 		// If we are on the render thread, queue up the copy on the RHIThread so it happens at the correct time.
 		if (ShouldDeferCmdListOperation(RHICmdList))
 		{
-			ALLOC_COMMAND_CL(*RHICmdList, FRHICommandUpdateTexture)(this, DestCopyLocation, UpdateRegion.DestX, UpdateRegion.DestY, 0, SourceCopyLocation, &UploadHeapResourceLocation);
+			ALLOC_COMMAND_CL(*RHICmdList, FRHICommandUpdateTexture)(&Texture, MipIndex, UpdateRegion.DestX, UpdateRegion.DestY, 0, SourceCopyLocation, &UploadHeapResourceLocation);
 		}
 		else
 		{
-			UpdateTexture(DestCopyLocation, UpdateRegion.DestX, UpdateRegion.DestY, 0, SourceCopyLocation);
+			Texture.UpdateTexture(MipIndex, UpdateRegion.DestX, UpdateRegion.DestY, 0, SourceCopyLocation);
 		}
 	}
 }
diff --git a/Engine/Source/Runtime/D3D12RHI/Private/D3D12Texture.h b/Engine/Source/Runtime/D3D12RHI/Private/D3D12Texture.h
--- a/Engine/Source/Runtime/D3D12RHI/Private/D3D12Texture.h
+++ b/Engine/Source/Runtime/D3D12RHI/Private/D3D12Texture.h
@@ -13,7 +13,8 @@ void SafeCreateTexture2D(FD3D12Device* pDevice,
 	FD3D12Adapter* Adapter,
 	const D3D12_RESOURCE_DESC& TextureDesc,
 	const D3D12_CLEAR_VALUE* ClearValue, 
-	FD3D12ResourceLocation* OutTexture2D, 
+	FD3D12ResourceLocation* OutTexture2D,
+	FD3D12BaseShaderResource* Owner,
 	uint8 Format, 
 	ETextureCreateFlags Flags,
 	D3D12_RESOURCE_STATES InitialState,
@@ -54,6 +55,7 @@ public:
 		{
 			DepthStencilViews[SubResourceIndex] = View;
 			NumDepthStencilViews = FMath::Max(SubResourceIndex + 1, NumDepthStencilViews);
+			AddView(View);
 		}
 		else
 		{
@@ -66,6 +68,7 @@ public:
 		if (SubResourceIndex < (uint32)RenderTargetViews.Num())
 		{
 			RenderTargetViews[SubResourceIndex] = View;
+			AddView(View);
 		}
 		else
 		{
@@ -77,6 +80,7 @@ public:
 	{
 		RenderTargetViews.Empty(1);
 		RenderTargetViews.Add(View);
+		AddView(View);
 	}
 
 	int64 GetMemorySize() const
@@ -102,7 +106,7 @@ public:
 	FD3D12BaseShaderResource* GetBaseShaderResource() const { return BaseShaderResource; }
 	inline const FTextureRHIRef& GetAliasingSourceTexture() const { return AliasingSourceTexture; }
 
-	void SetShaderResourceView(FD3D12ShaderResourceView* InShaderResourceView) { ShaderResourceView = InShaderResourceView; }
+	void SetShaderResourceView(FD3D12ShaderResourceView* InShaderResourceView) { ShaderResourceView = InShaderResourceView; AddView(InShaderResourceView); }
 
 	static inline bool ShouldDeferCmdListOperation(FRHICommandList* RHICmdList)
 	{
@@ -119,7 +123,7 @@ public:
 		return true;
 	}
 
-	void UpdateTexture(const D3D12_TEXTURE_COPY_LOCATION& DestCopyLocation, uint32 DestX, uint32 DestY, uint32 DestZ, const D3D12_TEXTURE_COPY_LOCATION& SourceCopyLocation);
+	void UpdateTexture(uint32 MipIndex, uint32 DestX, uint32 DestY, uint32 DestZ, const D3D12_TEXTURE_COPY_LOCATION& SourceCopyLocation);
 	void CopyTextureRegion(uint32 DestX, uint32 DestY, uint32 DestZ, FD3D12TextureBase* SourceTexture, const D3D12_BOX& SourceBox);
 	void InitializeTextureData(class FRHICommandListImmediate* RHICmdList, const void* InitData, uint32 InitDataSize, uint32 SizeX, uint32 SizeY, uint32 SizeZ, uint32 NumSlices, uint32 NumMips, EPixelFormat Format, D3D12_RESOURCE_STATES DestinationState);
 
diff --git a/Engine/Source/Runtime/D3D12RHI/Private/D3D12View.cpp b/Engine/Source/Runtime/D3D12RHI/Private/D3D12View.cpp
--- a/Engine/Source/Runtime/D3D12RHI/Private/D3D12View.cpp
+++ b/Engine/Source/Runtime/D3D12RHI/Private/D3D12View.cpp
@@ -91,7 +91,7 @@ FD3D12ShaderResourceView* CreateSRV(TextureType* Texture, D3D12_SHADER_RESOURCE_
 
 	return Adapter->CreateLinkedViews<TextureType, FD3D12ShaderResourceView>(Texture, [&Desc](TextureType* Texture)
 	{
-		return new FD3D12ShaderResourceView(Texture->GetParentDevice(), Desc, Texture->ResourceLocation);
+		return new FD3D12ShaderResourceView(Texture, Desc);
 	});
 }
 
diff --git a/Engine/Source/Runtime/D3D12RHI/Private/D3D12View.h b/Engine/Source/Runtime/D3D12RHI/Private/D3D12View.h
--- a/Engine/Source/Runtime/D3D12RHI/Private/D3D12View.h
+++ b/Engine/Source/Runtime/D3D12RHI/Private/D3D12View.h
@@ -817,6 +817,13 @@ public:
 		Initialize(InDesc, InResourceLocation, InStride, InStartOffsetBytes, InSkipFastClearFinalize);
 	}
 
+	template<typename TextureType>
+	FD3D12ShaderResourceView(TextureType* InTexture, D3D12_SHADER_RESOURCE_VIEW_DESC& InDesc)
+		: FD3D12ShaderResourceView(InTexture->GetParentDevice(), InDesc, InTexture->ResourceLocation)
+	{
+		InTexture->AddView(this);
+	}
+
 	~FD3D12ShaderResourceView()
 	{
 		// Non-resident mesh LODs can lead to null SRVs being created
@@ -1018,7 +1025,7 @@ public:
 };
 #endif
 
-class FD3D12RenderTargetView : public FD3D12View<D3D12_RENDER_TARGET_VIEW_DESC>, public FRHIResource, public FD3D12LinkedAdapterObject<FD3D12RenderTargetView>
+class FD3D12RenderTargetView : public FD3D12BaseShaderResourceView, public FD3D12View<D3D12_RENDER_TARGET_VIEW_DESC>, public FRHIResource, public FD3D12LinkedAdapterObject<FD3D12RenderTargetView>
 {
 public:
 	FD3D12RenderTargetView(FD3D12Device* InParent, const D3D12_RENDER_TARGET_VIEW_DESC& InRTVDesc, FD3D12ResourceLocation& InResourceLocation)
@@ -1026,9 +1033,22 @@ public:
 	{
 		CreateView(InRTVDesc, InResourceLocation);
 	}
+	
+	~FD3D12RenderTargetView()
+	{
+		FD3D12BaseShaderResourceView::Remove();
+	}
+
+	virtual void Recreate(FD3D12ResourceLocation& InResourceLocation) override
+	{
+		check(ResourceLocation == &InResourceLocation);
+		check(InResourceLocation.GetOffsetFromBaseOfResource() == 0);
+
+		CreateView(Desc, InResourceLocation);
+	}
 };
 
-class FD3D12DepthStencilView : public FD3D12View<D3D12_DEPTH_STENCIL_VIEW_DESC>, public FRHIResource, public FD3D12LinkedAdapterObject<FD3D12DepthStencilView>
+class FD3D12DepthStencilView : public FD3D12BaseShaderResourceView, public FD3D12View<D3D12_DEPTH_STENCIL_VIEW_DESC>, public FRHIResource, public FD3D12LinkedAdapterObject<FD3D12DepthStencilView>
 {
 	const bool bHasDepth;
 	const bool bHasStencil;
@@ -1065,6 +1085,11 @@ public:
 		}
 	}
 
+	~FD3D12DepthStencilView()
+	{
+		FD3D12BaseShaderResourceView::Remove();
+	}
+
 	bool HasDepth() const
 	{
 		return bHasDepth;
@@ -1086,6 +1111,14 @@ public:
 		check(bHasStencil);
 		return StencilOnlyViewSubresourceSubset;
 	}
+
+	virtual void Recreate(FD3D12ResourceLocation& InResourceLocation) override
+	{
+		check(ResourceLocation == &InResourceLocation);
+		check(InResourceLocation.GetOffsetFromBaseOfResource() == 0);
+
+		CreateView(Desc, InResourceLocation);
+	}
 };
 
 template<>
diff --git a/Engine/Source/Runtime/D3D12RHI/Private/D3D12Viewport.cpp b/Engine/Source/Runtime/D3D12RHI/Private/D3D12Viewport.cpp
--- a/Engine/Source/Runtime/D3D12RHI/Private/D3D12Viewport.cpp
+++ b/Engine/Source/Runtime/D3D12RHI/Private/D3D12Viewport.cpp
@@ -273,6 +273,7 @@ FD3D12Texture2D* GetSwapChainSurface(FD3D12Device* Parent, EPixelFormat PixelFor
 				BackBufferDesc,
 				nullptr, // &ClearValueBinding,
 				&NewTexture->ResourceLocation,
+				NewTexture,
 				PixelFormat,
 				TexCreate_RenderTargetable |  TexCreate_ShaderResource,
 				D3D12_RESOURCE_STATE_PRESENT,
diff --git a/Engine/Source/Runtime/D3D12RHI/Public/D3D12Resources.h b/Engine/Source/Runtime/D3D12RHI/Public/D3D12Resources.h
--- a/Engine/Source/Runtime/D3D12RHI/Public/D3D12Resources.h
+++ b/Engine/Source/Runtime/D3D12RHI/Public/D3D12Resources.h
@@ -560,7 +560,7 @@ public:
 
 	const inline bool IsValid() const { return Type != ResourceLocationType::eUndefined; }
 
-	void AsStandAlone(FD3D12Resource* Resource, uint32 BufferSize = 0, bool bInIsTransient = false);
+	void AsStandAlone(FD3D12Resource* Resource, uint64 InSize = 0, bool bInIsTransient = false);
 
 	inline void AsHeapAliased(FD3D12Resource* Resource)
 	{
diff --git a/Engine/Source/Runtime/RHICore/Private/RHIPoolAllocator.cpp b/Engine/Source/Runtime/RHICore/Private/RHIPoolAllocator.cpp
--- a/Engine/Source/Runtime/RHICore/Private/RHIPoolAllocator.cpp
+++ b/Engine/Source/Runtime/RHICore/Private/RHIPoolAllocator.cpp
@@ -196,6 +196,8 @@ FRHIMemoryPool::FRHIMemoryPool(int16 InPoolIndex, uint64 InPoolSize, uint32 InPo
 	AligmnentWaste(0), 
 	AllocatedBlocks(0)
 {	
+	// PoolSize should fit in 32 bits because FRHIPoolAllocationData only stores 32bit size
+	check(PoolSize < UINT32_MAX);
 }
 
 
@@ -518,8 +520,7 @@ FRHIPoolAllocator::~FRHIPoolAllocator()
 
 
 void FRHIPoolAllocator::Initialize()
-{
-	Pools.Add(CreateNewPool(0));
+{	
 }
 
 
