Fix issue with missing command list when creating buffers from RenderThread

diff --git a/Engine/Source/Runtime/D3D12RHI/Private/D3D12Buffer.cpp b/Engine/Source/Runtime/D3D12RHI/Private/D3D12Buffer.cpp
--- a/Engine/Source/Runtime/D3D12RHI/Private/D3D12Buffer.cpp
+++ b/Engine/Source/Runtime/D3D12RHI/Private/D3D12Buffer.cpp
@@ -397,6 +397,16 @@ D3D12_RESOURCE_DESC CreateBufferResourceDesc(uint32 Size, uint32 Usage)
 }
 
 FBufferRHIRef FD3D12DynamicRHI::RHICreateBuffer(uint32 Size, EBufferUsageFlags Usage, uint32 Stride, ERHIAccess InResourceState, FRHIResourceCreateInfo& CreateInfo)
+{	
+	return CreateBuffer(nullptr, Size, Usage, Stride, InResourceState, CreateInfo);
+}
+
+FBufferRHIRef FD3D12DynamicRHI::CreateBuffer_RenderThread(class FRHICommandListImmediate& RHICmdList, uint32 Size, EBufferUsageFlags Usage, uint32 Stride, ERHIAccess ResourceState, FRHIResourceCreateInfo& CreateInfo)
+{
+	return CreateBuffer(&RHICmdList, Size, Usage, Stride, ResourceState, CreateInfo);
+}
+
+FBufferRHIRef FD3D12DynamicRHI::CreateBuffer(FRHICommandListImmediate* RHICmdList, uint32 Size, EBufferUsageFlags Usage, uint32 Stride, ERHIAccess InResourceState, FRHIResourceCreateInfo& CreateInfo)
 {
 	if (CreateInfo.bWithoutNativeResource)
 	{
@@ -408,9 +418,9 @@ FBufferRHIRef FD3D12DynamicRHI::RHICreateBuffer(uint32 Size, EBufferUsageFlags U
 
 	const D3D12_RESOURCE_DESC Desc = CreateBufferResourceDesc(Size, Usage);
 	// Structured buffers, non-ByteAddress buffers, need to be aligned to their stride to ensure that they can be addressed correctly with element based offsets.
-	const uint32 Alignment = (Stride > 0) && (((Usage & BUF_StructuredBuffer) != 0) || ((Usage & (BUF_ByteAddressBuffer | BUF_DrawIndirect)) == 0))? Stride : 4;
+	const uint32 Alignment = (Stride > 0) && (((Usage & BUF_StructuredBuffer) != 0) || ((Usage & (BUF_ByteAddressBuffer | BUF_DrawIndirect)) == 0)) ? Stride : 4;
 
-	FD3D12Buffer* Buffer = GetAdapter().CreateRHIBuffer(nullptr, Desc, Alignment, Stride, Size, Usage, ED3D12ResourceStateMode::Default, InResourceState, CreateInfo);
+	FD3D12Buffer* Buffer = GetAdapter().CreateRHIBuffer(RHICmdList, Desc, Alignment, Stride, Size, Usage, ED3D12ResourceStateMode::Default, InResourceState, CreateInfo);
 	if (Buffer->ResourceLocation.IsTransient())
 	{
 		// TODO: this should ideally be set in platform-independent code, since this tracking is for the high level
@@ -420,11 +430,6 @@ FBufferRHIRef FD3D12DynamicRHI::RHICreateBuffer(uint32 Size, EBufferUsageFlags U
 	return Buffer;
 }
 
-FBufferRHIRef FD3D12DynamicRHI::CreateBuffer_RenderThread(class FRHICommandListImmediate& RHICmdList, uint32 Size, EBufferUsageFlags Usage, uint32 Stride, ERHIAccess ResourceState, FRHIResourceCreateInfo& CreateInfo)
-{
-	return RHICreateBuffer(Size, Usage, Stride, ResourceState, CreateInfo);
-}
-
 void* FD3D12DynamicRHI::LockBuffer(FRHICommandListImmediate* RHICmdList, FD3D12Buffer* Buffer, uint32 BufferSize, uint32 BufferUsage, uint32 Offset, uint32 Size, EResourceLockMode LockMode)
 {
 	SCOPE_CYCLE_COUNTER(STAT_D3D12LockBufferTime);
diff --git a/Engine/Source/Runtime/D3D12RHI/Private/D3D12RHIPrivate.h b/Engine/Source/Runtime/D3D12RHI/Private/D3D12RHIPrivate.h
--- a/Engine/Source/Runtime/D3D12RHI/Private/D3D12RHIPrivate.h
+++ b/Engine/Source/Runtime/D3D12RHI/Private/D3D12RHIPrivate.h
@@ -483,6 +483,7 @@ public:
 
 	static int32 GetResourceBarrierBatchSizeLimit();
 
+	FBufferRHIRef CreateBuffer(FRHICommandListImmediate* RHICmdList, uint32 Size, EBufferUsageFlags Usage, uint32 Stride, ERHIAccess ResourceState, FRHIResourceCreateInfo& CreateInfo);
 	void* LockBuffer(FRHICommandListImmediate* RHICmdList, FD3D12Buffer* Buffer, uint32 BufferSize, uint32 BufferUsage, uint32 Offset, uint32 Size, EResourceLockMode LockMode);
 	void UnlockBuffer(FRHICommandListImmediate* RHICmdList, FD3D12Buffer* Buffer, uint32 BufferUsage);
 
