Refactor/Generalize resource changed listener code in D3D12

#b yuriy.odonnell, emil.persson

diff --git a/Engine/Source/Runtime/D3D12RHI/Private/D3D12Buffer.cpp b/Engine/Source/Runtime/D3D12RHI/Private/D3D12Buffer.cpp
--- a/Engine/Source/Runtime/D3D12RHI/Private/D3D12Buffer.cpp
+++ b/Engine/Source/Runtime/D3D12RHI/Private/D3D12Buffer.cpp
@@ -323,7 +323,7 @@ FD3D12Buffer* FD3D12Adapter::CreateRHIBuffer(FRHICommandListImmediate* RHICmdLis
 void FD3D12Buffer::Rename(FD3D12ResourceLocation& NewLocation)
 {
 	FD3D12ResourceLocation::TransferOwnership(ResourceLocation, NewLocation);
-	RecreateViews(ResourceLocation);
+	ResourceRenamed(&ResourceLocation);
 }
 
 void FD3D12Buffer::RenameLDAChain(FD3D12ResourceLocation& NewLocation)
@@ -342,7 +342,7 @@ void FD3D12Buffer::RenameLDAChain(FD3D12ResourceLocation& NewLocation)
 		for (auto NextBuffer = ++FLinkedObjectIterator(this); NextBuffer; ++NextBuffer)
 		{
 			FD3D12ResourceLocation::ReferenceNode(NextBuffer->GetParentDevice(), NextBuffer->ResourceLocation, ResourceLocation);
-			NextBuffer->RecreateViews(NextBuffer->ResourceLocation);
+			NextBuffer->ResourceRenamed(&NextBuffer->ResourceLocation);
 		}
 	}
 }
@@ -363,7 +363,7 @@ void FD3D12Buffer::ReleaseUnderlyingResource()
 	{
 		check(!NextBuffer->LockedData.bLocked && NextBuffer->ResourceLocation.IsValid());
 		NextBuffer->ResourceLocation.Clear();
-		NextBuffer->RemoveAllViews();
+		NextBuffer->RemoveAllRenameListeners();
 	}
 }
 
diff --git a/Engine/Source/Runtime/D3D12RHI/Private/D3D12RayTracing.cpp b/Engine/Source/Runtime/D3D12RHI/Private/D3D12RayTracing.cpp
--- a/Engine/Source/Runtime/D3D12RHI/Private/D3D12RayTracing.cpp
+++ b/Engine/Source/Runtime/D3D12RHI/Private/D3D12RayTracing.cpp
@@ -2974,7 +2974,7 @@ void FD3D12RayTracingScene::BuildAccelerationStructure(FD3D12CommandContext& Com
 	if (bAccelerationStructureViewInitialized[GPUIndex])
 	{
 		check(AccelerationStructureView->GetParentDevice() == AccelerationStructureBuffer->GetParentDevice());
-		AccelerationStructureView->Rename(AccelerationStructureBuffer->ResourceLocation);
+		AccelerationStructureView->Rename(AccelerationStructureBuffer.GetReference());
 	}
 	else
 	{
@@ -2985,7 +2985,7 @@ void FD3D12RayTracingScene::BuildAccelerationStructure(FD3D12CommandContext& Com
 		SRVDesc.RaytracingAccelerationStructure.Location = AccelerationStructureBuffer->ResourceLocation.GetGPUVirtualAddress();
 
 		check(AccelerationStructureView->GetParentDevice() == AccelerationStructureBuffer->GetParentDevice());
-		AccelerationStructureView->Initialize(SRVDesc, AccelerationStructureBuffer->ResourceLocation, 4, 0);
+		AccelerationStructureView->Initialize(SRVDesc, AccelerationStructureBuffer, 4, 0);
 
 		bAccelerationStructureViewInitialized[GPUIndex] = true;
 	}
diff --git a/Engine/Source/Runtime/D3D12RHI/Private/D3D12RenderTarget.cpp b/Engine/Source/Runtime/D3D12RHI/Private/D3D12RenderTarget.cpp
--- a/Engine/Source/Runtime/D3D12RHI/Private/D3D12RenderTarget.cpp
+++ b/Engine/Source/Runtime/D3D12RHI/Private/D3D12RenderTarget.cpp
@@ -1406,7 +1406,7 @@ void FD3D12DynamicRHI::ReadSurfaceDataMSAARaw(FRHICommandList_RecursiveHazardous
 
 	RTVDesc.ViewDimension = D3D12_RTV_DIMENSION_TEXTURE2D;
 	RTVDesc.Texture2D.MipSlice = 0;
-	NonMSAARTV = new FD3D12RenderTargetView(Device, RTVDesc, ResourceLocation);
+	NonMSAARTV = new FD3D12RenderTargetView(Device, RTVDesc, Texture, ResourceLocation);
 
 	// Create a CPU-accessible staging texture to copy the resolved sample data to.
 	TRefCountPtr<FD3D12Resource> StagingTexture2D;
diff --git a/Engine/Source/Runtime/D3D12RHI/Private/D3D12Resources.cpp b/Engine/Source/Runtime/D3D12RHI/Private/D3D12Resources.cpp
--- a/Engine/Source/Runtime/D3D12RHI/Private/D3D12Resources.cpp
+++ b/Engine/Source/Runtime/D3D12RHI/Private/D3D12Resources.cpp
@@ -888,8 +888,8 @@ bool FD3D12ResourceLocation::OnAllocationMoved(FRHIPoolAllocationData* InNewData
 	GPUVirtualAddress = UnderlyingResource->GetGPUVirtualAddress() + OffsetFromBaseOfResource;
 	ResidencyHandle = UnderlyingResource->GetResidencyHandle();
 
-	// Recreate all the views (SRVs and UAVs)
-	Owner->RecreateViews(*this);
+	// Notify all the dependent resources about the change
+	Owner->ResourceRenamed(this);
 
 	return true;
 }
diff --git a/Engine/Source/Runtime/D3D12RHI/Private/D3D12Texture.cpp b/Engine/Source/Runtime/D3D12RHI/Private/D3D12Texture.cpp
--- a/Engine/Source/Runtime/D3D12RHI/Private/D3D12Texture.cpp
+++ b/Engine/Source/Runtime/D3D12RHI/Private/D3D12Texture.cpp
@@ -936,7 +936,7 @@ TD3D12Texture2D<BaseResourceType>* FD3D12DynamicRHI::CreateD3D12Texture2D(FRHICo
 						RTVDesc.Texture2DArray.MipSlice = MipIndex;
 						RTVDesc.Texture2DArray.PlaneSlice = GetPlaneSliceFromViewFormat(PlatformResourceFormat, RTVDesc.Format);
 
-						NewTexture->SetRenderTargetViewIndex(new FD3D12RenderTargetView(Device, RTVDesc, Location), RTVIndex++);
+						NewTexture->SetRenderTargetViewIndex(new FD3D12RenderTargetView(Device, RTVDesc, NewTexture), RTVIndex++);
 					}
 				}
 				else
@@ -978,7 +978,7 @@ TD3D12Texture2D<BaseResourceType>* FD3D12DynamicRHI::CreateD3D12Texture2D(FRHICo
 						}
 					}
 
-					NewTexture->SetRenderTargetViewIndex(new FD3D12RenderTargetView(Device, RTVDesc, Location), RTVIndex++);
+					NewTexture->SetRenderTargetViewIndex(new FD3D12RenderTargetView(Device, RTVDesc, NewTexture), RTVIndex++);
 				}
 			}
 		}
@@ -1028,7 +1028,7 @@ TD3D12Texture2D<BaseResourceType>* FD3D12DynamicRHI::CreateD3D12Texture2D(FRHICo
 					DSVDesc.Flags |= (AccessType & FExclusiveDepthStencil::DepthWrite_StencilRead) ? D3D12_DSV_FLAG_READ_ONLY_STENCIL : D3D12_DSV_FLAG_NONE;
 				}
 
-				NewTexture->SetDepthStencilView(new FD3D12DepthStencilView(Device, DSVDesc, Location, HasStencil), AccessType);
+				NewTexture->SetDepthStencilView(new FD3D12DepthStencilView(Device, DSVDesc, NewTexture, HasStencil), AccessType);
 			}
 		}
 
@@ -1088,7 +1088,7 @@ TD3D12Texture2D<BaseResourceType>* FD3D12DynamicRHI::CreateD3D12Texture2D(FRHICo
 				}
 			}
 
-			NewTexture->SetShaderResourceView(new FD3D12ShaderResourceView(Device, SRVDesc, Location));
+			NewTexture->SetShaderResourceView(new FD3D12ShaderResourceView(Device, SRVDesc, NewTexture));
 		}
 
 		return NewTexture;
@@ -1185,7 +1185,7 @@ FD3D12Texture3D* FD3D12DynamicRHI::CreateD3D12Texture3D(FRHICommandListImmediate
 			RTVDesc.Texture3D.FirstWSlice = 0;
 			RTVDesc.Texture3D.WSize = SizeZ;
 
-			Texture3D->SetRenderTargetView(new FD3D12RenderTargetView(Device, RTVDesc, Texture3D->ResourceLocation));
+			Texture3D->SetRenderTargetView(new FD3D12RenderTargetView(Device, RTVDesc, Texture3D));
 		}
 
 		// Create a shader resource view for the texture.
@@ -1196,7 +1196,7 @@ FD3D12Texture3D* FD3D12DynamicRHI::CreateD3D12Texture3D(FRHICommandListImmediate
 		SRVDesc.Texture3D.MipLevels = NumMips;
 		SRVDesc.Texture3D.MostDetailedMip = 0;
 
-		Texture3D->SetShaderResourceView(new FD3D12ShaderResourceView(Device, SRVDesc, Texture3D->ResourceLocation));
+		Texture3D->SetShaderResourceView(new FD3D12ShaderResourceView(Device, SRVDesc, Texture3D));
 
 		return Texture3D;
 	}); 
@@ -1328,7 +1328,7 @@ FTexture2DRHIRef FD3D12DynamicRHI::RHIAsyncCreateTexture2D(uint32 SizeX, uint32
 		SRVDesc.Texture2D.PlaneSlice = GetPlaneSliceFromViewFormat(PlatformResourceFormat, SRVDesc.Format);
 
 		// Create a wrapper for the SRV and set it on the texture
-		NewTexture->SetShaderResourceView(new FD3D12ShaderResourceView(Device, SRVDesc, NewTexture->ResourceLocation));
+		NewTexture->SetShaderResourceView(new FD3D12ShaderResourceView(Device, SRVDesc, NewTexture));
 
 		return NewTexture;
 	});
@@ -2809,7 +2809,7 @@ TD3D12Texture2D<BaseResourceType>* FD3D12DynamicRHI::CreateTextureFromResource(b
 					RTVDesc.Texture2DArray.MipSlice = MipIndex;
 					RTVDesc.Texture2DArray.PlaneSlice = GetPlaneSliceFromViewFormat(PlatformResourceFormat, RTVDesc.Format);
 
-					Texture2D->SetRenderTargetViewIndex(new FD3D12RenderTargetView(Device, RTVDesc, Location), RTVIndex++);
+					Texture2D->SetRenderTargetViewIndex(new FD3D12RenderTargetView(Device, RTVDesc, Texture2D), RTVIndex++);
 				}
 			}
 			else
@@ -2836,7 +2836,7 @@ TD3D12Texture2D<BaseResourceType>* FD3D12DynamicRHI::CreateTextureFromResource(b
 					RTVDesc.ViewDimension = D3D12_RTV_DIMENSION_TEXTURE2DMS;
 				}
 
-				Texture2D->SetRenderTargetViewIndex(new FD3D12RenderTargetView(Device, RTVDesc, Location), RTVIndex++);
+				Texture2D->SetRenderTargetViewIndex(new FD3D12RenderTargetView(Device, RTVDesc, Texture2D), RTVIndex++);
 			}
 		}
 	}
@@ -2873,7 +2873,7 @@ TD3D12Texture2D<BaseResourceType>* FD3D12DynamicRHI::CreateTextureFromResource(b
 				DSVDesc.Flags |= (AccessType & FExclusiveDepthStencil::DepthWrite_StencilRead) ? D3D12_DSV_FLAG_READ_ONLY_STENCIL : D3D12_DSV_FLAG_NONE;
 			}
 
-			Texture2D->SetDepthStencilView(new FD3D12DepthStencilView(Device, DSVDesc, Location, HasStencil), AccessType);
+			Texture2D->SetDepthStencilView(new FD3D12DepthStencilView(Device, DSVDesc, Texture2D, HasStencil), AccessType);
 		}
 	}
 
@@ -2920,7 +2920,7 @@ TD3D12Texture2D<BaseResourceType>* FD3D12DynamicRHI::CreateTextureFromResource(b
 	}
 
 	// Create a wrapper for the SRV and set it on the texture
-	Texture2D->SetShaderResourceView(new FD3D12ShaderResourceView(Device, SRVDesc, Location));
+	Texture2D->SetShaderResourceView(new FD3D12ShaderResourceView(Device, SRVDesc, Texture2D));
 
 	FD3D12TextureStats::D3D12TextureAllocated(*Texture2D);
 
diff --git a/Engine/Source/Runtime/D3D12RHI/Private/D3D12Texture.h b/Engine/Source/Runtime/D3D12RHI/Private/D3D12Texture.h
--- a/Engine/Source/Runtime/D3D12RHI/Private/D3D12Texture.h
+++ b/Engine/Source/Runtime/D3D12RHI/Private/D3D12Texture.h
@@ -55,7 +55,6 @@ public:
 		{
 			DepthStencilViews[SubResourceIndex] = View;
 			NumDepthStencilViews = FMath::Max(SubResourceIndex + 1, NumDepthStencilViews);
-			AddView(View);
 		}
 		else
 		{
@@ -68,7 +67,6 @@ public:
 		if (SubResourceIndex < (uint32)RenderTargetViews.Num())
 		{
 			RenderTargetViews[SubResourceIndex] = View;
-			AddView(View);
 		}
 		else
 		{
@@ -80,7 +78,6 @@ public:
 	{
 		RenderTargetViews.Empty(1);
 		RenderTargetViews.Add(View);
-		AddView(View);
 	}
 
 	int64 GetMemorySize() const
@@ -106,7 +103,7 @@ public:
 	FD3D12BaseShaderResource* GetBaseShaderResource() const { return BaseShaderResource; }
 	inline const FTextureRHIRef& GetAliasingSourceTexture() const { return AliasingSourceTexture; }
 
-	void SetShaderResourceView(FD3D12ShaderResourceView* InShaderResourceView) { ShaderResourceView = InShaderResourceView; AddView(InShaderResourceView); }
+	void SetShaderResourceView(FD3D12ShaderResourceView* InShaderResourceView) { ShaderResourceView = InShaderResourceView; }
 
 	static inline bool ShouldDeferCmdListOperation(FRHICommandList* RHICmdList)
 	{
diff --git a/Engine/Source/Runtime/D3D12RHI/Private/D3D12UAV.cpp b/Engine/Source/Runtime/D3D12RHI/Private/D3D12UAV.cpp
--- a/Engine/Source/Runtime/D3D12RHI/Private/D3D12UAV.cpp
+++ b/Engine/Source/Runtime/D3D12RHI/Private/D3D12UAV.cpp
@@ -24,9 +24,7 @@ inline FD3D12UnorderedAccessView* CreateUAV(D3D12_UNORDERED_ACCESS_VIEW_DESC& De
 			Device->GetParentAdapter()->CreateBuffer(D3D12_HEAP_TYPE_DEFAULT, Node, Node, 4, &CounterResource,  TEXT("Counter"), D3D12_RESOURCE_FLAG_ALLOW_UNORDERED_ACCESS);
 		}
 
-		FD3D12UnorderedAccessView* UAVView = new FD3D12UnorderedAccessView(Device, Desc, Resource->ResourceLocation, CounterResource);
-		Resource->AddView(UAVView);
-		return UAVView;
+		return new FD3D12UnorderedAccessView(Device, Desc, Resource, CounterResource);
 	});
 }
 
diff --git a/Engine/Source/Runtime/D3D12RHI/Private/D3D12View.cpp b/Engine/Source/Runtime/D3D12RHI/Private/D3D12View.cpp
--- a/Engine/Source/Runtime/D3D12RHI/Private/D3D12View.cpp
+++ b/Engine/Source/Runtime/D3D12RHI/Private/D3D12View.cpp
@@ -225,8 +225,7 @@ FShaderResourceViewRHIRef FD3D12DynamicRHI::RHICreateShaderResourceView(const FS
 					D3D12_SHADER_RESOURCE_VIEW_DESC SRVDesc = GetVertexBufferSRVDesc(VertexBuffer, CreationStride, Format, StartOffsetBytes, NumElements);
 
 					// Create a Shader Resource View
-					SRV->Initialize(SRVDesc, VertexBuffer->ResourceLocation, CreationStride, StartOffsetBytes);
-					VertexBuffer->AddView(SRV);
+					SRV->Initialize(SRVDesc, VertexBuffer, CreationStride, StartOffsetBytes);
 				}
 			};
 
@@ -309,8 +308,7 @@ FShaderResourceViewRHIRef FD3D12DynamicRHI::RHICreateShaderResourceView(const FS
 					SRVDesc.Buffer.FirstElement = Offset / Stride + StartElement;
 
 					// Create a Shader Resource View
-					SRV->Initialize(SRVDesc, StructuredBuffer->ResourceLocation, Stride, StartOffsetBytes);
-					StructuredBuffer->AddView(SRV);
+					SRV->Initialize(SRVDesc, StructuredBuffer, Stride, StartOffsetBytes);
 				}
 			};
 
@@ -347,11 +345,10 @@ FShaderResourceViewRHIRef FD3D12DynamicRHI::RHICreateShaderResourceView(const FS
 			{
 				check(Buffer);
 
-				FD3D12ResourceLocation& Location = Buffer->ResourceLocation;
 				const uint32 CreationStride = Buffer->GetStride();
 				D3D12_SHADER_RESOURCE_VIEW_DESC SRVDesc = GetIndexBufferSRVDesc(Buffer, Desc.StartOffsetBytes, Desc.NumElements);
 
-				FD3D12ShaderResourceView* ShaderResourceView = new FD3D12ShaderResourceView(Buffer->GetParentDevice(), SRVDesc, Location, CreationStride);
+				FD3D12ShaderResourceView* ShaderResourceView = new FD3D12ShaderResourceView(Buffer->GetParentDevice(), SRVDesc, Buffer, CreationStride);
 				return ShaderResourceView;
 			});
 		}
@@ -377,8 +374,7 @@ void FD3D12DynamicRHI::RHIUpdateShaderResourceView(FRHIShaderResourceView* SRV,
 			SRVD3D12 = It.GetSecond();
 
 			FD3D12Device* ParentDevice = Buffer->GetParentDevice();
-			SRVD3D12->Initialize(ParentDevice, SRVDesc, Buffer->ResourceLocation, Stride);
-			Buffer->AddView(SRVD3D12);
+			SRVD3D12->Initialize(ParentDevice, SRVDesc, Buffer, Stride);
 		}
 	}
 }
@@ -400,7 +396,7 @@ void FD3D12DynamicRHI::RHIUpdateShaderResourceView(FRHIShaderResourceView* SRV,
 			SRVD3D12 = It.GetSecond();
 
 			FD3D12Device* ParentDevice = Buffer->GetParentDevice();
-			SRVD3D12->Initialize(ParentDevice, SRVDesc, Buffer->ResourceLocation, Stride);
+			SRVD3D12->Initialize(ParentDevice, SRVDesc, Buffer, Stride);
 		}
 	}
 }
diff --git a/Engine/Source/Runtime/D3D12RHI/Private/D3D12View.h b/Engine/Source/Runtime/D3D12RHI/Private/D3D12View.h
--- a/Engine/Source/Runtime/D3D12RHI/Private/D3D12View.h
+++ b/Engine/Source/Runtime/D3D12RHI/Private/D3D12View.h
@@ -698,13 +698,14 @@ typedef TD3D12ViewDescriptorHandle<D3D12_DEPTH_STENCIL_VIEW_DESC>		FD3D12Descrip
 typedef TD3D12ViewDescriptorHandle<D3D12_UNORDERED_ACCESS_VIEW_DESC>	FD3D12DescriptorHandleUAV;
 
 template <typename TDesc>
-class FD3D12View
+class FD3D12View : public FD3D12ShaderResourceRenameListener
 {
 private:
 	TD3D12ViewDescriptorHandle<TDesc> Descriptor;
 
 protected:
 	ViewSubresourceSubsetFlags Flags;
+	FD3D12BaseShaderResource* BaseShaderResource;
 	FD3D12ResourceLocation* ResourceLocation;
 	FD3D12ResidencyHandle* ResidencyHandle;
 	FD3D12Resource* Resource;
@@ -718,6 +719,7 @@ protected:
 	explicit FD3D12View(FD3D12Device* InParent, ViewSubresourceSubsetFlags InFlags)
 		: Descriptor(InParent)
 		, Flags(InFlags)
+		, BaseShaderResource(nullptr)
 #if DO_CHECK
 		, bInitialized(false)
 #endif
@@ -725,14 +727,28 @@ protected:
 
 	virtual ~FD3D12View()
 	{
+		if (BaseShaderResource)
+		{
+			BaseShaderResource->RemoveRenameListener(this);
+		}
+
 #if DO_CHECK
 		bInitialized = false;
 #endif
 	}
 
 private:
-	void Initialize(const TDesc& InDesc, FD3D12ResourceLocation& InResourceLocation)
+	void InitializeInternal(const TDesc& InDesc, FD3D12BaseShaderResource* InBaseShaderResource, FD3D12ResourceLocation& InResourceLocation)
 	{
+		check(InBaseShaderResource);
+
+		// Only register the first time - init can be called again during rename
+		if (BaseShaderResource == nullptr)
+		{
+			InBaseShaderResource->AddRenameListener(this);
+		}
+
+		BaseShaderResource = InBaseShaderResource;
 		ResourceLocation = &InResourceLocation;
 		Resource = ResourceLocation->GetResource();
 		check(Resource);
@@ -753,23 +769,40 @@ private:
 	}
 
 protected:
-	void CreateView(const TDesc& InDesc, FD3D12ResourceLocation& InResourceLocation)
+	void CreateView(const TDesc& InDesc, FD3D12BaseShaderResource* InBaseShaderResource, FD3D12ResourceLocation& InResourceLocation)
 	{
-		Initialize(InDesc, InResourceLocation);
+		InitializeInternal(InDesc, InBaseShaderResource, InResourceLocation);
 
 		ID3D12Resource* D3DResource = ResourceLocation->GetResource()->GetResource();
 		Descriptor.CreateView(Desc, D3DResource);
 	}
 
-	void CreateViewWithCounter(const TDesc& InDesc, FD3D12ResourceLocation& InResourceLocation, FD3D12Resource* InCounterResource)
+	void CreateViewWithCounter(const TDesc& InDesc, FD3D12BaseShaderResource* InBaseShaderResource, FD3D12ResourceLocation& InResourceLocation, FD3D12Resource* InCounterResource)
 	{
-		Initialize(InDesc, InResourceLocation);
+		InitializeInternal(InDesc, InBaseShaderResource, InResourceLocation);
 
 		ID3D12Resource* D3DResource = ResourceLocation->GetResource()->GetResource();
 		ID3D12Resource* D3DCounterResource = InCounterResource ? InCounterResource->GetResource() : nullptr;
 		Descriptor.CreateViewWithCounter(Desc, D3DResource, D3DCounterResource);
 	}
 
+	virtual void ResourceRenamed(FD3D12BaseShaderResource* InRenamedResource, FD3D12ResourceLocation* InNewResourceLocation) override
+	{
+		check(InRenamedResource == BaseShaderResource);
+		if (InNewResourceLocation)
+		{
+			RecreateView();
+		}
+		else
+		{
+			// Marking not initialized will currently assert because views are used after the resource has been registered for delete
+			// Is that wrong?
+			BaseShaderResource = nullptr;
+		}
+	}
+
+	virtual void RecreateView() = 0;
+
 public:
 	inline FD3D12Device*					GetParentDevice()			const { return Descriptor.GetParentDevice(); }
 	inline FD3D12Device*					GetParentDevice_Unsafe()	const { return Descriptor.GetParentDevice_Unsafe(); }
@@ -794,7 +827,7 @@ public:
 };
 
 /** Shader resource view class. */
-class FD3D12ShaderResourceView : public FD3D12BaseShaderResourceView, public FRHIShaderResourceView, public FD3D12View<D3D12_SHADER_RESOURCE_VIEW_DESC>, public FD3D12LinkedAdapterObject<FD3D12ShaderResourceView>
+class FD3D12ShaderResourceView : public FRHIShaderResourceView, public FD3D12View<D3D12_SHADER_RESOURCE_VIEW_DESC>, public FD3D12LinkedAdapterObject<FD3D12ShaderResourceView>
 {
 	bool bContainsDepthPlane;
 	bool bContainsStencilPlane;
@@ -811,27 +844,25 @@ public:
 
 	// Used for all other SRV resource types. Initialization is immediate on the calling thread.
 	// Should not be used for dynamic resources which can be renamed.
-	FD3D12ShaderResourceView(FD3D12Device* InParent, D3D12_SHADER_RESOURCE_VIEW_DESC& InDesc, FD3D12ResourceLocation& InResourceLocation, uint32 InStride = -1, uint32 InStartOffsetBytes = 0, bool InSkipFastClearFinalize = false)
+	FD3D12ShaderResourceView(FD3D12Device* InParent, D3D12_SHADER_RESOURCE_VIEW_DESC& InDesc, FD3D12BaseShaderResource* InBaseShaderResource, uint32 InStride = -1, uint32 InStartOffsetBytes = 0, bool InSkipFastClearFinalize = false)
 		: FD3D12ShaderResourceView(InParent)
 	{
-		Initialize(InDesc, InResourceLocation, InStride, InStartOffsetBytes, InSkipFastClearFinalize);
+		Initialize(InDesc, InBaseShaderResource, InBaseShaderResource->ResourceLocation, InStride, InStartOffsetBytes, InSkipFastClearFinalize);
 	}
 
 	template<typename TextureType>
 	FD3D12ShaderResourceView(TextureType* InTexture, D3D12_SHADER_RESOURCE_VIEW_DESC& InDesc)
-		: FD3D12ShaderResourceView(InTexture->GetParentDevice(), InDesc, InTexture->ResourceLocation)
+		: FD3D12ShaderResourceView(InTexture->GetParentDevice(), InDesc, InTexture)
 	{
-		InTexture->AddView(this);
 	}
 
 	~FD3D12ShaderResourceView()
 	{
 		// Non-resident mesh LODs can lead to null SRVs being created
 		check(bInitialized || !GetParentDevice_Unsafe());
-		FD3D12BaseShaderResourceView::Remove();
 	}
 
-	void Initialize(D3D12_SHADER_RESOURCE_VIEW_DESC& InDesc, FD3D12ResourceLocation& InResourceLocation, uint32 InStride, uint32 InStartOffsetBytes, bool InSkipFastClearFinalize = false)
+	void Initialize(D3D12_SHADER_RESOURCE_VIEW_DESC& InDesc, FD3D12BaseShaderResource* InBaseShaderResource, FD3D12ResourceLocation& InResourceLocation, uint32 InStride, uint32 InStartOffsetBytes, bool InSkipFastClearFinalize)
 	{
 		Stride = InStride;
 		StartOffsetBytes = InStartOffsetBytes;
@@ -854,10 +885,15 @@ public:
 		}
 #endif
 
-		CreateView(InDesc, InResourceLocation);
+		CreateView(InDesc, InBaseShaderResource, InResourceLocation);
+	}
+
+	void Initialize(D3D12_SHADER_RESOURCE_VIEW_DESC& InDesc, FD3D12BaseShaderResource* InBaseShaderResource, uint32 InStride, uint32 InStartOffsetBytes, bool InSkipFastClearFinalize = false)
+	{
+		Initialize(InDesc, InBaseShaderResource, InBaseShaderResource->ResourceLocation, InStride, InStartOffsetBytes, InSkipFastClearFinalize);
 	}
 
-	void Initialize(FD3D12Device* InParent, D3D12_SHADER_RESOURCE_VIEW_DESC& InDesc, FD3D12ResourceLocation& InResourceLocation, uint32 InStride, bool InSkipFastClearFinalize = false)
+	void Initialize(FD3D12Device* InParent, D3D12_SHADER_RESOURCE_VIEW_DESC& InDesc, FD3D12BaseShaderResource* InBaseShaderResource, uint32 InStride, bool InSkipFastClearFinalize = false)
 	{
 		if (!this->GetParentDevice_Unsafe())
 		{
@@ -866,19 +902,22 @@ public:
 			this->SetParentDevice(InParent);
 		}
 		check(GetParentDevice() == InParent);
-		Initialize(InDesc, InResourceLocation, InStride, InSkipFastClearFinalize);
+		Initialize(InDesc, InBaseShaderResource, InBaseShaderResource->ResourceLocation, InStride, 0, InSkipFastClearFinalize);
 	}
 
-	void Rename(FD3D12ResourceLocation& InResourceLocation)
+	void Rename(FD3D12BaseShaderResource* InBaseShaderResource)
 	{
+		check(InBaseShaderResource == BaseShaderResource);
+		check(&InBaseShaderResource->ResourceLocation == ResourceLocation);
+
 		// Update the first element index, then reinitialize the SRV
 		if (Desc.ViewDimension == D3D12_SRV_DIMENSION_BUFFER)
 		{
 			uint32 StartElement = StartOffsetBytes / Stride;
-			Desc.Buffer.FirstElement = InResourceLocation.GetOffsetFromBaseOfResource() / Stride + StartElement;
+			Desc.Buffer.FirstElement = ResourceLocation->GetOffsetFromBaseOfResource() / Stride + StartElement;
 		}
 
-		Initialize(Desc, InResourceLocation, Stride, StartOffsetBytes);
+		Initialize(Desc, BaseShaderResource, Stride, StartOffsetBytes);
 	}
 
 	void Rename(float ResourceMinLODClamp)
@@ -889,12 +928,12 @@ public:
 
 		// Update the LODClamp, the reinitialize the SRV
 		Desc.Texture2D.ResourceMinLODClamp = ResourceMinLODClamp;
-		CreateView(Desc, *ResourceLocation);
+		CreateView(Desc, BaseShaderResource, *ResourceLocation);
 	}
 
-	virtual void Recreate(FD3D12ResourceLocation& InResourceLocation) override
+	virtual void RecreateView() override
 	{
-		Rename(InResourceLocation);
+		Rename(BaseShaderResource);
 	}
 
 	FORCEINLINE bool IsDepthStencilResource()	const { return bContainsDepthPlane || bContainsStencilPlane; }
@@ -912,17 +951,25 @@ public:
 		, ViewLocation(InParent)
 	{}
 
-	FD3D12ShaderResourceViewWithLocation(FD3D12Device* InParent, D3D12_SHADER_RESOURCE_VIEW_DESC& InDesc, FD3D12ResourceLocation& InResourceLocation, uint32 InStride = -1, bool InSkipFastClearFinalize = false)
+	FD3D12ShaderResourceViewWithLocation(FD3D12Device* InParent, D3D12_SHADER_RESOURCE_VIEW_DESC& InDesc, FD3D12BaseShaderResource* InBaseShaderResource, uint32 InStride = -1, bool InSkipFastClearFinalize = false)
 		: FD3D12ShaderResourceView(InParent)
 		, ViewLocation(InParent)
 	{
-		Initialize(InDesc, ViewLocation, InStride, 0, InSkipFastClearFinalize);
+		// Should not be pool allocated
+		check(InBaseShaderResource->ResourceLocation.GetAllocatorType() != FD3D12ResourceLocation::EAllocatorType::AT_Pool);
+		Initialize(InDesc, InBaseShaderResource, ViewLocation, InStride, 0, InSkipFastClearFinalize);
+	}
+
+	virtual void RecreateView() override
+	{
+		// Currently unsupported because the view location resource would also need to be recreated then
+		check(false);
 	}
 
 	FD3D12ResourceLocation ViewLocation;
 };
 
-class FD3D12UnorderedAccessView : public FD3D12BaseShaderResourceView, public FRHIUnorderedAccessView, public FD3D12View < D3D12_UNORDERED_ACCESS_VIEW_DESC >, public FD3D12LinkedAdapterObject<FD3D12UnorderedAccessView>
+class FD3D12UnorderedAccessView : public FRHIUnorderedAccessView, public FD3D12View < D3D12_UNORDERED_ACCESS_VIEW_DESC >, public FD3D12LinkedAdapterObject<FD3D12UnorderedAccessView>
 {
 public:
 	TRefCountPtr<FD3D12Resource> CounterResource;
@@ -933,31 +980,24 @@ public:
 	{
 	}
 
-	FD3D12UnorderedAccessView(FD3D12Device* InParent, D3D12_UNORDERED_ACCESS_VIEW_DESC& InDesc, FD3D12ResourceLocation& InResourceLocation, FD3D12Resource* InCounterResource = nullptr)
+	FD3D12UnorderedAccessView(FD3D12Device* InParent, D3D12_UNORDERED_ACCESS_VIEW_DESC& InDesc, FD3D12BaseShaderResource* InBaseShaderResource, FD3D12Resource* InCounterResource = nullptr)
 		: FD3D12View(InParent, ViewSubresourceSubsetFlags_None)
 		, CounterResource(InCounterResource)
 		, CounterResourceInitialized(false)
 	{
-		CreateViewWithCounter(InDesc, InResourceLocation, InCounterResource);
+		CreateViewWithCounter(InDesc, InBaseShaderResource, InBaseShaderResource->ResourceLocation, InCounterResource);
 	}
 
-	~FD3D12UnorderedAccessView()
+	void Initialize(D3D12_UNORDERED_ACCESS_VIEW_DESC& InDesc, FD3D12BaseShaderResource* InBaseShaderResource, FD3D12ResourceLocation& InResourceLocation)
 	{
-		FD3D12BaseShaderResourceView::Remove();
+		CreateViewWithCounter(InDesc, InBaseShaderResource, InResourceLocation, nullptr);
 	}
 
-	void Initialize(D3D12_UNORDERED_ACCESS_VIEW_DESC& InDesc, FD3D12ResourceLocation& InResourceLocation)
+	virtual void RecreateView() override
 	{
-		CreateViewWithCounter(InDesc, InResourceLocation, nullptr);
-	}
-
-	virtual void Recreate(FD3D12ResourceLocation& InResourceLocation) override
-	{
-		check(ResourceLocation == &InResourceLocation);
 		check(CounterResource == nullptr);
-		check(InResourceLocation.GetOffsetFromBaseOfResource() == 0);
-
-		Initialize(Desc, InResourceLocation);
+		check(ResourceLocation->GetOffsetFromBaseOfResource() == 0);
+		Initialize(Desc, BaseShaderResource, *ResourceLocation);
 	}
 };
 
@@ -970,6 +1010,12 @@ public:
 	{
 	}
 
+	virtual void RecreateView() override
+	{
+		// Currently unsupported because the view location resource would also need to be recreated then
+		check(false);
+	}
+
 	FD3D12ResourceLocation ViewLocation;
 };
 
@@ -1025,30 +1071,29 @@ public:
 };
 #endif
 
-class FD3D12RenderTargetView : public FD3D12BaseShaderResourceView, public FD3D12View<D3D12_RENDER_TARGET_VIEW_DESC>, public FRHIResource, public FD3D12LinkedAdapterObject<FD3D12RenderTargetView>
+class FD3D12RenderTargetView : public FD3D12View<D3D12_RENDER_TARGET_VIEW_DESC>, public FRHIResource, public FD3D12LinkedAdapterObject<FD3D12RenderTargetView>
 {
 public:
-	FD3D12RenderTargetView(FD3D12Device* InParent, const D3D12_RENDER_TARGET_VIEW_DESC& InRTVDesc, FD3D12ResourceLocation& InResourceLocation)
+	FD3D12RenderTargetView(FD3D12Device* InParent, const D3D12_RENDER_TARGET_VIEW_DESC& InRTVDesc, FD3D12BaseShaderResource* InBaseShaderResource)
 		: FD3D12View(InParent, ViewSubresourceSubsetFlags_None)
 	{
-		CreateView(InRTVDesc, InResourceLocation);
+		CreateView(InRTVDesc, InBaseShaderResource, InBaseShaderResource->ResourceLocation);
 	}
-	
-	~FD3D12RenderTargetView()
+
+	FD3D12RenderTargetView(FD3D12Device* InParent, const D3D12_RENDER_TARGET_VIEW_DESC& InRTVDesc, FD3D12BaseShaderResource* InBaseShaderResource, FD3D12ResourceLocation& InResourceLocation)
+		: FD3D12View(InParent, ViewSubresourceSubsetFlags_None)
 	{
-		FD3D12BaseShaderResourceView::Remove();
+		CreateView(InRTVDesc, InBaseShaderResource, InResourceLocation);
 	}
 
-	virtual void Recreate(FD3D12ResourceLocation& InResourceLocation) override
+	virtual void RecreateView() override
 	{
-		check(ResourceLocation == &InResourceLocation);
-		check(InResourceLocation.GetOffsetFromBaseOfResource() == 0);
-
-		CreateView(Desc, InResourceLocation);
+		check(ResourceLocation->GetOffsetFromBaseOfResource() == 0);
+		CreateView(Desc, BaseShaderResource, BaseShaderResource->ResourceLocation);
 	}
 };
 
-class FD3D12DepthStencilView : public FD3D12BaseShaderResourceView, public FD3D12View<D3D12_DEPTH_STENCIL_VIEW_DESC>, public FRHIResource, public FD3D12LinkedAdapterObject<FD3D12DepthStencilView>
+class FD3D12DepthStencilView : public FD3D12View<D3D12_DEPTH_STENCIL_VIEW_DESC>, public FRHIResource, public FD3D12LinkedAdapterObject<FD3D12DepthStencilView>
 {
 	const bool bHasDepth;
 	const bool bHasStencil;
@@ -1056,12 +1101,12 @@ class FD3D12DepthStencilView : public FD3D12BaseShaderResourceView, public FD3D1
 	CViewSubresourceSubset StencilOnlyViewSubresourceSubset;
 
 public:
-	FD3D12DepthStencilView(FD3D12Device* InParent, const D3D12_DEPTH_STENCIL_VIEW_DESC& InDSVDesc, FD3D12ResourceLocation& InResourceLocation, bool InHasStencil)
+	FD3D12DepthStencilView(FD3D12Device* InParent, const D3D12_DEPTH_STENCIL_VIEW_DESC& InDSVDesc, FD3D12BaseShaderResource* InBaseShaderResource, bool InHasStencil)
 		: FD3D12View(InParent, ViewSubresourceSubsetFlags_DepthAndStencilDsv)
 		, bHasDepth(true)				// Assume all DSVs have depth bits in their format
 		, bHasStencil(InHasStencil)		// Only some DSVs have stencil bits in their format
 	{
-		CreateView(InDSVDesc, InResourceLocation);
+		CreateView(InDSVDesc, InBaseShaderResource, InBaseShaderResource->ResourceLocation);
 
 		check(Resource);
 
@@ -1085,11 +1130,6 @@ public:
 		}
 	}
 
-	~FD3D12DepthStencilView()
-	{
-		FD3D12BaseShaderResourceView::Remove();
-	}
-
 	bool HasDepth() const
 	{
 		return bHasDepth;
@@ -1112,12 +1152,10 @@ public:
 		return StencilOnlyViewSubresourceSubset;
 	}
 
-	virtual void Recreate(FD3D12ResourceLocation& InResourceLocation) override
+	virtual void RecreateView() override
 	{
-		check(ResourceLocation == &InResourceLocation);
-		check(InResourceLocation.GetOffsetFromBaseOfResource() == 0);
-
-		CreateView(Desc, InResourceLocation);
+		check(ResourceLocation->GetOffsetFromBaseOfResource() == 0);
+		CreateView(Desc, BaseShaderResource, BaseShaderResource->ResourceLocation);
 	}
 };
 
diff --git a/Engine/Source/Runtime/D3D12RHI/Private/D3D12Viewport.cpp b/Engine/Source/Runtime/D3D12RHI/Private/D3D12Viewport.cpp
--- a/Engine/Source/Runtime/D3D12RHI/Private/D3D12Viewport.cpp
+++ b/Engine/Source/Runtime/D3D12RHI/Private/D3D12Viewport.cpp
@@ -304,8 +304,8 @@ FD3D12Texture2D* GetSwapChainSurface(FD3D12Device* Parent, EPixelFormat PixelFor
 			RTVDescRight.Texture2DArray.FirstArraySlice = 1;
 			RTVDescRight.Texture2DArray.ArraySize = 1;
 
-			BackBufferRenderTargetView = new FD3D12RenderTargetView(Device, RTVDescLeft, NewTexture->ResourceLocation);
-			BackBufferRenderTargetViewRight = new FD3D12RenderTargetView(Device, RTVDescRight, NewTexture->ResourceLocation);
+			BackBufferRenderTargetView = new FD3D12RenderTargetView(Device, RTVDescLeft, NewTexture);
+			BackBufferRenderTargetViewRight = new FD3D12RenderTargetView(Device, RTVDescRight, NewTexture);
 
 			NewTexture->SetNumRenderTargetViews(2);
 			NewTexture->SetRenderTargetViewIndex(BackBufferRenderTargetView, 0);
@@ -319,7 +319,7 @@ FD3D12Texture2D* GetSwapChainSurface(FD3D12Device* Parent, EPixelFormat PixelFor
 			RTVDesc.ViewDimension = D3D12_RTV_DIMENSION_TEXTURE2D;
 			RTVDesc.Texture2D.MipSlice = 0;
 
-			BackBufferRenderTargetView = new FD3D12RenderTargetView(Device, RTVDesc, NewTexture->ResourceLocation);
+			BackBufferRenderTargetView = new FD3D12RenderTargetView(Device, RTVDesc, NewTexture);
 			NewTexture->SetRenderTargetView(BackBufferRenderTargetView);
 		}
 
@@ -331,7 +331,7 @@ FD3D12Texture2D* GetSwapChainSurface(FD3D12Device* Parent, EPixelFormat PixelFor
 		SRVDesc.Texture2D.MostDetailedMip = 0;
 		SRVDesc.Texture2D.MipLevels = 1;
 
-		FD3D12ShaderResourceView* WrappedShaderResourceView = new FD3D12ShaderResourceView(Device, SRVDesc, NewTexture->ResourceLocation);
+		FD3D12ShaderResourceView* WrappedShaderResourceView = new FD3D12ShaderResourceView(Device, SRVDesc, NewTexture);
 		NewTexture->SetShaderResourceView(WrappedShaderResourceView);
 
 		if (Device->GetGPUIndex() == Parent->GetGPUIndex())
diff --git a/Engine/Source/Runtime/D3D12RHI/Public/D3D12Resources.h b/Engine/Source/Runtime/D3D12RHI/Public/D3D12Resources.h
--- a/Engine/Source/Runtime/D3D12RHI/Public/D3D12Resources.h
+++ b/Engine/Source/Runtime/D3D12RHI/Public/D3D12Resources.h
@@ -753,75 +753,64 @@ struct FD3D12LockedResource : public FD3D12DeviceChild
 	uint32 bHasNeverBeenLocked : 1;
 };
 
-class FD3D12BaseShaderResourceView
+/** Resource which might needs to be notified about changes on dependent resources (Views, RTGeometryObject, Cached binding tables) */
+class FD3D12ShaderResourceRenameListener
 {
 protected:
-	virtual ~FD3D12BaseShaderResourceView() { check(BaseShaderResource == nullptr); }
-
-	void Remove();
-	virtual void Recreate(FD3D12ResourceLocation& InResourceLocation) = 0;
 
 	friend class FD3D12BaseShaderResource;
-	FD3D12BaseShaderResource* BaseShaderResource = nullptr;
+	virtual void ResourceRenamed(FD3D12BaseShaderResource* InRenamedResource, FD3D12ResourceLocation* InNewResourceLocation) = 0;
 };
 
-/** The base class of resources that may be bound as shader resources. */
+
+/** The base class of resources that may be bound as shader resources (texture or buffer). */
 class FD3D12BaseShaderResource : public FD3D12DeviceChild, public IRefCountedObject
 {
 protected:
-	FCriticalSection ViewsCS;
-	TArray<class FD3D12BaseShaderResourceView*> Views;
+	FCriticalSection RenameListenersCS;
+	TArray<FD3D12ShaderResourceRenameListener*> RenameListeners;
 
 public:
 	FD3D12Resource* GetResource() const { return ResourceLocation.GetResource(); }
 
-	void AddView(FD3D12BaseShaderResourceView* InView)
+	void AddRenameListener(FD3D12ShaderResourceRenameListener* InRenameListener)
 	{
-		FScopeLock Lock(&ViewsCS);
-		check(InView->BaseShaderResource == nullptr);
-		InView->BaseShaderResource = this;
-		Views.Add(InView);
+		FScopeLock Lock(&RenameListenersCS);
+		RenameListeners.Add(InRenameListener);
 	}
 
-	void RemoveView(FD3D12BaseShaderResourceView* InSRV)
+	void RemoveRenameListener(FD3D12ShaderResourceRenameListener* InRenameListener)
 	{
-		FScopeLock Lock(&ViewsCS);
-		check(InSRV->BaseShaderResource == this);
-		InSRV->BaseShaderResource = nullptr;
-		uint32 Removed = Views.Remove(InSRV);
-		check(Removed == 1);
-	}
+		FScopeLock Lock(&RenameListenersCS);
+		uint32 Removed = RenameListeners.Remove(InRenameListener);
 
-	void RemoveAllViews()
-	{
-		FScopeLock Lock(&ViewsCS);
-		for (FD3D12BaseShaderResourceView* View : Views)
-		{
-			if (View != nullptr)
-			{
-				View->BaseShaderResource = nullptr;
-			}
-		}
-		Views.Reset();
+		checkf(Removed == 1, TEXT("Should have exactly one registered listener during remove (same listener shouldn't registered twice and we shouldn't call this if not registered"));
 	}
 
 	void Swap(FD3D12BaseShaderResource& Other)
 	{
+		// assume RHI thread when swapping listeners and resources
+		check(IsInRHIThread());
+
 		::Swap(Parent, Other.Parent);
 		ResourceLocation.Swap(Other.ResourceLocation);
 		::Swap(BufferAlignment, Other.BufferAlignment);
-		::Swap(Views, Other.Views);
+		::Swap(RenameListeners, Other.RenameListeners);
 	}
 
-	void RecreateViews(FD3D12ResourceLocation& InResourceLocation)
+	void RemoveAllRenameListeners()
 	{
-		FScopeLock Lock(&ViewsCS);
-		for (FD3D12BaseShaderResourceView* View : Views)
+		FScopeLock Lock(&RenameListenersCS);
+		ResourceRenamed(nullptr);
+		RenameListeners.Reset();
+	}
+
+	void ResourceRenamed(FD3D12ResourceLocation* InNewResourceLocation)
+	{
+		FScopeLock Lock(&RenameListenersCS);
+		for (FD3D12ShaderResourceRenameListener* RenameListener : RenameListeners)
 		{
-			if (View != nullptr)
-			{
-				View->Recreate(InResourceLocation);
-			}
+			RenameListener->ResourceRenamed(this, InNewResourceLocation);
 		}
 	}
 
@@ -838,22 +827,10 @@ public:
 
 	~FD3D12BaseShaderResource()
 	{
-		for (FD3D12BaseShaderResourceView* View : Views)
-		{
-			check(View->BaseShaderResource == this);
-			View->BaseShaderResource = nullptr;
-		}
+		RemoveAllRenameListeners();
 	}
 };
 
-inline void FD3D12BaseShaderResourceView::Remove()
-{
-	if (BaseShaderResource)
-	{
-		BaseShaderResource->RemoveView(this);
-	}
-}
-
 extern void UpdateBufferStats(FName Name, int64 RequestedSize);
 
 inline FName GetBufferStats(uint32 Usage)
