- Enable buffer pool defrag by default when pool allocator is enabled
- Add option to disable defrag per pool allocator (disabled by default for non default memory and BVH buffers)
- Fix used memory tracking size in stats

diff --git a/Engine/Source/Runtime/D3D12RHI/Private/D3D12Allocation.cpp b/Engine/Source/Runtime/D3D12RHI/Private/D3D12Allocation.cpp
--- a/Engine/Source/Runtime/D3D12RHI/Private/D3D12Allocation.cpp
+++ b/Engine/Source/Runtime/D3D12RHI/Private/D3D12Allocation.cpp
@@ -95,9 +95,7 @@ static FAutoConsoleVariableRef CVarD3D12D3D12VRAMTexturePoolDefragMaxCopySizePer
 	TEXT("Max amount of data to copy during defragmentation in a single frame (default 32MB)"),
 	ECVF_ReadOnly);
 
-// NOTE: buffer pool defrag is disabled by default because FD3D12RayTracingGeometry is currently not notified
-//       yet of defragged index and vertex buffers used by the the BVH data. Without RT it works fine.
-static int32 GD3D12VRAMBufferPoolDefrag = 0;
+static int32 GD3D12VRAMBufferPoolDefrag = 1;
 static FAutoConsoleVariableRef CVarD3D12D3D12VRAMBufferPoolDefrag(
 	TEXT("d3d12.VRAMBufferPoolDefrag"),
 	GD3D12VRAMBufferPoolDefrag,
@@ -1266,7 +1264,18 @@ FD3D12BufferPool* FD3D12DefaultBufferAllocator::CreateBufferPool(D3D12_HEAP_TYPE
 	uint64 PoolAlignment = (AllocationStrategy == EResourceAllocationStrategy::kPlacedResource) ? MIN_PLACED_RESOURCE_SIZE : 256;
 	uint64 MaxAllocationSize = InHeapType == D3D12_HEAP_TYPE_READBACK ? READBACK_BUFFER_POOL_MAX_ALLOC_SIZE : GD3D12PoolAllocatorBufferVRAMMaxAllocationSize;
 
-	FD3D12BufferPool* NewPool = new FD3D12PoolAllocator(Device, GetVisibilityMask(), InitConfig, Name, AllocationStrategy, PoolSize, PoolAlignment, MaxAllocationSize);
+	// Disable defrag if not Default memory
+	bool bDefragEnabled = (InitConfig.HeapType == D3D12_HEAP_TYPE_DEFAULT);
+
+#if D3D12_RHI_RAYTRACING
+	// disable defrag on the RT Acceleration pool - get GPU crashes after moving the data
+	if (InitConfig.InitialResourceState == D3D12_RESOURCE_STATE_RAYTRACING_ACCELERATION_STRUCTURE)
+	{
+		bDefragEnabled = false;
+	}
+#endif // D3D12_RHI_RAYTRACING
+
+	FD3D12BufferPool* NewPool = new FD3D12PoolAllocator(Device, GetVisibilityMask(), InitConfig, Name, AllocationStrategy, PoolSize, PoolAlignment, MaxAllocationSize, bDefragEnabled);
 
 #else // USE_POOL_ALLOCATOR
 
@@ -1480,7 +1489,8 @@ FD3D12TextureAllocatorPool::FD3D12TextureAllocatorPool(FD3D12Device* Device, FRH
 		const FString Name(L"D3D12 ReadOnly4K Texture Pool Allocator");
 		uint64 PoolSize = 4 * 1024 * 1024;
 		uint64 MaxAllocationSize = PoolSize;
-		FD3D12PoolAllocator* ReadOnly4KPool = new FD3D12PoolAllocator(Device, GetVisibilityMask(), InitConfig, Name, AllocationStrategy, PoolSize, D3D12_SMALL_RESOURCE_PLACEMENT_ALIGNMENT, MaxAllocationSize);
+		bool bDefragEnabled = false; // Disable defrag on 4K pool because it shouldn't really fragment - all allocations are 4K or multiple of 4K and pretty small
+		FD3D12PoolAllocator* ReadOnly4KPool = new FD3D12PoolAllocator(Device, GetVisibilityMask(), InitConfig, Name, AllocationStrategy, PoolSize, D3D12_SMALL_RESOURCE_PLACEMENT_ALIGNMENT, MaxAllocationSize, bDefragEnabled);
 		ReadOnly4KPool->Initialize();
 
 		PoolAllocators[(int)EPoolType::ReadOnly4K] = ReadOnly4KPool;
@@ -1493,7 +1503,8 @@ FD3D12TextureAllocatorPool::FD3D12TextureAllocatorPool(FD3D12Device* Device, FRH
 		const FString Name(L"D3D12 ReadOnly Texture Pool Allocator");
 		uint64 PoolSize = GD3D12PoolAllocatorReadOnlyTextureVRAMPoolSize;
 		uint64 MaxAllocationSize = GD3D12PoolAllocatorReadOnlyTextureVRAMMaxAllocationSize;
-		FD3D12PoolAllocator* ReadOnlyPool = new FD3D12PoolAllocator(Device, GetVisibilityMask(), InitConfig, Name, AllocationStrategy, PoolSize, D3D12_DEFAULT_RESOURCE_PLACEMENT_ALIGNMENT, MaxAllocationSize);
+		bool bDefragEnabled = true;
+		FD3D12PoolAllocator* ReadOnlyPool = new FD3D12PoolAllocator(Device, GetVisibilityMask(), InitConfig, Name, AllocationStrategy, PoolSize, D3D12_DEFAULT_RESOURCE_PLACEMENT_ALIGNMENT, MaxAllocationSize, bDefragEnabled);
 		ReadOnlyPool->Initialize();
 
 		PoolAllocators[(int)EPoolType::ReadOnly] = ReadOnlyPool;
@@ -1506,7 +1517,8 @@ FD3D12TextureAllocatorPool::FD3D12TextureAllocatorPool(FD3D12Device* Device, FRH
 		const FString Name(L"D3D12 RT Texture Pool Allocator");
 		uint64 PoolSize = GD3D12PoolAllocatorRTUAVTextureVRAMPoolSize;
 		uint64 MaxAllocationSize = GD3D12PoolAllocatorRTUAVTextureVRAMMaxAllocationSize;
-		FD3D12PoolAllocator* RTPool = new FD3D12PoolAllocator(Device, GetVisibilityMask(), InitConfig, Name, AllocationStrategy, PoolSize, D3D12_DEFAULT_RESOURCE_PLACEMENT_ALIGNMENT, MaxAllocationSize);
+		bool bDefragEnabled = true;
+		FD3D12PoolAllocator* RTPool = new FD3D12PoolAllocator(Device, GetVisibilityMask(), InitConfig, Name, AllocationStrategy, PoolSize, D3D12_DEFAULT_RESOURCE_PLACEMENT_ALIGNMENT, MaxAllocationSize, bDefragEnabled);
 		RTPool->Initialize();
 
 		PoolAllocators[(int)EPoolType::RenderTarget] = RTPool;
@@ -1519,7 +1531,8 @@ FD3D12TextureAllocatorPool::FD3D12TextureAllocatorPool(FD3D12Device* Device, FRH
 		const FString Name(L"D3D12 UAV Texture Pool Allocator");
 		uint64 PoolSize = GD3D12PoolAllocatorRTUAVTextureVRAMPoolSize;
 		uint64 MaxAllocationSize = GD3D12PoolAllocatorRTUAVTextureVRAMMaxAllocationSize;
-		FD3D12PoolAllocator* UAVPool = new FD3D12PoolAllocator(Device, GetVisibilityMask(), InitConfig, Name, AllocationStrategy, PoolSize, D3D12_DEFAULT_RESOURCE_PLACEMENT_ALIGNMENT, MaxAllocationSize);
+		bool bDefragEnabled = true;
+		FD3D12PoolAllocator* UAVPool = new FD3D12PoolAllocator(Device, GetVisibilityMask(), InitConfig, Name, AllocationStrategy, PoolSize, D3D12_DEFAULT_RESOURCE_PLACEMENT_ALIGNMENT, MaxAllocationSize, bDefragEnabled);
 		UAVPool->Initialize();
 
 		PoolAllocators[(int)EPoolType::UAV] = UAVPool;
@@ -1570,8 +1583,7 @@ void FD3D12TextureAllocatorPool::BeginFrame()
 	{
 		uint32 MaxCopySize = GD3D12VRAMTexturePoolDefragMaxCopySizePerFrame;
 		uint32 CopySize = 0;
-		// Don't defrag the 4K pool - not worth it because it shouldn't really fragment
-		for (uint32 PoolIndex = 1; PoolIndex < (uint32)EPoolType::Count; ++PoolIndex)
+		for (uint32 PoolIndex = 0; PoolIndex < (uint32)EPoolType::Count; ++PoolIndex)
 		{
 			PoolAllocators[PoolIndex]->Defrag(MaxCopySize, CopySize);
 		}
diff --git a/Engine/Source/Runtime/D3D12RHI/Private/D3D12PoolAllocator.cpp b/Engine/Source/Runtime/D3D12RHI/Private/D3D12PoolAllocator.cpp
--- a/Engine/Source/Runtime/D3D12RHI/Private/D3D12PoolAllocator.cpp
+++ b/Engine/Source/Runtime/D3D12RHI/Private/D3D12PoolAllocator.cpp
@@ -173,8 +173,8 @@ EResourceAllocationStrategy FD3D12PoolAllocator::GetResourceAllocationStrategy(D
 
 
 FD3D12PoolAllocator::FD3D12PoolAllocator(FD3D12Device* ParentDevice, FRHIGPUMask VisibleNodes, const FD3D12ResourceInitConfig& InInitConfig, const FString& InName,
-	EResourceAllocationStrategy InAllocationStrategy, uint64 InPoolSize, uint32 InPoolAlignment, uint32 InMaxAllocationSize) :
-	FRHIPoolAllocator(InPoolSize, InPoolAlignment, InMaxAllocationSize), 
+	EResourceAllocationStrategy InAllocationStrategy, uint64 InPoolSize, uint32 InPoolAlignment, uint32 InMaxAllocationSize, bool bInDefragEnabled) :
+	FRHIPoolAllocator(InPoolSize, InPoolAlignment, InMaxAllocationSize, bInDefragEnabled), 
 	FD3D12DeviceChild(ParentDevice), 
 	FD3D12MultiNodeGPUObject(ParentDevice->GetGPUMask(), VisibleNodes), 
 	InitConfig(InInitConfig), 
diff --git a/Engine/Source/Runtime/D3D12RHI/Private/D3D12PoolAllocator.h b/Engine/Source/Runtime/D3D12RHI/Private/D3D12PoolAllocator.h
--- a/Engine/Source/Runtime/D3D12RHI/Private/D3D12PoolAllocator.h
+++ b/Engine/Source/Runtime/D3D12RHI/Private/D3D12PoolAllocator.h
@@ -117,7 +117,7 @@ public:
 
 	// Constructor
 	FD3D12PoolAllocator(FD3D12Device* ParentDevice, FRHIGPUMask VisibleNodes, const FD3D12ResourceInitConfig& InInitConfig, const FString& InName,
-		EResourceAllocationStrategy InAllocationStrategy, uint64 InPoolSize, uint32 InPoolAlignment, uint32 InMaxAllocationSize);
+		EResourceAllocationStrategy InAllocationStrategy, uint64 InPoolSize, uint32 InPoolAlignment, uint32 InMaxAllocationSize, bool bInDefragEnabled);
 	~FD3D12PoolAllocator();
 	
 	// Function names currently have to match with FD3D12DefaultBufferPool until we can make full replacement of the allocator
diff --git a/Engine/Source/Runtime/RHICore/Private/RHIPoolAllocator.cpp b/Engine/Source/Runtime/RHICore/Private/RHIPoolAllocator.cpp
--- a/Engine/Source/Runtime/RHICore/Private/RHIPoolAllocator.cpp
+++ b/Engine/Source/Runtime/RHICore/Private/RHIPoolAllocator.cpp
@@ -505,10 +505,11 @@ void FRHIMemoryPool::Validate()
 //	Pool Allocator
 //-----------------------------------------------------------------------------
 
-FRHIPoolAllocator::FRHIPoolAllocator(uint64 InPoolSize, uint32 InPoolAlignment, uint32 InMaxAllocationSize)	: 
-	PoolSize(InPoolSize), 
-	PoolAlignment(InPoolAlignment), 
+FRHIPoolAllocator::FRHIPoolAllocator(uint64 InPoolSize, uint32 InPoolAlignment, uint32 InMaxAllocationSize, bool InDefragEnabled) :
+	PoolSize(InPoolSize),
+	PoolAlignment(InPoolAlignment),
 	MaxAllocationSize(InMaxAllocationSize),
+	bDefragEnabled(InDefragEnabled),
 	TotalAllocatedBlocks(0)
 {}
 
@@ -588,6 +589,12 @@ void FRHIPoolAllocator::DeallocateInternal(FRHIPoolAllocationData& AllocationDat
 
 void FRHIPoolAllocator::Defrag(uint32 InMaxCopySize, uint32& CurrentCopySize)
 {
+	// Don't do anything when defrag is disabled for this allocator
+	if (!bDefragEnabled)
+	{
+		return;
+	}
+
 	FScopeLock Lock(&CS);
 
 	TArray<FRHIMemoryPool*> SortedPools;
@@ -646,7 +653,7 @@ void FRHIPoolAllocator::UpdateMemoryStats(uint32& IOMemoryAllocated, uint32& IOM
 		if (Pool)
 		{
 			IOMemoryAllocated += Pool->GetPoolSize();
-			IOMemoryUsed += Pool->GetUsedSize();
+			IOMemoryUsed += Pool->GetUsedSize() - Pool->GetAlignmentWaste();
 			IOMemoryFree += Pool->GetFreeSize();
 			IOAlignmentWaste += Pool->GetAlignmentWaste();
 			IOAllocatedPageCount++;
diff --git a/Engine/Source/Runtime/RHICore/Public/RHIPoolAllocator.h b/Engine/Source/Runtime/RHICore/Public/RHIPoolAllocator.h
--- a/Engine/Source/Runtime/RHICore/Public/RHIPoolAllocator.h
+++ b/Engine/Source/Runtime/RHICore/Public/RHIPoolAllocator.h
@@ -187,7 +187,7 @@ class RHICORE_API FRHIPoolAllocator
 public:
 
 	// Constructor
-	FRHIPoolAllocator(uint64 InPoolSize, uint32 InPoolAlignment, uint32 InMaxAllocationSize);
+	FRHIPoolAllocator(uint64 InPoolSize, uint32 InPoolAlignment, uint32 InMaxAllocationSize, bool InDefragEnabled);
 	virtual ~FRHIPoolAllocator();
 
 	// Setup/Shutdown
@@ -217,6 +217,7 @@ protected:
 	const uint64 PoolSize;
 	const uint32 PoolAlignment;
 	const uint64 MaxAllocationSize;
+	const bool bDefragEnabled;
 
 	// Critical section to lock access to the pools
 	FCriticalSection CS;
