- Handle index/vertex buffer changes in RT Geometry objects and recreate the hitgroups on change

diff --git a/Engine/Source/Runtime/D3D12RHI/Private/D3D12RayTracing.cpp b/Engine/Source/Runtime/D3D12RHI/Private/D3D12RayTracing.cpp
--- a/Engine/Source/Runtime/D3D12RHI/Private/D3D12RayTracing.cpp
+++ b/Engine/Source/Runtime/D3D12RHI/Private/D3D12RayTracing.cpp
@@ -2379,6 +2379,7 @@ FD3D12RayTracingGeometry::FD3D12RayTracingGeometry(const FRayTracingGeometryInit
 
 	DebugName = Initializer.DebugName;
 	FMemory::Memzero(bHasPendingCompactionRequests);
+	FMemory::Memzero(bRegisteredAsRenameListener);
 
 	if(!FD3D12RayTracingGeometry::NullTransformBuffer.IsValid())
 	{
@@ -2484,6 +2485,7 @@ FD3D12RayTracingGeometry::FD3D12RayTracingGeometry(const FRayTracingGeometryInit
 
 FD3D12RayTracingGeometry::~FD3D12RayTracingGeometry()
 {
+	// Remove compaction request if still pending
 	for (uint32 GPUIndex = 0; GPUIndex < MAX_NUM_GPUS; ++GPUIndex)
 	{
 		if (bHasPendingCompactionRequests[GPUIndex])
@@ -2496,6 +2498,12 @@ FD3D12RayTracingGeometry::~FD3D12RayTracingGeometry()
 		}
 	}
 
+	// Unregister as dependent resource on vertex and index buffers
+	for (uint32 GPUIndex = 0; GPUIndex < MAX_NUM_GPUS; ++GPUIndex)
+	{
+		UnregisterAsRenameListener(GPUIndex);
+	}
+
 	for (TRefCountPtr<FD3D12Buffer>& Buffer : AccelerationStructureBuffers)
 	{
 		if (Buffer)
@@ -2518,6 +2526,69 @@ FD3D12RayTracingGeometry::~FD3D12RayTracingGeometry()
 	DEC_DWORD_STAT(STAT_D3D12RayTracingAllocatedBLAS);
 }
 
+void FD3D12RayTracingGeometry::RegisterAsRenameListener(uint32 InGPUIndex)
+{
+	check(!bRegisteredAsRenameListener[InGPUIndex]);
+
+	FD3D12Buffer* IndexBuffer = FD3D12DynamicRHI::ResourceCast(RHIIndexBuffer.GetReference(), InGPUIndex);
+	if (IndexBuffer)
+	{
+		IndexBuffer->AddRenameListener(this);
+	}
+
+	TArray<FD3D12Buffer*, TInlineAllocator<1>> UniqueVertexBuffers;
+	UniqueVertexBuffers.Reserve(Segments.Num());
+	for (const FRayTracingGeometrySegment& Segment : Segments)
+	{
+		FD3D12Buffer* VertexBuffer = FD3D12DynamicRHI::ResourceCast(Segment.VertexBuffer.GetReference(), InGPUIndex);
+		if (VertexBuffer && !UniqueVertexBuffers.Contains(VertexBuffer))
+		{
+			VertexBuffer->AddRenameListener(this);
+			UniqueVertexBuffers.Add(VertexBuffer);
+		}
+	}
+
+	bRegisteredAsRenameListener[InGPUIndex] = true;
+}
+
+void FD3D12RayTracingGeometry::UnregisterAsRenameListener(uint32 InGPUIndex)
+{
+	if (!bRegisteredAsRenameListener[InGPUIndex])
+	{
+		return;
+	}
+
+	FD3D12Buffer* IndexBuffer = FD3D12DynamicRHI::ResourceCast(RHIIndexBuffer.GetReference(), InGPUIndex);
+	if (IndexBuffer)
+	{
+		IndexBuffer->RemoveRenameListener(this);
+	}
+
+	TArray<FD3D12Buffer*, TInlineAllocator<1>> UniqueVertexBuffers;
+	UniqueVertexBuffers.Reserve(Segments.Num());
+	for (const FRayTracingGeometrySegment& Segment : Segments)
+	{
+		FD3D12Buffer* VertexBuffer = FD3D12DynamicRHI::ResourceCast(Segment.VertexBuffer.GetReference(), InGPUIndex);
+		if (VertexBuffer && !UniqueVertexBuffers.Contains(VertexBuffer))
+		{
+			VertexBuffer->RemoveRenameListener(this);
+			UniqueVertexBuffers.Add(VertexBuffer);
+		}
+	}
+
+	bRegisteredAsRenameListener[InGPUIndex] = false;
+}
+
+static constexpr uint32 IndicesPerPrimitive = 3; // Only triangle meshes are supported
+
+void FD3D12RayTracingGeometry::ResourceRenamed(FD3D12BaseShaderResource* InRenamedResource, FD3D12ResourceLocation* InNewResourceLocation)
+{
+	checkf(InNewResourceLocation, TEXT("Shouldn't release resources which are still referenced by the RayTracingGeometry"));
+
+	// Recreate the hit group parameters which cache the address to the index and vertex buffers directly
+	SetupHitGroupSystemParameters(InRenamedResource->GetParentDevice()->GetGPUIndex());	
+}
+
 void FD3D12RayTracingGeometry::TransitionBuffers(FD3D12CommandContext& CommandContext)
 {
 	// Transition vertex and index resources..
@@ -2541,6 +2612,30 @@ void FD3D12RayTracingGeometry::TransitionBuffers(FD3D12CommandContext& CommandCo
 	}
 }
 
+void FD3D12RayTracingGeometry::SetupHitGroupSystemParameters(uint32 InGPUIndex)
+{
+	TArray<FHitGroupSystemParameters>& HitGroupSystemParametersForThisGPU = HitGroupSystemParameters[InGPUIndex];
+	HitGroupSystemParametersForThisGPU.Reset(Segments.Num());
+
+	FD3D12Buffer* IndexBuffer = FD3D12DynamicRHI::ResourceCast(RHIIndexBuffer.GetReference(), InGPUIndex);
+	for (const FRayTracingGeometrySegment& Segment : Segments)
+	{
+		FD3D12Buffer* VertexBuffer = FD3D12DynamicRHI::ResourceCast(Segment.VertexBuffer.GetReference(), InGPUIndex);
+
+		FHitGroupSystemParameters SystemParameters = {};
+		SystemParameters.RootConstants.SetVertexAndIndexStride(Segment.VertexBufferStride, IndexStride);
+		SystemParameters.VertexBuffer = VertexBuffer->ResourceLocation.GetGPUVirtualAddress() + Segment.VertexBufferOffset;
+
+		if (GeometryType == D3D12_RAYTRACING_GEOMETRY_TYPE_TRIANGLES && IndexBuffer != nullptr)
+		{
+			SystemParameters.IndexBuffer = IndexBuffer->ResourceLocation.GetGPUVirtualAddress();
+			SystemParameters.RootConstants.IndexBufferOffsetInBytes = IndexOffsetInBytes + IndexStride * Segment.FirstPrimitive * IndicesPerPrimitive;
+		}
+
+		HitGroupSystemParametersForThisGPU.Add(SystemParameters);
+	}
+}
+
 static void CreateAccelerationStructureBuffers(
 	TRefCountPtr<FD3D12Buffer>& AccelerationStructureBuffer,
 	TRefCountPtr<FD3D12Buffer>&  ScratchBuffer,
@@ -2587,8 +2682,7 @@ void FD3D12RayTracingGeometry::BuildAccelerationStructure(FD3D12CommandContext&
 	IndexStride = RHIIndexBuffer ? RHIIndexBuffer->GetStride() : 0; // stride 0 means implicit triangle list for non-indexed geometry
 
 	const bool bIsUpdate = BuildMode == EAccelerationStructureBuildMode::Update;
-	static constexpr uint32 IndicesPerPrimitive = 3; // Only triangle meshes are supported
-
+		
 	// Array of geometry descriptions, one per segment (single-segment geometry is a common case).
 	TArray<D3D12_RAYTRACING_GEOMETRY_DESC, TInlineAllocator<1>> Descs;
 
@@ -2597,11 +2691,16 @@ void FD3D12RayTracingGeometry::BuildAccelerationStructure(FD3D12CommandContext&
 	FD3D12Buffer* IndexBuffer = CommandContext.RetrieveObject<FD3D12Buffer>(RHIIndexBuffer.GetReference());
 	FD3D12Buffer* NullTransformBufferD3D12 = CommandContext.RetrieveObject<FD3D12Buffer>(NullTransformBuffer.GetReference());
 
-	TArray<FHitGroupSystemParameters>& HitGroupSystemParametersForThisGPU = HitGroupSystemParameters[CommandContext.GetGPUIndex()];
-	HitGroupSystemParametersForThisGPU.Reset(Segments.Num());
+	// Recreate the hit group system parameters and use them during setup of the descs
+	SetupHitGroupSystemParameters(CommandContext.GetGPUIndex());
+	const TArray<FHitGroupSystemParameters>& HitGroupSystemParametersForThisGPU = HitGroupSystemParameters[CommandContext.GetGPUIndex()];
+	check(HitGroupSystemParametersForThisGPU.Num() == Segments.Num());
 
-	for (const FRayTracingGeometrySegment& Segment : Segments)
+	for (int32 SegmentIndex = 0; SegmentIndex < Segments.Num(); ++SegmentIndex)
 	{
+		const FRayTracingGeometrySegment& Segment = Segments[SegmentIndex];
+		const FHitGroupSystemParameters& SystemParameters = HitGroupSystemParametersForThisGPU[SegmentIndex];
+
 		D3D12_RAYTRACING_GEOMETRY_DESC Desc = {};
 
 		Desc.Flags = D3D12_RAYTRACING_GEOMETRY_FLAG_NONE;
@@ -2621,10 +2720,6 @@ void FD3D12RayTracingGeometry::BuildAccelerationStructure(FD3D12CommandContext&
 
 		FD3D12Buffer* VertexBuffer = CommandContext.RetrieveObject<FD3D12Buffer>(Segment.VertexBuffer.GetReference());
 		
-		FHitGroupSystemParameters SystemParameters = {};
-		SystemParameters.RootConstants.SetVertexAndIndexStride(Segment.VertexBufferStride, IndexStride);
-		SystemParameters.VertexBuffer = VertexBuffer->ResourceLocation.GetGPUVirtualAddress() + Segment.VertexBufferOffset;
-
 		// Conservative estimate of the maximum number of elements in this VB.
 		// Real used number will depend on the index buffer and is not available here right now.
 		// #dxr_todo: Add explicit vertex count to FRayTracingGeometrySegment.
@@ -2664,9 +2759,6 @@ void FD3D12RayTracingGeometry::BuildAccelerationStructure(FD3D12CommandContext&
 
 			if (IndexBuffer)
 			{
-				SystemParameters.IndexBuffer = IndexBuffer->ResourceLocation.GetGPUVirtualAddress();
-				SystemParameters.RootConstants.IndexBufferOffsetInBytes = IndexOffsetInBytes + IndexStride * Segment.FirstPrimitive * IndicesPerPrimitive;
-
 				Desc.Triangles.IndexFormat = IndexStride == 4 ? DXGI_FORMAT_R32_UINT : DXGI_FORMAT_R16_UINT;
 				Desc.Triangles.IndexCount = Segment.NumPrimitives * IndicesPerPrimitive;
 				Desc.Triangles.IndexBuffer = SystemParameters.IndexBuffer + SystemParameters.RootConstants.IndexBufferOffsetInBytes;
@@ -2708,13 +2800,14 @@ void FD3D12RayTracingGeometry::BuildAccelerationStructure(FD3D12CommandContext&
 			// #dxr_todo UE-72160: support various vertex buffer layouts (fetch/decode based on vertex stride and format)
 			checkf(Segment.VertexBufferElementType == VET_Float3 || Segment.VertexBufferElementType == VET_Float4, TEXT("Only VET_Float3 and Float4 are currently implemented and tested. Other formats will be supported in the future."));
 		}
-
-		HitGroupSystemParametersForThisGPU.Add(SystemParameters);
 	}
 
 	const uint32 GPUIndex = CommandContext.GetGPUIndex();
 	FD3D12Adapter* Adapter = CommandContext.GetParentAdapter();
 
+	// Register as rename listener to index/vertex buffers
+	RegisterAsRenameListener(GPUIndex);
+
 	ID3D12Device5* RayTracingDevice = CommandContext.GetParentDevice()->GetDevice5();
 
 	D3D12_RAYTRACING_ACCELERATION_STRUCTURE_BUILD_FLAGS LocalBuildFlags = D3D12_RAYTRACING_ACCELERATION_STRUCTURE_BUILD_FLAGS(BuildFlags);
@@ -3286,6 +3379,7 @@ void FD3D12CommandContext::RHIBuildAccelerationStructures(const TArrayView<const
 	for (const FAccelerationStructureBuildParams& P : Params)
 	{
 		FD3D12RayTracingGeometry* Geometry = FD3D12DynamicRHI::ResourceCast(P.Geometry.GetReference());
+		Geometry->UnregisterAsRenameListener(GetGPUIndex());
 
 		if (P.Segments.Num())
 		{
diff --git a/Engine/Source/Runtime/D3D12RHI/Private/D3D12RayTracing.h b/Engine/Source/Runtime/D3D12RHI/Private/D3D12RayTracing.h
--- a/Engine/Source/Runtime/D3D12RHI/Private/D3D12RayTracing.h
+++ b/Engine/Source/Runtime/D3D12RHI/Private/D3D12RayTracing.h
@@ -19,16 +19,23 @@ struct FHitGroupSystemParameters
 	FHitGroupSystemRootConstants RootConstants;
 };
 
-class FD3D12RayTracingGeometry : public FRHIRayTracingGeometry
+class FD3D12RayTracingGeometry : public FRHIRayTracingGeometry, public FD3D12ShaderResourceRenameListener
 {
 public:
 
 	FD3D12RayTracingGeometry(const FRayTracingGeometryInitializer& Initializer);
 	~FD3D12RayTracingGeometry();
 
+	void SetupHitGroupSystemParameters(uint32 InGPUIndex);
 	void TransitionBuffers(FD3D12CommandContext& CommandContext);
 	void BuildAccelerationStructure(FD3D12CommandContext& CommandContext, EAccelerationStructureBuildMode BuildMode);
 	void CompactAccelerationStructure(FD3D12CommandContext& CommandContext, uint32 InGPUIndex, uint64 InSizeAfterCompaction);
+	
+	// Implement FD3D12ShaderResourceRenameListener interface
+	virtual void ResourceRenamed(FD3D12BaseShaderResource* InRenamedResource, FD3D12ResourceLocation* InNewResourceLocation) override;
+
+	void RegisterAsRenameListener(uint32 InGPUIndex);
+	void UnregisterAsRenameListener(uint32 InGPUIndex);
 
 	bool bIsAccelerationStructureDirty[MAX_NUM_GPUS] = {};
 	void SetDirty(FRHIGPUMask GPUMask, bool bState)
@@ -58,6 +65,7 @@ public:
 	TRefCountPtr<FD3D12Buffer> AccelerationStructureBuffers[MAX_NUM_GPUS];
 	TRefCountPtr<FD3D12Buffer> ScratchBuffers[MAX_NUM_GPUS];
 
+	bool bRegisteredAsRenameListener[MAX_NUM_GPUS];
 	bool bHasPendingCompactionRequests[MAX_NUM_GPUS];
 
 	// Hit shader parameters per geometry segment
