Track buffer and texture pool stats on xbox as well

diff --git a/Engine/Source/Runtime/D3D12RHI/Private/D3D12Allocation.h b/Engine/Source/Runtime/D3D12RHI/Private/D3D12Allocation.h
--- a/Engine/Source/Runtime/D3D12RHI/Private/D3D12Allocation.h
+++ b/Engine/Source/Runtime/D3D12RHI/Private/D3D12Allocation.h
@@ -973,6 +973,7 @@ public:
 	void CleanUpAllocations() { ReadOnlyTexturePool.CleanUpAllocations(0); }
 
 	void Destroy() { ReadOnlyTexturePool.Destroy(); }
+	bool GetMemoryStats(uint64& OutTotalAllocated, uint64& OutTotalUnused) const { OutTotalAllocated = 0; OutTotalUnused = 0; }
 
 private:
 	FD3D12TextureAllocator ReadOnlyTexturePool;
