Workaround fix for check in RegisterAsRenameListener for double registration

diff --git a/Engine/Source/Runtime/D3D12RHI/Private/D3D12RayTracing.cpp b/Engine/Source/Runtime/D3D12RHI/Private/D3D12RayTracing.cpp
--- a/Engine/Source/Runtime/D3D12RHI/Private/D3D12RayTracing.cpp
+++ b/Engine/Source/Runtime/D3D12RHI/Private/D3D12RayTracing.cpp
@@ -2806,6 +2806,7 @@ void FD3D12RayTracingGeometry::BuildAccelerationStructure(FD3D12CommandContext&
 	FD3D12Adapter* Adapter = CommandContext.GetParentAdapter();
 
 	// Register as rename listener to index/vertex buffers
+	UnregisterAsRenameListener(GPUIndex);
 	RegisterAsRenameListener(GPUIndex);
 
 	ID3D12Device5* RayTracingDevice = CommandContext.GetParentDevice()->GetDevice5();
