Disable suballocated resources in readback until multithreaded map/unmap is fixed

diff --git a/Engine/Source/Runtime/D3D12RHI/Private/D3D12RHIPrivate.h b/Engine/Source/Runtime/D3D12RHI/Private/D3D12RHIPrivate.h
--- a/Engine/Source/Runtime/D3D12RHI/Private/D3D12RHIPrivate.h
+++ b/Engine/Source/Runtime/D3D12RHI/Private/D3D12RHIPrivate.h
@@ -109,8 +109,9 @@ typedef FD3D12StateCacheBase FD3D12StateCache;
 #endif //D3D12_RHI_RAYTRACING
 #endif //DEFAULT_BUFFER_POOL_MAX_ALLOC_SIZE
 
-#define READBACK_BUFFER_POOL_MAX_ALLOC_SIZE (64 * 1024)
-#define READBACK_BUFFER_POOL_DEFAULT_POOL_SIZE (4 * 1024 * 1024)
+// Disable readback sub allocation until multithreaded map/unmap has been fixed (was 64KB & 4 MB)
+#define READBACK_BUFFER_POOL_MAX_ALLOC_SIZE 0
+#define READBACK_BUFFER_POOL_DEFAULT_POOL_SIZE 0
 
 #define DEFAULT_CONTEXT_UPLOAD_POOL_SIZE (8 * 1024 * 1024)
 #define DEFAULT_CONTEXT_UPLOAD_POOL_MAX_ALLOC_SIZE (4 * 1024 * 1024)
