Flush command list when certain amount of initial resource copies have been enqueued to fix possible TDR issue during loading when enqueueing a LOT of buffer copies (default every 10K copies)

diff --git a/Engine/Source/Runtime/D3D12RHI/Private/D3D12Buffer.cpp b/Engine/Source/Runtime/D3D12RHI/Private/D3D12Buffer.cpp
--- a/Engine/Source/Runtime/D3D12RHI/Private/D3D12Buffer.cpp
+++ b/Engine/Source/Runtime/D3D12RHI/Private/D3D12Buffer.cpp
@@ -91,7 +91,8 @@ struct FD3D12RHICommandInitializeBuffer final : public FRHICommand<FD3D12RHIComm
 			FD3D12Resource* Destination = CurrentBuffer->ResourceLocation.GetResource();
 			FD3D12Device* Device = Destination->GetParentDevice();
 
-			FD3D12CommandListHandle& hCommandList = Device->GetDefaultCommandContext().CommandListHandle;
+			FD3D12CommandContext& CommandContext = Device->GetDefaultCommandContext();
+			FD3D12CommandListHandle& hCommandList = CommandContext.CommandListHandle;
 			// Copy from the temporary upload heap to the default resource
 			{
 				// if resource doesn't require state tracking then transition to copy dest here (could have been suballocated from shared resource) - not very optimal and should be batched
@@ -100,7 +101,7 @@ struct FD3D12RHICommandInitializeBuffer final : public FRHICommand<FD3D12RHIComm
 					hCommandList.AddTransitionBarrier(Destination, Destination->GetDefaultResourceState(), D3D12_RESOURCE_STATE_COPY_DEST, D3D12_RESOURCE_BARRIER_ALL_SUBRESOURCES);
 				}
 
-				Device->GetDefaultCommandContext().numCopies++;
+				CommandContext.numInitialResourceCopies++;
 				hCommandList.FlushResourceBarriers();
 				hCommandList->CopyBufferRegion(
 					Destination->GetResource(),
@@ -130,6 +131,8 @@ struct FD3D12RHICommandInitializeBuffer final : public FRHICommand<FD3D12RHIComm
 				}
 
 				hCommandList.UpdateResidency(SrcResourceLoc.GetResource());
+
+				CommandContext.ConditionalFlushCommandList();
 			}
 
 			// Buffer is now written and ready, so unlock the block (locked after creation and can be defragmented if needed)
diff --git a/Engine/Source/Runtime/D3D12RHI/Private/D3D12CommandContext.cpp b/Engine/Source/Runtime/D3D12RHI/Private/D3D12CommandContext.cpp
--- a/Engine/Source/Runtime/D3D12RHI/Private/D3D12CommandContext.cpp
+++ b/Engine/Source/Runtime/D3D12RHI/Private/D3D12CommandContext.cpp
@@ -27,6 +27,14 @@ static FAutoConsoleVariableRef CVarCommandListBatchingMode(
 	ECVF_RenderThreadSafe
 );
 
+int32 MaxInitialResourceCopiesPerCommandList = 10000;
+static FAutoConsoleVariableRef CVarSyncTemporalResources(
+	TEXT("D3D12.MaxInitialResourceCopiesPerCommandList"),
+	MaxInitialResourceCopiesPerCommandList,
+	TEXT("Flush command list to GPU after certain amount of enqueued initial resource copy operations (default value 10000)"),
+	ECVF_RenderThreadSafe
+);
+
 // We don't yet have a way to auto-detect that the Radeon Developer Panel is running
 // with profiling enabled, so for now, we have to manually toggle this console var.
 // It needs to be set before device creation, so it's read only.
@@ -324,6 +332,7 @@ void FD3D12CommandContext::OpenCommandList()
 	numClears = 0;
 	numBarriers = 0;
 	numCopies = 0;
+	numInitialResourceCopies = 0;
 	otherWorkCounter = 0;
 }
 
@@ -408,6 +417,16 @@ FD3D12CommandListHandle FD3D12CommandContext::FlushCommands(bool WaitForCompleti
 	return CommandListHandle;
 }
 
+void FD3D12CommandContext::ConditionalFlushCommandList()
+{
+	// Flush command list of reached maximum amount of initial resource command list which can be done in a single
+	// command list - too many can cause TDRs
+	if (MaxInitialResourceCopiesPerCommandList > 0 && numInitialResourceCopies > (uint32)MaxInitialResourceCopiesPerCommandList)
+	{
+		FlushCommands();
+	}
+}
+
 void FD3D12CommandContext::Finish(TArray<FD3D12CommandListHandle>& CommandLists)
 {
 	CloseCommandList();
diff --git a/Engine/Source/Runtime/D3D12RHI/Private/D3D12CommandContext.h b/Engine/Source/Runtime/D3D12RHI/Private/D3D12CommandContext.h
--- a/Engine/Source/Runtime/D3D12RHI/Private/D3D12CommandContext.h
+++ b/Engine/Source/Runtime/D3D12RHI/Private/D3D12CommandContext.h
@@ -121,6 +121,8 @@ public:
 	void ConditionalClearShaderResource(FD3D12ResourceLocation* Resource);
 	void ClearAllShaderResources();
 
+	void ConditionalFlushCommandList();
+
 	FD3D12FastConstantAllocator ConstantsAllocator;
 
 	// Handles to the command list and direct command allocator this context owns (granted by the command list manager/command allocator manager), and a direct pointer to the D3D command list/command allocator.
@@ -195,11 +197,12 @@ public:
 	uint32 numClears;
 	uint32 numBarriers;
 	uint32 numCopies;
+	uint32 numInitialResourceCopies;
 	uint32 otherWorkCounter;
 
 	bool HasDoneWork() const
 	{
-		return (numDraws + numDispatches + numClears + numBarriers + numCopies + otherWorkCounter) > 0;
+		return (numDraws + numDispatches + numClears + numBarriers + numCopies + numInitialResourceCopies + otherWorkCounter) > 0;
 	}
 
 	/** Constant buffers for Set*ShaderParameter calls. */
diff --git a/Engine/Source/Runtime/D3D12RHI/Private/D3D12Texture.cpp b/Engine/Source/Runtime/D3D12RHI/Private/D3D12Texture.cpp
--- a/Engine/Source/Runtime/D3D12RHI/Private/D3D12Texture.cpp
+++ b/Engine/Source/Runtime/D3D12RHI/Private/D3D12Texture.cpp
@@ -213,7 +213,7 @@ struct FD3D12RHICommandInitializeTexture final : public FRHICommand<FD3D12RHICom
 			FD3D12Resource* Resource = CurrentTexture.GetResource();
 
 			FD3D12CommandListHandle& hCommandList = Device->GetDefaultCommandContext().CommandListHandle;
-			hCommandList.GetCurrentOwningContext()->numCopies += NumSubresources;
+			hCommandList.GetCurrentOwningContext()->numInitialResourceCopies += NumSubresources;
 
 			// resource should be in copy dest already, because it's created like that, so no transition required here
 			
@@ -248,6 +248,8 @@ struct FD3D12RHICommandInitializeTexture final : public FRHICommand<FD3D12RHICom
 				check(Resource->GetDefaultResourceState() == DestinationState);
 			}
 
+			Device->GetDefaultCommandContext().ConditionalFlushCommandList();
+
 			// Texture is now written and ready, so unlock the block (locked after creation and can be defragmented if needed)
 			CurrentTexture.ResourceLocation.UnlockPoolData();
 		}
