Build fix - forgot to rename CVar variable

diff --git a/Engine/Source/Runtime/D3D12RHI/Private/D3D12CommandContext.cpp b/Engine/Source/Runtime/D3D12RHI/Private/D3D12CommandContext.cpp
--- a/Engine/Source/Runtime/D3D12RHI/Private/D3D12CommandContext.cpp
+++ b/Engine/Source/Runtime/D3D12RHI/Private/D3D12CommandContext.cpp
@@ -28,7 +28,7 @@ static FAutoConsoleVariableRef CVarCommandListBatchingMode(
 );
 
 int32 MaxInitialResourceCopiesPerCommandList = 10000;
-static FAutoConsoleVariableRef CVarSyncTemporalResources(
+static FAutoConsoleVariableRef CVarMaxInitialResourceCopiesPerCommandList(
 	TEXT("D3D12.MaxInitialResourceCopiesPerCommandList"),
 	MaxInitialResourceCopiesPerCommandList,
 	TEXT("Flush command list to GPU after certain amount of enqueued initial resource copy operations (default value 10000)"),
