Add UAV barrier before calling EmitRaytracingAccelerationStructurePostbuildInfo to make sure all previous build operations are finished

diff --git a/Engine/Source/Runtime/D3D12RHI/Private/D3D12RayTracing.cpp b/Engine/Source/Runtime/D3D12RHI/Private/D3D12RayTracing.cpp
--- a/Engine/Source/Runtime/D3D12RHI/Private/D3D12RayTracing.cpp
+++ b/Engine/Source/Runtime/D3D12RHI/Private/D3D12RayTracing.cpp
@@ -531,6 +531,9 @@ void FD3D12RayTracingCompactionRequestHandler::Update(FD3D12CommandContext& InCo
 
 		FD3D12DynamicRHI::TransitionResource(InCommandContext.CommandListHandle, PostBuildInfoBuffer->GetResource(), D3D12_RESOURCE_STATE_TBD, D3D12_RESOURCE_STATE_UNORDERED_ACCESS, 0, FD3D12DynamicRHI::ETransitionMode::Apply);
 
+		// Force UAV barrier to make sure all previous builds ops are finished
+		InCommandContext.CommandListHandle.AddUAVBarrier();
+
 		// Emit the RT post build info from the selected requests
 		ID3D12GraphicsCommandList4* CommandList4 = InCommandContext.CommandListHandle.RayTracingCommandList();
 		CommandList4->EmitRaytracingAccelerationStructurePostbuildInfo(&PostBuildInfoDesc, ActiveBLASGPUAddresses.Num(), ActiveBLASGPUAddresses.GetData());
