- Force flush pending transitions when calling EmitRaytracingAccelerationStructurePostbuildInfo
- Make sure commited resources with single state UAV are created with the correct initial state

diff --git a/Engine/Source/Runtime/D3D12RHI/Private/D3D12Allocation.cpp b/Engine/Source/Runtime/D3D12RHI/Private/D3D12Allocation.cpp
--- a/Engine/Source/Runtime/D3D12RHI/Private/D3D12Allocation.cpp
+++ b/Engine/Source/Runtime/D3D12RHI/Private/D3D12Allocation.cpp
@@ -1326,7 +1326,7 @@ bool FD3D12DefaultBufferAllocator::IsPlacedResource(D3D12_RESOURCE_FLAGS InResou
 }
 
 
-D3D12_RESOURCE_STATES FD3D12DefaultBufferAllocator::GetDefaultInitialResourceState(D3D12_HEAP_TYPE InHeapType, EBufferUsageFlags InBufferFlags)
+D3D12_RESOURCE_STATES FD3D12DefaultBufferAllocator::GetDefaultInitialResourceState(D3D12_HEAP_TYPE InHeapType, EBufferUsageFlags InBufferFlags, ED3D12ResourceStateMode InResourceStateMode)
 {
 	// Validate the create state
 	if (InHeapType == D3D12_HEAP_TYPE_READBACK)
@@ -1337,9 +1337,15 @@ D3D12_RESOURCE_STATES FD3D12DefaultBufferAllocator::GetDefaultInitialResourceSta
 	{
 		return D3D12_RESOURCE_STATE_GENERIC_READ;
 	}
+	else if (InBufferFlags == BUF_UnorderedAccess && InResourceStateMode == ED3D12ResourceStateMode::SingleState)
+	{
+		check(InHeapType == D3D12_HEAP_TYPE_DEFAULT);
+		return D3D12_RESOURCE_STATE_UNORDERED_ACCESS;
+	}
 #if D3D12_RHI_RAYTRACING
 	else if (InBufferFlags & BUF_AccelerationStructure)
 	{
+		check(InHeapType == D3D12_HEAP_TYPE_DEFAULT);
 		return D3D12_RESOURCE_STATE_RAYTRACING_ACCELERATION_STRUCTURE;
 	}
 #endif // D3D12_RHI_RAYTRACING
diff --git a/Engine/Source/Runtime/D3D12RHI/Private/D3D12Allocation.h b/Engine/Source/Runtime/D3D12RHI/Private/D3D12Allocation.h
--- a/Engine/Source/Runtime/D3D12RHI/Private/D3D12Allocation.h
+++ b/Engine/Source/Runtime/D3D12RHI/Private/D3D12Allocation.h
@@ -425,7 +425,7 @@ public:
 	void UpdateMemoryStats();
 
 	static bool IsPlacedResource(D3D12_RESOURCE_FLAGS InResourceFlags, ED3D12ResourceStateMode InResourceStateMode);
-	static D3D12_RESOURCE_STATES GetDefaultInitialResourceState(D3D12_HEAP_TYPE InHeapType, EBufferUsageFlags InBufferFlags);
+	static D3D12_RESOURCE_STATES GetDefaultInitialResourceState(D3D12_HEAP_TYPE InHeapType, EBufferUsageFlags InBufferFlags, ED3D12ResourceStateMode InResourceStateMode);
 
 private:
 
diff --git a/Engine/Source/Runtime/D3D12RHI/Private/D3D12Buffer.cpp b/Engine/Source/Runtime/D3D12RHI/Private/D3D12Buffer.cpp
--- a/Engine/Source/Runtime/D3D12RHI/Private/D3D12Buffer.cpp
+++ b/Engine/Source/Runtime/D3D12RHI/Private/D3D12Buffer.cpp
@@ -203,7 +203,7 @@ FD3D12Buffer* FD3D12Adapter::CreateRHIBuffer(FRHICommandListImmediate* RHICmdLis
 	D3D12_HEAP_TYPE HeapType = bIsDynamic ? D3D12_HEAP_TYPE_UPLOAD : D3D12_HEAP_TYPE_DEFAULT;
 	const FD3D12Resource::FD3D12ResourceTypeHelper Type(InDesc, HeapType);
 	const D3D12_RESOURCE_STATES InitialState = bSupportResourceStateTracking ? Type.GetOptimalInitialState(InResourceState, false) : 
-		FD3D12DefaultBufferAllocator::GetDefaultInitialResourceState(HeapType, (EBufferUsageFlags)InUsage);
+		FD3D12DefaultBufferAllocator::GetDefaultInitialResourceState(HeapType, (EBufferUsageFlags)InUsage, InResourceStateMode);
 
 	FD3D12Buffer* BufferOut = nullptr;
 	if (bIsDynamic)
diff --git a/Engine/Source/Runtime/D3D12RHI/Private/D3D12PoolAllocator.cpp b/Engine/Source/Runtime/D3D12RHI/Private/D3D12PoolAllocator.cpp
--- a/Engine/Source/Runtime/D3D12RHI/Private/D3D12PoolAllocator.cpp
+++ b/Engine/Source/Runtime/D3D12RHI/Private/D3D12PoolAllocator.cpp
@@ -210,6 +210,10 @@ void FD3D12PoolAllocator::AllocDefaultResource(D3D12_HEAP_TYPE InHeapType, const
 	{
 		check(InCreateState == D3D12_RESOURCE_STATE_GENERIC_READ);
 	}
+	else if (InBufferUsage == BUF_UnorderedAccess && InResourceStateMode == ED3D12ResourceStateMode::SingleState)
+	{
+		check(InCreateState == D3D12_RESOURCE_STATE_UNORDERED_ACCESS);
+	}
 #if D3D12_RHI_RAYTRACING
 	else if (InBufferUsage & BUF_AccelerationStructure)
 	{
diff --git a/Engine/Source/Runtime/D3D12RHI/Private/D3D12RayTracing.cpp b/Engine/Source/Runtime/D3D12RHI/Private/D3D12RayTracing.cpp
--- a/Engine/Source/Runtime/D3D12RHI/Private/D3D12RayTracing.cpp
+++ b/Engine/Source/Runtime/D3D12RHI/Private/D3D12RayTracing.cpp
@@ -533,6 +533,7 @@ void FD3D12RayTracingCompactionRequestHandler::Update(FD3D12CommandContext& InCo
 
 		// Force UAV barrier to make sure all previous builds ops are finished
 		InCommandContext.CommandListHandle.AddUAVBarrier();
+		InCommandContext.CommandListHandle.FlushResourceBarriers();
 
 		// Emit the RT post build info from the selected requests
 		ID3D12GraphicsCommandList4* CommandList4 = InCommandContext.CommandListHandle.RayTracingCommandList();
@@ -540,6 +541,8 @@ void FD3D12RayTracingCompactionRequestHandler::Update(FD3D12CommandContext& InCo
 
 		// Transition to copy source and perform the copy to readback
 		FD3D12DynamicRHI::TransitionResource(InCommandContext.CommandListHandle, PostBuildInfoBuffer.GetReference()->GetResource(), D3D12_RESOURCE_STATE_UNORDERED_ACCESS, D3D12_RESOURCE_STATE_COPY_SOURCE, 0, FD3D12DynamicRHI::ETransitionMode::Apply);
+		InCommandContext.CommandListHandle.FlushResourceBarriers();
+
 		InCommandContext.RHICopyToStagingBuffer(PostBuildInfoBuffer, PostBuildInfoStagingBuffer, 0, sizeof(uint64) * ActiveBLASGPUAddresses.Num());
 
 		// Update the fence value
