Make sure residency is updated before calling EmitRaytracingAccelerationStructurePostbuildInfo  on it

diff --git a/Engine/Source/Runtime/D3D12RHI/Private/D3D12RayTracing.cpp b/Engine/Source/Runtime/D3D12RHI/Private/D3D12RayTracing.cpp
--- a/Engine/Source/Runtime/D3D12RHI/Private/D3D12RayTracing.cpp
+++ b/Engine/Source/Runtime/D3D12RHI/Private/D3D12RayTracing.cpp
@@ -510,7 +510,11 @@ void FD3D12RayTracingCompactionRequestHandler::Update(FD3D12CommandContext& InCo
 	for (FD3D12RayTracingGeometry* RTGeomtry : PendingRequests)
 	{
 		ActiveRequests.Add(RTGeomtry);
-		ActiveBLASGPUAddresses.Add(RTGeomtry->AccelerationStructureBuffers[GPUIndex].GetReference()->ResourceLocation.GetGPUVirtualAddress());
+
+		FD3D12ResourceLocation& ResourceLocation = RTGeomtry->AccelerationStructureBuffers[GPUIndex].GetReference()->ResourceLocation;
+		ActiveBLASGPUAddresses.Add(ResourceLocation.GetGPUVirtualAddress());
+
+		ResourceLocation.GetResource()->UpdateResidency(InCommandContext.CommandListHandle);
 
 		// enqueued enough requests for this update round
 		if (ActiveRequests.Num() >= GD3D12RayTracingMaxBatchedCompaction)
