Keep readback memory mapped to fix map/unmap issues when called on the same suballocated parent resource from multiple threads at the same time (different sub allocated RHI resources so should be allowed)

diff --git a/Engine/Source/Runtime/D3D12RHI/Private/D3D12Allocation.cpp b/Engine/Source/Runtime/D3D12RHI/Private/D3D12Allocation.cpp
--- a/Engine/Source/Runtime/D3D12RHI/Private/D3D12Allocation.cpp
+++ b/Engine/Source/Runtime/D3D12RHI/Private/D3D12Allocation.cpp
@@ -232,7 +232,7 @@ void FD3D12BuddyAllocator::Initialize()
 		BackingHeap->SetHeap(Heap);
 
 		// Only track resources that cannot be accessed on the CPU.
-		if (IsCPUInaccessible(InitConfig.HeapType))
+		if (IsGPUOnly(InitConfig.HeapType))
 		{
 			BackingHeap->BeginTrackingResidency(Desc.SizeInBytes);
 		}
@@ -245,7 +245,7 @@ void FD3D12BuddyAllocator::Initialize()
 			VERIFYD3D12RESULT(Adapter->CreateBuffer(HeapProps, GetGPUMask(), InitConfig.InitialResourceState, ED3D12ResourceStateMode::SingleState, InitConfig.InitialResourceState, MaxBlockSize, BackingResource.GetInitReference(), TEXT("Resource Allocator Underlying Buffer"), InitConfig.ResourceFlags));
 		}
 
-		if (IsCPUWritable(InitConfig.HeapType))
+		if (IsCPUAccessible(InitConfig.HeapType))
 		{
 			BackingResource->Map();
 		}
@@ -379,7 +379,7 @@ void FD3D12BuddyAllocator::Allocate(uint32 SizeInBytes, uint32 Alignment, FD3D12
 		ResourceLocation.SetResource(BackingResource);
 		ResourceLocation.SetGPUVirtualAddress(BackingResource->GetGPUVirtualAddress() + AlignedOffsetFromResourceBase);
 
-		if (IsCPUWritable(InitConfig.HeapType))
+		if (IsCPUAccessible(InitConfig.HeapType))
 		{
 			ResourceLocation.SetMappedBaseAddress((uint8*)BackingResource->GetResourceBaseAddress() + AlignedOffsetFromResourceBase);
 		}
@@ -799,7 +799,7 @@ bool FD3D12BucketAllocator::TryAllocate(uint32 SizeInBytes, uint32 Alignment, FD
 		// Track the resource so we know when to delete it
 		SubAllocatedResources.Add(Resource);
 
-		if (IsCPUWritable(InitConfig.HeapType))
+		if (IsCPUAccessible(InitConfig.HeapType))
 		{
 			BaseAddress = Resource->Map();
 			check(BaseAddress);
@@ -843,7 +843,7 @@ bool FD3D12BucketAllocator::TryAllocate(uint32 SizeInBytes, uint32 Alignment, FD
 	ResourceLocation.SetOffsetFromBaseOfResource(AlignedBlockOffset);
 	ResourceLocation.SetGPUVirtualAddress(Block.ResourceHeap->GetGPUVirtualAddress() + AlignedBlockOffset);
 
-	if (IsCPUWritable(InitConfig.HeapType))
+	if (IsCPUAccessible(InitConfig.HeapType))
 	{
 		ResourceLocation.SetMappedBaseAddress((void*)((uint64)Block.ResourceHeap->GetResourceBaseAddress() + AlignedBlockOffset));
 	}
@@ -1846,15 +1846,9 @@ void* FD3D12FastAllocator::Allocate(uint32 Size, uint32 Alignment, class FD3D12R
 		ResourceName = FString::Printf(TEXT("Stand Alone Fast Allocation %lld"), UniqueID);
 #endif
 		VERIFYD3D12RESULT(Adapter->CreateBuffer(PagePool.GetHeapType(), GetGPUMask(), GetVisibilityMask(), Size + Alignment, &Resource, *ResourceName));
-
-		void* Data = nullptr;
-		if (PagePool.IsCPUWritable())
-		{
-			Data = Resource->Map();
-		}
 		ResourceLocation->AsStandAlone(Resource, Size + Alignment);
 
-		return Data;
+		return PagePool.IsCPUWritable() ? Resource->GetResourceBaseAddress() : nullptr;
 	}
 	else
 	{
diff --git a/Engine/Source/Runtime/D3D12RHI/Private/D3D12PoolAllocator.cpp b/Engine/Source/Runtime/D3D12RHI/Private/D3D12PoolAllocator.cpp
--- a/Engine/Source/Runtime/D3D12RHI/Private/D3D12PoolAllocator.cpp
+++ b/Engine/Source/Runtime/D3D12RHI/Private/D3D12PoolAllocator.cpp
@@ -68,7 +68,7 @@ void FD3D12MemoryPool::Init()
 		BackingHeap->SetHeap(Heap);
 
 		// Only track resources that cannot be accessed on the CPU.
-		if (IsCPUInaccessible(InitConfig.HeapType))
+		if (IsGPUOnly(InitConfig.HeapType))
 		{
 			BackingHeap->BeginTrackingResidency(Desc.SizeInBytes);
 		}
@@ -81,7 +81,7 @@ void FD3D12MemoryPool::Init()
 			VERIFYD3D12RESULT(Adapter->CreateBuffer(HeapProps, GetGPUMask(), InitConfig.InitialResourceState, ED3D12ResourceStateMode::SingleState, InitConfig.InitialResourceState, PoolSize, BackingResource.GetInitReference(), TEXT("Resource Allocator Underlying Buffer"), InitConfig.ResourceFlags));
 		}
 
-		if (IsCPUWritable(InitConfig.HeapType))
+		if (IsCPUAccessible(InitConfig.HeapType))
 		{
 			BackingResource->Map();
 		}
@@ -277,7 +277,7 @@ void FD3D12PoolAllocator::AllocResource(D3D12_HEAP_TYPE InHeapType, const D3D12_
 				ResourceLocation.SetResource(BackingResource);
 				ResourceLocation.SetGPUVirtualAddress(BackingResource->GetGPUVirtualAddress() + AllocationData.GetOffset());
 
-				if (IsCPUWritable(InitConfig.HeapType))
+				if (IsCPUAccessible(InitConfig.HeapType))
 				{
 					ResourceLocation.SetMappedBaseAddress((uint8*)BackingResource->GetResourceBaseAddress() + AllocationData.GetOffset());
 				}
diff --git a/Engine/Source/Runtime/D3D12RHI/Private/D3D12RHIPrivate.h b/Engine/Source/Runtime/D3D12RHI/Private/D3D12RHIPrivate.h
--- a/Engine/Source/Runtime/D3D12RHI/Private/D3D12RHIPrivate.h
+++ b/Engine/Source/Runtime/D3D12RHI/Private/D3D12RHIPrivate.h
@@ -109,9 +109,8 @@ typedef FD3D12StateCacheBase FD3D12StateCache;
 #endif //D3D12_RHI_RAYTRACING
 #endif //DEFAULT_BUFFER_POOL_MAX_ALLOC_SIZE
 
-// Disable readback sub allocation until multithreaded map/unmap has been fixed (was 64KB & 4 MB)
-#define READBACK_BUFFER_POOL_MAX_ALLOC_SIZE 0
-#define READBACK_BUFFER_POOL_DEFAULT_POOL_SIZE 0
+#define READBACK_BUFFER_POOL_MAX_ALLOC_SIZE (64 * 1024)
+#define READBACK_BUFFER_POOL_DEFAULT_POOL_SIZE (4 * 1024 * 1024)
 
 #define DEFAULT_CONTEXT_UPLOAD_POOL_SIZE (8 * 1024 * 1024)
 #define DEFAULT_CONTEXT_UPLOAD_POOL_MAX_ALLOC_SIZE (4 * 1024 * 1024)
diff --git a/Engine/Source/Runtime/D3D12RHI/Private/D3D12Resources.cpp b/Engine/Source/Runtime/D3D12RHI/Private/D3D12Resources.cpp
--- a/Engine/Source/Runtime/D3D12RHI/Private/D3D12Resources.cpp
+++ b/Engine/Source/Runtime/D3D12RHI/Private/D3D12Resources.cpp
@@ -302,7 +302,7 @@ FD3D12Resource::~FD3D12Resource()
 void FD3D12Resource::StartTrackingForResidency()
 {
 #if ENABLE_RESIDENCY_MANAGEMENT
-	check(IsCPUInaccessible(HeapType));	// This is checked at a higher level before calling this function.
+	check(IsGPUOnly(HeapType));	// This is checked at a higher level before calling this function.
 	check(D3DX12Residency::IsInitialized(ResidencyHandle) == false);
 	const D3D12_RESOURCE_DESC ResourceDesc = Resource->GetDesc();
 	const D3D12_RESOURCE_ALLOCATION_INFO Info = GetParentDevice()->GetDevice()->GetResourceAllocationInfo(0, 1, &ResourceDesc);
@@ -432,7 +432,7 @@ HRESULT FD3D12Adapter::CreateCommittedResource(const D3D12_RESOURCE_DESC& InDesc
 		SetName(*ppOutResource, Name);
 
 		// Only track resources that cannot be accessed on the CPU.
-		if (IsCPUInaccessible(HeapProps.Type))
+		if (IsGPUOnly(HeapProps.Type))
 		{
 			(*ppOutResource)->StartTrackingForResidency();
 		}
@@ -806,7 +806,7 @@ void FD3D12ResourceLocation::AsStandAlone(FD3D12Resource* Resource, uint64 InSiz
 	SetResource(Resource);
 	SetSize(InSize);
 
-	if (!IsCPUInaccessible(Resource->GetHeapType()))
+	if (IsCPUAccessible(Resource->GetHeapType()))
 	{
 		D3D12_RANGE range = { 0, IsCPUWritable(Resource->GetHeapType()) ? 0 : InSize };
 		SetMappedBaseAddress(Resource->Map(&range));
diff --git a/Engine/Source/Runtime/D3D12RHI/Private/D3D12Texture.cpp b/Engine/Source/Runtime/D3D12RHI/Private/D3D12Texture.cpp
--- a/Engine/Source/Runtime/D3D12RHI/Private/D3D12Texture.cpp
+++ b/Engine/Source/Runtime/D3D12RHI/Private/D3D12Texture.cpp
@@ -696,11 +696,6 @@ void SafeCreateTexture2D(FD3D12Device* pDevice,
 				FD3D12Resource* Resource = nullptr;
 				VERIFYD3D12CREATETEXTURERESULT(Adapter->CreateBuffer(HeapType, pDevice->GetGPUMask(), pDevice->GetVisibilityMask(), Size, &Resource, Name), TextureDesc, pDevice->GetDevice());
 				OutTexture2D->AsStandAlone(Resource);
-
-				if (IsCPUWritable(HeapType))
-				{
-					OutTexture2D->SetMappedBaseAddress(Resource->Map());
-				}
 			}
 			break;
 
diff --git a/Engine/Source/Runtime/D3D12RHI/Private/D3D12UAV.cpp b/Engine/Source/Runtime/D3D12RHI/Private/D3D12UAV.cpp
--- a/Engine/Source/Runtime/D3D12RHI/Private/D3D12UAV.cpp
+++ b/Engine/Source/Runtime/D3D12RHI/Private/D3D12UAV.cpp
@@ -339,10 +339,8 @@ void* FD3D12StagingBuffer::Lock(uint32 Offset, uint32 NumBytes)
 	bIsLocked = true;
 	if (ResourceLocation.IsValid())
 	{
-		D3D12_RANGE ReadRange;
-		ReadRange.Begin = ResourceLocation.GetOffsetFromBaseOfResource() + Offset;
-		ReadRange.End = ResourceLocation.GetOffsetFromBaseOfResource() + Offset + NumBytes;
-		return reinterpret_cast<uint8*>(ResourceLocation.GetResource()->Map(&ReadRange)) + ResourceLocation.GetOffsetFromBaseOfResource() + Offset;
+		// readback resource are kept mapped after creation
+		return reinterpret_cast<uint8*>(ResourceLocation.GetMappedBaseAddress()) + Offset;
 	}
 	else
 	{
@@ -354,8 +352,4 @@ void FD3D12StagingBuffer::Unlock()
 {
 	check(bIsLocked);
 	bIsLocked = false;
-	if (ResourceLocation.IsValid())
-	{
-		ResourceLocation.GetResource()->Unmap();
-	}
 }
\ No newline at end of file
diff --git a/Engine/Source/Runtime/D3D12RHI/Public/D3D12Resources.h b/Engine/Source/Runtime/D3D12RHI/Public/D3D12Resources.h
--- a/Engine/Source/Runtime/D3D12RHI/Public/D3D12Resources.h
+++ b/Engine/Source/Runtime/D3D12RHI/Public/D3D12Resources.h
@@ -173,6 +173,7 @@ public:
 	inline void* Map(const D3D12_RANGE* ReadRange = nullptr)
 	{
 		check(Resource);
+		check(ResourceBaseAddress == nullptr);
 		VERIFYD3D12RESULT(Resource->Map(0, ReadRange, &ResourceBaseAddress));
 
 		return ResourceBaseAddress;
@@ -564,6 +565,8 @@ public:
 
 	inline void AsHeapAliased(FD3D12Resource* Resource)
 	{
+		check(Resource->GetHeapType() != D3D12_HEAP_TYPE_READBACK);
+
 		SetType(FD3D12ResourceLocation::ResourceLocationType::eHeapAliased);
 		SetResource(Resource);
 		SetSize(0);
diff --git a/Engine/Source/Runtime/D3D12RHI/Public/D3D12Util.h b/Engine/Source/Runtime/D3D12RHI/Public/D3D12Util.h
--- a/Engine/Source/Runtime/D3D12RHI/Public/D3D12Util.h
+++ b/Engine/Source/Runtime/D3D12RHI/Public/D3D12Util.h
@@ -480,7 +480,7 @@ inline bool IsCPUWritable(D3D12_HEAP_TYPE HeapType, const D3D12_HEAP_PROPERTIES
 			(pCustomHeapProperties->CPUPageProperty == D3D12_CPU_PAGE_PROPERTY_WRITE_COMBINE || pCustomHeapProperties->CPUPageProperty == D3D12_CPU_PAGE_PROPERTY_WRITE_BACK));
 }
 
-inline bool IsCPUInaccessible(D3D12_HEAP_TYPE HeapType, const D3D12_HEAP_PROPERTIES *pCustomHeapProperties = nullptr)
+inline bool IsGPUOnly(D3D12_HEAP_TYPE HeapType, const D3D12_HEAP_PROPERTIES *pCustomHeapProperties = nullptr)
 {
 	check(HeapType == D3D12_HEAP_TYPE_CUSTOM ? pCustomHeapProperties != nullptr : true);
 	return HeapType == D3D12_HEAP_TYPE_DEFAULT ||
@@ -488,6 +488,11 @@ inline bool IsCPUInaccessible(D3D12_HEAP_TYPE HeapType, const D3D12_HEAP_PROPERT
 		(pCustomHeapProperties->CPUPageProperty == D3D12_CPU_PAGE_PROPERTY_NOT_AVAILABLE));
 }
 
+inline bool IsCPUAccessible(D3D12_HEAP_TYPE HeapType, const D3D12_HEAP_PROPERTIES* pCustomHeapProperties = nullptr)
+{
+	return !IsGPUOnly(HeapType, pCustomHeapProperties);
+}
+
 inline D3D12_RESOURCE_STATES DetermineInitialResourceState(D3D12_HEAP_TYPE HeapType, const D3D12_HEAP_PROPERTIES *pCustomHeapProperties = nullptr)
 {
 	if (HeapType == D3D12_HEAP_TYPE_DEFAULT || IsCPUWritable(HeapType, pCustomHeapProperties))
