Set r.D3D12.RayTracing.ExecuteIndirectWorkaround=1 by default to work around current NVIDIA driver issues when using indirect ray tracing.

This forces a command list flush before each indirect ray tracing dispatch, which causes a heavy GPU pipeline flush but ensures that indirect RT dispatches work.

diff --git a/Engine/Source/Runtime/D3D12RHI/Private/D3D12RayTracing.cpp b/Engine/Source/Runtime/D3D12RHI/Private/D3D12RayTracing.cpp
--- a/Engine/Source/Runtime/D3D12RHI/Private/D3D12RayTracing.cpp
+++ b/Engine/Source/Runtime/D3D12RHI/Private/D3D12RayTracing.cpp
@@ -83,11 +83,11 @@ static FAutoConsoleVariableRef CVarD3D12RayTracingMaxBatchedCompaction(
 	ECVF_ReadOnly
 );
 
-static int32 GD3D12RayTracingExecuteIndirectWorkaround = 0;
+static int32 GD3D12RayTracingExecuteIndirectWorkaround = 1;
 static FAutoConsoleVariableRef CVarD3D12RayTracingExecuteIndirectWorkaround(
 	TEXT("r.D3D12.RayTracing.ExecuteIndirectWorkaround"),
 	GD3D12RayTracingExecuteIndirectWorkaround,
-	TEXT("Adds a command list flush before ExecuteIndirect to avoid NVIDIA driver bug that causes some indirect RT commands to get lost. (default = 0)\n")
+	TEXT("Adds a command list flush before ExecuteIndirect to avoid NVIDIA driver bug that causes some indirect RT commands to get lost. (default = 1)\n")
 );
 
 // Ray tracing stat counters
