- Fix bug when copying RTAccelleration with missing offset from base address
- Enable defrag of RTAccelleration data pools
- Add CPU and Push/Pop event scopes for defrag times

diff --git a/Engine/Source/Runtime/D3D12RHI/Private/D3D12Allocation.cpp b/Engine/Source/Runtime/D3D12RHI/Private/D3D12Allocation.cpp
--- a/Engine/Source/Runtime/D3D12RHI/Private/D3D12Allocation.cpp
+++ b/Engine/Source/Runtime/D3D12RHI/Private/D3D12Allocation.cpp
@@ -1266,15 +1266,6 @@ FD3D12BufferPool* FD3D12DefaultBufferAllocator::CreateBufferPool(D3D12_HEAP_TYPE
 
 	// Disable defrag if not Default memory
 	bool bDefragEnabled = (InitConfig.HeapType == D3D12_HEAP_TYPE_DEFAULT);
-
-#if D3D12_RHI_RAYTRACING
-	// disable defrag on the RT Acceleration pool - get GPU crashes after moving the data
-	if (InitConfig.InitialResourceState == D3D12_RESOURCE_STATE_RAYTRACING_ACCELERATION_STRUCTURE)
-	{
-		bDefragEnabled = false;
-	}
-#endif // D3D12_RHI_RAYTRACING
-
 	FD3D12BufferPool* NewPool = new FD3D12PoolAllocator(Device, GetVisibilityMask(), InitConfig, Name, AllocationStrategy, PoolSize, PoolAlignment, MaxAllocationSize, bDefragEnabled);
 
 #else // USE_POOL_ALLOCATOR
@@ -1405,8 +1396,11 @@ void FD3D12DefaultBufferAllocator::FreeDefaultBufferPools()
 void FD3D12DefaultBufferAllocator::BeginFrame()
 {
 #if USE_POOL_ALLOCATOR
+
 	if (GD3D12VRAMBufferPoolDefrag > 0 && GD3D12VRAMBufferPoolDefragMaxCopySizePerFrame > 0)
 	{
+		TRACE_CPUPROFILER_EVENT_SCOPE(DefragBufferPool);
+
 		uint32 MaxCopySize = GD3D12VRAMBufferPoolDefragMaxCopySizePerFrame;
 		uint32 CopySize = 0;
 		for (FD3D12BufferPool* DefaultBufferPool : DefaultBufferPools)
@@ -1418,13 +1412,19 @@ void FD3D12DefaultBufferAllocator::BeginFrame()
 		}
 	}
 
-	FD3D12CommandContext& CommandContext = GetParentDevice()->GetDefaultCommandContext();
-	for (FD3D12BufferPool* DefaultBufferPool : DefaultBufferPools)
 	{
-		if (DefaultBufferPool)
+		TRACE_CPUPROFILER_EVENT_SCOPE(FlushPendingBufferCopyOps);
+
+		FD3D12CommandContext& CommandContext = GetParentDevice()->GetDefaultCommandContext();
+		CommandContext.RHIPushEvent(TEXT("BufferPoolCopyOps"), FColor::Emerald);
+		for (FD3D12BufferPool* DefaultBufferPool : DefaultBufferPools)
 		{
-			DefaultBufferPool->FlushPendingCopyOps(CommandContext);
+			if (DefaultBufferPool)
+			{
+				DefaultBufferPool->FlushPendingCopyOps(CommandContext);
+			}
 		}
+		CommandContext.RHIPopEvent();
 	}
 #endif // USE_POOL_ALLOCATOR
 }
@@ -1583,10 +1583,10 @@ HRESULT FD3D12TextureAllocatorPool::AllocateTexture(
 
 void FD3D12TextureAllocatorPool::BeginFrame()
 {
-	FD3D12CommandContext& CommandContext = GetParentDevice()->GetDefaultCommandContext();
-
 	if (GD3D12VRAMTexturePoolDefrag > 0 && GD3D12VRAMTexturePoolDefragMaxCopySizePerFrame > 0)
 	{
+		TRACE_CPUPROFILER_EVENT_SCOPE(DefragTexturePool);
+
 		uint32 MaxCopySize = GD3D12VRAMTexturePoolDefragMaxCopySizePerFrame;
 		uint32 CopySize = 0;
 		for (uint32 PoolIndex = 0; PoolIndex < (uint32)EPoolType::Count; ++PoolIndex)
@@ -1595,9 +1595,17 @@ void FD3D12TextureAllocatorPool::BeginFrame()
 		}
 	}
 
-	for (uint32 PoolIndex = 0; PoolIndex < (uint32)EPoolType::Count; ++PoolIndex)
 	{
-		PoolAllocators[PoolIndex]->FlushPendingCopyOps(CommandContext);
+		TRACE_CPUPROFILER_EVENT_SCOPE(FlushPendingTextureCopyOps);
+
+		FD3D12CommandContext& CommandContext = GetParentDevice()->GetDefaultCommandContext();
+
+		CommandContext.RHIPushEvent(TEXT("TexturePoolCopyOps"), FColor::Emerald);
+		for (uint32 PoolIndex = 0; PoolIndex < (uint32)EPoolType::Count; ++PoolIndex)
+		{
+			PoolAllocators[PoolIndex]->FlushPendingCopyOps(CommandContext);
+		}
+		CommandContext.RHIPopEvent();
 	}
 }
 
diff --git a/Engine/Source/Runtime/D3D12RHI/Private/D3D12PoolAllocator.cpp b/Engine/Source/Runtime/D3D12RHI/Private/D3D12PoolAllocator.cpp
--- a/Engine/Source/Runtime/D3D12RHI/Private/D3D12PoolAllocator.cpp
+++ b/Engine/Source/Runtime/D3D12RHI/Private/D3D12PoolAllocator.cpp
@@ -581,8 +581,8 @@ void FD3D12PoolAllocator::FlushPendingCopyOps(FD3D12CommandContext& InCommandCon
 #if D3D12_RHI_RAYTRACING
 		if (bRTAccelerationStructure)
 		{
-			CommandListHandle.GraphicsCommandList5()->CopyRaytracingAccelerationStructure(CopyOperation.DestResource->GetResource()->GetGPUVirtualAddress(),
-				CopyOperation.SourceResource->GetResource()->GetGPUVirtualAddress(), D3D12_RAYTRACING_ACCELERATION_STRUCTURE_COPY_MODE_CLONE);
+			CommandListHandle.GraphicsCommandList5()->CopyRaytracingAccelerationStructure(CopyOperation.DestResource->GetResource()->GetGPUVirtualAddress() + CopyOperation.DestOffset,
+				CopyOperation.SourceResource->GetResource()->GetGPUVirtualAddress() + CopyOperation.SourceOffset, D3D12_RAYTRACING_ACCELERATION_STRUCTURE_COPY_MODE_CLONE);
 		}
 		else
 #endif // D3D12_RHI_RAYTRACING
