Remove empty files after buffer refactor. Some leftover functions were moved to other files to allow files to be removed. No actual code changes.

diff --git a/Engine/Source/Runtime/Apple/MetalRHI/Private/MetalIndexBuffer.cpp b/Engine/Source/Runtime/Apple/MetalRHI/Private/MetalIndexBuffer.cpp
deleted file mode 100644
diff --git a/Engine/Source/Runtime/Apple/MetalRHI/Private/MetalStructuredBuffer.cpp b/Engine/Source/Runtime/Apple/MetalRHI/Private/MetalStructuredBuffer.cpp
deleted file mode 100644
diff --git a/Engine/Source/Runtime/Apple/MetalRHI/Private/MetalVertexBuffer.cpp b/Engine/Source/Runtime/Apple/MetalRHI/Private/MetalVertexBuffer.cpp
--- a/Engine/Source/Runtime/Apple/MetalRHI/Private/MetalVertexBuffer.cpp
+++ b/Engine/Source/Runtime/Apple/MetalRHI/Private/MetalVertexBuffer.cpp
@@ -846,6 +846,24 @@ void FMetalDynamicRHI::RHICopyBuffer(FRHIBuffer* SourceBufferRHI, FRHIBuffer* De
 	}
 }
 
+void FMetalDynamicRHI::RHITransferBufferUnderlyingResource(FRHIBuffer* DestBuffer, FRHIBuffer* SrcBuffer)
+{
+	@autoreleasepool {
+		check(DestBuffer);
+		FMetalResourceMultiBuffer* Dest = ResourceCast(DestBuffer);
+		if (!SrcBuffer)
+		{
+			TRefCountPtr<FMetalResourceMultiBuffer> DeletionProxy = new FMetalResourceMultiBuffer(0, Dest->GetUsage(), Dest->GetStride(), nullptr, Dest->Type);
+			Dest->Swap(*DeletionProxy);
+		}
+		else
+		{
+			FMetalResourceMultiBuffer* Src = ResourceCast(SrcBuffer);
+			Dest->Swap(*Src);
+		}
+	}
+}
+
 void FMetalRHIBuffer::Init_RenderThread(FRHICommandListImmediate& RHICmdList, uint32 InSize, uint32 InUsage, FRHIResourceCreateInfo& CreateInfo, FRHIResource* Resource)
 {
 	if (CreateInfo.ResourceArray)
diff --git a/Engine/Source/Runtime/D3D12RHI/Private/D3D12Buffer.cpp b/Engine/Source/Runtime/D3D12RHI/Private/D3D12Buffer.cpp
--- a/Engine/Source/Runtime/D3D12RHI/Private/D3D12Buffer.cpp
+++ b/Engine/Source/Runtime/D3D12RHI/Private/D3D12Buffer.cpp
@@ -636,3 +636,213 @@ void FD3D12DynamicRHI::RHITransferBufferUnderlyingResource(FRHIBuffer* DestBuffe
 		Dest->Swap(*Src);
 	}
 }
+
+void FD3D12DynamicRHI::RHICopyBuffer(FRHIBuffer* SourceBufferRHI, FRHIBuffer* DestBufferRHI)
+{
+	FD3D12Buffer* SrcBuffer = FD3D12DynamicRHI::ResourceCast(SourceBufferRHI);
+	FD3D12Buffer* DstBuffer = FD3D12DynamicRHI::ResourceCast(DestBufferRHI);
+	check(SrcBuffer->GetSize() == DstBuffer->GetSize());
+
+	FD3D12Buffer* SourceBuffer = SrcBuffer;
+	FD3D12Buffer* DestBuffer = DstBuffer;
+
+	for (FD3D12Buffer::FDualLinkedObjectIterator It(SourceBuffer, DestBuffer); It; ++It)
+	{
+		SourceBuffer = It.GetFirst();
+		DestBuffer = It.GetSecond();
+
+		FD3D12Device* Device = SourceBuffer->GetParentDevice();
+		check(Device == DestBuffer->GetParentDevice());
+
+		FD3D12Resource* pSourceResource = SourceBuffer->ResourceLocation.GetResource();
+		D3D12_RESOURCE_DESC const& SourceBufferDesc = pSourceResource->GetDesc();
+
+		FD3D12Resource* pDestResource = DestBuffer->ResourceLocation.GetResource();
+		D3D12_RESOURCE_DESC const& DestBufferDesc = pDestResource->GetDesc();
+
+		check(SourceBufferDesc.Width == DestBufferDesc.Width);
+
+		FD3D12CommandContext& Context = Device->GetDefaultCommandContext();
+		Context.numCopies++;
+		Context.CommandListHandle->CopyResource(pDestResource->GetResource(), pSourceResource->GetResource());
+		Context.CommandListHandle.UpdateResidency(pDestResource);
+		Context.CommandListHandle.UpdateResidency(pSourceResource);
+
+		DEBUG_EXECUTE_COMMAND_CONTEXT(Device->GetDefaultCommandContext());
+
+		Device->RegisterGPUWork(1);
+	}
+}
+
+#if D3D12_RHI_RAYTRACING
+void FD3D12CommandContext::RHICopyBufferRegion(FRHIBuffer* DestBufferRHI, uint64 DstOffset, FRHIBuffer* SourceBufferRHI, uint64 SrcOffset, uint64 NumBytes)
+{
+	FD3D12Buffer* SourceBuffer = RetrieveObject<FD3D12Buffer>(SourceBufferRHI);
+	FD3D12Buffer* DestBuffer = RetrieveObject<FD3D12Buffer>(DestBufferRHI);
+
+	FD3D12Device* Device = SourceBuffer->GetParentDevice();
+	check(Device == DestBuffer->GetParentDevice());
+	check(Device == GetParentDevice());
+
+	FD3D12Resource* pSourceResource = SourceBuffer->ResourceLocation.GetResource();
+	D3D12_RESOURCE_DESC const& SourceBufferDesc = pSourceResource->GetDesc();
+
+	FD3D12Resource* pDestResource = DestBuffer->ResourceLocation.GetResource();
+	D3D12_RESOURCE_DESC const& DestBufferDesc = pDestResource->GetDesc();
+
+	checkf(pSourceResource != pDestResource, TEXT("CopyBufferRegion cannot be used on the same resource. This can happen when both the source and the dest are suballocated from the same resource."));
+
+	check(DstOffset + NumBytes <= DestBufferDesc.Width);
+	check(SrcOffset + NumBytes <= SourceBufferDesc.Width);
+
+	numCopies++;
+
+	FScopedResourceBarrier ScopeResourceBarrierSource(CommandListHandle, pSourceResource, D3D12_RESOURCE_STATE_COPY_SOURCE, 0, FD3D12DynamicRHI::ETransitionMode::Validate);
+	FScopedResourceBarrier ScopeResourceBarrierDest(CommandListHandle, pDestResource, D3D12_RESOURCE_STATE_COPY_DEST, 0, FD3D12DynamicRHI::ETransitionMode::Validate);
+	CommandListHandle.FlushResourceBarriers();
+
+	CommandListHandle->CopyBufferRegion(pDestResource->GetResource(), DestBuffer->ResourceLocation.GetOffsetFromBaseOfResource() + DstOffset, pSourceResource->GetResource(), SourceBuffer->ResourceLocation.GetOffsetFromBaseOfResource() + SrcOffset, NumBytes);
+	CommandListHandle.UpdateResidency(pDestResource);
+	CommandListHandle.UpdateResidency(pSourceResource);
+
+	Device->RegisterGPUWork(1);
+}
+
+void FD3D12CommandContext::RHICopyBufferRegions(const TArrayView<const FCopyBufferRegionParams> Params)
+{
+	// Batched buffer copy finds unique source and destination buffer resources, performs transitions
+	// to copy source / dest state, then performs copies and finally restores original state.
+
+	using FLocalResourceArray = TArray<FD3D12Resource*, TInlineAllocator<16, TMemStackAllocator<>>>;
+	FLocalResourceArray SrcBuffers;
+	FLocalResourceArray DstBuffers;
+
+	SrcBuffers.Reserve(Params.Num());
+	DstBuffers.Reserve(Params.Num());
+
+	// Transition buffers to copy states
+	for (auto& Param : Params)
+	{
+		FD3D12Buffer* SourceBuffer = RetrieveObject<FD3D12Buffer>(Param.SourceBuffer);
+		FD3D12Buffer* DestBuffer = RetrieveObject<FD3D12Buffer>(Param.DestBuffer);
+		check(SourceBuffer);
+		check(DestBuffer);
+
+		FD3D12Device* Device = SourceBuffer->GetParentDevice();
+		check(Device == DestBuffer->GetParentDevice());
+		check(Device == GetParentDevice());
+
+		FD3D12Resource* pSourceResource = SourceBuffer->ResourceLocation.GetResource();
+		FD3D12Resource* pDestResource = DestBuffer->ResourceLocation.GetResource();
+
+		checkf(pSourceResource != pDestResource, TEXT("CopyBufferRegion cannot be used on the same resource. This can happen when both the source and the dest are suballocated from the same resource."));
+
+		SrcBuffers.Add(pSourceResource);
+		DstBuffers.Add(pDestResource);
+	}
+
+	Algo::Sort(SrcBuffers);
+	Algo::Sort(DstBuffers);
+
+	enum class EBatchCopyState
+	{
+		CopySource,
+		CopyDest,
+		FinalizeSource,
+		FinalizeDest,
+	};
+
+	auto TransitionResources = [](FD3D12CommandListHandle& CommandListHandle, FLocalResourceArray& SortedResources, EBatchCopyState State)
+	{
+		const uint32 Subresource = 0; // Buffers only have one subresource
+
+		FD3D12Resource* PrevResource = nullptr;
+		for (FD3D12Resource* Resource : SortedResources)
+		{
+			if (Resource == PrevResource) continue; // Skip duplicate resource barriers
+
+			const bool bUseDefaultState = !Resource->RequiresResourceStateTracking();
+
+			D3D12_RESOURCE_STATES DesiredState = D3D12_RESOURCE_STATE_CORRUPT;
+			D3D12_RESOURCE_STATES CurrentState = D3D12_RESOURCE_STATE_CORRUPT;
+			switch (State)
+			{
+			case EBatchCopyState::CopySource:
+				DesiredState = D3D12_RESOURCE_STATE_COPY_SOURCE;
+				CurrentState = bUseDefaultState ? Resource->GetDefaultResourceState() : CurrentState;
+				break;
+			case EBatchCopyState::CopyDest:
+				DesiredState = D3D12_RESOURCE_STATE_COPY_DEST;
+				CurrentState = bUseDefaultState ? Resource->GetDefaultResourceState() : CurrentState;
+				break;
+			case EBatchCopyState::FinalizeSource:
+				CurrentState = D3D12_RESOURCE_STATE_COPY_SOURCE;
+				DesiredState = bUseDefaultState ? Resource->GetDefaultResourceState() : D3D12_RESOURCE_STATE_GENERIC_READ;
+				break;
+			case EBatchCopyState::FinalizeDest:
+				CurrentState = D3D12_RESOURCE_STATE_COPY_DEST;
+				DesiredState = bUseDefaultState ? Resource->GetDefaultResourceState() : D3D12_RESOURCE_STATE_GENERIC_READ;
+				break;
+			default:
+				checkf(false, TEXT("Unexpected batch copy state"));
+				break;
+			}
+
+			if (bUseDefaultState)
+			{
+				check(CurrentState != D3D12_RESOURCE_STATE_CORRUPT);
+				CommandListHandle.AddTransitionBarrier(Resource, CurrentState, DesiredState, Subresource);
+			}
+			else
+			{
+				FD3D12DynamicRHI::TransitionResource(CommandListHandle, Resource, D3D12_RESOURCE_STATE_TBD, DesiredState, Subresource, FD3D12DynamicRHI::ETransitionMode::Validate);
+			}
+
+			PrevResource = Resource;
+		}
+	};
+
+	// Ensure that all previously pending barriers have been processed to avoid incorrect/conflicting transitions for non-tracked resources
+	CommandListHandle.FlushResourceBarriers();
+
+	TransitionResources(CommandListHandle, SrcBuffers, EBatchCopyState::CopySource);
+	TransitionResources(CommandListHandle, DstBuffers, EBatchCopyState::CopyDest);
+
+	// Issue all copy source/dest barriers before performing actual copies
+	CommandListHandle.FlushResourceBarriers();
+
+	for (auto& Param : Params)
+	{
+		FD3D12Buffer* SourceBuffer = RetrieveObject<FD3D12Buffer>(Param.SourceBuffer);
+		FD3D12Buffer* DestBuffer = RetrieveObject<FD3D12Buffer>(Param.DestBuffer);
+		uint64 SrcOffset = Param.SrcOffset;
+		uint64 DstOffset = Param.DstOffset;
+		uint64 NumBytes = Param.NumBytes;
+
+		FD3D12Device* Device = SourceBuffer->GetParentDevice();
+		check(Device == DestBuffer->GetParentDevice());
+
+		FD3D12Resource* pSourceResource = SourceBuffer->ResourceLocation.GetResource();
+		D3D12_RESOURCE_DESC const& SourceBufferDesc = pSourceResource->GetDesc();
+
+		FD3D12Resource* pDestResource = DestBuffer->ResourceLocation.GetResource();
+		D3D12_RESOURCE_DESC const& DestBufferDesc = pDestResource->GetDesc();
+
+		check(DstOffset + NumBytes <= DestBufferDesc.Width);
+		check(SrcOffset + NumBytes <= SourceBufferDesc.Width);
+
+		numCopies++;
+
+		CommandListHandle->CopyBufferRegion(pDestResource->GetResource(), DestBuffer->ResourceLocation.GetOffsetFromBaseOfResource() + DstOffset, pSourceResource->GetResource(), SourceBuffer->ResourceLocation.GetOffsetFromBaseOfResource() + SrcOffset, NumBytes);
+		CommandListHandle.UpdateResidency(pDestResource);
+		CommandListHandle.UpdateResidency(pSourceResource);
+
+		Device->RegisterGPUWork(1);
+	}
+
+	// Transition buffers back to default readable state
+
+	TransitionResources(CommandListHandle, SrcBuffers, EBatchCopyState::FinalizeSource);
+	TransitionResources(CommandListHandle, DstBuffers, EBatchCopyState::FinalizeDest);
+}
+#endif // D3D12_RHI_RAYTRACING
diff --git a/Engine/Source/Runtime/D3D12RHI/Private/D3D12IndexBuffer.cpp b/Engine/Source/Runtime/D3D12RHI/Private/D3D12IndexBuffer.cpp
deleted file mode 100644
diff --git a/Engine/Source/Runtime/D3D12RHI/Private/D3D12StructuredBuffer.cpp b/Engine/Source/Runtime/D3D12RHI/Private/D3D12StructuredBuffer.cpp
deleted file mode 100644
diff --git a/Engine/Source/Runtime/D3D12RHI/Private/D3D12VertexBuffer.cpp b/Engine/Source/Runtime/D3D12RHI/Private/D3D12VertexBuffer.cpp
deleted file mode 100644
diff --git a/Engine/Source/Runtime/OpenGLDrv/Private/OpenGLIndexBuffer.cpp b/Engine/Source/Runtime/OpenGLDrv/Private/OpenGLIndexBuffer.cpp
deleted file mode 100644
diff --git a/Engine/Source/Runtime/OpenGLDrv/Private/OpenGLStructuredBuffer.cpp b/Engine/Source/Runtime/OpenGLDrv/Private/OpenGLStructuredBuffer.cpp
deleted file mode 100644
diff --git a/Engine/Source/Runtime/OpenGLDrv/Private/OpenGLVertexBuffer.cpp b/Engine/Source/Runtime/OpenGLDrv/Private/OpenGLVertexBuffer.cpp
--- a/Engine/Source/Runtime/OpenGLDrv/Private/OpenGLVertexBuffer.cpp
+++ b/Engine/Source/Runtime/OpenGLDrv/Private/OpenGLVertexBuffer.cpp
@@ -265,6 +265,23 @@ void FOpenGLDynamicRHI::RHICopyBufferRegion(FRHIBuffer* DestBufferRHI, uint64 Ds
 	glBindBuffer(GL_COPY_WRITE_BUFFER, 0);
 }
 
+void FOpenGLDynamicRHI::RHITransferBufferUnderlyingResource(FRHIBuffer* DestBuffer, FRHIBuffer* SrcBuffer)
+{
+	VERIFY_GL_SCOPE();
+	check(DestBuffer);
+	FOpenGLBuffer* Dest = ResourceCast(DestBuffer);
+	if (!SrcBuffer)
+	{
+		TRefCountPtr<FOpenGLBuffer> Src = new FOpenGLBuffer();
+		Dest->Swap(*Src);
+	}
+	else
+	{
+		FOpenGLBuffer* Src = ResourceCast(SrcBuffer);
+		Dest->Swap(*Src);
+	}
+}
+
 FStagingBufferRHIRef FOpenGLDynamicRHI::RHICreateStagingBuffer()
 {
 	return new FOpenGLStagingBuffer();
diff --git a/Engine/Source/Runtime/VulkanRHI/Private/VulkanIndexBuffer.cpp b/Engine/Source/Runtime/VulkanRHI/Private/VulkanIndexBuffer.cpp
--- a/Engine/Source/Runtime/VulkanRHI/Private/VulkanIndexBuffer.cpp
+++ b/Engine/Source/Runtime/VulkanRHI/Private/VulkanIndexBuffer.cpp
@@ -449,6 +449,11 @@ void FVulkanDynamicRHI::UnlockBuffer_RenderThread(FRHICommandListImmediate& RHIC
 }
 #endif
 
+void FVulkanDynamicRHI::RHICopyBuffer(FRHIBuffer* SourceBufferRHI, FRHIBuffer* DestBufferRHI)
+{
+	VULKAN_SIGNAL_UNIMPLEMENTED();
+}
+
 void FVulkanDynamicRHI::RHITransferBufferUnderlyingResource(FRHIBuffer* DestBuffer, FRHIBuffer* SrcBuffer)
 {
 	check(DestBuffer);
diff --git a/Engine/Source/Runtime/VulkanRHI/Private/VulkanStructuredBuffer.cpp b/Engine/Source/Runtime/VulkanRHI/Private/VulkanStructuredBuffer.cpp
deleted file mode 100644
diff --git a/Engine/Source/Runtime/VulkanRHI/Private/VulkanVertexBuffer.cpp b/Engine/Source/Runtime/VulkanRHI/Private/VulkanVertexBuffer.cpp
deleted file mode 100644
diff --git a/Engine/Source/Runtime/Windows/D3D11RHI/Private/D3D11IndexBuffer.cpp b/Engine/Source/Runtime/Windows/D3D11RHI/Private/D3D11IndexBuffer.cpp
--- a/Engine/Source/Runtime/Windows/D3D11RHI/Private/D3D11IndexBuffer.cpp
+++ b/Engine/Source/Runtime/Windows/D3D11RHI/Private/D3D11IndexBuffer.cpp
@@ -239,6 +239,24 @@ void FD3D11DynamicRHI::UnlockBuffer_BottomOfPipe(FRHICommandListImmediate& RHICm
 	}
 }
 
+void FD3D11DynamicRHI::RHICopyBuffer(FRHIBuffer* SourceBufferRHI, FRHIBuffer* DestBufferRHI)
+{
+	FD3D11Buffer* SourceBuffer = ResourceCast(SourceBufferRHI);
+	FD3D11Buffer* DestBuffer = ResourceCast(DestBufferRHI);
+
+	D3D11_BUFFER_DESC SourceBufferDesc;
+	SourceBuffer->Resource->GetDesc(&SourceBufferDesc);
+	
+	D3D11_BUFFER_DESC DestBufferDesc;
+	DestBuffer->Resource->GetDesc(&DestBufferDesc);
+
+	check(SourceBufferDesc.ByteWidth == DestBufferDesc.ByteWidth);
+
+	Direct3DDeviceIMContext->CopyResource(DestBuffer->Resource,SourceBuffer->Resource);
+
+	GPUProfilingData.RegisterGPUWork(1);
+}
+
 void FD3D11DynamicRHI::RHITransferBufferUnderlyingResource(FRHIBuffer* DestBuffer, FRHIBuffer* SrcBuffer)
 {
 	check(DestBuffer);
diff --git a/Engine/Source/Runtime/Windows/D3D11RHI/Private/D3D11StructuredBuffer.cpp b/Engine/Source/Runtime/Windows/D3D11RHI/Private/D3D11StructuredBuffer.cpp
deleted file mode 100644
diff --git a/Engine/Source/Runtime/Windows/D3D11RHI/Private/D3D11VertexBuffer.cpp b/Engine/Source/Runtime/Windows/D3D11RHI/Private/D3D11VertexBuffer.cpp
deleted file mode 100644
