Make sure ResourceStateMode is forwarded to the comitted resource creation so it can be correctly tagged as a non tracking resource because state changes are never allowed on RT Accelleration structures (fixes D3D validation errors)

diff --git a/Engine/Source/Runtime/D3D12RHI/Private/D3D12PoolAllocator.cpp b/Engine/Source/Runtime/D3D12RHI/Private/D3D12PoolAllocator.cpp
--- a/Engine/Source/Runtime/D3D12RHI/Private/D3D12PoolAllocator.cpp
+++ b/Engine/Source/Runtime/D3D12RHI/Private/D3D12PoolAllocator.cpp
@@ -307,7 +307,7 @@ void FD3D12PoolAllocator::AllocResource(D3D12_HEAP_TYPE InHeapType, const D3D12_
 	const D3D12_HEAP_PROPERTIES HeapProps = CD3DX12_HEAP_PROPERTIES(InHeapType, GetGPUMask().GetNative(), GetVisibilityMask().GetNative());
 	D3D12_RESOURCE_DESC Desc = InDesc;
 	Desc.Alignment = 0;
-	VERIFYD3D12RESULT(Adapter->CreateCommittedResource(Desc, GetGPUMask(), HeapProps, InCreateState, InClearValue, &NewResource, InName, false));
+	VERIFYD3D12RESULT(Adapter->CreateCommittedResource(Desc, GetGPUMask(), HeapProps, InCreateState, InResourceStateMode, InCreateState, InClearValue, &NewResource, InName, false));
 	
 	ResourceLocation.AsStandAlone(NewResource, InSize);
 }
diff --git a/Engine/Source/Runtime/D3D12RHI/Public/D3D12Resources.h b/Engine/Source/Runtime/D3D12RHI/Public/D3D12Resources.h
--- a/Engine/Source/Runtime/D3D12RHI/Public/D3D12Resources.h
+++ b/Engine/Source/Runtime/D3D12RHI/Public/D3D12Resources.h
@@ -345,7 +345,7 @@ private:
 		{
 #if D3D12_RHI_RAYTRACING
 			// No state tracking for acceleration structures because they can't have another state
-			check(InDefaultState != D3D12_RESOURCE_STATE_RAYTRACING_ACCELERATION_STRUCTURE);
+			check(InDefaultState != D3D12_RESOURCE_STATE_RAYTRACING_ACCELERATION_STRUCTURE && InInitialState != D3D12_RESOURCE_STATE_RAYTRACING_ACCELERATION_STRUCTURE);
 #endif // D3D12_RHI_RAYTRACING
 
 			// Only a few resources (~1%) actually need resource state tracking
