Moving shader model detection around for better platform support.
Adding -sm6 to enable the new feature level from the command line.

diff --git a/Engine/Source/Runtime/D3D12RHI/Private/D3D12Adapter.cpp b/Engine/Source/Runtime/D3D12RHI/Private/D3D12Adapter.cpp
--- a/Engine/Source/Runtime/D3D12RHI/Private/D3D12Adapter.cpp
+++ b/Engine/Source/Runtime/D3D12RHI/Private/D3D12Adapter.cpp
@@ -708,38 +708,6 @@ void FD3D12Adapter::CreateRootDevice(bool bWithDebug)
 #endif
 }
 
-#if PLATFORM_WINDOWS
-static D3D_SHADER_MODEL FindHighestShaderModel(ID3D12Device* Device)
-{
-	// Because we can't guarantee older Windows versions will know about newer shader models, we need to check them all
-	// in descending order and return the first result that succeeds.
-	const D3D_SHADER_MODEL ShaderModelsToCheck[] =
-	{
-		D3D_SHADER_MODEL_6_6,
-		D3D_SHADER_MODEL_6_5,
-		D3D_SHADER_MODEL_6_4,
-		D3D_SHADER_MODEL_6_3,
-		D3D_SHADER_MODEL_6_2,
-		D3D_SHADER_MODEL_6_1,
-		D3D_SHADER_MODEL_6_0,
-	};
-
-	D3D12_FEATURE_DATA_SHADER_MODEL FeatureShaderModel{};
-	for (const D3D_SHADER_MODEL ShaderModelToCheck : ShaderModelsToCheck)
-	{
-		FeatureShaderModel.HighestShaderModel = ShaderModelToCheck;
-
-		if (SUCCEEDED(Device->CheckFeatureSupport(D3D12_FEATURE_SHADER_MODEL, &FeatureShaderModel, sizeof(FeatureShaderModel))))
-		{
-			return FeatureShaderModel.HighestShaderModel;
-		}
-	}
-
-	// Last ditch effort, the minimum requirement for DX12 is 5.1
-	return D3D_SHADER_MODEL_5_1;
-}
-#endif
-
 void FD3D12Adapter::InitializeDevices()
 {
 	check(IsInGameThread());
@@ -855,10 +823,6 @@ void FD3D12Adapter::InitializeDevices()
 		bDepthBoundsTestSupported = !!D3D12Caps2.DepthBoundsTestSupported;
 #endif
 
-#if PLATFORM_WINDOWS
-		HighestShaderModel = FindHighestShaderModel(RootDevice);
-#endif
-
 		D3D12_FEATURE_DATA_ROOT_SIGNATURE D3D12RootSignatureCaps = {};
 		D3D12RootSignatureCaps.HighestVersion = D3D_ROOT_SIGNATURE_VERSION_1_1;	// This is the highest version we currently support. If CheckFeatureSupport succeeds, the HighestVersion returned will not be greater than this.
 		if (FAILED(RootDevice->CheckFeatureSupport(D3D12_FEATURE_ROOT_SIGNATURE, &D3D12RootSignatureCaps, sizeof(D3D12RootSignatureCaps))))
diff --git a/Engine/Source/Runtime/D3D12RHI/Private/D3D12Adapter.h b/Engine/Source/Runtime/D3D12RHI/Private/D3D12Adapter.h
--- a/Engine/Source/Runtime/D3D12RHI/Private/D3D12Adapter.h
+++ b/Engine/Source/Runtime/D3D12RHI/Private/D3D12Adapter.h
@@ -45,16 +45,12 @@ class FD3D12DynamicRHI;
 
 struct FD3D12AdapterDesc
 {
-	FD3D12AdapterDesc()
-		: AdapterIndex(-1)
-		, MaxSupportedFeatureLevel((D3D_FEATURE_LEVEL)0)
-		, NumDeviceNodes(0)
-	{
-	}
+	FD3D12AdapterDesc() = default;
 
-	FD3D12AdapterDesc(DXGI_ADAPTER_DESC& DescIn , int32 InAdapterIndex, D3D_FEATURE_LEVEL InMaxSupportedFeatureLevel, uint32 NumNodes)
+	FD3D12AdapterDesc(const DXGI_ADAPTER_DESC& DescIn , int32 InAdapterIndex, D3D_FEATURE_LEVEL InMaxSupportedFeatureLevel, D3D_SHADER_MODEL InMaxSupportedShaderModel, uint32 NumNodes)
 		: AdapterIndex(InAdapterIndex)
 		, MaxSupportedFeatureLevel(InMaxSupportedFeatureLevel)
+		, MaxSupportedShaderModel(InMaxSupportedShaderModel)
 		, Desc(DescIn)
 		, NumDeviceNodes(NumNodes)
 	{
@@ -63,13 +59,16 @@ struct FD3D12AdapterDesc
 	bool IsValid() const { return MaxSupportedFeatureLevel != (D3D_FEATURE_LEVEL)0 && AdapterIndex >= 0; }
 
 	/** -1 if not supported or FindAdpater() wasn't called. Ideally we would store a pointer to IDXGIAdapter but it's unlikely the adpaters change during engine init. */
-	int32 AdapterIndex;
+	int32 AdapterIndex{ -1 };
+
 	/** The maximum D3D12 feature level supported. 0 if not supported or FindAdpater() wasn't called */
-	D3D_FEATURE_LEVEL MaxSupportedFeatureLevel;
+	D3D_FEATURE_LEVEL MaxSupportedFeatureLevel{ (D3D_FEATURE_LEVEL)0 };
 
-	DXGI_ADAPTER_DESC Desc;
+	D3D_SHADER_MODEL MaxSupportedShaderModel{ D3D_SHADER_MODEL_5_1 };
 
-	uint32 NumDeviceNodes;
+	DXGI_ADAPTER_DESC Desc{};
+
+	uint32 NumDeviceNodes{ 0 };
 };
 
 enum class ED3D12GPUCrashDebugginMode
@@ -96,6 +95,7 @@ public:
 	// Getters
 	FORCEINLINE const uint32 GetAdapterIndex() const { return Desc.AdapterIndex; }
 	FORCEINLINE const D3D_FEATURE_LEVEL GetFeatureLevel() const { return Desc.MaxSupportedFeatureLevel; }
+	FORCEINLINE D3D_SHADER_MODEL GetHighestShaderModel() const { return Desc.MaxSupportedShaderModel; }
 	FORCEINLINE ID3D12Device* GetD3DDevice() const { return RootDevice.GetReference(); }
 	FORCEINLINE ID3D12Device1* GetD3DDevice1() const { return RootDevice1.GetReference(); }
 #if PLATFORM_WINDOWS || PLATFORM_HOLOLENS
@@ -318,8 +318,6 @@ public:
 	void SubmitGapRecorderTimestamps();
 #endif
 
-	FORCEINLINE D3D_SHADER_MODEL GetHighestShaderModel() const { return HighestShaderModel; }
-
 protected:
 
 	virtual void CreateRootDevice(bool bWithDebug);
@@ -356,7 +354,6 @@ protected:
 	D3D12_RESOURCE_HEAP_TIER ResourceHeapTier;
 	D3D12_RESOURCE_BINDING_TIER ResourceBindingTier;
 	D3D_ROOT_SIGNATURE_VERSION RootSignatureVersion;
-	D3D_SHADER_MODEL HighestShaderModel = D3D_SHADER_MODEL_5_1;
 	bool bDepthBoundsTestSupported;
 	bool bHeapNotZeroedSupported;
 
diff --git a/Engine/Source/Runtime/D3D12RHI/Private/Windows/WindowsD3D12Device.cpp b/Engine/Source/Runtime/D3D12RHI/Private/Windows/WindowsD3D12Device.cpp
--- a/Engine/Source/Runtime/D3D12RHI/Private/Windows/WindowsD3D12Device.cpp
+++ b/Engine/Source/Runtime/D3D12RHI/Private/Windows/WindowsD3D12Device.cpp
@@ -195,11 +195,41 @@ static D3D_FEATURE_LEVEL GetRequiredD3DFeatureLevel()
 	return D3D_FEATURE_LEVEL_11_0;
 }
 
+static D3D_SHADER_MODEL FindHighestShaderModel(ID3D12Device* Device)
+{
+	// Because we can't guarantee older Windows versions will know about newer shader models, we need to check them all
+	// in descending order and return the first result that succeeds.
+	const D3D_SHADER_MODEL ShaderModelsToCheck[] =
+	{
+		D3D_SHADER_MODEL_6_6,
+		D3D_SHADER_MODEL_6_5,
+		D3D_SHADER_MODEL_6_4,
+		D3D_SHADER_MODEL_6_3,
+		D3D_SHADER_MODEL_6_2,
+		D3D_SHADER_MODEL_6_1,
+		D3D_SHADER_MODEL_6_0,
+	};
+
+	D3D12_FEATURE_DATA_SHADER_MODEL FeatureShaderModel{};
+	for (const D3D_SHADER_MODEL ShaderModelToCheck : ShaderModelsToCheck)
+	{
+		FeatureShaderModel.HighestShaderModel = ShaderModelToCheck;
+
+		if (SUCCEEDED(Device->CheckFeatureSupport(D3D12_FEATURE_SHADER_MODEL, &FeatureShaderModel, sizeof(FeatureShaderModel))))
+		{
+			return FeatureShaderModel.HighestShaderModel;
+		}
+	}
+
+	// Last ditch effort, the minimum requirement for DX12 is 5.1
+	return D3D_SHADER_MODEL_5_1;
+}
+
 /**
  * Attempts to create a D3D12 device for the adapter using at minimum MinFeatureLevel.
  * If creation is successful, true is returned and the max supported feature level is set in OutMaxFeatureLevel.
  */
-static bool SafeTestD3D12CreateDevice(IDXGIAdapter* Adapter, D3D_FEATURE_LEVEL MinFeatureLevel, D3D_FEATURE_LEVEL& OutMaxFeatureLevel, uint32& OutNumDeviceNodes)
+static bool SafeTestD3D12CreateDevice(IDXGIAdapter* Adapter, D3D_FEATURE_LEVEL MinFeatureLevel, D3D_FEATURE_LEVEL& OutMaxFeatureLevel, D3D_SHADER_MODEL& OutMaxShaderModel, uint32& OutNumDeviceNodes)
 {
 	const D3D_FEATURE_LEVEL FeatureLevels[] =
 	{
@@ -225,6 +255,8 @@ static bool SafeTestD3D12CreateDevice(IDXGIAdapter* Adapter, D3D_FEATURE_LEVEL M
 				MaxFeatureLevel = FeatureLevelCaps.MaxSupportedFeatureLevel;
 			}
 
+			OutMaxShaderModel = FindHighestShaderModel(pDevice);
+
 			OutMaxFeatureLevel = MaxFeatureLevel;
 			OutNumDeviceNodes = pDevice->GetNodeCount();
 
@@ -398,8 +430,10 @@ void FD3D12DynamicRHIModule::FindAdapter()
 		if (TempAdapter)
 		{
 			D3D_FEATURE_LEVEL MaxSupportedFeatureLevel = static_cast<D3D_FEATURE_LEVEL>(0);
+			D3D_SHADER_MODEL MaxSupportedShaderModel = D3D_SHADER_MODEL_5_1;
 			uint32 NumNodes = 0;
-			if (SafeTestD3D12CreateDevice(TempAdapter, MinRequiredFeatureLevel, MaxSupportedFeatureLevel, NumNodes))
+
+			if (SafeTestD3D12CreateDevice(TempAdapter, MinRequiredFeatureLevel, MaxSupportedFeatureLevel, MaxSupportedShaderModel, NumNodes))
 			{
 				check(NumNodes > 0);
 				// Log some information about the available D3D12 adapters.
@@ -448,7 +482,7 @@ void FD3D12DynamicRHIModule::FindAdapter()
 				// PerfHUD is for performance profiling
 				const bool bIsPerfHUD = !FCString::Stricmp(AdapterDesc.Description, TEXT("NVIDIA PerfHUD"));
 
-				FD3D12AdapterDesc CurrentAdapter(AdapterDesc, AdapterIndex, MaxSupportedFeatureLevel, NumNodes);
+				FD3D12AdapterDesc CurrentAdapter(AdapterDesc, AdapterIndex, MaxSupportedFeatureLevel, MaxSupportedShaderModel, NumNodes);
 				
 				// If requested WARP, then reject all other adapters. If WARP not requested, then reject the WARP device.
 				const bool bSkipRequestedWARP = (bRequestedWARP && !bIsWARP) || (!bRequestedWARP && bIsWARP);
@@ -530,8 +564,56 @@ void FD3D12DynamicRHIModule::FindAdapter()
 	}
 }
 
+static bool DoesAnyAdapterSupportSM6(const TArray<TSharedPtr<FD3D12Adapter>>& Adapters)
+{
+	for (const TSharedPtr<FD3D12Adapter>& Adapter : Adapters)
+	{
+		// TODO: D3D_FEATURE_LEVEL_12_2
+		if (Adapter->GetFeatureLevel() >= D3D_FEATURE_LEVEL_12_1 && Adapter->GetHighestShaderModel() >= D3D_SHADER_MODEL_6_5)
+		{
+			return true;
+		}
+	}
+
+	return false;
+}
+
 FDynamicRHI* FD3D12DynamicRHIModule::CreateRHI(ERHIFeatureLevel::Type RequestedFeatureLevel)
 {
+	const bool bAllowSM6 = GAllowShaderModel6
+		|| (RequestedFeatureLevel != ERHIFeatureLevel::Num && RequestedFeatureLevel >= ERHIFeatureLevel::SM6);
+
+	GShaderPlatformForFeatureLevel[ERHIFeatureLevel::ES3_1] = SP_PCD3D_ES3_1;
+	GShaderPlatformForFeatureLevel[ERHIFeatureLevel::SM5] = SP_PCD3D_SM5;
+	if (bAllowSM6)
+	{
+		GShaderPlatformForFeatureLevel[ERHIFeatureLevel::SM6] = SP_PCD3D_SM6;
+	}
+
+	ERHIFeatureLevel::Type PreviewFeatureLevel;
+	if (!GIsEditor && RHIGetPreviewFeatureLevel(PreviewFeatureLevel))
+	{
+		check(PreviewFeatureLevel == ERHIFeatureLevel::ES3_1);
+
+		// ES3.1 feature level emulation in D3D
+		GMaxRHIFeatureLevel = PreviewFeatureLevel;
+		if (GMaxRHIFeatureLevel == ERHIFeatureLevel::ES3_1)
+		{
+			GMaxRHIShaderPlatform = SP_PCD3D_ES3_1;
+		}
+	}
+	else
+	{
+		GMaxRHIFeatureLevel = ERHIFeatureLevel::SM5;
+		GMaxRHIShaderPlatform = SP_PCD3D_SM5;
+
+		if (bAllowSM6 && DoesAnyAdapterSupportSM6(ChosenAdapters))
+		{
+			GMaxRHIFeatureLevel = ERHIFeatureLevel::SM6;
+			GMaxRHIShaderPlatform = SP_PCD3D_SM6;
+		}
+	}
+
 #if USE_PIX
 	bool bPixEventEnabled = (WindowsPixDllHandle != nullptr);
 #else
@@ -955,41 +1037,6 @@ void FD3D12Device::InitPlatformSpecific()
 	CommandListManager = new FD3D12CommandListManager(this, D3D12_COMMAND_LIST_TYPE_DIRECT, ED3D12CommandQueueType::Direct);
 	CopyCommandListManager = new FD3D12CommandListManager(this, D3D12_COMMAND_LIST_TYPE_COPY, ED3D12CommandQueueType::Copy);
 	AsyncCommandListManager = new FD3D12CommandListManager(this, D3D12_COMMAND_LIST_TYPE_COMPUTE, ED3D12CommandQueueType::Async);
-
-	GShaderPlatformForFeatureLevel[ERHIFeatureLevel::ES3_1] = SP_PCD3D_ES3_1;
-	GShaderPlatformForFeatureLevel[ERHIFeatureLevel::SM5] = SP_PCD3D_SM5;
-	if (GAllowShaderModel6)
-	{
-		GShaderPlatformForFeatureLevel[ERHIFeatureLevel::SM6] = SP_PCD3D_SM6;
-	}
-
-	ERHIFeatureLevel::Type PreviewFeatureLevel;
-	if (!GIsEditor && RHIGetPreviewFeatureLevel(PreviewFeatureLevel))
-	{
-		check(PreviewFeatureLevel == ERHIFeatureLevel::ES3_1);
-
-		// ES3.1 feature level emulation in D3D
-		GMaxRHIFeatureLevel = PreviewFeatureLevel;
-		if (GMaxRHIFeatureLevel == ERHIFeatureLevel::ES3_1)
-		{
-			GMaxRHIShaderPlatform = SP_PCD3D_ES3_1;
-		}
-	}
-	else 
-	{
-		GMaxRHIFeatureLevel = ERHIFeatureLevel::SM5;
-		GMaxRHIShaderPlatform = SP_PCD3D_SM5;
-
-		if (GAllowShaderModel6)
-		{
-			// TODO: D3D_FEATURE_LEVEL_12_2
-			if (GetParentAdapter()->GetFeatureLevel() >= D3D_FEATURE_LEVEL_12_1 && GetParentAdapter()->GetHighestShaderModel() >= D3D_SHADER_MODEL_6_5)
-			{
-				GMaxRHIFeatureLevel = ERHIFeatureLevel::SM6;
-				GMaxRHIShaderPlatform = SP_PCD3D_SM6;
-			}
-		}
-	}
 }
 
 void FD3D12Device::CreateSamplerInternal(const D3D12_SAMPLER_DESC& Desc, D3D12_CPU_DESCRIPTOR_HANDLE Descriptor)
diff --git a/Engine/Source/Runtime/RHI/Private/Windows/WindowsDynamicRHI.cpp b/Engine/Source/Runtime/RHI/Private/Windows/WindowsDynamicRHI.cpp
--- a/Engine/Source/Runtime/RHI/Private/Windows/WindowsDynamicRHI.cpp
+++ b/Engine/Source/Runtime/RHI/Private/Windows/WindowsDynamicRHI.cpp
@@ -50,6 +50,7 @@ static IDynamicRHIModule* LoadDynamicRHIModule(ERHIFeatureLevel::Type& DesiredFe
 	}
 
 	bool bForceSM5 = FParse::Param(FCommandLine::Get(), TEXT("sm5"));
+	bool bForceSM6 = FParse::Param(FCommandLine::Get(), TEXT("sm6"));
 	bool bForceVulkan = FParse::Param(FCommandLine::Get(), TEXT("vulkan"));
 	bool bForceD3D11 = FParse::Param(FCommandLine::Get(), TEXT("d3d11")) || FParse::Param(FCommandLine::Get(), TEXT("dx11")) || (bForceSM5 && !bForceVulkan);
 	bool bForceD3D12 = FParse::Param(FCommandLine::Get(), TEXT("d3d12")) || FParse::Param(FCommandLine::Get(), TEXT("dx12"));
@@ -109,7 +110,11 @@ static IDynamicRHIModule* LoadDynamicRHIModule(ERHIFeatureLevel::Type& DesiredFe
 	}
 	else
 	{
-		if (bForceSM5)
+		if (bForceSM6)
+		{
+			DesiredFeatureLevel = ERHIFeatureLevel::SM6;
+		}
+		else if (bForceSM5)
 		{
 			DesiredFeatureLevel = ERHIFeatureLevel::SM5;
 		}
