Moved GRHISupportsRayTracingAMDHitToken and GRHISupportsRayTracingDispatchIndirect initializaation to WindowsD3D12Device.cpp

diff --git a/Engine/Source/Runtime/D3D12RHI/Private/D3D12Adapter.cpp b/Engine/Source/Runtime/D3D12RHI/Private/D3D12Adapter.cpp
--- a/Engine/Source/Runtime/D3D12RHI/Private/D3D12Adapter.cpp
+++ b/Engine/Source/Runtime/D3D12RHI/Private/D3D12Adapter.cpp
@@ -798,15 +798,12 @@ void FD3D12Adapter::InitializeDevices()
 
 					GRHISupportsRayTracing = true;
 
-					GRHISupportsRayTracingAMDHitToken = (OwningRHI->GetAmdSupportedExtensionFlags() & AGS_DX12_EXTENSION_INTRINSIC_RAY_TRACE_HIT_TOKEN) != 0;
-
 					if (D3D12Caps5.RaytracingTier >= D3D12_RAYTRACING_TIER_1_1
 						&& RootDevice7)
 					{
 						UE_LOG(LogD3D12RHI, Log, TEXT("D3D12 ray tracing 1.1 is supported."));
 
 						GRHISupportsRayTracingPSOAdditions = true;
-						GRHISupportsRayTracingDispatchIndirect = true;
 					}
 				}
 			}
diff --git a/Engine/Source/Runtime/D3D12RHI/Private/Windows/WindowsD3D12Device.cpp b/Engine/Source/Runtime/D3D12RHI/Private/Windows/WindowsD3D12Device.cpp
--- a/Engine/Source/Runtime/D3D12RHI/Private/Windows/WindowsD3D12Device.cpp
+++ b/Engine/Source/Runtime/D3D12RHI/Private/Windows/WindowsD3D12Device.cpp
@@ -818,6 +818,11 @@ void FD3D12DynamicRHI::Init()
 
 			UE_LOG(LogD3D12RHI, Warning, TEXT("Ray tracing is disabled because the driver is too old"));
 		}
+
+		if (GRHISupportsRayTracing)
+		{
+			GRHISupportsRayTracingAMDHitToken = (GetAmdSupportedExtensionFlags() & AGS_DX12_EXTENSION_INTRINSIC_RAY_TRACE_HIT_TOKEN) != 0;
+		}
 	}
 #endif // AMD_API_ENABLE
 
@@ -967,6 +972,18 @@ void FD3D12DynamicRHI::Init()
 		GVariableRateShadingImageTileSize 	= 1;
 	}
 
+	if (GRHISupportsRayTracing)
+	{
+		D3D12_FEATURE_DATA_D3D12_OPTIONS5 D3D12Caps5 = {};
+		if (SUCCEEDED(GetAdapter().GetD3DDevice()->CheckFeatureSupport(D3D12_FEATURE_D3D12_OPTIONS5, &D3D12Caps5, sizeof(D3D12Caps5))))
+		{
+			if (D3D12Caps5.RaytracingTier >= D3D12_RAYTRACING_TIER_1_1)
+			{
+				GRHISupportsRayTracingDispatchIndirect = true;
+			}
+		}
+	}
+
 	// Command lists need the validation RHI context if enabled, so call the global scope version of RHIGetDefaultContext() and RHIGetDefaultAsyncComputeContext().
 	GRHICommandList.GetImmediateCommandList().SetContext(::RHIGetDefaultContext());
 	GRHICommandList.GetImmediateAsyncComputeCommandList().SetComputeContext(::RHIGetDefaultAsyncComputeContext());
