Moved platform-specific D3D12 command signature initialization to platform-specific files.

diff --git a/Engine/Source/Runtime/D3D12RHI/Private/D3D12Adapter.cpp b/Engine/Source/Runtime/D3D12RHI/Private/D3D12Adapter.cpp
--- a/Engine/Source/Runtime/D3D12RHI/Private/D3D12Adapter.cpp
+++ b/Engine/Source/Runtime/D3D12RHI/Private/D3D12Adapter.cpp
@@ -834,7 +834,7 @@ void FD3D12Adapter::InitializeDevices()
 		StagingFence = new FD3D12Fence(this, FRHIGPUMask::All(), L"Staging Fence");
 		StagingFence->CreateFence();
 
-		CreateSignatures();
+		CreateCommandSignatures();
 
 		// Context redirectors allow RHI commands to be executed on multiple GPUs at the
 		// same time in a multi-GPU system. Redirectors have a physical mask for the GPUs
@@ -905,7 +905,7 @@ void FD3D12Adapter::InitializeRayTracing()
 #endif // D3D12_RHI_RAYTRACING
 }
 
-void FD3D12Adapter::CreateSignatures()
+void FD3D12Adapter::CreateCommandSignatures()
 {
 	ID3D12Device* Device = GetD3DDevice();
 
@@ -930,26 +930,7 @@ void FD3D12Adapter::CreateSignatures()
 	commandSignatureDesc.ByteStride = sizeof(D3D12_DISPATCH_ARGUMENTS);
 	VERIFYD3D12RESULT(Device->CreateCommandSignature(&commandSignatureDesc, nullptr, IID_PPV_ARGS(DispatchIndirectGraphicsCommandSignature.GetInitReference())));
 
-#if !PLATFORM_WINDOWS && !PLATFORM_HOLOLENS
-	indirectParameterDesc[0].Type = D3D12XBOX_INDIRECT_ARGUMENT_TYPE_DISPATCH_L2; // D3D12XBOX_INDIRECT_ARGUMENT_TYPE_DISPATCH_SINGLE
-#endif
-	VERIFYD3D12RESULT(Device->CreateCommandSignature(&commandSignatureDesc, nullptr, IID_PPV_ARGS(DispatchIndirectComputeCommandSignature.GetInitReference())));
-
-#if D3D12_RHI_RAYTRACING
-	if (GRHISupportsRayTracingDispatchIndirect)
-	{
-		D3D12_COMMAND_SIGNATURE_DESC SignatureDesc = {};
-		SignatureDesc.NumArgumentDescs = 1;
-		SignatureDesc.ByteStride = sizeof(D3D12_DISPATCH_RAYS_DESC);
-		SignatureDesc.NodeMask = FRHIGPUMask::All().GetNative();
-
-		D3D12_INDIRECT_ARGUMENT_DESC ArgumentDesc[1] = {};
-		ArgumentDesc[0].Type = D3D12_INDIRECT_ARGUMENT_TYPE_DISPATCH_RAYS;
-		SignatureDesc.pArgumentDescs = ArgumentDesc;
-
-		VERIFYD3D12RESULT(Device->CreateCommandSignature(&SignatureDesc, nullptr, IID_PPV_ARGS(DispatchRaysIndirectCommandSignature.GetInitReference())));
-	}
-#endif // D3D12_RHI_RAYTRACING
+	checkf(DispatchIndirectComputeCommandSignature.IsValid(), TEXT("Indirect compute dispatch command signature is expected to be created by platform-specific D3D12 adapter implementation."))
 }
 
 
diff --git a/Engine/Source/Runtime/D3D12RHI/Private/D3D12Adapter.h b/Engine/Source/Runtime/D3D12RHI/Private/D3D12Adapter.h
--- a/Engine/Source/Runtime/D3D12RHI/Private/D3D12Adapter.h
+++ b/Engine/Source/Runtime/D3D12RHI/Private/D3D12Adapter.h
@@ -334,7 +334,7 @@ protected:
 		FD3D12ResourceLocation& ResourceLocation);
 
 	// Creates default root and execute indirect signatures
-	void CreateSignatures();
+	virtual void CreateCommandSignatures();
 
 	FD3D12DynamicRHI* OwningRHI;
 
diff --git a/Engine/Source/Runtime/D3D12RHI/Private/Windows/WindowsD3D12Adapter.h b/Engine/Source/Runtime/D3D12RHI/Private/Windows/WindowsD3D12Adapter.h
new file mode 100644
--- /dev/null
+++ b/Engine/Source/Runtime/D3D12RHI/Private/Windows/WindowsD3D12Adapter.h
@@ -0,0 +1,20 @@
+// Copyright Epic Games, Inc. All Rights Reserved.
+
+#pragma once
+
+#include "D3D12Adapter.h"
+
+class FWindowsD3D12Adapter : public FD3D12Adapter
+{
+public:
+
+	FWindowsD3D12Adapter(FD3D12AdapterDesc& InDesc)
+		: FD3D12Adapter(InDesc)
+	{}
+
+protected:
+
+	virtual void CreateCommandSignatures() final override;
+
+};
+
diff --git a/Engine/Source/Runtime/D3D12RHI/Private/Windows/WindowsD3D12Device.cpp b/Engine/Source/Runtime/D3D12RHI/Private/Windows/WindowsD3D12Device.cpp
--- a/Engine/Source/Runtime/D3D12RHI/Private/Windows/WindowsD3D12Device.cpp
+++ b/Engine/Source/Runtime/D3D12RHI/Private/Windows/WindowsD3D12Device.cpp
@@ -5,6 +5,7 @@
 =============================================================================*/
 
 #include "D3D12RHIPrivate.h"
+#include "WindowsD3D12Adapter.h"
 #include "Modules/ModuleManager.h"
 #include "Windows/AllowWindowsPlatformTypes.h"
 #include <delayimp.h>
@@ -539,18 +540,18 @@ void FD3D12DynamicRHIModule::FindAdapter()
 		// We assume Intel is integrated graphics (slower than discrete) than NVIDIA or AMD cards and rather take a different one
 		if (FirstWithoutIntegratedAdapter.IsValid())
 		{
-			NewAdapter = TSharedPtr<FD3D12Adapter>(new FD3D12Adapter(FirstWithoutIntegratedAdapter));
+			NewAdapter = TSharedPtr<FD3D12Adapter>(new FWindowsD3D12Adapter(FirstWithoutIntegratedAdapter));
 			ChosenAdapters.Add(NewAdapter);
 		}
 		else
 		{
-			NewAdapter = TSharedPtr<FD3D12Adapter>(new FD3D12Adapter(FirstAdapter));
+			NewAdapter = TSharedPtr<FD3D12Adapter>(new FWindowsD3D12Adapter(FirstAdapter));
 			ChosenAdapters.Add(NewAdapter);
 		}
 	}
 	else
 	{
-		NewAdapter = TSharedPtr<FD3D12Adapter>(new FD3D12Adapter(FirstAdapter));
+		NewAdapter = TSharedPtr<FD3D12Adapter>(new FWindowsD3D12Adapter(FirstAdapter));
 		ChosenAdapters.Add(NewAdapter);
 	}
 
@@ -972,18 +973,6 @@ void FD3D12DynamicRHI::Init()
 		GVariableRateShadingImageTileSize 	= 1;
 	}
 
-	if (GRHISupportsRayTracing)
-	{
-		D3D12_FEATURE_DATA_D3D12_OPTIONS5 D3D12Caps5 = {};
-		if (SUCCEEDED(GetAdapter().GetD3DDevice()->CheckFeatureSupport(D3D12_FEATURE_D3D12_OPTIONS5, &D3D12Caps5, sizeof(D3D12Caps5))))
-		{
-			if (D3D12Caps5.RaytracingTier >= D3D12_RAYTRACING_TIER_1_1)
-			{
-				GRHISupportsRayTracingDispatchIndirect = true;
-			}
-		}
-	}
-
 	// Command lists need the validation RHI context if enabled, so call the global scope version of RHIGetDefaultContext() and RHIGetDefaultAsyncComputeContext().
 	GRHICommandList.GetImmediateCommandList().SetContext(::RHIGetDefaultContext());
 	GRHICommandList.GetImmediateAsyncComputeCommandList().SetComputeContext(::RHIGetDefaultAsyncComputeContext());
@@ -1217,3 +1206,59 @@ bool FD3D12DynamicRHI::RHIGetAvailableResolutions(FScreenResolutionArray& Resolu
 
 	return true;
 }
+
+void FWindowsD3D12Adapter::CreateCommandSignatures()
+{
+	ID3D12Device* Device = GetD3DDevice();
+
+#if D3D12_RHI_RAYTRACING
+
+	if (GRHISupportsRayTracing)
+	{
+		D3D12_FEATURE_DATA_D3D12_OPTIONS5 D3D12Caps5 = {};
+		if (SUCCEEDED(Device->CheckFeatureSupport(D3D12_FEATURE_D3D12_OPTIONS5, &D3D12Caps5, sizeof(D3D12Caps5))))
+		{
+			if (D3D12Caps5.RaytracingTier >= D3D12_RAYTRACING_TIER_1_1)
+			{
+				GRHISupportsRayTracingDispatchIndirect = true;
+			}
+		}
+	}
+
+	if (GRHISupportsRayTracingDispatchIndirect)
+	{
+		D3D12_COMMAND_SIGNATURE_DESC SignatureDesc = {};
+		SignatureDesc.NumArgumentDescs = 1;
+		SignatureDesc.ByteStride = sizeof(D3D12_DISPATCH_RAYS_DESC);
+		SignatureDesc.NodeMask = FRHIGPUMask::All().GetNative();
+
+		D3D12_INDIRECT_ARGUMENT_DESC ArgumentDesc[1] = {};
+		ArgumentDesc[0].Type = D3D12_INDIRECT_ARGUMENT_TYPE_DISPATCH_RAYS;
+		SignatureDesc.pArgumentDescs = ArgumentDesc;
+
+		checkf(DispatchRaysIndirectCommandSignature == nullptr, TEXT("Indirect ray tracing dispatch command signature is expected to be initialized by FWindowsD3D12Adapter."));
+		VERIFYD3D12RESULT(Device->CreateCommandSignature(&SignatureDesc, nullptr, IID_PPV_ARGS(DispatchRaysIndirectCommandSignature.GetInitReference())));
+	}
+
+#endif // D3D12_RHI_RAYTRACING
+
+	// Create windows-specific indirect compute dispatch command signature
+	{
+		D3D12_COMMAND_SIGNATURE_DESC CommandSignatureDesc = {};
+		CommandSignatureDesc.NumArgumentDescs = 1;
+		CommandSignatureDesc.ByteStride = sizeof(D3D12_DISPATCH_ARGUMENTS);
+		CommandSignatureDesc.NodeMask = FRHIGPUMask::All().GetNative();
+
+		D3D12_INDIRECT_ARGUMENT_DESC IndirectParameterDesc[1] = {};
+		IndirectParameterDesc[0].Type = D3D12_INDIRECT_ARGUMENT_TYPE_DISPATCH;
+		CommandSignatureDesc.pArgumentDescs = IndirectParameterDesc;
+
+		checkf(DispatchIndirectComputeCommandSignature == nullptr, TEXT("Indirect compute dispatch command signature is expected to be initialized by FWindowsD3D12Adapter."));
+		VERIFYD3D12RESULT(Device->CreateCommandSignature(&CommandSignatureDesc, nullptr, IID_PPV_ARGS(DispatchIndirectComputeCommandSignature.GetInitReference())));
+	}
+
+	// Create all the generic / cross-platform command signatures
+
+	FD3D12Adapter::CreateCommandSignatures();
+}
+
