Selectively allow/disallow Pascal cards to use ray tracing features.

r.D3D12.DXR.AllowEmulatedRayTracing (default=0) determines whether Pascal cards can use ray tracing features (without hardware acceleration).

diff --git a/Engine/Source/Runtime/D3D12RHI/Private/Windows/WindowsD3D12Device.cpp b/Engine/Source/Runtime/D3D12RHI/Private/Windows/WindowsD3D12Device.cpp
--- a/Engine/Source/Runtime/D3D12RHI/Private/Windows/WindowsD3D12Device.cpp
+++ b/Engine/Source/Runtime/D3D12RHI/Private/Windows/WindowsD3D12Device.cpp
@@ -102,6 +102,14 @@ static FAutoConsoleVariableRef CVarMinDriverVersionForRayTracingNVIDIA(
 	ECVF_ReadOnly | ECVF_RenderThreadSafe
 );
 
+int32 GAllowEmulatedRayTracing = 0;
+static FAutoConsoleVariableRef CVarAllowEmulatedRayTracing(
+	TEXT("r.D3D12.DXR.AllowEmulatedRayTracing"),
+	GAllowEmulatedRayTracing,
+	TEXT("Allows ray tracing emulation support on NVIDIA cards with the Pascal architecture (default=0)."),
+	ECVF_ReadOnly | ECVF_RenderThreadSafe
+);
+
 // Use AGS_MAKE_VERSION() macro to define the version.
 // i.e. AGS_MAKE_VERSION(major, minor, patch) ((major << 22) | (minor << 12) | patch)
 int32 GMinimumDriverVersionForRayTracingAMD = 0;
@@ -679,6 +687,45 @@ void FD3D12DynamicRHIModule::ShutdownModule()
 #endif
 }
 
+#if D3D12_RHI_RAYTRACING
+
+bool IsRayTracingEmulated(uint32 DeviceId)
+{
+	uint32 EmulatedRayTracingIds[] =
+	{
+		0x1B80, // "NVIDIA GeForce GTX 1080"
+		0x1B81, // "NVIDIA GeForce GTX 1070"
+		0x1B82, // "NVIDIA GeForce GTX 1070 Ti"
+		0x1B83, // "NVIDIA GeForce GTX 1060 6GB"
+		0x1B84, // "NVIDIA GeForce GTX 1060 3GB"
+		0x1C01, // "NVIDIA GeForce GTX 1050 Ti"
+		0x1C02, // "NVIDIA GeForce GTX 1060 3GB"
+		0x1C03, // "NVIDIA GeForce GTX 1060 6GB"
+		0x1C04, // "NVIDIA GeForce GTX 1060 5GB"
+		0x1C06, // "NVIDIA GeForce GTX 1060 6GB"
+		0x1C08, // "NVIDIA GeForce GTX 1050"
+		0x1C81, // "NVIDIA GeForce GTX 1050"
+		0x1C82, // "NVIDIA GeForce GTX 1050 Ti"
+		0x1C83, // "NVIDIA GeForce GTX 1050"
+		0x1B06, // "NVIDIA GeForce GTX 1080 Ti"
+		0x1B00, // "NVIDIA TITAN X (Pascal)"
+		0x1B02, // "NVIDIA TITAN Xp"
+		0x1D81, // "NVIDIA TITAN V"
+	};
+
+	for (int Index = 0; Index < UE_ARRAY_COUNT(EmulatedRayTracingIds); ++Index)
+	{
+		if (DeviceId == EmulatedRayTracingIds[Index])
+		{
+			return true;
+		}
+	}
+
+	return false;
+}
+
+#endif // D3D12_RHI_RAYTRACING
+
 void FD3D12DynamicRHI::Init()
 {
 	for (TSharedPtr<FD3D12Adapter>& Adapter : ChosenAdapters)
@@ -803,6 +850,15 @@ void FD3D12DynamicRHI::Init()
 				}
 			}
 		}
+
+		// Disable ray tracing emulation for Pascal architecture
+		if (GRHISupportsRayTracing
+			&& !GAllowEmulatedRayTracing
+			&& IsRayTracingEmulated(AdapterDesc.DeviceId))
+		{
+			GRHISupportsRayTracing = false;
+			UE_LOG(LogD3D12RHI, Warning, TEXT("Ray tracing is disabled for NVIDIA cards with the Pascal architecture. This can be overridden with the following CVar: r.D3D12.DXR.AllowEmulatedRayTracing=1"));
+		}
 	}
 #endif
 
