- Make sure there is no pending copy operation on a deleted resource
- Add extra validation during building of copy ops

diff --git a/Engine/Source/Runtime/D3D12RHI/Private/D3D12PoolAllocator.cpp b/Engine/Source/Runtime/D3D12RHI/Private/D3D12PoolAllocator.cpp
--- a/Engine/Source/Runtime/D3D12RHI/Private/D3D12PoolAllocator.cpp
+++ b/Engine/Source/Runtime/D3D12RHI/Private/D3D12PoolAllocator.cpp
@@ -338,6 +338,18 @@ void FD3D12PoolAllocator::Deallocate(FD3D12ResourceLocation& ResourceLocation)
 				break;
 			}
 		}
+
+		// If still pending copy then clear the copy operation data
+		for (FD3D12VRAMCopyOperation& CopyOperation : PendingCopyOps)
+		{
+			if (CopyOperation.SourceResource == ResourceLocation.GetResource() ||
+				CopyOperation.DestResource == ResourceLocation.GetResource())
+			{
+				CopyOperation.SourceResource = nullptr;
+				CopyOperation.DestResource = nullptr;
+				break;
+			}
+		}
 	}
 
 	int16 PoolIndex = AllocationData.GetPoolIndex();
@@ -418,6 +430,8 @@ bool FD3D12PoolAllocator::HandleDefragRequest(FRHIPoolAllocationData* InSourceBl
 	CopyOp.DestOffset		= Owner->GetOffsetFromBaseOfResource();
 	CopyOp.Size				= InSourceBlock->GetSize();
 	CopyOp.CopyType			= AllocationStrategy == EResourceAllocationStrategy::kManualSubAllocation ? FD3D12VRAMCopyOperation::ECopyType::BufferRegion : FD3D12VRAMCopyOperation::ECopyType::Resource;
+	check(CopyOp.SourceResource != nullptr);
+	check(CopyOp.DestResource != nullptr);
 	PendingCopyOps.Add(CopyOp);
 
 	// TODO: Using aliasing buffer on whole heap for copies to reduce flushes and resource transitions
@@ -550,6 +564,12 @@ void FD3D12PoolAllocator::FlushPendingCopyOps(FD3D12CommandContext& InCommandCon
 
 	for (FD3D12VRAMCopyOperation& CopyOperation : PendingCopyOps)
 	{		
+		// Cleared copy op?
+		if (CopyOperation.SourceResource == nullptr || CopyOperation.DestResource == nullptr)
+		{
+			continue;
+		}
+
 		bool bRTAccelerationStructure = false;
 		if (CopyOperation.SourceResource->RequiresResourceStateTracking())
 		{
@@ -581,8 +601,9 @@ void FD3D12PoolAllocator::FlushPendingCopyOps(FD3D12CommandContext& InCommandCon
 #if D3D12_RHI_RAYTRACING
 		if (bRTAccelerationStructure)
 		{
-			CommandListHandle.GraphicsCommandList5()->CopyRaytracingAccelerationStructure(CopyOperation.DestResource->GetResource()->GetGPUVirtualAddress() + CopyOperation.DestOffset,
-				CopyOperation.SourceResource->GetResource()->GetGPUVirtualAddress() + CopyOperation.SourceOffset, D3D12_RAYTRACING_ACCELERATION_STRUCTURE_COPY_MODE_CLONE);
+			D3D12_GPU_VIRTUAL_ADDRESS SrcAddress = CopyOperation.SourceResource->GetResource()->GetGPUVirtualAddress() + CopyOperation.SourceOffset;
+			D3D12_GPU_VIRTUAL_ADDRESS DestAddress = CopyOperation.DestResource->GetResource()->GetGPUVirtualAddress() + CopyOperation.DestOffset;
+			CommandListHandle.GraphicsCommandList5()->CopyRaytracingAccelerationStructure(DestAddress, SrcAddress, D3D12_RAYTRACING_ACCELERATION_STRUCTURE_COPY_MODE_CLONE);
 		}
 		else
 #endif // D3D12_RHI_RAYTRACING
