Fixed compile error due to IsRayTracingEmulated() being compiled out on some platforms

diff --git a/Engine/Source/Runtime/D3D12RHI/Private/Windows/WindowsD3D12Device.cpp b/Engine/Source/Runtime/D3D12RHI/Private/Windows/WindowsD3D12Device.cpp
--- a/Engine/Source/Runtime/D3D12RHI/Private/Windows/WindowsD3D12Device.cpp
+++ b/Engine/Source/Runtime/D3D12RHI/Private/Windows/WindowsD3D12Device.cpp
@@ -687,9 +687,7 @@ void FD3D12DynamicRHIModule::ShutdownModule()
 #endif
 }
 
-#if D3D12_RHI_RAYTRACING
-
-bool IsRayTracingEmulated(uint32 DeviceId)
+static bool IsRayTracingEmulated(uint32 DeviceId)
 {
 	uint32 EmulatedRayTracingIds[] =
 	{
@@ -724,8 +722,6 @@ bool IsRayTracingEmulated(uint32 DeviceId)
 	return false;
 }
 
-#endif // D3D12_RHI_RAYTRACING
-
 void FD3D12DynamicRHI::Init()
 {
 	for (TSharedPtr<FD3D12Adapter>& Adapter : ChosenAdapters)
