Use CmdList4 instead 5 for raytracing BVH copy because 5 might not be available when RT is on because you need Win 1903+ for that

diff --git a/Engine/Source/Runtime/D3D12RHI/Private/D3D12PoolAllocator.cpp b/Engine/Source/Runtime/D3D12RHI/Private/D3D12PoolAllocator.cpp
--- a/Engine/Source/Runtime/D3D12RHI/Private/D3D12PoolAllocator.cpp
+++ b/Engine/Source/Runtime/D3D12RHI/Private/D3D12PoolAllocator.cpp
@@ -603,7 +603,7 @@ void FD3D12PoolAllocator::FlushPendingCopyOps(FD3D12CommandContext& InCommandCon
 		{
 			D3D12_GPU_VIRTUAL_ADDRESS SrcAddress = CopyOperation.SourceResource->GetResource()->GetGPUVirtualAddress() + CopyOperation.SourceOffset;
 			D3D12_GPU_VIRTUAL_ADDRESS DestAddress = CopyOperation.DestResource->GetResource()->GetGPUVirtualAddress() + CopyOperation.DestOffset;
-			CommandListHandle.GraphicsCommandList5()->CopyRaytracingAccelerationStructure(DestAddress, SrcAddress, D3D12_RAYTRACING_ACCELERATION_STRUCTURE_COPY_MODE_CLONE);
+			CommandListHandle.RayTracingCommandList()->CopyRaytracingAccelerationStructure(DestAddress, SrcAddress, D3D12_RAYTRACING_ACCELERATION_STRUCTURE_COPY_MODE_CLONE);
 		}
 		else
 #endif // D3D12_RHI_RAYTRACING
