Remove invalid check in FD3D12LinearBatchedQueryPool::GetResolvedBatchResults()

  - waiting with no query batches outstanding immediately returns which is fine

diff --git a/Engine/Source/Runtime/D3D12RHI/Private/D3D12TimedIntervalQuery.cpp b/Engine/Source/Runtime/D3D12RHI/Private/D3D12TimedIntervalQuery.cpp
--- a/Engine/Source/Runtime/D3D12RHI/Private/D3D12TimedIntervalQuery.cpp
+++ b/Engine/Source/Runtime/D3D12RHI/Private/D3D12TimedIntervalQuery.cpp
@@ -145,7 +145,6 @@ public:
 			}
 		}
 
-		check(!bWait);
 		return false;
 	}
 
