Reduce template madness a bit

diff --git a/Engine/Source/Runtime/D3D12RHI/Private/D3D12Device.h b/Engine/Source/Runtime/D3D12RHI/Private/D3D12Device.h
--- a/Engine/Source/Runtime/D3D12RHI/Private/D3D12Device.h
+++ b/Engine/Source/Runtime/D3D12RHI/Private/D3D12Device.h
@@ -232,3 +232,31 @@ void TD3D12ViewDescriptorHandle<TDesc>::FreeDescriptorSlot()
 	}
 	check(!Handle.ptr);
 }
+
+inline void FD3D12DescriptorHandleSRV::CreateView(const D3D12_SHADER_RESOURCE_VIEW_DESC& Desc, ID3D12Resource* Resource)
+{
+#if D3D12_RHI_RAYTRACING
+	// NOTE (from D3D Debug runtime): When ViewDimension is D3D12_SRV_DIMENSION_RAYTRACING_ACCELLERATION_STRUCTURE, pResource must be NULL, since the resource location comes from a GPUVA in pDesc
+	if (Desc.ViewDimension == D3D12_SRV_DIMENSION_RAYTRACING_ACCELERATION_STRUCTURE)
+	{
+		Resource = nullptr;
+	}
+#endif // D3D12_RHI_RAYTRACING
+
+	GetParentDevice()->GetDevice()->CreateShaderResourceView(Resource, &Desc, Handle);
+}
+
+inline void FD3D12DescriptorHandleRTV::CreateView(const D3D12_RENDER_TARGET_VIEW_DESC& Desc, ID3D12Resource* Resource)
+{
+	GetParentDevice()->GetDevice()->CreateRenderTargetView(Resource, &Desc, Handle);
+}
+
+inline void FD3D12DescriptorHandleDSV::CreateView(const D3D12_DEPTH_STENCIL_VIEW_DESC& Desc, ID3D12Resource* Resource)
+{
+	GetParentDevice()->GetDevice()->CreateDepthStencilView(Resource, &Desc, Handle);
+}
+
+inline void FD3D12DescriptorHandleUAV::CreateViewWithCounter(const D3D12_UNORDERED_ACCESS_VIEW_DESC& Desc, ID3D12Resource* Resource, ID3D12Resource* CounterResource)
+{
+	GetParentDevice()->GetDevice()->CreateUnorderedAccessView(Resource, CounterResource, &Desc, Handle);
+}
diff --git a/Engine/Source/Runtime/D3D12RHI/Private/D3D12View.h b/Engine/Source/Runtime/D3D12RHI/Private/D3D12View.h
--- a/Engine/Source/Runtime/D3D12RHI/Private/D3D12View.h
+++ b/Engine/Source/Runtime/D3D12RHI/Private/D3D12View.h
@@ -621,27 +621,9 @@ inline uint32 CViewSubresourceSubset::MaxSubresource() const
 
 class FD3D12OfflineDescriptorManager;
 
-template<typename TDesc> 
-struct TIsD3D12SRVDescriptorHandle
-{
-	enum { Value = false };
-};
-
-template<>
-struct TIsD3D12SRVDescriptorHandle<D3D12_SHADER_RESOURCE_VIEW_DESC>
-{
-	enum { Value = true };
-};
-
 template <typename TDesc>
 class TD3D12ViewDescriptorHandle : public FD3D12DeviceChild
 {
-	template <typename TDesc> struct TCreateViewMap;
-	template<> struct TCreateViewMap<D3D12_SHADER_RESOURCE_VIEW_DESC>	{ static decltype(&ID3D12Device::CreateShaderResourceView)	GetCreate()	{ return &ID3D12Device::CreateShaderResourceView;	} };
-	template<> struct TCreateViewMap<D3D12_RENDER_TARGET_VIEW_DESC>		{ static decltype(&ID3D12Device::CreateRenderTargetView)	GetCreate()	{ return &ID3D12Device::CreateRenderTargetView;		} };
-	template<> struct TCreateViewMap<D3D12_DEPTH_STENCIL_VIEW_DESC>		{ static decltype(&ID3D12Device::CreateDepthStencilView)	GetCreate()	{ return &ID3D12Device::CreateDepthStencilView;		} };
-	template<> struct TCreateViewMap<D3D12_UNORDERED_ACCESS_VIEW_DESC>	{ static decltype(&ID3D12Device::CreateUnorderedAccessView)	GetCreate()	{ return &ID3D12Device::CreateUnorderedAccessView;	} };
-
 	CD3DX12_CPU_DESCRIPTOR_HANDLE Handle;
 	uint32 Index;
 
@@ -665,23 +647,9 @@ public:
 		AllocateDescriptorSlot();
 	}
 
-	void CreateView(const TDesc& Desc, ID3D12Resource* Resource)
-	{
-#if D3D12_RHI_RAYTRACING
-		// NOTE (from D3D Debug runtime): When ViewDimension is D3D12_SRV_DIMENSION_RAYTRACING_ACCELLERATION_STRUCTURE, pResource must be NULL, since the resource location comes from a GPUVA in pDesc
-		if (TIsD3D12SRVDescriptorHandle<TDesc>::Value && ((int)Desc.ViewDimension == (int)D3D12_SRV_DIMENSION_RAYTRACING_ACCELERATION_STRUCTURE))
-		{
-			Resource = nullptr;
-		}
-#endif // D3D12_RHI_RAYTRACING
-
-		(GetParentDevice()->GetDevice()->*TCreateViewMap<TDesc>::GetCreate()) (Resource, &Desc, Handle);
-	}
-
-	void CreateViewWithCounter(const TDesc& Desc, ID3D12Resource* Resource, ID3D12Resource* CounterResource)
-	{
-		(GetParentDevice()->GetDevice()->*TCreateViewMap<TDesc>::GetCreate()) (Resource, CounterResource, &Desc, Handle);
-	}
+	// Implemented in D3D12Device.h due to dependencies on FD3D12Device
+	inline void CreateView(const TDesc& Desc, ID3D12Resource* Resource);
+	inline void CreateViewWithCounter(const TDesc& Desc, ID3D12Resource* Resource, ID3D12Resource* CounterResource);
 
 	inline const CD3DX12_CPU_DESCRIPTOR_HANDLE& GetHandle() const { return Handle; }
 	inline uint32 GetIndex() const { return Index; }
