Dx12 fix

diff --git a/Engine/Source/Runtime/D3D12RHI/Private/D3D12Device.h b/Engine/Source/Runtime/D3D12RHI/Private/D3D12Device.h
--- a/Engine/Source/Runtime/D3D12RHI/Private/D3D12Device.h
+++ b/Engine/Source/Runtime/D3D12RHI/Private/D3D12Device.h
@@ -233,6 +233,7 @@ void TD3D12ViewDescriptorHandle<TDesc>::FreeDescriptorSlot()
 	check(!Handle.ptr);
 }
 
+template<>
 inline void FD3D12DescriptorHandleSRV::CreateView(const D3D12_SHADER_RESOURCE_VIEW_DESC& Desc, ID3D12Resource* Resource)
 {
 #if D3D12_RHI_RAYTRACING
@@ -246,16 +247,19 @@ inline void FD3D12DescriptorHandleSRV::CreateView(const D3D12_SHADER_RESOURCE_VI
 	GetParentDevice()->GetDevice()->CreateShaderResourceView(Resource, &Desc, Handle);
 }
 
+template<>
 inline void FD3D12DescriptorHandleRTV::CreateView(const D3D12_RENDER_TARGET_VIEW_DESC& Desc, ID3D12Resource* Resource)
 {
 	GetParentDevice()->GetDevice()->CreateRenderTargetView(Resource, &Desc, Handle);
 }
 
+template<>
 inline void FD3D12DescriptorHandleDSV::CreateView(const D3D12_DEPTH_STENCIL_VIEW_DESC& Desc, ID3D12Resource* Resource)
 {
 	GetParentDevice()->GetDevice()->CreateDepthStencilView(Resource, &Desc, Handle);
 }
 
+template<>
 inline void FD3D12DescriptorHandleUAV::CreateViewWithCounter(const D3D12_UNORDERED_ACCESS_VIEW_DESC& Desc, ID3D12Resource* Resource, ID3D12Resource* CounterResource)
 {
 	GetParentDevice()->GetDevice()->CreateUnorderedAccessView(Resource, CounterResource, &Desc, Handle);
