Avoid redundant copying of FRayTracingSceneInitializer::Instances array in FD3D12RayTracingScene.

Generate the D3D12_RAYTRACING_INSTANCE_DESC right away on scene creation, using ParallelFor.
Significantly reduces number of memory allocations in large scenes.

diff --git a/Engine/Source/Runtime/D3D12RHI/Private/D3D12RayTracing.cpp b/Engine/Source/Runtime/D3D12RHI/Private/D3D12RayTracing.cpp
--- a/Engine/Source/Runtime/D3D12RHI/Private/D3D12RayTracing.cpp
+++ b/Engine/Source/Runtime/D3D12RHI/Private/D3D12RayTracing.cpp
@@ -113,7 +113,6 @@ DECLARE_DWORD_COUNTER_STAT(TEXT("Used view descriptors (per frame)"), STAT_D3D12
 DECLARE_DWORD_COUNTER_STAT(TEXT("Built BL AS (per frame)"), STAT_D3D12RayTracingBuiltBLAS, STATGROUP_D3D12RayTracing);
 DECLARE_DWORD_COUNTER_STAT(TEXT("Updated BL AS (per frame)"), STAT_D3D12RayTracingUpdatedBLAS, STATGROUP_D3D12RayTracing);
 DECLARE_DWORD_COUNTER_STAT(TEXT("Built TL AS (per frame)"), STAT_D3D12RayTracingBuiltTLAS, STATGROUP_D3D12RayTracing);
-DECLARE_DWORD_COUNTER_STAT(TEXT("Updated TL AS (per frame)"), STAT_D3D12RayTracingUpdatedTLAS, STATGROUP_D3D12RayTracing);
 
 DECLARE_MEMORY_STAT(TEXT("BL AS Memory"), STAT_D3D12RayTracingBLASMemory, STATGROUP_D3D12RayTracing);
 DECLARE_MEMORY_STAT(TEXT("TL AS Memory"), STAT_D3D12RayTracingTLASMemory, STATGROUP_D3D12RayTracing);
@@ -2977,6 +2976,8 @@ FD3D12RayTracingScene::FD3D12RayTracingScene(FD3D12Adapter* Adapter, const FRayT
 {
 	INC_DWORD_STAT(STAT_D3D12RayTracingAllocatedTLAS);
 
+	Instances.SetAllowCPUAccess(true);
+
 	DebugName = Initializer.DebugName;
 
 	ShaderResourceView = Adapter->CreateLinkedObject<FD3D12ShaderResourceView>(FRHIGPUMask::All(), [&](FD3D12Device* Device)
@@ -2988,18 +2989,33 @@ FD3D12RayTracingScene::FD3D12RayTracingScene(FD3D12Adapter* Adapter, const FRayT
 
 	Lifetime = Initializer.Lifetime;
 	CreatedFrameFenceValue = Adapter->GetFrameFence().GetCurrentFence();
-	Instances = TArray<FRayTracingGeometryInstance>(Initializer.Instances.GetData(), Initializer.Instances.Num());
+
 	ShaderSlotsPerGeometrySegment = Initializer.ShaderSlotsPerGeometrySegment;
 
-	// Compute geometry segment count prefix sum to be later used in GetHitRecordBaseIndex()
-	SegmentPrefixSum.Reserve(Instances.Num());
+	Experimental::TSherwoodSet<FD3D12RayTracingGeometry*> UniqueGeometries;
+
+	const int32 NumSceneInstances = Initializer.Instances.Num();
+
+	int32 NumDxrInstances = 0;
+	int32 NumCopyCommands = 0;
 	NumTotalSegments = 0;
 
-	Experimental::TSherwoodSet<FD3D12RayTracingGeometry*> UniqueGeometries;
+	BaseInstancePrefixSum.Reserve(NumSceneInstances);
+	PerInstanceGeometries.Reserve(NumSceneInstances);
+	PerInstanceNumTransforms.Reserve(NumSceneInstances);
+	SegmentPrefixSum.Reserve(NumSceneInstances);
 
-	for (const FRayTracingGeometryInstance& Instance : Instances)
+	for (int32 InstanceIndex = 0; InstanceIndex < NumSceneInstances; ++InstanceIndex)
 	{
+		const FRayTracingGeometryInstance& Instance = Initializer.Instances[InstanceIndex];
+
+		checkf(Instance.GeometryRHI, TEXT("Ray tracing instance must have a valid geometry."));
+
 		FD3D12RayTracingGeometry* Geometry = FD3D12DynamicRHI::ResourceCast(Instance.GeometryRHI);
+
+		PerInstanceGeometries.Add(Geometry);
+
+		// Compute geometry segment count prefix sum to be later used in GetHitRecordBaseIndex()
 		SegmentPrefixSum.Add(NumTotalSegments);
 		NumTotalSegments += Geometry->Segments.Num();
 
@@ -3007,12 +3023,97 @@ FD3D12RayTracingScene::FD3D12RayTracingScene(FD3D12Adapter* Adapter, const FRayT
 		UniqueGeometries.Add(Geometry, &bIsAlreadyInSet);
 		if (!bIsAlreadyInSet)
 		{
-			Geometries.Add(Geometry);
+			ReferencedGeometries.Add(Geometry);
 		}
+
+		if (UNLIKELY(Instance.GPUTransformsSRV && Instance.NumTransforms))
+		{
+			FInstanceCopyCommand Command;
+			Command.Source = Instance.GPUTransformsSRV;
+			Command.BaseIndex = NumDxrInstances;
+			Command.Num = Instance.NumTransforms;
+			CopyCommands.Add(Command);
+		}
+
+		BaseInstancePrefixSum.Add(NumDxrInstances);
+		NumDxrInstances += Instance.NumTransforms;
+
+		PerInstanceNumTransforms.Add(Instance.NumTransforms);
+
+		TotalPrimitiveCount += uint64(Geometry->TotalPrimitiveCount) * Instance.NumTransforms;
 	}
 
-	// Reserve space for all segments
-	HitGroupSystemParametersCache.Reserve(NumTotalSegments);
+	Instances.SetNumUninitialized(NumDxrInstances);
+
+	const EParallelForFlags ParallelForFlags = EParallelForFlags::None; // set ForceSingleThread for testing
+	ParallelFor(NumSceneInstances, [this, &Initializer](int32 InstanceIndex)
+	{
+		const FRayTracingGeometryInstance& Instance = Initializer.Instances[InstanceIndex];
+		FD3D12RayTracingGeometry* Geometry = PerInstanceGeometries[InstanceIndex];
+
+		D3D12_RAYTRACING_INSTANCE_DESC InstanceDesc = {};
+
+		InstanceDesc.InstanceMask = Instance.Mask;
+		InstanceDesc.InstanceContributionToHitGroupIndex = SegmentPrefixSum[InstanceIndex] * ShaderSlotsPerGeometrySegment;
+		InstanceDesc.Flags = D3D12_RAYTRACING_INSTANCE_FLAG_TRIANGLE_FRONT_COUNTERCLOCKWISE; // #dxr_todo: convert cull mode based on instance mirroring or double-sidedness
+
+		if (Instance.bForceOpaque || GRayTracingDebugForceOpaque)
+		{
+			InstanceDesc.Flags |= D3D12_RAYTRACING_INSTANCE_FLAG_FORCE_OPAQUE;
+		}
+
+		if (Instance.bDoubleSided || GRayTracingDebugDisableTriangleCull)
+		{
+			InstanceDesc.Flags |= D3D12_RAYTRACING_INSTANCE_FLAG_TRIANGLE_CULL_DISABLE;
+		}
+
+		const uint32 NumTransforms = Instance.NumTransforms;
+
+		checkf(Instance.UserData.Num() == 0 || Instance.UserData.Num() == 1 || Instance.UserData.Num() == NumTransforms,
+			TEXT("User data array must be either be empty (implicit 0 is assumed), contain a single entry (same value applied to all instances) or contain one entry per entry in Transforms array."));
+
+		const bool bUseUniqueUserData = Instance.UserData.Num() > 1;
+		const uint32 CommonUserData = Instance.UserData.Num() == 1 ? Instance.UserData[0] : 0;
+
+		const bool bGpuInstance = Instance.GPUTransformsSRV.IsValid();
+		const bool bCpuInstance = !bGpuInstance;
+
+		uint32 DescIndex = BaseInstancePrefixSum[InstanceIndex];
+		for (uint32 TransformIndex = 0; TransformIndex < NumTransforms; ++TransformIndex)
+		{
+			InstanceDesc.InstanceID = bUseUniqueUserData ? Instance.UserData[TransformIndex] : CommonUserData;
+
+			if (LIKELY(bCpuInstance))
+			{
+				// DXR uses a 3x4 transform matrix in row major layout
+
+				const FMatrix& Transform = Instance.Transforms[TransformIndex];
+
+				InstanceDesc.Transform[0][0] = Transform.M[0][0];
+				InstanceDesc.Transform[0][1] = Transform.M[1][0];
+				InstanceDesc.Transform[0][2] = Transform.M[2][0];
+				InstanceDesc.Transform[0][3] = Transform.M[3][0];
+
+				InstanceDesc.Transform[1][0] = Transform.M[0][1];
+				InstanceDesc.Transform[1][1] = Transform.M[1][1];
+				InstanceDesc.Transform[1][2] = Transform.M[2][1];
+				InstanceDesc.Transform[1][3] = Transform.M[3][1];
+
+				InstanceDesc.Transform[2][0] = Transform.M[0][2];
+				InstanceDesc.Transform[2][1] = Transform.M[1][2];
+				InstanceDesc.Transform[2][2] = Transform.M[2][2];
+				InstanceDesc.Transform[2][3] = Transform.M[3][2];
+			}
+
+			Instances[DescIndex] = InstanceDesc;
+			++DescIndex;
+		}
+	}, ParallelForFlags);
+
+	BuildInputs.Type = D3D12_RAYTRACING_ACCELERATION_STRUCTURE_TYPE_TOP_LEVEL;
+	BuildInputs.DescsLayout = D3D12_ELEMENTS_LAYOUT_ARRAY;
+	BuildInputs.NumDescs = Instances.Num();
+	BuildInputs.Flags = D3D12_RAYTRACING_ACCELERATION_STRUCTURE_BUILD_FLAG_PREFER_FAST_TRACE;
 
 	NumCallableShaderSlots = Initializer.NumCallableShaderSlots;
 	NumMissShaderSlots = FMath::Max<uint32>(1, Initializer.NumMissShaderSlots);
@@ -3036,12 +3137,11 @@ FD3D12RayTracingScene::~FD3D12RayTracingScene()
 			DEC_MEMORY_STAT_BY(STAT_D3D12RayTracingTLASMemory, AccelerationStructureBuffer->GetSize());
 		}
 	}
-	
 
 	DEC_DWORD_STAT(STAT_D3D12RayTracingAllocatedTLAS);
 }
 
-void FD3D12RayTracingScene::BuildAccelerationStructure(FD3D12CommandContext& CommandContext, D3D12_RAYTRACING_ACCELERATION_STRUCTURE_BUILD_FLAGS BuildFlags)
+void FD3D12RayTracingScene::BuildAccelerationStructure(FD3D12CommandContext& CommandContext)
 {
 	TRACE_CPUPROFILER_EVENT_SCOPE(BuildAccelerationStructure_TopLevel);
 	SCOPE_CYCLE_COUNTER(STAT_D3D12BuildTLAS);
@@ -3053,26 +3153,8 @@ void FD3D12RayTracingScene::BuildAccelerationStructure(FD3D12CommandContext& Com
 	FD3D12Adapter* Adapter = CommandContext.GetParentAdapter();
 	ID3D12Device5* RayTracingDevice = CommandContext.GetParentDevice()->GetDevice5();
 
-	const uint32 NumSceneInstances = Instances.Num();
-
-	uint32 NumDxrInstances = 0;
-	BaseInstancePrefixSum.Reserve(NumSceneInstances);
-	for (uint32 InstanceIndex = 0; InstanceIndex < NumSceneInstances; ++InstanceIndex)
-	{
-		const FRayTracingGeometryInstance& Instance = Instances[InstanceIndex];
-		BaseInstancePrefixSum.Add(NumDxrInstances);
-		NumDxrInstances += Instance.NumTransforms;
-	}
-
-	D3D12_BUILD_RAYTRACING_ACCELERATION_STRUCTURE_INPUTS PrebuildDescInputs = {};
-
 	D3D12_RAYTRACING_ACCELERATION_STRUCTURE_PREBUILD_INFO PrebuildInfo = {};
-	PrebuildDescInputs.Type = D3D12_RAYTRACING_ACCELERATION_STRUCTURE_TYPE_TOP_LEVEL;
-	PrebuildDescInputs.DescsLayout = D3D12_ELEMENTS_LAYOUT_ARRAY;
-	PrebuildDescInputs.NumDescs = NumDxrInstances;
-	PrebuildDescInputs.Flags = BuildFlags;
-
-	RayTracingDevice->GetRaytracingAccelerationStructurePrebuildInfo(&PrebuildDescInputs, &PrebuildInfo);
+	RayTracingDevice->GetRaytracingAccelerationStructurePrebuildInfo(&BuildInputs, &PrebuildInfo);
 
 	TRefCountPtr<FD3D12Buffer>& AccelerationStructureBuffer = AccelerationStructureBuffers[GPUIndex];
 
@@ -3109,15 +3191,45 @@ void FD3D12RayTracingScene::BuildAccelerationStructure(FD3D12CommandContext& Com
 
 	// Create and fill instance buffer
 
-	TotalPrimitiveCount = 0;
-
-	if (NumSceneInstances)
+	if (BuildInputs.NumDescs)
 	{
+		// - Set up acceleration structure pointers and make them resident.
+		// - Generate HitGroupSystemParametersCache.
+
+		HitGroupSystemParametersCache.Reserve(NumTotalSegments);
+		HitGroupSystemParametersCache.Empty(NumTotalSegments);
+
+		{
+			int32 DxrInstanceIndex = 0;
+			const int32 NumSceneInstances = PerInstanceGeometries.Num();
+			for (int32 InstanceIndex = 0; InstanceIndex < NumSceneInstances; ++InstanceIndex)
+			{
+				FD3D12RayTracingGeometry* Geometry = PerInstanceGeometries[InstanceIndex];
+				Geometry->AccelerationStructureBuffers[GPUIndex]->ResourceLocation.GetResource()->UpdateResidency(CommandContext.CommandListHandle);
+
+				D3D12_GPU_VIRTUAL_ADDRESS BlasAddress = Geometry->AccelerationStructureBuffers[GPUIndex]->ResourceLocation.GetGPUVirtualAddress();
+
+				const int32 NumTransforms = PerInstanceNumTransforms[InstanceIndex];
+				for (int32 TransformIndex = 0; TransformIndex < NumTransforms; ++TransformIndex)
+				{
+					Instances[DxrInstanceIndex].AccelerationStructure = BlasAddress;
+					DxrInstanceIndex++;
+				}
+
+				// make a copy of system parameters to they can optimized fetch during SBT building (only done for GPU0)
+				check(Geometry->HitGroupSystemParameters[0].Num() > 0);
+				HitGroupSystemParametersCache.Append(Geometry->HitGroupSystemParameters[0]);
+			}
+		}
+
 		FRHIResourceCreateInfo CreateInfo(TEXT("InstanceBuffer"));
 		CreateInfo.GPUMask = FRHIGPUMask::FromIndex(GPUIndex);
+		CreateInfo.ResourceArray = &Instances;
 		
+		check(Instances.Num() == BuildInputs.NumDescs);
+
 		D3D12_RESOURCE_DESC InstanceBufferDesc = CD3DX12_RESOURCE_DESC::Buffer(
-			sizeof(D3D12_RAYTRACING_INSTANCE_DESC) * PrebuildDescInputs.NumDescs,
+			sizeof(D3D12_RAYTRACING_INSTANCE_DESC) * BuildInputs.NumDescs,
 			D3D12_RESOURCE_FLAG_ALLOW_UNORDERED_ACCESS);
 
 		// #dxr_todo multi state only when bShouldCopyIndirectInstances is set - will still need transition to copy_src for lock behind but this can be done on the complete pool in theory (have to check cost)
@@ -3125,15 +3237,6 @@ void FD3D12RayTracingScene::BuildAccelerationStructure(FD3D12CommandContext& Com
 			nullptr, InstanceBufferDesc, D3D12_RAYTRACING_INSTANCE_DESCS_BYTE_ALIGNMENT,
 			sizeof(D3D12_RAYTRACING_INSTANCE_DESC), InstanceBufferDesc.Width, BUF_UnorderedAccess, ED3D12ResourceStateMode::MultiState, ERHIAccess::UAVMask, CreateInfo);
 
-		D3D12_RAYTRACING_INSTANCE_DESC* MappedData = (D3D12_RAYTRACING_INSTANCE_DESC*)Adapter->GetOwningRHI()->LockBuffer(
-			nullptr, InstanceBuffer.GetReference(), InstanceBuffer->GetSize(), InstanceBuffer->GetUsage(), 0, InstanceBufferDesc.Width, RLM_WriteOnly);
-
-		check(MappedData);
-
-		uint32 DxrInstanceIndex = 0;
-
-		bool bShouldCopyIndirectInstances = false;
-
 		{
 			TArray<FD3D12ResidencyHandle*>& GeometryResidencyHandlesForThisGPU = GeometryResidencyHandles[GPUIndex];
 
@@ -3157,7 +3260,7 @@ void FD3D12RayTracingScene::BuildAccelerationStructure(FD3D12CommandContext& Com
 			#endif // ENABLE_RESIDENCY_MANAGEMENT
 			};
 
-			for (TRefCountPtr<FD3D12RayTracingGeometry>& Geometry : Geometries)
+			for (TRefCountPtr<FD3D12RayTracingGeometry>& Geometry : ReferencedGeometries)
 			{
 				checkf(!Geometry->IsDirty(CommandContext.GetGPUIndex()),
 					TEXT("Acceleration structures for all geometries must be built before building the top level acceleration structure for the scene."));
@@ -3181,122 +3284,45 @@ void FD3D12RayTracingScene::BuildAccelerationStructure(FD3D12CommandContext& Com
 			}
 		}
 
-		for (uint32 InstanceIndex = 0; InstanceIndex < NumSceneInstances; ++InstanceIndex)
-		{
-			const FRayTracingGeometryInstance& Instance = Instances[InstanceIndex];
-			FD3D12RayTracingGeometry* Geometry = FD3D12DynamicRHI::ResourceCast(Instance.GeometryRHI);
-
-			// make a copy of system parameters to they can optimized fetch during SBT building (only done for GPU0)
-			check(Geometry->HitGroupSystemParameters[0].Num() > 0);
-			HitGroupSystemParametersCache.Append(Geometry->HitGroupSystemParameters[0]);
-
-			D3D12_RAYTRACING_INSTANCE_DESC InstanceDesc = {};
-
-			InstanceDesc.InstanceMask = Instance.Mask;
-			InstanceDesc.InstanceContributionToHitGroupIndex = SegmentPrefixSum[InstanceIndex] * ShaderSlotsPerGeometrySegment;
-			InstanceDesc.Flags = D3D12_RAYTRACING_INSTANCE_FLAG_TRIANGLE_FRONT_COUNTERCLOCKWISE; // #dxr_todo: convert cull mode based on instance mirroring or double-sidedness
-			
-			if (Instance.bForceOpaque || GRayTracingDebugForceOpaque)
-			{
-				InstanceDesc.Flags |= D3D12_RAYTRACING_INSTANCE_FLAG_FORCE_OPAQUE;
-			}
-
-			if (Instance.bDoubleSided || GRayTracingDebugDisableTriangleCull)
-			{
-				InstanceDesc.Flags |= D3D12_RAYTRACING_INSTANCE_FLAG_TRIANGLE_CULL_DISABLE;
-			}
-
-			InstanceDesc.AccelerationStructure = Geometry->AccelerationStructureBuffers[GPUIndex]->ResourceLocation.GetGPUVirtualAddress();
-			Geometry->AccelerationStructureBuffers[GPUIndex]->ResourceLocation.GetResource()->UpdateResidency(CommandContext.CommandListHandle);
-
-			bool bHasGPUTransforms = Instance.GPUTransformsSRV.IsValid();
-
-			const uint32 NumTransforms = Instance.NumTransforms;
-
-			checkf(Instance.UserData.Num() == 0 || Instance.UserData.Num() == 1 || Instance.UserData.Num() == NumTransforms, 
-				TEXT("User data array must be either be empty (implicit 0 is assumed), contain a single entry (same value applied to all instances) or contain one entry per entry in Transforms array."));
-
-			const bool bUseUniqueUserData = Instance.UserData.Num() > 1;
-			const uint32 CommonUserData = Instance.UserData.Num() == 1 ? Instance.UserData[0] : 0;
-
-			bShouldCopyIndirectInstances |= bHasGPUTransforms;
-
-			for (uint32 TransformIndex = 0; TransformIndex < NumTransforms; ++TransformIndex)
-			{
-				InstanceDesc.InstanceID = bUseUniqueUserData ? Instance.UserData[TransformIndex] : CommonUserData;
-
-				if (!bHasGPUTransforms)
-				{
-					FMatrix TransformTransposed = Instance.Transforms[TransformIndex].GetTransposed();
-					
-					// Ensure the last row of the original Transform is <0,0,0,1>
-					check((TransformTransposed.M[3][0] == 0)
-						&& (TransformTransposed.M[3][1] == 0)
-						&& (TransformTransposed.M[3][2] == 0)
-						&& (TransformTransposed.M[3][3] == 1));
-
-					FMemory::Memcpy(&InstanceDesc.Transform, &TransformTransposed.M[0][0], sizeof(InstanceDesc.Transform));
-				}
-
-				MappedData[DxrInstanceIndex++] = InstanceDesc;
-			}
-
-			TotalPrimitiveCount += uint64(Geometry->TotalPrimitiveCount) * NumTransforms;
-		}
-
-		Adapter->GetOwningRHI()->UnlockBuffer(nullptr, InstanceBuffer.GetReference(), InstanceBuffer->GetUsage());
-
 		InstanceBuffer->GetResource()->UpdateResidency(CommandContext.CommandListHandle);
 
-		if(bShouldCopyIndirectInstances)
+		if(CopyCommands.Num())
 		{
 			TRHICommandList_RecursiveHazardous<FD3D12CommandContext> RHICmdList(&CommandContext);
 			FUnorderedAccessViewRHIRef InstancesDescUAV = RHICreateUnorderedAccessView(InstanceBuffer, false, false);
 
-			FD3D12DynamicRHI::TransitionResource(CommandContext.CommandListHandle, InstanceBuffer.GetReference()->GetResource(), D3D12_RESOURCE_STATE_TBD, D3D12_RESOURCE_STATE_UNORDERED_ACCESS, 0, FD3D12DynamicRHI::ETransitionMode::Apply);
+			FD3D12DynamicRHI::TransitionResource(CommandContext.CommandListHandle, InstanceBuffer.GetReference()->GetResource(),
+				D3D12_RESOURCE_STATE_TBD, D3D12_RESOURCE_STATE_UNORDERED_ACCESS, 0,
+				FD3D12DynamicRHI::ETransitionMode::Apply);
 
 			RHICmdList.BeginUAVOverlap(InstancesDescUAV);
 
-			uint32 DescOffset = 0;
-			for (uint32 InstanceIndex = 0; InstanceIndex < NumSceneInstances; ++InstanceIndex)
+			for (const FInstanceCopyCommand& Command : CopyCommands)
 			{
-				const FRayTracingGeometryInstance& Instance = Instances[InstanceIndex];
-				FD3D12RayTracingGeometry* Geometry = FD3D12DynamicRHI::ResourceCast(Instance.GeometryRHI);
-
-				const uint32 NumTransforms = Instance.NumTransforms;
-
-				if (NumTransforms > 0 && Instance.GPUTransformsSRV.IsValid())
-				{
-					FRHIShaderResourceView* TransformsSRV = Instance.GPUTransformsSRV.GetReference();
-
-					CopyRayTracingGPUInstances(RHICmdList, NumTransforms, DescOffset, TransformsSRV, InstancesDescUAV);
-				}
-
-				DescOffset += NumTransforms;
+				FRHIShaderResourceView* TransformsSRV = Command.Source.GetReference();
+				CopyRayTracingGPUInstances(RHICmdList, Command.Num, Command.BaseIndex, TransformsSRV, InstancesDescUAV);
 			}
 
 			RHICmdList.EndUAVOverlap(InstancesDescUAV);
 		}
 
-		FD3D12DynamicRHI::TransitionResource(CommandContext.CommandListHandle, InstanceBuffer.GetReference()->GetResource(), D3D12_RESOURCE_STATE_TBD, D3D12_RESOURCE_STATE_PIXEL_SHADER_RESOURCE | D3D12_RESOURCE_STATE_NON_PIXEL_SHADER_RESOURCE, 0, FD3D12DynamicRHI::ETransitionMode::Apply);
+		FD3D12DynamicRHI::TransitionResource(CommandContext.CommandListHandle, InstanceBuffer.GetReference()->GetResource(),
+			D3D12_RESOURCE_STATE_TBD, D3D12_RESOURCE_STATE_PIXEL_SHADER_RESOURCE | D3D12_RESOURCE_STATE_NON_PIXEL_SHADER_RESOURCE, 0,
+			FD3D12DynamicRHI::ETransitionMode::Apply);
 	}
 	
 	
 	// Build the actual acceleration structure
 
-	const bool bIsUpdateMode = false; // Top level acceleration structure is always built from scratch
-
 	AccelerationStructureBuffer->GetResource()->UpdateResidency(CommandContext.CommandListHandle);
 	ScratchBuffer->GetResource()->UpdateResidency(CommandContext.CommandListHandle);
 
 	D3D12_BUILD_RAYTRACING_ACCELERATION_STRUCTURE_DESC BuildDesc = {};
-	BuildDesc.Inputs = PrebuildDescInputs;
+	BuildDesc.Inputs = BuildInputs;
 	BuildDesc.Inputs.InstanceDescs = InstanceBuffer ? InstanceBuffer->ResourceLocation.GetGPUVirtualAddress() : D3D12_GPU_VIRTUAL_ADDRESS(0);
 	BuildDesc.DestAccelerationStructureData = AccelerationStructureBuffer->ResourceLocation.GetGPUVirtualAddress();
 	BuildDesc.ScratchAccelerationStructureData = ScratchBuffer->ResourceLocation.GetGPUVirtualAddress();
-	BuildDesc.SourceAccelerationStructureData = bIsUpdateMode
-		? AccelerationStructureBuffer->ResourceLocation.GetGPUVirtualAddress()
-		: D3D12_GPU_VIRTUAL_ADDRESS(0);
+	BuildDesc.SourceAccelerationStructureData = D3D12_GPU_VIRTUAL_ADDRESS(0); // Null source TLAS as this is a build command
 
 	// UAV barrier is used here to ensure that all bottom level acceleration structures are built
 	CommandContext.CommandListHandle.AddUAVBarrier();
@@ -3305,14 +3331,7 @@ void FD3D12RayTracingScene::BuildAccelerationStructure(FD3D12CommandContext& Com
 	ID3D12GraphicsCommandList4* RayTracingCommandList = CommandContext.CommandListHandle.RayTracingCommandList();
 	RayTracingCommandList->BuildRaytracingAccelerationStructure(&BuildDesc, 0, nullptr);
 
-	if (bIsUpdateMode)
-	{
-		INC_DWORD_STAT(STAT_D3D12RayTracingUpdatedTLAS);
-	}
-	else
-	{
-		INC_DWORD_STAT(STAT_D3D12RayTracingBuiltTLAS);
-	}
+	INC_DWORD_STAT(STAT_D3D12RayTracingBuiltTLAS);
 
 	// UAV barrier is used here to ensure that the acceleration structure build is complete before any rays are traced
 	// #dxr_todo: these barriers should ideally be inserted by the high level code to allow more overlapped execution
@@ -3451,7 +3470,7 @@ void FD3D12CommandContext::RHIBuildAccelerationStructures(const TArrayView<const
 void FD3D12CommandContext::RHIBuildAccelerationStructure(FRHIRayTracingScene* InScene)
 {
 	FD3D12RayTracingScene* Scene = FD3D12DynamicRHI::ResourceCast(InScene);
-	Scene->BuildAccelerationStructure(*this, D3D12_RAYTRACING_ACCELERATION_STRUCTURE_BUILD_FLAG_PREFER_FAST_TRACE);
+	Scene->BuildAccelerationStructure(*this);
 }
 
 void FD3D12CommandContext::RHIClearRayTracingBindings(FRHIRayTracingScene* InScene)
@@ -4348,8 +4367,7 @@ static void SetRayTracingHitGroup(
 	}
 	else
 	{
-		const FRayTracingGeometryInstance& Instance = Scene->Instances[InstanceIndex];
-		const FD3D12RayTracingGeometry* Geometry = FD3D12DynamicRHI::ResourceCast(Instance.GeometryRHI);
+		const FD3D12RayTracingGeometry* Geometry = Scene->PerInstanceGeometries[InstanceIndex];
 		const TArray<FHitGroupSystemParameters>& HitGroupSystemParametersForThisGPU = Geometry->HitGroupSystemParameters[GPUIndex];
 
 		SystemParameters = HitGroupSystemParametersForThisGPU[SegmentIndex];
diff --git a/Engine/Source/Runtime/D3D12RHI/Private/D3D12RayTracing.h b/Engine/Source/Runtime/D3D12RHI/Private/D3D12RayTracing.h
--- a/Engine/Source/Runtime/D3D12RHI/Private/D3D12RayTracing.h
+++ b/Engine/Source/Runtime/D3D12RHI/Private/D3D12RayTracing.h
@@ -74,7 +74,7 @@ public:
 	FName DebugName;
 };
 
-class FD3D12RayTracingScene : public FRHIRayTracingScene, public FD3D12AdapterChild
+class FD3D12RayTracingScene : public FRHIRayTracingScene, public FD3D12AdapterChild, public FNoncopyable
 {
 public:
 
@@ -86,16 +86,29 @@ public:
 	FD3D12RayTracingScene(FD3D12Adapter* Adapter, const FRayTracingSceneInitializer& Initializer);
 	~FD3D12RayTracingScene();
 
-	void BuildAccelerationStructure(FD3D12CommandContext& CommandContext, D3D12_RAYTRACING_ACCELERATION_STRUCTURE_BUILD_FLAGS BuildFlags);
+	void BuildAccelerationStructure(FD3D12CommandContext& CommandContext);
 
 	TRefCountPtr<FD3D12Buffer> AccelerationStructureBuffers[MAX_NUM_GPUS];
 	bool bAccelerationStructureViewInitialized[MAX_NUM_GPUS] = {};
 
-	TArray<FRayTracingGeometryInstance> Instances;
+	TResourceArray<D3D12_RAYTRACING_INSTANCE_DESC, 16> Instances;
+	D3D12_BUILD_RAYTRACING_ACCELERATION_STRUCTURE_INPUTS BuildInputs = {};
+	
+	struct FInstanceCopyCommand
+	{
+		FShaderResourceViewRHIRef Source;
+		int32 BaseIndex = -1;
+		int32 Num = 0;
+	};
+	TArray<FInstanceCopyCommand> CopyCommands;
+	
+	// One entry per instance in the Initializer
+	TArray<FD3D12RayTracingGeometry*> PerInstanceGeometries;
+	TArray<uint32> PerInstanceNumTransforms;
 
 	// Unique list of geometries referenced by all instances in this scene.
 	// Any referenced geometry is kept alive while the scene is alive.
-	TArray<TRefCountPtr<FD3D12RayTracingGeometry>> Geometries;
+	TArray<TRefCountPtr<FD3D12RayTracingGeometry>> ReferencedGeometries;
 
 	// Scene keeps track of child acceleration structure buffers to ensure
 	// they are resident when any ray tracing work is dispatched.
@@ -118,6 +131,9 @@ public:
 	uint32 NumCallableShaderSlots = 0;
 	uint32 NumMissShaderSlots = 1; // always at least the default
 
+	// Array of hit group parameters per geometry segment across all scene instance geometries.
+	// Accessed as HitGroupSystemParametersCache[SegmentPrefixSum[InstanceIndex] + SegmentIndex].
+	// Only used for GPU 0 (secondary GPUs take the slow path).
 	TArray<FHitGroupSystemParameters> HitGroupSystemParametersCache;
 
 	// #dxr_todo UE-68230: shader tables should be explicitly registered and unregistered with the scene
