Temp workaround for D3D12 NV vendor extensions using alternate API until NvAPI_D3D12_CreateGraphicsPipelineState is supported by RenderDoc

diff --git a/Engine/Source/Runtime/D3D12RHI/Private/Windows/WindowsD3D12PipelineState.cpp b/Engine/Source/Runtime/D3D12RHI/Private/Windows/WindowsD3D12PipelineState.cpp
--- a/Engine/Source/Runtime/D3D12RHI/Private/Windows/WindowsD3D12PipelineState.cpp
+++ b/Engine/Source/Runtime/D3D12RHI/Private/Windows/WindowsD3D12PipelineState.cpp
@@ -28,6 +28,9 @@
 #define D3D12RHI_USE_D3DDISASSEMBLE 1
 #endif
 
+// TODO: Temp workaround using alternate API until NvAPI_D3D12_CreateGraphicsPipelineState is supported by RenderDoc
+#define TEMP_RENDERDOC_WORKAROUND 1
+
 // D3D12RHI PSO file cache doesn't work anymore. Use FPipelineFileCache instead
 static TAutoConsoleVariable<int32> CVarPipelineStateDiskCache(
 	TEXT("D3D12.PSO.DiskCache"),
@@ -910,7 +913,12 @@ static FORCEINLINE NVAPI_D3D12_PSO_SET_SHADER_EXTENSION_SLOT_DESC GetNVShaderExt
 	ShdExtensionDesc.baseVersion = NV_PSO_EXTENSION_DESC_VER;
 	ShdExtensionDesc.version = NV_SET_SHADER_EXTENSION_SLOT_DESC_VER;
 	ShdExtensionDesc.uavSlot = UavSlot;
+#if TEMP_RENDERDOC_WORKAROUND
+	// TODO: Temp workaround until RenderDoc fixes SM 5.0 opcode detection (need to ignore register space)
+	ShdExtensionDesc.registerSpace = 0xFFFFFFFF;
+#else
 	ShdExtensionDesc.registerSpace = 0; // TODO: Should use the same special space as AMD to avoid driver pattern matching
+#endif
 	return ShdExtensionDesc;
 }
 
@@ -945,12 +953,31 @@ static void CreateGraphicsPipelineState(ID3D12PipelineState** PSO, FD3D12Adapter
 			{
 				if (Extension.Parameter.Type == EShaderParameterType::UAV)
 				{
-					D3D12_GRAPHICS_PIPELINE_STATE_DESC Desc = CreationArgs->Desc.Desc.GraphicsDescV0();
-
 					const NVAPI_D3D12_PSO_SET_SHADER_EXTENSION_SLOT_DESC ShdExtensionDesc = GetNVShaderExtensionDesc(Extension.Parameter.BaseIndex);
+				#if TEMP_RENDERDOC_WORKAROUND
+					// TODO: Temp workaround using alternate API until NvAPI_D3D12_CreateGraphicsPipelineState is supported by RenderDoc
+					NvAPI_Status NvStatus = NvAPI_D3D12_SetNvShaderExtnSlotSpaceLocalThread(
+						Adapter->GetD3DDevice(),
+						ShdExtensionDesc.uavSlot,
+						ShdExtensionDesc.registerSpace
+					);
+					check(NvStatus == NVAPI_OK);
+
+					CreatePipelineStateWrapper(PSO, Adapter, CreationArgs, false /* use stream */);
+
+					// Reset to default configuration without vendor extensions.
+					NvStatus = NvAPI_D3D12_SetNvShaderExtnSlotSpaceLocalThread(
+						Adapter->GetD3DDevice(),
+						0xFFFFFFFF,
+						0
+					);
+					check(NvStatus == NVAPI_OK);
+				#else
+					D3D12_GRAPHICS_PIPELINE_STATE_DESC Desc = CreationArgs->Desc.Desc.GraphicsDescV0();
 					const NVAPI_D3D12_PSO_EXTENSION_DESC* NvExtensions[] = { &ShdExtensionDesc };
 					NvAPI_Status NvStatus = NvAPI_D3D12_CreateGraphicsPipelineState(Adapter->GetD3DDevice(), &Desc, ARRAYSIZE(NvExtensions), NvExtensions, PSO);
 					check(NvStatus == NVAPI_OK);
+				#endif
 					return;
 				}
 			}
@@ -987,12 +1014,31 @@ static void CreateComputePipelineState(ID3D12PipelineState** PSO, FD3D12Adapter*
 			{
 				if (Extension.Parameter.Type == EShaderParameterType::UAV)
 				{
-					D3D12_COMPUTE_PIPELINE_STATE_DESC Desc = CreationArgs->Desc.Desc.ComputeDescV0();
-
 					const NVAPI_D3D12_PSO_SET_SHADER_EXTENSION_SLOT_DESC ShdExtensionDesc = GetNVShaderExtensionDesc(Extension.Parameter.BaseIndex);
+				#if TEMP_RENDERDOC_WORKAROUND
+					// TODO: Temp workaround using alternate API until NvAPI_D3D12_CreateGraphicsPipelineState is supported by RenderDoc
+					NvAPI_Status NvStatus = NvAPI_D3D12_SetNvShaderExtnSlotSpaceLocalThread(
+						Adapter->GetD3DDevice(),
+						ShdExtensionDesc.uavSlot,
+						ShdExtensionDesc.registerSpace
+					);
+					check(NvStatus == NVAPI_OK);
+
+					CreatePipelineStateWrapper(PSO, Adapter, CreationArgs, false /* use stream */);
+
+					// Reset to default configuration without vendor extensions.
+					NvStatus = NvAPI_D3D12_SetNvShaderExtnSlotSpaceLocalThread(
+						Adapter->GetD3DDevice(),
+						0xFFFFFFFF,
+						0
+					);
+					check(NvStatus == NVAPI_OK);
+				#else
+					D3D12_COMPUTE_PIPELINE_STATE_DESC Desc = CreationArgs->Desc.Desc.ComputeDescV0();
 					const NVAPI_D3D12_PSO_EXTENSION_DESC* NvExtensions[] = { &ShdExtensionDesc };
 					NvAPI_Status NvStatus = NvAPI_D3D12_CreateComputePipelineState(Adapter->GetD3DDevice(), &Desc, ARRAYSIZE(NvExtensions), NvExtensions, PSO);
 					check(NvStatus == NVAPI_OK);
+				#endif
 					return;
 				}
 			}
