Temporarily revert part of the RDoc workaround since RDoc disabled + NV D3D12 driver fails to ignore register spaces under SM5.0 (Nanite disappears). Most of the workaround remains, which fixes the device hooking, but UAV writes won't be tracked until the next nightly release.

diff --git a/Engine/Source/Runtime/D3D12RHI/Private/Windows/WindowsD3D12PipelineState.cpp b/Engine/Source/Runtime/D3D12RHI/Private/Windows/WindowsD3D12PipelineState.cpp
--- a/Engine/Source/Runtime/D3D12RHI/Private/Windows/WindowsD3D12PipelineState.cpp
+++ b/Engine/Source/Runtime/D3D12RHI/Private/Windows/WindowsD3D12PipelineState.cpp
@@ -913,7 +913,7 @@ static FORCEINLINE NVAPI_D3D12_PSO_SET_SHADER_EXTENSION_SLOT_DESC GetNVShaderExt
 	ShdExtensionDesc.baseVersion = NV_PSO_EXTENSION_DESC_VER;
 	ShdExtensionDesc.version = NV_SET_SHADER_EXTENSION_SLOT_DESC_VER;
 	ShdExtensionDesc.uavSlot = UavSlot;
-#if TEMP_RENDERDOC_WORKAROUND
+#if 0 // TEMP_RENDERDOC_WORKAROUND
 	// TODO: Temp workaround until RenderDoc fixes SM 5.0 opcode detection (need to ignore register space)
 	ShdExtensionDesc.registerSpace = 0xFFFFFFFF;
 #else
