- New IRHITransientResourceAllocator which can be used as base class for platform independent transient resource allocations - lifetime of all resources allocated via the transient allocator are well defined and ties to the lifetime of the allocator itself
- D3D12 RHI implementation of the transient resource allocator using aliasing barriers and discards on transient allocated resources

Not enabled yet and will be used by RDG. Doesn't replace the old transient code yet (only used by [REDACTED_PERSON]), but the plan is to replace over time when it's verified that it behaves the same or better than the old transient code on base xbox one

diff --git a/Engine/Source/Runtime/D3D12RHI/Private/D3D12Adapter.cpp b/Engine/Source/Runtime/D3D12RHI/Private/D3D12Adapter.cpp
--- a/Engine/Source/Runtime/D3D12RHI/Private/D3D12Adapter.cpp
+++ b/Engine/Source/Runtime/D3D12RHI/Private/D3D12Adapter.cpp
@@ -620,6 +620,10 @@ void FD3D12Adapter::CreateRootDevice(bool bWithDebug)
 				// to be interested in keeping PSO caches associated with old drivers around on disk, so it's better to just reset.
 				D3D12_MESSAGE_ID_CREATEPIPELINELIBRARY_DRIVERVERSIONMISMATCH,
 
+				// D3D complain about overlapping GPU addresses when aliasing DataBuffers in the same command list when using the Transient Allocator - it looks like
+				// it ignored the aliasing barriers to validate, and probably can't check them when called from IASetVertexBuffers because it only has GPU Virtual Addresses then
+				D3D12_MESSAGE_ID_HEAP_ADDRESS_RANGE_INTERSECTS_MULTIPLE_BUFFERS,
+
 #if ENABLE_RESIDENCY_MANAGEMENT
 				// TODO: Remove this when the debug layers work for executions which are guarded by a fence
 				D3D12_MESSAGE_ID_INVALID_USE_OF_NON_RESIDENT_RESOURCE,
diff --git a/Engine/Source/Runtime/D3D12RHI/Private/D3D12Adapter.h b/Engine/Source/Runtime/D3D12RHI/Private/D3D12Adapter.h
--- a/Engine/Source/Runtime/D3D12RHI/Private/D3D12Adapter.h
+++ b/Engine/Source/Runtime/D3D12RHI/Private/D3D12Adapter.h
@@ -279,7 +279,8 @@ public:
 		uint32 Alignment, uint32 Stride, uint32 Size, uint32 InUsage,
 		ED3D12ResourceStateMode InResourceStateMode,
 		ERHIAccess InResourceState,
-		FRHIResourceCreateInfo& CreateInfo);
+		FRHIResourceCreateInfo& CreateInfo,
+		ID3D12ResourceAllocator* ResourceAllocator);
 
 	template <typename ObjectType, typename CreationCoreFunction>
 	inline ObjectType* CreateLinkedObject(FRHIGPUMask GPUMask, const CreationCoreFunction& pfnCreationCore)
@@ -333,7 +334,8 @@ protected:
 		FRHIResourceCreateInfo& CreateInfo,
 		uint32 Alignment,
 		FD3D12Buffer* Buffer,
-		FD3D12ResourceLocation& ResourceLocation);
+		FD3D12ResourceLocation& ResourceLocation,
+		ID3D12ResourceAllocator* ResourceAllocator);
 
 	// Creates default root and execute indirect signatures
 	virtual void CreateCommandSignatures();
diff --git a/Engine/Source/Runtime/D3D12RHI/Private/D3D12Allocation.cpp b/Engine/Source/Runtime/D3D12RHI/Private/D3D12Allocation.cpp
--- a/Engine/Source/Runtime/D3D12RHI/Private/D3D12Allocation.cpp
+++ b/Engine/Source/Runtime/D3D12RHI/Private/D3D12Allocation.cpp
@@ -1263,10 +1263,11 @@ FD3D12BufferPool* FD3D12DefaultBufferAllocator::CreateBufferPool(D3D12_HEAP_TYPE
 	uint64 PoolSize = InHeapType == D3D12_HEAP_TYPE_READBACK ? READBACK_BUFFER_POOL_DEFAULT_POOL_SIZE : GD3D12PoolAllocatorBufferVRAMPoolSize;
 	uint64 PoolAlignment = (AllocationStrategy == EResourceAllocationStrategy::kPlacedResource) ? MIN_PLACED_RESOURCE_SIZE : 256;
 	uint64 MaxAllocationSize = InHeapType == D3D12_HEAP_TYPE_READBACK ? READBACK_BUFFER_POOL_MAX_ALLOC_SIZE : GD3D12PoolAllocatorBufferVRAMMaxAllocationSize;
+	FRHIMemoryPool::EFreeListOrder FreeListOrder = FRHIMemoryPool::EFreeListOrder::SortBySize;
 
 	// Disable defrag if not Default memory
 	bool bDefragEnabled = (InitConfig.HeapType == D3D12_HEAP_TYPE_DEFAULT);
-	FD3D12BufferPool* NewPool = new FD3D12PoolAllocator(Device, GetVisibilityMask(), InitConfig, Name, AllocationStrategy, PoolSize, PoolAlignment, MaxAllocationSize, bDefragEnabled);
+	FD3D12BufferPool* NewPool = new FD3D12PoolAllocator(Device, GetVisibilityMask(), InitConfig, Name, AllocationStrategy, PoolSize, PoolAlignment, MaxAllocationSize, FreeListOrder, bDefragEnabled);
 
 #else // USE_POOL_ALLOCATOR
 
@@ -1487,7 +1488,8 @@ FD3D12TextureAllocatorPool::FD3D12TextureAllocatorPool(FD3D12Device* Device, FRH
 	SharedInitConfig.InitialResourceState = D3D12_RESOURCE_STATE_COMMON;
 
 	EResourceAllocationStrategy AllocationStrategy = EResourceAllocationStrategy::kPlacedResource;
-	
+	FRHIMemoryPool::EFreeListOrder FreeListOrder = FRHIMemoryPool::EFreeListOrder::SortBySize;
+
 	{
 		FD3D12ResourceInitConfig InitConfig = SharedInitConfig;
 		InitConfig.HeapFlags = D3D12_HEAP_FLAG_ALLOW_ONLY_NON_RT_DS_TEXTURES;
@@ -1496,7 +1498,7 @@ FD3D12TextureAllocatorPool::FD3D12TextureAllocatorPool(FD3D12Device* Device, FRH
 		uint64 PoolSize = 4 * 1024 * 1024;
 		uint64 MaxAllocationSize = PoolSize;
 		bool bDefragEnabled = false; // Disable defrag on 4K pool because it shouldn't really fragment - all allocations are 4K or multiple of 4K and pretty small
-		FD3D12PoolAllocator* ReadOnly4KPool = new FD3D12PoolAllocator(Device, GetVisibilityMask(), InitConfig, Name, AllocationStrategy, PoolSize, D3D12_SMALL_RESOURCE_PLACEMENT_ALIGNMENT, MaxAllocationSize, bDefragEnabled);
+		FD3D12PoolAllocator* ReadOnly4KPool = new FD3D12PoolAllocator(Device, GetVisibilityMask(), InitConfig, Name, AllocationStrategy, PoolSize, D3D12_SMALL_RESOURCE_PLACEMENT_ALIGNMENT, MaxAllocationSize, FreeListOrder, bDefragEnabled);
 		ReadOnly4KPool->Initialize();
 
 		PoolAllocators[(int)EPoolType::ReadOnly4K] = ReadOnly4KPool;
@@ -1510,7 +1512,7 @@ FD3D12TextureAllocatorPool::FD3D12TextureAllocatorPool(FD3D12Device* Device, FRH
 		uint64 PoolSize = GD3D12PoolAllocatorReadOnlyTextureVRAMPoolSize;
 		uint64 MaxAllocationSize = GD3D12PoolAllocatorReadOnlyTextureVRAMMaxAllocationSize;
 		bool bDefragEnabled = true;
-		FD3D12PoolAllocator* ReadOnlyPool = new FD3D12PoolAllocator(Device, GetVisibilityMask(), InitConfig, Name, AllocationStrategy, PoolSize, D3D12_DEFAULT_RESOURCE_PLACEMENT_ALIGNMENT, MaxAllocationSize, bDefragEnabled);
+		FD3D12PoolAllocator* ReadOnlyPool = new FD3D12PoolAllocator(Device, GetVisibilityMask(), InitConfig, Name, AllocationStrategy, PoolSize, D3D12_DEFAULT_RESOURCE_PLACEMENT_ALIGNMENT, MaxAllocationSize, FreeListOrder, bDefragEnabled);
 		ReadOnlyPool->Initialize();
 
 		PoolAllocators[(int)EPoolType::ReadOnly] = ReadOnlyPool;
@@ -1524,7 +1526,7 @@ FD3D12TextureAllocatorPool::FD3D12TextureAllocatorPool(FD3D12Device* Device, FRH
 		uint64 PoolSize = GD3D12PoolAllocatorRTUAVTextureVRAMPoolSize;
 		uint64 MaxAllocationSize = GD3D12PoolAllocatorRTUAVTextureVRAMMaxAllocationSize;
 		bool bDefragEnabled = true;
-		FD3D12PoolAllocator* RTPool = new FD3D12PoolAllocator(Device, GetVisibilityMask(), InitConfig, Name, AllocationStrategy, PoolSize, D3D12_DEFAULT_RESOURCE_PLACEMENT_ALIGNMENT, MaxAllocationSize, bDefragEnabled);
+		FD3D12PoolAllocator* RTPool = new FD3D12PoolAllocator(Device, GetVisibilityMask(), InitConfig, Name, AllocationStrategy, PoolSize, D3D12_DEFAULT_RESOURCE_PLACEMENT_ALIGNMENT, MaxAllocationSize, FreeListOrder, bDefragEnabled);
 		RTPool->Initialize();
 
 		PoolAllocators[(int)EPoolType::RenderTarget] = RTPool;
@@ -1538,7 +1540,7 @@ FD3D12TextureAllocatorPool::FD3D12TextureAllocatorPool(FD3D12Device* Device, FRH
 		uint64 PoolSize = GD3D12PoolAllocatorRTUAVTextureVRAMPoolSize;
 		uint64 MaxAllocationSize = GD3D12PoolAllocatorRTUAVTextureVRAMMaxAllocationSize;
 		bool bDefragEnabled = true;
-		FD3D12PoolAllocator* UAVPool = new FD3D12PoolAllocator(Device, GetVisibilityMask(), InitConfig, Name, AllocationStrategy, PoolSize, D3D12_DEFAULT_RESOURCE_PLACEMENT_ALIGNMENT, MaxAllocationSize, bDefragEnabled);
+		FD3D12PoolAllocator* UAVPool = new FD3D12PoolAllocator(Device, GetVisibilityMask(), InitConfig, Name, AllocationStrategy, PoolSize, D3D12_DEFAULT_RESOURCE_PLACEMENT_ALIGNMENT, MaxAllocationSize, FreeListOrder, bDefragEnabled);
 		UAVPool->Initialize();
 
 		PoolAllocators[(int)EPoolType::UAV] = UAVPool;
@@ -1564,7 +1566,7 @@ HRESULT FD3D12TextureAllocatorPool::AllocateTexture(
 	const TCHAR* Name)
 {
 	// The top mip level must be less than 64 KB to use 4 KB alignment
-	bool b4KAligment = TextureCanBe4KAligned(Desc, UEFormat);
+	bool b4KAligment = TextureCanBe4KAligned(Desc, (EPixelFormat)UEFormat);
 	Desc.Alignment = b4KAligment ?	D3D12_SMALL_RESOURCE_PLACEMENT_ALIGNMENT : D3D12_DEFAULT_RESOURCE_PLACEMENT_ALIGNMENT;
 	const D3D12_RESOURCE_ALLOCATION_INFO Info = GetParentDevice()->GetDevice()->GetResourceAllocationInfo(0, 1, &Desc);
 
@@ -1575,7 +1577,7 @@ HRESULT FD3D12TextureAllocatorPool::AllocateTexture(
 						 Desc.Flags & D3D12_RESOURCE_FLAG_ALLOW_DEPTH_STENCIL);
 	ED3D12ResourceStateMode ResourceStateMode = bIsReadOnly ? ED3D12ResourceStateMode::Default : ED3D12ResourceStateMode::MultiState;
 	EPoolType PoolType = bIsReadOnly ? (b4KAligment ? EPoolType::ReadOnly4K : EPoolType::ReadOnly) : (bIsRenderTarget ? EPoolType::RenderTarget : EPoolType::UAV);
-	PoolAllocators[(int)PoolType]->AllocResource(D3D12_HEAP_TYPE_DEFAULT, Desc, Info.SizeInBytes, Info.Alignment, ResourceStateMode, InitialState, ClearValue, Name, TextureLocation);
+	PoolAllocators[(int)PoolType]->AllocateResource(D3D12_HEAP_TYPE_DEFAULT, Desc, Info.SizeInBytes, Info.Alignment, ResourceStateMode, InitialState, ClearValue, Name, TextureLocation);
 
 	return S_OK;
 }
@@ -1790,7 +1792,7 @@ HRESULT FD3D12TextureAllocatorPool::AllocateTexture(D3D12_RESOURCE_DESC Desc, co
 		Desc.SampleDesc.Count == 1)// Multi-Sample texures have much larger alignment requirements (4MB vs 64KB)
 	{
 		// The top mip level must be less than 64k
-		if (TextureCanBe4KAligned(Desc, UEFormat))
+		if (TextureCanBe4KAligned(Desc, (EPixelFormat)UEFormat))
 		{
 			Desc.Alignment = D3D12_SMALL_RESOURCE_PLACEMENT_ALIGNMENT; // request 4k alignment
 			return ReadOnlyTexturePool.AllocateTexture(Desc, ClearValue, TextureLocation, InitialState, Name);
diff --git a/Engine/Source/Runtime/D3D12RHI/Private/D3D12Buffer.cpp b/Engine/Source/Runtime/D3D12RHI/Private/D3D12Buffer.cpp
--- a/Engine/Source/Runtime/D3D12RHI/Private/D3D12Buffer.cpp
+++ b/Engine/Source/Runtime/D3D12RHI/Private/D3D12Buffer.cpp
@@ -150,8 +150,11 @@ void FD3D12Adapter::AllocateBuffer(FD3D12Device* Device,
 	FRHIResourceCreateInfo& CreateInfo,
 	uint32 Alignment,
 	FD3D12Buffer* Buffer,
-	FD3D12ResourceLocation& ResourceLocation)
+	FD3D12ResourceLocation& ResourceLocation,
+	ID3D12ResourceAllocator* ResourceAllocator)
 {
+	TRACE_CPUPROFILER_EVENT_SCOPE(D3D12RHI::AllocateBuffer);
+
 	// Explicitly check that the size is nonzero before allowing CreateBuffer to opaquely fail.
 	check(Size > 0);
 
@@ -159,6 +162,7 @@ void FD3D12Adapter::AllocateBuffer(FD3D12Device* Device,
 
 	if (bIsDynamic)
 	{
+		check(ResourceAllocator == nullptr);
 		check(InResourceStateMode != ED3D12ResourceStateMode::MultiState);
 		check(InCreateState == D3D12_RESOURCE_STATE_GENERIC_READ);
 		void* pData = GetUploadHeapAllocator(Device->GetGPUIndex()).AllocUploadResource(Size, Alignment, ResourceLocation);
@@ -175,7 +179,14 @@ void FD3D12Adapter::AllocateBuffer(FD3D12Device* Device,
 	}
 	else
 	{
-		Device->GetDefaultBufferAllocator().AllocDefaultResource(D3D12_HEAP_TYPE_DEFAULT, InDesc, (EBufferUsageFlags)InUsage, InResourceStateMode, InCreateState, ResourceLocation, Alignment, CreateInfo.DebugName);
+		if (ResourceAllocator)
+		{
+			ResourceAllocator->AllocateResource(D3D12_HEAP_TYPE_DEFAULT, InDesc, InDesc.Width, Alignment, InResourceStateMode, InCreateState, nullptr, CreateInfo.DebugName, ResourceLocation);
+		}
+		else
+		{
+			Device->GetDefaultBufferAllocator().AllocDefaultResource(D3D12_HEAP_TYPE_DEFAULT, InDesc, (EBufferUsageFlags)InUsage, InResourceStateMode, InCreateState, ResourceLocation, Alignment, CreateInfo.DebugName);
+		}
 		ResourceLocation.SetOwner(Buffer);
 		check(ResourceLocation.GetSize() == Size);
 	}
@@ -189,13 +200,21 @@ FD3D12Buffer* FD3D12Adapter::CreateRHIBuffer(FRHICommandListImmediate* RHICmdLis
 	uint32 InUsage,
 	ED3D12ResourceStateMode InResourceStateMode,
 	ERHIAccess InResourceState,
-	FRHIResourceCreateInfo& CreateInfo)
+	FRHIResourceCreateInfo& CreateInfo,
+	ID3D12ResourceAllocator* ResourceAllocator)
 {
+	TRACE_CPUPROFILER_EVENT_SCOPE(D3D12RHI::CreateRHIBuffer);
+
 	SCOPE_CYCLE_COUNTER(STAT_D3D12CreateBufferTime);
 
+	check(InDesc.Width == Size);
+
 	const bool bIsDynamic = (InUsage & BUF_AnyDynamic) ? true : false;
 	const uint32 FirstGPUIndex = CreateInfo.GPUMask.GetFirstIndex();
 
+	// Transient resources don't get any actual D3D12 resource asigned at this point yet
+	bool bIsTransient = GSupportsTransientResourceAliasing && (InUsage & BUF_Transient);
+	
 	// Does this resource support tracking?
 	const bool bSupportResourceStateTracking = !bIsDynamic && FD3D12DefaultBufferAllocator::IsPlacedResource(InDesc.Flags, InResourceStateMode);
 
@@ -208,6 +227,9 @@ FD3D12Buffer* FD3D12Adapter::CreateRHIBuffer(FRHICommandListImmediate* RHICmdLis
 	FD3D12Buffer* BufferOut = nullptr;
 	if (bIsDynamic)
 	{
+		// Assume not transient and dynamic
+		check(!bIsTransient);
+		
 		FD3D12Buffer* NewBuffer0 = nullptr;
 		BufferOut = CreateLinkedObject<FD3D12Buffer>(CreateInfo.GPUMask, [&](FD3D12Device* Device)
 		{
@@ -216,7 +238,7 @@ FD3D12Buffer* FD3D12Adapter::CreateRHIBuffer(FRHICommandListImmediate* RHICmdLis
 
 			if (Device->GetGPUIndex() == FirstGPUIndex)
 			{
-				AllocateBuffer(Device, InDesc, Size, InUsage, InResourceStateMode, InitialState, CreateInfo, Alignment, NewBuffer, NewBuffer->ResourceLocation);
+				AllocateBuffer(Device, InDesc, Size, InUsage, InResourceStateMode, InitialState, CreateInfo, Alignment, NewBuffer, NewBuffer->ResourceLocation, ResourceAllocator);
 				NewBuffer0 = NewBuffer;
 			}
 			else
@@ -237,8 +259,8 @@ FD3D12Buffer* FD3D12Adapter::CreateRHIBuffer(FRHICommandListImmediate* RHICmdLis
 		{
 			FD3D12Buffer* NewBuffer = new FD3D12Buffer(Device, Size, InUsage, Stride);
 			NewBuffer->BufferAlignment = Alignment;
-
-			AllocateBuffer(Device, InDesc, Size, InUsage, InResourceStateMode, CreateState, CreateInfo, Alignment, NewBuffer, NewBuffer->ResourceLocation);
+			AllocateBuffer(Device, InDesc, Size, InUsage, InResourceStateMode, CreateState, CreateInfo, Alignment, NewBuffer, NewBuffer->ResourceLocation, ResourceAllocator);
+			NewBuffer->ResourceLocation.SetTransient(bIsTransient);
 
 			// Unlock immediately if no initial data
 			if (CreateInfo.ResourceArray == nullptr)
@@ -252,6 +274,7 @@ FD3D12Buffer* FD3D12Adapter::CreateRHIBuffer(FRHICommandListImmediate* RHICmdLis
 
 	if (CreateInfo.ResourceArray)
 	{
+		check(!bIsTransient);
 		if (bIsDynamic == false && BufferOut->ResourceLocation.IsValid())
 		{
 			check(Size == CreateInfo.ResourceArray->GetResourceDataSize());
@@ -370,33 +393,35 @@ void FD3D12Buffer::ReleaseUnderlyingResource()
 	}
 }
 
-D3D12_RESOURCE_DESC CreateBufferResourceDesc(uint32 Size, uint32 Usage)
+void FD3D12Buffer::GetResourceDescAndAlignment(uint64 InSize, uint32 InStride, EBufferUsageFlags& InUsage, D3D12_RESOURCE_DESC& ResourceDesc, uint32& Alignment)
 {
-	D3D12_RESOURCE_DESC Desc = CD3DX12_RESOURCE_DESC::Buffer(Size);
+	ResourceDesc = CD3DX12_RESOURCE_DESC::Buffer(InSize);
 
-	if (Usage & BUF_UnorderedAccess)
+	if (InUsage & BUF_UnorderedAccess)
 	{
-		Desc.Flags |= D3D12_RESOURCE_FLAG_ALLOW_UNORDERED_ACCESS;
+		ResourceDesc.Flags |= D3D12_RESOURCE_FLAG_ALLOW_UNORDERED_ACCESS;
 
 		static bool bRequiresRawView = (GMaxRHIFeatureLevel < ERHIFeatureLevel::SM5);
 		if (bRequiresRawView)
 		{
 			// Force the buffer to be a raw, byte address buffer
-			Usage |= BUF_ByteAddressBuffer;
+			InUsage |= BUF_ByteAddressBuffer;
 		}
 	}
 
-	if ((Usage & BUF_ShaderResource) == 0)
+	if ((InUsage & BUF_ShaderResource) == 0)
 	{
-		Desc.Flags |= D3D12_RESOURCE_FLAG_DENY_SHADER_RESOURCE;
+		ResourceDesc.Flags |= D3D12_RESOURCE_FLAG_DENY_SHADER_RESOURCE;
 	}
 
-	if (Usage & BUF_DrawIndirect)
+	if (InUsage & BUF_DrawIndirect)
 	{
-		Desc.Flags |= D3D12RHI_RESOURCE_FLAG_ALLOW_INDIRECT_BUFFER;
+		ResourceDesc.Flags |= D3D12RHI_RESOURCE_FLAG_ALLOW_INDIRECT_BUFFER;
 	}
 
-	return Desc;
+	// Structured buffers, non-ByteAddress buffers, need to be aligned to their stride to ensure that they can be addressed correctly with element based offsets.
+	Alignment = (InStride > 0) && (((InUsage & BUF_StructuredBuffer) != 0) || ((InUsage & (BUF_ByteAddressBuffer | BUF_DrawIndirect)) == 0)) ? InStride : 4;
+
 }
 
 FBufferRHIRef FD3D12DynamicRHI::RHICreateBuffer(uint32 Size, EBufferUsageFlags Usage, uint32 Stride, ERHIAccess InResourceState, FRHIResourceCreateInfo& CreateInfo)
@@ -419,11 +444,17 @@ FBufferRHIRef FD3D12DynamicRHI::CreateBuffer(FRHICommandListImmediate* RHICmdLis
 			});
 	}
 
-	const D3D12_RESOURCE_DESC Desc = CreateBufferResourceDesc(Size, Usage);
-	// Structured buffers, non-ByteAddress buffers, need to be aligned to their stride to ensure that they can be addressed correctly with element based offsets.
-	const uint32 Alignment = (Stride > 0) && (((Usage & BUF_StructuredBuffer) != 0) || ((Usage & (BUF_ByteAddressBuffer | BUF_DrawIndirect)) == 0)) ? Stride : 4;
+	ID3D12ResourceAllocator* ResourceAllocator = nullptr;
+	return CreateD3D12Buffer(RHICmdList, Size, Usage, Stride, InResourceState, CreateInfo, ResourceAllocator);
+}
+
+FD3D12Buffer* FD3D12DynamicRHI::CreateD3D12Buffer(class FRHICommandListImmediate* RHICmdList, uint32 Size, EBufferUsageFlags Usage, uint32 Stride, ERHIAccess ResourceState, FRHIResourceCreateInfo& CreateInfo, ID3D12ResourceAllocator* ResourceAllocator)
+{
+	D3D12_RESOURCE_DESC Desc;
+	uint32 Alignment;
+	FD3D12Buffer::GetResourceDescAndAlignment(Size, Stride, Usage, Desc, Alignment);
 
-	FD3D12Buffer* Buffer = GetAdapter().CreateRHIBuffer(RHICmdList, Desc, Alignment, Stride, Size, Usage, ED3D12ResourceStateMode::Default, InResourceState, CreateInfo);
+	FD3D12Buffer* Buffer = GetAdapter().CreateRHIBuffer(RHICmdList, Desc, Alignment, Stride, Size, Usage, ED3D12ResourceStateMode::Default, ResourceState, CreateInfo, ResourceAllocator);
 	if (Buffer->ResourceLocation.IsTransient())
 	{
 		// TODO: this should ideally be set in platform-independent code, since this tracking is for the high level
diff --git a/Engine/Source/Runtime/D3D12RHI/Private/D3D12CommandContext.cpp b/Engine/Source/Runtime/D3D12RHI/Private/D3D12CommandContext.cpp
--- a/Engine/Source/Runtime/D3D12RHI/Private/D3D12CommandContext.cpp
+++ b/Engine/Source/Runtime/D3D12RHI/Private/D3D12CommandContext.cpp
@@ -513,6 +513,7 @@ void FD3D12CommandContextBase::RHIBeginFrame()
 
 		Device->GetDefaultBufferAllocator().BeginFrame();
 		Device->GetTextureAllocator().BeginFrame();
+		Device->GetTransientMemoryPoolManager().BeginFrame();
 
 #if D3D12_RHI_RAYTRACING
 		Device->GetRayTracingCompactionRequestHandler()->Update(*ContextAtIndex);
@@ -623,6 +624,7 @@ void FD3D12CommandContextBase::RHIEndFrame()
 		}
 
 		Device->GetTextureAllocator().CleanUpAllocations();
+		Device->GetTransientMemoryPoolManager().EndFrame();
 
 		// Only delete free blocks when not used in the last 2 frames, to make sure we are not allocating and releasing
 		// the same blocks every frame.
@@ -720,6 +722,7 @@ void FD3D12CommandContextBase::UpdateMemoryStats()
 #endif
 
 		Device->GetDefaultBufferAllocator().UpdateMemoryStats();
+		Device->GetTransientMemoryPoolManager().UpdateMemoryStats();
 	}
 #endif
 }
@@ -956,6 +959,15 @@ void FD3D12DynamicRHI::RHIReleaseTransition(FRHITransition* Transition)
 	Transition->GetPrivateData<FD3D12TransitionData>()->~FD3D12TransitionData();
 }
 
+
+IRHITransientResourceAllocator* FD3D12DynamicRHI::RHICreateTransientResourceAllocator()
+{
+	// Handle mGPU here?
+	FD3D12Device* Device = GetAdapter().GetDevice(0);
+	return new FD3D12TransientResourceAllocator(Device->GetTransientMemoryPoolManager());
+}
+
+
 //////////////////////////////////////////////////////////////////////////////////////////////////////////
 //
 // FD3D12CommandContextRedirector
diff --git a/Engine/Source/Runtime/D3D12RHI/Private/D3D12CommandContext.h b/Engine/Source/Runtime/D3D12RHI/Private/D3D12CommandContext.h
--- a/Engine/Source/Runtime/D3D12RHI/Private/D3D12CommandContext.h
+++ b/Engine/Source/Runtime/D3D12RHI/Private/D3D12CommandContext.h
@@ -251,6 +251,8 @@ public:
 	virtual void RHIBeginTransitions(TArrayView<const FRHITransition*> Transitions) final override;
 	virtual void RHIEndTransitions(TArrayView<const FRHITransition*> Transitions) final override;
 
+	virtual void RHIReleaseTransientResourceAllocator(IRHITransientResourceAllocator* InAllocator) final override;
+
 	// IRHIComputeContext interface
 	virtual void RHISetComputeShader(FRHIComputeShader* ComputeShader) final override;
 	virtual void RHISetComputePipelineState(FRHIComputePipelineState* ComputePipelineState) final override;
@@ -877,4 +879,4 @@ struct FD3D12TransitionData
 	TRefCountPtr<FD3D12Fence> Fence;
 
 	bool bCrossPipeline = false;
-};
\ No newline at end of file
+};
diff --git a/Engine/Source/Runtime/D3D12RHI/Private/D3D12CommandList.cpp b/Engine/Source/Runtime/D3D12RHI/Private/D3D12CommandList.cpp
--- a/Engine/Source/Runtime/D3D12RHI/Private/D3D12CommandList.cpp
+++ b/Engine/Source/Runtime/D3D12RHI/Private/D3D12CommandList.cpp
@@ -34,10 +34,10 @@ void FD3D12CommandListHandle::AddUAVBarrier()
 	CommandListData->CurrentOwningContext->numBarriers++;
 }
 
-void FD3D12CommandListHandle::AddAliasingBarrier(FD3D12Resource* pResource)
+void FD3D12CommandListHandle::AddAliasingBarrier(FD3D12Resource* InResourceBefore, FD3D12Resource* InResourceAfter)
 {
 	check(CommandListData);
-	CommandListData->ResourceBarrierBatcher.AddAliasingBarrier(pResource->GetResource());
+	CommandListData->ResourceBarrierBatcher.AddAliasingBarrier(InResourceBefore ? InResourceBefore->GetResource() : nullptr, InResourceAfter->GetResource());
 	CommandListData->CurrentOwningContext->numBarriers++;
 }
 
diff --git a/Engine/Source/Runtime/D3D12RHI/Private/D3D12CommandList.h b/Engine/Source/Runtime/D3D12RHI/Private/D3D12CommandList.h
--- a/Engine/Source/Runtime/D3D12RHI/Private/D3D12CommandList.h
+++ b/Engine/Source/Runtime/D3D12RHI/Private/D3D12CommandList.h
@@ -570,7 +570,7 @@ public:
 	// Adds a UAV barrier to the barrier batch
 	void AddUAVBarrier();
 
-	void AddAliasingBarrier(FD3D12Resource* pResource);
+	void AddAliasingBarrier(FD3D12Resource* InResourceBefore, FD3D12Resource* InResourceAfter);
 
 	// Flushes the batched resource barriers to the current command list
 	void FlushResourceBarriers()
diff --git a/Engine/Source/Runtime/D3D12RHI/Private/D3D12Commands.cpp b/Engine/Source/Runtime/D3D12RHI/Private/D3D12Commands.cpp
--- a/Engine/Source/Runtime/D3D12RHI/Private/D3D12Commands.cpp
+++ b/Engine/Source/Runtime/D3D12RHI/Private/D3D12Commands.cpp
@@ -290,12 +290,97 @@ void ProcessResource(FD3D12CommandContext& Context, const FRHITransitionInfo& In
 	}
 }
 
+void ProcessDiscardState(FD3D12CommandContext& Context, const FRHITransitionInfo& Info)
+{
+	// No discard before or after state, then done
+	if (!EnumHasAnyFlags(Info.AccessBefore, ERHIAccess::Discard) && !EnumHasAnyFlags(Info.AccessAfter, ERHIAccess::Discard))
+	{
+		return;
+	}
+
+	FD3D12BaseShaderResource* BaseShaderResource = nullptr;
+	switch (Info.Type)
+	{
+	case FRHITransitionInfo::EType::Buffer:
+	{
+		FD3D12Buffer* Buffer = Context.RetrieveObject<FD3D12Buffer>(Info.Buffer);
+		check(Buffer);
+		BaseShaderResource = Buffer;
+		break;
+	}
+	case FRHITransitionInfo::EType::Texture:
+	{
+		FD3D12TextureBase* Texture = Context.RetrieveTextureBase(Info.Texture);
+		check(Texture);
+		BaseShaderResource = Texture;
+		break;
+	}
+	default:
+		checkNoEntry();
+		break;
+	}
+
+	if (EnumHasAnyFlags(Info.AccessBefore, ERHIAccess::Discard))
+	{
+		TRACE_CPUPROFILER_EVENT_SCOPE(D3D12RHI::AcquireTransient);
+
+		FD3D12Resource* Resource = BaseShaderResource->ResourceLocation.GetResource();
+
+		// Try and see if the resource has overlapping resources
+		if (BaseShaderResource->ResourceLocation.IsTransient() && BaseShaderResource->ResourceLocation.GetAllocatorType() == FD3D12ResourceLocation::AT_Pool)
+		{
+			// We know it's a transient resource allocator
+			FD3D12TransientResourceAllocator* Allocator = (FD3D12TransientResourceAllocator*) BaseShaderResource->ResourceLocation.GetPoolAllocator();
+
+			// Get the overlapping resources and append aliasing barriers for each of them
+			TArrayView<FD3D12Resource*> OverlappingResources = Allocator->GetOverlappingResources(BaseShaderResource);
+			for (FD3D12Resource* OverlappingResource : OverlappingResources)
+			{
+				Context.CommandListHandle.AddAliasingBarrier(OverlappingResource, Resource);
+			}
+		}		
+
+		bool bIsAsyncCompute = false;
+		D3D12_RESOURCE_STATES AfterState = GetD3D12ResourceState(Info.AccessAfter, bIsAsyncCompute);
+		FD3D12DynamicRHI::TransitionResource(Context.CommandListHandle, Resource, D3D12_RESOURCE_STATE_TBD, AfterState, D3D12_RESOURCE_BARRIER_ALL_SUBRESOURCES, FD3D12DynamicRHI::ETransitionMode::Apply);
+		Context.CommandListHandle.FlushResourceBarriers();
+		
+		// Add clear or discard operation depending on flag
+		if (EnumHasAnyFlags(Info.Flags, EResourceTransitionFlags::Clear))
+		{
+			// add correct ops for clear
+			check(false);
+		}
+		else
+		{
+			check(Info.Flags & EResourceTransitionFlags::Discard);
+			Context.CommandListHandle->DiscardResource(Resource->GetResource(), nullptr);
+		}
+	}
+	else
+	{
+		TRACE_CPUPROFILER_EVENT_SCOPE(D3D12RHI::DiscardTransient);
+
+		// Remove from caches
+		Context.ConditionalClearShaderResource(&BaseShaderResource->ResourceLocation);
+
+		// Release the resource location - transient allocator internally will take care of correct caching
+		BaseShaderResource->ResourceLocation.Clear();
+
+		// Release the views after clearing resource location
+		BaseShaderResource->ResourceRenamed(&BaseShaderResource->ResourceLocation);
+	}
+}
+
 static void HandleResourceTransitions(FD3D12CommandContext& Context, const FD3D12TransitionData* TransitionData, D3D12_RESOURCE_STATES SkipFastClearEliminateState, bool& bUAVBarrier)
 {
 	const bool bDstAllPipelines = EnumHasAllFlags(TransitionData->DstPipelines, ERHIPipeline::All);
 
 	for (const FRHITransitionInfo& Info : TransitionData->Infos)
 	{
+		// Process the discard state
+		ProcessDiscardState(Context, Info);
+
 		const bool bUAVAccess = EnumHasAnyFlags(Info.AccessAfter, ERHIAccess::UAVMask);
 
 		// If targeting all pipelines and UAV access then add UAV barrier even with RHI based transitions
@@ -432,6 +517,12 @@ void FD3D12CommandContext::RHIEndTransitions(TArrayView<const FRHITransition*> T
 	}
 }
 
+void FD3D12CommandContext::RHIReleaseTransientResourceAllocator(IRHITransientResourceAllocator* InAllocator)
+{
+	// Should be fine to delete now
+	delete (FD3D12TransientResourceAllocator*)InAllocator;
+}
+
 void FD3D12CommandContext::RHISetStaticUniformBuffers(const FUniformBufferStaticBindings& InUniformBuffers)
 {
 	FMemory::Memzero(GlobalUniformBuffers.GetData(), GlobalUniformBuffers.Num() * sizeof(FRHIUniformBuffer*));
diff --git a/Engine/Source/Runtime/D3D12RHI/Private/D3D12Device.cpp b/Engine/Source/Runtime/D3D12RHI/Private/D3D12Device.cpp
--- a/Engine/Source/Runtime/D3D12RHI/Private/D3D12Device.cpp
+++ b/Engine/Source/Runtime/D3D12RHI/Private/D3D12Device.cpp
@@ -42,6 +42,7 @@ FD3D12Device::FD3D12Device(FRHIGPUMask InGPUMask, FD3D12Adapter* InAdapter) :
 	SamplerID(0),
 	DefaultFastAllocator(this, FRHIGPUMask::All(), D3D12_HEAP_TYPE_UPLOAD, 1024 * 1024 * 4),
 	TextureAllocator(this, FRHIGPUMask::All()),
+	TransientMemoryPoolManager(this, FRHIGPUMask::All()),
 	GPUProfilingData(this)
 {
 	for (uint32 QueueType = 0; QueueType < (uint32)ED3D12CommandQueueType::Count; ++QueueType)
@@ -64,6 +65,8 @@ FD3D12Device::~FD3D12Device()
 	// Cleanup the allocator near the end, as some resources may be returned to the allocator or references are shared by multiple GPUs
 	DefaultBufferAllocator.FreeDefaultBufferPools();
 
+	TransientMemoryPoolManager.Destroy();
+
 	DefaultFastAllocator.Destroy();
 
 	TextureAllocator.CleanUpAllocations();
@@ -483,3 +486,32 @@ void FD3D12Device::BlockUntilIdle()
 	GetCopyCommandListManager().WaitForCommandQueueFlush();
 	GetAsyncCommandListManager().WaitForCommandQueueFlush();
 }
+
+D3D12_RESOURCE_ALLOCATION_INFO FD3D12Device::GetResourceAllocationInfo(const D3D12_RESOURCE_DESC& InDesc)
+{
+	uint64 Hash = CityHash64((const char*)&InDesc, sizeof(D3D12_RESOURCE_DESC));
+
+	// By default there'll be more threads trying to read this than to write it.
+	ResourceAllocationInfoMapMutex.ReadLock();
+	D3D12_RESOURCE_ALLOCATION_INFO* CachedInfo = ResourceAllocationInfoMap.Find(Hash);
+	ResourceAllocationInfoMapMutex.ReadUnlock();
+
+	if (CachedInfo)
+	{
+		return *CachedInfo;
+	}
+	else
+	{
+		D3D12_RESOURCE_ALLOCATION_INFO Result = GetDevice()->GetResourceAllocationInfo(0, 1, &InDesc);
+
+		ResourceAllocationInfoMapMutex.WriteLock();
+		// Try search again with write lock because could have been added already
+		CachedInfo = ResourceAllocationInfoMap.Find(Hash);
+		if (CachedInfo == nullptr)
+		{
+			ResourceAllocationInfoMap.Add(Hash, Result);
+		}
+		ResourceAllocationInfoMapMutex.WriteUnlock();
+		return Result;
+	}
+}
\ No newline at end of file
diff --git a/Engine/Source/Runtime/D3D12RHI/Private/D3D12Device.h b/Engine/Source/Runtime/D3D12RHI/Private/D3D12Device.h
--- a/Engine/Source/Runtime/D3D12RHI/Private/D3D12Device.h
+++ b/Engine/Source/Runtime/D3D12RHI/Private/D3D12Device.h
@@ -80,6 +80,7 @@ public:
 	inline FD3D12CommandListManager& GetAsyncCommandListManager() { return *AsyncCommandListManager; }
 	inline FD3D12CommandAllocatorManager& GetTextureStreamingCommandAllocatorManager() { return TextureStreamingCommandAllocatorManager; }
 	inline FD3D12DefaultBufferAllocator& GetDefaultBufferAllocator() { return DefaultBufferAllocator; }
+	inline FD3D12TransientMemoryPoolManager& GetTransientMemoryPoolManager() { return TransientMemoryPoolManager; }
 	inline FD3D12GlobalOnlineSamplerHeap& GetGlobalSamplerHeap() { return GlobalSamplerHeap; }
 	inline FD3D12GlobalHeap& GetGlobalViewHeap() { return GlobalViewHeap; }
 
@@ -125,6 +126,8 @@ public:
 	FD3D12SamplerState* CreateSampler(const FSamplerStateInitializerRHI& Initializer);
 	void CreateSamplerInternal(const D3D12_SAMPLER_DESC& Desc, D3D12_CPU_DESCRIPTOR_HANDLE Descriptor);
 
+	D3D12_RESOURCE_ALLOCATION_INFO GetResourceAllocationInfo(const D3D12_RESOURCE_DESC& InDesc);
+
 	void BlockUntilIdle();
 
 	FORCEINLINE FD3DGPUProfiler& GetGPUProfiler() { return GPUProfilingData; }
@@ -169,6 +172,10 @@ protected:
 	TMap< D3D12_SAMPLER_DESC, TRefCountPtr<FD3D12SamplerState> > SamplerMap;
 	uint32 SamplerID;
 
+	/** Hashmap used to cache resource allocation size information */
+	FRWLock ResourceAllocationInfoMapMutex;
+	TMap< uint64, D3D12_RESOURCE_ALLOCATION_INFO> ResourceAllocationInfoMap;
+
 	// set by UpdateMSAASettings(), get by GetMSAAQuality()
 	// [SampleCount] = Quality, 0xffffffff if not supported
 	uint32 AvailableMSAAQualities[DX_MAX_MSAA_COUNT + 1];
@@ -189,8 +196,8 @@ protected:
 	void ReleasePooledUniformBuffers();
 
 	FD3D12FastAllocator DefaultFastAllocator;
-
 	FD3D12TextureAllocatorPool TextureAllocator;
+	FD3D12TransientMemoryPoolManager TransientMemoryPoolManager;
 
 	FD3D12ResidencyManager ResidencyManager;
 
diff --git a/Engine/Source/Runtime/D3D12RHI/Private/D3D12PoolAllocator.cpp b/Engine/Source/Runtime/D3D12RHI/Private/D3D12PoolAllocator.cpp
--- a/Engine/Source/Runtime/D3D12RHI/Private/D3D12PoolAllocator.cpp
+++ b/Engine/Source/Runtime/D3D12RHI/Private/D3D12PoolAllocator.cpp
@@ -12,8 +12,8 @@
 //-----------------------------------------------------------------------------
 
 FD3D12MemoryPool::FD3D12MemoryPool(FD3D12Device* ParentDevice, FRHIGPUMask VisibleNodes, const FD3D12ResourceInitConfig& InInitConfig, const FString& InName,
-	EResourceAllocationStrategy InAllocationStrategy, int16 InPoolIndex, uint64 InPoolSize, uint32 InPoolAlignment)
-	: FRHIMemoryPool(InPoolIndex, InPoolSize, InPoolAlignment), FD3D12DeviceChild(ParentDevice), FD3D12MultiNodeGPUObject(ParentDevice->GetGPUMask(), VisibleNodes)
+	EResourceAllocationStrategy InAllocationStrategy, int16 InPoolIndex, uint64 InPoolSize, uint32 InPoolAlignment, EFreeListOrder InFreeListOrder)
+	: FRHIMemoryPool(InPoolIndex, InPoolSize, InPoolAlignment, InFreeListOrder), FD3D12DeviceChild(ParentDevice), FD3D12MultiNodeGPUObject(ParentDevice->GetGPUMask(), VisibleNodes)
 	, InitConfig(InInitConfig), Name(InName), AllocationStrategy(InAllocationStrategy), LastUsedFrameFence(0)
 {
 }
@@ -173,8 +173,8 @@ EResourceAllocationStrategy FD3D12PoolAllocator::GetResourceAllocationStrategy(D
 
 
 FD3D12PoolAllocator::FD3D12PoolAllocator(FD3D12Device* ParentDevice, FRHIGPUMask VisibleNodes, const FD3D12ResourceInitConfig& InInitConfig, const FString& InName,
-	EResourceAllocationStrategy InAllocationStrategy, uint64 InPoolSize, uint32 InPoolAlignment, uint32 InMaxAllocationSize, bool bInDefragEnabled) :
-	FRHIPoolAllocator(InPoolSize, InPoolAlignment, InMaxAllocationSize, bInDefragEnabled), 
+	EResourceAllocationStrategy InAllocationStrategy, uint64 InPoolSize, uint32 InPoolAlignment, uint32 InMaxAllocationSize, FRHIMemoryPool::EFreeListOrder InFreeListOrder, bool bInDefragEnabled) :
+	FRHIPoolAllocator(InPoolSize, InPoolAlignment, InMaxAllocationSize, InFreeListOrder, bInDefragEnabled), 
 	FD3D12DeviceChild(ParentDevice), 
 	FD3D12MultiNodeGPUObject(ParentDevice->GetGPUMask(), VisibleNodes), 
 	InitConfig(InInitConfig), 
@@ -224,13 +224,15 @@ void FD3D12PoolAllocator::AllocDefaultResource(D3D12_HEAP_TYPE InHeapType, const
 #endif // D3D12_RHI_RAYTRACING
 #endif  // DO_CHECK
 
-	AllocResource(InHeapType, InDesc, InDesc.Width, InAllocationAlignment, InResourceStateMode, InCreateState, nullptr, InName, ResourceLocation);
+	AllocateResource(InHeapType, InDesc, InDesc.Width, InAllocationAlignment, InResourceStateMode, InCreateState, nullptr, InName, ResourceLocation);
 }
 
 
-void FD3D12PoolAllocator::AllocResource(D3D12_HEAP_TYPE InHeapType, const D3D12_RESOURCE_DESC& InDesc, uint64 InSize, uint32 InAllocationAlignment, ED3D12ResourceStateMode InResourceStateMode,
+void FD3D12PoolAllocator::AllocateResource(D3D12_HEAP_TYPE InHeapType, const D3D12_RESOURCE_DESC& InDesc, uint64 InSize, uint32 InAllocationAlignment, ED3D12ResourceStateMode InResourceStateMode,
 	D3D12_RESOURCE_STATES InCreateState, const D3D12_CLEAR_VALUE* InClearValue, const TCHAR* InName, FD3D12ResourceLocation& ResourceLocation)
 {
+	TRACE_CPUPROFILER_EVENT_SCOPE(D3D12RHI::AllocatePoolResource);
+
 	// If the resource location owns a block, this will deallocate it.
 	ResourceLocation.Clear();
 	if (InSize == 0)
@@ -246,11 +248,18 @@ void FD3D12PoolAllocator::AllocResource(D3D12_HEAP_TYPE InHeapType, const D3D12_
 	{
 		const bool bPlacedResource = (AllocationStrategy == EResourceAllocationStrategy::kPlacedResource);
 
+		uint32 AllocationAlignment = InAllocationAlignment;
+
 		// Ensure we're allocating from the correct pool
 		if (bPlacedResource)
 		{
 			// Writeable resources get separate ID3D12Resource* with their own resource state by using placed resources. Just make sure it's UAV, other flags are free to differ.
-			check(InDesc.Dimension != D3D12_RESOURCE_DIMENSION_BUFFER || (InDesc.Flags & InitConfig.ResourceFlags & D3D12_RESOURCE_FLAG_ALLOW_UNORDERED_ACCESS) != 0 || InHeapType == D3D12_HEAP_TYPE_READBACK);
+			check(InDesc.Dimension != D3D12_RESOURCE_DIMENSION_BUFFER || (InDesc.Flags & D3D12_RESOURCE_FLAG_ALLOW_UNORDERED_ACCESS) != 0 || InHeapType == D3D12_HEAP_TYPE_READBACK);
+
+			// If it's a placed resource then base offset will always be 0 from the actual d3d resource so ignore the allocation alignment - no extra offset required
+			// for creating the views!
+			check(InAllocationAlignment <= PoolAlignment);
+			AllocationAlignment = PoolAlignment;
 		}
 		else
 		{
@@ -260,7 +269,7 @@ void FD3D12PoolAllocator::AllocResource(D3D12_HEAP_TYPE InHeapType, const D3D12_
 
 		// Try to allocate in one of the pools
 		FRHIPoolAllocationData& AllocationData = ResourceLocation.GetPoolAllocatorPrivateData().PoolData;
-		if (TryAllocateInternal(InSize, InAllocationAlignment, AllocationData))
+		if (TryAllocateInternal(InSize, AllocationAlignment, AllocationData))
 		{
 			// Setup the resource location
 			ResourceLocation.SetType(FD3D12ResourceLocation::ResourceLocationType::eSubAllocation);
@@ -286,11 +295,7 @@ void FD3D12PoolAllocator::AllocResource(D3D12_HEAP_TYPE InHeapType, const D3D12_
 			{
 				check(ResourceLocation.GetResource() == nullptr);
 
-				FD3D12HeapAndOffset HeapAndOffset = GetBackingHeapAndAllocationOffsetInBytes(ResourceLocation);
-
-				FD3D12Resource* NewResource = nullptr;
-				VERIFYD3D12RESULT(Adapter->CreatePlacedResource(InDesc, HeapAndOffset.Heap, HeapAndOffset.Offset, InCreateState, InResourceStateMode, D3D12_RESOURCE_STATE_TBD, InClearValue, &NewResource, InName));
-
+				FD3D12Resource* NewResource = CreatePlacedResource(AllocationData, InDesc, InCreateState, InResourceStateMode, InClearValue, InName);
 				ResourceLocation.SetResource(NewResource);
 			}
 
@@ -313,7 +318,24 @@ void FD3D12PoolAllocator::AllocResource(D3D12_HEAP_TYPE InHeapType, const D3D12_
 }
 
 
-void FD3D12PoolAllocator::Deallocate(FD3D12ResourceLocation& ResourceLocation)
+FD3D12Resource* FD3D12PoolAllocator::CreatePlacedResource(
+	const FRHIPoolAllocationData& InAllocationData,
+	const D3D12_RESOURCE_DESC& InDesc, 
+	D3D12_RESOURCE_STATES InCreateState, 
+	ED3D12ResourceStateMode InResourceStateMode, 
+	const D3D12_CLEAR_VALUE* InClearValue, 
+	const TCHAR* InName)
+{
+	FD3D12Adapter* Adapter = GetParentDevice()->GetParentAdapter();
+	FD3D12HeapAndOffset HeapAndOffset = GetBackingHeapAndAllocationOffsetInBytes(InAllocationData);
+
+	FD3D12Resource* NewResource = nullptr;
+	VERIFYD3D12RESULT(Adapter->CreatePlacedResource(InDesc, HeapAndOffset.Heap, HeapAndOffset.Offset, InCreateState, InResourceStateMode, D3D12_RESOURCE_STATE_TBD, InClearValue, &NewResource, InName));
+	return NewResource;
+}
+
+
+void FD3D12PoolAllocator::DeallocateResource(FD3D12ResourceLocation& ResourceLocation)
 {
 	check(IsOwner(ResourceLocation));
 
@@ -388,7 +410,7 @@ void FD3D12PoolAllocator::Deallocate(FD3D12ResourceLocation& ResourceLocation)
 FRHIMemoryPool* FD3D12PoolAllocator::CreateNewPool(int16 InPoolIndex)
 {
 	FD3D12MemoryPool* NewPool = new FD3D12MemoryPool(GetParentDevice(),	GetVisibilityMask(), InitConfig,
-		Name, AllocationStrategy, InPoolIndex, PoolSize, PoolAlignment);
+		Name, AllocationStrategy, InPoolIndex, PoolSize, PoolAlignment, FreeListOrder);
 	NewPool->Init();
 	return NewPool;
 }
@@ -402,7 +424,7 @@ bool FD3D12PoolAllocator::HandleDefragRequest(FRHIPoolAllocationData* InSourceBl
 	FD3D12Resource* CurrentResource = Owner->GetResource();
 
 	// Release the current allocation (will only be freed on the next frame fence)
-	Deallocate(*Owner);
+	DeallocateResource(*Owner);
 
 	// Move temp allocation data to allocation data of the owner (part of different allocator now)		
 	bool bLocked = true;
@@ -543,10 +565,15 @@ FD3D12HeapAndOffset FD3D12PoolAllocator::GetBackingHeapAndAllocationOffsetInByte
 {
 	check(IsOwner(InResourceLocation));
 
+	return GetBackingHeapAndAllocationOffsetInBytes(InResourceLocation.GetPoolAllocatorPrivateData().PoolData);
+}
+
+
+FD3D12HeapAndOffset FD3D12PoolAllocator::GetBackingHeapAndAllocationOffsetInBytes(const FRHIPoolAllocationData& InAllocationData) const
+{
 	FD3D12HeapAndOffset HeapAndOffset;
-	FRHIPoolAllocationData& AllocationData = InResourceLocation.GetPoolAllocatorPrivateData().PoolData;
-	HeapAndOffset.Heap = ((FD3D12MemoryPool*)Pools[AllocationData.GetPoolIndex()])->GetBackingHeap();
-	HeapAndOffset.Offset = uint64(AlignDown(AllocationData.GetOffset(), PoolAlignment));
+	HeapAndOffset.Heap = ((FD3D12MemoryPool*)Pools[InAllocationData.GetPoolIndex()])->GetBackingHeap();
+	HeapAndOffset.Offset = uint64(AlignDown(InAllocationData.GetOffset(), PoolAlignment));
 	return HeapAndOffset;
 }
 
diff --git a/Engine/Source/Runtime/D3D12RHI/Private/D3D12PoolAllocator.h b/Engine/Source/Runtime/D3D12RHI/Private/D3D12PoolAllocator.h
--- a/Engine/Source/Runtime/D3D12RHI/Private/D3D12PoolAllocator.h
+++ b/Engine/Source/Runtime/D3D12RHI/Private/D3D12PoolAllocator.h
@@ -77,8 +77,8 @@ public:
 
 	// Constructor
 	FD3D12MemoryPool(FD3D12Device* ParentDevice, FRHIGPUMask VisibleNodes, const FD3D12ResourceInitConfig& InInitConfig, const FString& Name,
-		EResourceAllocationStrategy InAllocationStrategy, int16 InPoolIndex, uint64 InPoolSize, uint32 InPoolAlignment);
-	~FD3D12MemoryPool();
+		EResourceAllocationStrategy InAllocationStrategy, int16 InPoolIndex, uint64 InPoolSize, uint32 InPoolAlignment, EFreeListOrder InFreeListOrder);
+	virtual ~FD3D12MemoryPool();
 
 	// Setup/Shutdown
 	virtual void Init() override;
@@ -111,23 +111,24 @@ protected:
 /**
 @brief D3D12 specific implementation of the FRHIPoolAllocator
 **/
-class FD3D12PoolAllocator : public FRHIPoolAllocator, public FD3D12DeviceChild, public FD3D12MultiNodeGPUObject
+class FD3D12PoolAllocator : public FRHIPoolAllocator, public FD3D12DeviceChild, public FD3D12MultiNodeGPUObject, public ID3D12ResourceAllocator
 {
 public:
 
 	// Constructor
 	FD3D12PoolAllocator(FD3D12Device* ParentDevice, FRHIGPUMask VisibleNodes, const FD3D12ResourceInitConfig& InInitConfig, const FString& InName,
-		EResourceAllocationStrategy InAllocationStrategy, uint64 InPoolSize, uint32 InPoolAlignment, uint32 InMaxAllocationSize, bool bInDefragEnabled);
+		EResourceAllocationStrategy InAllocationStrategy, uint64 InPoolSize, uint32 InPoolAlignment, uint32 InMaxAllocationSize, FRHIMemoryPool::EFreeListOrder InFreeListOrder, bool bInDefragEnabled);
 	~FD3D12PoolAllocator();
 	
 	// Function names currently have to match with FD3D12DefaultBufferPool until we can make full replacement of the allocator
 	bool SupportsAllocation(D3D12_HEAP_TYPE InHeapType, D3D12_RESOURCE_FLAGS InResourceFlags, EBufferUsageFlags InBufferUsage, ED3D12ResourceStateMode InResourceStateMode) const;
 	void AllocDefaultResource(D3D12_HEAP_TYPE InHeapType, const D3D12_RESOURCE_DESC& InDesc, EBufferUsageFlags InBufferUsage, ED3D12ResourceStateMode InResourceStateMode,
 		D3D12_RESOURCE_STATES InCreateState, uint32 InAlignment, const TCHAR* InName, FD3D12ResourceLocation& ResourceLocation);
-	void Deallocate(FD3D12ResourceLocation& ResourceLocation);
+	virtual void DeallocateResource(FD3D12ResourceLocation& ResourceLocation);
 
-	void AllocResource(D3D12_HEAP_TYPE InHeapType, const D3D12_RESOURCE_DESC& InDesc, uint64 InSize, uint32 InAllocationAlignment, ED3D12ResourceStateMode InResourceStateMode,
-		D3D12_RESOURCE_STATES InCreateState, const D3D12_CLEAR_VALUE* InClearValue, const TCHAR* InName, FD3D12ResourceLocation& ResourceLocation);
+	// Implementation of ID3D12ResourceAllocator
+	virtual void AllocateResource(D3D12_HEAP_TYPE InHeapType, const D3D12_RESOURCE_DESC& InDesc, uint64 InSize, uint32 InAllocationAlignment, ED3D12ResourceStateMode InResourceStateMode,
+		D3D12_RESOURCE_STATES InCreateState, const D3D12_CLEAR_VALUE* InClearValue, const TCHAR* InName, FD3D12ResourceLocation& ResourceLocation) override;
 
 	void CleanUpAllocations(uint64 InFrameLag);
 	void FlushPendingCopyOps(FD3D12CommandContext& InCommandContext);
@@ -138,6 +139,7 @@ public:
 	EResourceAllocationStrategy GetAllocationStrategy() const { return AllocationStrategy; }
 	FD3D12Resource* GetBackingResource(FD3D12ResourceLocation& InResourceLocation) const;
 	FD3D12HeapAndOffset GetBackingHeapAndAllocationOffsetInBytes(FD3D12ResourceLocation& InResourceLocation) const;
+	FD3D12HeapAndOffset GetBackingHeapAndAllocationOffsetInBytes(const FRHIPoolAllocationData& InAllocationData) const;
 
 	static FD3D12ResourceInitConfig GetResourceAllocatorInitConfig(D3D12_HEAP_TYPE InHeapType, D3D12_RESOURCE_FLAGS InResourceFlags, EBufferUsageFlags InBufferUsage);
 	static EResourceAllocationStrategy GetResourceAllocationStrategy(D3D12_RESOURCE_FLAGS InResourceFlags, ED3D12ResourceStateMode InResourceStateMode);
@@ -148,6 +150,9 @@ protected:
 	virtual FRHIMemoryPool* CreateNewPool(int16 InPoolIndex) override;
 	virtual bool HandleDefragRequest(FRHIPoolAllocationData* InSourceBlock, FRHIPoolAllocationData& InTmpTargetBlock) override;
 	
+	// Placed resource allocation helper function which can be overriden
+	virtual FD3D12Resource* CreatePlacedResource(const FRHIPoolAllocationData& InAllocationData, const D3D12_RESOURCE_DESC& InDesc, D3D12_RESOURCE_STATES InCreateState, ED3D12ResourceStateMode InResourceStateMode, const D3D12_CLEAR_VALUE* InClearValue, const TCHAR* InName);
+	
 	// Locked allocation data which needs to do something specific at certain frame fence value	
 	struct FrameFencedAllocationData
 	{
diff --git a/Engine/Source/Runtime/D3D12RHI/Private/D3D12RHICommon.h b/Engine/Source/Runtime/D3D12RHI/Private/D3D12RHICommon.h
--- a/Engine/Source/Runtime/D3D12RHI/Private/D3D12RHICommon.h
+++ b/Engine/Source/Runtime/D3D12RHI/Private/D3D12RHICommon.h
@@ -9,6 +9,7 @@ D3D12RHICommon.h: Common D3D12 RHI definitions for Windows.
 DECLARE_STATS_GROUP(TEXT("D3D12RHI"), STATGROUP_D3D12RHI, STATCAT_Advanced);
 DECLARE_STATS_GROUP(TEXT("D3D12RHI: Memory"), STATGROUP_D3D12Memory, STATCAT_Advanced);
 DECLARE_STATS_GROUP(TEXT("D3D12RHI: Memory (Buffers)"), STATGROUP_D3D12MemoryBuffer, STATCAT_Advanced);
+DECLARE_STATS_GROUP(TEXT("D3D12RHI: Memory (Transient)"), STATGROUP_D3D12MemoryTransient, STATCAT_Advanced);
 DECLARE_STATS_GROUP(TEXT("D3D12RHI: Pipeline State (PSO)"), STATGROUP_D3D12PipelineState, STATCAT_Advanced);
 DECLARE_STATS_GROUP(TEXT("D3D12RHI: Descriptor Heap (GPU visible)"), STATGROUP_D3D12DescriptorHeap, STATCAT_Advanced);
 
diff --git a/Engine/Source/Runtime/D3D12RHI/Private/D3D12RHIPrivate.h b/Engine/Source/Runtime/D3D12RHI/Private/D3D12RHIPrivate.h
--- a/Engine/Source/Runtime/D3D12RHI/Private/D3D12RHIPrivate.h
+++ b/Engine/Source/Runtime/D3D12RHI/Private/D3D12RHIPrivate.h
@@ -83,6 +83,7 @@ extern int32 GDX12NVAfterMathMarkers;
 #include "D3D12StateCachePrivate.h"
 typedef FD3D12StateCacheBase FD3D12StateCache;
 #include "D3D12Allocation.h"
+#include "D3D12TransientResourceAllocator.h"
 #include "D3D12CommandContext.h"
 #include "D3D12Stats.h"
 #include "D3D12Device.h"
@@ -373,6 +374,8 @@ public:
 	virtual class IRHIComputeContext* RHIGetDefaultAsyncComputeContext() final override;
 	virtual class IRHICommandContextContainer* RHIGetCommandContextContainer(int32 Index, int32 Num) final override;
 
+	virtual IRHITransientResourceAllocator* RHICreateTransientResourceAllocator() final override;
+
 #if WITH_MGPU
 	virtual IRHICommandContextContainer* RHIGetCommandContextContainer(int32 Index, int32 Num, FRHIGPUMask GPUMask)final override;
 #endif
@@ -1040,11 +1043,17 @@ protected:
 	void* ZeroBuffer;
 	uint32 ZeroBufferSize;
 
+public:
+
 	template<typename BaseResourceType>
 	TD3D12Texture2D<BaseResourceType>* CreateD3D12Texture2D(class FRHICommandListImmediate* RHICmdList, uint32 SizeX, uint32 SizeY, uint32 SizeZ, bool bTextureArray, bool CubeTexture, EPixelFormat Format,
-		uint32 NumMips, uint32 NumSamples, ETextureCreateFlags InFlags, ERHIAccess InResourceState, FRHIResourceCreateInfo& CreateInfo);
+		uint32 NumMips, uint32 NumSamples, ETextureCreateFlags InFlags, ERHIAccess InResourceState, FRHIResourceCreateInfo& CreateInfo, ID3D12ResourceAllocator* ResourceAllocator);
 
-	FD3D12Texture3D* CreateD3D12Texture3D(class FRHICommandListImmediate* RHICmdList, uint32 SizeX, uint32 SizeY, uint32 SizeZ, EPixelFormat Format, uint32 NumMips, ETextureCreateFlags InFlags, ERHIAccess InResourceState, FRHIResourceCreateInfo& CreateInfo);
+	FD3D12Texture3D* CreateD3D12Texture3D(class FRHICommandListImmediate* RHICmdList, uint32 SizeX, uint32 SizeY, uint32 SizeZ, EPixelFormat Format, uint32 NumMips, ETextureCreateFlags InFlags, ERHIAccess InResourceState, FRHIResourceCreateInfo& CreateInfo, ID3D12ResourceAllocator* ResourceAllocator);
+
+	FD3D12Buffer* CreateD3D12Buffer(class FRHICommandListImmediate* RHICmdList, uint32 Size, EBufferUsageFlags Usage, uint32 Stride, ERHIAccess ResourceState, FRHIResourceCreateInfo& CreateInfo, ID3D12ResourceAllocator* ResourceAllocator);
+
+protected:
 
 	template<typename BaseResourceType>
 	TD3D12Texture2D<BaseResourceType>* CreateTextureFromResource(bool bTextureArray, bool bCubeTexture, EPixelFormat Format, ETextureCreateFlags TexCreateFlags, const FClearValueBinding& ClearValueBinding, ID3D12Resource* Resource);
diff --git a/Engine/Source/Runtime/D3D12RHI/Private/D3D12RayTracing.cpp b/Engine/Source/Runtime/D3D12RHI/Private/D3D12RayTracing.cpp
--- a/Engine/Source/Runtime/D3D12RHI/Private/D3D12RayTracing.cpp
+++ b/Engine/Source/Runtime/D3D12RHI/Private/D3D12RayTracing.cpp
@@ -427,8 +427,9 @@ FD3D12RayTracingCompactionRequestHandler::FD3D12RayTracingCompactionRequestHandl
 
 	FRHIResourceCreateInfo CreateInfo(TEXT("PostBuildInfoBuffer"));
 	CreateInfo.GPUMask = FRHIGPUMask::FromIndex(GetParentDevice()->GetGPUIndex());
+	ID3D12ResourceAllocator* ResourceAllocator = nullptr;
 	PostBuildInfoBuffer = GetParentDevice()->GetParentAdapter()->CreateRHIBuffer(nullptr, PostBuildInfoBufferDesc, 8,
-		0, PostBuildInfoBufferDesc.Width, BUF_UnorderedAccess | BUF_SourceCopy, ED3D12ResourceStateMode::MultiState, ERHIAccess::UAVMask, CreateInfo);
+		0, PostBuildInfoBufferDesc.Width, BUF_UnorderedAccess | BUF_SourceCopy, ED3D12ResourceStateMode::MultiState, ERHIAccess::UAVMask, CreateInfo, ResourceAllocator);
 	SetName(PostBuildInfoBuffer->GetResource(), TEXT("PostBuildInfoBuffer"));
 
 	PostBuildInfoStagingBuffer = RHICreateStagingBuffer();
@@ -1581,9 +1582,10 @@ public:
 		CreateInfo.ResourceArray = &Data;
 		CreateInfo.GPUMask = FRHIGPUMask::FromIndex(Device->GetGPUIndex());
 
+		ID3D12ResourceAllocator* ResourceAllocator = nullptr;
 		Buffer = Adapter->CreateRHIBuffer(
 			nullptr, BufferDesc, BufferDesc.Alignment,
-			0, BufferDesc.Width, BUF_Static, ED3D12ResourceStateMode::SingleState, ERHIAccess::SRVMask, CreateInfo);
+			0, BufferDesc.Width, BUF_Static, ED3D12ResourceStateMode::SingleState, ERHIAccess::SRVMask, CreateInfo, ResourceAllocator);
 
 		bIsDirty = false;
 	}
@@ -2656,9 +2658,10 @@ static void CreateAccelerationStructureBuffers(
 		PrebuildInfo.ResultDataMaxSizeInBytes, D3D12_RESOURCE_FLAG_ALLOW_UNORDERED_ACCESS);
 
 	CreateInfo.GPUMask = FRHIGPUMask::FromIndex(GPUIndex);
+	ID3D12ResourceAllocator* ResourceAllocator = nullptr;
 	AccelerationStructureBuffer = Adapter->CreateRHIBuffer(
 		nullptr, AccelerationStructureBufferDesc, D3D12_RAYTRACING_ACCELERATION_STRUCTURE_BYTE_ALIGNMENT,
-		0, AccelerationStructureBufferDesc.Width, BUF_AccelerationStructure, ED3D12ResourceStateMode::SingleState, ERHIAccess::RTAccelerationStructure, CreateInfo);
+		0, AccelerationStructureBufferDesc.Width, BUF_AccelerationStructure, ED3D12ResourceStateMode::SingleState, ERHIAccess::RTAccelerationStructure, CreateInfo, ResourceAllocator);
 
 	SetName(AccelerationStructureBuffer->GetResource(), DebugName);
 
@@ -2670,7 +2673,7 @@ static void CreateAccelerationStructureBuffers(
 	CreateInfo.DebugName = TEXT("ScratchBuffer");
 	ScratchBuffer = Adapter->CreateRHIBuffer(
 		nullptr, ScratchBufferDesc, D3D12_RAYTRACING_ACCELERATION_STRUCTURE_BYTE_ALIGNMENT,
-		0, ScratchBufferDesc.Width, BUF_UnorderedAccess, ED3D12ResourceStateMode::SingleState, ERHIAccess::UAVMask, CreateInfo);
+		0, ScratchBufferDesc.Width, BUF_UnorderedAccess, ED3D12ResourceStateMode::SingleState, ERHIAccess::UAVMask, CreateInfo, ResourceAllocator);
 
 	SetName(ScratchBuffer->GetResource(), TEXT("Acceleration structure scratch"));
 
@@ -2951,9 +2954,10 @@ void FD3D12RayTracingGeometry::CompactAccelerationStructure(FD3D12CommandContext
 		D3D12_RESOURCE_DESC AccelerationStructureBufferDesc = CD3DX12_RESOURCE_DESC::Buffer(InSizeAfterCompaction, D3D12_RESOURCE_FLAG_ALLOW_UNORDERED_ACCESS);
 
 		CreateInfo.GPUMask = FRHIGPUMask::FromIndex(InGPUIndex);
+		ID3D12ResourceAllocator* ResourceAllocator = nullptr;
 		AccelerationStructureBuffers[InGPUIndex] = CommandContext.GetParentAdapter()->CreateRHIBuffer(
 			nullptr, AccelerationStructureBufferDesc, D3D12_RAYTRACING_ACCELERATION_STRUCTURE_BYTE_ALIGNMENT,
-			0, AccelerationStructureBufferDesc.Width, BUF_AccelerationStructure, ED3D12ResourceStateMode::SingleState, ERHIAccess::RTAccelerationStructure, CreateInfo);
+			0, AccelerationStructureBufferDesc.Width, BUF_AccelerationStructure, ED3D12ResourceStateMode::SingleState, ERHIAccess::RTAccelerationStructure, CreateInfo, ResourceAllocator);
 
 		INC_MEMORY_STAT_BY(STAT_D3D12RayTracingUsedVideoMemory, AccelerationStructureBuffers[InGPUIndex]->GetSize());
 		INC_MEMORY_STAT_BY(STAT_D3D12RayTracingBLASMemory, AccelerationStructureBuffers[InGPUIndex]->GetSize());
@@ -3236,9 +3240,10 @@ void FD3D12RayTracingScene::BuildAccelerationStructure(FD3D12CommandContext& Com
 			D3D12_RESOURCE_FLAG_ALLOW_UNORDERED_ACCESS);
 
 		// #dxr_todo multi state only when bShouldCopyIndirectInstances is set - will still need transition to copy_src for lock behind but this can be done on the complete pool in theory (have to check cost)
+		ID3D12ResourceAllocator* ResourceAllocator = nullptr;
 		InstanceBuffer = Adapter->CreateRHIBuffer(
 			nullptr, InstanceBufferDesc, D3D12_RAYTRACING_INSTANCE_DESCS_BYTE_ALIGNMENT,
-			sizeof(D3D12_RAYTRACING_INSTANCE_DESC), InstanceBufferDesc.Width, BUF_UnorderedAccess, ED3D12ResourceStateMode::MultiState, ERHIAccess::UAVMask, CreateInfo);
+			sizeof(D3D12_RAYTRACING_INSTANCE_DESC), InstanceBufferDesc.Width, BUF_UnorderedAccess, ED3D12ResourceStateMode::MultiState, ERHIAccess::UAVMask, CreateInfo, ResourceAllocator);
 
 		{
 			TArray<FD3D12ResidencyHandle*>& GeometryResidencyHandlesForThisGPU = GeometryResidencyHandles[GPUIndex];
diff --git a/Engine/Source/Runtime/D3D12RHI/Private/D3D12Resources.cpp b/Engine/Source/Runtime/D3D12RHI/Private/D3D12Resources.cpp
--- a/Engine/Source/Runtime/D3D12RHI/Private/D3D12Resources.cpp
+++ b/Engine/Source/Runtime/D3D12RHI/Private/D3D12Resources.cpp
@@ -213,6 +213,26 @@ void FD3D12DeferredDeletionQueue::FD3D12AsyncDeletionWorker::DoWork()
 	}
 }
 
+
+/////////////////////////////////////////////////////////////////////
+//	ID3D12ResourceAllocator
+/////////////////////////////////////////////////////////////////////
+
+
+void ID3D12ResourceAllocator::AllocateTexture(D3D12_HEAP_TYPE InHeapType, const D3D12_RESOURCE_DESC& InDesc, EPixelFormat InUEFormat, ED3D12ResourceStateMode InResourceStateMode,
+	D3D12_RESOURCE_STATES InCreateState, const D3D12_CLEAR_VALUE* InClearValue, const TCHAR* InName, FD3D12ResourceLocation& ResourceLocation)
+{
+	// Check if texture can be 4K aligned
+	D3D12_RESOURCE_DESC Desc = InDesc;
+	bool b4KAligment = TextureCanBe4KAligned(Desc, InUEFormat);
+	Desc.Alignment = b4KAligment ? D3D12_SMALL_RESOURCE_PLACEMENT_ALIGNMENT : D3D12_DEFAULT_RESOURCE_PLACEMENT_ALIGNMENT;
+
+	// Get the size and alignment for the allocation
+	D3D12_RESOURCE_ALLOCATION_INFO Info = FD3D12DynamicRHI::GetD3DRHI()->GetAdapter().GetDevice(0)->GetResourceAllocationInfo(Desc);
+	AllocateResource(InHeapType, Desc, Info.SizeInBytes, Info.Alignment, InResourceStateMode, InCreateState, InClearValue, InName, ResourceLocation);
+}
+
+
 /////////////////////////////////////////////////////////////////////
 //	FD3D12 Resource
 /////////////////////////////////////////////////////////////////////
@@ -227,9 +247,10 @@ FD3D12Resource::FD3D12Resource(FD3D12Device* ParentDevice,
 	ID3D12Resource* InResource,
 	D3D12_RESOURCE_STATES InInitialState,
 	D3D12_RESOURCE_DESC const& InDesc,
+	const D3D12_CLEAR_VALUE* InClearValue,
 	FD3D12Heap* InHeap,
 	D3D12_HEAP_TYPE InHeapType) : 
-	FD3D12Resource(ParentDevice, VisibleNodes, InResource, InInitialState, ED3D12ResourceStateMode::Default, D3D12_RESOURCE_STATE_TBD, InDesc, InHeap, InHeapType)
+	FD3D12Resource(ParentDevice, VisibleNodes, InResource, InInitialState, ED3D12ResourceStateMode::Default, D3D12_RESOURCE_STATE_TBD, InDesc, InClearValue, InHeap, InHeapType)
 {
 }
 
@@ -240,6 +261,7 @@ FD3D12Resource::FD3D12Resource(FD3D12Device* ParentDevice,
 	ED3D12ResourceStateMode InResourceStateMode,
 	D3D12_RESOURCE_STATES InDefaultResourceState,
 	D3D12_RESOURCE_DESC const& InDesc,
+	const D3D12_CLEAR_VALUE* InClearValue,
 	FD3D12Heap* InHeap,
 	D3D12_HEAP_TYPE InHeapType)
 	: FD3D12DeviceChild(ParentDevice)
@@ -263,6 +285,16 @@ FD3D12Resource::FD3D12Resource(FD3D12Device* ParentDevice,
 	FPlatformAtomics::InterlockedIncrement(&TotalResourceCount);
 #endif
 
+	// Store the optional clear value - need for defrag and transient resource recreation
+	if (InClearValue)
+	{
+		ClearValue = *InClearValue;
+	}
+	else
+	{
+		FPlatformMemory::Memzero(&ClearValue, sizeof(D3D12_CLEAR_VALUE));
+	}
+
 	if (Resource
 #if PLATFORM_WINDOWS
 		&& Desc.Dimension == D3D12_RESOURCE_DIMENSION_BUFFER
@@ -435,7 +467,7 @@ HRESULT FD3D12Adapter::CreateCommittedResource(const D3D12_RESOURCE_DESC& InDesc
 	if (SUCCEEDED(hr))
 	{
 		// Set the output pointer
-		*ppOutResource = new FD3D12Resource(GetDevice(CreationNode.ToIndex()), CreationNode, pResource, InInitialState, InResourceStateMode, InDefaultState, InDesc, nullptr, HeapProps.Type);
+		*ppOutResource = new FD3D12Resource(GetDevice(CreationNode.ToIndex()), CreationNode, pResource, InInitialState, InResourceStateMode, InDefaultState, InDesc, ClearValue, nullptr, HeapProps.Type);
 		(*ppOutResource)->AddRef();
 
 		// Set a default name (can override later).
@@ -462,7 +494,7 @@ HRESULT FD3D12Adapter::CreatePlacedResource(const D3D12_RESOURCE_DESC& InDesc, F
 	ID3D12Heap* Heap = BackingHeap->GetHeap();
 
 	TRefCountPtr<ID3D12Resource> pResource;
-	const HRESULT hr = RootDevice->CreatePlacedResource(Heap, HeapOffset, &InDesc, InInitialState, nullptr, IID_PPV_ARGS(pResource.GetInitReference()));
+	const HRESULT hr = RootDevice->CreatePlacedResource(Heap, HeapOffset, &InDesc, InInitialState, ClearValue, IID_PPV_ARGS(pResource.GetInitReference()));
 
 	if (bVerifyHResult)
 	{
@@ -482,6 +514,7 @@ HRESULT FD3D12Adapter::CreatePlacedResource(const D3D12_RESOURCE_DESC& InDesc, F
 			InResourceStateMode,
 			InDefaultState,
 			InDesc,
+			ClearValue,
 			BackingHeap,
 			HeapDesc.Properties.Type);
 
@@ -707,7 +740,7 @@ void FD3D12ResourceLocation::ReleaseResource()
 		}
 		else if (AllocatorType == AT_Pool)
 		{
-			PoolAllocator->Deallocate(*this);
+			PoolAllocator->DeallocateResource(*this);
 		}
 		else
 		{
diff --git a/Engine/Source/Runtime/D3D12RHI/Private/D3D12Stats.h b/Engine/Source/Runtime/D3D12RHI/Private/D3D12Stats.h
--- a/Engine/Source/Runtime/D3D12RHI/Private/D3D12Stats.h
+++ b/Engine/Source/Runtime/D3D12RHI/Private/D3D12Stats.h
@@ -89,6 +89,24 @@ DECLARE_MEMORY_STAT_EXTERN(TEXT("BufferPool Memory Alignment Waste"), STAT_D3D12
 DECLARE_DWORD_ACCUMULATOR_STAT_EXTERN(TEXT("BufferPool Page Count"), STAT_D3D12BufferPoolPageCount, STATGROUP_D3D12MemoryBuffer, );
 DECLARE_DWORD_ACCUMULATOR_STAT_EXTERN(TEXT("BufferPool Full Pages"), STAT_D3D12BufferPoolFullPages, STATGROUP_D3D12MemoryBuffer, );
 
+DECLARE_MEMORY_STAT_EXTERN(TEXT("Transient Memory Pool Allocated"), STAT_D3D12TransientMemoryPoolAllocated, STATGROUP_D3D12Memory, );
+DECLARE_MEMORY_STAT_EXTERN(TEXT("Transient Memory Frame Pool Used"), STAT_D3D12TransientMemoryFramePoolUsed, STATGROUP_D3D12Memory, );
+DECLARE_MEMORY_STAT_EXTERN(TEXT("Transient Memory Frame Committed Used"), STAT_D3D12TransientMemoryFrameCommittedUsed, STATGROUP_D3D12Memory, );
+DECLARE_DWORD_ACCUMULATOR_STAT_EXTERN(TEXT("Transient Memory Pool Allocations"), STAT_D3D12TransientMemoryPoolAllocations, STATGROUP_D3D12Memory, );
+DECLARE_DWORD_ACCUMULATOR_STAT_EXTERN(TEXT("Transient Memory Committed Allocations"), STAT_D3D12TransientMemoryCommittedAllocations, STATGROUP_D3D12Memory, );
+
+DECLARE_MEMORY_STAT_EXTERN(TEXT("Buffer Memory Frame Used"), STAT_D3D12TransientMemoryBufferFrameUsed, STATGROUP_D3D12MemoryTransient, );
+DECLARE_MEMORY_STAT_EXTERN(TEXT("Buffer Memory Frame Total Requested"), STAT_D3D12TransientMemoryBufferFrameTotalRequested, STATGROUP_D3D12MemoryTransient, );
+DECLARE_MEMORY_STAT_EXTERN(TEXT("Buffer Memory Frame Committed Allocated"), STAT_D3D12TransientMemoryBufferFrameCommittedAllocated, STATGROUP_D3D12MemoryTransient, );
+DECLARE_DWORD_ACCUMULATOR_STAT_EXTERN(TEXT("Buffer Memory Pool Allocations"), STAT_D3D12TransientMemoryBufferPoolAllocations, STATGROUP_D3D12MemoryTransient, );
+DECLARE_DWORD_ACCUMULATOR_STAT_EXTERN(TEXT("Buffer Memory Allocations"), STAT_D3D12TransientMemoryBufferCommittedAllocations, STATGROUP_D3D12MemoryTransient, );
+
+DECLARE_MEMORY_STAT_EXTERN(TEXT("Texture Memory Frame Used"), STAT_D3D12TransientMemoryTextureFrameUsed, STATGROUP_D3D12MemoryTransient, );
+DECLARE_MEMORY_STAT_EXTERN(TEXT("Texture Memory Frame Total Requested"), STAT_D3D12TransientMemoryTextureFrameTotalRequested, STATGROUP_D3D12MemoryTransient, );
+DECLARE_MEMORY_STAT_EXTERN(TEXT("Texture Memory Frame Committed Allocated"), STAT_D3D12TransientMemoryTextureFrameCommittedAllocated, STATGROUP_D3D12MemoryTransient, );
+DECLARE_DWORD_ACCUMULATOR_STAT_EXTERN(TEXT("Texture Memory Pool Allocations"), STAT_D3D12TransientMemoryTexturePoolAllocations, STATGROUP_D3D12MemoryTransient, );
+DECLARE_DWORD_ACCUMULATOR_STAT_EXTERN(TEXT("Texture Memory Allocations"), STAT_D3D12TransientMemoryTextureCommittedAllocations, STATGROUP_D3D12MemoryTransient, );
+
 
 /**
 * Detailed Descriptor heap stats
diff --git a/Engine/Source/Runtime/D3D12RHI/Private/D3D12Texture.cpp b/Engine/Source/Runtime/D3D12RHI/Private/D3D12Texture.cpp
--- a/Engine/Source/Runtime/D3D12RHI/Private/D3D12Texture.cpp
+++ b/Engine/Source/Runtime/D3D12RHI/Private/D3D12Texture.cpp
@@ -355,8 +355,11 @@ void FD3D12TextureStats::D3D12TextureAllocated(TD3D12Texture2D<BaseResourceType>
 
 	if (D3D12Texture2D)
 	{
-		if ((Texture.Flags & TexCreate_Virtual) != TexCreate_Virtual)
+		// Don't update state for virtual or transient textures	
+		if (!EnumHasAnyFlags(Texture.Flags, TexCreate_Virtual | TexCreate_Transient))
 		{
+			TRACE_CPUPROFILER_EVENT_SCOPE(D3D12RHI::UpdateTextureStats);
+
 			if (!Desc)
 			{
 				Desc = &D3D12Texture2D->GetDesc();
@@ -386,18 +389,22 @@ void FD3D12TextureStats::D3D12TextureDeleted(TD3D12Texture2D<BaseResourceType>&
 
 	if (D3D12Texture2D)
 	{
-		const D3D12_RESOURCE_DESC& Desc = D3D12Texture2D->GetDesc();
-		const int64 TextureSize = Texture.GetMemorySize();
-		ensure(TextureSize > 0 || (Texture.Flags & TexCreate_Virtual) || Texture.GetAliasingSourceTexture() != nullptr);
+		// Don't update state for virtual or transient textures	
+		if(!EnumHasAnyFlags(Texture.Flags, TexCreate_Virtual | TexCreate_Transient))
+		{
+			const D3D12_RESOURCE_DESC& Desc = D3D12Texture2D->GetDesc();
+			const int64 TextureSize = Texture.GetMemorySize();
+			ensure(TextureSize > 0 || (Texture.Flags & TexCreate_Virtual) || Texture.GetAliasingSourceTexture() != nullptr);
 
-		UpdateD3D12TextureStats(Desc, -TextureSize, false, Texture.IsCubemap(), Texture.IsStreamable());
+			UpdateD3D12TextureStats(Desc, -TextureSize, false, Texture.IsCubemap(), Texture.IsStreamable());
 
 #if PLATFORM_WINDOWS
-		// On Windows there is no way to hook into the low level d3d allocations and frees.
-		// This means that we must manually add the tracking here.
-		LLM(FLowLevelMemTracker::Get().OnLowLevelFree(ELLMTracker::Platform, Texture.GetResource()->GetResource()));
-		LLM(FLowLevelMemTracker::Get().OnLowLevelFree(ELLMTracker::Default, Texture.GetResource()->GetResource()));
+			// On Windows there is no way to hook into the low level d3d allocations and frees.
+			// This means that we must manually add the tracking here.
+			LLM(FLowLevelMemTracker::Get().OnLowLevelFree(ELLMTracker::Platform, Texture.GetResource()->GetResource()));
+			LLM(FLowLevelMemTracker::Get().OnLowLevelFree(ELLMTracker::Default, Texture.GetResource()->GetResource()));
 #endif
+		}
 	}
 }
 
@@ -488,7 +495,7 @@ uint64 FD3D12DynamicRHI::RHICalcTexture2DPlatformSize(uint32 SizeX, uint32 SizeY
 	Desc.Width = SizeX;
 
 	// Check if the 4K aligment is possible
-	Desc.Alignment = TextureCanBe4KAligned(Desc, Format) ? D3D12_SMALL_RESOURCE_PLACEMENT_ALIGNMENT : 0;
+	Desc.Alignment = TextureCanBe4KAligned(Desc, (EPixelFormat)Format) ? D3D12_SMALL_RESOURCE_PLACEMENT_ALIGNMENT : 0;
 
 	const D3D12_RESOURCE_ALLOCATION_INFO AllocationInfo = GetAdapter().GetD3DDevice()->GetResourceAllocationInfo(0, 1, &Desc);
 	OutAlign = static_cast<uint32>(AllocationInfo.Alignment);
@@ -509,7 +516,7 @@ uint64 FD3D12DynamicRHI::RHICalcTexture2DArrayPlatformSize(uint32 SizeX, uint32
 	Desc.Width = SizeX;
 
 	// Check if the 4K aligment is possible
-	Desc.Alignment = TextureCanBe4KAligned(Desc, Format) ? D3D12_SMALL_RESOURCE_PLACEMENT_ALIGNMENT : 0;
+	Desc.Alignment = TextureCanBe4KAligned(Desc, (EPixelFormat)Format) ? D3D12_SMALL_RESOURCE_PLACEMENT_ALIGNMENT : 0;
 
 	const D3D12_RESOURCE_ALLOCATION_INFO AllocationInfo = GetAdapter().GetD3DDevice()->GetResourceAllocationInfo(0, 1, &Desc);
 	OutAlign = static_cast<uint32>(AllocationInfo.Alignment);
@@ -530,7 +537,7 @@ uint64 FD3D12DynamicRHI::RHICalcTexture3DPlatformSize(uint32 SizeX, uint32 SizeY
 	Desc.Width = SizeX;
 
 	// Check if the 4K aligment is possible
-	Desc.Alignment = TextureCanBe4KAligned(Desc, Format) ? D3D12_SMALL_RESOURCE_PLACEMENT_ALIGNMENT : 0;
+	Desc.Alignment = TextureCanBe4KAligned(Desc, (EPixelFormat)Format) ? D3D12_SMALL_RESOURCE_PLACEMENT_ALIGNMENT : 0;
 
 	const D3D12_RESOURCE_ALLOCATION_INFO AllocationInfo = GetAdapter().GetD3DDevice()->GetResourceAllocationInfo(0, 1, &Desc);
 	OutAlign = static_cast<uint32>(AllocationInfo.Alignment);
@@ -551,7 +558,7 @@ uint64 FD3D12DynamicRHI::RHICalcTextureCubePlatformSize(uint32 Size, uint8 Forma
 	Desc.Width = Size;
 
 	// Check if the 4K aligment is possible
-	Desc.Alignment = TextureCanBe4KAligned(Desc, Format) ? D3D12_SMALL_RESOURCE_PLACEMENT_ALIGNMENT : 0;
+	Desc.Alignment = TextureCanBe4KAligned(Desc, (EPixelFormat)Format) ? D3D12_SMALL_RESOURCE_PLACEMENT_ALIGNMENT : 0;
 
 	const D3D12_RESOURCE_ALLOCATION_INFO AllocationInfo = GetAdapter().GetD3DDevice()->GetResourceAllocationInfo(0, 1, &Desc);
 	OutAlign = static_cast<uint32>(AllocationInfo.Alignment);
@@ -731,11 +738,86 @@ void SafeCreateTexture2D(FD3D12Device* pDevice,
 #endif // #if GUARDED_TEXTURE_CREATES
 }
 
+
+static void DetermineTexture2DResourceFlagsAndLayout(uint32 SizeX, uint32 SizeY, uint32 SizeZ, uint32 NumMips, uint32 NumSamples, uint32 Flags, EPixelFormat Format, D3D12_RESOURCE_FLAGS& OutResourceFlags, D3D12_TEXTURE_LAYOUT& OutLayout, bool& bOutCreateRTV, bool& bOutCreateDSV, bool& bOutCreateSRV)
+{
+	TRACE_CPUPROFILER_EVENT_SCOPE(DetermineTexture2DResourceFlagsAndLayout);
+
+	OutResourceFlags = D3D12_RESOURCE_FLAG_NONE;
+	OutLayout = D3D12_TEXTURE_LAYOUT_UNKNOWN;
+	bOutCreateRTV = false;
+	bOutCreateDSV = false;
+	bOutCreateSRV = true;
+
+	if (Flags & TexCreate_CPUReadback)
+	{
+		check(!(Flags & TexCreate_RenderTargetable));
+		check(!(Flags & TexCreate_DepthStencilTargetable));
+		check(!(Flags & TexCreate_ShaderResource));
+		bOutCreateSRV = false;
+	}
+
+	if (Flags & TexCreate_DisableSRVCreation)
+	{
+		bOutCreateSRV = false;
+	}
+
+	if (Flags & TexCreate_Shared)
+	{
+		OutResourceFlags |= D3D12_RESOURCE_FLAG_ALLOW_SIMULTANEOUS_ACCESS;
+	}
+
+	if (Flags & TexCreate_RenderTargetable)
+	{
+		check(!(Flags & TexCreate_DepthStencilTargetable));
+		check(!(Flags & TexCreate_ResolveTargetable));
+		OutResourceFlags |= D3D12_RESOURCE_FLAG_ALLOW_RENDER_TARGET;
+		bOutCreateRTV = true;
+	}
+	else if (Flags & TexCreate_DepthStencilTargetable)
+	{
+		check(!(Flags & TexCreate_RenderTargetable));
+		check(!(Flags & TexCreate_ResolveTargetable));
+		OutResourceFlags |= D3D12_RESOURCE_FLAG_ALLOW_DEPTH_STENCIL;
+		bOutCreateDSV = true;
+	}
+	else if (Flags & TexCreate_ResolveTargetable)
+	{
+		check(!(Flags & TexCreate_RenderTargetable));
+		check(!(Flags & TexCreate_DepthStencilTargetable));
+		if (Format == PF_DepthStencil || Format == PF_ShadowDepth || Format == PF_D24)
+		{
+			OutResourceFlags |= D3D12_RESOURCE_FLAG_ALLOW_DEPTH_STENCIL;
+			bOutCreateDSV = true;
+		}
+		else
+		{
+			OutResourceFlags |= D3D12_RESOURCE_FLAG_ALLOW_RENDER_TARGET;
+			bOutCreateRTV = true;
+		}
+	}
+
+	if (Flags & TexCreate_UAV)
+	{
+		OutResourceFlags |= D3D12_RESOURCE_FLAG_ALLOW_UNORDERED_ACCESS;
+	}
+
+	if (bOutCreateDSV && !(Flags & TexCreate_ShaderResource))
+	{
+		// Only deny shader resources if it's a depth resource that will never be used as SRV
+		OutResourceFlags |= D3D12_RESOURCE_FLAG_DENY_SHADER_RESOURCE;
+		bOutCreateSRV = false;
+	}
+}
+
 template<typename BaseResourceType>
 TD3D12Texture2D<BaseResourceType>* FD3D12DynamicRHI::CreateD3D12Texture2D(FRHICommandListImmediate* RHICmdList, uint32 SizeX, uint32 SizeY, uint32 SizeZ, bool bTextureArray, bool bCubeTexture, EPixelFormat Format,
-	uint32 NumMips, uint32 NumSamples, ETextureCreateFlags Flags, ERHIAccess InResourceState, FRHIResourceCreateInfo& CreateInfo)
+	uint32 NumMips, uint32 NumSamples, ETextureCreateFlags Flags, ERHIAccess InResourceState, FRHIResourceCreateInfo& CreateInfo, ID3D12ResourceAllocator* ResourceAllocator)
 {
 #if PLATFORM_WINDOWS || PLATFORM_HOLOLENS
+	
+	TRACE_CPUPROFILER_EVENT_SCOPE(D3D12RHI::CreateD3D12Texture2D);
+
 	check(SizeX > 0 && SizeY > 0 && NumMips > 0);
 
 	if (bCubeTexture)
@@ -763,10 +845,7 @@ TD3D12Texture2D<BaseResourceType>* FD3D12DynamicRHI::CreateD3D12Texture2D(FRHICo
 	const DXGI_FORMAT PlatformRenderTargetFormat = FindShaderResourceDXGIFormat(PlatformResourceFormat, bSRGB);
 	const DXGI_FORMAT PlatformDepthStencilFormat = FindDepthStencilDXGIFormat(PlatformResourceFormat);
 
-	bool bCreateShaderResource = true;
-
 	uint32 ActualMSAACount = NumSamples;
-
 	uint32 ActualMSAAQuality = GetMaxMSAAQuality(ActualMSAACount);
 
 	// 0xffffffff means not supported
@@ -776,22 +855,8 @@ TD3D12Texture2D<BaseResourceType>* FD3D12DynamicRHI::CreateD3D12Texture2D(FRHICo
 		ActualMSAACount = 1;
 		ActualMSAAQuality = 0;
 	}
-
 	const bool bIsMultisampled = ActualMSAACount > 1;
 
-	if (Flags & TexCreate_CPUReadback)
-	{
-		check(!(Flags & TexCreate_RenderTargetable));
-		check(!(Flags & TexCreate_DepthStencilTargetable));
-		check(!(Flags & TexCreate_ShaderResource));
-		bCreateShaderResource = false;
-	}
-
-	if (Flags & TexCreate_DisableSRVCreation)
-	{
-		bCreateShaderResource = false;
-	}
-
 	// Describe the texture.
 	D3D12_RESOURCE_DESC TextureDesc = CD3DX12_RESOURCE_DESC::Tex2D(
 		PlatformResourceFormat,
@@ -804,56 +869,12 @@ TD3D12Texture2D<BaseResourceType>* FD3D12DynamicRHI::CreateD3D12Texture2D(FRHICo
 		D3D12_RESOURCE_FLAG_NONE);  // Add misc flags later
 
 	// Set up the texture bind flags.
-	bool bCreateRTV = false;
-	bool bCreateDSV = false;
-
-	if (Flags & TexCreate_Shared)
-	{
-		TextureDesc.Flags |= D3D12_RESOURCE_FLAG_ALLOW_SIMULTANEOUS_ACCESS;
-	}
-
-	if (Flags & TexCreate_RenderTargetable)
-	{
-		check(!(Flags & TexCreate_DepthStencilTargetable));
-		check(!(Flags & TexCreate_ResolveTargetable));
-		TextureDesc.Flags |= D3D12_RESOURCE_FLAG_ALLOW_RENDER_TARGET;
-		bCreateRTV = true;
-	}
-	else if (Flags & TexCreate_DepthStencilTargetable)
-	{
-		check(!(Flags & TexCreate_RenderTargetable));
-		check(!(Flags & TexCreate_ResolveTargetable));
-		TextureDesc.Flags |= D3D12_RESOURCE_FLAG_ALLOW_DEPTH_STENCIL;
-		bCreateDSV = true;
-	}
-	else if (Flags & TexCreate_ResolveTargetable)
-	{
-		check(!(Flags & TexCreate_RenderTargetable));
-		check(!(Flags & TexCreate_DepthStencilTargetable));
-		if (Format == PF_DepthStencil || Format == PF_ShadowDepth || Format == PF_D24)
-		{
-			TextureDesc.Flags |= D3D12_RESOURCE_FLAG_ALLOW_DEPTH_STENCIL;
-			bCreateDSV = true;
-		}
-		else
-		{
-			TextureDesc.Flags |= D3D12_RESOURCE_FLAG_ALLOW_RENDER_TARGET;
-			bCreateRTV = true;
-		}
-	}
-
-	if (Flags & TexCreate_UAV)
-	{
-		TextureDesc.Flags |= D3D12_RESOURCE_FLAG_ALLOW_UNORDERED_ACCESS;
-	}
-
-	if (bCreateDSV && !(Flags & TexCreate_ShaderResource))
-	{
-		// Only deny shader resources if it's a depth resource that will never be used as SRV
-		TextureDesc.Flags |= D3D12_RESOURCE_FLAG_DENY_SHADER_RESOURCE;
-		bCreateShaderResource = false;
-	}
+	bool bCreateRTV;
+	bool bCreateDSV;
+	bool bCreateShaderResource;
+	DetermineTexture2DResourceFlagsAndLayout(SizeX, SizeY, SizeZ, NumMips, ActualMSAACount, Flags, Format, TextureDesc.Flags, TextureDesc.Layout, bCreateRTV, bCreateDSV, bCreateShaderResource);
 
+	// Virtual textures currently not supported in default D3D12
 	Flags &= ~TexCreate_Virtual;
 
 	FD3D12Adapter* Adapter = &GetAdapter();
@@ -871,9 +892,13 @@ TD3D12Texture2D<BaseResourceType>* FD3D12DynamicRHI::CreateD3D12Texture2D(FRHICo
 		ClearValuePtr = &ClearValue;
 	}
 
+	// Transient resources don't get any actual D3D12 resource asigned at this point yet
+	bool bIsTransient = GSupportsTransientResourceAliasing && (Flags & TexCreate_Transient);
+
 	// The state this resource will be in when it leaves this function
 	const FD3D12Resource::FD3D12ResourceTypeHelper Type(TextureDesc, D3D12_HEAP_TYPE_DEFAULT);
 	const D3D12_RESOURCE_STATES InitialState = Type.GetOptimalInitialState(InResourceState, false);
+	const D3D12_RESOURCE_STATES CreateState = (CreateInfo.BulkData != nullptr) ? D3D12_RESOURCE_STATE_COPY_DEST : InitialState;
 
 	TD3D12Texture2D<BaseResourceType>* D3D12TextureOut = Adapter->CreateLinkedObject<TD3D12Texture2D<BaseResourceType>>(CreateInfo.GPUMask, [&](FD3D12Device* Device)
 	{
@@ -886,206 +911,227 @@ TD3D12Texture2D<BaseResourceType>* FD3D12DynamicRHI::CreateD3D12Texture2D(FRHICo
 			(EPixelFormat)Format,
 			bCubeTexture,
 			Flags,
-			CreateInfo.ClearValueBinding);
+			CreateInfo.ClearValueBinding);		
+
+		if (CreateInfo.DebugName)
+		{
+			TRACE_CPUPROFILER_EVENT_SCOPE(D3D12RHI::SetDebugName);
+			NewTexture->SetName(CreateInfo.DebugName);
+		}
 
 		FD3D12ResourceLocation& Location = NewTexture->ResourceLocation;
 
-		SafeCreateTexture2D(Device,
-			Adapter,
-			TextureDesc,
-			ClearValuePtr,
-			&Location,
-			NewTexture,
-			Format,
-			Flags,
-			(CreateInfo.BulkData != nullptr) ? D3D12_RESOURCE_STATE_COPY_DEST : InitialState,
-			CreateInfo.DebugName);
-		
+		if (ResourceAllocator)
+		{
+			const D3D12_HEAP_TYPE HeapType = (Flags & TexCreate_CPUReadback) ? D3D12_HEAP_TYPE_READBACK : D3D12_HEAP_TYPE_DEFAULT;
+			ResourceAllocator->AllocateTexture(HeapType, TextureDesc, (EPixelFormat)Format, ED3D12ResourceStateMode::Default, CreateState, ClearValuePtr, CreateInfo.DebugName, Location);
+			Location.SetOwner(NewTexture);
+		}
+		else
+		{
+			SafeCreateTexture2D(Device,
+				Adapter,
+				TextureDesc,
+				ClearValuePtr,
+				&Location,
+				NewTexture,
+				Format,
+				Flags,
+				CreateState,
+				CreateInfo.DebugName);
+		}
+
 		// Unlock immediately if no initial data
 		if (CreateInfo.BulkData == nullptr)
 		{
 			Location.UnlockPoolData();
 		}
 
+		// Mark transient
+		Location.SetTransient(bIsTransient);
+
 		uint32 RTVIndex = 0;
 
-		if (bCreateRTV)
 		{
-			const bool bCreateRTVsPerSlice = (Flags & TexCreate_TargetArraySlicesIndependently) && (bTextureArray || bCubeTexture);
-			NewTexture->SetNumRenderTargetViews(bCreateRTVsPerSlice ? NumMips * TextureDesc.DepthOrArraySize : NumMips);
-
-			// Create a render target view for each mip
-			for (uint32 MipIndex = 0; MipIndex < NumMips; MipIndex++)
+			TRACE_CPUPROFILER_EVENT_SCOPE(D3D12RHI::CreateViews);
+			if (bCreateRTV)
 			{
-				if (bCreateRTVsPerSlice)
-				{
-					NewTexture->SetCreatedRTVsPerSlice(true, TextureDesc.DepthOrArraySize);
-
-					for (uint32 SliceIndex = 0; SliceIndex < TextureDesc.DepthOrArraySize; SliceIndex++)
-					{
-						D3D12_RENDER_TARGET_VIEW_DESC RTVDesc;
-						FMemory::Memzero(RTVDesc);
-
-						RTVDesc.Format = PlatformRenderTargetFormat;
-						RTVDesc.ViewDimension = D3D12_RTV_DIMENSION_TEXTURE2DARRAY;
-						RTVDesc.Texture2DArray.FirstArraySlice = SliceIndex;
-						RTVDesc.Texture2DArray.ArraySize = 1;
-						RTVDesc.Texture2DArray.MipSlice = MipIndex;
-						RTVDesc.Texture2DArray.PlaneSlice = GetPlaneSliceFromViewFormat(PlatformResourceFormat, RTVDesc.Format);
+				const bool bCreateRTVsPerSlice = (Flags & TexCreate_TargetArraySlicesIndependently) && (bTextureArray || bCubeTexture);
+				NewTexture->SetNumRenderTargetViews(bCreateRTVsPerSlice ? NumMips * TextureDesc.DepthOrArraySize : NumMips);
 
-						NewTexture->SetRenderTargetViewIndex(new FD3D12RenderTargetView(Device, RTVDesc, NewTexture), RTVIndex++);
-					}
-				}
-				else
+				// Create a render target view for each mip
+				for (uint32 MipIndex = 0; MipIndex < NumMips; MipIndex++)
 				{
-					D3D12_RENDER_TARGET_VIEW_DESC RTVDesc;
-					FMemory::Memzero(RTVDesc);
-
-					RTVDesc.Format = PlatformRenderTargetFormat;
-
-					if (bTextureArray || bCubeTexture)
+					if (bCreateRTVsPerSlice)
 					{
-						if (bIsMultisampled)
-						{
-							RTVDesc.ViewDimension = D3D12_RTV_DIMENSION_TEXTURE2DMSARRAY;
-							RTVDesc.Texture2DMSArray.FirstArraySlice = 0;
-							RTVDesc.Texture2DMSArray.ArraySize = TextureDesc.DepthOrArraySize;
-						}
-						else
+						NewTexture->SetCreatedRTVsPerSlice(true, TextureDesc.DepthOrArraySize);
+
+						for (uint32 SliceIndex = 0; SliceIndex < TextureDesc.DepthOrArraySize; SliceIndex++)
 						{
+							D3D12_RENDER_TARGET_VIEW_DESC RTVDesc;
+							FMemory::Memzero(RTVDesc);
+
+							RTVDesc.Format = PlatformRenderTargetFormat;
 							RTVDesc.ViewDimension = D3D12_RTV_DIMENSION_TEXTURE2DARRAY;
-							RTVDesc.Texture2DArray.FirstArraySlice = 0;
-							RTVDesc.Texture2DArray.ArraySize = TextureDesc.DepthOrArraySize;
+							RTVDesc.Texture2DArray.FirstArraySlice = SliceIndex;
+							RTVDesc.Texture2DArray.ArraySize = 1;
 							RTVDesc.Texture2DArray.MipSlice = MipIndex;
 							RTVDesc.Texture2DArray.PlaneSlice = GetPlaneSliceFromViewFormat(PlatformResourceFormat, RTVDesc.Format);
+
+							NewTexture->SetRenderTargetViewIndex(new FD3D12RenderTargetView(Device, RTVDesc, NewTexture), RTVIndex++);
 						}
 					}
 					else
 					{
-						if (bIsMultisampled)
+						D3D12_RENDER_TARGET_VIEW_DESC RTVDesc;
+						FMemory::Memzero(RTVDesc);
+
+						RTVDesc.Format = PlatformRenderTargetFormat;
+
+						if (bTextureArray || bCubeTexture)
 						{
-							RTVDesc.ViewDimension = D3D12_RTV_DIMENSION_TEXTURE2DMS;
-							// Nothing to set
+							if (bIsMultisampled)
+							{
+								RTVDesc.ViewDimension = D3D12_RTV_DIMENSION_TEXTURE2DMSARRAY;
+								RTVDesc.Texture2DMSArray.FirstArraySlice = 0;
+								RTVDesc.Texture2DMSArray.ArraySize = TextureDesc.DepthOrArraySize;
+							}
+							else
+							{
+								RTVDesc.ViewDimension = D3D12_RTV_DIMENSION_TEXTURE2DARRAY;
+								RTVDesc.Texture2DArray.FirstArraySlice = 0;
+								RTVDesc.Texture2DArray.ArraySize = TextureDesc.DepthOrArraySize;
+								RTVDesc.Texture2DArray.MipSlice = MipIndex;
+								RTVDesc.Texture2DArray.PlaneSlice = GetPlaneSliceFromViewFormat(PlatformResourceFormat, RTVDesc.Format);
+							}
 						}
 						else
 						{
-							RTVDesc.ViewDimension = D3D12_RTV_DIMENSION_TEXTURE2D;
-							RTVDesc.Texture2D.MipSlice = MipIndex;
-							RTVDesc.Texture2D.PlaneSlice = GetPlaneSliceFromViewFormat(PlatformResourceFormat, RTVDesc.Format);
+							if (bIsMultisampled)
+							{
+								RTVDesc.ViewDimension = D3D12_RTV_DIMENSION_TEXTURE2DMS;
+								// Nothing to set
+							}
+							else
+							{
+								RTVDesc.ViewDimension = D3D12_RTV_DIMENSION_TEXTURE2D;
+								RTVDesc.Texture2D.MipSlice = MipIndex;
+								RTVDesc.Texture2D.PlaneSlice = GetPlaneSliceFromViewFormat(PlatformResourceFormat, RTVDesc.Format);
+							}
 						}
-					}
 
-					NewTexture->SetRenderTargetViewIndex(new FD3D12RenderTargetView(Device, RTVDesc, NewTexture), RTVIndex++);
+						NewTexture->SetRenderTargetViewIndex(new FD3D12RenderTargetView(Device, RTVDesc, NewTexture), RTVIndex++);
+					}
 				}
 			}
-		}
 
-		if (bCreateDSV)
-		{
-			// Create a depth-stencil-view for the texture.
-			D3D12_DEPTH_STENCIL_VIEW_DESC DSVDesc = {};
-			DSVDesc.Format = FindDepthStencilDXGIFormat(PlatformResourceFormat);
-			if (bTextureArray || bCubeTexture)
-			{
-				if (bIsMultisampled)
-				{
-					DSVDesc.ViewDimension = D3D12_DSV_DIMENSION_TEXTURE2DMSARRAY;
-					DSVDesc.Texture2DMSArray.FirstArraySlice = 0;
-					DSVDesc.Texture2DMSArray.ArraySize = TextureDesc.DepthOrArraySize;
-				}
-				else
-				{
-					DSVDesc.ViewDimension = D3D12_DSV_DIMENSION_TEXTURE2DARRAY;
-					DSVDesc.Texture2DArray.FirstArraySlice = 0;
-					DSVDesc.Texture2DArray.ArraySize = TextureDesc.DepthOrArraySize;
-					DSVDesc.Texture2DArray.MipSlice = 0;
-				}
-			}
-			else
+			if (bCreateDSV)
 			{
-				if (bIsMultisampled)
+				// Create a depth-stencil-view for the texture.
+				D3D12_DEPTH_STENCIL_VIEW_DESC DSVDesc = {};
+				DSVDesc.Format = FindDepthStencilDXGIFormat(PlatformResourceFormat);
+				if (bTextureArray || bCubeTexture)
 				{
-					DSVDesc.ViewDimension = D3D12_DSV_DIMENSION_TEXTURE2DMS;
-					// Nothing to set
+					if (bIsMultisampled)
+					{
+						DSVDesc.ViewDimension = D3D12_DSV_DIMENSION_TEXTURE2DMSARRAY;
+						DSVDesc.Texture2DMSArray.FirstArraySlice = 0;
+						DSVDesc.Texture2DMSArray.ArraySize = TextureDesc.DepthOrArraySize;
+					}
+					else
+					{
+						DSVDesc.ViewDimension = D3D12_DSV_DIMENSION_TEXTURE2DARRAY;
+						DSVDesc.Texture2DArray.FirstArraySlice = 0;
+						DSVDesc.Texture2DArray.ArraySize = TextureDesc.DepthOrArraySize;
+						DSVDesc.Texture2DArray.MipSlice = 0;
+					}
 				}
 				else
 				{
-					DSVDesc.ViewDimension = D3D12_DSV_DIMENSION_TEXTURE2D;
-					DSVDesc.Texture2D.MipSlice = 0;
+					if (bIsMultisampled)
+					{
+						DSVDesc.ViewDimension = D3D12_DSV_DIMENSION_TEXTURE2DMS;
+						// Nothing to set
+					}
+					else
+					{
+						DSVDesc.ViewDimension = D3D12_DSV_DIMENSION_TEXTURE2D;
+						DSVDesc.Texture2D.MipSlice = 0;
+					}
 				}
-			}
 
-			const bool HasStencil = HasStencilBits(DSVDesc.Format);
-			for (uint32 AccessType = 0; AccessType < FExclusiveDepthStencil::MaxIndex; ++AccessType)
-			{
-				// Create a read-only access views for the texture.
-				DSVDesc.Flags = (AccessType & FExclusiveDepthStencil::DepthRead_StencilWrite) ? D3D12_DSV_FLAG_READ_ONLY_DEPTH : D3D12_DSV_FLAG_NONE;
-				if (HasStencil)
+				const bool HasStencil = HasStencilBits(DSVDesc.Format);
+				for (uint32 AccessType = 0; AccessType < FExclusiveDepthStencil::MaxIndex; ++AccessType)
 				{
-					DSVDesc.Flags |= (AccessType & FExclusiveDepthStencil::DepthWrite_StencilRead) ? D3D12_DSV_FLAG_READ_ONLY_STENCIL : D3D12_DSV_FLAG_NONE;
-				}
+					// Create a read-only access views for the texture.
+					DSVDesc.Flags = (AccessType & FExclusiveDepthStencil::DepthRead_StencilWrite) ? D3D12_DSV_FLAG_READ_ONLY_DEPTH : D3D12_DSV_FLAG_NONE;
+					if (HasStencil)
+					{
+						DSVDesc.Flags |= (AccessType & FExclusiveDepthStencil::DepthWrite_StencilRead) ? D3D12_DSV_FLAG_READ_ONLY_STENCIL : D3D12_DSV_FLAG_NONE;
+					}
 
-				NewTexture->SetDepthStencilView(new FD3D12DepthStencilView(Device, DSVDesc, NewTexture, HasStencil), AccessType);
+					NewTexture->SetDepthStencilView(new FD3D12DepthStencilView(Device, DSVDesc, NewTexture, HasStencil), AccessType);
+				}
 			}
-		}
 
-		// Create a shader resource view for the texture.
-		if (bCreateShaderResource)
-		{
-			D3D12_SHADER_RESOURCE_VIEW_DESC SRVDesc = {};
-			SRVDesc.Shader4ComponentMapping = D3D12_DEFAULT_SHADER_4_COMPONENT_MAPPING;
-			SRVDesc.Format = PlatformShaderResourceFormat;
-
-			if (bCubeTexture && bTextureArray)
-			{
-				SRVDesc.ViewDimension = D3D12_SRV_DIMENSION_TEXTURECUBEARRAY;
-				SRVDesc.TextureCubeArray.MostDetailedMip = 0;
-				SRVDesc.TextureCubeArray.MipLevels = NumMips;
-				SRVDesc.TextureCubeArray.First2DArrayFace = 0;
-				SRVDesc.TextureCubeArray.NumCubes = SizeZ / 6;
-			}
-			else if (bCubeTexture)
+			// Create a shader resource view for the texture.
+			if (bCreateShaderResource)
 			{
-				SRVDesc.ViewDimension = D3D12_SRV_DIMENSION_TEXTURECUBE;
-				SRVDesc.TextureCube.MostDetailedMip = 0;
-				SRVDesc.TextureCube.MipLevels = NumMips;
-			}
-			else if (bTextureArray)
-			{
-				if (bIsMultisampled)
+				D3D12_SHADER_RESOURCE_VIEW_DESC SRVDesc = {};
+				SRVDesc.Shader4ComponentMapping = D3D12_DEFAULT_SHADER_4_COMPONENT_MAPPING;
+				SRVDesc.Format = PlatformShaderResourceFormat;
+
+				if (bCubeTexture && bTextureArray)
 				{
-					SRVDesc.ViewDimension = D3D12_SRV_DIMENSION_TEXTURE2DMSARRAY;
-					SRVDesc.Texture2DMSArray.FirstArraySlice = 0;
-					SRVDesc.Texture2DMSArray.ArraySize = TextureDesc.DepthOrArraySize;
-					//SRVDesc.Texture2DArray.PlaneSlice = GetPlaneSliceFromViewFormat(PlatformResourceFormat, SRVDesc.Format);
+					SRVDesc.ViewDimension = D3D12_SRV_DIMENSION_TEXTURECUBEARRAY;
+					SRVDesc.TextureCubeArray.MostDetailedMip = 0;
+					SRVDesc.TextureCubeArray.MipLevels = NumMips;
+					SRVDesc.TextureCubeArray.First2DArrayFace = 0;
+					SRVDesc.TextureCubeArray.NumCubes = SizeZ / 6;
 				}
-				else
+				else if (bCubeTexture)
 				{
-					SRVDesc.ViewDimension = D3D12_SRV_DIMENSION_TEXTURE2DARRAY;
-					SRVDesc.Texture2DArray.MostDetailedMip = 0;
-					SRVDesc.Texture2DArray.MipLevels = NumMips;
-					SRVDesc.Texture2DArray.FirstArraySlice = 0;
-					SRVDesc.Texture2DArray.ArraySize = TextureDesc.DepthOrArraySize;
-					SRVDesc.Texture2DArray.PlaneSlice = GetPlaneSliceFromViewFormat(PlatformResourceFormat, SRVDesc.Format);
+					SRVDesc.ViewDimension = D3D12_SRV_DIMENSION_TEXTURECUBE;
+					SRVDesc.TextureCube.MostDetailedMip = 0;
+					SRVDesc.TextureCube.MipLevels = NumMips;
 				}
-			}
-			else
-			{
-				if (bIsMultisampled)
+				else if (bTextureArray)
 				{
-					SRVDesc.ViewDimension = D3D12_SRV_DIMENSION_TEXTURE2DMS;
-					// Nothing to set
+					if (bIsMultisampled)
+					{
+						SRVDesc.ViewDimension = D3D12_SRV_DIMENSION_TEXTURE2DMSARRAY;
+						SRVDesc.Texture2DMSArray.FirstArraySlice = 0;
+						SRVDesc.Texture2DMSArray.ArraySize = TextureDesc.DepthOrArraySize;
+						//SRVDesc.Texture2DArray.PlaneSlice = GetPlaneSliceFromViewFormat(PlatformResourceFormat, SRVDesc.Format);
+					}
+					else
+					{
+						SRVDesc.ViewDimension = D3D12_SRV_DIMENSION_TEXTURE2DARRAY;
+						SRVDesc.Texture2DArray.MostDetailedMip = 0;
+						SRVDesc.Texture2DArray.MipLevels = NumMips;
+						SRVDesc.Texture2DArray.FirstArraySlice = 0;
+						SRVDesc.Texture2DArray.ArraySize = TextureDesc.DepthOrArraySize;
+						SRVDesc.Texture2DArray.PlaneSlice = GetPlaneSliceFromViewFormat(PlatformResourceFormat, SRVDesc.Format);
+					}
 				}
 				else
 				{
-					SRVDesc.ViewDimension = D3D12_SRV_DIMENSION_TEXTURE2D;
-					SRVDesc.Texture2D.MostDetailedMip = 0;
-					SRVDesc.Texture2D.MipLevels = NumMips;
-					SRVDesc.Texture2D.PlaneSlice = GetPlaneSliceFromViewFormat(PlatformResourceFormat, SRVDesc.Format);
+					if (bIsMultisampled)
+					{
+						SRVDesc.ViewDimension = D3D12_SRV_DIMENSION_TEXTURE2DMS;
+						// Nothing to set
+					}
+					else
+					{
+						SRVDesc.ViewDimension = D3D12_SRV_DIMENSION_TEXTURE2D;
+						SRVDesc.Texture2D.MostDetailedMip = 0;
+						SRVDesc.Texture2D.MipLevels = NumMips;
+						SRVDesc.Texture2D.PlaneSlice = GetPlaneSliceFromViewFormat(PlatformResourceFormat, SRVDesc.Format);
+					}
 				}
-			}
 
-			NewTexture->SetShaderResourceView(new FD3D12ShaderResourceView(Device, SRVDesc, NewTexture));
+				NewTexture->SetShaderResourceView(new FD3D12ShaderResourceView(Device, SRVDesc, NewTexture));
+			}
 		}
 
 		return NewTexture;
@@ -1108,7 +1154,7 @@ TD3D12Texture2D<BaseResourceType>* FD3D12DynamicRHI::CreateD3D12Texture2D(FRHICo
 #endif // PLATFORM_WINDOWS || PLATFORM_HOLOLENS
 }
 
-FD3D12Texture3D* FD3D12DynamicRHI::CreateD3D12Texture3D(FRHICommandListImmediate* RHICmdList, uint32 SizeX, uint32 SizeY, uint32 SizeZ, EPixelFormat Format, uint32 NumMips, ETextureCreateFlags Flags, ERHIAccess InResourceState, FRHIResourceCreateInfo& CreateInfo)
+FD3D12Texture3D* FD3D12DynamicRHI::CreateD3D12Texture3D(FRHICommandListImmediate* RHICmdList, uint32 SizeX, uint32 SizeY, uint32 SizeZ, EPixelFormat Format, uint32 NumMips, ETextureCreateFlags Flags, ERHIAccess InResourceState, FRHIResourceCreateInfo& CreateInfo, ID3D12ResourceAllocator* ResourceAllocator)
 {
 #if PLATFORM_WINDOWS || PLATFORM_HOLOLENS
 	SCOPE_CYCLE_COUNTER(STAT_D3D12CreateTextureTime);
@@ -1145,6 +1191,9 @@ FD3D12Texture3D* FD3D12DynamicRHI::CreateD3D12Texture3D(FRHICommandListImmediate
 	check(!(Flags & TexCreate_ResolveTargetable));
 	check(Flags & TexCreate_ShaderResource);
 
+	// Transient resources don't get any actual D3D12 resource asigned at this point yet
+	bool bIsTransient = GSupportsTransientResourceAliasing && (Flags & TexCreate_Transient);
+
 	D3D12_CLEAR_VALUE *ClearValuePtr = nullptr;
 	D3D12_CLEAR_VALUE ClearValue;
 	if (bCreateRTV && CreateInfo.ClearValueBinding.ColorBinding == EClearBinding::EColorBound)
@@ -1162,9 +1211,24 @@ FD3D12Texture3D* FD3D12DynamicRHI::CreateD3D12Texture3D(FRHICommandListImmediate
 	{
 		FD3D12Texture3D* Texture3D = new FD3D12Texture3D(Device, SizeX, SizeY, SizeZ, NumMips, Format, Flags, CreateInfo.ClearValueBinding);
 
-		VERIFYD3D12CREATETEXTURERESULT(Device->GetTextureAllocator().AllocateTexture(TextureDesc, ClearValuePtr, Format, Texture3D->ResourceLocation, (CreateInfo.BulkData != nullptr) ? D3D12_RESOURCE_STATE_COPY_DEST : InitialState, CreateInfo.DebugName), TextureDesc, Device->GetDevice());
+		if (CreateInfo.DebugName)
+		{
+			Texture3D->SetName(CreateInfo.DebugName);
+		}
+
+		if (ResourceAllocator)
+		{
+			ResourceAllocator->AllocateTexture(D3D12_HEAP_TYPE_DEFAULT, TextureDesc, (EPixelFormat)Format, ED3D12ResourceStateMode::Default, InitialState, ClearValuePtr, CreateInfo.DebugName, Texture3D->ResourceLocation);
+		}
+		else
+		{
+			VERIFYD3D12CREATETEXTURERESULT(Device->GetTextureAllocator().AllocateTexture(TextureDesc, ClearValuePtr, Format, Texture3D->ResourceLocation, (CreateInfo.BulkData != nullptr) ? D3D12_RESOURCE_STATE_COPY_DEST : InitialState, CreateInfo.DebugName), TextureDesc, Device->GetDevice());
+		}
 		Texture3D->ResourceLocation.SetOwner(Texture3D);
 
+		// Mark transient
+		Texture3D->ResourceLocation.SetTransient(bIsTransient);
+
 		// Unlock immediately if no initial data
 		if (CreateInfo.BulkData == nullptr)
 		{
@@ -1227,12 +1291,14 @@ FD3D12Texture3D* FD3D12DynamicRHI::CreateD3D12Texture3D(FRHICommandListImmediate
 
 FTexture2DRHIRef FD3D12DynamicRHI::RHICreateTexture2D_RenderThread(class FRHICommandListImmediate& RHICmdList, uint32 SizeX, uint32 SizeY, uint8 Format, uint32 NumMips, uint32 NumSamples, ETextureCreateFlags Flags, ERHIAccess InResourceState, FRHIResourceCreateInfo& CreateInfo)
 {
-	return CreateD3D12Texture2D<FD3D12BaseTexture2D>(&RHICmdList, SizeX, SizeY, 1, false, false, (EPixelFormat) Format, NumMips, NumSamples, Flags, InResourceState, CreateInfo);
+	ID3D12ResourceAllocator* ResourceAllocator = nullptr;
+	return CreateD3D12Texture2D<FD3D12BaseTexture2D>(&RHICmdList, SizeX, SizeY, 1, false, false, (EPixelFormat) Format, NumMips, NumSamples, Flags, InResourceState, CreateInfo, ResourceAllocator);
 }
 
 FTexture2DRHIRef FD3D12DynamicRHI::RHICreateTexture2D(uint32 SizeX, uint32 SizeY, uint8 Format, uint32 NumMips, uint32 NumSamples, ETextureCreateFlags Flags, ERHIAccess InResourceState, FRHIResourceCreateInfo& CreateInfo)
 {
-	return CreateD3D12Texture2D<FD3D12BaseTexture2D>(nullptr, SizeX, SizeY, 1, false, false, (EPixelFormat) Format, NumMips, NumSamples, Flags, InResourceState, CreateInfo);
+	ID3D12ResourceAllocator* ResourceAllocator = nullptr;
+	return CreateD3D12Texture2D<FD3D12BaseTexture2D>(nullptr, SizeX, SizeY, 1, false, false, (EPixelFormat) Format, NumMips, NumSamples, Flags, InResourceState, CreateInfo, ResourceAllocator);
 }
 
 FTexture2DRHIRef FD3D12DynamicRHI::RHIAsyncCreateTexture2D(uint32 SizeX, uint32 SizeY, uint8 Format, uint32 NumMips, ETextureCreateFlags Flags, ERHIAccess InResourceState, void** InitialMipData, uint32 NumInitialMips)
@@ -1504,25 +1570,31 @@ void FD3D12DynamicRHI::RHICopySharedMips(FRHITexture2D* DestTexture2DRHI, FRHITe
 FTexture2DArrayRHIRef FD3D12DynamicRHI::RHICreateTexture2DArray_RenderThread(class FRHICommandListImmediate& RHICmdList, uint32 SizeX, uint32 SizeY, uint32 SizeZ, uint8 Format, uint32 NumMips, uint32 NumSamples, ETextureCreateFlags Flags, ERHIAccess InResourceState, FRHIResourceCreateInfo& CreateInfo)
 {
 	check(SizeZ >= 1);
-	return CreateD3D12Texture2D<FD3D12BaseTexture2DArray>(&RHICmdList, SizeX, SizeY, SizeZ, true, false, (EPixelFormat) Format, NumMips, NumSamples, Flags, InResourceState, CreateInfo);
+
+	ID3D12ResourceAllocator* ResourceAllocator = nullptr;
+	return CreateD3D12Texture2D<FD3D12BaseTexture2DArray>(&RHICmdList, SizeX, SizeY, SizeZ, true, false, (EPixelFormat) Format, NumMips, NumSamples, Flags, InResourceState, CreateInfo, ResourceAllocator);
 }
 
 FTexture2DArrayRHIRef FD3D12DynamicRHI::RHICreateTexture2DArray(uint32 SizeX, uint32 SizeY, uint32 SizeZ, uint8 Format, uint32 NumMips, uint32 NumSamples, ETextureCreateFlags Flags, ERHIAccess InResourceState, FRHIResourceCreateInfo& CreateInfo)
 {
 	check(SizeZ >= 1);
-	return CreateD3D12Texture2D<FD3D12BaseTexture2DArray>(nullptr, SizeX, SizeY, SizeZ, true, false, (EPixelFormat) Format, NumMips, NumSamples, Flags, InResourceState, CreateInfo);
+
+	ID3D12ResourceAllocator* ResourceAllocator = nullptr;
+	return CreateD3D12Texture2D<FD3D12BaseTexture2DArray>(nullptr, SizeX, SizeY, SizeZ, true, false, (EPixelFormat) Format, NumMips, NumSamples, Flags, InResourceState, CreateInfo, ResourceAllocator);
 }
 
 FTexture3DRHIRef FD3D12DynamicRHI::RHICreateTexture3D_RenderThread(class FRHICommandListImmediate& RHICmdList, uint32 SizeX, uint32 SizeY, uint32 SizeZ, uint8 Format, uint32 NumMips, ETextureCreateFlags Flags, ERHIAccess InResourceState, FRHIResourceCreateInfo& CreateInfo)
 {
-	return CreateD3D12Texture3D(&RHICmdList, SizeX, SizeY, SizeZ, (EPixelFormat) Format, NumMips, Flags, InResourceState, CreateInfo);
+	ID3D12ResourceAllocator* ResourceAllocator = nullptr;
+	return CreateD3D12Texture3D(&RHICmdList, SizeX, SizeY, SizeZ, (EPixelFormat) Format, NumMips, Flags, InResourceState, CreateInfo, ResourceAllocator);
 }
 
 FTexture3DRHIRef FD3D12DynamicRHI::RHICreateTexture3D(uint32 SizeX, uint32 SizeY, uint32 SizeZ, uint8 Format, uint32 NumMips, ETextureCreateFlags Flags, ERHIAccess InResourceState, FRHIResourceCreateInfo& CreateInfo)
 {
 	check(SizeZ >= 1);
 #if PLATFORM_WINDOWS || PLATFORM_HOLOLENS
-	return CreateD3D12Texture3D(nullptr, SizeX, SizeY, SizeZ, (EPixelFormat) Format, NumMips, Flags, InResourceState, CreateInfo);
+	ID3D12ResourceAllocator* ResourceAllocator = nullptr;
+	return CreateD3D12Texture3D(nullptr, SizeX, SizeY, SizeZ, (EPixelFormat) Format, NumMips, Flags, InResourceState, CreateInfo, ResourceAllocator);
 #else
 	checkf(false, TEXT("XBOX_CODE_MERGE : Removed. The Xbox platform version should be used."));
 	return nullptr;
@@ -1627,7 +1699,8 @@ FTexture2DRHIRef FD3D12DynamicRHI::AsyncReallocateTexture2D_RenderThread(class F
 	// Allocate a new texture.
 	FRHIResourceCreateInfo CreateInfo(TEXT("AsyncReallocateTexture2D_RenderThread"));
 	ERHIAccess RHIAccess = ERHIAccess::Unknown;
-	FD3D12Texture2D* NewTexture2D = CreateD3D12Texture2D<FD3D12BaseTexture2D>(nullptr, NewSizeX, NewSizeY, 1, false, false, Texture2DRHI->GetFormat(), NewMipCount, 1, Texture2DRHI->GetFlags(), RHIAccess, CreateInfo);
+	ID3D12ResourceAllocator* ResourceAllocator = nullptr;
+	FD3D12Texture2D* NewTexture2D = CreateD3D12Texture2D<FD3D12BaseTexture2D>(nullptr, NewSizeX, NewSizeY, 1, false, false, Texture2DRHI->GetFormat(), NewMipCount, 1, Texture2DRHI->GetFlags(), RHIAccess, CreateInfo, ResourceAllocator);
 	
 	ALLOC_COMMAND_CL(RHICmdList, FRHICommandD3D12AsyncReallocateTexture2D)(Texture2D, NewTexture2D, NewMipCount, NewSizeX, NewSizeY, RequestStatus);
 
@@ -1657,7 +1730,8 @@ FTexture2DRHIRef FD3D12DynamicRHI::RHIAsyncReallocateTexture2D(FRHITexture2D* Te
 	// Allocate a new texture.
 	FRHIResourceCreateInfo CreateInfo(TEXT("RHIAsyncReallocateTexture2D"));
 	ERHIAccess RHIAccess = ERHIAccess::Unknown;
-	FD3D12Texture2D* NewTexture2D = CreateD3D12Texture2D<FD3D12BaseTexture2D>(nullptr, NewSizeX, NewSizeY, 1, false, false, Texture2DRHI->GetFormat(), NewMipCount, 1, Texture2DRHI->GetFlags(), RHIAccess, CreateInfo);
+	ID3D12ResourceAllocator* ResourceAllocator = nullptr;
+	FD3D12Texture2D* NewTexture2D = CreateD3D12Texture2D<FD3D12BaseTexture2D>(nullptr, NewSizeX, NewSizeY, 1, false, false, Texture2DRHI->GetFormat(), NewMipCount, 1, Texture2DRHI->GetFlags(), RHIAccess, CreateInfo, ResourceAllocator);
 	
 	DoAsyncReallocateTexture2D(Texture2D, NewTexture2D, NewMipCount, NewSizeX, NewSizeY, RequestStatus);
 
@@ -2618,22 +2692,26 @@ void FD3D12DynamicRHI::EndUpdateTexture3D_Internal(FUpdateTexture3DData& UpdateD
 	-----------------------------------------------------------------------------*/
 FTextureCubeRHIRef FD3D12DynamicRHI::RHICreateTextureCube_RenderThread(class FRHICommandListImmediate& RHICmdList, uint32 Size, uint8 Format, uint32 NumMips, ETextureCreateFlags Flags, ERHIAccess InResourceState, FRHIResourceCreateInfo& CreateInfo)
 {
-	return CreateD3D12Texture2D<FD3D12BaseTextureCube>(&RHICmdList, Size, Size, 6, false, true, (EPixelFormat) Format, NumMips, 1, Flags, InResourceState, CreateInfo);
+	ID3D12ResourceAllocator* ResourceAllocator = nullptr;
+	return CreateD3D12Texture2D<FD3D12BaseTextureCube>(&RHICmdList, Size, Size, 6, false, true, (EPixelFormat) Format, NumMips, 1, Flags, InResourceState, CreateInfo, ResourceAllocator);
 }
 
 FTextureCubeRHIRef FD3D12DynamicRHI::RHICreateTextureCube(uint32 Size, uint8 Format, uint32 NumMips, ETextureCreateFlags Flags, ERHIAccess InResourceState, FRHIResourceCreateInfo& CreateInfo)
 {
-	return CreateD3D12Texture2D<FD3D12BaseTextureCube>(nullptr, Size, Size, 6, false, true, (EPixelFormat) Format, NumMips, 1, Flags, InResourceState, CreateInfo);
+	ID3D12ResourceAllocator* ResourceAllocator = nullptr;
+	return CreateD3D12Texture2D<FD3D12BaseTextureCube>(nullptr, Size, Size, 6, false, true, (EPixelFormat) Format, NumMips, 1, Flags, InResourceState, CreateInfo, ResourceAllocator);
 }
 
 FTextureCubeRHIRef FD3D12DynamicRHI::RHICreateTextureCubeArray_RenderThread(class FRHICommandListImmediate& RHICmdList, uint32 Size, uint32 ArraySize, uint8 Format, uint32 NumMips, ETextureCreateFlags Flags, ERHIAccess InResourceState, FRHIResourceCreateInfo& CreateInfo)
 {
-	return CreateD3D12Texture2D<FD3D12BaseTextureCube>(&RHICmdList, Size, Size, 6 * ArraySize, true, true, (EPixelFormat) Format, NumMips, 1, Flags, InResourceState, CreateInfo);
+	ID3D12ResourceAllocator* ResourceAllocator = nullptr;
+	return CreateD3D12Texture2D<FD3D12BaseTextureCube>(&RHICmdList, Size, Size, 6 * ArraySize, true, true, (EPixelFormat) Format, NumMips, 1, Flags, InResourceState, CreateInfo, ResourceAllocator);
 }
 
 FTextureCubeRHIRef FD3D12DynamicRHI::RHICreateTextureCubeArray(uint32 Size, uint32 ArraySize, uint8 Format, uint32 NumMips, ETextureCreateFlags Flags, ERHIAccess InResourceState, FRHIResourceCreateInfo& CreateInfo)
 {
-	return CreateD3D12Texture2D<FD3D12BaseTextureCube>(nullptr, Size, Size, 6 * ArraySize, true, true, (EPixelFormat) Format, NumMips, 1, Flags, InResourceState, CreateInfo);
+	ID3D12ResourceAllocator* ResourceAllocator = nullptr;
+	return CreateD3D12Texture2D<FD3D12BaseTextureCube>(nullptr, Size, Size, 6 * ArraySize, true, true, (EPixelFormat) Format, NumMips, 1, Flags, InResourceState, CreateInfo, ResourceAllocator);
 }
 
 void* FD3D12DynamicRHI::RHILockTextureCubeFace(FRHITextureCube* TextureCubeRHI, uint32 FaceIndex, uint32 ArrayIndex, uint32 MipIndex, EResourceLockMode LockMode, uint32& DestStride, bool bLockWithinMiptail)
@@ -2676,8 +2754,10 @@ void FD3D12DynamicRHI::RHIBindDebugLabelName(FRHITexture* TextureRHI, const TCHA
 	{
 		SetName(BaseTexture->GetResource(), Name);
 	}
-
 #endif
+
+	// Also set on RHI object
+	TextureRHI->SetName(Name);
 }
 
 void FD3D12DynamicRHI::RHIVirtualTextureSetFirstMipInMemory(FRHITexture2D* TextureRHI, uint32 FirstMip)
@@ -2709,7 +2789,7 @@ TD3D12Texture2D<BaseResourceType>* FD3D12DynamicRHI::CreateTextureFromResource(b
 	uint32 SizeZ = TextureDesc.DepthOrArraySize;
 	uint32 NumMips = TextureDesc.MipLevels;
 	uint32 NumSamples = TextureDesc.SampleDesc.Count;
-
+	
 	check(TextureDesc.Dimension == D3D12_RESOURCE_DIMENSION_TEXTURE2D);
 	check(bTextureArray || (!bCubeTexture && SizeZ == 1) || (bCubeTexture && SizeZ == 6));
 
diff --git a/Engine/Source/Runtime/D3D12RHI/Private/D3D12Texture.h b/Engine/Source/Runtime/D3D12RHI/Private/D3D12Texture.h
--- a/Engine/Source/Runtime/D3D12RHI/Private/D3D12Texture.h
+++ b/Engine/Source/Runtime/D3D12RHI/Private/D3D12Texture.h
@@ -22,7 +22,7 @@ void SafeCreateTexture2D(FD3D12Device* pDevice,
 
 
 /** Texture base class. */
-class FD3D12TextureBase : public FD3D12BaseShaderResource, public FD3D12TransientResource, public FD3D12LinkedAdapterObject<FD3D12TextureBase>
+class FD3D12TextureBase : public FD3D12BaseShaderResource, public FD3D12LinkedAdapterObject<FD3D12TextureBase>
 {
 public:
 
@@ -237,7 +237,7 @@ public:
 
 	/** Initialization constructor. */
 	TD3D12Texture2D(
-	class FD3D12Device* InParent,
+		class FD3D12Device* InParent,
 		uint32 InSizeX,
 		uint32 InSizeY,
 		uint32 InSizeZ,
diff --git a/Engine/Source/Runtime/D3D12RHI/Private/D3D12TransientResourceAllocator.cpp b/Engine/Source/Runtime/D3D12RHI/Private/D3D12TransientResourceAllocator.cpp
new file mode 100644
--- /dev/null
+++ b/Engine/Source/Runtime/D3D12RHI/Private/D3D12TransientResourceAllocator.cpp
@@ -0,0 +1,785 @@
+// Copyright Epic Games, Inc. All Rights Reserved.
+
+#include "D3D12RHIPrivate.h"
+#include "D3D12TransientResourceAllocator.h"
+
+//PRAGMA_DISABLE_OPTIMIZATION
+
+static int32 GD3D12TransientAllocatorPoolSizeInMB = 128;
+static FAutoConsoleVariableRef CVarD3D12TransientAllocatorPoolSizeInMB(
+	TEXT("d3d12.TransientAllocator.PoolSizeInMB"),
+	GD3D12TransientAllocatorPoolSizeInMB,
+	TEXT("Size of a D3D12 transient allocator pool in MB (Default 128)"),
+	ECVF_ReadOnly);
+
+static int32 GD3D12TransientAllocatorPoolTextures = 1;
+static FAutoConsoleVariableRef CVarD3D12TransientAllocatorPoolTextures(
+	TEXT("d3d12.TransientAllocator.PoolTextures"),
+	GD3D12TransientAllocatorPoolTextures,
+	TEXT("Enable pooling of transient allocated RHITextures in the manager (default enabled)"),
+	ECVF_ReadOnly);
+
+static int32 GD3D12TransientAllocatorPoolBuffers = 0;
+static FAutoConsoleVariableRef CVarD3D12TransientAllocatorPoolBuffers(
+	TEXT("d3d12.TransientAllocator.PoolBuffers"),
+	GD3D12TransientAllocatorPoolBuffers,
+	TEXT("Enable pooling of transient allocated RHIBuffer in the manager (default disabled)"),
+	ECVF_ReadOnly);
+
+uint64 ComputeTextureCreateInfoHash(const FRHITextureCreateInfo& InCreateInfo)
+{
+	// Make sure all padding is removed
+	FRHITextureCreateInfo NewInfo;
+	FPlatformMemory::Memzero(&NewInfo, sizeof(FRHITextureCreateInfo));
+	NewInfo = InCreateInfo;
+	return CityHash64((const char*)&NewInfo, sizeof(FRHITextureCreateInfo));
+}
+
+uint64 ComputeBufferCreateInfoHash(const FRHIBufferCreateInfo& InCreateInfo)
+{
+	return CityHash64((const char*)&InCreateInfo, sizeof(FRHIBufferCreateInfo));
+}
+
+
+//-----------------------------------------------------------------------------
+//	FD3D12TransientMemoryPool
+//-----------------------------------------------------------------------------
+
+uint64 FD3D12TransientMemoryPool::FResourceCreateState::GetHash() const
+{
+	uint64 Hash = CityHash64((const char*)&AllocationOffset, sizeof(uint64));
+	Hash = CityHash64WithSeed((const char*)&ResourceDesc, sizeof(D3D12_RESOURCE_DESC), Hash);
+	Hash = CityHash64WithSeed((const char*)&ClearValue, sizeof(D3D12_CLEAR_VALUE), Hash);
+	return Hash;
+}
+
+
+FD3D12TransientMemoryPool::~FD3D12TransientMemoryPool()
+{
+	check(AllocatedBlocks == 0);
+
+	ClearActiveResources();
+
+	TArray<FD3D12Resource*> CachedResources;
+	CachedResourceMap.GenerateValueArray(CachedResources);
+	for (FD3D12Resource* Resource : CachedResources)
+	{
+		check(Resource->GetRefCount() == 1);
+		Resource->Release();
+	}
+	CachedResourceMap.Empty();
+}
+
+
+void FD3D12TransientMemoryPool::ResetPool()
+{
+	check(AllocatedBlocks == 0);
+	check(FreeBlocks.Num() == 1);
+	check(FreeBlocks[0]->GetOffset() == 0);
+	check(FreeBlocks[0]->GetSize() == PoolSize);
+	check(FreeBlocks[0]->GetAlignment() == PoolAlignment);
+
+	ClearActiveResources();
+}
+
+
+void FD3D12TransientMemoryPool::SetPoolIndex(int16 InNewPoolIndex)
+{
+	PoolIndex = InNewPoolIndex;
+
+	// Update pool index on single free block (validated during reset)
+	FRHIPoolAllocationData* FreeBlock = FreeBlocks[0];
+	FreeBlock->RemoveFromLinkedList();
+	FreeBlock->InitAsFree(PoolIndex, PoolSize, PoolAlignment, 0);
+	HeadBlock.AddAfter(FreeBlock);
+}
+
+
+void FD3D12TransientMemoryPool::ClearActiveResources()
+{
+	ActiveResources.Empty(ActiveResources.Num());
+}
+
+
+void FD3D12TransientMemoryPool::CheckActiveResources(const FInt64Range& InAllocationRange, TArray<FD3D12Resource *>& OverlappingResources)
+{		
+	// Do we already have active resources in the requested range
+	for (ActiveResourceData& ActiveResource : ActiveResources)
+	{
+		if (ActiveResource.ActiveRange.Overlaps(InAllocationRange))
+		{
+			check(ActiveResource.Resource != nullptr);
+
+			// Add to dependent list of overlapping resources
+			OverlappingResources.Add(ActiveResource.Resource);
+
+			// update the active ranges
+			if (ActiveResource.ActiveRange.GetLowerBoundValue() >= InAllocationRange.GetLowerBoundValue())
+			{
+				if (ActiveResource.ActiveRange.GetUpperBoundValue() > InAllocationRange.GetUpperBoundValue())
+				{
+					ActiveResource.ActiveRange.SetLowerBoundValue(InAllocationRange.GetUpperBoundValue());
+				}
+				else
+				{
+					// Full overlap, make empty range
+					ActiveResource.ActiveRange = FInt64Range(ActiveResource.ActiveRange.GetLowerBoundValue(), ActiveResource.ActiveRange.GetLowerBoundValue());
+				}
+			}
+			else
+			{
+				ActiveResource.ActiveRange.SetUpperBoundValue(InAllocationRange.GetLowerBoundValue());
+			}
+
+			// Mark as invalid because it has no active ranges anymore
+			if (ActiveResource.ActiveRange.IsEmpty())
+			{
+				ActiveResource.Resource = nullptr;
+			}
+		}
+	}
+}
+
+
+FD3D12Resource* FD3D12TransientMemoryPool::FindResourceInCache(uint64 InAllocationOffset, const D3D12_RESOURCE_DESC& InDesc, const D3D12_CLEAR_VALUE* InClearValue, const TCHAR* InName)
+{
+	// Try and find cached placed resource at given offset and location
+	FResourceCreateState CreateState;
+	CreateState.AllocationOffset = InAllocationOffset;
+	CreateState.ResourceDesc = InDesc;
+	if (InClearValue)
+	{
+		CreateState.ClearValue = *InClearValue;
+	}
+	else
+	{
+		FPlatformMemory::Memzero(&CreateState.ClearValue, sizeof(D3D12_CLEAR_VALUE));
+	}
+
+	uint64 CreateStateHash = CreateState.GetHash();
+	FD3D12Resource** Resource = CachedResourceMap.Find(CreateStateHash);
+	if (Resource)
+	{
+		check((*Resource)->GetDesc() == InDesc);		
+		CachedResourceMap.Remove(CreateStateHash);
+		return *Resource;
+	}
+	else
+	{
+		return nullptr;
+	}
+}
+
+
+void FD3D12TransientMemoryPool::ReleaseResource(FD3D12Resource* InResource, FRHIPoolAllocationData& InReleasedAllocationData, uint64 InFenceValue)
+{
+	check(InResource->GetHeap() == BackingHeap);
+
+	uint64 AllocationOffset = InReleasedAllocationData.GetOffset();
+	uint64 AllocationSize = InReleasedAllocationData.GetSize();
+
+	// Placed resource object is 'never' freed - it can be reused again this or next frame
+	// Will only be destroyed when the allocator is destroyed
+	FResourceCreateState CreateState;
+	CreateState.AllocationOffset = AllocationOffset;
+	CreateState.ResourceDesc = InResource->GetDesc();
+	CreateState.ClearValue = InResource->GetClearValue();
+
+	uint64 CreateStateHash = CreateState.GetHash();
+	check(CachedResourceMap.Find(CreateStateHash) == nullptr);
+	CachedResourceMap.Add(CreateStateHash, InResource);
+	
+	// Track this resource so aliasing barriers can be added if used again
+	ActiveResourceData ResourceData;
+	ResourceData.AllocationRange = FInt64Range(AllocationOffset, AllocationOffset + AllocationSize);
+	ResourceData.ActiveRange = ResourceData.AllocationRange;
+	ResourceData.Resource = InResource;
+	ActiveResources.Add(ResourceData);
+
+	// Free the pool data so this range can be reallocated again immediately (active range is tracked for aliasing barriers)
+	FRHIPoolAllocationData LockedAllocationData;
+	bool bLocked = true;
+	LockedAllocationData.MoveFrom(InReleasedAllocationData, bLocked);
+	Deallocate(LockedAllocationData);
+
+	// Update the last used frame fence (used during garbage collection)
+	UpdateLastUsedFrameFence(InFenceValue);
+}
+
+//-----------------------------------------------------------------------------
+//	FD3D12TransientMemoryPoolManager
+//-----------------------------------------------------------------------------
+
+
+FD3D12TransientMemoryPoolManager::FD3D12TransientMemoryPoolManager(FD3D12Device* InDevice, FRHIGPUMask VisibleNodes)
+	: FD3D12DeviceChild(InDevice)
+	, FD3D12MultiNodeGPUObject(InDevice->GetGPUMask(), VisibleNodes)
+{
+	// texture only interesting in VRAM for now
+	InitConfig.HeapType = D3D12_HEAP_TYPE_DEFAULT;
+
+	// unused for textures because placed and not suballocated
+	InitConfig.ResourceFlags = D3D12_RESOURCE_FLAG_NONE;
+	InitConfig.InitialResourceState = D3D12_RESOURCE_STATE_COMMON;
+
+	// Support RT and UAV for now (Tier 2 support) - need to add different pools for Tier 1
+	InitConfig.HeapFlags = D3D12_HEAP_FLAGS(0); // 0 means nothing is denied
+
+	PoolSize = GD3D12TransientAllocatorPoolSizeInMB * 1024 * 1024;
+	PoolAlignment = D3D12_DEFAULT_RESOURCE_PLACEMENT_ALIGNMENT;
+	MaxAllocationSize = PoolSize;
+}
+
+
+void FD3D12TransientMemoryPoolManager::Destroy()
+{
+	for (int32 PoolIndex = 0; PoolIndex < Pools.Num(); ++PoolIndex)
+	{
+		Pools[PoolIndex]->Destroy();
+		delete(Pools[PoolIndex]);
+	}
+	Pools.Empty();
+}
+
+
+void FD3D12TransientMemoryPoolManager::BeginFrame()
+{
+	TextureMemoryStats.Reset();
+	BufferMemoryStats.Reset();
+}
+
+
+void FD3D12TransientMemoryPoolManager::EndFrame()
+{
+	FScopeLock Lock(&CS);
+
+	static const int32 FrameLag = 20;
+
+	// Trim empty allocators if not used in last n frames
+	FD3D12Adapter* Adapter = GetParentDevice()->GetParentAdapter();
+	FD3D12Fence& FrameFence = Adapter->GetFrameFence();
+	const uint64 CompletedFence = FrameFence.UpdateLastCompletedFence();
+	for (int32 PoolIndex = 0; PoolIndex < Pools.Num(); ++PoolIndex)
+	{
+		FD3D12MemoryPool* MemoryPool = (FD3D12MemoryPool*)Pools[PoolIndex];
+		if (MemoryPool != nullptr && MemoryPool->IsEmpty() && (MemoryPool->GetLastUsedFrameFence() + FrameLag <= CompletedFence))
+		{
+			MemoryPool->Destroy();
+			delete(MemoryPool);
+
+			Pools.RemoveAt(PoolIndex);
+			PoolIndex--;
+		}
+	}
+}
+
+
+FD3D12TransientMemoryPool* FD3D12TransientMemoryPoolManager::GetOrCreateMemoryPool(int16 InPoolIndex)
+{
+	FScopeLock Lock(&CS);
+
+	FD3D12TransientMemoryPool* MemoryPool = nullptr;
+	if (Pools.Num() > 0)
+	{
+		MemoryPool = Pools.Pop(false);
+		MemoryPool->SetPoolIndex(InPoolIndex);
+	}
+	else
+	{
+		TRACE_CPUPROFILER_EVENT_SCOPE(D3D12RHI::AllocateTransientMemoryPool);
+
+		MemoryPool = new FD3D12TransientMemoryPool(GetParentDevice(), GetVisibilityMask(), InitConfig,
+			TEXT("TransientResourceMemoryPool"), EResourceAllocationStrategy::kPlacedResource, InPoolIndex, PoolSize, PoolAlignment);
+		MemoryPool->Init();
+
+#if PLATFORM_WINDOWS
+		// Boost priority to make sure it's not paged out
+		ID3D12Device* D3DDevice = GetParentDevice()->GetDevice();
+		TRefCountPtr<ID3D12Device5> D3DDevice5;
+		if (SUCCEEDED(D3DDevice->QueryInterface(IID_PPV_ARGS(D3DDevice5.GetInitReference()))))
+		{
+			ID3D12Pageable* HeapResource = MemoryPool->GetBackingHeap()->GetHeap();
+			D3D12_RESIDENCY_PRIORITY HeapPriority = D3D12_RESIDENCY_PRIORITY_HIGH;
+			D3DDevice5->SetResidencyPriority(1, &HeapResource, &HeapPriority);
+		}
+#endif // PLATFORM_WINDOWS
+	}
+
+	return MemoryPool;
+}
+
+
+void FD3D12TransientMemoryPoolManager::ReleaseMemoryPool(FD3D12TransientMemoryPool* InMemoryPool)
+{
+	FScopeLock Lock(&CS);
+	InMemoryPool->ResetPool();
+	Pools.Add(InMemoryPool);
+}
+
+
+FD3D12PooledTextureData FD3D12TransientMemoryPoolManager::GetPooledTexture(const FRHITextureCreateInfo& InCreateInfo, const TCHAR* InDebugName)
+{
+	FD3D12PooledTextureData Result;
+	if (GD3D12TransientAllocatorPoolTextures)
+	{
+		uint64 CreateHash = ComputeTextureCreateInfoHash(InCreateInfo);
+
+		// Try and search the cache for resource of same creation info (inside lock)
+		FScopeLock Lock(&CS);
+		TArray<FD3D12PooledTextureData>* FindResult = FreeTextures.Find(CreateHash);
+		if (FindResult && FindResult->Num() > 0)
+		{
+			// Make sure they are not referenced anymore! (should in theory not happen)
+			//check((*FindResult)[0].RHITexture->GetRefCount() == 1);
+			for (int32 Index = 0; Index < FindResult->Num(); ++Index)
+			{
+				FD3D12PooledTextureData& TextureData = (*FindResult)[Index];
+				if (TextureData.RHITexture->GetRefCount() == 1)
+				{
+					Result = TextureData;
+					check(Result.CreateInfo == InCreateInfo);
+					FindResult->RemoveAt(Index);
+					break;
+				}
+			}
+		}
+	}
+
+	return Result;
+}
+
+
+FD3D12PooledBufferData FD3D12TransientMemoryPoolManager::GetPooledBuffer(const FRHIBufferCreateInfo& InCreateInfo, const TCHAR* InDebugName)
+{
+	FD3D12PooledBufferData Result;
+	if (GD3D12TransientAllocatorPoolBuffers)
+	{
+		uint64 CreateHash = ComputeBufferCreateInfoHash(InCreateInfo);
+
+		// Try and search the cache for resource of same creation info (inside lock)
+		FScopeLock Lock(&CS);
+		TArray<FD3D12PooledBufferData>* FindResult = FreeBuffers.Find(CreateHash);
+		if (FindResult && FindResult->Num() > 0)
+		{
+			Result = FindResult->Pop(false);
+			check(Result.RHIBuffer->GetRefCount() == 1);
+		}
+	}
+
+	return Result;
+}
+
+
+void FD3D12TransientMemoryPoolManager::ReleaseResources(FD3D12TransientResourceAllocator* InAllocator)
+{
+	FScopeLock Lock(&CS);
+
+	// release the textures
+	if (GD3D12TransientAllocatorPoolTextures)
+	{
+		for (FD3D12PooledTextureData& TextureData : InAllocator->AllocatedTextures)
+		{
+			uint64 CreateHash = ComputeTextureCreateInfoHash(TextureData.CreateInfo);
+			TArray<FD3D12PooledTextureData>* FindResult = FreeTextures.Find(CreateHash);
+			if (FindResult)
+			{
+				FindResult->Add(TextureData);
+			}
+			else
+			{
+				TArray<FD3D12PooledTextureData> Textures;
+				Textures.Add(TextureData);
+				FreeTextures.Add(CreateHash, Textures);
+			}
+		}
+	}
+	InAllocator->AllocatedTextures.Empty();
+
+	// release the buffers
+	if (GD3D12TransientAllocatorPoolBuffers)
+	{
+		for (FD3D12PooledBufferData& BufferData : InAllocator->AllocatedBuffers)
+		{
+			uint64 CreateHash = ComputeBufferCreateInfoHash(BufferData.CreateInfo);
+			TArray<FD3D12PooledBufferData>* FindResult = FreeBuffers.Find(CreateHash);
+			if (FindResult)
+			{
+				FindResult->Add(BufferData);
+			}
+			else
+			{
+				TArray<FD3D12PooledBufferData> Buffers;
+				Buffers.Add(BufferData);
+				FreeBuffers.Add(CreateHash, Buffers);
+			}
+		}
+	}
+	InAllocator->AllocatedBuffers.Empty();
+}
+
+
+void FD3D12TransientMemoryPoolManager::UpdateMemoryStats()
+{
+	FScopeLock Lock(&CS);
+
+	uint32 MemoryAllocated = 0;
+	for (FD3D12TransientMemoryPool* MemoryPool : Pools)
+	{
+		MemoryAllocated += MemoryPool->GetPoolSize();
+	}
+
+	SET_MEMORY_STAT(STAT_D3D12TransientMemoryPoolAllocated, MemoryAllocated);
+	SET_MEMORY_STAT(STAT_D3D12TransientMemoryFramePoolUsed, TextureMemoryStats.MaxFrameAllocated + BufferMemoryStats.MaxFrameAllocated);
+	SET_MEMORY_STAT(STAT_D3D12TransientMemoryFrameCommittedUsed, TextureMemoryStats.CommittedAllocated + BufferMemoryStats.CommittedAllocated);
+	SET_DWORD_STAT(STAT_D3D12TransientMemoryPoolAllocations, TextureMemoryStats.PoolAllocations + BufferMemoryStats.PoolAllocations);
+	SET_DWORD_STAT(STAT_D3D12TransientMemoryCommittedAllocations, TextureMemoryStats.CommittedAllocations + BufferMemoryStats.CommittedAllocations);
+
+	SET_MEMORY_STAT(STAT_D3D12TransientMemoryBufferFrameUsed, BufferMemoryStats.MaxFrameAllocated);
+	SET_MEMORY_STAT(STAT_D3D12TransientMemoryBufferFrameTotalRequested, BufferMemoryStats.TotalRequested);
+	SET_MEMORY_STAT(STAT_D3D12TransientMemoryBufferFrameCommittedAllocated, BufferMemoryStats.CommittedAllocated);
+	SET_DWORD_STAT(STAT_D3D12TransientMemoryBufferPoolAllocations, BufferMemoryStats.PoolAllocations);
+	SET_DWORD_STAT(STAT_D3D12TransientMemoryBufferCommittedAllocations, BufferMemoryStats.CommittedAllocations);
+
+	SET_MEMORY_STAT(STAT_D3D12TransientMemoryTextureFrameUsed, TextureMemoryStats.MaxFrameAllocated);
+	SET_MEMORY_STAT(STAT_D3D12TransientMemoryTextureFrameTotalRequested, TextureMemoryStats.TotalRequested);
+	SET_MEMORY_STAT(STAT_D3D12TransientMemoryTextureFrameCommittedAllocated, TextureMemoryStats.CommittedAllocated);
+	SET_DWORD_STAT(STAT_D3D12TransientMemoryTexturePoolAllocations, TextureMemoryStats.PoolAllocations);
+	SET_DWORD_STAT(STAT_D3D12TransientMemoryTextureCommittedAllocations, TextureMemoryStats.CommittedAllocations);
+}
+
+
+//-----------------------------------------------------------------------------
+//	FD3D12TransientResourceAllocator
+//-----------------------------------------------------------------------------
+
+
+FD3D12TransientResourceAllocator::FD3D12TransientResourceAllocator(FD3D12TransientMemoryPoolManager& InMemoryPoolManager) :
+	FD3D12PoolAllocator(
+		InMemoryPoolManager.GetParentDevice(),
+		InMemoryPoolManager.GetVisibilityMask(),
+		InMemoryPoolManager.GetInitConfig(),
+		TEXT("FD3D12TransientResourceAllocator"),
+		EResourceAllocationStrategy::kPlacedResource,
+		InMemoryPoolManager.GetPoolSize(),
+		InMemoryPoolManager.GetPoolAlignment(),
+		InMemoryPoolManager.GetMaxAllocationSize(),
+		FRHIMemoryPool::EFreeListOrder::SortByOffset,
+		false /*defrag*/)
+{ 
+}
+
+
+FD3D12TransientResourceAllocator::~FD3D12TransientResourceAllocator()
+{
+	// release all resources back to pool so they can be reused if enabled (pools have already been freed)
+	FD3D12TransientMemoryPoolManager& PoolManager = GetParentDevice()->GetTransientMemoryPoolManager();
+	PoolManager.ReleaseResources(this);
+
+	PoolManager.UpdateTextureStats(TextureMemoryStats);
+	PoolManager.UpdateBufferStats(BufferMemoryStats);
+}
+
+
+FRHITexture* FD3D12TransientResourceAllocator::CreateTexture(const FRHITextureCreateInfo& InCreateInfo, const TCHAR* InDebugName)
+{
+	TRACE_CPUPROFILER_EVENT_SCOPE(D3D12RHI::CreateTransientTexture);
+
+	// Should be marked transient
+	check((InCreateInfo.Flags & TexCreate_Transient) == TexCreate_Transient);
+
+	// Try and get a pooled texture from the manager
+	FD3D12TransientMemoryPoolManager& PoolManager = GetParentDevice()->GetTransientMemoryPoolManager();
+	FD3D12PooledTextureData TextureData = PoolManager.GetPooledTexture(InCreateInfo, InDebugName);
+
+	FD3D12BaseShaderResource* BaseShaderResource = nullptr;
+
+	// Found something?
+	if (TextureData.RHITexture)
+	{
+		TRACE_CPUPROFILER_EVENT_SCOPE(SetupPoolTransientResource);
+
+		// Get the base shader resource
+		BaseShaderResource = GetBaseShaderResource(TextureData.RHITexture);
+		check(BaseShaderResource);
+		check(!BaseShaderResource->ResourceLocation.IsValid());
+
+		// Setup the optional clear value
+		D3D12_CLEAR_VALUE* ClearValuePtr = nullptr;
+		if (TextureData.ClearValue.Format != DXGI_FORMAT_UNKNOWN)
+		{
+			ClearValuePtr = &TextureData.ClearValue;
+		}
+
+		// Allocate a new FD3D12Resource on the resource location
+		AllocateTexture(D3D12_HEAP_TYPE_DEFAULT, TextureData.ResourceDesc, InCreateInfo.Format, ED3D12ResourceStateMode::MultiState, D3D12_RESOURCE_STATE_COMMON, ClearValuePtr, InDebugName, BaseShaderResource->ResourceLocation);
+
+		// Inform the listeners about the change - should ideally still get from cache
+		{
+			TRACE_CPUPROFILER_EVENT_SCOPE(RenameViews);
+			BaseShaderResource->ResourceRenamed(&BaseShaderResource->ResourceLocation);
+		}
+	}
+	else
+	{
+		// Create a new resource
+		FRHIResourceCreateInfo CreateInfo(InDebugName);
+		CreateInfo.ClearValueBinding = InCreateInfo.ClearValue;
+
+		FRHITexture* RHITexture = nullptr;
+		switch (InCreateInfo.Dimension)
+		{
+		case ETextureDimension::Texture2D:
+		{
+			const bool bTextureArray = false;
+			const bool bCubeTexture = false;
+			ID3D12ResourceAllocator* ResourceAllocator = this;
+			FD3D12Texture2D* Texture2D = FD3D12DynamicRHI::GetD3DRHI()->CreateD3D12Texture2D<FD3D12BaseTexture2D>(nullptr, InCreateInfo.Extent.X, InCreateInfo.Extent.Y, 1, bTextureArray, bCubeTexture, InCreateInfo.Format, InCreateInfo.NumMips, InCreateInfo.NumSamples, InCreateInfo.Flags, ERHIAccess::Discard, CreateInfo, ResourceAllocator);
+			TextureData.RHITexture = Texture2D;
+			TextureData.ResourceDesc = Texture2D->GetResource()->GetDesc();
+			TextureData.ClearValue = Texture2D->GetResource()->GetClearValue();
+			TextureData.CreateInfo = InCreateInfo;
+
+			BaseShaderResource = Texture2D;
+			break;
+		}
+		case ETextureDimension::Texture3D:
+		{
+			// Only support 2d textures for now
+			ID3D12ResourceAllocator* ResourceAllocator = this;
+			FD3D12Texture3D* Texture3D = FD3D12DynamicRHI::GetD3DRHI()->CreateD3D12Texture3D(nullptr, InCreateInfo.Extent.X, InCreateInfo.Extent.Y, InCreateInfo.Depth, InCreateInfo.Format, InCreateInfo.NumMips, InCreateInfo.Flags, ERHIAccess::Discard, CreateInfo, ResourceAllocator);
+			TextureData.RHITexture = Texture3D;
+			TextureData.ResourceDesc = Texture3D->GetResource()->GetDesc();
+			TextureData.ClearValue = Texture3D->GetResource()->GetClearValue();
+			TextureData.CreateInfo = InCreateInfo;
+			
+			BaseShaderResource = Texture3D;
+			break;
+		}
+		default:
+		{
+			// Only support 2d & 3d textures for now
+			check(false);
+			break;
+		}
+		}
+	}
+	
+	SetupAllocatedResource(BaseShaderResource, TextureMemoryStats);
+
+	// keep track of all allocated textures
+	AllocatedTextures.Add(TextureData);
+
+	return TextureData.RHITexture;
+}
+
+
+FRHIBuffer* FD3D12TransientResourceAllocator::CreateBuffer(const FRHIBufferCreateInfo& InCreateInfo, const TCHAR* InDebugName)
+{
+	TRACE_CPUPROFILER_EVENT_SCOPE(D3D12RHI::CreateTransientBuffer);
+
+	// Should be marked transient
+	check((InCreateInfo.Usage & BUF_Transient) == BUF_Transient);
+
+	// Try and get a pooled texture from the manager
+	FD3D12TransientMemoryPoolManager& PoolManager = GetParentDevice()->GetTransientMemoryPoolManager();
+	FD3D12PooledBufferData BufferData = PoolManager.GetPooledBuffer(InCreateInfo, InDebugName);
+
+	FD3D12Buffer* Buffer = nullptr;
+
+	// Found something?
+	if (BufferData.RHIBuffer)
+	{
+		// Get the D3D12 Buffer
+		Buffer = FD3D12DynamicRHI::ResourceCast((FRHIBuffer*)BufferData.RHIBuffer);
+		check(!Buffer->ResourceLocation.IsValid());
+
+		EBufferUsageFlags UsageFlags = InCreateInfo.Usage;
+		D3D12_RESOURCE_DESC Desc;
+		uint32 Alignment;
+		FD3D12Buffer::GetResourceDescAndAlignment(InCreateInfo.Size, InCreateInfo.Stride, UsageFlags, Desc, Alignment);
+
+		// Allocate a new FD3D12Resource on the resource location
+		AllocateResource(D3D12_HEAP_TYPE_DEFAULT, Desc, Desc.Width, Alignment, ED3D12ResourceStateMode::MultiState, D3D12_RESOURCE_STATE_COMMON, nullptr, InDebugName, Buffer->ResourceLocation);
+
+		// Inform the listeners about the change (should be empty?)
+		Buffer->ResourceRenamed(&Buffer->ResourceLocation);
+	}
+	else
+	{
+		FRHIResourceCreateInfo CreateInfo(InDebugName);
+		ID3D12ResourceAllocator* ResourceAllocator = this;
+		Buffer = FD3D12DynamicRHI::GetD3DRHI()->CreateD3D12Buffer(nullptr, InCreateInfo.Size, InCreateInfo.Usage, InCreateInfo.Stride, ERHIAccess::Discard, CreateInfo, ResourceAllocator);
+		BufferData.RHIBuffer = Buffer;
+		BufferData.CreateInfo = InCreateInfo;
+	}
+
+	SetupAllocatedResource((FD3D12BaseShaderResource*)Buffer, BufferMemoryStats);
+	
+	// keep track of all allocated buffers
+	AllocatedBuffers.Add(BufferData);
+
+	return BufferData.RHIBuffer;
+}
+
+
+void FD3D12TransientResourceAllocator::DeallocateMemory(FRHITexture* InTexture)
+{
+	// Get the correct base shader resource taken texture type into account
+	FD3D12BaseShaderResource* BaseShaderResource = GetBaseShaderResource(InTexture);
+	check(BaseShaderResource);
+	DeallocateMemory(BaseShaderResource, TextureMemoryStats);
+}
+
+
+void FD3D12TransientResourceAllocator::DeallocateMemory(FRHIBuffer* InBuffer)
+{
+	// cast to d3d12 object and call shared function
+	FD3D12Buffer* Buffer = FD3D12DynamicRHI::ResourceCast(InBuffer);
+	DeallocateMemory(Buffer, BufferMemoryStats);
+}
+
+
+FD3D12BaseShaderResource* FD3D12TransientResourceAllocator::GetBaseShaderResource(FRHITexture* InRHITexture)
+{
+	FD3D12BaseShaderResource* BaseShaderResource = nullptr;
+	if (FRHITexture2D* RHITexture2D = InRHITexture->GetTexture2D())
+	{
+		FD3D12Texture2D* Texture2D = FD3D12DynamicRHI::ResourceCast(RHITexture2D);
+		BaseShaderResource = Texture2D->GetBaseShaderResource();
+	}
+	else if (FRHITexture3D* RHITexture3D = InRHITexture->GetTexture3D())
+	{
+		FD3D12Texture3D* Texture3D = FD3D12DynamicRHI::ResourceCast(RHITexture3D);
+		BaseShaderResource = Texture3D;
+	}
+	return BaseShaderResource;
+}
+
+
+void FD3D12TransientResourceAllocator::SetupAllocatedResource(FD3D12BaseShaderResource* InBaseShaderResource, FD3D12TransientMemoryStats& Stats)
+{
+	uint64 AllocationSize = InBaseShaderResource->ResourceLocation.GetSize();
+
+	// If pool allocated then collected the overlapping resources as well
+	if (InBaseShaderResource->ResourceLocation.GetAllocatorType() == FD3D12ResourceLocation::AT_Pool)
+	{
+		FScopeLock Lock(&CS);
+				
+		FRHIPoolAllocationData& PoolAllocationData = InBaseShaderResource->ResourceLocation.GetPoolAllocatorPrivateData().PoolData;
+		FD3D12TransientMemoryPool* MemoryPool = ((FD3D12TransientMemoryPool*)Pools[PoolAllocationData.GetPoolIndex()]);
+		FInt64Range AllocationRange(PoolAllocationData.GetOffset(), PoolAllocationData.GetOffset() + PoolAllocationData.GetSize());
+		TArray<FD3D12Resource*> OverlappingResources;
+		MemoryPool->CheckActiveResources(AllocationRange, OverlappingResources);
+		if (OverlappingResources.Num() > 0)
+		{
+			OverlappingResourceData.Add(InBaseShaderResource, OverlappingResources);
+		}
+		
+		// update stats
+		Stats.PoolAllocations++;
+		Stats.CurrentPoolAllocated += AllocationSize;
+		Stats.TotalRequested += AllocationSize;
+		Stats.MaxFrameAllocated = FMath::Max(Stats.MaxFrameAllocated, Stats.CurrentPoolAllocated);
+	}
+	else
+	{
+		// Outside of pool size? Expensive operation if this happens every frame - should this not be allowed?		
+
+		// update stats
+		Stats.CommittedAllocations++;
+		Stats.CommittedAllocated += AllocationSize;
+		Stats.TotalRequested += AllocationSize;
+	}
+
+}
+
+
+void FD3D12TransientResourceAllocator::DeallocateMemory(FD3D12BaseShaderResource* InBaseShaderResource, FD3D12TransientMemoryStats& Stats)
+{
+	// If pool allocated then release the resource back to the pool
+	if (InBaseShaderResource->ResourceLocation.GetAllocatorType() == FD3D12ResourceLocation::AT_Pool)
+	{
+		FScopeLock Lock(&CS);
+
+		// update stats
+		uint64 AllocationSize = InBaseShaderResource->ResourceLocation.GetSize();
+		Stats.CurrentPoolAllocated -= AllocationSize;
+
+		FD3D12Adapter* Adapter = GetParentDevice()->GetParentAdapter();
+		FD3D12Fence& FrameFence = Adapter->GetFrameFence();
+
+		FRHIPoolAllocationData& PoolAllocationData = InBaseShaderResource->ResourceLocation.GetPoolAllocatorPrivateData().PoolData;
+		FD3D12TransientMemoryPool* MemoryPool = ((FD3D12TransientMemoryPool*)Pools[PoolAllocationData.GetPoolIndex()]);
+		MemoryPool->ReleaseResource(InBaseShaderResource->ResourceLocation.GetResource(), PoolAllocationData, FrameFence.GetCurrentFence());
+	}
+}
+
+
+void FD3D12TransientResourceAllocator::Freeze()
+{
+	// all memory should have been freed again - so release back to the manager
+	// can already be reused by the next transient allocator
+	// Resources are kept alive until the allocator is destroyed
+	FD3D12TransientMemoryPoolManager& PoolManager = GetParentDevice()->GetTransientMemoryPoolManager();
+	for (FRHIMemoryPool* RHIPool : Pools)
+	{
+		FD3D12TransientMemoryPool* MemoryPool = ((FD3D12TransientMemoryPool*)RHIPool);
+		check(MemoryPool->GetAllocatedBlocks() == 0);
+		PoolManager.ReleaseMemoryPool(MemoryPool);
+	}
+	Pools.Empty();
+}
+
+
+FRHIMemoryPool* FD3D12TransientResourceAllocator::CreateNewPool(int16 InPoolIndex)
+{
+	// Get pool from manager - don't reallocate each time
+	return GetParentDevice()->GetTransientMemoryPoolManager().GetOrCreateMemoryPool(InPoolIndex);
+}
+
+
+FD3D12Resource* FD3D12TransientResourceAllocator::CreatePlacedResource(const FRHIPoolAllocationData& InAllocationData, const D3D12_RESOURCE_DESC& InDesc, D3D12_RESOURCE_STATES InCreateState, ED3D12ResourceStateMode InResourceStateMode, const D3D12_CLEAR_VALUE* InClearValue, const TCHAR* InName)
+{
+	FD3D12Resource* Resource = nullptr;
+
+	// Try and find a cached resource at given offset and creation flags
+	{
+		FD3D12TransientMemoryPool* MemoryPool = ((FD3D12TransientMemoryPool*)Pools[InAllocationData.GetPoolIndex()]);
+		Resource = MemoryPool->FindResourceInCache(InAllocationData.GetOffset(), InDesc, InClearValue, InName);
+	}
+
+	// No cached resource then use base class
+	if (Resource == nullptr)
+	{
+		TRACE_CPUPROFILER_EVENT_SCOPE(D3D12RHI::CreatePlacedResource);
+		Resource = FD3D12PoolAllocator::CreatePlacedResource(InAllocationData, InDesc, InCreateState, InResourceStateMode, InClearValue, InName);
+	}
+
+	return Resource;
+}
+
+
+void FD3D12TransientResourceAllocator::DeallocateResource(FD3D12ResourceLocation& ResourceLocation)
+{
+	check(IsOwner(ResourceLocation));
+
+	// Don't touch the allocation data - it's probably already freed via call to DeallocateMemory
+	// On clear the data on the resource location itself
+	ResourceLocation.ClearAllocator();
+}
+
+
+TArrayView<FD3D12Resource*> FD3D12TransientResourceAllocator::GetOverlappingResources(FD3D12BaseShaderResource* InBaseShaderResource)
+{
+	check(IsOwner(InBaseShaderResource->ResourceLocation));
+
+	TArray<FD3D12Resource*>* FindResult = OverlappingResourceData.Find(InBaseShaderResource);
+	if (FindResult)
+	{
+		return TArrayView<FD3D12Resource*>(*FindResult);
+	}
+	else
+	{
+		return TArrayView<FD3D12Resource*>();
+	}
+}
\ No newline at end of file
diff --git a/Engine/Source/Runtime/D3D12RHI/Private/D3D12TransientResourceAllocator.h b/Engine/Source/Runtime/D3D12RHI/Private/D3D12TransientResourceAllocator.h
new file mode 100644
--- /dev/null
+++ b/Engine/Source/Runtime/D3D12RHI/Private/D3D12TransientResourceAllocator.h
@@ -0,0 +1,226 @@
+// Copyright Epic Games, Inc. All Rights Reserved.
+#pragma once
+
+#include "D3D12PoolAllocator.h"
+
+// Forward declare
+class FD3D12TransientResourceAllocator;
+
+/**
+@brief Tracking stats for transient resource allocations
+**/
+struct FD3D12TransientMemoryStats
+{
+	void Reset()
+	{
+		CurrentPoolAllocated = 0;
+		TotalRequested = 0;
+		MaxFrameAllocated = 0;
+		CommittedAllocated = 0;
+		PoolAllocations = 0;
+		CommittedAllocations = 0;
+	}
+	void Add(const FD3D12TransientMemoryStats& InOther)
+	{
+		CurrentPoolAllocated += InOther.CurrentPoolAllocated;
+		TotalRequested += InOther.TotalRequested;
+		MaxFrameAllocated += InOther.MaxFrameAllocated;
+		CommittedAllocated += InOther.CommittedAllocated;
+		PoolAllocations += InOther.PoolAllocations;
+		CommittedAllocations += InOther.CommittedAllocations;
+	}
+
+	uint64 CurrentPoolAllocated = 0;
+	uint64 TotalRequested = 0;
+	uint64 MaxFrameAllocated = 0;
+	uint64 CommittedAllocated = 0;
+	uint32 PoolAllocations = 0;
+	uint32 CommittedAllocations = 0;
+};
+
+
+struct FD3D12PooledTextureData
+{
+	FRHITextureCreateInfo CreateInfo;
+	D3D12_RESOURCE_DESC ResourceDesc;
+	D3D12_CLEAR_VALUE ClearValue;
+	FTextureRHIRef RHITexture;
+};
+
+struct FD3D12PooledBufferData
+{
+	FRHIBufferCreateInfo CreateInfo;
+	FBufferRHIRef RHIBuffer;
+};
+
+
+/**
+@brief D3D12 transient pool specific implementation of the FD3D12MemoryPool
+	   Caches previously created resource to remove placed resource creation overhead and 
+	   also keeps track of all open allocation ranges used to create the aliasing barriers
+**/
+class FD3D12TransientMemoryPool : public FD3D12MemoryPool
+{
+public:
+
+	FD3D12TransientMemoryPool(FD3D12Device* ParentDevice, FRHIGPUMask VisibleNodes, const FD3D12ResourceInitConfig& InInitConfig, const FString& Name,
+		EResourceAllocationStrategy InAllocationStrategy, int16 InPoolIndex, uint64 InPoolSize, uint32 InPoolAlignment) :
+		FD3D12MemoryPool(ParentDevice, VisibleNodes, InInitConfig, Name, InAllocationStrategy, InPoolIndex, InPoolSize, InPoolAlignment, FRHIMemoryPool::EFreeListOrder::SortByOffset) { }
+	virtual ~FD3D12TransientMemoryPool();
+
+	void ResetPool();
+	void SetPoolIndex(int16 InNewPoolIndex);
+
+	void ClearActiveResources();
+	void CheckActiveResources(const FInt64Range& InAllocationRange, TArray<FD3D12Resource*>& OverlappingResources);
+
+	FD3D12Resource* FindResourceInCache(uint64 InAllocationOffset, const D3D12_RESOURCE_DESC& InDesc, const D3D12_CLEAR_VALUE* InClearValue, const TCHAR* InName);
+	void ReleaseResource(FD3D12Resource* InResource, FRHIPoolAllocationData& InReleasedAllocationData, uint64 InFenceValue);
+		
+private:
+
+	// Track all active ranges on an allocation which require possible aliasing barriers
+	struct ActiveResourceData
+	{
+		FInt64Range AllocationRange;
+		FInt64Range ActiveRange;
+		FD3D12Resource* Resource;
+	};
+	TArray<ActiveResourceData> ActiveResources;
+
+	// Cache all previously created placed resources at given offset and creation params to reduce
+	// placed resource creation overhead
+	struct FResourceCreateState
+	{
+		// Key values used to create the resource
+		uint64 AllocationOffset;
+		D3D12_RESOURCE_DESC ResourceDesc;
+		D3D12_CLEAR_VALUE ClearValue;
+
+		uint64 GetHash() const;
+	};
+	TMap<uint64, FD3D12Resource*> CachedResourceMap;
+};
+
+
+//-----------------------------------------------------------------------------
+//	Transient Memory Pool Manager
+//-----------------------------------------------------------------------------
+
+/**
+@brief Manages an array of pools per device which can be reused by the temporary
+	   transient memory allocator - the pools are returned to the manager when the 
+	   allocator is destroyed, so the heap creation and placed resource caches are kept
+	   alive in between frames. Pools are deleted when not used for n amount of frames.
+**/
+class FD3D12TransientMemoryPoolManager : public FD3D12DeviceChild, public FD3D12MultiNodeGPUObject
+{
+public:
+	FD3D12TransientMemoryPoolManager(FD3D12Device* InParent, FRHIGPUMask VisibleNodes);
+	~FD3D12TransientMemoryPoolManager() { }
+
+	void Destroy();
+
+	void BeginFrame();
+	void EndFrame();
+
+	// Get or creation functions from cached resources
+	FD3D12TransientMemoryPool* GetOrCreateMemoryPool(int16 InPoolIndex);
+	void ReleaseMemoryPool(FD3D12TransientMemoryPool* InMemoryPool);
+	FD3D12PooledTextureData GetPooledTexture(const FRHITextureCreateInfo& InCreateInfo, const TCHAR* InDebugName);
+	FD3D12PooledBufferData GetPooledBuffer(const FRHIBufferCreateInfo& InCreateInfo, const TCHAR* InDebugName);
+	void ReleaseResources(FD3D12TransientResourceAllocator* InAllocator);
+
+	// Get the pool creation members
+	const FD3D12ResourceInitConfig& GetInitConfig() const { return InitConfig; }
+	uint64 GetPoolSize() const { return PoolSize; }
+	uint32 GetPoolAlignment() const { return PoolAlignment; }
+	uint64 GetMaxAllocationSize() const { return MaxAllocationSize; }
+
+	// Stat management
+	void UpdateTextureStats(const FD3D12TransientMemoryStats& InStats)
+	{
+		TextureMemoryStats.Add(InStats);
+	}
+	void UpdateBufferStats(const FD3D12TransientMemoryStats& InStats)
+	{
+		BufferMemoryStats.Add(InStats);
+	}
+	void UpdateMemoryStats();
+	
+private:
+
+	// Pool creation members
+	FD3D12ResourceInitConfig InitConfig;
+	uint64 PoolSize;
+	uint32 PoolAlignment;
+	uint64 MaxAllocationSize;
+
+	// Critical section to lock access to the pools & resource caches
+	FCriticalSection CS;
+	TArray<FD3D12TransientMemoryPool*> Pools;
+	TMap<uint64, TArray<FD3D12PooledTextureData>> FreeTextures;
+	TMap<uint64, TArray<FD3D12PooledBufferData>> FreeBuffers;
+
+	// Stats of last frame
+	FD3D12TransientMemoryStats TextureMemoryStats;
+	FD3D12TransientMemoryStats BufferMemoryStats;
+};
+
+
+//-----------------------------------------------------------------------------
+//	Transient Resource Allocator
+//-----------------------------------------------------------------------------
+
+/**
+@brief D3D12 Specific implementation of the IRHITransientResourceAllocator
+	   Uses a pool allocator as base to manage the internal D3D12 memory pools and also takes care of the actual
+	   D3D12 placed resource creation.
+	   Pools are retrieved from and released to the FD3D12TransientMemoryPoolManager
+**/
+class FD3D12TransientResourceAllocator : public IRHITransientResourceAllocator, public FD3D12PoolAllocator
+{
+public:
+
+	FD3D12TransientResourceAllocator(FD3D12TransientMemoryPoolManager& InMemoryPoolManager);
+	virtual ~FD3D12TransientResourceAllocator();
+
+	// Implementation of FRHITransientResourceAllocator interface
+	virtual FRHITexture* CreateTexture(const FRHITextureCreateInfo& InCreateInfo, const TCHAR* InDebugName) override;
+	virtual FRHIBuffer* CreateBuffer(const FRHIBufferCreateInfo& InCreateInfo, const TCHAR* InDebugName) override;
+	virtual void DeallocateMemory(FRHITexture* InTexture) override;
+	virtual void DeallocateMemory(FRHIBuffer* InBuffer) override;
+	virtual void Freeze() override;
+
+	// Override deallocation function - memory possibly already freed during DeallocateMemory
+	virtual void DeallocateResource(FD3D12ResourceLocation& ResourceLocation) override;
+
+	// Try and find overlapping resource data for given base resource
+	TArrayView<FD3D12Resource*> GetOverlappingResources(FD3D12BaseShaderResource* InBaseShaderResource);
+
+private:
+
+	// Override placed resource allocation helper function to get from cache if available
+	virtual FD3D12Resource* CreatePlacedResource(const FRHIPoolAllocationData& InAllocationData, const D3D12_RESOURCE_DESC& InDesc, D3D12_RESOURCE_STATES InCreateState, ED3D12ResourceStateMode InResourceStateMode, const D3D12_CLEAR_VALUE* InClearValue, const TCHAR* InName) override;
+	virtual FRHIMemoryPool* CreateNewPool(int16 InPoolIndex) override;
+
+	// Shared allocation/deallocation helper
+	FD3D12BaseShaderResource* GetBaseShaderResource(FRHITexture* InRHITexture);
+	void SetupAllocatedResource(FD3D12BaseShaderResource* InBaseShaderResource, FD3D12TransientMemoryStats& Stats);
+	void DeallocateMemory(FD3D12BaseShaderResource* InBaseShaderResource, FD3D12TransientMemoryStats& Stats);
+
+	// Reference all the allocated resources - owner of all the resources and will be returned back to the manager when the allocator is destroyed
+	TArray<FD3D12PooledTextureData> AllocatedTextures;
+	TArray<FD3D12PooledBufferData> AllocatedBuffers;
+
+	// All overlapping resource allocation information from all allocated resources on this allocator
+	// Not stored inside each resource to not increase the base object size for each resource
+	// (can still be moved if it the lookup becomes a perf bottleneck at some point)
+	TMap<FD3D12BaseShaderResource*, TArray<FD3D12Resource*>> OverlappingResourceData;
+
+	// Stat tracking
+	FD3D12TransientMemoryStats TextureMemoryStats;
+	FD3D12TransientMemoryStats BufferMemoryStats;
+
+	friend class FD3D12TransientMemoryPoolManager;
+};
\ No newline at end of file
diff --git a/Engine/Source/Runtime/D3D12RHI/Private/D3D12Util.cpp b/Engine/Source/Runtime/D3D12RHI/Private/D3D12Util.cpp
--- a/Engine/Source/Runtime/D3D12RHI/Private/D3D12Util.cpp
+++ b/Engine/Source/Runtime/D3D12RHI/Private/D3D12Util.cpp
@@ -1456,6 +1456,24 @@ DEFINE_STAT(STAT_D3D12TextureAllocatorAllocated);
 DEFINE_STAT(STAT_D3D12TextureAllocatorUnused);
 DEFINE_STAT(STAT_D3D12TextureAllocatorCount);
 
+DEFINE_STAT(STAT_D3D12TransientMemoryPoolAllocated);
+DEFINE_STAT(STAT_D3D12TransientMemoryFramePoolUsed);
+DEFINE_STAT(STAT_D3D12TransientMemoryFrameCommittedUsed);
+DEFINE_STAT(STAT_D3D12TransientMemoryPoolAllocations);
+DEFINE_STAT(STAT_D3D12TransientMemoryCommittedAllocations);
+
+DEFINE_STAT(STAT_D3D12TransientMemoryBufferFrameUsed);
+DEFINE_STAT(STAT_D3D12TransientMemoryBufferFrameTotalRequested);
+DEFINE_STAT(STAT_D3D12TransientMemoryBufferFrameCommittedAllocated);
+DEFINE_STAT(STAT_D3D12TransientMemoryBufferPoolAllocations);
+DEFINE_STAT(STAT_D3D12TransientMemoryBufferCommittedAllocations);
+
+DEFINE_STAT(STAT_D3D12TransientMemoryTextureFrameUsed);
+DEFINE_STAT(STAT_D3D12TransientMemoryTextureFrameTotalRequested);
+DEFINE_STAT(STAT_D3D12TransientMemoryTextureFrameCommittedAllocated);
+DEFINE_STAT(STAT_D3D12TransientMemoryTexturePoolAllocations);
+DEFINE_STAT(STAT_D3D12TransientMemoryTextureCommittedAllocations);
+
 DEFINE_STAT(STAT_D3D12BufferPoolMemoryAllocated);
 DEFINE_STAT(STAT_D3D12BufferPoolMemoryUsed);
 DEFINE_STAT(STAT_D3D12BufferPoolMemoryFree);
diff --git a/Engine/Source/Runtime/D3D12RHI/Private/D3D12View.h b/Engine/Source/Runtime/D3D12RHI/Private/D3D12View.h
--- a/Engine/Source/Runtime/D3D12RHI/Private/D3D12View.h
+++ b/Engine/Source/Runtime/D3D12RHI/Private/D3D12View.h
@@ -719,20 +719,27 @@ private:
 		BaseShaderResource = InBaseShaderResource;
 		ResourceLocation = &InResourceLocation;
 		Resource = ResourceLocation->GetResource();
-		check(Resource);
-
-		ResidencyHandle = Resource->GetResidencyHandle();
-
 		Desc = InDesc;
 
-		ViewSubresourceSubset = CViewSubresourceSubset(Desc,
-			Resource->GetMipLevels(),
-			Resource->GetArraySize(),
-			Resource->GetDesc().Format,
-			Flags);
+		// Transient resources might not have an actual resource yet
+		if (Resource)
+		{
+			ResidencyHandle = Resource->GetResidencyHandle();
+			ViewSubresourceSubset = CViewSubresourceSubset(Desc,
+				Resource->GetMipLevels(),
+				Resource->GetArraySize(),
+				Resource->GetDesc().Format,
+				Flags);
+		}
+		else
+		{
+			ResidencyHandle = nullptr;
+			ViewSubresourceSubset = CViewSubresourceSubset();
+		}
 
 #if DO_CHECK || USING_CODE_ANALYSIS
-		bInitialized = true;
+		// Only mark initialize if an actual resource is created for the base shader resource
+		bInitialized = (Resource != nullptr);
 #endif
 	}
 
@@ -741,17 +748,23 @@ protected:
 	{
 		InitializeInternal(InDesc, InBaseShaderResource, InResourceLocation);
 
-		ID3D12Resource* D3DResource = ResourceLocation->GetResource()->GetResource();
-		Descriptor.CreateView(Desc, D3DResource);
+		if (ResourceLocation->GetResource())
+		{
+			ID3D12Resource* D3DResource = ResourceLocation->GetResource()->GetResource();
+			Descriptor.CreateView(Desc, D3DResource);
+		}
 	}
 
 	void CreateViewWithCounter(const TDesc& InDesc, FD3D12BaseShaderResource* InBaseShaderResource, FD3D12ResourceLocation& InResourceLocation, FD3D12Resource* InCounterResource)
 	{
 		InitializeInternal(InDesc, InBaseShaderResource, InResourceLocation);
 
-		ID3D12Resource* D3DResource = ResourceLocation->GetResource()->GetResource();
-		ID3D12Resource* D3DCounterResource = InCounterResource ? InCounterResource->GetResource() : nullptr;
-		Descriptor.CreateViewWithCounter(Desc, D3DResource, D3DCounterResource);
+		if (ResourceLocation->GetResource())
+		{
+			ID3D12Resource* D3DResource = ResourceLocation->GetResource()->GetResource();
+			ID3D12Resource* D3DCounterResource = InCounterResource ? InCounterResource->GetResource() : nullptr;
+			Descriptor.CreateViewWithCounter(Desc, D3DResource, D3DCounterResource);
+		}
 	}
 
 	virtual void ResourceRenamed(FD3D12BaseShaderResource* InRenamedResource, FD3D12ResourceLocation* InNewResourceLocation) override
@@ -774,11 +787,11 @@ protected:
 public:
 	inline FD3D12Device*					GetParentDevice()			const { return Descriptor.GetParentDevice(); }
 	inline FD3D12Device*					GetParentDevice_Unsafe()	const { return Descriptor.GetParentDevice_Unsafe(); }
+	inline FD3D12ResourceLocation*			GetResourceLocation()		const { return ResourceLocation; }
 	inline const TDesc&						GetDesc()					const { checkf(bInitialized, TEXT("Uninitialized D3D12View size %d"), (uint32)sizeof(TDesc)); return Desc; }
 	inline CD3DX12_CPU_DESCRIPTOR_HANDLE	GetView()					const { checkf(bInitialized, TEXT("Uninitialized D3D12View size %d"), (uint32)sizeof(TDesc)); return Descriptor.GetHandle(); }
 	inline uint32							GetDescriptorHeapIndex()	const { checkf(bInitialized, TEXT("Uninitialized D3D12View size %d"), (uint32)sizeof(TDesc)); return Descriptor.GetIndex(); }
 	inline FD3D12Resource*					GetResource()				const { checkf(bInitialized, TEXT("Uninitialized D3D12View size %d"), (uint32)sizeof(TDesc)); return Resource; }
-	inline FD3D12ResourceLocation*			GetResourceLocation()		const { checkf(bInitialized, TEXT("Uninitialized D3D12View size %d"), (uint32)sizeof(TDesc)); return ResourceLocation; }
 	inline FD3D12ResidencyHandle*			GetResidencyHandle()		const { checkf(bInitialized, TEXT("Uninitialized D3D12View size %d"), (uint32)sizeof(TDesc)); return ResidencyHandle; }
 	inline const CViewSubresourceSubset&	GetViewSubresourceSubset()	const { checkf(bInitialized, TEXT("Uninitialized D3D12View size %d"), (uint32)sizeof(TDesc)); return ViewSubresourceSubset; }
 
@@ -826,32 +839,34 @@ public:
 
 	~FD3D12ShaderResourceView()
 	{
-		// Non-resident mesh LODs can lead to null SRVs being created
-		check(bInitialized || !GetParentDevice_Unsafe());
 	}
 
 	void Initialize(D3D12_SHADER_RESOURCE_VIEW_DESC& InDesc, FD3D12BaseShaderResource* InBaseShaderResource, FD3D12ResourceLocation& InResourceLocation, uint32 InStride, uint32 InStartOffsetBytes, bool InSkipFastClearFinalize)
 	{
 		Stride = InStride;
 		StartOffsetBytes = InStartOffsetBytes;
-		bContainsDepthPlane   = InResourceLocation.GetResource()->IsDepthStencilResource() && GetPlaneSliceFromViewFormat(InResourceLocation.GetResource()->GetDesc().Format, InDesc.Format) == 0;
-		bContainsStencilPlane = InResourceLocation.GetResource()->IsDepthStencilResource() && GetPlaneSliceFromViewFormat(InResourceLocation.GetResource()->GetDesc().Format, InDesc.Format) == 1;
-		bRequiresResourceStateTracking = InResourceLocation.GetResource()->RequiresResourceStateTracking();
 		bSkipFastClearFinalize = InSkipFastClearFinalize;
 
-#if DO_CHECK
-		// Check the plane slice of the SRV matches the texture format
-		// Texture2DMS does not have explicit plane index (it's implied by the format)
-		if (InDesc.ViewDimension == D3D12_SRV_DIMENSION_TEXTURE2D)
+		if (InResourceLocation.GetResource())
 		{
-			check(GetPlaneSliceFromViewFormat(InResourceLocation.GetResource()->GetDesc().Format, InDesc.Format) == InDesc.Texture2D.PlaneSlice);
-		}
+			bContainsDepthPlane = InResourceLocation.GetResource()->IsDepthStencilResource() && GetPlaneSliceFromViewFormat(InResourceLocation.GetResource()->GetDesc().Format, InDesc.Format) == 0;
+			bContainsStencilPlane = InResourceLocation.GetResource()->IsDepthStencilResource() && GetPlaneSliceFromViewFormat(InResourceLocation.GetResource()->GetDesc().Format, InDesc.Format) == 1;
+			bRequiresResourceStateTracking = InResourceLocation.GetResource()->RequiresResourceStateTracking();
 
-		if (InDesc.ViewDimension == D3D12_SRV_DIMENSION_BUFFER)
-		{
-			//check(InResourceLocation.GetOffsetFromBaseOfResource() == InDesc.Buffer.FirstElement * Stride);
-		}
+#if DO_CHECK
+			// Check the plane slice of the SRV matches the texture format
+			// Texture2DMS does not have explicit plane index (it's implied by the format)
+			if (InDesc.ViewDimension == D3D12_SRV_DIMENSION_TEXTURE2D)
+			{
+				check(GetPlaneSliceFromViewFormat(InResourceLocation.GetResource()->GetDesc().Format, InDesc.Format) == InDesc.Texture2D.PlaneSlice);
+			}
+
+			if (InDesc.ViewDimension == D3D12_SRV_DIMENSION_BUFFER)
+			{
+				//check(InResourceLocation.GetOffsetFromBaseOfResource() == InDesc.Buffer.FirstElement * Stride);
+			}
 #endif
+		}
 
 		CreateView(InDesc, InBaseShaderResource, InResourceLocation);
 	}
@@ -1075,27 +1090,7 @@ public:
 		, bHasStencil(InHasStencil)		// Only some DSVs have stencil bits in their format
 	{
 		CreateView(InDSVDesc, InBaseShaderResource, InBaseShaderResource->ResourceLocation);
-
-		check(Resource);
-
-		// Create individual subresource subsets for each plane
-		if (bHasDepth)
-		{
-			DepthOnlyViewSubresourceSubset = CViewSubresourceSubset(InDSVDesc,
-				Resource->GetMipLevels(),
-				Resource->GetArraySize(),
-				Resource->GetDesc().Format,
-				ViewSubresourceSubsetFlags_DepthOnlyDsv);
-		}
-
-		if (bHasStencil)
-		{
-			StencilOnlyViewSubresourceSubset = CViewSubresourceSubset(InDSVDesc,
-				Resource->GetMipLevels(),
-				Resource->GetArraySize(),
-				Resource->GetDesc().Format,
-				ViewSubresourceSubsetFlags_StencilOnlyDsv);
-		}
+		SetupDepthStencilViewSubresourceSubset();
 	}
 
 	bool HasDepth() const
@@ -1108,6 +1103,31 @@ public:
 		return bHasStencil;
 	}
 
+	void SetupDepthStencilViewSubresourceSubset()
+	{
+		if (Resource)
+		{
+			// Create individual subresource subsets for each plane
+			if (bHasDepth)
+			{
+				DepthOnlyViewSubresourceSubset = CViewSubresourceSubset(Desc,
+					Resource->GetMipLevels(),
+					Resource->GetArraySize(),
+					Resource->GetDesc().Format,
+					ViewSubresourceSubsetFlags_DepthOnlyDsv);
+			}
+
+			if (bHasStencil)
+			{
+				StencilOnlyViewSubresourceSubset = CViewSubresourceSubset(Desc,
+					Resource->GetMipLevels(),
+					Resource->GetArraySize(),
+					Resource->GetDesc().Format,
+					ViewSubresourceSubsetFlags_StencilOnlyDsv);
+			}
+		}
+	}
+
 	CViewSubresourceSubset& GetDepthOnlyViewSubresourceSubset()
 	{
 		check(bHasDepth);
@@ -1124,6 +1144,7 @@ public:
 	{
 		check(ResourceLocation->GetOffsetFromBaseOfResource() == 0);
 		CreateView(Desc, BaseShaderResource, BaseShaderResource->ResourceLocation);
+		SetupDepthStencilViewSubresourceSubset();
 	}
 };
 
diff --git a/Engine/Source/Runtime/D3D12RHI/Public/D3D12Resources.h b/Engine/Source/Runtime/D3D12RHI/Public/D3D12Resources.h
--- a/Engine/Source/Runtime/D3D12RHI/Public/D3D12Resources.h
+++ b/Engine/Source/Runtime/D3D12RHI/Public/D3D12Resources.h
@@ -117,6 +117,7 @@ private:
 
 	FD3D12ResidencyHandle ResidencyHandle;
 
+	D3D12_CLEAR_VALUE ClearValue;
 	D3D12_RESOURCE_DESC Desc;
 	uint8 PlaneCount;
 	uint16 SubresourceCount;
@@ -153,6 +154,7 @@ public:
 		ID3D12Resource* InResource,
 		D3D12_RESOURCE_STATES InInitialResourceState,
 		D3D12_RESOURCE_DESC const& InDesc,
+		const D3D12_CLEAR_VALUE* InClearValue = nullptr,
 		FD3D12Heap* InHeap = nullptr,
 		D3D12_HEAP_TYPE InHeapType = D3D12_HEAP_TYPE_DEFAULT);
 
@@ -163,6 +165,7 @@ public:
 		ED3D12ResourceStateMode InResourceStateMode,
 		D3D12_RESOURCE_STATES InDefaultResourceState,
 		D3D12_RESOURCE_DESC const& InDesc,
+		const D3D12_CLEAR_VALUE* InClearValue,
 		FD3D12Heap* InHeap,
 		D3D12_HEAP_TYPE InHeapType);
 
@@ -191,6 +194,7 @@ public:
 
 	ID3D12Pageable* GetPageable();
 	D3D12_RESOURCE_DESC const& GetDesc() const { return Desc; }
+	D3D12_CLEAR_VALUE const& GetClearValue() const { return ClearValue; }
 	D3D12_HEAP_TYPE GetHeapType() const { return HeapType; }
 	D3D12_GPU_VIRTUAL_ADDRESS GetGPUVirtualAddress() const { return GPUVirtualAddress; }
 	void* GetResourceBaseAddress() const { check(ResourceBaseAddress); return ResourceBaseAddress; }
@@ -266,7 +270,7 @@ public:
 		{
 			// Ignore the requested resource state for non tracked resource because RHI will assume it's always in default resource 
 			// state then when a transition is required (will transition via scoped push/pop to requested state)
-			if (!bSRVOnly && InResourceState != ERHIAccess::Unknown)
+			if (!bSRVOnly && InResourceState != ERHIAccess::Unknown && InResourceState != ERHIAccess::Discard)
 			{
 				bool bAsyncCompute = false;
 				return GetD3D12ResourceState(InResourceState, bAsyncCompute);
@@ -670,6 +674,18 @@ private:
 	EAllocatorType AllocatorType;
 };
 
+// Generic interface for every type D3D12 specific allocator
+struct ID3D12ResourceAllocator
+{
+	// Helper function for textures to compute the correct size and alignment
+	void AllocateTexture(D3D12_HEAP_TYPE InHeapType, const D3D12_RESOURCE_DESC& InDesc, EPixelFormat InUEFormat, ED3D12ResourceStateMode InResourceStateMode,
+		D3D12_RESOURCE_STATES InCreateState, const D3D12_CLEAR_VALUE* InClearValue, const TCHAR* InName, FD3D12ResourceLocation& ResourceLocation);
+
+	// Actual pure virtual resource allocation function
+	virtual void AllocateResource(D3D12_HEAP_TYPE InHeapType, const D3D12_RESOURCE_DESC& InDesc, uint64 InSize, uint32 InAllocationAlignment, ED3D12ResourceStateMode InResourceStateMode,
+		D3D12_RESOURCE_STATES InCreateState, const D3D12_CLEAR_VALUE* InClearValue, const TCHAR* InName, FD3D12ResourceLocation& ResourceLocation) = 0;
+};
+
 class FD3D12DeferredDeletionQueue : public FD3D12AdapterChild
 {
 public:
@@ -767,8 +783,27 @@ protected:
 };
 
 
+#if PLATFORM_WINDOWS || PLATFORM_HOLOLENS
+class FD3D12TransientResource
+{
+	// Nothing special for fast ram
+public:
+	void Swap(FD3D12TransientResource&) {}
+};
+class FD3D12FastClearResource
+{
+public:
+	inline void GetWriteMaskProperties(void*& OutData, uint32& OutSize)
+	{
+		OutData = nullptr;
+		OutSize = 0;
+	}
+};
+#endif
+
+
 /** The base class of resources that may be bound as shader resources (texture or buffer). */
-class FD3D12BaseShaderResource : public FD3D12DeviceChild, public IRefCountedObject
+class FD3D12BaseShaderResource : public FD3D12DeviceChild, public FD3D12TransientResource, public IRefCountedObject
 {
 protected:
 	FCriticalSection RenameListenersCS;
@@ -884,25 +919,7 @@ public:
 	virtual ~FD3D12UniformBuffer();
 };
 
-#if PLATFORM_WINDOWS || PLATFORM_HOLOLENS
-class FD3D12TransientResource
-{
-	// Nothing special for fast ram
-public:
-	void Swap(FD3D12TransientResource&) {}
-};
-class FD3D12FastClearResource
-{
-public:
-	inline void GetWriteMaskProperties(void*& OutData, uint32& OutSize)
-	{
-		OutData = nullptr;
-		OutSize = 0;
-	}
-};
-#endif
-
-class FD3D12Buffer : public FRHIBuffer, public FD3D12BaseShaderResource, public FD3D12TransientResource, public FD3D12LinkedAdapterObject<FD3D12Buffer>
+class FD3D12Buffer : public FRHIBuffer, public FD3D12BaseShaderResource, public FD3D12LinkedAdapterObject<FD3D12Buffer>
 {
 public:
 	FD3D12Buffer()
@@ -945,6 +962,7 @@ public:
 		return FRHIResource::GetRefCount();
 	}
 
+	static void GetResourceDescAndAlignment(uint64 InSize, uint32 InStride, EBufferUsageFlags& InUsage, D3D12_RESOURCE_DESC& ResourceDesc, uint32& Alignment);
 
 	FD3D12LockedResource LockedData;
 };
@@ -1015,14 +1033,14 @@ public:
 		return 1;
 	}
 
-	void AddAliasingBarrier(ID3D12Resource* pResource)
+	void AddAliasingBarrier(ID3D12Resource* InResourceBefore, ID3D12Resource* InResourceAfter)
 	{
 		Barriers.AddUninitialized();
 		D3D12_RESOURCE_BARRIER& Barrier = Barriers.Last();
 		Barrier.Type = D3D12_RESOURCE_BARRIER_TYPE_ALIASING;
 		Barrier.Flags = D3D12_RESOURCE_BARRIER_FLAG_NONE;
-		Barrier.Aliasing.pResourceBefore = NULL;
-		Barrier.Aliasing.pResourceAfter = pResource;
+		Barrier.Aliasing.pResourceBefore = InResourceBefore;
+		Barrier.Aliasing.pResourceAfter = InResourceAfter;
 	}
 
 	// Flush the batch to the specified command list then reset.
diff --git a/Engine/Source/Runtime/D3D12RHI/Public/D3D12Util.h b/Engine/Source/Runtime/D3D12RHI/Public/D3D12Util.h
--- a/Engine/Source/Runtime/D3D12RHI/Public/D3D12Util.h
+++ b/Engine/Source/Runtime/D3D12RHI/Public/D3D12Util.h
@@ -618,7 +618,7 @@ static uint32 GetHeightAlignment(DXGI_FORMAT Format)
 	}
 }
 
-static void Get4KTileShape(D3D12_TILE_SHAPE* pTileShape, DXGI_FORMAT Format, uint8 UEFormat, D3D12_RESOURCE_DIMENSION Dimension, uint32 SampleCount)
+static void Get4KTileShape(D3D12_TILE_SHAPE* pTileShape, DXGI_FORMAT DXGIFormat, EPixelFormat UEFormat, D3D12_RESOURCE_DIMENSION Dimension, uint32 SampleCount)
 {
 	//Bits per unit
 	uint32 BPU = GPixelFormats[UEFormat].BlockBytes * 8;
@@ -628,7 +628,7 @@ static void Get4KTileShape(D3D12_TILE_SHAPE* pTileShape, DXGI_FORMAT Format, uin
 	case D3D12_RESOURCE_DIMENSION_BUFFER:
 	case D3D12_RESOURCE_DIMENSION_TEXTURE1D:
 	{
-		check(!IsBlockCompressFormat(Format));
+		check(!IsBlockCompressFormat(DXGIFormat));
 		pTileShape->WidthInTexels = (BPU == 0) ? 4096 : 4096 * 8 / BPU;
 		pTileShape->HeightInTexels = 1;
 		pTileShape->DepthInTexels = 1;
@@ -637,19 +637,19 @@ static void Get4KTileShape(D3D12_TILE_SHAPE* pTileShape, DXGI_FORMAT Format, uin
 	case D3D12_RESOURCE_DIMENSION_TEXTURE2D:
 	{
 		pTileShape->DepthInTexels = 1;
-		if (IsBlockCompressFormat(Format))
+		if (IsBlockCompressFormat(DXGIFormat))
 		{
 			// Currently only supported block sizes are 64 and 128.
 			// These equations calculate the size in texels for a tile. It relies on the fact that 16*16*16 blocks fit in a tile if the block size is 128 bits.
 			check(BPU == 64 || BPU == 128);
-			pTileShape->WidthInTexels = 16 * GetWidthAlignment(Format);
-			pTileShape->HeightInTexels = 16 * GetHeightAlignment(Format);
+			pTileShape->WidthInTexels = 16 * GetWidthAlignment(DXGIFormat);
+			pTileShape->HeightInTexels = 16 * GetHeightAlignment(DXGIFormat);
 			if (BPU == 64)
 			{
 				// If bits per block are 64 we double width so it takes up the full tile size.
 				// This is only true for BC1 and BC4
-				check((Format >= DXGI_FORMAT_BC1_TYPELESS && Format <= DXGI_FORMAT_BC1_UNORM_SRGB) ||
-					(Format >= DXGI_FORMAT_BC4_TYPELESS && Format <= DXGI_FORMAT_BC4_SNORM));
+				check((DXGIFormat >= DXGI_FORMAT_BC1_TYPELESS && DXGIFormat <= DXGI_FORMAT_BC1_UNORM_SRGB) ||
+					(DXGIFormat >= DXGI_FORMAT_BC4_TYPELESS && DXGIFormat <= DXGI_FORMAT_BC4_SNORM));
 				pTileShape->WidthInTexels *= 2;
 			}
 		}
@@ -713,28 +713,28 @@ static void Get4KTileShape(D3D12_TILE_SHAPE* pTileShape, DXGI_FORMAT Format, uin
 				check(false);
 			}
 
-			check(GetWidthAlignment(Format) == 1);
-			check(GetHeightAlignment(Format) == 1);
+			check(GetWidthAlignment(DXGIFormat) == 1);
+			check(GetHeightAlignment(DXGIFormat) == 1);
 		}
 
 		break;
 	}
 	case D3D12_RESOURCE_DIMENSION_TEXTURE3D:
 	{
-		if (IsBlockCompressFormat(Format))
+		if (IsBlockCompressFormat(DXGIFormat))
 		{
 			// Currently only supported block sizes are 64 and 128.
 			// These equations calculate the size in texels for a tile. It relies on the fact that 16*16*16 blocks fit in a tile if the block size is 128 bits.
 			check(BPU == 64 || BPU == 128);
-			pTileShape->WidthInTexels = 8 * GetWidthAlignment(Format);
-			pTileShape->HeightInTexels = 8 * GetHeightAlignment(Format);
+			pTileShape->WidthInTexels = 8 * GetWidthAlignment(DXGIFormat);
+			pTileShape->HeightInTexels = 8 * GetHeightAlignment(DXGIFormat);
 			pTileShape->DepthInTexels = 4;
 			if (BPU == 64)
 			{
 				// If bits per block are 64 we double width so it takes up the full tile size.
 				// This is only true for BC1 and BC4
-				check((Format >= DXGI_FORMAT_BC1_TYPELESS && Format <= DXGI_FORMAT_BC1_UNORM_SRGB) ||
-					(Format >= DXGI_FORMAT_BC4_TYPELESS && Format <= DXGI_FORMAT_BC4_SNORM));
+				check((DXGIFormat >= DXGI_FORMAT_BC1_TYPELESS && DXGIFormat <= DXGI_FORMAT_BC1_UNORM_SRGB) ||
+					(DXGIFormat >= DXGI_FORMAT_BC4_TYPELESS && DXGIFormat <= DXGI_FORMAT_BC4_SNORM));
 				pTileShape->DepthInTexels *= 2;
 			}
 		}
@@ -775,8 +775,8 @@ static void Get4KTileShape(D3D12_TILE_SHAPE* pTileShape, DXGI_FORMAT Format, uin
 				check(false);
 			}
 
-			check(GetWidthAlignment(Format) == 1);
-			check(GetHeightAlignment(Format) == 1);
+			check(GetWidthAlignment(DXGIFormat) == 1);
+			check(GetHeightAlignment(DXGIFormat) == 1);
 		}
 	}
 	break;
@@ -785,7 +785,7 @@ static void Get4KTileShape(D3D12_TILE_SHAPE* pTileShape, DXGI_FORMAT Format, uin
 
 #define NUM_4K_BLOCKS_PER_64K_PAGE (16)
 
-static bool TextureCanBe4KAligned(D3D12_RESOURCE_DESC& Desc, uint8 UEFormat)
+static bool TextureCanBe4KAligned(D3D12_RESOURCE_DESC& Desc, EPixelFormat UEFormat)
 {
 	// 4KB alignment is only available for read only textures
 	if (!(Desc.Flags & D3D12_RESOURCE_FLAG_ALLOW_RENDER_TARGET ||
diff --git a/Engine/Source/Runtime/RHI/Private/RHI.cpp b/Engine/Source/Runtime/RHI/Private/RHI.cpp
--- a/Engine/Source/Runtime/RHI/Private/RHI.cpp
+++ b/Engine/Source/Runtime/RHI/Private/RHI.cpp
@@ -133,6 +133,7 @@ FString GetRHIAccessName(ERHIAccess Access)
 			case ERHIAccess::CopyDest:            return TEXT("CopyDest");
 			case ERHIAccess::ResolveDst:          return TEXT("ResolveDst");
 			case ERHIAccess::DSVWrite:            return TEXT("DSVWrite");
+			case ERHIAccess::Discard:             return TEXT("Discard");
 			}
 		});
 }
@@ -146,6 +147,8 @@ FString GetResourceTransitionFlagsName(EResourceTransitionFlags Flags)
 		default: checkNoEntry(); // fall through
 		case EResourceTransitionFlags::None:                return TEXT("None");
 		case EResourceTransitionFlags::MaintainCompression: return TEXT("MaintainCompression");
+		case EResourceTransitionFlags::Discard:				return TEXT("Discard");
+		case EResourceTransitionFlags::Clear:				return TEXT("Clear");
 		}
 	});
 }
@@ -1534,6 +1537,38 @@ void FRHIRenderPassInfo::Validate() const
 }
 #endif
 
+bool FRHITextureCreateInfo::IsValid() const
+{
+	if (Extent.X <= 0 || Extent.Y <= 0 || Depth == 0 || ArraySize == 0 || NumMips == 0 || NumSamples < 1 || NumSamples > 8)
+	{
+		return false;
+	}
+
+	if (NumSamples > 1 && !(Dimension == ETextureDimension::Texture2D || Dimension == ETextureDimension::Texture2DArray))
+	{
+		return false;
+	}
+
+	if (Dimension == ETextureDimension::Texture3D)
+	{
+		if (ArraySize > 1)
+		{
+			return false;
+		}
+	}
+	else if (Depth > 1)
+	{
+		return false;
+	}
+
+	if (Format == PF_Unknown)
+	{
+		return false;
+	}
+
+	return true;
+}
+
 static FRHIPanicEvent RHIPanicEvent;
 FRHIPanicEvent& RHIGetPanicDelegate()
 {
diff --git a/Engine/Source/Runtime/RHI/Private/RHIValidation.cpp b/Engine/Source/Runtime/RHI/Private/RHIValidation.cpp
--- a/Engine/Source/Runtime/RHI/Private/RHIValidation.cpp
+++ b/Engine/Source/Runtime/RHI/Private/RHIValidation.cpp
@@ -6,6 +6,7 @@
 
 #include "RHIValidation.h"
 #include "RHIValidationContext.h"
+#include "RHIValidationTransientResourceAllocator.h"
 #include "HAL/PlatformStackWalk.h"
 #include "Misc/OutputDeviceRedirector.h"
 
@@ -89,6 +90,13 @@ FValidationRHI::~FValidationRHI()
 	GRHIValidationEnabled = false;
 }
 
+IRHITransientResourceAllocator* FValidationRHI::RHICreateTransientResourceAllocator()
+{
+	// Wrap around validation allocator
+	IRHITransientResourceAllocator* RHIAllocator = RHI->RHICreateTransientResourceAllocator();
+	return new FValidationTransientResourceAllocator(RHIAllocator);
+}
+
 IRHICommandContext* FValidationRHI::RHIGetDefaultContext()
 {
 	IRHICommandContext* LowLevelContext = RHI->RHIGetDefaultContext();
@@ -1128,4 +1136,166 @@ namespace RHIValidation
 
 TLockFreePointerListUnordered<FValidationContext, PLATFORM_CACHE_LINE_SIZE> FValidationRHICommandContextContainer::ParallelCommandContexts;
 
+
+//-----------------------------------------------------------------------------
+//	Validation Transient Resource Allocator
+//-----------------------------------------------------------------------------
+
+#define TRANSIENT_RESOURCE_LOG_PREFIX_REASON(ReasonString) TEXT("RHI validation failed: " ReasonString "\n\n")\
+	TEXT("--------------------------------------------------------------------\n")\
+	TEXT("         RHI Transient Resource Allocation Validation Error		  \n")\
+	TEXT("--------------------------------------------------------------------\n")\
+	TEXT("\n")
+
+#define TRANSIENT_RESOURCE_LOG_SUFFIX TEXT("\n")\
+	TEXT("--------------------------------------------------------------------\n")\
+	TEXT("\n")
+
+FValidationTransientResourceAllocator::~FValidationTransientResourceAllocator()
+{
+	check(bFrozen);
+
+	// Validation of leaking textures resources disabled by default because we know textures are referenced by global uniforms
+	// This needs to be fixed first
+	static bool bValidateLeakingTextures = false;
+	static bool bValidateLeakingBuffers = true;
+	if (bValidateLeakingTextures || bValidateLeakingBuffers)
+	{
+		// Check all allocated resource data and make sure the transient allocator holds the last 'reference'
+		int32 LeakingResourceCount = 0;
+		for (auto It = AllocatedResourceMap.CreateIterator(); It; ++It)
+		{
+			const FRHIResource* RHIResource = It->Key;
+			const AllocatedResourceData& ResourceData = It->Value;
+			if (ResourceData.ResourceType == AllocatedResourceData::EType::Texture && bValidateLeakingTextures)
+			{
+				LeakingResourceCount += (RHIResource->GetRefCount() > 1) ? 1 : 0;
+			}
+			else if (ResourceData.ResourceType == AllocatedResourceData::EType::Buffer && bValidateLeakingBuffers)
+			{
+				LeakingResourceCount += (RHIResource->GetRefCount() > 1) ? 1 : 0;
+			}
+		}
+
+		if (LeakingResourceCount > 0)
+		{
+			FString ErrorMessage = FString::Printf(
+				TRANSIENT_RESOURCE_LOG_PREFIX_REASON("Leaking transient resources")
+				TEXT("%d Transient Resource allocations are still used when the alloctor is destroyed. The TransientResourceAllocator manages the lifetime of all allocated resource.\n\n")
+				TEXT("Leaking resources:\n"),
+				LeakingResourceCount);
+			for (auto It = AllocatedResourceMap.CreateIterator(); It; ++It)
+			{
+				const FRHIResource* RHIResource = It->Key;
+				const AllocatedResourceData& ResourceData = It->Value;
+				if (RHIResource->GetRefCount() > 1)
+				{
+					ErrorMessage += FString::Printf(TEXT("         %s (%s) - Refcount: %d\n"), *ResourceData.DebugName, ResourceData.ResourceType == AllocatedResourceData::EType::Texture ? TEXT("Texture") : TEXT("DataBuffer"), RHIResource->GetRefCount());
+				}
+			}
+			ErrorMessage += FString::Printf(TRANSIENT_RESOURCE_LOG_SUFFIX);
+			FValidationRHI::ReportValidationFailure(*ErrorMessage);
+		}
+	}
+
+	// Delete actual allocator
+	delete RHIAllocator;
+}
+
+FRHITexture* FValidationTransientResourceAllocator::CreateTexture(const FRHITextureCreateInfo& InCreateInfo, const TCHAR* InDebugName)
+{
+	check(!bFrozen);
+
+	FRHITexture* RHITexture = RHIAllocator->CreateTexture(InCreateInfo, InDebugName);
+
+	// Setup for resource tracking
+	RHITexture->InitBarrierTracking(InCreateInfo.NumMips, InCreateInfo.ArraySize, InCreateInfo.Format, InCreateInfo.Flags, ERHIAccess::Discard, InDebugName);
+
+	// Store allocation data
+	AllocatedResourceData ResourceData;
+	ResourceData.DebugName = InDebugName;
+	ResourceData.ResourceType = AllocatedResourceData::EType::Texture;
+	ResourceData.bMemoryAllocated = true;
+	AllocatedResourceMap.Add(RHITexture, ResourceData);
+
+	return RHITexture;
+}
+
+FRHIBuffer* FValidationTransientResourceAllocator::CreateBuffer(const FRHIBufferCreateInfo& InCreateInfo, const TCHAR* InDebugName)
+{
+	check(!bFrozen);
+
+	FRHIBuffer* RHIBuffer = RHIAllocator->CreateBuffer(InCreateInfo, InDebugName);
+
+	// Setup for resource tracking
+	RHIBuffer->InitBarrierTracking(ERHIAccess::Discard, InDebugName);
+	
+	// Store allocation data
+	AllocatedResourceData ResourceData;
+	ResourceData.DebugName = InDebugName;
+	ResourceData.ResourceType = AllocatedResourceData::EType::Buffer;
+	ResourceData.bMemoryAllocated = true;
+	AllocatedResourceMap.Add(RHIBuffer, ResourceData);
+
+	return RHIBuffer;
+}
+
+void FValidationTransientResourceAllocator::DeallocateMemory(FRHITexture* InTexture)
+{
+	check(!bFrozen);
+
+	RHIAllocator->DeallocateMemory(InTexture);
+
+	// Mark memory as freed
+	AllocatedResourceData* ResourceData = AllocatedResourceMap.Find(InTexture);
+	check(ResourceData);
+	ResourceData->bMemoryAllocated = false;
+}
+
+void FValidationTransientResourceAllocator::DeallocateMemory(FRHIBuffer* InBuffer)
+{
+	check(!bFrozen);
+
+	RHIAllocator->DeallocateMemory(InBuffer);
+
+	// Mark memory as freed
+	AllocatedResourceData* ResourceData = AllocatedResourceMap.Find(InBuffer);
+	check(ResourceData);
+	ResourceData->bMemoryAllocated = false;
+}
+
+void FValidationTransientResourceAllocator::Freeze()
+{
+	check(!bFrozen);
+
+	// Check all allocated resource data and make sure all memory is freed again
+	int32 AllocatedResourceCount = 0;
+	for (auto It = AllocatedResourceMap.CreateIterator(); It; ++It)
+	{
+		const AllocatedResourceData& ResourceData = It->Value;
+		AllocatedResourceCount += ResourceData.bMemoryAllocated ? 1 : 0;
+	}
+
+	if (AllocatedResourceCount > 0)
+	{
+		FString ErrorMessage = FString::Printf(
+			TRANSIENT_RESOURCE_LOG_PREFIX_REASON("Open transient allocations")
+			TEXT("%d Transient Resource allocations still have memory allocated. Call 'DeallocateMemory' on all transient allocated resources before calling Freeze.\n\n")
+			TEXT("Resources with Allocated Memory:\n"), 
+			AllocatedResourceCount);
+		for (auto It = AllocatedResourceMap.CreateIterator(); It; ++It)
+		{
+			const AllocatedResourceData& ResourceData = It->Value;
+			if (ResourceData.bMemoryAllocated)
+			{
+				ErrorMessage += FString::Printf(TEXT("         %s (%s)\n"), *ResourceData.DebugName, ResourceData.ResourceType == AllocatedResourceData::EType::Texture ? TEXT("Texture") : TEXT("DataBuffer"));
+			}
+		}
+		ErrorMessage += FString::Printf(TRANSIENT_RESOURCE_LOG_SUFFIX);
+		FValidationRHI::ReportValidationFailure(*ErrorMessage);
+	}
+
+	bFrozen = true;
+}
+
 #endif	// ENABLE_RHI_VALIDATION
diff --git a/Engine/Source/Runtime/RHI/Public/DynamicRHI.h b/Engine/Source/Runtime/RHI/Public/DynamicRHI.h
--- a/Engine/Source/Runtime/RHI/Public/DynamicRHI.h
+++ b/Engine/Source/Runtime/RHI/Public/DynamicRHI.h
@@ -435,6 +435,10 @@ public:
 	{
 	}
 
+	/**
+	* Create a new transient resource allocator
+	*/
+	virtual IRHITransientResourceAllocator* RHICreateTransientResourceAllocator() { return nullptr; }
 
 	/**
 	* Creates a staging buffer, which is memory visible to the cpu without any locking.
@@ -1642,6 +1646,11 @@ FORCEINLINE void RHIReleaseTransition(FRHITransition* Transition)
 	GDynamicRHI->RHIReleaseTransition(Transition);
 }
 
+FORCEINLINE IRHITransientResourceAllocator* RHICreateTransientResourceAllocator()
+{
+	return GDynamicRHI->RHICreateTransientResourceAllocator();
+}
+
 
 #if RHI_RAYTRACING
 
diff --git a/Engine/Source/Runtime/RHI/Public/RHI.h b/Engine/Source/Runtime/RHI/Public/RHI.h
--- a/Engine/Source/Runtime/RHI/Public/RHI.h
+++ b/Engine/Source/Runtime/RHI/Public/RHI.h
@@ -842,7 +842,10 @@ enum class ERHIAccess
 	DSVWrite            	= 1 << 14,
 	RTAccelerationStructure = 1 << 15,
 
-	Last = RTAccelerationStructure,
+	// Invalid released state (transient resources)
+	Discard					= 1 << 16,
+
+	Last = Discard,
 	None = Unknown,
 	Mask = (Last << 1) - 1,
 
@@ -1803,12 +1806,24 @@ enum class EResourceTransitionFlags
 	None                = 0,
 
 	MaintainCompression = 1 << 0, // Specifies that the transition should not decompress the resource, allowing us to read a compressed resource directly in its compressed state.
+	Discard				= 1 << 1, // Specifies that the data in the resource should be discarded during the transition - used for transient resource acquire when the resource will be fully overwritten
+	Clear				= 1 << 2, // Specifies that the data in the resource should be cleared during the transition - used for transient resource acquire when the resource might not be fully overwritten
 
-	Last = MaintainCompression,
+	Last = Clear,
 	Mask = (Last << 1) - 1
 };
 ENUM_CLASS_FLAGS(EResourceTransitionFlags);
 
+enum class ERHICreateTransientAllocationFlags
+{
+	// Discard data on acquire
+	Discard,
+
+	// Clear data on acquire
+	Clear,
+};
+ENUM_CLASS_FLAGS(ERHICreateTransientAllocationFlags);
+
 RHI_API FString GetRHIAccessName(ERHIAccess Access);
 RHI_API FString GetResourceTransitionFlagsName(EResourceTransitionFlags Flags);
 RHI_API FString GetRHIPipelineName(ERHIPipeline Pipeline);
diff --git a/Engine/Source/Runtime/RHI/Public/RHICommandList.h b/Engine/Source/Runtime/RHI/Public/RHICommandList.h
--- a/Engine/Source/Runtime/RHI/Public/RHICommandList.h
+++ b/Engine/Source/Runtime/RHI/Public/RHICommandList.h
@@ -1500,6 +1500,18 @@ FRHICOMMAND_MACRO(FRHICommandResourceTransition)
 	RHI_API void Execute(FRHICommandListBase& CmdList);
 };
 
+FRHICOMMAND_MACRO(FRHICommandReleaseTransientResourceAllocator)
+{
+	IRHITransientResourceAllocator* Allocator;
+
+	FRHICommandReleaseTransientResourceAllocator(IRHITransientResourceAllocator * InAllocator)
+		: Allocator(InAllocator)
+	{
+	}
+
+	RHI_API void Execute(FRHICommandListBase & CmdList);
+};
+
 struct FRHICommandSetAsyncComputeBudgetString
 {
 	static const TCHAR* TStr() { return TEXT("FRHICommandSetAsyncComputeBudget"); }
@@ -2422,6 +2434,16 @@ public:
 		ALLOC_COMMAND(FRHICommandClearUAVUint)(UnorderedAccessViewRHI, Values);
 	}
 
+	FORCEINLINE_DEBUGGABLE void ReleaseTransientResourceAllocator(IRHITransientResourceAllocator* InAllocator)
+	{
+		if (Bypass())
+		{
+			GetComputeContext().RHIReleaseTransientResourceAllocator(InAllocator);
+			return;
+		}
+		ALLOC_COMMAND(FRHICommandReleaseTransientResourceAllocator)(InAllocator);
+	}
+
 	FORCEINLINE_DEBUGGABLE void BeginTransitions(TArrayView<const FRHITransition*> Transitions)
 	{
 		if (Bypass())
diff --git a/Engine/Source/Runtime/RHI/Public/RHICommandListCommandExecutes.inl b/Engine/Source/Runtime/RHI/Public/RHICommandListCommandExecutes.inl
--- a/Engine/Source/Runtime/RHI/Public/RHICommandListCommandExecutes.inl
+++ b/Engine/Source/Runtime/RHI/Public/RHICommandListCommandExecutes.inl
@@ -355,6 +355,12 @@ void FRHICommandClearUAVUint::Execute(FRHICommandListBase& CmdList)
 	INTERNAL_DECORATOR_COMPUTE(RHIClearUAVUint)(UnorderedAccessViewRHI, Values);
 }
 
+void FRHICommandReleaseTransientResourceAllocator::Execute(FRHICommandListBase& CmdList)
+{
+	RHISTAT(ReleaseTransientResourceAllocator);
+	INTERNAL_DECORATOR_COMPUTE(RHIReleaseTransientResourceAllocator)(Allocator);
+}
+
 void FRHICommandCopyToResolveTarget::Execute(FRHICommandListBase& CmdList)
 {
 	RHISTAT(CopyToResolveTarget);
diff --git a/Engine/Source/Runtime/RHI/Public/RHIContext.h b/Engine/Source/Runtime/RHI/Public/RHIContext.h
--- a/Engine/Source/Runtime/RHI/Public/RHIContext.h
+++ b/Engine/Source/Runtime/RHI/Public/RHIContext.h
@@ -167,6 +167,11 @@ public:
 	*/
 	virtual void RHIClearUAVUint(FRHIUnorderedAccessView* UnorderedAccessViewRHI, const FUintVector4& Values) = 0;
 
+	/**
+	* Release the transient allocator and all internal resources when the command list is done executing
+	*/
+	virtual void RHIReleaseTransientResourceAllocator(IRHITransientResourceAllocator* InAllocator) {}
+
 	virtual void RHIBeginUAVOverlap() {}
 	virtual void RHIEndUAVOverlap() {}
 
diff --git a/Engine/Source/Runtime/RHI/Public/RHIResources.h b/Engine/Source/Runtime/RHI/Public/RHIResources.h
--- a/Engine/Source/Runtime/RHI/Public/RHIResources.h
+++ b/Engine/Source/Runtime/RHI/Public/RHIResources.h
@@ -2835,3 +2835,223 @@ struct FRHIRenderPassInfo
 private:
 	RHI_API void OnVerifyNumUAVsFailed(int32 InNumUAVs);
 };
+
+
+/** Descriptor used to create a texture resource */
+struct RHI_API FRHITextureCreateInfo
+{
+	static FRHITextureCreateInfo Create2D(
+		FIntPoint InExtent,
+		EPixelFormat InFormat,
+		FClearValueBinding InClearValue,
+		ETextureCreateFlags InFlags,
+		uint8 InNumMips = 1,
+		uint8 InNumSamples = 1)
+	{
+		return FRHITextureCreateInfo(ETextureDimension::Texture2D, InFlags, InFormat, InExtent, InClearValue, 1, 1, InNumMips, InNumSamples);
+	}
+
+	static FRHITextureCreateInfo Create2DArray(
+		FIntPoint InExtent,
+		EPixelFormat InFormat,
+		FClearValueBinding InClearValue,
+		ETextureCreateFlags InFlags,
+		uint32 InArraySize,
+		uint8 InNumMips = 1,
+		uint8 InNumSamples = 1)
+	{
+		return FRHITextureCreateInfo(ETextureDimension::Texture2DArray, InFlags, InFormat, InExtent, InClearValue, 1, InArraySize, InNumMips, InNumSamples);
+	}
+
+	static FRHITextureCreateInfo Create3D(
+		FIntVector InSize,
+		EPixelFormat InFormat,
+		FClearValueBinding InClearValue,
+		ETextureCreateFlags InFlags,
+		uint8 InNumMips = 1,
+		uint8 InNumSamples = 1)
+	{
+		return FRHITextureCreateInfo(ETextureDimension::Texture3D, InFlags, InFormat, FIntPoint(InSize.X, InSize.Y), InClearValue, InSize.Z, 1, InNumMips, InNumSamples);
+	}
+
+	static FRHITextureCreateInfo CreateCube(
+		uint32 InSizeInPixels,
+		EPixelFormat InFormat,
+		FClearValueBinding InClearValue,
+		ETextureCreateFlags InFlags,
+		uint8 InNumMips = 1,
+		uint8 InNumSamples = 1)
+	{
+		return FRHITextureCreateInfo(ETextureDimension::TextureCube, InFlags, InFormat, FIntPoint(InSizeInPixels, InSizeInPixels), InClearValue, 1, 1, InNumMips, InNumSamples);
+	}
+
+	static FRHITextureCreateInfo CreateCubeArray(
+		uint32 InSizeInPixels,
+		EPixelFormat InFormat,
+		FClearValueBinding InClearValue,
+		ETextureCreateFlags InFlags,
+		uint32 InArraySize,
+		uint8 InNumMips = 1,
+		uint8 InNumSamples = 1)
+	{
+		return FRHITextureCreateInfo(ETextureDimension::TextureCubeArray, InFlags, InFormat, FIntPoint(InSizeInPixels, InSizeInPixels), InClearValue, 1, InArraySize, InNumMips, InNumSamples);
+	}
+
+	FRHITextureCreateInfo() = default;
+	FRHITextureCreateInfo(
+		ETextureDimension InDimension,
+		ETextureCreateFlags InFlags,
+		EPixelFormat InFormat,
+		FIntPoint InExtent,
+		FClearValueBinding InClearValue,
+		uint16 InDepth = 1,
+		uint16 InArraySize = 1,
+		uint8 InNumMips = 1,
+		uint8 InNumSamples = 1)
+		: Dimension(InDimension)
+		, Flags(InFlags)
+		, Format(InFormat)
+		, Extent(InExtent)
+		, Depth(InDepth)
+		, ArraySize(InArraySize)
+		, NumMips(InNumMips)
+		, NumSamples(InNumSamples)
+		, ClearValue(InClearValue)
+	{}
+
+	bool operator == (const FRHITextureCreateInfo& Other) const
+	{
+		return Dimension == Other.Dimension
+			&& Flags == Other.Flags
+			&& Format == Other.Format
+			&& Extent == Other.Extent
+			&& Depth == Other.Depth
+			&& ArraySize == Other.ArraySize
+			&& NumMips == Other.NumMips
+			&& NumSamples == Other.NumSamples
+			&& ClearValue == Other.ClearValue;
+	}
+
+	bool operator != (const FRHITextureCreateInfo& Other) const
+	{
+		return !(*this == Other);
+	}
+
+	bool IsTexture2D() const
+	{
+		return Dimension == ETextureDimension::Texture2D || Dimension == ETextureDimension::Texture2DArray;
+	}
+
+	bool IsTexture3D() const
+	{
+		return Dimension == ETextureDimension::Texture3D;
+	}
+
+	bool IsTextureCube() const
+	{
+		return Dimension == ETextureDimension::TextureCube || Dimension == ETextureDimension::TextureCubeArray;
+	}
+
+	bool IsTextureArray() const
+	{
+		return Dimension == ETextureDimension::Texture2DArray || Dimension == ETextureDimension::TextureCubeArray;
+	}
+
+	bool IsMipChain() const
+	{
+		return NumMips > 1;
+	}
+
+	bool IsMultisample() const
+	{
+		return NumSamples > 1;
+	}
+
+	FIntVector GetSize() const
+	{
+		return FIntVector(Extent.X, Extent.Y, Depth);
+	}
+
+	/** Returns whether this descriptor conforms to requirements. */
+	bool IsValid() const;
+
+	/** Texture dimension to use when creating the RHI texture. */
+	ETextureDimension Dimension = ETextureDimension::Texture2D;
+
+	/** Texture flags passed on to RHI texture. */
+	ETextureCreateFlags Flags = TexCreate_None;
+
+	/** Pixel format used to create RHI texture. */
+	EPixelFormat Format = PF_Unknown;
+
+	/** Extent of the texture in x and y. */
+	FIntPoint Extent = FIntPoint(1, 1);
+
+	/** Depth of the texture if the dimension is 3D. */
+	uint16 Depth = 1;
+
+	/** The number of array elements in the texture. (Keep at 1 if dimension is 3D). */
+	uint16 ArraySize = 1;
+
+	/** Number of mips in the texture mip-map chain. */
+	uint8 NumMips = 1;
+
+	/** Number of samples in the texture. >1 for MSAA. */
+	uint8 NumSamples = 1;
+
+	/** Clear value to use when fast-clearing the texture. */
+	FClearValueBinding ClearValue;
+};
+
+
+/** Descriptor used to create a buffer resource */
+struct FRHIBufferCreateInfo
+{
+	bool operator == (const FRHIBufferCreateInfo& Other) const
+	{
+		return (
+			Size == Other.Size &&
+			Stride == Other.Stride &&
+			Usage == Other.Usage);
+	}
+
+	bool operator != (const FRHIBufferCreateInfo& Other) const
+	{
+		return !(*this == Other);
+	}
+
+	/** Total size of the buffer. */
+	uint32 Size = 1;
+
+	/** Stride in bytes */
+	uint32 Stride = 1;
+
+	/** Bitfields describing the uses of that buffer. */
+	EBufferUsageFlags Usage = BUF_None;
+};
+
+
+/**
+@brief Resource allocator which is managed for a certain part of the frame by the high level
+	   rendering code (usually RenderGraph). An allocator is created via RHICreateTransientResourceAllocator on the device
+	   and released via the command list function RHIReleaseTransientResourceAllocator when all the
+	   transient resources are no longer needed.
+	   The allocator manages the lifetime of all the created resources and will take care of destroying
+	   all the resources when the allocator itself is destroyed.	   
+**/
+struct RHI_API IRHITransientResourceAllocator
+{
+	// Virtual destructor
+	virtual ~IRHITransientResourceAllocator() {}
+
+	// Create functions
+	virtual FRHITexture* CreateTexture(const FRHITextureCreateInfo& InCreateInfo, const TCHAR* InDebugName) = 0;
+	virtual FRHIBuffer* CreateBuffer(const FRHIBufferCreateInfo& InCreateInfo, const TCHAR* InDebugName) = 0;
+
+	// Memory deallocation functions - doesn't destroy the resource, but frees internal memory for reuse by another transient resource
+	virtual void DeallocateMemory(FRHITexture* InTexture) = 0;
+	virtual void DeallocateMemory(FRHIBuffer* InBuffer) = 0;
+
+	// Freeze allocator for now creation and validates that all resources have their memory deallocated
+	virtual void Freeze() = 0;
+};
\ No newline at end of file
diff --git a/Engine/Source/Runtime/RHI/Public/RHIValidation.h b/Engine/Source/Runtime/RHI/Public/RHIValidation.h
--- a/Engine/Source/Runtime/RHI/Public/RHIValidation.h
+++ b/Engine/Source/Runtime/RHI/Public/RHIValidation.h
@@ -210,6 +210,8 @@ public:
 		RHI->RHIReleaseTransition(Transition);
 	}
 
+	virtual IRHITransientResourceAllocator* RHICreateTransientResourceAllocator() override final;
+
 	/**
 	* Creates a staging buffer, which is memory visible to the cpu without any locking.
 	* @return The new staging-buffer.
diff --git a/Engine/Source/Runtime/RHI/Public/RHIValidationCommon.h b/Engine/Source/Runtime/RHI/Public/RHIValidationCommon.h
--- a/Engine/Source/Runtime/RHI/Public/RHIValidationCommon.h
+++ b/Engine/Source/Runtime/RHI/Public/RHIValidationCommon.h
@@ -263,7 +263,7 @@ namespace RHIValidation
 
 		ELoggingMode LoggingMode = ELoggingMode::None;
 
-		void SetDebugName(const TCHAR* Name, const TCHAR* Suffix = nullptr);
+		RHI_API void SetDebugName(const TCHAR* Name, const TCHAR* Suffix = nullptr);
 		inline const TCHAR* GetDebugName() const { return DebugName.Len() ? *DebugName : nullptr; }
 
 		inline bool IsBarrierTrackingInitialized() const { return NumMips > 0 && NumArraySlices > 0; }
diff --git a/Engine/Source/Runtime/RHI/Public/RHIValidationContext.h b/Engine/Source/Runtime/RHI/Public/RHIValidationContext.h
--- a/Engine/Source/Runtime/RHI/Public/RHIValidationContext.h
+++ b/Engine/Source/Runtime/RHI/Public/RHIValidationContext.h
@@ -842,6 +842,11 @@ public:
 		RHIContext->RHIWriteGPUFence(FenceRHI);
 	}
 
+	virtual void RHIReleaseTransientResourceAllocator(IRHITransientResourceAllocator* InAllocator) override final
+	{
+		RHIContext->RHIReleaseTransientResourceAllocator(InAllocator);
+	}
+
 	virtual void RHICopyToStagingBuffer(FRHIBuffer* SourceBufferRHI, FRHIStagingBuffer* DestinationStagingBufferRHI, uint32 InOffset, uint32 InNumBytes) override final
 	{
 		Tracker->Assert(SourceBufferRHI->GetWholeResourceIdentity(), ERHIAccess::CopySrc); // @todo: is CopySrc correct?
diff --git a/Engine/Source/Runtime/RHI/Public/RHIValidationTransientResourceAllocator.h b/Engine/Source/Runtime/RHI/Public/RHIValidationTransientResourceAllocator.h
new file mode 100644
--- /dev/null
+++ b/Engine/Source/Runtime/RHI/Public/RHIValidationTransientResourceAllocator.h
@@ -0,0 +1,50 @@
+// Copyright Epic Games, Inc. All Rights Reserved.
+
+/*=============================================================================
+	RHIValidation.h: Public Valdation RHI definitions.
+=============================================================================*/
+
+#pragma once 
+
+#include "RHIResources.h"
+
+#if ENABLE_RHI_VALIDATION
+
+class FValidationTransientResourceAllocator : public IRHITransientResourceAllocator
+{
+public:
+	FValidationTransientResourceAllocator(IRHITransientResourceAllocator* InRHIAllocator) : RHIAllocator(InRHIAllocator) {}
+	virtual ~FValidationTransientResourceAllocator();
+
+	// Implementation of FRHITransientResourceAllocator interface
+	virtual FRHITexture* CreateTexture(const FRHITextureCreateInfo& InCreateInfo, const TCHAR* InDebugName) override final;
+	virtual FRHIBuffer* CreateBuffer(const FRHIBufferCreateInfo& InCreateInfo, const TCHAR* InDebugName) override final;
+	virtual void DeallocateMemory(FRHITexture* InTexture) override final;
+	virtual void DeallocateMemory(FRHIBuffer* InBuffer) override final;
+	virtual void Freeze() override final;
+
+private:
+
+	// Actual RHI transient allocator which will get all functions forwarded
+	IRHITransientResourceAllocator* RHIAllocator = nullptr;
+
+	// Allocator already frozen?
+	bool bFrozen = false;
+
+	// All the allocated resources on the transient allocator
+	struct AllocatedResourceData
+	{
+		enum class EType
+		{
+			Texture,
+			Buffer,
+		};
+
+		FString DebugName;
+		EType ResourceType;
+		bool bMemoryAllocated;
+	};
+	TMap<FRHIResource*, AllocatedResourceData> AllocatedResourceMap;
+};
+
+#endif	// ENABLE_RHI_VALIDATION
diff --git a/Engine/Source/Runtime/RHICore/Private/RHIPoolAllocator.cpp b/Engine/Source/Runtime/RHICore/Private/RHIPoolAllocator.cpp
--- a/Engine/Source/Runtime/RHICore/Private/RHIPoolAllocator.cpp
+++ b/Engine/Source/Runtime/RHICore/Private/RHIPoolAllocator.cpp
@@ -146,7 +146,7 @@ void FRHIPoolAllocationData::AddAfter(FRHIPoolAllocationData* InOther)
 
 
 /**
- * Sort predicate that sorts size of the allocation data
+ * Sort predicate that sorts on size of the allocation data
  */
 static bool LinkedListSortBySizePredicate(const FRHIPoolAllocationData* InLeft, const FRHIPoolAllocationData* InRight)
 {
@@ -154,6 +154,15 @@ static bool LinkedListSortBySizePredicate(const FRHIPoolAllocationData* InLeft,
 }
 
 
+/**
+ * Sort predicate that sorts on offset of the allocation data
+ */
+static bool LinkedListSortByOffsetPredicate(const FRHIPoolAllocationData* InLeft, const FRHIPoolAllocationData* InRight)
+{
+	return InLeft->GetOffset() < InRight->GetOffset();
+}
+
+
 //-----------------------------------------------------------------------------
 //	LinkedList Allocator
 //-----------------------------------------------------------------------------
@@ -188,10 +197,11 @@ uint32 FRHIMemoryPool::GetAlignedOffset(uint32 InOffset, uint32 InPoolAlignment,
 }
 
 
-FRHIMemoryPool::FRHIMemoryPool(int16 InPoolIndex, uint64 InPoolSize, uint32 InPoolAlignment) :
+FRHIMemoryPool::FRHIMemoryPool(int16 InPoolIndex, uint64 InPoolSize, uint32 InPoolAlignment, EFreeListOrder InFreeListOrder) :
 	PoolIndex(InPoolIndex), 
 	PoolSize(InPoolSize), 
-	PoolAlignment(InPoolAlignment), 
+	PoolAlignment(InPoolAlignment),
+	FreeListOrder(InFreeListOrder),
 	FreeSize(0), 
 	AligmnentWaste(0), 
 	AllocatedBlocks(0)
@@ -418,7 +428,20 @@ FRHIPoolAllocationData* FRHIMemoryPool::AddToFreeBlocks(FRHIPoolAllocationData*
 	}
 
 	// Sorted insert
-	int32 InsertIndex = Algo::LowerBound(FreeBlocks, FreeBlock, LinkedListSortBySizePredicate);
+	int32 InsertIndex = 0;
+	switch (FreeListOrder)
+	{
+		case EFreeListOrder::SortBySize:
+		{
+			InsertIndex = Algo::LowerBound(FreeBlocks, FreeBlock, LinkedListSortBySizePredicate);
+			break;
+		}
+		case EFreeListOrder::SortByOffset:
+		{
+			InsertIndex = Algo::LowerBound(FreeBlocks, FreeBlock, LinkedListSortByOffsetPredicate);
+			break;
+		}
+	}	
 	FreeBlocks.Insert(FreeBlock, InsertIndex);
 
 	Validate();
@@ -505,10 +528,11 @@ void FRHIMemoryPool::Validate()
 //	Pool Allocator
 //-----------------------------------------------------------------------------
 
-FRHIPoolAllocator::FRHIPoolAllocator(uint64 InPoolSize, uint32 InPoolAlignment, uint32 InMaxAllocationSize, bool InDefragEnabled) :
+FRHIPoolAllocator::FRHIPoolAllocator(uint64 InPoolSize, uint32 InPoolAlignment, uint32 InMaxAllocationSize, FRHIMemoryPool::EFreeListOrder InFreeListOrder, bool InDefragEnabled) :
 	PoolSize(InPoolSize),
 	PoolAlignment(InPoolAlignment),
 	MaxAllocationSize(InMaxAllocationSize),
+	FreeListOrder(InFreeListOrder),
 	bDefragEnabled(InDefragEnabled),
 	TotalAllocatedBlocks(0)
 {}
diff --git a/Engine/Source/Runtime/RHICore/Public/RHIPoolAllocator.h b/Engine/Source/Runtime/RHICore/Public/RHIPoolAllocator.h
--- a/Engine/Source/Runtime/RHICore/Public/RHIPoolAllocator.h
+++ b/Engine/Source/Runtime/RHICore/Public/RHIPoolAllocator.h
@@ -111,8 +111,14 @@ class RHICORE_API FRHIMemoryPool
 {
 public:
 
+	enum class EFreeListOrder
+	{
+		SortBySize,
+		SortByOffset,
+	};
+
 	// Constructor
-	FRHIMemoryPool(int16 InPoolIndex, uint64 InPoolSize, uint32 InPoolAlignment);
+	FRHIMemoryPool(int16 InPoolIndex, uint64 InPoolSize, uint32 InPoolAlignment, EFreeListOrder InFreeListOrder);
 	virtual ~FRHIMemoryPool();
 
 	// Setup/Shutdown
@@ -155,9 +161,10 @@ protected:
 	void Validate();
 
 	// Const creation members
-	const int16 PoolIndex;
+	int16 PoolIndex;
 	const uint64 PoolSize;
 	const uint32 PoolAlignment;
+	const EFreeListOrder FreeListOrder;
 
 	// Stats
 	uint64 FreeSize;
@@ -187,13 +194,13 @@ class RHICORE_API FRHIPoolAllocator
 public:
 
 	// Constructor
-	FRHIPoolAllocator(uint64 InPoolSize, uint32 InPoolAlignment, uint32 InMaxAllocationSize, bool InDefragEnabled);
+	FRHIPoolAllocator(uint64 InPoolSize, uint32 InPoolAlignment, uint32 InMaxAllocationSize, FRHIMemoryPool::EFreeListOrder InFreeListOrder, bool InDefragEnabled);
 	virtual ~FRHIPoolAllocator();
 
 	// Setup/Shutdown
 	void Initialize();
 	void Destroy();
-	
+
 	// Defrag & cleanup operation
 	void Defrag(uint32 InMaxCopySize, uint32& CurrentCopySize);
 
@@ -201,7 +208,7 @@ public:
 	void UpdateMemoryStats(uint32& IOMemoryAllocated, uint32& IOMemoryUsed, uint32& IOMemoryFree, uint32& IOAlignmentWaste, uint32& IOAllocatedPageCount, uint32& IOFullPageCount);
 
 protected:
-
+		
 	// Supported allocator operations
 	bool TryAllocateInternal(uint32 InSizeInBytes, uint32 InAllocationAlignment, FRHIPoolAllocationData& AllocationData);
 	void DeallocateInternal(FRHIPoolAllocationData& AllocationData);
@@ -217,6 +224,7 @@ protected:
 	const uint64 PoolSize;
 	const uint32 PoolAlignment;
 	const uint64 MaxAllocationSize;
+	const FRHIMemoryPool::EFreeListOrder FreeListOrder;
 	const bool bDefragEnabled;
 
 	// Critical section to lock access to the pools
