Fix Hololens build (UE-109945)

diff --git a/Engine/Source/Runtime/D3D12RHI/D3D12RHI.Build.cs b/Engine/Source/Runtime/D3D12RHI/D3D12RHI.Build.cs
--- a/Engine/Source/Runtime/D3D12RHI/D3D12RHI.Build.cs
+++ b/Engine/Source/Runtime/D3D12RHI/D3D12RHI.Build.cs
@@ -51,6 +51,10 @@ public class D3D12RHI : ModuleRules
             	AddEngineThirdPartyPrivateStaticDependencies(Target, "NVAftermath");
 				AddEngineThirdPartyPrivateStaticDependencies(Target, "IntelMetricsDiscovery");
             }
+            else
+            {
+				PrivateDefinitions.Add("D3D12RHI_USE_D3DDISASSEMBLE=0");
+			}
         }
     }
 }
