Added support for D3D12 experimental shader models, such as 6.6. Use CVar r.D3D12.ExperimentalShaderModels=1 to opt in (default 0).

diff --git a/Engine/Source/Runtime/D3D12RHI/Private/Windows/WindowsD3D12Device.cpp b/Engine/Source/Runtime/D3D12RHI/Private/Windows/WindowsD3D12Device.cpp
--- a/Engine/Source/Runtime/D3D12RHI/Private/Windows/WindowsD3D12Device.cpp
+++ b/Engine/Source/Runtime/D3D12RHI/Private/Windows/WindowsD3D12Device.cpp
@@ -126,6 +126,15 @@ static FAutoConsoleVariableRef CVarAllowShaderModel6(
 	ECVF_ReadOnly | ECVF_RenderThreadSafe
 );
 
+#if !UE_BUILD_SHIPPING
+static TAutoConsoleVariable<int32> CVarExperimentalShaderModels(
+	TEXT("r.D3D12.ExperimentalShaderModels"),
+	0,
+	TEXT("Controls whether D3D12 experimental shader models should be allowed. Not available in shipping builds. (default = 0)."),
+	ECVF_ReadOnly
+);
+#endif // !UE_BUILD_SHIPPING
+
 static inline int D3D12RHI_PreferAdapterVendor()
 {
 	if (FParse::Param(FCommandLine::Get(), TEXT("preferAMD")))
@@ -334,6 +343,18 @@ static bool SupportsHDROutput(FD3D12DynamicRHI* D3DRHI)
 	return bSupportsHDROutput;
 }
 
+static bool ShouldEnableExperimentalShaderModels()
+{
+#if !UE_BUILD_SHIPPING
+	if (CVarExperimentalShaderModels.GetValueOnAnyThread() == 1)
+	{
+		return true;
+	}
+#endif // !UE_BUILD_SHIPPING
+
+	return false;
+}
+
 bool FD3D12DynamicRHIModule::IsSupported()
 {
 #if !PLATFORM_HOLOLENS
@@ -343,6 +364,21 @@ bool FD3D12DynamicRHIModule::IsSupported()
 	}
 #endif
 
+	if (ShouldEnableExperimentalShaderModels())
+	{
+		// Experimental features must be enabled before doing anything else with D3D.
+
+		UUID ExperimentalFeatures[] =
+		{
+			D3D12ExperimentalShaderModels
+		};
+		HRESULT hr = D3D12EnableExperimentalFeatures(UE_ARRAY_COUNT(ExperimentalFeatures), ExperimentalFeatures, nullptr, nullptr);
+		if (SUCCEEDED(hr))
+		{
+			UE_LOG(LogD3D12RHI, Log, TEXT("D3D12 experimental shader models enabled"));
+		}
+	}
+
 	// If not computed yet
 	if (ChosenAdapters.Num() == 0)
 	{
@@ -451,10 +487,11 @@ void FD3D12DynamicRHIModule::FindAdapter()
 				uint32 OutputCount = CountAdapterOutputs(TempAdapter);
 
 				UE_LOG(LogD3D12RHI, Log,
-					TEXT("Found D3D12 adapter %u: %s (Max supported Feature Level %s)"),
+					TEXT("Found D3D12 adapter %u: %s (Max supported Feature Level %s, shader model %d.%d)"),
 					AdapterIndex,
 					AdapterDesc.Description,
-					GetFeatureLevelString(MaxSupportedFeatureLevel)
+					GetFeatureLevelString(MaxSupportedFeatureLevel),
+					(MaxSupportedShaderModel>>4), (MaxSupportedShaderModel & 0xF)
 					);
 				UE_LOG(LogD3D12RHI, Log,
 					TEXT("Adapter has %uMB of dedicated video memory, %uMB of dedicated system memory, and %uMB of shared system memory, %d output[s]"),
