Log native D3D12 SM6.6 64-bit atomic and mesh shader support

diff --git a/Engine/Source/Runtime/D3D12RHI/Private/D3D12Adapter.cpp b/Engine/Source/Runtime/D3D12RHI/Private/D3D12Adapter.cpp
--- a/Engine/Source/Runtime/D3D12RHI/Private/D3D12Adapter.cpp
+++ b/Engine/Source/Runtime/D3D12RHI/Private/D3D12Adapter.cpp
@@ -763,13 +763,13 @@ void FD3D12Adapter::InitializeDevices()
 		{
 			if (SUCCEEDED(RootDevice->QueryInterface(IID_PPV_ARGS(RootDevice1.GetInitReference()))))
 			{
-				UE_LOG(LogD3D12RHI, Log, TEXT("The system supports ID3D12Device1."));
+				UE_LOG(LogD3D12RHI, Log, TEXT("ID3D12Device1 is supported."));
 			}
 
 #if PLATFORM_WINDOWS || PLATFORM_HOLOLENS
 			if (SUCCEEDED(RootDevice->QueryInterface(IID_PPV_ARGS(RootDevice2.GetInitReference()))))
 			{
-				UE_LOG(LogD3D12RHI, Log, TEXT("The system supports ID3D12Device2."));
+				UE_LOG(LogD3D12RHI, Log, TEXT("ID3D12Device2 is supported."));
 			}
 #endif
 
@@ -782,12 +782,12 @@ void FD3D12Adapter::InitializeDevices()
 #if D3D12_RHI_RAYTRACING
 			if (SUCCEEDED(RootDevice->QueryInterface(IID_PPV_ARGS(RootDevice5.GetInitReference()))))
 			{
-				UE_LOG(LogD3D12RHI, Log, TEXT("The system supports ID3D12Device5."));
+				UE_LOG(LogD3D12RHI, Log, TEXT("ID3D12Device5 is supported."));
 			}
 
 			if (SUCCEEDED(RootDevice->QueryInterface(IID_PPV_ARGS(RootDevice7.GetInitReference()))))
 			{
-				UE_LOG(LogD3D12RHI, Log, TEXT("The system supports ID3D12Device7."));
+				UE_LOG(LogD3D12RHI, Log, TEXT("ID3D12Device7 is supported."));
 			}
 
 			D3D12_FEATURE_DATA_D3D12_OPTIONS5 D3D12Caps5 = {};
@@ -799,20 +799,53 @@ void FD3D12Adapter::InitializeDevices()
 					&& FDataDrivenShaderPlatformInfo::GetSupportsRayTracing(GMaxRHIShaderPlatform)
 					&& !FParse::Param(FCommandLine::Get(), TEXT("noraytracing")))
 				{
-					UE_LOG(LogD3D12RHI, Log, TEXT("D3D12 ray tracing 1.0 is supported."));
+					UE_LOG(LogD3D12RHI, Log, TEXT("D3D12 ray tracing tier 1.0 is supported."));
 
 					GRHISupportsRayTracing = true;
 
 					if (D3D12Caps5.RaytracingTier >= D3D12_RAYTRACING_TIER_1_1
 						&& RootDevice7)
 					{
-						UE_LOG(LogD3D12RHI, Log, TEXT("D3D12 ray tracing 1.1 is supported."));
+						UE_LOG(LogD3D12RHI, Log, TEXT("D3D12 ray tracing tier 1.1 is supported."));
 
 						GRHISupportsRayTracingPSOAdditions = true;
 					}
 				}
 			}
 #endif // D3D12_RHI_RAYTRACING
+
+#if D3D12_CORE_ENABLED
+			{
+				D3D12_FEATURE_DATA_D3D12_OPTIONS7 D3D12Caps7 = {};
+				RootDevice->CheckFeatureSupport(D3D12_FEATURE_D3D12_OPTIONS7, &D3D12Caps7, sizeof(D3D12Caps7));
+
+				D3D12_FEATURE_DATA_D3D12_OPTIONS9 D3D12Caps9 = {};
+				RootDevice->CheckFeatureSupport(D3D12_FEATURE_D3D12_OPTIONS9, &D3D12Caps9, sizeof(D3D12Caps9));
+
+				D3D12_FEATURE_DATA_D3D12_OPTIONS11 D3D12Caps11 = {};
+				RootDevice->CheckFeatureSupport(D3D12_FEATURE_D3D12_OPTIONS11, &D3D12Caps11, sizeof(D3D12Caps11));
+
+				if (D3D12Caps7.MeshShaderTier >= D3D12_MESH_SHADER_TIER_1)
+				{
+					UE_LOG(LogD3D12RHI, Log, TEXT("Mesh shader tier 1.0 is supported"));
+				}
+
+				if (D3D12Caps9.AtomicInt64OnTypedResourceSupported)
+				{
+					UE_LOG(LogD3D12RHI, Log, TEXT("AtomicInt64OnTypedResource is supported"));
+				}
+
+				if (D3D12Caps9.AtomicInt64OnGroupSharedSupported)
+				{
+					UE_LOG(LogD3D12RHI, Log, TEXT("AtomicInt64OnGroupShared is supported"));
+				}
+
+				if (D3D12Caps11.AtomicInt64OnDescriptorHeapResourceSupported)
+				{
+					UE_LOG(LogD3D12RHI, Log, TEXT("AtomicInt64OnDescriptorHeapResource is supported"));
+				}
+			}
+#endif // D3D12_CORE_ENABLED
 		}
 
 #if PLATFORM_WINDOWS || PLATFORM_HOLOLENS
