Fixed compile error due to undefined D3D12_CORE_ENABLED

diff --git a/Engine/Source/Runtime/D3D12RHI/Private/D3D12Adapter.cpp b/Engine/Source/Runtime/D3D12RHI/Private/D3D12Adapter.cpp
--- a/Engine/Source/Runtime/D3D12RHI/Private/D3D12Adapter.cpp
+++ b/Engine/Source/Runtime/D3D12RHI/Private/D3D12Adapter.cpp
@@ -814,7 +814,7 @@ void FD3D12Adapter::InitializeDevices()
 			}
 #endif // D3D12_RHI_RAYTRACING
 
-#if D3D12_CORE_ENABLED
+#if PLATFORM_WINDOWS && D3D12_CORE_ENABLED
 			{
 				D3D12_FEATURE_DATA_D3D12_OPTIONS7 D3D12Caps7 = {};
 				RootDevice->CheckFeatureSupport(D3D12_FEATURE_D3D12_OPTIONS7, &D3D12Caps7, sizeof(D3D12Caps7));
@@ -845,7 +845,7 @@ void FD3D12Adapter::InitializeDevices()
 					UE_LOG(LogD3D12RHI, Log, TEXT("AtomicInt64OnDescriptorHeapResource is supported"));
 				}
 			}
-#endif // D3D12_CORE_ENABLED
+#endif // PLATFORM_WINDOWS && D3D12_CORE_ENABLED
 		}
 
 #if PLATFORM_WINDOWS || PLATFORM_HOLOLENS
