Refactor acceleration structure allocations to happen in constructors instead of later during build

Kenzo.Terelst

diff --git a/Engine/Source/Runtime/D3D12RHI/Private/D3D12Allocation.cpp b/Engine/Source/Runtime/D3D12RHI/Private/D3D12Allocation.cpp
--- a/Engine/Source/Runtime/D3D12RHI/Private/D3D12Allocation.cpp
+++ b/Engine/Source/Runtime/D3D12RHI/Private/D3D12Allocation.cpp
@@ -1267,6 +1267,17 @@ FD3D12BufferPool* FD3D12DefaultBufferAllocator::CreateBufferPool(D3D12_HEAP_TYPE
 
 	// Disable defrag if not Default memory
 	bool bDefragEnabled = (InitConfig.HeapType == D3D12_HEAP_TYPE_DEFAULT);
+
+#if D3D12_RHI_RAYTRACING
+	// Disable defrag on the RT Acceleration pool - #kenzo_todo
+	// Acceleration structure buffers may be created, but not built immediately.
+	// If CopyRaytracingAccelerationStructure is called on such buffer, the GPU will crash.
+	if (InitConfig.InitialResourceState == D3D12_RESOURCE_STATE_RAYTRACING_ACCELERATION_STRUCTURE)
+	{
+		bDefragEnabled = false;
+	}
+#endif // D3D12_RHI_RAYTRACING
+
 	FD3D12BufferPool* NewPool = new FD3D12PoolAllocator(Device, GetVisibilityMask(), InitConfig, Name, AllocationStrategy, PoolSize, PoolAlignment, MaxAllocationSize, FreeListOrder, bDefragEnabled);
 
 #else // USE_POOL_ALLOCATOR
diff --git a/Engine/Source/Runtime/D3D12RHI/Private/D3D12RayTracing.cpp b/Engine/Source/Runtime/D3D12RHI/Private/D3D12RayTracing.cpp
--- a/Engine/Source/Runtime/D3D12RHI/Private/D3D12RayTracing.cpp
+++ b/Engine/Source/Runtime/D3D12RHI/Private/D3D12RayTracing.cpp
@@ -2368,7 +2368,8 @@ FRayTracingPipelineStateRHIRef FD3D12DynamicRHI::RHICreateRayTracingPipelineStat
 
 FRayTracingGeometryRHIRef FD3D12DynamicRHI::RHICreateRayTracingGeometry(const FRayTracingGeometryInitializer& Initializer)
 {
-	return new FD3D12RayTracingGeometry(Initializer);
+	FD3D12Adapter& Adapter = GetAdapter();
+	return new FD3D12RayTracingGeometry(&Adapter, Initializer);
 }
 
 FRayTracingSceneRHIRef FD3D12DynamicRHI::RHICreateRayTracingScene(const FRayTracingSceneInitializer& Initializer)
@@ -2382,11 +2383,62 @@ FRayTracingSceneRHIRef FD3D12DynamicRHI::RHICreateRayTracingScene(const FRayTrac
 
 FBufferRHIRef FD3D12RayTracingGeometry::NullTransformBuffer;
 
-FD3D12RayTracingGeometry::FD3D12RayTracingGeometry(const FRayTracingGeometryInitializer& Initializer)
+enum class ERayTracingBufferType
+{
+	AccelerationStructure,
+	Scratch
+};
+
+static TRefCountPtr<FD3D12Buffer> CreateRayTracingBuffer(FD3D12Adapter* Adapter, uint32 GPUIndex, uint64 Size, ERayTracingBufferType Type, FName DebugName)
+{
+	check(DebugName.IsValid());
+	check(Size);
+
+	TRefCountPtr<FD3D12Buffer> Result;
+
+	FRHIResourceCreateInfo CreateInfo(*DebugName.ToString());
+	CreateInfo.GPUMask = FRHIGPUMask::FromIndex(GPUIndex);
+
+	ID3D12ResourceAllocator* ResourceAllocator = nullptr;
+
+	D3D12_RESOURCE_DESC BufferDesc = CD3DX12_RESOURCE_DESC::Buffer(Size, D3D12_RESOURCE_FLAG_ALLOW_UNORDERED_ACCESS);
+
+	if (Type == ERayTracingBufferType::AccelerationStructure)
+	{
+		Result = Adapter->CreateRHIBuffer(
+			nullptr, BufferDesc, D3D12_RAYTRACING_ACCELERATION_STRUCTURE_BYTE_ALIGNMENT,
+			0, BufferDesc.Width, BUF_AccelerationStructure,
+			ED3D12ResourceStateMode::SingleState, ERHIAccess::RTAccelerationStructure,
+			CreateInfo, ResourceAllocator);
+	}
+	else if (Type == ERayTracingBufferType::Scratch)
+	{
+		Result = Adapter->CreateRHIBuffer(
+			nullptr, BufferDesc, D3D12_RAYTRACING_ACCELERATION_STRUCTURE_BYTE_ALIGNMENT,
+			0, BufferDesc.Width, BUF_UnorderedAccess,
+			ED3D12ResourceStateMode::SingleState, ERHIAccess::UAVMask,
+			CreateInfo, ResourceAllocator);
+	}
+	else
+	{
+		checkNoEntry();
+	}
+
+	// Elevates the raytracing acceleration structure heap priority, which may help performance / stability in low memory conditions
+	ID3D12Pageable* HeapResource = Result->GetResource()->GetPageable();
+	D3D12_RESIDENCY_PRIORITY HeapPriority = D3D12_RESIDENCY_PRIORITY_HIGH;
+	FD3D12Device* NodeDevice = Adapter->GetDevice(GPUIndex);
+	NodeDevice->GetDevice5()->SetResidencyPriority(1, &HeapResource, &HeapPriority);
+
+	return Result;
+}
+
+FD3D12RayTracingGeometry::FD3D12RayTracingGeometry(FD3D12Adapter* Adapter, const FRayTracingGeometryInitializer& Initializer)
+	: FD3D12AdapterChild(Adapter)
 {
 	INC_DWORD_STAT(STAT_D3D12RayTracingAllocatedBLAS);
 
-	DebugName = Initializer.DebugName;
+	DebugName = Initializer.DebugName.IsValid() ? Initializer.DebugName : FName(TEXT("BLAS"));
 	FMemory::Memzero(bHasPendingCompactionRequests);
 	FMemory::Memzero(bRegisteredAsRenameListener);
 
@@ -2449,46 +2501,149 @@ FD3D12RayTracingGeometry::FD3D12RayTracingGeometry(const FRayTracingGeometryInit
 	checkf(Initializer.Segments.Num() > 0, TEXT("Ray tracing geometry must be initialized with at least one segment."));
 	Segments = TArray<FRayTracingGeometrySegment>(Initializer.Segments.GetData(), Initializer.Segments.Num());
 
-#if DO_CHECK
+	IndexStride = RHIIndexBuffer ? RHIIndexBuffer->GetStride() : 0; // stride 0 means implicit triangle list for non-indexed geometry
+
 	static constexpr uint32 IndicesPerPrimitive = 3; // Only triangle meshes are supported
 
+	GeometryDescs.Reserve(Segments.Num());
+
+	uint32 ComputedPrimitiveCountForValidation = 0;
+
+	for (const FRayTracingGeometrySegment& Segment : Segments)
 	{
-		uint32 ComputedPrimitiveCountForValidation = 0;
+		checkf(Segment.VertexBuffer, TEXT("Position vertex buffer is required for ray tracing geometry"));
+		checkf(Segment.VertexBufferStride, TEXT("Position vertex buffer is required for ray tracing geometry"));
+		checkf(Segment.VertexBufferStride % 4 == 0, TEXT("Position vertex buffer stride must be aligned to 4 bytes for ByteAddressBuffer loads to work"));
 
-		for (const FRayTracingGeometrySegment& Segment : Segments)
+		if (Initializer.GeometryType == RTGT_Triangles)
+		{
+			checkf(Segment.VertexBufferElementType == VET_Float3
+				|| Segment.VertexBufferElementType == VET_Float4, TEXT("Only float3/4 vertex buffers are currently implemented.")); // #dxr_todo UE-72160: support other vertex buffer formats
+			checkf(Segment.VertexBufferStride >= 12, TEXT("Only deinterleaved float3 position vertex buffers are currently implemented.")); // #dxr_todo UE-72160: support interleaved vertex buffers
+		}
+		else if (Initializer.GeometryType == RTGT_Procedural)
 		{
-			checkf(Segment.VertexBuffer, TEXT("Position vertex buffer is required for ray tracing geometry"));
-			checkf(Segment.VertexBufferStride, TEXT("Position vertex buffer is required for ray tracing geometry"));
-			checkf(Segment.VertexBufferStride % 4 == 0, TEXT("Position vertex buffer stride must be aligned to 4 bytes for ByteAddressBuffer loads to work"));
+			checkf(Segment.VertexBufferStride >= (2 * sizeof(FVector)), TEXT("Procedural geometry vertex buffer must contain at least 2xFloat3 that defines 3D bounding boxes of primitives."));
+		}
 
-			if (Initializer.GeometryType == RTGT_Triangles)
-			{
-				// #dxr_todo UE-72160: temporary constraints on vertex and index buffer formats (this will be relaxed when more flexible vertex/index fetching is implemented)
-				checkf(Segment.VertexBufferElementType == VET_Float3
-					|| Segment.VertexBufferElementType == VET_Float4, TEXT("Only float3/4 vertex buffers are currently implemented.")); // #dxr_todo UE-72160: support other vertex buffer formats
-				checkf(Segment.VertexBufferStride >= 12, TEXT("Only deinterleaved float3 position vertex buffers are currently implemented.")); // #dxr_todo UE-72160: support interleaved vertex buffers
-			}
-			else if (Initializer.GeometryType == RTGT_Procedural)
+		ComputedPrimitiveCountForValidation += Segment.NumPrimitives;
+		check(Segment.FirstPrimitive + Segment.NumPrimitives <= Initializer.TotalPrimitiveCount);
+
+		if (Initializer.IndexBuffer)
+		{
+			check(Initializer.IndexBuffer->GetSize() >=
+				(Segment.FirstPrimitive + Segment.NumPrimitives) * IndicesPerPrimitive * IndexStride + IndexOffsetInBytes);
+		}
+
+
+		D3D12_RAYTRACING_GEOMETRY_DESC Desc = {};
+
+		Desc.Flags = D3D12_RAYTRACING_GEOMETRY_FLAG_NONE;
+		Desc.Type = GeometryType;
+
+		if (Segment.bForceOpaque)
+		{
+			// Deny anyhit shader invocations when this segment is hit
+			Desc.Flags |= D3D12_RAYTRACING_GEOMETRY_FLAG_OPAQUE;
+		}
+
+		if (!Segment.bAllowDuplicateAnyHitShaderInvocation)
+		{
+			// Allow only a single any-hit shader invocation per primitive
+			Desc.Flags |= D3D12_RAYTRACING_GEOMETRY_FLAG_NO_DUPLICATE_ANYHIT_INVOCATION;
+		}
+
+		// Conservative estimate of the maximum number of elements in this VB.
+		// Real used number will depend on the index buffer and is not available here right now.
+		// #dxr_todo: Add explicit vertex count to FRayTracingGeometrySegment.
+		const uint32 MaxSegmentVertices = (Segment.VertexBuffer->GetSize() - Segment.VertexBufferOffset) / Segment.VertexBufferStride;
+
+		switch (GeometryType)
+		{
+		case D3D12_RAYTRACING_GEOMETRY_TYPE_TRIANGLES:
+			switch (Segment.VertexBufferElementType)
 			{
-				checkf(Segment.VertexBufferStride >= (2 * sizeof(FVector)), TEXT("Procedural geometry vertex buffer must contain at least 2xFloat3 that defines 3D bounding boxes of primitives."));
+			case VET_Float4:
+				// While the DXGI_FORMAT_R32G32B32A32_FLOAT format is not supported by DXR, since we manually load vertex 
+				// data when we are building the BLAS, we can just rely on the vertex stride to offset the read index, 
+				// and read only the 3 vertex components, and so use the DXGI_FORMAT_R32G32B32_FLOAT vertex format
+			case VET_Float3:
+				Desc.Triangles.VertexFormat = DXGI_FORMAT_R32G32B32_FLOAT;
+				break;
+			case VET_Float2:
+				Desc.Triangles.VertexFormat = DXGI_FORMAT_R32G32_FLOAT;
+				break;
+			case VET_Half2:
+				Desc.Triangles.VertexFormat = DXGI_FORMAT_R16G16_FLOAT;
+				break;
+			default:
+				checkNoEntry();
+				break;
 			}
 
-			ComputedPrimitiveCountForValidation += Segment.NumPrimitives;
-			check(Segment.FirstPrimitive + Segment.NumPrimitives <= Initializer.TotalPrimitiveCount);
-
 			if (Initializer.IndexBuffer)
 			{
-				check(Initializer.IndexBuffer->GetSize() >=
-					(Segment.FirstPrimitive + Segment.NumPrimitives) * IndicesPerPrimitive * IndexStride + IndexOffsetInBytes);
+				// In some cases the geometry is created with 16 bit index buffer, but it's 32 bit at build time.
+				// We conservatively set this to 32 bit to allocate acceleration structure memory.
+				Desc.Triangles.IndexFormat = DXGI_FORMAT_R32_UINT;
+				Desc.Triangles.IndexCount = Segment.NumPrimitives * IndicesPerPrimitive;
+				Desc.Triangles.VertexCount = MaxSegmentVertices;
+			}
+			else
+			{
+				// Non-indexed geometry
+				checkf(Segments.Num() == 1, TEXT("Non-indexed geometry with multiple segments is not implemented."));
+				Desc.Triangles.IndexFormat = DXGI_FORMAT_UNKNOWN;
+				Desc.Triangles.VertexCount = FMath::Min<uint32>(MaxSegmentVertices, TotalPrimitiveCount * 3);
 			}
+			break;
+
+		case D3D12_RAYTRACING_GEOMETRY_TYPE_PROCEDURAL_PRIMITIVE_AABBS:
+			Desc.AABBs.AABBCount = Segment.NumPrimitives;
+			break;
+
+		default:
+			checkf(false, TEXT("Unexpected ray tracing geometry type"));
+			break;
 		}
-		check(ComputedPrimitiveCountForValidation == Initializer.TotalPrimitiveCount);
+
+		GeometryDescs.Add(Desc);
 	}
-#endif // DO_CHECK
+
+	check(ComputedPrimitiveCountForValidation == Initializer.TotalPrimitiveCount);
 
 	RHIIndexBuffer = Initializer.IndexBuffer;
 	SetDirty(FRHIGPUMask::All(), true);
 
+	D3D12_BUILD_RAYTRACING_ACCELERATION_STRUCTURE_INPUTS PrebuildDescInputs = {};
+
+	PrebuildDescInputs.Type = D3D12_RAYTRACING_ACCELERATION_STRUCTURE_TYPE_BOTTOM_LEVEL;
+	PrebuildDescInputs.DescsLayout = D3D12_ELEMENTS_LAYOUT_ARRAY;
+	PrebuildDescInputs.NumDescs = GeometryDescs.Num();
+	PrebuildDescInputs.pGeometryDescs = GeometryDescs.GetData();
+	PrebuildDescInputs.Flags = D3D12_RAYTRACING_ACCELERATION_STRUCTURE_BUILD_FLAGS(BuildFlags);
+
+	// Get maximum buffer sizes for all GPUs in the system
+	for (uint32 GPUIndex = 0; GPUIndex < MAX_NUM_GPUS && GPUIndex < GNumExplicitGPUsForRendering; ++GPUIndex)
+	{
+		ID3D12Device5* RayTracingDevice = Adapter->GetDevice(GPUIndex)->GetDevice5();
+		D3D12_RAYTRACING_ACCELERATION_STRUCTURE_PREBUILD_INFO PrebuildInfo = {};
+		RayTracingDevice->GetRaytracingAccelerationStructurePrebuildInfo(&PrebuildDescInputs, &PrebuildInfo);
+
+		AccelerationStructureSize = FMath::Max(AccelerationStructureSize, PrebuildInfo.ResultDataMaxSizeInBytes);
+		BuildScratchBufferSize = FMath::Max(BuildScratchBufferSize, PrebuildInfo.ScratchDataSizeInBytes);
+		UpdateScratchBufferSize = FMath::Max(UpdateScratchBufferSize, PrebuildInfo.UpdateScratchDataSizeInBytes);
+	}
+
+	// Allocate acceleration structure buffer
+	for (uint32 GPUIndex = 0; GPUIndex < MAX_NUM_GPUS && GPUIndex < GNumExplicitGPUsForRendering; ++GPUIndex)
+	{
+		AccelerationStructureBuffers[GPUIndex] = CreateRayTracingBuffer(Adapter, GPUIndex, AccelerationStructureSize, ERayTracingBufferType::AccelerationStructure, DebugName);
+
+		INC_MEMORY_STAT_BY(STAT_D3D12RayTracingUsedVideoMemory, AccelerationStructureBuffers[GPUIndex]->GetSize());
+		INC_MEMORY_STAT_BY(STAT_D3D12RayTracingBLASMemory, AccelerationStructureBuffers[GPUIndex]->GetSize());
+	}
+
 	INC_DWORD_STAT_BY(STAT_D3D12RayTracingTrianglesBLAS, Initializer.TotalPrimitiveCount);
 }
 
@@ -2522,15 +2677,6 @@ FD3D12RayTracingGeometry::~FD3D12RayTracingGeometry()
 		}
 	}
 
-	for (TRefCountPtr<FD3D12Buffer>& Buffer : ScratchBuffers)
-	{
-		if (Buffer)
-		{
-			DEC_MEMORY_STAT_BY(STAT_D3D12RayTracingUsedVideoMemory, Buffer->GetSize());
-			DEC_MEMORY_STAT_BY(STAT_D3D12RayTracingBLASMemory, Buffer->GetSize());
-		}
-	}
-
 	DEC_DWORD_STAT_BY(STAT_D3D12RayTracingTrianglesBLAS, TotalPrimitiveCount);
 	DEC_DWORD_STAT(STAT_D3D12RayTracingAllocatedBLAS);
 }
@@ -2645,55 +2791,10 @@ void FD3D12RayTracingGeometry::SetupHitGroupSystemParameters(uint32 InGPUIndex)
 	}
 }
 
-static void CreateAccelerationStructureBuffers(
-	TRefCountPtr<FD3D12Buffer>& AccelerationStructureBuffer,
-	TRefCountPtr<FD3D12Buffer>&  ScratchBuffer,
-	FD3D12Adapter* Adapter, uint32 GPUIndex,
-	const D3D12_RAYTRACING_ACCELERATION_STRUCTURE_PREBUILD_INFO& PrebuildInfo,
-	const TCHAR* DebugName)
+static bool ShouldCompactAfterBuild(D3D12_RAYTRACING_ACCELERATION_STRUCTURE_BUILD_FLAGS BuildFlags)
 {
-	FRHIResourceCreateInfo CreateInfo(DebugName);
-
-	D3D12_RESOURCE_DESC AccelerationStructureBufferDesc = CD3DX12_RESOURCE_DESC::Buffer(
-		PrebuildInfo.ResultDataMaxSizeInBytes, D3D12_RESOURCE_FLAG_ALLOW_UNORDERED_ACCESS);
-
-	CreateInfo.GPUMask = FRHIGPUMask::FromIndex(GPUIndex);
-	ID3D12ResourceAllocator* ResourceAllocator = nullptr;
-	AccelerationStructureBuffer = Adapter->CreateRHIBuffer(
-		nullptr, AccelerationStructureBufferDesc, D3D12_RAYTRACING_ACCELERATION_STRUCTURE_BYTE_ALIGNMENT,
-		0, AccelerationStructureBufferDesc.Width, BUF_AccelerationStructure, ED3D12ResourceStateMode::SingleState, ERHIAccess::RTAccelerationStructure, CreateInfo, ResourceAllocator);
-
-	SetName(AccelerationStructureBuffer->GetResource(), DebugName);
-
-	// #dxr_todo UE-72161: scratch buffers can be pooled and reused for different scenes and geometries
-	D3D12_RESOURCE_DESC ScratchBufferDesc = CD3DX12_RESOURCE_DESC::Buffer(
-		FMath::Max(PrebuildInfo.UpdateScratchDataSizeInBytes, PrebuildInfo.ScratchDataSizeInBytes), D3D12_RESOURCE_FLAG_ALLOW_UNORDERED_ACCESS);
-
-	CreateInfo.GPUMask = FRHIGPUMask::FromIndex(GPUIndex);
-	CreateInfo.DebugName = TEXT("ScratchBuffer");
-	ScratchBuffer = Adapter->CreateRHIBuffer(
-		nullptr, ScratchBufferDesc, D3D12_RAYTRACING_ACCELERATION_STRUCTURE_BYTE_ALIGNMENT,
-		0, ScratchBufferDesc.Width, BUF_UnorderedAccess, ED3D12ResourceStateMode::SingleState, ERHIAccess::UAVMask, CreateInfo, ResourceAllocator);
-
-	SetName(ScratchBuffer->GetResource(), TEXT("Acceleration structure scratch"));
-
-	// Elevates the raytracing acceleration structure heap priority, which may help performance / stability in low memory conditions
-	{
-		ID3D12Pageable* HeapResources[] =
-		{
-			AccelerationStructureBuffer->GetResource()->GetPageable(),
-			ScratchBuffer->GetResource()->GetPageable()
-		};
-
-		D3D12_RESIDENCY_PRIORITY HeapPriorities[] =
-		{
-			D3D12_RESIDENCY_PRIORITY_HIGH,
-			D3D12_RESIDENCY_PRIORITY_HIGH
-		};
-
-		FD3D12Device* NodeDevice = Adapter->GetDevice(GPUIndex);
-		NodeDevice->GetDevice5()->SetResidencyPriority(2, HeapResources, HeapPriorities);
-	}
+	return EnumHasAllFlags(BuildFlags, D3D12_RAYTRACING_ACCELERATION_STRUCTURE_BUILD_FLAG_ALLOW_COMPACTION | D3D12_RAYTRACING_ACCELERATION_STRUCTURE_BUILD_FLAG_PREFER_FAST_TRACE)
+		&& !EnumHasAnyFlags(BuildFlags, D3D12_RAYTRACING_ACCELERATION_STRUCTURE_BUILD_FLAG_ALLOW_UPDATE);
 }
 
 void FD3D12RayTracingGeometry::BuildAccelerationStructure(FD3D12CommandContext& CommandContext, EAccelerationStructureBuildMode BuildMode)
@@ -2709,11 +2810,10 @@ void FD3D12RayTracingGeometry::BuildAccelerationStructure(FD3D12CommandContext&
 	IndexStride = RHIIndexBuffer ? RHIIndexBuffer->GetStride() : 0; // stride 0 means implicit triangle list for non-indexed geometry
 
 	const bool bIsUpdate = BuildMode == EAccelerationStructureBuildMode::Update;
-		
-	// Array of geometry descriptions, one per segment (single-segment geometry is a common case).
-	TArray<D3D12_RAYTRACING_GEOMETRY_DESC, TInlineAllocator<1>> Descs;
 
-	Descs.Reserve(Segments.Num());
+	// Use the pre-built descs as template and set the GPU resource pointers (current VB/IB).
+	TArray<D3D12_RAYTRACING_GEOMETRY_DESC, TInlineAllocator<32>> LocalGeometryDescs;
+	LocalGeometryDescs.Reserve(GeometryDescs.Num());
 
 	FD3D12Buffer* IndexBuffer = CommandContext.RetrieveObject<FD3D12Buffer>(RHIIndexBuffer.GetReference());
 	FD3D12Buffer* NullTransformBufferD3D12 = CommandContext.RetrieveObject<FD3D12Buffer>(NullTransformBuffer.GetReference());
@@ -2725,26 +2825,11 @@ void FD3D12RayTracingGeometry::BuildAccelerationStructure(FD3D12CommandContext&
 
 	for (int32 SegmentIndex = 0; SegmentIndex < Segments.Num(); ++SegmentIndex)
 	{
+		D3D12_RAYTRACING_GEOMETRY_DESC Desc = GeometryDescs[SegmentIndex]; // Copy from template
+
 		const FRayTracingGeometrySegment& Segment = Segments[SegmentIndex];
 		const FHitGroupSystemParameters& SystemParameters = HitGroupSystemParametersForThisGPU[SegmentIndex];
 
-		D3D12_RAYTRACING_GEOMETRY_DESC Desc = {};
-
-		Desc.Flags = D3D12_RAYTRACING_GEOMETRY_FLAG_NONE;
-		Desc.Type = GeometryType;
-
-		if (Segment.bForceOpaque)
-		{
-			// Deny anyhit shader invocations when this segment is hit
-			Desc.Flags |= D3D12_RAYTRACING_GEOMETRY_FLAG_OPAQUE;
-		}
-
-		if (!Segment.bAllowDuplicateAnyHitShaderInvocation)
-		{
-			// Allow only a single any-hit shader invocation per primitive
-			Desc.Flags |= D3D12_RAYTRACING_GEOMETRY_FLAG_NO_DUPLICATE_ANYHIT_INVOCATION;
-		}
-
 		FD3D12Buffer* VertexBuffer = CommandContext.RetrieveObject<FD3D12Buffer>(Segment.VertexBuffer.GetReference());
 		
 		// Conservative estimate of the maximum number of elements in this VB.
@@ -2762,19 +2847,19 @@ void FD3D12RayTracingGeometry::BuildAccelerationStructure(FD3D12CommandContext&
 				// data when we are building the BLAS, we can just rely on the vertex stride to offset the read index, 
 				// and read only the 3 vertex components, and so use the DXGI_FORMAT_R32G32B32_FLOAT vertex format
 			case VET_Float3:
-				Desc.Triangles.VertexFormat = DXGI_FORMAT_R32G32B32_FLOAT;
+				check(Desc.Triangles.VertexFormat == DXGI_FORMAT_R32G32B32_FLOAT);
 				break;
 			case VET_Float2:
-				Desc.Triangles.VertexFormat = DXGI_FORMAT_R32G32_FLOAT;
+				check(Desc.Triangles.VertexFormat == DXGI_FORMAT_R32G32_FLOAT);
 				break;
 			case VET_Half2:
-				Desc.Triangles.VertexFormat = DXGI_FORMAT_R16G16_FLOAT;
+				check(Desc.Triangles.VertexFormat == DXGI_FORMAT_R16G16_FLOAT);
 				break;
 			default:
 				checkNoEntry();
 				break;
 			}
-
+			// #yuriy_todo: refactor to just have 0 primitives in this section
 			if (Segment.bEnabled)
 			{
 				Desc.Triangles.Transform3x4 = D3D12_GPU_VIRTUAL_ADDRESS(0);
@@ -2786,21 +2871,21 @@ void FD3D12RayTracingGeometry::BuildAccelerationStructure(FD3D12CommandContext&
 
 			if (IndexBuffer)
 			{
-				Desc.Triangles.IndexFormat = IndexStride == 4 ? DXGI_FORMAT_R32_UINT : DXGI_FORMAT_R16_UINT;
-				Desc.Triangles.IndexCount = Segment.NumPrimitives * IndicesPerPrimitive;
-				Desc.Triangles.IndexBuffer = SystemParameters.IndexBuffer + SystemParameters.RootConstants.IndexBufferOffsetInBytes;
-				Desc.Triangles.VertexCount = MaxSegmentVertices;
+				check(Desc.Triangles.IndexCount <= Segment.NumPrimitives * IndicesPerPrimitive);
+				check(Desc.Triangles.VertexCount <= MaxSegmentVertices);
 
+				Desc.Triangles.IndexFormat = (IndexStride == 4 ? DXGI_FORMAT_R32_UINT : DXGI_FORMAT_R16_UINT);
+				Desc.Triangles.IndexBuffer = SystemParameters.IndexBuffer + SystemParameters.RootConstants.IndexBufferOffsetInBytes;
 				IndexBuffer->GetResource()->UpdateResidency(CommandContext.CommandListHandle);
 			}
 			else
 			{
 				// Non-indexed geometry
 				checkf(Segments.Num() == 1, TEXT("Non-indexed geometry with multiple segments is not implemented."));
-				Desc.Triangles.IndexFormat = DXGI_FORMAT_UNKNOWN;
-				Desc.Triangles.IndexCount = 0;
-				Desc.Triangles.IndexBuffer = D3D12_GPU_VIRTUAL_ADDRESS(0);
-				Desc.Triangles.VertexCount = FMath::Min<uint32>(MaxSegmentVertices, TotalPrimitiveCount * 3);
+				check(Desc.Triangles.IndexFormat == DXGI_FORMAT_UNKNOWN);
+				check(Desc.Triangles.IndexCount == 0);
+				check(Desc.Triangles.IndexBuffer == D3D12_GPU_VIRTUAL_ADDRESS(0));
+				check(Desc.Triangles.VertexCount == FMath::Min<uint32>(MaxSegmentVertices, TotalPrimitiveCount * 3));
 			}
 
 			Desc.Triangles.VertexBuffer.StartAddress = SystemParameters.VertexBuffer;
@@ -2820,7 +2905,7 @@ void FD3D12RayTracingGeometry::BuildAccelerationStructure(FD3D12CommandContext&
 
 		VertexBuffer->ResourceLocation.GetResource()->UpdateResidency(CommandContext.CommandListHandle);
 
-		Descs.Add(Desc);
+		LocalGeometryDescs.Add(Desc);
 
 		if (GeometryType == D3D12_RAYTRACING_GEOMETRY_TYPE_TRIANGLES)
 		{
@@ -2852,30 +2937,30 @@ void FD3D12RayTracingGeometry::BuildAccelerationStructure(FD3D12CommandContext&
 
 	PrebuildDescInputs.Type = D3D12_RAYTRACING_ACCELERATION_STRUCTURE_TYPE_BOTTOM_LEVEL;
 	PrebuildDescInputs.DescsLayout = D3D12_ELEMENTS_LAYOUT_ARRAY;
-	PrebuildDescInputs.NumDescs = Descs.Num();
-	PrebuildDescInputs.pGeometryDescs = Descs.GetData();
+	PrebuildDescInputs.NumDescs = LocalGeometryDescs.Num();
+	PrebuildDescInputs.pGeometryDescs = LocalGeometryDescs.GetData();
 	PrebuildDescInputs.Flags = LocalBuildFlags;
 
-	if (!AccelerationStructureBuffers[GPUIndex])
-	{
-		check(!bIsUpdate);
-
-		D3D12_RAYTRACING_ACCELERATION_STRUCTURE_PREBUILD_INFO PrebuildInfo = {};
+	D3D12_RAYTRACING_ACCELERATION_STRUCTURE_PREBUILD_INFO PrebuildInfo = {};
 
-		RayTracingDevice->GetRaytracingAccelerationStructurePrebuildInfo(&PrebuildDescInputs, &PrebuildInfo);
+	RayTracingDevice->GetRaytracingAccelerationStructurePrebuildInfo(&PrebuildDescInputs, &PrebuildInfo);
 
-		CreateAccelerationStructureBuffers(
-			AccelerationStructureBuffers[GPUIndex],
-			ScratchBuffers[GPUIndex],
-			Adapter, GPUIndex,
-			PrebuildInfo, 
-			DebugName.IsValid() ? *DebugName.ToString() : TEXT("BLAS"));
+	// Must make sure that values computed in the constructor are valid.
+	check(PrebuildInfo.ResultDataMaxSizeInBytes <= AccelerationStructureSize);
 
-		INC_MEMORY_STAT_BY(STAT_D3D12RayTracingUsedVideoMemory, AccelerationStructureBuffers[GPUIndex]->GetSize());
-		INC_MEMORY_STAT_BY(STAT_D3D12RayTracingUsedVideoMemory, ScratchBuffers[GPUIndex]->GetSize());
+	TRefCountPtr<FD3D12Buffer> ScratchBuffer;
 
-		INC_MEMORY_STAT_BY(STAT_D3D12RayTracingBLASMemory, AccelerationStructureBuffers[GPUIndex]->GetSize());
-		INC_MEMORY_STAT_BY(STAT_D3D12RayTracingBLASMemory, ScratchBuffers[GPUIndex]->GetSize());
+	if (bIsUpdate)
+	{
+		check(PrebuildInfo.UpdateScratchDataSizeInBytes <= UpdateScratchBufferSize);
+		static const FName BufferName("UpdateScratchBLAS");
+		ScratchBuffer = CreateRayTracingBuffer(Adapter, GPUIndex, UpdateScratchBufferSize, ERayTracingBufferType::Scratch, BufferName);
+	}
+	else
+	{
+		check(PrebuildInfo.ScratchDataSizeInBytes <= BuildScratchBufferSize);
+		static const FName BufferName("BuildScratchBLAS");
+		ScratchBuffer = CreateRayTracingBuffer(Adapter, GPUIndex, BuildScratchBufferSize, ERayTracingBufferType::Scratch, BufferName);
 	}
 
 	TransitionBuffers(CommandContext);
@@ -2883,13 +2968,13 @@ void FD3D12RayTracingGeometry::BuildAccelerationStructure(FD3D12CommandContext&
 
 	if (IsDirty(GPUIndex))
 	{
-		ScratchBuffers[GPUIndex]->GetResource()->UpdateResidency(CommandContext.CommandListHandle);
+		ScratchBuffer->GetResource()->UpdateResidency(CommandContext.CommandListHandle);
 		AccelerationStructureBuffers[GPUIndex]->GetResource()->UpdateResidency(CommandContext.CommandListHandle);
 
 		D3D12_BUILD_RAYTRACING_ACCELERATION_STRUCTURE_DESC BuildDesc = {};
 		BuildDesc.Inputs = PrebuildDescInputs;
 		BuildDesc.DestAccelerationStructureData = AccelerationStructureBuffers[GPUIndex]->ResourceLocation.GetGPUVirtualAddress();
-		BuildDesc.ScratchAccelerationStructureData = ScratchBuffers[GPUIndex]->ResourceLocation.GetGPUVirtualAddress();
+		BuildDesc.ScratchAccelerationStructureData = ScratchBuffer->ResourceLocation.GetGPUVirtualAddress();
 		BuildDesc.SourceAccelerationStructureData = bIsUpdate
 			? AccelerationStructureBuffers[GPUIndex]->ResourceLocation.GetGPUVirtualAddress()
 			: D3D12_GPU_VIRTUAL_ADDRESS(0);
@@ -2899,11 +2984,7 @@ void FD3D12RayTracingGeometry::BuildAccelerationStructure(FD3D12CommandContext&
 		RayTracingCommandList->BuildRaytracingAccelerationStructure(&BuildDesc, 0, nullptr);
 
 		// Request compact BLAS if it is static
-		bool bShouldCompactAfterBuild = 
-			(BuildFlags & D3D12_RAYTRACING_ACCELERATION_STRUCTURE_BUILD_FLAG_ALLOW_COMPACTION) && 
-			(BuildFlags & D3D12_RAYTRACING_ACCELERATION_STRUCTURE_BUILD_FLAG_PREFER_FAST_TRACE) && 
-			!(BuildFlags & D3D12_RAYTRACING_ACCELERATION_STRUCTURE_BUILD_FLAG_ALLOW_UPDATE);		
-		if (bShouldCompactAfterBuild)
+		if (ShouldCompactAfterBuild(BuildFlags))
 		{
 			CommandContext.GetParentDevice()->GetRayTracingCompactionRequestHandler()->RequestCompact(this);
 			bHasPendingCompactionRequests[GPUIndex] = true;
@@ -2920,14 +3001,6 @@ void FD3D12RayTracingGeometry::BuildAccelerationStructure(FD3D12CommandContext&
 			INC_DWORD_STAT(STAT_D3D12RayTracingBuiltBLAS);
 		}
 	}
-
-	// We don't need to keep a scratch buffer after initial build if acceleration structure is static.
-	if (!(BuildFlags & D3D12_RAYTRACING_ACCELERATION_STRUCTURE_BUILD_FLAG_ALLOW_UPDATE))
-	{
-		DEC_MEMORY_STAT_BY(STAT_D3D12RayTracingUsedVideoMemory, ScratchBuffers[GPUIndex]->GetSize());
-		DEC_MEMORY_STAT_BY(STAT_D3D12RayTracingBLASMemory, ScratchBuffers[GPUIndex]->GetSize());
-		ScratchBuffers[GPUIndex] = nullptr;
-	}
 }
 
 void FD3D12RayTracingGeometry::CompactAccelerationStructure(FD3D12CommandContext& CommandContext, uint32 InGPUIndex, uint64 InSizeAfterCompaction)
@@ -2946,34 +3019,21 @@ void FD3D12RayTracingGeometry::CompactAccelerationStructure(FD3D12CommandContext
 	DEC_MEMORY_STAT_BY(STAT_D3D12RayTracingBLASMemory, AccelerationStructureBuffers[InGPUIndex]->GetSize());
 
 	// Move old AS into this temporary variable which gets released when this function returns
-	TRefCountPtr<FD3D12Buffer> OldAccelerationStructure = MoveTemp(AccelerationStructureBuffers[InGPUIndex]);
-
-	{
-		FRHIResourceCreateInfo CreateInfo(TEXT("AccelerationStructureBuffer"));
-
-		D3D12_RESOURCE_DESC AccelerationStructureBufferDesc = CD3DX12_RESOURCE_DESC::Buffer(InSizeAfterCompaction, D3D12_RESOURCE_FLAG_ALLOW_UNORDERED_ACCESS);
-
-		CreateInfo.GPUMask = FRHIGPUMask::FromIndex(InGPUIndex);
-		ID3D12ResourceAllocator* ResourceAllocator = nullptr;
-		AccelerationStructureBuffers[InGPUIndex] = CommandContext.GetParentAdapter()->CreateRHIBuffer(
-			nullptr, AccelerationStructureBufferDesc, D3D12_RAYTRACING_ACCELERATION_STRUCTURE_BYTE_ALIGNMENT,
-			0, AccelerationStructureBufferDesc.Width, BUF_AccelerationStructure, ED3D12ResourceStateMode::SingleState, ERHIAccess::RTAccelerationStructure, CreateInfo, ResourceAllocator);
+	TRefCountPtr<FD3D12Buffer> OldAccelerationStructure = AccelerationStructureBuffers[InGPUIndex];
 
-		INC_MEMORY_STAT_BY(STAT_D3D12RayTracingUsedVideoMemory, AccelerationStructureBuffers[InGPUIndex]->GetSize());
-		INC_MEMORY_STAT_BY(STAT_D3D12RayTracingBLASMemory, AccelerationStructureBuffers[InGPUIndex]->GetSize());
+	AccelerationStructureBuffers[InGPUIndex] = CreateRayTracingBuffer(CommandContext.GetParentAdapter(), InGPUIndex, InSizeAfterCompaction, ERayTracingBufferType::AccelerationStructure, DebugName);
 
-		SetName(AccelerationStructureBuffers[InGPUIndex]->GetResource(),
-			DebugName.IsValid() ? *DebugName.ToString() : TEXT("BLAS"));
+	INC_MEMORY_STAT_BY(STAT_D3D12RayTracingUsedVideoMemory, AccelerationStructureBuffers[InGPUIndex]->GetSize());
+	INC_MEMORY_STAT_BY(STAT_D3D12RayTracingBLASMemory, AccelerationStructureBuffers[InGPUIndex]->GetSize());
 
-		OldAccelerationStructure->GetResource()->UpdateResidency(CommandContext.CommandListHandle);
-		AccelerationStructureBuffers[InGPUIndex]->GetResource()->UpdateResidency(CommandContext.CommandListHandle);
+	OldAccelerationStructure->GetResource()->UpdateResidency(CommandContext.CommandListHandle);
+	AccelerationStructureBuffers[InGPUIndex]->GetResource()->UpdateResidency(CommandContext.CommandListHandle);
 
-		ID3D12GraphicsCommandList4* RayTracingCommandList = CommandContext.CommandListHandle.RayTracingCommandList();
-		RayTracingCommandList->CopyRaytracingAccelerationStructure(
-			AccelerationStructureBuffers[InGPUIndex]->ResourceLocation.GetGPUVirtualAddress(),
-			OldAccelerationStructure->ResourceLocation.GetGPUVirtualAddress(),
-			D3D12_RAYTRACING_ACCELERATION_STRUCTURE_COPY_MODE_COMPACT);
-	}
+	ID3D12GraphicsCommandList4* RayTracingCommandList = CommandContext.CommandListHandle.RayTracingCommandList();
+	RayTracingCommandList->CopyRaytracingAccelerationStructure(
+		AccelerationStructureBuffers[InGPUIndex]->ResourceLocation.GetGPUVirtualAddress(),
+		OldAccelerationStructure->ResourceLocation.GetGPUVirtualAddress(),
+		D3D12_RAYTRACING_ACCELERATION_STRUCTURE_COPY_MODE_COMPACT);
 }
 
 FD3D12RayTracingScene::FD3D12RayTracingScene(FD3D12Adapter* Adapter, const FRayTracingSceneInitializer& Initializer)
@@ -2983,7 +3043,7 @@ FD3D12RayTracingScene::FD3D12RayTracingScene(FD3D12Adapter* Adapter, const FRayT
 
 	Instances.SetAllowCPUAccess(true);
 
-	DebugName = Initializer.DebugName;
+	DebugName = Initializer.DebugName.IsValid() ? Initializer.DebugName : FName(TEXT("TLAS"));
 
 	ShaderResourceView = Adapter->CreateLinkedObject<FD3D12ShaderResourceView>(FRHIGPUMask::All(), [&](FD3D12Device* Device)
 	{
@@ -3124,6 +3184,34 @@ FD3D12RayTracingScene::FD3D12RayTracingScene(FD3D12Adapter* Adapter, const FRayT
 
 	NumCallableShaderSlots = Initializer.NumCallableShaderSlots;
 	NumMissShaderSlots = FMath::Max<uint32>(1, Initializer.NumMissShaderSlots);
+
+	for (uint32 GPUIndex = 0; GPUIndex < GNumExplicitGPUsForRendering; ++GPUIndex)
+	{
+		ID3D12Device5* RayTracingDevice = Adapter->GetDevice(GPUIndex)->GetDevice5();
+		D3D12_RAYTRACING_ACCELERATION_STRUCTURE_PREBUILD_INFO PrebuildInfo = {};
+		RayTracingDevice->GetRaytracingAccelerationStructurePrebuildInfo(&BuildInputs, &PrebuildInfo);
+
+		AccelerationStructureSize = FMath::Max(AccelerationStructureSize, PrebuildInfo.ResultDataMaxSizeInBytes);
+		BuildScratchBufferSize = FMath::Max(BuildScratchBufferSize, PrebuildInfo.ScratchDataSizeInBytes);
+	}
+
+	for (uint32 GPUIndex = 0; GPUIndex < GNumExplicitGPUsForRendering; ++GPUIndex)
+	{
+		AccelerationStructureBuffers[GPUIndex] = CreateRayTracingBuffer(Adapter, GPUIndex, AccelerationStructureSize, ERayTracingBufferType::AccelerationStructure, DebugName);
+
+		INC_MEMORY_STAT_BY(STAT_D3D12RayTracingUsedVideoMemory, AccelerationStructureBuffers[GPUIndex]->GetSize());
+		INC_MEMORY_STAT_BY(STAT_D3D12RayTracingTLASMemory, AccelerationStructureBuffers[GPUIndex]->GetSize());
+
+		FD3D12ShaderResourceView* View = FD3D12CommandContext::RetrieveObject<FD3D12ShaderResourceView>(ShaderResourceView.GetReference(), GPUIndex);
+
+		D3D12_SHADER_RESOURCE_VIEW_DESC SRVDesc = {};
+		SRVDesc.Format = DXGI_FORMAT_UNKNOWN;
+		SRVDesc.ViewDimension = D3D12_SRV_DIMENSION_RAYTRACING_ACCELERATION_STRUCTURE;
+		SRVDesc.Shader4ComponentMapping = D3D12_DEFAULT_SHADER_4_COMPONENT_MAPPING;
+		SRVDesc.RaytracingAccelerationStructure.Location = AccelerationStructureBuffers[GPUIndex]->ResourceLocation.GetGPUVirtualAddress();
+
+		View->Initialize(SRVDesc, AccelerationStructureBuffers[GPUIndex], 4, 0);
+	}
 };
 
 FD3D12RayTracingScene::~FD3D12RayTracingScene()
@@ -3162,40 +3250,14 @@ void FD3D12RayTracingScene::BuildAccelerationStructure(FD3D12CommandContext& Com
 
 	D3D12_RAYTRACING_ACCELERATION_STRUCTURE_PREBUILD_INFO PrebuildInfo = {};
 	RayTracingDevice->GetRaytracingAccelerationStructurePrebuildInfo(&BuildInputs, &PrebuildInfo);
+	check(PrebuildInfo.ResultDataMaxSizeInBytes <= AccelerationStructureSize);
+	check(PrebuildInfo.ScratchDataSizeInBytes <= BuildScratchBufferSize);
 
-	TRefCountPtr<FD3D12Buffer>& AccelerationStructureBuffer = AccelerationStructureBuffers[GPUIndex];
-
-	CreateAccelerationStructureBuffers(
-		AccelerationStructureBuffer,
-		ScratchBuffer,
-		Adapter, GPUIndex,
-		PrebuildInfo,
-		DebugName.IsValid() ? *DebugName.ToString() : TEXT("TLAS"));
-
-	INC_MEMORY_STAT_BY(STAT_D3D12RayTracingUsedVideoMemory, AccelerationStructureBuffer->GetSize());
-	INC_MEMORY_STAT_BY(STAT_D3D12RayTracingTLASMemory, AccelerationStructureBuffer->GetSize());
+	static const FName ScratchBufferName("BuildScratchTLAS");
+	ScratchBuffer = CreateRayTracingBuffer(Adapter, GPUIndex, PrebuildInfo.ScratchDataSizeInBytes, ERayTracingBufferType::Scratch, ScratchBufferName);
 
 	FD3D12ShaderResourceView* AccelerationStructureView = CommandContext.RetrieveObject<FD3D12ShaderResourceView>(ShaderResourceView.GetReference());
 
-	if (bAccelerationStructureViewInitialized[GPUIndex])
-	{
-		check(AccelerationStructureView->GetParentDevice() == AccelerationStructureBuffer->GetParentDevice());
-		AccelerationStructureView->Rename(AccelerationStructureBuffer.GetReference());
-	}
-	else
-	{
-		D3D12_SHADER_RESOURCE_VIEW_DESC SRVDesc = {};
-		SRVDesc.Format = DXGI_FORMAT_UNKNOWN;
-		SRVDesc.ViewDimension = D3D12_SRV_DIMENSION_RAYTRACING_ACCELERATION_STRUCTURE;
-		SRVDesc.Shader4ComponentMapping = D3D12_DEFAULT_SHADER_4_COMPONENT_MAPPING;
-		SRVDesc.RaytracingAccelerationStructure.Location = AccelerationStructureBuffer->ResourceLocation.GetGPUVirtualAddress();
-
-		check(AccelerationStructureView->GetParentDevice() == AccelerationStructureBuffer->GetParentDevice());
-		AccelerationStructureView->Initialize(SRVDesc, AccelerationStructureBuffer, 4, 0);
-
-		bAccelerationStructureViewInitialized[GPUIndex] = true;
-	}
-
 	// Create and fill instance buffer
 
 	if (BuildInputs.NumDescs)
@@ -3326,6 +3388,7 @@ void FD3D12RayTracingScene::BuildAccelerationStructure(FD3D12CommandContext& Com
 		Geometry->AccelerationStructureBuffers[GPUIndex]->ResourceLocation.GetResource()->UpdateResidency(CommandContext.CommandListHandle);
 	}
 
+	TRefCountPtr<FD3D12Buffer>& AccelerationStructureBuffer = AccelerationStructureBuffers[GPUIndex];
 	AccelerationStructureBuffer->GetResource()->UpdateResidency(CommandContext.CommandListHandle);
 	ScratchBuffer->GetResource()->UpdateResidency(CommandContext.CommandListHandle);
 
diff --git a/Engine/Source/Runtime/D3D12RHI/Private/D3D12RayTracing.h b/Engine/Source/Runtime/D3D12RHI/Private/D3D12RayTracing.h
--- a/Engine/Source/Runtime/D3D12RHI/Private/D3D12RayTracing.h
+++ b/Engine/Source/Runtime/D3D12RHI/Private/D3D12RayTracing.h
@@ -19,11 +19,11 @@ struct FHitGroupSystemParameters
 	FHitGroupSystemRootConstants RootConstants;
 };
 
-class FD3D12RayTracingGeometry : public FRHIRayTracingGeometry, public FD3D12ShaderResourceRenameListener
+class FD3D12RayTracingGeometry : public FRHIRayTracingGeometry, public FD3D12AdapterChild, public FD3D12ShaderResourceRenameListener, public FNoncopyable
 {
 public:
 
-	FD3D12RayTracingGeometry(const FRayTracingGeometryInitializer& Initializer);
+	FD3D12RayTracingGeometry(FD3D12Adapter* Adapter, const FRayTracingGeometryInitializer& Initializer);
 	~FD3D12RayTracingGeometry();
 
 	void SetupHitGroupSystemParameters(uint32 InGPUIndex);
@@ -63,7 +63,6 @@ public:
 	static FBufferRHIRef NullTransformBuffer; // Null transform for hidden sections
 
 	TRefCountPtr<FD3D12Buffer> AccelerationStructureBuffers[MAX_NUM_GPUS];
-	TRefCountPtr<FD3D12Buffer> ScratchBuffers[MAX_NUM_GPUS];
 
 	bool bRegisteredAsRenameListener[MAX_NUM_GPUS];
 	bool bHasPendingCompactionRequests[MAX_NUM_GPUS];
@@ -72,6 +71,16 @@ public:
 	TArray<FHitGroupSystemParameters> HitGroupSystemParameters[MAX_NUM_GPUS];
 
 	FName DebugName;
+
+	// Array of geometry descriptions, one per segment (single-segment geometry is a common case).
+	// Only references CPU-accessible structures (no GPU resources).
+	// Used as a template for BuildAccelerationStructure() later.
+	TArray<D3D12_RAYTRACING_GEOMETRY_DESC, TInlineAllocator<1>> GeometryDescs;
+
+	uint64 AccelerationStructureSize = 0;
+	uint64 AccelerationStructureCompactedSize = 0;
+	uint64 BuildScratchBufferSize = 0;
+	uint64 UpdateScratchBufferSize = 0;
 };
 
 class FD3D12RayTracingScene : public FRHIRayTracingScene, public FD3D12AdapterChild, public FNoncopyable
@@ -89,11 +98,13 @@ public:
 	void BuildAccelerationStructure(FD3D12CommandContext& CommandContext);
 
 	TRefCountPtr<FD3D12Buffer> AccelerationStructureBuffers[MAX_NUM_GPUS];
-	bool bAccelerationStructureViewInitialized[MAX_NUM_GPUS] = {};
 
 	TResourceArray<D3D12_RAYTRACING_INSTANCE_DESC, 16> Instances;
 	D3D12_BUILD_RAYTRACING_ACCELERATION_STRUCTURE_INPUTS BuildInputs = {};
-	
+
+	uint64 AccelerationStructureSize = 0;
+	uint64 BuildScratchBufferSize = 0;
+
 	struct FInstanceCopyCommand
 	{
 		FShaderResourceViewRHIRef Source;
