Disable several very verbose profiling events by default. Changed DECLARE_CYCLE_STAT_WITH_FLAGS_EXTERN to always include ClearEveryFrame, as that's what all cycle counter stats expect. Removed explicit ClearEveryFrame at use sites.

Disabled profiling events:
- FMaterial::TryGetShaders
- Various D3D12 RHI cycle stats
- CreateClusters GT cycle counter
- Moved FLinkerLoad::Preload event into slow path
- PSOGraphics/ComputeFindOrCreateTime (added explicit profiling events when slow path is taken)

Dramatically reduces trace file sizes, even with quite detailed instrumentation via -statnamedevents.

diff --git a/Engine/Source/Runtime/Core/Private/HAL/ThreadingBase.cpp b/Engine/Source/Runtime/Core/Private/HAL/ThreadingBase.cpp
--- a/Engine/Source/Runtime/Core/Private/HAL/ThreadingBase.cpp
+++ b/Engine/Source/Runtime/Core/Private/HAL/ThreadingBase.cpp
@@ -1249,7 +1249,7 @@ FQueuedThread::Run()
 			{
 				DECLARE_CYCLE_STAT_WITH_FLAGS(TEXT("FQueuedThread::Run.WaitForWork"),
 				STAT_FQueuedThread_Run_WaitForWork, STATGROUP_ThreadPoolAsyncTasks,
-					EStatFlags::ClearEveryFrame | EStatFlags::Verbose);
+					EStatFlags::Verbose);
 
 				SCOPE_CYCLE_COUNTER(STAT_FQueuedThread_Run_WaitForWork);
 
diff --git a/Engine/Source/Runtime/Core/Public/Misc/CoreStats.h b/Engine/Source/Runtime/Core/Public/Misc/CoreStats.h
--- a/Engine/Source/Runtime/Core/Public/Misc/CoreStats.h
+++ b/Engine/Source/Runtime/Core/Public/Misc/CoreStats.h
@@ -51,9 +51,9 @@ DECLARE_FLOAT_COUNTER_STAT_EXTERN(TEXT("Percentage CPU utilization (relative to
 	CPU Stalls
 -----------------------------------------------------------------------------*/
 
-DECLARE_CYCLE_STAT_WITH_FLAGS_EXTERN(TEXT("CPU Stall - Sleep"), STAT_Sleep, STATGROUP_CPUStalls, EStatFlags::ClearEveryFrame | EStatFlags::Verbose, CORE_API);
-DECLARE_CYCLE_STAT_WITH_FLAGS_EXTERN(TEXT("CPU Stall - Wait For Event" ), STAT_EventWait, STATGROUP_CPUStalls, EStatFlags::ClearEveryFrame | EStatFlags::Verbose,);
+DECLARE_CYCLE_STAT_WITH_FLAGS_EXTERN(TEXT("CPU Stall - Sleep"), STAT_Sleep, STATGROUP_CPUStalls, EStatFlags::Verbose, CORE_API);
+DECLARE_CYCLE_STAT_WITH_FLAGS_EXTERN(TEXT("CPU Stall - Wait For Event" ), STAT_EventWait, STATGROUP_CPUStalls, EStatFlags::Verbose,);
 
 /** The id will be stored as uint64 in the stat message. */
-DECLARE_CYCLE_STAT_WITH_FLAGS_EXTERN( TEXT( "CPU Stall - Wait For Event with ID" ), STAT_EventWaitWithId, STATGROUP_CPUStalls, EStatFlags::ClearEveryFrame | EStatFlags::Verbose, );
-DECLARE_CYCLE_STAT_WITH_FLAGS_EXTERN( TEXT( "CPU Stall - Trigger For Event with ID" ), STAT_EventTriggerWithId, STATGROUP_CPUStalls, EStatFlags::ClearEveryFrame | EStatFlags::Verbose, );
+DECLARE_CYCLE_STAT_WITH_FLAGS_EXTERN( TEXT( "CPU Stall - Wait For Event with ID" ), STAT_EventWaitWithId, STATGROUP_CPUStalls, EStatFlags::Verbose, );
+DECLARE_CYCLE_STAT_WITH_FLAGS_EXTERN( TEXT( "CPU Stall - Trigger For Event with ID" ), STAT_EventTriggerWithId, STATGROUP_CPUStalls, EStatFlags::Verbose, );
diff --git a/Engine/Source/Runtime/Core/Public/Stats/Stats2.h b/Engine/Source/Runtime/Core/Public/Stats/Stats2.h
--- a/Engine/Source/Runtime/Core/Public/Stats/Stats2.h
+++ b/Engine/Source/Runtime/Core/Public/Stats/Stats2.h
@@ -1917,7 +1917,7 @@ struct FStat_##StatName\
 	DECLARE_STAT(CounterName,StatId,GroupId,EStatDataType::ST_int64, EStatFlags::ClearEveryFrame | EStatFlags::CycleStat, FPlatformMemory::MCR_Invalid); \
 	extern APIX DEFINE_STAT(StatId);
 #define DECLARE_CYCLE_STAT_WITH_FLAGS_EXTERN(CounterName,StatId,GroupId,StatFlags, APIX) \
-	DECLARE_STAT(CounterName,StatId,GroupId,EStatDataType::ST_int64, (StatFlags) | EStatFlags::CycleStat, FPlatformMemory::MCR_Invalid); \
+	DECLARE_STAT(CounterName,StatId,GroupId,EStatDataType::ST_int64, (StatFlags) | EStatFlags::ClearEveryFrame | EStatFlags::CycleStat, FPlatformMemory::MCR_Invalid); \
 	extern APIX DEFINE_STAT(StatId);
 #define DECLARE_FLOAT_COUNTER_STAT_EXTERN(CounterName,StatId,GroupId, API) \
 	DECLARE_STAT(CounterName,StatId,GroupId,EStatDataType::ST_double, EStatFlags::ClearEveryFrame, FPlatformMemory::MCR_Invalid); \
diff --git a/Engine/Source/Runtime/CoreUObject/Public/Serialization/AsyncPackageLoader.h b/Engine/Source/Runtime/CoreUObject/Public/Serialization/AsyncPackageLoader.h
--- a/Engine/Source/Runtime/CoreUObject/Public/Serialization/AsyncPackageLoader.h
+++ b/Engine/Source/Runtime/CoreUObject/Public/Serialization/AsyncPackageLoader.h
@@ -16,7 +16,7 @@ DECLARE_STATS_GROUP(TEXT("Async Load Game Thread"), STATGROUP_AsyncLoadGameThrea
 DECLARE_CYCLE_STAT(TEXT("PostLoadObjects GT"), STAT_FAsyncPackage_PostLoadObjectsGameThread, STATGROUP_AsyncLoadGameThread);
 DECLARE_CYCLE_STAT(TEXT("TickAsyncLoading GT"), STAT_FAsyncPackage_TickAsyncLoadingGameThread, STATGROUP_AsyncLoadGameThread);
 DECLARE_CYCLE_STAT(TEXT("Flush Async Loading GT"), STAT_FAsyncPackage_FlushAsyncLoadingGameThread, STATGROUP_AsyncLoadGameThread);
-DECLARE_CYCLE_STAT(TEXT("CreateClusters GT"), STAT_FAsyncPackage_CreateClustersGameThread, STATGROUP_AsyncLoadGameThread);
+DECLARE_CYCLE_STAT_WITH_FLAGS(TEXT("CreateClusters GT"), STAT_FAsyncPackage_CreateClustersGameThread, STATGROUP_AsyncLoadGameThread, EStatFlags::Verbose);
 
 enum class ENotifyRegistrationType;
 enum class ENotifyRegistrationPhase;
diff --git a/Engine/Source/Runtime/D3D12RHI/Private/D3D12State.cpp b/Engine/Source/Runtime/D3D12RHI/Private/D3D12State.cpp
--- a/Engine/Source/Runtime/D3D12RHI/Private/D3D12State.cpp
+++ b/Engine/Source/Runtime/D3D12RHI/Private/D3D12State.cpp
@@ -11,8 +11,8 @@
 // modifying ref counts if not needed; so we manage that ourselves
 FCriticalSection GD3D12SamplerStateCacheLock;
 
-DECLARE_CYCLE_STAT(TEXT("Graphics: Find or Create time"), STAT_PSOGraphicsFindOrCreateTime, STATGROUP_D3D12PipelineState);
-DECLARE_CYCLE_STAT(TEXT("Compute: Find or Create time"), STAT_PSOComputeFindOrCreateTime, STATGROUP_D3D12PipelineState);
+DECLARE_CYCLE_STAT_WITH_FLAGS(TEXT("Graphics: Find or Create time"), STAT_PSOGraphicsFindOrCreateTime, STATGROUP_D3D12PipelineState, EStatFlags::Verbose);
+DECLARE_CYCLE_STAT_WITH_FLAGS(TEXT("Compute: Find or Create time"), STAT_PSOComputeFindOrCreateTime, STATGROUP_D3D12PipelineState, EStatFlags::Verbose);
 
 static D3D12_TEXTURE_ADDRESS_MODE TranslateAddressMode(ESamplerAddressMode AddressMode)
 {
@@ -474,6 +474,9 @@ FGraphicsPipelineStateRHIRef FD3D12DynamicRHI::RHICreateGraphicsPipelineState(co
 		return Found;
 	}
 #endif
+
+	TRACE_CPUPROFILER_EVENT_SCOPE(FD3D12DynamicRHI::RHICreateGraphicsPipelineState);
+
 	// TODO: Remove the need for BoundShaderState objects. Currently they are needed for the Root Signature.
 	FBoundShaderStateRHIRef const BSS = RHICreateBoundShaderState(
 		Initializer.BoundShaderState.VertexDeclarationRHI,
@@ -531,6 +534,8 @@ TRefCountPtr<FRHIComputePipelineState> FD3D12DynamicRHI::RHICreateComputePipelin
 		return Found;
 	}
 
+	TRACE_CPUPROFILER_EVENT_SCOPE(FD3D12DynamicRHI::RHICreateComputePipelineState);
+
 	// We need to actually create a PSO.
 	return PSOCache.CreateAndAdd(ComputeShader, LowLevelDesc);
 }
diff --git a/Engine/Source/Runtime/D3D12RHI/Private/D3D12Stats.h b/Engine/Source/Runtime/D3D12RHI/Private/D3D12Stats.h
--- a/Engine/Source/Runtime/D3D12RHI/Private/D3D12Stats.h
+++ b/Engine/Source/Runtime/D3D12RHI/Private/D3D12Stats.h
@@ -9,6 +9,7 @@ D3D12Stats.h: D3D12 Statistics and Timing Interfaces
 * The D3D RHI stats.
 */
 
+
 DECLARE_CYCLE_STAT_EXTERN(TEXT("Present time"), STAT_D3D12PresentTime, STATGROUP_D3D12RHI, );
 DECLARE_CYCLE_STAT_EXTERN(TEXT("CustomPresent time"), STAT_D3D12CustomPresentTime, STATGROUP_D3D12RHI, );
 
@@ -33,27 +34,27 @@ DECLARE_CYCLE_STAT_EXTERN(TEXT("New bound shader state time"), STAT_D3D12NewBoun
 DECLARE_DWORD_ACCUMULATOR_STAT_EXTERN(TEXT("Num bound shader states"), STAT_D3D12NumBoundShaderState, STATGROUP_D3D12RHI, );
 DECLARE_CYCLE_STAT_EXTERN(TEXT("Set bound shader state"), STAT_D3D12SetBoundShaderState, STATGROUP_D3D12RHI, );
 
-DECLARE_CYCLE_STAT_EXTERN(TEXT("Update uniform buffer"), STAT_D3D12UpdateUniformBufferTime, STATGROUP_D3D12RHI, );
+DECLARE_CYCLE_STAT_WITH_FLAGS_EXTERN(TEXT("Update uniform buffer"), STAT_D3D12UpdateUniformBufferTime, STATGROUP_D3D12RHI, EStatFlags::Verbose, );
 
 DECLARE_CYCLE_STAT_EXTERN(TEXT("Commit resource tables"), STAT_D3D12CommitResourceTables, STATGROUP_D3D12RHI, );
 DECLARE_DWORD_COUNTER_STAT_EXTERN(TEXT("Num textures in tables"), STAT_D3D12SetTextureInTableCalls, STATGROUP_D3D12RHI, );
 
-DECLARE_CYCLE_STAT_EXTERN(TEXT("Clear SRVs time"), STAT_D3D12ClearShaderResourceViewsTime, STATGROUP_D3D12RHI, );
-DECLARE_CYCLE_STAT_EXTERN(TEXT("Set SRV time"), STAT_D3D12SetShaderResourceViewTime, STATGROUP_D3D12RHI, );
-DECLARE_CYCLE_STAT_EXTERN(TEXT("Set UAV time"), STAT_D3D12SetUnorderedAccessViewTime, STATGROUP_D3D12RHI, );
-DECLARE_CYCLE_STAT_EXTERN(TEXT("Commit graphics constants (Set CBV time)"), STAT_D3D12CommitGraphicsConstants, STATGROUP_D3D12RHI, );
-DECLARE_CYCLE_STAT_EXTERN(TEXT("Commit compute constants (Set CBV time)"), STAT_D3D12CommitComputeConstants, STATGROUP_D3D12RHI, );
-DECLARE_CYCLE_STAT_EXTERN(TEXT("Set shader uniform buffer (Set CBV time)"), STAT_D3D12SetShaderUniformBuffer, STATGROUP_D3D12RHI, );
-
-DECLARE_CYCLE_STAT_EXTERN(TEXT("ApplyState time"), STAT_D3D12ApplyStateTime, STATGROUP_D3D12RHI, );
-DECLARE_CYCLE_STAT_EXTERN(TEXT("ApplyState: Rebuild PSO time"), STAT_D3D12ApplyStateRebuildPSOTime, STATGROUP_D3D12RHI, );
-DECLARE_CYCLE_STAT_EXTERN(TEXT("ApplyState: Find PSO time"), STAT_D3D12ApplyStateFindPSOTime, STATGROUP_D3D12RHI, );
-DECLARE_CYCLE_STAT_EXTERN(TEXT("ApplyState: Set SRV time"), STAT_D3D12ApplyStateSetSRVTime, STATGROUP_D3D12RHI, );
-DECLARE_CYCLE_STAT_EXTERN(TEXT("ApplyState: Set UAV time"), STAT_D3D12ApplyStateSetUAVTime, STATGROUP_D3D12RHI, );
-DECLARE_CYCLE_STAT_EXTERN(TEXT("ApplyState: Set Vertex Buffer time"), STAT_D3D12ApplyStateSetVertexBufferTime, STATGROUP_D3D12RHI, );
-DECLARE_CYCLE_STAT_EXTERN(TEXT("ApplyState: Set CBV time"), STAT_D3D12ApplyStateSetConstantBufferTime, STATGROUP_D3D12RHI, );
+DECLARE_CYCLE_STAT_WITH_FLAGS_EXTERN(TEXT("Clear SRVs time"), STAT_D3D12ClearShaderResourceViewsTime, STATGROUP_D3D12RHI, EStatFlags::Verbose, );
+DECLARE_CYCLE_STAT_WITH_FLAGS_EXTERN(TEXT("Set SRV time"), STAT_D3D12SetShaderResourceViewTime, STATGROUP_D3D12RHI, EStatFlags::Verbose, );
+DECLARE_CYCLE_STAT_WITH_FLAGS_EXTERN(TEXT("Set UAV time"), STAT_D3D12SetUnorderedAccessViewTime, STATGROUP_D3D12RHI, EStatFlags::Verbose, );
+DECLARE_CYCLE_STAT_WITH_FLAGS_EXTERN(TEXT("Commit graphics constants (Set CBV time)"), STAT_D3D12CommitGraphicsConstants, STATGROUP_D3D12RHI, EStatFlags::Verbose, );
+DECLARE_CYCLE_STAT_WITH_FLAGS_EXTERN(TEXT("Commit compute constants (Set CBV time)"), STAT_D3D12CommitComputeConstants, STATGROUP_D3D12RHI, EStatFlags::Verbose, );
+DECLARE_CYCLE_STAT_WITH_FLAGS_EXTERN(TEXT("Set shader uniform buffer (Set CBV time)"), STAT_D3D12SetShaderUniformBuffer, STATGROUP_D3D12RHI, EStatFlags::Verbose, );
+
+DECLARE_CYCLE_STAT_WITH_FLAGS_EXTERN(TEXT("ApplyState time"), STAT_D3D12ApplyStateTime, STATGROUP_D3D12RHI, EStatFlags::Verbose, );
+DECLARE_CYCLE_STAT_WITH_FLAGS_EXTERN(TEXT("ApplyState: Rebuild PSO time"), STAT_D3D12ApplyStateRebuildPSOTime, STATGROUP_D3D12RHI, EStatFlags::Verbose, );
+DECLARE_CYCLE_STAT_WITH_FLAGS_EXTERN(TEXT("ApplyState: Find PSO time"), STAT_D3D12ApplyStateFindPSOTime, STATGROUP_D3D12RHI, EStatFlags::Verbose, );
+DECLARE_CYCLE_STAT_WITH_FLAGS_EXTERN(TEXT("ApplyState: Set SRV time"), STAT_D3D12ApplyStateSetSRVTime, STATGROUP_D3D12RHI, EStatFlags::Verbose, );
+DECLARE_CYCLE_STAT_WITH_FLAGS_EXTERN(TEXT("ApplyState: Set UAV time"), STAT_D3D12ApplyStateSetUAVTime, STATGROUP_D3D12RHI, EStatFlags::Verbose, );
+DECLARE_CYCLE_STAT_WITH_FLAGS_EXTERN(TEXT("ApplyState: Set Vertex Buffer time"), STAT_D3D12ApplyStateSetVertexBufferTime, STATGROUP_D3D12RHI, EStatFlags::Verbose, );
+DECLARE_CYCLE_STAT_WITH_FLAGS_EXTERN(TEXT("ApplyState: Set CBV time"), STAT_D3D12ApplyStateSetConstantBufferTime, STATGROUP_D3D12RHI, EStatFlags::Verbose, );
 DECLARE_CYCLE_STAT_EXTERN(TEXT("PSO Create time"), STAT_D3D12PSOCreateTime, STATGROUP_D3D12RHI, );
-DECLARE_CYCLE_STAT_EXTERN(TEXT("Clear MRT time"), STAT_D3D12ClearMRT, STATGROUP_D3D12RHI, );
+DECLARE_CYCLE_STAT_WITH_FLAGS_EXTERN(TEXT("Clear MRT time"), STAT_D3D12ClearMRT, STATGROUP_D3D12RHI, EStatFlags::Verbose, );
 
 DECLARE_CYCLE_STAT_EXTERN(TEXT("ExecuteCommandList time"), STAT_D3D12ExecuteCommandListTime, STATGROUP_D3D12RHI, );
 DECLARE_CYCLE_STAT_EXTERN(TEXT("WaitForFence time"), STAT_D3D12WaitForFenceTime, STATGROUP_D3D12RHI, );
@@ -131,6 +132,7 @@ DECLARE_DWORD_ACCUMULATOR_STAT_EXTERN(TEXT("View Global: Used Descriptors"), STA
 DECLARE_DWORD_ACCUMULATOR_STAT_EXTERN(TEXT("View Global: Wasted Descriptors"), STAT_GlobalViewHeapWastedDescriptors, STATGROUP_D3D12DescriptorHeap, );
 DECLARE_DWORD_COUNTER_STAT_EXTERN(TEXT("View Global: Block Allocations"), STAT_GlobalViewHeapBlockAllocations, STATGROUP_D3D12DescriptorHeap, );
 
+
 struct FD3D12GlobalStats
 {
 	// in bytes, never change after RHI, needed to scale game features
diff --git a/Engine/Source/Runtime/Engine/Private/Materials/MaterialShared.cpp b/Engine/Source/Runtime/Engine/Private/Materials/MaterialShared.cpp
--- a/Engine/Source/Runtime/Engine/Private/Materials/MaterialShared.cpp
+++ b/Engine/Source/Runtime/Engine/Private/Materials/MaterialShared.cpp
@@ -2564,7 +2564,8 @@ TShaderRef<FShader> FMaterial::GetShader(FMeshMaterialShaderType* ShaderType, FV
 
 bool FMaterial::TryGetShaders(const FMaterialShaderTypes& InTypes, const FVertexFactoryType* InVertexFactoryType, FMaterialShaders& OutShaders) const
 {
-	TRACE_CPUPROFILER_EVENT_SCOPE(FMaterial::TryGetShaders);
+	// TRACE_CPUPROFILER_EVENT_SCOPE(FMaterial::TryGetShaders); // <-- disabled by default due to verbosity (hundreds of calls per frame)
+
 	static const auto* CVarShaderPipelines = IConsoleManager::Get().FindTConsoleVariableDataInt(TEXT("r.ShaderPipelines"));
 	const bool bIsInGameThread = IsInGameThread();
 	const FMaterialShaderMap* ShaderMap = bIsInGameThread ? GameThreadShaderMap : RenderingThreadShaderMap;
