Fixed use-after-free in CreateRayTracingBuffer() due to temporary string pointer being held in FRHIResourceCreateInfo.

diff --git a/Engine/Source/Runtime/D3D12RHI/Private/D3D12RayTracing.cpp b/Engine/Source/Runtime/D3D12RHI/Private/D3D12RayTracing.cpp
--- a/Engine/Source/Runtime/D3D12RHI/Private/D3D12RayTracing.cpp
+++ b/Engine/Source/Runtime/D3D12RHI/Private/D3D12RayTracing.cpp
@@ -2396,7 +2396,8 @@ static TRefCountPtr<FD3D12Buffer> CreateRayTracingBuffer(FD3D12Adapter* Adapter,
 
 	TRefCountPtr<FD3D12Buffer> Result;
 
-	FRHIResourceCreateInfo CreateInfo(*DebugName.ToString());
+	FString DebugNameString = DebugName.ToString();
+	FRHIResourceCreateInfo CreateInfo(*DebugNameString);
 	CreateInfo.GPUMask = FRHIGPUMask::FromIndex(GPUIndex);
 
 	ID3D12ResourceAllocator* ResourceAllocator = nullptr;
