Require MaxVertices to be provided per ray tracing geometry section when creating FRHIRayTracingGeometry objects.

diff --git a/Engine/Plugins/Experimental/GPULightmass/Source/GPULightmass/Private/Scene/Scene.cpp b/Engine/Plugins/Experimental/GPULightmass/Source/GPULightmass/Private/Scene/Scene.cpp
--- a/Engine/Plugins/Experimental/GPULightmass/Source/GPULightmass/Private/Scene/Scene.cpp
+++ b/Engine/Plugins/Experimental/GPULightmass/Source/GPULightmass/Private/Scene/Scene.cpp
@@ -1542,6 +1542,7 @@ void FScene::AddGeometryInstanceFromComponent(ULandscapeComponent* InComponent)
 					Segment.VertexBuffer = nullptr;
 					Segment.VertexBufferStride = sizeof(FVector);
 					Segment.VertexBufferElementType = VET_Float3;
+					Segment.MaxVertices = NumPrimitives*3;
 					Segment.NumPrimitives = NumPrimitives;
 					GeometryInitializer.Segments.Add(Segment);
 
diff --git a/Engine/Plugins/Experimental/GeometryCache/Source/GeometryCache/Private/GeometryCacheSceneProxy.cpp b/Engine/Plugins/Experimental/GeometryCache/Source/GeometryCache/Private/GeometryCacheSceneProxy.cpp
--- a/Engine/Plugins/Experimental/GeometryCache/Source/GeometryCache/Private/GeometryCacheSceneProxy.cpp
+++ b/Engine/Plugins/Experimental/GeometryCache/Source/GeometryCache/Private/GeometryCacheSceneProxy.cpp
@@ -167,6 +167,7 @@ FGeometryCacheSceneProxy::FGeometryCacheSceneProxy(UGeometryCacheComponent* Comp
 							Segment.FirstPrimitive = BatchInfo.StartIndex / 3;
 							Segment.NumPrimitives = BatchInfo.NumTriangles;
 							Segment.VertexBuffer = Section->PositionBuffers[PositionBufferIndex].VertexBufferRHI;
+							Segment.MaxVertices = Section->PositionBuffers[PositionBufferIndex].GetSizeInBytes() / Segment.VertexBufferStride; // conservative estimate
 							Segments.Add(Segment);
 							Initializer.TotalPrimitiveCount += BatchInfo.NumTriangles;
 						}
@@ -604,6 +605,7 @@ void FGeometryCacheSceneProxy::UpdateAnimation(float NewTime, bool bNewLooping,
 					Segment.FirstPrimitive = BatchInfo.StartIndex / 3;
 					Segment.NumPrimitives = BatchInfo.NumTriangles;
 					Segment.VertexBuffer = Section->PositionBuffers[PositionBufferIndex].VertexBufferRHI;
+					Segment.MaxVertices = Section->PositionBuffers[PositionBufferIndex].GetSizeInBytes() / Segment.VertexBufferStride; // conservative estimate
 
 					Segments.Add(Segment);
 					Section->RayTracingGeometry.Initializer.TotalPrimitiveCount += BatchInfo.NumTriangles;
diff --git a/Engine/Plugins/Runtime/HairStrands/Source/HairStrandsCore/Private/HairStrandsRendering.cpp b/Engine/Plugins/Runtime/HairStrands/Source/HairStrandsCore/Private/HairStrandsRendering.cpp
--- a/Engine/Plugins/Runtime/HairStrands/Source/HairStrandsCore/Private/HairStrandsRendering.cpp
+++ b/Engine/Plugins/Runtime/HairStrands/Source/HairStrandsCore/Private/HairStrandsRendering.cpp
@@ -1154,6 +1154,7 @@ static void BuildHairAccelerationStructure_Strands(FRHICommandList& RHICmdList,
 	Segment.VertexBuffer = PositionBuffer;
 	Segment.VertexBufferStride = FHairStrandsRaytracingFormat::SizeInByte;
 	Segment.VertexBufferElementType = FHairStrandsRaytracingFormat::VertexElementType;
+	Segment.MaxVertices = RaytracingVertexCount;
 	Segment.NumPrimitives = RaytracingVertexCount / 3;
 	Initializer.Segments.Add(Segment);
 
@@ -1188,6 +1189,7 @@ static void BuildHairAccelerationStructure_Cards(FRHICommandList& RHICmdList,
 	Segment.VertexBufferStride = FHairCardsPositionFormat::SizeInByte;
 	Segment.VertexBufferElementType = FHairCardsPositionFormat::VertexElementType;
 	Segment.NumPrimitives = RestResource->PrimitiveCount;
+	Segment.MaxVertices = RestResource->VertexCount;
 	Initializer.Segments.Add(Segment);
 
 	OutRayTracingGeometry->SetInitializer(Initializer);
@@ -1221,6 +1223,7 @@ static void BuildHairAccelerationStructure_Meshes(FRHICommandList& RHICmdList,
 	Segment.VertexBufferStride = FHairCardsPositionFormat::SizeInByte;
 	Segment.VertexBufferElementType = FHairCardsPositionFormat::VertexElementType;
 	Segment.NumPrimitives = RestResource->PrimitiveCount;
+	Segment.MaxVertices = RestResource->VertexCount;
 	Initializer.Segments.Add(Segment);
 
 	OutRayTracingGeometry->SetInitializer(Initializer);
diff --git a/Engine/Plugins/Runtime/ProceduralMeshComponent/Source/ProceduralMeshComponent/Private/ProceduralMeshComponent.cpp b/Engine/Plugins/Runtime/ProceduralMeshComponent/Source/ProceduralMeshComponent/Private/ProceduralMeshComponent.cpp
--- a/Engine/Plugins/Runtime/ProceduralMeshComponent/Source/ProceduralMeshComponent/Private/ProceduralMeshComponent.cpp
+++ b/Engine/Plugins/Runtime/ProceduralMeshComponent/Source/ProceduralMeshComponent/Private/ProceduralMeshComponent.cpp
@@ -197,6 +197,7 @@ public:
 						FRayTracingGeometrySegment Segment;
 						Segment.VertexBuffer = NewSection->VertexBuffers.PositionVertexBuffer.VertexBufferRHI;
 						Segment.NumPrimitives = NewSection->RayTracingGeometry.Initializer.TotalPrimitiveCount;
+						Segment.MaxVertices = NewSection->VertexBuffers.PositionVertexBuffer.GetNumVertices();
 						NewSection->RayTracingGeometry.Initializer.Segments.Add(Segment);
 
 						//#dxr_todo: add support for segments?
diff --git a/Engine/Source/Runtime/D3D12RHI/Private/D3D12RayTracing.cpp b/Engine/Source/Runtime/D3D12RHI/Private/D3D12RayTracing.cpp
--- a/Engine/Source/Runtime/D3D12RHI/Private/D3D12RayTracing.cpp
+++ b/Engine/Source/Runtime/D3D12RHI/Private/D3D12RayTracing.cpp
@@ -2516,6 +2516,10 @@ FD3D12RayTracingGeometry::FD3D12RayTracingGeometry(FD3D12Adapter* Adapter, const
 		checkf(Segment.VertexBufferStride, TEXT("Position vertex buffer is required for ray tracing geometry"));
 		checkf(Segment.VertexBufferStride % 4 == 0, TEXT("Position vertex buffer stride must be aligned to 4 bytes for ByteAddressBuffer loads to work"));
 
+		checkf(Segment.MaxVertices != 0 || Segment.NumPrimitives == 0,
+			TEXT("FRayTracingGeometrySegment.MaxVertices for '%s' must contain number of positions in the vertex buffer or maximum index buffer value+1 if index buffer is provided."),
+			*DebugName.ToString());
+
 		if (Initializer.GeometryType == RTGT_Triangles)
 		{
 			checkf(Segment.VertexBufferElementType == VET_Float3
@@ -2554,11 +2558,6 @@ FD3D12RayTracingGeometry::FD3D12RayTracingGeometry(FD3D12Adapter* Adapter, const
 			Desc.Flags |= D3D12_RAYTRACING_GEOMETRY_FLAG_NO_DUPLICATE_ANYHIT_INVOCATION;
 		}
 
-		// Conservative estimate of the maximum number of elements in this VB.
-		// Real used number will depend on the index buffer and is not available here right now.
-		// #dxr_todo: Add explicit vertex count to FRayTracingGeometrySegment.
-		const uint32 MaxSegmentVertices = (Segment.VertexBuffer->GetSize() - Segment.VertexBufferOffset) / Segment.VertexBufferStride;
-
 		switch (GeometryType)
 		{
 		case D3D12_RAYTRACING_GEOMETRY_TYPE_TRIANGLES:
@@ -2588,14 +2587,14 @@ FD3D12RayTracingGeometry::FD3D12RayTracingGeometry(FD3D12Adapter* Adapter, const
 				// We conservatively set this to 32 bit to allocate acceleration structure memory.
 				Desc.Triangles.IndexFormat = DXGI_FORMAT_R32_UINT;
 				Desc.Triangles.IndexCount = Segment.NumPrimitives * IndicesPerPrimitive;
-				Desc.Triangles.VertexCount = MaxSegmentVertices;
+				Desc.Triangles.VertexCount = Segment.MaxVertices;
 			}
 			else
 			{
 				// Non-indexed geometry
 				checkf(Segments.Num() == 1, TEXT("Non-indexed geometry with multiple segments is not implemented."));
 				Desc.Triangles.IndexFormat = DXGI_FORMAT_UNKNOWN;
-				Desc.Triangles.VertexCount = FMath::Min<uint32>(MaxSegmentVertices, TotalPrimitiveCount * 3);
+				Desc.Triangles.VertexCount = FMath::Min<uint32>(Segment.MaxVertices, TotalPrimitiveCount * 3);
 			}
 			break;
 
@@ -2832,11 +2831,6 @@ void FD3D12RayTracingGeometry::BuildAccelerationStructure(FD3D12CommandContext&
 		const FHitGroupSystemParameters& SystemParameters = HitGroupSystemParametersForThisGPU[SegmentIndex];
 
 		FD3D12Buffer* VertexBuffer = CommandContext.RetrieveObject<FD3D12Buffer>(Segment.VertexBuffer.GetReference());
-		
-		// Conservative estimate of the maximum number of elements in this VB.
-		// Real used number will depend on the index buffer and is not available here right now.
-		// #dxr_todo: Add explicit vertex count to FRayTracingGeometrySegment.
-		const uint32 MaxSegmentVertices = (VertexBuffer->ResourceLocation.GetSize() - Segment.VertexBufferOffset) / Segment.VertexBufferStride;
 
 		switch (GeometryType)
 		{
@@ -2873,7 +2867,6 @@ void FD3D12RayTracingGeometry::BuildAccelerationStructure(FD3D12CommandContext&
 			if (IndexBuffer)
 			{
 				check(Desc.Triangles.IndexCount <= Segment.NumPrimitives * IndicesPerPrimitive);
-				check(Desc.Triangles.VertexCount <= MaxSegmentVertices);
 
 				Desc.Triangles.IndexFormat = (IndexStride == 4 ? DXGI_FORMAT_R32_UINT : DXGI_FORMAT_R16_UINT);
 				Desc.Triangles.IndexBuffer = SystemParameters.IndexBuffer + SystemParameters.RootConstants.IndexBufferOffsetInBytes;
@@ -2886,7 +2879,6 @@ void FD3D12RayTracingGeometry::BuildAccelerationStructure(FD3D12CommandContext&
 				check(Desc.Triangles.IndexFormat == DXGI_FORMAT_UNKNOWN);
 				check(Desc.Triangles.IndexCount == 0);
 				check(Desc.Triangles.IndexBuffer == D3D12_GPU_VIRTUAL_ADDRESS(0));
-				check(Desc.Triangles.VertexCount == FMath::Min<uint32>(MaxSegmentVertices, TotalPrimitiveCount * 3));
 			}
 
 			Desc.Triangles.VertexBuffer.StartAddress = SystemParameters.VertexBuffer;
@@ -3507,6 +3499,11 @@ void FD3D12CommandContext::RHIBuildAccelerationStructures(const TArrayView<const
 
 			for (int32 i = 0; i < P.Segments.Num(); ++i)
 			{
+				checkf(P.Segments[i].MaxVertices <= Geometry->Segments[i].MaxVertices,
+					TEXT("Maximum number of vertices in a segment (%u) must not be smaller than what was declared during FRHIRayTracingGeometry creation (%u), as this controls BLAS memory allocation."),
+					P.Segments[i].MaxVertices, Geometry->Segments[i].MaxVertices
+				);
+
 				Geometry->Segments[i].VertexBuffer            = P.Segments[i].VertexBuffer;
 				Geometry->Segments[i].VertexBufferElementType = P.Segments[i].VertexBufferElementType;
 				Geometry->Segments[i].VertexBufferStride      = P.Segments[i].VertexBufferStride;
diff --git a/Engine/Source/Runtime/Engine/Private/GPUSkinCache.cpp b/Engine/Source/Runtime/Engine/Private/GPUSkinCache.cpp
--- a/Engine/Source/Runtime/Engine/Private/GPUSkinCache.cpp
+++ b/Engine/Source/Runtime/Engine/Private/GPUSkinCache.cpp
@@ -1527,6 +1527,7 @@ void FGPUSkinCache::ProcessRayTracingGeometryToUpdate(
 				Segment.VertexBufferElementType = VET_Float3;
 				Segment.VertexBufferStride = VertexBufferStride;
 				Segment.VertexBufferOffset = 0;
+				Segment.MaxVertices = Section.GetNumVertices();
 				Segment.FirstPrimitive = Section.BaseIndex / 3;
 				Segment.NumPrimitives = Section.NumTriangles;
 				Segment.bEnabled = !Section.bDisabled;
diff --git a/Engine/Source/Runtime/Engine/Private/SkeletalRenderCPUSkin.cpp b/Engine/Source/Runtime/Engine/Private/SkeletalRenderCPUSkin.cpp
--- a/Engine/Source/Runtime/Engine/Private/SkeletalRenderCPUSkin.cpp
+++ b/Engine/Source/Runtime/Engine/Private/SkeletalRenderCPUSkin.cpp
@@ -404,6 +404,7 @@ void FSkeletalMeshObjectCPUSkin::FSkeletalMeshObjectLOD::InitResources(FSkelMesh
 					Segment.VertexBufferStride = VertexBufferStride;
 					Segment.VertexBufferOffset = 0;
 					Segment.VertexBufferElementType = VET_Float3;
+					Segment.MaxVertices = Section.GetNumVertices();
 					Segment.FirstPrimitive = Section.BaseIndex / 3;
 					Segment.NumPrimitives = Section.NumTriangles;
 					Segment.bEnabled = !Section.bDisabled;
diff --git a/Engine/Source/Runtime/Engine/Private/SkeletalRenderStatic.cpp b/Engine/Source/Runtime/Engine/Private/SkeletalRenderStatic.cpp
--- a/Engine/Source/Runtime/Engine/Private/SkeletalRenderStatic.cpp
+++ b/Engine/Source/Runtime/Engine/Private/SkeletalRenderStatic.cpp
@@ -90,6 +90,7 @@ void FSkeletalMeshObjectStatic::InitResources(USkinnedMeshComponent* InMeshCompo
 								Segment.VertexBufferElementType = VET_Float3;
 								Segment.VertexBufferOffset = 0;
 								Segment.VertexBufferStride = VertexBufferStride;
+								Segment.MaxVertices = Section.GetNumVertices();
 								Segment.FirstPrimitive = Section.BaseIndex / 3;
 								Segment.NumPrimitives = Section.NumTriangles;
 								Segment.bEnabled = !Section.bDisabled;
diff --git a/Engine/Source/Runtime/Engine/Private/StaticMesh.cpp b/Engine/Source/Runtime/Engine/Private/StaticMesh.cpp
--- a/Engine/Source/Runtime/Engine/Private/StaticMesh.cpp
+++ b/Engine/Source/Runtime/Engine/Private/StaticMesh.cpp
@@ -1290,6 +1290,7 @@ void FStaticMeshLODResources::InitResources(UStaticMesh* Parent)
 					Segment.VertexBufferElementType = VET_Float3;
 					Segment.VertexBufferStride = VertexBuffers.PositionVertexBuffer.GetStride();
 					Segment.VertexBufferOffset = 0;
+					Segment.MaxVertices = VertexBuffers.PositionVertexBuffer.GetNumVertices();
 					Segment.FirstPrimitive = Section.FirstIndex / 3;
 					Segment.NumPrimitives = Section.NumTriangles;
 					Segment.bEnabled = Section.bVisibleInRayTracing;
diff --git a/Engine/Source/Runtime/Landscape/Private/LandscapeRender.cpp b/Engine/Source/Runtime/Landscape/Private/LandscapeRender.cpp
--- a/Engine/Source/Runtime/Landscape/Private/LandscapeRender.cpp
+++ b/Engine/Source/Runtime/Landscape/Private/LandscapeRender.cpp
@@ -1503,6 +1503,7 @@ void FLandscapeComponentSceneProxy::CreateRenderThreadResources()
 				Segment.VertexBuffer = nullptr;
 				Segment.VertexBufferStride = sizeof(FVector);
 				Segment.VertexBufferElementType = VET_Float3;
+				Segment.MaxVertices = FMath::Square(SubsectionSizeVerts);
 				Initializer.Segments.Add(Segment);
 				SectionRayTracingStates[SubSectionIdx].Geometry.SetInitializer(Initializer);
 				SectionRayTracingStates[SubSectionIdx].Geometry.InitResource();
diff --git a/Engine/Source/Runtime/RHI/Public/DynamicRHI.h b/Engine/Source/Runtime/RHI/Public/DynamicRHI.h
--- a/Engine/Source/Runtime/RHI/Public/DynamicRHI.h
+++ b/Engine/Source/Runtime/RHI/Public/DynamicRHI.h
@@ -157,6 +157,10 @@ public:
 	// Must be equal or greater than the size of the position vector.
 	LAYOUT_FIELD_INITIALIZED(uint32, VertexBufferStride, 12);
 
+	// Number of vertices (positions) in VertexBuffer.
+	// If an index buffer is present, this must be at least the maximum index value in the index buffer + 1.
+	LAYOUT_FIELD_INITIALIZED(uint32, MaxVertices, 0);
+
 	// Primitive range for this segment.
 	LAYOUT_FIELD_INITIALIZED(uint32, FirstPrimitive, 0);
 	LAYOUT_FIELD_INITIALIZED(uint32, NumPrimitives, 0);
diff --git a/Engine/Source/Runtime/Renderer/Private/RayTracing/RayTracingDynamicGeometry.cpp b/Engine/Source/Runtime/Renderer/Private/RayTracing/RayTracingDynamicGeometry.cpp
--- a/Engine/Source/Runtime/Renderer/Private/RayTracing/RayTracingDynamicGeometry.cpp
+++ b/Engine/Source/Runtime/Renderer/Private/RayTracing/RayTracingDynamicGeometry.cpp
@@ -278,6 +278,7 @@ void FRayTracingDynamicGeometryCollection::AddDynamicMeshBatchForGeometryUpdate(
 		Geometry.Initializer.Segments.Empty();
 		FRayTracingGeometrySegment Segment;
 		Segment.NumPrimitives = UpdateParams.NumTriangles;
+		Segment.MaxVertices = UpdateParams.NumVertices;
 		Geometry.Initializer.Segments.Add(Segment);
 		bRefit = false;
 	}
