Fixing root signature flags on devices that don't support mesh shaders.

diff --git a/Engine/Source/Runtime/D3D12RHI/Private/D3D12RootSignature.cpp b/Engine/Source/Runtime/D3D12RHI/Private/D3D12RootSignature.cpp
--- a/Engine/Source/Runtime/D3D12RHI/Private/D3D12RootSignature.cpp
+++ b/Engine/Source/Runtime/D3D12RHI/Private/D3D12RootSignature.cpp
@@ -108,9 +108,9 @@ FORCEINLINE D3D12_ROOT_SIGNATURE_FLAGS GetD3D12RootSignatureDenyFlag(EShaderVisi
 		return D3D12_ROOT_SIGNATURE_FLAG_DENY_PIXEL_SHADER_ROOT_ACCESS;
 #if PLATFORM_SUPPORTS_MESH_SHADERS
 	case SV_Mesh:
-		return D3D12_ROOT_SIGNATURE_FLAG_DENY_MESH_SHADER_ROOT_ACCESS;
+		return GRHISupportsMeshShaders ? D3D12_ROOT_SIGNATURE_FLAG_DENY_MESH_SHADER_ROOT_ACCESS : D3D12_ROOT_SIGNATURE_FLAG_NONE;
 	case SV_Amplification:
-		return D3D12_ROOT_SIGNATURE_FLAG_DENY_AMPLIFICATION_SHADER_ROOT_ACCESS;
+		return GRHISupportsMeshShaders ? D3D12_ROOT_SIGNATURE_FLAG_DENY_AMPLIFICATION_SHADER_ROOT_ACCESS : D3D12_ROOT_SIGNATURE_FLAG_NONE;
 #endif
 	case SV_All:
 		return D3D12_ROOT_SIGNATURE_FLAG_NONE;
