Temporary increase of maximum alloc size in d3d12 upload pool.
This is to avoid hitting slow path in nanite scattered uploads.
Incoming fix will slice those uploads and revert maximum alloc size.

diff --git a/Engine/Source/Runtime/D3D12RHI/Private/D3D12RHIPrivate.h b/Engine/Source/Runtime/D3D12RHI/Private/D3D12RHIPrivate.h
--- a/Engine/Source/Runtime/D3D12RHI/Private/D3D12RHIPrivate.h
+++ b/Engine/Source/Runtime/D3D12RHI/Private/D3D12RHIPrivate.h
@@ -114,7 +114,7 @@ typedef FD3D12StateCacheBase FD3D12StateCache;
 #define READBACK_BUFFER_POOL_DEFAULT_POOL_SIZE (4 * 1024 * 1024)
 
 #define DEFAULT_CONTEXT_UPLOAD_POOL_SIZE (8 * 1024 * 1024)
-#define DEFAULT_CONTEXT_UPLOAD_POOL_MAX_ALLOC_SIZE (4 * 1024 * 1024)
+#define DEFAULT_CONTEXT_UPLOAD_POOL_MAX_ALLOC_SIZE (32 * 1024 * 1024)
 #define DEFAULT_CONTEXT_UPLOAD_POOL_ALIGNMENT (D3D12_CONSTANT_BUFFER_DATA_PLACEMENT_ALIGNMENT)
 #define TEXTURE_POOL_SIZE (8 * 1024 * 1024)
 
