Add d3d12 resource tracking on windows.
Use D3D12.DumpTrackedAllocations and D3D12.DumpTrackedAllocationCallstacks to dump info to log.
Submitting on behalf of Kenzo.

diff --git a/Engine/Source/Runtime/D3D12RHI/Private/D3D12Adapter.cpp b/Engine/Source/Runtime/D3D12RHI/Private/D3D12Adapter.cpp
--- a/Engine/Source/Runtime/D3D12RHI/Private/D3D12Adapter.cpp
+++ b/Engine/Source/Runtime/D3D12RHI/Private/D3D12Adapter.cpp
@@ -11,6 +11,7 @@ D3D12Adapter.cpp:D3D12 Adapter implementation.
 #include "Windows/AllowWindowsPlatformTypes.h"
 #if PLATFORM_WINDOWS
 #include "Windows/WindowsPlatformMisc.h"
+#include "Windows/WindowsPlatformStackWalk.h"
 #endif
 
 #if !PLATFORM_CPU_ARM_FAMILY && (PLATFORM_WINDOWS)
@@ -1244,3 +1245,145 @@ void FD3D12Adapter::BlockUntilIdle()
 		GetDevice(GPUIndex)->BlockUntilIdle();
 	}
 }
+
+
+void FD3D12Adapter::TrackAllocationData(FD3D12ResourceLocation* InAllocation, const D3D12_RESOURCE_ALLOCATION_INFO& InInfo)
+{
+#if TRACK_RESOURCE_ALLOCATIONS
+	FTrackedAllocationData AllocationData;
+	AllocationData.ResourceAllocation = InAllocation;
+	AllocationData.AllocationInfo = InInfo;
+	AllocationData.StackDepth = FPlatformStackWalk::CaptureStackBackTrace(&AllocationData.Stack[0], FTrackedAllocationData::MaxStackDepth);
+
+	FScopeLock Lock(&TrackedAllocationDataCS);
+	check(!TrackedAllocationData.Contains(InAllocation));
+	TrackedAllocationData.Add(InAllocation, AllocationData);
+#endif
+}
+
+void FD3D12Adapter::ReleaseTrackedAllocationData(FD3D12ResourceLocation* InAllocation)
+{
+#if TRACK_RESOURCE_ALLOCATIONS
+	FScopeLock Lock(&TrackedAllocationDataCS);
+	verify(TrackedAllocationData.Remove(InAllocation) == 1);
+#endif
+}
+
+#if TRACK_RESOURCE_ALLOCATIONS
+
+static FAutoConsoleCommandWithOutputDevice GDumpTrackedD3D12AllocationsCmd(
+	TEXT("D3D12.DumpTrackedAllocations"),
+	TEXT("Dump all tracked d3d12 resource allocations."),
+	FConsoleCommandWithOutputDeviceDelegate::CreateStatic([](FOutputDevice& OutputDevice)
+	{
+		FD3D12DynamicRHI::GetD3DRHI()->GetAdapter().DumpTrackedAllocationData(OutputDevice, false);
+	})
+);
+
+static FAutoConsoleCommandWithOutputDevice GDumpTrackedD3D12AllocationCallstacksCmd(
+	TEXT("D3D12.DumpTrackedAllocationCallstacks"),
+	TEXT("Dump all tracked d3d12 resource allocation callstacks."),
+	FConsoleCommandWithOutputDeviceDelegate::CreateStatic([](FOutputDevice& OutputDevice)
+	{
+		FD3D12DynamicRHI::GetD3DRHI()->GetAdapter().DumpTrackedAllocationData(OutputDevice, true);
+	})
+);
+
+void FD3D12Adapter::DumpTrackedAllocationData(FOutputDevice& OutputDevice, bool bWithCallstack)
+{
+	FScopeLock Lock(&TrackedAllocationDataCS);
+
+	TArray<FTrackedAllocationData> Allocations;
+	TrackedAllocationData.GenerateValueArray(Allocations);
+	Allocations.Sort([](const FTrackedAllocationData& InLHS, const FTrackedAllocationData& InRHS)
+		{
+			return InLHS.AllocationInfo.SizeInBytes > InRHS.AllocationInfo.SizeInBytes;
+		});
+
+	TArray<FTrackedAllocationData> BufferAllocations;	
+	TArray<FTrackedAllocationData> TextureAllocations;
+	uint64 TotalBufferAllocationSize = 0;
+	uint64 TotalTextureAllocationSize = 0;
+	for (FTrackedAllocationData& AllocationData : Allocations)
+	{
+		D3D12_RESOURCE_DESC ResourceDesc = AllocationData.ResourceAllocation->GetResource()->GetDesc();
+		if (ResourceDesc.Dimension == D3D12_RESOURCE_DIMENSION_BUFFER)
+		{
+			BufferAllocations.Add(AllocationData);
+			TotalBufferAllocationSize += AllocationData.AllocationInfo.SizeInBytes;
+		}
+		else
+		{
+			TextureAllocations.Add(AllocationData);
+			TotalTextureAllocationSize += AllocationData.AllocationInfo.SizeInBytes;
+		}
+	}
+
+	const size_t STRING_SIZE = 16 * 1024;
+	ANSICHAR StackTrace[STRING_SIZE];
+
+	FString OutputData;
+	OutputData += FString::Printf(TEXT("\n%d Tracked Texture Allocations (Total size: %4.3fMB):\n"), TextureAllocations.Num(), TotalTextureAllocationSize / (1024.0f * 1024));
+	for (const FTrackedAllocationData& AllocationData : TextureAllocations)
+	{
+		D3D12_RESOURCE_DESC ResourceDesc = AllocationData.ResourceAllocation->GetResource()->GetDesc();
+
+		FString Flags;
+		if (EnumHasAnyFlags(ResourceDesc.Flags, D3D12_RESOURCE_FLAG_ALLOW_RENDER_TARGET))
+		{
+			Flags += "RT";
+		}
+		else if (EnumHasAnyFlags(ResourceDesc.Flags, D3D12_RESOURCE_FLAG_ALLOW_DEPTH_STENCIL))
+		{
+			Flags += "DS";
+		}
+		if (EnumHasAnyFlags(ResourceDesc.Flags, D3D12_RESOURCE_FLAG_ALLOW_UNORDERED_ACCESS))
+		{
+			Flags += EnumHasAnyFlags(ResourceDesc.Flags, D3D12_RESOURCE_FLAG_ALLOW_RENDER_TARGET | D3D12_RESOURCE_FLAG_ALLOW_DEPTH_STENCIL) ? "|UAV" : "UAV";
+		}
+
+		OutputData += FString::Printf(TEXT("\tName: %s - Size: %3.3fMB - Width: %d - Height: %d - DepthOrArraySize: %d - MipLevels: %d - Flags: %s\n"), 
+			*AllocationData.ResourceAllocation->GetResource()->GetName().ToString(), 
+			AllocationData.AllocationInfo.SizeInBytes / (1024.0f * 1024),
+			ResourceDesc.Width, ResourceDesc.Height, ResourceDesc.DepthOrArraySize, ResourceDesc.MipLevels,
+			Flags.IsEmpty() ? TEXT("None") : *Flags);
+
+		if (bWithCallstack)
+		{
+			static uint32 EntriesToSkip = 3;
+			for (uint32 Index = EntriesToSkip; Index < AllocationData.StackDepth; ++Index)
+			{
+				StackTrace[0] = 0;
+				FPlatformStackWalk::ProgramCounterToHumanReadableString(Index, AllocationData.Stack[Index], StackTrace, STRING_SIZE, nullptr);
+				OutputData += FString::Printf(TEXT("\t\t%d %s\n"), Index - EntriesToSkip, ANSI_TO_TCHAR(StackTrace));
+			}
+		}
+	}
+
+	OutputData += FString::Printf(TEXT("\n\n%d Tracked Buffer Allocations (Total size: %4.3fMB):\n"), BufferAllocations.Num(), TotalBufferAllocationSize / (1024.0f * 1024));
+	for (const FTrackedAllocationData& AllocationData : BufferAllocations)
+	{
+		D3D12_RESOURCE_DESC ResourceDesc = AllocationData.ResourceAllocation->GetResource()->GetDesc();
+
+		OutputData += FString::Printf(TEXT("\tName: %s - Size: %3.3fMB - Width: %d - UAV: %s\n"), 
+			*AllocationData.ResourceAllocation->GetResource()->GetName().ToString(), 
+			AllocationData.AllocationInfo.SizeInBytes / (1024.0f * 1024),
+			ResourceDesc.Width,
+			EnumHasAnyFlags(ResourceDesc.Flags, D3D12_RESOURCE_FLAG_ALLOW_UNORDERED_ACCESS) ? TEXT("Yes") : TEXT("No"));
+
+		if (bWithCallstack)
+		{
+			static uint32 EntriesToSkip = 3;
+			for (uint32 Index = EntriesToSkip; Index < AllocationData.StackDepth; ++Index)
+			{
+				StackTrace[0] = 0;
+				FPlatformStackWalk::ProgramCounterToHumanReadableString(Index, AllocationData.Stack[Index], StackTrace, STRING_SIZE, nullptr);
+				OutputData += FString::Printf(TEXT("\t\t%d %s\n"), Index - EntriesToSkip, ANSI_TO_TCHAR(StackTrace));
+			}
+		}
+	}
+
+	OutputDevice.Log(OutputData);
+}
+
+#endif // TRACK_RESOURCE_ALLOCATIONS
diff --git a/Engine/Source/Runtime/D3D12RHI/Private/D3D12Adapter.h b/Engine/Source/Runtime/D3D12RHI/Private/D3D12Adapter.h
--- a/Engine/Source/Runtime/D3D12RHI/Private/D3D12Adapter.h
+++ b/Engine/Source/Runtime/D3D12RHI/Private/D3D12Adapter.h
@@ -328,6 +328,10 @@ public:
 
 	FORCEINLINE uint32 GetFrameCount() const { return FrameCounter; }
 
+	void TrackAllocationData(FD3D12ResourceLocation* InAllocation, const D3D12_RESOURCE_ALLOCATION_INFO& InInfo);
+	void ReleaseTrackedAllocationData(FD3D12ResourceLocation* InAllocation);
+	void DumpTrackedAllocationData(FOutputDevice& OutputDevice, bool bWithCallstack);
+
 #if D3D12_SUBMISSION_GAP_RECORDER
 	FD3D12SubmissionGapRecorder SubmissionGapRecorder;
 
@@ -426,6 +430,20 @@ protected:
 
 	uint32 FrameCounter;
 
+	/** Information about an allocated resource. */
+	struct FTrackedAllocationData
+	{
+		static const int32 MaxStackDepth = 30;
+
+		FD3D12ResourceLocation* ResourceAllocation;
+		D3D12_RESOURCE_ALLOCATION_INFO AllocationInfo;
+		uint32 StackDepth;
+		uint64 Stack[MaxStackDepth];
+	};
+	/** Tracked resource information. */
+	TMap<FD3D12ResourceLocation*, FTrackedAllocationData> TrackedAllocationData;
+	FCriticalSection TrackedAllocationDataCS;
+
 #if D3D12_SUBMISSION_GAP_RECORDER
 	TArray<uint64> StartOfSubmissionTimestamps;
 	TArray<uint64> EndOfSubmissionTimestamps;
diff --git a/Engine/Source/Runtime/D3D12RHI/Private/D3D12RHIPrivate.h b/Engine/Source/Runtime/D3D12RHI/Private/D3D12RHIPrivate.h
--- a/Engine/Source/Runtime/D3D12RHI/Private/D3D12RHIPrivate.h
+++ b/Engine/Source/Runtime/D3D12RHI/Private/D3D12RHIPrivate.h
@@ -95,6 +95,7 @@ typedef FD3D12StateCacheBase FD3D12StateCache;
 #define ENABLE_PLACED_RESOURCES 0 // Disabled due to a couple of NVidia bugs related to placed resources. Works fine on Intel
 #define NAME_OBJECTS !(UE_BUILD_SHIPPING || UE_BUILD_TEST)	// Name objects in all builds except shipping
 #define LOG_PSO_CREATES (0 && STATS)	// Logs Create Pipeline State timings (also requires STATS)
+#define TRACK_RESOURCE_ALLOCATIONS (PLATFORM_WINDOWS && !UE_BUILD_SHIPPING && !UE_BUILD_TEST)
 
 //@TODO: Improve allocator efficiency so we can increase these thresholds and improve performance
 // We measured 149MB of wastage in 340MB of allocations with DEFAULT_BUFFER_POOL_MAX_ALLOC_SIZE set to 512KB
@@ -113,7 +114,7 @@ typedef FD3D12StateCacheBase FD3D12StateCache;
 #define READBACK_BUFFER_POOL_MAX_ALLOC_SIZE (64 * 1024)
 #define READBACK_BUFFER_POOL_DEFAULT_POOL_SIZE (4 * 1024 * 1024)
 
-#define DEFAULT_CONTEXT_UPLOAD_POOL_SIZE (8 * 1024 * 1024)
+#define DEFAULT_CONTEXT_UPLOAD_POOL_SIZE (32 * 1024 * 1024)
 #define DEFAULT_CONTEXT_UPLOAD_POOL_MAX_ALLOC_SIZE (32 * 1024 * 1024)
 #define DEFAULT_CONTEXT_UPLOAD_POOL_ALIGNMENT (D3D12_CONSTANT_BUFFER_DATA_PLACEMENT_ALIGNMENT)
 #define TEXTURE_POOL_SIZE (8 * 1024 * 1024)
diff --git a/Engine/Source/Runtime/D3D12RHI/Private/D3D12Resources.cpp b/Engine/Source/Runtime/D3D12RHI/Private/D3D12Resources.cpp
--- a/Engine/Source/Runtime/D3D12RHI/Private/D3D12Resources.cpp
+++ b/Engine/Source/Runtime/D3D12RHI/Private/D3D12Resources.cpp
@@ -829,6 +829,16 @@ void FD3D12ResourceLocation::UpdateStandAloneStats(bool bIncrement)
 				INC_MEMORY_STAT_BY(STAT_D3D12TextureStandAloneAllocated, SizeInBytes);
 			}
 		}
+
+		// Track all committed resource allocations
+		if (bIncrement)
+		{
+			UnderlyingResource->GetParentDevice()->GetParentAdapter()->TrackAllocationData(this, Info);
+		}
+		else
+		{
+			UnderlyingResource->GetParentDevice()->GetParentAdapter()->ReleaseTrackedAllocationData(this);
+		}
 	}
 }
 
