Fixing static analysis warnings by using D3D12_CPU_DESCRIPTOR_HANDLE instead of the helper type CD3DX12_CPU_DESCRIPTOR_HANDLE.

diff --git a/Engine/Source/Runtime/D3D12RHI/Private/D3D12DescriptorCache.cpp b/Engine/Source/Runtime/D3D12RHI/Private/D3D12DescriptorCache.cpp
--- a/Engine/Source/Runtime/D3D12RHI/Private/D3D12DescriptorCache.cpp
+++ b/Engine/Source/Runtime/D3D12RHI/Private/D3D12DescriptorCache.cpp
@@ -320,9 +320,9 @@ void FD3D12DescriptorCache::SetUAVs(const FD3D12RootSignature* RootSignature, FD
 	uint32 FirstSlotIndex = HeapSlot;
 	HeapSlot += SlotsNeeded;
 
-	CD3DX12_CPU_DESCRIPTOR_HANDLE DestDescriptor(CurrentViewHeap->GetCPUSlotHandle(FirstSlotIndex));
-	CD3DX12_GPU_DESCRIPTOR_HANDLE BindDescriptor(CurrentViewHeap->GetGPUSlotHandle(FirstSlotIndex));
-	CD3DX12_CPU_DESCRIPTOR_HANDLE SrcDescriptors[MAX_UAVS];
+	D3D12_CPU_DESCRIPTOR_HANDLE DestDescriptor(CurrentViewHeap->GetCPUSlotHandle(FirstSlotIndex));
+	D3D12_GPU_DESCRIPTOR_HANDLE BindDescriptor(CurrentViewHeap->GetGPUSlotHandle(FirstSlotIndex));
+	D3D12_CPU_DESCRIPTOR_HANDLE SrcDescriptors[MAX_UAVS];
 
 	FD3D12CommandListHandle& CommandList = CmdContext->CommandListHandle;
 
@@ -479,7 +479,7 @@ void FD3D12DescriptorCache::SetSamplers(const FD3D12RootSignature* RootSignature
 		checkSlow(SlotsNeeded <= MAX_SAMPLERS);
 
 		// Fill heap slots
-		CD3DX12_CPU_DESCRIPTOR_HANDLE SrcDescriptors[MAX_SAMPLERS];
+		D3D12_CPU_DESCRIPTOR_HANDLE SrcDescriptors[MAX_SAMPLERS];
 		for (uint32 SlotIndex = 0; SlotIndex < SlotsNeeded; SlotIndex++)
 		{
 			if (Samplers[SlotIndex] != nullptr)
diff --git a/Engine/Source/Runtime/D3D12RHI/Private/D3D12DescriptorCache.h b/Engine/Source/Runtime/D3D12RHI/Private/D3D12DescriptorCache.h
--- a/Engine/Source/Runtime/D3D12RHI/Private/D3D12DescriptorCache.h
+++ b/Engine/Source/Runtime/D3D12RHI/Private/D3D12DescriptorCache.h
@@ -108,11 +108,11 @@ typedef FD3D12ConservativeMap<FD3D12SamplerArrayDesc, D3D12_GPU_DESCRIPTOR_HANDL
 template< uint32 CPUTableSize>
 struct FD3D12UniqueDescriptorTable
 {
-	FD3D12UniqueDescriptorTable() : GPUHandle({}) {};
-	FD3D12UniqueDescriptorTable(FD3D12SamplerArrayDesc KeyIn, CD3DX12_CPU_DESCRIPTOR_HANDLE* Table) : GPUHandle({})
+	FD3D12UniqueDescriptorTable() = default;
+	FD3D12UniqueDescriptorTable(FD3D12SamplerArrayDesc KeyIn, D3D12_CPU_DESCRIPTOR_HANDLE* Table)
 	{
 		FMemory::Memcpy(&Key, &KeyIn, sizeof(Key));//Memcpy to avoid alignement issues
-		FMemory::Memcpy(CPUTable, Table, Key.Count * sizeof(CD3DX12_CPU_DESCRIPTOR_HANDLE));
+		FMemory::Memcpy(CPUTable, Table, Key.Count * sizeof(D3D12_CPU_DESCRIPTOR_HANDLE));
 	}
 
 	FORCEINLINE uint32 GetTypeHash(const FD3D12UniqueDescriptorTable& Table)
@@ -120,11 +120,11 @@ struct FD3D12UniqueDescriptorTable
 		return FD3D12PipelineStateCache::HashData((void*)Table.Key.SamplerID, Table.Key.Count * sizeof(Table.Key.SamplerID[0]));
 	}
 
-	FD3D12SamplerArrayDesc Key;
-	CD3DX12_CPU_DESCRIPTOR_HANDLE CPUTable[MAX_SAMPLERS];
+	FD3D12SamplerArrayDesc Key{};
+	D3D12_CPU_DESCRIPTOR_HANDLE CPUTable[MAX_SAMPLERS]{};
 
 	// This will point to the table start in the global heap
-	D3D12_GPU_DESCRIPTOR_HANDLE GPUHandle;
+	D3D12_GPU_DESCRIPTOR_HANDLE GPUHandle{};
 };
 
 template<typename FD3D12UniqueDescriptorTable, bool bInAllowDuplicateKeys = false>
diff --git a/Engine/Source/Runtime/D3D12RHI/Private/D3D12View.h b/Engine/Source/Runtime/D3D12RHI/Private/D3D12View.h
--- a/Engine/Source/Runtime/D3D12RHI/Private/D3D12View.h
+++ b/Engine/Source/Runtime/D3D12RHI/Private/D3D12View.h
@@ -624,7 +624,7 @@ class FD3D12OfflineDescriptorManager;
 template <typename TDesc>
 class TD3D12ViewDescriptorHandle : public FD3D12DeviceChild
 {
-	CD3DX12_CPU_DESCRIPTOR_HANDLE Handle;
+	D3D12_CPU_DESCRIPTOR_HANDLE Handle;
 	uint32 Index;
 
 public:
@@ -651,7 +651,7 @@ public:
 	inline void CreateView(const TDesc& Desc, ID3D12Resource* Resource);
 	inline void CreateViewWithCounter(const TDesc& Desc, ID3D12Resource* Resource, ID3D12Resource* CounterResource);
 
-	inline const CD3DX12_CPU_DESCRIPTOR_HANDLE& GetHandle() const { return Handle; }
+	inline const D3D12_CPU_DESCRIPTOR_HANDLE& GetHandle() const { return Handle; }
 	inline uint32 GetIndex() const { return Index; }
 
 private:
@@ -789,7 +789,7 @@ public:
 	inline FD3D12Device*					GetParentDevice_Unsafe()	const { return Descriptor.GetParentDevice_Unsafe(); }
 	inline FD3D12ResourceLocation*			GetResourceLocation()		const { return ResourceLocation; }
 	inline const TDesc&						GetDesc()					const { checkf(bInitialized, TEXT("Uninitialized D3D12View size %d"), (uint32)sizeof(TDesc)); return Desc; }
-	inline CD3DX12_CPU_DESCRIPTOR_HANDLE	GetView()					const { checkf(bInitialized, TEXT("Uninitialized D3D12View size %d"), (uint32)sizeof(TDesc)); return Descriptor.GetHandle(); }
+	inline D3D12_CPU_DESCRIPTOR_HANDLE		GetView()					const { checkf(bInitialized, TEXT("Uninitialized D3D12View size %d"), (uint32)sizeof(TDesc)); return Descriptor.GetHandle(); }
 	inline uint32							GetDescriptorHeapIndex()	const { checkf(bInitialized, TEXT("Uninitialized D3D12View size %d"), (uint32)sizeof(TDesc)); return Descriptor.GetIndex(); }
 	inline FD3D12Resource*					GetResource()				const { checkf(bInitialized, TEXT("Uninitialized D3D12View size %d"), (uint32)sizeof(TDesc)); return Resource; }
 	inline FD3D12ResidencyHandle*			GetResidencyHandle()		const { checkf(bInitialized, TEXT("Uninitialized D3D12View size %d"), (uint32)sizeof(TDesc)); return ResidencyHandle; }
@@ -1008,7 +1008,7 @@ class FD3D12ConstantBufferView : public FD3D12DeviceChild
 			public:
 //protected:
 	/** The handle to the descriptor in the offline descriptor heap */
-	CD3DX12_CPU_DESCRIPTOR_HANDLE OfflineDescriptorHandle;
+	D3D12_CPU_DESCRIPTOR_HANDLE OfflineDescriptorHandle;
 
 	/** Index of the descriptor in the offline heap */
 	uint32 OfflineHeapIndex;
