Added API to compute ray tracing acceleration buffer sizes. Implemented only for D3D12 RHI initially.

diff --git a/Engine/Source/Runtime/D3D12RHI/Private/D3D12RHIPrivate.h b/Engine/Source/Runtime/D3D12RHI/Private/D3D12RHIPrivate.h
--- a/Engine/Source/Runtime/D3D12RHI/Private/D3D12RHIPrivate.h
+++ b/Engine/Source/Runtime/D3D12RHI/Private/D3D12RHIPrivate.h
@@ -470,6 +470,9 @@ public:
 
 #if D3D12_RHI_RAYTRACING
 
+	virtual FRayTracingAccelerationStructureSize RHICalcRayTracingSceneSize(uint32 MaxInstances, ERayTracingAccelerationStructureFlags Flags) final override;
+	virtual FRayTracingAccelerationStructureSize RHICalcRayTracingGeometrySize(const FRayTracingGeometryInitializer& Initializer) final override;
+
 	virtual FRayTracingGeometryRHIRef RHICreateRayTracingGeometry(const FRayTracingGeometryInitializer& Initializer) final override;
 	virtual FRayTracingSceneRHIRef RHICreateRayTracingScene(const FRayTracingSceneInitializer& Initializer) final override;
 	virtual FRayTracingShaderRHIRef RHICreateRayTracingShader(TArrayView<const uint8> Code, const FSHAHash& Hash, EShaderFrequency ShaderFrequency) final override;
diff --git a/Engine/Source/Runtime/D3D12RHI/Private/D3D12RayTracing.cpp b/Engine/Source/Runtime/D3D12RHI/Private/D3D12RayTracing.cpp
--- a/Engine/Source/Runtime/D3D12RHI/Private/D3D12RayTracing.cpp
+++ b/Engine/Source/Runtime/D3D12RHI/Private/D3D12RayTracing.cpp
@@ -168,6 +168,46 @@
 
 static_assert(sizeof(FD3D12ShaderIdentifier) == D3D12_SHADER_IDENTIFIER_SIZE_IN_BYTES, "Unexpected shader identifier size");
 
+static D3D12_RAYTRACING_ACCELERATION_STRUCTURE_BUILD_FLAGS TranslateRayTracingAccelerationStructureFlags(ERayTracingAccelerationStructureFlags Flags)
+{
+	uint32 Result = {};
+
+	auto HandleFlag = [&Flags, &Result](ERayTracingAccelerationStructureFlags Engine, D3D12_RAYTRACING_ACCELERATION_STRUCTURE_BUILD_FLAGS Native)
+	{
+		if (EnumHasAllFlags(Flags, Engine))
+		{
+			Result |= (uint32)Native;
+			EnumRemoveFlags(Flags, Engine);
+		}
+	};
+
+	HandleFlag(ERayTracingAccelerationStructureFlags::AllowUpdate, D3D12_RAYTRACING_ACCELERATION_STRUCTURE_BUILD_FLAG_ALLOW_UPDATE);
+	HandleFlag(ERayTracingAccelerationStructureFlags::AllowCompaction, D3D12_RAYTRACING_ACCELERATION_STRUCTURE_BUILD_FLAG_ALLOW_COMPACTION);
+	HandleFlag(ERayTracingAccelerationStructureFlags::FastTrace, D3D12_RAYTRACING_ACCELERATION_STRUCTURE_BUILD_FLAG_PREFER_FAST_TRACE);
+	HandleFlag(ERayTracingAccelerationStructureFlags::FastBuild, D3D12_RAYTRACING_ACCELERATION_STRUCTURE_BUILD_FLAG_PREFER_FAST_BUILD);
+	HandleFlag(ERayTracingAccelerationStructureFlags::MinimizeMemory, D3D12_RAYTRACING_ACCELERATION_STRUCTURE_BUILD_FLAG_MINIMIZE_MEMORY);
+
+	checkf(!EnumHasAnyFlags(Flags, Flags), TEXT("Some ERayTracingAccelerationStructureFlags entries were not handled"));
+
+	return D3D12_RAYTRACING_ACCELERATION_STRUCTURE_BUILD_FLAGS(Result);
+}
+
+static D3D12_RAYTRACING_GEOMETRY_TYPE TranslateRayTracingGeometryType(ERayTracingGeometryType GeometryType)
+{
+	switch (GeometryType)
+	{
+	case RTGT_Triangles:
+		return D3D12_RAYTRACING_GEOMETRY_TYPE_TRIANGLES;
+		break;
+	case RTGT_Procedural:
+		return D3D12_RAYTRACING_GEOMETRY_TYPE_PROCEDURAL_PRIMITIVE_AABBS;
+		break;
+	default:
+		checkf(false, TEXT("Unexpected ray tracing geometry type"));
+		return D3D12_RAYTRACING_GEOMETRY_TYPE(0);
+	}
+}
+
 struct FDXILLibrary
 {
 	// No copy assignment or move because FDXILLibrary points to internal struct memory
@@ -439,9 +479,9 @@
 {
 	uint32 GPUIndex = GetParentDevice()->GetGPUIndex();
 	check(InRTGeometry->AccelerationStructureBuffers[GPUIndex]);
-	check((InRTGeometry->BuildFlags & D3D12_RAYTRACING_ACCELERATION_STRUCTURE_BUILD_FLAG_ALLOW_COMPACTION) &&
-		(InRTGeometry->BuildFlags & D3D12_RAYTRACING_ACCELERATION_STRUCTURE_BUILD_FLAG_PREFER_FAST_TRACE) &&
-		!(InRTGeometry->BuildFlags & D3D12_RAYTRACING_ACCELERATION_STRUCTURE_BUILD_FLAG_ALLOW_UPDATE));
+	check(EnumHasAllFlags(InRTGeometry->BuildFlags, ERayTracingAccelerationStructureFlags::AllowCompaction) &&
+		EnumHasAllFlags(InRTGeometry->BuildFlags, ERayTracingAccelerationStructureFlags::FastTrace) &&
+		!EnumHasAnyFlags(InRTGeometry->BuildFlags, ERayTracingAccelerationStructureFlags::AllowUpdate));
 
 	FScopeLock Lock(&CS);
 	PendingRequests.Add(InRTGeometry);
@@ -2358,167 +2398,89 @@
 	// Note: RayTracingDescriptorHeapCache is destroyed in ~FD3D12Device, after all deferred deletion is processed
 }
 
-FRayTracingPipelineStateRHIRef FD3D12DynamicRHI::RHICreateRayTracingPipelineState(const FRayTracingPipelineStateInitializer& Initializer)
+FRayTracingAccelerationStructureSize FD3D12DynamicRHI::RHICalcRayTracingSceneSize(uint32 MaxInstances, ERayTracingAccelerationStructureFlags Flags)
 {
-	FD3D12Device* Device = GetAdapter().GetDevice(0); // All pipelines are created on the first node, as they may be used on any other linked GPU.
-	FD3D12RayTracingPipelineState* Result = new FD3D12RayTracingPipelineState(Device, Initializer);
-
-	return Result;
-}
-
-FRayTracingGeometryRHIRef FD3D12DynamicRHI::RHICreateRayTracingGeometry(const FRayTracingGeometryInitializer& Initializer)
-{
-	FD3D12Adapter& Adapter = GetAdapter();
-	return new FD3D12RayTracingGeometry(&Adapter, Initializer);
-}
-
-FRayTracingSceneRHIRef FD3D12DynamicRHI::RHICreateRayTracingScene(const FRayTracingSceneInitializer& Initializer)
-{
-	TRACE_CPUPROFILER_EVENT_SCOPE(CreateRayTracingScene);
+	D3D12_BUILD_RAYTRACING_ACCELERATION_STRUCTURE_INPUTS BuildInputs = {};
+	BuildInputs.Type = D3D12_RAYTRACING_ACCELERATION_STRUCTURE_TYPE_TOP_LEVEL;
+	BuildInputs.DescsLayout = D3D12_ELEMENTS_LAYOUT_ARRAY;
+	BuildInputs.NumDescs = MaxInstances;
+	BuildInputs.Flags = TranslateRayTracingAccelerationStructureFlags(Flags);
 
 	FD3D12Adapter& Adapter = GetAdapter();
 
-	return new FD3D12RayTracingScene(&Adapter, Initializer);
+	FRayTracingAccelerationStructureSize SizeInfo = {};
+	for (uint32 GPUIndex = 0; GPUIndex < GNumExplicitGPUsForRendering; ++GPUIndex)
+	{
+		ID3D12Device5* RayTracingDevice = Adapter.GetDevice(GPUIndex)->GetDevice5();
+		D3D12_RAYTRACING_ACCELERATION_STRUCTURE_PREBUILD_INFO PrebuildInfo = {};
+		RayTracingDevice->GetRaytracingAccelerationStructurePrebuildInfo(&BuildInputs, &PrebuildInfo);
+
+		SizeInfo.ResultSize = FMath::Max(SizeInfo.ResultSize, PrebuildInfo.ResultDataMaxSizeInBytes);
+		SizeInfo.BuildScratchSize  = FMath::Max(SizeInfo.BuildScratchSize, PrebuildInfo.ScratchDataSizeInBytes);
+		SizeInfo.UpdateScratchSize = FMath::Max(SizeInfo.UpdateScratchSize, PrebuildInfo.UpdateScratchDataSizeInBytes);
+	}
+
+	return SizeInfo;
 }
 
-FBufferRHIRef FD3D12RayTracingGeometry::NullTransformBuffer;
+static ERayTracingAccelerationStructureFlags GetRayTracingAccelerationStructureBuildFlags(const FRayTracingGeometryInitializer& Initializer)
+{
+	ERayTracingAccelerationStructureFlags BuildFlags = ERayTracingAccelerationStructureFlags::None;
 
-enum class ERayTracingBufferType
-{
-	AccelerationStructure,
-	Scratch
-};
-
-static TRefCountPtr<FD3D12Buffer> CreateRayTracingBuffer(FD3D12Adapter* Adapter, uint32 GPUIndex, uint64 Size, ERayTracingBufferType Type, FName DebugName)
-{
-	check(DebugName.IsValid());
-	check(Size);
-
-	TRefCountPtr<FD3D12Buffer> Result;
-
-	FString DebugNameString = DebugName.ToString();
-	FRHIResourceCreateInfo CreateInfo(*DebugNameString);
-	CreateInfo.GPUMask = FRHIGPUMask::FromIndex(GPUIndex);
-
-	ID3D12ResourceAllocator* ResourceAllocator = nullptr;
-
-	D3D12_RESOURCE_DESC BufferDesc = CD3DX12_RESOURCE_DESC::Buffer(Size, D3D12_RESOURCE_FLAG_ALLOW_UNORDERED_ACCESS);
-
-	if (Type == ERayTracingBufferType::AccelerationStructure)
-	{
-		Result = Adapter->CreateRHIBuffer(
-			nullptr, BufferDesc, D3D12_RAYTRACING_ACCELERATION_STRUCTURE_BYTE_ALIGNMENT,
-			0, BufferDesc.Width, BUF_AccelerationStructure,
-			ED3D12ResourceStateMode::SingleState, ERHIAccess::BVHWrite,
-			CreateInfo, ResourceAllocator);
-	}
-	else if (Type == ERayTracingBufferType::Scratch)
-	{
-		Result = Adapter->CreateRHIBuffer(
-			nullptr, BufferDesc, D3D12_RAYTRACING_ACCELERATION_STRUCTURE_BYTE_ALIGNMENT,
-			0, BufferDesc.Width, BUF_UnorderedAccess,
-			ED3D12ResourceStateMode::SingleState, ERHIAccess::UAVMask,
-			CreateInfo, ResourceAllocator);
-	}
-	else
-	{
-		checkNoEntry();
-	}
-
-	// Elevates the raytracing acceleration structure heap priority, which may help performance / stability in low memory conditions
-	ID3D12Pageable* HeapResource = Result->GetResource()->GetPageable();
-	D3D12_RESIDENCY_PRIORITY HeapPriority = D3D12_RESIDENCY_PRIORITY_HIGH;
-	FD3D12Device* NodeDevice = Adapter->GetDevice(GPUIndex);
-	NodeDevice->GetDevice5()->SetResidencyPriority(1, &HeapResource, &HeapPriority);
-
-	return Result;
-}
-
-FD3D12RayTracingGeometry::FD3D12RayTracingGeometry(FD3D12Adapter* Adapter, const FRayTracingGeometryInitializer& Initializer)
-	: FD3D12AdapterChild(Adapter)
-{
-	INC_DWORD_STAT(STAT_D3D12RayTracingAllocatedBLAS);
-
-	DebugName = Initializer.DebugName.IsValid() ? Initializer.DebugName : FName(TEXT("BLAS"));
-	FMemory::Memzero(bHasPendingCompactionRequests);
-	FMemory::Memzero(bRegisteredAsRenameListener);
-
-	if(!FD3D12RayTracingGeometry::NullTransformBuffer.IsValid())
-	{
-		TResourceArray<float> NullTransformData;
-		NullTransformData.SetNumZeroed(12);
-
-		FRHIResourceCreateInfo CreateInfo(TEXT("NullTransformBuffer"));
-		CreateInfo.ResourceArray = &NullTransformData;
-
-		FD3D12RayTracingGeometry::NullTransformBuffer = RHICreateVertexBuffer(NullTransformData.GetResourceDataSize(), BUF_Static, CreateInfo);
-	}
-
-	IndexOffsetInBytes = Initializer.IndexBufferOffset;
-	TotalPrimitiveCount = Initializer.TotalPrimitiveCount;
-
-	switch (Initializer.GeometryType)
-	{
-	case RTGT_Triangles:
-		GeometryType = D3D12_RAYTRACING_GEOMETRY_TYPE_TRIANGLES;
-		break;
-	case RTGT_Procedural:
-		GeometryType = D3D12_RAYTRACING_GEOMETRY_TYPE_PROCEDURAL_PRIMITIVE_AABBS;
-		break;
-	default:
-		checkf(false, TEXT("Unexpected ray tracing geometry type"));
-	}
-
 	if (Initializer.bFastBuild)
 	{
-		BuildFlags = D3D12_RAYTRACING_ACCELERATION_STRUCTURE_BUILD_FLAG_PREFER_FAST_BUILD;
+		BuildFlags = ERayTracingAccelerationStructureFlags::FastBuild;
 	}
 	else
 	{
-		BuildFlags = D3D12_RAYTRACING_ACCELERATION_STRUCTURE_BUILD_FLAG_PREFER_FAST_TRACE;
+		BuildFlags = ERayTracingAccelerationStructureFlags::FastTrace;
 	}
 
 	if (Initializer.bAllowUpdate)
 	{
-		BuildFlags |= D3D12_RAYTRACING_ACCELERATION_STRUCTURE_BUILD_FLAG_ALLOW_UPDATE;
+		EnumAddFlags(BuildFlags, ERayTracingAccelerationStructureFlags::AllowUpdate);
 	}
 
 	if (!Initializer.bFastBuild && !Initializer.bAllowUpdate && GD3D12RayTracingAllowCompaction)
 	{
-		BuildFlags |= D3D12_RAYTRACING_ACCELERATION_STRUCTURE_BUILD_FLAG_ALLOW_COMPACTION;
+		EnumAddFlags(BuildFlags, ERayTracingAccelerationStructureFlags::AllowCompaction);
 	}
 
 	if (GRayTracingDebugForceBuildMode == 1)
 	{
-		BuildFlags |= D3D12_RAYTRACING_ACCELERATION_STRUCTURE_BUILD_FLAG_PREFER_FAST_BUILD;
-		BuildFlags &= ~D3D12_RAYTRACING_ACCELERATION_STRUCTURE_BUILD_FLAG_PREFER_FAST_TRACE;
+		EnumAddFlags(BuildFlags, ERayTracingAccelerationStructureFlags::FastBuild);
+		EnumRemoveFlags(BuildFlags, ERayTracingAccelerationStructureFlags::FastTrace);
 	}
 	else if (GRayTracingDebugForceBuildMode == 2)
 	{
-		BuildFlags |= D3D12_RAYTRACING_ACCELERATION_STRUCTURE_BUILD_FLAG_PREFER_FAST_TRACE;
-		BuildFlags &= ~D3D12_RAYTRACING_ACCELERATION_STRUCTURE_BUILD_FLAG_PREFER_FAST_BUILD;
+		EnumAddFlags(BuildFlags, ERayTracingAccelerationStructureFlags::FastTrace);
+		EnumRemoveFlags(BuildFlags, ERayTracingAccelerationStructureFlags::FastBuild);
 	}
 
-	checkf(Initializer.Segments.Num() > 0, TEXT("Ray tracing geometry must be initialized with at least one segment."));
-	Segments = TArray<FRayTracingGeometrySegment>(Initializer.Segments.GetData(), Initializer.Segments.Num());
+	return BuildFlags;
+}
 
-	IndexStride = RHIIndexBuffer ? RHIIndexBuffer->GetStride() : 0; // stride 0 means implicit triangle list for non-indexed geometry
+void TranslateRayTracingGeometryDescs(const FRayTracingGeometryInitializer& Initializer, TArrayView<D3D12_RAYTRACING_GEOMETRY_DESC> Output)
+{
+	check(Output.Num() == Initializer.Segments.Num());
 
+	D3D12_RAYTRACING_GEOMETRY_TYPE GeometryType = TranslateRayTracingGeometryType(Initializer.GeometryType);
+
 	static constexpr uint32 IndicesPerPrimitive = 3; // Only triangle meshes are supported
 
-	GeometryDescs.Reserve(Segments.Num());
-
 	uint32 ComputedPrimitiveCountForValidation = 0;
 
-	for (const FRayTracingGeometrySegment& Segment : Segments)
+	for (int32 SegmentIndex = 0; SegmentIndex < Initializer.Segments.Num(); ++SegmentIndex)
 	{
-		checkf(Segment.VertexBuffer, TEXT("Position vertex buffer is required for ray tracing geometry"));
-		checkf(Segment.VertexBufferStride, TEXT("Position vertex buffer is required for ray tracing geometry"));
-		checkf(Segment.VertexBufferStride % 4 == 0, TEXT("Position vertex buffer stride must be aligned to 4 bytes for ByteAddressBuffer loads to work"));
+		const FRayTracingGeometrySegment& Segment = Initializer.Segments[SegmentIndex];
 
+		checkf(Segment.VertexBuffer, TEXT("Position vertex buffer is required for ray tracing geometry."));
+		checkf(Segment.VertexBufferStride, TEXT("Non-zero position vertex buffer stride is required."));
+		checkf(Segment.VertexBufferStride % 4 == 0, TEXT("Position vertex buffer stride must be aligned to 4 bytes for ByteAddressBuffer loads to work."));
+
 		checkf(Segment.MaxVertices != 0 || Segment.NumPrimitives == 0,
 			TEXT("FRayTracingGeometrySegment.MaxVertices for '%s' must contain number of positions in the vertex buffer or maximum index buffer value+1 if index buffer is provided."),
-			*DebugName.ToString());
+			*Initializer.DebugName.ToString());
 
 		if (Initializer.GeometryType == RTGT_Triangles)
 		{
@@ -2531,16 +2493,15 @@
 			checkf(Segment.VertexBufferStride >= (2 * sizeof(FVector)), TEXT("Procedural geometry vertex buffer must contain at least 2xFloat3 that defines 3D bounding boxes of primitives."));
 		}
 
-		ComputedPrimitiveCountForValidation += Segment.NumPrimitives;
 		check(Segment.FirstPrimitive + Segment.NumPrimitives <= Initializer.TotalPrimitiveCount);
 
 		if (Initializer.IndexBuffer)
 		{
+			uint32 IndexStride = Initializer.IndexBuffer->GetStride();
 			check(Initializer.IndexBuffer->GetSize() >=
-				(Segment.FirstPrimitive + Segment.NumPrimitives) * IndicesPerPrimitive * IndexStride + IndexOffsetInBytes);
+				(Segment.FirstPrimitive + Segment.NumPrimitives) * IndicesPerPrimitive * IndexStride + Initializer.IndexBufferOffset);
 		}
 
-
 		D3D12_RAYTRACING_GEOMETRY_DESC Desc = {};
 
 		Desc.Flags = D3D12_RAYTRACING_GEOMETRY_FLAG_NONE;
@@ -2592,9 +2553,9 @@
 			else
 			{
 				// Non-indexed geometry
-				checkf(Segments.Num() == 1, TEXT("Non-indexed geometry with multiple segments is not implemented."));
+				checkf(Initializer.Segments.Num() == 1, TEXT("Non-indexed geometry with multiple segments is not implemented."));
 				Desc.Triangles.IndexFormat = DXGI_FORMAT_UNKNOWN;
-				Desc.Triangles.VertexCount = FMath::Min<uint32>(Segment.MaxVertices, TotalPrimitiveCount * 3);
+				Desc.Triangles.VertexCount = FMath::Min<uint32>(Segment.MaxVertices, Initializer.TotalPrimitiveCount * 3);
 			}
 			break;
 
@@ -2607,38 +2568,184 @@
 			break;
 		}
 
-		GeometryDescs.Add(Desc);
+		ComputedPrimitiveCountForValidation += Segment.NumPrimitives;
+
+
+		Output[SegmentIndex] = Desc;
 	}
 
 	check(ComputedPrimitiveCountForValidation == Initializer.TotalPrimitiveCount);
+}
 
-	RHIIndexBuffer = Initializer.IndexBuffer;
-	SetDirty(FRHIGPUMask::All(), true);
+FRayTracingAccelerationStructureSize FD3D12DynamicRHI::RHICalcRayTracingGeometrySize(const FRayTracingGeometryInitializer& Initializer)
+{
+	FRayTracingAccelerationStructureSize SizeInfo = {};
+	
+	ERayTracingAccelerationStructureFlags BuildFlags = GetRayTracingAccelerationStructureBuildFlags(Initializer);
 
 	D3D12_BUILD_RAYTRACING_ACCELERATION_STRUCTURE_INPUTS PrebuildDescInputs = {};
+
+	TArray<D3D12_RAYTRACING_GEOMETRY_DESC, TInlineAllocator<32>> GeometryDescs;
+	GeometryDescs.SetNumUninitialized(Initializer.Segments.Num());
+	TranslateRayTracingGeometryDescs(Initializer, GeometryDescs);
+
+	D3D12_RAYTRACING_GEOMETRY_TYPE GeometryType = TranslateRayTracingGeometryType(Initializer.GeometryType);
 
 	PrebuildDescInputs.Type = D3D12_RAYTRACING_ACCELERATION_STRUCTURE_TYPE_BOTTOM_LEVEL;
 	PrebuildDescInputs.DescsLayout = D3D12_ELEMENTS_LAYOUT_ARRAY;
 	PrebuildDescInputs.NumDescs = GeometryDescs.Num();
 	PrebuildDescInputs.pGeometryDescs = GeometryDescs.GetData();
-	PrebuildDescInputs.Flags = D3D12_RAYTRACING_ACCELERATION_STRUCTURE_BUILD_FLAGS(BuildFlags);
+	PrebuildDescInputs.Flags = TranslateRayTracingAccelerationStructureFlags(BuildFlags);
+
+	FD3D12Adapter& Adapter = GetAdapter();
 
 	// Get maximum buffer sizes for all GPUs in the system
-	for (uint32 GPUIndex = 0; GPUIndex < MAX_NUM_GPUS && GPUIndex < GNumExplicitGPUsForRendering; ++GPUIndex)
+	for (uint32 GPUIndex = 0; GPUIndex < GNumExplicitGPUsForRendering; ++GPUIndex)
 	{
-		ID3D12Device5* RayTracingDevice = Adapter->GetDevice(GPUIndex)->GetDevice5();
+		ID3D12Device5* RayTracingDevice = Adapter.GetDevice(GPUIndex)->GetDevice5();
 		D3D12_RAYTRACING_ACCELERATION_STRUCTURE_PREBUILD_INFO PrebuildInfo = {};
 		RayTracingDevice->GetRaytracingAccelerationStructurePrebuildInfo(&PrebuildDescInputs, &PrebuildInfo);
 
-		AccelerationStructureSize = FMath::Max(AccelerationStructureSize, PrebuildInfo.ResultDataMaxSizeInBytes);
-		BuildScratchBufferSize = FMath::Max(BuildScratchBufferSize, PrebuildInfo.ScratchDataSizeInBytes);
-		UpdateScratchBufferSize = FMath::Max(UpdateScratchBufferSize, PrebuildInfo.UpdateScratchDataSizeInBytes);
+		SizeInfo.ResultSize = FMath::Max(SizeInfo.ResultSize, PrebuildInfo.ResultDataMaxSizeInBytes);
+		SizeInfo.BuildScratchSize = FMath::Max(SizeInfo.BuildScratchSize, PrebuildInfo.ScratchDataSizeInBytes);
+		SizeInfo.UpdateScratchSize = FMath::Max(SizeInfo.UpdateScratchSize, PrebuildInfo.UpdateScratchDataSizeInBytes);
 	}
 
+	return SizeInfo;
+}
+
+FRayTracingPipelineStateRHIRef FD3D12DynamicRHI::RHICreateRayTracingPipelineState(const FRayTracingPipelineStateInitializer& Initializer)
+{
+	FD3D12Device* Device = GetAdapter().GetDevice(0); // All pipelines are created on the first node, as they may be used on any other linked GPU.
+	FD3D12RayTracingPipelineState* Result = new FD3D12RayTracingPipelineState(Device, Initializer);
+
+	return Result;
+}
+
+FRayTracingGeometryRHIRef FD3D12DynamicRHI::RHICreateRayTracingGeometry(const FRayTracingGeometryInitializer& Initializer)
+{
+	FD3D12Adapter& Adapter = GetAdapter();
+	return new FD3D12RayTracingGeometry(&Adapter, Initializer);
+}
+
+FRayTracingSceneRHIRef FD3D12DynamicRHI::RHICreateRayTracingScene(const FRayTracingSceneInitializer& Initializer)
+{
+	TRACE_CPUPROFILER_EVENT_SCOPE(CreateRayTracingScene);
+
+	FD3D12Adapter& Adapter = GetAdapter();
+
+	return new FD3D12RayTracingScene(&Adapter, Initializer);
+}
+
+FBufferRHIRef FD3D12RayTracingGeometry::NullTransformBuffer;
+
+enum class ERayTracingBufferType
+{
+	AccelerationStructure,
+	Scratch
+};
+
+static TRefCountPtr<FD3D12Buffer> CreateRayTracingBuffer(FD3D12Adapter* Adapter, uint32 GPUIndex, uint64 Size, ERayTracingBufferType Type, FName DebugName)
+{
+	check(DebugName.IsValid());
+	check(Size);
+
+	TRefCountPtr<FD3D12Buffer> Result;
+
+	FString DebugNameString = DebugName.ToString();
+	FRHIResourceCreateInfo CreateInfo(*DebugNameString);
+	CreateInfo.GPUMask = FRHIGPUMask::FromIndex(GPUIndex);
+
+	ID3D12ResourceAllocator* ResourceAllocator = nullptr;
+
+	D3D12_RESOURCE_DESC BufferDesc = CD3DX12_RESOURCE_DESC::Buffer(Size, D3D12_RESOURCE_FLAG_ALLOW_UNORDERED_ACCESS);
+
+	if (Type == ERayTracingBufferType::AccelerationStructure)
+	{
+		Result = Adapter->CreateRHIBuffer(
+			nullptr, BufferDesc, D3D12_RAYTRACING_ACCELERATION_STRUCTURE_BYTE_ALIGNMENT,
+			0, BufferDesc.Width, BUF_AccelerationStructure,
+			ED3D12ResourceStateMode::SingleState, ERHIAccess::BVHWrite,
+			CreateInfo, ResourceAllocator);
+	}
+	else if (Type == ERayTracingBufferType::Scratch)
+	{
+		Result = Adapter->CreateRHIBuffer(
+			nullptr, BufferDesc, D3D12_RAYTRACING_ACCELERATION_STRUCTURE_BYTE_ALIGNMENT,
+			0, BufferDesc.Width, BUF_UnorderedAccess,
+			ED3D12ResourceStateMode::SingleState, ERHIAccess::UAVMask,
+			CreateInfo, ResourceAllocator);
+	}
+	else
+	{
+		checkNoEntry();
+	}
+
+	// Elevates the raytracing acceleration structure heap priority, which may help performance / stability in low memory conditions
+	ID3D12Pageable* HeapResource = Result->GetResource()->GetPageable();
+	D3D12_RESIDENCY_PRIORITY HeapPriority = D3D12_RESIDENCY_PRIORITY_HIGH;
+	FD3D12Device* NodeDevice = Adapter->GetDevice(GPUIndex);
+	NodeDevice->GetDevice5()->SetResidencyPriority(1, &HeapResource, &HeapPriority);
+
+	return Result;
+}
+
+FD3D12RayTracingGeometry::FD3D12RayTracingGeometry(FD3D12Adapter* Adapter, const FRayTracingGeometryInitializer& Initializer)
+	: FD3D12AdapterChild(Adapter)
+{
+	INC_DWORD_STAT(STAT_D3D12RayTracingAllocatedBLAS);
+
+	DebugName = Initializer.DebugName.IsValid() ? Initializer.DebugName : FName(TEXT("BLAS"));
+	FMemory::Memzero(bHasPendingCompactionRequests);
+	FMemory::Memzero(bRegisteredAsRenameListener);
+
+	if(!FD3D12RayTracingGeometry::NullTransformBuffer.IsValid())
+	{
+		TResourceArray<float> NullTransformData;
+		NullTransformData.SetNumZeroed(12);
+
+		FRHIResourceCreateInfo CreateInfo(TEXT("NullTransformBuffer"));
+		CreateInfo.ResourceArray = &NullTransformData;
+
+		FD3D12RayTracingGeometry::NullTransformBuffer = RHICreateVertexBuffer(NullTransformData.GetResourceDataSize(), BUF_Static, CreateInfo);
+	}
+
+	IndexOffsetInBytes = Initializer.IndexBufferOffset;
+	TotalPrimitiveCount = Initializer.TotalPrimitiveCount;
+
+	GeometryType = TranslateRayTracingGeometryType(Initializer.GeometryType);
+
+	// #yuriy_todo: get flags directly through the initializer
+	BuildFlags = GetRayTracingAccelerationStructureBuildFlags(Initializer);
+
+	checkf(Initializer.Segments.Num() > 0, TEXT("Ray tracing geometry must be initialized with at least one segment."));
+	Segments = TArray<FRayTracingGeometrySegment>(Initializer.Segments.GetData(), Initializer.Segments.Num());
+
+	IndexStride = RHIIndexBuffer ? RHIIndexBuffer->GetStride() : 0; // stride 0 means implicit triangle list for non-indexed geometry
+
+	static constexpr uint32 IndicesPerPrimitive = 3; // Only triangle meshes are supported
+
+	GeometryDescs.SetNumUninitialized(Segments.Num());
+	TranslateRayTracingGeometryDescs(Initializer, GeometryDescs);
+
+	RHIIndexBuffer = Initializer.IndexBuffer;
+	SetDirty(FRHIGPUMask::All(), true);
+
+	D3D12_BUILD_RAYTRACING_ACCELERATION_STRUCTURE_INPUTS PrebuildDescInputs = {};
+
+	PrebuildDescInputs.Type = D3D12_RAYTRACING_ACCELERATION_STRUCTURE_TYPE_BOTTOM_LEVEL;
+	PrebuildDescInputs.DescsLayout = D3D12_ELEMENTS_LAYOUT_ARRAY;
+	PrebuildDescInputs.NumDescs = GeometryDescs.Num();
+	PrebuildDescInputs.pGeometryDescs = GeometryDescs.GetData();
+	PrebuildDescInputs.Flags = TranslateRayTracingAccelerationStructureFlags(BuildFlags);
+
+	// Get maximum buffer sizes for all GPUs in the system
+	SizeInfo = RHICalcRayTracingGeometrySize(Initializer);
+
 	// Allocate acceleration structure buffer
 	for (uint32 GPUIndex = 0; GPUIndex < MAX_NUM_GPUS && GPUIndex < GNumExplicitGPUsForRendering; ++GPUIndex)
 	{
-		AccelerationStructureBuffers[GPUIndex] = CreateRayTracingBuffer(Adapter, GPUIndex, AccelerationStructureSize, ERayTracingBufferType::AccelerationStructure, DebugName);
+		AccelerationStructureBuffers[GPUIndex] = CreateRayTracingBuffer(Adapter, GPUIndex, SizeInfo.ResultSize, ERayTracingBufferType::AccelerationStructure, DebugName);
 
 		INC_MEMORY_STAT_BY(STAT_D3D12RayTracingUsedVideoMemory, AccelerationStructureBuffers[GPUIndex]->GetSize());
 		INC_MEMORY_STAT_BY(STAT_D3D12RayTracingBLASMemory, AccelerationStructureBuffers[GPUIndex]->GetSize());
@@ -2791,10 +2898,10 @@
 	}
 }
 
-static bool ShouldCompactAfterBuild(D3D12_RAYTRACING_ACCELERATION_STRUCTURE_BUILD_FLAGS BuildFlags)
+static bool ShouldCompactAfterBuild(ERayTracingAccelerationStructureFlags BuildFlags)
 {
-	return EnumHasAllFlags(BuildFlags, D3D12_RAYTRACING_ACCELERATION_STRUCTURE_BUILD_FLAG_ALLOW_COMPACTION | D3D12_RAYTRACING_ACCELERATION_STRUCTURE_BUILD_FLAG_PREFER_FAST_TRACE)
-		&& !EnumHasAnyFlags(BuildFlags, D3D12_RAYTRACING_ACCELERATION_STRUCTURE_BUILD_FLAG_ALLOW_UPDATE);
+	return EnumHasAllFlags(BuildFlags, ERayTracingAccelerationStructureFlags::AllowCompaction | ERayTracingAccelerationStructureFlags::FastTrace)
+		&& !EnumHasAnyFlags(BuildFlags, ERayTracingAccelerationStructureFlags::AllowUpdate);
 }
 
 void FD3D12RayTracingGeometry::BuildAccelerationStructure(FD3D12CommandContext& CommandContext, EAccelerationStructureBuildMode BuildMode)
@@ -2920,7 +3027,7 @@
 
 	if (bIsUpdate)
 	{
-		checkf(BuildFlags & D3D12_RAYTRACING_ACCELERATION_STRUCTURE_BUILD_FLAG_ALLOW_UPDATE,
+		checkf(EnumHasAllFlags(BuildFlags, ERayTracingAccelerationStructureFlags::AllowUpdate),
 			TEXT("Acceleration structure must be created with FRayTracingGeometryInitializer::bAllowUpdate=true to perform refit / update."));
 
 		LocalBuildFlags |= D3D12_RAYTRACING_ACCELERATION_STRUCTURE_BUILD_FLAG_PERFORM_UPDATE;
@@ -2939,21 +3046,21 @@
 	RayTracingDevice->GetRaytracingAccelerationStructurePrebuildInfo(&PrebuildDescInputs, &PrebuildInfo);
 
 	// Must make sure that values computed in the constructor are valid.
-	check(PrebuildInfo.ResultDataMaxSizeInBytes <= AccelerationStructureSize);
+	check(PrebuildInfo.ResultDataMaxSizeInBytes <= SizeInfo.ResultSize);
 
 	TRefCountPtr<FD3D12Buffer> ScratchBuffer;
 
 	if (bIsUpdate)
 	{
-		check(PrebuildInfo.UpdateScratchDataSizeInBytes <= UpdateScratchBufferSize);
+		check(PrebuildInfo.UpdateScratchDataSizeInBytes <= SizeInfo.UpdateScratchSize);
 		static const FName BufferName("UpdateScratchBLAS");
-		ScratchBuffer = CreateRayTracingBuffer(Adapter, GPUIndex, UpdateScratchBufferSize, ERayTracingBufferType::Scratch, BufferName);
+		ScratchBuffer = CreateRayTracingBuffer(Adapter, GPUIndex, SizeInfo.UpdateScratchSize, ERayTracingBufferType::Scratch, BufferName);
 	}
 	else
 	{
-		check(PrebuildInfo.ScratchDataSizeInBytes <= BuildScratchBufferSize);
+		check(PrebuildInfo.ScratchDataSizeInBytes <= SizeInfo.BuildScratchSize);
 		static const FName BufferName("BuildScratchBLAS");
-		ScratchBuffer = CreateRayTracingBuffer(Adapter, GPUIndex, BuildScratchBufferSize, ERayTracingBufferType::Scratch, BufferName);
+		ScratchBuffer = CreateRayTracingBuffer(Adapter, GPUIndex, SizeInfo.BuildScratchSize, ERayTracingBufferType::Scratch, BufferName);
 	}
 
 	TransitionBuffers(CommandContext);
@@ -3170,27 +3277,23 @@
 		}
 	}, ParallelForFlags);
 
+	ERayTracingAccelerationStructureFlags BuildFlags = ERayTracingAccelerationStructureFlags::FastTrace; // #yuriy_todo: pass this in
+
+	BuildInputs = {};
 	BuildInputs.Type = D3D12_RAYTRACING_ACCELERATION_STRUCTURE_TYPE_TOP_LEVEL;
 	BuildInputs.DescsLayout = D3D12_ELEMENTS_LAYOUT_ARRAY;
 	BuildInputs.NumDescs = Instances.Num();
-	BuildInputs.Flags = D3D12_RAYTRACING_ACCELERATION_STRUCTURE_BUILD_FLAG_PREFER_FAST_TRACE;
+	BuildInputs.Flags = TranslateRayTracingAccelerationStructureFlags(BuildFlags);
 
 	NumCallableShaderSlots = Initializer.NumCallableShaderSlots;
 	NumMissShaderSlots = FMath::Max<uint32>(1, Initializer.NumMissShaderSlots);
 
-	for (uint32 GPUIndex = 0; GPUIndex < GNumExplicitGPUsForRendering; ++GPUIndex)
-	{
-		ID3D12Device5* RayTracingDevice = Adapter->GetDevice(GPUIndex)->GetDevice5();
-		D3D12_RAYTRACING_ACCELERATION_STRUCTURE_PREBUILD_INFO PrebuildInfo = {};
-		RayTracingDevice->GetRaytracingAccelerationStructurePrebuildInfo(&BuildInputs, &PrebuildInfo);
-
-		AccelerationStructureSize = FMath::Max(AccelerationStructureSize, PrebuildInfo.ResultDataMaxSizeInBytes);
-		BuildScratchBufferSize = FMath::Max(BuildScratchBufferSize, PrebuildInfo.ScratchDataSizeInBytes);
-	}
+	// Get maximum buffer sizes for all GPUs in the system
+	SizeInfo = RHICalcRayTracingSceneSize(BuildInputs.NumDescs, BuildFlags);
 
 	for (uint32 GPUIndex = 0; GPUIndex < GNumExplicitGPUsForRendering; ++GPUIndex)
 	{
-		AccelerationStructureBuffers[GPUIndex] = CreateRayTracingBuffer(Adapter, GPUIndex, AccelerationStructureSize, ERayTracingBufferType::AccelerationStructure, DebugName);
+		AccelerationStructureBuffers[GPUIndex] = CreateRayTracingBuffer(Adapter, GPUIndex, SizeInfo.ResultSize, ERayTracingBufferType::AccelerationStructure, DebugName);
 
 		INC_MEMORY_STAT_BY(STAT_D3D12RayTracingUsedVideoMemory, AccelerationStructureBuffers[GPUIndex]->GetSize());
 		INC_MEMORY_STAT_BY(STAT_D3D12RayTracingTLASMemory, AccelerationStructureBuffers[GPUIndex]->GetSize());
@@ -3243,8 +3346,8 @@
 
 	D3D12_RAYTRACING_ACCELERATION_STRUCTURE_PREBUILD_INFO PrebuildInfo = {};
 	RayTracingDevice->GetRaytracingAccelerationStructurePrebuildInfo(&BuildInputs, &PrebuildInfo);
-	check(PrebuildInfo.ResultDataMaxSizeInBytes <= AccelerationStructureSize);
-	check(PrebuildInfo.ScratchDataSizeInBytes <= BuildScratchBufferSize);
+	check(PrebuildInfo.ResultDataMaxSizeInBytes <= SizeInfo.ResultSize);
+	check(PrebuildInfo.ScratchDataSizeInBytes <= SizeInfo.BuildScratchSize);
 
 	static const FName ScratchBufferName("BuildScratchTLAS");
 	ScratchBuffer = CreateRayTracingBuffer(Adapter, GPUIndex, PrebuildInfo.ScratchDataSizeInBytes, ERayTracingBufferType::Scratch, ScratchBufferName);
diff --git a/Engine/Source/Runtime/D3D12RHI/Private/D3D12RayTracing.h b/Engine/Source/Runtime/D3D12RHI/Private/D3D12RayTracing.h
--- a/Engine/Source/Runtime/D3D12RHI/Private/D3D12RayTracing.h
+++ b/Engine/Source/Runtime/D3D12RHI/Private/D3D12RayTracing.h
@@ -57,7 +57,7 @@ public:
 	D3D12_RAYTRACING_GEOMETRY_TYPE GeometryType = D3D12_RAYTRACING_GEOMETRY_TYPE_TRIANGLES;
 
 	TArray<FRayTracingGeometrySegment> Segments; // Defines addressable parts of the mesh that can be used for material assignment (one segment = one SBT record)
-	D3D12_RAYTRACING_ACCELERATION_STRUCTURE_BUILD_FLAGS BuildFlags;
+	ERayTracingAccelerationStructureFlags BuildFlags = ERayTracingAccelerationStructureFlags::None;
 
 	FBufferRHIRef  RHIIndexBuffer;
 	static FBufferRHIRef NullTransformBuffer; // Null transform for hidden sections
@@ -77,10 +77,8 @@ public:
 	// Used as a template for BuildAccelerationStructure() later.
 	TArray<D3D12_RAYTRACING_GEOMETRY_DESC, TInlineAllocator<1>> GeometryDescs;
 
-	uint64 AccelerationStructureSize = 0;
 	uint64 AccelerationStructureCompactedSize = 0;
-	uint64 BuildScratchBufferSize = 0;
-	uint64 UpdateScratchBufferSize = 0;
+	FRayTracingAccelerationStructureSize SizeInfo = {};
 };
 
 class FD3D12RayTracingScene : public FRHIRayTracingScene, public FD3D12AdapterChild, public FNoncopyable
@@ -102,8 +100,7 @@ public:
 	TResourceArray<D3D12_RAYTRACING_INSTANCE_DESC, 16> Instances;
 	D3D12_BUILD_RAYTRACING_ACCELERATION_STRUCTURE_INPUTS BuildInputs = {};
 
-	uint64 AccelerationStructureSize = 0;
-	uint64 BuildScratchBufferSize = 0;
+	FRayTracingAccelerationStructureSize SizeInfo = {};
 
 	struct FInstanceCopyCommand
 	{
diff --git a/Engine/Source/Runtime/RHI/Public/DynamicRHI.h b/Engine/Source/Runtime/RHI/Public/DynamicRHI.h
--- a/Engine/Source/Runtime/RHI/Public/DynamicRHI.h
+++ b/Engine/Source/Runtime/RHI/Public/DynamicRHI.h
@@ -11,6 +11,7 @@ DynamicRHI.h: Dynamically bound Render Hardware Interface definitions.
 #include "MultiGPU.h"
 #include "Serialization/MemoryLayout.h"
 #include "Containers/ArrayView.h"
+#include "Misc/EnumClassFlags.h"
 
 class FBlendStateInitializerRHI;
 class FGraphicsPipelineStateInitializer;
@@ -214,6 +215,17 @@ enum ERayTracingSceneLifetime
 	// RTSL_MultiFrame,
 };
 
+enum class ERayTracingAccelerationStructureFlags
+{
+	None            = 0,
+	AllowUpdate     = 1 << 0,
+	AllowCompaction = 1 << 1,
+	FastTrace       = 1 << 2,
+	FastBuild       = 1 << 3,
+	MinimizeMemory  = 1 << 4,
+};
+ENUM_CLASS_FLAGS(ERayTracingAccelerationStructureFlags);
+
 struct FRayTracingSceneInitializer
 {
 	TArrayView<FRayTracingGeometryInstance> Instances;
@@ -241,6 +253,13 @@ struct FRayTracingSceneInitializer
 	FName DebugName;
 };
 
+struct FRayTracingAccelerationStructureSize
+{
+	uint64 ResultSize = 0;
+	uint64 BuildScratchSize = 0;
+	uint64 UpdateScratchSize = 0;
+};
+
 struct FShaderResourceViewInitializer
 {
 	struct FBufferShaderResourceViewInitializer
@@ -1375,6 +1394,19 @@ public:
 	virtual bool RHIRequiresComputeGenerateMips() const { return false; };
 
 #if RHI_RAYTRACING
+
+	virtual FRayTracingAccelerationStructureSize RHICalcRayTracingSceneSize(uint32 MaxInstances, ERayTracingAccelerationStructureFlags Flags)
+	{
+		checkNoEntry();
+		return {};
+	}
+
+	virtual FRayTracingAccelerationStructureSize RHICalcRayTracingGeometrySize(const FRayTracingGeometryInitializer& Initializer)
+	{
+		checkNoEntry();
+		return {};
+	}
+
 	virtual FRayTracingGeometryRHIRef RHICreateRayTracingGeometry(const FRayTracingGeometryInitializer& Initializer)
 	{
 		checkNoEntry();
@@ -1666,6 +1698,16 @@ FORCEINLINE IRHITransientResourceAllocator* RHICreateTransientResourceAllocator(
 
 #if RHI_RAYTRACING
 
+FORCEINLINE FRayTracingAccelerationStructureSize RHICalcRayTracingSceneSize(uint32 MaxInstances, ERayTracingAccelerationStructureFlags Flags)
+{
+	return GDynamicRHI->RHICalcRayTracingSceneSize(MaxInstances, Flags);
+}
+
+FORCEINLINE FRayTracingAccelerationStructureSize RHICalcRayTracingGeometrySize(const FRayTracingGeometryInitializer& Initializer)
+{
+	return GDynamicRHI->RHICalcRayTracingGeometrySize(Initializer);
+}
+
 FORCEINLINE FRayTracingGeometryRHIRef RHICreateRayTracingGeometry(const FRayTracingGeometryInitializer& Initializer)
 {
 	return GDynamicRHI->RHICreateRayTracingGeometry(Initializer);
