- Use pool allocator for upload allocator on PC by default (single shared allocator is easier to maintain)
- Pool allocator has extra array to define allocation order on pools (sorted once per 'frame' to define ideal allocation order)
- Allow max allocation size for pool allocator to be bigger than default pool size
- Restore original upload buffer size for default pool size
- FastConstant allocations use seperate upload allocator to fix a lot of fragmentation issues inside the UploadAllocator
- Remove (broken) USE_BUCKET_ALLOCATOR define and where it was used
- Add upload buffer allocator stats

diff --git a/Engine/Source/Runtime/D3D12RHI/Private/D3D12Adapter.cpp b/Engine/Source/Runtime/D3D12RHI/Private/D3D12Adapter.cpp
--- a/Engine/Source/Runtime/D3D12RHI/Private/D3D12Adapter.cpp
+++ b/Engine/Source/Runtime/D3D12RHI/Private/D3D12Adapter.cpp
@@ -905,14 +905,7 @@ void FD3D12Adapter::InitializeDevices()
 		for (uint32 GPUIndex : FRHIGPUMask::All())
 		{
 			// Safe to init as we have a device;
-			UploadHeapAllocator[GPUIndex] = new FD3D12DynamicHeapAllocator(this,
-				Devices[GPUIndex],
-				Name,
-				EResourceAllocationStrategy::kManualSubAllocation,
-				DEFAULT_CONTEXT_UPLOAD_POOL_MAX_ALLOC_SIZE,
-				DEFAULT_CONTEXT_UPLOAD_POOL_SIZE,
-				DEFAULT_CONTEXT_UPLOAD_POOL_ALIGNMENT);
-
+			UploadHeapAllocator[GPUIndex] = new FD3D12UploadHeapAllocator(this,	Devices[GPUIndex], Name);
 			UploadHeapAllocator[GPUIndex]->Init();
 		}
 
@@ -1247,12 +1240,12 @@ void FD3D12Adapter::BlockUntilIdle()
 }
 
 
-void FD3D12Adapter::TrackAllocationData(FD3D12ResourceLocation* InAllocation, const D3D12_RESOURCE_ALLOCATION_INFO& InInfo)
+void FD3D12Adapter::TrackAllocationData(FD3D12ResourceLocation* InAllocation, uint64 InAllocationSize)
 {
 #if TRACK_RESOURCE_ALLOCATIONS
 	FTrackedAllocationData AllocationData;
 	AllocationData.ResourceAllocation = InAllocation;
-	AllocationData.AllocationInfo = InInfo;
+	AllocationData.AllocationSize = InAllocationSize;
 	AllocationData.StackDepth = FPlatformStackWalk::CaptureStackBackTrace(&AllocationData.Stack[0], FTrackedAllocationData::MaxStackDepth);
 
 	FScopeLock Lock(&TrackedAllocationDataCS);
@@ -1297,7 +1290,7 @@ void FD3D12Adapter::DumpTrackedAllocationData(FOutputDevice& OutputDevice, bool
 	TrackedAllocationData.GenerateValueArray(Allocations);
 	Allocations.Sort([](const FTrackedAllocationData& InLHS, const FTrackedAllocationData& InRHS)
 		{
-			return InLHS.AllocationInfo.SizeInBytes > InRHS.AllocationInfo.SizeInBytes;
+			return InLHS.AllocationSize > InRHS.AllocationSize;
 		});
 
 	TArray<FTrackedAllocationData> BufferAllocations;	
@@ -1310,12 +1303,12 @@ void FD3D12Adapter::DumpTrackedAllocationData(FOutputDevice& OutputDevice, bool
 		if (ResourceDesc.Dimension == D3D12_RESOURCE_DIMENSION_BUFFER)
 		{
 			BufferAllocations.Add(AllocationData);
-			TotalBufferAllocationSize += AllocationData.AllocationInfo.SizeInBytes;
+			TotalBufferAllocationSize += AllocationData.AllocationSize;
 		}
 		else
 		{
 			TextureAllocations.Add(AllocationData);
-			TotalTextureAllocationSize += AllocationData.AllocationInfo.SizeInBytes;
+			TotalTextureAllocationSize += AllocationData.AllocationSize;
 		}
 	}
 
@@ -1344,7 +1337,7 @@ void FD3D12Adapter::DumpTrackedAllocationData(FOutputDevice& OutputDevice, bool
 
 		OutputData += FString::Printf(TEXT("\tName: %s - Size: %3.3fMB - Width: %d - Height: %d - DepthOrArraySize: %d - MipLevels: %d - Flags: %s\n"), 
 			*AllocationData.ResourceAllocation->GetResource()->GetName().ToString(), 
-			AllocationData.AllocationInfo.SizeInBytes / (1024.0f * 1024),
+			AllocationData.AllocationSize / (1024.0f * 1024),
 			ResourceDesc.Width, ResourceDesc.Height, ResourceDesc.DepthOrArraySize, ResourceDesc.MipLevels,
 			Flags.IsEmpty() ? TEXT("None") : *Flags);
 
@@ -1367,7 +1360,7 @@ void FD3D12Adapter::DumpTrackedAllocationData(FOutputDevice& OutputDevice, bool
 
 		OutputData += FString::Printf(TEXT("\tName: %s - Size: %3.3fMB - Width: %d - UAV: %s\n"), 
 			*AllocationData.ResourceAllocation->GetResource()->GetName().ToString(), 
-			AllocationData.AllocationInfo.SizeInBytes / (1024.0f * 1024),
+			AllocationData.AllocationSize / (1024.0f * 1024),
 			ResourceDesc.Width,
 			EnumHasAnyFlags(ResourceDesc.Flags, D3D12_RESOURCE_FLAG_ALLOW_UNORDERED_ACCESS) ? TEXT("Yes") : TEXT("No"));
 
diff --git a/Engine/Source/Runtime/D3D12RHI/Private/D3D12Adapter.h b/Engine/Source/Runtime/D3D12RHI/Private/D3D12Adapter.h
--- a/Engine/Source/Runtime/D3D12RHI/Private/D3D12Adapter.h
+++ b/Engine/Source/Runtime/D3D12RHI/Private/D3D12Adapter.h
@@ -201,7 +201,7 @@ public:
 	FORCEINLINE IDXGIFactory* GetDXGIFactory() const { return DxgiFactory; }
 	FORCEINLINE IDXGIFactory2* GetDXGIFactory2() const { return DxgiFactory2; }
 
-	FORCEINLINE FD3D12DynamicHeapAllocator& GetUploadHeapAllocator(uint32 GPUIndex) 
+	FORCEINLINE FD3D12UploadHeapAllocator& GetUploadHeapAllocator(uint32 GPUIndex)
 	{ 
 		return *(UploadHeapAllocator[GPUIndex]); 
 	}
@@ -328,7 +328,7 @@ public:
 
 	FORCEINLINE uint32 GetFrameCount() const { return FrameCounter; }
 
-	void TrackAllocationData(FD3D12ResourceLocation* InAllocation, const D3D12_RESOURCE_ALLOCATION_INFO& InInfo);
+	void TrackAllocationData(FD3D12ResourceLocation* InAllocation, uint64 InAllocationSize);
 	void ReleaseTrackedAllocationData(FD3D12ResourceLocation* InAllocation);
 	void DumpTrackedAllocationData(FOutputDevice& OutputDevice, bool bWithCallstack);
 
@@ -407,7 +407,7 @@ protected:
 
 	FD3D12FenceCorePool FenceCorePool;
 
-	FD3D12DynamicHeapAllocator*	UploadHeapAllocator[MAX_NUM_GPUS];
+	FD3D12UploadHeapAllocator*	UploadHeapAllocator[MAX_NUM_GPUS];
 
 	/** A list of all viewport RHIs that have been created. */
 	TArray<FD3D12Viewport*> Viewports;
@@ -436,7 +436,7 @@ protected:
 		static const int32 MaxStackDepth = 30;
 
 		FD3D12ResourceLocation* ResourceAllocation;
-		D3D12_RESOURCE_ALLOCATION_INFO AllocationInfo;
+		uint64 AllocationSize;
 		uint32 StackDepth;
 		uint64 Stack[MaxStackDepth];
 	};
diff --git a/Engine/Source/Runtime/D3D12RHI/Private/D3D12Allocation.cpp b/Engine/Source/Runtime/D3D12RHI/Private/D3D12Allocation.cpp
--- a/Engine/Source/Runtime/D3D12RHI/Private/D3D12Allocation.cpp
+++ b/Engine/Source/Runtime/D3D12RHI/Private/D3D12Allocation.cpp
@@ -82,28 +82,28 @@ static FAutoConsoleVariableRef CVarD3D12PoolAllocatorRTUAVTextureMaxAllocationSi
 	ECVF_ReadOnly);
 
 static int32 GD3D12VRAMTexturePoolDefrag = 1;
-static FAutoConsoleVariableRef CVarD3D12D3D12VRAMTexturePoolDefrag(
+static FAutoConsoleVariableRef CVarD3D12VRAMTexturePoolDefrag(
 	TEXT("d3d12.VRAMTexturePoolDefrag"),
 	GD3D12VRAMTexturePoolDefrag,
 	TEXT("Defrag the VRAM Texture pool (enabled by default)"),
 	ECVF_ReadOnly);
 
 static int32 GD3D12VRAMTexturePoolDefragMaxCopySizePerFrame = 32 * 1024 * 1024;
-static FAutoConsoleVariableRef CVarD3D12D3D12VRAMTexturePoolDefragMaxCopySizePerFrame(
+static FAutoConsoleVariableRef CVarD3D12VRAMTexturePoolDefragMaxCopySizePerFrame(
 	TEXT("d3d12.VRAMTexturePoolDefrag.MaxCopySizePerFrame"),
 	GD3D12VRAMTexturePoolDefragMaxCopySizePerFrame,
 	TEXT("Max amount of data to copy during defragmentation in a single frame (default 32MB)"),
 	ECVF_ReadOnly);
 
 static int32 GD3D12VRAMBufferPoolDefrag = 1;
-static FAutoConsoleVariableRef CVarD3D12D3D12VRAMBufferPoolDefrag(
+static FAutoConsoleVariableRef CVarD3D12VRAMBufferPoolDefrag(
 	TEXT("d3d12.VRAMBufferPoolDefrag"),
 	GD3D12VRAMBufferPoolDefrag,
 	TEXT("Defrag the VRAM buffer pool (disabled by default)"),
 	ECVF_ReadOnly);
 
 static int32 GD3D12VRAMBufferPoolDefragMaxCopySizePerFrame = 32 * 1024 * 1024;
-static FAutoConsoleVariableRef CVarD3D12D3D12VRAMBufferPoolDefragMaxCopySizePerFrame(
+static FAutoConsoleVariableRef CVarD3D12VRAMBufferPoolDefragMaxCopySizePerFrame(
 	TEXT("d3d12.VRAMBufferPoolDefrag.MaxCopySizePerFrame"),
 	GD3D12VRAMBufferPoolDefragMaxCopySizePerFrame,
 	TEXT("Max amount of data to copy during defragmentation in a single frame (default 32MB)"),
@@ -111,6 +111,42 @@ static FAutoConsoleVariableRef CVarD3D12D3D12VRAMBufferPoolDefragMaxCopySizePerF
 
 #endif // USE_POOL_ALLOCATOR
 
+static int32 GD3D12UploadHeapSmallBlockMaxAllocationSize = 64 * 1024;
+static FAutoConsoleVariableRef CVarD3D12UploadHeapSmallBlockMaxAllocationSize(
+	TEXT("d3d12.UploadHeap.SmallBlock.MaxAllocationSize"),
+	GD3D12UploadHeapSmallBlockMaxAllocationSize,
+	TEXT("Maximum allocation size on the small block allocator for upload memory"),
+	ECVF_ReadOnly);
+
+static int32 GD3D12UploadHeapSmallBlockPoolSize = 4 * 1024 * 1024;
+static FAutoConsoleVariableRef CVarD3D12UploadHeapSmallBlockPoolSize(
+	TEXT("d3d12.UploadHeap.SmallBlock.PoolSize"),
+	GD3D12UploadHeapSmallBlockPoolSize,
+	TEXT("Pool size for the upload memory small block allocator"),
+	ECVF_ReadOnly);
+
+static int32 GD3D12UploadHeapBigBlockMaxAllocationSize = 64 * 1024 * 1024;
+static FAutoConsoleVariableRef CVarD3D12UploadHeapBigBlockMaxAllocationSize(
+	TEXT("d3d12.UploadHeap.BigBlock.MaxAllocationSize"),
+	GD3D12UploadHeapBigBlockMaxAllocationSize,
+	TEXT("Maximum allocation size on the big block allocator for upload memory"),
+	ECVF_ReadOnly);
+
+static int32 GD3D12UploadHeapBigBlockPoolSize = 8 * 1024 * 1024;
+static FAutoConsoleVariableRef CVarD3D12UploadHeapBigBlockPoolSize(
+	TEXT("d3d12.UploadHeap.BigBlock.PoolSize"),
+	GD3D12UploadHeapBigBlockPoolSize,
+	TEXT("Pool size for the upload memory big block allocator"),
+	ECVF_ReadOnly);
+
+static int32 GD3D12FastConstantAllocatorPageSize = 64 * 1024;
+static FAutoConsoleVariableRef CVarD3D12FastConstantAllocatorPageSize(
+	TEXT("d3d12.FastConstantAllocatorPageSize"),
+	GD3D12FastConstantAllocatorPageSize,
+	TEXT("Page size for the fast constant allocator"),
+	ECVF_ReadOnly);
+
+
 #if D3D12RHI_SEGLIST_ALLOC_TRACK_WASTAGE
 static int32 GD3D12SegListTrackLeaks = 0;
 static FAutoConsoleVariableRef CVarD3D12SegListTrackLeaks(
@@ -958,96 +994,103 @@ void FD3D12BucketAllocator::Reset()
 //	Dynamic Buffer Allocator
 //-----------------------------------------------------------------------------
 
-FD3D12DynamicHeapAllocator::FD3D12DynamicHeapAllocator(FD3D12Adapter* InParent, FD3D12Device* InParentDevice, const FString& InName, EResourceAllocationStrategy InAllocationStrategy,
-	uint32 InMaxSizeForPooling,
-	uint32 InMaxBlockSize,
-	uint32 InMinBlockSize)
+FD3D12UploadHeapAllocator::FD3D12UploadHeapAllocator(FD3D12Adapter* InParent, FD3D12Device* InParentDevice, const FString& InName)
 	: FD3D12AdapterChild(InParent)
-	, FD3D12MultiNodeGPUObject(InParentDevice->GetGPUMask(), FRHIGPUMask::All()), // Dynamic heaps are upload memory, thus they can be trivially visibile to all GPUs
-#if USE_BUCKET_ALLOCATOR
-	Allocator(InParentDevice,
-		GetVisibilityMask(),
-		FD3D12ResourceInitConfig
-		{
-			D3D12_HEAP_TYPE_UPLOAD,
-			D3D12_HEAP_FLAG_ALLOW_ONLY_BUFFERS,
-			D3D12_RESOURCE_FLAG_NONE,
-			D3D12_RESOURCE_STATE_GENERIC_READ
-		},
-		InName,
-		5)
-#else
-	Allocator(InParentDevice,
-		GetVisibilityMask(),
-		FD3D12ResourceInitConfig
-		{
-			D3D12_HEAP_TYPE_UPLOAD,
-			D3D12_HEAP_FLAG_ALLOW_ONLY_BUFFERS,
-			D3D12_RESOURCE_FLAG_NONE,
-			D3D12_RESOURCE_STATE_GENERIC_READ
-		},
-		InName,
-		InAllocationStrategy,
-		InMaxSizeForPooling,
-		InMaxBlockSize,
-		InMinBlockSize
-		)
-#endif
+	, FD3D12MultiNodeGPUObject(InParentDevice->GetGPUMask(), FRHIGPUMask::All()), // Upload memory, thus they can be trivially visibile to all GPUs
+	SmallBlockAllocator(InParentDevice, GetVisibilityMask(), FD3D12ResourceInitConfig::CreateUpload(), InName, EResourceAllocationStrategy::kManualSubAllocation,
+		GD3D12UploadHeapSmallBlockMaxAllocationSize, GD3D12UploadHeapSmallBlockPoolSize, 256),
+	BigBlockAllocator(InParentDevice, GetVisibilityMask(), FD3D12ResourceInitConfig::CreateUpload(), InName, EResourceAllocationStrategy::kManualSubAllocation,
+		GD3D12UploadHeapBigBlockPoolSize, 256, GD3D12UploadHeapBigBlockMaxAllocationSize, FRHIMemoryPool::EFreeListOrder::SortByOffset, false /*defrag*/),
+	FastConstantPageAllocator(InParentDevice, GetVisibilityMask(), FD3D12ResourceInitConfig::CreateUpload(), InName, EResourceAllocationStrategy::kManualSubAllocation,
+		GD3D12FastConstantAllocatorPageSize * 64, GD3D12UploadHeapSmallBlockPoolSize, GD3D12FastConstantAllocatorPageSize)
 {
 }
 
-void FD3D12DynamicHeapAllocator::Init()
-{
 
+void FD3D12UploadHeapAllocator::Destroy()
+{
+	SmallBlockAllocator.Destroy();
+	BigBlockAllocator.Destroy();
+	FastConstantPageAllocator.Destroy();
 }
 
-void* FD3D12DynamicHeapAllocator::AllocUploadResource(uint32 Size, uint32 Alignment, FD3D12ResourceLocation& ResourceLocation)
+
+void* FD3D12UploadHeapAllocator::AllocUploadResource(uint32 InSize, uint32 InAlignment, FD3D12ResourceLocation& ResourceLocation)
 {
+	TRACE_CPUPROFILER_EVENT_SCOPE(FD3D12UploadHeapAllocator::AllocUploadResource);
+
+	// Clean up the release queue of resources which are currently not used by the GPU anymore
 	FD3D12Adapter* Adapter = GetParentAdapter();
+	if (Adapter->GetDeferredDeletionQueue().QueueSize() > 128)
+	{
+		Adapter->GetDeferredDeletionQueue().ReleaseResources(true, false);
+		SmallBlockAllocator.CleanUpAllocations(0); // 0 - no FrameLag, delete all unsued pages
+		BigBlockAllocator.CleanUpAllocations(0);
+	}
 
+	check(InSize > 0);
 	ResourceLocation.Clear();
 
-	//TODO: For some reason 0 sized buffers are being created and then expected to have a resource
-	if (Size == 0)
+	// Fit in small block allocator?
+	if (InSize <= SmallBlockAllocator.GetMaximumAllocationSizeForPooling())
 	{
-		Size = 16;
+		verify(SmallBlockAllocator.TryAllocate(InSize, InAlignment, ResourceLocation));
 	}
-	
-	// Clean up the release queue of resources which are currently not used by the GPU anymore
-	if (Adapter->GetDeferredDeletionQueue().QueueSize() > 128)
+	else
 	{
-		Adapter->GetDeferredDeletionQueue().ReleaseResources(true, false);
-		Allocator.CleanUpAllocations(0); // 0 - no FrameLag, delete all unsued pages
+		// Forward to the big block allocator
+		const D3D12_RESOURCE_DESC ResourceDesc = CD3DX12_RESOURCE_DESC::Buffer(InSize, D3D12_RESOURCE_FLAG_NONE);
+		BigBlockAllocator.AllocateResource(D3D12_HEAP_TYPE_UPLOAD, ResourceDesc, InSize, InAlignment, ED3D12ResourceStateMode::SingleState, D3D12_RESOURCE_STATE_GENERIC_READ, nullptr, nullptr, ResourceLocation);
+		ResourceLocation.UnlockPoolData();		
 	}
 	
-	if (Size <= Allocator.GetMaximumAllocationSizeForPooling())
-	{
-		if (Allocator.TryAllocate(Size, Alignment, ResourceLocation))
-		{
-			return ResourceLocation.GetMappedBaseAddress();
-		}
-	}
-
-	FD3D12Resource* NewResource = nullptr;
+	return ResourceLocation.GetMappedBaseAddress();
+}
 
-	//Allocate Standalone
-	VERIFYD3D12RESULT(Adapter->CreateBuffer(D3D12_HEAP_TYPE_UPLOAD, GetGPUMask(), GetVisibilityMask(), Size, &NewResource, TEXT("Stand Alone Upload Buffer")));
 
-	ResourceLocation.AsStandAlone(NewResource, Size);
+void* FD3D12UploadHeapAllocator::AllocFastConstantAllocationPage(uint32 InSize, uint32 InAlignment, FD3D12ResourceLocation& ResourceLocation)
+{
+	check(InSize > 0);
+	check(InSize <= FastConstantPageAllocator.GetMaximumAllocationSizeForPooling());
 
+	ResourceLocation.Clear();
+	verify(FastConstantPageAllocator.TryAllocate(InSize, InAlignment, ResourceLocation));
 	return ResourceLocation.GetMappedBaseAddress();
 }
 
-void FD3D12DynamicHeapAllocator::CleanUpAllocations(uint64 InFrameLag)
+
+void FD3D12UploadHeapAllocator::CleanUpAllocations(uint64 InFrameLag)
 {
-	Allocator.CleanUpAllocations(InFrameLag);
+	SmallBlockAllocator.CleanUpAllocations(InFrameLag);
+	BigBlockAllocator.CleanUpAllocations(InFrameLag);
+	FastConstantPageAllocator.CleanUpAllocations(InFrameLag);
 }
 
-void FD3D12DynamicHeapAllocator::Destroy()
+
+void FD3D12UploadHeapAllocator::UpdateMemoryStats()
 {
-	Allocator.Destroy();
+	uint32 MemoryAllocated = 0;
+	uint32 MemoryUsed = 0;
+	uint32 FreeMemory = 0;
+	uint32 AlignmentWaste = 0;
+	uint32 AllocatedPageCount = 0;
+	uint32 FullPageCount = 0;
+
+#if defined(D3D12RHI_TRACK_DETAILED_STATS)
+	SmallBlockAllocator.UpdateMemoryStats(MemoryAllocated, MemoryUsed, FreeMemory, AlignmentWaste, AllocatedPageCount, FullPageCount);
+	BigBlockAllocator.UpdateMemoryStats(MemoryAllocated, MemoryUsed, FreeMemory, AlignmentWaste, AllocatedPageCount, FullPageCount);
+	FastConstantPageAllocator.UpdateMemoryStats(MemoryAllocated, MemoryUsed, FreeMemory, AlignmentWaste, AllocatedPageCount, FullPageCount);
+#endif
+
+	SET_MEMORY_STAT(STAT_D3D12UploadPoolMemoryAllocated, MemoryAllocated);
+	SET_MEMORY_STAT(STAT_D3D12UploadPoolMemoryUsed, MemoryUsed);
+	SET_MEMORY_STAT(STAT_D3D12UploadPoolMemoryFree, FreeMemory);
+	SET_MEMORY_STAT(STAT_D3D12UploadPoolAlignmentWaste, AlignmentWaste);
+	SET_DWORD_STAT(STAT_D3D12UploadPoolPageCount, AllocatedPageCount);
+	SET_DWORD_STAT(STAT_D3D12UploadPoolFullPages, FullPageCount);
 }
 
+
 //-----------------------------------------------------------------------------
 //	Default Buffer Allocator
 //-----------------------------------------------------------------------------
@@ -1114,7 +1157,7 @@ EResourceAllocationStrategy FD3D12DefaultBufferPool::GetResourceAllocationStrate
 //-----------------------------------------------------------------------------
 
 
-FD3D12DefaultBufferPool::FD3D12DefaultBufferPool(FD3D12Device* InParent, FD3D12AllocatorType* InAllocator)
+FD3D12DefaultBufferPool::FD3D12DefaultBufferPool(FD3D12Device* InParent, FD3D12MultiBuddyAllocator* InAllocator)
 	: FD3D12DeviceChild(InParent)
 	, FD3D12MultiNodeGPUObject(InAllocator->GetGPUMask(), InAllocator->GetVisibilityMask())
 	, Allocator(InAllocator)
@@ -1125,13 +1168,8 @@ FD3D12DefaultBufferPool::FD3D12DefaultBufferPool(FD3D12Device* InParent, FD3D12A
 bool FD3D12DefaultBufferPool::SupportsAllocation(D3D12_HEAP_TYPE InHeapType, D3D12_RESOURCE_FLAGS InResourceFlags, EBufferUsageFlags InBufferUsage, ED3D12ResourceStateMode InResourceStateMode) const
 {
 	FD3D12ResourceInitConfig InitConfig = GetResourceAllocatorInitConfig(InHeapType, InResourceFlags, InBufferUsage);
-
-#if USE_BUCKET_ALLOCATOR
-	return Allocator->GetInitConfig() == InitConfig;
-#else
 	EResourceAllocationStrategy AllocationStrategy = GetResourceAllocationStrategy(InResourceFlags, InResourceStateMode);
 	return (Allocator->GetInitConfig() == InitConfig && Allocator->GetAllocationStrategy() == AllocationStrategy);
-#endif
 }
 
 
@@ -1179,11 +1217,7 @@ void FD3D12DefaultBufferPool::AllocDefaultResource(D3D12_HEAP_TYPE InHeapType, c
 
 	if (PoolResource)
 	{
-#if USE_BUCKET_ALLOCATOR
-		const bool bPlacedResource = false;
-#else
 		const bool bPlacedResource = (Allocator->GetAllocationStrategy() == EResourceAllocationStrategy::kPlacedResource);
-#endif // USE_BUCKET_ALLOCATOR
 
 		// Ensure we're allocating from the correct pool
 		if (bPlacedResource)
@@ -1252,8 +1286,6 @@ FD3D12DefaultBufferAllocator::FD3D12DefaultBufferAllocator(FD3D12Device* InParen
 FD3D12BufferPool* FD3D12DefaultBufferAllocator::CreateBufferPool(D3D12_HEAP_TYPE InHeapType, D3D12_RESOURCE_FLAGS InResourceFlags, EBufferUsageFlags InBufferUsage, ED3D12ResourceStateMode InResourceStateMode)
 {
 	FD3D12Device* Device = GetParentDevice();
-	FD3D12AllocatorType* Allocator = nullptr;
-
 	FD3D12ResourceInitConfig InitConfig = FD3D12BufferPool::GetResourceAllocatorInitConfig(InHeapType, InResourceFlags, InBufferUsage);
 
 #if USE_POOL_ALLOCATOR
@@ -1282,23 +1314,13 @@ FD3D12BufferPool* FD3D12DefaultBufferAllocator::CreateBufferPool(D3D12_HEAP_TYPE
 
 #else // USE_POOL_ALLOCATOR
 
-#if USE_BUCKET_ALLOCATOR
-
-	const FString Name(L"Default Buffer Bucket Allocator");
-	Allocator = new FD3D12BucketAllocator(Device,
-		GetVisibilityMask(),
-		InitConfig,
-		Name,
-		5);
-#else
-
 	EResourceAllocationStrategy AllocationStrategy = FD3D12DefaultBufferPool::GetResourceAllocationStrategy(InResourceFlags, InResourceStateMode);
 
 	// if placed then 64KB alignment required :(
 	uint32 MinBlockSize = (AllocationStrategy == EResourceAllocationStrategy::kPlacedResource) ? MIN_PLACED_RESOURCE_SIZE : 16;
 
 	const FString Name(L"Default Buffer Multi Buddy Allocator");
-	Allocator = new FD3D12MultiBuddyAllocator(Device,
+	FD3D12MultiBuddyAllocator* Allocator = new FD3D12MultiBuddyAllocator(Device,
 		GetVisibilityMask(),
 		InitConfig,
 		Name,
@@ -1307,7 +1329,6 @@ FD3D12BufferPool* FD3D12DefaultBufferAllocator::CreateBufferPool(D3D12_HEAP_TYPE
 		InHeapType == D3D12_HEAP_TYPE_READBACK ? READBACK_BUFFER_POOL_DEFAULT_POOL_SIZE : DEFAULT_BUFFER_POOL_DEFAULT_POOL_SIZE,
 		MinBlockSize
 		);
-#endif
 
 	FD3D12DefaultBufferPool* NewPool = new FD3D12DefaultBufferPool(Device, Allocator);
 
@@ -1320,12 +1341,8 @@ FD3D12BufferPool* FD3D12DefaultBufferAllocator::CreateBufferPool(D3D12_HEAP_TYPE
 
 bool FD3D12DefaultBufferAllocator::IsPlacedResource(D3D12_RESOURCE_FLAGS InResourceFlags, ED3D12ResourceStateMode InResourceStateMode)
 {
-#if USE_BUCKET_ALLOCATOR
-	return false;
-#else
 	EResourceAllocationStrategy AllocationStrategy = FD3D12BufferPool::GetResourceAllocationStrategy(InResourceFlags, InResourceStateMode);
 	return (AllocationStrategy == EResourceAllocationStrategy::kPlacedResource);
-#endif
 }
 
 
@@ -2038,8 +2055,8 @@ FD3D12FastConstantAllocator::FD3D12FastConstantAllocator(FD3D12Device* Parent, F
 	: FD3D12DeviceChild(Parent)
 	, FD3D12MultiNodeGPUObject(Parent->GetGPUMask(), VisibiltyMask)
 	, UnderlyingResource(Parent)
-	, Offset(64 * 1024) // Initial offset is at end of page so that first Allocate() call triggers a page allocation
-	, PageSize(64 * 1024)
+	, Offset(GD3D12FastConstantAllocatorPageSize) // Initial offset is at end of page so that first Allocate() call triggers a page allocation
+	, PageSize(GD3D12FastConstantAllocatorPageSize)
 {
 	check(PageSize % D3D12_CONSTANT_BUFFER_DATA_PLACEMENT_ALIGNMENT == 0);
 }
@@ -2061,8 +2078,8 @@ void* FD3D12FastConstantAllocator::Allocate(uint32 Bytes, FD3D12ResourceLocation
 		FD3D12Device* Device = GetParentDevice();
 		FD3D12Adapter* Adapter = Device->GetParentAdapter();
 
-		FD3D12DynamicHeapAllocator& Allocator = Adapter->GetUploadHeapAllocator(Device->GetGPUIndex());
-		Allocator.AllocUploadResource(PageSize, DEFAULT_CONTEXT_UPLOAD_POOL_ALIGNMENT, UnderlyingResource);
+		FD3D12UploadHeapAllocator& Allocator = Adapter->GetUploadHeapAllocator(Device->GetGPUIndex());
+		Allocator.AllocFastConstantAllocationPage(PageSize, D3D12_CONSTANT_BUFFER_DATA_PLACEMENT_ALIGNMENT, UnderlyingResource);
 	}
 
 	OutLocation.AsFastAllocation(UnderlyingResource.GetResource(),
diff --git a/Engine/Source/Runtime/D3D12RHI/Private/D3D12Allocation.h b/Engine/Source/Runtime/D3D12RHI/Private/D3D12Allocation.h
--- a/Engine/Source/Runtime/D3D12RHI/Private/D3D12Allocation.h
+++ b/Engine/Source/Runtime/D3D12RHI/Private/D3D12Allocation.h
@@ -348,38 +348,37 @@ private:
 	const uint64 BlockRetentionFrameCount;
 };
 
-#if USE_BUCKET_ALLOCATOR
-typedef FD3D12BucketAllocator FD3D12AllocatorType;
-#else
-typedef FD3D12MultiBuddyAllocator FD3D12AllocatorType;
-#endif
 
 //-----------------------------------------------------------------------------
-//	FD3D12DynamicHeapAllocator
+//	FD3D12UploadHeapAllocator
 //-----------------------------------------------------------------------------
 // This is designed for allocation of scratch memory such as temporary staging buffers
 // or shadow buffers for dynamic resources.
-class FD3D12DynamicHeapAllocator : public FD3D12AdapterChild, public FD3D12MultiNodeGPUObject
+class FD3D12UploadHeapAllocator : public FD3D12AdapterChild, public FD3D12MultiNodeGPUObject
 {
 public:
 
-	FD3D12DynamicHeapAllocator(FD3D12Adapter* InParent, FD3D12Device* InParentDevice, const FString& InName, EResourceAllocationStrategy InAllocationStrategy,
-		uint32 InMaxSizeForPooling,
-		uint32 InMaxBlockSize,
-		uint32 InMinBlockSize);
+	FD3D12UploadHeapAllocator(FD3D12Adapter* InParent, FD3D12Device* InParentDevice, const FString& InName);
 
-	void Init();
+	void Init() {}
+	void Destroy();
 
 	// Allocates <size> bytes from the end of an available resource heap.
-	void* AllocUploadResource(uint32 size, uint32 alignment, FD3D12ResourceLocation& ResourceLocation);
+	void* AllocUploadResource(uint32 InSize, uint32 InAlignment, FD3D12ResourceLocation& ResourceLocation);
+	void* AllocFastConstantAllocationPage(uint32 InSize, uint32 InAlignment, FD3D12ResourceLocation& ResourceLocation);
 
 	void CleanUpAllocations(uint64 InFrameLag);
-
-	void Destroy();
+	void UpdateMemoryStats();
 
 private:
 
-	FD3D12AllocatorType Allocator;
+	// Buddy allocator used for all 'small' allocation - fast but aligns to power of 2
+	FD3D12MultiBuddyAllocator SmallBlockAllocator;
+	// Pool allocator for all bigger allocations - less fast but less alignment waste
+	FD3D12PoolAllocator		  BigBlockAllocator;
+	// Seperate buddy allocator used for the fast constant allocator pages which get always freed within the same frame by default
+	// (different allocator to avoid fragmentation with the other pools - always the same size allocations)
+	FD3D12MultiBuddyAllocator FastConstantPageAllocator;
 };
 
 //-----------------------------------------------------------------------------
@@ -388,7 +387,7 @@ private:
 class FD3D12DefaultBufferPool : public FD3D12DeviceChild, public FD3D12MultiNodeGPUObject
 {
 public:
-	FD3D12DefaultBufferPool(FD3D12Device* InParent, FD3D12AllocatorType* InAllocator);
+	FD3D12DefaultBufferPool(FD3D12Device* InParent, FD3D12MultiBuddyAllocator* InAllocator);
 	~FD3D12DefaultBufferPool() { delete Allocator; }
 
 	bool SupportsAllocation(D3D12_HEAP_TYPE InHeapType, D3D12_RESOURCE_FLAGS InResourceFlags, EBufferUsageFlags InBufferUsage, ED3D12ResourceStateMode InResourceStateMode) const;
@@ -401,7 +400,7 @@ public:
 	static EResourceAllocationStrategy GetResourceAllocationStrategy(D3D12_RESOURCE_FLAGS InResourceFlags, ED3D12ResourceStateMode InResourceStateMode);
 
 private:
-	FD3D12AllocatorType* Allocator;
+	FD3D12MultiBuddyAllocator* Allocator;
 };
 
 #if USE_POOL_ALLOCATOR
diff --git a/Engine/Source/Runtime/D3D12RHI/Private/D3D12CommandContext.cpp b/Engine/Source/Runtime/D3D12RHI/Private/D3D12CommandContext.cpp
--- a/Engine/Source/Runtime/D3D12RHI/Private/D3D12CommandContext.cpp
+++ b/Engine/Source/Runtime/D3D12RHI/Private/D3D12CommandContext.cpp
@@ -724,6 +724,7 @@ void FD3D12CommandContextBase::UpdateMemoryStats()
 #endif
 
 		Device->GetDefaultBufferAllocator().UpdateMemoryStats();
+		ParentAdapter->GetUploadHeapAllocator(GPUIndex).UpdateMemoryStats();
 		Device->GetTransientMemoryPoolManager().UpdateMemoryStats();
 	}
 #endif
diff --git a/Engine/Source/Runtime/D3D12RHI/Private/D3D12Device.cpp b/Engine/Source/Runtime/D3D12RHI/Private/D3D12Device.cpp
--- a/Engine/Source/Runtime/D3D12RHI/Private/D3D12Device.cpp
+++ b/Engine/Source/Runtime/D3D12RHI/Private/D3D12Device.cpp
@@ -421,7 +421,7 @@ void FD3D12Device::Cleanup()
 	// Cleanup thread resources
 	for (int32 index; (index = FPlatformAtomics::InterlockedDecrement(&NumThreadDynamicHeapAllocators)) != -1;)
 	{
-		FD3D12DynamicHeapAllocator* pHeapAllocator = ThreadDynamicHeapAllocatorArray[index];
+		FD3D12UploadHeapAllocator* pHeapAllocator = ThreadDynamicHeapAllocatorArray[index];
 		pHeapAllocator->ReleaseAllResources();
 		delete(pHeapAllocator);
 	}
diff --git a/Engine/Source/Runtime/D3D12RHI/Private/D3D12PoolAllocator.cpp b/Engine/Source/Runtime/D3D12RHI/Private/D3D12PoolAllocator.cpp
--- a/Engine/Source/Runtime/D3D12RHI/Private/D3D12PoolAllocator.cpp
+++ b/Engine/Source/Runtime/D3D12RHI/Private/D3D12PoolAllocator.cpp
@@ -7,6 +7,25 @@
 #define NEEDS_D3D12_INDIRECT_ARGUMENT_HEAP_WORKAROUND 0
 #endif
 
+#if TRACK_RESOURCE_ALLOCATIONS
+
+// Limited size because tracking currently still takes some overhead when creating the callstack for each allocation
+// Arbitrary size right now for which we are interested to know where these allocations are coming from and if they are needed
+// and can be optimized.
+static int32 GD3D12MinPoolAllocationSizeToTrack = 16 * 1024 * 1024;
+static FAutoConsoleVariableRef CVarD3D12MinPoolAllocationSizeToTrack(
+	TEXT("d3d12.MinPoolAllocationSizeToTrack"),
+	GD3D12MinPoolAllocationSizeToTrack,
+	TEXT("Minimum allocation size to track for pool allocations (default 16MB)"),
+	ECVF_ReadOnly);
+
+static bool ShouldTrackAllocation(EResourceAllocationStrategy InAllocationStrategy, D3D12_HEAP_TYPE InHeapType, uint64 InAllocationSize)
+{
+	return (InHeapType == D3D12_HEAP_TYPE_DEFAULT && GD3D12MinPoolAllocationSizeToTrack > 0 && InAllocationSize >= (uint32)GD3D12MinPoolAllocationSizeToTrack);
+}
+
+#endif // TRACK_RESOURCE_ALLOCATIONS
+
 //-----------------------------------------------------------------------------
 //	FD3D12MemoryPool
 //-----------------------------------------------------------------------------
@@ -173,8 +192,8 @@ EResourceAllocationStrategy FD3D12PoolAllocator::GetResourceAllocationStrategy(D
 
 
 FD3D12PoolAllocator::FD3D12PoolAllocator(FD3D12Device* ParentDevice, FRHIGPUMask VisibleNodes, const FD3D12ResourceInitConfig& InInitConfig, const FString& InName,
-	EResourceAllocationStrategy InAllocationStrategy, uint64 InPoolSize, uint32 InPoolAlignment, uint32 InMaxAllocationSize, FRHIMemoryPool::EFreeListOrder InFreeListOrder, bool bInDefragEnabled) :
-	FRHIPoolAllocator(InPoolSize, InPoolAlignment, InMaxAllocationSize, InFreeListOrder, bInDefragEnabled), 
+	EResourceAllocationStrategy InAllocationStrategy, uint64 InDefaultPoolSize, uint32 InPoolAlignment, uint32 InMaxAllocationSize, FRHIMemoryPool::EFreeListOrder InFreeListOrder, bool bInDefragEnabled) :
+	FRHIPoolAllocator(InDefaultPoolSize, InPoolAlignment, InMaxAllocationSize, InFreeListOrder, bInDefragEnabled), 
 	FD3D12DeviceChild(ParentDevice), 
 	FD3D12MultiNodeGPUObject(ParentDevice->GetGPUMask(), VisibleNodes), 
 	InitConfig(InInitConfig), 
@@ -269,52 +288,55 @@ void FD3D12PoolAllocator::AllocateResource(D3D12_HEAP_TYPE InHeapType, const D3D
 
 		// Try to allocate in one of the pools
 		FRHIPoolAllocationData& AllocationData = ResourceLocation.GetPoolAllocatorPrivateData().PoolData;
-		if (TryAllocateInternal(InSize, AllocationAlignment, AllocationData))
-		{
-			// Setup the resource location
-			ResourceLocation.SetType(FD3D12ResourceLocation::ResourceLocationType::eSubAllocation);
-			ResourceLocation.SetPoolAllocator(this);
-			ResourceLocation.SetSize(InSize);
+		verify(TryAllocateInternal(InSize, AllocationAlignment, AllocationData));
 
-			AllocationData.SetOwner(&ResourceLocation);
+		// Setup the resource location
+		ResourceLocation.SetType(FD3D12ResourceLocation::ResourceLocationType::eSubAllocation);
+		ResourceLocation.SetPoolAllocator(this);
+		ResourceLocation.SetSize(InSize);
 
-			if (AllocationStrategy == EResourceAllocationStrategy::kManualSubAllocation)
-			{
-				FD3D12Resource* BackingResource = GetBackingResource(ResourceLocation);
+		AllocationData.SetOwner(&ResourceLocation);
 
-				ResourceLocation.SetOffsetFromBaseOfResource(AllocationData.GetOffset());
-				ResourceLocation.SetResource(BackingResource);
-				ResourceLocation.SetGPUVirtualAddress(BackingResource->GetGPUVirtualAddress() + AllocationData.GetOffset());
+		if (AllocationStrategy == EResourceAllocationStrategy::kManualSubAllocation)
+		{
+			FD3D12Resource* BackingResource = GetBackingResource(ResourceLocation);
 
-				if (IsCPUAccessible(InitConfig.HeapType))
-				{
-					ResourceLocation.SetMappedBaseAddress((uint8*)BackingResource->GetResourceBaseAddress() + AllocationData.GetOffset());
-				}
-			}
-			else
-			{
-				check(ResourceLocation.GetResource() == nullptr);
+			ResourceLocation.SetOffsetFromBaseOfResource(AllocationData.GetOffset());
+			ResourceLocation.SetResource(BackingResource);
+			ResourceLocation.SetGPUVirtualAddress(BackingResource->GetGPUVirtualAddress() + AllocationData.GetOffset());
 
-				FD3D12Resource* NewResource = CreatePlacedResource(AllocationData, InDesc, InCreateState, InResourceStateMode, InClearValue, InName);
-				ResourceLocation.SetResource(NewResource);
+			if (IsCPUAccessible(InitConfig.HeapType))
+			{
+				ResourceLocation.SetMappedBaseAddress((uint8*)BackingResource->GetResourceBaseAddress() + AllocationData.GetOffset());
 			}
+		}
+		else
+		{
+			check(ResourceLocation.GetResource() == nullptr);
 
-			// Successfully sub-allocated
-			return;
+			FD3D12Resource* NewResource = CreatePlacedResource(AllocationData, InDesc, InCreateState, InResourceStateMode, InClearValue, InName);
+			ResourceLocation.SetResource(NewResource);
 		}
 
-		// shouldn't fail
-		check(false);
+#if TRACK_RESOURCE_ALLOCATIONS
+		// Update stats if above a certain size
+		if (ShouldTrackAllocation(AllocationStrategy, InitConfig.HeapType, AllocationData.GetSize()))
+		{
+			Adapter->TrackAllocationData(&ResourceLocation, AllocationData.GetSize());
+		}
+#endif // TRACK_RESOURCE_ALLOCATIONS
 	}
+	else
+	{
+		// Allocate Standalone - move to owner of resource because this allocator should only manage pooled allocations (needed for now to do the same as FD3D12DefaultBufferPool)
+		FD3D12Resource* NewResource = nullptr;
+		const D3D12_HEAP_PROPERTIES HeapProps = CD3DX12_HEAP_PROPERTIES(InHeapType, GetGPUMask().GetNative(), GetVisibilityMask().GetNative());
+		D3D12_RESOURCE_DESC Desc = InDesc;
+		Desc.Alignment = 0;
+		VERIFYD3D12RESULT(Adapter->CreateCommittedResource(Desc, GetGPUMask(), HeapProps, InCreateState, InResourceStateMode, InCreateState, InClearValue, &NewResource, InName, false));
 
-	// Allocate Standalone - move to owner of resource because this allocator should only manage pooled allocations (needed for now to do the same as FD3D12DefaultBufferPool)
-	FD3D12Resource* NewResource = nullptr;
-	const D3D12_HEAP_PROPERTIES HeapProps = CD3DX12_HEAP_PROPERTIES(InHeapType, GetGPUMask().GetNative(), GetVisibilityMask().GetNative());
-	D3D12_RESOURCE_DESC Desc = InDesc;
-	Desc.Alignment = 0;
-	VERIFYD3D12RESULT(Adapter->CreateCommittedResource(Desc, GetGPUMask(), HeapProps, InCreateState, InResourceStateMode, InCreateState, InClearValue, &NewResource, InName, false));
-	
-	ResourceLocation.AsStandAlone(NewResource, InSize);
+		ResourceLocation.AsStandAlone(NewResource, InSize);
+	}
 }
 
 
@@ -339,12 +361,20 @@ void FD3D12PoolAllocator::DeallocateResource(FD3D12ResourceLocation& ResourceLoc
 {
 	check(IsOwner(ResourceLocation));
 
-	FScopeLock Lock(&CS);
-
 	// Mark allocation data as free
-	FRHIPoolAllocationData& AllocationData = ResourceLocation.GetPoolAllocatorPrivateData().PoolData;	
+	FRHIPoolAllocationData& AllocationData = ResourceLocation.GetPoolAllocatorPrivateData().PoolData;
 	check(AllocationData.IsAllocated());
 
+#if TRACK_RESOURCE_ALLOCATIONS
+	// Release tracked resource if tracked
+	if (ShouldTrackAllocation(AllocationStrategy, InitConfig.HeapType, AllocationData.GetSize()))
+	{
+		GetParentDevice()->GetParentAdapter()->ReleaseTrackedAllocationData(&ResourceLocation);
+	}
+#endif // TRACK_RESOURCE_ALLOCATIONS
+
+	FScopeLock Lock(&CS);
+
 	// If locked then assume still initial setup or in defragmentation unlock request
 	// Mark as nop because block will be deleted anyway
 	// TODO: optimize to not perform the full iteration here!
@@ -407,8 +437,19 @@ void FD3D12PoolAllocator::DeallocateResource(FD3D12ResourceLocation& ResourceLoc
 }
 
 
-FRHIMemoryPool* FD3D12PoolAllocator::CreateNewPool(int16 InPoolIndex)
+FRHIMemoryPool* FD3D12PoolAllocator::CreateNewPool(int16 InPoolIndex, uint32 InMinimumAllocationSize)
 {
+	TRACE_CPUPROFILER_EVENT_SCOPE(FD3D12PoolAllocator::CreateNewPool);
+
+	// Find out the pool size - use default, but if allocation doesn't fit then round up to next power of 2
+	// so it 'always' fits the pool allocator
+	uint32 PoolSize = DefaultPoolSize;
+	if (InMinimumAllocationSize > PoolSize)
+	{
+		check(InMinimumAllocationSize <= MaxAllocationSize);
+		PoolSize = FMath::Min(FMath::RoundUpToPowerOfTwo(InMinimumAllocationSize), (uint32)MaxAllocationSize);
+	}
+
 	FD3D12MemoryPool* NewPool = new FD3D12MemoryPool(GetParentDevice(),	GetVisibilityMask(), InitConfig,
 		Name, AllocationStrategy, InPoolIndex, PoolSize, PoolAlignment, FreeListOrder);
 	NewPool->Init();
@@ -456,6 +497,14 @@ bool FD3D12PoolAllocator::HandleDefragRequest(FRHIPoolAllocationData* InSourceBl
 	check(CopyOp.DestResource != nullptr);
 	PendingCopyOps.Add(CopyOp);
 
+#if TRACK_RESOURCE_ALLOCATIONS
+	// Update stats if above a certain size
+	if (ShouldTrackAllocation(AllocationStrategy, InitConfig.HeapType, InSourceBlock->GetSize()))
+	{
+		Adapter->TrackAllocationData(Owner, InSourceBlock->GetSize());
+	}
+#endif // TRACK_RESOURCE_ALLOCATIONS
+
 	// TODO: Using aliasing buffer on whole heap for copies to reduce flushes and resource transitions
 
 	return true;
@@ -536,6 +585,32 @@ void FD3D12PoolAllocator::CleanUpAllocations(uint64 InFrameLag)
 			Pools[PoolIndex] = nullptr;
 		}
 	}
+
+	// Update the allocation order
+	Algo::Sort(PoolAllocationOrder, [this](uint32 InLHS, uint32 InRHS)
+		{
+			// first allocate from 'fullest' pools when defrag is enabled, otherwise try and allocate from
+			// default sized pools with stable order
+			if (bDefragEnabled)
+			{
+				uint32 LHSUsePoolSize = Pools[InLHS] ? Pools[InLHS]->GetUsedSize() : UINT32_MAX;
+				uint32 RHSUsePoolSize = Pools[InRHS] ? Pools[InRHS]->GetUsedSize() : UINT32_MAX;
+				return LHSUsePoolSize > RHSUsePoolSize;
+			}
+			else
+			{
+				uint32 LHSPoolSize = Pools[InLHS] ? Pools[InLHS]->GetPoolSize() : UINT32_MAX;
+				uint32 RHSPoolSize = Pools[InRHS] ? Pools[InRHS]->GetPoolSize() : UINT32_MAX;
+				if (LHSPoolSize != RHSPoolSize)
+				{
+					return LHSPoolSize < RHSPoolSize;
+				}
+
+				uint32 LHSPoolIndex = Pools[InLHS] ? Pools[InLHS]->GetPoolIndex() : UINT32_MAX;
+				uint32 RHSPoolIndex = Pools[InRHS] ? Pools[InRHS]->GetPoolIndex() : UINT32_MAX;
+				return LHSPoolIndex < RHSPoolIndex;
+			}
+		});
 }
 
 
diff --git a/Engine/Source/Runtime/D3D12RHI/Private/D3D12PoolAllocator.h b/Engine/Source/Runtime/D3D12RHI/Private/D3D12PoolAllocator.h
--- a/Engine/Source/Runtime/D3D12RHI/Private/D3D12PoolAllocator.h
+++ b/Engine/Source/Runtime/D3D12RHI/Private/D3D12PoolAllocator.h
@@ -22,6 +22,16 @@ enum class EResourceAllocationStrategy
 // All data required to create a pool
 struct FD3D12ResourceInitConfig
 {
+	static FD3D12ResourceInitConfig CreateUpload()
+	{
+		FD3D12ResourceInitConfig Config;
+		Config.HeapType = D3D12_HEAP_TYPE_UPLOAD;
+		Config.HeapFlags = D3D12_HEAP_FLAG_ALLOW_ONLY_BUFFERS;
+		Config.ResourceFlags = D3D12_RESOURCE_FLAG_NONE;
+		Config.InitialResourceState = D3D12_RESOURCE_STATE_GENERIC_READ;
+		return Config;
+	}
+
 	bool operator==(const FD3D12ResourceInitConfig& InOther) const
 	{
 		return HeapType == InOther.HeapType &&
@@ -117,7 +127,7 @@ public:
 
 	// Constructor
 	FD3D12PoolAllocator(FD3D12Device* ParentDevice, FRHIGPUMask VisibleNodes, const FD3D12ResourceInitConfig& InInitConfig, const FString& InName,
-		EResourceAllocationStrategy InAllocationStrategy, uint64 InPoolSize, uint32 InPoolAlignment, uint32 InMaxAllocationSize, FRHIMemoryPool::EFreeListOrder InFreeListOrder, bool bInDefragEnabled);
+		EResourceAllocationStrategy InAllocationStrategy, uint64 InDefaultPoolSize, uint32 InPoolAlignment, uint32 InMaxAllocationSize, FRHIMemoryPool::EFreeListOrder InFreeListOrder, bool bInDefragEnabled);
 	~FD3D12PoolAllocator();
 	
 	// Function names currently have to match with FD3D12DefaultBufferPool until we can make full replacement of the allocator
@@ -147,7 +157,7 @@ public:
 protected:
 
 	// Implementation of FRHIPoolAllocator pure virtuals
-	virtual FRHIMemoryPool* CreateNewPool(int16 InPoolIndex) override;
+	virtual FRHIMemoryPool* CreateNewPool(int16 InPoolIndex, uint32 InMinimumAllocationSize) override;
 	virtual bool HandleDefragRequest(FRHIPoolAllocationData* InSourceBlock, FRHIPoolAllocationData& InTmpTargetBlock) override;
 	
 	// Placed resource allocation helper function which can be overriden
diff --git a/Engine/Source/Runtime/D3D12RHI/Private/D3D12RHIPrivate.h b/Engine/Source/Runtime/D3D12RHI/Private/D3D12RHIPrivate.h
--- a/Engine/Source/Runtime/D3D12RHI/Private/D3D12RHIPrivate.h
+++ b/Engine/Source/Runtime/D3D12RHI/Private/D3D12RHIPrivate.h
@@ -61,8 +61,6 @@ extern int32 GDX12NVAfterMathTrackResources;
 extern int32 GDX12NVAfterMathMarkers;
 #endif
 
-#define USE_BUCKET_ALLOCATOR 0
-
 #include "D3D12Residency.h"
 
 // D3D RHI public headers.
@@ -114,9 +112,6 @@ typedef FD3D12StateCacheBase FD3D12StateCache;
 #define READBACK_BUFFER_POOL_MAX_ALLOC_SIZE (64 * 1024)
 #define READBACK_BUFFER_POOL_DEFAULT_POOL_SIZE (4 * 1024 * 1024)
 
-#define DEFAULT_CONTEXT_UPLOAD_POOL_SIZE (32 * 1024 * 1024)
-#define DEFAULT_CONTEXT_UPLOAD_POOL_MAX_ALLOC_SIZE (32 * 1024 * 1024)
-#define DEFAULT_CONTEXT_UPLOAD_POOL_ALIGNMENT (D3D12_CONSTANT_BUFFER_DATA_PLACEMENT_ALIGNMENT)
 #define TEXTURE_POOL_SIZE (8 * 1024 * 1024)
 
 #define MAX_GPU_BREADCRUMB_DEPTH 1024
diff --git a/Engine/Source/Runtime/D3D12RHI/Private/D3D12Resources.cpp b/Engine/Source/Runtime/D3D12RHI/Private/D3D12Resources.cpp
--- a/Engine/Source/Runtime/D3D12RHI/Private/D3D12Resources.cpp
+++ b/Engine/Source/Runtime/D3D12RHI/Private/D3D12Resources.cpp
@@ -449,6 +449,8 @@ HRESULT FD3D12Adapter::CreateCommittedResource(const D3D12_RESOURCE_DESC& InDesc
 		return E_POINTER;
 	}
 
+	TRACE_CPUPROFILER_EVENT_SCOPE(CreateCommittedResource);
+
 	LLM_PLATFORM_SCOPE(ELLMTag::GraphicsPlatform);
 
 	TRefCountPtr<ID3D12Resource> pResource;
@@ -833,7 +835,7 @@ void FD3D12ResourceLocation::UpdateStandAloneStats(bool bIncrement)
 		// Track all committed resource allocations
 		if (bIncrement)
 		{
-			UnderlyingResource->GetParentDevice()->GetParentAdapter()->TrackAllocationData(this, Info);
+			UnderlyingResource->GetParentDevice()->GetParentAdapter()->TrackAllocationData(this, Info.SizeInBytes);
 		}
 		else
 		{
diff --git a/Engine/Source/Runtime/D3D12RHI/Private/D3D12Stats.h b/Engine/Source/Runtime/D3D12RHI/Private/D3D12Stats.h
--- a/Engine/Source/Runtime/D3D12RHI/Private/D3D12Stats.h
+++ b/Engine/Source/Runtime/D3D12RHI/Private/D3D12Stats.h
@@ -108,6 +108,13 @@ DECLARE_MEMORY_STAT_EXTERN(TEXT("Texture Memory Frame Committed Allocated"), STA
 DECLARE_DWORD_ACCUMULATOR_STAT_EXTERN(TEXT("Texture Memory Pool Allocations"), STAT_D3D12TransientMemoryTexturePoolAllocations, STATGROUP_D3D12MemoryTransient, );
 DECLARE_DWORD_ACCUMULATOR_STAT_EXTERN(TEXT("Texture Memory Allocations"), STAT_D3D12TransientMemoryTextureCommittedAllocations, STATGROUP_D3D12MemoryTransient, );
 
+DECLARE_MEMORY_STAT_EXTERN(TEXT("UploadPool Memory Allocated"), STAT_D3D12UploadPoolMemoryAllocated, STATGROUP_D3D12Memory, );
+DECLARE_MEMORY_STAT_EXTERN(TEXT("UploadPool Memory Used"), STAT_D3D12UploadPoolMemoryUsed, STATGROUP_D3D12Memory, );
+
+DECLARE_MEMORY_STAT_EXTERN(TEXT("UploadPool Memory Free"), STAT_D3D12UploadPoolMemoryFree, STATGROUP_D3D12MemoryBuffer, );
+DECLARE_MEMORY_STAT_EXTERN(TEXT("UploadPool Memory Alignment Waste"), STAT_D3D12UploadPoolAlignmentWaste, STATGROUP_D3D12MemoryBuffer, );
+DECLARE_DWORD_ACCUMULATOR_STAT_EXTERN(TEXT("UploadPool Page Count"), STAT_D3D12UploadPoolPageCount, STATGROUP_D3D12MemoryBuffer, );
+DECLARE_DWORD_ACCUMULATOR_STAT_EXTERN(TEXT("UploadPool Full Pages"), STAT_D3D12UploadPoolFullPages, STATGROUP_D3D12MemoryBuffer, );
 
 /**
 * Detailed Descriptor heap stats
diff --git a/Engine/Source/Runtime/D3D12RHI/Private/D3D12TransientResourceAllocator.cpp b/Engine/Source/Runtime/D3D12RHI/Private/D3D12TransientResourceAllocator.cpp
--- a/Engine/Source/Runtime/D3D12RHI/Private/D3D12TransientResourceAllocator.cpp
+++ b/Engine/Source/Runtime/D3D12RHI/Private/D3D12TransientResourceAllocator.cpp
@@ -725,7 +725,7 @@ void FD3D12TransientResourceAllocator::Freeze()
 }
 
 
-FRHIMemoryPool* FD3D12TransientResourceAllocator::CreateNewPool(int16 InPoolIndex)
+FRHIMemoryPool* FD3D12TransientResourceAllocator::CreateNewPool(int16 InPoolIndex, uint32 InMinimumAllocationSize)
 {
 	// Get pool from manager - don't reallocate each time
 	return GetParentDevice()->GetTransientMemoryPoolManager().GetOrCreateMemoryPool(InPoolIndex);
diff --git a/Engine/Source/Runtime/D3D12RHI/Private/D3D12TransientResourceAllocator.h b/Engine/Source/Runtime/D3D12RHI/Private/D3D12TransientResourceAllocator.h
--- a/Engine/Source/Runtime/D3D12RHI/Private/D3D12TransientResourceAllocator.h
+++ b/Engine/Source/Runtime/D3D12RHI/Private/D3D12TransientResourceAllocator.h
@@ -202,7 +202,7 @@ private:
 
 	// Override placed resource allocation helper function to get from cache if available
 	virtual FD3D12Resource* CreatePlacedResource(const FRHIPoolAllocationData& InAllocationData, const D3D12_RESOURCE_DESC& InDesc, D3D12_RESOURCE_STATES InCreateState, ED3D12ResourceStateMode InResourceStateMode, const D3D12_CLEAR_VALUE* InClearValue, const TCHAR* InName) override;
-	virtual FRHIMemoryPool* CreateNewPool(int16 InPoolIndex) override;
+	virtual FRHIMemoryPool* CreateNewPool(int16 InPoolIndex, uint32 InMinimumAllocationSize) override;
 
 	// Shared allocation/deallocation helper
 	FD3D12BaseShaderResource* GetBaseShaderResource(FRHITexture* InRHITexture);
diff --git a/Engine/Source/Runtime/D3D12RHI/Private/D3D12UniformBuffer.cpp b/Engine/Source/Runtime/D3D12RHI/Private/D3D12UniformBuffer.cpp
--- a/Engine/Source/Runtime/D3D12RHI/Private/D3D12UniformBuffer.cpp
+++ b/Engine/Source/Runtime/D3D12RHI/Private/D3D12UniformBuffer.cpp
@@ -42,8 +42,8 @@ FUniformBufferRHIRef FD3D12DynamicRHI::RHICreateUniformBuffer(const void* Conten
 			if (Usage == EUniformBufferUsage::UniformBuffer_MultiFrame)
 			{
 				// Uniform buffers that live for multiple frames must use the more expensive and persistent allocation path
-				FD3D12DynamicHeapAllocator& Allocator = GetAdapter().GetUploadHeapAllocator(Device->GetGPUIndex());
-				MappedData = Allocator.AllocUploadResource(NumBytesActualData, DEFAULT_CONTEXT_UPLOAD_POOL_ALIGNMENT, NewUniformBuffer->ResourceLocation);
+				FD3D12UploadHeapAllocator& Allocator = GetAdapter().GetUploadHeapAllocator(Device->GetGPUIndex());
+				MappedData = Allocator.AllocUploadResource(NumBytesActualData, D3D12_CONSTANT_BUFFER_DATA_PLACEMENT_ALIGNMENT, NewUniformBuffer->ResourceLocation);
 			}
 			else
 			{
@@ -174,8 +174,8 @@ void FD3D12DynamicRHI::RHIUpdateUniformBuffer(FRHIUniformBuffer* UniformBufferRH
 
 			if (UniformBuffer.UniformBufferUsage == UniformBuffer_MultiFrame)
 			{
-				FD3D12DynamicHeapAllocator& Allocator = GetAdapter().GetUploadHeapAllocator(Device->GetGPUIndex());
-				MappedData = Allocator.AllocUploadResource(NumBytes, DEFAULT_CONTEXT_UPLOAD_POOL_ALIGNMENT, UpdatedResourceLocation);
+				FD3D12UploadHeapAllocator& Allocator = GetAdapter().GetUploadHeapAllocator(Device->GetGPUIndex());
+				MappedData = Allocator.AllocUploadResource(NumBytes, D3D12_CONSTANT_BUFFER_DATA_PLACEMENT_ALIGNMENT, UpdatedResourceLocation);
 			}
 			else
 			{
diff --git a/Engine/Source/Runtime/D3D12RHI/Private/D3D12Util.cpp b/Engine/Source/Runtime/D3D12RHI/Private/D3D12Util.cpp
--- a/Engine/Source/Runtime/D3D12RHI/Private/D3D12Util.cpp
+++ b/Engine/Source/Runtime/D3D12RHI/Private/D3D12Util.cpp
@@ -1493,6 +1493,13 @@ DEFINE_STAT(STAT_D3D12BufferPoolAlignmentWaste);
 DEFINE_STAT(STAT_D3D12BufferPoolPageCount);
 DEFINE_STAT(STAT_D3D12BufferPoolFullPages);
 
+DEFINE_STAT(STAT_D3D12UploadPoolMemoryAllocated);
+DEFINE_STAT(STAT_D3D12UploadPoolMemoryUsed);
+DEFINE_STAT(STAT_D3D12UploadPoolMemoryFree);
+DEFINE_STAT(STAT_D3D12UploadPoolAlignmentWaste);
+DEFINE_STAT(STAT_D3D12UploadPoolPageCount);
+DEFINE_STAT(STAT_D3D12UploadPoolFullPages);
+
 DEFINE_STAT(STAT_UniqueSamplers);
 
 DEFINE_STAT(STAT_ViewHeapChanged);
diff --git a/Engine/Source/Runtime/D3D12RHI/Public/D3D12Resources.h b/Engine/Source/Runtime/D3D12RHI/Public/D3D12Resources.h
--- a/Engine/Source/Runtime/D3D12RHI/Public/D3D12Resources.h
+++ b/Engine/Source/Runtime/D3D12RHI/Public/D3D12Resources.h
@@ -438,11 +438,7 @@ private:
 	}
 };
 
-#if USE_BUCKET_ALLOCATOR
-typedef class FD3D12BucketAllocator FD3D12BaseAllocatorType;
-#else
 typedef class FD3D12BuddyAllocator FD3D12BaseAllocatorType;
-#endif
 
 struct FD3D12BuddyAllocatorPrivateData
 {
diff --git a/Engine/Source/Runtime/RHICore/Private/RHIPoolAllocator.cpp b/Engine/Source/Runtime/RHICore/Private/RHIPoolAllocator.cpp
--- a/Engine/Source/Runtime/RHICore/Private/RHIPoolAllocator.cpp
+++ b/Engine/Source/Runtime/RHICore/Private/RHIPoolAllocator.cpp
@@ -528,8 +528,8 @@ void FRHIMemoryPool::Validate()
 //	Pool Allocator
 //-----------------------------------------------------------------------------
 
-FRHIPoolAllocator::FRHIPoolAllocator(uint64 InPoolSize, uint32 InPoolAlignment, uint32 InMaxAllocationSize, FRHIMemoryPool::EFreeListOrder InFreeListOrder, bool InDefragEnabled) :
-	PoolSize(InPoolSize),
+FRHIPoolAllocator::FRHIPoolAllocator(uint64 InDefaultPoolSize, uint32 InPoolAlignment, uint32 InMaxAllocationSize, FRHIMemoryPool::EFreeListOrder InFreeListOrder, bool InDefragEnabled) :
+	DefaultPoolSize(InDefaultPoolSize),
 	PoolAlignment(InPoolAlignment),
 	MaxAllocationSize(InMaxAllocationSize),
 	FreeListOrder(InFreeListOrder),
@@ -561,16 +561,21 @@ void FRHIPoolAllocator::Destroy()
 	}
 
 	Pools.Empty();
+	PoolAllocationOrder.Empty();
 }
 
 
 bool FRHIPoolAllocator::TryAllocateInternal(uint32 InSizeInBytes, uint32 InAllocationAlignment, FRHIPoolAllocationData& AllocationData)
 {
+	TRACE_CPUPROFILER_EVENT_SCOPE(FRHIPoolAllocator::TryAllocateInternal);
+
 	FScopeLock Lock(&CS);
 
-	for (int32 PoolIndex = 0; PoolIndex < Pools.Num(); PoolIndex++)
+	check(PoolAllocationOrder.Num() == Pools.Num())
+	for (int32 PoolIndex : PoolAllocationOrder)
 	{
-		if (Pools[PoolIndex] && Pools[PoolIndex]->TryAllocate(InSizeInBytes, InAllocationAlignment, AllocationData))
+		FRHIMemoryPool* Pool = Pools[PoolIndex];
+		if (Pool != nullptr && !Pool->IsFull() && Pool->TryAllocate(InSizeInBytes, InAllocationAlignment, AllocationData))
 		{
 			return true;
 		}
@@ -589,12 +594,13 @@ bool FRHIPoolAllocator::TryAllocateInternal(uint32 InSizeInBytes, uint32 InAlloc
 
 	if (NewPoolIndex >= 0)
 	{
-		Pools[NewPoolIndex] = CreateNewPool(NewPoolIndex);
+		Pools[NewPoolIndex] = CreateNewPool(NewPoolIndex, InSizeInBytes);
 	}
 	else
 	{
-		NewPoolIndex = Pools.Num();
-		Pools.Add(CreateNewPool(NewPoolIndex));
+		NewPoolIndex = Pools.Num();		
+		Pools.Add(CreateNewPool(NewPoolIndex, InSizeInBytes));
+		PoolAllocationOrder.Add(NewPoolIndex);
 	}
 
 	return Pools[NewPoolIndex]->TryAllocate(InSizeInBytes, InAllocationAlignment, AllocationData);
diff --git a/Engine/Source/Runtime/RHICore/Public/RHIPoolAllocator.h b/Engine/Source/Runtime/RHICore/Public/RHIPoolAllocator.h
--- a/Engine/Source/Runtime/RHICore/Public/RHIPoolAllocator.h
+++ b/Engine/Source/Runtime/RHICore/Public/RHIPoolAllocator.h
@@ -194,7 +194,7 @@ class RHICORE_API FRHIPoolAllocator
 public:
 
 	// Constructor
-	FRHIPoolAllocator(uint64 InPoolSize, uint32 InPoolAlignment, uint32 InMaxAllocationSize, FRHIMemoryPool::EFreeListOrder InFreeListOrder, bool InDefragEnabled);
+	FRHIPoolAllocator(uint64 InDefaultPoolSize, uint32 InPoolAlignment, uint32 InMaxAllocationSize, FRHIMemoryPool::EFreeListOrder InFreeListOrder, bool InDefragEnabled);
 	virtual ~FRHIPoolAllocator();
 
 	// Setup/Shutdown
@@ -214,14 +214,14 @@ protected:
 	void DeallocateInternal(FRHIPoolAllocationData& AllocationData);
 	
 	// Helper function to create a new platform specific pool
-	virtual FRHIMemoryPool* CreateNewPool(int16 InPoolIndex) = 0;
+	virtual FRHIMemoryPool* CreateNewPool(int16 InPoolIndex, uint32 InMinimumAllocationSize) = 0;
 
 	// Handle a rhi specific defrag op
 	friend class FRHIMemoryPool;
 	virtual bool HandleDefragRequest(FRHIPoolAllocationData* InSourceBlock, FRHIPoolAllocationData& InTmpTargetBlock) = 0;
 
 	// Const creation members - used to create new pools
-	const uint64 PoolSize;
+	const uint64 DefaultPoolSize;
 	const uint32 PoolAlignment;
 	const uint64 MaxAllocationSize;
 	const FRHIMemoryPool::EFreeListOrder FreeListOrder;
@@ -233,6 +233,10 @@ protected:
 	// Actual pools managing each a linked list of allocations
 	TArray<FRHIMemoryPool*> Pools;
 
+	// Allocation order of the pools when performing new allocations
+	// Can' sort the pools directly because pool index is stored in the allocation info
+	TArray<uint32> PoolAllocationOrder;
+
 	// Stats
 	uint32 TotalAllocatedBlocks;
 };
