Add support for cubemaps in d3d12 transient resource allocator

diff --git a/Engine/Source/Runtime/D3D12RHI/Private/D3D12TransientResourceAllocator.cpp b/Engine/Source/Runtime/D3D12RHI/Private/D3D12TransientResourceAllocator.cpp
--- a/Engine/Source/Runtime/D3D12RHI/Private/D3D12TransientResourceAllocator.cpp
+++ b/Engine/Source/Runtime/D3D12RHI/Private/D3D12TransientResourceAllocator.cpp
@@ -552,9 +552,25 @@ FRHITexture* FD3D12TransientResourceAllocator::CreateTexture(const FRHITextureCr
 			BaseShaderResource = Texture3D;
 			break;
 		}
+		case ETextureDimension::TextureCube:
+		{
+			check(InCreateInfo.Extent.X == InCreateInfo.Extent.Y);
+
+			const bool bTextureArray = false;
+			const bool bCubeTexture = true;
+			ID3D12ResourceAllocator* ResourceAllocator = this;
+			FD3D12TextureCube* TextureCube = FD3D12DynamicRHI::GetD3DRHI()->CreateD3D12Texture2D<FD3D12BaseTextureCube>(nullptr, InCreateInfo.Extent.X, InCreateInfo.Extent.Y, 6, bTextureArray, bCubeTexture, InCreateInfo.Format, InCreateInfo.NumMips, InCreateInfo.NumSamples, InCreateInfo.Flags, ERHIAccess::Discard, CreateInfo, ResourceAllocator);
+			TextureData.RHITexture = TextureCube;
+			TextureData.ResourceDesc = TextureCube->GetResource()->GetDesc();
+			TextureData.ClearValue = TextureCube->GetResource()->GetClearValue();
+			TextureData.CreateInfo = InCreateInfo;
+
+			BaseShaderResource = TextureCube;
+			break;
+		}
 		default:
 		{
-			// Only support 2d & 3d textures for now
+			// Only support 2d, 3d & cube textures for now
 			check(false);
 			break;
 		}
@@ -646,6 +662,12 @@ FD3D12BaseShaderResource* FD3D12TransientResourceAllocator::GetBaseShaderResourc
 		FD3D12Texture3D* Texture3D = FD3D12DynamicRHI::ResourceCast(RHITexture3D);
 		BaseShaderResource = Texture3D;
 	}
+	else if (FRHITextureCube* RHITextureCube = InRHITexture->GetTextureCube())
+	{
+		FD3D12TextureCube* TextureCube = FD3D12DynamicRHI::ResourceCast(RHITextureCube);
+		BaseShaderResource = TextureCube;
+	}
+	check(BaseShaderResource);
 	return BaseShaderResource;
 }
 
