Don't track allocations for transient resources and slighty refactor/cleanup pool allocation tracking

diff --git a/Engine/Source/Runtime/D3D12RHI/Private/D3D12PoolAllocator.cpp b/Engine/Source/Runtime/D3D12RHI/Private/D3D12PoolAllocator.cpp
--- a/Engine/Source/Runtime/D3D12RHI/Private/D3D12PoolAllocator.cpp
+++ b/Engine/Source/Runtime/D3D12RHI/Private/D3D12PoolAllocator.cpp
@@ -19,11 +19,6 @@ static FAutoConsoleVariableRef CVarD3D12MinPoolAllocationSizeToTrack(
 	TEXT("Minimum allocation size to track for pool allocations (default 16MB)"),
 	ECVF_ReadOnly);
 
-static bool ShouldTrackAllocation(EResourceAllocationStrategy InAllocationStrategy, D3D12_HEAP_TYPE InHeapType, uint64 InAllocationSize)
-{
-	return (InHeapType == D3D12_HEAP_TYPE_DEFAULT && GD3D12MinPoolAllocationSizeToTrack > 0 && InAllocationSize >= (uint32)GD3D12MinPoolAllocationSizeToTrack);
-}
-
 #endif // TRACK_RESOURCE_ALLOCATIONS
 
 //-----------------------------------------------------------------------------
@@ -318,13 +313,7 @@ void FD3D12PoolAllocator::AllocateResource(D3D12_HEAP_TYPE InHeapType, const D3D
 			ResourceLocation.SetResource(NewResource);
 		}
 
-#if TRACK_RESOURCE_ALLOCATIONS
-		// Update stats if above a certain size
-		if (ShouldTrackAllocation(AllocationStrategy, InitConfig.HeapType, AllocationData.GetSize()))
-		{
-			Adapter->TrackAllocationData(&ResourceLocation, AllocationData.GetSize());
-		}
-#endif // TRACK_RESOURCE_ALLOCATIONS
+		UpdateAllocationTracking(ResourceLocation, EAllocationType::Allocate);
 	}
 	else
 	{
@@ -361,20 +350,14 @@ void FD3D12PoolAllocator::DeallocateResource(FD3D12ResourceLocation& ResourceLoc
 {
 	check(IsOwner(ResourceLocation));
 
+	UpdateAllocationTracking(ResourceLocation, EAllocationType::Free);
+	
+	FScopeLock Lock(&CS);
+
 	// Mark allocation data as free
 	FRHIPoolAllocationData& AllocationData = ResourceLocation.GetPoolAllocatorPrivateData().PoolData;
 	check(AllocationData.IsAllocated());
 
-#if TRACK_RESOURCE_ALLOCATIONS
-	// Release tracked resource if tracked
-	if (ShouldTrackAllocation(AllocationStrategy, InitConfig.HeapType, AllocationData.GetSize()))
-	{
-		GetParentDevice()->GetParentAdapter()->ReleaseTrackedAllocationData(&ResourceLocation);
-	}
-#endif // TRACK_RESOURCE_ALLOCATIONS
-
-	FScopeLock Lock(&CS);
-
 	// If locked then assume still initial setup or in defragmentation unlock request
 	// Mark as nop because block will be deleted anyway
 	// TODO: optimize to not perform the full iteration here!
@@ -497,13 +480,7 @@ bool FD3D12PoolAllocator::HandleDefragRequest(FRHIPoolAllocationData* InSourceBl
 	check(CopyOp.DestResource != nullptr);
 	PendingCopyOps.Add(CopyOp);
 
-#if TRACK_RESOURCE_ALLOCATIONS
-	// Update stats if above a certain size
-	if (ShouldTrackAllocation(AllocationStrategy, InitConfig.HeapType, InSourceBlock->GetSize()))
-	{
-		Adapter->TrackAllocationData(Owner, InSourceBlock->GetSize());
-	}
-#endif // TRACK_RESOURCE_ALLOCATIONS
+	UpdateAllocationTracking(*Owner, EAllocationType::Allocate);
 
 	// TODO: Using aliasing buffer on whole heap for copies to reduce flushes and resource transitions
 
@@ -653,6 +630,32 @@ FD3D12HeapAndOffset FD3D12PoolAllocator::GetBackingHeapAndAllocationOffsetInByte
 }
 
 
+void FD3D12PoolAllocator::UpdateAllocationTracking(FD3D12ResourceLocation& InAllocation, EAllocationType InAllocationType)
+{
+#if TRACK_RESOURCE_ALLOCATIONS
+	// Only track for default heap and of allocation is of certain minimum size
+	check(IsOwner(InAllocation));
+	FRHIPoolAllocationData& AllocationData = InAllocation.GetPoolAllocatorPrivateData().PoolData;
+	check(AllocationData.IsAllocated());
+	uint64 AllocationSize = AllocationData.GetSize();
+	if (InitConfig.HeapType == D3D12_HEAP_TYPE_DEFAULT && GD3D12MinPoolAllocationSizeToTrack > 0 && AllocationSize >= (uint32)GD3D12MinPoolAllocationSizeToTrack)
+	{
+		FD3D12Adapter* Adapter = GetParentDevice()->GetParentAdapter();		
+		if (InAllocationType == EAllocationType::Allocate)
+		{
+			Adapter->TrackAllocationData(&InAllocation, AllocationSize);
+		}
+		else
+		{
+			Adapter->ReleaseTrackedAllocationData(&InAllocation);
+		}
+	}
+#else 
+	return false;
+#endif // TRACK_RESOURCE_ALLOCATIONS
+}
+
+
 void FD3D12PoolAllocator::FlushPendingCopyOps(FD3D12CommandContext& InCommandContext)
 {
 	FScopeLock Lock(&CS);
diff --git a/Engine/Source/Runtime/D3D12RHI/Private/D3D12PoolAllocator.h b/Engine/Source/Runtime/D3D12RHI/Private/D3D12PoolAllocator.h
--- a/Engine/Source/Runtime/D3D12RHI/Private/D3D12PoolAllocator.h
+++ b/Engine/Source/Runtime/D3D12RHI/Private/D3D12PoolAllocator.h
@@ -163,6 +163,14 @@ protected:
 	// Placed resource allocation helper function which can be overriden
 	virtual FD3D12Resource* CreatePlacedResource(const FRHIPoolAllocationData& InAllocationData, const D3D12_RESOURCE_DESC& InDesc, D3D12_RESOURCE_STATES InCreateState, ED3D12ResourceStateMode InResourceStateMode, const D3D12_CLEAR_VALUE* InClearValue, const TCHAR* InName);
 	
+	// Track the allocation
+	enum class EAllocationType
+	{
+		Allocate,
+		Free
+	};
+	virtual void UpdateAllocationTracking(FD3D12ResourceLocation& InAllocation, EAllocationType InAllocationType);
+
 	// Locked allocation data which needs to do something specific at certain frame fence value	
 	struct FrameFencedAllocationData
 	{
diff --git a/Engine/Source/Runtime/D3D12RHI/Private/D3D12TransientResourceAllocator.h b/Engine/Source/Runtime/D3D12RHI/Private/D3D12TransientResourceAllocator.h
--- a/Engine/Source/Runtime/D3D12RHI/Private/D3D12TransientResourceAllocator.h
+++ b/Engine/Source/Runtime/D3D12RHI/Private/D3D12TransientResourceAllocator.h
@@ -209,6 +209,9 @@ private:
 	void SetupAllocatedResource(FD3D12BaseShaderResource* InBaseShaderResource, FD3D12TransientMemoryStats& Stats);
 	void DeallocateMemory(FD3D12BaseShaderResource* InBaseShaderResource, FD3D12TransientMemoryStats& Stats);
 
+	// Don't track any allocations
+	virtual void UpdateAllocationTracking(FD3D12ResourceLocation& InAllocation, EAllocationType InAllocationType) override {}
+
 	// Reference all the allocated resources - owner of all the resources and will be returned back to the manager when the allocator is destroyed
 	TArray<FD3D12PooledTextureData> AllocatedTextures;
 	TArray<FD3D12PooledBufferData> AllocatedBuffers;
