Compile fix for when resource tracking is disabled

diff --git a/Engine/Source/Runtime/D3D12RHI/Private/D3D12PoolAllocator.cpp b/Engine/Source/Runtime/D3D12RHI/Private/D3D12PoolAllocator.cpp
--- a/Engine/Source/Runtime/D3D12RHI/Private/D3D12PoolAllocator.cpp
+++ b/Engine/Source/Runtime/D3D12RHI/Private/D3D12PoolAllocator.cpp
@@ -650,8 +650,6 @@ void FD3D12PoolAllocator::UpdateAllocationTracking(FD3D12ResourceLocation& InAll
 			Adapter->ReleaseTrackedAllocationData(&InAllocation);
 		}
 	}
-#else 
-	return false;
 #endif // TRACK_RESOURCE_ALLOCATIONS
 }
 
