Removed GRHISupportsRayTracingExplicitMemoryManagement. This is now the default and only code path.

diff --git a/Engine/Source/Runtime/D3D12RHI/Private/D3D12RHI.cpp b/Engine/Source/Runtime/D3D12RHI/Private/D3D12RHI.cpp
--- a/Engine/Source/Runtime/D3D12RHI/Private/D3D12RHI.cpp
+++ b/Engine/Source/Runtime/D3D12RHI/Private/D3D12RHI.cpp
@@ -211,9 +211,6 @@ FD3D12DynamicRHI::FD3D12DynamicRHI(const TArray<TSharedPtr<FD3D12Adapter>>& Chos
 	// Manually enable Async BVH build for D3D12 RHI
 	GRHISupportsRayTracingAsyncBuildAccelerationStructure = true;
 
-	// Allow high-level code to allocate ray tracing buffers
-	GRHISupportsRayTracingExplicitMemoryManagement = true;
-
 	GRHISupportsPipelineFileCache = PLATFORM_WINDOWS;
 }
 
diff --git a/Engine/Source/Runtime/D3D12RHI/Private/D3D12RayTracing.cpp b/Engine/Source/Runtime/D3D12RHI/Private/D3D12RayTracing.cpp
--- a/Engine/Source/Runtime/D3D12RHI/Private/D3D12RayTracing.cpp
+++ b/Engine/Source/Runtime/D3D12RHI/Private/D3D12RayTracing.cpp
@@ -3291,28 +3291,22 @@ FD3D12RayTracingScene::FD3D12RayTracingScene(FD3D12Adapter* Adapter, const FRayT
 	// Get maximum buffer sizes for all GPUs in the system
 	SizeInfo = RHICalcRayTracingSceneSize(BuildInputs.NumDescs, BuildFlags);
 
-	bExplicitMemoryManagement = Initializer.bExplicitMemoryManagement;
-
-	// #yuriy_todo: this will be removed once all RHIs are moved to explicit buffer management. Leaving here for reference.
-	if (!Initializer.bExplicitMemoryManagement)
-	{
-		for (uint32 GPUIndex = 0; GPUIndex < GNumExplicitGPUsForRendering; ++GPUIndex)
-		{	
-			AccelerationStructureBuffers[GPUIndex] = CreateRayTracingBuffer(Adapter, GPUIndex, SizeInfo.ResultSize, ERayTracingBufferType::AccelerationStructure, DebugName);
+	for (uint32 GPUIndex = 0; GPUIndex < GNumExplicitGPUsForRendering; ++GPUIndex)
+	{	
+		AccelerationStructureBuffers[GPUIndex] = CreateRayTracingBuffer(Adapter, GPUIndex, SizeInfo.ResultSize, ERayTracingBufferType::AccelerationStructure, DebugName);
 
-			INC_MEMORY_STAT_BY(STAT_D3D12RayTracingUsedVideoMemory, AccelerationStructureBuffers[GPUIndex]->GetSize());
-			INC_MEMORY_STAT_BY(STAT_D3D12RayTracingTLASMemory, AccelerationStructureBuffers[GPUIndex]->GetSize());
+		INC_MEMORY_STAT_BY(STAT_D3D12RayTracingUsedVideoMemory, AccelerationStructureBuffers[GPUIndex]->GetSize());
+		INC_MEMORY_STAT_BY(STAT_D3D12RayTracingTLASMemory, AccelerationStructureBuffers[GPUIndex]->GetSize());
 	
-			FD3D12ShaderResourceView* View = FD3D12CommandContext::RetrieveObject<FD3D12ShaderResourceView>(ShaderResourceView.GetReference(), GPUIndex);
+		FD3D12ShaderResourceView* View = FD3D12CommandContext::RetrieveObject<FD3D12ShaderResourceView>(ShaderResourceView.GetReference(), GPUIndex);
 
-			D3D12_SHADER_RESOURCE_VIEW_DESC SRVDesc = {};
-			SRVDesc.Format = DXGI_FORMAT_UNKNOWN;
-			SRVDesc.ViewDimension = D3D12_SRV_DIMENSION_RAYTRACING_ACCELERATION_STRUCTURE;
-			SRVDesc.Shader4ComponentMapping = D3D12_DEFAULT_SHADER_4_COMPONENT_MAPPING;
-			SRVDesc.RaytracingAccelerationStructure.Location = AccelerationStructureBuffers[GPUIndex]->ResourceLocation.GetGPUVirtualAddress();
+		D3D12_SHADER_RESOURCE_VIEW_DESC SRVDesc = {};
+		SRVDesc.Format = DXGI_FORMAT_UNKNOWN;
+		SRVDesc.ViewDimension = D3D12_SRV_DIMENSION_RAYTRACING_ACCELERATION_STRUCTURE;
+		SRVDesc.Shader4ComponentMapping = D3D12_DEFAULT_SHADER_4_COMPONENT_MAPPING;
+		SRVDesc.RaytracingAccelerationStructure.Location = AccelerationStructureBuffers[GPUIndex]->ResourceLocation.GetGPUVirtualAddress();
 
-			View->Initialize(SRVDesc, AccelerationStructureBuffers[GPUIndex], 4, 0);
-		}
+		View->Initialize(SRVDesc, AccelerationStructureBuffers[GPUIndex], 4, 0);
 	}
 };
 
@@ -3340,7 +3334,6 @@ FD3D12RayTracingScene::~FD3D12RayTracingScene()
 
 void FD3D12RayTracingScene::BindBuffer(FRHIBuffer* InBuffer, uint32 InBufferOffset)
 {
-	checkf(bExplicitMemoryManagement, TEXT("BindBuffer API can only be used when the scene is created in explicit memory management mode."));
 	check(SizeInfo.ResultSize + InBufferOffset <= InBuffer->GetSize());
 
 	for (uint32 GPUIndex = 0; GPUIndex < GNumExplicitGPUsForRendering; ++GPUIndex)
@@ -3370,8 +3363,6 @@ void FD3D12RayTracingScene::BuildAccelerationStructure(FD3D12CommandContext& Com
 	static const FName ScratchBufferName("BuildScratchTLAS");
 	ScratchBuffer = CreateRayTracingBuffer(Adapter, GPUIndex, PrebuildInfo.ScratchDataSizeInBytes, ERayTracingBufferType::Scratch, ScratchBufferName);
 
-	FD3D12ShaderResourceView* AccelerationStructureView = CommandContext.RetrieveObject<FD3D12ShaderResourceView>(ShaderResourceView.GetReference());
-
 	// Create and fill instance buffer
 
 	if (BuildInputs.NumDescs)
@@ -3503,6 +3494,9 @@ void FD3D12RayTracingScene::BuildAccelerationStructure(FD3D12CommandContext& Com
 	}
 
 	TRefCountPtr<FD3D12Buffer>& AccelerationStructureBuffer = AccelerationStructureBuffers[GPUIndex];
+	checkf(AccelerationStructureBuffer.IsValid(), 
+		TEXT("Acceleration structure buffer must be set for this scene using RHIBindAccelerationStructureMemory() before build command is issued."));
+
 	AccelerationStructureBuffer->GetResource()->UpdateResidency(CommandContext.CommandListHandle);
 	ScratchBuffer->GetResource()->UpdateResidency(CommandContext.CommandListHandle);
 
diff --git a/Engine/Source/Runtime/D3D12RHI/Private/D3D12RayTracing.h b/Engine/Source/Runtime/D3D12RHI/Private/D3D12RayTracing.h
--- a/Engine/Source/Runtime/D3D12RHI/Private/D3D12RayTracing.h
+++ b/Engine/Source/Runtime/D3D12RHI/Private/D3D12RayTracing.h
@@ -157,8 +157,6 @@ public:
 
 	uint64 LastCommandListID = 0;
 
-	bool bExplicitMemoryManagement = false;
-
 	FName DebugName;
 };
 
diff --git a/Engine/Source/Runtime/RHI/Private/RHI.cpp b/Engine/Source/Runtime/RHI/Private/RHI.cpp
--- a/Engine/Source/Runtime/RHI/Private/RHI.cpp
+++ b/Engine/Source/Runtime/RHI/Private/RHI.cpp
@@ -927,7 +927,6 @@ bool GRHISupportsRayTracing = false;
 bool GRHISupportsRayTracingPSOAdditions = false;
 bool GRHISupportsRayTracingDispatchIndirect = false;
 bool GRHISupportsRayTracingAsyncBuildAccelerationStructure = false;
-bool GRHISupportsRayTracingExplicitMemoryManagement = false;
 bool GRHISupportsRayTracingAMDHitToken = false;
 bool GRHISupportsWaveOperations = false;
 int32 GRHIMinimumWaveSize = 4; // Minimum supported value in SM 6.0
diff --git a/Engine/Source/Runtime/RHI/Public/DynamicRHI.h b/Engine/Source/Runtime/RHI/Public/DynamicRHI.h
--- a/Engine/Source/Runtime/RHI/Public/DynamicRHI.h
+++ b/Engine/Source/Runtime/RHI/Public/DynamicRHI.h
@@ -250,13 +250,6 @@ struct FRayTracingSceneInitializer
 	// Currently only single-frame lifetime is supported.
 	ERayTracingSceneLifetime Lifetime = RTSL_SingleFrame;
 
-	// If this is enabled, the RHI will not create the acceleration structure result buffer and SRV automatically.
-	// User must create and manage those explicitly in high-level rendering code.
-	// Buffer must be created with BUF_AccelerationStructure usage.
-	// Buffer must be bound to the acceleration structure using RHIBindAccelerationStructureMemory command before 
-	// the acceleration structure is used for any other GPU command.
-	// This can only be used if GRHISupportsRayTracingExplicitMemoryManagement is true.
-	bool bExplicitMemoryManagement = false;
 
 	FName DebugName;
 };
diff --git a/Engine/Source/Runtime/RHI/Public/RHI.h b/Engine/Source/Runtime/RHI/Public/RHI.h
--- a/Engine/Source/Runtime/RHI/Public/RHI.h
+++ b/Engine/Source/Runtime/RHI/Public/RHI.h
@@ -575,9 +575,6 @@ extern RHI_API bool GRHISupportsRayTracingDispatchIndirect;
 /** Whether or not the RHI supports async building ray tracing acceleration structures. */
 extern RHI_API bool GRHISupportsRayTracingAsyncBuildAccelerationStructure;
 
-/** Whether RHI allows high-level code to explicitly allocate memory for ray tracing acceleration structures. */
-extern RHI_API bool GRHISupportsRayTracingExplicitMemoryManagement;
-
 /** Whether or not the RHI supports the AMD Hit Token extension. */
 extern RHI_API bool GRHISupportsRayTracingAMDHitToken;
 
diff --git a/Engine/Source/Runtime/Renderer/Private/DeferredShadingRenderer.cpp b/Engine/Source/Runtime/Renderer/Private/DeferredShadingRenderer.cpp
--- a/Engine/Source/Runtime/Renderer/Private/DeferredShadingRenderer.cpp
+++ b/Engine/Source/Runtime/Renderer/Private/DeferredShadingRenderer.cpp
@@ -1261,15 +1261,12 @@ bool FDeferredShadingSceneRenderer::DispatchRayTracingWorldUpdates(FRDGBuilder&
 		PrepareLumenHardwareRayTracingReflections(View, RayGenShaders);
 		PrepareLumenHardwareRayTracingVisualize(View, RayGenShaders);
 
-		if (GRHISupportsRayTracingExplicitMemoryManagement)
-		{
-			// #yuriy_todo: allocate these using RDG transient resources!
-			FRayTracingAccelerationStructureSize SizeInfo = RHICalcRayTracingSceneSize(NumTotalInstances, ERayTracingAccelerationStructureFlags::FastTrace);
-			check(SizeInfo.ResultSize < ~0u);
-			FRHIResourceCreateInfo CreateInfo(TEXT("RayTracingSceneBuffer"));
-			View.RayTracingSceneBuffer = RHICreateBuffer(uint32(SizeInfo.ResultSize), BUF_AccelerationStructure, 0, ERHIAccess::BVHWrite, CreateInfo);
-			View.RayTracingSceneSRV = RHICreateShaderResourceView(View.RayTracingSceneBuffer);
-		}
+		// #yuriy_todo: allocate these using RDG transient resources!
+		FRayTracingAccelerationStructureSize SizeInfo = RHICalcRayTracingSceneSize(NumTotalInstances, ERayTracingAccelerationStructureFlags::FastTrace);
+		check(SizeInfo.ResultSize < ~0u);
+		FRHIResourceCreateInfo CreateInfo(TEXT("RayTracingSceneBuffer"));
+		View.RayTracingSceneBuffer = RHICreateBuffer(uint32(SizeInfo.ResultSize), BUF_AccelerationStructure, 0, ERHIAccess::BVHWrite, CreateInfo);
+		View.RayTracingSceneSRV = RHICreateShaderResourceView(View.RayTracingSceneBuffer);
 
 		View.CreateRayTracingSceneTask = FFunctionGraphTask::CreateAndDispatchWhenReady([
 			NumTotalInstances,
@@ -1284,7 +1281,6 @@ bool FDeferredShadingSceneRenderer::DispatchRayTracingWorldUpdates(FRDGBuilder&
 			SceneInitializer.Instances = Instances;
 			SceneInitializer.ShaderSlotsPerGeometrySegment = RAY_TRACING_NUM_SHADER_SLOTS;
 			SceneInitializer.NumMissShaderSlots = RAY_TRACING_NUM_MISS_SHADER_SLOTS;
-			SceneInitializer.bExplicitMemoryManagement = GRHISupportsRayTracingExplicitMemoryManagement;
 
 			ResultScene = RHICreateRayTracingScene(SceneInitializer);
 
@@ -1339,12 +1335,10 @@ bool FDeferredShadingSceneRenderer::DispatchRayTracingWorldUpdates(FRDGBuilder&
 			{
 				FViewInfo& View = Views[ViewIndex];
 				FRHIRayTracingScene* RayTracingScene = View.GetRayTracingSceneChecked();
-				if (GRHISupportsRayTracingExplicitMemoryManagement)
-				{
-					check(View.RayTracingSceneBuffer.IsValid());
-					check(View.RayTracingSceneSRV.IsValid());
-					RHIAsyncCmdList.BindAccelerationStructureMemory(RayTracingScene, View.RayTracingSceneBuffer, 0);
-				}
+
+				check(View.RayTracingSceneBuffer.IsValid());
+				check(View.RayTracingSceneSRV.IsValid());
+				RHIAsyncCmdList.BindAccelerationStructureMemory(RayTracingScene, View.RayTracingSceneBuffer, 0);
 
 				RHIAsyncCmdList.BuildAccelerationStructure(RayTracingScene);
 			}
@@ -1372,12 +1366,9 @@ bool FDeferredShadingSceneRenderer::DispatchRayTracingWorldUpdates(FRDGBuilder&
 					FViewInfo& View = Views[ViewIndex];
 					FRHIRayTracingScene* RayTracingScene = View.GetRayTracingSceneChecked();
 
-					if (GRHISupportsRayTracingExplicitMemoryManagement)
-					{
-						check(View.RayTracingSceneBuffer.IsValid());
-						check(View.RayTracingSceneSRV.IsValid());
-						RHICmdList.BindAccelerationStructureMemory(RayTracingScene, View.RayTracingSceneBuffer, 0);
-					}
+					check(View.RayTracingSceneBuffer.IsValid());
+					check(View.RayTracingSceneSRV.IsValid());
+					RHICmdList.BindAccelerationStructureMemory(RayTracingScene, View.RayTracingSceneBuffer, 0);
 
 					RHICmdList.BuildAccelerationStructure(RayTracingScene);
 				}
diff --git a/Engine/Source/Runtime/Renderer/Private/SceneRendering.cpp b/Engine/Source/Runtime/Renderer/Private/SceneRendering.cpp
--- a/Engine/Source/Runtime/Renderer/Private/SceneRendering.cpp
+++ b/Engine/Source/Runtime/Renderer/Private/SceneRendering.cpp
@@ -1033,15 +1033,8 @@ FRHIRayTracingScene* FViewInfo::GetRayTracingSceneChecked() const
 }
 FRHIShaderResourceView* FViewInfo::GetRayTracingSceneViewChecked() const
 {
-	FRHIShaderResourceView* Result = nullptr;
-	if (GRHISupportsRayTracingExplicitMemoryManagement)
-	{
-		Result = RayTracingSceneSRV.GetReference();
-	}
-	else
-	{
-		Result = GetRayTracingSceneChecked()->GetShaderResourceView();
-	}
+	FRHIShaderResourceView* Result = RayTracingSceneSRV.GetReference();
+
 	checkf(Result, TEXT("Ray tracing scene SRV is expected to be created at this point."));
 	return Result;
 }
