Fix name conflicts between DX11 and DX12 RHIs. In a monolith build both will be linked into the exe, so for functions with identical signatures the linker will simply pick one or the other, causing one RHI to call into the other RHI's functions.

diff --git a/Engine/Plugins/Runtime/Oculus/OculusVR/Source/OculusHMD/Private/OculusHMDPrivateRHI.h b/Engine/Plugins/Runtime/Oculus/OculusVR/Source/OculusHMD/Private/OculusHMDPrivateRHI.h
--- a/Engine/Plugins/Runtime/Oculus/OculusVR/Source/OculusHMD/Private/OculusHMDPrivateRHI.h
+++ b/Engine/Plugins/Runtime/Oculus/OculusVR/Source/OculusHMD/Private/OculusHMDPrivateRHI.h
@@ -19,40 +19,8 @@
 //-------------------------------------------------------------------------------------------------
 
 #if OCULUS_HMD_SUPPORTED_PLATFORMS_D3D12
-#define GetD3D11CubeFace GetD3D12CubeFace
-#define VerifyD3D11Result VerifyD3D12Result
-#define GetD3D11TextureFromRHITexture GetD3D12TextureFromRHITexture
-#define FRingAllocation FRingAllocation_D3D12
-#define GetRenderTargetFormat GetRenderTargetFormat_D3D12
-#define ED3D11ShaderOffsetBuffer ED3D12ShaderOffsetBuffer
-#define FindShaderResourceDXGIFormat FindShaderResourceDXGIFormat_D3D12
-#define FindUnorderedAccessDXGIFormat FindUnorderedAccessDXGIFormat_D3D12
-#define FindDepthStencilDXGIFormat FindDepthStencilDXGIFormat_D3D12
-#define HasStencilBits HasStencilBits_D3D12
-#define FVector4VertexDeclaration FVector4VertexDeclaration_D3D12
-#define GLOBAL_CONSTANT_BUFFER_INDEX GLOBAL_CONSTANT_BUFFER_INDEX_D3D12
-#define MAX_CONSTANT_BUFFER_SLOTS MAX_CONSTANT_BUFFER_SLOTS_D3D12
-#define FD3DGPUProfiler FD3D12GPUProfiler
-#define FRangeAllocator FRangeAllocator_D3D12
-
 #include "D3D12RHIPrivate.h"
 #include "D3D12Util.h"
-
-#undef GetD3D11CubeFace
-#undef VerifyD3D11Result
-#undef GetD3D11TextureFromRHITexture
-#undef FRingAllocation
-#undef GetRenderTargetFormat
-#undef ED3D11ShaderOffsetBuffer
-#undef FindShaderResourceDXGIFormat
-#undef FindUnorderedAccessDXGIFormat
-#undef FindDepthStencilDXGIFormat
-#undef HasStencilBits
-#undef FVector4VertexDeclaration
-#undef GLOBAL_CONSTANT_BUFFER_INDEX
-#undef MAX_CONSTANT_BUFFER_SLOTS
-#undef FD3DGPUProfiler
-#undef FRangeAllocator
 #endif // OCULUS_HMD_SUPPORTED_PLATFORMS_D3D12
 
 
diff --git a/Engine/Plugins/Runtime/OpenXR/Source/OpenXRHMD/Private/OpenXRHMD_Swapchain.cpp b/Engine/Plugins/Runtime/OpenXR/Source/OpenXRHMD/Private/OpenXRHMD_Swapchain.cpp
--- a/Engine/Plugins/Runtime/OpenXR/Source/OpenXRHMD/Private/OpenXRHMD_Swapchain.cpp
+++ b/Engine/Plugins/Runtime/OpenXR/Source/OpenXRHMD/Private/OpenXRHMD_Swapchain.cpp
@@ -220,10 +220,10 @@ FXRSwapChainPtr CreateSwapchain_D3D11(XrSession InSession, uint8 Format, uint32
 	{
 		// We need to convert typeless to typed formats to create a swapchain
 		DXGI_FORMAT PlatformFormat = (DXGI_FORMAT)GPixelFormats[InFormat].PlatformFormat;
-		PlatformFormat = FindDepthStencilDXGIFormat(PlatformFormat);
+		PlatformFormat = ::FindDepthStencilDXGIFormat(PlatformFormat);
 
 		// UE4 renders a gamma-corrected image so we need to use an sRGB format if available
-		PlatformFormat = FindShaderResourceDXGIFormat(PlatformFormat, true);
+		PlatformFormat = ::FindShaderResourceDXGIFormat(PlatformFormat, true);
 		return PlatformFormat;
 	};
 
@@ -258,10 +258,10 @@ FXRSwapChainPtr CreateSwapchain_D3D12(XrSession InSession, uint8 Format, uint32
 	{
 		// We need to convert typeless to typed formats to create a swapchain
 		DXGI_FORMAT PlatformFormat = (DXGI_FORMAT)GPixelFormats[InFormat].PlatformFormat;
-		PlatformFormat = FindDepthStencilDXGIFormat_D3D12(PlatformFormat);
+		PlatformFormat = D3D12RHI::FindDepthStencilDXGIFormat(PlatformFormat);
 
 		// UE4 renders a gamma-corrected image so we need to use an sRGB format if available
-		PlatformFormat = FindShaderResourceDXGIFormat_D3D12(PlatformFormat, true);
+		PlatformFormat = D3D12RHI::FindShaderResourceDXGIFormat(PlatformFormat, true);
 		return PlatformFormat;
 	};
 
diff --git a/Engine/Plugins/Runtime/OpenXR/Source/OpenXRHMD/Private/OpenXRPlatformRHI.h b/Engine/Plugins/Runtime/OpenXR/Source/OpenXRHMD/Private/OpenXRPlatformRHI.h
--- a/Engine/Plugins/Runtime/OpenXR/Source/OpenXRHMD/Private/OpenXRPlatformRHI.h
+++ b/Engine/Plugins/Runtime/OpenXR/Source/OpenXRHMD/Private/OpenXRPlatformRHI.h
@@ -35,40 +35,8 @@
 //-------------------------------------------------------------------------------------------------
 
 #ifdef XR_USE_GRAPHICS_API_D3D12
-#define GetD3D11CubeFace GetD3D12CubeFace
-#define VerifyD3D11Result VerifyD3D12Result
-#define GetD3D11TextureFromRHITexture GetD3D12TextureFromRHITexture
-#define FRingAllocation FRingAllocation_D3D12
-#define GetRenderTargetFormat GetRenderTargetFormat_D3D12
-#define ED3D11ShaderOffsetBuffer ED3D12ShaderOffsetBuffer
-#define FindShaderResourceDXGIFormat FindShaderResourceDXGIFormat_D3D12
-#define FindUnorderedAccessDXGIFormat FindUnorderedAccessDXGIFormat_D3D12
-#define FindDepthStencilDXGIFormat FindDepthStencilDXGIFormat_D3D12
-#define HasStencilBits HasStencilBits_D3D12
-#define FVector4VertexDeclaration FVector4VertexDeclaration_D3D12
-#define GLOBAL_CONSTANT_BUFFER_INDEX GLOBAL_CONSTANT_BUFFER_INDEX_D3D12
-#define MAX_CONSTANT_BUFFER_SLOTS MAX_CONSTANT_BUFFER_SLOTS_D3D12
-#define FD3DGPUProfiler FD3D12GPUProfiler
-#define FRangeAllocator FRangeAllocator_D3D12
-
 #include "D3D12RHIPrivate.h"
 #include "D3D12Util.h"
-
-#undef GetD3D11CubeFace
-#undef VerifyD3D11Result
-#undef GetD3D11TextureFromRHITexture
-#undef FRingAllocation
-#undef GetRenderTargetFormat
-#undef ED3D11ShaderOffsetBuffer
-#undef FindShaderResourceDXGIFormat
-#undef FindUnorderedAccessDXGIFormat
-#undef FindDepthStencilDXGIFormat
-#undef HasStencilBits
-#undef FVector4VertexDeclaration
-#undef GLOBAL_CONSTANT_BUFFER_INDEX
-#undef MAX_CONSTANT_BUFFER_SLOTS
-#undef FD3DGPUProfiler
-#undef FRangeAllocator
 #endif // XR_USE_GRAPHICS_API_D3D12
 
 
diff --git a/Engine/Plugins/Runtime/nDisplay/Source/DisplayClusterProjection/Private/Policy/VIOSO/DisplayClusterProjectionVIOSOPolicy.cpp b/Engine/Plugins/Runtime/nDisplay/Source/DisplayClusterProjection/Private/Policy/VIOSO/DisplayClusterProjectionVIOSOPolicy.cpp
--- a/Engine/Plugins/Runtime/nDisplay/Source/DisplayClusterProjection/Private/Policy/VIOSO/DisplayClusterProjectionVIOSOPolicy.cpp
+++ b/Engine/Plugins/Runtime/nDisplay/Source/DisplayClusterProjection/Private/Policy/VIOSO/DisplayClusterProjectionVIOSOPolicy.cpp
@@ -30,42 +30,8 @@
 // D3D12
 //-------------------------------------------------------------------------------------------------
 #if VIOSO_USE_GRAPHICS_API_D3D12
-
-#define GetD3D11CubeFace GetD3D12CubeFace
-#define VerifyD3D11Result VerifyD3D12Result
-#define GetD3D11TextureFromRHITexture GetD3D12TextureFromRHITexture
-#define FRingAllocation FRingAllocation_D3D12
-#define GetRenderTargetFormat GetRenderTargetFormat_D3D12
-#define ED3D11ShaderOffsetBuffer ED3D12ShaderOffsetBuffer
-#define FindShaderResourceDXGIFormat FindShaderResourceDXGIFormat_D3D12
-#define FindUnorderedAccessDXGIFormat FindUnorderedAccessDXGIFormat_D3D12
-#define FindDepthStencilDXGIFormat FindDepthStencilDXGIFormat_D3D12
-#define HasStencilBits HasStencilBits_D3D12
-#define FVector4VertexDeclaration FVector4VertexDeclaration_D3D12
-#define GLOBAL_CONSTANT_BUFFER_INDEX GLOBAL_CONSTANT_BUFFER_INDEX_D3D12
-#define MAX_CONSTANT_BUFFER_SLOTS MAX_CONSTANT_BUFFER_SLOTS_D3D12
-#define FD3DGPUProfiler FD3D12GPUProfiler
-#define FRangeAllocator FRangeAllocator_D3D12
-
 #include "D3D12RHIPrivate.h"
 #include "D3D12Util.h"
-
-#undef GetD3D11CubeFace
-#undef VerifyD3D11Result
-#undef GetD3D11TextureFromRHITexture
-#undef FRingAllocation
-#undef GetRenderTargetFormat
-#undef ED3D11ShaderOffsetBuffer
-#undef FindShaderResourceDXGIFormat
-#undef FindUnorderedAccessDXGIFormat
-#undef FindDepthStencilDXGIFormat
-#undef HasStencilBits
-#undef FVector4VertexDeclaration
-#undef GLOBAL_CONSTANT_BUFFER_INDEX
-#undef MAX_CONSTANT_BUFFER_SLOTS
-#undef FD3DGPUProfiler
-#undef FRangeAllocator
-
 #endif // VIOSO_USE_GRAPHICS_API_D3D12
 
 static bool GetPooledTempRTT_RenderThread(FRHICommandListImmediate& RHICmdList, FIntPoint Size, EPixelFormat Format, bool bRTT, TRefCountPtr<IPooledRenderTarget>& OutPooledTempRTT)
diff --git a/Engine/Source/Runtime/D3D12RHI/Private/D3D12Device.h b/Engine/Source/Runtime/D3D12RHI/Private/D3D12Device.h
--- a/Engine/Source/Runtime/D3D12RHI/Private/D3D12Device.h
+++ b/Engine/Source/Runtime/D3D12RHI/Private/D3D12Device.h
@@ -130,7 +130,7 @@ public:
 
 	void BlockUntilIdle();
 
-	FORCEINLINE FD3DGPUProfiler& GetGPUProfiler() { return GPUProfilingData; }
+	FORCEINLINE D3D12RHI::FD3DGPUProfiler& GetGPUProfiler() { return GPUProfilingData; }
 
 protected:
 
@@ -213,7 +213,7 @@ protected:
 	FD3D12TimedIntervalQueryTracker* BackBufferWriteBarrierTracker = nullptr;
 #endif // #if PLATFORM_USE_BACKBUFFER_WRITE_TRANSITION_TRACKING
 
-	FD3DGPUProfiler GPUProfilingData;
+	D3D12RHI::FD3DGPUProfiler GPUProfilingData;
 };
 template <typename TDesc> 
 void TD3D12ViewDescriptorHandle<TDesc>::AllocateDescriptorSlot()
diff --git a/Engine/Source/Runtime/D3D12RHI/Private/D3D12RHIPrivate.h b/Engine/Source/Runtime/D3D12RHI/Private/D3D12RHIPrivate.h
--- a/Engine/Source/Runtime/D3D12RHI/Private/D3D12RHIPrivate.h
+++ b/Engine/Source/Runtime/D3D12RHI/Private/D3D12RHIPrivate.h
@@ -1320,6 +1320,185 @@ private:
 	}
 };
 
+// This namespace is needed to avoid a name clash with D3D11 RHI when linked together in monolithic builds. Otherwise the linker will just pick any variant instead of each RHI using their own version.
+namespace D3D12RHI
+{
+
+inline DXGI_FORMAT FindSharedResourceDXGIFormat(DXGI_FORMAT InFormat, bool bSRGB)
+{
+	if (bSRGB)
+	{
+		switch (InFormat)
+		{
+		case DXGI_FORMAT_B8G8R8X8_TYPELESS:    return DXGI_FORMAT_B8G8R8X8_UNORM_SRGB;
+		case DXGI_FORMAT_B8G8R8A8_TYPELESS:    return DXGI_FORMAT_B8G8R8A8_UNORM_SRGB;
+		case DXGI_FORMAT_R8G8B8A8_TYPELESS:    return DXGI_FORMAT_R8G8B8A8_UNORM_SRGB;
+		case DXGI_FORMAT_BC1_TYPELESS:         return DXGI_FORMAT_BC1_UNORM_SRGB;
+		case DXGI_FORMAT_BC2_TYPELESS:         return DXGI_FORMAT_BC2_UNORM_SRGB;
+		case DXGI_FORMAT_BC3_TYPELESS:         return DXGI_FORMAT_BC3_UNORM_SRGB;
+		case DXGI_FORMAT_BC7_TYPELESS:         return DXGI_FORMAT_BC7_UNORM_SRGB;
+		};
+	}
+	else
+	{
+		switch (InFormat)
+		{
+		case DXGI_FORMAT_B8G8R8X8_TYPELESS:    return DXGI_FORMAT_B8G8R8X8_UNORM;
+		case DXGI_FORMAT_B8G8R8A8_TYPELESS: return DXGI_FORMAT_B8G8R8A8_UNORM;
+		case DXGI_FORMAT_R8G8B8A8_TYPELESS: return DXGI_FORMAT_R8G8B8A8_UNORM;
+		case DXGI_FORMAT_BC1_TYPELESS:      return DXGI_FORMAT_BC1_UNORM;
+		case DXGI_FORMAT_BC2_TYPELESS:      return DXGI_FORMAT_BC2_UNORM;
+		case DXGI_FORMAT_BC3_TYPELESS:      return DXGI_FORMAT_BC3_UNORM;
+		case DXGI_FORMAT_BC7_TYPELESS:      return DXGI_FORMAT_BC7_UNORM;
+		};
+	}
+	switch (InFormat)
+	{
+	case DXGI_FORMAT_R32G32B32A32_TYPELESS: return DXGI_FORMAT_R32G32B32A32_UINT;
+	case DXGI_FORMAT_R32G32B32_TYPELESS:    return DXGI_FORMAT_R32G32B32_UINT;
+	case DXGI_FORMAT_R16G16B16A16_TYPELESS: return DXGI_FORMAT_R16G16B16A16_UNORM;
+	case DXGI_FORMAT_R32G32_TYPELESS:       return DXGI_FORMAT_R32G32_UINT;
+	case DXGI_FORMAT_R10G10B10A2_TYPELESS:  return DXGI_FORMAT_R10G10B10A2_UNORM;
+	case DXGI_FORMAT_R16G16_TYPELESS:       return DXGI_FORMAT_R16G16_UNORM;
+	case DXGI_FORMAT_R8G8_TYPELESS:         return DXGI_FORMAT_R8G8_UNORM;
+	case DXGI_FORMAT_R8_TYPELESS:           return DXGI_FORMAT_R8_UNORM;
+
+	case DXGI_FORMAT_BC4_TYPELESS:         return DXGI_FORMAT_BC4_UNORM;
+	case DXGI_FORMAT_BC5_TYPELESS:         return DXGI_FORMAT_BC5_UNORM;
+
+
+
+	case DXGI_FORMAT_R24G8_TYPELESS: return DXGI_FORMAT_R24_UNORM_X8_TYPELESS;
+	case DXGI_FORMAT_R32_TYPELESS: return DXGI_FORMAT_R32_FLOAT;
+	case DXGI_FORMAT_R16_TYPELESS: return DXGI_FORMAT_R16_UNORM;
+		// Changing Depth Buffers to 32 bit on Dingo as D24S8 is actually implemented as a 32 bit buffer in the hardware
+	case DXGI_FORMAT_R32G8X24_TYPELESS: return DXGI_FORMAT_R32_FLOAT_X8X24_TYPELESS;
+	}
+	return InFormat;
+}
+
+inline DXGI_FORMAT GetPlatformTextureResourceFormat(DXGI_FORMAT InFormat, uint32 InFlags)
+{
+	// Find valid shared texture format
+	if (InFlags & TexCreate_Shared)
+	{
+		return FindSharedResourceDXGIFormat(InFormat, InFlags & TexCreate_SRGB);
+	}
+
+	return InFormat;
+}
+
+/** Find an appropriate DXGI format for the input format and SRGB setting. */
+inline DXGI_FORMAT FindShaderResourceDXGIFormat(DXGI_FORMAT InFormat, bool bSRGB)
+{
+	if (bSRGB)
+	{
+		switch (InFormat)
+		{
+		case DXGI_FORMAT_B8G8R8A8_TYPELESS:    return DXGI_FORMAT_B8G8R8A8_UNORM_SRGB;
+		case DXGI_FORMAT_R8G8B8A8_TYPELESS:    return DXGI_FORMAT_R8G8B8A8_UNORM_SRGB;
+		case DXGI_FORMAT_BC1_TYPELESS:         return DXGI_FORMAT_BC1_UNORM_SRGB;
+		case DXGI_FORMAT_BC2_TYPELESS:         return DXGI_FORMAT_BC2_UNORM_SRGB;
+		case DXGI_FORMAT_BC3_TYPELESS:         return DXGI_FORMAT_BC3_UNORM_SRGB;
+		case DXGI_FORMAT_BC7_TYPELESS:         return DXGI_FORMAT_BC7_UNORM_SRGB;
+		};
+	}
+	else
+	{
+		switch (InFormat)
+		{
+		case DXGI_FORMAT_B8G8R8A8_TYPELESS: return DXGI_FORMAT_B8G8R8A8_UNORM;
+		case DXGI_FORMAT_R8G8B8A8_TYPELESS: return DXGI_FORMAT_R8G8B8A8_UNORM;
+		case DXGI_FORMAT_BC1_TYPELESS:      return DXGI_FORMAT_BC1_UNORM;
+		case DXGI_FORMAT_BC2_TYPELESS:      return DXGI_FORMAT_BC2_UNORM;
+		case DXGI_FORMAT_BC3_TYPELESS:      return DXGI_FORMAT_BC3_UNORM;
+		case DXGI_FORMAT_BC7_TYPELESS:      return DXGI_FORMAT_BC7_UNORM;
+		};
+	}
+	switch (InFormat)
+	{
+	case DXGI_FORMAT_R24G8_TYPELESS: return DXGI_FORMAT_R24_UNORM_X8_TYPELESS;
+	case DXGI_FORMAT_R32_TYPELESS: return DXGI_FORMAT_R32_FLOAT;
+	case DXGI_FORMAT_R16_TYPELESS: return DXGI_FORMAT_R16_UNORM;
+		// Changing Depth Buffers to 32 bit on Dingo as D24S8 is actually implemented as a 32 bit buffer in the hardware
+	case DXGI_FORMAT_R32G8X24_TYPELESS: return DXGI_FORMAT_R32_FLOAT_X8X24_TYPELESS;
+	}
+	return InFormat;
+}
+
+/** Find an appropriate DXGI format unordered access of the raw format. */
+inline DXGI_FORMAT FindUnorderedAccessDXGIFormat(DXGI_FORMAT InFormat)
+{
+	switch (InFormat)
+	{
+	case DXGI_FORMAT_B8G8R8A8_TYPELESS: return DXGI_FORMAT_B8G8R8A8_UNORM;
+	case DXGI_FORMAT_R8G8B8A8_TYPELESS: return DXGI_FORMAT_R8G8B8A8_UNORM;
+	}
+	return InFormat;
+}
+
+/** Find the appropriate depth-stencil targetable DXGI format for the given format. */
+inline DXGI_FORMAT FindDepthStencilDXGIFormat(DXGI_FORMAT InFormat)
+{
+	switch (InFormat)
+	{
+	case DXGI_FORMAT_R24G8_TYPELESS:
+		return DXGI_FORMAT_D24_UNORM_S8_UINT;
+		// Changing Depth Buffers to 32 bit on Dingo as D24S8 is actually implemented as a 32 bit buffer in the hardware
+	case DXGI_FORMAT_R32G8X24_TYPELESS:
+		return DXGI_FORMAT_D32_FLOAT_S8X24_UINT;
+	case DXGI_FORMAT_R32_TYPELESS:
+		return DXGI_FORMAT_D32_FLOAT;
+	case DXGI_FORMAT_R16_TYPELESS:
+		return DXGI_FORMAT_D16_UNORM;
+	};
+	return InFormat;
+}
+
+/**
+* Returns whether the given format contains stencil information.
+* Must be passed a format returned by FindDepthStencilDXGIFormat, so that typeless versions are converted to their corresponding depth stencil view format.
+*/
+inline bool HasStencilBits(DXGI_FORMAT InFormat)
+{
+	switch (InFormat)
+	{
+	case DXGI_FORMAT_D24_UNORM_S8_UINT:
+		// Changing Depth Buffers to 32 bit on Dingo as D24S8 is actually implemented as a 32 bit buffer in the hardware
+	case DXGI_FORMAT_D32_FLOAT_S8X24_UINT:
+		return true;
+	};
+
+	return false;
+}
+
+static void TranslateRenderTargetFormats(
+	const FGraphicsPipelineStateInitializer &PsoInit,
+	D3D12_RT_FORMAT_ARRAY& RTFormatArray,
+	DXGI_FORMAT& DSVFormat)
+{
+	RTFormatArray.NumRenderTargets = PsoInit.ComputeNumValidRenderTargets();
+
+	for (uint32 RTIdx = 0; RTIdx < PsoInit.RenderTargetsEnabled; ++RTIdx)
+	{
+		checkSlow(PsoInit.RenderTargetFormats[RTIdx] == PF_Unknown || GPixelFormats[PsoInit.RenderTargetFormats[RTIdx]].Supported);
+
+		DXGI_FORMAT PlatformFormat = (DXGI_FORMAT)GPixelFormats[PsoInit.RenderTargetFormats[RTIdx]].PlatformFormat;
+		uint32 Flags = PsoInit.RenderTargetFlags[RTIdx];
+
+		RTFormatArray.RTFormats[RTIdx] = D3D12RHI::FindShaderResourceDXGIFormat( GetPlatformTextureResourceFormat(PlatformFormat, Flags), (Flags & TexCreate_SRGB) != 0 );
+	}
+
+	checkSlow(PsoInit.DepthStencilTargetFormat == PF_Unknown || GPixelFormats[PsoInit.DepthStencilTargetFormat].Supported);
+
+	DXGI_FORMAT PlatformFormat = (DXGI_FORMAT)GPixelFormats[PsoInit.DepthStencilTargetFormat].PlatformFormat;
+	uint32 Flags = PsoInit.DepthStencilTargetFlag;
+
+	DSVFormat = D3D12RHI::FindDepthStencilDXGIFormat( GetPlatformTextureResourceFormat(PlatformFormat, PsoInit.DepthStencilTargetFlag) );
+}
+
+} // namespace D3D12RHI
+
 #if (PLATFORM_WINDOWS || PLATFORM_HOLOLENS)
 
 #ifndef DXGI_PRESENT_ALLOW_TEARING
diff --git a/Engine/Source/Runtime/D3D12RHI/Private/D3D12Viewport.cpp b/Engine/Source/Runtime/D3D12RHI/Private/D3D12Viewport.cpp
--- a/Engine/Source/Runtime/D3D12RHI/Private/D3D12Viewport.cpp
+++ b/Engine/Source/Runtime/D3D12RHI/Private/D3D12Viewport.cpp
@@ -232,7 +232,7 @@ FD3D12Texture2D* GetSwapChainSurface(FD3D12Device* Parent, EPixelFormat PixelFor
 		TextureDesc.Height = SizeY;
 		TextureDesc.DepthOrArraySize = 1;
 		TextureDesc.MipLevels = 1;
-		TextureDesc.Format = GetRenderTargetFormat(PixelFormat);
+		TextureDesc.Format = FD3D12Viewport::GetRenderTargetFormat(PixelFormat);
 		TextureDesc.SampleDesc.Count = 1;
 		TextureDesc.SampleDesc.Quality = 0;
 		TextureDesc.Layout = D3D12_TEXTURE_LAYOUT_UNKNOWN;
diff --git a/Engine/Source/Runtime/D3D12RHI/Public/D3D12Util.h b/Engine/Source/Runtime/D3D12RHI/Public/D3D12Util.h
--- a/Engine/Source/Runtime/D3D12RHI/Public/D3D12Util.h
+++ b/Engine/Source/Runtime/D3D12RHI/Public/D3D12Util.h
@@ -796,154 +796,6 @@ bool AssertResourceState(ID3D12CommandList* pCommandList, FD3D12View<TView>* pVi
 bool AssertResourceState(ID3D12CommandList* pCommandList, FD3D12Resource* pResource, const D3D12_RESOURCE_STATES& State, uint32 Subresource);
 bool AssertResourceState(ID3D12CommandList* pCommandList, FD3D12Resource* pResource, const D3D12_RESOURCE_STATES& State, const CViewSubresourceSubset& SubresourceSubset);
 
-inline DXGI_FORMAT FindSharedResourceDXGIFormat(DXGI_FORMAT InFormat, bool bSRGB)
-{
-	if (bSRGB)
-	{
-		switch (InFormat)
-		{
-		case DXGI_FORMAT_B8G8R8X8_TYPELESS:    return DXGI_FORMAT_B8G8R8X8_UNORM_SRGB;
-		case DXGI_FORMAT_B8G8R8A8_TYPELESS:    return DXGI_FORMAT_B8G8R8A8_UNORM_SRGB;
-		case DXGI_FORMAT_R8G8B8A8_TYPELESS:    return DXGI_FORMAT_R8G8B8A8_UNORM_SRGB;
-		case DXGI_FORMAT_BC1_TYPELESS:         return DXGI_FORMAT_BC1_UNORM_SRGB;
-		case DXGI_FORMAT_BC2_TYPELESS:         return DXGI_FORMAT_BC2_UNORM_SRGB;
-		case DXGI_FORMAT_BC3_TYPELESS:         return DXGI_FORMAT_BC3_UNORM_SRGB;
-		case DXGI_FORMAT_BC7_TYPELESS:         return DXGI_FORMAT_BC7_UNORM_SRGB;
-		};
-	}
-	else
-	{
-		switch (InFormat)
-		{
-		case DXGI_FORMAT_B8G8R8X8_TYPELESS:    return DXGI_FORMAT_B8G8R8X8_UNORM;
-		case DXGI_FORMAT_B8G8R8A8_TYPELESS: return DXGI_FORMAT_B8G8R8A8_UNORM;
-		case DXGI_FORMAT_R8G8B8A8_TYPELESS: return DXGI_FORMAT_R8G8B8A8_UNORM;
-		case DXGI_FORMAT_BC1_TYPELESS:      return DXGI_FORMAT_BC1_UNORM;
-		case DXGI_FORMAT_BC2_TYPELESS:      return DXGI_FORMAT_BC2_UNORM;
-		case DXGI_FORMAT_BC3_TYPELESS:      return DXGI_FORMAT_BC3_UNORM;
-		case DXGI_FORMAT_BC7_TYPELESS:      return DXGI_FORMAT_BC7_UNORM;
-		};
-	}
-	switch (InFormat)
-	{
-	case DXGI_FORMAT_R32G32B32A32_TYPELESS: return DXGI_FORMAT_R32G32B32A32_UINT;
-	case DXGI_FORMAT_R32G32B32_TYPELESS:    return DXGI_FORMAT_R32G32B32_UINT;
-	case DXGI_FORMAT_R16G16B16A16_TYPELESS: return DXGI_FORMAT_R16G16B16A16_UNORM;
-	case DXGI_FORMAT_R32G32_TYPELESS:       return DXGI_FORMAT_R32G32_UINT;
-	case DXGI_FORMAT_R10G10B10A2_TYPELESS:  return DXGI_FORMAT_R10G10B10A2_UNORM;
-	case DXGI_FORMAT_R16G16_TYPELESS:       return DXGI_FORMAT_R16G16_UNORM;
-	case DXGI_FORMAT_R8G8_TYPELESS:         return DXGI_FORMAT_R8G8_UNORM;
-	case DXGI_FORMAT_R8_TYPELESS:           return DXGI_FORMAT_R8_UNORM;
-
-	case DXGI_FORMAT_BC4_TYPELESS:         return DXGI_FORMAT_BC4_UNORM;
-	case DXGI_FORMAT_BC5_TYPELESS:         return DXGI_FORMAT_BC5_UNORM;
-
-
-
-	case DXGI_FORMAT_R24G8_TYPELESS: return DXGI_FORMAT_R24_UNORM_X8_TYPELESS;
-	case DXGI_FORMAT_R32_TYPELESS: return DXGI_FORMAT_R32_FLOAT;
-	case DXGI_FORMAT_R16_TYPELESS: return DXGI_FORMAT_R16_UNORM;
-		// Changing Depth Buffers to 32 bit on Dingo as D24S8 is actually implemented as a 32 bit buffer in the hardware
-	case DXGI_FORMAT_R32G8X24_TYPELESS: return DXGI_FORMAT_R32_FLOAT_X8X24_TYPELESS;
-	}
-	return InFormat;
-}
-
-inline DXGI_FORMAT GetPlatformTextureResourceFormat(DXGI_FORMAT InFormat, uint32 InFlags)
-{
-	// Find valid shared texture format
-	if (InFlags & TexCreate_Shared)
-	{
-		return FindSharedResourceDXGIFormat(InFormat, InFlags & TexCreate_SRGB);
-	}
-
-	return InFormat;
-}
-
-/** Find an appropriate DXGI format for the input format and SRGB setting. */
-inline DXGI_FORMAT FindShaderResourceDXGIFormat(DXGI_FORMAT InFormat, bool bSRGB)
-{
-	if (bSRGB)
-	{
-		switch (InFormat)
-		{
-		case DXGI_FORMAT_B8G8R8A8_TYPELESS:    return DXGI_FORMAT_B8G8R8A8_UNORM_SRGB;
-		case DXGI_FORMAT_R8G8B8A8_TYPELESS:    return DXGI_FORMAT_R8G8B8A8_UNORM_SRGB;
-		case DXGI_FORMAT_BC1_TYPELESS:         return DXGI_FORMAT_BC1_UNORM_SRGB;
-		case DXGI_FORMAT_BC2_TYPELESS:         return DXGI_FORMAT_BC2_UNORM_SRGB;
-		case DXGI_FORMAT_BC3_TYPELESS:         return DXGI_FORMAT_BC3_UNORM_SRGB;
-		case DXGI_FORMAT_BC7_TYPELESS:         return DXGI_FORMAT_BC7_UNORM_SRGB;
-		};
-	}
-	else
-	{
-		switch (InFormat)
-		{
-		case DXGI_FORMAT_B8G8R8A8_TYPELESS: return DXGI_FORMAT_B8G8R8A8_UNORM;
-		case DXGI_FORMAT_R8G8B8A8_TYPELESS: return DXGI_FORMAT_R8G8B8A8_UNORM;
-		case DXGI_FORMAT_BC1_TYPELESS:      return DXGI_FORMAT_BC1_UNORM;
-		case DXGI_FORMAT_BC2_TYPELESS:      return DXGI_FORMAT_BC2_UNORM;
-		case DXGI_FORMAT_BC3_TYPELESS:      return DXGI_FORMAT_BC3_UNORM;
-		case DXGI_FORMAT_BC7_TYPELESS:      return DXGI_FORMAT_BC7_UNORM;
-		};
-	}
-	switch (InFormat)
-	{
-	case DXGI_FORMAT_R24G8_TYPELESS: return DXGI_FORMAT_R24_UNORM_X8_TYPELESS;
-	case DXGI_FORMAT_R32_TYPELESS: return DXGI_FORMAT_R32_FLOAT;
-	case DXGI_FORMAT_R16_TYPELESS: return DXGI_FORMAT_R16_UNORM;
-		// Changing Depth Buffers to 32 bit on Dingo as D24S8 is actually implemented as a 32 bit buffer in the hardware
-	case DXGI_FORMAT_R32G8X24_TYPELESS: return DXGI_FORMAT_R32_FLOAT_X8X24_TYPELESS;
-	}
-	return InFormat;
-}
-
-/** Find an appropriate DXGI format unordered access of the raw format. */
-inline DXGI_FORMAT FindUnorderedAccessDXGIFormat(DXGI_FORMAT InFormat)
-{
-	switch (InFormat)
-	{
-	case DXGI_FORMAT_B8G8R8A8_TYPELESS: return DXGI_FORMAT_B8G8R8A8_UNORM;
-	case DXGI_FORMAT_R8G8B8A8_TYPELESS: return DXGI_FORMAT_R8G8B8A8_UNORM;
-	}
-	return InFormat;
-}
-
-/** Find the appropriate depth-stencil targetable DXGI format for the given format. */
-inline DXGI_FORMAT FindDepthStencilDXGIFormat(DXGI_FORMAT InFormat)
-{
-	switch (InFormat)
-	{
-	case DXGI_FORMAT_R24G8_TYPELESS:
-		return DXGI_FORMAT_D24_UNORM_S8_UINT;
-		// Changing Depth Buffers to 32 bit on Dingo as D24S8 is actually implemented as a 32 bit buffer in the hardware
-	case DXGI_FORMAT_R32G8X24_TYPELESS:
-		return DXGI_FORMAT_D32_FLOAT_S8X24_UINT;
-	case DXGI_FORMAT_R32_TYPELESS:
-		return DXGI_FORMAT_D32_FLOAT;
-	case DXGI_FORMAT_R16_TYPELESS:
-		return DXGI_FORMAT_D16_UNORM;
-	};
-	return InFormat;
-}
-
-/**
-* Returns whether the given format contains stencil information.
-* Must be passed a format returned by FindDepthStencilDXGIFormat, so that typeless versions are converted to their corresponding depth stencil view format.
-*/
-inline bool HasStencilBits(DXGI_FORMAT InFormat)
-{
-	switch (InFormat)
-	{
-	case DXGI_FORMAT_D24_UNORM_S8_UINT:
-		// Changing Depth Buffers to 32 bit on Dingo as D24S8 is actually implemented as a 32 bit buffer in the hardware
-	case DXGI_FORMAT_D32_FLOAT_S8X24_UINT:
-		return true;
-	};
-
-	return false;
-}
-
 FORCEINLINE_DEBUGGABLE D3D12_PRIMITIVE_TOPOLOGY_TYPE TranslatePrimitiveTopologyType(EPrimitiveTopologyType TopologyType)
 {
 	switch (TopologyType)
@@ -1043,37 +895,6 @@ FORCEINLINE_DEBUGGABLE D3D12_PRIMITIVE_TOPOLOGY_TYPE D3D12PrimitiveTypeToTopolog
 }
 #pragma warning(pop)
 
-static void TranslateRenderTargetFormats(
-	const FGraphicsPipelineStateInitializer &PsoInit,
-	D3D12_RT_FORMAT_ARRAY& RTFormatArray,
-	DXGI_FORMAT& DSVFormat
-	)
-{
-	RTFormatArray.NumRenderTargets = PsoInit.ComputeNumValidRenderTargets();
-
-	for (uint32 RTIdx = 0; RTIdx < PsoInit.RenderTargetsEnabled; ++RTIdx)
-	{
-		checkSlow(PsoInit.RenderTargetFormats[RTIdx] == PF_Unknown || GPixelFormats[PsoInit.RenderTargetFormats[RTIdx]].Supported);
-
-		DXGI_FORMAT PlatformFormat = (DXGI_FORMAT)GPixelFormats[PsoInit.RenderTargetFormats[RTIdx]].PlatformFormat;
-		uint32 Flags = PsoInit.RenderTargetFlags[RTIdx];
-
-		RTFormatArray.RTFormats[RTIdx] = FindShaderResourceDXGIFormat(
-			GetPlatformTextureResourceFormat(PlatformFormat, Flags),
-			(Flags & TexCreate_SRGB) != 0
-			);
-	}
-
-	checkSlow(PsoInit.DepthStencilTargetFormat == PF_Unknown || GPixelFormats[PsoInit.DepthStencilTargetFormat].Supported);
-
-	DXGI_FORMAT PlatformFormat = (DXGI_FORMAT)GPixelFormats[PsoInit.DepthStencilTargetFormat].PlatformFormat;
-	uint32 Flags = PsoInit.DepthStencilTargetFlag;
-
-	DSVFormat = FindDepthStencilDXGIFormat(
-		GetPlatformTextureResourceFormat(PlatformFormat, PsoInit.DepthStencilTargetFlag)
-		);
-}
-
 // @return 0xffffffff if not not supported
 FORCEINLINE_DEBUGGABLE uint32 GetMaxMSAAQuality(uint32 SampleCount)
 {
diff --git a/Engine/Source/Runtime/Windows/D3D11RHI/Private/D3D11Device.cpp b/Engine/Source/Runtime/Windows/D3D11RHI/Private/D3D11Device.cpp
--- a/Engine/Source/Runtime/Windows/D3D11RHI/Private/D3D11Device.cpp
+++ b/Engine/Source/Runtime/Windows/D3D11RHI/Private/D3D11Device.cpp
@@ -641,7 +641,7 @@ void* FD3D11DynamicRHI::RHIGetNativeInstance()
 
 static bool CanFormatBeDisplayed(const FD3D11DynamicRHI* InD3DRHI, EPixelFormat InPixelFormat)
 {
-	const DXGI_FORMAT DxgiFormat = GetRenderTargetFormat(InPixelFormat);
+	const DXGI_FORMAT DxgiFormat = FD3D11Viewport::GetRenderTargetFormat(InPixelFormat);
 
 	UINT FormatSupport = 0;
 	const HRESULT FormatSupportResult = InD3DRHI->GetDevice()->CheckFormatSupport(DxgiFormat, &FormatSupport);
