Fixed redundant TLAS buffer allocation. Added support for creating TLAS SRV with a buffer offset.

diff --git a/Engine/Source/Runtime/D3D12RHI/Private/D3D12RayTracing.cpp b/Engine/Source/Runtime/D3D12RHI/Private/D3D12RayTracing.cpp
--- a/Engine/Source/Runtime/D3D12RHI/Private/D3D12RayTracing.cpp
+++ b/Engine/Source/Runtime/D3D12RHI/Private/D3D12RayTracing.cpp
@@ -3145,11 +3145,6 @@ FD3D12RayTracingScene::FD3D12RayTracingScene(FD3D12Adapter* Adapter, const FRayT
 
 	DebugName = Initializer.DebugName.IsValid() ? Initializer.DebugName : FName(TEXT("TLAS"));
 
-	ShaderResourceView = Adapter->CreateLinkedObject<FD3D12ShaderResourceView>(FRHIGPUMask::All(), [&](FD3D12Device* Device)
-	{
-		return new FD3D12ShaderResourceView(Device);
-	});
-
 	checkf(Initializer.Lifetime == RTSL_SingleFrame, TEXT("Only single-frame ray tracing scenes are currently implemented."));
 
 	Lifetime = Initializer.Lifetime;
@@ -3290,24 +3285,6 @@ FD3D12RayTracingScene::FD3D12RayTracingScene(FD3D12Adapter* Adapter, const FRayT
 
 	// Get maximum buffer sizes for all GPUs in the system
 	SizeInfo = RHICalcRayTracingSceneSize(BuildInputs.NumDescs, BuildFlags);
-
-	for (uint32 GPUIndex = 0; GPUIndex < GNumExplicitGPUsForRendering; ++GPUIndex)
-	{	
-		AccelerationStructureBuffers[GPUIndex] = CreateRayTracingBuffer(Adapter, GPUIndex, SizeInfo.ResultSize, ERayTracingBufferType::AccelerationStructure, DebugName);
-
-		INC_MEMORY_STAT_BY(STAT_D3D12RayTracingUsedVideoMemory, AccelerationStructureBuffers[GPUIndex]->GetSize());
-		INC_MEMORY_STAT_BY(STAT_D3D12RayTracingTLASMemory, AccelerationStructureBuffers[GPUIndex]->GetSize());
-	
-		FD3D12ShaderResourceView* View = FD3D12CommandContext::RetrieveObject<FD3D12ShaderResourceView>(ShaderResourceView.GetReference(), GPUIndex);
-
-		D3D12_SHADER_RESOURCE_VIEW_DESC SRVDesc = {};
-		SRVDesc.Format = DXGI_FORMAT_UNKNOWN;
-		SRVDesc.ViewDimension = D3D12_SRV_DIMENSION_RAYTRACING_ACCELERATION_STRUCTURE;
-		SRVDesc.Shader4ComponentMapping = D3D12_DEFAULT_SHADER_4_COMPONENT_MAPPING;
-		SRVDesc.RaytracingAccelerationStructure.Location = AccelerationStructureBuffers[GPUIndex]->ResourceLocation.GetGPUVirtualAddress();
-
-		View->Initialize(SRVDesc, AccelerationStructureBuffers[GPUIndex], 4, 0);
-	}
 };
 
 FD3D12RayTracingScene::~FD3D12RayTracingScene()
@@ -3338,9 +3315,21 @@ void FD3D12RayTracingScene::BindBuffer(FRHIBuffer* InBuffer, uint32 InBufferOffs
 
 	for (uint32 GPUIndex = 0; GPUIndex < GNumExplicitGPUsForRendering; ++GPUIndex)
 	{
+		if (AccelerationStructureBuffers[GPUIndex])
+		{
+			DEC_MEMORY_STAT_BY(STAT_D3D12RayTracingUsedVideoMemory, AccelerationStructureBuffers[GPUIndex]->GetSize());
+			DEC_MEMORY_STAT_BY(STAT_D3D12RayTracingTLASMemory, AccelerationStructureBuffers[GPUIndex]->GetSize());
+		}
+
 		AccelerationStructureBuffers[GPUIndex] = FD3D12CommandContext::RetrieveObject<FD3D12Buffer>(InBuffer, GPUIndex);
+
+		INC_MEMORY_STAT_BY(STAT_D3D12RayTracingUsedVideoMemory, AccelerationStructureBuffers[GPUIndex]->GetSize());
+		INC_MEMORY_STAT_BY(STAT_D3D12RayTracingTLASMemory, AccelerationStructureBuffers[GPUIndex]->GetSize());
 	}
 	BufferOffset = InBufferOffset;
+
+	FShaderResourceViewInitializer ViewInitializer(InBuffer, InBufferOffset, 0);
+	ShaderResourceView = RHICreateShaderResourceView(ViewInitializer);
 }
 
 void FD3D12RayTracingScene::BuildAccelerationStructure(FD3D12CommandContext& CommandContext)
@@ -4426,7 +4415,7 @@ void FD3D12CommandContext::RHIRayTraceOcclusion(FRHIRayTracingScene* InScene,
 	DispatchDesc.Depth = 1;
 
 	FRayTracingShaderBindings Bindings;
-	Bindings.SRVs[0] = Scene->GetShaderResourceView();
+	Bindings.SRVs[0] = Scene->ShaderResourceView;
 	Bindings.SRVs[1] = Rays;
 	Bindings.UAVs[0] = Output;
 
@@ -4463,7 +4452,7 @@ void FD3D12CommandContext::RHIRayTraceIntersection(FRHIRayTracingScene* InScene,
 	DispatchDesc.Depth = 1;
 
 	FRayTracingShaderBindings Bindings;
-	Bindings.SRVs[0] = Scene->GetShaderResourceView();
+	Bindings.SRVs[0] = Scene->ShaderResourceView;
 	Bindings.SRVs[1] = Rays;
 	// #dxr_todo: intersection and occlusion shaders should be split into separate files to avoid resource slot collisions.
 	// Workaround for now is to bind a valid UAV to slots 0 and 1, even though only slot 1 is referenced.
diff --git a/Engine/Source/Runtime/D3D12RHI/Private/D3D12RayTracing.h b/Engine/Source/Runtime/D3D12RHI/Private/D3D12RayTracing.h
--- a/Engine/Source/Runtime/D3D12RHI/Private/D3D12RayTracing.h
+++ b/Engine/Source/Runtime/D3D12RHI/Private/D3D12RayTracing.h
@@ -96,6 +96,8 @@ public:
 	void BindBuffer(FRHIBuffer* Buffer, uint32 BufferOffset);
 	void BuildAccelerationStructure(FD3D12CommandContext& CommandContext);
 
+	FShaderResourceViewRHIRef ShaderResourceView;
+
 	TRefCountPtr<FD3D12Buffer> AccelerationStructureBuffers[MAX_NUM_GPUS];
 	uint32 BufferOffset = 0;
 
diff --git a/Engine/Source/Runtime/RHI/Private/DynamicRHI.cpp b/Engine/Source/Runtime/RHI/Private/DynamicRHI.cpp
--- a/Engine/Source/Runtime/RHI/Private/DynamicRHI.cpp
+++ b/Engine/Source/Runtime/RHI/Private/DynamicRHI.cpp
@@ -554,10 +554,13 @@ FShaderResourceViewInitializer::FShaderResourceViewInitializer(FRHIBuffer* InBuf
 FShaderResourceViewInitializer::FShaderResourceViewInitializer(FRHIBuffer* InBuffer, uint32 InStartOffsetBytes, uint32 InNumElements)
 	: BufferInitializer({ InBuffer, InStartOffsetBytes, InNumElements, PF_Unknown }), Type(EType::StructuredBufferSRV)
 {
-	check(InStartOffsetBytes % InBuffer->GetStride() == 0);
+	const uint32 Stride = (InBuffer->GetUsage() & BUF_AccelerationStructure) 
+		? 1 // Acceleration structure buffers don't have a stride as they are opaque and not indexable
+		: InBuffer->GetStride();
+
+	check(InStartOffsetBytes % Stride == 0);
 	if (!BufferInitializer.IsWholeResource())
 	{
-		const uint32 Stride = BufferInitializer.Buffer->GetStride();
 		check((BufferInitializer.NumElements * Stride + BufferInitializer.StartOffsetBytes) <= BufferInitializer.Buffer->GetSize());
 	}
 
diff --git a/Engine/Source/Runtime/RHI/Public/RHIResources.h b/Engine/Source/Runtime/RHI/Public/RHIResources.h
--- a/Engine/Source/Runtime/RHI/Public/RHIResources.h
+++ b/Engine/Source/Runtime/RHI/Public/RHIResources.h
@@ -1561,13 +1561,7 @@ class FRHIRayTracingGeometry : public FRHIRayTracingAccelerationStructure {};
 typedef TRefCountPtr<FRHIRayTracingGeometry>     FRayTracingGeometryRHIRef;
 
 /** Top level ray tracing acceleration structure (contains instances of meshes). */
-class FRHIRayTracingScene : public FRHIRayTracingAccelerationStructure
-{
-public:
-	FRHIShaderResourceView* GetShaderResourceView() { return ShaderResourceView; }
-protected:
-	FShaderResourceViewRHIRef ShaderResourceView;
-};
+class FRHIRayTracingScene : public FRHIRayTracingAccelerationStructure {};
 
 typedef TRefCountPtr<FRHIRayTracingScene>        FRayTracingSceneRHIRef;
 
