Reverting changes to UAVs in Root Signatures when Mesh shaders were added. We only expose UAVs to ALL stages, not individual stages. This was primarily breaking static root signatures.

diff --git a/Engine/Source/Runtime/D3D12RHI/Private/D3D12RootSignature.h b/Engine/Source/Runtime/D3D12RHI/Private/D3D12RootSignature.h
--- a/Engine/Source/Runtime/D3D12RHI/Private/D3D12RootSignature.h
+++ b/Engine/Source/Runtime/D3D12RHI/Private/D3D12RootSignature.h
@@ -11,7 +11,6 @@ enum ERootParameterKeys
 	PS_CBVs,
 	PS_RootCBVs,
 	PS_Samplers,
-	PS_UAVs,
 	VS_SRVs,
 	VS_CBVs,
 	VS_RootCBVs,
@@ -32,12 +31,10 @@ enum ERootParameterKeys
 	MS_CBVs,
 	MS_RootCBVs,
 	MS_Samplers,
-	MS_UAVs,
 	AS_SRVs,
 	AS_CBVs,
 	AS_RootCBVs,
 	AS_Samplers,
-	AS_UAVs,
 	ALL_SRVs,
 	ALL_CBVs,
 	ALL_RootCBVs,
@@ -205,20 +202,8 @@ public:
 
 	inline uint32 UAVRDTBindSlot(EShaderFrequency ShaderStage) const
 	{
-		check(ShaderStage == SF_Pixel || ShaderStage == SF_Mesh || ShaderStage == SF_Amplification || ShaderStage == SF_Compute);
-
-		switch (ShaderStage)
-		{
-		case SF_Pixel: return BindSlotMap[PS_UAVs];
-		case SF_Mesh: return BindSlotMap[MS_UAVs];
-		case SF_Amplification: return BindSlotMap[AS_UAVs];
-
-		case SF_Compute:
-		case SF_NumFrequencies: return BindSlotMap[ALL_UAVs];
-
-		default: check(false);
-			return UINT_MAX;
-		}
+		check(ShaderStage == SF_Pixel || ShaderStage == SF_Compute);
+		return BindSlotMap[ALL_UAVs];
 	}
 
 	inline bool HasUAVs() const { return bHasUAVs; }
@@ -361,21 +346,9 @@ private:
 
 	inline void SetUAVRDTBindSlot(EShaderFrequency SF, uint8 RootParameterIndex)
 	{
-		check(SF == SF_Pixel || SF == SF_Mesh || SF == SF_Amplification || SF == SF_Compute || SF == SF_NumFrequencies);
+		check(SF == SF_Pixel || SF == SF_Compute || SF == SF_NumFrequencies);
 
-		uint8* pBindSlot = nullptr;
-		switch (SF)
-		{
-		case SF_Pixel: pBindSlot = &BindSlotMap[PS_UAVs]; break;
-		case SF_Mesh: pBindSlot = &BindSlotMap[MS_UAVs]; break;
-		case SF_Amplification: pBindSlot = &BindSlotMap[AS_UAVs]; break;
-
-		case SF_Compute:
-		case SF_NumFrequencies: pBindSlot = &BindSlotMap[ALL_UAVs]; break;
-
-		default: check(false);
-			return;
-		}
+		uint8* pBindSlot = &BindSlotMap[ALL_UAVs];
 
 		check(*pBindSlot == 0xFF);
 		*pBindSlot = RootParameterIndex;
diff --git a/Engine/Source/Runtime/D3D12RHI/Private/D3D12Util.cpp b/Engine/Source/Runtime/D3D12RHI/Private/D3D12Util.cpp
--- a/Engine/Source/Runtime/D3D12RHI/Private/D3D12Util.cpp
+++ b/Engine/Source/Runtime/D3D12RHI/Private/D3D12Util.cpp
@@ -818,13 +818,13 @@ void QuantizeBoundShaderState(
 
 	if (const FD3D12MeshShader* const MeshShader = BSS->GetMeshShader())
 	{
-		FD3D12QuantizedBoundShaderState::InitShaderRegisterCounts(ResourceBindingTier, MeshShader->ResourceCounts, QBSS.RegisterCounts[SV_Mesh], true);
+		FD3D12QuantizedBoundShaderState::InitShaderRegisterCounts(ResourceBindingTier, MeshShader->ResourceCounts, QBSS.RegisterCounts[SV_Mesh]);
 		QBSS.bUseVendorExtension |= UsesVendorExtensionSpace(*MeshShader);
 	}
 
 	if (const FD3D12AmplificationShader* const AmplificationShader = BSS->GetAmplificationShader())
 	{
-		FD3D12QuantizedBoundShaderState::InitShaderRegisterCounts(ResourceBindingTier, AmplificationShader->ResourceCounts, QBSS.RegisterCounts[SV_Amplification], true);
+		FD3D12QuantizedBoundShaderState::InitShaderRegisterCounts(ResourceBindingTier, AmplificationShader->ResourceCounts, QBSS.RegisterCounts[SV_Amplification]);
 		QBSS.bUseVendorExtension |= UsesVendorExtensionSpace(*AmplificationShader);
 	}
 
