Clear bound TLAS buffer in FD3D12CommandContext::RHIClearRayTracingBindings

diff --git a/Engine/Source/Runtime/D3D12RHI/Private/D3D12RayTracing.cpp b/Engine/Source/Runtime/D3D12RHI/Private/D3D12RayTracing.cpp
--- a/Engine/Source/Runtime/D3D12RHI/Private/D3D12RayTracing.cpp
+++ b/Engine/Source/Runtime/D3D12RHI/Private/D3D12RayTracing.cpp
@@ -3297,6 +3297,13 @@ FD3D12RayTracingScene::~FD3D12RayTracingScene()
 		}
 	}
 
+	ReleaseBuffer();
+
+	DEC_DWORD_STAT(STAT_D3D12RayTracingAllocatedTLAS);
+}
+
+void FD3D12RayTracingScene::ReleaseBuffer()
+{
 	for (auto& AccelerationStructureBuffer : AccelerationStructureBuffers)
 	{
 		if (AccelerationStructureBuffer)
@@ -3305,8 +3312,7 @@ FD3D12RayTracingScene::~FD3D12RayTracingScene()
 			DEC_MEMORY_STAT_BY(STAT_D3D12RayTracingTLASMemory, AccelerationStructureBuffer->GetSize());
 		}
 	}
-
-	DEC_DWORD_STAT(STAT_D3D12RayTracingAllocatedTLAS);
+	ShaderResourceView = nullptr;
 }
 
 void FD3D12RayTracingScene::BindBuffer(FRHIBuffer* InBuffer, uint32 InBufferOffset)
@@ -3659,6 +3665,7 @@ void FD3D12CommandContext::RHIBindAccelerationStructureMemory(FRHIRayTracingScen
 void FD3D12CommandContext::RHIClearRayTracingBindings(FRHIRayTracingScene* InScene)
 {
 	FD3D12RayTracingScene* Scene = FD3D12DynamicRHI::ResourceCast(InScene);
+	check(Scene);
 
 	auto& Table = Scene->ShaderTables[GetGPUIndex()];
 	for (auto Item : Table)
@@ -3666,6 +3673,8 @@ void FD3D12CommandContext::RHIClearRayTracingBindings(FRHIRayTracingScene* InSce
 		delete Item.Value;
 	}
 	Table.Reset();
+
+	Scene->ReleaseBuffer();
 }
 
 struct FD3D12RayTracingGlobalResourceBinder
diff --git a/Engine/Source/Runtime/D3D12RHI/Private/D3D12RayTracing.h b/Engine/Source/Runtime/D3D12RHI/Private/D3D12RayTracing.h
--- a/Engine/Source/Runtime/D3D12RHI/Private/D3D12RayTracing.h
+++ b/Engine/Source/Runtime/D3D12RHI/Private/D3D12RayTracing.h
@@ -94,6 +94,8 @@ public:
 	~FD3D12RayTracingScene();
 
 	void BindBuffer(FRHIBuffer* Buffer, uint32 BufferOffset);
+	void ReleaseBuffer();
+
 	void BuildAccelerationStructure(FD3D12CommandContext& CommandContext);
 
 	FShaderResourceViewRHIRef ShaderResourceView;
