Added support for using an explicit scratch buffer for building ray tracing TLAS. Switched to this path in FDeferredShadingSceneRenderer, using RDG transient resource system.

diff --git a/Engine/Source/Runtime/D3D12RHI/Private/D3D12RayTracing.cpp b/Engine/Source/Runtime/D3D12RHI/Private/D3D12RayTracing.cpp
--- a/Engine/Source/Runtime/D3D12RHI/Private/D3D12RayTracing.cpp
+++ b/Engine/Source/Runtime/D3D12RHI/Private/D3D12RayTracing.cpp
@@ -3338,27 +3338,38 @@ void FD3D12RayTracingScene::BindBuffer(FRHIBuffer* InBuffer, uint32 InBufferOffs
 	ShaderResourceView = RHICreateShaderResourceView(ViewInitializer);
 }
 
-void FD3D12RayTracingScene::BuildAccelerationStructure(FD3D12CommandContext& CommandContext)
+void FD3D12RayTracingScene::BuildAccelerationStructure(FD3D12CommandContext& CommandContext,
+	FD3D12Buffer* ScratchBuffer, uint32 ScratchBufferOffset,
+	FD3D12Buffer* InstanceBuffer, uint32 InstanceBufferOffset,
+	uint32 NumInstanceDescs)
 {
 	TRACE_CPUPROFILER_EVENT_SCOPE(BuildAccelerationStructure_TopLevel);
 	SCOPE_CYCLE_COUNTER(STAT_D3D12BuildTLAS);
 
-	TRefCountPtr<FD3D12Buffer> InstanceBuffer;
-	TRefCountPtr<FD3D12Buffer> ScratchBuffer;
-
 	const uint32 GPUIndex = CommandContext.GetGPUIndex();
 	FD3D12Adapter* Adapter = CommandContext.GetParentAdapter();
 	ID3D12Device5* RayTracingDevice = CommandContext.GetParentDevice()->GetDevice5();
 
 	D3D12_RAYTRACING_ACCELERATION_STRUCTURE_PREBUILD_INFO PrebuildInfo = {};
+
 	RayTracingDevice->GetRaytracingAccelerationStructurePrebuildInfo(&BuildInputs, &PrebuildInfo);
 	check(PrebuildInfo.ResultDataMaxSizeInBytes <= SizeInfo.ResultSize);
 	check(PrebuildInfo.ScratchDataSizeInBytes <= SizeInfo.BuildScratchSize);
 
-	static const FName ScratchBufferName("BuildScratchTLAS");
-	ScratchBuffer = CreateRayTracingBuffer(Adapter, GPUIndex, PrebuildInfo.ScratchDataSizeInBytes, ERayTracingBufferType::Scratch, ScratchBufferName);
+	TRefCountPtr<FD3D12Buffer> AutoScratchBuffer;
+	if (ScratchBuffer == nullptr)
+	{
+		static const FName ScratchBufferName("AutoBuildScratchTLAS");
+		AutoScratchBuffer = CreateRayTracingBuffer(Adapter, GPUIndex, PrebuildInfo.ScratchDataSizeInBytes, ERayTracingBufferType::Scratch, ScratchBufferName);
+		ScratchBuffer = AutoScratchBuffer.GetReference();
+		ScratchBufferOffset = 0;
+	}
 
-	// Create and fill instance buffer
+	check(PrebuildInfo.ScratchDataSizeInBytes + ScratchBufferOffset <= ScratchBuffer->GetSize());
+
+	TRefCountPtr<FD3D12Buffer> AutoInstanceBuffer;
+
+	// Create and fill automatic instance buffer
 
 	if (BuildInputs.NumDescs)
 	{
@@ -3402,10 +3413,14 @@ void FD3D12RayTracingScene::BuildAccelerationStructure(FD3D12CommandContext& Com
 
 		// #dxr_todo multi state only when bShouldCopyIndirectInstances is set - will still need transition to copy_src for lock behind but this can be done on the complete pool in theory (have to check cost)
 		ID3D12ResourceAllocator* ResourceAllocator = nullptr;
-		InstanceBuffer = Adapter->CreateRHIBuffer(
+		AutoInstanceBuffer = Adapter->CreateRHIBuffer(
 			nullptr, InstanceBufferDesc, D3D12_RAYTRACING_INSTANCE_DESCS_BYTE_ALIGNMENT,
 			sizeof(D3D12_RAYTRACING_INSTANCE_DESC), InstanceBufferDesc.Width, BUF_UnorderedAccess, ED3D12ResourceStateMode::MultiState, ERHIAccess::UAVMask, CreateInfo, ResourceAllocator);
 
+		// #yuriy_todo: only create internal instance buffer if explicit one is not given
+		checkf(InstanceBuffer == nullptr, TEXT("Explicit instance buffer support is not yet implemented for D3D12 RHI."));
+		InstanceBuffer = AutoInstanceBuffer.GetReference();
+
 		{
 			TArray<FD3D12ResidencyHandle*>& GeometryResidencyHandlesForThisGPU = GeometryResidencyHandles[GPUIndex];
 
@@ -3460,7 +3475,7 @@ void FD3D12RayTracingScene::BuildAccelerationStructure(FD3D12CommandContext& Com
 			TRHICommandList_RecursiveHazardous<FD3D12CommandContext> RHICmdList(&CommandContext);
 			FUnorderedAccessViewRHIRef InstancesDescUAV = RHICreateUnorderedAccessView(InstanceBuffer, false, false);
 
-			FD3D12DynamicRHI::TransitionResource(CommandContext.CommandListHandle, InstanceBuffer.GetReference()->GetResource(),
+			FD3D12DynamicRHI::TransitionResource(CommandContext.CommandListHandle, InstanceBuffer->GetResource(),
 				D3D12_RESOURCE_STATE_TBD, D3D12_RESOURCE_STATE_UNORDERED_ACCESS, 0,
 				FD3D12DynamicRHI::ETransitionMode::Apply);
 
@@ -3475,7 +3490,7 @@ void FD3D12RayTracingScene::BuildAccelerationStructure(FD3D12CommandContext& Com
 			RHICmdList.EndUAVOverlap(InstancesDescUAV);
 		}
 
-		FD3D12DynamicRHI::TransitionResource(CommandContext.CommandListHandle, InstanceBuffer.GetReference()->GetResource(),
+		FD3D12DynamicRHI::TransitionResource(CommandContext.CommandListHandle, InstanceBuffer->GetResource(),
 			D3D12_RESOURCE_STATE_TBD, D3D12_RESOURCE_STATE_PIXEL_SHADER_RESOURCE | D3D12_RESOURCE_STATE_NON_PIXEL_SHADER_RESOURCE, 0,
 			FD3D12DynamicRHI::ETransitionMode::Apply);
 	}
@@ -3499,7 +3514,7 @@ void FD3D12RayTracingScene::BuildAccelerationStructure(FD3D12CommandContext& Com
 	BuildDesc.Inputs = BuildInputs;
 	BuildDesc.Inputs.InstanceDescs = InstanceBuffer ? InstanceBuffer->ResourceLocation.GetGPUVirtualAddress() : D3D12_GPU_VIRTUAL_ADDRESS(0);
 	BuildDesc.DestAccelerationStructureData = AccelerationStructureBuffer->ResourceLocation.GetGPUVirtualAddress() + BufferOffset;
-	BuildDesc.ScratchAccelerationStructureData = ScratchBuffer->ResourceLocation.GetGPUVirtualAddress();
+	BuildDesc.ScratchAccelerationStructureData = ScratchBuffer->ResourceLocation.GetGPUVirtualAddress() + ScratchBufferOffset;
 	BuildDesc.SourceAccelerationStructureData = D3D12_GPU_VIRTUAL_ADDRESS(0); // Null source TLAS as this is a build command
 
 	// UAV barrier is used here to ensure that all bottom level acceleration structures are built
@@ -3653,7 +3668,14 @@ void FD3D12CommandContext::RHIBuildAccelerationStructures(const TArrayView<const
 void FD3D12CommandContext::RHIBuildAccelerationStructure(const FRayTracingSceneBuildParams& SceneBuildParams)
 {
 	FD3D12RayTracingScene* Scene = FD3D12DynamicRHI::ResourceCast(SceneBuildParams.Scene);
-	Scene->BuildAccelerationStructure(*this);
+	FD3D12Buffer* ScratchBuffer = FD3D12DynamicRHI::ResourceCast(SceneBuildParams.ScratchBuffer);
+	FD3D12Buffer* InstanceBuffer = FD3D12DynamicRHI::ResourceCast(SceneBuildParams.InstanceBuffer);
+	Scene->BuildAccelerationStructure(
+		*this,
+		ScratchBuffer, SceneBuildParams.ScratchBufferOffset,
+		InstanceBuffer, SceneBuildParams.InstanceBufferOffset,
+		SceneBuildParams.NumInstances
+	);
 }
 
 void FD3D12CommandContext::RHIBindAccelerationStructureMemory(FRHIRayTracingScene* InScene, FRHIBuffer* InBuffer, uint32 InBufferOffset)
diff --git a/Engine/Source/Runtime/D3D12RHI/Private/D3D12RayTracing.h b/Engine/Source/Runtime/D3D12RHI/Private/D3D12RayTracing.h
--- a/Engine/Source/Runtime/D3D12RHI/Private/D3D12RayTracing.h
+++ b/Engine/Source/Runtime/D3D12RHI/Private/D3D12RayTracing.h
@@ -96,7 +96,11 @@ public:
 	void BindBuffer(FRHIBuffer* Buffer, uint32 BufferOffset);
 	void ReleaseBuffer();
 
-	void BuildAccelerationStructure(FD3D12CommandContext& CommandContext);
+	void BuildAccelerationStructure(FD3D12CommandContext& CommandContext,
+		FD3D12Buffer* ScratchBuffer, uint32 ScratchBufferOffset,
+		FD3D12Buffer* InstanceBuffer, uint32 InstanceBufferOffset,
+		uint32 NumInstanceDescs
+	);
 
 	FShaderResourceViewRHIRef ShaderResourceView;
 
diff --git a/Engine/Source/Runtime/Renderer/Private/DeferredShadingRenderer.cpp b/Engine/Source/Runtime/Renderer/Private/DeferredShadingRenderer.cpp
--- a/Engine/Source/Runtime/Renderer/Private/DeferredShadingRenderer.cpp
+++ b/Engine/Source/Runtime/Renderer/Private/DeferredShadingRenderer.cpp
@@ -1129,6 +1129,10 @@ static void DeduplicateRayGenerationShaders(TArray< FRHIRayTracingShader*>& RayG
 	RayGenShaders = UniqueRayGenShaders.Array();
 }
 
+BEGIN_SHADER_PARAMETER_STRUCT(FBuildAccelerationStructurePassParams, )
+RDG_BUFFER_ACCESS(RayTracingSceneScratchBuffer, ERHIAccess::UAVCompute)
+END_SHADER_PARAMETER_STRUCT()
+
 bool FDeferredShadingSceneRenderer::DispatchRayTracingWorldUpdates(FRDGBuilder& GraphBuilder)
 {
 	if (!IsRayTracingEnabled() || Views.Num() == 0)
@@ -1314,13 +1318,27 @@ bool FDeferredShadingSceneRenderer::DispatchRayTracingWorldUpdates(FRDGBuilder&
 
 		{
 			RDG_GPU_STAT_SCOPE(GraphBuilder, RayTracingScene);
-			AddPass(GraphBuilder, RDG_EVENT_NAME("RayTracingScene"), [this](FRHICommandList& RHICmdList)
+
+			FBuildAccelerationStructurePassParams* PassParams = GraphBuilder.AllocParameters<FBuildAccelerationStructurePassParams>();
+			FRDGBufferDesc ScratchBufferDesc;
+			ScratchBufferDesc.UnderlyingType = FRDGBufferDesc::EUnderlyingType::StructuredBuffer;
+			ScratchBufferDesc.Usage = BUF_UnorderedAccess;
+			ScratchBufferDesc.BytesPerElement = 4;
+			ScratchBufferDesc.NumElements = Scene->RayTracingScene.SizeInfo.BuildScratchSize / 4;
+			PassParams->RayTracingSceneScratchBuffer = GraphBuilder.CreateBuffer(ScratchBufferDesc, TEXT("RayTracingSceneScratchBuffer"));
+
+			GraphBuilder.AddPass(RDG_EVENT_NAME("RayTracingScene"), PassParams, ERDGPassFlags::Compute | ERDGPassFlags::NeverCull,
+				[this, PassParams](FRHICommandList& RHICmdList)
 			{
 				FRHIRayTracingScene* RayTracingSceneRHI = Scene->RayTracingScene.GetRHIRayTracingSceneChecked();
 
 				RHICmdList.BindAccelerationStructureMemory(RayTracingSceneRHI, Scene->RayTracingScene.GetBufferChecked(), 0);
 
-				RHICmdList.BuildAccelerationStructure(RayTracingSceneRHI);
+				FRayTracingSceneBuildParams BuildParams;
+				BuildParams.Scene = RayTracingSceneRHI;
+				BuildParams.ScratchBuffer = PassParams->RayTracingSceneScratchBuffer->GetRHI();
+				BuildParams.ScratchBufferOffset = 0;
+				RHICmdList.BuildAccelerationStructure(BuildParams);
 
 				// Submit potentially expensive BVH build commands to the GPU as soon as possible.
 				// Avoids a GPU bubble in some CPU-limited cases.
diff --git a/Engine/Source/Runtime/Renderer/Private/RayTracing/RayTracingScene.cpp b/Engine/Source/Runtime/Renderer/Private/RayTracing/RayTracingScene.cpp
--- a/Engine/Source/Runtime/Renderer/Private/RayTracing/RayTracingScene.cpp
+++ b/Engine/Source/Runtime/Renderer/Private/RayTracing/RayTracingScene.cpp
@@ -33,8 +33,7 @@ FGraphEventRef FRayTracingScene::BeginCreate()
 	static constexpr uint32 AllocationGranularity = 8 * 1024;
 	NumTotalInstances = FMath::DivideAndRoundUp(NumTotalInstances, AllocationGranularity) * AllocationGranularity;
 
-	// #yuriy_todo: allocate these using RDG transient resources!
-	FRayTracingAccelerationStructureSize SizeInfo = RHICalcRayTracingSceneSize(NumTotalInstances, ERayTracingAccelerationStructureFlags::FastTrace);
+	SizeInfo = RHICalcRayTracingSceneSize(NumTotalInstances, ERayTracingAccelerationStructureFlags::FastTrace);
 	check(SizeInfo.ResultSize < ~0u);
 
 	// Allocate GPU buffer if current one is too small or significantly larger than what we need.
diff --git a/Engine/Source/Runtime/Renderer/Private/RayTracing/RayTracingScene.h b/Engine/Source/Runtime/Renderer/Private/RayTracing/RayTracingScene.h
--- a/Engine/Source/Runtime/Renderer/Private/RayTracing/RayTracingScene.h
+++ b/Engine/Source/Runtime/Renderer/Private/RayTracing/RayTracingScene.h
@@ -70,6 +70,8 @@ public:
 	TSet<const FRayTracingGeometry*> GeometriesToValidate;
 #endif
 
+	FRayTracingAccelerationStructureSize SizeInfo = {};
+
 private:
 	// RHI object that abstracts mesh instnaces in this scene
 	FRayTracingSceneRHIRef RayTracingSceneRHI;
