Fixing static analysis errors.

diff --git a/Engine/Source/Runtime/D3D12RHI/Private/D3D12Util.cpp b/Engine/Source/Runtime/D3D12RHI/Private/D3D12Util.cpp
--- a/Engine/Source/Runtime/D3D12RHI/Private/D3D12Util.cpp
+++ b/Engine/Source/Runtime/D3D12RHI/Private/D3D12Util.cpp
@@ -837,13 +837,13 @@ void QuantizeBoundShaderState(
 	if (const FD3D12HullShader* const HullShader = BSS->GetHullShader())
 	{
 		FD3D12QuantizedBoundShaderState::InitShaderRegisterCounts(ResourceBindingTier, HullShader->ResourceCounts, QBSS.RegisterCounts[SV_Hull]);
-		QBSS.bUseVendorExtension |= UsesVendorExtensionSpace(*HullShader);
+		QBSS.bUseVendorExtension |= UsesVendorExtensionSpace(*HullShader); // -V522
 	}
 
 	if (const FD3D12DomainShader* const DomainShader = BSS->GetDomainShader())
 	{
 		FD3D12QuantizedBoundShaderState::InitShaderRegisterCounts(ResourceBindingTier, DomainShader->ResourceCounts, QBSS.RegisterCounts[SV_Domain]);
-		QBSS.bUseVendorExtension |= UsesVendorExtensionSpace(*DomainShader);
+		QBSS.bUseVendorExtension |= UsesVendorExtensionSpace(*DomainShader); // -V522
 	}
 
 	if (const FD3D12GeometryShader* const GeometryShader = BSS->GetGeometryShader())
