Significant improvements to D3D12 memory tracking/stats, and tracking/warning about video->system memory demotion.

* Removed FD3D12Adapter::GetLocalVideoMemoryInfo() and implemented UpdateMemoryInfo().
  * Local and Non-Local memory are both retrieved now.
  * Results are cached in a new FD3D12MemoryInfo struct on FD3D12Adapter. (updated only once per frame instead of every call like before).
  * Optimized FD3D12DynamicRHI::RHIGetTextureMemoryStats() to not call expensive DXGI query memory functions per texture.
* Implemented "Demoted Video Memory" D3D12 memory counter to show the amount of video memory that has been demoted to system memory.
  * aka. RIP performance
* Implemented "Used System Memory" D3D12 memory counter to show the amount of non-local memory allocated by a device adapter
  * Summation of demoted video memory, and also memory that should be non-local (i.e. upload heaps).
* Refactored all the D3D12 RHI stat counters.
  * Renamed/moved into cleaner groupings.
  * Avoid summations from being grouped in with children (making the stats hard to understand without looking at code).
  * Made it possible to hide various groupings where desired.
  * D3D12Memory -> D3D12Memory, D3D12Resources, D3D12Allocations
  * D3D12MemoryBuffer -> D3D12BufferAllocations
  * D3D12MemoryTransient -> D3D12TransientAllocations
  * D3D12PipelineState -> unchanged
  * D3D12DescriptorHeap -> unchanged
* Added global GDemotedLocalMemorySize to report from RHI to on screen warning system the amount of local video memory demoted to system memory.
* Implemented an on screen error message when video memory demotion occurs.
  * Does not depend on stats being displayed (but currently compiles out based on STATS)
  * Message display can be toggled with r.DemotedLocalMemoryWarning (default: on)

diff --git a/Engine/Source/Runtime/Core/Private/HAL/ConsoleManager.cpp b/Engine/Source/Runtime/Core/Private/HAL/ConsoleManager.cpp
--- a/Engine/Source/Runtime/Core/Private/HAL/ConsoleManager.cpp
+++ b/Engine/Source/Runtime/Core/Private/HAL/ConsoleManager.cpp
@@ -2682,6 +2682,12 @@ static TAutoConsoleVariable<int32> CVarPrecomputedVisibilityWarning(
 	TEXT("If set to 1, a warning will be displayed when rendering a scene from a view point without precomputed visibility."),
 	ECVF_RenderThreadSafe);
 
+static TAutoConsoleVariable<int32> CVarDemotedLocalMemoryWarning(
+	TEXT("r.DemotedLocalMemoryWarning"),
+	1,
+	TEXT("If set to 1, a warning will be displayed when local memory has been demoted to system memory."),
+	ECVF_RenderThreadSafe);
+
 static TAutoConsoleVariable<int32> CVarFeatureLevelPreview(
 	TEXT("r.FeatureLevelPreview"),
 	0,
diff --git a/Engine/Source/Runtime/D3D12RHI/Private/D3D12Adapter.cpp b/Engine/Source/Runtime/D3D12RHI/Private/D3D12Adapter.cpp
--- a/Engine/Source/Runtime/D3D12RHI/Private/D3D12Adapter.cpp
+++ b/Engine/Source/Runtime/D3D12RHI/Private/D3D12Adapter.cpp
@@ -218,7 +218,6 @@ void FD3D12Adapter::Initialize(FD3D12DynamicRHI* RHI)
 	OwningRHI = RHI;
 }
 
-
 /** Callback function called when the GPU crashes, when Aftermath is enabled */
 static void D3D12AftermathCrashCallback(const void* InGPUCrashDump, const uint32_t InGPUCrashDumpSize, void* InUserData)
 {
@@ -1215,13 +1214,40 @@ FD3D12FastConstantAllocator& FD3D12Adapter::GetTransientUniformBufferAllocator()
 	});
 }
 
-void FD3D12Adapter::GetLocalVideoMemoryInfo(DXGI_QUERY_VIDEO_MEMORY_INFO* LocalVideoMemoryInfo)
+void FD3D12Adapter::UpdateMemoryInfo()
 {
 #if PLATFORM_WINDOWS || PLATFORM_HOLOLENS
+	const uint64 UpdateFrame = FrameFence != nullptr ? FrameFence->GetCurrentFence() : 0;
+
+	// Avoid spurious query calls if we have already captured this frame.
+	if (MemoryInfo.UpdateFrameNumber == UpdateFrame)
+	{
+		return;
+	}
+
+	// Update the frame number that the memory is captured from.
+	MemoryInfo.UpdateFrameNumber = UpdateFrame;
+
 	TRefCountPtr<IDXGIAdapter3> Adapter3;
 	VERIFYD3D12RESULT(GetAdapter()->QueryInterface(IID_PPV_ARGS(Adapter3.GetInitReference())));
 
-	VERIFYD3D12RESULT(Adapter3->QueryVideoMemoryInfo(0, DXGI_MEMORY_SEGMENT_GROUP_LOCAL, LocalVideoMemoryInfo));
+	VERIFYD3D12RESULT(Adapter3->QueryVideoMemoryInfo(0, DXGI_MEMORY_SEGMENT_GROUP_LOCAL, &MemoryInfo.LocalMemoryInfo));
+	VERIFYD3D12RESULT(Adapter3->QueryVideoMemoryInfo(0, DXGI_MEMORY_SEGMENT_GROUP_NON_LOCAL, &MemoryInfo.NonLocalMemoryInfo));
+
+	// Over budget?
+	if (MemoryInfo.LocalMemoryInfo.CurrentUsage > MemoryInfo.LocalMemoryInfo.Budget)
+	{
+		MemoryInfo.AvailableLocalMemory = 0;
+		MemoryInfo.DemotedLocalMemory = MemoryInfo.LocalMemoryInfo.CurrentUsage - MemoryInfo.LocalMemoryInfo.Budget;
+	}
+	else
+	{
+		MemoryInfo.AvailableLocalMemory = MemoryInfo.LocalMemoryInfo.Budget - MemoryInfo.LocalMemoryInfo.CurrentUsage;
+		MemoryInfo.DemotedLocalMemory = 0;
+	}
+
+	// Update global RHI state (for warning output, etc.)
+	GDemotedLocalMemorySize = MemoryInfo.DemotedLocalMemory;
 
 	if (!GVirtualMGPU)
 	{
@@ -1229,8 +1255,15 @@ void FD3D12Adapter::GetLocalVideoMemoryInfo(DXGI_QUERY_VIDEO_MEMORY_INFO* LocalV
 		{
 			DXGI_QUERY_VIDEO_MEMORY_INFO TempVideoMemoryInfo;
 			VERIFYD3D12RESULT(Adapter3->QueryVideoMemoryInfo(Index, DXGI_MEMORY_SEGMENT_GROUP_LOCAL, &TempVideoMemoryInfo));
-			LocalVideoMemoryInfo->Budget = FMath::Min(LocalVideoMemoryInfo->Budget, TempVideoMemoryInfo.Budget);
-			LocalVideoMemoryInfo->CurrentUsage = FMath::Min(LocalVideoMemoryInfo->CurrentUsage, TempVideoMemoryInfo.CurrentUsage);
+
+			DXGI_QUERY_VIDEO_MEMORY_INFO TempSystemMemoryInfo;
+			VERIFYD3D12RESULT(Adapter3->QueryVideoMemoryInfo(Index, DXGI_MEMORY_SEGMENT_GROUP_NON_LOCAL, &TempSystemMemoryInfo));
+			
+			MemoryInfo.LocalMemoryInfo.Budget = FMath::Min(MemoryInfo.LocalMemoryInfo.Budget, TempVideoMemoryInfo.Budget);
+			MemoryInfo.LocalMemoryInfo.CurrentUsage = FMath::Min(MemoryInfo.LocalMemoryInfo.CurrentUsage, TempVideoMemoryInfo.CurrentUsage);
+
+			MemoryInfo.NonLocalMemoryInfo.Budget = FMath::Min(MemoryInfo.NonLocalMemoryInfo.Budget, TempSystemMemoryInfo.Budget);
+			MemoryInfo.NonLocalMemoryInfo.CurrentUsage = FMath::Min(MemoryInfo.NonLocalMemoryInfo.CurrentUsage, TempSystemMemoryInfo.CurrentUsage);
 		}
 	}
 #endif
diff --git a/Engine/Source/Runtime/D3D12RHI/Private/D3D12Adapter.h b/Engine/Source/Runtime/D3D12RHI/Private/D3D12Adapter.h
--- a/Engine/Source/Runtime/D3D12RHI/Private/D3D12Adapter.h
+++ b/Engine/Source/Runtime/D3D12RHI/Private/D3D12Adapter.h
@@ -71,6 +71,22 @@ struct FD3D12AdapterDesc
 	uint32 NumDeviceNodes{ 0 };
 };
 
+struct FD3D12MemoryInfo
+{
+	DXGI_QUERY_VIDEO_MEMORY_INFO LocalMemoryInfo{};
+	DXGI_QUERY_VIDEO_MEMORY_INFO NonLocalMemoryInfo{};
+
+	uint32 UpdateFrameNumber = ~uint32(0);
+
+	uint64 AvailableLocalMemory = 0;
+	uint64 DemotedLocalMemory = 0;
+
+	bool IsOverBudget() const
+	{
+		return DemotedLocalMemory > 0;
+	}
+};
+
 enum class ED3D12GPUCrashDebugginMode
 {
 	Disabled,
@@ -324,7 +340,8 @@ public:
 
 	void BlockUntilIdle();
 
-	void GetLocalVideoMemoryInfo(DXGI_QUERY_VIDEO_MEMORY_INFO* LocalVideoMemoryInfo);
+	void UpdateMemoryInfo();
+	FORCEINLINE const FD3D12MemoryInfo& GetMemoryInfo() const { return MemoryInfo; }
 
 	FORCEINLINE uint32 GetFrameCount() const { return FrameCounter; }
 
@@ -420,7 +437,7 @@ protected:
 	/** A Fence whos value increases every frame*/
 	TRefCountPtr<FD3D12ManualFence> FrameFence;
 
-	/** A Fence used to syncrhonize FD3D12GPUFence and FD3D12StagingBuffer */
+	/** A Fence used to synchronize FD3D12GPUFence and FD3D12StagingBuffer */
 	TRefCountPtr<FD3D12Fence> StagingFence;
 
 	FD3D12DeferredDeletionQueue DeferredDeletionQueue;
@@ -444,6 +461,8 @@ protected:
 	TMap<FD3D12ResourceLocation*, FTrackedAllocationData> TrackedAllocationData;
 	FCriticalSection TrackedAllocationDataCS;
 
+	FD3D12MemoryInfo MemoryInfo;
+
 #if D3D12_SUBMISSION_GAP_RECORDER
 	TArray<uint64> StartOfSubmissionTimestamps;
 	TArray<uint64> EndOfSubmissionTimestamps;
diff --git a/Engine/Source/Runtime/D3D12RHI/Private/D3D12CommandContext.cpp b/Engine/Source/Runtime/D3D12RHI/Private/D3D12CommandContext.cpp
--- a/Engine/Source/Runtime/D3D12RHI/Private/D3D12CommandContext.cpp
+++ b/Engine/Source/Runtime/D3D12RHI/Private/D3D12CommandContext.cpp
@@ -700,14 +700,16 @@ void FD3D12CommandContextBase::WaitForTransitionFences(TArrayView<const FRHITran
 void FD3D12CommandContextBase::UpdateMemoryStats()
 {
 #if PLATFORM_WINDOWS && STATS
-	DXGI_QUERY_VIDEO_MEMORY_INFO LocalVideoMemoryInfo;
-	ParentAdapter->GetLocalVideoMemoryInfo(&LocalVideoMemoryInfo);
-
-	const int64 Budget = LocalVideoMemoryInfo.Budget;
-	const int64 AvailableSpace = Budget - int64(LocalVideoMemoryInfo.CurrentUsage);
-	SET_MEMORY_STAT(STAT_D3D12UsedVideoMemory, LocalVideoMemoryInfo.CurrentUsage);
-	SET_MEMORY_STAT(STAT_D3D12AvailableVideoMemory, AvailableSpace);
-	SET_MEMORY_STAT(STAT_D3D12TotalVideoMemory, Budget);
+	// Refresh captured memory info.
+	ParentAdapter->UpdateMemoryInfo();
+
+	const FD3D12MemoryInfo& MemoryInfo = ParentAdapter->GetMemoryInfo();
+
+	SET_MEMORY_STAT(STAT_D3D12UsedVideoMemory, MemoryInfo.LocalMemoryInfo.CurrentUsage);
+	SET_MEMORY_STAT(STAT_D3D12UsedSystemMemory, MemoryInfo.NonLocalMemoryInfo.CurrentUsage);
+	SET_MEMORY_STAT(STAT_D3D12AvailableVideoMemory, MemoryInfo.AvailableLocalMemory);
+	SET_MEMORY_STAT(STAT_D3D12DemotedVideoMemory, MemoryInfo.DemotedLocalMemory);
+	SET_MEMORY_STAT(STAT_D3D12TotalVideoMemory, MemoryInfo.LocalMemoryInfo.Budget);
 
 	uint64 MaxTexAllocWastage = 0;
 	for (uint32 GPUIndex : GPUMask)
diff --git a/Engine/Source/Runtime/D3D12RHI/Private/D3D12RHICommon.h b/Engine/Source/Runtime/D3D12RHI/Private/D3D12RHICommon.h
--- a/Engine/Source/Runtime/D3D12RHI/Private/D3D12RHICommon.h
+++ b/Engine/Source/Runtime/D3D12RHI/Private/D3D12RHICommon.h
@@ -8,10 +8,12 @@ D3D12RHICommon.h: Common D3D12 RHI definitions for Windows.
 
 DECLARE_STATS_GROUP(TEXT("D3D12RHI"), STATGROUP_D3D12RHI, STATCAT_Advanced);
 DECLARE_STATS_GROUP(TEXT("D3D12RHI: Memory"), STATGROUP_D3D12Memory, STATCAT_Advanced);
-DECLARE_STATS_GROUP(TEXT("D3D12RHI: Memory (Buffers)"), STATGROUP_D3D12MemoryBuffer, STATCAT_Advanced);
-DECLARE_STATS_GROUP(TEXT("D3D12RHI: Memory (Transient)"), STATGROUP_D3D12MemoryTransient, STATCAT_Advanced);
+DECLARE_STATS_GROUP(TEXT("D3D12RHI: Resources"), STATGROUP_D3D12Resources, STATCAT_Advanced);
+DECLARE_STATS_GROUP(TEXT("D3D12RHI: Allocations"), STATGROUP_D3D12Allocations, STATCAT_Advanced);
+DECLARE_STATS_GROUP(TEXT("D3D12RHI: Allocations (Transient)"), STATGROUP_D3D12TransientAllocations, STATCAT_Advanced);
+DECLARE_STATS_GROUP(TEXT("D3D12RHI: Allocations (Buffer)"), STATGROUP_D3D12BufferAllocations, STATCAT_Advanced);
 DECLARE_STATS_GROUP(TEXT("D3D12RHI: Pipeline State (PSO)"), STATGROUP_D3D12PipelineState, STATCAT_Advanced);
-DECLARE_STATS_GROUP(TEXT("D3D12RHI: Descriptor Heap (GPU visible)"), STATGROUP_D3D12DescriptorHeap, STATCAT_Advanced);
+DECLARE_STATS_GROUP(TEXT("D3D12RHI: Descriptor Heap (GPU Visible)"), STATGROUP_D3D12DescriptorHeap, STATCAT_Advanced);
 
 #include "Windows/WindowsHWrapper.h"
 #include "D3D12RHI.h"
diff --git a/Engine/Source/Runtime/D3D12RHI/Private/D3D12Stats.h b/Engine/Source/Runtime/D3D12RHI/Private/D3D12Stats.h
--- a/Engine/Source/Runtime/D3D12RHI/Private/D3D12Stats.h
+++ b/Engine/Source/Runtime/D3D12RHI/Private/D3D12Stats.h
@@ -59,62 +59,63 @@ DECLARE_CYCLE_STAT_WITH_FLAGS_EXTERN(TEXT("Clear MRT time"), STAT_D3D12ClearMRT,
 DECLARE_CYCLE_STAT_EXTERN(TEXT("ExecuteCommandList time"), STAT_D3D12ExecuteCommandListTime, STATGROUP_D3D12RHI, );
 DECLARE_CYCLE_STAT_EXTERN(TEXT("WaitForFence time"), STAT_D3D12WaitForFenceTime, STATGROUP_D3D12RHI, );
 
-DECLARE_MEMORY_STAT_EXTERN(TEXT("Used Video Memory"), STAT_D3D12UsedVideoMemory, STATGROUP_D3D12Memory, );
-DECLARE_MEMORY_STAT_EXTERN(TEXT("Available Video Memory"), STAT_D3D12AvailableVideoMemory, STATGROUP_D3D12Memory, );
 DECLARE_MEMORY_STAT_EXTERN(TEXT("Total Video Memory"), STAT_D3D12TotalVideoMemory, STATGROUP_D3D12Memory, );
-
-DECLARE_MEMORY_STAT_EXTERN(TEXT("RenderTarget StandAlone Allocated"), STAT_D3D12RenderTargetStandAloneAllocated, STATGROUP_D3D12Memory, );
-DECLARE_MEMORY_STAT_EXTERN(TEXT("UAVTexture StandAlone Allocated"), STAT_D3D12UAVTextureStandAloneAllocated, STATGROUP_D3D12Memory, );
-DECLARE_MEMORY_STAT_EXTERN(TEXT("Texture StandAlone Allocated"), STAT_D3D12TextureStandAloneAllocated, STATGROUP_D3D12Memory, );
-DECLARE_MEMORY_STAT_EXTERN(TEXT("UAVBuffer StandAlone Allocated"), STAT_D3D12UAVBufferStandAloneAllocated, STATGROUP_D3D12Memory, );
-DECLARE_MEMORY_STAT_EXTERN(TEXT("Buffer StandAlone Allocated"), STAT_D3D12BufferStandAloneAllocated, STATGROUP_D3D12Memory, );
-
-DECLARE_DWORD_ACCUMULATOR_STAT_EXTERN(TEXT("RenderTarget StandAlone Allocated"), STAT_D3D12RenderTargetStandAloneCount, STATGROUP_D3D12Memory, );
-DECLARE_DWORD_ACCUMULATOR_STAT_EXTERN(TEXT("UAVTexture StandAlone Allocated"), STAT_D3D12UAVTextureStandAloneCount, STATGROUP_D3D12Memory, );
-DECLARE_DWORD_ACCUMULATOR_STAT_EXTERN(TEXT("Texture StandAlone Allocated"), STAT_D3D12TextureStandAloneCount, STATGROUP_D3D12Memory, );
-DECLARE_DWORD_ACCUMULATOR_STAT_EXTERN(TEXT("UAVBuffer StandAlone Allocated"), STAT_D3D12UAVBufferStandAloneCount, STATGROUP_D3D12Memory, );
-DECLARE_DWORD_ACCUMULATOR_STAT_EXTERN(TEXT("Buffer StandAlone Allocated"), STAT_D3D12BufferStandAloneCount, STATGROUP_D3D12Memory, );
-
-DECLARE_MEMORY_STAT_EXTERN(TEXT("Texture Allocator Allocated"), STAT_D3D12TextureAllocatorAllocated, STATGROUP_D3D12Memory, );
-DECLARE_MEMORY_STAT_EXTERN(TEXT("Texture Allocator Unused"), STAT_D3D12TextureAllocatorUnused, STATGROUP_D3D12Memory, );
-DECLARE_DWORD_ACCUMULATOR_STAT_EXTERN(TEXT("Texture Allocator Allocated"), STAT_D3D12TextureAllocatorCount, STATGROUP_D3D12Memory, );
-
-DECLARE_MEMORY_STAT_POOL_EXTERN(TEXT("View DescriptorHeap Memory (SRV, CBV, UAV)"), STAT_ViewOnlineDescriptorHeapMemory, STATGROUP_D3D12Memory, FPlatformMemory::MCR_GPUSystem, );
-DECLARE_MEMORY_STAT_POOL_EXTERN(TEXT("Sampler DescriptorHeap Memory"), STAT_SamplerOnlineDescriptorHeapMemory, STATGROUP_D3D12Memory, FPlatformMemory::MCR_GPUSystem, );
-
-DECLARE_MEMORY_STAT_EXTERN(TEXT("BufferPool Memory Allocated"), STAT_D3D12BufferPoolMemoryAllocated, STATGROUP_D3D12Memory, );
-DECLARE_MEMORY_STAT_EXTERN(TEXT("BufferPool Memory Used"), STAT_D3D12BufferPoolMemoryUsed, STATGROUP_D3D12Memory, );
-
-DECLARE_MEMORY_STAT_EXTERN(TEXT("BufferPool Memory Free"), STAT_D3D12BufferPoolMemoryFree, STATGROUP_D3D12MemoryBuffer, ); 
-DECLARE_MEMORY_STAT_EXTERN(TEXT("BufferPool Memory Alignment Waste"), STAT_D3D12BufferPoolAlignmentWaste, STATGROUP_D3D12MemoryBuffer, );
-DECLARE_DWORD_ACCUMULATOR_STAT_EXTERN(TEXT("BufferPool Page Count"), STAT_D3D12BufferPoolPageCount, STATGROUP_D3D12MemoryBuffer, );
-DECLARE_DWORD_ACCUMULATOR_STAT_EXTERN(TEXT("BufferPool Full Pages"), STAT_D3D12BufferPoolFullPages, STATGROUP_D3D12MemoryBuffer, );
-
-DECLARE_MEMORY_STAT_EXTERN(TEXT("Transient Memory Pool Allocated"), STAT_D3D12TransientMemoryPoolAllocated, STATGROUP_D3D12Memory, );
-DECLARE_MEMORY_STAT_EXTERN(TEXT("Transient Memory Frame Pool Used"), STAT_D3D12TransientMemoryFramePoolUsed, STATGROUP_D3D12Memory, );
-DECLARE_MEMORY_STAT_EXTERN(TEXT("Transient Memory Frame Committed Used"), STAT_D3D12TransientMemoryFrameCommittedUsed, STATGROUP_D3D12Memory, );
-DECLARE_DWORD_ACCUMULATOR_STAT_EXTERN(TEXT("Transient Memory Pool Allocations"), STAT_D3D12TransientMemoryPoolAllocations, STATGROUP_D3D12Memory, );
-DECLARE_DWORD_ACCUMULATOR_STAT_EXTERN(TEXT("Transient Memory Committed Allocations"), STAT_D3D12TransientMemoryCommittedAllocations, STATGROUP_D3D12Memory, );
-
-DECLARE_MEMORY_STAT_EXTERN(TEXT("Buffer Memory Frame Used"), STAT_D3D12TransientMemoryBufferFrameUsed, STATGROUP_D3D12MemoryTransient, );
-DECLARE_MEMORY_STAT_EXTERN(TEXT("Buffer Memory Frame Total Requested"), STAT_D3D12TransientMemoryBufferFrameTotalRequested, STATGROUP_D3D12MemoryTransient, );
-DECLARE_MEMORY_STAT_EXTERN(TEXT("Buffer Memory Frame Committed Allocated"), STAT_D3D12TransientMemoryBufferFrameCommittedAllocated, STATGROUP_D3D12MemoryTransient, );
-DECLARE_DWORD_ACCUMULATOR_STAT_EXTERN(TEXT("Buffer Memory Pool Allocations"), STAT_D3D12TransientMemoryBufferPoolAllocations, STATGROUP_D3D12MemoryTransient, );
-DECLARE_DWORD_ACCUMULATOR_STAT_EXTERN(TEXT("Buffer Memory Allocations"), STAT_D3D12TransientMemoryBufferCommittedAllocations, STATGROUP_D3D12MemoryTransient, );
-
-DECLARE_MEMORY_STAT_EXTERN(TEXT("Texture Memory Frame Used"), STAT_D3D12TransientMemoryTextureFrameUsed, STATGROUP_D3D12MemoryTransient, );
-DECLARE_MEMORY_STAT_EXTERN(TEXT("Texture Memory Frame Total Requested"), STAT_D3D12TransientMemoryTextureFrameTotalRequested, STATGROUP_D3D12MemoryTransient, );
-DECLARE_MEMORY_STAT_EXTERN(TEXT("Texture Memory Frame Committed Allocated"), STAT_D3D12TransientMemoryTextureFrameCommittedAllocated, STATGROUP_D3D12MemoryTransient, );
-DECLARE_DWORD_ACCUMULATOR_STAT_EXTERN(TEXT("Texture Memory Pool Allocations"), STAT_D3D12TransientMemoryTexturePoolAllocations, STATGROUP_D3D12MemoryTransient, );
-DECLARE_DWORD_ACCUMULATOR_STAT_EXTERN(TEXT("Texture Memory Allocations"), STAT_D3D12TransientMemoryTextureCommittedAllocations, STATGROUP_D3D12MemoryTransient, );
-
-DECLARE_MEMORY_STAT_EXTERN(TEXT("UploadPool Memory Allocated"), STAT_D3D12UploadPoolMemoryAllocated, STATGROUP_D3D12Memory, );
-DECLARE_MEMORY_STAT_EXTERN(TEXT("UploadPool Memory Used"), STAT_D3D12UploadPoolMemoryUsed, STATGROUP_D3D12Memory, );
-
-DECLARE_MEMORY_STAT_EXTERN(TEXT("UploadPool Memory Free"), STAT_D3D12UploadPoolMemoryFree, STATGROUP_D3D12MemoryBuffer, );
-DECLARE_MEMORY_STAT_EXTERN(TEXT("UploadPool Memory Alignment Waste"), STAT_D3D12UploadPoolAlignmentWaste, STATGROUP_D3D12MemoryBuffer, );
-DECLARE_DWORD_ACCUMULATOR_STAT_EXTERN(TEXT("UploadPool Page Count"), STAT_D3D12UploadPoolPageCount, STATGROUP_D3D12MemoryBuffer, );
-DECLARE_DWORD_ACCUMULATOR_STAT_EXTERN(TEXT("UploadPool Full Pages"), STAT_D3D12UploadPoolFullPages, STATGROUP_D3D12MemoryBuffer, );
+DECLARE_MEMORY_STAT_EXTERN(TEXT("Available Video Memory"), STAT_D3D12AvailableVideoMemory, STATGROUP_D3D12Memory, );
+DECLARE_MEMORY_STAT_EXTERN(TEXT("Demoted Video Memory"), STAT_D3D12DemotedVideoMemory, STATGROUP_D3D12Memory, );
+DECLARE_MEMORY_STAT_EXTERN(TEXT("Used Video Memory"), STAT_D3D12UsedVideoMemory, STATGROUP_D3D12Memory, );
+DECLARE_MEMORY_STAT_EXTERN(TEXT("Used System Memory"), STAT_D3D12UsedSystemMemory, STATGROUP_D3D12Memory, );
+
+DECLARE_MEMORY_STAT_EXTERN(TEXT("RenderTarget StandAlone Allocated"), STAT_D3D12RenderTargetStandAloneAllocated, STATGROUP_D3D12Resources, );
+DECLARE_MEMORY_STAT_EXTERN(TEXT("UAVTexture StandAlone Allocated"), STAT_D3D12UAVTextureStandAloneAllocated, STATGROUP_D3D12Resources, );
+DECLARE_MEMORY_STAT_EXTERN(TEXT("Texture StandAlone Allocated"), STAT_D3D12TextureStandAloneAllocated, STATGROUP_D3D12Resources, );
+DECLARE_MEMORY_STAT_EXTERN(TEXT("UAVBuffer StandAlone Allocated"), STAT_D3D12UAVBufferStandAloneAllocated, STATGROUP_D3D12Resources, );
+DECLARE_MEMORY_STAT_EXTERN(TEXT("Buffer StandAlone Allocated"), STAT_D3D12BufferStandAloneAllocated, STATGROUP_D3D12Resources, );
+
+DECLARE_DWORD_ACCUMULATOR_STAT_EXTERN(TEXT("RenderTarget StandAlone Allocated"), STAT_D3D12RenderTargetStandAloneCount, STATGROUP_D3D12Resources, );
+DECLARE_DWORD_ACCUMULATOR_STAT_EXTERN(TEXT("UAVTexture StandAlone Allocated"), STAT_D3D12UAVTextureStandAloneCount, STATGROUP_D3D12Resources, );
+DECLARE_DWORD_ACCUMULATOR_STAT_EXTERN(TEXT("Texture StandAlone Allocated"), STAT_D3D12TextureStandAloneCount, STATGROUP_D3D12Resources, );
+DECLARE_DWORD_ACCUMULATOR_STAT_EXTERN(TEXT("UAVBuffer StandAlone Allocated"), STAT_D3D12UAVBufferStandAloneCount, STATGROUP_D3D12Resources, );
+DECLARE_DWORD_ACCUMULATOR_STAT_EXTERN(TEXT("Buffer StandAlone Allocated"), STAT_D3D12BufferStandAloneCount, STATGROUP_D3D12Resources, );
+
+DECLARE_MEMORY_STAT_EXTERN(TEXT("Texture Allocator Allocated"), STAT_D3D12TextureAllocatorAllocated, STATGROUP_D3D12Allocations, );
+DECLARE_MEMORY_STAT_EXTERN(TEXT("Texture Allocator Unused"), STAT_D3D12TextureAllocatorUnused, STATGROUP_D3D12Allocations, );
+DECLARE_DWORD_ACCUMULATOR_STAT_EXTERN(TEXT("Texture Allocator Allocated"), STAT_D3D12TextureAllocatorCount, STATGROUP_D3D12Allocations, );
+
+DECLARE_MEMORY_STAT_POOL_EXTERN(TEXT("View DescriptorHeap Memory (SRV, CBV, UAV)"), STAT_ViewOnlineDescriptorHeapMemory, STATGROUP_D3D12Allocations, FPlatformMemory::MCR_GPUSystem, );
+DECLARE_MEMORY_STAT_POOL_EXTERN(TEXT("Sampler DescriptorHeap Memory"), STAT_SamplerOnlineDescriptorHeapMemory, STATGROUP_D3D12Allocations, FPlatformMemory::MCR_GPUSystem, );
+
+DECLARE_MEMORY_STAT_EXTERN(TEXT("BufferPool Memory Allocated"), STAT_D3D12BufferPoolMemoryAllocated, STATGROUP_D3D12Allocations, );
+DECLARE_MEMORY_STAT_EXTERN(TEXT("BufferPool Memory Used"), STAT_D3D12BufferPoolMemoryUsed, STATGROUP_D3D12Allocations, );
+DECLARE_MEMORY_STAT_EXTERN(TEXT("BufferPool Memory Free"), STAT_D3D12BufferPoolMemoryFree, STATGROUP_D3D12BufferAllocations, );
+DECLARE_MEMORY_STAT_EXTERN(TEXT("BufferPool Memory Alignment Waste"), STAT_D3D12BufferPoolAlignmentWaste, STATGROUP_D3D12BufferAllocations, );
+DECLARE_DWORD_ACCUMULATOR_STAT_EXTERN(TEXT("BufferPool Page Count"), STAT_D3D12BufferPoolPageCount, STATGROUP_D3D12BufferAllocations, );
+DECLARE_DWORD_ACCUMULATOR_STAT_EXTERN(TEXT("BufferPool Full Pages"), STAT_D3D12BufferPoolFullPages, STATGROUP_D3D12BufferAllocations, );
+
+DECLARE_MEMORY_STAT_EXTERN(TEXT("Transient Memory Pool Allocated"), STAT_D3D12TransientMemoryPoolAllocated, STATGROUP_D3D12Allocations, );
+DECLARE_MEMORY_STAT_EXTERN(TEXT("Transient Memory Frame Pool Used"), STAT_D3D12TransientMemoryFramePoolUsed, STATGROUP_D3D12Allocations, );
+DECLARE_MEMORY_STAT_EXTERN(TEXT("Transient Memory Frame Committed Used"), STAT_D3D12TransientMemoryFrameCommittedUsed, STATGROUP_D3D12Allocations, );
+DECLARE_DWORD_ACCUMULATOR_STAT_EXTERN(TEXT("Transient Memory Pool Allocations"), STAT_D3D12TransientMemoryPoolAllocations, STATGROUP_D3D12Allocations, );
+DECLARE_DWORD_ACCUMULATOR_STAT_EXTERN(TEXT("Transient Memory Committed Allocations"), STAT_D3D12TransientMemoryCommittedAllocations, STATGROUP_D3D12Allocations, );
+
+DECLARE_MEMORY_STAT_EXTERN(TEXT("Buffer Memory Frame Used"), STAT_D3D12TransientMemoryBufferFrameUsed, STATGROUP_D3D12TransientAllocations, );
+DECLARE_MEMORY_STAT_EXTERN(TEXT("Buffer Memory Frame Total Requested"), STAT_D3D12TransientMemoryBufferFrameTotalRequested, STATGROUP_D3D12TransientAllocations, );
+DECLARE_MEMORY_STAT_EXTERN(TEXT("Buffer Memory Frame Committed Allocated"), STAT_D3D12TransientMemoryBufferFrameCommittedAllocated, STATGROUP_D3D12TransientAllocations, );
+DECLARE_DWORD_ACCUMULATOR_STAT_EXTERN(TEXT("Buffer Memory Pool Allocations"), STAT_D3D12TransientMemoryBufferPoolAllocations, STATGROUP_D3D12TransientAllocations, );
+DECLARE_DWORD_ACCUMULATOR_STAT_EXTERN(TEXT("Buffer Memory Allocations"), STAT_D3D12TransientMemoryBufferCommittedAllocations, STATGROUP_D3D12TransientAllocations, );
+
+DECLARE_MEMORY_STAT_EXTERN(TEXT("Texture Memory Frame Used"), STAT_D3D12TransientMemoryTextureFrameUsed, STATGROUP_D3D12TransientAllocations, );
+DECLARE_MEMORY_STAT_EXTERN(TEXT("Texture Memory Frame Total Requested"), STAT_D3D12TransientMemoryTextureFrameTotalRequested, STATGROUP_D3D12TransientAllocations, );
+DECLARE_MEMORY_STAT_EXTERN(TEXT("Texture Memory Frame Committed Allocated"), STAT_D3D12TransientMemoryTextureFrameCommittedAllocated, STATGROUP_D3D12TransientAllocations, );
+DECLARE_DWORD_ACCUMULATOR_STAT_EXTERN(TEXT("Texture Memory Pool Allocations"), STAT_D3D12TransientMemoryTexturePoolAllocations, STATGROUP_D3D12TransientAllocations, );
+DECLARE_DWORD_ACCUMULATOR_STAT_EXTERN(TEXT("Texture Memory Allocations"), STAT_D3D12TransientMemoryTextureCommittedAllocations, STATGROUP_D3D12TransientAllocations, );
+
+DECLARE_MEMORY_STAT_EXTERN(TEXT("UploadPool Memory Allocated"), STAT_D3D12UploadPoolMemoryAllocated, STATGROUP_D3D12Allocations, );
+DECLARE_MEMORY_STAT_EXTERN(TEXT("UploadPool Memory Used"), STAT_D3D12UploadPoolMemoryUsed, STATGROUP_D3D12Allocations, );
+
+DECLARE_MEMORY_STAT_EXTERN(TEXT("UploadPool Memory Free"), STAT_D3D12UploadPoolMemoryFree, STATGROUP_D3D12BufferAllocations, );
+DECLARE_MEMORY_STAT_EXTERN(TEXT("UploadPool Memory Alignment Waste"), STAT_D3D12UploadPoolAlignmentWaste, STATGROUP_D3D12BufferAllocations, );
+DECLARE_DWORD_ACCUMULATOR_STAT_EXTERN(TEXT("UploadPool Page Count"), STAT_D3D12UploadPoolPageCount, STATGROUP_D3D12BufferAllocations, );
+DECLARE_DWORD_ACCUMULATOR_STAT_EXTERN(TEXT("UploadPool Full Pages"), STAT_D3D12UploadPoolFullPages, STATGROUP_D3D12BufferAllocations, );
 
 /**
 * Detailed Descriptor heap stats
@@ -132,7 +133,6 @@ DECLARE_DWORD_COUNTER_STAT_EXTERN(TEXT("View: Num reserved descriptors"), STAT_N
 DECLARE_DWORD_COUNTER_STAT_EXTERN(TEXT("Sampler: Num reserved descriptors"), STAT_NumReservedSamplerOnlineDescriptors, STATGROUP_D3D12DescriptorHeap, );
 DECLARE_DWORD_COUNTER_STAT_EXTERN(TEXT("Sampler: Num reused descriptors"), STAT_NumReusedSamplerOnlineDescriptors, STATGROUP_D3D12DescriptorHeap, );
 
-
 DECLARE_DWORD_ACCUMULATOR_STAT_EXTERN(TEXT("View Global: Free Descriptors"), STAT_GlobalViewHeapFreeDescriptors, STATGROUP_D3D12DescriptorHeap, );
 DECLARE_DWORD_ACCUMULATOR_STAT_EXTERN(TEXT("View Global: Reserved Descriptors"), STAT_GlobalViewHeapReservedDescriptors, STATGROUP_D3D12DescriptorHeap, );
 DECLARE_DWORD_ACCUMULATOR_STAT_EXTERN(TEXT("View Global: Used Descriptors"), STAT_GlobalViewHeapUsedDescriptors, STATGROUP_D3D12DescriptorHeap, );
diff --git a/Engine/Source/Runtime/D3D12RHI/Private/D3D12Texture.cpp b/Engine/Source/Runtime/D3D12RHI/Private/D3D12Texture.cpp
--- a/Engine/Source/Runtime/D3D12RHI/Private/D3D12Texture.cpp
+++ b/Engine/Source/Runtime/D3D12RHI/Private/D3D12Texture.cpp
@@ -584,8 +584,8 @@ void FD3D12DynamicRHI::RHIGetTextureMemoryStats(FTextureMemoryStats& OutStats)
 #if PLATFORM_WINDOWS || PLATFORM_HOLOLENS
 	if (GAdjustTexturePoolSizeBasedOnBudget)
 	{
-		DXGI_QUERY_VIDEO_MEMORY_INFO LocalVideoMemoryInfo;
-		GetAdapter().GetLocalVideoMemoryInfo(&LocalVideoMemoryInfo);
+		GetAdapter().UpdateMemoryInfo();
+		const DXGI_QUERY_VIDEO_MEMORY_INFO& LocalVideoMemoryInfo = GetAdapter().GetMemoryInfo().LocalMemoryInfo;
 
 		// Applications must explicitly manage their usage of physical memory and keep usage within the budget 
 		// assigned to the application process. Processes that cannot keep their usage within their assigned budgets 
diff --git a/Engine/Source/Runtime/D3D12RHI/Private/D3D12Util.cpp b/Engine/Source/Runtime/D3D12RHI/Private/D3D12Util.cpp
--- a/Engine/Source/Runtime/D3D12RHI/Private/D3D12Util.cpp
+++ b/Engine/Source/Runtime/D3D12RHI/Private/D3D12Util.cpp
@@ -1449,7 +1449,9 @@ DEFINE_STAT(STAT_D3D12ExecuteCommandListTime);
 DEFINE_STAT(STAT_D3D12WaitForFenceTime);
 
 DEFINE_STAT(STAT_D3D12UsedVideoMemory);
+DEFINE_STAT(STAT_D3D12UsedSystemMemory);
 DEFINE_STAT(STAT_D3D12AvailableVideoMemory);
+DEFINE_STAT(STAT_D3D12DemotedVideoMemory);
 DEFINE_STAT(STAT_D3D12TotalVideoMemory);
 
 DEFINE_STAT(STAT_D3D12RenderTargetStandAloneAllocated);
diff --git a/Engine/Source/Runtime/RHI/Private/RHI.cpp b/Engine/Source/Runtime/RHI/Private/RHI.cpp
--- a/Engine/Source/Runtime/RHI/Private/RHI.cpp
+++ b/Engine/Source/Runtime/RHI/Private/RHI.cpp
@@ -989,6 +989,8 @@ RHI_API int32 volatile GCurrentRendertargetMemorySize = 0;
 RHI_API int64 GTexturePoolSize = 0 * 1024 * 1024;
 RHI_API int32 GPoolSizeVRAMPercentage = 0;
 
+RHI_API uint64 GDemotedLocalMemorySize = 0;
+
 RHI_API EShaderPlatform GShaderPlatformForFeatureLevel[ERHIFeatureLevel::Num] = {SP_NumPlatforms,SP_NumPlatforms,SP_NumPlatforms,SP_NumPlatforms,SP_NumPlatforms};
 
 // simple stats about draw calls. GNum is the previous frame and 
diff --git a/Engine/Source/Runtime/RHI/Public/RHI.h b/Engine/Source/Runtime/RHI/Public/RHI.h
--- a/Engine/Source/Runtime/RHI/Public/RHI.h
+++ b/Engine/Source/Runtime/RHI/Public/RHI.h
@@ -541,6 +541,9 @@ extern RHI_API int64 GTexturePoolSize;
 /** In percent. If non-zero, the texture pool size is a percentage of GTotalGraphicsMemory. */
 extern RHI_API int32 GPoolSizeVRAMPercentage;
 
+/** Amount of local video memory demoted to system memory. In bytes. */
+extern RHI_API uint64 GDemotedLocalMemorySize;
+
 /** Some simple runtime stats, reset on every call to RHIBeginFrame */
 /** Num draw calls & primitives on previous frame (accurate on any thread)*/
 extern RHI_API int32 GNumDrawCallsRHI[MAX_NUM_GPUS];
diff --git a/Engine/Source/Runtime/Renderer/Private/SceneRendering.cpp b/Engine/Source/Runtime/Renderer/Private/SceneRendering.cpp
--- a/Engine/Source/Runtime/Renderer/Private/SceneRendering.cpp
+++ b/Engine/Source/Runtime/Renderer/Private/SceneRendering.cpp
@@ -2958,6 +2958,13 @@ void FSceneRenderer::RenderFinish(FRDGBuilder& GraphBuilder, FRDGTextureRef View
 			bShowPrecomputedVisibilityWarning = !bUsedPrecomputedVisibility;
 		}
 
+		bool bShowDemotedLocalMemoryWarning = false;
+		static const auto* CVarDemotedLocalMemoryWarning = IConsoleManager::Get().FindTConsoleVariableDataInt(TEXT("r.DemotedLocalMemoryWarning"));
+		if (CVarDemotedLocalMemoryWarning && CVarDemotedLocalMemoryWarning->GetValueOnRenderThread() == 1)
+		{
+			bShowDemotedLocalMemoryWarning = GDemotedLocalMemorySize > 0;
+		}
+
 		bool bShowGlobalClipPlaneWarning = false;
 
 		if (Scene->PlanarReflections.Num() > 0)
@@ -3037,7 +3044,7 @@ void FSceneRenderer::RenderFinish(FRDGBuilder& GraphBuilder, FRDGTextureRef View
 		FFXSystemInterface* FXInterface = Scene->GetFXSystem();
 		const bool bFxDebugDraw = FXInterface && FXInterface->ShouldDebugDraw_RenderThread();
 
-		const bool bAnyWarning = bShowPrecomputedVisibilityWarning || bShowGlobalClipPlaneWarning || bShowAtmosphericFogWarning || bShowSkylightWarning || bShowPointLightWarning 
+		const bool bAnyWarning = bShowPrecomputedVisibilityWarning || bShowDemotedLocalMemoryWarning || bShowGlobalClipPlaneWarning || bShowAtmosphericFogWarning || bShowSkylightWarning || bShowPointLightWarning
 			|| bShowDFAODisabledWarning || bShowShadowedLightOverflowWarning || bShowMobileDynamicCSMWarning || bShowMobileLowQualityLightmapWarning || bShowMobileMovableDirectionalLightWarning
 			|| bMobileShowVertexFogWarning || bShowSkinCacheOOM || bSingleLayerWaterWarning || bShowDFDisabledWarning || bShowNoSkyAtmosphereComponentWarning || bFxDebugDraw 
 			|| bShowSkyAtmosphereFogComponentsConflicts || bLumenEnabledButNoSoftwareTracing || bRealTimeSkyCaptureButNothingToCapture;
@@ -3059,7 +3066,7 @@ void FSceneRenderer::RenderFinish(FRDGBuilder& GraphBuilder, FRDGTextureRef View
 
 					AddDrawCanvasPass(GraphBuilder, {}, View, Output,
 						[this, &ReadOnlyCVARCache, ViewState, GPUSkinCacheExtraRequiredMemory,
-						bLocked, bShowPrecomputedVisibilityWarning, bShowGlobalClipPlaneWarning, bShowDFAODisabledWarning, bShowDFDisabledWarning,
+						bLocked, bShowPrecomputedVisibilityWarning, bShowDemotedLocalMemoryWarning, bShowGlobalClipPlaneWarning, bShowDFAODisabledWarning, bShowDFDisabledWarning,
 						bShowAtmosphericFogWarning, bViewParentOrFrozen, bShowSkylightWarning, bShowPointLightWarning, bShowShadowedLightOverflowWarning,
 						bShowMobileLowQualityLightmapWarning, bShowMobileMovableDirectionalLightWarning, bShowMobileDynamicCSMWarning, bMobileShowVertexFogWarning,
 						bShowSkinCacheOOM, bSingleLayerWaterWarning, bShowNoSkyAtmosphereComponentWarning, bFxDebugDraw, FXInterface, bShowSkyAtmosphereFogComponentsConflicts, 
@@ -3213,6 +3220,13 @@ void FSceneRenderer::RenderFinish(FRDGBuilder& GraphBuilder, FRDGTextureRef View
 							Canvas.DrawShadowedText(10, Y, Message, GetStatsFont(), FLinearColor(1.0, 0.05, 0.05, 1.0));
 							Y += 14;
 						}
+
+						if (bShowDemotedLocalMemoryWarning)
+						{
+							FString String = FString::Printf(TEXT("Video memory has been exhausted (%.3f MB over budget). Expect extremely poor performance."), float(GDemotedLocalMemorySize) / 1048576.0f);
+							Canvas.DrawShadowedText(10, Y, FText::FromString(String), GetStatsFont(), FLinearColor(1.0, 0.05, 0.05, 1.0));
+							Y += 14;
+						}
 					});
 					if (bFxDebugDraw)
 					{
