Adding missing Mesh shader support to D3D12 pipeline caches

diff --git a/Engine/Source/Runtime/D3D12RHI/Private/D3D12PipelineState.h b/Engine/Source/Runtime/D3D12RHI/Private/D3D12PipelineState.h
--- a/Engine/Source/Runtime/D3D12RHI/Private/D3D12PipelineState.h
+++ b/Engine/Source/Runtime/D3D12RHI/Private/D3D12PipelineState.h
@@ -195,6 +195,8 @@ template <> struct equality_pipeline_state_desc<FD3D12LowLevelGraphicsPipelineSt
 		// Order from most likely to change to least
 		PSO_IF_NOT_EQUAL_RETURN_FALSE(Desc.PS.BytecodeLength)
 		PSO_IF_NOT_EQUAL_RETURN_FALSE(Desc.VS.BytecodeLength)
+		PSO_IF_NOT_EQUAL_RETURN_FALSE(Desc.MS.BytecodeLength)
+		PSO_IF_NOT_EQUAL_RETURN_FALSE(Desc.AS.BytecodeLength)
 		PSO_IF_NOT_EQUAL_RETURN_FALSE(Desc.GS.BytecodeLength)
 		PSO_IF_NOT_EQUAL_RETURN_FALSE(Desc.DS.BytecodeLength)
 		PSO_IF_NOT_EQUAL_RETURN_FALSE(Desc.HS.BytecodeLength)
@@ -226,6 +228,8 @@ template <> struct equality_pipeline_state_desc<FD3D12LowLevelGraphicsPipelineSt
 		// should be tiny i.e if there were 1 quadrillion shaders the chance of a 
 		// collision is ~ 1 in 10^18. so only do a full check on debug builds
 		PSO_IF_NOT_EQUAL_RETURN_FALSE(VSHash)
+		PSO_IF_NOT_EQUAL_RETURN_FALSE(MSHash)
+		PSO_IF_NOT_EQUAL_RETURN_FALSE(ASHash)
 		PSO_IF_NOT_EQUAL_RETURN_FALSE(PSHash)
 		PSO_IF_NOT_EQUAL_RETURN_FALSE(GSHash)
 		PSO_IF_NOT_EQUAL_RETURN_FALSE(HSHash)
@@ -248,6 +252,8 @@ template <> struct equality_pipeline_state_desc<FD3D12LowLevelGraphicsPipelineSt
 
 	#if PLATFORM_WINDOWS
 		PSO_IF_NOT_EQUAL_RETURN_FALSE(VSExtensions);
+		PSO_IF_NOT_EQUAL_RETURN_FALSE(MSExtensions);
+		PSO_IF_NOT_EQUAL_RETURN_FALSE(ASExtensions);
 		PSO_IF_NOT_EQUAL_RETURN_FALSE(PSExtensions);
 		PSO_IF_NOT_EQUAL_RETURN_FALSE(GSExtensions);
 		PSO_IF_NOT_EQUAL_RETURN_FALSE(HSExtensions);
diff --git a/Engine/Source/Runtime/D3D12RHI/Private/D3D12Pipelinestate.cpp b/Engine/Source/Runtime/D3D12RHI/Private/D3D12Pipelinestate.cpp
--- a/Engine/Source/Runtime/D3D12RHI/Private/D3D12Pipelinestate.cpp
+++ b/Engine/Source/Runtime/D3D12RHI/Private/D3D12Pipelinestate.cpp
@@ -148,6 +148,8 @@ uint64 FD3D12PipelineStateCacheBase::HashPSODesc(const FD3D12LowLevelGraphicsPip
 	struct GraphicsPSOData
 	{
 		ShaderBytecodeHash VSHash;
+		ShaderBytecodeHash MSHash;
+		ShaderBytecodeHash ASHash;
 		ShaderBytecodeHash HSHash;
 		ShaderBytecodeHash DSHash;
 		ShaderBytecodeHash GSHash;
@@ -195,6 +197,8 @@ uint64 FD3D12PipelineStateCacheBase::HashPSODesc(const FD3D12LowLevelGraphicsPip
 	RenderTargetData* RTData = (RenderTargetData*) (Data + GraphicsPSODataSize);
 
 	PSOData->VSHash          = Desc.VSHash;
+	PSOData->MSHash          = Desc.MSHash;
+	PSOData->ASHash          = Desc.ASHash;
 	PSOData->HSHash          = Desc.HSHash;
 	PSOData->DSHash          = Desc.DSHash;
 	PSOData->GSHash          = Desc.GSHash;
diff --git a/Engine/Source/Runtime/D3D12RHI/Private/D3D12RootSignature.cpp b/Engine/Source/Runtime/D3D12RHI/Private/D3D12RootSignature.cpp
--- a/Engine/Source/Runtime/D3D12RHI/Private/D3D12RootSignature.cpp
+++ b/Engine/Source/Runtime/D3D12RHI/Private/D3D12RootSignature.cpp
@@ -420,7 +420,7 @@ const D3D12_VERSIONED_ROOT_SIGNATURE_DESC& FD3D12RootSignatureDesc::GetStaticMes
 {
 	// TODO: Support vendor extensions for static root signatures?
 
-	static const uint32 DescriptorTableCount = 16;
+	static const uint32 DescriptorTableCount = 10;
 	static struct
 	{
 		D3D12_SHADER_VISIBILITY Vis;
diff --git a/Engine/Source/Runtime/D3D12RHI/Private/D3D12StateCache.cpp b/Engine/Source/Runtime/D3D12RHI/Private/D3D12StateCache.cpp
--- a/Engine/Source/Runtime/D3D12RHI/Private/D3D12StateCache.cpp
+++ b/Engine/Source/Runtime/D3D12RHI/Private/D3D12StateCache.cpp
@@ -353,6 +353,11 @@ void FD3D12StateCacheBase::SetScissorRects(uint32 Count, const D3D12_RECT* const
 	}
 }
 
+inline bool ShouldSkipStage(uint32 Stage)
+{
+	return ((Stage == SF_Mesh || Stage == SF_Amplification) && !GRHISupportsMeshShaders);
+}
+
 template <ED3D12PipelineType PipelineType>
 void FD3D12StateCacheBase::ApplyState()
 {
@@ -520,6 +525,11 @@ void FD3D12StateCacheBase::ApplyState()
 
 		for (uint32 Stage = StartStage; Stage < EndStage; ++Stage)
 		{
+			if (ShouldSkipStage(Stage))
+			{
+				continue;
+			}
+
 			// Note this code assumes the starting register is index 0.
 			const SRVSlotMask CurrentShaderSRVRegisterMask = BitMask<SRVSlotMask>(PipelineState.Common.CurrentShaderSRVCounts[Stage]);
 			CurrentShaderDirtySRVSlots[Stage] = CurrentShaderSRVRegisterMask & PipelineState.Common.SRVCache.DirtySlotMask[Stage];
@@ -675,6 +685,11 @@ void FD3D12StateCacheBase::ApplySamplers(const FD3D12RootSignature* const pRootS
 
 		for (uint32 Stage = StartStage; Stage < EndStage; ++Stage)
 		{
+			if (ShouldSkipStage(Stage))
+			{
+				continue;
+			}
+
 			// Note this code assumes the starting register is index 0.
 			const SamplerSlotMask CurrentShaderSamplerRegisterMask = ((SamplerSlotMask)1 << PipelineState.Common.CurrentShaderSamplerCounts[Stage]) - (SamplerSlotMask)1;
 			CurrentShaderDirtySamplerSlots[Stage] = CurrentShaderSamplerRegisterMask & Cache.DirtySlotMask[Stage];
@@ -705,6 +720,11 @@ void FD3D12StateCacheBase::ApplySamplers(const FD3D12RootSignature* const pRootS
 
 		for (uint32 Stage = StartStage; Stage < EndStage; Stage++)
 		{
+			if (ShouldSkipStage(Stage))
+			{
+				continue;
+			}
+
 			if ((CurrentShaderDirtySamplerSlots[Stage] && NumSamplers[Stage]))
 			{
 				SamplerSlotMask& CurrentDirtySlotMask = Cache.DirtySlotMask[Stage];
@@ -843,6 +863,11 @@ bool FD3D12StateCacheBase::AssertResourceStates(ED3D12PipelineType PipelineType)
 	bool bSRVIntersectsWithStencil = false;
 	for (uint32 Stage = StartStage; Stage < EndStage; Stage++)
 	{
+		if (ShouldSkipStage(Stage))
+		{
+			continue;
+		}
+
 		// UAVs
 		{
 			const uint32 numUAVs = PipelineState.Common.CurrentShaderUAVCounts[Stage];
diff --git a/Engine/Source/Runtime/D3D12RHI/Private/D3D12StateCachePrivate.h b/Engine/Source/Runtime/D3D12RHI/Private/D3D12StateCachePrivate.h
--- a/Engine/Source/Runtime/D3D12RHI/Private/D3D12StateCachePrivate.h
+++ b/Engine/Source/Runtime/D3D12RHI/Private/D3D12StateCachePrivate.h
@@ -135,6 +135,8 @@ struct FD3D12ResourceCache
 	inline void DirtyGraphics(const ResourceSlotMask& SlotMask = -1)
 	{
 		Dirty(SF_Vertex, SlotMask);
+		Dirty(SF_Mesh, SlotMask);
+		Dirty(SF_Amplification, SlotMask);
 		Dirty(SF_Hull, SlotMask);
 		Dirty(SF_Domain, SlotMask);
 		Dirty(SF_Pixel, SlotMask);
diff --git a/Engine/Source/Runtime/D3D12RHI/Private/Windows/WindowsD3D12DiskCache.h b/Engine/Source/Runtime/D3D12RHI/Private/Windows/WindowsD3D12DiskCache.h
--- a/Engine/Source/Runtime/D3D12RHI/Private/Windows/WindowsD3D12DiskCache.h
+++ b/Engine/Source/Runtime/D3D12RHI/Private/Windows/WindowsD3D12DiskCache.h
@@ -10,7 +10,7 @@ class FDiskCacheInterface
 {
 	// Increment this if changes are made to the
 	// disk caches so stale caches get updated correctly
-	static const uint32 mCurrentHeaderVersion = 9;
+	static const uint32 mCurrentHeaderVersion = 10;
 	struct FDiskCacheHeader
 	{
 		uint32 mHeaderVersion;
diff --git a/Engine/Source/Runtime/D3D12RHI/Private/Windows/WindowsD3D12PipelineState.cpp b/Engine/Source/Runtime/D3D12RHI/Private/Windows/WindowsD3D12PipelineState.cpp
--- a/Engine/Source/Runtime/D3D12RHI/Private/Windows/WindowsD3D12PipelineState.cpp
+++ b/Engine/Source/Runtime/D3D12RHI/Private/Windows/WindowsD3D12PipelineState.cpp
@@ -165,6 +165,8 @@ void ComputePipelineCreationArgs_POD::Destroy()
 void GraphicsPipelineCreationArgs_POD::Destroy()
 {
 	FMemory::Free((void*)Desc.Desc.VS.pShaderBytecode);
+	FMemory::Free((void*)Desc.Desc.MS.pShaderBytecode);
+	FMemory::Free((void*)Desc.Desc.AS.pShaderBytecode);
 	FMemory::Free((void*)Desc.Desc.PS.pShaderBytecode);
 	FMemory::Free((void*)Desc.Desc.GS.pShaderBytecode);
 	FMemory::Free((void*)Desc.Desc.HS.pShaderBytecode);
@@ -280,6 +282,14 @@ void FD3D12PipelineStateCache::RebuildFromDiskCache(ID3D12RootSignature* Graphic
 		{
 			DiskCaches[PSO_CACHE_GRAPHICS].SetPointerAndAdvanceFilePosition((void**)&PSODesc->VS.pShaderBytecode, PSODesc->VS.BytecodeLength, bBackShadersWithSystemMemory);
 		}
+		if (PSODesc->MS.BytecodeLength)
+		{
+			DiskCaches[PSO_CACHE_GRAPHICS].SetPointerAndAdvanceFilePosition((void**)&PSODesc->MS.pShaderBytecode, PSODesc->MS.BytecodeLength, bBackShadersWithSystemMemory);
+		}
+		if (PSODesc->AS.BytecodeLength)
+		{
+			DiskCaches[PSO_CACHE_GRAPHICS].SetPointerAndAdvanceFilePosition((void**)&PSODesc->AS.pShaderBytecode, PSODesc->AS.BytecodeLength, bBackShadersWithSystemMemory);
+		}
 		if (PSODesc->PS.BytecodeLength)
 		{
 			DiskCaches[PSO_CACHE_GRAPHICS].SetPointerAndAdvanceFilePosition((void**)&PSODesc->PS.pShaderBytecode, PSODesc->PS.BytecodeLength, bBackShadersWithSystemMemory);
@@ -420,6 +430,14 @@ void FD3D12PipelineStateCache::AddToDiskCache(const FD3D12LowLevelGraphicsPipeli
 		{
 			DiskCache.AppendData((void*)PsoDesc.VS.pShaderBytecode, PsoDesc.VS.BytecodeLength);
 		}
+		if (PsoDesc.MS.BytecodeLength)
+		{
+			DiskCache.AppendData((void*)PsoDesc.MS.pShaderBytecode, PsoDesc.MS.BytecodeLength);
+		}
+		if (PsoDesc.AS.BytecodeLength)
+		{
+			DiskCache.AppendData((void*)PsoDesc.AS.pShaderBytecode, PsoDesc.AS.BytecodeLength);
+		}
 		if (PsoDesc.PS.BytecodeLength)
 		{
 			DiskCache.AppendData((void*)PsoDesc.PS.pShaderBytecode, PsoDesc.PS.BytecodeLength);
