- Make sure transient resources are created and released in correct default state
- Apply transition on first usage when before state is known and valid (skip adding for patch up command list)

diff --git a/Engine/Source/Runtime/D3D12RHI/Private/D3D12Commands.cpp b/Engine/Source/Runtime/D3D12RHI/Private/D3D12Commands.cpp
--- a/Engine/Source/Runtime/D3D12RHI/Private/D3D12Commands.cpp
+++ b/Engine/Source/Runtime/D3D12RHI/Private/D3D12Commands.cpp
@@ -335,47 +335,20 @@ static void HandleResourceDiscardTransitions(
 				return;
 			}
 
-			const D3D12_RESOURCE_FLAGS ResourceFlags = Resource->GetDesc().Flags;
-
-			D3D12_RESOURCE_STATES DiscardState;
-
-			if (EnumHasAnyFlags(ResourceFlags, D3D12_RESOURCE_FLAG_ALLOW_RENDER_TARGET))
-			{
-				DiscardState = D3D12_RESOURCE_STATE_RENDER_TARGET;
-			}
-			else if (EnumHasAnyFlags(ResourceFlags, D3D12_RESOURCE_FLAG_ALLOW_DEPTH_STENCIL))
-			{
-				DiscardState = D3D12_RESOURCE_STATE_DEPTH_WRITE;
-			}
-			else if (EnumHasAnyFlags(ResourceFlags, D3D12_RESOURCE_FLAG_ALLOW_UNORDERED_ACCESS))
-			{
-				DiscardState = D3D12_RESOURCE_STATE_UNORDERED_ACCESS;
-			}
-			else
-			{
-				return;
-			}
-
-			if (EnumHasAnyFlags(Info.Flags, EResourceTransitionFlags::MaintainCompression))
-			{
-				DiscardState |= SkipFastClearEliminateState;
-			}
-
+			// Get the initial state to force a 'nop' transition so the internal command list resource tracking has the correct state
+			// already to make sure it's not added to the pending transition list to be kicked before this command list (resource is not valid then yet)
+			D3D12_RESOURCE_STATES InitialState = FD3D12TransientResourceAllocator::GetInitialResourceState(Resource->GetDesc());
 			if (Info.IsWholeResource() || Resource->GetSubresourceCount() == 1)
 			{
-				if (FD3D12DynamicRHI::TransitionResource(Context.CommandListHandle, Resource, D3D12_RESOURCE_STATE_TBD, DiscardState, D3D12_RESOURCE_BARRIER_ALL_SUBRESOURCES, FD3D12DynamicRHI::ETransitionMode::Apply))
-				{
-					ResourcesToDiscard.Emplace(Resource, Info.Flags, D3D12_RESOURCE_BARRIER_ALL_SUBRESOURCES);
-				}
+				FD3D12DynamicRHI::TransitionResource(Context.CommandListHandle, Resource, InitialState, InitialState, D3D12_RESOURCE_BARRIER_ALL_SUBRESOURCES, FD3D12DynamicRHI::ETransitionMode::Apply);
+				ResourcesToDiscard.Emplace(Resource, Info.Flags, D3D12_RESOURCE_BARRIER_ALL_SUBRESOURCES);
 			}
 			else
 			{
 				EnumerateSubresources(Resource, Info, [&](uint32 Subresource)
 				{
-					if (FD3D12DynamicRHI::TransitionResource(Context.CommandListHandle, Resource, D3D12_RESOURCE_STATE_TBD, DiscardState, Subresource, FD3D12DynamicRHI::ETransitionMode::Apply))
-					{
-						ResourcesToDiscard.Emplace(Resource, Info.Flags, Subresource);
-					}
+					FD3D12DynamicRHI::TransitionResource(Context.CommandListHandle, Resource, InitialState, InitialState, Subresource, FD3D12DynamicRHI::ETransitionMode::Apply);
+					ResourcesToDiscard.Emplace(Resource, Info.Flags, Subresource);
 				});
 			}
 		});
@@ -451,12 +424,11 @@ static void HandleTransientAliasing(FD3D12CommandContext& Context, const FD3D12T
 			break;
 		}
 
+		FD3D12Resource* Resource = BaseShaderResource->ResourceLocation.GetResource();
 		if (Info.Action == FRHITransientAliasingInfo::EAction::Acquire)
 		{
 			TRACE_CPUPROFILER_EVENT_SCOPE(D3D12RHI::AcquireTransient);
 
-			FD3D12Resource* Resource = BaseShaderResource->ResourceLocation.GetResource();
-
 			// Try and see if the resource has overlapping resources
 			if (BaseShaderResource->ResourceLocation.IsTransient() && BaseShaderResource->ResourceLocation.GetAllocatorType() == FD3D12ResourceLocation::AT_Pool)
 			{
@@ -475,6 +447,10 @@ static void HandleTransientAliasing(FD3D12CommandContext& Context, const FD3D12T
 		{
 			TRACE_CPUPROFILER_EVENT_SCOPE(D3D12RHI::DiscardTransient);
 
+			// Restore the resource back to the initial state when done
+			D3D12_RESOURCE_STATES FinalState = FD3D12TransientResourceAllocator::GetInitialResourceState(BaseShaderResource->GetResource()->GetDesc());
+			FD3D12DynamicRHI::TransitionResource(Context.CommandListHandle, Resource, D3D12_RESOURCE_STATE_TBD, FinalState, D3D12_RESOURCE_BARRIER_ALL_SUBRESOURCES, FD3D12DynamicRHI::ETransitionMode::Apply);
+
 			// Remove from caches
 			Context.ConditionalClearShaderResource(&BaseShaderResource->ResourceLocation);
 
diff --git a/Engine/Source/Runtime/D3D12RHI/Private/D3D12RHIPrivate.h b/Engine/Source/Runtime/D3D12RHI/Private/D3D12RHIPrivate.h
--- a/Engine/Source/Runtime/D3D12RHI/Private/D3D12RHIPrivate.h
+++ b/Engine/Source/Runtime/D3D12RHI/Private/D3D12RHIPrivate.h
@@ -907,7 +907,8 @@ public:
 		bool bValidateState = !GUseInternalTransitions && (InMode == ETransitionMode::Validate);
 
 		// Try and get the correct D3D before state for the transition
-		D3D12_RESOURCE_STATES BeforeState = InResourceState.GetSubresourceState(InSubresourceIndex);
+		D3D12_RESOURCE_STATES TrackedState = InResourceState.GetSubresourceState(InSubresourceIndex);
+		D3D12_RESOURCE_STATES BeforeState = TrackedState;
 
 		// Still untracked in this command list, then try and find out a before state to use
 		if (BeforeState == D3D12_RESOURCE_STATE_TBD)
@@ -924,9 +925,21 @@ public:
 			}
 			else if (GUseInternalTransitions)
 			{
-				// We need a pending resource barrier so we can setup the state before this command list executes
-				InCommandList.AddPendingResourceBarrier(InResource, InAfterState, InSubresourceIndex);
-				InResourceState.SetSubresourceState(InSubresourceIndex, InAfterState);
+				// Already perform transition here if possible to skip patch up during command list execution
+				if (InBeforeState != D3D12_RESOURCE_STATE_TBD)
+				{
+					check(BeforeState == D3D12_RESOURCE_STATE_TBD || BeforeState == InBeforeState);
+					BeforeState = InBeforeState;
+
+					// Add dummy pending barrier, because the end state needs to be updated during execute
+					InCommandList.AddPendingResourceBarrier(InResource, D3D12_RESOURCE_STATE_TBD, InSubresourceIndex);
+				}
+				else
+				{
+					// We need a pending resource barrier so we can setup the state before this command list executes
+					InResourceState.SetSubresourceState(InSubresourceIndex, InAfterState);
+					InCommandList.AddPendingResourceBarrier(InResource, InAfterState, InSubresourceIndex);
+				}
 			}
 			else
 			{
@@ -978,7 +991,12 @@ public:
 				{
 					InCommandList.AddTransitionBarrier(InResource, BeforeState, InAfterState, InSubresourceIndex);
 					InResourceState.SetSubresourceState(InSubresourceIndex, InAfterState);
-				}				
+				}
+				// Force update the state when the tracked state is still unknown
+				else if (TrackedState == D3D12_RESOURCE_STATE_TBD)
+				{
+					InResourceState.SetSubresourceState(InSubresourceIndex, InAfterState);
+				}
 			}
 		}
 
diff --git a/Engine/Source/Runtime/D3D12RHI/Private/D3D12TransientResourceAllocator.cpp b/Engine/Source/Runtime/D3D12RHI/Private/D3D12TransientResourceAllocator.cpp
--- a/Engine/Source/Runtime/D3D12RHI/Private/D3D12TransientResourceAllocator.cpp
+++ b/Engine/Source/Runtime/D3D12RHI/Private/D3D12TransientResourceAllocator.cpp
@@ -177,7 +177,7 @@ void FD3D12TransientMemoryPool::ReleaseResource(FD3D12Resource* InResource, FRHI
 
 	uint64 AllocationOffset = InReleasedAllocationData.GetOffset();
 	uint64 AllocationSize = InReleasedAllocationData.GetSize();
-
+	
 	// Placed resource object is 'never' freed - it can be reused again this or next frame
 	// Will only be destroyed when the allocator is destroyed
 	FResourceCreateState CreateState;
@@ -188,7 +188,7 @@ void FD3D12TransientMemoryPool::ReleaseResource(FD3D12Resource* InResource, FRHI
 	uint64 CreateStateHash = CreateState.GetHash();
 	check(CachedResourceMap.Find(CreateStateHash) == nullptr);
 	CachedResourceMap.Add(CreateStateHash, InResource);
-	
+
 	// Track this resource so aliasing barriers can be added if used again
 	ActiveResourceData ResourceData;
 	ResourceData.AllocationRange = FInt64Range(AllocationOffset, AllocationOffset + AllocationSize);
@@ -453,6 +453,27 @@ void FD3D12TransientMemoryPoolManager::UpdateMemoryStats()
 //-----------------------------------------------------------------------------
 
 
+D3D12_RESOURCE_STATES FD3D12TransientResourceAllocator::GetInitialResourceState(const D3D12_RESOURCE_DESC& InDesc)
+{
+	// Validate the creation state
+	D3D12_RESOURCE_STATES State = D3D12_RESOURCE_STATE_COMMON;
+	if (EnumHasAnyFlags(InDesc.Flags, D3D12_RESOURCE_FLAG_ALLOW_RENDER_TARGET))
+	{
+		State = D3D12_RESOURCE_STATE_RENDER_TARGET;
+	}
+	else if (EnumHasAnyFlags(InDesc.Flags, D3D12_RESOURCE_FLAG_ALLOW_DEPTH_STENCIL))
+	{
+		State = D3D12_RESOURCE_STATE_DEPTH_WRITE;
+	}
+	else if (EnumHasAnyFlags(InDesc.Flags, D3D12_RESOURCE_FLAG_ALLOW_UNORDERED_ACCESS))
+	{
+		State = D3D12_RESOURCE_STATE_UNORDERED_ACCESS;
+	}
+	check(State != D3D12_RESOURCE_STATE_COMMON);
+	return State;
+}
+
+
 FD3D12TransientResourceAllocator::FD3D12TransientResourceAllocator(FD3D12TransientMemoryPoolManager& InMemoryPoolManager) :
 	FD3D12PoolAllocator(
 		InMemoryPoolManager.GetParentDevice(),
@@ -489,6 +510,16 @@ FRHITexture* FD3D12TransientResourceAllocator::CreateTexture(const FRHITextureCr
 	FD3D12PooledTextureData TextureData = PoolManager.GetPooledTexture(InCreateInfo, InDebugName);
 
 	FD3D12BaseShaderResource* BaseShaderResource = nullptr;
+	
+	ERHIAccess InitialState = ERHIAccess::UAVMask;
+	if (EnumHasAnyFlags(InCreateInfo.Flags, TexCreate_RenderTargetable))
+	{
+		InitialState = ERHIAccess::RTV;
+	}
+	else if (EnumHasAnyFlags(InCreateInfo.Flags, TexCreate_DepthStencilTargetable))
+	{
+		InitialState = ERHIAccess::DSVWrite;
+	}
 
 	// Found something?
 	if (TextureData.RHITexture)
@@ -507,8 +538,10 @@ FRHITexture* FD3D12TransientResourceAllocator::CreateTexture(const FRHITextureCr
 			ClearValuePtr = &TextureData.ClearValue;
 		}
 
+		D3D12_RESOURCE_STATES CreateState = GetD3D12ResourceState(InitialState, false);
+
 		// Allocate a new FD3D12Resource on the resource location
-		AllocateTexture(D3D12_HEAP_TYPE_DEFAULT, TextureData.ResourceDesc, InCreateInfo.Format, ED3D12ResourceStateMode::MultiState, D3D12_RESOURCE_STATE_COMMON, ClearValuePtr, InDebugName, BaseShaderResource->ResourceLocation);
+		AllocateTexture(D3D12_HEAP_TYPE_DEFAULT, TextureData.ResourceDesc, InCreateInfo.Format, ED3D12ResourceStateMode::MultiState, CreateState, ClearValuePtr, InDebugName, BaseShaderResource->ResourceLocation);
 
 		// Inform the listeners about the change - should ideally still get from cache
 		{
@@ -530,7 +563,7 @@ FRHITexture* FD3D12TransientResourceAllocator::CreateTexture(const FRHITextureCr
 			const bool bTextureArray = false;
 			const bool bCubeTexture = false;
 			ID3D12ResourceAllocator* ResourceAllocator = this;
-			FD3D12Texture2D* Texture2D = FD3D12DynamicRHI::GetD3DRHI()->CreateD3D12Texture2D<FD3D12BaseTexture2D>(nullptr, InCreateInfo.Extent.X, InCreateInfo.Extent.Y, 1, bTextureArray, bCubeTexture, InCreateInfo.Format, InCreateInfo.NumMips, InCreateInfo.NumSamples, InCreateInfo.Flags, ERHIAccess::Discard, CreateInfo, ResourceAllocator);
+			FD3D12Texture2D* Texture2D = FD3D12DynamicRHI::GetD3DRHI()->CreateD3D12Texture2D<FD3D12BaseTexture2D>(nullptr, InCreateInfo.Extent.X, InCreateInfo.Extent.Y, 1, bTextureArray, bCubeTexture, InCreateInfo.Format, InCreateInfo.NumMips, InCreateInfo.NumSamples, InCreateInfo.Flags, InitialState, CreateInfo, ResourceAllocator);
 			TextureData.RHITexture = Texture2D;
 			TextureData.ResourceDesc = Texture2D->GetResource()->GetDesc();
 			TextureData.ClearValue = Texture2D->GetResource()->GetClearValue();
@@ -543,7 +576,7 @@ FRHITexture* FD3D12TransientResourceAllocator::CreateTexture(const FRHITextureCr
 		{
 			// Only support 2d textures for now
 			ID3D12ResourceAllocator* ResourceAllocator = this;
-			FD3D12Texture3D* Texture3D = FD3D12DynamicRHI::GetD3DRHI()->CreateD3D12Texture3D(nullptr, InCreateInfo.Extent.X, InCreateInfo.Extent.Y, InCreateInfo.Depth, InCreateInfo.Format, InCreateInfo.NumMips, InCreateInfo.Flags, ERHIAccess::Discard, CreateInfo, ResourceAllocator);
+			FD3D12Texture3D* Texture3D = FD3D12DynamicRHI::GetD3DRHI()->CreateD3D12Texture3D(nullptr, InCreateInfo.Extent.X, InCreateInfo.Extent.Y, InCreateInfo.Depth, InCreateInfo.Format, InCreateInfo.NumMips, InCreateInfo.Flags, InitialState, CreateInfo, ResourceAllocator);
 			TextureData.RHITexture = Texture3D;
 			TextureData.ResourceDesc = Texture3D->GetResource()->GetDesc();
 			TextureData.ClearValue = Texture3D->GetResource()->GetClearValue();
@@ -559,7 +592,7 @@ FRHITexture* FD3D12TransientResourceAllocator::CreateTexture(const FRHITextureCr
 			const bool bTextureArray = false;
 			const bool bCubeTexture = true;
 			ID3D12ResourceAllocator* ResourceAllocator = this;
-			FD3D12TextureCube* TextureCube = FD3D12DynamicRHI::GetD3DRHI()->CreateD3D12Texture2D<FD3D12BaseTextureCube>(nullptr, InCreateInfo.Extent.X, InCreateInfo.Extent.Y, 6, bTextureArray, bCubeTexture, InCreateInfo.Format, InCreateInfo.NumMips, InCreateInfo.NumSamples, InCreateInfo.Flags, ERHIAccess::Discard, CreateInfo, ResourceAllocator);
+			FD3D12TextureCube* TextureCube = FD3D12DynamicRHI::GetD3DRHI()->CreateD3D12Texture2D<FD3D12BaseTextureCube>(nullptr, InCreateInfo.Extent.X, InCreateInfo.Extent.Y, 6, bTextureArray, bCubeTexture, InCreateInfo.Format, InCreateInfo.NumMips, InCreateInfo.NumSamples, InCreateInfo.Flags, InitialState, CreateInfo, ResourceAllocator);
 			TextureData.RHITexture = TextureCube;
 			TextureData.ResourceDesc = TextureCube->GetResource()->GetDesc();
 			TextureData.ClearValue = TextureCube->GetResource()->GetClearValue();
@@ -593,6 +626,7 @@ FRHIBuffer* FD3D12TransientResourceAllocator::CreateBuffer(const FRHIBufferCreat
 	// Try and get a pooled texture from the manager
 	FD3D12TransientMemoryPoolManager& PoolManager = GetParentDevice()->GetTransientMemoryPoolManager();
 	FD3D12PooledBufferData BufferData = PoolManager.GetPooledBuffer(InCreateInfo, InDebugName);
+	ERHIAccess InitialState = ERHIAccess::UAVMask;
 
 	FD3D12Buffer* Buffer = nullptr;
 
@@ -608,8 +642,10 @@ FRHIBuffer* FD3D12TransientResourceAllocator::CreateBuffer(const FRHIBufferCreat
 		uint32 Alignment;
 		FD3D12Buffer::GetResourceDescAndAlignment(InCreateInfo.Size, InCreateInfo.Stride, UsageFlags, Desc, Alignment);
 
+		D3D12_RESOURCE_STATES CreateState = GetD3D12ResourceState(InitialState, false);
+
 		// Allocate a new FD3D12Resource on the resource location
-		AllocateResource(D3D12_HEAP_TYPE_DEFAULT, Desc, Desc.Width, Alignment, ED3D12ResourceStateMode::MultiState, D3D12_RESOURCE_STATE_COMMON, nullptr, InDebugName, Buffer->ResourceLocation);
+		AllocateResource(D3D12_HEAP_TYPE_DEFAULT, Desc, Desc.Width, Alignment, ED3D12ResourceStateMode::MultiState, CreateState, nullptr, InDebugName, Buffer->ResourceLocation);
 
 		// Inform the listeners about the change (should be empty?)
 		Buffer->ResourceRenamed(&Buffer->ResourceLocation);
@@ -618,7 +654,7 @@ FRHIBuffer* FD3D12TransientResourceAllocator::CreateBuffer(const FRHIBufferCreat
 	{
 		FRHIResourceCreateInfo CreateInfo(InDebugName);
 		ID3D12ResourceAllocator* ResourceAllocator = this;
-		Buffer = FD3D12DynamicRHI::GetD3DRHI()->CreateD3D12Buffer(nullptr, InCreateInfo.Size, InCreateInfo.Usage, InCreateInfo.Stride, ERHIAccess::Discard, CreateInfo, ResourceAllocator);
+		Buffer = FD3D12DynamicRHI::GetD3DRHI()->CreateD3D12Buffer(nullptr, InCreateInfo.Size, InCreateInfo.Usage, InCreateInfo.Stride, InitialState, CreateInfo, ResourceAllocator);
 		BufferData.RHIBuffer = Buffer;
 		BufferData.CreateInfo = InCreateInfo;
 	}
@@ -767,6 +803,9 @@ FD3D12Resource* FD3D12TransientResourceAllocator::CreatePlacedResource(const FRH
 	// No cached resource then use base class
 	if (Resource == nullptr)
 	{
+		// Validate the creation state
+		check(InCreateState == GetInitialResourceState(InDesc));
+
 		TRACE_CPUPROFILER_EVENT_SCOPE(D3D12RHI::CreatePlacedResource);
 		Resource = FD3D12PoolAllocator::CreatePlacedResource(InAllocationData, InDesc, InCreateState, InResourceStateMode, InClearValue, InName);
 	}
diff --git a/Engine/Source/Runtime/D3D12RHI/Private/D3D12TransientResourceAllocator.h b/Engine/Source/Runtime/D3D12RHI/Private/D3D12TransientResourceAllocator.h
--- a/Engine/Source/Runtime/D3D12RHI/Private/D3D12TransientResourceAllocator.h
+++ b/Engine/Source/Runtime/D3D12RHI/Private/D3D12TransientResourceAllocator.h
@@ -198,6 +198,9 @@ public:
 	// Try and find overlapping resource data for given base resource
 	TArrayView<FD3D12Resource*> GetOverlappingResources(FD3D12BaseShaderResource* InBaseShaderResource);
 
+	// Get the initial resource state
+	static D3D12_RESOURCE_STATES GetInitialResourceState(const D3D12_RESOURCE_DESC& InDesc);
+
 private:
 
 	// Override placed resource allocation helper function to get from cache if available
