Disable pool allocator on Xbox

diff --git a/Engine/Source/Runtime/D3D12RHI/Private/D3D12RHIPrivate.h b/Engine/Source/Runtime/D3D12RHI/Private/D3D12RHIPrivate.h
--- a/Engine/Source/Runtime/D3D12RHI/Private/D3D12RHIPrivate.h
+++ b/Engine/Source/Runtime/D3D12RHI/Private/D3D12RHIPrivate.h
@@ -370,7 +370,7 @@ public:
 	virtual class IRHIComputeContext* RHIGetDefaultAsyncComputeContext() final override;
 	virtual class IRHICommandContextContainer* RHIGetCommandContextContainer(int32 Index, int32 Num) final override;
 
-	virtual IRHITransientResourceAllocator* RHICreateTransientResourceAllocator() final override;
+	virtual IRHITransientResourceAllocator* RHICreateTransientResourceAllocator() final;
 
 #if WITH_MGPU
 	virtual IRHICommandContextContainer* RHIGetCommandContextContainer(int32 Index, int32 Num, FRHIGPUMask GPUMask)final override;
