Add support for dynamic pool sizes to reduce memory overhead and still support big allocations of upto 512 MB

diff --git a/Engine/Source/Runtime/D3D12RHI/Private/D3D12TransientResourceAllocator.cpp b/Engine/Source/Runtime/D3D12RHI/Private/D3D12TransientResourceAllocator.cpp
--- a/Engine/Source/Runtime/D3D12RHI/Private/D3D12TransientResourceAllocator.cpp
+++ b/Engine/Source/Runtime/D3D12RHI/Private/D3D12TransientResourceAllocator.cpp
@@ -3,13 +3,20 @@
 #include "D3D12RHIPrivate.h"
 #include "D3D12TransientResourceAllocator.h"
 
-//PRAGMA_DISABLE_OPTIMIZATION
+PRAGMA_DISABLE_OPTIMIZATION
 
-static int32 GD3D12TransientAllocatorPoolSizeInMB = 128;
+static int32 GD3D12TransientAllocatorDefaultPoolSizeInMB = 64;
 static FAutoConsoleVariableRef CVarD3D12TransientAllocatorPoolSizeInMB(
-	TEXT("d3d12.TransientAllocator.PoolSizeInMB"),
-	GD3D12TransientAllocatorPoolSizeInMB,
-	TEXT("Size of a D3D12 transient allocator pool in MB (Default 128)"),
+	TEXT("d3d12.TransientAllocator.DefaultPoolSizeInMB"),
+	GD3D12TransientAllocatorDefaultPoolSizeInMB,
+	TEXT("Default size of a D3D12 transient allocator pool in MB (Default 128)"),
+	ECVF_ReadOnly);
+
+static int32 GD3D12TransientAllocatorMaxAllocationSizeInMB = 512;
+static FAutoConsoleVariableRef CVarD3D12TransientAllocatorMaxAllocationSizeInMB(
+	TEXT("d3d12.TransientAllocator.MaxPoolSizeInMB"),
+	GD3D12TransientAllocatorMaxAllocationSizeInMB,
+	TEXT("Max size of an D3D12 transient allocation in MB (Default 512)"),
 	ECVF_ReadOnly);
 
 static int32 GD3D12TransientAllocatorPoolTextures = 1;
@@ -177,7 +184,7 @@ void FD3D12TransientMemoryPool::ReleaseResource(FD3D12Resource* InResource, FRHI
 
 	uint64 AllocationOffset = InReleasedAllocationData.GetOffset();
 	uint64 AllocationSize = InReleasedAllocationData.GetSize();
-	
+
 	// Placed resource object is 'never' freed - it can be reused again this or next frame
 	// Will only be destroyed when the allocator is destroyed
 	FResourceCreateState CreateState;
@@ -188,7 +195,7 @@ void FD3D12TransientMemoryPool::ReleaseResource(FD3D12Resource* InResource, FRHI
 	uint64 CreateStateHash = CreateState.GetHash();
 	check(CachedResourceMap.Find(CreateStateHash) == nullptr);
 	CachedResourceMap.Add(CreateStateHash, InResource);
-
+	
 	// Track this resource so aliasing barriers can be added if used again
 	ActiveResourceData ResourceData;
 	ResourceData.AllocationRange = FInt64Range(AllocationOffset, AllocationOffset + AllocationSize);
@@ -225,9 +232,9 @@ FD3D12TransientMemoryPoolManager::FD3D12TransientMemoryPoolManager(FD3D12Device*
 	// Support RT and UAV for now (Tier 2 support) - need to add different pools for Tier 1
 	InitConfig.HeapFlags = D3D12_HEAP_FLAGS(0); // 0 means nothing is denied
 
-	PoolSize = GD3D12TransientAllocatorPoolSizeInMB * 1024 * 1024;
+	DefaultPoolSize = GD3D12TransientAllocatorDefaultPoolSizeInMB * 1024 * 1024;
 	PoolAlignment = D3D12_DEFAULT_RESOURCE_PLACEMENT_ALIGNMENT;
-	MaxAllocationSize = PoolSize;
+	MaxAllocationSize = GD3D12TransientAllocatorMaxAllocationSizeInMB * 1024 * 1024;
 }
 
 
@@ -274,20 +281,39 @@ void FD3D12TransientMemoryPoolManager::EndFrame()
 }
 
 
-FD3D12TransientMemoryPool* FD3D12TransientMemoryPoolManager::GetOrCreateMemoryPool(int16 InPoolIndex)
+FD3D12TransientMemoryPool* FD3D12TransientMemoryPoolManager::GetOrCreateMemoryPool(int16 InPoolIndex, uint32 InMinimumAllocationSize)
 {
 	FScopeLock Lock(&CS);
 
 	FD3D12TransientMemoryPool* MemoryPool = nullptr;
-	if (Pools.Num() > 0)
+
+	// Try and find a pool which fits this allocations (sort by size so will take best fit)
+	for (int32 PoolIndex = 0; PoolIndex < Pools.Num(); ++PoolIndex)
 	{
-		MemoryPool = Pools.Pop(false);
-		MemoryPool->SetPoolIndex(InPoolIndex);
+		FD3D12TransientMemoryPool* Pool = Pools[PoolIndex];
+		if (Pool->GetPoolSize() >= InMinimumAllocationSize)
+		{			
+			MemoryPool = Pool;
+			MemoryPool->SetPoolIndex(InPoolIndex);
+			Pools.RemoveAt(PoolIndex);
+			break;
+		}
 	}
-	else
+
+	// couldn't find a pool?
+	if (MemoryPool == nullptr)
 	{
 		TRACE_CPUPROFILER_EVENT_SCOPE(D3D12RHI::AllocateTransientMemoryPool);
 
+		// Find out the pool size - use default, but if allocation doesn't fit then round up to next power of 2
+		// so it 'always' fits the pool allocator
+		uint32 PoolSize = DefaultPoolSize;
+		if (InMinimumAllocationSize > PoolSize)
+		{
+			check(InMinimumAllocationSize <= MaxAllocationSize);
+			PoolSize = FMath::Min(FMath::RoundUpToPowerOfTwo(InMinimumAllocationSize), (uint32)MaxAllocationSize);
+		}
+
 		MemoryPool = new FD3D12TransientMemoryPool(GetParentDevice(), GetVisibilityMask(), InitConfig,
 			TEXT("TransientResourceMemoryPool"), EResourceAllocationStrategy::kPlacedResource, InPoolIndex, PoolSize, PoolAlignment);
 		MemoryPool->Init();
@@ -313,7 +339,19 @@ void FD3D12TransientMemoryPoolManager::ReleaseMemoryPool(FD3D12TransientMemoryPo
 {
 	FScopeLock Lock(&CS);
 	InMemoryPool->ResetPool();
-	Pools.Add(InMemoryPool);
+
+	// sorted insert
+	int32 InsertIndex = 0;
+	for (; InsertIndex < Pools.Num(); ++InsertIndex)
+	{
+		// First by size then by last frame usage
+		if (Pools[InsertIndex]->GetPoolSize() >= InMemoryPool->GetPoolSize() &&
+			Pools[InsertIndex]->GetLastUsedFrameFence() <= InMemoryPool->GetLastUsedFrameFence())
+		{
+			break;
+		}
+	}
+	Pools.Insert(InMemoryPool, InsertIndex);
 }
 
 
@@ -481,7 +519,7 @@ FD3D12TransientResourceAllocator::FD3D12TransientResourceAllocator(FD3D12Transie
 		InMemoryPoolManager.GetInitConfig(),
 		TEXT("FD3D12TransientResourceAllocator"),
 		EResourceAllocationStrategy::kPlacedResource,
-		InMemoryPoolManager.GetPoolSize(),
+		InMemoryPoolManager.GetDefaultPoolSize(),
 		InMemoryPoolManager.GetPoolAlignment(),
 		InMemoryPoolManager.GetMaxAllocationSize(),
 		FRHIMemoryPool::EFreeListOrder::SortByOffset,
@@ -510,7 +548,7 @@ FRHITexture* FD3D12TransientResourceAllocator::CreateTexture(const FRHITextureCr
 	FD3D12PooledTextureData TextureData = PoolManager.GetPooledTexture(InCreateInfo, InDebugName);
 
 	FD3D12BaseShaderResource* BaseShaderResource = nullptr;
-	
+
 	ERHIAccess InitialState = ERHIAccess::UAVMask;
 	if (EnumHasAnyFlags(InCreateInfo.Flags, TexCreate_RenderTargetable))
 	{
@@ -786,7 +824,7 @@ void FD3D12TransientResourceAllocator::Freeze(FRHICommandListImmediate&)
 FRHIMemoryPool* FD3D12TransientResourceAllocator::CreateNewPool(int16 InPoolIndex, uint32 InMinimumAllocationSize)
 {
 	// Get pool from manager - don't reallocate each time
-	return GetParentDevice()->GetTransientMemoryPoolManager().GetOrCreateMemoryPool(InPoolIndex);
+	return GetParentDevice()->GetTransientMemoryPoolManager().GetOrCreateMemoryPool(InPoolIndex, InMinimumAllocationSize);
 }
 
 
diff --git a/Engine/Source/Runtime/D3D12RHI/Private/D3D12TransientResourceAllocator.h b/Engine/Source/Runtime/D3D12RHI/Private/D3D12TransientResourceAllocator.h
--- a/Engine/Source/Runtime/D3D12RHI/Private/D3D12TransientResourceAllocator.h
+++ b/Engine/Source/Runtime/D3D12RHI/Private/D3D12TransientResourceAllocator.h
@@ -125,7 +125,7 @@ public:
 	void EndFrame();
 
 	// Get or creation functions from cached resources
-	FD3D12TransientMemoryPool* GetOrCreateMemoryPool(int16 InPoolIndex);
+	FD3D12TransientMemoryPool* GetOrCreateMemoryPool(int16 InPoolIndex, uint32 InMinimumAllocationSize);
 	void ReleaseMemoryPool(FD3D12TransientMemoryPool* InMemoryPool);
 	FD3D12PooledTextureData GetPooledTexture(const FRHITextureCreateInfo& InCreateInfo, const TCHAR* InDebugName);
 	FD3D12PooledBufferData GetPooledBuffer(const FRHIBufferCreateInfo& InCreateInfo, const TCHAR* InDebugName);
@@ -133,7 +133,7 @@ public:
 
 	// Get the pool creation members
 	const FD3D12ResourceInitConfig& GetInitConfig() const { return InitConfig; }
-	uint64 GetPoolSize() const { return PoolSize; }
+	uint64 GetDefaultPoolSize() const { return DefaultPoolSize; }
 	uint32 GetPoolAlignment() const { return PoolAlignment; }
 	uint64 GetMaxAllocationSize() const { return MaxAllocationSize; }
 
@@ -152,7 +152,7 @@ private:
 
 	// Pool creation members
 	FD3D12ResourceInitConfig InitConfig;
-	uint64 PoolSize;
+	uint64 DefaultPoolSize;
 	uint32 PoolAlignment;
 	uint64 MaxAllocationSize;
 
