Remove forgotten PRAGMA_DISABLE_OPTIMIZATION

diff --git a/Engine/Source/Runtime/D3D12RHI/Private/D3D12TransientResourceAllocator.cpp b/Engine/Source/Runtime/D3D12RHI/Private/D3D12TransientResourceAllocator.cpp
--- a/Engine/Source/Runtime/D3D12RHI/Private/D3D12TransientResourceAllocator.cpp
+++ b/Engine/Source/Runtime/D3D12RHI/Private/D3D12TransientResourceAllocator.cpp
@@ -3,8 +3,6 @@
 #include "D3D12RHIPrivate.h"
 #include "D3D12TransientResourceAllocator.h"
 
-PRAGMA_DISABLE_OPTIMIZATION
-
 static int32 GD3D12TransientAllocatorDefaultPoolSizeInMB = 64;
 static FAutoConsoleVariableRef CVarD3D12TransientAllocatorPoolSizeInMB(
 	TEXT("d3d12.TransientAllocator.DefaultPoolSizeInMB"),
