Add support for Tier 1 (no merged buffer/texture heaps) on Transient allocator in d3d12

diff --git a/Engine/Source/Runtime/D3D12RHI/Private/D3D12PoolAllocator.cpp b/Engine/Source/Runtime/D3D12RHI/Private/D3D12PoolAllocator.cpp
--- a/Engine/Source/Runtime/D3D12RHI/Private/D3D12PoolAllocator.cpp
+++ b/Engine/Source/Runtime/D3D12RHI/Private/D3D12PoolAllocator.cpp
@@ -26,8 +26,8 @@ static FAutoConsoleVariableRef CVarD3D12MinPoolAllocationSizeToTrack(
 //-----------------------------------------------------------------------------
 
 FD3D12MemoryPool::FD3D12MemoryPool(FD3D12Device* ParentDevice, FRHIGPUMask VisibleNodes, const FD3D12ResourceInitConfig& InInitConfig, const FString& InName,
-	EResourceAllocationStrategy InAllocationStrategy, int16 InPoolIndex, uint64 InPoolSize, uint32 InPoolAlignment, EFreeListOrder InFreeListOrder)
-	: FRHIMemoryPool(InPoolIndex, InPoolSize, InPoolAlignment, InFreeListOrder), FD3D12DeviceChild(ParentDevice), FD3D12MultiNodeGPUObject(ParentDevice->GetGPUMask(), VisibleNodes)
+	EResourceAllocationStrategy InAllocationStrategy, int16 InPoolIndex, uint64 InPoolSize, uint32 InPoolAlignment, ERHIPoolResourceTypes InSupportedResourceTypes, EFreeListOrder InFreeListOrder)
+	: FRHIMemoryPool(InPoolIndex, InPoolSize, InPoolAlignment, InSupportedResourceTypes, InFreeListOrder), FD3D12DeviceChild(ParentDevice), FD3D12MultiNodeGPUObject(ParentDevice->GetGPUMask(), VisibleNodes)
 	, InitConfig(InInitConfig), Name(InName), AllocationStrategy(InAllocationStrategy), LastUsedFrameFence(0)
 {
 }
@@ -281,9 +281,24 @@ void FD3D12PoolAllocator::AllocateResource(D3D12_HEAP_TYPE InHeapType, const D3D
 			check(InDesc.Flags == InitConfig.ResourceFlags);
 		}
 
+		// Find the correct allocation resource type
+		ERHIPoolResourceTypes AllocationResourceType;
+		if (InDesc.Dimension == D3D12_RESOURCE_DIMENSION_BUFFER)
+		{
+			AllocationResourceType = ERHIPoolResourceTypes::Buffers;
+		}
+		else if (EnumHasAnyFlags(InDesc.Flags, D3D12_RESOURCE_FLAG_ALLOW_RENDER_TARGET | D3D12_RESOURCE_FLAG_ALLOW_DEPTH_STENCIL))
+		{
+			AllocationResourceType = ERHIPoolResourceTypes::RTDSTextures;
+		}
+		else
+		{
+			AllocationResourceType = ERHIPoolResourceTypes::NonRTDSTextures;
+		}
+
 		// Try to allocate in one of the pools
 		FRHIPoolAllocationData& AllocationData = ResourceLocation.GetPoolAllocatorPrivateData().PoolData;
-		verify(TryAllocateInternal(InSize, AllocationAlignment, AllocationData));
+		verify(TryAllocateInternal(InSize, AllocationAlignment, AllocationResourceType, AllocationData));
 
 		// Setup the resource location
 		ResourceLocation.SetType(FD3D12ResourceLocation::ResourceLocationType::eSubAllocation);
@@ -420,7 +435,7 @@ void FD3D12PoolAllocator::DeallocateResource(FD3D12ResourceLocation& ResourceLoc
 }
 
 
-FRHIMemoryPool* FD3D12PoolAllocator::CreateNewPool(int16 InPoolIndex, uint32 InMinimumAllocationSize)
+FRHIMemoryPool* FD3D12PoolAllocator::CreateNewPool(int16 InPoolIndex, uint32 InMinimumAllocationSize, ERHIPoolResourceTypes InAllocationResourceType)
 {
 	TRACE_CPUPROFILER_EVENT_SCOPE(FD3D12PoolAllocator::CreateNewPool);
 
@@ -434,7 +449,7 @@ FRHIMemoryPool* FD3D12PoolAllocator::CreateNewPool(int16 InPoolIndex, uint32 InM
 	}
 
 	FD3D12MemoryPool* NewPool = new FD3D12MemoryPool(GetParentDevice(),	GetVisibilityMask(), InitConfig,
-		Name, AllocationStrategy, InPoolIndex, PoolSize, PoolAlignment, FreeListOrder);
+		Name, AllocationStrategy, InPoolIndex, PoolSize, PoolAlignment, InAllocationResourceType, FreeListOrder);
 	NewPool->Init();
 	return NewPool;
 }
diff --git a/Engine/Source/Runtime/D3D12RHI/Private/D3D12PoolAllocator.h b/Engine/Source/Runtime/D3D12RHI/Private/D3D12PoolAllocator.h
--- a/Engine/Source/Runtime/D3D12RHI/Private/D3D12PoolAllocator.h
+++ b/Engine/Source/Runtime/D3D12RHI/Private/D3D12PoolAllocator.h
@@ -87,7 +87,7 @@ public:
 
 	// Constructor
 	FD3D12MemoryPool(FD3D12Device* ParentDevice, FRHIGPUMask VisibleNodes, const FD3D12ResourceInitConfig& InInitConfig, const FString& Name,
-		EResourceAllocationStrategy InAllocationStrategy, int16 InPoolIndex, uint64 InPoolSize, uint32 InPoolAlignment, EFreeListOrder InFreeListOrder);
+		EResourceAllocationStrategy InAllocationStrategy, int16 InPoolIndex, uint64 InPoolSize, uint32 InPoolAlignment, ERHIPoolResourceTypes InSupportedResourceTypes, EFreeListOrder InFreeListOrder);
 	virtual ~FD3D12MemoryPool();
 
 	// Setup/Shutdown
@@ -157,7 +157,7 @@ public:
 protected:
 
 	// Implementation of FRHIPoolAllocator pure virtuals
-	virtual FRHIMemoryPool* CreateNewPool(int16 InPoolIndex, uint32 InMinimumAllocationSize) override;
+	virtual FRHIMemoryPool* CreateNewPool(int16 InPoolIndex, uint32 InMinimumAllocationSize, ERHIPoolResourceTypes InAllocationResourceType) override;
 	virtual bool HandleDefragRequest(FRHIPoolAllocationData* InSourceBlock, FRHIPoolAllocationData& InTmpTargetBlock) override;
 	
 	// Placed resource allocation helper function which can be overriden
diff --git a/Engine/Source/Runtime/D3D12RHI/Private/D3D12TransientResourceAllocator.cpp b/Engine/Source/Runtime/D3D12RHI/Private/D3D12TransientResourceAllocator.cpp
--- a/Engine/Source/Runtime/D3D12RHI/Private/D3D12TransientResourceAllocator.cpp
+++ b/Engine/Source/Runtime/D3D12RHI/Private/D3D12TransientResourceAllocator.cpp
@@ -220,6 +220,9 @@ FD3D12TransientMemoryPoolManager::FD3D12TransientMemoryPoolManager(FD3D12Device*
 	: FD3D12DeviceChild(InDevice)
 	, FD3D12MultiNodeGPUObject(InDevice->GetGPUMask(), VisibleNodes)
 {
+	// Does the device support Heap Tier 2 - merged heaps for buffers & textures
+	bMergedTypeHeapSupported = (InDevice->GetParentAdapter()->GetResourceHeapTier() == D3D12_RESOURCE_HEAP_TIER_2);
+
 	// texture only interesting in VRAM for now
 	InitConfig.HeapType = D3D12_HEAP_TYPE_DEFAULT;
 
@@ -227,7 +230,7 @@ FD3D12TransientMemoryPoolManager::FD3D12TransientMemoryPoolManager(FD3D12Device*
 	InitConfig.ResourceFlags = D3D12_RESOURCE_FLAG_NONE;
 	InitConfig.InitialResourceState = D3D12_RESOURCE_STATE_COMMON;
 
-	// Support RT and UAV for now (Tier 2 support) - need to add different pools for Tier 1
+	// By default the manager and allocator support all resource types
 	InitConfig.HeapFlags = D3D12_HEAP_FLAGS(0); // 0 means nothing is denied
 
 	DefaultPoolSize = GD3D12TransientAllocatorDefaultPoolSizeInMB * 1024 * 1024;
@@ -279,7 +282,7 @@ void FD3D12TransientMemoryPoolManager::EndFrame()
 }
 
 
-FD3D12TransientMemoryPool* FD3D12TransientMemoryPoolManager::GetOrCreateMemoryPool(int16 InPoolIndex, uint32 InMinimumAllocationSize)
+FD3D12TransientMemoryPool* FD3D12TransientMemoryPoolManager::GetOrCreateMemoryPool(int16 InPoolIndex, uint32 InMinimumAllocationSize, ERHIPoolResourceTypes InAllocationResourceType)
 {
 	FScopeLock Lock(&CS);
 
@@ -289,7 +292,7 @@ FD3D12TransientMemoryPool* FD3D12TransientMemoryPoolManager::GetOrCreateMemoryPo
 	for (int32 PoolIndex = 0; PoolIndex < Pools.Num(); ++PoolIndex)
 	{
 		FD3D12TransientMemoryPool* Pool = Pools[PoolIndex];
-		if (Pool->GetPoolSize() >= InMinimumAllocationSize)
+		if (Pool->GetPoolSize() >= InMinimumAllocationSize && Pool->IsResourceTypeSupported(InAllocationResourceType))
 		{			
 			MemoryPool = Pool;
 			MemoryPool->SetPoolIndex(InPoolIndex);
@@ -312,7 +315,32 @@ FD3D12TransientMemoryPool* FD3D12TransientMemoryPoolManager::GetOrCreateMemoryPo
 			PoolSize = FMath::Min(FMath::RoundUpToPowerOfTwo(InMinimumAllocationSize), (uint32)MaxAllocationSize);
 		}
 
-		MemoryPool = new FD3D12TransientMemoryPool(GetParentDevice(), GetVisibilityMask(), InitConfig,
+		// If merged heap types are not supported then create a resource type specific heap
+		FD3D12ResourceInitConfig PoolInitConfig = InitConfig;		
+		if (!bMergedTypeHeapSupported)
+		{
+			switch (InAllocationResourceType)
+			{
+			case ERHIPoolResourceTypes::Buffers:
+			{
+				PoolInitConfig.HeapFlags = D3D12_HEAP_FLAG_ALLOW_ONLY_BUFFERS;
+				break;
+			}
+			case ERHIPoolResourceTypes::RTDSTextures:
+			{
+				PoolInitConfig.HeapFlags = D3D12_HEAP_FLAG_ALLOW_ONLY_RT_DS_TEXTURES;
+				break;
+			}
+			case ERHIPoolResourceTypes::NonRTDSTextures:
+			{
+				PoolInitConfig.HeapFlags = D3D12_HEAP_FLAG_ALLOW_ONLY_NON_RT_DS_TEXTURES;
+				break;
+			}
+			}
+		}
+		ERHIPoolResourceTypes SupportedResourceTypes = bMergedTypeHeapSupported ? ERHIPoolResourceTypes::All : InAllocationResourceType;
+
+		MemoryPool = new FD3D12TransientMemoryPool(GetParentDevice(), GetVisibilityMask(), PoolInitConfig, SupportedResourceTypes,
 			TEXT("TransientResourceMemoryPool"), EResourceAllocationStrategy::kPlacedResource, InPoolIndex, PoolSize, PoolAlignment);
 		MemoryPool->Init();
 
@@ -819,10 +847,10 @@ void FD3D12TransientResourceAllocator::Freeze(FRHICommandListImmediate&)
 }
 
 
-FRHIMemoryPool* FD3D12TransientResourceAllocator::CreateNewPool(int16 InPoolIndex, uint32 InMinimumAllocationSize)
+FRHIMemoryPool* FD3D12TransientResourceAllocator::CreateNewPool(int16 InPoolIndex, uint32 InMinimumAllocationSize, ERHIPoolResourceTypes InAllocationResourceType)
 {
 	// Get pool from manager - don't reallocate each time
-	return GetParentDevice()->GetTransientMemoryPoolManager().GetOrCreateMemoryPool(InPoolIndex, InMinimumAllocationSize);
+	return GetParentDevice()->GetTransientMemoryPoolManager().GetOrCreateMemoryPool(InPoolIndex, InMinimumAllocationSize, InAllocationResourceType);
 }
 
 
diff --git a/Engine/Source/Runtime/D3D12RHI/Private/D3D12TransientResourceAllocator.h b/Engine/Source/Runtime/D3D12RHI/Private/D3D12TransientResourceAllocator.h
--- a/Engine/Source/Runtime/D3D12RHI/Private/D3D12TransientResourceAllocator.h
+++ b/Engine/Source/Runtime/D3D12RHI/Private/D3D12TransientResourceAllocator.h
@@ -63,9 +63,9 @@ class FD3D12TransientMemoryPool : public FD3D12MemoryPool
 {
 public:
 
-	FD3D12TransientMemoryPool(FD3D12Device* ParentDevice, FRHIGPUMask VisibleNodes, const FD3D12ResourceInitConfig& InInitConfig, const FString& Name,
+	FD3D12TransientMemoryPool(FD3D12Device* ParentDevice, FRHIGPUMask VisibleNodes, const FD3D12ResourceInitConfig& InInitConfig, ERHIPoolResourceTypes InSupportedResourceTypes, const FString& Name,
 		EResourceAllocationStrategy InAllocationStrategy, int16 InPoolIndex, uint64 InPoolSize, uint32 InPoolAlignment) :
-		FD3D12MemoryPool(ParentDevice, VisibleNodes, InInitConfig, Name, InAllocationStrategy, InPoolIndex, InPoolSize, InPoolAlignment, FRHIMemoryPool::EFreeListOrder::SortByOffset) { }
+		FD3D12MemoryPool(ParentDevice, VisibleNodes, InInitConfig, Name, InAllocationStrategy, InPoolIndex, InPoolSize, InPoolAlignment, InSupportedResourceTypes, FRHIMemoryPool::EFreeListOrder::SortByOffset) { }
 	virtual ~FD3D12TransientMemoryPool();
 
 	void ResetPool();
@@ -125,7 +125,7 @@ public:
 	void EndFrame();
 
 	// Get or creation functions from cached resources
-	FD3D12TransientMemoryPool* GetOrCreateMemoryPool(int16 InPoolIndex, uint32 InMinimumAllocationSize);
+	FD3D12TransientMemoryPool* GetOrCreateMemoryPool(int16 InPoolIndex, uint32 InMinimumAllocationSize, ERHIPoolResourceTypes InAllocationResourceType);
 	void ReleaseMemoryPool(FD3D12TransientMemoryPool* InMemoryPool);
 	FD3D12PooledTextureData GetPooledTexture(const FRHITextureCreateInfo& InCreateInfo, const TCHAR* InDebugName);
 	FD3D12PooledBufferData GetPooledBuffer(const FRHIBufferCreateInfo& InCreateInfo, const TCHAR* InDebugName);
@@ -155,6 +155,7 @@ private:
 	uint64 DefaultPoolSize;
 	uint32 PoolAlignment;
 	uint64 MaxAllocationSize;
+	bool bMergedTypeHeapSupported;
 
 	// Critical section to lock access to the pools & resource caches
 	FCriticalSection CS;
@@ -205,7 +206,7 @@ private:
 
 	// Override placed resource allocation helper function to get from cache if available
 	virtual FD3D12Resource* CreatePlacedResource(const FRHIPoolAllocationData& InAllocationData, const D3D12_RESOURCE_DESC& InDesc, D3D12_RESOURCE_STATES InCreateState, ED3D12ResourceStateMode InResourceStateMode, const D3D12_CLEAR_VALUE* InClearValue, const TCHAR* InName) override;
-	virtual FRHIMemoryPool* CreateNewPool(int16 InPoolIndex, uint32 InMinimumAllocationSize) override;
+	virtual FRHIMemoryPool* CreateNewPool(int16 InPoolIndex, uint32 InMinimumAllocationSize, ERHIPoolResourceTypes InAllocationResourceType) override;
 
 	// Shared allocation/deallocation helper
 	FD3D12BaseShaderResource* GetBaseShaderResource(FRHITexture* InRHITexture);
diff --git a/Engine/Source/Runtime/RHICore/Private/RHIPoolAllocator.cpp b/Engine/Source/Runtime/RHICore/Private/RHIPoolAllocator.cpp
--- a/Engine/Source/Runtime/RHICore/Private/RHIPoolAllocator.cpp
+++ b/Engine/Source/Runtime/RHICore/Private/RHIPoolAllocator.cpp
@@ -197,10 +197,11 @@ uint32 FRHIMemoryPool::GetAlignedOffset(uint32 InOffset, uint32 InPoolAlignment,
 }
 
 
-FRHIMemoryPool::FRHIMemoryPool(int16 InPoolIndex, uint64 InPoolSize, uint32 InPoolAlignment, EFreeListOrder InFreeListOrder) :
+FRHIMemoryPool::FRHIMemoryPool(int16 InPoolIndex, uint64 InPoolSize, uint32 InPoolAlignment, ERHIPoolResourceTypes InSupportedResourceTypes, EFreeListOrder InFreeListOrder) :
 	PoolIndex(InPoolIndex), 
 	PoolSize(InPoolSize), 
 	PoolAlignment(InPoolAlignment),
+	SupportedResourceTypes(InSupportedResourceTypes),
 	FreeListOrder(InFreeListOrder),
 	FreeSize(0), 
 	AligmnentWaste(0), 
@@ -256,8 +257,10 @@ void FRHIMemoryPool::Destroy()
 }
 
 
-bool FRHIMemoryPool::TryAllocate(uint32 InSizeInBytes, uint32 InAllocationAlignment, FRHIPoolAllocationData& AllocationData)
+bool FRHIMemoryPool::TryAllocate(uint32 InSizeInBytes, uint32 InAllocationAlignment, ERHIPoolResourceTypes InAllocationResourceType, FRHIPoolAllocationData& AllocationData)
 {
+	check(IsResourceTypeSupported(InAllocationResourceType));
+
 	int32 FreeBlockIndex = FindFreeBlock(InSizeInBytes, InAllocationAlignment);
 	if (FreeBlockIndex != INDEX_NONE)
 	{
@@ -341,7 +344,7 @@ void FRHIMemoryPool::TryClear(FRHIPoolAllocator* InAllocator, uint32 InMaxCopySi
 			// Try find pool to move requested allocation (in desired order already)
 			for (FRHIMemoryPool* TargetPool : InTargetPools)
 			{
-				if (TargetPool->TryAllocate(SizeToAllocate, AllocationAlignment, TempTargetAllocation))
+				if (TargetPool->TryAllocate(SizeToAllocate, AllocationAlignment, SupportedResourceTypes, TempTargetAllocation))
 				{
 					// RHI specific handling of the actual defrag request
 					InAllocator->HandleDefragRequest(BlockToMove, TempTargetAllocation);
@@ -565,7 +568,7 @@ void FRHIPoolAllocator::Destroy()
 }
 
 
-bool FRHIPoolAllocator::TryAllocateInternal(uint32 InSizeInBytes, uint32 InAllocationAlignment, FRHIPoolAllocationData& AllocationData)
+bool FRHIPoolAllocator::TryAllocateInternal(uint32 InSizeInBytes, uint32 InAllocationAlignment, ERHIPoolResourceTypes InAllocationResourceType, FRHIPoolAllocationData& AllocationData)
 {
 	TRACE_CPUPROFILER_EVENT_SCOPE(FRHIPoolAllocator::TryAllocateInternal);
 
@@ -575,7 +578,7 @@ bool FRHIPoolAllocator::TryAllocateInternal(uint32 InSizeInBytes, uint32 InAlloc
 	for (int32 PoolIndex : PoolAllocationOrder)
 	{
 		FRHIMemoryPool* Pool = Pools[PoolIndex];
-		if (Pool != nullptr && !Pool->IsFull() && Pool->TryAllocate(InSizeInBytes, InAllocationAlignment, AllocationData))
+		if (Pool != nullptr && Pool->IsResourceTypeSupported(InAllocationResourceType) && !Pool->IsFull() && Pool->TryAllocate(InSizeInBytes, InAllocationAlignment, InAllocationResourceType, AllocationData))
 		{
 			return true;
 		}
@@ -594,16 +597,16 @@ bool FRHIPoolAllocator::TryAllocateInternal(uint32 InSizeInBytes, uint32 InAlloc
 
 	if (NewPoolIndex >= 0)
 	{
-		Pools[NewPoolIndex] = CreateNewPool(NewPoolIndex, InSizeInBytes);
+		Pools[NewPoolIndex] = CreateNewPool(NewPoolIndex, InSizeInBytes, InAllocationResourceType);
 	}
 	else
 	{
 		NewPoolIndex = Pools.Num();		
-		Pools.Add(CreateNewPool(NewPoolIndex, InSizeInBytes));
+		Pools.Add(CreateNewPool(NewPoolIndex, InSizeInBytes, InAllocationResourceType));
 		PoolAllocationOrder.Add(NewPoolIndex);
 	}
 
-	return Pools[NewPoolIndex]->TryAllocate(InSizeInBytes, InAllocationAlignment, AllocationData);
+	return Pools[NewPoolIndex]->TryAllocate(InSizeInBytes, InAllocationAlignment, InAllocationResourceType, AllocationData);
 }
 
 
diff --git a/Engine/Source/Runtime/RHICore/Public/RHIPoolAllocator.h b/Engine/Source/Runtime/RHICore/Public/RHIPoolAllocator.h
--- a/Engine/Source/Runtime/RHICore/Public/RHIPoolAllocator.h
+++ b/Engine/Source/Runtime/RHICore/Public/RHIPoolAllocator.h
@@ -15,6 +15,18 @@ class FRHIPoolResource
 };
 
 
+/**
+@brief Resource type supported by the pools (can be any or all)
+**/
+enum class ERHIPoolResourceTypes
+{
+	Buffers			= 0x1,
+	RTDSTextures	= 0x2,
+	NonRTDSTextures = 0x4,
+	All				= Buffers | RTDSTextures | NonRTDSTextures,
+};
+
+
 /**
 @brief Pool allocator internal data
 **/
@@ -118,7 +130,7 @@ public:
 	};
 
 	// Constructor
-	FRHIMemoryPool(int16 InPoolIndex, uint64 InPoolSize, uint32 InPoolAlignment, EFreeListOrder InFreeListOrder);
+	FRHIMemoryPool(int16 InPoolIndex, uint64 InPoolSize, uint32 InPoolAlignment, ERHIPoolResourceTypes InSupportedResourceTypes, EFreeListOrder InFreeListOrder);
 	virtual ~FRHIMemoryPool();
 
 	// Setup/Shutdown
@@ -126,7 +138,7 @@ public:
 	virtual void Destroy();
 
 	// AllocationData functions: allocate, free, move, ...
-	bool TryAllocate(uint32 InSizeInBytes, uint32 InAllocationAlignment, FRHIPoolAllocationData& AllocationData);
+	bool TryAllocate(uint32 InSizeInBytes, uint32 InAllocationAlignment, ERHIPoolResourceTypes InAllocationResourceType, FRHIPoolAllocationData& AllocationData);
 	void Deallocate(FRHIPoolAllocationData& AllocationData);
 
 	// Bookkeeping and clearing
@@ -139,6 +151,8 @@ public:
 	uint64 GetUsedSize() const { return (PoolSize - FreeSize); }
 	uint32 GetAlignmentWaste() const { return AligmnentWaste; }
 	uint32 GetAllocatedBlocks() const { return AllocatedBlocks; }
+	ERHIPoolResourceTypes GetSupportedResourceTypes() const { return SupportedResourceTypes; }
+	bool IsResourceTypeSupported(ERHIPoolResourceTypes InType) const { return EnumHasAnyFlags(SupportedResourceTypes, InType); }
 	bool IsEmpty() const { return GetUsedSize() == 0; }
 	bool IsFull() const { return FreeSize == 0; }
 
@@ -164,6 +178,7 @@ protected:
 	int16 PoolIndex;
 	const uint64 PoolSize;
 	const uint32 PoolAlignment;
+	const ERHIPoolResourceTypes SupportedResourceTypes;
 	const EFreeListOrder FreeListOrder;
 
 	// Stats
@@ -210,11 +225,11 @@ public:
 protected:
 		
 	// Supported allocator operations
-	bool TryAllocateInternal(uint32 InSizeInBytes, uint32 InAllocationAlignment, FRHIPoolAllocationData& AllocationData);
+	bool TryAllocateInternal(uint32 InSizeInBytes, uint32 InAllocationAlignment, ERHIPoolResourceTypes InAllocationResourceType, FRHIPoolAllocationData& AllocationData);
 	void DeallocateInternal(FRHIPoolAllocationData& AllocationData);
 	
 	// Helper function to create a new platform specific pool
-	virtual FRHIMemoryPool* CreateNewPool(int16 InPoolIndex, uint32 InMinimumAllocationSize) = 0;
+	virtual FRHIMemoryPool* CreateNewPool(int16 InPoolIndex, uint32 InMinimumAllocationSize, ERHIPoolResourceTypes InAllocationResourceType) = 0;
 
 	// Handle a rhi specific defrag op
 	friend class FRHIMemoryPool;
