Add info during allocation dump if resource is resident or not and add option to only dump resident resources

diff --git a/Engine/Source/Runtime/D3D12RHI/Private/D3D12Adapter.cpp b/Engine/Source/Runtime/D3D12RHI/Private/D3D12Adapter.cpp
--- a/Engine/Source/Runtime/D3D12RHI/Private/D3D12Adapter.cpp
+++ b/Engine/Source/Runtime/D3D12RHI/Private/D3D12Adapter.cpp
@@ -1307,7 +1307,7 @@ static FAutoConsoleCommandWithOutputDevice GDumpTrackedD3D12AllocationsCmd(
 	TEXT("Dump all tracked d3d12 resource allocations."),
 	FConsoleCommandWithOutputDeviceDelegate::CreateStatic([](FOutputDevice& OutputDevice)
 	{
-		FD3D12DynamicRHI::GetD3DRHI()->GetAdapter().DumpTrackedAllocationData(OutputDevice, false);
+		FD3D12DynamicRHI::GetD3DRHI()->GetAdapter().DumpTrackedAllocationData(OutputDevice, false, false);
 	})
 );
 
@@ -1316,11 +1316,29 @@ static FAutoConsoleCommandWithOutputDevice GDumpTrackedD3D12AllocationCallstacks
 	TEXT("Dump all tracked d3d12 resource allocation callstacks."),
 	FConsoleCommandWithOutputDeviceDelegate::CreateStatic([](FOutputDevice& OutputDevice)
 	{
-		FD3D12DynamicRHI::GetD3DRHI()->GetAdapter().DumpTrackedAllocationData(OutputDevice, true);
+		FD3D12DynamicRHI::GetD3DRHI()->GetAdapter().DumpTrackedAllocationData(OutputDevice, false, true);
 	})
 );
 
-void FD3D12Adapter::DumpTrackedAllocationData(FOutputDevice& OutputDevice, bool bWithCallstack)
+static FAutoConsoleCommandWithOutputDevice GDumpTrackedD3D12ResidentAllocationsCmd(
+	TEXT("D3D12.DumpTrackedResidentAllocations"),
+	TEXT("Dump all tracked resisdent d3d12 resource allocations."),
+	FConsoleCommandWithOutputDeviceDelegate::CreateStatic([](FOutputDevice& OutputDevice)
+		{
+			FD3D12DynamicRHI::GetD3DRHI()->GetAdapter().DumpTrackedAllocationData(OutputDevice, true, false);
+		})
+);
+
+static FAutoConsoleCommandWithOutputDevice GDumpTrackedD3D12ResidentAllocationCallstacksCmd(
+	TEXT("D3D12.DumpTrackedResidentAllocationCallstacks"),
+	TEXT("Dump all tracked resident d3d12 resource allocation callstacks."),
+	FConsoleCommandWithOutputDeviceDelegate::CreateStatic([](FOutputDevice& OutputDevice)
+		{
+			FD3D12DynamicRHI::GetD3DRHI()->GetAdapter().DumpTrackedAllocationData(OutputDevice, true, true);
+		})
+);
+
+void FD3D12Adapter::DumpTrackedAllocationData(FOutputDevice& OutputDevice, bool bResidentOnly, bool bWithCallstack)
 {
 	FScopeLock Lock(&TrackedAllocationDataCS);
 
@@ -1333,20 +1351,32 @@ void FD3D12Adapter::DumpTrackedAllocationData(FOutputDevice& OutputDevice, bool
 
 	TArray<FTrackedAllocationData> BufferAllocations;	
 	TArray<FTrackedAllocationData> TextureAllocations;
-	uint64 TotalBufferAllocationSize = 0;
-	uint64 TotalTextureAllocationSize = 0;
+	uint64 TotalAllocatedBufferSize = 0;
+	uint64 TotalResidentBufferSize = 0;
+	uint64 TotalAllocatedTextureSize = 0;
+	uint64 TotalResidentTextureSize = 0;
 	for (FTrackedAllocationData& AllocationData : Allocations)
 	{
 		D3D12_RESOURCE_DESC ResourceDesc = AllocationData.ResourceAllocation->GetResource()->GetDesc();
 		if (ResourceDesc.Dimension == D3D12_RESOURCE_DIMENSION_BUFFER)
 		{
 			BufferAllocations.Add(AllocationData);
-			TotalBufferAllocationSize += AllocationData.AllocationSize;
+			TotalAllocatedBufferSize += AllocationData.AllocationSize;			
+#if ENABLE_RESIDENCY_MANAGEMENT
+			TotalResidentBufferSize += (AllocationData.ResourceAllocation->GetResidencyHandle()->ResidencyStatus == D3DX12Residency::ManagedObject::RESIDENCY_STATUS::RESIDENT) ? AllocationData.AllocationSize : 0;
+#else
+			TotalResidentBufferSize += AllocationData.AllocationSize;
+#endif 
 		}
 		else
 		{
 			TextureAllocations.Add(AllocationData);
-			TotalTextureAllocationSize += AllocationData.AllocationSize;
+			TotalAllocatedTextureSize += AllocationData.AllocationSize;
+#if ENABLE_RESIDENCY_MANAGEMENT
+			TotalResidentTextureSize += (AllocationData.ResourceAllocation->GetResidencyHandle()->ResidencyStatus == D3DX12Residency::ManagedObject::RESIDENCY_STATUS::RESIDENT) ? AllocationData.AllocationSize : 0;
+#else
+			TotalResidentTextureSize += AllocationData.AllocationSize;
+#endif 
 		}
 	}
 
@@ -1354,11 +1384,20 @@ void FD3D12Adapter::DumpTrackedAllocationData(FOutputDevice& OutputDevice, bool
 	ANSICHAR StackTrace[STRING_SIZE];
 
 	FString OutputData;
-	OutputData += FString::Printf(TEXT("\n%d Tracked Texture Allocations (Total size: %4.3fMB):\n"), TextureAllocations.Num(), TotalTextureAllocationSize / (1024.0f * 1024));
+	OutputData += FString::Printf(TEXT("\n%d Tracked Texture Allocations (Total size: %4.3fMB - Resident: %4.3fMB):\n"), TextureAllocations.Num(), TotalAllocatedTextureSize / (1024.0f * 1024), TotalResidentTextureSize / (1024.0f * 1024));
 	for (const FTrackedAllocationData& AllocationData : TextureAllocations)
 	{
 		D3D12_RESOURCE_DESC ResourceDesc = AllocationData.ResourceAllocation->GetResource()->GetDesc();
 
+		bool bResident = true;
+#if ENABLE_RESIDENCY_MANAGEMENT
+		bResident = AllocationData.ResourceAllocation->GetResidencyHandle()->ResidencyStatus == D3DX12Residency::ManagedObject::RESIDENCY_STATUS::RESIDENT;
+#endif 
+		if (!bResident && bResidentOnly)
+		{
+			continue;
+		}
+
 		FString Flags;
 		if (EnumHasAnyFlags(ResourceDesc.Flags, D3D12_RESOURCE_FLAG_ALLOW_RENDER_TARGET))
 		{
@@ -1373,11 +1412,12 @@ void FD3D12Adapter::DumpTrackedAllocationData(FOutputDevice& OutputDevice, bool
 			Flags += EnumHasAnyFlags(ResourceDesc.Flags, D3D12_RESOURCE_FLAG_ALLOW_RENDER_TARGET | D3D12_RESOURCE_FLAG_ALLOW_DEPTH_STENCIL) ? "|UAV" : "UAV";
 		}
 
-		OutputData += FString::Printf(TEXT("\tName: %s - Size: %3.3fMB - Width: %d - Height: %d - DepthOrArraySize: %d - MipLevels: %d - Flags: %s\n"), 
+		OutputData += FString::Printf(TEXT("\tName: %s - Size: %3.3fMB - Width: %d - Height: %d - DepthOrArraySize: %d - MipLevels: %d - Flags: %s - Resident: %s\n"), 
 			*AllocationData.ResourceAllocation->GetResource()->GetName().ToString(), 
 			AllocationData.AllocationSize / (1024.0f * 1024),
 			ResourceDesc.Width, ResourceDesc.Height, ResourceDesc.DepthOrArraySize, ResourceDesc.MipLevels,
-			Flags.IsEmpty() ? TEXT("None") : *Flags);
+			Flags.IsEmpty() ? TEXT("None") : *Flags,
+			bResident ? TEXT("Yes") : TEXT("No"));
 
 		if (bWithCallstack)
 		{
@@ -1391,16 +1431,26 @@ void FD3D12Adapter::DumpTrackedAllocationData(FOutputDevice& OutputDevice, bool
 		}
 	}
 
-	OutputData += FString::Printf(TEXT("\n\n%d Tracked Buffer Allocations (Total size: %4.3fMB):\n"), BufferAllocations.Num(), TotalBufferAllocationSize / (1024.0f * 1024));
+	OutputData += FString::Printf(TEXT("\n\n%d Tracked Buffer Allocations (Total size: %4.3fMB - Resident: %4.3fMB):\n"), BufferAllocations.Num(), TotalAllocatedBufferSize / (1024.0f * 1024), TotalResidentBufferSize / (1024.0f * 1024));
 	for (const FTrackedAllocationData& AllocationData : BufferAllocations)
 	{
 		D3D12_RESOURCE_DESC ResourceDesc = AllocationData.ResourceAllocation->GetResource()->GetDesc();
 
-		OutputData += FString::Printf(TEXT("\tName: %s - Size: %3.3fMB - Width: %d - UAV: %s\n"), 
+		bool bResident = true;
+#if ENABLE_RESIDENCY_MANAGEMENT
+		bResident = AllocationData.ResourceAllocation->GetResidencyHandle()->ResidencyStatus == D3DX12Residency::ManagedObject::RESIDENCY_STATUS::RESIDENT;
+#endif 
+		if (!bResident && bResidentOnly)
+		{
+			continue;
+		}
+
+		OutputData += FString::Printf(TEXT("\tName: %s - Size: %3.3fMB - Width: %d - UAV: %s - Resident: %s\n"), 
 			*AllocationData.ResourceAllocation->GetResource()->GetName().ToString(), 
 			AllocationData.AllocationSize / (1024.0f * 1024),
 			ResourceDesc.Width,
-			EnumHasAnyFlags(ResourceDesc.Flags, D3D12_RESOURCE_FLAG_ALLOW_UNORDERED_ACCESS) ? TEXT("Yes") : TEXT("No"));
+			EnumHasAnyFlags(ResourceDesc.Flags, D3D12_RESOURCE_FLAG_ALLOW_UNORDERED_ACCESS) ? TEXT("Yes") : TEXT("No"),
+			bResident ? TEXT("Yes") : TEXT("No"));
 
 		if (bWithCallstack)
 		{
diff --git a/Engine/Source/Runtime/D3D12RHI/Private/D3D12Adapter.h b/Engine/Source/Runtime/D3D12RHI/Private/D3D12Adapter.h
--- a/Engine/Source/Runtime/D3D12RHI/Private/D3D12Adapter.h
+++ b/Engine/Source/Runtime/D3D12RHI/Private/D3D12Adapter.h
@@ -347,7 +347,7 @@ public:
 
 	void TrackAllocationData(FD3D12ResourceLocation* InAllocation, uint64 InAllocationSize);
 	void ReleaseTrackedAllocationData(FD3D12ResourceLocation* InAllocation);
-	void DumpTrackedAllocationData(FOutputDevice& OutputDevice, bool bWithCallstack);
+	void DumpTrackedAllocationData(FOutputDevice& OutputDevice, bool bResidentOnly, bool bWithCallstack);
 
 #if D3D12_SUBMISSION_GAP_RECORDER
 	FD3D12SubmissionGapRecorder SubmissionGapRecorder;
