Added D3D12.DumpRayTracingGeometries command to report sizes and other stats for all live DXR bottom level acceleration structures. Available in non-shipping builds.

diff --git a/Engine/Source/Runtime/D3D12RHI/Private/D3D12RayTracing.cpp b/Engine/Source/Runtime/D3D12RHI/Private/D3D12RayTracing.cpp
--- a/Engine/Source/Runtime/D3D12RHI/Private/D3D12RayTracing.cpp
+++ b/Engine/Source/Runtime/D3D12RHI/Private/D3D12RayTracing.cpp
@@ -15,6 +15,7 @@
 #include "Misc/ScopeLock.h"
 #include "RayTracingInstanceCopyShader.h"
 #include "Async/ParallelFor.h"
+#include "Misc/BufferedOutputDevice.h"
 
 static int32 GRayTracingDebugForceOpaque = 0;
 static FAutoConsoleVariableRef CVarRayTracingDebugForceOpaque(
@@ -128,6 +129,93 @@ DECLARE_CYCLE_STAT(TEXT("BuildTopLevel"), STAT_D3D12BuildTLAS, STATGROUP_D3D12Ra
 DECLARE_CYCLE_STAT(TEXT("BuildBottomLevel"), STAT_D3D12BuildBLAS, STATGROUP_D3D12RayTracing);
 DECLARE_CYCLE_STAT(TEXT("DispatchRays"), STAT_D3D12DispatchRays, STATGROUP_D3D12RayTracing);
 
+
+#if UE_BUILD_SHIPPING
+inline void RegisterD3D12RayTracingGeometry(FD3D12RayTracingGeometry* Geometry) {};
+inline void UnregisterD3D12RayTracingGeometry(FD3D12RayTracingGeometry* Geometry) {};
+#else
+struct FD3D12RayTracingGeometryTracker
+{
+	TSet<FD3D12RayTracingGeometry*> Geometries;
+	FCriticalSection CS;
+
+	void Add(FD3D12RayTracingGeometry* Geometry)
+	{
+		FScopeLock Lock(&CS);
+		Geometries.Add(Geometry);
+	}
+
+	void Remove(FD3D12RayTracingGeometry* Geometry)
+	{
+		FScopeLock Lock(&CS);
+		Geometries.Remove(Geometry);
+	}
+};
+
+static FD3D12RayTracingGeometryTracker& GetD3D12RayTracingGeometryTracker()
+{
+	static FD3D12RayTracingGeometryTracker Instance;
+	return Instance;
+}
+
+static FAutoConsoleCommandWithOutputDevice GD3D12DumpRayTracingGeometriesCmd(
+	TEXT("D3D12.DumpRayTracingGeometries"),
+	TEXT("Dump memory allocations for ray tracing resources."),
+	FConsoleCommandWithOutputDeviceDelegate::CreateStatic([](FOutputDevice& OutputDevice)
+{
+	FD3D12RayTracingGeometryTracker& Tracker = GetD3D12RayTracingGeometryTracker();
+	FScopeLock Lock(&Tracker.CS);
+
+	auto GetGeometrySize = [](FD3D12RayTracingGeometry& Geometry)
+	{
+		if (Geometry.AccelerationStructureCompactedSize != 0)
+		{
+			return Geometry.AccelerationStructureCompactedSize;
+		}
+		else
+		{
+			return Geometry.SizeInfo.ResultSize;
+		}
+	};
+
+	TArray<FD3D12RayTracingGeometry*> Geometries = Tracker.Geometries.Array();
+	Geometries.Sort([GetGeometrySize](FD3D12RayTracingGeometry& A, FD3D12RayTracingGeometry& B)
+	{
+		return GetGeometrySize(A) > GetGeometrySize(B);
+	});
+
+	FBufferedOutputDevice BufferedOutput;
+	FName CategoryName(TEXT("D3D12RayTracing"));
+	uint64 TotalSizeBytes = 0;
+	BufferedOutput.CategorizedLogf(CategoryName, ELogVerbosity::Log, TEXT("Tracked FD3D12RayTracingGeometry objects"));
+	for (FD3D12RayTracingGeometry* Geometry : Geometries)
+	{
+		uint64 SizeBytes = GetGeometrySize(*Geometry);
+		BufferedOutput.CategorizedLogf(CategoryName, ELogVerbosity::Log, TEXT("Name: %s - Size: %.3f MB - Prims: %d - Segments: %d -  Compaction: %d - Update: %d"),
+			Geometry->DebugName.IsValid() ? *Geometry->DebugName.ToString() : TEXT("*UNKNOWN*"),
+			SizeBytes / double(1 << 20),
+			Geometry->TotalPrimitiveCount,
+			Geometry->Segments.Num(),
+			(int32)EnumHasAllFlags(Geometry->BuildFlags, ERayTracingAccelerationStructureFlags::AllowCompaction),
+			(int32)EnumHasAllFlags(Geometry->BuildFlags, ERayTracingAccelerationStructureFlags::AllowUpdate));
+		TotalSizeBytes += SizeBytes;
+	}
+
+	BufferedOutput.CategorizedLogf(CategoryName, ELogVerbosity::Log, TEXT("Total size: %.3f MB"), TotalSizeBytes / double(1<<20));
+
+	BufferedOutput.RedirectTo(OutputDevice);
+}));
+
+inline void RegisterD3D12RayTracingGeometry(FD3D12RayTracingGeometry* Geometry)
+{
+	GetD3D12RayTracingGeometryTracker().Add(Geometry);
+}
+inline void UnregisterD3D12RayTracingGeometry(FD3D12RayTracingGeometry* Geometry)
+{
+	GetD3D12RayTracingGeometryTracker().Remove(Geometry);
+}
+#endif // UE_BUILD_SHIPPING
+
 // Whether to compare the full descriptor table on cache lookup or only use CityHash64 digest.
 #ifndef RAY_TRACING_DESCRIPTOR_CACHE_FULL_COMPARE
 #define RAY_TRACING_DESCRIPTOR_CACHE_FULL_COMPARE 1
@@ -2752,10 +2840,14 @@ FD3D12RayTracingGeometry::FD3D12RayTracingGeometry(FD3D12Adapter* Adapter, const
 	}
 
 	INC_DWORD_STAT_BY(STAT_D3D12RayTracingTrianglesBLAS, Initializer.TotalPrimitiveCount);
+
+	RegisterD3D12RayTracingGeometry(this);
 }
 
 FD3D12RayTracingGeometry::~FD3D12RayTracingGeometry()
 {
+	UnregisterD3D12RayTracingGeometry(this);
+
 	// Remove compaction request if still pending
 	for (uint32 GPUIndex = 0; GPUIndex < MAX_NUM_GPUS; ++GPUIndex)
 	{
@@ -3134,6 +3226,8 @@ void FD3D12RayTracingGeometry::CompactAccelerationStructure(FD3D12CommandContext
 		AccelerationStructureBuffers[InGPUIndex]->ResourceLocation.GetGPUVirtualAddress(),
 		OldAccelerationStructure->ResourceLocation.GetGPUVirtualAddress(),
 		D3D12_RAYTRACING_ACCELERATION_STRUCTURE_COPY_MODE_COMPACT);
+
+	AccelerationStructureCompactedSize = InSizeAfterCompaction;
 }
 
 FD3D12RayTracingScene::FD3D12RayTracingScene(FD3D12Adapter* Adapter, const FRayTracingSceneInitializer& Initializer)
