Tessellation Removal: Removing Tessellation from D3D12RHI

diff --git a/Engine/Source/Runtime/D3D12RHI/Private/D3D12CommandContext.cpp b/Engine/Source/Runtime/D3D12RHI/Private/D3D12CommandContext.cpp
--- a/Engine/Source/Runtime/D3D12RHI/Private/D3D12CommandContext.cpp
+++ b/Engine/Source/Runtime/D3D12RHI/Private/D3D12CommandContext.cpp
@@ -112,7 +112,6 @@ FD3D12CommandContext::FD3D12CommandContext(FD3D12Device* InParent, ED3D12Command
 	bOuterOcclusionQuerySubmitted(false),
 	bDiscardSharedGraphicsConstants(false),
 	bDiscardSharedComputeConstants(false),
-	bUsingTessellation(false),
 #if PLATFORM_SUPPORTS_VARIABLE_RATE_SHADING
 	VRSCombiners{ D3D12_SHADING_RATE_COMBINER_PASSTHROUGH, D3D12_SHADING_RATE_COMBINER_PASSTHROUGH },
 	VRSShadingRate(D3D12_SHADING_RATE_1X1),
@@ -125,8 +124,6 @@ FD3D12CommandContext::FD3D12CommandContext(FD3D12Device* InParent, ED3D12Command
 	VSConstantBuffer(InParent, ConstantsAllocator),
 	MSConstantBuffer(InParent, ConstantsAllocator),
 	ASConstantBuffer(InParent, ConstantsAllocator),
-	HSConstantBuffer(InParent, ConstantsAllocator),
-	DSConstantBuffer(InParent, ConstantsAllocator),
 	PSConstantBuffer(InParent, ConstantsAllocator),
 	GSConstantBuffer(InParent, ConstantsAllocator),
 	CSConstantBuffer(InParent, ConstantsAllocator)
@@ -551,8 +548,6 @@ void FD3D12CommandContext::ClearState()
 		CurrentDepthTexture = nullptr;
 
 		CurrentDSVAccessType = FExclusiveDepthStencil::DepthWrite_StencilWrite;
-
-		bUsingTessellation = false;
 	}
 }
 
@@ -560,8 +555,8 @@ void FD3D12CommandContext::ConditionalClearShaderResource(FD3D12ResourceLocation
 {
 	check(Resource);
 	StateCache.ClearShaderResourceViews<SF_Vertex>(Resource);
-	StateCache.ClearShaderResourceViews<SF_Hull>(Resource);
-	StateCache.ClearShaderResourceViews<SF_Domain>(Resource);
+	StateCache.ClearShaderResourceViews<SF_Mesh>(Resource);
+	StateCache.ClearShaderResourceViews<SF_Amplification>(Resource);
 	StateCache.ClearShaderResourceViews<SF_Pixel>(Resource);
 	StateCache.ClearShaderResourceViews<SF_Geometry>(Resource);
 	StateCache.ClearShaderResourceViews<SF_Compute>(Resource);
diff --git a/Engine/Source/Runtime/D3D12RHI/Private/D3D12CommandContext.h b/Engine/Source/Runtime/D3D12RHI/Private/D3D12CommandContext.h
--- a/Engine/Source/Runtime/D3D12RHI/Private/D3D12CommandContext.h
+++ b/Engine/Source/Runtime/D3D12RHI/Private/D3D12CommandContext.h
@@ -163,9 +163,6 @@ public:
 	/** When a new compute PSO is set, we discard all old constants set for the previous shader. */
 	bool bDiscardSharedComputeConstants;
 
-	/** Set to true when the current shading setup uses tessellation */
-	bool bUsingTessellation;
-
 	/** Used by variable rate shading to cache the current state of the combiners and the constant shading rate*/
 #if PLATFORM_SUPPORTS_VARIABLE_RATE_SHADING
 	D3D12_SHADING_RATE_COMBINER		VRSCombiners[D3D12_RS_SET_SHADING_RATE_COMBINER_COUNT];
@@ -209,8 +206,6 @@ public:
 	FD3D12ConstantBuffer VSConstantBuffer;
 	FD3D12ConstantBuffer MSConstantBuffer;
 	FD3D12ConstantBuffer ASConstantBuffer;
-	FD3D12ConstantBuffer HSConstantBuffer;
-	FD3D12ConstantBuffer DSConstantBuffer;
 	FD3D12ConstantBuffer PSConstantBuffer;
 	FD3D12ConstantBuffer GSConstantBuffer;
 	FD3D12ConstantBuffer CSConstantBuffer;
diff --git a/Engine/Source/Runtime/D3D12RHI/Private/D3D12Commands.cpp b/Engine/Source/Runtime/D3D12RHI/Private/D3D12Commands.cpp
--- a/Engine/Source/Runtime/D3D12RHI/Private/D3D12Commands.cpp
+++ b/Engine/Source/Runtime/D3D12RHI/Private/D3D12Commands.cpp
@@ -40,8 +40,6 @@ DECLARE_ISBOUNDSHADER(MeshShader)
 DECLARE_ISBOUNDSHADER(AmplificationShader)
 DECLARE_ISBOUNDSHADER(PixelShader)
 DECLARE_ISBOUNDSHADER(GeometryShader)
-DECLARE_ISBOUNDSHADER(HullShader)
-DECLARE_ISBOUNDSHADER(DomainShader)
 DECLARE_ISBOUNDSHADER(ComputeShader)
 
 #define DECLARE_ISVALIDUNIFORMBUFFERHASH(ShaderType) inline void ValidateBoundUniformBuffer(FD3D12UniformBuffer* InUniformBuffer, FRHI##ShaderType* In##RHI, uint32 InBufferIndex) \
@@ -60,8 +58,6 @@ DECLARE_ISVALIDUNIFORMBUFFERHASH(MeshShader)
 DECLARE_ISVALIDUNIFORMBUFFERHASH(AmplificationShader)
 DECLARE_ISVALIDUNIFORMBUFFERHASH(PixelShader)
 DECLARE_ISVALIDUNIFORMBUFFERHASH(GeometryShader)
-DECLARE_ISVALIDUNIFORMBUFFERHASH(HullShader)
-DECLARE_ISVALIDUNIFORMBUFFERHASH(DomainShader)
 DECLARE_ISVALIDUNIFORMBUFFERHASH(ComputeShader)
 
 #if EXECUTE_DEBUG_COMMAND_LISTS
@@ -730,15 +726,11 @@ void FD3D12CommandContext::RHISetGraphicsPipelineState(FRHIGraphicsPipelineState
 	FD3D12GraphicsPipelineState* GraphicsPipelineState = FD3D12DynamicRHI::ResourceCast(GraphicsState);
 
 	// TODO: [PSO API] Every thing inside this scope is only necessary to keep the PSO shadow in sync while we convert the high level to only use PSOs
-	const bool bWasUsingTessellation = bUsingTessellation;
-	bUsingTessellation = GraphicsPipelineState->GetHullShader() && GraphicsPipelineState->GetDomainShader();
 	// Ensure the command buffers are reset to reduce the amount of data that needs to be versioned.
 	VSConstantBuffer.Reset();
 	MSConstantBuffer.Reset();
 	ASConstantBuffer.Reset();
 	PSConstantBuffer.Reset();
-	HSConstantBuffer.Reset();
-	DSConstantBuffer.Reset();
 	GSConstantBuffer.Reset();
 	
 	// @TODO : really should only discard the constants if the shader state has actually changed.
@@ -754,7 +746,7 @@ void FD3D12CommandContext::RHISetGraphicsPipelineState(FRHIGraphicsPipelineState
 		StateCache.SetShadingRate(GraphicsPipelineState->PipelineStateInitializer.ShadingRate, VRSRB_Passthrough);
 	}
 
-	StateCache.SetGraphicsPipelineState(GraphicsPipelineState, bUsingTessellation != bWasUsingTessellation);
+	StateCache.SetGraphicsPipelineState(GraphicsPipelineState);
 	StateCache.SetStencilRef(0);
 
 	if (bApplyAdditionalState)
@@ -762,8 +754,6 @@ void FD3D12CommandContext::RHISetGraphicsPipelineState(FRHIGraphicsPipelineState
 		ApplyStaticUniformBuffers(GraphicsPipelineState->GetVertexShader());
 		ApplyStaticUniformBuffers(GraphicsPipelineState->GetMeshShader());
 		ApplyStaticUniformBuffers(GraphicsPipelineState->GetAmplificationShader());
-		ApplyStaticUniformBuffers(GraphicsPipelineState->GetHullShader());
-		ApplyStaticUniformBuffers(GraphicsPipelineState->GetDomainShader());
 		ApplyStaticUniformBuffers(GraphicsPipelineState->GetGeometryShader());
 		ApplyStaticUniformBuffers(GraphicsPipelineState->GetPixelShader());
 	}
@@ -816,20 +806,6 @@ void FD3D12CommandContext::RHISetShaderTexture(FRHIGraphicsShader* ShaderRHI, ui
 		StateCache.SetShaderResourceView<SF_Amplification>(NewTexture ? NewTexture->GetShaderResourceView() : nullptr, TextureIndex);
 	}
 	break;
-	case SF_Hull:
-	{
-		FRHIHullShader* HullShaderRHI = static_cast<FRHIHullShader*>(ShaderRHI);
-		VALIDATE_BOUND_SHADER(HullShaderRHI);
-		StateCache.SetShaderResourceView<SF_Hull>(NewTexture ? NewTexture->GetShaderResourceView() : nullptr, TextureIndex);
-	}
-	break;
-	case SF_Domain:
-	{
-		FRHIDomainShader* DomainShaderRHI = static_cast<FRHIDomainShader*>(ShaderRHI);
-		VALIDATE_BOUND_SHADER(DomainShaderRHI);
-		StateCache.SetShaderResourceView<SF_Domain>(NewTexture ? NewTexture->GetShaderResourceView() : nullptr, TextureIndex);
-	}
-	break;
 	case SF_Geometry:
 	{
 		FRHIGeometryShader* GeometryShaderRHI = static_cast<FRHIGeometryShader*>(ShaderRHI);
@@ -929,20 +905,6 @@ void FD3D12CommandContext::RHISetShaderResourceViewParameter(FRHIGraphicsShader*
 		StateCache.SetShaderResourceView<SF_Amplification>(SRV, TextureIndex);
 	}
 	break;
-	case SF_Hull:
-	{
-		FRHIHullShader* HullShaderRHI = static_cast<FRHIHullShader*>(ShaderRHI);
-		VALIDATE_BOUND_SHADER(HullShaderRHI);
-		StateCache.SetShaderResourceView<SF_Hull>(SRV, TextureIndex);
-	}
-	break;
-	case SF_Domain:
-	{
-		FRHIDomainShader* DomainShaderRHI = static_cast<FRHIDomainShader*>(ShaderRHI);
-		VALIDATE_BOUND_SHADER(DomainShaderRHI);
-		StateCache.SetShaderResourceView<SF_Domain>(SRV, TextureIndex);
-	}
-	break;
 	case SF_Geometry:
 	{
 		FRHIGeometryShader* GeometryShaderRHI = static_cast<FRHIGeometryShader*>(ShaderRHI);
@@ -996,20 +958,6 @@ void FD3D12CommandContext::RHISetShaderSampler(FRHIGraphicsShader* ShaderRHI, ui
 		StateCache.SetSamplerState<SF_Amplification>(NewState, SamplerIndex);
 	}
 	break;
-	case SF_Hull:
-	{
-		FRHIHullShader* HullShaderRHI = static_cast<FRHIHullShader*>(ShaderRHI);
-		VALIDATE_BOUND_SHADER(HullShaderRHI);
-		StateCache.SetSamplerState<SF_Hull>(NewState, SamplerIndex);
-	}
-	break;
-	case SF_Domain:
-	{
-		FRHIDomainShader* DomainShaderRHI = static_cast<FRHIDomainShader*>(ShaderRHI);
-		VALIDATE_BOUND_SHADER(DomainShaderRHI);
-		StateCache.SetSamplerState<SF_Domain>(NewState, SamplerIndex);
-	}
-	break;
 	case SF_Geometry:
 	{
 		FRHIGeometryShader* GeometryShaderRHI = static_cast<FRHIGeometryShader*>(ShaderRHI);
@@ -1070,24 +1018,6 @@ void FD3D12CommandContext::RHISetShaderUniformBuffer(FRHIGraphicsShader* ShaderR
 		Stage = SF_Amplification;
 	}
 	break;
-	case SF_Hull:
-	{
-		FRHIHullShader* HullShaderRHI = static_cast<FRHIHullShader*>(ShaderRHI);
-		VALIDATE_BOUND_SHADER(HullShaderRHI);
-		VALIDATE_BOUND_UNIFORMBUFFER_HASH(Buffer, HullShaderRHI, BufferIndex);
-		StateCache.SetConstantsFromUniformBuffer<SF_Hull>(BufferIndex, Buffer);
-		Stage = SF_Hull;
-	}
-	break;
-	case SF_Domain:
-	{
-		FRHIDomainShader* DomainShaderRHI = static_cast<FRHIDomainShader*>(ShaderRHI);
-		VALIDATE_BOUND_SHADER(DomainShaderRHI);
-		VALIDATE_BOUND_UNIFORMBUFFER_HASH(Buffer, DomainShaderRHI, BufferIndex);
-		StateCache.SetConstantsFromUniformBuffer<SF_Domain>(BufferIndex, Buffer);
-		Stage = SF_Domain;
-	}
-	break;
 	case SF_Geometry:
 	{
 		FRHIGeometryShader* GeometryShaderRHI = static_cast<FRHIGeometryShader*>(ShaderRHI);
@@ -1164,20 +1094,6 @@ void FD3D12CommandContext::RHISetShaderParameter(FRHIGraphicsShader* ShaderRHI,
 		ASConstantBuffer.UpdateConstant((const uint8*)NewValue, BaseIndex, NumBytes);
 	}
 	break;
-	case SF_Hull:
-	{
-		FRHIHullShader* HullShaderRHI = static_cast<FRHIHullShader*>(ShaderRHI);
-		VALIDATE_BOUND_SHADER(HullShaderRHI);
-		HSConstantBuffer.UpdateConstant((const uint8*)NewValue, BaseIndex, NumBytes);
-	}
-	break;
-	case SF_Domain:
-	{
-		FRHIDomainShader* DomainShaderRHI = static_cast<FRHIDomainShader*>(ShaderRHI);
-		VALIDATE_BOUND_SHADER(DomainShaderRHI);
-		DSConstantBuffer.UpdateConstant((const uint8*)NewValue, BaseIndex, NumBytes);
-	}
-	break;
 	case SF_Geometry:
 	{
 		FRHIGeometryShader* GeometryShaderRHI = static_cast<FRHIGeometryShader*>(ShaderRHI);
@@ -1591,23 +1507,6 @@ void FD3D12CommandContext::CommitNonComputeShaderConstants()
 		StateCache.SetConstantBuffer<SF_Amplification>(ASConstantBuffer, bDiscardSharedGraphicsConstants);
 	}
 
-	// Skip HS/DS CB updates in cases where tessellation isn't being used
-	// Note that this is *potentially* unsafe because bDiscardSharedGraphicsConstants is cleared at the
-	// end of the function, however we're OK for now because bDiscardSharedGraphicsConstants
-	// is always reset whenever bUsingTessellation changes in SetBoundShaderState()
-	if (bUsingTessellation)
-	{
-		if (GraphicPSO->bShaderNeedsGlobalConstantBuffer[SF_Hull])
-		{
-			StateCache.SetConstantBuffer<SF_Hull>(HSConstantBuffer, bDiscardSharedGraphicsConstants);
-		}
-
-		if (GraphicPSO->bShaderNeedsGlobalConstantBuffer[SF_Domain])
-		{
-			StateCache.SetConstantBuffer<SF_Domain>(DSConstantBuffer, bDiscardSharedGraphicsConstants);
-		}
-	}
-
 	if (GraphicPSO->bShaderNeedsGlobalConstantBuffer[SF_Geometry])
 	{
 		StateCache.SetConstantBuffer<SF_Geometry>(GSConstantBuffer, bDiscardSharedGraphicsConstants);
@@ -1893,14 +1792,6 @@ void FD3D12CommandContext::CommitGraphicsResourceTables()
 	{
 		SetResourcesFromTables(PixelShader);
 	}
-	if (auto* Shader = GraphicPSO->GetHullShader())
-	{
-		SetResourcesFromTables(Shader);
-	}
-	if (auto* Shader = GraphicPSO->GetDomainShader())
-	{
-		SetResourcesFromTables(Shader);
-	}
 	if (auto* Shader = GraphicPSO->GetGeometryShader())
 	{
 		SetResourcesFromTables(Shader);
diff --git a/Engine/Source/Runtime/D3D12RHI/Private/D3D12DescriptorCache.cpp b/Engine/Source/Runtime/D3D12RHI/Private/D3D12DescriptorCache.cpp
--- a/Engine/Source/Runtime/D3D12RHI/Private/D3D12DescriptorCache.cpp
+++ b/Engine/Source/Runtime/D3D12RHI/Private/D3D12DescriptorCache.cpp
@@ -18,8 +18,6 @@ static FAutoConsoleVariableRef CVarGlobalViewHeapBlockSize(
 template void FD3D12DescriptorCache::SetConstantBuffers<SF_Vertex>(const FD3D12RootSignature* RootSignature, FD3D12ConstantBufferCache& Cache, const CBVSlotMask& SlotsNeededMask, uint32 Count, uint32& HeapSlot);
 template void FD3D12DescriptorCache::SetConstantBuffers<SF_Mesh>(const FD3D12RootSignature* RootSignature, FD3D12ConstantBufferCache& Cache, const CBVSlotMask& SlotsNeededMask, uint32 Count, uint32& HeapSlot);
 template void FD3D12DescriptorCache::SetConstantBuffers<SF_Amplification>(const FD3D12RootSignature* RootSignature, FD3D12ConstantBufferCache& Cache, const CBVSlotMask& SlotsNeededMask, uint32 Count, uint32& HeapSlot);
-template void FD3D12DescriptorCache::SetConstantBuffers<SF_Hull>(const FD3D12RootSignature* RootSignature, FD3D12ConstantBufferCache& Cache, const CBVSlotMask& SlotsNeededMask, uint32 Count, uint32& HeapSlot);
-template void FD3D12DescriptorCache::SetConstantBuffers<SF_Domain>(const FD3D12RootSignature* RootSignature, FD3D12ConstantBufferCache& Cache, const CBVSlotMask& SlotsNeededMask, uint32 Count, uint32& HeapSlot);
 template void FD3D12DescriptorCache::SetConstantBuffers<SF_Geometry>(const FD3D12RootSignature* RootSignature, FD3D12ConstantBufferCache& Cache, const CBVSlotMask& SlotsNeededMask, uint32 Count, uint32& HeapSlot);
 template void FD3D12DescriptorCache::SetConstantBuffers<SF_Pixel>(const FD3D12RootSignature* RootSignature, FD3D12ConstantBufferCache& Cache, const CBVSlotMask& SlotsNeededMask, uint32 Count, uint32& HeapSlot);
 template void FD3D12DescriptorCache::SetConstantBuffers<SF_Compute>(const FD3D12RootSignature* RootSignature, FD3D12ConstantBufferCache& Cache, const CBVSlotMask& SlotsNeededMask, uint32 Count, uint32& HeapSlot);
@@ -27,8 +25,6 @@ template void FD3D12DescriptorCache::SetConstantBuffers<SF_Compute>(const FD3D12
 template void FD3D12DescriptorCache::SetConstantBuffers<SF_Vertex>(const FD3D12RootSignature* RootSignature, FD3D12ConstantBufferCache& Cache, const CBVSlotMask& SlotsNeededMask);
 template void FD3D12DescriptorCache::SetConstantBuffers<SF_Mesh>(const FD3D12RootSignature* RootSignature, FD3D12ConstantBufferCache& Cache, const CBVSlotMask& SlotsNeededMask);
 template void FD3D12DescriptorCache::SetConstantBuffers<SF_Amplification>(const FD3D12RootSignature* RootSignature, FD3D12ConstantBufferCache& Cache, const CBVSlotMask& SlotsNeededMask);
-template void FD3D12DescriptorCache::SetConstantBuffers<SF_Hull>(const FD3D12RootSignature* RootSignature, FD3D12ConstantBufferCache& Cache, const CBVSlotMask& SlotsNeededMask);
-template void FD3D12DescriptorCache::SetConstantBuffers<SF_Domain>(const FD3D12RootSignature* RootSignature, FD3D12ConstantBufferCache& Cache, const CBVSlotMask& SlotsNeededMask);
 template void FD3D12DescriptorCache::SetConstantBuffers<SF_Geometry>(const FD3D12RootSignature* RootSignature, FD3D12ConstantBufferCache& Cache, const CBVSlotMask& SlotsNeededMask);
 template void FD3D12DescriptorCache::SetConstantBuffers<SF_Pixel>(const FD3D12RootSignature* RootSignature, FD3D12ConstantBufferCache& Cache, const CBVSlotMask& SlotsNeededMask);
 template void FD3D12DescriptorCache::SetConstantBuffers<SF_Compute>(const FD3D12RootSignature* RootSignature, FD3D12ConstantBufferCache& Cache, const CBVSlotMask& SlotsNeededMask);
@@ -37,8 +33,6 @@ template void FD3D12DescriptorCache::SetConstantBuffers<SF_Compute>(const FD3D12
 template void FD3D12DescriptorCache::SetSRVs<SF_Vertex>(const FD3D12RootSignature* RootSignature, FD3D12ShaderResourceViewCache& Cache, const SRVSlotMask& SlotsNeededMask, uint32 SlotsNeeded, uint32& HeapSlot);
 template void FD3D12DescriptorCache::SetSRVs<SF_Mesh>(const FD3D12RootSignature* RootSignature, FD3D12ShaderResourceViewCache& Cache, const SRVSlotMask& SlotsNeededMask, uint32 SlotsNeeded, uint32& HeapSlot);
 template void FD3D12DescriptorCache::SetSRVs<SF_Amplification>(const FD3D12RootSignature* RootSignature, FD3D12ShaderResourceViewCache& Cache, const SRVSlotMask& SlotsNeededMask, uint32 SlotsNeeded, uint32& HeapSlot);
-template void FD3D12DescriptorCache::SetSRVs<SF_Hull>(const FD3D12RootSignature* RootSignature, FD3D12ShaderResourceViewCache& Cache, const SRVSlotMask& SlotsNeededMask, uint32 SlotsNeeded, uint32& HeapSlot);
-template void FD3D12DescriptorCache::SetSRVs<SF_Domain>(const FD3D12RootSignature* RootSignature, FD3D12ShaderResourceViewCache& Cache, const SRVSlotMask& SlotsNeededMask, uint32 SlotsNeeded, uint32& HeapSlot);
 template void FD3D12DescriptorCache::SetSRVs<SF_Geometry>(const FD3D12RootSignature* RootSignature, FD3D12ShaderResourceViewCache& Cache, const SRVSlotMask& SlotsNeededMask, uint32 SlotsNeeded, uint32& HeapSlot);
 template void FD3D12DescriptorCache::SetSRVs<SF_Pixel>(const FD3D12RootSignature* RootSignature, FD3D12ShaderResourceViewCache& Cache, const SRVSlotMask& SlotsNeededMask, uint32 SlotsNeeded, uint32& HeapSlot);
 template void FD3D12DescriptorCache::SetSRVs<SF_Compute>(const FD3D12RootSignature* RootSignature, FD3D12ShaderResourceViewCache& Cache, const SRVSlotMask& SlotsNeededMask, uint32 SlotsNeeded, uint32& HeapSlot);
@@ -49,8 +43,6 @@ template void FD3D12DescriptorCache::SetUAVs<SF_Compute>(const FD3D12RootSignatu
 template void FD3D12DescriptorCache::SetSamplers<SF_Vertex>(const FD3D12RootSignature* RootSignature, FD3D12SamplerStateCache& Cache, const SamplerSlotMask& SlotsNeededMask, uint32 SlotsNeeded, uint32& HeapSlot);
 template void FD3D12DescriptorCache::SetSamplers<SF_Mesh>(const FD3D12RootSignature* RootSignature, FD3D12SamplerStateCache& Cache, const SamplerSlotMask& SlotsNeededMask, uint32 SlotsNeeded, uint32& HeapSlot);
 template void FD3D12DescriptorCache::SetSamplers<SF_Amplification>(const FD3D12RootSignature* RootSignature, FD3D12SamplerStateCache& Cache, const SamplerSlotMask& SlotsNeededMask, uint32 SlotsNeeded, uint32& HeapSlot);
-template void FD3D12DescriptorCache::SetSamplers<SF_Hull>(const FD3D12RootSignature* RootSignature, FD3D12SamplerStateCache& Cache, const SamplerSlotMask& SlotsNeededMask, uint32 SlotsNeeded, uint32& HeapSlot);
-template void FD3D12DescriptorCache::SetSamplers<SF_Domain>(const FD3D12RootSignature* RootSignature, FD3D12SamplerStateCache& Cache, const SamplerSlotMask& SlotsNeededMask, uint32 SlotsNeeded, uint32& HeapSlot);
 template void FD3D12DescriptorCache::SetSamplers<SF_Geometry>(const FD3D12RootSignature* RootSignature, FD3D12SamplerStateCache& Cache, const SamplerSlotMask& SlotsNeededMask, uint32 SlotsNeeded, uint32& HeapSlot);
 template void FD3D12DescriptorCache::SetSamplers<SF_Pixel>(const FD3D12RootSignature* RootSignature, FD3D12SamplerStateCache& Cache, const SamplerSlotMask& SlotsNeededMask, uint32 SlotsNeeded, uint32& HeapSlot);
 template void FD3D12DescriptorCache::SetSamplers<SF_Compute>(const FD3D12RootSignature* RootSignature, FD3D12SamplerStateCache& Cache, const SamplerSlotMask& SlotsNeededMask, uint32 SlotsNeeded, uint32& HeapSlot);
diff --git a/Engine/Source/Runtime/D3D12RHI/Private/D3D12PipelineState.h b/Engine/Source/Runtime/D3D12RHI/Private/D3D12PipelineState.h
--- a/Engine/Source/Runtime/D3D12RHI/Private/D3D12PipelineState.h
+++ b/Engine/Source/Runtime/D3D12RHI/Private/D3D12PipelineState.h
@@ -52,8 +52,6 @@ struct FD3D12_GRAPHICS_PIPELINE_STATE_DESC
 	D3D12_SHADER_BYTECODE MS;
 	D3D12_SHADER_BYTECODE AS;
 	D3D12_SHADER_BYTECODE PS;
-	D3D12_SHADER_BYTECODE DS;
-	D3D12_SHADER_BYTECODE HS;
 	D3D12_SHADER_BYTECODE GS;
 #if !D3D12_USE_DERIVED_PSO || D3D12_USE_DERIVED_PSO_SHADER_EXPORTS
 	D3D12_BLEND_DESC BlendState;
@@ -88,8 +86,6 @@ struct FD3D12LowLevelGraphicsPipelineStateDesc
 	ShaderBytecodeHash VSHash;
 	ShaderBytecodeHash MSHash;
 	ShaderBytecodeHash ASHash;
-	ShaderBytecodeHash HSHash;
-	ShaderBytecodeHash DSHash;
 	ShaderBytecodeHash GSHash;
 	ShaderBytecodeHash PSHash;
 	uint32 InputLayoutHash;
@@ -109,8 +105,6 @@ struct FD3D12LowLevelGraphicsPipelineStateDesc
 	const TArray<FShaderCodeVendorExtension>* VSExtensions;
 	const TArray<FShaderCodeVendorExtension>* MSExtensions;
 	const TArray<FShaderCodeVendorExtension>* ASExtensions;
-	const TArray<FShaderCodeVendorExtension>* HSExtensions;
-	const TArray<FShaderCodeVendorExtension>* DSExtensions;
 	const TArray<FShaderCodeVendorExtension>* GSExtensions;
 	const TArray<FShaderCodeVendorExtension>* PSExtensions;
 
@@ -121,9 +115,7 @@ struct FD3D12LowLevelGraphicsPipelineStateDesc
 			MSExtensions != nullptr ||
 			ASExtensions != nullptr ||
 			PSExtensions != nullptr ||
-			GSExtensions != nullptr ||
-			HSExtensions != nullptr ||
-			DSExtensions != nullptr);
+			GSExtensions != nullptr);
 	}
 #else
 	FORCEINLINE bool HasVendorExtensions() const { return false; }
@@ -198,8 +190,6 @@ template <> struct equality_pipeline_state_desc<FD3D12LowLevelGraphicsPipelineSt
 		PSO_IF_NOT_EQUAL_RETURN_FALSE(Desc.MS.BytecodeLength)
 		PSO_IF_NOT_EQUAL_RETURN_FALSE(Desc.AS.BytecodeLength)
 		PSO_IF_NOT_EQUAL_RETURN_FALSE(Desc.GS.BytecodeLength)
-		PSO_IF_NOT_EQUAL_RETURN_FALSE(Desc.DS.BytecodeLength)
-		PSO_IF_NOT_EQUAL_RETURN_FALSE(Desc.HS.BytecodeLength)
 		PSO_IF_NOT_EQUAL_RETURN_FALSE(Desc.InputLayout.NumElements)
 		PSO_IF_NOT_EQUAL_RETURN_FALSE(Desc.RTFormatArray.NumRenderTargets)
 		PSO_IF_NOT_EQUAL_RETURN_FALSE(Desc.DSVFormat)
@@ -232,8 +222,6 @@ template <> struct equality_pipeline_state_desc<FD3D12LowLevelGraphicsPipelineSt
 		PSO_IF_NOT_EQUAL_RETURN_FALSE(ASHash)
 		PSO_IF_NOT_EQUAL_RETURN_FALSE(PSHash)
 		PSO_IF_NOT_EQUAL_RETURN_FALSE(GSHash)
-		PSO_IF_NOT_EQUAL_RETURN_FALSE(HSHash)
-		PSO_IF_NOT_EQUAL_RETURN_FALSE(DSHash)
 
 		if (lhs.Desc.InputLayout.pInputElementDescs != rhs.Desc.InputLayout.pInputElementDescs &&
 			lhs.Desc.InputLayout.NumElements)
@@ -256,8 +244,6 @@ template <> struct equality_pipeline_state_desc<FD3D12LowLevelGraphicsPipelineSt
 		PSO_IF_NOT_EQUAL_RETURN_FALSE(ASExtensions);
 		PSO_IF_NOT_EQUAL_RETURN_FALSE(PSExtensions);
 		PSO_IF_NOT_EQUAL_RETURN_FALSE(GSExtensions);
-		PSO_IF_NOT_EQUAL_RETURN_FALSE(HSExtensions);
-		PSO_IF_NOT_EQUAL_RETURN_FALSE(DSExtensions);
 	#endif
 
 		return true;
@@ -395,13 +381,6 @@ struct FD3D12GraphicsPipelineState : public FRHIGraphicsPipelineState
 	FORCEINLINE class FD3D12MeshShader*				GetMeshShader() const { return nullptr; }
 	FORCEINLINE class FD3D12AmplificationShader*	GetAmplificationShader() const { return nullptr; }
 #endif
-#if PLATFORM_SUPPORTS_TESSELLATION_SHADERS
-	FORCEINLINE class FD3D12HullShader*     GetHullShader() const { return (FD3D12HullShader*)PipelineStateInitializer.BoundShaderState.HullShaderRHI; }
-	FORCEINLINE class FD3D12DomainShader*   GetDomainShader() const { return (FD3D12DomainShader*)PipelineStateInitializer.BoundShaderState.DomainShaderRHI; }
-#else
-	FORCEINLINE class FD3D12HullShader*     GetHullShader() const { return nullptr; }
-	FORCEINLINE class FD3D12DomainShader*   GetDomainShader() const { return nullptr; }
-#endif
 #if PLATFORM_SUPPORTS_GEOMETRY_SHADERS
 	FORCEINLINE class FD3D12GeometryShader* GetGeometryShader() const { return (FD3D12GeometryShader*)PipelineStateInitializer.BoundShaderState.GeometryShaderRHI; }
 #else
diff --git a/Engine/Source/Runtime/D3D12RHI/Private/D3D12Pipelinestate.cpp b/Engine/Source/Runtime/D3D12RHI/Private/D3D12Pipelinestate.cpp
--- a/Engine/Source/Runtime/D3D12RHI/Private/D3D12Pipelinestate.cpp
+++ b/Engine/Source/Runtime/D3D12RHI/Private/D3D12Pipelinestate.cpp
@@ -35,14 +35,7 @@ FD3D12LowLevelGraphicsPipelineStateDesc GetLowLevelGraphicsPipelineStateDesc(con
 	Desc.Desc.DepthStencilState = Initializer.DepthStencilState ? CD3DX12_DEPTH_STENCIL_DESC1(FD3D12DynamicRHI::ResourceCast(Initializer.DepthStencilState)->Desc) : CD3DX12_DEPTH_STENCIL_DESC1(D3D12_DEFAULT);
 #endif // !D3D12_USE_DERIVED_PSO
 
-	if (Initializer.BoundShaderState.HullShaderRHI && Initializer.BoundShaderState.DomainShaderRHI)
-	{
-		Desc.Desc.PrimitiveTopologyType = D3D12_PRIMITIVE_TOPOLOGY_TYPE_PATCH;
-	}
-	else
-	{
-		Desc.Desc.PrimitiveTopologyType = D3D12PrimitiveTypeToTopologyType(TranslatePrimitiveType(Initializer.PrimitiveType));
-	}
+	Desc.Desc.PrimitiveTopologyType = D3D12PrimitiveTypeToTopologyType(TranslatePrimitiveType(Initializer.PrimitiveType));
 
 	TranslateRenderTargetFormats(Initializer, Desc.Desc.RTFormatArray, Desc.Desc.DSVFormat);
 
@@ -66,8 +59,6 @@ FD3D12LowLevelGraphicsPipelineStateDesc GetLowLevelGraphicsPipelineStateDesc(con
 	COPY_SHADER(M, Mesh);
 	COPY_SHADER(A, Amplification);
 	COPY_SHADER(P, Pixel);
-	COPY_SHADER(D, Domain);
-	COPY_SHADER(H, Hull);
 	COPY_SHADER(G, Geometry);
 #undef COPY_SHADER
 
@@ -84,8 +75,6 @@ FD3D12LowLevelGraphicsPipelineStateDesc GetLowLevelGraphicsPipelineStateDesc(con
 	EXT_SHADER(M, Mesh);
 	EXT_SHADER(A, Amplification);
 	EXT_SHADER(P, Pixel);
-	EXT_SHADER(D, Domain);
-	EXT_SHADER(H, Hull);
 	EXT_SHADER(G, Geometry);
 #undef EXT_SHADER
 #endif
@@ -150,8 +139,6 @@ uint64 FD3D12PipelineStateCacheBase::HashPSODesc(const FD3D12LowLevelGraphicsPip
 		ShaderBytecodeHash VSHash;
 		ShaderBytecodeHash MSHash;
 		ShaderBytecodeHash ASHash;
-		ShaderBytecodeHash HSHash;
-		ShaderBytecodeHash DSHash;
 		ShaderBytecodeHash GSHash;
 		ShaderBytecodeHash PSHash;
 		uint32 InputLayoutHash;
@@ -199,8 +186,6 @@ uint64 FD3D12PipelineStateCacheBase::HashPSODesc(const FD3D12LowLevelGraphicsPip
 	PSOData->VSHash          = Desc.VSHash;
 	PSOData->MSHash          = Desc.MSHash;
 	PSOData->ASHash          = Desc.ASHash;
-	PSOData->HSHash          = Desc.HSHash;
-	PSOData->DSHash          = Desc.DSHash;
 	PSOData->GSHash          = Desc.GSHash;
 	PSOData->PSHash          = Desc.PSHash;
 	PSOData->InputLayoutHash = Desc.InputLayoutHash;
@@ -344,8 +329,6 @@ FD3D12GraphicsPipelineState::FD3D12GraphicsPipelineState(
 	bShaderNeedsGlobalConstantBuffer[SF_Mesh] = GetMeshShader() && GetMeshShader()->ResourceCounts.bGlobalUniformBufferUsed;
 	bShaderNeedsGlobalConstantBuffer[SF_Amplification] = GetAmplificationShader() && GetAmplificationShader()->ResourceCounts.bGlobalUniformBufferUsed;
 	bShaderNeedsGlobalConstantBuffer[SF_Pixel] = GetPixelShader() && GetPixelShader()->ResourceCounts.bGlobalUniformBufferUsed;
-	bShaderNeedsGlobalConstantBuffer[SF_Hull] = GetHullShader() && GetHullShader()->ResourceCounts.bGlobalUniformBufferUsed;
-	bShaderNeedsGlobalConstantBuffer[SF_Domain] = GetDomainShader() && GetDomainShader()->ResourceCounts.bGlobalUniformBufferUsed;
 	bShaderNeedsGlobalConstantBuffer[SF_Geometry] = GetGeometryShader() && GetGeometryShader()->ResourceCounts.bGlobalUniformBufferUsed;
 }
 
diff --git a/Engine/Source/Runtime/D3D12RHI/Private/D3D12RHIPrivate.h b/Engine/Source/Runtime/D3D12RHI/Private/D3D12RHIPrivate.h
--- a/Engine/Source/Runtime/D3D12RHI/Private/D3D12RHIPrivate.h
+++ b/Engine/Source/Runtime/D3D12RHI/Private/D3D12RHIPrivate.h
@@ -268,8 +268,6 @@ public:
 	virtual FVertexShaderRHIRef RHICreateVertexShader(TArrayView<const uint8> Code, const FSHAHash& Hash) final override;
 	virtual FMeshShaderRHIRef RHICreateMeshShader(TArrayView<const uint8> Code, const FSHAHash& Hash) final override;
 	virtual FAmplificationShaderRHIRef RHICreateAmplificationShader(TArrayView<const uint8> Code, const FSHAHash& Hash) final override;
-	virtual FHullShaderRHIRef RHICreateHullShader(TArrayView<const uint8> Code, const FSHAHash& Hash) final override;
-	virtual FDomainShaderRHIRef RHICreateDomainShader(TArrayView<const uint8> Code, const FSHAHash& Hash) final override;
 	virtual FGeometryShaderRHIRef RHICreateGeometryShader(TArrayView<const uint8> Code, const FSHAHash& Hash) final override;
 	virtual FComputeShaderRHIRef RHICreateComputeShader(TArrayView<const uint8> Code, const FSHAHash& Hash) override;
 	virtual FGPUFenceRHIRef RHICreateGPUFence(const FName& Name) final override;
@@ -398,16 +396,6 @@ public:
 		return RHICreateGeometryShader(Code, Hash);
 	}
 
-	virtual FHullShaderRHIRef CreateHullShader_RenderThread(class FRHICommandListImmediate& RHICmdList, TArrayView<const uint8> Code, const FSHAHash& Hash) override final
-	{
-		return RHICreateHullShader(Code, Hash);
-	}
-
-	virtual FDomainShaderRHIRef CreateDomainShader_RenderThread(class FRHICommandListImmediate& RHICmdList, TArrayView<const uint8> Code, const FSHAHash& Hash) override final
-	{
-		return RHICreateDomainShader(Code, Hash);
-	}
-
 	virtual FPixelShaderRHIRef CreatePixelShader_RenderThread(class FRHICommandListImmediate& RHICmdList, TArrayView<const uint8> Code, const FSHAHash& Hash) override final
 	{
 		return RHICreatePixelShader(Code, Hash);
diff --git a/Engine/Source/Runtime/D3D12RHI/Private/D3D12RootSignature.cpp b/Engine/Source/Runtime/D3D12RHI/Private/D3D12RootSignature.cpp
--- a/Engine/Source/Runtime/D3D12RHI/Private/D3D12RootSignature.cpp
+++ b/Engine/Source/Runtime/D3D12RHI/Private/D3D12RootSignature.cpp
@@ -69,10 +69,6 @@ FORCEINLINE D3D12_SHADER_VISIBILITY GetD3D12ShaderVisibility(EShaderVisibility V
 	{
 	case SV_Vertex:
 		return D3D12_SHADER_VISIBILITY_VERTEX;
-	case SV_Hull:
-		return D3D12_SHADER_VISIBILITY_HULL;
-	case SV_Domain:
-		return D3D12_SHADER_VISIBILITY_DOMAIN;
 	case SV_Geometry:
 		return D3D12_SHADER_VISIBILITY_GEOMETRY;
 	case SV_Pixel:
@@ -98,10 +94,6 @@ FORCEINLINE D3D12_ROOT_SIGNATURE_FLAGS GetD3D12RootSignatureDenyFlag(EShaderVisi
 	{
 	case SV_Vertex:
 		return D3D12_ROOT_SIGNATURE_FLAG_DENY_VERTEX_SHADER_ROOT_ACCESS;
-	case SV_Hull:
-		return D3D12_ROOT_SIGNATURE_FLAG_DENY_HULL_SHADER_ROOT_ACCESS;
-	case SV_Domain:
-		return D3D12_ROOT_SIGNATURE_FLAG_DENY_DOMAIN_SHADER_ROOT_ACCESS;
 	case SV_Geometry:
 		return D3D12_ROOT_SIGNATURE_FLAG_DENY_GEOMETRY_SHADER_ROOT_ACCESS;
 	case SV_Pixel:
@@ -125,7 +117,7 @@ FORCEINLINE D3D12_ROOT_SIGNATURE_FLAGS GetD3D12RootSignatureDenyFlag(EShaderVisi
 FD3D12RootSignatureDesc::FD3D12RootSignatureDesc(const FD3D12QuantizedBoundShaderState& QBSS, const D3D12_RESOURCE_BINDING_TIER ResourceBindingTier)
 	: RootParametersSize(0)
 {
-	const EShaderVisibility ShaderVisibilityPriorityOrder[] = { SV_Pixel, SV_Vertex, SV_Geometry, SV_Hull, SV_Domain, SV_Mesh, SV_Amplification, SV_All };
+	const EShaderVisibility ShaderVisibilityPriorityOrder[] = { SV_Pixel, SV_Vertex, SV_Geometry, SV_Mesh, SV_Amplification, SV_All };
 	const D3D12_ROOT_PARAMETER_TYPE RootParameterTypePriorityOrder[] = { D3D12_ROOT_PARAMETER_TYPE_DESCRIPTOR_TABLE, D3D12_ROOT_PARAMETER_TYPE_CBV };
 	uint32 RootParameterCount = 0;
 
@@ -362,7 +354,7 @@ const D3D12_VERSIONED_ROOT_SIGNATURE_DESC& FD3D12RootSignatureDesc::GetStaticGra
 {
 	// TODO: Support vendor extensions for static root signatures?
 
-	static const uint32 DescriptorTableCount = 16;
+	static const uint32 DescriptorTableCount = 10;
 	static struct
 	{
 		D3D12_SHADER_VISIBILITY Vis;
@@ -384,14 +376,6 @@ const D3D12_VERSIONED_ROOT_SIGNATURE_DESC& FD3D12RootSignatureDesc::GetStaticGra
 		{ D3D12_SHADER_VISIBILITY_GEOMETRY, D3D12_DESCRIPTOR_RANGE_TYPE_CBV, MAX_CBS, 0, D3D12ShaderUtils::StaticRootSignatureConstants::CBVDescriptorRangeFlags },
 		{ D3D12_SHADER_VISIBILITY_GEOMETRY, D3D12_DESCRIPTOR_RANGE_TYPE_SAMPLER, MAX_SAMPLERS, 0, D3D12ShaderUtils::StaticRootSignatureConstants::SamplerDescriptorRangeFlags },
 
-		{ D3D12_SHADER_VISIBILITY_HULL, D3D12_DESCRIPTOR_RANGE_TYPE_SRV, MAX_SRVS, 0, D3D12ShaderUtils::StaticRootSignatureConstants::SRVDescriptorRangeFlags },
-		{ D3D12_SHADER_VISIBILITY_HULL, D3D12_DESCRIPTOR_RANGE_TYPE_CBV, MAX_CBS, 0, D3D12ShaderUtils::StaticRootSignatureConstants::CBVDescriptorRangeFlags },
-		{ D3D12_SHADER_VISIBILITY_HULL, D3D12_DESCRIPTOR_RANGE_TYPE_SAMPLER, MAX_SAMPLERS, 0, D3D12ShaderUtils::StaticRootSignatureConstants::SamplerDescriptorRangeFlags },
-
-		{ D3D12_SHADER_VISIBILITY_DOMAIN, D3D12_DESCRIPTOR_RANGE_TYPE_SRV, MAX_SRVS, 0, D3D12ShaderUtils::StaticRootSignatureConstants::SRVDescriptorRangeFlags },
-		{ D3D12_SHADER_VISIBILITY_DOMAIN, D3D12_DESCRIPTOR_RANGE_TYPE_CBV, MAX_CBS, 0, D3D12ShaderUtils::StaticRootSignatureConstants::CBVDescriptorRangeFlags },
-		{ D3D12_SHADER_VISIBILITY_DOMAIN, D3D12_DESCRIPTOR_RANGE_TYPE_SAMPLER, MAX_SAMPLERS, 0, D3D12ShaderUtils::StaticRootSignatureConstants::SamplerDescriptorRangeFlags },
-
 		{ D3D12_SHADER_VISIBILITY_ALL, D3D12_DESCRIPTOR_RANGE_TYPE_UAV, MAX_UAVS, 0, D3D12ShaderUtils::StaticRootSignatureConstants::UAVDescriptorRangeFlags },
 	};
 
@@ -615,8 +599,6 @@ void FD3D12RootSignature::InternalAnalyzeSignature(const RootSignatureDescType&
 	}
 
 	const bool bDenyVS = (Desc.Flags & D3D12_ROOT_SIGNATURE_FLAG_DENY_VERTEX_SHADER_ROOT_ACCESS) != 0;
-	const bool bDenyHS = (Desc.Flags & D3D12_ROOT_SIGNATURE_FLAG_DENY_HULL_SHADER_ROOT_ACCESS) != 0;
-	const bool bDenyDS = (Desc.Flags & D3D12_ROOT_SIGNATURE_FLAG_DENY_DOMAIN_SHADER_ROOT_ACCESS) != 0;
 	const bool bDenyGS = (Desc.Flags & D3D12_ROOT_SIGNATURE_FLAG_DENY_GEOMETRY_SHADER_ROOT_ACCESS) != 0;
 	const bool bDenyPS = (Desc.Flags & D3D12_ROOT_SIGNATURE_FLAG_DENY_PIXEL_SHADER_ROOT_ACCESS) != 0;
 #if PLATFORM_SUPPORTS_MESH_SHADERS
@@ -678,12 +660,6 @@ void FD3D12RootSignature::InternalAnalyzeSignature(const RootSignatureDescType&
 		case D3D12_SHADER_VISIBILITY_VERTEX:
 			CurrentVisibleSF = SF_Vertex;
 			break;
-		case D3D12_SHADER_VISIBILITY_HULL:
-			CurrentVisibleSF = SF_Hull;
-			break;
-		case D3D12_SHADER_VISIBILITY_DOMAIN:
-			CurrentVisibleSF = SF_Domain;
-			break;
 		case D3D12_SHADER_VISIBILITY_GEOMETRY:
 			CurrentVisibleSF = SF_Geometry;
 			break;
@@ -708,8 +684,6 @@ void FD3D12RootSignature::InternalAnalyzeSignature(const RootSignatureDescType&
 		// Determine shader stage visibility.
 		{
 			Stage[SF_Vertex].bVisible = Stage[SF_Vertex].bVisible || (!bDenyVS && HasVisibility(CurrentParameter.ShaderVisibility, D3D12_SHADER_VISIBILITY_VERTEX));
-			Stage[SF_Hull].bVisible = Stage[SF_Hull].bVisible || (!bDenyHS && HasVisibility(CurrentParameter.ShaderVisibility, D3D12_SHADER_VISIBILITY_HULL));
-			Stage[SF_Domain].bVisible = Stage[SF_Domain].bVisible || (!bDenyDS && HasVisibility(CurrentParameter.ShaderVisibility, D3D12_SHADER_VISIBILITY_DOMAIN));
 			Stage[SF_Geometry].bVisible = Stage[SF_Geometry].bVisible || (!bDenyGS && HasVisibility(CurrentParameter.ShaderVisibility, D3D12_SHADER_VISIBILITY_GEOMETRY));
 			Stage[SF_Pixel].bVisible = Stage[SF_Pixel].bVisible || (!bDenyPS && HasVisibility(CurrentParameter.ShaderVisibility, D3D12_SHADER_VISIBILITY_PIXEL));
 
diff --git a/Engine/Source/Runtime/D3D12RHI/Private/D3D12RootSignature.h b/Engine/Source/Runtime/D3D12RHI/Private/D3D12RootSignature.h
--- a/Engine/Source/Runtime/D3D12RHI/Private/D3D12RootSignature.h
+++ b/Engine/Source/Runtime/D3D12RHI/Private/D3D12RootSignature.h
@@ -19,14 +19,6 @@ enum ERootParameterKeys
 	GS_CBVs,
 	GS_RootCBVs,
 	GS_Samplers,
-	HS_SRVs,
-	HS_CBVs,
-	HS_RootCBVs,
-	HS_Samplers,
-	DS_SRVs,
-	DS_CBVs,
-	DS_RootCBVs,
-	DS_Samplers,
 	MS_SRVs,
 	MS_CBVs,
 	MS_RootCBVs,
@@ -129,8 +121,6 @@ public:
 		case SF_Amplification: return BindSlotMap[AS_Samplers];
 		case SF_Pixel: return BindSlotMap[PS_Samplers];
 		case SF_Geometry: return BindSlotMap[GS_Samplers];
-		case SF_Hull: return BindSlotMap[HS_Samplers];
-		case SF_Domain: return BindSlotMap[DS_Samplers];
 		case SF_Compute: return BindSlotMap[ALL_Samplers];
 
 		default: check(false);
@@ -147,8 +137,6 @@ public:
 		case SF_Amplification: return BindSlotMap[AS_SRVs];
 		case SF_Pixel: return BindSlotMap[PS_SRVs];
 		case SF_Geometry: return BindSlotMap[GS_SRVs];
-		case SF_Hull: return BindSlotMap[HS_SRVs];
-		case SF_Domain: return BindSlotMap[DS_SRVs];
 		case SF_Compute: return BindSlotMap[ALL_SRVs];
 
 		default: check(false);
@@ -165,8 +153,6 @@ public:
 		case SF_Amplification: return BindSlotMap[AS_CBVs];
 		case SF_Pixel: return BindSlotMap[PS_CBVs];
 		case SF_Geometry: return BindSlotMap[GS_CBVs];
-		case SF_Hull: return BindSlotMap[HS_CBVs];
-		case SF_Domain: return BindSlotMap[DS_CBVs];
 		case SF_Compute: return BindSlotMap[ALL_CBVs];
 
 		default: check(false);
@@ -183,8 +169,6 @@ public:
 		case SF_Amplification: return BindSlotMap[AS_RootCBVs];
 		case SF_Pixel: return BindSlotMap[PS_RootCBVs];
 		case SF_Geometry: return BindSlotMap[GS_RootCBVs];
-		case SF_Hull: return BindSlotMap[HS_RootCBVs];
-		case SF_Domain: return BindSlotMap[DS_RootCBVs];
 
 		case SF_NumFrequencies:
 		case SF_Compute: return BindSlotMap[ALL_RootCBVs];
@@ -213,8 +197,6 @@ public:
 	inline bool HasVS() const { return Stage[SF_Vertex].bVisible; }
 	inline bool HasMS() const { return Stage[SF_Mesh].bVisible; }
 	inline bool HasAS() const { return Stage[SF_Amplification].bVisible; }
-	inline bool HasHS() const { return Stage[SF_Hull].bVisible; }
-	inline bool HasDS() const { return Stage[SF_Domain].bVisible; }
 	inline bool HasGS() const { return Stage[SF_Geometry].bVisible; }
 	inline bool HasPS() const { return Stage[SF_Pixel].bVisible; }
 	inline bool HasCS() const { return Stage[SF_Compute].bVisible; }	// Root signatures can be used for Graphics and/or Compute because they exist in separate bind spaces.
@@ -248,8 +230,6 @@ private:
 		case SF_Amplification: pBindSlot = &BindSlotMap[AS_Samplers]; break;
 		case SF_Pixel: pBindSlot = &BindSlotMap[PS_Samplers]; break;
 		case SF_Geometry: pBindSlot = &BindSlotMap[GS_Samplers]; break;
-		case SF_Hull: pBindSlot = &BindSlotMap[HS_Samplers]; break;
-		case SF_Domain: pBindSlot = &BindSlotMap[DS_Samplers]; break;
 
 		case SF_Compute:
 		case SF_NumFrequencies: pBindSlot = &BindSlotMap[ALL_Samplers]; break;
@@ -274,8 +254,6 @@ private:
 		case SF_Amplification: pBindSlot = &BindSlotMap[AS_SRVs]; break;
 		case SF_Pixel: pBindSlot = &BindSlotMap[PS_SRVs]; break;
 		case SF_Geometry: pBindSlot = &BindSlotMap[GS_SRVs]; break;
-		case SF_Hull: pBindSlot = &BindSlotMap[HS_SRVs]; break;
-		case SF_Domain: pBindSlot = &BindSlotMap[DS_SRVs]; break;
 
 		case SF_Compute:
 		case SF_NumFrequencies: pBindSlot = &BindSlotMap[ALL_SRVs]; break;
@@ -300,8 +278,6 @@ private:
 		case SF_Amplification: pBindSlot = &BindSlotMap[AS_CBVs]; break;
 		case SF_Pixel: pBindSlot = &BindSlotMap[PS_CBVs]; break;
 		case SF_Geometry: pBindSlot = &BindSlotMap[GS_CBVs]; break;
-		case SF_Hull: pBindSlot = &BindSlotMap[HS_CBVs]; break;
-		case SF_Domain: pBindSlot = &BindSlotMap[DS_CBVs]; break;
 
 		case SF_Compute:
 		case SF_NumFrequencies: pBindSlot = &BindSlotMap[ALL_CBVs]; break;
@@ -327,8 +303,6 @@ private:
 		case SF_Amplification: pBindSlot = &BindSlotMap[AS_RootCBVs]; break;
 		case SF_Pixel: pBindSlot = &BindSlotMap[PS_RootCBVs]; break;
 		case SF_Geometry: pBindSlot = &BindSlotMap[GS_RootCBVs]; break;
-		case SF_Hull: pBindSlot = &BindSlotMap[HS_RootCBVs]; break;
-		case SF_Domain: pBindSlot = &BindSlotMap[DS_RootCBVs]; break;
 
 		case SF_Compute:
 		case SF_NumFrequencies: pBindSlot = &BindSlotMap[ALL_RootCBVs]; break;
diff --git a/Engine/Source/Runtime/D3D12RHI/Private/D3D12Shader.h b/Engine/Source/Runtime/D3D12RHI/Private/D3D12Shader.h
--- a/Engine/Source/Runtime/D3D12RHI/Private/D3D12Shader.h
+++ b/Engine/Source/Runtime/D3D12RHI/Private/D3D12Shader.h
@@ -112,38 +112,6 @@ public:
 	FShaderCodePackedResourceCounts ResourceCounts;
 };
 
-class FD3D12HullShader : public FRHIHullShader, public FD3D12ShaderData
-{
-public:
-	enum { StaticFrequency = SF_Hull };
-
-	/** The shader's bytecode. */
-	FD3D12ShaderBytecode ShaderBytecode;
-
-	FD3D12ShaderResourceTable ShaderResourceTable;
-
-	/** The shader's bytecode, with custom data in the last byte. */
-	TArray<uint8> Code;
-
-	FShaderCodePackedResourceCounts ResourceCounts;
-};
-
-class FD3D12DomainShader : public FRHIDomainShader, public FD3D12ShaderData
-{
-public:
-	enum { StaticFrequency = SF_Domain };
-
-	/** The shader's bytecode. */
-	FD3D12ShaderBytecode ShaderBytecode;
-
-	FD3D12ShaderResourceTable ShaderResourceTable;
-
-	/** The shader's bytecode, with custom data in the last byte. */
-	TArray<uint8> Code;
-
-	FShaderCodePackedResourceCounts ResourceCounts;
-};
-
 class FD3D12PixelShader : public FRHIPixelShader, public FD3D12ShaderData
 {
 public:
@@ -222,8 +190,6 @@ public:
 	FORCEINLINE FD3D12MeshShader*          GetMeshShader()          const { return (FD3D12MeshShader*)          CacheLink.GetMeshShader();          }
 	FORCEINLINE FD3D12AmplificationShader* GetAmplificationShader() const { return (FD3D12AmplificationShader*) CacheLink.GetAmplificationShader(); }
 	FORCEINLINE FD3D12PixelShader*         GetPixelShader()         const { return (FD3D12PixelShader*)         CacheLink.GetPixelShader();         }
-	FORCEINLINE FD3D12HullShader*          GetHullShader()          const { return (FD3D12HullShader*)          CacheLink.GetHullShader();          }
-	FORCEINLINE FD3D12DomainShader*        GetDomainShader()        const { return (FD3D12DomainShader*)        CacheLink.GetDomainShader();        }
 	FORCEINLINE FD3D12GeometryShader*      GetGeometryShader()      const { return (FD3D12GeometryShader*)      CacheLink.GetGeometryShader();      }
 };
 
@@ -276,16 +242,6 @@ struct TD3D12ResourceTraits<FRHIGeometryShader>
 	typedef FD3D12GeometryShader TConcreteType;
 };
 template<>
-struct TD3D12ResourceTraits<FRHIHullShader>
-{
-	typedef FD3D12HullShader TConcreteType;
-};
-template<>
-struct TD3D12ResourceTraits<FRHIDomainShader>
-{
-	typedef FD3D12DomainShader TConcreteType;
-};
-template<>
 struct TD3D12ResourceTraits<FRHIPixelShader>
 {
 	typedef FD3D12PixelShader TConcreteType;
diff --git a/Engine/Source/Runtime/D3D12RHI/Private/D3D12Shaders.cpp b/Engine/Source/Runtime/D3D12RHI/Private/D3D12Shaders.cpp
--- a/Engine/Source/Runtime/D3D12RHI/Private/D3D12Shaders.cpp
+++ b/Engine/Source/Runtime/D3D12RHI/Private/D3D12Shaders.cpp
@@ -215,72 +215,6 @@ FPixelShaderRHIRef FD3D12DynamicRHI::RHICreatePixelShader(TArrayView<const uint8
 	return Shader;
 }
 
-FHullShaderRHIRef FD3D12DynamicRHI::RHICreateHullShader(TArrayView<const uint8> Code, const FSHAHash& Hash)
-{
-	FShaderCodeReader ShaderCode(Code);
-
-	FD3D12HullShader* Shader = new FD3D12HullShader;
-
-	FMemoryReaderView Ar(Code, true);
-	Ar << Shader->ShaderResourceTable;
-	int32 Offset = Ar.Tell();
-	const uint8* CodePtr = Code.GetData() + Offset;
-	const SIZE_T CodeSize = ShaderCode.GetActualShaderCodeSize() - Offset;
-	bool bFoundCodeFeatures;
-	FShaderCodeFeatures CodeFeatures;
-	if (!ReadShaderOptionalData(ShaderCode, *Shader, bFoundCodeFeatures, CodeFeatures))
-	{
-		return nullptr;
-	}
-	if (bFoundCodeFeatures && CodeFeatures.bUsesWaveOps && !GRHISupportsWaveOperations)
-	{
-		return nullptr;
-	}
-
-	Shader->Code = Code;
-	InitUniformBufferStaticSlots(Shader);
-
-	D3D12_SHADER_BYTECODE ShaderBytecode;
-	ShaderBytecode.pShaderBytecode = Shader->Code.GetData() + Offset;
-	ShaderBytecode.BytecodeLength = CodeSize;
-	Shader->ShaderBytecode.SetShaderBytecode(ShaderBytecode);
-
-	return Shader;
-}
-
-FDomainShaderRHIRef FD3D12DynamicRHI::RHICreateDomainShader(TArrayView<const uint8> Code, const FSHAHash& Hash)
-{
-	FShaderCodeReader ShaderCode(Code);
-
-	FD3D12DomainShader* Shader = new FD3D12DomainShader;
-
-	FMemoryReaderView Ar(Code, true);
-	Ar << Shader->ShaderResourceTable;
-	int32 Offset = Ar.Tell();
-	const uint8* CodePtr = Code.GetData() + Offset;
-	const SIZE_T CodeSize = ShaderCode.GetActualShaderCodeSize() - Offset;
-	bool bFoundCodeFeatures;
-	FShaderCodeFeatures CodeFeatures;
-	if (!ReadShaderOptionalData(ShaderCode, *Shader, bFoundCodeFeatures, CodeFeatures))
-	{
-		return nullptr;
-	}
-	if (bFoundCodeFeatures && CodeFeatures.bUsesWaveOps && !GRHISupportsWaveOperations)
-	{
-		return nullptr;
-	}
-
-	Shader->Code = Code;
-	InitUniformBufferStaticSlots(Shader);
-
-	D3D12_SHADER_BYTECODE ShaderBytecode;
-	ShaderBytecode.pShaderBytecode = Shader->Code.GetData() + Offset;
-	ShaderBytecode.BytecodeLength = CodeSize;
-	Shader->ShaderBytecode.SetShaderBytecode(ShaderBytecode);
-
-	return Shader;
-}
-
 FGeometryShaderRHIRef FD3D12DynamicRHI::RHICreateGeometryShader(TArrayView<const uint8> Code, const FSHAHash& Hash)
 {
 	FShaderCodeReader ShaderCode(Code);
@@ -562,8 +496,6 @@ FBoundShaderStateRHIRef FD3D12DynamicRHI::DX12CreateBoundShaderState(const FBoun
 * @param VertexDeclaration - existing vertex decl
 * @param StreamStrides - optional stream strides
 * @param VertexShader - existing vertex shader
-* @param HullShader - existing hull shader
-* @param DomainShader - existing domain shader
 * @param PixelShader - existing pixel shader
 * @param GeometryShader - existing geometry shader
 */
diff --git a/Engine/Source/Runtime/D3D12RHI/Private/D3D12StateCache.cpp b/Engine/Source/Runtime/D3D12RHI/Private/D3D12StateCache.cpp
--- a/Engine/Source/Runtime/D3D12RHI/Private/D3D12StateCache.cpp
+++ b/Engine/Source/Runtime/D3D12RHI/Private/D3D12StateCache.cpp
@@ -47,8 +47,6 @@ inline bool operator!=(D3D12_CPU_DESCRIPTOR_HANDLE lhs, D3D12_CPU_DESCRIPTOR_HAN
 template void FD3D12StateCacheBase::SetShaderResourceView<SF_Vertex>(FD3D12ShaderResourceView* SRV, uint32 ResourceIndex);
 template void FD3D12StateCacheBase::SetShaderResourceView<SF_Mesh>(FD3D12ShaderResourceView* SRV, uint32 ResourceIndex);
 template void FD3D12StateCacheBase::SetShaderResourceView<SF_Amplification>(FD3D12ShaderResourceView* SRV, uint32 ResourceIndex);
-template void FD3D12StateCacheBase::SetShaderResourceView<SF_Hull>(FD3D12ShaderResourceView* SRV, uint32 ResourceIndex);
-template void FD3D12StateCacheBase::SetShaderResourceView<SF_Domain>(FD3D12ShaderResourceView* SRV, uint32 ResourceIndex);
 template void FD3D12StateCacheBase::SetShaderResourceView<SF_Geometry>(FD3D12ShaderResourceView* SRV, uint32 ResourceIndex);
 template void FD3D12StateCacheBase::SetShaderResourceView<SF_Pixel>(FD3D12ShaderResourceView* SRV, uint32 ResourceIndex);
 template void FD3D12StateCacheBase::SetShaderResourceView<SF_Compute>(FD3D12ShaderResourceView* SRV, uint32 ResourceIndex);
@@ -355,7 +353,7 @@ void FD3D12StateCacheBase::SetScissorRects(uint32 Count, const D3D12_RECT* const
 
 inline bool ShouldSkipStage(uint32 Stage)
 {
-	return ((Stage == SF_Mesh || Stage == SF_Amplification) && !GRHISupportsMeshShaders);
+	return ((Stage == SF_Mesh || Stage == SF_Amplification) && !GRHISupportsMeshShaders) || Stage == SF_Hull || Stage == SF_Domain;
 }
 
 template <ED3D12PipelineType PipelineType>
@@ -612,8 +610,6 @@ void FD3D12StateCacheBase::ApplyState()
 			CONDITIONAL_SET_SRVS(SF_Vertex);
 			CONDITIONAL_SET_SRVS(SF_Mesh);
 			CONDITIONAL_SET_SRVS(SF_Amplification);
-			CONDITIONAL_SET_SRVS(SF_Hull);
-			CONDITIONAL_SET_SRVS(SF_Domain);
 			CONDITIONAL_SET_SRVS(SF_Geometry);
 			CONDITIONAL_SET_SRVS(SF_Pixel);
 		}
@@ -649,8 +645,6 @@ void FD3D12StateCacheBase::ApplyState()
 			CONDITIONAL_SET_CBVS(SF_Vertex);
 			CONDITIONAL_SET_CBVS(SF_Mesh);
 			CONDITIONAL_SET_CBVS(SF_Amplification);
-			CONDITIONAL_SET_CBVS(SF_Hull);
-			CONDITIONAL_SET_CBVS(SF_Domain);
 			CONDITIONAL_SET_CBVS(SF_Geometry);
 			CONDITIONAL_SET_CBVS(SF_Pixel);
 		}
@@ -821,8 +815,6 @@ void FD3D12StateCacheBase::ApplySamplers(const FD3D12RootSignature* const pRootS
 		CONDITIONAL_SET_SAMPLERS(SF_Vertex);
 		CONDITIONAL_SET_SAMPLERS(SF_Mesh);
 		CONDITIONAL_SET_SAMPLERS(SF_Amplification);
-		CONDITIONAL_SET_SAMPLERS(SF_Hull);
-		CONDITIONAL_SET_SAMPLERS(SF_Domain);
 		CONDITIONAL_SET_SAMPLERS(SF_Geometry);
 		CONDITIONAL_SET_SAMPLERS(SF_Pixel);
 	}
diff --git a/Engine/Source/Runtime/D3D12RHI/Private/D3D12StateCachePrivate.h b/Engine/Source/Runtime/D3D12RHI/Private/D3D12StateCachePrivate.h
--- a/Engine/Source/Runtime/D3D12RHI/Private/D3D12StateCachePrivate.h
+++ b/Engine/Source/Runtime/D3D12RHI/Private/D3D12StateCachePrivate.h
@@ -137,8 +137,6 @@ struct FD3D12ResourceCache
 		Dirty(SF_Vertex, SlotMask);
 		Dirty(SF_Mesh, SlotMask);
 		Dirty(SF_Amplification, SlotMask);
-		Dirty(SF_Hull, SlotMask);
-		Dirty(SF_Domain, SlotMask);
 		Dirty(SF_Pixel, SlotMask);
 		Dirty(SF_Geometry, SlotMask);
 	}
@@ -262,7 +260,7 @@ struct FD3D12SamplerStateCache : public FD3D12ResourceCache<SamplerSlotMask>
 };
 
 
-static inline D3D_PRIMITIVE_TOPOLOGY GetD3D12PrimitiveType(uint32 PrimitiveType, bool bUsingTessellation)
+static inline D3D_PRIMITIVE_TOPOLOGY GetD3D12PrimitiveType(uint32 PrimitiveType)
 {
 	static const uint8 D3D12PrimitiveType[] =
 	{
@@ -311,19 +309,6 @@ static inline D3D_PRIMITIVE_TOPOLOGY GetD3D12PrimitiveType(uint32 PrimitiveType,
 	};
 	static_assert(UE_ARRAY_COUNT(D3D12PrimitiveType) == PT_Num, "Primitive lookup table is wrong size");
 
-	if (bUsingTessellation)
-	{
-		if (PrimitiveType == PT_TriangleList)
-		{
-			// This is the case for tessellation without AEN or other buffers, so just flip to 3 CPs
-			return D3D_PRIMITIVE_TOPOLOGY_3_CONTROL_POINT_PATCHLIST;
-		}
-		else/* if (PrimitiveType < PT_1_ControlPointPatchList)*/
-		{
-			checkf(PrimitiveType >= PT_1_ControlPointPatchList, TEXT("Invalid type specified for tessellated render, probably missing a case in FSkeletalMeshSceneProxy::DrawDynamicElementsByMaterial or FStaticMeshSceneProxy::GetMeshElement"));
-		}
-	}
-
 	D3D_PRIMITIVE_TOPOLOGY D3DType = (D3D_PRIMITIVE_TOPOLOGY) D3D12PrimitiveType[PrimitiveType];
 	checkf(D3DType, TEXT("Unknown primitive type: %u"), PrimitiveType);
 	return D3DType;
@@ -460,10 +445,6 @@ protected:
 	DECLARE_SHADER_TRAITS(Amplification);
 #endif
 	DECLARE_SHADER_TRAITS(Pixel);
-#if PLATFORM_SUPPORTS_TESSELLATION_SHADERS
-	DECLARE_SHADER_TRAITS(Domain);
-	DECLARE_SHADER_TRAITS(Hull);
-#endif
 #if PLATFORM_SUPPORTS_GEOMETRY_SHADERS
 	DECLARE_SHADER_TRAITS(Geometry);
 #endif
@@ -766,24 +747,6 @@ public:
 #endif
 	}
 
-	D3D12_STATE_CACHE_INLINE void GetHullShader(FD3D12HullShader** Shader)
-	{
-#if PLATFORM_SUPPORTS_TESSELLATION_SHADERS
-		GetShader(Shader);
-#else
-		*Shader = nullptr;
-#endif
-	}
-
-	D3D12_STATE_CACHE_INLINE void GetDomainShader(FD3D12DomainShader** Shader)
-	{
-#if PLATFORM_SUPPORTS_TESSELLATION_SHADERS
-		GetShader(Shader);
-#else
-		*Shader = nullptr;
-#endif
-	}
-
 	D3D12_STATE_CACHE_INLINE void GetGeometryShader(FD3D12GeometryShader** Shader)
 	{
 #if PLATFORM_SUPPORTS_GEOMETRY_SHADERS
@@ -798,7 +761,7 @@ public:
 		GetShader(Shader);
 	}
 
-	D3D12_STATE_CACHE_INLINE void SetGraphicsPipelineState(FD3D12GraphicsPipelineState* GraphicsPipelineState, bool bTessellationChanged)
+	D3D12_STATE_CACHE_INLINE void SetGraphicsPipelineState(FD3D12GraphicsPipelineState* GraphicsPipelineState)
 	{
 		check(GraphicsPipelineState);
 		if (PipelineState.Graphics.CurrentPipelineStateObject != GraphicsPipelineState)
@@ -810,10 +773,6 @@ public:
 			SetShader(GraphicsPipelineState->GetAmplificationShader());
 #endif
 			SetShader(GraphicsPipelineState->GetPixelShader());
-#if PLATFORM_SUPPORTS_TESSELLATION_SHADERS
-			SetShader(GraphicsPipelineState->GetDomainShader());
-			SetShader(GraphicsPipelineState->GetHullShader());
-#endif
 #if PLATFORM_SUPPORTS_GEOMETRY_SHADERS
 			SetShader(GraphicsPipelineState->GetGeometryShader());
 #endif
@@ -828,11 +787,10 @@ public:
 			PipelineState.Graphics.CurrentPipelineStateObject = GraphicsPipelineState;
 
 			EPrimitiveType PrimitiveType = GraphicsPipelineState->PipelineStateInitializer.PrimitiveType;
-			if (PipelineState.Graphics.CurrentPrimitiveType != PrimitiveType || bTessellationChanged)
+			if (PipelineState.Graphics.CurrentPrimitiveType != PrimitiveType)
 			{
-				const bool bUsingTessellation = GraphicsPipelineState->GetHullShader() && GraphicsPipelineState->GetDomainShader();
 				PipelineState.Graphics.CurrentPrimitiveType = PrimitiveType;
-				PipelineState.Graphics.CurrentPrimitiveTopology = GetD3D12PrimitiveType(PrimitiveType, bUsingTessellation);
+				PipelineState.Graphics.CurrentPrimitiveTopology = GetD3D12PrimitiveType(PrimitiveType);
 				bNeedSetPrimitiveTopology = true;
 
 				static_assert(PT_Num == 38, "This computation needs to be updated, matching that of GetVertexCountForPrimitiveCount()");
diff --git a/Engine/Source/Runtime/D3D12RHI/Private/D3D12Util.cpp b/Engine/Source/Runtime/D3D12RHI/Private/D3D12Util.cpp
--- a/Engine/Source/Runtime/D3D12RHI/Private/D3D12Util.cpp
+++ b/Engine/Source/Runtime/D3D12RHI/Private/D3D12Util.cpp
@@ -834,18 +834,6 @@ void QuantizeBoundShaderState(
 		QBSS.bUseVendorExtension |= UsesVendorExtensionSpace(*PixelShader);
 	}
 
-	if (const FD3D12HullShader* const HullShader = BSS->GetHullShader())
-	{
-		FD3D12QuantizedBoundShaderState::InitShaderRegisterCounts(ResourceBindingTier, HullShader->ResourceCounts, QBSS.RegisterCounts[SV_Hull]);
-		QBSS.bUseVendorExtension |= UsesVendorExtensionSpace(*HullShader); // -V522
-	}
-
-	if (const FD3D12DomainShader* const DomainShader = BSS->GetDomainShader())
-	{
-		FD3D12QuantizedBoundShaderState::InitShaderRegisterCounts(ResourceBindingTier, DomainShader->ResourceCounts, QBSS.RegisterCounts[SV_Domain]);
-		QBSS.bUseVendorExtension |= UsesVendorExtensionSpace(*DomainShader); // -V522
-	}
-
 	if (const FD3D12GeometryShader* const GeometryShader = BSS->GetGeometryShader())
 	{
 		FD3D12QuantizedBoundShaderState::InitShaderRegisterCounts(ResourceBindingTier, GeometryShader->ResourceCounts, QBSS.RegisterCounts[SV_Geometry]);
diff --git a/Engine/Source/Runtime/D3D12RHI/Private/Windows/WindowsD3D12DiskCache.h b/Engine/Source/Runtime/D3D12RHI/Private/Windows/WindowsD3D12DiskCache.h
--- a/Engine/Source/Runtime/D3D12RHI/Private/Windows/WindowsD3D12DiskCache.h
+++ b/Engine/Source/Runtime/D3D12RHI/Private/Windows/WindowsD3D12DiskCache.h
@@ -10,7 +10,7 @@ class FDiskCacheInterface
 {
 	// Increment this if changes are made to the
 	// disk caches so stale caches get updated correctly
-	static const uint32 mCurrentHeaderVersion = 10;
+	static const uint32 mCurrentHeaderVersion = 11;
 	struct FDiskCacheHeader
 	{
 		uint32 mHeaderVersion;
diff --git a/Engine/Source/Runtime/D3D12RHI/Private/Windows/WindowsD3D12PipelineState.cpp b/Engine/Source/Runtime/D3D12RHI/Private/Windows/WindowsD3D12PipelineState.cpp
--- a/Engine/Source/Runtime/D3D12RHI/Private/Windows/WindowsD3D12PipelineState.cpp
+++ b/Engine/Source/Runtime/D3D12RHI/Private/Windows/WindowsD3D12PipelineState.cpp
@@ -63,8 +63,6 @@ FD3D12_GRAPHICS_PIPELINE_STATE_STREAM FD3D12_GRAPHICS_PIPELINE_STATE_DESC::Pipel
 	Stream.PrimitiveTopologyType = this->PrimitiveTopologyType;
 	Stream.VS = this->VS;
 	Stream.GS = this->GS;
-	Stream.HS = this->HS;
-	Stream.DS = this->DS;
 	Stream.PS = this->PS;
 	Stream.BlendState = CD3DX12_BLEND_DESC(this->BlendState);
 	Stream.DepthStencilState = CD3DX12_DEPTH_STENCIL_DESC1(this->DepthStencilState);
@@ -109,8 +107,6 @@ D3D12_GRAPHICS_PIPELINE_STATE_DESC FD3D12_GRAPHICS_PIPELINE_STATE_DESC::Graphics
 	D.PrimitiveTopologyType = this->PrimitiveTopologyType;
 	D.VS = this->VS;
 	D.GS = this->GS;
-	D.HS = this->HS;
-	D.DS = this->DS;
 	D.PS = this->PS;
 	D.BlendState = this->BlendState;
 	D.DepthStencilState = CD3DX12_DEPTH_STENCIL_DESC1(this->DepthStencilState);
@@ -169,8 +165,6 @@ void GraphicsPipelineCreationArgs_POD::Destroy()
 	FMemory::Free((void*)Desc.Desc.AS.pShaderBytecode);
 	FMemory::Free((void*)Desc.Desc.PS.pShaderBytecode);
 	FMemory::Free((void*)Desc.Desc.GS.pShaderBytecode);
-	FMemory::Free((void*)Desc.Desc.HS.pShaderBytecode);
-	FMemory::Free((void*)Desc.Desc.DS.pShaderBytecode);
 }
 
 void FD3D12PipelineStateCache::OnPSOCreated(FD3D12PipelineState* PipelineState, const FD3D12LowLevelGraphicsPipelineStateDesc& Desc)
@@ -294,14 +288,6 @@ void FD3D12PipelineStateCache::RebuildFromDiskCache(ID3D12RootSignature* Graphic
 		{
 			DiskCaches[PSO_CACHE_GRAPHICS].SetPointerAndAdvanceFilePosition((void**)&PSODesc->PS.pShaderBytecode, PSODesc->PS.BytecodeLength, bBackShadersWithSystemMemory);
 		}
-		if (PSODesc->DS.BytecodeLength)
-		{
-			DiskCaches[PSO_CACHE_GRAPHICS].SetPointerAndAdvanceFilePosition((void**)&PSODesc->DS.pShaderBytecode, PSODesc->DS.BytecodeLength, bBackShadersWithSystemMemory);
-		}
-		if (PSODesc->HS.BytecodeLength)
-		{
-			DiskCaches[PSO_CACHE_GRAPHICS].SetPointerAndAdvanceFilePosition((void**)&PSODesc->HS.pShaderBytecode, PSODesc->HS.BytecodeLength, bBackShadersWithSystemMemory);
-		}
 		if (PSODesc->GS.BytecodeLength)
 		{
 			DiskCaches[PSO_CACHE_GRAPHICS].SetPointerAndAdvanceFilePosition((void**)&PSODesc->GS.pShaderBytecode, PSODesc->GS.BytecodeLength, bBackShadersWithSystemMemory);
@@ -442,14 +428,6 @@ void FD3D12PipelineStateCache::AddToDiskCache(const FD3D12LowLevelGraphicsPipeli
 		{
 			DiskCache.AppendData((void*)PsoDesc.PS.pShaderBytecode, PsoDesc.PS.BytecodeLength);
 		}
-		if (PsoDesc.DS.BytecodeLength)
-		{
-			DiskCache.AppendData((void*)PsoDesc.DS.pShaderBytecode, PsoDesc.DS.BytecodeLength);
-		}
-		if (PsoDesc.HS.BytecodeLength)
-		{
-			DiskCache.AppendData((void*)PsoDesc.HS.pShaderBytecode, PsoDesc.HS.BytecodeLength);
-		}
 		if (PsoDesc.GS.BytecodeLength)
 		{
 			DiskCache.AppendData((void*)PsoDesc.GS.pShaderBytecode, PsoDesc.GS.BytecodeLength);
@@ -767,8 +745,6 @@ static void DumpGraphicsPSO(const FD3D12_GRAPHICS_PIPELINE_STATE_DESC& Desc, con
 		DumpShaderAsm(String, Desc.MS);
 		DumpShaderAsm(String, Desc.AS);
 		DumpShaderAsm(String, Desc.GS);
-		DumpShaderAsm(String, Desc.HS);
-		DumpShaderAsm(String, Desc.DS);
 		DumpShaderAsm(String, Desc.PS);
 	}
 
@@ -1027,8 +1003,6 @@ static void CreateGraphicsPipelineState(ID3D12PipelineState** PSO, FD3D12Adapter
 		MERGE_EXT(M);
 		MERGE_EXT(A);
 		MERGE_EXT(P);
-		MERGE_EXT(D);
-		MERGE_EXT(H);
 		MERGE_EXT(G);
 	#undef MERGE_EXT
 
diff --git a/Engine/Source/Runtime/D3D12RHI/Private/Windows/WindowsD3D12PipelineState.h b/Engine/Source/Runtime/D3D12RHI/Private/Windows/WindowsD3D12PipelineState.h
--- a/Engine/Source/Runtime/D3D12RHI/Private/Windows/WindowsD3D12PipelineState.h
+++ b/Engine/Source/Runtime/D3D12RHI/Private/Windows/WindowsD3D12PipelineState.h
@@ -111,8 +111,6 @@ struct GraphicsPipelineCreationArgs_POD
 		SaveByteCode(Desc.Desc.AS);
 		SaveByteCode(Desc.Desc.PS);
 		SaveByteCode(Desc.Desc.GS);
-		SaveByteCode(Desc.Desc.HS);
-		SaveByteCode(Desc.Desc.DS);
 	}
 
 	void Destroy();
diff --git a/Engine/Source/Runtime/D3D12RHI/Public/D3D12RootSignatureDefinitions.h b/Engine/Source/Runtime/D3D12RHI/Public/D3D12RootSignatureDefinitions.h
--- a/Engine/Source/Runtime/D3D12RHI/Public/D3D12RootSignatureDefinitions.h
+++ b/Engine/Source/Runtime/D3D12RHI/Public/D3D12RootSignatureDefinitions.h
@@ -85,12 +85,6 @@ namespace D3D12ShaderUtils
 		Creator->AddTable(SF_Geometry, FRootSignatureCreator::SRV, MAX_SRVS);
 		Creator->AddTable(SF_Geometry, FRootSignatureCreator::CBV, MAX_CBS);
 		Creator->AddTable(SF_Geometry, FRootSignatureCreator::Sampler, MAX_SAMPLERS);
-		Creator->AddTable(SF_Hull, FRootSignatureCreator::SRV, MAX_SRVS);
-		Creator->AddTable(SF_Hull, FRootSignatureCreator::CBV, MAX_CBS);
-		Creator->AddTable(SF_Hull, FRootSignatureCreator::Sampler, MAX_SAMPLERS);
-		Creator->AddTable(SF_Domain, FRootSignatureCreator::SRV, MAX_SRVS);
-		Creator->AddTable(SF_Domain, FRootSignatureCreator::CBV, MAX_CBS);
-		Creator->AddTable(SF_Domain, FRootSignatureCreator::Sampler, MAX_SAMPLERS);
 		Creator->AddTable(SF_NumFrequencies, FRootSignatureCreator::UAV, MAX_UAVS);
 	}
 
@@ -117,10 +111,6 @@ namespace D3D12ShaderUtils
 			return D3D12_SHADER_VISIBILITY_PIXEL;
 		case SF_Geometry:
 			return D3D12_SHADER_VISIBILITY_GEOMETRY;
-		case SF_Domain:
-			return D3D12_SHADER_VISIBILITY_DOMAIN;
-		case SF_Hull:
-			return D3D12_SHADER_VISIBILITY_HULL;
 		default:
 			return D3D12_SHADER_VISIBILITY_ALL;
 		}
diff --git a/Engine/Source/Runtime/D3D12RHI/Public/D3D12Util.h b/Engine/Source/Runtime/D3D12RHI/Public/D3D12Util.h
--- a/Engine/Source/Runtime/D3D12RHI/Public/D3D12Util.h
+++ b/Engine/Source/Runtime/D3D12RHI/Public/D3D12Util.h
@@ -64,8 +64,6 @@ enum EShaderVisibility
 {
 	SV_Vertex,
 	SV_Pixel,
-	SV_Hull,
-	SV_Domain,
 	SV_Geometry,
 	SV_Mesh,
 	SV_Amplification,
