Don't track rhi stats for transient allocations - getting size of texture is slow and they are deallocated already by the end of the render thread frame as well (rhi validated)

diff --git a/Engine/Source/Runtime/D3D12RHI/Private/D3D12Buffer.cpp b/Engine/Source/Runtime/D3D12RHI/Private/D3D12Buffer.cpp
--- a/Engine/Source/Runtime/D3D12RHI/Private/D3D12Buffer.cpp
+++ b/Engine/Source/Runtime/D3D12RHI/Private/D3D12Buffer.cpp
@@ -212,8 +212,8 @@ FD3D12Buffer* FD3D12Adapter::CreateRHIBuffer(FRHICommandListImmediate* RHICmdLis
 	const bool bIsDynamic = (InUsage & BUF_AnyDynamic) ? true : false;
 	const uint32 FirstGPUIndex = CreateInfo.GPUMask.GetFirstIndex();
 
-	// Transient resources don't get any actual D3D12 resource asigned at this point yet
-	bool bIsTransient = GSupportsTransientResourceAliasing && (InUsage & BUF_Transient);
+	// Transient flag set?
+	bool bIsTransient = (InUsage & BUF_Transient);
 	
 	// Does this resource support tracking?
 	const bool bSupportResourceStateTracking = !bIsDynamic && FD3D12DefaultBufferAllocator::IsPlacedResource(InDesc.Flags, InResourceStateMode);
@@ -341,7 +341,11 @@ FD3D12Buffer* FD3D12Adapter::CreateRHIBuffer(FRHICommandListImmediate* RHICmdLis
 		CreateInfo.ResourceArray->Discard();
 	}
 
-	UpdateBufferStats(GetBufferStats(InUsage), BufferOut->ResourceLocation.GetSize());
+	// Don't update stats for transient resources
+	if (!bIsTransient)
+	{
+		UpdateBufferStats(GetBufferStats(InUsage), BufferOut->ResourceLocation.GetSize());
+	}
 
 	return BufferOut;
 }
diff --git a/Engine/Source/Runtime/D3D12RHI/Private/D3D12Texture.cpp b/Engine/Source/Runtime/D3D12RHI/Private/D3D12Texture.cpp
--- a/Engine/Source/Runtime/D3D12RHI/Private/D3D12Texture.cpp
+++ b/Engine/Source/Runtime/D3D12RHI/Private/D3D12Texture.cpp
@@ -419,20 +419,24 @@ void FD3D12TextureStats::D3D12TextureAllocated(FD3D12Texture3D& Texture)
 
 	if (D3D12Texture3D)
 	{
-		const D3D12_RESOURCE_DESC& Desc = D3D12Texture3D->GetDesc();
-		const D3D12_RESOURCE_ALLOCATION_INFO AllocationInfo = Texture.GetParentDevice()->GetDevice()->GetResourceAllocationInfo(0, 1, &Desc);
-		const int64 TextureSize = AllocationInfo.SizeInBytes;
+		// Don't update state for virtual or transient textures	
+		if (!EnumHasAnyFlags(Texture.GetFlags(), TexCreate_Virtual | TexCreate_Transient))
+		{
+			const D3D12_RESOURCE_DESC& Desc = D3D12Texture3D->GetDesc();
+			const D3D12_RESOURCE_ALLOCATION_INFO AllocationInfo = Texture.GetParentDevice()->GetDevice()->GetResourceAllocationInfo(0, 1, &Desc);
+			const int64 TextureSize = AllocationInfo.SizeInBytes;
 
-		Texture.SetMemorySize(TextureSize);
+			Texture.SetMemorySize(TextureSize);
 
-		UpdateD3D12TextureStats(Desc, TextureSize, true, false, Texture.IsStreamable());
+			UpdateD3D12TextureStats(Desc, TextureSize, true, false, Texture.IsStreamable());
 
 #if PLATFORM_WINDOWS
-		// On Windows there is no way to hook into the low level d3d allocations and frees.
-		// This means that we must manually add the tracking here.
-		LLM(FLowLevelMemTracker::Get().OnLowLevelAlloc(ELLMTracker::Platform, Texture.GetResource()->GetResource(), Texture.GetMemorySize(), ELLMTag::GraphicsPlatform));
-		LLM(FLowLevelMemTracker::Get().OnLowLevelAlloc(ELLMTracker::Default, Texture.GetResource()->GetResource(), Texture.GetMemorySize(), ELLMTag::Textures));
+			// On Windows there is no way to hook into the low level d3d allocations and frees.
+			// This means that we must manually add the tracking here.
+			LLM(FLowLevelMemTracker::Get().OnLowLevelAlloc(ELLMTracker::Platform, Texture.GetResource()->GetResource(), Texture.GetMemorySize(), ELLMTag::GraphicsPlatform));
+			LLM(FLowLevelMemTracker::Get().OnLowLevelAlloc(ELLMTracker::Default, Texture.GetResource()->GetResource(), Texture.GetMemorySize(), ELLMTag::Textures));
 #endif
+		}
 	}
 }
 
@@ -442,18 +446,22 @@ void FD3D12TextureStats::D3D12TextureDeleted(FD3D12Texture3D& Texture)
 
 	if (D3D12Texture3D)
 	{
-		const D3D12_RESOURCE_DESC& Desc = D3D12Texture3D->GetDesc();
-		const int64 TextureSize = Texture.GetMemorySize();
-		if (TextureSize > 0)
+		// Don't update state for virtual or transient textures	
+		if (!EnumHasAnyFlags(Texture.GetFlags(), TexCreate_Virtual | TexCreate_Transient))
 		{
-			UpdateD3D12TextureStats(Desc, -TextureSize, true, false, Texture.IsStreamable());
+			const D3D12_RESOURCE_DESC& Desc = D3D12Texture3D->GetDesc();
+			const int64 TextureSize = Texture.GetMemorySize();
+			if (TextureSize > 0)
+			{
+				UpdateD3D12TextureStats(Desc, -TextureSize, true, false, Texture.IsStreamable());
 
 #if PLATFORM_WINDOWS
-			// On Windows there is no way to hook into the low level d3d allocations and frees.
-			// This means that we must manually add the tracking here.
-			LLM(FLowLevelMemTracker::Get().OnLowLevelFree(ELLMTracker::Platform, Texture.GetResource()->GetResource()));
-			LLM(FLowLevelMemTracker::Get().OnLowLevelFree(ELLMTracker::Default, Texture.GetResource()->GetResource()));
+				// On Windows there is no way to hook into the low level d3d allocations and frees.
+				// This means that we must manually add the tracking here.
+				LLM(FLowLevelMemTracker::Get().OnLowLevelFree(ELLMTracker::Platform, Texture.GetResource()->GetResource()));
+				LLM(FLowLevelMemTracker::Get().OnLowLevelFree(ELLMTracker::Default, Texture.GetResource()->GetResource()));
 #endif
+			}
 		}
 	}
 }
@@ -892,8 +900,8 @@ TD3D12Texture2D<BaseResourceType>* FD3D12DynamicRHI::CreateD3D12Texture2D(FRHICo
 		ClearValuePtr = &ClearValue;
 	}
 
-	// Transient resources don't get any actual D3D12 resource asigned at this point yet
-	bool bIsTransient = GSupportsTransientResourceAliasing && (Flags & TexCreate_Transient);
+	// Check if transient
+	bool bIsTransient = EnumHasAnyFlags(Flags, TexCreate_Transient);
 
 	// The state this resource will be in when it leaves this function
 	const FD3D12Resource::FD3D12ResourceTypeHelper Type(TextureDesc, D3D12_HEAP_TYPE_DEFAULT);
@@ -1191,8 +1199,8 @@ FD3D12Texture3D* FD3D12DynamicRHI::CreateD3D12Texture3D(FRHICommandListImmediate
 	check(!(Flags & TexCreate_ResolveTargetable));
 	check(Flags & TexCreate_ShaderResource);
 
-	// Transient resources don't get any actual D3D12 resource asigned at this point yet
-	bool bIsTransient = GSupportsTransientResourceAliasing && (Flags & TexCreate_Transient);
+	// Is transient
+	bool bIsTransient = EnumHasAnyFlags(Flags, TexCreate_Transient);
 
 	D3D12_CLEAR_VALUE *ClearValuePtr = nullptr;
 	D3D12_CLEAR_VALUE ClearValue;
diff --git a/Engine/Source/Runtime/D3D12RHI/Private/D3D12TransientResourceAllocator.cpp b/Engine/Source/Runtime/D3D12RHI/Private/D3D12TransientResourceAllocator.cpp
--- a/Engine/Source/Runtime/D3D12RHI/Private/D3D12TransientResourceAllocator.cpp
+++ b/Engine/Source/Runtime/D3D12RHI/Private/D3D12TransientResourceAllocator.cpp
@@ -619,7 +619,10 @@ FRHITexture* FD3D12TransientResourceAllocator::CreateTexture(const FRHITextureCr
 		FRHIResourceCreateInfo CreateInfo(InDebugName);
 		CreateInfo.ClearValueBinding = InCreateInfo.ClearValue;
 
-		FRHITexture* RHITexture = nullptr;
+		// Force transient flag
+		ETextureCreateFlags CreationFlags = InCreateInfo.Flags;
+		CreationFlags |= TexCreate_Transient;
+
 		switch (InCreateInfo.Dimension)
 		{
 		case ETextureDimension::Texture2D:
@@ -627,7 +630,7 @@ FRHITexture* FD3D12TransientResourceAllocator::CreateTexture(const FRHITextureCr
 			const bool bTextureArray = false;
 			const bool bCubeTexture = false;
 			ID3D12ResourceAllocator* ResourceAllocator = this;
-			FD3D12Texture2D* Texture2D = FD3D12DynamicRHI::GetD3DRHI()->CreateD3D12Texture2D<FD3D12BaseTexture2D>(nullptr, InCreateInfo.Extent.X, InCreateInfo.Extent.Y, 1, bTextureArray, bCubeTexture, InCreateInfo.Format, InCreateInfo.NumMips, InCreateInfo.NumSamples, InCreateInfo.Flags, InitialState, CreateInfo, ResourceAllocator);
+			FD3D12Texture2D* Texture2D = FD3D12DynamicRHI::GetD3DRHI()->CreateD3D12Texture2D<FD3D12BaseTexture2D>(nullptr, InCreateInfo.Extent.X, InCreateInfo.Extent.Y, 1, bTextureArray, bCubeTexture, InCreateInfo.Format, InCreateInfo.NumMips, InCreateInfo.NumSamples, CreationFlags, InitialState, CreateInfo, ResourceAllocator);
 			TextureData.RHITexture = Texture2D;
 			TextureData.ResourceDesc = Texture2D->GetResource()->GetDesc();
 			TextureData.ClearValue = Texture2D->GetResource()->GetClearValue();
@@ -640,7 +643,7 @@ FRHITexture* FD3D12TransientResourceAllocator::CreateTexture(const FRHITextureCr
 		{
 			// Only support 2d textures for now
 			ID3D12ResourceAllocator* ResourceAllocator = this;
-			FD3D12Texture3D* Texture3D = FD3D12DynamicRHI::GetD3DRHI()->CreateD3D12Texture3D(nullptr, InCreateInfo.Extent.X, InCreateInfo.Extent.Y, InCreateInfo.Depth, InCreateInfo.Format, InCreateInfo.NumMips, InCreateInfo.Flags, InitialState, CreateInfo, ResourceAllocator);
+			FD3D12Texture3D* Texture3D = FD3D12DynamicRHI::GetD3DRHI()->CreateD3D12Texture3D(nullptr, InCreateInfo.Extent.X, InCreateInfo.Extent.Y, InCreateInfo.Depth, InCreateInfo.Format, InCreateInfo.NumMips, CreationFlags, InitialState, CreateInfo, ResourceAllocator);
 			TextureData.RHITexture = Texture3D;
 			TextureData.ResourceDesc = Texture3D->GetResource()->GetDesc();
 			TextureData.ClearValue = Texture3D->GetResource()->GetClearValue();
@@ -656,7 +659,7 @@ FRHITexture* FD3D12TransientResourceAllocator::CreateTexture(const FRHITextureCr
 			const bool bTextureArray = false;
 			const bool bCubeTexture = true;
 			ID3D12ResourceAllocator* ResourceAllocator = this;
-			FD3D12TextureCube* TextureCube = FD3D12DynamicRHI::GetD3DRHI()->CreateD3D12Texture2D<FD3D12BaseTextureCube>(nullptr, InCreateInfo.Extent.X, InCreateInfo.Extent.Y, 6, bTextureArray, bCubeTexture, InCreateInfo.Format, InCreateInfo.NumMips, InCreateInfo.NumSamples, InCreateInfo.Flags, InitialState, CreateInfo, ResourceAllocator);
+			FD3D12TextureCube* TextureCube = FD3D12DynamicRHI::GetD3DRHI()->CreateD3D12Texture2D<FD3D12BaseTextureCube>(nullptr, InCreateInfo.Extent.X, InCreateInfo.Extent.Y, 6, bTextureArray, bCubeTexture, InCreateInfo.Format, InCreateInfo.NumMips, InCreateInfo.NumSamples, CreationFlags, InitialState, CreateInfo, ResourceAllocator);
 			TextureData.RHITexture = TextureCube;
 			TextureData.ResourceDesc = TextureCube->GetResource()->GetDesc();
 			TextureData.ClearValue = TextureCube->GetResource()->GetClearValue();
@@ -694,6 +697,10 @@ FRHIBuffer* FD3D12TransientResourceAllocator::CreateBuffer(const FRHIBufferCreat
 
 	FD3D12Buffer* Buffer = nullptr;
 
+	// Force transient flag
+	EBufferUsageFlags UsageFlags = InCreateInfo.Usage;
+	UsageFlags |= BUF_Transient;
+
 	// Found something?
 	if (BufferData.RHIBuffer)
 	{
@@ -701,7 +708,6 @@ FRHIBuffer* FD3D12TransientResourceAllocator::CreateBuffer(const FRHIBufferCreat
 		Buffer = FD3D12DynamicRHI::ResourceCast((FRHIBuffer*)BufferData.RHIBuffer);
 		check(!Buffer->ResourceLocation.IsValid());
 
-		EBufferUsageFlags UsageFlags = InCreateInfo.Usage;
 		D3D12_RESOURCE_DESC Desc;
 		uint32 Alignment;
 		FD3D12Buffer::GetResourceDescAndAlignment(InCreateInfo.Size, InCreateInfo.Stride, UsageFlags, Desc, Alignment);
@@ -718,7 +724,7 @@ FRHIBuffer* FD3D12TransientResourceAllocator::CreateBuffer(const FRHIBufferCreat
 	{
 		FRHIResourceCreateInfo CreateInfo(InDebugName);
 		ID3D12ResourceAllocator* ResourceAllocator = this;
-		Buffer = FD3D12DynamicRHI::GetD3DRHI()->CreateD3D12Buffer(nullptr, InCreateInfo.Size, InCreateInfo.Usage, InCreateInfo.Stride, InitialState, CreateInfo, ResourceAllocator);
+		Buffer = FD3D12DynamicRHI::GetD3DRHI()->CreateD3D12Buffer(nullptr, InCreateInfo.Size, UsageFlags, InCreateInfo.Stride, InitialState, CreateInfo, ResourceAllocator);
 		BufferData.RHIBuffer = Buffer;
 		BufferData.CreateInfo = InCreateInfo;
 	}
