Console command D3D12.DumpRayTracingGeometries now only shows top 50 largest geometries. Use D3D12.DumpRayTracingGeometries N or D3D12.DumpRayTracingGeometries all to customize the behavior.

diff --git a/Engine/Source/Runtime/D3D12RHI/Private/D3D12RayTracing.cpp b/Engine/Source/Runtime/D3D12RHI/Private/D3D12RayTracing.cpp
--- a/Engine/Source/Runtime/D3D12RHI/Private/D3D12RayTracing.cpp
+++ b/Engine/Source/Runtime/D3D12RHI/Private/D3D12RayTracing.cpp
@@ -16,6 +16,7 @@
 #include "RayTracingInstanceCopyShader.h"
 #include "Async/ParallelFor.h"
 #include "Misc/BufferedOutputDevice.h"
+#include "String/LexFromString.h"
 
 static int32 GRayTracingDebugForceOpaque = 0;
 static FAutoConsoleVariableRef CVarRayTracingDebugForceOpaque(
@@ -158,14 +159,38 @@ static FD3D12RayTracingGeometryTracker& GetD3D12RayTracingGeometryTracker()
 	return Instance;
 }
 
-static FAutoConsoleCommandWithOutputDevice GD3D12DumpRayTracingGeometriesCmd(
+static FAutoConsoleCommandWithWorldArgsAndOutputDevice GD3D12DumpRayTracingGeometriesCmd(
 	TEXT("D3D12.DumpRayTracingGeometries"),
 	TEXT("Dump memory allocations for ray tracing resources."),
-	FConsoleCommandWithOutputDeviceDelegate::CreateStatic([](FOutputDevice& OutputDevice)
+	FConsoleCommandWithWorldArgsAndOutputDeviceDelegate::CreateStatic([](const TArray<FString>& Args, UWorld*, FOutputDevice& OutputDevice)
 {
 	FD3D12RayTracingGeometryTracker& Tracker = GetD3D12RayTracingGeometryTracker();
 	FScopeLock Lock(&Tracker.CS);
 
+	enum class EMode
+	{
+		Top,
+		All,
+	};
+
+	// Default: show top 50 largest objects.
+	EMode Mode = EMode::Top;
+	int32 NumEntriesToShow = 50;
+
+	if (Args.Num())
+	{
+		if (Args[0] == TEXT("all"))
+		{
+			Mode = EMode::All;
+			NumEntriesToShow = -1;
+		}
+		else if (FCString::IsNumeric(*Args[0]))
+		{
+			Mode = EMode::Top;
+			LexFromString(NumEntriesToShow, *Args[0]);
+		}
+	}
+
 	auto GetGeometrySize = [](FD3D12RayTracingGeometry& Geometry)
 	{
 		if (Geometry.AccelerationStructureCompactedSize != 0)
@@ -187,21 +212,49 @@ static FAutoConsoleCommandWithOutputDevice GD3D12DumpRayTracingGeometriesCmd(
 	FBufferedOutputDevice BufferedOutput;
 	FName CategoryName(TEXT("D3D12RayTracing"));
 	uint64 TotalSizeBytes = 0;
+	uint64 TopSizeBytes = 0;
 	BufferedOutput.CategorizedLogf(CategoryName, ELogVerbosity::Log, TEXT("Tracked FD3D12RayTracingGeometry objects"));
-	for (FD3D12RayTracingGeometry* Geometry : Geometries)
+
+	if (NumEntriesToShow < 0 || NumEntriesToShow > Geometries.Num())
 	{
+		NumEntriesToShow = Geometries.Num();
+	}
+
+	if (NumEntriesToShow != Geometries.Num())
+	{
+		BufferedOutput.CategorizedLogf(CategoryName, ELogVerbosity::Log, TEXT("Showing %d out of %d"), NumEntriesToShow, Geometries.Num());
+	}
+
+	for (int32 i=0; i< Geometries.Num(); ++i)
+	{
+		FD3D12RayTracingGeometry* Geometry = Geometries[i];
 		uint64 SizeBytes = GetGeometrySize(*Geometry);
-		BufferedOutput.CategorizedLogf(CategoryName, ELogVerbosity::Log, TEXT("Name: %s - Size: %.3f MB - Prims: %d - Segments: %d -  Compaction: %d - Update: %d"),
-			Geometry->DebugName.IsValid() ? *Geometry->DebugName.ToString() : TEXT("*UNKNOWN*"),
-			SizeBytes / double(1 << 20),
-			Geometry->TotalPrimitiveCount,
-			Geometry->Segments.Num(),
-			(int32)EnumHasAllFlags(Geometry->BuildFlags, ERayTracingAccelerationStructureFlags::AllowCompaction),
-			(int32)EnumHasAllFlags(Geometry->BuildFlags, ERayTracingAccelerationStructureFlags::AllowUpdate));
+		if (i < NumEntriesToShow)
+		{
+			BufferedOutput.CategorizedLogf(CategoryName, ELogVerbosity::Log, TEXT("Name: %s - Size: %.3f MB - Prims: %d - Segments: %d -  Compaction: %d - Update: %d"),
+				Geometry->DebugName.IsValid() ? *Geometry->DebugName.ToString() : TEXT("*UNKNOWN*"),
+				SizeBytes / double(1 << 20),
+				Geometry->TotalPrimitiveCount,
+				Geometry->Segments.Num(),
+				(int32)EnumHasAllFlags(Geometry->BuildFlags, ERayTracingAccelerationStructureFlags::AllowCompaction),
+				(int32)EnumHasAllFlags(Geometry->BuildFlags, ERayTracingAccelerationStructureFlags::AllowUpdate));
+			TopSizeBytes += SizeBytes;
+		}
 		TotalSizeBytes += SizeBytes;
 	}
 
-	BufferedOutput.CategorizedLogf(CategoryName, ELogVerbosity::Log, TEXT("Total size: %.3f MB"), TotalSizeBytes / double(1<<20));
+	double TotalSizeF = double(TotalSizeBytes) / double(1 << 20);
+	double TopSizeF = double(TopSizeBytes) / double(1 << 20);
+
+	if (NumEntriesToShow != Geometries.Num())
+	{
+		BufferedOutput.CategorizedLogf(CategoryName, ELogVerbosity::Log, TEXT("Use command `D3D12.DumpRayTracingGeometries all` to dump all objects."));
+
+		BufferedOutput.CategorizedLogf(CategoryName, ELogVerbosity::Log, TEXT("Top %d size: %.3f MB (%.2f%% of total)"),
+			NumEntriesToShow, TopSizeF, 100.0 * TopSizeF / TotalSizeF);
+	}
+
+	BufferedOutput.CategorizedLogf(CategoryName, ELogVerbosity::Log, TEXT("Total size: %.3f MB"), TotalSizeF);
 
 	BufferedOutput.RedirectTo(OutputDevice);
 }));
