Only allow mesh shaders when using SM6.

diff --git a/Engine/Source/Runtime/D3D12RHI/Private/D3D12Adapter.cpp b/Engine/Source/Runtime/D3D12RHI/Private/D3D12Adapter.cpp
--- a/Engine/Source/Runtime/D3D12RHI/Private/D3D12Adapter.cpp
+++ b/Engine/Source/Runtime/D3D12RHI/Private/D3D12Adapter.cpp
@@ -491,7 +491,7 @@ void FD3D12Adapter::CreateRootDevice(bool bWithDebug)
 		{
 			bHeapNotZeroedSupported = true;
 
-			GRHISupportsMeshShaders = (Features.MeshShaderTier != D3D12_MESH_SHADER_TIER_NOT_SUPPORTED);
+			GRHISupportsMeshShaders = GMaxRHIFeatureLevel >= ERHIFeatureLevel::SM6 && (Features.MeshShaderTier != D3D12_MESH_SHADER_TIER_NOT_SUPPORTED);
 		}
 	}
 #endif
