Rewrite of RHI transient resource system.

 - Views are cached on RHI transient resources; view renames are no longer necessary.
 - RHI Transient resources utilize a single cache per heap keyed off of the descriptor + offset. Resource caches and heaps are garbage collected.
 - CPU performance is effectively equivalent to the existing pooled resource method.
 - Added common RHI transient resource allocator implementation in RHI core; significantly reduces the amount of platform code.
 - Resource aliasing overlaps are tracked by the RHI and submitted through an acquire operation.
 - Fixed D3D12 implementation to support multi-GPU.
 - Removed condition that excluded small (<64k) buffers in the transient allocator.
 - RHI validation now checks that resource overlaps are valid; i.e. if an overlap occurs between resource A and B during an acquire of B, validation checks that A has been discarded.

diff --git a/Engine/Source/Runtime/D3D12RHI/Private/D3D12Adapter.cpp b/Engine/Source/Runtime/D3D12RHI/Private/D3D12Adapter.cpp
--- a/Engine/Source/Runtime/D3D12RHI/Private/D3D12Adapter.cpp
+++ b/Engine/Source/Runtime/D3D12RHI/Private/D3D12Adapter.cpp
@@ -977,6 +977,8 @@ void FD3D12Adapter::InitializeDevices()
 
 		PipelineStateCache.RebuildFromDiskCache(StaticGraphicsRS, StaticComputeRS);
 	}
+
+	TransientResourceSystem = FD3D12TransientResourceSystem::Create(this, FRHIGPUMask::All());
 }
 
 void FD3D12Adapter::InitializeRayTracing()
@@ -1026,6 +1028,8 @@ void FD3D12Adapter::Cleanup()
 	// Reset the RHI initialized flag.
 	GIsRHIInitialized = false;
 
+	TransientResourceSystem.Reset();
+
 	for (auto& Viewport : Viewports)
 	{
 		Viewport->IssueFrameEvent();
@@ -1226,6 +1230,7 @@ void FD3D12Adapter::EndFrame()
 		GetUploadHeapAllocator(GPUIndex).CleanUpAllocations(FrameLag);
 	}
 	GetDeferredDeletionQueue().ReleaseResources(false, false);
+	GetTransientResourceSystem().GarbageCollect();
 
 #if D3D12_SUBMISSION_GAP_RECORDER
 	SubmitGapRecorderTimestamps();
diff --git a/Engine/Source/Runtime/D3D12RHI/Private/D3D12Adapter.h b/Engine/Source/Runtime/D3D12RHI/Private/D3D12Adapter.h
--- a/Engine/Source/Runtime/D3D12RHI/Private/D3D12Adapter.h
+++ b/Engine/Source/Runtime/D3D12RHI/Private/D3D12Adapter.h
@@ -338,6 +338,8 @@ public:
 	inline FD3D12CommandContextRedirector& GetDefaultContextRedirector() { return DefaultContextRedirector; }
 	inline FD3D12CommandContextRedirector& GetDefaultAsyncComputeContextRedirector() { return DefaultAsyncComputeContextRedirector; }
 
+	inline FD3D12TransientResourceSystem& GetTransientResourceSystem() { return *TransientResourceSystem; }
+
 	FD3D12TemporalEffect* GetTemporalEffect(const FName& EffectName);
 
 	FD3D12FastConstantAllocator& GetTransientUniformBufferAllocator();
@@ -476,6 +478,8 @@ protected:
 	TMap<FName, FD3D12TemporalEffect> TemporalEffectMap;
 #endif
 
+	TUniquePtr<FD3D12TransientResourceSystem> TransientResourceSystem;
+
 	// Each of these devices represents a physical GPU 'Node'
 	FD3D12Device* Devices[MAX_NUM_GPUS];
 
diff --git a/Engine/Source/Runtime/D3D12RHI/Private/D3D12Allocation.cpp b/Engine/Source/Runtime/D3D12RHI/Private/D3D12Allocation.cpp
--- a/Engine/Source/Runtime/D3D12RHI/Private/D3D12Allocation.cpp
+++ b/Engine/Source/Runtime/D3D12RHI/Private/D3D12Allocation.cpp
@@ -998,6 +998,7 @@ void FD3D12BucketAllocator::Reset()
 
 FD3D12UploadHeapAllocator::FD3D12UploadHeapAllocator(FD3D12Adapter* InParent, FD3D12Device* InParentDevice, const FString& InName)
 	: FD3D12AdapterChild(InParent)
+	, FD3D12DeviceChild(InParentDevice)
 	, FD3D12MultiNodeGPUObject(InParentDevice->GetGPUMask(), FRHIGPUMask::All()), // Upload memory, thus they can be trivially visibile to all GPUs
 	SmallBlockAllocator(InParentDevice, GetVisibilityMask(), FD3D12ResourceInitConfig::CreateUpload(), InName, EResourceAllocationStrategy::kManualSubAllocation,
 		GD3D12UploadHeapSmallBlockMaxAllocationSize, GD3D12UploadHeapSmallBlockPoolSize, 256),
@@ -1042,7 +1043,7 @@ void* FD3D12UploadHeapAllocator::AllocUploadResource(uint32 InSize, uint32 InAli
 	{
 		// Forward to the big block allocator
 		const D3D12_RESOURCE_DESC ResourceDesc = CD3DX12_RESOURCE_DESC::Buffer(InSize, D3D12_RESOURCE_FLAG_NONE);
-		BigBlockAllocator.AllocateResource(D3D12_HEAP_TYPE_UPLOAD, ResourceDesc, InSize, InAlignment, ED3D12ResourceStateMode::SingleState, D3D12_RESOURCE_STATE_GENERIC_READ, nullptr, nullptr, ResourceLocation);
+		BigBlockAllocator.AllocateResource(GetParentDevice()->GetGPUIndex(), D3D12_HEAP_TYPE_UPLOAD, ResourceDesc, InSize, InAlignment, ED3D12ResourceStateMode::SingleState, D3D12_RESOURCE_STATE_GENERIC_READ, nullptr, nullptr, ResourceLocation);
 		ResourceLocation.UnlockPoolData();		
 	}
 	
@@ -1607,7 +1608,7 @@ HRESULT FD3D12TextureAllocatorPool::AllocateTexture(
 						 Desc.Flags & D3D12_RESOURCE_FLAG_ALLOW_DEPTH_STENCIL);
 	ED3D12ResourceStateMode ResourceStateMode = bIsReadOnly ? ED3D12ResourceStateMode::Default : ED3D12ResourceStateMode::MultiState;
 	EPoolType PoolType = bIsReadOnly ? (b4KAligment ? EPoolType::ReadOnly4K : EPoolType::ReadOnly) : (bIsRenderTarget ? EPoolType::RenderTarget : EPoolType::UAV);
-	PoolAllocators[(int)PoolType]->AllocateResource(D3D12_HEAP_TYPE_DEFAULT, Desc, Info.SizeInBytes, Info.Alignment, ResourceStateMode, InitialState, ClearValue, Name, TextureLocation);
+	PoolAllocators[(int)PoolType]->AllocateResource(GetParentDevice()->GetGPUIndex(), D3D12_HEAP_TYPE_DEFAULT, Desc, Info.SizeInBytes, Info.Alignment, ResourceStateMode, InitialState, ClearValue, Name, TextureLocation);
 
 	return S_OK;
 }
diff --git a/Engine/Source/Runtime/D3D12RHI/Private/D3D12Allocation.h b/Engine/Source/Runtime/D3D12RHI/Private/D3D12Allocation.h
--- a/Engine/Source/Runtime/D3D12RHI/Private/D3D12Allocation.h
+++ b/Engine/Source/Runtime/D3D12RHI/Private/D3D12Allocation.h
@@ -354,7 +354,7 @@ private:
 //-----------------------------------------------------------------------------
 // This is designed for allocation of scratch memory such as temporary staging buffers
 // or shadow buffers for dynamic resources.
-class FD3D12UploadHeapAllocator : public FD3D12AdapterChild, public FD3D12MultiNodeGPUObject
+class FD3D12UploadHeapAllocator : public FD3D12AdapterChild, public FD3D12DeviceChild, public FD3D12MultiNodeGPUObject
 {
 public:
 
diff --git a/Engine/Source/Runtime/D3D12RHI/Private/D3D12Buffer.cpp b/Engine/Source/Runtime/D3D12RHI/Private/D3D12Buffer.cpp
--- a/Engine/Source/Runtime/D3D12RHI/Private/D3D12Buffer.cpp
+++ b/Engine/Source/Runtime/D3D12RHI/Private/D3D12Buffer.cpp
@@ -181,7 +181,7 @@ void FD3D12Adapter::AllocateBuffer(FD3D12Device* Device,
 	{
 		if (ResourceAllocator)
 		{
-			ResourceAllocator->AllocateResource(D3D12_HEAP_TYPE_DEFAULT, InDesc, InDesc.Width, Alignment, InResourceStateMode, InCreateState, nullptr, CreateInfo.DebugName, ResourceLocation);
+			ResourceAllocator->AllocateResource(Device->GetGPUIndex(), D3D12_HEAP_TYPE_DEFAULT, InDesc, InDesc.Width, Alignment, InResourceStateMode, InCreateState, nullptr, CreateInfo.DebugName, ResourceLocation);
 		}
 		else
 		{
diff --git a/Engine/Source/Runtime/D3D12RHI/Private/D3D12CommandContext.cpp b/Engine/Source/Runtime/D3D12RHI/Private/D3D12CommandContext.cpp
--- a/Engine/Source/Runtime/D3D12RHI/Private/D3D12CommandContext.cpp
+++ b/Engine/Source/Runtime/D3D12RHI/Private/D3D12CommandContext.cpp
@@ -512,7 +512,6 @@ void FD3D12CommandContextBase::RHIBeginFrame()
 
 		Device->GetDefaultBufferAllocator().BeginFrame();
 		Device->GetTextureAllocator().BeginFrame();
-		Device->GetTransientMemoryPoolManager().BeginFrame();
 
 #if D3D12_RHI_RAYTRACING
 		Device->GetRayTracingCompactionRequestHandler()->Update(*ContextAtIndex);
@@ -621,7 +620,6 @@ void FD3D12CommandContextBase::RHIEndFrame()
 		}
 
 		Device->GetTextureAllocator().CleanUpAllocations();
-		Device->GetTransientMemoryPoolManager().EndFrame();
 
 		// Only delete free blocks when not used in the last 2 frames, to make sure we are not allocating and releasing
 		// the same blocks every frame.
@@ -706,6 +704,8 @@ void FD3D12CommandContextBase::UpdateMemoryStats()
 	SET_MEMORY_STAT(STAT_D3D12DemotedVideoMemory, MemoryInfo.DemotedLocalMemory);
 	SET_MEMORY_STAT(STAT_D3D12TotalVideoMemory, MemoryInfo.LocalMemoryInfo.Budget);
 
+	ParentAdapter->GetTransientResourceSystem().UpdateStats();
+
 	uint64 MaxTexAllocWastage = 0;
 	for (uint32 GPUIndex : GPUMask)
 	{
@@ -722,7 +722,6 @@ void FD3D12CommandContextBase::UpdateMemoryStats()
 
 		Device->GetDefaultBufferAllocator().UpdateMemoryStats();
 		ParentAdapter->GetUploadHeapAllocator(GPUIndex).UpdateMemoryStats();
-		Device->GetTransientMemoryPoolManager().UpdateMemoryStats();
 	}
 #endif
 }
@@ -952,6 +951,27 @@ void FD3D12DynamicRHI::RHICreateTransition(FRHITransition* Transition, const FRH
 	Data->bCrossPipeline = bCrossPipeline;
 	Data->TransitionInfos = CreateInfo.TransitionInfos;
 	Data->AliasingInfos = CreateInfo.AliasingInfos;
+
+	uint32 AliasingOverlapCount = 0;
+
+	for (const FRHITransientAliasingInfo& AliasingInfo : Data->AliasingInfos)
+	{
+		AliasingOverlapCount += AliasingInfo.Overlaps.Num();
+	}
+
+	Data->AliasingOverlaps.Reserve(AliasingOverlapCount);
+
+	for (FRHITransientAliasingInfo& AliasingInfo : Data->AliasingInfos)
+	{
+		const int32 OverlapCount = AliasingInfo.Overlaps.Num();
+
+		if (OverlapCount > 0)
+		{
+			const int32 OverlapOffset = Data->AliasingOverlaps.Num();
+			Data->AliasingOverlaps.Append(AliasingInfo.Overlaps.GetData(), OverlapCount);
+			AliasingInfo.Overlaps = MakeArrayView(&Data->AliasingOverlaps[OverlapOffset], OverlapCount);
+		}
+	}
 }
 
 void FD3D12DynamicRHI::RHIReleaseTransition(FRHITransition* Transition)
@@ -960,15 +980,11 @@ void FD3D12DynamicRHI::RHIReleaseTransition(FRHITransition* Transition)
 	Transition->GetPrivateData<FD3D12TransitionData>()->~FD3D12TransitionData();
 }
 
-
 IRHITransientResourceAllocator* FD3D12DynamicRHI::RHICreateTransientResourceAllocator()
 {
-	// Handle mGPU here?
-	FD3D12Device* Device = GetAdapter().GetDevice(0);
-	return new FD3D12TransientResourceAllocator(Device->GetTransientMemoryPoolManager());
+	return new FD3D12TransientResourceAllocator(GetAdapter().GetTransientResourceSystem());
 }
 
-
 //////////////////////////////////////////////////////////////////////////////////////////////////////////
 //
 // FD3D12CommandContextRedirector
diff --git a/Engine/Source/Runtime/D3D12RHI/Private/D3D12CommandContext.h b/Engine/Source/Runtime/D3D12RHI/Private/D3D12CommandContext.h
--- a/Engine/Source/Runtime/D3D12RHI/Private/D3D12CommandContext.h
+++ b/Engine/Source/Runtime/D3D12RHI/Private/D3D12CommandContext.h
@@ -249,8 +249,6 @@ public:
 	virtual void RHIBeginTransitions(TArrayView<const FRHITransition*> Transitions) final override;
 	virtual void RHIEndTransitions(TArrayView<const FRHITransition*> Transitions) final override;
 
-	virtual void RHIReleaseTransientResourceAllocator(IRHITransientResourceAllocator* InAllocator) final override;
-
 	// IRHIComputeContext interface
 	virtual void RHISetComputeShader(FRHIComputeShader* ComputeShader) final override;
 	virtual void RHISetComputePipelineState(FRHIComputePipelineState* ComputePipelineState) final override;
@@ -899,6 +897,7 @@ struct FD3D12TransitionData
 
 	TArray<FRHITransitionInfo, TInlineAllocator<4>> TransitionInfos;
 	TArray<FRHITransientAliasingInfo, TInlineAllocator<4>> AliasingInfos;
+	TArray<FRHITransientAliasingOverlap, TInlineAllocator<4>> AliasingOverlaps;
 	TRefCountPtr<FD3D12Fence> Fence;
 
 	bool bCrossPipeline = false;
diff --git a/Engine/Source/Runtime/D3D12RHI/Private/D3D12Commands.cpp b/Engine/Source/Runtime/D3D12RHI/Private/D3D12Commands.cpp
--- a/Engine/Source/Runtime/D3D12RHI/Private/D3D12Commands.cpp
+++ b/Engine/Source/Runtime/D3D12RHI/Private/D3D12Commands.cpp
@@ -25,6 +25,13 @@ static FAutoConsoleVariableRef CVarSyncTemporalResources(
 	ECVF_RenderThreadSafe
 	);
 
+static int32 GD3D12TransientAllocatorFullAliasingBarrier = 0;
+static FAutoConsoleVariableRef CVarD3D12TransientAllocatorFullAliasingBarrier(
+	TEXT("d3d12.TransientAllocator.FullAliasingBarrier"),
+	GD3D12TransientAllocatorFullAliasingBarrier,
+	TEXT("Inserts a full aliasing barrier on an transient acquire operation. Useful to debug if an aliasing barrier is missing."),
+	ECVF_RenderThreadSafe);
+
 using namespace D3D12RHI;
 
 #define DECLARE_ISBOUNDSHADER(ShaderType) inline void ValidateBoundShader(FD3D12StateCache& InStateCache, FRHI##ShaderType* ShaderType##RHI) \
@@ -333,7 +340,7 @@ static void HandleResourceDiscardTransitions(
 
 			// Get the initial state to force a 'nop' transition so the internal command list resource tracking has the correct state
 			// already to make sure it's not added to the pending transition list to be kicked before this command list (resource is not valid then yet)
-			D3D12_RESOURCE_STATES InitialState = FD3D12TransientResourceAllocator::GetInitialResourceState(Resource->GetDesc());
+			D3D12_RESOURCE_STATES InitialState = GetInitialResourceState(Resource->GetDesc());
 			if (Info.IsWholeResource() || Resource->GetSubresourceCount() == 1)
 			{
 				FD3D12DynamicRHI::TransitionResource(Context.CommandListHandle, Resource, InitialState, InitialState, D3D12_RESOURCE_BARRIER_ALL_SUBRESOURCES, FD3D12DynamicRHI::ETransitionMode::Apply);
@@ -425,17 +432,28 @@ static void HandleTransientAliasing(FD3D12CommandContext& Context, const FD3D12T
 		{
 			TRACE_CPUPROFILER_EVENT_SCOPE(D3D12RHI::AcquireTransient);
 
-			// Try and see if the resource has overlapping resources
-			if (BaseShaderResource->ResourceLocation.IsTransient() && BaseShaderResource->ResourceLocation.GetAllocatorType() == FD3D12ResourceLocation::AT_Pool)
+			if (GD3D12TransientAllocatorFullAliasingBarrier != 0)
 			{
-				// We know it's a transient resource allocator
-				FD3D12TransientResourceAllocator* Allocator = (FD3D12TransientResourceAllocator*)BaseShaderResource->ResourceLocation.GetPoolAllocator();
-
-				// Get the overlapping resources and append aliasing barriers for each of them
-				TArrayView<FD3D12Resource*> OverlappingResources = Allocator->GetOverlappingResources(BaseShaderResource);
-				for (FD3D12Resource* OverlappingResource : OverlappingResources)
+				Context.CommandListHandle.AddAliasingBarrier(nullptr, Resource);
+			}
+			else
+			{
+				for (const FRHITransientAliasingOverlap& Overlap : Info.Overlaps)
 				{
-					Context.CommandListHandle.AddAliasingBarrier(OverlappingResource, Resource);
+					FD3D12Resource* ResourceBefore{};
+
+					switch (Overlap.Type)
+					{
+					case FRHITransientAliasingOverlap::EType::Texture:
+						ResourceBefore = Context.RetrieveTextureBase(Overlap.Texture)->GetResource();
+						break;
+					case FRHITransientAliasingOverlap::EType::Buffer:
+						ResourceBefore = Context.RetrieveObject<FD3D12Buffer>(Overlap.Buffer)->GetResource();
+						break;
+					}
+
+					check(ResourceBefore);
+					Context.CommandListHandle.AddAliasingBarrier(ResourceBefore, Resource);
 				}
 			}
 		}
@@ -444,17 +462,11 @@ static void HandleTransientAliasing(FD3D12CommandContext& Context, const FD3D12T
 			TRACE_CPUPROFILER_EVENT_SCOPE(D3D12RHI::DiscardTransient);
 
 			// Restore the resource back to the initial state when done
-			D3D12_RESOURCE_STATES FinalState = FD3D12TransientResourceAllocator::GetInitialResourceState(BaseShaderResource->GetResource()->GetDesc());
+			D3D12_RESOURCE_STATES FinalState = GetInitialResourceState(Resource->GetDesc());
 			FD3D12DynamicRHI::TransitionResource(Context.CommandListHandle, Resource, D3D12_RESOURCE_STATE_TBD, FinalState, D3D12_RESOURCE_BARRIER_ALL_SUBRESOURCES, FD3D12DynamicRHI::ETransitionMode::Apply);
 
 			// Remove from caches
 			Context.ConditionalClearShaderResource(&BaseShaderResource->ResourceLocation);
-
-			// Release the resource location - transient allocator internally will take care of correct caching
-			BaseShaderResource->ResourceLocation.Clear();
-
-			// Release the views after clearing resource location
-			BaseShaderResource->ResourceRenamed(&BaseShaderResource->ResourceLocation);
 		}
 	}
 }
@@ -610,12 +622,6 @@ void FD3D12CommandContext::RHIEndTransitions(TArrayView<const FRHITransition*> T
 	}
 }
 
-void FD3D12CommandContext::RHIReleaseTransientResourceAllocator(IRHITransientResourceAllocator* InAllocator)
-{
-	// Should be fine to delete now
-	delete (FD3D12TransientResourceAllocator*)InAllocator;
-}
-
 void FD3D12CommandContext::RHISetStaticUniformBuffers(const FUniformBufferStaticBindings& InUniformBuffers)
 {
 	FMemory::Memzero(GlobalUniformBuffers.GetData(), GlobalUniformBuffers.Num() * sizeof(FRHIUniformBuffer*));
diff --git a/Engine/Source/Runtime/D3D12RHI/Private/D3D12Device.cpp b/Engine/Source/Runtime/D3D12RHI/Private/D3D12Device.cpp
--- a/Engine/Source/Runtime/D3D12RHI/Private/D3D12Device.cpp
+++ b/Engine/Source/Runtime/D3D12RHI/Private/D3D12Device.cpp
@@ -42,7 +42,6 @@ FD3D12Device::FD3D12Device(FRHIGPUMask InGPUMask, FD3D12Adapter* InAdapter) :
 	SamplerID(0),
 	DefaultFastAllocator(this, FRHIGPUMask::All(), D3D12_HEAP_TYPE_UPLOAD, 1024 * 1024 * 4),
 	TextureAllocator(this, FRHIGPUMask::All()),
-	TransientMemoryPoolManager(this, FRHIGPUMask::All()),
 	GPUProfilingData(this)
 {
 	for (uint32 QueueType = 0; QueueType < (uint32)ED3D12CommandQueueType::Count; ++QueueType)
@@ -65,8 +64,6 @@ FD3D12Device::~FD3D12Device()
 	// Cleanup the allocator near the end, as some resources may be returned to the allocator or references are shared by multiple GPUs
 	DefaultBufferAllocator.FreeDefaultBufferPools();
 
-	TransientMemoryPoolManager.Destroy();
-
 	DefaultFastAllocator.Destroy();
 
 	TextureAllocator.CleanUpAllocations();
diff --git a/Engine/Source/Runtime/D3D12RHI/Private/D3D12Device.h b/Engine/Source/Runtime/D3D12RHI/Private/D3D12Device.h
--- a/Engine/Source/Runtime/D3D12RHI/Private/D3D12Device.h
+++ b/Engine/Source/Runtime/D3D12RHI/Private/D3D12Device.h
@@ -80,7 +80,6 @@ public:
 	inline FD3D12CommandListManager& GetAsyncCommandListManager() { return *AsyncCommandListManager; }
 	inline FD3D12CommandAllocatorManager& GetTextureStreamingCommandAllocatorManager() { return TextureStreamingCommandAllocatorManager; }
 	inline FD3D12DefaultBufferAllocator& GetDefaultBufferAllocator() { return DefaultBufferAllocator; }
-	inline FD3D12TransientMemoryPoolManager& GetTransientMemoryPoolManager() { return TransientMemoryPoolManager; }
 	inline FD3D12GlobalOnlineSamplerHeap& GetGlobalSamplerHeap() { return GlobalSamplerHeap; }
 	inline FD3D12GlobalHeap& GetGlobalViewHeap() { return GlobalViewHeap; }
 
@@ -197,7 +196,6 @@ protected:
 
 	FD3D12FastAllocator DefaultFastAllocator;
 	FD3D12TextureAllocatorPool TextureAllocator;
-	FD3D12TransientMemoryPoolManager TransientMemoryPoolManager;
 
 	FD3D12ResidencyManager ResidencyManager;
 
diff --git a/Engine/Source/Runtime/D3D12RHI/Private/D3D12PoolAllocator.cpp b/Engine/Source/Runtime/D3D12RHI/Private/D3D12PoolAllocator.cpp
--- a/Engine/Source/Runtime/D3D12RHI/Private/D3D12PoolAllocator.cpp
+++ b/Engine/Source/Runtime/D3D12RHI/Private/D3D12PoolAllocator.cpp
@@ -238,11 +238,11 @@ void FD3D12PoolAllocator::AllocDefaultResource(D3D12_HEAP_TYPE InHeapType, const
 #endif // D3D12_RHI_RAYTRACING
 #endif  // DO_CHECK
 
-	AllocateResource(InHeapType, InDesc, InDesc.Width, InAllocationAlignment, InResourceStateMode, InCreateState, nullptr, InName, ResourceLocation);
+	AllocateResource(GetParentDevice()->GetGPUIndex(), InHeapType, InDesc, InDesc.Width, InAllocationAlignment, InResourceStateMode, InCreateState, nullptr, InName, ResourceLocation);
 }
 
 
-void FD3D12PoolAllocator::AllocateResource(D3D12_HEAP_TYPE InHeapType, const D3D12_RESOURCE_DESC& InDesc, uint64 InSize, uint32 InAllocationAlignment, ED3D12ResourceStateMode InResourceStateMode,
+void FD3D12PoolAllocator::AllocateResource(uint32 GPUIndex, D3D12_HEAP_TYPE InHeapType, const D3D12_RESOURCE_DESC& InDesc, uint64 InSize, uint32 InAllocationAlignment, ED3D12ResourceStateMode InResourceStateMode,
 	D3D12_RESOURCE_STATES InCreateState, const D3D12_CLEAR_VALUE* InClearValue, const TCHAR* InName, FD3D12ResourceLocation& ResourceLocation)
 {
 	TRACE_CPUPROFILER_EVENT_SCOPE(D3D12RHI::AllocatePoolResource);
@@ -257,6 +257,8 @@ void FD3D12PoolAllocator::AllocateResource(D3D12_HEAP_TYPE InHeapType, const D3D
 	FD3D12Device* Device = GetParentDevice();
 	FD3D12Adapter* Adapter = Device->GetParentAdapter();
 
+	check(GPUIndex == Device->GetGPUIndex());
+
 	const bool PoolResource = InSize <= MaxAllocationSize;
 	if (PoolResource)
 	{
diff --git a/Engine/Source/Runtime/D3D12RHI/Private/D3D12PoolAllocator.h b/Engine/Source/Runtime/D3D12RHI/Private/D3D12PoolAllocator.h
--- a/Engine/Source/Runtime/D3D12RHI/Private/D3D12PoolAllocator.h
+++ b/Engine/Source/Runtime/D3D12RHI/Private/D3D12PoolAllocator.h
@@ -137,7 +137,7 @@ public:
 	virtual void DeallocateResource(FD3D12ResourceLocation& ResourceLocation);
 
 	// Implementation of ID3D12ResourceAllocator
-	virtual void AllocateResource(D3D12_HEAP_TYPE InHeapType, const D3D12_RESOURCE_DESC& InDesc, uint64 InSize, uint32 InAllocationAlignment, ED3D12ResourceStateMode InResourceStateMode,
+	virtual void AllocateResource(uint32 GPUIndex, D3D12_HEAP_TYPE InHeapType, const D3D12_RESOURCE_DESC& InDesc, uint64 InSize, uint32 InAllocationAlignment, ED3D12ResourceStateMode InResourceStateMode,
 		D3D12_RESOURCE_STATES InCreateState, const D3D12_CLEAR_VALUE* InClearValue, const TCHAR* InName, FD3D12ResourceLocation& ResourceLocation) override;
 
 	void CleanUpAllocations(uint64 InFrameLag);
diff --git a/Engine/Source/Runtime/D3D12RHI/Private/D3D12RHICommon.h b/Engine/Source/Runtime/D3D12RHI/Private/D3D12RHICommon.h
--- a/Engine/Source/Runtime/D3D12RHI/Private/D3D12RHICommon.h
+++ b/Engine/Source/Runtime/D3D12RHI/Private/D3D12RHICommon.h
@@ -10,7 +10,6 @@ DECLARE_STATS_GROUP(TEXT("D3D12RHI"), STATGROUP_D3D12RHI, STATCAT_Advanced);
 DECLARE_STATS_GROUP(TEXT("D3D12RHI: Memory"), STATGROUP_D3D12Memory, STATCAT_Advanced);
 DECLARE_STATS_GROUP(TEXT("D3D12RHI: Resources"), STATGROUP_D3D12Resources, STATCAT_Advanced);
 DECLARE_STATS_GROUP(TEXT("D3D12RHI: Allocations"), STATGROUP_D3D12Allocations, STATCAT_Advanced);
-DECLARE_STATS_GROUP(TEXT("D3D12RHI: Allocations (Transient)"), STATGROUP_D3D12TransientAllocations, STATCAT_Advanced);
 DECLARE_STATS_GROUP(TEXT("D3D12RHI: Allocations (Buffer)"), STATGROUP_D3D12BufferAllocations, STATCAT_Advanced);
 DECLARE_STATS_GROUP(TEXT("D3D12RHI: Pipeline State (PSO)"), STATGROUP_D3D12PipelineState, STATCAT_Advanced);
 DECLARE_STATS_GROUP(TEXT("D3D12RHI: Descriptor Heap (GPU Visible)"), STATGROUP_D3D12DescriptorHeap, STATCAT_Advanced);
diff --git a/Engine/Source/Runtime/D3D12RHI/Private/D3D12Resources.cpp b/Engine/Source/Runtime/D3D12RHI/Private/D3D12Resources.cpp
--- a/Engine/Source/Runtime/D3D12RHI/Private/D3D12Resources.cpp
+++ b/Engine/Source/Runtime/D3D12RHI/Private/D3D12Resources.cpp
@@ -219,7 +219,7 @@ void FD3D12DeferredDeletionQueue::FD3D12AsyncDeletionWorker::DoWork()
 /////////////////////////////////////////////////////////////////////
 
 
-void ID3D12ResourceAllocator::AllocateTexture(D3D12_HEAP_TYPE InHeapType, const D3D12_RESOURCE_DESC& InDesc, EPixelFormat InUEFormat, ED3D12ResourceStateMode InResourceStateMode,
+void ID3D12ResourceAllocator::AllocateTexture(uint32 GPUIndex, D3D12_HEAP_TYPE InHeapType, const D3D12_RESOURCE_DESC& InDesc, EPixelFormat InUEFormat, ED3D12ResourceStateMode InResourceStateMode,
 	D3D12_RESOURCE_STATES InCreateState, const D3D12_CLEAR_VALUE* InClearValue, const TCHAR* InName, FD3D12ResourceLocation& ResourceLocation)
 {
 	// Check if texture can be 4K aligned
@@ -229,7 +229,7 @@ void ID3D12ResourceAllocator::AllocateTexture(D3D12_HEAP_TYPE InHeapType, const
 
 	// Get the size and alignment for the allocation
 	D3D12_RESOURCE_ALLOCATION_INFO Info = FD3D12DynamicRHI::GetD3DRHI()->GetAdapter().GetDevice(0)->GetResourceAllocationInfo(Desc);
-	AllocateResource(InHeapType, Desc, Info.SizeInBytes, Info.Alignment, InResourceStateMode, InCreateState, InClearValue, InName, ResourceLocation);
+	AllocateResource(GPUIndex, InHeapType, Desc, Info.SizeInBytes, Info.Alignment, InResourceStateMode, InCreateState, InClearValue, InName, ResourceLocation);
 }
 
 
diff --git a/Engine/Source/Runtime/D3D12RHI/Private/D3D12Stats.h b/Engine/Source/Runtime/D3D12RHI/Private/D3D12Stats.h
--- a/Engine/Source/Runtime/D3D12RHI/Private/D3D12Stats.h
+++ b/Engine/Source/Runtime/D3D12RHI/Private/D3D12Stats.h
@@ -91,24 +91,6 @@ DECLARE_MEMORY_STAT_EXTERN(TEXT("BufferPool Memory Alignment Waste"), STAT_D3D12
 DECLARE_DWORD_ACCUMULATOR_STAT_EXTERN(TEXT("BufferPool Page Count"), STAT_D3D12BufferPoolPageCount, STATGROUP_D3D12BufferAllocations, );
 DECLARE_DWORD_ACCUMULATOR_STAT_EXTERN(TEXT("BufferPool Full Pages"), STAT_D3D12BufferPoolFullPages, STATGROUP_D3D12BufferAllocations, );
 
-DECLARE_MEMORY_STAT_EXTERN(TEXT("Transient Memory Pool Allocated"), STAT_D3D12TransientMemoryPoolAllocated, STATGROUP_D3D12Allocations, );
-DECLARE_MEMORY_STAT_EXTERN(TEXT("Transient Memory Frame Pool Used"), STAT_D3D12TransientMemoryFramePoolUsed, STATGROUP_D3D12Allocations, );
-DECLARE_MEMORY_STAT_EXTERN(TEXT("Transient Memory Frame Committed Used"), STAT_D3D12TransientMemoryFrameCommittedUsed, STATGROUP_D3D12Allocations, );
-DECLARE_DWORD_ACCUMULATOR_STAT_EXTERN(TEXT("Transient Memory Pool Allocations"), STAT_D3D12TransientMemoryPoolAllocations, STATGROUP_D3D12Allocations, );
-DECLARE_DWORD_ACCUMULATOR_STAT_EXTERN(TEXT("Transient Memory Committed Allocations"), STAT_D3D12TransientMemoryCommittedAllocations, STATGROUP_D3D12Allocations, );
-
-DECLARE_MEMORY_STAT_EXTERN(TEXT("Buffer Memory Frame Used"), STAT_D3D12TransientMemoryBufferFrameUsed, STATGROUP_D3D12TransientAllocations, );
-DECLARE_MEMORY_STAT_EXTERN(TEXT("Buffer Memory Frame Total Requested"), STAT_D3D12TransientMemoryBufferFrameTotalRequested, STATGROUP_D3D12TransientAllocations, );
-DECLARE_MEMORY_STAT_EXTERN(TEXT("Buffer Memory Frame Committed Allocated"), STAT_D3D12TransientMemoryBufferFrameCommittedAllocated, STATGROUP_D3D12TransientAllocations, );
-DECLARE_DWORD_ACCUMULATOR_STAT_EXTERN(TEXT("Buffer Memory Pool Allocations"), STAT_D3D12TransientMemoryBufferPoolAllocations, STATGROUP_D3D12TransientAllocations, );
-DECLARE_DWORD_ACCUMULATOR_STAT_EXTERN(TEXT("Buffer Memory Allocations"), STAT_D3D12TransientMemoryBufferCommittedAllocations, STATGROUP_D3D12TransientAllocations, );
-
-DECLARE_MEMORY_STAT_EXTERN(TEXT("Texture Memory Frame Used"), STAT_D3D12TransientMemoryTextureFrameUsed, STATGROUP_D3D12TransientAllocations, );
-DECLARE_MEMORY_STAT_EXTERN(TEXT("Texture Memory Frame Total Requested"), STAT_D3D12TransientMemoryTextureFrameTotalRequested, STATGROUP_D3D12TransientAllocations, );
-DECLARE_MEMORY_STAT_EXTERN(TEXT("Texture Memory Frame Committed Allocated"), STAT_D3D12TransientMemoryTextureFrameCommittedAllocated, STATGROUP_D3D12TransientAllocations, );
-DECLARE_DWORD_ACCUMULATOR_STAT_EXTERN(TEXT("Texture Memory Pool Allocations"), STAT_D3D12TransientMemoryTexturePoolAllocations, STATGROUP_D3D12TransientAllocations, );
-DECLARE_DWORD_ACCUMULATOR_STAT_EXTERN(TEXT("Texture Memory Allocations"), STAT_D3D12TransientMemoryTextureCommittedAllocations, STATGROUP_D3D12TransientAllocations, );
-
 DECLARE_MEMORY_STAT_EXTERN(TEXT("UploadPool Memory Allocated"), STAT_D3D12UploadPoolMemoryAllocated, STATGROUP_D3D12Allocations, );
 DECLARE_MEMORY_STAT_EXTERN(TEXT("UploadPool Memory Used"), STAT_D3D12UploadPoolMemoryUsed, STATGROUP_D3D12Allocations, );
 
diff --git a/Engine/Source/Runtime/D3D12RHI/Private/D3D12Texture.cpp b/Engine/Source/Runtime/D3D12RHI/Private/D3D12Texture.cpp
--- a/Engine/Source/Runtime/D3D12RHI/Private/D3D12Texture.cpp
+++ b/Engine/Source/Runtime/D3D12RHI/Private/D3D12Texture.cpp
@@ -932,7 +932,7 @@ TD3D12Texture2D<BaseResourceType>* FD3D12DynamicRHI::CreateD3D12Texture2D(FRHICo
 		if (ResourceAllocator)
 		{
 			const D3D12_HEAP_TYPE HeapType = (Flags & TexCreate_CPUReadback) ? D3D12_HEAP_TYPE_READBACK : D3D12_HEAP_TYPE_DEFAULT;
-			ResourceAllocator->AllocateTexture(HeapType, TextureDesc, (EPixelFormat)Format, ED3D12ResourceStateMode::Default, CreateState, ClearValuePtr, CreateInfo.DebugName, Location);
+			ResourceAllocator->AllocateTexture(Device->GetGPUIndex(), HeapType, TextureDesc, (EPixelFormat)Format, ED3D12ResourceStateMode::Default, CreateState, ClearValuePtr, CreateInfo.DebugName, Location);
 			Location.SetOwner(NewTexture);
 		}
 		else
@@ -958,6 +958,8 @@ TD3D12Texture2D<BaseResourceType>* FD3D12DynamicRHI::CreateD3D12Texture2D(FRHICo
 		// Mark transient
 		Location.SetTransient(bIsTransient);
 
+		check(Location.IsValid());
+
 		uint32 RTVIndex = 0;
 
 		{
@@ -1226,7 +1228,7 @@ FD3D12Texture3D* FD3D12DynamicRHI::CreateD3D12Texture3D(FRHICommandListImmediate
 
 		if (ResourceAllocator)
 		{
-			ResourceAllocator->AllocateTexture(D3D12_HEAP_TYPE_DEFAULT, TextureDesc, (EPixelFormat)Format, ED3D12ResourceStateMode::Default, InitialState, ClearValuePtr, CreateInfo.DebugName, Texture3D->ResourceLocation);
+			ResourceAllocator->AllocateTexture(Device->GetGPUIndex(), D3D12_HEAP_TYPE_DEFAULT, TextureDesc, (EPixelFormat)Format, ED3D12ResourceStateMode::Default, InitialState, ClearValuePtr, CreateInfo.DebugName, Texture3D->ResourceLocation);
 		}
 		else
 		{
@@ -1914,6 +1916,119 @@ void* TD3D12Texture2D<RHIResourceType>::Lock(class FRHICommandListImmediate* RHI
 	return Data;
 }
 
+D3D12_RESOURCE_DESC FD3D12TextureBase::GetResourceDesc(const FRHITextureCreateInfo& CreateInfo)
+{
+	D3D12_RESOURCE_DESC Desc;
+
+	const DXGI_FORMAT Format = GetPlatformTextureResourceFormat((DXGI_FORMAT)GPixelFormats[CreateInfo.Format].PlatformFormat, CreateInfo.Flags);
+
+	if (CreateInfo.Dimension != ETextureDimension::Texture3D)
+	{
+		if (CreateInfo.IsTextureCube())
+		{
+			check(CreateInfo.Extent.X <= (int32)GetMaxCubeTextureDimension());
+			check(CreateInfo.Extent.X == CreateInfo.Extent.Y);
+		}
+		else
+		{
+			check(CreateInfo.Extent.X <= (int32)GetMax2DTextureDimension());
+			check(CreateInfo.Extent.Y <= (int32)GetMax2DTextureDimension());
+		}
+
+		if (CreateInfo.IsTextureArray())
+		{
+			check(CreateInfo.ArraySize <= (int32)GetMaxTextureArrayLayers());
+		}
+
+		uint32 ActualMSAACount = CreateInfo.NumSamples;
+		uint32 ActualMSAAQuality = GetMaxMSAAQuality(ActualMSAACount);
+
+		// 0xffffffff means not supported
+		if (ActualMSAAQuality == 0xffffffff || EnumHasAnyFlags(CreateInfo.Flags, TexCreate_Shared))
+		{
+			// no MSAA
+			ActualMSAACount = 1;
+			ActualMSAAQuality = 0;
+		}
+
+		Desc = CD3DX12_RESOURCE_DESC::Tex2D(
+			Format,
+			CreateInfo.Extent.X,
+			CreateInfo.Extent.Y,
+			CreateInfo.ArraySize * (CreateInfo.IsTextureCube() ? 6 : 1),  // Array size
+			CreateInfo.NumMips,
+			ActualMSAACount,
+			ActualMSAAQuality,
+			D3D12_RESOURCE_FLAG_NONE);  // Add misc flags later
+
+		if (EnumHasAnyFlags(CreateInfo.Flags, TexCreate_Shared))
+		{
+			Desc.Flags |= D3D12_RESOURCE_FLAG_ALLOW_SIMULTANEOUS_ACCESS;
+		}
+
+		if (EnumHasAnyFlags(CreateInfo.Flags, TexCreate_RenderTargetable))
+		{
+			check(!EnumHasAnyFlags(CreateInfo.Flags, TexCreate_DepthStencilTargetable | TexCreate_ResolveTargetable));
+			Desc.Flags |= D3D12_RESOURCE_FLAG_ALLOW_RENDER_TARGET;
+		}
+		else if (EnumHasAnyFlags(CreateInfo.Flags, TexCreate_DepthStencilTargetable))
+		{
+			check(!EnumHasAnyFlags(CreateInfo.Flags, TexCreate_RenderTargetable | TexCreate_ResolveTargetable));
+			Desc.Flags |= D3D12_RESOURCE_FLAG_ALLOW_DEPTH_STENCIL;
+		}
+		else if (EnumHasAnyFlags(CreateInfo.Flags, TexCreate_ResolveTargetable))
+		{
+			check(!EnumHasAnyFlags(CreateInfo.Flags, TexCreate_RenderTargetable | TexCreate_DepthStencilTargetable));
+			if (CreateInfo.Format == PF_DepthStencil || CreateInfo.Format == PF_ShadowDepth || CreateInfo.Format == PF_D24)
+			{
+				Desc.Flags |= D3D12_RESOURCE_FLAG_ALLOW_DEPTH_STENCIL;
+			}
+			else
+			{
+				Desc.Flags |= D3D12_RESOURCE_FLAG_ALLOW_RENDER_TARGET;
+			}
+		}
+
+		if (EnumHasAnyFlags(CreateInfo.Flags, TexCreate_UAV))
+		{
+			Desc.Flags |= D3D12_RESOURCE_FLAG_ALLOW_UNORDERED_ACCESS;
+		}
+
+		if (EnumHasAnyFlags(CreateInfo.Flags, TexCreate_DepthStencilTargetable) && !EnumHasAnyFlags(CreateInfo.Flags, TexCreate_ShaderResource))
+		{
+			// Only deny shader resources if it's a depth resource that will never be used as SRV
+			Desc.Flags |= D3D12_RESOURCE_FLAG_DENY_SHADER_RESOURCE;
+		}
+	}
+	else // ETextureDimension::Texture3D
+	{
+		check(CreateInfo.Dimension == ETextureDimension::Texture3D)
+		check(!EnumHasAnyFlags(CreateInfo.Flags, TexCreate_DepthStencilTargetable | TexCreate_ResolveTargetable));
+		check(EnumHasAnyFlags(CreateInfo.Flags, TexCreate_ShaderResource));
+
+		Desc = CD3DX12_RESOURCE_DESC::Tex3D(
+			Format,
+			CreateInfo.Extent.X,
+			CreateInfo.Extent.Y,
+			CreateInfo.Depth,
+			CreateInfo.NumMips);
+
+		if (EnumHasAnyFlags(CreateInfo.Flags, TexCreate_UAV))
+		{
+			Desc.Flags |= D3D12_RESOURCE_FLAG_ALLOW_UNORDERED_ACCESS;
+		}
+
+		if (EnumHasAnyFlags(CreateInfo.Flags, TexCreate_RenderTargetable))
+		{
+			Desc.Flags |= D3D12_RESOURCE_FLAG_ALLOW_RENDER_TARGET;
+		}
+	}
+
+	Desc.Alignment = D3D12_DEFAULT_RESOURCE_PLACEMENT_ALIGNMENT;
+
+	return Desc;
+}
+
 void FD3D12TextureBase::UpdateTexture(uint32 MipIndex, uint32 DestX, uint32 DestY, uint32 DestZ, const D3D12_TEXTURE_COPY_LOCATION& SourceCopyLocation)
 {
 	FD3D12CommandContext& DefaultContext = GetParentDevice()->GetDefaultCommandContext();
diff --git a/Engine/Source/Runtime/D3D12RHI/Private/D3D12Texture.h b/Engine/Source/Runtime/D3D12RHI/Private/D3D12Texture.h
--- a/Engine/Source/Runtime/D3D12RHI/Private/D3D12Texture.h
+++ b/Engine/Source/Runtime/D3D12RHI/Private/D3D12Texture.h
@@ -193,6 +193,8 @@ public:
 
 	FD3D12CLSyncPoint ReadBackSyncPoint;
 
+	static D3D12_RESOURCE_DESC GetResourceDesc(const FRHITextureCreateInfo& CreateInfo);
+
 protected:
 
 	/** Amount of memory allocated by this texture, in bytes. */
diff --git a/Engine/Source/Runtime/D3D12RHI/Private/D3D12TransientResourceAllocator.cpp b/Engine/Source/Runtime/D3D12RHI/Private/D3D12TransientResourceAllocator.cpp
--- a/Engine/Source/Runtime/D3D12RHI/Private/D3D12TransientResourceAllocator.cpp
+++ b/Engine/Source/Runtime/D3D12RHI/Private/D3D12TransientResourceAllocator.cpp
@@ -3,908 +3,240 @@
 #include "D3D12RHIPrivate.h"
 #include "D3D12TransientResourceAllocator.h"
 
-static int32 GD3D12TransientAllocatorDefaultPoolSizeInMB = 64;
-static FAutoConsoleVariableRef CVarD3D12TransientAllocatorPoolSizeInMB(
-	TEXT("d3d12.TransientAllocator.DefaultPoolSizeInMB"),
-	GD3D12TransientAllocatorDefaultPoolSizeInMB,
-	TEXT("Default size of a D3D12 transient allocator pool in MB (Default 128)"),
-	ECVF_ReadOnly);
-
-static int32 GD3D12TransientAllocatorMaxAllocationSizeInMB = 512;
-static FAutoConsoleVariableRef CVarD3D12TransientAllocatorMaxAllocationSizeInMB(
-	TEXT("d3d12.TransientAllocator.MaxPoolSizeInMB"),
-	GD3D12TransientAllocatorMaxAllocationSizeInMB,
-	TEXT("Max size of an D3D12 transient allocation in MB (Default 512)"),
-	ECVF_ReadOnly);
-
-static int32 GD3D12TransientAllocatorPoolTextures = 1;
-static FAutoConsoleVariableRef CVarD3D12TransientAllocatorPoolTextures(
-	TEXT("d3d12.TransientAllocator.PoolTextures"),
-	GD3D12TransientAllocatorPoolTextures,
-	TEXT("Enable pooling of transient allocated RHITextures in the manager (default enabled)"),
-	ECVF_ReadOnly);
-
-static int32 GD3D12TransientAllocatorPoolBuffers = 0;
-static FAutoConsoleVariableRef CVarD3D12TransientAllocatorPoolBuffers(
-	TEXT("d3d12.TransientAllocator.PoolBuffers"),
-	GD3D12TransientAllocatorPoolBuffers,
-	TEXT("Enable pooling of transient allocated RHIBuffer in the manager (default disabled)"),
-	ECVF_ReadOnly);
-
-uint64 ComputeTextureCreateInfoHash(const FRHITextureCreateInfo& InCreateInfo)
+D3D12_RESOURCE_STATES GetInitialResourceState(const D3D12_RESOURCE_DESC& InDesc)
 {
-	// Make sure all padding is removed
-	FRHITextureCreateInfo NewInfo;
-	FPlatformMemory::Memzero(&NewInfo, sizeof(FRHITextureCreateInfo));
-	NewInfo = InCreateInfo;
-	return CityHash64((const char*)&NewInfo, sizeof(FRHITextureCreateInfo));
-}
-
-uint64 ComputeBufferCreateInfoHash(const FRHIBufferCreateInfo& InCreateInfo)
-{
-	return CityHash64((const char*)&InCreateInfo, sizeof(FRHIBufferCreateInfo));
-}
-
-
-//-----------------------------------------------------------------------------
-//	FD3D12TransientMemoryPool
-//-----------------------------------------------------------------------------
-
-uint64 FD3D12TransientMemoryPool::FResourceCreateState::GetHash() const
-{
-	uint64 Hash = CityHash64((const char*)&AllocationOffset, sizeof(uint64));
-	Hash = CityHash64WithSeed((const char*)&ResourceDesc, sizeof(D3D12_RESOURCE_DESC), Hash);
-	Hash = CityHash64WithSeed((const char*)&ClearValue, sizeof(D3D12_CLEAR_VALUE), Hash);
-	return Hash;
-}
-
-
-FD3D12TransientMemoryPool::~FD3D12TransientMemoryPool()
-{
-	check(AllocatedBlocks == 0);
-
-	ClearActiveResources();
-
-	TArray<FD3D12Resource*> CachedResources;
-	CachedResourceMap.GenerateValueArray(CachedResources);
-	for (FD3D12Resource* Resource : CachedResources)
-	{
-		check(Resource->GetRefCount() == 1);
-		Resource->Release();
-	}
-	CachedResourceMap.Empty();
-}
-
-
-void FD3D12TransientMemoryPool::ResetPool()
-{
-	check(AllocatedBlocks == 0);
-	check(FreeBlocks.Num() == 1);
-	check(FreeBlocks[0]->GetOffset() == 0);
-	check(FreeBlocks[0]->GetSize() == PoolSize);
-	check(FreeBlocks[0]->GetAlignment() == PoolAlignment);
-
-	ClearActiveResources();
-}
-
-
-void FD3D12TransientMemoryPool::SetPoolIndex(int16 InNewPoolIndex)
-{
-	PoolIndex = InNewPoolIndex;
-
-	// Update pool index on single free block (validated during reset)
-	FRHIPoolAllocationData* FreeBlock = FreeBlocks[0];
-	FreeBlock->RemoveFromLinkedList();
-	FreeBlock->InitAsFree(PoolIndex, PoolSize, PoolAlignment, 0);
-	HeadBlock.AddAfter(FreeBlock);
-}
-
-
-void FD3D12TransientMemoryPool::ClearActiveResources()
-{
-	ActiveResources.Empty(ActiveResources.Num());
-}
-
-
-void FD3D12TransientMemoryPool::CheckActiveResources(const FInt64Range& InAllocationRange, TArray<FD3D12Resource *>& OverlappingResources)
-{		
-	// Do we already have active resources in the requested range
-	for (ActiveResourceData& ActiveResource : ActiveResources)
-	{
-		if (ActiveResource.ActiveRange.Overlaps(InAllocationRange))
-		{
-			check(ActiveResource.Resource != nullptr);
-
-			// Add to dependent list of overlapping resources
-			OverlappingResources.Add(ActiveResource.Resource);
-
-			// update the active ranges
-			if (ActiveResource.ActiveRange.GetLowerBoundValue() >= InAllocationRange.GetLowerBoundValue())
-			{
-				if (ActiveResource.ActiveRange.GetUpperBoundValue() > InAllocationRange.GetUpperBoundValue())
-				{
-					ActiveResource.ActiveRange.SetLowerBoundValue(InAllocationRange.GetUpperBoundValue());
-				}
-				else
-				{
-					// Full overlap, make empty range
-					ActiveResource.ActiveRange = FInt64Range(ActiveResource.ActiveRange.GetLowerBoundValue(), ActiveResource.ActiveRange.GetLowerBoundValue());
-				}
-			}
-			else
-			{
-				ActiveResource.ActiveRange.SetUpperBoundValue(InAllocationRange.GetLowerBoundValue());
-			}
-
-			// Mark as invalid because it has no active ranges anymore
-			if (ActiveResource.ActiveRange.IsEmpty())
-			{
-				ActiveResource.Resource = nullptr;
-			}
-		}
-	}
-}
-
-
-FD3D12Resource* FD3D12TransientMemoryPool::FindResourceInCache(uint64 InAllocationOffset, const D3D12_RESOURCE_DESC& InDesc, const D3D12_CLEAR_VALUE* InClearValue, const TCHAR* InName)
-{
-	// Try and find cached placed resource at given offset and location
-	FResourceCreateState CreateState;
-	CreateState.AllocationOffset = InAllocationOffset;
-	CreateState.ResourceDesc = InDesc;
-	if (InClearValue)
-	{
-		CreateState.ClearValue = *InClearValue;
-	}
-	else
-	{
-		FPlatformMemory::Memzero(&CreateState.ClearValue, sizeof(D3D12_CLEAR_VALUE));
-	}
-
-	uint64 CreateStateHash = CreateState.GetHash();
-	FD3D12Resource** Resource = CachedResourceMap.Find(CreateStateHash);
-	if (Resource)
+	// Validate the creation state
+	D3D12_RESOURCE_STATES State = D3D12_RESOURCE_STATE_COMMON;
+	if (EnumHasAnyFlags(InDesc.Flags, D3D12_RESOURCE_FLAG_ALLOW_RENDER_TARGET))
 	{
-		check((*Resource)->GetDesc() == InDesc);		
-		CachedResourceMap.Remove(CreateStateHash);
-		return *Resource;
+		State = D3D12_RESOURCE_STATE_RENDER_TARGET;
 	}
-	else
+	else if (EnumHasAnyFlags(InDesc.Flags, D3D12_RESOURCE_FLAG_ALLOW_DEPTH_STENCIL))
 	{
-		return nullptr;
+		State = D3D12_RESOURCE_STATE_DEPTH_WRITE;
 	}
-}
-
-
-void FD3D12TransientMemoryPool::ReleaseResource(FD3D12Resource* InResource, FRHIPoolAllocationData& InReleasedAllocationData, uint64 InFenceValue)
-{
-	check(InResource->GetHeap() == BackingHeap);
-
-	uint64 AllocationOffset = InReleasedAllocationData.GetOffset();
-	uint64 AllocationSize = InReleasedAllocationData.GetSize();
-
-	// Placed resource object is 'never' freed - it can be reused again this or next frame
-	// Will only be destroyed when the allocator is destroyed
-	FResourceCreateState CreateState;
-	CreateState.AllocationOffset = AllocationOffset;
-	CreateState.ResourceDesc = InResource->GetDesc();
-	CreateState.ClearValue = InResource->GetClearValue();
-
-	uint64 CreateStateHash = CreateState.GetHash();
-	check(CachedResourceMap.Find(CreateStateHash) == nullptr);
-	CachedResourceMap.Add(CreateStateHash, InResource);
-	
-	// Track this resource so aliasing barriers can be added if used again
-	ActiveResourceData ResourceData;
-	ResourceData.AllocationRange = FInt64Range(AllocationOffset, AllocationOffset + AllocationSize);
-	ResourceData.ActiveRange = ResourceData.AllocationRange;
-	ResourceData.Resource = InResource;
-	ActiveResources.Add(ResourceData);
-
-	// Free the pool data so this range can be reallocated again immediately (active range is tracked for aliasing barriers)
-	FRHIPoolAllocationData LockedAllocationData;
-	bool bLocked = true;
-	LockedAllocationData.MoveFrom(InReleasedAllocationData, bLocked);
-	Deallocate(LockedAllocationData);
-
-	// Update the last used frame fence (used during garbage collection)
-	UpdateLastUsedFrameFence(InFenceValue);
-}
-
-//-----------------------------------------------------------------------------
-//	FD3D12TransientMemoryPoolManager
-//-----------------------------------------------------------------------------
-
-
-FD3D12TransientMemoryPoolManager::FD3D12TransientMemoryPoolManager(FD3D12Device* InDevice, FRHIGPUMask VisibleNodes)
-	: FD3D12DeviceChild(InDevice)
-	, FD3D12MultiNodeGPUObject(InDevice->GetGPUMask(), VisibleNodes)
-{
-	// Does the device support Heap Tier 2 - merged heaps for buffers & textures
-	bMergedTypeHeapSupported = (InDevice->GetParentAdapter()->GetResourceHeapTier() == D3D12_RESOURCE_HEAP_TIER_2);
-
-	// texture only interesting in VRAM for now
-	InitConfig.HeapType = D3D12_HEAP_TYPE_DEFAULT;
-
-	// unused for textures because placed and not suballocated
-	InitConfig.ResourceFlags = D3D12_RESOURCE_FLAG_NONE;
-	InitConfig.InitialResourceState = D3D12_RESOURCE_STATE_COMMON;
-
-	// By default the manager and allocator support all resource types
-	InitConfig.HeapFlags = D3D12_HEAP_FLAGS(0); // 0 means nothing is denied
-
-	DefaultPoolSize = GD3D12TransientAllocatorDefaultPoolSizeInMB * 1024 * 1024;
-	PoolAlignment = D3D12_DEFAULT_RESOURCE_PLACEMENT_ALIGNMENT;
-	MaxAllocationSize = GD3D12TransientAllocatorMaxAllocationSizeInMB * 1024 * 1024;
-}
-
-
-void FD3D12TransientMemoryPoolManager::Destroy()
-{
-	for (int32 PoolIndex = 0; PoolIndex < Pools.Num(); ++PoolIndex)
+	else if (EnumHasAnyFlags(InDesc.Flags, D3D12_RESOURCE_FLAG_ALLOW_UNORDERED_ACCESS))
 	{
-		Pools[PoolIndex]->Destroy();
-		delete(Pools[PoolIndex]);
+		State = D3D12_RESOURCE_STATE_UNORDERED_ACCESS;
 	}
-	Pools.Empty();
-}
-
-
-void FD3D12TransientMemoryPoolManager::BeginFrame()
-{
-	TextureMemoryStats.Reset();
-	BufferMemoryStats.Reset();
+	check(State != D3D12_RESOURCE_STATE_COMMON);
+	return State;
 }
 
-
-void FD3D12TransientMemoryPoolManager::EndFrame()
+FD3D12TransientHeap::FD3D12TransientHeap(const FRHITransientHeapInitializer& Initializer, FD3D12Adapter* Adapter, FD3D12Device* Device, FRHIGPUMask VisibleNodeMask)
+	: FRHITransientHeap(Initializer)
 {
-	FScopeLock Lock(&CS);
-
-	static const int32 FrameLag = 20;
+	D3D12_HEAP_FLAGS HeapFlags = D3D12_HEAP_FLAG_ALLOW_ALL_BUFFERS_AND_TEXTURES;
 
-	// Trim empty allocators if not used in last n frames
-	FD3D12Adapter* Adapter = GetParentDevice()->GetParentAdapter();
-	FD3D12Fence& FrameFence = Adapter->GetFrameFence();
-	const uint64 CompletedFence = FrameFence.UpdateLastCompletedFence();
-	for (int32 PoolIndex = 0; PoolIndex < Pools.Num(); ++PoolIndex)
+	if (Initializer.Flags != ERHITransientHeapFlags::AllowAll)
 	{
-		FD3D12MemoryPool* MemoryPool = (FD3D12MemoryPool*)Pools[PoolIndex];
-		if (MemoryPool != nullptr && MemoryPool->IsEmpty() && (MemoryPool->GetLastUsedFrameFence() + FrameLag <= CompletedFence))
+		switch (Initializer.Flags)
 		{
-			MemoryPool->Destroy();
-			delete(MemoryPool);
+		case ERHITransientHeapFlags::AllowBuffers:
+			HeapFlags = D3D12_HEAP_FLAG_ALLOW_ONLY_BUFFERS;
+			break;
+
+		case ERHITransientHeapFlags::AllowTextures:
+			HeapFlags = D3D12_HEAP_FLAG_ALLOW_ONLY_NON_RT_DS_TEXTURES;
+			break;
 
-			Pools.RemoveAt(PoolIndex);
-			PoolIndex--;
+		case ERHITransientHeapFlags::AllowRenderTargets:
+			HeapFlags = D3D12_HEAP_FLAG_ALLOW_ONLY_RT_DS_TEXTURES;
+			break;
 		}
 	}
-}
-
 
-FD3D12TransientMemoryPool* FD3D12TransientMemoryPoolManager::GetOrCreateMemoryPool(int16 InPoolIndex, uint32 InMinimumAllocationSize, ERHIPoolResourceTypes InAllocationResourceType)
-{
-	FScopeLock Lock(&CS);
+	D3D12_HEAP_PROPERTIES HeapProperties = CD3DX12_HEAP_PROPERTIES(D3D12_HEAP_TYPE_DEFAULT);
+	HeapProperties.CreationNodeMask = FRHIGPUMask::FromIndex(Device->GetGPUIndex()).GetNative();
+	HeapProperties.VisibleNodeMask = VisibleNodeMask.GetNative();
 
-	FD3D12TransientMemoryPool* MemoryPool = nullptr;
+	D3D12_HEAP_DESC Desc = {};
+	Desc.SizeInBytes = Initializer.Size;
+	Desc.Properties = HeapProperties;
+	Desc.Alignment = 0;
+	Desc.Flags = HeapFlags;
 
-	// Try and find a pool which fits this allocations (sort by size so will take best fit)
-	for (int32 PoolIndex = 0; PoolIndex < Pools.Num(); ++PoolIndex)
+	if (Adapter->IsHeapNotZeroedSupported())
 	{
-		FD3D12TransientMemoryPool* Pool = Pools[PoolIndex];
-		if (Pool->GetPoolSize() >= InMinimumAllocationSize && Pool->IsResourceTypeSupported(InAllocationResourceType))
-		{			
-			MemoryPool = Pool;
-			MemoryPool->SetPoolIndex(InPoolIndex);
-			Pools.RemoveAt(PoolIndex);
-			break;
-		}
+		Desc.Flags |= FD3D12_HEAP_FLAG_CREATE_NOT_ZEROED;
 	}
 
-	// couldn't find a pool?
-	if (MemoryPool == nullptr)
+	ID3D12Heap* D3DHeap = nullptr;
 	{
-		TRACE_CPUPROFILER_EVENT_SCOPE(D3D12RHI::AllocateTransientMemoryPool);
-
-		// Find out the pool size - use default, but if allocation doesn't fit then round up to next power of 2
-		// so it 'always' fits the pool allocator
-		uint32 PoolSize = DefaultPoolSize;
-		if (InMinimumAllocationSize > PoolSize)
-		{
-			check(InMinimumAllocationSize <= MaxAllocationSize);
-			PoolSize = FMath::Min(FMath::RoundUpToPowerOfTwo(InMinimumAllocationSize), (uint32)MaxAllocationSize);
-		}
+		ID3D12Device* D3DDevice = Device->GetDevice();
 
-		// If merged heap types are not supported then create a resource type specific heap
-		FD3D12ResourceInitConfig PoolInitConfig = InitConfig;		
-		if (!bMergedTypeHeapSupported)
-		{
-			switch (InAllocationResourceType)
-			{
-			case ERHIPoolResourceTypes::Buffers:
-			{
-				PoolInitConfig.HeapFlags = D3D12_HEAP_FLAG_ALLOW_ONLY_BUFFERS;
-				break;
-			}
-			case ERHIPoolResourceTypes::RTDSTextures:
-			{
-				PoolInitConfig.HeapFlags = D3D12_HEAP_FLAG_ALLOW_ONLY_RT_DS_TEXTURES;
-				break;
-			}
-			case ERHIPoolResourceTypes::NonRTDSTextures:
-			{
-				PoolInitConfig.HeapFlags = D3D12_HEAP_FLAG_ALLOW_ONLY_NON_RT_DS_TEXTURES;
-				break;
-			}
-			}
-		}
-		ERHIPoolResourceTypes SupportedResourceTypes = bMergedTypeHeapSupported ? ERHIPoolResourceTypes::All : InAllocationResourceType;
+		LLM_PLATFORM_SCOPE(ELLMTag::GraphicsPlatform);
 
-		MemoryPool = new FD3D12TransientMemoryPool(GetParentDevice(), GetVisibilityMask(), PoolInitConfig, SupportedResourceTypes,
-			TEXT("TransientResourceMemoryPool"), EResourceAllocationStrategy::kPlacedResource, InPoolIndex, PoolSize, PoolAlignment);
-		MemoryPool->Init();
+		// We are tracking allocations ourselves, so don't let XMemAlloc track these as well
+		LLM_SCOPED_PAUSE_TRACKING_FOR_TRACKER(ELLMTracker::Default, ELLMAllocType::System);
+		VERIFYD3D12RESULT(D3DDevice->CreateHeap(&Desc, IID_PPV_ARGS(&D3DHeap)));
 
 #if PLATFORM_WINDOWS
 		// Boost priority to make sure it's not paged out
-		ID3D12Device* D3DDevice = GetParentDevice()->GetDevice();
 		TRefCountPtr<ID3D12Device5> D3DDevice5;
 		if (SUCCEEDED(D3DDevice->QueryInterface(IID_PPV_ARGS(D3DDevice5.GetInitReference()))))
 		{
-			ID3D12Pageable* HeapResource = MemoryPool->GetBackingHeap()->GetHeap();
+			ID3D12Pageable* Pageable = D3DHeap;
 			D3D12_RESIDENCY_PRIORITY HeapPriority = D3D12_RESIDENCY_PRIORITY_HIGH;
-			D3DDevice5->SetResidencyPriority(1, &HeapResource, &HeapPriority);
+			D3DDevice5->SetResidencyPriority(1, &Pageable, &HeapPriority);
 		}
 #endif // PLATFORM_WINDOWS
 	}
+	SetName(D3DHeap, L"TransientResourceAllocator Backing Heap");
 
-	return MemoryPool;
+	Heap = new FD3D12Heap(Device, VisibleNodeMask);
+	Heap->SetHeap(D3DHeap);
+	Heap->BeginTrackingResidency(Desc.SizeInBytes);
 }
 
-
-void FD3D12TransientMemoryPoolManager::ReleaseMemoryPool(FD3D12TransientMemoryPool* InMemoryPool)
+FD3D12TransientHeap::~FD3D12TransientHeap()
 {
-	FScopeLock Lock(&CS);
-	InMemoryPool->ResetPool();
-
-	// sorted insert
-	int32 InsertIndex = 0;
-	for (; InsertIndex < Pools.Num(); ++InsertIndex)
-	{
-		// First by size then by last frame usage
-		if (Pools[InsertIndex]->GetPoolSize() >= InMemoryPool->GetPoolSize() &&
-			Pools[InsertIndex]->GetLastUsedFrameFence() <= InMemoryPool->GetLastUsedFrameFence())
-		{
-			break;
-		}
-	}
-	Pools.Insert(InMemoryPool, InsertIndex);
+	LLM_SCOPED_PAUSE_TRACKING_FOR_TRACKER(ELLMTracker::Default, ELLMAllocType::System);
+	Heap->Destroy();
 }
 
-
-FD3D12PooledTextureData FD3D12TransientMemoryPoolManager::GetPooledTexture(const FRHITextureCreateInfo& InCreateInfo, const TCHAR* InDebugName)
+TUniquePtr<FD3D12TransientResourceSystem> FD3D12TransientResourceSystem::Create(FD3D12Adapter* ParentAdapter, FRHIGPUMask VisibleNodeMask)
 {
-	FD3D12PooledTextureData Result;
-	if (GD3D12TransientAllocatorPoolTextures)
-	{
-		uint64 CreateHash = ComputeTextureCreateInfoHash(InCreateInfo);
+	FRHITransientResourceSystemInitializer Initializer = FRHITransientResourceSystemInitializer::CreateDefault();
+	Initializer.HeapAlignment = D3D12_DEFAULT_RESOURCE_PLACEMENT_ALIGNMENT;
 
-		// Try and search the cache for resource of same creation info (inside lock)
-		FScopeLock Lock(&CS);
-		TArray<FD3D12PooledTextureData>* FindResult = FreeTextures.Find(CreateHash);
-		if (FindResult && FindResult->Num() > 0)
-		{
-			// Make sure they are not referenced anymore! (should in theory not happen)
-			//check((*FindResult)[0].RHITexture->GetRefCount() == 1);
-			for (int32 Index = 0; Index < FindResult->Num(); ++Index)
-			{
-				FD3D12PooledTextureData& TextureData = (*FindResult)[Index];
-				if (TextureData.RHITexture->GetRefCount() == 1)
-				{
-					Result = TextureData;
-					check(Result.CreateInfo == InCreateInfo);
-					FindResult->RemoveAt(Index);
-					break;
-				}
-			}
-		}
-	}
+	// Tier2 hardware is able to mix resource types onto the same heap.
+	Initializer.bSupportsAllHeapFlags = ParentAdapter->GetResourceHeapTier() == D3D12_RESOURCE_HEAP_TIER_2;
 
-	return Result;
+	return TUniquePtr<FD3D12TransientResourceSystem>(new FD3D12TransientResourceSystem(Initializer, ParentAdapter, VisibleNodeMask));
 }
 
+FD3D12TransientResourceSystem::FD3D12TransientResourceSystem(const FRHITransientResourceSystemInitializer& Initializer, FD3D12Adapter* ParentAdapter, FRHIGPUMask InVisibleNodeMask)
+	: FRHITransientResourceSystem(Initializer)
+	, FD3D12AdapterChild(ParentAdapter)
+	, VisibleNodeMask(InVisibleNodeMask)
+{}
 
-FD3D12PooledBufferData FD3D12TransientMemoryPoolManager::GetPooledBuffer(const FRHIBufferCreateInfo& InCreateInfo, const TCHAR* InDebugName)
+FRHITransientHeap* FD3D12TransientResourceSystem::CreateHeap(const FRHITransientHeapInitializer& HeapInitializer)
 {
-	FD3D12PooledBufferData Result;
-	if (GD3D12TransientAllocatorPoolBuffers)
+	return GetParentAdapter()->CreateLinkedObject<FD3D12TransientHeap>(VisibleNodeMask, [&](FD3D12Device* Device)
 	{
-		uint64 CreateHash = ComputeBufferCreateInfoHash(InCreateInfo);
-
-		// Try and search the cache for resource of same creation info (inside lock)
-		FScopeLock Lock(&CS);
-		TArray<FD3D12PooledBufferData>* FindResult = FreeBuffers.Find(CreateHash);
-		if (FindResult && FindResult->Num() > 0)
-		{
-			Result = FindResult->Pop(false);
-			check(Result.RHIBuffer->GetRefCount() == 1);
-		}
-	}
-
-	return Result;
+		return new FD3D12TransientHeap(HeapInitializer, GetParentAdapter(), Device, VisibleNodeMask);
+	});
 }
 
+FD3D12TransientResourceAllocator::FD3D12TransientResourceAllocator(FD3D12TransientResourceSystem& InParentSystem)
+	: FD3D12AdapterChild(InParentSystem.GetParentAdapter())
+	, Allocator(InParentSystem)
+	, AllocationInfoQueryDevice(GetParentAdapter()->GetDevice(0))
+{}
 
-void FD3D12TransientMemoryPoolManager::ReleaseResources(FD3D12TransientResourceAllocator* InAllocator)
+FRHITransientTexture* FD3D12TransientResourceAllocator::CreateTexture(const FRHITextureCreateInfo& InCreateInfo, const TCHAR* InDebugName)
 {
-	FScopeLock Lock(&CS);
+	D3D12_RESOURCE_DESC Desc = FD3D12TextureBase::GetResourceDesc(InCreateInfo);
+	D3D12_RESOURCE_ALLOCATION_INFO Info = AllocationInfoQueryDevice->GetResourceAllocationInfo(Desc);
 
-	// release the textures
-	if (GD3D12TransientAllocatorPoolTextures)
+	return Allocator.CreateTexture(InCreateInfo, InDebugName, Info.SizeInBytes, Info.Alignment,
+		[&](const FRHITransientResourceAllocator::FResourceInitializer& Initializer)
 	{
-		for (FD3D12PooledTextureData& TextureData : InAllocator->AllocatedTextures)
-		{
-			uint64 CreateHash = ComputeTextureCreateInfoHash(TextureData.CreateInfo);
-			TArray<FD3D12PooledTextureData>* FindResult = FreeTextures.Find(CreateHash);
-			if (FindResult)
-			{
-				FindResult->Add(TextureData);
-			}
-			else
-			{
-				TArray<FD3D12PooledTextureData> Textures;
-				Textures.Add(TextureData);
-				FreeTextures.Add(CreateHash, Textures);
-			}
-		}
-	}
-	InAllocator->AllocatedTextures.Empty();
+		ERHIAccess InitialState = ERHIAccess::UAVMask;
 
-	// release the buffers
-	if (GD3D12TransientAllocatorPoolBuffers)
-	{
-		for (FD3D12PooledBufferData& BufferData : InAllocator->AllocatedBuffers)
+		if (EnumHasAnyFlags(InCreateInfo.Flags, TexCreate_RenderTargetable))
 		{
-			uint64 CreateHash = ComputeBufferCreateInfoHash(BufferData.CreateInfo);
-			TArray<FD3D12PooledBufferData>* FindResult = FreeBuffers.Find(CreateHash);
-			if (FindResult)
-			{
-				FindResult->Add(BufferData);
-			}
-			else
-			{
-				TArray<FD3D12PooledBufferData> Buffers;
-				Buffers.Add(BufferData);
-				FreeBuffers.Add(CreateHash, Buffers);
-			}
+			InitialState = ERHIAccess::RTV;
 		}
-	}
-	InAllocator->AllocatedBuffers.Empty();
-}
-
-
-void FD3D12TransientMemoryPoolManager::UpdateMemoryStats()
-{
-	FScopeLock Lock(&CS);
-
-	uint32 MemoryAllocated = 0;
-	for (FD3D12TransientMemoryPool* MemoryPool : Pools)
-	{
-		MemoryAllocated += MemoryPool->GetPoolSize();
-	}
-
-	SET_MEMORY_STAT(STAT_D3D12TransientMemoryPoolAllocated, MemoryAllocated);
-	SET_MEMORY_STAT(STAT_D3D12TransientMemoryFramePoolUsed, TextureMemoryStats.MaxFrameAllocated + BufferMemoryStats.MaxFrameAllocated);
-	SET_MEMORY_STAT(STAT_D3D12TransientMemoryFrameCommittedUsed, TextureMemoryStats.CommittedAllocated + BufferMemoryStats.CommittedAllocated);
-	SET_DWORD_STAT(STAT_D3D12TransientMemoryPoolAllocations, TextureMemoryStats.PoolAllocations + BufferMemoryStats.PoolAllocations);
-	SET_DWORD_STAT(STAT_D3D12TransientMemoryCommittedAllocations, TextureMemoryStats.CommittedAllocations + BufferMemoryStats.CommittedAllocations);
-
-	SET_MEMORY_STAT(STAT_D3D12TransientMemoryBufferFrameUsed, BufferMemoryStats.MaxFrameAllocated);
-	SET_MEMORY_STAT(STAT_D3D12TransientMemoryBufferFrameTotalRequested, BufferMemoryStats.TotalRequested);
-	SET_MEMORY_STAT(STAT_D3D12TransientMemoryBufferFrameCommittedAllocated, BufferMemoryStats.CommittedAllocated);
-	SET_DWORD_STAT(STAT_D3D12TransientMemoryBufferPoolAllocations, BufferMemoryStats.PoolAllocations);
-	SET_DWORD_STAT(STAT_D3D12TransientMemoryBufferCommittedAllocations, BufferMemoryStats.CommittedAllocations);
-
-	SET_MEMORY_STAT(STAT_D3D12TransientMemoryTextureFrameUsed, TextureMemoryStats.MaxFrameAllocated);
-	SET_MEMORY_STAT(STAT_D3D12TransientMemoryTextureFrameTotalRequested, TextureMemoryStats.TotalRequested);
-	SET_MEMORY_STAT(STAT_D3D12TransientMemoryTextureFrameCommittedAllocated, TextureMemoryStats.CommittedAllocated);
-	SET_DWORD_STAT(STAT_D3D12TransientMemoryTexturePoolAllocations, TextureMemoryStats.PoolAllocations);
-	SET_DWORD_STAT(STAT_D3D12TransientMemoryTextureCommittedAllocations, TextureMemoryStats.CommittedAllocations);
-}
-
-
-//-----------------------------------------------------------------------------
-//	FD3D12TransientResourceAllocator
-//-----------------------------------------------------------------------------
-
-
-D3D12_RESOURCE_STATES FD3D12TransientResourceAllocator::GetInitialResourceState(const D3D12_RESOURCE_DESC& InDesc)
-{
-	// Validate the creation state
-	D3D12_RESOURCE_STATES State = D3D12_RESOURCE_STATE_COMMON;
-	if (EnumHasAnyFlags(InDesc.Flags, D3D12_RESOURCE_FLAG_ALLOW_RENDER_TARGET))
-	{
-		State = D3D12_RESOURCE_STATE_RENDER_TARGET;
-	}
-	else if (EnumHasAnyFlags(InDesc.Flags, D3D12_RESOURCE_FLAG_ALLOW_DEPTH_STENCIL))
-	{
-		State = D3D12_RESOURCE_STATE_DEPTH_WRITE;
-	}
-	else if (EnumHasAnyFlags(InDesc.Flags, D3D12_RESOURCE_FLAG_ALLOW_UNORDERED_ACCESS))
-	{
-		State = D3D12_RESOURCE_STATE_UNORDERED_ACCESS;
-	}
-	check(State != D3D12_RESOURCE_STATE_COMMON);
-	return State;
-}
-
-
-FD3D12TransientResourceAllocator::FD3D12TransientResourceAllocator(FD3D12TransientMemoryPoolManager& InMemoryPoolManager) :
-	FD3D12PoolAllocator(
-		InMemoryPoolManager.GetParentDevice(),
-		InMemoryPoolManager.GetVisibilityMask(),
-		InMemoryPoolManager.GetInitConfig(),
-		TEXT("FD3D12TransientResourceAllocator"),
-		EResourceAllocationStrategy::kPlacedResource,
-		InMemoryPoolManager.GetDefaultPoolSize(),
-		InMemoryPoolManager.GetPoolAlignment(),
-		InMemoryPoolManager.GetMaxAllocationSize(),
-		FRHIMemoryPool::EFreeListOrder::SortByOffset,
-		false /*defrag*/)
-{ 
-}
-
-
-FD3D12TransientResourceAllocator::~FD3D12TransientResourceAllocator()
-{
-	// release all resources back to pool so they can be reused if enabled (pools have already been freed)
-	FD3D12TransientMemoryPoolManager& PoolManager = GetParentDevice()->GetTransientMemoryPoolManager();
-	PoolManager.ReleaseResources(this);
-
-	PoolManager.UpdateTextureStats(TextureMemoryStats);
-	PoolManager.UpdateBufferStats(BufferMemoryStats);
-}
-
-
-FRHITexture* FD3D12TransientResourceAllocator::CreateTexture(const FRHITextureCreateInfo& InCreateInfo, const TCHAR* InDebugName)
-{
-	TRACE_CPUPROFILER_EVENT_SCOPE(D3D12RHI::CreateTransientTexture);
-
-	// Try and get a pooled texture from the manager
-	FD3D12TransientMemoryPoolManager& PoolManager = GetParentDevice()->GetTransientMemoryPoolManager();
-	FD3D12PooledTextureData TextureData = PoolManager.GetPooledTexture(InCreateInfo, InDebugName);
-
-	FD3D12BaseShaderResource* BaseShaderResource = nullptr;
-
-	ERHIAccess InitialState = ERHIAccess::UAVMask;
-	if (EnumHasAnyFlags(InCreateInfo.Flags, TexCreate_RenderTargetable))
-	{
-		InitialState = ERHIAccess::RTV;
-	}
-	else if (EnumHasAnyFlags(InCreateInfo.Flags, TexCreate_DepthStencilTargetable))
-	{
-		InitialState = ERHIAccess::DSVWrite;
-	}
-
-	// Found something?
-	if (TextureData.RHITexture)
-	{
-		TRACE_CPUPROFILER_EVENT_SCOPE(SetupPoolTransientResource);
-
-		// Get the base shader resource
-		BaseShaderResource = GetBaseShaderResource(TextureData.RHITexture);
-		check(BaseShaderResource);
-		check(!BaseShaderResource->ResourceLocation.IsValid());
-
-		// Setup the optional clear value
-		D3D12_CLEAR_VALUE* ClearValuePtr = nullptr;
-		if (TextureData.ClearValue.Format != DXGI_FORMAT_UNKNOWN)
+		else if (EnumHasAnyFlags(InCreateInfo.Flags, TexCreate_DepthStencilTargetable))
 		{
-			ClearValuePtr = &TextureData.ClearValue;
+			InitialState = ERHIAccess::DSVWrite;
 		}
 
-		D3D12_RESOURCE_STATES CreateState = GetD3D12ResourceState(InitialState, false);
-
-		// Allocate a new FD3D12Resource on the resource location
-		AllocateTexture(D3D12_HEAP_TYPE_DEFAULT, TextureData.ResourceDesc, InCreateInfo.Format, ED3D12ResourceStateMode::MultiState, CreateState, ClearValuePtr, InDebugName, BaseShaderResource->ResourceLocation);
-
-		// Inform the listeners about the change - should ideally still get from cache
-		{
-			TRACE_CPUPROFILER_EVENT_SCOPE(RenameViews);
-			BaseShaderResource->ResourceRenamed(&BaseShaderResource->ResourceLocation);
-		}
-	}
-	else
-	{
-		// Create a new resource
-		FRHIResourceCreateInfo CreateInfo(InDebugName);
-		CreateInfo.ClearValueBinding = InCreateInfo.ClearValue;
+		FRHIResourceCreateInfo ResourceCreateInfo(InDebugName, InCreateInfo.ClearValue);
+		FResourceAllocatorAdapter ResourceAllocatorAdapter(GetParentAdapter(), static_cast<FD3D12TransientHeap&>(Initializer.Heap), Initializer.Allocation);
+		FRHITexture* Texture = nullptr;
 
-		// Force transient flag
-		ETextureCreateFlags CreationFlags = InCreateInfo.Flags;
-		CreationFlags |= TexCreate_Transient;
+		const bool bTextureArray = InCreateInfo.ArraySize > 1;
 
 		switch (InCreateInfo.Dimension)
 		{
 		case ETextureDimension::Texture2D:
 		{
-			const bool bTextureArray = false;
 			const bool bCubeTexture = false;
-			ID3D12ResourceAllocator* ResourceAllocator = this;
-			FD3D12Texture2D* Texture2D = FD3D12DynamicRHI::GetD3DRHI()->CreateD3D12Texture2D<FD3D12BaseTexture2D>(nullptr, InCreateInfo.Extent.X, InCreateInfo.Extent.Y, 1, bTextureArray, bCubeTexture, InCreateInfo.Format, InCreateInfo.NumMips, InCreateInfo.NumSamples, CreationFlags, InitialState, CreateInfo, ResourceAllocator);
-			TextureData.RHITexture = Texture2D;
-			TextureData.ResourceDesc = Texture2D->GetResource()->GetDesc();
-			TextureData.ClearValue = Texture2D->GetResource()->GetClearValue();
-			TextureData.CreateInfo = InCreateInfo;
-
-			BaseShaderResource = Texture2D;
-			break;
+			Texture = FD3D12DynamicRHI::GetD3DRHI()->CreateD3D12Texture2D<FD3D12BaseTexture2D>(nullptr, InCreateInfo.Extent.X, InCreateInfo.Extent.Y, 1, bTextureArray, bCubeTexture, InCreateInfo.Format, InCreateInfo.NumMips, InCreateInfo.NumSamples, InCreateInfo.Flags, InitialState, ResourceCreateInfo, &ResourceAllocatorAdapter);
 		}
+		break;
 		case ETextureDimension::Texture3D:
 		{
-			// Only support 2d textures for now
-			ID3D12ResourceAllocator* ResourceAllocator = this;
-			FD3D12Texture3D* Texture3D = FD3D12DynamicRHI::GetD3DRHI()->CreateD3D12Texture3D(nullptr, InCreateInfo.Extent.X, InCreateInfo.Extent.Y, InCreateInfo.Depth, InCreateInfo.Format, InCreateInfo.NumMips, CreationFlags, InitialState, CreateInfo, ResourceAllocator);
-			TextureData.RHITexture = Texture3D;
-			TextureData.ResourceDesc = Texture3D->GetResource()->GetDesc();
-			TextureData.ClearValue = Texture3D->GetResource()->GetClearValue();
-			TextureData.CreateInfo = InCreateInfo;
-
-			BaseShaderResource = Texture3D;
-			break;
+			Texture = FD3D12DynamicRHI::GetD3DRHI()->CreateD3D12Texture3D(nullptr, InCreateInfo.Extent.X, InCreateInfo.Extent.Y, InCreateInfo.Depth, InCreateInfo.Format, InCreateInfo.NumMips, InCreateInfo.Flags, InitialState, ResourceCreateInfo, &ResourceAllocatorAdapter);
 		}
+		break;
 		case ETextureDimension::TextureCube:
 		{
-			check(InCreateInfo.Extent.X == InCreateInfo.Extent.Y);
-
-			const bool bTextureArray = false;
 			const bool bCubeTexture = true;
-			ID3D12ResourceAllocator* ResourceAllocator = this;
-			FD3D12TextureCube* TextureCube = FD3D12DynamicRHI::GetD3DRHI()->CreateD3D12Texture2D<FD3D12BaseTextureCube>(nullptr, InCreateInfo.Extent.X, InCreateInfo.Extent.Y, 6, bTextureArray, bCubeTexture, InCreateInfo.Format, InCreateInfo.NumMips, InCreateInfo.NumSamples, CreationFlags, InitialState, CreateInfo, ResourceAllocator);
-			TextureData.RHITexture = TextureCube;
-			TextureData.ResourceDesc = TextureCube->GetResource()->GetDesc();
-			TextureData.ClearValue = TextureCube->GetResource()->GetClearValue();
-			TextureData.CreateInfo = InCreateInfo;
-
-			BaseShaderResource = TextureCube;
-			break;
-		}
-		default:
-		{
-			// Only support 2d, 3d & cube textures for now
-			check(false);
-			break;
+			check(InCreateInfo.Extent.X == InCreateInfo.Extent.Y);
+			Texture = FD3D12DynamicRHI::GetD3DRHI()->CreateD3D12Texture2D<FD3D12BaseTextureCube>(nullptr, InCreateInfo.Extent.X, InCreateInfo.Extent.Y, 6, bTextureArray, bCubeTexture, InCreateInfo.Format, InCreateInfo.NumMips, InCreateInfo.NumSamples, InCreateInfo.Flags, InitialState, ResourceCreateInfo, &ResourceAllocatorAdapter);
 		}
+		break;
 		}
-	}
 
-	SetupAllocatedResource(BaseShaderResource, TextureMemoryStats);
-
-	// keep track of all allocated textures
-	AllocatedTextures.Add(TextureData);
+		// The D3D12_RESOURCE_DESC's are built in two different functions right now. This checks that they actually match what we expect.
+#if DO_CHECK
+		{
+			CD3DX12_RESOURCE_DESC CreatedDesc(GetD3D12TextureFromRHITexture(Texture)->GetResource()->GetDesc());
+			CD3DX12_RESOURCE_DESC DerivedDesc(Desc);
+			check(CreatedDesc == DerivedDesc);
+		}
+#endif
 
-	return TextureData.RHITexture;
+		return new FRHITransientTexture(Texture, Initializer.Hash, InCreateInfo);
+	});
 }
 
-
-FRHIBuffer* FD3D12TransientResourceAllocator::CreateBuffer(const FRHIBufferCreateInfo& InCreateInfo, const TCHAR* InDebugName)
+void FD3D12TransientResourceAllocator::FResourceAllocatorAdapter::AllocateResource(
+	uint32 GPUIndex, D3D12_HEAP_TYPE, const D3D12_RESOURCE_DESC& InDesc, uint64 InSize, uint32, ED3D12ResourceStateMode InResourceStateMode,
+	D3D12_RESOURCE_STATES InCreateState, const D3D12_CLEAR_VALUE* InClearValue, const TCHAR* InName, FD3D12ResourceLocation& ResourceLocation)
 {
-	TRACE_CPUPROFILER_EVENT_SCOPE(D3D12RHI::CreateTransientBuffer);
-
-	// Try and get a pooled texture from the manager
-	FD3D12TransientMemoryPoolManager& PoolManager = GetParentDevice()->GetTransientMemoryPoolManager();
-	FD3D12PooledBufferData BufferData = PoolManager.GetPooledBuffer(InCreateInfo, InDebugName);
-	ERHIAccess InitialState = ERHIAccess::UAVMask;
-
-	FD3D12Buffer* Buffer = nullptr;
-
-	// Force transient flag
-	EBufferUsageFlags UsageFlags = InCreateInfo.Usage;
-	UsageFlags |= BUF_Transient;
-
-	// Found something?
-	if (BufferData.RHIBuffer)
-	{
-		// Get the D3D12 Buffer
-		Buffer = FD3D12DynamicRHI::ResourceCast((FRHIBuffer*)BufferData.RHIBuffer);
-		check(!Buffer->ResourceLocation.IsValid());
-
-		D3D12_RESOURCE_DESC Desc;
-		uint32 Alignment;
-		FD3D12Buffer::GetResourceDescAndAlignment(InCreateInfo.Size, InCreateInfo.Stride, UsageFlags, Desc, Alignment);
+	FD3D12Resource* NewResource = nullptr;
+	VERIFYD3D12RESULT(GetParentAdapter()->CreatePlacedResource(InDesc, Heap.GetLinkedObject(GPUIndex)->Get(), Allocation.Offset, InCreateState, InResourceStateMode, D3D12_RESOURCE_STATE_TBD, InClearValue, &NewResource, InName));
 
-		D3D12_RESOURCE_STATES CreateState = GetD3D12ResourceState(InitialState, false);
-
-		// Allocate a new FD3D12Resource on the resource location
-		AllocateResource(D3D12_HEAP_TYPE_DEFAULT, Desc, Desc.Width, Alignment, ED3D12ResourceStateMode::MultiState, CreateState, nullptr, InDebugName, Buffer->ResourceLocation);
-
-		// Inform the listeners about the change (should be empty?)
-		Buffer->ResourceRenamed(&Buffer->ResourceLocation);
-	}
-	else
-	{
-		FRHIResourceCreateInfo CreateInfo(InDebugName);
-		ID3D12ResourceAllocator* ResourceAllocator = this;
-		Buffer = FD3D12DynamicRHI::GetD3DRHI()->CreateD3D12Buffer(nullptr, InCreateInfo.Size, UsageFlags, InCreateInfo.Stride, InitialState, CreateInfo, ResourceAllocator);
-		BufferData.RHIBuffer = Buffer;
-		BufferData.CreateInfo = InCreateInfo;
-	}
-
-	SetupAllocatedResource((FD3D12BaseShaderResource*)Buffer, BufferMemoryStats);
-	
-	// keep track of all allocated buffers
-	AllocatedBuffers.Add(BufferData);
-
-	return BufferData.RHIBuffer;
-}
-
-
-void FD3D12TransientResourceAllocator::DeallocateMemory(FRHITexture* InTexture)
-{
-	// Get the correct base shader resource taken texture type into account
-	FD3D12BaseShaderResource* BaseShaderResource = GetBaseShaderResource(InTexture);
-	check(BaseShaderResource);
-	DeallocateMemory(BaseShaderResource, TextureMemoryStats);
+	check(!ResourceLocation.IsValid());
+	ResourceLocation.AsHeapAliased(NewResource);
+	ResourceLocation.SetSize(InSize);
 }
 
-
-void FD3D12TransientResourceAllocator::DeallocateMemory(FRHIBuffer* InBuffer)
+FRHITransientBuffer* FD3D12TransientResourceAllocator::CreateBuffer(const FRHIBufferCreateInfo& InCreateInfo, const TCHAR* InDebugName)
 {
-	// cast to d3d12 object and call shared function
-	FD3D12Buffer* Buffer = FD3D12DynamicRHI::ResourceCast(InBuffer);
-	DeallocateMemory(Buffer, BufferMemoryStats);
-}
-
+	D3D12_RESOURCE_DESC Desc;
+	uint32 Alignment;
+	EBufferUsageFlags BufferUsage = InCreateInfo.Usage;
+	FD3D12Buffer::GetResourceDescAndAlignment(InCreateInfo.Size, InCreateInfo.Stride, BufferUsage, Desc, Alignment);
 
-FD3D12BaseShaderResource* FD3D12TransientResourceAllocator::GetBaseShaderResource(FRHITexture* InRHITexture)
-{
-	FD3D12BaseShaderResource* BaseShaderResource = nullptr;
-	if (FRHITexture2D* RHITexture2D = InRHITexture->GetTexture2D())
-	{
-		FD3D12Texture2D* Texture2D = FD3D12DynamicRHI::ResourceCast(RHITexture2D);
-		BaseShaderResource = Texture2D->GetBaseShaderResource();
-	}
-	else if (FRHITexture3D* RHITexture3D = InRHITexture->GetTexture3D())
-	{
-		FD3D12Texture3D* Texture3D = FD3D12DynamicRHI::ResourceCast(RHITexture3D);
-		BaseShaderResource = Texture3D;
-	}
-	else if (FRHITextureCube* RHITextureCube = InRHITexture->GetTextureCube())
+	return Allocator.CreateBuffer(InCreateInfo, InDebugName, Desc.Width, Alignment,
+		[&](const FRHITransientResourceAllocator::FResourceInitializer& Initializer)
 	{
-		FD3D12TextureCube* TextureCube = FD3D12DynamicRHI::ResourceCast(RHITextureCube);
-		BaseShaderResource = TextureCube;
-	}
-	check(BaseShaderResource);
-	return BaseShaderResource;
-}
-
+		FResourceAllocatorAdapter ResourceAllocatorAdapter(GetParentAdapter(), static_cast<FD3D12TransientHeap&>(Initializer.Heap), Initializer.Allocation);
+		FRHIResourceCreateInfo ResourceCreateInfo(InDebugName);
 
-void FD3D12TransientResourceAllocator::SetupAllocatedResource(FD3D12BaseShaderResource* InBaseShaderResource, FD3D12TransientMemoryStats& Stats)
-{
-	uint64 AllocationSize = InBaseShaderResource->ResourceLocation.GetSize();
+		FD3D12Buffer* Buffer = FD3D12DynamicRHI::GetD3DRHI()->CreateD3D12Buffer(nullptr, InCreateInfo.Size, BufferUsage, InCreateInfo.Stride, ERHIAccess::UAVMask, ResourceCreateInfo, &ResourceAllocatorAdapter);
 
-	// If pool allocated then collected the overlapping resources as well
-	if (InBaseShaderResource->ResourceLocation.GetAllocatorType() == FD3D12ResourceLocation::AT_Pool)
-	{
-		FScopeLock Lock(&CS);
-				
-		FRHIPoolAllocationData& PoolAllocationData = InBaseShaderResource->ResourceLocation.GetPoolAllocatorPrivateData().PoolData;
-		FD3D12TransientMemoryPool* MemoryPool = ((FD3D12TransientMemoryPool*)Pools[PoolAllocationData.GetPoolIndex()]);
-		FInt64Range AllocationRange(PoolAllocationData.GetOffset(), PoolAllocationData.GetOffset() + PoolAllocationData.GetSize());
-		TArray<FD3D12Resource*> OverlappingResources;
-		MemoryPool->CheckActiveResources(AllocationRange, OverlappingResources);
-		if (OverlappingResources.Num() > 0)
+		// The D3D12_RESOURCE_DESC's are built in two different functions right now. This checks that they actually match what we expect.
+#if DO_CHECK
 		{
-			OverlappingResourceData.Add(InBaseShaderResource, OverlappingResources);
+			CD3DX12_RESOURCE_DESC CreatedDesc(Buffer->GetResource()->GetDesc());
+			CD3DX12_RESOURCE_DESC DerivedDesc(Desc);
+			check(CreatedDesc == DerivedDesc);
 		}
-		
-		// update stats
-		Stats.PoolAllocations++;
-		Stats.CurrentPoolAllocated += AllocationSize;
-		Stats.TotalRequested += AllocationSize;
-		Stats.MaxFrameAllocated = FMath::Max(Stats.MaxFrameAllocated, Stats.CurrentPoolAllocated);
-	}
-	else
-	{
-		// Outside of pool size? Expensive operation if this happens every frame - should this not be allowed?		
-
-		// update stats
-		Stats.CommittedAllocations++;
-		Stats.CommittedAllocated += AllocationSize;
-		Stats.TotalRequested += AllocationSize;
-	}
-
-}
-
-
-void FD3D12TransientResourceAllocator::DeallocateMemory(FD3D12BaseShaderResource* InBaseShaderResource, FD3D12TransientMemoryStats& Stats)
-{
-	// If pool allocated then release the resource back to the pool
-	if (InBaseShaderResource->ResourceLocation.GetAllocatorType() == FD3D12ResourceLocation::AT_Pool)
-	{
-		FScopeLock Lock(&CS);
-
-		// update stats
-		uint64 AllocationSize = InBaseShaderResource->ResourceLocation.GetSize();
-		Stats.CurrentPoolAllocated -= AllocationSize;
-
-		FD3D12Adapter* Adapter = GetParentDevice()->GetParentAdapter();
-		FD3D12Fence& FrameFence = Adapter->GetFrameFence();
-
-		FRHIPoolAllocationData& PoolAllocationData = InBaseShaderResource->ResourceLocation.GetPoolAllocatorPrivateData().PoolData;
-		FD3D12TransientMemoryPool* MemoryPool = ((FD3D12TransientMemoryPool*)Pools[PoolAllocationData.GetPoolIndex()]);
-		MemoryPool->ReleaseResource(InBaseShaderResource->ResourceLocation.GetResource(), PoolAllocationData, FrameFence.GetCurrentFence());
-	}
-}
-
+#endif
 
-void FD3D12TransientResourceAllocator::Freeze(FRHICommandListImmediate&)
-{
-	// all memory should have been freed again - so release back to the manager
-	// can already be reused by the next transient allocator
-	// Resources are kept alive until the allocator is destroyed
-	FD3D12TransientMemoryPoolManager& PoolManager = GetParentDevice()->GetTransientMemoryPoolManager();
-	for (FRHIMemoryPool* RHIPool : Pools)
-	{
-		FD3D12TransientMemoryPool* MemoryPool = ((FD3D12TransientMemoryPool*)RHIPool);
-		check(MemoryPool->GetAllocatedBlocks() == 0);
-		PoolManager.ReleaseMemoryPool(MemoryPool);
-	}
-	Pools.Empty();
-}
-
-
-FRHIMemoryPool* FD3D12TransientResourceAllocator::CreateNewPool(int16 InPoolIndex, uint32 InMinimumAllocationSize, ERHIPoolResourceTypes InAllocationResourceType)
-{
-	// Get pool from manager - don't reallocate each time
-	return GetParentDevice()->GetTransientMemoryPoolManager().GetOrCreateMemoryPool(InPoolIndex, InMinimumAllocationSize, InAllocationResourceType);
+		return new FRHITransientBuffer(Buffer, Initializer.Hash, InCreateInfo);
+	});
 }
 
-
-FD3D12Resource* FD3D12TransientResourceAllocator::CreatePlacedResource(const FRHIPoolAllocationData& InAllocationData, const D3D12_RESOURCE_DESC& InDesc, D3D12_RESOURCE_STATES InCreateState, ED3D12ResourceStateMode InResourceStateMode, const D3D12_CLEAR_VALUE* InClearValue, const TCHAR* InName)
+void FD3D12TransientResourceAllocator::DeallocateMemory(FRHITransientTexture* InTexture)
 {
-	FD3D12Resource* Resource = nullptr;
-
-	// Try and find a cached resource at given offset and creation flags
-	{
-		FD3D12TransientMemoryPool* MemoryPool = ((FD3D12TransientMemoryPool*)Pools[InAllocationData.GetPoolIndex()]);
-		Resource = MemoryPool->FindResourceInCache(InAllocationData.GetOffset(), InDesc, InClearValue, InName);
-	}
-
-	// No cached resource then use base class
-	if (Resource == nullptr)
-	{
-		// Validate the creation state
-		check(InCreateState == GetInitialResourceState(InDesc));
-
-		TRACE_CPUPROFILER_EVENT_SCOPE(D3D12RHI::CreatePlacedResource);
-		Resource = FD3D12PoolAllocator::CreatePlacedResource(InAllocationData, InDesc, InCreateState, InResourceStateMode, InClearValue, InName);
-	}
-
-	return Resource;
+	Allocator.DeallocateMemory(InTexture);
 }
 
-
-void FD3D12TransientResourceAllocator::DeallocateResource(FD3D12ResourceLocation& ResourceLocation)
+void FD3D12TransientResourceAllocator::DeallocateMemory(FRHITransientBuffer* InBuffer)
 {
-	check(IsOwner(ResourceLocation));
-
-	// Don't touch the allocation data - it's probably already freed via call to DeallocateMemory
-	// On clear the data on the resource location itself
-	ResourceLocation.ClearAllocator();
+	Allocator.DeallocateMemory(InBuffer);
 }
 
-
-TArrayView<FD3D12Resource*> FD3D12TransientResourceAllocator::GetOverlappingResources(FD3D12BaseShaderResource* InBaseShaderResource)
+void FD3D12TransientResourceAllocator::Freeze(FRHICommandListImmediate& RHICmdList)
 {
-	check(IsOwner(InBaseShaderResource->ResourceLocation));
-
-	TArray<FD3D12Resource*>* FindResult = OverlappingResourceData.Find(InBaseShaderResource);
-	if (FindResult)
-	{
-		return TArrayView<FD3D12Resource*>(*FindResult);
-	}
-	else
-	{
-		return TArrayView<FD3D12Resource*>();
-	}
+	Allocator.Freeze(RHICmdList);
 }
\ No newline at end of file
diff --git a/Engine/Source/Runtime/D3D12RHI/Private/D3D12TransientResourceAllocator.h b/Engine/Source/Runtime/D3D12RHI/Private/D3D12TransientResourceAllocator.h
--- a/Engine/Source/Runtime/D3D12RHI/Private/D3D12TransientResourceAllocator.h
+++ b/Engine/Source/Runtime/D3D12RHI/Private/D3D12TransientResourceAllocator.h
@@ -1,233 +1,77 @@
 // Copyright Epic Games, Inc. All Rights Reserved.
-#pragma once
-
-#include "D3D12PoolAllocator.h"
-
-// Forward declare
-class FD3D12TransientResourceAllocator;
 
-/**
-@brief Tracking stats for transient resource allocations
-**/
-struct FD3D12TransientMemoryStats
-{
-	void Reset()
-	{
-		CurrentPoolAllocated = 0;
-		TotalRequested = 0;
-		MaxFrameAllocated = 0;
-		CommittedAllocated = 0;
-		PoolAllocations = 0;
-		CommittedAllocations = 0;
-	}
-	void Add(const FD3D12TransientMemoryStats& InOther)
-	{
-		CurrentPoolAllocated += InOther.CurrentPoolAllocated;
-		TotalRequested += InOther.TotalRequested;
-		MaxFrameAllocated += InOther.MaxFrameAllocated;
-		CommittedAllocated += InOther.CommittedAllocated;
-		PoolAllocations += InOther.PoolAllocations;
-		CommittedAllocations += InOther.CommittedAllocations;
-	}
-
-	uint64 CurrentPoolAllocated = 0;
-	uint64 TotalRequested = 0;
-	uint64 MaxFrameAllocated = 0;
-	uint64 CommittedAllocated = 0;
-	uint32 PoolAllocations = 0;
-	uint32 CommittedAllocations = 0;
-};
-
-
-struct FD3D12PooledTextureData
-{
-	FRHITextureCreateInfo CreateInfo;
-	D3D12_RESOURCE_DESC ResourceDesc;
-	D3D12_CLEAR_VALUE ClearValue;
-	FTextureRHIRef RHITexture;
-};
+#pragma once
 
-struct FD3D12PooledBufferData
-{
-	FRHIBufferCreateInfo CreateInfo;
-	FBufferRHIRef RHIBuffer;
-};
+#include "D3D12Resources.h"
+#include "RHICoreTransientResourceAllocator.h"
 
+extern D3D12_RESOURCE_STATES GetInitialResourceState(const D3D12_RESOURCE_DESC& InDesc);
 
-/**
-@brief D3D12 transient pool specific implementation of the FD3D12MemoryPool
-	   Caches previously created resource to remove placed resource creation overhead and 
-	   also keeps track of all open allocation ranges used to create the aliasing barriers
-**/
-class FD3D12TransientMemoryPool : public FD3D12MemoryPool
+class FD3D12TransientHeap final
+	: public FRHITransientHeap
+	, public FRefCountBase
+	, public FD3D12LinkedAdapterObject<FD3D12TransientHeap>
 {
 public:
+	FD3D12TransientHeap(const FRHITransientHeapInitializer& Initializer, FD3D12Adapter* Adapter, FD3D12Device* Device, FRHIGPUMask VisibleNodeMask);
+	~FD3D12TransientHeap();
 
-	FD3D12TransientMemoryPool(FD3D12Device* ParentDevice, FRHIGPUMask VisibleNodes, const FD3D12ResourceInitConfig& InInitConfig, ERHIPoolResourceTypes InSupportedResourceTypes, const FString& Name,
-		EResourceAllocationStrategy InAllocationStrategy, int16 InPoolIndex, uint64 InPoolSize, uint32 InPoolAlignment) :
-		FD3D12MemoryPool(ParentDevice, VisibleNodes, InInitConfig, Name, InAllocationStrategy, InPoolIndex, InPoolSize, InPoolAlignment, InSupportedResourceTypes, FRHIMemoryPool::EFreeListOrder::SortByOffset) { }
-	virtual ~FD3D12TransientMemoryPool();
+	FD3D12Heap* Get() { return Heap; }
 
-	void ResetPool();
-	void SetPoolIndex(int16 InNewPoolIndex);
-
-	void ClearActiveResources();
-	void CheckActiveResources(const FInt64Range& InAllocationRange, TArray<FD3D12Resource*>& OverlappingResources);
-
-	FD3D12Resource* FindResourceInCache(uint64 InAllocationOffset, const D3D12_RESOURCE_DESC& InDesc, const D3D12_CLEAR_VALUE* InClearValue, const TCHAR* InName);
-	void ReleaseResource(FD3D12Resource* InResource, FRHIPoolAllocationData& InReleasedAllocationData, uint64 InFenceValue);
-		
 private:
-
-	// Track all active ranges on an allocation which require possible aliasing barriers
-	struct ActiveResourceData
-	{
-		FInt64Range AllocationRange;
-		FInt64Range ActiveRange;
-		FD3D12Resource* Resource;
-	};
-	TArray<ActiveResourceData> ActiveResources;
-
-	// Cache all previously created placed resources at given offset and creation params to reduce
-	// placed resource creation overhead
-	struct FResourceCreateState
-	{
-		// Key values used to create the resource
-		uint64 AllocationOffset;
-		D3D12_RESOURCE_DESC ResourceDesc;
-		D3D12_CLEAR_VALUE ClearValue;
-
-		uint64 GetHash() const;
-	};
-	TMap<uint64, FD3D12Resource*> CachedResourceMap;
+	TRefCountPtr<FD3D12Heap> Heap;
 };
 
-
-//-----------------------------------------------------------------------------
-//	Transient Memory Pool Manager
-//-----------------------------------------------------------------------------
-
-/**
-@brief Manages an array of pools per device which can be reused by the temporary
-	   transient memory allocator - the pools are returned to the manager when the 
-	   allocator is destroyed, so the heap creation and placed resource caches are kept
-	   alive in between frames. Pools are deleted when not used for n amount of frames.
-**/
-class FD3D12TransientMemoryPoolManager : public FD3D12DeviceChild, public FD3D12MultiNodeGPUObject
+class FD3D12TransientResourceSystem final
+	: public FRHITransientResourceSystem
+	, public FD3D12AdapterChild
 {
 public:
-	FD3D12TransientMemoryPoolManager(FD3D12Device* InParent, FRHIGPUMask VisibleNodes);
-	~FD3D12TransientMemoryPoolManager() { }
-
-	void Destroy();
-
-	void BeginFrame();
-	void EndFrame();
-
-	// Get or creation functions from cached resources
-	FD3D12TransientMemoryPool* GetOrCreateMemoryPool(int16 InPoolIndex, uint32 InMinimumAllocationSize, ERHIPoolResourceTypes InAllocationResourceType);
-	void ReleaseMemoryPool(FD3D12TransientMemoryPool* InMemoryPool);
-	FD3D12PooledTextureData GetPooledTexture(const FRHITextureCreateInfo& InCreateInfo, const TCHAR* InDebugName);
-	FD3D12PooledBufferData GetPooledBuffer(const FRHIBufferCreateInfo& InCreateInfo, const TCHAR* InDebugName);
-	void ReleaseResources(FD3D12TransientResourceAllocator* InAllocator);
+	static TUniquePtr<FD3D12TransientResourceSystem> Create(FD3D12Adapter* ParentAdapter, FRHIGPUMask VisibleNodeMask);
 
-	// Get the pool creation members
-	const FD3D12ResourceInitConfig& GetInitConfig() const { return InitConfig; }
-	uint64 GetDefaultPoolSize() const { return DefaultPoolSize; }
-	uint32 GetPoolAlignment() const { return PoolAlignment; }
-	uint64 GetMaxAllocationSize() const { return MaxAllocationSize; }
+	//! FRHITransientResourceSystem Overrides
+	FRHITransientHeap* CreateHeap(const FRHITransientHeapInitializer& Initializer) override;
 
-	// Stat management
-	void UpdateTextureStats(const FD3D12TransientMemoryStats& InStats)
-	{
-		TextureMemoryStats.Add(InStats);
-	}
-	void UpdateBufferStats(const FD3D12TransientMemoryStats& InStats)
-	{
-		BufferMemoryStats.Add(InStats);
-	}
-	void UpdateMemoryStats();
-	
 private:
+	FD3D12TransientResourceSystem(const FRHITransientResourceSystemInitializer& Initializer, FD3D12Adapter* ParentAdapter, FRHIGPUMask VisibleNodeMask);
 
-	// Pool creation members
-	FD3D12ResourceInitConfig InitConfig;
-	uint64 DefaultPoolSize;
-	uint32 PoolAlignment;
-	uint64 MaxAllocationSize;
-	bool bMergedTypeHeapSupported;
-
-	// Critical section to lock access to the pools & resource caches
-	FCriticalSection CS;
-	TArray<FD3D12TransientMemoryPool*> Pools;
-	TMap<uint64, TArray<FD3D12PooledTextureData>> FreeTextures;
-	TMap<uint64, TArray<FD3D12PooledBufferData>> FreeBuffers;
-
-	// Stats of last frame
-	FD3D12TransientMemoryStats TextureMemoryStats;
-	FD3D12TransientMemoryStats BufferMemoryStats;
+	FRHIGPUMask VisibleNodeMask;
 };
 
-
-//-----------------------------------------------------------------------------
-//	Transient Resource Allocator
-//-----------------------------------------------------------------------------
-
-/**
-@brief D3D12 Specific implementation of the IRHITransientResourceAllocator
-	   Uses a pool allocator as base to manage the internal D3D12 memory pools and also takes care of the actual
-	   D3D12 placed resource creation.
-	   Pools are retrieved from and released to the FD3D12TransientMemoryPoolManager
-**/
-class FD3D12TransientResourceAllocator : public IRHITransientResourceAllocator, public FD3D12PoolAllocator
+class FD3D12TransientResourceAllocator final
+	: public IRHITransientResourceAllocator
+	, public FD3D12AdapterChild
 {
 public:
+	FD3D12TransientResourceAllocator(FD3D12TransientResourceSystem& InParentSystem);
 
-	FD3D12TransientResourceAllocator(FD3D12TransientMemoryPoolManager& InMemoryPoolManager);
-	virtual ~FD3D12TransientResourceAllocator();
-
-	// Implementation of FRHITransientResourceAllocator interface
-	virtual FRHITexture* CreateTexture(const FRHITextureCreateInfo& InCreateInfo, const TCHAR* InDebugName) override;
-	virtual FRHIBuffer* CreateBuffer(const FRHIBufferCreateInfo& InCreateInfo, const TCHAR* InDebugName) override;
-	virtual void DeallocateMemory(FRHITexture* InTexture) override;
-	virtual void DeallocateMemory(FRHIBuffer* InBuffer) override;
-	virtual void Freeze(FRHICommandListImmediate&) override;
-
-	// Override deallocation function - memory possibly already freed during DeallocateMemory
-	virtual void DeallocateResource(FD3D12ResourceLocation& ResourceLocation) override;
-
-	// Try and find overlapping resource data for given base resource
-	TArrayView<FD3D12Resource*> GetOverlappingResources(FD3D12BaseShaderResource* InBaseShaderResource);
-
-	// Get the initial resource state
-	static D3D12_RESOURCE_STATES GetInitialResourceState(const D3D12_RESOURCE_DESC& InDesc);
+	//! IRHITransientResourceAllocator Overrides
+	FRHITransientTexture* CreateTexture(const FRHITextureCreateInfo& InCreateInfo, const TCHAR* InDebugName) override;
+	FRHITransientBuffer* CreateBuffer(const FRHIBufferCreateInfo& InCreateInfo, const TCHAR* InDebugName) override;
+	void DeallocateMemory(FRHITransientTexture* InTexture) override;
+	void DeallocateMemory(FRHITransientBuffer* InBuffer) override;
+	void Freeze(FRHICommandListImmediate& RHICmdList) override;
 
 private:
 
-	// Override placed resource allocation helper function to get from cache if available
-	virtual FD3D12Resource* CreatePlacedResource(const FRHIPoolAllocationData& InAllocationData, const D3D12_RESOURCE_DESC& InDesc, D3D12_RESOURCE_STATES InCreateState, ED3D12ResourceStateMode InResourceStateMode, const D3D12_CLEAR_VALUE* InClearValue, const TCHAR* InName) override;
-	virtual FRHIMemoryPool* CreateNewPool(int16 InPoolIndex, uint32 InMinimumAllocationSize, ERHIPoolResourceTypes InAllocationResourceType) override;
-
-	// Shared allocation/deallocation helper
-	FD3D12BaseShaderResource* GetBaseShaderResource(FRHITexture* InRHITexture);
-	void SetupAllocatedResource(FD3D12BaseShaderResource* InBaseShaderResource, FD3D12TransientMemoryStats& Stats);
-	void DeallocateMemory(FD3D12BaseShaderResource* InBaseShaderResource, FD3D12TransientMemoryStats& Stats);
-
-	// Don't track any allocations
-	virtual void UpdateAllocationTracking(FD3D12ResourceLocation& InAllocation, EAllocationType InAllocationType) override {}
-
-	// Reference all the allocated resources - owner of all the resources and will be returned back to the manager when the allocator is destroyed
-	TArray<FD3D12PooledTextureData> AllocatedTextures;
-	TArray<FD3D12PooledBufferData> AllocatedBuffers;
-
-	// All overlapping resource allocation information from all allocated resources on this allocator
-	// Not stored inside each resource to not increase the base object size for each resource
-	// (can still be moved if it the lookup becomes a perf bottleneck at some point)
-	TMap<FD3D12BaseShaderResource*, TArray<FD3D12Resource*>> OverlappingResourceData;
-
-	// Stat tracking
-	FD3D12TransientMemoryStats TextureMemoryStats;
-	FD3D12TransientMemoryStats BufferMemoryStats;
+	// This adapter mocks the D3D12 resource allocator interface and performs placed resource creation.
+	class FResourceAllocatorAdapter final : public FD3D12AdapterChild, public ID3D12ResourceAllocator
+	{
+	public:
+		FResourceAllocatorAdapter(FD3D12Adapter* Adapter, FD3D12TransientHeap& InHeap, const FRHITransientHeapAllocation& InAllocation)
+			: FD3D12AdapterChild(Adapter)
+			, Heap(InHeap)
+			, Allocation(InAllocation)
+		{}
+
+		void AllocateResource(
+			uint32 GPUIndex, D3D12_HEAP_TYPE InHeapType, const D3D12_RESOURCE_DESC& InDesc, uint64 InSize, uint32 InAllocationAlignment, ED3D12ResourceStateMode InResourceStateMode,
+			D3D12_RESOURCE_STATES InCreateState, const D3D12_CLEAR_VALUE* InClearValue, const TCHAR* InName, FD3D12ResourceLocation& ResourceLocation) override;
+
+		FD3D12TransientHeap& Heap;
+		const FRHITransientHeapAllocation& Allocation;
+	};
 
-	friend class FD3D12TransientMemoryPoolManager;
+	FRHITransientResourceAllocator Allocator;
+	FD3D12Device* AllocationInfoQueryDevice;
 };
\ No newline at end of file
diff --git a/Engine/Source/Runtime/D3D12RHI/Private/D3D12Util.cpp b/Engine/Source/Runtime/D3D12RHI/Private/D3D12Util.cpp
--- a/Engine/Source/Runtime/D3D12RHI/Private/D3D12Util.cpp
+++ b/Engine/Source/Runtime/D3D12RHI/Private/D3D12Util.cpp
@@ -1464,24 +1464,6 @@ DEFINE_STAT(STAT_D3D12TextureAllocatorAllocated);
 DEFINE_STAT(STAT_D3D12TextureAllocatorUnused);
 DEFINE_STAT(STAT_D3D12TextureAllocatorCount);
 
-DEFINE_STAT(STAT_D3D12TransientMemoryPoolAllocated);
-DEFINE_STAT(STAT_D3D12TransientMemoryFramePoolUsed);
-DEFINE_STAT(STAT_D3D12TransientMemoryFrameCommittedUsed);
-DEFINE_STAT(STAT_D3D12TransientMemoryPoolAllocations);
-DEFINE_STAT(STAT_D3D12TransientMemoryCommittedAllocations);
-
-DEFINE_STAT(STAT_D3D12TransientMemoryBufferFrameUsed);
-DEFINE_STAT(STAT_D3D12TransientMemoryBufferFrameTotalRequested);
-DEFINE_STAT(STAT_D3D12TransientMemoryBufferFrameCommittedAllocated);
-DEFINE_STAT(STAT_D3D12TransientMemoryBufferPoolAllocations);
-DEFINE_STAT(STAT_D3D12TransientMemoryBufferCommittedAllocations);
-
-DEFINE_STAT(STAT_D3D12TransientMemoryTextureFrameUsed);
-DEFINE_STAT(STAT_D3D12TransientMemoryTextureFrameTotalRequested);
-DEFINE_STAT(STAT_D3D12TransientMemoryTextureFrameCommittedAllocated);
-DEFINE_STAT(STAT_D3D12TransientMemoryTexturePoolAllocations);
-DEFINE_STAT(STAT_D3D12TransientMemoryTextureCommittedAllocations);
-
 DEFINE_STAT(STAT_D3D12BufferPoolMemoryAllocated);
 DEFINE_STAT(STAT_D3D12BufferPoolMemoryUsed);
 DEFINE_STAT(STAT_D3D12BufferPoolMemoryFree);
diff --git a/Engine/Source/Runtime/D3D12RHI/Public/D3D12Resources.h b/Engine/Source/Runtime/D3D12RHI/Public/D3D12Resources.h
--- a/Engine/Source/Runtime/D3D12RHI/Public/D3D12Resources.h
+++ b/Engine/Source/Runtime/D3D12RHI/Public/D3D12Resources.h
@@ -674,11 +674,11 @@ private:
 struct ID3D12ResourceAllocator
 {
 	// Helper function for textures to compute the correct size and alignment
-	void AllocateTexture(D3D12_HEAP_TYPE InHeapType, const D3D12_RESOURCE_DESC& InDesc, EPixelFormat InUEFormat, ED3D12ResourceStateMode InResourceStateMode,
+	void AllocateTexture(uint32 GPUIndex, D3D12_HEAP_TYPE InHeapType, const D3D12_RESOURCE_DESC& InDesc, EPixelFormat InUEFormat, ED3D12ResourceStateMode InResourceStateMode,
 		D3D12_RESOURCE_STATES InCreateState, const D3D12_CLEAR_VALUE* InClearValue, const TCHAR* InName, FD3D12ResourceLocation& ResourceLocation);
 
 	// Actual pure virtual resource allocation function
-	virtual void AllocateResource(D3D12_HEAP_TYPE InHeapType, const D3D12_RESOURCE_DESC& InDesc, uint64 InSize, uint32 InAllocationAlignment, ED3D12ResourceStateMode InResourceStateMode,
+	virtual void AllocateResource(uint32 GPUIndex, D3D12_HEAP_TYPE InHeapType, const D3D12_RESOURCE_DESC& InDesc, uint64 InSize, uint32 InAllocationAlignment, ED3D12ResourceStateMode InResourceStateMode,
 		D3D12_RESOURCE_STATES InCreateState, const D3D12_CLEAR_VALUE* InClearValue, const TCHAR* InName, FD3D12ResourceLocation& ResourceLocation) = 0;
 };
 
diff --git a/Engine/Source/Runtime/RHI/Private/RHITransientResourceAllocator.cpp b/Engine/Source/Runtime/RHI/Private/RHITransientResourceAllocator.cpp
new file mode 100644
--- /dev/null
+++ b/Engine/Source/Runtime/RHI/Private/RHITransientResourceAllocator.cpp
@@ -0,0 +1,135 @@
+// Copyright Epic Games, Inc. All Rights Reserved.
+
+#include "RHITransientResourceAllocator.h"
+#include "RHICommandList.h"
+
+FRHIShaderResourceView* FRHITransientTexture::GetOrCreateSRV(const FRHITextureSRVCreateInfo& SRVCreateInfo)
+{
+	for (const auto& KeyValue : SRVs)
+	{
+		if (KeyValue.Key == SRVCreateInfo)
+		{
+			return KeyValue.Value.GetReference();
+		}
+	}
+
+	FShaderResourceViewRHIRef RHIShaderResourceView;
+
+	if (SRVCreateInfo.MetaData == ERHITextureMetaDataAccess::None)
+	{
+		RHIShaderResourceView = RHICreateShaderResourceView(GetRHI(), SRVCreateInfo);
+	}
+	else
+	{
+		switch (SRVCreateInfo.MetaData)
+		{
+		case ERHITextureMetaDataAccess::HTile:
+			check(GRHISupportsExplicitHTile);
+			RHIShaderResourceView = RHICreateShaderResourceViewHTile((FRHITexture2D*)GetRHI());
+			break;
+
+		case ERHITextureMetaDataAccess::FMask:
+			RHIShaderResourceView = RHICreateShaderResourceViewFMask((FRHITexture2D*)GetRHI());
+			break;
+
+		case ERHITextureMetaDataAccess::CMask:
+			RHIShaderResourceView = RHICreateShaderResourceViewWriteMask((FRHITexture2D*)GetRHI());
+			break;
+		}
+	}
+
+	check(RHIShaderResourceView);
+	FRHIShaderResourceView* View = RHIShaderResourceView.GetReference();
+	SRVs.Emplace(SRVCreateInfo, MoveTemp(RHIShaderResourceView));
+	return View;
+}
+
+FRHIUnorderedAccessView* FRHITransientTexture::GetOrCreateUAV(const FRHITextureUAVCreateInfo& UAVCreateInfo)
+{
+	for (const auto& KeyValue : UAVs)
+	{
+		if (KeyValue.Key == UAVCreateInfo)
+		{
+			return KeyValue.Value.GetReference();
+		}
+	}
+
+	FUnorderedAccessViewRHIRef RHIUnorderedAccessView;
+
+	if (UAVCreateInfo.MetaData == ERHITextureMetaDataAccess::HTile)
+	{
+		check(GRHISupportsExplicitHTile);
+		RHIUnorderedAccessView = RHICreateUnorderedAccessViewHTile((FRHITexture2D*)GetRHI());
+	}
+	else if (UAVCreateInfo.MetaData == ERHITextureMetaDataAccess::Stencil)
+	{
+		RHIUnorderedAccessView = RHICreateUnorderedAccessViewStencil((FRHITexture2D*)GetRHI(), 0);
+	}
+	else
+	{
+		RHIUnorderedAccessView = RHICreateUnorderedAccessView(GetRHI(), UAVCreateInfo.MipLevel);
+	}
+
+	check(RHIUnorderedAccessView);
+	FRHIUnorderedAccessView* View = RHIUnorderedAccessView.GetReference();
+	UAVs.Emplace(UAVCreateInfo, MoveTemp(RHIUnorderedAccessView));
+	return View;
+}
+
+
+FRHIShaderResourceView* FRHITransientBuffer::GetOrCreateSRV(const FRHIBufferSRVCreateInfo& SRVCreateInfo)
+{
+	for (const auto& KeyValue : SRVs)
+	{
+		if (KeyValue.Key == SRVCreateInfo)
+		{
+			return KeyValue.Value.GetReference();
+		}
+	}
+
+	FShaderResourceViewRHIRef RHIShaderResourceView;
+
+	if (SRVCreateInfo.Format != PF_Unknown)
+	{
+		RHIShaderResourceView = RHICreateShaderResourceView(GetRHI(), SRVCreateInfo.BytesPerElement, SRVCreateInfo.Format);
+	}
+	else
+	{
+		RHIShaderResourceView = RHICreateShaderResourceView(GetRHI());
+	}
+
+	FRHIShaderResourceView* View = RHIShaderResourceView.GetReference();
+	SRVs.Emplace(SRVCreateInfo, MoveTemp(RHIShaderResourceView));
+	return View;
+}
+
+FRHIUnorderedAccessView* FRHITransientBuffer::GetOrCreateUAV(const FRHIBufferUAVCreateInfo& UAVCreateInfo)
+{
+	for (const auto& KeyValue : UAVs)
+	{
+		if (KeyValue.Key == UAVCreateInfo)
+		{
+			return KeyValue.Value.GetReference();
+		}
+	}
+
+	FUnorderedAccessViewRHIRef RHIUnorderedAccessView;
+
+	if (UAVCreateInfo.Format != PF_Unknown)
+	{
+		RHIUnorderedAccessView = RHICreateUnorderedAccessView(GetRHI(), UAVCreateInfo.Format);
+	}
+	else
+	{
+		RHIUnorderedAccessView = RHICreateUnorderedAccessView(GetRHI(), UAVCreateInfo.bSupportsAtomicCounter, UAVCreateInfo.bSupportsAppendBuffer);
+	}
+
+	FRHIUnorderedAccessView* View = RHIUnorderedAccessView.GetReference();
+	UAVs.Emplace(UAVCreateInfo, MoveTemp(RHIUnorderedAccessView));
+	return View;
+}
+
+void IRHITransientResourceAllocator::Release(FRHICommandListImmediate&)
+{
+	delete this;
+}
\ No newline at end of file
diff --git a/Engine/Source/Runtime/RHI/Private/RHIValidation.cpp b/Engine/Source/Runtime/RHI/Private/RHIValidation.cpp
--- a/Engine/Source/Runtime/RHI/Private/RHIValidation.cpp
+++ b/Engine/Source/Runtime/RHI/Private/RHIValidation.cpp
@@ -256,9 +256,28 @@ void FValidationRHI::RHICreateTransition(FRHITransition* Transition, const FRHIT
 			checkf(SrcPipelines == ERHIPipeline::Graphics, TEXT("Acquiring a transient resource (%s) must begin on the graphics pipe."), Resource->GetDebugName());
 
 			AcquireOps.Emplace(FOperation::AcquireTransientResource(Resource, nullptr));
+
+			for (const FRHITransientAliasingOverlap& Overlap : Info.Overlaps)
+			{
+				FResource* ResourceBefore = nullptr;
+
+				if (Overlap.Type == FRHITransientAliasingOverlap::EType::Texture)
+				{
+					ResourceBefore = Overlap.Texture->GetTrackerResource();
+				}
+				else
+				{
+					ResourceBefore = Overlap.Buffer;
+				}
+
+				checkf(ResourceBefore, TEXT("Null resource provided as an aliasing overlap of %s"), Resource->GetDebugName());
+
+				AcquireOps.Emplace(FOperation::AliasingOverlap(ResourceBefore, Resource, nullptr));
+			}
 		}
 		else
 		{
+			checkf(Info.Overlaps.IsEmpty(), TEXT("Aliasing overlaps provided on a Discard for resource %s. Overlaps must be provided on an acquire."), Resource->GetDebugName());
 			checkf(Resource->TransientState.bTransient, TEXT("Discarding resource %s which is not transient. Only transient resources can be discarded."), Resource->GetDebugName());
 			checkf(DstPipelines == ERHIPipeline::Graphics, TEXT("Discarding a transient resource (%s) must end on the graphics pipe."), Resource->GetDebugName());
 
@@ -308,7 +327,22 @@ void FValidationRHI::RHICreateTransition(FRHITransition* Transition, const FRHIT
 	if (bDoTrace)
 	{
 		void* Backtrace = CaptureBacktrace();
-		for (FOperation& Op : AcquireOps) { Op.Data_AcquireTransient.CreateBacktrace = Backtrace; }
+
+		for (FOperation& Op : AcquireOps)
+		{
+			switch (Op.Type)
+			{
+			case EOpType::AcquireTransient:
+				Op.Data_AcquireTransient.CreateBacktrace = Backtrace;
+				break;
+			case EOpType::AliasingOverlap:
+				Op.Data_AliasingOverlap.CreateBacktrace = Backtrace;
+				break;
+			default:
+				checkNoEntry();
+			}
+		}
+
 		for (FOperation& Op : DiscardOps) { Op.Data_DiscardTransient.CreateBacktrace = Backtrace; }
 		for (FOperation& Op : BeginOps) { Op.Data_BeginTransition.CreateBacktrace = Backtrace; }
 		for (FOperation& Op : EndOps) { Op.Data_EndTransition.CreateBacktrace = Backtrace; }
@@ -685,6 +719,35 @@ namespace RHIValidation
 			*DebugName);
 	}
 
+	static inline FString GetReasonString_AliasingOverlapNonDiscarded(FResource* ResourceBefore, FResource* ResourceAfter, void* CreateTrace)
+	{
+		FString DebugNameBefore = GetResourceDebugName(ResourceBefore, {});
+		FString DebugNameAfter = GetResourceDebugName(ResourceAfter, {});
+		return FString::Printf(
+			BARRIER_TRACKER_LOG_PREFIX_RESNAME
+			TEXT("Attempted to overlap resource %s (before) with resource %s (after), but %s (before) has not been discarded.\n")
+			TEXT("%s")
+			BARRIER_TRACKER_LOG_SUFFIX,
+			*DebugNameAfter,
+			*DebugNameBefore,
+			*DebugNameAfter,
+			*DebugNameBefore,
+			*GetReasonString_Backtrace(TEXT("acquire"), TEXT("RHICreateTransition"), CreateTrace));
+	}
+
+	static inline FString GetReasonString_AliasingOverlapNonTransient(FResource* ResourceBefore, FResource* ResourceAfter)
+	{
+		FString DebugNameBefore = GetResourceDebugName(ResourceBefore, {});
+		FString DebugNameAfter = GetResourceDebugName(ResourceAfter, {});
+		return FString::Printf(
+			BARRIER_TRACKER_LOG_PREFIX_RESNAME
+			TEXT("Attempted to overlap non-transient resource %s when acquiring resource %s. Only transient resources may be used in an aliasing overlap operation.\n")
+			BARRIER_TRACKER_LOG_SUFFIX,
+			*DebugNameBefore,
+			*DebugNameBefore,
+			*DebugNameAfter);
+	}
+
 	static inline FString GetReasonString_DuplicateAcquireTransient(FResource* Resource, void* PreviousAcquireTrace, void* CurrentAcquireTrace)
 	{
 		FString DebugName = GetResourceDebugName(Resource, {});
@@ -937,6 +1000,26 @@ namespace RHIValidation
 		}
 	}
 
+	void FTransientState::AliasingOverlap(FResource* ResourceBefore, FResource* ResourceAfter, void* CreateTrace)
+	{
+		FTransientState& TransientStateBefore = ResourceBefore->TransientState;
+		FTransientState& TransientStateAfter = ResourceAfter->TransientState;
+
+		// Acquire should validate whether ResourceAfter is transient. We assume it is here.
+		RHI_VALIDATION_CHECK(TransientStateBefore.bTransient, *GetReasonString_AliasingOverlapNonTransient(ResourceBefore, ResourceAfter));
+		RHI_VALIDATION_CHECK(TransientStateBefore.IsDiscarded(), *GetReasonString_AliasingOverlapNonDiscarded(ResourceBefore, ResourceAfter, CreateTrace));
+
+		if (ResourceBefore->LoggingMode != ELoggingMode::None)
+		{
+			Log(ResourceBefore, {}, CreateTrace, TEXT("AliasingOverlap"), TEXT("AliasingOverlap"), TEXT("Aliasing Overlap (Before)"));
+		}
+
+		if (ResourceAfter->LoggingMode != ELoggingMode::None)
+		{
+			Log(ResourceAfter, {}, CreateTrace, TEXT("AliasingOverlap"), TEXT("AliasingOverlap"), TEXT("Aliasing Overlap (After)"));
+		}
+	}
+
 	void FResource::SetDebugName(const TCHAR* Name, const TCHAR* Suffix)
 	{
 		DebugName = Suffix
@@ -1202,6 +1285,12 @@ namespace RHIValidation
 			Data_EndTransition.Identity.Resource->ReleaseOpRef();
 			break;
 
+		case EOpType::AliasingOverlap:
+			FTransientState::AliasingOverlap(Data_AliasingOverlap.ResourceBefore, Data_AliasingOverlap.ResourceAfter, Data_AliasingOverlap.CreateBacktrace);
+			Data_AliasingOverlap.ResourceBefore->ReleaseOpRef();
+			Data_AliasingOverlap.ResourceAfter->ReleaseOpRef();
+			break;
+
 		case EOpType::AcquireTransient:
 			Data_AcquireTransient.Resource->TransientState.Acquire(Data_AcquireTransient.Resource, Data_AcquireTransient.CreateBacktrace);
 			Data_AcquireTransient.Resource->ReleaseOpRef();
@@ -1394,60 +1483,22 @@ TLockFreePointerListUnordered<FValidationContext, PLATFORM_CACHE_LINE_SIZE> FVal
 FValidationTransientResourceAllocator::~FValidationTransientResourceAllocator()
 {
 	check(bFrozen);
-
-	// Validation of leaking textures resources disabled by default because we know textures are referenced by global uniforms
-	// This needs to be fixed first
-	static bool bValidateLeakingTextures = false;
-	static bool bValidateLeakingBuffers = true;
-	if (bValidateLeakingTextures || bValidateLeakingBuffers)
-	{
-		// Check all allocated resource data and make sure the transient allocator holds the last 'reference'
-		int32 LeakingResourceCount = 0;
-		for (auto It = AllocatedResourceMap.CreateIterator(); It; ++It)
-		{
-			const FRHIResource* RHIResource = It->Key;
-			const AllocatedResourceData& ResourceData = It->Value;
-			if (ResourceData.ResourceType == AllocatedResourceData::EType::Texture && bValidateLeakingTextures)
-			{
-				LeakingResourceCount += (RHIResource->GetRefCount() > 1) ? 1 : 0;
-			}
-			else if (ResourceData.ResourceType == AllocatedResourceData::EType::Buffer && bValidateLeakingBuffers)
-			{
-				LeakingResourceCount += (RHIResource->GetRefCount() > 1) ? 1 : 0;
-			}
-		}
-
-		if (LeakingResourceCount > 0)
-		{
-			FString ErrorMessage = FString::Printf(
-				TRANSIENT_RESOURCE_LOG_PREFIX_REASON("Leaking transient resources")
-				TEXT("%d Transient Resource allocations are still used when the alloctor is destroyed. The TransientResourceAllocator manages the lifetime of all allocated resource.\n\n")
-				TEXT("Leaking resources:\n"),
-				LeakingResourceCount);
-			for (auto It = AllocatedResourceMap.CreateIterator(); It; ++It)
-			{
-				const FRHIResource* RHIResource = It->Key;
-				const AllocatedResourceData& ResourceData = It->Value;
-				if (RHIResource->GetRefCount() > 1)
-				{
-					ErrorMessage += FString::Printf(TEXT("         %s (%s) - Refcount: %d\n"), ResourceData.DebugName, ResourceData.ResourceType == AllocatedResourceData::EType::Texture ? TEXT("Texture") : TEXT("DataBuffer"), RHIResource->GetRefCount());
-				}
-			}
-			ErrorMessage += FString::Printf(TRANSIENT_RESOURCE_LOG_SUFFIX);
-			FValidationRHI::ReportValidationFailure(*ErrorMessage);
-		}
-	}
-
-	// Delete actual allocator
-	delete RHIAllocator;
+	check(bReleased);
 }
 
-FRHITexture* FValidationTransientResourceAllocator::CreateTexture(const FRHITextureCreateInfo& InCreateInfo, const TCHAR* InDebugName)
+FRHITransientTexture* FValidationTransientResourceAllocator::CreateTexture(const FRHITextureCreateInfo& InCreateInfo, const TCHAR* InDebugName)
 {
 	check(!bFrozen);
 	checkf(!EnumHasAnyFlags(InCreateInfo.Flags, TexCreate_Transient), TEXT("Attempted to create transient texture %s using legacy TexCreate_Transient flag."), InDebugName);
 
-	FRHITexture* RHITexture = RHIAllocator->CreateTexture(InCreateInfo, InDebugName);
+	FRHITransientTexture* TransientTexture = RHIAllocator->CreateTexture(InCreateInfo, InDebugName);
+
+	if (!TransientTexture)
+	{
+		return nullptr;
+	}
+
+	FRHITexture* RHITexture = TransientTexture->GetRHI();
 
 	bool bReinitializeTrackerResource = true;
 
@@ -1461,9 +1512,9 @@ FRHITexture* FValidationTransientResourceAllocator::CreateTexture(const FRHIText
 	}
 
 	// Store allocation data
-	AllocatedResourceData ResourceData;
+	FAllocatedResourceData ResourceData;
 	ResourceData.DebugName = InDebugName;
-	ResourceData.ResourceType = AllocatedResourceData::EType::Texture;
+	ResourceData.ResourceType = FAllocatedResourceData::EType::Texture;
 	ResourceData.bMemoryAllocated = true;
 	ResourceData.bReinitializeBarrierTracking = bReinitializeTrackerResource;
 	ResourceData.Texture.Flags = InCreateInfo.Flags;
@@ -1472,15 +1523,22 @@ FRHITexture* FValidationTransientResourceAllocator::CreateTexture(const FRHIText
 	ResourceData.Texture.NumMips = InCreateInfo.NumMips;
 	AllocatedResourceMap.Add(RHITexture, ResourceData);
 
-	return RHITexture;
+	return TransientTexture;
 }
 
-FRHIBuffer* FValidationTransientResourceAllocator::CreateBuffer(const FRHIBufferCreateInfo& InCreateInfo, const TCHAR* InDebugName)
+FRHITransientBuffer* FValidationTransientResourceAllocator::CreateBuffer(const FRHIBufferCreateInfo& InCreateInfo, const TCHAR* InDebugName)
 {
 	check(!bFrozen);
 	checkf(!EnumHasAnyFlags(InCreateInfo.Usage, BUF_Transient), TEXT("Attempted to create transient buffer %s using legacy BUF_Transient flag."), InDebugName);
-	
-	FRHIBuffer* RHIBuffer = RHIAllocator->CreateBuffer(InCreateInfo, InDebugName);
+
+	FRHITransientBuffer* TransientBuffer = RHIAllocator->CreateBuffer(InCreateInfo, InDebugName);
+
+	if (!TransientBuffer)
+	{
+		return nullptr;
+	}
+
+	FRHIBuffer* RHIBuffer = TransientBuffer->GetRHI();
 
 	bool bReinitializeTrackerResource = true;
 
@@ -1491,36 +1549,36 @@ FRHIBuffer* FValidationTransientResourceAllocator::CreateBuffer(const FRHIBuffer
 	}
 
 	// Store allocation data
-	AllocatedResourceData ResourceData;
+	FAllocatedResourceData ResourceData;
 	ResourceData.DebugName = InDebugName;
-	ResourceData.ResourceType = AllocatedResourceData::EType::Buffer;
+	ResourceData.ResourceType = FAllocatedResourceData::EType::Buffer;
 	ResourceData.bMemoryAllocated = true;
 	ResourceData.bReinitializeBarrierTracking = bReinitializeTrackerResource;
 	AllocatedResourceMap.Add(RHIBuffer, ResourceData);
 
-	return RHIBuffer;
+	return TransientBuffer;
 }
 
-void FValidationTransientResourceAllocator::DeallocateMemory(FRHITexture* InTexture)
+void FValidationTransientResourceAllocator::DeallocateMemory(FRHITransientTexture* InTransientTexture)
 {
 	check(!bFrozen);
 
-	RHIAllocator->DeallocateMemory(InTexture);
+	RHIAllocator->DeallocateMemory(InTransientTexture);
 
 	// Mark memory as freed
-	AllocatedResourceData* ResourceData = AllocatedResourceMap.Find(InTexture);
+	FAllocatedResourceData* ResourceData = AllocatedResourceMap.Find(InTransientTexture->GetRHI());
 	check(ResourceData);
 	ResourceData->bMemoryAllocated = false;
 }
 
-void FValidationTransientResourceAllocator::DeallocateMemory(FRHIBuffer* InBuffer)
+void FValidationTransientResourceAllocator::DeallocateMemory(FRHITransientBuffer* InTransientBuffer)
 {
 	check(!bFrozen);
 
-	RHIAllocator->DeallocateMemory(InBuffer);
+	RHIAllocator->DeallocateMemory(InTransientBuffer);
 
 	// Mark memory as freed
-	AllocatedResourceData* ResourceData = AllocatedResourceMap.Find(InBuffer);
+	FAllocatedResourceData* ResourceData = AllocatedResourceMap.Find(InTransientBuffer->GetRHI());
 	check(ResourceData);
 	ResourceData->bMemoryAllocated = false;
 }
@@ -1530,30 +1588,32 @@ void FValidationTransientResourceAllocator::Freeze(FRHICommandListImmediate& RHI
 	check(!bFrozen);
 
 	// Check all allocated resource data and make sure all memory is freed again
-	int32 AllocatedResourceCount = 0;
-	for (auto It = AllocatedResourceMap.CreateIterator(); It; ++It)
 	{
-		const AllocatedResourceData& ResourceData = It->Value;
-		AllocatedResourceCount += ResourceData.bMemoryAllocated ? 1 : 0;
-	}
+		TArray<const FAllocatedResourceData*> AllocatedResources;
 
-	if (AllocatedResourceCount > 0)
-	{
-		FString ErrorMessage = FString::Printf(
-			TRANSIENT_RESOURCE_LOG_PREFIX_REASON("Open transient allocations")
-			TEXT("%d Transient Resource allocations still have memory allocated. Call 'DeallocateMemory' on all transient allocated resources before calling Freeze.\n\n")
-			TEXT("Resources with Allocated Memory:\n"), 
-			AllocatedResourceCount);
-		for (auto It = AllocatedResourceMap.CreateIterator(); It; ++It)
+		for (const auto& KeyValue : AllocatedResourceMap)
 		{
-			const AllocatedResourceData& ResourceData = It->Value;
+			const FAllocatedResourceData& ResourceData = KeyValue.Value;
 			if (ResourceData.bMemoryAllocated)
 			{
-				ErrorMessage += FString::Printf(TEXT("         %s (%s)\n"), ResourceData.DebugName, ResourceData.ResourceType == AllocatedResourceData::EType::Texture ? TEXT("Texture") : TEXT("DataBuffer"));
+				AllocatedResources.Emplace(&ResourceData);
+			}
+		}
+
+		if (AllocatedResources.Num() > 0)
+		{
+			FString ErrorMessage = FString::Printf(
+				TRANSIENT_RESOURCE_LOG_PREFIX_REASON("Open transient allocations")
+				TEXT("%d Transient Resource allocations still have memory allocated. Call 'DeallocateMemory' on all transient allocated resources before calling Freeze.\n\n")
+				TEXT("Resources with Allocated Memory:\n"),
+				AllocatedResources.Num());
+			for (const FAllocatedResourceData* ResourceData : AllocatedResources)
+			{
+				ErrorMessage += FString::Printf(TEXT("         %s (%s)\n"), ResourceData->DebugName, ResourceData->ResourceType == FAllocatedResourceData::EType::Texture ? TEXT("Texture") : TEXT("Buffer"));
 			}
+			ErrorMessage += FString::Printf(TRANSIENT_RESOURCE_LOG_SUFFIX);
+			FValidationRHI::ReportValidationFailure(*ErrorMessage);
 		}
-		ErrorMessage += FString::Printf(TRANSIENT_RESOURCE_LOG_SUFFIX);
-		FValidationRHI::ReportValidationFailure(*ErrorMessage);
 	}
 
 	RHICmdList.EnqueueLambda([this](FRHICommandListImmediate& InRHICmdList)
@@ -1569,6 +1629,15 @@ void FValidationTransientResourceAllocator::Freeze(FRHICommandListImmediate& RHI
 	RHIAllocator->Freeze(RHICmdList);
 }
 
+void FValidationTransientResourceAllocator::Release(FRHICommandListImmediate& RHICmdList)
+{
+	check(bFrozen);
+
+	RHIAllocator->Release(RHICmdList);
+	bReleased = true;
+	delete this;
+}
+
 void FValidationTransientResourceAllocator::InitBarrierTracking()
 {
 	// Barrier tracking initialization has to happen on the RHI thread, because RHI resources are pooled and reused.
@@ -1576,7 +1645,7 @@ void FValidationTransientResourceAllocator::InitBarrierTracking()
 	for (const auto& Entry : AllocatedResourceMap)
 	{
 		FRHIResource* Resource = Entry.Key;
-		const AllocatedResourceData& ResourceData = Entry.Value;
+		const FAllocatedResourceData& ResourceData = Entry.Value;
 
 		if (!ResourceData.bReinitializeBarrierTracking)
 		{
@@ -1585,13 +1654,13 @@ void FValidationTransientResourceAllocator::InitBarrierTracking()
 
 		switch (ResourceData.ResourceType)
 		{
-		case AllocatedResourceData::EType::Texture:
+		case FAllocatedResourceData::EType::Texture:
 		{
 			FRHITexture* Texture = static_cast<FRHITexture*>(Resource);
 			Texture->InitBarrierTracking(ResourceData.Texture.NumMips, ResourceData.Texture.ArraySize, ResourceData.Texture.Format, ResourceData.Texture.Flags, ERHIAccess::Discard, ResourceData.DebugName);
 		}
 		break;
-		case AllocatedResourceData::EType::Buffer:
+		case FAllocatedResourceData::EType::Buffer:
 		{
 			FRHIBuffer* Buffer = static_cast<FRHIBuffer*>(Resource);
 			Buffer->InitBarrierTracking(ERHIAccess::Discard, ResourceData.DebugName);
diff --git a/Engine/Source/Runtime/RHI/Public/DynamicRHI.h b/Engine/Source/Runtime/RHI/Public/DynamicRHI.h
--- a/Engine/Source/Runtime/RHI/Public/DynamicRHI.h
+++ b/Engine/Source/Runtime/RHI/Public/DynamicRHI.h
@@ -20,6 +20,7 @@ class FReadSurfaceDataFlags;
 class FRHICommandList;
 class FRHIComputeFence;
 class FRayTracingPipelineState;
+class IRHITransientResourceAllocator;
 struct FDepthStencilStateInitializerRHI;
 struct FRasterizerStateInitializerRHI;
 struct FRHIResourceCreateInfo;
diff --git a/Engine/Source/Runtime/RHI/Public/RHI.h b/Engine/Source/Runtime/RHI/Public/RHI.h
--- a/Engine/Source/Runtime/RHI/Public/RHI.h
+++ b/Engine/Source/Runtime/RHI/Public/RHI.h
@@ -1627,73 +1627,6 @@ struct FRHIResourceCreateInfo
 	uint32 ExtData;
 };
 
-enum ERHITextureSRVOverrideSRGBType
-{
-	SRGBO_Default,
-	SRGBO_ForceDisable,
-};
-
-struct FRHITextureSRVCreateInfo
-{
-	explicit FRHITextureSRVCreateInfo(uint8 InMipLevel = 0u, uint8 InNumMipLevels = 1u, uint8 InFormat = PF_Unknown)
-		: Format(InFormat)
-		, MipLevel(InMipLevel)
-		, NumMipLevels(InNumMipLevels)
-		, SRGBOverride(SRGBO_Default)
-		, FirstArraySlice(0)
-		, NumArraySlices(0)
-	{}
-
-	explicit FRHITextureSRVCreateInfo(uint8 InMipLevel, uint8 InNumMipLevels, uint32 InFirstArraySlice, uint32 InNumArraySlices, uint8 InFormat = PF_Unknown)
-		: Format(InFormat)
-		, MipLevel(InMipLevel)
-		, NumMipLevels(InNumMipLevels)
-		, SRGBOverride(SRGBO_Default)
-		, FirstArraySlice(InFirstArraySlice)
-		, NumArraySlices(InNumArraySlices)
-	{}
-
-	/** View the texture with a different format. Leave as PF_Unknown to use original format. Useful when sampling stencil */
-	uint8 Format;
-
-	/** Specify the mip level to use. Useful when rendering to one mip while sampling from another */
-	uint8 MipLevel;
-
-	/** Create a view to a single, or multiple mip levels */
-	uint8 NumMipLevels;
-
-	/** Potentially override the texture's sRGB flag */
-	ERHITextureSRVOverrideSRGBType SRGBOverride;
-
-	/** Specify first array slice index. By default 0. */
-	uint32 FirstArraySlice;
-
-	/** Specify number of array slices. If FirstArraySlice and NumArraySlices are both zero, the SRV is created for all array slices. By default 0. */
-	uint32 NumArraySlices;
-
-	FORCEINLINE bool operator==(const FRHITextureSRVCreateInfo& Other)const
-	{
-		return (
-			Format == Other.Format &&
-			MipLevel == Other.MipLevel &&
-			NumMipLevels == Other.NumMipLevels &&
-			SRGBOverride == Other.SRGBOverride &&
-			FirstArraySlice == Other.FirstArraySlice &&
-			NumArraySlices == Other.NumArraySlices);
-	}
-
-	FORCEINLINE bool operator!=(const FRHITextureSRVCreateInfo& Other)const
-	{
-		return !(*this == Other);
-	}
-};
-
-FORCEINLINE uint32 GetTypeHash(const FRHITextureSRVCreateInfo& Var)
-{
-	uint32 Hash0 = uint32(Var.Format) | uint32(Var.MipLevel) << 8 | uint32(Var.NumMipLevels) << 16 | uint32(Var.SRGBOverride) << 24;
-	return HashCombine(HashCombine(GetTypeHash(Hash0), GetTypeHash(Var.FirstArraySlice)), GetTypeHash(Var.NumArraySlices));
-}
-
 struct FResolveRect
 {
 	int32 X1;
@@ -2017,6 +1950,59 @@ struct FRHITransitionInfo : public FRHISubresourceRange
 	}
 };
 
+struct FRHITransientAliasingOverlap
+{
+	union
+	{
+		class FRHIResource* Resource = nullptr;
+		class FRHITexture* Texture;
+		class FRHIBuffer* Buffer;
+	};
+
+	enum class EType : uint8
+	{
+		Texture,
+		Buffer
+	} Type = EType::Texture;
+
+	FRHITransientAliasingOverlap() = default;
+
+	FRHITransientAliasingOverlap(FRHIResource* InResource, EType InType)
+		: Resource(InResource)
+		, Type(InType)
+	{}
+
+	FRHITransientAliasingOverlap(FRHITexture* InTexture)
+		: Texture(InTexture)
+		, Type(EType::Texture)
+	{}
+
+	FRHITransientAliasingOverlap(FRHIBuffer* InBuffer)
+		: Buffer(InBuffer)
+		, Type(EType::Buffer)
+	{}
+
+	bool IsTexture() const
+	{
+		return Type == EType::Texture;
+	}
+
+	bool IsBuffer() const
+	{
+		return Type == EType::Buffer;
+	}
+
+	bool operator == (const FRHITransientAliasingOverlap& Other) const
+	{
+		return Resource == Other.Resource;
+	}
+
+	inline bool operator != (const FRHITransientAliasingOverlap& RHS) const
+	{
+		return !(*this == RHS);
+	}
+};
+
 struct FRHITransientAliasingInfo
 {
 	union
@@ -2026,6 +2012,9 @@ struct FRHITransientAliasingInfo
 		class FRHIBuffer* Buffer;
 	};
 
+	// List of prior resource overlaps to use when acquiring. Must be empty for discard operations.
+	TArrayView<const FRHITransientAliasingOverlap> Overlaps;
+
 	enum class EType : uint8
 	{
 		Texture,
@@ -2040,19 +2029,21 @@ struct FRHITransientAliasingInfo
 
 	FRHITransientAliasingInfo() = default;
 
-	static FRHITransientAliasingInfo Acquire(class FRHITexture* Texture)
+	static FRHITransientAliasingInfo Acquire(class FRHITexture* Texture, TArrayView<const FRHITransientAliasingOverlap> InOverlaps)
 	{
 		FRHITransientAliasingInfo Info;
 		Info.Texture = Texture;
+		Info.Overlaps = InOverlaps;
 		Info.Type = EType::Texture;
 		Info.Action = EAction::Acquire;
 		return Info;
 	}
 
-	static FRHITransientAliasingInfo Acquire(class FRHIBuffer* Buffer)
+	static FRHITransientAliasingInfo Acquire(class FRHIBuffer* Buffer, TArrayView<const FRHITransientAliasingOverlap> InOverlaps)
 	{
 		FRHITransientAliasingInfo Info;
 		Info.Buffer = Buffer;
+		Info.Overlaps = InOverlaps;
 		Info.Type = EType::Buffer;
 		Info.Action = EAction::Acquire;
 		return Info;
diff --git a/Engine/Source/Runtime/RHI/Public/RHICommandList.h b/Engine/Source/Runtime/RHI/Public/RHICommandList.h
--- a/Engine/Source/Runtime/RHI/Public/RHICommandList.h
+++ b/Engine/Source/Runtime/RHI/Public/RHICommandList.h
@@ -1541,18 +1541,6 @@ FRHICOMMAND_MACRO(FRHICommandResourceTransition)
 	RHI_API void Execute(FRHICommandListBase& CmdList);
 };
 
-FRHICOMMAND_MACRO(FRHICommandReleaseTransientResourceAllocator)
-{
-	IRHITransientResourceAllocator* Allocator;
-
-	FRHICommandReleaseTransientResourceAllocator(IRHITransientResourceAllocator * InAllocator)
-		: Allocator(InAllocator)
-	{
-	}
-
-	RHI_API void Execute(FRHICommandListBase & CmdList);
-};
-
 struct FRHICommandSetAsyncComputeBudgetString
 {
 	static const TCHAR* TStr() { return TEXT("FRHICommandSetAsyncComputeBudget"); }
@@ -4260,7 +4248,7 @@ public:
 		return GDynamicRHI->RHICreateShaderResourceView_RenderThread(*this, Texture, CreateInfo);
 	}
 	
-	FORCEINLINE FShaderResourceViewRHIRef CreateShaderResourceView(FRHITexture* Texture, uint8 MipLevel, uint8 NumMipLevels, uint8 Format)
+	FORCEINLINE FShaderResourceViewRHIRef CreateShaderResourceView(FRHITexture* Texture, uint8 MipLevel, uint8 NumMipLevels, EPixelFormat Format)
 	{
 		LLM_SCOPE(ELLMTag::RHIMisc);
 		checkf(Texture, TEXT("Can't create a view off a null resource!"));
@@ -4469,17 +4457,6 @@ public:
 		return GDynamicRHI->RHICreateRenderQuery_RenderThread(*this, QueryType);
 	}
 
-	FORCEINLINE_DEBUGGABLE void ReleaseTransientResourceAllocator(IRHITransientResourceAllocator* InAllocator)
-	{
-		check(InAllocator);
-		if (Bypass())
-		{
-			GetContext().RHIReleaseTransientResourceAllocator(InAllocator);
-			return;
-		}
-		ALLOC_COMMAND(FRHICommandReleaseTransientResourceAllocator)(InAllocator);
-	}
-
 	FORCEINLINE void AcquireTransientResource_RenderThread(FRHITexture* Texture)
 	{
 		if (!Texture->IsCommitted() )
@@ -5317,7 +5294,7 @@ FORCEINLINE FShaderResourceViewRHIRef RHICreateShaderResourceView(FRHITexture* T
 	return FRHICommandListExecutor::GetImmediateCommandList().CreateShaderResourceView(Texture, MipLevel);
 }
 
-FORCEINLINE FShaderResourceViewRHIRef RHICreateShaderResourceView(FRHITexture* Texture, uint8 MipLevel, uint8 NumMipLevels, uint8 Format)
+FORCEINLINE FShaderResourceViewRHIRef RHICreateShaderResourceView(FRHITexture* Texture, uint8 MipLevel, uint8 NumMipLevels, EPixelFormat Format)
 {
 	return FRHICommandListExecutor::GetImmediateCommandList().CreateShaderResourceView(Texture, MipLevel, NumMipLevels, Format);
 }
diff --git a/Engine/Source/Runtime/RHI/Public/RHICommandListCommandExecutes.inl b/Engine/Source/Runtime/RHI/Public/RHICommandListCommandExecutes.inl
--- a/Engine/Source/Runtime/RHI/Public/RHICommandListCommandExecutes.inl
+++ b/Engine/Source/Runtime/RHI/Public/RHICommandListCommandExecutes.inl
@@ -374,12 +374,6 @@ void FRHICommandClearUAVUint::Execute(FRHICommandListBase& CmdList)
 	INTERNAL_DECORATOR_COMPUTE(RHIClearUAVUint)(UnorderedAccessViewRHI, Values);
 }
 
-void FRHICommandReleaseTransientResourceAllocator::Execute(FRHICommandListBase& CmdList)
-{
-	RHISTAT(ReleaseTransientResourceAllocator);
-	INTERNAL_DECORATOR(RHIReleaseTransientResourceAllocator)(Allocator);
-}
-
 void FRHICommandCopyToResolveTarget::Execute(FRHICommandListBase& CmdList)
 {
 	RHISTAT(CopyToResolveTarget);
diff --git a/Engine/Source/Runtime/RHI/Public/RHIContext.h b/Engine/Source/Runtime/RHI/Public/RHIContext.h
--- a/Engine/Source/Runtime/RHI/Public/RHIContext.h
+++ b/Engine/Source/Runtime/RHI/Public/RHIContext.h
@@ -677,11 +677,6 @@ public:
 	}
 #endif
 
-	/**
-	* Release the transient allocator and all internal resources when the command list is done executing
-	*/
-	virtual void RHIReleaseTransientResourceAllocator(IRHITransientResourceAllocator* InAllocator) {}
-
 	virtual void RHIClearRayTracingBindings(FRHIRayTracingScene* Scene)
 	{
 		checkNoEntry();
diff --git a/Engine/Source/Runtime/RHI/Public/RHIDefinitions.h b/Engine/Source/Runtime/RHI/Public/RHIDefinitions.h
--- a/Engine/Source/Runtime/RHI/Public/RHIDefinitions.h
+++ b/Engine/Source/Runtime/RHI/Public/RHIDefinitions.h
@@ -2033,6 +2033,24 @@ inline bool IsRayTracingShaderFrequency(EShaderFrequency Frequency)
 	}
 }
 
+inline ERHIResourceType GetRHIResourceType(ETextureDimension Dimension)
+{
+	switch (Dimension)
+	{
+	case ETextureDimension::Texture2D:
+		return ERHIResourceType::RRT_Texture2D;
+	case ETextureDimension::Texture2DArray:
+		return ERHIResourceType::RRT_Texture2DArray;
+	case ETextureDimension::Texture3D:
+		return ERHIResourceType::RRT_Texture3D;
+	case ETextureDimension::TextureCube:
+	case ETextureDimension::TextureCubeArray:
+		return ERHIResourceType::RRT_TextureCube;
+	}
+	checkNoEntry();
+	return ERHIResourceType::RRT_None;
+}
+
 #if PLATFORM_SUPPORTS_GEOMETRY_SHADERS
 	#define GEOMETRY_SHADER(GeometryShader)	(GeometryShader)
 #else
diff --git a/Engine/Source/Runtime/RHI/Public/RHIResources.h b/Engine/Source/Runtime/RHI/Public/RHIResources.h
--- a/Engine/Source/Runtime/RHI/Public/RHIResources.h
+++ b/Engine/Source/Runtime/RHI/Public/RHIResources.h
@@ -2861,7 +2861,6 @@ private:
 	RHI_API void OnVerifyNumUAVsFailed(int32 InNumUAVs);
 };
 
-
 /** Descriptor used to create a texture resource */
 struct RHI_API FRHITextureCreateInfo
 {
@@ -2893,10 +2892,9 @@ struct RHI_API FRHITextureCreateInfo
 		EPixelFormat InFormat,
 		FClearValueBinding InClearValue,
 		ETextureCreateFlags InFlags,
-		uint8 InNumMips = 1,
-		uint8 InNumSamples = 1)
+		uint8 InNumMips = 1)
 	{
-		return FRHITextureCreateInfo(ETextureDimension::Texture3D, InFlags, InFormat, FIntPoint(InSize.X, InSize.Y), InClearValue, InSize.Z, 1, InNumMips, InNumSamples);
+		return FRHITextureCreateInfo(ETextureDimension::Texture3D, InFlags, InFormat, FIntPoint(InSize.X, InSize.Y), InClearValue, InSize.Z, 1, InNumMips);
 	}
 
 	static FRHITextureCreateInfo CreateCube(
@@ -3038,6 +3036,123 @@ struct RHI_API FRHITextureCreateInfo
 	uint8 NumSamples = 1;
 };
 
+/** Used to specify a texture metadata plane when creating a view. */
+enum class ERHITextureMetaDataAccess : uint8
+{
+	/** The primary plane is used with default compression behavior. */
+	None = 0,
+
+	/** The primary plane is used without decompressing it. */
+	CompressedSurface,
+
+	/** The depth plane is used with default compression behavior. */
+	Depth,
+
+	/** The stencil plane is used with default compression behavior. */
+	Stencil,
+
+	/** The HTile plane is used. */
+	HTile,
+
+	/** the FMask plane is used. */
+	FMask,
+
+	/** the CMask plane is used. */
+	CMask
+};
+
+enum ERHITextureSRVOverrideSRGBType
+{
+	SRGBO_Default,
+	SRGBO_ForceDisable,
+};
+
+struct FRHITextureSRVCreateInfo
+{
+	explicit FRHITextureSRVCreateInfo(uint8 InMipLevel = 0u, uint8 InNumMipLevels = 1u, EPixelFormat InFormat = PF_Unknown)
+		: Format(InFormat)
+		, MipLevel(InMipLevel)
+		, NumMipLevels(InNumMipLevels)
+		, SRGBOverride(SRGBO_Default)
+		, FirstArraySlice(0)
+		, NumArraySlices(0)
+	{}
+
+	explicit FRHITextureSRVCreateInfo(uint8 InMipLevel, uint8 InNumMipLevels, uint32 InFirstArraySlice, uint32 InNumArraySlices, EPixelFormat InFormat = PF_Unknown)
+		: Format(InFormat)
+		, MipLevel(InMipLevel)
+		, NumMipLevels(InNumMipLevels)
+		, SRGBOverride(SRGBO_Default)
+		, FirstArraySlice(InFirstArraySlice)
+		, NumArraySlices(InNumArraySlices)
+	{}
+
+	/** View the texture with a different format. Leave as PF_Unknown to use original format. Useful when sampling stencil */
+	EPixelFormat Format;
+
+	/** Specify the mip level to use. Useful when rendering to one mip while sampling from another */
+	uint8 MipLevel;
+
+	/** Create a view to a single, or multiple mip levels */
+	uint8 NumMipLevels;
+
+	/** Potentially override the texture's sRGB flag */
+	ERHITextureSRVOverrideSRGBType SRGBOverride;
+
+	/** Specify first array slice index. By default 0. */
+	uint32 FirstArraySlice;
+
+	/** Specify number of array slices. If FirstArraySlice and NumArraySlices are both zero, the SRV is created for all array slices. By default 0. */
+	uint32 NumArraySlices;
+
+	/** Specify the metadata plane to use when creating a view. */
+	ERHITextureMetaDataAccess MetaData = ERHITextureMetaDataAccess::None;
+
+	FORCEINLINE bool operator==(const FRHITextureSRVCreateInfo& Other)const
+	{
+		return (
+			Format == Other.Format &&
+			MipLevel == Other.MipLevel &&
+			NumMipLevels == Other.NumMipLevels &&
+			SRGBOverride == Other.SRGBOverride &&
+			FirstArraySlice == Other.FirstArraySlice &&
+			NumArraySlices == Other.NumArraySlices &&
+			MetaData == Other.MetaData);
+	}
+
+	FORCEINLINE bool operator!=(const FRHITextureSRVCreateInfo& Other)const
+	{
+		return !(*this == Other);
+	}
+};
+
+FORCEINLINE uint32 GetTypeHash(const FRHITextureSRVCreateInfo& Var)
+{
+	uint32 Hash0 = uint32(Var.Format) | uint32(Var.MipLevel) << 8 | uint32(Var.NumMipLevels) << 16 | uint32(Var.SRGBOverride) << 24;
+	return HashCombine(HashCombine(GetTypeHash(Hash0), GetTypeHash(Var.FirstArraySlice)), GetTypeHash(Var.NumArraySlices));
+}
+
+struct FRHITextureUAVCreateInfo
+{
+public:
+	explicit FRHITextureUAVCreateInfo(uint8 InMipLevel = 0u, ERHITextureMetaDataAccess InMetaData = ERHITextureMetaDataAccess::None)
+		: MipLevel(InMipLevel)
+		, MetaData(InMetaData)
+	{}
+
+	FORCEINLINE bool operator==(const FRHITextureUAVCreateInfo& Other)const
+	{
+		return MipLevel == Other.MipLevel && MetaData == Other.MetaData;
+	}
+
+	FORCEINLINE bool operator!=(const FRHITextureUAVCreateInfo& Other)const
+	{
+		return !(*this == Other);
+	}
+
+	uint8 MipLevel;
+	ERHITextureMetaDataAccess MetaData;
+};
 
 /** Descriptor used to create a buffer resource */
 struct FRHIBufferCreateInfo
@@ -3065,29 +3180,55 @@ struct FRHIBufferCreateInfo
 	EBufferUsageFlags Usage = BUF_None;
 };
 
+struct FRHIBufferSRVCreateInfo
+{
+	explicit FRHIBufferSRVCreateInfo() = default;
+
+	explicit FRHIBufferSRVCreateInfo(EPixelFormat InFormat)
+		: Format(InFormat)
+	{
+		BytesPerElement = GPixelFormats[Format].BlockBytes;
+	}
+
+	FORCEINLINE bool operator==(const FRHIBufferSRVCreateInfo& Other)const
+	{
+		return BytesPerElement == Other.BytesPerElement && Format == Other.Format;
+	}
+
+	FORCEINLINE bool operator!=(const FRHIBufferSRVCreateInfo& Other)const
+	{
+		return !(*this == Other);
+	}
+
+	/** Number of bytes per element. */
+	uint32 BytesPerElement = 1;
 
-/**
-@brief Resource allocator which is managed for a certain part of the frame by the high level
-	   rendering code (usually RenderGraph). An allocator is created via RHICreateTransientResourceAllocator on the device
-	   and released via the command list function RHIReleaseTransientResourceAllocator when all the
-	   transient resources are no longer needed.
-	   The allocator manages the lifetime of all the created resources and will take care of destroying
-	   all the resources when the allocator itself is destroyed.	   
-**/
-class RHI_API IRHITransientResourceAllocator
+	/** Encoding format for the element. */
+	EPixelFormat Format = PF_Unknown;
+};
+
+struct FRHIBufferUAVCreateInfo
 {
-public:
-	// Virtual destructor
-	virtual ~IRHITransientResourceAllocator() {}
+	FRHIBufferUAVCreateInfo() = default;
 
-	// Create functions
-	virtual FRHITexture* CreateTexture(const FRHITextureCreateInfo& InCreateInfo, const TCHAR* InDebugName) = 0;
-	virtual FRHIBuffer* CreateBuffer(const FRHIBufferCreateInfo& InCreateInfo, const TCHAR* InDebugName) = 0;
+	explicit FRHIBufferUAVCreateInfo(EPixelFormat InFormat)
+		: Format(InFormat)
+	{}
 
-	// Memory deallocation functions - doesn't destroy the resource, but frees internal memory for reuse by another transient resource
-	virtual void DeallocateMemory(FRHITexture* InTexture) = 0;
-	virtual void DeallocateMemory(FRHIBuffer* InBuffer) = 0;
+	FORCEINLINE bool operator==(const FRHIBufferUAVCreateInfo& Other)const
+	{
+		return Format == Other.Format && bSupportsAtomicCounter == Other.bSupportsAtomicCounter && bSupportsAppendBuffer == Other.bSupportsAppendBuffer;
+	}
+
+	FORCEINLINE bool operator!=(const FRHIBufferUAVCreateInfo& Other)const
+	{
+		return !(*this == Other);
+	}
+
+	/** Number of bytes per element (used for typed buffers). */
+	EPixelFormat Format = PF_Unknown;
 
-	// Freeze allocator for now creation and validates that all resources have their memory deallocated
-	virtual void Freeze(class FRHICommandListImmediate& RHICmdList) = 0;
+	/** Whether the uav supports atomic counter or append buffer operations (used for structured buffers) */
+	bool bSupportsAtomicCounter = false;
+	bool bSupportsAppendBuffer = false;
 };
\ No newline at end of file
diff --git a/Engine/Source/Runtime/RHI/Public/RHITransientResourceAllocator.h b/Engine/Source/Runtime/RHI/Public/RHITransientResourceAllocator.h
new file mode 100644
--- /dev/null
+++ b/Engine/Source/Runtime/RHI/Public/RHITransientResourceAllocator.h
@@ -0,0 +1,150 @@
+// Copyright Epic Games, Inc. All Rights Reserved.
+
+#pragma once
+
+#include "RHI.h"
+#include "RHIResources.h"
+
+class FRHICommandListImmediate;
+
+enum class ERHITransientResourceType
+{
+	Texture,
+	Buffer
+};
+
+class FRHITransientResource
+{
+public:
+	FRHITransientResource(FRHIResource* InResource, uint64 InHash, ERHITransientResourceType InType)
+		: Resource(InResource)
+		, Hash(InHash)
+		, Type(InType)
+	{}
+
+	virtual ~FRHITransientResource() = default;
+
+	// (Internal) Initializes the transient resource with a new allocation / name.
+	void Init(const TCHAR* InName, uint32 InAllocationIndex)
+	{
+		Name = InName;
+		AllocationIndex = InAllocationIndex;
+		AliasingOverlaps.Reset();
+	}
+
+	// (Internal) Adds a new transient resource overlap.
+	FORCEINLINE void AddAliasingOverlap(FRHITransientResource* InResource)
+	{
+		AliasingOverlaps.Emplace(InResource->GetRHI(), InResource->IsTexture() ? FRHITransientAliasingOverlap::EType::Texture : FRHITransientAliasingOverlap::EType::Buffer);
+	}
+
+	// Returns the underlying RHI resource.
+	FORCEINLINE FRHIResource* GetRHI() const { return Resource; }
+
+	// Returns the name assigned to the transient resource at allocation time.
+	FORCEINLINE const TCHAR* GetName() const { return Name; }
+
+	// (Internal)Returns the hash used to uniquely identify this resource if cached.
+	FORCEINLINE uint64 GetHash() const { return Hash; }
+
+	// (Internal) Returns the platform-specific allocation index.
+	FORCEINLINE uint32 GetAllocationIndex() const { return AllocationIndex; }
+
+	// Returns the aliasing overlaps for this resource.
+	FORCEINLINE TConstArrayView<FRHITransientAliasingOverlap> GetAliasingOverlaps() const { return AliasingOverlaps; }
+
+	FORCEINLINE ERHITransientResourceType GetType() const { return Type; }
+	FORCEINLINE bool IsTexture() const { return Type == ERHITransientResourceType::Texture; }
+	FORCEINLINE bool IsBuffer() const { return Type == ERHITransientResourceType::Buffer; }
+
+private:
+	// Underlying RHI resource.
+	TRefCountPtr<FRHIResource> Resource;
+
+	// The hash used to uniquely identify this resource if cached.
+	uint64 Hash;
+
+	// Debug name of the resource. Updated with each allocation.
+	const TCHAR* Name{};
+
+	// List of aliasing resources overlapping with this one.
+	TArray<FRHITransientAliasingOverlap> AliasingOverlaps;
+
+	// Type of the underlying RHI resource.
+	ERHITransientResourceType Type;
+
+	// An index to the underlying allocation info on the internal platform allocator.
+	uint32 AllocationIndex = ~0u;
+};
+
+class RHI_API FRHITransientTexture final : public FRHITransientResource
+{
+public:
+	FRHITransientTexture(FRHITexture* InTexture, uint64 InHash, const FRHITextureCreateInfo& InCreateInfo)
+		: FRHITransientResource(InTexture, InHash, ERHITransientResourceType::Texture)
+		, CreateInfo(InCreateInfo)
+	{}
+
+	// Returns the underlying RHI texture.
+	FRHITexture* GetRHI() const { return static_cast<FRHITexture*>(FRHITransientResource::GetRHI()); }
+
+	// Returns the create info struct used when creating this texture.
+	FORCEINLINE const FRHITextureCreateInfo& GetCreateInfo() const { return CreateInfo; }
+
+	// Finds a UAV matching the descriptor in the cache or creates a new one and updates the cache.
+	FRHIUnorderedAccessView* GetOrCreateUAV(const FRHITextureUAVCreateInfo& CreateInfo);
+
+	// Finds a SRV matching the descriptor in the cache or creates a new one and updates the cache.
+	FRHIShaderResourceView* GetOrCreateSRV(const FRHITextureSRVCreateInfo& CreateInfo);
+
+private:
+	FRHITextureCreateInfo CreateInfo;
+	TArray<TPair<FRHITextureUAVCreateInfo, FUnorderedAccessViewRHIRef>, TInlineAllocator<1>> UAVs;
+	TArray<TPair<FRHITextureSRVCreateInfo, FShaderResourceViewRHIRef>, TInlineAllocator<1>> SRVs;
+};
+
+class RHI_API FRHITransientBuffer final : public FRHITransientResource
+{
+public:
+	FORCEINLINE FRHITransientBuffer(FRHIBuffer* InBuffer, uint64 InHash, const FRHIBufferCreateInfo& InCreateInfo)
+		: FRHITransientResource(InBuffer, InHash, ERHITransientResourceType::Buffer)
+		, CreateInfo(InCreateInfo)
+	{}
+
+	// Returns the underlying RHI buffer.
+	FORCEINLINE FRHIBuffer* GetRHI() const { return static_cast<FRHIBuffer*>(FRHITransientResource::GetRHI()); }
+
+	// Returns the create info used when creating this buffer.
+	FORCEINLINE const FRHIBufferCreateInfo& GetCreateInfo() const { return CreateInfo; }
+
+	// Finds a UAV matching the descriptor in the cache or creates a new one and updates the cache.
+	FRHIUnorderedAccessView* GetOrCreateUAV(const FRHIBufferUAVCreateInfo& CreateInfo);
+
+	// Finds a SRV matching the descriptor in the cache or creates a new one and updates the cache.
+	FRHIShaderResourceView* GetOrCreateSRV(const FRHIBufferSRVCreateInfo& CreateInfo);
+
+private:
+	FRHIBufferCreateInfo CreateInfo;
+	TArray<TPair<FRHIBufferUAVCreateInfo, FUnorderedAccessViewRHIRef>, TInlineAllocator<1>> UAVs;
+	TArray<TPair<FRHIBufferSRVCreateInfo, FShaderResourceViewRHIRef>, TInlineAllocator<1>> SRVs;
+};
+
+class RHI_API IRHITransientResourceAllocator
+{
+public:
+	virtual ~IRHITransientResourceAllocator() = default;
+
+	// Allocates a new transient resource with memory backed by the transient allocator.
+	virtual FRHITransientTexture* CreateTexture(const FRHITextureCreateInfo& InCreateInfo, const TCHAR* InDebugName) = 0;
+	virtual FRHITransientBuffer* CreateBuffer(const FRHIBufferCreateInfo& InCreateInfo, const TCHAR* InDebugName) = 0;
+
+	// Deallocates the underlying memory for use by a future resource creation call.
+	virtual void DeallocateMemory(FRHITransientTexture* InTexture) = 0;
+	virtual void DeallocateMemory(FRHITransientBuffer* InBuffer) = 0;
+
+	// Freezes all allocations and validates that all resources have their memory deallocated.
+	virtual void Freeze(FRHICommandListImmediate& RHICmdList) = 0;
+
+	// Releases the transient allocator and deletes the instance. Any FRHITransientResource* access after this call is not allowed.
+	virtual void Release(FRHICommandListImmediate& RHICmdList);
+};
\ No newline at end of file
diff --git a/Engine/Source/Runtime/RHI/Public/RHIValidationCommon.h b/Engine/Source/Runtime/RHI/Public/RHIValidationCommon.h
--- a/Engine/Source/Runtime/RHI/Public/RHIValidationCommon.h
+++ b/Engine/Source/Runtime/RHI/Public/RHIValidationCommon.h
@@ -253,13 +253,13 @@ namespace RHIValidation
 		bool bTransient = false;
 		EStatus Status = EStatus::None;
 
-		bool IsAcquired() const
-		{
-			return Status == EStatus::Acquired;
-		}
+		FORCEINLINE bool IsAcquired() const { return Status == EStatus::Acquired; }
+		FORCEINLINE bool IsDiscarded() const { return Status == EStatus::Discarded; }
 
 		void Acquire(FResource* Resource, void* CreateTrace);
 		void Discard(FResource* Resource, void* CreateTrace);
+
+		static void AliasingOverlap(FResource* ResourceBefore, FResource* ResourceAfter, void* CreateTrace);
 	};
 
 	class FResource
@@ -269,6 +269,7 @@ namespace RHIValidation
 		friend FOperation;
 		friend FSubresourceState;
 		friend FSubresourceRange;
+		friend FTransientState;
 		friend FValidationRHI;
 
 	protected:
@@ -542,6 +543,7 @@ namespace RHIValidation
 	{
 		BeginTransition,
 		EndTransition,
+		AliasingOverlap,
 		AcquireTransient,
 		DiscardTransient,
 		Assert,
@@ -586,6 +588,13 @@ namespace RHIValidation
 				void* CreateBacktrace;
 			} Data_EndTransition;
 
+			struct
+			{
+				FResource* ResourceBefore;
+				FResource* ResourceAfter;
+				void* CreateBacktrace;
+			} Data_AliasingOverlap;
+
 			struct
 			{
 				FResource* Resource;
@@ -676,6 +685,19 @@ namespace RHIValidation
 			return MoveTemp(Op);
 		}
 
+		static inline FOperation AliasingOverlap(FResource* ResourceBefore, FResource* ResourceAfter, void* CreateBacktrace)
+		{
+			ResourceBefore->AddOpRef();
+			ResourceAfter->AddOpRef();
+
+			FOperation Op;
+			Op.Type = EOpType::AliasingOverlap;
+			Op.Data_AliasingOverlap.ResourceBefore = ResourceBefore;
+			Op.Data_AliasingOverlap.ResourceAfter = ResourceAfter;
+			Op.Data_AliasingOverlap.CreateBacktrace = CreateBacktrace;
+			return Op;
+		}
+
 		static inline FOperation AcquireTransientResource(FResource* Resource, void* CreateBacktrace)
 		{
 			Resource->AddOpRef();
diff --git a/Engine/Source/Runtime/RHI/Public/RHIValidationContext.h b/Engine/Source/Runtime/RHI/Public/RHIValidationContext.h
--- a/Engine/Source/Runtime/RHI/Public/RHIValidationContext.h
+++ b/Engine/Source/Runtime/RHI/Public/RHIValidationContext.h
@@ -904,11 +904,6 @@ public:
 		RHIContext->RHIWriteGPUFence(FenceRHI);
 	}
 
-	virtual void RHIReleaseTransientResourceAllocator(IRHITransientResourceAllocator* InAllocator) override final
-	{
-		RHIContext->RHIReleaseTransientResourceAllocator(InAllocator);
-	}
-
 	virtual void RHICopyToStagingBuffer(FRHIBuffer* SourceBufferRHI, FRHIStagingBuffer* DestinationStagingBufferRHI, uint32 InOffset, uint32 InNumBytes) override final
 	{
 		Tracker->Assert(SourceBufferRHI->GetWholeResourceIdentity(), ERHIAccess::CopySrc); // @todo: is CopySrc correct?
diff --git a/Engine/Source/Runtime/RHI/Public/RHIValidationTransientResourceAllocator.h b/Engine/Source/Runtime/RHI/Public/RHIValidationTransientResourceAllocator.h
--- a/Engine/Source/Runtime/RHI/Public/RHIValidationTransientResourceAllocator.h
+++ b/Engine/Source/Runtime/RHI/Public/RHIValidationTransientResourceAllocator.h
@@ -1,27 +1,27 @@
 // Copyright Epic Games, Inc. All Rights Reserved.
 
-/*=============================================================================
-	RHIValidation.h: Public Valdation RHI definitions.
-=============================================================================*/
-
 #pragma once 
 
-#include "RHIResources.h"
+#include "RHITransientResourceAllocator.h"
 
 #if ENABLE_RHI_VALIDATION
 
 class FValidationTransientResourceAllocator : public IRHITransientResourceAllocator
 {
 public:
-	FValidationTransientResourceAllocator(IRHITransientResourceAllocator* InRHIAllocator) : RHIAllocator(InRHIAllocator) {}
+	FValidationTransientResourceAllocator(IRHITransientResourceAllocator* InRHIAllocator)
+		: RHIAllocator(InRHIAllocator)
+	{}
+
 	virtual ~FValidationTransientResourceAllocator();
 
 	// Implementation of FRHITransientResourceAllocator interface
-	virtual FRHITexture* CreateTexture(const FRHITextureCreateInfo& InCreateInfo, const TCHAR* InDebugName) override final;
-	virtual FRHIBuffer* CreateBuffer(const FRHIBufferCreateInfo& InCreateInfo, const TCHAR* InDebugName) override final;
-	virtual void DeallocateMemory(FRHITexture* InTexture) override final;
-	virtual void DeallocateMemory(FRHIBuffer* InBuffer) override final;
+	virtual FRHITransientTexture* CreateTexture(const FRHITextureCreateInfo& InCreateInfo, const TCHAR* InDebugName) override final;
+	virtual FRHITransientBuffer* CreateBuffer(const FRHIBufferCreateInfo& InCreateInfo, const TCHAR* InDebugName) override final;
+	virtual void DeallocateMemory(FRHITransientTexture* InTexture) override final;
+	virtual void DeallocateMemory(FRHITransientBuffer* InBuffer) override final;
 	virtual void Freeze(FRHICommandListImmediate&) override final;
+	virtual void Release(FRHICommandListImmediate&) override final;
 
 private:
 
@@ -31,11 +31,11 @@ private:
 	// Actual RHI transient allocator which will get all functions forwarded
 	IRHITransientResourceAllocator* RHIAllocator = nullptr;
 
-	// Allocator already frozen?
 	bool bFrozen = false;
+	bool bReleased = false;
 
 	// All the allocated resources on the transient allocator
-	struct AllocatedResourceData
+	struct FAllocatedResourceData
 	{
 		enum class EType
 		{
@@ -56,7 +56,7 @@ private:
 			uint8 NumMips = 0;
 		} Texture;
 	};
-	TMap<FRHIResource*, AllocatedResourceData> AllocatedResourceMap;
+	TMap<FRHIResource*, FAllocatedResourceData> AllocatedResourceMap;
 };
 
 #endif	// ENABLE_RHI_VALIDATION
diff --git a/Engine/Source/Runtime/RHICore/Private/RHICoreTransientResourceAllocator.cpp b/Engine/Source/Runtime/RHICore/Private/RHICoreTransientResourceAllocator.cpp
new file mode 100644
--- /dev/null
+++ b/Engine/Source/Runtime/RHICore/Private/RHICoreTransientResourceAllocator.cpp
@@ -0,0 +1,673 @@
+// Copyright Epic Games, Inc. All Rights Reserved.
+
+#include "RHICoreTransientResourceAllocator.h"
+
+static int32 GRHITransientAllocatorMinimumHeapSize = 128;
+static FAutoConsoleVariableRef CVarGRHITransientAllocatorMinimumHeapSize(
+	TEXT("RHI.TransientAllocator.MinimumHeapSize"),
+	GRHITransientAllocatorMinimumHeapSize,
+	TEXT("Minimum size of an RHI transient heap in MB. Heaps will default to this size and grow to the maximum based on the first allocation (Default 128)."),
+	ECVF_ReadOnly);
+
+static int32 GRHITransientAllocatorMaximumHeapSize = 512;
+static FAutoConsoleVariableRef CVarRHITransientAllocatorMaximumHeapSize(
+	TEXT("RHI.TransientAllocator.MaximumHeapSize"),
+	GRHITransientAllocatorMaximumHeapSize,
+	TEXT("Maximum size of an RHI transient allocation in MB. Allocations larger than this will fail the transient allocator (Default 512)."),
+	ECVF_ReadOnly);
+
+static int32 GRHITransientAllocatorBufferCacheSize = 64;
+static FAutoConsoleVariableRef CVarRHITransientAllocatorBufferCacheSize(
+	TEXT("RHI.TransientAllocator.BufferCacheSize"),
+	GRHITransientAllocatorBufferCacheSize,
+	TEXT("The maximum number of RHI buffers to cache on each heap before garbage collecting."),
+	ECVF_ReadOnly);
+
+static int32 GRHITransientAllocatorTextureCacheSize = 64;
+static FAutoConsoleVariableRef CVarRHITransientAllocatorTextureCacheSize(
+	TEXT("RHI.TransientAllocator.TextureCacheSize"),
+	GRHITransientAllocatorTextureCacheSize,
+	TEXT("The maximum number of RHI textures to cache on each heap before garbage collecting."),
+	ECVF_ReadOnly);
+
+DECLARE_STATS_GROUP(TEXT("RHI: Transient Memory"), STATGROUP_RHITransientMemory, STATCAT_Advanced);
+
+DECLARE_MEMORY_STAT(TEXT("Memory Allocated"), STAT_RHITransientMemoryAllocated, STATGROUP_RHITransientMemory);
+DECLARE_MEMORY_STAT(TEXT("Memory Requested"), STAT_RHITransientMemoryRequested, STATGROUP_RHITransientMemory);
+DECLARE_MEMORY_STAT(TEXT("Memory Used"), STAT_RHITransientMemoryUsed, STATGROUP_RHITransientMemory);
+DECLARE_MEMORY_STAT(TEXT("Buffer Memory Used"), STAT_RHITransientBufferMemoryUsed, STATGROUP_RHITransientMemory);
+DECLARE_MEMORY_STAT(TEXT("Buffer Memory Requested"), STAT_RHITransientBufferMemoryRequested, STATGROUP_RHITransientMemory);
+DECLARE_MEMORY_STAT(TEXT("Texture Memory Used"), STAT_RHITransientTextureMemoryUsed, STATGROUP_RHITransientMemory);
+DECLARE_MEMORY_STAT(TEXT("Texture Memory Requested"), STAT_RHITransientTextureMemoryRequested, STATGROUP_RHITransientMemory);
+
+DECLARE_DWORD_ACCUMULATOR_STAT(TEXT("Resources"), STAT_RHITransientResources, STATGROUP_RHITransientMemory);
+DECLARE_DWORD_ACCUMULATOR_STAT(TEXT("Textures"), STAT_RHITransientTextures, STATGROUP_RHITransientMemory);
+DECLARE_DWORD_ACCUMULATOR_STAT(TEXT("Buffers"), STAT_RHITransientBuffers, STATGROUP_RHITransientMemory);
+DECLARE_DWORD_ACCUMULATOR_STAT(TEXT("Heaps"), STAT_RHITransientHeaps, STATGROUP_RHITransientMemory);
+
+inline uint64 ComputeHash(const FRHITextureCreateInfo& InCreateInfo, uint64 HeapOffset)
+{
+	// Make sure all padding is removed.
+	FRHITextureCreateInfo NewInfo;
+	FPlatformMemory::Memzero(&NewInfo, sizeof(FRHITextureCreateInfo));
+	NewInfo = InCreateInfo;
+	return CityHash64WithSeed((const char*)&NewInfo, sizeof(FRHITextureCreateInfo), HeapOffset);
+}
+
+inline uint64 ComputeHash(const FRHIBufferCreateInfo& InCreateInfo, uint64 HeapOffset)
+{
+	return CityHash64WithSeed((const char*)&InCreateInfo, sizeof(FRHIBufferCreateInfo), HeapOffset);
+}
+
+//////////////////////////////////////////////////////////////////////////
+
+template <typename CreateFunctionType>
+FRHITransientTexture* FRHITransientHeap::AcquireTexture(const FRHITextureCreateInfo& CreateInfo, uint64 HeapOffset, CreateFunctionType CreateFunction)
+{
+	const uint64 Hash = ComputeHash(CreateInfo, HeapOffset);
+	FRHITransientTexture* Texture = Textures.Acquire(Hash, CreateFunction);
+	AllocatedTextures.Emplace(Texture);
+	return Texture;
+}
+
+template <typename CreateFunctionType>
+FRHITransientBuffer* FRHITransientHeap::AcquireBuffer(const FRHIBufferCreateInfo& CreateInfo, uint64 HeapOffset, CreateFunctionType CreateFunction)
+{
+	const uint64 Hash = ComputeHash(CreateInfo, HeapOffset);
+	FRHITransientBuffer* Buffer = Buffers.Acquire(Hash, CreateFunction);
+	AllocatedBuffers.Emplace(Buffer);
+	return Buffer;
+}
+
+void FRHITransientHeap::ForfeitResources()
+{
+	Textures.Forfeit(AllocatedTextures);
+	AllocatedTextures.Reset();
+
+	Buffers.Forfeit(AllocatedBuffers);
+	AllocatedBuffers.Reset();
+}
+
+//////////////////////////////////////////////////////////////////////////
+
+FRHITransientResourceSystemInitializer FRHITransientResourceSystemInitializer::CreateDefault()
+{
+	FRHITransientResourceSystemInitializer Initializer;
+	Initializer.MinimumHeapSize = GRHITransientAllocatorMinimumHeapSize * 1024 * 1024;
+	Initializer.MaximumHeapSize = GRHITransientAllocatorMaximumHeapSize * 1024 * 1024;
+	Initializer.BufferCacheSize = GRHITransientAllocatorBufferCacheSize;
+	Initializer.TextureCacheSize = GRHITransientAllocatorTextureCacheSize;
+	return Initializer;
+}
+
+FRHITransientResourceSystem::~FRHITransientResourceSystem()
+{
+	ReleaseHeaps();
+}
+
+FRHITransientHeap* FRHITransientResourceSystem::AcquireHeap(uint64 FirstAllocationSize, ERHITransientHeapFlags FirstAllocationHeapFlags)
+{
+	{
+		FScopeLock Lock(&HeapCriticalSection);
+
+		for (int32 HeapIndex = 0; HeapIndex < Heaps.Num(); ++HeapIndex)
+		{
+			FRHITransientHeap* Heap = Heaps[HeapIndex];
+
+			if (Heap->IsAllocationSupported(FirstAllocationSize, FirstAllocationHeapFlags))
+			{
+				Heaps.RemoveAt(HeapIndex);
+				return Heap;
+			}
+		}
+	}
+
+	FRHITransientHeapInitializer HeapInitializer;
+	HeapInitializer.Size = GetHeapSize(FirstAllocationSize);
+	HeapInitializer.Alignment = GetHeapAlignment();
+	HeapInitializer.Flags = (Initializer.bSupportsAllHeapFlags ? ERHITransientHeapFlags::AllowAll : FirstAllocationHeapFlags);
+	HeapInitializer.TextureCacheSize = Initializer.TextureCacheSize;
+	HeapInitializer.BufferCacheSize = Initializer.BufferCacheSize;
+
+	return CreateHeap(HeapInitializer);
+}
+
+void FRHITransientResourceSystem::ForfeitHeaps(TArrayView<FRHITransientHeap* const> InForfeitedHeaps)
+{
+	FScopeLock Lock(&HeapCriticalSection);
+
+	for (FRHITransientHeap* Heap : InForfeitedHeaps)
+	{
+		Heap->ForfeitResources();
+		Heap->LastUsedGarbageCollectCycle = GarbageCollectCycle;
+	}
+
+	Heaps.Append(InForfeitedHeaps.GetData(), InForfeitedHeaps.Num());
+
+	Algo::Sort(Heaps, [](const FRHITransientHeap* LHS, const FRHITransientHeap* RHS)
+	{
+		// Sort by smaller heap first.
+		if (LHS->GetCapacity() != RHS->GetCapacity())
+		{
+			return LHS->GetCapacity() < RHS->GetCapacity();
+		}
+
+		// Sort next by most recently used first.
+		return LHS->GetLastUsedGarbageCollectCycle() >= RHS->GetLastUsedGarbageCollectCycle();
+	});
+}
+
+void FRHITransientResourceSystem::GarbageCollect()
+{
+	FScopeLock Lock(&HeapCriticalSection);
+
+	for (int32 HeapIndex = 0; HeapIndex < Heaps.Num(); ++HeapIndex)
+	{
+		FRHITransientHeap* Heap = Heaps[HeapIndex];
+
+		if (Heap->GetLastUsedGarbageCollectCycle() + Initializer.GarbageCollectLatency <= GarbageCollectCycle)
+		{
+			Heaps.RemoveAt(HeapIndex);
+			HeapIndex--;
+
+			delete Heap;
+		}
+	}
+
+	GarbageCollectCycle++;
+}
+
+void FRHITransientResourceSystem::ReleaseHeaps()
+{
+	FScopeLock Lock(&HeapCriticalSection);
+	for (FRHITransientHeap* Heap : Heaps)
+	{
+		delete Heap;
+	}
+	Heaps.Empty();
+}
+
+void FRHITransientResourceSystem::UpdateStats()
+{
+	FStats Stats;
+	Stats.Textures = TextureStats;
+	Stats.Buffers = BufferStats;
+
+	uint32 HeapCount;
+
+	{
+		FScopeLock Lock(&HeapCriticalSection);
+		for (FRHITransientHeap* Heap : Heaps)
+		{
+			Stats.TotalMemoryUsed += Heap->GetCapacity();
+		}
+		HeapCount = Heaps.Num();
+	}
+
+	ReportStats(Stats);
+
+	SET_MEMORY_STAT(STAT_RHITransientMemoryAllocated, Stats.TotalMemoryUsed);
+	SET_MEMORY_STAT(STAT_RHITransientMemoryRequested, Stats.Textures.TotalSize + Stats.Buffers.TotalSize);
+	SET_MEMORY_STAT(STAT_RHITransientMemoryUsed, Stats.Textures.TotalSizeWithAliasing + Stats.Buffers.TotalSizeWithAliasing);
+	SET_MEMORY_STAT(STAT_RHITransientBufferMemoryRequested, Stats.Buffers.TotalSize);
+	SET_MEMORY_STAT(STAT_RHITransientBufferMemoryUsed, Stats.Buffers.TotalSizeWithAliasing);
+	SET_MEMORY_STAT(STAT_RHITransientTextureMemoryRequested, Stats.Textures.TotalSize);
+	SET_MEMORY_STAT(STAT_RHITransientTextureMemoryUsed, Stats.Textures.TotalSizeWithAliasing);
+
+	SET_DWORD_STAT(STAT_RHITransientTextures, Stats.Textures.AllocationCount);
+	SET_DWORD_STAT(STAT_RHITransientBuffers, Stats.Buffers.AllocationCount);
+	SET_DWORD_STAT(STAT_RHITransientResources, Stats.Textures.AllocationCount + Stats.Buffers.AllocationCount);
+	SET_DWORD_STAT(STAT_RHITransientHeaps, HeapCount);
+
+	TextureStats = {};
+	BufferStats = {};
+}
+
+//////////////////////////////////////////////////////////////////////////
+
+FRHITransientHeapAllocator::FRHITransientHeapAllocator(const FRHITransientHeapInitializer& InInitializer, uint32 InHeapIndex)
+	: Initializer(InInitializer)
+	, HeapIndex(InHeapIndex)
+{
+	HeadHandle = CreateRange();
+	InsertRange(HeadHandle, 0, Initializer.Size);
+}
+
+FRHITransientHeapAllocator::~FRHITransientHeapAllocator()
+{
+#if RHICORE_TRANSIENT_ALLOCATOR_DEBUG
+	checkf(AllocationCount == 0, TEXT("%d allocations remain on heap."), AllocationCount);
+	check(HeadHandle != InvalidRangeHandle);
+
+	FRangeHandle FirstFreeHandle = GetFirstFreeRangeHandle();
+	check(FirstFreeHandle != InvalidRangeHandle);
+
+	const FRange& FirstFreeRange = Ranges[FirstFreeHandle];
+	check(FirstFreeRange.NextFreeHandle == InvalidRangeHandle);
+	check(FirstFreeRange.Offset == 0);
+	check(FirstFreeRange.Size == GetCapacity());
+#endif
+}
+
+FRHITransientHeapAllocation FRHITransientHeapAllocator::Allocate(uint64 Size, uint32 Alignment)
+{
+	check(Size > 0);
+
+	Alignment = FMath::Max(Alignment, Initializer.Alignment);
+
+	FFindResult FindResult = FindFreeRange(Size, Alignment);
+
+	if (FindResult.FoundHandle == InvalidRangeHandle)
+	{
+		return {};
+	}
+
+	FRange& FoundRange = Ranges[FindResult.FoundHandle];
+
+	const uint64 AlignedSize   = FoundRange.Size   - FindResult.LeftoverSize;
+	const uint64 AlignmentPad  = AlignedSize       - Size;
+	const uint64 AlignedOffset = FoundRange.Offset + AlignmentPad;
+	const uint64 AllocationEnd = AlignedOffset     + Size;
+
+	// Adjust the range if there is space left over.
+	if (FindResult.LeftoverSize)
+	{
+		FoundRange.Offset = AllocationEnd;
+		FoundRange.Size  = FindResult.LeftoverSize;
+	}
+	// Otherwise, remove it.
+	else
+	{
+		RemoveRange(FindResult.PreviousHandle, FindResult.FoundHandle);
+	}
+
+	AllocationCount++;
+	UsedSize       += AlignedSize;
+	AlignmentWaste += AlignmentPad;
+
+	FRHITransientHeapAllocation Allocation;
+	Allocation.Size         = Size;
+	Allocation.Offset       = AlignedOffset;
+	Allocation.AlignmentPad = AlignmentPad;
+	Allocation.HeapIndex    = HeapIndex;
+
+	Validate();
+
+	return Allocation;
+}
+
+void FRHITransientHeapAllocator::Deallocate(FRHITransientHeapAllocation Allocation)
+{
+	check(Allocation.Size > 0 && Allocation.Size <= UsedSize);
+
+	// Reconstruct the original range offset by subtracting the alignment pad, and expand the size accordingly.
+	const uint64 RangeToFreeOffset = Allocation.Offset - Allocation.AlignmentPad;
+	const uint64 RangeToFreeSize   = Allocation.Size   + Allocation.AlignmentPad;
+	const uint64 RangeToFreeEnd    = RangeToFreeOffset + RangeToFreeSize;
+
+	FRangeHandle PreviousHandle = HeadHandle;
+	FRangeHandle NextHandle     = InvalidRangeHandle;
+	FRangeHandle Handle         = GetFirstFreeRangeHandle();
+
+	while (Handle != InvalidRangeHandle)
+	{
+		const FRange& Range = Ranges[Handle];
+
+		// Find the first free range after the one being freed.
+		if (RangeToFreeOffset < Range.Offset)
+		{
+			NextHandle = Handle;
+			break;
+		}
+
+		PreviousHandle = Handle;
+		Handle         = Range.NextFreeHandle;
+	}
+
+	uint64 MergedFreeRangeStart = RangeToFreeOffset;
+	uint64 MergedFreeRangeEnd   = RangeToFreeEnd;
+	bool bMergedPrevious        = false;
+	bool bMergedNext            = false;
+
+	if (PreviousHandle != HeadHandle)
+	{
+		FRange& PreviousRange = Ranges[PreviousHandle];
+
+		// Attempt to merge the previous range with the range being freed.
+		if (PreviousRange.GetEnd() == RangeToFreeOffset)
+		{
+			PreviousRange.Size   += RangeToFreeSize;
+			MergedFreeRangeStart  = PreviousRange.Offset;
+			MergedFreeRangeEnd    = PreviousRange.GetEnd();
+			bMergedPrevious       = true;
+		}
+	}
+
+	if (NextHandle != InvalidRangeHandle)
+	{
+		FRange& NextRange = Ranges[NextHandle];
+
+		// Attempt to merge the next range with the range being freed.
+		if (RangeToFreeEnd == NextRange.Offset)
+		{
+			NextRange.Size      += RangeToFreeSize;
+			NextRange.Offset     = RangeToFreeOffset;
+			MergedFreeRangeStart = FMath::Min(MergedFreeRangeStart, RangeToFreeOffset);
+			MergedFreeRangeEnd   = NextRange.GetEnd();
+			bMergedNext          = true;
+		}
+	}
+
+	// With both previous and next ranges merged with the freed range, they now overlap. Remove next and expand previous to cover all three.
+	if (bMergedPrevious && bMergedNext)
+	{
+		FRange& PreviousRange = Ranges[PreviousHandle];
+		FRange& NextRange     = Ranges[NextHandle];
+
+		PreviousRange.Size = MergedFreeRangeEnd - MergedFreeRangeStart;
+		RemoveRange(PreviousHandle, NextHandle);
+	}
+	// If neither previous or next were merged, insert a new range between them.
+	else if (!bMergedPrevious && !bMergedNext)
+	{
+		InsertRange(PreviousHandle, RangeToFreeOffset, RangeToFreeSize);
+	}
+
+	UsedSize       -= RangeToFreeSize;
+	AlignmentWaste -= Allocation.AlignmentPad;
+	AllocationCount--;
+
+	Validate();
+}
+
+FRHITransientHeapAllocator::FFindResult FRHITransientHeapAllocator::FindFreeRange(uint64 Size, uint32 Alignment)
+{
+	FFindResult FindResult;
+	FindResult.PreviousHandle = HeadHandle;
+
+	FRangeHandle Handle = GetFirstFreeRangeHandle();
+	while (Handle != InvalidRangeHandle)
+	{
+		FRange& Range = Ranges[Handle];
+
+		// Due to alignment we may have to shift the offset and expand the size accordingly.
+		const uint64 AlignmentPad = Align(Range.Offset, Alignment) - Range.Offset;
+		const uint64 RequiredSize = Size + AlignmentPad;
+
+		if (RequiredSize <= Range.Size)
+		{
+			FindResult.FoundHandle  = Handle;
+			FindResult.LeftoverSize = Range.Size - RequiredSize;
+			return FindResult;
+		}
+
+		FindResult.PreviousHandle = Handle;
+		Handle                    = Range.NextFreeHandle;
+	}
+
+	return {};
+}
+
+void FRHITransientHeapAllocator::Validate()
+{
+#if RHICORE_TRANSIENT_ALLOCATOR_DEBUG
+	uint64 DerivedFreeSize = 0;
+
+	FRangeHandle PreviousHandle = HeadHandle;
+	FRangeHandle NextHandle     = InvalidRangeHandle;
+	FRangeHandle Handle         = GetFirstFreeRangeHandle();
+
+	while (Handle != InvalidRangeHandle)
+	{
+		const FRange& Range = Ranges[Handle];
+		DerivedFreeSize    += Range.Size;
+
+		if (PreviousHandle != HeadHandle)
+		{
+			const FRange& PreviousRange = Ranges[PreviousHandle];
+
+			// Checks that the ranges are sorted.
+			check(PreviousRange.Offset + PreviousRange.Size < Range.Offset);
+		}
+
+		PreviousHandle = Handle;
+		Handle         = Range.NextFreeHandle;
+	}
+
+	check(Initializer.Size == DerivedFreeSize + UsedSize);
+#endif
+}
+
+//////////////////////////////////////////////////////////////////////////
+
+void FRHITransientResourceOverlapTracker::Track(FRHITransientResource* TransientResource, uint64 OffsetMin, uint64 OffsetMax)
+{
+	check(TransientResource);
+
+	FResourceRange ResourceRangeNew;
+	ResourceRangeNew.Resource  = TransientResource;
+	ResourceRangeNew.OffsetMin = OffsetMin;
+	ResourceRangeNew.OffsetMax = OffsetMax;
+
+	for (int32 Index = 0; Index < ResourceRanges.Num(); ++Index)
+	{
+		FResourceRange& ResourceRangeOld = ResourceRanges[Index];
+
+		// If the old range starts later in the heap and doesn't overlap, the sort invariant guarantees no future range will overlap.
+		if (ResourceRangeOld.OffsetMin >= ResourceRangeNew.OffsetMin)
+		{
+			break;
+		}
+
+		// If the old range starts earlier in the heap and doesn't overlap, we keep searching.
+		if (ResourceRangeOld.OffsetMax <= ResourceRangeNew.OffsetMin)
+		{
+			continue;
+		}
+
+		TransientResource->AddAliasingOverlap(ResourceRangeOld.Resource);
+
+		// Complete overlap.
+		if (ResourceRangeOld.OffsetMin >= ResourceRangeNew.OffsetMin && ResourceRangeOld.OffsetMax <= ResourceRangeNew.OffsetMax)
+		{
+			ResourceRanges.RemoveAt(Index);
+			Index--;
+		}
+		// Partial overlap, can manifest as three cases:
+		else
+		{
+			bool bResizedOld = false;
+
+			// 1) New:    ********
+			//            |||        ->
+			//    Old: ======             ===********
+			if (ResourceRangeOld.OffsetMin < ResourceRangeNew.OffsetMin)
+			{
+				ResourceRangeOld.OffsetMax = ResourceRangeNew.OffsetMin;
+				bResizedOld = true;
+			}
+
+			if (ResourceRangeOld.OffsetMax > ResourceRangeNew.OffsetMax)
+			{
+				// 2) New:    ********
+				//                |||      ->
+				//    Old:        ======         ********===
+				if (!bResizedOld)
+				{
+					ResourceRangeOld.OffsetMin = ResourceRangeNew.OffsetMax;
+				}
+
+				// 3) New:    ********
+				//            ||||||||      ->
+				//    Old: ==============        ===********===
+				else
+				{
+					// Lower bound has been resized already; add an upper bound.
+					FResourceRange ResourceRangeOldUpper = ResourceRangeOld;
+					ResourceRangeOldUpper.OffsetMin = ResourceRangeNew.OffsetMax;
+					Index = ResourceRanges.Insert(ResourceRangeOldUpper, Index + 1);
+				}
+			}
+		}
+	}
+
+	for (int32 Index = 0; Index < ResourceRanges.Num(); ++Index)
+	{
+		const FResourceRange& ResourceRangeOld = ResourceRanges[Index];
+
+		// Insert new range while preserving order.
+		if (ResourceRangeOld.OffsetMin > ResourceRangeNew.OffsetMin)
+		{
+			ResourceRanges.Insert(ResourceRangeNew, Index);
+			return;
+		}
+	}
+
+	ResourceRanges.Add(ResourceRangeNew);
+}
+
+//////////////////////////////////////////////////////////////////////////
+
+FRHITransientResourceAllocator::FRHITransientResourceAllocator(FRHITransientResourceSystem& ParentInSystem)
+	: ParentSystem(ParentInSystem)
+{
+	Heaps.Reserve(ParentSystem.GetHeapCount());
+	HeapAllocators.Reserve(ParentSystem.GetHeapCount());
+}
+
+FRHITransientTexture* FRHITransientResourceAllocator::CreateTexture(
+	const FRHITextureCreateInfo& CreateInfo,
+	const TCHAR* DebugName,
+	uint64 TextureSize,
+	uint32 TextureAlignment,
+	FCreateTextureFunction CreateTextureFunction)
+{
+	if (TextureSize > ParentSystem.GetMaximumHeapSize())
+	{
+		return {};
+	}
+
+	const ERHITransientHeapFlags TextureHeapFlags =
+		EnumHasAnyFlags(CreateInfo.Flags, TexCreate_RenderTargetable | TexCreate_ResolveTargetable | TexCreate_DepthStencilTargetable | TexCreate_DepthStencilResolveTarget)
+		? ERHITransientHeapFlags::AllowRenderTargets
+		: ERHITransientHeapFlags::AllowTextures;
+
+	const FRHITransientHeapAllocation Allocation = Allocate(TextureStats, TextureSize, TextureAlignment, TextureHeapFlags);
+
+	FRHITransientHeap* Heap = Heaps[Allocation.HeapIndex];
+
+	FRHITransientTexture* TransientTexture = Heap->AcquireTexture(CreateInfo, Allocation.Offset, [&] (uint64 Hash)
+	{
+		const FResourceInitializer ResourceInitializer(*Heap, Allocation, Hash);
+		return CreateTextureFunction(ResourceInitializer);
+	});
+
+	InitResource(TransientTexture, Allocation, DebugName);
+
+#if RHICORE_TRANSIENT_ALLOCATOR_DEBUG
+	DebugTextures.Emplace(TransientTexture);
+	check(TransientTexture->GetCreateInfo() == CreateInfo && TransientTexture->GetName() == DebugName);
+#endif
+
+	return MoveTemp(TransientTexture);
+}
+
+FRHITransientBuffer* FRHITransientResourceAllocator::CreateBuffer(
+	const FRHIBufferCreateInfo& CreateInfo,
+	const TCHAR* DebugName,
+	uint32 BufferSize,
+	uint32 BufferAlignment,
+	FCreateBufferFunction CreateBufferFunction)
+{
+	if (BufferSize > ParentSystem.GetMaximumHeapSize())
+	{
+		return {};
+	}
+
+	const FRHITransientHeapAllocation Allocation = Allocate(BufferStats, BufferSize, BufferAlignment, ERHITransientHeapFlags::AllowBuffers);
+
+	FRHITransientHeap* Heap = Heaps[Allocation.HeapIndex];
+
+	FRHITransientBuffer* TransientBuffer = Heap->AcquireBuffer(CreateInfo, Allocation.Offset, [&](uint64 Hash)
+	{
+		const FResourceInitializer ResourceInitializer(*Heap, Allocation, Hash);
+		return CreateBufferFunction(ResourceInitializer);
+	});
+
+	InitResource(TransientBuffer, Allocation, DebugName);
+
+#if RHICORE_TRANSIENT_ALLOCATOR_DEBUG
+	DebugBuffers.Emplace(TransientBuffer);
+	check(TransientBuffer->GetCreateInfo() == CreateInfo && TransientBuffer->GetName() == DebugName);
+#endif
+
+	return TransientBuffer;
+}
+
+void FRHITransientResourceAllocator::DeallocateMemoryInternal(FRHITransientResource* InResource, FMemoryStats& StatsToUpdate)
+{
+	const FRHITransientHeapAllocation& Allocation = HeapAllocations[InResource->GetAllocationIndex()];
+
+	HeapAllocators[Allocation.HeapIndex].Deallocate(Allocation);
+
+	StatsToUpdate.CurrentSizeWithAliasing -= Allocation.Size;
+}
+
+void FRHITransientResourceAllocator::Freeze(FRHICommandListImmediate& RHICmdList)
+{
+	ParentSystem.ForfeitHeaps(Heaps);
+
+	RHICmdList.EnqueueLambda([&ParentSystem = ParentSystem, BufferStats = BufferStats, TextureStats = TextureStats](FRHICommandListImmediate&)
+	{
+		ParentSystem.TextureStats.Add(TextureStats);
+		ParentSystem.BufferStats.Add(BufferStats);
+	});
+}
+
+FRHITransientHeapAllocation FRHITransientResourceAllocator::Allocate(FMemoryStats& StatsToUpdate, uint64 Size, uint32 Alignment, ERHITransientHeapFlags ResourceHeapFlags)
+{
+	FRHITransientHeapAllocation Allocation;
+
+	for (int32 Index = 0; Index < HeapAllocators.Num(); ++Index)
+	{
+		FRHITransientHeapAllocator& HeapAllocator = HeapAllocators[Index];
+
+		if (!HeapAllocator.IsAllocationSupported(Size, ResourceHeapFlags))
+		{
+			continue;
+		}
+
+		Allocation = HeapAllocator.Allocate(Size, Alignment);
+
+		if (Allocation.IsValid())
+		{
+			break;
+		}
+	}
+
+	if (!Allocation.IsValid())
+	{
+		const uint32 HeapIndex = Heaps.Num();
+
+		FRHITransientHeap* Heap = ParentSystem.AcquireHeap(Size, ResourceHeapFlags);
+		Heaps.Emplace(Heap);
+
+		FRHITransientHeapAllocator& HeapAllocator = HeapAllocators.Emplace_GetRef(Heap->GetInitializer(), HeapIndex);
+		Allocation = HeapAllocator.Allocate(Size, Alignment);
+
+		check(Allocation.IsValid());
+	}
+
+	StatsToUpdate.AllocationCount++;
+	StatsToUpdate.TotalSize += Allocation.Size;
+	StatsToUpdate.CurrentSizeWithAliasing += Allocation.Size;
+	StatsToUpdate.TotalSizeWithAliasing = FMath::Max(StatsToUpdate.TotalSizeWithAliasing, StatsToUpdate.CurrentSizeWithAliasing);
+
+	return Allocation;
+}
+
+void FRHITransientResourceAllocator::InitResource(FRHITransientResource* TransientResource, const FRHITransientHeapAllocation& Allocation, const TCHAR* Name)
+{
+	TransientResource->Init(Name, HeapAllocations.Emplace(Allocation));
+	HeapAllocators[Allocation.HeapIndex].TrackOverlap(TransientResource, Allocation);
+}
\ No newline at end of file
diff --git a/Engine/Source/Runtime/RHICore/Public/RHICoreTransientResourceAllocator.h b/Engine/Source/Runtime/RHICore/Public/RHICoreTransientResourceAllocator.h
new file mode 100644
--- /dev/null
+++ b/Engine/Source/Runtime/RHICore/Public/RHICoreTransientResourceAllocator.h
@@ -0,0 +1,598 @@
+// Copyright Epic Games, Inc. All Rights Reserved.
+#pragma once
+
+#include "RHITransientResourceAllocator.h"
+
+class FRHITransientResourceSystem;
+class FRHITransientResourceAllocator;
+
+#define RHICORE_TRANSIENT_ALLOCATOR_DEBUG (!UE_BUILD_SHIPPING && !UE_BUILD_TEST)
+
+#if RHICORE_TRANSIENT_ALLOCATOR_DEBUG
+	#define IF_RHICORE_TRANSIENT_ALLOCATOR_DEBUG(Op) Op
+#else
+	#define IF_RHICORE_TRANSIENT_ALLOCATOR_DEBUG(Op)
+#endif
+
+struct FRHITransientMemoryStats
+{
+	FRHITransientMemoryStats() = default;
+
+	void Add(const FRHITransientMemoryStats& Other)
+	{
+		TotalSizeWithAliasing += Other.TotalSizeWithAliasing;
+		TotalSize += Other.TotalSize;
+		AllocationCount += Other.AllocationCount;
+	}
+
+	// Total allocated memory usage with aliasing.
+	uint64 TotalSizeWithAliasing = 0;
+
+	// Total allocated memory usage without aliasing.
+	uint64 TotalSize = 0;
+
+	// The number of allocations made from the transient allocator.
+	uint32 AllocationCount = 0;
+};
+
+/** An RHI transient resource cache designed to optimize fetches for resources placed into a heap with an offset.
+ *  The cache has a fixed capacity whereby no garbage collection will occur. Once that capacity is exceeded, garbage
+ *  collection is invoked on resources older than a specified generation (where generation is incremented with each
+ *  cycle of forfeiting acquired resources).
+ */
+template <typename TransientResourceType>
+class TRHITransientResourceCache
+{
+public:
+	static const uint32 kInfinity = ~0u;
+	static const uint32 kDefaultCapacity = kInfinity;
+	static const uint32 kDefaultGarbageCollectLatency = 32;
+
+	TRHITransientResourceCache(uint32 InCapacity = kDefaultCapacity, uint32 InGarbageCollectLatency = kDefaultGarbageCollectLatency)
+		: GarbageCollectLatency(InGarbageCollectLatency)
+		, Capacity(InCapacity)
+	{
+		if (Capacity != kInfinity)
+		{
+			Cache.Reserve(Capacity);
+		}
+	}
+
+	~TRHITransientResourceCache()
+	{
+		for (const FCacheItem& Item : Cache)
+		{
+			delete Item.Resource;
+		}
+	}
+
+	template <typename CreateFunctionType>
+	TransientResourceType* Acquire(uint64 Hash, CreateFunctionType CreateFunction)
+	{
+		for (int32 Index = 0; Index < Cache.Num(); ++Index)
+		{
+			const FCacheItem& CacheItem = Cache[Index];
+
+			if (CacheItem.Hash == Hash)
+			{
+				TransientResourceType* Resource = CacheItem.Resource;
+				Cache.RemoveAtSwap(Index);
+				return Resource;
+			}
+		}
+
+		return CreateFunction(Hash);
+	}
+
+	void Forfeit(TArrayView<TransientResourceType* const> Resources)
+	{
+		CurrentGeneration++;
+
+		for (TransientResourceType* Resource : Resources)
+		{
+			Cache.Emplace(Resource, Resource->GetHash(), CurrentGeneration);
+		}
+
+		Algo::Sort(Cache, [](const FCacheItem& LHS, const FCacheItem& RHS)
+		{
+			return LHS.Generation > RHS.Generation;
+		});
+
+		while (uint32(Cache.Num()) > Capacity)
+		{
+			if (!TryReleaseItem())
+			{
+				break;
+			}
+		}
+	}
+
+	FORCEINLINE uint32 GetSize() const { return Cache.Num(); }
+	FORCEINLINE uint32 GetCapacity() const { return Capacity; }
+
+private:
+	bool TryReleaseItem()
+	{
+		const FCacheItem& Item = Cache.Top();
+
+		if (Item.Generation + GarbageCollectLatency < CurrentGeneration)
+		{
+			delete Item.Resource;
+			Cache.Pop();
+			return true;
+		}
+
+		return false;
+	}
+
+	struct FCacheItem
+	{
+		FCacheItem(TransientResourceType* InResource, uint64 InHash, uint64 InGeneration)
+			: Resource(InResource)
+			, Hash(InHash)
+			, Generation(InGeneration)
+		{}
+
+		TransientResourceType* Resource;
+		uint64 Hash{};
+		uint64 Generation{};
+	};
+
+	TArray<FCacheItem> Cache;
+	uint64 CurrentGeneration{};
+	uint32 GarbageCollectLatency;
+	uint32 Capacity;
+};
+
+enum class ERHITransientHeapFlags : uint8
+{
+	// Supports placing buffers onto the heap.
+	AllowBuffers = 1 << 0,
+
+	// Supports placing textures with UAV support onto the heap.
+	AllowTextures = 1 << 1,
+
+	// Supports placing render targets onto the heap.
+	AllowRenderTargets = 1 << 2,
+
+	// Supports all resource types.
+	AllowAll = AllowBuffers | AllowTextures | AllowRenderTargets
+};
+
+ENUM_CLASS_FLAGS(ERHITransientHeapFlags);
+
+struct FRHITransientHeapInitializer
+{
+	// Size of the heap in bytes.
+	uint64 Size{};
+
+	// Alignment of the heap in bytes.
+	uint32 Alignment{};
+
+	// Flags used to filter resource allocations within the heap.
+	ERHITransientHeapFlags Flags = ERHITransientHeapFlags::AllowAll;
+
+	// Size of the texture cache before elements are evicted.
+	uint32 TextureCacheSize{};
+
+	// Size of the buffer cache before elements are evicted.
+	uint32 BufferCacheSize{};
+};
+
+/** The base class for a platform heap implementation. Transient resources are placed on the heap at specific
+ *  byte offsets. Each heap additionally contains a cache of RHI transient resources, each with its own RHI
+ *  resource and cache of RHI views. The lifetime of the resource cache is tied to the heap.
+ */
+class RHICORE_API FRHITransientHeap
+{
+public:
+	FRHITransientHeap(const FRHITransientHeapInitializer& InInitializer)
+		: Initializer(InInitializer)
+		, Textures(InInitializer.TextureCacheSize)
+		, Buffers(InInitializer.BufferCacheSize)
+	{}
+
+	virtual ~FRHITransientHeap() = default;
+
+	FORCEINLINE const FRHITransientHeapInitializer& GetInitializer() const { return Initializer; }
+
+	FORCEINLINE uint64 GetCapacity() const { return Initializer.Size; }
+
+	FORCEINLINE uint64 GetLastUsedGarbageCollectCycle() const { return LastUsedGarbageCollectCycle; }
+
+	FORCEINLINE bool IsAllocationSupported(uint64 Size, ERHITransientHeapFlags Flags) const
+	{
+		return Size <= Initializer.Size && EnumHasAnyFlags(Flags, Flags);
+	}
+
+private:
+	template <typename CreateFunctionType>
+	FRHITransientTexture* AcquireTexture(const FRHITextureCreateInfo& CreateInfo, uint64 HeapOffset, CreateFunctionType CreateFunction);
+
+	template <typename CreateFunctionType>
+	FRHITransientBuffer* AcquireBuffer(const FRHIBufferCreateInfo& CreateInfo, uint64 HeapOffset, CreateFunctionType CreateFunction);
+
+	void ForfeitResources();
+
+	FRHITransientHeapInitializer Initializer;
+	uint64 LastUsedGarbageCollectCycle = 0;
+
+	TRHITransientResourceCache<FRHITransientTexture> Textures;
+	TRHITransientResourceCache<FRHITransientBuffer> Buffers;
+	TArray<FRHITransientTexture*> AllocatedTextures;
+	TArray<FRHITransientBuffer*> AllocatedBuffers;
+
+	friend FRHITransientResourceSystem;
+	friend FRHITransientResourceAllocator;
+};
+
+struct RHICORE_API FRHITransientResourceSystemInitializer
+{
+	// Creates a default initializer using common RHI CVars.
+	static FRHITransientResourceSystemInitializer CreateDefault();
+
+	static const uint32 kDefaultGarbageCollectLatency = 20;
+	static const uint32 kDefaultResourceCacheSize = 256;
+
+	// The minimum size to use when creating a heap. This is the default but can grow based on allocations.
+	uint64 MinimumHeapSize{};
+
+	// The maximum size of a pool. Allocations above this size will fail.
+	uint64 MaximumHeapSize{};
+
+	// The alignment of all heaps in the cache.
+	uint32 HeapAlignment{};
+
+	// The latency between the completed fence value and the used fence value to invoke garbage collection of the heap.
+	uint32 GarbageCollectLatency = kDefaultGarbageCollectLatency;
+
+	// Size of the texture cache before elements are evicted.
+	uint32 TextureCacheSize = kDefaultResourceCacheSize;
+
+	// Size of the buffer cache before elements are evicted.
+	uint32 BufferCacheSize = kDefaultResourceCacheSize;
+
+	// Whether all heaps should be created with the AllowAll heap flag.
+	bool bSupportsAllHeapFlags = true;
+};
+
+/** The RHI transient resource system is a base class for the platform implementation. It has a persistent lifetime
+ *  and contains a cache of transient heaps. The transient allocator acquires heaps from the system and forfeits them
+ *  at the end of its lifetime. Garbage collection of heaps is done using an internal counter that increments with
+ *  each GarbageCollect call. This should be done periodically. Heaps older than a platform-specified fence latency
+ *  are destroyed. Additionally, statistics are gathered automatically and reported to the platform implementation for
+ *  stats tracking. However, 'rhitransientmemory' stats are reported automatically by the implementation itself.
+ */
+class RHICORE_API FRHITransientResourceSystem
+{
+public:
+	FRHITransientResourceSystem(const FRHITransientResourceSystemInitializer& InInitializer)
+		: Initializer(InInitializer)
+	{}
+
+	FRHITransientResourceSystem(const FRHITransientResourceSystem&) = delete;
+
+	virtual ~FRHITransientResourceSystem();
+
+	void ReleaseHeaps();
+
+	void GarbageCollect();
+
+	void UpdateStats();
+
+	FORCEINLINE uint64 GetHeapSize(uint64 RequestedHeapSize) const
+	{
+		check(RequestedHeapSize <= Initializer.MaximumHeapSize);
+		return FMath::Clamp(FMath::RoundUpToPowerOfTwo64(RequestedHeapSize), Initializer.MinimumHeapSize, Initializer.MaximumHeapSize);
+	}
+
+	FORCEINLINE uint64 GetMinimumHeapSize() const { return Initializer.MinimumHeapSize; }
+	FORCEINLINE uint64 GetMaximumHeapSize() const { return Initializer.MaximumHeapSize; }
+	FORCEINLINE uint32 GetHeapAlignment() const { return Initializer.HeapAlignment; }
+	FORCEINLINE uint32 GetHeapCount() const { return Heaps.Num(); }
+
+	FORCEINLINE FRHITransientHeap* GetHeap(uint32 Index) { return Heaps[Index]; }
+
+protected:
+	struct FStats
+	{
+		FRHITransientMemoryStats Textures;
+		FRHITransientMemoryStats Buffers;
+
+		// Total memory used by the underlying heaps.
+		uint64 TotalMemoryUsed{};
+	};
+
+private:
+	// Called by the transient allocator to acquire a heap from the cache.
+	FRHITransientHeap* AcquireHeap(uint64 FirstAllocationSize, ERHITransientHeapFlags FirstAllocationHeapFlags);
+
+	// Called by the transient allocator to forfeit all acquired heaps back to the cache.
+	void ForfeitHeaps(TArrayView<FRHITransientHeap* const> Heaps);
+
+	//////////////////////////////////////////////////////////////////////////
+	//! Platform API
+
+	// Called when a new heap is being created and added to the pool.
+	virtual FRHITransientHeap* CreateHeap(const FRHITransientHeapInitializer& Initializer) = 0;
+
+	// Called for the platform to report stats.
+	virtual void ReportStats(const FStats& Stats) {};
+
+	//////////////////////////////////////////////////////////////////////////
+
+	FRHITransientResourceSystemInitializer Initializer;
+
+	uint64 GarbageCollectCycle{};
+
+	FCriticalSection HeapCriticalSection;
+	TArray<FRHITransientHeap*> Heaps;
+
+	FRHITransientMemoryStats TextureStats;
+	FRHITransientMemoryStats BufferStats;
+
+	friend FRHITransientResourceAllocator;
+};
+
+/** Tracks resource allocations on the heap and adds overlap events to transient resources. */
+class RHICORE_API FRHITransientResourceOverlapTracker final
+{
+public:
+	void Track(FRHITransientResource* InResource, uint64 OffsetMin, uint64 OffsetMax);
+
+private:
+	struct FResourceRange
+	{
+		FRHITransientResource* Resource{};
+		uint64 OffsetMin{};
+		uint64 OffsetMax{};
+	};
+
+	TArray<FResourceRange, TMemStackAllocator<>> ResourceRanges;
+};
+
+/** Represents an allocation from the transient heap. */
+struct FRHITransientHeapAllocation
+{
+	FRHITransientHeapAllocation() = default;
+
+	FORCEINLINE bool IsValid() const { return Size != 0; }
+
+	bool operator==(const FRHITransientHeapAllocation& Other) const
+	{
+		return Size == Other.Size && Offset == Other.Offset && AlignmentPad == Other.AlignmentPad && HeapIndex == Other.HeapIndex;
+	}
+
+	bool operator!=(const FRHITransientHeapAllocation& Other) const
+	{
+		return !(*this == Other);
+	}
+
+	// Size of the allocation made from the allocator (aligned).
+	uint64 Size = 0;
+
+	// Offset in the transient heap; front of the heap starts at 0.
+	uint64 Offset = 0;
+
+	// Number of bytes of padding were added to the offset.
+	uint32 AlignmentPad = 0;
+
+	// Index of the transient heap.
+	uint32 HeapIndex = 0;
+};
+
+/** A simple first-fit allocator for placing resources onto a transient heap and tracking their aliasing overlap events. */
+class RHICORE_API FRHITransientHeapAllocator final
+{
+public:
+	FRHITransientHeapAllocator(const FRHITransientHeapInitializer& Initializer, uint32 HeapIndex);
+	~FRHITransientHeapAllocator();
+
+	FRHITransientHeapAllocation Allocate(uint64 Size, uint32 Alignment);
+	void Deallocate(FRHITransientHeapAllocation Allocation);
+
+	FORCEINLINE void TrackOverlap(FRHITransientResource* Resource, const FRHITransientHeapAllocation& Allocation)
+	{
+		OverlapTracker.Track(Resource, Allocation.Offset, Allocation.Offset + Allocation.Size);
+	}
+
+	FORCEINLINE uint64 GetCapacity() const { return Initializer.Size; }
+	FORCEINLINE uint64 GetUsedSize() const { return UsedSize; }
+	FORCEINLINE uint64 GetFreeSize() const { return Initializer.Size - UsedSize; }
+	FORCEINLINE uint64 GetAlignmentWaste() const { return AlignmentWaste; }
+	FORCEINLINE uint32 GetAllocationCount() const { return AllocationCount; }
+
+	FORCEINLINE bool IsFull() const { return UsedSize == Initializer.Size; }
+	FORCEINLINE bool IsEmpty() const { return UsedSize == 0; }
+
+	// Returns whether the requested allocation can succeed.
+	FORCEINLINE bool IsAllocationSupported(uint64 Size, ERHITransientHeapFlags Flags) const
+	{
+		return Size <= GetFreeSize() && EnumHasAnyFlags(Initializer.Flags, Flags);
+	}
+
+private:
+	using FRangeHandle = uint16;
+	static const FRangeHandle InvalidRangeHandle = FRangeHandle(~0);
+
+	struct FRange
+	{
+		uint64 Size{};
+		uint64 Offset{};
+		FRangeHandle NextFreeHandle = InvalidRangeHandle;
+
+		FORCEINLINE uint64 GetStart() const { return Offset; }
+		FORCEINLINE uint64 GetEnd() const { return Size + Offset; }
+	};
+
+	FORCEINLINE FRangeHandle GetFirstFreeRangeHandle()
+	{
+		return Ranges[HeadHandle].NextFreeHandle;
+	}
+
+	FRangeHandle CreateRange()
+	{
+		if (!RangeFreeList.IsEmpty())
+		{
+			return RangeFreeList.Pop();
+		}
+		Ranges.Emplace();
+		return FRangeHandle(Ranges.Num() - 1);
+	}
+
+	FRangeHandle InsertRange(FRangeHandle PreviousHandle, uint64 Offset, uint64 Size)
+	{
+		FRangeHandle Handle = CreateRange();
+
+		FRange& CurrentRange = Ranges[Handle];
+		CurrentRange.Offset = Offset;
+		CurrentRange.Size = Size;
+
+		FRange& PreviousRange = Ranges[PreviousHandle];
+		CurrentRange.NextFreeHandle = PreviousRange.NextFreeHandle;
+		PreviousRange.NextFreeHandle = Handle;
+
+		return Handle;
+	}
+
+	void RemoveRange(FRangeHandle PreviousHandle, FRangeHandle CurrentHandle)
+	{
+		FRange& PreviousRange = Ranges[PreviousHandle];
+		FRange& CurrentRange = Ranges[CurrentHandle];
+
+		PreviousRange.NextFreeHandle = CurrentRange.NextFreeHandle;
+		CurrentRange.NextFreeHandle = InvalidRangeHandle;
+
+		RangeFreeList.Add(CurrentHandle);
+	}
+
+	struct FFindResult
+	{
+		uint64 LeftoverSize = 0;
+		FRangeHandle PreviousHandle = InvalidRangeHandle;
+		FRangeHandle FoundHandle = InvalidRangeHandle;
+	};
+
+	FFindResult FindFreeRange(uint64 Size, uint32 Alignment);
+
+	void Validate();
+
+	FRHITransientHeapInitializer Initializer;
+
+	uint64 UsedSize{};
+	uint64 AlignmentWaste{};
+	uint32 AllocationCount{};
+	uint32 HeapIndex{};
+
+	FRangeHandle HeadHandle = InvalidRangeHandle;
+	TArray<FRangeHandle, TInlineAllocator<4, TMemStackAllocator<>>> RangeFreeList;
+	TArray<FRange, TInlineAllocator<4, TMemStackAllocator<>>> Ranges;
+
+	FRHITransientResourceOverlapTracker OverlapTracker;
+};
+
+/** A helper class for implementing IRHITransientResourceAllocator. This class is designed for composition instead of
+ *  inheritance to keep the platform implementation clean. Its methods are not straight overrides of the user-facing
+ *  class; rather, they are designed to streamline the platform implementation. This object is designed to match the
+ *  lifecycle of IRHITransientResourceAllocator. Specifically, it must be short-lived and its API used on the render
+ *  thread, as it utilizes MemStack-allocated containers for performance. Its state is built from scratch each allocation
+ *  cycle and released at the end. Heaps are acquired from the parent transient resource system, and transient resources
+ *  from respective heaps.
+ */
+class RHICORE_API FRHITransientResourceAllocator final
+{
+public:
+	FRHITransientResourceAllocator(FRHITransientResourceSystem& InParentSystem);
+	FRHITransientResourceAllocator(const FRHITransientResourceAllocator&) = delete;
+
+	struct FResourceInitializer
+	{
+		FResourceInitializer(FRHITransientHeap& InHeap, const FRHITransientHeapAllocation& InAllocation, uint64 InHash)
+			: Heap(InHeap)
+			, Allocation(InAllocation)
+			, Hash(InHash)
+		{}
+
+		// The heap on which to create the resource.
+		FRHITransientHeap& Heap;
+
+		// The allocation (offset / size) on the provided heap.
+		const FRHITransientHeapAllocation& Allocation;
+
+		// The unique hash computed from the create info and allocation offset.
+		const uint64 Hash;
+	};
+
+	using FCreateTextureFunction = TFunction<FRHITransientTexture*(const FResourceInitializer&)>;
+
+	/** Allocates a texture on a heap at a specific offset, returning a cached RHI transient texture pointer, or null
+	 *  if the allocation failed. TextureSize and TextureAlignment are platform specific and must be derived from the
+	 *  texture create info and passed in, along with a platform-specific texture creation function if no cached resource
+	 *  if found.
+	 */
+	FRHITransientTexture* CreateTexture(
+		const FRHITextureCreateInfo& CreateInfo,
+		const TCHAR* DebugName,
+		uint64 TextureSize,
+		uint32 TextureAlignment,
+		FCreateTextureFunction CreateTextureFunction);
+
+	using FCreateBufferFunction = TFunction<FRHITransientBuffer*(const FResourceInitializer&)>;
+
+	/** Allocates a buffer on a heap at a specific offset, returning a cached RHI transient buffer pointer, or null
+	 *  if the allocation failed. BufferSize and BufferAlignment are platform specific and must be derived from the
+	 *  buffer create info and passed in, along with a platform-specific buffer creation function if no cached resource
+	 *  if found.
+	 */
+	FRHITransientBuffer* CreateBuffer(
+		const FRHIBufferCreateInfo& CreateInfo,
+		const TCHAR* DebugName,
+		uint32 BufferSize,
+		uint32 BufferAlignment,
+		FCreateBufferFunction CreateBufferFunction);
+
+	// Deallocates a texture from its parent heap. Provide the current platform fence value used to update the heap.
+	FORCEINLINE void DeallocateMemory(FRHITransientTexture* Texture)
+	{
+		DeallocateMemoryInternal(Texture, TextureStats);
+	}
+
+	// Deallocates a buffer from its parent heap. Provide the current platform fence value used to update the heap.
+	FORCEINLINE void DeallocateMemory(FRHITransientBuffer* Buffer)
+	{
+		DeallocateMemoryInternal(Buffer, BufferStats);
+	}
+
+	// Called to signify all allocations have completed. Forfeits all resources / heaps back to the parent system.
+	void Freeze(FRHICommandListImmediate& RHICmdList);
+
+private:
+	struct FMemoryStats : FRHITransientMemoryStats
+	{
+		uint64 CurrentSizeWithAliasing{};
+	};
+
+	FRHITransientHeapAllocation Allocate(FMemoryStats& StatsToUpdate, uint64 Size, uint32 Alignment, ERHITransientHeapFlags ResourceHeapFlags);
+
+	void DeallocateMemoryInternal(FRHITransientResource* InResource, FMemoryStats& StatsToUpdate);
+
+	void InitResource(FRHITransientResource* TransientResource, const FRHITransientHeapAllocation& Allocation, const TCHAR* Name);
+
+	FRHITransientResourceSystem& ParentSystem;
+
+	// Tracks state on the rendering thread; must be cleared before the destructor.
+	TArray<FRHITransientHeap*, TInlineAllocator<4, TMemStackAllocator<>>> Heaps;
+	TArray<FRHITransientHeapAllocator, TInlineAllocator<4, TMemStackAllocator<>>> HeapAllocators;
+	TArray<FRHITransientHeapAllocation, TInlineAllocator<4, TMemStackAllocator<>>> HeapAllocations;
+
+#if RHICORE_TRANSIENT_ALLOCATOR_DEBUG
+	TArray<FRHITransientTexture*, TMemStackAllocator<>> DebugTextures;
+	TArray<FRHITransientBuffer*, TMemStackAllocator<>> DebugBuffers;
+#endif
+
+	FMemoryStats TextureStats;
+	FMemoryStats BufferStats;
+
+	uint64 AllocatedTextureSize{};
+	uint64 AllocatedBufferSize{};
+};
\ No newline at end of file
diff --git a/Engine/Source/Runtime/RenderCore/Private/RenderGraphBuilder.cpp b/Engine/Source/Runtime/RenderCore/Private/RenderGraphBuilder.cpp
--- a/Engine/Source/Runtime/RenderCore/Private/RenderGraphBuilder.cpp
+++ b/Engine/Source/Runtime/RenderCore/Private/RenderGraphBuilder.cpp
@@ -508,10 +508,7 @@ bool FRDGBuilder::IsTransient(FRDGBufferRef Buffer) const
 		return false;
 	}
 
-	// Buffers smaller than a page go through normal pooled allocation.
-	const uint32 TransientSizeMin = 64 * 1024;
-
-	return Buffer->Desc.GetTotalNumBytes() >= TransientSizeMin && EnumHasAnyFlags(Buffer->Desc.Usage, BUF_UnorderedAccess);
+	return EnumHasAnyFlags(Buffer->Desc.Usage, BUF_UnorderedAccess);
 }
 
 bool FRDGBuilder::IsTransient(FRDGTextureRef Texture) const
@@ -528,11 +525,18 @@ bool FRDGBuilder::IsTransientInternal(FRDGParentResourceRef Resource) const
 	}
 
 	// Transient resources must stay within the graph.
-	if (Resource->bExternal || Resource->bExtracted)
+	if (Resource->bExternal || Resource->bExtracted || Resource->bUserSetNonTransient)
 	{
 		return false;
 	}
 
+#if RDG_ENABLE_DEBUG
+	if (GRDGDebugDisableTransientResources != 0 && IsDebugAllowedForResource(Resource->Name))
+	{
+		return false;
+	}
+#endif
+
 	// This resource cannot be extracted or made external later.
 	check(bSetupComplete);
 
@@ -543,7 +547,7 @@ FRDGBuilder::FTransientResourceAllocator::~FTransientResourceAllocator()
 {
 	if (Allocator)
 	{
-		RHICmdList.ReleaseTransientResourceAllocator(Allocator);
+		Allocator->Release(RHICmdList);
 	}
 }
 
@@ -949,6 +953,16 @@ void FRDGBuilder::Compile()
 			AddCullingDependency(Buffer->LastProducer, StateFinal, ERHIPipeline::Graphics);
 			Buffer->ReferenceCount++;
 		}
+
+		EnumerateExtendedLifetimeResources(Textures, [](FRDGTextureRef Texture)
+		{
+			Texture->ReferenceCount++;
+		});
+
+		EnumerateExtendedLifetimeResources(Buffers, [](FRDGBufferRef Buffer)
+		{
+			Buffer->ReferenceCount++;
+		});
 	}
 
 	// All dependencies in the raw graph have been specified; if enabled, all passes are marked as culled and a
@@ -1469,6 +1483,16 @@ void FRDGBuilder::Execute()
 				EndResourceRHI(TransientResourceAllocator, EpiloguePassHandle, Query.Key, 1);
 			}
 
+			EnumerateExtendedLifetimeResources(Textures, [&](FRDGTextureRef Texture)
+			{
+				EndResourceRHI(TransientResourceAllocator, EpiloguePassHandle, Texture, 1);
+			});
+
+			EnumerateExtendedLifetimeResources(Buffers, [&](FRDGBufferRef Buffer)
+			{
+				EndResourceRHI(TransientResourceAllocator, EpiloguePassHandle, Buffer, 1);
+			});
+
 			if (TransientResourceAllocator)
 			{
 				TransientResourceAllocator->Freeze(RHICmdList);
@@ -2430,7 +2454,7 @@ void FRDGBuilder::BeginResourceRHI(FRDGUniformBuffer* UniformBuffer)
 {
 	check(UniformBuffer);
 
-	if (UniformBuffer->UniformBufferRHI)
+	if (UniformBuffer->GetRHIUnchecked())
 	{
 		return;
 	}
@@ -2452,7 +2476,7 @@ void FRDGBuilder::BeginResourceRHI(FTransientResourceAllocator* TransientResourc
 {
 	check(Texture);
 
-	if (Texture->PooledTexture)
+	if (Texture->GetRHIUnchecked())
 	{
 		return;
 	}
@@ -2471,40 +2495,36 @@ void FRDGBuilder::BeginResourceRHI(FTransientResourceAllocator* TransientResourc
 	}
 #endif
 
-	bool bTransientAllocated = false;
-
 	if (IsTransient(Texture))
 	{
 		check(TransientResourceAllocator);
 
 		if (IRHITransientResourceAllocator* AllocatorRHI = TransientResourceAllocator->GetOrCreate())
 		{
-			TRACE_CPUPROFILER_EVENT_SCOPE(FRDGBuilder::SetupTransientTexture);
-
-			Texture->SetRHI(FRDGPooledTexture::CreateTransient(Allocator, Texture->Desc, AllocatorRHI->CreateTexture(Texture->Desc, Texture->Name)));
-			Texture->bTransient = true;
+			if (FRHITransientTexture* TransientTexture = AllocatorRHI->CreateTexture(Texture->Desc, Texture->Name))
+			{
+				Texture->SetRHI(TransientTexture, Allocator);
 
-			check(GetPrologueBarrierPassHandle(PassHandle) == PassHandle);
+				check(GetPrologueBarrierPassHandle(PassHandle) == PassHandle);
 
-			FRDGSubresourceState InitialState;
-			InitialState.SetPass(ERHIPipeline::Graphics, PassHandle);
-			InitialState.Access = ERHIAccess::Discard;
-			InitAsWholeResource(Texture->GetState(), InitialState);
+				FRDGSubresourceState InitialState;
+				InitialState.SetPass(ERHIPipeline::Graphics, PassHandle);
+				InitialState.Access = ERHIAccess::Discard;
+				InitAsWholeResource(Texture->GetState(), InitialState);
 
-			AddToPrologueBarriers(PassHandle, [&](FRDGBarrierBatchBegin& Barriers)
-			{
-				Barriers.AddAlias(Texture, FRHITransientAliasingInfo::Acquire(Texture->GetRHIUnchecked()));
-			});
-
-			bTransientAllocated = true;
+				AddToPrologueBarriers(PassHandle, [&](FRDGBarrierBatchBegin& Barriers)
+				{
+					Barriers.AddAlias(Texture, FRHITransientAliasingInfo::Acquire(Texture->GetRHIUnchecked(), Texture->TransientTexture->GetAliasingOverlaps()));
+				});
 
-		#if STATS
-			GRDGStatTransientTextureCount++;
-		#endif
+			#if STATS
+				GRDGStatTransientTextureCount++;
+			#endif
+			}
 		}
 	}
 
-	if (!bTransientAllocated)
+	if (!Texture->bTransient)
 	{
 		Texture->SetRHI(GRenderTargetPool.FindFreeElementForRDG(RHICmdList, Texture->Desc, Texture->Name));
 	}
@@ -2518,69 +2538,29 @@ void FRDGBuilder::BeginResourceRHI(FTransientResourceAllocator*, FRDGPassHandle
 {
 	check(SRV);
 
-	if (SRV->ResourceRHI)
+	if (SRV->GetRHIUnchecked())
 	{
 		return;
 	}
 
 	FRDGTextureRef Texture = SRV->Desc.Texture;
-	FRDGPooledTexture* PooledTexture = Texture->PooledTexture;
-	checkf(PooledTexture, TEXT("Pass parameters contained an SRV of RDG Texture %s, before that texture was referenced as a UAV or RTV, which isn't supported.  Make sure to call ClearUnusedGraphResources to remove unbound parameters which can also cause this."), Texture->Name);
-
-	if (SRV->Desc.MetaData == ERDGTextureMetaDataAccess::HTile)
-	{
-		check(GRHISupportsExplicitHTile);
-		if (!PooledTexture->HTileSRV)
-		{
-			PooledTexture->HTileSRV = RHICreateShaderResourceViewHTile((FRHITexture2D*)PooledTexture->Texture.GetReference());
-		}
-		SRV->ResourceRHI = PooledTexture->HTileSRV;
-		check(SRV->ResourceRHI);
-		return;
-	}
-
-	if (SRV->Desc.MetaData == ERDGTextureMetaDataAccess::FMask)
-	{
-		if (!PooledTexture->FMaskSRV)
-		{
-			PooledTexture->FMaskSRV = RHICreateShaderResourceViewFMask((FRHITexture2D*)PooledTexture->Texture.GetReference());
-		}
-		SRV->ResourceRHI = PooledTexture->FMaskSRV;
-		check(SRV->ResourceRHI);
-		return;
-	}
+	check(Texture->GetRHIUnchecked());
 
-	if (SRV->Desc.MetaData == ERDGTextureMetaDataAccess::CMask)
+	if (Texture->bTransient)
 	{
-		if (!PooledTexture->CMaskSRV)
-		{
-			PooledTexture->CMaskSRV = RHICreateShaderResourceViewWriteMask((FRHITexture2D*)PooledTexture->Texture.GetReference());
-		}
-		SRV->ResourceRHI = PooledTexture->CMaskSRV;
-		check(SRV->ResourceRHI);
-		return;
+		SRV->ResourceRHI = Texture->TransientTexture->GetOrCreateSRV(SRV->Desc);
 	}
-
-	for (const auto& SRVPair : PooledTexture->SRVs)
+	else
 	{
-		if (SRVPair.Key == SRV->Desc)
-		{
-			SRV->ResourceRHI = SRVPair.Value;
-			return;
-		}
+		SRV->ResourceRHI = Texture->PooledTexture->GetOrCreateSRV(SRV->Desc);
 	}
-
-	FShaderResourceViewRHIRef RHIShaderResourceView = RHICreateShaderResourceView(PooledTexture->Texture, SRV->Desc);
-
-	SRV->ResourceRHI = RHIShaderResourceView;
-	PooledTexture->SRVs.Emplace(SRV->Desc, MoveTemp(RHIShaderResourceView));
 }
 
 void FRDGBuilder::BeginResourceRHI(FTransientResourceAllocator* TransientResourceAllocator, FRDGPassHandle PassHandle, FRDGTextureUAVRef UAV)
 {
 	check(UAV);
 
-	if (UAV->ResourceRHI)
+	if (UAV->GetRHIUnchecked())
 	{
 		return;
 	}
@@ -2588,40 +2568,22 @@ void FRDGBuilder::BeginResourceRHI(FTransientResourceAllocator* TransientResourc
 	BeginResourceRHI(TransientResourceAllocator, PassHandle, UAV->Desc.Texture);
 
 	FRDGTextureRef Texture = UAV->Desc.Texture;
-	FRDGPooledTexture* PooledTexture = Texture->PooledTexture;
-	check(PooledTexture);
 
-	if (UAV->Desc.MetaData == ERDGTextureMetaDataAccess::HTile)
+	if (Texture->bTransient)
 	{
-		check(GRHISupportsExplicitHTile);
-		if (!PooledTexture->HTileUAV)
-		{
-			PooledTexture->HTileUAV = RHICreateUnorderedAccessViewHTile((FRHITexture2D*)PooledTexture->Texture.GetReference());
-		}
-		UAV->ResourceRHI = PooledTexture->HTileUAV;
-		check(UAV->ResourceRHI);
-		return;
+		UAV->ResourceRHI = Texture->TransientTexture->GetOrCreateUAV(UAV->Desc);
 	}
-
-	if (UAV->Desc.MetaData == ERDGTextureMetaDataAccess::Stencil)
+	else
 	{
-		if (!PooledTexture->StencilUAV)
-		{
-			PooledTexture->StencilUAV = RHICreateUnorderedAccessViewStencil((FRHITexture2D*)PooledTexture->Texture.GetReference(), 0);
-		}
-		UAV->ResourceRHI = PooledTexture->StencilUAV;
-		check(UAV->ResourceRHI);
-		return;
+		UAV->ResourceRHI = Texture->PooledTexture->GetOrCreateUAV(UAV->Desc);
 	}
-
-	UAV->ResourceRHI = PooledTexture->MipUAVs[UAV->Desc.MipLevel];
 }
 
 void FRDGBuilder::BeginResourceRHI(FTransientResourceAllocator* TransientResourceAllocator, FRDGPassHandle PassHandle, FRDGBufferRef Buffer)
 {
 	check(Buffer);
 
-	if (Buffer->PooledBuffer)
+	if (Buffer->GetRHIUnchecked())
 	{
 		return;
 	}
@@ -2629,8 +2591,6 @@ void FRDGBuilder::BeginResourceRHI(FTransientResourceAllocator* TransientResourc
 	check(TransientResourceAllocator || Buffer->bExternal);
 	check(Buffer->ReferenceCount > 0 || Buffer->bExternal || IsImmediateMode());
 
-	bool bTransientAllocated = false;
-
 	// If transient then create the resource on the transient allocator. External or extracted resource can't be transient because of lifetime tracking issues.
 	if (IsTransient(Buffer))
 	{
@@ -2638,67 +2598,100 @@ void FRDGBuilder::BeginResourceRHI(FTransientResourceAllocator* TransientResourc
 
 		if (IRHITransientResourceAllocator* AllocatorRHI = TransientResourceAllocator->GetOrCreate())
 		{
-			TRACE_CPUPROFILER_EVENT_SCOPE(FRDGBuilder::SetupTransientBuffer);
-
-			Buffer->SetRHI(FRDGPooledBuffer::CreateTransient(Allocator, AllocatorRHI->CreateBuffer(Translate(Buffer->Desc), Buffer->Name), Buffer->Desc));
-			Buffer->bTransient = true;
-
-			check(GetPrologueBarrierPassHandle(PassHandle) == PassHandle);
+			if (FRHITransientBuffer* TransientBuffer = AllocatorRHI->CreateBuffer(Translate(Buffer->Desc), Buffer->Name))
+			{
+				Buffer->SetRHI(TransientBuffer, Allocator);
 
-			FRDGSubresourceState& InitialState = Buffer->GetState();
-			InitialState.SetPass(ERHIPipeline::Graphics, PassHandle);
-			InitialState.Access = ERHIAccess::Discard;
+				check(GetPrologueBarrierPassHandle(PassHandle) == PassHandle);
 
-			AddToPrologueBarriers(PassHandle, [&](FRDGBarrierBatchBegin& Barriers)
-			{
-				Barriers.AddAlias(Buffer, FRHITransientAliasingInfo::Acquire(Buffer->GetRHIUnchecked()));
-			});
+				FRDGSubresourceState& InitialState = Buffer->GetState();
+				InitialState.SetPass(ERHIPipeline::Graphics, PassHandle);
+				InitialState.Access = ERHIAccess::Discard;
 
-			bTransientAllocated = true;
+				AddToPrologueBarriers(PassHandle, [&](FRDGBarrierBatchBegin& Barriers)
+				{
+					Barriers.AddAlias(Buffer, FRHITransientAliasingInfo::Acquire(Buffer->GetRHIUnchecked(), Buffer->TransientBuffer->GetAliasingOverlaps()));
+				});
 
-		#if STATS
-			GRDGStatTransientBufferCount++;
-		#endif
+			#if STATS
+				GRDGStatTransientBufferCount++;
+			#endif
+			}
 		}
 	}
 
-	if (!bTransientAllocated)
+	if (!Buffer->bTransient)
 	{
 		Buffer->SetRHI(GRenderGraphResourcePool.FindFreeBufferInternal(RHICmdList, Buffer->Desc, Buffer->Name));
 	}
 
 	Buffer->FirstPass = PassHandle;
 
-	check(Buffer->PooledBuffer);
+	check(Buffer->GetRHIUnchecked());
 }
 
 void FRDGBuilder::BeginResourceRHI(FTransientResourceAllocator*, FRDGPassHandle PassHandle, FRDGBufferSRVRef SRV)
 {
 	check(SRV);
 
-	if (SRV->ResourceRHI)
+	if (SRV->GetRHIUnchecked())
 	{
 		return;
 	}
 
 	FRDGBufferRef Buffer = SRV->Desc.Buffer;
-	checkf(Buffer->PooledBuffer, TEXT("Pass parameters contained an SRV of RDG buffer %s, before that buffer was referenced as a UAV, which isn't supported.  Make sure to call ClearUnusedGraphResources to remove unbound parameters which can also cause this."), Buffer->Name);
+	checkf(Buffer->GetRHIUnchecked(), TEXT("Pass parameters contained an SRV of RDG buffer %s, before that buffer was referenced as a UAV, which isn't supported.  Make sure to call ClearUnusedGraphResources to remove unbound parameters which can also cause this."), Buffer->Name);
+
+	FRHIBufferSRVCreateInfo SRVCreateInfo = SRV->Desc;
+
+	if (Buffer->Desc.UnderlyingType == FRDGBufferDesc::EUnderlyingType::StructuredBuffer)
+	{
+		// RDG allows structured buffer views to be typed, but the view creation logic requires that it
+		// be unknown (as do platform APIs -- structured buffers are not typed). This could be validated
+		// at the high level but the current API makes it confusing. For now, it's considered a no-op.
+		SRVCreateInfo.Format = PF_Unknown;
+	}
 
-	SRV->ResourceRHI = Buffer->PooledBuffer->GetOrCreateSRV(SRV->Desc);
+	if (Buffer->bTransient)
+	{
+		SRV->ResourceRHI = Buffer->TransientBuffer->GetOrCreateSRV(SRVCreateInfo);
+	}
+	else
+	{
+		SRV->ResourceRHI = Buffer->PooledBuffer->GetOrCreateSRV(SRVCreateInfo);
+	}
 }
 
 void FRDGBuilder::BeginResourceRHI(FTransientResourceAllocator* TransientResourceAllocator, FRDGPassHandle PassHandle, FRDGBufferUAV* UAV)
 {
 	check(UAV);
 
-	if (UAV->ResourceRHI)
+	if (UAV->GetRHIUnchecked())
 	{
 		return;
 	}
 
 	FRDGBufferRef Buffer = UAV->Desc.Buffer;
 	BeginResourceRHI(TransientResourceAllocator, PassHandle, Buffer);
-	UAV->ResourceRHI = Buffer->PooledBuffer->GetOrCreateUAV(UAV->Desc);
+
+	FRHIBufferUAVCreateInfo UAVCreateInfo = UAV->Desc;
+
+	if (Buffer->Desc.UnderlyingType == FRDGBufferDesc::EUnderlyingType::StructuredBuffer)
+	{
+		// RDG allows structured buffer views to be typed, but the view creation logic requires that it
+		// be unknown (as do platform APIs -- structured buffers are not typed). This could be validated
+		// at the high level but the current API makes it confusing. For now, it's considered a no-op.
+		UAVCreateInfo.Format = PF_Unknown;
+	}
+
+	if (Buffer->bTransient)
+	{
+		UAV->ResourceRHI = Buffer->TransientBuffer->GetOrCreateUAV(UAVCreateInfo);
+	}
+	else
+	{
+		UAV->ResourceRHI = Buffer->PooledBuffer->GetOrCreateUAV(UAVCreateInfo);
+	}
 }
 
 void FRDGBuilder::EndResourceRHI(FTransientResourceAllocator& TransientResourceAllocator, FRDGPassHandle PassHandle, FRDGTextureRef Texture, uint32 ReferenceCount)
@@ -2711,7 +2704,7 @@ void FRDGBuilder::EndResourceRHI(FTransientResourceAllocator& TransientResourceA
 	{
 		if (Texture->bTransient)
 		{
-			TransientResourceAllocator->DeallocateMemory(Texture->GetRHIUnchecked());
+			TransientResourceAllocator->DeallocateMemory(Texture->TransientTexture);
 
 			AddToEpilogueBarriers(PassHandle, [&](FRDGBarrierBatchBegin& Barriers)
 			{
@@ -2737,7 +2730,7 @@ void FRDGBuilder::EndResourceRHI(FTransientResourceAllocator& TransientResourceA
 	{
 		if (Buffer->bTransient)
 		{
-			TransientResourceAllocator->DeallocateMemory(Buffer->GetRHIUnchecked());
+			TransientResourceAllocator->DeallocateMemory(Buffer->TransientBuffer);
 
 			AddToEpilogueBarriers(PassHandle, [&](FRDGBarrierBatchBegin& Barriers)
 			{
diff --git a/Engine/Source/Runtime/RenderCore/Private/RenderGraphPrivate.cpp b/Engine/Source/Runtime/RenderCore/Private/RenderGraphPrivate.cpp
--- a/Engine/Source/Runtime/RenderCore/Private/RenderGraphPrivate.cpp
+++ b/Engine/Source/Runtime/RenderCore/Private/RenderGraphPrivate.cpp
@@ -39,6 +39,23 @@ FAutoConsoleVariableRef CVarRDGDebugFlushGPU(
 	}),
 	ECVF_RenderThreadSafe);
 
+int32 GRDGDebugExtendResourceLifetimes = 0;
+FAutoConsoleVariableRef CVarRDGDebugExtendResourceLifetimes(
+	TEXT("r.RDG.Debug.ExtendResourceLifetimes"),
+	GRDGDebugExtendResourceLifetimes,
+	TEXT("Extends the resource lifetimes of resources (or a specific resource filter specified by r.RDG.Debug.ResourceFilter) ")
+	TEXT("so that they cannot overlap memory with any other resource within the graph. Useful to debug if transient aliasing is causing issues.\n")
+	TEXT(" 0: disabled (default);\n")
+	TEXT(" 1: enabled;\n"),
+	ECVF_RenderThreadSafe);
+
+int32 GRDGDebugDisableTransientResources = 0;
+FAutoConsoleVariableRef CVarRDGDebugDisableTransientResource(
+	TEXT("r.RDG.Debug.DisableTransientResources"),
+	GRDGDebugDisableTransientResources,
+	TEXT("Filters out transient resources from the transient allocator. Use r.rdg.debug.resourcefilter to specify the filter. Defaults to all resources if enabled."),
+	ECVF_RenderThreadSafe);
+
 int32 GRDGDumpGraph = 0;
 FAutoConsoleVariableRef CVarDumpGraph(
 	TEXT("r.RDG.DumpGraph"),
@@ -97,14 +114,23 @@ FAutoConsoleVariableRef CVarRDGTransitionLog(
 
 TAutoConsoleVariable<FString> CVarRDGDebugGraphFilter(
 	TEXT("r.RDG.Debug.GraphFilter"), TEXT(""),
-	TEXT("Filters certain debug events to a specific graph.\n"),
+	TEXT("Filters certain debug events to a specific graph. Set to 'None' to reset.\n"),
 	ECVF_Default);
 
 FString GRDGDebugGraphFilterName;
 
+inline FString GetDebugFilterString(const FString& InputString)
+{
+	if (!InputString.Compare(TEXT("None"), ESearchCase::IgnoreCase))
+	{
+		return {};
+	}
+	return InputString;
+}
+
 FAutoConsoleVariableSink CVarRDGDebugGraphSink(FConsoleCommandDelegate::CreateLambda([]()
 {
-	GRDGDebugGraphFilterName = CVarRDGDebugGraphFilter.GetValueOnGameThread();
+	GRDGDebugGraphFilterName = GetDebugFilterString(CVarRDGDebugGraphFilter.GetValueOnGameThread());
 }));
 
 inline bool IsDebugAllowed(const FString& FilterString, const TCHAR* Name)
@@ -130,14 +156,14 @@ bool IsDebugAllowedForGraph(const TCHAR* GraphName)
 
 TAutoConsoleVariable<FString> CVarRDGDebugPassFilter(
 	TEXT("r.RDG.Debug.PassFilter"), TEXT(""),
-	TEXT("Filters certain debug events to specific passes.\n"),
+	TEXT("Filters certain debug events to specific passes. Set to 'None' to reset.\n"),
 	ECVF_Default);
 
 FString GRDGDebugPassFilterName;
 
 FAutoConsoleVariableSink CVarRDGDebugPassSink(FConsoleCommandDelegate::CreateLambda([]()
 {
-	GRDGDebugPassFilterName = CVarRDGDebugPassFilter.GetValueOnGameThread();
+	GRDGDebugPassFilterName = GetDebugFilterString(CVarRDGDebugPassFilter.GetValueOnGameThread());
 }));
 
 bool IsDebugAllowedForPass(const TCHAR* PassName)
@@ -147,14 +173,14 @@ bool IsDebugAllowedForPass(const TCHAR* PassName)
 
 TAutoConsoleVariable<FString> CVarRDGDebugResourceFilter(
 	TEXT("r.RDG.Debug.ResourceFilter"), TEXT(""),
-	TEXT("Filters certain debug events to a specific resource.\n"),
+	TEXT("Filters certain debug events to a specific resource. Set to 'None' to reset.\n"),
 	ECVF_Default);
 
 FString GRDGDebugResourceFilterName;
 
 FAutoConsoleVariableSink CVarRDGDebugResourceSink(FConsoleCommandDelegate::CreateLambda([]()
 {
-	GRDGDebugResourceFilterName = CVarRDGDebugResourceFilter.GetValueOnGameThread();
+	GRDGDebugResourceFilterName = GetDebugFilterString(CVarRDGDebugResourceFilter.GetValueOnGameThread());
 }));
 
 bool IsDebugAllowedForResource(const TCHAR* ResourceName)
@@ -271,10 +297,12 @@ FAutoConsoleVariableRef CVarRDGMergeRenderPasses(
 	TEXT(" 1:on(default);\n"),
 	ECVF_RenderThreadSafe);
 
-int32 GRDGTransientAllocator = 0;
+int32 GRDGTransientAllocator = 1;
 FAutoConsoleVariableRef CVarRDGUseTransientAllocator(
 	TEXT("r.RDG.TransientAllocator"), GRDGTransientAllocator,
-	TEXT("RDG will use the RHITransientResourceAllocator to allocate all transient resources."),
+	TEXT("RDG will use the RHITransientResourceAllocator to allocate all transient resources.")
+	TEXT(" 0: disables the transient allocator;")
+	TEXT(" 1: enables the transient allocator (default)"),
 	ECVF_RenderThreadSafe);
 
 #if CSV_PROFILER
@@ -339,6 +367,11 @@ void InitRenderGraph()
 		GRDGDebug = 1;
 	}
 
+	if (FParse::Param(FCommandLine::Get(), TEXT("rdgdebugextendresourcelifetimes")))
+	{
+		GRDGDebugExtendResourceLifetimes = 1;
+	}
+
 	if (FParse::Param(FCommandLine::Get(), TEXT("rdgtransitionlog")))
 	{
 		// Set to -1 to specify infinite number of frames.
diff --git a/Engine/Source/Runtime/RenderCore/Private/RenderGraphPrivate.h b/Engine/Source/Runtime/RenderCore/Private/RenderGraphPrivate.h
--- a/Engine/Source/Runtime/RenderCore/Private/RenderGraphPrivate.h
+++ b/Engine/Source/Runtime/RenderCore/Private/RenderGraphPrivate.h
@@ -25,6 +25,8 @@ extern int32 GRDGAsyncCompute;
 extern int32 GRDGClobberResources;
 extern int32 GRDGDebug;
 extern int32 GRDGDebugFlushGPU;
+extern int32 GRDGDebugExtendResourceLifetimes;
+extern int32 GRDGDebugDisableTransientResources;
 extern int32 GRDGDumpGraph;
 extern int32 GRDGDumpGraphUnknownCount;
 extern int32 GRDGBreakpoint;
@@ -87,6 +89,8 @@ void EmitRDGWarning(const FString& WarningMessage);
 const int32 GRDGClobberResources = 0;
 const int32 GRDGDebug = 0;
 const int32 GRDGDebugFlushGPU = 0;
+const int32 GRDGDebugExtendResourceLifetimes = 0;
+const int32 GRDGDebugDisableTransientResources = 0;
 const int32 GRDGDumpGraph = 0;
 const int32 GRDGBreakpoint = 0;
 const int32 GRDGTransitionLog = 0;
@@ -103,6 +107,7 @@ extern int32 GRDGAsyncCompute;
 extern int32 GRDGCullPasses;
 extern int32 GRDGMergeRenderPasses;
 extern int32 GRDGTransientAllocator;
+extern int32 GRDGTransientMinimumBufferSize;
 
 #if CSV_PROFILER
 extern int32 GRDGVerboseCSVStats;
@@ -169,4 +174,23 @@ inline const TCHAR* GetEpilogueBarriersToBeginDebugName(ERHIPipeline Pipelines)
 inline bool IsImmediateMode()
 {
 	return GRDGImmediateMode != 0;
+}
+
+template <typename ResourceRegistryType, typename FunctionType>
+inline void EnumerateExtendedLifetimeResources(ResourceRegistryType& Registry, FunctionType Function)
+{
+#if RDG_ENABLE_DEBUG
+	if (GRDGDebugExtendResourceLifetimes)
+	{
+		for (auto Handle = Registry.Begin(); Handle != Registry.End(); ++Handle)
+		{
+			auto* Resource = Registry[Handle];
+
+			if (IsDebugAllowedForResource(Resource->Name))
+			{
+				Function(Resource);
+			}
+		}
+	}
+#endif
 }
\ No newline at end of file
diff --git a/Engine/Source/Runtime/RenderCore/Private/RenderGraphResourcePool.cpp b/Engine/Source/Runtime/RenderCore/Private/RenderGraphResourcePool.cpp
--- a/Engine/Source/Runtime/RenderCore/Private/RenderGraphResourcePool.cpp
+++ b/Engine/Source/Runtime/RenderCore/Private/RenderGraphResourcePool.cpp
@@ -74,7 +74,7 @@ TRefCountPtr<FRDGPooledBuffer> FRenderGraphResourcePool::FindFreeBufferInternal(
 			check(0);
 		}
 
-		TRefCountPtr<FRDGPooledBuffer> PooledBuffer = FRDGPooledBuffer::CreateCommitted(MoveTemp(BufferRHI), Desc);
+		TRefCountPtr<FRDGPooledBuffer> PooledBuffer = new FRDGPooledBuffer(MoveTemp(BufferRHI), Desc);
 		AllocatedBuffers.Add(PooledBuffer);
 		check(PooledBuffer->GetRefCount() == 2);
 
diff --git a/Engine/Source/Runtime/RenderCore/Private/RenderGraphResources.cpp b/Engine/Source/Runtime/RenderCore/Private/RenderGraphResources.cpp
--- a/Engine/Source/Runtime/RenderCore/Private/RenderGraphResources.cpp
+++ b/Engine/Source/Runtime/RenderCore/Private/RenderGraphResources.cpp
@@ -27,6 +27,7 @@ FRDGParentResource::FRDGParentResource(const TCHAR* InName, const ERDGParentReso
 	, bExtracted(0)
 	, bProduced(0)
 	, bTransient(0)
+	, bUserSetNonTransient(0)
 	, bFinalizedAccess(0)
 	, bLastOwner(1)
 	// Culling logic runs only when immediate mode is off.
@@ -207,6 +208,75 @@ void FRDGPooledTexture::Reset()
 	Owner = nullptr;
 }
 
+FRHIShaderResourceView* FRDGPooledTexture::GetOrCreateSRV(const FRHITextureSRVCreateInfo& Desc)
+{
+	if (Desc.MetaData == ERDGTextureMetaDataAccess::HTile)
+	{
+		check(GRHISupportsExplicitHTile);
+		if (!HTileSRV)
+		{
+			HTileSRV = RHICreateShaderResourceViewHTile((FRHITexture2D*)GetRHI());
+		}
+		return HTileSRV;
+	}
+
+	if (Desc.MetaData == ERDGTextureMetaDataAccess::FMask)
+	{
+		if (!FMaskSRV)
+		{
+			FMaskSRV = RHICreateShaderResourceViewFMask((FRHITexture2D*)GetRHI());
+		}
+		return FMaskSRV;
+	}
+
+	if (Desc.MetaData == ERDGTextureMetaDataAccess::CMask)
+	{
+		if (!CMaskSRV)
+		{
+			CMaskSRV = RHICreateShaderResourceViewWriteMask((FRHITexture2D*)GetRHI());
+		}
+		return CMaskSRV;
+	}
+
+	for (const auto& SRVPair : SRVs)
+	{
+		if (SRVPair.Key == Desc)
+		{
+			return SRVPair.Value;
+		}
+	}
+
+	FShaderResourceViewRHIRef RHIShaderResourceView = RHICreateShaderResourceView(GetRHI(), Desc);
+
+	FRHIShaderResourceView* View = RHIShaderResourceView;
+	SRVs.Emplace(Desc, MoveTemp(RHIShaderResourceView));
+	return View;
+}
+
+FRHIUnorderedAccessView* FRDGPooledTexture::GetOrCreateUAV(const FRHITextureUAVCreateInfo& UAVDesc)
+{
+	if (UAVDesc.MetaData == ERDGTextureMetaDataAccess::HTile)
+	{
+		check(GRHISupportsExplicitHTile);
+		if (!HTileUAV)
+		{
+			HTileUAV = RHICreateUnorderedAccessViewHTile((FRHITexture2D*)GetRHI());
+		}
+		return HTileUAV;
+	}
+
+	if (UAVDesc.MetaData == ERDGTextureMetaDataAccess::Stencil)
+	{
+		if (!StencilUAV)
+		{
+			StencilUAV = RHICreateUnorderedAccessViewStencil((FRHITexture2D*)GetRHI(), 0);
+		}
+		return StencilUAV;
+	}
+
+	return MipUAVs[UAVDesc.MipLevel];
+}
+
 FRDGTextureSubresourceRange FRDGTexture::GetSubresourceRangeSRV() const
 {
 	FRDGTextureSubresourceRange Range = GetSubresourceRange();
@@ -221,7 +291,6 @@ FRDGTextureSubresourceRange FRDGTexture::GetSubresourceRangeSRV() const
 
 void FRDGTexture::SetRHI(FPooledRenderTarget* InPooledRenderTarget)
 {
-	check(InPooledRenderTarget);
 	Allocation = InPooledRenderTarget;
 	PooledRenderTarget = InPooledRenderTarget;
 
@@ -235,7 +304,6 @@ void FRDGTexture::SetRHI(FPooledRenderTarget* InPooledRenderTarget)
 
 void FRDGTexture::SetRHI(FRDGPooledTexture* InPooledTexture)
 {
-	check(InPooledTexture);
 	PooledTexture = InPooledTexture;
 	State = &PooledTexture->State;
 
@@ -251,7 +319,15 @@ void FRDGTexture::SetRHI(FRDGPooledTexture* InPooledTexture)
 	}
 
 	ResourceRHI = PooledTexture->GetRHI();
-	check(ResourceRHI);
+}
+
+void FRDGTexture::SetRHI(FRHITransientTexture* InTransientTexture, FRDGAllocator& Allocator)
+{
+	TransientTexture = InTransientTexture;
+	State = Allocator.AllocNoDestruct<FRDGTextureSubresourceState>();
+	ResourceRHI = TransientTexture->GetRHI();
+
+	bTransient = true;
 }
 
 void FRDGTexture::Finalize()
@@ -261,28 +337,35 @@ void FRDGTexture::Finalize()
 
 	if (bLastOwner)
 	{
-		// External and extracted resources are user controlled, so we cannot assume the texture stays in its final state.
-		if (bExternal || bExtracted)
+		if (bTransient)
 		{
-			PooledTexture->Reset();
+			// Manually deconstruct the allocated state so as not to invoke overhead from the allocators destructor tracking.
+			State->~FRDGTextureSubresourceState();
+			State = nullptr;
 		}
 		else
 		{
-			PooledTexture->Finalize();
-		}
-
-		// Restore the reference to the last owner in the aliasing chain; not necessary for the transient resource allocator.
-		if (PooledRenderTarget)
-		{
-			Allocation = PooledRenderTarget;
+			// External and extracted resources are user controlled, so we cannot assume the texture stays in its final state.
+			if (bExternal || bExtracted)
+			{
+				PooledTexture->Reset();
+			}
+			else
+			{
+				PooledTexture->Finalize();
+			}
+
+			// Restore the reference to the last owner in the aliasing chain; not necessary for the transient resource allocator.
+			if (PooledRenderTarget)
+			{
+				Allocation = PooledRenderTarget;
+			}
 		}
 	}
 }
 
 void FRDGBuffer::SetRHI(FRDGPooledBuffer* InPooledBuffer)
 {
-	check(InPooledBuffer);
-
 	// Return the previous owner and assign this buffer as the new one.
 	FRDGBuffer* PreviousOwner = InPooledBuffer->Owner;
 	InPooledBuffer->Owner = this;
@@ -297,8 +380,16 @@ void FRDGBuffer::SetRHI(FRDGPooledBuffer* InPooledBuffer)
 	PooledBuffer = InPooledBuffer;
 	Allocation = InPooledBuffer;
 	State = &PooledBuffer->State;
-	ResourceRHI = InPooledBuffer->Buffer;
-	check(ResourceRHI);
+	ResourceRHI = InPooledBuffer->GetRHI();
+}
+
+void FRDGBuffer::SetRHI(FRHITransientBuffer* InTransientBuffer, FRDGAllocator& Allocator)
+{
+	TransientBuffer = InTransientBuffer;
+	State = Allocator.AllocNoDestruct<FRDGSubresourceState>();
+	ResourceRHI = TransientBuffer->GetRHI();
+
+	bTransient = true;;
 }
 
 void FRDGBuffer::Finalize()
@@ -309,44 +400,46 @@ void FRDGBuffer::Finalize()
 
 	if (bLastOwner)
 	{
-		if (bExternal || bExtracted)
+		if (bTransient)
 		{
-			PooledBuffer->Reset();
+			State = nullptr;
 		}
 		else
 		{
-			PooledBuffer->Finalize();
+			if (bExternal || bExtracted)
+			{
+				PooledBuffer->Reset();
+			}
+			else
+			{
+				PooledBuffer->Finalize();
+			}
+
+			// Restore the reference to the last owner in the aliasing chain.
+			Allocation = PooledBuffer;
 		}
-
-		// Restore the reference to the last owner in the aliasing chain.
-		Allocation = PooledBuffer;
 	}
 }
 
-FRHIShaderResourceView* FRDGPooledBuffer::GetOrCreateSRV(FRDGBufferSRVDesc SRVDesc)
+FRHIShaderResourceView* FRDGPooledBuffer::GetOrCreateSRV(const FRHIBufferSRVCreateInfo& SRVDesc)
 {
-	if (const auto* FoundPtr = SRVs.Find(SRVDesc))
+	for (const auto& KeyValue : SRVs)
 	{
-		return FoundPtr->GetReference();
+		if (KeyValue.Key == SRVDesc)
+		{
+			return KeyValue.Value.GetReference();
+		}
 	}
 
 	FShaderResourceViewRHIRef RHIShaderResourceView;
 
-	if (Desc.UnderlyingType == FRDGBufferDesc::EUnderlyingType::VertexBuffer)
+	if (SRVDesc.Format != PF_Unknown)
 	{
-		RHIShaderResourceView = RHICreateShaderResourceView(Buffer, SRVDesc.BytesPerElement, SRVDesc.Format);
-	}
-	else if (Desc.UnderlyingType == FRDGBufferDesc::EUnderlyingType::StructuredBuffer)
-	{
-		RHIShaderResourceView = RHICreateShaderResourceView(Buffer);
-	}
-	else if (Desc.UnderlyingType == FRDGBufferDesc::EUnderlyingType::AccelerationStructure)
-	{
-		RHIShaderResourceView = RHICreateShaderResourceView(Buffer);
+		RHIShaderResourceView = RHICreateShaderResourceView(GetRHI(), SRVDesc.BytesPerElement, SRVDesc.Format);
 	}
 	else
 	{
-		checkNoEntry();
+		RHIShaderResourceView = RHICreateShaderResourceView(GetRHI());
 	}
 
 	FRHIShaderResourceView* View = RHIShaderResourceView.GetReference();
@@ -354,26 +447,25 @@ FRHIShaderResourceView* FRDGPooledBuffer::GetOrCreateSRV(FRDGBufferSRVDesc SRVDe
 	return View;
 }
 
-FRHIUnorderedAccessView* FRDGPooledBuffer::GetOrCreateUAV(FRDGBufferUAVDesc UAVDesc)
+FRHIUnorderedAccessView* FRDGPooledBuffer::GetOrCreateUAV(const FRHIBufferUAVCreateInfo& UAVDesc)
 {
-	if (const auto* FoundPtr = UAVs.Find(UAVDesc))
+	for (const auto& KeyValue : UAVs)
 	{
-		return FoundPtr->GetReference();
+		if (KeyValue.Key == UAVDesc)
+		{
+			return KeyValue.Value.GetReference();
+		}
 	}
 
 	FUnorderedAccessViewRHIRef RHIUnorderedAccessView;
 
-	if (Desc.UnderlyingType == FRDGBufferDesc::EUnderlyingType::VertexBuffer)
-	{
-		RHIUnorderedAccessView = RHICreateUnorderedAccessView(Buffer, UAVDesc.Format);
-	}
-	else if (Desc.UnderlyingType == FRDGBufferDesc::EUnderlyingType::StructuredBuffer)
+	if (UAVDesc.Format != PF_Unknown)
 	{
-		RHIUnorderedAccessView = RHICreateUnorderedAccessView(Buffer, UAVDesc.bSupportsAtomicCounter, UAVDesc.bSupportsAppendBuffer);
+		RHIUnorderedAccessView = RHICreateUnorderedAccessView(GetRHI(), UAVDesc.Format);
 	}
 	else
 	{
-		checkNoEntry();
+		RHIUnorderedAccessView = RHICreateUnorderedAccessView(GetRHI(), UAVDesc.bSupportsAtomicCounter, UAVDesc.bSupportsAppendBuffer);
 	}
 
 	FRHIUnorderedAccessView* View = RHIUnorderedAccessView.GetReference();
diff --git a/Engine/Source/Runtime/RenderCore/Private/RenderGraphUtils.cpp b/Engine/Source/Runtime/RenderCore/Private/RenderGraphUtils.cpp
--- a/Engine/Source/Runtime/RenderCore/Private/RenderGraphUtils.cpp
+++ b/Engine/Source/Runtime/RenderCore/Private/RenderGraphUtils.cpp
@@ -759,11 +759,15 @@ const void* GetInitialData(FRDGBuilder& GraphBuilder, const void* InitialData, u
 
 void AddBufferUploadPass(FRDGBuilder& GraphBuilder, FRDGBufferRef Buffer, const void* InitialData, uint64 InitialDataSize, ERDGInitialDataFlags InitialDataFlags)
 {
+	check(Buffer);
+
 	const void* SourcePtr = GetInitialData(GraphBuilder, InitialData, InitialDataSize, InitialDataFlags);
 
 	FCopyBufferParameters* PassParameters = GraphBuilder.AllocParameters<FCopyBufferParameters>();
 	PassParameters->Buffer = Buffer;
 
+	Buffer->SetNonTransient();
+
 	GraphBuilder.AddPass(
 		RDG_EVENT_NAME("BufferUpload(%s)", Buffer->Name),
 		PassParameters,
diff --git a/Engine/Source/Runtime/RenderCore/Private/RenderTargetPool.cpp b/Engine/Source/Runtime/RenderCore/Private/RenderTargetPool.cpp
--- a/Engine/Source/Runtime/RenderCore/Private/RenderTargetPool.cpp
+++ b/Engine/Source/Runtime/RenderCore/Private/RenderTargetPool.cpp
@@ -1223,12 +1223,12 @@ void FPooledRenderTarget::InitRDG()
 
 	if (RenderTargetItem.TargetableTexture)
 	{
-		TargetableTexture = FRDGPooledTexture::CreateCommitted(RenderTargetItem.TargetableTexture, Translate(Desc, ERenderTargetTexture::Targetable), RenderTargetItem.UAV);
+		TargetableTexture = new FRDGPooledTexture(RenderTargetItem.TargetableTexture, Translate(Desc, ERenderTargetTexture::Targetable), RenderTargetItem.UAV);
 	}
 
 	if (RenderTargetItem.ShaderResourceTexture != RenderTargetItem.TargetableTexture)
 	{
-		ShaderResourceTexture = FRDGPooledTexture::CreateCommitted(RenderTargetItem.ShaderResourceTexture, Translate(Desc, ERenderTargetTexture::ShaderResource), nullptr);
+		ShaderResourceTexture = new FRDGPooledTexture(RenderTargetItem.ShaderResourceTexture, Translate(Desc, ERenderTargetTexture::ShaderResource), nullptr);
 	}
 	else
 	{
diff --git a/Engine/Source/Runtime/RenderCore/Public/RenderGraphDefinitions.h b/Engine/Source/Runtime/RenderCore/Public/RenderGraphDefinitions.h
--- a/Engine/Source/Runtime/RenderCore/Public/RenderGraphDefinitions.h
+++ b/Engine/Source/Runtime/RenderCore/Public/RenderGraphDefinitions.h
@@ -177,30 +177,7 @@ inline ERDGParentResourceType GetParentResourceType(ERDGViewType ViewType)
 	}
 }
 
-/** Used to specify a texture metadata plane when creating a view. */
-enum class ERDGTextureMetaDataAccess : uint8
-{
-	/** The primary plane is used with default compression behavior. */
-	None = 0,
-
-	/** The primary plane is used without decompressing it. */
-	CompressedSurface,
-
-	/** The depth plane is used with default compression behavior. */
-	Depth,
-
-	/** The stencil plane is used with default compression behavior. */
-	Stencil,
-
-	/** The HTile plane is used. */
-	HTile,
-
-	/** the FMask plane is used. */
-	FMask,
-
-	/** the CMask plane is used. */
-	CMask
-};
+using ERDGTextureMetaDataAccess = ERHITextureMetaDataAccess;
 
 /** Returns the associated FRHITransitionInfo plane index. */
 inline int32 GetResourceTransitionPlaneForMetadataAccess(ERDGTextureMetaDataAccess Metadata)
@@ -357,13 +334,7 @@ class TRDGHandleBitArray : public TBitArray<FRDGBitArrayAllocator>
 {
 	using Base = TBitArray<FRDGBitArrayAllocator>;
 public:
-	TRDGHandleBitArray() = default;
-	TRDGHandleBitArray(const TRDGHandleBitArray&) = default;
-	TRDGHandleBitArray& operator=(const TRDGHandleBitArray&) = default;
-
-	explicit TRDGHandleBitArray(bool bValue, int32 InNumBits)
-		: Base(bValue, InNumBits)
-	{}
+	using Base::Base;
 
 	FORCEINLINE FBitReference operator[](HandleType Handle)
 	{
diff --git a/Engine/Source/Runtime/RenderCore/Public/RenderGraphResources.h b/Engine/Source/Runtime/RenderCore/Public/RenderGraphResources.h
--- a/Engine/Source/Runtime/RenderCore/Public/RenderGraphResources.h
+++ b/Engine/Source/Runtime/RenderCore/Public/RenderGraphResources.h
@@ -5,6 +5,7 @@
 #include "RenderGraphParameter.h"
 #include "RenderGraphTextureSubresource.h"
 #include "RendererInterface.h"
+#include "RHITransientResourceAllocator.h"
 
 struct FPooledRenderTarget;
 class FRenderTargetPool;
@@ -234,6 +235,13 @@ public:
 		return bProduced;
 	}
 
+	/** Marks a resource as excluded from the transient allocator, if it would otherwise be included. */
+	void SetNonTransient()
+	{
+		check(!bTransient);
+		bUserSetNonTransient = 1;
+	}
+
 protected:
 	FRDGParentResource(const TCHAR* InName, ERDGParentResourceType InType);
 	~FRDGParentResource();
@@ -247,9 +255,12 @@ protected:
 	/** Whether any sub-resource has been used for write by a pass. */
 	uint8 bProduced : 1;
 
-	/** Whether this resource needs acquire / discard. */
+	/** Whether this resource is allocated through the transient resource allocator. */
 	uint8 bTransient : 1;
 
+	/** Whether this resource is set to be non-transient by the user. */
+	uint8 bUserSetNonTransient : 1;
+
 	/** (External | Extracted only) If true, the resource is locked in its current state and will not be transitioned any more. */
 	uint8 bFinalizedAccess : 1;
 
@@ -328,73 +339,23 @@ inline FRDGTextureDesc Translate(const FPooledRenderTargetDesc& InDesc, ERenderT
 /** Translates from an RDG texture descriptor to a pooled render target descriptor. */
 inline FPooledRenderTargetDesc Translate(const FRDGTextureDesc& InDesc);
 
-class FRDGPooledResource
-{
-public:
-	virtual ~FRDGPooledResource() = default;
-
-	uint32 GetRefCount() const
-	{
-		return RefCount;
-	}
-
-	uint32 AddRef() const
-	{
-		return ++RefCount;
-	}
-
-	uint32 Release() const
-	{
-		check(RefCount > 0);
-		uint32 LocalRefCount = --RefCount;
-		if (LocalRefCount == 0)
-		{
-			if (Allocator == EAllocator::Default)
-			{
-				delete this;
-			}
-			else
-			{
-				this->~FRDGPooledResource();
-			}
-		}
-		return LocalRefCount;
-	}
-
-protected:
-	enum class EAllocator
-	{
-		/** Allocated with the default allocator (new / delete). */
-		Default,
-
-		/** Allocated with the RDG allocator. */
-		RDG
-	};
-
-	FRDGPooledResource(EAllocator InAllocator)
-		: Allocator(InAllocator)
-	{}
-
-	EAllocator Allocator;
-
-	mutable uint32 RefCount = 0;
-};
-
 class RENDERCORE_API FRDGPooledTexture final
-	: public FRDGPooledResource
+	: public FRefCountedObject
 {
 public:
-	/** Creates a pooled texture using the default allocator. Ref-counting determines deletion of the memory. Safe to be persistently cached. */
-	static TRefCountPtr<FRDGPooledTexture> CreateCommitted(TRefCountPtr<FRHITexture> Texture, const FRDGTextureSubresourceLayout& Layout, const FUnorderedAccessViewRHIRef& FirstMipUAV = nullptr)
+	FRDGPooledTexture(FRHITexture* InTexture, const FRDGTextureSubresourceLayout& InLayout, const FUnorderedAccessViewRHIRef& FirstMipUAV)
+		: Texture(InTexture)
+		, Layout(InLayout)
 	{
-		return TRefCountPtr<FRDGPooledTexture>(new FRDGPooledTexture(MoveTemp(Texture), Layout, FirstMipUAV, EAllocator::Default));
+		InitViews(FirstMipUAV);
+		Reset();
 	}
 
-	/** Creates a pooled texture using the RDG allocator. All references should be released before the allocator is freed. */
-	static TRefCountPtr<FRDGPooledTexture> CreateTransient(FRDGAllocator& Allocator, const FRDGTextureSubresourceLayout& Layout, TRefCountPtr<FRHITexture> Texture, const FUnorderedAccessViewRHIRef& FirstMipUAV = nullptr)
-	{
-		return TRefCountPtr<FRDGPooledTexture>(Allocator.AllocNoDestruct<FRDGPooledTexture>(MoveTemp(Texture), Layout, FirstMipUAV, EAllocator::RDG));
-	}
+	/** Finds a UAV matching the descriptor in the cache or creates a new one and updates the cache. */
+	FRHIUnorderedAccessView* GetOrCreateUAV(const FRHITextureUAVCreateInfo& UAVDesc);
+
+	/** Finds a SRV matching the descriptor in the cache or creates a new one and updates the cache. */
+	FRHIShaderResourceView* GetOrCreateSRV(const FRHITextureSRVCreateInfo& SRVDesc);
 
 	FRHITexture* GetRHI() const
 	{
@@ -407,15 +368,6 @@ public:
 	}
 
 private:
-	FRDGPooledTexture(TRefCountPtr<FRHITexture> InTexture, const FRDGTextureSubresourceLayout& InLayout, const FUnorderedAccessViewRHIRef& FirstMipUAV, EAllocator InAllocator)
-		: FRDGPooledResource(InAllocator)
-		, Texture(MoveTemp(InTexture))
-		, Layout(InLayout)
-	{
-		InitViews(FirstMipUAV);
-		Reset();
-	}
-
 	/** Initializes cached views. Safe to call multiple times; each call will recreate. */
 	void InitViews(const FUnorderedAccessViewRHIRef& FirstMipUAV);
 
@@ -501,10 +453,15 @@ private:
 		InitAsWholeResource(LastProducers);
 	}
 
-	/** Assigns the pooled texture to this texture; returns the previous texture to own the allocation. */
+	/** Assigns a pooled render target as the backing RHI resource. */
 	void SetRHI(FPooledRenderTarget* PooledRenderTarget);
+
+	/** Assigns a pooled buffer as the backing RHI resource. */
 	void SetRHI(FRDGPooledTexture* PooledTexture);
 
+	/** Assigns a transient buffer as the backing RHI resource. */
+	void SetRHI(FRHITransientTexture* TransientTexture, FRDGAllocator& Allocator);
+
 	/** Finalizes the texture for execution; no other transitions are allowed after calling this. */
 	void Finalize();
 
@@ -542,8 +499,14 @@ private:
 	/** The assigned pooled render target to use during execution. Never reset. */
 	IPooledRenderTarget* PooledRenderTarget = nullptr;
 
-	/** The assigned pooled texture to use during execution. Never reset. */
-	TRefCountPtr<FRDGPooledTexture> PooledTexture;
+	union
+	{
+		/** The assigned pooled texture to use during execution. Never reset. */
+		FRDGPooledTexture* PooledTexture = nullptr;
+
+		/** The assigned transient texture to use during execution. Never reset. */
+		FRHITransientTexture* TransientTexture;
+	};
 
 	/** Cached state pointer from the pooled texture. */
 	FRDGTextureSubresourceState* State = nullptr;
@@ -628,7 +591,6 @@ public:
 	FRDGTextureSRVDesc() = default;
 	
 	FRDGTextureRef Texture = nullptr;
-	ERDGTextureMetaDataAccess MetaData = ERDGTextureMetaDataAccess::None;
 
 	/** Create SRV that access all sub resources of texture. */
 	static FRDGTextureSRVDesc Create(FRDGTextureRef Texture)
@@ -693,14 +655,15 @@ private:
 };
 
 /** Descriptor for render graph tracked UAV. */
-class FRDGTextureUAVDesc
+class FRDGTextureUAVDesc final
+	: public FRHITextureUAVCreateInfo
 {
 public:
 	FRDGTextureUAVDesc() = default;
 
-	FRDGTextureUAVDesc(FRDGTextureRef InTexture, uint8 InMipLevel = 0)
-		: Texture(InTexture)
-		, MipLevel(InMipLevel)
+	FRDGTextureUAVDesc(FRDGTextureRef InTexture, uint8 InMipLevel = 0, ERHITextureMetaDataAccess InMetaData = ERHITextureMetaDataAccess::None)
+		: FRHITextureUAVCreateInfo(InMipLevel, InMetaData)
+		, Texture(InTexture)
 	{}
 
 	/** Create UAV with access to a specific meta data plane */
@@ -712,8 +675,6 @@ public:
 	}
 
 	FRDGTextureRef Texture = nullptr;
-	uint8 MipLevel = 0;
-	ERDGTextureMetaDataAccess MetaData = ERDGTextureMetaDataAccess::None;
 };
 
 /** Render graph tracked texture UAV. */
@@ -854,75 +815,57 @@ struct FRDGBufferDesc
 	EUnderlyingType UnderlyingType = EUnderlyingType::VertexBuffer;
 };
 
-struct FRDGBufferSRVDesc
+struct FRDGBufferSRVDesc final
+	: public FRHIBufferSRVCreateInfo
 {
 	FRDGBufferSRVDesc() = default;
 
 	FRDGBufferSRVDesc(FRDGBufferRef InBuffer);
 
 	FRDGBufferSRVDesc(FRDGBufferRef InBuffer, EPixelFormat InFormat)
-		: Buffer(InBuffer)
-		, Format(InFormat)
+		: FRHIBufferSRVCreateInfo(InFormat)
+		, Buffer(InBuffer)
 	{
 		BytesPerElement = GPixelFormats[Format].BlockBytes;
 	}
 
 	FRDGBufferRef Buffer = nullptr;
-
-	/** Number of bytes per element (used for vertex buffer). */
-	uint32 BytesPerElement = 1;
-
-	/** Encoding format for the element (used for vertex buffer). */
-	EPixelFormat Format = PF_Unknown;
 };
 
-struct FRDGBufferUAVDesc
+struct FRDGBufferUAVDesc final
+	: public FRHIBufferUAVCreateInfo
 {
 	FRDGBufferUAVDesc() = default;
 
 	FRDGBufferUAVDesc(FRDGBufferRef InBuffer);
 
 	FRDGBufferUAVDesc(FRDGBufferRef InBuffer, EPixelFormat InFormat)
-		: Buffer(InBuffer)
-		, Format(InFormat)
+		: FRHIBufferUAVCreateInfo(InFormat)
+		, Buffer(InBuffer)
 	{}
 
 	FRDGBufferRef Buffer = nullptr;
-
-	/** Number of bytes per element (used for vertex buffer). */
-	EPixelFormat Format = PF_Unknown;
-
-	/** Whether the uav supports atomic counter or append buffer operations (used for structured buffers) */
-	bool bSupportsAtomicCounter = false;
-	bool bSupportsAppendBuffer = false;
 };
 
 /** Translates from a RDG buffer descriptor to a RHI buffer creation info */
 inline FRHIBufferCreateInfo Translate(const FRDGBufferDesc& InDesc);
 
 class RENDERCORE_API FRDGPooledBuffer final
-	: public FRDGPooledResource
+	: public FRefCountedObject
 {
 public:
-	/** Creates a pooled texture using the default allocator. Ref-counting determines deletion of the memory. Safe to be persistently cached. */
-	static TRefCountPtr<FRDGPooledBuffer> CreateCommitted(TRefCountPtr<FRHIBuffer> Buffer, const FRDGBufferDesc& Desc)
-	{
-		return TRefCountPtr<FRDGPooledBuffer>(new FRDGPooledBuffer(MoveTemp(Buffer), Desc, EAllocator::Default));
-	}
-
-	/** Creates a pooled texture using the RDG allocator. All references should be released before the allocator is freed. */
-	static TRefCountPtr<FRDGPooledBuffer> CreateTransient(FRDGAllocator& Allocator, TRefCountPtr<FRHIBuffer> Buffer, const FRDGBufferDesc& Desc)
-	{
-		return TRefCountPtr<FRDGPooledBuffer>(Allocator.AllocNoDestruct<FRDGPooledBuffer>(MoveTemp(Buffer), Desc, EAllocator::RDG));
-	}
+	FRDGPooledBuffer(TRefCountPtr<FRHIBuffer> InBuffer, const FRDGBufferDesc& InDesc)
+		: Desc(InDesc)
+		, Buffer(MoveTemp(InBuffer))
+	{}
 
 	const FRDGBufferDesc Desc;
 
 	/** Finds a UAV matching the descriptor in the cache or creates a new one and updates the cache. */
-	FRHIUnorderedAccessView* GetOrCreateUAV(FRDGBufferUAVDesc UAVDesc);
+	FRHIUnorderedAccessView* GetOrCreateUAV(const FRHIBufferUAVCreateInfo& UAVDesc);
 
 	/** Finds a SRV matching the descriptor in the cache or creates a new one and updates the cache. */
-	FRHIShaderResourceView* GetOrCreateSRV(FRDGBufferSRVDesc SRVDesc);
+	FRHIShaderResourceView* GetOrCreateSRV(const FRHIBufferSRVCreateInfo& SRVDesc);
 
 	/** Returns the RHI buffer. */
 	FRHIBuffer* GetRHI() const
@@ -942,56 +885,10 @@ public:
 		return Buffer;
 	}
 
-	template<typename KeyType, typename ValueType>
-	struct TSRVFuncs : BaseKeyFuncs<TPair<KeyType, ValueType>, KeyType, /* bInAllowDuplicateKeys = */ false>
-	{
-		typedef typename TTypeTraits<KeyType>::ConstPointerType KeyInitType;
-		typedef const TPairInitializer<typename TTypeTraits<KeyType>::ConstInitType, typename TTypeTraits<ValueType>::ConstInitType>& ElementInitType;
-
-		static FORCEINLINE KeyInitType GetSetKey(ElementInitType Element)
-		{
-			return Element.Key;
-		}
-		static FORCEINLINE bool Matches(KeyInitType A, KeyInitType B)
-		{
-			return A.BytesPerElement == B.BytesPerElement && A.Format == B.Format;
-		}
-		static FORCEINLINE uint32 GetKeyHash(KeyInitType Key)
-		{
-			return HashCombine(uint32(Key.BytesPerElement), uint32(Key.Format));
-		}
-	};
-
-	template<typename KeyType, typename ValueType>
-	struct TUAVFuncs : BaseKeyFuncs<TPair<KeyType, ValueType>, KeyType, /* bInAllowDuplicateKeys = */ false>
-	{
-		typedef typename TTypeTraits<KeyType>::ConstPointerType KeyInitType;
-		typedef const TPairInitializer<typename TTypeTraits<KeyType>::ConstInitType, typename TTypeTraits<ValueType>::ConstInitType>& ElementInitType;
-
-		static FORCEINLINE KeyInitType GetSetKey(ElementInitType Element)
-		{
-			return Element.Key;
-		}
-		static FORCEINLINE bool Matches(KeyInitType A, KeyInitType B)
-		{
-			return A.Format == B.Format && A.bSupportsAtomicCounter == B.bSupportsAtomicCounter && A.bSupportsAppendBuffer == B.bSupportsAppendBuffer;
-		}
-		static FORCEINLINE uint32 GetKeyHash(KeyInitType Key)
-		{
-			return (uint32(Key.bSupportsAtomicCounter) << 8) | (uint32(Key.bSupportsAppendBuffer) << 9) | uint32(Key.Format);
-		}
-	};
-
 private:
-	FRDGPooledBuffer(TRefCountPtr<FRHIBuffer> InBuffer, const FRDGBufferDesc& InDesc, EAllocator InAllocator)
-		: FRDGPooledResource(InAllocator)
-		, Desc(InDesc)
-		, Buffer(MoveTemp(InBuffer))
-	{}
-
 	TRefCountPtr<FRHIBuffer> Buffer;
-	TMap<FRDGBufferUAVDesc, FUnorderedAccessViewRHIRef, FDefaultSetAllocator, TUAVFuncs<FRDGBufferUAVDesc, FUnorderedAccessViewRHIRef>> UAVs;
-	TMap<FRDGBufferSRVDesc, FShaderResourceViewRHIRef, FDefaultSetAllocator, TSRVFuncs<FRDGBufferSRVDesc, FShaderResourceViewRHIRef>> SRVs;
+	TArray<TPair<FRHIBufferUAVCreateInfo, FUnorderedAccessViewRHIRef>, TInlineAllocator<1>> UAVs;
+	TArray<TPair<FRHIBufferSRVCreateInfo, FShaderResourceViewRHIRef>, TInlineAllocator<1>> SRVs;
 
 	void Reset()
 	{
@@ -1069,18 +966,20 @@ private:
 		, Flags(InFlags)
 	{}
 
-	/** Assigns the pooled buffer to this buffer; returns the previous buffer to own the allocation. */
+	/** Assigns a pooled buffer as the backing RHI resource. */
 	void SetRHI(FRDGPooledBuffer* InPooledBuffer);
 
-	/** Returns RHI buffer without access checks. */
+	/** Assigns a transient buffer as the backing RHI resource. */
+	void SetRHI(FRHITransientBuffer* InTransientBuffer, FRDGAllocator& Allocator);
+
+	/** Finalizes the buffer for execution; no other transitions are allowed after calling this. */
+	void Finalize();
+
 	FRHIBuffer* GetRHIUnchecked() const
 	{
 		return static_cast<FRHIBuffer*>(FRDGResource::GetRHIUnchecked());
 	}
 
-	/** Finalizes the buffer for execution; no other transitions are allowed after calling this. */
-	void Finalize();
-
 	/** Returns the current buffer state. Only valid to call after SetRHI. */
 	FRDGSubresourceState& GetState() const
 	{
@@ -1097,10 +996,16 @@ private:
 	/** The next buffer to own the PooledBuffer allocation during execution. */
 	FRDGBufferHandle NextOwner;
 
-	/** Assigned pooled buffer pointer. Never reset once assigned. */
-	FRDGPooledBuffer* PooledBuffer = nullptr;
+	union
+	{
+		/** Assigned pooled buffer pointer. Never reset once assigned. */
+		FRDGPooledBuffer* PooledBuffer = nullptr;
+
+		/** Assigned transient buffer pointer. Never reset once assigned. */
+		FRHITransientBuffer* TransientBuffer;
+	};
 
-	/** Cached state pointer from the pooled buffer. */
+	/** Cached state pointer from the pooled / transient buffer. */
 	FRDGSubresourceState* State = nullptr;
 
 	/** Valid strictly when holding a strong reference; use PooledBuffer instead. */
