Added ability to filter geometries by name when using D3D12.DumpRayTracingGeometries command. For example `D3D12.DumpRayTracingGeometries 50 skm_` will show 50 largest geometries that have "skm_" in their name.

diff --git a/Engine/Source/Runtime/D3D12RHI/Private/D3D12RayTracing.cpp b/Engine/Source/Runtime/D3D12RHI/Private/D3D12RayTracing.cpp
--- a/Engine/Source/Runtime/D3D12RHI/Private/D3D12RayTracing.cpp
+++ b/Engine/Source/Runtime/D3D12RHI/Private/D3D12RayTracing.cpp
@@ -177,6 +177,8 @@ static FAutoConsoleCommandWithWorldArgsAndOutputDevice GD3D12DumpRayTracingGeome
 	EMode Mode = EMode::Top;
 	int32 NumEntriesToShow = 50;
 
+	FString NameFilter;
+
 	if (Args.Num())
 	{
 		if (Args[0] == TEXT("all"))
@@ -189,6 +191,11 @@ static FAutoConsoleCommandWithWorldArgsAndOutputDevice GD3D12DumpRayTracingGeome
 			Mode = EMode::Top;
 			LexFromString(NumEntriesToShow, *Args[0]);
 		}
+
+		if (Args.Num() > 1)
+		{
+			NameFilter = Args[1];
+		}
 	}
 
 	auto GetGeometrySize = [](FD3D12RayTracingGeometry& Geometry)
@@ -225,11 +232,31 @@ static FAutoConsoleCommandWithWorldArgsAndOutputDevice GD3D12DumpRayTracingGeome
 		BufferedOutput.CategorizedLogf(CategoryName, ELogVerbosity::Log, TEXT("Showing %d out of %d"), NumEntriesToShow, Geometries.Num());
 	}
 
+	auto ShouldShow = [&NameFilter](FD3D12RayTracingGeometry* Entry)
+	{
+		if (NameFilter.IsEmpty())
+		{
+			return true;
+		}
+
+		FString DebugName = Entry->DebugName.ToString();
+		if (DebugName.Find(NameFilter, ESearchCase::IgnoreCase) != INDEX_NONE)
+		{
+			return true;
+		}
+		else
+		{
+			return false;
+		}
+	};
+
+	int32 ShownEntries = 0;
 	for (int32 i=0; i< Geometries.Num(); ++i)
 	{
 		FD3D12RayTracingGeometry* Geometry = Geometries[i];
 		uint64 SizeBytes = GetGeometrySize(*Geometry);
-		if (i < NumEntriesToShow)
+
+		if (ShownEntries < NumEntriesToShow && ShouldShow(Geometry))
 		{
 			BufferedOutput.CategorizedLogf(CategoryName, ELogVerbosity::Log, TEXT("Name: %s - Size: %.3f MB - Prims: %d - Segments: %d -  Compaction: %d - Update: %d"),
 				Geometry->DebugName.IsValid() ? *Geometry->DebugName.ToString() : TEXT("*UNKNOWN*"),
@@ -239,19 +266,22 @@ static FAutoConsoleCommandWithWorldArgsAndOutputDevice GD3D12DumpRayTracingGeome
 				(int32)EnumHasAllFlags(Geometry->BuildFlags, ERayTracingAccelerationStructureFlags::AllowCompaction),
 				(int32)EnumHasAllFlags(Geometry->BuildFlags, ERayTracingAccelerationStructureFlags::AllowUpdate));
 			TopSizeBytes += SizeBytes;
+			++ShownEntries;
 		}
+
 		TotalSizeBytes += SizeBytes;
 	}
 
 	double TotalSizeF = double(TotalSizeBytes) / double(1 << 20);
 	double TopSizeF = double(TopSizeBytes) / double(1 << 20);
 
-	if (NumEntriesToShow != Geometries.Num())
+	if (ShownEntries != Geometries.Num() && ShownEntries)
 	{
-		BufferedOutput.CategorizedLogf(CategoryName, ELogVerbosity::Log, TEXT("Use command `D3D12.DumpRayTracingGeometries all` to dump all objects."));
-
-		BufferedOutput.CategorizedLogf(CategoryName, ELogVerbosity::Log, TEXT("Top %d size: %.3f MB (%.2f%% of total)"),
-			NumEntriesToShow, TopSizeF, 100.0 * TopSizeF / TotalSizeF);
+		BufferedOutput.CategorizedLogf(CategoryName, ELogVerbosity::Log,
+			TEXT("Use command `D3D12.DumpRayTracingGeometries all/N [name]` to dump all or N objects. ")
+			TEXT("Optionally add 'name' to filter entries, such as 'skm_'."));
+		BufferedOutput.CategorizedLogf(CategoryName, ELogVerbosity::Log, TEXT("Shown %d entries. Size: %.3f MB (%.2f%% of total)"),
+			ShownEntries, TopSizeF, 100.0 * TopSizeF / TotalSizeF);
 	}
 
 	BufferedOutput.CategorizedLogf(CategoryName, ELogVerbosity::Log, TEXT("Total size: %.3f MB"), TotalSizeF);
