Moved r.RayTracing.DebugDisableTriangleCull and r.RayTracing.DebugForceOpaque to the high-level rendering code. Previously was in D3D12 RHI.
Removed r.RayTracing.DebugForceFullBuild (was unused). Renamed remaining D3D12-specific ray tracing CVars to include r.D3D12 prefix.
Added TranslateRayTracingInstanceFlags() to translate high-level engine instance flags to D3D12.
Fixed r.D3D12.RayTracing.MaxBatchedCompaction CVar.

diff --git a/Engine/Source/Runtime/D3D12RHI/Private/D3D12RayTracing.cpp b/Engine/Source/Runtime/D3D12RHI/Private/D3D12RayTracing.cpp
--- a/Engine/Source/Runtime/D3D12RHI/Private/D3D12RayTracing.cpp
+++ b/Engine/Source/Runtime/D3D12RHI/Private/D3D12RayTracing.cpp
@@ -18,16 +18,9 @@
 #include "Misc/BufferedOutputDevice.h"
 #include "String/LexFromString.h"
 
-static int32 GRayTracingDebugForceOpaque = 0;
-static FAutoConsoleVariableRef CVarRayTracingDebugForceOpaque(
-	TEXT("r.RayTracing.DebugForceOpaque"),
-	GRayTracingDebugForceOpaque,
-	TEXT("Forces all ray tracing geometry instances to be opaque, effectively disabling any-hit shaders. This is useful for debugging and profiling. (default = 0)")
-);
-
 static int32 GRayTracingDebugForceBuildMode = 0;
 static FAutoConsoleVariableRef CVarRayTracingDebugForceFastTrace(
-	TEXT("r.RayTracing.DebugForceBuildMode"),
+	TEXT("r.D3D12.RayTracing.DebugForceBuildMode"),
 	GRayTracingDebugForceBuildMode,
 	TEXT("Forces specific acceleration structure build mode (not runtime-tweakable).\n")
 	TEXT("0: Use build mode requested by high-level code (Default)\n")
@@ -36,25 +29,9 @@ static FAutoConsoleVariableRef CVarRayTracingDebugForceFastTrace(
 	ECVF_ReadOnly
 );
 
-static int32 GRayTracingDebugForceFullBuild = 0;
-static FAutoConsoleVariableRef CVarRayTracingDebugForceFullBuild(
-	TEXT("r.RayTracing.DebugForceFullBuild"),
-	GRayTracingDebugForceFullBuild,
-	TEXT("Forces all acceleration structure updates to always perform a full build.\n")
-	TEXT("0: Allow update (Default)\n")
-	TEXT("1: Force full build\n")
-);
-
-static int32 GRayTracingDebugDisableTriangleCull = 0;
-static FAutoConsoleVariableRef CVarRayTracingDebugDisableTriangleCull(
-	TEXT("r.RayTracing.DebugDisableTriangleCull"),
-	GRayTracingDebugDisableTriangleCull,
-	TEXT("Forces all ray tracing geometry instances to be double-sided by disabling back-face culling. This is useful for debugging and profiling. (default = 0)")
-);
-
 static int32 GRayTracingCacheShaderRecords = 1;
 static FAutoConsoleVariableRef CVarRayTracingShaderRecordCache(
-	TEXT("r.RayTracing.CacheShaderRecords"),
+	TEXT("r.D3D12.RayTracing.CacheShaderRecords"),
 	GRayTracingCacheShaderRecords,
 	TEXT("Automatically cache and re-use SBT hit group records. This significantly improves CPU performance in large scenes with many identical mesh instances. (default = 1)\n")
 	TEXT("This mode assumes that contents of uniform buffers does not change during ray tracing resource binding.")
@@ -81,7 +58,7 @@ static FAutoConsoleVariableRef CVarD3D12RayTracingAllowCompaction(
 static int32 GD3D12RayTracingMaxBatchedCompaction = 64;
 static FAutoConsoleVariableRef CVarD3D12RayTracingMaxBatchedCompaction(
 	TEXT("r.D3D12.RayTracing.MaxBatchedCompaction"),
-	GD3D12RayTracingAllowCompaction,
+	GD3D12RayTracingMaxBatchedCompaction,
 	TEXT("Maximum of amount of compaction requests and rebuilds per frame. (default = 64)\n"),
 	ECVF_ReadOnly
 );
@@ -2569,6 +2546,34 @@ void FD3D12Device::CleanupRayTracing()
 	// Note: RayTracingDescriptorHeapCache is destroyed in ~FD3D12Device, after all deferred deletion is processed
 }
 
+static D3D12_RAYTRACING_INSTANCE_FLAGS TranslateRayTracingInstanceFlags(ERayTracingInstanceFlags InFlags)
+{
+	D3D12_RAYTRACING_INSTANCE_FLAGS Result = D3D12_RAYTRACING_INSTANCE_FLAG_NONE;
+
+	if (EnumHasAnyFlags(InFlags, ERayTracingInstanceFlags::TriangleCullDisable))
+	{
+		Result |= D3D12_RAYTRACING_INSTANCE_FLAG_TRIANGLE_CULL_DISABLE;
+	}
+
+	if (!EnumHasAnyFlags(InFlags, ERayTracingInstanceFlags::TriangleCullReverse))
+	{
+		// Counterclockwise is the default for UE. Reversing culling is achieved by *not* setting this flag.
+		Result |= D3D12_RAYTRACING_INSTANCE_FLAG_TRIANGLE_FRONT_COUNTERCLOCKWISE;
+	}
+
+	if (EnumHasAnyFlags(InFlags, ERayTracingInstanceFlags::ForceOpaque))
+	{
+		Result |= D3D12_RAYTRACING_INSTANCE_FLAG_FORCE_OPAQUE;
+	}
+
+	if (EnumHasAnyFlags(InFlags, ERayTracingInstanceFlags::ForceNonOpaque))
+	{
+		Result |= D3D12_RAYTRACING_INSTANCE_FLAG_FORCE_NON_OPAQUE;
+	}
+
+	return Result;
+}
+
 FRayTracingAccelerationStructureSize FD3D12DynamicRHI::RHICalcRayTracingSceneSize(uint32 MaxInstances, ERayTracingAccelerationStructureFlags Flags)
 {
 	D3D12_BUILD_RAYTRACING_ACCELERATION_STRUCTURE_INPUTS BuildInputs = {};
@@ -3016,8 +3021,6 @@ void FD3D12RayTracingGeometry::UnregisterAsRenameListener(uint32 InGPUIndex)
 	bRegisteredAsRenameListener[InGPUIndex] = false;
 }
 
-static constexpr uint32 IndicesPerPrimitive = 3; // Only triangle meshes are supported
-
 void FD3D12RayTracingGeometry::ResourceRenamed(FD3D12BaseShaderResource* InRenamedResource, FD3D12ResourceLocation* InNewResourceLocation)
 {
 	checkf(InNewResourceLocation, TEXT("Shouldn't release resources which are still referenced by the RayTracingGeometry"));
@@ -3065,6 +3068,7 @@ void FD3D12RayTracingGeometry::SetupHitGroupSystemParameters(uint32 InGPUIndex)
 
 		if (GeometryType == D3D12_RAYTRACING_GEOMETRY_TYPE_TRIANGLES && IndexBuffer != nullptr)
 		{
+			static constexpr uint32 IndicesPerPrimitive = 3; // Only triangle meshes are supported
 			SystemParameters.IndexBuffer = IndexBuffer->ResourceLocation.GetGPUVirtualAddress();
 			SystemParameters.RootConstants.IndexBufferOffsetInBytes = IndexOffsetInBytes + IndexStride * Segment.FirstPrimitive * IndicesPerPrimitive;
 		}
@@ -3148,6 +3152,7 @@ void FD3D12RayTracingGeometry::BuildAccelerationStructure(FD3D12CommandContext&
 
 			if (IndexBuffer)
 			{
+				static constexpr uint32 IndicesPerPrimitive = 3; // Only triangle meshes are supported
 				check(Desc.Triangles.IndexCount <= Segment.NumPrimitives * IndicesPerPrimitive);
 
 				Desc.Triangles.IndexFormat = (IndexStride == 4 ? DXGI_FORMAT_R32_UINT : DXGI_FORMAT_R16_UINT);
@@ -3392,27 +3397,7 @@ FD3D12RayTracingScene::FD3D12RayTracingScene(FD3D12Adapter* Adapter, const FRayT
 
 		InstanceDesc.InstanceMask = Instance.Mask;
 		InstanceDesc.InstanceContributionToHitGroupIndex = SegmentPrefixSum[InstanceIndex] * ShaderSlotsPerGeometrySegment;
-
-		if (EnumHasAnyFlags(Instance.Flags, ERayTracingInstanceFlags::TriangleCullDisable) || GRayTracingDebugDisableTriangleCull)
-		{
-			InstanceDesc.Flags |= D3D12_RAYTRACING_INSTANCE_FLAG_TRIANGLE_CULL_DISABLE;
-		}
-
-		if (!EnumHasAnyFlags(Instance.Flags, ERayTracingInstanceFlags::TriangleCullReverse))
-		{
-			// Counterclockwise is the default for UE. Reversing culling is achieved by *not* setting this flag.
-			InstanceDesc.Flags |= D3D12_RAYTRACING_INSTANCE_FLAG_TRIANGLE_FRONT_COUNTERCLOCKWISE;
-		}
-
-		if (EnumHasAnyFlags(Instance.Flags, ERayTracingInstanceFlags::ForceOpaque) || GRayTracingDebugForceOpaque)
-		{
-			InstanceDesc.Flags |= D3D12_RAYTRACING_INSTANCE_FLAG_FORCE_OPAQUE;
-		}
-
-		if (EnumHasAnyFlags(Instance.Flags, ERayTracingInstanceFlags::ForceNonOpaque))
-		{
-			InstanceDesc.Flags |= D3D12_RAYTRACING_INSTANCE_FLAG_FORCE_NON_OPAQUE;
-		}
+		InstanceDesc.Flags = TranslateRayTracingInstanceFlags(Instance.Flags);
 
 		const uint32 NumTransforms = Instance.NumTransforms;
 
diff --git a/Engine/Source/Runtime/Renderer/Private/DeferredShadingRenderer.cpp b/Engine/Source/Runtime/Renderer/Private/DeferredShadingRenderer.cpp
--- a/Engine/Source/Runtime/Renderer/Private/DeferredShadingRenderer.cpp
+++ b/Engine/Source/Runtime/Renderer/Private/DeferredShadingRenderer.cpp
@@ -200,6 +200,21 @@ static TAutoConsoleVariable<int32> CVarRayTracingAutoInstance(
 	ECVF_RenderThreadSafe
 );
 
+static int32 GRayTracingDebugDisableTriangleCull = 0;
+static FAutoConsoleVariableRef CVarRayTracingDebugDisableTriangleCull(
+	TEXT("r.RayTracing.DebugDisableTriangleCull"),
+	GRayTracingDebugDisableTriangleCull,
+	TEXT("Forces all ray tracing geometry instances to be double-sided by disabling back-face culling. This is useful for debugging and profiling. (default = 0)")
+);
+
+
+static int32 GRayTracingDebugForceOpaque = 0;
+static FAutoConsoleVariableRef CVarRayTracingDebugForceOpaque(
+	TEXT("r.RayTracing.DebugForceOpaque"),
+	GRayTracingDebugForceOpaque,
+	TEXT("Forces all ray tracing geometry instances to be opaque, effectively disabling any-hit shaders. This is useful for debugging and profiling. (default = 0)")
+);
+
 #if !UE_BUILD_SHIPPING
 static TAutoConsoleVariable<int32> CVarForceBlackVelocityBuffer(
 	TEXT("r.Test.ForceBlackVelocityBuffer"), 0,
@@ -463,6 +478,18 @@ static void RenderOpaqueFX(
 
 #if RHI_RAYTRACING
 
+static void AddDebugRayTracingInstanceFlags(ERayTracingInstanceFlags& InOutFlags)
+{
+	if (GRayTracingDebugForceOpaque)
+	{
+		InOutFlags |= ERayTracingInstanceFlags::ForceOpaque;
+	}
+	if (GRayTracingDebugDisableTriangleCull)
+	{
+		InOutFlags |= ERayTracingInstanceFlags::TriangleCullDisable;
+	}
+}
+
 bool FDeferredShadingSceneRenderer::GatherRayTracingWorldInstancesForView(FRHICommandListImmediate& RHICmdList, FViewInfo& View, FRayTracingScene& RayTracingScene)
 {
 	if (!IsRayTracingEnabled())
@@ -956,6 +983,7 @@ bool FDeferredShadingSceneRenderer::GatherRayTracingWorldInstancesForView(FRHICo
 					{
 						RayTracingInstance.Flags |= ERayTracingInstanceFlags::TriangleCullDisable;
 					}
+					AddDebugRayTracingInstanceFlags(RayTracingInstance.Flags);
 
 					if (Instance.InstanceGPUTransformsSRV.IsValid())
 					{
@@ -1145,6 +1173,7 @@ bool FDeferredShadingSceneRenderer::GatherRayTracingWorldInstancesForView(FRHICo
 				}
 
 				RayTracingScene.Instances.Add(SceneInfo->CachedRayTracingInstance);
+				AddDebugRayTracingInstanceFlags(RayTracingScene.Instances.Last().Flags);
 			}
 			else
 			{
@@ -1223,6 +1252,7 @@ bool FDeferredShadingSceneRenderer::GatherRayTracingWorldInstancesForView(FRHICo
 					{
 						RayTracingInstance.Flags |= ERayTracingInstanceFlags::TriangleCullDisable;
 					}
+					AddDebugRayTracingInstanceFlags(RayTracingInstance.Flags);
 				}
 			}
 		}
