Transient allocator support  for platform

diff --git a/Engine/Source/Runtime/D3D12RHI/Private/D3D12Adapter.cpp b/Engine/Source/Runtime/D3D12RHI/Private/D3D12Adapter.cpp
--- a/Engine/Source/Runtime/D3D12RHI/Private/D3D12Adapter.cpp
+++ b/Engine/Source/Runtime/D3D12RHI/Private/D3D12Adapter.cpp
@@ -1230,7 +1230,9 @@ void FD3D12Adapter::EndFrame()
 		GetUploadHeapAllocator(GPUIndex).CleanUpAllocations(FrameLag);
 	}
 	GetDeferredDeletionQueue().ReleaseResources(false, false);
-	GetTransientResourceSystem().GarbageCollect();
+
+	TransientResourceSystem->GarbageCollect();
+	TransientResourceSystem->UpdateStats();
 
 #if D3D12_SUBMISSION_GAP_RECORDER
 	SubmitGapRecorderTimestamps();
diff --git a/Engine/Source/Runtime/D3D12RHI/Private/D3D12Buffer.cpp b/Engine/Source/Runtime/D3D12RHI/Private/D3D12Buffer.cpp
--- a/Engine/Source/Runtime/D3D12RHI/Private/D3D12Buffer.cpp
+++ b/Engine/Source/Runtime/D3D12RHI/Private/D3D12Buffer.cpp
@@ -472,6 +472,12 @@ FD3D12Buffer* FD3D12DynamicRHI::CreateD3D12Buffer(class FRHICommandListImmediate
 	return Buffer;
 }
 
+FRHIBuffer* FD3D12DynamicRHI::CreateBuffer(const FRHIBufferCreateInfo& CreateInfo, const TCHAR* DebugName, ERHIAccess InitialState, ID3D12ResourceAllocator* ResourceAllocator)
+{
+	FRHIResourceCreateInfo ResourceCreateInfo(DebugName);
+	return CreateD3D12Buffer(nullptr, CreateInfo.Size, CreateInfo.Usage, CreateInfo.Stride, InitialState, ResourceCreateInfo, ResourceAllocator);
+}
+
 void* FD3D12DynamicRHI::LockBuffer(FRHICommandListImmediate* RHICmdList, FD3D12Buffer* Buffer, uint32 BufferSize, uint32 BufferUsage, uint32 Offset, uint32 Size, EResourceLockMode LockMode)
 {
 	SCOPE_CYCLE_COUNTER(STAT_D3D12LockBufferTime);
diff --git a/Engine/Source/Runtime/D3D12RHI/Private/D3D12CommandContext.cpp b/Engine/Source/Runtime/D3D12RHI/Private/D3D12CommandContext.cpp
--- a/Engine/Source/Runtime/D3D12RHI/Private/D3D12CommandContext.cpp
+++ b/Engine/Source/Runtime/D3D12RHI/Private/D3D12CommandContext.cpp
@@ -704,8 +704,6 @@ void FD3D12CommandContextBase::UpdateMemoryStats()
 	SET_MEMORY_STAT(STAT_D3D12DemotedVideoMemory, MemoryInfo.DemotedLocalMemory);
 	SET_MEMORY_STAT(STAT_D3D12TotalVideoMemory, MemoryInfo.LocalMemoryInfo.Budget);
 
-	ParentAdapter->GetTransientResourceSystem().UpdateStats();
-
 	uint64 MaxTexAllocWastage = 0;
 	for (uint32 GPUIndex : GPUMask)
 	{
diff --git a/Engine/Source/Runtime/D3D12RHI/Private/D3D12RHIPrivate.h b/Engine/Source/Runtime/D3D12RHI/Private/D3D12RHIPrivate.h
--- a/Engine/Source/Runtime/D3D12RHI/Private/D3D12RHIPrivate.h
+++ b/Engine/Source/Runtime/D3D12RHI/Private/D3D12RHIPrivate.h
@@ -1054,9 +1054,14 @@ public:
 	TD3D12Texture2D<BaseResourceType>* CreateD3D12Texture2D(class FRHICommandListImmediate* RHICmdList, uint32 SizeX, uint32 SizeY, uint32 SizeZ, bool bTextureArray, bool CubeTexture, EPixelFormat Format,
 		uint32 NumMips, uint32 NumSamples, ETextureCreateFlags InFlags, ERHIAccess InResourceState, FRHIResourceCreateInfo& CreateInfo, ID3D12ResourceAllocator* ResourceAllocator);
 
+	FD3D12Buffer* CreateD3D12Buffer(class FRHICommandListImmediate* RHICmdList, uint32 Size, EBufferUsageFlags Usage, uint32 Stride, ERHIAccess ResourceState, FRHIResourceCreateInfo& CreateInfo, ID3D12ResourceAllocator* ResourceAllocator);
+
 	FD3D12Texture3D* CreateD3D12Texture3D(class FRHICommandListImmediate* RHICmdList, uint32 SizeX, uint32 SizeY, uint32 SizeZ, EPixelFormat Format, uint32 NumMips, ETextureCreateFlags InFlags, ERHIAccess InResourceState, FRHIResourceCreateInfo& CreateInfo, ID3D12ResourceAllocator* ResourceAllocator);
 
-	FD3D12Buffer* CreateD3D12Buffer(class FRHICommandListImmediate* RHICmdList, uint32 Size, EBufferUsageFlags Usage, uint32 Stride, ERHIAccess ResourceState, FRHIResourceCreateInfo& CreateInfo, ID3D12ResourceAllocator* ResourceAllocator);
+	// Interface for FD3D12TransientResourceAllocator.
+	virtual D3D12_RESOURCE_DESC GetResourceDesc(const FRHITextureCreateInfo& CreateInfo) const;
+	virtual FRHITexture* CreateTexture(const FRHITextureCreateInfo& CreateInfo, const TCHAR* DebugName, ERHIAccess InitialState, ID3D12ResourceAllocator* ResourceAllocator);
+	FRHIBuffer* CreateBuffer(const FRHIBufferCreateInfo& CreateInfo, const TCHAR* DebugName, ERHIAccess InitialState, ID3D12ResourceAllocator* ResourceAllocator);
 
 protected:
 
diff --git a/Engine/Source/Runtime/D3D12RHI/Private/D3D12Texture.cpp b/Engine/Source/Runtime/D3D12RHI/Private/D3D12Texture.cpp
--- a/Engine/Source/Runtime/D3D12RHI/Private/D3D12Texture.cpp
+++ b/Engine/Source/Runtime/D3D12RHI/Private/D3D12Texture.cpp
@@ -1295,6 +1295,35 @@ FD3D12Texture3D* FD3D12DynamicRHI::CreateD3D12Texture3D(FRHICommandListImmediate
 #endif // PLATFORM_WINDOWS || PLATFORM_HOLOLENS
 }
 
+FRHITexture* FD3D12DynamicRHI::CreateTexture(const FRHITextureCreateInfo& CreateInfo, const TCHAR* DebugName, ERHIAccess InitialState, ID3D12ResourceAllocator* ResourceAllocator)
+{
+	FRHIResourceCreateInfo ResourceCreateInfo(DebugName, CreateInfo.ClearValue);
+
+	const bool bTextureArray = CreateInfo.IsTextureArray();
+	const bool bTextureCube = CreateInfo.IsTextureCube();
+
+	switch (CreateInfo.Dimension)
+	{
+	case ETextureDimension::Texture2D:
+		return CreateD3D12Texture2D<FD3D12BaseTexture2D>(nullptr, CreateInfo.Extent.X, CreateInfo.Extent.Y, 1, bTextureArray, bTextureCube, CreateInfo.Format, CreateInfo.NumMips, CreateInfo.NumSamples, CreateInfo.Flags, InitialState, ResourceCreateInfo, ResourceAllocator);
+
+	case ETextureDimension::Texture2DArray:
+		return CreateD3D12Texture2D<FD3D12BaseTexture2DArray>(nullptr, CreateInfo.Extent.X, CreateInfo.Extent.Y, 1, bTextureArray, bTextureCube, CreateInfo.Format, CreateInfo.NumMips, CreateInfo.NumSamples, CreateInfo.Flags, InitialState, ResourceCreateInfo, ResourceAllocator);
+
+	case ETextureDimension::TextureCube:
+	case ETextureDimension::TextureCubeArray:
+		return CreateD3D12Texture2D<FD3D12BaseTextureCube>(nullptr, CreateInfo.Extent.X, CreateInfo.Extent.Y, 6, bTextureArray, bTextureCube, CreateInfo.Format, CreateInfo.NumMips, CreateInfo.NumSamples, CreateInfo.Flags, InitialState, ResourceCreateInfo, ResourceAllocator);
+
+	case ETextureDimension::Texture3D:
+		return CreateD3D12Texture3D(nullptr, CreateInfo.Extent.X, CreateInfo.Extent.Y, CreateInfo.Depth, CreateInfo.Format, CreateInfo.NumMips, CreateInfo.Flags, InitialState, ResourceCreateInfo, ResourceAllocator);
+
+	default:
+		checkNoEntry();
+	}
+
+	return nullptr;
+}
+
 /*-----------------------------------------------------------------------------
 	2D texture support.
 	-----------------------------------------------------------------------------*/
@@ -1916,7 +1945,7 @@ void* TD3D12Texture2D<RHIResourceType>::Lock(class FRHICommandListImmediate* RHI
 	return Data;
 }
 
-D3D12_RESOURCE_DESC FD3D12TextureBase::GetResourceDesc(const FRHITextureCreateInfo& CreateInfo)
+D3D12_RESOURCE_DESC FD3D12DynamicRHI::GetResourceDesc(const FRHITextureCreateInfo& CreateInfo) const
 {
 	D3D12_RESOURCE_DESC Desc;
 
diff --git a/Engine/Source/Runtime/D3D12RHI/Private/D3D12Texture.h b/Engine/Source/Runtime/D3D12RHI/Private/D3D12Texture.h
--- a/Engine/Source/Runtime/D3D12RHI/Private/D3D12Texture.h
+++ b/Engine/Source/Runtime/D3D12RHI/Private/D3D12Texture.h
@@ -193,8 +193,6 @@ public:
 
 	FD3D12CLSyncPoint ReadBackSyncPoint;
 
-	static D3D12_RESOURCE_DESC GetResourceDesc(const FRHITextureCreateInfo& CreateInfo);
-
 protected:
 
 	/** Amount of memory allocated by this texture, in bytes. */
diff --git a/Engine/Source/Runtime/D3D12RHI/Private/D3D12TransientResourceAllocator.cpp b/Engine/Source/Runtime/D3D12RHI/Private/D3D12TransientResourceAllocator.cpp
--- a/Engine/Source/Runtime/D3D12RHI/Private/D3D12TransientResourceAllocator.cpp
+++ b/Engine/Source/Runtime/D3D12RHI/Private/D3D12TransientResourceAllocator.cpp
@@ -128,8 +128,10 @@ FD3D12TransientResourceAllocator::FD3D12TransientResourceAllocator(FD3D12Transie
 
 FRHITransientTexture* FD3D12TransientResourceAllocator::CreateTexture(const FRHITextureCreateInfo& InCreateInfo, const TCHAR* InDebugName)
 {
-	D3D12_RESOURCE_DESC Desc = FD3D12TextureBase::GetResourceDesc(InCreateInfo);
-	D3D12_RESOURCE_ALLOCATION_INFO Info = AllocationInfoQueryDevice->GetResourceAllocationInfo(Desc);
+	FD3D12DynamicRHI* DynamicRHI = FD3D12DynamicRHI::GetD3DRHI();
+
+	const D3D12_RESOURCE_DESC Desc = DynamicRHI->GetResourceDesc(InCreateInfo);
+	const D3D12_RESOURCE_ALLOCATION_INFO Info = AllocationInfoQueryDevice->GetResourceAllocationInfo(Desc);
 
 	return Allocator.CreateTexture(InCreateInfo, InDebugName, Info.SizeInBytes, Info.Alignment,
 		[&](const FRHITransientResourceAllocator::FResourceInitializer& Initializer)
@@ -145,43 +147,8 @@ FRHITransientTexture* FD3D12TransientResourceAllocator::CreateTexture(const FRHI
 			InitialState = ERHIAccess::DSVWrite;
 		}
 
-		FRHIResourceCreateInfo ResourceCreateInfo(InDebugName, InCreateInfo.ClearValue);
-		FResourceAllocatorAdapter ResourceAllocatorAdapter(GetParentAdapter(), static_cast<FD3D12TransientHeap&>(Initializer.Heap), Initializer.Allocation);
-		FRHITexture* Texture = nullptr;
-
-		const bool bTextureArray = InCreateInfo.ArraySize > 1;
-
-		switch (InCreateInfo.Dimension)
-		{
-		case ETextureDimension::Texture2D:
-		{
-			const bool bCubeTexture = false;
-			Texture = FD3D12DynamicRHI::GetD3DRHI()->CreateD3D12Texture2D<FD3D12BaseTexture2D>(nullptr, InCreateInfo.Extent.X, InCreateInfo.Extent.Y, 1, bTextureArray, bCubeTexture, InCreateInfo.Format, InCreateInfo.NumMips, InCreateInfo.NumSamples, InCreateInfo.Flags, InitialState, ResourceCreateInfo, &ResourceAllocatorAdapter);
-		}
-		break;
-		case ETextureDimension::Texture3D:
-		{
-			Texture = FD3D12DynamicRHI::GetD3DRHI()->CreateD3D12Texture3D(nullptr, InCreateInfo.Extent.X, InCreateInfo.Extent.Y, InCreateInfo.Depth, InCreateInfo.Format, InCreateInfo.NumMips, InCreateInfo.Flags, InitialState, ResourceCreateInfo, &ResourceAllocatorAdapter);
-		}
-		break;
-		case ETextureDimension::TextureCube:
-		{
-			const bool bCubeTexture = true;
-			check(InCreateInfo.Extent.X == InCreateInfo.Extent.Y);
-			Texture = FD3D12DynamicRHI::GetD3DRHI()->CreateD3D12Texture2D<FD3D12BaseTextureCube>(nullptr, InCreateInfo.Extent.X, InCreateInfo.Extent.Y, 6, bTextureArray, bCubeTexture, InCreateInfo.Format, InCreateInfo.NumMips, InCreateInfo.NumSamples, InCreateInfo.Flags, InitialState, ResourceCreateInfo, &ResourceAllocatorAdapter);
-		}
-		break;
-		}
-
-		// The D3D12_RESOURCE_DESC's are built in two different functions right now. This checks that they actually match what we expect.
-#if DO_CHECK
-		{
-			CD3DX12_RESOURCE_DESC CreatedDesc(GetD3D12TextureFromRHITexture(Texture)->GetResource()->GetDesc());
-			CD3DX12_RESOURCE_DESC DerivedDesc(Desc);
-			check(CreatedDesc == DerivedDesc);
-		}
-#endif
-
+		FResourceAllocatorAdapter ResourceAllocatorAdapter(GetParentAdapter(), static_cast<FD3D12TransientHeap&>(Initializer.Heap), Initializer.Allocation, Desc);
+		FRHITexture* Texture = DynamicRHI->CreateTexture(InCreateInfo, InDebugName, InitialState, &ResourceAllocatorAdapter);
 		return new FRHITransientTexture(Texture, Initializer.Hash, InCreateInfo);
 	});
 }
@@ -190,6 +157,15 @@ void FD3D12TransientResourceAllocator::FResourceAllocatorAdapter::AllocateResour
 	uint32 GPUIndex, D3D12_HEAP_TYPE, const D3D12_RESOURCE_DESC& InDesc, uint64 InSize, uint32, ED3D12ResourceStateMode InResourceStateMode,
 	D3D12_RESOURCE_STATES InCreateState, const D3D12_CLEAR_VALUE* InClearValue, const TCHAR* InName, FD3D12ResourceLocation& ResourceLocation)
 {
+	// The D3D12_RESOURCE_DESC's are built in two different functions right now. This checks that they actually match what we expect.
+#if DO_CHECK
+	{
+		CD3DX12_RESOURCE_DESC CreatedDesc(InDesc);
+		CD3DX12_RESOURCE_DESC DerivedDesc(Desc);
+		check(CreatedDesc == DerivedDesc);
+	}
+#endif
+
 	FD3D12Resource* NewResource = nullptr;
 	VERIFYD3D12RESULT(GetParentAdapter()->CreatePlacedResource(InDesc, Heap.GetLinkedObject(GPUIndex)->Get(), Allocation.Offset, InCreateState, InResourceStateMode, D3D12_RESOURCE_STATE_TBD, InClearValue, &NewResource, InName));
 
@@ -205,23 +181,14 @@ FRHITransientBuffer* FD3D12TransientResourceAllocator::CreateBuffer(const FRHIBu
 	EBufferUsageFlags BufferUsage = InCreateInfo.Usage;
 	FD3D12Buffer::GetResourceDescAndAlignment(InCreateInfo.Size, InCreateInfo.Stride, BufferUsage, Desc, Alignment);
 
-	return Allocator.CreateBuffer(InCreateInfo, InDebugName, Desc.Width, Alignment,
+	Alignment = D3D12_DEFAULT_RESOURCE_PLACEMENT_ALIGNMENT;
+	uint64 Size = Align(Desc.Width, Alignment);
+
+	return Allocator.CreateBuffer(InCreateInfo, InDebugName, Size, Alignment,
 		[&](const FRHITransientResourceAllocator::FResourceInitializer& Initializer)
 	{
-		FResourceAllocatorAdapter ResourceAllocatorAdapter(GetParentAdapter(), static_cast<FD3D12TransientHeap&>(Initializer.Heap), Initializer.Allocation);
-		FRHIResourceCreateInfo ResourceCreateInfo(InDebugName);
-
-		FD3D12Buffer* Buffer = FD3D12DynamicRHI::GetD3DRHI()->CreateD3D12Buffer(nullptr, InCreateInfo.Size, BufferUsage, InCreateInfo.Stride, ERHIAccess::UAVMask, ResourceCreateInfo, &ResourceAllocatorAdapter);
-
-		// The D3D12_RESOURCE_DESC's are built in two different functions right now. This checks that they actually match what we expect.
-#if DO_CHECK
-		{
-			CD3DX12_RESOURCE_DESC CreatedDesc(Buffer->GetResource()->GetDesc());
-			CD3DX12_RESOURCE_DESC DerivedDesc(Desc);
-			check(CreatedDesc == DerivedDesc);
-		}
-#endif
-
+		FResourceAllocatorAdapter ResourceAllocatorAdapter(GetParentAdapter(), static_cast<FD3D12TransientHeap&>(Initializer.Heap), Initializer.Allocation, Desc);
+		FRHIBuffer* Buffer = FD3D12DynamicRHI::GetD3DRHI()->CreateBuffer(InCreateInfo, InDebugName, ERHIAccess::UAVMask, &ResourceAllocatorAdapter);
 		return new FRHITransientBuffer(Buffer, Initializer.Hash, InCreateInfo);
 	});
 }
diff --git a/Engine/Source/Runtime/D3D12RHI/Private/D3D12TransientResourceAllocator.h b/Engine/Source/Runtime/D3D12RHI/Private/D3D12TransientResourceAllocator.h
--- a/Engine/Source/Runtime/D3D12RHI/Private/D3D12TransientResourceAllocator.h
+++ b/Engine/Source/Runtime/D3D12RHI/Private/D3D12TransientResourceAllocator.h
@@ -58,10 +58,11 @@ private:
 	class FResourceAllocatorAdapter final : public FD3D12AdapterChild, public ID3D12ResourceAllocator
 	{
 	public:
-		FResourceAllocatorAdapter(FD3D12Adapter* Adapter, FD3D12TransientHeap& InHeap, const FRHITransientHeapAllocation& InAllocation)
+		FResourceAllocatorAdapter(FD3D12Adapter* Adapter, FD3D12TransientHeap& InHeap, const FRHITransientHeapAllocation& InAllocation, const D3D12_RESOURCE_DESC& InDesc)
 			: FD3D12AdapterChild(Adapter)
 			, Heap(InHeap)
 			, Allocation(InAllocation)
+			, Desc(InDesc)
 		{}
 
 		void AllocateResource(
@@ -70,6 +71,7 @@ private:
 
 		FD3D12TransientHeap& Heap;
 		const FRHITransientHeapAllocation& Allocation;
+		const D3D12_RESOURCE_DESC Desc;
 	};
 
 	FRHITransientResourceAllocator Allocator;
diff --git a/Engine/Source/Runtime/RenderCore/Private/RenderGraphResourcePool.cpp b/Engine/Source/Runtime/RenderCore/Private/RenderGraphResourcePool.cpp
--- a/Engine/Source/Runtime/RenderCore/Private/RenderGraphResourcePool.cpp
+++ b/Engine/Source/Runtime/RenderCore/Private/RenderGraphResourcePool.cpp
@@ -52,6 +52,8 @@ TRefCountPtr<FRDGPooledBuffer> FRenderGraphResourcePool::FindFreeBufferInternal(
 
 	// Allocate new one
 	{
+		TRACE_CPUPROFILER_EVENT_SCOPE(FRenderGraphResourcePool::CreateBuffer);
+
 		uint32 NumBytes = Desc.GetTotalNumBytes();
 
 		FRHIResourceCreateInfo CreateInfo(InDebugName);
diff --git a/Engine/Source/Runtime/RenderCore/Private/RenderTargetPool.cpp b/Engine/Source/Runtime/RenderCore/Private/RenderTargetPool.cpp
--- a/Engine/Source/Runtime/RenderCore/Private/RenderTargetPool.cpp
+++ b/Engine/Source/Runtime/RenderCore/Private/RenderTargetPool.cpp
@@ -340,6 +340,7 @@ Done:
 
 	if (!Found)
 	{
+		TRACE_CPUPROFILER_EVENT_SCOPE(FRenderTargetPool::CreateTexture);
 		UE_LOG(LogRenderTargetPool, Display, TEXT("%d MB, NewRT %s %s"), (AllocationLevelInKB + 1023) / 1024, *Desc.GenerateInfoString(), InDebugName);
 
 		// not found in the pool, create a new element
