Fix transient allocator for texture arrays

The array size was not getting passed down, which caused subsequent asserts to trigger. Discovered by running the path tracer IES test.

diff --git a/Engine/Source/Runtime/D3D12RHI/Private/D3D12Texture.cpp b/Engine/Source/Runtime/D3D12RHI/Private/D3D12Texture.cpp
--- a/Engine/Source/Runtime/D3D12RHI/Private/D3D12Texture.cpp
+++ b/Engine/Source/Runtime/D3D12RHI/Private/D3D12Texture.cpp
@@ -1308,11 +1308,11 @@ FRHITexture* FD3D12DynamicRHI::CreateTexture(const FRHITextureCreateInfo& Create
 		return CreateD3D12Texture2D<FD3D12BaseTexture2D>(nullptr, CreateInfo.Extent.X, CreateInfo.Extent.Y, 1, bTextureArray, bTextureCube, CreateInfo.Format, CreateInfo.NumMips, CreateInfo.NumSamples, CreateInfo.Flags, InitialState, ResourceCreateInfo, ResourceAllocator);
 
 	case ETextureDimension::Texture2DArray:
-		return CreateD3D12Texture2D<FD3D12BaseTexture2DArray>(nullptr, CreateInfo.Extent.X, CreateInfo.Extent.Y, 1, bTextureArray, bTextureCube, CreateInfo.Format, CreateInfo.NumMips, CreateInfo.NumSamples, CreateInfo.Flags, InitialState, ResourceCreateInfo, ResourceAllocator);
+		return CreateD3D12Texture2D<FD3D12BaseTexture2DArray>(nullptr, CreateInfo.Extent.X, CreateInfo.Extent.Y, CreateInfo.ArraySize, bTextureArray, bTextureCube, CreateInfo.Format, CreateInfo.NumMips, CreateInfo.NumSamples, CreateInfo.Flags, InitialState, ResourceCreateInfo, ResourceAllocator);
 
 	case ETextureDimension::TextureCube:
 	case ETextureDimension::TextureCubeArray:
-		return CreateD3D12Texture2D<FD3D12BaseTextureCube>(nullptr, CreateInfo.Extent.X, CreateInfo.Extent.Y, 6, bTextureArray, bTextureCube, CreateInfo.Format, CreateInfo.NumMips, CreateInfo.NumSamples, CreateInfo.Flags, InitialState, ResourceCreateInfo, ResourceAllocator);
+		return CreateD3D12Texture2D<FD3D12BaseTextureCube>(nullptr, CreateInfo.Extent.X, CreateInfo.Extent.Y, 6 * CreateInfo.ArraySize, bTextureArray, bTextureCube, CreateInfo.Format, CreateInfo.NumMips, CreateInfo.NumSamples, CreateInfo.Flags, InitialState, ResourceCreateInfo, ResourceAllocator);
 
 	case ETextureDimension::Texture3D:
 		return CreateD3D12Texture3D(nullptr, CreateInfo.Extent.X, CreateInfo.Extent.Y, CreateInfo.Depth, CreateInfo.Format, CreateInfo.NumMips, CreateInfo.Flags, InitialState, ResourceCreateInfo, ResourceAllocator);
