Disabled ray tracing geometry sections now set the number of indices to 0 instead of using a dummy zero transform matrix.

Also, BLAS buffer sizes are now computed based on 16 and 32 bit index buffers, as we don't know the final index buffer format in the geometry constructor.

diff --git a/Engine/Source/Runtime/D3D12RHI/Private/D3D12RayTracing.cpp b/Engine/Source/Runtime/D3D12RHI/Private/D3D12RayTracing.cpp
--- a/Engine/Source/Runtime/D3D12RHI/Private/D3D12RayTracing.cpp
+++ b/Engine/Source/Runtime/D3D12RHI/Private/D3D12RayTracing.cpp
@@ -2775,16 +2775,33 @@ FRayTracingAccelerationStructureSize FD3D12DynamicRHI::RHICalcRayTracingGeometry
 
 	FD3D12Adapter& Adapter = GetAdapter();
 
-	// Get maximum buffer sizes for all GPUs in the system
-	for (uint32 GPUIndex = 0; GPUIndex < GNumExplicitGPUsForRendering; ++GPUIndex)
+	// We don't know the final index buffer format, so take maximum of 16 and 32 bit.
+	DXGI_FORMAT IndexFormats[] =
 	{
-		ID3D12Device5* RayTracingDevice = Adapter.GetDevice(GPUIndex)->GetDevice5();
-		D3D12_RAYTRACING_ACCELERATION_STRUCTURE_PREBUILD_INFO PrebuildInfo = {};
-		RayTracingDevice->GetRaytracingAccelerationStructurePrebuildInfo(&PrebuildDescInputs, &PrebuildInfo);
+		DXGI_FORMAT_R16_UINT, DXGI_FORMAT_R32_UINT
+	};
 
-		SizeInfo.ResultSize = FMath::Max(SizeInfo.ResultSize, PrebuildInfo.ResultDataMaxSizeInBytes);
-		SizeInfo.BuildScratchSize = FMath::Max(SizeInfo.BuildScratchSize, PrebuildInfo.ScratchDataSizeInBytes);
-		SizeInfo.UpdateScratchSize = FMath::Max(SizeInfo.UpdateScratchSize, PrebuildInfo.UpdateScratchDataSizeInBytes);
+	for (DXGI_FORMAT IndexFormat : IndexFormats)
+	{
+		for (D3D12_RAYTRACING_GEOMETRY_DESC& GeometryDesc : GeometryDescs)
+		{
+			if (GeometryDesc.Type == D3D12_RAYTRACING_GEOMETRY_TYPE_TRIANGLES)
+			{
+				GeometryDesc.Triangles.IndexFormat = IndexFormat;
+			}
+		}
+
+		// Get maximum buffer sizes for all GPUs in the system
+		for (uint32 GPUIndex = 0; GPUIndex < GNumExplicitGPUsForRendering; ++GPUIndex)
+		{
+			ID3D12Device5* RayTracingDevice = Adapter.GetDevice(GPUIndex)->GetDevice5();
+			D3D12_RAYTRACING_ACCELERATION_STRUCTURE_PREBUILD_INFO PrebuildInfo = {};
+			RayTracingDevice->GetRaytracingAccelerationStructurePrebuildInfo(&PrebuildDescInputs, &PrebuildInfo);
+
+			SizeInfo.ResultSize = FMath::Max(SizeInfo.ResultSize, PrebuildInfo.ResultDataMaxSizeInBytes);
+			SizeInfo.BuildScratchSize = FMath::Max(SizeInfo.BuildScratchSize, PrebuildInfo.ScratchDataSizeInBytes);
+			SizeInfo.UpdateScratchSize = FMath::Max(SizeInfo.UpdateScratchSize, PrebuildInfo.UpdateScratchDataSizeInBytes);
+		}
 	}
 
 	return SizeInfo;
@@ -3140,16 +3157,14 @@ void FD3D12RayTracingGeometry::BuildAccelerationStructure(FD3D12CommandContext&
 				checkNoEntry();
 				break;
 			}
-			// #yuriy_todo: refactor to just have 0 primitives in this section
-			if (Segment.bEnabled)
-			{
-				Desc.Triangles.Transform3x4 = D3D12_GPU_VIRTUAL_ADDRESS(0);
-			}
-			else
+
+			if (!Segment.bEnabled)
 			{
-				Desc.Triangles.Transform3x4 = NullTransformBufferD3D12->ResourceLocation.GetGPUVirtualAddress();
+				Desc.Triangles.IndexCount = 0;
 			}
 
+			checkf(Desc.Triangles.Transform3x4 == D3D12_GPU_VIRTUAL_ADDRESS(0), TEXT("BLAS geometry transforms are not supported!"));
+
 			if (IndexBuffer)
 			{
 				static constexpr uint32 IndicesPerPrimitive = 3; // Only triangle meshes are supported
