Fixed RHICalcRayTracingGeometrySize implementation when index buffer is null.

diff --git a/Engine/Source/Runtime/D3D12RHI/Private/D3D12RayTracing.cpp b/Engine/Source/Runtime/D3D12RHI/Private/D3D12RayTracing.cpp
--- a/Engine/Source/Runtime/D3D12RHI/Private/D3D12RayTracing.cpp
+++ b/Engine/Source/Runtime/D3D12RHI/Private/D3D12RayTracing.cpp
@@ -2776,10 +2776,13 @@ FRayTracingAccelerationStructureSize FD3D12DynamicRHI::RHICalcRayTracingGeometry
 	FD3D12Adapter& Adapter = GetAdapter();
 
 	// We don't know the final index buffer format, so take maximum of 16 and 32 bit.
-	DXGI_FORMAT IndexFormats[] =
-	{
-		DXGI_FORMAT_R16_UINT, DXGI_FORMAT_R32_UINT
-	};
+
+	static const DXGI_FORMAT ValidIndexBufferFormats[] = { DXGI_FORMAT_R16_UINT, DXGI_FORMAT_R32_UINT };
+	static const DXGI_FORMAT NullIndexBufferFormats[] = { DXGI_FORMAT_UNKNOWN };
+
+	TArrayView<const DXGI_FORMAT> IndexFormats = Initializer.IndexBuffer.IsValid() 
+		? MakeArrayView(ValidIndexBufferFormats)
+		: MakeArrayView(NullIndexBufferFormats);
 
 	for (DXGI_FORMAT IndexFormat : IndexFormats)
 	{
