Removed r.D3D12.RayTracing.ExecuteIndirectWorkaround CVar and replaced it with an NVIDIA driver version check (require 466.11 or newer).

diff --git a/Engine/Source/Runtime/D3D12RHI/Private/D3D12RayTracing.cpp b/Engine/Source/Runtime/D3D12RHI/Private/D3D12RayTracing.cpp
--- a/Engine/Source/Runtime/D3D12RHI/Private/D3D12RayTracing.cpp
+++ b/Engine/Source/Runtime/D3D12RHI/Private/D3D12RayTracing.cpp
@@ -63,13 +63,6 @@ static FAutoConsoleVariableRef CVarD3D12RayTracingMaxBatchedCompaction(
 	ECVF_ReadOnly
 );
 
-static int32 GD3D12RayTracingExecuteIndirectWorkaround = 1;
-static FAutoConsoleVariableRef CVarD3D12RayTracingExecuteIndirectWorkaround(
-	TEXT("r.D3D12.RayTracing.ExecuteIndirectWorkaround"),
-	GD3D12RayTracingExecuteIndirectWorkaround,
-	TEXT("Adds a command list flush before ExecuteIndirect to avoid NVIDIA driver bug that causes some indirect RT commands to get lost. (default = 1)\n")
-);
-
 // Ray tracing stat counters
 
 DECLARE_STATS_GROUP(TEXT("D3D12RHI: Ray Tracing"), STATGROUP_D3D12RayTracing, STATCAT_Advanced);
@@ -4519,11 +4512,6 @@ static void DispatchRays(FD3D12CommandContext& CommandContext,
 		}
 
 		CommandListHandle.AddTransitionBarrier(DispatchRaysDescBufferResource, D3D12_RESOURCE_STATE_UNORDERED_ACCESS, D3D12_RESOURCE_STATE_INDIRECT_ARGUMENT, 0);
-
-		if (GD3D12RayTracingExecuteIndirectWorkaround)
-		{
-			CommandContext.RHISubmitCommandsHint();
-		}
 	}
 
 	// Setup state for RT dispatch
diff --git a/Engine/Source/Runtime/D3D12RHI/Private/Windows/WindowsD3D12Device.cpp b/Engine/Source/Runtime/D3D12RHI/Private/Windows/WindowsD3D12Device.cpp
--- a/Engine/Source/Runtime/D3D12RHI/Private/Windows/WindowsD3D12Device.cpp
+++ b/Engine/Source/Runtime/D3D12RHI/Private/Windows/WindowsD3D12Device.cpp
@@ -859,12 +859,11 @@ void FD3D12DynamicRHI::Init()
 #if NV_API_ENABLE
 	if (IsRHIDeviceNVIDIA() && bAllowVendorDevice)
 	{
-		NvAPI_Status NvStatus;
-		NvStatus = NvAPI_Initialize();
+		const NvAPI_Status NvStatus = NvAPI_Initialize();
 		if (NvStatus == NVAPI_OK)
 		{
-			NvStatus = NvAPI_D3D12_IsNvShaderExtnOpCodeSupported(GetAdapter().GetD3DDevice(), NV_EXTN_OP_UINT64_ATOMIC, &GRHISupportsAtomicUInt64);
-			if (NvStatus != NVAPI_OK)
+			NvAPI_Status NvStatusAtomicU64 = NvAPI_D3D12_IsNvShaderExtnOpCodeSupported(GetAdapter().GetD3DDevice(), NV_EXTN_OP_UINT64_ATOMIC, &GRHISupportsAtomicUInt64);
+			if (NvStatusAtomicU64 != NVAPI_OK)
 			{
 				UE_LOG(LogD3D12RHI, Warning, TEXT("Failed to query support for 64 bit atomics"));
 			}
@@ -874,24 +873,36 @@ void FD3D12DynamicRHI::Init()
 			UE_LOG(LogD3D12RHI, Warning, TEXT("Failed to initialize NVAPI"));
 		}
 
-		// Disable ray tracing for old Nvidia drivers
-		if (GRHISupportsRayTracing
-			&& GMinimumDriverVersionForRayTracingNVIDIA > 0
-			&& NvStatus == NVAPI_OK)
+		NvU32 DriverVersion = UINT32_MAX;
+
+		if (NvStatus == NVAPI_OK)
 		{
-			NvU32 DriverVersion = UINT32_MAX;
 			NvAPI_ShortString BranchString("");
-			if (NvAPI_SYS_GetDriverAndBranchVersion(&DriverVersion, BranchString) == NVAPI_OK)
+			if (NvAPI_SYS_GetDriverAndBranchVersion(&DriverVersion, BranchString) != NVAPI_OK)
 			{
-				if (DriverVersion < (uint32)GMinimumDriverVersionForRayTracingNVIDIA)
-				{
-					GRHISupportsRayTracing = false;
-
-					UE_LOG(LogD3D12RHI, Warning, TEXT("Ray tracing is disabled because the driver is too old"));
-				}
+				UE_LOG(LogD3D12RHI, Warning, TEXT("Failed to query NVIDIA driver version"));
 			}
 		}
 
+		// Disable ray tracing for old Nvidia drivers
+		if (GRHISupportsRayTracing 
+			&& GMinimumDriverVersionForRayTracingNVIDIA > 0
+			&& DriverVersion < (uint32)GMinimumDriverVersionForRayTracingNVIDIA)
+		{
+			GRHISupportsRayTracing = false;
+			UE_LOG(LogD3D12RHI, Warning, TEXT("Ray tracing is disabled because the driver is too old"));
+		}
+
+		// Disable indirect ray tracing dispatch on drivers that have a known bug.
+		if (GRHISupportsRayTracingDispatchIndirect 
+			&& DriverVersion < (uint32)46611)
+		{
+			GRHISupportsRayTracingDispatchIndirect = false;
+			UE_LOG(LogD3D12RHI, Warning,
+				TEXT("Indirect ray tracing dispatch is disabled because of known bugs in the current driver. ")
+				TEXT("Please update to NVIDIA driver version 466.11 or newer."));
+		}
+
 		// Disable ray tracing emulation for Pascal architecture
 		if (GRHISupportsRayTracing
 			&& !GAllowEmulatedRayTracing
