Fixes for swap support on ResourceLocation when pool allocated

diff --git a/Engine/Source/Runtime/D3D12RHI/Private/D3D12Resources.cpp b/Engine/Source/Runtime/D3D12RHI/Private/D3D12Resources.cpp
--- a/Engine/Source/Runtime/D3D12RHI/Private/D3D12Resources.cpp
+++ b/Engine/Source/Runtime/D3D12RHI/Private/D3D12Resources.cpp
@@ -679,11 +679,27 @@ void FD3D12ResourceLocation::Swap(FD3D12ResourceLocation& Other)
 	}
 #endif
 
-	// Should not be linked list allocated - otherwise internal linked list data needs to be updated as well in a threadsafe way
-	check(Other.GetAllocatorType() != FD3D12ResourceLocation::AT_Pool);
-	check(GetAllocatorType() != FD3D12ResourceLocation::AT_Pool);
-
-	::Swap(*this, Other);
+	if (Other.GetAllocatorType() == FD3D12ResourceLocation::AT_Pool)
+	{
+		check(GetAllocatorType() != FD3D12ResourceLocation::AT_Pool);
+		
+		// Cache the allocator data and reset before swap
+		FD3D12PoolAllocatorPrivateData& PoolData = Other.GetPoolAllocatorPrivateData();
+		FD3D12PoolAllocatorPrivateData TmpPoolData = PoolData;
+		PoolData.Init();
+
+		// Perform swap
+		::Swap(*this, Other);
+
+		// Restore allocator data and perform pool aware swap
+		PoolData = TmpPoolData;	
+		Other.SetPoolAllocator(GetPoolAllocator());
+		GetPoolAllocator()->TransferOwnership(Other, *this);
+	}
+	else
+	{
+		::Swap(*this, Other);
+	}
 }
 
 void FD3D12ResourceLocation::Alias(FD3D12ResourceLocation & Destination, FD3D12ResourceLocation & Source)
diff --git a/Engine/Source/Runtime/D3D12RHI/Public/D3D12Resources.h b/Engine/Source/Runtime/D3D12RHI/Public/D3D12Resources.h
--- a/Engine/Source/Runtime/D3D12RHI/Public/D3D12Resources.h
+++ b/Engine/Source/Runtime/D3D12RHI/Public/D3D12Resources.h
@@ -829,6 +829,7 @@ public:
 
 		::Swap(Parent, Other.Parent);
 		ResourceLocation.Swap(Other.ResourceLocation);
+		ResourceLocation.SetOwner(this);
 		::Swap(BufferAlignment, Other.BufferAlignment);
 		::Swap(RenameListeners, Other.RenameListeners);
 	}
