- Add tracking of all resources and also keep track of releases resources during last 100 frames
- When tracking is enabled all resources have a valid GPU address
- On GPU crash with page fault use tracked allocation data to find resources in area of page fault and also recently freed resources

diff --git a/Engine/Source/Runtime/D3D12RHI/Private/D3D12Adapter.cpp b/Engine/Source/Runtime/D3D12RHI/Private/D3D12Adapter.cpp
--- a/Engine/Source/Runtime/D3D12RHI/Private/D3D12Adapter.cpp
+++ b/Engine/Source/Runtime/D3D12RHI/Private/D3D12Adapter.cpp
@@ -81,6 +81,14 @@ static FAutoConsoleVariableRef CVarD3D12EnableDRED(
 	TEXT("Has GPU overhead but gives the most information on the current GPU state when it crashes or hangs.\n"),
 	ECVF_RenderThreadSafe | ECVF_ReadOnly);
 
+bool GD3D12TrackAllAlocations = false;
+static TAutoConsoleVariable<int32> CVarD3D12TrackAllAllocations(
+	TEXT("D3D12.TrackAllAllocations"),
+	GD3D12TrackAllAlocations,
+	TEXT("Controls whether D3D12 RHI should track all allocation information (default = off)."),
+	ECVF_ReadOnly
+);
+
 static bool CheckD3DStoredMessages()
 {
 	bool bResult = false;
@@ -184,7 +192,7 @@ FD3D12Adapter::FD3D12Adapter(FD3D12AdapterDesc& DescIn)
 	, bHeapNotZeroedSupported(false)
 	, VRSTileSize(0)
 	, bDebugDevice(false)
-	, GPUCrashDebuggingModes(ED3D12GPUCrashDebuggingModes::None)
+	, GPUCrashDebuggingModes(ED3D12GPUCrashDebuggingModes::None)	
 	, bDeviceRemoved(false)
 	, Desc(DescIn)
 	, RootSignatureManager(this)
@@ -928,6 +936,12 @@ void FD3D12Adapter::InitializeDevices()
 		StagingFence = new FD3D12Fence(this, FRHIGPUMask::All(), L"Staging Fence");
 		StagingFence->CreateFence();
 
+#if TRACK_RESOURCE_ALLOCATIONS
+		// Set flag if we want to track all allocations - comes with some overhead and only possible when Tier 2 is available
+		// (because we will create placed buffers for texture allocation to retrieve the GPU virtual addresses)
+		bTrackAllAllocation = (GD3D12TrackAllAlocations || GPUCrashDebuggingModes == ED3D12GPUCrashDebuggingModes::All) && (ResourceHeapTier == D3D12_RESOURCE_HEAP_TIER_2);
+#endif 
+
 		CreateCommandSignatures();
 
 		// Context redirectors allow RHI commands to be executed on multiple GPUs at the
@@ -1237,6 +1251,24 @@ void FD3D12Adapter::EndFrame()
 #if D3D12_SUBMISSION_GAP_RECORDER
 	SubmitGapRecorderTimestamps();
 #endif
+
+#if TRACK_RESOURCE_ALLOCATIONS
+	// remove tracked released resources older than n amount of frames
+	static const uint64 FrameRetention = 100;
+	int32 ReleaseCount = 0;
+	uint64 CurrentFrameID = GetFrameFence().GetCurrentFence();
+	for (; ReleaseCount < ReleasedAllocationData.Num(); ++ReleaseCount)
+	{
+		if (ReleasedAllocationData[ReleaseCount].ReleasedFrameID + FrameRetention > CurrentFrameID)
+		{
+			break;
+		}
+	}
+	if (ReleaseCount > 0)
+	{
+		ReleasedAllocationData.RemoveAt(0, ReleaseCount, false);
+	}
+#endif
 }
 
 #if WITH_MGPU
@@ -1334,13 +1366,20 @@ void FD3D12Adapter::BlockUntilIdle()
 }
 
 
-void FD3D12Adapter::TrackAllocationData(FD3D12ResourceLocation* InAllocation, uint64 InAllocationSize)
+void FD3D12Adapter::TrackAllocationData(FD3D12ResourceLocation* InAllocation, uint64 InAllocationSize, bool bCollectCallstack)
 {
 #if TRACK_RESOURCE_ALLOCATIONS
 	FTrackedAllocationData AllocationData;
 	AllocationData.ResourceAllocation = InAllocation;
 	AllocationData.AllocationSize = InAllocationSize;
-	AllocationData.StackDepth = FPlatformStackWalk::CaptureStackBackTrace(&AllocationData.Stack[0], FTrackedAllocationData::MaxStackDepth);
+	if (bCollectCallstack)
+	{
+		AllocationData.StackDepth = FPlatformStackWalk::CaptureStackBackTrace(&AllocationData.Stack[0], FTrackedAllocationData::MaxStackDepth);
+	}
+	else 
+	{
+		AllocationData.StackDepth = 0;
+	}
 
 	FScopeLock Lock(&TrackedAllocationDataCS);
 	check(!TrackedAllocationData.Contains(InAllocation));
@@ -1348,14 +1387,80 @@ void FD3D12Adapter::TrackAllocationData(FD3D12ResourceLocation* InAllocation, ui
 #endif
 }
 
-void FD3D12Adapter::ReleaseTrackedAllocationData(FD3D12ResourceLocation* InAllocation)
+void FD3D12Adapter::ReleaseTrackedAllocationData(FD3D12ResourceLocation* InAllocation, bool bDefragFree)
 {
 #if TRACK_RESOURCE_ALLOCATIONS
 	FScopeLock Lock(&TrackedAllocationDataCS);
+
+	D3D12_GPU_VIRTUAL_ADDRESS GPUAddress = InAllocation->GetGPUVirtualAddress();
+	if (GPUAddress != 0 || IsTrackingAllAllocations())
+	{
+		FReleasedAllocationData ReleasedData;
+		ReleasedData.GPUVirtualAddress = GPUAddress;
+		ReleasedData.AllocationSize = InAllocation->GetSize();
+		ReleasedData.ResourceName = InAllocation->GetResource()->GetName();
+		ReleasedData.ReleasedFrameID = GetFrameFence().GetCurrentFence();
+		ReleasedData.bDefragFree = bDefragFree;
+		ReleasedData.bBackBuffer = InAllocation->GetResource()->IsBackBuffer();
+		ReleasedData.bTransient = InAllocation->IsTransient();
+		// Only the backbuffer doesn't have a valid gpu virtual address
+		check(ReleasedData.GPUVirtualAddress != 0 || ReleasedData.bBackBuffer);
+		ReleasedAllocationData.Add(ReleasedData);
+	}
+
 	verify(TrackedAllocationData.Remove(InAllocation) == 1);
 #endif
 }
 
+void FD3D12Adapter::FindResourcesNearGPUAddress(D3D12_GPU_VIRTUAL_ADDRESS InGPUVirtualAddress, uint64 InRange, TArray<FAllocatedResourceResult>& OutResources)
+{
+#if TRACK_RESOURCE_ALLOCATIONS
+	FScopeLock Lock(&TrackedAllocationDataCS);
+
+	TArray<FTrackedAllocationData> Allocations;
+	TrackedAllocationData.GenerateValueArray(Allocations);
+	FInt64Range TrackRange((int64)(InGPUVirtualAddress - InRange), (int64)(InGPUVirtualAddress + InRange));
+	for (FTrackedAllocationData& AllocationData : Allocations)
+	{
+		D3D12_GPU_VIRTUAL_ADDRESS GPUAddress = AllocationData.ResourceAllocation->GetResource()->GetGPUVirtualAddress();
+		FInt64Range AllocationRange((int64)GPUAddress, (int64)(GPUAddress + AllocationData.AllocationSize));
+		if (TrackRange.Overlaps(AllocationRange))
+		{
+			bool bContainsAllocation = AllocationRange.Contains(InGPUVirtualAddress);
+			int64 Distance = bContainsAllocation ? 0 : ((InGPUVirtualAddress < GPUAddress) ? GPUAddress - InGPUVirtualAddress : InGPUVirtualAddress - AllocationRange.GetUpperBoundValue());
+			check(Distance > 0);
+
+			FAllocatedResourceResult Result;
+			Result.Allocation = AllocationData.ResourceAllocation;
+			Result.Distance = Distance;
+			OutResources.Add(Result);
+		}
+	}
+
+	// Sort the resources on distance from the requested address
+	Algo::Sort(OutResources, [InGPUVirtualAddress](const FAllocatedResourceResult& InLHS, const FAllocatedResourceResult& InRHS)
+		{
+			return InLHS.Distance < InRHS.Distance;
+		});
+#endif
+}
+
+void FD3D12Adapter::FindReleasedAllocationData(D3D12_GPU_VIRTUAL_ADDRESS InGPUVirtualAddress, TArray<FReleasedAllocationData>& OutAllocationData)
+{
+#if TRACK_RESOURCE_ALLOCATIONS
+	FScopeLock Lock(&TrackedAllocationDataCS);
+
+	for (FReleasedAllocationData& AllocationData : ReleasedAllocationData)
+	{
+		if (InGPUVirtualAddress >= AllocationData.GPUVirtualAddress && InGPUVirtualAddress < (AllocationData.GPUVirtualAddress + AllocationData.AllocationSize))
+		{
+			// Add in reverse order, so last released resources at first in the array
+			OutAllocationData.EmplaceAt(0, AllocationData);
+		}
+	}
+#endif
+}
+
 #if TRACK_RESOURCE_ALLOCATIONS
 
 static FAutoConsoleCommandWithOutputDevice GDumpTrackedD3D12AllocationsCmd(
diff --git a/Engine/Source/Runtime/D3D12RHI/Private/D3D12Adapter.h b/Engine/Source/Runtime/D3D12RHI/Private/D3D12Adapter.h
--- a/Engine/Source/Runtime/D3D12RHI/Private/D3D12Adapter.h
+++ b/Engine/Source/Runtime/D3D12RHI/Private/D3D12Adapter.h
@@ -351,10 +351,30 @@ public:
 
 	FORCEINLINE uint32 GetFrameCount() const { return FrameCounter; }
 
-	void TrackAllocationData(FD3D12ResourceLocation* InAllocation, uint64 InAllocationSize);
-	void ReleaseTrackedAllocationData(FD3D12ResourceLocation* InAllocation);
+	bool IsTrackingAllAllocations() const { return bTrackAllAllocation; }
+	void TrackAllocationData(FD3D12ResourceLocation* InAllocation, uint64 InAllocationSize, bool bCollectCallstack);
+	void ReleaseTrackedAllocationData(FD3D12ResourceLocation* InAllocation, bool bDefragFree);
 	void DumpTrackedAllocationData(FOutputDevice& OutputDevice, bool bResidentOnly, bool bWithCallstack);
 
+	struct FAllocatedResourceResult
+	{
+		FD3D12ResourceLocation* Allocation = nullptr;
+		uint64 Distance = 0;
+	};
+	void FindResourcesNearGPUAddress(D3D12_GPU_VIRTUAL_ADDRESS InGPUVirtualAddress, uint64 InRange, TArray<FAllocatedResourceResult>& OutResources);
+	
+	struct FReleasedAllocationData
+	{
+		D3D12_GPU_VIRTUAL_ADDRESS GPUVirtualAddress = 0;
+		uint64 AllocationSize = 0;
+		FName ResourceName;
+		uint64 ReleasedFrameID = 0;
+		bool bDefragFree = false;
+		bool bBackBuffer = false;
+		bool bTransient = false;
+	};
+	void FindReleasedAllocationData(D3D12_GPU_VIRTUAL_ADDRESS InGPUVirtualAddress, TArray<FReleasedAllocationData>& OutAllocationData);
+
 #if D3D12_SUBMISSION_GAP_RECORDER
 	FD3D12SubmissionGapRecorder SubmissionGapRecorder;
 
@@ -453,6 +473,8 @@ protected:
 
 	uint32 FrameCounter;
 
+	bool bTrackAllAllocation = false;
+
 	/** Information about an allocated resource. */
 	struct FTrackedAllocationData
 	{
@@ -463,8 +485,10 @@ protected:
 		uint32 StackDepth;
 		uint64 Stack[MaxStackDepth];
 	};
+
 	/** Tracked resource information. */
 	TMap<FD3D12ResourceLocation*, FTrackedAllocationData> TrackedAllocationData;
+	TArray<FReleasedAllocationData> ReleasedAllocationData;
 	FCriticalSection TrackedAllocationDataCS;
 
 	FD3D12MemoryInfo MemoryInfo;
diff --git a/Engine/Source/Runtime/D3D12RHI/Private/D3D12Allocation.cpp b/Engine/Source/Runtime/D3D12RHI/Private/D3D12Allocation.cpp
--- a/Engine/Source/Runtime/D3D12RHI/Private/D3D12Allocation.cpp
+++ b/Engine/Source/Runtime/D3D12RHI/Private/D3D12Allocation.cpp
@@ -1598,7 +1598,7 @@ HRESULT FD3D12TextureAllocatorPool::AllocateTexture(
 {
 	// The top mip level must be less than 64 KB to use 4 KB alignment
 	bool b4KAligment = TextureCanBe4KAligned(Desc, (EPixelFormat)UEFormat);
-	Desc.Alignment = b4KAligment ?	D3D12_SMALL_RESOURCE_PLACEMENT_ALIGNMENT : D3D12_DEFAULT_RESOURCE_PLACEMENT_ALIGNMENT;
+	Desc.Alignment = b4KAligment ?	D3D12_SMALL_RESOURCE_PLACEMENT_ALIGNMENT : (Desc.SampleDesc.Count > 1 ? D3D12_DEFAULT_MSAA_RESOURCE_PLACEMENT_ALIGNMENT : D3D12_DEFAULT_RESOURCE_PLACEMENT_ALIGNMENT);
 	const D3D12_RESOURCE_ALLOCATION_INFO Info = GetParentDevice()->GetDevice()->GetResourceAllocationInfo(0, 1, &Desc);
 
 	bool bIsReadOnly = !(Desc.Flags & D3D12_RESOURCE_FLAG_ALLOW_RENDER_TARGET ||
diff --git a/Engine/Source/Runtime/D3D12RHI/Private/D3D12Buffer.cpp b/Engine/Source/Runtime/D3D12RHI/Private/D3D12Buffer.cpp
--- a/Engine/Source/Runtime/D3D12RHI/Private/D3D12Buffer.cpp
+++ b/Engine/Source/Runtime/D3D12RHI/Private/D3D12Buffer.cpp
@@ -211,10 +211,7 @@ FD3D12Buffer* FD3D12Adapter::CreateRHIBuffer(FRHICommandListImmediate* RHICmdLis
 
 	const bool bIsDynamic = (InUsage & BUF_AnyDynamic) ? true : false;
 	const uint32 FirstGPUIndex = CreateInfo.GPUMask.GetFirstIndex();
-
-	// Transient flag set?
-	bool bIsTransient = (InUsage & BUF_Transient);
-	
+		
 	// Does this resource support tracking?
 	const bool bSupportResourceStateTracking = !bIsDynamic && FD3D12DefaultBufferAllocator::IsPlacedResource(InDesc.Flags, InResourceStateMode);
 
@@ -226,10 +223,7 @@ FD3D12Buffer* FD3D12Adapter::CreateRHIBuffer(FRHICommandListImmediate* RHICmdLis
 
 	FD3D12Buffer* BufferOut = nullptr;
 	if (bIsDynamic)
-	{
-		// Assume not transient and dynamic
-		check(!bIsTransient);
-		
+	{		
 		FD3D12Buffer* NewBuffer0 = nullptr;
 		BufferOut = CreateLinkedObject<FD3D12Buffer>(CreateInfo.GPUMask, [&](FD3D12Device* Device)
 		{
@@ -260,8 +254,7 @@ FD3D12Buffer* FD3D12Adapter::CreateRHIBuffer(FRHICommandListImmediate* RHICmdLis
 			FD3D12Buffer* NewBuffer = new FD3D12Buffer(Device, Size, InUsage, Stride);
 			NewBuffer->BufferAlignment = Alignment;
 			AllocateBuffer(Device, InDesc, Size, InUsage, InResourceStateMode, CreateState, CreateInfo, Alignment, NewBuffer, NewBuffer->ResourceLocation, ResourceAllocator);
-			NewBuffer->ResourceLocation.SetTransient(bIsTransient);
-
+			
 			// Unlock immediately if no initial data
 			if (CreateInfo.ResourceArray == nullptr)
 			{
@@ -274,7 +267,6 @@ FD3D12Buffer* FD3D12Adapter::CreateRHIBuffer(FRHICommandListImmediate* RHICmdLis
 
 	if (CreateInfo.ResourceArray)
 	{
-		check(!bIsTransient);
 		if (bIsDynamic == false && BufferOut->ResourceLocation.IsValid())
 		{
 			check(Size == CreateInfo.ResourceArray->GetResourceDataSize());
@@ -341,12 +333,6 @@ FD3D12Buffer* FD3D12Adapter::CreateRHIBuffer(FRHICommandListImmediate* RHICmdLis
 		CreateInfo.ResourceArray->Discard();
 	}
 
-	// Don't update stats for transient resources
-	if (!bIsTransient)
-	{
-		UpdateBufferStats(GetBufferStats(InUsage), BufferOut->ResourceLocation.GetSize());
-	}
-
 	return BufferOut;
 }
 
diff --git a/Engine/Source/Runtime/D3D12RHI/Private/D3D12PoolAllocator.cpp b/Engine/Source/Runtime/D3D12RHI/Private/D3D12PoolAllocator.cpp
--- a/Engine/Source/Runtime/D3D12RHI/Private/D3D12PoolAllocator.cpp
+++ b/Engine/Source/Runtime/D3D12RHI/Private/D3D12PoolAllocator.cpp
@@ -7,20 +7,6 @@
 #define NEEDS_D3D12_INDIRECT_ARGUMENT_HEAP_WORKAROUND 0
 #endif
 
-#if TRACK_RESOURCE_ALLOCATIONS
-
-// Limited size because tracking currently still takes some overhead when creating the callstack for each allocation
-// Arbitrary size right now for which we are interested to know where these allocations are coming from and if they are needed
-// and can be optimized.
-static int32 GD3D12MinPoolAllocationSizeToTrack = 16 * 1024 * 1024;
-static FAutoConsoleVariableRef CVarD3D12MinPoolAllocationSizeToTrack(
-	TEXT("d3d12.MinPoolAllocationSizeToTrack"),
-	GD3D12MinPoolAllocationSizeToTrack,
-	TEXT("Minimum allocation size to track for pool allocations (default 16MB)"),
-	ECVF_ReadOnly);
-
-#endif // TRACK_RESOURCE_ALLOCATIONS
-
 //-----------------------------------------------------------------------------
 //	FD3D12MemoryPool
 //-----------------------------------------------------------------------------
@@ -63,6 +49,14 @@ void FD3D12MemoryPool::Init()
 		Desc.Properties = HeapProps;
 		Desc.Alignment = 0;
 		Desc.Flags = InitConfig.HeapFlags;
+
+		// All all resources types on heap when tracking allocations (need to be able to create a buffer on it to extract the base gpu address)
+		// (Tier 2 support already checked when this flag is enabled)
+		if (Adapter->IsTrackingAllAllocations())
+		{
+			Desc.Flags = D3D12_HEAP_FLAG_ALLOW_ALL_BUFFERS_AND_TEXTURES;
+		}
+
 		if (Adapter->IsHeapNotZeroedSupported())
 		{
 			Desc.Flags |= FD3D12_HEAP_FLAG_CREATE_NOT_ZEROED;
@@ -328,18 +322,43 @@ void FD3D12PoolAllocator::AllocateResource(uint32 GPUIndex, D3D12_HEAP_TYPE InHe
 
 			FD3D12Resource* NewResource = CreatePlacedResource(AllocationData, InDesc, InCreateState, InResourceStateMode, InClearValue, InName);
 			ResourceLocation.SetResource(NewResource);
+			ResourceLocation.SetGPUVirtualAddress(NewResource->GetGPUVirtualAddress());
 		}
 
 		UpdateAllocationTracking(ResourceLocation, EAllocationType::Allocate);
 	}
 	else
 	{
-		// Allocate Standalone - move to owner of resource because this allocator should only manage pooled allocations (needed for now to do the same as FD3D12DefaultBufferPool)
 		FD3D12Resource* NewResource = nullptr;
 		const D3D12_HEAP_PROPERTIES HeapProps = CD3DX12_HEAP_PROPERTIES(InHeapType, GetGPUMask().GetNative(), GetVisibilityMask().GetNative());
 		D3D12_RESOURCE_DESC Desc = InDesc;
 		Desc.Alignment = 0;
-		VERIFYD3D12RESULT(Adapter->CreateCommittedResource(Desc, GetGPUMask(), HeapProps, InCreateState, InResourceStateMode, InCreateState, InClearValue, &NewResource, InName, false));
+
+		// If we are tracking all allocation data and allocating a standalone texture, then first create a heap so we can retrieve the GPU virtual address as well
+		if (Adapter->IsTrackingAllAllocations() && InDesc.Dimension != D3D12_RESOURCE_DIMENSION_BUFFER)
+		{
+			D3D12_HEAP_DESC HeapDesc = {};
+			HeapDesc.SizeInBytes = InSize;
+			HeapDesc.Properties = HeapProps;
+			HeapDesc.Alignment = InDesc.SampleDesc.Count > 1 ? D3D12_DEFAULT_MSAA_RESOURCE_PLACEMENT_ALIGNMENT : 0;
+			HeapDesc.Flags = D3D12_HEAP_FLAG_ALLOW_ALL_BUFFERS_AND_TEXTURES;
+			if (Adapter->IsHeapNotZeroedSupported())
+			{
+				HeapDesc.Flags |= FD3D12_HEAP_FLAG_CREATE_NOT_ZEROED;
+			}
+
+			ID3D12Heap* Heap = nullptr;
+			VERIFYD3D12RESULT(Adapter->GetD3DDevice()->CreateHeap(&HeapDesc, IID_PPV_ARGS(&Heap)));
+			TRefCountPtr<FD3D12Heap> BackingHeap = new FD3D12Heap(GetParentDevice(), GetVisibilityMask());
+			BackingHeap->SetHeap(Heap);		
+			
+			VERIFYD3D12RESULT(Adapter->CreatePlacedResource(Desc, BackingHeap, 0, InCreateState, InResourceStateMode, InCreateState, InClearValue, &NewResource, InName));
+		}
+		else
+		{
+			// Allocate Standalone - move to owner of resource because this allocator should only manage pooled allocations (needed for now to do the same as FD3D12DefaultBufferPool)
+			VERIFYD3D12RESULT(Adapter->CreateCommittedResource(Desc, GetGPUMask(), HeapProps, InCreateState, InResourceStateMode, InCreateState, InClearValue, &NewResource, InName, false));
+		}
 
 		ResourceLocation.AsStandAlone(NewResource, InSize);
 	}
@@ -363,11 +382,11 @@ FD3D12Resource* FD3D12PoolAllocator::CreatePlacedResource(
 }
 
 
-void FD3D12PoolAllocator::DeallocateResource(FD3D12ResourceLocation& ResourceLocation)
+void FD3D12PoolAllocator::DeallocateResource(FD3D12ResourceLocation& ResourceLocation, bool bDefragFree)
 {
 	check(IsOwner(ResourceLocation));
 
-	UpdateAllocationTracking(ResourceLocation, EAllocationType::Free);
+	UpdateAllocationTracking(ResourceLocation, bDefragFree ? EAllocationType::DefragFree : EAllocationType::Free);
 	
 	FScopeLock Lock(&CS);
 
@@ -465,7 +484,8 @@ bool FD3D12PoolAllocator::HandleDefragRequest(FRHIPoolAllocationData* InSourceBl
 	FD3D12Resource* CurrentResource = Owner->GetResource();
 
 	// Release the current allocation (will only be freed on the next frame fence)
-	DeallocateResource(*Owner);
+	bool bDefragFree = true;
+	DeallocateResource(*Owner, bDefragFree);
 
 	// Move temp allocation data to allocation data of the owner (part of different allocator now)		
 	bool bLocked = true;
@@ -650,21 +670,24 @@ FD3D12HeapAndOffset FD3D12PoolAllocator::GetBackingHeapAndAllocationOffsetInByte
 void FD3D12PoolAllocator::UpdateAllocationTracking(FD3D12ResourceLocation& InAllocation, EAllocationType InAllocationType)
 {
 #if TRACK_RESOURCE_ALLOCATIONS
-	// Only track for default heap and of allocation is of certain minimum size
 	check(IsOwner(InAllocation));
 	FRHIPoolAllocationData& AllocationData = InAllocation.GetPoolAllocatorPrivateData().PoolData;
 	check(AllocationData.IsAllocated());
 	uint64 AllocationSize = AllocationData.GetSize();
-	if (InitConfig.HeapType == D3D12_HEAP_TYPE_DEFAULT && GD3D12MinPoolAllocationSizeToTrack > 0 && AllocationSize >= (uint32)GD3D12MinPoolAllocationSizeToTrack)
+
+	// Track only when tracking all resources is set
+	FD3D12Adapter* Adapter = GetParentDevice()->GetParentAdapter();
+	if (InitConfig.HeapType == D3D12_HEAP_TYPE_DEFAULT && AllocationStrategy == EResourceAllocationStrategy::kPlacedResource && Adapter->IsTrackingAllAllocations())
 	{
-		FD3D12Adapter* Adapter = GetParentDevice()->GetParentAdapter();		
 		if (InAllocationType == EAllocationType::Allocate)
 		{
-			Adapter->TrackAllocationData(&InAllocation, AllocationSize);
+			bool bCollectCallstack = false;
+			Adapter->TrackAllocationData(&InAllocation, AllocationSize, bCollectCallstack);
 		}
 		else
 		{
-			Adapter->ReleaseTrackedAllocationData(&InAllocation);
+			bool bDefragFree = InAllocationType == EAllocationType::DefragFree;
+			Adapter->ReleaseTrackedAllocationData(&InAllocation, bDefragFree);
 		}
 	}
 #endif // TRACK_RESOURCE_ALLOCATIONS
diff --git a/Engine/Source/Runtime/D3D12RHI/Private/D3D12PoolAllocator.h b/Engine/Source/Runtime/D3D12RHI/Private/D3D12PoolAllocator.h
--- a/Engine/Source/Runtime/D3D12RHI/Private/D3D12PoolAllocator.h
+++ b/Engine/Source/Runtime/D3D12RHI/Private/D3D12PoolAllocator.h
@@ -134,7 +134,7 @@ public:
 	bool SupportsAllocation(D3D12_HEAP_TYPE InHeapType, D3D12_RESOURCE_FLAGS InResourceFlags, EBufferUsageFlags InBufferUsage, ED3D12ResourceStateMode InResourceStateMode) const;
 	void AllocDefaultResource(D3D12_HEAP_TYPE InHeapType, const D3D12_RESOURCE_DESC& InDesc, EBufferUsageFlags InBufferUsage, ED3D12ResourceStateMode InResourceStateMode,
 		D3D12_RESOURCE_STATES InCreateState, uint32 InAlignment, const TCHAR* InName, FD3D12ResourceLocation& ResourceLocation);
-	virtual void DeallocateResource(FD3D12ResourceLocation& ResourceLocation);
+	virtual void DeallocateResource(FD3D12ResourceLocation& ResourceLocation, bool bDefragFree = false);
 
 	// Implementation of ID3D12ResourceAllocator
 	virtual void AllocateResource(uint32 GPUIndex, D3D12_HEAP_TYPE InHeapType, const D3D12_RESOURCE_DESC& InDesc, uint64 InSize, uint32 InAllocationAlignment, ED3D12ResourceStateMode InResourceStateMode,
@@ -166,8 +166,9 @@ protected:
 	// Track the allocation
 	enum class EAllocationType
 	{
-		Allocate,
-		Free
+		Allocate,		// Update tracking data on allocation
+		DefragFree,		// Update tracking data on free during defrag op
+		Free			// Update tracking data when resource is actually freed
 	};
 	virtual void UpdateAllocationTracking(FD3D12ResourceLocation& InAllocation, EAllocationType InAllocationType);
 
diff --git a/Engine/Source/Runtime/D3D12RHI/Private/D3D12Resources.cpp b/Engine/Source/Runtime/D3D12RHI/Private/D3D12Resources.cpp
--- a/Engine/Source/Runtime/D3D12RHI/Private/D3D12Resources.cpp
+++ b/Engine/Source/Runtime/D3D12RHI/Private/D3D12Resources.cpp
@@ -295,6 +295,7 @@ FD3D12Resource::FD3D12Resource(FD3D12Device* ParentDevice,
 		FPlatformMemory::Memzero(&ClearValue, sizeof(D3D12_CLEAR_VALUE));
 	}
 
+	// On Windows it's sadly enough not possible to get the GPU virtual address from the resource directly
 	if (Resource
 #if PLATFORM_WINDOWS
 		&& Desc.Dimension == D3D12_RESOURCE_DIMENSION_BUFFER
@@ -409,6 +410,26 @@ FD3D12Heap::~FD3D12Heap()
 	Destroy();
 }
 
+void FD3D12Heap::SetHeap(ID3D12Heap* HeapIn)
+{
+	*Heap.GetInitReference() = HeapIn; 
+	
+	D3D12_HEAP_DESC HeapDesc = Heap->GetDesc();
+	HeapType = HeapDesc.Properties.Type;
+
+	// Create a buffer placed resource on the heap to extract the gpu virtual address
+	// if we are tracking all allocations
+	FD3D12Adapter* Adapter = GetParentDevice()->GetParentAdapter();	
+	if (Adapter->IsTrackingAllAllocations() && HeapType == D3D12_HEAP_TYPE_DEFAULT)
+	{
+		uint64 HeapSize = HeapDesc.SizeInBytes;
+		TRefCountPtr<ID3D12Resource> TempResource;
+		const D3D12_RESOURCE_DESC BufDesc = CD3DX12_RESOURCE_DESC::Buffer(HeapSize, D3D12_RESOURCE_FLAG_NONE);
+		const HRESULT hr = Adapter->GetD3DDevice()->CreatePlacedResource(Heap, 0, &BufDesc, D3D12_RESOURCE_STATE_COMMON, nullptr, IID_PPV_ARGS(TempResource.GetInitReference()));
+		GPUVirtualAddress = TempResource->GetGPUVirtualAddress();
+	}
+}
+
 void FD3D12Heap::UpdateResidency(FD3D12CommandListHandle& CommandList)
 {
 #if ENABLE_RESIDENCY_MANAGEMENT
@@ -530,6 +551,23 @@ HRESULT FD3D12Adapter::CreatePlacedResource(const D3D12_RESOURCE_DESC& InDesc, F
 			BackingHeap,
 			HeapDesc.Properties.Type);
 
+#if PLATFORM_WINDOWS
+		if (IsTrackingAllAllocations() && BackingHeap->GetHeapType() == D3D12_HEAP_TYPE_DEFAULT)
+		{
+			// Manually set the GPU virtual address from the heap gpu virtual address & offset
+			if (InDesc.Dimension != D3D12_RESOURCE_DIMENSION_BUFFER)
+			{
+				check(BackingHeap->GetGPUVirtualAddress() != 0);
+				(*ppOutResource)->SetGPUVirtualAddress(BackingHeap->GetGPUVirtualAddress() + HeapOffset);
+			}
+			else
+			{
+				check((*ppOutResource)->GetGPUVirtualAddress() != 0);
+				check((*ppOutResource)->GetGPUVirtualAddress() == BackingHeap->GetGPUVirtualAddress() + HeapOffset);
+			}
+		}
+#endif		
+
 		// Set a default name (can override later).
 		SetName(*ppOutResource, Name);
 
@@ -736,6 +774,18 @@ void FD3D12ResourceLocation::ReferenceNode(FD3D12Device* DestinationDevice, FD3D
 
 void FD3D12ResourceLocation::ReleaseResource()
 {
+#if TRACK_RESOURCE_ALLOCATIONS
+	if (IsTransient())
+	{
+		FD3D12Adapter* Adapter = GetParentDevice()->GetParentAdapter();
+		if (Adapter->IsTrackingAllAllocations())
+		{
+			bool bDefragFree = false;
+			Adapter->ReleaseTrackedAllocationData(this, bDefragFree);
+		}
+	}
+#endif
+
 	switch (Type)
 	{
 	case ResourceLocationType::eStandAlone:
@@ -860,11 +910,13 @@ void FD3D12ResourceLocation::UpdateStandAloneStats(bool bIncrement)
 		// Track all committed resource allocations
 		if (bIncrement)
 		{
-			UnderlyingResource->GetParentDevice()->GetParentAdapter()->TrackAllocationData(this, Info.SizeInBytes);
+			bool bCollectCallstack = true;
+			UnderlyingResource->GetParentDevice()->GetParentAdapter()->TrackAllocationData(this, Info.SizeInBytes, bCollectCallstack);
 		}
 		else
 		{
-			UnderlyingResource->GetParentDevice()->GetParentAdapter()->ReleaseTrackedAllocationData(this);
+			bool bDefragFree = false;
+			UnderlyingResource->GetParentDevice()->GetParentAdapter()->ReleaseTrackedAllocationData(this, bDefragFree);
 		}
 	}
 }
diff --git a/Engine/Source/Runtime/D3D12RHI/Private/D3D12Texture.cpp b/Engine/Source/Runtime/D3D12RHI/Private/D3D12Texture.cpp
--- a/Engine/Source/Runtime/D3D12RHI/Private/D3D12Texture.cpp
+++ b/Engine/Source/Runtime/D3D12RHI/Private/D3D12Texture.cpp
@@ -900,9 +900,6 @@ TD3D12Texture2D<BaseResourceType>* FD3D12DynamicRHI::CreateD3D12Texture2D(FRHICo
 		ClearValuePtr = &ClearValue;
 	}
 
-	// Check if transient
-	bool bIsTransient = EnumHasAnyFlags(Flags, TexCreate_Transient);
-
 	// The state this resource will be in when it leaves this function
 	const FD3D12Resource::FD3D12ResourceTypeHelper Type(TextureDesc, D3D12_HEAP_TYPE_DEFAULT);
 	const D3D12_RESOURCE_STATES InitialState = Type.GetOptimalInitialState(InResourceState, false);
@@ -955,9 +952,6 @@ TD3D12Texture2D<BaseResourceType>* FD3D12DynamicRHI::CreateD3D12Texture2D(FRHICo
 			Location.UnlockPoolData();
 		}
 
-		// Mark transient
-		Location.SetTransient(bIsTransient);
-
 		check(Location.IsValid());
 
 		uint32 RTVIndex = 0;
@@ -1201,9 +1195,6 @@ FD3D12Texture3D* FD3D12DynamicRHI::CreateD3D12Texture3D(FRHICommandListImmediate
 	check(!(Flags & TexCreate_ResolveTargetable));
 	check(Flags & TexCreate_ShaderResource);
 
-	// Is transient
-	bool bIsTransient = EnumHasAnyFlags(Flags, TexCreate_Transient);
-
 	D3D12_CLEAR_VALUE *ClearValuePtr = nullptr;
 	D3D12_CLEAR_VALUE ClearValue;
 	if (bCreateRTV && CreateInfo.ClearValueBinding.ColorBinding == EClearBinding::EColorBound)
@@ -1236,9 +1227,6 @@ FD3D12Texture3D* FD3D12DynamicRHI::CreateD3D12Texture3D(FRHICommandListImmediate
 		}
 		Texture3D->ResourceLocation.SetOwner(Texture3D);
 
-		// Mark transient
-		Texture3D->ResourceLocation.SetTransient(bIsTransient);
-
 		// Unlock immediately if no initial data
 		if (CreateInfo.BulkData == nullptr)
 		{
diff --git a/Engine/Source/Runtime/D3D12RHI/Private/D3D12TransientResourceAllocator.cpp b/Engine/Source/Runtime/D3D12RHI/Private/D3D12TransientResourceAllocator.cpp
--- a/Engine/Source/Runtime/D3D12RHI/Private/D3D12TransientResourceAllocator.cpp
+++ b/Engine/Source/Runtime/D3D12RHI/Private/D3D12TransientResourceAllocator.cpp
@@ -167,11 +167,21 @@ void FD3D12TransientResourceAllocator::FResourceAllocatorAdapter::AllocateResour
 #endif
 
 	FD3D12Resource* NewResource = nullptr;
-	VERIFYD3D12RESULT(GetParentAdapter()->CreatePlacedResource(InDesc, Heap.GetLinkedObject(GPUIndex)->Get(), Allocation.Offset, InCreateState, InResourceStateMode, D3D12_RESOURCE_STATE_TBD, InClearValue, &NewResource, InName));
+	FD3D12Adapter* Adapter = GetParentAdapter();
+	VERIFYD3D12RESULT(Adapter->CreatePlacedResource(InDesc, Heap.GetLinkedObject(GPUIndex)->Get(), Allocation.Offset, InCreateState, InResourceStateMode, D3D12_RESOURCE_STATE_TBD, InClearValue, &NewResource, InName));
 
 	check(!ResourceLocation.IsValid());
 	ResourceLocation.AsHeapAliased(NewResource);
 	ResourceLocation.SetSize(InSize);
+	ResourceLocation.SetTransient(true);
+
+#if TRACK_RESOURCE_ALLOCATIONS
+	if (Adapter->IsTrackingAllAllocations())
+	{
+		bool bCollectCallstack = false;
+		Adapter->TrackAllocationData(&ResourceLocation, Allocation.Size, bCollectCallstack);
+	}
+#endif
 }
 
 FRHITransientBuffer* FD3D12TransientResourceAllocator::CreateBuffer(const FRHIBufferCreateInfo& InCreateInfo, const TCHAR* InDebugName)
diff --git a/Engine/Source/Runtime/D3D12RHI/Private/D3D12Util.cpp b/Engine/Source/Runtime/D3D12RHI/Private/D3D12Util.cpp
--- a/Engine/Source/Runtime/D3D12RHI/Private/D3D12Util.cpp
+++ b/Engine/Source/Runtime/D3D12RHI/Private/D3D12Util.cpp
@@ -316,7 +316,7 @@ struct FDred_1_2
 
 /** Log the DRED data to Error log if available */
 template <typename FDred_T>
-static bool LogDREDData(ID3D12Device* Device)
+static bool LogDREDData(ID3D12Device* Device, bool bTrackingAllAllocations, D3D12_GPU_VIRTUAL_ADDRESS& OutPageFaultGPUAddress)
 {
 	// Should match all values from D3D12_AUTO_BREADCRUMB_OP
 	static const TCHAR* OpNames[] =
@@ -464,18 +464,23 @@ static bool LogDREDData(ID3D12Device* Device)
 		D3D12_DRED_PAGE_FAULT_OUTPUT DredPageFaultOutput;
 		if (SUCCEEDED(Dred.Data->GetPageFaultAllocationOutput(&DredPageFaultOutput)) && DredPageFaultOutput.PageFaultVA != 0)
 		{
+			OutPageFaultGPUAddress = DredPageFaultOutput.PageFaultVA;
 			UE_LOG(LogD3D12RHI, Error, TEXT("DRED: PageFault at VA GPUAddress \"0x%llX\""), (long long)DredPageFaultOutput.PageFaultVA);
-
+			
 			const D3D12_DRED_ALLOCATION_NODE* Node = DredPageFaultOutput.pHeadExistingAllocationNode;
 			if (Node)
 			{
 				UE_LOG(LogD3D12RHI, Error, TEXT("DRED: Active objects with VA ranges that match the faulting VA:"));
 				while (Node)
 				{
-					int32 alloc_type_index = Node->AllocationType - D3D12_DRED_ALLOCATION_TYPE_COMMAND_QUEUE;
-					const TCHAR* AllocTypeName = (alloc_type_index < UE_ARRAY_COUNT(AllocTypesNames)) ? AllocTypesNames[alloc_type_index] : TEXT("Unknown Alloc");
-					UE_LOG(LogD3D12RHI, Error, TEXT("\tName: %s (Type: %s)"), Node->ObjectNameW, AllocTypeName);
-
+					// When tracking all allocations then empty named dummy resources (heap & buffer)
+					// are created for each texture to extract the GPUBaseAddress so don't write these out
+					if (!bTrackingAllAllocations || Node->ObjectNameW)
+					{
+						int32 alloc_type_index = Node->AllocationType - D3D12_DRED_ALLOCATION_TYPE_COMMAND_QUEUE;
+						const TCHAR* AllocTypeName = (alloc_type_index < UE_ARRAY_COUNT(AllocTypesNames)) ? AllocTypesNames[alloc_type_index] : TEXT("Unknown Alloc");
+						UE_LOG(LogD3D12RHI, Error, TEXT("\tName: %s (Type: %s)"), Node->ObjectNameW, AllocTypeName);
+					}
 					Node = Node->pNext;
 				}
 			}
@@ -486,9 +491,13 @@ static bool LogDREDData(ID3D12Device* Device)
 				UE_LOG(LogD3D12RHI, Error, TEXT("DRED: Recent freed objects with VA ranges that match the faulting VA:"));
 				while (Node)
 				{
-					int32 alloc_type_index = Node->AllocationType - D3D12_DRED_ALLOCATION_TYPE_COMMAND_QUEUE;
-					const TCHAR* AllocTypeName = (alloc_type_index < UE_ARRAY_COUNT(AllocTypesNames)) ? AllocTypesNames[alloc_type_index] : TEXT("Unknown Alloc");
-					UE_LOG(LogD3D12RHI, Error, TEXT("\tName: %s (Type: %s)"), Node->ObjectNameW, AllocTypeName);
+					// See comments above
+					if (!bTrackingAllAllocations || Node->ObjectNameW)
+					{
+						int32 alloc_type_index = Node->AllocationType - D3D12_DRED_ALLOCATION_TYPE_COMMAND_QUEUE;
+						const TCHAR* AllocTypeName = (alloc_type_index < UE_ARRAY_COUNT(AllocTypesNames)) ? AllocTypesNames[alloc_type_index] : TEXT("Unknown Alloc");
+						UE_LOG(LogD3D12RHI, Error, TEXT("\tName: %s (Type: %s)"), Node->ObjectNameW, AllocTypeName);
+					}
 
 					Node = Node->pNext;
 				}
@@ -507,6 +516,49 @@ static bool LogDREDData(ID3D12Device* Device)
 	}
 }
 
+void LogPageFaultData(FD3D12Adapter* InAdapter, D3D12_GPU_VIRTUAL_ADDRESS InPageFaultAddress)
+{
+	if (InPageFaultAddress == 0)
+	{
+		return;
+	}
+
+	UE_LOG(LogD3D12RHI, Error, TEXT("PageFault: PageFault at VA GPUAddress \"0x%llX\""), (long long)InPageFaultAddress);
+	UE_LOG(LogD3D12RHI, Error, TEXT("PageFault: Logging all resource enabled: %s"), InAdapter->IsTrackingAllAllocations() ? TEXT("Yes") : TEXT("No"));
+
+	// Try and find all current allocations near that range
+	static const int64 CheckRangeRadius = 16 * 1024 * 1024;
+	TArray<FD3D12Adapter::FAllocatedResourceResult> OverlappingResources;
+	InAdapter->FindResourcesNearGPUAddress(InPageFaultAddress, CheckRangeRadius, OverlappingResources);
+	UE_LOG(LogD3D12RHI, Error, TEXT("PageFault: Found %d active tracked resources in %3.2f MB range of page fault address"), OverlappingResources.Num(), CheckRangeRadius / (1024.0f * 1024));
+	if (OverlappingResources.Num() > 0)
+	{
+		uint32 PrintCount = FMath::Min(OverlappingResources.Num(), 100);
+		for (uint32 Index = 0; Index < PrintCount; ++Index)
+		{
+			FD3D12Adapter::FAllocatedResourceResult OverlappingResource = OverlappingResources[Index];
+			D3D12_GPU_VIRTUAL_ADDRESS ResourceAddress = OverlappingResource.Allocation->GetGPUVirtualAddress();
+			UE_LOG(LogD3D12RHI, Error, TEXT("\tGPU Address: \"0x%llX\" - Size: %3.2f MB - Distance to page fault: %3.2f MB - Transient: %d - Name: %s"),
+				(long long)ResourceAddress, OverlappingResource.Allocation->GetSize() / (1024.0f * 1024), OverlappingResource.Distance / (1024.0f * 1024), OverlappingResource.Allocation->IsTransient(), *OverlappingResource.Allocation->GetResource()->GetName().ToString());
+		}
+	}
+
+	// Try and find all released allocations within the faulting address
+	TArray<FD3D12Adapter::FReleasedAllocationData> ReleasedResources;
+	InAdapter->FindReleasedAllocationData(InPageFaultAddress, ReleasedResources);
+	UE_LOG(LogD3D12RHI, Error, TEXT("PageFault: Found %d released resources containing the page fault address during last 100 frames (current frame ID: %d)"), ReleasedResources.Num(), InAdapter->GetFrameFence().GetCurrentFence());
+	if (ReleasedResources.Num() > 0)
+	{
+		uint32 PrintCount = FMath::Min(ReleasedResources.Num(), 100);
+		for (uint32 Index = 0; Index < PrintCount; ++Index)
+		{
+			FD3D12Adapter::FReleasedAllocationData& AllocationData = ReleasedResources[Index];
+			UE_LOG(LogD3D12RHI, Error, TEXT("\tGPU Address: \"0x%llX\" - Size: %3.2f MB - FrameID: %4d - DefragFree: %d - Transient: %d - Name: %s"),
+				(long long)AllocationData.GPUVirtualAddress, AllocationData.AllocationSize / (1024.0f * 1024), AllocationData.ReleasedFrameID, AllocationData.bDefragFree, AllocationData.bTransient, *AllocationData.ResourceName.ToString());
+		}
+	}
+}
+
 #endif  // PLATFORM_WINDOWS
 
 extern CORE_API bool GIsGPUCrashed;
@@ -560,10 +612,14 @@ namespace D3D12RHI
 			{
 				if (InDevice == nullptr || InDevice == IterationDevice->GetDevice())
 				{
-					if (!LogDREDData<FDred_1_2>(IterationDevice->GetDevice()))
+					D3D12_GPU_VIRTUAL_ADDRESS PageFaultAddress = 0;
+					bool bIsTrackingAllAllocations = IterationDevice->GetParentAdapter()->IsTrackingAllAllocations();
+					if (!LogDREDData<FDred_1_2>(IterationDevice->GetDevice(), bIsTrackingAllAllocations, PageFaultAddress))
 					{
-						LogDREDData<FDred_1_1>(IterationDevice->GetDevice());
+						LogDREDData<FDred_1_1>(IterationDevice->GetDevice(), bIsTrackingAllAllocations, PageFaultAddress);
 					}
+
+					LogPageFaultData(IterationDevice->GetParentAdapter(), PageFaultAddress);
 				}
 			});
 #endif  // PLATFORM_WINDOWS
diff --git a/Engine/Source/Runtime/D3D12RHI/Public/D3D12Resources.h b/Engine/Source/Runtime/D3D12RHI/Public/D3D12Resources.h
--- a/Engine/Source/Runtime/D3D12RHI/Public/D3D12Resources.h
+++ b/Engine/Source/Runtime/D3D12RHI/Public/D3D12Resources.h
@@ -94,18 +94,20 @@ public:
 	~FD3D12Heap();
 
 	inline ID3D12Heap* GetHeap() { return Heap.GetReference(); }
-	inline void SetHeap(ID3D12Heap* HeapIn) { *Heap.GetInitReference() = HeapIn; }
+	void SetHeap(ID3D12Heap* HeapIn);
 
 	void UpdateResidency(FD3D12CommandListHandle& CommandList);
-
 	void BeginTrackingResidency(uint64 Size);
-
 	void Destroy();
 
+	inline D3D12_HEAP_TYPE GetHeapType() const { return HeapType; }
 	inline FD3D12ResidencyHandle* GetResidencyHandle() { return &ResidencyHandle; }
+	inline D3D12_GPU_VIRTUAL_ADDRESS GetGPUVirtualAddress() const { return GPUVirtualAddress; }
 
 private:
 	TRefCountPtr<ID3D12Heap> Heap;
+	D3D12_HEAP_TYPE HeapType = D3D12_HEAP_TYPE_DEFAULT;
+	D3D12_GPU_VIRTUAL_ADDRESS GPUVirtualAddress = 0;
 	FD3D12ResidencyHandle ResidencyHandle;
 };
 
@@ -197,6 +199,7 @@ public:
 	D3D12_CLEAR_VALUE const& GetClearValue() const { return ClearValue; }
 	D3D12_HEAP_TYPE GetHeapType() const { return HeapType; }
 	D3D12_GPU_VIRTUAL_ADDRESS GetGPUVirtualAddress() const { return GPUVirtualAddress; }
+	inline void SetGPUVirtualAddress(D3D12_GPU_VIRTUAL_ADDRESS Value) { GPUVirtualAddress = Value; }
 	void* GetResourceBaseAddress() const { check(ResourceBaseAddress); return ResourceBaseAddress; }
 	uint16 GetMipLevels() const { return Desc.MipLevels; }
 	uint16 GetArraySize() const { return (Desc.Dimension == D3D12_RESOURCE_DIMENSION_TEXTURE3D) ? 1 : Desc.DepthOrArraySize; }
