Added a basic implementation of assert in shaders for D3D12. Requires SM6 support which can be enabled through r.D3D12.AllowShaderModel6=1.

Example usage in a shader:

PLATFORM_ASSERT4(GroupThread > 4321, 0xDF7DC903, GroupThread, GroupIndex, GroupBase, asuint(123.456f));

Example GPU crash report output:

LogD3D12RHI: Error: [GPUBreadCrumb] Last tracked GPU operations:
LogD3D12RHI: Error: [GPUBreadCrumb]	3D Queue 0 - Begin: Frame 348 - Scene - ReflectionIndirect - RayTracingReflections - RayTracingDeferredReflectionsShade 957x466
LogD3D12RHI: Error: [GPUBreadCrumb]	Shader assertion failed! ID: 0xDF7DC903
LogD3D12RHI: Error: [GPUBreadCrumb]	Payload  [int32]: 416 42 172032 1123477881
LogD3D12RHI: Error: [GPUBreadCrumb]	Payload [uint32]: 416 42 172032 1123477881
LogD3D12RHI: Error: [GPUBreadCrumb]	Payload    [hex]: 0x000001A0 0x0000002A 0x0002A000 0x42F6E979
LogD3D12RHI: Error: [GPUBreadCrumb]	Payload  [float]: 0.000000 0.000000 0.000000 123.456001

This works through automatic diagnostic buffer binding in D3D12 RHI, which becomes available in *any* shader as a UAV.
Opens the door for arbitrary human-readable message output from any shader in the future.

PLATFORM_ASSERT macro is always enabled and available on supported platforms. It is intended to be a low-level mechanism for a higher-level check() that would be controlled via CVars.

diff --git a/Engine/Shaders/Public/Platform.ush b/Engine/Shaders/Public/Platform.ush
--- a/Engine/Shaders/Public/Platform.ush
+++ b/Engine/Shaders/Public/Platform.ush
@@ -352,6 +352,13 @@
 	#define PLATFORM_BREAK() 
 #endif
 
+#if !defined(PLATFORM_ASSERT)
+	#define PLATFORM_ASSERT(condition, assert_id)
+	#define PLATFORM_ASSERT1(condition, assert_id, a)
+	#define PLATFORM_ASSERT2(condition, assert_id, a, b)
+	#define PLATFORM_ASSERT3(condition, assert_id, a, b, c)
+	#define PLATFORM_ASSERT4(condition, assert_id, a, b, c, d)
+#endif
 
 // Hlslcc platforms ignore the uniform keyword as it can't properly optimize flow
 #if COMPILER_HLSLCC
diff --git a/Engine/Shaders/Public/Platform/D3D/D3DCommon.ush b/Engine/Shaders/Public/Platform/D3D/D3DCommon.ush
--- a/Engine/Shaders/Public/Platform/D3D/D3DCommon.ush
+++ b/Engine/Shaders/Public/Platform/D3D/D3DCommon.ush
@@ -26,3 +26,39 @@ SamplerState D3DStaticTrilinearClampedSampler   : register(s5, space1000);
 #define GetGlobalSampler(Filter, WrapMode) D3DStatic##Filter##WrapMode##Sampler
 
 #endif // PLATFORM_SUPPORTS_STATIC_SAMPLERS
+
+#if PLATFORM_SUPPORTS_DIAGNOSTIC_BUFFER
+
+// Diagnostic buffer which will be automatically bound by the RHI as a root parameter
+// Counterpart to FD3D12DiagnosticBufferData in D3D12DirectCommandListManager.h
+// Diagnostic message layout: [uint counter, uint message_id, uint4 payload] (6 DWORDs total)
+RWStructuredBuffer<uint> UEDiagnosticBuffer : register(u0, space999);
+
+void UEReportAssertWithPayload(uint AssertID, uint4 Payload)
+{
+	if (WaveIsFirstLane())
+	{
+		// Write some diagnostic data if we're first
+		uint OldValue = 0;
+		InterlockedAdd(UEDiagnosticBuffer[0], 1, OldValue);
+		if (OldValue == 0)
+		{
+			UEDiagnosticBuffer[1] = AssertID;
+			UEDiagnosticBuffer[2] = Payload.x;
+			UEDiagnosticBuffer[3] = Payload.y;
+			UEDiagnosticBuffer[4] = Payload.z;
+			UEDiagnosticBuffer[5] = Payload.w;
+		}
+		// Try to trigger a crash
+		// TODO: use PLATFORM_BREAK
+		UEDiagnosticBuffer[0xFFFFFFFF] = 0;
+	}
+}
+
+#define PLATFORM_ASSERT(condition, assert_id)              { if (!(condition)) UEReportAssertWithPayload(assert_id, uint4(0, 0, 0, 0)); }
+#define PLATFORM_ASSERT1(condition, assert_id, a)          { if (!(condition)) UEReportAssertWithPayload(assert_id, uint4(a, 0, 0, 0));  }
+#define PLATFORM_ASSERT2(condition, assert_id, a, b)       { if (!(condition)) UEReportAssertWithPayload(assert_id, uint4(a, b, 0, 0));  }
+#define PLATFORM_ASSERT3(condition, assert_id, a, b, c)    { if (!(condition)) UEReportAssertWithPayload(assert_id, uint4(a, b, c, 0));  }
+#define PLATFORM_ASSERT4(condition, assert_id, a, b, c, d) { if (!(condition)) UEReportAssertWithPayload(assert_id, uint4(a, b, c, d));  }
+
+#endif // PLATFORM_SUPPORTS_DIAGNOSTIC_BUFFER
diff --git a/Engine/Shaders/Public/ShaderVersion.ush b/Engine/Shaders/Public/ShaderVersion.ush
--- a/Engine/Shaders/Public/ShaderVersion.ush
+++ b/Engine/Shaders/Public/ShaderVersion.ush
@@ -3,4 +3,4 @@
 // in Platform.ush (which should be included in any shader) it allows to invalidate the shader DDC.
 // 
 // If you are merging streams and there is a conflict with this GUID you should make a new GUID rather than taking one or the other.
-// GUID = 14D3C2851B1D43B385782BB210BC73EF
\ No newline at end of file
+// GUID = 24BB391CB1F842E9922C0FEDC88DF989
\ No newline at end of file
diff --git a/Engine/Source/Developer/Windows/ShaderFormatD3D/Private/D3DShaderCompiler.cpp b/Engine/Source/Developer/Windows/ShaderFormatD3D/Private/D3DShaderCompiler.cpp
--- a/Engine/Source/Developer/Windows/ShaderFormatD3D/Private/D3DShaderCompiler.cpp
+++ b/Engine/Source/Developer/Windows/ShaderFormatD3D/Private/D3DShaderCompiler.cpp
@@ -657,6 +657,7 @@ bool CompileAndProcessD3DShaderFXC(FString& PreprocessedShaderSource, const FStr
 	if (SUCCEEDED(Result))
 	{
 		bool bGlobalUniformBufferUsed = false;
+		bool bDiagnosticBufferUsed = false;
 		uint32 NumInstructions = 0;
 		uint32 NumSamplers = 0;
 		uint32 NumSRVs = 0;
@@ -786,7 +787,8 @@ bool CompileAndProcessD3DShaderFXC(FString& PreprocessedShaderSource, const FStr
 				ID3D11ShaderReflection, D3D11_SHADER_DESC, D3D11_SHADER_INPUT_BIND_DESC,
 				ID3D11ShaderReflectionConstantBuffer, D3D11_SHADER_BUFFER_DESC,
 				ID3D11ShaderReflectionVariable, D3D11_SHADER_VARIABLE_DESC>(Input.Target.Platform, BindingSpace, Input.VirtualSourceFilePath, Reflector, ShaderDesc,
-					bGlobalUniformBufferUsed, NumSamplers, NumSRVs, NumCBs, NumUAVs,
+					bGlobalUniformBufferUsed, bDiagnosticBufferUsed,
+					NumSamplers, NumSRVs, NumCBs, NumUAVs,
 					Output, UniformBufferNames, UsedUniformBufferSlots, VendorExtensions);
 
 			NumInstructions = ShaderDesc.InstructionCount;
@@ -907,6 +909,7 @@ void CompileD3DShader(const FShaderCompilerInput& Input, FShaderCompilerOutput&
 	{
 		AdditionalDefines.SetDefine(TEXT("PLATFORM_SUPPORTS_SM6_0_WAVE_OPERATIONS"), 1);
 		AdditionalDefines.SetDefine(TEXT("PLATFORM_SUPPORTS_STATIC_SAMPLERS"), 1);
+		AdditionalDefines.SetDefine(TEXT("PLATFORM_SUPPORTS_DIAGNOSTIC_BUFFER"), 1);
 		AdditionalDefines.SetDefine(TEXT("COMPILER_SUPPORTS_NOINLINE"), 1);
 	}
 
diff --git a/Engine/Source/Developer/Windows/ShaderFormatD3D/Private/D3DShaderCompiler.inl b/Engine/Source/Developer/Windows/ShaderFormatD3D/Private/D3DShaderCompiler.inl
--- a/Engine/Source/Developer/Windows/ShaderFormatD3D/Private/D3DShaderCompiler.inl
+++ b/Engine/Source/Developer/Windows/ShaderFormatD3D/Private/D3DShaderCompiler.inl
@@ -241,7 +241,7 @@ template <typename ID3D1xShaderReflection, typename D3D1x_SHADER_DESC, typename
 	typename ID3D1xShaderReflectionVariable, typename D3D1x_SHADER_VARIABLE_DESC>
 inline void ExtractParameterMapFromD3DShader(
 		uint32 TargetPlatform, uint32 BindingSpace, const FString& VirtualSourceFilePath, ID3D1xShaderReflection* Reflector, const D3D1x_SHADER_DESC& ShaderDesc,
-		bool& bGlobalUniformBufferUsed, uint32& NumSamplers, uint32& NumSRVs, uint32& NumCBs, uint32& NumUAVs,
+		bool& bGlobalUniformBufferUsed, bool& bDiagnosticBufferUsed, uint32& NumSamplers, uint32& NumSRVs, uint32& NumCBs, uint32& NumUAVs,
 		FShaderCompilerOutput& Output, TArray<FString>& UniformBufferNames, TBitArray<>& UsedUniformBufferSlots, TArray<FShaderCodeVendorExtension>& VendorExtensions)
 {
 	// Add parameters for shader resources (constant buffers, textures, samplers, etc. */
@@ -374,6 +374,9 @@ inline void ExtractParameterMapFromD3DShader(
 
 			const bool bIsVendorParameter = bIsNVExtension || bIsIntelExtension || bIsAMDExtensionDX11 || bIsAMDExtensionDX12;
 
+			// See D3DCommon.ush
+			const bool bIsDiagnosticBufferParameter = (FCStringAnsi::Strcmp(BindDesc.Name, "UEDiagnosticBuffer") == 0);
+
 			TCHAR OfficialName[1024];
 			FCString::Strcpy(OfficialName, ANSI_TO_TCHAR(BindDesc.Name));
 
@@ -399,6 +402,10 @@ inline void ExtractParameterMapFromD3DShader(
 				VendorExtension.Parameter.Type = EShaderParameterType::UAV;
 				VendorExtensions.Add(VendorExtension);
 			}
+			else if (bIsDiagnosticBufferParameter)
+			{
+				bDiagnosticBufferUsed = true;
+			}
 			else
 			{
 				Output.ParameterMap.AddParameterAllocation(
diff --git a/Engine/Source/Developer/Windows/ShaderFormatD3D/Private/D3DShaderCompilerDXC.cpp b/Engine/Source/Developer/Windows/ShaderFormatD3D/Private/D3DShaderCompilerDXC.cpp
--- a/Engine/Source/Developer/Windows/ShaderFormatD3D/Private/D3DShaderCompilerDXC.cpp
+++ b/Engine/Source/Developer/Windows/ShaderFormatD3D/Private/D3DShaderCompilerDXC.cpp
@@ -407,6 +407,12 @@ inline bool IsCompatibleBinding(const D3D12_SHADER_INPUT_BIND_DESC& BindDesc, ui
 		const bool bIsAMDExtensionDX12 = (FCStringAnsi::Strcmp(BindDesc.Name, "AmdExtD3DShaderIntrinsicsUAV") == 0);
 		bIsCompatibleBinding = bIsAMDExtensionDX12 && (BindDesc.Space == AGS_DX12_SHADER_INSTRINSICS_SPACE_ID);
 	}
+	if (!bIsCompatibleBinding)
+	{
+		// #todo: there is currently no common header where a binding space number or buffer name could be defined. See D3DCommon.ush and D3D12RootSignature.cpp.
+		const bool bIsUEDebugBuffer = (FCStringAnsi::Strcmp(BindDesc.Name, "UEDiagnosticBuffer") == 0);
+		bIsCompatibleBinding = bIsUEDebugBuffer && (BindDesc.Space == 999);
+	}
 
 	return bIsCompatibleBinding;
 }
@@ -724,6 +730,7 @@ bool CompileAndProcessD3DShaderDXC(FString& PreprocessedShaderSource,
 		TArray<FShaderCodeVendorExtension> VendorExtensions;
 
 		bool bGlobalUniformBufferUsed = false;
+		bool bDiagnosticBufferUsed = false;
 		uint32 NumInstructions = 0;
 		uint32 NumSamplers = 0;
 		uint32 NumSRVs = 0;
@@ -792,7 +799,9 @@ bool CompileAndProcessD3DShaderDXC(FString& PreprocessedShaderSource,
 						ExtractParameterMapFromD3DShader<ID3D12FunctionReflection, D3D12_FUNCTION_DESC, D3D12_SHADER_INPUT_BIND_DESC,
 							ID3D12ShaderReflectionConstantBuffer, D3D12_SHADER_BUFFER_DESC,
 							ID3D12ShaderReflectionVariable, D3D12_SHADER_VARIABLE_DESC>(
-								Input.Target.Platform, AutoBindingSpace, Input.VirtualSourceFilePath, FunctionReflection, FunctionDesc, bGlobalUniformBufferUsed, NumSamplers, NumSRVs, NumCBs, NumUAVs,
+								Input.Target.Platform, AutoBindingSpace, Input.VirtualSourceFilePath, FunctionReflection, FunctionDesc, 
+								bGlobalUniformBufferUsed, bDiagnosticBufferUsed,
+								NumSamplers, NumSRVs, NumCBs, NumUAVs,
 								Output, UniformBufferNames, UsedUniformBufferSlots, VendorExtensions);
 
 						NumFoundEntryPoints++;
@@ -860,7 +869,9 @@ bool CompileAndProcessD3DShaderDXC(FString& PreprocessedShaderSource,
 			ExtractParameterMapFromD3DShader<ID3D12ShaderReflection, D3D12_SHADER_DESC, D3D12_SHADER_INPUT_BIND_DESC,
 				ID3D12ShaderReflectionConstantBuffer, D3D12_SHADER_BUFFER_DESC,
 				ID3D12ShaderReflectionVariable, D3D12_SHADER_VARIABLE_DESC>(
-					Input.Target.Platform, AutoBindingSpace, Input.VirtualSourceFilePath, ShaderReflection, ShaderDesc, bGlobalUniformBufferUsed, NumSamplers, NumSRVs, NumCBs, NumUAVs,
+					Input.Target.Platform, AutoBindingSpace, Input.VirtualSourceFilePath, ShaderReflection, ShaderDesc,
+						bGlobalUniformBufferUsed, bDiagnosticBufferUsed, 
+						NumSamplers, NumSRVs, NumCBs, NumUAVs,
 					Output, UniformBufferNames, UsedUniformBufferSlots, VendorExtensions);
 
 
@@ -890,6 +901,7 @@ bool CompileAndProcessD3DShaderDXC(FString& PreprocessedShaderSource,
 				FShaderCodeFeatures CodeFeatures;
 				//#todo-rco: Really should look inside DXIL
 				CodeFeatures.bUsesWaveOps = Input.Environment.CompilerFlags.Contains(CFLAG_WaveOperations);
+				CodeFeatures.bUsesDiagnosticBuffer = bDiagnosticBufferUsed;
 
 				// We only need this to appear when using a DXC shader
 				ShaderCode.AddOptionalData<FShaderCodeFeatures>(CodeFeatures);
diff --git a/Engine/Source/Runtime/D3D12RHI/Private/D3D12DescriptorCache.cpp b/Engine/Source/Runtime/D3D12RHI/Private/D3D12DescriptorCache.cpp
--- a/Engine/Source/Runtime/D3D12RHI/Private/D3D12DescriptorCache.cpp
+++ b/Engine/Source/Runtime/D3D12RHI/Private/D3D12DescriptorCache.cpp
@@ -367,6 +367,20 @@ void FD3D12DescriptorCache::SetUAVs(const FD3D12RootSignature* RootSignature, FD
 #ifdef VERBOSE_DESCRIPTOR_HEAP_DEBUG
 	FMsg::Logf(__FILE__, __LINE__, TEXT("DescriptorCache"), ELogVerbosity::Log, TEXT("SetUnorderedAccessViewTable [STAGE %d] to slots %d - %d"), (int32)ShaderStage, FirstSlotIndex, FirstSlotIndex + SlotsNeeded - 1);
 #endif
+
+	const int8 DiagnosticBufferSlot = RootSignature->GetDiagnosticBufferSlot();
+	D3D12_GPU_VIRTUAL_ADDRESS DiagnosticBufferAddress = CmdContext->GetCommandListManager().GetDiagnosticBufferGPUAddress();
+	if (DiagnosticBufferSlot >= 0 && DiagnosticBufferAddress)
+	{
+		if (ShaderStage == SF_Compute)
+		{
+			CommandList->SetComputeRootUnorderedAccessView(DiagnosticBufferSlot, DiagnosticBufferAddress);
+		}
+		else
+		{
+			CommandList->SetGraphicsRootUnorderedAccessView(DiagnosticBufferSlot, DiagnosticBufferAddress);
+		}
+	}
 }
 
 void FD3D12DescriptorCache::SetRenderTargets(FD3D12RenderTargetView** RenderTargetViewArray, uint32 Count, FD3D12DepthStencilView* DepthStencilTarget)
diff --git a/Engine/Source/Runtime/D3D12RHI/Private/D3D12DirectCommandListManager.cpp b/Engine/Source/Runtime/D3D12RHI/Private/D3D12DirectCommandListManager.cpp
--- a/Engine/Source/Runtime/D3D12RHI/Private/D3D12DirectCommandListManager.cpp
+++ b/Engine/Source/Runtime/D3D12RHI/Private/D3D12DirectCommandListManager.cpp
@@ -426,9 +426,15 @@ void FD3D12CommandListManager::Create(const TCHAR* Name, uint32 NumCommandLists,
 			int32 GPUCrashDataDepth = GetParentDevice()->GetGPUProfiler().GPUCrashDataDepth;
 			int32 MaxEventCount = GPUCrashDataDepth > 0 ? FMath::Min(GPUCrashDataDepth, MAX_GPU_BREADCRUMB_DEPTH) : MAX_GPU_BREADCRUMB_DEPTH;			
 
+			const uint32 ShaderDiagnosticBufferSize = sizeof(FD3D12DiagnosticBufferData);
+
 			// Allocate persistent CPU readable memory which will still be valid after a device lost and wrap this data in a placed resource
 			// so the GPU command list can write to it
-			const uint32 BreadCrumbBufferSize = MaxEventCount * sizeof(uint32);
+			const uint32 EventBufferSize = MaxEventCount * sizeof(uint32);
+			const uint32 BreadCrumbBufferSize = EventBufferSize + ShaderDiagnosticBufferSize;
+
+			DiagnosticBufferOffset = EventBufferSize;
+
 			BreadCrumbResourceAddress = VirtualAlloc(nullptr, BreadCrumbBufferSize, MEM_COMMIT, PAGE_READWRITE);
 			if (BreadCrumbResourceAddress)
 			{
@@ -448,6 +454,7 @@ void FD3D12CommandListManager::Create(const TCHAR* Name, uint32 NumCommandLists,
 					if (SUCCEEDED(hr))
 					{
 						UE_LOG(LogD3D12RHI, Log, TEXT("[GPUBreadCrumb] Successfully setup breadcrumb resource for %s"), Name);
+						BreadCrumbResourceGPUAddress = BreadCrumbResource->GetGPUVirtualAddress();
 					}
 					else
 					{
diff --git a/Engine/Source/Runtime/D3D12RHI/Private/D3D12DirectCommandListManager.h b/Engine/Source/Runtime/D3D12RHI/Private/D3D12DirectCommandListManager.h
--- a/Engine/Source/Runtime/D3D12RHI/Private/D3D12DirectCommandListManager.h
+++ b/Engine/Source/Runtime/D3D12RHI/Private/D3D12DirectCommandListManager.h
@@ -206,6 +206,22 @@ private:
 	const D3D12_COMMAND_LIST_TYPE Type;
 };
 
+// Counterpart to UEDiagnosticBuffer in D3DCommon.ush
+struct FD3D12DiagnosticBufferData
+{
+	uint32 Counter;
+	uint32 MessageID;
+	union
+	{
+		int32  AsInt[4];
+		uint32 AsUint[4];
+		float  AsFloat[4];
+	} Payload;
+};
+
+static_assert(sizeof(FD3D12DiagnosticBufferData) == 6*sizeof(uint32),
+	"Remember to change UEDiagnosticBuffer layout in the shaders when changing FD3D12DiagnosticBufferData");
+
 class FD3D12CommandListManager : public FD3D12DeviceChild, public FD3D12SingleNodeGPUObject
 {
 public:
@@ -267,6 +283,15 @@ public:
 
 	/** Get the CPU readable data from the breadcrumb data - this data is still valid after the Device is Lost */
 	const void* GetBreadCrumbResourceAddress() const { return BreadCrumbResourceAddress; }
+	const FD3D12DiagnosticBufferData* GetDiagnosticBufferData() const
+	{ 
+		const uint8* Address = BreadCrumbResourceAddress ? reinterpret_cast<const uint8*>(GetBreadCrumbResourceAddress()) + DiagnosticBufferOffset : nullptr;
+		return reinterpret_cast<const FD3D12DiagnosticBufferData*>(Address);
+	}
+	const D3D12_GPU_VIRTUAL_ADDRESS GetDiagnosticBufferGPUAddress() const
+	{ 
+		return BreadCrumbResourceGPUAddress ? BreadCrumbResourceGPUAddress + DiagnosticBufferOffset : 0;
+	}
 
 	void WaitForCommandQueueFlush();
 
@@ -346,7 +371,9 @@ protected:
 	void* BreadCrumbResourceAddress;
 	TRefCountPtr<FD3D12Heap> BreadCrumbHeap;
 	TRefCountPtr<FD3D12Resource> BreadCrumbResource;
-	
+	D3D12_GPU_VIRTUAL_ADDRESS BreadCrumbResourceGPUAddress = {};
+	uint32 DiagnosticBufferOffset = 0;
+
 #if WITH_PROFILEGPU || D3D12_SUBMISSION_GAP_RECORDER
 	uint64 CmdListTimingQueryBatchTokens[2];
 	TArray<FResolvedCmdListExecTime> ResolvedTimingPairs;
diff --git a/Engine/Source/Runtime/D3D12RHI/Private/D3D12RootSignature.cpp b/Engine/Source/Runtime/D3D12RHI/Private/D3D12RootSignature.cpp
--- a/Engine/Source/Runtime/D3D12RHI/Private/D3D12RootSignature.cpp
+++ b/Engine/Source/Runtime/D3D12RHI/Private/D3D12RootSignature.cpp
@@ -22,6 +22,11 @@
 #include "Windows/HideWindowsPlatformTypes.h"
 #endif
 
+// Allows to automatically bind UEDiagnosticBuffer UAV, available to all shaders.
+// If a shader is compiled with diagnostics enabled, it will fail to load/create 
+// unless D3D12_ALLOW_SHADER_DIAGNOSTIC_BUFFER is enabled.
+#define D3D12_ALLOW_SHADER_DIAGNOSTIC_BUFFER 1
+
 namespace
 {
 	// Root parameter costs in DWORDs as described here: https://docs.microsoft.com/en-us/windows/desktop/direct3d12/root-signature-limits
@@ -143,6 +148,10 @@ FD3D12RootSignatureDesc::FD3D12RootSignatureDesc(const FD3D12QuantizedBoundShade
 
 	uint32 BindingSpace = 0; // Default binding space for D3D 11 & 12 shaders
 
+	const bool bUseShaderDiagnosticBuffer = D3D12_ALLOW_SHADER_DIAGNOSTIC_BUFFER
+		&& QBSS.bUseDiagnosticBuffer
+		&& QBSS.RootSignatureType != RS_RayTracingLocal;
+
 #if D3D12_RHI_RAYTRACING
 	if (QBSS.RootSignatureType == RS_RayTracingLocal)
 	{
@@ -313,10 +322,19 @@ FD3D12RootSignatureDesc::FD3D12RootSignatureDesc(const FD3D12QuantizedBoundShade
 		check(RootParameterCount < MaxRootParameters);
 		TableSlots[RootParameterCount].InitAsUnorderedAccessView(0, AGS_DX12_SHADER_INSTRINSICS_SPACE_ID, D3D12_ROOT_DESCRIPTOR_FLAG_DATA_VOLATILE, D3D12_SHADER_VISIBILITY_ALL);
 		RootParameterCount++;
-		RootParametersSize += RootDescriptorTableCost;
+		RootParametersSize += RootDescriptorCost;
 	}
 #endif
 
+	if (bUseShaderDiagnosticBuffer)
+	{
+		check(RootParameterCount < MaxRootParameters);
+		DiagnosticBufferSlot = int8(RootParameterCount);
+		TableSlots[RootParameterCount].InitAsUnorderedAccessView(0, 999, D3D12_ROOT_DESCRIPTOR_FLAG_DATA_VOLATILE, D3D12_SHADER_VISIBILITY_ALL);
+		RootParameterCount++;
+		RootParametersSize += RootDescriptorCost;
+	}
+
 	// Init the desc (warn about the size if necessary).
 #if !NO_LOGGING
 	const uint32 SizeWarningThreshold = 12;
@@ -502,8 +520,11 @@ void FD3D12RootSignature::Init(const FD3D12QuantizedBoundShaderState& InQBSS)
 {
 	// Create a root signature desc from the quantized bound shader state.
 	const D3D12_RESOURCE_BINDING_TIER ResourceBindingTier = GetParentAdapter()->GetResourceBindingTier();
+
 	FD3D12RootSignatureDesc Desc(InQBSS, ResourceBindingTier);
 
+	DiagnosticBufferSlot = Desc.GetDiagnosticBufferSlot();
+
 	uint32 BindingSpace = 0; // Default binding space for D3D 11 & 12 shaders
 
 	if (InQBSS.RootSignatureType == RS_RayTracingGlobal)
diff --git a/Engine/Source/Runtime/D3D12RHI/Private/D3D12RootSignature.h b/Engine/Source/Runtime/D3D12RHI/Private/D3D12RootSignature.h
--- a/Engine/Source/Runtime/D3D12RHI/Private/D3D12RootSignature.h
+++ b/Engine/Source/Runtime/D3D12RHI/Private/D3D12RootSignature.h
@@ -52,9 +52,12 @@ public:
 
 	static constexpr uint32 MaxRootParameters = 32;	// Arbitrary max, increase as needed.
 
+	inline int8 GetDiagnosticBufferSlot() const { return DiagnosticBufferSlot; }
+
 private:
 
 	uint32 RootParametersSize;	// The size of all root parameters in the root signature. Size in DWORDs, the limit is 64.
+	int8 DiagnosticBufferSlot = -1;
 	CD3DX12_ROOT_PARAMETER1 TableSlots[MaxRootParameters];
 	CD3DX12_DESCRIPTOR_RANGE1 DescriptorRanges[MaxRootParameters];
 	CD3DX12_VERSIONED_ROOT_SIGNATURE_DESC RootDesc;
@@ -209,6 +212,9 @@ public:
 	uint32 GetBindSlotOffsetInBytes(uint8 BindSlotIndex) const { check(BindSlotIndex < UE_ARRAY_COUNT(BindSlotOffsetsInDWORDs)); return 4 * BindSlotOffsetsInDWORDs[BindSlotIndex]; }
 	uint32 GetTotalRootSignatureSizeInBytes() const { return 4 * TotalRootSignatureSizeInDWORDs; }
 
+	// Returns root parameter slot for the internal shader diagnostic buffer (used for asserts, etc.) or -1 if not available.
+	inline int8 GetDiagnosticBufferSlot() const { return DiagnosticBufferSlot; }
+
 private:
 	void AnalyzeSignature(const D3D12_VERSIONED_ROOT_SIGNATURE_DESC& Desc, uint32 BindingSpace);
 
@@ -451,6 +457,7 @@ private:
 
 	uint8 BindSlotOffsetsInDWORDs[FD3D12RootSignatureDesc::MaxRootParameters] = {};
 	uint8 TotalRootSignatureSizeInDWORDs = 0;
+	int8 DiagnosticBufferSlot = -1;
 
 	uint8 bHasUAVs : 1;
 	uint8 bHasSRVs : 1;
diff --git a/Engine/Source/Runtime/D3D12RHI/Private/D3D12Shader.h b/Engine/Source/Runtime/D3D12RHI/Private/D3D12Shader.h
--- a/Engine/Source/Runtime/D3D12RHI/Private/D3D12Shader.h
+++ b/Engine/Source/Runtime/D3D12RHI/Private/D3D12Shader.h
@@ -42,6 +42,8 @@ struct FD3D12ShaderData
 
 	/** The static slot associated with the resource table index in ShaderResourceTable. */
 	TArray<FUniformBufferStaticSlot> StaticSlots;
+
+	bool bUsesDiagnosticBuffer = false;
 };
 
 /** This represents a vertex shader that hasn't been combined with a specific declaration to create a bound shader. */
diff --git a/Engine/Source/Runtime/D3D12RHI/Private/D3D12Shaders.cpp b/Engine/Source/Runtime/D3D12RHI/Private/D3D12Shaders.cpp
--- a/Engine/Source/Runtime/D3D12RHI/Private/D3D12Shaders.cpp
+++ b/Engine/Source/Runtime/D3D12RHI/Private/D3D12Shaders.cpp
@@ -48,6 +48,8 @@ static inline bool ReadShaderOptionalData(FShaderCodeReader& InShaderCode, TShad
 	{
 		bOutFoundCodeFeatures = true;
 		OutCodeFeatures = *CodeFeatures;
+
+		OutShader.bUsesDiagnosticBuffer = CodeFeatures->bUsesDiagnosticBuffer;
 	}
 	else
 	{
diff --git a/Engine/Source/Runtime/D3D12RHI/Private/D3D12Util.cpp b/Engine/Source/Runtime/D3D12RHI/Private/D3D12Util.cpp
--- a/Engine/Source/Runtime/D3D12RHI/Private/D3D12Util.cpp
+++ b/Engine/Source/Runtime/D3D12RHI/Private/D3D12Util.cpp
@@ -243,6 +243,46 @@ static bool LogBreadcrumbData(D3D12RHI::FD3DGPUProfiler& GPUProfiler, FD3D12Comm
 
 	UE_LOG(LogD3D12RHI, Error, TEXT("%s"), *gpu_progress);
 
+	const FD3D12DiagnosticBufferData* DiagnosticData = CommandListManager.GetDiagnosticBufferData();
+	if (DiagnosticData && DiagnosticData->Counter)
+	{
+		UE_LOG(LogD3D12RHI, Error, TEXT("[GPUBreadCrumb]\t\tShader assertion failed! ID: 0x%08X"), DiagnosticData->MessageID);
+
+		{
+			const int32* Payload = DiagnosticData->Payload.AsInt;
+			if (Payload[0] < 0 || Payload[1] < 0 || Payload[2] < 0 || Payload[3] < 0)
+			{
+				UE_LOG(LogD3D12RHI, Error,
+					TEXT("[GPUBreadCrumb]\t\tPayload  [int32]: %d %d %d %d"),
+					Payload[0], Payload[1], Payload[2], Payload[3]);
+			}
+		}
+
+		{
+			const uint32* Payload = DiagnosticData->Payload.AsUint;
+			UE_LOG(LogD3D12RHI, Error,
+				TEXT("[GPUBreadCrumb]\t\tPayload [uint32]: %u %u %u %u"),
+				Payload[0], Payload[1], Payload[2], Payload[3]);
+		}
+
+		{
+			const uint32* Payload = DiagnosticData->Payload.AsUint;
+			UE_LOG(LogD3D12RHI, Error,
+				TEXT("[GPUBreadCrumb]\t\tPayload    [hex]: 0x%08X 0x%08X 0x%08X 0x%08X"),
+				Payload[0], Payload[1], Payload[2], Payload[3]);
+		}
+
+		{
+			const float* Payload = DiagnosticData->Payload.AsFloat;
+			UE_LOG(LogD3D12RHI, Error,
+				TEXT("[GPUBreadCrumb]\t\tPayload  [float]: %f %f %f %f"),
+				Payload[0], Payload[1], Payload[2], Payload[3]);
+		}
+	}
+
+	GLog->PanicFlushThreadedLogs();
+	GLog->Flush();
+
 	return true;
 }
 
@@ -870,30 +910,35 @@ void QuantizeBoundShaderState(
 	{
 		FD3D12QuantizedBoundShaderState::InitShaderRegisterCounts(ResourceBindingTier, VertexShader->ResourceCounts, QBSS.RegisterCounts[SV_Vertex]);
 		QBSS.bUseVendorExtension |= UsesVendorExtensionSpace(*VertexShader);
+		QBSS.bUseDiagnosticBuffer = VertexShader->bUsesDiagnosticBuffer;
 	}
 
 	if (const FD3D12MeshShader* const MeshShader = BSS->GetMeshShader())
 	{
 		FD3D12QuantizedBoundShaderState::InitShaderRegisterCounts(ResourceBindingTier, MeshShader->ResourceCounts, QBSS.RegisterCounts[SV_Mesh]);
 		QBSS.bUseVendorExtension |= UsesVendorExtensionSpace(*MeshShader);
+		QBSS.bUseDiagnosticBuffer = MeshShader->bUsesDiagnosticBuffer;
 	}
 
 	if (const FD3D12AmplificationShader* const AmplificationShader = BSS->GetAmplificationShader())
 	{
 		FD3D12QuantizedBoundShaderState::InitShaderRegisterCounts(ResourceBindingTier, AmplificationShader->ResourceCounts, QBSS.RegisterCounts[SV_Amplification]);
 		QBSS.bUseVendorExtension |= UsesVendorExtensionSpace(*AmplificationShader);
+		QBSS.bUseDiagnosticBuffer = AmplificationShader->bUsesDiagnosticBuffer;
 	}
 
 	if (const FD3D12PixelShader* const PixelShader = BSS->GetPixelShader())
 	{
 		FD3D12QuantizedBoundShaderState::InitShaderRegisterCounts(ResourceBindingTier, PixelShader->ResourceCounts, QBSS.RegisterCounts[SV_Pixel], true);
 		QBSS.bUseVendorExtension |= UsesVendorExtensionSpace(*PixelShader);
+		QBSS.bUseDiagnosticBuffer = PixelShader->bUsesDiagnosticBuffer;
 	}
 
 	if (const FD3D12GeometryShader* const GeometryShader = BSS->GetGeometryShader())
 	{
 		FD3D12QuantizedBoundShaderState::InitShaderRegisterCounts(ResourceBindingTier, GeometryShader->ResourceCounts, QBSS.RegisterCounts[SV_Geometry]);
 		QBSS.bUseVendorExtension |= UsesVendorExtensionSpace(*GeometryShader);
+		QBSS.bUseDiagnosticBuffer = GeometryShader->bUsesDiagnosticBuffer;
 	}
 }
 
@@ -923,6 +968,7 @@ void QuantizeBoundShaderState(
 	QuantizeBoundShaderStateCommon(ResourceBindingTier, ComputeShader->ResourceCounts, SV_All, bAllosUAVs, OutQBSS);
 	check(OutQBSS.bAllowIAInputLayout == false); // No access to vertex buffers needed
 	OutQBSS.bUseVendorExtension = (ComputeShader->VendorExtensions.Num() > 0);
+	OutQBSS.bUseDiagnosticBuffer = ComputeShader->bUsesDiagnosticBuffer;
 }
 
 #if D3D12_RHI_RAYTRACING
@@ -989,6 +1035,11 @@ void QuantizeBoundShaderState(
 	default:
 		checkNoEntry(); // Unexpected shader target frequency
 	}
+
+	if (RayTracingShader)
+	{
+		OutQBSS.bUseDiagnosticBuffer = RayTracingShader->bUsesDiagnosticBuffer;
+	}
 }
 #endif // D3D12_RHI_RAYTRACING
 
diff --git a/Engine/Source/Runtime/D3D12RHI/Public/D3D12Util.h b/Engine/Source/Runtime/D3D12RHI/Public/D3D12Util.h
--- a/Engine/Source/Runtime/D3D12RHI/Public/D3D12Util.h
+++ b/Engine/Source/Runtime/D3D12RHI/Public/D3D12Util.h
@@ -92,7 +92,8 @@ struct FD3D12QuantizedBoundShaderState
 	ERTRootSignatureType RootSignatureType = RS_Raster;
 	uint8 bAllowIAInputLayout : 1;
 	uint8 bUseVendorExtension : 1;
-	uint8 Padding : 6;
+	uint8 bUseDiagnosticBuffer : 1;
+	uint8 Padding : 5;
 
 	inline bool operator==(const FD3D12QuantizedBoundShaderState& RHS) const
 	{
diff --git a/Engine/Source/Runtime/RenderCore/Public/ShaderCore.h b/Engine/Source/Runtime/RenderCore/Public/ShaderCore.h
--- a/Engine/Source/Runtime/RenderCore/Public/ShaderCore.h
+++ b/Engine/Source/Runtime/RenderCore/Public/ShaderCore.h
@@ -543,7 +543,8 @@ struct FShaderCodeFeatures
 	// for FindOptionalData() and AddOptionalData()
 	static const uint8 Key = 'x';
 
-	bool bUsesWaveOps;
+	bool bUsesWaveOps = false;
+	bool bUsesDiagnosticBuffer = false;
 };
 
 // if this changes you need to make sure all shaders get invalidated
