D3D12:  Remove debug barriers previously added to resolve issues with indirect args on Nanite

- tested on Reverb P_PerfCave on impacted platforms

#rnx

diff --git a/Engine/Source/Runtime/D3D12RHI/Private/D3D12Commands.cpp b/Engine/Source/Runtime/D3D12RHI/Private/D3D12Commands.cpp
--- a/Engine/Source/Runtime/D3D12RHI/Private/D3D12Commands.cpp
+++ b/Engine/Source/Runtime/D3D12RHI/Private/D3D12Commands.cpp
@@ -153,14 +153,6 @@ void FD3D12CommandContext::RHIDispatchComputeShader(uint32 ThreadGroupCountX, ui
 	numDispatches++;
 	CommandListHandle->Dispatch(ThreadGroupCountX, ThreadGroupCountY, ThreadGroupCountZ);
 
-#if DEBUG_BARRIER_DISPATCHES
-	D3D12_RESOURCE_BARRIER Barrier = {};
-	Barrier.Type = D3D12_RESOURCE_BARRIER_TYPE_UAV;
-	Barrier.Flags = D3D12_RESOURCE_BARRIER_FLAG_NONE;
-	Barrier.UAV.pResource = nullptr;
-	CommandListHandle->ResourceBarrier(1, &Barrier);
-#endif
-
 	DEBUG_EXECUTE_COMMAND_LIST(this);
 }
 
@@ -205,14 +197,6 @@ void FD3D12CommandContext::RHIDispatchIndirectComputeShader(FRHIBuffer* Argument
 		);
 	CommandListHandle.UpdateResidency(Location.GetResource());
 
-#if DEBUG_BARRIER_DISPATCHES
-	D3D12_RESOURCE_BARRIER Barrier = {};
-	Barrier.Type = D3D12_RESOURCE_BARRIER_TYPE_UAV;
-	Barrier.Flags = D3D12_RESOURCE_BARRIER_FLAG_NONE;
-	Barrier.UAV.pResource = nullptr;
-	CommandListHandle->ResourceBarrier(1, &Barrier);
-#endif
-
 	DEBUG_EXECUTE_COMMAND_LIST(this);
 }
 
diff --git a/Engine/Source/Runtime/D3D12RHI/Private/D3D12RHIPrivate.h b/Engine/Source/Runtime/D3D12RHI/Private/D3D12RHIPrivate.h
--- a/Engine/Source/Runtime/D3D12RHI/Private/D3D12RHIPrivate.h
+++ b/Engine/Source/Runtime/D3D12RHI/Private/D3D12RHIPrivate.h
@@ -89,7 +89,6 @@ typedef FD3D12StateCacheBase FD3D12StateCache;
 
 
 #define EXECUTE_DEBUG_COMMAND_LISTS 0
-#define DEBUG_BARRIER_DISPATCHES 1
 #define ENABLE_PLACED_RESOURCES 0 // Disabled due to a couple of NVidia bugs related to placed resources. Works fine on Intel
 #define NAME_OBJECTS !(UE_BUILD_SHIPPING || UE_BUILD_TEST)	// Name objects in all builds except shipping
 #define LOG_PSO_CREATES (0 && STATS)	// Logs Create Pipeline State timings (also requires STATS)
