- refactor d3d12 buffer creation function to split resource creation from initial data copies
- Add option to use copy queue for initial buffer data copy
- refactor command context execute to split build of patch up transition command list
- don't execute patch up command list on direct queue if not needed (can be done in async if not graphics stats)
- correctly sync and forward the copy queue sync points for upload
- execute transition on first use in command list if before state is known and valid (less patch up work to do)
- Use copy queue for instance 'InstanceBuffer' copy of TLAS
- Fix problem with execution of TLAS build on async compute queue (dynamic geometry updates don't run on async compute queue)

diff --git a/Engine/Source/Runtime/D3D12RHI/Private/D3D12Adapter.h b/Engine/Source/Runtime/D3D12RHI/Private/D3D12Adapter.h
--- a/Engine/Source/Runtime/D3D12RHI/Private/D3D12Adapter.h
+++ b/Engine/Source/Runtime/D3D12RHI/Private/D3D12Adapter.h
@@ -311,13 +311,15 @@ public:
 		const TCHAR* Name,
 		D3D12_RESOURCE_FLAGS Flags = D3D12_RESOURCE_FLAG_NONE);
 
-	FD3D12Buffer* CreateRHIBuffer(FRHICommandListImmediate* RHICmdList,
+	FD3D12Buffer* CreateRHIBuffer(
 		const D3D12_RESOURCE_DESC& Desc,
 		uint32 Alignment, uint32 Stride, uint32 Size, uint32 InUsage,
 		ED3D12ResourceStateMode InResourceStateMode,
-		ERHIAccess InResourceState,
-		FRHIResourceCreateInfo& CreateInfo,
-		ID3D12ResourceAllocator* ResourceAllocator);
+		D3D12_RESOURCE_STATES InCreateState,
+		bool bHasInitialData,
+		const FRHIGPUMask& InGPUMask,
+		ID3D12ResourceAllocator* ResourceAllocator,
+		const TCHAR* InDebugName);
 
 	template <typename ObjectType, typename CreationCoreFunction>
 	inline ObjectType* CreateLinkedObject(FRHIGPUMask GPUMask, const CreationCoreFunction& pfnCreationCore)
@@ -395,11 +397,11 @@ protected:
 		uint32 InUsage,
 		ED3D12ResourceStateMode InResourceStateMode,
 		D3D12_RESOURCE_STATES InCreateState,
-		FRHIResourceCreateInfo& CreateInfo,
 		uint32 Alignment,
 		FD3D12Buffer* Buffer,
 		FD3D12ResourceLocation& ResourceLocation,
-		ID3D12ResourceAllocator* ResourceAllocator);
+		ID3D12ResourceAllocator* ResourceAllocator,
+		const TCHAR* InDebugName);
 
 	// Creates default root and execute indirect signatures
 	virtual void CreateCommandSignatures();
diff --git a/Engine/Source/Runtime/D3D12RHI/Private/D3D12Buffer.cpp b/Engine/Source/Runtime/D3D12RHI/Private/D3D12Buffer.cpp
--- a/Engine/Source/Runtime/D3D12RHI/Private/D3D12Buffer.cpp
+++ b/Engine/Source/Runtime/D3D12RHI/Private/D3D12Buffer.cpp
@@ -79,19 +79,19 @@ struct FD3D12RHICommandInitializeBuffer final : public FRHICommand<FD3D12RHIComm
 		FD3D12ResourceLocation::TransferOwnership(SrcResourceLoc, InSrcResourceLoc);
 	}
 
-	void Execute(FRHICommandListBase& /* unused */)
+	void Execute(FRHICommandListBase& CmdList)
 	{
-		ExecuteNoCmdList();
+		FD3D12CommandContext& CommandContext = (FD3D12CommandContext&)(CmdList.IsImmediateAsyncCompute() ? CmdList.GetComputeContext().GetLowestLevelContext() : CmdList.GetContext().GetLowestLevelContext());
+		ExecuteOnCommandContext(CommandContext);
 	}
 
-	void ExecuteNoCmdList()
+	void ExecuteOnCommandContext(FD3D12CommandContext& CommandContext)
 	{
 		for (FD3D12Buffer::FLinkedObjectIterator CurrentBuffer(Buffer); CurrentBuffer; ++CurrentBuffer)
 		{
 			FD3D12Resource* Destination = CurrentBuffer->ResourceLocation.GetResource();
 			FD3D12Device* Device = Destination->GetParentDevice();
 
-			FD3D12CommandContext& CommandContext = Device->GetDefaultCommandContext();
 			FD3D12CommandListHandle& hCommandList = CommandContext.CommandListHandle;
 			// Copy from the temporary upload heap to the default resource
 			{
@@ -141,17 +141,158 @@ struct FD3D12RHICommandInitializeBuffer final : public FRHICommand<FD3D12RHIComm
 	}
 };
 
+
+void FD3D12Buffer::UploadResourceData(class FRHICommandListImmediate* RHICmdList, FResourceArrayInterface* InResourceArray, D3D12_RESOURCE_STATES InDestinationState)
+{
+	check(InResourceArray);
+	check(ResourceLocation.IsValid());
+
+	uint32 BufferSize = GetSize();
+	check(BufferSize == InResourceArray->GetResourceDataSize());
+
+	const bool bIsDynamic = (GetUsage() & BUF_AnyDynamic) ? true : false;
+	if (bIsDynamic)
+	{
+		// Copy directly in mapped data
+		void* MappedUploadData = ResourceLocation.GetMappedBaseAddress(); 
+		FMemory::Memcpy(MappedUploadData, InResourceArray->GetResourceData(), BufferSize);
+	}
+	else
+	{
+		const bool bOnAsyncThread = !IsInRHIThread() && !IsInRenderingThread();
+
+		// Get an upload heap and initialize data
+		FD3D12ResourceLocation SrcResourceLoc(GetParentDevice());
+		void* pData;
+		if (bOnAsyncThread)
+		{
+			const uint32 GPUIdx = SrcResourceLoc.GetParentDevice()->GetGPUIndex();
+			pData = GetParentDevice()->GetParentAdapter()->GetUploadHeapAllocator(GPUIdx).AllocUploadResource(BufferSize, 4u, SrcResourceLoc);
+		}
+		else
+		{
+			pData = SrcResourceLoc.GetParentDevice()->GetDefaultFastAllocator().Allocate(BufferSize, 4UL, &SrcResourceLoc);
+		}
+		check(pData);
+		FMemory::Memcpy(pData, InResourceArray->GetResourceData(), BufferSize);
+
+		if (bOnAsyncThread)
+		{
+			// Need to update buffer content on RHI thread (immediate context) because the buffer can be a
+			// sub-allocation and its backing resource may be in a state incompatible with the copy queue.
+			// TODO:
+			// Create static buffers in COMMON state, rely on state promotion/decay to avoid transition barriers,
+			// and initialize them asynchronously on the copy queue. D3D12 buffers always allow simultaneous acess
+			// so it is legal to write to a region on the copy queue while other non-overlapping regions are
+			// being read on the graphics/compute queue. Currently, d3ddebug throws error for such usage.
+			// Once Microsoft (via Windows update) fix the debug layer, async static buffer initialization should
+			// be done on the copy queue.
+			FD3D12ResourceLocation* SrcResourceLoc_Heap = new FD3D12ResourceLocation(SrcResourceLoc.GetParentDevice());
+			FD3D12ResourceLocation::TransferOwnership(*SrcResourceLoc_Heap, SrcResourceLoc);
+			ENQUEUE_RENDER_COMMAND(CmdD3D12InitializeBuffer)(
+				[this, SrcResourceLoc_Heap, BufferSize, InDestinationState](FRHICommandListImmediate& RHICmdList)
+				{
+					if (RHICmdList.Bypass())
+					{
+						FD3D12RHICommandInitializeBuffer Command(this, *SrcResourceLoc_Heap, BufferSize, InDestinationState);
+						Command.Execute(RHICmdList);
+					}
+					else
+					{
+						new (RHICmdList.AllocCommand<FD3D12RHICommandInitializeBuffer>()) FD3D12RHICommandInitializeBuffer(this, *SrcResourceLoc_Heap, BufferSize, InDestinationState);
+					}
+					delete SrcResourceLoc_Heap;
+				});
+		}
+		else if (!RHICmdList || RHICmdList->Bypass())
+		{
+			// On RHIT or RT (when bypassing), we can access immediate context directly
+			FD3D12RHICommandInitializeBuffer Command(this, SrcResourceLoc, BufferSize, InDestinationState);
+			if (RHICmdList)
+			{
+				Command.Execute(*RHICmdList);
+			}
+			else
+			{
+				FD3D12CommandContext& CommandContext = GetParentDevice()->GetDefaultCommandContext();
+				Command.ExecuteOnCommandContext(CommandContext);
+			}
+		}
+		else
+		{
+			// On RT but not bypassing
+			new (RHICmdList->AllocCommand<FD3D12RHICommandInitializeBuffer>()) FD3D12RHICommandInitializeBuffer(this, SrcResourceLoc, BufferSize, InDestinationState);
+		}
+	}
+
+	// Discard the resource array's contents.
+	InResourceArray->Discard();
+}
+
+
+FD3D12SyncPoint FD3D12Buffer::UploadResourceDataViaCopyQueue(FResourceArrayInterface* InResourceArray)
+{
+	// assume not dynamic and not on async thread (probably fine but untested)
+	check(IsInRHIThread() || IsInRenderingThread());
+	check(!(GetUsage() & BUF_AnyDynamic));
+
+	uint32 BufferSize = GetSize();
+
+	// Get an upload heap and copy the data
+	FD3D12ResourceLocation SrcResourceLoc(GetParentDevice());
+	void* pData = GetParentDevice()->GetDefaultFastAllocator().Allocate(BufferSize, 4UL, &SrcResourceLoc);
+	check(pData);
+	FMemory::Memcpy(pData, InResourceArray->GetResourceData(), BufferSize);
+
+	// Allocate copy queue command list and perform the copy op
+	FD3D12Device* Device = SrcResourceLoc.GetParentDevice();
+	FD3D12CommandAllocatorManager& CommandAllocatorManager = Device->GetTextureStreamingCommandAllocatorManager();
+	FD3D12CommandAllocator* CurrentCommandAllocator = CommandAllocatorManager.ObtainCommandAllocator();
+	FD3D12CommandListHandle hCopyCommandList = Device->GetCopyCommandListManager().ObtainCommandList(*CurrentCommandAllocator);
+
+	// Required for stat tracking ?!?
+	hCopyCommandList.SetCurrentOwningContext(&Device->GetDefaultCommandContext());
+	hCopyCommandList.GetCurrentOwningContext()->numCopies++;
+
+	// Perform actual copy op		
+	hCopyCommandList->CopyBufferRegion(
+		ResourceLocation.GetResource()->GetResource(),
+		ResourceLocation.GetOffsetFromBaseOfResource(),
+		SrcResourceLoc.GetResource()->GetResource(),
+		SrcResourceLoc.GetOffsetFromBaseOfResource(), BufferSize);
+
+	// Residency update needed since it's just been created?
+	hCopyCommandList.UpdateResidency(ResourceLocation.GetResource());
+
+	// Close and kick the command list without waiting for it
+	hCopyCommandList.Close();
+	bool bWaitForCompletion = false;
+	FD3D12SyncPoint CopyQueueSyncPoint = Device->GetCopyCommandListManager().ExecuteCommandListNoCopyQueueSync(hCopyCommandList, bWaitForCompletion);
+
+	// Release command allocator (has the sync point as well)
+	CommandAllocatorManager.ReleaseCommandAllocator(CurrentCommandAllocator);
+
+	// Buffer is now written and ready, so unlock the block
+	ResourceLocation.UnlockPoolData();
+
+	// Discard the resource array's contents.
+	InResourceArray->Discard();
+
+	return CopyQueueSyncPoint;
+}
+
+
 void FD3D12Adapter::AllocateBuffer(FD3D12Device* Device,
 	const D3D12_RESOURCE_DESC& InDesc,
 	uint32 Size,
 	uint32 InUsage,
 	ED3D12ResourceStateMode InResourceStateMode,
 	D3D12_RESOURCE_STATES InCreateState,
-	FRHIResourceCreateInfo& CreateInfo,
 	uint32 Alignment,
 	FD3D12Buffer* Buffer,
 	FD3D12ResourceLocation& ResourceLocation,
-	ID3D12ResourceAllocator* ResourceAllocator)
+	ID3D12ResourceAllocator* ResourceAllocator,
+	const TCHAR* InDebugName)
 {
 	TRACE_CPUPROFILER_EVENT_SCOPE(D3D12RHI::AllocateBuffer);
 
@@ -165,74 +306,58 @@ void FD3D12Adapter::AllocateBuffer(FD3D12Device* Device,
 		check(ResourceAllocator == nullptr);
 		check(InResourceStateMode != ED3D12ResourceStateMode::MultiState);
 		check(InCreateState == D3D12_RESOURCE_STATE_GENERIC_READ);
-		void* pData = GetUploadHeapAllocator(Device->GetGPUIndex()).AllocUploadResource(Size, Alignment, ResourceLocation);
+		GetUploadHeapAllocator(Device->GetGPUIndex()).AllocUploadResource(Size, Alignment, ResourceLocation);
 		check(ResourceLocation.GetSize() == Size);
-
-		if (CreateInfo.ResourceArray)
-		{
-			const void* InitialData = CreateInfo.ResourceArray->GetResourceData();
-
-			check(Size == CreateInfo.ResourceArray->GetResourceDataSize());
-			// Handle initial data
-			FMemory::Memcpy(pData, InitialData, Size);
-		}
 	}
 	else
 	{
 		if (ResourceAllocator)
 		{
-			ResourceAllocator->AllocateResource(Device->GetGPUIndex(), D3D12_HEAP_TYPE_DEFAULT, InDesc, InDesc.Width, Alignment, InResourceStateMode, InCreateState, nullptr, CreateInfo.DebugName, ResourceLocation);
+			ResourceAllocator->AllocateResource(Device->GetGPUIndex(), D3D12_HEAP_TYPE_DEFAULT, InDesc, InDesc.Width, Alignment, InResourceStateMode, InCreateState, nullptr, InDebugName, ResourceLocation);
 		}
 		else
 		{
-			Device->GetDefaultBufferAllocator().AllocDefaultResource(D3D12_HEAP_TYPE_DEFAULT, InDesc, (EBufferUsageFlags)InUsage, InResourceStateMode, InCreateState, ResourceLocation, Alignment, CreateInfo.DebugName);
+			Device->GetDefaultBufferAllocator().AllocDefaultResource(D3D12_HEAP_TYPE_DEFAULT, InDesc, (EBufferUsageFlags)InUsage, InResourceStateMode, InCreateState, ResourceLocation, Alignment, InDebugName);
 		}
 		ResourceLocation.SetOwner(Buffer);
 		check(ResourceLocation.GetSize() == Size);
 	}
 }
 
-FD3D12Buffer* FD3D12Adapter::CreateRHIBuffer(FRHICommandListImmediate* RHICmdList,
+FD3D12Buffer* FD3D12Adapter::CreateRHIBuffer(
 	const D3D12_RESOURCE_DESC& InDesc,
 	uint32 Alignment,
 	uint32 Stride,
 	uint32 Size,
 	uint32 InUsage,
 	ED3D12ResourceStateMode InResourceStateMode,
-	ERHIAccess InResourceState,
-	FRHIResourceCreateInfo& CreateInfo,
-	ID3D12ResourceAllocator* ResourceAllocator)
+	D3D12_RESOURCE_STATES InCreateState,
+	bool bHasInitialData,
+	const FRHIGPUMask& InGPUMask,
+	ID3D12ResourceAllocator* ResourceAllocator,
+	const TCHAR* InDebugName)
 {
 	TRACE_CPUPROFILER_EVENT_SCOPE(D3D12RHI::CreateRHIBuffer);
-
 	SCOPE_CYCLE_COUNTER(STAT_D3D12CreateBufferTime);
 
 	check(InDesc.Width == Size);
 
-	const bool bIsDynamic = (InUsage & BUF_AnyDynamic) ? true : false;
-	const uint32 FirstGPUIndex = CreateInfo.GPUMask.GetFirstIndex();
-		
-	// Does this resource support tracking?
-	const bool bSupportResourceStateTracking = !bIsDynamic && FD3D12DefaultBufferAllocator::IsPlacedResource(InDesc.Flags, InResourceStateMode);
-
-	// Initial state is derived from the InResourceState if it's supports tracking
-	D3D12_HEAP_TYPE HeapType = bIsDynamic ? D3D12_HEAP_TYPE_UPLOAD : D3D12_HEAP_TYPE_DEFAULT;
-	const FD3D12Resource::FD3D12ResourceTypeHelper Type(InDesc, HeapType);
-	const D3D12_RESOURCE_STATES InitialState = bSupportResourceStateTracking ? Type.GetOptimalInitialState(InResourceState, false) : 
-		FD3D12DefaultBufferAllocator::GetDefaultInitialResourceState(HeapType, (EBufferUsageFlags)InUsage, InResourceStateMode);
-
 	FD3D12Buffer* BufferOut = nullptr;
+
+	const bool bIsDynamic = (InUsage & BUF_AnyDynamic) ? true : false;		
 	if (bIsDynamic)
-	{		
+	{
+		const uint32 FirstGPUIndex = InGPUMask.GetFirstIndex();
+
 		FD3D12Buffer* NewBuffer0 = nullptr;
-		BufferOut = CreateLinkedObject<FD3D12Buffer>(CreateInfo.GPUMask, [&](FD3D12Device* Device)
+		BufferOut = CreateLinkedObject<FD3D12Buffer>(InGPUMask, [&](FD3D12Device* Device)
 		{
 			FD3D12Buffer* NewBuffer = new FD3D12Buffer(Device, Size, InUsage, Stride);
 			NewBuffer->BufferAlignment = Alignment;
 
 			if (Device->GetGPUIndex() == FirstGPUIndex)
 			{
-				AllocateBuffer(Device, InDesc, Size, InUsage, InResourceStateMode, InitialState, CreateInfo, Alignment, NewBuffer, NewBuffer->ResourceLocation, ResourceAllocator);
+				AllocateBuffer(Device, InDesc, Size, InUsage, InResourceStateMode, InCreateState, Alignment, NewBuffer, NewBuffer->ResourceLocation, ResourceAllocator, InDebugName);
 				NewBuffer0 = NewBuffer;
 			}
 			else
@@ -246,17 +371,14 @@ FD3D12Buffer* FD3D12Adapter::CreateRHIBuffer(FRHICommandListImmediate* RHICmdLis
 	}
 	else
 	{
-		// Setup the state at which the resource needs to be created - copy dest only supported for placed resources
-		const D3D12_RESOURCE_STATES CreateState = (CreateInfo.ResourceArray && bSupportResourceStateTracking) ? D3D12_RESOURCE_STATE_COPY_DEST : InitialState;
-
-		BufferOut = CreateLinkedObject<FD3D12Buffer>(CreateInfo.GPUMask, [&](FD3D12Device* Device)
+		BufferOut = CreateLinkedObject<FD3D12Buffer>(InGPUMask, [&](FD3D12Device* Device)
 		{
 			FD3D12Buffer* NewBuffer = new FD3D12Buffer(Device, Size, InUsage, Stride);
 			NewBuffer->BufferAlignment = Alignment;
-			AllocateBuffer(Device, InDesc, Size, InUsage, InResourceStateMode, CreateState, CreateInfo, Alignment, NewBuffer, NewBuffer->ResourceLocation, ResourceAllocator);
+			AllocateBuffer(Device, InDesc, Size, InUsage, InResourceStateMode, InCreateState, Alignment, NewBuffer, NewBuffer->ResourceLocation, ResourceAllocator, InDebugName);
 			
-			// Unlock immediately if no initial data
-			if (CreateInfo.ResourceArray == nullptr)
+			// Unlock immediately if there is no initial data
+			if (!bHasInitialData)
 			{
 				NewBuffer->ResourceLocation.UnlockPoolData();
 			}
@@ -265,72 +387,10 @@ FD3D12Buffer* FD3D12Adapter::CreateRHIBuffer(FRHICommandListImmediate* RHICmdLis
 		});
 	}
 
-	if (CreateInfo.ResourceArray)
+	// Don't track transient buffer stats here
+	if (!BufferOut->ResourceLocation.IsTransient())
 	{
-		if (bIsDynamic == false && BufferOut->ResourceLocation.IsValid())
-		{
-			check(Size == CreateInfo.ResourceArray->GetResourceDataSize());
-
-			const bool bOnAsyncThread = !IsInRHIThread() && !IsInRenderingThread();
-
-			// Get an upload heap and initialize data
-			FD3D12ResourceLocation SrcResourceLoc(BufferOut->GetParentDevice());
-			void* pData;
-			if (bOnAsyncThread)
-			{
-				const uint32 GPUIdx = SrcResourceLoc.GetParentDevice()->GetGPUIndex();
-				pData = GetUploadHeapAllocator(GPUIdx).AllocUploadResource(Size, 4u, SrcResourceLoc);
-			}
-			else
-			{
-				pData = SrcResourceLoc.GetParentDevice()->GetDefaultFastAllocator().Allocate(Size, 4UL, &SrcResourceLoc);
-			}
-			check(pData);
-			FMemory::Memcpy(pData, CreateInfo.ResourceArray->GetResourceData(), Size);
-			
-			if (bOnAsyncThread)
-			{
-				// Need to update buffer content on RHI thread (immediate context) because the buffer can be a
-				// sub-allocation and its backing resource may be in a state incompatible with the copy queue.
-				// TODO:
-				// Create static buffers in COMMON state, rely on state promotion/decay to avoid transition barriers,
-				// and initialize them asynchronously on the copy queue. D3D12 buffers always allow simultaneous acess
-				// so it is legal to write to a region on the copy queue while other non-overlapping regions are
-				// being read on the graphics/compute queue. Currently, d3ddebug throws error for such usage.
-				// Once Microsoft (via Windows update) fix the debug layer, async static buffer initialization should
-				// be done on the copy queue.
-				FD3D12ResourceLocation* SrcResourceLoc_Heap = new FD3D12ResourceLocation(SrcResourceLoc.GetParentDevice());
-				FD3D12ResourceLocation::TransferOwnership(*SrcResourceLoc_Heap, SrcResourceLoc);
-				ENQUEUE_RENDER_COMMAND(CmdD3D12InitializeBuffer)(
-					[BufferOut, SrcResourceLoc_Heap, Size, InitialState](FRHICommandListImmediate& RHICmdList)
-				{
-					if (RHICmdList.Bypass())
-					{
-						FD3D12RHICommandInitializeBuffer Command(BufferOut, *SrcResourceLoc_Heap, Size, InitialState);
-						Command.ExecuteNoCmdList();
-					}
-					else
-					{
-						new (RHICmdList.AllocCommand<FD3D12RHICommandInitializeBuffer>()) FD3D12RHICommandInitializeBuffer(BufferOut, *SrcResourceLoc_Heap, Size, InitialState);
-					}
-					delete SrcResourceLoc_Heap;
-				});
-			}
-			else if (!RHICmdList || RHICmdList->Bypass())
-			{
-				// On RHIT or RT (when bypassing), we can access immediate context directly
-				FD3D12RHICommandInitializeBuffer Command(BufferOut, SrcResourceLoc, Size, InitialState);
-				Command.ExecuteNoCmdList();
-			}
-			else
-			{
-				// On RT but not bypassing
-				new (RHICmdList->AllocCommand<FD3D12RHICommandInitializeBuffer>()) FD3D12RHICommandInitializeBuffer(BufferOut, SrcResourceLoc, Size, InitialState);
-			}
-		}
-
-		// Discard the resource array's contents.
-		CreateInfo.ResourceArray->Discard();
+		UpdateBufferStats(GetBufferStats(InUsage), BufferOut->ResourceLocation.GetSize());
 	}
 
 	return BufferOut;
@@ -438,7 +498,7 @@ FBufferRHIRef FD3D12DynamicRHI::CreateBuffer(FRHICommandListImmediate* RHICmdLis
 	return CreateD3D12Buffer(RHICmdList, Size, Usage, Stride, InResourceState, CreateInfo, ResourceAllocator);
 }
 
-FD3D12Buffer* FD3D12DynamicRHI::CreateD3D12Buffer(class FRHICommandListImmediate* RHICmdList, uint32 Size, EBufferUsageFlags Usage, uint32 Stride, ERHIAccess ResourceState, FRHIResourceCreateInfo& CreateInfo, ID3D12ResourceAllocator* ResourceAllocator)
+FD3D12Buffer* FD3D12DynamicRHI::CreateD3D12Buffer(class FRHICommandListImmediate* RHICmdList, uint32 Size, EBufferUsageFlags Usage, uint32 Stride, ERHIAccess InResourceState, FRHIResourceCreateInfo& CreateInfo, ID3D12ResourceAllocator* ResourceAllocator)
 {
 	D3D12_RESOURCE_DESC Desc;
 	uint32 Alignment;
@@ -448,11 +508,28 @@ FD3D12Buffer* FD3D12DynamicRHI::CreateD3D12Buffer(class FRHICommandListImmediate
 		? ED3D12ResourceStateMode::SingleState 
 		: ED3D12ResourceStateMode::Default;
 
-	FD3D12Buffer* Buffer = GetAdapter().CreateRHIBuffer(RHICmdList, Desc, Alignment, Stride, Size, Usage, StateMode, ResourceState, CreateInfo, ResourceAllocator);
-	if (Buffer->ResourceLocation.IsTransient())
+	// Does this resource support tracking?
+	const bool bIsDynamic = (Usage & BUF_AnyDynamic) ? true : false;
+	const bool bSupportResourceStateTracking = !bIsDynamic && FD3D12DefaultBufferAllocator::IsPlacedResource(Desc.Flags, StateMode);
+
+	// Initial state is derived from the InResourceState if it's supports tracking
+	D3D12_HEAP_TYPE HeapType = bIsDynamic ? D3D12_HEAP_TYPE_UPLOAD : D3D12_HEAP_TYPE_DEFAULT;
+	const FD3D12Resource::FD3D12ResourceTypeHelper Type(Desc, HeapType);
+	D3D12_RESOURCE_STATES DesiredState = bSupportResourceStateTracking ? Type.GetOptimalInitialState(InResourceState, false) :
+		FD3D12DefaultBufferAllocator::GetDefaultInitialResourceState(HeapType, (EBufferUsageFlags)Usage, StateMode);
+
+	// Setup the state at which the resource needs to be created - copy dest only supported for placed resources
+	D3D12_RESOURCE_STATES CreateState = (CreateInfo.ResourceArray && bSupportResourceStateTracking) ? D3D12_RESOURCE_STATE_COPY_DEST : DesiredState;
+	bool bHasInitialData = CreateInfo.ResourceArray != nullptr;
+
+	FD3D12SyncPoint CopyQueueSyncPoint;
+	FD3D12Buffer* Buffer = GetAdapter().CreateRHIBuffer(Desc, Alignment, Stride, Size, Usage, StateMode, CreateState, bHasInitialData, CreateInfo.GPUMask, ResourceAllocator, CreateInfo.DebugName);
+	check(Buffer->ResourceLocation.IsValid());
+
+	// Copy the resource data if available 
+	if (bHasInitialData)
 	{
-		// TODO: this should ideally be set in platform-independent code, since this tracking is for the high level
-		Buffer->SetCommitted(false);
+		Buffer->UploadResourceData(RHICmdList, CreateInfo.ResourceArray, DesiredState);
 	}
 
 	return Buffer;
diff --git a/Engine/Source/Runtime/D3D12RHI/Private/D3D12CommandContext.cpp b/Engine/Source/Runtime/D3D12RHI/Private/D3D12CommandContext.cpp
--- a/Engine/Source/Runtime/D3D12RHI/Private/D3D12CommandContext.cpp
+++ b/Engine/Source/Runtime/D3D12RHI/Private/D3D12CommandContext.cpp
@@ -394,7 +394,7 @@ FD3D12CommandListHandle FD3D12CommandContext::FlushCommands(bool WaitForCompleti
 				// which could insert a flush in between pending command lists
 				check(IsInRenderingThread());
 			}
-			GetCommandListManager().ExecuteCommandLists(Device->PendingCommandLists, WaitForCompletion);
+			GetCommandListManager().ExecuteCommandLists(Device->PendingCommandLists, CopyQueueSyncPoint, WaitForCompletion);
 			Device->PendingCommandLists.Reset();
 		}
 		else
@@ -403,9 +403,12 @@ FD3D12CommandListHandle FD3D12CommandContext::FlushCommands(bool WaitForCompleti
 			check(bIsCommandListOpen);
 
 			// Just submit the current command list
-			CommandListHandle.Execute(WaitForCompletion);
+			CommandListHandle.Execute(CopyQueueSyncPoint, WaitForCompletion);
 		}
 
+		// reset sync point again
+		CopyQueueSyncPoint = FD3D12SyncPoint();
+
 		if (bIsCommandListOpen)
 		{
 			// Get a new command list to replace the one we submitted for execution. 
@@ -890,7 +893,10 @@ public:
 				const bool bHasPendingWork = Device->PendingCommandLists.Num() > 0;
 				if (bFlush && bHasPendingWork)
 				{
-					Device->GetCommandListManager().ExecuteCommandLists(Device->PendingCommandLists);
+					// TODO: find out copy queue sync point of all parallel translate jobs?
+					FD3D12SyncPoint EmptyCopyQueueSyncPoint;
+					bool bWaitForCompletion = false;
+					Device->GetCommandListManager().ExecuteCommandLists(Device->PendingCommandLists, EmptyCopyQueueSyncPoint, bWaitForCompletion);
 					Device->PendingCommandLists.Reset();
 				}
 
diff --git a/Engine/Source/Runtime/D3D12RHI/Private/D3D12CommandContext.h b/Engine/Source/Runtime/D3D12RHI/Private/D3D12CommandContext.h
--- a/Engine/Source/Runtime/D3D12RHI/Private/D3D12CommandContext.h
+++ b/Engine/Source/Runtime/D3D12RHI/Private/D3D12CommandContext.h
@@ -130,7 +130,10 @@ public:
 	FD3D12CommandListHandle CommandListHandle;
 	FD3D12CommandAllocator* CommandAllocator;
 	FD3D12CommandAllocatorManager CommandAllocatorManager;
-	
+
+	// Sync point with copy queue which needs to be checked before kicking this command lists
+	FD3D12SyncPoint CopyQueueSyncPoint;
+
 	// Current GPU event stack
 	TArray<uint32> GPUEventStack;
 
diff --git a/Engine/Source/Runtime/D3D12RHI/Private/D3D12CommandList.cpp b/Engine/Source/Runtime/D3D12RHI/Private/D3D12CommandList.cpp
--- a/Engine/Source/Runtime/D3D12RHI/Private/D3D12CommandList.cpp
+++ b/Engine/Source/Runtime/D3D12RHI/Private/D3D12CommandList.cpp
@@ -293,10 +293,10 @@ void FD3D12CommandListHandle::FD3D12CommandListData::FCommandListResourceState::
 	ResourceStates.Empty();
 }
 
-void FD3D12CommandListHandle::Execute(bool WaitForCompletion)
+void FD3D12CommandListHandle::Execute(FD3D12SyncPoint& CopyQueueSyncPoint, bool WaitForCompletion)
 {
 	check(CommandListData);
-	CommandListData->CommandListManager->ExecuteCommandList(*this, WaitForCompletion);
+	CommandListData->CommandListManager->ExecuteCommandList(*this, CopyQueueSyncPoint, WaitForCompletion);
 }
 
 FD3D12CommandAllocator::FD3D12CommandAllocator(ID3D12Device* InDevice, const D3D12_COMMAND_LIST_TYPE& InType)
diff --git a/Engine/Source/Runtime/D3D12RHI/Private/D3D12CommandList.h b/Engine/Source/Runtime/D3D12RHI/Private/D3D12CommandList.h
--- a/Engine/Source/Runtime/D3D12RHI/Private/D3D12CommandList.h
+++ b/Engine/Source/Runtime/D3D12RHI/Private/D3D12CommandList.h
@@ -378,7 +378,7 @@ public:
 
 	void Create(FD3D12Device* ParentDevice, D3D12_COMMAND_LIST_TYPE CommandListType, FD3D12CommandAllocator& CommandAllocator, FD3D12CommandListManager* InCommandListManager);
 
-	void Execute(bool WaitForCompletion = false);
+	void Execute(FD3D12SyncPoint& CopyQueueSyncPoint, bool WaitForCompletion = false);
 	
 	void Close()
 	{
diff --git a/Engine/Source/Runtime/D3D12RHI/Private/D3D12Commands.cpp b/Engine/Source/Runtime/D3D12RHI/Private/D3D12Commands.cpp
--- a/Engine/Source/Runtime/D3D12RHI/Private/D3D12Commands.cpp
+++ b/Engine/Source/Runtime/D3D12RHI/Private/D3D12Commands.cpp
@@ -2353,7 +2353,8 @@ void FD3D12CommandContext::RHIBroadcastTemporalEffect(const FName& InEffectName,
 	}
 	hCopyCommandList.Close();
 
-	CopyCommandListManager.ExecuteCommandList(hCopyCommandList);
+	bool bWaitForCompletion = false;
+	CopyCommandListManager.ExecuteCommandListNoCopyQueueSync(hCopyCommandList, bWaitForCompletion);
 	CopyCommandAllocatorManager.ReleaseCommandAllocator(CopyCommandAllocator);
 
 	// Signal again once the copy queue copy is complete.
diff --git a/Engine/Source/Runtime/D3D12RHI/Private/D3D12DirectCommandListManager.cpp b/Engine/Source/Runtime/D3D12RHI/Private/D3D12DirectCommandListManager.cpp
--- a/Engine/Source/Runtime/D3D12RHI/Private/D3D12DirectCommandListManager.cpp
+++ b/Engine/Source/Runtime/D3D12RHI/Private/D3D12DirectCommandListManager.cpp
@@ -324,7 +324,7 @@ void FD3D12CommandAllocatorManager::ReleaseCommandAllocator(FD3D12CommandAllocat
 FD3D12CommandListManager::FD3D12CommandListManager(FD3D12Device* InParent, D3D12_COMMAND_LIST_TYPE InCommandListType, ED3D12CommandQueueType InQueueType)
 	: FD3D12DeviceChild(InParent)
 	, FD3D12SingleNodeGPUObject(InParent->GetGPUMask())
-	, ResourceBarrierCommandAllocatorManager(InParent, D3D12_COMMAND_LIST_TYPE_DIRECT)
+	, ResourceBarrierCommandAllocatorManager(InParent, InCommandListType)
 	, ResourceBarrierCommandAllocator(nullptr)
 	, CommandListFence(nullptr)
 	, CommandListType(InCommandListType)
@@ -530,12 +530,18 @@ void FD3D12CommandListManager::ReleaseCommandList(FD3D12CommandListHandle& hList
 	ReadyLists.Enqueue(hList);
 }
 
-void FD3D12CommandListManager::ExecuteCommandList(FD3D12CommandListHandle& hList, bool WaitForCompletion)
+FD3D12SyncPoint FD3D12CommandListManager::ExecuteCommandListNoCopyQueueSync(FD3D12CommandListHandle& hList, bool WaitForCompletion)
+{
+	FD3D12SyncPoint EmptySyncPoint;
+	return ExecuteCommandList(hList, EmptySyncPoint, WaitForCompletion);
+}
+
+FD3D12SyncPoint FD3D12CommandListManager::ExecuteCommandList(FD3D12CommandListHandle& hList, FD3D12SyncPoint& CopyQueueSyncPoint, bool WaitForCompletion)
 {
 	TArray<FD3D12CommandListHandle> Lists;
 	Lists.Add(hList);
 
-	ExecuteCommandLists(Lists, WaitForCompletion);
+	return ExecuteCommandLists(Lists, CopyQueueSyncPoint, WaitForCompletion);
 }
 
 uint64 FD3D12CommandListManager::ExecuteAndIncrementFence(FD3D12CommandListPayload& Payload, FD3D12Fence &Fence)
@@ -604,11 +610,18 @@ void FD3D12CommandListManager::WaitOnExecuteTask()
 	}
 }
 
-void FD3D12CommandListManager::ExecuteCommandLists(TArray<FD3D12CommandListHandle>& Lists, bool WaitForCompletion)
+FD3D12SyncPoint FD3D12CommandListManager::ExecuteCommandLists(TArray<FD3D12CommandListHandle>& Lists, FD3D12SyncPoint& CopyQueueSyncPoint, bool WaitForCompletion)
 {
 	// Still has a pending execute task, then make sure the current one is finished before executing a new command list set
 	WaitOnExecuteTask();
 
+	// Need to sync with copy queue before submit?
+	if (CopyQueueSyncPoint.IsValid())
+	{
+		// Command queue should wait for copy queue to finish
+		CopyQueueSyncPoint.GPUWait(GetQueueType());
+	}
+
 	check(ExecuteCommandListHandles.Num() == 0);
 
 	// Do we want to kick via a task - only for direct/graphics queue for now
@@ -620,33 +633,51 @@ void FD3D12CommandListManager::ExecuteCommandLists(TArray<FD3D12CommandListHandl
 		// Increment the pending fence value so all object can be corrected fenced again future pending signal
 		CommandListFence->AdvancePendingFenceValue();
 
+		// Setup the future sync point already on the pending fence value
+		FD3D12SyncPoint PendingSyncPoint(CommandListFence, CommandListFence->GetCurrentFence());
+
 		ExecuteTask = FFunctionGraphTask::CreateAndDispatchWhenReady(
 			[this]()
 			{	
-				ExecuteCommandListInternal(ExecuteCommandListHandles, false);
+				// Empty sync point because the copy queue sync point is only needed for async compute command list which can't execute
+				// it's barriers on that queue (we know this is on direct queue already so synced already above
+				FD3D12SyncPoint EmptySyncPoint;
+				ExecuteCommandListInternal(ExecuteCommandListHandles, EmptySyncPoint, false);
 				ExecuteCommandListHandles.Reset();
 			}, TStatId(), nullptr, ENamedThreads::AnyThread);
+
+		// Return pending sync point which can be used to sync with already if needed
+		return PendingSyncPoint;
 	}
 	else
 	{
-		ExecuteCommandListInternal(Lists, WaitForCompletion);
+		return ExecuteCommandListInternal(Lists, CopyQueueSyncPoint, WaitForCompletion);
 	}
 }
 
 
-void FD3D12CommandListManager::ExecuteCommandListInternal(TArray<FD3D12CommandListHandle>& Lists, bool WaitForCompletion)
+FD3D12SyncPoint FD3D12CommandListManager::ExecuteCommandListInternal(TArray<FD3D12CommandListHandle>& Lists, FD3D12SyncPoint& CopyQueueSyncPoint, bool WaitForCompletion)
 {
 	SCOPE_CYCLE_COUNTER(STAT_D3D12ExecuteCommandListTime);
 	check(CommandListFence);
 
-	bool NeedsResourceBarriers = false;
-	for (int32 i = 0; i < Lists.Num(); i++)
+	// Collect all the barrier desc data and find out if there are any graphics state transitions
+	bool bNeedsResourceBarriers = false;
+	bool bNeedsGraphicStateResourceBarriers = false;
+	TArray<FBarrierDescInfo> BarrierDescInfos;
+	BarrierDescInfos.SetNum(Lists.Num());
 	{
-		FD3D12CommandListHandle& commandList = Lists[i];
-		if (commandList.PendingResourceBarriers().Num() > 0)
+		FScopeLock Lock(&ResourceStateCS);
+
+		for (int32 CommandListIndex = 0; CommandListIndex < Lists.Num(); CommandListIndex++)
 		{
-			NeedsResourceBarriers = true;
-			break;
+			FBarrierDescInfo& BarrierDescInfo = BarrierDescInfos[CommandListIndex];
+			uint32 BarrierCount = CollectInitialResourceBarriersAndUpdateFinalState(Lists[CommandListIndex], BarrierDescInfo);
+			if (BarrierCount > 0)
+			{
+				bNeedsResourceBarriers = true;
+				bNeedsGraphicStateResourceBarriers = bNeedsGraphicStateResourceBarriers | BarrierDescInfo.bHasGraphicStates;
+			}
 		}
 	}
 
@@ -655,9 +686,8 @@ void FD3D12CommandListManager::ExecuteCommandListInternal(TArray<FD3D12CommandLi
 	FD3D12SyncPoint SyncPoint;
 	FD3D12SyncPoint BarrierSyncPoint;
 
-	FD3D12CommandListManager& DirectCommandListManager = GetParentDevice()->GetCommandListManager();
-	FD3D12Fence& DirectFence = DirectCommandListManager.GetFence();
-	checkf(DirectFence.GetGPUMask() == GetGPUMask(), TEXT("Fence GPU masks does not fit with the command list mask!"));
+	FD3D12CommandListManager* BarrierCommandListManager = bNeedsGraphicStateResourceBarriers ? &GetParentDevice()->GetCommandListManager() : this;
+	bool bBarriersUseDifferentQueue = (BarrierCommandListManager != this);
 
 	int32 commandListIndex = 0;
 	int32 barrierCommandListIndex = 0;
@@ -665,12 +695,14 @@ void FD3D12CommandListManager::ExecuteCommandListInternal(TArray<FD3D12CommandLi
 	// Close the resource barrier lists, get the raw command list pointers, and enqueue the command list handles
 	// Note: All command lists will share the same fence
 	FD3D12CommandListPayload CurrentCommandListPayload;
-	FD3D12CommandListPayload ComputeBarrierPayload;
 
 	check(Lists.Num() <= FD3D12CommandListPayload::MaxCommandListsPerPayload);
 	FD3D12CommandListHandle BarrierCommandList[128];
-	if (NeedsResourceBarriers)
+	if (bNeedsResourceBarriers)
 	{
+		// should not be copy queue
+		check(GetQueueType() != ED3D12CommandQueueType::Copy);
+
 #if UE_BUILD_DEBUG	
 		if (ResourceStateCS.TryLock())
 		{
@@ -685,39 +717,54 @@ void FD3D12CommandListManager::ExecuteCommandListInternal(TArray<FD3D12CommandLi
 		}
 #endif
 		
-		// If not direct queue, then make sure the direct queue is done executing commands 
-		// before trying to use the barrier command list from it
-		if (CommandListType != D3D12_COMMAND_LIST_TYPE_DIRECT)
+		// Make sure the queue is done executing commands before trying to use the barrier command list from it
+		BarrierCommandListManager->WaitOnExecuteTask();
+
+		// If the barrier command list manager is different then this command list then it might need another sync point
+		// with the copy queue (didn't happen yet)
+		if (bBarriersUseDifferentQueue && CopyQueueSyncPoint.IsValid())
 		{
-			DirectCommandListManager.WaitOnExecuteTask();
+			CopyQueueSyncPoint.GPUWait(BarrierCommandListManager->GetQueueType());
 		}
 
-		FScopeLock Lock(&ResourceStateCS);
-
-		for (int32 i = 0; i < Lists.Num(); i++)
+		// Kick all the command lists
+		for (int32 CommandListIndex = 0; CommandListIndex < Lists.Num(); CommandListIndex++)
 		{
-			FD3D12CommandListHandle& commandList = Lists[i];
+			FD3D12CommandListHandle& commandList = Lists[CommandListIndex];
+			FBarrierDescInfo& BarrierDescInfo = BarrierDescInfos[CommandListIndex];
 
-			FD3D12CommandListHandle barrierCommandList ={};
-			// Async compute cannot perform all resource transitions, and so it uses the direct context
-			const uint32 numBarriers = DirectCommandListManager.GetResourceBarrierCommandList(commandList, barrierCommandList);
-			if (numBarriers)
+			// Got any barriers to execute?
+			if (BarrierDescInfo.BarrierDescs.Num() > 0 || BarrierDescInfo.BackBufferBarrierDescs.Num() > 0)
 			{
+				FD3D12CommandListHandle barrierCommandList = {};
+
+				// Async compute cannot perform all resource transitions, and so it uses the direct context
+				const uint32 numBarriers = BarrierCommandListManager->GetResourceBarrierCommandList(BarrierDescInfo, barrierCommandList);
+				check(numBarriers > 0);
+
 				// TODO: Unnecessary assignment here, but fixing this will require refactoring GetResourceBarrierCommandList
 				BarrierCommandList[barrierCommandListIndex] = barrierCommandList;
 				barrierCommandListIndex++;
 
 				barrierCommandList.Close();
 
-				if (CommandListType == D3D12_COMMAND_LIST_TYPE_COMPUTE)
+				// Kick on different queue and add fence (ideally this could be done per command list and chec if it has grahics transitions
+				// but then command list also needs to be build on that manager and managing all those command list handles become more complicated)
+				// Not sure if it's worth it yet
+				if (bBarriersUseDifferentQueue)
 				{
+					FD3D12Fence& BarrierFence = BarrierCommandListManager->GetFence();
+					checkf(BarrierFence.GetGPUMask() == GetGPUMask(), TEXT("Fence GPU masks does not fit with the command list mask!"));
+
+					FD3D12CommandListPayload ComputeBarrierPayload;
 					ComputeBarrierPayload.Reset();
 					ComputeBarrierPayload.Append(barrierCommandList.CommandList(), &barrierCommandList.GetResidencySet());
-					BarrierFenceValue = DirectCommandListManager.ExecuteAndIncrementFence(ComputeBarrierPayload, DirectFence);
-					DirectFence.GpuWait(QueueType, BarrierFenceValue);
+					BarrierFenceValue = BarrierCommandListManager->ExecuteAndIncrementFence(ComputeBarrierPayload, BarrierFence);
+					BarrierFence.GpuWait(QueueType, BarrierFenceValue);
 				}
 				else
 				{
+					check(BarrierCommandListManager == this);
 					CurrentCommandListPayload.Append(barrierCommandList.CommandList(), &barrierCommandList.GetResidencySet());
 				}
 			}
@@ -725,11 +772,13 @@ void FD3D12CommandListManager::ExecuteCommandListInternal(TArray<FD3D12CommandLi
 			CurrentCommandListPayload.Append(commandList.CommandList(), &commandList.GetResidencySet());
 			commandList.LogResourceBarriers();
 		}
+
 		SignaledFenceValue = ExecuteAndIncrementFence(CurrentCommandListPayload, *CommandListFence);
 		SyncPoint = FD3D12SyncPoint(CommandListFence, SignaledFenceValue);
-		if (CommandListType == D3D12_COMMAND_LIST_TYPE_COMPUTE)
+
+		if (bBarriersUseDifferentQueue)
 		{
-			BarrierSyncPoint = FD3D12SyncPoint(&DirectFence, BarrierFenceValue);
+			BarrierSyncPoint = FD3D12SyncPoint(&BarrierCommandListManager->GetFence(), BarrierFenceValue);
 		}
 		else
 		{
@@ -744,7 +793,6 @@ void FD3D12CommandListManager::ExecuteCommandListInternal(TArray<FD3D12CommandLi
 			Lists[i].LogResourceBarriers();
 		}
 		SignaledFenceValue = ExecuteAndIncrementFence(CurrentCommandListPayload, *CommandListFence);
-		check(CommandListType != D3D12_COMMAND_LIST_TYPE_COMPUTE);
 		SyncPoint = FD3D12SyncPoint(CommandListFence, SignaledFenceValue);
 		BarrierSyncPoint = SyncPoint;
 	}
@@ -766,7 +814,7 @@ void FD3D12CommandListManager::ExecuteCommandListInternal(TArray<FD3D12CommandLi
 		// Set a sync point on the command list so we know when it's current generation is complete on the GPU, then release it so it can be reused later.
 		// Note this also updates the command list's command allocator
 		commandList.SetSyncPoint(BarrierSyncPoint);
-		DirectCommandListManager.ReleaseCommandList(commandList);
+		BarrierCommandListManager->ReleaseCommandList(commandList);
 	}
 
 	if (WaitForCompletion)
@@ -774,6 +822,8 @@ void FD3D12CommandListManager::ExecuteCommandListInternal(TArray<FD3D12CommandLi
 		CommandListFence->WaitForFence(SignaledFenceValue);
 		check(SyncPoint.IsComplete());
 	}
+
+	return SyncPoint;
 }
 
 void FD3D12CommandListManager::ReleaseResourceBarrierCommandListAllocator()
@@ -889,19 +939,18 @@ void FD3D12CommandListManager::FlushPendingTimingPairs(bool bBlock)
 #endif
 }
 
-uint32 FD3D12CommandListManager::GetResourceBarrierCommandList(FD3D12CommandListHandle& hList, FD3D12CommandListHandle& hResourceBarrierList)
+uint32 FD3D12CommandListManager::CollectInitialResourceBarriersAndUpdateFinalState(FD3D12CommandListHandle& InCommandListHandle, FBarrierDescInfo& BarrierDescInfo)
 {
-	TRACE_CPUPROFILER_EVENT_SCOPE(GetResourceBarrierCommandList);
+	TRACE_CPUPROFILER_EVENT_SCOPE(CollectInitialResourceBarriersAndUpdateFinalState);
 
-	TArray<FD3D12PendingResourceBarrier>& PendingResourceBarriers = hList.PendingResourceBarriers();
+	BarrierDescInfo.bHasGraphicStates = false;
+
+	TArray<FD3D12PendingResourceBarrier>& PendingResourceBarriers = InCommandListHandle.PendingResourceBarriers();
 	const uint32 NumPendingResourceBarriers = PendingResourceBarriers.Num();
-	if (NumPendingResourceBarriers)
+	if (NumPendingResourceBarriers > 0)
 	{
 		// Reserve space for the descs
-		TArray<D3D12_RESOURCE_BARRIER> BarrierDescs;
-		BarrierDescs.Reserve(NumPendingResourceBarriers);
-
-		TArray<D3D12_RESOURCE_BARRIER, TInlineAllocator<2>> BackBufferBarrierDescs;
+		BarrierDescInfo.BarrierDescs.Reserve(NumPendingResourceBarriers);
 
 		// Fill out the descs
 		D3D12_RESOURCE_BARRIER Desc = {};
@@ -930,92 +979,100 @@ uint32 FD3D12CommandListManager::GetResourceBarrierCommandList(FD3D12CommandList
 				Desc.Transition.StateBefore = Before;
 				Desc.Transition.StateAfter = After;
 
+				if (IsDirectQueueExclusiveD3D12State(Before) || IsDirectQueueExclusiveD3D12State(After))
+				{
+					BarrierDescInfo.bHasGraphicStates = true;
+				}
+
 				// Add the desc
 				if (PRB.Resource->IsBackBuffer() && (After & BackBufferBarrierWriteTransitionTargets))
 				{
-					BackBufferBarrierDescs.Add(Desc);
+					BarrierDescInfo.BackBufferBarrierDescs.Add(Desc);
 				}
 				else
 				{
-					BarrierDescs.Add(Desc);
+					BarrierDescInfo.BarrierDescs.Add(Desc);
 				}
 			}
 
 			// Update the state to the what it will be after hList executes
-			const D3D12_RESOURCE_STATES CommandListState = hList.GetResourceState(PRB.Resource).GetSubresourceState(Desc.Transition.Subresource);
+			const D3D12_RESOURCE_STATES CommandListState = InCommandListHandle.GetResourceState(PRB.Resource).GetSubresourceState(Desc.Transition.Subresource);
 			const D3D12_RESOURCE_STATES LastState = (CommandListState != D3D12_RESOURCE_STATE_TBD) ? CommandListState : After;
-			
+
 			if (Before != LastState)
 			{
 				ResourceState.SetSubresourceState(Desc.Transition.Subresource, LastState);
 			}
+
+#if ENABLE_RESIDENCY_MANAGEMENT
+			BarrierDescInfo.ResidencyHandles.Add(PRB.Resource->GetResidencyHandle());
+#endif // ENABLE_RESIDENCY_MANAGEMENT
 		}
+	}
 
-		const uint32 BarrierCount = BarrierDescs.Num() + BackBufferBarrierDescs.Num();
+	return BarrierDescInfo.BarrierDescs.Num() + BarrierDescInfo.BackBufferBarrierDescs.Num();
+}
+
+uint32 FD3D12CommandListManager::GetResourceBarrierCommandList(FBarrierDescInfo& InBarrierDescInfo, FD3D12CommandListHandle& hResourceBarrierList)
+{
+	TRACE_CPUPROFILER_EVENT_SCOPE(GetResourceBarrierCommandList);
 
-		if (BarrierCount > 0)
+	const uint32 BarrierCount = InBarrierDescInfo.BarrierDescs.Num() + InBarrierDescInfo.BackBufferBarrierDescs.Num();
+	if (BarrierCount > 0)
+	{
+		// Get a new resource barrier command allocator if we don't already have one.
+		if (ResourceBarrierCommandAllocator == nullptr)
 		{
-			// Get a new resource barrier command allocator if we don't already have one.
-			if (ResourceBarrierCommandAllocator == nullptr)
-			{
-				ResourceBarrierCommandAllocator = ResourceBarrierCommandAllocatorManager.ObtainCommandAllocator();
-			}
+			ResourceBarrierCommandAllocator = ResourceBarrierCommandAllocatorManager.ObtainCommandAllocator();
+		}
 
-			hResourceBarrierList = ObtainCommandList(*ResourceBarrierCommandAllocator, BackBufferBarrierDescs.Num() > 0);
+		hResourceBarrierList = ObtainCommandList(*ResourceBarrierCommandAllocator, InBarrierDescInfo.BackBufferBarrierDescs.Num() > 0);
 
 #if ENABLE_RESIDENCY_MANAGEMENT
-			//TODO: Update the logic so that this loop can occur above!
-			for (uint32 i = 0; i < NumPendingResourceBarriers; ++i)
-			{
-				const FD3D12PendingResourceBarrier& PRB = PendingResourceBarriers[i];
-				hResourceBarrierList.UpdateResidency(PRB.Resource);
-			}
+		hResourceBarrierList.UpdateResidency(InBarrierDescInfo.ResidencyHandles.GetData(), InBarrierDescInfo.ResidencyHandles.Num());
 #endif // #if ENABLE_RESIDENCY_MANAGEMENT
+
 #if DEBUG_RESOURCE_STATES
-			LogResourceBarriers(BarrierDescs.Num(), BarrierDescs.GetData(), hResourceBarrierList.CommandList());
-			LogResourceBarriers(BackBufferBarrierDescs.Num(), BackBufferBarrierDescs.GetData(), BackBufferBarrierDescs.CommandList());
+		LogResourceBarriers(InBarrierDescInfo.BarrierDescs.Num(), InBarrierDescInfo.BarrierDescs.GetData(), hResourceBarrierList.CommandList());
+		LogResourceBarriers(InBarrierDescInfo.BackBufferBarrierDescs.Num(), InBarrierDescInfo.BackBufferBarrierDescs.GetData(), hResourceBarrierList.CommandList());
 #endif // #if DEBUG_RESOURCE_STATES
-			const int32 BarrierBatchMax = FD3D12DynamicRHI::GetResourceBarrierBatchSizeLimit();
+
+		const int32 BarrierBatchMax = FD3D12DynamicRHI::GetResourceBarrierBatchSizeLimit();
 
 #if PLATFORM_USE_BACKBUFFER_WRITE_TRANSITION_TRACKING
-			if (BackBufferBarrierDescs.Num())
+		if (InBarrierDescInfo.BackBufferBarrierDescs.Num())
+		{
+			FD3D12ScopedTimedIntervalQuery BarrierScopeTimer(GetParentDevice()->GetBackBufferWriteBarrierTracker(), hResourceBarrierList.GraphicsCommandList());
+			if (InBarrierDescInfo.BackBufferBarrierDescs.Num() > BarrierBatchMax)
 			{
-				check(hList.GetCurrentOwningContext());
-				FD3D12ScopedTimedIntervalQuery BarrierScopeTimer(GetParentDevice()->GetBackBufferWriteBarrierTracker(), hResourceBarrierList.GraphicsCommandList());
-				if (BackBufferBarrierDescs.Num() > BarrierBatchMax)
-				{
-					int Num = BackBufferBarrierDescs.Num();
-					D3D12_RESOURCE_BARRIER* Ptr = BackBufferBarrierDescs.GetData();
-					while (Num > 0)
-					{
-						const int DispatchNum = FMath::Min(Num, BarrierBatchMax);
-						hResourceBarrierList->ResourceBarrier(DispatchNum, Ptr);
-						Ptr += BarrierBatchMax;
-						Num -= BarrierBatchMax;
-					}
-				}
-				else
+				int Num = InBarrierDescInfo.BackBufferBarrierDescs.Num();
+				D3D12_RESOURCE_BARRIER* Ptr = InBarrierDescInfo.BackBufferBarrierDescs.GetData();
+				while (Num > 0)
 				{
-					hResourceBarrierList->ResourceBarrier(BackBufferBarrierDescs.Num(), BackBufferBarrierDescs.GetData());
+					const int DispatchNum = FMath::Min(Num, BarrierBatchMax);
+					hResourceBarrierList->ResourceBarrier(DispatchNum, Ptr);
+					Ptr += BarrierBatchMax;
+					Num -= BarrierBatchMax;
 				}
 			}
-
-			if (BarrierDescs.Num())
-#else
-			BarrierDescs.Append(BackBufferBarrierDescs);
-			BackBufferBarrierDescs.Empty();
-#endif // #if PLATFORM_USE_BACKBUFFER_WRITE_TRANSITION_TRACKING
+			else
 			{
-				extern void ResourceBarriersSeparateRTV2SRV(ID3D12GraphicsCommandList*, const TArray<D3D12_RESOURCE_BARRIER>&, int32);
-
-				ResourceBarriersSeparateRTV2SRV(hResourceBarrierList.GraphicsCommandList(), BarrierDescs, BarrierBatchMax);
+				hResourceBarrierList->ResourceBarrier(InBarrierDescInfo.BackBufferBarrierDescs.Num(), InBarrierDescInfo.BackBufferBarrierDescs.GetData());
 			}
 		}
 
-		return BarrierCount;
+		if (BarrierDescs.Num())
+#else
+		InBarrierDescInfo.BarrierDescs.Append(InBarrierDescInfo.BackBufferBarrierDescs);
+		InBarrierDescInfo.BackBufferBarrierDescs.Empty();
+#endif // #if PLATFORM_USE_BACKBUFFER_WRITE_TRANSITION_TRACKING
+		{
+			extern void ResourceBarriersSeparateRTV2SRV(ID3D12GraphicsCommandList*, const TArray<D3D12_RESOURCE_BARRIER>&, int32);
+			ResourceBarriersSeparateRTV2SRV(hResourceBarrierList.GraphicsCommandList(), InBarrierDescInfo.BarrierDescs, BarrierBatchMax);
+		}
 	}
 
-	return 0;
+	return BarrierCount;
 }
 
 bool FD3D12CommandListManager::IsComplete(const FD3D12CLSyncPoint& hSyncPoint, uint64 FenceOffset)
diff --git a/Engine/Source/Runtime/D3D12RHI/Private/D3D12DirectCommandListManager.h b/Engine/Source/Runtime/D3D12RHI/Private/D3D12DirectCommandListManager.h
--- a/Engine/Source/Runtime/D3D12RHI/Private/D3D12DirectCommandListManager.h
+++ b/Engine/Source/Runtime/D3D12RHI/Private/D3D12DirectCommandListManager.h
@@ -255,8 +255,9 @@ public:
 	FD3D12CommandListHandle ObtainCommandList(FD3D12CommandAllocator& CommandAllocator, bool bHasBackbufferWriteTransition = false);
 	void ReleaseCommandList(FD3D12CommandListHandle& hList);
 
-	void ExecuteCommandList(FD3D12CommandListHandle& hList, bool WaitForCompletion = false);
-	virtual void ExecuteCommandLists(TArray<FD3D12CommandListHandle>& Lists, bool WaitForCompletion = false);
+	FD3D12SyncPoint ExecuteCommandListNoCopyQueueSync(FD3D12CommandListHandle& hList, bool WaitForCompletion);
+	FD3D12SyncPoint ExecuteCommandList(FD3D12CommandListHandle& hList, FD3D12SyncPoint& CopyQueueSyncPoint, bool WaitForCompletion);
+	virtual FD3D12SyncPoint ExecuteCommandLists(TArray<FD3D12CommandListHandle>& Lists, FD3D12SyncPoint& CopyQueueSyncPoint, bool WaitForCompletion);
 
 	void WaitOnExecuteTask();
 
@@ -338,8 +339,18 @@ protected:
 		{}
 	};
 
-	void ExecuteCommandListInternal(TArray<FD3D12CommandListHandle>& Lists, bool WaitForCompletion);
-	uint32 GetResourceBarrierCommandList(FD3D12CommandListHandle& hList, FD3D12CommandListHandle& hResourceBarrierList);
+	FD3D12SyncPoint ExecuteCommandListInternal(TArray<FD3D12CommandListHandle>& Lists, FD3D12SyncPoint& CopyQueueSyncPoint, bool WaitForCompletion);
+	struct FBarrierDescInfo
+	{
+		TArray<D3D12_RESOURCE_BARRIER> BarrierDescs;
+		TArray<D3D12_RESOURCE_BARRIER, TInlineAllocator<2>> BackBufferBarrierDescs;
+#if ENABLE_RESIDENCY_MANAGEMENT
+		TArray<FD3D12ResidencyHandle*> ResidencyHandles;
+#endif // ENABLE_RESIDENCY_MANAGEMENT
+		bool bHasGraphicStates;
+	};
+	static uint32 CollectInitialResourceBarriersAndUpdateFinalState(FD3D12CommandListHandle& InCommandListHandle, FBarrierDescInfo& BarrierDescInfo);
+	uint32 GetResourceBarrierCommandList(FBarrierDescInfo& InBarrierDescInfo, FD3D12CommandListHandle& hResourceBarrierList);
 
 	// Returns signaled Fence
 	uint64 ExecuteAndIncrementFence(FD3D12CommandListPayload& Payload, FD3D12Fence &Fence);
diff --git a/Engine/Source/Runtime/D3D12RHI/Private/D3D12RayTracing.cpp b/Engine/Source/Runtime/D3D12RHI/Private/D3D12RayTracing.cpp
--- a/Engine/Source/Runtime/D3D12RHI/Private/D3D12RayTracing.cpp
+++ b/Engine/Source/Runtime/D3D12RHI/Private/D3D12RayTracing.cpp
@@ -606,11 +606,11 @@ FD3D12RayTracingCompactionRequestHandler::FD3D12RayTracingCompactionRequestHandl
 {
 	D3D12_RESOURCE_DESC PostBuildInfoBufferDesc = CD3DX12_RESOURCE_DESC::Buffer(GD3D12RayTracingMaxBatchedCompaction * sizeof(uint64), D3D12_RESOURCE_FLAG_ALLOW_UNORDERED_ACCESS);
 
-	FRHIResourceCreateInfo CreateInfo(TEXT("PostBuildInfoBuffer"));
-	CreateInfo.GPUMask = FRHIGPUMask::FromIndex(GetParentDevice()->GetGPUIndex());
+	FRHIGPUMask GPUMask = FRHIGPUMask::FromIndex(GetParentDevice()->GetGPUIndex());
 	ID3D12ResourceAllocator* ResourceAllocator = nullptr;
-	PostBuildInfoBuffer = GetParentDevice()->GetParentAdapter()->CreateRHIBuffer(nullptr, PostBuildInfoBufferDesc, 8,
-		0, PostBuildInfoBufferDesc.Width, BUF_UnorderedAccess | BUF_SourceCopy, ED3D12ResourceStateMode::MultiState, ERHIAccess::UAVMask, CreateInfo, ResourceAllocator);
+	bool bHasInitialData = false;
+	PostBuildInfoBuffer = GetParentDevice()->GetParentAdapter()->CreateRHIBuffer(PostBuildInfoBufferDesc, 8,
+		0, PostBuildInfoBufferDesc.Width, BUF_UnorderedAccess | BUF_SourceCopy, ED3D12ResourceStateMode::MultiState, D3D12_RESOURCE_STATE_UNORDERED_ACCESS, bHasInitialData, GPUMask, ResourceAllocator, TEXT("PostBuildInfoBuffer"));
 	SetName(PostBuildInfoBuffer->GetResource(), TEXT("PostBuildInfoBuffer"));
 
 	PostBuildInfoStagingBuffer = RHICreateStagingBuffer();
@@ -1756,17 +1756,18 @@ public:
 		checkf(Data.Num(), TEXT("Shader table is expected to be initialized before copying to GPU."));
 
 		FD3D12Adapter* Adapter = Device->GetParentAdapter();
-
 		D3D12_RESOURCE_DESC BufferDesc = CD3DX12_RESOURCE_DESC::Buffer(Data.GetResourceDataSize(), D3D12_RESOURCE_FLAG_NONE, D3D12_RAYTRACING_SHADER_TABLE_BYTE_ALIGNMENT);
-
-		FRHIResourceCreateInfo CreateInfo(TEXT("Shader binding table"));
-		CreateInfo.ResourceArray = &Data;
-		CreateInfo.GPUMask = FRHIGPUMask::FromIndex(Device->GetGPUIndex());
+		FRHIGPUMask GPUMask = FRHIGPUMask::FromIndex(Device->GetGPUIndex());
+		bool bHasInitialData = true;
 
 		ID3D12ResourceAllocator* ResourceAllocator = nullptr;
 		Buffer = Adapter->CreateRHIBuffer(
-			nullptr, BufferDesc, BufferDesc.Alignment,
-			0, BufferDesc.Width, BUF_Static, ED3D12ResourceStateMode::SingleState, ERHIAccess::SRVMask, CreateInfo, ResourceAllocator);
+			BufferDesc, BufferDesc.Alignment, 0, BufferDesc.Width, BUF_Static, ED3D12ResourceStateMode::SingleState, 
+			D3D12_RESOURCE_STATE_GENERIC_READ, bHasInitialData, GPUMask, ResourceAllocator, TEXT("Shader binding table"));
+
+		// providing nullprt immidiate command least means execute on default command list
+		FRHICommandListImmediate* RHICmdList = nullptr;
+		Buffer->UploadResourceData(RHICmdList, &Data, D3D12_RESOURCE_STATE_GENERIC_READ);
 
 		bIsDirty = false;
 	}
@@ -1974,7 +1975,7 @@ public:
 
 		for (FD3D12ShaderResourceView* SRV : TransitionSRVs[0])
 		{
-			FD3D12DynamicRHI::TransitionResource(CommandContext.CommandListHandle, SRV, D3D12_RESOURCE_STATE_PIXEL_SHADER_RESOURCE | D3D12_RESOURCE_STATE_NON_PIXEL_SHADER_RESOURCE, FD3D12DynamicRHI::ETransitionMode::Apply);
+			FD3D12DynamicRHI::TransitionResource(CommandContext.CommandListHandle, SRV, D3D12_RESOURCE_STATE_NON_PIXEL_SHADER_RESOURCE, FD3D12DynamicRHI::ETransitionMode::Apply);
 		}
 
 		for (FD3D12UnorderedAccessView* UAV : TransitionUAVs[0])
@@ -2842,28 +2843,26 @@ static TRefCountPtr<FD3D12Buffer> CreateRayTracingBuffer(FD3D12Adapter* Adapter,
 	TRefCountPtr<FD3D12Buffer> Result;
 
 	FString DebugNameString = DebugName.ToString();
-	FRHIResourceCreateInfo CreateInfo(*DebugNameString);
-	CreateInfo.GPUMask = FRHIGPUMask::FromIndex(GPUIndex);
-
 	ID3D12ResourceAllocator* ResourceAllocator = nullptr;
-
 	D3D12_RESOURCE_DESC BufferDesc = CD3DX12_RESOURCE_DESC::Buffer(Size, D3D12_RESOURCE_FLAG_ALLOW_UNORDERED_ACCESS);
+	FRHIGPUMask GPUMask = FRHIGPUMask::FromIndex(GPUIndex);
+	bool bHasInitialData = false;
 
 	if (Type == ERayTracingBufferType::AccelerationStructure)
 	{
 		Result = Adapter->CreateRHIBuffer(
-			nullptr, BufferDesc, D3D12_RAYTRACING_ACCELERATION_STRUCTURE_BYTE_ALIGNMENT,
+			BufferDesc, D3D12_RAYTRACING_ACCELERATION_STRUCTURE_BYTE_ALIGNMENT,
 			0, BufferDesc.Width, BUF_AccelerationStructure,
-			ED3D12ResourceStateMode::SingleState, ERHIAccess::BVHWrite,
-			CreateInfo, ResourceAllocator);
+			ED3D12ResourceStateMode::SingleState, D3D12_RESOURCE_STATE_RAYTRACING_ACCELERATION_STRUCTURE, bHasInitialData,
+			GPUMask, ResourceAllocator, *DebugNameString);
 	}
 	else if (Type == ERayTracingBufferType::Scratch)
 	{
 		Result = Adapter->CreateRHIBuffer(
-			nullptr, BufferDesc, D3D12_RAYTRACING_ACCELERATION_STRUCTURE_BYTE_ALIGNMENT,
+			BufferDesc, D3D12_RAYTRACING_ACCELERATION_STRUCTURE_BYTE_ALIGNMENT,
 			0, BufferDesc.Width, BUF_UnorderedAccess,
-			ED3D12ResourceStateMode::SingleState, ERHIAccess::UAVMask,
-			CreateInfo, ResourceAllocator);
+			ED3D12ResourceStateMode::SingleState, D3D12_RESOURCE_STATE_UNORDERED_ACCESS, bHasInitialData,
+			GPUMask, ResourceAllocator, *DebugNameString);
 	}
 	else
 	{
@@ -3050,7 +3049,7 @@ void FD3D12RayTracingGeometry::TransitionBuffers(FD3D12CommandContext& CommandCo
 		FD3D12Buffer* IndexBuffer = CommandContext.RetrieveObject<FD3D12Buffer>(RHIIndexBuffer.GetReference());
 		if (IndexBuffer->GetResource()->RequiresResourceStateTracking())
 		{
-			FD3D12DynamicRHI::TransitionResource(CommandContext.CommandListHandle, IndexBuffer->GetResource(), D3D12_RESOURCE_STATE_TBD, D3D12_RESOURCE_STATE_PIXEL_SHADER_RESOURCE | D3D12_RESOURCE_STATE_NON_PIXEL_SHADER_RESOURCE, 0, FD3D12DynamicRHI::ETransitionMode::Apply);
+			FD3D12DynamicRHI::TransitionResource(CommandContext.CommandListHandle, IndexBuffer->GetResource(), D3D12_RESOURCE_STATE_TBD, D3D12_RESOURCE_STATE_NON_PIXEL_SHADER_RESOURCE, 0, FD3D12DynamicRHI::ETransitionMode::Apply);
 		}
 	}
 
@@ -3060,7 +3059,7 @@ void FD3D12RayTracingGeometry::TransitionBuffers(FD3D12CommandContext& CommandCo
 		FD3D12Buffer* VertexBuffer = CommandContext.RetrieveObject<FD3D12Buffer>(RHIVertexBuffer.GetReference());
 		if (VertexBuffer->GetResource()->RequiresResourceStateTracking())
 		{
-			FD3D12DynamicRHI::TransitionResource(CommandContext.CommandListHandle, VertexBuffer->GetResource(), D3D12_RESOURCE_STATE_TBD, D3D12_RESOURCE_STATE_PIXEL_SHADER_RESOURCE | D3D12_RESOURCE_STATE_NON_PIXEL_SHADER_RESOURCE, 0, FD3D12DynamicRHI::ETransitionMode::Apply);
+			FD3D12DynamicRHI::TransitionResource(CommandContext.CommandListHandle, VertexBuffer->GetResource(), D3D12_RESOURCE_STATE_TBD, D3D12_RESOURCE_STATE_NON_PIXEL_SHADER_RESOURCE, 0, FD3D12DynamicRHI::ETransitionMode::Apply);
 		}
 	}
 }
@@ -3581,10 +3580,6 @@ void FD3D12RayTracingScene::BuildAccelerationStructure(FD3D12CommandContext& Com
 			}
 		}
 
-		FRHIResourceCreateInfo CreateInfo(TEXT("InstanceBuffer"));
-		CreateInfo.GPUMask = FRHIGPUMask::FromIndex(GPUIndex);
-		CreateInfo.ResourceArray = &Instances;
-		
 		check(Instances.Num() == BuildInputs.NumDescs);
 
 		D3D12_RESOURCE_DESC InstanceBufferDesc = CD3DX12_RESOURCE_DESC::Buffer(
@@ -3592,15 +3587,53 @@ void FD3D12RayTracingScene::BuildAccelerationStructure(FD3D12CommandContext& Com
 			D3D12_RESOURCE_FLAG_ALLOW_UNORDERED_ACCESS);
 
 		// #dxr_todo multi state only when bShouldCopyIndirectInstances is set - will still need transition to copy_src for lock behind but this can be done on the complete pool in theory (have to check cost)
+		// #dxr_todo use dynamic with upload memory when no CopyCommands requested - resource will be in upload memory then instead of VRAM
+
 		ID3D12ResourceAllocator* ResourceAllocator = nullptr;
+		FRHIGPUMask GPUMask = FRHIGPUMask::FromIndex(GPUIndex);
+		bool bHasInitialData = true;
+
 		AutoInstanceBuffer = Adapter->CreateRHIBuffer(
-			nullptr, InstanceBufferDesc, D3D12_RAYTRACING_INSTANCE_DESCS_BYTE_ALIGNMENT,
-			sizeof(D3D12_RAYTRACING_INSTANCE_DESC), InstanceBufferDesc.Width, BUF_UnorderedAccess, ED3D12ResourceStateMode::MultiState, ERHIAccess::UAVMask, CreateInfo, ResourceAllocator);
+			InstanceBufferDesc, D3D12_RAYTRACING_INSTANCE_DESCS_BYTE_ALIGNMENT,
+			sizeof(D3D12_RAYTRACING_INSTANCE_DESC), InstanceBufferDesc.Width, BUF_UnorderedAccess, ED3D12ResourceStateMode::MultiState, D3D12_RESOURCE_STATE_COPY_DEST, bHasInitialData, GPUMask, ResourceAllocator, TEXT("InstanceBuffer"));
+
+		// Use copy queue for uploading the data
+		FD3D12SyncPoint CopyQueueSyncPoint = AutoInstanceBuffer->UploadResourceDataViaCopyQueue(&Instances);
+		CommandContext.CopyQueueSyncPoint.Merge(CopyQueueSyncPoint);
 
 		// #yuriy_todo: only create internal instance buffer if explicit one is not given
 		checkf(InstanceBuffer == nullptr, TEXT("Explicit instance buffer support is not yet implemented for D3D12 RHI."));
 		InstanceBuffer = AutoInstanceBuffer.GetReference();
 
+		if (CopyCommands.Num() > 0)
+		{
+			TRHICommandList_RecursiveHazardous<FD3D12CommandContext> RHICmdList(&CommandContext);
+			FUnorderedAccessViewRHIRef InstancesDescUAV = RHICreateUnorderedAccessView(InstanceBuffer, false, false);
+
+			FD3D12DynamicRHI::TransitionResource(CommandContext.CommandListHandle, InstanceBuffer->GetResource(),
+				D3D12_RESOURCE_STATE_COPY_DEST, D3D12_RESOURCE_STATE_UNORDERED_ACCESS, 0, FD3D12DynamicRHI::ETransitionMode::Apply);
+
+			RHICmdList.BeginUAVOverlap(InstancesDescUAV);
+
+			for (const FInstanceCopyCommand& Command : CopyCommands)
+			{
+				FRHIShaderResourceView* TransformsSRV = Command.Source.GetReference();
+				CopyRayTracingGPUInstances(RHICmdList, Command.Num, Command.BaseIndex, TransformsSRV, InstancesDescUAV);
+			}
+
+			RHICmdList.EndUAVOverlap(InstancesDescUAV);
+
+			FD3D12DynamicRHI::TransitionResource(CommandContext.CommandListHandle, InstanceBuffer->GetResource(),
+				D3D12_RESOURCE_STATE_UNORDERED_ACCESS, D3D12_RESOURCE_STATE_NON_PIXEL_SHADER_RESOURCE, 0, FD3D12DynamicRHI::ETransitionMode::Apply);
+		}
+		else
+		{
+			FD3D12DynamicRHI::TransitionResource(CommandContext.CommandListHandle, InstanceBuffer->GetResource(),
+				D3D12_RESOURCE_STATE_COPY_DEST, D3D12_RESOURCE_STATE_NON_PIXEL_SHADER_RESOURCE, 0, FD3D12DynamicRHI::ETransitionMode::Apply);
+		}
+
+		InstanceBuffer->GetResource()->UpdateResidency(CommandContext.CommandListHandle);
+
 		{
 			TArray<FD3D12ResidencyHandle*>& GeometryResidencyHandlesForThisGPU = GeometryResidencyHandles[GPUIndex];
 
@@ -3647,32 +3680,6 @@ void FD3D12RayTracingScene::BuildAccelerationStructure(FD3D12CommandContext& Com
 				}
 			}
 		}
-
-		InstanceBuffer->GetResource()->UpdateResidency(CommandContext.CommandListHandle);
-
-		if(CopyCommands.Num())
-		{
-			TRHICommandList_RecursiveHazardous<FD3D12CommandContext> RHICmdList(&CommandContext);
-			FUnorderedAccessViewRHIRef InstancesDescUAV = RHICreateUnorderedAccessView(InstanceBuffer, false, false);
-
-			FD3D12DynamicRHI::TransitionResource(CommandContext.CommandListHandle, InstanceBuffer->GetResource(),
-				D3D12_RESOURCE_STATE_TBD, D3D12_RESOURCE_STATE_UNORDERED_ACCESS, 0,
-				FD3D12DynamicRHI::ETransitionMode::Apply);
-
-			RHICmdList.BeginUAVOverlap(InstancesDescUAV);
-
-			for (const FInstanceCopyCommand& Command : CopyCommands)
-			{
-				FRHIShaderResourceView* TransformsSRV = Command.Source.GetReference();
-				CopyRayTracingGPUInstances(RHICmdList, Command.Num, Command.BaseIndex, TransformsSRV, InstancesDescUAV);
-			}
-
-			RHICmdList.EndUAVOverlap(InstancesDescUAV);
-		}
-
-		FD3D12DynamicRHI::TransitionResource(CommandContext.CommandListHandle, InstanceBuffer->GetResource(),
-			D3D12_RESOURCE_STATE_TBD, D3D12_RESOURCE_STATE_PIXEL_SHADER_RESOURCE | D3D12_RESOURCE_STATE_NON_PIXEL_SHADER_RESOURCE, 0,
-			FD3D12DynamicRHI::ETransitionMode::Apply);
 	}
 	
 	
@@ -3924,7 +3931,7 @@ struct FD3D12RayTracingGlobalResourceBinder
 
 	void AddResourceTransition(FD3D12ShaderResourceView* SRV)
 	{
-		FD3D12DynamicRHI::TransitionResource(CommandContext.CommandListHandle, SRV, D3D12_RESOURCE_STATE_PIXEL_SHADER_RESOURCE | D3D12_RESOURCE_STATE_NON_PIXEL_SHADER_RESOURCE, FD3D12DynamicRHI::ETransitionMode::Apply);
+		FD3D12DynamicRHI::TransitionResource(CommandContext.CommandListHandle, SRV, D3D12_RESOURCE_STATE_NON_PIXEL_SHADER_RESOURCE, FD3D12DynamicRHI::ETransitionMode::Apply);
 	}
 
 	void AddResourceTransition(FD3D12UnorderedAccessView* UAV)
diff --git a/Engine/Source/Runtime/D3D12RHI/Private/D3D12Texture.cpp b/Engine/Source/Runtime/D3D12RHI/Private/D3D12Texture.cpp
--- a/Engine/Source/Runtime/D3D12RHI/Private/D3D12Texture.cpp
+++ b/Engine/Source/Runtime/D3D12RHI/Private/D3D12Texture.cpp
@@ -1506,7 +1506,9 @@ FTexture2DRHIRef FD3D12DynamicRHI::RHIAsyncCreateTexture2D(uint32 SizeX, uint32
 
 			// Wait for the copy context to finish before continuing as this function is only expected to return once all the texture streaming has finished.
 			hCopyCommandList.Close();
-			Device->GetCopyCommandListManager().ExecuteCommandList(hCopyCommandList, true);
+
+			bool bWaitForCompletion = true;
+			Device->GetCopyCommandListManager().ExecuteCommandListNoCopyQueueSync(hCopyCommandList, bWaitForCompletion);
 
 			CommandAllocatorManager.ReleaseCommandAllocator(CurrentCommandAllocator);
 		}
diff --git a/Engine/Source/Runtime/D3D12RHI/Private/D3D12Util.cpp b/Engine/Source/Runtime/D3D12RHI/Private/D3D12Util.cpp
--- a/Engine/Source/Runtime/D3D12RHI/Private/D3D12Util.cpp
+++ b/Engine/Source/Runtime/D3D12RHI/Private/D3D12Util.cpp
@@ -1421,6 +1421,12 @@ void FD3D12SyncPoint::WaitForCompletion() const
 	Fence->WaitForFence(Value);
 }
 
+void FD3D12SyncPoint::GPUWait(ED3D12CommandQueueType InCommandQueueType) const
+{
+	check(IsValid());
+	Fence->GpuWait(InCommandQueueType, Value);
+}
+
 // Forward declarations are required for the template functions
 template bool AssertResourceState(ID3D12CommandList* pCommandList, FD3D12View<D3D12_RENDER_TARGET_VIEW_DESC>* pView, const D3D12_RESOURCE_STATES& State);
 template bool AssertResourceState(ID3D12CommandList* pCommandList, FD3D12View<D3D12_UNORDERED_ACCESS_VIEW_DESC>* pView, const D3D12_RESOURCE_STATES& State);
diff --git a/Engine/Source/Runtime/D3D12RHI/Public/D3D12Resources.h b/Engine/Source/Runtime/D3D12RHI/Public/D3D12Resources.h
--- a/Engine/Source/Runtime/D3D12RHI/Public/D3D12Resources.h
+++ b/Engine/Source/Runtime/D3D12RHI/Public/D3D12Resources.h
@@ -938,9 +938,16 @@ public:
 	virtual ~FD3D12Buffer()
 	{
 		int64 BufferSize = ResourceLocation.GetSize();
-		UpdateBufferStats(GetBufferStats(GetUsage()), -BufferSize);
+		bool bTransient = ResourceLocation.IsTransient();
+		if (!bTransient)
+		{
+			UpdateBufferStats(GetBufferStats(GetUsage()), -BufferSize);
+		}
 	}
 
+	void UploadResourceData(class FRHICommandListImmediate* InRHICmdList, FResourceArrayInterface* InResourceArray, D3D12_RESOURCE_STATES InDestinationState);
+	FD3D12SyncPoint UploadResourceDataViaCopyQueue(FResourceArrayInterface* InResourceArray);
+
 	void Rename(FD3D12ResourceLocation& NewLocation);
 	void RenameLDAChain(FD3D12ResourceLocation& NewLocation);
 
diff --git a/Engine/Source/Runtime/D3D12RHI/Public/D3D12Util.h b/Engine/Source/Runtime/D3D12RHI/Public/D3D12Util.h
--- a/Engine/Source/Runtime/D3D12RHI/Public/D3D12Util.h
+++ b/Engine/Source/Runtime/D3D12RHI/Public/D3D12Util.h
@@ -261,6 +261,11 @@ static bool IsValidD3D12ResourceState(D3D12_RESOURCE_STATES InState)
 	return (InState != D3D12_RESOURCE_STATE_TBD && InState != D3D12_RESOURCE_STATE_CORRUPT);
 }
 
+static bool IsDirectQueueExclusiveD3D12State(D3D12_RESOURCE_STATES InState)
+{
+	return EnumHasAnyFlags(InState, D3D12_RESOURCE_STATE_RENDER_TARGET | D3D12_RESOURCE_STATE_DEPTH_WRITE | D3D12_RESOURCE_STATE_DEPTH_READ | D3D12_RESOURCE_STATE_PIXEL_SHADER_RESOURCE);
+}
+
 D3D12_RESOURCE_STATES GetD3D12ResourceState(ERHIAccess InRHIAccess, bool InIsAsyncCompute);
 
 //==================================================================================================================================
@@ -504,6 +509,14 @@ public:
 	bool IsValid() const;
 	bool IsComplete() const;
 	void WaitForCompletion() const;
+	void GPUWait(ED3D12CommandQueueType InCommandQueueType) const;
+
+	void Merge(const FD3D12SyncPoint& InOther)
+	{
+		check(Fence == nullptr || Fence == InOther.Fence);
+		Fence = InOther.Fence;
+		Value = FMath::Max(Value, InOther.Value);
+	}
 
 private:
 	FD3D12Fence* Fence;
