Compile fix for when PLATFORM_USE_BACKBUFFER_WRITE_TRANSITION_TRACKING is defined

diff --git a/Engine/Source/Runtime/D3D12RHI/Private/D3D12DirectCommandListManager.cpp b/Engine/Source/Runtime/D3D12RHI/Private/D3D12DirectCommandListManager.cpp
--- a/Engine/Source/Runtime/D3D12RHI/Private/D3D12DirectCommandListManager.cpp
+++ b/Engine/Source/Runtime/D3D12RHI/Private/D3D12DirectCommandListManager.cpp
@@ -1061,7 +1061,7 @@ uint32 FD3D12CommandListManager::GetResourceBarrierCommandList(FBarrierDescInfo&
 			}
 		}
 
-		if (BarrierDescs.Num())
+		if (InBarrierDescInfo.BarrierDescs.Num())
 #else
 		InBarrierDescInfo.BarrierDescs.Append(InBarrierDescInfo.BackBufferBarrierDescs);
 		InBarrierDescInfo.BackBufferBarrierDescs.Empty();
