Add defrag % and amount info in buffer detail stats

diff --git a/Engine/Source/Runtime/D3D12RHI/Private/D3D12Allocation.cpp b/Engine/Source/Runtime/D3D12RHI/Private/D3D12Allocation.cpp
--- a/Engine/Source/Runtime/D3D12RHI/Private/D3D12Allocation.cpp
+++ b/Engine/Source/Runtime/D3D12RHI/Private/D3D12Allocation.cpp
@@ -1072,13 +1072,14 @@ void FD3D12UploadHeapAllocator::UpdateMemoryStats()
 	uint32 MemoryAllocated = 0;
 	uint32 MemoryUsed = 0;
 	uint32 FreeMemory = 0;
+	uint32 EndFreeMemory = 0;
 	uint32 AlignmentWaste = 0;
 	uint32 AllocatedPageCount = 0;
 	uint32 FullPageCount = 0;
 
 #if defined(D3D12RHI_TRACK_DETAILED_STATS)
 	SmallBlockAllocator.UpdateMemoryStats(MemoryAllocated, MemoryUsed, FreeMemory, AlignmentWaste, AllocatedPageCount, FullPageCount);
-	BigBlockAllocator.UpdateMemoryStats(MemoryAllocated, MemoryUsed, FreeMemory, AlignmentWaste, AllocatedPageCount, FullPageCount);
+	BigBlockAllocator.UpdateMemoryStats(MemoryAllocated, MemoryUsed, FreeMemory, EndFreeMemory, AlignmentWaste, AllocatedPageCount, FullPageCount);
 	FastConstantPageAllocator.UpdateMemoryStats(MemoryAllocated, MemoryUsed, FreeMemory, AlignmentWaste, AllocatedPageCount, FullPageCount);
 #endif
 
@@ -1263,7 +1264,7 @@ void FD3D12DefaultBufferPool::AllocDefaultResource(D3D12_HEAP_TYPE InHeapType, c
 
 
 
-void FD3D12DefaultBufferPool::UpdateMemoryStats(uint32& IOMemoryAllocated, uint32& IOMemoryUsed, uint32& IOMemoryFree, uint32& IOAlignmentWaste, uint32& IOAllocatedPageCount, uint32& IOFullPageCount)
+void FD3D12DefaultBufferPool::UpdateMemoryStats(uint32& IOMemoryAllocated, uint32& IOMemoryUsed, uint32& IOMemoryFree, uint32& IOMemoryEndFree, uint32& IOAlignmentWaste, uint32& IOAllocatedPageCount, uint32& IOFullPageCount)
 {
 	Allocator->UpdateMemoryStats(IOMemoryAllocated, IOMemoryUsed, IOMemoryFree, IOAlignmentWaste, IOAllocatedPageCount, IOFullPageCount);
 }
@@ -1473,6 +1474,7 @@ void FD3D12DefaultBufferAllocator::UpdateMemoryStats()
 	uint32 MemoryAllocated = 0;
 	uint32 MemoryUsed = 0;
 	uint32 FreeMemory = 0;
+	uint32 EndFreeMemory = 0;
 	uint32 AlignmentWaste = 0;
 	uint32 AllocatedPageCount = 0;
 	uint32 FullPageCount = 0;
@@ -1482,17 +1484,25 @@ void FD3D12DefaultBufferAllocator::UpdateMemoryStats()
 	{
 		if (DefaultBufferPool)
 		{
-			DefaultBufferPool->UpdateMemoryStats(MemoryAllocated, MemoryUsed, FreeMemory, AlignmentWaste, AllocatedPageCount, FullPageCount);
+			DefaultBufferPool->UpdateMemoryStats(MemoryAllocated, MemoryUsed, FreeMemory, EndFreeMemory, AlignmentWaste, AllocatedPageCount, FullPageCount);
 		}
 	}
 #endif
 
+	//check((MemoryUsed + AlignmentWaste + FreeMemory) == MemoryAllocated);
+
+	// compute fragmentation percentage stats:
+	uint32 Fragmentation = FreeMemory - EndFreeMemory;
+	float FragmentationPercentage = float(Fragmentation) / float(MemoryUsed + AlignmentWaste + Fragmentation);
+
 	SET_MEMORY_STAT(STAT_D3D12BufferPoolMemoryAllocated, MemoryAllocated);
 	SET_MEMORY_STAT(STAT_D3D12BufferPoolMemoryUsed, MemoryUsed);
 	SET_MEMORY_STAT(STAT_D3D12BufferPoolMemoryFree, FreeMemory);
 	SET_MEMORY_STAT(STAT_D3D12BufferPoolAlignmentWaste, AlignmentWaste);
 	SET_DWORD_STAT(STAT_D3D12BufferPoolPageCount, AllocatedPageCount);
 	SET_DWORD_STAT(STAT_D3D12BufferPoolFullPages, FullPageCount);
+	SET_MEMORY_STAT(STAT_D3D12BufferPoolFragmentation, Fragmentation);
+	SET_FLOAT_STAT(STAT_D3D12BufferPoolFragmentationPercentage, FragmentationPercentage);
 }
 
 //-----------------------------------------------------------------------------
@@ -1651,13 +1661,14 @@ bool FD3D12TextureAllocatorPool::GetMemoryStats(uint64& OutTotalAllocated, uint6
 	uint32 MemoryAllocated = 0;
 	uint32 MemoryUsed = 0;
 	uint32 FreeMemory = 0;
+	uint32 EndFreeMemory = 0;
 	uint32 AlignmentWaste = 0;
 	uint32 AllocatedPageCount = 0;
 	uint32 FullPageCount = 0;
 
 	for (uint32 PoolIndex = 0; PoolIndex < (uint32)EPoolType::Count; ++PoolIndex)
 	{
-		PoolAllocators[PoolIndex]->UpdateMemoryStats(MemoryAllocated, MemoryUsed, FreeMemory, AlignmentWaste, AllocatedPageCount, FullPageCount);
+		PoolAllocators[PoolIndex]->UpdateMemoryStats(MemoryAllocated, MemoryUsed, FreeMemory, EndFreeMemory, AlignmentWaste, AllocatedPageCount, FullPageCount);
 	}
 
 	OutTotalAllocated = MemoryAllocated;
diff --git a/Engine/Source/Runtime/D3D12RHI/Private/D3D12Allocation.h b/Engine/Source/Runtime/D3D12RHI/Private/D3D12Allocation.h
--- a/Engine/Source/Runtime/D3D12RHI/Private/D3D12Allocation.h
+++ b/Engine/Source/Runtime/D3D12RHI/Private/D3D12Allocation.h
@@ -394,7 +394,7 @@ public:
 	void AllocDefaultResource(D3D12_HEAP_TYPE InHeapType, const D3D12_RESOURCE_DESC& InDesc, EBufferUsageFlags InBufferUsage, ED3D12ResourceStateMode InResourceStateMode,
 		D3D12_RESOURCE_STATES InCreateState, uint32 InAlignment, const TCHAR* InName, FD3D12ResourceLocation& ResourceLocation);
 	void CleanUpAllocations(uint64 InFrameLag);
-	void UpdateMemoryStats(uint32& IOMemoryAllocated, uint32& IOMemoryUsed, uint32& IOMemoryFree, uint32& IOAlignmentWaste, uint32& IOAllocatedPageCount, uint32& IOFullPageCount);
+	void UpdateMemoryStats(uint32& IOMemoryAllocated, uint32& IOMemoryUsed, uint32& IOMemoryFree, uint32& IOMemoryEndFree, uint32& IOAlignmentWaste, uint32& IOAllocatedPageCount, uint32& IOFullPageCount);
 
 	static FD3D12ResourceInitConfig GetResourceAllocatorInitConfig(D3D12_HEAP_TYPE InHeapType, D3D12_RESOURCE_FLAGS InResourceFlags, EBufferUsageFlags InBufferUsage);
 	static EResourceAllocationStrategy GetResourceAllocationStrategy(D3D12_RESOURCE_FLAGS InResourceFlags, ED3D12ResourceStateMode InResourceStateMode);
diff --git a/Engine/Source/Runtime/D3D12RHI/Private/D3D12Stats.h b/Engine/Source/Runtime/D3D12RHI/Private/D3D12Stats.h
--- a/Engine/Source/Runtime/D3D12RHI/Private/D3D12Stats.h
+++ b/Engine/Source/Runtime/D3D12RHI/Private/D3D12Stats.h
@@ -90,6 +90,8 @@ DECLARE_MEMORY_STAT_EXTERN(TEXT("BufferPool Memory Free"), STAT_D3D12BufferPoolM
 DECLARE_MEMORY_STAT_EXTERN(TEXT("BufferPool Memory Alignment Waste"), STAT_D3D12BufferPoolAlignmentWaste, STATGROUP_D3D12BufferAllocations, );
 DECLARE_DWORD_ACCUMULATOR_STAT_EXTERN(TEXT("BufferPool Page Count"), STAT_D3D12BufferPoolPageCount, STATGROUP_D3D12BufferAllocations, );
 DECLARE_DWORD_ACCUMULATOR_STAT_EXTERN(TEXT("BufferPool Full Pages"), STAT_D3D12BufferPoolFullPages, STATGROUP_D3D12BufferAllocations, );
+DECLARE_MEMORY_STAT_EXTERN(TEXT("Bufferool Fragmentation"), STAT_D3D12BufferPoolFragmentation, STATGROUP_D3D12BufferAllocations, );
+DECLARE_FLOAT_ACCUMULATOR_STAT_EXTERN(TEXT("Bufferool Fragmentation Percentage"), STAT_D3D12BufferPoolFragmentationPercentage, STATGROUP_D3D12BufferAllocations, );
 
 DECLARE_MEMORY_STAT_EXTERN(TEXT("UploadPool Memory Allocated"), STAT_D3D12UploadPoolMemoryAllocated, STATGROUP_D3D12Allocations, );
 DECLARE_MEMORY_STAT_EXTERN(TEXT("UploadPool Memory Used"), STAT_D3D12UploadPoolMemoryUsed, STATGROUP_D3D12Allocations, );
diff --git a/Engine/Source/Runtime/D3D12RHI/Private/D3D12Util.cpp b/Engine/Source/Runtime/D3D12RHI/Private/D3D12Util.cpp
--- a/Engine/Source/Runtime/D3D12RHI/Private/D3D12Util.cpp
+++ b/Engine/Source/Runtime/D3D12RHI/Private/D3D12Util.cpp
@@ -1583,6 +1583,8 @@ DEFINE_STAT(STAT_D3D12BufferPoolMemoryFree);
 DEFINE_STAT(STAT_D3D12BufferPoolAlignmentWaste);
 DEFINE_STAT(STAT_D3D12BufferPoolPageCount);
 DEFINE_STAT(STAT_D3D12BufferPoolFullPages);
+DEFINE_STAT(STAT_D3D12BufferPoolFragmentation);
+DEFINE_STAT(STAT_D3D12BufferPoolFragmentationPercentage);
 
 DEFINE_STAT(STAT_D3D12UploadPoolMemoryAllocated);
 DEFINE_STAT(STAT_D3D12UploadPoolMemoryUsed);
diff --git a/Engine/Source/Runtime/RHICore/Private/RHIPoolAllocator.cpp b/Engine/Source/Runtime/RHICore/Private/RHIPoolAllocator.cpp
--- a/Engine/Source/Runtime/RHICore/Private/RHIPoolAllocator.cpp
+++ b/Engine/Source/Runtime/RHICore/Private/RHIPoolAllocator.cpp
@@ -676,11 +676,14 @@ void FRHIPoolAllocator::Defrag(uint32 InMaxCopySize, uint32& CurrentCopySize)
 }
 
 
-void FRHIPoolAllocator::UpdateMemoryStats(uint32& IOMemoryAllocated, uint32& IOMemoryUsed, uint32& IOMemoryFree, uint32& IOAlignmentWaste, uint32& IOAllocatedPageCount, uint32& IOFullPageCount)
+void FRHIPoolAllocator::UpdateMemoryStats(uint32& IOMemoryAllocated, uint32& IOMemoryUsed, uint32& IOMemoryFree, uint32& IOMemoryEndFree, uint32& IOAlignmentWaste, uint32& IOAllocatedPageCount, uint32& IOFullPageCount)
 {
 	FScopeLock Lock(&CS);
 
 	TotalAllocatedBlocks = 0;
+
+	// found out which pool has maximum available free size
+	uint32 MaxEndFree = 0;
 	for (FRHIMemoryPool* Pool : Pools)
 	{
 		if (Pool)
@@ -688,6 +691,7 @@ void FRHIPoolAllocator::UpdateMemoryStats(uint32& IOMemoryAllocated, uint32& IOM
 			IOMemoryAllocated += Pool->GetPoolSize();
 			IOMemoryUsed += Pool->GetUsedSize() - Pool->GetAlignmentWaste();
 			IOMemoryFree += Pool->GetFreeSize();
+			MaxEndFree = FMath::Max(MaxEndFree, (uint32)Pool->GetFreeSize());
 			IOAlignmentWaste += Pool->GetAlignmentWaste();
 			IOAllocatedPageCount++;
 			if (Pool->IsFull())
@@ -698,4 +702,6 @@ void FRHIPoolAllocator::UpdateMemoryStats(uint32& IOMemoryAllocated, uint32& IOM
 			TotalAllocatedBlocks += Pool->GetAllocatedBlocks();
 		}
 	}
+
+	IOMemoryEndFree += MaxEndFree;
 }
diff --git a/Engine/Source/Runtime/RHICore/Public/RHIPoolAllocator.h b/Engine/Source/Runtime/RHICore/Public/RHIPoolAllocator.h
--- a/Engine/Source/Runtime/RHICore/Public/RHIPoolAllocator.h
+++ b/Engine/Source/Runtime/RHICore/Public/RHIPoolAllocator.h
@@ -220,7 +220,7 @@ public:
 	void Defrag(uint32 InMaxCopySize, uint32& CurrentCopySize);
 
 	// Stats
-	void UpdateMemoryStats(uint32& IOMemoryAllocated, uint32& IOMemoryUsed, uint32& IOMemoryFree, uint32& IOAlignmentWaste, uint32& IOAllocatedPageCount, uint32& IOFullPageCount);
+	void UpdateMemoryStats(uint32& IOMemoryAllocated, uint32& IOMemoryUsed, uint32& IOMemoryFree, uint32& IOMemoryEndFree, uint32& IOAlignmentWaste, uint32& IOAllocatedPageCount, uint32& IOFullPageCount);
 
 protected:
 		
