Boost priority on RT enabled resources

diff --git a/Engine/Source/Runtime/D3D12RHI/Private/D3D12PoolAllocator.cpp b/Engine/Source/Runtime/D3D12RHI/Private/D3D12PoolAllocator.cpp
--- a/Engine/Source/Runtime/D3D12RHI/Private/D3D12PoolAllocator.cpp
+++ b/Engine/Source/Runtime/D3D12RHI/Private/D3D12PoolAllocator.cpp
@@ -362,6 +362,17 @@ void FD3D12PoolAllocator::AllocateResource(uint32 GPUIndex, D3D12_HEAP_TYPE InHe
 
 		ResourceLocation.AsStandAlone(NewResource, InSize);
 	}
+
+#if D3D12_RHI_RAYTRACING
+	// Elevates the raytracing acceleration structure heap priority, which may help performance / stability in low memory conditions
+	if (InCreateState == D3D12_RESOURCE_STATE_RAYTRACING_ACCELERATION_STRUCTURE)
+	{
+		ID3D12Pageable* HeapResource = ResourceLocation.GetResource()->GetPageable();
+		D3D12_RESIDENCY_PRIORITY HeapPriority = D3D12_RESIDENCY_PRIORITY_HIGH;
+		FD3D12Device* NodeDevice = Adapter->GetDevice(GPUIndex);
+		NodeDevice->GetDevice5()->SetResidencyPriority(1, &HeapResource, &HeapPriority);
+	}
+#endif // D3D12_RHI_RAYTRACING
 }
 
 
diff --git a/Engine/Source/Runtime/D3D12RHI/Private/D3D12RayTracing.cpp b/Engine/Source/Runtime/D3D12RHI/Private/D3D12RayTracing.cpp
--- a/Engine/Source/Runtime/D3D12RHI/Private/D3D12RayTracing.cpp
+++ b/Engine/Source/Runtime/D3D12RHI/Private/D3D12RayTracing.cpp
@@ -2863,18 +2863,19 @@ static TRefCountPtr<FD3D12Buffer> CreateRayTracingBuffer(FD3D12Adapter* Adapter,
 			0, BufferDesc.Width, BUF_UnorderedAccess,
 			ED3D12ResourceStateMode::SingleState, D3D12_RESOURCE_STATE_UNORDERED_ACCESS, bHasInitialData,
 			GPUMask, ResourceAllocator, *DebugNameString);
+
+		// Elevates the scratch buffer heap priority, which may help performance / stability in low memory conditions 
+		// (Acceleration structure already boosted from allocation side)
+		ID3D12Pageable* HeapResource = Result->GetResource()->GetPageable();
+		D3D12_RESIDENCY_PRIORITY HeapPriority = D3D12_RESIDENCY_PRIORITY_HIGH;
+		FD3D12Device* NodeDevice = Adapter->GetDevice(GPUIndex);
+		NodeDevice->GetDevice5()->SetResidencyPriority(1, &HeapResource, &HeapPriority);
 	}
 	else
 	{
 		checkNoEntry();
 	}
 
-	// Elevates the raytracing acceleration structure heap priority, which may help performance / stability in low memory conditions
-	ID3D12Pageable* HeapResource = Result->GetResource()->GetPageable();
-	D3D12_RESIDENCY_PRIORITY HeapPriority = D3D12_RESIDENCY_PRIORITY_HIGH;
-	FD3D12Device* NodeDevice = Adapter->GetDevice(GPUIndex);
-	NodeDevice->GetDevice5()->SetResidencyPriority(1, &HeapResource, &HeapPriority);
-
 	return Result;
 }
 
