Fix DredSettings leaking

diff --git a/Engine/Source/Runtime/D3D12RHI/Private/D3D12Adapter.cpp b/Engine/Source/Runtime/D3D12RHI/Private/D3D12Adapter.cpp
--- a/Engine/Source/Runtime/D3D12RHI/Private/D3D12Adapter.cpp
+++ b/Engine/Source/Runtime/D3D12RHI/Private/D3D12Adapter.cpp
@@ -380,8 +380,8 @@ void FD3D12Adapter::CreateRootDevice(bool bWithDebug)
 	bool bDREDContext = false;
 	if (EnumHasAnyFlags(GPUCrashDebuggingModes, ED3D12GPUCrashDebuggingModes::DRED))
 	{
-		ID3D12DeviceRemovedExtendedDataSettings* DredSettings = nullptr;
-		HRESULT hr = D3D12GetDebugInterface(IID_PPV_ARGS(&DredSettings));
+		TRefCountPtr<ID3D12DeviceRemovedExtendedDataSettings> DredSettings;
+		HRESULT hr = D3D12GetDebugInterface(IID_PPV_ARGS(DredSettings.GetInitReference()));
 
 		// Can fail if not on correct Windows Version - needs 1903 or newer
 		if (SUCCEEDED(hr))
