- Fixup small issues for making non UAV buffers also placed and tracked (requested as multistate)
- Use copy queue upload for RT Shader resource tables

diff --git a/Engine/Source/Runtime/D3D12RHI/Private/D3D12PoolAllocator.cpp b/Engine/Source/Runtime/D3D12RHI/Private/D3D12PoolAllocator.cpp
--- a/Engine/Source/Runtime/D3D12RHI/Private/D3D12PoolAllocator.cpp
+++ b/Engine/Source/Runtime/D3D12RHI/Private/D3D12PoolAllocator.cpp
@@ -263,12 +263,10 @@ void FD3D12PoolAllocator::AllocateResource(uint32 GPUIndex, D3D12_HEAP_TYPE InHe
 		// Ensure we're allocating from the correct pool
 		if (bPlacedResource)
 		{
-			// Writeable resources get separate ID3D12Resource* with their own resource state by using placed resources. Just make sure it's UAV, other flags are free to differ.
-			check(InDesc.Dimension != D3D12_RESOURCE_DIMENSION_BUFFER || (InDesc.Flags & D3D12_RESOURCE_FLAG_ALLOW_UNORDERED_ACCESS) != 0 || InHeapType == D3D12_HEAP_TYPE_READBACK);
-
 			// If it's a placed resource then base offset will always be 0 from the actual d3d resource so ignore the allocation alignment - no extra offset required
 			// for creating the views!
 			check(InAllocationAlignment <= PoolAlignment);
+			check(InDesc.Alignment <= PoolAlignment);
 			AllocationAlignment = PoolAlignment;
 		}
 		else
@@ -319,8 +317,11 @@ void FD3D12PoolAllocator::AllocateResource(uint32 GPUIndex, D3D12_HEAP_TYPE InHe
 		else
 		{
 			check(ResourceLocation.GetResource() == nullptr);
+			
+			D3D12_RESOURCE_DESC Desc = InDesc;
+			Desc.Alignment = AllocationAlignment;
 
-			FD3D12Resource* NewResource = CreatePlacedResource(AllocationData, InDesc, InCreateState, InResourceStateMode, InClearValue, InName);
+			FD3D12Resource* NewResource = CreatePlacedResource(AllocationData, Desc, InCreateState, InResourceStateMode, InClearValue, InName);
 			ResourceLocation.SetResource(NewResource);
 			ResourceLocation.SetGPUVirtualAddress(NewResource->GetGPUVirtualAddress());
 		}
diff --git a/Engine/Source/Runtime/D3D12RHI/Private/D3D12RayTracing.cpp b/Engine/Source/Runtime/D3D12RHI/Private/D3D12RayTracing.cpp
--- a/Engine/Source/Runtime/D3D12RHI/Private/D3D12RayTracing.cpp
+++ b/Engine/Source/Runtime/D3D12RHI/Private/D3D12RayTracing.cpp
@@ -1747,7 +1747,7 @@ public:
 		}
 	}
 
-	void CopyToGPU(FD3D12Device* Device)
+	void CopyToGPU(FD3D12CommandContext& CommandContext)
 	{
 		TRACE_CPUPROFILER_EVENT_SCOPE(ShaderTableCopyToGPU);
 
@@ -1755,6 +1755,7 @@ public:
 
 		checkf(Data.Num(), TEXT("Shader table is expected to be initialized before copying to GPU."));
 
+		FD3D12Device* Device = CommandContext.GetParentDevice();
 		FD3D12Adapter* Adapter = Device->GetParentAdapter();
 		D3D12_RESOURCE_DESC BufferDesc = CD3DX12_RESOURCE_DESC::Buffer(Data.GetResourceDataSize(), D3D12_RESOURCE_FLAG_NONE, D3D12_RAYTRACING_SHADER_TABLE_BYTE_ALIGNMENT);
 		FRHIGPUMask GPUMask = FRHIGPUMask::FromIndex(Device->GetGPUIndex());
@@ -1762,12 +1763,16 @@ public:
 
 		ID3D12ResourceAllocator* ResourceAllocator = nullptr;
 		Buffer = Adapter->CreateRHIBuffer(
-			BufferDesc, BufferDesc.Alignment, 0, BufferDesc.Width, BUF_Static, ED3D12ResourceStateMode::SingleState, 
-			D3D12_RESOURCE_STATE_GENERIC_READ, bHasInitialData, GPUMask, ResourceAllocator, TEXT("Shader binding table"));
+			BufferDesc, BufferDesc.Alignment, 0, BufferDesc.Width, BUF_Static, ED3D12ResourceStateMode::MultiState,
+			D3D12_RESOURCE_STATE_COPY_DEST, bHasInitialData, GPUMask, ResourceAllocator, TEXT("Shader binding table"));
 
-		// providing nullprt immidiate command least means execute on default command list
-		FRHICommandListImmediate* RHICmdList = nullptr;
-		Buffer->UploadResourceData(RHICmdList, &Data, D3D12_RESOURCE_STATE_GENERIC_READ);
+		// Use copy queue for uploading the data
+		FD3D12SyncPoint CopyQueueSyncPoint = Buffer->UploadResourceDataViaCopyQueue(&Data);
+		CommandContext.CopyQueueSyncPoint.Merge(CopyQueueSyncPoint);
+
+		// Enqueue transition to SRV
+		FD3D12DynamicRHI::TransitionResource(CommandContext.CommandListHandle, Buffer->GetResource(),
+			D3D12_RESOURCE_STATE_COPY_DEST, D3D12_RESOURCE_STATE_NON_PIXEL_SHADER_RESOURCE, 0, FD3D12DynamicRHI::ETransitionMode::Apply);
 
 		bIsDirty = false;
 	}
@@ -4617,7 +4622,7 @@ void FD3D12CommandContext::RHIRayTraceOcclusion(FRHIRayTracingScene* InScene,
 
 	if (ShaderTable.bIsDirty)
 	{
-		ShaderTable.CopyToGPU(GetParentDevice());
+		ShaderTable.CopyToGPU(*this);
 	}
 
 	Scene->UpdateResidency(*this);
@@ -4654,7 +4659,7 @@ void FD3D12CommandContext::RHIRayTraceIntersection(FRHIRayTracingScene* InScene,
 
 	if (ShaderTable.bIsDirty)
 	{
-		ShaderTable.CopyToGPU(GetParentDevice());
+		ShaderTable.CopyToGPU(*this);
 	}
 
 	Scene->UpdateResidency(*this);
@@ -4691,7 +4696,7 @@ void FD3D12CommandContext::RHIRayTraceDispatch(FRHIRayTracingPipelineState* InRa
 
 	if (ShaderTable->bIsDirty)
 	{
-		ShaderTable->CopyToGPU(GetParentDevice());
+		ShaderTable->CopyToGPU(*this);
 	}
 
 	Scene->UpdateResidency(*this);
@@ -4727,7 +4732,7 @@ void FD3D12CommandContext::RHIRayTraceDispatchIndirect(FRHIRayTracingPipelineSta
 
 	if (ShaderTable->bIsDirty)
 	{
-		ShaderTable->CopyToGPU(GetParentDevice());
+		ShaderTable->CopyToGPU(*this);
 	}
 
 	Scene->UpdateResidency(*this);
diff --git a/Engine/Source/Runtime/D3D12RHI/Public/D3D12Resources.h b/Engine/Source/Runtime/D3D12RHI/Public/D3D12Resources.h
--- a/Engine/Source/Runtime/D3D12RHI/Public/D3D12Resources.h
+++ b/Engine/Source/Runtime/D3D12RHI/Public/D3D12Resources.h
@@ -346,7 +346,7 @@ private:
 		}
 		else
 		{
-			DetermineResourceStates(InDefaultState);
+			DetermineResourceStates(InDefaultState, InResourceStateMode);
 		}
 
 		if (bRequiresResourceStateTracking)
@@ -362,7 +362,7 @@ private:
 		}
 	}
 
-	void DetermineResourceStates(D3D12_RESOURCE_STATES InDefaultState)
+	void DetermineResourceStates(D3D12_RESOURCE_STATES InDefaultState, ED3D12ResourceStateMode InResourceStateMode)
 	{
 		const FD3D12ResourceTypeHelper Type(Desc, HeapType);
 
@@ -372,7 +372,7 @@ private:
 		SetCompressedState(D3D12_RESOURCE_STATE_COMMON);
 #endif
 
-		if (Type.bWritable)
+		if (Type.bWritable || InResourceStateMode == ED3D12ResourceStateMode::MultiState)
 		{
 			// Determine the resource's write/read states.
 			if (Type.bRTV)
@@ -390,52 +390,29 @@ private:
 			}
 			else
 			{
-				check(Type.bUAV && !Type.bRTV && !Type.bDSV);
-				WritableState = D3D12_RESOURCE_STATE_UNORDERED_ACCESS;
+				WritableState = Type.bUAV ? D3D12_RESOURCE_STATE_UNORDERED_ACCESS : D3D12_RESOURCE_STATE_CORRUPT;
 				ReadableState = Type.bSRV ? D3D12_RESOURCE_STATE_NON_PIXEL_SHADER_RESOURCE | D3D12_RESOURCE_STATE_PIXEL_SHADER_RESOURCE : D3D12_RESOURCE_STATE_CORRUPT;
 			}
 		}
-
-		if (Type.bBuffer)
+		else
 		{
-			if (!Type.bWritable)
-			{
-				// Buffer used for input, like Vertex/Index buffer.
-				// Don't bother tracking state for this resource.
+			bRequiresResourceStateTracking = false;
+						
 #if UE_BUILD_DEBUG
-				FPlatformAtomics::InterlockedIncrement(&NoStateTrackingResourceCount);
+			FPlatformAtomics::InterlockedIncrement(&NoStateTrackingResourceCount);
 #endif
-				if (InDefaultState != D3D12_RESOURCE_STATE_TBD)
-				{
-					DefaultResourceState = InDefaultState;
-				}
-				else
-				{
-					DefaultResourceState = (HeapType == D3D12_HEAP_TYPE_READBACK) ? D3D12_RESOURCE_STATE_COPY_DEST : D3D12_RESOURCE_STATE_GENERIC_READ;
-				}
-				bRequiresResourceStateTracking = false;
-				return;
+			if (InDefaultState != D3D12_RESOURCE_STATE_TBD)
+			{
+				DefaultResourceState = InDefaultState;
 			}
-		}
-		else
-		{
-			if (Type.bSRVOnly)
+			else if (Type.bBuffer)
 			{
-				// Texture used only as a SRV.
-				// Don't bother tracking state for this resource.
-#if UE_BUILD_DEBUG
-				FPlatformAtomics::InterlockedIncrement(&NoStateTrackingResourceCount);
-#endif
-				if (InDefaultState != D3D12_RESOURCE_STATE_TBD)
-				{
-					DefaultResourceState = InDefaultState;
-				}
-				else
-				{
-					DefaultResourceState = D3D12_RESOURCE_STATE_NON_PIXEL_SHADER_RESOURCE | D3D12_RESOURCE_STATE_PIXEL_SHADER_RESOURCE;
-				}
-				bRequiresResourceStateTracking = false;
-				return;
+				DefaultResourceState = (HeapType == D3D12_HEAP_TYPE_READBACK) ? D3D12_RESOURCE_STATE_COPY_DEST : D3D12_RESOURCE_STATE_GENERIC_READ;
+			}
+			else
+			{
+				check(Type.bSRVOnly);
+				DefaultResourceState = D3D12_RESOURCE_STATE_NON_PIXEL_SHADER_RESOURCE | D3D12_RESOURCE_STATE_PIXEL_SHADER_RESOURCE;
 			}
 		}
 	}
