Also mark pending barriers as pending work during flush of command list (fixes possible missing barriers to back buffer to present)

diff --git a/Engine/Source/Runtime/D3D12RHI/Private/D3D12CommandContext.cpp b/Engine/Source/Runtime/D3D12RHI/Private/D3D12CommandContext.cpp
--- a/Engine/Source/Runtime/D3D12RHI/Private/D3D12CommandContext.cpp
+++ b/Engine/Source/Runtime/D3D12RHI/Private/D3D12CommandContext.cpp
@@ -330,6 +330,7 @@ void FD3D12CommandContext::OpenCommandList()
 	numDispatches = 0;
 	numClears = 0;
 	numBarriers = 0;
+	numPendingBarriers = 0;
 	numCopies = 0;
 	numInitialResourceCopies = 0;
 	otherWorkCounter = 0;
diff --git a/Engine/Source/Runtime/D3D12RHI/Private/D3D12CommandContext.h b/Engine/Source/Runtime/D3D12RHI/Private/D3D12CommandContext.h
--- a/Engine/Source/Runtime/D3D12RHI/Private/D3D12CommandContext.h
+++ b/Engine/Source/Runtime/D3D12RHI/Private/D3D12CommandContext.h
@@ -197,13 +197,14 @@ public:
 	uint32 numDispatches;
 	uint32 numClears;
 	uint32 numBarriers;
+	uint32 numPendingBarriers;
 	uint32 numCopies;
 	uint32 numInitialResourceCopies;
 	uint32 otherWorkCounter;
 
 	bool HasDoneWork() const
 	{
-		return (numDraws + numDispatches + numClears + numBarriers + numCopies + numInitialResourceCopies + otherWorkCounter) > 0;
+		return (numDraws + numDispatches + numClears + numBarriers + numPendingBarriers + numCopies + numInitialResourceCopies + otherWorkCounter) > 0;
 	}
 
 	/** Constant buffers for Set*ShaderParameter calls. */
diff --git a/Engine/Source/Runtime/D3D12RHI/Private/D3D12CommandList.cpp b/Engine/Source/Runtime/D3D12RHI/Private/D3D12CommandList.cpp
--- a/Engine/Source/Runtime/D3D12RHI/Private/D3D12CommandList.cpp
+++ b/Engine/Source/Runtime/D3D12RHI/Private/D3D12CommandList.cpp
@@ -17,6 +17,16 @@ static uint64 GenerateCommandListID()
 	return FPlatformAtomics::InterlockedIncrement(&GCommandListIDCounter);
 }
 
+void FD3D12CommandListHandle::AddPendingResourceBarrier(FD3D12Resource* Resource, D3D12_RESOURCE_STATES State, uint32 SubResource)
+{
+	check(CommandListData);
+
+	FD3D12PendingResourceBarrier PRB = { Resource, State, SubResource };
+
+	CommandListData->PendingResourceBarriers.Add(PRB);
+	CommandListData->CurrentOwningContext->numPendingBarriers++;
+}
+
 void FD3D12CommandListHandle::AddTransitionBarrier(FD3D12Resource* pResource, D3D12_RESOURCE_STATES Before, D3D12_RESOURCE_STATES After, uint32 Subresource)
 {
 	check(CommandListData);
diff --git a/Engine/Source/Runtime/D3D12RHI/Private/D3D12CommandList.h b/Engine/Source/Runtime/D3D12RHI/Private/D3D12CommandList.h
--- a/Engine/Source/Runtime/D3D12RHI/Private/D3D12CommandList.h
+++ b/Engine/Source/Runtime/D3D12RHI/Private/D3D12CommandList.h
@@ -526,14 +526,7 @@ public:
 		return CommandListData->TrackedResourceState.GetResourceState(pResource);
 	}
 
-	void AddPendingResourceBarrier(FD3D12Resource* Resource, D3D12_RESOURCE_STATES State, uint32 SubResource)
-	{
-		check(CommandListData);
-
-		FD3D12PendingResourceBarrier PRB = { Resource, State, SubResource };
-
-		CommandListData->PendingResourceBarriers.Add(PRB);
-	}
+	void AddPendingResourceBarrier(FD3D12Resource* Resource, D3D12_RESOURCE_STATES State, uint32 SubResource);
 
 	TArray<FD3D12PendingResourceBarrier>& PendingResourceBarriers()
 	{
