Add D3D12RHI_API to pending barrier function because it's called directly from a plugin - build fix

diff --git a/Engine/Source/Runtime/D3D12RHI/Private/D3D12CommandList.h b/Engine/Source/Runtime/D3D12RHI/Private/D3D12CommandList.h
--- a/Engine/Source/Runtime/D3D12RHI/Private/D3D12CommandList.h
+++ b/Engine/Source/Runtime/D3D12RHI/Private/D3D12CommandList.h
@@ -526,7 +526,7 @@ public:
 		return CommandListData->TrackedResourceState.GetResourceState(pResource);
 	}
 
-	void AddPendingResourceBarrier(FD3D12Resource* Resource, D3D12_RESOURCE_STATES State, uint32 SubResource);
+	D3D12RHI_API void AddPendingResourceBarrier(FD3D12Resource* Resource, D3D12_RESOURCE_STATES State, uint32 SubResource);
 
 	TArray<FD3D12PendingResourceBarrier>& PendingResourceBarriers()
 	{
