Clear global descriptor heap block after GPU is done using it to null SRVs on PC to fix page fault issues caused by accessing already destroyed heaps - still unkown what actual underlying issue is but looks like it 'fixes' the issue for now

diff --git a/Engine/Source/Runtime/D3D12RHI/Private/D3D12DescriptorCache.cpp b/Engine/Source/Runtime/D3D12RHI/Private/D3D12DescriptorCache.cpp
--- a/Engine/Source/Runtime/D3D12RHI/Private/D3D12DescriptorCache.cpp
+++ b/Engine/Source/Runtime/D3D12RHI/Private/D3D12DescriptorCache.cpp
@@ -1031,6 +1031,24 @@ void FD3D12GlobalHeap::Init(D3D12_DESCRIPTOR_HEAP_TYPE InType, uint32 InTotalSiz
 		FreeBlocks.Enqueue(new FD3D12GlobalHeapBlock(CurrentBaseSlot, ActualBlockSize));
 		CurrentBaseSlot += ActualBlockSize;
 	}
+
+	// Create a default SRV view
+	D3D12_SHADER_RESOURCE_VIEW_DESC SRVDesc = {};
+	SRVDesc.Shader4ComponentMapping = D3D12_DEFAULT_SHADER_4_COMPONENT_MAPPING;
+	SRVDesc.ViewDimension = D3D12_SRV_DIMENSION_TEXTURE2D;
+	SRVDesc.Format = DXGI_FORMAT_B8G8R8A8_UNORM;
+	SRVDesc.Texture2D.MipLevels = 1;
+	SRVDesc.Texture2D.MostDetailedMip = 0;
+	SRVDesc.Texture2D.ResourceMinLODClamp = 0.0f;
+	pNullSRV = new FD3D12DescriptorHandleSRV(GetParentDevice());
+	pNullSRV->CreateView(SRVDesc, nullptr);
+
+	// clear the whole heap with null srvs
+	EmptyDescriptors.SetNum(BlockSize);
+	for (uint32 SlotIndex = 0; SlotIndex < BlockSize; SlotIndex++)
+	{
+		EmptyDescriptors[SlotIndex] = pNullSRV->GetHandle();
+	}
 }
 
 
@@ -1094,6 +1112,19 @@ void FD3D12GlobalHeap::UpdateFreeBlocks()
 			DEC_DWORD_STAT_BY(STAT_GlobalViewHeapWastedDescriptors, ReleasedBlock->Size - ReleasedBlock->SizeUsed);
 			INC_DWORD_STAT_BY(STAT_GlobalViewHeapFreeDescriptors, ReleasedBlock->Size);
 
+#if PLATFORM_WINDOWS
+			// clear the whole heap with null srvs to make sure they are not used or read by the GPU anymore
+			// Fixes some strange are rare page fault GPU crashes on nVidia
+			// clearing on reuse still crashes so current assumption is some strange random access caused
+			// outside of application control
+			{			
+				check(ReleasedBlock->Size <= (uint32)EmptyDescriptors.Num());
+				ID3D12Device* Device = GetParentDevice()->GetDevice();
+				D3D12_CPU_DESCRIPTOR_HANDLE TargetCPUBase = GetCPUSlotHandle(ReleasedBlock);
+				Device->CopyDescriptorsSimple(ReleasedBlock->Size, TargetCPUBase, EmptyDescriptors.GetData()[0], D3D12_DESCRIPTOR_HEAP_TYPE_CBV_SRV_UAV);
+			}
+#endif  // PLATFORM_WINDOWS
+
 			ReleasedBlock->SizeUsed = 0;
 			FreeBlocks.Enqueue(ReleasedBlock);
 
diff --git a/Engine/Source/Runtime/D3D12RHI/Private/D3D12DescriptorCache.h b/Engine/Source/Runtime/D3D12RHI/Private/D3D12DescriptorCache.h
--- a/Engine/Source/Runtime/D3D12RHI/Private/D3D12DescriptorCache.h
+++ b/Engine/Source/Runtime/D3D12RHI/Private/D3D12DescriptorCache.h
@@ -455,6 +455,9 @@ private:
 
 	TQueue<FD3D12GlobalHeapBlock*> FreeBlocks;
 	TArray<FD3D12GlobalHeapBlock*> ReleasedBlocks;
+	
+	FD3D12DescriptorHandleSRV* pNullSRV;
+	TArray<D3D12_CPU_DESCRIPTOR_HANDLE> EmptyDescriptors;
 
 	FCriticalSection CriticalSection;
 };
@@ -488,7 +491,7 @@ private:
 
 	D3D12_DESCRIPTOR_HEAP_TYPE HeapType = D3D12_DESCRIPTOR_HEAP_TYPE_NUM_TYPES;
 	FD3D12GlobalHeapBlock* CurrentBlock = nullptr;
-
+	
 	FD3D12DescriptorCache* DescriptorCache = nullptr;
 	FD3D12CommandListHandle CurrentCommandList;
 };
