Improved GPU crash reporting detail. Added exact address ranges and sizes. Added human-readable resource format to Aftermath report.

diff --git a/Engine/Source/Runtime/D3D12RHI/Private/D3D12RHI.cpp b/Engine/Source/Runtime/D3D12RHI/Private/D3D12RHI.cpp
--- a/Engine/Source/Runtime/D3D12RHI/Private/D3D12RHI.cpp
+++ b/Engine/Source/Runtime/D3D12RHI/Private/D3D12RHI.cpp
@@ -816,4 +816,134 @@ void FD3D12SubmissionGapRecorder::OnRenderThreadAdvanceFrame()
 		FrameRingbuffer[WriteIndexRT].bSafeToReadOnRenderThread = false;
 	}
 }
-#endif
\ No newline at end of file
+#endif
+
+const TCHAR* LexToString(DXGI_FORMAT Format)
+{
+	switch (Format)
+	{
+	default:
+	case DXGI_FORMAT_UNKNOWN: return TEXT("DXGI_FORMAT_UNKNOWN");
+	case DXGI_FORMAT_R32G32B32A32_TYPELESS: return TEXT("DXGI_FORMAT_R32G32B32A32_TYPELESS");
+	case DXGI_FORMAT_R32G32B32A32_FLOAT: return TEXT("DXGI_FORMAT_R32G32B32A32_FLOAT");
+	case DXGI_FORMAT_R32G32B32A32_UINT: return TEXT("DXGI_FORMAT_R32G32B32A32_UINT");
+	case DXGI_FORMAT_R32G32B32A32_SINT: return TEXT("DXGI_FORMAT_R32G32B32A32_SINT");
+	case DXGI_FORMAT_R32G32B32_TYPELESS: return TEXT("DXGI_FORMAT_R32G32B32_TYPELESS");
+	case DXGI_FORMAT_R32G32B32_FLOAT: return TEXT("DXGI_FORMAT_R32G32B32_FLOAT");
+	case DXGI_FORMAT_R32G32B32_UINT: return TEXT("DXGI_FORMAT_R32G32B32_UINT");
+	case DXGI_FORMAT_R32G32B32_SINT: return TEXT("DXGI_FORMAT_R32G32B32_SINT");
+	case DXGI_FORMAT_R16G16B16A16_TYPELESS: return TEXT("DXGI_FORMAT_R16G16B16A16_TYPELESS");
+	case DXGI_FORMAT_R16G16B16A16_FLOAT: return TEXT("DXGI_FORMAT_R16G16B16A16_FLOAT");
+	case DXGI_FORMAT_R16G16B16A16_UNORM: return TEXT("DXGI_FORMAT_R16G16B16A16_UNORM");
+	case DXGI_FORMAT_R16G16B16A16_UINT: return TEXT("DXGI_FORMAT_R16G16B16A16_UINT");
+	case DXGI_FORMAT_R16G16B16A16_SNORM: return TEXT("DXGI_FORMAT_R16G16B16A16_SNORM");
+	case DXGI_FORMAT_R16G16B16A16_SINT: return TEXT("DXGI_FORMAT_R16G16B16A16_SINT");
+	case DXGI_FORMAT_R32G32_TYPELESS: return TEXT("DXGI_FORMAT_R32G32_TYPELESS");
+	case DXGI_FORMAT_R32G32_FLOAT: return TEXT("DXGI_FORMAT_R32G32_FLOAT");
+	case DXGI_FORMAT_R32G32_UINT: return TEXT("DXGI_FORMAT_R32G32_UINT");
+	case DXGI_FORMAT_R32G32_SINT: return TEXT("DXGI_FORMAT_R32G32_SINT");
+	case DXGI_FORMAT_R32G8X24_TYPELESS: return TEXT("DXGI_FORMAT_R32G8X24_TYPELESS");
+	case DXGI_FORMAT_D32_FLOAT_S8X24_UINT: return TEXT("DXGI_FORMAT_D32_FLOAT_S8X24_UINT");
+	case DXGI_FORMAT_R32_FLOAT_X8X24_TYPELESS: return TEXT("DXGI_FORMAT_R32_FLOAT_X8X24_TYPELESS");
+	case DXGI_FORMAT_X32_TYPELESS_G8X24_UINT: return TEXT("DXGI_FORMAT_X32_TYPELESS_G8X24_UINT");
+	case DXGI_FORMAT_R10G10B10A2_TYPELESS: return TEXT("DXGI_FORMAT_R10G10B10A2_TYPELESS");
+	case DXGI_FORMAT_R10G10B10A2_UNORM: return TEXT("DXGI_FORMAT_R10G10B10A2_UNORM");
+	case DXGI_FORMAT_R10G10B10A2_UINT: return TEXT("DXGI_FORMAT_R10G10B10A2_UINT");
+	case DXGI_FORMAT_R11G11B10_FLOAT: return TEXT("DXGI_FORMAT_R11G11B10_FLOAT");
+	case DXGI_FORMAT_R8G8B8A8_TYPELESS: return TEXT("DXGI_FORMAT_R8G8B8A8_TYPELESS");
+	case DXGI_FORMAT_R8G8B8A8_UNORM: return TEXT("DXGI_FORMAT_R8G8B8A8_UNORM");
+	case DXGI_FORMAT_R8G8B8A8_UNORM_SRGB: return TEXT("DXGI_FORMAT_R8G8B8A8_UNORM_SRGB");
+	case DXGI_FORMAT_R8G8B8A8_UINT: return TEXT("DXGI_FORMAT_R8G8B8A8_UINT");
+	case DXGI_FORMAT_R8G8B8A8_SNORM: return TEXT("DXGI_FORMAT_R8G8B8A8_SNORM");
+	case DXGI_FORMAT_R8G8B8A8_SINT: return TEXT("DXGI_FORMAT_R8G8B8A8_SINT");
+	case DXGI_FORMAT_R16G16_TYPELESS: return TEXT("DXGI_FORMAT_R16G16_TYPELESS");
+	case DXGI_FORMAT_R16G16_FLOAT: return TEXT("DXGI_FORMAT_R16G16_FLOAT");
+	case DXGI_FORMAT_R16G16_UNORM: return TEXT("DXGI_FORMAT_R16G16_UNORM");
+	case DXGI_FORMAT_R16G16_UINT: return TEXT("DXGI_FORMAT_R16G16_UINT");
+	case DXGI_FORMAT_R16G16_SNORM: return TEXT("DXGI_FORMAT_R16G16_SNORM");
+	case DXGI_FORMAT_R16G16_SINT: return TEXT("DXGI_FORMAT_R16G16_SINT");
+	case DXGI_FORMAT_R32_TYPELESS: return TEXT("DXGI_FORMAT_R32_TYPELESS");
+	case DXGI_FORMAT_D32_FLOAT: return TEXT("DXGI_FORMAT_D32_FLOAT");
+	case DXGI_FORMAT_R32_FLOAT: return TEXT("DXGI_FORMAT_R32_FLOAT");
+	case DXGI_FORMAT_R32_UINT: return TEXT("DXGI_FORMAT_R32_UINT");
+	case DXGI_FORMAT_R32_SINT: return TEXT("DXGI_FORMAT_R32_SINT");
+	case DXGI_FORMAT_R24G8_TYPELESS: return TEXT("DXGI_FORMAT_R24G8_TYPELESS");
+	case DXGI_FORMAT_D24_UNORM_S8_UINT: return TEXT("DXGI_FORMAT_D24_UNORM_S8_UINT");
+	case DXGI_FORMAT_R24_UNORM_X8_TYPELESS: return TEXT("DXGI_FORMAT_R24_UNORM_X8_TYPELESS");
+	case DXGI_FORMAT_X24_TYPELESS_G8_UINT: return TEXT("DXGI_FORMAT_X24_TYPELESS_G8_UINT");
+	case DXGI_FORMAT_R8G8_TYPELESS: return TEXT("DXGI_FORMAT_R8G8_TYPELESS");
+	case DXGI_FORMAT_R8G8_UNORM: return TEXT("DXGI_FORMAT_R8G8_UNORM");
+	case DXGI_FORMAT_R8G8_UINT: return TEXT("DXGI_FORMAT_R8G8_UINT");
+	case DXGI_FORMAT_R8G8_SNORM: return TEXT("DXGI_FORMAT_R8G8_SNORM");
+	case DXGI_FORMAT_R8G8_SINT: return TEXT("DXGI_FORMAT_R8G8_SINT");
+	case DXGI_FORMAT_R16_TYPELESS: return TEXT("DXGI_FORMAT_R16_TYPELESS");
+	case DXGI_FORMAT_R16_FLOAT: return TEXT("DXGI_FORMAT_R16_FLOAT");
+	case DXGI_FORMAT_D16_UNORM: return TEXT("DXGI_FORMAT_D16_UNORM");
+	case DXGI_FORMAT_R16_UNORM: return TEXT("DXGI_FORMAT_R16_UNORM");
+	case DXGI_FORMAT_R16_UINT: return TEXT("DXGI_FORMAT_R16_UINT");
+	case DXGI_FORMAT_R16_SNORM: return TEXT("DXGI_FORMAT_R16_SNORM");
+	case DXGI_FORMAT_R16_SINT: return TEXT("DXGI_FORMAT_R16_SINT");
+	case DXGI_FORMAT_R8_TYPELESS: return TEXT("DXGI_FORMAT_R8_TYPELESS");
+	case DXGI_FORMAT_R8_UNORM: return TEXT("DXGI_FORMAT_R8_UNORM");
+	case DXGI_FORMAT_R8_UINT: return TEXT("DXGI_FORMAT_R8_UINT");
+	case DXGI_FORMAT_R8_SNORM: return TEXT("DXGI_FORMAT_R8_SNORM");
+	case DXGI_FORMAT_R8_SINT: return TEXT("DXGI_FORMAT_R8_SINT");
+	case DXGI_FORMAT_A8_UNORM: return TEXT("DXGI_FORMAT_A8_UNORM");
+	case DXGI_FORMAT_R1_UNORM: return TEXT("DXGI_FORMAT_R1_UNORM");
+	case DXGI_FORMAT_R9G9B9E5_SHAREDEXP: return TEXT("DXGI_FORMAT_R9G9B9E5_SHAREDEXP");
+	case DXGI_FORMAT_R8G8_B8G8_UNORM: return TEXT("DXGI_FORMAT_R8G8_B8G8_UNORM");
+	case DXGI_FORMAT_G8R8_G8B8_UNORM: return TEXT("DXGI_FORMAT_G8R8_G8B8_UNORM");
+	case DXGI_FORMAT_BC1_TYPELESS: return TEXT("DXGI_FORMAT_BC1_TYPELESS");
+	case DXGI_FORMAT_BC1_UNORM: return TEXT("DXGI_FORMAT_BC1_UNORM");
+	case DXGI_FORMAT_BC1_UNORM_SRGB: return TEXT("DXGI_FORMAT_BC1_UNORM_SRGB");
+	case DXGI_FORMAT_BC2_TYPELESS: return TEXT("DXGI_FORMAT_BC2_TYPELESS");
+	case DXGI_FORMAT_BC2_UNORM: return TEXT("DXGI_FORMAT_BC2_UNORM");
+	case DXGI_FORMAT_BC2_UNORM_SRGB: return TEXT("DXGI_FORMAT_BC2_UNORM_SRGB");
+	case DXGI_FORMAT_BC3_TYPELESS: return TEXT("DXGI_FORMAT_BC3_TYPELESS");
+	case DXGI_FORMAT_BC3_UNORM: return TEXT("DXGI_FORMAT_BC3_UNORM");
+	case DXGI_FORMAT_BC3_UNORM_SRGB: return TEXT("DXGI_FORMAT_BC3_UNORM_SRGB");
+	case DXGI_FORMAT_BC4_TYPELESS: return TEXT("DXGI_FORMAT_BC4_TYPELESS");
+	case DXGI_FORMAT_BC4_UNORM: return TEXT("DXGI_FORMAT_BC4_UNORM");
+	case DXGI_FORMAT_BC4_SNORM: return TEXT("DXGI_FORMAT_BC4_SNORM");
+	case DXGI_FORMAT_BC5_TYPELESS: return TEXT("DXGI_FORMAT_BC5_TYPELESS");
+	case DXGI_FORMAT_BC5_UNORM: return TEXT("DXGI_FORMAT_BC5_UNORM");
+	case DXGI_FORMAT_BC5_SNORM: return TEXT("DXGI_FORMAT_BC5_SNORM");
+	case DXGI_FORMAT_B5G6R5_UNORM: return TEXT("DXGI_FORMAT_B5G6R5_UNORM");
+	case DXGI_FORMAT_B5G5R5A1_UNORM: return TEXT("DXGI_FORMAT_B5G5R5A1_UNORM");
+	case DXGI_FORMAT_B8G8R8A8_UNORM: return TEXT("DXGI_FORMAT_B8G8R8A8_UNORM");
+	case DXGI_FORMAT_B8G8R8X8_UNORM: return TEXT("DXGI_FORMAT_B8G8R8X8_UNORM");
+	case DXGI_FORMAT_R10G10B10_XR_BIAS_A2_UNORM: return TEXT("DXGI_FORMAT_R10G10B10_XR_BIAS_A2_UNORM");
+	case DXGI_FORMAT_B8G8R8A8_TYPELESS: return TEXT("DXGI_FORMAT_B8G8R8A8_TYPELESS");
+	case DXGI_FORMAT_B8G8R8A8_UNORM_SRGB: return TEXT("DXGI_FORMAT_B8G8R8A8_UNORM_SRGB");
+	case DXGI_FORMAT_B8G8R8X8_TYPELESS: return TEXT("DXGI_FORMAT_B8G8R8X8_TYPELESS");
+	case DXGI_FORMAT_B8G8R8X8_UNORM_SRGB: return TEXT("DXGI_FORMAT_B8G8R8X8_UNORM_SRGB");
+	case DXGI_FORMAT_BC6H_TYPELESS: return TEXT("DXGI_FORMAT_BC6H_TYPELESS");
+	case DXGI_FORMAT_BC6H_UF16: return TEXT("DXGI_FORMAT_BC6H_UF16");
+	case DXGI_FORMAT_BC6H_SF16: return TEXT("DXGI_FORMAT_BC6H_SF16");
+	case DXGI_FORMAT_BC7_TYPELESS: return TEXT("DXGI_FORMAT_BC7_TYPELESS");
+	case DXGI_FORMAT_BC7_UNORM: return TEXT("DXGI_FORMAT_BC7_UNORM");
+	case DXGI_FORMAT_BC7_UNORM_SRGB: return TEXT("DXGI_FORMAT_BC7_UNORM_SRGB");
+	case DXGI_FORMAT_AYUV: return TEXT("DXGI_FORMAT_AYUV");
+	case DXGI_FORMAT_Y410: return TEXT("DXGI_FORMAT_Y410");
+	case DXGI_FORMAT_Y416: return TEXT("DXGI_FORMAT_Y416");
+	case DXGI_FORMAT_NV12: return TEXT("DXGI_FORMAT_NV12");
+	case DXGI_FORMAT_P010: return TEXT("DXGI_FORMAT_P010");
+	case DXGI_FORMAT_P016: return TEXT("DXGI_FORMAT_P016");
+	case DXGI_FORMAT_420_OPAQUE: return TEXT("DXGI_FORMAT_420_OPAQUE");
+	case DXGI_FORMAT_YUY2: return TEXT("DXGI_FORMAT_YUY2");
+	case DXGI_FORMAT_Y210: return TEXT("DXGI_FORMAT_Y210");
+	case DXGI_FORMAT_Y216: return TEXT("DXGI_FORMAT_Y216");
+	case DXGI_FORMAT_NV11: return TEXT("DXGI_FORMAT_NV11");
+	case DXGI_FORMAT_AI44: return TEXT("DXGI_FORMAT_AI44");
+	case DXGI_FORMAT_IA44: return TEXT("DXGI_FORMAT_IA44");
+	case DXGI_FORMAT_P8: return TEXT("DXGI_FORMAT_P8");
+	case DXGI_FORMAT_A8P8: return TEXT("DXGI_FORMAT_A8P8");
+	case DXGI_FORMAT_B4G4R4A4_UNORM: return TEXT("DXGI_FORMAT_B4G4R4A4_UNORM");
+	case DXGI_FORMAT_P208: return TEXT("DXGI_FORMAT_P208");
+	case DXGI_FORMAT_V208: return TEXT("DXGI_FORMAT_V208");
+	case DXGI_FORMAT_V408: return TEXT("DXGI_FORMAT_V408");
+	case DXGI_FORMAT_SAMPLER_FEEDBACK_MIN_MIP_OPAQUE: return TEXT("DXGI_FORMAT_SAMPLER_FEEDBACK_MIN_MIP_OPAQUE");
+	case DXGI_FORMAT_SAMPLER_FEEDBACK_MIP_REGION_USED_OPAQUE: return TEXT("DXGI_FORMAT_SAMPLER_FEEDBACK_MIP_REGION_USED_OPAQUE");
+	}
+}
+
diff --git a/Engine/Source/Runtime/D3D12RHI/Private/D3D12RHIPrivate.h b/Engine/Source/Runtime/D3D12RHI/Private/D3D12RHIPrivate.h
--- a/Engine/Source/Runtime/D3D12RHI/Private/D3D12RHIPrivate.h
+++ b/Engine/Source/Runtime/D3D12RHI/Private/D3D12RHIPrivate.h
@@ -1511,6 +1511,9 @@ static void TranslateRenderTargetFormats(
 
 } // namespace D3D12RHI
 
+// Returns the given format as a string. Unsupported formats are treated as DXGI_FORMAT_UNKNOWN.
+const TCHAR* LexToString(DXGI_FORMAT Format);
+
 #if (PLATFORM_WINDOWS || PLATFORM_HOLOLENS)
 
 #ifndef DXGI_PRESENT_ALLOW_TEARING
diff --git a/Engine/Source/Runtime/D3D12RHI/Private/D3D12Stats.cpp b/Engine/Source/Runtime/D3D12RHI/Private/D3D12Stats.cpp
--- a/Engine/Source/Runtime/D3D12RHI/Private/D3D12Stats.cpp
+++ b/Engine/Source/Runtime/D3D12RHI/Private/D3D12Stats.cpp
@@ -385,7 +385,15 @@ bool D3D12RHI::FD3DGPUProfiler::CheckGpuHeartbeat() const
 						UE_LOG(LogRHI, Error, TEXT("[Aftermath] Faulting resource dims: %d x %d x %d"), FaultInformation.resourceDesc.width, FaultInformation.resourceDesc.height, FaultInformation.resourceDesc.depth);
 						UE_LOG(LogRHI, Error, TEXT("[Aftermath] Faulting result size: %llu bytes"), FaultInformation.resourceDesc.size);
 						UE_LOG(LogRHI, Error, TEXT("[Aftermath] Faulting resource mips: %d"), FaultInformation.resourceDesc.mipLevels);
-						UE_LOG(LogRHI, Error, TEXT("[Aftermath] Faulting resource format: 0x%x"), FaultInformation.resourceDesc.format);
+
+						DXGI_FORMAT ResourceFormat = (DXGI_FORMAT)FaultInformation.resourceDesc.format;
+						const TCHAR* FormatStr = LexToString(ResourceFormat);
+						const TCHAR* FormatPrefix = TEXT("DXGI_FORMAT_");
+						if (FCString::Strstr(FormatStr, FormatPrefix) == FormatStr)
+						{
+							FormatStr += FCString::Strlen(FormatPrefix);
+						}
+						UE_LOG(LogRHI, Error, TEXT("[Aftermath] Faulting resource format: %s (0x%x)"), FormatStr, (int32)ResourceFormat);
 					}
 					else
 					{
diff --git a/Engine/Source/Runtime/D3D12RHI/Private/D3D12Util.cpp b/Engine/Source/Runtime/D3D12RHI/Private/D3D12Util.cpp
--- a/Engine/Source/Runtime/D3D12RHI/Private/D3D12Util.cpp
+++ b/Engine/Source/Runtime/D3D12RHI/Private/D3D12Util.cpp
@@ -578,8 +578,15 @@ void LogPageFaultData(FD3D12Adapter* InAdapter, D3D12_GPU_VIRTUAL_ADDRESS InPage
 		{
 			FD3D12Adapter::FAllocatedResourceResult OverlappingResource = OverlappingResources[Index];
 			D3D12_GPU_VIRTUAL_ADDRESS ResourceAddress = OverlappingResource.Allocation->GetGPUVirtualAddress();
-			UE_LOG(LogD3D12RHI, Error, TEXT("\tGPU Address: \"0x%llX\" - Size: %3.2f MB - Distance to page fault: %3.2f MB - Transient: %d - Name: %s"),
-				(long long)ResourceAddress, OverlappingResource.Allocation->GetSize() / (1024.0f * 1024), OverlappingResource.Distance / (1024.0f * 1024), OverlappingResource.Allocation->IsTransient(), *OverlappingResource.Allocation->GetResource()->GetName().ToString());
+			UE_LOG(LogD3D12RHI, Error, TEXT("\tGPU Address: [0x%llX .. 0x%llX] - Size: %lld bytes, %3.2f MB - Distance to page fault: %lld bytes, %3.2f MB - Transient: %d - Name: %s"),
+				(uint64)ResourceAddress,
+				(uint64)ResourceAddress + OverlappingResource.Allocation->GetSize(),
+				OverlappingResource.Allocation->GetSize(),
+				OverlappingResource.Allocation->GetSize() / (1024.0f * 1024), 
+				OverlappingResource.Distance,
+				OverlappingResource.Distance / (1024.0f * 1024), 
+				OverlappingResource.Allocation->IsTransient(), 
+				*OverlappingResource.Allocation->GetResource()->GetName().ToString());
 		}
 	}
 
@@ -593,8 +600,15 @@ void LogPageFaultData(FD3D12Adapter* InAdapter, D3D12_GPU_VIRTUAL_ADDRESS InPage
 		for (uint32 Index = 0; Index < PrintCount; ++Index)
 		{
 			FD3D12Adapter::FReleasedAllocationData& AllocationData = ReleasedResources[Index];
-			UE_LOG(LogD3D12RHI, Error, TEXT("\tGPU Address: \"0x%llX\" - Size: %3.2f MB - FrameID: %4d - DefragFree: %d - Transient: %d - Name: %s"),
-				(long long)AllocationData.GPUVirtualAddress, AllocationData.AllocationSize / (1024.0f * 1024), AllocationData.ReleasedFrameID, AllocationData.bDefragFree, AllocationData.bTransient, *AllocationData.ResourceName.ToString());
+			UE_LOG(LogD3D12RHI, Error, TEXT("\tGPU Address: [0x%llX .. 0x%llX] - Size: %lld bytes, %3.2f MB - FrameID: %4d - DefragFree: %d - Transient: %d - Name: %s"),
+				(uint64)AllocationData.GPUVirtualAddress,
+				(uint64)AllocationData.GPUVirtualAddress + AllocationData.AllocationSize,
+				AllocationData.AllocationSize,
+				AllocationData.AllocationSize / (1024.0f * 1024),
+				AllocationData.ReleasedFrameID,
+				AllocationData.bDefragFree,
+				AllocationData.bTransient,
+				*AllocationData.ResourceName.ToString());
 		}
 	}
 }
