Fixed compile error due to DXGI_FORMAT_SAMPLER_FEEDBACK_xxx format identifiers not being present (older D3D12 headers used)

diff --git a/Engine/Source/Runtime/D3D12RHI/Private/D3D12RHI.cpp b/Engine/Source/Runtime/D3D12RHI/Private/D3D12RHI.cpp
--- a/Engine/Source/Runtime/D3D12RHI/Private/D3D12RHI.cpp
+++ b/Engine/Source/Runtime/D3D12RHI/Private/D3D12RHI.cpp
@@ -942,8 +942,8 @@ const TCHAR* LexToString(DXGI_FORMAT Format)
 	case DXGI_FORMAT_P208: return TEXT("DXGI_FORMAT_P208");
 	case DXGI_FORMAT_V208: return TEXT("DXGI_FORMAT_V208");
 	case DXGI_FORMAT_V408: return TEXT("DXGI_FORMAT_V408");
-	case DXGI_FORMAT_SAMPLER_FEEDBACK_MIN_MIP_OPAQUE: return TEXT("DXGI_FORMAT_SAMPLER_FEEDBACK_MIN_MIP_OPAQUE");
-	case DXGI_FORMAT_SAMPLER_FEEDBACK_MIP_REGION_USED_OPAQUE: return TEXT("DXGI_FORMAT_SAMPLER_FEEDBACK_MIP_REGION_USED_OPAQUE");
+	case 189: return TEXT("DXGI_FORMAT_SAMPLER_FEEDBACK_MIN_MIP_OPAQUE");
+	case 190: return TEXT("DXGI_FORMAT_SAMPLER_FEEDBACK_MIP_REGION_USED_OPAQUE");
 	}
 }
 
