- Fix clearing of descriptors with correct CopyDescriptor call
- Use FrameFence instead of Command list fence for reuse because there is clearly an issue with this fence somehow (work around for now)

diff --git a/Engine/Source/Runtime/D3D12RHI/Private/D3D12DescriptorCache.cpp b/Engine/Source/Runtime/D3D12RHI/Private/D3D12DescriptorCache.cpp
--- a/Engine/Source/Runtime/D3D12RHI/Private/D3D12DescriptorCache.cpp
+++ b/Engine/Source/Runtime/D3D12RHI/Private/D3D12DescriptorCache.cpp
@@ -13,6 +13,18 @@ static FAutoConsoleVariableRef CVarGlobalViewHeapBlockSize(
 	ECVF_ReadOnly
 );
 
+#if PLATFORM_WINDOWS
+
+int32 GGlobalViewHeapClobberAfterUse = 1;
+static FAutoConsoleVariableRef CVarGlobalViewHeapClobberSRVAfterUse(
+	TEXT("D3D12.GlobalViewHeapClobberSRVAfterUse"),
+	GGlobalViewHeapClobberAfterUse,
+	TEXT("Workaround for current outstanding suspected NVIDIA driver bug. JIRA UE-108339"),
+	ECVF_ReadOnly
+);
+
+#endif // PLATFORM_WINDOWS
+
 // Define template functions that are only declared in the header.
 #if USE_STATIC_ROOT_SIGNATURE
 template void FD3D12DescriptorCache::SetConstantBuffers<SF_Vertex>(const FD3D12RootSignature* RootSignature, FD3D12ConstantBufferCache& Cache, const CBVSlotMask& SlotsNeededMask, uint32 Count, uint32& HeapSlot);
@@ -1031,6 +1043,17 @@ void FD3D12GlobalHeap::Init(D3D12_DESCRIPTOR_HEAP_TYPE InType, uint32 InTotalSiz
 		FreeBlocks.Enqueue(new FD3D12GlobalHeapBlock(CurrentBaseSlot, ActualBlockSize));
 		CurrentBaseSlot += ActualBlockSize;
 	}
+
+	// Create a default SRV view
+	D3D12_SHADER_RESOURCE_VIEW_DESC SRVDesc = {};
+	SRVDesc.Shader4ComponentMapping = D3D12_DEFAULT_SHADER_4_COMPONENT_MAPPING;
+	SRVDesc.ViewDimension = D3D12_SRV_DIMENSION_TEXTURE2D;
+	SRVDesc.Format = DXGI_FORMAT_B8G8R8A8_UNORM;
+	SRVDesc.Texture2D.MipLevels = 1;
+	SRVDesc.Texture2D.MostDetailedMip = 0;
+	SRVDesc.Texture2D.ResourceMinLODClamp = 0.0f;
+	NullSRVDescriptor = new FD3D12DescriptorHandleSRV(GetParentDevice());
+	NullSRVDescriptor->CreateView(SRVDesc, nullptr);
 }
 
 
@@ -1074,6 +1097,8 @@ void FD3D12GlobalHeap::FreeHeapBlock(FD3D12GlobalHeapBlock* InHeapBlock)
 	INC_DWORD_STAT_BY(STAT_GlobalViewHeapUsedDescriptors, InHeapBlock->SizeUsed);
 	INC_DWORD_STAT_BY(STAT_GlobalViewHeapWastedDescriptors, InHeapBlock->Size - InHeapBlock->SizeUsed);
 
+	InHeapBlock->FrameFence = GetParentDevice()->GetParentAdapter()->GetFrameFence().GetCurrentFence();
+
 	ReleasedBlocks.Add(InHeapBlock);
 }
 
@@ -1083,17 +1108,40 @@ Find all the blocks which are not used by the GPU anymore
 **/
 void FD3D12GlobalHeap::UpdateFreeBlocks()
 {
+	FD3D12ManualFence& FrameFence = GetParentDevice()->GetParentAdapter()->GetFrameFence();
 	for (int32 BlockIndex = 0; BlockIndex < ReleasedBlocks.Num(); ++BlockIndex)
 	{
 		// Check if GPU is ready consuming the block data
 		FD3D12GlobalHeapBlock* ReleasedBlock = ReleasedBlocks[BlockIndex];
-		if (ReleasedBlock->SyncPoint.IsComplete())
+		if (ReleasedBlock->SyncPoint.IsComplete() && FrameFence.IsFenceComplete(ReleasedBlock->FrameFence))
 		{
 			// Update stats
 			DEC_DWORD_STAT_BY(STAT_GlobalViewHeapUsedDescriptors, ReleasedBlock->SizeUsed);
 			DEC_DWORD_STAT_BY(STAT_GlobalViewHeapWastedDescriptors, ReleasedBlock->Size - ReleasedBlock->SizeUsed);
 			INC_DWORD_STAT_BY(STAT_GlobalViewHeapFreeDescriptors, ReleasedBlock->Size);
 
+#if PLATFORM_WINDOWS
+			// Clear the global descriptor heap block with null srvs to make sure they are not referencing some possibly deleted resources/heaps anymore (we know the GPU should be done accessing them so copy is valid).
+			// This fixes a rare page fault GPU crash on nVidia GPUs.
+			// Clearing the descriptors on reuse (during AllocateBlock) still crashes so current assumption is some random/invalid access done by the driver.
+			// See: JIRA UE-108339
+			if (GGlobalViewHeapClobberAfterUse)
+			{
+				ID3D12Device* Device = GetParentDevice()->GetDevice();
+				D3D12_CPU_DESCRIPTOR_HANDLE TargetCPUBase = GetCPUSlotHandle(ReleasedBlock);
+				D3D12_CPU_DESCRIPTOR_HANDLE SourceNullSRVHandle = NullSRVDescriptor->GetHandle();
+				for (uint32 SlotIndex = 0; SlotIndex < ReleasedBlock->Size; SlotIndex++)
+				{
+					D3D12_CPU_DESCRIPTOR_HANDLE TargetCPU = { TargetCPUBase.ptr + SlotIndex * DescriptorSize };
+					Device->CopyDescriptorsSimple(
+						1,
+						TargetCPU,
+						SourceNullSRVHandle,
+						D3D12_DESCRIPTOR_HEAP_TYPE_CBV_SRV_UAV);
+				}
+			}
+#endif  // PLATFORM_WINDOWS
+			
 			ReleasedBlock->SizeUsed = 0;
 			FreeBlocks.Enqueue(ReleasedBlock);
 
diff --git a/Engine/Source/Runtime/D3D12RHI/Private/D3D12DescriptorCache.h b/Engine/Source/Runtime/D3D12RHI/Private/D3D12DescriptorCache.h
--- a/Engine/Source/Runtime/D3D12RHI/Private/D3D12DescriptorCache.h
+++ b/Engine/Source/Runtime/D3D12RHI/Private/D3D12DescriptorCache.h
@@ -414,6 +414,8 @@ public:
 	uint32 Size;
 	uint32 SizeUsed;
 	FD3D12CLSyncPoint SyncPoint;
+
+	uint64 FrameFence = 0;
 };
 
 
@@ -456,6 +458,8 @@ private:
 	TQueue<FD3D12GlobalHeapBlock*> FreeBlocks;
 	TArray<FD3D12GlobalHeapBlock*> ReleasedBlocks;
 
+	FD3D12DescriptorHandleSRV* NullSRVDescriptor;
+
 	FCriticalSection CriticalSection;
 };
 
