- Only remove sync point from active generations when it's actually done instead of before waiting to finish - can otherwise cause invalid return true during call to IsComplete() when it can't find this generation ID anymore
- Remove 'unneeded' frame fence when checking if global descriptor heap block is done (was added because sync point code wasn't safe)

diff --git a/Engine/Source/Runtime/D3D12RHI/Private/D3D12CommandList.h b/Engine/Source/Runtime/D3D12RHI/Private/D3D12CommandList.h
--- a/Engine/Source/Runtime/D3D12RHI/Private/D3D12CommandList.h
+++ b/Engine/Source/Runtime/D3D12RHI/Private/D3D12CommandList.h
@@ -85,6 +85,8 @@ private:
 					if (Generation < GenerationSyncPoint.Key)
 					{
 						// The requested generation is older than the oldest tracked generation, so it must be complete.
+						// (LastCompleteGeneration probably updated while trying to take the lock)
+						check(Generation <= LastCompleteGeneration);
 						return true;
 					}
 					else
@@ -102,6 +104,11 @@ private:
 						}
 					}
 				}
+				else
+				{
+					// LastCompleteGeneration probably updated while trying to take the lock, validate that the GPU is actually processed this generation
+					check(Generation <= LastCompleteGeneration);
+				}
 			}
 
 			return true;
@@ -120,7 +127,6 @@ private:
 					while (ActiveGenerations.Peek(GenerationSyncPoint) && (Generation > LastCompleteGeneration))
 					{
 						check(Generation >= GenerationSyncPoint.Key);
-						ActiveGenerations.Dequeue(GenerationSyncPoint);
 
 						// Unblock other threads while we wait for the command list to complete
 						ActiveGenerationsCS.Unlock();
@@ -128,6 +134,7 @@ private:
 						GenerationSyncPoint.Value.WaitForCompletion();
 
 						ActiveGenerationsCS.Lock();
+						ActiveGenerations.Dequeue(GenerationSyncPoint);
 						LastCompleteGeneration = FMath::Max(LastCompleteGeneration, GenerationSyncPoint.Key);
 					}
 				}
diff --git a/Engine/Source/Runtime/D3D12RHI/Private/D3D12DescriptorCache.cpp b/Engine/Source/Runtime/D3D12RHI/Private/D3D12DescriptorCache.cpp
--- a/Engine/Source/Runtime/D3D12RHI/Private/D3D12DescriptorCache.cpp
+++ b/Engine/Source/Runtime/D3D12RHI/Private/D3D12DescriptorCache.cpp
@@ -1097,8 +1097,6 @@ void FD3D12GlobalHeap::FreeHeapBlock(FD3D12GlobalHeapBlock* InHeapBlock)
 	INC_DWORD_STAT_BY(STAT_GlobalViewHeapUsedDescriptors, InHeapBlock->SizeUsed);
 	INC_DWORD_STAT_BY(STAT_GlobalViewHeapWastedDescriptors, InHeapBlock->Size - InHeapBlock->SizeUsed);
 
-	InHeapBlock->FrameFence = GetParentDevice()->GetParentAdapter()->GetFrameFence().GetCurrentFence();
-
 	ReleasedBlocks.Add(InHeapBlock);
 }
 
@@ -1108,12 +1106,11 @@ Find all the blocks which are not used by the GPU anymore
 **/
 void FD3D12GlobalHeap::UpdateFreeBlocks()
 {
-	FD3D12ManualFence& FrameFence = GetParentDevice()->GetParentAdapter()->GetFrameFence();
 	for (int32 BlockIndex = 0; BlockIndex < ReleasedBlocks.Num(); ++BlockIndex)
 	{
 		// Check if GPU is ready consuming the block data
 		FD3D12GlobalHeapBlock* ReleasedBlock = ReleasedBlocks[BlockIndex];
-		if (ReleasedBlock->SyncPoint.IsComplete() && FrameFence.IsFenceComplete(ReleasedBlock->FrameFence))
+		if (ReleasedBlock->SyncPoint.IsComplete())
 		{
 			// Update stats
 			DEC_DWORD_STAT_BY(STAT_GlobalViewHeapUsedDescriptors, ReleasedBlock->SizeUsed);
@@ -1141,7 +1138,7 @@ void FD3D12GlobalHeap::UpdateFreeBlocks()
 				}
 			}
 #endif  // PLATFORM_WINDOWS
-			
+
 			ReleasedBlock->SizeUsed = 0;
 			FreeBlocks.Enqueue(ReleasedBlock);
 
diff --git a/Engine/Source/Runtime/D3D12RHI/Private/D3D12DescriptorCache.h b/Engine/Source/Runtime/D3D12RHI/Private/D3D12DescriptorCache.h
--- a/Engine/Source/Runtime/D3D12RHI/Private/D3D12DescriptorCache.h
+++ b/Engine/Source/Runtime/D3D12RHI/Private/D3D12DescriptorCache.h
@@ -414,8 +414,6 @@ public:
 	uint32 Size;
 	uint32 SizeUsed;
 	FD3D12CLSyncPoint SyncPoint;
-
-	uint64 FrameFence = 0;
 };
 
 
