Extending Mesh shader support.

diff --git a/Engine/Source/Runtime/D3D12RHI/D3D12RHI.Build.cs b/Engine/Source/Runtime/D3D12RHI/D3D12RHI.Build.cs
--- a/Engine/Source/Runtime/D3D12RHI/D3D12RHI.Build.cs
+++ b/Engine/Source/Runtime/D3D12RHI/D3D12RHI.Build.cs
@@ -33,7 +33,7 @@ public class D3D12RHI : ModuleRules
         // Platform specific defines
         ///////////////////////////////////////////////////////////////
 
-        if (!Target.Platform.IsInGroup(UnrealPlatformGroup.Windows))
+		if (!Target.Platform.IsInGroup(UnrealPlatformGroup.Windows))
         {
             PrecompileForTargets = PrecompileTargetsType.None;
         }
@@ -41,6 +41,8 @@ public class D3D12RHI : ModuleRules
         if (Target.Platform.IsInGroup(UnrealPlatformGroup.Windows) ||
             Target.Platform == UnrealTargetPlatform.HoloLens)
 		{
+			PrivateDefinitions.Add("D3D12RHI_USE_PIPELINE_STATE_STREAM=1");
+
 			AddEngineThirdPartyPrivateStaticDependencies(Target, "DX12");
 			AddEngineThirdPartyPrivateStaticDependencies(Target, "DX11");
 			AddEngineThirdPartyPrivateStaticDependencies(Target, "AMD_AGS");
diff --git a/Engine/Source/Runtime/D3D12RHI/Private/D3D12Adapter.cpp b/Engine/Source/Runtime/D3D12RHI/Private/D3D12Adapter.cpp
--- a/Engine/Source/Runtime/D3D12RHI/Private/D3D12Adapter.cpp
+++ b/Engine/Source/Runtime/D3D12RHI/Private/D3D12Adapter.cpp
@@ -829,35 +829,56 @@ void FD3D12Adapter::InitializeDevices()
 		}
 		else
 		{
+#if D3D12_MAX_DEVICE_INTERFACE >= 1
 			if (SUCCEEDED(RootDevice->QueryInterface(IID_PPV_ARGS(RootDevice1.GetInitReference()))))
 			{
 				UE_LOG(LogD3D12RHI, Log, TEXT("ID3D12Device1 is supported."));
 			}
-
-#if PLATFORM_WINDOWS || PLATFORM_HOLOLENS
+#endif
+#if D3D12_MAX_DEVICE_INTERFACE >= 2
 			if (SUCCEEDED(RootDevice->QueryInterface(IID_PPV_ARGS(RootDevice2.GetInitReference()))))
 			{
 				UE_LOG(LogD3D12RHI, Log, TEXT("ID3D12Device2 is supported."));
 			}
 #endif
-
-			D3D12_FEATURE_DATA_D3D12_OPTIONS D3D12Caps;
-			FMemory::Memzero(&D3D12Caps, sizeof(D3D12Caps));
-			VERIFYD3D12RESULT(RootDevice->CheckFeatureSupport(D3D12_FEATURE_D3D12_OPTIONS, &D3D12Caps, sizeof(D3D12Caps)));
-			ResourceHeapTier = D3D12Caps.ResourceHeapTier;
-			ResourceBindingTier = D3D12Caps.ResourceBindingTier;
-
-#if D3D12_RHI_RAYTRACING
+#if D3D12_MAX_DEVICE_INTERFACE >= 3
+			if (SUCCEEDED(RootDevice->QueryInterface(IID_PPV_ARGS(RootDevice3.GetInitReference()))))
+			{
+				UE_LOG(LogD3D12RHI, Log, TEXT("ID3D12Device3 is supported."));
+			}
+#endif
+#if D3D12_MAX_DEVICE_INTERFACE >= 4
+			if (SUCCEEDED(RootDevice->QueryInterface(IID_PPV_ARGS(RootDevice4.GetInitReference()))))
+			{
+				UE_LOG(LogD3D12RHI, Log, TEXT("ID3D12Device4 is supported."));
+			}
+#endif
+#if D3D12_MAX_DEVICE_INTERFACE >= 5
 			if (SUCCEEDED(RootDevice->QueryInterface(IID_PPV_ARGS(RootDevice5.GetInitReference()))))
 			{
 				UE_LOG(LogD3D12RHI, Log, TEXT("ID3D12Device5 is supported."));
 			}
-
+#endif
+#if D3D12_MAX_DEVICE_INTERFACE >= 6
+			if (SUCCEEDED(RootDevice->QueryInterface(IID_PPV_ARGS(RootDevice6.GetInitReference()))))
+			{
+				UE_LOG(LogD3D12RHI, Log, TEXT("ID3D12Device6 is supported."));
+			}
+#endif
+#if D3D12_MAX_DEVICE_INTERFACE >= 7
 			if (SUCCEEDED(RootDevice->QueryInterface(IID_PPV_ARGS(RootDevice7.GetInitReference()))))
 			{
 				UE_LOG(LogD3D12RHI, Log, TEXT("ID3D12Device7 is supported."));
 			}
+#endif
 
+			D3D12_FEATURE_DATA_D3D12_OPTIONS D3D12Caps;
+			FMemory::Memzero(&D3D12Caps, sizeof(D3D12Caps));
+			VERIFYD3D12RESULT(RootDevice->CheckFeatureSupport(D3D12_FEATURE_D3D12_OPTIONS, &D3D12Caps, sizeof(D3D12Caps)));
+			ResourceHeapTier = D3D12Caps.ResourceHeapTier;
+			ResourceBindingTier = D3D12Caps.ResourceBindingTier;
+
+#if D3D12_RHI_RAYTRACING
 			D3D12_FEATURE_DATA_D3D12_OPTIONS5 D3D12Caps5 = {};
 			if (SUCCEEDED(RootDevice->CheckFeatureSupport(D3D12_FEATURE_D3D12_OPTIONS5, &D3D12Caps5, sizeof(D3D12Caps5))))
 			{
diff --git a/Engine/Source/Runtime/D3D12RHI/Private/D3D12Adapter.h b/Engine/Source/Runtime/D3D12RHI/Private/D3D12Adapter.h
--- a/Engine/Source/Runtime/D3D12RHI/Private/D3D12Adapter.h
+++ b/Engine/Source/Runtime/D3D12RHI/Private/D3D12Adapter.h
@@ -121,12 +121,16 @@ public:
 	FORCEINLINE const D3D_FEATURE_LEVEL GetFeatureLevel() const { return Desc.MaxSupportedFeatureLevel; }
 	FORCEINLINE D3D_SHADER_MODEL GetHighestShaderModel() const { return Desc.MaxSupportedShaderModel; }
 	FORCEINLINE ID3D12Device* GetD3DDevice() const { return RootDevice.GetReference(); }
+#if D3D12_MAX_DEVICE_INTERFACE >= 1
 	FORCEINLINE ID3D12Device1* GetD3DDevice1() const { return RootDevice1.GetReference(); }
-#if PLATFORM_WINDOWS || PLATFORM_HOLOLENS
+#endif
+#if D3D12_MAX_DEVICE_INTERFACE >= 2
 	FORCEINLINE ID3D12Device2* GetD3DDevice2() const { return RootDevice2.GetReference(); }
 #endif
-#if D3D12_RHI_RAYTRACING
+#if D3D12_MAX_DEVICE_INTERFACE >= 5
 	FORCEINLINE ID3D12Device5* GetD3DDevice5() { return RootDevice5.GetReference(); }
+#endif
+#if D3D12_MAX_DEVICE_INTERFACE >= 7
 	FORCEINLINE ID3D12Device7* GetD3DDevice7() { return RootDevice7.GetReference(); }
 #endif // D3D12_RHI_RAYTRACING
 	FORCEINLINE void SetDeviceRemoved(bool value) { bDeviceRemoved = value; }
@@ -165,13 +169,6 @@ public:
 		static const FD3D12RootSignature StaticGraphicsRootSignature(this, FD3D12RootSignatureDesc::GetStaticGraphicsRootSignatureDesc());
 		return &StaticGraphicsRootSignature;
 	}
-#if PLATFORM_SUPPORTS_MESH_SHADERS
-	FORCEINLINE const FD3D12RootSignature* GetStaticMeshRootSignature()
-	{
-		static const FD3D12RootSignature StaticMeshRootSignature(this, FD3D12RootSignatureDesc::GetStaticMeshRootSignatureDesc());
-		return &StaticMeshRootSignature;
-	}
-#endif
 	FORCEINLINE const FD3D12RootSignature* GetStaticComputeRootSignature()
 	{
 		static const FD3D12RootSignature StaticComputeRootSignature(this, FD3D12RootSignatureDesc::GetStaticComputeRootSignatureDesc());
@@ -410,17 +407,33 @@ protected:
 
 	// LDA setups have one ID3D12Device
 	TRefCountPtr<ID3D12Device> RootDevice;
+#if D3D12_MAX_DEVICE_INTERFACE >= 1
 	TRefCountPtr<ID3D12Device1> RootDevice1;
-#if PLATFORM_WINDOWS || PLATFORM_HOLOLENS
+#endif
+#if D3D12_MAX_DEVICE_INTERFACE >= 2
 	TRefCountPtr<ID3D12Device2> RootDevice2;
-	TRefCountPtr<IDXGIDebug> DXGIDebug;
-
-	HANDLE ExceptionHandlerHandle = INVALID_HANDLE_VALUE;
 #endif
-#if D3D12_RHI_RAYTRACING
+#if D3D12_MAX_DEVICE_INTERFACE >= 3
+	TRefCountPtr<ID3D12Device3> RootDevice3;
+#endif
+#if D3D12_MAX_DEVICE_INTERFACE >= 4
+	TRefCountPtr<ID3D12Device4> RootDevice4;
+#endif
+#if D3D12_MAX_DEVICE_INTERFACE >= 5
 	TRefCountPtr<ID3D12Device5> RootDevice5;
+#endif
+#if D3D12_MAX_DEVICE_INTERFACE >= 6
+	TRefCountPtr<ID3D12Device6> RootDevice6;
+#endif
+#if D3D12_MAX_DEVICE_INTERFACE >= 7
 	TRefCountPtr<ID3D12Device7> RootDevice7;
-#endif // D3D12_RHI_RAYTRACING
+#endif
+
+#if PLATFORM_WINDOWS || PLATFORM_HOLOLENS
+	TRefCountPtr<IDXGIDebug> DXGIDebug;
+	HANDLE ExceptionHandlerHandle = INVALID_HANDLE_VALUE;
+#endif
+
 	D3D12_RESOURCE_HEAP_TIER ResourceHeapTier;
 	D3D12_RESOURCE_BINDING_TIER ResourceBindingTier;
 	D3D_ROOT_SIGNATURE_VERSION RootSignatureVersion;
diff --git a/Engine/Source/Runtime/D3D12RHI/Private/D3D12PipelineState.h b/Engine/Source/Runtime/D3D12RHI/Private/D3D12PipelineState.h
--- a/Engine/Source/Runtime/D3D12RHI/Private/D3D12PipelineState.h
+++ b/Engine/Source/Runtime/D3D12RHI/Private/D3D12PipelineState.h
@@ -25,6 +25,10 @@
 	#endif // D3D12_USE_DERIVED_PSO_SHADER_EXPORTS
 #endif // D3D12_USE_DERIVED_PSO
 
+#ifndef D3D12RHI_USE_PIPELINE_STATE_STREAM
+	#define D3D12RHI_USE_PIPELINE_STATE_STREAM 0
+#endif
+
 #if D3D12RHI_USE_HIGH_LEVEL_PSO_CACHE
 DECLARE_DWORD_ACCUMULATOR_STAT(TEXT("Graphics: Num high-level cache entries"), STAT_PSOGraphicsNumHighlevelCacheEntries, STATGROUP_D3D12PipelineState);
 DECLARE_DWORD_COUNTER_STAT(TEXT("Graphics: High-level cache hit"), STAT_PSOGraphicsHighlevelCacheHit, STATGROUP_D3D12PipelineState);
@@ -71,10 +75,12 @@ struct FD3D12_GRAPHICS_PIPELINE_STATE_DESC
 	D3D12_CACHED_PIPELINE_STATE CachedPSO;
 	D3D12_PIPELINE_STATE_FLAGS Flags;
 
-#if PLATFORM_WINDOWS || PLATFORM_HOLOLENS
+#if D3D12RHI_USE_PIPELINE_STATE_STREAM
 	inline bool RequiresStream() const { return MS.BytecodeLength > 0; }
 	FD3D12_GRAPHICS_PIPELINE_STATE_STREAM PipelineStateStream() const;
+#if PLATFORM_SUPPORTS_MESH_SHADERS
 	FD3D12_MESH_PIPELINE_STATE_STREAM MeshPipelineStateStream() const;
+#endif
 	D3D12_GRAPHICS_PIPELINE_STATE_DESC GraphicsDescV0() const;
 #endif // PLATFORM_WINDOWS
 };
@@ -131,7 +137,7 @@ struct FD3D12LowLevelGraphicsPipelineStateDesc
 // Compute pipeline struct that represents the latest versions of PSO subobjects currently supported by the RHI.
 struct FD3D12_COMPUTE_PIPELINE_STATE_DESC : public D3D12_COMPUTE_PIPELINE_STATE_DESC
 {
-#if PLATFORM_WINDOWS || PLATFORM_HOLOLENS
+#if D3D12RHI_USE_PIPELINE_STATE_STREAM
 	FD3D12_COMPUTE_PIPELINE_STATE_STREAM PipelineStateStream() const;
 	D3D12_COMPUTE_PIPELINE_STATE_DESC ComputeDescV0() const;
 #endif
diff --git a/Engine/Source/Runtime/D3D12RHI/Private/D3D12RootSignature.cpp b/Engine/Source/Runtime/D3D12RHI/Private/D3D12RootSignature.cpp
--- a/Engine/Source/Runtime/D3D12RHI/Private/D3D12RootSignature.cpp
+++ b/Engine/Source/Runtime/D3D12RHI/Private/D3D12RootSignature.cpp
@@ -371,8 +371,6 @@ FD3D12RootSignatureDesc::FD3D12RootSignatureDesc(const FD3D12QuantizedBoundShade
 const D3D12_VERSIONED_ROOT_SIGNATURE_DESC& FD3D12RootSignatureDesc::GetStaticGraphicsRootSignatureDesc()
 {
 	// TODO: Support vendor extensions for static root signatures?
-
-	static const uint32 DescriptorTableCount = 10;
 	static struct
 	{
 		D3D12_SHADER_VISIBILITY Vis;
@@ -380,7 +378,7 @@ const D3D12_VERSIONED_ROOT_SIGNATURE_DESC& FD3D12RootSignatureDesc::GetStaticGra
 		uint32 Count;
 		uint32 BaseShaderReg;
 		D3D12_DESCRIPTOR_RANGE_FLAGS Flags;
-	} RangeDesc[DescriptorTableCount] =
+	} RangeDesc[] =
 	{
 		{ D3D12_SHADER_VISIBILITY_PIXEL, D3D12_DESCRIPTOR_RANGE_TYPE_SRV, MAX_SRVS, 0, D3D12ShaderUtils::StaticRootSignatureConstants::SRVDescriptorRangeFlags },
 		{ D3D12_SHADER_VISIBILITY_PIXEL, D3D12_DESCRIPTOR_RANGE_TYPE_CBV, MAX_CBS, 0, D3D12ShaderUtils::StaticRootSignatureConstants::CBVDescriptorRangeFlags },
@@ -394,48 +392,7 @@ const D3D12_VERSIONED_ROOT_SIGNATURE_DESC& FD3D12RootSignatureDesc::GetStaticGra
 		{ D3D12_SHADER_VISIBILITY_GEOMETRY, D3D12_DESCRIPTOR_RANGE_TYPE_CBV, MAX_CBS, 0, D3D12ShaderUtils::StaticRootSignatureConstants::CBVDescriptorRangeFlags },
 		{ D3D12_SHADER_VISIBILITY_GEOMETRY, D3D12_DESCRIPTOR_RANGE_TYPE_SAMPLER, MAX_SAMPLERS, 0, D3D12ShaderUtils::StaticRootSignatureConstants::SamplerDescriptorRangeFlags },
 
-		{ D3D12_SHADER_VISIBILITY_ALL, D3D12_DESCRIPTOR_RANGE_TYPE_UAV, MAX_UAVS, 0, D3D12ShaderUtils::StaticRootSignatureConstants::UAVDescriptorRangeFlags },
-	};
-
-	static CD3DX12_ROOT_PARAMETER1 TableSlots[DescriptorTableCount];
-	static CD3DX12_DESCRIPTOR_RANGE1 DescriptorRanges[DescriptorTableCount];
-
-	for (uint32 i = 0; i < DescriptorTableCount; i++)
-	{
-		DescriptorRanges[i].Init(
-			RangeDesc[i].Type,
-			RangeDesc[i].Count,
-			RangeDesc[i].BaseShaderReg,
-			0u,
-			RangeDesc[i].Flags
-		);
-
-		TableSlots[i].InitAsDescriptorTable(1, &DescriptorRanges[i], RangeDesc[i].Vis);
-	}
-
-	static CD3DX12_VERSIONED_ROOT_SIGNATURE_DESC RootDesc(DescriptorTableCount, TableSlots, 0, nullptr, D3D12_ROOT_SIGNATURE_FLAG_ALLOW_INPUT_ASSEMBLER_INPUT_LAYOUT);
-	return RootDesc;
-}
-
 #if PLATFORM_SUPPORTS_MESH_SHADERS
-const D3D12_VERSIONED_ROOT_SIGNATURE_DESC& FD3D12RootSignatureDesc::GetStaticMeshRootSignatureDesc()
-{
-	// TODO: Support vendor extensions for static root signatures?
-
-	static const uint32 DescriptorTableCount = 10;
-	static struct
-	{
-		D3D12_SHADER_VISIBILITY Vis;
-		D3D12_DESCRIPTOR_RANGE_TYPE Type;
-		uint32 Count;
-		uint32 BaseShaderReg;
-		D3D12_DESCRIPTOR_RANGE_FLAGS Flags;
-	} RangeDesc[DescriptorTableCount] =
-	{
-		{ D3D12_SHADER_VISIBILITY_PIXEL, D3D12_DESCRIPTOR_RANGE_TYPE_SRV, MAX_SRVS, 0, D3D12ShaderUtils::StaticRootSignatureConstants::SRVDescriptorRangeFlags },
-		{ D3D12_SHADER_VISIBILITY_PIXEL, D3D12_DESCRIPTOR_RANGE_TYPE_CBV, MAX_CBS, 0, D3D12ShaderUtils::StaticRootSignatureConstants::CBVDescriptorRangeFlags },
-		{ D3D12_SHADER_VISIBILITY_PIXEL, D3D12_DESCRIPTOR_RANGE_TYPE_SAMPLER, MAX_SAMPLERS, 0, D3D12ShaderUtils::StaticRootSignatureConstants::SamplerDescriptorRangeFlags },
-
 		{ D3D12_SHADER_VISIBILITY_MESH, D3D12_DESCRIPTOR_RANGE_TYPE_SRV, MAX_SRVS, 0, D3D12ShaderUtils::StaticRootSignatureConstants::SRVDescriptorRangeFlags },
 		{ D3D12_SHADER_VISIBILITY_MESH, D3D12_DESCRIPTOR_RANGE_TYPE_CBV, MAX_CBS, 0, D3D12ShaderUtils::StaticRootSignatureConstants::CBVDescriptorRangeFlags },
 		{ D3D12_SHADER_VISIBILITY_MESH, D3D12_DESCRIPTOR_RANGE_TYPE_SAMPLER, MAX_SAMPLERS, 0, D3D12ShaderUtils::StaticRootSignatureConstants::SamplerDescriptorRangeFlags },
@@ -443,10 +400,12 @@ const D3D12_VERSIONED_ROOT_SIGNATURE_DESC& FD3D12RootSignatureDesc::GetStaticMes
 		{ D3D12_SHADER_VISIBILITY_AMPLIFICATION, D3D12_DESCRIPTOR_RANGE_TYPE_SRV, MAX_SRVS, 0, D3D12ShaderUtils::StaticRootSignatureConstants::SRVDescriptorRangeFlags },
 		{ D3D12_SHADER_VISIBILITY_AMPLIFICATION, D3D12_DESCRIPTOR_RANGE_TYPE_CBV, MAX_CBS, 0, D3D12ShaderUtils::StaticRootSignatureConstants::CBVDescriptorRangeFlags },
 		{ D3D12_SHADER_VISIBILITY_AMPLIFICATION, D3D12_DESCRIPTOR_RANGE_TYPE_SAMPLER, MAX_SAMPLERS, 0, D3D12ShaderUtils::StaticRootSignatureConstants::SamplerDescriptorRangeFlags },
+#endif // PLATFORM_SUPPORTS_MESH_SHADERS
 
 		{ D3D12_SHADER_VISIBILITY_ALL, D3D12_DESCRIPTOR_RANGE_TYPE_UAV, MAX_UAVS, 0, D3D12ShaderUtils::StaticRootSignatureConstants::UAVDescriptorRangeFlags },
 	};
 
+	static const uint32 DescriptorTableCount = UE_ARRAY_COUNT(RangeDesc);
 	static CD3DX12_ROOT_PARAMETER1 TableSlots[DescriptorTableCount];
 	static CD3DX12_DESCRIPTOR_RANGE1 DescriptorRanges[DescriptorTableCount];
 
@@ -463,10 +422,9 @@ const D3D12_VERSIONED_ROOT_SIGNATURE_DESC& FD3D12RootSignatureDesc::GetStaticMes
 		TableSlots[i].InitAsDescriptorTable(1, &DescriptorRanges[i], RangeDesc[i].Vis);
 	}
 
-	static CD3DX12_VERSIONED_ROOT_SIGNATURE_DESC RootDesc(DescriptorTableCount, TableSlots, 0, nullptr, D3D12_ROOT_SIGNATURE_FLAG_NONE);
+	static CD3DX12_VERSIONED_ROOT_SIGNATURE_DESC RootDesc(DescriptorTableCount, TableSlots, 0, nullptr, D3D12_ROOT_SIGNATURE_FLAG_ALLOW_INPUT_ASSEMBLER_INPUT_LAYOUT);
 	return RootDesc;
 }
-#endif // PLATFORM_SUPPORTS_MESH_SHADERS
 
 const D3D12_VERSIONED_ROOT_SIGNATURE_DESC& FD3D12RootSignatureDesc::GetStaticComputeRootSignatureDesc()
 {
diff --git a/Engine/Source/Runtime/D3D12RHI/Private/D3D12RootSignature.h b/Engine/Source/Runtime/D3D12RHI/Private/D3D12RootSignature.h
--- a/Engine/Source/Runtime/D3D12RHI/Private/D3D12RootSignature.h
+++ b/Engine/Source/Runtime/D3D12RHI/Private/D3D12RootSignature.h
@@ -43,9 +43,6 @@ public:
 	inline const D3D12_VERSIONED_ROOT_SIGNATURE_DESC& GetDesc() const { return RootDesc; }
 
 	static const D3D12_VERSIONED_ROOT_SIGNATURE_DESC& GetStaticGraphicsRootSignatureDesc();
-#if PLATFORM_SUPPORTS_MESH_SHADERS
-	static const D3D12_VERSIONED_ROOT_SIGNATURE_DESC& GetStaticMeshRootSignatureDesc();
-#endif
 	static const D3D12_VERSIONED_ROOT_SIGNATURE_DESC& GetStaticComputeRootSignatureDesc();
 	static const D3D12_VERSIONED_ROOT_SIGNATURE_DESC& GetStaticRayTracingGlobalRootSignatureDesc();
 	static const D3D12_VERSIONED_ROOT_SIGNATURE_DESC& GetStaticRayTracingLocalRootSignatureDesc();
diff --git a/Engine/Source/Runtime/D3D12RHI/Private/D3D12Shaders.cpp b/Engine/Source/Runtime/D3D12RHI/Private/D3D12Shaders.cpp
--- a/Engine/Source/Runtime/D3D12RHI/Private/D3D12Shaders.cpp
+++ b/Engine/Source/Runtime/D3D12RHI/Private/D3D12Shaders.cpp
@@ -407,7 +407,7 @@ FD3D12BoundShaderState::FD3D12BoundShaderState(
 	INC_DWORD_STAT(STAT_D3D12NumBoundShaderState);
 
 #if USE_STATIC_ROOT_SIGNATURE
-	pRootSignature = InAdapter->GetStaticMeshRootSignature();
+	pRootSignature = InAdapter->GetStaticGraphicsRootSignature();
 #else
 	const D3D12_RESOURCE_BINDING_TIER Tier = InAdapter->GetResourceBindingTier();
 	FD3D12QuantizedBoundShaderState QuantizedBoundShaderState;
diff --git a/Engine/Source/Runtime/D3D12RHI/Private/Windows/WindowsD3D12PipelineState.cpp b/Engine/Source/Runtime/D3D12RHI/Private/Windows/WindowsD3D12PipelineState.cpp
--- a/Engine/Source/Runtime/D3D12RHI/Private/Windows/WindowsD3D12PipelineState.cpp
+++ b/Engine/Source/Runtime/D3D12RHI/Private/Windows/WindowsD3D12PipelineState.cpp
@@ -75,6 +75,7 @@ FD3D12_GRAPHICS_PIPELINE_STATE_STREAM FD3D12_GRAPHICS_PIPELINE_STATE_DESC::Pipel
 	return Stream;
 }
 
+#if PLATFORM_SUPPORTS_MESH_SHADERS
 FD3D12_MESH_PIPELINE_STATE_STREAM FD3D12_GRAPHICS_PIPELINE_STATE_DESC::MeshPipelineStateStream() const
 {
 	FD3D12_MESH_PIPELINE_STATE_STREAM Stream{};
@@ -95,6 +96,7 @@ FD3D12_MESH_PIPELINE_STATE_STREAM FD3D12_GRAPHICS_PIPELINE_STATE_DESC::MeshPipel
 	Stream.CachedPSO = this->CachedPSO;
 	return Stream;
 }
+#endif // PLATFORM_SUPPORTS_MESH_SHADERS
 
 D3D12_GRAPHICS_PIPELINE_STATE_DESC FD3D12_GRAPHICS_PIPELINE_STATE_DESC::GraphicsDescV0() const
 {
@@ -871,6 +873,7 @@ static void CreatePipelineStateWrapper(ID3D12PipelineState** PSO, FD3D12Adapter*
 
 	// Use pipeline streams if the system supports it.
 	ID3D12Device2* const pDevice2 = Adapter->GetD3DDevice2();
+#if PLATFORM_SUPPORTS_MESH_SHADERS
 	if (CreationArgs->Desc.UsesMeshShaders())
 	{
 		check(pDevice2);
@@ -892,7 +895,9 @@ static void CreatePipelineStateWrapper(ID3D12PipelineState** PSO, FD3D12Adapter*
 			}
 		}
 	}
-	else if (pDevice2 && bUseStream)
+	else
+#endif // PLATFORM_SUPPORTS_MESH_SHADERS
+	if (pDevice2 && bUseStream)
 	{
 		FD3D12_GRAPHICS_PIPELINE_STATE_STREAM Stream = CreationArgs->Desc.Desc.PipelineStateStream();
 		const D3D12_PIPELINE_STATE_STREAM_DESC StreamDesc = { sizeof(Stream), &Stream };
diff --git a/Engine/Source/Runtime/D3D12RHI/Private/Windows/WindowsD3D12PipelineState.h b/Engine/Source/Runtime/D3D12RHI/Private/Windows/WindowsD3D12PipelineState.h
--- a/Engine/Source/Runtime/D3D12RHI/Private/Windows/WindowsD3D12PipelineState.h
+++ b/Engine/Source/Runtime/D3D12RHI/Private/Windows/WindowsD3D12PipelineState.h
@@ -37,6 +37,7 @@ struct FD3D12_GRAPHICS_PIPELINE_STATE_STREAM
 	//CD3DX12_PIPELINE_STATE_STREAM_VIEW_INSTANCING ViewInstancingDesc;
 };
 
+#if PLATFORM_SUPPORTS_MESH_SHADERS
 struct FD3D12_MESH_PIPELINE_STATE_STREAM
 {
 	// Note: Unused members are currently commented out to exclude them from the stream.
@@ -59,6 +60,7 @@ struct FD3D12_MESH_PIPELINE_STATE_STREAM
 	CD3DX12_PIPELINE_STATE_STREAM_CACHED_PSO CachedPSO;
 	//CD3DX12_PIPELINE_STATE_STREAM_VIEW_INSTANCING ViewInstancingDesc;
 };
+#endif
 
 // Compute pipeline stream struct that represents the latest version of PSO subobjects currently used by the RHI.
 struct FD3D12_COMPUTE_PIPELINE_STATE_STREAM
diff --git a/Engine/Source/Runtime/D3D12RHI/Public/D3D12RootSignatureDefinitions.h b/Engine/Source/Runtime/D3D12RHI/Public/D3D12RootSignatureDefinitions.h
--- a/Engine/Source/Runtime/D3D12RHI/Public/D3D12RootSignatureDefinitions.h
+++ b/Engine/Source/Runtime/D3D12RHI/Public/D3D12RootSignatureDefinitions.h
@@ -68,7 +68,7 @@ namespace D3D12ShaderUtils
 	};
 
 	// Fat/Static Gfx Root Signature
-	inline void CreateGfxRootSignature(FRootSignatureCreator* Creator)
+	inline void CreateGfxRootSignature(FRootSignatureCreator* Creator, bool bAllowMeshShaders)
 	{
 		// Ensure the creator starts in a clean state (in cases of creator reuse, etc.).
 		Creator->Reset();
@@ -85,6 +85,15 @@ namespace D3D12ShaderUtils
 		Creator->AddTable(SF_Geometry, FRootSignatureCreator::SRV, MAX_SRVS);
 		Creator->AddTable(SF_Geometry, FRootSignatureCreator::CBV, MAX_CBS);
 		Creator->AddTable(SF_Geometry, FRootSignatureCreator::Sampler, MAX_SAMPLERS);
+		if (bAllowMeshShaders)
+		{
+			Creator->AddTable(SF_Mesh, FRootSignatureCreator::SRV, MAX_SRVS);
+			Creator->AddTable(SF_Mesh, FRootSignatureCreator::CBV, MAX_CBS);
+			Creator->AddTable(SF_Mesh, FRootSignatureCreator::Sampler, MAX_SAMPLERS);
+			Creator->AddTable(SF_Amplification, FRootSignatureCreator::SRV, MAX_SRVS);
+			Creator->AddTable(SF_Amplification, FRootSignatureCreator::CBV, MAX_CBS);
+			Creator->AddTable(SF_Amplification, FRootSignatureCreator::Sampler, MAX_SAMPLERS);
+		}
 		Creator->AddTable(SF_NumFrequencies, FRootSignatureCreator::UAV, MAX_UAVS);
 	}
 
@@ -163,7 +172,7 @@ namespace D3D12ShaderUtils
 
 		void Compile()
 		{
-			D3D12ShaderUtils::CreateGfxRootSignature(this);
+			D3D12ShaderUtils::CreateGfxRootSignature(this, false);
 
 			// Patch pointers
 			for (auto& Pair : ParameterToRangeMap)
