Remove FName lookup from fence creation in Test / Shipping builds.

diff --git a/Engine/Source/Runtime/D3D12RHI/Private/D3D12CommandContext.cpp b/Engine/Source/Runtime/D3D12RHI/Private/D3D12CommandContext.cpp
--- a/Engine/Source/Runtime/D3D12RHI/Private/D3D12CommandContext.cpp
+++ b/Engine/Source/Runtime/D3D12RHI/Private/D3D12CommandContext.cpp
@@ -948,7 +948,11 @@ void FD3D12DynamicRHI::RHICreateTransition(FRHITransition* Transition, const FRH
 
 	if (bCrossPipeline && !EnumHasAnyFlags(Data->CreateFlags, ERHITransitionCreateFlags::NoFence))
 	{
+#if (!UE_BUILD_SHIPPING && !UE_BUILD_TEST)
 		const FName Name = CreateInfo.SrcPipelines == ERHIPipeline::Graphics ? L"<Graphics To AsyncCompute>" : L"<AsyncCompute To Graphics>";
+#else
+		const FName Name = {};
+#endif
 
 		Data->Fence = new FD3D12Fence(&GetAdapter(), FRHIGPUMask::All(), Name);
 		Data->Fence->CreateFence();
