Removing D3D12RHI_USE_PIPELINE_STATE_STREAM now that it's confirmed to work on all relevant platforms.

diff --git a/Engine/Source/Runtime/D3D12RHI/D3D12RHI.Build.cs b/Engine/Source/Runtime/D3D12RHI/D3D12RHI.Build.cs
--- a/Engine/Source/Runtime/D3D12RHI/D3D12RHI.Build.cs
+++ b/Engine/Source/Runtime/D3D12RHI/D3D12RHI.Build.cs
@@ -41,8 +41,6 @@ public class D3D12RHI : ModuleRules
         if (Target.Platform.IsInGroup(UnrealPlatformGroup.Windows) ||
             Target.Platform == UnrealTargetPlatform.HoloLens)
 		{
-			PrivateDefinitions.Add("D3D12RHI_USE_PIPELINE_STATE_STREAM=1");
-
 			AddEngineThirdPartyPrivateStaticDependencies(Target, "DX12");
 			AddEngineThirdPartyPrivateStaticDependencies(Target, "DX11");
 			AddEngineThirdPartyPrivateStaticDependencies(Target, "AMD_AGS");
diff --git a/Engine/Source/Runtime/D3D12RHI/Private/D3D12PipelineState.h b/Engine/Source/Runtime/D3D12RHI/Private/D3D12PipelineState.h
--- a/Engine/Source/Runtime/D3D12RHI/Private/D3D12PipelineState.h
+++ b/Engine/Source/Runtime/D3D12RHI/Private/D3D12PipelineState.h
@@ -25,10 +25,6 @@
 	#endif // D3D12_USE_DERIVED_PSO_SHADER_EXPORTS
 #endif // D3D12_USE_DERIVED_PSO
 
-#ifndef D3D12RHI_USE_PIPELINE_STATE_STREAM
-	#define D3D12RHI_USE_PIPELINE_STATE_STREAM 0
-#endif
-
 #if D3D12RHI_USE_HIGH_LEVEL_PSO_CACHE
 DECLARE_DWORD_ACCUMULATOR_STAT(TEXT("Graphics: Num high-level cache entries"), STAT_PSOGraphicsNumHighlevelCacheEntries, STATGROUP_D3D12PipelineState);
 DECLARE_DWORD_COUNTER_STAT(TEXT("Graphics: High-level cache hit"), STAT_PSOGraphicsHighlevelCacheHit, STATGROUP_D3D12PipelineState);
@@ -75,14 +71,12 @@ struct FD3D12_GRAPHICS_PIPELINE_STATE_DESC
 	D3D12_CACHED_PIPELINE_STATE CachedPSO;
 	D3D12_PIPELINE_STATE_FLAGS Flags;
 
-#if D3D12RHI_USE_PIPELINE_STATE_STREAM
 	inline bool RequiresStream() const { return MS.BytecodeLength > 0; }
 	FD3D12_GRAPHICS_PIPELINE_STATE_STREAM PipelineStateStream() const;
 #if PLATFORM_SUPPORTS_MESH_SHADERS
 	FD3D12_MESH_PIPELINE_STATE_STREAM MeshPipelineStateStream() const;
 #endif
 	D3D12_GRAPHICS_PIPELINE_STATE_DESC GraphicsDescV0() const;
-#endif // PLATFORM_WINDOWS
 };
 
 struct FD3D12LowLevelGraphicsPipelineStateDesc
@@ -137,10 +131,8 @@ struct FD3D12LowLevelGraphicsPipelineStateDesc
 // Compute pipeline struct that represents the latest versions of PSO subobjects currently supported by the RHI.
 struct FD3D12_COMPUTE_PIPELINE_STATE_DESC : public D3D12_COMPUTE_PIPELINE_STATE_DESC
 {
-#if D3D12RHI_USE_PIPELINE_STATE_STREAM
 	FD3D12_COMPUTE_PIPELINE_STATE_STREAM PipelineStateStream() const;
 	D3D12_COMPUTE_PIPELINE_STATE_DESC ComputeDescV0() const;
-#endif
 };
 
 struct FD3D12ComputePipelineStateDesc
