Only flushing pending execute task when on a different queue - fix deadlock with submit task

diff --git a/Engine/Source/Runtime/D3D12RHI/Private/D3D12DirectCommandListManager.cpp b/Engine/Source/Runtime/D3D12RHI/Private/D3D12DirectCommandListManager.cpp
--- a/Engine/Source/Runtime/D3D12RHI/Private/D3D12DirectCommandListManager.cpp
+++ b/Engine/Source/Runtime/D3D12RHI/Private/D3D12DirectCommandListManager.cpp
@@ -718,7 +718,10 @@ FD3D12SyncPoint FD3D12CommandListManager::ExecuteCommandListInternal(TArray<FD3D
 #endif
 		
 		// Make sure the queue is done executing commands before trying to use the barrier command list from it
-		BarrierCommandListManager->WaitOnExecuteTask();
+		if (bBarriersUseDifferentQueue)
+		{
+			BarrierCommandListManager->WaitOnExecuteTask();
+		}
 
 		// If the barrier command list manager is different then this command list then it might need another sync point
 		// with the copy queue (didn't happen yet)
