Fix a use-after-release crash

diff --git a/Engine/Source/Runtime/D3D12RHI/Private/D3D12Buffer.cpp b/Engine/Source/Runtime/D3D12RHI/Private/D3D12Buffer.cpp
--- a/Engine/Source/Runtime/D3D12RHI/Private/D3D12Buffer.cpp
+++ b/Engine/Source/Runtime/D3D12RHI/Private/D3D12Buffer.cpp
@@ -65,13 +65,13 @@ struct FD3D12RHICommandInitializeBufferString
 };
 struct FD3D12RHICommandInitializeBuffer final : public FRHICommand<FD3D12RHICommandInitializeBuffer, FD3D12RHICommandInitializeBufferString>
 {
-	FD3D12Buffer* Buffer;
+	TRefCountPtr<FD3D12Buffer> Buffer;
 	FD3D12ResourceLocation SrcResourceLoc;
 	uint32 Size;
 	D3D12_RESOURCE_STATES DestinationState;
 
-	FORCEINLINE_DEBUGGABLE FD3D12RHICommandInitializeBuffer(FD3D12Buffer* InBuffer, FD3D12ResourceLocation& InSrcResourceLoc, uint32 InSize, D3D12_RESOURCE_STATES InDestinationState)
-		: Buffer(InBuffer)
+	FORCEINLINE_DEBUGGABLE FD3D12RHICommandInitializeBuffer(TRefCountPtr<FD3D12Buffer>&& InBuffer, FD3D12ResourceLocation& InSrcResourceLoc, uint32 InSize, D3D12_RESOURCE_STATES InDestinationState)
+		: Buffer(MoveTemp(InBuffer))
 		, SrcResourceLoc(InSrcResourceLoc.GetParentDevice())
 		, Size(InSize)
 		, DestinationState(InDestinationState)
@@ -189,17 +189,18 @@ void FD3D12Buffer::UploadResourceData(class FRHICommandListImmediate* RHICmdList
 			// be done on the copy queue.
 			FD3D12ResourceLocation* SrcResourceLoc_Heap = new FD3D12ResourceLocation(SrcResourceLoc.GetParentDevice());
 			FD3D12ResourceLocation::TransferOwnership(*SrcResourceLoc_Heap, SrcResourceLoc);
+			TRefCountPtr<FD3D12Buffer> ThisRef = this;
 			ENQUEUE_RENDER_COMMAND(CmdD3D12InitializeBuffer)(
-				[this, SrcResourceLoc_Heap, BufferSize, InDestinationState](FRHICommandListImmediate& RHICmdList)
+				[DestinationBuffer = MoveTemp(ThisRef), SrcResourceLoc_Heap, BufferSize, InDestinationState](FRHICommandListImmediate& RHICmdList) mutable
 				{
 					if (RHICmdList.Bypass())
 					{
-						FD3D12RHICommandInitializeBuffer Command(this, *SrcResourceLoc_Heap, BufferSize, InDestinationState);
+						FD3D12RHICommandInitializeBuffer Command(MoveTemp(DestinationBuffer), *SrcResourceLoc_Heap, BufferSize, InDestinationState);
 						Command.Execute(RHICmdList);
 					}
 					else
 					{
-						new (RHICmdList.AllocCommand<FD3D12RHICommandInitializeBuffer>()) FD3D12RHICommandInitializeBuffer(this, *SrcResourceLoc_Heap, BufferSize, InDestinationState);
+						new (RHICmdList.AllocCommand<FD3D12RHICommandInitializeBuffer>()) FD3D12RHICommandInitializeBuffer(MoveTemp(DestinationBuffer), *SrcResourceLoc_Heap, BufferSize, InDestinationState);
 					}
 					delete SrcResourceLoc_Heap;
 				});
