Reinstate the option to not clear GBufferD

diff --git a/Engine/Source/Runtime/D3D12RHI/Private/D3D12CommandContext.h b/Engine/Source/Runtime/D3D12RHI/Private/D3D12CommandContext.h
--- a/Engine/Source/Runtime/D3D12RHI/Private/D3D12CommandContext.h
+++ b/Engine/Source/Runtime/D3D12RHI/Private/D3D12CommandContext.h
@@ -315,7 +315,7 @@ public:
 	virtual void RHIBackBufferWaitTrackingBeginFrame(uint64 FrameToken, bool bDeferred) final override;
 #endif // #if PLATFORM_USE_BACKBUFFER_WRITE_TRANSITION_TRACKING
 
-	virtual void RHIClearMRTImpl(bool bClearColor, int32 NumClearColors, const FLinearColor* ColorArray, bool bClearDepth, float Depth, bool bClearStencil, uint32 Stencil);
+	virtual void RHIClearMRTImpl(bool* bClearColorArray, int32 NumClearColors, const FLinearColor* ColorArray, bool bClearDepth, float Depth, bool bClearStencil, uint32 Stencil);
 
 
 	virtual void RHIBeginRenderPass(const FRHIRenderPassInfo& InInfo, const TCHAR* InName)
@@ -456,7 +456,7 @@ protected:
 	void WriteGPUEventStackToBreadCrumbData(bool bBeginEvent);
 
 private:
-	void RHIClearMRT(bool bClearColor, int32 NumClearColors, const FLinearColor* ColorArray, bool bClearDepth, float Depth, bool bClearStencil, uint32 Stencil);
+	void RHIClearMRT(bool* bClearColorArray, int32 NumClearColors, const FLinearColor* ColorArray, bool bClearDepth, float Depth, bool bClearStencil, uint32 Stencil);
 
 	static void ClearUAV(TRHICommandList_RecursiveHazardous<FD3D12CommandContext>& RHICmdList, FD3D12UnorderedAccessView* UAV, const void* ClearValues, bool bFloat);
 
diff --git a/Engine/Source/Runtime/D3D12RHI/Private/D3D12Commands.cpp b/Engine/Source/Runtime/D3D12RHI/Private/D3D12Commands.cpp
--- a/Engine/Source/Runtime/D3D12RHI/Private/D3D12Commands.cpp
+++ b/Engine/Source/Runtime/D3D12RHI/Private/D3D12Commands.cpp
@@ -1369,6 +1369,7 @@ void FD3D12CommandContext::SetRenderTargetsAndClear(const FRHISetRenderTargetsIn
 	if (RenderTargetsInfo.bClearColor || RenderTargetsInfo.bClearStencil || RenderTargetsInfo.bClearDepth)
 	{
 		FLinearColor ClearColors[MaxSimultaneousRenderTargets];
+		bool bClearColorArray[MaxSimultaneousRenderTargets];
 		float DepthClear = 0.0;
 		uint32 StencilClear = 0;
 
@@ -1386,6 +1387,7 @@ void FD3D12CommandContext::SetRenderTargetsAndClear(const FRHISetRenderTargetsIn
 				{
 					ClearColors[i] = FLinearColor(ForceInitToZero);
 				}
+				bClearColorArray[i] = RenderTargetsInfo.ColorRenderTarget[i].LoadAction == ERenderTargetLoadAction::EClear;
 			}
 		}
 		if (RenderTargetsInfo.bClearDepth || RenderTargetsInfo.bClearStencil)
@@ -1395,7 +1397,7 @@ void FD3D12CommandContext::SetRenderTargetsAndClear(const FRHISetRenderTargetsIn
 			ClearValue.GetDepthStencil(DepthClear, StencilClear);
 		}
 
-		this->RHIClearMRTImpl(RenderTargetsInfo.bClearColor, RenderTargetsInfo.NumColorRenderTargets, ClearColors, RenderTargetsInfo.bClearDepth, DepthClear, RenderTargetsInfo.bClearStencil, StencilClear);
+		this->RHIClearMRTImpl(RenderTargetsInfo.bClearColor ? bClearColorArray : nullptr, RenderTargetsInfo.NumColorRenderTargets, ClearColors, RenderTargetsInfo.bClearDepth, DepthClear, RenderTargetsInfo.bClearStencil, StencilClear);
 	}
 
 #if PLATFORM_SUPPORTS_VARIABLE_RATE_SHADING
@@ -2068,12 +2070,12 @@ void FD3D12CommandContext::RHIDispatchIndirectMeshShader(FRHIBuffer* ArgumentBuf
 #endif // PLATFORM_SUPPORTS_MESH_SHADERS
 
 // Raster operations.
-void FD3D12CommandContext::RHIClearMRT(bool bClearColor, int32 NumClearColors, const FLinearColor* ClearColorArray, bool bClearDepth, float Depth, bool bClearStencil, uint32 Stencil)
+void FD3D12CommandContext::RHIClearMRT(bool* bClearColorArray, int32 NumClearColors, const FLinearColor* ClearColorArray, bool bClearDepth, float Depth, bool bClearStencil, uint32 Stencil)
 {
-	RHIClearMRTImpl(bClearColor, NumClearColors, ClearColorArray, bClearDepth, Depth, bClearStencil, Stencil);
+	RHIClearMRTImpl(bClearColorArray, NumClearColors, ClearColorArray, bClearDepth, Depth, bClearStencil, Stencil);
 }
 
-void FD3D12CommandContext::RHIClearMRTImpl(bool bClearColor, int32 NumClearColors, const FLinearColor* ClearColorArray, bool bClearDepth, float Depth, bool bClearStencil, uint32 Stencil)
+void FD3D12CommandContext::RHIClearMRTImpl(bool* bClearColorArray, int32 NumClearColors, const FLinearColor* ClearColorArray, bool bClearDepth, float Depth, bool bClearStencil, uint32 Stencil)
 {
 	SCOPE_CYCLE_COUNTER(STAT_D3D12ClearMRT);
 
@@ -2105,7 +2107,7 @@ void FD3D12CommandContext::RHIClearMRTImpl(bool bClearColor, int32 NumClearColor
 	}
 
 	// Must specify enough clear colors for all active RTs
-	check(!bClearColor || NumClearColors >= BoundRenderTargets.GetNumActiveTargets());
+	check(!bClearColorArray || NumClearColors >= BoundRenderTargets.GetNumActiveTargets());
 
 	const bool bSupportsFastClear = true;
 	uint32 ClearRectCount = 0;
@@ -2132,7 +2134,7 @@ void FD3D12CommandContext::RHIClearMRTImpl(bool bClearColor, int32 NumClearColor
 		}
 	}
 
-	const bool ClearRTV = bClearColor && BoundRenderTargets.GetNumActiveTargets() > 0;
+	const bool ClearRTV = bClearColorArray && BoundRenderTargets.GetNumActiveTargets() > 0;
 	const bool ClearDSV = (bClearDepth || bClearStencil) && DepthStencilView;
 
 	uint32 ClearFlags = 0;
@@ -2167,7 +2169,7 @@ void FD3D12CommandContext::RHIClearMRTImpl(bool bClearColor, int32 NumClearColor
 			{
 				FD3D12RenderTargetView* RTView = BoundRenderTargets.GetRenderTargetView(TargetIndex);
 
-				if (RTView != nullptr)
+				if (RTView != nullptr && bClearColorArray[TargetIndex])
 				{
 					numClears++;
 					CommandListHandle->ClearRenderTargetView(RTView->GetView(), (float*)&ClearColorArray[TargetIndex], ClearRectCount, pClearRects);
diff --git a/Engine/Source/Runtime/Renderer/Private/BasePassRendering.cpp b/Engine/Source/Runtime/Renderer/Private/BasePassRendering.cpp
--- a/Engine/Source/Runtime/Renderer/Private/BasePassRendering.cpp
+++ b/Engine/Source/Runtime/Renderer/Private/BasePassRendering.cpp
@@ -77,6 +77,12 @@ static TAutoConsoleVariable<int32> CVarParallelBasePass(
 	TEXT("Toggles parallel base pass rendering. Parallel rendering must be enabled for this to have an effect."),
 	ECVF_RenderThreadSafe);
 
+static TAutoConsoleVariable<int32> CVarClearGBufferDBeforeBasePass(
+	TEXT("r.ClearGBufferDBeforeBasePass"),
+	1,
+	TEXT("Whether to clear GBuffer D before basepass"),
+	ECVF_RenderThreadSafe);
+
 // Scene color alpha is used during scene captures and planar reflections.  1 indicates background should be shown, 0 indicates foreground is fully present.
 static const float kSceneColorClearAlpha = 1.0f;
 
@@ -799,6 +805,11 @@ void FDeferredShadingSceneRenderer::RenderBasePass(
 		auto* PassParameters = GraphBuilder.AllocParameters<FRenderTargetParameters>();
 		PassParameters->RenderTargets = GetRenderTargetBindings(ColorLoadAction, BasePassTexturesView);
 
+		if (!CVarClearGBufferDBeforeBasePass.GetValueOnRenderThread() && SceneTextures.Config.GBufferD.Index > 0 && SceneTextures.Config.GBufferD.Index < (int32)BasePassTextureCount)
+		{
+			PassParameters->RenderTargets[SceneTextures.Config.GBufferD.Index].SetLoadAction(ERenderTargetLoadAction::ENoAction);
+		}
+
 		GraphBuilder.AddPass(RDG_EVENT_NAME("GBufferClear"), PassParameters, ERDGPassFlags::Raster,
 			[PassParameters, ColorLoadAction, SceneColorClearValue](FRHICommandList& RHICmdList)
 		{
