Fixed detection of AMD ray token extension. Support query was previously incorrectly behind minimum driver version check.

diff --git a/Engine/Source/Runtime/D3D12RHI/Private/Windows/WindowsD3D12Device.cpp b/Engine/Source/Runtime/D3D12RHI/Private/Windows/WindowsD3D12Device.cpp
--- a/Engine/Source/Runtime/D3D12RHI/Private/Windows/WindowsD3D12Device.cpp
+++ b/Engine/Source/Runtime/D3D12RHI/Private/Windows/WindowsD3D12Device.cpp
@@ -924,11 +924,11 @@ void FD3D12DynamicRHI::Init()
 #if AMD_API_ENABLE
 	if (GRHISupportsRayTracing
 		&& IsRHIDeviceAMD()
-		&& GMinimumDriverVersionForRayTracingAMD > 0
 		&& AmdAgsContext
 		&& AmdAgsGpuInfo.radeonSoftwareVersion)
 	{
-		if (agsCheckDriverVersion(AmdAgsGpuInfo.radeonSoftwareVersion, GMinimumDriverVersionForRayTracingAMD) == AGS_SOFTWAREVERSIONCHECK_OLDER)
+		if (GMinimumDriverVersionForRayTracingAMD > 0 
+			&& agsCheckDriverVersion(AmdAgsGpuInfo.radeonSoftwareVersion, GMinimumDriverVersionForRayTracingAMD) == AGS_SOFTWAREVERSIONCHECK_OLDER)
 		{
 			GRHISupportsRayTracing = false;
 
@@ -938,6 +938,7 @@ void FD3D12DynamicRHI::Init()
 		if (GRHISupportsRayTracing)
 		{
 			GRHISupportsRayTracingAMDHitToken = (GetAmdSupportedExtensionFlags() & AGS_DX12_EXTENSION_INTRINSIC_RAY_TRACE_HIT_TOKEN) != 0;
+			UE_LOG(LogD3D12RHI, Log, TEXT("AMD hit token extension is %s"), GRHISupportsRayTracingAMDHitToken ? TEXT("supported") : TEXT("not supported"));
 		}
 	}
 #endif // AMD_API_ENABLE
