Print DXGI memory stats on GPU crash to see if it's over VRAM budget or not

diff --git a/Engine/Source/Runtime/D3D12RHI/Private/D3D12Util.cpp b/Engine/Source/Runtime/D3D12RHI/Private/D3D12Util.cpp
--- a/Engine/Source/Runtime/D3D12RHI/Private/D3D12Util.cpp
+++ b/Engine/Source/Runtime/D3D12RHI/Private/D3D12Util.cpp
@@ -500,6 +500,10 @@ static bool LogDREDData(ID3D12Device* Device, bool bTrackingAllAllocations, D3D1
 				UE_LOG(LogD3D12RHI, Error, TEXT("DRED: No command list found with active outstanding operations (all finished or not started yet)."));
 			}
 		}
+		else
+		{
+			UE_LOG(LogD3D12RHI, Error, TEXT("DRED: No breadcrumb head found."));
+		}
 
 		D3D12_DRED_PAGE_FAULT_OUTPUT DredPageFaultOutput;
 		if (SUCCEEDED(Dred.Data->GetPageFaultAllocationOutput(&DredPageFaultOutput)) && DredPageFaultOutput.PageFaultVA != 0)
@@ -613,6 +617,16 @@ void LogPageFaultData(FD3D12Adapter* InAdapter, D3D12_GPU_VIRTUAL_ADDRESS InPage
 	}
 }
 
+void LogMemoryInfo(FD3D12Adapter* InAdapter)
+{
+	InAdapter->UpdateMemoryInfo();
+	const FD3D12MemoryInfo& MemoryInfo = InAdapter->GetMemoryInfo();
+
+	UE_LOG(LogD3D12RHI, Error, TEXT("Memory Info from frame ID %d:"), MemoryInfo.UpdateFrameNumber);
+	UE_LOG(LogD3D12RHI, Error, TEXT("\tBudget:\t%7.2f MB"), MemoryInfo.LocalMemoryInfo.Budget / (1024.0f * 1024));
+	UE_LOG(LogD3D12RHI, Error, TEXT("\tUsed:\t%7.2f MB"), MemoryInfo.LocalMemoryInfo.CurrentUsage / (1024.0f * 1024));
+}
+
 #endif  // PLATFORM_WINDOWS
 
 extern CORE_API bool GIsGPUCrashed;
@@ -670,10 +684,15 @@ namespace D3D12RHI
 					bool bIsTrackingAllAllocations = IterationDevice->GetParentAdapter()->IsTrackingAllAllocations();
 					if (!LogDREDData<FDred_1_2>(IterationDevice->GetDevice(), bIsTrackingAllAllocations, PageFaultAddress))
 					{
-						LogDREDData<FDred_1_1>(IterationDevice->GetDevice(), bIsTrackingAllAllocations, PageFaultAddress);
+						if (!LogDREDData<FDred_1_1>(IterationDevice->GetDevice(), bIsTrackingAllAllocations, PageFaultAddress))
+						{
+							UE_LOG(LogD3D12RHI, Error, TEXT("DRED: could not find DRED data (might not be enabled or available). Run with -dred or -gpucrashdebugging to enable dred if available."));
+						}						
 					}
 
-					LogPageFaultData(IterationDevice->GetParentAdapter(), PageFaultAddress);
+					FD3D12Adapter* Adapter = IterationDevice->GetParentAdapter();
+					LogPageFaultData(Adapter, PageFaultAddress);
+					LogMemoryInfo(Adapter);
 				}
 			});
 #endif  // PLATFORM_WINDOWS
