Fix for allocating a structured buffer with a large stride (>256) and no UAV flag crashing D3D12. Fix is to force a placed resource when the requested alignment exceeds the manual suballocation alignment.

diff --git a/Engine/Source/Runtime/D3D12RHI/Private/D3D12Allocation.cpp b/Engine/Source/Runtime/D3D12RHI/Private/D3D12Allocation.cpp
--- a/Engine/Source/Runtime/D3D12RHI/Private/D3D12Allocation.cpp
+++ b/Engine/Source/Runtime/D3D12RHI/Private/D3D12Allocation.cpp
@@ -1139,8 +1139,13 @@ FD3D12ResourceInitConfig FD3D12DefaultBufferPool::GetResourceAllocatorInitConfig
 }
 
 
-EResourceAllocationStrategy FD3D12DefaultBufferPool::GetResourceAllocationStrategy(D3D12_RESOURCE_FLAGS InResourceFlags, ED3D12ResourceStateMode InResourceStateMode)
+EResourceAllocationStrategy FD3D12DefaultBufferPool::GetResourceAllocationStrategy(D3D12_RESOURCE_FLAGS InResourceFlags, ED3D12ResourceStateMode InResourceStateMode, uint32 Alignment)
 {
+	if (Alignment > kD3D12ManualSubAllocationAlignment)
+	{
+		return EResourceAllocationStrategy::kPlacedResource;
+	}
+
 	// Does the resource need state tracking and transitions
 	ED3D12ResourceStateMode ResourceStateMode = InResourceStateMode;
 	if (ResourceStateMode == ED3D12ResourceStateMode::Default)
@@ -1166,10 +1171,10 @@ FD3D12DefaultBufferPool::FD3D12DefaultBufferPool(FD3D12Device* InParent, FD3D12M
 }
 
 
-bool FD3D12DefaultBufferPool::SupportsAllocation(D3D12_HEAP_TYPE InHeapType, D3D12_RESOURCE_FLAGS InResourceFlags, EBufferUsageFlags InBufferUsage, ED3D12ResourceStateMode InResourceStateMode) const
+bool FD3D12DefaultBufferPool::SupportsAllocation(D3D12_HEAP_TYPE InHeapType, D3D12_RESOURCE_FLAGS InResourceFlags, EBufferUsageFlags InBufferUsage, ED3D12ResourceStateMode InResourceStateMode, uint32 Alignment) const
 {
 	FD3D12ResourceInitConfig InitConfig = GetResourceAllocatorInitConfig(InHeapType, InResourceFlags, InBufferUsage);
-	EResourceAllocationStrategy AllocationStrategy = GetResourceAllocationStrategy(InResourceFlags, InResourceStateMode);
+	EResourceAllocationStrategy AllocationStrategy = GetResourceAllocationStrategy(InResourceFlags, InResourceStateMode, Alignment);
 	return (Allocator->GetInitConfig() == InitConfig && Allocator->GetAllocationStrategy() == AllocationStrategy);
 }
 
@@ -1284,7 +1289,7 @@ FD3D12DefaultBufferAllocator::FD3D12DefaultBufferAllocator(FD3D12Device* InParen
 	FMemory::Memset(DefaultBufferPools, 0);
 }
 
-FD3D12BufferPool* FD3D12DefaultBufferAllocator::CreateBufferPool(D3D12_HEAP_TYPE InHeapType, D3D12_RESOURCE_FLAGS InResourceFlags, EBufferUsageFlags InBufferUsage, ED3D12ResourceStateMode InResourceStateMode)
+FD3D12BufferPool* FD3D12DefaultBufferAllocator::CreateBufferPool(D3D12_HEAP_TYPE InHeapType, D3D12_RESOURCE_FLAGS InResourceFlags, EBufferUsageFlags InBufferUsage, ED3D12ResourceStateMode InResourceStateMode, uint32 Alignment)
 {
 	FD3D12Device* Device = GetParentDevice();
 	FD3D12ResourceInitConfig InitConfig = FD3D12BufferPool::GetResourceAllocatorInitConfig(InHeapType, InResourceFlags, InBufferUsage);
@@ -1292,9 +1297,9 @@ FD3D12BufferPool* FD3D12DefaultBufferAllocator::CreateBufferPool(D3D12_HEAP_TYPE
 #if USE_POOL_ALLOCATOR
 
 	const FString Name(L"D3D12 Pool Allocator");
-	EResourceAllocationStrategy AllocationStrategy = FD3D12PoolAllocator::GetResourceAllocationStrategy(InResourceFlags, InResourceStateMode);
+	EResourceAllocationStrategy AllocationStrategy = FD3D12PoolAllocator::GetResourceAllocationStrategy(InResourceFlags, InResourceStateMode, Alignment);
 	uint64 PoolSize = InHeapType == D3D12_HEAP_TYPE_READBACK ? READBACK_BUFFER_POOL_DEFAULT_POOL_SIZE : GD3D12PoolAllocatorBufferVRAMPoolSize;
-	uint64 PoolAlignment = (AllocationStrategy == EResourceAllocationStrategy::kPlacedResource) ? MIN_PLACED_RESOURCE_SIZE : 256;
+	uint64 PoolAlignment = (AllocationStrategy == EResourceAllocationStrategy::kPlacedResource) ? MIN_PLACED_RESOURCE_SIZE : kD3D12ManualSubAllocationAlignment;
 	uint64 MaxAllocationSize = InHeapType == D3D12_HEAP_TYPE_READBACK ? READBACK_BUFFER_POOL_MAX_ALLOC_SIZE : GD3D12PoolAllocatorBufferVRAMMaxAllocationSize;
 	FRHIMemoryPool::EFreeListOrder FreeListOrder = FRHIMemoryPool::EFreeListOrder::SortBySize;
 
@@ -1315,7 +1320,7 @@ FD3D12BufferPool* FD3D12DefaultBufferAllocator::CreateBufferPool(D3D12_HEAP_TYPE
 
 #else // USE_POOL_ALLOCATOR
 
-	EResourceAllocationStrategy AllocationStrategy = FD3D12DefaultBufferPool::GetResourceAllocationStrategy(InResourceFlags, InResourceStateMode);
+	EResourceAllocationStrategy AllocationStrategy = FD3D12DefaultBufferPool::GetResourceAllocationStrategy(InResourceFlags, InResourceStateMode, Alignment);
 
 	// if placed then 64KB alignment required :(
 	uint32 MinBlockSize = (AllocationStrategy == EResourceAllocationStrategy::kPlacedResource) ? MIN_PLACED_RESOURCE_SIZE : 16;
@@ -1340,9 +1345,9 @@ FD3D12BufferPool* FD3D12DefaultBufferAllocator::CreateBufferPool(D3D12_HEAP_TYPE
 }
 
 
-bool FD3D12DefaultBufferAllocator::IsPlacedResource(D3D12_RESOURCE_FLAGS InResourceFlags, ED3D12ResourceStateMode InResourceStateMode)
+bool FD3D12DefaultBufferAllocator::IsPlacedResource(D3D12_RESOURCE_FLAGS InResourceFlags, ED3D12ResourceStateMode InResourceStateMode, uint32 Alignment)
 {
-	EResourceAllocationStrategy AllocationStrategy = FD3D12BufferPool::GetResourceAllocationStrategy(InResourceFlags, InResourceStateMode);
+	EResourceAllocationStrategy AllocationStrategy = FD3D12BufferPool::GetResourceAllocationStrategy(InResourceFlags, InResourceStateMode, Alignment);
 	return (AllocationStrategy == EResourceAllocationStrategy::kPlacedResource);
 }
 
@@ -1389,7 +1394,7 @@ void FD3D12DefaultBufferAllocator::AllocDefaultResource(D3D12_HEAP_TYPE InHeapTy
 	FD3D12BufferPool* BufferPool = nullptr;
 	for (FD3D12BufferPool* Pool : DefaultBufferPools)
 	{
-		if (Pool->SupportsAllocation(InHeapType, ResourceDesc.Flags, InBufferUsage, InResourceStateMode))
+		if (Pool->SupportsAllocation(InHeapType, ResourceDesc.Flags, InBufferUsage, InResourceStateMode, Alignment))
 		{
 			BufferPool = Pool;
 			break;
@@ -1399,7 +1404,7 @@ void FD3D12DefaultBufferAllocator::AllocDefaultResource(D3D12_HEAP_TYPE InHeapTy
 	// No pool yet, then create one
 	if (BufferPool == nullptr)
 	{
-		BufferPool = CreateBufferPool(InHeapType, ResourceDesc.Flags, InBufferUsage, InResourceStateMode);
+		BufferPool = CreateBufferPool(InHeapType, ResourceDesc.Flags, InBufferUsage, InResourceStateMode, Alignment);
 	}
 
 	// Perform actual allocation
diff --git a/Engine/Source/Runtime/D3D12RHI/Private/D3D12Allocation.h b/Engine/Source/Runtime/D3D12RHI/Private/D3D12Allocation.h
--- a/Engine/Source/Runtime/D3D12RHI/Private/D3D12Allocation.h
+++ b/Engine/Source/Runtime/D3D12RHI/Private/D3D12Allocation.h
@@ -64,8 +64,7 @@ FORCEINLINE uint32 GetTypeHash(const FD3D12SegListAllocatorLeakTrack& S)
 }
 
 
-
-
+const uint32 kD3D12ManualSubAllocationAlignment = 256;
 
 class FD3D12SegList;
 
@@ -390,14 +389,14 @@ public:
 	FD3D12DefaultBufferPool(FD3D12Device* InParent, FD3D12MultiBuddyAllocator* InAllocator);
 	~FD3D12DefaultBufferPool() { delete Allocator; }
 
-	bool SupportsAllocation(D3D12_HEAP_TYPE InHeapType, D3D12_RESOURCE_FLAGS InResourceFlags, EBufferUsageFlags InBufferUsage, ED3D12ResourceStateMode InResourceStateMode) const;
+	bool SupportsAllocation(D3D12_HEAP_TYPE InHeapType, D3D12_RESOURCE_FLAGS InResourceFlags, EBufferUsageFlags InBufferUsage, ED3D12ResourceStateMode InResourceStateMode, uint32 Alignment) const;
 	void AllocDefaultResource(D3D12_HEAP_TYPE InHeapType, const D3D12_RESOURCE_DESC& InDesc, EBufferUsageFlags InBufferUsage, ED3D12ResourceStateMode InResourceStateMode,
 		D3D12_RESOURCE_STATES InCreateState, uint32 InAlignment, const TCHAR* InName, FD3D12ResourceLocation& ResourceLocation);
 	void CleanUpAllocations(uint64 InFrameLag);
 	void UpdateMemoryStats(uint32& IOMemoryAllocated, uint32& IOMemoryUsed, uint32& IOMemoryFree, uint32& IOMemoryEndFree, uint32& IOAlignmentWaste, uint32& IOAllocatedPageCount, uint32& IOFullPageCount);
 
 	static FD3D12ResourceInitConfig GetResourceAllocatorInitConfig(D3D12_HEAP_TYPE InHeapType, D3D12_RESOURCE_FLAGS InResourceFlags, EBufferUsageFlags InBufferUsage);
-	static EResourceAllocationStrategy GetResourceAllocationStrategy(D3D12_RESOURCE_FLAGS InResourceFlags, ED3D12ResourceStateMode InResourceStateMode);
+	static EResourceAllocationStrategy GetResourceAllocationStrategy(D3D12_RESOURCE_FLAGS InResourceFlags, ED3D12ResourceStateMode InResourceStateMode, uint32 Alignment);
 
 private:
 	FD3D12MultiBuddyAllocator* Allocator;
@@ -423,12 +422,12 @@ public:
 	void CleanupFreeBlocks(uint64 InFrameLag);	
 	void UpdateMemoryStats();
 
-	static bool IsPlacedResource(D3D12_RESOURCE_FLAGS InResourceFlags, ED3D12ResourceStateMode InResourceStateMode);
+	static bool IsPlacedResource(D3D12_RESOURCE_FLAGS InResourceFlags, ED3D12ResourceStateMode InResourceStateMode, uint32 Alignment);
 	static D3D12_RESOURCE_STATES GetDefaultInitialResourceState(D3D12_HEAP_TYPE InHeapType, EBufferUsageFlags InBufferFlags, ED3D12ResourceStateMode InResourceStateMode);
 
 private:
 
-	FD3D12BufferPool* CreateBufferPool(D3D12_HEAP_TYPE InHeapType, D3D12_RESOURCE_FLAGS InResourceFlags, EBufferUsageFlags InBufferUsage, ED3D12ResourceStateMode InResourceStateMode);
+	FD3D12BufferPool* CreateBufferPool(D3D12_HEAP_TYPE InHeapType, D3D12_RESOURCE_FLAGS InResourceFlags, EBufferUsageFlags InBufferUsage, ED3D12ResourceStateMode InResourceStateMode, uint32 Alignment);
 
 	TArray<FD3D12BufferPool*> DefaultBufferPools;
 };
diff --git a/Engine/Source/Runtime/D3D12RHI/Private/D3D12Buffer.cpp b/Engine/Source/Runtime/D3D12RHI/Private/D3D12Buffer.cpp
--- a/Engine/Source/Runtime/D3D12RHI/Private/D3D12Buffer.cpp
+++ b/Engine/Source/Runtime/D3D12RHI/Private/D3D12Buffer.cpp
@@ -524,7 +524,7 @@ FD3D12Buffer* FD3D12DynamicRHI::CreateD3D12Buffer(class FRHICommandListImmediate
 
 	// Does this resource support tracking?
 	const bool bIsDynamic = (Usage & BUF_AnyDynamic) ? true : false;
-	const bool bSupportResourceStateTracking = !bIsDynamic && FD3D12DefaultBufferAllocator::IsPlacedResource(Desc.Flags, StateMode);
+	const bool bSupportResourceStateTracking = !bIsDynamic && FD3D12DefaultBufferAllocator::IsPlacedResource(Desc.Flags, StateMode, Alignment);
 
 	// Initial state is derived from the InResourceState if it's supports tracking
 	D3D12_HEAP_TYPE HeapType = bIsDynamic ? D3D12_HEAP_TYPE_UPLOAD : D3D12_HEAP_TYPE_DEFAULT;
diff --git a/Engine/Source/Runtime/D3D12RHI/Private/D3D12PoolAllocator.cpp b/Engine/Source/Runtime/D3D12RHI/Private/D3D12PoolAllocator.cpp
--- a/Engine/Source/Runtime/D3D12RHI/Private/D3D12PoolAllocator.cpp
+++ b/Engine/Source/Runtime/D3D12RHI/Private/D3D12PoolAllocator.cpp
@@ -165,9 +165,13 @@ FD3D12ResourceInitConfig FD3D12PoolAllocator::GetResourceAllocatorInitConfig(D3D
 	return InitConfig;
 }
 
-
-EResourceAllocationStrategy FD3D12PoolAllocator::GetResourceAllocationStrategy(D3D12_RESOURCE_FLAGS InResourceFlags, ED3D12ResourceStateMode InResourceStateMode)
+EResourceAllocationStrategy FD3D12PoolAllocator::GetResourceAllocationStrategy(D3D12_RESOURCE_FLAGS InResourceFlags, ED3D12ResourceStateMode InResourceStateMode, uint32 Alignment)
 {
+	if (Alignment > kD3D12ManualSubAllocationAlignment)
+	{
+		return EResourceAllocationStrategy::kPlacedResource;
+	}
+
 	// Does the resource need state tracking and transitions
 	ED3D12ResourceStateMode ResourceStateMode = InResourceStateMode;
 	if (ResourceStateMode == ED3D12ResourceStateMode::Default)
@@ -197,10 +201,10 @@ FD3D12PoolAllocator::~FD3D12PoolAllocator()
 }
 
 
-bool FD3D12PoolAllocator::SupportsAllocation(D3D12_HEAP_TYPE InHeapType, D3D12_RESOURCE_FLAGS InResourceFlags, EBufferUsageFlags InBufferUsage, ED3D12ResourceStateMode InResourceStateMode) const
+bool FD3D12PoolAllocator::SupportsAllocation(D3D12_HEAP_TYPE InHeapType, D3D12_RESOURCE_FLAGS InResourceFlags, EBufferUsageFlags InBufferUsage, ED3D12ResourceStateMode InResourceStateMode, uint32 Alignment) const
 {
 	FD3D12ResourceInitConfig InInitConfig = GetResourceAllocatorInitConfig(InHeapType, InResourceFlags, InBufferUsage);
-	EResourceAllocationStrategy InAllocationStrategy = GetResourceAllocationStrategy(InResourceFlags, InResourceStateMode);
+	EResourceAllocationStrategy InAllocationStrategy = GetResourceAllocationStrategy(InResourceFlags, InResourceStateMode, Alignment);
 	return (InitConfig == InInitConfig && AllocationStrategy == InAllocationStrategy);
 }
 
diff --git a/Engine/Source/Runtime/D3D12RHI/Private/D3D12PoolAllocator.h b/Engine/Source/Runtime/D3D12RHI/Private/D3D12PoolAllocator.h
--- a/Engine/Source/Runtime/D3D12RHI/Private/D3D12PoolAllocator.h
+++ b/Engine/Source/Runtime/D3D12RHI/Private/D3D12PoolAllocator.h
@@ -131,7 +131,7 @@ public:
 	~FD3D12PoolAllocator();
 	
 	// Function names currently have to match with FD3D12DefaultBufferPool until we can make full replacement of the allocator
-	bool SupportsAllocation(D3D12_HEAP_TYPE InHeapType, D3D12_RESOURCE_FLAGS InResourceFlags, EBufferUsageFlags InBufferUsage, ED3D12ResourceStateMode InResourceStateMode) const;
+	bool SupportsAllocation(D3D12_HEAP_TYPE InHeapType, D3D12_RESOURCE_FLAGS InResourceFlags, EBufferUsageFlags InBufferUsage, ED3D12ResourceStateMode InResourceStateMode, uint32 Alignment) const;
 	void AllocDefaultResource(D3D12_HEAP_TYPE InHeapType, const D3D12_RESOURCE_DESC& InDesc, EBufferUsageFlags InBufferUsage, ED3D12ResourceStateMode InResourceStateMode,
 		D3D12_RESOURCE_STATES InCreateState, uint32 InAlignment, const TCHAR* InName, FD3D12ResourceLocation& ResourceLocation);
 	virtual void DeallocateResource(FD3D12ResourceLocation& ResourceLocation, bool bDefragFree = false);
@@ -152,7 +152,7 @@ public:
 	FD3D12HeapAndOffset GetBackingHeapAndAllocationOffsetInBytes(const FRHIPoolAllocationData& InAllocationData) const;
 
 	static FD3D12ResourceInitConfig GetResourceAllocatorInitConfig(D3D12_HEAP_TYPE InHeapType, D3D12_RESOURCE_FLAGS InResourceFlags, EBufferUsageFlags InBufferUsage);
-	static EResourceAllocationStrategy GetResourceAllocationStrategy(D3D12_RESOURCE_FLAGS InResourceFlags, ED3D12ResourceStateMode InResourceStateMode);
+	static EResourceAllocationStrategy GetResourceAllocationStrategy(D3D12_RESOURCE_FLAGS InResourceFlags, ED3D12ResourceStateMode InResourceStateMode, uint32 Alignment);
 	
 protected:
 
