Track d3d12 heap allocations as well for page fault checking

diff --git a/Engine/Source/Runtime/D3D12RHI/Private/D3D12Adapter.cpp b/Engine/Source/Runtime/D3D12RHI/Private/D3D12Adapter.cpp
--- a/Engine/Source/Runtime/D3D12RHI/Private/D3D12Adapter.cpp
+++ b/Engine/Source/Runtime/D3D12RHI/Private/D3D12Adapter.cpp
@@ -1293,13 +1293,15 @@ void FD3D12Adapter::EndFrame()
 #endif
 
 #if TRACK_RESOURCE_ALLOCATIONS
+	FScopeLock Lock(&TrackedAllocationDataCS);
+
 	// remove tracked released resources older than n amount of frames
-	static const uint64 FrameRetention = 100;
+	static const uint64 ReleaseFrameRetention = 100;
 	int32 ReleaseCount = 0;
 	uint64 CurrentFrameID = GetFrameFence().GetCurrentFence();
 	for (; ReleaseCount < ReleasedAllocationData.Num(); ++ReleaseCount)
 	{
-		if (ReleasedAllocationData[ReleaseCount].ReleasedFrameID + FrameRetention > CurrentFrameID)
+		if (ReleasedAllocationData[ReleaseCount].ReleasedFrameID + ReleaseFrameRetention > CurrentFrameID)
 		{
 			break;
 		}
@@ -1452,6 +1454,37 @@ void FD3D12Adapter::ReleaseTrackedAllocationData(FD3D12ResourceLocation* InAlloc
 #endif
 }
 
+
+void FD3D12Adapter::TrackHeapAllocation(FD3D12Heap* InHeap)
+{
+#if TRACK_RESOURCE_ALLOCATIONS
+	FScopeLock Lock(&TrackedAllocationDataCS);
+	check(!TrackedHeaps.Contains(InHeap));
+	TrackedHeaps.Add(InHeap);
+#endif
+}
+
+void FD3D12Adapter::ReleaseTrackedHeap(FD3D12Heap* InHeap)
+{
+#if TRACK_RESOURCE_ALLOCATIONS
+	FScopeLock Lock(&TrackedAllocationDataCS);
+
+	D3D12_GPU_VIRTUAL_ADDRESS GPUVirtualAddress = InHeap->GetGPUVirtualAddress();
+	if (GPUVirtualAddress != 0 || IsTrackingAllAllocations())
+	{
+		FReleasedAllocationData ReleasedData;
+		ReleasedData.GPUVirtualAddress	= GPUVirtualAddress;
+		ReleasedData.AllocationSize		= InHeap->GetHeapDesc().SizeInBytes;
+		ReleasedData.ResourceName		= InHeap->GetName();
+		ReleasedData.ReleasedFrameID	= GetFrameFence().GetCurrentFence();
+		ReleasedData.bHeap				= true;
+		ReleasedAllocationData.Add(ReleasedData);
+	}
+
+	verify(TrackedHeaps.Remove(InHeap) == 1);
+#endif
+}
+
 void FD3D12Adapter::FindResourcesNearGPUAddress(D3D12_GPU_VIRTUAL_ADDRESS InGPUVirtualAddress, uint64 InRange, TArray<FAllocatedResourceResult>& OutResources)
 {
 #if TRACK_RESOURCE_ALLOCATIONS
@@ -1468,7 +1501,7 @@ void FD3D12Adapter::FindResourcesNearGPUAddress(D3D12_GPU_VIRTUAL_ADDRESS InGPUV
 		{
 			bool bContainsAllocation = AllocationRange.Contains(InGPUVirtualAddress);
 			int64 Distance = bContainsAllocation ? 0 : ((InGPUVirtualAddress < GPUAddress) ? GPUAddress - InGPUVirtualAddress : InGPUVirtualAddress - AllocationRange.GetUpperBoundValue());
-			check(Distance > 0);
+			check(Distance >= 0);
 
 			FAllocatedResourceResult Result;
 			Result.Allocation = AllocationData.ResourceAllocation;
@@ -1485,6 +1518,23 @@ void FD3D12Adapter::FindResourcesNearGPUAddress(D3D12_GPU_VIRTUAL_ADDRESS InGPUV
 #endif
 }
 
+void FD3D12Adapter::FindHeapsContainingGPUAddress(D3D12_GPU_VIRTUAL_ADDRESS InGPUVirtualAddress, TArray<FD3D12Heap*>& OutHeaps)
+{
+#if TRACK_RESOURCE_ALLOCATIONS
+	FScopeLock Lock(&TrackedAllocationDataCS);
+
+	for (FD3D12Heap* AllocatedHeap : TrackedHeaps)
+	{
+		D3D12_GPU_VIRTUAL_ADDRESS GPUAddress = AllocatedHeap->GetGPUVirtualAddress();
+		FInt64Range HeapRange((int64)GPUAddress, (int64)(GPUAddress + AllocatedHeap->GetHeapDesc().SizeInBytes));
+		if (HeapRange.Contains(InGPUVirtualAddress))
+		{			
+			OutHeaps.Add(AllocatedHeap);
+		}
+	}
+#endif
+}
+
 void FD3D12Adapter::FindReleasedAllocationData(D3D12_GPU_VIRTUAL_ADDRESS InGPUVirtualAddress, TArray<FReleasedAllocationData>& OutAllocationData)
 {
 #if TRACK_RESOURCE_ALLOCATIONS
diff --git a/Engine/Source/Runtime/D3D12RHI/Private/D3D12Adapter.h b/Engine/Source/Runtime/D3D12RHI/Private/D3D12Adapter.h
--- a/Engine/Source/Runtime/D3D12RHI/Private/D3D12Adapter.h
+++ b/Engine/Source/Runtime/D3D12RHI/Private/D3D12Adapter.h
@@ -357,6 +357,8 @@ public:
 	bool IsTrackingAllAllocations() const { return bTrackAllAllocation; }
 	void TrackAllocationData(FD3D12ResourceLocation* InAllocation, uint64 InAllocationSize, bool bCollectCallstack);
 	void ReleaseTrackedAllocationData(FD3D12ResourceLocation* InAllocation, bool bDefragFree);
+	void TrackHeapAllocation(FD3D12Heap* InHeap);
+	void ReleaseTrackedHeap(FD3D12Heap* InHeap);
 	void DumpTrackedAllocationData(FOutputDevice& OutputDevice, bool bResidentOnly, bool bWithCallstack);
 
 	struct FAllocatedResourceResult
@@ -365,6 +367,7 @@ public:
 		uint64 Distance = 0;
 	};
 	void FindResourcesNearGPUAddress(D3D12_GPU_VIRTUAL_ADDRESS InGPUVirtualAddress, uint64 InRange, TArray<FAllocatedResourceResult>& OutResources);
+	void FindHeapsContainingGPUAddress(D3D12_GPU_VIRTUAL_ADDRESS InGPUVirtualAddress, TArray<FD3D12Heap*>& OutHeaps);
 	
 	struct FReleasedAllocationData
 	{
@@ -375,6 +378,7 @@ public:
 		bool bDefragFree = false;
 		bool bBackBuffer = false;
 		bool bTransient = false;
+		bool bHeap = false;
 	};
 	void FindReleasedAllocationData(D3D12_GPU_VIRTUAL_ADDRESS InGPUVirtualAddress, TArray<FReleasedAllocationData>& OutAllocationData);
 
@@ -507,6 +511,7 @@ protected:
 
 	/** Tracked resource information. */
 	TMap<FD3D12ResourceLocation*, FTrackedAllocationData> TrackedAllocationData;
+	TArray<FD3D12Heap*> TrackedHeaps;
 	TArray<FReleasedAllocationData> ReleasedAllocationData;
 	FCriticalSection TrackedAllocationDataCS;
 
diff --git a/Engine/Source/Runtime/D3D12RHI/Private/D3D12Allocation.cpp b/Engine/Source/Runtime/D3D12RHI/Private/D3D12Allocation.cpp
--- a/Engine/Source/Runtime/D3D12RHI/Private/D3D12Allocation.cpp
+++ b/Engine/Source/Runtime/D3D12RHI/Private/D3D12Allocation.cpp
@@ -259,10 +259,9 @@ void FD3D12BuddyAllocator::Initialize()
 			LLM_SCOPED_PAUSE_TRACKING_FOR_TRACKER(ELLMTracker::Default, ELLMAllocType::System);
 			VERIFYD3D12RESULT(Adapter->GetD3DDevice()->CreateHeap(&Desc, IID_PPV_ARGS(&Heap)));
 		}
-		SetName(Heap, L"Placed Resource Allocator Backing Heap");
 
 		BackingHeap = new FD3D12Heap(GetParentDevice(), GetVisibilityMask());
-		BackingHeap->SetHeap(Heap);
+		BackingHeap->SetHeap(Heap, TEXT("BuddyAllocator Heap"));
 
 		// Only track resources that cannot be accessed on the CPU.
 		if (IsGPUOnly(InitConfig.HeapType))
@@ -559,11 +558,6 @@ void FD3D12BuddyAllocator::ReleaseAllResources()
 		ensure(BackingResource->GetRefCount() == 1 || GNumExplicitGPUsForRendering > 1);
 		BackingResource = nullptr;
 	}
-
-	if (BackingHeap)
-	{
-		BackingHeap->Destroy();
-	}
 }
 
 void FD3D12BuddyAllocator::DumpAllocatorStats(class FOutputDevice& Ar)
diff --git a/Engine/Source/Runtime/D3D12RHI/Private/D3D12Allocation.h b/Engine/Source/Runtime/D3D12RHI/Private/D3D12Allocation.h
--- a/Engine/Source/Runtime/D3D12RHI/Private/D3D12Allocation.h
+++ b/Engine/Source/Runtime/D3D12RHI/Private/D3D12Allocation.h
@@ -548,7 +548,7 @@ private:
 		ArrayIdx(Idx),
 		FirstFreeOffset(0)
 	{
-		this->SetHeap(NewHeap);
+		this->SetHeap(NewHeap, TEXT("SegListHeap"));
 		BeginTrackingResidency(HeapSize);
 	}
 
diff --git a/Engine/Source/Runtime/D3D12RHI/Private/D3D12DirectCommandListManager.cpp b/Engine/Source/Runtime/D3D12RHI/Private/D3D12DirectCommandListManager.cpp
--- a/Engine/Source/Runtime/D3D12RHI/Private/D3D12DirectCommandListManager.cpp
+++ b/Engine/Source/Runtime/D3D12RHI/Private/D3D12DirectCommandListManager.cpp
@@ -444,7 +444,7 @@ void FD3D12CommandListManager::Create(const TCHAR* Name, uint32 NumCommandLists,
 				if (SUCCEEDED(hr))
 				{
 					BreadCrumbHeap = new FD3D12Heap(Device, GetVisibilityMask());
-					BreadCrumbHeap->SetHeap(D3D12Heap);
+					BreadCrumbHeap->SetHeap(D3D12Heap, TEXT("BreadCrumbHeap"));
 
 					TCHAR TempStr[MAX_SPRINTF] = TEXT("");
 					FCString::Sprintf(TempStr, TEXT("BreadCrumbResource_%s"), Name);
diff --git a/Engine/Source/Runtime/D3D12RHI/Private/D3D12PoolAllocator.cpp b/Engine/Source/Runtime/D3D12RHI/Private/D3D12PoolAllocator.cpp
--- a/Engine/Source/Runtime/D3D12RHI/Private/D3D12PoolAllocator.cpp
+++ b/Engine/Source/Runtime/D3D12RHI/Private/D3D12PoolAllocator.cpp
@@ -70,10 +70,9 @@ void FD3D12MemoryPool::Init()
 			LLM_SCOPED_PAUSE_TRACKING_FOR_TRACKER(ELLMTracker::Default, ELLMAllocType::System);
 			VERIFYD3D12RESULT(Adapter->GetD3DDevice()->CreateHeap(&Desc, IID_PPV_ARGS(&Heap)));
 		}
-		SetName(Heap, L"LinkListAllocator Backing Heap");
 
 		BackingHeap = new FD3D12Heap(GetParentDevice(), GetVisibilityMask());
-		BackingHeap->SetHeap(Heap);
+		BackingHeap->SetHeap(Heap, TEXT("PoolAllocator Heap"));
 
 		// Only track resources that cannot be accessed on the CPU.
 		if (IsGPUOnly(InitConfig.HeapType))
@@ -110,11 +109,6 @@ void FD3D12MemoryPool::Destroy()
 		ensure(BackingResource->GetRefCount() == 1 || GNumExplicitGPUsForRendering > 1);
 		BackingResource = nullptr;
 	}
-
-	if (BackingHeap)
-	{
-		BackingHeap->Destroy();
-	}
 }
 
 
@@ -355,7 +349,8 @@ void FD3D12PoolAllocator::AllocateResource(uint32 GPUIndex, D3D12_HEAP_TYPE InHe
 			ID3D12Heap* Heap = nullptr;
 			VERIFYD3D12RESULT(Adapter->GetD3DDevice()->CreateHeap(&HeapDesc, IID_PPV_ARGS(&Heap)));
 			TRefCountPtr<FD3D12Heap> BackingHeap = new FD3D12Heap(GetParentDevice(), GetVisibilityMask());
-			BackingHeap->SetHeap(Heap);		
+			bool bTrack = false;
+			BackingHeap->SetHeap(Heap, InName, bTrack);		
 			
 			VERIFYD3D12RESULT(Adapter->CreatePlacedResource(Desc, BackingHeap, 0, InCreateState, InResourceStateMode, InCreateState, InClearValue, &NewResource, InName));
 		}
diff --git a/Engine/Source/Runtime/D3D12RHI/Private/D3D12Resources.cpp b/Engine/Source/Runtime/D3D12RHI/Private/D3D12Resources.cpp
--- a/Engine/Source/Runtime/D3D12RHI/Private/D3D12Resources.cpp
+++ b/Engine/Source/Runtime/D3D12RHI/Private/D3D12Resources.cpp
@@ -407,26 +407,52 @@ FD3D12Heap::FD3D12Heap(FD3D12Device* Parent, FRHIGPUMask VisibleNodes) :
 
 FD3D12Heap::~FD3D12Heap()
 {
-	Destroy();
+#if TRACK_RESOURCE_ALLOCATIONS
+	FD3D12Adapter* Adapter = GetParentDevice()->GetParentAdapter();
+	if (GPUVirtualAddress != 0 && bTrack)
+	{
+		Adapter->ReleaseTrackedHeap(this);
+	}
+#endif // TRACK_RESOURCE_ALLOCATIONS
+
+#if ENABLE_RESIDENCY_MANAGEMENT
+	if (D3DX12Residency::IsInitialized(ResidencyHandle))
+	{
+		D3DX12Residency::EndTrackingObject(GetParentDevice()->GetResidencyManager(), ResidencyHandle);
+		ResidencyHandle = {};
+	}
+#endif // ENABLE_RESIDENCY_MANAGEMENT
+
+	// Release actual d3d object
+	Heap.SafeRelease();
 }
 
-void FD3D12Heap::SetHeap(ID3D12Heap* HeapIn)
+void FD3D12Heap::SetHeap(ID3D12Heap* HeapIn, const TCHAR* const InName, bool bInTrack)
 {
 	*Heap.GetInitReference() = HeapIn; 
-	
-	D3D12_HEAP_DESC HeapDesc = Heap->GetDesc();
-	HeapType = HeapDesc.Properties.Type;
+	bTrack = bInTrack;
+	HeapName = InName;
+	HeapDesc = Heap->GetDesc();
 
+	SetName(HeapIn, InName);
+	
 	// Create a buffer placed resource on the heap to extract the gpu virtual address
 	// if we are tracking all allocations
 	FD3D12Adapter* Adapter = GetParentDevice()->GetParentAdapter();	
-	if (Adapter->IsTrackingAllAllocations() && HeapType == D3D12_HEAP_TYPE_DEFAULT)
+	if (Adapter->IsTrackingAllAllocations() && HeapDesc.Properties.Type == D3D12_HEAP_TYPE_DEFAULT)
 	{
 		uint64 HeapSize = HeapDesc.SizeInBytes;
 		TRefCountPtr<ID3D12Resource> TempResource;
 		const D3D12_RESOURCE_DESC BufDesc = CD3DX12_RESOURCE_DESC::Buffer(HeapSize, D3D12_RESOURCE_FLAG_NONE);
 		const HRESULT hr = Adapter->GetD3DDevice()->CreatePlacedResource(Heap, 0, &BufDesc, D3D12_RESOURCE_STATE_COMMON, nullptr, IID_PPV_ARGS(TempResource.GetInitReference()));
 		GPUVirtualAddress = TempResource->GetGPUVirtualAddress();
+				
+#if TRACK_RESOURCE_ALLOCATIONS
+		if (bTrack)
+		{
+			Adapter->TrackHeapAllocation(this);
+		}
+#endif
 	}
 }
 
@@ -440,16 +466,6 @@ void FD3D12Heap::UpdateResidency(FD3D12CommandListHandle& CommandList)
 #endif
 }
 
-void FD3D12Heap::Destroy()
-{
-	//TODO: Check ref counts?
-	if (D3DX12Residency::IsInitialized(ResidencyHandle))
-	{
-		D3DX12Residency::EndTrackingObject(GetParentDevice()->GetResidencyManager(), ResidencyHandle);
-		ResidencyHandle = {};
-	}
-}
-
 void FD3D12Heap::BeginTrackingResidency(uint64 Size)
 {
 #if ENABLE_RESIDENCY_MANAGEMENT
@@ -552,7 +568,7 @@ HRESULT FD3D12Adapter::CreatePlacedResource(const D3D12_RESOURCE_DESC& InDesc, F
 			HeapDesc.Properties.Type);
 
 #if PLATFORM_WINDOWS
-		if (IsTrackingAllAllocations() && BackingHeap->GetHeapType() == D3D12_HEAP_TYPE_DEFAULT)
+		if (IsTrackingAllAllocations() && BackingHeap->GetHeapDesc().Properties.Type == D3D12_HEAP_TYPE_DEFAULT)
 		{
 			// Manually set the GPU virtual address from the heap gpu virtual address & offset
 			if (InDesc.Dimension != D3D12_RESOURCE_DIMENSION_BUFFER)
diff --git a/Engine/Source/Runtime/D3D12RHI/Private/D3D12StateCachePrivate.h b/Engine/Source/Runtime/D3D12RHI/Private/D3D12StateCachePrivate.h
--- a/Engine/Source/Runtime/D3D12RHI/Private/D3D12StateCachePrivate.h
+++ b/Engine/Source/Runtime/D3D12RHI/Private/D3D12StateCachePrivate.h
@@ -850,7 +850,7 @@ public:
 				continue;
 			}
 
-			for (int32 j = 0; j < PipelineState.Common.SRVCache.MaxBoundIndex[i]; ++j)
+			for (int32 j = 0; j <= PipelineState.Common.SRVCache.MaxBoundIndex[i]; ++j)
 			{
 				if (PipelineState.Common.SRVCache.Views[i][j] && PipelineState.Common.SRVCache.Views[i][j]->GetResourceLocation())
 				{
diff --git a/Engine/Source/Runtime/D3D12RHI/Private/D3D12TransientResourceAllocator.cpp b/Engine/Source/Runtime/D3D12RHI/Private/D3D12TransientResourceAllocator.cpp
--- a/Engine/Source/Runtime/D3D12RHI/Private/D3D12TransientResourceAllocator.cpp
+++ b/Engine/Source/Runtime/D3D12RHI/Private/D3D12TransientResourceAllocator.cpp
@@ -82,17 +82,15 @@ FD3D12TransientHeap::FD3D12TransientHeap(const FRHITransientHeapInitializer& Ini
 		}
 #endif // PLATFORM_WINDOWS
 	}
-	SetName(D3DHeap, L"TransientResourceAllocator Backing Heap");
 
 	Heap = new FD3D12Heap(Device, VisibleNodeMask);
-	Heap->SetHeap(D3DHeap);
+	Heap->SetHeap(D3DHeap, TEXT("TransientResourceAllocator Backing Heap"));
 	Heap->BeginTrackingResidency(Desc.SizeInBytes);
 }
 
 FD3D12TransientHeap::~FD3D12TransientHeap()
 {
 	LLM_SCOPED_PAUSE_TRACKING_FOR_TRACKER(ELLMTracker::Default, ELLMAllocType::System);
-	Heap->Destroy();
 }
 
 TUniquePtr<FD3D12TransientResourceSystem> FD3D12TransientResourceSystem::Create(FD3D12Adapter* ParentAdapter, FRHIGPUMask VisibleNodeMask)
@@ -216,4 +214,4 @@ void FD3D12TransientResourceAllocator::DeallocateMemory(FRHITransientBuffer* InB
 void FD3D12TransientResourceAllocator::Freeze(FRHICommandListImmediate& RHICmdList)
 {
 	Allocator.Freeze(RHICmdList);
-}
\ No newline at end of file
+}
diff --git a/Engine/Source/Runtime/D3D12RHI/Private/D3D12Util.cpp b/Engine/Source/Runtime/D3D12RHI/Private/D3D12Util.cpp
--- a/Engine/Source/Runtime/D3D12RHI/Private/D3D12Util.cpp
+++ b/Engine/Source/Runtime/D3D12RHI/Private/D3D12Util.cpp
@@ -567,7 +567,10 @@ void LogPageFaultData(FD3D12Adapter* InAdapter, D3D12_GPU_VIRTUAL_ADDRESS InPage
 		return;
 	}
 
+	FD3D12ManualFence& FrameFence = InAdapter->GetFrameFence();
+
 	UE_LOG(LogD3D12RHI, Error, TEXT("PageFault: PageFault at VA GPUAddress \"0x%llX\""), (long long)InPageFaultAddress);
+	UE_LOG(LogD3D12RHI, Error, TEXT("PageFault: Last completed frame ID: %d (cached: %d) - Current frame ID: %d"), FrameFence.PeekLastCompletedFence(), FrameFence.GetLastCompletedFenceFast(), FrameFence.GetCurrentFence());
 	UE_LOG(LogD3D12RHI, Error, TEXT("PageFault: Logging all resource enabled: %s"), InAdapter->IsTrackingAllAllocations() ? TEXT("Yes") : TEXT("No"));
 
 	// Try and find all current allocations near that range
@@ -594,17 +597,28 @@ void LogPageFaultData(FD3D12Adapter* InAdapter, D3D12_GPU_VIRTUAL_ADDRESS InPage
 		}
 	}
 
+	// Try and find all current heaps containing the page fault address
+	TArray<FD3D12Heap*> OverlappingHeaps;
+	InAdapter->FindHeapsContainingGPUAddress(InPageFaultAddress, OverlappingHeaps);
+	UE_LOG(LogD3D12RHI, Error, TEXT("PageFault: Found %d active heaps containing page fault address"), OverlappingHeaps.Num());
+	for (int32 Index = 0; Index < OverlappingHeaps.Num(); ++Index)
+	{
+		FD3D12Heap* Heap = OverlappingHeaps[Index];
+		UE_LOG(LogD3D12RHI, Error, TEXT("\tGPU Address: \"0x%llX\" - Size: %3.2f MB - Name: %s"),
+			(long long)Heap->GetGPUVirtualAddress(), Heap->GetHeapDesc().SizeInBytes / (1024.0f * 1024), *(Heap->GetName().ToString()));
+	}
+
 	// Try and find all released allocations within the faulting address
 	TArray<FD3D12Adapter::FReleasedAllocationData> ReleasedResources;
 	InAdapter->FindReleasedAllocationData(InPageFaultAddress, ReleasedResources);
-	UE_LOG(LogD3D12RHI, Error, TEXT("PageFault: Found %d released resources containing the page fault address during last 100 frames (current frame ID: %d)"), ReleasedResources.Num(), InAdapter->GetFrameFence().GetCurrentFence());
+	UE_LOG(LogD3D12RHI, Error, TEXT("PageFault: Found %d released resources containing the page fault address during last 100 frames"), ReleasedResources.Num());
 	if (ReleasedResources.Num() > 0)
 	{
 		uint32 PrintCount = FMath::Min(ReleasedResources.Num(), 100);
 		for (uint32 Index = 0; Index < PrintCount; ++Index)
 		{
 			FD3D12Adapter::FReleasedAllocationData& AllocationData = ReleasedResources[Index];
-			UE_LOG(LogD3D12RHI, Error, TEXT("\tGPU Address: [0x%llX .. 0x%llX] - Size: %lld bytes, %3.2f MB - FrameID: %4d - DefragFree: %d - Transient: %d - Name: %s"),
+			UE_LOG(LogD3D12RHI, Error, TEXT("\tGPU Address: [0x%llX .. 0x%llX] - Size: %lld bytes, %3.2f MB - FrameID: %4d - DefragFree: %d - Transient: %d - Heap: %d - Name: %s"),
 				(uint64)AllocationData.GPUVirtualAddress,
 				(uint64)AllocationData.GPUVirtualAddress + AllocationData.AllocationSize,
 				AllocationData.AllocationSize,
@@ -612,6 +626,7 @@ void LogPageFaultData(FD3D12Adapter* InAdapter, D3D12_GPU_VIRTUAL_ADDRESS InPage
 				AllocationData.ReleasedFrameID,
 				AllocationData.bDefragFree,
 				AllocationData.bTransient,
+				AllocationData.bHeap,
 				*AllocationData.ResourceName.ToString());
 		}
 	}
diff --git a/Engine/Source/Runtime/D3D12RHI/Public/D3D12Resources.h b/Engine/Source/Runtime/D3D12RHI/Public/D3D12Resources.h
--- a/Engine/Source/Runtime/D3D12RHI/Public/D3D12Resources.h
+++ b/Engine/Source/Runtime/D3D12RHI/Public/D3D12Resources.h
@@ -94,19 +94,22 @@ public:
 	~FD3D12Heap();
 
 	inline ID3D12Heap* GetHeap() { return Heap.GetReference(); }
-	void SetHeap(ID3D12Heap* HeapIn);
+	void SetHeap(ID3D12Heap* HeapIn, const TCHAR* const InName, bool bTrack = true);
 
 	void UpdateResidency(FD3D12CommandListHandle& CommandList);
 	void BeginTrackingResidency(uint64 Size);
-	void Destroy();
-
-	inline D3D12_HEAP_TYPE GetHeapType() const { return HeapType; }
+	
+	inline FName GetName() const { return HeapName; }
+	inline D3D12_HEAP_DESC GetHeapDesc() const { return HeapDesc; }
 	inline FD3D12ResidencyHandle* GetResidencyHandle() { return &ResidencyHandle; }
 	inline D3D12_GPU_VIRTUAL_ADDRESS GetGPUVirtualAddress() const { return GPUVirtualAddress; }
 
 private:
+
 	TRefCountPtr<ID3D12Heap> Heap;
-	D3D12_HEAP_TYPE HeapType = D3D12_HEAP_TYPE_DEFAULT;
+	FName HeapName;
+	bool bTrack = true;
+	D3D12_HEAP_DESC HeapDesc;
 	D3D12_GPU_VIRTUAL_ADDRESS GPUVirtualAddress = 0;
 	FD3D12ResidencyHandle ResidencyHandle;
 };
