Make released tracked resource frame retention a cvar instead of hardcoded 100 frames

diff --git a/Engine/Source/Runtime/D3D12RHI/Private/D3D12Adapter.cpp b/Engine/Source/Runtime/D3D12RHI/Private/D3D12Adapter.cpp
--- a/Engine/Source/Runtime/D3D12RHI/Private/D3D12Adapter.cpp
+++ b/Engine/Source/Runtime/D3D12RHI/Private/D3D12Adapter.cpp
@@ -50,6 +50,16 @@ static FAutoConsoleVariableRef CVarGapRecorderUseBlockingCall(
 );
 #endif
 
+#if TRACK_RESOURCE_ALLOCATIONS
+int32 GTrackedReleasedAllocationFrameRetention = 100;
+static FAutoConsoleVariableRef CTrackedReleasedAllocationFrameRetention(
+	TEXT("D3D12.TrackedReleasedAllocationFrameRetention"),
+	GTrackedReleasedAllocationFrameRetention,
+	TEXT("Amount of frames for which we keep freed allocation data around when resource tracking is enabled"),
+	ECVF_RenderThreadSafe
+);
+#endif
+
 #if PLATFORM_WINDOWS || PLATFORM_HOLOLENS
 
 #if UE_BUILD_SHIPPING || UE_BUILD_TEST
@@ -1293,15 +1303,14 @@ void FD3D12Adapter::EndFrame()
 #endif
 
 #if TRACK_RESOURCE_ALLOCATIONS
-	FScopeLock Lock(&TrackedAllocationDataCS);
+	FScopeLock Lock(&TrackedAllocationDataCS); 
 
 	// remove tracked released resources older than n amount of frames
-	static const uint64 ReleaseFrameRetention = 100;
 	int32 ReleaseCount = 0;
 	uint64 CurrentFrameID = GetFrameFence().GetCurrentFence();
 	for (; ReleaseCount < ReleasedAllocationData.Num(); ++ReleaseCount)
 	{
-		if (ReleasedAllocationData[ReleaseCount].ReleasedFrameID + ReleaseFrameRetention > CurrentFrameID)
+		if (ReleasedAllocationData[ReleaseCount].ReleasedFrameID + GTrackedReleasedAllocationFrameRetention > CurrentFrameID)
 		{
 			break;
 		}
