Fallback to upload heap allocator when fast allocator doesn't fit page size anymore for upload heap allocations - upload heap allocator supports dynamic sized allocation instead of immediatly falling back to committed resource allocation

diff --git a/Engine/Source/Runtime/D3D12RHI/Private/D3D12Allocation.cpp b/Engine/Source/Runtime/D3D12RHI/Private/D3D12Allocation.cpp
--- a/Engine/Source/Runtime/D3D12RHI/Private/D3D12Allocation.cpp
+++ b/Engine/Source/Runtime/D3D12RHI/Private/D3D12Allocation.cpp
@@ -994,9 +994,9 @@ FD3D12UploadHeapAllocator::FD3D12UploadHeapAllocator(FD3D12Adapter* InParent, FD
 	, FD3D12DeviceChild(InParentDevice)
 	, FD3D12MultiNodeGPUObject(InParentDevice->GetGPUMask(), FRHIGPUMask::All()), // Upload memory, thus they can be trivially visibile to all GPUs
 	SmallBlockAllocator(InParentDevice, GetVisibilityMask(), FD3D12ResourceInitConfig::CreateUpload(), InName, EResourceAllocationStrategy::kManualSubAllocation,
-		GD3D12UploadHeapSmallBlockMaxAllocationSize, GD3D12UploadHeapSmallBlockPoolSize, 256),
+		GD3D12UploadHeapSmallBlockMaxAllocationSize, GD3D12UploadHeapSmallBlockPoolSize, D3D12_TEXTURE_DATA_PLACEMENT_ALIGNMENT),
 	BigBlockAllocator(InParentDevice, GetVisibilityMask(), FD3D12ResourceInitConfig::CreateUpload(), InName, EResourceAllocationStrategy::kManualSubAllocation,
-		GD3D12UploadHeapBigBlockPoolSize, 256, GD3D12UploadHeapBigBlockMaxAllocationSize, FRHIMemoryPool::EFreeListOrder::SortByOffset, false /*defrag*/),
+		GD3D12UploadHeapBigBlockPoolSize, D3D12_TEXTURE_DATA_PLACEMENT_ALIGNMENT, GD3D12UploadHeapBigBlockMaxAllocationSize, FRHIMemoryPool::EFreeListOrder::SortByOffset, false /*defrag*/),
 	FastConstantPageAllocator(InParentDevice, GetVisibilityMask(), FD3D12ResourceInitConfig::CreateUpload(), InName, EResourceAllocationStrategy::kManualSubAllocation,
 		GD3D12FastConstantAllocatorPageSize * 64, GD3D12UploadHeapSmallBlockPoolSize, GD3D12FastConstantAllocatorPageSize)
 {
@@ -1880,6 +1880,12 @@ void* FD3D12FastAllocator::Allocate(uint32 Size, uint32 Alignment, class FD3D12R
 	{
 		FD3D12Adapter* Adapter = GetParentDevice()->GetParentAdapter();
 
+		// If upload memory then fallback to the shader upload heap allocator which support dynamic sized allocation of bigger sizes
+		if (PagePool.GetHeapType() == D3D12_HEAP_TYPE_UPLOAD)
+		{
+			return Adapter->GetUploadHeapAllocator(GetGPUMask().ToIndex()).AllocUploadResource(Size, Alignment, *ResourceLocation);
+		}
+
 		//Allocations are 64k aligned
 		if (Alignment)
 		{
