Force IndirectArg buffers to not use transient allocator and make them committed resources in d3d12

diff --git a/Engine/Source/Runtime/D3D12RHI/Private/D3D12PoolAllocator.cpp b/Engine/Source/Runtime/D3D12RHI/Private/D3D12PoolAllocator.cpp
--- a/Engine/Source/Runtime/D3D12RHI/Private/D3D12PoolAllocator.cpp
+++ b/Engine/Source/Runtime/D3D12RHI/Private/D3D12PoolAllocator.cpp
@@ -7,6 +7,14 @@
 #define NEEDS_D3D12_INDIRECT_ARGUMENT_HEAP_WORKAROUND 0
 #endif
 
+// Fix for random GPU crashes on draw indirects on multiple IHVs. Force all indirect arg buffers as committed resources (see UE-115982)
+static int32 GD3D12AllowPoolAllocateIndirectArgBuffers = 0;
+static FAutoConsoleVariableRef CVarD3D12AllowPoolAllocateIndirectArgBuffers(
+	TEXT("d3d12.AllowPoolAllocateIndirectArgBuffers"),
+	GD3D12AllowPoolAllocateIndirectArgBuffers,
+	TEXT("Allow indirect args to be pool allocated (otherwise they will be committed resources) (default: 0)"),
+	ECVF_ReadOnly);
+
 //-----------------------------------------------------------------------------
 //	FD3D12MemoryPool
 //-----------------------------------------------------------------------------
@@ -230,7 +238,23 @@ void FD3D12PoolAllocator::AllocDefaultResource(D3D12_HEAP_TYPE InHeapType, const
 #endif // D3D12_RHI_RAYTRACING
 #endif  // DO_CHECK
 
-	AllocateResource(GetParentDevice()->GetGPUIndex(), InHeapType, InDesc, InDesc.Width, InAllocationAlignment, InResourceStateMode, InCreateState, nullptr, InName, ResourceLocation);
+	// Force indirect args to stand alone allocations instead of pooled
+	if (!GD3D12AllowPoolAllocateIndirectArgBuffers && EnumHasAnyFlags(InBufferUsage, BUF_DrawIndirect))
+	{
+		ResourceLocation.Clear();
+
+		FD3D12Resource* NewResource = nullptr;
+		const D3D12_HEAP_PROPERTIES HeapProps = CD3DX12_HEAP_PROPERTIES(InHeapType, GetGPUMask().GetNative(), GetVisibilityMask().GetNative());
+		D3D12_RESOURCE_DESC Desc = InDesc;
+		Desc.Alignment = 0;
+		VERIFYD3D12RESULT(GetParentDevice()->GetParentAdapter()->CreateCommittedResource(Desc, GetGPUMask(), HeapProps, InCreateState, InResourceStateMode, InCreateState, nullptr, &NewResource, InName, false));
+
+		ResourceLocation.AsStandAlone(NewResource, InDesc.Width);
+	}
+	else
+	{
+		AllocateResource(GetParentDevice()->GetGPUIndex(), InHeapType, InDesc, InDesc.Width, InAllocationAlignment, InResourceStateMode, InCreateState, nullptr, InName, ResourceLocation);
+	}
 }
 
 
diff --git a/Engine/Source/Runtime/RenderCore/Private/RenderGraphBuilder.cpp b/Engine/Source/Runtime/RenderCore/Private/RenderGraphBuilder.cpp
--- a/Engine/Source/Runtime/RenderCore/Private/RenderGraphBuilder.cpp
+++ b/Engine/Source/Runtime/RenderCore/Private/RenderGraphBuilder.cpp
@@ -409,6 +409,11 @@ bool FRDGBuilder::IsTransient(FRDGBufferRef Buffer) const
 		return false;
 	}
 
+	if (!GRDGTransientIndirectArgBuffers && EnumHasAnyFlags(Buffer->Desc.Usage, BUF_DrawIndirect))
+	{
+		return false;
+	}
+
 	return EnumHasAnyFlags(Buffer->Desc.Usage, BUF_UnorderedAccess);
 }
 
diff --git a/Engine/Source/Runtime/RenderCore/Private/RenderGraphPrivate.cpp b/Engine/Source/Runtime/RenderCore/Private/RenderGraphPrivate.cpp
--- a/Engine/Source/Runtime/RenderCore/Private/RenderGraphPrivate.cpp
+++ b/Engine/Source/Runtime/RenderCore/Private/RenderGraphPrivate.cpp
@@ -305,6 +305,13 @@ FAutoConsoleVariableRef CVarRDGUseTransientAllocator(
 	TEXT(" 1: enables the transient allocator (default)"),
 	ECVF_RenderThreadSafe);
 
+// Fix for random GPU crashes on draw indirects on multiple IHVs. Force all indirect arg buffers as non transient (see UE-115982)
+int32 GRDGTransientIndirectArgBuffers = 0;
+FAutoConsoleVariableRef CVarRDGIndirectArgBufferTransientAllocated(
+	TEXT("r.RDG.TransientAllocator.IndirectArgumentBuffers"), GRDGTransientIndirectArgBuffers,
+	TEXT("Whether indirect argument buffers should use transient resource allocator. Default: 0"),
+	ECVF_RenderThreadSafe);
+
 int32 GRDGDrain = 0;
 FAutoConsoleVariableRef CVarRDGDrain(
 	TEXT("r.RDG.Drain"), GRDGDrain,
diff --git a/Engine/Source/Runtime/RenderCore/Private/RenderGraphPrivate.h b/Engine/Source/Runtime/RenderCore/Private/RenderGraphPrivate.h
--- a/Engine/Source/Runtime/RenderCore/Private/RenderGraphPrivate.h
+++ b/Engine/Source/Runtime/RenderCore/Private/RenderGraphPrivate.h
@@ -107,6 +107,7 @@ extern int32 GRDGAsyncCompute;
 extern int32 GRDGCullPasses;
 extern int32 GRDGMergeRenderPasses;
 extern int32 GRDGTransientAllocator;
+extern int32 GRDGTransientIndirectArgBuffers;
 extern int32 GRDGDrain;
 
 #if CSV_PROFILER
