Back out Null SRV workaround, it causes Echo and other geometry to randomly flicker

diff --git a/Engine/Source/Runtime/D3D12RHI/Private/D3D12CommandList.h b/Engine/Source/Runtime/D3D12RHI/Private/D3D12CommandList.h
--- a/Engine/Source/Runtime/D3D12RHI/Private/D3D12CommandList.h
+++ b/Engine/Source/Runtime/D3D12RHI/Private/D3D12CommandList.h
@@ -85,8 +85,6 @@ private:
 					if (Generation < GenerationSyncPoint.Key)
 					{
 						// The requested generation is older than the oldest tracked generation, so it must be complete.
-						// (LastCompleteGeneration probably updated while trying to take the lock)
-						check(Generation <= LastCompleteGeneration);
 						return true;
 					}
 					else
@@ -104,11 +102,6 @@ private:
 						}
 					}
 				}
-				else
-				{
-					// LastCompleteGeneration probably updated while trying to take the lock, validate that the GPU is actually processed this generation
-					check(Generation <= LastCompleteGeneration);
-				}
 			}
 
 			return true;
@@ -127,6 +120,7 @@ private:
 					while (ActiveGenerations.Peek(GenerationSyncPoint) && (Generation > LastCompleteGeneration))
 					{
 						check(Generation >= GenerationSyncPoint.Key);
+						ActiveGenerations.Dequeue(GenerationSyncPoint);
 
 						// Unblock other threads while we wait for the command list to complete
 						ActiveGenerationsCS.Unlock();
@@ -134,7 +128,6 @@ private:
 						GenerationSyncPoint.Value.WaitForCompletion();
 
 						ActiveGenerationsCS.Lock();
-						ActiveGenerations.Dequeue(GenerationSyncPoint);
 						LastCompleteGeneration = FMath::Max(LastCompleteGeneration, GenerationSyncPoint.Key);
 					}
 				}
diff --git a/Engine/Source/Runtime/D3D12RHI/Private/D3D12DescriptorCache.cpp b/Engine/Source/Runtime/D3D12RHI/Private/D3D12DescriptorCache.cpp
--- a/Engine/Source/Runtime/D3D12RHI/Private/D3D12DescriptorCache.cpp
+++ b/Engine/Source/Runtime/D3D12RHI/Private/D3D12DescriptorCache.cpp
@@ -13,18 +13,6 @@ static FAutoConsoleVariableRef CVarGlobalViewHeapBlockSize(
 	ECVF_ReadOnly
 );
 
-#if PLATFORM_WINDOWS
-
-int32 GGlobalViewHeapClobberAfterUse = 0;
-static FAutoConsoleVariableRef CVarGlobalViewHeapClobberSRVAfterUse(
-	TEXT("D3D12.GlobalViewHeapClobberSRVAfterUse"),
-	GGlobalViewHeapClobberAfterUse,
-	TEXT("Workaround for current outstanding suspected NVIDIA driver bug. JIRA UE-108339"),
-	ECVF_ReadOnly
-);
-
-#endif // PLATFORM_WINDOWS
-
 // Define template functions that are only declared in the header.
 #if USE_STATIC_ROOT_SIGNATURE
 template void FD3D12DescriptorCache::SetConstantBuffers<SF_Vertex>(const FD3D12RootSignature* RootSignature, FD3D12ConstantBufferCache& Cache, const CBVSlotMask& SlotsNeededMask, uint32 Count, uint32& HeapSlot);
@@ -1043,17 +1031,6 @@ void FD3D12GlobalHeap::Init(D3D12_DESCRIPTOR_HEAP_TYPE InType, uint32 InTotalSiz
 		FreeBlocks.Enqueue(new FD3D12GlobalHeapBlock(CurrentBaseSlot, ActualBlockSize));
 		CurrentBaseSlot += ActualBlockSize;
 	}
-
-	// Create a default SRV view
-	D3D12_SHADER_RESOURCE_VIEW_DESC SRVDesc = {};
-	SRVDesc.Shader4ComponentMapping = D3D12_DEFAULT_SHADER_4_COMPONENT_MAPPING;
-	SRVDesc.ViewDimension = D3D12_SRV_DIMENSION_TEXTURE2D;
-	SRVDesc.Format = DXGI_FORMAT_B8G8R8A8_UNORM;
-	SRVDesc.Texture2D.MipLevels = 1;
-	SRVDesc.Texture2D.MostDetailedMip = 0;
-	SRVDesc.Texture2D.ResourceMinLODClamp = 0.0f;
-	NullSRVDescriptor = new FD3D12DescriptorHandleSRV(GetParentDevice());
-	NullSRVDescriptor->CreateView(SRVDesc, nullptr);
 }
 
 
@@ -1117,28 +1094,6 @@ void FD3D12GlobalHeap::UpdateFreeBlocks()
 			DEC_DWORD_STAT_BY(STAT_GlobalViewHeapWastedDescriptors, ReleasedBlock->Size - ReleasedBlock->SizeUsed);
 			INC_DWORD_STAT_BY(STAT_GlobalViewHeapFreeDescriptors, ReleasedBlock->Size);
 
-#if PLATFORM_WINDOWS
-			// Clear the global descriptor heap block with null srvs to make sure they are not referencing some possibly deleted resources/heaps anymore (we know the GPU should be done accessing them so copy is valid).
-			// This fixes a rare page fault GPU crash on nVidia GPUs.
-			// Clearing the descriptors on reuse (during AllocateBlock) still crashes so current assumption is some random/invalid access done by the driver.
-			// See: JIRA UE-108339
-			if (GGlobalViewHeapClobberAfterUse)
-			{
-				ID3D12Device* Device = GetParentDevice()->GetDevice();
-				D3D12_CPU_DESCRIPTOR_HANDLE TargetCPUBase = GetCPUSlotHandle(ReleasedBlock);
-				D3D12_CPU_DESCRIPTOR_HANDLE SourceNullSRVHandle = NullSRVDescriptor->GetHandle();
-				for (uint32 SlotIndex = 0; SlotIndex < ReleasedBlock->Size; SlotIndex++)
-				{
-					D3D12_CPU_DESCRIPTOR_HANDLE TargetCPU = { TargetCPUBase.ptr + SlotIndex * DescriptorSize };
-					Device->CopyDescriptorsSimple(
-						1,
-						TargetCPU,
-						SourceNullSRVHandle,
-						D3D12_DESCRIPTOR_HEAP_TYPE_CBV_SRV_UAV);
-				}
-			}
-#endif  // PLATFORM_WINDOWS
-
 			ReleasedBlock->SizeUsed = 0;
 			FreeBlocks.Enqueue(ReleasedBlock);
 
diff --git a/Engine/Source/Runtime/D3D12RHI/Private/D3D12DescriptorCache.h b/Engine/Source/Runtime/D3D12RHI/Private/D3D12DescriptorCache.h
--- a/Engine/Source/Runtime/D3D12RHI/Private/D3D12DescriptorCache.h
+++ b/Engine/Source/Runtime/D3D12RHI/Private/D3D12DescriptorCache.h
@@ -456,8 +456,6 @@ private:
 	TQueue<FD3D12GlobalHeapBlock*> FreeBlocks;
 	TArray<FD3D12GlobalHeapBlock*> ReleasedBlocks;
 
-	FD3D12DescriptorHandleSRV* NullSRVDescriptor;
-
 	FCriticalSection CriticalSection;
 };
 
