D3D12 race condition fix for ActiveGenerations having dequeue called concurrently in WaitForCompletion and CleanupActiveGenerations, resulting in an incorrect double pop, causing an unrelated sync point to be dropped from processing (fix for random meshes flickering, fix for random GPU crashes/hangs, possible fix for numerous texture corruption cases) - group collaboration with Yuriy/Rune/Kenzo.

diff --git a/Engine/Source/Runtime/D3D12RHI/Private/D3D12CommandList.h b/Engine/Source/Runtime/D3D12RHI/Private/D3D12CommandList.h
--- a/Engine/Source/Runtime/D3D12RHI/Private/D3D12CommandList.h
+++ b/Engine/Source/Runtime/D3D12RHI/Private/D3D12CommandList.h
@@ -122,13 +122,10 @@ private:
 						check(Generation >= GenerationSyncPoint.Key);
 						ActiveGenerations.Dequeue(GenerationSyncPoint);
 
-						// Unblock other threads while we wait for the command list to complete
-						ActiveGenerationsCS.Unlock();
-
 						GenerationSyncPoint.Value.WaitForCompletion();
 
-						ActiveGenerationsCS.Lock();
-						LastCompleteGeneration = FMath::Max(LastCompleteGeneration, GenerationSyncPoint.Key);
+						check(GenerationSyncPoint.Key > LastCompleteGeneration);
+						LastCompleteGeneration = GenerationSyncPoint.Key;
 					}
 				}
 			}
