Allow -d3ddebug on all D3D12 platforms.

#rnx

diff --git a/Engine/Source/Runtime/D3D12RHI/Private/D3D12Adapter.cpp b/Engine/Source/Runtime/D3D12RHI/Private/D3D12Adapter.cpp
--- a/Engine/Source/Runtime/D3D12RHI/Private/D3D12Adapter.cpp
+++ b/Engine/Source/Runtime/D3D12RHI/Private/D3D12Adapter.cpp
@@ -262,6 +262,40 @@ void FD3D12Adapter::CreateRootDevice(bool bWithDebug)
 {
 	const bool bAllowVendorDevice = !FParse::Param(FCommandLine::Get(), TEXT("novendordevice"));
 
+	// -d3ddebug is always allowed on Windows, but only allowed in non-shipping builds on other platforms.
+	// -gpuvalidation is only supported on Windows.
+#if PLATFORM_WINDOWS || !UE_BUILD_SHIPPING
+	bool bWithGPUValidation = PLATFORM_WINDOWS && (FParse::Param(FCommandLine::Get(), TEXT("d3d12gpuvalidation")) || FParse::Param(FCommandLine::Get(), TEXT("gpuvalidation")));
+	// If GPU validation is requested, automatically enable the debug layer.
+	bWithDebug |= bWithGPUValidation;
+	if (bWithDebug)
+	{
+		TRefCountPtr<ID3D12Debug> DebugController;
+		if (SUCCEEDED(D3D12GetDebugInterface(IID_PPV_ARGS(DebugController.GetInitReference()))))
+		{
+			DebugController->EnableDebugLayer();
+			bDebugDevice = true;
+
+#if PLATFORM_WINDOWS
+			if (bWithGPUValidation)
+			{
+				TRefCountPtr<ID3D12Debug1> DebugController1;
+				VERIFYD3D12RESULT(DebugController->QueryInterface(IID_PPV_ARGS(DebugController1.GetInitReference())));
+				DebugController1->SetEnableGPUBasedValidation(true);
+				SetEmitDrawEvents(true);
+			}
+#endif
+		}
+		else
+		{
+			UE_LOG(LogD3D12RHI, Fatal, TEXT("The debug interface requires the D3D12 SDK Layers. Please install the Graphics Tools for Windows. See: https://docs.microsoft.com/en-us/windows/uwp/gaming/use-the-directx-runtime-and-visual-studio-graphics-diagnostic-features"));
+		}
+	}
+
+	FGenericCrashContext::SetEngineData(TEXT("RHI.D3DDebug"), bWithDebug ? TEXT("true") : TEXT("false"));
+	UE_LOG(LogD3D12RHI, Log, TEXT("InitD3DDevice: -D3DDebug = %s -D3D12GPUValidation = %s"), bWithDebug ? TEXT("on") : TEXT("off"), bWithGPUValidation ? TEXT("on") : TEXT("off"));
+#endif
+
 #if PLATFORM_WINDOWS || (PLATFORM_HOLOLENS && !UE_BUILD_SHIPPING && WITH_PIX_EVENT_RUNTIME)
 	
 	// Multiple ways to enable the different D3D12 crash debugging modes:
@@ -357,34 +391,6 @@ void FD3D12Adapter::CreateRootDevice(bool bWithDebug)
 	}
 #endif
 
-	bool bD3d12gpuvalidation = false;
-	if (bWithDebug)
-	{
-		TRefCountPtr<ID3D12Debug> DebugController;
-		if (SUCCEEDED(D3D12GetDebugInterface(IID_PPV_ARGS(DebugController.GetInitReference()))))
-		{
-			DebugController->EnableDebugLayer();
-			bDebugDevice = true;
-
-			if (FParse::Param(FCommandLine::Get(), TEXT("d3d12gpuvalidation")) || FParse::Param(FCommandLine::Get(), TEXT("gpuvalidation")))
-			{
-				TRefCountPtr<ID3D12Debug1> DebugController1;
-				VERIFYD3D12RESULT(DebugController->QueryInterface(IID_PPV_ARGS(DebugController1.GetInitReference())));
-				DebugController1->SetEnableGPUBasedValidation(true);
-
-				SetEmitDrawEvents(true);
-				bD3d12gpuvalidation = true;
-			}
-		}
-		else
-		{
-			bWithDebug = false;
-			UE_LOG(LogD3D12RHI, Fatal, TEXT("The debug interface requires the D3D12 SDK Layers. Please install the Graphics Tools for Windows. See: https://docs.microsoft.com/en-us/windows/uwp/gaming/use-the-directx-runtime-and-visual-studio-graphics-diagnostic-features"));
-		}
-	}
-
-	FGenericCrashContext::SetEngineData(TEXT("RHI.D3DDebug"), bWithDebug ? TEXT("true") : TEXT("false"));
-
 	// Setup DRED if requested
 	bool bDRED = false;
 	bool bDREDContext = false;
@@ -423,8 +429,6 @@ void FD3D12Adapter::CreateRootDevice(bool bWithDebug)
 	FGenericCrashContext::SetEngineData(TEXT("RHI.DRED"), bDRED ? TEXT("true") : TEXT("false"));
 	FGenericCrashContext::SetEngineData(TEXT("RHI.DREDContext"), bDREDContext ? TEXT("true") : TEXT("false"));
 
-	UE_LOG(LogD3D12RHI, Log, TEXT("InitD3DDevice: -D3DDebug = %s -D3D12GPUValidation = %s"), bWithDebug ? TEXT("on") : TEXT("off"), bD3d12gpuvalidation ? TEXT("on") : TEXT("off"));
-
 #endif // PLATFORM_WINDOWS || (PLATFORM_HOLOLENS && !UE_BUILD_SHIPPING && WITH_PIX_EVENT_RUNTIME)
 
 #if USE_PIX
