Adding GRHISupportsUAVFormatAliasing support to D3D12RHI.
When creating a BCn texture with UAV support, create a secondary ID3D12Resource that aliases the same memory and can be used as its UAV. This alias resource has to be considered during transitions to/from the UAV state where we insert an aliasing barrier and possibly another state transition for the 'real' resource.
Created FD3D12ResourceDesc which extends D3D12_RESOURCE_DESC. This type should allow us to pass around future UE-only creation info without modifying all the allocation function signatures. This is only implemented enough to cover Textures right now, can be extended in the future.

diff --git a/Engine/Source/Runtime/D3D12RHI/Private/D3D12Adapter.h b/Engine/Source/Runtime/D3D12RHI/Private/D3D12Adapter.h
--- a/Engine/Source/Runtime/D3D12RHI/Private/D3D12Adapter.h
+++ b/Engine/Source/Runtime/D3D12RHI/Private/D3D12Adapter.h
@@ -234,7 +234,7 @@ public:
 	void EndFrame();
 
 	// Resource Creation
-	HRESULT CreateCommittedResource(const D3D12_RESOURCE_DESC& InDesc,
+	HRESULT CreateCommittedResource(const FD3D12ResourceDesc& InDesc,
 		FRHIGPUMask CreationNode,
 		const D3D12_HEAP_PROPERTIES& HeapProps,
 		D3D12_RESOURCE_STATES InInitialState,
@@ -246,7 +246,7 @@ public:
 		return CreateCommittedResource(InDesc, CreationNode, HeapProps, InInitialState, ED3D12ResourceStateMode::Default, D3D12_RESOURCE_STATE_TBD, ClearValue, ppOutResource, Name, bVerifyHResult);
 	}
 
-	HRESULT CreateCommittedResource(const D3D12_RESOURCE_DESC& Desc,
+	HRESULT CreateCommittedResource(const FD3D12ResourceDesc& Desc,
 		FRHIGPUMask CreationNode,
 		const D3D12_HEAP_PROPERTIES& HeapProps,
 		D3D12_RESOURCE_STATES InInitialState,
@@ -257,7 +257,7 @@ public:
 		const TCHAR* Name,
 		bool bVerifyHResult = true);
 
-	HRESULT CreatePlacedResource(const D3D12_RESOURCE_DESC& InDesc,
+	HRESULT CreatePlacedResource(const FD3D12ResourceDesc& InDesc,
 		FD3D12Heap* BackingHeap,
 		uint64 HeapOffset,
 		D3D12_RESOURCE_STATES InInitialState,
@@ -269,7 +269,7 @@ public:
 		return CreatePlacedResource(InDesc, BackingHeap, HeapOffset, InInitialState, ED3D12ResourceStateMode::Default, D3D12_RESOURCE_STATE_TBD, ClearValue, ppOutResource, Name, bVerifyHResult);
 	}
 
-	HRESULT CreatePlacedResource(const D3D12_RESOURCE_DESC& Desc,
+	HRESULT CreatePlacedResource(const FD3D12ResourceDesc& Desc,
 		FD3D12Heap* BackingHeap,
 		uint64 HeapOffset,
 		D3D12_RESOURCE_STATES InInitialState,
diff --git a/Engine/Source/Runtime/D3D12RHI/Private/D3D12Allocation.cpp b/Engine/Source/Runtime/D3D12RHI/Private/D3D12Allocation.cpp
--- a/Engine/Source/Runtime/D3D12RHI/Private/D3D12Allocation.cpp
+++ b/Engine/Source/Runtime/D3D12RHI/Private/D3D12Allocation.cpp
@@ -1595,9 +1595,9 @@ void FD3D12TextureAllocatorPool::Destroy()
 }
 
 HRESULT FD3D12TextureAllocatorPool::AllocateTexture(
-	D3D12_RESOURCE_DESC Desc,
+	FD3D12ResourceDesc Desc,
 	const D3D12_CLEAR_VALUE* ClearValue,
-	uint8 UEFormat,
+	EPixelFormat UEFormat,
 	FD3D12ResourceLocation& TextureLocation,
 	const D3D12_RESOURCE_STATES InitialState,
 	const TCHAR* Name)
@@ -1607,11 +1607,9 @@ HRESULT FD3D12TextureAllocatorPool::AllocateTexture(
 	Desc.Alignment = b4KAligment ?	D3D12_SMALL_RESOURCE_PLACEMENT_ALIGNMENT : (Desc.SampleDesc.Count > 1 ? D3D12_DEFAULT_MSAA_RESOURCE_PLACEMENT_ALIGNMENT : D3D12_DEFAULT_RESOURCE_PLACEMENT_ALIGNMENT);
 	const D3D12_RESOURCE_ALLOCATION_INFO Info = GetParentDevice()->GetDevice()->GetResourceAllocationInfo(0, 1, &Desc);
 
-	bool bIsReadOnly = !(Desc.Flags & D3D12_RESOURCE_FLAG_ALLOW_RENDER_TARGET ||
-						 Desc.Flags & D3D12_RESOURCE_FLAG_ALLOW_DEPTH_STENCIL ||
-						 Desc.Flags & D3D12_RESOURCE_FLAG_ALLOW_UNORDERED_ACCESS);
-	bool bIsRenderTarget = (Desc.Flags & D3D12_RESOURCE_FLAG_ALLOW_RENDER_TARGET ||
-						 Desc.Flags & D3D12_RESOURCE_FLAG_ALLOW_DEPTH_STENCIL);
+	const bool bIsRenderTarget = EnumHasAnyFlags(Desc.Flags, D3D12_RESOURCE_FLAG_ALLOW_RENDER_TARGET | D3D12_RESOURCE_FLAG_ALLOW_DEPTH_STENCIL);
+	const bool bIsReadOnly = !bIsRenderTarget && !EnumHasAnyFlags(Desc.Flags, D3D12_RESOURCE_FLAG_ALLOW_UNORDERED_ACCESS) && !Desc.NeedsUAVAliasWorkarounds();
+
 	ED3D12ResourceStateMode ResourceStateMode = bIsReadOnly ? ED3D12ResourceStateMode::Default : ED3D12ResourceStateMode::MultiState;
 	EPoolType PoolType = bIsReadOnly ? (b4KAligment ? EPoolType::ReadOnly4K : EPoolType::ReadOnly) : (bIsRenderTarget ? EPoolType::RenderTarget : EPoolType::UAV);
 	PoolAllocators[(int)PoolType]->AllocateResource(GetParentDevice()->GetGPUIndex(), D3D12_HEAP_TYPE_DEFAULT, Desc, Info.SizeInBytes, Info.Alignment, ResourceStateMode, InitialState, ClearValue, Name, TextureLocation);
@@ -1694,9 +1692,9 @@ FD3D12TextureAllocatorPool::FD3D12TextureAllocatorPool(FD3D12Device* Device, FRH
 }
 
 HRESULT FD3D12TextureAllocatorPool::AllocateTexture(
-	D3D12_RESOURCE_DESC Desc,
+	FD3D12ResourceDesc Desc,
 	const D3D12_CLEAR_VALUE* ClearValue,
-	uint8 UEFormat,
+	EPixelFormat UEFormat,
 	FD3D12ResourceLocation& TextureLocation,
 	const D3D12_RESOURCE_STATES InitialState,
 	const TCHAR* Name)
@@ -1708,9 +1706,8 @@ HRESULT FD3D12TextureAllocatorPool::AllocateTexture(
 
 	TextureLocation.Clear();
 
-	if (!(Desc.Flags & D3D12_RESOURCE_FLAG_ALLOW_RENDER_TARGET ||
-		Desc.Flags & D3D12_RESOURCE_FLAG_ALLOW_DEPTH_STENCIL ||
-		Desc.Flags & D3D12_RESOURCE_FLAG_ALLOW_UNORDERED_ACCESS) &&
+	if (!EnumHasAnyFlags(Desc.Flags, D3D12_RESOURCE_FLAG_ALLOW_RENDER_TARGET|D3D12_RESOURCE_FLAG_ALLOW_DEPTH_STENCIL|D3D12_RESOURCE_FLAG_ALLOW_UNORDERED_ACCESS) &&
+		!Desc.NeedsUAVAliasWorkarounds() &&
 		Desc.SampleDesc.Count == 1)
 	{
 		// The top mip level must be less than 64 KB to use 4 KB alignment
@@ -1776,7 +1773,7 @@ FD3D12TextureAllocator::~FD3D12TextureAllocator()
 {
 }
 
-HRESULT FD3D12TextureAllocator::AllocateTexture(D3D12_RESOURCE_DESC Desc, const D3D12_CLEAR_VALUE* ClearValue, FD3D12ResourceLocation& TextureLocation, const D3D12_RESOURCE_STATES InitialState, const TCHAR* Name)
+HRESULT FD3D12TextureAllocator::AllocateTexture(FD3D12ResourceDesc Desc, const D3D12_CLEAR_VALUE* ClearValue, FD3D12ResourceLocation& TextureLocation, const D3D12_RESOURCE_STATES InitialState, const TCHAR* Name)
 {
 	FD3D12Device* Device = GetParentDevice();
 	FD3D12Adapter* Adapter = Device->GetParentAdapter();
@@ -1821,12 +1818,11 @@ FD3D12TextureAllocatorPool::FD3D12TextureAllocatorPool(FD3D12Device* Device, FRH
 	ReadOnlyTexturePool(Device, VisibilityNode, FString(L"Small Read-Only Texture allocator"), TEXTURE_POOL_SIZE, D3D12_HEAP_FLAG_ALLOW_ONLY_NON_RT_DS_TEXTURES)
 {};
 
-HRESULT FD3D12TextureAllocatorPool::AllocateTexture(D3D12_RESOURCE_DESC Desc, const D3D12_CLEAR_VALUE* ClearValue, uint8 UEFormat, FD3D12ResourceLocation& TextureLocation, const D3D12_RESOURCE_STATES InitialState, const TCHAR* Name)
+HRESULT FD3D12TextureAllocatorPool::AllocateTexture(FD3D12ResourceDesc Desc, const D3D12_CLEAR_VALUE* ClearValue, EPixelFormat UEFormat, FD3D12ResourceLocation& TextureLocation, const D3D12_RESOURCE_STATES InitialState, const TCHAR* Name)
 {
 	// 4KB alignment is only available for read only textures
-	if ((Desc.Flags & D3D12_RESOURCE_FLAG_ALLOW_RENDER_TARGET ||
-		Desc.Flags & D3D12_RESOURCE_FLAG_ALLOW_DEPTH_STENCIL ||
-		Desc.Flags & D3D12_RESOURCE_FLAG_ALLOW_UNORDERED_ACCESS) == false &&
+	if (!EnumHasAnyFlags(Desc.Flags, D3D12_RESOURCE_FLAG_ALLOW_RENDER_TARGET|D3D12_RESOURCE_FLAG_ALLOW_DEPTH_STENCIL|D3D12_RESOURCE_FLAG_ALLOW_UNORDERED_ACCESS) &&
+		!Desc.NeedsUAVAliasWorkarounds() &&
 		Desc.SampleDesc.Count == 1)// Multi-Sample texures have much larger alignment requirements (4MB vs 64KB)
 	{
 		// The top mip level must be less than 64k
@@ -1841,6 +1837,39 @@ HRESULT FD3D12TextureAllocatorPool::AllocateTexture(D3D12_RESOURCE_DESC Desc, co
 	FD3D12Resource* Resource = nullptr;
 
 	const D3D12_HEAP_PROPERTIES HeapProps = CD3DX12_HEAP_PROPERTIES(D3D12_HEAP_TYPE_DEFAULT, GetGPUMask().GetNative(), GetVisibilityMask().GetNative());
+
+	// UAV Aliasing needs a Heap to create the aliased resource in.
+	if (Desc.NeedsUAVAliasWorkarounds())
+	{
+		const D3D12_RESOURCE_ALLOCATION_INFO Info = GetParentDevice()->GetDevice()->GetResourceAllocationInfo(0, 1, &Desc);
+
+		D3D12_HEAP_DESC HeapDesc{};
+		HeapDesc.SizeInBytes = Info.SizeInBytes;
+		HeapDesc.Properties = HeapProps;
+		HeapDesc.Alignment = Desc.SampleDesc.Count > 1 ? D3D12_DEFAULT_MSAA_RESOURCE_PLACEMENT_ALIGNMENT : 0;
+		HeapDesc.Flags = D3D12_HEAP_FLAG_ALLOW_ALL_BUFFERS_AND_TEXTURES;
+		if (Adapter->IsHeapNotZeroedSupported())
+		{
+			HeapDesc.Flags |= FD3D12_HEAP_FLAG_CREATE_NOT_ZEROED;
+		}
+
+		ID3D12Heap* Heap = nullptr;
+		VERIFYD3D12RESULT(Adapter->GetD3DDevice()->CreateHeap(&HeapDesc, IID_PPV_ARGS(&Heap)));
+		TRefCountPtr<FD3D12Heap> BackingHeap = new FD3D12Heap(GetParentDevice(), GetVisibilityMask());
+		bool bTrack = false;
+		BackingHeap->SetHeap(Heap, Name, bTrack);
+
+		HRESULT hr = Adapter->CreatePlacedResource(Desc, BackingHeap, 0, InitialState, ED3D12ResourceStateMode::MultiState, InitialState, ClearValue, &Resource, Name);
+
+		if (SUCCEEDED(hr))
+		{
+			TextureLocation.SetType(FD3D12ResourceLocation::ResourceLocationType::eStandAlone);
+			TextureLocation.SetResource(Resource);
+		}
+
+		return hr;
+	}
+
 	HRESULT hr = Adapter->CreateCommittedResource(Desc, GetGPUMask(), HeapProps, InitialState, ClearValue, &Resource, Name, false);
 
 	if (SUCCEEDED(hr))
diff --git a/Engine/Source/Runtime/D3D12RHI/Private/D3D12Allocation.h b/Engine/Source/Runtime/D3D12RHI/Private/D3D12Allocation.h
--- a/Engine/Source/Runtime/D3D12RHI/Private/D3D12Allocation.h
+++ b/Engine/Source/Runtime/D3D12RHI/Private/D3D12Allocation.h
@@ -902,7 +902,7 @@ class FD3D12TextureAllocatorPool : public FD3D12DeviceChild, public FD3D12MultiN
 public:
 	FD3D12TextureAllocatorPool(FD3D12Device* Device, FRHIGPUMask VisibilityNode);
 
-	HRESULT AllocateTexture(D3D12_RESOURCE_DESC Desc, const D3D12_CLEAR_VALUE* ClearValue, uint8 UEFormat, FD3D12ResourceLocation& TextureLocation, const D3D12_RESOURCE_STATES InitialState, const TCHAR* Name);
+	HRESULT AllocateTexture(FD3D12ResourceDesc Desc, const D3D12_CLEAR_VALUE* ClearValue, EPixelFormat UEFormat, FD3D12ResourceLocation& TextureLocation, const D3D12_RESOURCE_STATES InitialState, const TCHAR* Name);
 	
 	void BeginFrame();
 	void CleanUpAllocations();
@@ -928,7 +928,7 @@ class FD3D12TextureAllocatorPool : public FD3D12DeviceChild, public FD3D12MultiN
 public:
 	FD3D12TextureAllocatorPool(FD3D12Device* Device, FRHIGPUMask VisibilityNode);
 
-	HRESULT AllocateTexture(D3D12_RESOURCE_DESC Desc, const D3D12_CLEAR_VALUE* ClearValue, uint8 UEFormat, FD3D12ResourceLocation& TextureLocation, const D3D12_RESOURCE_STATES InitialState, const TCHAR* Name);
+	HRESULT AllocateTexture(FD3D12ResourceDesc Desc, const D3D12_CLEAR_VALUE* ClearValue, EPixelFormat UEFormat, FD3D12ResourceLocation& TextureLocation, const D3D12_RESOURCE_STATES InitialState, const TCHAR* Name);
 
 	void BeginFrame() {}
 	void CleanUpAllocations()
@@ -957,7 +957,7 @@ public:
 
 	~FD3D12TextureAllocator();
 
-	HRESULT AllocateTexture(D3D12_RESOURCE_DESC Desc, const D3D12_CLEAR_VALUE* ClearValue, FD3D12ResourceLocation& TextureLocation, const D3D12_RESOURCE_STATES InitialState, const TCHAR* Name);
+	HRESULT AllocateTexture(FD3D12ResourceDesc Desc, const D3D12_CLEAR_VALUE* ClearValue, FD3D12ResourceLocation& TextureLocation, const D3D12_RESOURCE_STATES InitialState, const TCHAR* Name);
 };
 
 class FD3D12TextureAllocatorPool : public FD3D12DeviceChild, public FD3D12MultiNodeGPUObject
@@ -965,7 +965,7 @@ class FD3D12TextureAllocatorPool : public FD3D12DeviceChild, public FD3D12MultiN
 public:
 	FD3D12TextureAllocatorPool(FD3D12Device* Device, FRHIGPUMask VisibilityNode);
 
-	HRESULT AllocateTexture(D3D12_RESOURCE_DESC Desc, const D3D12_CLEAR_VALUE* ClearValue, uint8 UEFormat, FD3D12ResourceLocation& TextureLocation, const D3D12_RESOURCE_STATES InitialState, const TCHAR* Name);
+	HRESULT AllocateTexture(FD3D12ResourceDesc Desc, const D3D12_CLEAR_VALUE* ClearValue, EPixelFormat UEFormat, FD3D12ResourceLocation& TextureLocation, const D3D12_RESOURCE_STATES InitialState, const TCHAR* Name);
 	
 	void BeginFrame() {}
 	void CleanUpAllocations() { ReadOnlyTexturePool.CleanUpAllocations(0); }
diff --git a/Engine/Source/Runtime/D3D12RHI/Private/D3D12DirectCommandListManager.cpp b/Engine/Source/Runtime/D3D12RHI/Private/D3D12DirectCommandListManager.cpp
--- a/Engine/Source/Runtime/D3D12RHI/Private/D3D12DirectCommandListManager.cpp
+++ b/Engine/Source/Runtime/D3D12RHI/Private/D3D12DirectCommandListManager.cpp
@@ -955,10 +955,6 @@ uint32 FD3D12CommandListManager::CollectInitialResourceBarriersAndUpdateFinalSta
 		// Reserve space for the descs
 		BarrierDescInfo.BarrierDescs.Reserve(NumPendingResourceBarriers);
 
-		// Fill out the descs
-		D3D12_RESOURCE_BARRIER Desc = {};
-		Desc.Type = D3D12_RESOURCE_BARRIER_TYPE_TRANSITION;
-
 		for (uint32 i = 0; i < NumPendingResourceBarriers; ++i)
 		{
 			const FD3D12PendingResourceBarrier& PRB = PendingResourceBarriers[i];
@@ -968,8 +964,7 @@ uint32 FD3D12CommandListManager::CollectInitialResourceBarriersAndUpdateFinalSta
 
 			CResourceState& ResourceState = PRB.Resource->GetResourceState();
 
-			Desc.Transition.Subresource = PRB.SubResource;
-			const D3D12_RESOURCE_STATES Before = ResourceState.GetSubresourceState(Desc.Transition.Subresource);
+			const D3D12_RESOURCE_STATES Before = ResourceState.GetSubresourceState(PRB.SubResource);
 
 			// If state unknown then we don't enqueue a transition - only want to update the end state on the resource
 			// and not really enqueue a transition
@@ -978,33 +973,28 @@ uint32 FD3D12CommandListManager::CollectInitialResourceBarriersAndUpdateFinalSta
 			check(Before != D3D12_RESOURCE_STATE_TBD && Before != D3D12_RESOURCE_STATE_CORRUPT);
 			if (Before != After)
 			{
-				Desc.Transition.pResource = PRB.Resource->GetResource();
-				Desc.Transition.StateBefore = Before;
-				Desc.Transition.StateAfter = After;
-
 				if (IsDirectQueueExclusiveD3D12State(Before) || IsDirectQueueExclusiveD3D12State(After))
 				{
 					BarrierDescInfo.bHasGraphicStates = true;
 				}
 
-				// Add the desc
-				if (PRB.Resource->IsBackBuffer() && (After & BackBufferBarrierWriteTransitionTargets))
+				if (PRB.Resource->IsBackBuffer() && EnumHasAnyFlags(After, BackBufferBarrierWriteTransitionTargets))
 				{
-					BarrierDescInfo.BackBufferBarrierDescs.Add(Desc);
+					AddTransitionBarrier(BarrierDescInfo.BackBufferBarrierDescs, PRB.Resource, Before, After, PRB.SubResource);
 				}
 				else
 				{
-					BarrierDescInfo.BarrierDescs.Add(Desc);
+					AddTransitionBarrierWithUAVAccessOverrides(BarrierDescInfo.BarrierDescs, PRB.Resource, Before, After, PRB.SubResource);
 				}
 			}
 
 			// Update the state to the what it will be after hList executes
-			const D3D12_RESOURCE_STATES CommandListState = InCommandListHandle.GetResourceState(PRB.Resource).GetSubresourceState(Desc.Transition.Subresource);
+			const D3D12_RESOURCE_STATES CommandListState = InCommandListHandle.GetResourceState(PRB.Resource).GetSubresourceState(PRB.SubResource);
 			const D3D12_RESOURCE_STATES LastState = (CommandListState != D3D12_RESOURCE_STATE_TBD) ? CommandListState : After;
 
 			if (Before != LastState)
 			{
-				ResourceState.SetSubresourceState(Desc.Transition.Subresource, LastState);
+				ResourceState.SetSubresourceState(PRB.SubResource, LastState);
 			}
 
 #if ENABLE_RESIDENCY_MANAGEMENT
diff --git a/Engine/Source/Runtime/D3D12RHI/Private/D3D12PoolAllocator.cpp b/Engine/Source/Runtime/D3D12RHI/Private/D3D12PoolAllocator.cpp
--- a/Engine/Source/Runtime/D3D12RHI/Private/D3D12PoolAllocator.cpp
+++ b/Engine/Source/Runtime/D3D12RHI/Private/D3D12PoolAllocator.cpp
@@ -211,7 +211,7 @@ bool FD3D12PoolAllocator::SupportsAllocation(D3D12_HEAP_TYPE InHeapType, D3D12_R
 }
 
 
-void FD3D12PoolAllocator::AllocDefaultResource(D3D12_HEAP_TYPE InHeapType, const D3D12_RESOURCE_DESC& InDesc, EBufferUsageFlags InBufferUsage, ED3D12ResourceStateMode InResourceStateMode,
+void FD3D12PoolAllocator::AllocDefaultResource(D3D12_HEAP_TYPE InHeapType, const FD3D12ResourceDesc& InDesc, EBufferUsageFlags InBufferUsage, ED3D12ResourceStateMode InResourceStateMode,
 	D3D12_RESOURCE_STATES InCreateState, uint32 InAllocationAlignment, const TCHAR* InName, FD3D12ResourceLocation& ResourceLocation)
 {
 #if DO_CHECK
@@ -258,7 +258,7 @@ void FD3D12PoolAllocator::AllocDefaultResource(D3D12_HEAP_TYPE InHeapType, const
 }
 
 
-void FD3D12PoolAllocator::AllocateResource(uint32 GPUIndex, D3D12_HEAP_TYPE InHeapType, const D3D12_RESOURCE_DESC& InDesc, uint64 InSize, uint32 InAllocationAlignment, ED3D12ResourceStateMode InResourceStateMode,
+void FD3D12PoolAllocator::AllocateResource(uint32 GPUIndex, D3D12_HEAP_TYPE InHeapType, const FD3D12ResourceDesc& InDesc, uint64 InSize, uint32 InAllocationAlignment, ED3D12ResourceStateMode InResourceStateMode,
 	D3D12_RESOURCE_STATES InCreateState, const D3D12_CLEAR_VALUE* InClearValue, const TCHAR* InName, FD3D12ResourceLocation& ResourceLocation)
 {
 	TRACE_CPUPROFILER_EVENT_SCOPE(D3D12RHI::AllocatePoolResource);
@@ -340,7 +340,7 @@ void FD3D12PoolAllocator::AllocateResource(uint32 GPUIndex, D3D12_HEAP_TYPE InHe
 		{
 			check(ResourceLocation.GetResource() == nullptr);
 			
-			D3D12_RESOURCE_DESC Desc = InDesc;
+			FD3D12ResourceDesc Desc = InDesc;
 			Desc.Alignment = AllocationAlignment;
 
 			FD3D12Resource* NewResource = CreatePlacedResource(AllocationData, Desc, InCreateState, InResourceStateMode, InClearValue, InName);
@@ -354,11 +354,12 @@ void FD3D12PoolAllocator::AllocateResource(uint32 GPUIndex, D3D12_HEAP_TYPE InHe
 	{
 		FD3D12Resource* NewResource = nullptr;
 		const D3D12_HEAP_PROPERTIES HeapProps = CD3DX12_HEAP_PROPERTIES(InHeapType, GetGPUMask().GetNative(), GetVisibilityMask().GetNative());
-		D3D12_RESOURCE_DESC Desc = InDesc;
+		FD3D12ResourceDesc Desc = InDesc;
 		Desc.Alignment = 0;
 
 		// If we are tracking all allocation data and allocating a standalone texture, then first create a heap so we can retrieve the GPU virtual address as well
-		if (Adapter->IsTrackingAllAllocations() && InDesc.Dimension != D3D12_RESOURCE_DIMENSION_BUFFER)
+		// UAV Aliasing needs a Heap to create the aliased resource in.
+		if (InDesc.NeedsUAVAliasWorkarounds() || (Adapter->IsTrackingAllAllocations() && InDesc.Dimension != D3D12_RESOURCE_DIMENSION_BUFFER))
 		{
 			D3D12_HEAP_DESC HeapDesc = {};
 			HeapDesc.SizeInBytes = InSize;
@@ -402,7 +403,7 @@ void FD3D12PoolAllocator::AllocateResource(uint32 GPUIndex, D3D12_HEAP_TYPE InHe
 
 FD3D12Resource* FD3D12PoolAllocator::CreatePlacedResource(
 	const FRHIPoolAllocationData& InAllocationData,
-	const D3D12_RESOURCE_DESC& InDesc, 
+	const FD3D12ResourceDesc& InDesc,
 	D3D12_RESOURCE_STATES InCreateState, 
 	ED3D12ResourceStateMode InResourceStateMode, 
 	const D3D12_CLEAR_VALUE* InClearValue, 
diff --git a/Engine/Source/Runtime/D3D12RHI/Private/D3D12PoolAllocator.h b/Engine/Source/Runtime/D3D12RHI/Private/D3D12PoolAllocator.h
--- a/Engine/Source/Runtime/D3D12RHI/Private/D3D12PoolAllocator.h
+++ b/Engine/Source/Runtime/D3D12RHI/Private/D3D12PoolAllocator.h
@@ -132,12 +132,12 @@ public:
 	
 	// Function names currently have to match with FD3D12DefaultBufferPool until we can make full replacement of the allocator
 	bool SupportsAllocation(D3D12_HEAP_TYPE InHeapType, D3D12_RESOURCE_FLAGS InResourceFlags, EBufferUsageFlags InBufferUsage, ED3D12ResourceStateMode InResourceStateMode, uint32 Alignment) const;
-	void AllocDefaultResource(D3D12_HEAP_TYPE InHeapType, const D3D12_RESOURCE_DESC& InDesc, EBufferUsageFlags InBufferUsage, ED3D12ResourceStateMode InResourceStateMode,
+	void AllocDefaultResource(D3D12_HEAP_TYPE InHeapType, const FD3D12ResourceDesc& InDesc, EBufferUsageFlags InBufferUsage, ED3D12ResourceStateMode InResourceStateMode,
 		D3D12_RESOURCE_STATES InCreateState, uint32 InAlignment, const TCHAR* InName, FD3D12ResourceLocation& ResourceLocation);
 	virtual void DeallocateResource(FD3D12ResourceLocation& ResourceLocation, bool bDefragFree = false);
 
 	// Implementation of ID3D12ResourceAllocator
-	virtual void AllocateResource(uint32 GPUIndex, D3D12_HEAP_TYPE InHeapType, const D3D12_RESOURCE_DESC& InDesc, uint64 InSize, uint32 InAllocationAlignment, ED3D12ResourceStateMode InResourceStateMode,
+	virtual void AllocateResource(uint32 GPUIndex, D3D12_HEAP_TYPE InHeapType, const FD3D12ResourceDesc& InDesc, uint64 InSize, uint32 InAllocationAlignment, ED3D12ResourceStateMode InResourceStateMode,
 		D3D12_RESOURCE_STATES InCreateState, const D3D12_CLEAR_VALUE* InClearValue, const TCHAR* InName, FD3D12ResourceLocation& ResourceLocation) override;
 
 	void CleanUpAllocations(uint64 InFrameLag);
@@ -161,7 +161,7 @@ protected:
 	virtual bool HandleDefragRequest(FRHIPoolAllocationData* InSourceBlock, FRHIPoolAllocationData& InTmpTargetBlock) override;
 	
 	// Placed resource allocation helper function which can be overriden
-	virtual FD3D12Resource* CreatePlacedResource(const FRHIPoolAllocationData& InAllocationData, const D3D12_RESOURCE_DESC& InDesc, D3D12_RESOURCE_STATES InCreateState, ED3D12ResourceStateMode InResourceStateMode, const D3D12_CLEAR_VALUE* InClearValue, const TCHAR* InName);
+	virtual FD3D12Resource* CreatePlacedResource(const FRHIPoolAllocationData& InAllocationData, const FD3D12ResourceDesc& InDesc, D3D12_RESOURCE_STATES InCreateState, ED3D12ResourceStateMode InResourceStateMode, const D3D12_CLEAR_VALUE* InClearValue, const TCHAR* InName);
 	
 	// Track the allocation
 	enum class EAllocationType
diff --git a/Engine/Source/Runtime/D3D12RHI/Private/D3D12RHIPrivate.h b/Engine/Source/Runtime/D3D12RHI/Private/D3D12RHIPrivate.h
--- a/Engine/Source/Runtime/D3D12RHI/Private/D3D12RHIPrivate.h
+++ b/Engine/Source/Runtime/D3D12RHI/Private/D3D12RHIPrivate.h
@@ -288,6 +288,7 @@ public:
 	virtual void RHIUnlockBuffer(FRHICommandListImmediate& RHICmdList, FRHIBuffer* Buffer) final override;
 	virtual FUnorderedAccessViewRHIRef RHICreateUnorderedAccessView(FRHIBuffer* Buffer, bool bUseUAVCounter, bool bAppendBuffer) final override;
 	virtual FUnorderedAccessViewRHIRef RHICreateUnorderedAccessView(FRHITexture* Texture, uint32 MipLevel) final override;
+	virtual FUnorderedAccessViewRHIRef RHICreateUnorderedAccessView(FRHITexture* Texture, uint32 MipLevel, uint8 Format) final override;
 	virtual FUnorderedAccessViewRHIRef RHICreateUnorderedAccessView(FRHIBuffer* Buffer, uint8 Format) final override;
 	virtual FShaderResourceViewRHIRef RHICreateShaderResourceView(FRHIBuffer* Buffer) final override;
 	virtual FShaderResourceViewRHIRef RHICreateShaderResourceView(FRHIBuffer* Buffer, uint32 Stride, uint8 Format) final override;
diff --git a/Engine/Source/Runtime/D3D12RHI/Private/D3D12Resources.cpp b/Engine/Source/Runtime/D3D12RHI/Private/D3D12Resources.cpp
--- a/Engine/Source/Runtime/D3D12RHI/Private/D3D12Resources.cpp
+++ b/Engine/Source/Runtime/D3D12RHI/Private/D3D12Resources.cpp
@@ -219,11 +219,11 @@ void FD3D12DeferredDeletionQueue::FD3D12AsyncDeletionWorker::DoWork()
 /////////////////////////////////////////////////////////////////////
 
 
-void ID3D12ResourceAllocator::AllocateTexture(uint32 GPUIndex, D3D12_HEAP_TYPE InHeapType, const D3D12_RESOURCE_DESC& InDesc, EPixelFormat InUEFormat, ED3D12ResourceStateMode InResourceStateMode,
+void ID3D12ResourceAllocator::AllocateTexture(uint32 GPUIndex, D3D12_HEAP_TYPE InHeapType, const FD3D12ResourceDesc& InDesc, EPixelFormat InUEFormat, ED3D12ResourceStateMode InResourceStateMode,
 	D3D12_RESOURCE_STATES InCreateState, const D3D12_CLEAR_VALUE* InClearValue, const TCHAR* InName, FD3D12ResourceLocation& ResourceLocation)
 {
 	// Check if texture can be 4K aligned
-	D3D12_RESOURCE_DESC Desc = InDesc;
+	FD3D12ResourceDesc Desc = InDesc;
 	bool b4KAligment = TextureCanBe4KAligned(Desc, InUEFormat);
 	Desc.Alignment = b4KAligment ? D3D12_SMALL_RESOURCE_PLACEMENT_ALIGNMENT : D3D12_DEFAULT_RESOURCE_PLACEMENT_ALIGNMENT;
 
@@ -246,7 +246,7 @@ FD3D12Resource::FD3D12Resource(FD3D12Device* ParentDevice,
 	FRHIGPUMask VisibleNodes,
 	ID3D12Resource* InResource,
 	D3D12_RESOURCE_STATES InInitialState,
-	D3D12_RESOURCE_DESC const& InDesc,
+	const FD3D12ResourceDesc& InDesc,
 	const D3D12_CLEAR_VALUE* InClearValue,
 	FD3D12Heap* InHeap,
 	D3D12_HEAP_TYPE InHeapType) : 
@@ -260,7 +260,7 @@ FD3D12Resource::FD3D12Resource(FD3D12Device* ParentDevice,
 	D3D12_RESOURCE_STATES InInitialState,
 	ED3D12ResourceStateMode InResourceStateMode,
 	D3D12_RESOURCE_STATES InDefaultResourceState,
-	D3D12_RESOURCE_DESC const& InDesc,
+	const FD3D12ResourceDesc& InDesc,
 	const D3D12_CLEAR_VALUE* InClearValue,
 	FD3D12Heap* InHeap,
 	D3D12_HEAP_TYPE InHeapType)
@@ -478,7 +478,7 @@ void FD3D12Heap::BeginTrackingResidency(uint64 Size)
 //	FD3D12 Adapter
 /////////////////////////////////////////////////////////////////////
 
-HRESULT FD3D12Adapter::CreateCommittedResource(const D3D12_RESOURCE_DESC& InDesc, FRHIGPUMask CreationNode, const D3D12_HEAP_PROPERTIES& HeapProps, D3D12_RESOURCE_STATES InInitialState,
+HRESULT FD3D12Adapter::CreateCommittedResource(const FD3D12ResourceDesc& InDesc, FRHIGPUMask CreationNode, const D3D12_HEAP_PROPERTIES& HeapProps, D3D12_RESOURCE_STATES InInitialState,
 	ED3D12ResourceStateMode InResourceStateMode, D3D12_RESOURCE_STATES InDefaultState, const D3D12_CLEAR_VALUE* ClearValue, FD3D12Resource** ppOutResource, const TCHAR* Name, bool bVerifyHResult)
 {
 	if (!ppOutResource)
@@ -498,7 +498,7 @@ HRESULT FD3D12Adapter::CreateCommittedResource(const D3D12_RESOURCE_DESC& InDesc
 		HeapFlags |= D3D12_HEAP_FLAG_SHARED;
 	}
 
-	D3D12_RESOURCE_DESC LocalDesc = InDesc;
+	FD3D12ResourceDesc LocalDesc = InDesc;
 
 #if D3D12_RHI_RAYTRACING
 	if (InDefaultState == D3D12_RESOURCE_STATE_RAYTRACING_ACCELERATION_STRUCTURE)
@@ -532,7 +532,7 @@ HRESULT FD3D12Adapter::CreateCommittedResource(const D3D12_RESOURCE_DESC& InDesc
 	return hr;
 }
 
-HRESULT FD3D12Adapter::CreatePlacedResource(const D3D12_RESOURCE_DESC& InDesc, FD3D12Heap* BackingHeap, uint64 HeapOffset, D3D12_RESOURCE_STATES InInitialState, ED3D12ResourceStateMode InResourceStateMode, 
+HRESULT FD3D12Adapter::CreatePlacedResource(const FD3D12ResourceDesc& InDesc, FD3D12Heap* BackingHeap, uint64 HeapOffset, D3D12_RESOURCE_STATES InInitialState, ED3D12ResourceStateMode InResourceStateMode,
 	D3D12_RESOURCE_STATES InDefaultState, const D3D12_CLEAR_VALUE* ClearValue, FD3D12Resource** ppOutResource, const TCHAR* Name, bool bVerifyHResult)
 {
 	if (!ppOutResource)
diff --git a/Engine/Source/Runtime/D3D12RHI/Private/D3D12Texture.cpp b/Engine/Source/Runtime/D3D12RHI/Private/D3D12Texture.cpp
--- a/Engine/Source/Runtime/D3D12RHI/Private/D3D12Texture.cpp
+++ b/Engine/Source/Runtime/D3D12RHI/Private/D3D12Texture.cpp
@@ -261,14 +261,14 @@ struct FD3D12RHICommandInitializeTexture final : public FRHICommand<FD3D12RHICom
 // Texture Stats
 ///////////////////////////////////////////////////////////////////////////////////////////
 
-bool FD3D12TextureStats::ShouldCountAsTextureMemory(uint32 MiscFlags)
+bool FD3D12TextureStats::ShouldCountAsTextureMemory(D3D12_RESOURCE_FLAGS MiscFlags)
 {
 	// Shouldn't be used for DEPTH, RENDER TARGET, or UNORDERED ACCESS
-	return (0 == (MiscFlags & (D3D12_RESOURCE_FLAG_ALLOW_DEPTH_STENCIL | D3D12_RESOURCE_FLAG_ALLOW_RENDER_TARGET | D3D12_RESOURCE_FLAG_ALLOW_UNORDERED_ACCESS)));
+	return !EnumHasAnyFlags(MiscFlags, D3D12_RESOURCE_FLAG_ALLOW_DEPTH_STENCIL | D3D12_RESOURCE_FLAG_ALLOW_RENDER_TARGET | D3D12_RESOURCE_FLAG_ALLOW_UNORDERED_ACCESS);
 }
 
 // @param b3D true:3D, false:2D or cube map
-TStatId FD3D12TextureStats::GetD3D12StatEnum(uint32 MiscFlags, bool bCubeMap, bool b3D)
+TStatId FD3D12TextureStats::GetD3D12StatEnum(D3D12_RESOURCE_FLAGS MiscFlags, bool bCubeMap, bool b3D)
 {
 #if STATS
 	if (ShouldCountAsTextureMemory(MiscFlags))
@@ -679,11 +679,11 @@ bool FD3D12DynamicRHI::RHIGetTextureMemoryVisualizeData(FColor* /*TextureData*/,
  */
 void SafeCreateTexture2D(FD3D12Device* pDevice, 
 	FD3D12Adapter* Adapter,
-	const D3D12_RESOURCE_DESC& TextureDesc,
+	const FD3D12ResourceDesc& TextureDesc,
 	const D3D12_CLEAR_VALUE* ClearValue, 
 	FD3D12ResourceLocation* OutTexture2D, 
 	FD3D12BaseShaderResource* Owner,
-	uint8 Format, 
+	EPixelFormat Format,
 	ETextureCreateFlags Flags,
 	D3D12_RESOURCE_STATES InitialState,
 	const TCHAR* Name)
@@ -742,6 +742,50 @@ void SafeCreateTexture2D(FD3D12Device* pDevice,
 #endif // #if GUARDED_TEXTURE_CREATES
 }
 
+void CreateUAVAliasResource(FD3D12Adapter* Adapter, D3D12_CLEAR_VALUE* ClearValuePtr, const TCHAR* DebugName, FD3D12ResourceLocation& Location)
+{
+	FD3D12Resource* SourceResource = Location.GetResource();
+
+	const FD3D12ResourceDesc& SourceDesc = SourceResource->GetDesc();
+	const FD3D12Heap* const ResourceHeap = SourceResource->GetHeap();
+
+	const EPixelFormat SourceFormat = SourceDesc.PixelFormat;
+	const EPixelFormat AliasTextureFormat = SourceDesc.UAVAliasPixelFormat;
+
+	if (ensure(ResourceHeap != nullptr) && ensure(SourceFormat != PF_Unknown) && SourceFormat != AliasTextureFormat)
+	{
+		const uint64 SourceOffset = Location.GetOffsetFromBaseOfResource();
+
+		FD3D12ResourceDesc AliasTextureDesc = SourceDesc;
+		AliasTextureDesc.Format = (DXGI_FORMAT)GPixelFormats[AliasTextureFormat].PlatformFormat;
+		AliasTextureDesc.Width = SourceDesc.Width / GPixelFormats[SourceFormat].BlockSizeX;
+		AliasTextureDesc.Height = SourceDesc.Height / GPixelFormats[SourceFormat].BlockSizeY;
+		EnumAddFlags(AliasTextureDesc.Flags, D3D12_RESOURCE_FLAG_ALLOW_UNORDERED_ACCESS);
+
+		AliasTextureDesc.UAVAliasPixelFormat = PF_Unknown;
+
+		TRefCountPtr<ID3D12Resource> pAliasResource;
+		HRESULT AliasHR = Adapter->GetD3DDevice()->CreatePlacedResource(
+			ResourceHeap->GetHeap(),
+			SourceOffset,
+			&AliasTextureDesc,
+			D3D12_RESOURCE_STATE_UNORDERED_ACCESS,
+			ClearValuePtr,
+			IID_PPV_ARGS(pAliasResource.GetInitReference()));
+
+		if (pAliasResource && DebugName)
+		{
+			TCHAR NameBuffer[512]{};
+			FCString::Snprintf(NameBuffer, UE_ARRAY_COUNT(NameBuffer), TEXT("%s UAVAlias"), DebugName);
+			SetName(pAliasResource, NameBuffer);
+		}
+
+		if (SUCCEEDED(AliasHR))
+		{
+			SourceResource->SetUAVAccessResource(pAliasResource);
+		}
+	}
+}
 
 static void DetermineTexture2DResourceFlagsAndLayout(uint32 SizeX, uint32 SizeY, uint32 SizeZ, uint32 NumMips, uint32 NumSamples, ETextureCreateFlags Flags, EPixelFormat Format, D3D12_RESOURCE_FLAGS& OutResourceFlags, D3D12_TEXTURE_LAYOUT& OutLayout, bool& bOutCreateRTV, bool& bOutCreateDSV, bool& bOutCreateSRV)
 {
@@ -857,7 +901,7 @@ TD3D12Texture2D<BaseResourceType>* FD3D12DynamicRHI::CreateD3D12Texture2D(FRHICo
 	const bool bIsMultisampled = ActualMSAACount > 1;
 
 	// Describe the texture.
-	D3D12_RESOURCE_DESC TextureDesc = CD3DX12_RESOURCE_DESC::Tex2D(
+	FD3D12ResourceDesc TextureDesc = CD3DX12_RESOURCE_DESC::Tex2D(
 		PlatformResourceFormat,
 		SizeX,
 		SizeY,
@@ -867,6 +911,15 @@ TD3D12Texture2D<BaseResourceType>* FD3D12DynamicRHI::CreateD3D12Texture2D(FRHICo
 		ActualMSAAQuality,
 		D3D12_RESOURCE_FLAG_NONE);  // Add misc flags later
 
+	TextureDesc.PixelFormat = Format;
+
+	bool bBCTextureNeedsUAVAlias = EnumHasAnyFlags(Flags, TexCreate_UAV) && IsBlockCompressedFormat(Format);
+	if (bBCTextureNeedsUAVAlias)
+	{
+		EnumRemoveFlags(Flags, TexCreate_UAV);
+		TextureDesc.UAVAliasPixelFormat = GetBlockCompressedFormatUAVAliasFormat(Format);
+	}
+
 	// Set up the texture bind flags.
 	bool bCreateRTV;
 	bool bCreateDSV;
@@ -945,6 +998,11 @@ TD3D12Texture2D<BaseResourceType>* FD3D12DynamicRHI::CreateD3D12Texture2D(FRHICo
 
 		check(Location.IsValid());
 
+		if (bBCTextureNeedsUAVAlias)
+		{
+			CreateUAVAliasResource(Adapter, ClearValuePtr, CreateInfo.DebugName, Location);
+		}
+
 		uint32 RTVIndex = 0;
 
 		{
@@ -1209,7 +1267,7 @@ FD3D12Texture3D* FD3D12DynamicRHI::CreateD3D12Texture3D(FRHICommandListImmediate
 
 		if (ResourceAllocator)
 		{
-			ResourceAllocator->AllocateTexture(Device->GetGPUIndex(), D3D12_HEAP_TYPE_DEFAULT, TextureDesc, (EPixelFormat)Format, ED3D12ResourceStateMode::Default, InitialState, ClearValuePtr, CreateInfo.DebugName, Texture3D->ResourceLocation);
+			ResourceAllocator->AllocateTexture(Device->GetGPUIndex(), D3D12_HEAP_TYPE_DEFAULT, TextureDesc, Format, ED3D12ResourceStateMode::Default, InitialState, ClearValuePtr, CreateInfo.DebugName, Texture3D->ResourceLocation);
 		}
 		else
 		{
@@ -1394,7 +1452,7 @@ FTexture2DRHIRef FD3D12DynamicRHI::RHIAsyncCreateTexture2D(uint32 SizeX, uint32
 			nullptr,
 			&NewTexture->ResourceLocation,
 			NewTexture,
-			Format,
+			(EPixelFormat)Format,
 			Flags,
 			InitialState,
 			nullptr);
diff --git a/Engine/Source/Runtime/D3D12RHI/Private/D3D12Texture.h b/Engine/Source/Runtime/D3D12RHI/Private/D3D12Texture.h
--- a/Engine/Source/Runtime/D3D12RHI/Private/D3D12Texture.h
+++ b/Engine/Source/Runtime/D3D12RHI/Private/D3D12Texture.h
@@ -8,18 +8,43 @@ D3D12Texture.h: Implementation of D3D12 Texture
 /** If true, guard texture creates with SEH to log more information about a driver crash we are seeing during texture streaming. */
 #define GUARDED_TEXTURE_CREATES (PLATFORM_WINDOWS && !(UE_BUILD_SHIPPING || UE_BUILD_TEST || PLATFORM_COMPILER_CLANG))
 
+static bool TextureCanBe4KAligned(const FD3D12ResourceDesc& Desc, EPixelFormat UEFormat)
+{
+	// 4KB alignment is only available for read only textures
+	if (!EnumHasAnyFlags(Desc.Flags, D3D12_RESOURCE_FLAG_ALLOW_RENDER_TARGET | D3D12_RESOURCE_FLAG_ALLOW_DEPTH_STENCIL | D3D12_RESOURCE_FLAG_ALLOW_UNORDERED_ACCESS) &&
+		!Desc.NeedsUAVAliasWorkarounds() && // UAV aliased resources are secretly writable.
+		Desc.SampleDesc.Count == 1)
+	{
+		D3D12_TILE_SHAPE Tile = {};
+		Get4KTileShape(&Tile, Desc.Format, UEFormat, Desc.Dimension, Desc.SampleDesc.Count);
+
+		uint32 TilesNeeded = GetTilesNeeded(Desc.Width, Desc.Height, Desc.DepthOrArraySize, Tile);
+
+		constexpr uint32 NUM_4K_BLOCKS_PER_64K_PAGE = 16;
+		return TilesNeeded <= NUM_4K_BLOCKS_PER_64K_PAGE;
+	}
+	else
+	{
+		return false;
+	}
+}
 
 void SafeCreateTexture2D(FD3D12Device* pDevice, 
 	FD3D12Adapter* Adapter,
-	const D3D12_RESOURCE_DESC& TextureDesc,
+	const FD3D12ResourceDesc& TextureDesc,
 	const D3D12_CLEAR_VALUE* ClearValue, 
 	FD3D12ResourceLocation* OutTexture2D,
 	FD3D12BaseShaderResource* Owner,
-	uint8 Format, 
+	EPixelFormat Format,
 	ETextureCreateFlags Flags,
 	D3D12_RESOURCE_STATES InitialState,
 	const TCHAR* Name);
 
+void CreateUAVAliasResource(
+	FD3D12Adapter* Adapter,
+	D3D12_CLEAR_VALUE* ClearValuePtr,
+	const TCHAR* DebugName,
+	FD3D12ResourceLocation& Location);
 
 /** Texture base class. */
 class FD3D12TextureBase : public FD3D12BaseShaderResource, public FD3D12LinkedAdapterObject<FD3D12TextureBase>
@@ -592,9 +617,9 @@ class FD3D12TextureStats
 {
 public:
 
-	static bool ShouldCountAsTextureMemory(uint32 MiscFlags);
+	static bool ShouldCountAsTextureMemory(D3D12_RESOURCE_FLAGS MiscFlags);
 	// @param b3D true:3D, false:2D or cube map
-	static TStatId GetD3D12StatEnum(uint32 MiscFlags, bool bCubeMap, bool b3D);
+	static TStatId GetD3D12StatEnum(D3D12_RESOURCE_FLAGS MiscFlags, bool bCubeMap, bool b3D);
 
 	// Note: This function can be called from many different threads
 	// @param TextureSize >0 to allocate, <0 to deallocate
diff --git a/Engine/Source/Runtime/D3D12RHI/Private/D3D12TransientResourceAllocator.cpp b/Engine/Source/Runtime/D3D12RHI/Private/D3D12TransientResourceAllocator.cpp
--- a/Engine/Source/Runtime/D3D12RHI/Private/D3D12TransientResourceAllocator.cpp
+++ b/Engine/Source/Runtime/D3D12RHI/Private/D3D12TransientResourceAllocator.cpp
@@ -152,7 +152,7 @@ FRHITransientTexture* FD3D12TransientResourceAllocator::CreateTexture(const FRHI
 }
 
 void FD3D12TransientResourceAllocator::FResourceAllocatorAdapter::AllocateResource(
-	uint32 GPUIndex, D3D12_HEAP_TYPE, const D3D12_RESOURCE_DESC& InDesc, uint64 InSize, uint32, ED3D12ResourceStateMode InResourceStateMode,
+	uint32 GPUIndex, D3D12_HEAP_TYPE, const FD3D12ResourceDesc& InDesc, uint64 InSize, uint32, ED3D12ResourceStateMode InResourceStateMode,
 	D3D12_RESOURCE_STATES InCreateState, const D3D12_CLEAR_VALUE* InClearValue, const TCHAR* InName, FD3D12ResourceLocation& ResourceLocation)
 {
 	// The D3D12_RESOURCE_DESC's are built in two different functions right now. This checks that they actually match what we expect.
diff --git a/Engine/Source/Runtime/D3D12RHI/Private/D3D12TransientResourceAllocator.h b/Engine/Source/Runtime/D3D12RHI/Private/D3D12TransientResourceAllocator.h
--- a/Engine/Source/Runtime/D3D12RHI/Private/D3D12TransientResourceAllocator.h
+++ b/Engine/Source/Runtime/D3D12RHI/Private/D3D12TransientResourceAllocator.h
@@ -66,7 +66,7 @@ private:
 		{}
 
 		void AllocateResource(
-			uint32 GPUIndex, D3D12_HEAP_TYPE InHeapType, const D3D12_RESOURCE_DESC& InDesc, uint64 InSize, uint32 InAllocationAlignment, ED3D12ResourceStateMode InResourceStateMode,
+			uint32 GPUIndex, D3D12_HEAP_TYPE InHeapType, const FD3D12ResourceDesc& InDesc, uint64 InSize, uint32 InAllocationAlignment, ED3D12ResourceStateMode InResourceStateMode,
 			D3D12_RESOURCE_STATES InCreateState, const D3D12_CLEAR_VALUE* InClearValue, const TCHAR* InName, FD3D12ResourceLocation& ResourceLocation) override;
 
 		FD3D12TransientHeap& Heap;
diff --git a/Engine/Source/Runtime/D3D12RHI/Private/D3D12UAV.cpp b/Engine/Source/Runtime/D3D12RHI/Private/D3D12UAV.cpp
--- a/Engine/Source/Runtime/D3D12RHI/Private/D3D12UAV.cpp
+++ b/Engine/Source/Runtime/D3D12RHI/Private/D3D12UAV.cpp
@@ -65,13 +65,13 @@ FUnorderedAccessViewRHIRef FD3D12DynamicRHI::RHICreateUnorderedAccessView(FRHIBu
 	return CreateUAV(UAVDesc, Buffer, bNeedsCounterResource);
 }
 
-FUnorderedAccessViewRHIRef FD3D12DynamicRHI::RHICreateUnorderedAccessView(FRHITexture* TextureRHI, uint32 MipLevel)
+FUnorderedAccessViewRHIRef FD3D12DynamicRHI::RHICreateUnorderedAccessView(FRHITexture* TextureRHI, uint32 MipLevel, uint8 Format)
 {
 	FD3D12TextureBase* Texture = GetD3D12TextureFromRHITexture(TextureRHI);
 
 	D3D12_UNORDERED_ACCESS_VIEW_DESC UAVDesc = {};
 
-	const DXGI_FORMAT PlatformResourceFormat = (DXGI_FORMAT)GPixelFormats[TextureRHI->GetFormat()].PlatformFormat;
+	const DXGI_FORMAT PlatformResourceFormat = (DXGI_FORMAT)GPixelFormats[Format].PlatformFormat;
 	UAVDesc.Format = FindShaderResourceDXGIFormat(PlatformResourceFormat, false);
 
 	if (TextureRHI->GetTexture3D() != NULL)
@@ -117,6 +117,11 @@ FUnorderedAccessViewRHIRef FD3D12DynamicRHI::RHICreateUnorderedAccessView(FRHITe
 	}
 }
 
+FUnorderedAccessViewRHIRef FD3D12DynamicRHI::RHICreateUnorderedAccessView(FRHITexture* TextureRHI, uint32 MipLevel)
+{
+	return RHICreateUnorderedAccessView(TextureRHI, MipLevel, TextureRHI->GetFormat());
+}
+
 FUnorderedAccessViewRHIRef FD3D12DynamicRHI::RHICreateUnorderedAccessView(FRHIBuffer* BufferRHI, uint8 Format)
 {
 	FD3D12Buffer* Buffer = FD3D12DynamicRHI::ResourceCast(BufferRHI);
diff --git a/Engine/Source/Runtime/D3D12RHI/Private/D3D12View.h b/Engine/Source/Runtime/D3D12RHI/Private/D3D12View.h
--- a/Engine/Source/Runtime/D3D12RHI/Private/D3D12View.h
+++ b/Engine/Source/Runtime/D3D12RHI/Private/D3D12View.h
@@ -668,10 +668,9 @@ typedef TD3D12ViewDescriptorHandle<D3D12_UNORDERED_ACCESS_VIEW_DESC>	FD3D12Descr
 template <typename TDesc>
 class FD3D12View : public FD3D12ShaderResourceRenameListener
 {
-private:
+protected:
 	TD3D12ViewDescriptorHandle<TDesc> Descriptor;
 
-protected:
 	ViewSubresourceSubsetFlags Flags;
 	FD3D12BaseShaderResource* BaseShaderResource;
 	FD3D12ResourceLocation* ResourceLocation;
@@ -705,7 +704,7 @@ protected:
 #endif
 	}
 
-private:
+protected:
 	void InitializeInternal(const TDesc& InDesc, FD3D12BaseShaderResource* InBaseShaderResource, FD3D12ResourceLocation& InResourceLocation)
 	{
 		check(InBaseShaderResource);
@@ -743,7 +742,6 @@ private:
 #endif
 	}
 
-protected:
 	void CreateView(const TDesc& InDesc, FD3D12BaseShaderResource* InBaseShaderResource, FD3D12ResourceLocation& InResourceLocation)
 	{
 		InitializeInternal(InDesc, InBaseShaderResource, InResourceLocation);
@@ -755,18 +753,6 @@ protected:
 		}
 	}
 
-	void CreateViewWithCounter(const TDesc& InDesc, FD3D12BaseShaderResource* InBaseShaderResource, FD3D12ResourceLocation& InResourceLocation, FD3D12Resource* InCounterResource)
-	{
-		InitializeInternal(InDesc, InBaseShaderResource, InResourceLocation);
-
-		if (ResourceLocation->GetResource())
-		{
-			ID3D12Resource* D3DResource = ResourceLocation->GetResource()->GetResource();
-			ID3D12Resource* D3DCounterResource = InCounterResource ? InCounterResource->GetResource() : nullptr;
-			Descriptor.CreateViewWithCounter(Desc, D3DResource, D3DCounterResource);
-		}
-	}
-
 	virtual void ResourceRenamed(FD3D12BaseShaderResource* InRenamedResource, FD3D12ResourceLocation* InNewResourceLocation) override
 	{
 		check(InRenamedResource == BaseShaderResource);
@@ -982,6 +968,19 @@ public:
 		check(ResourceLocation->GetOffsetFromBaseOfResource() == 0);
 		Initialize(Desc, BaseShaderResource, *ResourceLocation);
 	}
+
+protected:
+	void CreateViewWithCounter(const D3D12_UNORDERED_ACCESS_VIEW_DESC& InDesc, FD3D12BaseShaderResource* InBaseShaderResource, FD3D12ResourceLocation& InResourceLocation, FD3D12Resource* InCounterResource)
+	{
+		InitializeInternal(InDesc, InBaseShaderResource, InResourceLocation);
+
+		if (Resource)
+		{
+			ID3D12Resource* D3DResource = Resource->GetUAVAccessResource() ? Resource->GetUAVAccessResource() : Resource->GetResource();
+			ID3D12Resource* D3DCounterResource = InCounterResource ? InCounterResource->GetResource() : nullptr;
+			Descriptor.CreateViewWithCounter(Desc, D3DResource, D3DCounterResource);
+		}
+	}
 };
 
 class FD3D12UnorderedAccessViewWithLocation : public FD3D12UnorderedAccessView
diff --git a/Engine/Source/Runtime/D3D12RHI/Private/Windows/WindowsD3D12Device.cpp b/Engine/Source/Runtime/D3D12RHI/Private/Windows/WindowsD3D12Device.cpp
--- a/Engine/Source/Runtime/D3D12RHI/Private/Windows/WindowsD3D12Device.cpp
+++ b/Engine/Source/Runtime/D3D12RHI/Private/Windows/WindowsD3D12Device.cpp
@@ -1129,6 +1129,8 @@ void FD3D12DynamicRHI::Init()
 		GRHIVariableRateShadingImageTileMaxHeight = 1;
 	}
 
+	GRHISupportsUAVFormatAliasing = true;
+
 	// Command lists need the validation RHI context if enabled, so call the global scope version of RHIGetDefaultContext() and RHIGetDefaultAsyncComputeContext().
 	GRHICommandList.GetImmediateCommandList().SetContext(::RHIGetDefaultContext());
 	GRHICommandList.GetImmediateAsyncComputeCommandList().SetComputeContext(::RHIGetDefaultAsyncComputeContext());
diff --git a/Engine/Source/Runtime/D3D12RHI/Public/D3D12Resources.h b/Engine/Source/Runtime/D3D12RHI/Public/D3D12Resources.h
--- a/Engine/Source/Runtime/D3D12RHI/Public/D3D12Resources.h
+++ b/Engine/Source/Runtime/D3D12RHI/Public/D3D12Resources.h
@@ -93,7 +93,7 @@ public:
 	FD3D12Heap(FD3D12Device* Parent, FRHIGPUMask VisibleNodes);
 	~FD3D12Heap();
 
-	inline ID3D12Heap* GetHeap() { return Heap.GetReference(); }
+	inline ID3D12Heap* GetHeap() const { return Heap.GetReference(); }
 	void SetHeap(ID3D12Heap* HeapIn, const TCHAR* const InName, bool bTrack = true);
 
 	void UpdateResidency(FD3D12CommandListHandle& CommandList);
@@ -114,16 +114,42 @@ private:
 	FD3D12ResidencyHandle ResidencyHandle;
 };
 
+struct FD3D12ResourceDesc : public D3D12_RESOURCE_DESC
+{
+	FD3D12ResourceDesc() = default;
+	FD3D12ResourceDesc(const CD3DX12_RESOURCE_DESC& Other)
+		: D3D12_RESOURCE_DESC(Other)
+	{
+	}
+	
+	// TODO: use this type everywhere and disallow implicit conversion
+	/*explicit*/ FD3D12ResourceDesc(const D3D12_RESOURCE_DESC& Other)
+		: D3D12_RESOURCE_DESC(Other)
+	{
+	}
+
+	EPixelFormat PixelFormat{ PF_Unknown };
+	// PixelFormat for the Resource that aliases our current resource.
+	EPixelFormat UAVAliasPixelFormat{ PF_Unknown };
+
+	// Used primarily to help treat this resource description as writable.
+	inline bool NeedsUAVAliasWorkarounds() const { return UAVAliasPixelFormat != PF_Unknown; }
+};
+
 class FD3D12Resource : public FD3D12RefCount, public FD3D12DeviceChild, public FD3D12MultiNodeGPUObject
 {
 private:
 	TRefCountPtr<ID3D12Resource> Resource;
+	// Since certain formats cannot be aliased in D3D12, we have to create a separate ID3D12Resource that aliases the
+	// resource's memory and use this separate resource to create the UAV.
+	// TODO: UE-116727: request better DX12 API support and clean this up.
+	TRefCountPtr<ID3D12Resource> UAVAccessResource;
 	TRefCountPtr<FD3D12Heap> Heap;
 
 	FD3D12ResidencyHandle ResidencyHandle;
 
 	D3D12_CLEAR_VALUE ClearValue;
-	D3D12_RESOURCE_DESC Desc;
+	const FD3D12ResourceDesc Desc;
 	uint8 PlaneCount;
 	uint16 SubresourceCount;
 	CResourceState ResourceState;
@@ -133,6 +159,7 @@ private:
 #ifdef PLATFORM_SUPPORTS_RESOURCE_COMPRESSION
 	D3D12_RESOURCE_STATES CompressedState;
 #endif
+	D3D12_RESOURCE_STATES UAVHiddenResourceState{ D3D12_RESOURCE_STATE_CORRUPT };
 
 	bool bRequiresResourceStateTracking : 1;
 	bool bDepthStencil : 1;
@@ -158,7 +185,7 @@ public:
 		FRHIGPUMask VisibleNodes,
 		ID3D12Resource* InResource,
 		D3D12_RESOURCE_STATES InInitialResourceState,
-		D3D12_RESOURCE_DESC const& InDesc,
+		FD3D12ResourceDesc const& InDesc,
 		const D3D12_CLEAR_VALUE* InClearValue = nullptr,
 		FD3D12Heap* InHeap = nullptr,
 		D3D12_HEAP_TYPE InHeapType = D3D12_HEAP_TYPE_DEFAULT);
@@ -169,7 +196,7 @@ public:
 		D3D12_RESOURCE_STATES InInitialResourceState,
 		ED3D12ResourceStateMode InResourceStateMode,
 		D3D12_RESOURCE_STATES InDefaultResourceState,
-		D3D12_RESOURCE_DESC const& InDesc,
+		FD3D12ResourceDesc const& InDesc,
 		const D3D12_CLEAR_VALUE* InClearValue,
 		FD3D12Heap* InHeap,
 		D3D12_HEAP_TYPE InHeapType);
@@ -177,7 +204,13 @@ public:
 	virtual ~FD3D12Resource();
 
 	operator ID3D12Resource&() { return *Resource; }
+
 	ID3D12Resource* GetResource() const { return Resource.GetReference(); }
+	ID3D12Resource* GetUAVAccessResource() const { return UAVAccessResource.GetReference(); }
+	void SetUAVAccessResource(ID3D12Resource* InUAVAccessResource) { UAVAccessResource = InUAVAccessResource; }
+
+	inline D3D12_RESOURCE_STATES GetUAVHiddenResourceState() const { return UAVHiddenResourceState; }
+	void SetUAVHiddenResourceState(D3D12_RESOURCE_STATES InUAVHiddenResourceState) { UAVHiddenResourceState = InUAVHiddenResourceState; }
 
 	inline void* Map(const D3D12_RANGE* ReadRange = nullptr)
 	{
@@ -198,7 +231,7 @@ public:
 	}
 
 	ID3D12Pageable* GetPageable();
-	D3D12_RESOURCE_DESC const& GetDesc() const { return Desc; }
+	const FD3D12ResourceDesc& GetDesc() const { return Desc; }
 	D3D12_CLEAR_VALUE const& GetClearValue() const { return ClearValue; }
 	D3D12_HEAP_TYPE GetHeapType() const { return HeapType; }
 	D3D12_GPU_VIRTUAL_ADDRESS GetGPUVirtualAddress() const { return GPUVirtualAddress; }
@@ -261,11 +294,11 @@ public:
 
 	struct FD3D12ResourceTypeHelper
 	{
-		FD3D12ResourceTypeHelper(const D3D12_RESOURCE_DESC& Desc, D3D12_HEAP_TYPE HeapType) :
-			bSRV((Desc.Flags & D3D12_RESOURCE_FLAG_DENY_SHADER_RESOURCE) == 0),
-			bDSV((Desc.Flags & D3D12_RESOURCE_FLAG_ALLOW_DEPTH_STENCIL) != 0),
-			bRTV((Desc.Flags & D3D12_RESOURCE_FLAG_ALLOW_RENDER_TARGET) != 0),
-			bUAV((Desc.Flags & D3D12_RESOURCE_FLAG_ALLOW_UNORDERED_ACCESS) != 0),
+		FD3D12ResourceTypeHelper(const FD3D12ResourceDesc& Desc, D3D12_HEAP_TYPE HeapType) :
+			bSRV(!EnumHasAnyFlags(Desc.Flags, D3D12_RESOURCE_FLAG_DENY_SHADER_RESOURCE)),
+			bDSV(EnumHasAnyFlags(Desc.Flags, D3D12_RESOURCE_FLAG_ALLOW_DEPTH_STENCIL)),
+			bRTV(EnumHasAnyFlags(Desc.Flags, D3D12_RESOURCE_FLAG_ALLOW_RENDER_TARGET)),
+			bUAV(EnumHasAnyFlags(Desc.Flags, D3D12_RESOURCE_FLAG_ALLOW_UNORDERED_ACCESS) || Desc.NeedsUAVAliasWorkarounds()),
 			bWritable(bDSV || bRTV || bUAV),
 			bSRVOnly(bSRV && !bWritable),
 			bBuffer(Desc.Dimension == D3D12_RESOURCE_DIMENSION_BUFFER),
@@ -657,11 +690,11 @@ private:
 struct ID3D12ResourceAllocator
 {
 	// Helper function for textures to compute the correct size and alignment
-	void AllocateTexture(uint32 GPUIndex, D3D12_HEAP_TYPE InHeapType, const D3D12_RESOURCE_DESC& InDesc, EPixelFormat InUEFormat, ED3D12ResourceStateMode InResourceStateMode,
+	void AllocateTexture(uint32 GPUIndex, D3D12_HEAP_TYPE InHeapType, const FD3D12ResourceDesc& InDesc, EPixelFormat InUEFormat, ED3D12ResourceStateMode InResourceStateMode,
 		D3D12_RESOURCE_STATES InCreateState, const D3D12_CLEAR_VALUE* InClearValue, const TCHAR* InName, FD3D12ResourceLocation& ResourceLocation);
 
 	// Actual pure virtual resource allocation function
-	virtual void AllocateResource(uint32 GPUIndex, D3D12_HEAP_TYPE InHeapType, const D3D12_RESOURCE_DESC& InDesc, uint64 InSize, uint32 InAllocationAlignment, ED3D12ResourceStateMode InResourceStateMode,
+	virtual void AllocateResource(uint32 GPUIndex, D3D12_HEAP_TYPE InHeapType, const FD3D12ResourceDesc& InDesc, uint64 InSize, uint32 InAllocationAlignment, ED3D12ResourceStateMode InResourceStateMode,
 		D3D12_RESOURCE_STATES InCreateState, const D3D12_CLEAR_VALUE* InClearValue, const TCHAR* InName, FD3D12ResourceLocation& ResourceLocation) = 0;
 };
 
@@ -968,6 +1001,43 @@ public:
 
 class FD3D12ShaderResourceView;
 
+template<typename InAllocatorType>
+inline void AddTransitionBarrier(TArray<D3D12_RESOURCE_BARRIER, InAllocatorType>& BarrierList, FD3D12Resource* pResource, D3D12_RESOURCE_STATES Before, D3D12_RESOURCE_STATES After, uint32 Subresource)
+{
+	BarrierList.Add(CD3DX12_RESOURCE_BARRIER::Transition(pResource->GetResource(), Before, After, Subresource));
+}
+
+template<typename InAllocatorType>
+inline void AddTransitionBarrierWithUAVAccessOverrides(TArray<D3D12_RESOURCE_BARRIER, InAllocatorType>& BarrierList, FD3D12Resource* pResource, D3D12_RESOURCE_STATES Before, D3D12_RESOURCE_STATES After, uint32 Subresource)
+{
+	if (pResource->GetUAVAccessResource() && EnumHasAnyFlags(Before | After, D3D12_RESOURCE_STATE_UNORDERED_ACCESS))
+	{
+		// inject an aliasing barrier
+		const bool bFromUAV = EnumHasAnyFlags(Before, D3D12_RESOURCE_STATE_UNORDERED_ACCESS);
+		const bool bToUAV = EnumHasAnyFlags(After, D3D12_RESOURCE_STATE_UNORDERED_ACCESS);
+
+		check(bFromUAV != bToUAV);
+
+		BarrierList.Add(CD3DX12_RESOURCE_BARRIER::Aliasing(
+			bFromUAV ? pResource->GetUAVAccessResource() : pResource->GetResource(),
+			bToUAV ? pResource->GetUAVAccessResource() : pResource->GetResource()
+		));
+
+		if (bToUAV)
+		{
+			pResource->SetUAVHiddenResourceState(Before);
+		}
+		else if (D3D12_RESOURCE_STATES HiddenState = pResource->GetUAVHiddenResourceState(); HiddenState != After && HiddenState != D3D12_RESOURCE_STATE_CORRUPT)
+		{
+			AddTransitionBarrier(BarrierList, pResource, HiddenState, After, Subresource);
+		}
+	}
+	else
+	{
+		AddTransitionBarrier(BarrierList, pResource, Before, After, Subresource);
+	}
+}
+
 class FD3D12ResourceBarrierBatcher : public FNoncopyable
 {
 public:
@@ -1009,26 +1079,17 @@ public:
 
 		check(IsValidD3D12ResourceState(Before) && IsValidD3D12ResourceState(After));
 
-		D3D12_RESOURCE_BARRIER* Barrier = nullptr;
 #if PLATFORM_USE_BACKBUFFER_WRITE_TRANSITION_TRACKING
-		if (pResource->IsBackBuffer() && (After & BackBufferBarrierWriteTransitionTargets))
+		if (pResource->IsBackBuffer() && EnumHasAnyFlags(After, BackBufferBarrierWriteTransitionTargets))
 		{
-			BackBufferBarriers.AddUninitialized();
-			Barrier = &BackBufferBarriers.Last();
+			AddTransitionBarrier(BackBufferBarriers, pResource, Before, After, Subresource);
 		}
 		else
-#endif // #if PLATFORM_USE_BACKBUFFER_WRITE_TRANSITION_TRACKING
+#endif
 		{
-			Barriers.AddUninitialized();
-			Barrier = &Barriers.Last();
+			AddTransitionBarrierWithUAVAccessOverrides(Barriers, pResource, Before, After, Subresource);
 		}
 
-		Barrier->Type = D3D12_RESOURCE_BARRIER_TYPE_TRANSITION;
-		Barrier->Flags = D3D12_RESOURCE_BARRIER_FLAG_NONE;
-		Barrier->Transition.StateBefore = Before;
-		Barrier->Transition.StateAfter = After;
-		Barrier->Transition.Subresource = Subresource;
-		Barrier->Transition.pResource = pResource->GetResource();
 		return 1;
 	}
 
diff --git a/Engine/Source/Runtime/D3D12RHI/Public/D3D12Util.h b/Engine/Source/Runtime/D3D12RHI/Public/D3D12Util.h
--- a/Engine/Source/Runtime/D3D12RHI/Public/D3D12Util.h
+++ b/Engine/Source/Runtime/D3D12RHI/Public/D3D12Util.h
@@ -775,29 +775,6 @@ static void Get4KTileShape(D3D12_TILE_SHAPE* pTileShape, DXGI_FORMAT DXGIFormat,
 	}
 }
 
-#define NUM_4K_BLOCKS_PER_64K_PAGE (16)
-
-static bool TextureCanBe4KAligned(D3D12_RESOURCE_DESC& Desc, EPixelFormat UEFormat)
-{
-	// 4KB alignment is only available for read only textures
-	if (!(Desc.Flags & D3D12_RESOURCE_FLAG_ALLOW_RENDER_TARGET ||
-		Desc.Flags & D3D12_RESOURCE_FLAG_ALLOW_DEPTH_STENCIL ||
-		Desc.Flags & D3D12_RESOURCE_FLAG_ALLOW_UNORDERED_ACCESS) &&
-		Desc.SampleDesc.Count == 1)
-	{
-		D3D12_TILE_SHAPE Tile = {};
-		Get4KTileShape(&Tile, Desc.Format, UEFormat, Desc.Dimension, Desc.SampleDesc.Count);
-
-		uint32 TilesNeeded = GetTilesNeeded(Desc.Width, Desc.Height, Desc.DepthOrArraySize, Tile);
-
-		return TilesNeeded <= NUM_4K_BLOCKS_PER_64K_PAGE;
-	}
-	else
-	{
-		return false;
-	}
-}
-
 template <class TView>
 class FD3D12View;
 class CViewSubresourceSubset;
diff --git a/Engine/Source/Runtime/RHI/Public/RHIUtilities.h b/Engine/Source/Runtime/RHI/Public/RHIUtilities.h
--- a/Engine/Source/Runtime/RHI/Public/RHIUtilities.h
+++ b/Engine/Source/Runtime/RHI/Public/RHIUtilities.h
@@ -65,6 +65,25 @@ static inline bool IsBlockCompressedFormat(EPixelFormat Format)
 	return false;
 }
 
+static inline EPixelFormat GetBlockCompressedFormatUAVAliasFormat(EPixelFormat Format)
+{
+	switch (Format)
+	{
+	case PF_DXT1:
+	case PF_BC4:
+		return PF_R32G32_UINT;
+
+	case PF_DXT3:
+	case PF_DXT5:
+	case PF_BC5:
+	case PF_BC6H:
+	case PF_BC7:
+		return PF_R32G32B32A32_UINT;
+	}
+
+	return Format;
+}
+
 static bool IsFloatFormat(EPixelFormat Format)
 {
 	switch (Format)
diff --git a/Engine/Source/Runtime/RHI/Public/RHIValidation.h b/Engine/Source/Runtime/RHI/Public/RHIValidation.h
--- a/Engine/Source/Runtime/RHI/Public/RHIValidation.h
+++ b/Engine/Source/Runtime/RHI/Public/RHIValidation.h
@@ -404,6 +404,17 @@ public:
 		return UAV;
 	}
 
+	/** Creates an unordered access view of the given texture. */
+	// FlushType: Wait RHI Thread
+	virtual FUnorderedAccessViewRHIRef RHICreateUnorderedAccessView(FRHITexture* Texture, uint32 MipLevel, uint8 Format) override final
+	{
+		RHI_VALIDATION_CHECK(Texture->GetTextureReference() == nullptr, TEXT("Creating an unordered access view of an FRHITextureReference is not supported."));
+
+		FUnorderedAccessViewRHIRef UAV = RHI->RHICreateUnorderedAccessView(Texture, MipLevel, Format);
+		UAV->ViewIdentity = Texture->GetViewIdentity(MipLevel, 1, 0, 0, uint32(RHIValidation::EResourcePlane::Common), 1);
+		return UAV;
+	}
+
 	/** Creates an unordered access view of the given buffer. */
 	// FlushType: Wait RHI Thread
 	virtual FUnorderedAccessViewRHIRef RHICreateUnorderedAccessView(FRHIBuffer* Buffer, uint8 Format) override final
diff --git a/Engine/Source/Runtime/Renderer/Private/VT/VirtualTexturePhysicalSpace.cpp b/Engine/Source/Runtime/Renderer/Private/VT/VirtualTexturePhysicalSpace.cpp
--- a/Engine/Source/Runtime/Renderer/Private/VT/VirtualTexturePhysicalSpace.cpp
+++ b/Engine/Source/Runtime/Renderer/Private/VT/VirtualTexturePhysicalSpace.cpp
@@ -56,21 +56,13 @@ FVirtualTexturePhysicalSpace::~FVirtualTexturePhysicalSpace()
 {
 }
 
-EPixelFormat GetUnorderedAccessViewFormat(EPixelFormat InFormat)
+static EPixelFormat GetUnorderedAccessViewFormat(EPixelFormat InFormat)
 {
 	// Use alias formats for compressed textures on APIs where that is possible
 	// This allows us to compress runtime data directly to the physical texture
-	const bool bUAVAliasForCompressedTextures = GRHISupportsUAVFormatAliasing;
-
-	switch (InFormat)
+	if (IsBlockCompressedFormat(InFormat))
 	{
-	case PF_DXT1: 
-	case PF_BC4: 
-		return bUAVAliasForCompressedTextures ? PF_R32G32_UINT : PF_Unknown;
-	case PF_DXT3: 
-	case PF_DXT5: 
-	case PF_BC5: 
-		return bUAVAliasForCompressedTextures ? PF_R32G32B32A32_UINT : PF_Unknown;
+		return GRHISupportsUAVFormatAliasing ? GetBlockCompressedFormatUAVAliasFormat(InFormat) : PF_Unknown;
 	}
 
 	return InFormat;
