Disabling UAVFormatAliasing on Heap Tier 1 hardware until we figure out how to handle the limited heaps.

diff --git a/Engine/Source/Runtime/D3D12RHI/Private/Windows/WindowsD3D12Device.cpp b/Engine/Source/Runtime/D3D12RHI/Private/Windows/WindowsD3D12Device.cpp
--- a/Engine/Source/Runtime/D3D12RHI/Private/Windows/WindowsD3D12Device.cpp
+++ b/Engine/Source/Runtime/D3D12RHI/Private/Windows/WindowsD3D12Device.cpp
@@ -1129,7 +1129,7 @@ void FD3D12DynamicRHI::Init()
 		GRHIVariableRateShadingImageTileMaxHeight = 1;
 	}
 
-	GRHISupportsUAVFormatAliasing = true;
+	GRHISupportsUAVFormatAliasing = GetAdapter().GetResourceHeapTier() > D3D12_RESOURCE_HEAP_TIER_1;
 
 	// Command lists need the validation RHI context if enabled, so call the global scope version of RHIGetDefaultContext() and RHIGetDefaultAsyncComputeContext().
 	GRHICommandList.GetImmediateCommandList().SetContext(::RHIGetDefaultContext());
