Fix a leaking texture memory stat. It leaks because sizes of virtual textures are not deducted when they are released

diff --git a/Engine/Source/Runtime/D3D12RHI/Private/D3D12Texture.cpp b/Engine/Source/Runtime/D3D12RHI/Private/D3D12Texture.cpp
--- a/Engine/Source/Runtime/D3D12RHI/Private/D3D12Texture.cpp
+++ b/Engine/Source/Runtime/D3D12RHI/Private/D3D12Texture.cpp
@@ -385,8 +385,8 @@ void FD3D12TextureStats::D3D12TextureDeleted(TD3D12Texture2D<BaseResourceType>&
 
 	if (D3D12Texture2D)
 	{
-		// Don't update state for virtual or transient textures	
-		if(!EnumHasAnyFlags(Texture.Flags, TexCreate_Virtual | TexCreate_Transient))
+		// Don't update state for transient textures	
+		if(!EnumHasAnyFlags(Texture.Flags, TexCreate_Transient))
 		{
 			const D3D12_RESOURCE_DESC& Desc = D3D12Texture2D->GetDesc();
 			const int64 TextureSize = Texture.GetMemorySize();
@@ -442,8 +442,8 @@ void FD3D12TextureStats::D3D12TextureDeleted(FD3D12Texture3D& Texture)
 
 	if (D3D12Texture3D)
 	{
-		// Don't update state for virtual or transient textures	
-		if (!EnumHasAnyFlags(Texture.GetFlags(), TexCreate_Virtual | TexCreate_Transient))
+		// Don't update state for transient textures	
+		if (!EnumHasAnyFlags(Texture.GetFlags(), TexCreate_Transient))
 		{
 			const D3D12_RESOURCE_DESC& Desc = D3D12Texture3D->GetDesc();
 			const int64 TextureSize = Texture.GetMemorySize();
