Fixing template resolve order by moving CreateViewWithCounter back into FD3D12View.

diff --git a/Engine/Source/Runtime/D3D12RHI/Private/D3D12View.h b/Engine/Source/Runtime/D3D12RHI/Private/D3D12View.h
--- a/Engine/Source/Runtime/D3D12RHI/Private/D3D12View.h
+++ b/Engine/Source/Runtime/D3D12RHI/Private/D3D12View.h
@@ -753,6 +753,18 @@ protected:
 		}
 	}
 
+	void CreateViewWithCounter(const TDesc& InDesc, FD3D12BaseShaderResource* InBaseShaderResource, FD3D12ResourceLocation& InResourceLocation, FD3D12Resource* InCounterResource)
+	{
+		InitializeInternal(InDesc, InBaseShaderResource, InResourceLocation);
+
+		if (Resource)
+		{
+			ID3D12Resource* D3DResource = Resource->GetUAVAccessResource() ? Resource->GetUAVAccessResource() : Resource->GetResource();
+			ID3D12Resource* D3DCounterResource = InCounterResource ? InCounterResource->GetResource() : nullptr;
+			Descriptor.CreateViewWithCounter(Desc, D3DResource, D3DCounterResource);
+		}
+	}
+
 	virtual void ResourceRenamed(FD3D12BaseShaderResource* InRenamedResource, FD3D12ResourceLocation* InNewResourceLocation) override
 	{
 		check(InRenamedResource == BaseShaderResource);
@@ -968,19 +980,6 @@ public:
 		check(ResourceLocation->GetOffsetFromBaseOfResource() == 0);
 		Initialize(Desc, BaseShaderResource, *ResourceLocation);
 	}
-
-protected:
-	void CreateViewWithCounter(const D3D12_UNORDERED_ACCESS_VIEW_DESC& InDesc, FD3D12BaseShaderResource* InBaseShaderResource, FD3D12ResourceLocation& InResourceLocation, FD3D12Resource* InCounterResource)
-	{
-		InitializeInternal(InDesc, InBaseShaderResource, InResourceLocation);
-
-		if (Resource)
-		{
-			ID3D12Resource* D3DResource = Resource->GetUAVAccessResource() ? Resource->GetUAVAccessResource() : Resource->GetResource();
-			ID3D12Resource* D3DCounterResource = InCounterResource ? InCounterResource->GetResource() : nullptr;
-			Descriptor.CreateViewWithCounter(Desc, D3DResource, D3DCounterResource);
-		}
-	}
 };
 
 class FD3D12UnorderedAccessViewWithLocation : public FD3D12UnorderedAccessView
