Added descriptions to validation checks in FD3D12RayTracingScene::BuildAccelerationStructure()

diff --git a/Engine/Source/Runtime/D3D12RHI/Private/D3D12RayTracing.cpp b/Engine/Source/Runtime/D3D12RHI/Private/D3D12RayTracing.cpp
--- a/Engine/Source/Runtime/D3D12RHI/Private/D3D12RayTracing.cpp
+++ b/Engine/Source/Runtime/D3D12RHI/Private/D3D12RayTracing.cpp
@@ -3538,8 +3538,14 @@ void FD3D12RayTracingScene::BuildAccelerationStructure(FD3D12CommandContext& Com
 	D3D12_RAYTRACING_ACCELERATION_STRUCTURE_PREBUILD_INFO PrebuildInfo = {};
 
 	RayTracingDevice->GetRaytracingAccelerationStructurePrebuildInfo(&BuildInputs, &PrebuildInfo);
-	check(PrebuildInfo.ResultDataMaxSizeInBytes <= SizeInfo.ResultSize);
-	check(PrebuildInfo.ScratchDataSizeInBytes <= SizeInfo.BuildScratchSize);
+
+	checkf(PrebuildInfo.ResultDataMaxSizeInBytes <= SizeInfo.ResultSize,
+		TEXT("TLAS build result buffer now requires %lld bytes, but only %lld was calculated in the constructor."),
+			PrebuildInfo.ResultDataMaxSizeInBytes, SizeInfo.ResultSize);
+
+	checkf(PrebuildInfo.ScratchDataSizeInBytes <= SizeInfo.BuildScratchSize,
+		TEXT("TLAS build scratch buffer now requires %lld bytes, but only %lld was calculated in the constructor."),
+			PrebuildInfo.ScratchDataSizeInBytes, SizeInfo.BuildScratchSize);
 
 	TRefCountPtr<FD3D12Buffer> AutoScratchBuffer;
 	if (ScratchBuffer == nullptr)
@@ -3550,7 +3556,12 @@ void FD3D12RayTracingScene::BuildAccelerationStructure(FD3D12CommandContext& Com
 		ScratchBufferOffset = 0;
 	}
 
-	check(PrebuildInfo.ScratchDataSizeInBytes + ScratchBufferOffset <= ScratchBuffer->GetSize());
+	checkf(PrebuildInfo.ScratchDataSizeInBytes + ScratchBufferOffset <= ScratchBuffer->GetSize(),
+		TEXT("TLAS scratch buffer size is %lld bytes with offset %lld (%lld bytes available), but the build requires %lld bytes. ")
+		TEXT("BuildInputs.NumDescs = %d, Instances.Num = %d, PerInstanceGeometries.Num = %d."),
+			ScratchBuffer->GetSize(), ScratchBufferOffset, ScratchBuffer->GetSize() - ScratchBufferOffset, 
+			PrebuildInfo.ScratchDataSizeInBytes,
+			BuildInputs.NumDescs, Instances.Num(), PerInstanceGeometries.Num());
 
 	TRefCountPtr<FD3D12Buffer> AutoInstanceBuffer;
 
@@ -3710,6 +3721,8 @@ void FD3D12RayTracingScene::BuildAccelerationStructure(FD3D12CommandContext& Com
 	BuildDesc.ScratchAccelerationStructureData = ScratchBuffer->ResourceLocation.GetGPUVirtualAddress() + ScratchBufferOffset;
 	BuildDesc.SourceAccelerationStructureData = D3D12_GPU_VIRTUAL_ADDRESS(0); // Null source TLAS as this is a build command
 
+	checkf(BuildDesc.ScratchAccelerationStructureData % 256 == 0, TEXT("TLAS build scratch buffer must have 256 byte alignment."));
+
 	// UAV barrier is used here to ensure that all bottom level acceleration structures are built
 	CommandContext.CommandListHandle.AddUAVBarrier();
 	CommandContext.CommandListHandle.FlushResourceBarriers();
