XB1 transient allocator implementation that works with RDG. Courtesy to Kenzo for most of the implementation

diff --git a/Engine/Source/Runtime/D3D12RHI/Private/D3D12Adapter.h b/Engine/Source/Runtime/D3D12RHI/Private/D3D12Adapter.h
--- a/Engine/Source/Runtime/D3D12RHI/Private/D3D12Adapter.h
+++ b/Engine/Source/Runtime/D3D12RHI/Private/D3D12Adapter.h
@@ -315,6 +315,7 @@ public:
 		D3D12_RESOURCE_STATES InCreateState,
 		bool bHasInitialData,
 		const FRHIGPUMask& InGPUMask,
+		ED3D12ResourceTransientMode TransientMode,
 		ID3D12ResourceAllocator* ResourceAllocator,
 		const TCHAR* InDebugName);
 
@@ -401,6 +402,7 @@ protected:
 		uint32 Alignment,
 		FD3D12Buffer* Buffer,
 		FD3D12ResourceLocation& ResourceLocation,
+		ED3D12ResourceTransientMode TransientMode,
 		ID3D12ResourceAllocator* ResourceAllocator,
 		const TCHAR* InDebugName);
 
diff --git a/Engine/Source/Runtime/D3D12RHI/Private/D3D12Buffer.cpp b/Engine/Source/Runtime/D3D12RHI/Private/D3D12Buffer.cpp
--- a/Engine/Source/Runtime/D3D12RHI/Private/D3D12Buffer.cpp
+++ b/Engine/Source/Runtime/D3D12RHI/Private/D3D12Buffer.cpp
@@ -292,6 +292,7 @@ void FD3D12Adapter::AllocateBuffer(FD3D12Device* Device,
 	uint32 Alignment,
 	FD3D12Buffer* Buffer,
 	FD3D12ResourceLocation& ResourceLocation,
+	ED3D12ResourceTransientMode TransientMode,
 	ID3D12ResourceAllocator* ResourceAllocator,
 	const TCHAR* InDebugName)
 {
@@ -335,6 +336,7 @@ FD3D12Buffer* FD3D12Adapter::CreateRHIBuffer(
 	D3D12_RESOURCE_STATES InCreateState,
 	bool bHasInitialData,
 	const FRHIGPUMask& InGPUMask,
+	ED3D12ResourceTransientMode TransientMode,
 	ID3D12ResourceAllocator* ResourceAllocator,
 	const TCHAR* InDebugName)
 {
@@ -364,7 +366,7 @@ FD3D12Buffer* FD3D12Adapter::CreateRHIBuffer(
 
 			if (Device->GetGPUIndex() == FirstGPUIndex)
 			{
-				AllocateBuffer(Device, InDesc, Size, InUsage, InResourceStateMode, InCreateState, Alignment, NewBuffer, NewBuffer->ResourceLocation, ResourceAllocator, InDebugName);
+				AllocateBuffer(Device, InDesc, Size, InUsage, InResourceStateMode, InCreateState, Alignment, NewBuffer, NewBuffer->ResourceLocation, TransientMode, ResourceAllocator, InDebugName);
 				NewBuffer0 = NewBuffer;
 			}
 			else
@@ -389,7 +391,7 @@ FD3D12Buffer* FD3D12Adapter::CreateRHIBuffer(
 				NewBuffer->SetName(InDebugName);
 			}
 
-			AllocateBuffer(Device, InDesc, Size, InUsage, InResourceStateMode, InCreateState, Alignment, NewBuffer, NewBuffer->ResourceLocation, ResourceAllocator, InDebugName);
+			AllocateBuffer(Device, InDesc, Size, InUsage, InResourceStateMode, InCreateState, Alignment, NewBuffer, NewBuffer->ResourceLocation, TransientMode, ResourceAllocator, InDebugName);
 			
 			// Unlock immediately if there is no initial data
 			if (!bHasInitialData)
@@ -508,11 +510,10 @@ FBufferRHIRef FD3D12DynamicRHI::CreateBuffer(FRHICommandListImmediate* RHICmdLis
 			});
 	}
 
-	ID3D12ResourceAllocator* ResourceAllocator = nullptr;
-	return CreateD3D12Buffer(RHICmdList, Size, Usage, Stride, InResourceState, CreateInfo, ResourceAllocator);
+	return CreateD3D12Buffer(RHICmdList, Size, Usage, Stride, InResourceState, CreateInfo);
 }
 
-FD3D12Buffer* FD3D12DynamicRHI::CreateD3D12Buffer(class FRHICommandListImmediate* RHICmdList, uint32 Size, EBufferUsageFlags Usage, uint32 Stride, ERHIAccess InResourceState, FRHIResourceCreateInfo& CreateInfo, ID3D12ResourceAllocator* ResourceAllocator)
+FD3D12Buffer* FD3D12DynamicRHI::CreateD3D12Buffer(class FRHICommandListImmediate* RHICmdList, uint32 Size, EBufferUsageFlags Usage, uint32 Stride, ERHIAccess InResourceState, FRHIResourceCreateInfo& CreateInfo, ED3D12ResourceTransientMode TransientMode, ID3D12ResourceAllocator* ResourceAllocator)
 {
 	D3D12_RESOURCE_DESC Desc;
 	uint32 Alignment;
@@ -537,7 +538,7 @@ FD3D12Buffer* FD3D12DynamicRHI::CreateD3D12Buffer(class FRHICommandListImmediate
 	bool bHasInitialData = CreateInfo.ResourceArray != nullptr;
 
 	FD3D12SyncPoint CopyQueueSyncPoint;
-	FD3D12Buffer* Buffer = GetAdapter().CreateRHIBuffer(Desc, Alignment, Stride, Size, Usage, StateMode, CreateState, bHasInitialData, CreateInfo.GPUMask, ResourceAllocator, CreateInfo.DebugName);
+	FD3D12Buffer* Buffer = GetAdapter().CreateRHIBuffer(Desc, Alignment, Stride, Size, Usage, StateMode, CreateState, bHasInitialData, CreateInfo.GPUMask, TransientMode, ResourceAllocator, CreateInfo.DebugName);
 	check(Buffer->ResourceLocation.IsValid());
 
 	// Copy the resource data if available 
@@ -549,10 +550,10 @@ FD3D12Buffer* FD3D12DynamicRHI::CreateD3D12Buffer(class FRHICommandListImmediate
 	return Buffer;
 }
 
-FRHIBuffer* FD3D12DynamicRHI::CreateBuffer(const FRHIBufferCreateInfo& CreateInfo, const TCHAR* DebugName, ERHIAccess InitialState, ID3D12ResourceAllocator* ResourceAllocator)
+FRHIBuffer* FD3D12DynamicRHI::CreateBuffer(const FRHIBufferCreateInfo& CreateInfo, const TCHAR* DebugName, ERHIAccess InitialState, ED3D12ResourceTransientMode TransientMode, ID3D12ResourceAllocator* ResourceAllocator)
 {
 	FRHIResourceCreateInfo ResourceCreateInfo(DebugName);
-	return CreateD3D12Buffer(nullptr, CreateInfo.Size, CreateInfo.Usage, CreateInfo.Stride, InitialState, ResourceCreateInfo, ResourceAllocator);
+	return CreateD3D12Buffer(nullptr, CreateInfo.Size, CreateInfo.Usage, CreateInfo.Stride, InitialState, ResourceCreateInfo, TransientMode, ResourceAllocator);
 }
 
 void* FD3D12DynamicRHI::LockBuffer(FRHICommandListImmediate* RHICmdList, FD3D12Buffer* Buffer, uint32 BufferSize, uint32 BufferUsage, uint32 Offset, uint32 Size, EResourceLockMode LockMode)
diff --git a/Engine/Source/Runtime/D3D12RHI/Private/D3D12Commands.cpp b/Engine/Source/Runtime/D3D12RHI/Private/D3D12Commands.cpp
--- a/Engine/Source/Runtime/D3D12RHI/Private/D3D12Commands.cpp
+++ b/Engine/Source/Runtime/D3D12RHI/Private/D3D12Commands.cpp
@@ -25,13 +25,21 @@ static FAutoConsoleVariableRef CVarSyncTemporalResources(
 	ECVF_RenderThreadSafe
 	);
 
-static int32 GD3D12TransientAllocatorFullAliasingBarrier = 0;
+// TODO test and enable overlap tracking
+static int32 GD3D12TransientAllocatorFullAliasingBarrier = 1;
 static FAutoConsoleVariableRef CVarD3D12TransientAllocatorFullAliasingBarrier(
 	TEXT("d3d12.TransientAllocator.FullAliasingBarrier"),
 	GD3D12TransientAllocatorFullAliasingBarrier,
 	TEXT("Inserts a full aliasing barrier on an transient acquire operation. Useful to debug if an aliasing barrier is missing."),
 	ECVF_RenderThreadSafe);
 
+static int32 GD3D12AllowDiscardResources = 1;
+static FAutoConsoleVariableRef CVarD3D12AllowDiscardResources(
+	TEXT("d3d12.AllowDiscardResources"),
+	GD3D12AllowDiscardResources,
+	TEXT("Wheter to call DiscardResources after transient aliasing acquire. This is not needed on some platforms if newly acquired resources are cleared before use."),
+	ECVF_RenderThreadSafe);
+
 using namespace D3D12RHI;
 
 #define DECLARE_ISBOUNDSHADER(ShaderType) inline void ValidateBoundShader(FD3D12StateCache& InStateCache, FRHI##ShaderType* ShaderType##RHI) \
@@ -356,6 +364,11 @@ static void HandleDiscardResources(
 		HandleResourceDiscardTransitions(Context, Data, SkipFastClearEliminateState, ResourcesToDiscard);
 	}
 
+	if (!GD3D12AllowDiscardResources)
+	{
+		return;
+	}
+
 	Context.CommandListHandle.FlushResourceBarriers();
 
 	for (const FD3D12DiscardResource& DiscardResource : ResourcesToDiscard)
diff --git a/Engine/Source/Runtime/D3D12RHI/Private/D3D12RHIPrivate.h b/Engine/Source/Runtime/D3D12RHI/Private/D3D12RHIPrivate.h
--- a/Engine/Source/Runtime/D3D12RHI/Private/D3D12RHIPrivate.h
+++ b/Engine/Source/Runtime/D3D12RHI/Private/D3D12RHIPrivate.h
@@ -1055,16 +1055,17 @@ public:
 
 	template<typename BaseResourceType>
 	TD3D12Texture2D<BaseResourceType>* CreateD3D12Texture2D(class FRHICommandListImmediate* RHICmdList, uint32 SizeX, uint32 SizeY, uint32 SizeZ, bool bTextureArray, bool CubeTexture, EPixelFormat Format,
-		uint32 NumMips, uint32 NumSamples, ETextureCreateFlags InFlags, ERHIAccess InResourceState, FRHIResourceCreateInfo& CreateInfo, ID3D12ResourceAllocator* ResourceAllocator);
+		uint32 NumMips, uint32 NumSamples, ETextureCreateFlags InFlags, ERHIAccess InResourceState, FRHIResourceCreateInfo& CreateInfo, ED3D12ResourceTransientMode TransientMode = ED3D12ResourceTransientMode::NonTransient, ID3D12ResourceAllocator* ResourceAllocator = nullptr);
 
-	FD3D12Buffer* CreateD3D12Buffer(class FRHICommandListImmediate* RHICmdList, uint32 Size, EBufferUsageFlags Usage, uint32 Stride, ERHIAccess ResourceState, FRHIResourceCreateInfo& CreateInfo, ID3D12ResourceAllocator* ResourceAllocator);
+	FD3D12Buffer* CreateD3D12Buffer(class FRHICommandListImmediate* RHICmdList, uint32 Size, EBufferUsageFlags Usage, uint32 Stride, ERHIAccess ResourceState, FRHIResourceCreateInfo& CreateInfo, ED3D12ResourceTransientMode TransientMode = ED3D12ResourceTransientMode::NonTransient, ID3D12ResourceAllocator* ResourceAllocator = nullptr);
+
+	FD3D12Texture3D* CreateD3D12Texture3D(class FRHICommandListImmediate* RHICmdList, uint32 SizeX, uint32 SizeY, uint32 SizeZ, EPixelFormat Format, uint32 NumMips, ETextureCreateFlags InFlags, ERHIAccess InResourceState, FRHIResourceCreateInfo& CreateInfo, ED3D12ResourceTransientMode TransientMode = ED3D12ResourceTransientMode::NonTransient, ID3D12ResourceAllocator* ResourceAllocator = nullptr);
 
-	FD3D12Texture3D* CreateD3D12Texture3D(class FRHICommandListImmediate* RHICmdList, uint32 SizeX, uint32 SizeY, uint32 SizeZ, EPixelFormat Format, uint32 NumMips, ETextureCreateFlags InFlags, ERHIAccess InResourceState, FRHIResourceCreateInfo& CreateInfo, ID3D12ResourceAllocator* ResourceAllocator);
 
 	// Interface for FD3D12TransientResourceAllocator.
 	virtual D3D12_RESOURCE_DESC GetResourceDesc(const FRHITextureCreateInfo& CreateInfo) const;
-	virtual FRHITexture* CreateTexture(const FRHITextureCreateInfo& CreateInfo, const TCHAR* DebugName, ERHIAccess InitialState, ID3D12ResourceAllocator* ResourceAllocator);
-	FRHIBuffer* CreateBuffer(const FRHIBufferCreateInfo& CreateInfo, const TCHAR* DebugName, ERHIAccess InitialState, ID3D12ResourceAllocator* ResourceAllocator);
+	virtual FRHITexture* CreateTexture(const FRHITextureCreateInfo& CreateInfo, const TCHAR* DebugName, ERHIAccess InitialState, ED3D12ResourceTransientMode AllocationMode, ID3D12ResourceAllocator* ResourceAllocator);
+	FRHIBuffer* CreateBuffer(const FRHIBufferCreateInfo& CreateInfo, const TCHAR* DebugName, ERHIAccess InitialState, ED3D12ResourceTransientMode TransientMode, ID3D12ResourceAllocator* ResourceAllocator);
 
 protected:
 
diff --git a/Engine/Source/Runtime/D3D12RHI/Private/D3D12RayTracing.cpp b/Engine/Source/Runtime/D3D12RHI/Private/D3D12RayTracing.cpp
--- a/Engine/Source/Runtime/D3D12RHI/Private/D3D12RayTracing.cpp
+++ b/Engine/Source/Runtime/D3D12RHI/Private/D3D12RayTracing.cpp
@@ -608,9 +608,10 @@ FD3D12RayTracingCompactionRequestHandler::FD3D12RayTracingCompactionRequestHandl
 
 	FRHIGPUMask GPUMask = FRHIGPUMask::FromIndex(GetParentDevice()->GetGPUIndex());
 	ID3D12ResourceAllocator* ResourceAllocator = nullptr;
+	ED3D12ResourceTransientMode TransientMode = ED3D12ResourceTransientMode::NonTransient;
 	bool bHasInitialData = false;
 	PostBuildInfoBuffer = GetParentDevice()->GetParentAdapter()->CreateRHIBuffer(PostBuildInfoBufferDesc, 8,
-		0, PostBuildInfoBufferDesc.Width, BUF_UnorderedAccess | BUF_SourceCopy, ED3D12ResourceStateMode::MultiState, D3D12_RESOURCE_STATE_UNORDERED_ACCESS, bHasInitialData, GPUMask, ResourceAllocator, TEXT("PostBuildInfoBuffer"));
+		0, PostBuildInfoBufferDesc.Width, BUF_UnorderedAccess | BUF_SourceCopy, ED3D12ResourceStateMode::MultiState, D3D12_RESOURCE_STATE_UNORDERED_ACCESS, bHasInitialData, GPUMask, TransientMode, ResourceAllocator, TEXT("PostBuildInfoBuffer"));
 	SetName(PostBuildInfoBuffer->GetResource(), TEXT("PostBuildInfoBuffer"));
 
 	PostBuildInfoStagingBuffer = RHICreateStagingBuffer();
@@ -1762,9 +1763,10 @@ public:
 		bool bHasInitialData = true;
 
 		ID3D12ResourceAllocator* ResourceAllocator = nullptr;
+		ED3D12ResourceTransientMode TransientMode = ED3D12ResourceTransientMode::NonTransient;
 		Buffer = Adapter->CreateRHIBuffer(
 			BufferDesc, BufferDesc.Alignment, 0, BufferDesc.Width, BUF_Static, ED3D12ResourceStateMode::MultiState,
-			D3D12_RESOURCE_STATE_COPY_DEST, bHasInitialData, GPUMask, ResourceAllocator, TEXT("Shader binding table"));
+			D3D12_RESOURCE_STATE_COPY_DEST, bHasInitialData, GPUMask, TransientMode, ResourceAllocator, TEXT("Shader binding table"));
 
 		// Use copy queue for uploading the data
 		FD3D12SyncPoint CopyQueueSyncPoint = Buffer->UploadResourceDataViaCopyQueue(&Data);
@@ -2848,6 +2850,7 @@ static TRefCountPtr<FD3D12Buffer> CreateRayTracingBuffer(FD3D12Adapter* Adapter,
 	TRefCountPtr<FD3D12Buffer> Result;
 
 	FString DebugNameString = DebugName.ToString();
+	ED3D12ResourceTransientMode TransientMode = ED3D12ResourceTransientMode::NonTransient;
 	ID3D12ResourceAllocator* ResourceAllocator = nullptr;
 	D3D12_RESOURCE_DESC BufferDesc = CD3DX12_RESOURCE_DESC::Buffer(Size, D3D12_RESOURCE_FLAG_ALLOW_UNORDERED_ACCESS);
 	FRHIGPUMask GPUMask = FRHIGPUMask::FromIndex(GPUIndex);
@@ -2859,7 +2862,7 @@ static TRefCountPtr<FD3D12Buffer> CreateRayTracingBuffer(FD3D12Adapter* Adapter,
 			BufferDesc, D3D12_RAYTRACING_ACCELERATION_STRUCTURE_BYTE_ALIGNMENT,
 			0, BufferDesc.Width, BUF_AccelerationStructure,
 			ED3D12ResourceStateMode::SingleState, D3D12_RESOURCE_STATE_RAYTRACING_ACCELERATION_STRUCTURE, bHasInitialData,
-			GPUMask, ResourceAllocator, *DebugNameString);
+			GPUMask, TransientMode, ResourceAllocator, *DebugNameString);
 	}
 	else if (Type == ERayTracingBufferType::Scratch)
 	{
@@ -2867,7 +2870,7 @@ static TRefCountPtr<FD3D12Buffer> CreateRayTracingBuffer(FD3D12Adapter* Adapter,
 			BufferDesc, D3D12_RAYTRACING_ACCELERATION_STRUCTURE_BYTE_ALIGNMENT,
 			0, BufferDesc.Width, BUF_UnorderedAccess,
 			ED3D12ResourceStateMode::SingleState, D3D12_RESOURCE_STATE_UNORDERED_ACCESS, bHasInitialData,
-			GPUMask, ResourceAllocator, *DebugNameString);
+			GPUMask, TransientMode, ResourceAllocator, *DebugNameString);
 
 		// Elevates the scratch buffer heap priority, which may help performance / stability in low memory conditions 
 		// (Acceleration structure already boosted from allocation side)
@@ -3606,13 +3609,15 @@ void FD3D12RayTracingScene::BuildAccelerationStructure(FD3D12CommandContext& Com
 		// #dxr_todo multi state only when bShouldCopyIndirectInstances is set - will still need transition to copy_src for lock behind but this can be done on the complete pool in theory (have to check cost)
 		// #dxr_todo use dynamic with upload memory when no CopyCommands requested - resource will be in upload memory then instead of VRAM
 
+		ED3D12ResourceTransientMode TransientMode = ED3D12ResourceTransientMode::NonTransient;
 		ID3D12ResourceAllocator* ResourceAllocator = nullptr;
 		FRHIGPUMask GPUMask = FRHIGPUMask::FromIndex(GPUIndex);
 		bool bHasInitialData = true;
 
 		AutoInstanceBuffer = Adapter->CreateRHIBuffer(
 			InstanceBufferDesc, D3D12_RAYTRACING_INSTANCE_DESCS_BYTE_ALIGNMENT,
-			sizeof(D3D12_RAYTRACING_INSTANCE_DESC), InstanceBufferDesc.Width, BUF_UnorderedAccess, ED3D12ResourceStateMode::MultiState, D3D12_RESOURCE_STATE_COPY_DEST, bHasInitialData, GPUMask, ResourceAllocator, TEXT("InstanceBuffer"));
+			sizeof(D3D12_RAYTRACING_INSTANCE_DESC), InstanceBufferDesc.Width, BUF_UnorderedAccess, ED3D12ResourceStateMode::MultiState, D3D12_RESOURCE_STATE_COPY_DEST, 
+			bHasInitialData, GPUMask, TransientMode, ResourceAllocator, TEXT("InstanceBuffer"));
 
 		// Use copy queue for uploading the data
 		FD3D12SyncPoint CopyQueueSyncPoint = AutoInstanceBuffer->UploadResourceDataViaCopyQueue(&Instances);
diff --git a/Engine/Source/Runtime/D3D12RHI/Private/D3D12Texture.cpp b/Engine/Source/Runtime/D3D12RHI/Private/D3D12Texture.cpp
--- a/Engine/Source/Runtime/D3D12RHI/Private/D3D12Texture.cpp
+++ b/Engine/Source/Runtime/D3D12RHI/Private/D3D12Texture.cpp
@@ -855,7 +855,7 @@ static void DetermineTexture2DResourceFlagsAndLayout(uint32 SizeX, uint32 SizeY,
 
 template<typename BaseResourceType>
 TD3D12Texture2D<BaseResourceType>* FD3D12DynamicRHI::CreateD3D12Texture2D(FRHICommandListImmediate* RHICmdList, uint32 SizeX, uint32 SizeY, uint32 SizeZ, bool bTextureArray, bool bCubeTexture, EPixelFormat Format,
-	uint32 NumMips, uint32 NumSamples, ETextureCreateFlags Flags, ERHIAccess InResourceState, FRHIResourceCreateInfo& CreateInfo, ID3D12ResourceAllocator* ResourceAllocator)
+	uint32 NumMips, uint32 NumSamples, ETextureCreateFlags Flags, ERHIAccess InResourceState, FRHIResourceCreateInfo& CreateInfo, ED3D12ResourceTransientMode TransientMode, ID3D12ResourceAllocator* ResourceAllocator)
 {
 #if PLATFORM_WINDOWS || PLATFORM_HOLOLENS
 	
@@ -1215,7 +1215,7 @@ TD3D12Texture2D<BaseResourceType>* FD3D12DynamicRHI::CreateD3D12Texture2D(FRHICo
 #endif // PLATFORM_WINDOWS || PLATFORM_HOLOLENS
 }
 
-FD3D12Texture3D* FD3D12DynamicRHI::CreateD3D12Texture3D(FRHICommandListImmediate* RHICmdList, uint32 SizeX, uint32 SizeY, uint32 SizeZ, EPixelFormat Format, uint32 NumMips, ETextureCreateFlags Flags, ERHIAccess InResourceState, FRHIResourceCreateInfo& CreateInfo, ID3D12ResourceAllocator* ResourceAllocator)
+FD3D12Texture3D* FD3D12DynamicRHI::CreateD3D12Texture3D(FRHICommandListImmediate* RHICmdList, uint32 SizeX, uint32 SizeY, uint32 SizeZ, EPixelFormat Format, uint32 NumMips, ETextureCreateFlags Flags, ERHIAccess InResourceState, FRHIResourceCreateInfo& CreateInfo, ED3D12ResourceTransientMode TransientMode, ID3D12ResourceAllocator* ResourceAllocator)
 {
 #if PLATFORM_WINDOWS || PLATFORM_HOLOLENS
 	SCOPE_CYCLE_COUNTER(STAT_D3D12CreateTextureTime);
@@ -1339,7 +1339,7 @@ FD3D12Texture3D* FD3D12DynamicRHI::CreateD3D12Texture3D(FRHICommandListImmediate
 #endif // PLATFORM_WINDOWS || PLATFORM_HOLOLENS
 }
 
-FRHITexture* FD3D12DynamicRHI::CreateTexture(const FRHITextureCreateInfo& CreateInfo, const TCHAR* DebugName, ERHIAccess InitialState, ID3D12ResourceAllocator* ResourceAllocator)
+FRHITexture* FD3D12DynamicRHI::CreateTexture(const FRHITextureCreateInfo& CreateInfo, const TCHAR* DebugName, ERHIAccess InitialState, ED3D12ResourceTransientMode TransientMode, ID3D12ResourceAllocator* ResourceAllocator)
 {
 	FRHIResourceCreateInfo ResourceCreateInfo(DebugName, CreateInfo.ClearValue);
 
@@ -1349,17 +1349,17 @@ FRHITexture* FD3D12DynamicRHI::CreateTexture(const FRHITextureCreateInfo& Create
 	switch (CreateInfo.Dimension)
 	{
 	case ETextureDimension::Texture2D:
-		return CreateD3D12Texture2D<FD3D12BaseTexture2D>(nullptr, CreateInfo.Extent.X, CreateInfo.Extent.Y, 1, bTextureArray, bTextureCube, CreateInfo.Format, CreateInfo.NumMips, CreateInfo.NumSamples, CreateInfo.Flags, InitialState, ResourceCreateInfo, ResourceAllocator);
+		return CreateD3D12Texture2D<FD3D12BaseTexture2D>(nullptr, CreateInfo.Extent.X, CreateInfo.Extent.Y, 1, bTextureArray, bTextureCube, CreateInfo.Format, CreateInfo.NumMips, CreateInfo.NumSamples, CreateInfo.Flags, InitialState, ResourceCreateInfo, TransientMode, ResourceAllocator);
 
 	case ETextureDimension::Texture2DArray:
-		return CreateD3D12Texture2D<FD3D12BaseTexture2DArray>(nullptr, CreateInfo.Extent.X, CreateInfo.Extent.Y, CreateInfo.ArraySize, bTextureArray, bTextureCube, CreateInfo.Format, CreateInfo.NumMips, CreateInfo.NumSamples, CreateInfo.Flags, InitialState, ResourceCreateInfo, ResourceAllocator);
+		return CreateD3D12Texture2D<FD3D12BaseTexture2DArray>(nullptr, CreateInfo.Extent.X, CreateInfo.Extent.Y, CreateInfo.ArraySize, bTextureArray, bTextureCube, CreateInfo.Format, CreateInfo.NumMips, CreateInfo.NumSamples, CreateInfo.Flags, InitialState, ResourceCreateInfo, TransientMode, ResourceAllocator);
 
 	case ETextureDimension::TextureCube:
 	case ETextureDimension::TextureCubeArray:
-		return CreateD3D12Texture2D<FD3D12BaseTextureCube>(nullptr, CreateInfo.Extent.X, CreateInfo.Extent.Y, 6 * CreateInfo.ArraySize, bTextureArray, bTextureCube, CreateInfo.Format, CreateInfo.NumMips, CreateInfo.NumSamples, CreateInfo.Flags, InitialState, ResourceCreateInfo, ResourceAllocator);
+		return CreateD3D12Texture2D<FD3D12BaseTextureCube>(nullptr, CreateInfo.Extent.X, CreateInfo.Extent.Y, 6 * CreateInfo.ArraySize, bTextureArray, bTextureCube, CreateInfo.Format, CreateInfo.NumMips, CreateInfo.NumSamples, CreateInfo.Flags, InitialState, ResourceCreateInfo, TransientMode, ResourceAllocator);
 
 	case ETextureDimension::Texture3D:
-		return CreateD3D12Texture3D(nullptr, CreateInfo.Extent.X, CreateInfo.Extent.Y, CreateInfo.Depth, CreateInfo.Format, CreateInfo.NumMips, CreateInfo.Flags, InitialState, ResourceCreateInfo, ResourceAllocator);
+		return CreateD3D12Texture3D(nullptr, CreateInfo.Extent.X, CreateInfo.Extent.Y, CreateInfo.Depth, CreateInfo.Format, CreateInfo.NumMips, CreateInfo.Flags, InitialState, ResourceCreateInfo, TransientMode, ResourceAllocator);
 
 	default:
 		checkNoEntry();
@@ -1374,14 +1374,12 @@ FRHITexture* FD3D12DynamicRHI::CreateTexture(const FRHITextureCreateInfo& Create
 
 FTexture2DRHIRef FD3D12DynamicRHI::RHICreateTexture2D_RenderThread(class FRHICommandListImmediate& RHICmdList, uint32 SizeX, uint32 SizeY, uint8 Format, uint32 NumMips, uint32 NumSamples, ETextureCreateFlags Flags, ERHIAccess InResourceState, FRHIResourceCreateInfo& CreateInfo)
 {
-	ID3D12ResourceAllocator* ResourceAllocator = nullptr;
-	return CreateD3D12Texture2D<FD3D12BaseTexture2D>(&RHICmdList, SizeX, SizeY, 1, false, false, (EPixelFormat) Format, NumMips, NumSamples, Flags, InResourceState, CreateInfo, ResourceAllocator);
+	return CreateD3D12Texture2D<FD3D12BaseTexture2D>(&RHICmdList, SizeX, SizeY, 1, false, false, (EPixelFormat) Format, NumMips, NumSamples, Flags, InResourceState, CreateInfo);
 }
 
 FTexture2DRHIRef FD3D12DynamicRHI::RHICreateTexture2D(uint32 SizeX, uint32 SizeY, uint8 Format, uint32 NumMips, uint32 NumSamples, ETextureCreateFlags Flags, ERHIAccess InResourceState, FRHIResourceCreateInfo& CreateInfo)
 {
-	ID3D12ResourceAllocator* ResourceAllocator = nullptr;
-	return CreateD3D12Texture2D<FD3D12BaseTexture2D>(nullptr, SizeX, SizeY, 1, false, false, (EPixelFormat) Format, NumMips, NumSamples, Flags, InResourceState, CreateInfo, ResourceAllocator);
+	return CreateD3D12Texture2D<FD3D12BaseTexture2D>(nullptr, SizeX, SizeY, 1, false, false, (EPixelFormat) Format, NumMips, NumSamples, Flags, InResourceState, CreateInfo);
 }
 
 FTexture2DRHIRef FD3D12DynamicRHI::RHIAsyncCreateTexture2D(uint32 SizeX, uint32 SizeY, uint8 Format, uint32 NumMips, ETextureCreateFlags Flags, ERHIAccess InResourceState, void** InitialMipData, uint32 NumInitialMips)
@@ -1656,30 +1654,26 @@ FTexture2DArrayRHIRef FD3D12DynamicRHI::RHICreateTexture2DArray_RenderThread(cla
 {
 	check(SizeZ >= 1);
 
-	ID3D12ResourceAllocator* ResourceAllocator = nullptr;
-	return CreateD3D12Texture2D<FD3D12BaseTexture2DArray>(&RHICmdList, SizeX, SizeY, SizeZ, true, false, (EPixelFormat) Format, NumMips, NumSamples, Flags, InResourceState, CreateInfo, ResourceAllocator);
+	return CreateD3D12Texture2D<FD3D12BaseTexture2DArray>(&RHICmdList, SizeX, SizeY, SizeZ, true, false, (EPixelFormat) Format, NumMips, NumSamples, Flags, InResourceState, CreateInfo);
 }
 
 FTexture2DArrayRHIRef FD3D12DynamicRHI::RHICreateTexture2DArray(uint32 SizeX, uint32 SizeY, uint32 SizeZ, uint8 Format, uint32 NumMips, uint32 NumSamples, ETextureCreateFlags Flags, ERHIAccess InResourceState, FRHIResourceCreateInfo& CreateInfo)
 {
 	check(SizeZ >= 1);
 
-	ID3D12ResourceAllocator* ResourceAllocator = nullptr;
-	return CreateD3D12Texture2D<FD3D12BaseTexture2DArray>(nullptr, SizeX, SizeY, SizeZ, true, false, (EPixelFormat) Format, NumMips, NumSamples, Flags, InResourceState, CreateInfo, ResourceAllocator);
+	return CreateD3D12Texture2D<FD3D12BaseTexture2DArray>(nullptr, SizeX, SizeY, SizeZ, true, false, (EPixelFormat) Format, NumMips, NumSamples, Flags, InResourceState, CreateInfo);
 }
 
 FTexture3DRHIRef FD3D12DynamicRHI::RHICreateTexture3D_RenderThread(class FRHICommandListImmediate& RHICmdList, uint32 SizeX, uint32 SizeY, uint32 SizeZ, uint8 Format, uint32 NumMips, ETextureCreateFlags Flags, ERHIAccess InResourceState, FRHIResourceCreateInfo& CreateInfo)
 {
-	ID3D12ResourceAllocator* ResourceAllocator = nullptr;
-	return CreateD3D12Texture3D(&RHICmdList, SizeX, SizeY, SizeZ, (EPixelFormat) Format, NumMips, Flags, InResourceState, CreateInfo, ResourceAllocator);
+	return CreateD3D12Texture3D(&RHICmdList, SizeX, SizeY, SizeZ, (EPixelFormat) Format, NumMips, Flags, InResourceState, CreateInfo);
 }
 
 FTexture3DRHIRef FD3D12DynamicRHI::RHICreateTexture3D(uint32 SizeX, uint32 SizeY, uint32 SizeZ, uint8 Format, uint32 NumMips, ETextureCreateFlags Flags, ERHIAccess InResourceState, FRHIResourceCreateInfo& CreateInfo)
 {
 	check(SizeZ >= 1);
 #if PLATFORM_WINDOWS || PLATFORM_HOLOLENS
-	ID3D12ResourceAllocator* ResourceAllocator = nullptr;
-	return CreateD3D12Texture3D(nullptr, SizeX, SizeY, SizeZ, (EPixelFormat) Format, NumMips, Flags, InResourceState, CreateInfo, ResourceAllocator);
+	return CreateD3D12Texture3D(nullptr, SizeX, SizeY, SizeZ, (EPixelFormat) Format, NumMips, Flags, InResourceState, CreateInfo);
 #else
 	checkf(false, TEXT("XBOX_CODE_MERGE : Removed. The Xbox platform version should be used."));
 	return nullptr;
@@ -1780,8 +1774,7 @@ FTexture2DRHIRef FD3D12DynamicRHI::AsyncReallocateTexture2D_RenderThread(class F
 	// Allocate a new texture.
 	FRHIResourceCreateInfo CreateInfo(TEXT("AsyncReallocateTexture2D_RenderThread"));
 	ERHIAccess RHIAccess = ERHIAccess::Unknown;
-	ID3D12ResourceAllocator* ResourceAllocator = nullptr;
-	FD3D12Texture2D* NewTexture2D = CreateD3D12Texture2D<FD3D12BaseTexture2D>(nullptr, NewSizeX, NewSizeY, 1, false, false, Texture2DRHI->GetFormat(), NewMipCount, 1, Texture2DRHI->GetFlags(), RHIAccess, CreateInfo, ResourceAllocator);
+	FD3D12Texture2D* NewTexture2D = CreateD3D12Texture2D<FD3D12BaseTexture2D>(nullptr, NewSizeX, NewSizeY, 1, false, false, Texture2DRHI->GetFormat(), NewMipCount, 1, Texture2DRHI->GetFlags(), RHIAccess, CreateInfo);
 	
 	ALLOC_COMMAND_CL(RHICmdList, FRHICommandD3D12AsyncReallocateTexture2D)(Texture2D, NewTexture2D, NewMipCount, NewSizeX, NewSizeY, RequestStatus);
 
@@ -1811,8 +1804,7 @@ FTexture2DRHIRef FD3D12DynamicRHI::RHIAsyncReallocateTexture2D(FRHITexture2D* Te
 	// Allocate a new texture.
 	FRHIResourceCreateInfo CreateInfo(TEXT("RHIAsyncReallocateTexture2D"));
 	ERHIAccess RHIAccess = ERHIAccess::Unknown;
-	ID3D12ResourceAllocator* ResourceAllocator = nullptr;
-	FD3D12Texture2D* NewTexture2D = CreateD3D12Texture2D<FD3D12BaseTexture2D>(nullptr, NewSizeX, NewSizeY, 1, false, false, Texture2DRHI->GetFormat(), NewMipCount, 1, Texture2DRHI->GetFlags(), RHIAccess, CreateInfo, ResourceAllocator);
+	FD3D12Texture2D* NewTexture2D = CreateD3D12Texture2D<FD3D12BaseTexture2D>(nullptr, NewSizeX, NewSizeY, 1, false, false, Texture2DRHI->GetFormat(), NewMipCount, 1, Texture2DRHI->GetFlags(), RHIAccess, CreateInfo);
 	
 	DoAsyncReallocateTexture2D(Texture2D, NewTexture2D, NewMipCount, NewSizeX, NewSizeY, RequestStatus);
 
@@ -2884,26 +2876,22 @@ void FD3D12DynamicRHI::EndUpdateTexture3D_Internal(FUpdateTexture3DData& UpdateD
 	-----------------------------------------------------------------------------*/
 FTextureCubeRHIRef FD3D12DynamicRHI::RHICreateTextureCube_RenderThread(class FRHICommandListImmediate& RHICmdList, uint32 Size, uint8 Format, uint32 NumMips, ETextureCreateFlags Flags, ERHIAccess InResourceState, FRHIResourceCreateInfo& CreateInfo)
 {
-	ID3D12ResourceAllocator* ResourceAllocator = nullptr;
-	return CreateD3D12Texture2D<FD3D12BaseTextureCube>(&RHICmdList, Size, Size, 6, false, true, (EPixelFormat) Format, NumMips, 1, Flags, InResourceState, CreateInfo, ResourceAllocator);
+	return CreateD3D12Texture2D<FD3D12BaseTextureCube>(&RHICmdList, Size, Size, 6, false, true, (EPixelFormat) Format, NumMips, 1, Flags, InResourceState, CreateInfo);
 }
 
 FTextureCubeRHIRef FD3D12DynamicRHI::RHICreateTextureCube(uint32 Size, uint8 Format, uint32 NumMips, ETextureCreateFlags Flags, ERHIAccess InResourceState, FRHIResourceCreateInfo& CreateInfo)
 {
-	ID3D12ResourceAllocator* ResourceAllocator = nullptr;
-	return CreateD3D12Texture2D<FD3D12BaseTextureCube>(nullptr, Size, Size, 6, false, true, (EPixelFormat) Format, NumMips, 1, Flags, InResourceState, CreateInfo, ResourceAllocator);
+	return CreateD3D12Texture2D<FD3D12BaseTextureCube>(nullptr, Size, Size, 6, false, true, (EPixelFormat) Format, NumMips, 1, Flags, InResourceState, CreateInfo);
 }
 
 FTextureCubeRHIRef FD3D12DynamicRHI::RHICreateTextureCubeArray_RenderThread(class FRHICommandListImmediate& RHICmdList, uint32 Size, uint32 ArraySize, uint8 Format, uint32 NumMips, ETextureCreateFlags Flags, ERHIAccess InResourceState, FRHIResourceCreateInfo& CreateInfo)
 {
-	ID3D12ResourceAllocator* ResourceAllocator = nullptr;
-	return CreateD3D12Texture2D<FD3D12BaseTextureCube>(&RHICmdList, Size, Size, 6 * ArraySize, true, true, (EPixelFormat) Format, NumMips, 1, Flags, InResourceState, CreateInfo, ResourceAllocator);
+	return CreateD3D12Texture2D<FD3D12BaseTextureCube>(&RHICmdList, Size, Size, 6 * ArraySize, true, true, (EPixelFormat) Format, NumMips, 1, Flags, InResourceState, CreateInfo);
 }
 
 FTextureCubeRHIRef FD3D12DynamicRHI::RHICreateTextureCubeArray(uint32 Size, uint32 ArraySize, uint8 Format, uint32 NumMips, ETextureCreateFlags Flags, ERHIAccess InResourceState, FRHIResourceCreateInfo& CreateInfo)
 {
-	ID3D12ResourceAllocator* ResourceAllocator = nullptr;
-	return CreateD3D12Texture2D<FD3D12BaseTextureCube>(nullptr, Size, Size, 6 * ArraySize, true, true, (EPixelFormat) Format, NumMips, 1, Flags, InResourceState, CreateInfo, ResourceAllocator);
+	return CreateD3D12Texture2D<FD3D12BaseTextureCube>(nullptr, Size, Size, 6 * ArraySize, true, true, (EPixelFormat) Format, NumMips, 1, Flags, InResourceState, CreateInfo);
 }
 
 void* FD3D12DynamicRHI::RHILockTextureCubeFace(FRHITextureCube* TextureCubeRHI, uint32 FaceIndex, uint32 ArrayIndex, uint32 MipIndex, EResourceLockMode LockMode, uint32& DestStride, bool bLockWithinMiptail)
diff --git a/Engine/Source/Runtime/D3D12RHI/Private/D3D12TransientResourceAllocator.cpp b/Engine/Source/Runtime/D3D12RHI/Private/D3D12TransientResourceAllocator.cpp
--- a/Engine/Source/Runtime/D3D12RHI/Private/D3D12TransientResourceAllocator.cpp
+++ b/Engine/Source/Runtime/D3D12RHI/Private/D3D12TransientResourceAllocator.cpp
@@ -145,8 +145,10 @@ FRHITransientTexture* FD3D12TransientResourceAllocator::CreateTexture(const FRHI
 			InitialState = ERHIAccess::DSVWrite;
 		}
 
+		ED3D12ResourceTransientMode TransientMode = ED3D12ResourceTransientMode::Transient;
 		FResourceAllocatorAdapter ResourceAllocatorAdapter(GetParentAdapter(), static_cast<FD3D12TransientHeap&>(Initializer.Heap), Initializer.Allocation, Desc);
-		FRHITexture* Texture = DynamicRHI->CreateTexture(InCreateInfo, InDebugName, InitialState, &ResourceAllocatorAdapter);
+		
+		FRHITexture* Texture = DynamicRHI->CreateTexture(InCreateInfo, InDebugName, InitialState, TransientMode, &ResourceAllocatorAdapter);
 		return new FRHITransientTexture(Texture, Initializer.Hash, InCreateInfo);
 	});
 }
@@ -195,8 +197,9 @@ FRHITransientBuffer* FD3D12TransientResourceAllocator::CreateBuffer(const FRHIBu
 	return Allocator.CreateBuffer(InCreateInfo, InDebugName, Size, Alignment,
 		[&](const FRHITransientResourceAllocator::FResourceInitializer& Initializer)
 	{
+		ED3D12ResourceTransientMode TransientMode = ED3D12ResourceTransientMode::Transient;
 		FResourceAllocatorAdapter ResourceAllocatorAdapter(GetParentAdapter(), static_cast<FD3D12TransientHeap&>(Initializer.Heap), Initializer.Allocation, Desc);
-		FRHIBuffer* Buffer = FD3D12DynamicRHI::GetD3DRHI()->CreateBuffer(InCreateInfo, InDebugName, ERHIAccess::UAVMask, &ResourceAllocatorAdapter);
+		FRHIBuffer* Buffer = FD3D12DynamicRHI::GetD3DRHI()->CreateBuffer(InCreateInfo, InDebugName, ERHIAccess::UAVMask, TransientMode, &ResourceAllocatorAdapter);
 		return new FRHITransientBuffer(Buffer, Initializer.Hash, InCreateInfo);
 	});
 }
diff --git a/Engine/Source/Runtime/D3D12RHI/Public/D3D12Resources.h b/Engine/Source/Runtime/D3D12RHI/Public/D3D12Resources.h
--- a/Engine/Source/Runtime/D3D12RHI/Public/D3D12Resources.h
+++ b/Engine/Source/Runtime/D3D12RHI/Public/D3D12Resources.h
@@ -42,6 +42,12 @@ enum class ED3D12ResourceStateMode
 	MultiState,					//< Force enable state tracking of resource
 };
 
+enum class ED3D12ResourceTransientMode
+{
+	NonTransient,				//< Resource is not transient
+	Transient,					//< Resource is transient
+};
+
 class FD3D12PendingResourceBarrier
 {
 public:
@@ -840,8 +846,8 @@ public:
 
 	void Swap(FD3D12BaseShaderResource& Other)
 	{
-		// assume RHI thread when swapping listeners and resources
-		check(IsInRHIThread());
+		// assume RHI thread when swapping listeners and resources unless RHI thread is disabled
+		check(IsInRHIThread() || IsInRenderingThread());
 
 		::Swap(Parent, Other.Parent);
 		ResourceLocation.Swap(Other.ResourceLocation);
diff --git a/Engine/Source/Runtime/RenderCore/Private/RenderGraphBuilder.cpp b/Engine/Source/Runtime/RenderCore/Private/RenderGraphBuilder.cpp
--- a/Engine/Source/Runtime/RenderCore/Private/RenderGraphBuilder.cpp
+++ b/Engine/Source/Runtime/RenderCore/Private/RenderGraphBuilder.cpp
@@ -419,12 +419,17 @@ bool FRDGBuilder::IsTransient(FRDGBufferRef Buffer) const
 		return false;
 	}
 
+	if (GRDGTransientAllocator == 2 && !EnumHasAnyFlags(Buffer->Desc.Usage, BUF_FastVRAM))
+	{
+		return false;
+	}
+
 	return EnumHasAnyFlags(Buffer->Desc.Usage, BUF_UnorderedAccess);
 }
 
 bool FRDGBuilder::IsTransient(FRDGTextureRef Texture) const
 {
-	return IsTransientInternal(Texture);
+	return IsTransientInternal(Texture) && (GRDGTransientAllocator != 2 || EnumHasAnyFlags(Texture->Desc.Flags, TexCreate_FastVRAM));
 }
 
 bool FRDGBuilder::IsTransientInternal(FRDGParentResourceRef Resource) const
diff --git a/Engine/Source/Runtime/RenderCore/Private/RenderGraphPrivate.cpp b/Engine/Source/Runtime/RenderCore/Private/RenderGraphPrivate.cpp
--- a/Engine/Source/Runtime/RenderCore/Private/RenderGraphPrivate.cpp
+++ b/Engine/Source/Runtime/RenderCore/Private/RenderGraphPrivate.cpp
@@ -302,7 +302,8 @@ FAutoConsoleVariableRef CVarRDGUseTransientAllocator(
 	TEXT("r.RDG.TransientAllocator"), GRDGTransientAllocator,
 	TEXT("RDG will use the RHITransientResourceAllocator to allocate all transient resources.")
 	TEXT(" 0: disables the transient allocator;")
-	TEXT(" 1: enables the transient allocator (default)"),
+	TEXT(" 1: enables the transient allocator (default);")
+	TEXT(" 2: enables the transient allocator for resources with FastVRAM flag only"),
 	ECVF_RenderThreadSafe);
 
 // Fix for random GPU crashes on draw indirects on multiple IHVs. Force all indirect arg buffers as non transient (see UE-115982)
