Disable full aliasing barriers for d3d12 by default and only enable it on the platform that needs it. Also fix some typos

diff --git a/Engine/Source/Runtime/D3D12RHI/Private/D3D12Commands.cpp b/Engine/Source/Runtime/D3D12RHI/Private/D3D12Commands.cpp
--- a/Engine/Source/Runtime/D3D12RHI/Private/D3D12Commands.cpp
+++ b/Engine/Source/Runtime/D3D12RHI/Private/D3D12Commands.cpp
@@ -25,8 +25,7 @@ static FAutoConsoleVariableRef CVarSyncTemporalResources(
 	ECVF_RenderThreadSafe
 	);
 
-// TODO test and enable overlap tracking
-static int32 GD3D12TransientAllocatorFullAliasingBarrier = 1;
+static int32 GD3D12TransientAllocatorFullAliasingBarrier = 0;
 static FAutoConsoleVariableRef CVarD3D12TransientAllocatorFullAliasingBarrier(
 	TEXT("d3d12.TransientAllocator.FullAliasingBarrier"),
 	GD3D12TransientAllocatorFullAliasingBarrier,
@@ -37,7 +36,7 @@ static int32 GD3D12AllowDiscardResources = 1;
 static FAutoConsoleVariableRef CVarD3D12AllowDiscardResources(
 	TEXT("d3d12.AllowDiscardResources"),
 	GD3D12AllowDiscardResources,
-	TEXT("Wheter to call DiscardResources after transient aliasing acquire. This is not needed on some platforms if newly acquired resources are cleared before use."),
+	TEXT("Whether to call DiscardResources after transient aliasing acquire. This is not needed on some platforms if newly acquired resources are cleared before use."),
 	ECVF_RenderThreadSafe);
 
 using namespace D3D12RHI;
