Prevent a crash in DX12 if GFSDK_Aftermath_Lib.x64.dll is not found and improve the finding of the dll when using -basedir=

- The original code is using ::LoadLibrary with a relative path to find the dll and when running from staged data via -basedir= we will be relying on that path being relative to the current directory set by the command line. However many modules call FPlatformProcess::PushDllDirectory/PopDllDirectory to find dlls and when all of the paths are fully popped we will call ::SetDllDirectory with an empty string which will remove the current directory from the list of valid paths for ::LoadLibrary to search.
- To fix this we can call FPlatformProcess::PushDllDirectory/PopDllDirectory ourselves, with the location of the dll and then request the handle of the dll instead (examples of this can be found in other modules)
- When running with DX12 we no longer attempt to enable NvAftermath if we previously failed to find the dll in order to prevent a crash. We do this by recording if the dll has been loaded or not via a new global variable GDX12NVAfterMathModuleLoaded, similar to how the DX11 code path tracks this.

diff --git a/Engine/Source/Runtime/D3D12RHI/Private/D3D12Adapter.cpp b/Engine/Source/Runtime/D3D12RHI/Private/D3D12Adapter.cpp
--- a/Engine/Source/Runtime/D3D12RHI/Private/D3D12Adapter.cpp
+++ b/Engine/Source/Runtime/D3D12RHI/Private/D3D12Adapter.cpp
@@ -358,7 +358,7 @@ void FD3D12Adapter::CreateRootDevice(bool bWithDebug)
 	FGenericCrashContext::SetEngineData(TEXT("RHI.Breadcrumbs"), bBreadcrumbs ? TEXT("true") : TEXT("false"));
 
 #if NV_AFTERMATH
-	if (IsRHIDeviceNVIDIA())
+	if (IsRHIDeviceNVIDIA() && GDX12NVAfterMathModuleLoaded)
 	{
 		// GPUcrash dump handler must be attached prior to device creation
 		if (EnumHasAnyFlags(GPUCrashDebuggingModes, ED3D12GPUCrashDebuggingModes::NvAftermath))
diff --git a/Engine/Source/Runtime/D3D12RHI/Private/D3D12RHIPrivate.h b/Engine/Source/Runtime/D3D12RHI/Private/D3D12RHIPrivate.h
--- a/Engine/Source/Runtime/D3D12RHI/Private/D3D12RHIPrivate.h
+++ b/Engine/Source/Runtime/D3D12RHI/Private/D3D12RHIPrivate.h
@@ -56,6 +56,7 @@ DECLARE_LOG_CATEGORY_EXTERN(LogD3D12GapRecorder, Log, All);
 #include "GFSDK_Aftermath.h"
 #include "GFSDK_Aftermath_GpuCrashdump.h"
 #undef GFSDK_Aftermath_WITH_DX12
+extern bool GDX12NVAfterMathModuleLoaded;
 extern int32 GDX12NVAfterMathEnabled;
 extern int32 GDX12NVAfterMathTrackResources;
 extern int32 GDX12NVAfterMathMarkers;
diff --git a/Engine/Source/Runtime/D3D12RHI/Private/Windows/WindowsD3D12Device.cpp b/Engine/Source/Runtime/D3D12RHI/Private/Windows/WindowsD3D12Device.cpp
--- a/Engine/Source/Runtime/D3D12RHI/Private/Windows/WindowsD3D12Device.cpp
+++ b/Engine/Source/Runtime/D3D12RHI/Private/Windows/WindowsD3D12Device.cpp
@@ -67,6 +67,7 @@ static FAutoConsoleVariableRef CVarValidateInternalTransitions(
 );
 
 #if NV_AFTERMATH
+bool GDX12NVAfterMathModuleLoaded = false;
 // Disabled by default since introduces stalls between render and driver threads
 int32 GDX12NVAfterMathEnabled = 0;
 static FAutoConsoleVariableRef CVarDX12NVAfterMathBufferSize(
@@ -714,17 +715,22 @@ void FD3D12DynamicRHIModule::StartupModule()
 {
 #if NV_AFTERMATH
 	// Note - can't check device type here, we'll check for that before actually initializing Aftermath
-	FString AftermathBinariesRoot = FPaths::EngineDir() / TEXT("Binaries/ThirdParty/NVIDIA/NVaftermath/Win64/");
-	if (LoadLibraryW(*(AftermathBinariesRoot + "GFSDK_Aftermath_Lib.x64.dll")) == nullptr)
+	const FString AftermathBinariesRoot = FPaths::EngineDir() / TEXT("Binaries/ThirdParty/NVIDIA/NVaftermath/Win64/");
+	
+	FPlatformProcess::PushDllDirectory(*AftermathBinariesRoot);
+	void* Handle = FPlatformProcess::GetDllHandle(TEXT("GFSDK_Aftermath_Lib.x64.dll"));
+	FPlatformProcess::PopDllDirectory(*AftermathBinariesRoot);
+
+	if (Handle == nullptr)
 	{
 		UE_LOG(LogD3D12RHI, Warning, TEXT("Failed to load GFSDK_Aftermath_Lib.x64.dll"));
-		GDX12NVAfterMathEnabled = 0;
+		GDX12NVAfterMathModuleLoaded = false;
 		return;
 	}
 	else
 	{
 		UE_LOG(LogD3D12RHI, Log, TEXT("Aftermath initialized"));
-		GDX12NVAfterMathEnabled = 1;
+		GDX12NVAfterMathModuleLoaded = true;
 	}
 #endif
 
diff --git a/Engine/Source/Runtime/VulkanRHI/Private/Windows/VulkanWindowsPlatform.cpp b/Engine/Source/Runtime/VulkanRHI/Private/Windows/VulkanWindowsPlatform.cpp
--- a/Engine/Source/Runtime/VulkanRHI/Private/Windows/VulkanWindowsPlatform.cpp
+++ b/Engine/Source/Runtime/VulkanRHI/Private/Windows/VulkanWindowsPlatform.cpp
@@ -31,8 +31,13 @@ bool FVulkanWindowsPlatform::LoadVulkanLibrary()
 	if (bAllowVendorDevice)
 	{
 		// Note - can't check device type here, we'll check for that before actually initializing Aftermath
-		FString AftermathBinariesRoot = FPaths::EngineDir() / TEXT("Binaries/ThirdParty/NVIDIA/NVaftermath/Win64/");
-		if (LoadLibraryW(*(AftermathBinariesRoot + "GFSDK_Aftermath_Lib.x64.dll")) == nullptr)
+		const FString AftermathBinariesRoot = FPaths::EngineDir() / TEXT("Binaries/ThirdParty/NVIDIA/NVaftermath/Win64/");
+
+		FPlatformProcess::PushDllDirectory(*AftermathBinariesRoot);
+		void* Handle = FPlatformProcess::GetDllHandle(TEXT("GFSDK_Aftermath_Lib.x64.dll"));
+		FPlatformProcess::PopDllDirectory(*AftermathBinariesRoot);
+
+		if (Handle == nullptr)
 		{
 			UE_LOG(LogVulkanRHI, Warning, TEXT("Failed to load GFSDK_Aftermath_Lib.x64.dll"));
 		}
diff --git a/Engine/Source/Runtime/Windows/D3D11RHI/Private/Windows/WindowsD3D11Device.cpp b/Engine/Source/Runtime/Windows/D3D11RHI/Private/Windows/WindowsD3D11Device.cpp
--- a/Engine/Source/Runtime/Windows/D3D11RHI/Private/Windows/WindowsD3D11Device.cpp
+++ b/Engine/Source/Runtime/Windows/D3D11RHI/Private/Windows/WindowsD3D11Device.cpp
@@ -767,8 +767,13 @@ void FD3D11DynamicRHIModule::StartupModule()
 	{
 		// Note - can't check device type here, we'll check for that before actually initializing Aftermath
 
-		FString AftermathBinariesRoot = FPaths::EngineDir() / TEXT("Binaries/ThirdParty/NVIDIA/NVaftermath/Win64/");
-		if (LoadLibraryW(*(AftermathBinariesRoot + "GFSDK_Aftermath_Lib.x64.dll")) == nullptr)
+		const FString AftermathBinariesRoot = FPaths::EngineDir() / TEXT("Binaries/ThirdParty/NVIDIA/NVaftermath/Win64/");
+
+		FPlatformProcess::PushDllDirectory(*AftermathBinariesRoot);
+		void* Handle = FPlatformProcess::GetDllHandle(TEXT("GFSDK_Aftermath_Lib.x64.dll"));
+		FPlatformProcess::PopDllDirectory(*AftermathBinariesRoot);
+
+		if (Handle == nullptr)
 		{
 			UE_LOG(LogD3D11RHI, Warning, TEXT("Failed to load GFSDK_Aftermath_Lib.x64.dll"));
 			GNVAftermathModuleLoaded = false;
