Adding quicker utility for getting a shader type's name.
Rewriting ValidateBoundShader and ValidateBoundUniformBuffer to be templates instead of defines to make them easier to investigate.

diff --git a/Engine/Source/Runtime/D3D12RHI/Private/D3D12Commands.cpp b/Engine/Source/Runtime/D3D12RHI/Private/D3D12Commands.cpp
--- a/Engine/Source/Runtime/D3D12RHI/Private/D3D12Commands.cpp
+++ b/Engine/Source/Runtime/D3D12RHI/Private/D3D12Commands.cpp
@@ -41,49 +41,37 @@ static FAutoConsoleVariableRef CVarD3D12AllowDiscardResources(
 
 using namespace D3D12RHI;
 
-#define DECLARE_ISBOUNDSHADER(ShaderType) inline void ValidateBoundShader(FD3D12StateCache& InStateCache, FRHI##ShaderType* ShaderType##RHI) \
-{ \
-	FD3D12##ShaderType* CachedShader; \
-	InStateCache.Get##ShaderType(&CachedShader); \
-	FD3D12##ShaderType* ShaderType = FD3D12DynamicRHI::ResourceCast(ShaderType##RHI); \
-	ensureMsgf(CachedShader == ShaderType, TEXT("Parameters are being set for a %s which is not currently bound"), TEXT( #ShaderType )); \
-}
-
-DECLARE_ISBOUNDSHADER(VertexShader)
-DECLARE_ISBOUNDSHADER(MeshShader)
-DECLARE_ISBOUNDSHADER(AmplificationShader)
-DECLARE_ISBOUNDSHADER(PixelShader)
-DECLARE_ISBOUNDSHADER(GeometryShader)
-DECLARE_ISBOUNDSHADER(ComputeShader)
-
-#define DECLARE_ISVALIDUNIFORMBUFFERHASH(ShaderType) inline void ValidateBoundUniformBuffer(FD3D12UniformBuffer* InUniformBuffer, FRHI##ShaderType* In##RHI, uint32 InBufferIndex) \
-{ \
-	FD3D12##ShaderType* D3D12Shader = FD3D12DynamicRHI::ResourceCast(In##RHI); \
-	if (InBufferIndex < (uint32)D3D12Shader->ShaderResourceTable.ResourceTableLayoutHashes.Num()) \
-	{ \
-		uint32 UniformBufferHash = InUniformBuffer->GetLayout().GetHash(); \
-		uint32 ShaderTableHash = D3D12Shader->ShaderResourceTable.ResourceTableLayoutHashes[InBufferIndex]; \
-		ensureMsgf(ShaderTableHash == 0 || UniformBufferHash == ShaderTableHash, TEXT("Invalid uniform buffer %s bound on shader %s at index %d."), *(InUniformBuffer->GetLayout().GetDebugName()), TEXT( #ShaderType ), InBufferIndex); \
-	} \
-}
-
-DECLARE_ISVALIDUNIFORMBUFFERHASH(VertexShader)
-DECLARE_ISVALIDUNIFORMBUFFERHASH(MeshShader)
-DECLARE_ISVALIDUNIFORMBUFFERHASH(AmplificationShader)
-DECLARE_ISVALIDUNIFORMBUFFERHASH(PixelShader)
-DECLARE_ISVALIDUNIFORMBUFFERHASH(GeometryShader)
-DECLARE_ISVALIDUNIFORMBUFFERHASH(ComputeShader)
-
-#if EXECUTE_DEBUG_COMMAND_LISTS
-bool GIsDoingQuery = false;
+template<typename TRHIShaderType>
+inline void ValidateBoundShader(FD3D12StateCache & InStateCache, TRHIShaderType* InShaderRHI)
+{
+#if DO_CHECK
+	auto* CachedShader = InStateCache.GetShader<typename TD3D12ResourceTraits<TRHIShaderType>::TConcreteType>();
+	auto* ShaderType = FD3D12DynamicRHI::ResourceCast(InShaderRHI);
+	ensureMsgf(CachedShader == ShaderType, TEXT("Parameters are being set for a %sShader which is not currently bound"),
+		GetShaderFrequencyString<TRHIShaderType>(false));
 #endif
+}
 
+template<typename TRHIShaderType>
+inline void ValidateBoundUniformBuffer(FD3D12UniformBuffer * InUniformBuffer, TRHIShaderType* InShaderRHI, uint32 InBufferIndex)
+{
 #if DO_CHECK
-#define VALIDATE_BOUND_SHADER(s) ValidateBoundShader(StateCache, s)
-#define VALIDATE_BOUND_UNIFORMBUFFER_HASH(ub, s, i) ValidateBoundUniformBuffer(ub, s, i)
-#else
-#define VALIDATE_BOUND_SHADER(s)
-#define VALIDATE_BOUND_UNIFORMBUFFER_HASH(ub, s, i)
+	auto* D3D12Shader = FD3D12DynamicRHI::ResourceCast(InShaderRHI);
+	if (InBufferIndex < (uint32)D3D12Shader->ShaderResourceTable.ResourceTableLayoutHashes.Num())
+	{
+		uint32 UniformBufferHash = InUniformBuffer->GetLayout().GetHash(); 
+		uint32 ShaderTableHash = D3D12Shader->ShaderResourceTable.ResourceTableLayoutHashes[InBufferIndex];
+		ensureMsgf(ShaderTableHash == 0 || UniformBufferHash == ShaderTableHash,
+			TEXT("Invalid uniform buffer %s bound on %sShader at index %d."),
+			*(InUniformBuffer->GetLayout().GetDebugName()),
+			GetShaderFrequencyString<TRHIShaderType>(false),
+			InBufferIndex);
+	}
+#endif
+}
+
+#if EXECUTE_DEBUG_COMMAND_LISTS
+bool GIsDoingQuery = false;
 #endif
 
 #if !defined(D3D12_PLATFORM_SUPPORTS_RESOLVE_SHADERS)
@@ -790,35 +778,37 @@ void FD3D12CommandContext::RHISetShaderTexture(FRHIGraphicsShader* ShaderRHI, ui
 	case SF_Vertex:
 	{
 		FRHIVertexShader* VertexShaderRHI = static_cast<FRHIVertexShader*>(ShaderRHI);
-		VALIDATE_BOUND_SHADER(VertexShaderRHI);
+		ValidateBoundShader(StateCache, VertexShaderRHI);
 		StateCache.SetShaderResourceView<SF_Vertex>(NewTexture ? NewTexture->GetShaderResourceView() : nullptr, TextureIndex);
 	}
 	break;
+#if PLATFORM_SUPPORTS_MESH_SHADERS
 	case SF_Mesh:
 	{
 		FRHIMeshShader* MeshShaderRHI = static_cast<FRHIMeshShader*>(ShaderRHI);
-		VALIDATE_BOUND_SHADER(MeshShaderRHI);
+		ValidateBoundShader(StateCache, MeshShaderRHI);
 		StateCache.SetShaderResourceView<SF_Mesh>(NewTexture ? NewTexture->GetShaderResourceView() : nullptr, TextureIndex);
 	}
 	break;
 	case SF_Amplification:
 	{
 		FRHIAmplificationShader* AmplificationShaderRHI = static_cast<FRHIAmplificationShader*>(ShaderRHI);
-		VALIDATE_BOUND_SHADER(AmplificationShaderRHI);
+		ValidateBoundShader(StateCache, AmplificationShaderRHI);
 		StateCache.SetShaderResourceView<SF_Amplification>(NewTexture ? NewTexture->GetShaderResourceView() : nullptr, TextureIndex);
 	}
 	break;
+#endif
 	case SF_Geometry:
 	{
 		FRHIGeometryShader* GeometryShaderRHI = static_cast<FRHIGeometryShader*>(ShaderRHI);
-		VALIDATE_BOUND_SHADER(GeometryShaderRHI);
+		ValidateBoundShader(StateCache, GeometryShaderRHI);
 		StateCache.SetShaderResourceView<SF_Geometry>(NewTexture ? NewTexture->GetShaderResourceView() : nullptr, TextureIndex);
 	}
 	break;
 	case SF_Pixel:
 	{
 		FRHIPixelShader* PixelShaderRHI = static_cast<FRHIPixelShader*>(ShaderRHI);
-		VALIDATE_BOUND_SHADER(PixelShaderRHI);
+		ValidateBoundShader(StateCache, PixelShaderRHI);
 		StateCache.SetShaderResourceView<SF_Pixel>(NewTexture ? NewTexture->GetShaderResourceView() : nullptr, TextureIndex);
 	}
 	break;
@@ -829,7 +819,7 @@ void FD3D12CommandContext::RHISetShaderTexture(FRHIGraphicsShader* ShaderRHI, ui
 
 void FD3D12CommandContext::RHISetShaderTexture(FRHIComputeShader* ComputeShaderRHI, uint32 TextureIndex, FRHITexture* NewTextureRHI)
 {
-	//VALIDATE_BOUND_SHADER(ComputeShaderRHI);
+	//ValidateBoundShader(StateCache, ComputeShaderRHI);
 	FD3D12TextureBase* const NewTexture = RetrieveTextureBase(NewTextureRHI);
 	StateCache.SetShaderResourceView<SF_Compute>(NewTexture ? NewTexture->GetShaderResourceView() : nullptr, TextureIndex);
 }
@@ -852,7 +842,7 @@ void FD3D12CommandContext::RHISetUAVParameter(FRHIPixelShader* PixelShaderRHI, u
 
 void FD3D12CommandContext::RHISetUAVParameter(FRHIComputeShader* ComputeShaderRHI, uint32 UAVIndex, FRHIUnorderedAccessView* UAVRHI)
 {
-	//VALIDATE_BOUND_SHADER(ComputeShaderRHI);
+	//ValidateBoundShader(StateCache, ComputeShaderRHI);
 
 	FD3D12UnorderedAccessView* UAV = RetrieveObject<FD3D12UnorderedAccessView>(UAVRHI);
 
@@ -869,7 +859,7 @@ void FD3D12CommandContext::RHISetUAVParameter(FRHIComputeShader* ComputeShaderRH
 
 void FD3D12CommandContext::RHISetUAVParameter(FRHIComputeShader* ComputeShaderRHI, uint32 UAVIndex, FRHIUnorderedAccessView* UAVRHI, uint32 InitialCount)
 {
-	//VALIDATE_BOUND_SHADER(ComputeShaderRHI);
+	//ValidateBoundShader(StateCache, ComputeShaderRHI);
 
 	FD3D12UnorderedAccessView* UAV = RetrieveObject<FD3D12UnorderedAccessView>(UAVRHI);
 
@@ -889,35 +879,37 @@ void FD3D12CommandContext::RHISetShaderResourceViewParameter(FRHIGraphicsShader*
 	case SF_Vertex:
 	{
 		FRHIVertexShader* VertexShaderRHI = static_cast<FRHIVertexShader*>(ShaderRHI);
-		VALIDATE_BOUND_SHADER(VertexShaderRHI);
+		ValidateBoundShader(StateCache, VertexShaderRHI);
 		StateCache.SetShaderResourceView<SF_Vertex>(SRV, TextureIndex);
 	}
 	break;
+#if PLATFORM_SUPPORTS_MESH_SHADERS
 	case SF_Mesh:
 	{
 		FRHIMeshShader* MeshShaderRHI = static_cast<FRHIMeshShader*>(ShaderRHI);
-		VALIDATE_BOUND_SHADER(MeshShaderRHI);
+		ValidateBoundShader(StateCache, MeshShaderRHI);
 		StateCache.SetShaderResourceView<SF_Mesh>(SRV, TextureIndex);
 	}
 	break;
 	case SF_Amplification:
 	{
 		FRHIAmplificationShader* AmplificationShaderRHI = static_cast<FRHIAmplificationShader*>(ShaderRHI);
-		VALIDATE_BOUND_SHADER(AmplificationShaderRHI);
+		ValidateBoundShader(StateCache, AmplificationShaderRHI);
 		StateCache.SetShaderResourceView<SF_Amplification>(SRV, TextureIndex);
 	}
 	break;
+#endif
 	case SF_Geometry:
 	{
 		FRHIGeometryShader* GeometryShaderRHI = static_cast<FRHIGeometryShader*>(ShaderRHI);
-		VALIDATE_BOUND_SHADER(GeometryShaderRHI);
+		ValidateBoundShader(StateCache, GeometryShaderRHI);
 		StateCache.SetShaderResourceView<SF_Geometry>(SRV, TextureIndex);
 	}
 	break;
 	case SF_Pixel:
 	{
 		FRHIPixelShader* PixelShaderRHI = static_cast<FRHIPixelShader*>(ShaderRHI);
-		VALIDATE_BOUND_SHADER(PixelShaderRHI);
+		ValidateBoundShader(StateCache, PixelShaderRHI);
 		StateCache.SetShaderResourceView<SF_Pixel>(SRV, TextureIndex);
 	}
 	break;
@@ -929,7 +921,7 @@ void FD3D12CommandContext::RHISetShaderResourceViewParameter(FRHIGraphicsShader*
 
 void FD3D12CommandContext::RHISetShaderResourceViewParameter(FRHIComputeShader* ComputeShaderRHI, uint32 TextureIndex, FRHIShaderResourceView* SRVRHI)
 {
-	//VALIDATE_BOUND_SHADER(ComputeShaderRHI);
+	//ValidateBoundShader(StateCache, ComputeShaderRHI);
 	FD3D12ShaderResourceView* const SRV = RetrieveObject<FD3D12ShaderResourceView>(SRVRHI);
 	StateCache.SetShaderResourceView<SF_Compute>(SRV, TextureIndex);
 }
@@ -942,35 +934,37 @@ void FD3D12CommandContext::RHISetShaderSampler(FRHIGraphicsShader* ShaderRHI, ui
 	case SF_Vertex:
 	{
 		FRHIVertexShader* VertexShaderRHI = static_cast<FRHIVertexShader*>(ShaderRHI);
-		VALIDATE_BOUND_SHADER(VertexShaderRHI);
+		ValidateBoundShader(StateCache, VertexShaderRHI);
 		StateCache.SetSamplerState<SF_Vertex>(NewState, SamplerIndex);
 	}
 	break;
+#if PLATFORM_SUPPORTS_MESH_SHADERS
 	case SF_Mesh:
 	{
 		FRHIMeshShader* MeshShaderRHI = static_cast<FRHIMeshShader*>(ShaderRHI);
-		VALIDATE_BOUND_SHADER(MeshShaderRHI);
+		ValidateBoundShader(StateCache, MeshShaderRHI);
 		StateCache.SetSamplerState<SF_Mesh>(NewState, SamplerIndex);
 	}
 	break;
 	case SF_Amplification:
 	{
 		FRHIAmplificationShader* AmplificationShaderRHI = static_cast<FRHIAmplificationShader*>(ShaderRHI);
-		VALIDATE_BOUND_SHADER(AmplificationShaderRHI);
+		ValidateBoundShader(StateCache, AmplificationShaderRHI);
 		StateCache.SetSamplerState<SF_Amplification>(NewState, SamplerIndex);
 	}
 	break;
+#endif
 	case SF_Geometry:
 	{
 		FRHIGeometryShader* GeometryShaderRHI = static_cast<FRHIGeometryShader*>(ShaderRHI);
-		VALIDATE_BOUND_SHADER(GeometryShaderRHI);
+		ValidateBoundShader(StateCache, GeometryShaderRHI);
 		StateCache.SetSamplerState<SF_Geometry>(NewState, SamplerIndex);
 	}
 	break;
 	case SF_Pixel:
 	{
 		FRHIPixelShader* PixelShaderRHI = static_cast<FRHIPixelShader*>(ShaderRHI);
-		VALIDATE_BOUND_SHADER(PixelShaderRHI);
+		ValidateBoundShader(StateCache, PixelShaderRHI);
 		StateCache.SetSamplerState<SF_Pixel>(NewState, SamplerIndex);
 	}
 	break;
@@ -981,7 +975,7 @@ void FD3D12CommandContext::RHISetShaderSampler(FRHIGraphicsShader* ShaderRHI, ui
 
 void FD3D12CommandContext::RHISetShaderSampler(FRHIComputeShader* ComputeShaderRHI, uint32 SamplerIndex, FRHISamplerState* NewStateRHI)
 {
-	//VALIDATE_BOUND_SHADER(ComputeShaderRHI);
+	//ValidateBoundShader(StateCache, ComputeShaderRHI);
 	FD3D12SamplerState* NewState = RetrieveObject<FD3D12SamplerState>(NewStateRHI);
 	StateCache.SetSamplerState<SF_Compute>(NewState, SamplerIndex);
 }
@@ -996,17 +990,18 @@ void FD3D12CommandContext::RHISetShaderUniformBuffer(FRHIGraphicsShader* ShaderR
 	case SF_Vertex:
 	{
 		FRHIVertexShader* VertexShaderRHI = static_cast<FRHIVertexShader*>(ShaderRHI);
-		VALIDATE_BOUND_SHADER(VertexShaderRHI);
-		VALIDATE_BOUND_UNIFORMBUFFER_HASH(Buffer, VertexShaderRHI, BufferIndex);
+		ValidateBoundShader(StateCache, VertexShaderRHI);
+		ValidateBoundUniformBuffer(Buffer, VertexShaderRHI, BufferIndex);
 		StateCache.SetConstantsFromUniformBuffer<SF_Vertex>(BufferIndex, Buffer);
 		Stage = SF_Vertex;
 	}
 	break;
+#if PLATFORM_SUPPORTS_MESH_SHADERS
 	case SF_Mesh:
 	{
 		FRHIMeshShader* MeshShaderRHI = static_cast<FRHIMeshShader*>(ShaderRHI);
-		VALIDATE_BOUND_SHADER(MeshShaderRHI);
-		VALIDATE_BOUND_UNIFORMBUFFER_HASH(Buffer, MeshShaderRHI, BufferIndex);
+		ValidateBoundShader(StateCache, MeshShaderRHI);
+		ValidateBoundUniformBuffer(Buffer, MeshShaderRHI, BufferIndex);
 		StateCache.SetConstantsFromUniformBuffer<SF_Mesh>(BufferIndex, Buffer);
 		Stage = SF_Mesh;
 	}
@@ -1014,17 +1009,18 @@ void FD3D12CommandContext::RHISetShaderUniformBuffer(FRHIGraphicsShader* ShaderR
 	case SF_Amplification:
 	{
 		FRHIAmplificationShader* AmplificationShaderRHI = static_cast<FRHIAmplificationShader*>(ShaderRHI);
-		VALIDATE_BOUND_SHADER(AmplificationShaderRHI);
-		VALIDATE_BOUND_UNIFORMBUFFER_HASH(Buffer, AmplificationShaderRHI, BufferIndex);
+		ValidateBoundShader(StateCache, AmplificationShaderRHI);
+		ValidateBoundUniformBuffer(Buffer, AmplificationShaderRHI, BufferIndex);
 		StateCache.SetConstantsFromUniformBuffer<SF_Amplification>(BufferIndex, Buffer);
 		Stage = SF_Amplification;
 	}
 	break;
+#endif
 	case SF_Geometry:
 	{
 		FRHIGeometryShader* GeometryShaderRHI = static_cast<FRHIGeometryShader*>(ShaderRHI);
-		VALIDATE_BOUND_SHADER(GeometryShaderRHI);
-		VALIDATE_BOUND_UNIFORMBUFFER_HASH(Buffer, GeometryShaderRHI, BufferIndex);
+		ValidateBoundShader(StateCache, GeometryShaderRHI);
+		ValidateBoundUniformBuffer(Buffer, GeometryShaderRHI, BufferIndex);
 		StateCache.SetConstantsFromUniformBuffer<SF_Geometry>(BufferIndex, Buffer);
 		Stage = SF_Geometry;
 	}
@@ -1032,8 +1028,8 @@ void FD3D12CommandContext::RHISetShaderUniformBuffer(FRHIGraphicsShader* ShaderR
 	case SF_Pixel:
 	{
 		FRHIPixelShader* PixelShaderRHI = static_cast<FRHIPixelShader*>(ShaderRHI);
-		VALIDATE_BOUND_SHADER(PixelShaderRHI);
-		VALIDATE_BOUND_UNIFORMBUFFER_HASH(Buffer, PixelShaderRHI, BufferIndex);
+		ValidateBoundShader(StateCache, PixelShaderRHI);
+		ValidateBoundUniformBuffer(Buffer, PixelShaderRHI, BufferIndex);
 		StateCache.SetConstantsFromUniformBuffer<SF_Pixel>(BufferIndex, Buffer);
 		Stage = SF_Pixel;
 	}
@@ -1055,7 +1051,7 @@ void FD3D12CommandContext::RHISetShaderUniformBuffer(FRHIGraphicsShader* ShaderR
 void FD3D12CommandContext::RHISetShaderUniformBuffer(FRHIComputeShader* ComputeShader, uint32 BufferIndex, FRHIUniformBuffer* BufferRHI)
 {
 	//SCOPE_CYCLE_COUNTER(STAT_D3D12SetShaderUniformBuffer);
-	//VALIDATE_BOUND_SHADER(ComputeShader);
+	//ValidateBoundShader(StateCache, ComputeShader);
 	FD3D12UniformBuffer* Buffer = RetrieveObject<FD3D12UniformBuffer>(BufferRHI);
 
 	StateCache.SetConstantsFromUniformBuffer<SF_Compute>(BufferIndex, Buffer);
@@ -1078,35 +1074,37 @@ void FD3D12CommandContext::RHISetShaderParameter(FRHIGraphicsShader* ShaderRHI,
 	case SF_Vertex:
 	{
 		FRHIVertexShader* VertexShaderRHI = static_cast<FRHIVertexShader*>(ShaderRHI);
-		VALIDATE_BOUND_SHADER(VertexShaderRHI);
+		ValidateBoundShader(StateCache, VertexShaderRHI);
 		VSConstantBuffer.UpdateConstant((const uint8*)NewValue, BaseIndex, NumBytes);
 	}
 	break;
+#if PLATFORM_SUPPORTS_MESH_SHADERS
 	case SF_Mesh:
 	{
 		FRHIMeshShader* MeshShaderRHI = static_cast<FRHIMeshShader*>(ShaderRHI);
-		VALIDATE_BOUND_SHADER(MeshShaderRHI);
+		ValidateBoundShader(StateCache, MeshShaderRHI);
 		MSConstantBuffer.UpdateConstant((const uint8*)NewValue, BaseIndex, NumBytes);
 	}
 	break;
 	case SF_Amplification:
 	{
 		FRHIAmplificationShader* AmplificationShaderRHI = static_cast<FRHIAmplificationShader*>(ShaderRHI);
-		VALIDATE_BOUND_SHADER(AmplificationShaderRHI);
+		ValidateBoundShader(StateCache, AmplificationShaderRHI);
 		ASConstantBuffer.UpdateConstant((const uint8*)NewValue, BaseIndex, NumBytes);
 	}
 	break;
+#endif
 	case SF_Geometry:
 	{
 		FRHIGeometryShader* GeometryShaderRHI = static_cast<FRHIGeometryShader*>(ShaderRHI);
-		VALIDATE_BOUND_SHADER(GeometryShaderRHI);
+		ValidateBoundShader(StateCache, GeometryShaderRHI);
 		GSConstantBuffer.UpdateConstant((const uint8*)NewValue, BaseIndex, NumBytes);
 	}
 	break;
 	case SF_Pixel:
 	{
 		FRHIPixelShader* PixelShaderRHI = static_cast<FRHIPixelShader*>(ShaderRHI);
-		VALIDATE_BOUND_SHADER(PixelShaderRHI);
+		ValidateBoundShader(StateCache, PixelShaderRHI);
 		PSConstantBuffer.UpdateConstant((const uint8*)NewValue, BaseIndex, NumBytes);
 	}
 	break;
@@ -1117,7 +1115,7 @@ void FD3D12CommandContext::RHISetShaderParameter(FRHIGraphicsShader* ShaderRHI,
 
 void FD3D12CommandContext::RHISetShaderParameter(FRHIComputeShader* ComputeShaderRHI, uint32 BufferIndex, uint32 BaseIndex, uint32 NumBytes, const void* NewValue)
 {
-	//VALIDATE_BOUND_SHADER(ComputeShaderRHI);
+	//ValidateBoundShader(StateCache, ComputeShaderRHI);
 	checkSlow(BufferIndex == 0);
 	CSConstantBuffer.UpdateConstant((const uint8*)NewValue, BaseIndex, NumBytes);
 }
@@ -1511,6 +1509,7 @@ void FD3D12CommandContext::CommitNonComputeShaderConstants()
 		StateCache.SetConstantBuffer<SF_Vertex>(VSConstantBuffer, bDiscardSharedGraphicsConstants);
 	}
 
+#if PLATFORM_SUPPORTS_MESH_SHADERS
 	if (GraphicPSO->bShaderNeedsGlobalConstantBuffer[SF_Mesh])
 	{
 		StateCache.SetConstantBuffer<SF_Mesh>(MSConstantBuffer, bDiscardSharedGraphicsConstants);
@@ -1520,6 +1519,7 @@ void FD3D12CommandContext::CommitNonComputeShaderConstants()
 	{
 		StateCache.SetConstantBuffer<SF_Amplification>(ASConstantBuffer, bDiscardSharedGraphicsConstants);
 	}
+#endif
 
 	if (GraphicPSO->bShaderNeedsGlobalConstantBuffer[SF_Geometry])
 	{
diff --git a/Engine/Source/Runtime/D3D12RHI/Private/D3D12StateCachePrivate.h b/Engine/Source/Runtime/D3D12RHI/Private/D3D12StateCachePrivate.h
--- a/Engine/Source/Runtime/D3D12RHI/Private/D3D12StateCachePrivate.h
+++ b/Engine/Source/Runtime/D3D12RHI/Private/D3D12StateCachePrivate.h
@@ -467,11 +467,6 @@ protected:
 		}
 	}
 
-	template <typename TShader> D3D12_STATE_CACHE_INLINE void GetShader(TShader** Shader)
-	{
-		*Shader = StateCacheShaderTraits<TShader>::GetShader(GetGraphicsPipelineState());
-	}
-
 	template <ED3D12PipelineType PipelineType>
 	D3D12_STATE_CACHE_INLINE void InternalSetPipelineState()
 	{
@@ -724,6 +719,18 @@ public:
 	void SetStencilRef(uint32 StencilRef);
 	uint32 GetStencilRef() const { return PipelineState.Graphics.CurrentReferenceStencil; }
 
+	template <typename TShader>
+	D3D12_STATE_CACHE_INLINE TShader* GetShader()
+	{
+		return StateCacheShaderTraits<TShader>::GetShader(GetGraphicsPipelineState());
+	}
+
+	template <typename TShader>
+	D3D12_STATE_CACHE_INLINE void GetShader(TShader** Shader)
+	{
+		*Shader = GetShader<TShader>();
+	}
+
 	D3D12_STATE_CACHE_INLINE void GetVertexShader(FD3D12VertexShader** Shader)
 	{
 		GetShader(Shader);
diff --git a/Engine/Source/Runtime/RHI/Public/RHIResources.h b/Engine/Source/Runtime/RHI/Public/RHIResources.h
--- a/Engine/Source/Runtime/RHI/Public/RHIResources.h
+++ b/Engine/Source/Runtime/RHI/Public/RHIResources.h
@@ -2164,26 +2164,32 @@ public:
 
 typedef TRefCountPtr<FRHICustomPresent> FCustomPresentRHIRef;
 
-// Template magic to convert an FRHI*Shader to its enum
+// Templates to convert an FRHI*Shader to its enum
 template<typename TRHIShader> struct TRHIShaderToEnum {};
-template<> struct TRHIShaderToEnum<FRHIVertexShader>		{ enum { ShaderFrequency = SF_Vertex }; };
-template<> struct TRHIShaderToEnum<FRHIMeshShader>			{ enum { ShaderFrequency = SF_Mesh}; };
-template<> struct TRHIShaderToEnum<FRHIAmplificationShader>	{ enum { ShaderFrequency = SF_Amplification}; };
-template<> struct TRHIShaderToEnum<FRHIPixelShader>			{ enum { ShaderFrequency = SF_Pixel }; };
-template<> struct TRHIShaderToEnum<FRHIGeometryShader>		{ enum { ShaderFrequency = SF_Geometry }; };
-template<> struct TRHIShaderToEnum<FRHIComputeShader>		{ enum { ShaderFrequency = SF_Compute }; };
-template<> struct TRHIShaderToEnum<FRHIVertexShader*>		{ enum { ShaderFrequency = SF_Vertex }; };
-template<> struct TRHIShaderToEnum<FRHIMeshShader*>			{ enum { ShaderFrequency = SF_Mesh }; };
-template<> struct TRHIShaderToEnum<FRHIAmplificationShader*>{ enum { ShaderFrequency = SF_Amplification }; };
-template<> struct TRHIShaderToEnum<FRHIPixelShader*>		{ enum { ShaderFrequency = SF_Pixel }; };
-template<> struct TRHIShaderToEnum<FRHIGeometryShader*>		{ enum { ShaderFrequency = SF_Geometry }; };
-template<> struct TRHIShaderToEnum<FRHIComputeShader*>		{ enum { ShaderFrequency = SF_Compute }; };
-template<> struct TRHIShaderToEnum<FVertexShaderRHIRef>		{ enum { ShaderFrequency = SF_Vertex }; };
-template<> struct TRHIShaderToEnum<FMeshShaderRHIRef>		{ enum { ShaderFrequency = SF_Mesh }; };
-template<> struct TRHIShaderToEnum<FAmplificationShaderRHIRef>{ enum { ShaderFrequency = SF_Amplification }; };
-template<> struct TRHIShaderToEnum<FPixelShaderRHIRef>		{ enum { ShaderFrequency = SF_Pixel }; };
-template<> struct TRHIShaderToEnum<FGeometryShaderRHIRef>	{ enum { ShaderFrequency = SF_Geometry }; };
-template<> struct TRHIShaderToEnum<FComputeShaderRHIRef>	{ enum { ShaderFrequency = SF_Compute }; };
+template<> struct TRHIShaderToEnum<FRHIVertexShader>           { enum { ShaderFrequency = SF_Vertex        }; };
+template<> struct TRHIShaderToEnum<FRHIMeshShader>             { enum { ShaderFrequency = SF_Mesh          }; };
+template<> struct TRHIShaderToEnum<FRHIAmplificationShader>    { enum { ShaderFrequency = SF_Amplification }; };
+template<> struct TRHIShaderToEnum<FRHIPixelShader>            { enum { ShaderFrequency = SF_Pixel         }; };
+template<> struct TRHIShaderToEnum<FRHIGeometryShader>         { enum { ShaderFrequency = SF_Geometry      }; };
+template<> struct TRHIShaderToEnum<FRHIComputeShader>          { enum { ShaderFrequency = SF_Compute       }; };
+template<> struct TRHIShaderToEnum<FRHIVertexShader*>          { enum { ShaderFrequency = SF_Vertex        }; };
+template<> struct TRHIShaderToEnum<FRHIMeshShader*>            { enum { ShaderFrequency = SF_Mesh          }; };
+template<> struct TRHIShaderToEnum<FRHIAmplificationShader*>   { enum { ShaderFrequency = SF_Amplification }; };
+template<> struct TRHIShaderToEnum<FRHIPixelShader*>           { enum { ShaderFrequency = SF_Pixel         }; };
+template<> struct TRHIShaderToEnum<FRHIGeometryShader*>        { enum { ShaderFrequency = SF_Geometry      }; };
+template<> struct TRHIShaderToEnum<FRHIComputeShader*>         { enum { ShaderFrequency = SF_Compute       }; };
+template<> struct TRHIShaderToEnum<FVertexShaderRHIRef>        { enum { ShaderFrequency = SF_Vertex        }; };
+template<> struct TRHIShaderToEnum<FMeshShaderRHIRef>          { enum { ShaderFrequency = SF_Mesh          }; };
+template<> struct TRHIShaderToEnum<FAmplificationShaderRHIRef> { enum { ShaderFrequency = SF_Amplification }; };
+template<> struct TRHIShaderToEnum<FPixelShaderRHIRef>         { enum { ShaderFrequency = SF_Pixel         }; };
+template<> struct TRHIShaderToEnum<FGeometryShaderRHIRef>      { enum { ShaderFrequency = SF_Geometry      }; };
+template<> struct TRHIShaderToEnum<FComputeShaderRHIRef>       { enum { ShaderFrequency = SF_Compute       }; };
+
+template<typename TRHIShaderType>
+inline const TCHAR* GetShaderFrequencyString(bool bIncludePrefix = true)
+{
+	return GetShaderFrequencyString(static_cast<EShaderFrequency>(TRHIShaderToEnum<TRHIShaderType>::ShaderFrequency), bIncludePrefix);
+}
 
 struct FBoundShaderStateInput
 {
