Fixing query waiting stats. Specifically, we can't add job thread wait times to GRenderThreadIdle.

diff --git a/Engine/Source/Runtime/D3D12RHI/Private/D3D12Query.cpp b/Engine/Source/Runtime/D3D12RHI/Private/D3D12Query.cpp
--- a/Engine/Source/Runtime/D3D12RHI/Private/D3D12Query.cpp
+++ b/Engine/Source/Runtime/D3D12RHI/Private/D3D12Query.cpp
@@ -175,8 +175,11 @@ bool FD3D12Device::GetQueryData(FD3D12RenderQuery& Query, bool bWait)
 
 		SyncPoint.WaitForCompletion();
 
-		GRenderThreadIdle[ERenderThreadIdleTypes::WaitingForGPUQuery] += FPlatformTime::Cycles() - IdleStart;
-		GRenderThreadNumIdle[ERenderThreadIdleTypes::WaitingForGPUQuery]++;
+		if (IsInActualRenderingThread())
+		{
+			GRenderThreadIdle[ERenderThreadIdleTypes::WaitingForGPUQuery] += FPlatformTime::Cycles() - IdleStart;
+			GRenderThreadNumIdle[ERenderThreadIdleTypes::WaitingForGPUQuery]++;
+		}
 	}
 
 	// Read the data from the query's result buffer.
diff --git a/Engine/Source/Runtime/VulkanRHI/Private/VulkanQuery.cpp b/Engine/Source/Runtime/VulkanRHI/Private/VulkanQuery.cpp
--- a/Engine/Source/Runtime/VulkanRHI/Private/VulkanQuery.cpp
+++ b/Engine/Source/Runtime/VulkanRHI/Private/VulkanQuery.cpp
@@ -165,8 +165,11 @@ bool FVulkanOcclusionQueryPool::InternalTryGetResults(bool bWait)
 				++NumLoops;
 			}
 
-			GRenderThreadIdle[ERenderThreadIdleTypes::WaitingForGPUQuery] += FPlatformTime::Cycles() - IdleStart;
-			GRenderThreadNumIdle[ERenderThreadIdleTypes::WaitingForGPUQuery]++;
+			if (IsInActualRenderingThread())
+			{
+				GRenderThreadIdle[ERenderThreadIdleTypes::WaitingForGPUQuery] += FPlatformTime::Cycles() - IdleStart;
+				GRenderThreadNumIdle[ERenderThreadIdleTypes::WaitingForGPUQuery]++;
+			}
 
 			State = EState::RT_PostGetResults;
 			return true;
