Make sure delete tasks on the d3d12 deferred deletion queue are thread safe because of 'random' flush of pending delete resources which can happen during upload buffer allocation

diff --git a/Engine/Source/Runtime/D3D12RHI/Private/D3D12Resources.cpp b/Engine/Source/Runtime/D3D12RHI/Private/D3D12Resources.cpp
--- a/Engine/Source/Runtime/D3D12RHI/Private/D3D12Resources.cpp
+++ b/Engine/Source/Runtime/D3D12RHI/Private/D3D12Resources.cpp
@@ -69,6 +69,8 @@ void FD3D12DeferredDeletionQueue::EnqueueResource(ID3D12Object* pResource, FD3D1
 
 bool FD3D12DeferredDeletionQueue::ReleaseResources(bool bDeleteImmediately, bool bIsShutDown)
 {
+	FScopeLock ScopeLock(&DeleteTaskCS);
+
 	FD3D12Adapter* Adapter = GetParentAdapter();
 
 	if (GD3D12AsyncDeferredDeletion)
diff --git a/Engine/Source/Runtime/D3D12RHI/Public/D3D12Resources.h b/Engine/Source/Runtime/D3D12RHI/Public/D3D12Resources.h
--- a/Engine/Source/Runtime/D3D12RHI/Public/D3D12Resources.h
+++ b/Engine/Source/Runtime/D3D12RHI/Public/D3D12Resources.h
@@ -758,6 +758,8 @@ public:
 	};
 
 private:
+
+	FCriticalSection DeleteTaskCS;
 	TQueue<FAsyncTask<FD3D12AsyncDeletionWorker>*> DeleteTasks;
 
 };
