Log creation params on failure to create committed or placed resource in d3d12

diff --git a/Engine/Source/Runtime/D3D12RHI/Private/D3D12Resources.cpp b/Engine/Source/Runtime/D3D12RHI/Private/D3D12Resources.cpp
--- a/Engine/Source/Runtime/D3D12RHI/Private/D3D12Resources.cpp
+++ b/Engine/Source/Runtime/D3D12RHI/Private/D3D12Resources.cpp
@@ -510,11 +510,6 @@ HRESULT FD3D12Adapter::CreateCommittedResource(const FD3D12ResourceDesc& InDesc,
 #endif // D3D12_RHI_RAYTRACING
 
 	const HRESULT hr = RootDevice->CreateCommittedResource(&HeapProps, HeapFlags, &LocalDesc, InInitialState, ClearValue, IID_PPV_ARGS(pResource.GetInitReference()));
-	if (bVerifyHResult)
-	{
-		VERIFYD3D12RESULT_EX(hr, RootDevice);
-	}
-
 	if (SUCCEEDED(hr))
 	{
 		// Set the output pointer
@@ -530,6 +525,16 @@ HRESULT FD3D12Adapter::CreateCommittedResource(const FD3D12ResourceDesc& InDesc,
 			(*ppOutResource)->StartTrackingForResidency();
 		}
 	}
+	else	
+	{
+		UE_LOG(LogD3D12RHI, Display, TEXT("D3D12 CreateCommittedResource failed with params:\n\tHeap Type: %d\n\tHeap Flags: %d\n\tResource Dimension: %d\n\tResource Width: %d\n\tResource Height: %d\n\tFormat: %d\n\tResource Flags: %d"),
+			HeapProps.Type, HeapFlags, LocalDesc.Dimension, LocalDesc.Width, LocalDesc.Height, LocalDesc.PixelFormat, LocalDesc.Flags);
+
+		if (bVerifyHResult)
+		{
+			VERIFYD3D12RESULT_EX(hr, RootDevice);
+		}
+	}
 
 	return hr;
 }
@@ -546,12 +551,6 @@ HRESULT FD3D12Adapter::CreatePlacedResource(const FD3D12ResourceDesc& InDesc, FD
 
 	TRefCountPtr<ID3D12Resource> pResource;
 	const HRESULT hr = RootDevice->CreatePlacedResource(Heap, HeapOffset, &InDesc, InInitialState, ClearValue, IID_PPV_ARGS(pResource.GetInitReference()));
-
-	if (bVerifyHResult)
-	{
-		VERIFYD3D12RESULT_EX(hr, RootDevice);
-	}
-
 	if (SUCCEEDED(hr))
 	{
 		FD3D12Device* Device = BackingHeap->GetParentDevice();
@@ -591,6 +590,16 @@ HRESULT FD3D12Adapter::CreatePlacedResource(const FD3D12ResourceDesc& InDesc, FD
 
 		(*ppOutResource)->AddRef();
 	}
+	else
+	{
+		UE_LOG(LogD3D12RHI, Display, TEXT("D3D12 CreatePlacedResource failed with params:\n\tHeap Type: %d\n\tHeap Flags: %d\n\tResource Dimension: %d\n\tResource Width: %d\n\tResource Height: %d\n\tHeightFormat: %d\n\tResource Flags: %d"),
+			BackingHeap->GetHeapDesc().Properties.Type, BackingHeap->GetHeapDesc().Flags, InDesc.Dimension, InDesc.Width, InDesc.Height, InDesc.PixelFormat, InDesc.Flags);
+
+		if (bVerifyHResult)
+		{
+			VERIFYD3D12RESULT_EX(hr, RootDevice);
+		}
+	}
 
 	return hr;
 }
