Fix a crash due to trying to create a buffer from a heap that denies buffers

diff --git a/Engine/Source/Runtime/D3D12RHI/Private/D3D12Resources.cpp b/Engine/Source/Runtime/D3D12RHI/Private/D3D12Resources.cpp
--- a/Engine/Source/Runtime/D3D12RHI/Private/D3D12Resources.cpp
+++ b/Engine/Source/Runtime/D3D12RHI/Private/D3D12Resources.cpp
@@ -437,16 +437,18 @@ void FD3D12Heap::SetHeap(ID3D12Heap* HeapIn, const TCHAR* const InName, bool bIn
 	HeapDesc = Heap->GetDesc();
 
 	SetName(HeapIn, InName);
-	
+
 	// Create a buffer placed resource on the heap to extract the gpu virtual address
 	// if we are tracking all allocations
 	FD3D12Adapter* Adapter = GetParentDevice()->GetParentAdapter();	
-	if ((bForceGetGPUAddress || Adapter->IsTrackingAllAllocations()) && HeapDesc.Properties.Type == D3D12_HEAP_TYPE_DEFAULT)
+	if ((bForceGetGPUAddress || Adapter->IsTrackingAllAllocations())
+		&& !(HeapDesc.Flags & D3D12_HEAP_FLAG_DENY_BUFFERS)
+		&& HeapDesc.Properties.Type == D3D12_HEAP_TYPE_DEFAULT)
 	{
 		uint64 HeapSize = HeapDesc.SizeInBytes;
 		TRefCountPtr<ID3D12Resource> TempResource;
 		const D3D12_RESOURCE_DESC BufDesc = CD3DX12_RESOURCE_DESC::Buffer(HeapSize, D3D12_RESOURCE_FLAG_NONE);
-		const HRESULT hr = Adapter->GetD3DDevice()->CreatePlacedResource(Heap, 0, &BufDesc, D3D12_RESOURCE_STATE_COMMON, nullptr, IID_PPV_ARGS(TempResource.GetInitReference()));
+		VERIFYD3D12RESULT(Adapter->GetD3DDevice()->CreatePlacedResource(Heap, 0, &BufDesc, D3D12_RESOURCE_STATE_COMMON, nullptr, IID_PPV_ARGS(TempResource.GetInitReference())));
 		GPUVirtualAddress = TempResource->GetGPUVirtualAddress();
 				
 #if TRACK_RESOURCE_ALLOCATIONS
