D3D12 - Added support for the "-AllowSoftwareRendering" command line option
* Added for consistency with D3D11 which already supports this option
* Differ from the D3D12 specific "-warp" option, which forces the use of a WARP device. Also, WARP devices were ignored without this flag.

diff --git a/Engine/Source/Runtime/D3D12RHI/Private/D3D12RHIPrivate.h b/Engine/Source/Runtime/D3D12RHI/Private/D3D12RHIPrivate.h
--- a/Engine/Source/Runtime/D3D12RHI/Private/D3D12RHIPrivate.h
+++ b/Engine/Source/Runtime/D3D12RHI/Private/D3D12RHIPrivate.h
@@ -180,6 +180,12 @@ static bool D3D12RHI_ShouldCreateWithWarp()
 	return bCreateWithWarp;
 }
 
+static bool D3D12RHI_AllowSoftwareFallback()
+{
+	static bool bAllowSoftwareRendering = FParse::Param(FCommandLine::Get(), TEXT("AllowSoftwareRendering"));
+	return bAllowSoftwareRendering;
+}
+
 static bool D3D12RHI_ShouldAllowAsyncResourceCreation()
 {
 	static bool bAllowAsyncResourceCreation = !FParse::Param(FCommandLine::Get(), TEXT("nod3dasync"));
diff --git a/Engine/Source/Runtime/D3D12RHI/Private/Windows/WindowsD3D12Device.cpp b/Engine/Source/Runtime/D3D12RHI/Private/Windows/WindowsD3D12Device.cpp
--- a/Engine/Source/Runtime/D3D12RHI/Private/Windows/WindowsD3D12Device.cpp
+++ b/Engine/Source/Runtime/D3D12RHI/Private/Windows/WindowsD3D12Device.cpp
@@ -42,6 +42,7 @@ IMPLEMENT_MODULE(FD3D12DynamicRHIModule, D3D12RHI);
 
 extern bool D3D12RHI_ShouldCreateWithD3DDebug();
 extern bool D3D12RHI_ShouldCreateWithWarp();
+extern bool D3D12RHI_AllowSoftwareFallback();
 extern bool D3D12RHI_ShouldAllowAsyncResourceCreation();
 extern bool D3D12RHI_ShouldForceCompatibility();
 
@@ -481,18 +482,19 @@ void FD3D12DynamicRHIModule::FindAdapter()
 	int32 CVarExplicitAdapterValue = HmdGraphicsAdapterLuid == 0 ? (CVarGraphicsAdapter ? CVarGraphicsAdapter->GetValueOnGameThread() : -1) : -2;
 	FParse::Value(FCommandLine::Get(), TEXT("graphicsadapter="), CVarExplicitAdapterValue);
 
-	const bool bFavorNonIntegrated = CVarExplicitAdapterValue == -1;
+	const bool bFavorDiscreteAdapter = CVarExplicitAdapterValue == -1;
 
 	TRefCountPtr<IDXGIAdapter> TempAdapter;
 	const D3D_FEATURE_LEVEL MinRequiredFeatureLevel = GetRequiredD3DFeatureLevel();
 
-	FD3D12AdapterDesc FirstWithoutIntegratedAdapter;
+	FD3D12AdapterDesc FirstDiscreteAdapter;
 	FD3D12AdapterDesc FirstAdapter;
 
 	bool bIsAnyAMD = false;
 	bool bIsAnyIntel = false;
 	bool bIsAnyNVIDIA = false;
 	bool bRequestedWARP = D3D12RHI_ShouldCreateWithWarp();
+	bool bAllowSoftwareRendering = D3D12RHI_AllowSoftwareFallback();
 
 	int PreferredVendor = D3D12RHI_PreferAdapterVendor();
 	// Enumerate the DXGIFactory's adapters.
@@ -557,8 +559,8 @@ void FD3D12DynamicRHIModule::FindAdapter()
 
 				FD3D12AdapterDesc CurrentAdapter(AdapterDesc, AdapterIndex, MaxSupportedFeatureLevel, MaxSupportedShaderModel, NumNodes, bIsIntegrated);
 				
-				// If requested WARP, then reject all other adapters. If WARP not requested, then reject the WARP device.
-				const bool bSkipRequestedWARP = (bRequestedWARP && !bIsWARP) || (!bRequestedWARP && bIsWARP);
+				// If requested WARP, then reject all other adapters. If WARP not requested, then reject the WARP device if software rendering support is disallowed
+				const bool bSkipWARP = (bRequestedWARP && !bIsWARP) || (!bRequestedWARP && bIsWARP && !bAllowSoftwareRendering);
 
 				// we don't allow the PerfHUD adapter
 				const bool bSkipPerfHUDAdapter = bIsPerfHUD && !bAllowPerfHUD;
@@ -569,17 +571,17 @@ void FD3D12DynamicRHIModule::FindAdapter()
 				// the user wants a specific adapter, not this one
 				const bool bSkipExplicitAdapter = CVarExplicitAdapterValue >= 0 && AdapterIndex != CVarExplicitAdapterValue;
 
-				const bool bSkipAdapter = bSkipRequestedWARP || bSkipPerfHUDAdapter || bSkipHmdGraphicsAdapter || bSkipExplicitAdapter;
+				const bool bSkipAdapter = bSkipWARP || bSkipPerfHUDAdapter || bSkipHmdGraphicsAdapter || bSkipExplicitAdapter;
 
 				if (!bSkipAdapter)
 				{
-					if (!bIsIntegrated && !FirstWithoutIntegratedAdapter.IsValid())
+					if (!bIsWARP && !bIsIntegrated && !FirstDiscreteAdapter.IsValid())
 					{
-						FirstWithoutIntegratedAdapter = CurrentAdapter;
+						FirstDiscreteAdapter = CurrentAdapter;
 					}
-					else if (PreferredVendor == AdapterDesc.VendorId && FirstWithoutIntegratedAdapter.IsValid())
+					else if (PreferredVendor == AdapterDesc.VendorId && FirstDiscreteAdapter.IsValid())
 					{
-						FirstWithoutIntegratedAdapter = CurrentAdapter;
+						FirstDiscreteAdapter = CurrentAdapter;
 					}
 
 					if (!FirstAdapter.IsValid())
@@ -607,12 +609,12 @@ void FD3D12DynamicRHIModule::FindAdapter()
 #endif // PLATFORM_DESKTOP
 
 	TSharedPtr<FD3D12Adapter> NewAdapter;
-	if (bFavorNonIntegrated)
+	if (bFavorDiscreteAdapter)
 	{
 		// We assume Intel is integrated graphics (slower than discrete) than NVIDIA or AMD cards and rather take a different one
-		if (FirstWithoutIntegratedAdapter.IsValid())
+		if (FirstDiscreteAdapter.IsValid())
 		{
-			NewAdapter = TSharedPtr<FD3D12Adapter>(new FWindowsD3D12Adapter(FirstWithoutIntegratedAdapter));
+			NewAdapter = TSharedPtr<FD3D12Adapter>(new FWindowsD3D12Adapter(FirstDiscreteAdapter));
 			ChosenAdapters.Add(NewAdapter);
 		}
 		else
