Fix D3D12 crash when using Windows versions older than 1903.

We were adding D3D12_RESOURCE_STATE_SHADING_RATE_SOURCE on transitions without checking if VRS is supported, which causes a fatal error on 1809 and earlier.

#rnx

diff --git a/Engine/Source/Runtime/D3D12RHI/Private/D3D12Util.cpp b/Engine/Source/Runtime/D3D12RHI/Private/D3D12Util.cpp
--- a/Engine/Source/Runtime/D3D12RHI/Private/D3D12Util.cpp
+++ b/Engine/Source/Runtime/D3D12RHI/Private/D3D12Util.cpp
@@ -1276,7 +1276,7 @@ D3D12_RESOURCE_STATES GetD3D12ResourceState(ERHIAccess InRHIAccess, bool InIsAsy
 				State |= D3D12_RESOURCE_STATE_DEPTH_READ;
 			}
 #if PLATFORM_SUPPORTS_VARIABLE_RATE_SHADING
-			if (EnumHasAnyFlags(InRHIAccess, ERHIAccess::ShadingRateSource))
+			if (EnumHasAnyFlags(InRHIAccess, ERHIAccess::ShadingRateSource) && GRHISupportsPipelineVariableRateShading)
 			{
 				State |= D3D12_RESOURCE_STATE_SHADING_RATE_SOURCE;
 			}
