Fix invalid source state for the first transition on large structured buffers.

FD3D12DynamicRHI::CreateD3D12Buffer assumed static placed buffers will require resource tracking and computed the initial state accordingly. However, is the buffer is read-only, FD3D12Resource::DetermineResourceStates forces tracking to off after the resource is created, and overrides the default state to D3D12_RESOURCE_STATE_GENERIC_READ, even though the resource was created in the state originally computed by CreateD3D12Buffer. This mismatch caused the source state to be incorrect when the buffer was transitioned to COPY_DEST for uploading the initial contents.

The solution is to make CreateD3D12Buffer aware of the fact that read-only resources don't need tracking, in which case it uses GetDefaultInitialResourceState() to determine the initial state, which corresponds to what DetermineResourceStates computes.

#rnx

diff --git a/Engine/Source/Runtime/D3D12RHI/Private/D3D12Buffer.cpp b/Engine/Source/Runtime/D3D12RHI/Private/D3D12Buffer.cpp
--- a/Engine/Source/Runtime/D3D12RHI/Private/D3D12Buffer.cpp
+++ b/Engine/Source/Runtime/D3D12RHI/Private/D3D12Buffer.cpp
@@ -523,13 +523,14 @@ FD3D12Buffer* FD3D12DynamicRHI::CreateD3D12Buffer(class FRHICommandListImmediate
 		? ED3D12ResourceStateMode::SingleState 
 		: ED3D12ResourceStateMode::Default;
 
-	// Does this resource support tracking?
 	const bool bIsDynamic = (Usage & BUF_AnyDynamic) ? true : false;
-	const bool bSupportResourceStateTracking = !bIsDynamic && FD3D12DefaultBufferAllocator::IsPlacedResource(Desc.Flags, StateMode, Alignment);
-
-	// Initial state is derived from the InResourceState if it's supports tracking
 	D3D12_HEAP_TYPE HeapType = bIsDynamic ? D3D12_HEAP_TYPE_UPLOAD : D3D12_HEAP_TYPE_DEFAULT;
 	const FD3D12Resource::FD3D12ResourceTypeHelper Type(Desc, HeapType);
+
+	// Does this resource support tracking?
+	const bool bSupportResourceStateTracking = !bIsDynamic && FD3D12DefaultBufferAllocator::IsPlacedResource(Desc.Flags, StateMode, Alignment) && Type.bWritable;
+
+	// Initial state is derived from the InResourceState if it supports tracking
 	D3D12_RESOURCE_STATES DesiredState = bSupportResourceStateTracking ? Type.GetOptimalInitialState(InResourceState, false) :
 		FD3D12DefaultBufferAllocator::GetDefaultInitialResourceState(HeapType, (EBufferUsageFlags)Usage, StateMode);
 
