Don't add transient render targets to STAT_RenderTargetMemoryXX in D3D12.

#rnx

diff --git a/Engine/Source/Runtime/D3D12RHI/Private/D3D12Texture.cpp b/Engine/Source/Runtime/D3D12RHI/Private/D3D12Texture.cpp
--- a/Engine/Source/Runtime/D3D12RHI/Private/D3D12Texture.cpp
+++ b/Engine/Source/Runtime/D3D12RHI/Private/D3D12Texture.cpp
@@ -352,7 +352,7 @@ void FD3D12TextureStats::D3D12TextureAllocated(TD3D12Texture2D<BaseResourceType>
 	if (D3D12Texture2D)
 	{
 		// Don't update state for virtual or transient textures	
-		if (!EnumHasAnyFlags(Texture.Flags, TexCreate_Virtual | TexCreate_Transient))
+		if (!EnumHasAnyFlags(Texture.Flags, TexCreate_Virtual | TexCreate_Transient) && !Texture.ResourceLocation.IsTransient())
 		{
 			TRACE_CPUPROFILER_EVENT_SCOPE(D3D12RHI::UpdateTextureStats);
 
@@ -386,7 +386,7 @@ void FD3D12TextureStats::D3D12TextureDeleted(TD3D12Texture2D<BaseResourceType>&
 	if (D3D12Texture2D)
 	{
 		// Don't update state for transient textures	
-		if(!EnumHasAnyFlags(Texture.Flags, TexCreate_Transient))
+		if(!EnumHasAnyFlags(Texture.Flags, TexCreate_Transient) && !Texture.ResourceLocation.IsTransient())
 		{
 			const D3D12_RESOURCE_DESC& Desc = D3D12Texture2D->GetDesc();
 			const int64 TextureSize = Texture.GetMemorySize();
@@ -416,7 +416,7 @@ void FD3D12TextureStats::D3D12TextureAllocated(FD3D12Texture3D& Texture)
 	if (D3D12Texture3D)
 	{
 		// Don't update state for virtual or transient textures	
-		if (!EnumHasAnyFlags(Texture.GetFlags(), TexCreate_Virtual | TexCreate_Transient))
+		if (!EnumHasAnyFlags(Texture.GetFlags(), TexCreate_Virtual | TexCreate_Transient) && !Texture.ResourceLocation.IsTransient())
 		{
 			const D3D12_RESOURCE_DESC& Desc = D3D12Texture3D->GetDesc();
 			const D3D12_RESOURCE_ALLOCATION_INFO AllocationInfo = Texture.GetParentDevice()->GetDevice()->GetResourceAllocationInfo(0, 1, &Desc);
@@ -443,7 +443,7 @@ void FD3D12TextureStats::D3D12TextureDeleted(FD3D12Texture3D& Texture)
 	if (D3D12Texture3D)
 	{
 		// Don't update state for transient textures	
-		if (!EnumHasAnyFlags(Texture.GetFlags(), TexCreate_Transient))
+		if (!EnumHasAnyFlags(Texture.GetFlags(), TexCreate_Transient) && !Texture.ResourceLocation.IsTransient())
 		{
 			const D3D12_RESOURCE_DESC& Desc = D3D12Texture3D->GetDesc();
 			const int64 TextureSize = Texture.GetMemorySize();
