D3D12:  Fix for DEBUG_RESOURCE_STATES logging

- add a case for aliasing barriers as it currently asserts when one is encountered

diff --git a/Engine/Source/Runtime/D3D12RHI/Private/D3D12Util.cpp b/Engine/Source/Runtime/D3D12RHI/Private/D3D12Util.cpp
--- a/Engine/Source/Runtime/D3D12RHI/Private/D3D12Util.cpp
+++ b/Engine/Source/Runtime/D3D12RHI/Private/D3D12Util.cpp
@@ -1206,6 +1206,10 @@ void LogResourceBarriers(uint32 NumBarriers, D3D12_RESOURCE_BARRIER* pBarriers,
 			UE_LOG(LogD3D12RHI, Log, TEXT("*** BARRIER (CmdList: %016llX) %u/%u: UAV Barrier"), pCommandList, i + 1, NumBarriers);
 			break;
 
+		case D3D12_RESOURCE_BARRIER_TYPE_ALIASING:
+			UE_LOG(LogD3D12RHI, Log, TEXT("*** BARRIER (CmdList: %016llX) %u/%u: Aliasing Barrier, %016llX -> %016llX"), pCommandList, i + 1, NumBarriers, currentBarrier.Aliasing.pResourceBefore, currentBarrier.Aliasing.pResourceAfter);
+			break;
+
 		default:
 			check(false);
 			break;
