- Small cleanup/refactor of d3d12 resource and memory stats
- Remove Windows d3d12 VRAM llm hooks because they weren't tracking all allocations correctly (will be soon replaced by proper Insight VRAM tracking)

diff --git a/Engine/Source/Runtime/D3D12RHI/Private/D3D12Buffer.cpp b/Engine/Source/Runtime/D3D12RHI/Private/D3D12Buffer.cpp
--- a/Engine/Source/Runtime/D3D12RHI/Private/D3D12Buffer.cpp
+++ b/Engine/Source/Runtime/D3D12RHI/Private/D3D12Buffer.cpp
@@ -406,7 +406,7 @@ FD3D12Buffer* FD3D12Adapter::CreateRHIBuffer(
 	// Don't track transient buffer stats here
 	if (!BufferOut->ResourceLocation.IsTransient())
 	{
-		UpdateBufferStats(GetBufferStats(InUsage), BufferOut->ResourceLocation.GetSize());
+		UpdateBufferStats((EBufferUsageFlags)InUsage, BufferOut->ResourceLocation.GetSize());
 	}
 
 	return BufferOut;
diff --git a/Engine/Source/Runtime/D3D12RHI/Private/D3D12RHICommon.h b/Engine/Source/Runtime/D3D12RHI/Private/D3D12RHICommon.h
--- a/Engine/Source/Runtime/D3D12RHI/Private/D3D12RHICommon.h
+++ b/Engine/Source/Runtime/D3D12RHI/Private/D3D12RHICommon.h
@@ -8,9 +8,9 @@ D3D12RHICommon.h: Common D3D12 RHI definitions for Windows.
 
 DECLARE_STATS_GROUP(TEXT("D3D12RHI"), STATGROUP_D3D12RHI, STATCAT_Advanced);
 DECLARE_STATS_GROUP(TEXT("D3D12RHI: Memory"), STATGROUP_D3D12Memory, STATCAT_Advanced);
+DECLARE_STATS_GROUP(TEXT("D3D12RHI: Memory Details"), STATGROUP_D3D12MemoryDetails, STATCAT_Advanced);
 DECLARE_STATS_GROUP(TEXT("D3D12RHI: Resources"), STATGROUP_D3D12Resources, STATCAT_Advanced);
-DECLARE_STATS_GROUP(TEXT("D3D12RHI: Allocations"), STATGROUP_D3D12Allocations, STATCAT_Advanced);
-DECLARE_STATS_GROUP(TEXT("D3D12RHI: Allocations (Buffer)"), STATGROUP_D3D12BufferAllocations, STATCAT_Advanced);
+DECLARE_STATS_GROUP(TEXT("D3D12RHI: Buffer Details"), STATGROUP_D3D12BufferDetails, STATCAT_Advanced);
 DECLARE_STATS_GROUP(TEXT("D3D12RHI: Pipeline State (PSO)"), STATGROUP_D3D12PipelineState, STATCAT_Advanced);
 DECLARE_STATS_GROUP(TEXT("D3D12RHI: Descriptor Heap (GPU Visible)"), STATGROUP_D3D12DescriptorHeap, STATCAT_Advanced);
 
diff --git a/Engine/Source/Runtime/D3D12RHI/Private/D3D12Stats.cpp b/Engine/Source/Runtime/D3D12RHI/Private/D3D12Stats.cpp
--- a/Engine/Source/Runtime/D3D12RHI/Private/D3D12Stats.cpp
+++ b/Engine/Source/Runtime/D3D12RHI/Private/D3D12Stats.cpp
@@ -281,22 +281,46 @@ float FD3D12EventNode::GetTiming()
 	return Result;
 }
 
-void UpdateBufferStats(FName Name, int64 RequestedSize)
+inline FName GetRHIBufferStats(EBufferUsageFlags Usage)
 {
-	INC_MEMORY_STAT_BY_FName(Name, RequestedSize);
-
-#if PLATFORM_WINDOWS
-	// this is a work-around on Windows. Due to the fact that there is no way
-	// to hook the actual d3d allocations it is very difficult to track memory
-	// in the normal way. The problem is that some buffers are allocated from
-	// the allocators and some are allocated from the device. Ideally this
-	// tracking would be moved to where the actual d3d resource is created and
-	// released and the tracking could be re-enabled in the buddy allocator.
-	// The problem is that the releasing of resources happens in a generic way
-	// (see FD3D12ResourceLocation) 
-	LLM_SCOPED_PAUSE_TRACKING_WITH_ENUM_AND_AMOUNT(ELLMTag::Meshes, RequestedSize, ELLMTracker::Default, ELLMAllocType::None);
-	LLM_SCOPED_PAUSE_TRACKING_WITH_ENUM_AND_AMOUNT(ELLMTag::GraphicsPlatform, RequestedSize, ELLMTracker::Platform, ELLMAllocType::None);
-#endif
+	if (Usage & BUF_VertexBuffer)
+	{
+		return GET_STATFNAME(STAT_VertexBufferMemory);
+	}
+	else if (Usage & BUF_IndexBuffer)
+	{
+		return GET_STATFNAME(STAT_IndexBufferMemory);
+	}
+	else if (Usage & BUF_AccelerationStructure)
+	{
+		return GET_STATFNAME(STAT_RTAccelerationStructureMemory);
+	}
+	else
+	{
+		return GET_STATFNAME(STAT_StructuredBufferMemory);
+	}
+}
+
+static FName GetD3D12BufferStat(EBufferUsageFlags InUsageFlags)
+{
+	if (EnumHasAnyFlags(InUsageFlags, BUF_UnorderedAccess))
+	{
+		return GET_STATFNAME(STAT_D3D12UAVBuffers);
+	}
+	else if (EnumHasAnyFlags(InUsageFlags, BUF_AccelerationStructure))
+	{
+		return GET_STATFNAME(STAT_D3D12RTBuffers);
+	}
+	else
+	{
+		return GET_STATFNAME(STAT_D3D12Buffer);
+	}
+}
+
+void UpdateBufferStats(EBufferUsageFlags InUsageFlags, int64 RequestedSize)
+{
+	INC_MEMORY_STAT_BY_FName(GetRHIBufferStats(InUsageFlags), RequestedSize);
+	INC_MEMORY_STAT_BY_FName(GetD3D12BufferStat(InUsageFlags), RequestedSize);
 }
 
 #if NV_AFTERMATH
diff --git a/Engine/Source/Runtime/D3D12RHI/Private/D3D12Stats.h b/Engine/Source/Runtime/D3D12RHI/Private/D3D12Stats.h
--- a/Engine/Source/Runtime/D3D12RHI/Private/D3D12Stats.h
+++ b/Engine/Source/Runtime/D3D12RHI/Private/D3D12Stats.h
@@ -65,41 +65,50 @@ DECLARE_MEMORY_STAT_EXTERN(TEXT("Demoted Video Memory"), STAT_D3D12DemotedVideoM
 DECLARE_MEMORY_STAT_EXTERN(TEXT("Used Video Memory"), STAT_D3D12UsedVideoMemory, STATGROUP_D3D12Memory, );
 DECLARE_MEMORY_STAT_EXTERN(TEXT("Used System Memory"), STAT_D3D12UsedSystemMemory, STATGROUP_D3D12Memory, );
 
-DECLARE_MEMORY_STAT_EXTERN(TEXT("RenderTarget StandAlone Allocated"), STAT_D3D12RenderTargetStandAloneAllocated, STATGROUP_D3D12Resources, );
-DECLARE_MEMORY_STAT_EXTERN(TEXT("UAVTexture StandAlone Allocated"), STAT_D3D12UAVTextureStandAloneAllocated, STATGROUP_D3D12Resources, );
-DECLARE_MEMORY_STAT_EXTERN(TEXT("Texture StandAlone Allocated"), STAT_D3D12TextureStandAloneAllocated, STATGROUP_D3D12Resources, );
-DECLARE_MEMORY_STAT_EXTERN(TEXT("UAVBuffer StandAlone Allocated"), STAT_D3D12UAVBufferStandAloneAllocated, STATGROUP_D3D12Resources, );
-DECLARE_MEMORY_STAT_EXTERN(TEXT("Buffer StandAlone Allocated"), STAT_D3D12BufferStandAloneAllocated, STATGROUP_D3D12Resources, );
-
-DECLARE_DWORD_ACCUMULATOR_STAT_EXTERN(TEXT("RenderTarget StandAlone Allocated"), STAT_D3D12RenderTargetStandAloneCount, STATGROUP_D3D12Resources, );
-DECLARE_DWORD_ACCUMULATOR_STAT_EXTERN(TEXT("UAVTexture StandAlone Allocated"), STAT_D3D12UAVTextureStandAloneCount, STATGROUP_D3D12Resources, );
-DECLARE_DWORD_ACCUMULATOR_STAT_EXTERN(TEXT("Texture StandAlone Allocated"), STAT_D3D12TextureStandAloneCount, STATGROUP_D3D12Resources, );
-DECLARE_DWORD_ACCUMULATOR_STAT_EXTERN(TEXT("UAVBuffer StandAlone Allocated"), STAT_D3D12UAVBufferStandAloneCount, STATGROUP_D3D12Resources, );
-DECLARE_DWORD_ACCUMULATOR_STAT_EXTERN(TEXT("Buffer StandAlone Allocated"), STAT_D3D12BufferStandAloneCount, STATGROUP_D3D12Resources, );
-
-DECLARE_MEMORY_STAT_EXTERN(TEXT("Texture Allocator Allocated"), STAT_D3D12TextureAllocatorAllocated, STATGROUP_D3D12Allocations, );
-DECLARE_MEMORY_STAT_EXTERN(TEXT("Texture Allocator Unused"), STAT_D3D12TextureAllocatorUnused, STATGROUP_D3D12Allocations, );
-DECLARE_DWORD_ACCUMULATOR_STAT_EXTERN(TEXT("Texture Allocator Allocated"), STAT_D3D12TextureAllocatorCount, STATGROUP_D3D12Allocations, );
-
-DECLARE_MEMORY_STAT_POOL_EXTERN(TEXT("View DescriptorHeap Memory (SRV, CBV, UAV)"), STAT_ViewOnlineDescriptorHeapMemory, STATGROUP_D3D12Allocations, FPlatformMemory::MCR_GPUSystem, );
-DECLARE_MEMORY_STAT_POOL_EXTERN(TEXT("Sampler DescriptorHeap Memory"), STAT_SamplerOnlineDescriptorHeapMemory, STATGROUP_D3D12Allocations, FPlatformMemory::MCR_GPUSystem, );
-
-DECLARE_MEMORY_STAT_EXTERN(TEXT("BufferPool Memory Allocated"), STAT_D3D12BufferPoolMemoryAllocated, STATGROUP_D3D12Allocations, );
-DECLARE_MEMORY_STAT_EXTERN(TEXT("BufferPool Memory Used"), STAT_D3D12BufferPoolMemoryUsed, STATGROUP_D3D12Allocations, );
-DECLARE_MEMORY_STAT_EXTERN(TEXT("BufferPool Memory Free"), STAT_D3D12BufferPoolMemoryFree, STATGROUP_D3D12BufferAllocations, );
-DECLARE_MEMORY_STAT_EXTERN(TEXT("BufferPool Memory Alignment Waste"), STAT_D3D12BufferPoolAlignmentWaste, STATGROUP_D3D12BufferAllocations, );
-DECLARE_DWORD_ACCUMULATOR_STAT_EXTERN(TEXT("BufferPool Page Count"), STAT_D3D12BufferPoolPageCount, STATGROUP_D3D12BufferAllocations, );
-DECLARE_DWORD_ACCUMULATOR_STAT_EXTERN(TEXT("BufferPool Full Pages"), STAT_D3D12BufferPoolFullPages, STATGROUP_D3D12BufferAllocations, );
-DECLARE_MEMORY_STAT_EXTERN(TEXT("Bufferool Fragmentation"), STAT_D3D12BufferPoolFragmentation, STATGROUP_D3D12BufferAllocations, );
-DECLARE_FLOAT_ACCUMULATOR_STAT_EXTERN(TEXT("Bufferool Fragmentation Percentage"), STAT_D3D12BufferPoolFragmentationPercentage, STATGROUP_D3D12BufferAllocations, );
-
-DECLARE_MEMORY_STAT_EXTERN(TEXT("UploadPool Memory Allocated"), STAT_D3D12UploadPoolMemoryAllocated, STATGROUP_D3D12Allocations, );
-DECLARE_MEMORY_STAT_EXTERN(TEXT("UploadPool Memory Used"), STAT_D3D12UploadPoolMemoryUsed, STATGROUP_D3D12Allocations, );
-
-DECLARE_MEMORY_STAT_EXTERN(TEXT("UploadPool Memory Free"), STAT_D3D12UploadPoolMemoryFree, STATGROUP_D3D12BufferAllocations, );
-DECLARE_MEMORY_STAT_EXTERN(TEXT("UploadPool Memory Alignment Waste"), STAT_D3D12UploadPoolAlignmentWaste, STATGROUP_D3D12BufferAllocations, );
-DECLARE_DWORD_ACCUMULATOR_STAT_EXTERN(TEXT("UploadPool Page Count"), STAT_D3D12UploadPoolPageCount, STATGROUP_D3D12BufferAllocations, );
-DECLARE_DWORD_ACCUMULATOR_STAT_EXTERN(TEXT("UploadPool Full Pages"), STAT_D3D12UploadPoolFullPages, STATGROUP_D3D12BufferAllocations, );
+DECLARE_MEMORY_STAT_EXTERN(TEXT("Render Targets"), STAT_D3D12RenderTargets, STATGROUP_D3D12Resources, );
+DECLARE_MEMORY_STAT_EXTERN(TEXT("UAV Textures"), STAT_D3D12UAVTextures, STATGROUP_D3D12Resources, );
+DECLARE_MEMORY_STAT_EXTERN(TEXT("Textures"), STAT_D3D12Textures, STATGROUP_D3D12Resources, );
+DECLARE_MEMORY_STAT_EXTERN(TEXT("UAV Buffers"), STAT_D3D12UAVBuffers, STATGROUP_D3D12Resources, );
+DECLARE_MEMORY_STAT_EXTERN(TEXT("RayTracing Buffers"), STAT_D3D12RTBuffers, STATGROUP_D3D12Resources, );
+DECLARE_MEMORY_STAT_EXTERN(TEXT("Buffers"), STAT_D3D12Buffer, STATGROUP_D3D12Resources, );
+DECLARE_MEMORY_STAT_EXTERN(TEXT("Transient Heaps"), STAT_D3D12TransientHeaps, STATGROUP_D3D12Resources, );
+
+DECLARE_MEMORY_STAT_EXTERN(TEXT("RenderTarget StandAlone Allocated"), STAT_D3D12RenderTargetStandAloneAllocated, STATGROUP_D3D12MemoryDetails, );
+DECLARE_MEMORY_STAT_EXTERN(TEXT("UAVTexture StandAlone Allocated"), STAT_D3D12UAVTextureStandAloneAllocated, STATGROUP_D3D12MemoryDetails, );
+DECLARE_MEMORY_STAT_EXTERN(TEXT("Texture StandAlone Allocated"), STAT_D3D12TextureStandAloneAllocated, STATGROUP_D3D12MemoryDetails, );
+DECLARE_MEMORY_STAT_EXTERN(TEXT("UAVBuffer StandAlone Allocated"), STAT_D3D12UAVBufferStandAloneAllocated, STATGROUP_D3D12MemoryDetails, );
+DECLARE_MEMORY_STAT_EXTERN(TEXT("Buffer StandAlone Allocated"), STAT_D3D12BufferStandAloneAllocated, STATGROUP_D3D12MemoryDetails, );
+
+DECLARE_DWORD_ACCUMULATOR_STAT_EXTERN(TEXT("RenderTarget StandAlone Allocated"), STAT_D3D12RenderTargetStandAloneCount, STATGROUP_D3D12MemoryDetails, );
+DECLARE_DWORD_ACCUMULATOR_STAT_EXTERN(TEXT("UAVTexture StandAlone Allocated"), STAT_D3D12UAVTextureStandAloneCount, STATGROUP_D3D12MemoryDetails, );
+DECLARE_DWORD_ACCUMULATOR_STAT_EXTERN(TEXT("Texture StandAlone Allocated"), STAT_D3D12TextureStandAloneCount, STATGROUP_D3D12MemoryDetails, );
+DECLARE_DWORD_ACCUMULATOR_STAT_EXTERN(TEXT("UAVBuffer StandAlone Allocated"), STAT_D3D12UAVBufferStandAloneCount, STATGROUP_D3D12MemoryDetails, );
+DECLARE_DWORD_ACCUMULATOR_STAT_EXTERN(TEXT("Buffer StandAlone Allocated"), STAT_D3D12BufferStandAloneCount, STATGROUP_D3D12MemoryDetails, );
+
+DECLARE_MEMORY_STAT_EXTERN(TEXT("Texture Allocator Allocated"), STAT_D3D12TextureAllocatorAllocated, STATGROUP_D3D12MemoryDetails, );
+DECLARE_MEMORY_STAT_EXTERN(TEXT("Texture Allocator Unused"), STAT_D3D12TextureAllocatorUnused, STATGROUP_D3D12MemoryDetails, );
+DECLARE_DWORD_ACCUMULATOR_STAT_EXTERN(TEXT("Texture Allocator Allocated"), STAT_D3D12TextureAllocatorCount, STATGROUP_D3D12MemoryDetails, );
+
+DECLARE_MEMORY_STAT_POOL_EXTERN(TEXT("View DescriptorHeap Memory (SRV, CBV, UAV)"), STAT_ViewOnlineDescriptorHeapMemory, STATGROUP_D3D12MemoryDetails, FPlatformMemory::MCR_GPUSystem, );
+DECLARE_MEMORY_STAT_POOL_EXTERN(TEXT("Sampler DescriptorHeap Memory"), STAT_SamplerOnlineDescriptorHeapMemory, STATGROUP_D3D12MemoryDetails, FPlatformMemory::MCR_GPUSystem, );
+
+DECLARE_MEMORY_STAT_EXTERN(TEXT("BufferPool Memory Allocated"), STAT_D3D12BufferPoolMemoryAllocated, STATGROUP_D3D12MemoryDetails, );
+DECLARE_MEMORY_STAT_EXTERN(TEXT("BufferPool Memory Used"), STAT_D3D12BufferPoolMemoryUsed, STATGROUP_D3D12MemoryDetails, );
+
+DECLARE_MEMORY_STAT_EXTERN(TEXT("UploadPool Memory Allocated"), STAT_D3D12UploadPoolMemoryAllocated, STATGROUP_D3D12MemoryDetails, );
+DECLARE_MEMORY_STAT_EXTERN(TEXT("UploadPool Memory Used"), STAT_D3D12UploadPoolMemoryUsed, STATGROUP_D3D12MemoryDetails, );
+
+DECLARE_MEMORY_STAT_EXTERN(TEXT("BufferPool Memory Free"), STAT_D3D12BufferPoolMemoryFree, STATGROUP_D3D12BufferDetails, );
+DECLARE_MEMORY_STAT_EXTERN(TEXT("BufferPool Memory Alignment Waste"), STAT_D3D12BufferPoolAlignmentWaste, STATGROUP_D3D12BufferDetails, );
+DECLARE_DWORD_ACCUMULATOR_STAT_EXTERN(TEXT("BufferPool Page Count"), STAT_D3D12BufferPoolPageCount, STATGROUP_D3D12BufferDetails, );
+DECLARE_DWORD_ACCUMULATOR_STAT_EXTERN(TEXT("BufferPool Full Pages"), STAT_D3D12BufferPoolFullPages, STATGROUP_D3D12BufferDetails, );
+DECLARE_MEMORY_STAT_EXTERN(TEXT("Bufferool Fragmentation"), STAT_D3D12BufferPoolFragmentation, STATGROUP_D3D12BufferDetails, );
+DECLARE_FLOAT_ACCUMULATOR_STAT_EXTERN(TEXT("Bufferool Fragmentation Percentage"), STAT_D3D12BufferPoolFragmentationPercentage, STATGROUP_D3D12BufferDetails, );
+
+DECLARE_MEMORY_STAT_EXTERN(TEXT("UploadPool Memory Free"), STAT_D3D12UploadPoolMemoryFree, STATGROUP_D3D12BufferDetails, );
+DECLARE_MEMORY_STAT_EXTERN(TEXT("UploadPool Memory Alignment Waste"), STAT_D3D12UploadPoolAlignmentWaste, STATGROUP_D3D12BufferDetails, );
+DECLARE_DWORD_ACCUMULATOR_STAT_EXTERN(TEXT("UploadPool Page Count"), STAT_D3D12UploadPoolPageCount, STATGROUP_D3D12BufferDetails, );
+DECLARE_DWORD_ACCUMULATOR_STAT_EXTERN(TEXT("UploadPool Full Pages"), STAT_D3D12UploadPoolFullPages, STATGROUP_D3D12BufferDetails, );
 
 /**
 * Detailed Descriptor heap stats
diff --git a/Engine/Source/Runtime/D3D12RHI/Private/D3D12Texture.cpp b/Engine/Source/Runtime/D3D12RHI/Private/D3D12Texture.cpp
--- a/Engine/Source/Runtime/D3D12RHI/Private/D3D12Texture.cpp
+++ b/Engine/Source/Runtime/D3D12RHI/Private/D3D12Texture.cpp
@@ -268,7 +268,7 @@ bool FD3D12TextureStats::ShouldCountAsTextureMemory(D3D12_RESOURCE_FLAGS MiscFla
 }
 
 // @param b3D true:3D, false:2D or cube map
-TStatId FD3D12TextureStats::GetD3D12StatEnum(D3D12_RESOURCE_FLAGS MiscFlags, bool bCubeMap, bool b3D)
+TStatId FD3D12TextureStats::GetRHIStatEnum(D3D12_RESOURCE_FLAGS MiscFlags, bool bCubeMap, bool b3D)
 {
 #if STATS
 	if (ShouldCountAsTextureMemory(MiscFlags))
@@ -307,6 +307,25 @@ TStatId FD3D12TextureStats::GetD3D12StatEnum(D3D12_RESOURCE_FLAGS MiscFlags, boo
 	return TStatId();
 }
 
+TStatId FD3D12TextureStats::GetD3D12StatEnum(D3D12_RESOURCE_FLAGS MiscFlags)
+{
+#if STATS
+	if (EnumHasAnyFlags(MiscFlags, D3D12_RESOURCE_FLAG_ALLOW_RENDER_TARGET | D3D12_RESOURCE_FLAG_ALLOW_DEPTH_STENCIL))
+	{
+		return GET_STATID(STAT_D3D12RenderTargets);
+	}
+	else if (EnumHasAnyFlags(MiscFlags, D3D12_RESOURCE_FLAG_ALLOW_UNORDERED_ACCESS))
+	{
+		return GET_STATID(STAT_D3D12UAVTextures);
+	}
+	else
+	{
+		return GET_STATID(STAT_D3D12Textures);
+	}
+#endif
+	return TStatId();
+}
+
 // Note: This function can be called from many different threads
 // @param TextureSize >0 to allocate, <0 to deallocate
 // @param b3D true:3D, false:2D or cube map
@@ -332,7 +351,8 @@ void FD3D12TextureStats::UpdateD3D12TextureStats(const D3D12_RESOURCE_DESC& Desc
 		FPlatformAtomics::InterlockedAdd(&GCurrentRendertargetMemorySize, AlignedSize);
 	}
 
-	INC_MEMORY_STAT_BY_FName(GetD3D12StatEnum(Desc.Flags, bCubeMap, b3D).GetName(), TextureSize);
+	INC_MEMORY_STAT_BY_FName(GetD3D12StatEnum(Desc.Flags).GetName(), TextureSize);
+	INC_MEMORY_STAT_BY_FName(GetRHIStatEnum(Desc.Flags, bCubeMap, b3D).GetName(), TextureSize);
 
 	if (TextureSize > 0)
 	{
@@ -352,7 +372,7 @@ void FD3D12TextureStats::D3D12TextureAllocated(TD3D12Texture2D<BaseResourceType>
 	if (D3D12Texture2D)
 	{
 		// Don't update state for virtual or transient textures	
-		if (!EnumHasAnyFlags(Texture.Flags, TexCreate_Virtual | TexCreate_Transient) && !Texture.ResourceLocation.IsTransient())
+		if (!EnumHasAnyFlags(Texture.Flags, TexCreate_Virtual) && !Texture.ResourceLocation.IsTransient())
 		{
 			TRACE_CPUPROFILER_EVENT_SCOPE(D3D12RHI::UpdateTextureStats);
 
@@ -367,13 +387,10 @@ void FD3D12TextureStats::D3D12TextureAllocated(TD3D12Texture2D<BaseResourceType>
 			Texture.SetMemorySize(TextureSize);
 
 			UpdateD3D12TextureStats(*Desc, TextureSize, false, Texture.IsCubemap(), Texture.IsStreamable());
-
-#if PLATFORM_WINDOWS
-			// On Windows there is no way to hook into the low level d3d allocations and frees.
-			// This means that we must manually add the tracking here.
-			LLM(FLowLevelMemTracker::Get().OnLowLevelAlloc(ELLMTracker::Platform, Texture.GetResource()->GetResource(), Texture.GetMemorySize(), ELLMTag::GraphicsPlatform));
-			LLM(FLowLevelMemTracker::Get().OnLowLevelAlloc(ELLMTracker::Default, Texture.GetResource()->GetResource(), Texture.GetMemorySize(), ELLMTag::Textures));
-#endif
+		}
+		else
+		{
+			Texture.SetMemorySize(Texture.ResourceLocation.GetSize());
 		}
 	}
 }
@@ -386,20 +403,13 @@ void FD3D12TextureStats::D3D12TextureDeleted(TD3D12Texture2D<BaseResourceType>&
 	if (D3D12Texture2D)
 	{
 		// Don't update state for transient textures	
-		if(!EnumHasAnyFlags(Texture.Flags, TexCreate_Transient) && !Texture.ResourceLocation.IsTransient())
+		if (!Texture.ResourceLocation.IsTransient())
 		{
 			const D3D12_RESOURCE_DESC& Desc = D3D12Texture2D->GetDesc();
 			const int64 TextureSize = Texture.GetMemorySize();
 			ensure(TextureSize > 0 || EnumHasAnyFlags(Texture.Flags, TexCreate_Virtual) || Texture.GetAliasingSourceTexture() != nullptr);
 
 			UpdateD3D12TextureStats(Desc, -TextureSize, false, Texture.IsCubemap(), Texture.IsStreamable());
-
-#if PLATFORM_WINDOWS
-			// On Windows there is no way to hook into the low level d3d allocations and frees.
-			// This means that we must manually add the tracking here.
-			LLM(FLowLevelMemTracker::Get().OnLowLevelFree(ELLMTracker::Platform, Texture.GetResource()->GetResource()));
-			LLM(FLowLevelMemTracker::Get().OnLowLevelFree(ELLMTracker::Default, Texture.GetResource()->GetResource()));
-#endif
 		}
 	}
 }
@@ -416,7 +426,7 @@ void FD3D12TextureStats::D3D12TextureAllocated(FD3D12Texture3D& Texture)
 	if (D3D12Texture3D)
 	{
 		// Don't update state for virtual or transient textures	
-		if (!EnumHasAnyFlags(Texture.GetFlags(), TexCreate_Virtual | TexCreate_Transient) && !Texture.ResourceLocation.IsTransient())
+		if (!EnumHasAnyFlags(Texture.GetFlags(), TexCreate_Virtual) && !Texture.ResourceLocation.IsTransient())
 		{
 			const D3D12_RESOURCE_DESC& Desc = D3D12Texture3D->GetDesc();
 			const D3D12_RESOURCE_ALLOCATION_INFO AllocationInfo = Texture.GetParentDevice()->GetDevice()->GetResourceAllocationInfo(0, 1, &Desc);
@@ -425,13 +435,10 @@ void FD3D12TextureStats::D3D12TextureAllocated(FD3D12Texture3D& Texture)
 			Texture.SetMemorySize(TextureSize);
 
 			UpdateD3D12TextureStats(Desc, TextureSize, true, false, Texture.IsStreamable());
-
-#if PLATFORM_WINDOWS
-			// On Windows there is no way to hook into the low level d3d allocations and frees.
-			// This means that we must manually add the tracking here.
-			LLM(FLowLevelMemTracker::Get().OnLowLevelAlloc(ELLMTracker::Platform, Texture.GetResource()->GetResource(), Texture.GetMemorySize(), ELLMTag::GraphicsPlatform));
-			LLM(FLowLevelMemTracker::Get().OnLowLevelAlloc(ELLMTracker::Default, Texture.GetResource()->GetResource(), Texture.GetMemorySize(), ELLMTag::Textures));
-#endif
+		}
+		else
+		{
+			Texture.SetMemorySize(Texture.ResourceLocation.GetSize());
 		}
 	}
 }
@@ -443,20 +450,13 @@ void FD3D12TextureStats::D3D12TextureDeleted(FD3D12Texture3D& Texture)
 	if (D3D12Texture3D)
 	{
 		// Don't update state for transient textures	
-		if (!EnumHasAnyFlags(Texture.GetFlags(), TexCreate_Transient) && !Texture.ResourceLocation.IsTransient())
+		if (!Texture.ResourceLocation.IsTransient())
 		{
 			const D3D12_RESOURCE_DESC& Desc = D3D12Texture3D->GetDesc();
 			const int64 TextureSize = Texture.GetMemorySize();
 			if (TextureSize > 0)
 			{
 				UpdateD3D12TextureStats(Desc, -TextureSize, true, false, Texture.IsStreamable());
-
-#if PLATFORM_WINDOWS
-				// On Windows there is no way to hook into the low level d3d allocations and frees.
-				// This means that we must manually add the tracking here.
-				LLM(FLowLevelMemTracker::Get().OnLowLevelFree(ELLMTracker::Platform, Texture.GetResource()->GetResource()));
-				LLM(FLowLevelMemTracker::Get().OnLowLevelFree(ELLMTracker::Default, Texture.GetResource()->GetResource()));
-#endif
 			}
 		}
 	}
diff --git a/Engine/Source/Runtime/D3D12RHI/Private/D3D12Texture.h b/Engine/Source/Runtime/D3D12RHI/Private/D3D12Texture.h
--- a/Engine/Source/Runtime/D3D12RHI/Private/D3D12Texture.h
+++ b/Engine/Source/Runtime/D3D12RHI/Private/D3D12Texture.h
@@ -620,8 +620,10 @@ class FD3D12TextureStats
 public:
 
 	static bool ShouldCountAsTextureMemory(D3D12_RESOURCE_FLAGS MiscFlags);
+
 	// @param b3D true:3D, false:2D or cube map
-	static TStatId GetD3D12StatEnum(D3D12_RESOURCE_FLAGS MiscFlags, bool bCubeMap, bool b3D);
+	static TStatId GetRHIStatEnum(D3D12_RESOURCE_FLAGS MiscFlags, bool bCubeMap, bool b3D);
+	static TStatId GetD3D12StatEnum(D3D12_RESOURCE_FLAGS MiscFlags);
 
 	// Note: This function can be called from many different threads
 	// @param TextureSize >0 to allocate, <0 to deallocate
diff --git a/Engine/Source/Runtime/D3D12RHI/Private/D3D12TransientResourceAllocator.cpp b/Engine/Source/Runtime/D3D12RHI/Private/D3D12TransientResourceAllocator.cpp
--- a/Engine/Source/Runtime/D3D12RHI/Private/D3D12TransientResourceAllocator.cpp
+++ b/Engine/Source/Runtime/D3D12RHI/Private/D3D12TransientResourceAllocator.cpp
@@ -2,6 +2,7 @@
 
 #include "D3D12RHIPrivate.h"
 #include "D3D12TransientResourceAllocator.h"
+#include "D3D12Stats.h"
 
 D3D12_RESOURCE_STATES GetInitialResourceState(const D3D12_RESOURCE_DESC& InDesc)
 {
@@ -88,11 +89,19 @@ FD3D12TransientHeap::FD3D12TransientHeap(const FRHITransientHeapInitializer& Ini
 	Heap->BeginTrackingResidency(Desc.SizeInBytes);
 
 	BaseGPUVirtualAddress = Heap->GetGPUVirtualAddress();
+
+	INC_MEMORY_STAT_BY(STAT_D3D12TransientHeaps, Desc.SizeInBytes);
 }
 
 FD3D12TransientHeap::~FD3D12TransientHeap()
 {
 	LLM_SCOPED_PAUSE_TRACKING_FOR_TRACKER(ELLMTracker::Default, ELLMAllocType::System);
+
+	if (Heap)
+	{
+		D3D12_HEAP_DESC Desc = Heap->GetHeapDesc();
+		DEC_MEMORY_STAT_BY(STAT_D3D12TransientHeaps, Desc.SizeInBytes);
+	}
 }
 
 TUniquePtr<FD3D12TransientResourceSystem> FD3D12TransientResourceSystem::Create(FD3D12Adapter* ParentAdapter, FRHIGPUMask VisibleNodeMask)
diff --git a/Engine/Source/Runtime/D3D12RHI/Private/D3D12UniformBuffer.cpp b/Engine/Source/Runtime/D3D12RHI/Private/D3D12UniformBuffer.cpp
--- a/Engine/Source/Runtime/D3D12RHI/Private/D3D12UniformBuffer.cpp
+++ b/Engine/Source/Runtime/D3D12RHI/Private/D3D12UniformBuffer.cpp
@@ -88,7 +88,7 @@ FUniformBufferRHIRef FD3D12DynamicRHI::RHICreateUniformBuffer(const void* Conten
 		}
 	}
 
-	UpdateBufferStats(GET_STATFNAME(STAT_UniformBufferMemory), UniformBufferOut->ResourceLocation.GetSize());
+	INC_MEMORY_STAT_BY(STAT_UniformBufferMemory, UniformBufferOut->ResourceLocation.GetSize());
 
 	return UniformBufferOut;
 }
@@ -213,7 +213,7 @@ FD3D12UniformBuffer::~FD3D12UniformBuffer()
 	check(IsInRHIThread() || IsInRenderingThread());
 
 	int64 BufferSize = ResourceLocation.GetSize();
-	UpdateBufferStats(GET_STATFNAME(STAT_UniformBufferMemory), -BufferSize);
+	DEC_MEMORY_STAT_BY(STAT_UniformBufferMemory, BufferSize);
 
 #if USE_STATIC_ROOT_SIGNATURE
 	delete View;
diff --git a/Engine/Source/Runtime/D3D12RHI/Private/D3D12Util.cpp b/Engine/Source/Runtime/D3D12RHI/Private/D3D12Util.cpp
--- a/Engine/Source/Runtime/D3D12RHI/Private/D3D12Util.cpp
+++ b/Engine/Source/Runtime/D3D12RHI/Private/D3D12Util.cpp
@@ -1613,6 +1613,14 @@ DEFINE_STAT(STAT_D3D12AvailableVideoMemory);
 DEFINE_STAT(STAT_D3D12DemotedVideoMemory);
 DEFINE_STAT(STAT_D3D12TotalVideoMemory);
 
+DEFINE_STAT(STAT_D3D12RenderTargets);
+DEFINE_STAT(STAT_D3D12UAVTextures);
+DEFINE_STAT(STAT_D3D12Textures);
+DEFINE_STAT(STAT_D3D12UAVBuffers);
+DEFINE_STAT(STAT_D3D12RTBuffers);
+DEFINE_STAT(STAT_D3D12Buffer);
+DEFINE_STAT(STAT_D3D12TransientHeaps);
+
 DEFINE_STAT(STAT_D3D12RenderTargetStandAloneAllocated);
 DEFINE_STAT(STAT_D3D12UAVTextureStandAloneAllocated);
 DEFINE_STAT(STAT_D3D12TextureStandAloneAllocated);
diff --git a/Engine/Source/Runtime/D3D12RHI/Public/D3D12Resources.h b/Engine/Source/Runtime/D3D12RHI/Public/D3D12Resources.h
--- a/Engine/Source/Runtime/D3D12RHI/Public/D3D12Resources.h
+++ b/Engine/Source/Runtime/D3D12RHI/Public/D3D12Resources.h
@@ -848,8 +848,8 @@ public:
 
 	void Swap(FD3D12BaseShaderResource& Other)
 	{
-		// assume RHI thread when swapping listeners and resources unless RHI thread is disabled
-		check(IsInRHIThread() || IsInRenderingThread());
+		// assume RHI thread when swapping listeners and resources
+		check(!IsRunningRHIInSeparateThread() || IsInRHIThread());
 
 		::Swap(Parent, Other.Parent);
 		ResourceLocation.Swap(Other.ResourceLocation);
@@ -891,27 +891,7 @@ public:
 	}
 };
 
-extern void UpdateBufferStats(FName Name, int64 RequestedSize);
-
-inline FName GetBufferStats(uint32 Usage)
-{
-	if (Usage & BUF_VertexBuffer)
-	{
-		return GET_STATFNAME(STAT_VertexBufferMemory);
-	}
-	else if (Usage & BUF_IndexBuffer)
-	{
-		return GET_STATFNAME(STAT_IndexBufferMemory);
-	}
-	else if (Usage & BUF_AccelerationStructure)
-	{
-		return GET_STATFNAME(STAT_RTAccelerationStructureMemory);
-	}
-	else
-	{
-		return GET_STATFNAME(STAT_StructuredBufferMemory);
-	}
-}
+extern void UpdateBufferStats(EBufferUsageFlags UsageFlags, int64 RequestedSize);
 
 /** Uniform buffer resource class. */
 class FD3D12UniformBuffer : public FRHIUniformBuffer, public FD3D12DeviceChild, public FD3D12LinkedAdapterObject<FD3D12UniformBuffer>
@@ -966,7 +946,7 @@ public:
 		bool bTransient = ResourceLocation.IsTransient();
 		if (!bTransient)
 		{
-			UpdateBufferStats(GetBufferStats(GetUsage()), -BufferSize);
+			UpdateBufferStats((EBufferUsageFlags)GetUsage(), -BufferSize);
 		}
 	}
 
