Flush command list to D3D if more than certain amount of commands have been enqueued (helps with TDR issues on very BIG command lists)

diff --git a/Engine/Source/Runtime/D3D12RHI/Private/D3D12Buffer.cpp b/Engine/Source/Runtime/D3D12RHI/Private/D3D12Buffer.cpp
--- a/Engine/Source/Runtime/D3D12RHI/Private/D3D12Buffer.cpp
+++ b/Engine/Source/Runtime/D3D12RHI/Private/D3D12Buffer.cpp
@@ -792,6 +792,8 @@ void FD3D12DynamicRHI::RHICopyBuffer(FRHIBuffer* SourceBufferRHI, FRHIBuffer* De
 		Context.CommandListHandle.UpdateResidency(pDestResource);
 		Context.CommandListHandle.UpdateResidency(pSourceResource);
 
+		Context.ConditionalFlushCommandList();
+
 		DEBUG_EXECUTE_COMMAND_CONTEXT(Device->GetDefaultCommandContext());
 
 		Device->RegisterGPUWork(1);
@@ -829,6 +831,8 @@ void FD3D12CommandContext::RHICopyBufferRegion(FRHIBuffer* DestBufferRHI, uint64
 	CommandListHandle.UpdateResidency(pDestResource);
 	CommandListHandle.UpdateResidency(pSourceResource);
 
+	ConditionalFlushCommandList();
+
 	Device->RegisterGPUWork(1);
 }
 
@@ -968,5 +972,7 @@ void FD3D12CommandContext::RHICopyBufferRegions(const TArrayView<const FCopyBuff
 
 	TransitionResources(CommandListHandle, SrcBuffers, EBatchCopyState::FinalizeSource);
 	TransitionResources(CommandListHandle, DstBuffers, EBatchCopyState::FinalizeDest);
+	
+	ConditionalFlushCommandList();
 }
 #endif // D3D12_RHI_RAYTRACING
diff --git a/Engine/Source/Runtime/D3D12RHI/Private/D3D12CommandContext.cpp b/Engine/Source/Runtime/D3D12RHI/Private/D3D12CommandContext.cpp
--- a/Engine/Source/Runtime/D3D12RHI/Private/D3D12CommandContext.cpp
+++ b/Engine/Source/Runtime/D3D12RHI/Private/D3D12CommandContext.cpp
@@ -27,11 +27,11 @@ static FAutoConsoleVariableRef CVarCommandListBatchingMode(
 	ECVF_RenderThreadSafe
 );
 
-int32 MaxInitialResourceCopiesPerCommandList = 10000;
-static FAutoConsoleVariableRef CVarMaxInitialResourceCopiesPerCommandList(
-	TEXT("D3D12.MaxInitialResourceCopiesPerCommandList"),
-	MaxInitialResourceCopiesPerCommandList,
-	TEXT("Flush command list to GPU after certain amount of enqueued initial resource copy operations (default value 10000)"),
+int32 MaxCommandsPerCommandList = 10000;
+static FAutoConsoleVariableRef CVarMaxCommandsPerCommandList(
+	TEXT("D3D12.MaxCommandsPerCommandList"),
+	MaxCommandsPerCommandList,
+	TEXT("Flush command list to GPU after certain amount of enqueued commands (draw, dispatch, copy, ...) (default value 10000)"),
 	ECVF_RenderThreadSafe
 );
 
@@ -326,6 +326,8 @@ void FD3D12CommandContext::OpenCommandList()
 	// Mark state as dirty so next time ApplyState is called, it will set all state on this new command list
 	StateCache.DirtyStateForNewCommandList();
 
+	bIsDoingQuery = false;
+
 	numDraws = 0;
 	numDispatches = 0;
 	numClears = 0;
@@ -346,6 +348,9 @@ FD3D12CommandListHandle FD3D12CommandContext::FlushCommands(bool WaitForCompleti
 	// We should only be flushing the default context
 	check(IsDefaultContext());
 
+	// We should be in a query anymore
+	check(!bIsDoingQuery);
+
 	bool bHasProfileGPUAction = false;
 #if WITH_PROFILEGPU || D3D12_SUBMISSION_GAP_RECORDER
 	// Only graphics command list supports ID3D12GraphicsCommandList::EndQuery currently
@@ -423,10 +428,12 @@ FD3D12CommandListHandle FD3D12CommandContext::FlushCommands(bool WaitForCompleti
 
 void FD3D12CommandContext::ConditionalFlushCommandList()
 {
-	// Flush command list of reached maximum amount of initial resource command list which can be done in a single
+	// Flush command list of reached maximum amount of commands which can be done in a single
 	// command list - too many can cause TDRs
-	if (MaxInitialResourceCopiesPerCommandList > 0 && numInitialResourceCopies > (uint32)MaxInitialResourceCopiesPerCommandList)
+	// (can't flush when query is open!)
+	if (IsDefaultContext() && !bIsDoingQuery && MaxCommandsPerCommandList > 0 && GetTotalWorkCount() > (uint32)MaxCommandsPerCommandList)
 	{
+		UE_LOG(LogD3D12RHI, Warning, TEXT("Force flushing command list to GPU because too many commands have been enqueued already (%d commands)"), GetTotalWorkCount());
 		FlushCommands();
 	}
 }
diff --git a/Engine/Source/Runtime/D3D12RHI/Private/D3D12CommandContext.h b/Engine/Source/Runtime/D3D12RHI/Private/D3D12CommandContext.h
--- a/Engine/Source/Runtime/D3D12RHI/Private/D3D12CommandContext.h
+++ b/Engine/Source/Runtime/D3D12RHI/Private/D3D12CommandContext.h
@@ -193,6 +193,8 @@ public:
 	virtual void FlushTextureCache() {};
 #endif
 
+	bool bIsDoingQuery = false;
+
 	uint32 numDraws;
 	uint32 numDispatches;
 	uint32 numClears;
@@ -202,9 +204,14 @@ public:
 	uint32 numInitialResourceCopies;
 	uint32 otherWorkCounter;
 
+	uint32 GetTotalWorkCount() const
+	{
+		return numDraws + numDispatches + numClears + numBarriers + numPendingBarriers + numCopies + numInitialResourceCopies + otherWorkCounter;
+	}
+
 	bool HasDoneWork() const
 	{
-		return (numDraws + numDispatches + numClears + numBarriers + numPendingBarriers + numCopies + numInitialResourceCopies + otherWorkCounter) > 0;
+		return GetTotalWorkCount() > 0;
 	}
 
 	/** Constant buffers for Set*ShaderParameter calls. */
diff --git a/Engine/Source/Runtime/D3D12RHI/Private/D3D12Commands.cpp b/Engine/Source/Runtime/D3D12RHI/Private/D3D12Commands.cpp
--- a/Engine/Source/Runtime/D3D12RHI/Private/D3D12Commands.cpp
+++ b/Engine/Source/Runtime/D3D12RHI/Private/D3D12Commands.cpp
@@ -70,10 +70,6 @@ inline void ValidateBoundUniformBuffer(FD3D12UniformBuffer * InUniformBuffer, TR
 #endif
 }
 
-#if EXECUTE_DEBUG_COMMAND_LISTS
-bool GIsDoingQuery = false;
-#endif
-
 #if !defined(D3D12_PLATFORM_SUPPORTS_RESOLVE_SHADERS)
 	#define D3D12_PLATFORM_SUPPORTS_RESOLVE_SHADERS 1
 #endif
@@ -147,6 +143,8 @@ void FD3D12CommandContext::RHIDispatchComputeShader(uint32 ThreadGroupCountX, ui
 
 	numDispatches++;
 	CommandListHandle->Dispatch(ThreadGroupCountX, ThreadGroupCountY, ThreadGroupCountZ);
+	
+	ConditionalFlushCommandList();
 
 	DEBUG_EXECUTE_COMMAND_LIST(this);
 }
@@ -191,6 +189,8 @@ void FD3D12CommandContext::RHIDispatchIndirectComputeShader(FRHIBuffer* Argument
 		0
 		);
 	CommandListHandle.UpdateResidency(Location.GetResource());
+	
+	ConditionalFlushCommandList();
 
 	DEBUG_EXECUTE_COMMAND_LIST(this);
 }
@@ -665,6 +665,8 @@ void FD3D12CommandContext::RHICopyToStagingBuffer(FRHIBuffer* SourceBufferRHI, F
 		CommandListHandle->CopyBufferRegion(pDestResource->GetResource(), DestOffset, pSourceResource->GetResource(), Offset + SourceOffset, NumBytes);
 		CommandListHandle.UpdateResidency(pDestResource);
 		CommandListHandle.UpdateResidency(pSourceResource);
+
+		ConditionalFlushCommandList();
 	}
 }
 
@@ -1447,9 +1449,7 @@ void FD3D12CommandContext::RHIBeginRenderQuery(FRHIRenderQuery* QueryRHI)
 
 	GetParentDevice()->GetOcclusionQueryHeap(CommandQueueType)->BeginQuery(*this, Query);
 
-#if EXECUTE_DEBUG_COMMAND_LISTS
-	GIsDoingQuery = true;
-#endif
+	bIsDoingQuery = true;
 }
 
 void FD3D12CommandContext::RHIEndRenderQuery(FRHIRenderQuery* QueryRHI)
@@ -1480,9 +1480,7 @@ void FD3D12CommandContext::RHIEndRenderQuery(FRHIRenderQuery* QueryRHI)
 	// Query data isn't ready until it has been resolved.
 	ensure(Query->bResultIsCached == false && Query->bResolved == false);
 
-#if EXECUTE_DEBUG_COMMAND_LISTS
-	GIsDoingQuery = false;
-#endif
+	bIsDoingQuery = false;
 }
 
 void FD3D12CommandContext::RHICalibrateTimers(FRHITimestampCalibrationQuery* CalibrationQuery)
@@ -1839,6 +1837,8 @@ void FD3D12CommandContext::RHIDrawPrimitive(uint32 BaseVertexIndex, uint32 NumPr
 	StateCache.ApplyState<D3D12PT_Graphics>();
 	CommandListHandle->DrawInstanced(VertexCount, NumInstances, BaseVertexIndex, 0);
 
+	ConditionalFlushCommandList();
+
 #if UE_BUILD_DEBUG	
 	OwningRHI.DrawCount++;
 #endif
@@ -1879,6 +1879,8 @@ void FD3D12CommandContext::RHIDrawPrimitiveIndirect(FRHIBuffer* ArgumentBufferRH
 
 	CommandListHandle.UpdateResidency(Location.GetResource());
 
+	ConditionalFlushCommandList();
+
 #if UE_BUILD_DEBUG	
 	OwningRHI.DrawCount++;
 #endif
@@ -1928,6 +1930,8 @@ void FD3D12CommandContext::RHIDrawIndexedIndirect(FRHIBuffer* IndexBufferRHI, FR
 
 	CommandListHandle.UpdateResidency(Location.GetResource());
 
+	ConditionalFlushCommandList();
+
 #if UE_BUILD_DEBUG	
 	OwningRHI.DrawCount++;
 #endif
@@ -1966,6 +1970,8 @@ void FD3D12CommandContext::RHIDrawIndexedPrimitive(FRHIBuffer* IndexBufferRHI, i
 
 	CommandListHandle->DrawIndexedInstanced(IndexCount, NumInstances, StartIndex, BaseVertexIndex, FirstInstance);
 
+	ConditionalFlushCommandList();
+
 #if UE_BUILD_DEBUG	
 	OwningRHI.DrawCount++;
 #endif
@@ -2012,6 +2018,8 @@ void FD3D12CommandContext::RHIDrawIndexedPrimitiveIndirect(FRHIBuffer* IndexBuff
 
 	CommandListHandle.UpdateResidency(Location.GetResource());
 
+	ConditionalFlushCommandList();
+
 #if UE_BUILD_DEBUG	
 	OwningRHI.DrawCount++;
 #endif
@@ -2034,6 +2042,8 @@ void FD3D12CommandContext::RHIDispatchMeshShader(uint32 ThreadGroupCountX, uint3
 
 	CommandListHandle.GraphicsCommandList6()->DispatchMesh(ThreadGroupCountX, ThreadGroupCountY, ThreadGroupCountZ);
 
+	ConditionalFlushCommandList();
+
 #if UE_BUILD_DEBUG	
 	OwningRHI.DrawCount++;
 #endif
@@ -2074,6 +2084,8 @@ void FD3D12CommandContext::RHIDispatchIndirectMeshShader(FRHIBuffer* ArgumentBuf
 
 	CommandListHandle.UpdateResidency(Location.GetResource());
 
+	ConditionalFlushCommandList();
+
 #if UE_BUILD_DEBUG	
 	OwningRHI.DrawCount++;
 #endif
@@ -2196,6 +2208,8 @@ void FD3D12CommandContext::RHIClearMRTImpl(bool* bClearColorArray, int32 NumClea
 			CommandListHandle->ClearDepthStencilView(DepthStencilView->GetView(), (D3D12_CLEAR_FLAGS)ClearFlags, Depth, Stencil, ClearRectCount, pClearRects);
 			CommandListHandle.UpdateResidency(DepthStencilView->GetResource());
 		}
+
+		ConditionalFlushCommandList();
 	}
 
 	if (IsDefaultContext())
diff --git a/Engine/Source/Runtime/D3D12RHI/Private/D3D12PoolAllocator.cpp b/Engine/Source/Runtime/D3D12RHI/Private/D3D12PoolAllocator.cpp
--- a/Engine/Source/Runtime/D3D12RHI/Private/D3D12PoolAllocator.cpp
+++ b/Engine/Source/Runtime/D3D12RHI/Private/D3D12PoolAllocator.cpp
@@ -813,6 +813,8 @@ void FD3D12PoolAllocator::FlushPendingCopyOps(FD3D12CommandContext& InCommandCon
 		}
 	}
 
+	InCommandContext.ConditionalFlushCommandList();
+
 	PendingCopyOps.Empty(PendingCopyOps.Num());
 }
 
diff --git a/Engine/Source/Runtime/D3D12RHI/Private/D3D12RHI.cpp b/Engine/Source/Runtime/D3D12RHI/Private/D3D12RHI.cpp
--- a/Engine/Source/Runtime/D3D12RHI/Private/D3D12RHI.cpp
+++ b/Engine/Source/Runtime/D3D12RHI/Private/D3D12RHI.cpp
@@ -345,6 +345,8 @@ void FD3D12DynamicRHI::UpdateBuffer(FD3D12Resource* Dest, uint32 DestOffset, FD3
 	hCommandList->CopyBufferRegion(Dest->GetResource(), DestOffset, Source->GetResource(), SourceOffset, NumBytes);
 	hCommandList.UpdateResidency(Dest);
 	hCommandList.UpdateResidency(Source);
+	
+	DefaultContext.ConditionalFlushCommandList();
 
 	DEBUG_RHI_EXECUTE_COMMAND_LIST(this);
 }
diff --git a/Engine/Source/Runtime/D3D12RHI/Private/D3D12RHIPrivate.h b/Engine/Source/Runtime/D3D12RHI/Private/D3D12RHIPrivate.h
--- a/Engine/Source/Runtime/D3D12RHI/Private/D3D12RHIPrivate.h
+++ b/Engine/Source/Runtime/D3D12RHI/Private/D3D12RHIPrivate.h
@@ -140,10 +140,9 @@ typedef FD3D12StateCacheBase FD3D12StateCache;
 #endif
 
 #if EXECUTE_DEBUG_COMMAND_LISTS
-extern bool GIsDoingQuery;
-#define DEBUG_EXECUTE_COMMAND_LIST(scope) if (!GIsDoingQuery) { scope##->FlushCommands(true); }
-#define DEBUG_EXECUTE_COMMAND_CONTEXT(context) if (!GIsDoingQuery) { context##.FlushCommands(true); }
-#define DEBUG_RHI_EXECUTE_COMMAND_LIST(scope) if (!GIsDoingQuery) { scope##->GetRHIDevice()->GetDefaultCommandContext().FlushCommands(true); }
+#define DEBUG_EXECUTE_COMMAND_LIST(scope) if (!scope##->bIsDoingQuery) { scope##->FlushCommands(true); }
+#define DEBUG_EXECUTE_COMMAND_CONTEXT(context) if (!context.bIsDoingQuery) { context##.FlushCommands(true); }
+#define DEBUG_RHI_EXECUTE_COMMAND_LIST(scope) if (!scope##->GetRHIDevice(0)->GetDefaultCommandContext().bIsDoingQuery) { scope##->GetRHIDevice(0)->GetDefaultCommandContext().FlushCommands(true); }
 #else
 #define DEBUG_EXECUTE_COMMAND_LIST(scope) 
 #define DEBUG_EXECUTE_COMMAND_CONTEXT(context) 
diff --git a/Engine/Source/Runtime/D3D12RHI/Private/D3D12RenderTarget.cpp b/Engine/Source/Runtime/D3D12RHI/Private/D3D12RenderTarget.cpp
--- a/Engine/Source/Runtime/D3D12RHI/Private/D3D12RenderTarget.cpp
+++ b/Engine/Source/Runtime/D3D12RHI/Private/D3D12RenderTarget.cpp
@@ -491,6 +491,8 @@ void FD3D12CommandContext::RHICopyToResolveTarget(FRHITexture* SourceTextureRHI,
 		RHICmdList.TransitionResource(ResolveParams.DestAccessFinal, DestTextureRHI);
 	}
 
+	ConditionalFlushCommandList();
+
 	DEBUG_EXECUTE_COMMAND_LIST(this);
 }
 
diff --git a/Engine/Source/Runtime/D3D12RHI/Private/D3D12Texture.cpp b/Engine/Source/Runtime/D3D12RHI/Private/D3D12Texture.cpp
--- a/Engine/Source/Runtime/D3D12RHI/Private/D3D12Texture.cpp
+++ b/Engine/Source/Runtime/D3D12RHI/Private/D3D12Texture.cpp
@@ -1646,6 +1646,8 @@ void FD3D12DynamicRHI::RHICopySharedMips(FRHITexture2D* DestTexture2DRHI, FRHITe
 		// unlock the pool allocated resource because all data has been written
 		DestTexture2D->ResourceLocation.UnlockPoolData();
 
+		Device->GetDefaultCommandContext().ConditionalFlushCommandList();
+
 		DEBUG_EXECUTE_COMMAND_CONTEXT(Device->GetDefaultCommandContext());
 	}
 }
@@ -1746,6 +1748,8 @@ static void DoAsyncReallocateTexture2D(FD3D12Texture2D* Texture2D, FD3D12Texture
 			hCommandList.UpdateResidency(NewTexture2D->GetResource());
 			hCommandList.UpdateResidency(Texture2D->GetResource());
 
+			Device->GetDefaultCommandContext().ConditionalFlushCommandList();
+
 			DEBUG_EXECUTE_COMMAND_CONTEXT(Device->GetDefaultCommandContext());
 		}
 	}
@@ -2111,6 +2115,8 @@ void FD3D12TextureBase::UpdateTexture(uint32 MipIndex, uint32 DestX, uint32 Dest
 		nullptr);
 
 	hCommandList.UpdateResidency(GetResource());
+	
+	DefaultContext.ConditionalFlushCommandList();
 
 	DEBUG_EXECUTE_COMMAND_CONTEXT(DefaultContext);
 }
@@ -2585,7 +2591,6 @@ public:
 					nullptr);
 
 				NativeCmdList.UpdateResidency(TextureLink.GetResource());
-
 				DEBUG_EXECUTE_COMMAND_CONTEXT(Device->GetDefaultCommandContext());
 #if USE_PIX
 				if (FD3D12DynamicRHI::GetD3DRHI()->IsPixEventEnabled())
@@ -2594,6 +2599,8 @@ public:
 				}
 #endif
 			}
+
+			Device->GetDefaultCommandContext().ConditionalFlushCommandList();
 		}
 	}
 
@@ -2821,6 +2828,7 @@ public:
 
 			NativeCmdList.UpdateResidency(TextureLink.GetResource());
 
+			Device->GetDefaultCommandContext().ConditionalFlushCommandList();
 			DEBUG_EXECUTE_COMMAND_CONTEXT(Device->GetDefaultCommandContext());
 #if USE_PIX
 			if (FD3D12DynamicRHI::GetD3DRHI()->IsPixEventEnabled())
@@ -3482,6 +3490,8 @@ void FD3D12CommandContext::RHICopyTexture(FRHITexture* SourceTextureRHI, FRHITex
 
 	CommandListHandle.UpdateResidency(SourceTexture->GetResource());
 	CommandListHandle.UpdateResidency(DestTexture->GetResource());
+	
+	ConditionalFlushCommandList();
 
 	// Save the command list handle. This lets us check when this command list is complete. Note: This must be saved before we execute the command list
 	DestTexture->SetReadBackListHandle(CommandListHandle);
diff --git a/Engine/Source/Runtime/D3D12RHI/Private/D3D12UAV.cpp b/Engine/Source/Runtime/D3D12RHI/Private/D3D12UAV.cpp
--- a/Engine/Source/Runtime/D3D12RHI/Private/D3D12UAV.cpp
+++ b/Engine/Source/Runtime/D3D12RHI/Private/D3D12UAV.cpp
@@ -249,6 +249,7 @@ void FD3D12CommandContext::ClearUAV(TRHICommandList_RecursiveHazardous<FD3D12Com
 				Context.CommandListHandle.FlushResourceBarriers();
 				Context.CommandListHandle->ClearUnorderedAccessViewUint(GPUHandle, CPUHandle, Resource, *reinterpret_cast<const UINT(*)[4]>(ClearValues), 0, nullptr);
 				Context.CommandListHandle.UpdateResidency(UnorderedAccessView->GetResource());
+				Context.ConditionalFlushCommandList();
 
 				if (Context.IsDefaultContext())
 				{
