Allow 0 size textures - fix boottest issues

diff --git a/Engine/Source/Runtime/D3D12RHI/Private/D3D12Texture.h b/Engine/Source/Runtime/D3D12RHI/Private/D3D12Texture.h
--- a/Engine/Source/Runtime/D3D12RHI/Private/D3D12Texture.h
+++ b/Engine/Source/Runtime/D3D12RHI/Private/D3D12Texture.h
@@ -112,7 +112,7 @@ public:
 
 	void SetMemorySize(int64 InMemorySize)
 	{
-		check(InMemorySize > 0);
+		check(InMemorySize >= 0);
 		MemorySize = InMemorySize;
 	}
 
