Fix for assertion on RayTracingTable allocation

- FD3D12DefaultBufferPool::GetResourceAllocationStrategy() forces MultiState resources to be placed resources.  Assertion in FD3D12DefaultBufferPool::AllocDefaultResource() needs to account for it.

diff --git a/Engine/Source/Runtime/D3D12RHI/Private/D3D12Allocation.cpp b/Engine/Source/Runtime/D3D12RHI/Private/D3D12Allocation.cpp
--- a/Engine/Source/Runtime/D3D12RHI/Private/D3D12Allocation.cpp
+++ b/Engine/Source/Runtime/D3D12RHI/Private/D3D12Allocation.cpp
@@ -1223,7 +1223,7 @@ void FD3D12DefaultBufferPool::AllocDefaultResource(D3D12_HEAP_TYPE InHeapType, c
 		if (bPlacedResource)
 		{
 			// Writeable resources get separate ID3D12Resource* with their own resource state by using placed resources. Just make sure it's UAV, other flags are free to differ.
-			check((Desc.Flags & Allocator->GetInitConfig().ResourceFlags & D3D12_RESOURCE_FLAG_ALLOW_UNORDERED_ACCESS) != 0 || InHeapType == D3D12_HEAP_TYPE_READBACK || Alignment > kD3D12ManualSubAllocationAlignment);
+			check((Desc.Flags & Allocator->GetInitConfig().ResourceFlags & D3D12_RESOURCE_FLAG_ALLOW_UNORDERED_ACCESS) != 0 || InHeapType == D3D12_HEAP_TYPE_READBACK || Alignment > kD3D12ManualSubAllocationAlignment || InResourceStateMode == ED3D12ResourceStateMode::MultiState);
 		}
 		else
 		{
