Check if more than 1GB of pending delete requests on upload memory on allocate and if so, force flush GPU & garbage collect to avoid high upload memory peaks during loading

diff --git a/Engine/Source/Runtime/D3D12RHI/Private/D3D12Allocation.cpp b/Engine/Source/Runtime/D3D12RHI/Private/D3D12Allocation.cpp
--- a/Engine/Source/Runtime/D3D12RHI/Private/D3D12Allocation.cpp
+++ b/Engine/Source/Runtime/D3D12RHI/Private/D3D12Allocation.cpp
@@ -160,6 +160,13 @@ static FAutoConsoleVariableRef CVarD3D12FastAllocatorMinPagesToRetain(
 	TEXT("Minimum number of pages to retain. Pages below this limit will never be released. Pages above can be released after being unused for a certain number of frames."),
 	ECVF_Default);
 
+static int32 GD3D12UploadAllocatorPendingDeleteSizeForceFlushInGB = 1;
+static FAutoConsoleVariableRef CVarD3D12UploadAllocatorPendingDeleteSizeForceFlushInGB(
+	TEXT("d3d12.UploadAllocator.PendingDeleteSizeForceFlushInGB"),
+	GD3D12FastAllocatorMinPagesToRetain,
+	TEXT("If given threshold of GBs in the pending delete is queue is reached, then a force GPU flush is triggered to reduce memory load (1 by default, 0 to disable)"),
+	ECVF_Default);
+
 namespace ED3D12AllocatorID
 {
 	enum Type
@@ -1017,11 +1024,29 @@ void* FD3D12UploadHeapAllocator::AllocUploadResource(uint32 InSize, uint32 InAli
 
 	// Clean up the release queue of resources which are currently not used by the GPU anymore
 	FD3D12Adapter* Adapter = GetParentAdapter();
-	if (Adapter->GetDeferredDeletionQueue().QueueSize() > 128)
+	bool bFlushDeferredDeletionQueue = Adapter->GetDeferredDeletionQueue().QueueSize() > 128;
+	bool bFlushPendingDeleteRequests = GD3D12UploadAllocatorPendingDeleteSizeForceFlushInGB > 0 && BigBlockAllocator.GetPendingDeleteRequestSize() > (GD3D12UploadAllocatorPendingDeleteSizeForceFlushInGB * 1024 * 1024 * 1024);
+	if ((bFlushDeferredDeletionQueue || bFlushPendingDeleteRequests) && IsInRenderingThread())
 	{
-		Adapter->GetDeferredDeletionQueue().ReleaseResources(true, false);
+		if (bFlushPendingDeleteRequests)
+		{
+			UE_LOG(LogD3D12RHI, Warning, TEXT("Force flushing GPU because pending upload allocations reached its limit"));
+
+			// Flush to GPU & Wait (stall the RHI thread)
+			FScopedRHIThreadStaller StallRHIThread(FRHICommandListExecutor::GetImmediateCommandList());
+			Adapter->GetDevice(0)->GetDefaultCommandContext().FlushCommands(true);	// Don't wait yet, since we're stalling the RHI thread.
+
+			// Waited for GPU to finish so all sync points are ready so can force free all pending deletes (done while RHI thread is stalled)
+			bool bForceFreePendingDeletes = true;
+			BigBlockAllocator.CleanUpAllocations(0, bForceFreePendingDeletes);
+		}
+		else
+		{
+			BigBlockAllocator.CleanUpAllocations(0);
+		}
+		
 		SmallBlockAllocator.CleanUpAllocations(0); // 0 - no FrameLag, delete all unsued pages
-		BigBlockAllocator.CleanUpAllocations(0);
+		Adapter->GetDeferredDeletionQueue().ReleaseResources(true, false);
 	}
 
 	check(InSize > 0);
diff --git a/Engine/Source/Runtime/D3D12RHI/Private/D3D12PoolAllocator.cpp b/Engine/Source/Runtime/D3D12RHI/Private/D3D12PoolAllocator.cpp
--- a/Engine/Source/Runtime/D3D12RHI/Private/D3D12PoolAllocator.cpp
+++ b/Engine/Source/Runtime/D3D12RHI/Private/D3D12PoolAllocator.cpp
@@ -477,6 +477,8 @@ void FD3D12PoolAllocator::DeallocateResource(FD3D12ResourceLocation& ResourceLoc
 	DeleteRequest.FrameFence = FrameFence.GetCurrentFence();
 	DeleteRequest.AllocationData = ReleasedAllocationData;
 
+	PendingDeleteRequestSize += DeleteRequest.AllocationData->GetSize();
+
 	// Update the last used frame fence (used during garbage collection)
 	((FD3D12MemoryPool*)Pools[PoolIndex])->UpdateLastUsedFrameFence(DeleteRequest.FrameFence);
 
@@ -561,7 +563,7 @@ bool FD3D12PoolAllocator::HandleDefragRequest(FRHIPoolAllocationData* InSourceBl
 }
 
 
-void FD3D12PoolAllocator::CleanUpAllocations(uint64 InFrameLag)
+void FD3D12PoolAllocator::CleanUpAllocations(uint64 InFrameLag, bool bForceFree)
 {
 	FScopeLock Lock(&CS);
 
@@ -572,12 +574,16 @@ void FD3D12PoolAllocator::CleanUpAllocations(uint64 InFrameLag)
 	for (int32 i = 0; i < FrameFencedOperations.Num(); i++)
 	{
 		FrameFencedAllocationData& Operation = FrameFencedOperations[i];
-		if (FrameFence.IsFenceComplete(Operation.FrameFence))
+		if (bForceFree || FrameFence.IsFenceComplete(Operation.FrameFence))
 		{
 			switch (Operation.Operation)
 			{
 			case FrameFencedAllocationData::EOperation::Deallocate:
 			{
+				// Not pending anymore
+				check(PendingDeleteRequestSize >= Operation.AllocationData->GetSize());
+				PendingDeleteRequestSize -= Operation.AllocationData->GetSize();
+
 				// Deallocate the locked block (actually free now)
 				DeallocateInternal(*Operation.AllocationData);
 				Operation.AllocationData->Reset();
@@ -628,7 +634,7 @@ void FD3D12PoolAllocator::CleanUpAllocations(uint64 InFrameLag)
 	for (int32 PoolIndex = 0; PoolIndex < Pools.Num(); ++PoolIndex)
 	{
 		FD3D12MemoryPool* MemoryPool = (FD3D12MemoryPool*) Pools[PoolIndex];
-		if (MemoryPool != nullptr && MemoryPool->IsEmpty() && (MemoryPool->GetLastUsedFrameFence() + InFrameLag <= CompletedFence))
+		if (MemoryPool != nullptr && MemoryPool->IsEmpty() && (bForceFree || (MemoryPool->GetLastUsedFrameFence() + InFrameLag <= CompletedFence)))
 		{
 			MemoryPool->Destroy();
 			delete(MemoryPool);
diff --git a/Engine/Source/Runtime/D3D12RHI/Private/D3D12PoolAllocator.h b/Engine/Source/Runtime/D3D12RHI/Private/D3D12PoolAllocator.h
--- a/Engine/Source/Runtime/D3D12RHI/Private/D3D12PoolAllocator.h
+++ b/Engine/Source/Runtime/D3D12RHI/Private/D3D12PoolAllocator.h
@@ -140,7 +140,7 @@ public:
 	virtual void AllocateResource(uint32 GPUIndex, D3D12_HEAP_TYPE InHeapType, const FD3D12ResourceDesc& InDesc, uint64 InSize, uint32 InAllocationAlignment, ED3D12ResourceStateMode InResourceStateMode,
 		D3D12_RESOURCE_STATES InCreateState, const D3D12_CLEAR_VALUE* InClearValue, const TCHAR* InName, FD3D12ResourceLocation& ResourceLocation) override;
 
-	void CleanUpAllocations(uint64 InFrameLag);
+	void CleanUpAllocations(uint64 InFrameLag, bool bForceFree = false);
 	void FlushPendingCopyOps(FD3D12CommandContext& InCommandContext);
 
 	void TransferOwnership(FD3D12ResourceLocation& InSource, FD3D12ResourceLocation& InDest);
@@ -150,6 +150,7 @@ public:
 	FD3D12Resource* GetBackingResource(FD3D12ResourceLocation& InResourceLocation) const;
 	FD3D12HeapAndOffset GetBackingHeapAndAllocationOffsetInBytes(FD3D12ResourceLocation& InResourceLocation) const;
 	FD3D12HeapAndOffset GetBackingHeapAndAllocationOffsetInBytes(const FRHIPoolAllocationData& InAllocationData) const;
+	uint64 GetPendingDeleteRequestSize() const { return PendingDeleteRequestSize; }
 
 	static FD3D12ResourceInitConfig GetResourceAllocatorInitConfig(D3D12_HEAP_TYPE InHeapType, D3D12_RESOURCE_FLAGS InResourceFlags, EBufferUsageFlags InBufferUsage);
 	static EResourceAllocationStrategy GetResourceAllocationStrategy(D3D12_RESOURCE_FLAGS InResourceFlags, ED3D12ResourceStateMode InResourceStateMode, uint32 Alignment);
@@ -196,6 +197,9 @@ protected:
 	// All operations which need to happen on specific frame fences
 	TArray<FrameFencedAllocationData> FrameFencedOperations;
 
+	// Size of all pending delete requests in the frame fence operation array
+	uint64 PendingDeleteRequestSize = 0;
+
 	// Pending copy ops which need to be flush this frame
 	TArray<FD3D12VRAMCopyOperation> PendingCopyOps;
 
