fix crashes due to copy queue execution calls from multiple threads on console platforms

diff --git a/Engine/Source/Runtime/D3D12RHI/Private/D3D12Buffer.cpp b/Engine/Source/Runtime/D3D12RHI/Private/D3D12Buffer.cpp
--- a/Engine/Source/Runtime/D3D12RHI/Private/D3D12Buffer.cpp
+++ b/Engine/Source/Runtime/D3D12RHI/Private/D3D12Buffer.cpp
@@ -5,6 +5,7 @@ D3D12Buffer.cpp: D3D Common code for buffers.
 =============================================================================*/
 
 #include "D3D12RHIPrivate.h"
+#include "D3D12RHIBridge.h"
 
 struct FRHICommandUpdateBufferString
 {
@@ -267,8 +268,12 @@ FD3D12SyncPoint FD3D12Buffer::UploadResourceDataViaCopyQueue(FResourceArrayInter
 	// Close and kick the command list without waiting for it
 	hCopyCommandList.Close();
 	bool bWaitForCompletion = false;
-	FD3D12SyncPoint CopyQueueSyncPoint = Device->GetCopyCommandListManager().ExecuteCommandListNoCopyQueueSync(hCopyCommandList, bWaitForCompletion);
 
+	FD3D12SyncPoint CopyQueueSyncPoint;
+	D3D12RHI::ExecuteCodeWithCopyCommandQueueUsage([Device, &hCopyCommandList, bWaitForCompletion, &CopyQueueSyncPoint](ID3D12CommandQueue* D3DCommandQueue) -> void
+	{
+		CopyQueueSyncPoint = Device->GetCopyCommandListManager().ExecuteCommandListNoCopyQueueSync(hCopyCommandList, bWaitForCompletion);
+	});
 	// Release command allocator (has the sync point as well)
 	CommandAllocatorManager.ReleaseCommandAllocator(CurrentCommandAllocator);
 
diff --git a/Engine/Source/Runtime/D3D12RHI/Private/D3D12Commands.cpp b/Engine/Source/Runtime/D3D12RHI/Private/D3D12Commands.cpp
--- a/Engine/Source/Runtime/D3D12RHI/Private/D3D12Commands.cpp
+++ b/Engine/Source/Runtime/D3D12RHI/Private/D3D12Commands.cpp
@@ -5,6 +5,7 @@ D3D12Commands.cpp: D3D RHI commands implementation.
 =============================================================================*/
 
 #include "D3D12RHIPrivate.h"
+#include "D3D12RHIBridge.h"
 #include "StaticBoundShaderState.h"
 #include "GlobalShader.h"
 #include "OneColorShader.h"
@@ -2390,7 +2391,12 @@ void FD3D12CommandContext::RHIBroadcastTemporalEffect(const FName& InEffectName,
 	hCopyCommandList.Close();
 
 	bool bWaitForCompletion = false;
-	CopyCommandListManager.ExecuteCommandListNoCopyQueueSync(hCopyCommandList, bWaitForCompletion);
+
+	D3D12RHI::ExecuteCodeWithCopyCommandQueueUsage([bWaitForCompletion, &CopyCommandListManager, &hCopyCommandList](ID3D12CommandQueue* D3DCommandQueue) -> void
+	{
+		CopyCommandListManager.ExecuteCommandListNoCopyQueueSync(hCopyCommandList, bWaitForCompletion);
+	});
+
 	CopyCommandAllocatorManager.ReleaseCommandAllocator(CopyCommandAllocator);
 
 	// Signal again once the copy queue copy is complete.
diff --git a/Engine/Source/Runtime/D3D12RHI/Private/D3D12Texture.cpp b/Engine/Source/Runtime/D3D12RHI/Private/D3D12Texture.cpp
--- a/Engine/Source/Runtime/D3D12RHI/Private/D3D12Texture.cpp
+++ b/Engine/Source/Runtime/D3D12RHI/Private/D3D12Texture.cpp
@@ -5,6 +5,7 @@
 	=============================================================================*/
 
 #include "D3D12RHIPrivate.h"
+#include "D3D12RHIBridge.h"
 
 int64 FD3D12GlobalStats::GDedicatedVideoMemory = 0;
 int64 FD3D12GlobalStats::GDedicatedSystemMemory = 0;
@@ -1566,7 +1567,11 @@ FTexture2DRHIRef FD3D12DynamicRHI::RHIAsyncCreateTexture2D(uint32 SizeX, uint32
 			hCopyCommandList.Close();
 
 			bool bWaitForCompletion = true;
-			Device->GetCopyCommandListManager().ExecuteCommandListNoCopyQueueSync(hCopyCommandList, bWaitForCompletion);
+
+			D3D12RHI::ExecuteCodeWithCopyCommandQueueUsage([Device, bWaitForCompletion, &hCopyCommandList](ID3D12CommandQueue* D3DCommandQueue) -> void
+			{
+				Device->GetCopyCommandListManager().ExecuteCommandListNoCopyQueueSync(hCopyCommandList, bWaitForCompletion);
+			});
 
 			CommandAllocatorManager.ReleaseCommandAllocator(CurrentCommandAllocator);
 		}
