Fixing D3D12 with -rhivalidation

diff --git a/Engine/Source/Runtime/D3D12RHI/Private/D3D12CommandList.cpp b/Engine/Source/Runtime/D3D12RHI/Private/D3D12CommandList.cpp
--- a/Engine/Source/Runtime/D3D12RHI/Private/D3D12CommandList.cpp
+++ b/Engine/Source/Runtime/D3D12RHI/Private/D3D12CommandList.cpp
@@ -358,7 +358,7 @@ namespace D3D12RHI
 	{
 		IRHICommandContext* Context = GDynamicRHI->RHIGetDefaultContext();
 		checkSlow(Context);
-		FD3D12CommandContextBase& BaseCmdContext = (FD3D12CommandContextBase&)*Context;
+		FD3D12CommandContextBase& BaseCmdContext = (FD3D12CommandContextBase&)Context->GetLowestLevelContext();
 
 		ID3D12CommandQueue* CommandQueue = BaseCmdContext.GetParentAdapter()->GetDevice(0)->GetD3DCommandQueue(ED3D12CommandQueueType::Copy);
 		checkSlow(CommandQueue);
