Adding missing overloads for Mesh and Amplification shaders.

diff --git a/Engine/Source/Runtime/D3D12RHI/Private/D3D12RHIPrivate.h b/Engine/Source/Runtime/D3D12RHI/Private/D3D12RHIPrivate.h
--- a/Engine/Source/Runtime/D3D12RHI/Private/D3D12RHIPrivate.h
+++ b/Engine/Source/Runtime/D3D12RHI/Private/D3D12RHIPrivate.h
@@ -399,6 +399,16 @@ public:
 		return RHICreateVertexShader(Code, Hash);
 	}
 
+	virtual FMeshShaderRHIRef CreateMeshShader_RenderThread(class FRHICommandListImmediate& RHICmdList, TArrayView<const uint8> Code, const FSHAHash& Hash) override final
+	{
+		return RHICreateMeshShader(Code, Hash);
+	}
+
+	virtual FAmplificationShaderRHIRef CreateAmplificationShader_RenderThread(class FRHICommandListImmediate& RHICmdList, TArrayView<const uint8> Code, const FSHAHash& Hash) override final
+	{
+		return RHICreateAmplificationShader(Code, Hash);
+	}
+
 	virtual FGeometryShaderRHIRef CreateGeometryShader_RenderThread(class FRHICommandListImmediate& RHICmdList, TArrayView<const uint8> Code, const FSHAHash& Hash) override final
 	{
 		return RHICreateGeometryShader(Code, Hash);
