[LLM] - D3D12 LLM stats, highligts the BLAS allocations (previous to this they were considered RHIMisc)

diff --git a/Engine/Source/Runtime/D3D12RHI/Private/D3D12CommandContext.cpp b/Engine/Source/Runtime/D3D12RHI/Private/D3D12CommandContext.cpp
--- a/Engine/Source/Runtime/D3D12RHI/Private/D3D12CommandContext.cpp
+++ b/Engine/Source/Runtime/D3D12RHI/Private/D3D12CommandContext.cpp
@@ -292,6 +292,7 @@ void FD3D12CommandContext::ConditionalObtainCommandAllocator()
 {
 	if (CommandAllocator == nullptr)
 	{
+		LLM_SCOPE_BYNAME(TEXT("RHIMisc/ConditionalObtainCommandAllocator"));
 		// Obtain a command allocator if the context doesn't already have one.
 		// This will check necessary fence values to ensure the returned command allocator isn't being used by the GPU, then reset it.
 		CommandAllocator = CommandAllocatorManager.ObtainCommandAllocator();
diff --git a/Engine/Source/Runtime/D3D12RHI/Private/D3D12RayTracing.cpp b/Engine/Source/Runtime/D3D12RHI/Private/D3D12RayTracing.cpp
--- a/Engine/Source/Runtime/D3D12RHI/Private/D3D12RayTracing.cpp
+++ b/Engine/Source/Runtime/D3D12RHI/Private/D3D12RayTracing.cpp
@@ -656,6 +656,7 @@ bool FD3D12RayTracingCompactionRequestHandler::ReleaseRequest(FD3D12RayTracingGe
 
 void FD3D12RayTracingCompactionRequestHandler::Update(FD3D12CommandContext& InCommandContext)
 {
+	LLM_SCOPE_BYNAME(TEXT("FD3D12RT/Compaction"));
 	FScopeLock Lock(&CS);
 
 	// process previous build request data retrieval
@@ -749,7 +750,7 @@ public:
 		, DefaultLocalRootSignature(Device->GetParentAdapter())
 	{
 		// Default empty local root signature
-
+		LLM_SCOPE_BYNAME(TEXT("FD3D12RT/PipelineCache"));
 		D3D12_VERSIONED_ROOT_SIGNATURE_DESC LocalRootSignatureDesc = {};
 		if (Device->GetParentAdapter()->GetRootSignatureVersion() >= D3D_ROOT_SIGNATURE_VERSION_1_1)
 		{
@@ -2526,6 +2527,7 @@ public:
 
 void FD3D12Device::InitRayTracing()
 {
+	LLM_SCOPE_BYNAME(TEXT("FD3D12RT"));
 	check(RayTracingPipelineCache == nullptr);
 	RayTracingPipelineCache = new FD3D12RayTracingPipelineCache(this);
 
@@ -3108,6 +3110,7 @@ void FD3D12RayTracingGeometry::BuildAccelerationStructure(FD3D12CommandContext&
 {
 	TRACE_CPUPROFILER_EVENT_SCOPE(BuildAccelerationStructure_BottomLevel);
 	SCOPE_CYCLE_COUNTER(STAT_D3D12BuildBLAS);
+	LLM_SCOPE_BYNAME(TEXT("FD3D12RT/BLAS"));
 
 	if (RHIIndexBuffer)
 	{
@@ -3304,6 +3307,7 @@ void FD3D12RayTracingGeometry::BuildAccelerationStructure(FD3D12CommandContext&
 
 void FD3D12RayTracingGeometry::CompactAccelerationStructure(FD3D12CommandContext& CommandContext, uint32 InGPUIndex, uint64 InSizeAfterCompaction)
 {
+	LLM_SCOPE_BYNAME(TEXT("FD3D12RT/CompactBLAS"));
 	// Should have a pending request
 	check(bHasPendingCompactionRequests[InGPUIndex]);
 	bHasPendingCompactionRequests[InGPUIndex] = false;
