Remove deprecated *SetShadingRateImage* functions (flagged deprecated in 4.27).

diff --git a/Engine/Source/Runtime/D3D12RHI/Private/D3D12CommandContext.h b/Engine/Source/Runtime/D3D12RHI/Private/D3D12CommandContext.h
--- a/Engine/Source/Runtime/D3D12RHI/Private/D3D12CommandContext.h
+++ b/Engine/Source/Runtime/D3D12RHI/Private/D3D12CommandContext.h
@@ -317,7 +317,6 @@ public:
 #endif
 	virtual void RHISetDepthBounds(float MinDepth, float MaxDepth) final override;
     virtual void RHISetShadingRate(EVRSShadingRate ShadingRate, EVRSRateCombiner Combiner) final override;
-    virtual void RHISetShadingRateImage(FRHITexture* RateImageTexture, EVRSRateCombiner Combiner) final override;
 #if PLATFORM_USE_BACKBUFFER_WRITE_TRANSITION_TRACKING
 	virtual void RHIBackBufferWaitTrackingBeginFrame(uint64 FrameToken, bool bDeferred) final override;
 #endif // #if PLATFORM_USE_BACKBUFFER_WRITE_TRANSITION_TRACKING
@@ -712,11 +711,6 @@ public:
 		ContextRedirect(RHISetShadingRate(ShadingRate, Combiner));
 	}
 
-	FORCEINLINE virtual void RHISetShadingRateImage(FRHITexture* RateImageTexture, EVRSRateCombiner Combiner) final override
-	{
-		checkf(false, TEXT("RHISetShadingRateImage API is deprecated. Use the ShadingRateImage attachment in the RHISetRenderTargetsInfo struct instead."));
-	}
-
 	FORCEINLINE virtual void RHIWaitForTemporalEffect(const FName& InEffectName) final AFR_API_OVERRIDE
 	{
 		ContextRedirect(RHIWaitForTemporalEffect(InEffectName));
diff --git a/Engine/Source/Runtime/D3D12RHI/Private/D3D12Commands.cpp b/Engine/Source/Runtime/D3D12RHI/Private/D3D12Commands.cpp
--- a/Engine/Source/Runtime/D3D12RHI/Private/D3D12Commands.cpp
+++ b/Engine/Source/Runtime/D3D12RHI/Private/D3D12Commands.cpp
@@ -2294,13 +2294,6 @@ void FD3D12CommandContext::RHISetShadingRate(EVRSShadingRate ShadingRate, EVRSRa
  #endif
 }
 
-PRAGMA_DISABLE_DEPRECATION_WARNINGS
-void FD3D12CommandContext::RHISetShadingRateImage(FRHITexture* RateImageTexture, EVRSRateCombiner Combiner)
-{
-	checkf(false, TEXT("RHISetShadingRateImage API is deprecated. Use the ShadingRateImage attachment in the RHISetRenderTargetsInfo struct instead."));
-}
-PRAGMA_ENABLE_DEPRECATION_WARNINGS
-
 void FD3D12CommandContext::RHISubmitCommandsHint()
 {
 	// Resolve any timestamp queries so far (if any).
diff --git a/Engine/Source/Runtime/RHI/Public/RHICommandList.h b/Engine/Source/Runtime/RHI/Public/RHICommandList.h
--- a/Engine/Source/Runtime/RHI/Public/RHICommandList.h
+++ b/Engine/Source/Runtime/RHI/Public/RHICommandList.h
@@ -1502,19 +1502,6 @@ FRHICOMMAND_MACRO(FRHICommandSetShadingRate)
 	RHI_API void Execute(FRHICommandListBase& CmdList);
 };
 
-FRHICOMMAND_MACRO(FRHICommandSetShadingRateImage)
-{
-	FRHITexture* RateImageTexture;
-	EVRSRateCombiner  Combiner;
-
-	FORCEINLINE_DEBUGGABLE FRHICommandSetShadingRateImage(FRHITexture* InRateImageTexture, EVRSRateCombiner InCombiner)
-		: RateImageTexture(InRateImageTexture),
-		Combiner(InCombiner)
-	{
-	}
-	RHI_API void Execute(FRHICommandListBase& CmdList);
-};
-
 FRHICOMMAND_MACRO(FRHICommandClearUAVFloat)
 {
 	FRHIUnorderedAccessView* UnorderedAccessViewRHI;
@@ -3516,12 +3503,6 @@ public:
 #endif
 	}
 
-	UE_DEPRECATED(4.27, "SetShadingRateImage is deprecated. Bind the shading rate image as part of the FRHIRenderPassInfo struct.")
-	FORCEINLINE_DEBUGGABLE void SetShadingRateImage(FRHITexture* RateImageTexture, EVRSRateCombiner Combiner)
-	{
-		check(false);
-	}
-
 	FORCEINLINE_DEBUGGABLE void CopyToResolveTarget(FRHITexture* SourceTextureRHI, FRHITexture* DestTextureRHI, const FResolveParams& ResolveParams)
 	{
 		//check(IsOutsideRenderPass());
diff --git a/Engine/Source/Runtime/RHI/Public/RHICommandListCommandExecutes.inl b/Engine/Source/Runtime/RHI/Public/RHICommandListCommandExecutes.inl
--- a/Engine/Source/Runtime/RHI/Public/RHICommandListCommandExecutes.inl
+++ b/Engine/Source/Runtime/RHI/Public/RHICommandListCommandExecutes.inl
@@ -355,12 +355,6 @@ void FRHICommandSetShadingRate::Execute(FRHICommandListBase& CmdList)
 	INTERNAL_DECORATOR(RHISetShadingRate)(ShadingRate, Combiner);
 }
 
-void FRHICommandSetShadingRateImage::Execute(FRHICommandListBase& CmdList)
-{
-	RHISTAT(SetShadingRateImage);
-	checkf(false, TEXT("RHISetShadingRateImage API is deprecated. Use the ShadingRateImage attachment in the RHISetRenderTargetsInfo struct instead."));
-}
-
 void FRHICommandSetDepthBounds::Execute(FRHICommandListBase& CmdList)
 {
 	RHISTAT(EnableDepthBoundsTest);
diff --git a/Engine/Source/Runtime/RHI/Public/RHIContext.h b/Engine/Source/Runtime/RHI/Public/RHIContext.h
--- a/Engine/Source/Runtime/RHI/Public/RHIContext.h
+++ b/Engine/Source/Runtime/RHI/Public/RHIContext.h
@@ -666,12 +666,6 @@ public:
 		/* empty default implementation */
 	} 
 	
-	UE_DEPRECATED(4.27, "RHISetShadingRateImage is deprecated. Bind the shading rate image as part of the RHISetRenderTargetsInfo struct.")
-	virtual void RHISetShadingRateImage(FRHITexture* RateImageTexture, EVRSRateCombiner Combiner)
-	{
-		/* empty default implementation */
-	}
-
 	virtual void RHIBeginRenderPass(const FRHIRenderPassInfo& InInfo, const TCHAR* InName) = 0;
 
 	virtual void RHIEndRenderPass() = 0;
