Handle destroyed index/vertex buffers used by the RT geometry - remove as listener and mark RT geometry dirty (invalid for use in scene). Fixes crashes during mesh stream out with Raytracing enabled.

diff --git a/Engine/Source/Runtime/D3D12RHI/Private/D3D12RayTracing.cpp b/Engine/Source/Runtime/D3D12RHI/Private/D3D12RayTracing.cpp
--- a/Engine/Source/Runtime/D3D12RHI/Private/D3D12RayTracing.cpp
+++ b/Engine/Source/Runtime/D3D12RHI/Private/D3D12RayTracing.cpp
@@ -3046,10 +3046,33 @@ void FD3D12RayTracingGeometry::UnregisterAsRenameListener(uint32 InGPUIndex)
 
 void FD3D12RayTracingGeometry::ResourceRenamed(FD3D12BaseShaderResource* InRenamedResource, FD3D12ResourceLocation* InNewResourceLocation)
 {
-	checkf(InNewResourceLocation, TEXT("Shouldn't release resources which are still referenced by the RayTracingGeometry"));
+	// If there is no valid resource location (released) then release  as listener to the resource and mark the raytracing geometry dirty 
+	// (can't be used anymore since one of its pending resources are invalid now).
+	// Ideally this should not happen with correct order of destruction but it can happen that the RT geometry gets deleted a bit later than the buffers because it's
+	// still referenced by the previous RT scene (but not used anymore). This only happens during freeing of RHI buffers during mesh streaming.
+	if (InNewResourceLocation == nullptr)
+	{	
+		FD3D12Buffer* IndexBuffer = FD3D12DynamicRHI::ResourceCast(RHIIndexBuffer.GetReference(), InRenamedResource->GetParentDevice()->GetGPUIndex());
+		if (IndexBuffer == InRenamedResource)
+		{
+			RHIIndexBuffer = nullptr;
+		}
+		for (FRayTracingGeometrySegment& Segment : Segments)
+		{
+			FD3D12Buffer* VertexBuffer = FD3D12DynamicRHI::ResourceCast(Segment.VertexBuffer.GetReference(), InRenamedResource->GetParentDevice()->GetGPUIndex());
+			if (VertexBuffer == InRenamedResource)
+			{
+				Segment.VertexBuffer = nullptr;
+			}
+		}
 
-	// Recreate the hit group parameters which cache the address to the index and vertex buffers directly
-	SetupHitGroupSystemParameters(InRenamedResource->GetParentDevice()->GetGPUIndex());	
+		SetDirty(FRHIGPUMask::FromIndex(InRenamedResource->GetParentDevice()->GetGPUIndex()), true);
+	}
+	else
+	{
+		// Recreate the hit group parameters which cache the address to the index and vertex buffers directly
+		SetupHitGroupSystemParameters(InRenamedResource->GetParentDevice()->GetGPUIndex());
+	}
 }
 
 void FD3D12RayTracingGeometry::TransitionBuffers(FD3D12CommandContext& CommandContext)
