r.D3D12.DXR.AllowEmulatedRayTracing is now compiled out by default through a new define DXR_ALLOW_EMULATED_RAYTRACING

diff --git a/Engine/Source/Runtime/D3D12RHI/Private/Windows/WindowsD3D12Device.cpp b/Engine/Source/Runtime/D3D12RHI/Private/Windows/WindowsD3D12Device.cpp
--- a/Engine/Source/Runtime/D3D12RHI/Private/Windows/WindowsD3D12Device.cpp
+++ b/Engine/Source/Runtime/D3D12RHI/Private/Windows/WindowsD3D12Device.cpp
@@ -103,6 +103,9 @@ static FAutoConsoleVariableRef CVarMinDriverVersionForRayTracingNVIDIA(
 	ECVF_ReadOnly | ECVF_RenderThreadSafe
 );
 
+#define DXR_ALLOW_EMULATED_RAYTRACING 0
+
+#if DXR_ALLOW_EMULATED_RAYTRACING
 int32 GAllowEmulatedRayTracing = 0;
 static FAutoConsoleVariableRef CVarAllowEmulatedRayTracing(
 	TEXT("r.D3D12.DXR.AllowEmulatedRayTracing"),
@@ -110,6 +113,7 @@ static FAutoConsoleVariableRef CVarAllowEmulatedRayTracing(
 	TEXT("Allows ray tracing emulation support on NVIDIA cards with the Pascal architecture (default=0)."),
 	ECVF_ReadOnly | ECVF_RenderThreadSafe
 );
+#endif
 
 // Use AGS_MAKE_VERSION() macro to define the version.
 // i.e. AGS_MAKE_VERSION(major, minor, patch) ((major << 22) | (minor << 12) | patch)
@@ -957,14 +961,20 @@ void FD3D12DynamicRHI::Init()
 				TEXT("Please update to NVIDIA driver version 466.11 or newer."));
 		}
 
-		// Disable ray tracing emulation for Pascal architecture
-		if (GRHISupportsRayTracing
-			&& !GAllowEmulatedRayTracing
-			&& IsRayTracingEmulated(AdapterDesc.DeviceId))
+	if (GRHISupportsRayTracing
+		&& IsRayTracingEmulated(AdapterDesc.DeviceId))
+		{
+#if DXR_ALLOW_EMULATED_RAYTRACING
+		if (!GAllowEmulatedRayTracing)
 		{
 			GRHISupportsRayTracing = false;
 			UE_LOG(LogD3D12RHI, Warning, TEXT("Ray tracing is disabled for NVIDIA cards with the Pascal architecture. This can be overridden with the following CVar: r.D3D12.DXR.AllowEmulatedRayTracing=1"));
 		}
+#else
+		GRHISupportsRayTracing = false;
+		UE_LOG(LogD3D12RHI, Warning, TEXT("Ray tracing is disabled for NVIDIA cards with the Pascal architecture."));
+#endif // DXR_ALLOW_EMULATED_RAYTRACING
+		}
 	}
 #endif
 
