Fixed D3D12 profilegpu erroneously reporting 0.0 for certain leaf markers.

diff --git a/Engine/Source/Runtime/D3D12RHI/Private/D3D12CommandContext.cpp b/Engine/Source/Runtime/D3D12RHI/Private/D3D12CommandContext.cpp
--- a/Engine/Source/Runtime/D3D12RHI/Private/D3D12CommandContext.cpp
+++ b/Engine/Source/Runtime/D3D12RHI/Private/D3D12CommandContext.cpp
@@ -329,6 +329,8 @@ void FD3D12CommandContext::OpenCommandList()
 
 	bIsDoingQuery = false;
 
+	numPrimitives = 0;
+	numVertices = 0;
 	numDraws = 0;
 	numDispatches = 0;
 	numClears = 0;
@@ -476,8 +478,6 @@ void FD3D12CommandContextBase::RHIBeginFrame()
 	{
 		FD3D12Device* Device = ParentAdapter->GetDevice(GPUIndex);
 
-		bTrackingEvents |= bIsDefaultContext && Device->GetGPUProfiler().bTrackingEvents;
-
 #if D3D12_SUBMISSION_GAP_RECORDER
 		if (GEnableGapRecorder && !GTriggerGPUProfile)
 		{
@@ -526,6 +526,8 @@ void FD3D12CommandContextBase::RHIBeginFrame()
 		Device->GetDefaultBufferAllocator().BeginFrame();
 		Device->GetTextureAllocator().BeginFrame();
 
+		bTrackingEvents |= bIsDefaultContext && Device->GetGPUProfiler().bTrackingEvents;
+
 #if D3D12_RHI_RAYTRACING
 		Device->GetRayTracingCompactionRequestHandler()->Update(*ContextAtIndex);
 #endif // D3D12_RHI_RAYTRACING
@@ -764,6 +766,14 @@ class FD3D12CommandContextContainer : public IRHICommandContextContainer
 	FD3D12CommandContextRedirector* CmdContextRedirector;
 	FRHIGPUMask GPUMask;
 
+	struct FStats
+	{
+		uint32 NumDraws = 0;
+		uint32 NumPrimitives = 0;
+		uint32 NumVertices = 0;
+
+	} StatsPerGPU[MAX_NUM_GPUS];
+
 	TArray<FD3D12CommandListHandle> CommandLists;
 
 public:
@@ -835,6 +845,10 @@ public:
 
 		if (CmdContext)
 		{
+			StatsPerGPU[0].NumDraws = CmdContext->numDraws;
+			StatsPerGPU[0].NumPrimitives = CmdContext->numPrimitives;
+			StatsPerGPU[0].NumVertices = CmdContext->numVertices;
+
 			CmdContext->Finish(CommandLists);
 			CmdContext->GetParentDevice()->ReleaseCommandContext(CmdContext);
 			CmdContext = nullptr;
@@ -845,6 +859,11 @@ public:
 			for (uint32 GPUIndex : GPUMask)
 			{
 				CmdContext = CmdContextRedirector->GetContext(GPUIndex);
+
+				StatsPerGPU[GPUIndex].NumDraws = CmdContext->numDraws;
+				StatsPerGPU[GPUIndex].NumPrimitives = CmdContext->numPrimitives;
+				StatsPerGPU[GPUIndex].NumVertices = CmdContext->numVertices;
+
 				CmdContext->Finish(CommandLists);
 				CmdContext->GetParentDevice()->ReleaseCommandContext(CmdContext);
 				CmdContext = nullptr;
@@ -893,6 +912,8 @@ public:
 			FD3D12Device* Device = Adapter->GetDevice(GPUIndex);
 			check(Device);
 
+			Device->GetGPUProfiler().RegisterGPUWork(StatsPerGPU[GPUIndex].NumDraws, StatsPerGPU[GPUIndex].NumPrimitives, StatsPerGPU[GPUIndex].NumVertices);
+
 			if (Index == (Num - 1))
 			{
 				// Determine if we should flush:
diff --git a/Engine/Source/Runtime/D3D12RHI/Private/D3D12CommandContext.h b/Engine/Source/Runtime/D3D12RHI/Private/D3D12CommandContext.h
--- a/Engine/Source/Runtime/D3D12RHI/Private/D3D12CommandContext.h
+++ b/Engine/Source/Runtime/D3D12RHI/Private/D3D12CommandContext.h
@@ -195,6 +195,8 @@ public:
 
 	bool bIsDoingQuery = false;
 
+	uint32 numPrimitives;
+	uint32 numVertices;
 	uint32 numDraws;
 	uint32 numDispatches;
 	uint32 numClears;
diff --git a/Engine/Source/Runtime/D3D12RHI/Private/D3D12Commands.cpp b/Engine/Source/Runtime/D3D12RHI/Private/D3D12Commands.cpp
--- a/Engine/Source/Runtime/D3D12RHI/Private/D3D12Commands.cpp
+++ b/Engine/Source/Runtime/D3D12RHI/Private/D3D12Commands.cpp
@@ -1830,6 +1830,8 @@ void FD3D12CommandContext::RHIDrawPrimitive(uint32 BaseVertexIndex, uint32 NumPr
 	uint32 VertexCount = StateCache.GetVertexCountAndIncrementStat(NumPrimitives);
 	NumInstances = FMath::Max<uint32>(1, NumInstances);
 	numDraws++;
+	numPrimitives += NumPrimitives * NumInstances;
+	numVertices += VertexCount * NumInstances;
 	if (bTrackingEvents)
 	{
 		GetParentDevice()->RegisterGPUWork(NumPrimitives * NumInstances, VertexCount * NumInstances);
@@ -1947,6 +1949,8 @@ void FD3D12CommandContext::RHIDrawIndexedPrimitive(FRHIBuffer* IndexBufferRHI, i
 
 	NumInstances = FMath::Max<uint32>(1, NumInstances);
 	numDraws++;
+	numPrimitives += NumPrimitives * NumInstances;
+	numVertices += NumVertices * NumInstances;
 	if (bTrackingEvents)
 	{
 		GetParentDevice()->RegisterGPUWork(NumPrimitives * NumInstances, NumVertices * NumInstances);
diff --git a/Engine/Source/Runtime/RHI/Public/GPUProfiler.h b/Engine/Source/Runtime/RHI/Public/GPUProfiler.h
--- a/Engine/Source/Runtime/RHI/Public/GPUProfiler.h
+++ b/Engine/Source/Runtime/RHI/Public/GPUProfiler.h
@@ -293,17 +293,22 @@ struct RHI_API FGPUProfiler
 	{
 	}
 
-	void RegisterGPUWork(uint32 NumPrimitives = 0, uint32 NumVertices = 0)
+	void RegisterGPUWork(uint32 NumDraws, uint32 NumPrimitives, uint32 NumVertices)
 	{
 		if (bTrackingEvents && CurrentEventNode)
 		{
 			check(IsInRenderingThread() || IsInRHIThread());
-			CurrentEventNode->NumDraws++;
+			CurrentEventNode->NumDraws += NumDraws;
 			CurrentEventNode->NumPrimitives += NumPrimitives;
 			CurrentEventNode->NumVertices += NumVertices;
 		}
 	}
 
+	void RegisterGPUWork(uint32 NumPrimitives = 0, uint32 NumVertices = 0)
+	{
+		RegisterGPUWork(1, NumPrimitives, NumVertices);
+	}
+
 	void RegisterGPUDispatch(FIntVector GroupCount)
 	{
 		if (bTrackingEvents && CurrentEventNode)
