[LLM] - further LLM breakdowns for D3D12 allocations

diff --git a/Engine/Source/Runtime/D3D12RHI/Private/D3D12CommandList.cpp b/Engine/Source/Runtime/D3D12RHI/Private/D3D12CommandList.cpp
--- a/Engine/Source/Runtime/D3D12RHI/Private/D3D12CommandList.cpp
+++ b/Engine/Source/Runtime/D3D12RHI/Private/D3D12CommandList.cpp
@@ -212,6 +212,7 @@ void FD3D12CommandListHandle::FD3D12CommandListData::FlushResourceBarriers()
 
 void FD3D12CommandListHandle::FD3D12CommandListData::Reset(FD3D12CommandAllocator& CommandAllocator, bool bTrackExecTime)
 {
+	LLM_SCOPE_BYNAME(TEXT("D3D12/Commandlist"));
 	VERIFYD3D12RESULT(CommandList->Reset(CommandAllocator, nullptr));
 
 	CurrentCommandAllocator = &CommandAllocator;
diff --git a/Engine/Source/Runtime/D3D12RHI/Private/D3D12Texture.cpp b/Engine/Source/Runtime/D3D12RHI/Private/D3D12Texture.cpp
--- a/Engine/Source/Runtime/D3D12RHI/Private/D3D12Texture.cpp
+++ b/Engine/Source/Runtime/D3D12RHI/Private/D3D12Texture.cpp
@@ -2103,6 +2103,7 @@ D3D12_RESOURCE_DESC FD3D12DynamicRHI::GetResourceDesc(const FRHITextureCreateInf
 
 void FD3D12TextureBase::UpdateTexture(uint32 MipIndex, uint32 DestX, uint32 DestY, uint32 DestZ, const D3D12_TEXTURE_COPY_LOCATION& SourceCopyLocation)
 {
+	LLM_SCOPE_BYNAME(TEXT("D3D12/CopyTextureRegion"));
 	FD3D12CommandContext& DefaultContext = GetParentDevice()->GetDefaultCommandContext();
 	FD3D12CommandListHandle& hCommandList = DefaultContext.CommandListHandle;
 
