Implemented a 64bit context register sort key on the base FD3D12PipelineState, such that certain RHIs can populate this for more accurate sorting of draw calls for context roll reduction than just using the cached pipeline id in the high level code.

diff --git a/Engine/Source/Runtime/D3D12RHI/Private/D3D12PipelineState.h b/Engine/Source/Runtime/D3D12RHI/Private/D3D12PipelineState.h
--- a/Engine/Source/Runtime/D3D12RHI/Private/D3D12PipelineState.h
+++ b/Engine/Source/Runtime/D3D12RHI/Private/D3D12PipelineState.h
@@ -339,6 +339,16 @@ public:
 		}
 	}
 
+	FORCEINLINE uint64 GetContextSortKey() const
+	{
+		return ContextSortKey;
+	}
+
+	FORCEINLINE void SetContextSortKey(uint64 InContextSortKey)
+	{
+		ContextSortKey = InContextSortKey;
+	}
+
 	FD3D12PipelineState& operator=(const FD3D12PipelineState& other) = delete;
 
 private:
@@ -356,6 +366,9 @@ protected:
 		CreationFailed,
 	};
 	volatile PSOInitState InitState;
+
+	// GRHISupportsPipelineStateSortKey
+	uint64 ContextSortKey = 0;
 };
 
 struct FD3D12GraphicsPipelineState : public FRHIGraphicsPipelineState
diff --git a/Engine/Source/Runtime/D3D12RHI/Private/D3D12Pipelinestate.cpp b/Engine/Source/Runtime/D3D12RHI/Private/D3D12Pipelinestate.cpp
--- a/Engine/Source/Runtime/D3D12RHI/Private/D3D12Pipelinestate.cpp
+++ b/Engine/Source/Runtime/D3D12RHI/Private/D3D12Pipelinestate.cpp
@@ -330,6 +330,9 @@ FD3D12GraphicsPipelineState::FD3D12GraphicsPipelineState(
 	bShaderNeedsGlobalConstantBuffer[SF_Amplification] = GetAmplificationShader() && GetAmplificationShader()->ResourceCounts.bGlobalUniformBufferUsed;
 	bShaderNeedsGlobalConstantBuffer[SF_Pixel] = GetPixelShader() && GetPixelShader()->ResourceCounts.bGlobalUniformBufferUsed;
 	bShaderNeedsGlobalConstantBuffer[SF_Geometry] = GetGeometryShader() && GetGeometryShader()->ResourceCounts.bGlobalUniformBufferUsed;
+
+	// GRHISupportsPipelineStateSortKey
+	SetSortKey(InPipelineState->GetContextSortKey());
 }
 
 FD3D12GraphicsPipelineState::~FD3D12GraphicsPipelineState()
