FD3D12Resource cleanup:
- FD3D12Resource reduced by 32 bytes.
- FD3D12ResourceLocation reduced by 8 bytes.
- FD3D12TextureBase reduced by 16 bytes (24 total).

diff --git a/Engine/Source/Runtime/D3D12RHI/Private/D3D12Resources.cpp b/Engine/Source/Runtime/D3D12RHI/Private/D3D12Resources.cpp
--- a/Engine/Source/Runtime/D3D12RHI/Private/D3D12Resources.cpp
+++ b/Engine/Source/Runtime/D3D12RHI/Private/D3D12Resources.cpp
@@ -249,10 +249,9 @@ FD3D12Resource::FD3D12Resource(FD3D12Device* ParentDevice,
 	ID3D12Resource* InResource,
 	D3D12_RESOURCE_STATES InInitialState,
 	const FD3D12ResourceDesc& InDesc,
-	const D3D12_CLEAR_VALUE* InClearValue,
 	FD3D12Heap* InHeap,
 	D3D12_HEAP_TYPE InHeapType) : 
-	FD3D12Resource(ParentDevice, VisibleNodes, InResource, InInitialState, ED3D12ResourceStateMode::Default, D3D12_RESOURCE_STATE_TBD, InDesc, InClearValue, InHeap, InHeapType)
+	FD3D12Resource(ParentDevice, VisibleNodes, InResource, InInitialState, ED3D12ResourceStateMode::Default, D3D12_RESOURCE_STATE_TBD, InDesc, InHeap, InHeapType)
 {
 }
 
@@ -263,40 +262,24 @@ FD3D12Resource::FD3D12Resource(FD3D12Device* ParentDevice,
 	ED3D12ResourceStateMode InResourceStateMode,
 	D3D12_RESOURCE_STATES InDefaultResourceState,
 	const FD3D12ResourceDesc& InDesc,
-	const D3D12_CLEAR_VALUE* InClearValue,
 	FD3D12Heap* InHeap,
 	D3D12_HEAP_TYPE InHeapType)
 	: FD3D12DeviceChild(ParentDevice)
 	, FD3D12MultiNodeGPUObject(ParentDevice->GetGPUMask(), VisibleNodes)
 	, Resource(InResource)
 	, Heap(InHeap)
-	, ResidencyHandle()
 	, Desc(InDesc)
+	, HeapType(InHeapType)
 	, PlaneCount(::GetPlaneCount(Desc.Format))
-	, SubresourceCount(0)
-	, DefaultResourceState(D3D12_RESOURCE_STATE_TBD)
 	, bRequiresResourceStateTracking(true)
 	, bDepthStencil(false)
 	, bDeferDelete(true)
 	, bBackBuffer(false)
-	, HeapType(InHeapType)
-	, GPUVirtualAddress(0)
-	, ResourceBaseAddress(nullptr)
 {
 #if UE_BUILD_DEBUG
 	FPlatformAtomics::InterlockedIncrement(&TotalResourceCount);
 #endif
 
-	// Store the optional clear value - need for defrag and transient resource recreation
-	if (InClearValue)
-	{
-		ClearValue = *InClearValue;
-	}
-	else
-	{
-		FPlatformMemory::Memzero(&ClearValue, sizeof(D3D12_CLEAR_VALUE));
-	}
-
 	// On Windows it's sadly enough not possible to get the GPU virtual address from the resource directly
 	if (Resource
 #if PLATFORM_WINDOWS
@@ -515,7 +498,7 @@ HRESULT FD3D12Adapter::CreateCommittedResource(const FD3D12ResourceDesc& InDesc,
 	if (SUCCEEDED(hr))
 	{
 		// Set the output pointer
-		*ppOutResource = new FD3D12Resource(GetDevice(CreationNode.ToIndex()), CreationNode, pResource, InInitialState, InResourceStateMode, InDefaultState, InDesc, ClearValue, nullptr, HeapProps.Type);
+		*ppOutResource = new FD3D12Resource(GetDevice(CreationNode.ToIndex()), CreationNode, pResource, InInitialState, InResourceStateMode, InDefaultState, InDesc, nullptr, HeapProps.Type);
 		(*ppOutResource)->AddRef();
 
 		// Set a default name (can override later).
@@ -566,7 +549,6 @@ HRESULT FD3D12Adapter::CreatePlacedResource(const FD3D12ResourceDesc& InDesc, FD
 			InResourceStateMode,
 			InDefaultState,
 			InDesc,
-			ClearValue,
 			BackingHeap,
 			HeapDesc.Properties.Type);
 
@@ -651,17 +633,7 @@ HRESULT FD3D12Adapter::CreateBuffer(const D3D12_HEAP_PROPERTIES& HeapProps,
 
 FD3D12ResourceLocation::FD3D12ResourceLocation(FD3D12Device* Parent)
 	: FD3D12DeviceChild(Parent)
-	, Type(ResourceLocationType::eUndefined)
-	, Owner(nullptr)
-	, UnderlyingResource(nullptr)
-	, ResidencyHandle(nullptr)
 	, Allocator(nullptr)
-	, MappedBaseAddress(nullptr)
-	, GPUVirtualAddress(0)
-	, OffsetFromBaseOfResource(0)
-	, Size(0)
-	, bTransient(false)
-	, AllocatorType(AT_Unknown)
 {
 	FMemory::Memzero(AllocatorData);
 }
diff --git a/Engine/Source/Runtime/D3D12RHI/Private/D3D12Texture.h b/Engine/Source/Runtime/D3D12RHI/Private/D3D12Texture.h
--- a/Engine/Source/Runtime/D3D12RHI/Private/D3D12Texture.h
+++ b/Engine/Source/Runtime/D3D12RHI/Private/D3D12Texture.h
@@ -53,10 +53,6 @@ public:
 
 	FD3D12TextureBase(class FD3D12Device* InParent)
 		: FD3D12BaseShaderResource(InParent)
-		, MemorySize(0)
-		, BaseShaderResource(this)
-		, bCreatedRTVsPerSlice(false)
-		, NumDepthStencilViews(0)
 	{
 	}
 
@@ -79,7 +75,6 @@ public:
 		if (SubResourceIndex < FExclusiveDepthStencil::MaxIndex)
 		{
 			DepthStencilViews[SubResourceIndex] = View;
-			NumDepthStencilViews = FMath::Max(SubResourceIndex + 1, NumDepthStencilViews);
 		}
 		else
 		{
@@ -125,7 +120,6 @@ public:
 	FD3D12Resource* GetResource() const { return ResourceLocation.GetResource(); }
 	uint64 GetOffset() const { return ResourceLocation.GetOffsetFromBaseOfResource(); }
 	FD3D12ShaderResourceView* GetShaderResourceView() const { return ShaderResourceView; }
-	FD3D12BaseShaderResource* GetBaseShaderResource() const { return BaseShaderResource; }
 	inline const FTextureRHIRef& GetAliasingSourceTexture() const { return AliasingSourceTexture; }
 
 	void SetShaderResourceView(FD3D12ShaderResourceView* InShaderResourceView) { ShaderResourceView = InShaderResourceView; }
@@ -180,13 +174,6 @@ public:
 		return DepthStencilViews[AccessType.GetIndex()];
 	}
 
-	// New Monolithic Graphics drivers have optional "fast calls" replacing various D3d functions
-	// You can't use fast version of XXSetShaderResources (called XXSetFastShaderResource) on dynamic or d/s targets
-	inline bool HasDepthStencilView() const
-	{
-		return (NumDepthStencilViews > 0);
-	}
-
 	inline bool HasRenderTargetViews() const
 	{
 		return (RenderTargetViews.Num() > 0);
@@ -197,9 +184,6 @@ public:
 		// Alias the location, will perform an addref underneath
 		FD3D12ResourceLocation::Alias(ResourceLocation, Texture->ResourceLocation);
 
-		// Do not copy the BaseShaderResource from the source texture (this is initialized correctly here, and is used for
-		// state caching logic).
-
 		ShaderResourceView = Texture->ShaderResourceView;
 
 		for (uint32 Index = 0; Index < FExclusiveDepthStencil::MaxIndex; Index++)
@@ -221,10 +205,7 @@ public:
 protected:
 
 	/** Amount of memory allocated by this texture, in bytes. */
-	int64 MemorySize;
-
-	/** Pointer to the base shader resource. Usually the object itself, but not for texture references. */
-	FD3D12BaseShaderResource* BaseShaderResource;
+	int64 MemorySize{};
 
 	/** A shader resource view of the texture. */
 	TRefCountPtr<FD3D12ShaderResourceView> ShaderResourceView;
@@ -232,19 +213,16 @@ protected:
 	/** A render targetable view of the texture. */
 	TArray<TRefCountPtr<FD3D12RenderTargetView>, TInlineAllocator<1>> RenderTargetViews;
 
-	bool bCreatedRTVsPerSlice;
-
-	int32 RTVArraySize;
-
 	/** A depth-stencil targetable view of the texture. */
 	TRefCountPtr<FD3D12DepthStencilView> DepthStencilViews[FExclusiveDepthStencil::MaxIndex];
 
-	/** Number of Depth Stencil Views - used for fast call tracking. */
-	uint32	NumDepthStencilViews;
-
 	TMap<uint32, FD3D12LockedResource*> LockedMap;
 
 	FTextureRHIRef AliasingSourceTexture;
+
+	int32 RTVArraySize{};
+
+	bool bCreatedRTVsPerSlice{ false };
 };
 
 #if PLATFORM_WINDOWS || PLATFORM_HOLOLENS
diff --git a/Engine/Source/Runtime/D3D12RHI/Public/D3D12Resources.h b/Engine/Source/Runtime/D3D12RHI/Public/D3D12Resources.h
--- a/Engine/Source/Runtime/D3D12RHI/Public/D3D12Resources.h
+++ b/Engine/Source/Runtime/D3D12RHI/Public/D3D12Resources.h
@@ -154,46 +154,46 @@ private:
 
 	FD3D12ResidencyHandle ResidencyHandle;
 
-	D3D12_CLEAR_VALUE ClearValue;
+	D3D12_GPU_VIRTUAL_ADDRESS GPUVirtualAddress{};
+	void* ResourceBaseAddress{};
+
+#if NV_AFTERMATH
+	GFSDK_Aftermath_ResourceHandle AftermathHandle{};
+#endif
+
 	const FD3D12ResourceDesc Desc;
-	uint8 PlaneCount;
-	uint16 SubresourceCount;
 	CResourceState ResourceState;
-	D3D12_RESOURCE_STATES DefaultResourceState;
-	D3D12_RESOURCE_STATES ReadableState;
-	D3D12_RESOURCE_STATES WritableState;
+	D3D12_RESOURCE_STATES DefaultResourceState{ D3D12_RESOURCE_STATE_TBD };
+	D3D12_RESOURCE_STATES ReadableState{ D3D12_RESOURCE_STATE_CORRUPT };
+	D3D12_RESOURCE_STATES WritableState{ D3D12_RESOURCE_STATE_CORRUPT };
 #ifdef PLATFORM_SUPPORTS_RESOURCE_COMPRESSION
-	D3D12_RESOURCE_STATES CompressedState;
+	D3D12_RESOURCE_STATES CompressedState{ D3D12_RESOURCE_STATE_CORRUPT };
 #endif
 	D3D12_RESOURCE_STATES UAVHiddenResourceState{ D3D12_RESOURCE_STATE_CORRUPT };
 
+	D3D12_HEAP_TYPE HeapType;
+	FName DebugName;
+
+	int32 NumMapCalls = 0;
+	uint16 SubresourceCount{};
+	uint8 PlaneCount;
 	bool bRequiresResourceStateTracking : 1;
 	bool bDepthStencil : 1;
 	bool bDeferDelete : 1;
 	bool bBackBuffer : 1;
 
-	D3D12_HEAP_TYPE HeapType;
-	D3D12_GPU_VIRTUAL_ADDRESS GPUVirtualAddress;
-	void* ResourceBaseAddress;
-	int32 NumMapCalls = 0;
-	FName DebugName;
-
-#if NV_AFTERMATH
-	GFSDK_Aftermath_ResourceHandle AftermathHandle;
-#endif
-
 #if UE_BUILD_DEBUG
 	static int64 TotalResourceCount;
 	static int64 NoStateTrackingResourceCount;
 #endif
 
 public:
+	FD3D12Resource() = delete;
 	explicit FD3D12Resource(FD3D12Device* ParentDevice,
 		FRHIGPUMask VisibleNodes,
 		ID3D12Resource* InResource,
 		D3D12_RESOURCE_STATES InInitialResourceState,
 		FD3D12ResourceDesc const& InDesc,
-		const D3D12_CLEAR_VALUE* InClearValue = nullptr,
 		FD3D12Heap* InHeap = nullptr,
 		D3D12_HEAP_TYPE InHeapType = D3D12_HEAP_TYPE_DEFAULT);
 
@@ -204,7 +204,6 @@ public:
 		ED3D12ResourceStateMode InResourceStateMode,
 		D3D12_RESOURCE_STATES InDefaultResourceState,
 		FD3D12ResourceDesc const& InDesc,
-		const D3D12_CLEAR_VALUE* InClearValue,
 		FD3D12Heap* InHeap,
 		D3D12_HEAP_TYPE InHeapType);
 
@@ -223,9 +222,9 @@ public:
 	{
 		if (NumMapCalls == 0)
 		{
-		check(Resource);
-		check(ResourceBaseAddress == nullptr);
-		VERIFYD3D12RESULT(Resource->Map(0, ReadRange, &ResourceBaseAddress));
+			check(Resource);
+			check(ResourceBaseAddress == nullptr);
+			VERIFYD3D12RESULT(Resource->Map(0, ReadRange, &ResourceBaseAddress));
 		}
 		else
 		{
@@ -245,14 +244,13 @@ public:
 		--NumMapCalls;
 		if (NumMapCalls == 0)
 		{
-		Resource->Unmap(0, nullptr);
-		ResourceBaseAddress = nullptr;
-	}
+			Resource->Unmap(0, nullptr);
+			ResourceBaseAddress = nullptr;
+		}
 	}
 
 	ID3D12Pageable* GetPageable();
 	const FD3D12ResourceDesc& GetDesc() const { return Desc; }
-	D3D12_CLEAR_VALUE const& GetClearValue() const { return ClearValue; }
 	D3D12_HEAP_TYPE GetHeapType() const { return HeapType; }
 	D3D12_GPU_VIRTUAL_ADDRESS GetGPUVirtualAddress() const { return GPUVirtualAddress; }
 	inline void SetGPUVirtualAddress(D3D12_GPU_VIRTUAL_ADDRESS Value) { GPUVirtualAddress = Value; }
@@ -531,8 +529,7 @@ class FD3D12BaseShaderResource;
 class FD3D12ResourceLocation : public FRHIPoolResource, public FD3D12DeviceChild, public FNoncopyable
 {
 public:
-
-	enum class ResourceLocationType
+	enum class ResourceLocationType : uint8
 	{
 		eUndefined,
 		eStandAlone,
@@ -670,11 +667,9 @@ private:
 	void ReleaseResource();
 	void UpdateStandAloneStats(bool bIncrement);
 
-	ResourceLocationType Type;
-
-	FD3D12BaseShaderResource* Owner;
-	FD3D12Resource* UnderlyingResource;
-	FD3D12ResidencyHandle* ResidencyHandle;
+	FD3D12BaseShaderResource* Owner{};
+	FD3D12Resource* UnderlyingResource{};
+	FD3D12ResidencyHandle* ResidencyHandle{};
 
 	// Which allocator this belongs to
 	union
@@ -694,16 +689,16 @@ private:
 	} AllocatorData;
 
 	// Note: These values refer to the start of this location including any padding *NOT* the start of the underlying resource
-	void* MappedBaseAddress;
-	D3D12_GPU_VIRTUAL_ADDRESS GPUVirtualAddress;
-	uint64 OffsetFromBaseOfResource;
+	void* MappedBaseAddress{};
+	D3D12_GPU_VIRTUAL_ADDRESS GPUVirtualAddress{};
+	uint64 OffsetFromBaseOfResource{};
 
 	// The size the application asked for
-	uint64 Size;
-
-	bool bTransient;
+	uint64 Size{};
 
-	EAllocatorType AllocatorType;
+	ResourceLocationType Type{ ResourceLocationType::eUndefined };
+	EAllocatorType AllocatorType{ AT_Unknown };
+	bool bTransient{ false };
 };
 
 // Generic interface for every type D3D12 specific allocator
