Support initializing FD3D12RayTracingGeometry using offline data.

diff --git a/Engine/Source/Runtime/D3D12RHI/Private/D3D12RayTracing.cpp b/Engine/Source/Runtime/D3D12RHI/Private/D3D12RayTracing.cpp
--- a/Engine/Source/Runtime/D3D12RHI/Private/D3D12RayTracing.cpp
+++ b/Engine/Source/Runtime/D3D12RHI/Private/D3D12RayTracing.cpp
@@ -2525,6 +2525,50 @@ public:
 	FD3D12RayTracingPipelineState* Intersection;
 };
 
+struct FD3D12RHICommandInitializeRayTracingGeometryString
+{
+	static const TCHAR* TStr() { return TEXT("FD3D12RHICommandInitializeRayTracingGeometryString"); }
+};
+struct FD3D12RHICommandInitializeRayTracingGeometry final : public FRHICommand<FD3D12RHICommandInitializeRayTracingGeometry, FD3D12RHICommandInitializeRayTracingGeometryString>
+{
+	FD3D12RayTracingGeometry* Geometry;
+	FD3D12ResourceLocation SrcResourceLoc;
+
+	FORCEINLINE_DEBUGGABLE FD3D12RHICommandInitializeRayTracingGeometry(FD3D12RayTracingGeometry* InGeometry, FD3D12ResourceLocation& InSrcResourceLoc)
+		: Geometry(InGeometry)
+		, SrcResourceLoc(InSrcResourceLoc.GetParentDevice())
+	{
+		FD3D12ResourceLocation::TransferOwnership(SrcResourceLoc, InSrcResourceLoc);
+	}
+
+	void Execute(FRHICommandListBase& /* unused */)
+	{
+		ExecuteNoCmdList();
+	}
+
+	void ExecuteNoCmdList()
+	{
+		for (uint32 GPUIndex = 0; GPUIndex < MAX_NUM_GPUS && GPUIndex < GNumExplicitGPUsForRendering; ++GPUIndex)
+		{
+			FD3D12Buffer* AccelerationStructure = Geometry->AccelerationStructureBuffers[GPUIndex];
+			FD3D12Resource* Destination = AccelerationStructure->ResourceLocation.GetResource();
+			FD3D12Device* Device = Destination->GetParentDevice();
+
+			FD3D12CommandListHandle& hCommandList = Device->GetDefaultCommandContext().CommandListHandle;
+			hCommandList.GetCurrentOwningContext()->numInitialResourceCopies += 1;
+
+			ID3D12GraphicsCommandList4* CmdList = hCommandList.RayTracingCommandList();
+
+			CmdList->CopyRaytracingAccelerationStructure(AccelerationStructure->ResourceLocation.GetGPUVirtualAddress(), SrcResourceLoc.GetGPUVirtualAddress(), D3D12_RAYTRACING_ACCELERATION_STRUCTURE_COPY_MODE_DESERIALIZE);
+
+			Device->GetDefaultCommandContext().ConditionalFlushCommandList();
+
+			Geometry->SetupHitGroupSystemParameters(GPUIndex);
+			Geometry->SetDirty(FRHIGPUMask::FromIndex(GPUIndex), false);
+		}
+	}
+};
+
 void FD3D12Device::InitRayTracing()
 {
 	LLM_SCOPE_BYNAME(TEXT("FD3D12RT"));
@@ -2938,8 +2982,28 @@ FD3D12RayTracingGeometry::FD3D12RayTracingGeometry(FD3D12Adapter* Adapter, const
 	PrebuildDescInputs.pGeometryDescs = GeometryDescs.GetData();
 	PrebuildDescInputs.Flags = TranslateRayTracingAccelerationStructureFlags(BuildFlags);
 
-	// Get maximum buffer sizes for all GPUs in the system
-	SizeInfo = RHICalcRayTracingGeometrySize(Initializer);
+	struct FOfflineBVHHeader
+	{
+		uint64 ResultSize;
+		uint64 BuildScratchSize;
+		uint64 UpdateScratchSize;
+	};
+
+	if (Initializer.OfflineData)
+	{
+		FOfflineBVHHeader* DataHeader = (FOfflineBVHHeader*)Initializer.OfflineData->GetResourceData();
+
+		SizeInfo.ResultSize = DataHeader->ResultSize;
+		SizeInfo.BuildScratchSize = DataHeader->BuildScratchSize;
+		SizeInfo.UpdateScratchSize = DataHeader->UpdateScratchSize;
+
+		AccelerationStructureCompactedSize = DataHeader->ResultSize;
+	}
+	else
+	{
+		// Get maximum buffer sizes for all GPUs in the system
+		SizeInfo = RHICalcRayTracingGeometrySize(Initializer);
+	}
 
 	// Allocate acceleration structure buffer
 	for (uint32 GPUIndex = 0; GPUIndex < MAX_NUM_GPUS && GPUIndex < GNumExplicitGPUsForRendering; ++GPUIndex)
@@ -2953,6 +3017,32 @@ FD3D12RayTracingGeometry::FD3D12RayTracingGeometry(FD3D12Adapter* Adapter, const
 	INC_DWORD_STAT_BY(STAT_D3D12RayTracingTrianglesBLAS, Initializer.TotalPrimitiveCount);
 
 	RegisterD3D12RayTracingGeometry(this);
+
+	if (Initializer.OfflineData)
+	{
+		FD3D12Device* Device = Adapter->GetDevice(0);
+
+		const uint8* Data = (const uint8*)Initializer.OfflineData->GetResourceData();
+		uint32 Size = Initializer.OfflineData->GetResourceDataSize();
+
+		FD3D12ResourceLocation SrcResourceLoc(Device);
+		uint8* DstDataBase = (uint8*)Device->GetDefaultFastAllocator().Allocate(Size, 256, &SrcResourceLoc);
+		FMemory::Memcpy(DstDataBase, Data + sizeof(FOfflineBVHHeader), Size);
+
+		FRHICommandListImmediate& RHICmdList = FRHICommandListExecutor::GetImmediateCommandList();
+
+		if (!RHICmdList.Bypass() && IsRunningRHIInSeparateThread())
+		{
+			ALLOC_COMMAND_CL(RHICmdList, FD3D12RHICommandInitializeRayTracingGeometry)(this, SrcResourceLoc);
+		}
+		else
+		{
+			FD3D12RHICommandInitializeRayTracingGeometry Command(this, SrcResourceLoc);
+			Command.ExecuteNoCmdList();
+		}
+		
+		Initializer.OfflineData->Discard();
+	}
 }
 
 FD3D12RayTracingGeometry::~FD3D12RayTracingGeometry()
