Break up D3D12RaytracingGeometry::BuildAccelerationStructure into separate parts so they can be modified per platform more easily.
Additionally, clean up the code to avoid any unexpected side effects when building bottom level structures.

diff --git a/Engine/Source/Runtime/D3D12RHI/Private/D3D12CommandContext.h b/Engine/Source/Runtime/D3D12RHI/Private/D3D12CommandContext.h
--- a/Engine/Source/Runtime/D3D12RHI/Private/D3D12CommandContext.h
+++ b/Engine/Source/Runtime/D3D12RHI/Private/D3D12CommandContext.h
@@ -383,6 +383,7 @@ public:
 	virtual void RHICopyBufferRegion(FRHIBuffer* DestBuffer, uint64 DstOffset, FRHIBuffer* SourceBuffer, uint64 SrcOffset, uint64 NumBytes) final override;
 	virtual void RHICopyBufferRegions(const TArrayView<const FCopyBufferRegionParams> Params) final override;
 	virtual void RHIBindAccelerationStructureMemory(FRHIRayTracingScene* Scene, FRHIBuffer* Buffer, uint32 BufferOffset) final override;
+	virtual void BuildAccelerationStructuresInternal(const TArrayView<D3D12_BUILD_RAYTRACING_ACCELERATION_STRUCTURE_DESC> BuildDesc);
 	virtual void RHIBuildAccelerationStructures(const TArrayView<const FRayTracingGeometryBuildParams> Params) final override;
 	virtual void RHIBuildAccelerationStructure(const FRayTracingSceneBuildParams& SceneBuildParams) final override;
 	virtual void RHIClearRayTracingBindings(FRHIRayTracingScene* Scene) final override;
diff --git a/Engine/Source/Runtime/D3D12RHI/Private/D3D12CommandList.h b/Engine/Source/Runtime/D3D12RHI/Private/D3D12CommandList.h
--- a/Engine/Source/Runtime/D3D12RHI/Private/D3D12CommandList.h
+++ b/Engine/Source/Runtime/D3D12RHI/Private/D3D12CommandList.h
@@ -446,7 +446,7 @@ public:
 #if D3D12_MAX_COMMANDLIST_INTERFACE >= 6
 	ID3D12GraphicsCommandList6* GraphicsCommandList6() const
 	{
-		check(CommandListData && (CommandListData->CommandListType == D3D12_COMMAND_LIST_TYPE_DIRECT));
+		check(CommandListData && (CommandListData->CommandListType == D3D12_COMMAND_LIST_TYPE_DIRECT || CommandListData->CommandListType == D3D12_COMMAND_LIST_TYPE_COMPUTE));
 		return CommandListData->CommandList6.GetReference();
 	}
 #endif
diff --git a/Engine/Source/Runtime/D3D12RHI/Private/D3D12RayTracing.cpp b/Engine/Source/Runtime/D3D12RHI/Private/D3D12RayTracing.cpp
--- a/Engine/Source/Runtime/D3D12RHI/Private/D3D12RayTracing.cpp
+++ b/Engine/Source/Runtime/D3D12RHI/Private/D3D12RayTracing.cpp
@@ -2965,7 +2965,6 @@ FD3D12RayTracingGeometry::FD3D12RayTracingGeometry(FD3D12Adapter* Adapter, const
 	Segments = TArray<FRayTracingGeometrySegment>(Initializer.Segments.GetData(), Initializer.Segments.Num());
 
 	RHIIndexBuffer = Initializer.IndexBuffer;
-	IndexStride = RHIIndexBuffer ? RHIIndexBuffer->GetStride() : 0; // stride 0 means implicit triangle list for non-indexed geometry
 
 	static constexpr uint32 IndicesPerPrimitive = 3; // Only triangle meshes are supported
 
@@ -3188,12 +3187,32 @@ void FD3D12RayTracingGeometry::TransitionBuffers(FD3D12CommandContext& CommandCo
 	}
 }
 
+void FD3D12RayTracingGeometry::UpdateResidency(FD3D12CommandContext& CommandContext)
+{
+	if (RHIIndexBuffer)
+	{
+		FD3D12Buffer* IndexBuffer = CommandContext.RetrieveObject<FD3D12Buffer>(RHIIndexBuffer.GetReference());
+		IndexBuffer->GetResource()->UpdateResidency(CommandContext.CommandListHandle);		
+	}
+
+	for (const FRayTracingGeometrySegment& Segment : Segments)
+	{
+		const FBufferRHIRef& RHIVertexBuffer = Segment.VertexBuffer;
+		FD3D12Buffer* VertexBuffer = CommandContext.RetrieveObject<FD3D12Buffer>(RHIVertexBuffer.GetReference());		
+		VertexBuffer->ResourceLocation.GetResource()->UpdateResidency(CommandContext.CommandListHandle);
+	}
+
+	const uint32 GPUIndex = CommandContext.GetGPUIndex();
+	AccelerationStructureBuffers[GPUIndex]->GetResource()->UpdateResidency(CommandContext.CommandListHandle);
+}
+
 void FD3D12RayTracingGeometry::SetupHitGroupSystemParameters(uint32 InGPUIndex)
 {
 	TArray<FHitGroupSystemParameters>& HitGroupSystemParametersForThisGPU = HitGroupSystemParameters[InGPUIndex];
 	HitGroupSystemParametersForThisGPU.Reset(Segments.Num());
 
 	FD3D12Buffer* IndexBuffer = FD3D12DynamicRHI::ResourceCast(RHIIndexBuffer.GetReference(), InGPUIndex);
+	const uint32 IndexStride = IndexBuffer ? IndexBuffer->GetStride() : 0;
 	for (const FRayTracingGeometrySegment& Segment : Segments)
 	{
 		FD3D12Buffer* VertexBuffer = FD3D12DynamicRHI::ResourceCast(Segment.VertexBuffer.GetReference(), InGPUIndex);
@@ -3213,42 +3232,30 @@ void FD3D12RayTracingGeometry::SetupHitGroupSystemParameters(uint32 InGPUIndex)
 	}
 }
 
-static bool ShouldCompactAfterBuild(ERayTracingAccelerationStructureFlags BuildFlags)
+void FD3D12RayTracingGeometry::CreateAccelerationStructureBuildDesc(FD3D12CommandContext& CommandContext, EAccelerationStructureBuildMode BuildMode, const TRefCountPtr<FD3D12Buffer>& ScratchBuffer, D3D12_BUILD_RAYTRACING_ACCELERATION_STRUCTURE_DESC& OutDesc, TArrayView<D3D12_RAYTRACING_GEOMETRY_DESC>& OutGeometryDescs) const
 {
-	return EnumHasAllFlags(BuildFlags, ERayTracingAccelerationStructureFlags::AllowCompaction | ERayTracingAccelerationStructureFlags::FastTrace)
-		&& !EnumHasAnyFlags(BuildFlags, ERayTracingAccelerationStructureFlags::AllowUpdate);
-}
-
-void FD3D12RayTracingGeometry::BuildAccelerationStructure(FD3D12CommandContext& CommandContext, EAccelerationStructureBuildMode BuildMode)
-{
-	TRACE_CPUPROFILER_EVENT_SCOPE(BuildAccelerationStructure_BottomLevel);
-	SCOPE_CYCLE_COUNTER(STAT_D3D12BuildBLAS);
-	LLM_SCOPE_BYNAME(TEXT("FD3D12RT/BLAS"));
-
 	if (RHIIndexBuffer)
 	{
 		checkf(RHIIndexBuffer->GetStride() == 2 || RHIIndexBuffer->GetStride() == 4, TEXT("Index buffer must be 16 or 32 bit."));
 	}
 
-	IndexStride = RHIIndexBuffer ? RHIIndexBuffer->GetStride() : 0; // stride 0 means implicit triangle list for non-indexed geometry
-
+	const uint32 GPUIndex = CommandContext.GetGPUIndex();
+	const uint32 IndexStride = RHIIndexBuffer ? RHIIndexBuffer->GetStride() : 0;
 	const bool bIsUpdate = BuildMode == EAccelerationStructureBuildMode::Update;
 
 	// Use the pre-built descs as template and set the GPU resource pointers (current VB/IB).
-	TArray<D3D12_RAYTRACING_GEOMETRY_DESC, TInlineAllocator<32>> LocalGeometryDescs;
-	LocalGeometryDescs.Reserve(GeometryDescs.Num());
+	check(OutGeometryDescs.Num() == GeometryDescs.Num());
 
 	FD3D12Buffer* IndexBuffer = CommandContext.RetrieveObject<FD3D12Buffer>(RHIIndexBuffer.GetReference());
 	FD3D12Buffer* NullTransformBufferD3D12 = CommandContext.RetrieveObject<FD3D12Buffer>(NullTransformBuffer.GetReference());
 
-	// Recreate the hit group system parameters and use them during setup of the descs
-	SetupHitGroupSystemParameters(CommandContext.GetGPUIndex());
-	const TArray<FHitGroupSystemParameters>& HitGroupSystemParametersForThisGPU = HitGroupSystemParameters[CommandContext.GetGPUIndex()];
+	const TArray<FHitGroupSystemParameters>& HitGroupSystemParametersForThisGPU = HitGroupSystemParameters[GPUIndex];
 	check(HitGroupSystemParametersForThisGPU.Num() == Segments.Num());
 
 	for (int32 SegmentIndex = 0; SegmentIndex < Segments.Num(); ++SegmentIndex)
 	{
-		D3D12_RAYTRACING_GEOMETRY_DESC Desc = GeometryDescs[SegmentIndex]; // Copy from template
+		D3D12_RAYTRACING_GEOMETRY_DESC& Desc = OutGeometryDescs[SegmentIndex];
+		Desc = GeometryDescs[SegmentIndex]; // Copy from template
 
 		const FRayTracingGeometrySegment& Segment = Segments[SegmentIndex];
 		const FHitGroupSystemParameters& SystemParameters = HitGroupSystemParametersForThisGPU[SegmentIndex];
@@ -3292,7 +3299,6 @@ void FD3D12RayTracingGeometry::BuildAccelerationStructure(FD3D12CommandContext&
 
 				Desc.Triangles.IndexFormat = (IndexStride == 4 ? DXGI_FORMAT_R32_UINT : DXGI_FORMAT_R16_UINT);
 				Desc.Triangles.IndexBuffer = SystemParameters.IndexBuffer + SystemParameters.RootConstants.IndexBufferOffsetInBytes;
-				IndexBuffer->GetResource()->UpdateResidency(CommandContext.CommandListHandle);
 			}
 			else
 			{
@@ -3318,10 +3324,6 @@ void FD3D12RayTracingGeometry::BuildAccelerationStructure(FD3D12CommandContext&
 			break;
 		}
 
-		VertexBuffer->ResourceLocation.GetResource()->UpdateResidency(CommandContext.CommandListHandle);
-
-		LocalGeometryDescs.Add(Desc);
-
 		if (GeometryType == D3D12_RAYTRACING_GEOMETRY_TYPE_TRIANGLES)
 		{
 			// #dxr_todo UE-72160: support various vertex buffer layouts (fetch/decode based on vertex stride and format)
@@ -3329,15 +3331,6 @@ void FD3D12RayTracingGeometry::BuildAccelerationStructure(FD3D12CommandContext&
 		}
 	}
 
-	const uint32 GPUIndex = CommandContext.GetGPUIndex();
-	FD3D12Adapter* Adapter = CommandContext.GetParentAdapter();
-
-	// Register as rename listener to index/vertex buffers
-	UnregisterAsRenameListener(GPUIndex);
-	RegisterAsRenameListener(GPUIndex);
-
-	ID3D12Device5* RayTracingDevice = CommandContext.GetParentDevice()->GetDevice5();
-
 	D3D12_RAYTRACING_ACCELERATION_STRUCTURE_BUILD_FLAGS LocalBuildFlags = D3D12_RAYTRACING_ACCELERATION_STRUCTURE_BUILD_FLAGS(BuildFlags);
 
 	if (bIsUpdate)
@@ -3352,70 +3345,42 @@ void FD3D12RayTracingGeometry::BuildAccelerationStructure(FD3D12CommandContext&
 
 	PrebuildDescInputs.Type = D3D12_RAYTRACING_ACCELERATION_STRUCTURE_TYPE_BOTTOM_LEVEL;
 	PrebuildDescInputs.DescsLayout = D3D12_ELEMENTS_LAYOUT_ARRAY;
-	PrebuildDescInputs.NumDescs = LocalGeometryDescs.Num();
-	PrebuildDescInputs.pGeometryDescs = LocalGeometryDescs.GetData();
+	PrebuildDescInputs.NumDescs = OutGeometryDescs.Num();
+	PrebuildDescInputs.pGeometryDescs = OutGeometryDescs.GetData();
 	PrebuildDescInputs.Flags = LocalBuildFlags;
 
 	D3D12_RAYTRACING_ACCELERATION_STRUCTURE_PREBUILD_INFO PrebuildInfo = {};
 
+	ID3D12Device5* RayTracingDevice = CommandContext.GetParentDevice()->GetDevice5();
 	RayTracingDevice->GetRaytracingAccelerationStructurePrebuildInfo(&PrebuildDescInputs, &PrebuildInfo);
 
 	// Must make sure that values computed in the constructor are valid.
 	check(PrebuildInfo.ResultDataMaxSizeInBytes <= SizeInfo.ResultSize);
 
-	TRefCountPtr<FD3D12Buffer> ScratchBuffer;
-
 	if (bIsUpdate)
 	{
 		check(PrebuildInfo.UpdateScratchDataSizeInBytes <= SizeInfo.UpdateScratchSize);
-		static const FName BufferName("UpdateScratchBLAS");
-		ScratchBuffer = CreateRayTracingBuffer(Adapter, GPUIndex, SizeInfo.UpdateScratchSize, ERayTracingBufferType::Scratch, BufferName);
 	}
 	else
 	{
 		check(PrebuildInfo.ScratchDataSizeInBytes <= SizeInfo.BuildScratchSize);
-		static const FName BufferName("BuildScratchBLAS");
-		ScratchBuffer = CreateRayTracingBuffer(Adapter, GPUIndex, SizeInfo.BuildScratchSize, ERayTracingBufferType::Scratch, BufferName);
 	}
 
-	TransitionBuffers(CommandContext);
-	CommandContext.CommandListHandle.FlushResourceBarriers();
-
-	if (IsDirty(GPUIndex))
-	{
-		ScratchBuffer->GetResource()->UpdateResidency(CommandContext.CommandListHandle);
-		AccelerationStructureBuffers[GPUIndex]->GetResource()->UpdateResidency(CommandContext.CommandListHandle);
-
-		D3D12_BUILD_RAYTRACING_ACCELERATION_STRUCTURE_DESC BuildDesc = {};
-		BuildDesc.Inputs = PrebuildDescInputs;
-		BuildDesc.DestAccelerationStructureData = AccelerationStructureBuffers[GPUIndex]->ResourceLocation.GetGPUVirtualAddress();
-		BuildDesc.ScratchAccelerationStructureData = ScratchBuffer->ResourceLocation.GetGPUVirtualAddress();
-		BuildDesc.SourceAccelerationStructureData = bIsUpdate
-			? AccelerationStructureBuffers[GPUIndex]->ResourceLocation.GetGPUVirtualAddress()
-			: D3D12_GPU_VIRTUAL_ADDRESS(0);
-
-		// Build
-		ID3D12GraphicsCommandList4* RayTracingCommandList = CommandContext.CommandListHandle.RayTracingCommandList();
-		RayTracingCommandList->BuildRaytracingAccelerationStructure(&BuildDesc, 0, nullptr);
-
-		// Request compact BLAS if it is static
-		if (ShouldCompactAfterBuild(BuildFlags))
-		{
-			CommandContext.GetParentDevice()->GetRayTracingCompactionRequestHandler()->RequestCompact(this);
-			bHasPendingCompactionRequests[GPUIndex] = true;
-		}
-
-		SetDirty(CommandContext.GetGPUMask(), false);
+	D3D12_BUILD_RAYTRACING_ACCELERATION_STRUCTURE_DESC BuildDesc = {};
+	BuildDesc.Inputs = PrebuildDescInputs;
+	BuildDesc.DestAccelerationStructureData = AccelerationStructureBuffers[GPUIndex]->ResourceLocation.GetGPUVirtualAddress();
+	BuildDesc.ScratchAccelerationStructureData = ScratchBuffer->ResourceLocation.GetGPUVirtualAddress();
+	BuildDesc.SourceAccelerationStructureData = bIsUpdate
+		? AccelerationStructureBuffers[GPUIndex]->ResourceLocation.GetGPUVirtualAddress()
+		: D3D12_GPU_VIRTUAL_ADDRESS(0);
+
+	OutDesc = BuildDesc;
+}
 
-		if (bIsUpdate)
-		{
-			INC_DWORD_STAT(STAT_D3D12RayTracingUpdatedBLAS);
-		}
-		else
-		{
-			INC_DWORD_STAT(STAT_D3D12RayTracingBuiltBLAS);
-		}
-	}
+static bool ShouldCompactAfterBuild(ERayTracingAccelerationStructureFlags BuildFlags)
+{
+	return EnumHasAllFlags(BuildFlags, ERayTracingAccelerationStructureFlags::AllowCompaction | ERayTracingAccelerationStructureFlags::FastTrace)
+		&& !EnumHasAnyFlags(BuildFlags, ERayTracingAccelerationStructureFlags::AllowUpdate);
 }
 
 void FD3D12RayTracingGeometry::CompactAccelerationStructure(FD3D12CommandContext& CommandContext, uint32 InGPUIndex, uint64 InSizeAfterCompaction)
@@ -3936,8 +3901,20 @@ FD3D12RayTracingShaderTable* FD3D12RayTracingScene::FindOrCreateShaderTable(cons
 	return CreatedShaderTable;
 }
 
+void FD3D12CommandContext::BuildAccelerationStructuresInternal(const TArrayView<D3D12_BUILD_RAYTRACING_ACCELERATION_STRUCTURE_DESC> BuildDescs)
+{
+	for (const D3D12_BUILD_RAYTRACING_ACCELERATION_STRUCTURE_DESC& Desc : BuildDescs)
+	{
+		CommandListHandle.GraphicsCommandList4()->BuildRaytracingAccelerationStructure(&Desc, 0, nullptr);
+	}
+}
+
 void FD3D12CommandContext::RHIBuildAccelerationStructures(const TArrayView<const FRayTracingGeometryBuildParams> Params)
 {
+	TRACE_CPUPROFILER_EVENT_SCOPE(BuildAccelerationStructure_BottomLevel);
+	SCOPE_CYCLE_COUNTER(STAT_D3D12BuildBLAS);
+	LLM_SCOPE_BYNAME(TEXT("FD3D12RT/BLAS"));
+
 	// Update geometry vertex buffers
 
 	for (const FRayTracingGeometryBuildParams& P : Params)
@@ -3974,15 +3951,86 @@ void FD3D12CommandContext::RHIBuildAccelerationStructures(const TArrayView<const
 	}
 
 	CommandListHandle.FlushResourceBarriers();
+		
+	const uint32 GPUIndex = GetGPUIndex();
 
 	// Then do all work
+	TArray<TRefCountPtr<FD3D12Buffer>, TInlineAllocator<32>> ScratchBuffers;
+	TArray<D3D12_BUILD_RAYTRACING_ACCELERATION_STRUCTURE_DESC, TInlineAllocator<32>> BuildDescs;
 
-	for (const FRayTracingGeometryBuildParams& P : Params)
+	BuildDescs.Reserve(Params.Num());
+	ScratchBuffers.Reserve(Params.Num());
+
+	FMemMark Mark(FMemStack::Get());
+
+	for (int32 i = 0; i < Params.Num(); i++)
 	{
+		const FRayTracingGeometryBuildParams& P = Params[i];
+
 		FD3D12RayTracingGeometry* Geometry = FD3D12DynamicRHI::ResourceCast(P.Geometry.GetReference());
 		Geometry->SetDirty(GetGPUMask(), true);
 
-		Geometry->BuildAccelerationStructure(*this, P.BuildMode);
+		// Register as rename listener to index/vertex buffers
+		Geometry->UnregisterAsRenameListener(GPUIndex);
+		Geometry->RegisterAsRenameListener(GPUIndex);
+
+		// Recreate the hit group system parameters and use them during setup of the descs
+		Geometry->SetupHitGroupSystemParameters(GPUIndex);
+
+		if (Geometry->IsDirty(GPUIndex))
+		{
+			TRefCountPtr<FD3D12Buffer> ScratchBuffer;
+			if (P.BuildMode == EAccelerationStructureBuildMode::Update)
+			{
+				static const FName BufferName("UpdateScratchBLAS");
+				ScratchBuffer = CreateRayTracingBuffer(GetParentAdapter(), GPUIndex, Geometry->SizeInfo.UpdateScratchSize, ERayTracingBufferType::Scratch, BufferName);
+			}
+			else
+			{
+				static const FName BufferName("BuildScratchBLAS");
+				ScratchBuffer = CreateRayTracingBuffer(GetParentAdapter(), GPUIndex, Geometry->SizeInfo.BuildScratchSize, ERayTracingBufferType::Scratch, BufferName);
+			}
+
+			// We need to keep D3D12_RAYTRACING_GEOMETRY_DESCs that are used in D3D12_BUILD_RAYTRACING_ACCELERATION_STRUCTURE_DESC alive.
+			const uint32 NumGeometryDescs = Geometry->GeometryDescs.Num();
+			D3D12_RAYTRACING_GEOMETRY_DESC* LocalGeometryDescsMemory = (D3D12_RAYTRACING_GEOMETRY_DESC*)FMemStack::Get().Alloc(NumGeometryDescs * sizeof(D3D12_RAYTRACING_GEOMETRY_DESC), alignof(D3D12_RAYTRACING_GEOMETRY_DESC));
+			TArrayView<D3D12_RAYTRACING_GEOMETRY_DESC> LocalGeometryDescs = MakeArrayView(LocalGeometryDescsMemory, NumGeometryDescs);
+
+			D3D12_BUILD_RAYTRACING_ACCELERATION_STRUCTURE_DESC& BuildDesc = BuildDescs.AddZeroed_GetRef();			
+			Geometry->CreateAccelerationStructureBuildDesc(*this, P.BuildMode, ScratchBuffer, BuildDesc, LocalGeometryDescs);
+
+			Geometry->UpdateResidency(*this);
+			ScratchBuffer->GetResource()->UpdateResidency(CommandListHandle);
+			ScratchBuffers.Add(ScratchBuffer);
+
+			if (P.BuildMode == EAccelerationStructureBuildMode::Update)
+			{
+				INC_DWORD_STAT(STAT_D3D12RayTracingUpdatedBLAS);
+			}
+			else
+			{
+				INC_DWORD_STAT(STAT_D3D12RayTracingBuiltBLAS);
+			}
+		}		
+	}
+
+	TArrayView<D3D12_BUILD_RAYTRACING_ACCELERATION_STRUCTURE_DESC> LocalBuildDescs = MakeArrayView(BuildDescs.GetData(), BuildDescs.Num());
+	BuildAccelerationStructuresInternal(LocalBuildDescs);
+
+	for (const FRayTracingGeometryBuildParams& P : Params)
+	{
+		FD3D12RayTracingGeometry* Geometry = FD3D12DynamicRHI::ResourceCast(P.Geometry.GetReference());
+
+		if (Geometry->IsDirty(GPUIndex))
+		{
+			if (ShouldCompactAfterBuild(Geometry->BuildFlags))
+			{
+				GetParentDevice()->GetRayTracingCompactionRequestHandler()->RequestCompact(Geometry);
+				Geometry->bHasPendingCompactionRequests[GPUIndex] = true;
+			}
+
+			Geometry->SetDirty(GetGPUMask(), false);
+		}
 	}
 
 	// Add a UAV barrier after each acceleration structure build batch.
diff --git a/Engine/Source/Runtime/D3D12RHI/Private/D3D12RayTracing.h b/Engine/Source/Runtime/D3D12RHI/Private/D3D12RayTracing.h
--- a/Engine/Source/Runtime/D3D12RHI/Private/D3D12RayTracing.h
+++ b/Engine/Source/Runtime/D3D12RHI/Private/D3D12RayTracing.h
@@ -28,8 +28,10 @@ public:
 
 	void SetupHitGroupSystemParameters(uint32 InGPUIndex);
 	void TransitionBuffers(FD3D12CommandContext& CommandContext);
-	void BuildAccelerationStructure(FD3D12CommandContext& CommandContext, EAccelerationStructureBuildMode BuildMode);
+	void UpdateResidency(FD3D12CommandContext& CommandContext);
 	void CompactAccelerationStructure(FD3D12CommandContext& CommandContext, uint32 InGPUIndex, uint64 InSizeAfterCompaction);
+	void CreateAccelerationStructureBuildDesc(FD3D12CommandContext& CommandContext, EAccelerationStructureBuildMode BuildMode, const TRefCountPtr<FD3D12Buffer>& ScratchBuffer,
+											D3D12_BUILD_RAYTRACING_ACCELERATION_STRUCTURE_DESC& OutDesc, TArrayView<D3D12_RAYTRACING_GEOMETRY_DESC>& OutGeometryDescs) const;
 	
 	// Implement FD3D12ShaderResourceRenameListener interface
 	virtual void ResourceRenamed(FD3D12BaseShaderResource* InRenamedResource, FD3D12ResourceLocation* InNewResourceLocation) override;
@@ -50,7 +52,6 @@ public:
 		return bIsAccelerationStructureDirty[GPUIndex];
 	}
 
-	uint32 IndexStride = 0; // 0 for non-indexed / implicit triangle list, 2 for uint16, 4 for uint32
 	uint32 IndexOffsetInBytes = 0;
 	uint32 TotalPrimitiveCount = 0; // Combined number of primitives in all mesh segments
 
