Added an option to create ray tracing PSOs grouped by expected ray generation shader performance/occupancy

Specialization is enbabled via new CVar r.D3D12.RayTracing.SpecializeStateObjects=1 (default: 0)

diff --git a/Engine/Source/Runtime/D3D12RHI/Private/D3D12Device.h b/Engine/Source/Runtime/D3D12RHI/Private/D3D12Device.h
--- a/Engine/Source/Runtime/D3D12RHI/Private/D3D12Device.h
+++ b/Engine/Source/Runtime/D3D12RHI/Private/D3D12Device.h
@@ -15,6 +15,7 @@ class FD3D12RayTracingPipelineCache;
 class FD3D12RayTracingCompactionRequestHandler;
 class FD3D12TimedIntervalQueryTracker;
 class FD3D12LinearQueryHeap;
+struct FD3D12RayTracingPipelineInfo;
 
 class FD3D12Device : public FD3D12SingleNodeGPUObject, public FNoncopyable, public FD3D12AdapterChild
 {
@@ -55,6 +56,10 @@ public:
 	FD3D12RayTracingPipelineCache*			GetRayTracingPipelineCache() { return RayTracingPipelineCache; }
 	FD3D12RayTracingCompactionRequestHandler* GetRayTracingCompactionRequestHandler() { return RayTracingCompactionRequestHandler; }
 	TRefCountPtr<ID3D12StateObject>			DeserializeRayTracingStateObject(D3D12_SHADER_BYTECODE Bytecode, ID3D12RootSignature* RootSignature);
+
+	/** Queries ray tracing pipeline state object metrics such as VGPR usage (if available/supported). Returns true if query succeeded. */
+	bool GetRayTracingPipelineInfo(ID3D12StateObject* Pipeline, FD3D12RayTracingPipelineInfo* OutInfo);
+
 #endif // D3D12_RHI_RAYTRACING
 
 	FD3D12DynamicRHI* GetOwningRHI();
diff --git a/Engine/Source/Runtime/D3D12RHI/Private/D3D12RHIPrivate.h b/Engine/Source/Runtime/D3D12RHI/Private/D3D12RHIPrivate.h
--- a/Engine/Source/Runtime/D3D12RHI/Private/D3D12RHIPrivate.h
+++ b/Engine/Source/Runtime/D3D12RHI/Private/D3D12RHIPrivate.h
@@ -206,6 +206,24 @@ struct FD3D12UpdateTexture3DData
 	bool bComputeShaderCopy;
 };
 
+/**
+* Structure that represents various RTPSO properties (0 if unknown).
+* These can be used to report performance characteristics, sort shaders by occupancy, etc.
+*/
+struct FD3D12RayTracingPipelineInfo
+{
+	static constexpr uint32 MaxPerformanceGroups = 10;
+
+	// Estimated RTPSO group based on occupancy or other platform-specific heuristics.
+	// Group 0 is expected to be performing worst, 9 (MaxPerformanceGroups-1) is expected to be the best.
+	uint32 PerformanceGroup = 0;
+
+	uint32 NumVGPR = 0;
+	uint32 NumSGPR = 0;
+	uint32 StackSize = 0;
+	uint32 ScratchSize = 0;
+};
+
 /** Forward declare the context for the AMD AGS utility library. */
 struct AGSContext;
 
diff --git a/Engine/Source/Runtime/D3D12RHI/Private/D3D12RayTracing.cpp b/Engine/Source/Runtime/D3D12RHI/Private/D3D12RayTracing.cpp
--- a/Engine/Source/Runtime/D3D12RHI/Private/D3D12RayTracing.cpp
+++ b/Engine/Source/Runtime/D3D12RHI/Private/D3D12RayTracing.cpp
@@ -63,6 +63,23 @@ static FAutoConsoleVariableRef CVarD3D12RayTracingMaxBatchedCompaction(
 	ECVF_ReadOnly
 );
 
+static int32 GRayTracingSpecializeStateObjects = 0;
+static FAutoConsoleVariableRef CVarRayTracingSpecializeStateObjects(
+	TEXT("r.D3D12.RayTracing.SpecializeStateObjects"),
+	GRayTracingSpecializeStateObjects,
+	TEXT("Whether to create specialized unique ray tracing pipeline state objects for each ray generation shader. (default = 0)\n")
+	TEXT("This option can produce more more efficient PSOs for the GPU at the cost of longer creation times and more memory. Requires DXR 1.1.\n"),
+	ECVF_ReadOnly
+);
+
+static int32 GRayTracingAllowSpecializedStateObjects = 0;
+static FAutoConsoleVariableRef CVarRayTracingAllowSpecializedStateObjects(
+	TEXT("r.D3D12.RayTracing.AllowSpecializedStateObjects"),
+	GRayTracingAllowSpecializedStateObjects,
+	TEXT("Whether to use specialized RTPSOs if they have been created. ")
+	TEXT("This is intended for performance testingand has no effect if r.D3D12.RayTracing.SpecializeStateObjects is 0. (default = 1)\n")
+);
+
 // Ray tracing stat counters
 
 DECLARE_STATS_GROUP(TEXT("D3D12RHI: Ray Tracing"), STATGROUP_D3D12RayTracing, STATCAT_Advanced);
@@ -792,6 +809,15 @@ public:
 		}
 	};
 
+	enum class ECollectionType
+	{
+		Unknown,
+		RayGen,
+		Miss,
+		HitGroup,
+		Callable,
+	};
+
 	struct FEntry
 	{
 		// Move-only type
@@ -819,9 +845,13 @@ public:
 			return *(ExportNames[0]);
 		}
 
+		ECollectionType CollectionType = ECollectionType::Unknown;
+
 		TRefCountPtr<FD3D12RayTracingShader> Shader;
 
 		TRefCountPtr<ID3D12StateObject> StateObject;
+		FD3D12RayTracingPipelineInfo PipelineInfo;
+
 		FGraphEventRef CompileEvent;
 		bool bDeserialized = false;
 
@@ -833,18 +863,12 @@ public:
 		float CompileTimeMS = 0.0f;
 	};
 
-	enum class ECollectionType
-	{
-		RayGen,
-		Miss,
-		HitGroup,
-		Callable,
-	};
-
 	static const TCHAR* GetCollectionTypeName(ECollectionType Type)
 	{
 		switch (Type)
 		{
+		case ECollectionType::Unknown:
+			return TEXT("Unknown");
 		case ECollectionType::RayGen:
 			return TEXT("RayGen");
 		case ECollectionType::Miss:
@@ -867,13 +891,14 @@ public:
 		FShaderCompileTask(
 				FEntry& InEntry,
 				FKey InCacheKey,
-				ID3D12Device5* InRayTracingDevice,
+				FD3D12Device* InDevice,
 				ECollectionType InCollectionType)
 			: Entry(InEntry)
 			, CacheKey(InCacheKey)
-			, RayTracingDevice(InRayTracingDevice)
+			, Device(InDevice)
+			, RayTracingDevice(InDevice->GetDevice5())
 			, CollectionType(InCollectionType)
-			{
+		{
 		}
 
 		static ESubsequentsMode::Type GetSubsequentsMode() { return ESubsequentsMode::TrackSubsequents; }
@@ -963,6 +988,11 @@ public:
 				{}, // ExistingCollections
 				D3D12_STATE_OBJECT_TYPE_COLLECTION);
 
+			if (Entry.StateObject)
+			{
+				Device->GetRayTracingPipelineInfo(Entry.StateObject, &Entry.PipelineInfo);
+			}
+
 			// Shader identifier can be queried immediately here per PSO collection, however this does not work on old NVIDIA drivers (430.00).
 			// Therefore shader identifiers need to be queried from the final linked pipeline (JIRA DH-2182).
 			// Entry.Identifier = GetShaderIdentifier(Entry.StateObject, Entry.GetPrimaryExportNameChars());
@@ -984,6 +1014,7 @@ public:
 
 		FEntry& Entry;
 		FKey CacheKey;
+		FD3D12Device* Device;
 		ID3D12Device5* RayTracingDevice;
 		ECollectionType CollectionType;
 
@@ -1037,12 +1068,17 @@ public:
 
 			FEntry& Entry = *FindResult;
 
+			Entry.CollectionType = CollectionType;
 			Entry.Shader = Shader;
 
 			if (Shader->bPrecompiledPSO)
 			{
 				D3D12_SHADER_BYTECODE Bytecode = Shader->ShaderBytecode.GetShaderBytecode();
 				Entry.StateObject = Device->DeserializeRayTracingStateObject(Bytecode, GlobalRootSignature);
+				if (Entry.StateObject)
+				{
+					Device->GetRayTracingPipelineInfo(Entry.StateObject, &Entry.PipelineInfo);
+				}
 
 				checkf(Entry.StateObject != nullptr, TEXT("Failed to deserialize RTPSO"));
 
@@ -1061,7 +1097,7 @@ public:
 				Entry.CompileEvent = TGraphTask<FShaderCompileTask>::CreateTask().ConstructAndDispatchWhenReady(
 					Entry,
 					CacheKey,
-					Device->GetDevice5(),
+					Device,
 					CollectionType
 				);
 			}
@@ -2046,6 +2082,138 @@ struct FD3D12RayTracingShaderLibrary
 	TArray<FD3D12ShaderIdentifier> Identifiers;
 };
 
+static void CreateSpecializedStateObjects(
+	ID3D12Device5* RayTracingDevice,
+	ID3D12RootSignature* GlobalRootSignature,
+	uint32 MaxPayloadSizeInBytes,
+	const FD3D12RayTracingShaderLibrary& RayGenShaders,
+	const TArray<FD3D12RayTracingPipelineCache::FEntry*>& UniqueShaderCollections,
+	const TMap<FSHAHash, int32>& RayGenShaderIndexByHash,
+	TArray<TRefCountPtr<ID3D12StateObject>>& OutSpecializedStateObjects,
+	TArray<int32>& OutSpecializationIndices)
+{
+	static constexpr uint32 MaxSpecializationBuckets = FD3D12RayTracingPipelineInfo::MaxPerformanceGroups;
+
+	if (RayGenShaders.Shaders.Num() <= 1)
+	{
+		// No specializations needed
+		return;
+	}
+
+	// Initialize raygen shader PSO specialization map to default values
+	OutSpecializationIndices.Reserve(RayGenShaders.Shaders.Num());
+	for (int32 It = 0; It < RayGenShaders.Shaders.Num(); ++It)
+	{
+		OutSpecializationIndices.Add(INDEX_NONE);
+	}
+
+	struct FRayGenShaderSpecialization
+	{
+		D3D12_EXISTING_COLLECTION_DESC Desc = {};
+		int32 ShaderIndex = INDEX_NONE;
+	};
+	TArray<FRayGenShaderSpecialization> RayGenShaderCollectionBuckets[MaxSpecializationBuckets];
+	TArray<D3D12_EXISTING_COLLECTION_DESC> ShaderCollectionDescs;
+
+	// Find useful performance group range for non-raygen shaders.
+	// It is not necessary to create PSO specializations for high-occupancy RGS if overall PSO will be limited by low-occupancy hit shaders.
+	// Also not necessary to create specializations if all raygen shaders are already in the same group.
+	uint32 MaxPerformanceGroupRGS = 0;
+	uint32 MinPerformanceGroupRGS = MaxSpecializationBuckets - 1;
+	uint32 MaxPerformanceGroupOther = 0;
+	uint32 MinPerformanceGroupOther = MaxSpecializationBuckets - 1;
+	int32 LastRayGenShaderCollectionIndex = INDEX_NONE;
+
+	for (int32 EntryIndex = 0; EntryIndex < UniqueShaderCollections.Num(); ++EntryIndex)
+	{
+		FD3D12RayTracingPipelineCache::FEntry* Entry = UniqueShaderCollections[EntryIndex];
+
+		const uint32 Group = FMath::Min<uint32>(Entry->PipelineInfo.PerformanceGroup, MaxSpecializationBuckets);
+
+		if (Entry->CollectionType == FD3D12RayTracingPipelineCache::ECollectionType::RayGen)
+		{
+			MaxPerformanceGroupRGS = FMath::Max<uint32>(MaxPerformanceGroupRGS, Group);
+			MinPerformanceGroupRGS = FMath::Min<uint32>(MinPerformanceGroupRGS, Group);
+			LastRayGenShaderCollectionIndex = EntryIndex;
+		}
+		else
+		{
+			checkf(EntryIndex > LastRayGenShaderCollectionIndex, TEXT("Ray generation shaders are expected to be first in the UniqueShaderCollections list."));
+
+			MaxPerformanceGroupOther = FMath::Max<uint32>(MaxPerformanceGroupOther, Group);
+			MinPerformanceGroupOther = FMath::Min<uint32>(MinPerformanceGroupOther, Group);
+
+			// This is a hit/miss/callable shader which will be common for all specialized RTPSOs.
+			ShaderCollectionDescs.Add(Entry->GetCollectionDesc());
+		}
+	}
+
+	if (MinPerformanceGroupRGS == MaxPerformanceGroupRGS)
+	{
+		// No need to create a specialized PSO if all raygen shaders are already in the same group
+		return;
+	}
+
+	// Split RGS collections into a separate lists, organized by performance group
+	for (int32 EntryIndex = 0; EntryIndex <= LastRayGenShaderCollectionIndex; ++EntryIndex)
+	{
+		FD3D12RayTracingPipelineCache::FEntry* Entry = UniqueShaderCollections[EntryIndex];
+
+		check(Entry->CollectionType == FD3D12RayTracingPipelineCache::ECollectionType::RayGen);
+
+		// Don't create specializations for raygen shaders that have better occupancy than worst non-raygen shader
+		const uint32 SpecializationBucket = FMath::Min<uint32>(Entry->PipelineInfo.PerformanceGroup, MinPerformanceGroupOther);
+
+		// Don't create extra specialized pipelines for group 0 (worst-performing) and just use the default RTPSO.
+		if (SpecializationBucket > 0)
+		{
+			FRayGenShaderSpecialization Specialization;
+			Specialization.Desc = Entry->GetCollectionDesc();
+			Specialization.ShaderIndex = RayGenShaderIndexByHash.FindChecked(Entry->Shader->GetHash());
+			RayGenShaderCollectionBuckets[SpecializationBucket].Add(Specialization);
+		}
+	}
+
+	OutSpecializedStateObjects.Reserve(MaxSpecializationBuckets);
+
+	const uint32 ShaderCollectionDescsSize = ShaderCollectionDescs.Num();
+
+	for (const TArray<FRayGenShaderSpecialization>& SpecializationBucket : RayGenShaderCollectionBuckets)
+	{
+		if (SpecializationBucket.IsEmpty())
+		{
+			continue;
+		}
+
+		const int32 SpecializationIndex = OutSpecializedStateObjects.Num();
+
+		for (const FRayGenShaderSpecialization& Specialization : SpecializationBucket)
+		{
+			// Temporarily add the RGSs to complete shader collection
+			ShaderCollectionDescs.Add(Specialization.Desc);
+			OutSpecializationIndices[Specialization.ShaderIndex] = SpecializationIndex;
+		}
+
+		TRefCountPtr<ID3D12StateObject> SpecializedPSO = CreateRayTracingStateObject(
+			RayTracingDevice,
+			{}, // Libraries,
+			{}, // LibraryExports,
+			MaxPayloadSizeInBytes,
+			{}, // HitGroups
+			GlobalRootSignature,
+			{}, // LocalRootSignatures
+			{}, // LocalRootSignatureAssociations,
+			ShaderCollectionDescs,
+			D3D12_STATE_OBJECT_TYPE_RAYTRACING_PIPELINE);
+
+		OutSpecializedStateObjects.Add(SpecializedPSO);
+
+		// Remove the temporary RGSs
+		ShaderCollectionDescs.SetNum(ShaderCollectionDescsSize);
+	}
+}
+
+
 class FD3D12RayTracingPipelineState : public FRHIRayTracingPipelineState
 {
 public:
@@ -2154,6 +2322,9 @@ public:
 
 		RayGenShaders.Reserve(InitializerRayGenShaders.Num());
 		RayGenShaderEntries.Reserve(InitializerRayGenShaders.Num());
+		TMap<FSHAHash, int32> RayGenShaderIndexByHash;
+
+		checkf(UniqueShaderCollections.Num() == 0, TEXT("Ray generation shaders are expected to be first in the UniqueShaderCollections list."));
 
 		for (FRHIRayTracingShader* ShaderRHI : InitializerRayGenShaders)
 		{
@@ -2164,6 +2335,7 @@ public:
 			FD3D12RayTracingPipelineCache::FEntry* ShaderCacheEntry = AddShaderCollection(Shader, FD3D12RayTracingPipelineCache::ECollectionType::RayGen);
 
 			RayGenShaderEntries.Add(ShaderCacheEntry);
+			RayGenShaderIndexByHash.Add(Shader->GetHash(), RayGenShaders.Shaders.Num());
 			RayGenShaders.Shaders.Add(Shader);
 		}
 
@@ -2274,7 +2446,8 @@ public:
 
 		LinkTime -= FPlatformTime::Cycles64();
 
-		if (BasePipeline)
+		// Extending RTPSOs is currently not compatible with PSO specializations
+		if (BasePipeline && GRayTracingSpecializeStateObjects == 0)
 		{
 			if (UniqueShaderCollectionDescs.Num() == 0)
 			{
@@ -2325,6 +2498,20 @@ public:
 				D3D12_STATE_OBJECT_TYPE_RAYTRACING_PIPELINE);
 		}
 
+		if (GRayTracingSpecializeStateObjects != 0 && Initializer.GetRayGenTable().Num() > 1)
+		{
+			CreateSpecializedStateObjects(
+				RayTracingDevice,
+				GlobalRootSignature,
+				Initializer.MaxPayloadSizeInBytes,
+				RayGenShaders,
+				UniqueShaderCollections,
+				RayGenShaderIndexByHash,
+				SpecializedStateObjects, // out param
+				SpecializationIndices // out param
+			);
+		}
+
 		LinkTime += FPlatformTime::Cycles64();
 
 		HRESULT QueryInterfaceResult = StateObject->QueryInterface(IID_PPV_ARGS(PipelineProperties.GetInitReference()));
@@ -2451,6 +2638,12 @@ public:
 	TRefCountPtr<ID3D12StateObject> StateObject;
 	TRefCountPtr<ID3D12StateObjectProperties> PipelineProperties;
 
+	// Maps raygen shader index to a specialized state object (may be -1 if no specialization is used for a shader)
+	TArray<int32> SpecializationIndices;
+
+	// State objects with raygen shaders grouped by occupancy
+	TArray<TRefCountPtr<ID3D12StateObject>> SpecializedStateObjects;
+
 	static constexpr uint32 ShaderIdentifierSize = D3D12_SHADER_IDENTIFIER_SIZE_IN_BYTES;
 	bool bAllowHitGroupIndexing = true;
 
@@ -4758,7 +4951,25 @@ static void DispatchRays(FD3D12CommandContext& CommandContext,
 
 		CommandContext.CommandListHandle.FlushResourceBarriers();
 
-		ID3D12StateObject* RayTracingStateObject = Pipeline->StateObject.GetReference();
+		ID3D12StateObject* RayTracingStateObject = nullptr;
+
+		// Select a specialized RTPSO, if one is available
+		if (GRayTracingAllowSpecializedStateObjects
+			&& !Pipeline->SpecializedStateObjects.IsEmpty() 
+			&& !Pipeline->SpecializationIndices.IsEmpty())
+		{
+			int32 SpecializationIndex = Pipeline->SpecializationIndices[RayGenShaderIndex];
+			if (SpecializationIndex != INDEX_NONE)
+			{
+				RayTracingStateObject = Pipeline->SpecializedStateObjects[SpecializationIndex];
+			}
+		}
+
+		// Fall back to default full RTPSO if specialization is not available
+		if (!RayTracingStateObject)
+		{
+			RayTracingStateObject = Pipeline->StateObject.GetReference();
+		}
 
 		ID3D12GraphicsCommandList4* RayTracingCommandList = CommandContext.CommandListHandle.RayTracingCommandList();
 		RayTracingCommandList->SetPipelineState1(RayTracingStateObject);
diff --git a/Engine/Source/Runtime/D3D12RHI/Private/Windows/WindowsD3D12Device.cpp b/Engine/Source/Runtime/D3D12RHI/Private/Windows/WindowsD3D12Device.cpp
--- a/Engine/Source/Runtime/D3D12RHI/Private/Windows/WindowsD3D12Device.cpp
+++ b/Engine/Source/Runtime/D3D12RHI/Private/Windows/WindowsD3D12Device.cpp
@@ -1249,6 +1249,16 @@ TRefCountPtr<ID3D12StateObject> FD3D12Device::DeserializeRayTracingStateObject(D
 	TRefCountPtr<ID3D12StateObject> Result;
 	return Result;
 }
+
+bool FD3D12Device::GetRayTracingPipelineInfo(ID3D12StateObject* Pipeline, FD3D12RayTracingPipelineInfo* OutInfo)
+{
+	// Return a safe default result on Windows, as there is no API to query interesting pipeline metrics.
+	FD3D12RayTracingPipelineInfo Result = {};
+	*OutInfo = Result;
+
+	return false;
+}
+
 #endif // D3D12_RHI_RAYTRACING
 
 /**
