Toggled default value for r.D3D12.RayTracing.AllowSpecializedStateObjects from 0 to 1. This change currently has no practical effect, since dependent setting is still off by default.

diff --git a/Engine/Source/Runtime/D3D12RHI/Private/D3D12RayTracing.cpp b/Engine/Source/Runtime/D3D12RHI/Private/D3D12RayTracing.cpp
--- a/Engine/Source/Runtime/D3D12RHI/Private/D3D12RayTracing.cpp
+++ b/Engine/Source/Runtime/D3D12RHI/Private/D3D12RayTracing.cpp
@@ -72,7 +72,7 @@ static FAutoConsoleVariableRef CVarRayTracingSpecializeStateObjects(
 	ECVF_ReadOnly
 );
 
-static int32 GRayTracingAllowSpecializedStateObjects = 0;
+static int32 GRayTracingAllowSpecializedStateObjects = 1;
 static FAutoConsoleVariableRef CVarRayTracingAllowSpecializedStateObjects(
 	TEXT("r.D3D12.RayTracing.AllowSpecializedStateObjects"),
 	GRayTracingAllowSpecializedStateObjects,
