Fix for oversized memcpy when creating raytracing geometry from offline data.
On behalf of Dave.Barrett.

diff --git a/Engine/Source/Runtime/D3D12RHI/Private/D3D12RayTracing.cpp b/Engine/Source/Runtime/D3D12RHI/Private/D3D12RayTracing.cpp
--- a/Engine/Source/Runtime/D3D12RHI/Private/D3D12RayTracing.cpp
+++ b/Engine/Source/Runtime/D3D12RHI/Private/D3D12RayTracing.cpp
@@ -3214,12 +3214,12 @@ FD3D12RayTracingGeometry::FD3D12RayTracingGeometry(FD3D12Adapter* Adapter, const
 	{
 		FD3D12Device* Device = Adapter->GetDevice(0);
 
-		const uint8* Data = (const uint8*)Initializer.OfflineData->GetResourceData();
-		uint32 Size = Initializer.OfflineData->GetResourceDataSize();
+		const uint8* Data = ((const uint8*)Initializer.OfflineData->GetResourceData()) + sizeof(FOfflineBVHHeader);
+		uint32 Size = Initializer.OfflineData->GetResourceDataSize() - sizeof(FOfflineBVHHeader);
 
 		FD3D12ResourceLocation SrcResourceLoc(Device);
 		uint8* DstDataBase = (uint8*)Device->GetDefaultFastAllocator().Allocate(Size, 256, &SrcResourceLoc);
-		FMemory::Memcpy(DstDataBase, Data + sizeof(FOfflineBVHHeader), Size);
+		FMemory::Memcpy(DstDataBase, Data, Size);
 
 		FRHICommandListImmediate& RHICmdList = FRHICommandListExecutor::GetImmediateCommandList();
 
