UE-125001 //UE5/Main - Compile UnrealGame HoloLens - Several compile errors
-Fixed d3d12 hololens build breaks.
-Fixed HololensD3D11Viewport build break.

diff --git a/Engine/Source/Runtime/D3D12RHI/Private/D3D12Adapter.cpp b/Engine/Source/Runtime/D3D12RHI/Private/D3D12Adapter.cpp
--- a/Engine/Source/Runtime/D3D12RHI/Private/D3D12Adapter.cpp
+++ b/Engine/Source/Runtime/D3D12RHI/Private/D3D12Adapter.cpp
@@ -440,7 +440,7 @@ void FD3D12Adapter::CreateRootDevice(bool bWithDebug)
 
 	// QI for the Adapter
 	TRefCountPtr<IDXGIAdapter> TempAdapter;
-#if PLATFORM_WINDOWS
+#if PLATFORM_WINDOWS || PLATFORM_HOLOLENS
 	Desc.EnumAdapters(DxgiFactory, DxgiFactory6, TempAdapter.GetInitReference());
 #else
 	DxgiFactory->EnumAdapters(Desc.AdapterIndex, TempAdapter.GetInitReference());
diff --git a/Engine/Source/Runtime/D3D12RHI/Private/D3D12Adapter.h b/Engine/Source/Runtime/D3D12RHI/Private/D3D12Adapter.h
--- a/Engine/Source/Runtime/D3D12RHI/Private/D3D12Adapter.h
+++ b/Engine/Source/Runtime/D3D12RHI/Private/D3D12Adapter.h
@@ -48,7 +48,7 @@ struct FD3D12AdapterDesc
 	FD3D12AdapterDesc() = default;
 
 	FD3D12AdapterDesc(const DXGI_ADAPTER_DESC& DescIn , int32 InAdapterIndex, D3D_FEATURE_LEVEL InMaxSupportedFeatureLevel, D3D_SHADER_MODEL InMaxSupportedShaderModel, uint32 NumNodes, bool InIsIntegrated
-#if PLATFORM_WINDOWS
+#if PLATFORM_WINDOWS || PLATFORM_HOLOLENS
 		, DXGI_GPU_PREFERENCE InGpuPreference
 #endif
 	)
@@ -58,7 +58,7 @@ struct FD3D12AdapterDesc
 		, Desc(DescIn)
 		, NumDeviceNodes(NumNodes)
 		, bIsIntegrated(InIsIntegrated)
-#if PLATFORM_WINDOWS
+#if PLATFORM_WINDOWS || PLATFORM_HOLOLENS
 		, GpuPreference(InGpuPreference)
 #endif
 	{
@@ -66,7 +66,7 @@ struct FD3D12AdapterDesc
 
 	bool IsValid() const { return MaxSupportedFeatureLevel != (D3D_FEATURE_LEVEL)0 && AdapterIndex >= 0; }
 
-#if PLATFORM_WINDOWS
+#if PLATFORM_WINDOWS || PLATFORM_HOLOLENS
 	static HRESULT EnumAdapters(int32 AdapterIndex, DXGI_GPU_PREFERENCE GpuPreference, IDXGIFactory* DxgiFactory, IDXGIFactory6* DxgiFactory6, IDXGIAdapter** TempAdapter)
 	{
 		if (!DxgiFactory6 || GpuPreference == DXGI_GPU_PREFERENCE_UNSPECIFIED)
@@ -100,7 +100,7 @@ struct FD3D12AdapterDesc
 	/** Whether the GPU is integrated or discrete. */
 	bool bIsIntegrated = false;
 
-#if PLATFORM_WINDOWS
+#if PLATFORM_WINDOWS || PLATFORM_HOLOLENS
 	DXGI_GPU_PREFERENCE GpuPreference = DXGI_GPU_PREFERENCE_UNSPECIFIED;
 #endif
 };
@@ -251,7 +251,7 @@ public:
 	void CreateDXGIFactory(bool bWithDebug);
 	FORCEINLINE IDXGIFactory* GetDXGIFactory() const { return DxgiFactory; }
 	FORCEINLINE IDXGIFactory2* GetDXGIFactory2() const { return DxgiFactory2; }
-#if PLATFORM_WINDOWS
+#if PLATFORM_WINDOWS || PLATFORM_HOLOLENS
 	FORCEINLINE IDXGIFactory6* GetDXGIFactory6() const { return DxgiFactory6; }
 #endif
 
@@ -518,7 +518,7 @@ protected:
 
 	TRefCountPtr<IDXGIFactory> DxgiFactory;
 	TRefCountPtr<IDXGIFactory2> DxgiFactory2;
-#if PLATFORM_WINDOWS
+#if PLATFORM_WINDOWS || PLATFORM_HOLOLENS
 	TRefCountPtr<IDXGIFactory6> DxgiFactory6;
 #endif
 
diff --git a/Engine/Source/Runtime/D3D12RHI/Private/D3D12RHI.cpp b/Engine/Source/Runtime/D3D12RHI/Private/D3D12RHI.cpp
--- a/Engine/Source/Runtime/D3D12RHI/Private/D3D12RHI.cpp
+++ b/Engine/Source/Runtime/D3D12RHI/Private/D3D12RHI.cpp
@@ -401,7 +401,7 @@ void FD3D12DynamicRHI::RHIGetSupportedResolution(uint32& Width, uint32& Height)
 	{
 		HRESULT HResult = S_OK;
 		TRefCountPtr<IDXGIAdapter> Adapter;
-#if PLATFORM_WINDOWS
+#if PLATFORM_WINDOWS || PLATFORM_HOLOLENS
 		HResult = GetAdapter().GetDesc().EnumAdapters(GetAdapter().GetDXGIFactory(), GetAdapter().GetDXGIFactory6(), Adapter.GetInitReference());
 #else
 		HResult = GetAdapter().GetDXGIFactory()->EnumAdapters(GetAdapter().GetAdapterIndex(), Adapter.GetInitReference());
diff --git a/Engine/Source/Runtime/Windows/D3D11RHI/Private/HoloLens/HoloLensD3D11Viewport.cpp b/Engine/Source/Runtime/Windows/D3D11RHI/Private/HoloLens/HoloLensD3D11Viewport.cpp
--- a/Engine/Source/Runtime/Windows/D3D11RHI/Private/HoloLens/HoloLensD3D11Viewport.cpp
+++ b/Engine/Source/Runtime/Windows/D3D11RHI/Private/HoloLens/HoloLensD3D11Viewport.cpp
@@ -187,7 +187,12 @@ FD3D11Viewport::FD3D11Viewport(FD3D11DynamicRHI* InD3DRHI,HWND InWindowHandle,ui
 	// Create a RHI surface to represent the viewport's back buffer.
 	BackBuffer = GetSwapChainSurface(D3DRHI, PixelFormat, SizeX, SizeY, SwapChain);
 
-	BeginInitResource(&FrameSyncEvent);
+	ENQUEUE_RENDER_COMMAND(FD3D11Viewport)(
+	[this](FRHICommandListImmediate& RHICmdList)
+	{
+		// Initialize the query by issuing an initial event.
+		FrameSyncEvent.IssueEvent();
+	});
 }
 
 
