UE5:  D3D12 command buffer memory leak while idling in Sandbox.

Need to call "ReleaseResourceBarrierCommandListAllocator" on all command list managers each frame.  We were leaking around 10 MB tracked and 75 MB untracked memory per minute without this fix.

#rnx

diff --git a/Engine/Source/Runtime/D3D12RHI/Private/D3D12CommandContext.cpp b/Engine/Source/Runtime/D3D12RHI/Private/D3D12CommandContext.cpp
--- a/Engine/Source/Runtime/D3D12RHI/Private/D3D12CommandContext.cpp
+++ b/Engine/Source/Runtime/D3D12RHI/Private/D3D12CommandContext.cpp
@@ -650,6 +650,8 @@ void FD3D12CommandContextBase::RHIEndFrame()
 	{
 		Device = ParentAdapter->GetDevice(GPUIndex);
 		Device->GetCommandListManager().ReleaseResourceBarrierCommandListAllocator();
+		Device->GetAsyncCommandListManager().ReleaseResourceBarrierCommandListAllocator();
+		Device->GetCopyCommandListManager().ReleaseResourceBarrierCommandListAllocator();
 	}
 
 	UpdateMemoryStats();
