Transition vertex buffers right before calling IASetVertexBuffers, instead doing it when the high level calls RHISetStreamSource.

If a command list flush happens in the middle of a sequence of draws which use the same vertex buffer, the tracked internal state of the buffer is reset. Since the high level doesn't call RHISetStreamSource again after the flush, no pending barrier was inserted. If a subsequent dispatch needed the buffer in a different state, that transition was erroneously inserted in the patch-up command list, which modified the state too soon (before the draws which still need the buffer in the VB state). Adding the transition in FD3D12DescriptorCache::SetVertexBuffers ensures that the correct state is applied, since this method is called internally before each draw.

#rnx

diff --git a/Engine/Source/Runtime/D3D12RHI/Private/D3D12DescriptorCache.cpp b/Engine/Source/Runtime/D3D12RHI/Private/D3D12DescriptorCache.cpp
--- a/Engine/Source/Runtime/D3D12RHI/Private/D3D12DescriptorCache.cpp
+++ b/Engine/Source/Runtime/D3D12RHI/Private/D3D12DescriptorCache.cpp
@@ -295,6 +295,19 @@ void FD3D12DescriptorCache::SetVertexBuffers(FD3D12VertexBufferCache& Cache)
 
 	CmdContext->CommandListHandle.UpdateResidency(Cache.ResidencyHandles, Count);
 	CmdContext->CommandListHandle->IASetVertexBuffers(0, Count, Cache.CurrentVertexBufferViews);
+
+	for (uint32 i = 0; i < Count; ++i)
+	{
+		if (Cache.CurrentVertexBufferResources[i])
+		{
+			FD3D12Resource* Resource = Cache.CurrentVertexBufferResources[i]->GetResource();
+			if (Resource->RequiresResourceStateTracking())
+			{
+				check(Resource->GetSubresourceCount() == 1);
+				FD3D12DynamicRHI::TransitionResource(CmdContext->CommandListHandle, Resource, D3D12_RESOURCE_STATE_TBD, D3D12_RESOURCE_STATE_VERTEX_AND_CONSTANT_BUFFER, D3D12_RESOURCE_BARRIER_ALL_SUBRESOURCES, FD3D12DynamicRHI::ETransitionMode::Validate);
+			}
+		}
+	}
 }
 
 template <EShaderFrequency ShaderStage>
diff --git a/Engine/Source/Runtime/D3D12RHI/Private/D3D12StateCache.cpp b/Engine/Source/Runtime/D3D12RHI/Private/D3D12StateCache.cpp
--- a/Engine/Source/Runtime/D3D12RHI/Private/D3D12StateCache.cpp
+++ b/Engine/Source/Runtime/D3D12RHI/Private/D3D12StateCache.cpp
@@ -1198,17 +1198,6 @@ void FD3D12StateCacheBase::InternalSetStreamSource(FD3D12ResourceLocation* Verte
 			PipelineState.Graphics.VBCache.MaxBoundVertexBufferIndex = INDEX_NONE;
 		}
 	}
-
-	if (VertexBufferLocation)
-	{
-		FD3D12Resource* const pResource = VertexBufferLocation->GetResource();
-		if (pResource && pResource->RequiresResourceStateTracking())
-		{
-			check(pResource->GetSubresourceCount() == 1);
-			FD3D12DynamicRHI::TransitionResource(CmdContext->CommandListHandle, pResource, D3D12_RESOURCE_STATE_TBD, D3D12_RESOURCE_STATE_VERTEX_AND_CONSTANT_BUFFER, D3D12_RESOURCE_BARRIER_ALL_SUBRESOURCES, FD3D12DynamicRHI::ETransitionMode::Validate);
-		}
-	}
-
 }
 
 template <EShaderFrequency ShaderFrequency>
