D3D shader cleanup:
- Removed some unnecessary code from D3D shader compilers.
- Removing shader VendorExtensions runtime data from platforms that don't need them.
- Turning FShaderCodeFeatures into a set of flags to reduce its size. Adding a few new flags for future use.
- Moved the common D3D12 shader data and creation code into shared locations to cut down on code duplication.
- Removed the unnecessary copying of header data from all D3D12 shader types.
- Removed FD3D12ShaderBytecode since its values can be inferred from other shader data.

diff --git a/Engine/Source/Developer/Windows/ShaderFormatD3D/Private/D3DShaderCompiler.cpp b/Engine/Source/Developer/Windows/ShaderFormatD3D/Private/D3DShaderCompiler.cpp
--- a/Engine/Source/Developer/Windows/ShaderFormatD3D/Private/D3DShaderCompiler.cpp
+++ b/Engine/Source/Developer/Windows/ShaderFormatD3D/Private/D3DShaderCompiler.cpp
@@ -792,8 +792,6 @@ bool CompileAndProcessD3DShaderFXC(FString& PreprocessedShaderSource, const FStr
 	}
 
 	// Gather reflection information
-	int32 NumInterpolants = 0;
-	TIndirectArray<FString> InterpolantNames;
 	TArray<FString> ShaderInputs;
 	TArray<FShaderCodeVendorExtension> VendorExtensions;
 
@@ -808,7 +806,6 @@ bool CompileAndProcessD3DShaderFXC(FString& PreprocessedShaderSource, const FStr
 		uint32 NumUAVs = 0;
 
 		TArray<FString> UniformBufferNames;
-		TArray<FString> ShaderOutputs;
 
 		TBitArray<> UsedUniformBufferSlots;
 		UsedUniformBufferSlots.Init(false, 32);
@@ -831,38 +828,13 @@ bool CompileAndProcessD3DShaderFXC(FString& PreprocessedShaderSource, const FStr
 			D3D11_SHADER_DESC ShaderDesc;
 			Reflector->GetDesc(&ShaderDesc);
 
-			if (Input.Target.Frequency == SF_Vertex)
-			{
-				for (uint32 Index = 0; Index < ShaderDesc.OutputParameters; ++Index)
-				{
-					// VC++ horrible hack: Runtime ESP checks get confused and fail for some reason calling Reflector->GetOutputParameterDesc() (because it comes from another DLL?)
-					// so "guard it" using the middle of an array; it's been confirmed NO corruption is really happening.
-					D3D11_SIGNATURE_PARAMETER_DESC ParamDescs[3];
-					D3D11_SIGNATURE_PARAMETER_DESC& ParamDesc = ParamDescs[1];
-					Reflector->GetOutputParameterDesc(Index, &ParamDesc);
-					if (ParamDesc.SystemValueType == D3D_NAME_UNDEFINED && ParamDesc.Mask != 0)
-					{
-						++NumInterpolants;
-						InterpolantNames.Add(new FString(FString::Printf(TEXT("%s%d"), ANSI_TO_TCHAR(ParamDesc.SemanticName), ParamDesc.SemanticIndex)));
-						ShaderOutputs.Add(*InterpolantNames.Last());
-					}
-				}
-			}
-			else if (Input.Target.Frequency == SF_Pixel)
+			if (Input.Target.Frequency == SF_Pixel)
 			{
 				if (GD3DAllowRemoveUnused != 0 && Input.bCompilingForShaderPipeline)
 				{
 					// Handy place for a breakpoint for debugging...
 					++GBreakpoint;
 				}
-				for(uint32 Index = 0; Index < ShaderDesc.OutputParameters; ++Index)
-				{
-					// VC++ horrible hack: Runtime ESP checks get confused and fail for some reason calling Reflector->GetInputParameterDesc() (because it comes from another DLL?)
-					// so "guard it" using the middle of an array; it's been confirmed NO corruption is really happening.
-					D3D11_SIGNATURE_PARAMETER_DESC ParamDescs[3];
-					D3D11_SIGNATURE_PARAMETER_DESC& ParamDesc = ParamDescs[1];
-					Reflector->GetOutputParameterDesc(Index, &ParamDesc);
-				}
 
 				bool bFoundUnused = false;
 				for (uint32 Index = 0; Index < ShaderDesc.InputParameters; ++Index)
diff --git a/Engine/Source/Developer/Windows/ShaderFormatD3D/Private/D3DShaderCompiler.inl b/Engine/Source/Developer/Windows/ShaderFormatD3D/Private/D3DShaderCompiler.inl
--- a/Engine/Source/Developer/Windows/ShaderFormatD3D/Private/D3DShaderCompiler.inl
+++ b/Engine/Source/Developer/Windows/ShaderFormatD3D/Private/D3DShaderCompiler.inl
@@ -416,13 +416,7 @@ template <typename ID3D1xShaderReflection, typename D3D1x_SHADER_DESC, typename
 
 			if (bIsVendorParameter)
 			{
-				FShaderCodeVendorExtension VendorExtension;
-				VendorExtension.VendorId = 0x1002; // AMD
-				VendorExtension.Parameter.BufferIndex = 0;
-				VendorExtension.Parameter.BaseIndex = BindDesc.BindPoint;
-				VendorExtension.Parameter.Size = BindCount;
-				VendorExtension.Parameter.Type = ParameterType;
-				VendorExtensions.Add(VendorExtension);
+				VendorExtensions.Emplace(0x1002 /*AMD*/, 0, BindDesc.BindPoint, BindCount, ParameterType);
 			}
 			else
 			{
@@ -465,24 +459,12 @@ template <typename ID3D1xShaderReflection, typename D3D1x_SHADER_DESC, typename
 			const uint32 BindCount = 1;
 			if (bIsVendorParameter)
 			{
-				FShaderCodeVendorExtension VendorExtension;
-				if (bIsNVExtension)
-				{
-					VendorExtension.VendorId = 0x10DE; // NVIDIA
-				}
-				else if (bIsAMDExtensionDX11 || bIsAMDExtensionDX12)
-				{
-					VendorExtension.VendorId = 0x1002; // AMD
-				}
-				else if (bIsIntelExtension)
-				{
-					VendorExtension.VendorId = 0x8086; // INTEL
-				}
-				VendorExtension.Parameter.BufferIndex = 0;
-				VendorExtension.Parameter.BaseIndex = BindDesc.BindPoint;
-				VendorExtension.Parameter.Size = BindCount;
-				VendorExtension.Parameter.Type = EShaderParameterType::UAV;
-				VendorExtensions.Add(VendorExtension);
+				const uint32 VendorId =
+					bIsNVExtension ? 0x10DE : // NVIDIA
+					(bIsAMDExtensionDX11 || bIsAMDExtensionDX12) ? 0x1002 : // AMD
+					bIsIntelExtension ? 0x8086 : // Intel
+					0;
+				VendorExtensions.Emplace(VendorId, 0, BindDesc.BindPoint, BindCount, EShaderParameterType::UAV);
 			}
 			else if (bIsDiagnosticBufferParameter)
 			{
diff --git a/Engine/Source/Developer/Windows/ShaderFormatD3D/Private/D3DShaderCompilerDXC.cpp b/Engine/Source/Developer/Windows/ShaderFormatD3D/Private/D3DShaderCompilerDXC.cpp
--- a/Engine/Source/Developer/Windows/ShaderFormatD3D/Private/D3DShaderCompilerDXC.cpp
+++ b/Engine/Source/Developer/Windows/ShaderFormatD3D/Private/D3DShaderCompilerDXC.cpp
@@ -721,8 +721,6 @@ bool CompileAndProcessD3DShaderDXC(FString& PreprocessedShaderSource,
 	if (SUCCEEDED(D3DCompileToDxilResult))
 	{
 		// Gather reflection information
-		int32 NumInterpolants = 0;
-		TIndirectArray<FString> InterpolantNames;
 		TArray<FString> ShaderInputs;
 		TArray<FShaderCodeVendorExtension> VendorExtensions;
 
@@ -739,6 +737,8 @@ bool CompileAndProcessD3DShaderDXC(FString& PreprocessedShaderSource,
 		TBitArray<> UsedUniformBufferSlots;
 		UsedUniformBufferSlots.Init(false, 32);
 
+		uint64 ShaderRequiresFlags{};
+
 		dxc::DxcDllSupport& DxcDllHelper = GetDxcDllHelper();
 		TRefCountPtr<IDxcUtils> Utils;
 		VERIFYHRESULT(DxcDllHelper.CreateInstance(CLSID_DxcUtils, Utils.GetInitReference()));
@@ -787,6 +787,8 @@ bool CompileAndProcessD3DShaderDXC(FString& PreprocessedShaderSource,
 				FunctionReflection = LibraryReflection->GetFunctionByIndex(FunctionIndex);
 				FunctionReflection->GetDesc(&FunctionDesc);
 
+				ShaderRequiresFlags |= FunctionDesc.RequiredFeatureFlags;
+
 				for (const FString& MangledEntryPoint : MangledEntryPoints)
 				{
 					// Entry point parameters are currently not included in the partial mangling, therefore partial substring match is used here.
@@ -853,7 +855,6 @@ bool CompileAndProcessD3DShaderDXC(FString& PreprocessedShaderSource,
 		}
 		else
 		{
-
 			TRefCountPtr<ID3D12ShaderReflection> ShaderReflection;
 			const HRESULT CreateReflectionResult = Utils->CreateReflection(&ReflBuffer, IID_PPV_ARGS(ShaderReflection.GetInitReference()));
 			if (FAILED(CreateReflectionResult))
@@ -864,16 +865,17 @@ bool CompileAndProcessD3DShaderDXC(FString& PreprocessedShaderSource,
 			D3D12_SHADER_DESC ShaderDesc = {};
 			ShaderReflection->GetDesc(&ShaderDesc);
 
+			ShaderRequiresFlags = ShaderReflection->GetRequiresFlags();
+
 			ExtractParameterMapFromD3DShader<ID3D12ShaderReflection, D3D12_SHADER_DESC, D3D12_SHADER_INPUT_BIND_DESC,
 				ID3D12ShaderReflectionConstantBuffer, D3D12_SHADER_BUFFER_DESC,
 				ID3D12ShaderReflectionVariable, D3D12_SHADER_VARIABLE_DESC>(
 					Input, ShaderParameterParser,
 					AutoBindingSpace, ShaderReflection, ShaderDesc,
-					bGlobalUniformBufferUsed, bDiagnosticBufferUsed, 
+					bGlobalUniformBufferUsed, bDiagnosticBufferUsed,
 					NumSamplers, NumSRVs, NumCBs, NumUAVs,
 					Output, UniformBufferNames, UsedUniformBufferSlots, VendorExtensions);
 
-
 			Output.bSucceeded = true;
 		}
 
@@ -898,14 +900,50 @@ bool CompileAndProcessD3DShaderDXC(FString& PreprocessedShaderSource,
 			auto AddOptionalDataCallback = [&](FShaderCode& ShaderCode)
 			{
 				FShaderCodeFeatures CodeFeatures;
-				//#todo-rco: Really should look inside DXIL
-				CodeFeatures.bUsesWaveOps = Input.Environment.CompilerFlags.Contains(CFLAG_WaveOperations);
-				CodeFeatures.bUsesDiagnosticBuffer = bDiagnosticBufferUsed;
+
+				if ((ShaderRequiresFlags & D3D_SHADER_REQUIRES_WAVE_OPS) != 0)
+				{
+					EnumAddFlags(CodeFeatures.CodeFeatures, EShaderCodeFeatures::WaveOps);
+				}
+
+				if ((ShaderRequiresFlags & D3D_SHADER_REQUIRES_NATIVE_16BIT_OPS) != 0)
+				{
+					EnumAddFlags(CodeFeatures.CodeFeatures, EShaderCodeFeatures::SixteenBitTypes);
+				}
+
+				if ((ShaderRequiresFlags & D3D_SHADER_REQUIRES_TYPED_UAV_LOAD_ADDITIONAL_FORMATS) != 0)
+				{
+					EnumAddFlags(CodeFeatures.CodeFeatures, EShaderCodeFeatures::TypedUAVLoadsExtended);
+				}
+
+				if ((ShaderRequiresFlags & (D3D_SHADER_REQUIRES_ATOMIC_INT64_ON_TYPED_RESOURCE| D3D_SHADER_REQUIRES_ATOMIC_INT64_ON_GROUP_SHARED)) != 0)
+				{
+					EnumAddFlags(CodeFeatures.CodeFeatures, EShaderCodeFeatures::Atomic64);
+				}
+
+				if (bDiagnosticBufferUsed)
+				{
+					EnumAddFlags(CodeFeatures.CodeFeatures, EShaderCodeFeatures::DiagnosticBuffer);
+				}
+
+				if ((ShaderRequiresFlags & D3D_SHADER_REQUIRES_RESOURCE_DESCRIPTOR_HEAP_INDEXING) != 0)
+				{
+					EnumAddFlags(CodeFeatures.CodeFeatures, EShaderCodeFeatures::BindlessResources);
+				}
+
+				if ((ShaderRequiresFlags & D3D_SHADER_REQUIRES_SAMPLER_DESCRIPTOR_HEAP_INDEXING) != 0)
+				{
+					EnumAddFlags(CodeFeatures.CodeFeatures, EShaderCodeFeatures::BindlessSamplers);
+				}
 
 				// We only need this to appear when using a DXC shader
 				ShaderCode.AddOptionalData<FShaderCodeFeatures>(CodeFeatures);
-				uint8 IsSM6 = 1;
-				ShaderCode.AddOptionalData('6', &IsSM6, 1);
+
+				if (Language != ELanguage::SM6)
+				{
+					uint8 IsSM6 = 1;
+					ShaderCode.AddOptionalData('6', &IsSM6, 1);
+				}
 			};
 
 			//#todo-rco: Should compress ShaderCode?
diff --git a/Engine/Source/Developer/Windows/ShaderFormatD3D/Private/ShaderFormatD3D.cpp b/Engine/Source/Developer/Windows/ShaderFormatD3D/Private/ShaderFormatD3D.cpp
--- a/Engine/Source/Developer/Windows/ShaderFormatD3D/Private/ShaderFormatD3D.cpp
+++ b/Engine/Source/Developer/Windows/ShaderFormatD3D/Private/ShaderFormatD3D.cpp
@@ -16,10 +16,12 @@ class FShaderFormatD3D : public IShaderFormat
 {
 	enum
 	{
+		UE_SHADER_PCD3D_SHARED_VER = 1,
+
 		/** Version for shader format, this becomes part of the DDC key. */
-		UE_SHADER_PCD3D_SM6_VER = 2,
-		UE_SHADER_PCD3D_SM5_VER = 9,
-		UE_SHADER_PCD3D_ES3_1_VER = 8,
+		UE_SHADER_PCD3D_SM6_VER = UE_SHADER_PCD3D_SHARED_VER + 2,
+		UE_SHADER_PCD3D_SM5_VER = UE_SHADER_PCD3D_SHARED_VER + 9,
+		UE_SHADER_PCD3D_ES3_1_VER = UE_SHADER_PCD3D_SHARED_VER + 8,
 	};
 
 	void CheckFormat(FName Format) const
diff --git a/Engine/Source/Runtime/D3D12RHI/Private/D3D12PipelineState.h b/Engine/Source/Runtime/D3D12RHI/Private/D3D12PipelineState.h
--- a/Engine/Source/Runtime/D3D12RHI/Private/D3D12PipelineState.h
+++ b/Engine/Source/Runtime/D3D12RHI/Private/D3D12PipelineState.h
@@ -98,7 +98,7 @@ struct FD3D12LowLevelGraphicsPipelineStateDesc
 		return Desc.MS.BytecodeLength > 0;
 	}
 
-#if PLATFORM_WINDOWS
+#if D3D12RHI_NEEDS_VENDOR_EXTENSIONS
 	// TODO: Replace with a global hash lookup to reduce overall footprint?
 	// Very few permutations, so a single > 0 u32 hash code would be lower
 	// memory usage, and very rarely cause a look up.
@@ -117,8 +117,6 @@ struct FD3D12LowLevelGraphicsPipelineStateDesc
 			PSExtensions != nullptr ||
 			GSExtensions != nullptr);
 	}
-#else
-	FORCEINLINE bool HasVendorExtensions() const { return false; }
 #endif
 
 	FORCEINLINE FString GetName() const { return FString::Printf(TEXT("%llu"), CombinedHash); }
@@ -143,11 +141,9 @@ struct FD3D12ComputePipelineStateDesc
 
 	SIZE_T CombinedHash;
 
-#if PLATFORM_WINDOWS
+#if D3D12RHI_NEEDS_VENDOR_EXTENSIONS
 	const TArray<FShaderCodeVendorExtension>* Extensions;
 	FORCEINLINE bool HasVendorExtensions() const { return (Extensions != nullptr); }
-#else
-	FORCEINLINE bool HasVendorExtensions() const { return false; }
 #endif
 
 	FORCEINLINE FString GetName() const { return FString::Printf(TEXT("%llu"), CombinedHash); }
diff --git a/Engine/Source/Runtime/D3D12RHI/Private/D3D12Pipelinestate.cpp b/Engine/Source/Runtime/D3D12RHI/Private/D3D12Pipelinestate.cpp
--- a/Engine/Source/Runtime/D3D12RHI/Private/D3D12Pipelinestate.cpp
+++ b/Engine/Source/Runtime/D3D12RHI/Private/D3D12Pipelinestate.cpp
@@ -52,8 +52,8 @@ FD3D12LowLevelGraphicsPipelineStateDesc GetLowLevelGraphicsPipelineStateDesc(con
 #define COPY_SHADER(Initial, Name) \
 	if (FD3D12##Name##Shader* Shader = (FD3D12##Name##Shader*) Initializer.BoundShaderState.##Name##ShaderRHI) \
 	{ \
-		Desc.Desc.Initial##S = Shader->ShaderBytecode.GetShaderBytecode(); \
-		Desc.Initial##SHash = Shader->ShaderBytecode.GetHash(); \
+		Desc.Desc.Initial##S = Shader->GetShaderBytecode(); \
+		Desc.Initial##SHash = Shader->GetBytecodeHash(); \
 	}
 	COPY_SHADER(V, Vertex);
 	COPY_SHADER(M, Mesh);
@@ -62,7 +62,7 @@ FD3D12LowLevelGraphicsPipelineStateDesc GetLowLevelGraphicsPipelineStateDesc(con
 	COPY_SHADER(G, Geometry);
 #undef COPY_SHADER
 
-#if PLATFORM_WINDOWS
+#if D3D12RHI_NEEDS_VENDOR_EXTENSIONS
 #define EXT_SHADER(Initial, Name) \
 	if (FD3D12##Name##Shader* Shader = (FD3D12##Name##Shader*) Initializer.BoundShaderState.##Name##ShaderRHI) \
 	{ \
@@ -97,9 +97,9 @@ FD3D12ComputePipelineStateDesc GetComputePipelineStateDesc(const FD3D12ComputeSh
 
 	Desc.pRootSignature = ComputeShader->pRootSignature;
 	Desc.Desc.pRootSignature = Desc.pRootSignature->GetRootSignature();
-	Desc.Desc.CS = ComputeShader->ShaderBytecode.GetShaderBytecode();
-	Desc.CSHash = ComputeShader->ShaderBytecode.GetHash();
-#if PLATFORM_WINDOWS
+	Desc.Desc.CS = ComputeShader->GetShaderBytecode();
+	Desc.CSHash = ComputeShader->GetBytecodeHash();
+#if D3D12RHI_NEEDS_VENDOR_EXTENSIONS
 	if (ComputeShader->VendorExtensions.Num() > 0)
 	{
 		Desc.Extensions = &ComputeShader->VendorExtensions;
diff --git a/Engine/Source/Runtime/D3D12RHI/Private/D3D12RayTracing.cpp b/Engine/Source/Runtime/D3D12RHI/Private/D3D12RayTracing.cpp
--- a/Engine/Source/Runtime/D3D12RHI/Private/D3D12RayTracing.cpp
+++ b/Engine/Source/Runtime/D3D12RHI/Private/D3D12RayTracing.cpp
@@ -404,11 +404,6 @@ struct FDXILLibrary
 		InitFromDXIL(ShaderBytecode.pShaderBytecode, ShaderBytecode.BytecodeLength, InEntryNames, InExportNames, NumEntryNames);
 	}
 
-	void InitFromDXIL(const FD3D12ShaderBytecode& ShaderBytecode, LPCWSTR* InEntryNames, LPCWSTR* InExportNames, uint32 NumEntryNames)
-	{
-		InitFromDXIL(ShaderBytecode.GetShaderBytecode(), InEntryNames, InExportNames, NumEntryNames);
-	}
-
 	D3D12_STATE_SUBOBJECT GetSubobject() const
 	{
 		D3D12_STATE_SUBOBJECT Subobject = {};
@@ -972,7 +967,7 @@ public:
 			check(Entry.ExportNames.Num() <= Entry.MaxExports);
 
 			FDXILLibrary Library;
-			Library.InitFromDXIL(Shader->ShaderBytecode, OriginalEntryPoints.GetData(), RenamedEntryPoints.GetData(), OriginalEntryPoints.Num());
+			Library.InitFromDXIL(Shader->GetShaderBytecode(), OriginalEntryPoints.GetData(), RenamedEntryPoints.GetData(), OriginalEntryPoints.Num());
 
 			const FDXILLibrary* LibraryPtr = &Library;
 
@@ -1073,7 +1068,7 @@ public:
 
 			if (Shader->bPrecompiledPSO)
 			{
-				D3D12_SHADER_BYTECODE Bytecode = Shader->ShaderBytecode.GetShaderBytecode();
+				D3D12_SHADER_BYTECODE Bytecode = Shader->GetShaderBytecode();
 				Entry.StateObject = Device->DeserializeRayTracingStateObject(Bytecode, GlobalRootSignature);
 				if (Entry.StateObject)
 				{
@@ -2048,7 +2043,7 @@ static void GetBuildInShaderLibrary(FDXILLibrary& ShaderLibrary)
 {
 	FD3D12RayTracingShader* RayTracingShader = GetBuildInRayTracingShader<ShaderType>();
 	LPCWSTR EntryName[1] = { *RayTracingShader->EntryPoint };
-	ShaderLibrary.InitFromDXIL(RayTracingShader->ShaderBytecode.GetShaderBytecode().pShaderBytecode, RayTracingShader->ShaderBytecode.GetShaderBytecode().BytecodeLength, EntryName, EntryName, 1);
+	ShaderLibrary.InitFromDXIL(RayTracingShader->GetShaderBytecode().pShaderBytecode, RayTracingShader->GetShaderBytecode().BytecodeLength, EntryName, EntryName, 1);
 }
 
 void FD3D12Device::DestroyRayTracingDescriptorCache()
@@ -2594,7 +2589,7 @@ public:
 		{
 			FShaderStats Stats;
 			Stats.Name = *(Entry->Shader->EntryPoint);
-			Stats.ShaderSize = uint32(Entry->Shader->ShaderBytecode.GetShaderBytecode().BytecodeLength);
+			Stats.ShaderSize = Entry->Shader->Code.Num();
 			Stats.CompileTimeMS = Entry->CompileTimeMS;
 
 		#if PLATFORM_WINDOWS
diff --git a/Engine/Source/Runtime/D3D12RHI/Private/D3D12Shader.h b/Engine/Source/Runtime/D3D12RHI/Private/D3D12Shader.h
--- a/Engine/Source/Runtime/D3D12RHI/Private/D3D12Shader.h
+++ b/Engine/Source/Runtime/D3D12RHI/Private/D3D12Shader.h
@@ -38,12 +38,56 @@ public:
 
 struct FD3D12ShaderData
 {
-	TArray<FShaderCodeVendorExtension> VendorExtensions;
+	/** The shader's bytecode, with custom data in the last byte. */
+	TArray<uint8> Code;
+
+	FD3D12ShaderResourceTable ShaderResourceTable;
+
+	FShaderCodePackedResourceCounts ResourceCounts{};
 
 	/** The static slot associated with the resource table index in ShaderResourceTable. */
 	TArray<FUniformBufferStaticSlot> StaticSlots;
 
-	bool bUsesDiagnosticBuffer = false;
+#if D3D12RHI_NEEDS_VENDOR_EXTENSIONS
+	TArray<FShaderCodeVendorExtension> VendorExtensions;
+#endif
+
+#if D3D12RHI_NEEDS_SHADER_FEATURE_CHECKS
+	EShaderCodeFeatures Features = EShaderCodeFeatures::None;
+#endif
+
+	D3D12_SHADER_BYTECODE GetShaderBytecode() const
+	{
+		return CD3DX12_SHADER_BYTECODE(Code.GetData(), Code.Num());
+	}
+
+	ShaderBytecodeHash GetBytecodeHash() const
+	{
+		ShaderBytecodeHash Hash;
+		if (Code.Num() == 0)
+		{
+			Hash.Hash[0] = Hash.Hash[1] = 0;
+		}
+		else
+		{
+			// D3D shader bytecode contains a 128bit checksum in DWORD 1-4. We can just use that directly instead of hashing the whole shader bytecode ourselves.
+			const uint8* pData = Code.GetData() + 4;
+			Hash = *reinterpret_cast<const ShaderBytecodeHash*>(pData);
+		}
+		return Hash;
+	}
+
+#if D3D12RHI_NEEDS_SHADER_FEATURE_CHECKS
+	FORCEINLINE EShaderCodeFeatures GetFeatures() const { return Features; }
+#else
+	FORCEINLINE EShaderCodeFeatures GetFeatures() const { return EShaderCodeFeatures::None; }
+#endif
+
+	FORCEINLINE bool UsesDiagnosticBuffer() const { return EnumHasAnyFlags(GetFeatures(), EShaderCodeFeatures::DiagnosticBuffer); }
+	FORCEINLINE bool UsesBindlessResources() const { return EnumHasAnyFlags(GetFeatures(), EShaderCodeFeatures::BindlessResources); }
+	FORCEINLINE bool UsesBindlessSamplers() const { return EnumHasAnyFlags(GetFeatures(), EShaderCodeFeatures::BindlessSamplers); }
+
+	bool InitCommon(TArrayView<const uint8> InCode);
 };
 
 /** This represents a vertex shader that hasn't been combined with a specific declaration to create a bound shader. */
@@ -51,83 +95,30 @@ class FD3D12VertexShader : public FRHIVertexShader, public FD3D12ShaderData
 {
 public:
 	enum { StaticFrequency = SF_Vertex };
-
-	/** The shader's bytecode. */
-	FD3D12ShaderBytecode ShaderBytecode;
-
-	FD3D12ShaderResourceTable ShaderResourceTable;
-
-	/** The vertex shader's bytecode, with custom data in the last byte. */
-	TArray<uint8> Code;
-
-	// TEMP remove with removal of bound shader state
-	int32 Offset;
-
-	FShaderCodePackedResourceCounts ResourceCounts;
 };
 
 class FD3D12MeshShader : public FRHIMeshShader, public FD3D12ShaderData
 {
 public:
 	enum { StaticFrequency = SF_Mesh };
-
-	/** The shader's bytecode. */
-	FD3D12ShaderBytecode ShaderBytecode;
-
-	FD3D12ShaderResourceTable ShaderResourceTable;
-
-	/** The mesh shader's bytecode, with custom data in the last byte. */
-	TArray<uint8> Code;
-
-	FShaderCodePackedResourceCounts ResourceCounts;
 };
 
 class FD3D12AmplificationShader : public FRHIAmplificationShader, public FD3D12ShaderData
 {
 public:
 	enum { StaticFrequency = SF_Amplification };
-
-	/** The shader's bytecode. */
-	FD3D12ShaderBytecode ShaderBytecode;
-
-	FD3D12ShaderResourceTable ShaderResourceTable;
-
-	/** The amplification shader's bytecode, with custom data in the last byte. */
-	TArray<uint8> Code;
-
-	FShaderCodePackedResourceCounts ResourceCounts;
 };
 
 class FD3D12GeometryShader : public FRHIGeometryShader, public FD3D12ShaderData
 {
 public:
 	enum { StaticFrequency = SF_Geometry };
-
-	/** The shader's bytecode. */
-	FD3D12ShaderBytecode ShaderBytecode;
-
-	FD3D12ShaderResourceTable ShaderResourceTable;
-
-	/** The shader's bytecode, with custom data in the last byte. */
-	TArray<uint8> Code;
-
-	FShaderCodePackedResourceCounts ResourceCounts;
 };
 
 class FD3D12PixelShader : public FRHIPixelShader, public FD3D12ShaderData
 {
 public:
 	enum { StaticFrequency = SF_Pixel };
-
-	/** The shader's bytecode. */
-	FD3D12ShaderBytecode ShaderBytecode;
-
-	/** The shader's bytecode, with custom data in the last byte. */
-	TArray<uint8> Code;
-
-	FD3D12ShaderResourceTable ShaderResourceTable;
-
-	FShaderCodePackedResourceCounts ResourceCounts;
 };
 
 class FD3D12ComputeShader : public FRHIComputeShader, public FD3D12ShaderData
@@ -135,16 +126,7 @@ class FD3D12ComputeShader : public FRHIComputeShader, public FD3D12ShaderData
 public:
 	enum { StaticFrequency = SF_Compute };
 
-	/** The shader's bytecode. */
-	FD3D12ShaderBytecode ShaderBytecode;
-
-	/** The shader's bytecode, with custom data in the last byte. */
-	TArray<uint8> Code;
-
-	FD3D12ShaderResourceTable ShaderResourceTable;
-
-	FShaderCodePackedResourceCounts ResourceCounts;
-	const FD3D12RootSignature* pRootSignature;
+	const FD3D12RootSignature* pRootSignature = nullptr;
 };
 
 /**
@@ -202,23 +184,13 @@ class FD3D12RayTracingShader : public FRHIRayTracingShader, public FD3D12ShaderD
 public:
 	explicit FD3D12RayTracingShader(EShaderFrequency InFrequency) : FRHIRayTracingShader(InFrequency) {}
 
-	/** The shader's bytecode. */
-	FD3D12ShaderBytecode ShaderBytecode;
-
-	FD3D12ShaderResourceTable ShaderResourceTable;
-
-	/** The shader's bytecode, with custom data in the last byte. */
-	TArray<uint8> Code;
+	const FD3D12RootSignature* pRootSignature = nullptr;
 
 	/** The shader's DXIL entrypoint & base export name for DXR (required for RTPSO creation) */
 	FString EntryPoint; // Primary entry point for all ray tracing shaders. Assumed to be closest hit shader for SF_RayHitGroup.
 	FString AnyHitEntryPoint; // Optional any-hit shader entry point for SF_RayHitGroup.
 	FString IntersectionEntryPoint; // Optional intersection shader entry point for SF_RayHitGroup.
 	bool bPrecompiledPSO = false;
-
-	FShaderCodePackedResourceCounts ResourceCounts;
-
-	const FD3D12RootSignature* pRootSignature = nullptr;
 };
 
 #endif // D3D12_RHI_RAYTRACING
diff --git a/Engine/Source/Runtime/D3D12RHI/Private/D3D12Shaders.cpp b/Engine/Source/Runtime/D3D12RHI/Private/D3D12Shaders.cpp
--- a/Engine/Source/Runtime/D3D12RHI/Private/D3D12Shaders.cpp
+++ b/Engine/Source/Runtime/D3D12RHI/Private/D3D12Shaders.cpp
@@ -7,7 +7,7 @@
 #include "D3D12RHIPrivate.h"
 
 template <typename TShaderType>
-static inline bool ReadShaderOptionalData(FShaderCodeReader& InShaderCode, TShaderType& OutShader, bool& bOutFoundCodeFeatures, FShaderCodeFeatures& OutCodeFeatures)
+static inline bool ReadShaderOptionalData(FShaderCodeReader& InShaderCode, TShaderType& OutShader)
 {
 	const FShaderCodePackedResourceCounts* PackedResourceCounts = InShaderCode.FindOptionalData<FShaderCodePackedResourceCounts>();
 	if (!PackedResourceCounts)
@@ -15,12 +15,13 @@ static inline bool ReadShaderOptionalData(FShaderCodeReader& InShaderCode, TShad
 		return false;
 	}
 	OutShader.ResourceCounts = *PackedResourceCounts;
+
 #if !(UE_BUILD_SHIPPING || UE_BUILD_TEST)
 	OutShader.ShaderName = InShaderCode.FindOptionalData(FShaderCodeName::Key);
 
+#if 0
 	int32 UniformBufferTableSize = 0;
 	const auto* UniformBufferData = InShaderCode.FindOptionalDataAndSize('u', UniformBufferTableSize);
-#if 0
 	//#todo-rco
 	if (UniformBufferData && UniformBufferTableSize > 0)
 	{
@@ -35,6 +36,8 @@ static inline bool ReadShaderOptionalData(FShaderCodeReader& InShaderCode, TShad
 	}
 #endif
 #endif
+
+#if D3D12RHI_NEEDS_VENDOR_EXTENSIONS
 	int32 VendorExtensionTableSize = 0;
 	auto* VendorExtensionData = InShaderCode.FindOptionalDataAndSize(FShaderCodeVendorExtension::Key, VendorExtensionTableSize);
 	if (VendorExtensionData && VendorExtensionTableSize > 0)
@@ -42,36 +45,69 @@ static inline bool ReadShaderOptionalData(FShaderCodeReader& InShaderCode, TShad
 		FBufferReader Ar((void*)VendorExtensionData, VendorExtensionTableSize, false);
 		Ar << OutShader.VendorExtensions;
 	}
+#endif
 
-	const FShaderCodeFeatures* CodeFeatures = InShaderCode.FindOptionalData<FShaderCodeFeatures>();
-	if (CodeFeatures)
+#if D3D12RHI_NEEDS_SHADER_FEATURE_CHECKS
+	if (const FShaderCodeFeatures* CodeFeatures = InShaderCode.FindOptionalData<FShaderCodeFeatures>())
 	{
-		bOutFoundCodeFeatures = true;
-		OutCodeFeatures = *CodeFeatures;
-
-		OutShader.bUsesDiagnosticBuffer = CodeFeatures->bUsesDiagnosticBuffer;
+		OutShader.Features = CodeFeatures->CodeFeatures;
 	}
-	else
+#endif
+
+	return true;
+}
+
+static bool ValidateShaderIsUsable(FD3D12ShaderData* InShader)
+{
+#if D3D12RHI_NEEDS_SHADER_FEATURE_CHECKS
+	if (EnumHasAnyFlags(InShader->Features, EShaderCodeFeatures::WaveOps) && !GRHISupportsWaveOperations)
 	{
-		bOutFoundCodeFeatures = false;
+		return false;
 	}
+#endif
 
 	return true;
 }
 
 template <typename TShaderType>
-static inline bool ReadShaderOptionalData(FShaderCodeReader& InShaderCode, TShaderType& OutShader)
+bool InitShaderCommon(FShaderCodeReader& ShaderCode, int32 Offset, TShaderType* InShader)
 {
-	bool bHasCodeFeatures = false;
-	FShaderCodeFeatures Features;
-	return ReadShaderOptionalData(InShaderCode, OutShader, bHasCodeFeatures, Features);
+	if (!ReadShaderOptionalData(ShaderCode, *InShader))
+	{
+		return false;
+	}
+
+	if (!ValidateShaderIsUsable(InShader))
+	{
+		return false;
+	}
+
+	// Copy the native shader data only, skipping any of our own headers.
+	InShader->Code = ShaderCode.GetOffsetShaderCode(Offset);
+
+	return true;
 }
 
 template <typename TShaderType>
-static inline void InitUniformBufferStaticSlots(TShaderType* Shader)
+TShaderType* CreateStandardShader(TArrayView<const uint8> InCode)
 {
-	const FBaseShaderResourceTable& SRT = Shader->ShaderResourceTable;
+	FShaderCodeReader ShaderCode(InCode);
+	TShaderType* Shader = new TShaderType();
+
+	FMemoryReaderView Ar(InCode, true);
+	Ar << Shader->ShaderResourceTable;
+
+	const int32 Offset = Ar.Tell();
+
+	if (!InitShaderCommon(ShaderCode, Offset, Shader))
+	{
+		delete Shader;
+		return nullptr;
+	}
+
+	// InitUniformBufferStaticSlots
 
+	const FBaseShaderResourceTable& SRT = Shader->ShaderResourceTable;
 	Shader->StaticSlots.Reserve(SRT.ResourceTableLayoutHashes.Num());
 
 	for (uint32 LayoutHash : SRT.ResourceTableLayoutHashes)
@@ -85,205 +121,51 @@ static inline void InitUniformBufferStaticSlots(TShaderType* Shader)
 			Shader->StaticSlots.Add(MAX_UNIFORM_BUFFER_STATIC_SLOTS);
 		}
 	}
+
+	return Shader;
 }
 
 FVertexShaderRHIRef FD3D12DynamicRHI::RHICreateVertexShader(TArrayView<const uint8> Code, const FSHAHash& Hash)
 {
-	FShaderCodeReader ShaderCode(Code);
-	FD3D12VertexShader* Shader = new FD3D12VertexShader;
-
-	FMemoryReaderView Ar(Code, true);
-	Ar << Shader->ShaderResourceTable;
-	int32 Offset = Ar.Tell();
-	const SIZE_T CodeSize = ShaderCode.GetActualShaderCodeSize() - Offset;
-	bool bFoundCodeFeatures;
-	FShaderCodeFeatures CodeFeatures;
-	if (!ReadShaderOptionalData(ShaderCode, *Shader, bFoundCodeFeatures, CodeFeatures))
-	{
-		return nullptr;
-	}
-	if (bFoundCodeFeatures && CodeFeatures.bUsesWaveOps && !GRHISupportsWaveOperations)
-	{
-		return nullptr;
-	}
-
-	Shader->Code = Code;
-	Shader->Offset = Offset;
-	InitUniformBufferStaticSlots(Shader);
-
-	D3D12_SHADER_BYTECODE ShaderBytecode;
-	ShaderBytecode.pShaderBytecode = Shader->Code.GetData() + Offset;
-	ShaderBytecode.BytecodeLength = CodeSize;
-	Shader->ShaderBytecode.SetShaderBytecode(ShaderBytecode);
-
-	return Shader;
+	return CreateStandardShader<FD3D12VertexShader>(Code);
 }
 
 FMeshShaderRHIRef FD3D12DynamicRHI::RHICreateMeshShader(TArrayView<const uint8> Code, const FSHAHash& Hash)
 {
-	FShaderCodeReader ShaderCode(Code);
-	FD3D12MeshShader* Shader = new FD3D12MeshShader;
-
-	FMemoryReaderView Ar(Code, true);
-	Ar << Shader->ShaderResourceTable;
-	int32 Offset = Ar.Tell();
-	const SIZE_T CodeSize = ShaderCode.GetActualShaderCodeSize() - Offset;
-	bool bFoundCodeFeatures;
-	FShaderCodeFeatures CodeFeatures;
-	if (!ReadShaderOptionalData(ShaderCode, *Shader, bFoundCodeFeatures, CodeFeatures))
-	{
-		return nullptr;
-	}
-	if (bFoundCodeFeatures && CodeFeatures.bUsesWaveOps && !GRHISupportsWaveOperations)
-	{
-		return nullptr;
-	}
-
-	Shader->Code = Code;
-	InitUniformBufferStaticSlots(Shader);
-
-	D3D12_SHADER_BYTECODE ShaderBytecode;
-	ShaderBytecode.pShaderBytecode = Shader->Code.GetData() + Offset;
-	ShaderBytecode.BytecodeLength = CodeSize;
-	Shader->ShaderBytecode.SetShaderBytecode(ShaderBytecode);
-
-	return Shader;
+	return CreateStandardShader<FD3D12MeshShader>(Code);
 }
 
 FAmplificationShaderRHIRef FD3D12DynamicRHI::RHICreateAmplificationShader(TArrayView<const uint8> Code, const FSHAHash& Hash)
 {
-	FShaderCodeReader ShaderCode(Code);
-	FD3D12AmplificationShader* Shader = new FD3D12AmplificationShader;
-
-	FMemoryReaderView Ar(Code, true);
-	Ar << Shader->ShaderResourceTable;
-	int32 Offset = Ar.Tell();
-	const SIZE_T CodeSize = ShaderCode.GetActualShaderCodeSize() - Offset;
-	bool bFoundCodeFeatures;
-	FShaderCodeFeatures CodeFeatures;
-	if (!ReadShaderOptionalData(ShaderCode, *Shader, bFoundCodeFeatures, CodeFeatures))
-	{
-		return nullptr;
-	}
-	if (bFoundCodeFeatures && CodeFeatures.bUsesWaveOps && !GRHISupportsWaveOperations)
-	{
-		return nullptr;
-	}
-
-	Shader->Code = Code;
-	InitUniformBufferStaticSlots(Shader);
-
-	D3D12_SHADER_BYTECODE ShaderBytecode;
-	ShaderBytecode.pShaderBytecode = Shader->Code.GetData() + Offset;
-	ShaderBytecode.BytecodeLength = CodeSize;
-	Shader->ShaderBytecode.SetShaderBytecode(ShaderBytecode);
-
-	return Shader;
+	return CreateStandardShader<FD3D12AmplificationShader>(Code);
 }
 
 FPixelShaderRHIRef FD3D12DynamicRHI::RHICreatePixelShader(TArrayView<const uint8> Code, const FSHAHash& Hash)
 {
-	FShaderCodeReader ShaderCode(Code);
-
-	FD3D12PixelShader* Shader = new FD3D12PixelShader;
-
-	FMemoryReaderView Ar(Code, true);
-	Ar << Shader->ShaderResourceTable;
-	int32 Offset = Ar.Tell();
-	const SIZE_T CodeSize = ShaderCode.GetActualShaderCodeSize() - Offset;
-	bool bFoundCodeFeatures;
-	FShaderCodeFeatures CodeFeatures;
-	if (!ReadShaderOptionalData(ShaderCode, *Shader, bFoundCodeFeatures, CodeFeatures))
-	{
-		return nullptr;
-	}
-	if (bFoundCodeFeatures && CodeFeatures.bUsesWaveOps && !GRHISupportsWaveOperations)
-	{
-		return nullptr;
-	}
-
-	Shader->Code = Code;
-	InitUniformBufferStaticSlots(Shader);
-
-	D3D12_SHADER_BYTECODE ShaderBytecode;
-	ShaderBytecode.pShaderBytecode = Shader->Code.GetData() + Offset;
-	ShaderBytecode.BytecodeLength = CodeSize;
-	Shader->ShaderBytecode.SetShaderBytecode(ShaderBytecode);
-
-	return Shader;
+	return CreateStandardShader<FD3D12PixelShader>(Code);
 }
 
 FGeometryShaderRHIRef FD3D12DynamicRHI::RHICreateGeometryShader(TArrayView<const uint8> Code, const FSHAHash& Hash)
 {
-	FShaderCodeReader ShaderCode(Code);
-
-	FD3D12GeometryShader* Shader = new FD3D12GeometryShader;
-
-	FMemoryReaderView Ar(Code, true);
-	Ar << Shader->ShaderResourceTable;
-	int32 Offset = Ar.Tell();
-	const SIZE_T CodeSize = ShaderCode.GetActualShaderCodeSize() - Offset;
-	bool bFoundCodeFeatures;
-	FShaderCodeFeatures CodeFeatures;
-	if (!ReadShaderOptionalData(ShaderCode, *Shader, bFoundCodeFeatures, CodeFeatures))
-	{
-		return nullptr;
-	}
-	if (bFoundCodeFeatures && CodeFeatures.bUsesWaveOps && !GRHISupportsWaveOperations)
-	{
-		return nullptr;
-	}
-
-	Shader->Code = Code;
-	InitUniformBufferStaticSlots(Shader);
-
-	D3D12_SHADER_BYTECODE ShaderBytecode;
-	ShaderBytecode.pShaderBytecode = Shader->Code.GetData() + Offset;
-	ShaderBytecode.BytecodeLength = CodeSize;
-	Shader->ShaderBytecode.SetShaderBytecode(ShaderBytecode);
-
-	return Shader;
+	return CreateStandardShader<FD3D12GeometryShader>(Code);
 }
 
 FComputeShaderRHIRef FD3D12DynamicRHI::RHICreateComputeShader(TArrayView<const uint8> Code, const FSHAHash& Hash)
 {
-	FShaderCodeReader ShaderCode(Code);
-
-	FD3D12ComputeShader* Shader = new FD3D12ComputeShader;
-
-	FMemoryReaderView Ar(Code, true);
-	Ar << Shader->ShaderResourceTable;
-	int32 Offset = Ar.Tell();
-	const SIZE_T CodeSize = ShaderCode.GetActualShaderCodeSize() - Offset;
-	bool bFoundCodeFeatures;
-	FShaderCodeFeatures CodeFeatures;
-	if (!ReadShaderOptionalData(ShaderCode, *Shader, bFoundCodeFeatures, CodeFeatures))
+	FD3D12ComputeShader* Shader = CreateStandardShader<FD3D12ComputeShader>(Code);
+	if (Shader)
 	{
-		return nullptr;
-	}
-	if (bFoundCodeFeatures && CodeFeatures.bUsesWaveOps && !GRHISupportsWaveOperations)
-	{
-		return nullptr;
-	}
-
-	Shader->Code = Code;
-	InitUniformBufferStaticSlots(Shader);
-
-	D3D12_SHADER_BYTECODE ShaderBytecode;
-	ShaderBytecode.pShaderBytecode = Shader->Code.GetData() + Offset;
-	ShaderBytecode.BytecodeLength = CodeSize;
-	Shader->ShaderBytecode.SetShaderBytecode(ShaderBytecode);
-
-	FD3D12Adapter& Adapter = GetAdapter();
+		FD3D12Adapter& Adapter = GetAdapter();
 
 #if USE_STATIC_ROOT_SIGNATURE
-	Shader->pRootSignature = Adapter.GetStaticComputeRootSignature();
+		Shader->pRootSignature = Adapter.GetStaticComputeRootSignature();
 #else
-	const D3D12_RESOURCE_BINDING_TIER Tier = Adapter.GetResourceBindingTier();
-	FD3D12QuantizedBoundShaderState QBSS;
-	QuantizeBoundShaderState(Tier, Shader, QBSS);
-	Shader->pRootSignature = Adapter.GetRootSignature(QBSS);
+		const D3D12_RESOURCE_BINDING_TIER Tier = Adapter.GetResourceBindingTier();
+		FD3D12QuantizedBoundShaderState QBSS;
+		QuantizeBoundShaderState(Tier, Shader, QBSS);
+		Shader->pRootSignature = Adapter.GetRootSignature(QBSS);
 #endif
+	}
 
 	return Shader;
 }
@@ -305,13 +187,6 @@ FRayTracingShaderRHIRef FD3D12DynamicRHI::RHICreateRayTracingShader(TArrayView<c
 
 	int32 Offset = Ar.Tell();
 
-	bool bFoundCodeFeatures;
-	FShaderCodeFeatures CodeFeatures;
-	if (!ReadShaderOptionalData(ShaderCode, *Shader, bFoundCodeFeatures, CodeFeatures))
-	{
-		return nullptr;
-	}
-
 	int32 PrecompiledKey = 0;
 	Ar << PrecompiledKey;
 	if (PrecompiledKey == RayTracingPrecompiledPSOKey)
@@ -320,15 +195,11 @@ FRayTracingShaderRHIRef FD3D12DynamicRHI::RHICreateRayTracingShader(TArrayView<c
 		Shader->bPrecompiledPSO = true;
 	}
 
-	// Copy the native shader data only, skipping any of our own headers.
-	TArrayView<const uint8> NativeShaderCode = MakeArrayView(Code.GetData() + Offset, ShaderCode.GetActualShaderCodeSize() - Offset);
-	Shader->Code = TArray<uint8>(NativeShaderCode);
-
-	D3D12_SHADER_BYTECODE ShaderBytecode;
-	ShaderBytecode.pShaderBytecode = Shader->Code.GetData();
-	ShaderBytecode.BytecodeLength = Shader->Code.Num();
-
-	Shader->ShaderBytecode.SetShaderBytecode(ShaderBytecode);
+	if (!InitShaderCommon(ShaderCode, Offset, Shader))
+	{
+		delete Shader;
+		return nullptr;
+	}
 
 	FD3D12Adapter& Adapter = GetAdapter();
 
diff --git a/Engine/Source/Runtime/D3D12RHI/Private/D3D12Util.cpp b/Engine/Source/Runtime/D3D12RHI/Private/D3D12Util.cpp
--- a/Engine/Source/Runtime/D3D12RHI/Private/D3D12Util.cpp
+++ b/Engine/Source/Runtime/D3D12RHI/Private/D3D12Util.cpp
@@ -930,6 +930,7 @@ void FD3D12QuantizedBoundShaderState::InitShaderRegisterCounts(const D3D12_RESOU
 
 bool UsesVendorExtensionSpace(const FD3D12ShaderData& ShaderData)
 {
+#if D3D12RHI_NEEDS_VENDOR_EXTENSIONS
 	for (const FShaderCodeVendorExtension& Extension : ShaderData.VendorExtensions)
 	{
 		if (Extension.VendorId == 0x1002) // AMD
@@ -938,6 +939,7 @@ bool UsesVendorExtensionSpace(const FD3D12ShaderData& ShaderData)
 			return true;
 		}
 	}
+#endif
 
 	return false;
 }
@@ -958,35 +960,35 @@ void QuantizeBoundShaderState(
 	{
 		FD3D12QuantizedBoundShaderState::InitShaderRegisterCounts(ResourceBindingTier, VertexShader->ResourceCounts, QBSS.RegisterCounts[SV_Vertex]);
 		QBSS.bUseVendorExtension |= UsesVendorExtensionSpace(*VertexShader);
-		QBSS.bUseDiagnosticBuffer = VertexShader->bUsesDiagnosticBuffer;
+		QBSS.bUseDiagnosticBuffer = VertexShader->UsesDiagnosticBuffer();
 	}
 
 	if (const FD3D12MeshShader* const MeshShader = BSS->GetMeshShader())
 	{
 		FD3D12QuantizedBoundShaderState::InitShaderRegisterCounts(ResourceBindingTier, MeshShader->ResourceCounts, QBSS.RegisterCounts[SV_Mesh]);
 		QBSS.bUseVendorExtension |= UsesVendorExtensionSpace(*MeshShader);
-		QBSS.bUseDiagnosticBuffer = MeshShader->bUsesDiagnosticBuffer;
+		QBSS.bUseDiagnosticBuffer = MeshShader->UsesDiagnosticBuffer();
 	}
 
 	if (const FD3D12AmplificationShader* const AmplificationShader = BSS->GetAmplificationShader())
 	{
 		FD3D12QuantizedBoundShaderState::InitShaderRegisterCounts(ResourceBindingTier, AmplificationShader->ResourceCounts, QBSS.RegisterCounts[SV_Amplification]);
 		QBSS.bUseVendorExtension |= UsesVendorExtensionSpace(*AmplificationShader);
-		QBSS.bUseDiagnosticBuffer = AmplificationShader->bUsesDiagnosticBuffer;
+		QBSS.bUseDiagnosticBuffer = AmplificationShader->UsesDiagnosticBuffer();
 	}
 
 	if (const FD3D12PixelShader* const PixelShader = BSS->GetPixelShader())
 	{
 		FD3D12QuantizedBoundShaderState::InitShaderRegisterCounts(ResourceBindingTier, PixelShader->ResourceCounts, QBSS.RegisterCounts[SV_Pixel], true);
 		QBSS.bUseVendorExtension |= UsesVendorExtensionSpace(*PixelShader);
-		QBSS.bUseDiagnosticBuffer = PixelShader->bUsesDiagnosticBuffer;
+		QBSS.bUseDiagnosticBuffer = PixelShader->UsesDiagnosticBuffer();
 	}
 
 	if (const FD3D12GeometryShader* const GeometryShader = BSS->GetGeometryShader())
 	{
 		FD3D12QuantizedBoundShaderState::InitShaderRegisterCounts(ResourceBindingTier, GeometryShader->ResourceCounts, QBSS.RegisterCounts[SV_Geometry]);
 		QBSS.bUseVendorExtension |= UsesVendorExtensionSpace(*GeometryShader);
-		QBSS.bUseDiagnosticBuffer = GeometryShader->bUsesDiagnosticBuffer;
+		QBSS.bUseDiagnosticBuffer = GeometryShader->UsesDiagnosticBuffer();
 	}
 }
 
@@ -1015,8 +1017,10 @@ void QuantizeBoundShaderState(
 	const bool bAllosUAVs = true;
 	QuantizeBoundShaderStateCommon(ResourceBindingTier, ComputeShader->ResourceCounts, SV_All, bAllosUAVs, OutQBSS);
 	check(OutQBSS.bAllowIAInputLayout == false); // No access to vertex buffers needed
+#if D3D12RHI_NEEDS_VENDOR_EXTENSIONS
 	OutQBSS.bUseVendorExtension = (ComputeShader->VendorExtensions.Num() > 0);
-	OutQBSS.bUseDiagnosticBuffer = ComputeShader->bUsesDiagnosticBuffer;
+#endif
+	OutQBSS.bUseDiagnosticBuffer = ComputeShader->UsesDiagnosticBuffer();
 }
 
 #if D3D12_RHI_RAYTRACING
@@ -1086,7 +1090,7 @@ void QuantizeBoundShaderState(
 
 	if (RayTracingShader)
 	{
-		OutQBSS.bUseDiagnosticBuffer = RayTracingShader->bUsesDiagnosticBuffer;
+		OutQBSS.bUseDiagnosticBuffer = RayTracingShader->UsesDiagnosticBuffer();
 	}
 }
 #endif // D3D12_RHI_RAYTRACING
diff --git a/Engine/Source/Runtime/D3D12RHI/Private/HoloLens/D3D12RHIBasePrivate.h b/Engine/Source/Runtime/D3D12RHI/Private/HoloLens/D3D12RHIBasePrivate.h
--- a/Engine/Source/Runtime/D3D12RHI/Private/HoloLens/D3D12RHIBasePrivate.h
+++ b/Engine/Source/Runtime/D3D12RHI/Private/HoloLens/D3D12RHIBasePrivate.h
@@ -32,6 +32,9 @@
 #define D3D12RHI_RESOURCE_FLAG_ALLOW_INDIRECT_BUFFER D3D12_RESOURCE_FLAG_NONE
 #define D3D12RHI_HEAP_FLAG_ALLOW_INDIRECT_BUFFERS		D3D12_HEAP_FLAG_NONE
 
+#define D3D12RHI_NEEDS_VENDOR_EXTENSIONS     0
+#define D3D12RHI_NEEDS_SHADER_FEATURE_CHECKS 0
+
 // Heap create not zeroed flag is not available on Hololens so use internal define to disable the feature
 // but make code path shared when it becomes available
 #define FD3D12_HEAP_FLAG_CREATE_NOT_ZEROED D3D12_HEAP_FLAG_NONE
diff --git a/Engine/Source/Runtime/D3D12RHI/Private/Windows/D3D12RHIBasePrivate.h b/Engine/Source/Runtime/D3D12RHI/Private/Windows/D3D12RHIBasePrivate.h
--- a/Engine/Source/Runtime/D3D12RHI/Private/Windows/D3D12RHIBasePrivate.h
+++ b/Engine/Source/Runtime/D3D12RHI/Private/Windows/D3D12RHIBasePrivate.h
@@ -33,6 +33,9 @@ THIRD_PARTY_INCLUDES_END
 #define D3D12RHI_RESOURCE_FLAG_ALLOW_INDIRECT_BUFFER D3D12_RESOURCE_FLAG_NONE
 #define D3D12RHI_HEAP_FLAG_ALLOW_INDIRECT_BUFFERS		D3D12_HEAP_FLAG_NONE
 
+#define D3D12RHI_NEEDS_VENDOR_EXTENSIONS     1
+#define D3D12RHI_NEEDS_SHADER_FEATURE_CHECKS 1
+
 #include "../Public/D3D12Util.h"
 #include "WindowsD3D12DiskCache.h"
 #include "WindowsD3D12PipelineState.h"
diff --git a/Engine/Source/Runtime/D3D12RHI/Private/Windows/WindowsD3D12PipelineState.cpp b/Engine/Source/Runtime/D3D12RHI/Private/Windows/WindowsD3D12PipelineState.cpp
--- a/Engine/Source/Runtime/D3D12RHI/Private/Windows/WindowsD3D12PipelineState.cpp
+++ b/Engine/Source/Runtime/D3D12RHI/Private/Windows/WindowsD3D12PipelineState.cpp
@@ -968,7 +968,7 @@ static void CreatePipelineStateWrapper(ID3D12PipelineState** PSO, FD3D12Adapter*
 	}
 }
 
-#if PLATFORM_WINDOWS
+#if D3D12RHI_NEEDS_VENDOR_EXTENSIONS
 
 static FORCEINLINE NVAPI_D3D12_PSO_SET_SHADER_EXTENSION_SLOT_DESC GetNVShaderExtensionDesc(uint32 UavSlot)
 {
diff --git a/Engine/Source/Runtime/D3D12RHI/Public/D3D12Util.h b/Engine/Source/Runtime/D3D12RHI/Public/D3D12Util.h
--- a/Engine/Source/Runtime/D3D12RHI/Public/D3D12Util.h
+++ b/Engine/Source/Runtime/D3D12RHI/Public/D3D12Util.h
@@ -318,50 +318,6 @@ struct ShaderBytecodeHash
 	}
 };
 
-class FD3D12ShaderBytecode
-{
-public:
-	FD3D12ShaderBytecode()
-	{
-		FMemory::Memzero(Shader);
-		FMemory::Memzero(Hash);
-	}
-
-	FD3D12ShaderBytecode(const D3D12_SHADER_BYTECODE &InShader)
-		: Shader(InShader)
-	{
-		HashShader();
-	}
-
-	void SetShaderBytecode(const D3D12_SHADER_BYTECODE &InShader)
-	{
-		Shader = InShader;
-		HashShader();
-	}
-
-	const D3D12_SHADER_BYTECODE& GetShaderBytecode() const { return Shader; }
-	const ShaderBytecodeHash& GetHash() const { return Hash; }
-
-private:
-	void HashShader()
-	{
-		if (Shader.pShaderBytecode && Shader.BytecodeLength > 0)
-		{
-			// D3D shader bytecode contains a 128bit checksum in DWORD 1-4. We can just use that directly instead of hashing the whole shader bytecode ourselves.
-			check(Shader.BytecodeLength >= sizeof(uint32) + sizeof(Hash));
-			FMemory::Memcpy(&Hash, ((uint32*) Shader.pShaderBytecode) + 1, sizeof(Hash));
-		}
-		else
-		{
-			FMemory::Memzero(Hash);
-		}
-	}
-
-private:
-	ShaderBytecodeHash Hash;
-	D3D12_SHADER_BYTECODE Shader;
-};
-
 /**
  * The base class of threadsafe reference counted objects.
  */
diff --git a/Engine/Source/Runtime/RenderCore/Public/ShaderCore.h b/Engine/Source/Runtime/RenderCore/Public/ShaderCore.h
--- a/Engine/Source/Runtime/RenderCore/Public/ShaderCore.h
+++ b/Engine/Source/Runtime/RenderCore/Public/ShaderCore.h
@@ -177,16 +177,20 @@ enum class EShaderParameterType : uint8
 
 struct FParameterAllocation
 {
-	uint16 BufferIndex;
-	uint16 BaseIndex;
-	uint16 Size;
-	EShaderParameterType Type;
-	mutable bool bBound;
-
-	FParameterAllocation() :
-		Type(EShaderParameterType::Num),
-		bBound(false)
-	{}
+	uint16 BufferIndex = 0;
+	uint16 BaseIndex = 0;
+	uint16 Size = 0;
+	EShaderParameterType Type{ EShaderParameterType::Num };
+	mutable bool bBound = false;
+
+	FParameterAllocation() = default;
+	FParameterAllocation(uint16 InBufferIndex, uint16 InBaseIndex, uint16 InSize, EShaderParameterType InType)
+		: BufferIndex(InBufferIndex)
+		, BaseIndex(InBaseIndex)
+		, Size(InSize)
+		, Type(InType)
+	{
+	}
 
 	friend FArchive& operator<<(FArchive& Ar,FParameterAllocation& Allocation)
 	{
@@ -538,13 +542,25 @@ struct FShaderCodeResourceMasks
 };
 
 // if this changes you need to make sure all shaders get invalidated
+enum class EShaderCodeFeatures : uint8
+{
+	None                    = 0,
+	WaveOps                 = 1 << 0,
+	SixteenBitTypes         = 1 << 1,
+	TypedUAVLoadsExtended   = 1 << 2,
+	Atomic64                = 1 << 3,
+	DiagnosticBuffer        = 1 << 4,
+	BindlessResources       = 1 << 5,
+	BindlessSamplers        = 1 << 6,
+};
+ENUM_CLASS_FLAGS(EShaderCodeFeatures);
+
 struct FShaderCodeFeatures
 {
 	// for FindOptionalData() and AddOptionalData()
 	static const uint8 Key = 'x';
 
-	bool bUsesWaveOps = false;
-	bool bUsesDiagnosticBuffer = false;
+	EShaderCodeFeatures CodeFeatures = EShaderCodeFeatures::None;
 };
 
 // if this changes you need to make sure all shaders get invalidated
@@ -561,9 +577,16 @@ struct FShaderCodeVendorExtension
 	// for FindOptionalData() and AddOptionalData()
 	static const uint8 Key = 'v';
 
-	uint32 VendorId;
+	uint32 VendorId = 0;
 	FParameterAllocation Parameter;
 
+	FShaderCodeVendorExtension() = default;
+	FShaderCodeVendorExtension(uint32 InVendorId, uint16 InBufferIndex, uint16 InBaseIndex, uint16 InSize, EShaderParameterType InType)
+		: VendorId(InVendorId)
+		, Parameter(InBufferIndex, InBaseIndex, InSize, InType)
+	{
+	}
+
 	friend FArchive& operator<<(FArchive& Ar, FShaderCodeVendorExtension& Extension)
 	{
 		return Ar << Extension.VendorId << Extension.Parameter;
@@ -604,6 +627,11 @@ public:
 		return ShaderCode.Num() - GetOptionalDataSize();
 	}
 
+	TArrayView<const uint8> GetOffsetShaderCode(int32 Offset)
+	{
+		return MakeArrayView(ShaderCode.GetData() + Offset, GetActualShaderCodeSize() - Offset);
+	}
+
 	// for convenience
 	template <class T>
 	const T* FindOptionalData() const
