Undo changes to transition vertex buffers each time they are bound (CL 17489558).

Unfortunately the state cache can contain dangling pointers to buffers which have been freed or renamed mid-frame, because nothing clears slots in the vertex buffer cache until the end of the frame. This was "fine" before, since we only touch the views in SetVertexBuffers, and as long as the shader doesn't use the bogus slots, nothing bad happens. However, the change made that function touch CurrentVertexBufferResources, which can contain pointers to freed memory.

This brings back the problem that the code can miss transitions. We'll have to revisit this and properly clear the cache when resources are deleted or renamed.

#rnx

diff --git a/Engine/Source/Runtime/D3D12RHI/Private/D3D12DescriptorCache.cpp b/Engine/Source/Runtime/D3D12RHI/Private/D3D12DescriptorCache.cpp
--- a/Engine/Source/Runtime/D3D12RHI/Private/D3D12DescriptorCache.cpp
+++ b/Engine/Source/Runtime/D3D12RHI/Private/D3D12DescriptorCache.cpp
@@ -295,19 +295,6 @@ void FD3D12DescriptorCache::SetVertexBuffers(FD3D12VertexBufferCache& Cache)
 
 	CmdContext->CommandListHandle.UpdateResidency(Cache.ResidencyHandles, Count);
 	CmdContext->CommandListHandle->IASetVertexBuffers(0, Count, Cache.CurrentVertexBufferViews);
-
-	for (uint32 i = 0; i < Count; ++i)
-	{
-		if (Cache.CurrentVertexBufferResources[i])
-		{
-			FD3D12Resource* Resource = Cache.CurrentVertexBufferResources[i]->GetResource();
-			if (Resource->RequiresResourceStateTracking())
-			{
-				check(Resource->GetSubresourceCount() == 1);
-				FD3D12DynamicRHI::TransitionResource(CmdContext->CommandListHandle, Resource, D3D12_RESOURCE_STATE_TBD, D3D12_RESOURCE_STATE_VERTEX_AND_CONSTANT_BUFFER, D3D12_RESOURCE_BARRIER_ALL_SUBRESOURCES, FD3D12DynamicRHI::ETransitionMode::Validate);
-			}
-		}
-	}
 }
 
 template <EShaderFrequency ShaderStage>
diff --git a/Engine/Source/Runtime/D3D12RHI/Private/D3D12StateCache.cpp b/Engine/Source/Runtime/D3D12RHI/Private/D3D12StateCache.cpp
--- a/Engine/Source/Runtime/D3D12RHI/Private/D3D12StateCache.cpp
+++ b/Engine/Source/Runtime/D3D12RHI/Private/D3D12StateCache.cpp
@@ -1198,6 +1198,17 @@ void FD3D12StateCacheBase::InternalSetStreamSource(FD3D12ResourceLocation* Verte
 			PipelineState.Graphics.VBCache.MaxBoundVertexBufferIndex = INDEX_NONE;
 		}
 	}
+
+	if (VertexBufferLocation)
+	{
+		FD3D12Resource* const pResource = VertexBufferLocation->GetResource();
+		if (pResource && pResource->RequiresResourceStateTracking())
+		{
+			check(pResource->GetSubresourceCount() == 1);
+			FD3D12DynamicRHI::TransitionResource(CmdContext->CommandListHandle, pResource, D3D12_RESOURCE_STATE_TBD, D3D12_RESOURCE_STATE_VERTEX_AND_CONSTANT_BUFFER, D3D12_RESOURCE_BARRIER_ALL_SUBRESOURCES, FD3D12DynamicRHI::ETransitionMode::Validate);
+		}
+	}
+
 }
 
 template <EShaderFrequency ShaderFrequency>
