More D3D cleanup
- Ripping out all non-stream PSO creation now that we're locked into ID3D12Device2 support.
- Tidying up some VendorExtension handling.
- Renaming FD3D12QuantizedBoundShaderState::bUseVendorExtension to bNeedsAgsIntrinsicsSpace to be more explicit on what it does. This was accidentally being enabled when Nvidia extensions were used in Compute shaders.

diff --git a/Engine/Source/Runtime/D3D12RHI/Private/D3D12PipelineState.h b/Engine/Source/Runtime/D3D12RHI/Private/D3D12PipelineState.h
--- a/Engine/Source/Runtime/D3D12RHI/Private/D3D12PipelineState.h
+++ b/Engine/Source/Runtime/D3D12RHI/Private/D3D12PipelineState.h
@@ -71,12 +71,10 @@ struct FD3D12_GRAPHICS_PIPELINE_STATE_DESC
 	D3D12_CACHED_PIPELINE_STATE CachedPSO;
 	D3D12_PIPELINE_STATE_FLAGS Flags;
 
-	inline bool RequiresStream() const { return MS.BytecodeLength > 0; }
 	FD3D12_GRAPHICS_PIPELINE_STATE_STREAM PipelineStateStream() const;
 #if PLATFORM_SUPPORTS_MESH_SHADERS
 	FD3D12_MESH_PIPELINE_STATE_STREAM MeshPipelineStateStream() const;
 #endif
-	D3D12_GRAPHICS_PIPELINE_STATE_DESC GraphicsDescV0() const;
 };
 
 struct FD3D12LowLevelGraphicsPipelineStateDesc
@@ -130,7 +128,6 @@ struct FD3D12LowLevelGraphicsPipelineStateDesc
 struct FD3D12_COMPUTE_PIPELINE_STATE_DESC : public D3D12_COMPUTE_PIPELINE_STATE_DESC
 {
 	FD3D12_COMPUTE_PIPELINE_STATE_STREAM PipelineStateStream() const;
-	D3D12_COMPUTE_PIPELINE_STATE_DESC ComputeDescV0() const;
 };
 
 struct FD3D12ComputePipelineStateDesc
diff --git a/Engine/Source/Runtime/D3D12RHI/Private/D3D12RootSignature.cpp b/Engine/Source/Runtime/D3D12RHI/Private/D3D12RootSignature.cpp
--- a/Engine/Source/Runtime/D3D12RHI/Private/D3D12RootSignature.cpp
+++ b/Engine/Source/Runtime/D3D12RHI/Private/D3D12RootSignature.cpp
@@ -316,8 +316,8 @@ FD3D12RootSignatureDesc::FD3D12RootSignatureDesc(const FD3D12QuantizedBoundShade
 		}
 	}
 
-#if !PLATFORM_CPU_ARM_FAMILY && (PLATFORM_WINDOWS || PLATFORM_HOLOLENS)
-	if (QBSS.bUseVendorExtension)
+#if D3D12RHI_NEEDS_VENDOR_EXTENSIONS
+	if (QBSS.bNeedsAgsIntrinsicsSpace)
 	{
 		check(RootParameterCount < MaxRootParameters);
 		TableSlots[RootParameterCount].InitAsUnorderedAccessView(0, AGS_DX12_SHADER_INSTRINSICS_SPACE_ID, D3D12_ROOT_DESCRIPTOR_FLAG_DATA_VOLATILE, D3D12_SHADER_VISIBILITY_ALL);
diff --git a/Engine/Source/Runtime/D3D12RHI/Private/D3D12Util.cpp b/Engine/Source/Runtime/D3D12RHI/Private/D3D12Util.cpp
--- a/Engine/Source/Runtime/D3D12RHI/Private/D3D12Util.cpp
+++ b/Engine/Source/Runtime/D3D12RHI/Private/D3D12Util.cpp
@@ -928,7 +928,7 @@ void FD3D12QuantizedBoundShaderState::InitShaderRegisterCounts(const D3D12_RESOU
 	}
 }
 
-bool UsesVendorExtensionSpace(const FD3D12ShaderData& ShaderData)
+bool NeedsAgsIntrinsicsSpace(const FD3D12ShaderData& ShaderData)
 {
 #if D3D12RHI_NEEDS_VENDOR_EXTENSIONS
 	for (const FShaderCodeVendorExtension& Extension : ShaderData.VendorExtensions)
@@ -944,6 +944,22 @@ bool UsesVendorExtensionSpace(const FD3D12ShaderData& ShaderData)
 	return false;
 }
 
+static void QuantizeBoundShaderStateCommon(
+	FD3D12QuantizedBoundShaderState& OutQBSS,
+	const FD3D12ShaderData* ShaderData,
+	D3D12_RESOURCE_BINDING_TIER ResourceBindingTier,
+	EShaderVisibility ShaderVisibility,
+	bool bAllowUAVs = false
+)
+{
+	if (ShaderData)
+	{
+		FD3D12QuantizedBoundShaderState::InitShaderRegisterCounts(ResourceBindingTier, ShaderData->ResourceCounts, OutQBSS.RegisterCounts[ShaderVisibility], bAllowUAVs);
+		OutQBSS.bNeedsAgsIntrinsicsSpace |= NeedsAgsIntrinsicsSpace(*ShaderData);
+		OutQBSS.bUseDiagnosticBuffer = ShaderData->UsesDiagnosticBuffer();
+	}
+}
+
 void QuantizeBoundShaderState(
 	const D3D12_RESOURCE_BINDING_TIER& ResourceBindingTier,
 	const FD3D12BoundShaderState* const BSS,
@@ -954,57 +970,14 @@ void QuantizeBoundShaderState(
 	// The objective is to allow a single root signature to represent many bound shader state objects.
 	// The bigger the quantization step sizes, the fewer the root signatures.
 	FMemory::Memzero(&QBSS, sizeof(QBSS));
-	QBSS.bAllowIAInputLayout = BSS->GetVertexDeclaration() != nullptr;	// Does the root signature need access to vertex buffers?
 
-	if (const FD3D12VertexShader* const VertexShader = BSS->GetVertexShader())
-	{
-		FD3D12QuantizedBoundShaderState::InitShaderRegisterCounts(ResourceBindingTier, VertexShader->ResourceCounts, QBSS.RegisterCounts[SV_Vertex]);
-		QBSS.bUseVendorExtension |= UsesVendorExtensionSpace(*VertexShader);
-		QBSS.bUseDiagnosticBuffer = VertexShader->UsesDiagnosticBuffer();
-	}
-
-	if (const FD3D12MeshShader* const MeshShader = BSS->GetMeshShader())
-	{
-		FD3D12QuantizedBoundShaderState::InitShaderRegisterCounts(ResourceBindingTier, MeshShader->ResourceCounts, QBSS.RegisterCounts[SV_Mesh]);
-		QBSS.bUseVendorExtension |= UsesVendorExtensionSpace(*MeshShader);
-		QBSS.bUseDiagnosticBuffer = MeshShader->UsesDiagnosticBuffer();
-	}
-
-	if (const FD3D12AmplificationShader* const AmplificationShader = BSS->GetAmplificationShader())
-	{
-		FD3D12QuantizedBoundShaderState::InitShaderRegisterCounts(ResourceBindingTier, AmplificationShader->ResourceCounts, QBSS.RegisterCounts[SV_Amplification]);
-		QBSS.bUseVendorExtension |= UsesVendorExtensionSpace(*AmplificationShader);
-		QBSS.bUseDiagnosticBuffer = AmplificationShader->UsesDiagnosticBuffer();
-	}
-
-	if (const FD3D12PixelShader* const PixelShader = BSS->GetPixelShader())
-	{
-		FD3D12QuantizedBoundShaderState::InitShaderRegisterCounts(ResourceBindingTier, PixelShader->ResourceCounts, QBSS.RegisterCounts[SV_Pixel], true);
-		QBSS.bUseVendorExtension |= UsesVendorExtensionSpace(*PixelShader);
-		QBSS.bUseDiagnosticBuffer = PixelShader->UsesDiagnosticBuffer();
-	}
-
-	if (const FD3D12GeometryShader* const GeometryShader = BSS->GetGeometryShader())
-	{
-		FD3D12QuantizedBoundShaderState::InitShaderRegisterCounts(ResourceBindingTier, GeometryShader->ResourceCounts, QBSS.RegisterCounts[SV_Geometry]);
-		QBSS.bUseVendorExtension |= UsesVendorExtensionSpace(*GeometryShader);
-		QBSS.bUseDiagnosticBuffer = GeometryShader->UsesDiagnosticBuffer();
-	}
-}
+	QBSS.bAllowIAInputLayout = BSS->GetVertexDeclaration() != nullptr;	// Does the root signature need access to vertex buffers?
 
-static void QuantizeBoundShaderStateCommon(
-	D3D12_RESOURCE_BINDING_TIER ResourceBindingTier,
-	const FShaderCodePackedResourceCounts& ResourceCounts,
-	EShaderVisibility ShaderVisibility,
-	bool bAllowUAVs,
-	FD3D12QuantizedBoundShaderState &OutQBSS
-)
-{
-	// BSS quantizer. There is a 1:1 mapping of quantized bound shader state objects to root signatures.
-	// The objective is to allow a single root signature to represent many bound shader state objects.
-	// The bigger the quantization step sizes, the fewer the root signatures.
-	FMemory::Memzero(&OutQBSS, sizeof(OutQBSS));
-	FD3D12QuantizedBoundShaderState::InitShaderRegisterCounts(ResourceBindingTier, ResourceCounts, OutQBSS.RegisterCounts[ShaderVisibility], bAllowUAVs);
+	QuantizeBoundShaderStateCommon(QBSS, BSS->GetVertexShader(),        ResourceBindingTier, SV_Vertex);
+	QuantizeBoundShaderStateCommon(QBSS, BSS->GetMeshShader(),          ResourceBindingTier, SV_Mesh);
+	QuantizeBoundShaderStateCommon(QBSS, BSS->GetAmplificationShader(), ResourceBindingTier, SV_Amplification);
+	QuantizeBoundShaderStateCommon(QBSS, BSS->GetPixelShader(),         ResourceBindingTier, SV_Pixel, true /*bAllowUAVs*/);
+	QuantizeBoundShaderStateCommon(QBSS, BSS->GetGeometryShader(),      ResourceBindingTier, SV_Geometry);
 }
 
 void QuantizeBoundShaderState(
@@ -1014,13 +987,15 @@ void QuantizeBoundShaderState(
 	)
 {
 	check(ComputeShader);
-	const bool bAllosUAVs = true;
-	QuantizeBoundShaderStateCommon(ResourceBindingTier, ComputeShader->ResourceCounts, SV_All, bAllosUAVs, OutQBSS);
+
+	// BSS quantizer. There is a 1:1 mapping of quantized bound shader state objects to root signatures.
+	// The objective is to allow a single root signature to represent many bound shader state objects.
+	// The bigger the quantization step sizes, the fewer the root signatures.
+	FMemory::Memzero(&OutQBSS, sizeof(OutQBSS));
+
+	QuantizeBoundShaderStateCommon(OutQBSS, ComputeShader, ResourceBindingTier, SV_All, true /*bAllowUAVs*/);
+
 	check(OutQBSS.bAllowIAInputLayout == false); // No access to vertex buffers needed
-#if D3D12RHI_NEEDS_VENDOR_EXTENSIONS
-	OutQBSS.bUseVendorExtension = (ComputeShader->VendorExtensions.Num() > 0);
-#endif
-	OutQBSS.bUseDiagnosticBuffer = ComputeShader->UsesDiagnosticBuffer();
 }
 
 #if D3D12_RHI_RAYTRACING
diff --git a/Engine/Source/Runtime/D3D12RHI/Private/Windows/WindowsD3D12PipelineState.cpp b/Engine/Source/Runtime/D3D12RHI/Private/Windows/WindowsD3D12PipelineState.cpp
--- a/Engine/Source/Runtime/D3D12RHI/Private/Windows/WindowsD3D12PipelineState.cpp
+++ b/Engine/Source/Runtime/D3D12RHI/Private/Windows/WindowsD3D12PipelineState.cpp
@@ -28,9 +28,6 @@
 #define D3D12RHI_USE_D3DDISASSEMBLE 1
 #endif
 
-// TODO: Temp workaround using alternate API until NvAPI_D3D12_CreateGraphicsPipelineState is supported by RenderDoc
-#define TEMP_RENDERDOC_WORKAROUND 1
-
 // D3D12RHI PSO file cache doesn't work anymore. Use FPipelineFileCache instead
 static TAutoConsoleVariable<int32> CVarPipelineStateDiskCache(
 	TEXT("D3D12.PSO.DiskCache"),
@@ -98,31 +95,6 @@ FD3D12_MESH_PIPELINE_STATE_STREAM FD3D12_GRAPHICS_PIPELINE_STATE_DESC::MeshPipel
 }
 #endif // PLATFORM_SUPPORTS_MESH_SHADERS
 
-D3D12_GRAPHICS_PIPELINE_STATE_DESC FD3D12_GRAPHICS_PIPELINE_STATE_DESC::GraphicsDescV0() const
-{
-	D3D12_GRAPHICS_PIPELINE_STATE_DESC D = {};
-	D.Flags = this->Flags;
-	D.NodeMask = this->NodeMask;
-	D.pRootSignature = this->pRootSignature;
-	D.InputLayout = this->InputLayout;
-	D.IBStripCutValue = this->IBStripCutValue;
-	D.PrimitiveTopologyType = this->PrimitiveTopologyType;
-	D.VS = this->VS;
-	D.GS = this->GS;
-	D.PS = this->PS;
-	D.BlendState = this->BlendState;
-	D.DepthStencilState = CD3DX12_DEPTH_STENCIL_DESC1(this->DepthStencilState);
-	D.DSVFormat = this->DSVFormat;
-	D.RasterizerState = this->RasterizerState;
-	D.NumRenderTargets = this->RTFormatArray.NumRenderTargets;
-	FMemory::Memcpy(D.RTVFormats, this->RTFormatArray.RTFormats, sizeof(D.RTVFormats));
-	FMemory::Memzero(&D.StreamOutput, sizeof(D.StreamOutput));
-	D.SampleDesc = this->SampleDesc;
-	D.SampleMask = this->SampleMask;
-	D.CachedPSO = this->CachedPSO;
-	return D;
-}
-
 FD3D12_COMPUTE_PIPELINE_STATE_STREAM FD3D12_COMPUTE_PIPELINE_STATE_DESC::PipelineStateStream() const
 {
 	FD3D12_COMPUTE_PIPELINE_STATE_STREAM Stream = {};
@@ -134,17 +106,6 @@ FD3D12_COMPUTE_PIPELINE_STATE_STREAM FD3D12_COMPUTE_PIPELINE_STATE_DESC::Pipelin
 	return Stream;
 }
 
-D3D12_COMPUTE_PIPELINE_STATE_DESC FD3D12_COMPUTE_PIPELINE_STATE_DESC::ComputeDescV0() const
-{
-	D3D12_COMPUTE_PIPELINE_STATE_DESC D = {};
-	D.Flags = this->Flags;
-	D.NodeMask = this->NodeMask;
-	D.pRootSignature = this->pRootSignature;
-	D.CS = this->CS;
-	D.CachedPSO = this->CachedPSO;
-	return D;
-}
-
 void SaveByteCode(D3D12_SHADER_BYTECODE& ByteCode)
 {
 	if (ByteCode.pShaderBytecode)
@@ -769,49 +730,6 @@ static void DumpComputePSO(const FD3D12_COMPUTE_PIPELINE_STATE_DESC& Desc, const
 	UE_LOG(LogD3D12RHI, Warning, TEXT("Failed to create Compute PSO with hash 0x%s:\n%s"), Name, *String);
 }
 
-
-// Thread-safe create graphics/compute pipeline state. Conditionally load/store the PSO using a Pipeline Library.
-static HRESULT CreatePipelineState(ID3D12PipelineState*&PSO, ID3D12Device* Device, const D3D12_GRAPHICS_PIPELINE_STATE_DESC* Desc, ID3D12PipelineLibrary* Library, const TCHAR* Name)
-{
-	HRESULT hr = S_OK;
-	if (Library == nullptr || Library->LoadGraphicsPipeline(Name, Desc, IID_PPV_ARGS(&PSO)) == E_INVALIDARG)
-	{
-		SCOPE_CYCLE_COUNTER(STAT_PSOCreateTime);
-		hr = Device->CreateGraphicsPipelineState(Desc, IID_PPV_ARGS(&PSO));
-	}
-
-	if (Library && SUCCEEDED(hr))
-	{
-		HRESULT r = Library->StorePipeline(Name, PSO);
-		check(r != E_INVALIDARG);
-	}
-
-	return hr;
-}
-
-static HRESULT CreatePipelineState(ID3D12PipelineState*&PSO, ID3D12Device* Device, const D3D12_COMPUTE_PIPELINE_STATE_DESC* Desc, ID3D12PipelineLibrary* Library, const TCHAR* Name)
-{
-	HRESULT hr = S_OK;
-	if (Library == nullptr || Library->LoadComputePipeline(Name, Desc, IID_PPV_ARGS(&PSO)) == E_INVALIDARG)
-	{
-		SCOPE_CYCLE_COUNTER(STAT_PSOCreateTime);
-		hr = Device->CreateComputePipelineState(Desc, IID_PPV_ARGS(&PSO));
-		if (FAILED(hr))
-		{
-			UE_LOG(LogD3D12RHI, Error, TEXT("Failed to create PipelineState with hash %s"), Name);
-		}
-		VERIFYD3D12RESULT(hr);
-	}
-
-	if (Library && SUCCEEDED(hr))
-	{
-		HRESULT r = Library->StorePipeline(Name, PSO);
-		check(r != E_INVALIDARG);
-	}
-
-	return hr;
-}
-
 // Thread-safe create graphics/compute pipeline state. Conditionally load/store the PSO using a Pipeline Library.
 static HRESULT CreatePipelineStateFromStream(ID3D12PipelineState*& PSO, ID3D12Device2* Device, const D3D12_PIPELINE_STATE_STREAM_DESC* Desc, ID3D12PipelineLibrary1* Library, const TCHAR* Name)
 {
@@ -860,7 +778,7 @@ static inline void FastHashName(wchar_t Name[17], uint64 Hash)
 	Name[16] = 0;
 }
 
-static void CreatePipelineStateWrapper(ID3D12PipelineState** PSO, FD3D12Adapter* Adapter, const GraphicsPipelineCreationArgs_POD* CreationArgs, bool bUseStream)
+static void CreatePipelineStateWrapper(ID3D12PipelineState** PSO, FD3D12Adapter* Adapter, const GraphicsPipelineCreationArgs_POD* CreationArgs)
 {
 	// Get the pipeline state name, currently based on the hash.
 	wchar_t Name[17];
@@ -873,12 +791,11 @@ static void CreatePipelineStateWrapper(ID3D12PipelineState** PSO, FD3D12Adapter*
 
 	// Use pipeline streams if the system supports it.
 	ID3D12Device2* const pDevice2 = Adapter->GetD3DDevice2();
+	check(pDevice2);
+
 #if PLATFORM_SUPPORTS_MESH_SHADERS
 	if (CreationArgs->Desc.UsesMeshShaders())
 	{
-		check(pDevice2);
-		checkf(bUseStream, TEXT("Mesh shader was misconfigured, make sure it is not using extensions that force stream creation off"));
-
 		FD3D12_MESH_PIPELINE_STATE_STREAM Stream = CreationArgs->Desc.Desc.MeshPipelineStateStream();
 		const D3D12_PIPELINE_STATE_STREAM_DESC StreamDesc = { sizeof(Stream), &Stream };
 		HRESULT hr = CreatePipelineStateFromStream(*PSO, pDevice2, &StreamDesc, static_cast<ID3D12PipelineLibrary1*>(CreationArgs->Library), Name);	// Static cast to ID3D12PipelineLibrary1 since we already checked for ID3D12Device2.
@@ -897,7 +814,6 @@ static void CreatePipelineStateWrapper(ID3D12PipelineState** PSO, FD3D12Adapter*
 	}
 	else
 #endif // PLATFORM_SUPPORTS_MESH_SHADERS
-	if (pDevice2 && bUseStream)
 	{
 		FD3D12_GRAPHICS_PIPELINE_STATE_STREAM Stream = CreationArgs->Desc.Desc.PipelineStateStream();
 		const D3D12_PIPELINE_STATE_STREAM_DESC StreamDesc = { sizeof(Stream), &Stream };
@@ -915,26 +831,9 @@ static void CreatePipelineStateWrapper(ID3D12PipelineState** PSO, FD3D12Adapter*
 			}
 		}
 	}
-	else
-	{
-		const D3D12_GRAPHICS_PIPELINE_STATE_DESC Desc = CreationArgs->Desc.Desc.GraphicsDescV0();
-		HRESULT hr = CreatePipelineState(*PSO, Adapter->GetD3DDevice(), &Desc, CreationArgs->Library, Name);
-		if (FAILED(hr))
-		{
-			// First check if D3D device removed, hung or out of memory and handle that separately 
-			if (hr == DXGI_ERROR_DEVICE_REMOVED || hr == DXGI_ERROR_DEVICE_HUNG || hr == E_OUTOFMEMORY)
-			{
-				VERIFYD3D12RESULT_EX(hr, pDevice2);
-			}
-			else
-			{
-				DumpGraphicsPSO(CreationArgs->Desc.Desc, Name);
-			}
-		}
-	}
 }
 
-static void CreatePipelineStateWrapper(ID3D12PipelineState** PSO, FD3D12Adapter* Adapter, const ComputePipelineCreationArgs_POD* CreationArgs, bool bUseStream)
+static void CreatePipelineStateWrapper(ID3D12PipelineState** PSO, FD3D12Adapter* Adapter, const ComputePipelineCreationArgs_POD* CreationArgs)
 {
 	// Get the pipeline state name, currently based on the hash.
 	wchar_t Name[17];
@@ -945,23 +844,21 @@ static void CreatePipelineStateWrapper(ID3D12PipelineState** PSO, FD3D12Adapter*
 	SCOPE_LOG_TIME(*CreatePipelineStateMessage, &GD3D12CreatePSOTime);
 #endif
 
-	// Use pipeline streams if the system supports it.
 	ID3D12Device2* const pDevice2 = Adapter->GetD3DDevice2();
-	if (pDevice2 && bUseStream)
+	check(pDevice2);
+
+	FD3D12_COMPUTE_PIPELINE_STATE_STREAM Stream = CreationArgs->Desc.Desc.PipelineStateStream();
+	const D3D12_PIPELINE_STATE_STREAM_DESC StreamDesc = { sizeof(Stream), &Stream };
+
+	HRESULT hr = CreatePipelineStateFromStream(*PSO, pDevice2, &StreamDesc, static_cast<ID3D12PipelineLibrary1*>(CreationArgs->Library), Name);	// Static cast to ID3D12PipelineLibrary1 since we already checked for ID3D12Device2.
+	if (FAILED(hr))
 	{
-		FD3D12_COMPUTE_PIPELINE_STATE_STREAM Stream = CreationArgs->Desc.Desc.PipelineStateStream();
-		const D3D12_PIPELINE_STATE_STREAM_DESC StreamDesc = { sizeof(Stream), &Stream };
-		HRESULT hr = CreatePipelineStateFromStream(*PSO, pDevice2, &StreamDesc, static_cast<ID3D12PipelineLibrary1*>(CreationArgs->Library), Name);	// Static cast to ID3D12PipelineLibrary1 since we already checked for ID3D12Device2.
-		if (FAILED(hr))
+		// First check if D3D device removed, hung or out of memory and handle that separately 
+		if (hr == DXGI_ERROR_DEVICE_REMOVED || hr == DXGI_ERROR_DEVICE_HUNG || hr == E_OUTOFMEMORY)
 		{
-			DumpComputePSO(CreationArgs->Desc.Desc, Name);
+			VERIFYD3D12RESULT_EX(hr, pDevice2);
 		}
-	}
-	else
-	{
-		D3D12_COMPUTE_PIPELINE_STATE_DESC Desc = CreationArgs->Desc.Desc.ComputeDescV0();
-		HRESULT hr = CreatePipelineState(*PSO, Adapter->GetD3DDevice(), &Desc, CreationArgs->Library, Name);
-		if (FAILED(hr))
+		else
 		{
 			DumpComputePSO(CreationArgs->Desc.Desc, Name);
 		}
@@ -969,7 +866,6 @@ static void CreatePipelineStateWrapper(ID3D12PipelineState** PSO, FD3D12Adapter*
 }
 
 #if D3D12RHI_NEEDS_VENDOR_EXTENSIONS
-
 static FORCEINLINE NVAPI_D3D12_PSO_SET_SHADER_EXTENSION_SLOT_DESC GetNVShaderExtensionDesc(uint32 UavSlot)
 {
 	// https://developer.nvidia.com/unlocking-gpu-intrinsics-hlsl
@@ -987,8 +883,57 @@ static FORCEINLINE NVAPI_D3D12_PSO_SET_SHADER_EXTENSION_SLOT_DESC GetNVShaderExt
 	return ShdExtensionDesc;
 }
 
+template<typename TCreationArgs>
+static void CreatePipelineStateWithExtensions(ID3D12PipelineState** PSO, FD3D12Adapter* Adapter, const TCreationArgs* CreationArgs, TArrayView<const FShaderCodeVendorExtension> VendorExtensions)
+{
+	for (const FShaderCodeVendorExtension& Extension : VendorExtensions)
+	{
+		if (Extension.VendorId == 0x10DE) // NVIDIA
+		{
+			if (Extension.Parameter.Type == EShaderParameterType::UAV)
+			{
+				const NVAPI_D3D12_PSO_SET_SHADER_EXTENSION_SLOT_DESC ShdExtensionDesc = GetNVShaderExtensionDesc(Extension.Parameter.BaseIndex);
+
+				NvAPI_Status NvStatus = NvAPI_D3D12_SetNvShaderExtnSlotSpaceLocalThread(
+					Adapter->GetD3DDevice(),
+					ShdExtensionDesc.uavSlot,
+					ShdExtensionDesc.registerSpace
+				);
+				check(NvStatus == NVAPI_OK);
+
+				CreatePipelineStateWrapper(PSO, Adapter, CreationArgs);
+
+				// Reset to default configuration without vendor extensions.
+				NvStatus = NvAPI_D3D12_SetNvShaderExtnSlotSpaceLocalThread(
+					Adapter->GetD3DDevice(),
+					0xFFFFFFFF,
+					0
+				);
+				check(NvStatus == NVAPI_OK);
+
+				return;
+			}
+		}
+		else if (Extension.VendorId == 0x1002) // AMD
+		{
+			// [REDACTED_PROFILE_URL]/AGS_SDK/blob/master/ags_lib/hlsl/ags_shader_intrinsics_dx12.hlsl
+			// No special create override needed, pass through to default:
+			CreatePipelineStateWrapper(PSO, Adapter, CreationArgs);
+			return;
+		}
+		else if (Extension.VendorId == 0x8086) // INTEL
+		{
+			// TODO: [REDACTED_PROFILE_URL]/intel-graphics-compiler/blob/master/inc/IntelExtensions.hlsl
+		}
+	}
+
+	check(false); // Unimplemented extension path
+}
+#endif
+
 static void CreateGraphicsPipelineState(ID3D12PipelineState** PSO, FD3D12Adapter* Adapter, const GraphicsPipelineCreationArgs_POD* CreationArgs)
 {
+#if D3D12RHI_NEEDS_VENDOR_EXTENSIONS
 	if (CreationArgs->Desc.HasVendorExtensions())
 	{
 		// Need to merge extensions across all stages for a single PSO
@@ -997,10 +942,9 @@ static void CreateGraphicsPipelineState(ID3D12PipelineState** PSO, FD3D12Adapter
 	#define MERGE_EXT(Initial) \
 		if (CreationArgs->Desc.##Initial##SExtensions != nullptr) \
 		{ \
-			const TArray<FShaderCodeVendorExtension>& Extensions = *CreationArgs->Desc.##Initial##SExtensions; \
-			for (int32 ExtIndex = 0; ExtIndex < Extensions.Num(); ++ExtIndex) \
+			for (const FShaderCodeVendorExtension& Extension : *CreationArgs->Desc.##Initial##SExtensions) \
 			{ \
-				MergedExtensions.AddUnique(Extensions[ExtIndex]); \
+				MergedExtensions.AddUnique(Extension); \
 			} \
 		}
 
@@ -1011,137 +955,29 @@ static void CreateGraphicsPipelineState(ID3D12PipelineState** PSO, FD3D12Adapter
 		MERGE_EXT(G);
 	#undef MERGE_EXT
 
-		for (int32 ExtIndex = 0; ExtIndex < MergedExtensions.Num(); ++ExtIndex)
-		{
-			const FShaderCodeVendorExtension& Extension = MergedExtensions[ExtIndex];
-			if (Extension.VendorId == 0x10DE) // NVIDIA
-			{
-				if (Extension.Parameter.Type == EShaderParameterType::UAV)
-				{
-					const NVAPI_D3D12_PSO_SET_SHADER_EXTENSION_SLOT_DESC ShdExtensionDesc = GetNVShaderExtensionDesc(Extension.Parameter.BaseIndex);
-				#if TEMP_RENDERDOC_WORKAROUND
-					// TODO: Temp workaround using alternate API until NvAPI_D3D12_CreateGraphicsPipelineState is supported by RenderDoc
-					NvAPI_Status NvStatus = NvAPI_D3D12_SetNvShaderExtnSlotSpaceLocalThread(
-						Adapter->GetD3DDevice(),
-						ShdExtensionDesc.uavSlot,
-						ShdExtensionDesc.registerSpace
-					);
-					check(NvStatus == NVAPI_OK);
-
-					CreatePipelineStateWrapper(PSO, Adapter, CreationArgs, false /* use stream */);
-
-					// Reset to default configuration without vendor extensions.
-					NvStatus = NvAPI_D3D12_SetNvShaderExtnSlotSpaceLocalThread(
-						Adapter->GetD3DDevice(),
-						0xFFFFFFFF,
-						0
-					);
-					check(NvStatus == NVAPI_OK);
-				#else
-					D3D12_GRAPHICS_PIPELINE_STATE_DESC Desc = CreationArgs->Desc.Desc.GraphicsDescV0();
-					const NVAPI_D3D12_PSO_EXTENSION_DESC* NvExtensions[] = { &ShdExtensionDesc };
-					NvAPI_Status NvStatus = NvAPI_D3D12_CreateGraphicsPipelineState(Adapter->GetD3DDevice(), &Desc, ARRAYSIZE(NvExtensions), NvExtensions, PSO);
-					check(NvStatus == NVAPI_OK);
-				#endif
-					return;
-				}
-			}
-			else if (Extension.VendorId == 0x1002) // AMD
-			{
-				// [REDACTED_PROFILE_URL]/AGS_SDK/blob/master/ags_lib/hlsl/ags_shader_intrinsics_dx12.hlsl
-				// No special create override needed, pass through to default:
-				CreatePipelineStateWrapper(PSO, Adapter, CreationArgs, false /* use stream */);
-				return;
-			}
-			else if (Extension.VendorId == 0x8086) // INTEL
-			{
-				// TODO: [REDACTED_PROFILE_URL]/intel-graphics-compiler/blob/master/inc/IntelExtensions.hlsl
-			}
-		}
-
-		check(false); // Unimplemented extension path
+		CreatePipelineStateWithExtensions(PSO, Adapter, CreationArgs, MergedExtensions);
 	}
 	else
+#endif // D3D12RHI_NEEDS_VENDOR_EXTENSIONS
 	{
-		CreatePipelineStateWrapper(PSO, Adapter, CreationArgs, true /* use stream */);
+		CreatePipelineStateWrapper(PSO, Adapter, CreationArgs);
 	}
 }
 
 static void CreateComputePipelineState(ID3D12PipelineState** PSO, FD3D12Adapter* Adapter, const ComputePipelineCreationArgs_POD* CreationArgs)
 {
+#if D3D12RHI_NEEDS_VENDOR_EXTENSIONS
 	if (CreationArgs->Desc.HasVendorExtensions())
 	{
-		const TArray<FShaderCodeVendorExtension>& Extensions = *CreationArgs->Desc.Extensions;
-		for (int32 ExtIndex = 0; ExtIndex < Extensions.Num(); ++ExtIndex)
-		{
-			const FShaderCodeVendorExtension& Extension = Extensions[ExtIndex];
-			if (Extension.VendorId == 0x10DE) // NVIDIA
-			{
-				if (Extension.Parameter.Type == EShaderParameterType::UAV)
-				{
-					const NVAPI_D3D12_PSO_SET_SHADER_EXTENSION_SLOT_DESC ShdExtensionDesc = GetNVShaderExtensionDesc(Extension.Parameter.BaseIndex);
-				#if TEMP_RENDERDOC_WORKAROUND
-					// TODO: Temp workaround using alternate API until NvAPI_D3D12_CreateGraphicsPipelineState is supported by RenderDoc
-					NvAPI_Status NvStatus = NvAPI_D3D12_SetNvShaderExtnSlotSpaceLocalThread(
-						Adapter->GetD3DDevice(),
-						ShdExtensionDesc.uavSlot,
-						ShdExtensionDesc.registerSpace
-					);
-					check(NvStatus == NVAPI_OK);
-
-					CreatePipelineStateWrapper(PSO, Adapter, CreationArgs, false /* use stream */);
-
-					// Reset to default configuration without vendor extensions.
-					NvStatus = NvAPI_D3D12_SetNvShaderExtnSlotSpaceLocalThread(
-						Adapter->GetD3DDevice(),
-						0xFFFFFFFF,
-						0
-					);
-					check(NvStatus == NVAPI_OK);
-				#else
-					D3D12_COMPUTE_PIPELINE_STATE_DESC Desc = CreationArgs->Desc.Desc.ComputeDescV0();
-					const NVAPI_D3D12_PSO_EXTENSION_DESC* NvExtensions[] = { &ShdExtensionDesc };
-					NvAPI_Status NvStatus = NvAPI_D3D12_CreateComputePipelineState(Adapter->GetD3DDevice(), &Desc, ARRAYSIZE(NvExtensions), NvExtensions, PSO);
-					check(NvStatus == NVAPI_OK);
-				#endif
-					return;
-				}
-			}
-			else if (Extension.VendorId == 0x1002) // AMD
-			{
-				// [REDACTED_PROFILE_URL]/AGS_SDK/blob/master/ags_lib/hlsl/ags_shader_intrinsics_dx12.hlsl
-				// No special create override needed, pass through to default:
-				CreatePipelineStateWrapper(PSO, Adapter, CreationArgs, false /* use stream */);
-				return;
-			}
-			else if (Extension.VendorId == 0x8086) // INTEL
-			{
-				// TODO: [REDACTED_PROFILE_URL]/intel-graphics-compiler/blob/master/inc/IntelExtensions.hlsl
-			}
-		}
-
-		check(false); // Unimplemented extension path
+		CreatePipelineStateWithExtensions(PSO, Adapter, CreationArgs, TArrayView<const FShaderCodeVendorExtension>(*CreationArgs->Desc.Extensions));
 	}
 	else
+#endif // D3D12RHI_NEEDS_VENDOR_EXTENSIONS
 	{
-		CreatePipelineStateWrapper(PSO, Adapter, CreationArgs, true /* use stream */);
+		CreatePipelineStateWrapper(PSO, Adapter, CreationArgs);
 	}
 }
 
-#else
-
-static FORCEINLINE void CreateGraphicsPipelineState(ID3D12PipelineState** PSO, FD3D12Adapter* Adapter, const GraphicsPipelineCreationArgs_POD* CreationArgs)
-{
-	CreatePipelineStateWrapper(PSO, Adapter, CreationArgs, false /* use stream */);
-}
-
-static FORCEINLINE void CreateComputePipelineState(ID3D12PipelineState** PSO, FD3D12Adapter* Adapter, const ComputePipelineCreationArgs_POD* CreationArgs)
-{
-	CreatePipelineStateWrapper(PSO, Adapter, CreationArgs, false /* use stream */);
-}
-
-#endif
-
 void FD3D12PipelineState::Create(const ComputePipelineCreationArgs& InCreationArgs)
 {
 	check(PipelineState.GetReference() == nullptr);
diff --git a/Engine/Source/Runtime/D3D12RHI/Public/D3D12Util.h b/Engine/Source/Runtime/D3D12RHI/Public/D3D12Util.h
--- a/Engine/Source/Runtime/D3D12RHI/Public/D3D12Util.h
+++ b/Engine/Source/Runtime/D3D12RHI/Public/D3D12Util.h
@@ -91,7 +91,7 @@ struct FD3D12QuantizedBoundShaderState
 	FShaderRegisterCounts RegisterCounts[SV_ShaderVisibilityCount];
 	ERTRootSignatureType RootSignatureType = RS_Raster;
 	uint8 bAllowIAInputLayout : 1;
-	uint8 bUseVendorExtension : 1;
+	uint8 bNeedsAgsIntrinsicsSpace : 1;
 	uint8 bUseDiagnosticBuffer : 1;
 	uint8 Padding : 5;
 
diff --git a/Engine/Source/Runtime/Windows/D3D11RHI/Private/D3D11Shaders.cpp b/Engine/Source/Runtime/Windows/D3D11RHI/Private/D3D11Shaders.cpp
--- a/Engine/Source/Runtime/Windows/D3D11RHI/Private/D3D11Shaders.cpp
+++ b/Engine/Source/Runtime/Windows/D3D11RHI/Private/D3D11Shaders.cpp
@@ -54,13 +54,17 @@ static inline void ReadShaderOptionalData(FShaderCodeReader& InShaderCode, TShad
 	OutShader.bIsSm6Shader = IsSm6Shader && IsSm6ShaderSize && *IsSm6Shader;
 }
 
-static bool ApplyVendorExtensions(ID3D11Device* Direct3DDevice, EShaderFrequency Frequency, const TArray<FShaderCodeVendorExtension>& VendorExtensions)
+static bool ApplyVendorExtensions(ID3D11Device* Direct3DDevice, EShaderFrequency Frequency, const FD3D11ShaderData* ShaderData, bool& OutNeedsReset)
 {
+	if (ShaderData->bIsSm6Shader)
+	{
+		return false;
+	}
+
 	bool IsValidHardwareExtension = true;
 #if !PLATFORM_HOLOLENS
-	for (int32 ExtensionIndex = 0; ExtensionIndex < VendorExtensions.Num(); ++ExtensionIndex)
+	for (const FShaderCodeVendorExtension& Extension : ShaderData->VendorExtensions)
 	{
-		const FShaderCodeVendorExtension& Extension = VendorExtensions[ExtensionIndex];
 		if (Extension.VendorId == 0x10DE) // NVIDIA
 		{
 			if (!IsRHIDeviceNVIDIA())
@@ -73,6 +77,7 @@ static bool ApplyVendorExtensions(ID3D11Device* Direct3DDevice, EShaderFrequency
 			if (Extension.Parameter.Type == EShaderParameterType::UAV)
 			{
 				NvAPI_D3D11_SetNvShaderExtnSlot(Direct3DDevice, Extension.Parameter.BaseIndex);
+				OutNeedsReset = true;
 			}
 		}
 		else if (Extension.VendorId == 0x1002) // AMD
@@ -98,32 +103,17 @@ static bool ApplyVendorExtensions(ID3D11Device* Direct3DDevice, EShaderFrequency
 	return IsValidHardwareExtension;
 }
 
-static void ResetVendorExtensions(ID3D11Device* Direct3DDevice, EShaderFrequency Frequency, const TArray<FShaderCodeVendorExtension>& VendorExtensions)
+static void ResetVendorExtensions(ID3D11Device* Direct3DDevice)
 {
 #if !PLATFORM_HOLOLENS
-	for (int32 ExtensionIndex = 0; ExtensionIndex < VendorExtensions.Num(); ++ExtensionIndex)
+	if (IsRHIDeviceNVIDIA())
 	{
-		const FShaderCodeVendorExtension& Extension = VendorExtensions[ExtensionIndex];
-		if (Extension.VendorId == 0x10DE)
-		{
-			// NVIDIA
-			if (Extension.Parameter.Type == EShaderParameterType::UAV)
-			{
-				NvAPI_D3D11_SetNvShaderExtnSlot(Direct3DDevice, ~uint32(0));
-			}
-		}
-		else if (Extension.VendorId == 0x1002) // AMD
-		{
-		}
-		else if (Extension.VendorId == 0x8086) // Intel
-		{
-		}
+		NvAPI_D3D11_SetNvShaderExtnSlot(Direct3DDevice, ~uint32(0));
 	}
 #endif
 }
 
-template <typename TShaderType>
-static inline void InitUniformBufferStaticSlots(TShaderType* Shader)
+static void InitUniformBufferStaticSlots(FD3D11ShaderData* Shader)
 {
 	const FBaseShaderResourceTable& SRT = Shader->ShaderResourceTable;
 
@@ -145,20 +135,24 @@ static inline void InitUniformBufferStaticSlots(TShaderType* Shader)
 FVertexShaderRHIRef FD3D11DynamicRHI::RHICreateVertexShader(TArrayView<const uint8> Code, const FSHAHash& Hash)
 {
 	FShaderCodeReader ShaderCode(Code);
-
 	FD3D11VertexShader* Shader = new FD3D11VertexShader;
 
 	FMemoryReaderView Ar( Code, true );
 	Ar << Shader->ShaderResourceTable;
 	int32 Offset = Ar.Tell();
-	const uint8* CodePtr = Code.GetData() + Offset;
-	const size_t CodeSize = ShaderCode.GetActualShaderCodeSize() - Offset;
+
+	TArrayView<const uint8> ActualCode = ShaderCode.GetOffsetShaderCode(Offset);
 
 	ReadShaderOptionalData(ShaderCode, *Shader);
-	if (!Shader->bIsSm6Shader && ApplyVendorExtensions(Direct3DDevice, SF_Vertex, Shader->VendorExtensions))
+
+	bool bNeedsReset = false;
+	if (ApplyVendorExtensions(Direct3DDevice, SF_Vertex, Shader, bNeedsReset))
 	{
-		VERIFYD3D11SHADERRESULT(Direct3DDevice->CreateVertexShader((void*)CodePtr, CodeSize, nullptr, Shader->Resource.GetInitReference()), Shader, Direct3DDevice);
-		ResetVendorExtensions(Direct3DDevice, SF_Vertex, Shader->VendorExtensions);
+		VERIFYD3D11SHADERRESULT(Direct3DDevice->CreateVertexShader(ActualCode.GetData(), ActualCode.Num(), nullptr, Shader->Resource.GetInitReference()), Shader, Direct3DDevice);
+		if (bNeedsReset)
+		{
+			ResetVendorExtensions(Direct3DDevice);
+		}
 		InitUniformBufferStaticSlots(Shader);
 	}
 	
@@ -179,20 +173,24 @@ FVertexShaderRHIRef FD3D11DynamicRHI::CreateVertexShader_RenderThread(
 FGeometryShaderRHIRef FD3D11DynamicRHI::RHICreateGeometryShader(TArrayView<const uint8> Code, const FSHAHash& Hash)
 { 
 	FShaderCodeReader ShaderCode(Code);
-
 	FD3D11GeometryShader* Shader = new FD3D11GeometryShader;
 
 	FMemoryReaderView Ar( Code, true );
 	Ar << Shader->ShaderResourceTable;
 	int32 Offset = Ar.Tell();
-	const uint8* CodePtr = Code.GetData() + Offset;
-	const size_t CodeSize = ShaderCode.GetActualShaderCodeSize() - Offset;
+
+	TArrayView<const uint8> ActualCode = ShaderCode.GetOffsetShaderCode(Offset);
 
 	ReadShaderOptionalData(ShaderCode, *Shader);
-	if (!Shader->bIsSm6Shader && ApplyVendorExtensions(Direct3DDevice, SF_Geometry, Shader->VendorExtensions))
+
+	bool bNeedsReset = false;
+	if (ApplyVendorExtensions(Direct3DDevice, SF_Geometry, Shader, bNeedsReset))
 	{
-		VERIFYD3D11SHADERRESULT(Direct3DDevice->CreateGeometryShader((void*)CodePtr, CodeSize, nullptr, Shader->Resource.GetInitReference()), Shader, Direct3DDevice);
-		ResetVendorExtensions(Direct3DDevice, SF_Geometry, Shader->VendorExtensions);
+		VERIFYD3D11SHADERRESULT(Direct3DDevice->CreateGeometryShader(ActualCode.GetData(), ActualCode.Num(), nullptr, Shader->Resource.GetInitReference()), Shader, Direct3DDevice);
+		if (bNeedsReset)
+		{
+			ResetVendorExtensions(Direct3DDevice);
+		}
 		InitUniformBufferStaticSlots(Shader);
 	}
 
@@ -209,20 +207,24 @@ FGeometryShaderRHIRef FD3D11DynamicRHI::CreateGeometryShader_RenderThread(
 FPixelShaderRHIRef FD3D11DynamicRHI::RHICreatePixelShader(TArrayView<const uint8> Code, const FSHAHash& Hash)
 {
 	FShaderCodeReader ShaderCode(Code);
-
 	FD3D11PixelShader* Shader = new FD3D11PixelShader;
 
 	FMemoryReaderView Ar( Code, true );
 	Ar << Shader->ShaderResourceTable;
 	int32 Offset = Ar.Tell();
-	const uint8* CodePtr = Code.GetData() + Offset;
-	const size_t CodeSize = ShaderCode.GetActualShaderCodeSize() - Offset;
+
+	TArrayView<const uint8> ActualCode = ShaderCode.GetOffsetShaderCode(Offset);
 
 	ReadShaderOptionalData(ShaderCode, *Shader);
-	if (!Shader->bIsSm6Shader && ApplyVendorExtensions(Direct3DDevice, SF_Pixel, Shader->VendorExtensions))
+
+	bool bNeedsReset = false;
+	if (ApplyVendorExtensions(Direct3DDevice, SF_Pixel, Shader, bNeedsReset))
 	{
-		VERIFYD3D11SHADERRESULT(Direct3DDevice->CreatePixelShader((void*)CodePtr, CodeSize, nullptr, Shader->Resource.GetInitReference()), Shader, Direct3DDevice);
-		ResetVendorExtensions(Direct3DDevice, SF_Pixel, Shader->VendorExtensions);
+		VERIFYD3D11SHADERRESULT(Direct3DDevice->CreatePixelShader(ActualCode.GetData(), ActualCode.Num(), nullptr, Shader->Resource.GetInitReference()), Shader, Direct3DDevice);
+		if (bNeedsReset)
+		{
+			ResetVendorExtensions(Direct3DDevice);
+		}
 		InitUniformBufferStaticSlots(Shader);
 	}
 
@@ -239,20 +241,24 @@ FPixelShaderRHIRef FD3D11DynamicRHI::CreatePixelShader_RenderThread(
 FComputeShaderRHIRef FD3D11DynamicRHI::RHICreateComputeShader(TArrayView<const uint8> Code, const FSHAHash& Hash)
 { 
 	FShaderCodeReader ShaderCode(Code);
-
 	FD3D11ComputeShader* Shader = new FD3D11ComputeShader;
 
 	FMemoryReaderView Ar( Code, true );
 	Ar << Shader->ShaderResourceTable;
 	int32 Offset = Ar.Tell();
-	const uint8* CodePtr = Code.GetData() + Offset;
-	const size_t CodeSize = ShaderCode.GetActualShaderCodeSize() - Offset;
+
+	TArrayView<const uint8> ActualCode = ShaderCode.GetOffsetShaderCode(Offset);
 
 	ReadShaderOptionalData(ShaderCode, *Shader);
-	if (!Shader->bIsSm6Shader && ApplyVendorExtensions(Direct3DDevice, SF_Compute, Shader->VendorExtensions))
+
+	bool bNeedsReset = false;
+	if (ApplyVendorExtensions(Direct3DDevice, SF_Compute, Shader, bNeedsReset))
 	{
-		VERIFYD3D11SHADERRESULT(Direct3DDevice->CreateComputeShader((void*)CodePtr, CodeSize, nullptr, Shader->Resource.GetInitReference()), Shader, Direct3DDevice);
-		ResetVendorExtensions(Direct3DDevice, SF_Compute, Shader->VendorExtensions);
+		VERIFYD3D11SHADERRESULT(Direct3DDevice->CreateComputeShader(ActualCode.GetData(), ActualCode.Num(), nullptr, Shader->Resource.GetInitReference()), Shader, Direct3DDevice);
+		if (bNeedsReset)
+		{
+			ResetVendorExtensions(Direct3DDevice);
+		}
 		InitUniformBufferStaticSlots(Shader);
 	}
 
