Adding Dynamic Resource heap flags to Root Signatures.
Fixing issue where bUseDiagnosticBuffer would be overwritten by the last used stage.
Moving static Root Signature text generation to a public file.

diff --git a/Engine/Source/Runtime/D3D12RHI/Private/D3D12RootSignature.cpp b/Engine/Source/Runtime/D3D12RHI/Private/D3D12RootSignature.cpp
--- a/Engine/Source/Runtime/D3D12RHI/Private/D3D12RootSignature.cpp
+++ b/Engine/Source/Runtime/D3D12RHI/Private/D3D12RootSignature.cpp
@@ -282,6 +282,18 @@ FD3D12RootSignatureDesc::FD3D12RootSignatureDesc(const FD3D12QuantizedBoundShade
 
 	D3D12_ROOT_SIGNATURE_FLAGS Flags = D3D12_ROOT_SIGNATURE_FLAG_NONE;
 
+#if !defined(D3D12RHI_DYNAMIC_RESOURCES_UNSUPPORTED)
+	if (QBSS.bUseDirectlyIndexedResourceHeap)
+	{
+		Flags |= D3D12_ROOT_SIGNATURE_FLAG_CBV_SRV_UAV_HEAP_DIRECTLY_INDEXED;
+	}
+
+	if (QBSS.bUseDirectlyIndexedSamplerHeap)
+	{
+		Flags |= D3D12_ROOT_SIGNATURE_FLAG_SAMPLER_HEAP_DIRECTLY_INDEXED;
+	}
+#endif // !defined(D3D12RHI_DYNAMIC_RESOURCES_UNSUPPORTED)
+
 #if D3D12_RHI_RAYTRACING
 	if (QBSS.RootSignatureType == RS_RayTracingLocal)
 	{
diff --git a/Engine/Source/Runtime/D3D12RHI/Private/D3D12Util.cpp b/Engine/Source/Runtime/D3D12RHI/Private/D3D12Util.cpp
--- a/Engine/Source/Runtime/D3D12RHI/Private/D3D12Util.cpp
+++ b/Engine/Source/Runtime/D3D12RHI/Private/D3D12Util.cpp
@@ -944,6 +944,16 @@ bool NeedsAgsIntrinsicsSpace(const FD3D12ShaderData& ShaderData)
 	return false;
 }
 
+static void SetBoundShaderStateFlags(FD3D12QuantizedBoundShaderState& OutQBSS, const FD3D12ShaderData* ShaderData)
+{
+	if (ShaderData)
+	{
+		OutQBSS.bUseDiagnosticBuffer |= ShaderData->UsesDiagnosticBuffer();
+		OutQBSS.bUseDirectlyIndexedResourceHeap |= ShaderData->UsesBindlessResources();
+		OutQBSS.bUseDirectlyIndexedSamplerHeap |= ShaderData->UsesBindlessSamplers();
+	}
+}
+
 static void QuantizeBoundShaderStateCommon(
 	FD3D12QuantizedBoundShaderState& OutQBSS,
 	const FD3D12ShaderData* ShaderData,
@@ -956,8 +966,9 @@ static void QuantizeBoundShaderStateCommon(
 	{
 		FD3D12QuantizedBoundShaderState::InitShaderRegisterCounts(ResourceBindingTier, ShaderData->ResourceCounts, OutQBSS.RegisterCounts[ShaderVisibility], bAllowUAVs);
 		OutQBSS.bNeedsAgsIntrinsicsSpace |= NeedsAgsIntrinsicsSpace(*ShaderData);
-		OutQBSS.bUseDiagnosticBuffer = ShaderData->UsesDiagnosticBuffer();
 	}
+
+	SetBoundShaderStateFlags(OutQBSS, ShaderData);
 }
 
 void QuantizeBoundShaderState(
@@ -1063,10 +1074,7 @@ void QuantizeBoundShaderState(
 		checkNoEntry(); // Unexpected shader target frequency
 	}
 
-	if (RayTracingShader)
-	{
-		OutQBSS.bUseDiagnosticBuffer = RayTracingShader->UsesDiagnosticBuffer();
-	}
+	SetBoundShaderStateFlags(OutQBSS, RayTracingShader);
 }
 #endif // D3D12_RHI_RAYTRACING
 
diff --git a/Engine/Source/Runtime/D3D12RHI/Public/D3D12RootSignatureDefinitions.h b/Engine/Source/Runtime/D3D12RHI/Public/D3D12RootSignatureDefinitions.h
--- a/Engine/Source/Runtime/D3D12RHI/Public/D3D12RootSignatureDefinitions.h
+++ b/Engine/Source/Runtime/D3D12RHI/Public/D3D12RootSignatureDefinitions.h
@@ -6,8 +6,6 @@ D3D12RootSignatureDefinitions.h: D3D12 utilities for Root signatures.
 
 #pragma once
 
-#include "D3D12RHI.h"
-
 namespace D3D12ShaderUtils
 {
 	namespace StaticRootSignatureConstants
@@ -19,110 +17,181 @@ namespace D3D12ShaderUtils
 		const D3D12_DESCRIPTOR_RANGE_FLAGS SamplerDescriptorRangeFlags = D3D12_DESCRIPTOR_RANGE_FLAG_DESCRIPTORS_VOLATILE;
 	}
 
+	enum class ERootSignatureType
+	{
+		CBV,
+		SRV,
+		UAV,
+		Sampler,
+	};
+
 	// Simple base class to help write out a root signature (subclass to generate either to a binary struct or a #define)
 	struct FRootSignatureCreator
 	{
 		virtual ~FRootSignatureCreator() { }
-		virtual void AddRootFlag(D3D12_ROOT_SIGNATURE_FLAGS Flag) = 0;
-		enum EType
-		{
-			CBV,
-			SRV,
-			UAV,
-			Sampler,
-		};
-		virtual void AddTable(EShaderFrequency Stage, EType Type, int32 NumDescriptors/*, bool bAppend*/) = 0;
-		virtual void Reset() = 0;
 
-		static D3D12_DESCRIPTOR_RANGE_TYPE GetD3D12DescriptorRangeType(EType Type)
+		virtual FRootSignatureCreator& Reset() = 0;
+		virtual FRootSignatureCreator& AddRootFlag(D3D12_ROOT_SIGNATURE_FLAGS Flag) = 0;
+		virtual FRootSignatureCreator& AddTable(EShaderFrequency Stage, ERootSignatureType Type, int32 NumDescriptors/*, bool bAppend*/) = 0;
+
+		inline D3D12_DESCRIPTOR_RANGE_TYPE GetD3D12DescriptorRangeType(ERootSignatureType Type)
 		{
 			switch (Type)
 			{
-			case D3D12ShaderUtils::FRootSignatureCreator::SRV:
+			case ERootSignatureType::SRV:
 				return D3D12_DESCRIPTOR_RANGE_TYPE_SRV;
-			case D3D12ShaderUtils::FRootSignatureCreator::UAV:
+			case ERootSignatureType::UAV:
 				return D3D12_DESCRIPTOR_RANGE_TYPE_UAV;
-			case D3D12ShaderUtils::FRootSignatureCreator::Sampler:
+			case ERootSignatureType::Sampler:
 				return D3D12_DESCRIPTOR_RANGE_TYPE_SAMPLER;
 			default:
 				return D3D12_DESCRIPTOR_RANGE_TYPE_CBV;
 			}
 		}
 
-		static D3D12_DESCRIPTOR_RANGE_FLAGS GetD3D12DescriptorRangeFlags(EType Type)
+		inline D3D12_DESCRIPTOR_RANGE_FLAGS GetD3D12DescriptorRangeFlags(ERootSignatureType Type)
 		{
 			switch (Type)
 			{
-			case D3D12ShaderUtils::FRootSignatureCreator::SRV:
+			case ERootSignatureType::SRV:
 				return D3D12ShaderUtils::StaticRootSignatureConstants::SRVDescriptorRangeFlags;
-			case D3D12ShaderUtils::FRootSignatureCreator::CBV:
+			case ERootSignatureType::CBV:
 				return D3D12ShaderUtils::StaticRootSignatureConstants::CBVDescriptorRangeFlags;
-			case D3D12ShaderUtils::FRootSignatureCreator::UAV:
+			case ERootSignatureType::UAV:
 				return D3D12ShaderUtils::StaticRootSignatureConstants::UAVDescriptorRangeFlags;
-			case D3D12ShaderUtils::FRootSignatureCreator::Sampler:
+			case ERootSignatureType::Sampler:
 				return D3D12ShaderUtils::StaticRootSignatureConstants::SamplerDescriptorRangeFlags;
 			default:
 				return D3D12_DESCRIPTOR_RANGE_FLAG_NONE;
 			}
 		}
+
+		inline D3D12_SHADER_VISIBILITY GetD3D12ShaderVisibility(EShaderFrequency Stage)
+		{
+			switch (Stage)
+			{
+			case SF_Vertex:
+				return D3D12_SHADER_VISIBILITY_VERTEX;
+			case SF_Pixel:
+				return D3D12_SHADER_VISIBILITY_PIXEL;
+			case SF_Geometry:
+				return D3D12_SHADER_VISIBILITY_GEOMETRY;
+			default:
+				return D3D12_SHADER_VISIBILITY_ALL;
+			}
+		}
+
+		inline const TCHAR* GetVisibilityFlag(EShaderFrequency Stage)
+		{
+			switch (Stage)
+			{
+			case SF_Vertex:
+				return TEXT("SHADER_VISIBILITY_VERTEX");
+			case SF_Geometry:
+				return TEXT("SHADER_VISIBILITY_GEOMETRY");
+			case SF_Pixel:
+				return TEXT("SHADER_VISIBILITY_PIXEL");
+			case SF_Mesh:
+				return TEXT("SHADER_VISIBILITY_MESH");
+			case SF_Amplification:
+				return TEXT("SHADER_VISIBILITY_AMPLIFICATION");
+			default:
+				return TEXT("SHADER_VISIBILITY_ALL");
+			}
+		};
+
+		inline const TCHAR* GetTypePrefix(ERootSignatureType Type)
+		{
+			switch (Type)
+			{
+			case ERootSignatureType::SRV:
+				return TEXT("SRV(t");
+			case ERootSignatureType::UAV:
+				return TEXT("UAV(u");
+			case ERootSignatureType::Sampler:
+				return TEXT("Sampler(s");
+			default:
+				return TEXT("CBV(b");
+			}
+		}
+
+		inline TCHAR* GetFlagName(D3D12_ROOT_SIGNATURE_FLAGS Flag)
+		{
+			switch (Flag)
+			{
+			case D3D12_ROOT_SIGNATURE_FLAG_ALLOW_INPUT_ASSEMBLER_INPUT_LAYOUT:
+				return TEXT("ALLOW_INPUT_ASSEMBLER_INPUT_LAYOUT");
+			case D3D12_ROOT_SIGNATURE_FLAG_DENY_VERTEX_SHADER_ROOT_ACCESS:
+				return TEXT("DENY_VERTEX_SHADER_ROOT_ACCESS");
+			case D3D12_ROOT_SIGNATURE_FLAG_DENY_GEOMETRY_SHADER_ROOT_ACCESS:
+				return TEXT("DENY_GEOMETRY_SHADER_ROOT_ACCESS");
+			case D3D12_ROOT_SIGNATURE_FLAG_DENY_PIXEL_SHADER_ROOT_ACCESS:
+				return TEXT("DENY_PIXEL_SHADER_ROOT_ACCESS");
+			case D3D12_ROOT_SIGNATURE_FLAG_ALLOW_STREAM_OUTPUT:
+				return TEXT("ALLOW_STREAM_OUTPUT");
+
+#if !defined(D3D12RHI_TOOLS_MESH_SHADERS_UNSUPPORTED)
+			case D3D12_ROOT_SIGNATURE_FLAG_DENY_AMPLIFICATION_SHADER_ROOT_ACCESS:
+				return TEXT("DENY_AMPLIFICATION_SHADER_ROOT_ACCESS");
+			case D3D12_ROOT_SIGNATURE_FLAG_DENY_MESH_SHADER_ROOT_ACCESS:
+				return TEXT("DENY_MESH_SHADER_ROOT_ACCESS");
+#endif
+
+#if !defined(D3D12RHI_DYNAMIC_RESOURCES_UNSUPPORTED)
+			case D3D12_ROOT_SIGNATURE_FLAG_CBV_SRV_UAV_HEAP_DIRECTLY_INDEXED:
+				return TEXT("CBV_SRV_UAV_HEAP_DIRECTLY_INDEXED");
+			case D3D12_ROOT_SIGNATURE_FLAG_SAMPLER_HEAP_DIRECTLY_INDEXED:
+				return TEXT("SAMPLER_HEAP_DIRECTLY_INDEXED");
+#endif
+
+			default:
+				break;
+			}
+
+			return TEXT("0");
+		};
 	};
 
 	// Fat/Static Gfx Root Signature
-	inline void CreateGfxRootSignature(FRootSignatureCreator* Creator, bool bAllowMeshShaders)
+	inline void CreateGfxRootSignature(FRootSignatureCreator& Creator, bool bAllowMeshShaders)
 	{
 		// Ensure the creator starts in a clean state (in cases of creator reuse, etc.).
-		Creator->Reset();
-
-		Creator->AddRootFlag(D3D12_ROOT_SIGNATURE_FLAG_ALLOW_INPUT_ASSEMBLER_INPUT_LAYOUT);
-		//Creator->AddRootFlag(D3D12_ROOT_SIGNATURE_FLAG_DENY_HULL_SHADER_ROOT_ACCESS);
-		//Creator->AddRootFlag(D3D12_ROOT_SIGNATURE_FLAG_DENY_DOMAIN_SHADER_ROOT_ACCESS);
-		Creator->AddTable(SF_Pixel, FRootSignatureCreator::SRV, MAX_SRVS);
-		Creator->AddTable(SF_Pixel, FRootSignatureCreator::CBV, MAX_CBS);
-		Creator->AddTable(SF_Pixel, FRootSignatureCreator::Sampler, MAX_SAMPLERS);
-		Creator->AddTable(SF_Vertex, FRootSignatureCreator::SRV, MAX_SRVS);
-		Creator->AddTable(SF_Vertex, FRootSignatureCreator::CBV, MAX_CBS);
-		Creator->AddTable(SF_Vertex, FRootSignatureCreator::Sampler, MAX_SAMPLERS);
-		Creator->AddTable(SF_Geometry, FRootSignatureCreator::SRV, MAX_SRVS);
-		Creator->AddTable(SF_Geometry, FRootSignatureCreator::CBV, MAX_CBS);
-		Creator->AddTable(SF_Geometry, FRootSignatureCreator::Sampler, MAX_SAMPLERS);
+		Creator
+			.Reset()
+			.AddRootFlag(D3D12_ROOT_SIGNATURE_FLAG_ALLOW_INPUT_ASSEMBLER_INPUT_LAYOUT)
+			.AddTable(SF_Pixel, ERootSignatureType::SRV, MAX_SRVS)
+			.AddTable(SF_Pixel, ERootSignatureType::CBV, MAX_CBS)
+			.AddTable(SF_Pixel, ERootSignatureType::Sampler, MAX_SAMPLERS)
+			.AddTable(SF_Vertex, ERootSignatureType::SRV, MAX_SRVS)
+			.AddTable(SF_Vertex, ERootSignatureType::CBV, MAX_CBS)
+			.AddTable(SF_Vertex, ERootSignatureType::Sampler, MAX_SAMPLERS)
+			.AddTable(SF_Geometry, ERootSignatureType::SRV, MAX_SRVS)
+			.AddTable(SF_Geometry, ERootSignatureType::CBV, MAX_CBS)
+			.AddTable(SF_Geometry, ERootSignatureType::Sampler, MAX_SAMPLERS);
 		if (bAllowMeshShaders)
 		{
-			Creator->AddTable(SF_Mesh, FRootSignatureCreator::SRV, MAX_SRVS);
-			Creator->AddTable(SF_Mesh, FRootSignatureCreator::CBV, MAX_CBS);
-			Creator->AddTable(SF_Mesh, FRootSignatureCreator::Sampler, MAX_SAMPLERS);
-			Creator->AddTable(SF_Amplification, FRootSignatureCreator::SRV, MAX_SRVS);
-			Creator->AddTable(SF_Amplification, FRootSignatureCreator::CBV, MAX_CBS);
-			Creator->AddTable(SF_Amplification, FRootSignatureCreator::Sampler, MAX_SAMPLERS);
+			Creator
+				.AddTable(SF_Mesh, ERootSignatureType::SRV, MAX_SRVS)
+				.AddTable(SF_Mesh, ERootSignatureType::CBV, MAX_CBS)
+				.AddTable(SF_Mesh, ERootSignatureType::Sampler, MAX_SAMPLERS)
+				.AddTable(SF_Amplification, ERootSignatureType::SRV, MAX_SRVS)
+				.AddTable(SF_Amplification, ERootSignatureType::CBV, MAX_CBS)
+				.AddTable(SF_Amplification, ERootSignatureType::Sampler, MAX_SAMPLERS);
 		}
-		Creator->AddTable(SF_NumFrequencies, FRootSignatureCreator::UAV, MAX_UAVS);
+		Creator.AddTable(SF_NumFrequencies, ERootSignatureType::UAV, MAX_UAVS);
 	}
 
 	// Fat/Static Compute Root Signature
-	inline void CreateComputeRootSignature(FRootSignatureCreator* Creator)
+	inline void CreateComputeRootSignature(FRootSignatureCreator& Creator)
 	{
 		// Ensure the creator starts in a clean state (in cases of creator reuse, etc.).
-		Creator->Reset();
-
-		Creator->AddRootFlag(D3D12_ROOT_SIGNATURE_FLAG_NONE);
-		Creator->AddTable(SF_NumFrequencies, FRootSignatureCreator::SRV, MAX_SRVS);
-		Creator->AddTable(SF_NumFrequencies, FRootSignatureCreator::CBV, MAX_CBS);
-		Creator->AddTable(SF_NumFrequencies, FRootSignatureCreator::Sampler, MAX_SAMPLERS);
-		Creator->AddTable(SF_NumFrequencies, FRootSignatureCreator::UAV, MAX_UAVS);
-	}
-
-	inline D3D12_SHADER_VISIBILITY TranslateShaderVisibility(EShaderFrequency Stage)
-	{
-		switch (Stage)
-		{
-		case SF_Vertex:
-			return D3D12_SHADER_VISIBILITY_VERTEX;
-		case SF_Pixel:
-			return D3D12_SHADER_VISIBILITY_PIXEL;
-		case SF_Geometry:
-			return D3D12_SHADER_VISIBILITY_GEOMETRY;
-		default:
-			return D3D12_SHADER_VISIBILITY_ALL;
-		}
+		Creator
+			.Reset()
+			.AddRootFlag(D3D12_ROOT_SIGNATURE_FLAG_NONE)
+			.AddTable(SF_NumFrequencies, ERootSignatureType::SRV, MAX_SRVS)
+			.AddTable(SF_NumFrequencies, ERootSignatureType::CBV, MAX_CBS)
+			.AddTable(SF_NumFrequencies, ERootSignatureType::Sampler, MAX_SAMPLERS)
+			.AddTable(SF_NumFrequencies, ERootSignatureType::UAV, MAX_UAVS);
 	}
 
 	struct FBinaryRootSignatureCreator : public FRootSignatureCreator
@@ -131,16 +200,24 @@ namespace D3D12ShaderUtils
 		TArray<D3D12_ROOT_PARAMETER1> Parameters;
 		TMap<uint32, uint32> ParameterToRangeMap;
 
-		D3D12_ROOT_SIGNATURE_FLAGS Flags;
+		D3D12_ROOT_SIGNATURE_FLAGS Flags = D3D12_ROOT_SIGNATURE_FLAG_NONE;
 
-		FBinaryRootSignatureCreator() : Flags(D3D12_ROOT_SIGNATURE_FLAG_NONE) {}
+		FRootSignatureCreator& Reset() override
+		{
+			DescriptorRanges.Empty();
+			Parameters.Empty();
+			ParameterToRangeMap.Empty();
+			Flags = D3D12_ROOT_SIGNATURE_FLAG_NONE;
+			return *this;
+		}
 
-		virtual void AddRootFlag(D3D12_ROOT_SIGNATURE_FLAGS Flag) override
+		FRootSignatureCreator& AddRootFlag(D3D12_ROOT_SIGNATURE_FLAGS Flag) override
 		{
 			Flags |= Flag;
+			return *this;
 		}
 
-		virtual void AddTable(EShaderFrequency Stage, EType Type, int32 NumDescriptors)
+		FRootSignatureCreator& AddTable(EShaderFrequency Stage, ERootSignatureType Type, int32 NumDescriptors) override
 		{
 			int32 ParameterIndex = Parameters.AddZeroed();
 			int32 RangeIndex = DescriptorRanges.AddZeroed();
@@ -155,24 +232,29 @@ namespace D3D12ShaderUtils
 			//Parameter.DescriptorTable.pDescriptorRanges = &DescriptorRanges[0];
 			ParameterToRangeMap.Add(ParameterIndex, RangeIndex);
 
-			Range.RangeType = FRootSignatureCreator::GetD3D12DescriptorRangeType(Type);
+			Range.RangeType = GetD3D12DescriptorRangeType(Type);
 			Range.NumDescriptors = NumDescriptors;
-			Range.Flags = FRootSignatureCreator::GetD3D12DescriptorRangeFlags(Type);
+			Range.Flags = GetD3D12DescriptorRangeFlags(Type);
 			//Range.OffsetInDescriptorsFromTableStart = UINT32_MAX;
-			Parameter.ShaderVisibility =  D3D12ShaderUtils::TranslateShaderVisibility(Stage);
-		}
+			Parameter.ShaderVisibility =  GetD3D12ShaderVisibility(Stage);
 
-		virtual void Reset()
-		{
-			DescriptorRanges.Empty();
-			Parameters.Empty();
-			ParameterToRangeMap.Empty();
-			Flags = D3D12_ROOT_SIGNATURE_FLAG_NONE;
+			return *this;
 		}
 
-		void Compile()
+		void Compile(EShaderFrequency Freq)
 		{
-			D3D12ShaderUtils::CreateGfxRootSignature(this, false);
+			if (Freq < SF_NumGraphicsFrequencies)
+			{
+				bool bAllowMeshShaders = false;
+#if !defined(D3D12RHI_TOOLS_MESH_SHADERS_UNSUPPORTED)
+				bAllowMeshShaders = true;
+#endif
+				D3D12ShaderUtils::CreateGfxRootSignature(*this, bAllowMeshShaders);
+			}
+			else //if (Freq == SF_Compute)
+			{
+				D3D12ShaderUtils::CreateComputeRootSignature(*this);
+			}
 
 			// Patch pointers
 			for (auto& Pair : ParameterToRangeMap)
@@ -181,4 +263,75 @@ namespace D3D12ShaderUtils
 			}
 		}
 	};
+
+	/* Root signature generator for DXC */
+	struct FTextRootSignatureCreator : public FRootSignatureCreator
+	{
+		FRootSignatureCreator& AddRootFlag(D3D12_ROOT_SIGNATURE_FLAGS InFlag) override
+		{
+			if (Flags.Len() > 0)
+			{
+				Flags += "|";
+			}
+			Flags += GetFlagName(InFlag);
+			return *this;
+		}
+
+		FRootSignatureCreator& AddTable(EShaderFrequency InStage, ERootSignatureType Type, int32 NumDescriptors) override
+		{
+			FString Line = FString::Printf(TEXT("DescriptorTable(visibility=%s, %s0, numDescriptors=%d%s"/*, flags = DESCRIPTORS_VOLATILE*/"))"),
+				GetVisibilityFlag(InStage), GetTypePrefix(Type), NumDescriptors,
+				TEXT("")
+			);
+			if (Table.Len() > 0)
+			{
+				Table += ",";
+			}
+			Table += Line;
+			return *this;
+		}
+
+		FRootSignatureCreator& Reset() override
+		{
+			Flags.Empty();
+			Table.Empty();
+			return *this;
+		}
+
+		FString CreateAndGenerateString(EShaderFrequency Freq)
+		{
+			if (Freq < SF_NumGraphicsFrequencies)
+			{
+				bool bAllowMeshShaders = false;
+#if !defined(D3D12RHI_TOOLS_MESH_SHADERS_UNSUPPORTED)
+				bAllowMeshShaders = true;
+#endif
+				D3D12ShaderUtils::CreateGfxRootSignature(*this, bAllowMeshShaders);
+			}
+			else //if (Freq == SF_Compute)
+			{
+				D3D12ShaderUtils::CreateComputeRootSignature(*this);
+			}
+
+			//FString String = FString::Printf(TEXT("#define TheRootSignature \"RootFlags(%s),%s\"\n[RootSignature(TheRootSignature)\n"),
+			FString String = FString::Printf(TEXT("\"RootFlags(%s),%s\""),
+				Flags.Len() == 0 ? TEXT("0") : *Flags,
+				*Table);
+			return String;
+		}
+
+		FString Flags;
+		FString Table;
+	};
+
+	inline FString GenerateRootSignatureString(EShaderFrequency InFrequency)
+	{
+		if (InFrequency < SF_NumStandardFrequencies)
+		{
+			FTextRootSignatureCreator Creator;
+			return Creator.CreateAndGenerateString(InFrequency);
+		}
+
+		return TEXT("");
+	}
 }
diff --git a/Engine/Source/Runtime/D3D12RHI/Public/D3D12Util.h b/Engine/Source/Runtime/D3D12RHI/Public/D3D12Util.h
--- a/Engine/Source/Runtime/D3D12RHI/Public/D3D12Util.h
+++ b/Engine/Source/Runtime/D3D12RHI/Public/D3D12Util.h
@@ -93,7 +93,9 @@ struct FD3D12QuantizedBoundShaderState
 	uint8 bAllowIAInputLayout : 1;
 	uint8 bNeedsAgsIntrinsicsSpace : 1;
 	uint8 bUseDiagnosticBuffer : 1;
-	uint8 Padding : 5;
+	uint8 bUseDirectlyIndexedResourceHeap : 1;
+	uint8 bUseDirectlyIndexedSamplerHeap : 1;
+	uint8 Padding : 3;
 
 	inline bool operator==(const FD3D12QuantizedBoundShaderState& RHS) const
 	{
