Fix occasional D3D12 crash after allocating a resource via RHICreateTexture2DFromResource and releasing it.

FD3D12DynamicRHI::CreateTextureFromResource() shouldn't call Location.AsStandAlone(), since that adds the allocation to FD3D12Adapter::TrackedAllocationData. However, since the type is then changed to eAliased, the resource will not be removed from TrackedAllocationData when it's released (FD3D12ResourceLocation::ReleaseResource() only calls UpdateStandAloneStats() for eStandAlone resources, as it should). This skews the stats, but more importantly causes an assert if the memory for the released FD3D12ResourceLocation object is reused by a subsequent allocation.

#rnx

diff --git a/Engine/Source/Runtime/D3D12RHI/Private/D3D12Texture.cpp b/Engine/Source/Runtime/D3D12RHI/Private/D3D12Texture.cpp
--- a/Engine/Source/Runtime/D3D12RHI/Private/D3D12Texture.cpp
+++ b/Engine/Source/Runtime/D3D12RHI/Private/D3D12Texture.cpp
@@ -3030,8 +3030,9 @@ TD3D12Texture2D<BaseResourceType>* FD3D12DynamicRHI::CreateTextureFromResource(b
 	});
 
 	FD3D12ResourceLocation& Location = Texture2D->ResourceLocation;
-	Location.AsStandAlone(TextureResource);
 	Location.SetType(FD3D12ResourceLocation::ResourceLocationType::eAliased);
+	Location.SetResource(TextureResource);
+	Location.SetGPUVirtualAddress(TextureResource->GetGPUVirtualAddress());
 
 	uint32 RTVIndex = 0;
 
