modifications needed for dbuffer write mask optimization

diff --git a/Engine/Source/Runtime/D3D12RHI/Private/D3D12View.h b/Engine/Source/Runtime/D3D12RHI/Private/D3D12View.h
--- a/Engine/Source/Runtime/D3D12RHI/Private/D3D12View.h
+++ b/Engine/Source/Runtime/D3D12RHI/Private/D3D12View.h
@@ -705,7 +705,12 @@ protected:
 	}
 
 protected:
-	void InitializeInternal(const TDesc& InDesc, FD3D12BaseShaderResource* InBaseShaderResource, FD3D12ResourceLocation& InResourceLocation)
+	void SetDesc(const TDesc& InDesc)
+	{
+		Desc = InDesc;
+	}
+
+	void InitializeInternal(FD3D12BaseShaderResource* InBaseShaderResource, FD3D12ResourceLocation& InResourceLocation)
 	{
 		check(InBaseShaderResource);
 
@@ -718,7 +723,6 @@ protected:
 		BaseShaderResource = InBaseShaderResource;
 		ResourceLocation = &InResourceLocation;
 		Resource = ResourceLocation->GetResource();
-		Desc = InDesc;
 
 		// Transient resources might not have an actual resource yet
 		if (Resource)
@@ -742,29 +746,6 @@ protected:
 #endif
 	}
 
-	void CreateView(const TDesc& InDesc, FD3D12BaseShaderResource* InBaseShaderResource, FD3D12ResourceLocation& InResourceLocation)
-	{
-		InitializeInternal(InDesc, InBaseShaderResource, InResourceLocation);
-
-		if (ResourceLocation->GetResource())
-		{
-			ID3D12Resource* D3DResource = ResourceLocation->GetResource()->GetResource();
-			Descriptor.CreateView(Desc, D3DResource);
-		}
-	}
-
-	void CreateViewWithCounter(const TDesc& InDesc, FD3D12BaseShaderResource* InBaseShaderResource, FD3D12ResourceLocation& InResourceLocation, FD3D12Resource* InCounterResource)
-	{
-		InitializeInternal(InDesc, InBaseShaderResource, InResourceLocation);
-
-		if (Resource)
-		{
-			ID3D12Resource* D3DResource = Resource->GetUAVAccessResource() ? Resource->GetUAVAccessResource() : Resource->GetResource();
-			ID3D12Resource* D3DCounterResource = InCounterResource ? InCounterResource->GetResource() : nullptr;
-			Descriptor.CreateViewWithCounter(Desc, D3DResource, D3DCounterResource);
-		}
-	}
-
 	virtual void ResourceRenamed(FD3D12BaseShaderResource* InRenamedResource, FD3D12ResourceLocation* InNewResourceLocation) override
 	{
 		check(InRenamedResource == BaseShaderResource);
@@ -839,7 +820,7 @@ public:
 	{
 	}
 
-	void Initialize(D3D12_SHADER_RESOURCE_VIEW_DESC& InDesc, FD3D12BaseShaderResource* InBaseShaderResource, FD3D12ResourceLocation& InResourceLocation, uint32 InStride, uint32 InStartOffsetBytes, bool InSkipFastClearFinalize)
+	void PreCreateView(FD3D12ResourceLocation& InResourceLocation, uint32 InStride, uint32 InStartOffsetBytes, bool InSkipFastClearFinalize)
 	{
 		Stride = InStride;
 		StartOffsetBytes = InStartOffsetBytes;
@@ -847,25 +828,42 @@ public:
 
 		if (InResourceLocation.GetResource())
 		{
-			bContainsDepthPlane = InResourceLocation.GetResource()->IsDepthStencilResource() && GetPlaneSliceFromViewFormat(InResourceLocation.GetResource()->GetDesc().Format, InDesc.Format) == 0;
-			bContainsStencilPlane = InResourceLocation.GetResource()->IsDepthStencilResource() && GetPlaneSliceFromViewFormat(InResourceLocation.GetResource()->GetDesc().Format, InDesc.Format) == 1;
+			bContainsDepthPlane = InResourceLocation.GetResource()->IsDepthStencilResource() && GetPlaneSliceFromViewFormat(InResourceLocation.GetResource()->GetDesc().Format, Desc.Format) == 0;
+			bContainsStencilPlane = InResourceLocation.GetResource()->IsDepthStencilResource() && GetPlaneSliceFromViewFormat(InResourceLocation.GetResource()->GetDesc().Format, Desc.Format) == 1;
 			bRequiresResourceStateTracking = InResourceLocation.GetResource()->RequiresResourceStateTracking();
 
 #if DO_CHECK
 			// Check the plane slice of the SRV matches the texture format
 			// Texture2DMS does not have explicit plane index (it's implied by the format)
-			if (InDesc.ViewDimension == D3D12_SRV_DIMENSION_TEXTURE2D)
+			if (Desc.ViewDimension == D3D12_SRV_DIMENSION_TEXTURE2D)
 			{
-				check(GetPlaneSliceFromViewFormat(InResourceLocation.GetResource()->GetDesc().Format, InDesc.Format) == InDesc.Texture2D.PlaneSlice);
+				check(GetPlaneSliceFromViewFormat(InResourceLocation.GetResource()->GetDesc().Format, Desc.Format) == Desc.Texture2D.PlaneSlice);
 			}
 
-			if (InDesc.ViewDimension == D3D12_SRV_DIMENSION_BUFFER)
+			if (Desc.ViewDimension == D3D12_SRV_DIMENSION_BUFFER)
 			{
-				//check(InResourceLocation.GetOffsetFromBaseOfResource() == InDesc.Buffer.FirstElement * Stride);
+				//check(InResourceLocation.GetOffsetFromBaseOfResource() == Desc.Buffer.FirstElement * Stride);
 			}
 #endif
 		}
+	}
+
+	void CreateView(const D3D12_SHADER_RESOURCE_VIEW_DESC& InDesc, FD3D12BaseShaderResource* InBaseShaderResource, FD3D12ResourceLocation& InResourceLocation)
+	{
+		InitializeInternal(InBaseShaderResource, InResourceLocation);
+
+		if (ResourceLocation->GetResource())
+		{
+			ID3D12Resource* D3DResource = ResourceLocation->GetResource()->GetResource();
+			Descriptor.CreateView(Desc, D3DResource);
+		}
+	}
 
+	void Initialize(D3D12_SHADER_RESOURCE_VIEW_DESC& InDesc, FD3D12BaseShaderResource* InBaseShaderResource, FD3D12ResourceLocation& InResourceLocation, uint32 InStride, uint32 InStartOffsetBytes, bool InSkipFastClearFinalize)
+	{
+		SetDesc(InDesc);
+
+		PreCreateView(InResourceLocation, InStride, InStartOffsetBytes, InSkipFastClearFinalize);
 		CreateView(InDesc, InBaseShaderResource, InResourceLocation);
 	}
 
@@ -955,28 +953,40 @@ class FD3D12UnorderedAccessView : public FRHIUnorderedAccessView, public FD3D12V
 public:
 
 	FD3D12UnorderedAccessView(FD3D12Device* InParent)
-	: FD3D12View(InParent, ViewSubresourceSubsetFlags_None)
+		: FD3D12View(InParent, ViewSubresourceSubsetFlags_None)
 	{
 	}
 
 	FD3D12UnorderedAccessView(FD3D12Device* InParent, D3D12_UNORDERED_ACCESS_VIEW_DESC& InDesc, FD3D12BaseShaderResource* InBaseShaderResource, FD3D12Resource* InCounterResource = nullptr)
 		: FD3D12View(InParent, ViewSubresourceSubsetFlags_None)
 		, CounterResource(InCounterResource)
-		, CounterResourceInitialized(false)
 	{
-		CreateViewWithCounter(InDesc, InBaseShaderResource, InBaseShaderResource->ResourceLocation, InCounterResource);
+		Initialize(InDesc, InBaseShaderResource, InBaseShaderResource->ResourceLocation, InCounterResource);
+	}
+
+	void CreateView(const D3D12_UNORDERED_ACCESS_VIEW_DESC& InDesc, FD3D12BaseShaderResource* InBaseShaderResource, FD3D12ResourceLocation& InResourceLocation, FD3D12Resource* InCounterResource)
+	{
+		InitializeInternal(InBaseShaderResource, InResourceLocation);
+
+		if (Resource)
+		{
+			ID3D12Resource* D3DResource = Resource->GetUAVAccessResource() ? Resource->GetUAVAccessResource() : Resource->GetResource();
+			ID3D12Resource* D3DCounterResource = InCounterResource ? InCounterResource->GetResource() : nullptr;
+			Descriptor.CreateViewWithCounter(Desc, D3DResource, D3DCounterResource);
+		}
 	}
 
-	void Initialize(D3D12_UNORDERED_ACCESS_VIEW_DESC& InDesc, FD3D12BaseShaderResource* InBaseShaderResource, FD3D12ResourceLocation& InResourceLocation)
+	void Initialize(D3D12_UNORDERED_ACCESS_VIEW_DESC& InDesc, FD3D12BaseShaderResource* InBaseShaderResource, FD3D12ResourceLocation& InResourceLocation, FD3D12Resource* InCounterResource = nullptr)
 	{
-		CreateViewWithCounter(InDesc, InBaseShaderResource, InResourceLocation, nullptr);
+		SetDesc(InDesc);
+		CreateView(InDesc, InBaseShaderResource, InResourceLocation, InCounterResource);
 	}
 
 	virtual void RecreateView() override
 	{
 		check(CounterResource == nullptr);
 		check(ResourceLocation->GetOffsetFromBaseOfResource() == 0);
-		Initialize(Desc, BaseShaderResource, *ResourceLocation);
+		CreateView(Desc, BaseShaderResource, *ResourceLocation, nullptr);
 	}
 
 	bool IsCounterResourceInitialized() const { return CounterResourceInitialized; }
@@ -984,10 +994,9 @@ public:
 
 	FD3D12Resource* GetCounterResource() { return CounterResource; }
 
-private:
-
+protected:
 	TRefCountPtr<FD3D12Resource> CounterResource;
-	bool CounterResourceInitialized;
+	bool CounterResourceInitialized = false;
 };
 
 class FD3D12UnorderedAccessViewWithLocation : public FD3D12UnorderedAccessView
@@ -1077,6 +1086,18 @@ public:
 		CreateView(InRTVDesc, InBaseShaderResource, InResourceLocation);
 	}
 
+	void CreateView(const D3D12_RENDER_TARGET_VIEW_DESC& InDesc, FD3D12BaseShaderResource* InBaseShaderResource, FD3D12ResourceLocation& InResourceLocation)
+	{
+		SetDesc(InDesc);
+		InitializeInternal(InBaseShaderResource, InResourceLocation);
+
+		if (ResourceLocation->GetResource())
+		{
+			ID3D12Resource* D3DResource = ResourceLocation->GetResource()->GetResource();
+			Descriptor.CreateView(Desc, D3DResource);
+		}
+	}
+
 	virtual void RecreateView() override
 	{
 		check(ResourceLocation->GetOffsetFromBaseOfResource() == 0);
@@ -1102,6 +1123,18 @@ public:
 		SetupDepthStencilViewSubresourceSubset();
 	}
 
+	void CreateView(const D3D12_DEPTH_STENCIL_VIEW_DESC& InDesc, FD3D12BaseShaderResource* InBaseShaderResource, FD3D12ResourceLocation& InResourceLocation)
+	{
+		SetDesc(InDesc);
+		InitializeInternal(InBaseShaderResource, InResourceLocation);
+
+		if (ResourceLocation->GetResource())
+		{
+			ID3D12Resource* D3DResource = ResourceLocation->GetResource()->GetResource();
+			Descriptor.CreateView(Desc, D3DResource);
+		}
+	}
+
 	bool HasDepth() const
 	{
 		return bHasDepth;
