D3D12 Descriptor Management Reorganization
- Some reorganization in preparation for better descriptor management.
- Adding even higher level Descriptor Heap Manager that all descriptor heaps should be allocated through. This manager also has an internal Heap that all GPU Visible resource heaps are suballocated from.
- Reworking all existing Descriptor Heap allocations to go through the new manager.
- Adding a Resource Descriptor Manager for views to allocate resource descriptors.
- Renamed FD3D12GlobalHeap to FD3D12OnlineDescriptorManager to better reflect where it's used.
- Merging SRV, UAV and CBV descriptor heaps into a single "Standard" heap type.
- Reworking View Descriptor Handles to not be bound to specific view types.

diff --git a/Engine/Source/Runtime/D3D12RHI/Private/D3D12Allocation.cpp b/Engine/Source/Runtime/D3D12RHI/Private/D3D12Allocation.cpp
--- a/Engine/Source/Runtime/D3D12RHI/Private/D3D12Allocation.cpp
+++ b/Engine/Source/Runtime/D3D12RHI/Private/D3D12Allocation.cpp
@@ -2133,11 +2133,7 @@ FD3D12FastConstantAllocator::FD3D12FastConstantAllocator(FD3D12Device* Parent, F
 	check(PageSize % D3D12_CONSTANT_BUFFER_DATA_PLACEMENT_ALIGNMENT == 0);
 }
 
-#if USE_STATIC_ROOT_SIGNATURE
 void* FD3D12FastConstantAllocator::Allocate(uint32 Bytes, FD3D12ResourceLocation& OutLocation, FD3D12ConstantBufferView* OutCBView)
-#else
-void* FD3D12FastConstantAllocator::Allocate(uint32 Bytes, FD3D12ResourceLocation& OutLocation)
-#endif
 {
 	check(Bytes <= PageSize);
 
@@ -2161,12 +2157,10 @@ void* FD3D12FastConstantAllocator::Allocate(uint32 Bytes, FD3D12ResourceLocation
 		UnderlyingResource.GetOffsetFromBaseOfResource(), // AllocUploadResource returns a suballocated resource where we're suballocating (again) from
 		Offset);
 
-#if USE_STATIC_ROOT_SIGNATURE
 	if (OutCBView)
 	{
 		OutCBView->Create(UnderlyingResource.GetGPUVirtualAddress() + Offset, AlignedSize);
 	}
-#endif
 
 	Offset += AlignedSize;
 
diff --git a/Engine/Source/Runtime/D3D12RHI/Private/D3D12Allocation.h b/Engine/Source/Runtime/D3D12RHI/Private/D3D12Allocation.h
--- a/Engine/Source/Runtime/D3D12RHI/Private/D3D12Allocation.h
+++ b/Engine/Source/Runtime/D3D12RHI/Private/D3D12Allocation.h
@@ -515,12 +515,7 @@ class FD3D12FastConstantAllocator : public FD3D12DeviceChild, public FD3D12Multi
 public:
 	FD3D12FastConstantAllocator(FD3D12Device* Parent, FRHIGPUMask VisibiltyMask);
 
-#if USE_STATIC_ROOT_SIGNATURE
 	void* Allocate(uint32 Bytes, class FD3D12ResourceLocation& OutLocation, FD3D12ConstantBufferView* OutCBView);
-#else
-	void* Allocate(uint32 Bytes, class FD3D12ResourceLocation& OutLocation);
-#endif
-
 
 private:
 	FD3D12ResourceLocation UnderlyingResource;
diff --git a/Engine/Source/Runtime/D3D12RHI/Private/D3D12ConstantBuffer.cpp b/Engine/Source/Runtime/D3D12RHI/Private/D3D12ConstantBuffer.cpp
--- a/Engine/Source/Runtime/D3D12RHI/Private/D3D12ConstantBuffer.cpp
+++ b/Engine/Source/Runtime/D3D12RHI/Private/D3D12ConstantBuffer.cpp
@@ -11,9 +11,6 @@ DEFINE_STAT(STAT_D3D12GlobalConstantBufferUpdateTime);
 // New circular buffer system for faster constant uploads.  Avoids CopyResource and speeds things up considerably
 FD3D12ConstantBuffer::FD3D12ConstantBuffer(FD3D12Device* InParent, FD3D12FastConstantAllocator& InAllocator) :
 	FD3D12DeviceChild(InParent),
-#if USE_STATIC_ROOT_SIGNATURE
-	View(nullptr),
-#endif
 	CurrentUpdateSize(0),
 	TotalUpdateSize(0),
 	bIsDirty(false),
@@ -21,7 +18,7 @@ FD3D12ConstantBuffer::FD3D12ConstantBuffer(FD3D12Device* InParent, FD3D12FastCon
 {
 	FMemory::Memset(ShadowData, 0);
 #if USE_STATIC_ROOT_SIGNATURE
-	View = new FD3D12ConstantBufferView(InParent, nullptr);
+	View = new FD3D12ConstantBufferView(InParent);
 #endif
 }
 
@@ -53,16 +50,13 @@ bool FD3D12ConstantBuffer::Version(FD3D12ResourceLocation& BufferOut, bool bDisc
 		TotalUpdateSize = FMath::Max(CurrentUpdateSize, TotalUpdateSize);
 	}
 
+	FD3D12ConstantBufferView* ViewToUse = nullptr;
 #if USE_STATIC_ROOT_SIGNATURE
-	View->AllocateHeapSlot();
+	ViewToUse = View;
 #endif
 
 	// Get the next constant buffer
-#if USE_STATIC_ROOT_SIGNATURE
-	void* Data = Allocator.Allocate(TotalUpdateSize, BufferOut, View);
-#else
-	void* Data = Allocator.Allocate(TotalUpdateSize, BufferOut);
-#endif
+	void* Data = Allocator.Allocate(TotalUpdateSize, BufferOut, ViewToUse);
 
 	check(TotalUpdateSize <= sizeof(ShadowData));
 	FMemory::Memcpy(Data, ShadowData, TotalUpdateSize);
diff --git a/Engine/Source/Runtime/D3D12RHI/Private/D3D12DescriptorCache.cpp b/Engine/Source/Runtime/D3D12RHI/Private/D3D12DescriptorCache.cpp
--- a/Engine/Source/Runtime/D3D12RHI/Private/D3D12DescriptorCache.cpp
+++ b/Engine/Source/Runtime/D3D12RHI/Private/D3D12DescriptorCache.cpp
@@ -5,14 +5,6 @@
 //-----------------------------------------------------------------------------
 #include "D3D12RHIPrivate.h"
 
-int32 GGlobalViewHeapBlockSize = 2000;
-static FAutoConsoleVariableRef CVarGlobalViewHeapBlockSize(
-	TEXT("D3D12.GlobalViewHeapBlockSize"),
-	GGlobalViewHeapBlockSize,
-	TEXT("Block size for sub allocations on the global view descriptor heap."),
-	ECVF_ReadOnly
-);
-
 // Define template functions that are only declared in the header.
 #if USE_STATIC_ROOT_SIGNATURE
 template void FD3D12DescriptorCache::SetConstantBuffers<SF_Vertex>(const FD3D12RootSignature* RootSignature, FD3D12ConstantBufferCache& Cache, const CBVSlotMask& SlotsNeededMask, uint32 Count, uint32& HeapSlot);
@@ -47,15 +39,15 @@ template void FD3D12DescriptorCache::SetSamplers<SF_Geometry>(const FD3D12RootSi
 template void FD3D12DescriptorCache::SetSamplers<SF_Pixel>(const FD3D12RootSignature* RootSignature, FD3D12SamplerStateCache& Cache, const SamplerSlotMask& SlotsNeededMask, uint32 SlotsNeeded, uint32& HeapSlot);
 template void FD3D12DescriptorCache::SetSamplers<SF_Compute>(const FD3D12RootSignature* RootSignature, FD3D12SamplerStateCache& Cache, const SamplerSlotMask& SlotsNeededMask, uint32 SlotsNeeded, uint32& HeapSlot);
 
-bool FD3D12DescriptorCache::HeapRolledOver(D3D12_DESCRIPTOR_HEAP_TYPE Type)
+bool FD3D12DescriptorCache::HeapRolledOver(ED3D12DescriptorHeapType InHeapType)
 {
 	// A heap rolled over, so set the descriptor heaps again and return if the heaps actually changed.
 	return SetDescriptorHeaps();
 }
 
-void FD3D12DescriptorCache::HeapLoopedAround(D3D12_DESCRIPTOR_HEAP_TYPE Type)
+void FD3D12DescriptorCache::HeapLoopedAround(ED3D12DescriptorHeapType InHeapType)
 {
-	if (Type == FD3D12_DESCRIPTOR_HEAP_TYPE_SAMPLER)
+	if (InHeapType == ED3D12DescriptorHeapType::Sampler)
 	{
 		SamplerMap.Reset();
 	}
@@ -64,16 +56,13 @@ void FD3D12DescriptorCache::HeapLoopedAround(D3D12_DESCRIPTOR_HEAP_TYPE Type)
 FD3D12DescriptorCache::FD3D12DescriptorCache(FRHIGPUMask Node)
 	: FD3D12DeviceChild(nullptr)
 	, FD3D12SingleNodeGPUObject(Node)
-	, pNullSRV(nullptr)
-	, pNullRTV(nullptr)
-	, pNullUAV(nullptr)
 	, pPreviousViewHeap(nullptr)
 	, pPreviousSamplerHeap(nullptr)
 	, CurrentViewHeap(nullptr)
 	, CurrentSamplerHeap(nullptr)
 	, LocalViewHeap(nullptr)
-	, LocalSamplerHeap(nullptr, Node, this)
-	, SubAllocatedViewHeap(Node, this)
+	, LocalSamplerHeap(nullptr, this)
+	, SubAllocatedViewHeap(this)
 	, SamplerMap(271) // Prime numbers for better hashing
 	, bUsingGlobalSamplerHeap(false)
 	, NumLocalViewDescriptors(0)
@@ -87,12 +76,12 @@ void FD3D12DescriptorCache::Init(FD3D12Device* InParent, FD3D12CommandContext* I
 	CmdContext = InCmdContext;
 
 	LocalSamplerHeap.SetParentDevice(InParent);
-	SubAllocatedViewHeap.Init(InParent, D3D12_DESCRIPTOR_HEAP_TYPE_CBV_SRV_UAV);
+	SubAllocatedViewHeap.Init(InParent);
 
 	// Always Init a local sampler heap as the high level cache will always miss initialy
 	// so we need something to fall back on (The view heap never rolls over so we init that one
 	// lazily as a backup to save memory)
-	LocalSamplerHeap.Init(InNumSamplerDescriptors, FD3D12_DESCRIPTOR_HEAP_TYPE_SAMPLER);
+	LocalSamplerHeap.Init(InNumSamplerDescriptors, ED3D12DescriptorHeapType::Sampler);
 
 	NumLocalViewDescriptors = InNumLocalViewDescriptors;
 
@@ -108,25 +97,25 @@ void FD3D12DescriptorCache::Init(FD3D12Device* InParent, FD3D12CommandContext* I
 	SRVDesc.Texture2D.MipLevels = 1;
 	SRVDesc.Texture2D.MostDetailedMip = 0;
 	SRVDesc.Texture2D.ResourceMinLODClamp = 0.0f;
-	pNullSRV = new FD3D12DescriptorHandleSRV(GetParentDevice());
-	pNullSRV->CreateView(SRVDesc, nullptr);
+	NullSRV = new FD3D12ViewDescriptorHandle(GetParentDevice(), ED3D12DescriptorHeapType::Standard);
+	NullSRV->CreateView(SRVDesc, nullptr);
 
 	D3D12_RENDER_TARGET_VIEW_DESC RTVDesc = {};
 	RTVDesc.ViewDimension = D3D12_RTV_DIMENSION_TEXTURE2D;
 	RTVDesc.Format = DXGI_FORMAT_B8G8R8A8_UNORM;
 	RTVDesc.Texture2D.MipSlice = 0;
-	pNullRTV = new FD3D12DescriptorHandleRTV(GetParentDevice());
-	pNullRTV->CreateView(RTVDesc, nullptr);
+	NullRTV = new FD3D12ViewDescriptorHandle(GetParentDevice(), ED3D12DescriptorHeapType::RenderTarget);
+	NullRTV->CreateView(RTVDesc, nullptr);
 
 	D3D12_UNORDERED_ACCESS_VIEW_DESC UAVDesc = {};
 	UAVDesc.ViewDimension = D3D12_UAV_DIMENSION_TEXTURE2D;
 	UAVDesc.Format = DXGI_FORMAT_R8G8B8A8_UNORM;
 	UAVDesc.Texture2D.MipSlice = 0;
-	pNullUAV = new FD3D12DescriptorHandleUAV(GetParentDevice());
-	pNullUAV->CreateViewWithCounter(UAVDesc, nullptr, nullptr);
+	NullUAV = new FD3D12ViewDescriptorHandle(GetParentDevice(), ED3D12DescriptorHeapType::Standard);
+	NullUAV->CreateView(UAVDesc, nullptr, nullptr);
 
 #if USE_STATIC_ROOT_SIGNATURE
-	pNullCBV = new FD3D12ConstantBufferView(GetParentDevice(), nullptr);
+	NullCBV = new FD3D12ConstantBufferView(GetParentDevice());
 #endif
 
 	const FSamplerStateInitializerRHI SamplerDesc(
@@ -142,20 +131,20 @@ void FD3D12DescriptorCache::Init(FD3D12Device* InParent, FD3D12CommandContext* I
 
 	FSamplerStateRHIRef Sampler = InParent->CreateSampler(SamplerDesc);
 
-	pDefaultSampler = static_cast<FD3D12SamplerState*>(Sampler.GetReference());
+	DefaultSampler = static_cast<FD3D12SamplerState*>(Sampler.GetReference());
 
 	// The default sampler must have ID=0
 	// DescriptorCache::SetSamplers relies on this
-	check(pDefaultSampler->ID == 0);
+	check(DefaultSampler->ID == 0);
 }
 
 void FD3D12DescriptorCache::Clear()
 {
-	delete pNullSRV; pNullSRV = nullptr;
-	delete pNullUAV; pNullUAV = nullptr;
-	delete pNullRTV; pNullRTV = nullptr;
+	delete NullSRV; NullSRV = nullptr;
+	delete NullUAV; NullUAV = nullptr;
+	delete NullRTV; NullRTV = nullptr;
 #if USE_STATIC_ROOT_SIGNATURE
-	delete pNullCBV; pNullCBV = nullptr;
+	delete NullCBV; NullCBV = nullptr;
 #endif
 }
 
@@ -326,7 +315,7 @@ void FD3D12DescriptorCache::SetUAVs(const FD3D12RootSignature* RootSignature, FD
 	{
 		if ((SlotIndex < UAVStartSlot) || (UAVs[SlotIndex] == nullptr))
 		{
-			SrcDescriptors[SlotIndex] = pNullUAV->GetHandle();
+			SrcDescriptors[SlotIndex] = NullUAV->GetOfflineHandle();
 		}
 		else
 		{
@@ -405,7 +394,7 @@ void FD3D12DescriptorCache::SetRenderTargets(FD3D12RenderTargetView** RenderTarg
 		}
 		else
 		{
-			RTVDescriptors[i] = pNullRTV->GetHandle();
+			RTVDescriptors[i] = NullRTV->GetOfflineHandle();
 		}
 	}
 
@@ -494,7 +483,7 @@ void FD3D12DescriptorCache::SetSamplers(const FD3D12RootSignature* RootSignature
 			}
 			else
 			{
-				SrcDescriptors[SlotIndex] = pDefaultSampler->Descriptor;
+				SrcDescriptors[SlotIndex] = DefaultSampler->Descriptor;
 			}
 		}
 		FD3D12SamplerStateCache::CleanSlots(CurrentDirtySlotMask, SlotsNeeded);
@@ -584,7 +573,7 @@ void FD3D12DescriptorCache::SetSRVs(const FD3D12RootSignature* RootSignature, FD
 		}
 		else
 		{
-			SrcDescriptors[SlotIndex] = pNullSRV->GetHandle();
+			SrcDescriptors[SlotIndex] = NullSRV->GetOfflineHandle();
 		}
 		check(SrcDescriptors[SlotIndex].ptr != 0);
 	}
@@ -649,12 +638,12 @@ void FD3D12DescriptorCache::SetConstantBuffers(const FD3D12RootSignature* RootSi
 	check(SlotsNeeded != 0);
 	HeapSlot += SlotsNeeded;
 
-	D3D12_CPU_DESCRIPTOR_HANDLE DestDescriptor = CurrentViewHeap->GetCPUSlotHandle(FirstSlotIndex);
-	const uint32 DescriptorSize = CurrentViewHeap->GetDescriptorSize();
+	uint32 DestDescriptorSlot = FirstSlotIndex;
 
-	//Device->CopyDescriptors(1, &DestDescriptor, &DescriptorSize, SlotsNeeded, CBVHandles, SrcSizes, D3D12_DESCRIPTOR_HEAP_TYPE_CBV_SRV_UAV);
 	for (uint32 SlotIndex = 0; SlotIndex < SlotsNeeded; SlotIndex++)
 	{
+		const D3D12_CPU_DESCRIPTOR_HANDLE DestDescriptor = CurrentViewHeap->GetCPUSlotHandle(DestDescriptorSlot);
+
 		if (CBVHandles[SlotIndex].ptr != 0)
 		{
 			Device->CopyDescriptorsSimple(1, DestDescriptor, CBVHandles[SlotIndex], D3D12_DESCRIPTOR_HEAP_TYPE_CBV_SRV_UAV);
@@ -664,10 +653,10 @@ void FD3D12DescriptorCache::SetConstantBuffers(const FD3D12RootSignature* RootSi
 		}
 		else
 		{
-			Device->CopyDescriptorsSimple(1, DestDescriptor, pNullCBV->OfflineDescriptorHandle, D3D12_DESCRIPTOR_HEAP_TYPE_CBV_SRV_UAV);
+			Device->CopyDescriptorsSimple(1, DestDescriptor, NullCBV->GetView(), D3D12_DESCRIPTOR_HEAP_TYPE_CBV_SRV_UAV);
 		}
 
-		DestDescriptor.ptr += DescriptorSize;
+		DestDescriptorSlot++;
 
 		// Clear the dirty bit.
 		FD3D12ConstantBufferCache::CleanSlot(CurrentDirtySlotMask, SlotIndex);
@@ -743,11 +732,11 @@ bool FD3D12DescriptorCache::SwitchToContextLocalViewHeap(const FD3D12CommandList
 		UE_LOG(LogD3D12RHI, Log, TEXT("This should only happen in the Editor where it doesn't matter as much. If it happens in game you should increase the device global heap size!"));
 		
 		// Allocate the heap lazily
-		LocalViewHeap = new FD3D12LocalOnlineHeap(GetParentDevice(), GetGPUMask(), this);
+		LocalViewHeap = new FD3D12LocalOnlineHeap(GetParentDevice(), this);
 		if (LocalViewHeap)
 		{
 			check(NumLocalViewDescriptors);
-			LocalViewHeap->Init(NumLocalViewDescriptors, D3D12_DESCRIPTOR_HEAP_TYPE_CBV_SRV_UAV);
+			LocalViewHeap->Init(NumLocalViewDescriptors, ED3D12DescriptorHeapType::Standard);
 		}
 		else
 		{
@@ -802,16 +791,13 @@ bool FD3D12DescriptorCache::SwitchToGlobalSamplerHeap()
 /**
 Initialization constructor
 **/
-FD3D12OnlineHeap::FD3D12OnlineHeap(FD3D12Device* Device, FRHIGPUMask Node, bool CanLoopAround)
+FD3D12OnlineHeap::FD3D12OnlineHeap(FD3D12Device* Device, bool CanLoopAround)
 	: FD3D12DeviceChild(Device)
-	, FD3D12SingleNodeGPUObject(Node)
-	, DescriptorSize(0)
 	, bCanLoopAround(CanLoopAround)
-	, NextSlotIndex(0)
-	, FirstUsedSlot(0)
-	, Desc({})
-{};
+{
+}
 
+FD3D12OnlineHeap::~FD3D12OnlineHeap() = default;
 
 /**
 Check if requested number of slots still fit the heap
@@ -867,9 +853,11 @@ Reserve requested amount of descriptor slots - should fit, user has to check wit
 **/
 uint32 FD3D12OnlineHeap::ReserveSlots(uint32 NumSlotsRequested)
 {
+	const ED3D12DescriptorHeapType HeapType = Heap->GetType();
+
 #ifdef VERBOSE_DESCRIPTOR_HEAP_DEBUG
-	FMsg::Logf(__FILE__, __LINE__, TEXT("DescriptorCache"), ELogVerbosity::Log, TEXT("Requesting reservation [TYPE %d] with %d slots, required fence is %d"),
-		(int32)Desc.Type, NumSlotsRequested, RequiredFenceForCurrentCL);
+	FMsg::Logf(__FILE__, __LINE__, TEXT("DescriptorCache"), ELogVerbosity::Log, TEXT("Requesting reservation [TYPE %s] with %d slots, required fence is %d"),
+		ToString(HeapType), NumSlotsRequested, RequiredFenceForCurrentCL);
 #endif
 
 	const uint32 HeapSize = GetTotalSize();
@@ -906,7 +894,7 @@ uint32 FD3D12OnlineHeap::ReserveSlots(uint32 NumSlotsRequested)
 	// Note where to start looking next time
 	NextSlotIndex = SlotAfterReservation;
 
-	if (Desc.Type == D3D12_DESCRIPTOR_HEAP_TYPE_CBV_SRV_UAV)
+	if (HeapType == ED3D12DescriptorHeapType::Standard)
 	{
 		INC_DWORD_STAT_BY(STAT_NumReservedViewOnlineDescriptors, NumSlotsRequested);
 	}
@@ -928,7 +916,7 @@ void FD3D12OnlineHeap::SetNextSlot(uint32 NextSlot)
 	// This is used to correct for the actual number of heap slots used
 	check(NextSlot <= NextSlotIndex);
 
-	check(Desc.Type != D3D12_DESCRIPTOR_HEAP_TYPE_CBV_SRV_UAV);
+	check(Heap->GetType() != ED3D12DescriptorHeapType::Standard);
 	DEC_DWORD_STAT_BY(STAT_NumReservedSamplerOnlineDescriptors, NextSlotIndex - NextSlot);
 
 	NextSlotIndex = NextSlot;
@@ -939,190 +927,54 @@ void FD3D12OnlineHeap::SetNextSlot(uint32 NextSlot)
 // FD3D12GlobalSamplerOnlineHeap
 ////////////////////////////////////////////////////////////////////////////////////////////////////////////////////////
 
-/**
-Allocate and initialize the global sampler heap
-**/
-void FD3D12GlobalOnlineSamplerHeap::Init(uint32 TotalSize)
+FD3D12GlobalOnlineSamplerHeap::FD3D12GlobalOnlineSamplerHeap(FD3D12Device* Device)
+	: FD3D12OnlineHeap(Device, false)
 {
-	D3D12_DESCRIPTOR_HEAP_FLAGS HeapFlags = D3D12_DESCRIPTOR_HEAP_FLAG_SHADER_VISIBLE;
-
-	Desc = {};
-	Desc.Flags = HeapFlags;
-	Desc.Type = D3D12_DESCRIPTOR_HEAP_TYPE_SAMPLER;
-	Desc.NumDescriptors = TotalSize;
-	Desc.NodeMask = GetGPUMask().GetNative();
+}
 
-	VERIFYD3D12RESULT(GetParentDevice()->GetDevice()->CreateDescriptorHeap(&Desc, IID_PPV_ARGS(Heap.GetInitReference())));
-	SetName(Heap, L"Device Global - Online Sampler Heap");
+FD3D12GlobalOnlineSamplerHeap::~FD3D12GlobalOnlineSamplerHeap() = default;
 
-	CPUBase = Heap->GetCPUDescriptorHandleForHeapStart();
-	GPUBase = Heap->GetGPUDescriptorHandleForHeapStart();
-	DescriptorSize = GetParentDevice()->GetDevice()->GetDescriptorHandleIncrementSize(Desc.Type);
+/** Allocate and initialize the global sampler heap */
+void FD3D12GlobalOnlineSamplerHeap::Init(uint32 TotalSize)
+{
+	Heap = GetParentDevice()->GetDescriptorHeapManager().AllocateHeap(
+		TEXT("Device Global - Online Sampler Heap"),
+		ED3D12DescriptorHeapType::Sampler,
+		TotalSize,
+		D3D12_DESCRIPTOR_HEAP_FLAG_SHADER_VISIBLE
+	);
 
 	INC_DWORD_STAT(STAT_NumSamplerOnlineDescriptorHeaps);
-	INC_MEMORY_STAT_BY(STAT_SamplerOnlineDescriptorHeapMemory, Desc.NumDescriptors * GetDescriptorSize());
+	INC_MEMORY_STAT_BY(STAT_SamplerOnlineDescriptorHeapMemory, Heap->GetMemorySize());
 }
 
 
-/**
-No rollover supported
-**/
 bool FD3D12GlobalOnlineSamplerHeap::RollOver()
 {
+	// No rollover supported
 	check(false);
 	UE_LOG(LogD3D12RHI, Fatal, TEXT("Global Descriptor heaps can't roll over!"));
 	return false;
 }
 
-
-////////////////////////////////////////////////////////////////////////////////////////////////////////////////////////
-// FD3D12GlobalHeap
 ////////////////////////////////////////////////////////////////////////////////////////////////////////////////////////
+// FD3D12SubAllocatedOnlineHeap
 
-/**
-@brief Allocate and initialize the global heap
-**/
-void FD3D12GlobalHeap::Init(D3D12_DESCRIPTOR_HEAP_TYPE InType, uint32 InTotalSize)
-{
-	Type = InType;
-	TotalSize = InTotalSize;
-
-	// Setup the descriptor
-	D3D12_DESCRIPTOR_HEAP_DESC Desc = {};
-	Desc.Flags = D3D12_DESCRIPTOR_HEAP_FLAG_SHADER_VISIBLE;
-	Desc.Type = InType;
-	Desc.NumDescriptors = TotalSize;
-	Desc.NodeMask = GetGPUMask().GetNative();
-
-	// Allocate the heap and name it
-	VERIFYD3D12RESULT(GetParentDevice()->GetDevice()->CreateDescriptorHeap(&Desc, IID_PPV_ARGS(Heap.GetInitReference())));
-	SetName(Heap, Desc.Type == D3D12_DESCRIPTOR_HEAP_TYPE_CBV_SRV_UAV ? L"Device Global - Online View Heap" : L"Device Global - Online Sampler Heap");
-
-	// Extract useful data from created heap
-	CPUBase = Heap->GetCPUDescriptorHandleForHeapStart();
-	GPUBase = Heap->GetGPUDescriptorHandleForHeapStart();
-	DescriptorSize = GetParentDevice()->GetDevice()->GetDescriptorHandleIncrementSize(Desc.Type);
-
-	// Update the stats
-	if (Desc.Type == D3D12_DESCRIPTOR_HEAP_TYPE_CBV_SRV_UAV)
-	{
-		INC_DWORD_STAT(STAT_NumViewOnlineDescriptorHeaps);
-		INC_MEMORY_STAT_BY(STAT_ViewOnlineDescriptorHeapMemory, Desc.NumDescriptors * DescriptorSize);
-	}
-	else
-	{
-		INC_DWORD_STAT(STAT_NumSamplerOnlineDescriptorHeaps);
-		INC_MEMORY_STAT_BY(STAT_SamplerOnlineDescriptorHeapMemory, Desc.NumDescriptors * DescriptorSize);
-	}
-
-	INC_DWORD_STAT_BY(STAT_GlobalViewHeapFreeDescriptors, TotalSize);
-	
-	// Compute amount of free blocks
-	uint32 BlockSize = GGlobalViewHeapBlockSize;
-	uint32 BlockCount = TotalSize / BlockSize;
-	ReleasedBlocks.Reserve(BlockCount);
-
-	// Allocate the free blocks
-	uint32 CurrentBaseSlot = 0;
-	for (uint32 BlockIndex = 0; BlockIndex < BlockCount; ++BlockIndex)
-	{
-		// Last entry take the rest
-		uint32 ActualBlockSize = (BlockIndex == (BlockCount - 1)) ? TotalSize - CurrentBaseSlot : BlockSize;
-		FreeBlocks.Enqueue(new FD3D12GlobalHeapBlock(CurrentBaseSlot, ActualBlockSize));
-		CurrentBaseSlot += ActualBlockSize;
-	}
-}
-
-
-/**
-Allocate a new heap block - will also check if released blocks can be freed again
-**/
-FD3D12GlobalHeapBlock* FD3D12GlobalHeap::AllocateHeapBlock()
-{
-	SCOPED_NAMED_EVENT(FD3D12GlobalHeap_AllocateHeapBlock, FColor::Silver);
-
-	FScopeLock Lock(&CriticalSection);
-
-	// Check if certain released blocks are free again
-	UpdateFreeBlocks();
-
-	// Free block
-	FD3D12GlobalHeapBlock* Result = nullptr;
-	FreeBlocks.Dequeue(Result);
-
-	if (Result)
-	{
-		// Update stats
-		INC_DWORD_STAT(STAT_GlobalViewHeapBlockAllocations);
-		DEC_DWORD_STAT_BY(STAT_GlobalViewHeapFreeDescriptors, Result->Size);
-		INC_DWORD_STAT_BY(STAT_GlobalViewHeapReservedDescriptors, Result->Size);
-	}
-
-	return Result;
-}
-
-
-/**
-Free given block - can still be used by the GPU (SyncPoint needs to be setup by the caller and will be used to check if the block can be reused again)
-**/
-void FD3D12GlobalHeap::FreeHeapBlock(FD3D12GlobalHeapBlock* InHeapBlock)
+FD3D12SubAllocatedOnlineHeap::FD3D12SubAllocatedOnlineHeap(FD3D12DescriptorCache* InDescriptorCache)
+	: FD3D12OnlineHeap(nullptr, false)
+	, DescriptorCache(InDescriptorCache)
 {
-	FScopeLock Lock(&CriticalSection);
-
-	// Update stats
-	DEC_DWORD_STAT_BY(STAT_GlobalViewHeapReservedDescriptors, InHeapBlock->Size);
-	INC_DWORD_STAT_BY(STAT_GlobalViewHeapUsedDescriptors, InHeapBlock->SizeUsed);
-	INC_DWORD_STAT_BY(STAT_GlobalViewHeapWastedDescriptors, InHeapBlock->Size - InHeapBlock->SizeUsed);
-
-	ReleasedBlocks.Add(InHeapBlock);
 }
 
+FD3D12SubAllocatedOnlineHeap::~FD3D12SubAllocatedOnlineHeap() = default;
 
-/**
-Find all the blocks which are not used by the GPU anymore
-**/
-void FD3D12GlobalHeap::UpdateFreeBlocks()
-{
-	for (int32 BlockIndex = 0; BlockIndex < ReleasedBlocks.Num(); ++BlockIndex)
-	{
-		// Check if GPU is ready consuming the block data
-		FD3D12GlobalHeapBlock* ReleasedBlock = ReleasedBlocks[BlockIndex];
-		if (ReleasedBlock->SyncPoint.IsComplete())
-		{
-			// Update stats
-			DEC_DWORD_STAT_BY(STAT_GlobalViewHeapUsedDescriptors, ReleasedBlock->SizeUsed);
-			DEC_DWORD_STAT_BY(STAT_GlobalViewHeapWastedDescriptors, ReleasedBlock->Size - ReleasedBlock->SizeUsed);
-			INC_DWORD_STAT_BY(STAT_GlobalViewHeapFreeDescriptors, ReleasedBlock->Size);
-
-			ReleasedBlock->SizeUsed = 0;
-			FreeBlocks.Enqueue(ReleasedBlock);
-
-			// don't want to resize, but optional parameter is missing
-			ReleasedBlocks.RemoveAtSwap(BlockIndex);
-			BlockIndex--;
-		}
-	}
-
-}
-
-
-////////////////////////////////////////////////////////////////////////////////////////////////////////////////////////
-// FD3D12SubAllocatedOnlineHeap
-////////////////////////////////////////////////////////////////////////////////////////////////////////////////////////
-
-/**
-Initialize the sub allocated online heap
-**/
-void FD3D12SubAllocatedOnlineHeap::Init(FD3D12Device* InDevice, D3D12_DESCRIPTOR_HEAP_TYPE InHeapType)
+/** Initialize the sub allocated online heap */
+void FD3D12SubAllocatedOnlineHeap::Init(FD3D12Device* InDevice)
 {
 	SetParentDevice(InDevice);
-	HeapType = InHeapType;
 }
 
-
-/**
-Handle roll over on the sub allocated online heap - needs a new block
-**/
+/** Handle roll over on the sub allocated online heap - needs a new block */
 bool FD3D12SubAllocatedOnlineHeap::RollOver()
 {
 	// Try and allocate a new block from the global heap
@@ -1132,10 +984,7 @@ bool FD3D12SubAllocatedOnlineHeap::RollOver()
 	return CurrentBlock == nullptr;
 }
 
-
-/**
-Set the current command list which needs to be notified about changes
-**/
+/** Set the current command list which needs to be notified about changes */
 void FD3D12SubAllocatedOnlineHeap::SetCurrentCommandList(const FD3D12CommandListHandle& CommandListHandle)
 {
 	// Update the current command list
@@ -1148,13 +997,10 @@ void FD3D12SubAllocatedOnlineHeap::SetCurrentCommandList(const FD3D12CommandList
 	}
 }
 
-
-/**
-Tries to allocate a new block from the global heap - if it fails then it will switch to thread local view heap
-**/
+/** Tries to allocate a new block from the global heap - if it fails then it will switch to thread local view heap */
 bool FD3D12SubAllocatedOnlineHeap::AllocateBlock()
 {
-	FD3D12GlobalHeap& GlobalHeap = GetParentDevice()->GetGlobalViewHeap();
+	FD3D12OnlineDescriptorManager& OnlineManager = GetParentDevice()->GetOnlineDescriptorManager();
 
 	// If we still have a block, then free it first
 	if (CurrentBlock)
@@ -1166,12 +1012,12 @@ bool FD3D12SubAllocatedOnlineHeap::AllocateBlock()
 		// Create the sync point on the current command list
 		CurrentBlock->SyncPoint = FD3D12CLSyncPoint(CurrentCommandList);
 
-		GlobalHeap.FreeHeapBlock(CurrentBlock);		
+		OnlineManager.FreeHeapBlock(CurrentBlock);
 		CurrentBlock = nullptr;
 	}
 
 	// Try and allocate from the global heap
-	CurrentBlock = GlobalHeap.AllocateHeapBlock();
+	CurrentBlock = OnlineManager.AllocateHeapBlock();
 
 	// Reset counters
 	NextSlotIndex = 0;
@@ -1181,11 +1027,7 @@ bool FD3D12SubAllocatedOnlineHeap::AllocateBlock()
 	// Extract global heap data
 	if (CurrentBlock)
 	{
-		DescriptorSize = GlobalHeap.GetDescriptorSize();
-		CPUBase = GlobalHeap.GetCPUSlotHandle(CurrentBlock);
-		GPUBase = GlobalHeap.GetGPUSlotHandle(CurrentBlock);
-		Heap = GlobalHeap.GetHeap();
-		Desc = Heap->GetDesc();
+		Heap = new FD3D12DescriptorHeap(OnlineManager.GetDescriptorHeap(), CurrentBlock->BaseSlot, CurrentBlock->Size);
 	}
 	else
 	{
@@ -1202,38 +1044,39 @@ bool FD3D12SubAllocatedOnlineHeap::AllocateBlock()
 
 ////////////////////////////////////////////////////////////////////////////////////////////////////////////////////////
 // FD3D12LocalOnlineHeap
-////////////////////////////////////////////////////////////////////////////////////////////////////////////////////////
+
+FD3D12LocalOnlineHeap::FD3D12LocalOnlineHeap(FD3D12Device* Device, FD3D12DescriptorCache* InDescriptorCache)
+	: FD3D12OnlineHeap(Device, true)
+	, DescriptorCache(InDescriptorCache)
+{
+}
+
+FD3D12LocalOnlineHeap::~FD3D12LocalOnlineHeap() = default;
 
 /**
 Initialize a thread local online heap
 **/
-void FD3D12LocalOnlineHeap::Init(uint32 NumDescriptors, D3D12_DESCRIPTOR_HEAP_TYPE Type)
+void FD3D12LocalOnlineHeap::Init(uint32 InNumDescriptors, ED3D12DescriptorHeapType InHeapType)
 {
-	Desc = {};
-	Desc.Flags = D3D12_DESCRIPTOR_HEAP_FLAG_SHADER_VISIBLE;
-	Desc.Type = Type;
-	Desc.NumDescriptors = NumDescriptors;
-	Desc.NodeMask = GetGPUMask().GetNative();
-
-	//LLM_SCOPE(ELLMTag::DescriptorCache);
-	VERIFYD3D12RESULT(GetParentDevice()->GetDevice()->CreateDescriptorHeap(&Desc, IID_PPV_ARGS(Heap.GetInitReference())));
-	SetName(Heap, Desc.Type == D3D12_DESCRIPTOR_HEAP_TYPE_CBV_SRV_UAV ? L"Thread Local - Online View Heap" : L"Thread Local - Online Sampler Heap");
+	const TCHAR* DebugName = InHeapType == ED3D12DescriptorHeapType::Standard ? L"Thread Local - Online View Heap" : L"Thread Local - Online Sampler Heap";
+	Heap = GetParentDevice()->GetDescriptorHeapManager().AllocateHeap(
+		DebugName,
+		InHeapType,
+		InNumDescriptors,
+		D3D12_DESCRIPTOR_HEAP_FLAG_SHADER_VISIBLE
+	);
 
 	Entry.Heap = Heap;
 
-	CPUBase = Heap->GetCPUDescriptorHandleForHeapStart();
-	GPUBase = Heap->GetGPUDescriptorHandleForHeapStart();
-	DescriptorSize = GetParentDevice()->GetDevice()->GetDescriptorHandleIncrementSize(Type);
-
-	if (Desc.Type == D3D12_DESCRIPTOR_HEAP_TYPE_CBV_SRV_UAV)
+	if (InHeapType == ED3D12DescriptorHeapType::Standard)
 	{
 		INC_DWORD_STAT(STAT_NumViewOnlineDescriptorHeaps);
-		INC_MEMORY_STAT_BY(STAT_ViewOnlineDescriptorHeapMemory, Desc.NumDescriptors * GetDescriptorSize());
+		INC_MEMORY_STAT_BY(STAT_ViewOnlineDescriptorHeapMemory, Heap->GetMemorySize());
 	}
 	else
 	{
 		INC_DWORD_STAT(STAT_NumSamplerOnlineDescriptorHeaps);
-		INC_MEMORY_STAT_BY(STAT_SamplerOnlineDescriptorHeapMemory, Desc.NumDescriptors * GetDescriptorSize());
+		INC_MEMORY_STAT_BY(STAT_SamplerOnlineDescriptorHeapMemory, Heap->GetMemorySize());
 	}
 }
 
@@ -1260,18 +1103,26 @@ bool FD3D12LocalOnlineHeap::RollOver()
 
 		//LLM_SCOPE(ELLMTag::DescriptorCache);
 
-		VERIFYD3D12RESULT(GetParentDevice()->GetDevice()->CreateDescriptorHeap(&Desc, IID_PPV_ARGS(Heap.GetInitReference())));
-		SetName(Heap, Desc.Type == D3D12_DESCRIPTOR_HEAP_TYPE_CBV_SRV_UAV ? L"Thread Local - Online View Heap" : L"Thread Local - Online Sampler Heap");
+		const ED3D12DescriptorHeapType HeapType = Heap->GetType();
+		const uint32 NumDescriptors = Heap->GetNumDescriptors();
+
+		const TCHAR* DebugName = HeapType == ED3D12DescriptorHeapType::Standard ? L"Thread Local - Online View Heap" : L"Thread Local - Online Sampler Heap";
+		Heap = GetParentDevice()->GetDescriptorHeapManager().AllocateHeap(
+			DebugName,
+			HeapType,
+			NumDescriptors,
+			D3D12_DESCRIPTOR_HEAP_FLAG_SHADER_VISIBLE
+		);
 
-		if (Desc.Type == D3D12_DESCRIPTOR_HEAP_TYPE_CBV_SRV_UAV)
+		if (HeapType == ED3D12DescriptorHeapType::Standard)
 		{
 			INC_DWORD_STAT(STAT_NumViewOnlineDescriptorHeaps);
-			INC_MEMORY_STAT_BY(STAT_ViewOnlineDescriptorHeapMemory, Desc.NumDescriptors * GetDescriptorSize());
+			INC_MEMORY_STAT_BY(STAT_ViewOnlineDescriptorHeapMemory, Heap->GetMemorySize());
 		}
 		else
 		{
 			INC_DWORD_STAT(STAT_NumSamplerOnlineDescriptorHeaps);
-			INC_MEMORY_STAT_BY(STAT_SamplerOnlineDescriptorHeapMemory, Desc.NumDescriptors * GetDescriptorSize());
+			INC_MEMORY_STAT_BY(STAT_SamplerOnlineDescriptorHeapMemory, Heap->GetMemorySize());
 		}
 
 		Entry.Heap = Heap;
@@ -1280,10 +1131,7 @@ bool FD3D12LocalOnlineHeap::RollOver()
 	NextSlotIndex = 0;
 	FirstUsedSlot = 0;
 
-	// Notify other layers of heap change
-	CPUBase = Heap->GetCPUDescriptorHandleForHeapStart();
-	GPUBase = Heap->GetGPUDescriptorHandleForHeapStart();
-	return DescriptorCache->HeapRolledOver(Desc.Type);
+	return DescriptorCache->HeapRolledOver(Heap->GetType());
 }
 
 
@@ -1292,7 +1140,7 @@ Handle loop around on the heap
 **/
 void FD3D12LocalOnlineHeap::HeapLoopedAround()
 {
-	DescriptorCache->HeapLoopedAround(Desc.Type);
+	DescriptorCache->HeapLoopedAround(Heap->GetType());
 }
 
 
diff --git a/Engine/Source/Runtime/D3D12RHI/Private/D3D12DescriptorCache.h b/Engine/Source/Runtime/D3D12RHI/Private/D3D12DescriptorCache.h
--- a/Engine/Source/Runtime/D3D12RHI/Private/D3D12DescriptorCache.h
+++ b/Engine/Source/Runtime/D3D12RHI/Private/D3D12DescriptorCache.h
@@ -1,11 +1,9 @@
 // Copyright Epic Games, Inc. All Rights Reserved.
 
-/*=============================================================================
-D3D12DescriptorCache.h: D3D12 State application functionality
-=============================================================================*/
 #pragma once
 
 class FD3D12DynamicRHI;
+class FD3D12DescriptorCache;
 struct FD3D12VertexBufferCache;
 struct FD3D12IndexBufferCache;
 struct FD3D12ConstantBufferCache;
@@ -160,173 +158,17 @@ typedef FD3D12UniqueDescriptorTable<MAX_SAMPLERS> FD3D12UniqueSamplerTable;
 
 typedef TSet<FD3D12UniqueSamplerTable, FD3D12UniqueDescriptorTableKeyFuncs<FD3D12UniqueSamplerTable>> FD3D12SamplerSet;
 
-class FD3D12DescriptorCache;
-
-class FD3D12OfflineDescriptorManager : public FD3D12SingleNodeGPUObject
-{
-public: // Types
-	typedef D3D12_CPU_DESCRIPTOR_HANDLE HeapOffset;
-	typedef decltype(HeapOffset::ptr) HeapOffsetRaw;
-	typedef uint32 HeapIndex;
-
-private: // Types
-	struct SFreeRange { HeapOffsetRaw Start; HeapOffsetRaw End; };
-	struct SHeapEntry
-	{
-		TRefCountPtr<ID3D12DescriptorHeap> m_Heap;
-		TDoubleLinkedList<SFreeRange> m_FreeList;
-
-		SHeapEntry() { }
-	};
-	typedef TArray<SHeapEntry> THeapMap;
-
-	static D3D12_DESCRIPTOR_HEAP_DESC CreateDescriptor(FRHIGPUMask Node, D3D12_DESCRIPTOR_HEAP_TYPE Type, uint32 NumDescriptorsPerHeap)
-	{
-		D3D12_DESCRIPTOR_HEAP_DESC Desc = {};
-		Desc.Type = Type;
-		Desc.NumDescriptors = NumDescriptorsPerHeap;
-		Desc.Flags = D3D12_DESCRIPTOR_HEAP_FLAG_NONE;// None as this heap is offline
-		Desc.NodeMask = Node.GetNative();
-
-		return Desc;
-	}
-
-public: // Methods
-	FD3D12OfflineDescriptorManager(FRHIGPUMask Node, D3D12_DESCRIPTOR_HEAP_TYPE Type, uint32 NumDescriptorsPerHeap)
-		: FD3D12SingleNodeGPUObject(Node)
-		, m_Desc(CreateDescriptor(Node, Type, NumDescriptorsPerHeap))
-		, m_DescriptorSize(0)
-		, m_pDevice(nullptr)
-	{}
-
-	void Init(ID3D12Device* pDevice)
-	{
-		m_pDevice = pDevice;
-		m_DescriptorSize = pDevice->GetDescriptorHandleIncrementSize(m_Desc.Type);
-	}
-
-	HeapOffset AllocateHeapSlot(HeapIndex &outIndex)
-	{
-		FScopeLock Lock(&CritSect);
-		if (0 == m_FreeHeaps.Num())
-		{
-			AllocateHeap();
-		}
-		check(0 != m_FreeHeaps.Num());
-		auto Head = m_FreeHeaps.GetHead();
-		outIndex = Head->GetValue();
-		SHeapEntry &HeapEntry = m_Heaps[outIndex];
-		check(0 != HeapEntry.m_FreeList.Num());
-		SFreeRange &Range = HeapEntry.m_FreeList.GetHead()->GetValue();
-		HeapOffset Ret = { Range.Start };
-		Range.Start += m_DescriptorSize;
-
-		if (Range.Start == Range.End)
-		{
-			HeapEntry.m_FreeList.RemoveNode(HeapEntry.m_FreeList.GetHead());
-			if (0 == HeapEntry.m_FreeList.Num())
-			{
-				m_FreeHeaps.RemoveNode(Head);
-			}
-		}
-		return Ret;
-	}
-
-	void FreeHeapSlot(HeapOffset Offset, HeapIndex index)
-	{
-		FScopeLock Lock(&CritSect);
-		SHeapEntry &HeapEntry = m_Heaps[index];
-
-		SFreeRange NewRange =
-		{
-			Offset.ptr,
-			Offset.ptr + m_DescriptorSize
-		};
-
-		bool bFound = false;
-		for (auto Node = HeapEntry.m_FreeList.GetHead();
-		Node != nullptr && !bFound;
-			Node = Node->GetNextNode())
-		{
-			SFreeRange &Range = Node->GetValue();
-			check(Range.Start < Range.End);
-			if (Range.Start == Offset.ptr + m_DescriptorSize)
-			{
-				Range.Start = Offset.ptr;
-				bFound = true;
-			}
-			else if (Range.End == Offset.ptr)
-			{
-				Range.End += m_DescriptorSize;
-				bFound = true;
-			}
-			else
-			{
-				check(Range.End < Offset.ptr || Range.Start > Offset.ptr);
-				if (Range.Start > Offset.ptr)
-				{
-					HeapEntry.m_FreeList.InsertNode(NewRange, Node);
-					bFound = true;
-				}
-			}
-		}
-
-		if (!bFound)
-		{
-			if (0 == HeapEntry.m_FreeList.Num())
-			{
-				m_FreeHeaps.AddTail(index);
-			}
-			HeapEntry.m_FreeList.AddTail(NewRange);
-		}
-	}
-
-private: // Methods
-	void AllocateHeap()
-	{
-		TRefCountPtr<ID3D12DescriptorHeap> Heap;
-		VERIFYD3D12RESULT(m_pDevice->CreateDescriptorHeap(&m_Desc, IID_PPV_ARGS(Heap.GetInitReference())));
-		SetName(Heap, L"FD3D12OfflineDescriptorManager Descriptor Heap");
-
-		HeapOffset HeapBase = Heap->GetCPUDescriptorHandleForHeapStart();
-		check(HeapBase.ptr != 0);
-
-		// Allocate and initialize a single new entry in the map
-		m_Heaps.SetNum(m_Heaps.Num() + 1);
-		SHeapEntry& HeapEntry = m_Heaps.Last();
-		HeapEntry.m_FreeList.AddTail({ HeapBase.ptr,
-			HeapBase.ptr + m_Desc.NumDescriptors * m_DescriptorSize });
-		HeapEntry.m_Heap = Heap;
-		m_FreeHeaps.AddTail(m_Heaps.Num() - 1);
-	}
-
-private: // Members
-	const D3D12_DESCRIPTOR_HEAP_DESC m_Desc;
-	uint32 m_DescriptorSize;
-	ID3D12Device* m_pDevice; // weak-ref
-
-	THeapMap m_Heaps;
-	TDoubleLinkedList<HeapIndex> m_FreeHeaps;
-	FCriticalSection CritSect;
-};
-
-
-/**
-Manages a D3D heap which is GPU visible - base class which can be used by the FD3D12DescriptorCache
-**/
-class FD3D12OnlineHeap : public FD3D12DeviceChild, public FD3D12SingleNodeGPUObject
+/** Manages a D3D heap which is GPU visible - base class which can be used by the FD3D12DescriptorCache */
+class FD3D12OnlineHeap : public FD3D12DeviceChild
 {
 public:
+	FD3D12OnlineHeap(FD3D12Device* Device, bool CanLoopAround);
+	virtual ~FD3D12OnlineHeap();
 
-	FD3D12OnlineHeap(FD3D12Device* Device, FRHIGPUMask Node, bool CanLoopAround);
-	virtual ~FD3D12OnlineHeap() { }
-
-	FORCEINLINE D3D12_CPU_DESCRIPTOR_HANDLE GetCPUSlotHandle(uint32 Slot) const { return{ CPUBase.ptr + Slot * DescriptorSize }; }
-	FORCEINLINE D3D12_GPU_DESCRIPTOR_HANDLE GetGPUSlotHandle(uint32 Slot) const { return{ GPUBase.ptr + Slot * DescriptorSize }; }
+	ID3D12DescriptorHeap* GetHeap() { return Heap->GetHeap(); }
 
-	inline const uint32 GetDescriptorSize() const { return DescriptorSize; }
-	ID3D12DescriptorHeap* GetHeap() { return Heap.GetReference(); }
-	const D3D12_DESCRIPTOR_HEAP_DESC& GetDesc() const { return Desc; }
+	FORCEINLINE D3D12_CPU_DESCRIPTOR_HANDLE GetCPUSlotHandle(uint32 Slot) const { return Heap->GetCPUSlotHandle(Slot); }
+	FORCEINLINE D3D12_GPU_DESCRIPTOR_HANDLE GetGPUSlotHandle(uint32 Slot) const { return Heap->GetGPUSlotHandle(Slot); }
 
 	// Call this to reserve descriptor heap slots for use by the command list you are currently recording. This will wait if
 	// necessary until slots are free (if they are currently in use by another command list.) If the reservation can be
@@ -342,45 +184,30 @@ public:
 	virtual bool RollOver() = 0;
 	virtual void HeapLoopedAround() { }
 	virtual void SetCurrentCommandList(const FD3D12CommandListHandle& CommandListHandle) { }
-	virtual uint32 GetTotalSize() { return Desc.NumDescriptors; }
+	virtual uint32 GetTotalSize() { return Heap->GetNumDescriptors(); }
 
 	static const uint32 HeapExhaustedValue = uint32(-1);
 
 protected:
-		
-	// Handles for manipulation of the heap
-	uint32 DescriptorSize;
-	D3D12_CPU_DESCRIPTOR_HANDLE CPUBase;
-	D3D12_GPU_DESCRIPTOR_HANDLE GPUBase;
-
-	// Does the heap support loop around allocations
-	const bool bCanLoopAround;
+	// Keeping this ptr around is basically just for lifetime management
+	TRefCountPtr<FD3D12DescriptorHeap> Heap;
 
 	// This index indicate where the next set of descriptors should be placed *if* there's room
-	uint32 NextSlotIndex;
+	uint32 NextSlotIndex = 0;
 
 	// Indicates the last free slot marked by the command list being finished
-	uint32 FirstUsedSlot;
-
-	// Keeping this ptr around is basically just for lifetime management
-	TRefCountPtr<ID3D12DescriptorHeap> Heap;
-
-	// Desc contains the number of slots and allows for easy recreation
-	D3D12_DESCRIPTOR_HEAP_DESC Desc;
+	uint32 FirstUsedSlot = 0;
 
+	// Does the heap support loop around allocations
+	const bool bCanLoopAround;
 };
 
-
-/**
-Global sampler heap managed by the device which stored a unique set of sampler sets
-**/
+/** Global sampler heap managed by the device which stored a unique set of sampler sets */
 class FD3D12GlobalOnlineSamplerHeap : public FD3D12OnlineHeap
 {
 public:
-	FD3D12GlobalOnlineSamplerHeap(FD3D12Device* Device, FRHIGPUMask Node)
-		: FD3D12OnlineHeap(Device, Node, false)
-		, bUniqueDescriptorTablesAreDirty(false)
-	{ }
+	FD3D12GlobalOnlineSamplerHeap(FD3D12Device* Device);
+	~FD3D12GlobalOnlineSamplerHeap();
 
 	void Init(uint32 TotalSize);
 
@@ -393,85 +220,20 @@ public:
 	virtual bool RollOver() final override;
 
 private:
-
 	FD3D12SamplerSet UniqueDescriptorTables;
-	bool bUniqueDescriptorTablesAreDirty;
-
 	FCriticalSection CriticalSection;
+	bool bUniqueDescriptorTablesAreDirty = false;
 };
 
-
-/**
-Heap sub block of a global heap
-**/
-struct FD3D12GlobalHeapBlock
-{
-public:
-	FD3D12GlobalHeapBlock(uint32 InBaseSlot, uint32 InSize) :
-		BaseSlot(InBaseSlot), Size(InSize), SizeUsed(0) {};
-
-	uint32 BaseSlot;
-	uint32 Size;
-	uint32 SizeUsed;
-	FD3D12CLSyncPoint SyncPoint;
-};
-
-
-/**
-Global per device heap from which sub blocks can be allocated and freed
-**/
-class FD3D12GlobalHeap : public FD3D12DeviceChild, public FD3D12SingleNodeGPUObject
-{
-public:
-	FD3D12GlobalHeap(FD3D12Device* Device, FRHIGPUMask Node)
-		: FD3D12DeviceChild(Device), FD3D12SingleNodeGPUObject(Node)
-	{ }
-
-	// Setup the actual heap
-	void Init(D3D12_DESCRIPTOR_HEAP_TYPE InType, uint32 InTotalSize);
-
-	// Allocate an available sub heap block from the global heap
-	FD3D12GlobalHeapBlock* AllocateHeapBlock();
-	void FreeHeapBlock(FD3D12GlobalHeapBlock* InHeapBlock);
-
-	// Get the CPU & GPU descriptor handle for specific slot index
-	uint32 GetDescriptorSize() const { return DescriptorSize; }
-	ID3D12DescriptorHeap* GetHeap() { return Heap.GetReference(); }
-	D3D12_CPU_DESCRIPTOR_HANDLE GetCPUSlotHandle(FD3D12GlobalHeapBlock* InBlock) const { return{ CPUBase.ptr + InBlock->BaseSlot * DescriptorSize }; }
-	D3D12_GPU_DESCRIPTOR_HANDLE GetGPUSlotHandle(FD3D12GlobalHeapBlock* InBlock) const { return{ GPUBase.ptr + InBlock->BaseSlot * DescriptorSize }; }
-
-private:
-
-	// Check all released blocks and check which ones are not used by the GPU anymore
-	void UpdateFreeBlocks();
-
-	D3D12_DESCRIPTOR_HEAP_TYPE Type = D3D12_DESCRIPTOR_HEAP_TYPE_NUM_TYPES;
-	uint32 TotalSize = 0;
-	TRefCountPtr<ID3D12DescriptorHeap> Heap;
-
-	uint32 DescriptorSize = 0;
-	D3D12_CPU_DESCRIPTOR_HANDLE CPUBase;
-	D3D12_GPU_DESCRIPTOR_HANDLE GPUBase;
-
-	TQueue<FD3D12GlobalHeapBlock*> FreeBlocks;
-	TArray<FD3D12GlobalHeapBlock*> ReleasedBlocks;
-
-	FCriticalSection CriticalSection;
-};
-
-
-/**
-Online heap which can be used by a FD3D12DescriptorCache to manage a block allocated from a GLobalHeap
-**/
+/** Online heap which can be used by a FD3D12DescriptorCache to manage a block allocated from a GlobalHeap */
 class FD3D12SubAllocatedOnlineHeap : public FD3D12OnlineHeap
 {
 public:
-
-	FD3D12SubAllocatedOnlineHeap(FRHIGPUMask InNode, FD3D12DescriptorCache* InDescriptorCache) :
-		FD3D12OnlineHeap(nullptr, InNode, false), DescriptorCache(InDescriptorCache) {};
+	FD3D12SubAllocatedOnlineHeap(FD3D12DescriptorCache* InDescriptorCache);
+	~FD3D12SubAllocatedOnlineHeap();
 
 	// Setup the online heap data
-	void Init(FD3D12Device* InDevice, D3D12_DESCRIPTOR_HEAP_TYPE InHeapType);
+	void Init(FD3D12Device* InDevice);
 
 	// Override FD3D12OnlineHeap functions
 	virtual bool RollOver() final override;
@@ -482,36 +244,32 @@ public:
 	}
 
 private:
-
 	// Allocate a new block from the global heap - return true if allocation succeeds
 	bool AllocateBlock();
 
-	D3D12_DESCRIPTOR_HEAP_TYPE HeapType = D3D12_DESCRIPTOR_HEAP_TYPE_NUM_TYPES;
-	FD3D12GlobalHeapBlock* CurrentBlock = nullptr;
+	FD3D12OnlineDescriptorBlock* CurrentBlock = nullptr;
 
 	FD3D12DescriptorCache* DescriptorCache = nullptr;
 	FD3D12CommandListHandle CurrentCommandList;
 };
 
 
-/**
-Online heap which is not shared between multiple FD3D12DescriptorCache - used as overflow heap when the global heaps are full or don't contain the required data
-**/
+/** Online heap which is not shared between multiple FD3D12DescriptorCache.
+ *  Used as overflow heap when the global heaps are full or don't contain the required data
+ */
 class FD3D12LocalOnlineHeap : public FD3D12OnlineHeap
 {
 public:
-	FD3D12LocalOnlineHeap(FD3D12Device* Device, FRHIGPUMask Node, FD3D12DescriptorCache* InDescriptorCache)
-		: FD3D12OnlineHeap(Device, Node, true), DescriptorCache(InDescriptorCache)
-	{ }
+	FD3D12LocalOnlineHeap(FD3D12Device* Device, FD3D12DescriptorCache* InDescriptorCache);
+	~FD3D12LocalOnlineHeap();
 
 	// Allocate the actual overflow heap
-	void Init(uint32 NumDescriptors, D3D12_DESCRIPTOR_HEAP_TYPE Type);
+	void Init(uint32 InNumDescriptors, ED3D12DescriptorHeapType InHeapType);
 
 	// Override FD3D12OnlineHeap functions
 	virtual bool RollOver() final override;
 	virtual void HeapLoopedAround() final override;
 	virtual void SetCurrentCommandList(const FD3D12CommandListHandle& CommandListHandle) final override;
-	virtual uint32 GetTotalSize() final override { return Desc.NumDescriptors; }	
 
 private:
 	struct SyncPointEntry
@@ -537,7 +295,7 @@ private:
 
 	struct PoolEntry
 	{
-		TRefCountPtr<ID3D12DescriptorHeap> Heap;
+		TRefCountPtr<FD3D12DescriptorHeap> Heap;
 		FD3D12CLSyncPoint SyncPoint;
 
 		PoolEntry() 
@@ -560,11 +318,6 @@ private:
 	FD3D12CommandListHandle CurrentCommandList;
 };
 
-
-//-----------------------------------------------------------------------------
-//	FD3D12DescriptorCache Class Definition
-//-----------------------------------------------------------------------------
-
 class FD3D12DescriptorCache : public FD3D12DeviceChild, public FD3D12SingleNodeGPUObject
 {
 protected:
@@ -581,16 +334,6 @@ public:
 		if (LocalViewHeap) { delete(LocalViewHeap); }
 	}
 
-	inline ID3D12DescriptorHeap* GetViewDescriptorHeap()
-	{
-		return CurrentViewHeap->GetHeap();
-	}
-
-	inline ID3D12DescriptorHeap* GetSamplerDescriptorHeap()
-	{
-		return CurrentSamplerHeap->GetHeap();
-	}
-
 	// Checks if the specified descriptor heap has been set on the current command list.
 	bool IsHeapSet(ID3D12DescriptorHeap* const pHeap) const
 	{
@@ -607,14 +350,14 @@ public:
 
 	// null views
 
-	FD3D12DescriptorHandleSRV* pNullSRV;
-	FD3D12DescriptorHandleRTV* pNullRTV;
-	FD3D12DescriptorHandleUAV* pNullUAV;
+	FD3D12ViewDescriptorHandle* NullSRV{};
+	FD3D12ViewDescriptorHandle* NullRTV{};
+	FD3D12ViewDescriptorHandle* NullUAV{};
 
 #if USE_STATIC_ROOT_SIGNATURE
-	FD3D12ConstantBufferView* pNullCBV;
+	FD3D12ConstantBufferView* NullCBV{};
 #endif
-	TRefCountPtr<FD3D12SamplerState> pDefaultSampler;
+	TRefCountPtr<FD3D12SamplerState> DefaultSampler;
 
 	void SetVertexBuffers(FD3D12VertexBufferCache& Cache);
 	void SetRenderTargets(FD3D12RenderTargetView** RenderTargetViewArray, uint32 Count, FD3D12DepthStencilView* DepthStencilTarget);
@@ -637,8 +380,8 @@ public:
 
 	void SetStreamOutTargets(FD3D12Resource **Buffers, uint32 Count, const uint32* Offsets);
 
-	bool HeapRolledOver(D3D12_DESCRIPTOR_HEAP_TYPE Type);
-	void HeapLoopedAround(D3D12_DESCRIPTOR_HEAP_TYPE Type);
+	bool HeapRolledOver(ED3D12DescriptorHeapType InHeapType);
+	void HeapLoopedAround(ED3D12DescriptorHeapType InHeapType);
 	void Init(FD3D12Device* InParent, FD3D12CommandContext* InCmdContext, uint32 InNumLocalViewDescriptors, uint32 InNumSamplerDescriptors);
 	void Clear();
 	void BeginFrame();
diff --git a/Engine/Source/Runtime/D3D12RHI/Private/D3D12Descriptors.cpp b/Engine/Source/Runtime/D3D12RHI/Private/D3D12Descriptors.cpp
new file mode 100644
--- /dev/null
+++ b/Engine/Source/Runtime/D3D12RHI/Private/D3D12Descriptors.cpp
@@ -0,0 +1,598 @@
+// Copyright Epic Games, Inc. All Rights Reserved.
+
+#include "D3D12RHIPrivate.h"
+#include "D3D12Descriptors.h"
+
+const TCHAR* ToString(ED3D12DescriptorHeapType InHeapType)
+{
+	switch (InHeapType)
+	{
+	default: checkNoEntry();                     return TEXT("UnknownDescriptorHeapType");
+	case ED3D12DescriptorHeapType::Standard:     return TEXT("Standard");
+	case ED3D12DescriptorHeapType::RenderTarget: return TEXT("RenderTarget");
+	case ED3D12DescriptorHeapType::DepthStencil: return TEXT("DepthStencil");
+	case ED3D12DescriptorHeapType::Sampler:      return TEXT("Sampler");
+	}
+}
+
+inline D3D12_DESCRIPTOR_HEAP_TYPE Translate(ED3D12DescriptorHeapType InHeapType)
+{
+	switch (InHeapType)
+	{
+	default: checkNoEntry();
+	case ED3D12DescriptorHeapType::Standard:     return D3D12_DESCRIPTOR_HEAP_TYPE_CBV_SRV_UAV;
+	case ED3D12DescriptorHeapType::RenderTarget: return D3D12_DESCRIPTOR_HEAP_TYPE_RTV;
+	case ED3D12DescriptorHeapType::DepthStencil: return D3D12_DESCRIPTOR_HEAP_TYPE_DSV;
+	case ED3D12DescriptorHeapType::Sampler:      return FD3D12_DESCRIPTOR_HEAP_TYPE_SAMPLER;
+	}
+}
+
+////////////////////////////////////////////////////////////////////////////////////////////////////////////////////////
+// FD3D12DescriptorHeap
+
+FD3D12DescriptorHeap::FD3D12DescriptorHeap(FD3D12Device* InDevice, ID3D12DescriptorHeap* InHeap, uint32 InNumDescriptors, ED3D12DescriptorHeapType InType, D3D12_DESCRIPTOR_HEAP_FLAGS InFlags, bool bInIsGlobal)
+	: FD3D12DeviceChild(InDevice)
+	, Heap(InHeap)
+	, CpuBase(InHeap->GetCPUDescriptorHandleForHeapStart())
+	, GpuBase((InFlags& D3D12_DESCRIPTOR_HEAP_FLAG_SHADER_VISIBLE) ? InHeap->GetGPUDescriptorHandleForHeapStart() : D3D12_GPU_DESCRIPTOR_HANDLE{})
+	, Offset(0)
+	, NumDescriptors(InNumDescriptors)
+	, DescriptorSize(InDevice->GetDevice()->GetDescriptorHandleIncrementSize(Translate(InType)))
+	, Type(InType)
+	, Flags(InFlags)
+	, bIsGlobal(bInIsGlobal)
+	, bIsSuballocation(false)
+{
+}
+
+FD3D12DescriptorHeap::FD3D12DescriptorHeap(FD3D12DescriptorHeap* SubAllocateSourceHeap, uint32 InOffset, uint32 InNumDescriptors)
+	: FD3D12DeviceChild(SubAllocateSourceHeap->GetParentDevice())
+	, Heap(SubAllocateSourceHeap->Heap)
+	, CpuBase(SubAllocateSourceHeap->CpuBase, InOffset, SubAllocateSourceHeap->DescriptorSize)
+	, GpuBase(SubAllocateSourceHeap->GpuBase, InOffset, SubAllocateSourceHeap->DescriptorSize)
+	, Offset(InOffset)
+	, NumDescriptors(InNumDescriptors)
+	, DescriptorSize(SubAllocateSourceHeap->DescriptorSize)
+	, Type(SubAllocateSourceHeap->GetType())
+	, Flags(SubAllocateSourceHeap->GetFlags())
+	, bIsGlobal(SubAllocateSourceHeap->IsGlobal())
+	, bIsSuballocation(true)
+{
+}
+
+FD3D12DescriptorHeap::~FD3D12DescriptorHeap() = default;
+
+////////////////////////////////////////////////////////////////////////////////////////////////////////////////////////
+// FDescriptorAllocator
+
+struct FDescriptorAllocatorRange
+{
+	FDescriptorAllocatorRange(uint32 InFirst, uint32 InLast) : First(InFirst), Last(InLast) {}
+	uint32 First;
+	uint32 Last;
+};
+
+FDescriptorAllocator::FDescriptorAllocator()
+{
+}
+
+FDescriptorAllocator::~FDescriptorAllocator()
+{
+}
+
+void FDescriptorAllocator::Init(uint32 InNumDescriptors)
+{
+	Capacity = InNumDescriptors;
+	Ranges.Emplace(0, InNumDescriptors - 1);
+}
+
+bool FDescriptorAllocator::Allocate(uint32 NumDescriptors, uint32& OutSlot)
+{
+	FScopeLock Lock(&CriticalSection);
+
+	if (const uint32 NumRanges = Ranges.Num(); NumRanges > 0)
+	{
+		uint32 Index = 0;
+		do
+		{
+			FDescriptorAllocatorRange& CurrentRange = Ranges[Index];
+			const uint32 Size = 1 + CurrentRange.Last - CurrentRange.First;
+			if (NumDescriptors <= Size)
+			{
+				uint32 First = CurrentRange.First;
+				if (NumDescriptors == Size && Index + 1 < NumRanges)
+				{
+					// Range is full and a new range exists, so move on to that one
+					Ranges.RemoveAt(Index);
+				}
+				else
+				{
+					CurrentRange.First += NumDescriptors;
+				}
+				OutSlot = First;
+				return true;
+			}
+			++Index;
+		} while (Index < NumRanges);
+	}
+
+	OutSlot = UINT_MAX;
+	return false;
+}
+
+void FDescriptorAllocator::Free(uint32 Offset, uint32 NumDescriptors)
+{
+	if (Offset == UINT_MAX || NumDescriptors == 0)
+	{
+		return;
+	}
+
+	FScopeLock Lock(&CriticalSection);
+
+	const uint32 End = Offset + NumDescriptors;
+	// Binary search of the range list
+	uint32 Index0 = 0;
+	uint32 Index1 = Ranges.Num() -1;
+	for (;;)
+	{
+		const uint32 Index = (Index0 + Index1) / 2;
+		if (Offset < Ranges[Index].First)
+		{
+			// Before current range, check if neighboring
+			if (End >= Ranges[Index].First)
+			{
+				check(End == Ranges[Index].First); // Can't overlap a range of free IDs
+				// Neighbor id, check if neighboring previous range too
+				if (Index > Index0 && Offset - 1 == Ranges[Index - 1].Last)
+				{
+					// Merge with previous range
+					Ranges[Index - 1].Last = Ranges[Index].Last;
+					Ranges.RemoveAt(Index);
+				}
+				else
+				{
+					// Just grow range
+					Ranges[Index].First = Offset;
+				}
+				return;
+			}
+			else
+			{
+				// Non-neighbor id
+				if (Index != Index0)
+				{
+					// Cull upper half of list
+					Index1 = Index - 1;
+				}
+				else
+				{
+					// Found our position in the list, insert the deleted range here
+					Ranges.EmplaceAt(Index, Offset, End -1);
+					return;
+				}
+			}
+		}
+		else if (Offset > Ranges[Index].Last)
+		{
+			// After current range, check if neighboring
+			if (Offset - 1 == Ranges[Index].Last)
+			{
+				// Neighbor id, check if neighboring next range too
+				if (Index < Index1 && End == Ranges[Index + 1].First)
+				{
+					// Merge with next range
+					Ranges[Index].Last = Ranges[Index + 1].Last;
+					Ranges.RemoveAt(Index);
+				}
+				else
+				{
+					// Just grow range
+					Ranges[Index].Last += NumDescriptors;
+				}
+				return;
+			}
+			else
+			{
+				// Non-neighbor id
+				if (Index != Index1)
+				{
+					// Cull bottom half of list
+					Index0 = Index + 1;
+				}
+				else
+				{
+					// Found our position in the list, insert the deleted range here
+					Ranges.EmplaceAt(Index + 1, Offset, End - 1);
+					return;
+				}
+			}
+		}
+		else
+		{
+			// Inside a free block, not a valid offset
+			checkNoEntry();
+		}
+	}
+}
+
+////////////////////////////////////////////////////////////////////////////////////////////////////////////////////////
+// FD3D12ResourceDescriptorManager
+
+FD3D12ResourceDescriptorManager::FD3D12ResourceDescriptorManager(FD3D12Device* Device)
+	: FD3D12DeviceChild(Device)
+{
+}
+
+FD3D12ResourceDescriptorManager::~FD3D12ResourceDescriptorManager() = default;
+
+void FD3D12ResourceDescriptorManager::Init(uint32 InTotalSize)
+{
+	if (InTotalSize > 0)
+	{
+		Heap = GetParentDevice()->GetDescriptorHeapManager().AllocateHeap(
+			TEXT("ResourceDescriptorHeap"),
+			ED3D12DescriptorHeapType::Standard,
+			InTotalSize,
+			D3D12_DESCRIPTOR_HEAP_FLAG_SHADER_VISIBLE
+		);
+
+		Allocator.Init(InTotalSize);
+	}
+}
+
+void FD3D12ResourceDescriptorManager::Destroy()
+{
+	if (Heap)
+	{
+		GetParentDevice()->GetDescriptorHeapManager().FreeHeap(Heap);
+		Heap.SafeRelease();
+	}
+}
+
+bool FD3D12ResourceDescriptorManager::AllocateDescriptor(uint32& OutSlot)
+{
+	return Allocator.Allocate(1, OutSlot);
+}
+
+bool FD3D12ResourceDescriptorManager::AllocateDescriptors(uint32 NumDescriptors, uint32& OutSlot)
+{
+	return Allocator.Allocate(NumDescriptors, OutSlot);
+}
+
+void FD3D12ResourceDescriptorManager::FreeDescriptor(uint32 Slot)
+{
+	Allocator.Free(Slot, 1);
+}
+
+void FD3D12ResourceDescriptorManager::FreeDescriptors(uint32 Slot, uint32 NumDescriptors)
+{
+	Allocator.Free(Slot, NumDescriptors);
+}
+
+
+////////////////////////////////////////////////////////////////////////////////////////////////////////////////////////
+// FD3D12OnlineDescriptorManager
+
+FD3D12OnlineDescriptorManager::FD3D12OnlineDescriptorManager(FD3D12Device* Device)
+	: FD3D12DeviceChild(Device)
+{
+}
+
+FD3D12OnlineDescriptorManager::~FD3D12OnlineDescriptorManager() = default;
+
+/** Allocate and initialize the online heap */
+void FD3D12OnlineDescriptorManager::Init(uint32 InTotalSize, uint32 InBlockSize)
+{
+	Heap = GetParentDevice()->GetDescriptorHeapManager().AllocateHeap(
+		TEXT("Device Global - Online View Heap"),
+		ED3D12DescriptorHeapType::Standard,
+		InTotalSize,
+		D3D12_DESCRIPTOR_HEAP_FLAG_SHADER_VISIBLE);
+
+	// Update the stats
+	INC_DWORD_STAT(STAT_NumViewOnlineDescriptorHeaps);
+	INC_MEMORY_STAT_BY(STAT_ViewOnlineDescriptorHeapMemory, Heap->GetMemorySize());
+	INC_DWORD_STAT_BY(STAT_GlobalViewHeapFreeDescriptors, InTotalSize);
+
+	// Compute amount of free blocks
+	uint32 BlockCount = InTotalSize / InBlockSize;
+	ReleasedBlocks.Reserve(BlockCount);
+
+	// Allocate the free blocks
+	uint32 CurrentBaseSlot = 0;
+	for (uint32 BlockIndex = 0; BlockIndex < BlockCount; ++BlockIndex)
+	{
+		// Last entry take the rest
+		uint32 ActualBlockSize = (BlockIndex == (BlockCount - 1)) ? InTotalSize - CurrentBaseSlot : InBlockSize;
+		FreeBlocks.Enqueue(new FD3D12OnlineDescriptorBlock(CurrentBaseSlot, ActualBlockSize));
+		CurrentBaseSlot += ActualBlockSize;
+	}
+}
+
+/** Allocate a new heap block - will also check if released blocks can be freed again */
+FD3D12OnlineDescriptorBlock* FD3D12OnlineDescriptorManager::AllocateHeapBlock()
+{
+	SCOPED_NAMED_EVENT(FD3D12OnlineViewHeap_AllocateHeapBlock, FColor::Silver);
+
+	FScopeLock Lock(&CriticalSection);
+
+	// Check if certain released blocks are free again
+	UpdateFreeBlocks();
+
+	// Free block
+	FD3D12OnlineDescriptorBlock* Result = nullptr;
+	FreeBlocks.Dequeue(Result);
+
+	if (Result)
+	{
+		// Update stats
+		INC_DWORD_STAT(STAT_GlobalViewHeapBlockAllocations);
+		DEC_DWORD_STAT_BY(STAT_GlobalViewHeapFreeDescriptors, Result->Size);
+		INC_DWORD_STAT_BY(STAT_GlobalViewHeapReservedDescriptors, Result->Size);
+	}
+
+	return Result;
+}
+
+/** Free given block - can still be used by the GPU (SyncPoint needs to be setup by the caller and will be used to check if the block can be reused again) */
+void FD3D12OnlineDescriptorManager::FreeHeapBlock(FD3D12OnlineDescriptorBlock * InHeapBlock)
+{
+	FScopeLock Lock(&CriticalSection);
+
+	// Update stats
+	DEC_DWORD_STAT_BY(STAT_GlobalViewHeapReservedDescriptors, InHeapBlock->Size);
+	INC_DWORD_STAT_BY(STAT_GlobalViewHeapUsedDescriptors, InHeapBlock->SizeUsed);
+	INC_DWORD_STAT_BY(STAT_GlobalViewHeapWastedDescriptors, InHeapBlock->Size - InHeapBlock->SizeUsed);
+
+	ReleasedBlocks.Add(InHeapBlock);
+}
+
+/** Find all the blocks which are not used by the GPU anymore */
+void FD3D12OnlineDescriptorManager::UpdateFreeBlocks()
+{
+	for (int32 BlockIndex = 0; BlockIndex < ReleasedBlocks.Num(); )
+	{
+		// Check if GPU is ready consuming the block data
+		FD3D12OnlineDescriptorBlock* ReleasedBlock = ReleasedBlocks[BlockIndex];
+		if (ReleasedBlock->SyncPoint.IsComplete())
+		{
+			// Update stats
+			DEC_DWORD_STAT_BY(STAT_GlobalViewHeapUsedDescriptors, ReleasedBlock->SizeUsed);
+			DEC_DWORD_STAT_BY(STAT_GlobalViewHeapWastedDescriptors, ReleasedBlock->Size - ReleasedBlock->SizeUsed);
+			INC_DWORD_STAT_BY(STAT_GlobalViewHeapFreeDescriptors, ReleasedBlock->Size);
+
+			ReleasedBlock->SizeUsed = 0;
+			FreeBlocks.Enqueue(ReleasedBlock);
+
+			// don't want to resize, but optional parameter is missing
+			ReleasedBlocks.RemoveAtSwap(BlockIndex);
+		}
+		else
+		{
+			BlockIndex++;
+		}
+	}
+}
+
+////////////////////////////////////////////////////////////////////////////////////////////////////////////////////////
+// FD3D12OfflineDescriptorManager
+
+static uint32 GetOfflineDescriptorHeapDefaultSize(ED3D12DescriptorHeapType InHeapType)
+{
+	switch (InHeapType)
+	{
+	default: checkNoEntry();
+#if USE_STATIC_ROOT_SIGNATURE
+	case ED3D12DescriptorHeapType::Standard:     return 4096;
+#else
+	case ED3D12DescriptorHeapType::Standard:     return 2048;
+#endif
+	case ED3D12DescriptorHeapType::RenderTarget: return 256;
+	case ED3D12DescriptorHeapType::DepthStencil: return 256;
+	case ED3D12DescriptorHeapType::Sampler:      return 128;
+	}
+}
+
+struct FD3D12OfflineHeapFreeRange
+{
+	size_t Start;
+	size_t End;
+};
+
+struct FD3D12OfflineHeapEntry
+{
+	FD3D12OfflineHeapEntry(FD3D12DescriptorHeap* InHeap, const D3D12_CPU_DESCRIPTOR_HANDLE& InHeapBase, size_t InSize)
+		: Heap(InHeap)
+	{
+		FreeList.AddTail({ InHeapBase.ptr, InHeapBase.ptr + InSize });
+	}
+
+	TRefCountPtr<FD3D12DescriptorHeap> Heap;
+	TDoubleLinkedList<FD3D12OfflineHeapFreeRange> FreeList;
+};
+
+FD3D12OfflineDescriptorManager::FD3D12OfflineDescriptorManager(FD3D12Device* InDevice)
+	: FD3D12DeviceChild(InDevice)
+{
+}
+
+FD3D12OfflineDescriptorManager::~FD3D12OfflineDescriptorManager() = default;
+
+void FD3D12OfflineDescriptorManager::Init(ED3D12DescriptorHeapType InHeapType)
+{
+	HeapType = InHeapType;
+	DescriptorSize = GetParentDevice()->GetDevice()->GetDescriptorHandleIncrementSize(Translate(InHeapType));
+	NumDescriptorsPerHeap = GetOfflineDescriptorHeapDefaultSize(InHeapType);
+}
+
+void FD3D12OfflineDescriptorManager::AllocateHeap()
+{
+	checkf(NumDescriptorsPerHeap != 0, TEXT("Init() needs to be called before allocating heaps."));
+	TRefCountPtr<FD3D12DescriptorHeap> Heap = GetParentDevice()->GetDescriptorHeapManager().AllocateHeap(
+		TEXT("FD3D12OfflineDescriptorManager"),
+		HeapType,
+		NumDescriptorsPerHeap,
+		D3D12_DESCRIPTOR_HEAP_FLAG_NONE
+	);
+
+	D3D12_CPU_DESCRIPTOR_HANDLE HeapBase = Heap->GetCPUSlotHandle(0);
+	check(HeapBase.ptr != 0);
+
+	// Allocate and initialize a single new entry in the map
+	const uint32 NewHeapIndex = Heaps.Num();
+
+	Heaps.Emplace(Heap, HeapBase, NumDescriptorsPerHeap * DescriptorSize);
+	FreeHeaps.AddTail(NewHeapIndex);
+}
+
+D3D12_CPU_DESCRIPTOR_HANDLE FD3D12OfflineDescriptorManager::AllocateHeapSlot(uint32& OutIndex)
+{
+	FScopeLock Lock(&CriticalSection);
+	if (FreeHeaps.Num() == 0)
+	{
+		AllocateHeap();
+	}
+
+	check(FreeHeaps.Num() != 0);
+
+	auto Head = FreeHeaps.GetHead();
+	OutIndex = Head->GetValue();
+
+	FD3D12OfflineHeapEntry& HeapEntry = Heaps[OutIndex];
+	check(0 != HeapEntry.FreeList.Num());
+	FD3D12OfflineHeapFreeRange& Range = HeapEntry.FreeList.GetHead()->GetValue();
+
+	const D3D12_CPU_DESCRIPTOR_HANDLE Ret{ Range.Start };
+	Range.Start += DescriptorSize;
+
+	if (Range.Start == Range.End)
+	{
+		HeapEntry.FreeList.RemoveNode(HeapEntry.FreeList.GetHead());
+		if (HeapEntry.FreeList.Num() == 0)
+		{
+			FreeHeaps.RemoveNode(Head);
+		}
+	}
+	return Ret;
+}
+
+void FD3D12OfflineDescriptorManager::FreeHeapSlot(D3D12_CPU_DESCRIPTOR_HANDLE Offset, uint32 InIndex)
+{
+	FScopeLock Lock(&CriticalSection);
+	FD3D12OfflineHeapEntry& HeapEntry = Heaps[InIndex];
+
+	const FD3D12OfflineHeapFreeRange NewRange{ Offset.ptr, Offset.ptr + DescriptorSize };
+
+	bool bFound = false;
+	for (auto Node = HeapEntry.FreeList.GetHead();
+		Node != nullptr && !bFound;
+		Node = Node->GetNextNode())
+	{
+		FD3D12OfflineHeapFreeRange& Range = Node->GetValue();
+		check(Range.Start < Range.End);
+		if (Range.Start == Offset.ptr + DescriptorSize)
+		{
+			Range.Start = Offset.ptr;
+			bFound = true;
+		}
+		else if (Range.End == Offset.ptr)
+		{
+			Range.End += DescriptorSize;
+			bFound = true;
+		}
+		else
+		{
+			check(Range.End < Offset.ptr || Range.Start > Offset.ptr);
+			if (Range.Start > Offset.ptr)
+			{
+				HeapEntry.FreeList.InsertNode(NewRange, Node);
+				bFound = true;
+			}
+		}
+	}
+
+	if (!bFound)
+	{
+		if (HeapEntry.FreeList.Num() == 0)
+		{
+			FreeHeaps.AddTail(InIndex);
+		}
+		HeapEntry.FreeList.AddTail(NewRange);
+	}
+}
+
+////////////////////////////////////////////////////////////////////////////////////////////////////////////////////////
+// FD3D12DescriptorHeapManager
+
+FD3D12DescriptorHeapManager::FD3D12DescriptorHeapManager(FD3D12Device* InDevice)
+	: FD3D12DeviceChild(InDevice)
+{
+}
+
+FD3D12DescriptorHeapManager::~FD3D12DescriptorHeapManager()
+{
+	Destroy();
+}
+
+static FD3D12DescriptorHeap* CreateDescriptorHeap(FD3D12Device* Device, const TCHAR* DebugName, ED3D12DescriptorHeapType HeapType, uint32 NumDescriptors, D3D12_DESCRIPTOR_HEAP_FLAGS Flags, bool bIsGlobal)
+{
+	D3D12_DESCRIPTOR_HEAP_DESC Desc{};
+	Desc.Type = Translate(HeapType);
+	Desc.NumDescriptors = NumDescriptors;
+	Desc.Flags = Flags;
+	Desc.NodeMask = Device->GetGPUMask().GetNative();
+
+	TRefCountPtr<ID3D12DescriptorHeap> Heap;
+	VERIFYD3D12RESULT(Device->GetDevice()->CreateDescriptorHeap(&Desc, IID_PPV_ARGS(Heap.GetInitReference())));
+
+	SetName(Heap, DebugName);
+
+	return new FD3D12DescriptorHeap(Device, Heap, NumDescriptors, HeapType, Flags, bIsGlobal);
+}
+
+void FD3D12DescriptorHeapManager::Init(uint32 InNumGlobalDescriptors)
+{
+	if (InNumGlobalDescriptors > 0)
+	{
+		GlobalHeap = CreateDescriptorHeap(
+			GetParentDevice(),
+			TEXT("ViewDescriptorHeap"),
+			ED3D12DescriptorHeapType::Standard,
+			InNumGlobalDescriptors,
+			D3D12_DESCRIPTOR_HEAP_FLAG_SHADER_VISIBLE,
+			true /* bIsGlobal */
+		);
+
+		Allocator.Init(InNumGlobalDescriptors);
+	}
+}
+
+void FD3D12DescriptorHeapManager::Destroy()
+{
+}
+
+FD3D12DescriptorHeap* FD3D12DescriptorHeapManager::AllocateHeap(const TCHAR* InDebugName, ED3D12DescriptorHeapType InHeapType, uint32 InNumDescriptors, D3D12_DESCRIPTOR_HEAP_FLAGS InHeapFlags)
+{
+	if (GlobalHeap && InHeapType == GlobalHeap->GetType() && InHeapFlags == GlobalHeap->GetFlags())
+	{
+		uint32 Offset = 0;
+		if (Allocator.Allocate(InNumDescriptors, Offset))
+		{
+			return new FD3D12DescriptorHeap(GlobalHeap, Offset, InNumDescriptors);
+		}
+
+		// TODO: handle running out of space
+		checkNoEntry();
+	}
+
+	return CreateDescriptorHeap(GetParentDevice(), InDebugName, InHeapType, InNumDescriptors, InHeapFlags, false);
+}
+
+void FD3D12DescriptorHeapManager::FreeHeap(FD3D12DescriptorHeap* InHeap)
+{
+	if (InHeap->IsGlobal())
+	{
+		check(InHeap->GetHeap() == GlobalHeap->GetHeap());
+		Allocator.Free(InHeap->GetOffset(), InHeap->GetNumDescriptors());
+	}
+}
diff --git a/Engine/Source/Runtime/D3D12RHI/Private/D3D12Descriptors.h b/Engine/Source/Runtime/D3D12RHI/Private/D3D12Descriptors.h
new file mode 100644
--- /dev/null
+++ b/Engine/Source/Runtime/D3D12RHI/Private/D3D12Descriptors.h
@@ -0,0 +1,213 @@
+// Copyright Epic Games, Inc. All Rights Reserved.
+
+#pragma once
+
+#include "D3D12CommandList.h"
+
+struct FD3D12OfflineHeapEntry;
+
+// Internal platform agnostic Descriptor Heap type.
+enum class ED3D12DescriptorHeapType
+{
+	Standard,
+	RenderTarget,
+	DepthStencil,
+	Sampler,
+	count
+};
+const TCHAR* ToString(ED3D12DescriptorHeapType InHeapType);
+
+struct FD3D12DescriptorHeap : public FD3D12DeviceChild, public FD3D12RefCount
+{
+public:
+	FD3D12DescriptorHeap() = delete;
+
+	// Heap created with its own D3D heap object.
+	FD3D12DescriptorHeap(FD3D12Device* InDevice, ID3D12DescriptorHeap* InHeap, uint32 InNumDescriptors, ED3D12DescriptorHeapType InType, D3D12_DESCRIPTOR_HEAP_FLAGS InFlags, bool bInIsGlobal);
+
+	// Heap created as a suballocation of another heap.
+	FD3D12DescriptorHeap(FD3D12DescriptorHeap* SubAllocateSourceHeap, uint32 InOffset, uint32 InNumDescriptors);
+
+	~FD3D12DescriptorHeap();
+
+	inline ID3D12DescriptorHeap*       GetHeap()  const { return Heap; }
+	inline ED3D12DescriptorHeapType    GetType()  const { return Type; }
+	inline D3D12_DESCRIPTOR_HEAP_FLAGS GetFlags() const { return Flags; }
+
+	inline uint32 GetOffset()         const { return Offset; }
+	inline uint32 GetNumDescriptors() const { return NumDescriptors; }
+	inline uint32 GetDescriptorSize() const { return DescriptorSize; }
+	inline bool   IsGlobal()          const { return bIsGlobal; }
+	inline bool   IsSuballocation()   const { return bIsSuballocation; }
+
+	inline uint32 GetMemorySize() const { return DescriptorSize * NumDescriptors; }
+
+	inline D3D12_CPU_DESCRIPTOR_HANDLE GetCPUSlotHandle(uint32 Slot) const { return CD3DX12_CPU_DESCRIPTOR_HANDLE(CpuBase, Slot, DescriptorSize); }
+	inline D3D12_GPU_DESCRIPTOR_HANDLE GetGPUSlotHandle(uint32 Slot) const { return CD3DX12_GPU_DESCRIPTOR_HANDLE(GpuBase, Slot, DescriptorSize); }
+
+private:
+	TRefCountPtr<ID3D12DescriptorHeap> Heap;
+
+	const CD3DX12_CPU_DESCRIPTOR_HANDLE CpuBase;
+	const CD3DX12_GPU_DESCRIPTOR_HANDLE GpuBase;
+
+	// Offset in descriptors into the heap, only used when heap is suballocated.
+	const uint32 Offset = 0;
+
+	// Total number of descriptors in this heap.
+	const uint32 NumDescriptors;
+
+	// Device provided size of each descriptor in this heap.
+	const uint32 DescriptorSize;
+
+	const ED3D12DescriptorHeapType Type;
+	const D3D12_DESCRIPTOR_HEAP_FLAGS Flags;
+
+	// Enabled if this heap is the "global" heap.
+	const bool bIsGlobal;
+
+	// Enabled if this heap was allocated inside another heap.
+	const bool bIsSuballocation;
+};
+
+struct FDescriptorAllocatorRange;
+
+class FDescriptorAllocator
+{
+public:
+	FDescriptorAllocator();
+	~FDescriptorAllocator();
+
+	void Init(uint32 InNumDescriptors);
+
+	bool Allocate(uint32 NumDescriptors, uint32& OutSlot);
+	void Free(uint32 Slot, uint32 NumDescriptors);
+
+private:
+	TArray<FDescriptorAllocatorRange> Ranges;
+	uint32 Capacity = 0;
+
+	FCriticalSection CriticalSection;
+};
+
+/** Manager for resource descriptors used in bindless rendering. */
+class FD3D12ResourceDescriptorManager : public FD3D12DeviceChild
+{
+public:
+	FD3D12ResourceDescriptorManager(FD3D12Device* Device);
+	~FD3D12ResourceDescriptorManager();
+
+	void Init(uint32 InTotalSize);
+	void Destroy();
+
+	bool AllocateDescriptor(uint32& OutSlot);
+	bool AllocateDescriptors(uint32 NumDescriptors, uint32& OutSlot);
+
+	void FreeDescriptor(uint32 Slot);
+	void FreeDescriptors(uint32 Slot, uint32 NumDescriptors);
+
+	inline D3D12_CPU_DESCRIPTOR_HANDLE GetResourceHandle(uint32 InSlot) const { return Heap->GetCPUSlotHandle(InSlot); }
+
+private:
+	TRefCountPtr<FD3D12DescriptorHeap> Heap;
+
+	FDescriptorAllocator Allocator;
+};
+
+/** Heap sub block of an online heap */
+struct FD3D12OnlineDescriptorBlock
+{
+	FD3D12OnlineDescriptorBlock() = delete;
+	FD3D12OnlineDescriptorBlock(uint32 InBaseSlot, uint32 InSize)
+		: BaseSlot(InBaseSlot)
+		, Size(InSize)
+	{
+	}
+
+	uint32 BaseSlot;
+	uint32 Size;
+	uint32 SizeUsed = 0;
+	FD3D12CLSyncPoint SyncPoint;
+};
+
+/** Primary online heap from which sub blocks can be allocated and freed. Used when allocating blocks of descriptors for tables. */
+class FD3D12OnlineDescriptorManager : public FD3D12DeviceChild
+{
+public:
+	FD3D12OnlineDescriptorManager(FD3D12Device* Device);
+	~FD3D12OnlineDescriptorManager();
+
+	// Setup the actual heap
+	void Init(uint32 InTotalSize, uint32 InBlockSize);
+
+	// Allocate an available sub heap block from the global heap
+	FD3D12OnlineDescriptorBlock* AllocateHeapBlock();
+	void FreeHeapBlock(FD3D12OnlineDescriptorBlock* InHeapBlock);
+
+	ID3D12DescriptorHeap* GetHeap() { return Heap->GetHeap(); }
+	FD3D12DescriptorHeap* GetDescriptorHeap() { return Heap.GetReference(); }
+
+	D3D12_CPU_DESCRIPTOR_HANDLE GetCPUSlotHandle(FD3D12OnlineDescriptorBlock* InBlock) const { return Heap->GetCPUSlotHandle(InBlock->BaseSlot); }
+	D3D12_GPU_DESCRIPTOR_HANDLE GetGPUSlotHandle(FD3D12OnlineDescriptorBlock* InBlock) const { return Heap->GetGPUSlotHandle(InBlock->BaseSlot); }
+
+private:
+	// Check all released blocks and check which ones are not used by the GPU anymore
+	void UpdateFreeBlocks();
+
+	TRefCountPtr<FD3D12DescriptorHeap> Heap;
+
+	TQueue<FD3D12OnlineDescriptorBlock*> FreeBlocks;
+	TArray<FD3D12OnlineDescriptorBlock*> ReleasedBlocks;
+
+	FCriticalSection CriticalSection;
+};
+
+/** Manages and allows allocations of CPU descriptors only. Creates small heaps on demand to satisfy allocations. */
+class FD3D12OfflineDescriptorManager : public FD3D12DeviceChild
+{
+public:
+	FD3D12OfflineDescriptorManager() = delete;
+	FD3D12OfflineDescriptorManager(FD3D12Device* InDevice);
+	~FD3D12OfflineDescriptorManager();
+
+	inline ED3D12DescriptorHeapType GetHeapType() const { return HeapType; }
+
+	void Init(ED3D12DescriptorHeapType InHeapType);
+
+	D3D12_CPU_DESCRIPTOR_HANDLE AllocateHeapSlot(uint32& outIndex);
+	void FreeHeapSlot(D3D12_CPU_DESCRIPTOR_HANDLE Offset, uint32 index);
+
+private:
+	void AllocateHeap();
+
+	TArray<FD3D12OfflineHeapEntry> Heaps;
+	TDoubleLinkedList<uint32> FreeHeaps;
+
+	ED3D12DescriptorHeapType HeapType;
+	uint32 NumDescriptorsPerHeap{};
+	uint32 DescriptorSize{};
+
+	FCriticalSection CriticalSection;
+};
+
+/** Primary descriptor heap and descriptor manager. All heap allocations come from here.
+	All GPU visible resource heap allocations will be sub-allocated from a single heap in this manager. */
+class FD3D12DescriptorHeapManager : public FD3D12DeviceChild
+{
+public:
+	FD3D12DescriptorHeapManager(FD3D12Device* InDevice);
+	~FD3D12DescriptorHeapManager();
+
+	void Init(uint32 InNumGlobalDescriptors);
+	void Destroy();
+
+	FD3D12DescriptorHeap* AllocateHeap(const TCHAR* DebugName, ED3D12DescriptorHeapType InHeapType, uint32 InNumDescriptors, D3D12_DESCRIPTOR_HEAP_FLAGS InFlags);
+	void FreeHeap(FD3D12DescriptorHeap* InHeap);
+
+	inline FD3D12DescriptorHeap* GetGlobalHeap() const { return GlobalHeap; }
+
+private:
+	TRefCountPtr<FD3D12DescriptorHeap> GlobalHeap;
+	FDescriptorAllocator Allocator;
+};
+
diff --git a/Engine/Source/Runtime/D3D12RHI/Private/D3D12Device.cpp b/Engine/Source/Runtime/D3D12RHI/Private/D3D12Device.cpp
--- a/Engine/Source/Runtime/D3D12RHI/Private/D3D12Device.cpp
+++ b/Engine/Source/Runtime/D3D12RHI/Private/D3D12Device.cpp
@@ -1,21 +1,12 @@
 // Copyright Epic Games, Inc. All Rights Reserved.
 
-/*=============================================================================
-D3D12Device.cpp: D3D device RHI implementation.
-=============================================================================*/
 #include "D3D12RHIPrivate.h"
 #include "D3D12RayTracing.h"
 
-namespace D3D12RHI
-{
-	extern void EmptyD3DSamplerStateCache();
-}
-using namespace D3D12RHI;
-
 FD3D12Device::FD3D12Device() :
 	FD3D12Device(FRHIGPUMask::GPU0(), nullptr)
-	{
-	}
+{
+}
 
 FD3D12Device::FD3D12Device(FRHIGPUMask InGPUMask, FD3D12Adapter* InAdapter) :
 	FD3D12SingleNodeGPUObject(InGPUMask),
@@ -24,16 +15,16 @@ FD3D12Device::FD3D12Device(FRHIGPUMask InGPUMask, FD3D12Adapter* InAdapter) :
 	CopyCommandListManager(nullptr),
 	AsyncCommandListManager(nullptr),
 	TextureStreamingCommandAllocatorManager(this, D3D12_COMMAND_LIST_TYPE_COPY),
-	RTVAllocator(InGPUMask, D3D12_DESCRIPTOR_HEAP_TYPE_RTV, 256),
-	DSVAllocator(InGPUMask, D3D12_DESCRIPTOR_HEAP_TYPE_DSV, 256),
-	SRVAllocator(InGPUMask, D3D12_DESCRIPTOR_HEAP_TYPE_CBV_SRV_UAV, 1024),
-	UAVAllocator(InGPUMask, D3D12_DESCRIPTOR_HEAP_TYPE_CBV_SRV_UAV, 1024),
-#if USE_STATIC_ROOT_SIGNATURE
-	CBVAllocator(InGPUMask, D3D12_DESCRIPTOR_HEAP_TYPE_CBV_SRV_UAV, 2048),
-#endif
-	SamplerAllocator(InGPUMask, FD3D12_DESCRIPTOR_HEAP_TYPE_SAMPLER, 128),
-	GlobalSamplerHeap(this, InGPUMask),
-	GlobalViewHeap(this, InGPUMask),
+	DescriptorHeapManager(this),
+	ResourceDescriptorManager(this),
+	OfflineDescriptorManagers{
+		FD3D12OfflineDescriptorManager(this),
+		FD3D12OfflineDescriptorManager(this),
+		FD3D12OfflineDescriptorManager(this),
+		FD3D12OfflineDescriptorManager(this),
+	},
+	GlobalSamplerHeap(this),
+	OnlineDescriptorManager(this),
 	OcclusionQueryHeap(this, D3D12_QUERY_TYPE_OCCLUSION, 65536, 4 /*frames to keep results */ * 1 /*batches per frame*/),
 #if WITH_PROFILEGPU || D3D12_SUBMISSION_GAP_RECORDER
 	CmdListExecTimeQueryHeap(new FD3D12LinearQueryHeap(this, D3D12_QUERY_HEAP_TYPE_TIMESTAMP, 8192)),
@@ -322,21 +313,6 @@ void FD3D12Device::SetupAfterDeviceCreation()
 	}
 #endif // (PLATFORM_WINDOWS || PLATFORM_HOLOLENS)
 
-	// Init offline descriptor allocators
-	RTVAllocator.Init(Direct3DDevice);
-	DSVAllocator.Init(Direct3DDevice);
-	SRVAllocator.Init(Direct3DDevice);
-	UAVAllocator.Init(Direct3DDevice);
-#if USE_STATIC_ROOT_SIGNATURE
-	CBVAllocator.Init(Direct3DDevice);
-#endif
-	SamplerAllocator.Init(Direct3DDevice);
-
-	GlobalSamplerHeap.Init(NUM_SAMPLER_DESCRIPTORS);
-
-	// This value can be tuned on a per app basis. I.e. most apps will never run into descriptor heap pressure so
-	// can make this global heap smaller
-	uint32 NumGlobalViewDesc = GGlobalViewHeapSize;
 
 	const D3D12_RESOURCE_BINDING_TIER Tier = GetParentAdapter()->GetResourceBindingTier();
 	uint32 MaximumSupportedHeapSize = NUM_VIEW_DESCRIPTORS_TIER_1;
@@ -354,9 +330,25 @@ void FD3D12Device::SetupAfterDeviceCreation()
 		MaximumSupportedHeapSize = NUM_VIEW_DESCRIPTORS_TIER_3;
 		break;
 	}
-	check(NumGlobalViewDesc <= MaximumSupportedHeapSize);
-		
-	GlobalViewHeap.Init(D3D12_DESCRIPTOR_HEAP_TYPE_CBV_SRV_UAV, NumGlobalViewDesc);
+
+	// This value can be tuned on a per app basis. I.e. most apps will never run into descriptor heap pressure so
+	// can make this global heap smaller
+	check((uint32)GGlobalDescriptorHeapSize <= MaximumSupportedHeapSize);
+	check(GOnlineDescriptorHeapSize <= GGlobalDescriptorHeapSize);
+	check(GResourceDescriptorHeapSize <= GGlobalDescriptorHeapSize);
+
+	DescriptorHeapManager.Init(GGlobalDescriptorHeapSize);
+	ResourceDescriptorManager.Init(GResourceDescriptorHeapSize);
+
+	// Init offline descriptor managers
+	for (uint32 Index = 0; Index < static_cast<uint32>(ED3D12DescriptorHeapType::count); Index++)
+	{
+		OfflineDescriptorManagers[Index].Init(static_cast<ED3D12DescriptorHeapType>(Index));
+	}
+
+	GlobalSamplerHeap.Init(NUM_SAMPLER_DESCRIPTORS);
+
+	OnlineDescriptorManager.Init(GOnlineDescriptorHeapSize, GOnlineDescriptorHeapBlockSize);
 
 	// Init the occlusion and timestamp query heaps
 	OcclusionQueryHeap.Init();
diff --git a/Engine/Source/Runtime/D3D12RHI/Private/D3D12Device.h b/Engine/Source/Runtime/D3D12RHI/Private/D3D12Device.h
--- a/Engine/Source/Runtime/D3D12RHI/Private/D3D12Device.h
+++ b/Engine/Source/Runtime/D3D12RHI/Private/D3D12Device.h
@@ -7,6 +7,7 @@ D3D12Device.h: D3D12 Device Interfaces
 #pragma once
 
 #include "CoreMinimal.h"
+#include "D3D12Descriptors.h"
 
 class FD3D12DynamicRHI;
 class FD3D12BasicRayTracingPipeline;
@@ -68,25 +69,21 @@ public:
 	inline FD3D12QueryHeap* GetTimestampQueryHeap(ED3D12CommandQueueType inQueueType) { return TimestampQueryHeaps[(uint32)inQueueType]; }
 	FD3D12LinearQueryHeap* GetCmdListExecTimeQueryHeap();
 
-	template <typename TViewDesc> FD3D12OfflineDescriptorManager& GetViewDescriptorAllocator();
-	template<> FD3D12OfflineDescriptorManager& GetViewDescriptorAllocator<D3D12_SHADER_RESOURCE_VIEW_DESC>() { return SRVAllocator; }
-	template<> FD3D12OfflineDescriptorManager& GetViewDescriptorAllocator<D3D12_RENDER_TARGET_VIEW_DESC>() { return RTVAllocator; }
-	template<> FD3D12OfflineDescriptorManager& GetViewDescriptorAllocator<D3D12_DEPTH_STENCIL_VIEW_DESC>() { return DSVAllocator; }
-	template<> FD3D12OfflineDescriptorManager& GetViewDescriptorAllocator<D3D12_UNORDERED_ACCESS_VIEW_DESC>() { return UAVAllocator; }
-#if USE_STATIC_ROOT_SIGNATURE
-	template<> FD3D12OfflineDescriptorManager& GetViewDescriptorAllocator<D3D12_CONSTANT_BUFFER_VIEW_DESC>() { return CBVAllocator; }
-#else
-	template<> FD3D12OfflineDescriptorManager& GetViewDescriptorAllocator<D3D12_CONSTANT_BUFFER_VIEW_DESC>() { check(false); abort(); }
-#endif
+	inline FD3D12DescriptorHeapManager& GetDescriptorHeapManager() { return DescriptorHeapManager; }
+	inline FD3D12ResourceDescriptorManager& GetResourceDescriptorManager() { return ResourceDescriptorManager; }
+	inline FD3D12OfflineDescriptorManager& GetOfflineDescriptorManager(ED3D12DescriptorHeapType InType)
+	{
+		check(InType < ED3D12DescriptorHeapType::count);
+		return OfflineDescriptorManagers[static_cast<int>(InType)];
+	}
 
-	inline FD3D12OfflineDescriptorManager& GetSamplerDescriptorAllocator() { return SamplerAllocator; }
 	inline FD3D12CommandListManager& GetCommandListManager() { return *CommandListManager; }
 	inline FD3D12CommandListManager& GetCopyCommandListManager() { return *CopyCommandListManager; }
 	inline FD3D12CommandListManager& GetAsyncCommandListManager() { return *AsyncCommandListManager; }
 	inline FD3D12CommandAllocatorManager& GetTextureStreamingCommandAllocatorManager() { return TextureStreamingCommandAllocatorManager; }
 	inline FD3D12DefaultBufferAllocator& GetDefaultBufferAllocator() { return DefaultBufferAllocator; }
 	inline FD3D12GlobalOnlineSamplerHeap& GetGlobalSamplerHeap() { return GlobalSamplerHeap; }
-	inline FD3D12GlobalHeap& GetGlobalViewHeap() { return GlobalViewHeap; }
+	inline FD3D12OnlineDescriptorManager& GetOnlineDescriptorManager() { return OnlineDescriptorManager; }
 
 	bool IsGPUIdle();
 
@@ -147,17 +144,12 @@ protected:
 	FD3D12CommandAllocatorManager TextureStreamingCommandAllocatorManager;
 
 	// Must be before the StateCache so that destructor ordering is valid
-	FD3D12OfflineDescriptorManager RTVAllocator;
-	FD3D12OfflineDescriptorManager DSVAllocator;
-	FD3D12OfflineDescriptorManager SRVAllocator;
-	FD3D12OfflineDescriptorManager UAVAllocator;
-#if USE_STATIC_ROOT_SIGNATURE
-	FD3D12OfflineDescriptorManager CBVAllocator;
-#endif
-	FD3D12OfflineDescriptorManager SamplerAllocator;
+	FD3D12DescriptorHeapManager DescriptorHeapManager;
+	FD3D12ResourceDescriptorManager ResourceDescriptorManager;
+	FD3D12OfflineDescriptorManager OfflineDescriptorManagers[static_cast<int>(ED3D12DescriptorHeapType::count)];
 
 	FD3D12GlobalOnlineSamplerHeap GlobalSamplerHeap;
-	FD3D12GlobalHeap GlobalViewHeap;
+	FD3D12OnlineDescriptorManager OnlineDescriptorManager;
 
 	FD3D12QueryHeap OcclusionQueryHeap;
 	FD3D12QueryHeap* TimestampQueryHeaps[(uint32)ED3D12CommandQueueType::Count];
@@ -218,59 +210,3 @@ protected:
 
 	D3D12RHI::FD3DGPUProfiler GPUProfilingData;
 };
-template <typename TDesc> 
-void TD3D12ViewDescriptorHandle<TDesc>::AllocateDescriptorSlot()
-{
-	if (Parent)
-	{
-		FD3D12Device* Device = GetParentDevice();
-		FD3D12OfflineDescriptorManager& DescriptorAllocator = Device->template GetViewDescriptorAllocator<TDesc>();
-		Handle = DescriptorAllocator.AllocateHeapSlot(Index);
-		check(Handle.ptr != 0);
-	}
-}
-
-template <typename TDesc> 
-void TD3D12ViewDescriptorHandle<TDesc>::FreeDescriptorSlot()
-{
-	if (Parent)
-	{
-		FD3D12Device* Device = GetParentDevice();
-		FD3D12OfflineDescriptorManager& DescriptorAllocator = Device->template GetViewDescriptorAllocator<TDesc>();
-		DescriptorAllocator.FreeHeapSlot(Handle, Index);
-		Handle.ptr = 0;
-	}
-	check(!Handle.ptr);
-}
-
-template<>
-inline void FD3D12DescriptorHandleSRV::CreateView(const D3D12_SHADER_RESOURCE_VIEW_DESC& Desc, ID3D12Resource* Resource)
-{
-#if D3D12_RHI_RAYTRACING
-	// NOTE (from D3D Debug runtime): When ViewDimension is D3D12_SRV_DIMENSION_RAYTRACING_ACCELLERATION_STRUCTURE, pResource must be NULL, since the resource location comes from a GPUVA in pDesc
-	if (Desc.ViewDimension == D3D12_SRV_DIMENSION_RAYTRACING_ACCELERATION_STRUCTURE)
-	{
-		Resource = nullptr;
-	}
-#endif // D3D12_RHI_RAYTRACING
-
-	GetParentDevice()->GetDevice()->CreateShaderResourceView(Resource, &Desc, Handle);
-}
-
-template<>
-inline void FD3D12DescriptorHandleRTV::CreateView(const D3D12_RENDER_TARGET_VIEW_DESC& Desc, ID3D12Resource* Resource)
-{
-	GetParentDevice()->GetDevice()->CreateRenderTargetView(Resource, &Desc, Handle);
-}
-
-template<>
-inline void FD3D12DescriptorHandleDSV::CreateView(const D3D12_DEPTH_STENCIL_VIEW_DESC& Desc, ID3D12Resource* Resource)
-{
-	GetParentDevice()->GetDevice()->CreateDepthStencilView(Resource, &Desc, Handle);
-}
-
-template<>
-inline void FD3D12DescriptorHandleUAV::CreateViewWithCounter(const D3D12_UNORDERED_ACCESS_VIEW_DESC& Desc, ID3D12Resource* Resource, ID3D12Resource* CounterResource)
-{
-	GetParentDevice()->GetDevice()->CreateUnorderedAccessView(Resource, CounterResource, &Desc, Handle);
-}
diff --git a/Engine/Source/Runtime/D3D12RHI/Private/D3D12RayTracing.cpp b/Engine/Source/Runtime/D3D12RHI/Private/D3D12RayTracing.cpp
--- a/Engine/Source/Runtime/D3D12RHI/Private/D3D12RayTracing.cpp
+++ b/Engine/Source/Runtime/D3D12RHI/Private/D3D12RayTracing.cpp
@@ -4404,10 +4404,10 @@ struct FD3D12RayTracingLocalResourceBinder
 #if USE_STATIC_ROOT_SIGNATURE
 	D3D12_CPU_DESCRIPTOR_HANDLE CreateTransientConstantBufferView(D3D12_GPU_VIRTUAL_ADDRESS Address, uint32 DataSize)
 	{
-		FD3D12ConstantBufferView* BufferView = new FD3D12ConstantBufferView(GetDevice(), nullptr);
+		FD3D12ConstantBufferView* BufferView = new FD3D12ConstantBufferView(GetDevice());
 		BufferView->Create(Address, DataSize);
 		ShaderTable.TransientCBVs.Add(BufferView);
-		return BufferView->OfflineDescriptorHandle;
+		return BufferView->GetView();
 	}
 #endif // USE_STATIC_ROOT_SIGNATURE
 
@@ -4420,11 +4420,7 @@ struct FD3D12RayTracingLocalResourceBinder
 		FD3D12FastConstantAllocator& Allocator = Device.GetParentAdapter()->GetTransientUniformBufferAllocator();
 		FD3D12ResourceLocation ResourceLocation(&Device);
 
-	#if USE_STATIC_ROOT_SIGNATURE
 		void* MappedData = Allocator.Allocate(DataSize, ResourceLocation, nullptr);
-	#else // USE_STATIC_ROOT_SIGNATURE
-		void* MappedData = Allocator.Allocate(DataSize, ResourceLocation);
-	#endif // USE_STATIC_ROOT_SIGNATURE
 
 		FMemory::Memcpy(MappedData, Data, DataSize);
 
@@ -4539,7 +4535,7 @@ static bool SetRayTracingShaderResources(
 		{
 			FD3D12UniformBuffer* CBV = FD3D12CommandContext::RetrieveObject<FD3D12UniformBuffer>(Resource, GPUIndex);
 		#if USE_STATIC_ROOT_SIGNATURE
-			LocalCBVs[CBVIndex] = CBV->View->OfflineDescriptorHandle;
+			LocalCBVs[CBVIndex] = CBV->View->GetView();
 		#else // USE_STATIC_ROOT_SIGNATURE
 			LocalCBVs[CBVIndex] = CBV->ResourceLocation.GetGPUVirtualAddress();
 		#endif // USE_STATIC_ROOT_SIGNATURE
diff --git a/Engine/Source/Runtime/D3D12RHI/Private/D3D12State.cpp b/Engine/Source/Runtime/D3D12RHI/Private/D3D12State.cpp
--- a/Engine/Source/Runtime/D3D12RHI/Private/D3D12State.cpp
+++ b/Engine/Source/Runtime/D3D12RHI/Private/D3D12State.cpp
@@ -540,7 +540,7 @@ FD3D12SamplerState::FD3D12SamplerState(FD3D12Device* InParent, const D3D12_SAMPL
 	, ID(SamplerID)
 {
 	Descriptor.ptr = 0;
-	FD3D12OfflineDescriptorManager& DescriptorAllocator = GetParentDevice()->GetSamplerDescriptorAllocator();
+	FD3D12OfflineDescriptorManager& DescriptorAllocator = GetParentDevice()->GetOfflineDescriptorManager(ED3D12DescriptorHeapType::Sampler);
 	Descriptor = DescriptorAllocator.AllocateHeapSlot(DescriptorHeapIndex);
 
 	GetParentDevice()->CreateSamplerInternal(Desc, Descriptor);
@@ -550,7 +550,7 @@ FD3D12SamplerState::~FD3D12SamplerState()
 {
 	if (Descriptor.ptr)
 	{
-		FD3D12OfflineDescriptorManager& DescriptorAllocator = GetParentDevice()->GetSamplerDescriptorAllocator();
+		FD3D12OfflineDescriptorManager& DescriptorAllocator = GetParentDevice()->GetOfflineDescriptorManager(ED3D12DescriptorHeapType::Sampler);
 		DescriptorAllocator.FreeHeapSlot(Descriptor, DescriptorHeapIndex);
 		Descriptor.ptr = 0;
 	}
diff --git a/Engine/Source/Runtime/D3D12RHI/Private/D3D12StateCache.cpp b/Engine/Source/Runtime/D3D12RHI/Private/D3D12StateCache.cpp
--- a/Engine/Source/Runtime/D3D12RHI/Private/D3D12StateCache.cpp
+++ b/Engine/Source/Runtime/D3D12RHI/Private/D3D12StateCache.cpp
@@ -11,6 +11,24 @@
 #include <emmintrin.h>
 #endif
 
+// This value defines how many descriptors will be in the device global descriptor heap. This heap contains all shader visible view descriptors.
+// Other shader visible descriptor heaps (e.g. OnlineViewHeap) are allocated from this pool. Non-visible heaps (e.g. LocalViewHeap) are allocated as standalone.
+int32 GGlobalDescriptorHeapSize = 1000 * 1000;
+static FAutoConsoleVariableRef CVarGlobalDescriptorHeapSize(
+	TEXT("D3D12.GlobalDescriptorHeapSize"),
+	GGlobalDescriptorHeapSize,
+	TEXT("Global descriptor heap size"),
+	ECVF_ReadOnly
+);
+
+int32 GResourceDescriptorHeapSize = 0;// 500 * 1000;
+static FAutoConsoleVariableRef CVarResourceDescriptorHeapSize(
+	TEXT("D3D12.ResourceDescriptorHeapSize"),
+	GResourceDescriptorHeapSize,
+	TEXT("Resource descriptor heap size"),
+	ECVF_ReadOnly
+);
+
 // This value defines how many descriptors will be in the device local view heap which
 // This should be tweaked for each title as heaps require VRAM. The default value of 512k takes up ~16MB
 int32 GLocalViewHeapSize = 500 * 1000;
@@ -21,14 +39,22 @@ static FAutoConsoleVariableRef CVarLocalViewHeapSize(
 	ECVF_ReadOnly
 );
 
-// This value defines how many descriptors will be in the device global view heap which
+// This value defines how many descriptors will be in the device online view heap which
 // is shared across contexts to allow the driver to eliminate redundant descriptor heap sets.
 // This should be tweaked for each title as heaps require VRAM. The default value of 512k takes up ~16MB
-int32 GGlobalViewHeapSize = 500 * 1000;
-static FAutoConsoleVariableRef CVarGlobalViewHeapSize(
-	TEXT("D3D12.GlobalViewHeapSize"),
-	GGlobalViewHeapSize,
-	TEXT("Global view heap size"),
+int32 GOnlineDescriptorHeapSize = 500 * 1000;
+static FAutoConsoleVariableRef CVarOnlineDescriptorHeapSize(
+	TEXT("D3D12.OnlineDescriptorHeapSize"),
+	GOnlineDescriptorHeapSize,
+	TEXT("Online descriptor heap size"),
+	ECVF_ReadOnly
+);
+
+int32 GOnlineDescriptorHeapBlockSize = 2000;
+static FAutoConsoleVariableRef CVarOnlineDescriptorHeapBlockSize(
+	TEXT("D3D12.OnlineDescriptorHeapBlockSize"),
+	GOnlineDescriptorHeapBlockSize,
+	TEXT("Block size for sub allocations on the global view descriptor heap."),
 	ECVF_ReadOnly
 );
 
@@ -102,8 +128,10 @@ void FD3D12StateCacheBase::Init(FD3D12Device* InParent, FD3D12CommandContext* In
 		ResourceBindingTier == D3D12_RESOURCE_BINDING_TIER_2 ? NUM_VIEW_DESCRIPTORS_TIER_2 :
 		NUM_VIEW_DESCRIPTORS_TIER_1;
 
+	check(GGlobalDescriptorHeapSize <= (int32)MaxDescriptorsForTier);
+	check(GResourceDescriptorHeapSize <= (int32)MaxDescriptorsForTier);
 	check(GLocalViewHeapSize <= (int32)MaxDescriptorsForTier);
-	check(GGlobalViewHeapSize <= (int32)MaxDescriptorsForTier);
+	check(GOnlineDescriptorHeapSize <= (int32)MaxDescriptorsForTier);
 
 	const uint32 NumSamplerDescriptors = NUM_SAMPLER_DESCRIPTORS;
 	DescriptorCache.Init(InParent, InCmdContext, GLocalViewHeapSize, NumSamplerDescriptors);
@@ -1060,11 +1088,7 @@ void FD3D12StateCacheBase::SetUAVs(uint32 UAVStartSlot, uint32 NumSimultaneousUA
 				FD3D12Device* Device = CounterResource->GetParentDevice();
 				FD3D12ResourceLocation UploadBufferLocation(Device);
 
-#if USE_STATIC_ROOT_SIGNATURE
 				uint32* CounterUploadHeapData = static_cast<uint32*>(CmdContext->ConstantsAllocator.Allocate(sizeof(uint32), UploadBufferLocation, nullptr));
-#else
-				uint32* CounterUploadHeapData = static_cast<uint32*>(CmdContext->ConstantsAllocator.Allocate(sizeof(uint32), UploadBufferLocation));
-#endif
 
 				// Initialize the counter to 0 if it's not been previously initialized and the UAVInitialCount is -1, if not use the value that was passed.
 				*CounterUploadHeapData = (!UAV->IsCounterResourceInitialized() && UAVInitialCountArray[i] == -1) ? 0 : UAVInitialCountArray[i];
diff --git a/Engine/Source/Runtime/D3D12RHI/Private/D3D12StateCachePrivate.h b/Engine/Source/Runtime/D3D12RHI/Private/D3D12StateCachePrivate.h
--- a/Engine/Source/Runtime/D3D12RHI/Private/D3D12StateCachePrivate.h
+++ b/Engine/Source/Runtime/D3D12RHI/Private/D3D12StateCachePrivate.h
@@ -47,8 +47,10 @@ extern bool GD3D12SkipStateCaching;
 static const bool GD3D12SkipStateCaching = false;
 #endif
 
-extern int32 GGlobalViewHeapSize;
-
+extern int32 GGlobalDescriptorHeapSize;
+extern int32 GResourceDescriptorHeapSize;
+extern int32 GOnlineDescriptorHeapSize;
+extern int32 GOnlineDescriptorHeapBlockSize;
 
 enum ED3D12PipelineType
 {
@@ -670,7 +672,7 @@ public:
 			}
 
 #if USE_STATIC_ROOT_SIGNATURE
-			CBVCache.CBHandles[ShaderFrequency][SlotIndex] = UniformBuffer->View->OfflineDescriptorHandle;
+			CBVCache.CBHandles[ShaderFrequency][SlotIndex] = UniformBuffer->View->GetView();
 #endif
 		}
 		else if (CurrentGPUVirtualAddress != 0)
@@ -708,7 +710,7 @@ public:
 			FD3D12ConstantBufferCache::DirtySlot(CBVCache.DirtySlotMask[ShaderFrequency], SlotIndex);
 
 #if USE_STATIC_ROOT_SIGNATURE
-			CBVCache.CBHandles[ShaderFrequency][SlotIndex] = Buffer.View->OfflineDescriptorHandle;
+			CBVCache.CBHandles[ShaderFrequency][SlotIndex] = Buffer.GetView();
 #endif
 		}
 	}
diff --git a/Engine/Source/Runtime/D3D12RHI/Private/D3D12UAV.cpp b/Engine/Source/Runtime/D3D12RHI/Private/D3D12UAV.cpp
--- a/Engine/Source/Runtime/D3D12RHI/Private/D3D12UAV.cpp
+++ b/Engine/Source/Runtime/D3D12RHI/Private/D3D12UAV.cpp
@@ -230,8 +230,8 @@ void FD3D12CommandContext::ClearUAV(TRHICommandList_RecursiveHazardous<FD3D12Com
 				}
 
 				// Scoped descriptor handle will free the offline CPU handle once we return
-				FD3D12DescriptorHandleUAV UAVHandle(ParentDevice);
-				UAVHandle.CreateViewWithCounter(R32UAVDesc, Resource, nullptr);
+				FD3D12ViewDescriptorHandle UAVHandle(ParentDevice, ED3D12DescriptorHeapType::Standard);
+				UAVHandle.CreateView(R32UAVDesc, Resource, nullptr);
 
 				// Check if the view heap is full and needs to rollover.
 				if (!Context.StateCache.GetDescriptorCache()->GetCurrentViewHeap()->CanReserveSlots(1))
@@ -240,7 +240,7 @@ void FD3D12CommandContext::ClearUAV(TRHICommandList_RecursiveHazardous<FD3D12Com
 				}
 
 				uint32 ReservedSlot = Context.StateCache.GetDescriptorCache()->GetCurrentViewHeap()->ReserveSlots(1);
-				D3D12_CPU_DESCRIPTOR_HANDLE CPUHandle = UAVHandle.GetHandle();
+				D3D12_CPU_DESCRIPTOR_HANDLE CPUHandle = UAVHandle.GetOfflineHandle();
 				D3D12_CPU_DESCRIPTOR_HANDLE DestSlot = Context.StateCache.GetDescriptorCache()->GetCurrentViewHeap()->GetCPUSlotHandle(ReservedSlot);
 				D3D12_GPU_DESCRIPTOR_HANDLE GPUHandle = Context.StateCache.GetDescriptorCache()->GetCurrentViewHeap()->GetGPUSlotHandle(ReservedSlot);
 
diff --git a/Engine/Source/Runtime/D3D12RHI/Private/D3D12UniformBuffer.cpp b/Engine/Source/Runtime/D3D12RHI/Private/D3D12UniformBuffer.cpp
--- a/Engine/Source/Runtime/D3D12RHI/Private/D3D12UniformBuffer.cpp
+++ b/Engine/Source/Runtime/D3D12RHI/Private/D3D12UniformBuffer.cpp
@@ -36,7 +36,7 @@ FUniformBufferRHIRef FD3D12DynamicRHI::RHICreateUniformBuffer(const void* Conten
 
 #if USE_STATIC_ROOT_SIGNATURE
 			// Create an offline CBV descriptor
-			NewUniformBuffer->View = new FD3D12ConstantBufferView(Device, nullptr);
+			NewUniformBuffer->View = new FD3D12ConstantBufferView(Device);
 #endif
 			void* MappedData = nullptr;
 			if (Usage == EUniformBufferUsage::UniformBuffer_MultiFrame)
@@ -49,11 +49,7 @@ FUniformBufferRHIRef FD3D12DynamicRHI::RHICreateUniformBuffer(const void* Conten
 			{
 				// Uniform buffers which will live for 1 frame at the max can be allocated very efficiently from a ring buffer
 				FD3D12FastConstantAllocator& Allocator = GetAdapter().GetTransientUniformBufferAllocator();
-#if USE_STATIC_ROOT_SIGNATURE
 				MappedData = Allocator.Allocate(NumBytesActualData, NewUniformBuffer->ResourceLocation, nullptr);
-#else
-				MappedData = Allocator.Allocate(NumBytesActualData, NewUniformBuffer->ResourceLocation);
-#endif
 			}
 			check(NewUniformBuffer->ResourceLocation.GetOffsetFromBaseOfResource() % 16 == 0);
 
@@ -181,11 +177,7 @@ void FD3D12DynamicRHI::RHIUpdateUniformBuffer(FRHIUniformBuffer* UniformBufferRH
 			{
 				FD3D12FastConstantAllocator& Allocator = GetAdapter().GetTransientUniformBufferAllocator();
 
-#if USE_STATIC_ROOT_SIGNATURE
 				MappedData = Allocator.Allocate(NumBytes, UpdatedResourceLocation, nullptr);
-#else
-				MappedData = Allocator.Allocate(NumBytes, UpdatedResourceLocation);
-#endif
 			}
 
 			check(MappedData != nullptr);
diff --git a/Engine/Source/Runtime/D3D12RHI/Private/D3D12View.cpp b/Engine/Source/Runtime/D3D12RHI/Private/D3D12View.cpp
--- a/Engine/Source/Runtime/D3D12RHI/Private/D3D12View.cpp
+++ b/Engine/Source/Runtime/D3D12RHI/Private/D3D12View.cpp
@@ -1,11 +1,116 @@
 // Copyright Epic Games, Inc. All Rights Reserved.
 
-/*=============================================================================
-	D3D12View.cpp: 
-=============================================================================*/
-
 #include "D3D12RHIPrivate.h"
 
+////////////////////////////////////////////////////////////////////////////////////////////////////////////////////////
+// FD3D12ViewDescriptorHandle
+
+FD3D12ViewDescriptorHandle::FD3D12ViewDescriptorHandle(FD3D12Device* InParentDevice, ED3D12DescriptorHeapType InHeapType)
+	: FD3D12DeviceChild(InParentDevice)
+	, HeapType(InHeapType)
+{
+	AllocateDescriptorSlot();
+}
+
+FD3D12ViewDescriptorHandle::~FD3D12ViewDescriptorHandle()
+{
+	FreeDescriptorSlot();
+	check(OfflineHandle.ptr == 0);
+	check(OfflineIndex == UINT_MAX);
+	check(ResourceIndex == UINT_MAX);
+}
+
+void FD3D12ViewDescriptorHandle::CreateView(const D3D12_RENDER_TARGET_VIEW_DESC& Desc, ID3D12Resource* Resource)
+{
+	check(HeapType == ED3D12DescriptorHeapType::RenderTarget);
+
+	GetParentDevice()->GetDevice()->CreateRenderTargetView(Resource, &Desc, OfflineHandle);
+}
+
+void FD3D12ViewDescriptorHandle::CreateView(const D3D12_DEPTH_STENCIL_VIEW_DESC& Desc, ID3D12Resource* Resource)
+{
+	check(HeapType == ED3D12DescriptorHeapType::DepthStencil);
+
+	GetParentDevice()->GetDevice()->CreateDepthStencilView(Resource, &Desc, OfflineHandle);
+}
+
+void FD3D12ViewDescriptorHandle::CreateView(const D3D12_CONSTANT_BUFFER_VIEW_DESC& Desc)
+{
+	check(HeapType == ED3D12DescriptorHeapType::Standard);
+
+	GetParentDevice()->GetDevice()->CreateConstantBufferView(&Desc, OfflineHandle);
+}
+
+void FD3D12ViewDescriptorHandle::CreateView(const D3D12_SHADER_RESOURCE_VIEW_DESC& Desc, ID3D12Resource* Resource)
+{
+	check(HeapType == ED3D12DescriptorHeapType::Standard);
+
+#if D3D12_RHI_RAYTRACING
+	// NOTE (from D3D Debug runtime): When ViewDimension is D3D12_SRV_DIMENSION_RAYTRACING_ACCELLERATION_STRUCTURE, pResource must be NULL, since the resource location comes from a GPUVA in pDesc
+	if (Desc.ViewDimension == D3D12_SRV_DIMENSION_RAYTRACING_ACCELERATION_STRUCTURE)
+	{
+		Resource = nullptr;
+	}
+#endif // D3D12_RHI_RAYTRACING
+
+	GetParentDevice()->GetDevice()->CreateShaderResourceView(Resource, &Desc, OfflineHandle);
+
+	if (IsResourceHandle())
+	{
+		GetParentDevice()->GetDevice()->CopyDescriptorsSimple(1, GetResourceHandle(), GetOfflineHandle(), D3D12_DESCRIPTOR_HEAP_TYPE_CBV_SRV_UAV);
+	}
+}
+
+void FD3D12ViewDescriptorHandle::CreateView(const D3D12_UNORDERED_ACCESS_VIEW_DESC& Desc, ID3D12Resource* Resource, ID3D12Resource* CounterResource)
+{
+	check(HeapType == ED3D12DescriptorHeapType::Standard);
+
+	GetParentDevice()->GetDevice()->CreateUnorderedAccessView(Resource, CounterResource, &Desc, OfflineHandle);
+
+	if (IsResourceHandle())
+	{
+		GetParentDevice()->GetDevice()->CopyDescriptorsSimple(1, GetResourceHandle(), GetOfflineHandle(), D3D12_DESCRIPTOR_HEAP_TYPE_CBV_SRV_UAV);
+	}
+}
+
+D3D12_CPU_DESCRIPTOR_HANDLE FD3D12ViewDescriptorHandle::GetResourceHandle() const
+{
+	return GetParentDevice()->GetResourceDescriptorManager().GetResourceHandle(ResourceIndex);
+}
+
+void FD3D12ViewDescriptorHandle::AllocateDescriptorSlot()
+{
+	if (FD3D12Device* Device = GetParentDevice())
+	{
+		OfflineHandle = Device->GetOfflineDescriptorManager(HeapType).AllocateHeapSlot(OfflineIndex);
+		check(OfflineHandle.ptr != 0);
+
+		if (HeapType == ED3D12DescriptorHeapType::Standard)
+		{
+			Device->GetResourceDescriptorManager().AllocateDescriptor(ResourceIndex);
+		}
+	}
+}
+
+void FD3D12ViewDescriptorHandle::FreeDescriptorSlot()
+{
+	if (FD3D12Device* Device = GetParentDevice())
+	{
+		Device->GetOfflineDescriptorManager(HeapType).FreeHeapSlot(OfflineHandle, OfflineIndex);
+		OfflineIndex = UINT_MAX;
+		OfflineHandle.ptr = 0;
+
+		if (IsResourceHandle())
+		{
+			Device->GetResourceDescriptorManager().FreeDescriptor(ResourceIndex);
+			ResourceIndex = UINT_MAX;
+		}
+	}
+	check(OfflineHandle.ptr == 0);
+}
+
+////////////////////////////////////////////////////////////////////////////////////////////////////////////////////////
+
 static FORCEINLINE D3D12_SHADER_RESOURCE_VIEW_DESC GetVertexBufferSRVDesc(FD3D12Buffer* VertexBuffer, uint32& CreationStride, uint8 Format, uint32 StartOffsetBytes, uint32 NumElements)
 {
 	const uint32 BufferSize = VertexBuffer->GetSize();
@@ -460,31 +565,15 @@ FShaderResourceViewRHIRef FD3D12DynamicRHI::RHICreateShaderResourceViewWriteMask
 	return RHICreateShaderResourceViewWriteMask(Texture2D);
 }
 
-#if USE_STATIC_ROOT_SIGNATURE
-void FD3D12ConstantBufferView::AllocateHeapSlot()
-{
-	if (!OfflineDescriptorHandle.ptr)
-	{
-	    FD3D12OfflineDescriptorManager& DescriptorAllocator = GetParentDevice()->GetViewDescriptorAllocator<D3D12_CONSTANT_BUFFER_VIEW_DESC>();
-	    OfflineDescriptorHandle = DescriptorAllocator.AllocateHeapSlot(OfflineHeapIndex);
-	    check(OfflineDescriptorHandle.ptr != 0);
-	}
-}
-
-void FD3D12ConstantBufferView::FreeHeapSlot()
+FD3D12ConstantBufferView::FD3D12ConstantBufferView(FD3D12Device* InParent)
+	: FD3D12DeviceChild(InParent)
+	, Descriptor(InParent, ED3D12DescriptorHeapType::Standard)
 {
-	if (OfflineDescriptorHandle.ptr)
-	{
-		FD3D12OfflineDescriptorManager& DescriptorAllocator = GetParentDevice()->GetViewDescriptorAllocator<D3D12_CONSTANT_BUFFER_VIEW_DESC>();
-		DescriptorAllocator.FreeHeapSlot(OfflineDescriptorHandle, OfflineHeapIndex);
-		OfflineDescriptorHandle.ptr = 0;
-	}
 }
 
 void FD3D12ConstantBufferView::Create(D3D12_GPU_VIRTUAL_ADDRESS GPUAddress, const uint32 AlignedSize)
 {
 	Desc.BufferLocation = GPUAddress;
 	Desc.SizeInBytes = AlignedSize;
-	GetParentDevice()->GetDevice()->CreateConstantBufferView(&Desc, OfflineDescriptorHandle);
+	Descriptor.CreateView(Desc);
 }
-#endif
diff --git a/Engine/Source/Runtime/D3D12RHI/Private/D3D12View.h b/Engine/Source/Runtime/D3D12RHI/Private/D3D12View.h
--- a/Engine/Source/Runtime/D3D12RHI/Private/D3D12View.h
+++ b/Engine/Source/Runtime/D3D12RHI/Private/D3D12View.h
@@ -1,10 +1,9 @@
 // Copyright Epic Games, Inc. All Rights Reserved.
 
-/*=============================================================================
-D3D12View.h: D3D12 Resource Views
-=============================================================================*/
 #pragma once
 
+#include "D3D12Descriptors.h"
+
 enum ViewSubresourceSubsetFlags
 {
 	ViewSubresourceSubsetFlags_None = 0x0,
@@ -619,57 +618,44 @@ inline uint32 CViewSubresourceSubset::MaxSubresource() const
 	return (*(--end())).Value;
 }
 
-class FD3D12OfflineDescriptorManager;
-
-template <typename TDesc>
-class TD3D12ViewDescriptorHandle : public FD3D12DeviceChild
+/** Manages descriptor allocations and view creation */
+class FD3D12ViewDescriptorHandle : public FD3D12DeviceChild
 {
-	D3D12_CPU_DESCRIPTOR_HANDLE Handle;
-	uint32 Index;
-
 public:
-	TD3D12ViewDescriptorHandle(FD3D12Device* InParentDevice)
-		: FD3D12DeviceChild(InParentDevice)
-	{
-		Handle.ptr = 0;
-		AllocateDescriptorSlot();
-	}
+	FD3D12ViewDescriptorHandle() = delete;
+	FD3D12ViewDescriptorHandle(FD3D12Device* InParentDevice, ED3D12DescriptorHeapType InHeapType);
+	~FD3D12ViewDescriptorHandle();
 
-	~TD3D12ViewDescriptorHandle()
-	{
-		FreeDescriptorSlot();
-	}
+	void CreateView(const D3D12_RENDER_TARGET_VIEW_DESC& Desc, ID3D12Resource* Resource);
+	void CreateView(const D3D12_DEPTH_STENCIL_VIEW_DESC& Desc, ID3D12Resource* Resource);
+	void CreateView(const D3D12_CONSTANT_BUFFER_VIEW_DESC& Desc);
+	void CreateView(const D3D12_SHADER_RESOURCE_VIEW_DESC& Desc, ID3D12Resource* Resource);
+	void CreateView(const D3D12_UNORDERED_ACCESS_VIEW_DESC& Desc, ID3D12Resource* Resource, ID3D12Resource* CounterResource);
 
-	void SetParentDevice(FD3D12Device* InParent)
-	{
-		check(!Parent && !Handle.ptr);
-		FD3D12DeviceChild::SetParentDevice(InParent);
-		AllocateDescriptorSlot();
-	}
+	D3D12_CPU_DESCRIPTOR_HANDLE GetResourceHandle() const;
 
-	// Implemented in D3D12Device.h due to dependencies on FD3D12Device
-	inline void CreateView(const TDesc& Desc, ID3D12Resource* Resource);
-	inline void CreateViewWithCounter(const TDesc& Desc, ID3D12Resource* Resource, ID3D12Resource* CounterResource);
+	inline D3D12_CPU_DESCRIPTOR_HANDLE GetOfflineHandle()  const { return OfflineHandle; }
+	inline ED3D12DescriptorHeapType    GetHeapType()       const { return HeapType; }
+	inline uint32                      GetOfflineIndex()   const { return OfflineIndex; }
+	inline uint32                      GetResourceIndex()  const { return ResourceIndex; }
 
-	inline const D3D12_CPU_DESCRIPTOR_HANDLE& GetHandle() const { return Handle; }
-	inline uint32 GetIndex() const { return Index; }
+	inline bool IsResourceHandle() const { return GetResourceIndex() != UINT_MAX; }
 
 private:
-	// Implemented in D3D12Device.h due to dependencies on FD3D12Device
-	inline void AllocateDescriptorSlot();
-	inline void FreeDescriptorSlot();
-};
+	void AllocateDescriptorSlot();
+	void FreeDescriptorSlot();
 
-typedef TD3D12ViewDescriptorHandle<D3D12_SHADER_RESOURCE_VIEW_DESC>		FD3D12DescriptorHandleSRV;
-typedef TD3D12ViewDescriptorHandle<D3D12_RENDER_TARGET_VIEW_DESC>		FD3D12DescriptorHandleRTV;
-typedef TD3D12ViewDescriptorHandle<D3D12_DEPTH_STENCIL_VIEW_DESC>		FD3D12DescriptorHandleDSV;
-typedef TD3D12ViewDescriptorHandle<D3D12_UNORDERED_ACCESS_VIEW_DESC>	FD3D12DescriptorHandleUAV;
+	D3D12_CPU_DESCRIPTOR_HANDLE OfflineHandle{};
+	const ED3D12DescriptorHeapType HeapType{};
+	uint32 OfflineIndex{ UINT_MAX };
+	uint32 ResourceIndex{ UINT_MAX };
+};
 
 template <typename TDesc>
 class FD3D12View : public FD3D12ShaderResourceRenameListener
 {
 protected:
-	TD3D12ViewDescriptorHandle<TDesc> Descriptor;
+	FD3D12ViewDescriptorHandle Descriptor;
 
 	ViewSubresourceSubsetFlags Flags;
 	FD3D12BaseShaderResource* BaseShaderResource;
@@ -683,8 +669,8 @@ protected:
 	bool bInitialized;
 #endif
 
-	explicit FD3D12View(FD3D12Device* InParent, ViewSubresourceSubsetFlags InFlags)
-		: Descriptor(InParent)
+	explicit FD3D12View(FD3D12Device* InParent, ED3D12DescriptorHeapType InHeapType, ViewSubresourceSubsetFlags InFlags)
+		: Descriptor(InParent, InHeapType)
 		, Flags(InFlags)
 		, BaseShaderResource(nullptr)
 #if DO_CHECK || USING_CODE_ANALYSIS
@@ -768,7 +754,7 @@ public:
 	inline FD3D12Device*					GetParentDevice_Unsafe()	const { return Descriptor.GetParentDevice_Unsafe(); }
 	inline FD3D12ResourceLocation*			GetResourceLocation()		const { return ResourceLocation; }
 	inline const TDesc&						GetDesc()					const { checkf(bInitialized, TEXT("Uninitialized D3D12View size %d"), (uint32)sizeof(TDesc)); return Desc; }
-	inline D3D12_CPU_DESCRIPTOR_HANDLE		GetView()					const { checkf(bInitialized, TEXT("Uninitialized D3D12View size %d"), (uint32)sizeof(TDesc)); return Descriptor.GetHandle(); }
+	inline D3D12_CPU_DESCRIPTOR_HANDLE		GetView()					const { checkf(bInitialized, TEXT("Uninitialized D3D12View size %d"), (uint32)sizeof(TDesc)); return Descriptor.GetOfflineHandle(); }
 	inline uint32							GetDescriptorHeapIndex()	const { checkf(bInitialized, TEXT("Uninitialized D3D12View size %d"), (uint32)sizeof(TDesc)); return Descriptor.GetIndex(); }
 	inline FD3D12Resource*					GetResource()				const { checkf(bInitialized, TEXT("Uninitialized D3D12View size %d"), (uint32)sizeof(TDesc)); return Resource; }
 	inline FD3D12ResidencyHandle*			GetResidencyHandle()		const { checkf(bInitialized, TEXT("Uninitialized D3D12View size %d"), (uint32)sizeof(TDesc)); return ResidencyHandle; }
@@ -789,18 +775,16 @@ public:
 /** Shader resource view class. */
 class FD3D12ShaderResourceView : public FRHIShaderResourceView, public FD3D12View<D3D12_SHADER_RESOURCE_VIEW_DESC>, public FD3D12LinkedAdapterObject<FD3D12ShaderResourceView>
 {
-	bool bContainsDepthPlane;
-	bool bContainsStencilPlane;
-	bool bSkipFastClearFinalize;
-	bool bRequiresResourceStateTracking;
-	uint32 Stride;
-	uint32 StartOffsetBytes;
-
 public:
 	// Used for dynamic buffer SRVs, which can be renamed. Must be explicitly Initialize()'d before it can be used.
 	FD3D12ShaderResourceView(FD3D12Device* InParent)
-		: FD3D12View(InParent, ViewSubresourceSubsetFlags_None)
-	{}
+		: FD3D12View(InParent, ED3D12DescriptorHeapType::Standard, ViewSubresourceSubsetFlags_None)
+		, bContainsDepthPlane(false)
+		, bContainsStencilPlane(false)
+		, bSkipFastClearFinalize(false)
+		, bRequiresResourceStateTracking(false)
+	{
+	}
 
 	// Used for all other SRV resource types. Initialization is immediate on the calling thread.
 	// Should not be used for dynamic resources which can be renamed.
@@ -820,51 +804,12 @@ public:
 	{
 	}
 
-	void PreCreateView(FD3D12ResourceLocation& InResourceLocation, uint32 InStride, uint32 InStartOffsetBytes, bool InSkipFastClearFinalize)
-	{
-		Stride = InStride;
-		StartOffsetBytes = InStartOffsetBytes;
-		bSkipFastClearFinalize = InSkipFastClearFinalize;
-
-		if (InResourceLocation.GetResource())
-		{
-			bContainsDepthPlane = InResourceLocation.GetResource()->IsDepthStencilResource() && GetPlaneSliceFromViewFormat(InResourceLocation.GetResource()->GetDesc().Format, Desc.Format) == 0;
-			bContainsStencilPlane = InResourceLocation.GetResource()->IsDepthStencilResource() && GetPlaneSliceFromViewFormat(InResourceLocation.GetResource()->GetDesc().Format, Desc.Format) == 1;
-			bRequiresResourceStateTracking = InResourceLocation.GetResource()->RequiresResourceStateTracking();
-
-#if DO_CHECK
-			// Check the plane slice of the SRV matches the texture format
-			// Texture2DMS does not have explicit plane index (it's implied by the format)
-			if (Desc.ViewDimension == D3D12_SRV_DIMENSION_TEXTURE2D)
-			{
-				check(GetPlaneSliceFromViewFormat(InResourceLocation.GetResource()->GetDesc().Format, Desc.Format) == Desc.Texture2D.PlaneSlice);
-			}
-
-			if (Desc.ViewDimension == D3D12_SRV_DIMENSION_BUFFER)
-			{
-				//check(InResourceLocation.GetOffsetFromBaseOfResource() == Desc.Buffer.FirstElement * Stride);
-			}
-#endif
-		}
-	}
-
-	void CreateView(const D3D12_SHADER_RESOURCE_VIEW_DESC& InDesc, FD3D12BaseShaderResource* InBaseShaderResource, FD3D12ResourceLocation& InResourceLocation)
-	{
-		InitializeInternal(InBaseShaderResource, InResourceLocation);
-
-		if (ResourceLocation->GetResource())
-		{
-			ID3D12Resource* D3DResource = ResourceLocation->GetResource()->GetResource();
-			Descriptor.CreateView(Desc, D3DResource);
-		}
-	}
-
 	void Initialize(D3D12_SHADER_RESOURCE_VIEW_DESC& InDesc, FD3D12BaseShaderResource* InBaseShaderResource, FD3D12ResourceLocation& InResourceLocation, uint32 InStride, uint32 InStartOffsetBytes, bool InSkipFastClearFinalize)
 	{
 		SetDesc(InDesc);
 
 		PreCreateView(InResourceLocation, InStride, InStartOffsetBytes, InSkipFastClearFinalize);
-		CreateView(InDesc, InBaseShaderResource, InResourceLocation);
+		CreateView(InBaseShaderResource, InResourceLocation);
 	}
 
 	void Initialize(D3D12_SHADER_RESOURCE_VIEW_DESC& InDesc, FD3D12BaseShaderResource* InBaseShaderResource, uint32 InStride, uint32 InStartOffsetBytes, bool InSkipFastClearFinalize = false)
@@ -907,7 +852,7 @@ public:
 
 		// Update the LODClamp, the reinitialize the SRV
 		Desc.Texture2D.ResourceMinLODClamp = ResourceMinLODClamp;
-		CreateView(Desc, BaseShaderResource, *ResourceLocation);
+		CreateView(BaseShaderResource, *ResourceLocation);
 	}
 
 	virtual void RecreateView() override
@@ -920,6 +865,53 @@ public:
 	FORCEINLINE bool IsStencilPlaneResource()	const { return bContainsStencilPlane; }
 	FORCEINLINE bool GetSkipFastClearFinalize()	const { return bSkipFastClearFinalize; }
 	FORCEINLINE bool RequiresResourceStateTracking() const { return bRequiresResourceStateTracking; }
+
+protected:
+	void PreCreateView(FD3D12ResourceLocation& InResourceLocation, uint32 InStride, uint32 InStartOffsetBytes, bool InSkipFastClearFinalize)
+	{
+		Stride = InStride;
+		StartOffsetBytes = InStartOffsetBytes;
+		bSkipFastClearFinalize = InSkipFastClearFinalize;
+
+		if (InResourceLocation.GetResource())
+		{
+			bContainsDepthPlane = InResourceLocation.GetResource()->IsDepthStencilResource() && GetPlaneSliceFromViewFormat(InResourceLocation.GetResource()->GetDesc().Format, Desc.Format) == 0;
+			bContainsStencilPlane = InResourceLocation.GetResource()->IsDepthStencilResource() && GetPlaneSliceFromViewFormat(InResourceLocation.GetResource()->GetDesc().Format, Desc.Format) == 1;
+			bRequiresResourceStateTracking = InResourceLocation.GetResource()->RequiresResourceStateTracking();
+
+#if DO_CHECK
+			// Check the plane slice of the SRV matches the texture format
+			// Texture2DMS does not have explicit plane index (it's implied by the format)
+			if (Desc.ViewDimension == D3D12_SRV_DIMENSION_TEXTURE2D)
+			{
+				check(GetPlaneSliceFromViewFormat(InResourceLocation.GetResource()->GetDesc().Format, Desc.Format) == Desc.Texture2D.PlaneSlice);
+			}
+
+			if (Desc.ViewDimension == D3D12_SRV_DIMENSION_BUFFER)
+			{
+				//check(InResourceLocation.GetOffsetFromBaseOfResource() == Desc.Buffer.FirstElement * Stride);
+			}
+#endif
+		}
+	}
+
+	void CreateView(FD3D12BaseShaderResource* InBaseShaderResource, FD3D12ResourceLocation& InResourceLocation)
+	{
+		InitializeInternal(InBaseShaderResource, InResourceLocation);
+
+		if (ResourceLocation->GetResource())
+		{
+			ID3D12Resource* D3DResource = ResourceLocation->GetResource()->GetResource();
+			Descriptor.CreateView(Desc, D3DResource);
+		}
+	}
+
+	uint32 Stride{};
+	uint32 StartOffsetBytes{};
+	bool bContainsDepthPlane : 1;
+	bool bContainsStencilPlane : 1;
+	bool bSkipFastClearFinalize : 1;
+	bool bRequiresResourceStateTracking : 1;
 };
 
 class FD3D12ShaderResourceViewWithLocation : public FD3D12ShaderResourceView
@@ -951,42 +943,29 @@ public:
 class FD3D12UnorderedAccessView : public FRHIUnorderedAccessView, public FD3D12View < D3D12_UNORDERED_ACCESS_VIEW_DESC >, public FD3D12LinkedAdapterObject<FD3D12UnorderedAccessView>
 {
 public:
-
 	FD3D12UnorderedAccessView(FD3D12Device* InParent)
-		: FD3D12View(InParent, ViewSubresourceSubsetFlags_None)
+	: FD3D12View(InParent, ED3D12DescriptorHeapType::Standard, ViewSubresourceSubsetFlags_None)
 	{
 	}
 
 	FD3D12UnorderedAccessView(FD3D12Device* InParent, D3D12_UNORDERED_ACCESS_VIEW_DESC& InDesc, FD3D12BaseShaderResource* InBaseShaderResource, FD3D12Resource* InCounterResource = nullptr)
-		: FD3D12View(InParent, ViewSubresourceSubsetFlags_None)
+		: FD3D12View(InParent, ED3D12DescriptorHeapType::Standard, ViewSubresourceSubsetFlags_None)
 		, CounterResource(InCounterResource)
 	{
 		Initialize(InDesc, InBaseShaderResource, InBaseShaderResource->ResourceLocation, InCounterResource);
 	}
 
-	void CreateView(const D3D12_UNORDERED_ACCESS_VIEW_DESC& InDesc, FD3D12BaseShaderResource* InBaseShaderResource, FD3D12ResourceLocation& InResourceLocation, FD3D12Resource* InCounterResource)
-	{
-		InitializeInternal(InBaseShaderResource, InResourceLocation);
-
-		if (Resource)
-		{
-			ID3D12Resource* D3DResource = Resource->GetUAVAccessResource() ? Resource->GetUAVAccessResource() : Resource->GetResource();
-			ID3D12Resource* D3DCounterResource = InCounterResource ? InCounterResource->GetResource() : nullptr;
-			Descriptor.CreateViewWithCounter(Desc, D3DResource, D3DCounterResource);
-		}
-	}
-
 	void Initialize(D3D12_UNORDERED_ACCESS_VIEW_DESC& InDesc, FD3D12BaseShaderResource* InBaseShaderResource, FD3D12ResourceLocation& InResourceLocation, FD3D12Resource* InCounterResource = nullptr)
 	{
 		SetDesc(InDesc);
-		CreateView(InDesc, InBaseShaderResource, InResourceLocation, InCounterResource);
+		CreateView(InBaseShaderResource, InResourceLocation, InCounterResource);
 	}
 
 	virtual void RecreateView() override
 	{
 		check(CounterResource == nullptr);
 		check(ResourceLocation->GetOffsetFromBaseOfResource() == 0);
-		CreateView(Desc, BaseShaderResource, *ResourceLocation, nullptr);
+		CreateView(BaseShaderResource, *ResourceLocation, nullptr);
 	}
 
 	bool IsCounterResourceInitialized() const { return CounterResourceInitialized; }
@@ -995,6 +974,18 @@ public:
 	FD3D12Resource* GetCounterResource() { return CounterResource; }
 
 protected:
+	void CreateView(FD3D12BaseShaderResource* InBaseShaderResource, FD3D12ResourceLocation& InResourceLocation, FD3D12Resource* InCounterResource)
+	{
+		InitializeInternal(InBaseShaderResource, InResourceLocation);
+
+		if (Resource)
+		{
+			ID3D12Resource* D3DResource = Resource->GetUAVAccessResource() ? Resource->GetUAVAccessResource() : Resource->GetResource();
+			ID3D12Resource* D3DCounterResource = InCounterResource ? InCounterResource->GetResource() : nullptr;
+			Descriptor.CreateView(Desc, D3DResource, D3DCounterResource);
+		}
+	}
+
 	TRefCountPtr<FD3D12Resource> CounterResource;
 	bool CounterResourceInitialized = false;
 };
@@ -1017,78 +1008,54 @@ public:
 	FD3D12ResourceLocation ViewLocation;
 };
 
-#if USE_STATIC_ROOT_SIGNATURE
 class FD3D12ConstantBufferView : public FD3D12DeviceChild
 {
-			public:
-//protected:
-	/** The handle to the descriptor in the offline descriptor heap */
-	D3D12_CPU_DESCRIPTOR_HANDLE OfflineDescriptorHandle;
-
-	/** Index of the descriptor in the offline heap */
-	uint32 OfflineHeapIndex;
-
-	D3D12_CONSTANT_BUFFER_VIEW_DESC Desc;
-
-//protected:
-
-	explicit FD3D12ConstantBufferView(FD3D12Device* InParent, const D3D12_CONSTANT_BUFFER_VIEW_DESC* InDesc)
-		: FD3D12DeviceChild(InParent)
-		, OfflineHeapIndex(UINT_MAX)
-	{
-		OfflineDescriptorHandle.ptr = 0;
-		Init(InDesc);
-	}
-
-	virtual ~FD3D12ConstantBufferView()
-	{
-		FreeHeapSlot();
-	}
-
-	void Init(const D3D12_CONSTANT_BUFFER_VIEW_DESC* InDesc)
-	{
-		if (InDesc)
-		{
-			Desc = *InDesc;
-		}
-		else
-		{
-			FMemory::Memzero(&Desc, sizeof(Desc));
-		}
-		AllocateHeapSlot();
-	}
-
-	void AllocateHeapSlot();
+public:
+	FD3D12ConstantBufferView() = delete;
+	FD3D12ConstantBufferView(FD3D12Device* InParent);
 
-	void FreeHeapSlot();
+	inline D3D12_CPU_DESCRIPTOR_HANDLE            GetView() const { return Descriptor.GetOfflineHandle(); }
+	inline const D3D12_CONSTANT_BUFFER_VIEW_DESC& GetDesc() const { return Desc; }
 
 	void Create(D3D12_GPU_VIRTUAL_ADDRESS GPUAddress, const uint32 AlignedSize);
 
-public:
-	const D3D12_CONSTANT_BUFFER_VIEW_DESC& GetDesc() const { return Desc; }
+protected:
+	FD3D12ViewDescriptorHandle Descriptor;
+	D3D12_CONSTANT_BUFFER_VIEW_DESC Desc{};
 };
-#endif
 
 class FD3D12RenderTargetView : public FD3D12View<D3D12_RENDER_TARGET_VIEW_DESC>, public FRHIResource, public FD3D12LinkedAdapterObject<FD3D12RenderTargetView>
 {
 public:
 	FD3D12RenderTargetView(FD3D12Device* InParent, const D3D12_RENDER_TARGET_VIEW_DESC& InRTVDesc, FD3D12BaseShaderResource* InBaseShaderResource)
-		: FD3D12View(InParent, ViewSubresourceSubsetFlags_None)
+		: FD3D12View(InParent, ED3D12DescriptorHeapType::RenderTarget, ViewSubresourceSubsetFlags_None)
 		, FRHIResource(RRT_None)
 	{
-		CreateView(InRTVDesc, InBaseShaderResource, InBaseShaderResource->ResourceLocation);
+		Initialize(InRTVDesc, InBaseShaderResource, InBaseShaderResource->ResourceLocation);
 	}
 
 	FD3D12RenderTargetView(FD3D12Device* InParent, const D3D12_RENDER_TARGET_VIEW_DESC& InRTVDesc, FD3D12BaseShaderResource* InBaseShaderResource, FD3D12ResourceLocation& InResourceLocation)
-		: FD3D12View(InParent, ViewSubresourceSubsetFlags_None)
+		: FD3D12View(InParent, ED3D12DescriptorHeapType::RenderTarget, ViewSubresourceSubsetFlags_None)
 		, FRHIResource(RRT_None)
 	{
-		CreateView(InRTVDesc, InBaseShaderResource, InResourceLocation);
+		Initialize(InRTVDesc, InBaseShaderResource, InResourceLocation);
 	}
 
-	void CreateView(const D3D12_RENDER_TARGET_VIEW_DESC& InDesc, FD3D12BaseShaderResource* InBaseShaderResource, FD3D12ResourceLocation& InResourceLocation)
+	void Initialize(const D3D12_RENDER_TARGET_VIEW_DESC& InDesc, FD3D12BaseShaderResource* InBaseShaderResource, FD3D12ResourceLocation& InResourceLocation)
 	{
 		SetDesc(InDesc);
+		CreateView(InBaseShaderResource, InResourceLocation);
+	}
+
+	virtual void RecreateView() override
+	{
+		check(ResourceLocation->GetOffsetFromBaseOfResource() == 0);
+		CreateView(BaseShaderResource, BaseShaderResource->ResourceLocation);
+	}
+
+protected:
+	void CreateView(FD3D12BaseShaderResource* InBaseShaderResource, FD3D12ResourceLocation& InResourceLocation)
+	{
 		InitializeInternal(InBaseShaderResource, InResourceLocation);
 
 		if (ResourceLocation->GetResource())
@@ -1097,42 +1064,25 @@ public:
 			Descriptor.CreateView(Desc, D3DResource);
 		}
 	}
-
-	virtual void RecreateView() override
-	{
-		check(ResourceLocation->GetOffsetFromBaseOfResource() == 0);
-		CreateView(Desc, BaseShaderResource, BaseShaderResource->ResourceLocation);
-	}
 };
 
 class FD3D12DepthStencilView : public FD3D12View<D3D12_DEPTH_STENCIL_VIEW_DESC>, public FRHIResource, public FD3D12LinkedAdapterObject<FD3D12DepthStencilView>
 {
-	const bool bHasDepth;
-	const bool bHasStencil;
-	CViewSubresourceSubset DepthOnlyViewSubresourceSubset;
-	CViewSubresourceSubset StencilOnlyViewSubresourceSubset;
-
 public:
 	FD3D12DepthStencilView(FD3D12Device* InParent, const D3D12_DEPTH_STENCIL_VIEW_DESC& InDSVDesc, FD3D12BaseShaderResource* InBaseShaderResource, bool InHasStencil)
-		: FD3D12View(InParent, ViewSubresourceSubsetFlags_DepthAndStencilDsv)
+		: FD3D12View(InParent, ED3D12DescriptorHeapType::DepthStencil, ViewSubresourceSubsetFlags_DepthAndStencilDsv)
 		, FRHIResource(RRT_None)
 		, bHasDepth(true)				// Assume all DSVs have depth bits in their format
 		, bHasStencil(InHasStencil)		// Only some DSVs have stencil bits in their format
 	{
-		CreateView(InDSVDesc, InBaseShaderResource, InBaseShaderResource->ResourceLocation);
+		Initialize(InDSVDesc, InBaseShaderResource, InBaseShaderResource->ResourceLocation);
 		SetupDepthStencilViewSubresourceSubset();
 	}
 
-	void CreateView(const D3D12_DEPTH_STENCIL_VIEW_DESC& InDesc, FD3D12BaseShaderResource* InBaseShaderResource, FD3D12ResourceLocation& InResourceLocation)
+	void Initialize(const D3D12_DEPTH_STENCIL_VIEW_DESC& InDesc, FD3D12BaseShaderResource* InBaseShaderResource, FD3D12ResourceLocation& InResourceLocation)
 	{
 		SetDesc(InDesc);
-		InitializeInternal(InBaseShaderResource, InResourceLocation);
-
-		if (ResourceLocation->GetResource())
-		{
-			ID3D12Resource* D3DResource = ResourceLocation->GetResource()->GetResource();
-			Descriptor.CreateView(Desc, D3DResource);
-		}
+		CreateView(InBaseShaderResource, InResourceLocation);
 	}
 
 	bool HasDepth() const
@@ -1185,9 +1135,26 @@ public:
 	virtual void RecreateView() override
 	{
 		check(ResourceLocation->GetOffsetFromBaseOfResource() == 0);
-		CreateView(Desc, BaseShaderResource, BaseShaderResource->ResourceLocation);
+		CreateView(BaseShaderResource, BaseShaderResource->ResourceLocation);
 		SetupDepthStencilViewSubresourceSubset();
 	}
+
+protected:
+	void CreateView(FD3D12BaseShaderResource* InBaseShaderResource, FD3D12ResourceLocation& InResourceLocation)
+	{
+		InitializeInternal(InBaseShaderResource, InResourceLocation);
+
+		if (ResourceLocation->GetResource())
+		{
+			ID3D12Resource* D3DResource = ResourceLocation->GetResource()->GetResource();
+			Descriptor.CreateView(Desc, D3DResource);
+		}
+	}
+
+	CViewSubresourceSubset DepthOnlyViewSubresourceSubset;
+	CViewSubresourceSubset StencilOnlyViewSubresourceSubset;
+	const bool bHasDepth : 1;
+	const bool bHasStencil : 1;
 };
 
 template<>
diff --git a/Engine/Source/Runtime/D3D12RHI/Public/D3D12ConstantBuffer.h b/Engine/Source/Runtime/D3D12RHI/Public/D3D12ConstantBuffer.h
--- a/Engine/Source/Runtime/D3D12RHI/Public/D3D12ConstantBuffer.h
+++ b/Engine/Source/Runtime/D3D12RHI/Public/D3D12ConstantBuffer.h
@@ -19,9 +19,6 @@ class FD3D12FastConstantAllocator;
 class FD3D12ConstantBuffer : public FD3D12DeviceChild
 {
 public:
-#if USE_STATIC_ROOT_SIGNATURE
-	FD3D12ConstantBufferView* View;
-#endif
 	FD3D12ConstantBuffer(FD3D12Device* InParent, FD3D12FastConstantAllocator& InAllocator);
 	virtual ~FD3D12ConstantBuffer();
 
@@ -45,7 +42,15 @@ public:
 
 	bool Version(FD3D12ResourceLocation& BufferOut, bool bDiscardSharedConstants);
 
+#if USE_STATIC_ROOT_SIGNATURE
+	inline D3D12_CPU_DESCRIPTOR_HANDLE GetView() const { return View->GetView(); }
+#endif
+
 protected:
+#if USE_STATIC_ROOT_SIGNATURE
+	FD3D12ConstantBufferView* View = nullptr;
+#endif
+
 	__declspec(align(16)) uint8 ShadowData[MAX_GLOBAL_CONSTANT_BUFFER_SIZE];
 	
 	/** Size of all constants that has been updated since the last call to Commit. */
