New version of RHICreateRayTracingScene.
Moved calculating RT instance prefix sums and some other logic to higher level code.

diff --git a/Engine/Plugins/Experimental/GPULightmass/Source/GPULightmass/Private/LightmapRenderer.cpp b/Engine/Plugins/Experimental/GPULightmass/Source/GPULightmass/Private/LightmapRenderer.cpp
--- a/Engine/Plugins/Experimental/GPULightmass/Source/GPULightmass/Private/LightmapRenderer.cpp
+++ b/Engine/Plugins/Experimental/GPULightmass/Source/GPULightmass/Private/LightmapRenderer.cpp
@@ -25,6 +25,7 @@
 #include "RayTracingGeometryManager.h"
 #include "ScreenPass.h"
 #include "RayTracingDynamicGeometryCollection.h"
+#include "Experimental/Containers/SherwoodHashTable.h"
 
 class FCopyConvergedLightmapTilesCS : public FGlobalShader
 {
@@ -165,6 +166,62 @@ struct FGPUBatchedTileRequests
 	TResourceArray<FGPUTileDescription> BatchedTilesDesc;
 };
 
+static FRayTracingSceneRHIRef CreateRayTracingScene(TArrayView<FRayTracingGeometryInstance> Instances, uint32& OutNumNativeInstances)
+{
+	const uint32 NumSceneInstances = Instances.Num();
+
+	FRayTracingSceneInitializer2 Initializer;
+	Initializer.Instances = Instances;
+	Initializer.ShaderSlotsPerGeometrySegment = RAY_TRACING_NUM_SHADER_SLOTS;
+	Initializer.PerInstanceGeometries.SetNumUninitialized(NumSceneInstances);
+	Initializer.PerInstanceNumTransforms.SetNumUninitialized(NumSceneInstances);
+	Initializer.BaseInstancePrefixSum.SetNumUninitialized(NumSceneInstances);
+	Initializer.SegmentPrefixSum.SetNumUninitialized(NumSceneInstances);
+	Initializer.NumNativeInstances = 0;
+	Initializer.NumTotalSegments = 0;
+
+	Experimental::TSherwoodSet<FRHIRayTracingGeometry*> UniqueGeometries;
+
+	// Compute geometry segment and instance count prefix sums.
+	// These are later used by GetHitRecordBaseIndex() during resource binding
+	// and by GetBaseInstanceIndex() in shaders to emulate SV_InstanceIndex.
+
+	for (uint32 InstanceIndex = 0; InstanceIndex < NumSceneInstances; ++InstanceIndex)
+	{
+		const FRayTracingGeometryInstance& InstanceDesc = Instances[InstanceIndex];
+
+		checkf(InstanceDesc.GPUTransformsSRV || InstanceDesc.NumTransforms <= uint32(InstanceDesc.Transforms.Num()),
+			TEXT("Expected at most %d ray tracing geometry instance transforms, but got %d."),
+			InstanceDesc.NumTransforms, InstanceDesc.Transforms.Num());
+
+		checkf(InstanceDesc.GeometryRHI, TEXT("Ray tracing instance must have a valid geometry."));
+
+		ensure(InstanceDesc.NumTransforms == uint32(InstanceDesc.Transforms.Num()));
+
+		Initializer.PerInstanceGeometries[InstanceIndex] = InstanceDesc.GeometryRHI;
+
+		// Compute geometry segment count prefix sum to be later used in GetHitRecordBaseIndex()
+		Initializer.SegmentPrefixSum[InstanceIndex] = Initializer.NumTotalSegments;
+		Initializer.NumTotalSegments += InstanceDesc.GeometryRHI->GetNumSegments();
+
+		bool bIsAlreadyInSet = false;
+		UniqueGeometries.Add(InstanceDesc.GeometryRHI, &bIsAlreadyInSet);
+		if (!bIsAlreadyInSet)
+		{
+			Initializer.ReferencedGeometries.Add(InstanceDesc.GeometryRHI);
+		}
+
+		Initializer.BaseInstancePrefixSum[InstanceIndex] = Initializer.NumNativeInstances;
+		Initializer.NumNativeInstances += InstanceDesc.NumTransforms;
+
+		Initializer.PerInstanceNumTransforms[InstanceIndex] = InstanceDesc.NumTransforms;
+	}
+
+	OutNumNativeInstances = Initializer.NumNativeInstances;
+
+	return RHICreateRayTracingScene(MoveTemp(Initializer));
+}
+
 namespace GPULightmass
 {
 
@@ -999,21 +1056,12 @@ void FSceneRenderState::SetupRayTracingScene(int32 LODIndex)
 			}
 		}
 
-		FRayTracingSceneInitializer Initializer;
-		Initializer.Instances = RayTracingGeometryInstances;
-		Initializer.ShaderSlotsPerGeometrySegment = RAY_TRACING_NUM_SHADER_SLOTS;
 		if (IsRayTracingEnabled())
 		{
 			SCOPED_GPU_MASK(RHICmdList, FRHIGPUMask::All());
 
-			RayTracingScene = RHICreateRayTracingScene(Initializer);
-
 			uint32 NumTotalInstances = 0;
-			for (const FRayTracingGeometryInstance& Instance : RayTracingGeometryInstances)
-			{
-				ensure(Instance.NumTransforms == uint32(Instance.Transforms.Num()));
-				NumTotalInstances += Instance.NumTransforms;
-			}
+			RayTracingScene = CreateRayTracingScene(RayTracingGeometryInstances, NumTotalInstances);
 
 			FRayTracingAccelerationStructureSize SizeInfo = RHICalcRayTracingSceneSize(NumTotalInstances, ERayTracingAccelerationStructureFlags::FastTrace);
 			FRHIResourceCreateInfo BufferCreateInfo(TEXT("LightmassRayTracingSceneBuffer"));
diff --git a/Engine/Source/Runtime/D3D12RHI/Private/D3D12RHIPrivate.h b/Engine/Source/Runtime/D3D12RHI/Private/D3D12RHIPrivate.h
--- a/Engine/Source/Runtime/D3D12RHI/Private/D3D12RHIPrivate.h
+++ b/Engine/Source/Runtime/D3D12RHI/Private/D3D12RHIPrivate.h
@@ -495,6 +495,7 @@ public:
 
 	virtual FRayTracingGeometryRHIRef RHICreateRayTracingGeometry(const FRayTracingGeometryInitializer& Initializer) final override;
 	virtual FRayTracingSceneRHIRef RHICreateRayTracingScene(const FRayTracingSceneInitializer& Initializer) final override;
+	virtual FRayTracingSceneRHIRef RHICreateRayTracingScene(FRayTracingSceneInitializer2 Initializer) final override;
 	virtual FRayTracingShaderRHIRef RHICreateRayTracingShader(TArrayView<const uint8> Code, const FSHAHash& Hash, EShaderFrequency ShaderFrequency) final override;
 	virtual FRayTracingPipelineStateRHIRef RHICreateRayTracingPipelineState(const FRayTracingPipelineStateInitializer& Initializer) final override;
 
diff --git a/Engine/Source/Runtime/D3D12RHI/Private/D3D12RayTracing.cpp b/Engine/Source/Runtime/D3D12RHI/Private/D3D12RayTracing.cpp
--- a/Engine/Source/Runtime/D3D12RHI/Private/D3D12RayTracing.cpp
+++ b/Engine/Source/Runtime/D3D12RHI/Private/D3D12RayTracing.cpp
@@ -3063,9 +3063,63 @@ FRayTracingSceneRHIRef FD3D12DynamicRHI::RHICreateRayTracingScene(const FRayTrac
 {
 	TRACE_CPUPROFILER_EVENT_SCOPE(CreateRayTracingScene);
 
+	const uint32 NumSceneInstances = Initializer.Instances.Num();
+
+	FRayTracingSceneInitializer2 Initializer2;
+	Initializer2.DebugName = Initializer.DebugName;
+	Initializer2.ShaderSlotsPerGeometrySegment = Initializer.ShaderSlotsPerGeometrySegment;
+	Initializer2.NumMissShaderSlots = Initializer.NumMissShaderSlots;
+	Initializer2.Instances = Initializer.Instances;
+	Initializer2.PerInstanceGeometries.SetNumUninitialized(NumSceneInstances);
+	Initializer2.PerInstanceNumTransforms.SetNumUninitialized(NumSceneInstances);
+	Initializer2.BaseInstancePrefixSum.SetNumUninitialized(NumSceneInstances);
+	Initializer2.SegmentPrefixSum.SetNumUninitialized(NumSceneInstances);
+	Initializer2.NumNativeInstances = 0;
+	Initializer2.NumTotalSegments = 0;
+
+	Experimental::TSherwoodSet<FRHIRayTracingGeometry*> UniqueGeometries;
+
+	for (uint32 InstanceIndex = 0; InstanceIndex < NumSceneInstances; ++InstanceIndex)
+	{
+		const FRayTracingGeometryInstance& InstanceDesc = Initializer.Instances[InstanceIndex];
+
+		checkf(InstanceDesc.GPUTransformsSRV || InstanceDesc.NumTransforms <= uint32(InstanceDesc.Transforms.Num()),
+			TEXT("Expected at most %d ray tracing geometry instance transforms, but got %d."),
+			InstanceDesc.NumTransforms, InstanceDesc.Transforms.Num());
+
+		checkf(InstanceDesc.GeometryRHI, TEXT("Ray tracing instance must have a valid geometry."));
+
+		Initializer2.PerInstanceGeometries[InstanceIndex] = InstanceDesc.GeometryRHI;
+
+		// Compute geometry segment count prefix sum to be later used in GetHitRecordBaseIndex()
+		Initializer2.SegmentPrefixSum[InstanceIndex] = Initializer2.NumTotalSegments;
+		Initializer2.NumTotalSegments += InstanceDesc.GeometryRHI->GetNumSegments();
+
+		bool bIsAlreadyInSet = false;
+		UniqueGeometries.Add(InstanceDesc.GeometryRHI, &bIsAlreadyInSet);
+		if (!bIsAlreadyInSet)
+		{
+			Initializer2.ReferencedGeometries.Add(InstanceDesc.GeometryRHI);
+		}
+
+		Initializer2.BaseInstancePrefixSum[InstanceIndex] = Initializer2.NumNativeInstances;
+		Initializer2.NumNativeInstances += InstanceDesc.NumTransforms;
+
+		Initializer2.PerInstanceNumTransforms[InstanceIndex] = InstanceDesc.NumTransforms;
+	}
+
 	FD3D12Adapter& Adapter = GetAdapter();
 
-	return new FD3D12RayTracingScene(&Adapter, Initializer);
+	return new FD3D12RayTracingScene(&Adapter, MoveTemp(Initializer2));
+}
+
+FRayTracingSceneRHIRef FD3D12DynamicRHI::RHICreateRayTracingScene(FRayTracingSceneInitializer2 Initializer)
+{
+	TRACE_CPUPROFILER_EVENT_SCOPE(CreateRayTracingScene);
+
+	FD3D12Adapter& Adapter = GetAdapter();
+
+	return new FD3D12RayTracingScene(&Adapter, MoveTemp(Initializer));
 }
 
 FBufferRHIRef FD3D12RayTracingGeometry::NullTransformBuffer;
@@ -3607,57 +3661,23 @@ void FD3D12RayTracingGeometry::CompactAccelerationStructure(FD3D12CommandContext
 	AccelerationStructureCompactedSize = InSizeAfterCompaction;
 }
 
-FD3D12RayTracingScene::FD3D12RayTracingScene(FD3D12Adapter* Adapter, const FRayTracingSceneInitializer& Initializer)
-	: FD3D12AdapterChild(Adapter)
+FD3D12RayTracingScene::FD3D12RayTracingScene(FD3D12Adapter* Adapter, FRayTracingSceneInitializer2 InInitializer)
+	: FD3D12AdapterChild(Adapter), Initializer(MoveTemp(InInitializer))
 {
 	INC_DWORD_STAT(STAT_D3D12RayTracingAllocatedTLAS);
 
-	Instances.SetAllowCPUAccess(true);
-
-	DebugName = Initializer.DebugName.IsValid() ? Initializer.DebugName : FName(TEXT("TLAS"));
-
+	checkf(Initializer.NumMissShaderSlots >= 1, TEXT("Need at least 1 miss shader slot."));
 	checkf(Initializer.Lifetime == RTSL_SingleFrame, TEXT("Only single-frame ray tracing scenes are currently implemented."));
 
-	Lifetime = Initializer.Lifetime;
-	CreatedFrameFenceValue = Adapter->GetFrameFence().GetCurrentFence();
-
-	ShaderSlotsPerGeometrySegment = Initializer.ShaderSlotsPerGeometrySegment;
-
-	Experimental::TSherwoodSet<FD3D12RayTracingGeometry*> UniqueGeometries;
-
-	const int32 NumSceneInstances = Initializer.Instances.Num();
-
 	int32 NumDxrInstances = 0;
 	int32 NumInactiveDxrInstances = 0;
 	int32 NumCopyCommands = 0;
-	NumTotalSegments = 0;
-
-	BaseInstancePrefixSum.Reserve(NumSceneInstances);
-	PerInstanceGeometries.Reserve(NumSceneInstances);
-	PerInstanceNumTransforms.Reserve(NumSceneInstances);
-	SegmentPrefixSum.Reserve(NumSceneInstances);
 
+	const int32 NumSceneInstances = Initializer.Instances.Num();
 	for (int32 InstanceIndex = 0; InstanceIndex < NumSceneInstances; ++InstanceIndex)
 	{
 		const FRayTracingGeometryInstance& Instance = Initializer.Instances[InstanceIndex];
 
-		checkf(Instance.GeometryRHI, TEXT("Ray tracing instance must have a valid geometry."));
-
-		FD3D12RayTracingGeometry* Geometry = FD3D12DynamicRHI::ResourceCast(Instance.GeometryRHI);
-
-		PerInstanceGeometries.Add(Geometry);
-
-		// Compute geometry segment count prefix sum to be later used in GetHitRecordBaseIndex()
-		SegmentPrefixSum.Add(NumTotalSegments);
-		NumTotalSegments += Geometry->Segments.Num();
-
-		bool bIsAlreadyInSet = false;
-		UniqueGeometries.Add(Geometry, &bIsAlreadyInSet);
-		if (!bIsAlreadyInSet)
-		{
-			ReferencedGeometries.Add(Geometry);
-		}
-
 		if (UNLIKELY(Instance.GPUTransformsSRV != nullptr && Instance.NumTransforms))
 		{
 			FInstanceCopyCommand Command;
@@ -3667,27 +3687,23 @@ FD3D12RayTracingScene::FD3D12RayTracingScene(FD3D12Adapter* Adapter, const FRayT
 			CopyCommands.Add(Command);
 		}
 
-		BaseInstancePrefixSum.Add(NumDxrInstances);
 		NumDxrInstances += Instance.NumTransforms;
-
-		PerInstanceNumTransforms.Add(Instance.NumTransforms);
-
-		TotalPrimitiveCount += uint64(Geometry->TotalPrimitiveCount) * Instance.NumTransforms;
 	}
 
+	Instances.SetAllowCPUAccess(true);
 	Instances.SetNumUninitialized(NumDxrInstances);
 
 	const EParallelForFlags ParallelForFlags = EParallelForFlags::None; // set ForceSingleThread for testing
-	ParallelFor(NumSceneInstances, [this, &Initializer, &NumInactiveDxrInstances](int32 InstanceIndex)
+	ParallelFor(NumSceneInstances, [this, &NumInactiveDxrInstances](int32 InstanceIndex)
 	{
 		const FRayTracingGeometryInstance& Instance = Initializer.Instances[InstanceIndex];
-		FD3D12RayTracingGeometry* Geometry = PerInstanceGeometries[InstanceIndex];
+		FD3D12RayTracingGeometry* Geometry = FD3D12DynamicRHI::ResourceCast(Initializer.PerInstanceGeometries[InstanceIndex]);
 
 		D3D12_RAYTRACING_INSTANCE_DESC InstanceDesc = {};
 		check(InstanceDesc.AccelerationStructure == 0);
 
 		InstanceDesc.InstanceMask = Instance.Mask;
-		InstanceDesc.InstanceContributionToHitGroupIndex = SegmentPrefixSum[InstanceIndex] * ShaderSlotsPerGeometrySegment;
+		InstanceDesc.InstanceContributionToHitGroupIndex = Initializer.SegmentPrefixSum[InstanceIndex] * Initializer.ShaderSlotsPerGeometrySegment;
 		InstanceDesc.Flags = TranslateRayTracingInstanceFlags(Instance.Flags);
 
 		const uint32 NumTransforms = Instance.NumTransforms;
@@ -3700,7 +3716,7 @@ FD3D12RayTracingScene::FD3D12RayTracingScene(FD3D12Adapter* Adapter, const FRayT
 		const bool bGpuInstance = Instance.GPUTransformsSRV != nullptr;
 		const bool bCpuInstance = !bGpuInstance;
 
-		uint32 DescIndex = BaseInstancePrefixSum[InstanceIndex];
+		uint32 DescIndex = Initializer.BaseInstancePrefixSum[InstanceIndex];
 
 		int32 NumInactiveDxrInstancesThisSceneInstance = 0;
 
@@ -3762,13 +3778,10 @@ FD3D12RayTracingScene::FD3D12RayTracingScene(FD3D12Adapter* Adapter, const FRayT
 	BuildInputs.NumDescs = Instances.Num();
 	BuildInputs.Flags = TranslateRayTracingAccelerationStructureFlags(BuildFlags);
 
-	NumCallableShaderSlots = Initializer.NumCallableShaderSlots;
-	NumMissShaderSlots = FMath::Max<uint32>(1, Initializer.NumMissShaderSlots);
-
 	// Get maximum buffer sizes for all GPUs in the system
 	SizeInfo = RHICalcRayTracingSceneSize(BuildInputs.NumDescs, BuildFlags);
 
-	SET_DWORD_STAT(STAT_RayTracingInstances, NumDxrInstances - NumInactiveDxrInstances);
+	SET_DWORD_STAT(STAT_RayTracingInstances, Initializer.NumNativeInstances - NumInactiveDxrInstances);
 };
 
 FD3D12RayTracingScene::~FD3D12RayTracingScene()
@@ -3860,7 +3873,7 @@ void FD3D12RayTracingScene::BuildAccelerationStructure(FD3D12CommandContext& Com
 		TEXT("BuildInputs.NumDescs = %d, Instances.Num = %d, PerInstanceGeometries.Num = %d."),
 			ScratchBuffer->GetSize(), ScratchBufferOffset, ScratchBuffer->GetSize() - ScratchBufferOffset, 
 			PrebuildInfo.ScratchDataSizeInBytes,
-			BuildInputs.NumDescs, Instances.Num(), PerInstanceGeometries.Num());
+			BuildInputs.NumDescs, Instances.Num(), Initializer.PerInstanceGeometries.Num());
 
 	TRefCountPtr<FD3D12Buffer> AutoInstanceBuffer;
 
@@ -3871,19 +3884,19 @@ void FD3D12RayTracingScene::BuildAccelerationStructure(FD3D12CommandContext& Com
 		// - Set up acceleration structure pointers and make them resident.
 		// - Generate HitGroupSystemParametersCache.
 
-		HitGroupSystemParametersCache.Reserve(NumTotalSegments);
-		HitGroupSystemParametersCache.Empty(NumTotalSegments);
+		HitGroupSystemParametersCache.Reserve(Initializer.NumTotalSegments);
+		HitGroupSystemParametersCache.Empty(Initializer.NumTotalSegments);
 
 		{
 			int32 DxrInstanceIndex = 0;
-			const int32 NumSceneInstances = PerInstanceGeometries.Num();
+			const int32 NumSceneInstances = Initializer.PerInstanceGeometries.Num();
 			for (int32 InstanceIndex = 0; InstanceIndex < NumSceneInstances; ++InstanceIndex)
 			{
-				FD3D12RayTracingGeometry* Geometry = PerInstanceGeometries[InstanceIndex];
+				FD3D12RayTracingGeometry* Geometry = FD3D12DynamicRHI::ResourceCast(Initializer.PerInstanceGeometries[InstanceIndex]);
 
 				D3D12_GPU_VIRTUAL_ADDRESS BlasAddress = Geometry->AccelerationStructureBuffers[GPUIndex]->ResourceLocation.GetGPUVirtualAddress();
 
-				const int32 NumTransforms = PerInstanceNumTransforms[InstanceIndex];
+				const int32 NumTransforms = Initializer.PerInstanceNumTransforms[InstanceIndex];
 				for (int32 TransformIndex = 0; TransformIndex < NumTransforms; ++TransformIndex)
 				{
 					// If this instance has been deactivated, set BLAS address to 0
@@ -3984,8 +3997,11 @@ void FD3D12RayTracingScene::BuildAccelerationStructure(FD3D12CommandContext& Com
 			#endif // ENABLE_RESIDENCY_MANAGEMENT
 			};
 
-			for (TRefCountPtr<FD3D12RayTracingGeometry>& Geometry : ReferencedGeometries)
+			const uint32 NumReferencedGeometries = Initializer.ReferencedGeometries.Num();
+			for (uint32 Index = 0; Index < NumReferencedGeometries; ++Index)
 			{
+				FD3D12RayTracingGeometry* Geometry = FD3D12DynamicRHI::ResourceCast(Initializer.ReferencedGeometries[Index].GetReference());
+
 				checkf(!Geometry->IsDirty(CommandContext.GetGPUIndex()),
 					TEXT("Acceleration structures for all geometries must be built before building the top level acceleration structure for the scene."));
 
@@ -4011,9 +4027,10 @@ void FD3D12RayTracingScene::BuildAccelerationStructure(FD3D12CommandContext& Com
 	
 	
 	// Build the actual acceleration structure
-
-	for (TRefCountPtr<FD3D12RayTracingGeometry>& Geometry : ReferencedGeometries)
+	const uint32 NumReferencedGeometries = Initializer.ReferencedGeometries.Num();
+	for (uint32 Index = 0; Index < NumReferencedGeometries; ++Index)
 	{
+		FD3D12RayTracingGeometry* Geometry = FD3D12DynamicRHI::ResourceCast(Initializer.ReferencedGeometries[Index].GetReference());
 		Geometry->AccelerationStructureBuffers[GPUIndex]->ResourceLocation.GetResource()->UpdateResidency(CommandContext.CommandListHandle);
 	}
 
@@ -4101,16 +4118,16 @@ FD3D12RayTracingShaderTable* FD3D12RayTracingScene::FindOrCreateShaderTable(cons
 	ID3D12Device5* RayTracingDevice = Device->GetDevice5();
 	const uint32 GPUIndex = Device->GetGPUIndex();
 
-	const uint32 NumHitGroupSlots = Pipeline->bAllowHitGroupIndexing ? NumTotalSegments * ShaderSlotsPerGeometrySegment : 0;
+	const uint32 NumHitGroupSlots = Pipeline->bAllowHitGroupIndexing ? Initializer.NumTotalSegments * Initializer.ShaderSlotsPerGeometrySegment : 0;
 
 	checkf(Pipeline->MaxLocalRootSignatureSize >= sizeof(FHitGroupSystemParameters), TEXT("All local root signatures are expected to contain ray tracing system root parameters (2x root buffers + 4x root DWORD)"));
 
 	FD3D12RayTracingShaderTable::FInitializer SBTInitializer = {};
 	SBTInitializer.NumRayGenShaders = Pipeline->RayGenShaders.Identifiers.Num();
 	SBTInitializer.NumMissShaders = Pipeline->MissShaders.Identifiers.Num();
-	SBTInitializer.NumMissRecords = NumMissShaderSlots;
+	SBTInitializer.NumMissRecords = Initializer.NumMissShaderSlots;
 	SBTInitializer.NumHitRecords = NumHitGroupSlots;
-	SBTInitializer.NumCallableRecords = NumCallableShaderSlots;
+	SBTInitializer.NumCallableRecords = Initializer.NumCallableShaderSlots;
 	SBTInitializer.LocalRootDataSize = Pipeline->MaxLocalRootSignatureSize;
 	SBTInitializer.MaxViewDescriptorsPerRecord = Pipeline->MaxHitGroupViewDescriptors;
 
@@ -5176,7 +5193,7 @@ static void SetRayTracingHitGroup(
 	uint32 UserData,
 	uint32 WorkerIndex)
 {
-	checkf(ShaderSlot < Scene->ShaderSlotsPerGeometrySegment, TEXT("Shader slot is invalid. Make sure that ShaderSlotsPerGeometrySegment is correct on FRayTracingSceneInitializer."));
+	checkf(ShaderSlot < Scene->Initializer.ShaderSlotsPerGeometrySegment, TEXT("Shader slot is invalid. Make sure that ShaderSlotsPerGeometrySegment is correct on FRayTracingSceneInitializer."));
 
 	const uint32 RecordIndex = Scene->GetHitRecordBaseIndex(InstanceIndex, SegmentIndex) + ShaderSlot;
 	const uint32 GPUIndex = Device->GetGPUIndex();
@@ -5184,18 +5201,18 @@ static void SetRayTracingHitGroup(
 	FHitGroupSystemParameters SystemParameters;
 	if (GPUIndex == 0)
 	{
-		uint32 PrefixedSegmentIndex = Scene->SegmentPrefixSum[InstanceIndex];
+		uint32 PrefixedSegmentIndex = Scene->Initializer.SegmentPrefixSum[InstanceIndex];
 		SystemParameters = Scene->HitGroupSystemParametersCache[PrefixedSegmentIndex + SegmentIndex];
 	}
 	else
 	{
-		const FD3D12RayTracingGeometry* Geometry = Scene->PerInstanceGeometries[InstanceIndex];
+		const FD3D12RayTracingGeometry* Geometry = FD3D12DynamicRHI::ResourceCast(Scene->Initializer.PerInstanceGeometries[InstanceIndex]);
 		const TArray<FHitGroupSystemParameters>& HitGroupSystemParametersForThisGPU = Geometry->HitGroupSystemParameters[GPUIndex];
 
 		SystemParameters = HitGroupSystemParametersForThisGPU[SegmentIndex];
 	}
 
-	SystemParameters.RootConstants.BaseInstanceIndex = Scene->BaseInstancePrefixSum[InstanceIndex];
+	SystemParameters.RootConstants.BaseInstanceIndex = Scene->Initializer.BaseInstancePrefixSum[InstanceIndex];
 	SystemParameters.RootConstants.UserData = UserData;
 	ShaderTable->SetHitGroupSystemParameters(RecordIndex, SystemParameters);
 
@@ -5204,7 +5221,7 @@ static void SetRayTracingHitGroup(
 	FD3D12RayTracingShaderTable::FShaderRecordCacheKey CacheKey;
 
 	const bool bCanUseRecordCache = GRayTracingCacheShaderRecords
-		&& Scene->Lifetime == RTSL_SingleFrame
+		&& Scene->Initializer.Lifetime == RTSL_SingleFrame
 		&& LooseParameterDataSize == 0 // loose parameters end up in unique constant buffers, so SBT records can't be shared
 		&& NumUniformBuffers > 0 // there is no benefit from cache if no resources are being bound
 		&& NumUniformBuffers <= CacheKey.MaxUniformBuffers;
@@ -5329,7 +5346,7 @@ void FD3D12CommandContext::RHISetRayTracingHitGroup(
 	FD3D12RayTracingScene* Scene = FD3D12DynamicRHI::ResourceCast(InScene);
 	FD3D12RayTracingPipelineState* Pipeline = FD3D12DynamicRHI::ResourceCast(InPipeline);
 
-	checkf(ShaderSlot < Scene->ShaderSlotsPerGeometrySegment, TEXT("Shader slot is invalid. Make sure that ShaderSlotsPerGeometrySegment is correct on FRayTracingSceneInitializer."));
+	checkf(ShaderSlot < Scene->Initializer.ShaderSlotsPerGeometrySegment, TEXT("Shader slot is invalid. Make sure that ShaderSlotsPerGeometrySegment is correct on FRayTracingSceneInitializer."));
 
 	FD3D12RayTracingShaderTable* ShaderTable = Scene->FindOrCreateShaderTable(Pipeline, GetParentDevice());
 	checkf(ShaderTable->LocalShaderTableOffset == ShaderTable->HitGroupShaderTableOffset,
@@ -5358,7 +5375,7 @@ void FD3D12CommandContext::RHISetRayTracingCallableShader(
 	FD3D12RayTracingScene* Scene = FD3D12DynamicRHI::ResourceCast(InScene);
 	FD3D12RayTracingPipelineState* Pipeline = FD3D12DynamicRHI::ResourceCast(InPipeline);
 
-	checkf(ShaderSlotInScene < Scene->NumCallableShaderSlots, TEXT("Shader slot is invalid. Make sure that NumCallableShaderSlots is correct on FRayTracingSceneInitializer."));
+	checkf(ShaderSlotInScene < Scene->Initializer.NumCallableShaderSlots, TEXT("Shader slot is invalid. Make sure that NumCallableShaderSlots is correct on FRayTracingSceneInitializer."));
 
 	FD3D12RayTracingShaderTable* ShaderTable = Scene->FindOrCreateShaderTable(Pipeline, GetParentDevice());
 
@@ -5395,7 +5412,7 @@ void FD3D12CommandContext::RHISetRayTracingMissShader(
 	FD3D12RayTracingScene* Scene = FD3D12DynamicRHI::ResourceCast(InScene);
 	FD3D12RayTracingPipelineState* Pipeline = FD3D12DynamicRHI::ResourceCast(InPipeline);
 
-	checkf(ShaderSlotInScene < Scene->NumMissShaderSlots, TEXT("Shader slot is invalid. Make sure that NumMissShaderSlots is correct on FRayTracingSceneInitializer."));
+	checkf(ShaderSlotInScene < Scene->Initializer.NumMissShaderSlots, TEXT("Shader slot is invalid. Make sure that NumMissShaderSlots is correct on FRayTracingSceneInitializer."));
 	checkf(ShaderSlotInScene > 0, TEXT("Shader slot 0 is reserved for the default MissShader."));
 
 	FD3D12RayTracingShaderTable* ShaderTable = Scene->FindOrCreateShaderTable(Pipeline, GetParentDevice());
diff --git a/Engine/Source/Runtime/D3D12RHI/Private/D3D12RayTracing.h b/Engine/Source/Runtime/D3D12RHI/Private/D3D12RayTracing.h
--- a/Engine/Source/Runtime/D3D12RHI/Private/D3D12RayTracing.h
+++ b/Engine/Source/Runtime/D3D12RHI/Private/D3D12RayTracing.h
@@ -26,6 +26,8 @@ public:
 	FD3D12RayTracingGeometry(FD3D12Adapter* Adapter, const FRayTracingGeometryInitializer& Initializer);
 	~FD3D12RayTracingGeometry();
 
+	virtual uint32 GetNumSegments() const final override { return Segments.Num(); }
+
 	void SetupHitGroupSystemParameters(uint32 InGPUIndex);
 	void TransitionBuffers(FD3D12CommandContext& CommandContext);
 	void UpdateResidency(FD3D12CommandContext& CommandContext);
@@ -91,9 +93,11 @@ public:
 	// Scaling beyond 5 total threads does not yield any speedup in practice.
 	static constexpr uint32 MaxBindingWorkers = 5; // RHI thread + 4 parallel workers.
 
-	FD3D12RayTracingScene(FD3D12Adapter* Adapter, const FRayTracingSceneInitializer& Initializer);
+	FD3D12RayTracingScene(FD3D12Adapter* Adapter, FRayTracingSceneInitializer2 Initializer);
 	~FD3D12RayTracingScene();
 
+	const FRayTracingSceneInitializer2& GetInitializer() const override final { return Initializer; }
+
 	void BindBuffer(FRHIBuffer* Buffer, uint32 BufferOffset);
 	void ReleaseBuffer();
 
@@ -120,14 +124,8 @@ public:
 		int32 Num = 0;
 	};
 	TArray<FInstanceCopyCommand> CopyCommands;
-	
-	// One entry per instance in the Initializer
-	TArray<FD3D12RayTracingGeometry*> PerInstanceGeometries;
-	TArray<uint32> PerInstanceNumTransforms;
 
-	// Unique list of geometries referenced by all instances in this scene.
-	// Any referenced geometry is kept alive while the scene is alive.
-	TArray<TRefCountPtr<FD3D12RayTracingGeometry>> ReferencedGeometries;
+	const FRayTracingSceneInitializer2 Initializer;
 
 	// Scene keeps track of child acceleration structure buffers to ensure
 	// they are resident when any ray tracing work is dispatched.
@@ -135,20 +133,7 @@ public:
 
 	void UpdateResidency(FD3D12CommandContext& CommandContext);
 
-	uint32 ShaderSlotsPerGeometrySegment = 1;
-
-	// Exclusive prefix sum of `Instance.NumTransforms` for all instances in this scene. Used to emulate SV_InstanceID in hit shaders.
-	TArray<uint32> BaseInstancePrefixSum;
-
-	// Exclusive prefix sum of instance geometry segments is used to calculate SBT record address from instance and segment indices.
-	TArray<uint32> SegmentPrefixSum;
-	uint32 NumTotalSegments = 0;
-	uint32 GetHitRecordBaseIndex(uint32 InstanceIndex, uint32 SegmentIndex) const { return (SegmentPrefixSum[InstanceIndex] + SegmentIndex) * ShaderSlotsPerGeometrySegment; }
-
-	uint64 TotalPrimitiveCount = 0; // Combined number of primitives in all geometry instances
-
-	uint32 NumCallableShaderSlots = 0;
-	uint32 NumMissShaderSlots = 1; // always at least the default
+	uint32 GetHitRecordBaseIndex(uint32 InstanceIndex, uint32 SegmentIndex) const { return (Initializer.SegmentPrefixSum[InstanceIndex] + SegmentIndex) * Initializer.ShaderSlotsPerGeometrySegment; }
 
 	// Array of hit group parameters per geometry segment across all scene instance geometries.
 	// Accessed as HitGroupSystemParametersCache[SegmentPrefixSum[InstanceIndex] + SegmentIndex].
@@ -161,12 +146,7 @@ public:
 
 	TMap<const FD3D12RayTracingPipelineState*, FD3D12RayTracingShaderTable*> ShaderTables[MAX_NUM_GPUS];
 
-	ERayTracingSceneLifetime Lifetime = RTSL_SingleFrame;
-	uint64 CreatedFrameFenceValue = 0;
-
 	uint64 LastCommandListID = 0;
-
-	FName DebugName;
 };
 
 // Manages all the pending BLAS compaction requests
diff --git a/Engine/Source/Runtime/RHI/Public/DynamicRHI.h b/Engine/Source/Runtime/RHI/Public/DynamicRHI.h
--- a/Engine/Source/Runtime/RHI/Public/DynamicRHI.h
+++ b/Engine/Source/Runtime/RHI/Public/DynamicRHI.h
@@ -252,6 +252,50 @@ struct FRayTracingSceneInitializer
 	// Currently only single-frame lifetime is supported.
 	ERayTracingSceneLifetime Lifetime = RTSL_SingleFrame;
 
+	FName DebugName;
+};
+
+struct FRayTracingSceneInitializer2
+{
+	// WIP - need these arrays until building instance buffer is moved to higher level code
+	// will be removed before 5.0 release
+	TArrayView<FRayTracingGeometryInstance> Instances;
+	TArray<uint32> PerInstanceNumTransforms;
+
+	// Unique list of geometries referenced by all instances in this scene.
+	// Any referenced geometry is kept alive while the scene is alive.
+	TArray<FRayTracingGeometryRHIRef> ReferencedGeometries;
+	// One entry per instance
+	TArray<FRHIRayTracingGeometry*> PerInstanceGeometries;
+	// Exclusive prefix sum of `Instance.NumTransforms` for all instances in this scene. Used to emulate SV_InstanceID in hit shaders.
+	TArray<uint32> BaseInstancePrefixSum;
+	// Exclusive prefix sum of instance geometry segments is used to calculate SBT record address from instance and segment indices.
+	TArray<uint32> SegmentPrefixSum;
+
+	// Total flattened number of ray tracing geometry instances (a single FRayTracingGeometryInstance may represent many).
+	uint32 NumNativeInstances = 0;
+
+	uint32 NumTotalSegments = 0;
+
+	// This value controls how many elements will be allocated in the shader binding table per geometry segment.
+	// Changing this value allows different hit shaders to be used for different effects.
+	// For example, setting this to 2 allows one hit shader for regular material evaluation and a different one for shadows.
+	// Desired hit shader can be selected by providing appropriate RayContributionToHitGroupIndex to TraceRay() function.
+	// Use ShaderSlot argument in SetRayTracingHitGroup() to assign shaders and resources for specific part of the shder binding table record.
+	uint32 ShaderSlotsPerGeometrySegment = 1;
+
+	// Defines how many different callable shaders with unique resource bindings can be bound to this scene.
+	// Shaders and resources are assigned to slots in the scene using SetRayTracingCallableShader().
+	uint32 NumCallableShaderSlots = 0;
+
+	// At least one miss shader must be present in a ray tracing scene.
+	// Default miss shader is always in slot 0. Default shader must not use local resources.
+	// Custom miss shaders can be bound to other slots using SetRayTracingMissShader().
+	uint32 NumMissShaderSlots = 1;
+
+	// Defines whether data in this scene should persist between frames.
+	// Currently only single-frame lifetime is supported.
+	ERayTracingSceneLifetime Lifetime = RTSL_SingleFrame;
 
 	FName DebugName;
 };
@@ -1454,6 +1498,12 @@ public:
 		checkNoEntry();
 		return nullptr;
 	}
+	
+	virtual FRayTracingSceneRHIRef RHICreateRayTracingScene(FRayTracingSceneInitializer2 Initializer)
+	{
+		checkNoEntry();
+		return nullptr;
+	}
 
 	virtual FRayTracingShaderRHIRef RHICreateRayTracingShader(TArrayView<const uint8> Code, const FSHAHash& Hash, EShaderFrequency ShaderFrequency)
 	{
@@ -1761,11 +1811,17 @@ FORCEINLINE FRayTracingGeometryRHIRef RHICreateRayTracingGeometry(const FRayTrac
 	return GDynamicRHI->RHICreateRayTracingGeometry(Initializer);
 }
 
+UE_DEPRECATED(5.0, "This version of RHICreateRayTracingScene is deprecated. FRayTracingSceneInitializer2 should be used instead.")
 FORCEINLINE FRayTracingSceneRHIRef RHICreateRayTracingScene(const FRayTracingSceneInitializer& Initializer)
 {
 	return GDynamicRHI->RHICreateRayTracingScene(Initializer);
 }
 
+FORCEINLINE FRayTracingSceneRHIRef RHICreateRayTracingScene(FRayTracingSceneInitializer2 Initializer)
+{
+	return GDynamicRHI->RHICreateRayTracingScene(MoveTemp(Initializer));
+}
+
 FORCEINLINE FRayTracingShaderRHIRef RHICreateRayTracingShader(TArrayView<const uint8> Code, const FSHAHash& Hash, EShaderFrequency ShaderFrequency)
 {
 	return GDynamicRHI->RHICreateRayTracingShader(Code, Hash, ShaderFrequency);
diff --git a/Engine/Source/Runtime/RHI/Public/RHIResources.h b/Engine/Source/Runtime/RHI/Public/RHIResources.h
--- a/Engine/Source/Runtime/RHI/Public/RHIResources.h
+++ b/Engine/Source/Runtime/RHI/Public/RHIResources.h
@@ -1906,12 +1906,22 @@ public:
 };
 
 /** Bottom level ray tracing acceleration structure (contains triangles). */
-class FRHIRayTracingGeometry : public FRHIRayTracingAccelerationStructure {};
+class FRHIRayTracingGeometry : public FRHIRayTracingAccelerationStructure
+{
+public:
+	virtual uint32 GetNumSegments() const = 0;
+};
 
 typedef TRefCountPtr<FRHIRayTracingGeometry>     FRayTracingGeometryRHIRef;
 
+struct FRayTracingSceneInitializer2;
+
 /** Top level ray tracing acceleration structure (contains instances of meshes). */
-class FRHIRayTracingScene : public FRHIRayTracingAccelerationStructure {};
+class FRHIRayTracingScene : public FRHIRayTracingAccelerationStructure
+{
+public:
+	virtual const FRayTracingSceneInitializer2& GetInitializer() const = 0;
+};
 
 typedef TRefCountPtr<FRHIRayTracingScene>        FRayTracingSceneRHIRef;
 
diff --git a/Engine/Source/Runtime/RHI/Public/RHIValidation.h b/Engine/Source/Runtime/RHI/Public/RHIValidation.h
--- a/Engine/Source/Runtime/RHI/Public/RHIValidation.h
+++ b/Engine/Source/Runtime/RHI/Public/RHIValidation.h
@@ -1736,6 +1736,14 @@ public:
 		return Result;
 	}
 
+	virtual FRayTracingSceneRHIRef RHICreateRayTracingScene(FRayTracingSceneInitializer2 Initializer) override final
+	{
+		FName DebugName = Initializer.DebugName;
+		FRayTracingSceneRHIRef Result = RHI->RHICreateRayTracingScene(MoveTemp(Initializer));
+		Result->InitBarrierTracking(ERHIAccess::BVHWrite, *DebugName.ToString()); // BVHs are always created in BVHWrite state
+		return Result;
+	}
+
 	virtual FRayTracingShaderRHIRef RHICreateRayTracingShader(TArrayView<const uint8> Code, const FSHAHash& Hash, EShaderFrequency ShaderFrequency) override final
 	{
 		return RHI->RHICreateRayTracingShader(Code, Hash, ShaderFrequency);
diff --git a/Engine/Source/Runtime/Renderer/Private/RayTracing/RayTracingScene.cpp b/Engine/Source/Runtime/Renderer/Private/RayTracing/RayTracingScene.cpp
--- a/Engine/Source/Runtime/Renderer/Private/RayTracing/RayTracingScene.cpp
+++ b/Engine/Source/Runtime/Renderer/Private/RayTracing/RayTracingScene.cpp
@@ -7,6 +7,7 @@
 #include "RenderCore.h"
 #include "RayTracingDefinitions.h"
 #include "RenderGraphBuilder.h"
+#include "Experimental/Containers/SherwoodHashTable.h"
 
 FRayTracingScene::FRayTracingScene()
 {
@@ -23,18 +24,58 @@ FGraphEventRef FRayTracingScene::BeginCreate(FRDGBuilder& GraphBuilder)
 	WaitForTasks();
 
 	NumNativeInstances = 0;
+	NumTotalSegments = 0;
 
-	for (const FRayTracingGeometryInstance& InstanceDesc : Instances)
+	const uint32 NumSceneInstances = Instances.Num();
+
+	FRayTracingSceneInitializer2 SceneInitializer;
+	SceneInitializer.DebugName = FName(TEXT("FRayTracingScene"));
+	SceneInitializer.ShaderSlotsPerGeometrySegment = RAY_TRACING_NUM_SHADER_SLOTS;
+	SceneInitializer.NumMissShaderSlots = RAY_TRACING_NUM_MISS_SHADER_SLOTS;
+	SceneInitializer.Instances = Instances;
+	SceneInitializer.PerInstanceGeometries.SetNumUninitialized(NumSceneInstances);
+	SceneInitializer.PerInstanceNumTransforms.SetNumUninitialized(NumSceneInstances);
+	SceneInitializer.BaseInstancePrefixSum.SetNumUninitialized(NumSceneInstances);
+	SceneInitializer.SegmentPrefixSum.SetNumUninitialized(NumSceneInstances);
+
+	Experimental::TSherwoodSet<FRHIRayTracingGeometry*> UniqueGeometries;
+
+	// Compute geometry segment and instance count prefix sums.
+	// These are later used by GetHitRecordBaseIndex() during resource binding
+	// and by GetBaseInstanceIndex() in shaders to emulate SV_InstanceIndex.
+
+	for (uint32 InstanceIndex = 0; InstanceIndex < NumSceneInstances; ++InstanceIndex)
 	{
+		const FRayTracingGeometryInstance& InstanceDesc = Instances[InstanceIndex];
+
 		checkf(InstanceDesc.GPUTransformsSRV || InstanceDesc.NumTransforms <= uint32(InstanceDesc.Transforms.Num()),
 			TEXT("Expected at most %d ray tracing geometry instance transforms, but got %d."),
 			InstanceDesc.NumTransforms, InstanceDesc.Transforms.Num());
 
 		checkf(InstanceDesc.GeometryRHI, TEXT("Ray tracing instance must have a valid geometry."));
 
+		SceneInitializer.PerInstanceGeometries[InstanceIndex] = InstanceDesc.GeometryRHI;
+
+		// Compute geometry segment count prefix sum to be later used in GetHitRecordBaseIndex()
+		SceneInitializer.SegmentPrefixSum[InstanceIndex] = NumTotalSegments;
+		NumTotalSegments += InstanceDesc.GeometryRHI->GetNumSegments();
+
+		bool bIsAlreadyInSet = false;
+		UniqueGeometries.Add(InstanceDesc.GeometryRHI, &bIsAlreadyInSet);
+		if (!bIsAlreadyInSet)
+		{
+			SceneInitializer.ReferencedGeometries.Add(InstanceDesc.GeometryRHI);
+		}
+
+		SceneInitializer.BaseInstancePrefixSum[InstanceIndex] = NumNativeInstances;
 		NumNativeInstances += InstanceDesc.NumTransforms;
+
+		SceneInitializer.PerInstanceNumTransforms[InstanceIndex] = InstanceDesc.NumTransforms;
 	}
 
+	SceneInitializer.NumNativeInstances = NumNativeInstances;
+	SceneInitializer.NumTotalSegments = NumTotalSegments;
+
 	// Round up number of instances to some multiple to avoid pathological growth reallocations.
 	static constexpr uint32 AllocationGranularity = 8 * 1024;
 	uint32 NumNativeInstancesAligned = FMath::DivideAndRoundUp(NumNativeInstances, AllocationGranularity) * AllocationGranularity;
@@ -63,17 +104,11 @@ FGraphEventRef FRayTracingScene::BeginCreate(FRDGBuilder& GraphBuilder)
 
 	CreateRayTracingSceneTask = FFunctionGraphTask::CreateAndDispatchWhenReady([
 			&ResultScene = RayTracingSceneRHI,
-			Instances = MakeArrayView(Instances)]()
+			SceneInitializer = MoveTemp(SceneInitializer)]() mutable
 		{
 			FTaskTagScope TaskTagScope(ETaskTag::EParallelRenderingThread);
 
-			FRayTracingSceneInitializer SceneInitializer;
-			SceneInitializer.DebugName = FName(TEXT("FRayTracingScene"));
-			SceneInitializer.Instances = Instances;
-			SceneInitializer.ShaderSlotsPerGeometrySegment = RAY_TRACING_NUM_SHADER_SLOTS;
-			SceneInitializer.NumMissShaderSlots = RAY_TRACING_NUM_MISS_SHADER_SLOTS;
-
-			ResultScene = RHICreateRayTracingScene(SceneInitializer);
+			ResultScene = RHICreateRayTracingScene(MoveTemp(SceneInitializer));
 
 		}, TStatId(), nullptr, ENamedThreads::AnyThread);
 
diff --git a/Engine/Source/Runtime/Renderer/Private/RayTracing/RayTracingScene.h b/Engine/Source/Runtime/Renderer/Private/RayTracing/RayTracingScene.h
--- a/Engine/Source/Runtime/Renderer/Private/RayTracing/RayTracingScene.h
+++ b/Engine/Source/Runtime/Renderer/Private/RayTracing/RayTracingScene.h
@@ -78,6 +78,8 @@ public:
 	// Total flattened number of ray tracing geometry instances (a single FRayTracingGeometryInstance may represent many).
 	uint32 NumNativeInstances = 0;
 
+	uint32 NumTotalSegments = 0;
+
 	// Geometries which still have a pending build request but are used this frame and require a force build.
 	TArray<const FRayTracingGeometry*> GeometriesToBuild;
 
diff --git a/Engine/Source/Runtime/Renderer/Private/Tests/RayTracingTestbed.cpp b/Engine/Source/Runtime/Renderer/Private/Tests/RayTracingTestbed.cpp
--- a/Engine/Source/Runtime/Renderer/Private/Tests/RayTracingTestbed.cpp
+++ b/Engine/Source/Runtime/Renderer/Private/Tests/RayTracingTestbed.cpp
@@ -114,10 +114,6 @@ bool RunRayTracingTestbed_RenderThread(const FString& Parameters)
 	Instances[0].NumTransforms = NumTransforms;
 	Instances[0].Transforms = MakeArrayView(&FMatrix::Identity, 1);
 
-	FRayTracingSceneInitializer Initializer;
-	Initializer.Instances = Instances;
-	Initializer.ShaderSlotsPerGeometrySegment = RAY_TRACING_NUM_SHADER_SLOTS;
-
 	ERayTracingAccelerationStructureFlags SceneBuildFlags = ERayTracingAccelerationStructureFlags::FastTrace;
 	FRayTracingAccelerationStructureSize SceneSizeInfo = RHICalcRayTracingSceneSize(1, SceneBuildFlags);
 	FRHIResourceCreateInfo SceneBufferCreateInfo(TEXT("RayTracingTestBedSceneBuffer"));
@@ -128,7 +124,23 @@ bool RunRayTracingTestbed_RenderThread(const FString& Parameters)
 
 	FRHICommandListImmediate& RHICmdList = FRHICommandListExecutor::GetImmediateCommandList();
 
-	FRayTracingSceneRHIRef Scene = RHICreateRayTracingScene(Initializer);
+	FRayTracingSceneRHIRef Scene;
+
+	{
+		FRayTracingSceneInitializer2 Initializer;
+		Initializer.DebugName = FName(TEXT("FRayTracingScene"));
+		Initializer.ShaderSlotsPerGeometrySegment = RAY_TRACING_NUM_SHADER_SLOTS;
+		Initializer.Instances = Instances;
+		Initializer.PerInstanceGeometries.Add(Geometry);
+		Initializer.PerInstanceNumTransforms.Add(NumTransforms);
+		Initializer.BaseInstancePrefixSum.Add(0);
+		Initializer.SegmentPrefixSum.Add(0);
+		Initializer.NumNativeInstances = NumInstances * NumTransforms;
+		Initializer.NumTotalSegments = Geometry->GetNumSegments();
+
+		Scene = RHICreateRayTracingScene(MoveTemp(Initializer));
+	}
+
 	RHICmdList.BindAccelerationStructureMemory(Scene, SceneBuffer, 0);
 
 	RHICmdList.BuildAccelerationStructure(Geometry);
diff --git a/Engine/Source/Runtime/VulkanRHI/Private/VulkanRayTracing.h b/Engine/Source/Runtime/VulkanRHI/Private/VulkanRayTracing.h
--- a/Engine/Source/Runtime/VulkanRHI/Private/VulkanRayTracing.h
+++ b/Engine/Source/Runtime/VulkanRHI/Private/VulkanRayTracing.h
@@ -84,6 +84,8 @@ public:
 	FVulkanRayTracingGeometry(const FRayTracingGeometryInitializer& Initializer, const FVulkanDevice* InDevice);
 	~FVulkanRayTracingGeometry();
 
+	virtual uint32 GetNumSegments() const final override { return Segments.Num(); }
+
 	void BuildAccelerationStructure(FVulkanCommandListContext& CommandContext, EAccelerationStructureBuildMode BuildMode);
 	
 	VkDeviceAddress GetAccelerationStructureAddress() const
