Check TypedUAVLoad support for D3D12. Changes from Jessica Chen at Intel.

diff --git a/Engine/Source/Runtime/D3D12RHI/Private/D3D12Device.cpp b/Engine/Source/Runtime/D3D12RHI/Private/D3D12Device.cpp
--- a/Engine/Source/Runtime/D3D12RHI/Private/D3D12Device.cpp
+++ b/Engine/Source/Runtime/D3D12RHI/Private/D3D12Device.cpp
@@ -574,4 +574,9 @@ D3D12_RESOURCE_ALLOCATION_INFO FD3D12Device::GetResourceAllocationInfo(const D3D
 		ResourceAllocationInfoMapMutex.WriteUnlock();
 		return Result;
 	}
+}
+
+bool FD3D12DynamicRHI::RHIIsTypedUAVLoadSupported(EPixelFormat PixelFormat)
+{
+	return EnumHasAnyFlags(GPixelFormats[PixelFormat].Capabilities, EPixelFormatCapabilities::TypedUAVLoad);
 }
\ No newline at end of file
diff --git a/Engine/Source/Runtime/D3D12RHI/Private/D3D12RHIPrivate.h b/Engine/Source/Runtime/D3D12RHI/Private/D3D12RHIPrivate.h
--- a/Engine/Source/Runtime/D3D12RHI/Private/D3D12RHIPrivate.h
+++ b/Engine/Source/Runtime/D3D12RHI/Private/D3D12RHIPrivate.h
@@ -396,6 +396,8 @@ public:
 
 	virtual IRHITransientResourceAllocator* RHICreateTransientResourceAllocator() override;
 
+	virtual bool RHIIsTypedUAVLoadSupported(EPixelFormat PixelFormat) override;
+
 #if WITH_MGPU
 	virtual IRHICommandContextContainer* RHIGetCommandContextContainer(int32 Index, int32 Num, FRHIGPUMask GPUMask)final override;
 #endif
