Adds setting for disabling D3D12 on particular IHVs

diff --git a/Engine/Source/Runtime/D3D12RHI/Private/Windows/WindowsD3D12Device.cpp b/Engine/Source/Runtime/D3D12RHI/Private/Windows/WindowsD3D12Device.cpp
--- a/Engine/Source/Runtime/D3D12RHI/Private/Windows/WindowsD3D12Device.cpp
+++ b/Engine/Source/Runtime/D3D12RHI/Private/Windows/WindowsD3D12Device.cpp
@@ -390,10 +390,29 @@ bool FD3D12DynamicRHIModule::IsSupported()
 		FindAdapter();
 	}
 
+	FString BlockedIHVString;
+	bool bBlockD3D12 = false;
+	if (ChosenAdapters.Num() > 0 && ChosenAdapters[0].IsValid() && GConfig->GetString(TEXT("SystemSettings"), TEXT("RHI.BlockIHVD3D12"), BlockedIHVString, GEngineIni))
+	{
+		TArray<FString> BlockedIHVs;
+		BlockedIHVString.ParseIntoArray(BlockedIHVs, TEXT(","));
+
+		const TCHAR* VendorId = RHIVendorIdToString(EGpuVendorId(ChosenAdapters[0]->GetD3DAdapterDesc().VendorId));
+		for (const FString& BlockedVendor : BlockedIHVs)
+		{
+			if (BlockedVendor.Equals(VendorId, ESearchCase::IgnoreCase))
+			{
+				bBlockD3D12 = true;
+				break;
+			}
+		}
+	}
+
 	// The hardware must support at least 11.0.
 	return ChosenAdapters.Num() > 0
 		&& ChosenAdapters[0]->GetDesc().IsValid()
-		&& ChosenAdapters[0]->GetDesc().MaxSupportedFeatureLevel >= D3D_FEATURE_LEVEL_11_0;
+		&& ChosenAdapters[0]->GetDesc().MaxSupportedFeatureLevel >= D3D_FEATURE_LEVEL_11_0
+		&& !bBlockD3D12;
 }
 
 namespace D3D12RHI
