Integrated 14742935, 14742972, 14781977, 14782581 : Add vertex buffer support to RHIBroadcastTemporalEffect.

diff --git a/Engine/Source/Runtime/D3D12RHI/Private/D3D12CommandContext.h b/Engine/Source/Runtime/D3D12RHI/Private/D3D12CommandContext.h
--- a/Engine/Source/Runtime/D3D12RHI/Private/D3D12CommandContext.h
+++ b/Engine/Source/Runtime/D3D12RHI/Private/D3D12CommandContext.h
@@ -369,6 +369,7 @@ public:
 
 	// This should be called right after the effect generates the resources which will be used in subsequent frame(s).
 	virtual void RHIBroadcastTemporalEffect(const FName& InEffectName, const TArrayView<FRHITexture*> InTextures) final AFR_API_OVERRIDE;
+	virtual void RHIBroadcastTemporalEffect(const FName& InEffectName, const TArrayView<FRHIVertexBuffer*> InBuffers) final AFR_API_OVERRIDE;
 
 #if D3D12_RHI_RAYTRACING
 	virtual void RHICopyBufferRegion(FRHIBuffer* DestBuffer, uint64 DstOffset, FRHIBuffer* SourceBuffer, uint64 SrcOffset, uint64 NumBytes) final override;
@@ -456,6 +457,10 @@ protected:
 	void WriteGPUEventStackToBreadCrumbData(bool bBeginEvent);
 
 private:
+
+	template <typename TD3D12Resource, typename TCopyFunction>
+	void RHIBroadcastTemporalEffect(const FName& InEffectName, const TArrayView<TD3D12Resource*> InResources, const TCopyFunction& InCopyFunction);
+
 	void RHIClearMRT(bool* bClearColorArray, int32 NumClearColors, const FLinearColor* ColorArray, bool bClearDepth, float Depth, bool bClearStencil, uint32 Stencil);
 
 	static void ClearUAV(TRHICommandList_RecursiveHazardous<FD3D12CommandContext>& RHICmdList, FD3D12UnorderedAccessView* UAV, const void* ClearValues, bool bFloat);
@@ -715,6 +720,11 @@ public:
 		ContextRedirect(RHIBroadcastTemporalEffect(InEffectName, InTextures));
 	}
 
+	FORCEINLINE virtual void RHIBroadcastTemporalEffect(const FName& InEffectName, const TArrayView<FRHIVertexBuffer*> InBuffers) final AFR_API_OVERRIDE
+	{
+		ContextRedirect(RHIBroadcastTemporalEffect(InEffectName, InBuffers));
+	}
+
 	virtual void RHIBeginRenderPass(const FRHIRenderPassInfo& InInfo, const TCHAR* InName) final override
 	{
 		ContextRedirect(RHIBeginRenderPass(InInfo, InName));
diff --git a/Engine/Source/Runtime/D3D12RHI/Private/D3D12Commands.cpp b/Engine/Source/Runtime/D3D12RHI/Private/D3D12Commands.cpp
--- a/Engine/Source/Runtime/D3D12RHI/Private/D3D12Commands.cpp
+++ b/Engine/Source/Runtime/D3D12RHI/Private/D3D12Commands.cpp
@@ -2313,7 +2313,8 @@ void FD3D12CommandContext::RHIWaitForTemporalEffect(const FName& InEffectName)
 #endif // WITH_MGPU
 }
 
-void FD3D12CommandContext::RHIBroadcastTemporalEffect(const FName& InEffectName, const TArrayView<FRHITexture*> InTextures)
+template <typename TD3D12Resource, typename TCopyFunction>
+void FD3D12CommandContext::RHIBroadcastTemporalEffect(const FName& InEffectName, const TArrayView<TD3D12Resource*> InResources, const TCopyFunction& InCopyFunction)
 {
 #if WITH_MGPU
 	check(IsDefaultContext());
@@ -2324,14 +2325,7 @@ void FD3D12CommandContext::RHIBroadcastTemporalEffect(const FName& InEffectName,
 	}
 
 	const uint32 GPUIndex = GetGPUIndex();
-	TArray<FD3D12TextureBase*, TInlineAllocator<MAX_NUM_GPUS>> SrcTextures, DstTextures;
-	const int32 NumTextures = InTextures.Num();
-	for (int32 i = 0; i < NumTextures; i++)
-	{
-		SrcTextures.Emplace(RetrieveTextureBase(InTextures[i]));
-		const uint32 NextSiblingGPUIndex = AFRUtils::GetNextSiblingGPUIndex(GPUIndex);
-		DstTextures.Emplace(RetrieveTextureBase(InTextures[i], NextSiblingGPUIndex));
-	}
+	const uint32 NextSiblingGPUIndex = AFRUtils::GetNextSiblingGPUIndex(GPUIndex);
 
 	FD3D12Device* Device = GetParentDevice();
 	FD3D12Adapter* Adapter = Device->GetParentAdapter();
@@ -2339,11 +2333,12 @@ void FD3D12CommandContext::RHIBroadcastTemporalEffect(const FName& InEffectName,
 
 #if USE_COPY_QUEUE_FOR_RESOURCE_SYNC
 
-	for (int32 i = 0; i < NumTextures; i++)
+	for (TD3D12Resource* SrcResource : InResources)
 	{
 		// Resources must be in the COMMON state before using on the copy queue.
-		FD3D12DynamicRHI::TransitionResource(CommandListHandle, SrcTextures[i]->GetResource(), D3D12_RESOURCE_STATE_TBD, D3D12_RESOURCE_STATE_COMMON, D3D12_RESOURCE_BARRIER_ALL_SUBRESOURCES, FD3D12DynamicRHI::ETransitionMode::Apply);
-		FD3D12DynamicRHI::TransitionResource(CommandListHandle, DstTextures[i]->GetResource(), D3D12_RESOURCE_STATE_TBD, D3D12_RESOURCE_STATE_COMMON, D3D12_RESOURCE_BARRIER_ALL_SUBRESOURCES, FD3D12DynamicRHI::ETransitionMode::Apply);
+		TD3D12Resource* DstResource = SrcResource->GetLinkedObject(NextSiblingGPUIndex);
+		FD3D12DynamicRHI::TransitionResource(CommandListHandle, SrcResource->GetResource(), D3D12_RESOURCE_STATE_TBD, D3D12_RESOURCE_STATE_COMMON, D3D12_RESOURCE_BARRIER_ALL_SUBRESOURCES, FD3D12DynamicRHI::ETransitionMode::Apply);
+		FD3D12DynamicRHI::TransitionResource(CommandListHandle, DstResource->GetResource(), D3D12_RESOURCE_STATE_TBD, D3D12_RESOURCE_STATE_COMMON, D3D12_RESOURCE_BARRIER_ALL_SUBRESOURCES, FD3D12DynamicRHI::ETransitionMode::Apply);
 	}
 	CommandListHandle.FlushResourceBarriers();
 
@@ -2360,10 +2355,11 @@ void FD3D12CommandContext::RHIBroadcastTemporalEffect(const FName& InEffectName,
 	FD3D12CommandListHandle hCopyCommandList = CopyCommandListManager.ObtainCommandList(*CopyCommandAllocator);
 	hCopyCommandList.SetCurrentOwningContext(this);
 
-	for (int32 i = 0; i < NumTextures; i++)
+	for (TD3D12Resource* SrcResource : InResources)
 	{
 		// NB: We do not increment numCopies here because the main context isn't doing any work.
-		hCopyCommandList->CopyResource(DstTextures[i]->GetResource()->GetResource(), SrcTextures[i]->GetResource()->GetResource());
+		TD3D12Resource* DstResource = SrcResource->GetLinkedObject(NextSiblingGPUIndex);
+		InCopyFunction(hCopyCommandList, DstResource, SrcResource);
 	}
 	hCopyCommandList.Close();
 
@@ -2376,17 +2372,19 @@ void FD3D12CommandContext::RHIBroadcastTemporalEffect(const FName& InEffectName,
 
 #else
 
-	for (int32 i = 0; i < NumTextures; i++)
+	for (TD3D12Resource* SrcResource : InResources)
 	{
-		FD3D12DynamicRHI::TransitionResource(CommandListHandle, SrcTextures[i]->GetResource(), D3D12_RESOURCE_STATE_TBD, D3D12_RESOURCE_STATE_COPY_SOURCE, D3D12_RESOURCE_BARRIER_ALL_SUBRESOURCES);
-		FD3D12DynamicRHI::TransitionResource(CommandListHandle, DstTextures[i]->GetResource(), D3D12_RESOURCE_STATE_TBD, D3D12_RESOURCE_STATE_COPY_DEST, D3D12_RESOURCE_BARRIER_ALL_SUBRESOURCES);
+		TD3D12Resource* DstResource = SrcResource->GetLinkedObject(NextSiblingGPUIndex);;
+		FD3D12DynamicRHI::TransitionResource(CommandListHandle, SrcResource->GetResource(), D3D12_RESOURCE_STATE_TBD, D3D12_RESOURCE_STATE_COPY_SOURCE, D3D12_RESOURCE_BARRIER_ALL_SUBRESOURCES);
+		FD3D12DynamicRHI::TransitionResource(CommandListHandle, DstResource->GetResource(), D3D12_RESOURCE_STATE_TBD, D3D12_RESOURCE_STATE_COPY_DEST, D3D12_RESOURCE_BARRIER_ALL_SUBRESOURCES);
 	}
 	CommandListHandle.FlushResourceBarriers();
 
-	for (int32 i = 0; i < NumTextures; i++)
+	for (TD3D12Resource* SrcResource : InResources)
 	{
 		numCopies++;
-		CommandListHandle->CopyResource(DstTextures[i]->GetResource()->GetResource(), SrcTextures[i]->GetResource()->GetResource());
+		TD3D12Resource* DstResource = SrcResource->GetLinkedObject(NextSiblingGPUIndex);;
+		InCopyFunction(CommandListHandle, DstResource, SrcResource);
 	}
 
 	FlushCommands();
@@ -2395,4 +2393,46 @@ void FD3D12CommandContext::RHIBroadcastTemporalEffect(const FName& InEffectName,
 
 #endif // USE_COPY_QUEUE_FOR_RESOURCE_SYNC
 #endif // WITH_MGPU
-}
\ No newline at end of file
+}
+void FD3D12CommandContext::RHIBroadcastTemporalEffect(const FName& InEffectName, const TArrayView<FRHITexture*> InTextures)
+{
+#if WITH_MGPU
+	FMemMark Mark(FMemStack::Get());
+	TArray<FD3D12TextureBase*, TMemStackAllocator<>> Resources;
+	Resources.Reserve(InTextures.Num());
+	for (FRHITexture* Texture : InTextures)
+	{
+		Resources.Emplace(RetrieveTextureBase(Texture));
+	}
+	auto CopyFunction = [](FD3D12CommandListHandle& CommandList, FD3D12TextureBase* Dst, FD3D12TextureBase* Src)
+	{
+		CommandList->CopyResource(Dst->GetResource()->GetResource(), Src->GetResource()->GetResource());
+	};
+	RHIBroadcastTemporalEffect(InEffectName, MakeArrayView(Resources), CopyFunction);
+#endif // WITH_MGPU
+}
+
+void FD3D12CommandContext::RHIBroadcastTemporalEffect(const FName& InEffectName, const TArrayView<FRHIVertexBuffer*> InBuffers)
+{
+#if WITH_MGPU
+	FMemMark Mark(FMemStack::Get());
+	TArray<FD3D12Buffer*, TMemStackAllocator<>> Resources;
+	Resources.Reserve(InBuffers.Num());
+	for (auto* Buffer : InBuffers)
+	{
+		Resources.Emplace(RetrieveObject<FD3D12Buffer>(Buffer));
+	}
+	auto CopyFunction = [](FD3D12CommandListHandle& CommandList, FD3D12Buffer* Dst, FD3D12Buffer* Src)
+	{
+		CommandList->CopyBufferRegion(
+			Dst->GetResource()->GetResource(),
+			Dst->ResourceLocation.GetOffsetFromBaseOfResource(),
+			Src->GetResource()->GetResource(),
+			Src->ResourceLocation.GetOffsetFromBaseOfResource(),
+			Src->ResourceLocation.GetSize()
+		);
+	};
+	RHIBroadcastTemporalEffect(InEffectName, MakeArrayView(Resources), CopyFunction);
+#endif // WITH_MGPU
+}
+
diff --git a/Engine/Source/Runtime/RHI/Public/RHICommandList.h b/Engine/Source/Runtime/RHI/Public/RHICommandList.h
--- a/Engine/Source/Runtime/RHI/Public/RHICommandList.h
+++ b/Engine/Source/Runtime/RHI/Public/RHICommandList.h
@@ -880,15 +880,18 @@ FRHICOMMAND_MACRO(FRHICommandWaitForTemporalEffect)
 	RHI_API void Execute(FRHICommandListBase& CmdList);
 };
 
-FRHICOMMAND_MACRO(FRHICommandBroadcastTemporalEffect)
+struct FRHICommandBroadcastTemporalEffectString
+{
+	static const TCHAR* TStr() { return TEXT("FRHICommandBroadcastTemporalEffect"); }
+};
+template <typename TRHIResource>
+struct FRHICommandBroadcastTemporalEffect final	: public FRHICommand<FRHICommandBroadcastTemporalEffect<TRHIResource>, FRHICommandBroadcastTemporalEffectString>
 {
 	FName EffectName;
-	FRHITexture** Textures;
-	int32 NumTextures;
-	FORCEINLINE_DEBUGGABLE FRHICommandBroadcastTemporalEffect(const FName& InEffectName, FRHITexture** InTextures, int32 InNumTextures)
+	const TArrayView<TRHIResource*> Resources;
+	FORCEINLINE_DEBUGGABLE FRHICommandBroadcastTemporalEffect(const FName& InEffectName, const TArrayView<TRHIResource*> InResources)
 		: EffectName(InEffectName)
-		, Textures(InTextures)
-		, NumTextures(InNumTextures)
+		, Resources(InResources)
 	{
 	}
 	RHI_API void Execute(FRHICommandListBase& CmdList);
@@ -3062,15 +3065,19 @@ public:
 			return;
 		}
 
-		// Allocate space to hold the list of textures inline in the command list itself.
-		const int32 NumTextures = Textures.Num();
-		FRHITexture** InlineTextureArray = (FRHITexture**)Alloc(sizeof(FRHITexture*) * NumTextures, alignof(FRHITexture*));
-		for (int32 Index = 0; Index < NumTextures; ++Index)
+		ALLOC_COMMAND(FRHICommandBroadcastTemporalEffect<FRHITexture>)(EffectName, AllocArray(Textures));
+	}
+
+	FORCEINLINE_DEBUGGABLE void BroadcastTemporalEffect(const FName& EffectName, const TArrayView<FRHIVertexBuffer*> Buffers)
+	{
+		//check(IsOutsideRenderPass());
+		if (Bypass())
 		{
-			InlineTextureArray[Index] = Textures[Index];
+			GetContext().RHIBroadcastTemporalEffect(EffectName, Buffers);
+			return;
 		}
 
-		ALLOC_COMMAND(FRHICommandBroadcastTemporalEffect)(EffectName, InlineTextureArray, NumTextures);
+		ALLOC_COMMAND(FRHICommandBroadcastTemporalEffect<FRHIVertexBuffer>)(EffectName, AllocArray(Buffers));
 	}
 #endif // WITH_MGPU
 
diff --git a/Engine/Source/Runtime/RHI/Public/RHICommandListCommandExecutes.inl b/Engine/Source/Runtime/RHI/Public/RHICommandListCommandExecutes.inl
--- a/Engine/Source/Runtime/RHI/Public/RHICommandListCommandExecutes.inl
+++ b/Engine/Source/Runtime/RHI/Public/RHICommandListCommandExecutes.inl
@@ -98,10 +98,15 @@ void FRHICommandWaitForTemporalEffect::Execute(FRHICommandListBase& CmdList)
 	INTERNAL_DECORATOR(RHIWaitForTemporalEffect)(EffectName);
 }
 
-void FRHICommandBroadcastTemporalEffect::Execute(FRHICommandListBase& CmdList)
+template <> void FRHICommandBroadcastTemporalEffect<FRHITexture>::Execute(FRHICommandListBase& CmdList)
 {
 	RHISTAT(BroadcastTemporalEffect);
-	INTERNAL_DECORATOR(RHIBroadcastTemporalEffect)(EffectName, { Textures, NumTextures });
+	INTERNAL_DECORATOR(RHIBroadcastTemporalEffect)(EffectName, Resources);
+}
+template <> void FRHICommandBroadcastTemporalEffect<FRHIVertexBuffer>::Execute(FRHICommandListBase& CmdList)
+{
+	RHISTAT(BroadcastTemporalEffect);
+	INTERNAL_DECORATOR(RHIBroadcastTemporalEffect)(EffectName, Resources);
 }
 
 void FRHICommandTransferTextures::Execute(FRHICommandListBase& CmdList)
diff --git a/Engine/Source/Runtime/RHI/Public/RHIContext.h b/Engine/Source/Runtime/RHI/Public/RHIContext.h
--- a/Engine/Source/Runtime/RHI/Public/RHIContext.h
+++ b/Engine/Source/Runtime/RHI/Public/RHIContext.h
@@ -308,6 +308,11 @@ public:
 	{
 		/* empty default implementation */
 	}
+
+	virtual void RHIBroadcastTemporalEffect(const FName& InEffectName, const TArrayView<FRHIVertexBuffer*> InBuffers)
+	{
+		/* empty default implementation */
+	}
 #endif // WITH_MGPU
 
 	/**
