fixing FRHIVertexBuffer deprecation warning

diff --git a/Engine/Source/Runtime/D3D12RHI/Private/D3D12CommandContext.h b/Engine/Source/Runtime/D3D12RHI/Private/D3D12CommandContext.h
--- a/Engine/Source/Runtime/D3D12RHI/Private/D3D12CommandContext.h
+++ b/Engine/Source/Runtime/D3D12RHI/Private/D3D12CommandContext.h
@@ -369,7 +369,7 @@ public:
 
 	// This should be called right after the effect generates the resources which will be used in subsequent frame(s).
 	virtual void RHIBroadcastTemporalEffect(const FName& InEffectName, const TArrayView<FRHITexture*> InTextures) final AFR_API_OVERRIDE;
-	virtual void RHIBroadcastTemporalEffect(const FName& InEffectName, const TArrayView<FRHIVertexBuffer*> InBuffers) final AFR_API_OVERRIDE;
+	virtual void RHIBroadcastTemporalEffect(const FName& InEffectName, const TArrayView<FRHIBuffer*> InBuffers) final AFR_API_OVERRIDE;
 
 #if D3D12_RHI_RAYTRACING
 	virtual void RHICopyBufferRegion(FRHIBuffer* DestBuffer, uint64 DstOffset, FRHIBuffer* SourceBuffer, uint64 SrcOffset, uint64 NumBytes) final override;
@@ -720,7 +720,7 @@ public:
 		ContextRedirect(RHIBroadcastTemporalEffect(InEffectName, InTextures));
 	}
 
-	FORCEINLINE virtual void RHIBroadcastTemporalEffect(const FName& InEffectName, const TArrayView<FRHIVertexBuffer*> InBuffers) final AFR_API_OVERRIDE
+	FORCEINLINE virtual void RHIBroadcastTemporalEffect(const FName& InEffectName, const TArrayView<FRHIBuffer*> InBuffers) final AFR_API_OVERRIDE
 	{
 		ContextRedirect(RHIBroadcastTemporalEffect(InEffectName, InBuffers));
 	}
