More bring-up work on bindless in D3D12RHI and initial public RHI support:
- Bindless rendering is not enabled anywhere yet.
- Adding bSupportsBindless to DDSPI.
- Moving ED3D12DescriptorHeapType into RHI and renaming to ERHIDescriptorHeapType.
- Adding FRHIDescriptorHandle for public usage of bindless indices.
- Adding GetBindlessHandle() to RHI SamplerState, SRV, UAV.
- Adding GetDefaultShaderResourceView() to RHI Textures.
- Adding bindless Sampler heap support to D3D12.
- Cleaned up some common shader code in prep for future bindless work.

diff --git a/Engine/Source/Developer/ShaderCompilerCommon/Private/ShaderCompilerCommon.cpp b/Engine/Source/Developer/ShaderCompilerCommon/Private/ShaderCompilerCommon.cpp
--- a/Engine/Source/Developer/ShaderCompilerCommon/Private/ShaderCompilerCommon.cpp
+++ b/Engine/Source/Developer/ShaderCompilerCommon/Private/ShaderCompilerCommon.cpp
@@ -532,6 +532,12 @@ void AddUnboundShaderParameterError(
 	AddNoteToDisplayShaderParameterStructureOnCppSide(CompilerInput.RootParametersStructure, CompilerOutput);
 }
 
+inline bool MemberWasPotentiallyMoved(const FShaderParametersMetadata::FMember& InMember)
+{
+	const EUniformBufferBaseType BaseType = InMember.GetBaseType();
+	return BaseType == UBMT_INT32 || BaseType == UBMT_UINT32 || BaseType == UBMT_FLOAT32;
+}
+
 bool FShaderParameterParser::ParseAndMoveShaderParametersToRootConstantBuffer(
 	const FShaderCompilerInput& CompilerInput,
 	FShaderCompilerOutput& CompilerOutput,
@@ -1003,56 +1009,43 @@ bool FShaderParameterParser::ParseAndMoveShaderParametersToRootConstantBuffer(
 				const TCHAR* ShaderBindingName,
 				uint16 ByteOffset)
 		{
-			if (
-				Member.GetBaseType() != UBMT_INT32 &&
-				Member.GetBaseType() != UBMT_UINT32 &&
-				Member.GetBaseType() != UBMT_FLOAT32)
-			{
-				return;
-			}
-
-			FParsedShaderParameter& ParsedParameter = ParsedParameters[ShaderBindingName];
-
-			if (ParsedParameter.IsFound())
+			if (MemberWasPotentiallyMoved(Member))
 			{
-				FString HLSLOffset;
+				FParsedShaderParameter* ParsedParameter = ParsedParameters.Find(ShaderBindingName);
+				if (ParsedParameter && ParsedParameter->IsFound())
 				{
-					HLSLOffset = FString::FromInt(ByteOffset / 16);
+					uint32 ConstantRegister = ByteOffset / 16;
+					const TCHAR* ConstantSwizzle = [ByteOffset]()
+						{
+							switch (ByteOffset % 16)
+							{
+							default: unimplemented();
+							case 0:  return TEXT("");
+							case 4:  return TEXT(".y");
+							case 8:  return TEXT(".z");
+							case 12: return TEXT(".w");
+							}
+						}();
 
-					switch (ByteOffset % 16)
+					if (!ParsedParameter->ParsedArraySize.IsEmpty())
 					{
-					case 0:
-						break;
-					case 4:
-						HLSLOffset.Append(TEXT(".y"));
-						break;
-					case 8:
-						HLSLOffset.Append(TEXT(".z"));
-						break;
-					case 12:
-						HLSLOffset.Append(TEXT(".w"));
-						break;
-					default:
-						unimplemented();
+						RootCBufferContent.Append(FString::Printf(
+							TEXT("%s %s[%s] : packoffset(c%d%s);\r\n"),
+							*ParsedParameter->ParsedType,
+							ShaderBindingName,
+							*ParsedParameter->ParsedArraySize,
+							ConstantRegister,
+							ConstantSwizzle));
+					}
+					else
+					{
+						RootCBufferContent.Append(FString::Printf(
+							TEXT("%s %s : packoffset(c%d%s);\r\n"),
+							*ParsedParameter->ParsedType,
+							ShaderBindingName,
+							ConstantRegister,
+							ConstantSwizzle));
 					}
-				}
-
-				if (!ParsedParameter.ParsedArraySize.IsEmpty())
-				{
-					RootCBufferContent.Append(FString::Printf(
-						TEXT("%s %s[%s] : packoffset(c%s);\r\n"),
-						*ParsedParameter.ParsedType,
-						ShaderBindingName,
-						*ParsedParameter.ParsedArraySize,
-						*HLSLOffset));
-				}
-				else
-				{
-					RootCBufferContent.Append(FString::Printf(
-						TEXT("%s %s : packoffset(c%s);\r\n"),
-						*ParsedParameter.ParsedType,
-						ShaderBindingName,
-						*HLSLOffset));
 				}
 			}
 		});
diff --git a/Engine/Source/Runtime/D3D12RHI/Private/D3D12Commands.cpp b/Engine/Source/Runtime/D3D12RHI/Private/D3D12Commands.cpp
--- a/Engine/Source/Runtime/D3D12RHI/Private/D3D12Commands.cpp
+++ b/Engine/Source/Runtime/D3D12RHI/Private/D3D12Commands.cpp
@@ -1068,7 +1068,7 @@ void FD3D12CommandContext::RHISetShaderUniformBuffer(FRHIComputeShader* ComputeS
 	DirtyUniformBuffers[SF_Compute] |= (1 << BufferIndex);
 }
 
-void FD3D12CommandContext::RHISetShaderParameter(FRHIGraphicsShader* ShaderRHI, uint32 BufferIndex, uint32 BaseIndex, uint32 NumBytes, const void* NewValue)
+void FD3D12CommandContext::RHISetShaderParameter(FRHIGraphicsShader* ShaderRHI, uint32 BufferIndex, uint32 Offset, uint32 NumBytes, const void* NewValue)
 {
 	checkSlow(BufferIndex == 0);
 
@@ -1078,7 +1078,7 @@ void FD3D12CommandContext::RHISetShaderParameter(FRHIGraphicsShader* ShaderRHI,
 	{
 		FRHIVertexShader* VertexShaderRHI = static_cast<FRHIVertexShader*>(ShaderRHI);
 		ValidateBoundShader(StateCache, VertexShaderRHI);
-		VSConstantBuffer.UpdateConstant((const uint8*)NewValue, BaseIndex, NumBytes);
+		VSConstantBuffer.UpdateConstant((const uint8*)NewValue, Offset, NumBytes);
 	}
 	break;
 #if PLATFORM_SUPPORTS_MESH_SHADERS
@@ -1086,14 +1086,14 @@ void FD3D12CommandContext::RHISetShaderParameter(FRHIGraphicsShader* ShaderRHI,
 	{
 		FRHIMeshShader* MeshShaderRHI = static_cast<FRHIMeshShader*>(ShaderRHI);
 		ValidateBoundShader(StateCache, MeshShaderRHI);
-		MSConstantBuffer.UpdateConstant((const uint8*)NewValue, BaseIndex, NumBytes);
+		MSConstantBuffer.UpdateConstant((const uint8*)NewValue, Offset, NumBytes);
 	}
 	break;
 	case SF_Amplification:
 	{
 		FRHIAmplificationShader* AmplificationShaderRHI = static_cast<FRHIAmplificationShader*>(ShaderRHI);
 		ValidateBoundShader(StateCache, AmplificationShaderRHI);
-		ASConstantBuffer.UpdateConstant((const uint8*)NewValue, BaseIndex, NumBytes);
+		ASConstantBuffer.UpdateConstant((const uint8*)NewValue, Offset, NumBytes);
 	}
 	break;
 #endif
@@ -1101,14 +1101,14 @@ void FD3D12CommandContext::RHISetShaderParameter(FRHIGraphicsShader* ShaderRHI,
 	{
 		FRHIGeometryShader* GeometryShaderRHI = static_cast<FRHIGeometryShader*>(ShaderRHI);
 		ValidateBoundShader(StateCache, GeometryShaderRHI);
-		GSConstantBuffer.UpdateConstant((const uint8*)NewValue, BaseIndex, NumBytes);
+		GSConstantBuffer.UpdateConstant((const uint8*)NewValue, Offset, NumBytes);
 	}
 	break;
 	case SF_Pixel:
 	{
 		FRHIPixelShader* PixelShaderRHI = static_cast<FRHIPixelShader*>(ShaderRHI);
 		ValidateBoundShader(StateCache, PixelShaderRHI);
-		PSConstantBuffer.UpdateConstant((const uint8*)NewValue, BaseIndex, NumBytes);
+		PSConstantBuffer.UpdateConstant((const uint8*)NewValue, Offset, NumBytes);
 	}
 	break;
 	default:
@@ -1116,11 +1116,11 @@ void FD3D12CommandContext::RHISetShaderParameter(FRHIGraphicsShader* ShaderRHI,
 	}
 }
 
-void FD3D12CommandContext::RHISetShaderParameter(FRHIComputeShader* ComputeShaderRHI, uint32 BufferIndex, uint32 BaseIndex, uint32 NumBytes, const void* NewValue)
+void FD3D12CommandContext::RHISetShaderParameter(FRHIComputeShader* ComputeShaderRHI, uint32 BufferIndex, uint32 Offset, uint32 NumBytes, const void* NewValue)
 {
 	//ValidateBoundShader(StateCache, ComputeShaderRHI);
 	checkSlow(BufferIndex == 0);
-	CSConstantBuffer.UpdateConstant((const uint8*)NewValue, BaseIndex, NumBytes);
+	CSConstantBuffer.UpdateConstant((const uint8*)NewValue, Offset, NumBytes);
 }
 
 void FD3D12CommandContext::ValidateExclusiveDepthStencilAccess(FExclusiveDepthStencil RequestedAccess) const
@@ -2201,7 +2201,7 @@ void FD3D12CommandContext::RHIClearMRTImpl(bool* bClearColorArray, int32 NumClea
 				if (RTView != nullptr && bClearColorArray[TargetIndex])
 				{
 					numClears++;
-					CommandListHandle->ClearRenderTargetView(RTView->GetView(), (float*)&ClearColorArray[TargetIndex], ClearRectCount, pClearRects);
+					CommandListHandle->ClearRenderTargetView(RTView->GetOfflineCpuHandle(), (float*)&ClearColorArray[TargetIndex], ClearRectCount, pClearRects);
 					CommandListHandle.UpdateResidency(RTView->GetResource());
 				}
 			}
@@ -2210,7 +2210,7 @@ void FD3D12CommandContext::RHIClearMRTImpl(bool* bClearColorArray, int32 NumClea
 		if (ClearDSV)
 		{
 			numClears++;
-			CommandListHandle->ClearDepthStencilView(DepthStencilView->GetView(), (D3D12_CLEAR_FLAGS)ClearFlags, Depth, Stencil, ClearRectCount, pClearRects);
+			CommandListHandle->ClearDepthStencilView(DepthStencilView->GetOfflineCpuHandle(), (D3D12_CLEAR_FLAGS)ClearFlags, Depth, Stencil, ClearRectCount, pClearRects);
 			CommandListHandle.UpdateResidency(DepthStencilView->GetResource());
 		}
 
diff --git a/Engine/Source/Runtime/D3D12RHI/Private/D3D12DescriptorCache.cpp b/Engine/Source/Runtime/D3D12RHI/Private/D3D12DescriptorCache.cpp
--- a/Engine/Source/Runtime/D3D12RHI/Private/D3D12DescriptorCache.cpp
+++ b/Engine/Source/Runtime/D3D12RHI/Private/D3D12DescriptorCache.cpp
@@ -39,15 +39,15 @@ template void FD3D12DescriptorCache::SetSamplers<SF_Geometry>(const FD3D12RootSi
 template void FD3D12DescriptorCache::SetSamplers<SF_Pixel>(const FD3D12RootSignature* RootSignature, FD3D12SamplerStateCache& Cache, const SamplerSlotMask& SlotsNeededMask, uint32 SlotsNeeded, uint32& HeapSlot);
 template void FD3D12DescriptorCache::SetSamplers<SF_Compute>(const FD3D12RootSignature* RootSignature, FD3D12SamplerStateCache& Cache, const SamplerSlotMask& SlotsNeededMask, uint32 SlotsNeeded, uint32& HeapSlot);
 
-bool FD3D12DescriptorCache::HeapRolledOver(ED3D12DescriptorHeapType InHeapType)
+bool FD3D12DescriptorCache::HeapRolledOver(ERHIDescriptorHeapType InHeapType)
 {
 	// A heap rolled over, so set the descriptor heaps again and return if the heaps actually changed.
 	return SetDescriptorHeaps();
 }
 
-void FD3D12DescriptorCache::HeapLoopedAround(ED3D12DescriptorHeapType InHeapType)
+void FD3D12DescriptorCache::HeapLoopedAround(ERHIDescriptorHeapType InHeapType)
 {
-	if (InHeapType == ED3D12DescriptorHeapType::Sampler)
+	if (InHeapType == ERHIDescriptorHeapType::Sampler)
 	{
 		SamplerMap.Reset();
 	}
@@ -81,7 +81,7 @@ void FD3D12DescriptorCache::Init(FD3D12Device* InParent, FD3D12CommandContext* I
 	// Always Init a local sampler heap as the high level cache will always miss initialy
 	// so we need something to fall back on (The view heap never rolls over so we init that one
 	// lazily as a backup to save memory)
-	LocalSamplerHeap.Init(InNumSamplerDescriptors, ED3D12DescriptorHeapType::Sampler);
+	LocalSamplerHeap.Init(InNumSamplerDescriptors, ERHIDescriptorHeapType::Sampler);
 
 	NumLocalViewDescriptors = InNumLocalViewDescriptors;
 
@@ -97,21 +97,21 @@ void FD3D12DescriptorCache::Init(FD3D12Device* InParent, FD3D12CommandContext* I
 	SRVDesc.Texture2D.MipLevels = 1;
 	SRVDesc.Texture2D.MostDetailedMip = 0;
 	SRVDesc.Texture2D.ResourceMinLODClamp = 0.0f;
-	NullSRV = new FD3D12ViewDescriptorHandle(GetParentDevice(), ED3D12DescriptorHeapType::Standard);
+	NullSRV = new FD3D12ViewDescriptorHandle(GetParentDevice(), ERHIDescriptorHeapType::Standard);
 	NullSRV->CreateView(SRVDesc, nullptr);
 
 	D3D12_RENDER_TARGET_VIEW_DESC RTVDesc = {};
 	RTVDesc.ViewDimension = D3D12_RTV_DIMENSION_TEXTURE2D;
 	RTVDesc.Format = DXGI_FORMAT_B8G8R8A8_UNORM;
 	RTVDesc.Texture2D.MipSlice = 0;
-	NullRTV = new FD3D12ViewDescriptorHandle(GetParentDevice(), ED3D12DescriptorHeapType::RenderTarget);
+	NullRTV = new FD3D12ViewDescriptorHandle(GetParentDevice(), ERHIDescriptorHeapType::RenderTarget);
 	NullRTV->CreateView(RTVDesc, nullptr);
 
 	D3D12_UNORDERED_ACCESS_VIEW_DESC UAVDesc = {};
 	UAVDesc.ViewDimension = D3D12_UAV_DIMENSION_TEXTURE2D;
 	UAVDesc.Format = DXGI_FORMAT_R8G8B8A8_UNORM;
 	UAVDesc.Texture2D.MipSlice = 0;
-	NullUAV = new FD3D12ViewDescriptorHandle(GetParentDevice(), ED3D12DescriptorHeapType::Standard);
+	NullUAV = new FD3D12ViewDescriptorHandle(GetParentDevice(), ERHIDescriptorHeapType::Standard);
 	NullUAV->CreateView(UAVDesc, nullptr, nullptr);
 
 #if USE_STATIC_ROOT_SIGNATURE
@@ -315,11 +315,11 @@ void FD3D12DescriptorCache::SetUAVs(const FD3D12RootSignature* RootSignature, FD
 	{
 		if ((SlotIndex < UAVStartSlot) || (UAVs[SlotIndex] == nullptr))
 		{
-			SrcDescriptors[SlotIndex] = NullUAV->GetOfflineHandle();
+			SrcDescriptors[SlotIndex] = NullUAV->GetOfflineCpuHandle();
 		}
 		else
 		{
-			SrcDescriptors[SlotIndex] = UAVs[SlotIndex]->GetView();
+			SrcDescriptors[SlotIndex] = UAVs[SlotIndex]->GetOfflineCpuHandle();
 
 			FD3D12DynamicRHI::TransitionResource(CommandList, UAVs[SlotIndex], D3D12_RESOURCE_STATE_UNORDERED_ACCESS, FD3D12DynamicRHI::ETransitionMode::Validate);
 			CommandList.UpdateResidency(Cache.ResidencyHandles[ShaderStage][SlotIndex]);
@@ -389,18 +389,18 @@ void FD3D12DescriptorCache::SetRenderTargets(FD3D12RenderTargetView** RenderTarg
 	{
 		if (RenderTargetViewArray[i] != NULL)
 		{
-			RTVDescriptors[i] = RenderTargetViewArray[i]->GetView();
+			RTVDescriptors[i] = RenderTargetViewArray[i]->GetOfflineCpuHandle();
 			CommandList.UpdateResidency(RenderTargetViewArray[i]->GetResource());
 		}
 		else
 		{
-			RTVDescriptors[i] = NullRTV->GetOfflineHandle();
+			RTVDescriptors[i] = NullRTV->GetOfflineCpuHandle();
 		}
 	}
 
 	if (DepthStencilTarget != nullptr)
 	{
-		const D3D12_CPU_DESCRIPTOR_HANDLE DSVDescriptor = DepthStencilTarget->GetView();
+		const D3D12_CPU_DESCRIPTOR_HANDLE DSVDescriptor = DepthStencilTarget->GetOfflineCpuHandle();
 		CommandList->OMSetRenderTargets(Count, RTVDescriptors, 0, &DSVDescriptor);
 		CommandList.UpdateResidency(DepthStencilTarget->GetResource());
 	}
@@ -479,11 +479,11 @@ void FD3D12DescriptorCache::SetSamplers(const FD3D12RootSignature* RootSignature
 		{
 			if (Samplers[SlotIndex] != nullptr)
 			{
-				SrcDescriptors[SlotIndex] = Samplers[SlotIndex]->Descriptor;
+				SrcDescriptors[SlotIndex] = Samplers[SlotIndex]->OfflineHandle;
 			}
 			else
 			{
-				SrcDescriptors[SlotIndex] = DefaultSampler->Descriptor;
+				SrcDescriptors[SlotIndex] = DefaultSampler->OfflineHandle;
 			}
 		}
 		FD3D12SamplerStateCache::CleanSlots(CurrentDirtySlotMask, SlotsNeeded);
@@ -553,7 +553,7 @@ void FD3D12DescriptorCache::SetSRVs(const FD3D12RootSignature* RootSignature, FD
 	{
 		if (SRVs[SlotIndex] != nullptr)
 		{
-			SrcDescriptors[SlotIndex] = SRVs[SlotIndex]->GetView();
+			SrcDescriptors[SlotIndex] = SRVs[SlotIndex]->GetOfflineCpuHandle();
 
 			D3D12_RESOURCE_STATES State = (ShaderStage == SF_Compute) ? D3D12_RESOURCE_STATE_NON_PIXEL_SHADER_RESOURCE : D3D12_RESOURCE_STATE_NON_PIXEL_SHADER_RESOURCE | D3D12_RESOURCE_STATE_PIXEL_SHADER_RESOURCE;
 			if (SRVs[SlotIndex]->IsDepthStencilResource())
@@ -573,7 +573,7 @@ void FD3D12DescriptorCache::SetSRVs(const FD3D12RootSignature* RootSignature, FD
 		}
 		else
 		{
-			SrcDescriptors[SlotIndex] = NullSRV->GetOfflineHandle();
+			SrcDescriptors[SlotIndex] = NullSRV->GetOfflineCpuHandle();
 		}
 		check(SrcDescriptors[SlotIndex].ptr != 0);
 	}
@@ -653,7 +653,7 @@ void FD3D12DescriptorCache::SetConstantBuffers(const FD3D12RootSignature* RootSi
 		}
 		else
 		{
-			Device->CopyDescriptorsSimple(1, DestDescriptor, NullCBV->GetView(), D3D12_DESCRIPTOR_HEAP_TYPE_CBV_SRV_UAV);
+			Device->CopyDescriptorsSimple(1, DestDescriptor, NullCBV->GetOfflineCpuHandle(), D3D12_DESCRIPTOR_HEAP_TYPE_CBV_SRV_UAV);
 		}
 
 		DestDescriptorSlot++;
@@ -736,7 +736,7 @@ bool FD3D12DescriptorCache::SwitchToContextLocalViewHeap(const FD3D12CommandList
 		if (LocalViewHeap)
 		{
 			check(NumLocalViewDescriptors);
-			LocalViewHeap->Init(NumLocalViewDescriptors, ED3D12DescriptorHeapType::Standard);
+			LocalViewHeap->Init(NumLocalViewDescriptors, ERHIDescriptorHeapType::Standard);
 		}
 		else
 		{
@@ -853,7 +853,7 @@ Reserve requested amount of descriptor slots - should fit, user has to check wit
 **/
 uint32 FD3D12OnlineHeap::ReserveSlots(uint32 NumSlotsRequested)
 {
-	const ED3D12DescriptorHeapType HeapType = Heap->GetType();
+	const ERHIDescriptorHeapType HeapType = Heap->GetType();
 
 #ifdef VERBOSE_DESCRIPTOR_HEAP_DEBUG
 	FMsg::Logf(__FILE__, __LINE__, TEXT("DescriptorCache"), ELogVerbosity::Log, TEXT("Requesting reservation [TYPE %s] with %d slots, required fence is %d"),
@@ -894,7 +894,7 @@ uint32 FD3D12OnlineHeap::ReserveSlots(uint32 NumSlotsRequested)
 	// Note where to start looking next time
 	NextSlotIndex = SlotAfterReservation;
 
-	if (HeapType == ED3D12DescriptorHeapType::Standard)
+	if (HeapType == ERHIDescriptorHeapType::Standard)
 	{
 		INC_DWORD_STAT_BY(STAT_NumReservedViewOnlineDescriptors, NumSlotsRequested);
 	}
@@ -916,7 +916,7 @@ void FD3D12OnlineHeap::SetNextSlot(uint32 NextSlot)
 	// This is used to correct for the actual number of heap slots used
 	check(NextSlot <= NextSlotIndex);
 
-	check(Heap->GetType() != ED3D12DescriptorHeapType::Standard);
+	check(Heap->GetType() != ERHIDescriptorHeapType::Standard);
 	DEC_DWORD_STAT_BY(STAT_NumReservedSamplerOnlineDescriptors, NextSlotIndex - NextSlot);
 
 	NextSlotIndex = NextSlot;
@@ -939,7 +939,7 @@ void FD3D12GlobalOnlineSamplerHeap::Init(uint32 TotalSize)
 {
 	Heap = GetParentDevice()->GetDescriptorHeapManager().AllocateHeap(
 		TEXT("Device Global - Online Sampler Heap"),
-		ED3D12DescriptorHeapType::Sampler,
+		ERHIDescriptorHeapType::Sampler,
 		TotalSize,
 		D3D12_DESCRIPTOR_HEAP_FLAG_SHADER_VISIBLE
 	);
@@ -1056,9 +1056,9 @@ FD3D12LocalOnlineHeap::~FD3D12LocalOnlineHeap() = default;
 /**
 Initialize a thread local online heap
 **/
-void FD3D12LocalOnlineHeap::Init(uint32 InNumDescriptors, ED3D12DescriptorHeapType InHeapType)
+void FD3D12LocalOnlineHeap::Init(uint32 InNumDescriptors, ERHIDescriptorHeapType InHeapType)
 {
-	const TCHAR* DebugName = InHeapType == ED3D12DescriptorHeapType::Standard ? L"Thread Local - Online View Heap" : L"Thread Local - Online Sampler Heap";
+	const TCHAR* DebugName = InHeapType == ERHIDescriptorHeapType::Standard ? L"Thread Local - Online View Heap" : L"Thread Local - Online Sampler Heap";
 	Heap = GetParentDevice()->GetDescriptorHeapManager().AllocateHeap(
 		DebugName,
 		InHeapType,
@@ -1068,7 +1068,7 @@ void FD3D12LocalOnlineHeap::Init(uint32 InNumDescriptors, ED3D12DescriptorHeapTy
 
 	Entry.Heap = Heap;
 
-	if (InHeapType == ED3D12DescriptorHeapType::Standard)
+	if (InHeapType == ERHIDescriptorHeapType::Standard)
 	{
 		INC_DWORD_STAT(STAT_NumViewOnlineDescriptorHeaps);
 		INC_MEMORY_STAT_BY(STAT_ViewOnlineDescriptorHeapMemory, Heap->GetMemorySize());
@@ -1103,10 +1103,10 @@ bool FD3D12LocalOnlineHeap::RollOver()
 
 		//LLM_SCOPE(ELLMTag::DescriptorCache);
 
-		const ED3D12DescriptorHeapType HeapType = Heap->GetType();
+		const ERHIDescriptorHeapType HeapType = Heap->GetType();
 		const uint32 NumDescriptors = Heap->GetNumDescriptors();
 
-		const TCHAR* DebugName = HeapType == ED3D12DescriptorHeapType::Standard ? L"Thread Local - Online View Heap" : L"Thread Local - Online Sampler Heap";
+		const TCHAR* DebugName = HeapType == ERHIDescriptorHeapType::Standard ? L"Thread Local - Online View Heap" : L"Thread Local - Online Sampler Heap";
 		Heap = GetParentDevice()->GetDescriptorHeapManager().AllocateHeap(
 			DebugName,
 			HeapType,
@@ -1114,7 +1114,7 @@ bool FD3D12LocalOnlineHeap::RollOver()
 			D3D12_DESCRIPTOR_HEAP_FLAG_SHADER_VISIBLE
 		);
 
-		if (HeapType == ED3D12DescriptorHeapType::Standard)
+		if (HeapType == ERHIDescriptorHeapType::Standard)
 		{
 			INC_DWORD_STAT(STAT_NumViewOnlineDescriptorHeaps);
 			INC_MEMORY_STAT_BY(STAT_ViewOnlineDescriptorHeapMemory, Heap->GetMemorySize());
diff --git a/Engine/Source/Runtime/D3D12RHI/Private/D3D12DescriptorCache.h b/Engine/Source/Runtime/D3D12RHI/Private/D3D12DescriptorCache.h
--- a/Engine/Source/Runtime/D3D12RHI/Private/D3D12DescriptorCache.h
+++ b/Engine/Source/Runtime/D3D12RHI/Private/D3D12DescriptorCache.h
@@ -264,7 +264,7 @@ public:
 	~FD3D12LocalOnlineHeap();
 
 	// Allocate the actual overflow heap
-	void Init(uint32 InNumDescriptors, ED3D12DescriptorHeapType InHeapType);
+	void Init(uint32 InNumDescriptors, ERHIDescriptorHeapType InHeapType);
 
 	// Override FD3D12OnlineHeap functions
 	virtual bool RollOver() final override;
@@ -380,8 +380,8 @@ public:
 
 	void SetStreamOutTargets(FD3D12Resource **Buffers, uint32 Count, const uint32* Offsets);
 
-	bool HeapRolledOver(ED3D12DescriptorHeapType InHeapType);
-	void HeapLoopedAround(ED3D12DescriptorHeapType InHeapType);
+	bool HeapRolledOver(ERHIDescriptorHeapType InHeapType);
+	void HeapLoopedAround(ERHIDescriptorHeapType InHeapType);
 	void Init(FD3D12Device* InParent, FD3D12CommandContext* InCmdContext, uint32 InNumLocalViewDescriptors, uint32 InNumSamplerDescriptors);
 	void Clear();
 	void BeginFrame();
diff --git a/Engine/Source/Runtime/D3D12RHI/Private/D3D12Descriptors.cpp b/Engine/Source/Runtime/D3D12RHI/Private/D3D12Descriptors.cpp
--- a/Engine/Source/Runtime/D3D12RHI/Private/D3D12Descriptors.cpp
+++ b/Engine/Source/Runtime/D3D12RHI/Private/D3D12Descriptors.cpp
@@ -3,34 +3,22 @@
 #include "D3D12RHIPrivate.h"
 #include "D3D12Descriptors.h"
 
-const TCHAR* ToString(ED3D12DescriptorHeapType InHeapType)
-{
-	switch (InHeapType)
-	{
-	default: checkNoEntry();                     return TEXT("UnknownDescriptorHeapType");
-	case ED3D12DescriptorHeapType::Standard:     return TEXT("Standard");
-	case ED3D12DescriptorHeapType::RenderTarget: return TEXT("RenderTarget");
-	case ED3D12DescriptorHeapType::DepthStencil: return TEXT("DepthStencil");
-	case ED3D12DescriptorHeapType::Sampler:      return TEXT("Sampler");
-	}
-}
-
-inline D3D12_DESCRIPTOR_HEAP_TYPE Translate(ED3D12DescriptorHeapType InHeapType)
+inline D3D12_DESCRIPTOR_HEAP_TYPE Translate(ERHIDescriptorHeapType InHeapType)
 {
 	switch (InHeapType)
 	{
 	default: checkNoEntry();
-	case ED3D12DescriptorHeapType::Standard:     return D3D12_DESCRIPTOR_HEAP_TYPE_CBV_SRV_UAV;
-	case ED3D12DescriptorHeapType::RenderTarget: return D3D12_DESCRIPTOR_HEAP_TYPE_RTV;
-	case ED3D12DescriptorHeapType::DepthStencil: return D3D12_DESCRIPTOR_HEAP_TYPE_DSV;
-	case ED3D12DescriptorHeapType::Sampler:      return FD3D12_DESCRIPTOR_HEAP_TYPE_SAMPLER;
+	case ERHIDescriptorHeapType::Standard:     return D3D12_DESCRIPTOR_HEAP_TYPE_CBV_SRV_UAV;
+	case ERHIDescriptorHeapType::RenderTarget: return D3D12_DESCRIPTOR_HEAP_TYPE_RTV;
+	case ERHIDescriptorHeapType::DepthStencil: return D3D12_DESCRIPTOR_HEAP_TYPE_DSV;
+	case ERHIDescriptorHeapType::Sampler:      return FD3D12_DESCRIPTOR_HEAP_TYPE_SAMPLER;
 	}
 }
 
 ////////////////////////////////////////////////////////////////////////////////////////////////////////////////////////
 // FD3D12DescriptorHeap
 
-FD3D12DescriptorHeap::FD3D12DescriptorHeap(FD3D12Device* InDevice, ID3D12DescriptorHeap* InHeap, uint32 InNumDescriptors, ED3D12DescriptorHeapType InType, D3D12_DESCRIPTOR_HEAP_FLAGS InFlags, bool bInIsGlobal)
+FD3D12DescriptorHeap::FD3D12DescriptorHeap(FD3D12Device* InDevice, ID3D12DescriptorHeap* InHeap, uint32 InNumDescriptors, ERHIDescriptorHeapType InType, D3D12_DESCRIPTOR_HEAP_FLAGS InFlags, bool bInIsGlobal)
 	: FD3D12DeviceChild(InDevice)
 	, Heap(InHeap)
 	, CpuBase(InHeap->GetCPUDescriptorHandleForHeapStart())
@@ -216,31 +204,33 @@ void FDescriptorAllocator::Free(uint32 Offset, uint32 NumDescriptors)
 }
 
 ////////////////////////////////////////////////////////////////////////////////////////////////////////////////////////
-// FD3D12ResourceDescriptorManager
+// FD3D12DescriptorManager
 
-FD3D12ResourceDescriptorManager::FD3D12ResourceDescriptorManager(FD3D12Device* Device)
+FD3D12DescriptorManager::FD3D12DescriptorManager(FD3D12Device* Device, FD3D12DescriptorHeap* InHeap)
 	: FD3D12DeviceChild(Device)
+	, Heap(InHeap)
 {
+	Allocator.Init(InHeap->GetNumDescriptors());
 }
 
-FD3D12ResourceDescriptorManager::~FD3D12ResourceDescriptorManager() = default;
+FD3D12DescriptorManager::~FD3D12DescriptorManager() = default;
 
-void FD3D12ResourceDescriptorManager::Init(uint32 InTotalSize)
+void FD3D12DescriptorManager::Init(TCHAR* InName, ERHIDescriptorHeapType InType, uint32 InNumDescriptors)
 {
-	if (InTotalSize > 0)
+	if (InNumDescriptors > 0)
 	{
 		Heap = GetParentDevice()->GetDescriptorHeapManager().AllocateHeap(
-			TEXT("ResourceDescriptorHeap"),
-			ED3D12DescriptorHeapType::Standard,
-			InTotalSize,
+			InName,
+			InType,
+			InNumDescriptors,
 			D3D12_DESCRIPTOR_HEAP_FLAG_SHADER_VISIBLE
 		);
 
-		Allocator.Init(InTotalSize);
+		Allocator.Init(InNumDescriptors);
 	}
 }
 
-void FD3D12ResourceDescriptorManager::Destroy()
+void FD3D12DescriptorManager::Destroy()
 {
 	if (Heap)
 	{
@@ -249,26 +239,109 @@ void FD3D12ResourceDescriptorManager::Destroy()
 	}
 }
 
-bool FD3D12ResourceDescriptorManager::AllocateDescriptor(uint32& OutSlot)
+FRHIDescriptorHandle FD3D12DescriptorManager::AllocateDescriptor()
 {
-	return Allocator.Allocate(1, OutSlot);
+	uint32 Index{};
+	if (Allocator.Allocate(1, Index))
+	{
+		return FRHIDescriptorHandle(Heap->GetType(), Index);
+	}
+	return FRHIDescriptorHandle();
 }
 
-bool FD3D12ResourceDescriptorManager::AllocateDescriptors(uint32 NumDescriptors, uint32& OutSlot)
+bool FD3D12DescriptorManager::AllocateDescriptors(uint32 NumDescriptors, uint32& OutSlot)
 {
 	return Allocator.Allocate(NumDescriptors, OutSlot);
 }
 
-void FD3D12ResourceDescriptorManager::FreeDescriptor(uint32 Slot)
+void FD3D12DescriptorManager::FreeDescriptor(FRHIDescriptorHandle InHandle)
 {
-	Allocator.Free(Slot, 1);
+	check(InHandle.GetType() == Heap->GetType());
+	Allocator.Free(InHandle.GetIndex(), 1);
 }
 
-void FD3D12ResourceDescriptorManager::FreeDescriptors(uint32 Slot, uint32 NumDescriptors)
+void FD3D12DescriptorManager::FreeDescriptors(uint32 Slot, uint32 NumDescriptors)
 {
 	Allocator.Free(Slot, NumDescriptors);
 }
 
+////////////////////////////////////////////////////////////////////////////////////////////////////////////////////////
+// FD3D12BindlessDescriptorManager
+
+FD3D12BindlessDescriptorManager::FD3D12BindlessDescriptorManager(FD3D12Device* Device)
+	: FD3D12DeviceChild(Device)
+{
+}
+
+void FD3D12BindlessDescriptorManager::Init(uint32 InNumResourceDescriptors, uint32 InNumSamplerDescriptors)
+{
+	FD3D12DescriptorHeapManager& HeapManager = GetParentDevice()->GetDescriptorHeapManager();
+
+	if (InNumResourceDescriptors > 0)
+	{
+		FD3D12DescriptorHeap* Heap = HeapManager.AllocateHeap(
+			TEXT("BindlessResources"),
+			ERHIDescriptorHeapType::Standard,
+			InNumResourceDescriptors,
+			D3D12_DESCRIPTOR_HEAP_FLAG_SHADER_VISIBLE
+		);
+		Managers.Emplace(GetParentDevice(), Heap);
+	}
+
+	if (InNumSamplerDescriptors > 0)
+	{
+		FD3D12DescriptorHeap* Heap = HeapManager.AllocateHeap(
+			TEXT("BindlessSamplers"),
+			ERHIDescriptorHeapType::Sampler,
+			InNumSamplerDescriptors,
+			D3D12_DESCRIPTOR_HEAP_FLAG_SHADER_VISIBLE
+		);
+		Managers.Emplace(GetParentDevice(), Heap);
+	}
+}
+
+FRHIDescriptorHandle FD3D12BindlessDescriptorManager::AllocateDescriptor(ERHIDescriptorHeapType InType)
+{
+	for (FD3D12DescriptorManager& Manager : Managers)
+	{
+		if (Manager.HandlesAllocation(InType))
+		{
+			return Manager.AllocateDescriptor();
+		}
+	}
+	return FRHIDescriptorHandle();
+}
+
+void FD3D12BindlessDescriptorManager::FreeDescriptor(FRHIDescriptorHandle InHandle)
+{
+	for (FD3D12DescriptorManager& Manager : Managers)
+	{
+		if (Manager.HandlesAllocation(InHandle.GetType()))
+		{
+			Manager.FreeDescriptor(InHandle);
+			return;
+		}
+	}
+
+	// Bad configuration?
+	checkNoEntry();
+}
+
+D3D12_CPU_DESCRIPTOR_HANDLE FD3D12BindlessDescriptorManager::GetCpuDescriptorHandle(FRHIDescriptorHandle InHandle) const
+{
+	for (const FD3D12DescriptorManager& Manager : Managers)
+	{
+		if (Manager.HandlesAllocation(InHandle.GetType()))
+		{
+			return Manager.GetCpuDescriptorHandle(InHandle.GetIndex());
+		}
+	}
+
+	checkNoEntry();
+
+	D3D12_CPU_DESCRIPTOR_HANDLE NullHandle{};
+	return NullHandle;
+}
 
 ////////////////////////////////////////////////////////////////////////////////////////////////////////////////////////
 // FD3D12OnlineDescriptorManager
@@ -285,7 +358,7 @@ void FD3D12OnlineDescriptorManager::Init(uint32 InTotalSize, uint32 InBlockSize)
 {
 	Heap = GetParentDevice()->GetDescriptorHeapManager().AllocateHeap(
 		TEXT("Device Global - Online View Heap"),
-		ED3D12DescriptorHeapType::Standard,
+		ERHIDescriptorHeapType::Standard,
 		InTotalSize,
 		D3D12_DESCRIPTOR_HEAP_FLAG_SHADER_VISIBLE);
 
@@ -377,19 +450,19 @@ void FD3D12OnlineDescriptorManager::UpdateFreeBlocks()
 ////////////////////////////////////////////////////////////////////////////////////////////////////////////////////////
 // FD3D12OfflineDescriptorManager
 
-static uint32 GetOfflineDescriptorHeapDefaultSize(ED3D12DescriptorHeapType InHeapType)
+static uint32 GetOfflineDescriptorHeapDefaultSize(ERHIDescriptorHeapType InHeapType)
 {
 	switch (InHeapType)
 	{
 	default: checkNoEntry();
 #if USE_STATIC_ROOT_SIGNATURE
-	case ED3D12DescriptorHeapType::Standard:     return 4096;
+	case ERHIDescriptorHeapType::Standard:     return 4096;
 #else
-	case ED3D12DescriptorHeapType::Standard:     return 2048;
+	case ERHIDescriptorHeapType::Standard:     return 2048;
 #endif
-	case ED3D12DescriptorHeapType::RenderTarget: return 256;
-	case ED3D12DescriptorHeapType::DepthStencil: return 256;
-	case ED3D12DescriptorHeapType::Sampler:      return 128;
+	case ERHIDescriptorHeapType::RenderTarget: return 256;
+	case ERHIDescriptorHeapType::DepthStencil: return 256;
+	case ERHIDescriptorHeapType::Sampler:      return 128;
 	}
 }
 
@@ -418,7 +491,7 @@ FD3D12OfflineDescriptorManager::FD3D12OfflineDescriptorManager(FD3D12Device* InD
 
 FD3D12OfflineDescriptorManager::~FD3D12OfflineDescriptorManager() = default;
 
-void FD3D12OfflineDescriptorManager::Init(ED3D12DescriptorHeapType InHeapType)
+void FD3D12OfflineDescriptorManager::Init(ERHIDescriptorHeapType InHeapType)
 {
 	HeapType = InHeapType;
 	DescriptorSize = GetParentDevice()->GetDevice()->GetDescriptorHandleIncrementSize(Translate(InHeapType));
@@ -534,7 +607,7 @@ FD3D12DescriptorHeapManager::~FD3D12DescriptorHeapManager()
 	Destroy();
 }
 
-static FD3D12DescriptorHeap* CreateDescriptorHeap(FD3D12Device* Device, const TCHAR* DebugName, ED3D12DescriptorHeapType HeapType, uint32 NumDescriptors, D3D12_DESCRIPTOR_HEAP_FLAGS Flags, bool bIsGlobal)
+static FD3D12DescriptorHeap* CreateDescriptorHeap(FD3D12Device* Device, const TCHAR* DebugName, ERHIDescriptorHeapType HeapType, uint32 NumDescriptors, D3D12_DESCRIPTOR_HEAP_FLAGS Flags, bool bIsGlobal)
 {
 	D3D12_DESCRIPTOR_HEAP_DESC Desc{};
 	Desc.Type = Translate(HeapType);
@@ -550,20 +623,34 @@ static FD3D12DescriptorHeap* CreateDescriptorHeap(FD3D12Device* Device, const TC
 	return new FD3D12DescriptorHeap(Device, Heap, NumDescriptors, HeapType, Flags, bIsGlobal);
 }
 
-void FD3D12DescriptorHeapManager::Init(uint32 InNumGlobalDescriptors)
+void FD3D12DescriptorHeapManager::Init(uint32 InNumGlobalResourceDescriptors, uint32 InNumGlobalSamplerDescriptors)
 {
-	if (InNumGlobalDescriptors > 0)
+	if (InNumGlobalResourceDescriptors > 0)
+	{
+		FD3D12DescriptorHeap* Heap = CreateDescriptorHeap(
+			GetParentDevice(),
+			TEXT("GlobalResourceHeap"),
+			ERHIDescriptorHeapType::Standard,
+			InNumGlobalResourceDescriptors,
+			D3D12_DESCRIPTOR_HEAP_FLAG_SHADER_VISIBLE,
+			true /* bIsGlobal */
+		);
+
+		GlobalHeaps.Emplace(GetParentDevice(), Heap);
+	}
+
+	if (InNumGlobalSamplerDescriptors > 0)
 	{
-		GlobalHeap = CreateDescriptorHeap(
+		FD3D12DescriptorHeap* Heap = CreateDescriptorHeap(
 			GetParentDevice(),
-			TEXT("ViewDescriptorHeap"),
-			ED3D12DescriptorHeapType::Standard,
-			InNumGlobalDescriptors,
+			TEXT("GlobalSamplerHeap"),
+			ERHIDescriptorHeapType::Sampler,
+			InNumGlobalSamplerDescriptors,
 			D3D12_DESCRIPTOR_HEAP_FLAG_SHADER_VISIBLE,
 			true /* bIsGlobal */
 		);
 
-		Allocator.Init(InNumGlobalDescriptors);
+		GlobalHeaps.Emplace(GetParentDevice(), Heap);
 	}
 }
 
@@ -571,18 +658,21 @@ void FD3D12DescriptorHeapManager::Destroy()
 {
 }
 
-FD3D12DescriptorHeap* FD3D12DescriptorHeapManager::AllocateHeap(const TCHAR* InDebugName, ED3D12DescriptorHeapType InHeapType, uint32 InNumDescriptors, D3D12_DESCRIPTOR_HEAP_FLAGS InHeapFlags)
+FD3D12DescriptorHeap* FD3D12DescriptorHeapManager::AllocateHeap(const TCHAR* InDebugName, ERHIDescriptorHeapType InHeapType, uint32 InNumDescriptors, D3D12_DESCRIPTOR_HEAP_FLAGS InHeapFlags)
 {
-	if (GlobalHeap && InHeapType == GlobalHeap->GetType() && InHeapFlags == GlobalHeap->GetFlags())
+	for (FD3D12DescriptorManager& GlobalHeap : GlobalHeaps)
 	{
-		uint32 Offset = 0;
-		if (Allocator.Allocate(InNumDescriptors, Offset))
+		if (GlobalHeap.HandlesAllocation(InHeapType, InHeapFlags))
 		{
-			return new FD3D12DescriptorHeap(GlobalHeap, Offset, InNumDescriptors);
-		}
+			uint32 Offset = 0;
+			if (GlobalHeap.AllocateDescriptors(InNumDescriptors, Offset))
+			{
+				return new FD3D12DescriptorHeap(GlobalHeap.GetHeap(), Offset, InNumDescriptors);
+			}
 
-		// TODO: handle running out of space
-		checkNoEntry();
+			// TODO: handle running out of space
+			//checkNoEntry();
+		}
 	}
 
 	return CreateDescriptorHeap(GetParentDevice(), InDebugName, InHeapType, InNumDescriptors, InHeapFlags, false);
@@ -590,9 +680,13 @@ FD3D12DescriptorHeap* FD3D12DescriptorHeapManager::AllocateHeap(const TCHAR* InD
 
 void FD3D12DescriptorHeapManager::FreeHeap(FD3D12DescriptorHeap* InHeap)
 {
-	if (InHeap->IsGlobal())
+	for (FD3D12DescriptorManager& GlobalHeap : GlobalHeaps)
 	{
-		check(InHeap->GetHeap() == GlobalHeap->GetHeap());
-		Allocator.Free(InHeap->GetOffset(), InHeap->GetNumDescriptors());
+		if (GlobalHeap.IsHeapAChild(InHeap))
+		{
+			GlobalHeap.FreeDescriptors(InHeap->GetOffset(), InHeap->GetNumDescriptors());
+			return;
+		}
 	}
+	check(!InHeap->IsGlobal());
 }
diff --git a/Engine/Source/Runtime/D3D12RHI/Private/D3D12Descriptors.h b/Engine/Source/Runtime/D3D12RHI/Private/D3D12Descriptors.h
--- a/Engine/Source/Runtime/D3D12RHI/Private/D3D12Descriptors.h
+++ b/Engine/Source/Runtime/D3D12RHI/Private/D3D12Descriptors.h
@@ -6,24 +6,13 @@
 
 struct FD3D12OfflineHeapEntry;
 
-// Internal platform agnostic Descriptor Heap type.
-enum class ED3D12DescriptorHeapType
-{
-	Standard,
-	RenderTarget,
-	DepthStencil,
-	Sampler,
-	count
-};
-const TCHAR* ToString(ED3D12DescriptorHeapType InHeapType);
-
 struct FD3D12DescriptorHeap : public FD3D12DeviceChild, public FD3D12RefCount
 {
 public:
 	FD3D12DescriptorHeap() = delete;
 
 	// Heap created with its own D3D heap object.
-	FD3D12DescriptorHeap(FD3D12Device* InDevice, ID3D12DescriptorHeap* InHeap, uint32 InNumDescriptors, ED3D12DescriptorHeapType InType, D3D12_DESCRIPTOR_HEAP_FLAGS InFlags, bool bInIsGlobal);
+	FD3D12DescriptorHeap(FD3D12Device* InDevice, ID3D12DescriptorHeap* InHeap, uint32 InNumDescriptors, ERHIDescriptorHeapType InType, D3D12_DESCRIPTOR_HEAP_FLAGS InFlags, bool bInIsGlobal);
 
 	// Heap created as a suballocation of another heap.
 	FD3D12DescriptorHeap(FD3D12DescriptorHeap* SubAllocateSourceHeap, uint32 InOffset, uint32 InNumDescriptors);
@@ -31,7 +20,7 @@ public:
 	~FD3D12DescriptorHeap();
 
 	inline ID3D12DescriptorHeap*       GetHeap()  const { return Heap; }
-	inline ED3D12DescriptorHeapType    GetType()  const { return Type; }
+	inline ERHIDescriptorHeapType      GetType()  const { return Type; }
 	inline D3D12_DESCRIPTOR_HEAP_FLAGS GetFlags() const { return Flags; }
 
 	inline uint32 GetOffset()         const { return Offset; }
@@ -60,7 +49,7 @@ private:
 	// Device provided size of each descriptor in this heap.
 	const uint32 DescriptorSize;
 
-	const ED3D12DescriptorHeapType Type;
+	const ERHIDescriptorHeapType Type;
 	const D3D12_DESCRIPTOR_HEAP_FLAGS Flags;
 
 	// Enabled if this heap is the "global" heap.
@@ -69,6 +58,7 @@ private:
 	// Enabled if this heap was allocated inside another heap.
 	const bool bIsSuballocation;
 };
+using FD3D12DescriptorHeapPtr = TRefCountPtr<FD3D12DescriptorHeap>;
 
 struct FDescriptorAllocatorRange;
 
@@ -90,30 +80,63 @@ private:
 	FCriticalSection CriticalSection;
 };
 
-/** Manager for resource descriptors used in bindless rendering. */
-class FD3D12ResourceDescriptorManager : public FD3D12DeviceChild
+/** Manager for resource descriptor allocations. */
+class FD3D12DescriptorManager : public FD3D12DeviceChild
 {
 public:
-	FD3D12ResourceDescriptorManager(FD3D12Device* Device);
-	~FD3D12ResourceDescriptorManager();
+	FD3D12DescriptorManager() = delete;
+	FD3D12DescriptorManager(FD3D12Device* Device, FD3D12DescriptorHeap* InHeap);
+	~FD3D12DescriptorManager();
 
-	void Init(uint32 InTotalSize);
+	void Init(TCHAR* InName, ERHIDescriptorHeapType InType, uint32 InNumDescriptors);
 	void Destroy();
 
-	bool AllocateDescriptor(uint32& OutSlot);
+	FRHIDescriptorHandle AllocateDescriptor();
 	bool AllocateDescriptors(uint32 NumDescriptors, uint32& OutSlot);
 
-	void FreeDescriptor(uint32 Slot);
+	void FreeDescriptor(FRHIDescriptorHandle InHandle);
 	void FreeDescriptors(uint32 Slot, uint32 NumDescriptors);
 
-	inline D3D12_CPU_DESCRIPTOR_HANDLE GetResourceHandle(uint32 InSlot) const { return Heap->GetCPUSlotHandle(InSlot); }
+	inline FD3D12DescriptorHeap* GetHeap() { return Heap.GetReference(); }
+	inline D3D12_CPU_DESCRIPTOR_HANDLE GetCpuDescriptorHandle(uint32 InSlot) const { return Heap->GetCPUSlotHandle(InSlot); }
 
-private:
-	TRefCountPtr<FD3D12DescriptorHeap> Heap;
+	inline bool HandlesAllocation(ERHIDescriptorHeapType InHeapType) const
+	{
+		return Heap && Heap->GetType() == InHeapType;
+	}
+
+	inline bool HandlesAllocation(ERHIDescriptorHeapType InHeapType, D3D12_DESCRIPTOR_HEAP_FLAGS InHeapFlags) const
+	{
+		return HandlesAllocation(InHeapType) && Heap->GetFlags() == InHeapFlags;
+	}
+
+	inline bool IsHeapAChild(const FD3D12DescriptorHeap* InHeap)
+	{
+		return InHeap->GetHeap() == Heap->GetHeap();
+	}
 
+private:
+	FD3D12DescriptorHeapPtr Heap;
 	FDescriptorAllocator Allocator;
 };
 
+/** Manager for resource descriptors used in bindless rendering. */
+class FD3D12BindlessDescriptorManager : public FD3D12DeviceChild
+{
+public:
+	FD3D12BindlessDescriptorManager(FD3D12Device* Device);
+
+	void Init(uint32 InNumResourceDescriptors, uint32 InNumSamplerDescriptors);
+
+	FRHIDescriptorHandle AllocateDescriptor(ERHIDescriptorHeapType InType);
+	void FreeDescriptor(FRHIDescriptorHandle InHandle);
+
+	D3D12_CPU_DESCRIPTOR_HANDLE GetCpuDescriptorHandle(FRHIDescriptorHandle InHandle) const;
+
+private:
+	TArray<FD3D12DescriptorManager> Managers;
+};
+
 /** Heap sub block of an online heap */
 struct FD3D12OnlineDescriptorBlock
 {
@@ -154,7 +177,7 @@ private:
 	// Check all released blocks and check which ones are not used by the GPU anymore
 	void UpdateFreeBlocks();
 
-	TRefCountPtr<FD3D12DescriptorHeap> Heap;
+	FD3D12DescriptorHeapPtr Heap;
 
 	TQueue<FD3D12OnlineDescriptorBlock*> FreeBlocks;
 	TArray<FD3D12OnlineDescriptorBlock*> ReleasedBlocks;
@@ -170,9 +193,9 @@ public:
 	FD3D12OfflineDescriptorManager(FD3D12Device* InDevice);
 	~FD3D12OfflineDescriptorManager();
 
-	inline ED3D12DescriptorHeapType GetHeapType() const { return HeapType; }
+	inline ERHIDescriptorHeapType GetHeapType() const { return HeapType; }
 
-	void Init(ED3D12DescriptorHeapType InHeapType);
+	void Init(ERHIDescriptorHeapType InHeapType);
 
 	D3D12_CPU_DESCRIPTOR_HANDLE AllocateHeapSlot(uint32& outIndex);
 	void FreeHeapSlot(D3D12_CPU_DESCRIPTOR_HANDLE Offset, uint32 index);
@@ -183,7 +206,7 @@ private:
 	TArray<FD3D12OfflineHeapEntry> Heaps;
 	TDoubleLinkedList<uint32> FreeHeaps;
 
-	ED3D12DescriptorHeapType HeapType;
+	ERHIDescriptorHeapType HeapType;
 	uint32 NumDescriptorsPerHeap{};
 	uint32 DescriptorSize{};
 
@@ -198,16 +221,15 @@ public:
 	FD3D12DescriptorHeapManager(FD3D12Device* InDevice);
 	~FD3D12DescriptorHeapManager();
 
-	void Init(uint32 InNumGlobalDescriptors);
+	void Init(uint32 InNumGlobalResourceDescriptors, uint32 InNumGlobalSamplerDescriptors);
 	void Destroy();
 
-	FD3D12DescriptorHeap* AllocateHeap(const TCHAR* DebugName, ED3D12DescriptorHeapType InHeapType, uint32 InNumDescriptors, D3D12_DESCRIPTOR_HEAP_FLAGS InFlags);
+	FD3D12DescriptorHeap* AllocateHeap(const TCHAR* DebugName, ERHIDescriptorHeapType InHeapType, uint32 InNumDescriptors, D3D12_DESCRIPTOR_HEAP_FLAGS InFlags);
 	void FreeHeap(FD3D12DescriptorHeap* InHeap);
 
-	inline FD3D12DescriptorHeap* GetGlobalHeap() const { return GlobalHeap; }
+	//inline FD3D12DescriptorHeap* GetGlobalHeap() const { return GlobalHeap; }
 
 private:
-	TRefCountPtr<FD3D12DescriptorHeap> GlobalHeap;
-	FDescriptorAllocator Allocator;
+	TArray<FD3D12DescriptorManager> GlobalHeaps;
 };
 
diff --git a/Engine/Source/Runtime/D3D12RHI/Private/D3D12Device.cpp b/Engine/Source/Runtime/D3D12RHI/Private/D3D12Device.cpp
--- a/Engine/Source/Runtime/D3D12RHI/Private/D3D12Device.cpp
+++ b/Engine/Source/Runtime/D3D12RHI/Private/D3D12Device.cpp
@@ -16,7 +16,7 @@ FD3D12Device::FD3D12Device(FRHIGPUMask InGPUMask, FD3D12Adapter* InAdapter) :
 	AsyncCommandListManager(nullptr),
 	TextureStreamingCommandAllocatorManager(this, D3D12_COMMAND_LIST_TYPE_COPY),
 	DescriptorHeapManager(this),
-	ResourceDescriptorManager(this),
+	BindlessDescriptorManager(this),
 	OfflineDescriptorManagers{
 		FD3D12OfflineDescriptorManager(this),
 		FD3D12OfflineDescriptorManager(this),
@@ -316,6 +316,7 @@ void FD3D12Device::SetupAfterDeviceCreation()
 
 	const D3D12_RESOURCE_BINDING_TIER Tier = GetParentAdapter()->GetResourceBindingTier();
 	uint32 MaximumSupportedHeapSize = NUM_VIEW_DESCRIPTORS_TIER_1;
+	uint32 MaximumSupportedSamplerHeapSize = D3D12_MAX_SHADER_VISIBLE_SAMPLER_HEAP_SIZE;
 
 	switch (Tier)
 	{
@@ -333,17 +334,19 @@ void FD3D12Device::SetupAfterDeviceCreation()
 
 	// This value can be tuned on a per app basis. I.e. most apps will never run into descriptor heap pressure so
 	// can make this global heap smaller
-	check((uint32)GGlobalDescriptorHeapSize <= MaximumSupportedHeapSize);
-	check(GOnlineDescriptorHeapSize <= GGlobalDescriptorHeapSize);
-	check(GResourceDescriptorHeapSize <= GGlobalDescriptorHeapSize);
+	check((uint32)GGlobalResourceDescriptorHeapSize <= MaximumSupportedHeapSize);
+	check((uint32)GGlobalSamplerDescriptorHeapSize <= MaximumSupportedSamplerHeapSize);
+	check(GOnlineDescriptorHeapSize <= GGlobalResourceDescriptorHeapSize);
+	check(GResourceDescriptorHeapSize <= GGlobalResourceDescriptorHeapSize);
+	check(GSamplerDescriptorHeapSize <= GGlobalSamplerDescriptorHeapSize);
 
-	DescriptorHeapManager.Init(GGlobalDescriptorHeapSize);
-	ResourceDescriptorManager.Init(GResourceDescriptorHeapSize);
+	DescriptorHeapManager.Init(GGlobalResourceDescriptorHeapSize, GGlobalSamplerDescriptorHeapSize);
+	BindlessDescriptorManager.Init(GResourceDescriptorHeapSize, GSamplerDescriptorHeapSize);
 
 	// Init offline descriptor managers
-	for (uint32 Index = 0; Index < static_cast<uint32>(ED3D12DescriptorHeapType::count); Index++)
+	for (uint32 Index = 0; Index < static_cast<uint32>(ERHIDescriptorHeapType::count); Index++)
 	{
-		OfflineDescriptorManagers[Index].Init(static_cast<ED3D12DescriptorHeapType>(Index));
+		OfflineDescriptorManagers[Index].Init(static_cast<ERHIDescriptorHeapType>(Index));
 	}
 
 	GlobalSamplerHeap.Init(NUM_SAMPLER_DESCRIPTORS);
diff --git a/Engine/Source/Runtime/D3D12RHI/Private/D3D12Device.h b/Engine/Source/Runtime/D3D12RHI/Private/D3D12Device.h
--- a/Engine/Source/Runtime/D3D12RHI/Private/D3D12Device.h
+++ b/Engine/Source/Runtime/D3D12RHI/Private/D3D12Device.h
@@ -70,10 +70,10 @@ public:
 	FD3D12LinearQueryHeap* GetCmdListExecTimeQueryHeap();
 
 	inline FD3D12DescriptorHeapManager& GetDescriptorHeapManager() { return DescriptorHeapManager; }
-	inline FD3D12ResourceDescriptorManager& GetResourceDescriptorManager() { return ResourceDescriptorManager; }
-	inline FD3D12OfflineDescriptorManager& GetOfflineDescriptorManager(ED3D12DescriptorHeapType InType)
+	inline FD3D12BindlessDescriptorManager& GetBindlessDescriptorManager() { return BindlessDescriptorManager; }
+	inline FD3D12OfflineDescriptorManager& GetOfflineDescriptorManager(ERHIDescriptorHeapType InType)
 	{
-		check(InType < ED3D12DescriptorHeapType::count);
+		check(InType < ERHIDescriptorHeapType::count);
 		return OfflineDescriptorManagers[static_cast<int>(InType)];
 	}
 
@@ -145,8 +145,8 @@ protected:
 
 	// Must be before the StateCache so that destructor ordering is valid
 	FD3D12DescriptorHeapManager DescriptorHeapManager;
-	FD3D12ResourceDescriptorManager ResourceDescriptorManager;
-	FD3D12OfflineDescriptorManager OfflineDescriptorManagers[static_cast<int>(ED3D12DescriptorHeapType::count)];
+	FD3D12BindlessDescriptorManager BindlessDescriptorManager;
+	FD3D12OfflineDescriptorManager OfflineDescriptorManagers[static_cast<int>(ERHIDescriptorHeapType::count)];
 
 	FD3D12GlobalOnlineSamplerHeap GlobalSamplerHeap;
 	FD3D12OnlineDescriptorManager OnlineDescriptorManager;
diff --git a/Engine/Source/Runtime/D3D12RHI/Private/D3D12RayTracing.cpp b/Engine/Source/Runtime/D3D12RHI/Private/D3D12RayTracing.cpp
--- a/Engine/Source/Runtime/D3D12RHI/Private/D3D12RayTracing.cpp
+++ b/Engine/Source/Runtime/D3D12RHI/Private/D3D12RayTracing.cpp
@@ -4424,7 +4424,7 @@ struct FD3D12RayTracingLocalResourceBinder
 		FD3D12ConstantBufferView* BufferView = new FD3D12ConstantBufferView(GetDevice());
 		BufferView->Create(Address, DataSize);
 		ShaderTable.TransientCBVs.Add(BufferView);
-		return BufferView->GetView();
+		return BufferView->GetOfflineCpuHandle();
 	}
 #endif // USE_STATIC_ROOT_SIGNATURE
 
@@ -4523,7 +4523,7 @@ static bool SetRayTracingShaderResources(
 		{
 			FD3D12TextureBase* Texture = FD3D12CommandContext::RetrieveTextureBase(Resource, GPUIndex);
 			FD3D12ShaderResourceView* SRV = Texture->GetShaderResourceView();
-			LocalSRVs[SRVIndex] = SRV->GetView();
+			LocalSRVs[SRVIndex] = SRV->GetOfflineCpuHandle();
 			BoundSRVMask |= 1ull << SRVIndex;
 
 			ReferencedResources.Add({ Texture->GetResource(), Resource });
@@ -4537,7 +4537,7 @@ static bool SetRayTracingShaderResources(
 		if (Resource)
 		{
 			FD3D12ShaderResourceView* SRV = FD3D12CommandContext::RetrieveObject<FD3D12ShaderResourceView>(Resource, GPUIndex);
-			LocalSRVs[SRVIndex] = SRV->GetView();
+			LocalSRVs[SRVIndex] = SRV->GetOfflineCpuHandle();
 			BoundSRVMask |= 1ull << SRVIndex;
 
 			ReferencedResources.Add({ SRV->GetResource(), Resource });
@@ -4552,7 +4552,7 @@ static bool SetRayTracingShaderResources(
 		{
 			FD3D12UniformBuffer* CBV = FD3D12CommandContext::RetrieveObject<FD3D12UniformBuffer>(Resource, GPUIndex);
 		#if USE_STATIC_ROOT_SIGNATURE
-			LocalCBVs[CBVIndex] = CBV->View->GetView();
+			LocalCBVs[CBVIndex] = CBV->View->GetOfflineCpuHandle();
 		#else // USE_STATIC_ROOT_SIGNATURE
 			LocalCBVs[CBVIndex] = CBV->ResourceLocation.GetGPUVirtualAddress();
 		#endif // USE_STATIC_ROOT_SIGNATURE
@@ -4567,7 +4567,7 @@ static bool SetRayTracingShaderResources(
 		FRHISamplerState* Resource = Samplers[SamplerIndex];
 		if (Resource)
 		{
-			LocalSamplers[SamplerIndex] = FD3D12CommandContext::RetrieveObject<FD3D12SamplerState>(Resource, GPUIndex)->Descriptor;
+			LocalSamplers[SamplerIndex] = FD3D12CommandContext::RetrieveObject<FD3D12SamplerState>(Resource, GPUIndex)->OfflineHandle;
 			BoundSamplerMask |= 1ull << SamplerIndex;
 		}
 	}
@@ -4578,7 +4578,7 @@ static bool SetRayTracingShaderResources(
 		if (Resource)
 		{
 			FD3D12UnorderedAccessView* UAV = FD3D12CommandContext::RetrieveObject<FD3D12UnorderedAccessView>(Resource, GPUIndex);
-			LocalUAVs[UAVIndex] = UAV->GetView();
+			LocalUAVs[UAVIndex] = UAV->GetOfflineCpuHandle();
 			BoundUAVMask |= 1ull << UAVIndex;
 
 			ReferencedResources.Add({ UAV->GetResource(), Resource });
@@ -4632,7 +4632,7 @@ static bool SetRayTracingShaderResources(
 					}
 					check(SRV != nullptr);
 
-					LocalSRVs[BindIndex] = SRV->GetView();
+					LocalSRVs[BindIndex] = SRV->GetOfflineCpuHandle();
 					BoundSRVMask |= 1ull << BindIndex;
 
 					ReferencedResources.Add({ SRV->GetResource(), SRV });
@@ -4667,7 +4667,7 @@ static bool SetRayTracingShaderResources(
 					FD3D12ShaderResourceView* SRV = FD3D12CommandContext::RetrieveObject<FD3D12ShaderResourceView>(RHISRV, GPUIndex);
 					check(SRV != nullptr);
 
-					LocalSRVs[BindIndex] = SRV->GetView();
+					LocalSRVs[BindIndex] = SRV->GetOfflineCpuHandle();
 					BoundSRVMask |= 1ull << BindIndex;
 
 					ReferencedResources.Add({ SRV->GetResource(), SRV });
@@ -4702,7 +4702,7 @@ static bool SetRayTracingShaderResources(
 					FD3D12UnorderedAccessView* UAV = FD3D12CommandContext::RetrieveObject<FD3D12UnorderedAccessView>(RHIUAV, GPUIndex);
 					check(UAV != nullptr);
 
-					LocalUAVs[BindIndex] = UAV->GetView();
+					LocalUAVs[BindIndex] = UAV->GetOfflineCpuHandle();
 					BoundUAVMask |= 1ull << BindIndex;
 
 					ReferencedResources.Add({ UAV->GetResource(), UAV });
@@ -4737,7 +4737,7 @@ static bool SetRayTracingShaderResources(
 					FD3D12SamplerState* Sampler = FD3D12CommandContext::RetrieveObject<FD3D12SamplerState>(RHISampler, GPUIndex);
 					check(Sampler != nullptr);
 
-					LocalSamplers[BindIndex] = Sampler->Descriptor;
+					LocalSamplers[BindIndex] = Sampler->OfflineHandle;
 					BoundSamplerMask |= 1ull << BindIndex;
 
 					ResourceInfo = *ResourceInfos++;
diff --git a/Engine/Source/Runtime/D3D12RHI/Private/D3D12State.cpp b/Engine/Source/Runtime/D3D12RHI/Private/D3D12State.cpp
--- a/Engine/Source/Runtime/D3D12RHI/Private/D3D12State.cpp
+++ b/Engine/Source/Runtime/D3D12RHI/Private/D3D12State.cpp
@@ -539,19 +539,31 @@ FD3D12SamplerState::FD3D12SamplerState(FD3D12Device* InParent, const D3D12_SAMPL
 	: FD3D12DeviceChild(InParent)
 	, ID(SamplerID)
 {
-	Descriptor.ptr = 0;
-	FD3D12OfflineDescriptorManager& DescriptorAllocator = GetParentDevice()->GetOfflineDescriptorManager(ED3D12DescriptorHeapType::Sampler);
-	Descriptor = DescriptorAllocator.AllocateHeapSlot(DescriptorHeapIndex);
+	FD3D12OfflineDescriptorManager& OfflineAllocator = GetParentDevice()->GetOfflineDescriptorManager(ERHIDescriptorHeapType::Sampler);
+	OfflineHandle = OfflineAllocator.AllocateHeapSlot(OfflineIndex);
 
-	GetParentDevice()->CreateSamplerInternal(Desc, Descriptor);
+	GetParentDevice()->CreateSamplerInternal(Desc, OfflineHandle);
+
+	FD3D12BindlessDescriptorManager& BindlessDescriptorManager = GetParentDevice()->GetBindlessDescriptorManager();
+	BindlessHandle = BindlessDescriptorManager.AllocateDescriptor(ERHIDescriptorHeapType::Sampler);
+
+	if (BindlessHandle.IsValid())
+	{
+		D3D12_CPU_DESCRIPTOR_HANDLE OnlineHandle = BindlessDescriptorManager.GetCpuDescriptorHandle(BindlessHandle);
+		GetParentDevice()->GetDevice()->CopyDescriptorsSimple(1, OnlineHandle, OfflineHandle, D3D12_DESCRIPTOR_HEAP_TYPE_SAMPLER);
+	}
 }
 
 FD3D12SamplerState::~FD3D12SamplerState()
 {
-	if (Descriptor.ptr)
+	if (OfflineHandle.ptr)
 	{
-		FD3D12OfflineDescriptorManager& DescriptorAllocator = GetParentDevice()->GetOfflineDescriptorManager(ED3D12DescriptorHeapType::Sampler);
-		DescriptorAllocator.FreeHeapSlot(Descriptor, DescriptorHeapIndex);
-		Descriptor.ptr = 0;
+		FD3D12OfflineDescriptorManager& OfflineAllocator = GetParentDevice()->GetOfflineDescriptorManager(ERHIDescriptorHeapType::Sampler);
+		OfflineAllocator.FreeHeapSlot(OfflineHandle, OfflineIndex);
+
+		if (BindlessHandle.IsValid())
+		{
+			GetParentDevice()->GetBindlessDescriptorManager().FreeDescriptor(BindlessHandle);
+		}
 	}
-}
\ No newline at end of file
+}
diff --git a/Engine/Source/Runtime/D3D12RHI/Private/D3D12StateCache.cpp b/Engine/Source/Runtime/D3D12RHI/Private/D3D12StateCache.cpp
--- a/Engine/Source/Runtime/D3D12RHI/Private/D3D12StateCache.cpp
+++ b/Engine/Source/Runtime/D3D12RHI/Private/D3D12StateCache.cpp
@@ -13,15 +13,23 @@
 
 // This value defines how many descriptors will be in the device global descriptor heap. This heap contains all shader visible view descriptors.
 // Other shader visible descriptor heaps (e.g. OnlineViewHeap) are allocated from this pool. Non-visible heaps (e.g. LocalViewHeap) are allocated as standalone.
-int32 GGlobalDescriptorHeapSize = 1000 * 1000;
-static FAutoConsoleVariableRef CVarGlobalDescriptorHeapSize(
-	TEXT("D3D12.GlobalDescriptorHeapSize"),
-	GGlobalDescriptorHeapSize,
-	TEXT("Global descriptor heap size"),
+int32 GGlobalResourceDescriptorHeapSize = 1000 * 1000;
+static FAutoConsoleVariableRef CVarGlobalResourceDescriptorHeapSize(
+	TEXT("D3D12.GlobalResourceDescriptorHeapSize"),
+	GGlobalResourceDescriptorHeapSize,
+	TEXT("Global resource descriptor heap size"),
 	ECVF_ReadOnly
 );
 
-int32 GResourceDescriptorHeapSize = 0;// 500 * 1000;
+int32 GGlobalSamplerDescriptorHeapSize = 2048;
+static FAutoConsoleVariableRef CVarGlobalSamplerDescriptorHeapSize(
+	TEXT("D3D12.GlobalSamplerDescriptorHeapSize"),
+	GGlobalSamplerDescriptorHeapSize,
+	TEXT("Global sampler descriptor heap size"),
+	ECVF_ReadOnly
+);
+
+int32 GResourceDescriptorHeapSize = 0;
 static FAutoConsoleVariableRef CVarResourceDescriptorHeapSize(
 	TEXT("D3D12.ResourceDescriptorHeapSize"),
 	GResourceDescriptorHeapSize,
@@ -29,6 +37,14 @@ static FAutoConsoleVariableRef CVarResourceDescriptorHeapSize(
 	ECVF_ReadOnly
 );
 
+int32 GSamplerDescriptorHeapSize = 0;
+static FAutoConsoleVariableRef CVarSamplerDescriptorHeapSize(
+	TEXT("D3D12.SamplerDescriptorHeapSize"),
+	GSamplerDescriptorHeapSize,
+	TEXT("Sampler descriptor heap size"),
+	ECVF_ReadOnly
+);
+
 // This value defines how many descriptors will be in the device local view heap which
 // This should be tweaked for each title as heaps require VRAM. The default value of 512k takes up ~16MB
 int32 GLocalViewHeapSize = 500 * 1000;
@@ -128,7 +144,8 @@ void FD3D12StateCacheBase::Init(FD3D12Device* InParent, FD3D12CommandContext* In
 		ResourceBindingTier == D3D12_RESOURCE_BINDING_TIER_2 ? NUM_VIEW_DESCRIPTORS_TIER_2 :
 		NUM_VIEW_DESCRIPTORS_TIER_1;
 
-	check(GGlobalDescriptorHeapSize <= (int32)MaxDescriptorsForTier);
+	check(GGlobalResourceDescriptorHeapSize <= (int32)MaxDescriptorsForTier);
+	check(GGlobalSamplerDescriptorHeapSize <= (int32)MaxDescriptorsForTier);
 	check(GResourceDescriptorHeapSize <= (int32)MaxDescriptorsForTier);
 	check(GLocalViewHeapSize <= (int32)MaxDescriptorsForTier);
 	check(GOnlineDescriptorHeapSize <= (int32)MaxDescriptorsForTier);
diff --git a/Engine/Source/Runtime/D3D12RHI/Private/D3D12StateCachePrivate.h b/Engine/Source/Runtime/D3D12RHI/Private/D3D12StateCachePrivate.h
--- a/Engine/Source/Runtime/D3D12RHI/Private/D3D12StateCachePrivate.h
+++ b/Engine/Source/Runtime/D3D12RHI/Private/D3D12StateCachePrivate.h
@@ -47,8 +47,12 @@ extern bool GD3D12SkipStateCaching;
 static const bool GD3D12SkipStateCaching = false;
 #endif
 
-extern int32 GGlobalDescriptorHeapSize;
+extern int32 GGlobalResourceDescriptorHeapSize;
+extern int32 GGlobalSamplerDescriptorHeapSize;
+
 extern int32 GResourceDescriptorHeapSize;
+extern int32 GSamplerDescriptorHeapSize;
+
 extern int32 GOnlineDescriptorHeapSize;
 extern int32 GOnlineDescriptorHeapBlockSize;
 
@@ -672,7 +676,7 @@ public:
 			}
 
 #if USE_STATIC_ROOT_SIGNATURE
-			CBVCache.CBHandles[ShaderFrequency][SlotIndex] = UniformBuffer->View->GetView();
+			CBVCache.CBHandles[ShaderFrequency][SlotIndex] = UniformBuffer->View->GetOfflineCpuHandle();
 #endif
 		}
 		else if (CurrentGPUVirtualAddress != 0)
@@ -710,7 +714,7 @@ public:
 			FD3D12ConstantBufferCache::DirtySlot(CBVCache.DirtySlotMask[ShaderFrequency], SlotIndex);
 
 #if USE_STATIC_ROOT_SIGNATURE
-			CBVCache.CBHandles[ShaderFrequency][SlotIndex] = Buffer.GetView();
+			CBVCache.CBHandles[ShaderFrequency][SlotIndex] = Buffer.GetOfflineCpuHandle();
 #endif
 		}
 	}
diff --git a/Engine/Source/Runtime/D3D12RHI/Private/D3D12Texture.h b/Engine/Source/Runtime/D3D12RHI/Private/D3D12Texture.h
--- a/Engine/Source/Runtime/D3D12RHI/Private/D3D12Texture.h
+++ b/Engine/Source/Runtime/D3D12RHI/Private/D3D12Texture.h
@@ -327,6 +327,11 @@ public:
 
 	bool IsLastMipFirst() const { return bMipOrderDescending; }
 
+	virtual FRHIShaderResourceView* GetDefaultShaderResourceView() const final
+	{
+		return GetShaderResourceView();
+	}
+
 	/** FRHITexture override.  See FRHITexture::GetNativeResource() */
 	virtual void* GetNativeResource() const override final
 	{
diff --git a/Engine/Source/Runtime/D3D12RHI/Private/D3D12UAV.cpp b/Engine/Source/Runtime/D3D12RHI/Private/D3D12UAV.cpp
--- a/Engine/Source/Runtime/D3D12RHI/Private/D3D12UAV.cpp
+++ b/Engine/Source/Runtime/D3D12RHI/Private/D3D12UAV.cpp
@@ -230,7 +230,7 @@ void FD3D12CommandContext::ClearUAV(TRHICommandList_RecursiveHazardous<FD3D12Com
 				}
 
 				// Scoped descriptor handle will free the offline CPU handle once we return
-				FD3D12ViewDescriptorHandle UAVHandle(ParentDevice, ED3D12DescriptorHeapType::Standard);
+				FD3D12ViewDescriptorHandle UAVHandle(ParentDevice, ERHIDescriptorHeapType::Standard);
 				UAVHandle.CreateView(R32UAVDesc, Resource, nullptr);
 
 				// Check if the view heap is full and needs to rollover.
@@ -240,7 +240,7 @@ void FD3D12CommandContext::ClearUAV(TRHICommandList_RecursiveHazardous<FD3D12Com
 				}
 
 				uint32 ReservedSlot = Context.StateCache.GetDescriptorCache()->GetCurrentViewHeap()->ReserveSlots(1);
-				D3D12_CPU_DESCRIPTOR_HANDLE CPUHandle = UAVHandle.GetOfflineHandle();
+				D3D12_CPU_DESCRIPTOR_HANDLE CPUHandle = UAVHandle.GetOfflineCpuHandle();
 				D3D12_CPU_DESCRIPTOR_HANDLE DestSlot = Context.StateCache.GetDescriptorCache()->GetCurrentViewHeap()->GetCPUSlotHandle(ReservedSlot);
 				D3D12_GPU_DESCRIPTOR_HANDLE GPUHandle = Context.StateCache.GetDescriptorCache()->GetCurrentViewHeap()->GetGPUSlotHandle(ReservedSlot);
 
diff --git a/Engine/Source/Runtime/D3D12RHI/Private/D3D12View.cpp b/Engine/Source/Runtime/D3D12RHI/Private/D3D12View.cpp
--- a/Engine/Source/Runtime/D3D12RHI/Private/D3D12View.cpp
+++ b/Engine/Source/Runtime/D3D12RHI/Private/D3D12View.cpp
@@ -5,7 +5,7 @@
 ////////////////////////////////////////////////////////////////////////////////////////////////////////////////////////
 // FD3D12ViewDescriptorHandle
 
-FD3D12ViewDescriptorHandle::FD3D12ViewDescriptorHandle(FD3D12Device* InParentDevice, ED3D12DescriptorHeapType InHeapType)
+FD3D12ViewDescriptorHandle::FD3D12ViewDescriptorHandle(FD3D12Device* InParentDevice, ERHIDescriptorHeapType InHeapType)
 	: FD3D12DeviceChild(InParentDevice)
 	, HeapType(InHeapType)
 {
@@ -15,35 +15,35 @@ FD3D12ViewDescriptorHandle::FD3D12ViewDescriptorHandle(FD3D12Device* InParentDev
 FD3D12ViewDescriptorHandle::~FD3D12ViewDescriptorHandle()
 {
 	FreeDescriptorSlot();
-	check(OfflineHandle.ptr == 0);
-	check(OfflineIndex == UINT_MAX);
-	check(ResourceIndex == UINT_MAX);
+	check(OfflineCpuHandle.ptr == 0);
+	check(OfflineHeapIndex == UINT_MAX);
+	check(!BindlessHandle.IsValid());
 }
 
 void FD3D12ViewDescriptorHandle::CreateView(const D3D12_RENDER_TARGET_VIEW_DESC& Desc, ID3D12Resource* Resource)
 {
-	check(HeapType == ED3D12DescriptorHeapType::RenderTarget);
+	check(HeapType == ERHIDescriptorHeapType::RenderTarget);
 
-	GetParentDevice()->GetDevice()->CreateRenderTargetView(Resource, &Desc, OfflineHandle);
+	GetParentDevice()->GetDevice()->CreateRenderTargetView(Resource, &Desc, OfflineCpuHandle);
 }
 
 void FD3D12ViewDescriptorHandle::CreateView(const D3D12_DEPTH_STENCIL_VIEW_DESC& Desc, ID3D12Resource* Resource)
 {
-	check(HeapType == ED3D12DescriptorHeapType::DepthStencil);
+	check(HeapType == ERHIDescriptorHeapType::DepthStencil);
 
-	GetParentDevice()->GetDevice()->CreateDepthStencilView(Resource, &Desc, OfflineHandle);
+	GetParentDevice()->GetDevice()->CreateDepthStencilView(Resource, &Desc, OfflineCpuHandle);
 }
 
 void FD3D12ViewDescriptorHandle::CreateView(const D3D12_CONSTANT_BUFFER_VIEW_DESC& Desc)
 {
-	check(HeapType == ED3D12DescriptorHeapType::Standard);
+	check(HeapType == ERHIDescriptorHeapType::Standard);
 
-	GetParentDevice()->GetDevice()->CreateConstantBufferView(&Desc, OfflineHandle);
+	GetParentDevice()->GetDevice()->CreateConstantBufferView(&Desc, OfflineCpuHandle);
 }
 
 void FD3D12ViewDescriptorHandle::CreateView(const D3D12_SHADER_RESOURCE_VIEW_DESC& Desc, ID3D12Resource* Resource)
 {
-	check(HeapType == ED3D12DescriptorHeapType::Standard);
+	check(HeapType == ERHIDescriptorHeapType::Standard);
 
 #if D3D12_RHI_RAYTRACING
 	// NOTE (from D3D Debug runtime): When ViewDimension is D3D12_SRV_DIMENSION_RAYTRACING_ACCELLERATION_STRUCTURE, pResource must be NULL, since the resource location comes from a GPUVA in pDesc
@@ -53,41 +53,36 @@ void FD3D12ViewDescriptorHandle::CreateView(const D3D12_SHADER_RESOURCE_VIEW_DES
 	}
 #endif // D3D12_RHI_RAYTRACING
 
-	GetParentDevice()->GetDevice()->CreateShaderResourceView(Resource, &Desc, OfflineHandle);
+	GetParentDevice()->GetDevice()->CreateShaderResourceView(Resource, &Desc, OfflineCpuHandle);
 
-	if (IsResourceHandle())
+	if (IsBindless())
 	{
-		GetParentDevice()->GetDevice()->CopyDescriptorsSimple(1, GetResourceHandle(), GetOfflineHandle(), D3D12_DESCRIPTOR_HEAP_TYPE_CBV_SRV_UAV);
+		GetParentDevice()->GetDevice()->CopyDescriptorsSimple(1, GetBindlessCpuDescriptorHandle(), GetOfflineCpuHandle(), D3D12_DESCRIPTOR_HEAP_TYPE_CBV_SRV_UAV);
 	}
 }
 
 void FD3D12ViewDescriptorHandle::CreateView(const D3D12_UNORDERED_ACCESS_VIEW_DESC& Desc, ID3D12Resource* Resource, ID3D12Resource* CounterResource)
 {
-	check(HeapType == ED3D12DescriptorHeapType::Standard);
+	check(HeapType == ERHIDescriptorHeapType::Standard);
 
-	GetParentDevice()->GetDevice()->CreateUnorderedAccessView(Resource, CounterResource, &Desc, OfflineHandle);
+	GetParentDevice()->GetDevice()->CreateUnorderedAccessView(Resource, CounterResource, &Desc, OfflineCpuHandle);
 
-	if (IsResourceHandle())
+	if (IsBindless())
 	{
-		GetParentDevice()->GetDevice()->CopyDescriptorsSimple(1, GetResourceHandle(), GetOfflineHandle(), D3D12_DESCRIPTOR_HEAP_TYPE_CBV_SRV_UAV);
+		GetParentDevice()->GetDevice()->CopyDescriptorsSimple(1, GetBindlessCpuDescriptorHandle(), GetOfflineCpuHandle(), D3D12_DESCRIPTOR_HEAP_TYPE_CBV_SRV_UAV);
 	}
 }
 
-D3D12_CPU_DESCRIPTOR_HANDLE FD3D12ViewDescriptorHandle::GetResourceHandle() const
-{
-	return GetParentDevice()->GetResourceDescriptorManager().GetResourceHandle(ResourceIndex);
-}
-
 void FD3D12ViewDescriptorHandle::AllocateDescriptorSlot()
 {
 	if (FD3D12Device* Device = GetParentDevice())
 	{
-		OfflineHandle = Device->GetOfflineDescriptorManager(HeapType).AllocateHeapSlot(OfflineIndex);
-		check(OfflineHandle.ptr != 0);
+		OfflineCpuHandle = Device->GetOfflineDescriptorManager(HeapType).AllocateHeapSlot(OfflineHeapIndex);
+		check(OfflineCpuHandle.ptr != 0);
 
-		if (HeapType == ED3D12DescriptorHeapType::Standard)
+		if (HeapType == ERHIDescriptorHeapType::Standard)
 		{
-			Device->GetResourceDescriptorManager().AllocateDescriptor(ResourceIndex);
+			BindlessHandle = Device->GetBindlessDescriptorManager().AllocateDescriptor(ERHIDescriptorHeapType::Standard);
 		}
 	}
 }
@@ -96,17 +91,22 @@ void FD3D12ViewDescriptorHandle::FreeDescriptorSlot()
 {
 	if (FD3D12Device* Device = GetParentDevice())
 	{
-		Device->GetOfflineDescriptorManager(HeapType).FreeHeapSlot(OfflineHandle, OfflineIndex);
-		OfflineIndex = UINT_MAX;
-		OfflineHandle.ptr = 0;
+		Device->GetOfflineDescriptorManager(HeapType).FreeHeapSlot(OfflineCpuHandle, OfflineHeapIndex);
+		OfflineHeapIndex = UINT_MAX;
+		OfflineCpuHandle.ptr = 0;
 
-		if (IsResourceHandle())
+		if (IsBindless())
 		{
-			Device->GetResourceDescriptorManager().FreeDescriptor(ResourceIndex);
-			ResourceIndex = UINT_MAX;
+			Device->GetBindlessDescriptorManager().FreeDescriptor(BindlessHandle);
+			BindlessHandle = FRHIDescriptorHandle();
 		}
 	}
-	check(OfflineHandle.ptr == 0);
+	check(OfflineCpuHandle.ptr == 0);
+}
+
+D3D12_CPU_DESCRIPTOR_HANDLE FD3D12ViewDescriptorHandle::GetBindlessCpuDescriptorHandle() const
+{
+	return GetParentDevice()->GetBindlessDescriptorManager().GetCpuDescriptorHandle(BindlessHandle);
 }
 
 ////////////////////////////////////////////////////////////////////////////////////////////////////////////////////////
@@ -567,7 +567,7 @@ FShaderResourceViewRHIRef FD3D12DynamicRHI::RHICreateShaderResourceViewWriteMask
 
 FD3D12ConstantBufferView::FD3D12ConstantBufferView(FD3D12Device* InParent)
 	: FD3D12DeviceChild(InParent)
-	, Descriptor(InParent, ED3D12DescriptorHeapType::Standard)
+	, Descriptor(InParent, ERHIDescriptorHeapType::Standard)
 {
 }
 
diff --git a/Engine/Source/Runtime/D3D12RHI/Private/D3D12View.h b/Engine/Source/Runtime/D3D12RHI/Private/D3D12View.h
--- a/Engine/Source/Runtime/D3D12RHI/Private/D3D12View.h
+++ b/Engine/Source/Runtime/D3D12RHI/Private/D3D12View.h
@@ -623,7 +623,7 @@ class FD3D12ViewDescriptorHandle : public FD3D12DeviceChild
 {
 public:
 	FD3D12ViewDescriptorHandle() = delete;
-	FD3D12ViewDescriptorHandle(FD3D12Device* InParentDevice, ED3D12DescriptorHeapType InHeapType);
+	FD3D12ViewDescriptorHandle(FD3D12Device* InParentDevice, ERHIDescriptorHeapType InHeapType);
 	~FD3D12ViewDescriptorHandle();
 
 	void CreateView(const D3D12_RENDER_TARGET_VIEW_DESC& Desc, ID3D12Resource* Resource);
@@ -632,23 +632,23 @@ public:
 	void CreateView(const D3D12_SHADER_RESOURCE_VIEW_DESC& Desc, ID3D12Resource* Resource);
 	void CreateView(const D3D12_UNORDERED_ACCESS_VIEW_DESC& Desc, ID3D12Resource* Resource, ID3D12Resource* CounterResource);
 
-	D3D12_CPU_DESCRIPTOR_HANDLE GetResourceHandle() const;
+	inline D3D12_CPU_DESCRIPTOR_HANDLE GetOfflineCpuHandle() const { return OfflineCpuHandle; }
+	inline uint32                      GetOfflineHeapIndex() const { return OfflineHeapIndex; }
+	inline FRHIDescriptorHandle        GetBindlessHandle()   const { return BindlessHandle;   }
 
-	inline D3D12_CPU_DESCRIPTOR_HANDLE GetOfflineHandle()  const { return OfflineHandle; }
-	inline ED3D12DescriptorHeapType    GetHeapType()       const { return HeapType; }
-	inline uint32                      GetOfflineIndex()   const { return OfflineIndex; }
-	inline uint32                      GetResourceIndex()  const { return ResourceIndex; }
-
-	inline bool IsResourceHandle() const { return GetResourceIndex() != UINT_MAX; }
+	inline bool IsBindless() const { return BindlessHandle.IsValid(); }
 
 private:
 	void AllocateDescriptorSlot();
 	void FreeDescriptorSlot();
 
-	D3D12_CPU_DESCRIPTOR_HANDLE OfflineHandle{};
-	const ED3D12DescriptorHeapType HeapType{};
-	uint32 OfflineIndex{ UINT_MAX };
-	uint32 ResourceIndex{ UINT_MAX };
+	D3D12_CPU_DESCRIPTOR_HANDLE GetBindlessCpuDescriptorHandle() const;
+
+	D3D12_CPU_DESCRIPTOR_HANDLE OfflineCpuHandle{};
+	uint32 OfflineHeapIndex{ UINT_MAX };
+
+	FRHIDescriptorHandle BindlessHandle;
+	const ERHIDescriptorHeapType HeapType;
 };
 
 template <typename TDesc>
@@ -669,7 +669,7 @@ protected:
 	bool bInitialized;
 #endif
 
-	explicit FD3D12View(FD3D12Device* InParent, ED3D12DescriptorHeapType InHeapType, ViewSubresourceSubsetFlags InFlags)
+	explicit FD3D12View(FD3D12Device* InParent, ERHIDescriptorHeapType InHeapType, ViewSubresourceSubsetFlags InFlags)
 		: Descriptor(InParent, InHeapType)
 		, Flags(InFlags)
 		, BaseShaderResource(nullptr)
@@ -754,7 +754,7 @@ public:
 	inline FD3D12Device*					GetParentDevice_Unsafe()	const { return Descriptor.GetParentDevice_Unsafe(); }
 	inline FD3D12ResourceLocation*			GetResourceLocation()		const { return ResourceLocation; }
 	inline const TDesc&						GetDesc()					const { checkf(bInitialized, TEXT("Uninitialized D3D12View size %d"), (uint32)sizeof(TDesc)); return Desc; }
-	inline D3D12_CPU_DESCRIPTOR_HANDLE		GetView()					const { checkf(bInitialized, TEXT("Uninitialized D3D12View size %d"), (uint32)sizeof(TDesc)); return Descriptor.GetOfflineHandle(); }
+	inline D3D12_CPU_DESCRIPTOR_HANDLE		GetOfflineCpuHandle()		const { checkf(bInitialized, TEXT("Uninitialized D3D12View size %d"), (uint32)sizeof(TDesc)); return Descriptor.GetOfflineCpuHandle(); }
 	inline uint32							GetDescriptorHeapIndex()	const { checkf(bInitialized, TEXT("Uninitialized D3D12View size %d"), (uint32)sizeof(TDesc)); return Descriptor.GetIndex(); }
 	inline FD3D12Resource*					GetResource()				const { checkf(bInitialized, TEXT("Uninitialized D3D12View size %d"), (uint32)sizeof(TDesc)); return Resource; }
 	inline FD3D12ResidencyHandle*			GetResidencyHandle()		const { checkf(bInitialized, TEXT("Uninitialized D3D12View size %d"), (uint32)sizeof(TDesc)); return ResidencyHandle; }
@@ -778,7 +778,7 @@ class FD3D12ShaderResourceView : public FRHIShaderResourceView, public FD3D12Vie
 public:
 	// Used for dynamic buffer SRVs, which can be renamed. Must be explicitly Initialize()'d before it can be used.
 	FD3D12ShaderResourceView(FD3D12Device* InParent)
-		: FD3D12View(InParent, ED3D12DescriptorHeapType::Standard, ViewSubresourceSubsetFlags_None)
+		: FD3D12View(InParent, ERHIDescriptorHeapType::Standard, ViewSubresourceSubsetFlags_None)
 		, bContainsDepthPlane(false)
 		, bContainsStencilPlane(false)
 		, bSkipFastClearFinalize(false)
@@ -866,6 +866,8 @@ public:
 	FORCEINLINE bool GetSkipFastClearFinalize()	const { return bSkipFastClearFinalize; }
 	FORCEINLINE bool RequiresResourceStateTracking() const { return bRequiresResourceStateTracking; }
 
+	virtual FRHIDescriptorHandle GetBindlessHandle() const override { return Descriptor.GetBindlessHandle(); }
+
 protected:
 	void PreCreateView(FD3D12ResourceLocation& InResourceLocation, uint32 InStride, uint32 InStartOffsetBytes, bool InSkipFastClearFinalize)
 	{
@@ -944,12 +946,12 @@ class FD3D12UnorderedAccessView : public FRHIUnorderedAccessView, public FD3D12V
 {
 public:
 	FD3D12UnorderedAccessView(FD3D12Device* InParent)
-	: FD3D12View(InParent, ED3D12DescriptorHeapType::Standard, ViewSubresourceSubsetFlags_None)
+	: FD3D12View(InParent, ERHIDescriptorHeapType::Standard, ViewSubresourceSubsetFlags_None)
 	{
 	}
 
 	FD3D12UnorderedAccessView(FD3D12Device* InParent, D3D12_UNORDERED_ACCESS_VIEW_DESC& InDesc, FD3D12BaseShaderResource* InBaseShaderResource, FD3D12Resource* InCounterResource = nullptr)
-		: FD3D12View(InParent, ED3D12DescriptorHeapType::Standard, ViewSubresourceSubsetFlags_None)
+		: FD3D12View(InParent, ERHIDescriptorHeapType::Standard, ViewSubresourceSubsetFlags_None)
 		, CounterResource(InCounterResource)
 	{
 		Initialize(InDesc, InBaseShaderResource, InBaseShaderResource->ResourceLocation, InCounterResource);
@@ -973,6 +975,8 @@ public:
 
 	FD3D12Resource* GetCounterResource() { return CounterResource; }
 
+	virtual FRHIDescriptorHandle GetBindlessHandle() const override { return Descriptor.GetBindlessHandle(); }
+
 protected:
 	void CreateView(FD3D12BaseShaderResource* InBaseShaderResource, FD3D12ResourceLocation& InResourceLocation, FD3D12Resource* InCounterResource)
 	{
@@ -1014,7 +1018,7 @@ public:
 	FD3D12ConstantBufferView() = delete;
 	FD3D12ConstantBufferView(FD3D12Device* InParent);
 
-	inline D3D12_CPU_DESCRIPTOR_HANDLE            GetView() const { return Descriptor.GetOfflineHandle(); }
+	inline D3D12_CPU_DESCRIPTOR_HANDLE            GetOfflineCpuHandle() const { return Descriptor.GetOfflineCpuHandle(); }
 	inline const D3D12_CONSTANT_BUFFER_VIEW_DESC& GetDesc() const { return Desc; }
 
 	void Create(D3D12_GPU_VIRTUAL_ADDRESS GPUAddress, const uint32 AlignedSize);
@@ -1028,14 +1032,14 @@ class FD3D12RenderTargetView : public FD3D12View<D3D12_RENDER_TARGET_VIEW_DESC>,
 {
 public:
 	FD3D12RenderTargetView(FD3D12Device* InParent, const D3D12_RENDER_TARGET_VIEW_DESC& InRTVDesc, FD3D12BaseShaderResource* InBaseShaderResource)
-		: FD3D12View(InParent, ED3D12DescriptorHeapType::RenderTarget, ViewSubresourceSubsetFlags_None)
+		: FD3D12View(InParent, ERHIDescriptorHeapType::RenderTarget, ViewSubresourceSubsetFlags_None)
 		, FRHIResource(RRT_None)
 	{
 		Initialize(InRTVDesc, InBaseShaderResource, InBaseShaderResource->ResourceLocation);
 	}
 
 	FD3D12RenderTargetView(FD3D12Device* InParent, const D3D12_RENDER_TARGET_VIEW_DESC& InRTVDesc, FD3D12BaseShaderResource* InBaseShaderResource, FD3D12ResourceLocation& InResourceLocation)
-		: FD3D12View(InParent, ED3D12DescriptorHeapType::RenderTarget, ViewSubresourceSubsetFlags_None)
+		: FD3D12View(InParent, ERHIDescriptorHeapType::RenderTarget, ViewSubresourceSubsetFlags_None)
 		, FRHIResource(RRT_None)
 	{
 		Initialize(InRTVDesc, InBaseShaderResource, InResourceLocation);
@@ -1053,6 +1057,8 @@ public:
 		CreateView(BaseShaderResource, BaseShaderResource->ResourceLocation);
 	}
 
+	inline D3D12_CPU_DESCRIPTOR_HANDLE GetView() const { return GetOfflineCpuHandle(); }
+
 protected:
 	void CreateView(FD3D12BaseShaderResource* InBaseShaderResource, FD3D12ResourceLocation& InResourceLocation)
 	{
@@ -1070,7 +1076,7 @@ class FD3D12DepthStencilView : public FD3D12View<D3D12_DEPTH_STENCIL_VIEW_DESC>,
 {
 public:
 	FD3D12DepthStencilView(FD3D12Device* InParent, const D3D12_DEPTH_STENCIL_VIEW_DESC& InDSVDesc, FD3D12BaseShaderResource* InBaseShaderResource, bool InHasStencil)
-		: FD3D12View(InParent, ED3D12DescriptorHeapType::DepthStencil, ViewSubresourceSubsetFlags_DepthAndStencilDsv)
+		: FD3D12View(InParent, ERHIDescriptorHeapType::DepthStencil, ViewSubresourceSubsetFlags_DepthAndStencilDsv)
 		, FRHIResource(RRT_None)
 		, bHasDepth(true)				// Assume all DSVs have depth bits in their format
 		, bHasStencil(InHasStencil)		// Only some DSVs have stencil bits in their format
@@ -1139,6 +1145,8 @@ public:
 		SetupDepthStencilViewSubresourceSubset();
 	}
 
+	inline D3D12_CPU_DESCRIPTOR_HANDLE GetView() const { return GetOfflineCpuHandle(); }
+
 protected:
 	void CreateView(FD3D12BaseShaderResource* InBaseShaderResource, FD3D12ResourceLocation& InResourceLocation)
 	{
diff --git a/Engine/Source/Runtime/D3D12RHI/Public/D3D12ConstantBuffer.h b/Engine/Source/Runtime/D3D12RHI/Public/D3D12ConstantBuffer.h
--- a/Engine/Source/Runtime/D3D12RHI/Public/D3D12ConstantBuffer.h
+++ b/Engine/Source/Runtime/D3D12RHI/Public/D3D12ConstantBuffer.h
@@ -43,7 +43,7 @@ public:
 	bool Version(FD3D12ResourceLocation& BufferOut, bool bDiscardSharedConstants);
 
 #if USE_STATIC_ROOT_SIGNATURE
-	inline D3D12_CPU_DESCRIPTOR_HANDLE GetView() const { return View->GetView(); }
+	inline D3D12_CPU_DESCRIPTOR_HANDLE GetOfflineCpuHandle() const { return View->GetOfflineCpuHandle(); }
 #endif
 
 protected:
diff --git a/Engine/Source/Runtime/D3D12RHI/Public/D3D12State.h b/Engine/Source/Runtime/D3D12RHI/Public/D3D12State.h
--- a/Engine/Source/Runtime/D3D12RHI/Public/D3D12State.h
+++ b/Engine/Source/Runtime/D3D12RHI/Public/D3D12State.h
@@ -9,12 +9,16 @@
 class FD3D12SamplerState : public FRHISamplerState, public FD3D12DeviceChild, public FD3D12LinkedAdapterObject<FD3D12SamplerState>
 {
 public:
-	D3D12_CPU_DESCRIPTOR_HANDLE Descriptor;
-	uint32 DescriptorHeapIndex;
+	D3D12_CPU_DESCRIPTOR_HANDLE OfflineHandle{};
+	uint32 OfflineIndex{};
+	FRHIDescriptorHandle BindlessHandle;
 	const uint16 ID;
 
+	FD3D12SamplerState() = delete;
 	FD3D12SamplerState(FD3D12Device* InParent, const D3D12_SAMPLER_DESC& Desc, uint16 SamplerID);
 	~FD3D12SamplerState();
+
+	virtual FRHIDescriptorHandle GetBindlessHandle() const final { return BindlessHandle; }
 };
 
 class FD3D12RasterizerState : public FRHIRasterizerState
diff --git a/Engine/Source/Runtime/Engine/Private/ShaderCompiler/ShaderCompiler.cpp b/Engine/Source/Runtime/Engine/Private/ShaderCompiler/ShaderCompiler.cpp
--- a/Engine/Source/Runtime/Engine/Private/ShaderCompiler/ShaderCompiler.cpp
+++ b/Engine/Source/Runtime/Engine/Private/ShaderCompiler/ShaderCompiler.cpp
@@ -5791,6 +5791,7 @@ void GlobalBeginCompileShader(
 	Input.Environment.SetDefine(TEXT("PLATFORM_SUPPORTS_DISTANCE_FIELDS"), DoesPlatformSupportDistanceFields(EShaderPlatform(Target.Platform)) ? 1 : 0);
 	Input.Environment.SetDefine(TEXT("PLATFORM_SUPPORTS_MESH_SHADERS_TIER0"), RHISupportsMeshShadersTier0(EShaderPlatform(Target.Platform)) ? 1 : 0);
 	Input.Environment.SetDefine(TEXT("PLATFORM_SUPPORTS_MESH_SHADERS_TIER1"), RHISupportsMeshShadersTier1(EShaderPlatform(Target.Platform)) ? 1 : 0);
+	Input.Environment.SetDefine(TEXT("PLATFORM_SUPPORTS_BINDLESS"), RHISupportsBindless(EShaderPlatform(Target.Platform)) ? 1 : 0);
 
 	{
 		static const auto CVar = IConsoleManager::Get().FindConsoleVariable(TEXT("r.VT.AnisotropicFiltering"));
diff --git a/Engine/Source/Runtime/RHI/Private/RHI.cpp b/Engine/Source/Runtime/RHI/Private/RHI.cpp
--- a/Engine/Source/Runtime/RHI/Private/RHI.cpp
+++ b/Engine/Source/Runtime/RHI/Private/RHI.cpp
@@ -2135,6 +2135,18 @@ FString LexToString(ERHIFeatureLevel::Type Level)
 	return TEXT("UnknownFeatureLevel");
 }
 
+const TCHAR* LexToString(ERHIDescriptorHeapType InHeapType)
+{
+	switch (InHeapType)
+	{
+	default: checkNoEntry();                   return TEXT("UnknownDescriptorHeapType");
+	case ERHIDescriptorHeapType::Standard:     return TEXT("Standard");
+	case ERHIDescriptorHeapType::RenderTarget: return TEXT("RenderTarget");
+	case ERHIDescriptorHeapType::DepthStencil: return TEXT("DepthStencil");
+	case ERHIDescriptorHeapType::Sampler:      return TEXT("Sampler");
+	}
+}
+
 const FName LANGUAGE_D3D("D3D");
 const FName LANGUAGE_Metal("Metal");
 const FName LANGUAGE_OpenGL("OpenGL");
@@ -2270,6 +2282,7 @@ void FGenericDataDrivenShaderPlatformInfo::ParseDataDrivenShaderInfo(const FConf
 	GET_SECTION_BOOL_HELPER(bOverrideFMaterial_NeedsGBufferEnabled);
 	GET_SECTION_BOOL_HELPER(bSupportsMobileDistanceField);
 	GET_SECTION_BOOL_HELPER(bSupportsFFTBloom);
+	GET_SECTION_BOOL_HELPER(bSupportsBindless);
 #undef GET_SECTION_BOOL_HELPER
 #undef GET_SECTION_INT_HELPER
 
diff --git a/Engine/Source/Runtime/RHI/Public/RHI.h b/Engine/Source/Runtime/RHI/Public/RHI.h
--- a/Engine/Source/Runtime/RHI/Public/RHI.h
+++ b/Engine/Source/Runtime/RHI/Public/RHI.h
@@ -257,6 +257,12 @@ inline RHI_API bool RHISupportsConservativeRasterization(EShaderPlatform Platfor
 	return FDataDrivenShaderPlatformInfo::GetSupportsConservativeRasterization(Platform);
 }
 
+/** True if the given shader platform supports bindless resources/views. */
+inline bool RHISupportsBindless(EShaderPlatform Platform)
+{
+	return FDataDrivenShaderPlatformInfo::GetSupportsBindless(Platform);
+}
+
 // Wrapper for GRHI## global variables, allows values to be overridden for mobile preview modes.
 template <typename TValueType>
 class TRHIGlobal
@@ -850,6 +856,9 @@ extern RHI_API FString LexToString(EShaderPlatform Platform);
 /** Stringifies ERHIFeatureLevel */
 extern RHI_API FString LexToString(ERHIFeatureLevel::Type Level);
 
+/** Stringifies ERHIDescriptorHeapType */
+extern RHI_API const TCHAR* LexToString(ERHIDescriptorHeapType InHeapType);
+
 /** Finds a corresponding ERHIShadingPath::Type given an FName, or returns false if one could not be found. */
 extern RHI_API bool GetShadingPathFromName(FName Name, ERHIShadingPath::Type& OutShadingPath);
 
diff --git a/Engine/Source/Runtime/RHI/Public/RHIDefinitions.h b/Engine/Source/Runtime/RHI/Public/RHIDefinitions.h
--- a/Engine/Source/Runtime/RHI/Public/RHIDefinitions.h
+++ b/Engine/Source/Runtime/RHI/Public/RHIDefinitions.h
@@ -375,6 +375,7 @@ class RHI_API FGenericDataDrivenShaderPlatformInfo
 	uint32 bOverrideFMaterial_NeedsGBufferEnabled : 1;
 	uint32 bSupportsMobileDistanceField : 1;
 	uint32 bSupportsFFTBloom : 1;
+	uint32 bSupportsBindless : 1;
 
 		
 #if WITH_EDITOR
@@ -772,6 +773,11 @@ public:
 		return Infos[Platform].bSupportsFFTBloom;
 	}
 
+	static FORCEINLINE_DEBUGGABLE const bool GetSupportsBindless(const FStaticShaderPlatform Platform)
+	{
+		return Infos[Platform].bSupportsBindless;
+	}
+
 #if WITH_EDITOR
 	static FORCEINLINE_DEBUGGABLE FText GetFriendlyName(const FStaticShaderPlatform Platform)
 	{
@@ -1773,6 +1779,34 @@ enum class EAsyncComputeBudget
 	EAll_4,				//Async compute can use the entire GPU.
 };
 
+enum class ERHIDescriptorHeapType : uint8
+{
+	Standard,
+	Sampler,
+	RenderTarget,
+	DepthStencil,
+	count
+};
+
+struct FRHIDescriptorHandle
+{
+	FRHIDescriptorHandle() = default;
+	FRHIDescriptorHandle(ERHIDescriptorHeapType InType, uint32 InIndex)
+		: Index(InIndex)
+		, Type(InType)
+	{
+	}
+
+	inline uint32                GetIndex() const { return Index; }
+	inline ERHIDescriptorHeapType GetType() const { return Type; }
+
+	inline bool IsValid() const { return Index != UINT_MAX && Type != ERHIDescriptorHeapType::count; }
+
+private:
+	uint32                 Index{ UINT_MAX };
+	ERHIDescriptorHeapType Type{ ERHIDescriptorHeapType::count };
+};
+
 inline bool IsPCPlatform(const FStaticShaderPlatform Platform)
 {
 	return Platform == SP_PCD3D_SM5 ||
diff --git a/Engine/Source/Runtime/RHI/Public/RHIResources.h b/Engine/Source/Runtime/RHI/Public/RHIResources.h
--- a/Engine/Source/Runtime/RHI/Public/RHIResources.h
+++ b/Engine/Source/Runtime/RHI/Public/RHIResources.h
@@ -600,6 +600,7 @@ class FRHISamplerState : public FRHIResource
 public:
 	FRHISamplerState() : FRHIResource(RRT_SamplerState) {}
 	virtual bool IsImmutable() const { return false; }
+	virtual FRHIDescriptorHandle GetBindlessHandle() const { return FRHIDescriptorHandle(); }
 };
 
 class FRHIRasterizerState : public FRHIResource
@@ -1245,7 +1246,9 @@ public:
 	virtual class FRHITexture3D* GetTexture3D() { return NULL; }
 	virtual class FRHITextureCube* GetTextureCube() { return NULL; }
 	virtual class FRHITextureReference* GetTextureReference() { return NULL; }
-	
+
+	virtual class FRHIShaderResourceView* GetDefaultShaderResourceView() const { return nullptr; }
+
 	// Slower method to get Size X, Y & Z information. Prefer sub-classes' GetSizeX(), etc
 	virtual FIntVector GetSizeXYZ() const = 0;
 
@@ -1841,6 +1844,7 @@ class FRHIUnorderedAccessView : public FRHIResource
 {
 public:
 	FRHIUnorderedAccessView() : FRHIResource(RRT_UnorderedAccessView) {}
+	virtual FRHIDescriptorHandle GetBindlessHandle() const { return FRHIDescriptorHandle(); }
 };
 
 class FRHIShaderResourceView : public FRHIResource 
@@ -1850,6 +1854,7 @@ class FRHIShaderResourceView : public FRHIResource
 {
 public:
 	FRHIShaderResourceView() : FRHIResource(RRT_ShaderResourceView) {}
+	virtual FRHIDescriptorHandle GetBindlessHandle() const { return FRHIDescriptorHandle(); }
 };
 
 
diff --git a/Engine/Source/Runtime/RenderCore/Private/Shader.cpp b/Engine/Source/Runtime/RenderCore/Private/Shader.cpp
--- a/Engine/Source/Runtime/RenderCore/Private/Shader.cpp
+++ b/Engine/Source/Runtime/RenderCore/Private/Shader.cpp
@@ -596,88 +596,84 @@ void FShader::Finalize(const FShaderMapResourceCode* Code)
 
 void FShader::BuildParameterMapInfo(const TMap<FString, FParameterAllocation>& ParameterMap)
 {
-	for (int32 ParameterTypeIndex = 0; ParameterTypeIndex < (int32)EShaderParameterType::Num; ParameterTypeIndex++)
+	uint32 UniformCount = 0;
+	uint32 SamplerCount = 0;
+	uint32 SRVCount = 0;
+
+	for (TMap<FString, FParameterAllocation>::TConstIterator ParameterIt(ParameterMap); ParameterIt; ++ParameterIt)
 	{
-		EShaderParameterType CurrentParameterType = (EShaderParameterType)ParameterTypeIndex;
+		const FParameterAllocation& ParamValue = ParameterIt.Value();
 
-		if (CurrentParameterType == EShaderParameterType::LooseData)
+		switch (ParamValue.Type)
 		{
-			for (TMap<FString, FParameterAllocation>::TConstIterator ParameterIt(ParameterMap); ParameterIt; ++ParameterIt)
-			{
-				const FParameterAllocation& ParamValue = ParameterIt.Value();
-
-				if (ParamValue.Type == CurrentParameterType)
-				{
-					bool bAddedToExistingBuffer = false;
+		case EShaderParameterType::UniformBuffer:
+			UniformCount++;
+			break;
+		case EShaderParameterType::Sampler:
+			SamplerCount++;
+			break;
+		case EShaderParameterType::SRV:
+			SRVCount++;
+			break;
+		}
+	}
 
-					for (int32 LooseParameterBufferIndex = 0; LooseParameterBufferIndex < ParameterMapInfo.LooseParameterBuffers.Num(); LooseParameterBufferIndex++)
-					{
-						FShaderLooseParameterBufferInfo& LooseParameterBufferInfo = ParameterMapInfo.LooseParameterBuffers[LooseParameterBufferIndex];
-
-						if (LooseParameterBufferInfo.BaseIndex == ParamValue.BufferIndex)
-						{
-							FShaderParameterInfo ParameterInfo(ParamValue.BaseIndex, ParamValue.Size);
-							LooseParameterBufferInfo.Parameters.Add(ParameterInfo);
-							LooseParameterBufferInfo.Size += ParamValue.Size;
-							bAddedToExistingBuffer = true;
-						}
-					}
+	ParameterMapInfo.UniformBuffers.Empty(UniformCount);
+	ParameterMapInfo.TextureSamplers.Empty(SamplerCount);
+	ParameterMapInfo.SRVs.Empty(SRVCount);
 
-					if (!bAddedToExistingBuffer)
-					{
-						FShaderLooseParameterBufferInfo NewParameterBufferInfo(ParamValue.BufferIndex, ParamValue.Size);
+	auto GetParameterMap = [this](EShaderParameterType ParameterType) -> TMemoryImageArray<FShaderParameterInfo>*
+	{
+		switch (ParameterType)
+		{
+		case EShaderParameterType::UniformBuffer:
+			return &ParameterMapInfo.UniformBuffers;
+		case EShaderParameterType::Sampler:
+			return &ParameterMapInfo.TextureSamplers;
+		case EShaderParameterType::SRV:
+			return &ParameterMapInfo.SRVs;
+		default:
+			return nullptr;
+		}
+	};
 
-						FShaderParameterInfo ParameterInfo(ParamValue.BaseIndex, ParamValue.Size);
-						NewParameterBufferInfo.Parameters.Add(ParameterInfo);
+	for (TMap<FString, FParameterAllocation>::TConstIterator ParameterIt(ParameterMap); ParameterIt; ++ParameterIt)
+	{
+		const FParameterAllocation& ParamValue = ParameterIt.Value();
 
-						ParameterMapInfo.LooseParameterBuffers.Add(NewParameterBufferInfo);
-					}
-				}
-			}
-		}
-		else if (CurrentParameterType != EShaderParameterType::UAV)
+		if (ParamValue.Type == EShaderParameterType::LooseData)
 		{
-			int32 NumParameters = 0;
+			bool bAddedToExistingBuffer = false;
 
-			for (TMap<FString, FParameterAllocation>::TConstIterator ParameterIt(ParameterMap); ParameterIt; ++ParameterIt)
+			for (int32 LooseParameterBufferIndex = 0; LooseParameterBufferIndex < ParameterMapInfo.LooseParameterBuffers.Num(); LooseParameterBufferIndex++)
 			{
-				const FParameterAllocation& ParamValue = ParameterIt.Value();
+				FShaderLooseParameterBufferInfo& LooseParameterBufferInfo = ParameterMapInfo.LooseParameterBuffers[LooseParameterBufferIndex];
 
-				if (ParamValue.Type == CurrentParameterType)
+				if (LooseParameterBufferInfo.BaseIndex == ParamValue.BufferIndex)
 				{
-					NumParameters++;
+					FShaderParameterInfo ParameterInfo(ParamValue.BaseIndex, ParamValue.Size);
+					LooseParameterBufferInfo.Parameters.Add(ParameterInfo);
+					LooseParameterBufferInfo.Size += ParamValue.Size;
+					bAddedToExistingBuffer = true;
 				}
 			}
 
-			TMemoryImageArray<FShaderParameterInfo>* ParameterInfoArray = &ParameterMapInfo.UniformBuffers;
-
-			if (CurrentParameterType == EShaderParameterType::Sampler)
-			{
-				ParameterInfoArray = &ParameterMapInfo.TextureSamplers;
-			}
-			else if (CurrentParameterType == EShaderParameterType::SRV)
-			{
-				ParameterInfoArray = &ParameterMapInfo.SRVs;
-			}
-			else
+			if (!bAddedToExistingBuffer)
 			{
-				check(CurrentParameterType == EShaderParameterType::UniformBuffer);
-			}
+				FShaderLooseParameterBufferInfo NewParameterBufferInfo(ParamValue.BufferIndex, ParamValue.Size);
 
-			ParameterInfoArray->Empty(NumParameters);
-		
-			for (TMap<FString, FParameterAllocation>::TConstIterator ParameterIt(ParameterMap); ParameterIt; ++ParameterIt)
-			{
-				const FParameterAllocation& ParamValue = ParameterIt.Value();
+				FShaderParameterInfo ParameterInfo(ParamValue.BaseIndex, ParamValue.Size);
+				NewParameterBufferInfo.Parameters.Add(ParameterInfo);
 
-				if (ParamValue.Type == CurrentParameterType)
-				{
-					const uint16 BaseIndex = CurrentParameterType == EShaderParameterType::UniformBuffer ? ParamValue.BufferIndex : ParamValue.BaseIndex;
-					FShaderParameterInfo ParameterInfo(BaseIndex, ParamValue.Size);
-					ParameterInfoArray->Add(ParameterInfo);
-				}
+				ParameterMapInfo.LooseParameterBuffers.Add(NewParameterBufferInfo);
 			}
 		}
+		else if (TMemoryImageArray<FShaderParameterInfo>* ParameterInfoArray = GetParameterMap(ParamValue.Type))
+		{
+			const uint16 BaseIndex = ParamValue.Type == EShaderParameterType::UniformBuffer ? ParamValue.BufferIndex : ParamValue.BaseIndex;
+			FShaderParameterInfo ParameterInfo(BaseIndex, ParamValue.Size);
+			ParameterInfoArray->Add(ParameterInfo);
+		}
 	}
 
 	for (FShaderLooseParameterBufferInfo& Info : ParameterMapInfo.LooseParameterBuffers)
