Reduce ray tracing sampler heap allocation high water mark through better sampler descriptor table deduplication

Samplers were previously deduplicated using per-thread hash tables, which results in some cross-thread duplicates. This change adds an exhaustive brute force search through the entire global sampler table if per-thread cache did not find a match.

Tested by running a replay which previously reached 1400+ samplers (out of hard limit of 2048). This change reduces the maximum sampler allocations to 400 in this case. Performance impact is minimal.

New CVar r.D3D12.RayTracing.DeduplicateSamplers (default=1) can be used to control the new deduplication scheme. Value 0 restores old behavior.

diff --git a/Engine/Source/Runtime/D3D12RHI/Private/D3D12RayTracing.cpp b/Engine/Source/Runtime/D3D12RHI/Private/D3D12RayTracing.cpp
--- a/Engine/Source/Runtime/D3D12RHI/Private/D3D12RayTracing.cpp
+++ b/Engine/Source/Runtime/D3D12RHI/Private/D3D12RayTracing.cpp
@@ -80,6 +80,13 @@ static FAutoConsoleVariableRef CVarRayTracingAllowSpecializedStateObjects(
 	TEXT("This is intended for performance testingand has no effect if r.D3D12.RayTracing.SpecializeStateObjects is 0. (default = 1)\n")
 );
 
+static int32 GD3D12RayTracingDeduplicateSamplers = 1;
+static FAutoConsoleVariableRef CVarD3D12RayTracingDeduplicateSamplers(
+	TEXT("r.D3D12.RayTracing.DeduplicateSamplers"),
+	GD3D12RayTracingDeduplicateSamplers,
+	TEXT("Use an exhaustive search to deduplicate sampler descriptors when generating shader binding tables. Reduces sampler heap usage at the cost of some CPU time. (default = 1)")
+);
+
 // Ray tracing stat counters
 
 DECLARE_STATS_GROUP(TEXT("D3D12RHI: Ray Tracing"), STATGROUP_D3D12RayTracing, STATCAT_Advanced);
@@ -97,6 +104,7 @@ DECLARE_DWORD_ACCUMULATOR_STAT(TEXT("Allocated sampler descriptors"), STAT_D3D12
 DECLARE_DWORD_ACCUMULATOR_STAT(TEXT("Allocated view descriptor heaps"), STAT_D3D12RayTracingViewDescriptorHeaps, STATGROUP_D3D12RayTracing);
 DECLARE_DWORD_ACCUMULATOR_STAT(TEXT("Allocated view descriptors"), STAT_D3D12RayTracingViewDescriptors, STATGROUP_D3D12RayTracing);
 
+DECLARE_DWORD_COUNTER_STAT(TEXT("Max used sampler descriptors in a single heap"), STAT_D3D12RayTracingMaxUsedSamplerDescriptors, STATGROUP_D3D12RayTracing);
 DECLARE_DWORD_COUNTER_STAT(TEXT("Used sampler descriptors (per frame)"), STAT_D3D12RayTracingUsedSamplerDescriptors, STATGROUP_D3D12RayTracing);
 DECLARE_DWORD_COUNTER_STAT(TEXT("Used view descriptors (per frame)"), STAT_D3D12RayTracingUsedViewDescriptors, STATGROUP_D3D12RayTracing);
 
@@ -1438,6 +1446,9 @@ struct FD3D12RayTracingDescriptorHeap : public FD3D12DeviceChild
 
 	int32 NumAllocatedDescriptors = 0;
 
+	// Marks the valid range of the heap when exhaustive sampler deduplication is enabled. Not used otherwise.
+	int32 NumWrittenSamplerDescriptors = 0;
+
 	uint32 DescriptorSize = 0;
 	D3D12_CPU_DESCRIPTOR_HANDLE CPUBase = {};
 	D3D12_GPU_DESCRIPTOR_HANDLE GPUBase = {};
@@ -1514,6 +1525,33 @@ public:
 			}
 		}
 
+		const bool bExhaustiveSamplerDeduplication = GD3D12RayTracingDeduplicateSamplers == 1;
+
+		if (bExhaustiveSamplerDeduplication && Type == D3D12_DESCRIPTOR_HEAP_TYPE_SAMPLER)
+		{
+			// Exhaustive search for a sampler table:
+			// We have to do this because sampler heap space is precious (hard limit of 2048 total entries).
+			// Per-thread descriptor table deduplication hash tables introduce a lot of redundancy in the heap
+			// which is reduced by looking for global matches on hash table lookup miss.
+
+			int32 FoundIndex = InvalidIndex;
+			int32 SearchEndPos = Heap.NumWrittenSamplerDescriptors;
+			for (int32 SearchIndex = 0; SearchIndex + int32(NumDescriptors) < SearchEndPos; ++SearchIndex)
+			{
+				if (Heap.CompareDescriptors(SearchIndex, Descriptors, NumDescriptors))
+				{
+					FoundIndex = SearchIndex;
+					break;
+				}
+			}
+
+			if (FoundIndex != InvalidIndex)
+			{
+				DescriptorTableBaseIndex = FoundIndex;
+				return DescriptorTableBaseIndex;
+			}
+		}
+
 		DescriptorTableBaseIndex = Heap.Allocate(NumDescriptors);
 
 		if (DescriptorTableBaseIndex == InvalidIndex)
@@ -1523,6 +1561,11 @@ public:
 
 		Heap.CopyDescriptors(DescriptorTableBaseIndex, Descriptors, NumDescriptors);
 
+		if (bExhaustiveSamplerDeduplication && Type == D3D12_DESCRIPTOR_HEAP_TYPE_SAMPLER)
+		{
+			FPlatformAtomics::InterlockedAdd(&Heap.NumWrittenSamplerDescriptors, NumDescriptors);
+		}
+
 		if (Type == D3D12_DESCRIPTOR_HEAP_TYPE_CBV_SRV_UAV)
 		{
 			INC_DWORD_STAT_BY(STAT_D3D12RayTracingUsedViewDescriptors, NumDescriptors);
@@ -1530,6 +1573,13 @@ public:
 		else
 		{
 			INC_DWORD_STAT_BY(STAT_D3D12RayTracingUsedSamplerDescriptors, NumDescriptors);
+
+			static uint64 GMaxNumUsedSamplerDescriptors = 0;
+			if (Heap.NumAllocatedDescriptors > GMaxNumUsedSamplerDescriptors)
+			{
+				GMaxNumUsedSamplerDescriptors = Heap.NumAllocatedDescriptors;
+			}
+			SET_DWORD_STAT(STAT_D3D12RayTracingMaxUsedSamplerDescriptors, GMaxNumUsedSamplerDescriptors);
 		}
 
 		return DescriptorTableBaseIndex;
