Build RT Instance Buffer in compute shader
- FRayTracingScene::Create() kicks off instance data upload to parallel thread along with RDG pass to build platform specific instance buffer.
- Implemented helper functions FillInstanceUploadBuffer and BuildRayTracingInstanceBuffer
- AutoInstanceBuffer is deprecated but still supported (for one release) when FRayTracingSceneInitializer is used.
- Temporarily remove support for Niagara meshes in TLAS.

diff --git a/Engine/Plugins/Experimental/GPULightmass/Source/GPULightmass/Private/LightmapRenderer.cpp b/Engine/Plugins/Experimental/GPULightmass/Source/GPULightmass/Private/LightmapRenderer.cpp
--- a/Engine/Plugins/Experimental/GPULightmass/Source/GPULightmass/Private/LightmapRenderer.cpp
+++ b/Engine/Plugins/Experimental/GPULightmass/Source/GPULightmass/Private/LightmapRenderer.cpp
@@ -23,6 +23,7 @@
 #include "EngineModule.h"
 #include "PostProcess/PostProcessing.h"
 #include "RayTracingGeometryManager.h"
+#include "RayTracingInstanceCopyShader.h"
 #include "ScreenPass.h"
 #include "RayTracingDynamicGeometryCollection.h"
 #include "Experimental/Containers/SherwoodHashTable.h"
@@ -166,21 +167,20 @@ struct FGPUBatchedTileRequests
 	TResourceArray<FGPUTileDescription> BatchedTilesDesc;
 };
 
-static FRayTracingSceneRHIRef CreateRayTracingScene(TArrayView<FRayTracingGeometryInstance> Instances, uint32& OutNumNativeInstances)
+static FRayTracingSceneRHIRef CreateRayTracingScene(TArrayView<FRayTracingGeometryInstance> Instances, TArray<uint32>& OutInstancesGeometryIndex)
 {
 	const uint32 NumSceneInstances = Instances.Num();
 
 	FRayTracingSceneInitializer2 Initializer;
-	Initializer.Instances = Instances;
 	Initializer.ShaderSlotsPerGeometrySegment = RAY_TRACING_NUM_SHADER_SLOTS;
 	Initializer.PerInstanceGeometries.SetNumUninitialized(NumSceneInstances);
-	Initializer.PerInstanceNumTransforms.SetNumUninitialized(NumSceneInstances);
 	Initializer.BaseInstancePrefixSum.SetNumUninitialized(NumSceneInstances);
 	Initializer.SegmentPrefixSum.SetNumUninitialized(NumSceneInstances);
 	Initializer.NumNativeInstances = 0;
 	Initializer.NumTotalSegments = 0;
 
-	Experimental::TSherwoodSet<FRHIRayTracingGeometry*> UniqueGeometries;
+	Experimental::TSherwoodMap<FRHIRayTracingGeometry*, uint32> UniqueGeometries;
+	OutInstancesGeometryIndex.SetNumUninitialized(NumSceneInstances);
 
 	// Compute geometry segment and instance count prefix sums.
 	// These are later used by GetHitRecordBaseIndex() during resource binding
@@ -204,21 +204,17 @@ static FRayTracingSceneRHIRef CreateRayTracingScene(TArrayView<FRayTracingGeomet
 		Initializer.SegmentPrefixSum[InstanceIndex] = Initializer.NumTotalSegments;
 		Initializer.NumTotalSegments += InstanceDesc.GeometryRHI->GetNumSegments();
 
-		bool bIsAlreadyInSet = false;
-		UniqueGeometries.Add(InstanceDesc.GeometryRHI, &bIsAlreadyInSet);
-		if (!bIsAlreadyInSet)
+		uint32 GeometryIndex = UniqueGeometries.FindOrAdd(InstanceDesc.GeometryRHI, Initializer.ReferencedGeometries.Num());
+		OutInstancesGeometryIndex[InstanceIndex] = GeometryIndex;
+		if (GeometryIndex == Initializer.ReferencedGeometries.Num())
 		{
 			Initializer.ReferencedGeometries.Add(InstanceDesc.GeometryRHI);
 		}
 
 		Initializer.BaseInstancePrefixSum[InstanceIndex] = Initializer.NumNativeInstances;
 		Initializer.NumNativeInstances += InstanceDesc.NumTransforms;
-
-		Initializer.PerInstanceNumTransforms[InstanceIndex] = InstanceDesc.NumTransforms;
 	}
 
-	OutNumNativeInstances = Initializer.NumNativeInstances;
-
 	return RHICreateRayTracingScene(MoveTemp(Initializer));
 }
 
@@ -1060,16 +1056,78 @@ void FSceneRenderState::SetupRayTracingScene(int32 LODIndex)
 		{
 			SCOPED_GPU_MASK(RHICmdList, FRHIGPUMask::All());
 
-			uint32 NumTotalInstances = 0;
-			RayTracingScene = CreateRayTracingScene(RayTracingGeometryInstances, NumTotalInstances);
+			TArray<uint32> InstancesGeometryIndex;
+			RayTracingScene = CreateRayTracingScene(RayTracingGeometryInstances, InstancesGeometryIndex);
+
+			const FRayTracingSceneInitializer2& SceneInitializer = RayTracingScene->GetInitializer();
 
-			FRayTracingAccelerationStructureSize SizeInfo = RHICalcRayTracingSceneSize(NumTotalInstances, ERayTracingAccelerationStructureFlags::FastTrace);
+			FRayTracingAccelerationStructureSize SizeInfo = RHICalcRayTracingSceneSize(SceneInitializer.NumNativeInstances, ERayTracingAccelerationStructureFlags::FastTrace);
 			FRHIResourceCreateInfo BufferCreateInfo(TEXT("LightmassRayTracingSceneBuffer"));
 			RayTracingSceneBuffer = RHICreateBuffer(uint32(SizeInfo.ResultSize), BUF_AccelerationStructure, 0, ERHIAccess::BVHWrite, BufferCreateInfo);
 			RayTracingSceneSRV = RHICreateShaderResourceView(RayTracingSceneBuffer);
 
+			FRHIResourceCreateInfo ScratchBufferCreateInfo(TEXT("LightmassRayTracingScratchBuffer"));
+			FBufferRHIRef ScratchBuffer = RHICreateBuffer(
+				uint32(SizeInfo.BuildScratchSize),
+				BUF_UnorderedAccess, GRHIRayTracingAccelerationStructureAlignment, ERHIAccess::UAVCompute,
+				ScratchBufferCreateInfo);
+
+			FRWBufferStructured InstanceBuffer;
+			InstanceBuffer.Initialize(TEXT("LightmassRayTracingInstanceBuffer"), GRHIRayTracingInstanceDescriptorSize, 1);
+
+			FByteAddressBuffer AccelerationStructureAddressesBuffer;
+			AccelerationStructureAddressesBuffer.Initialize(TEXT("LightmassRayTracingAccelerationStructureAddressesBuffer"), sizeof(FRayTracingAccelerationStructureAddress), BUF_Volatile);
+
+			FBufferRHIRef InstanceUploadBuffer;
+			FShaderResourceViewRHIRef InstanceUploadSRV;
+			{
+				const uint32 UploadBufferSize = SceneInitializer.NumNativeInstances * sizeof(FRayTracingInstanceDescriptorInput);
+
+				FRHIResourceCreateInfo CreateInfo(TEXT("RayTracingTestBedInstanceUploadBuffer"));
+				InstanceUploadBuffer = RHICreateStructuredBuffer(sizeof(FRayTracingInstanceDescriptorInput), UploadBufferSize, BUF_ShaderResource | BUF_Volatile, CreateInfo);
+				InstanceUploadSRV = RHICreateShaderResourceView(InstanceUploadBuffer);
+			}
+
+			{
+				FRayTracingInstanceDescriptorInput* InstanceUploadData = (FRayTracingInstanceDescriptorInput*)RHICmdList.LockBuffer(InstanceUploadBuffer, 0, sizeof(FRayTracingInstanceDescriptorInput), RLM_WriteOnly);
+				FillInstanceUploadBuffer(RayTracingGeometryInstances, InstancesGeometryIndex, RayTracingScene, MakeArrayView(InstanceUploadData, SceneInitializer.NumNativeInstances));
+				RHICmdList.UnlockBuffer(InstanceUploadBuffer);
+			}
+
+			RHICmdList.EnqueueLambda([&AccelerationStructureAddressesBuffer, &SceneInitializer](FRHICommandListImmediate& RHICmdList)
+				{
+					FRayTracingAccelerationStructureAddress* AddressesPtr = (FRayTracingAccelerationStructureAddress*)RHICmdList.LockBuffer(
+						AccelerationStructureAddressesBuffer.Buffer,
+						0,
+						SceneInitializer.ReferencedGeometries.Num() * sizeof(FRayTracingAccelerationStructureAddress), RLM_WriteOnly);
+
+					const uint32 NumGeometries = SceneInitializer.ReferencedGeometries.Num();
+					for (uint32 GeometryIndex = 0; GeometryIndex < NumGeometries; ++GeometryIndex)
+					{
+						AddressesPtr[GeometryIndex] = SceneInitializer.ReferencedGeometries[GeometryIndex]->GetAccelerationStructureAddress(RHICmdList.GetGPUMask().ToIndex());
+					}
+
+					RHICmdList.UnlockBuffer(AccelerationStructureAddressesBuffer.Buffer);
+				});
+
+			BuildRayTracingInstanceBuffer(
+				RHICmdList,
+				SceneInitializer.NumNativeInstances,
+				InstanceBuffer.UAV,
+				InstanceUploadSRV,
+				AccelerationStructureAddressesBuffer.SRV);
+
 			RHICmdList.BindAccelerationStructureMemory(RayTracingScene, RayTracingSceneBuffer, 0);
-			RHICmdList.BuildAccelerationStructure(RayTracingScene);
+
+			{
+				FRayTracingSceneBuildParams BuildParams;
+				BuildParams.Scene = RayTracingScene;
+				BuildParams.ScratchBuffer = ScratchBuffer;
+				BuildParams.ScratchBufferOffset = 0;
+				BuildParams.InstanceBuffer = InstanceBuffer.Buffer;
+				BuildParams.InstanceBufferOffset = 0;
+				RHICmdList.BuildAccelerationStructure(BuildParams);
+			}
 
 			FRayTracingPipelineStateInitializer PSOInitializer;
 
diff --git a/Engine/Shaders/Private/RayTracing/RayTracingInstanceCopy.usf b/Engine/Shaders/Private/RayTracing/RayTracingInstanceCopy.usf
--- a/Engine/Shaders/Private/RayTracing/RayTracingInstanceCopy.usf
+++ b/Engine/Shaders/Private/RayTracing/RayTracingInstanceCopy.usf
@@ -5,13 +5,14 @@
 ===============================================================================================*/
 
 #include "RayTracingCommon.ush"
+//#include "../SceneData.ush"
 #include "../Common.ush"
 
 #ifndef THREADGROUP_SIZE
 #define THREADGROUP_SIZE 1
 #endif
 
-uint InstancesCPUCount;
+uint NumInstances;
 uint DescBufferOffset;
 StructuredBuffer<float3x4> InstancesTransforms;
 RWStructuredBuffer<FRayTracingInstanceDescriptor> InstancesDescriptors;
@@ -20,13 +21,56 @@ RWStructuredBuffer<FRayTracingInstanceDescriptor> InstancesDescriptors;
 [numthreads(THREADGROUP_SIZE, 1, 1)]
 void RayTracingInstanceCopyShaderCS(uint3 DispatchThreadId : SV_DispatchThreadID)
 {
-	uint InstanceIdx = DispatchThreadId.x;	
+	uint Index = DispatchThreadId.x;	
 
-	float3x4 InstanceXForm = InstancesTransforms[InstanceIdx];
-	uint DescIndex = DescBufferOffset + InstanceIdx;
+	float3x4 InstanceXForm = InstancesTransforms[Index];
+	uint DescIndex = DescBufferOffset + Index;
 
-	if (InstanceIdx < InstancesCPUCount)
+	if (Index < NumInstances)
 	{
 		SetRayTracingInstanceTransform(InstancesDescriptors[DescIndex], InstanceXForm);
 	}
 }
+
+struct FRayTracingInstanceDescriptorInput
+{
+	//uint GPUSceneInstanceIndex;
+	float3x4 LocalToWorld;
+	uint AccelerationStructureIndex;
+	uint InstanceId;
+	uint InstanceMaskAndFlags;
+	uint InstanceContributionToHitGroupIndex;
+};
+
+StructuredBuffer<FRayTracingInstanceDescriptorInput> InputInstanceDescriptors;
+ByteAddressBuffer AccelerationStructureAddresses;
+
+[numthreads(THREADGROUP_SIZE, 1, 1)]
+void RayTracingBuildInstanceBufferCS(uint3 DispatchThreadId : SV_DispatchThreadID)
+{
+	uint Index = DispatchThreadId.x;
+
+	if (Index < NumInstances)
+	{
+		FRayTracingInstanceDescriptorInput InputDesc = InputInstanceDescriptors[Index];
+
+		//FInstanceSceneData InstanceSceneData = GetInstanceSceneData(InputDesc.GPUSceneInstanceIndex, View.InstanceSceneDataSOAStride);
+		//float4x4 LocalToWorld = InstanceSceneData.LocalToWorld;
+		float3x4 LocalToWorld = InputDesc.LocalToWorld;
+
+		uint2 BlasAddress = 0;
+
+		if (InputDesc.AccelerationStructureIndex != 0xFFFFFFFF)
+		{
+			BlasAddress = AccelerationStructureAddresses.Load2(InputDesc.AccelerationStructureIndex * 8);
+		}
+
+		InstancesDescriptors[Index] = BuildPlatformRayTracingInstanceDesc(
+			InputDesc.InstanceMaskAndFlags & 0xFF,
+			InputDesc.InstanceId,
+			TranslateRayTracingInstanceFlags((InputDesc.InstanceMaskAndFlags >> 8) & 0xFF),
+			InputDesc.InstanceContributionToHitGroupIndex,
+			LocalToWorld,
+			BlasAddress);
+	}
+}
diff --git a/Engine/Source/Runtime/D3D12RHI/Private/D3D12RayTracing.cpp b/Engine/Source/Runtime/D3D12RHI/Private/D3D12RayTracing.cpp
--- a/Engine/Source/Runtime/D3D12RHI/Private/D3D12RayTracing.cpp
+++ b/Engine/Source/Runtime/D3D12RHI/Private/D3D12RayTracing.cpp
@@ -13,7 +13,6 @@
 #include "HAL/CriticalSection.h"
 #include "HAL/IConsoleManager.h"
 #include "Misc/ScopeLock.h"
-#include "RayTracingInstanceCopyShader.h"
 #include "Async/ParallelFor.h"
 #include "Misc/BufferedOutputDevice.h"
 #include "String/LexFromString.h"
@@ -3109,6 +3108,7 @@ FRayTracingGeometryRHIRef FD3D12DynamicRHI::RHICreateRayTracingGeometry(const FR
 	return new FD3D12RayTracingGeometry(&Adapter, Initializer);
 }
 
+// This code path is deprecated and will be removed in the next release. Please use FRayTracingSceneInitializer2 instead.
 FRayTracingSceneRHIRef FD3D12DynamicRHI::RHICreateRayTracingScene(const FRayTracingSceneInitializer& Initializer)
 {
 	TRACE_CPUPROFILER_EVENT_SCOPE(CreateRayTracingScene);
@@ -3119,14 +3119,15 @@ FRayTracingSceneRHIRef FD3D12DynamicRHI::RHICreateRayTracingScene(const FRayTrac
 	Initializer2.DebugName = Initializer.DebugName;
 	Initializer2.ShaderSlotsPerGeometrySegment = Initializer.ShaderSlotsPerGeometrySegment;
 	Initializer2.NumMissShaderSlots = Initializer.NumMissShaderSlots;
-	Initializer2.Instances = Initializer.Instances;
 	Initializer2.PerInstanceGeometries.SetNumUninitialized(NumSceneInstances);
-	Initializer2.PerInstanceNumTransforms.SetNumUninitialized(NumSceneInstances);
 	Initializer2.BaseInstancePrefixSum.SetNumUninitialized(NumSceneInstances);
 	Initializer2.SegmentPrefixSum.SetNumUninitialized(NumSceneInstances);
 	Initializer2.NumNativeInstances = 0;
 	Initializer2.NumTotalSegments = 0;
 
+	TArray<uint32> PerInstanceNumTransforms;
+	PerInstanceNumTransforms.SetNumUninitialized(NumSceneInstances);
+
 	Experimental::TSherwoodSet<FRHIRayTracingGeometry*> UniqueGeometries;
 
 	for (uint32 InstanceIndex = 0; InstanceIndex < NumSceneInstances; ++InstanceIndex)
@@ -3155,12 +3156,88 @@ FRayTracingSceneRHIRef FD3D12DynamicRHI::RHICreateRayTracingScene(const FRayTrac
 		Initializer2.BaseInstancePrefixSum[InstanceIndex] = Initializer2.NumNativeInstances;
 		Initializer2.NumNativeInstances += InstanceDesc.NumTransforms;
 
-		Initializer2.PerInstanceNumTransforms[InstanceIndex] = InstanceDesc.NumTransforms;
+		PerInstanceNumTransforms[InstanceIndex] = InstanceDesc.NumTransforms;
 	}
 
+	TResourceArray<D3D12_RAYTRACING_INSTANCE_DESC, 16> NativeInstances;
+	NativeInstances.SetAllowCPUAccess(true);
+	NativeInstances.SetNumUninitialized(Initializer2.NumNativeInstances);
+
+	const EParallelForFlags ParallelForFlags = EParallelForFlags::None; // set ForceSingleThread for testing
+	ParallelFor(NumSceneInstances, [&Initializer2, Instances = Initializer.Instances, &NativeInstances](int32 InstanceIndex)
+		{
+			const FRayTracingGeometryInstance& Instance = Instances[InstanceIndex];
+			FD3D12RayTracingGeometry* Geometry = FD3D12DynamicRHI::ResourceCast(Initializer2.PerInstanceGeometries[InstanceIndex]);
+
+			D3D12_RAYTRACING_INSTANCE_DESC InstanceDesc = {};
+			check(InstanceDesc.AccelerationStructure == 0);
+
+			InstanceDesc.InstanceMask = Instance.Mask;
+			InstanceDesc.InstanceContributionToHitGroupIndex = Initializer2.SegmentPrefixSum[InstanceIndex] * Initializer2.ShaderSlotsPerGeometrySegment;
+			InstanceDesc.Flags = TranslateRayTracingInstanceFlags(Instance.Flags);
+
+			const uint32 NumTransforms = Instance.NumTransforms;
+
+			checkf(Instance.UserData.Num() == 0 || Instance.UserData.Num() >= int32(NumTransforms),
+				TEXT("User data array must be either be empty (Instance.DefaultUserData is used), or contain one entry per entry in Transforms array."));
+
+			const bool bUseUniqueUserData = Instance.UserData.Num() != 0;
+
+			const bool bGpuInstance = Instance.GPUTransformsSRV != nullptr;
+			const bool bCpuInstance = !bGpuInstance;
+
+			uint32 DescIndex = Initializer2.BaseInstancePrefixSum[InstanceIndex];
+
+			int32 NumInactiveDxrInstancesThisSceneInstance = 0;
+
+			for (uint32 TransformIndex = 0; TransformIndex < NumTransforms; ++TransformIndex)
+			{
+				InstanceDesc.InstanceID = bUseUniqueUserData ? Instance.UserData[TransformIndex] : Instance.DefaultUserData;
+
+				if (LIKELY(bCpuInstance))
+				{
+					// Set flag for deactivated instances
+					if (!Instance.ActivationMask.IsEmpty())
+					{
+						if ((Instance.ActivationMask[TransformIndex / 32] & (1 << (TransformIndex % 32))) == 0)
+						{
+							InstanceDesc.AccelerationStructure = 0xFFFFFFFFFFFFFFFF;
+							NumInactiveDxrInstancesThisSceneInstance++;
+						}
+						else
+						{
+							InstanceDesc.AccelerationStructure = 0;
+						}
+					}
+
+					// DXR uses a 3x4 transform matrix in row major layout
+
+					const FMatrix& Transform = Instance.Transforms[TransformIndex];
+
+					InstanceDesc.Transform[0][0] = Transform.M[0][0];
+					InstanceDesc.Transform[0][1] = Transform.M[1][0];
+					InstanceDesc.Transform[0][2] = Transform.M[2][0];
+					InstanceDesc.Transform[0][3] = Transform.M[3][0];
+
+					InstanceDesc.Transform[1][0] = Transform.M[0][1];
+					InstanceDesc.Transform[1][1] = Transform.M[1][1];
+					InstanceDesc.Transform[1][2] = Transform.M[2][1];
+					InstanceDesc.Transform[1][3] = Transform.M[3][1];
+
+					InstanceDesc.Transform[2][0] = Transform.M[0][2];
+					InstanceDesc.Transform[2][1] = Transform.M[1][2];
+					InstanceDesc.Transform[2][2] = Transform.M[2][2];
+					InstanceDesc.Transform[2][3] = Transform.M[3][2];
+				}
+
+				NativeInstances[DescIndex] = InstanceDesc;
+				++DescIndex;
+			}
+		}, ParallelForFlags);
+
 	FD3D12Adapter& Adapter = GetAdapter();
 
-	return new FD3D12RayTracingScene(&Adapter, MoveTemp(Initializer2));
+	return new FD3D12RayTracingScene(&Adapter, MoveTemp(Initializer2), MoveTemp(NativeInstances), MoveTemp(PerInstanceNumTransforms));
 }
 
 FRayTracingSceneRHIRef FD3D12DynamicRHI::RHICreateRayTracingScene(FRayTracingSceneInitializer2 Initializer)
@@ -3169,7 +3246,7 @@ FRayTracingSceneRHIRef FD3D12DynamicRHI::RHICreateRayTracingScene(FRayTracingSce
 
 	FD3D12Adapter& Adapter = GetAdapter();
 
-	return new FD3D12RayTracingScene(&Adapter, MoveTemp(Initializer));
+	return new FD3D12RayTracingScene(&Adapter, MoveTemp(Initializer), {}, {});
 }
 
 FBufferRHIRef FD3D12RayTracingGeometry::NullTransformBuffer;
@@ -3711,127 +3788,24 @@ void FD3D12RayTracingGeometry::CompactAccelerationStructure(FD3D12CommandContext
 	AccelerationStructureCompactedSize = InSizeAfterCompaction;
 }
 
-FD3D12RayTracingScene::FD3D12RayTracingScene(FD3D12Adapter* Adapter, FRayTracingSceneInitializer2 InInitializer)
-	: FD3D12AdapterChild(Adapter), Initializer(MoveTemp(InInitializer))
+FD3D12RayTracingScene::FD3D12RayTracingScene(FD3D12Adapter* Adapter, FRayTracingSceneInitializer2 InInitializer, TResourceArray<D3D12_RAYTRACING_INSTANCE_DESC, 16> InInstances, TArray<uint32> InPerInstanceNumTransforms)
+	: FD3D12AdapterChild(Adapter), Initializer(MoveTemp(InInitializer)), Instances(MoveTemp(InInstances)), PerInstanceNumTransforms(MoveTemp(InPerInstanceNumTransforms))
 {
 	INC_DWORD_STAT(STAT_D3D12RayTracingAllocatedTLAS);
 
 	checkf(Initializer.NumMissShaderSlots >= 1, TEXT("Need at least 1 miss shader slot."));
 	checkf(Initializer.Lifetime == RTSL_SingleFrame, TEXT("Only single-frame ray tracing scenes are currently implemented."));
 
-	int32 NumDxrInstances = 0;
-	int32 NumInactiveDxrInstances = 0;
-	int32 NumCopyCommands = 0;
-
-	const int32 NumSceneInstances = Initializer.Instances.Num();
-	for (int32 InstanceIndex = 0; InstanceIndex < NumSceneInstances; ++InstanceIndex)
-	{
-		const FRayTracingGeometryInstance& Instance = Initializer.Instances[InstanceIndex];
-
-		if (UNLIKELY(Instance.GPUTransformsSRV != nullptr && Instance.NumTransforms))
-		{
-			FInstanceCopyCommand Command;
-			Command.Source = Instance.GPUTransformsSRV;
-			Command.BaseIndex = NumDxrInstances;
-			Command.Num = Instance.NumTransforms;
-			CopyCommands.Add(Command);
-		}
-
-		NumDxrInstances += Instance.NumTransforms;
-	}
-
-	Instances.SetAllowCPUAccess(true);
-	Instances.SetNumUninitialized(NumDxrInstances);
-
-	const EParallelForFlags ParallelForFlags = EParallelForFlags::None; // set ForceSingleThread for testing
-	ParallelFor(NumSceneInstances, [this, &NumInactiveDxrInstances](int32 InstanceIndex)
-	{
-		const FRayTracingGeometryInstance& Instance = Initializer.Instances[InstanceIndex];
-		FD3D12RayTracingGeometry* Geometry = FD3D12DynamicRHI::ResourceCast(Initializer.PerInstanceGeometries[InstanceIndex]);
-
-		D3D12_RAYTRACING_INSTANCE_DESC InstanceDesc = {};
-		check(InstanceDesc.AccelerationStructure == 0);
-
-		InstanceDesc.InstanceMask = Instance.Mask;
-		InstanceDesc.InstanceContributionToHitGroupIndex = Initializer.SegmentPrefixSum[InstanceIndex] * Initializer.ShaderSlotsPerGeometrySegment;
-		InstanceDesc.Flags = TranslateRayTracingInstanceFlags(Instance.Flags);
-
-		const uint32 NumTransforms = Instance.NumTransforms;
-
-		checkf(Instance.UserData.Num() == 0 || Instance.UserData.Num() >= int32(NumTransforms),
-			TEXT("User data array must be either be empty (Instance.DefaultUserData is used), or contain one entry per entry in Transforms array."));
-
-		const bool bUseUniqueUserData = Instance.UserData.Num() != 0;
-
-		const bool bGpuInstance = Instance.GPUTransformsSRV != nullptr;
-		const bool bCpuInstance = !bGpuInstance;
-
-		uint32 DescIndex = Initializer.BaseInstancePrefixSum[InstanceIndex];
-
-		int32 NumInactiveDxrInstancesThisSceneInstance = 0;
-
-		for (uint32 TransformIndex = 0; TransformIndex < NumTransforms; ++TransformIndex)
-		{
-			InstanceDesc.InstanceID = bUseUniqueUserData ? Instance.UserData[TransformIndex] : Instance.DefaultUserData;
-
-			if (LIKELY(bCpuInstance))
-			{
-				// Set flag for deactivated instances
-				if (!Instance.ActivationMask.IsEmpty())
-				{
-					if ((Instance.ActivationMask[TransformIndex / 32] & (1 << (TransformIndex % 32))) == 0)
-					{
-						InstanceDesc.AccelerationStructure = 0xFFFFFFFFFFFFFFFF;
-						NumInactiveDxrInstancesThisSceneInstance++;
-					}
-					else
-					{
-						InstanceDesc.AccelerationStructure = 0;
-					}
-				}
-
-				// DXR uses a 3x4 transform matrix in row major layout
-
-				const FMatrix& Transform = Instance.Transforms[TransformIndex];
-
-				InstanceDesc.Transform[0][0] = Transform.M[0][0];
-				InstanceDesc.Transform[0][1] = Transform.M[1][0];
-				InstanceDesc.Transform[0][2] = Transform.M[2][0];
-				InstanceDesc.Transform[0][3] = Transform.M[3][0];
-
-				InstanceDesc.Transform[1][0] = Transform.M[0][1];
-				InstanceDesc.Transform[1][1] = Transform.M[1][1];
-				InstanceDesc.Transform[1][2] = Transform.M[2][1];
-				InstanceDesc.Transform[1][3] = Transform.M[3][1];
-
-				InstanceDesc.Transform[2][0] = Transform.M[0][2];
-				InstanceDesc.Transform[2][1] = Transform.M[1][2];
-				InstanceDesc.Transform[2][2] = Transform.M[2][2];
-				InstanceDesc.Transform[2][3] = Transform.M[3][2];
-			}
-
-			Instances[DescIndex] = InstanceDesc;
-			++DescIndex;
-		}
-
-	#if STATS
-		FPlatformAtomics::InterlockedAdd(&NumInactiveDxrInstances, NumInactiveDxrInstancesThisSceneInstance);
-	#endif
-
-	}, ParallelForFlags);
-
 	ERayTracingAccelerationStructureFlags BuildFlags = ERayTracingAccelerationStructureFlags::FastTrace; // #yuriy_todo: pass this in
 
 	BuildInputs = {};
 	BuildInputs.Type = D3D12_RAYTRACING_ACCELERATION_STRUCTURE_TYPE_TOP_LEVEL;
 	BuildInputs.DescsLayout = D3D12_ELEMENTS_LAYOUT_ARRAY;
-	BuildInputs.NumDescs = Instances.Num();
+	BuildInputs.NumDescs = Initializer.NumNativeInstances;
 	BuildInputs.Flags = TranslateRayTracingAccelerationStructureFlags(BuildFlags);
 
 	// Get maximum buffer sizes for all GPUs in the system
 	SizeInfo = RHICalcRayTracingSceneSize(BuildInputs.NumDescs, BuildFlags);
-
-	SET_DWORD_STAT(STAT_RayTracingInstances, Initializer.NumNativeInstances - NumInactiveDxrInstances);
 };
 
 FD3D12RayTracingScene::~FD3D12RayTracingScene()
@@ -3887,8 +3861,7 @@ void FD3D12RayTracingScene::BindBuffer(FRHIBuffer* InBuffer, uint32 InBufferOffs
 
 void FD3D12RayTracingScene::BuildAccelerationStructure(FD3D12CommandContext& CommandContext,
 	FD3D12Buffer* ScratchBuffer, uint32 ScratchBufferOffset,
-	FD3D12Buffer* InstanceBuffer, uint32 InstanceBufferOffset,
-	uint32 NumInstanceDescs)
+	FD3D12Buffer* InstanceBuffer, uint32 InstanceBufferOffset)
 {
 	TRACE_CPUPROFILER_EVENT_SCOPE(BuildAccelerationStructure_TopLevel);
 	SCOPE_CYCLE_COUNTER(STAT_D3D12BuildTLAS);
@@ -3918,17 +3891,23 @@ void FD3D12RayTracingScene::BuildAccelerationStructure(FD3D12CommandContext& Com
 		ScratchBufferOffset = 0;
 	}
 
+	checkf(ScratchBuffer, TEXT("TLAS build requires scratch buffer of at least %lld bytes."), PrebuildInfo.ScratchDataSizeInBytes);
+
+	const D3D12_GPU_VIRTUAL_ADDRESS ScratchAddress = ScratchBuffer->ResourceLocation.GetGPUVirtualAddress() + ScratchBufferOffset;
+
+	checkf(ScratchAddress % GRHIRayTracingAccelerationStructureAlignment == 0,
+		TEXT("TLAS scratch buffer (plus offset) must be aligned to %lld bytes."),
+		GRHIRayTracingAccelerationStructureAlignment);
+
 	checkf(PrebuildInfo.ScratchDataSizeInBytes + ScratchBufferOffset <= ScratchBuffer->GetSize(),
 		TEXT("TLAS scratch buffer size is %lld bytes with offset %lld (%lld bytes available), but the build requires %lld bytes. ")
 		TEXT("BuildInputs.NumDescs = %d, Instances.Num = %d, PerInstanceGeometries.Num = %d."),
 			ScratchBuffer->GetSize(), ScratchBufferOffset, ScratchBuffer->GetSize() - ScratchBufferOffset, 
 			PrebuildInfo.ScratchDataSizeInBytes,
-			BuildInputs.NumDescs, Instances.Num(), Initializer.PerInstanceGeometries.Num());
+			BuildInputs.NumDescs, Initializer.NumNativeInstances, Initializer.PerInstanceGeometries.Num());
 
 	TRefCountPtr<FD3D12Buffer> AutoInstanceBuffer;
 
-	// Create and fill automatic instance buffer
-
 	if (BuildInputs.NumDescs)
 	{
 		// - Set up acceleration structure pointers and make them resident.
@@ -3938,86 +3917,65 @@ void FD3D12RayTracingScene::BuildAccelerationStructure(FD3D12CommandContext& Com
 		HitGroupSystemParametersCache.Empty(Initializer.NumTotalSegments);
 
 		{
-			int32 DxrInstanceIndex = 0;
+			int32 NativeInstanceIndex = 0;
+
 			const int32 NumSceneInstances = Initializer.PerInstanceGeometries.Num();
 			for (int32 InstanceIndex = 0; InstanceIndex < NumSceneInstances; ++InstanceIndex)
 			{
 				FD3D12RayTracingGeometry* Geometry = FD3D12DynamicRHI::ResourceCast(Initializer.PerInstanceGeometries[InstanceIndex]);
 
-				D3D12_GPU_VIRTUAL_ADDRESS BlasAddress = Geometry->AccelerationStructureBuffers[GPUIndex]->ResourceLocation.GetGPUVirtualAddress();
+				// make a copy of system parameters to they can optimized fetch during SBT building (only done for GPU0)
+				check(Geometry->HitGroupSystemParameters[0].Num() > 0);
+				HitGroupSystemParametersCache.Append(Geometry->HitGroupSystemParameters[0]);
 
-				const int32 NumTransforms = Initializer.PerInstanceNumTransforms[InstanceIndex];
-				for (int32 TransformIndex = 0; TransformIndex < NumTransforms; ++TransformIndex)
+				if (InstanceBuffer == nullptr)
 				{
-					// If this instance has been deactivated, set BLAS address to 0
-					if (Instances[DxrInstanceIndex].AccelerationStructure == 0xFFFFFFFFFFFFFFFF)
-					{
-						Instances[DxrInstanceIndex].AccelerationStructure = 0;
-					}
-					else
+					D3D12_GPU_VIRTUAL_ADDRESS BlasAddress = Geometry->AccelerationStructureBuffers[GPUIndex]->ResourceLocation.GetGPUVirtualAddress();
+
+					const int32 NumTransforms = PerInstanceNumTransforms[InstanceIndex];
+					for (int32 TransformIndex = 0; TransformIndex < NumTransforms; ++TransformIndex)
 					{
-						Instances[DxrInstanceIndex].AccelerationStructure = BlasAddress;
+						// If this instance has been deactivated, set BLAS address to 0
+						if (Instances[NativeInstanceIndex].AccelerationStructure == 0xFFFFFFFFFFFFFFFF)
+						{
+							Instances[NativeInstanceIndex].AccelerationStructure = 0;
+						}
+						else
+						{
+							Instances[NativeInstanceIndex].AccelerationStructure = BlasAddress;
+						}
+
+						NativeInstanceIndex++;
 					}
-
-					DxrInstanceIndex++;
 				}
-
-				// make a copy of system parameters to they can optimized fetch during SBT building (only done for GPU0)
-				check(Geometry->HitGroupSystemParameters[0].Num() > 0);
-				HitGroupSystemParametersCache.Append(Geometry->HitGroupSystemParameters[0]);
 			}
 		}
 
-		check(Instances.Num() == BuildInputs.NumDescs);
-
-		D3D12_RESOURCE_DESC InstanceBufferDesc = CD3DX12_RESOURCE_DESC::Buffer(
-			sizeof(D3D12_RAYTRACING_INSTANCE_DESC) * BuildInputs.NumDescs,
-			D3D12_RESOURCE_FLAG_ALLOW_UNORDERED_ACCESS);
-
-		// #dxr_todo multi state only when bShouldCopyIndirectInstances is set - will still need transition to copy_src for lock behind but this can be done on the complete pool in theory (have to check cost)
-		// #dxr_todo use dynamic with upload memory when no CopyCommands requested - resource will be in upload memory then instead of VRAM
-
-		ED3D12ResourceTransientMode TransientMode = ED3D12ResourceTransientMode::NonTransient;
-		ID3D12ResourceAllocator* ResourceAllocator = nullptr;
-		FRHIGPUMask GPUMask = FRHIGPUMask::FromIndex(GPUIndex);
-		bool bHasInitialData = true;
-
-		AutoInstanceBuffer = Adapter->CreateRHIBuffer(
-			InstanceBufferDesc, D3D12_RAYTRACING_INSTANCE_DESCS_BYTE_ALIGNMENT,
-			sizeof(D3D12_RAYTRACING_INSTANCE_DESC), InstanceBufferDesc.Width, BUF_UnorderedAccess, ED3D12ResourceStateMode::MultiState, D3D12_RESOURCE_STATE_COPY_DEST, 
-			bHasInitialData, GPUMask, TransientMode, ResourceAllocator, TEXT("InstanceBuffer"));
-
-		// Use copy queue for uploading the data
-		FD3D12SyncPoint CopyQueueSyncPoint = AutoInstanceBuffer->UploadResourceDataViaCopyQueue(&Instances);
-		CommandContext.CopyQueueSyncPoint.Merge(CopyQueueSyncPoint);
-
-		// #yuriy_todo: only create internal instance buffer if explicit one is not given
-		checkf(InstanceBuffer == nullptr, TEXT("Explicit instance buffer support is not yet implemented for D3D12 RHI."));
-		InstanceBuffer = AutoInstanceBuffer.GetReference();
-
-		if (CopyCommands.Num() > 0)
+		if (InstanceBuffer == nullptr)
 		{
-			TRHICommandList_RecursiveHazardous<FD3D12CommandContext> RHICmdList(&CommandContext);
-			FUnorderedAccessViewRHIRef InstancesDescUAV = RHICreateUnorderedAccessView(InstanceBuffer, false, false);
+			check(Instances.Num() == BuildInputs.NumDescs);
 
-			FD3D12DynamicRHI::TransitionResource(CommandContext.CommandListHandle, InstanceBuffer->GetResource(),
-				D3D12_RESOURCE_STATE_COPY_DEST, D3D12_RESOURCE_STATE_UNORDERED_ACCESS, 0, FD3D12DynamicRHI::ETransitionMode::Apply);
+			D3D12_RESOURCE_DESC InstanceBufferDesc = CD3DX12_RESOURCE_DESC::Buffer(
+				sizeof(D3D12_RAYTRACING_INSTANCE_DESC) * BuildInputs.NumDescs,
+				D3D12_RESOURCE_FLAG_ALLOW_UNORDERED_ACCESS);
 
-			RHICmdList.BeginUAVOverlap(InstancesDescUAV);
+			ED3D12ResourceTransientMode TransientMode = ED3D12ResourceTransientMode::NonTransient;
+			ID3D12ResourceAllocator* ResourceAllocator = nullptr;
+			FRHIGPUMask GPUMask = FRHIGPUMask::FromIndex(GPUIndex);
+			bool bHasInitialData = true;
 
-			for (const FInstanceCopyCommand& Command : CopyCommands)
-			{
-				FRHIShaderResourceView* TransformsSRV = Command.Source.GetReference();
-				CopyRayTracingGPUInstances(RHICmdList, Command.Num, Command.BaseIndex, TransformsSRV, InstancesDescUAV);
-			}
+			AutoInstanceBuffer = Adapter->CreateRHIBuffer(
+				InstanceBufferDesc, D3D12_RAYTRACING_INSTANCE_DESCS_BYTE_ALIGNMENT,
+				sizeof(D3D12_RAYTRACING_INSTANCE_DESC), InstanceBufferDesc.Width, BUF_UnorderedAccess, ED3D12ResourceStateMode::MultiState, D3D12_RESOURCE_STATE_COPY_DEST,
+				bHasInitialData, GPUMask, TransientMode, ResourceAllocator, TEXT("InstanceBuffer"));
 
-			RHICmdList.EndUAVOverlap(InstancesDescUAV);
+			// Use copy queue for uploading the data
+			FD3D12SyncPoint CopyQueueSyncPoint = AutoInstanceBuffer->UploadResourceDataViaCopyQueue(&Instances);
+			CommandContext.CopyQueueSyncPoint.Merge(CopyQueueSyncPoint);
+
+			InstanceBuffer = AutoInstanceBuffer.GetReference();
+			InstanceBufferOffset = 0;
 
-			FD3D12DynamicRHI::TransitionResource(CommandContext.CommandListHandle, InstanceBuffer->GetResource(),
-				D3D12_RESOURCE_STATE_UNORDERED_ACCESS, D3D12_RESOURCE_STATE_NON_PIXEL_SHADER_RESOURCE, 0, FD3D12DynamicRHI::ETransitionMode::Apply);
-		}
-		else
-		{
 			FD3D12DynamicRHI::TransitionResource(CommandContext.CommandListHandle, InstanceBuffer->GetResource(),
 				D3D12_RESOURCE_STATE_COPY_DEST, D3D12_RESOURCE_STATE_NON_PIXEL_SHADER_RESOURCE, 0, FD3D12DynamicRHI::ETransitionMode::Apply);
 		}
@@ -4047,8 +4005,8 @@ void FD3D12RayTracingScene::BuildAccelerationStructure(FD3D12CommandContext& Com
 			#endif // ENABLE_RESIDENCY_MANAGEMENT
 			};
 
-			const uint32 NumReferencedGeometries = Initializer.ReferencedGeometries.Num();
-			for (uint32 Index = 0; Index < NumReferencedGeometries; ++Index)
+			const int32 NumReferencedGeometries = Initializer.ReferencedGeometries.Num();
+			for (int32 Index = 0; Index < NumReferencedGeometries; ++Index)
 			{
 				FD3D12RayTracingGeometry* Geometry = FD3D12DynamicRHI::ResourceCast(Initializer.ReferencedGeometries[Index].GetReference());
 
@@ -4075,10 +4033,10 @@ void FD3D12RayTracingScene::BuildAccelerationStructure(FD3D12CommandContext& Com
 		}
 	}
 	
-	
 	// Build the actual acceleration structure
-	const uint32 NumReferencedGeometries = Initializer.ReferencedGeometries.Num();
-	for (uint32 Index = 0; Index < NumReferencedGeometries; ++Index)
+
+	const int32 NumReferencedGeometries = Initializer.ReferencedGeometries.Num();
+	for (int32 Index = 0; Index < NumReferencedGeometries; ++Index)
 	{
 		FD3D12RayTracingGeometry* Geometry = FD3D12DynamicRHI::ResourceCast(Initializer.ReferencedGeometries[Index].GetReference());
 		Geometry->AccelerationStructureBuffers[GPUIndex]->ResourceLocation.GetResource()->UpdateResidency(CommandContext.CommandListHandle);
@@ -4093,9 +4051,9 @@ void FD3D12RayTracingScene::BuildAccelerationStructure(FD3D12CommandContext& Com
 
 	D3D12_BUILD_RAYTRACING_ACCELERATION_STRUCTURE_DESC BuildDesc = {};
 	BuildDesc.Inputs = BuildInputs;
-	BuildDesc.Inputs.InstanceDescs = InstanceBuffer ? InstanceBuffer->ResourceLocation.GetGPUVirtualAddress() : D3D12_GPU_VIRTUAL_ADDRESS(0);
+	BuildDesc.Inputs.InstanceDescs = InstanceBuffer ? InstanceBuffer->ResourceLocation.GetGPUVirtualAddress() + InstanceBufferOffset : D3D12_GPU_VIRTUAL_ADDRESS(0);
 	BuildDesc.DestAccelerationStructureData = AccelerationStructureBuffer->ResourceLocation.GetGPUVirtualAddress() + BufferOffset;
-	BuildDesc.ScratchAccelerationStructureData = ScratchBuffer->ResourceLocation.GetGPUVirtualAddress() + ScratchBufferOffset;
+	BuildDesc.ScratchAccelerationStructureData = ScratchAddress;
 	BuildDesc.SourceAccelerationStructureData = D3D12_GPU_VIRTUAL_ADDRESS(0); // Null source TLAS as this is a build command
 
 	checkf(BuildDesc.ScratchAccelerationStructureData % 256 == 0, TEXT("TLAS build scratch buffer must have 256 byte alignment."));
@@ -4339,8 +4297,7 @@ void FD3D12CommandContext::RHIBuildAccelerationStructure(const FRayTracingSceneB
 	Scene->BuildAccelerationStructure(
 		*this,
 		ScratchBuffer, SceneBuildParams.ScratchBufferOffset,
-		InstanceBuffer, SceneBuildParams.InstanceBufferOffset,
-		SceneBuildParams.NumInstances
+		InstanceBuffer, SceneBuildParams.InstanceBufferOffset
 	);
 }
 
diff --git a/Engine/Source/Runtime/D3D12RHI/Private/D3D12RayTracing.h b/Engine/Source/Runtime/D3D12RHI/Private/D3D12RayTracing.h
--- a/Engine/Source/Runtime/D3D12RHI/Private/D3D12RayTracing.h
+++ b/Engine/Source/Runtime/D3D12RHI/Private/D3D12RayTracing.h
@@ -26,6 +26,11 @@ public:
 	FD3D12RayTracingGeometry(FD3D12Adapter* Adapter, const FRayTracingGeometryInitializer& Initializer);
 	~FD3D12RayTracingGeometry();
 
+	virtual FRayTracingAccelerationStructureAddress GetAccelerationStructureAddress(uint64 GPUIndex) const final override
+	{
+		checkf(IsInRHIThread(), TEXT("Acceleration structure addresses can only be accessed on RHI timeline due to compaction and defragmentation."));
+		return AccelerationStructureBuffers[GPUIndex]->ResourceLocation.GetGPUVirtualAddress();
+	}
 	virtual uint32 GetNumSegments() const final override { return Segments.Num(); }
 
 	void SetupHitGroupSystemParameters(uint32 InGPUIndex);
@@ -93,7 +98,7 @@ public:
 	// Scaling beyond 5 total threads does not yield any speedup in practice.
 	static constexpr uint32 MaxBindingWorkers = 5; // RHI thread + 4 parallel workers.
 
-	FD3D12RayTracingScene(FD3D12Adapter* Adapter, FRayTracingSceneInitializer2 Initializer);
+	FD3D12RayTracingScene(FD3D12Adapter* Adapter, FRayTracingSceneInitializer2 Initializer, TResourceArray<D3D12_RAYTRACING_INSTANCE_DESC, 16> Instances, TArray<uint32> PerInstanceNumTransforms);
 	~FD3D12RayTracingScene();
 
 	const FRayTracingSceneInitializer2& GetInitializer() const override final { return Initializer; }
@@ -103,8 +108,7 @@ public:
 
 	void BuildAccelerationStructure(FD3D12CommandContext& CommandContext,
 		FD3D12Buffer* ScratchBuffer, uint32 ScratchBufferOffset,
-		FD3D12Buffer* InstanceBuffer, uint32 InstanceBufferOffset,
-		uint32 NumInstanceDescs
+		FD3D12Buffer* InstanceBuffer, uint32 InstanceBufferOffset
 	);
 
 	FShaderResourceViewRHIRef ShaderResourceView;
@@ -112,21 +116,15 @@ public:
 	TRefCountPtr<FD3D12Buffer> AccelerationStructureBuffers[MAX_NUM_GPUS];
 	uint32 BufferOffset = 0;
 
-	TResourceArray<D3D12_RAYTRACING_INSTANCE_DESC, 16> Instances;
 	D3D12_BUILD_RAYTRACING_ACCELERATION_STRUCTURE_INPUTS BuildInputs = {};
 
 	FRayTracingAccelerationStructureSize SizeInfo = {};
 
-	struct FInstanceCopyCommand
-	{
-		FShaderResourceViewRHIRef Source;
-		int32 BaseIndex = -1;
-		int32 Num = 0;
-	};
-	TArray<FInstanceCopyCommand> CopyCommands;
-
 	const FRayTracingSceneInitializer2 Initializer;
 
+	TResourceArray<D3D12_RAYTRACING_INSTANCE_DESC, 16> Instances;
+	TArray<uint32> PerInstanceNumTransforms;
+
 	// Scene keeps track of child acceleration structure buffers to ensure
 	// they are resident when any ray tracing work is dispatched.
 	TArray<FD3D12ResidencyHandle*> GeometryResidencyHandles[MAX_NUM_GPUS];
diff --git a/Engine/Source/Runtime/RHI/Public/DynamicRHI.h b/Engine/Source/Runtime/RHI/Public/DynamicRHI.h
--- a/Engine/Source/Runtime/RHI/Public/DynamicRHI.h
+++ b/Engine/Source/Runtime/RHI/Public/DynamicRHI.h
@@ -257,11 +257,6 @@ struct FRayTracingSceneInitializer
 
 struct FRayTracingSceneInitializer2
 {
-	// WIP - need these arrays until building instance buffer is moved to higher level code
-	// will be removed before 5.0 release
-	TArrayView<FRayTracingGeometryInstance> Instances;
-	TArray<uint32> PerInstanceNumTransforms;
-
 	// Unique list of geometries referenced by all instances in this scene.
 	// Any referenced geometry is kept alive while the scene is alive.
 	TArray<FRayTracingGeometryRHIRef> ReferencedGeometries;
diff --git a/Engine/Source/Runtime/RHI/Public/RHICommandList.h b/Engine/Source/Runtime/RHI/Public/RHICommandList.h
--- a/Engine/Source/Runtime/RHI/Public/RHICommandList.h
+++ b/Engine/Source/Runtime/RHI/Public/RHICommandList.h
@@ -2977,6 +2977,7 @@ public:
 		}
 	}
 
+	UE_DEPRECATED(5.0, "Use FRayTracingSceneBuildParams instead.")
 	FORCEINLINE_DEBUGGABLE void BuildAccelerationStructure(FRHIRayTracingScene* Scene)
 	{
 		FRayTracingSceneBuildParams Params;
diff --git a/Engine/Source/Runtime/RHI/Public/RHIContext.h b/Engine/Source/Runtime/RHI/Public/RHIContext.h
--- a/Engine/Source/Runtime/RHI/Public/RHIContext.h
+++ b/Engine/Source/Runtime/RHI/Public/RHIContext.h
@@ -404,18 +404,13 @@ struct FRayTracingSceneBuildParams
 	FRHIBuffer* ResultBuffer = nullptr;
 	uint32 ResultBufferOffset = 0;
 
-	// Optional explicit build scratch buffer. Must be in UAV state.
-	// Scratch buffer will be created automatically if this is not provided.
+	// Scratch buffer used to build Acceleration structure. Must be in UAV state.
 	FRHIBuffer* ScratchBuffer = nullptr;
 	uint32 ScratchBufferOffset = 0;
 
-	// Optional explicit buffer of native ray tracing instance descriptors. Must be in SRV state.
-	// Instance list will be taken from the scene if this is not provided.
+	// Buffer of native ray tracing instance descriptors. Must be in SRV state.
 	FRHIBuffer* InstanceBuffer = nullptr;
 	uint32 InstanceBufferOffset = 0;
-
-	// Number of elements in the InstanceBuffer. Only used if InstanceBuffer is provided.
-	uint32 NumInstances = 0;
 };
 
 struct FCopyBufferRegionParams
diff --git a/Engine/Source/Runtime/RHI/Public/RHIResources.h b/Engine/Source/Runtime/RHI/Public/RHIResources.h
--- a/Engine/Source/Runtime/RHI/Public/RHIResources.h
+++ b/Engine/Source/Runtime/RHI/Public/RHIResources.h
@@ -1910,10 +1910,13 @@ public:
 	FRHIRayTracingAccelerationStructure() : FRHIResource(RRT_RayTracingAccelerationStructure) {}
 };
 
+using FRayTracingAccelerationStructureAddress = uint64;
+
 /** Bottom level ray tracing acceleration structure (contains triangles). */
 class FRHIRayTracingGeometry : public FRHIRayTracingAccelerationStructure
 {
 public:
+	virtual FRayTracingAccelerationStructureAddress GetAccelerationStructureAddress(uint64 GPUIndex) const = 0;
 	virtual uint32 GetNumSegments() const = 0;
 };
 
diff --git a/Engine/Source/Runtime/RenderCore/Private/RayTracingInstanceCopyShader.cpp b/Engine/Source/Runtime/RenderCore/Private/RayTracingInstanceCopyShader.cpp
--- a/Engine/Source/Runtime/RenderCore/Private/RayTracingInstanceCopyShader.cpp
+++ b/Engine/Source/Runtime/RenderCore/Private/RayTracingInstanceCopyShader.cpp
@@ -1,14 +1,166 @@
 // Copyright Epic Games, Inc. All Rights Reserved.
 
 #include "RayTracingInstanceCopyShader.h"
+
+#include "RayTracingDefinitions.h"
+
+#include "RenderGraphBuilder.h"
 #include "ShaderParameterUtils.h"
 #include "RendererInterface.h"
+#include "ShaderParameterStruct.h"
+#include "GlobalShader.h"
+
+#include "Async/ParallelFor.h"
 
 #if RHI_RAYTRACING
-IMPLEMENT_SHADER_TYPE3(RayTracingInstanceCopyCS, SF_Compute);
-#endif //RHI_RAYTRACING
 
-bool RayTracingInstanceCopyCS::ShouldCompilePermutation(const FGlobalShaderPermutationParameters& Parameters)
+void FillInstanceUploadBuffer(
+	TConstArrayView<FRayTracingGeometryInstance> Instances,
+	TConstArrayView<uint32> InstancesGeometryIndex,
+	FRayTracingSceneRHIRef RayTracingSceneRHI,
+	TArrayView<FRayTracingInstanceDescriptorInput> OutInstanceUploadData)
+{
+	const FRayTracingSceneInitializer2& SceneInitializer = RayTracingSceneRHI->GetInitializer();
+
+	int32 NumInactiveNativeInstances = 0;
+
+	const int32 NumSceneInstances = Instances.Num();
+	ParallelFor(NumSceneInstances, [OutInstanceUploadData, Instances, InstancesGeometryIndex, &SceneInitializer, &NumInactiveNativeInstances](int32 SceneInstanceIndex)
+		{
+			const FRayTracingGeometryInstance& SceneInstance = Instances[SceneInstanceIndex];
+
+			const uint32 NumTransforms = SceneInstance.NumTransforms;
+
+			checkf(SceneInstance.UserData.Num() == 0 || SceneInstance.UserData.Num() >= int32(NumTransforms),
+				TEXT("User data array must be either be empty (Instance.DefaultUserData is used), or contain one entry per entry in Transforms array."));
+
+			check(SceneInstance.ActivationMask.IsEmpty() || SceneInstance.ActivationMask.Num() * 32 >= int32(NumTransforms));
+
+			const bool bUseUniqueUserData = SceneInstance.UserData.Num() != 0;
+
+			const bool bGpuInstance = SceneInstance.GPUTransformsSRV != nullptr;
+			const bool bCpuInstance = !bGpuInstance;
+
+			const uint32 AccelerationStructureIndex = InstancesGeometryIndex[SceneInstanceIndex];
+
+			uint32 BaseInstanceIndex = SceneInitializer.BaseInstancePrefixSum[SceneInstanceIndex];
+
+			int32 NumInactiveNativeInstancesThisSceneInstance = 0;
+			for (uint32 TransformIndex = 0; TransformIndex < NumTransforms; ++TransformIndex)
+			{
+				FRayTracingInstanceDescriptorInput& InstanceDesc = OutInstanceUploadData[BaseInstanceIndex + TransformIndex];
+				InstanceDesc.InstanceMaskAndFlags = SceneInstance.Mask | ((uint32)SceneInstance.Flags << 8);
+				InstanceDesc.InstanceContributionToHitGroupIndex = SceneInitializer.SegmentPrefixSum[SceneInstanceIndex] * RAY_TRACING_NUM_SHADER_SLOTS;
+				InstanceDesc.InstanceId = bUseUniqueUserData ? SceneInstance.UserData[TransformIndex] : SceneInstance.DefaultUserData;
+				//InstanceDesc.GPUSceneInstanceIndex = 0; // TODO
+
+				InstanceDesc.AccelerationStructureIndex = AccelerationStructureIndex;
+
+				if (!SceneInstance.ActivationMask.IsEmpty() && (SceneInstance.ActivationMask[TransformIndex / 32] & (1 << (TransformIndex % 32))) == 0)
+				{
+					// Set flag for deactivated instances
+					InstanceDesc.AccelerationStructureIndex = 0xFFFFFFFF;
+					NumInactiveNativeInstancesThisSceneInstance++;
+					continue;
+				}
+
+				if (LIKELY(bCpuInstance))
+				{
+					const FMatrix44f LocalToWorld = SceneInstance.Transforms[TransformIndex].GetTransposed();
+					InstanceDesc.LocalToWorld[0] = *(FVector4f*)&LocalToWorld.M[0];
+					InstanceDesc.LocalToWorld[1] = *(FVector4f*)&LocalToWorld.M[1];
+					InstanceDesc.LocalToWorld[2] = *(FVector4f*)&LocalToWorld.M[2];
+				}
+				else
+				{
+					// #todo: GPU-based instances transform should be copied from GPUTransformsSRV
+					FMemory::Memset(&InstanceDesc.LocalToWorld, 0, sizeof(InstanceDesc.LocalToWorld));
+				}
+			}
+
+#if STATS
+			FPlatformAtomics::InterlockedAdd(&NumInactiveNativeInstances, NumInactiveNativeInstancesThisSceneInstance);
+#endif
+		});
+
+	SET_DWORD_STAT(STAT_RayTracingInstances, SceneInitializer.NumNativeInstances - NumInactiveNativeInstances);
+}
+
+struct FRayTracingInstanceCopyCS : public FGlobalShader
 {
-	return (ShouldCompileRayTracingShadersForProject(Parameters.Platform) && RHISupportsComputeShaders(Parameters.Platform) && !(Parameters.Platform == EShaderPlatform::SP_METAL || Parameters.Platform == EShaderPlatform::SP_METAL_TVOS || IsMobilePlatform(Parameters.Platform)));
-}
\ No newline at end of file
+	DECLARE_GLOBAL_SHADER(FRayTracingInstanceCopyCS);
+	SHADER_USE_PARAMETER_STRUCT(FRayTracingInstanceCopyCS, FGlobalShader);
+
+	BEGIN_SHADER_PARAMETER_STRUCT(FParameters, )
+		SHADER_PARAMETER_RDG_BUFFER_UAV(RWStructuredBuffer, InstancesDescriptors)
+		SHADER_PARAMETER_RDG_BUFFER_SRV(StructuredBuffer, InstancesTransforms)
+		SHADER_PARAMETER(uint32, NumInstances)
+		SHADER_PARAMETER(uint32, DescBufferOffset)
+	END_SHADER_PARAMETER_STRUCT()
+
+	static constexpr uint32 ThreadGroupSize = 64;
+	static inline void ModifyCompilationEnvironment(const FGlobalShaderPermutationParameters& Parameters, FShaderCompilerEnvironment& OutEnvironment)
+	{
+		FGlobalShader::ModifyCompilationEnvironment(Parameters, OutEnvironment);
+
+		OutEnvironment.SetDefine(TEXT("THREADGROUP_SIZE"), ThreadGroupSize);
+	}
+
+	static bool ShouldCompilePermutation(const FGlobalShaderPermutationParameters& Parameters)
+	{
+		return (ShouldCompileRayTracingShadersForProject(Parameters.Platform) && RHISupportsComputeShaders(Parameters.Platform) && !(Parameters.Platform == EShaderPlatform::SP_METAL || Parameters.Platform == EShaderPlatform::SP_METAL_TVOS || IsMobilePlatform(Parameters.Platform)));
+	}
+};
+
+IMPLEMENT_GLOBAL_SHADER(FRayTracingInstanceCopyCS, "/Engine/Private/Raytracing/RayTracingInstanceCopy.usf", "RayTracingInstanceCopyShaderCS", SF_Compute);
+
+struct FRayTracingInstanceBufferCS : public FGlobalShader
+{
+	DECLARE_GLOBAL_SHADER(FRayTracingInstanceBufferCS);
+	SHADER_USE_PARAMETER_STRUCT(FRayTracingInstanceBufferCS, FGlobalShader);
+
+	BEGIN_SHADER_PARAMETER_STRUCT(FParameters, )
+		SHADER_PARAMETER_UAV(RWStructuredBuffer, InstancesDescriptors)
+		SHADER_PARAMETER_SRV(StructuredBuffer<FRayTracingInstanceDescriptorInput>, InputInstanceDescriptors)
+		SHADER_PARAMETER_SRV(ByteAddressBuffer, AccelerationStructureAddresses)
+		SHADER_PARAMETER(uint32, NumInstances)
+	END_SHADER_PARAMETER_STRUCT()
+		
+	static constexpr uint32 ThreadGroupSize = 64;
+
+	static inline void ModifyCompilationEnvironment(const FGlobalShaderPermutationParameters& Parameters, FShaderCompilerEnvironment& OutEnvironment)
+	{
+		FGlobalShader::ModifyCompilationEnvironment(Parameters, OutEnvironment);
+
+		OutEnvironment.SetDefine(TEXT("THREADGROUP_SIZE"), ThreadGroupSize);
+	}
+
+	static bool ShouldCompilePermutation(const FGlobalShaderPermutationParameters& Parameters)
+	{
+		return (ShouldCompileRayTracingShadersForProject(Parameters.Platform) && RHISupportsComputeShaders(Parameters.Platform) && !(Parameters.Platform == EShaderPlatform::SP_METAL || Parameters.Platform == EShaderPlatform::SP_METAL_TVOS || IsMobilePlatform(Parameters.Platform)));
+	}
+};
+
+IMPLEMENT_GLOBAL_SHADER(FRayTracingInstanceBufferCS, "/Engine/Private/Raytracing/RayTracingInstanceCopy.usf", "RayTracingBuildInstanceBufferCS", SF_Compute);
+
+void BuildRayTracingInstanceBuffer(FRHICommandList& RHICmdList, uint32 NumInstances, FUnorderedAccessViewRHIRef InstancesUAV, FShaderResourceViewRHIRef InstanceUploadSRV, FShaderResourceViewRHIRef AccelerationStructureAddressesSRV)
+{
+	FRayTracingInstanceBufferCS::FParameters PassParams;
+	PassParams.InstancesDescriptors = InstancesUAV;
+	PassParams.InputInstanceDescriptors = InstanceUploadSRV;
+	PassParams.AccelerationStructureAddresses = AccelerationStructureAddressesSRV;
+	PassParams.NumInstances = NumInstances;
+
+	auto ComputeShader = GetGlobalShaderMap(GMaxRHIFeatureLevel)->GetShader<FRayTracingInstanceBufferCS>();
+	const int32 GroupSize = FMath::DivideAndRoundUp(PassParams.NumInstances, FRayTracingInstanceBufferCS::ThreadGroupSize);
+
+	RHICmdList.SetComputeShader(ComputeShader.GetComputeShader());
+
+	SetShaderParameters(RHICmdList, ComputeShader, ComputeShader.GetComputeShader(), PassParams);
+
+	DispatchComputeShader(RHICmdList, ComputeShader.GetShader(), GroupSize, 1, 1);
+
+	UnsetShaderUAVs(RHICmdList, ComputeShader, ComputeShader.GetComputeShader());
+}
+
+#endif //RHI_RAYTRACING
diff --git a/Engine/Source/Runtime/RenderCore/Public/RayTracingInstanceCopyShader.h b/Engine/Source/Runtime/RenderCore/Public/RayTracingInstanceCopyShader.h
--- a/Engine/Source/Runtime/RenderCore/Public/RayTracingInstanceCopyShader.h
+++ b/Engine/Source/Runtime/RenderCore/Public/RayTracingInstanceCopyShader.h
@@ -2,84 +2,32 @@
 
 #pragma once
 
+#include "RenderGraphDefinitions.h"
 #include "CoreMinimal.h"
-#include "ShaderParameters.h"
-#include "ShaderParameterUtils.h"
-#include "Shader.h"
-#include "GlobalShader.h"
 
-struct RayTracingInstanceCopyCS : public FGlobalShader
-{
-	DECLARE_EXPORTED_SHADER_TYPE(RayTracingInstanceCopyCS, Global, RENDERCORE_API);
-
-public:
-	static constexpr uint32 ThreadGroupSize = 64;
-
-	RayTracingInstanceCopyCS() {}
-	RayTracingInstanceCopyCS(const ShaderMetaType::CompiledShaderInitializerType& Initializer)
-		: FGlobalShader(Initializer)
-	{
-		InstancesCPUCountParam.Bind(Initializer.ParameterMap, TEXT("InstancesCPUCount"), SPF_Mandatory);
-		DescBufferOffsetParam.Bind(Initializer.ParameterMap, TEXT("DescBufferOffset"), SPF_Mandatory);
-		InstancesTransformsParam.Bind(Initializer.ParameterMap, TEXT("InstancesTransforms"), SPF_Mandatory);
-		InstancesDescriptorsParam.Bind(Initializer.ParameterMap, TEXT("InstancesDescriptors"), SPF_Mandatory);
-	}
-
-	static const TCHAR* GetSourceFilename() { return TEXT("/Engine/Private/Raytracing/RayTracingInstanceCopy.usf"); }	
-	static const TCHAR* GetFunctionName() { return TEXT("RayTracingInstanceCopyShaderCS"); }
-
-	static bool ShouldCompilePermutation(const FGlobalShaderPermutationParameters& Parameters);
-	static inline void ModifyCompilationEnvironment(const FGlobalShaderPermutationParameters& Parameters, FShaderCompilerEnvironment& OutEnvironment)
-	{
-		FGlobalShader::ModifyCompilationEnvironment(Parameters, OutEnvironment);
-
-		OutEnvironment.SetDefine(TEXT("THREADGROUP_SIZE"), ThreadGroupSize);
-	}
-
-	const FShaderParameter& GetInstancesCountParam() const { return InstancesCPUCountParam; }
-	inline const FShaderParameter& GetInstancesDescBufferOffsetParam() const { return DescBufferOffsetParam; }
-	inline const FShaderResourceParameter& GetInstancesTransformsParam() const { return InstancesTransformsParam; }	
-	inline const FShaderResourceParameter& GetInstancesDescriptorsParam() const { return InstancesDescriptorsParam; }
+#if RHI_RAYTRACING
 
-private:
-	LAYOUT_FIELD(FShaderParameter, InstancesCPUCountParam);
-	LAYOUT_FIELD(FShaderParameter, DescBufferOffsetParam);
-	LAYOUT_FIELD(FShaderResourceParameter, InstancesTransformsParam);
-	LAYOUT_FIELD(FShaderResourceParameter, InstancesDescriptorsParam);
-};
-
-/**
- * CS can be dispatched from inside low level RHIs via FRHICommandList_RecursiveHazardous. 
- * ResourceBindCallback is provided to allow the RHI to override how the UAV resource is bound to the underlying platform context.
- */
-inline void CopyRayTracingGPUInstances(FRHICommandList& RHICmdList, uint32 InstancesCount, int32 DescBufferOffset, FRHIShaderResourceView* TransformsSRV, FRHIUnorderedAccessView* InstancesDescUAV, TFunctionRef<void(FRHIComputeShader*, const FShaderResourceParameter&, const FShaderResourceParameter&, bool)> ResourceBindCallback)
+struct FRayTracingInstanceDescriptorInput
 {
-	TShaderMapRef<RayTracingInstanceCopyCS> ComputeShader(GetGlobalShaderMap(GMaxRHIFeatureLevel));
-	FRHIComputeShader* ShaderRHI = ComputeShader.GetComputeShader();
-	RHICmdList.SetComputeShader(ShaderRHI);
-
-	SetShaderValue(RHICmdList, ShaderRHI, ComputeShader->GetInstancesCountParam(), InstancesCount);
-	SetShaderValue(RHICmdList, ShaderRHI, ComputeShader->GetInstancesDescBufferOffsetParam(), DescBufferOffset);
-
-	ResourceBindCallback(ShaderRHI, ComputeShader->GetInstancesTransformsParam(), ComputeShader->GetInstancesDescriptorsParam(), true);
+	//uint32 GPUSceneInstanceIndex;
+	FVector4f LocalToWorld[3];
+	uint32 AccelerationStructureIndex;
+	uint32 InstanceId;
+	uint32 InstanceMaskAndFlags;
+	uint32 InstanceContributionToHitGroupIndex;
+};
 
-	uint32 NGroups = FMath::DivideAndRoundUp(InstancesCount, RayTracingInstanceCopyCS::ThreadGroupSize);
-	RHICmdList.DispatchComputeShader(NGroups, 1, 1);
+RENDERCORE_API void FillInstanceUploadBuffer(
+	TConstArrayView<FRayTracingGeometryInstance> Instances,
+	TConstArrayView<uint32> InstancesGeometryIndex,
+	FRayTracingSceneRHIRef RayTracingSceneRHI,
+	TArrayView<FRayTracingInstanceDescriptorInput> OutInstanceUploadData);
 
-	ResourceBindCallback(ShaderRHI, ComputeShader->GetInstancesTransformsParam(), ComputeShader->GetInstancesDescriptorsParam(), false);
-}
+RENDERCORE_API void BuildRayTracingInstanceBuffer(
+	FRHICommandList& RHICmdList, 
+	uint32 NumInstances,
+	FUnorderedAccessViewRHIRef InstancesUAV,
+	FShaderResourceViewRHIRef InstanceUploadSRV,
+	FShaderResourceViewRHIRef AccelerationStructureAddressesSRV);
 
-inline void CopyRayTracingGPUInstances(FRHICommandList& RHICmdList, int32 InstancesCount, int32 DescBufferOffset, FRHIShaderResourceView* TransformsSRV, FRHIUnorderedAccessView* InstancesDescUAV)
-{
-	CopyRayTracingGPUInstances(RHICmdList, InstancesCount, DescBufferOffset, TransformsSRV, InstancesDescUAV,
-		[&RHICmdList, TransformsSRV, InstancesDescUAV](
-			FRHIComputeShader* ShaderRHI, 
-			const FShaderResourceParameter& InstancesTransformsParam,
-			const FShaderResourceParameter& InstancesDescriptorsParam, 
-			bool bSet)
-	{
-		SetSRVParameter(RHICmdList, ShaderRHI, InstancesTransformsParam, bSet ? TransformsSRV : nullptr);
-		SetUAVParameter(RHICmdList, ShaderRHI, InstancesDescriptorsParam, bSet ? InstancesDescUAV : nullptr);
-	}
-	);
-}
+#endif // RHI_RAYTRACING
\ No newline at end of file
diff --git a/Engine/Source/Runtime/RenderCore/RenderCore.Build.cs b/Engine/Source/Runtime/RenderCore/RenderCore.Build.cs
--- a/Engine/Source/Runtime/RenderCore/RenderCore.Build.cs
+++ b/Engine/Source/Runtime/RenderCore/RenderCore.Build.cs
@@ -5,8 +5,10 @@ using UnrealBuildTool;
 public class RenderCore : ModuleRules
 {
 	public RenderCore(ReadOnlyTargetRules Target) : base(Target)
-    {
-        PublicDependencyModuleNames.AddRange(new string[] { "RHI" });
+	{
+		PrivateIncludePaths.Add("../Shaders/Shared");
+
+		PublicDependencyModuleNames.AddRange(new string[] { "RHI" });
 
 		PrivateIncludePathModuleNames.AddRange(new string[] { "TargetPlatform" });
 
diff --git a/Engine/Source/Runtime/Renderer/Private/DeferredShadingRenderer.cpp b/Engine/Source/Runtime/Renderer/Private/DeferredShadingRenderer.cpp
--- a/Engine/Source/Runtime/Renderer/Private/DeferredShadingRenderer.cpp
+++ b/Engine/Source/Runtime/Renderer/Private/DeferredShadingRenderer.cpp
@@ -1312,7 +1312,8 @@ static void DeduplicateRayGenerationShaders(TArray< FRHIRayTracingShader*>& RayG
 }
 
 BEGIN_SHADER_PARAMETER_STRUCT(FBuildAccelerationStructurePassParams, )
-RDG_BUFFER_ACCESS(RayTracingSceneScratchBuffer, ERHIAccess::UAVCompute)
+	RDG_BUFFER_ACCESS(RayTracingSceneScratchBuffer, ERHIAccess::UAVCompute)
+	RDG_BUFFER_ACCESS(RayTracingSceneInstanceBuffer, ERHIAccess::SRVCompute)
 END_SHADER_PARAMETER_STRUCT()
 
 bool FDeferredShadingSceneRenderer::SetupRayTracingPipelineStates(FRHICommandListImmediate& RHICmdList)
@@ -1478,37 +1479,53 @@ bool FDeferredShadingSceneRenderer::DispatchRayTracingWorldUpdates(FRDGBuilder&
 
 	RDG_GPU_MASK_SCOPE(GraphBuilder, FRHIGPUMask::All());
 
-	RayTracingScene.BeginCreate(GraphBuilder);
+	RayTracingScene.Create(GraphBuilder);
 
 	const bool bRayTracingAsyncBuild = CVarRayTracingAsyncBuild.GetValueOnRenderThread() != 0;
 
 	if (bRayTracingAsyncBuild && GRHISupportsRayTracingAsyncBuildAccelerationStructure)
 	{
-		AddPass(GraphBuilder, RDG_EVENT_NAME("BuildAccelerationStructure"), [this](FRHICommandListImmediate& RHICmdList)
-		{
-			check(RayTracingDynamicGeometryUpdateEndTransition == nullptr);
-			const FRHITransition* RayTracingDynamicGeometryUpdateBeginTransition = RHICreateTransition(FRHITransitionCreateInfo(ERHIPipeline::Graphics, ERHIPipeline::AsyncCompute));
-			RayTracingDynamicGeometryUpdateEndTransition = RHICreateTransition(FRHITransitionCreateInfo(ERHIPipeline::AsyncCompute, ERHIPipeline::Graphics));
+		FBuildAccelerationStructurePassParams* PassParams = GraphBuilder.AllocParameters<FBuildAccelerationStructurePassParams>();
+		PassParams->RayTracingSceneScratchBuffer = Scene->RayTracingScene.BuildScratchBuffer;
+		PassParams->RayTracingSceneInstanceBuffer = Scene->RayTracingScene.InstanceBuffer;
 
-			FRHIAsyncComputeCommandListImmediate& RHIAsyncCmdList = FRHICommandListExecutor::GetImmediateAsyncComputeCommandList();
+		GraphBuilder.AddPass(RDG_EVENT_NAME("RayTracingScene"), PassParams, ERDGPassFlags::Compute | ERDGPassFlags::NeverCull,
+			[this, PassParams](FRHICommandListImmediate& RHICmdList)
+			{
+				FRHIAsyncComputeCommandListImmediate& RHIAsyncCmdList = FRHICommandListExecutor::GetImmediateAsyncComputeCommandList();
 
-			RHICmdList.BeginTransition(RayTracingDynamicGeometryUpdateBeginTransition);
-			RHIAsyncCmdList.EndTransition(RayTracingDynamicGeometryUpdateBeginTransition);
+				const FRHITransition* RayTracingDynamicGeometryUpdateBeginTransition = RHICreateTransition(FRHITransitionCreateInfo(ERHIPipeline::Graphics, ERHIPipeline::AsyncCompute));
+				RHICmdList.BeginTransition(RayTracingDynamicGeometryUpdateBeginTransition);
+				RHIAsyncCmdList.EndTransition(RayTracingDynamicGeometryUpdateBeginTransition);
 
-			Scene->GetRayTracingDynamicGeometryCollection()->DispatchUpdates(RHIAsyncCmdList);
+				Scene->GetRayTracingDynamicGeometryCollection()->DispatchUpdates(RHIAsyncCmdList);
 
-			FRHIRayTracingScene* RayTracingSceneRHI = Scene->RayTracingScene.GetRHIRayTracingSceneChecked();
+				FRHIRayTracingScene* RayTracingSceneRHI = Scene->RayTracingScene.GetRHIRayTracingSceneChecked();
+				FRHIBuffer* AccelerationStructureBuffer = Scene->RayTracingScene.GetBufferChecked();
+				FRHIBuffer* InstanceBuffer = PassParams->RayTracingSceneInstanceBuffer->GetRHI();
+				FRHIBuffer* ScratchBuffer = PassParams->RayTracingSceneScratchBuffer->GetRHI();
 
-			RHIAsyncCmdList.BindAccelerationStructureMemory(RayTracingSceneRHI, Scene->RayTracingScene.GetBufferChecked(), 0);
+				RHIAsyncCmdList.BindAccelerationStructureMemory(RayTracingSceneRHI, AccelerationStructureBuffer, 0);
 
-			{
-				SCOPED_DRAW_EVENT(RHIAsyncCmdList, RayTracingScene);
-				RHIAsyncCmdList.BuildAccelerationStructure(RayTracingSceneRHI);				
-			}
+				{
+					SCOPED_DRAW_EVENT(RHIAsyncCmdList, RayTracingScene);
 
-			RHIAsyncCmdList.BeginTransition(RayTracingDynamicGeometryUpdateEndTransition);
-			FRHIAsyncComputeCommandListImmediate::ImmediateDispatch(RHIAsyncCmdList);
-		});
+					FRayTracingSceneBuildParams BuildParams;
+					BuildParams.Scene = RayTracingSceneRHI;
+					BuildParams.ScratchBuffer = ScratchBuffer;
+					BuildParams.ScratchBufferOffset = 0;
+					BuildParams.InstanceBuffer = InstanceBuffer;
+					BuildParams.InstanceBufferOffset = 0;
+
+					RHIAsyncCmdList.BuildAccelerationStructure(BuildParams);
+				}
+
+				check(RayTracingDynamicGeometryUpdateEndTransition == nullptr)
+				RayTracingDynamicGeometryUpdateEndTransition = RHICreateTransition(FRHITransitionCreateInfo(ERHIPipeline::AsyncCompute, ERHIPipeline::Graphics));
+				RHIAsyncCmdList.BeginTransition(RayTracingDynamicGeometryUpdateEndTransition);
+
+				FRHIAsyncComputeCommandListImmediate::ImmediateDispatch(RHIAsyncCmdList);
+			});
 	}
 	else
 	{
@@ -1525,6 +1542,7 @@ bool FDeferredShadingSceneRenderer::DispatchRayTracingWorldUpdates(FRDGBuilder&
 
 			FBuildAccelerationStructurePassParams* PassParams = GraphBuilder.AllocParameters<FBuildAccelerationStructurePassParams>();
 			PassParams->RayTracingSceneScratchBuffer = Scene->RayTracingScene.BuildScratchBuffer;
+			PassParams->RayTracingSceneInstanceBuffer = Scene->RayTracingScene.InstanceBuffer;
 
 			GraphBuilder.AddPass(RDG_EVENT_NAME("RayTracingScene"), PassParams, ERDGPassFlags::Compute | ERDGPassFlags::NeverCull,
 				[this, PassParams](FRHICommandListImmediate& RHICmdList)
@@ -1532,17 +1550,20 @@ bool FDeferredShadingSceneRenderer::DispatchRayTracingWorldUpdates(FRDGBuilder&
 				FRHIRayTracingScene* RayTracingSceneRHI = Scene->RayTracingScene.GetRHIRayTracingSceneChecked();
 				FRHIBuffer* AccelerationStructureBuffer = Scene->RayTracingScene.GetBufferChecked();
 				FRHIBuffer* ScratchBuffer = PassParams->RayTracingSceneScratchBuffer->GetRHI();
+				FRHIBuffer* InstanceBuffer = PassParams->RayTracingSceneInstanceBuffer->GetRHI();
 
 				FRayTracingSceneBuildParams BuildParams;
 				BuildParams.Scene = RayTracingSceneRHI;
 				BuildParams.ScratchBuffer = ScratchBuffer;
 				BuildParams.ScratchBufferOffset = 0;
+				BuildParams.InstanceBuffer = InstanceBuffer;
+				BuildParams.InstanceBufferOffset = 0;
 
 				// Sanity check acceleration structure buffer sizes
 			#if DO_CHECK
 				{
 					FRayTracingAccelerationStructureSize SizeInfo = RHICalcRayTracingSceneSize(
-						Scene->RayTracingScene.NumNativeInstances, ERayTracingAccelerationStructureFlags::FastTrace);
+						RayTracingSceneRHI->GetInitializer().NumNativeInstances, ERayTracingAccelerationStructureFlags::FastTrace);
 
 					check(SizeInfo.ResultSize <= Scene->RayTracingScene.SizeInfo.ResultSize);
 					check(SizeInfo.BuildScratchSize <= Scene->RayTracingScene.SizeInfo.BuildScratchSize);
@@ -1675,8 +1696,6 @@ void FDeferredShadingSceneRenderer::WaitForRayTracingScene(FRDGBuilder& GraphBui
 				}
 			}
 
-			Scene->RayTracingScene.WaitForTasks();
-
 			const bool bCopyDataToInlineStorage = false; // Storage is already allocated from RHICmdList, no extra copy necessary
 			RHICmdList.SetRayTracingHitGroups(
 				ReferenceView.GetRayTracingSceneChecked(),
diff --git a/Engine/Source/Runtime/Renderer/Private/RayTracing/RayTracingScene.cpp b/Engine/Source/Runtime/Renderer/Private/RayTracing/RayTracingScene.cpp
--- a/Engine/Source/Runtime/Renderer/Private/RayTracing/RayTracingScene.cpp
+++ b/Engine/Source/Runtime/Renderer/Private/RayTracing/RayTracingScene.cpp
@@ -4,11 +4,18 @@
 
 #if RHI_RAYTRACING
 
+#include "RayTracingInstanceCopyShader.h"
 #include "RenderCore.h"
 #include "RayTracingDefinitions.h"
 #include "RenderGraphBuilder.h"
+#include "RenderGraphUtils.h"
+
 #include "Experimental/Containers/SherwoodHashTable.h"
 
+BEGIN_SHADER_PARAMETER_STRUCT(FBuildInstanceBufferPassParams, )
+	SHADER_PARAMETER_RDG_BUFFER_UAV(RWStructuredBuffer, InstanceBuffer)
+END_SHADER_PARAMETER_STRUCT()
+
 FRayTracingScene::FRayTracingScene()
 {
 
@@ -19,68 +26,71 @@ FRayTracingScene::~FRayTracingScene()
 	WaitForTasks();
 }
 
-FGraphEventRef FRayTracingScene::BeginCreate(FRDGBuilder& GraphBuilder)
+void FRayTracingScene::Create(FRDGBuilder& GraphBuilder)
 {
-	WaitForTasks();
-
-	NumNativeInstances = 0;
-	NumTotalSegments = 0;
-
-	const uint32 NumSceneInstances = Instances.Num();
+	QUICK_SCOPE_CYCLE_COUNTER(FRayTracingScene_BeginCreate);
 
-	FRayTracingSceneInitializer2 SceneInitializer;
-	SceneInitializer.DebugName = FName(TEXT("FRayTracingScene"));
-	SceneInitializer.ShaderSlotsPerGeometrySegment = RAY_TRACING_NUM_SHADER_SLOTS;
-	SceneInitializer.NumMissShaderSlots = RAY_TRACING_NUM_MISS_SHADER_SLOTS;
-	SceneInitializer.Instances = Instances;
-	SceneInitializer.PerInstanceGeometries.SetNumUninitialized(NumSceneInstances);
-	SceneInitializer.PerInstanceNumTransforms.SetNumUninitialized(NumSceneInstances);
-	SceneInitializer.BaseInstancePrefixSum.SetNumUninitialized(NumSceneInstances);
-	SceneInitializer.SegmentPrefixSum.SetNumUninitialized(NumSceneInstances);
-
-	Experimental::TSherwoodSet<FRHIRayTracingGeometry*> UniqueGeometries;
+	WaitForTasks();
 
-	// Compute geometry segment and instance count prefix sums.
-	// These are later used by GetHitRecordBaseIndex() during resource binding
-	// and by GetBaseInstanceIndex() in shaders to emulate SV_InstanceIndex.
+	TArray<uint32> InstancesGeometryIndex;
 
-	for (uint32 InstanceIndex = 0; InstanceIndex < NumSceneInstances; ++InstanceIndex)
 	{
-		const FRayTracingGeometryInstance& InstanceDesc = Instances[InstanceIndex];
+		const uint32 NumSceneInstances = Instances.Num();
+
+		FRayTracingSceneInitializer2 SceneInitializer;
+		SceneInitializer.DebugName = FName(TEXT("FRayTracingScene"));
+		SceneInitializer.ShaderSlotsPerGeometrySegment = RAY_TRACING_NUM_SHADER_SLOTS;
+		SceneInitializer.NumMissShaderSlots = RAY_TRACING_NUM_MISS_SHADER_SLOTS;
+		SceneInitializer.PerInstanceGeometries.SetNumUninitialized(NumSceneInstances);
+		SceneInitializer.BaseInstancePrefixSum.SetNumUninitialized(NumSceneInstances);
+		SceneInitializer.SegmentPrefixSum.SetNumUninitialized(NumSceneInstances);
+		SceneInitializer.NumNativeInstances = 0;
+		SceneInitializer.NumTotalSegments = 0;
+
+		Experimental::TSherwoodMap<FRHIRayTracingGeometry*, uint32> UniqueGeometries;
+		InstancesGeometryIndex.SetNumUninitialized(NumSceneInstances);
+
+		// Compute geometry segment and instance count prefix sums.
+		// These are later used by GetHitRecordBaseIndex() during resource binding
+		// and by GetBaseInstanceIndex() in shaders to emulate SV_InstanceIndex.
+
+		for (uint32 InstanceIndex = 0; InstanceIndex < NumSceneInstances; ++InstanceIndex)
+		{
+			const FRayTracingGeometryInstance& InstanceDesc = Instances[InstanceIndex];
 
-		checkf(InstanceDesc.GPUTransformsSRV || InstanceDesc.NumTransforms <= uint32(InstanceDesc.Transforms.Num()),
-			TEXT("Expected at most %d ray tracing geometry instance transforms, but got %d."),
-			InstanceDesc.NumTransforms, InstanceDesc.Transforms.Num());
+			checkf(InstanceDesc.GPUTransformsSRV || InstanceDesc.NumTransforms <= uint32(InstanceDesc.Transforms.Num()),
+				TEXT("Expected at most %d ray tracing geometry instance transforms, but got %d."),
+				InstanceDesc.NumTransforms, InstanceDesc.Transforms.Num());
 
-		checkf(InstanceDesc.GeometryRHI, TEXT("Ray tracing instance must have a valid geometry."));
+			checkf(InstanceDesc.GeometryRHI, TEXT("Ray tracing instance must have a valid geometry."));
 
-		SceneInitializer.PerInstanceGeometries[InstanceIndex] = InstanceDesc.GeometryRHI;
+			SceneInitializer.PerInstanceGeometries[InstanceIndex] = InstanceDesc.GeometryRHI;
 
-		// Compute geometry segment count prefix sum to be later used in GetHitRecordBaseIndex()
-		SceneInitializer.SegmentPrefixSum[InstanceIndex] = NumTotalSegments;
-		NumTotalSegments += InstanceDesc.GeometryRHI->GetNumSegments();
+			// Compute geometry segment count prefix sum to be later used in GetHitRecordBaseIndex()
+			SceneInitializer.SegmentPrefixSum[InstanceIndex] = SceneInitializer.NumTotalSegments;
+			SceneInitializer.NumTotalSegments += InstanceDesc.GeometryRHI->GetNumSegments();
 
-		bool bIsAlreadyInSet = false;
-		UniqueGeometries.Add(InstanceDesc.GeometryRHI, &bIsAlreadyInSet);
-		if (!bIsAlreadyInSet)
-		{
-			SceneInitializer.ReferencedGeometries.Add(InstanceDesc.GeometryRHI);
-		}
+			uint32 GeometryIndex = UniqueGeometries.FindOrAdd(InstanceDesc.GeometryRHI, SceneInitializer.ReferencedGeometries.Num());
+			InstancesGeometryIndex[InstanceIndex] = GeometryIndex;
+			if (GeometryIndex == SceneInitializer.ReferencedGeometries.Num())
+			{
+				SceneInitializer.ReferencedGeometries.Add(InstanceDesc.GeometryRHI);
+			}
 
-		SceneInitializer.BaseInstancePrefixSum[InstanceIndex] = NumNativeInstances;
-		NumNativeInstances += InstanceDesc.NumTransforms;
+			SceneInitializer.BaseInstancePrefixSum[InstanceIndex] = SceneInitializer.NumNativeInstances;
+			SceneInitializer.NumNativeInstances += InstanceDesc.NumTransforms;
+		}
 
-		SceneInitializer.PerInstanceNumTransforms[InstanceIndex] = InstanceDesc.NumTransforms;
+		RayTracingSceneRHI = RHICreateRayTracingScene(MoveTemp(SceneInitializer));
 	}
 
-	SceneInitializer.NumNativeInstances = NumNativeInstances;
-	SceneInitializer.NumTotalSegments = NumTotalSegments;
+	const FRayTracingSceneInitializer2& SceneInitializer = RayTracingSceneRHI->GetInitializer();
 
 	// Round up number of instances to some multiple to avoid pathological growth reallocations.
 	static constexpr uint32 AllocationGranularity = 8 * 1024;
-	uint32 NumNativeInstancesAligned = FMath::DivideAndRoundUp(NumNativeInstances, AllocationGranularity) * AllocationGranularity;
+	uint32 NumNativeInstancesAligned = FMath::DivideAndRoundUp(FMath::Max(SceneInitializer.NumNativeInstances, 1U), AllocationGranularity) * AllocationGranularity;
 
-	SizeInfo = RHICalcRayTracingSceneSize(NumNativeInstances, ERayTracingAccelerationStructureFlags::FastTrace);
+	SizeInfo = RHICalcRayTracingSceneSize(SceneInitializer.NumNativeInstances, ERayTracingAccelerationStructureFlags::FastTrace);
 	FRayTracingAccelerationStructureSize SizeInfoAligned = RHICalcRayTracingSceneSize(NumNativeInstancesAligned, ERayTracingAccelerationStructureFlags::FastTrace);
 	SizeInfo.ResultSize = FMath::Max(SizeInfo.ResultSize, SizeInfoAligned.ResultSize);
 	SizeInfo.BuildScratchSize = FMath::Max(SizeInfo.BuildScratchSize, SizeInfoAligned.BuildScratchSize);
@@ -100,54 +110,130 @@ FGraphEventRef FRayTracingScene::BeginCreate(FRDGBuilder& GraphBuilder)
 		RayTracingSceneSRV = RHICreateShaderResourceView(RayTracingSceneBuffer);
 	}
 
-	RayTracingSceneRHI.SafeRelease();
+	{
+		const uint64 ScratchAlignment = GRHIRayTracingAccelerationStructureAlignment;
+		FRDGBufferDesc ScratchBufferDesc;
+		ScratchBufferDesc.UnderlyingType = FRDGBufferDesc::EUnderlyingType::StructuredBuffer;
+		ScratchBufferDesc.Usage = BUF_UnorderedAccess;
+		ScratchBufferDesc.BytesPerElement = uint32(ScratchAlignment);
+		ScratchBufferDesc.NumElements = uint32(FMath::DivideAndRoundUp(SizeInfo.BuildScratchSize, ScratchAlignment));
+
+		BuildScratchBuffer = GraphBuilder.CreateBuffer(ScratchBufferDesc, TEXT("RayTracingScratchBuffer"));
+	}
+
+	{
+		FRDGBufferDesc InstanceBufferDesc;
+		InstanceBufferDesc.UnderlyingType = FRDGBufferDesc::EUnderlyingType::StructuredBuffer;
+		InstanceBufferDesc.Usage = BUF_UnorderedAccess | BUF_ShaderResource;
+		InstanceBufferDesc.BytesPerElement = GRHIRayTracingInstanceDescriptorSize;
+		InstanceBufferDesc.NumElements = NumNativeInstancesAligned;
+
+		InstanceBuffer = GraphBuilder.CreateBuffer(InstanceBufferDesc, TEXT("FRayTracingScene::InstanceBuffer"));
+	}
+
+	{
+		// Round to PoT to avoid resizing too often
+		const uint32 NumGeometries = FMath::RoundUpToPowerOfTwo(SceneInitializer.ReferencedGeometries.Num());
+		const uint32 AccelerationStructureAddressesBufferSize = NumGeometries * sizeof(FRayTracingAccelerationStructureAddress);
 
-	CreateRayTracingSceneTask = FFunctionGraphTask::CreateAndDispatchWhenReady([
-			&ResultScene = RayTracingSceneRHI,
-			SceneInitializer = MoveTemp(SceneInitializer)]() mutable
+		if (AccelerationStructureAddressesBuffer.NumBytes < AccelerationStructureAddressesBufferSize)
 		{
-			FTaskTagScope TaskTagScope(ETaskTag::EParallelRenderingThread);
+			AccelerationStructureAddressesBuffer.Initialize(TEXT("RayTracingAccelerationStructureAddressesBuffer"), AccelerationStructureAddressesBufferSize, BUF_Volatile);
+		}
+	}
 
-			ResultScene = RHICreateRayTracingScene(MoveTemp(SceneInitializer));
+	{
+		// Lock instance upload buffer (and resize if necessary)
+		const uint32 UploadBufferSize = NumNativeInstancesAligned * sizeof(FRayTracingInstanceDescriptorInput);
 
-		}, TStatId(), nullptr, ENamedThreads::AnyThread);
+		if (!InstanceUploadBuffer.IsValid()
+			|| UploadBufferSize > InstanceUploadBuffer->GetSize()
+			|| UploadBufferSize < InstanceUploadBuffer->GetSize() / 2)
+		{
+			FRHIResourceCreateInfo CreateInfo(TEXT("RayTracingSceneInstanceUploadBuffer"));
+			InstanceUploadBuffer = RHICreateStructuredBuffer(sizeof(FRayTracingInstanceDescriptorInput), UploadBufferSize, BUF_ShaderResource | BUF_Volatile, CreateInfo);
+			InstanceUploadSRV = RHICreateShaderResourceView(InstanceUploadBuffer);
+		}
+	}
+
+	if (SceneInitializer.NumNativeInstances > 0)
+	{
+		const uint32 UploadBytes = SceneInitializer.NumNativeInstances * sizeof(FRayTracingInstanceDescriptorInput);
 
-	const uint64 ScratchAlignment = GRHIRayTracingAccelerationStructureAlignment;
-	FRDGBufferDesc ScratchBufferDesc;
-	ScratchBufferDesc.UnderlyingType = FRDGBufferDesc::EUnderlyingType::StructuredBuffer;
-	ScratchBufferDesc.Usage = BUF_UnorderedAccess;
-	ScratchBufferDesc.BytesPerElement = uint32(ScratchAlignment);
-	ScratchBufferDesc.NumElements = uint32(FMath::DivideAndRoundUp(SizeInfo.BuildScratchSize, ScratchAlignment));
+		FRayTracingInstanceDescriptorInput* InstanceUploadData = (FRayTracingInstanceDescriptorInput*)RHILockBuffer(InstanceUploadBuffer, 0, UploadBytes, RLM_WriteOnly);
 
-	BuildScratchBuffer = GraphBuilder.CreateBuffer(ScratchBufferDesc, TEXT("FRayTracingScene::BuildScratchBuffer"));
+		// Fill instance upload buffer on separate thread since results are only needed in RHI thread
+		FillInstanceUploadBufferTask = FFunctionGraphTask::CreateAndDispatchWhenReady(
+			[InstanceUploadData,
+			NumNativeInstances = SceneInitializer.NumNativeInstances,
+			Instances = MakeArrayView(Instances),
+			InstancesGeometryIndex = MoveTemp(InstancesGeometryIndex),
+			RayTracingSceneRHI = RayTracingSceneRHI]()
+		{
+			FTaskTagScope TaskTagScope(ETaskTag::EParallelRenderingThread);
+
+			FillInstanceUploadBuffer(Instances, InstancesGeometryIndex, RayTracingSceneRHI, MakeArrayView(InstanceUploadData, NumNativeInstances));
+		}, TStatId(), nullptr, ENamedThreads::AnyThread);
 
-	return CreateRayTracingSceneTask;
+		FBuildInstanceBufferPassParams* PassParams = GraphBuilder.AllocParameters<FBuildInstanceBufferPassParams>();
+		PassParams->InstanceBuffer = GraphBuilder.CreateUAV(InstanceBuffer);
+
+		GraphBuilder.AddPass(
+			RDG_EVENT_NAME("BuildTLASInstanceBuffer"),
+			PassParams,
+			ERDGPassFlags::Compute,
+			[PassParams, this, &SceneInitializer](FRHICommandListImmediate& RHICmdList)
+			{
+				WaitForTasks();
+				RHICmdList.UnlockBuffer(InstanceUploadBuffer);
+
+				RHICmdList.EnqueueLambda([this, &SceneInitializer](FRHICommandListImmediate& RHICmdList)
+					{
+						QUICK_SCOPE_CYCLE_COUNTER(GetAccelerationStructuresAddresses);
+
+						FRayTracingAccelerationStructureAddress* AddressesPtr = (FRayTracingAccelerationStructureAddress*)RHICmdList.LockBuffer(
+							AccelerationStructureAddressesBuffer.Buffer, 
+							0, 
+							SceneInitializer.ReferencedGeometries.Num() * sizeof(FRayTracingAccelerationStructureAddress), RLM_WriteOnly);
+
+						const uint32 NumGeometries = SceneInitializer.ReferencedGeometries.Num();
+						for (uint32 GeometryIndex = 0; GeometryIndex < NumGeometries; ++GeometryIndex)
+						{
+							AddressesPtr[GeometryIndex] = SceneInitializer.ReferencedGeometries[GeometryIndex]->GetAccelerationStructureAddress(RHICmdList.GetGPUMask().ToIndex());
+						}
+
+						RHICmdList.UnlockBuffer(AccelerationStructureAddressesBuffer.Buffer);
+					});
+
+				::BuildRayTracingInstanceBuffer(
+					RHICmdList,
+					SceneInitializer.NumNativeInstances,
+					PassParams->InstanceBuffer->GetRHI(),
+					InstanceUploadSRV,
+					AccelerationStructureAddressesBuffer.SRV);
+			});
+	}
 }
 
 void FRayTracingScene::WaitForTasks() const
 {
-	if (CreateRayTracingSceneTask.IsValid())
+	if (FillInstanceUploadBufferTask.IsValid())
 	{
-		TRACE_CPUPROFILER_EVENT_SCOPE(WaitForCreateRayTracingScene);
-		FTaskGraphInterface::Get().WaitUntilTaskCompletes(CreateRayTracingSceneTask, ENamedThreads::GetRenderThread_Local());
+		QUICK_SCOPE_CYCLE_COUNTER(WaitForRayTracingSceneFillInstanceUploadBuffer);
+		TRACE_CPUPROFILER_EVENT_SCOPE(WaitForRayTracingSceneFillInstanceUploadBuffer);
+		FTaskGraphInterface::Get().WaitUntilTaskCompletes(FillInstanceUploadBufferTask, ENamedThreads::GetRenderThread_Local());
 
-		CreateRayTracingSceneTask = {};
+		FillInstanceUploadBufferTask = {};
 	}
 }
 
 bool FRayTracingScene::IsCreated() const
 {
-	// Either scene creation task has been started (CreateRayTracingSceneTask != null) or we already have a valid RayTracingSceneRHI.
-	return CreateRayTracingSceneTask.IsValid() || RayTracingSceneRHI.IsValid();
+	return RayTracingSceneRHI.IsValid();
 }
 
 FRHIRayTracingScene* FRayTracingScene::GetRHIRayTracingScene() const
 {
-	WaitForTasks();
-
-	checkf(!CreateRayTracingSceneTask.IsValid() || CreateRayTracingSceneTask->IsComplete(),
-		TEXT("Ray tracing scene creation task is expected to be waited upon and ready when we get here."));
-
 	return RayTracingSceneRHI.GetReference();
 }
 
diff --git a/Engine/Source/Runtime/Renderer/Private/RayTracing/RayTracingScene.h b/Engine/Source/Runtime/Renderer/Private/RayTracing/RayTracingScene.h
--- a/Engine/Source/Runtime/Renderer/Private/RayTracing/RayTracingScene.h
+++ b/Engine/Source/Runtime/Renderer/Private/RayTracing/RayTracingScene.h
@@ -30,12 +30,10 @@ public:
 	FRayTracingScene();
 	~FRayTracingScene();
 
-	// Starts asynchronous creation of RayTracingSceneRHI.
-	// This can be an expensive process, depending on the number of instances in the scene.
-	// Immediately allocates GPU memory to fit at least the current number of instances and updates the SRV.
-	// Returns the task that must be waited upon before accessing RayTracingSceneRHI.
-	FGraphEventRef BeginCreate(FRDGBuilder& GraphBuilder);
-	void WaitForTasks() const;
+	// Creates RayTracingSceneRHI.
+	// Allocates GPU memory to fit at least the current number of instances.
+	// Kicks off instance buffer build to parallel thread along with RDG pass
+	void Create(FRDGBuilder& GraphBuilder);
 
 	// Resets the instance list and reserves memory for this frame.
 	void Reset();
@@ -51,11 +49,11 @@ public:
 		return MakeArrayView(new(Allocator) T[Count], Count);
 	}
 
-	// Returns true if RHI ray tracing scene has been created or asynchronous creation has been kicked.
+	// Returns true if RHI ray tracing scene has been created.
 	// i.e. returns true after BeginCreate() and before Reset().
 	RENDERER_API bool IsCreated() const;
 
-	// Waits for the async creation task if it's active and then returns RayTracingSceneRHI object (may return null).
+	// Returns RayTracingSceneRHI object (may return null).
 	RENDERER_API  FRHIRayTracingScene* GetRHIRayTracingScene() const;
 
 	// Similar to GetRayTracingScene, but checks that ray tracing scene RHI object is valid.
@@ -75,11 +73,6 @@ public:
 	// This must be filled before calling BeginCreateTask().
 	TArray<FRayTracingGeometryInstance> Instances;
 
-	// Total flattened number of ray tracing geometry instances (a single FRayTracingGeometryInstance may represent many).
-	uint32 NumNativeInstances = 0;
-
-	uint32 NumTotalSegments = 0;
-
 	// Geometries which still have a pending build request but are used this frame and require a force build.
 	TArray<const FRayTracingGeometry*> GeometriesToBuild;
 
@@ -88,15 +81,15 @@ public:
 
 	FRayTracingAccelerationStructureSize SizeInfo = {};
 
+	FRDGBufferRef InstanceBuffer;
 	FRDGBufferRef BuildScratchBuffer;
 
 private:
+	void WaitForTasks() const;
+
 	// RHI object that abstracts mesh instnaces in this scene
 	FRayTracingSceneRHIRef RayTracingSceneRHI;
 
-	// Task that asynchronously creates RayTracingSceneRHI
-	mutable FGraphEventRef CreateRayTracingSceneTask;
-
 	// Persistently allocated buffer that holds the built TLAS
 	FBufferRHIRef RayTracingSceneBuffer;
 
@@ -106,6 +99,13 @@ private:
 	// Transient memory allocator
 	FMemStackBase Allocator;
 
+	FBufferRHIRef InstanceUploadBuffer;
+	FShaderResourceViewRHIRef InstanceUploadSRV;
+
+	FByteAddressBuffer AccelerationStructureAddressesBuffer;
+
+	mutable FGraphEventRef FillInstanceUploadBufferTask;
+
 };
 
 #endif // RHI_RAYTRACING
diff --git a/Engine/Source/Runtime/Renderer/Private/SceneRendering.cpp b/Engine/Source/Runtime/Renderer/Private/SceneRendering.cpp
--- a/Engine/Source/Runtime/Renderer/Private/SceneRendering.cpp
+++ b/Engine/Source/Runtime/Renderer/Private/SceneRendering.cpp
@@ -1031,24 +1031,22 @@ bool FViewInfo::HasRayTracingScene() const
 	}
 	return false;
 }
-FRHIRayTracingScene* FViewInfo::GetRayTracingScene() const
+
+FRHIRayTracingScene* FViewInfo::GetRayTracingSceneChecked() const
 {
 	check(Family);
 	if (Family->Scene)
 	{
 		if (FScene* Scene = Family->Scene->GetRenderScene())
 		{
-			return Scene->RayTracingScene.GetRHIRayTracingScene();
+			FRHIRayTracingScene* Result = Scene->RayTracingScene.GetRHIRayTracingScene();
+			checkf(Result, TEXT("Ray tracing scene is expected to be created at this point."));
+			return Result;
 		}
 	}
 	return nullptr;
 }
-FRHIRayTracingScene* FViewInfo::GetRayTracingSceneChecked() const
-{
-	FRHIRayTracingScene* Result = GetRayTracingScene();
-	checkf(Result, TEXT("Ray tracing scene is expected to be created at this point."));
-	return Result;
-}
+
 FRHIShaderResourceView* FViewInfo::GetRayTracingSceneViewChecked() const
 {
 	FRHIShaderResourceView* Result = nullptr;
diff --git a/Engine/Source/Runtime/Renderer/Private/SceneRendering.h b/Engine/Source/Runtime/Renderer/Private/SceneRendering.h
--- a/Engine/Source/Runtime/Renderer/Private/SceneRendering.h
+++ b/Engine/Source/Runtime/Renderer/Private/SceneRendering.h
@@ -1483,7 +1483,6 @@ public:
 
 #if RHI_RAYTRACING
 	bool HasRayTracingScene() const;
-	FRHIRayTracingScene* GetRayTracingScene() const; // Soft-deprecated method, use FScene.RayTracingScene instead.
 	FRHIRayTracingScene* GetRayTracingSceneChecked() const; // Soft-deprecated method, use FScene.RayTracingScene instead.
 	FRHIShaderResourceView* GetRayTracingSceneViewChecked() const; // Soft-deprecated method, use FScene.RayTracingScene instead.
 
diff --git a/Engine/Source/Runtime/Renderer/Private/Tests/RayTracingTestbed.cpp b/Engine/Source/Runtime/Renderer/Private/Tests/RayTracingTestbed.cpp
--- a/Engine/Source/Runtime/Renderer/Private/Tests/RayTracingTestbed.cpp
+++ b/Engine/Source/Runtime/Renderer/Private/Tests/RayTracingTestbed.cpp
@@ -10,9 +10,10 @@ IMPLEMENT_SIMPLE_AUTOMATION_TEST(FRayTracingTestbed, "System.Renderer.RayTracing
 
 #if RHI_RAYTRACING
 
-#include "Containers/DynamicRHIResourceArray.h"
-#include "GlobalShader.h"
+#include "RayTracingInstanceCopyShader.h"
 #include "RayTracingDefinitions.h"
+#include "GlobalShader.h"
+#include "Containers/DynamicRHIResourceArray.h"
 
 // HINT: Execute this test via console command in editor: Automation RunTest System.Renderer.RayTracing.BasicRayTracing
 bool RunRayTracingTestbed_RenderThread(const FString& Parameters)
@@ -114,15 +115,8 @@ bool RunRayTracingTestbed_RenderThread(const FString& Parameters)
 	Instances[0].NumTransforms = NumTransforms;
 	Instances[0].Transforms = MakeArrayView(&FMatrix::Identity, 1);
 
-	ERayTracingAccelerationStructureFlags SceneBuildFlags = ERayTracingAccelerationStructureFlags::FastTrace;
-	FRayTracingAccelerationStructureSize SceneSizeInfo = RHICalcRayTracingSceneSize(1, SceneBuildFlags);
-	FRHIResourceCreateInfo SceneBufferCreateInfo(TEXT("RayTracingTestBedSceneBuffer"));
-	FBufferRHIRef SceneBuffer = RHICreateBuffer(
-		uint32(SceneSizeInfo.ResultSize),
-		BUF_AccelerationStructure, 0 /*Stride*/, ERHIAccess::BVHWrite,
-		SceneBufferCreateInfo);
-
-	FRHICommandListImmediate& RHICmdList = FRHICommandListExecutor::GetImmediateCommandList();
+	uint32 InstancesGeometryIndex[NumInstances] = {};
+	InstancesGeometryIndex[0] = 0;
 
 	FRayTracingSceneRHIRef Scene;
 
@@ -130,17 +124,79 @@ bool RunRayTracingTestbed_RenderThread(const FString& Parameters)
 		FRayTracingSceneInitializer2 Initializer;
 		Initializer.DebugName = FName(TEXT("FRayTracingScene"));
 		Initializer.ShaderSlotsPerGeometrySegment = RAY_TRACING_NUM_SHADER_SLOTS;
-		Initializer.Instances = Instances;
 		Initializer.PerInstanceGeometries.Add(Geometry);
-		Initializer.PerInstanceNumTransforms.Add(NumTransforms);
 		Initializer.BaseInstancePrefixSum.Add(0);
 		Initializer.SegmentPrefixSum.Add(0);
+		Initializer.ReferencedGeometries.Add(Geometry);
 		Initializer.NumNativeInstances = NumInstances * NumTransforms;
 		Initializer.NumTotalSegments = Geometry->GetNumSegments();
 
 		Scene = RHICreateRayTracingScene(MoveTemp(Initializer));
 	}
 
+	const FRayTracingSceneInitializer2& SceneInitializer = Scene->GetInitializer();
+
+	ERayTracingAccelerationStructureFlags SceneBuildFlags = ERayTracingAccelerationStructureFlags::FastTrace;
+	FRayTracingAccelerationStructureSize SceneSizeInfo = RHICalcRayTracingSceneSize(1, SceneBuildFlags);
+	FRHIResourceCreateInfo SceneBufferCreateInfo(TEXT("RayTracingTestBedSceneBuffer"));
+	FBufferRHIRef SceneBuffer = RHICreateBuffer(
+		uint32(SceneSizeInfo.ResultSize),
+		BUF_AccelerationStructure, 0 /*Stride*/, ERHIAccess::BVHWrite,
+		SceneBufferCreateInfo);
+
+	FRHIResourceCreateInfo ScratchBufferCreateInfo(TEXT("RayTracingTestBedScratchBuffer"));
+	FBufferRHIRef ScratchBuffer = RHICreateBuffer(
+		uint32(SceneSizeInfo.BuildScratchSize),
+		BUF_UnorderedAccess, GRHIRayTracingAccelerationStructureAlignment, ERHIAccess::UAVCompute,
+		ScratchBufferCreateInfo);
+
+	FRWBufferStructured InstanceBuffer;
+	InstanceBuffer.Initialize(TEXT("RayTracingTestBedInstanceBuffer"), GRHIRayTracingInstanceDescriptorSize, 1);
+
+	FByteAddressBuffer AccelerationStructureAddressesBuffer;
+	AccelerationStructureAddressesBuffer.Initialize(TEXT("RayTracingTestBedAccelerationStructureAddressesBuffer"), sizeof(FRayTracingAccelerationStructureAddress), BUF_Volatile);
+
+	FBufferRHIRef InstanceUploadBuffer;
+	FShaderResourceViewRHIRef InstanceUploadSRV;
+	{
+		const uint32 UploadBufferSize = 1 * sizeof(FRayTracingInstanceDescriptorInput);
+
+		FRHIResourceCreateInfo CreateInfo(TEXT("RayTracingTestBedInstanceUploadBuffer"));
+		InstanceUploadBuffer = RHICreateStructuredBuffer(sizeof(FRayTracingInstanceDescriptorInput), UploadBufferSize, BUF_ShaderResource | BUF_Volatile, CreateInfo);
+		InstanceUploadSRV = RHICreateShaderResourceView(InstanceUploadBuffer);
+	}
+
+	FRHICommandListImmediate& RHICmdList = FRHICommandListExecutor::GetImmediateCommandList();
+
+	{
+		FRayTracingInstanceDescriptorInput* InstanceUploadData = (FRayTracingInstanceDescriptorInput*)RHICmdList.LockBuffer(InstanceUploadBuffer, 0, sizeof(FRayTracingInstanceDescriptorInput), RLM_WriteOnly);
+		FillInstanceUploadBuffer(Instances, InstancesGeometryIndex, Scene, MakeArrayView(InstanceUploadData, 1));
+		RHICmdList.UnlockBuffer(InstanceUploadBuffer);
+	}
+
+	RHICmdList.EnqueueLambda([&AccelerationStructureAddressesBuffer, &SceneInitializer](FRHICommandListImmediate& RHICmdList)
+		{
+			FRayTracingAccelerationStructureAddress* AddressesPtr = (FRayTracingAccelerationStructureAddress*)RHICmdList.LockBuffer(
+				AccelerationStructureAddressesBuffer.Buffer, 
+				0, 
+				SceneInitializer.ReferencedGeometries.Num() * sizeof(FRayTracingAccelerationStructureAddress), RLM_WriteOnly);
+
+			const uint32 NumGeometries = SceneInitializer.ReferencedGeometries.Num();
+			for (uint32 GeometryIndex = 0; GeometryIndex < NumGeometries; ++GeometryIndex)
+			{
+				AddressesPtr[GeometryIndex] = SceneInitializer.ReferencedGeometries[GeometryIndex]->GetAccelerationStructureAddress(RHICmdList.GetGPUMask().ToIndex());
+			}
+
+			RHICmdList.UnlockBuffer(AccelerationStructureAddressesBuffer.Buffer);
+		});
+
+	BuildRayTracingInstanceBuffer(
+		RHICmdList,
+		SceneInitializer.NumNativeInstances,
+		InstanceBuffer.UAV,
+		InstanceUploadSRV,
+		AccelerationStructureAddressesBuffer.SRV);
+
 	RHICmdList.BindAccelerationStructureMemory(Scene, SceneBuffer, 0);
 
 	RHICmdList.BuildAccelerationStructure(Geometry);
@@ -148,7 +204,14 @@ bool RunRayTracingTestbed_RenderThread(const FString& Parameters)
 	// #yuriy_todo: explicit transitions and state validation for BLAS
 	// RHICmdList.Transition(FRHITransitionInfo(Geometry.GetReference(), ERHIAccess::BVHWrite, ERHIAccess::BVHRead));
 
-	RHICmdList.BuildAccelerationStructure(Scene);
+	FRayTracingSceneBuildParams Params;
+	Params.Scene = Scene;
+	Params.ScratchBuffer = ScratchBuffer;
+	Params.ScratchBufferOffset = 0;
+	Params.InstanceBuffer = InstanceBuffer.Buffer;
+	Params.InstanceBufferOffset = 0;
+
+	RHICmdList.BuildAccelerationStructure(Params);
 
 	RHICmdList.Transition(FRHITransitionInfo(Scene.GetReference(), ERHIAccess::BVHWrite, ERHIAccess::BVHRead));
 
