Adding missing condition to resource descriptor heap size checks.

diff --git a/Engine/Source/Runtime/D3D12RHI/Private/D3D12StateCache.cpp b/Engine/Source/Runtime/D3D12RHI/Private/D3D12StateCache.cpp
--- a/Engine/Source/Runtime/D3D12RHI/Private/D3D12StateCache.cpp
+++ b/Engine/Source/Runtime/D3D12RHI/Private/D3D12StateCache.cpp
@@ -152,7 +152,7 @@ void FD3D12StateCacheBase::Init(FD3D12Device* InParent, FD3D12CommandContext* In
 	const int32 MaximumSamplerHeapSize = Adapter->GetMaxDescriptorsForHeapType(ERHIDescriptorHeapType::Sampler);
 
 	check(GLocalViewHeapSize <= MaximumResourceHeapSize || MaximumResourceHeapSize < 0);
-	check(GOnlineDescriptorHeapSize <= MaximumResourceHeapSize);
+	check(GOnlineDescriptorHeapSize <= MaximumResourceHeapSize || MaximumResourceHeapSize < 0);
 
 	const uint32 NumSamplerDescriptors = NUM_SAMPLER_DESCRIPTORS;
 	check(NumSamplerDescriptors <= MaximumSamplerHeapSize);
