Log extended GPU page fault information in Aftermath crash handler

diff --git a/Engine/Source/Runtime/D3D12RHI/Private/D3D12Stats.cpp b/Engine/Source/Runtime/D3D12RHI/Private/D3D12Stats.cpp
--- a/Engine/Source/Runtime/D3D12RHI/Private/D3D12Stats.cpp
+++ b/Engine/Source/Runtime/D3D12RHI/Private/D3D12Stats.cpp
@@ -416,6 +416,12 @@ bool D3D12RHI::FD3DGPUProfiler::CheckGpuHeartbeat() const
 							FormatStr += FCString::Strlen(FormatPrefix);
 						}
 						UE_LOG(LogRHI, Error, TEXT("[Aftermath] Faulting resource format: %s (0x%x)"), FormatStr, (int32)ResourceFormat);
+
+						if (FaultInformation.faultingGpuVA)
+						{
+							FD3D12Adapter* Adapter = GetParentDevice()->GetParentAdapter();
+							D3D12RHI::LogPageFaultData(Adapter, D3D12_GPU_VIRTUAL_ADDRESS(FaultInformation.faultingGpuVA));
+						}
 					}
 					else
 					{
diff --git a/Engine/Source/Runtime/D3D12RHI/Private/D3D12Util.cpp b/Engine/Source/Runtime/D3D12RHI/Private/D3D12Util.cpp
--- a/Engine/Source/Runtime/D3D12RHI/Private/D3D12Util.cpp
+++ b/Engine/Source/Runtime/D3D12RHI/Private/D3D12Util.cpp
@@ -560,6 +560,10 @@ static bool LogDREDData(ID3D12Device* Device, bool bTrackingAllAllocations, D3D1
 	}
 }
 
+
+namespace D3D12RHI
+{
+
 void LogPageFaultData(FD3D12Adapter* InAdapter, D3D12_GPU_VIRTUAL_ADDRESS InPageFaultAddress)
 {
 	if (InPageFaultAddress == 0)
@@ -632,6 +636,8 @@ void LogPageFaultData(FD3D12Adapter* InAdapter, D3D12_GPU_VIRTUAL_ADDRESS InPage
 	}
 }
 
+} // namespace D3D12RHI
+
 void LogMemoryInfo(FD3D12Adapter* InAdapter)
 {
 	InAdapter->UpdateMemoryInfo();
diff --git a/Engine/Source/Runtime/D3D12RHI/Public/D3D12Util.h b/Engine/Source/Runtime/D3D12RHI/Public/D3D12Util.h
--- a/Engine/Source/Runtime/D3D12RHI/Public/D3D12Util.h
+++ b/Engine/Source/Runtime/D3D12RHI/Public/D3D12Util.h
@@ -10,6 +10,8 @@
 #include "Containers/Queue.h"
 #include "RenderUtils.h"
 
+class FD3D12Adapter;
+
 namespace D3D12RHI
 {
 	/**
@@ -51,7 +53,11 @@ namespace D3D12RHI
 
 	/** Returns a string for the provided DXGI format. */
 	const TCHAR* GetD3D12TextureFormatString(DXGI_FORMAT TextureFormat);
-}
+
+	/** Checks if given GPU virtual address corresponds to any known resource allocations and logs results */
+	void LogPageFaultData(FD3D12Adapter* InAdapter, D3D12_GPU_VIRTUAL_ADDRESS InPageFaultAddress);
+	
+} // namespace D3D12RHI
 
 using namespace D3D12RHI;
 
