U5 - fix d3d12 case where resolve targettable (MSAA result) textures had invalid starting state (UAV)

diff --git a/Engine/Source/Runtime/D3D12RHI/Private/D3D12TransientResourceAllocator.cpp b/Engine/Source/Runtime/D3D12RHI/Private/D3D12TransientResourceAllocator.cpp
--- a/Engine/Source/Runtime/D3D12RHI/Private/D3D12TransientResourceAllocator.cpp
+++ b/Engine/Source/Runtime/D3D12RHI/Private/D3D12TransientResourceAllocator.cpp
@@ -143,7 +143,7 @@ FRHITransientTexture* FD3D12TransientResourceAllocator::CreateTexture(const FRHI
 	{
 		ERHIAccess InitialState = ERHIAccess::UAVMask;
 
-		if (EnumHasAnyFlags(InCreateInfo.Flags, TexCreate_RenderTargetable))
+		if (EnumHasAnyFlags(InCreateInfo.Flags, TexCreate_RenderTargetable | TexCreate_ResolveTargetable))
 		{
 			InitialState = ERHIAccess::RTV;
 		}
