Add support for GPU asserts in ray generation shaders.

diff --git a/Engine/Shaders/Public/Platform/D3D/D3DCommon.ush b/Engine/Shaders/Public/Platform/D3D/D3DCommon.ush
--- a/Engine/Shaders/Public/Platform/D3D/D3DCommon.ush
+++ b/Engine/Shaders/Public/Platform/D3D/D3DCommon.ush
@@ -28,7 +28,7 @@ SamplerState D3DStaticTrilinearClampedSampler   : register(s5, space1000);
 #endif // PLATFORM_SUPPORTS_STATIC_SAMPLERS
 
 // TODO: Add support for the UEDiagnosticBuffer to other shader types
-#if PLATFORM_SUPPORTS_DIAGNOSTIC_BUFFER && (PIXELSHADER || COMPUTESHADER)
+#if PLATFORM_SUPPORTS_DIAGNOSTIC_BUFFER && (PIXELSHADER || COMPUTESHADER || RAYGENSHADER)
 
 // Diagnostic buffer which will be automatically bound by the RHI as a root parameter
 // Counterpart to FD3D12DiagnosticBufferData in D3D12DirectCommandListManager.h
diff --git a/Engine/Source/Runtime/D3D12RHI/Private/D3D12RayTracing.cpp b/Engine/Source/Runtime/D3D12RHI/Private/D3D12RayTracing.cpp
--- a/Engine/Source/Runtime/D3D12RHI/Private/D3D12RayTracing.cpp
+++ b/Engine/Source/Runtime/D3D12RHI/Private/D3D12RayTracing.cpp
@@ -4969,6 +4969,23 @@ static void DispatchRays(FD3D12CommandContext& CommandContext,
 
 	CommandContext.CommandListHandle.GraphicsCommandList()->SetComputeRootSignature(Pipeline->GlobalRootSignature);
 
+	// Bind diagnostic buffer to allow asserts in ray generation shaders
+	{
+#if USE_STATIC_ROOT_SIGNATURE
+		FD3D12RootSignature* RootSignature = Adapter->GetStaticRayTracingGlobalRootSignature();
+#else // USE_STATIC_ROOT_SIGNATURE
+		FD3D12QuantizedBoundShaderState RSDesc = GetRayTracingGlobalRootSignatureDesc();
+		FD3D12RootSignature* RootSignature = Adapter->GetRootSignature(RSDesc);
+#endif // USE_STATIC_ROOT_SIGNATURE
+
+		const int8 DiagnosticBufferSlot = RootSignature->GetDiagnosticBufferSlot();
+		D3D12_GPU_VIRTUAL_ADDRESS DiagnosticBufferAddress = CommandContext.GetParentDevice()->GetCommandListManager().GetDiagnosticBufferGPUAddress();
+		if (DiagnosticBufferSlot >= 0 && DiagnosticBufferAddress)
+		{
+			CommandContext.CommandListHandle.GraphicsCommandList()->SetComputeRootUnorderedAccessView(DiagnosticBufferSlot, DiagnosticBufferAddress);
+		}
+	}
+
 	FD3D12RayTracingShader* RayGenShader = Pipeline->RayGenShaders.Shaders[RayGenShaderIndex];
 
 	bool bResourcesBound = false;
diff --git a/Engine/Source/Runtime/D3D12RHI/Private/D3D12Util.cpp b/Engine/Source/Runtime/D3D12RHI/Private/D3D12Util.cpp
--- a/Engine/Source/Runtime/D3D12RHI/Private/D3D12Util.cpp
+++ b/Engine/Source/Runtime/D3D12RHI/Private/D3D12Util.cpp
@@ -246,7 +246,7 @@ static bool LogBreadcrumbData(D3D12RHI::FD3DGPUProfiler& GPUProfiler, FD3D12Comm
 	const FD3D12DiagnosticBufferData* DiagnosticData = CommandListManager.GetDiagnosticBufferData();
 	if (DiagnosticData && DiagnosticData->Counter)
 	{
-		UE_LOG(LogD3D12RHI, Error, TEXT("[GPUBreadCrumb]\t\tShader assertion failed! ID: 0x%08X"), DiagnosticData->MessageID);
+		UE_LOG(LogD3D12RHI, Error, TEXT("[GPUBreadCrumb]\t\tShader assertion failed! ID: 0x%08X (%d)"), DiagnosticData->MessageID, DiagnosticData->MessageID);
 
 		{
 			const int32* Payload = DiagnosticData->Payload.AsInt;
@@ -1023,6 +1023,7 @@ FD3D12QuantizedBoundShaderState GetRayTracingGlobalRootSignatureDesc()
 	FShaderRegisterCounts& QBSSRegisterCounts = OutQBSS.RegisterCounts[SV_All];
 
 	OutQBSS.RootSignatureType = RS_RayTracingGlobal;
+	OutQBSS.bUseDiagnosticBuffer = true;
 
 	QBSSRegisterCounts.SamplerCount = MAX_SAMPLERS;
 	QBSSRegisterCounts.ShaderResourceCount = MAX_SRVS;
