Fix compile error due to missing const

diff --git a/Engine/Source/Runtime/D3D12RHI/Private/D3D12RayTracing.cpp b/Engine/Source/Runtime/D3D12RHI/Private/D3D12RayTracing.cpp
--- a/Engine/Source/Runtime/D3D12RHI/Private/D3D12RayTracing.cpp
+++ b/Engine/Source/Runtime/D3D12RHI/Private/D3D12RayTracing.cpp
@@ -4972,10 +4972,10 @@ static void DispatchRays(FD3D12CommandContext& CommandContext,
 	// Bind diagnostic buffer to allow asserts in ray generation shaders
 	{
 #if USE_STATIC_ROOT_SIGNATURE
-		FD3D12RootSignature* RootSignature = Adapter->GetStaticRayTracingGlobalRootSignature();
+		const FD3D12RootSignature* RootSignature = Adapter->GetStaticRayTracingGlobalRootSignature();
 #else // USE_STATIC_ROOT_SIGNATURE
 		FD3D12QuantizedBoundShaderState RSDesc = GetRayTracingGlobalRootSignatureDesc();
-		FD3D12RootSignature* RootSignature = Adapter->GetRootSignature(RSDesc);
+		const FD3D12RootSignature* RootSignature = Adapter->GetRootSignature(RSDesc);
 #endif // USE_STATIC_ROOT_SIGNATURE
 
 		const int8 DiagnosticBufferSlot = RootSignature->GetDiagnosticBufferSlot();
