Fixing asserts on platforms that don't support mesh shaders by excluding mesh/amplification from EShaderVisibility on those platforms.

diff --git a/Engine/Source/Runtime/D3D12RHI/Private/D3D12RootSignature.cpp b/Engine/Source/Runtime/D3D12RHI/Private/D3D12RootSignature.cpp
--- a/Engine/Source/Runtime/D3D12RHI/Private/D3D12RootSignature.cpp
+++ b/Engine/Source/Runtime/D3D12RHI/Private/D3D12RootSignature.cpp
@@ -122,7 +122,17 @@ FORCEINLINE D3D12_ROOT_SIGNATURE_FLAGS GetD3D12RootSignatureDenyFlag(EShaderVisi
 FD3D12RootSignatureDesc::FD3D12RootSignatureDesc(const FD3D12QuantizedBoundShaderState& QBSS, const D3D12_RESOURCE_BINDING_TIER ResourceBindingTier)
 	: RootParametersSize(0)
 {
-	const EShaderVisibility ShaderVisibilityPriorityOrder[] = { SV_Pixel, SV_Vertex, SV_Geometry, SV_Mesh, SV_Amplification, SV_All };
+	const EShaderVisibility ShaderVisibilityPriorityOrder[] =
+	{
+		SV_Pixel,
+		SV_Vertex,
+		SV_Geometry,
+#if PLATFORM_SUPPORTS_MESH_SHADERS
+		SV_Mesh,
+		SV_Amplification,
+#endif
+		SV_All
+	};
 	const D3D12_ROOT_PARAMETER_TYPE RootParameterTypePriorityOrder[] = { D3D12_ROOT_PARAMETER_TYPE_DESCRIPTOR_TABLE, D3D12_ROOT_PARAMETER_TYPE_CBV };
 	uint32 RootParameterCount = 0;
 
@@ -213,7 +223,7 @@ FD3D12RootSignatureDesc::FD3D12RootSignatureDesc(const FD3D12QuantizedBoundShade
 		// ... and each shader stage visibility ...
 		for (uint32 ShaderVisibilityIndex = 0; ShaderVisibilityIndex < UE_ARRAY_COUNT(ShaderVisibilityPriorityOrder); ShaderVisibilityIndex++)
 		{
-			const EShaderVisibility& Visibility = ShaderVisibilityPriorityOrder[ShaderVisibilityIndex];
+			const EShaderVisibility Visibility = ShaderVisibilityPriorityOrder[ShaderVisibilityIndex];
 			const FShaderRegisterCounts& Shader = QBSS.RegisterCounts[Visibility];
 
 			switch (RootParameterType)
@@ -313,7 +323,7 @@ FD3D12RootSignatureDesc::FD3D12RootSignatureDesc(const FD3D12QuantizedBoundShade
 
 		for (uint32 ShaderVisibilityIndex = 0; ShaderVisibilityIndex < UE_ARRAY_COUNT(ShaderVisibilityPriorityOrder); ShaderVisibilityIndex++)
 		{
-			const EShaderVisibility& Visibility = ShaderVisibilityPriorityOrder[ShaderVisibilityIndex];
+			const EShaderVisibility Visibility = ShaderVisibilityPriorityOrder[ShaderVisibilityIndex];
 			const FShaderRegisterCounts& Shader = QBSS.RegisterCounts[Visibility];
 			if ((Shader.ShaderResourceCount == 0) &&
 				(Shader.ConstantBufferCount == 0) &&
diff --git a/Engine/Source/Runtime/D3D12RHI/Private/D3D12Util.cpp b/Engine/Source/Runtime/D3D12RHI/Private/D3D12Util.cpp
--- a/Engine/Source/Runtime/D3D12RHI/Private/D3D12Util.cpp
+++ b/Engine/Source/Runtime/D3D12RHI/Private/D3D12Util.cpp
@@ -991,8 +991,10 @@ void QuantizeBoundShaderState(
 	QBSS.bAllowIAInputLayout = BSS->GetVertexDeclaration() != nullptr;	// Does the root signature need access to vertex buffers?
 
 	QuantizeBoundShaderStateCommon(QBSS, BSS->GetVertexShader(),        ResourceBindingTier, SV_Vertex);
+#if PLATFORM_SUPPORTS_MESH_SHADERS
 	QuantizeBoundShaderStateCommon(QBSS, BSS->GetMeshShader(),          ResourceBindingTier, SV_Mesh);
 	QuantizeBoundShaderStateCommon(QBSS, BSS->GetAmplificationShader(), ResourceBindingTier, SV_Amplification);
+#endif
 	QuantizeBoundShaderStateCommon(QBSS, BSS->GetPixelShader(),         ResourceBindingTier, SV_Pixel, true /*bAllowUAVs*/);
 	QuantizeBoundShaderStateCommon(QBSS, BSS->GetGeometryShader(),      ResourceBindingTier, SV_Geometry);
 }
diff --git a/Engine/Source/Runtime/D3D12RHI/Public/D3D12Util.h b/Engine/Source/Runtime/D3D12RHI/Public/D3D12Util.h
--- a/Engine/Source/Runtime/D3D12RHI/Public/D3D12Util.h
+++ b/Engine/Source/Runtime/D3D12RHI/Public/D3D12Util.h
@@ -71,8 +71,10 @@ enum EShaderVisibility
 	SV_Vertex,
 	SV_Pixel,
 	SV_Geometry,
+#if PLATFORM_SUPPORTS_MESH_SHADERS
 	SV_Mesh,
 	SV_Amplification,
+#endif
 	SV_All,
 	SV_ShaderVisibilityCount
 };
