Shelf '18242529': Restrict UAV aliasing to actual supporting hardware.

diff --git a/Engine/Source/Runtime/D3D12RHI/Private/Windows/WindowsD3D12Device.cpp b/Engine/Source/Runtime/D3D12RHI/Private/Windows/WindowsD3D12Device.cpp
--- a/Engine/Source/Runtime/D3D12RHI/Private/Windows/WindowsD3D12Device.cpp
+++ b/Engine/Source/Runtime/D3D12RHI/Private/Windows/WindowsD3D12Device.cpp
@@ -1163,7 +1163,7 @@ void FD3D12DynamicRHI::Init()
 		GRHIVariableRateShadingImageTileMaxHeight = 1;
 	}
 
-	GRHISupportsUAVFormatAliasing = GetAdapter().GetResourceHeapTier() > D3D12_RESOURCE_HEAP_TIER_1;
+	GRHISupportsUAVFormatAliasing = (GetAdapter().GetResourceHeapTier() > D3D12_RESOURCE_HEAP_TIER_1 && IsRHIDeviceNVIDIA());
 
 	// Command lists need the validation RHI context if enabled, so call the global scope version of RHIGetDefaultContext() and RHIGetDefaultAsyncComputeContext().
 	GRHICommandList.GetImmediateCommandList().SetContext(::RHIGetDefaultContext());
