Report resource dimensions and type when logging live and recently released resources after a GPU crash.

This makes it easier to correlate resources reported by Aftermath with the live and recently released allocations in the engine.

diff --git a/Engine/Source/Runtime/D3D12RHI/Private/D3D12Adapter.cpp b/Engine/Source/Runtime/D3D12RHI/Private/D3D12Adapter.cpp
--- a/Engine/Source/Runtime/D3D12RHI/Private/D3D12Adapter.cpp
+++ b/Engine/Source/Runtime/D3D12RHI/Private/D3D12Adapter.cpp
@@ -1549,6 +1549,7 @@ void FD3D12Adapter::ReleaseTrackedAllocationData(FD3D12ResourceLocation* InAlloc
 		ReleasedData.GPUVirtualAddress = GPUAddress;
 		ReleasedData.AllocationSize = InAllocation->GetSize();
 		ReleasedData.ResourceName = InAllocation->GetResource()->GetName();
+		ReleasedData.ResourceDesc = InAllocation->GetResource()->GetDesc();
 		ReleasedData.ReleasedFrameID = GetFrameFence().GetCurrentFence();
 		ReleasedData.bDefragFree = bDefragFree;
 		ReleasedData.bBackBuffer = InAllocation->GetResource()->IsBackBuffer();
diff --git a/Engine/Source/Runtime/D3D12RHI/Private/D3D12Adapter.h b/Engine/Source/Runtime/D3D12RHI/Private/D3D12Adapter.h
--- a/Engine/Source/Runtime/D3D12RHI/Private/D3D12Adapter.h
+++ b/Engine/Source/Runtime/D3D12RHI/Private/D3D12Adapter.h
@@ -410,6 +410,7 @@ public:
 		D3D12_GPU_VIRTUAL_ADDRESS GPUVirtualAddress = 0;
 		uint64 AllocationSize = 0;
 		FName ResourceName;
+		D3D12_RESOURCE_DESC ResourceDesc = {};
 		uint64 ReleasedFrameID = 0;
 		bool bDefragFree = false;
 		bool bBackBuffer = false;
diff --git a/Engine/Source/Runtime/D3D12RHI/Private/D3D12Util.cpp b/Engine/Source/Runtime/D3D12RHI/Private/D3D12Util.cpp
--- a/Engine/Source/Runtime/D3D12RHI/Private/D3D12Util.cpp
+++ b/Engine/Source/Runtime/D3D12RHI/Private/D3D12Util.cpp
@@ -564,6 +564,26 @@ static bool LogDREDData(ID3D12Device* Device, bool bTrackingAllAllocations, D3D1
 namespace D3D12RHI
 {
 
+
+static FString MakeResourceDescDebugString(const D3D12_RESOURCE_DESC& Desc)
+{
+	FString ResourceDescString;
+	switch (Desc.Dimension)
+	{
+	default:
+		ResourceDescString = TEXT("Unknown");
+		break;
+	case D3D12_RESOURCE_DIMENSION_BUFFER:
+		ResourceDescString = FString::Printf(TEXT("Buffer %d bytes"), Desc.Width);
+		break;
+	case D3D12_RESOURCE_DIMENSION_TEXTURE1D:
+	case D3D12_RESOURCE_DIMENSION_TEXTURE2D:
+	case D3D12_RESOURCE_DIMENSION_TEXTURE3D:
+		ResourceDescString = FString::Printf(TEXT("Texture %dx%dx%d %s"), Desc.Width, Desc.Height, Desc.DepthOrArraySize, LexToString(Desc.Format));
+	}
+	return ResourceDescString;
+}
+
 void LogPageFaultData(FD3D12Adapter* InAdapter, D3D12_GPU_VIRTUAL_ADDRESS InPageFaultAddress)
 {
 	if (InPageFaultAddress == 0)
@@ -589,7 +609,11 @@ void LogPageFaultData(FD3D12Adapter* InAdapter, D3D12_GPU_VIRTUAL_ADDRESS InPage
 		{
 			FD3D12Adapter::FAllocatedResourceResult OverlappingResource = OverlappingResources[Index];
 			D3D12_GPU_VIRTUAL_ADDRESS ResourceAddress = OverlappingResource.Allocation->GetGPUVirtualAddress();
-			UE_LOG(LogD3D12RHI, Error, TEXT("\tGPU Address: [0x%llX .. 0x%llX] - Size: %lld bytes, %3.2f MB - Distance to page fault: %lld bytes, %3.2f MB - Transient: %d - Name: %s"),
+
+			const FD3D12Resource* Resource = OverlappingResource.Allocation->GetResource();
+			FString ResourceDescString = MakeResourceDescDebugString(Resource->GetDesc());
+
+			UE_LOG(LogD3D12RHI, Error, TEXT("\tGPU Address: [0x%llX .. 0x%llX] - Size: %lld bytes, %3.2f MB - Distance to page fault: %lld bytes, %3.2f MB - Transient: %d - Name: %s - Desc: %s"),
 				(uint64)ResourceAddress,
 				(uint64)ResourceAddress + OverlappingResource.Allocation->GetSize(),
 				OverlappingResource.Allocation->GetSize(),
@@ -597,7 +621,8 @@ void LogPageFaultData(FD3D12Adapter* InAdapter, D3D12_GPU_VIRTUAL_ADDRESS InPage
 				OverlappingResource.Distance,
 				OverlappingResource.Distance / (1024.0f * 1024), 
 				OverlappingResource.Allocation->IsTransient(), 
-				*OverlappingResource.Allocation->GetResource()->GetName().ToString());
+				*Resource->GetName().ToString(),
+				*ResourceDescString);
 		}
 	}
 
@@ -622,7 +647,10 @@ void LogPageFaultData(FD3D12Adapter* InAdapter, D3D12_GPU_VIRTUAL_ADDRESS InPage
 		for (uint32 Index = 0; Index < PrintCount; ++Index)
 		{
 			FD3D12Adapter::FReleasedAllocationData& AllocationData = ReleasedResources[Index];
-			UE_LOG(LogD3D12RHI, Error, TEXT("\tGPU Address: [0x%llX .. 0x%llX] - Size: %lld bytes, %3.2f MB - FrameID: %4d - DefragFree: %d - Transient: %d - Heap: %d - Name: %s"),
+
+			FString ResourceDescString = MakeResourceDescDebugString(AllocationData.ResourceDesc);
+
+			UE_LOG(LogD3D12RHI, Error, TEXT("\tGPU Address: [0x%llX .. 0x%llX] - Size: %lld bytes, %3.2f MB - FrameID: %4d - DefragFree: %d - Transient: %d - Heap: %d - Name: %s - Desc: %s"),
 				(uint64)AllocationData.GPUVirtualAddress,
 				(uint64)AllocationData.GPUVirtualAddress + AllocationData.AllocationSize,
 				AllocationData.AllocationSize,
@@ -631,7 +659,8 @@ void LogPageFaultData(FD3D12Adapter* InAdapter, D3D12_GPU_VIRTUAL_ADDRESS InPage
 				AllocationData.bDefragFree,
 				AllocationData.bTransient,
 				AllocationData.bHeap,
-				*AllocationData.ResourceName.ToString());
+				*AllocationData.ResourceName.ToString(),
+				*ResourceDescString);
 		}
 	}
 }
