View rename fixes:
- FD3D12ViewDescriptorHandle wasn't handling the null parent device that's passed in when creating the temporary "null" view
- Implementing SetParentDevice on FD3D12ViewDescriptorHandle to handle renaming from a null view

diff --git a/Engine/Source/Runtime/D3D12RHI/Private/D3D12DescriptorCache.cpp b/Engine/Source/Runtime/D3D12RHI/Private/D3D12DescriptorCache.cpp
--- a/Engine/Source/Runtime/D3D12RHI/Private/D3D12DescriptorCache.cpp
+++ b/Engine/Source/Runtime/D3D12RHI/Private/D3D12DescriptorCache.cpp
@@ -61,7 +61,7 @@ FD3D12DescriptorCache::FD3D12DescriptorCache(FRHIGPUMask Node)
 	, CurrentViewHeap(nullptr)
 	, CurrentSamplerHeap(nullptr)
 	, LocalViewHeap(nullptr)
-	, LocalSamplerHeap(nullptr, this)
+	, LocalSamplerHeap(this)
 	, SubAllocatedViewHeap(this)
 	, SamplerMap(271) // Prime numbers for better hashing
 	, bUsingGlobalSamplerHeap(false)
@@ -75,13 +75,12 @@ void FD3D12DescriptorCache::Init(FD3D12Device* InParent, FD3D12CommandContext* I
 	Parent = InParent;
 	CmdContext = InCmdContext;
 
-	LocalSamplerHeap.SetParentDevice(InParent);
 	SubAllocatedViewHeap.Init(InParent);
 
 	// Always Init a local sampler heap as the high level cache will always miss initialy
 	// so we need something to fall back on (The view heap never rolls over so we init that one
 	// lazily as a backup to save memory)
-	LocalSamplerHeap.Init(InNumSamplerDescriptors, ERHIDescriptorHeapType::Sampler);
+	LocalSamplerHeap.Init(InParent, InNumSamplerDescriptors, ERHIDescriptorHeapType::Sampler);
 
 	NumLocalViewDescriptors = InNumLocalViewDescriptors;
 
@@ -732,11 +731,11 @@ bool FD3D12DescriptorCache::SwitchToContextLocalViewHeap(const FD3D12CommandList
 		UE_LOG(LogD3D12RHI, Log, TEXT("This should only happen in the Editor where it doesn't matter as much. If it happens in game you should increase the device global heap size!"));
 		
 		// Allocate the heap lazily
-		LocalViewHeap = new FD3D12LocalOnlineHeap(GetParentDevice(), this);
+		LocalViewHeap = new FD3D12LocalOnlineHeap(this);
 		if (LocalViewHeap)
 		{
 			check(NumLocalViewDescriptors);
-			LocalViewHeap->Init(NumLocalViewDescriptors, ERHIDescriptorHeapType::Standard);
+			LocalViewHeap->Init(GetParentDevice(), NumLocalViewDescriptors, ERHIDescriptorHeapType::Standard);
 		}
 		else
 		{
@@ -969,9 +968,10 @@ FD3D12SubAllocatedOnlineHeap::FD3D12SubAllocatedOnlineHeap(FD3D12DescriptorCache
 FD3D12SubAllocatedOnlineHeap::~FD3D12SubAllocatedOnlineHeap() = default;
 
 /** Initialize the sub allocated online heap */
-void FD3D12SubAllocatedOnlineHeap::Init(FD3D12Device* InDevice)
+void FD3D12SubAllocatedOnlineHeap::Init(FD3D12Device* InParent)
 {
-	SetParentDevice(InDevice);
+	check(Parent == nullptr);
+	Parent = InParent;
 }
 
 /** Handle roll over on the sub allocated online heap - needs a new block */
@@ -1045,8 +1045,8 @@ bool FD3D12SubAllocatedOnlineHeap::AllocateBlock()
 ////////////////////////////////////////////////////////////////////////////////////////////////////////////////////////
 // FD3D12LocalOnlineHeap
 
-FD3D12LocalOnlineHeap::FD3D12LocalOnlineHeap(FD3D12Device* Device, FD3D12DescriptorCache* InDescriptorCache)
-	: FD3D12OnlineHeap(Device, true)
+FD3D12LocalOnlineHeap::FD3D12LocalOnlineHeap(FD3D12DescriptorCache* InDescriptorCache)
+	: FD3D12OnlineHeap(nullptr, true)
 	, DescriptorCache(InDescriptorCache)
 {
 }
@@ -1056,8 +1056,11 @@ FD3D12LocalOnlineHeap::~FD3D12LocalOnlineHeap() = default;
 /**
 Initialize a thread local online heap
 **/
-void FD3D12LocalOnlineHeap::Init(uint32 InNumDescriptors, ERHIDescriptorHeapType InHeapType)
+void FD3D12LocalOnlineHeap::Init(FD3D12Device* InParent, uint32 InNumDescriptors, ERHIDescriptorHeapType InHeapType)
 {
+	check(Parent == nullptr);
+	Parent = InParent;
+
 	const TCHAR* DebugName = InHeapType == ERHIDescriptorHeapType::Standard ? L"Thread Local - Online View Heap" : L"Thread Local - Online Sampler Heap";
 	Heap = GetParentDevice()->GetDescriptorHeapManager().AllocateHeap(
 		DebugName,
diff --git a/Engine/Source/Runtime/D3D12RHI/Private/D3D12DescriptorCache.h b/Engine/Source/Runtime/D3D12RHI/Private/D3D12DescriptorCache.h
--- a/Engine/Source/Runtime/D3D12RHI/Private/D3D12DescriptorCache.h
+++ b/Engine/Source/Runtime/D3D12RHI/Private/D3D12DescriptorCache.h
@@ -233,7 +233,7 @@ public:
 	~FD3D12SubAllocatedOnlineHeap();
 
 	// Setup the online heap data
-	void Init(FD3D12Device* InDevice);
+	void Init(FD3D12Device* InParent);
 
 	// Override FD3D12OnlineHeap functions
 	virtual bool RollOver() final override;
@@ -260,11 +260,11 @@ private:
 class FD3D12LocalOnlineHeap : public FD3D12OnlineHeap
 {
 public:
-	FD3D12LocalOnlineHeap(FD3D12Device* Device, FD3D12DescriptorCache* InDescriptorCache);
+	FD3D12LocalOnlineHeap(FD3D12DescriptorCache* InDescriptorCache);
 	~FD3D12LocalOnlineHeap();
 
 	// Allocate the actual overflow heap
-	void Init(uint32 InNumDescriptors, ERHIDescriptorHeapType InHeapType);
+	void Init(FD3D12Device* InParent, uint32 InNumDescriptors, ERHIDescriptorHeapType InHeapType);
 
 	// Override FD3D12OnlineHeap functions
 	virtual bool RollOver() final override;
diff --git a/Engine/Source/Runtime/D3D12RHI/Private/D3D12RHICommon.h b/Engine/Source/Runtime/D3D12RHI/Private/D3D12RHICommon.h
--- a/Engine/Source/Runtime/D3D12RHI/Private/D3D12RHICommon.h
+++ b/Engine/Source/Runtime/D3D12RHI/Private/D3D12RHICommon.h
@@ -74,13 +74,6 @@ public:
 	{
 		return Parent;
 	}
-
-	// To be used with delayed setup
-	inline void SetParentDevice(FD3D12Device* InParent)
-	{
-		check(Parent == nullptr);
-		Parent = InParent;
-	}
 };
 
 class FD3D12GPUObject
diff --git a/Engine/Source/Runtime/D3D12RHI/Private/D3D12View.cpp b/Engine/Source/Runtime/D3D12RHI/Private/D3D12View.cpp
--- a/Engine/Source/Runtime/D3D12RHI/Private/D3D12View.cpp
+++ b/Engine/Source/Runtime/D3D12RHI/Private/D3D12View.cpp
@@ -20,6 +20,14 @@ FD3D12ViewDescriptorHandle::~FD3D12ViewDescriptorHandle()
 	check(!BindlessHandle.IsValid());
 }
 
+void FD3D12ViewDescriptorHandle::SetParentDevice(FD3D12Device* InParent)
+{
+	check(Parent == nullptr);
+	Parent = InParent;
+
+	AllocateDescriptorSlot();
+}
+
 void FD3D12ViewDescriptorHandle::CreateView(const D3D12_RENDER_TARGET_VIEW_DESC& Desc, ID3D12Resource* Resource)
 {
 	check(HeapType == ERHIDescriptorHeapType::RenderTarget);
@@ -75,7 +83,7 @@ void FD3D12ViewDescriptorHandle::CreateView(const D3D12_UNORDERED_ACCESS_VIEW_DE
 
 void FD3D12ViewDescriptorHandle::AllocateDescriptorSlot()
 {
-	if (FD3D12Device* Device = GetParentDevice())
+	if (FD3D12Device* Device = GetParentDevice_Unsafe())
 	{
 		OfflineCpuHandle = Device->GetOfflineDescriptorManager(HeapType).AllocateHeapSlot(OfflineHeapIndex);
 		check(OfflineCpuHandle.ptr != 0);
@@ -89,7 +97,7 @@ void FD3D12ViewDescriptorHandle::AllocateDescriptorSlot()
 
 void FD3D12ViewDescriptorHandle::FreeDescriptorSlot()
 {
-	if (FD3D12Device* Device = GetParentDevice())
+	if (FD3D12Device* Device = GetParentDevice_Unsafe())
 	{
 		Device->GetOfflineDescriptorManager(HeapType).FreeHeapSlot(OfflineCpuHandle, OfflineHeapIndex);
 		OfflineHeapIndex = UINT_MAX;
diff --git a/Engine/Source/Runtime/D3D12RHI/Private/D3D12View.h b/Engine/Source/Runtime/D3D12RHI/Private/D3D12View.h
--- a/Engine/Source/Runtime/D3D12RHI/Private/D3D12View.h
+++ b/Engine/Source/Runtime/D3D12RHI/Private/D3D12View.h
@@ -626,6 +626,9 @@ public:
 	FD3D12ViewDescriptorHandle(FD3D12Device* InParentDevice, ERHIDescriptorHeapType InHeapType);
 	~FD3D12ViewDescriptorHandle();
 
+	// Called when streaming decides it's ready to actually make a view
+	void SetParentDevice(FD3D12Device* InParent);
+
 	void CreateView(const D3D12_RENDER_TARGET_VIEW_DESC& Desc, ID3D12Resource* Resource);
 	void CreateView(const D3D12_DEPTH_STENCIL_VIEW_DESC& Desc, ID3D12Resource* Resource);
 	void CreateView(const D3D12_CONSTANT_BUFFER_VIEW_DESC& Desc);
