UE5_RELEASE:  MGPU, fixes to avoid startup and shutdown crashes.

* GetLowestLevelContext returns the first physical context, rather than the multi-GPU context.  The one place this function is used in practice, we are initializing the first copy of a multi-GPU resource, so we always want the first GPU.  If the caller wants a more specific context, they can use other functions.
* Repeatedly loop when flushing pending deletes on shutdown, until nothing gets deleted.  Sometimes a pending delete will release a reference on something else, triggered a chained delete.  Without MGPU, this just leaked resources, but with MGPU, this was actually crashing, because we call the device cleanup once per device, and some of the leftover deletes are shared across devices, which crashes because the first device has already been destroyed.
* Pool allocated resources were asserting in FD3D12ResourceLocation::ReferenceNode, because there wasn't a proper implementation to add these nodes to linked lists for tracking.  A comment suggested threading matters -- the code isn't called in multiple threads, but it did corrupt the linked lists if you skip the assert, crashing things later, which my fix solves.
* Ray tracing acceleration structure GPU updates were asserting due to a call to GetGPUMask().ToIndex() on a mask with multiple bits set.  We need to update the structure on the GPU where the buffer in question is located.

#rnx

diff --git a/Engine/Source/Runtime/D3D12RHI/Private/D3D12CommandContext.h b/Engine/Source/Runtime/D3D12RHI/Private/D3D12CommandContext.h
--- a/Engine/Source/Runtime/D3D12RHI/Private/D3D12CommandContext.h
+++ b/Engine/Source/Runtime/D3D12RHI/Private/D3D12CommandContext.h
@@ -820,6 +820,11 @@ public:
 		check(PhysicalGPUMask.ContainsAll(GPUMask));
 	}
 
+	virtual FRHIGPUMask RHIGetGPUMask() const final override
+	{
+		return GPUMask;
+	}
+
 	// Sets the mask of which GPUs can be supported, as opposed to the currently active
 	// set. RHISetGPUMask checks that the active mask is a subset of the physical mask.
 	FORCEINLINE void SetPhysicalGPUMask(FRHIGPUMask InGPUMask)
@@ -857,6 +862,11 @@ public:
 		return PhysicalContexts[GPUIndex];
 	}
 
+	IRHIComputeContext& GetLowestLevelContext() final override
+	{
+		return *PhysicalContexts[0];
+	}
+
 private:
 
 	// Make every GPU in the provided mask to wait on one another.
diff --git a/Engine/Source/Runtime/D3D12RHI/Private/D3D12Device.cpp b/Engine/Source/Runtime/D3D12RHI/Private/D3D12Device.cpp
--- a/Engine/Source/Runtime/D3D12RHI/Private/D3D12Device.cpp
+++ b/Engine/Source/Runtime/D3D12RHI/Private/D3D12Device.cpp
@@ -455,8 +455,12 @@ void FD3D12Device::Cleanup()
 	ReleasePooledUniformBuffers();
 
 	// Flush all pending deletes before destroying the device or any command contexts.
-	FRHIResource::FlushPendingDeletes(RHICmdList);
-	RHICmdList.ImmediateFlush(EImmediateFlushType::FlushRHIThread);
+	int32 DeletedCount;
+	do
+	{
+		DeletedCount = FRHIResource::FlushPendingDeletes(RHICmdList);
+		RHICmdList.ImmediateFlush(EImmediateFlushType::FlushRHIThread);
+	} while (DeletedCount);
 
 	// Delete array index 0 (the default context) last
 	for (int32 i = CommandContextArray.Num() - 1; i >= 0; i--)
diff --git a/Engine/Source/Runtime/D3D12RHI/Private/D3D12Resources.cpp b/Engine/Source/Runtime/D3D12RHI/Private/D3D12Resources.cpp
--- a/Engine/Source/Runtime/D3D12RHI/Private/D3D12Resources.cpp
+++ b/Engine/Source/Runtime/D3D12RHI/Private/D3D12Resources.cpp
@@ -759,8 +759,6 @@ void FD3D12ResourceLocation::Alias(FD3D12ResourceLocation & Destination, FD3D12R
 
 void FD3D12ResourceLocation::ReferenceNode(FD3D12Device* DestinationDevice, FD3D12ResourceLocation& Destination, FD3D12ResourceLocation& Source)
 {
-	// Should not be linked list allocated - otherwise internal linked list data needs to be updated as well in a threadsafe way
-	check(Source.GetAllocatorType() != FD3D12ResourceLocation::AT_Pool);
 	check(Source.GetResource() != nullptr);
 	Destination.Clear();
 
@@ -771,6 +769,12 @@ void FD3D12ResourceLocation::ReferenceNode(FD3D12Device* DestinationDevice, FD3D
 
 	// Addref the source as another resource location references it
 	Source.GetResource()->AddRef();
+
+	if (Source.GetAllocatorType() == FD3D12ResourceLocation::AT_Pool)
+	{
+		Source.GetPoolAllocatorPrivateData().PoolData.AddAlias(
+			&Destination.GetPoolAllocatorPrivateData().PoolData);
+	}
 }
 
 void FD3D12ResourceLocation::ReleaseResource()
@@ -819,6 +823,15 @@ void FD3D12ResourceLocation::ReleaseResource()
 		}
 		else if (AllocatorType == AT_Pool)
 		{
+			// Unlink any aliases -- the contents of aliases are cleaned up separately elsewhere via iteration over
+			// the FD3D12LinkedAdapterObject.
+			for (FRHIPoolAllocationData* Alias = GetPoolAllocatorPrivateData().PoolData.GetFirstAlias();
+				 Alias;
+				 Alias = GetPoolAllocatorPrivateData().PoolData.GetFirstAlias())
+			{
+				Alias->RemoveAlias();
+			}
+
 			PoolAllocator->DeallocateResource(*this);
 		}
 		else
@@ -830,6 +843,11 @@ void FD3D12ResourceLocation::ReleaseResource()
 	case ResourceLocationType::eNodeReference:
 	case ResourceLocationType::eAliased:
 	{
+		if (GetAllocatorType() == FD3D12ResourceLocation::AT_Pool)
+		{
+			GetPoolAllocatorPrivateData().PoolData.RemoveAlias();
+		}
+
 		if (UnderlyingResource->ShouldDeferDelete() && UnderlyingResource->GetRefCount() == 1)
 		{
 			UnderlyingResource->DeferDelete();
@@ -1022,6 +1040,17 @@ bool FD3D12ResourceLocation::OnAllocationMoved(FRHIPoolAllocationData* InNewData
 	GPUVirtualAddress = UnderlyingResource->GetGPUVirtualAddress() + OffsetFromBaseOfResource;
 	ResidencyHandle = UnderlyingResource->GetResidencyHandle();
 
+	// Refresh aliases
+	for (FRHIPoolAllocationData* OtherAlias = AllocationData.GetFirstAlias(); OtherAlias; OtherAlias = OtherAlias->GetNext())
+	{
+		FD3D12ResourceLocation* OtherResourceLocation = (FD3D12ResourceLocation*)OtherAlias->GetOwner();
+
+		OtherResourceLocation->OffsetFromBaseOfResource = OffsetFromBaseOfResource;
+		OtherResourceLocation->UnderlyingResource = UnderlyingResource;
+		OtherResourceLocation->GPUVirtualAddress = GPUVirtualAddress;
+		OtherResourceLocation->ResidencyHandle = ResidencyHandle;
+	}
+
 	// Notify all the dependent resources about the change
 	Owner->ResourceRenamed(this);
 
@@ -1131,4 +1160,10 @@ void FD3D12ResourceBarrierBatcher::Flush(FD3D12Device* Device, ID3D12GraphicsCom
 #endif // #if PLATFORM_USE_BACKBUFFER_WRITE_TRANSITION_TRACKING
 
 	Reset();
-}
\ No newline at end of file
+}
+
+uint32 FD3D12Buffer::GetParentGPUIndex() const
+{
+	return Parent->GetGPUIndex();
+}
+
diff --git a/Engine/Source/Runtime/D3D12RHI/Public/D3D12Resources.h b/Engine/Source/Runtime/D3D12RHI/Public/D3D12Resources.h
--- a/Engine/Source/Runtime/D3D12RHI/Public/D3D12Resources.h
+++ b/Engine/Source/Runtime/D3D12RHI/Public/D3D12Resources.h
@@ -959,6 +959,8 @@ public:
 		}
 	}
 
+	virtual uint32 GetParentGPUIndex() const override;
+
 	void UploadResourceData(class FRHICommandListImmediate* InRHICmdList, FResourceArrayInterface* InResourceArray, D3D12_RESOURCE_STATES InDestinationState);
 	FD3D12SyncPoint UploadResourceDataViaCopyQueue(FResourceArrayInterface* InResourceArray);
 
diff --git a/Engine/Source/Runtime/RHI/Public/RHIContext.h b/Engine/Source/Runtime/RHI/Public/RHIContext.h
--- a/Engine/Source/Runtime/RHI/Public/RHIContext.h
+++ b/Engine/Source/Runtime/RHI/Public/RHIContext.h
@@ -285,6 +285,11 @@ public:
 		ensure(GPUMask == FRHIGPUMask::GPU0());
 	}
 
+	virtual FRHIGPUMask RHIGetGPUMask() const
+	{
+		return FRHIGPUMask::GPU0();
+	}
+
 #if WITH_MGPU
 	virtual void RHIWaitForTemporalEffect(const FName& InEffectName)
 	{
@@ -362,6 +367,10 @@ public:
 		return WrappingContext ? *WrappingContext : *this;
 	}
 
+#elif WITH_MGPU
+	// Needs to be virtual, because it's required for multi GPU resource uploads (FD3D12RHICommandInitializeBuffer, etc)
+	virtual IRHIComputeContext& GetLowestLevelContext() { return *this; }
+	inline IRHIComputeContext& GetHighestLevelContext() { return *this; }
 #else
 
 	// Fast implementations when the RHI validation layer is disabled.
diff --git a/Engine/Source/Runtime/RHI/Public/RHIResources.h b/Engine/Source/Runtime/RHI/Public/RHIResources.h
--- a/Engine/Source/Runtime/RHI/Public/RHIResources.h
+++ b/Engine/Source/Runtime/RHI/Public/RHIResources.h
@@ -1177,6 +1177,8 @@ public:
 
 	FName GetName() const { return BufferName; }
 
+	virtual uint32 GetParentGPUIndex() const { return 0; }
+
 protected:
 	void Swap(FRHIBuffer& Other)
 	{
diff --git a/Engine/Source/Runtime/RHICore/Private/RHIPoolAllocator.cpp b/Engine/Source/Runtime/RHICore/Private/RHIPoolAllocator.cpp
--- a/Engine/Source/Runtime/RHICore/Private/RHIPoolAllocator.cpp
+++ b/Engine/Source/Runtime/RHICore/Private/RHIPoolAllocator.cpp
@@ -144,6 +144,47 @@ void FRHIPoolAllocationData::AddAfter(FRHIPoolAllocationData* InOther)
 	InOther->PreviousAllocation = this;
 }
 
+// The "AliasAllocation" in the original allocation points to a linked list of aliases.  The aliases
+// are doubly linked with each other via NextAllocation/PreviousAllocation, and also point back to their
+// list head via "AliasAllocation".
+void FRHIPoolAllocationData::AddAlias(FRHIPoolAllocationData* InAlias)
+{
+	if (AliasAllocation)
+	{
+		AliasAllocation->PreviousAllocation = InAlias;
+	}
+
+	InAlias->NextAllocation = AliasAllocation;		// Add link to previous alias linked list head
+	InAlias->PreviousAllocation = nullptr;			// New item is at the front, so it doesn't have a predecessor
+	InAlias->AliasAllocation = this;				// Point to the original allocation that contains the aliases
+
+	AliasAllocation = InAlias;						// Finally, update current alias linked list head
+}
+
+void FRHIPoolAllocationData::RemoveAlias()
+{
+	if (AliasAllocation)
+	{
+		if (PreviousAllocation)
+		{
+			PreviousAllocation->NextAllocation = NextAllocation;
+		}
+		else
+		{
+			// Removing item at list head
+			AliasAllocation->AliasAllocation = NextAllocation;
+		}
+		if (NextAllocation)
+		{
+			NextAllocation->PreviousAllocation = PreviousAllocation;
+		}
+
+		NextAllocation = nullptr;
+		PreviousAllocation = nullptr;
+		AliasAllocation = nullptr;
+	}
+}
+
 
 /**
  * Sort predicate that sorts on size of the allocation data
diff --git a/Engine/Source/Runtime/RHICore/Public/RHIPoolAllocator.h b/Engine/Source/Runtime/RHICore/Public/RHIPoolAllocator.h
--- a/Engine/Source/Runtime/RHICore/Public/RHIPoolAllocator.h
+++ b/Engine/Source/Runtime/RHICore/Public/RHIPoolAllocator.h
@@ -52,6 +52,11 @@ public:
 	void AddBefore(FRHIPoolAllocationData* InOther);
 	void AddAfter(FRHIPoolAllocationData* InOther);
 
+	// Alias operations
+	void AddAlias(FRHIPoolAllocationData* InOther);
+	void RemoveAlias();
+	FRHIPoolAllocationData* GetFirstAlias() const { return AliasAllocation; }
+
 	// Type getters
 	bool IsHead() const { return GetAllocationType() == EAllocationType::Head; }
 	bool IsFree() const { return GetAllocationType() == EAllocationType::Free; }
@@ -106,6 +111,7 @@ private:
 	// TODO: Link list data could be stored as index into a fixed array managed by the memory pool. Then only a uin16 or uint32 is needed here to get the next and previous index
 	FRHIPoolAllocationData* PreviousAllocation;
 	FRHIPoolAllocationData* NextAllocation;
+	FRHIPoolAllocationData* AliasAllocation;		// Linked list of aliases
 
 	FRHIPoolResource* Owner;
 };
diff --git a/Engine/Source/Runtime/Renderer/Private/RayTracing/RayTracingScene.cpp b/Engine/Source/Runtime/Renderer/Private/RayTracing/RayTracingScene.cpp
--- a/Engine/Source/Runtime/Renderer/Private/RayTracing/RayTracingScene.cpp
+++ b/Engine/Source/Runtime/Renderer/Private/RayTracing/RayTracingScene.cpp
@@ -184,10 +184,12 @@ void FRayTracingScene::Create(FRDGBuilder& GraphBuilder, const FGPUScene& GPUSce
 							0, 
 							SceneInitializer.ReferencedGeometries.Num() * sizeof(FRayTracingAccelerationStructureAddress), RLM_WriteOnly);
 
+						const uint32 ParentGPUIndex = AccelerationStructureAddressesBuffer.Buffer->GetParentGPUIndex();
+
 						const uint32 NumGeometries = SceneInitializer.ReferencedGeometries.Num();
 						for (uint32 GeometryIndex = 0; GeometryIndex < NumGeometries; ++GeometryIndex)
 						{
-							AddressesPtr[GeometryIndex] = SceneInitializer.ReferencedGeometries[GeometryIndex]->GetAccelerationStructureAddress(RHICmdList.GetGPUMask().ToIndex());
+							AddressesPtr[GeometryIndex] = SceneInitializer.ReferencedGeometries[GeometryIndex]->GetAccelerationStructureAddress(ParentGPUIndex);
 						}
 
 						RHICmdList.UnlockBuffer(AccelerationStructureAddressesBuffer.Buffer);
