Transition vertex buffers right before calling IASetVertexBuffers, instead doing it when the high level calls RHISetStreamSource.

If a command list flush happens in the middle of a sequence of draws which use the same vertex buffer, the tracked internal state of the buffer is reset. Since the high level doesn't call RHISetStreamSource again after the flush, no pending barrier was inserted. If a subsequent dispatch needed the buffer in a different state, that transition was erroneously inserted in the patch-up command list, which modified the state too soon (before the draws which still need the buffer in the VB state). Adding the transition in FD3D12DescriptorCache::SetVertexBuffers ensures that the correct state is applied, since this method is called internally before each draw.

Resubmit with fixes for level reload caued by dangling pointers in the state cache: state cache is cleared during calls to RHISubmitCommandsHint when scene or viewport rendering is not enabled (eg level loading)

diff --git a/Engine/Source/Runtime/D3D12RHI/Private/D3D12CommandContext.cpp b/Engine/Source/Runtime/D3D12RHI/Private/D3D12CommandContext.cpp
--- a/Engine/Source/Runtime/D3D12RHI/Private/D3D12CommandContext.cpp
+++ b/Engine/Source/Runtime/D3D12RHI/Private/D3D12CommandContext.cpp
@@ -743,10 +743,14 @@ void FD3D12CommandContextBase::UpdateMemoryStats()
 
 void FD3D12CommandContext::RHIBeginScene()
 {
+	ensure(!bDrawingScene);
+	bDrawingScene = true;
 }
 
 void FD3D12CommandContext::RHIEndScene()
 {
+	ensure(bDrawingScene);
+	bDrawingScene = false;
 }
 
 #if PLATFORM_USE_BACKBUFFER_WRITE_TRANSITION_TRACKING
diff --git a/Engine/Source/Runtime/D3D12RHI/Private/D3D12CommandContext.h b/Engine/Source/Runtime/D3D12RHI/Private/D3D12CommandContext.h
--- a/Engine/Source/Runtime/D3D12RHI/Private/D3D12CommandContext.h
+++ b/Engine/Source/Runtime/D3D12RHI/Private/D3D12CommandContext.h
@@ -55,6 +55,8 @@ public:
 
 	virtual void RHISetAsyncComputeBudget(EAsyncComputeBudget Budget) {}
 
+	bool IsDrawingSceneOrViewport() const {	return bDrawingScene || bDrawingViewport; }
+
 protected:
 	virtual FD3D12CommandContext* GetContext(uint32 InGPUIndex) = 0;
 
@@ -63,6 +65,8 @@ protected:
 
 	FRHIGPUMask GPUMask;
 
+	bool bDrawingViewport = false;
+	bool bDrawingScene = false;
 	bool bTrackingEvents;
 	const ED3D12CommandQueueType CommandQueueType;
 	const bool bIsDefaultContext;
diff --git a/Engine/Source/Runtime/D3D12RHI/Private/D3D12Commands.cpp b/Engine/Source/Runtime/D3D12RHI/Private/D3D12Commands.cpp
--- a/Engine/Source/Runtime/D3D12RHI/Private/D3D12Commands.cpp
+++ b/Engine/Source/Runtime/D3D12RHI/Private/D3D12Commands.cpp
@@ -2319,6 +2319,12 @@ void FD3D12CommandContext::RHISubmitCommandsHint()
 
 	// Submit the work we have so far, and start a new command list.
 	FlushCommands();
+
+	// Clear the state on requested command hint flush and not in scene or viewport rendering
+	if (!IsDrawingSceneOrViewport())
+	{
+		ClearState();
+	}
 }
 
 #define USE_COPY_QUEUE_FOR_RESOURCE_SYNC 1
diff --git a/Engine/Source/Runtime/D3D12RHI/Private/D3D12DescriptorCache.cpp b/Engine/Source/Runtime/D3D12RHI/Private/D3D12DescriptorCache.cpp
--- a/Engine/Source/Runtime/D3D12RHI/Private/D3D12DescriptorCache.cpp
+++ b/Engine/Source/Runtime/D3D12RHI/Private/D3D12DescriptorCache.cpp
@@ -284,6 +284,19 @@ void FD3D12DescriptorCache::SetVertexBuffers(FD3D12VertexBufferCache& Cache)
 
 	CmdContext->CommandListHandle.UpdateResidency(Cache.ResidencyHandles, Count);
 	CmdContext->CommandListHandle->IASetVertexBuffers(0, Count, Cache.CurrentVertexBufferViews);
+
+	for (uint32 i = 0; i < Count; ++i)
+	{
+		if (Cache.CurrentVertexBufferResources[i])
+		{
+			FD3D12Resource* Resource = Cache.CurrentVertexBufferResources[i]->GetResource();
+			if (Resource->RequiresResourceStateTracking())
+			{
+				check(Resource->GetSubresourceCount() == 1);
+				FD3D12DynamicRHI::TransitionResource(CmdContext->CommandListHandle, Resource, D3D12_RESOURCE_STATE_TBD, D3D12_RESOURCE_STATE_VERTEX_AND_CONSTANT_BUFFER, D3D12_RESOURCE_BARRIER_ALL_SUBRESOURCES, FD3D12DynamicRHI::ETransitionMode::Validate);
+			}
+		}
+	}
 }
 
 template <EShaderFrequency ShaderStage>
diff --git a/Engine/Source/Runtime/D3D12RHI/Private/D3D12StateCache.cpp b/Engine/Source/Runtime/D3D12RHI/Private/D3D12StateCache.cpp
--- a/Engine/Source/Runtime/D3D12RHI/Private/D3D12StateCache.cpp
+++ b/Engine/Source/Runtime/D3D12RHI/Private/D3D12StateCache.cpp
@@ -1245,17 +1245,6 @@ void FD3D12StateCacheBase::InternalSetStreamSource(FD3D12ResourceLocation* Verte
 			PipelineState.Graphics.VBCache.MaxBoundVertexBufferIndex = INDEX_NONE;
 		}
 	}
-
-	if (VertexBufferLocation)
-	{
-		FD3D12Resource* const pResource = VertexBufferLocation->GetResource();
-		if (pResource && pResource->RequiresResourceStateTracking())
-		{
-			check(pResource->GetSubresourceCount() == 1);
-			FD3D12DynamicRHI::TransitionResource(CmdContext->CommandListHandle, pResource, D3D12_RESOURCE_STATE_TBD, D3D12_RESOURCE_STATE_VERTEX_AND_CONSTANT_BUFFER, D3D12_RESOURCE_BARRIER_ALL_SUBRESOURCES, FD3D12DynamicRHI::ETransitionMode::Validate);
-		}
-	}
-
 }
 
 template <EShaderFrequency ShaderFrequency>
diff --git a/Engine/Source/Runtime/D3D12RHI/Private/D3D12Viewport.cpp b/Engine/Source/Runtime/D3D12RHI/Private/D3D12Viewport.cpp
--- a/Engine/Source/Runtime/D3D12RHI/Private/D3D12Viewport.cpp
+++ b/Engine/Source/Runtime/D3D12RHI/Private/D3D12Viewport.cpp
@@ -959,6 +959,9 @@ void FD3D12DynamicRHI::RHITick(float DeltaTime)
 
 void FD3D12CommandContextBase::RHIBeginDrawingViewport(FRHIViewport* ViewportRHI, FRHITexture* RenderTargetRHI)
 {
+	ensure(!bDrawingViewport);
+	bDrawingViewport = true;
+
 	FD3D12Viewport* Viewport = FD3D12DynamicRHI::ResourceCast(ViewportRHI);
 
 	SCOPE_CYCLE_COUNTER(STAT_D3D12PresentTime);
@@ -996,6 +999,9 @@ void FD3D12CommandContextBase::RHIEndDrawingViewport(FRHIViewport* ViewportRHI,
 	FD3D12DynamicRHI& RHI = *ParentAdapter->GetOwningRHI();
 	FD3D12Viewport* Viewport = FD3D12DynamicRHI::ResourceCast(ViewportRHI);
 
+	ensure(bDrawingViewport);
+	bDrawingViewport = false;
+
 #if !UE_BUILD_SHIPPING
 	if (RHIConsoleVariables::LogViewportEvents)
 	{
