Fix for Test/Shipping monolithic configs

diff --git a/Engine/Source/Runtime/D3D12RHI/Private/D3D12CommandContext.h b/Engine/Source/Runtime/D3D12RHI/Private/D3D12CommandContext.h
--- a/Engine/Source/Runtime/D3D12RHI/Private/D3D12CommandContext.h
+++ b/Engine/Source/Runtime/D3D12RHI/Private/D3D12CommandContext.h
@@ -866,10 +866,12 @@ public:
 		return PhysicalContexts[GPUIndex];
 	}
 
+#if ENABLE_RHI_VALIDATION || WITH_MGPU
 	IRHIComputeContext& GetLowestLevelContext() final override
 	{
 		return *PhysicalContexts[0];
 	}
+#endif
 
 private:
 
