Fix crash when the underlying resource of a vertex buffer has been deleted.

#rnx

diff --git a/Engine/Source/Runtime/D3D12RHI/Private/D3D12DescriptorCache.cpp b/Engine/Source/Runtime/D3D12RHI/Private/D3D12DescriptorCache.cpp
--- a/Engine/Source/Runtime/D3D12RHI/Private/D3D12DescriptorCache.cpp
+++ b/Engine/Source/Runtime/D3D12RHI/Private/D3D12DescriptorCache.cpp
@@ -290,7 +290,7 @@ void FD3D12DescriptorCache::SetVertexBuffers(FD3D12VertexBufferCache& Cache)
 		if (Cache.CurrentVertexBufferResources[i])
 		{
 			FD3D12Resource* Resource = Cache.CurrentVertexBufferResources[i]->GetResource();
-			if (Resource->RequiresResourceStateTracking())
+			if (Resource && Resource->RequiresResourceStateTracking())
 			{
 				check(Resource->GetSubresourceCount() == 1);
 				FD3D12DynamicRHI::TransitionResource(CmdContext->CommandListHandle, Resource, D3D12_RESOURCE_STATE_TBD, D3D12_RESOURCE_STATE_VERTEX_AND_CONSTANT_BUFFER, D3D12_RESOURCE_BARRIER_ALL_SUBRESOURCES, FD3D12DynamicRHI::ETransitionMode::Validate);
