Keep track of how long a query is valid depending on the query heap settings and return false on RHIGetRenderQueryResult when the query is too old (no valid data anymore)

diff --git a/Engine/Source/Runtime/D3D12RHI/Private/D3D12Device.cpp b/Engine/Source/Runtime/D3D12RHI/Private/D3D12Device.cpp
--- a/Engine/Source/Runtime/D3D12RHI/Private/D3D12Device.cpp
+++ b/Engine/Source/Runtime/D3D12RHI/Private/D3D12Device.cpp
@@ -25,7 +25,7 @@ FD3D12Device::FD3D12Device(FRHIGPUMask InGPUMask, FD3D12Adapter* InAdapter) :
 	},
 	GlobalSamplerHeap(this),
 	OnlineDescriptorManager(this),
-	OcclusionQueryHeap(this, D3D12_QUERY_TYPE_OCCLUSION, 65536, 4 /*frames to keep results */ * 1 /*batches per frame*/),
+	OcclusionQueryHeap(this, D3D12_QUERY_TYPE_OCCLUSION, 65536, 4 /*frames to keep results */, 1 /*batches per frame*/),
 #if WITH_PROFILEGPU || D3D12_SUBMISSION_GAP_RECORDER
 	CmdListExecTimeQueryHeap(new FD3D12LinearQueryHeap(this, D3D12_QUERY_HEAP_TYPE_TIMESTAMP, 8192)),
 #endif
@@ -37,7 +37,7 @@ FD3D12Device::FD3D12Device(FRHIGPUMask InGPUMask, FD3D12Adapter* InAdapter) :
 {
 	for (uint32 QueueType = 0; QueueType < (uint32)ED3D12CommandQueueType::Count; ++QueueType)
 	{
-		TimestampQueryHeaps[QueueType] = new FD3D12QueryHeap(this, D3D12_QUERY_TYPE_TIMESTAMP, 8192, 4 /*frames to keep results */ * 5 /*batches per frame*/);
+		TimestampQueryHeaps[QueueType] = new FD3D12QueryHeap(this, D3D12_QUERY_TYPE_TIMESTAMP, 8192, 4 /*frames to keep results */, 5 /*batches per frame*/);
 	}
 
 	InitPlatformSpecific();
diff --git a/Engine/Source/Runtime/D3D12RHI/Private/D3D12Query.cpp b/Engine/Source/Runtime/D3D12RHI/Private/D3D12Query.cpp
--- a/Engine/Source/Runtime/D3D12RHI/Private/D3D12Query.cpp
+++ b/Engine/Source/Runtime/D3D12RHI/Private/D3D12Query.cpp
@@ -144,7 +144,11 @@ bool FD3D12DynamicRHI::RHIGetRenderQueryResult(FRHIRenderQuery* QueryRHI, uint64
 		else
 		{
 			OutResult = FMath::Max<uint64>(Query->Result, OutResult);
-			bSuccess = true;
+
+			// Make sure the query is still valid (not to old and query data already reused for a new frame)
+			int32 FramesToKeepResults = Adapter.GetDevice(GPUIndex)->GetOcclusionQueryHeap(ED3D12CommandQueueType::Direct)->GetNumFramesToKeepResults();
+			int32 LastValidFrameFenceValue = FMath::Max(0, (int32)Adapter.GetFrameFence().GetCurrentFence() - FramesToKeepResults);
+			bSuccess = (int32)Query->FrameSubmitted >= LastValidFrameFenceValue;
 		}
 	}
 	return bSuccess;
@@ -231,10 +235,11 @@ void FD3D12CommandContext::RHIEndOcclusionQueryBatch()
 * class FD3D12QueryHeap
 *=============================================================================*/
 
-FD3D12QueryHeap::FD3D12QueryHeap(FD3D12Device* InParent, const D3D12_QUERY_TYPE InQueryType, uint32 InQueryHeapCount, uint32 InMaxActiveBatches)
+FD3D12QueryHeap::FD3D12QueryHeap(FD3D12Device* InParent, const D3D12_QUERY_TYPE InQueryType, uint32 InQueryHeapCount, uint32 InNumFramesToKeepResults, uint32 InBatchPerFrame)
 	: FD3D12DeviceChild(InParent)
 	, FD3D12SingleNodeGPUObject(InParent->GetGPUMask())
-	, LastBatch(InMaxActiveBatches - 1)
+	, NumFramesToKeepResults(InNumFramesToKeepResults)
+	, BatchesPerFrame(InBatchPerFrame)
 	, ActiveAllocatedElementCount(0)
 	, LastAllocatedElement(InQueryHeapCount - 1)
 	, QueryType(InQueryType)
@@ -246,8 +251,11 @@ FD3D12QueryHeap::FD3D12QueryHeap(FD3D12Device* InParent, const D3D12_QUERY_TYPE
 
 	CurrentQueryBatch.Clear();
 
-	ActiveQueryBatches.Reserve(InMaxActiveBatches);
-	ActiveQueryBatches.AddZeroed(InMaxActiveBatches);
+	uint32 MaxActiveBatches = NumFramesToKeepResults * BatchesPerFrame;
+	LastBatch = MaxActiveBatches - 1;
+
+	ActiveQueryBatches.Reserve(MaxActiveBatches);
+	ActiveQueryBatches.AddZeroed(MaxActiveBatches);
 
 	// Don't Init() until the RHI has created the device
 }
@@ -463,6 +471,10 @@ void FD3D12QueryHeap::EndQuery(FD3D12CommandContext& CmdContext, FD3D12RenderQue
 	if (QueryType == D3D12_QUERY_TYPE_OCCLUSION)
 	{
 		check(CurrentQueryBatch.bOpen);
+		if (GetParentDevice())
+		{
+			RenderQuery->FrameSubmitted = GetParentDevice()->GetParentAdapter()->GetFrameFence().GetCurrentFence();
+		}
 	}
 	else
 	{
diff --git a/Engine/Source/Runtime/D3D12RHI/Private/D3D12Query.h b/Engine/Source/Runtime/D3D12RHI/Private/D3D12Query.h
--- a/Engine/Source/Runtime/D3D12RHI/Private/D3D12Query.h
+++ b/Engine/Source/Runtime/D3D12RHI/Private/D3D12Query.h
@@ -116,11 +116,13 @@ private:
 	};
 
 public:
-	FD3D12QueryHeap(class FD3D12Device* InParent, const D3D12_QUERY_TYPE InQueryType, uint32 InQueryHeapCount, uint32 InMaxActiveBatches);
+	FD3D12QueryHeap(class FD3D12Device* InParent, const D3D12_QUERY_TYPE InQueryType, uint32 InQueryHeapCount, uint32 InNumFramesToKeepResults, uint32 InBatchPerFrame);
 
 	void Init();
 	void Destroy();
 
+	uint32 GetNumFramesToKeepResults() const { return NumFramesToKeepResults; }
+
 	// Start tracking a new batch of begin/end query calls that will be resolved together
 	void StartQueryBatch(FD3D12CommandContext& CmdContext, uint32 NumQueriesInBatch);
 	// Stop tracking the current batch of begin/end query calls that will be resolved together.
@@ -143,6 +145,9 @@ private:
 private:
 	QueryBatch CurrentQueryBatch;                       // The current recording batch.
 
+	uint32 NumFramesToKeepResults;						// Amount of frames which have valid results
+	uint32 BatchesPerFrame;								// Amount of batches per frame expected
+
 	TArray<QueryBatch> ActiveQueryBatches;              // List of active query batches. The data for these is in use.
 	uint32 LastBatch;                                   // The index of the newest batch.
 
