Use pool allocator also on Xbox and some minor fixes: (disabled in this CL - will be enabled in follow up change)
- only check heap flags for placed resource (per resource flags are not important than)
- disable dedicated pool for BLAS scratch because it has more memory overhead than it saves (most BLAS scratch is transient now)
- don't add require UAV flag for BLAS allocation to match the renderer side allocation (make sure they are part of the same pool now)
- check if Indirect buffer needs to be committed allocated before creating the pool which won't be used
- split pool allocator define for databuffers & textures
- allow tweaking of pool sizes via defines instead of CVars

diff --git a/Engine/Source/Runtime/D3D12RHI/Private/D3D12Allocation.cpp b/Engine/Source/Runtime/D3D12RHI/Private/D3D12Allocation.cpp
--- a/Engine/Source/Runtime/D3D12RHI/Private/D3D12Allocation.cpp
+++ b/Engine/Source/Runtime/D3D12RHI/Private/D3D12Allocation.cpp
@@ -10,6 +10,13 @@
 #include "Misc/BufferedOutputDevice.h"
 #include "HAL/PlatformStackWalk.h"
 
+// Fix for random GPU crashes on draw indirects on multiple IHVs. Force all indirect arg buffers as committed resources (see UE-115982)
+static int32 GD3D12AllowPoolAllocateIndirectArgBuffers = 0;
+static FAutoConsoleVariableRef CVarD3D12AllowPoolAllocateIndirectArgBuffers(
+	TEXT("d3d12.AllowPoolAllocateIndirectArgBuffers"),
+	GD3D12AllowPoolAllocateIndirectArgBuffers,
+	TEXT("Allow indirect args to be pool allocated (otherwise they will be committed resources) (default: 0)"),
+	ECVF_ReadOnly);
 
 #if D3D12RHI_SEGREGATED_TEXTURE_ALLOC
 static int32 GD3D12ReadOnlyTextureAllocatorMinPoolSize = 4 * 1024 * 1024;
@@ -35,21 +42,44 @@ static FAutoConsoleVariableRef CVarD3D12ReadOnlyTextureAllocatorMaxPoolSize(
 	ECVF_ReadOnly);
 #endif
 
-#if USE_POOL_ALLOCATOR
-static int32 GD3D12PoolAllocatorBufferVRAMPoolSize = 32 * 1024 * 1024;
-static FAutoConsoleVariableRef CVarD3D12PoolAllocatorBufferVRAMPoolSize(
-	TEXT("d3d12.PoolAllocator.BufferVRAMPoolSize"),
-	GD3D12PoolAllocatorBufferVRAMPoolSize,
-	TEXT("Pool size of a single VRAM DataBuffer memory pool (default 32MB)"),
+#if USE_BUFFER_POOL_ALLOCATOR
+
+#if !defined(BUFFER_POOL_DEFRAG_MAX_COPY_SIZE_PER_FRAME)
+#define BUFFER_POOL_DEFRAG_MAX_COPY_SIZE_PER_FRAME				32*1024*1024
+#endif
+
+#if !defined(BUFFER_POOL_DEFAULT_POOL_SIZE)
+#define BUFFER_POOL_DEFAULT_POOL_SIZE							32*1024*1024
+#endif
+
+#if !defined(BUFFER_POOL_DEFAULT_POOL_MAX_ALLOC_SIZE)
+#define BUFFER_POOL_DEFAULT_POOL_MAX_ALLOC_SIZE					16*1024*1024
+#endif
+
+#if !defined(BUFFER_POOL_RT_ACCELERATION_STRUCTURE_POOL_SIZE)
+#define BUFFER_POOL_RT_ACCELERATION_STRUCTURE_POOL_SIZE			32*1024*1024
+#endif
+
+#if !defined(BUFFER_POOL_RT_ACCELERATION_STRUCTURE_MAX_ALLOC_SIZE)
+#define BUFFER_POOL_RT_ACCELERATION_STRUCTURE_MAX_ALLOC_SIZE	16*1024*1024
+#endif
+
+static int32 GD3D12VRAMBufferPoolDefrag = 1;
+static FAutoConsoleVariableRef CVarD3D12VRAMBufferPoolDefrag(
+	TEXT("d3d12.VRAMBufferPoolDefrag"),
+	GD3D12VRAMBufferPoolDefrag,
+	TEXT("Defrag the VRAM buffer pool"),
 	ECVF_ReadOnly);
 
-static int32 GD3D12PoolAllocatorBufferVRAMMaxAllocationSize = 16 * 1024 * 1024;
-static FAutoConsoleVariableRef CVarD3D12PoolAllocatorBufferMaxAllocationSize(
-	TEXT("d3d12.PoolAllocator.BufferMaxAllocationSize"),
-	GD3D12PoolAllocatorBufferVRAMMaxAllocationSize,
-	TEXT("Maximum size of a single allocation in the VRAM DataBuffer pool allocator (default 16MB)"),
+static int32 GD3D12VRAMBufferPoolDefragMaxCopySizePerFrame = BUFFER_POOL_DEFRAG_MAX_COPY_SIZE_PER_FRAME;
+static FAutoConsoleVariableRef CVarD3D12VRAMBufferPoolDefragMaxCopySizePerFrame(
+	TEXT("d3d12.VRAMBufferPoolDefrag.MaxCopySizePerFrame"),
+	GD3D12VRAMBufferPoolDefragMaxCopySizePerFrame,
+	TEXT("Max amount of data to copy during defragmentation in a single frame (default 32MB)"),
 	ECVF_ReadOnly);
+#endif // USE_BUFFER_POOL_ALLOCATOR
 
+#if USE_TEXTURE_POOL_ALLOCATOR
 static int32 GD3D12PoolAllocatorReadOnlyTextureVRAMPoolSize = 64 * 1024 * 1024;
 static FAutoConsoleVariableRef CVarD3D12PoolAllocatorReadOnlyTextureVRAMPoolSize(
 	TEXT("d3d12.PoolAllocator.ReadOnlyTextureVRAMPoolSize"),
@@ -91,22 +121,7 @@ static FAutoConsoleVariableRef CVarD3D12VRAMTexturePoolDefragMaxCopySizePerFrame
 	GD3D12VRAMTexturePoolDefragMaxCopySizePerFrame,
 	TEXT("Max amount of data to copy during defragmentation in a single frame (default 32MB)"),
 	ECVF_ReadOnly);
-
-static int32 GD3D12VRAMBufferPoolDefrag = 1;
-static FAutoConsoleVariableRef CVarD3D12VRAMBufferPoolDefrag(
-	TEXT("d3d12.VRAMBufferPoolDefrag"),
-	GD3D12VRAMBufferPoolDefrag,
-	TEXT("Defrag the VRAM buffer pool (disabled by default)"),
-	ECVF_ReadOnly);
-
-static int32 GD3D12VRAMBufferPoolDefragMaxCopySizePerFrame = 32 * 1024 * 1024;
-static FAutoConsoleVariableRef CVarD3D12VRAMBufferPoolDefragMaxCopySizePerFrame(
-	TEXT("d3d12.VRAMBufferPoolDefrag.MaxCopySizePerFrame"),
-	GD3D12VRAMBufferPoolDefragMaxCopySizePerFrame,
-	TEXT("Max amount of data to copy during defragmentation in a single frame (default 32MB)"),
-	ECVF_ReadOnly);
-
-#endif // USE_POOL_ALLOCATOR
+#endif // USE_TEXTURE_POOL_ALLOCATOR
 
 static int32 GD3D12UploadHeapSmallBlockMaxAllocationSize = 64 * 1024;
 static FAutoConsoleVariableRef CVarD3D12UploadHeapSmallBlockMaxAllocationSize(
@@ -1313,17 +1328,17 @@ FD3D12BufferPool* FD3D12DefaultBufferAllocator::CreateBufferPool(D3D12_HEAP_TYPE
 	FD3D12Device* Device = GetParentDevice();
 	FD3D12ResourceInitConfig InitConfig = FD3D12BufferPool::GetResourceAllocatorInitConfig(InHeapType, InResourceFlags, InBufferUsage);
 
-#if USE_POOL_ALLOCATOR
+#if USE_BUFFER_POOL_ALLOCATOR
 
 	const FString Name(L"D3D12 Pool Allocator");
 	EResourceAllocationStrategy AllocationStrategy = FD3D12PoolAllocator::GetResourceAllocationStrategy(InResourceFlags, InResourceStateMode, Alignment);
-	uint64 PoolSize = InHeapType == D3D12_HEAP_TYPE_READBACK ? READBACK_BUFFER_POOL_DEFAULT_POOL_SIZE : GD3D12PoolAllocatorBufferVRAMPoolSize;
+	uint64 PoolSize = InHeapType == D3D12_HEAP_TYPE_READBACK ? READBACK_BUFFER_POOL_DEFAULT_POOL_SIZE : BUFFER_POOL_DEFAULT_POOL_SIZE;
 	uint64 PoolAlignment = (AllocationStrategy == EResourceAllocationStrategy::kPlacedResource) ? MIN_PLACED_RESOURCE_SIZE : kD3D12ManualSubAllocationAlignment;
-	uint64 MaxAllocationSize = InHeapType == D3D12_HEAP_TYPE_READBACK ? READBACK_BUFFER_POOL_MAX_ALLOC_SIZE : GD3D12PoolAllocatorBufferVRAMMaxAllocationSize;
+	uint64 MaxAllocationSize = InHeapType == D3D12_HEAP_TYPE_READBACK ? READBACK_BUFFER_POOL_MAX_ALLOC_SIZE : BUFFER_POOL_DEFAULT_POOL_MAX_ALLOC_SIZE;
 	FRHIMemoryPool::EFreeListOrder FreeListOrder = FRHIMemoryPool::EFreeListOrder::SortBySize;
 
 	// Disable defrag if not Default memory
-	bool bDefragEnabled = (InitConfig.HeapType == D3D12_HEAP_TYPE_DEFAULT);
+	bool bDefragEnabled = (InitConfig.HeapType == D3D12_HEAP_TYPE_DEFAULT && AllocationStrategy == EResourceAllocationStrategy::kPlacedResource);
 
 #if D3D12_RHI_RAYTRACING
 	// Disable defrag on the RT Acceleration pool - #kenzo_todo
@@ -1332,12 +1347,16 @@ FD3D12BufferPool* FD3D12DefaultBufferAllocator::CreateBufferPool(D3D12_HEAP_TYPE
 	if (InitConfig.InitialResourceState == D3D12_RESOURCE_STATE_RAYTRACING_ACCELERATION_STRUCTURE)
 	{
 		bDefragEnabled = false;
+
+		// Use custom pool and allocation size for RT structures because they don't defrag and will thus 'waste' more memory
+		PoolSize = BUFFER_POOL_RT_ACCELERATION_STRUCTURE_POOL_SIZE;
+		MaxAllocationSize = BUFFER_POOL_RT_ACCELERATION_STRUCTURE_MAX_ALLOC_SIZE;
 	}
 #endif // D3D12_RHI_RAYTRACING
 
 	FD3D12BufferPool* NewPool = new FD3D12PoolAllocator(Device, GetVisibilityMask(), InitConfig, Name, AllocationStrategy, PoolSize, PoolAlignment, MaxAllocationSize, FreeListOrder, bDefragEnabled);
 
-#else // USE_POOL_ALLOCATOR
+#else // USE_BUFFER_POOL_ALLOCATOR
 
 	EResourceAllocationStrategy AllocationStrategy = FD3D12DefaultBufferPool::GetResourceAllocationStrategy(InResourceFlags, InResourceStateMode, Alignment);
 
@@ -1357,7 +1376,7 @@ FD3D12BufferPool* FD3D12DefaultBufferAllocator::CreateBufferPool(D3D12_HEAP_TYPE
 
 	FD3D12DefaultBufferPool* NewPool = new FD3D12DefaultBufferPool(Device, Allocator);
 
-#endif // USE_POOL_ALLOCATOR
+#endif // USE_BUFFER_POOL_ALLOCATOR
 
 	DefaultBufferPools.Add(NewPool);
 	return NewPool;
@@ -1404,6 +1423,21 @@ D3D12_RESOURCE_STATES FD3D12DefaultBufferAllocator::GetDefaultInitialResourceSta
 // Grab a buffer from the available buffers or create a new buffer if none are available
 void FD3D12DefaultBufferAllocator::AllocDefaultResource(D3D12_HEAP_TYPE InHeapType, const D3D12_RESOURCE_DESC& InResourceDesc, EBufferUsageFlags InBufferUsage, ED3D12ResourceStateMode InResourceStateMode, D3D12_RESOURCE_STATES InCreateState, FD3D12ResourceLocation& ResourceLocation, uint32 Alignment, const TCHAR* Name)
 {
+	// Force indirect args to stand alone allocations instead of pooled
+	if (!GD3D12AllowPoolAllocateIndirectArgBuffers && EnumHasAnyFlags(InBufferUsage, BUF_DrawIndirect))
+	{
+		ResourceLocation.Clear();
+
+		FD3D12Resource* NewResource = nullptr;
+		const D3D12_HEAP_PROPERTIES HeapProps = CD3DX12_HEAP_PROPERTIES(InHeapType, GetGPUMask().GetNative(), GetVisibilityMask().GetNative());
+		D3D12_RESOURCE_DESC Desc = InResourceDesc;
+		Desc.Alignment = 0;
+		VERIFYD3D12RESULT(GetParentDevice()->GetParentAdapter()->CreateCommittedResource(Desc, GetGPUMask(), HeapProps, InCreateState, InResourceStateMode, InCreateState, nullptr, &NewResource, Name, false));
+
+		ResourceLocation.AsStandAlone(NewResource, InResourceDesc.Width);
+		return;
+	}
+
 	// Patch out deny shader resource because it doesn't add anything for buffers and allows more pool sharing
 	// TODO: check if this is different on Xbox?
 	D3D12_RESOURCE_DESC ResourceDesc = InResourceDesc;
@@ -1449,7 +1483,7 @@ void FD3D12DefaultBufferAllocator::FreeDefaultBufferPools()
 
 void FD3D12DefaultBufferAllocator::BeginFrame()
 {
-#if USE_POOL_ALLOCATOR
+#if USE_BUFFER_POOL_ALLOCATOR
 
 	if (GD3D12VRAMBufferPoolDefrag > 0 && GD3D12VRAMBufferPoolDefragMaxCopySizePerFrame > 0)
 	{
@@ -1480,7 +1514,7 @@ void FD3D12DefaultBufferAllocator::BeginFrame()
 		}
 		CommandContext.RHIPopEvent();
 	}
-#endif // USE_POOL_ALLOCATOR
+#endif // USE_BUFFER_POOL_ALLOCATOR
 }
 
 
@@ -1535,7 +1569,7 @@ void FD3D12DefaultBufferAllocator::UpdateMemoryStats()
 //	Texture Allocator
 //-----------------------------------------------------------------------------
 
-#if USE_POOL_ALLOCATOR
+#if USE_TEXTURE_POOL_ALLOCATOR
 FD3D12TextureAllocatorPool::FD3D12TextureAllocatorPool(FD3D12Device* Device, FRHIGPUMask VisibilityNode) :
 	FD3D12DeviceChild(Device),
 	FD3D12MultiNodeGPUObject(Device->GetGPUMask(), VisibilityNode)
diff --git a/Engine/Source/Runtime/D3D12RHI/Private/D3D12Allocation.h b/Engine/Source/Runtime/D3D12RHI/Private/D3D12Allocation.h
--- a/Engine/Source/Runtime/D3D12RHI/Private/D3D12Allocation.h
+++ b/Engine/Source/Runtime/D3D12RHI/Private/D3D12Allocation.h
@@ -11,7 +11,13 @@ D3D12Allocation.h: A Collection of allocators
 // Enable pool allocator by default only on Windows until all initial issues have been resolved.
 // Xbox is currently using the buddy allocator with a lot smaller pool sizes so the alignment waste is a lot less
 // It also has less of a problem with committed resource allocations
-#define USE_POOL_ALLOCATOR (PLATFORM_WINDOWS)
+#if !defined(USE_BUFFER_POOL_ALLOCATOR)
+#define USE_BUFFER_POOL_ALLOCATOR (PLATFORM_WINDOWS)
+#endif
+
+#if !defined(USE_TEXTURE_POOL_ALLOCATOR)
+#define USE_TEXTURE_POOL_ALLOCATOR (PLATFORM_WINDOWS)
+#endif
 
 // Segregated free list texture allocator
 // Description:
@@ -402,7 +408,7 @@ private:
 	FD3D12MultiBuddyAllocator* Allocator;
 };
 
-#if USE_POOL_ALLOCATOR
+#if USE_BUFFER_POOL_ALLOCATOR
 typedef FD3D12PoolAllocator FD3D12BufferPool;
 #else
 typedef FD3D12DefaultBufferPool FD3D12BufferPool;
@@ -891,7 +897,7 @@ private:
 //	FD3D12TextureAllocator
 //-----------------------------------------------------------------------------
 
-#if USE_POOL_ALLOCATOR
+#if USE_TEXTURE_POOL_ALLOCATOR
 class FD3D12TextureAllocatorPool : public FD3D12DeviceChild, public FD3D12MultiNodeGPUObject
 {
 public:
diff --git a/Engine/Source/Runtime/D3D12RHI/Private/D3D12Buffer.cpp b/Engine/Source/Runtime/D3D12RHI/Private/D3D12Buffer.cpp
--- a/Engine/Source/Runtime/D3D12RHI/Private/D3D12Buffer.cpp
+++ b/Engine/Source/Runtime/D3D12RHI/Private/D3D12Buffer.cpp
@@ -359,11 +359,13 @@ FD3D12Buffer* FD3D12Adapter::CreateRHIBuffer(
 			FD3D12Buffer* NewBuffer = new FD3D12Buffer(Device, Size, InUsage, Stride);
 			NewBuffer->BufferAlignment = Alignment;
 
+#if NAME_OBJECTS
 			if (InDebugName)
 			{
 				TRACE_CPUPROFILER_EVENT_SCOPE(D3D12RHI::SetDebugName);
 				NewBuffer->SetName(InDebugName);
 			}
+#endif // NAME_OBJECTS
 
 			if (Device->GetGPUIndex() == FirstGPUIndex)
 			{
@@ -386,11 +388,13 @@ FD3D12Buffer* FD3D12Adapter::CreateRHIBuffer(
 			FD3D12Buffer* NewBuffer = new FD3D12Buffer(Device, Size, InUsage, Stride);
 			NewBuffer->BufferAlignment = Alignment;
 
+#if NAME_OBJECTS
 			if (InDebugName)
 			{
 				TRACE_CPUPROFILER_EVENT_SCOPE(D3D12RHI::SetDebugName);
 				NewBuffer->SetName(InDebugName);
 			}
+#endif // NAME_OBJECTS
 
 			AllocateBuffer(Device, InDesc, Size, InUsage, InResourceStateMode, InCreateState, Alignment, NewBuffer, NewBuffer->ResourceLocation, TransientMode, ResourceAllocator, InDebugName);
 			
diff --git a/Engine/Source/Runtime/D3D12RHI/Private/D3D12PoolAllocator.cpp b/Engine/Source/Runtime/D3D12RHI/Private/D3D12PoolAllocator.cpp
--- a/Engine/Source/Runtime/D3D12RHI/Private/D3D12PoolAllocator.cpp
+++ b/Engine/Source/Runtime/D3D12RHI/Private/D3D12PoolAllocator.cpp
@@ -7,14 +7,6 @@
 #define NEEDS_D3D12_INDIRECT_ARGUMENT_HEAP_WORKAROUND 0
 #endif
 
-// Fix for random GPU crashes on draw indirects on multiple IHVs. Force all indirect arg buffers as committed resources (see UE-115982)
-static int32 GD3D12AllowPoolAllocateIndirectArgBuffers = 0;
-static FAutoConsoleVariableRef CVarD3D12AllowPoolAllocateIndirectArgBuffers(
-	TEXT("d3d12.AllowPoolAllocateIndirectArgBuffers"),
-	GD3D12AllowPoolAllocateIndirectArgBuffers,
-	TEXT("Allow indirect args to be pool allocated (otherwise they will be committed resources) (default: 0)"),
-	ECVF_ReadOnly);
-
 //-----------------------------------------------------------------------------
 //	FD3D12MemoryPool
 //-----------------------------------------------------------------------------
@@ -207,7 +199,16 @@ bool FD3D12PoolAllocator::SupportsAllocation(D3D12_HEAP_TYPE InHeapType, D3D12_R
 {
 	FD3D12ResourceInitConfig InInitConfig = GetResourceAllocatorInitConfig(InHeapType, InResourceFlags, InBufferUsage);
 	EResourceAllocationStrategy InAllocationStrategy = GetResourceAllocationStrategy(InResourceFlags, InResourceStateMode, Alignment);
-	return (InitConfig == InInitConfig && AllocationStrategy == InAllocationStrategy);
+
+	// Don't check the resource flags & states for places resource but only the heap flags
+	if (InAllocationStrategy == EResourceAllocationStrategy::kPlacedResource && AllocationStrategy == InAllocationStrategy)
+	{
+		return (InitConfig.HeapType == InInitConfig.HeapType && InitConfig.HeapFlags == InInitConfig.HeapFlags);
+	}
+	else
+	{
+		return (InitConfig == InInitConfig && AllocationStrategy == InAllocationStrategy);
+	}
 }
 
 
@@ -238,23 +239,7 @@ void FD3D12PoolAllocator::AllocDefaultResource(D3D12_HEAP_TYPE InHeapType, const
 #endif // D3D12_RHI_RAYTRACING
 #endif  // DO_CHECK
 
-	// Force indirect args to stand alone allocations instead of pooled
-	if (!GD3D12AllowPoolAllocateIndirectArgBuffers && EnumHasAnyFlags(InBufferUsage, BUF_DrawIndirect))
-	{
-		ResourceLocation.Clear();
-
-		FD3D12Resource* NewResource = nullptr;
-		const D3D12_HEAP_PROPERTIES HeapProps = CD3DX12_HEAP_PROPERTIES(InHeapType, GetGPUMask().GetNative(), GetVisibilityMask().GetNative());
-		D3D12_RESOURCE_DESC Desc = InDesc;
-		Desc.Alignment = 0;
-		VERIFYD3D12RESULT(GetParentDevice()->GetParentAdapter()->CreateCommittedResource(Desc, GetGPUMask(), HeapProps, InCreateState, InResourceStateMode, InCreateState, nullptr, &NewResource, InName, false));
-
-		ResourceLocation.AsStandAlone(NewResource, InDesc.Width);
-	}
-	else
-	{
-		AllocateResource(GetParentDevice()->GetGPUIndex(), InHeapType, InDesc, InDesc.Width, InAllocationAlignment, InResourceStateMode, InCreateState, nullptr, InName, ResourceLocation);
-	}
+	AllocateResource(GetParentDevice()->GetGPUIndex(), InHeapType, InDesc, InDesc.Width, InAllocationAlignment, InResourceStateMode, InCreateState, nullptr, InName, ResourceLocation);
 }
 
 
@@ -382,7 +367,8 @@ void FD3D12PoolAllocator::AllocateResource(uint32 GPUIndex, D3D12_HEAP_TYPE InHe
 		else
 		{
 			// Allocate Standalone - move to owner of resource because this allocator should only manage pooled allocations (needed for now to do the same as FD3D12DefaultBufferPool)
-			VERIFYD3D12RESULT(Adapter->CreateCommittedResource(Desc, GetGPUMask(), HeapProps, InCreateState, InResourceStateMode, InCreateState, InClearValue, &NewResource, InName, false));
+			D3D12_RESOURCE_STATES DefaultState = (InResourceStateMode == ED3D12ResourceStateMode::Default) ? D3D12_RESOURCE_STATE_TBD : InCreateState;
+			VERIFYD3D12RESULT(Adapter->CreateCommittedResource(Desc, GetGPUMask(), HeapProps, InCreateState, InResourceStateMode, DefaultState, InClearValue, &NewResource, InName, false));
 		}
 
 		ResourceLocation.AsStandAlone(NewResource, InSize);
diff --git a/Engine/Source/Runtime/D3D12RHI/Private/D3D12RayTracing.cpp b/Engine/Source/Runtime/D3D12RHI/Private/D3D12RayTracing.cpp
--- a/Engine/Source/Runtime/D3D12RHI/Private/D3D12RayTracing.cpp
+++ b/Engine/Source/Runtime/D3D12RHI/Private/D3D12RayTracing.cpp
@@ -3291,12 +3291,12 @@ static TRefCountPtr<FD3D12Buffer> CreateRayTracingBuffer(FD3D12Adapter* Adapter,
 	FString DebugNameString = DebugName.ToString();
 	ED3D12ResourceTransientMode TransientMode = ED3D12ResourceTransientMode::NonTransient;
 	ID3D12ResourceAllocator* ResourceAllocator = nullptr;
-	D3D12_RESOURCE_DESC BufferDesc = CD3DX12_RESOURCE_DESC::Buffer(Size, D3D12_RESOURCE_FLAG_ALLOW_UNORDERED_ACCESS);
 	FRHIGPUMask GPUMask = FRHIGPUMask::FromIndex(GPUIndex);
 	bool bHasInitialData = false;
 
 	if (Type == ERayTracingBufferType::AccelerationStructure)
 	{
+		D3D12_RESOURCE_DESC BufferDesc = CD3DX12_RESOURCE_DESC::Buffer(Size, D3D12_RESOURCE_FLAG_NONE);
 		Result = Adapter->CreateRHIBuffer(
 			BufferDesc, D3D12_RAYTRACING_ACCELERATION_STRUCTURE_BYTE_ALIGNMENT,
 			0, BufferDesc.Width, BUF_AccelerationStructure,
@@ -3305,10 +3305,13 @@ static TRefCountPtr<FD3D12Buffer> CreateRayTracingBuffer(FD3D12Adapter* Adapter,
 	}
 	else if (Type == ERayTracingBufferType::Scratch)
 	{
+		// Scratch doesn't need single state anymore because there are only a few scratch allocations left and allocating a 
+		// dedicated single state heap for it wastes memory - ideally all scratch allocations should be transient
+		D3D12_RESOURCE_DESC BufferDesc = CD3DX12_RESOURCE_DESC::Buffer(Size, D3D12_RESOURCE_FLAG_ALLOW_UNORDERED_ACCESS);
 		Result = Adapter->CreateRHIBuffer(
 			BufferDesc, D3D12_RAYTRACING_ACCELERATION_STRUCTURE_BYTE_ALIGNMENT,
 			0, BufferDesc.Width, BUF_UnorderedAccess,
-			ED3D12ResourceStateMode::SingleState, D3D12_RESOURCE_STATE_UNORDERED_ACCESS, bHasInitialData,
+			ED3D12ResourceStateMode::Default, D3D12_RESOURCE_STATE_UNORDERED_ACCESS, bHasInitialData,
 			GPUMask, TransientMode, ResourceAllocator, *DebugNameString);
 
 		// Elevates the scratch buffer heap priority, which may help performance / stability in low memory conditions 
