Handle blocks in defragmentation when swapping ownership of allocation data

diff --git a/Engine/Source/Runtime/D3D12RHI/Private/D3D12PoolAllocator.cpp b/Engine/Source/Runtime/D3D12RHI/Private/D3D12PoolAllocator.cpp
--- a/Engine/Source/Runtime/D3D12RHI/Private/D3D12PoolAllocator.cpp
+++ b/Engine/Source/Runtime/D3D12RHI/Private/D3D12PoolAllocator.cpp
@@ -431,7 +431,7 @@ void FD3D12PoolAllocator::DeallocateResource(FD3D12ResourceLocation& ResourceLoc
 				break;
 			}
 		}
-
+		
 		// If still pending copy then clear the copy operation data
 		for (FD3D12VRAMCopyOperation& CopyOperation : PendingCopyOps)
 		{
@@ -662,9 +662,29 @@ void FD3D12PoolAllocator::TransferOwnership(FD3D12ResourceLocation& InSource, FD
 
 	check(IsOwner(InSource));
 
-	// Don't need to lock - ownership simply changed
-	bool bLocked = false;
+	FRHIPoolAllocationData& SourcePoolData = InSource.GetPoolAllocatorPrivateData().PoolData;
 	FRHIPoolAllocationData& DestinationPoolData = InDest.GetPoolAllocatorPrivateData().PoolData;
+
+	// If locked then assume the block is currently being defragged and then the frame fenced operation
+	// also needs to be updated
+	if (SourcePoolData.IsLocked())
+	{
+		bool bFound = false;
+		for (FrameFencedAllocationData& Operation : FrameFencedOperations)
+		{
+			if (Operation.AllocationData == &SourcePoolData)
+			{
+				check(Operation.Operation == FrameFencedAllocationData::EOperation::Unlock);
+				Operation.AllocationData = &DestinationPoolData;
+				bFound = true;
+				break;
+			}
+		}
+		check(bFound);
+	}
+
+	// Take over lock as well - can be defragging
+	bool bLocked = SourcePoolData.IsLocked();
 	DestinationPoolData.MoveFrom(InSource.GetPoolAllocatorPrivateData().PoolData, bLocked);
 	DestinationPoolData.SetOwner(&InDest);
 }
diff --git a/Engine/Source/Runtime/D3D12RHI/Private/D3D12Resources.cpp b/Engine/Source/Runtime/D3D12RHI/Private/D3D12Resources.cpp
--- a/Engine/Source/Runtime/D3D12RHI/Private/D3D12Resources.cpp
+++ b/Engine/Source/Runtime/D3D12RHI/Private/D3D12Resources.cpp
@@ -731,7 +731,11 @@ void FD3D12ResourceLocation::Swap(FD3D12ResourceLocation& Other)
 		::Swap(*this, Other);
 
 		// Restore allocator data and perform pool aware swap
-		PoolData = TmpPoolData;	
+		PoolData = TmpPoolData;
+
+		// Reset the tmp pool data again, because it's not needed anymore - all data copied over
+		TmpPoolData.Init();
+
 		Other.SetPoolAllocator(GetPoolAllocator());
 		GetPoolAllocator()->TransferOwnership(Other, *this);
 	}
diff --git a/Engine/Source/Runtime/RHICore/Private/RHIPoolAllocator.cpp b/Engine/Source/Runtime/RHICore/Private/RHIPoolAllocator.cpp
--- a/Engine/Source/Runtime/RHICore/Private/RHIPoolAllocator.cpp
+++ b/Engine/Source/Runtime/RHICore/Private/RHIPoolAllocator.cpp
@@ -594,7 +594,7 @@ void FRHIMemoryPool::Validate()
 		FRHIPoolAllocationData* FreeBlock = FreeBlocks[FreeBlockIndex];
 		check(!FreeBlock->GetPrev()->IsFree() || FreeBlock->GetPrev()->IsLocked());
 		check(!FreeBlock->GetNext()->IsFree() || FreeBlock->GetNext()->IsLocked());
-		check(FreeBlockIndex == (FreeBlocks.Num() - 1) || FreeBlock->GetSize() <= FreeBlocks[FreeBlockIndex + 1]->GetSize());
+		check(FreeListOrder != EFreeListOrder::SortBySize || FreeBlockIndex == (FreeBlocks.Num() - 1) || FreeBlock->GetSize() <= FreeBlocks[FreeBlockIndex + 1]->GetSize());
 		check(IsAligned(FreeBlock->GetOffset(), PoolAlignment));
 		TotalFreeSize += FreeBlock->GetSize();
 	}
