Implement LockTexture2DArray_RenderThread & UnlockTexture2DArray_RenderThread in d3d12 and agc rhi to remove the unrequired RHI thread flush which isn't required on these platforms

diff --git a/Engine/Source/Runtime/D3D12RHI/Private/D3D12RHIPrivate.h b/Engine/Source/Runtime/D3D12RHI/Private/D3D12RHIPrivate.h
--- a/Engine/Source/Runtime/D3D12RHI/Private/D3D12RHIPrivate.h
+++ b/Engine/Source/Runtime/D3D12RHI/Private/D3D12RHIPrivate.h
@@ -453,6 +453,9 @@ public:
 	virtual void* LockTexture2D_RenderThread(class FRHICommandListImmediate& RHICmdList, FRHITexture2D* Texture, uint32 MipIndex, EResourceLockMode LockMode, uint32& DestStride, bool bLockWithinMiptail, bool bNeedsDefaultRHIFlush = true) override final;
 	virtual void UnlockTexture2D_RenderThread(class FRHICommandListImmediate& RHICmdList, FRHITexture2D* Texture, uint32 MipIndex, bool bLockWithinMiptail, bool bNeedsDefaultRHIFlush = true) override final;
 
+	virtual void* LockTexture2DArray_RenderThread(class FRHICommandListImmediate& RHICmdList, FRHITexture2DArray* Texture, uint32 ArrayIndex, uint32 MipIndex, EResourceLockMode LockMode, uint32& DestStride, bool bLockWithinMiptail) override final;
+	virtual void UnlockTexture2DArray_RenderThread(class FRHICommandListImmediate& RHICmdList, FRHITexture2DArray* Texture, uint32 ArrayIndex, uint32 MipIndex, bool bLockWithinMiptail) override final;
+
 	virtual FTexture2DRHIRef AsyncReallocateTexture2D_RenderThread(class FRHICommandListImmediate& RHICmdList, FRHITexture2D* Texture2D, int32 NewMipCount, int32 NewSizeX, int32 NewSizeY, FThreadSafeCounter* RequestStatus);
 	virtual ETextureReallocationStatus FinalizeAsyncReallocateTexture2D_RenderThread(class FRHICommandListImmediate& RHICmdList, FRHITexture2D* Texture2D, bool bBlockUntilCompleted)
 	{
diff --git a/Engine/Source/Runtime/D3D12RHI/Private/D3D12Texture.cpp b/Engine/Source/Runtime/D3D12RHI/Private/D3D12Texture.cpp
--- a/Engine/Source/Runtime/D3D12RHI/Private/D3D12Texture.cpp
+++ b/Engine/Source/Runtime/D3D12RHI/Private/D3D12Texture.cpp
@@ -2514,6 +2514,13 @@ void FD3D12DynamicRHI::RHIUnlockTexture2D(FRHITexture2D* TextureRHI, uint32 MipI
 	Texture->Unlock(nullptr, MipIndex, 0);
 }
 
+void* FD3D12DynamicRHI::LockTexture2DArray_RenderThread(class FRHICommandListImmediate& RHICmdList, FRHITexture2DArray* TextureRHI, uint32 TextureIndex, uint32 MipIndex, EResourceLockMode LockMode, uint32& DestStride, bool bLockWithinMiptail)
+{
+	check(TextureRHI);
+	FD3D12Texture2DArray* Texture = FD3D12DynamicRHI::ResourceCast(TextureRHI);
+	return Texture->Lock(&RHICmdList, MipIndex, TextureIndex, LockMode, DestStride);
+}
+
 void* FD3D12DynamicRHI::RHILockTexture2DArray(FRHITexture2DArray* TextureRHI, uint32 TextureIndex, uint32 MipIndex, EResourceLockMode LockMode, uint32& DestStride, bool bLockWithinMiptail)
 {
 	check(TextureRHI);
@@ -2521,6 +2528,13 @@ void* FD3D12DynamicRHI::RHILockTexture2DArray(FRHITexture2DArray* TextureRHI, ui
 	return Texture->Lock(nullptr, MipIndex, TextureIndex, LockMode, DestStride);
 }
 
+void FD3D12DynamicRHI::UnlockTexture2DArray_RenderThread(class FRHICommandListImmediate& RHICmdList, FRHITexture2DArray* TextureRHI, uint32 TextureIndex, uint32 MipIndex, bool bLockWithinMiptail)
+{
+	check(TextureRHI);
+	FD3D12Texture2DArray* Texture = FD3D12DynamicRHI::ResourceCast(TextureRHI);
+	Texture->Unlock(&RHICmdList, MipIndex, TextureIndex);
+}
+
 void FD3D12DynamicRHI::RHIUnlockTexture2DArray(FRHITexture2DArray* TextureRHI, uint32 TextureIndex, uint32 MipIndex, bool bLockWithinMiptail)
 {
 	check(TextureRHI);
