Register Offline BVH as an IB/VB rename listener: when IB/VB address changes we need to reconstruct HitGroup data otherwise we use incorrect/stale data in RT shaders.

diff --git a/Engine/Source/Runtime/D3D12RHI/Private/D3D12RayTracing.cpp b/Engine/Source/Runtime/D3D12RHI/Private/D3D12RayTracing.cpp
--- a/Engine/Source/Runtime/D3D12RHI/Private/D3D12RayTracing.cpp
+++ b/Engine/Source/Runtime/D3D12RHI/Private/D3D12RayTracing.cpp
@@ -2800,6 +2800,7 @@ struct FD3D12RHICommandInitializeRayTracingGeometry final : public FRHICommand<F
 
 			Device->GetDefaultCommandContext().ConditionalFlushCommandList();
 
+			Geometry->RegisterAsRenameListener(GPUIndex);
 			Geometry->SetupHitGroupSystemParameters(GPUIndex);
 			Geometry->SetDirty(FRHIGPUMask::FromIndex(GPUIndex), false);
 		}
