Fix duplicated code from merge.

diff --git a/Engine/Source/Runtime/D3D12RHI/Private/D3D12RayTracing.cpp b/Engine/Source/Runtime/D3D12RHI/Private/D3D12RayTracing.cpp
--- a/Engine/Source/Runtime/D3D12RHI/Private/D3D12RayTracing.cpp
+++ b/Engine/Source/Runtime/D3D12RHI/Private/D3D12RayTracing.cpp
@@ -3926,11 +3926,6 @@ void FD3D12RayTracingScene::BuildAccelerationStructure(FD3D12CommandContext& Com
 	checkf(ScratchAddress % GRHIRayTracingAccelerationStructureAlignment == 0,
 		TEXT("TLAS scratch buffer (plus offset) must be aligned to %lld bytes."),
 		GRHIRayTracingAccelerationStructureAlignment);
-		
-	checkf(ScratchBuffer, TEXT("TLAS build requires scratch buffer of at least %lld bytes."), PrebuildInfo.ScratchDataSizeInBytes);
-
-	const D3D12_GPU_VIRTUAL_ADDRESS ScratchAddress = ScratchBuffer->ResourceLocation.GetGPUVirtualAddress() + ScratchBufferOffset;
-
 
 	checkf(PrebuildInfo.ScratchDataSizeInBytes + ScratchBufferOffset <= ScratchBuffer->GetSize(),
 		TEXT("TLAS scratch buffer size is %lld bytes with offset %lld (%lld bytes available), but the build requires %lld bytes. ")
