D3D shader debug data changes
- Putting UniformBuffer name array key in a global header
- UniformBuffer name array shouldn't even been compiled in configurations that don't populate it
- Load shader UniformBuffer names in D3D12
- Adding error logging to D3D12 resource binding, uses shader UniformBuffer names. Implementation copied from D3D11.

diff --git a/Engine/Source/Developer/Windows/ShaderFormatD3D/Private/D3DShaderCompiler.inl b/Engine/Source/Developer/Windows/ShaderFormatD3D/Private/D3DShaderCompiler.inl
--- a/Engine/Source/Developer/Windows/ShaderFormatD3D/Private/D3DShaderCompiler.inl
+++ b/Engine/Source/Developer/Windows/ShaderFormatD3D/Private/D3DShaderCompiler.inl
@@ -613,7 +613,7 @@ inline void GenerateFinalOutput(TRefCountPtr<TBlob>& CompressedData,
 	// Append data that is generate from the shader code and assist the usage, mostly needed for DX12 
 	{
 		Output.ShaderCode.AddOptionalData(PackedResourceCounts);
-		Output.ShaderCode.AddOptionalData('u', UniformBufferNameBytes.GetData(), UniformBufferNameBytes.Num());
+		Output.ShaderCode.AddOptionalData(FShaderCodeUniformBuffers::Key, UniformBufferNameBytes.GetData(), UniformBufferNameBytes.Num());
 		AddOptionalDataCallback(Output.ShaderCode);
 	}
 
diff --git a/Engine/Source/Runtime/D3D12RHI/Private/D3D12Commands.cpp b/Engine/Source/Runtime/D3D12RHI/Private/D3D12Commands.cpp
--- a/Engine/Source/Runtime/D3D12RHI/Private/D3D12Commands.cpp
+++ b/Engine/Source/Runtime/D3D12RHI/Private/D3D12Commands.cpp
@@ -1577,7 +1577,7 @@ FORCEINLINE void SetResource(FD3D12CommandContext& CmdContext, uint32 BindIndex,
 }
 
 template <EShaderFrequency ShaderFrequency>
-inline int32 SetShaderResourcesFromBuffer_Surface(FD3D12CommandContext& CmdContext, FD3D12UniformBuffer* RESTRICT Buffer, const uint32 * RESTRICT ResourceMap, int32 BufferIndex)
+inline int32 SetShaderResourcesFromBuffer_Surface(FD3D12CommandContext& CmdContext, FD3D12UniformBuffer* RESTRICT Buffer, const uint32 * RESTRICT ResourceMap, int32 BufferIndex, const TCHAR* LayoutName)
 {
 #if ENABLE_RHI_VALIDATION
 	constexpr ERHIAccess SRVAccess = (ShaderFrequency == SF_Compute) ? ERHIAccess::SRVCompute : ERHIAccess::SRVGraphics;
@@ -1598,6 +1598,10 @@ inline int32 SetShaderResourcesFromBuffer_Surface(FD3D12CommandContext& CmdConte
 			const uint8 BindIndex = FRHIResourceTableEntry::GetBindIndex(ResourceInfo);
 
 			FRHITexture* TextureRHI = (FRHITexture*)Resources[ResourceIndex].GetReference();
+			if (!TextureRHI)
+			{
+				UE_LOG(LogD3D12RHI, Fatal, TEXT("Null texture (resource %d bind %d) on UB Layout %s"), ResourceIndex, BindIndex, LayoutName);
+			}
 			TextureRHI->SetLastRenderTime(CurrentTime);
 
 			FD3D12TextureBase* TextureD3D12 = CmdContext.RetrieveTextureBase(TextureRHI);
@@ -1625,7 +1629,7 @@ inline int32 SetShaderResourcesFromBuffer_Surface(FD3D12CommandContext& CmdConte
 }
 
 template <EShaderFrequency ShaderFrequency>
-inline int32 SetShaderResourcesFromBuffer_SRV(FD3D12CommandContext& CmdContext, FD3D12UniformBuffer* RESTRICT Buffer, const uint32 * RESTRICT ResourceMap, int32 BufferIndex)
+inline int32 SetShaderResourcesFromBuffer_SRV(FD3D12CommandContext& CmdContext, FD3D12UniformBuffer* RESTRICT Buffer, const uint32 * RESTRICT ResourceMap, int32 BufferIndex, const TCHAR* LayoutName)
 {
 #if ENABLE_RHI_VALIDATION
 	constexpr ERHIAccess SRVAccess = (ShaderFrequency == SF_Compute) ? ERHIAccess::SRVCompute : ERHIAccess::SRVGraphics;
@@ -1645,6 +1649,10 @@ inline int32 SetShaderResourcesFromBuffer_SRV(FD3D12CommandContext& CmdContext,
 			const uint8 BindIndex = FRHIResourceTableEntry::GetBindIndex(ResourceInfo);
 
 			FD3D12ShaderResourceView* D3D12Resource = CmdContext.RetrieveObject<FD3D12ShaderResourceView>((FRHIShaderResourceView*)(Resources[ResourceIndex].GetReference()));
+			if (!D3D12Resource)
+			{
+				UE_LOG(LogD3D12RHI, Fatal, TEXT("Null SRV (resource %d bind %d) on UB Layout %s"), ResourceIndex, BindIndex, LayoutName);
+			}
 
 #if ENABLE_RHI_VALIDATION
 			if (CmdContext.Tracker)
@@ -1706,13 +1714,66 @@ void FD3D12CommandContext::SetResourcesFromTables(const ShaderType* RESTRICT Sha
 		const int32 BufferIndex = FMath::FloorLog2(LowestBitMask); // todo: This has a branch on zero, we know it could never be zero...
 		DirtyBits ^= LowestBitMask;
 		FD3D12UniformBuffer* Buffer = BoundUniformBuffers[ShaderType::StaticFrequency][BufferIndex];
-		check(Buffer);
+
 		check(BufferIndex < Shader->ShaderResourceTable.ResourceTableLayoutHashes.Num());
-		check(Buffer->GetLayout().GetHash() == Shader->ShaderResourceTable.ResourceTableLayoutHashes[BufferIndex]);
+
+		if (!Buffer)
+		{
+			FString ShaderUB;
+#if RHI_INCLUDE_SHADER_DEBUG_DATA
+			if (BufferIndex < Shader->UniformBuffers.Num())
+			{
+				Shader->UniformBuffers[BufferIndex].ToString(ShaderUB);
+			}
+#endif
+			UE_LOG(LogD3D12RHI, Fatal, TEXT("Shader expected a uniform buffer at slot %u but got null instead (Shader='%s' UB='%s').  Rendering code needs to set a valid uniform buffer for this slot."), BufferIndex, Shader->GetShaderName(), *ShaderUB);
+		}
+
+#if RHI_INCLUDE_SHADER_DEBUG_DATA
+		// to track down OR-7159 CRASH: Client crashed at start of match in D3D11Commands.cpp
+		{
+			const uint32 LayoutHash = Buffer->GetLayout().GetHash();
+
+			if (LayoutHash != Shader->ShaderResourceTable.ResourceTableLayoutHashes[BufferIndex])
+			{
+				auto& BufferLayout = Buffer->GetLayout();
+				const auto& DebugName = BufferLayout.GetDebugName();
+				const FString& ShaderName = Shader->ShaderName;
+#if UE_BUILD_DEBUG
+				FString ShaderUB;
+				if (BufferIndex < Shader->UniformBuffers.Num())
+				{
+					ShaderUB = FString::Printf(TEXT("expecting UB '%s'"), *Shader->UniformBuffers[BufferIndex].GetPlainNameString());
+				}
+				UE_LOG(LogD3D12RHI, Error, TEXT("SetResourcesFromTables upcoming check(%08x != %08x); Bound Layout='%s' Shader='%s' %s"), BufferLayout.GetHash(), Shader->ShaderResourceTable.ResourceTableLayoutHashes[BufferIndex], *DebugName, *ShaderName, *ShaderUB);
+				FString ResourcesString;
+				for (int32 Index = 0; Index < BufferLayout.Resources.Num(); ++Index)
+				{
+					ResourcesString += FString::Printf(TEXT("%d "), BufferLayout.Resources[Index].MemberType);
+				}
+				UE_LOG(LogD3D12RHI, Error, TEXT("Layout CB Size %d %d Resources: %s"), BufferLayout.ConstantBufferSize, BufferLayout.Resources.Num(), *ResourcesString);
+#else
+				UE_LOG(LogD3D12RHI, Error, TEXT("Bound Layout='%s' Shader='%s', Layout CB Size %d %d"), *DebugName, *ShaderName, BufferLayout.ConstantBufferSize, BufferLayout.Resources.Num());
+#endif
+				// this might mean you are accessing a data you haven't bound e.g. GBuffer
+				checkf(BufferLayout.GetHash() == Shader->ShaderResourceTable.ResourceTableLayoutHashes[BufferIndex],
+					TEXT("Uniform buffer bound to slot %u is not what the shader expected:\n")
+					TEXT("\tBound:    Uniform Buffer[%s] with Hash[%u]\n")
+					TEXT("\tExpected: Uniform Buffer[%s] with Hash[%u]"),
+					BufferIndex, *DebugName, BufferLayout.GetHash(), *Shader->UniformBuffers[BufferIndex].GetPlainNameString(), Shader->ShaderResourceTable.ResourceTableLayoutHashes[BufferIndex]);
+			}
+		}
+#endif
+
+#if RHI_INCLUDE_SHADER_DEBUG_DATA
+		const TCHAR* LayoutName = *Buffer->GetLayout().GetDebugName();
+#else 
+		const TCHAR* LayoutName = nullptr;
+#endif
 
 		// todo: could make this two pass: gather then set
-		SetShaderResourcesFromBuffer_Surface<(EShaderFrequency)ShaderType::StaticFrequency>(*this, Buffer, Shader->ShaderResourceTable.TextureMap.GetData(), BufferIndex);
-		SetShaderResourcesFromBuffer_SRV<(EShaderFrequency)ShaderType::StaticFrequency>(*this, Buffer, Shader->ShaderResourceTable.ShaderResourceViewMap.GetData(), BufferIndex);
+		SetShaderResourcesFromBuffer_Surface<(EShaderFrequency)ShaderType::StaticFrequency>(*this, Buffer, Shader->ShaderResourceTable.TextureMap.GetData(), BufferIndex, LayoutName);
+		SetShaderResourcesFromBuffer_SRV<(EShaderFrequency)ShaderType::StaticFrequency>(*this, Buffer, Shader->ShaderResourceTable.ShaderResourceViewMap.GetData(), BufferIndex, LayoutName);
 		SetShaderResourcesFromBuffer_Sampler<(EShaderFrequency)ShaderType::StaticFrequency>(*this, Buffer, Shader->ShaderResourceTable.SamplerMap.GetData(), BufferIndex);
 	}
 
@@ -1721,7 +1782,7 @@ void FD3D12CommandContext::SetResourcesFromTables(const ShaderType* RESTRICT Sha
 
 
 template <EShaderFrequency ShaderFrequency>
-inline int32 SetShaderResourcesFromBuffer_UAVPS(FD3D12CommandContext& CmdContext, FD3D12UniformBuffer* RESTRICT Buffer, const uint32 * RESTRICT ResourceMap, int32 BufferIndex)
+inline int32 SetShaderResourcesFromBuffer_UAVPS(FD3D12CommandContext& CmdContext, FD3D12UniformBuffer* RESTRICT Buffer, const uint32 * RESTRICT ResourceMap, int32 BufferIndex, const TCHAR* LayoutName)
 {
 #if ENABLE_RHI_VALIDATION
 	constexpr ERHIAccess UAVAccess = (ShaderFrequency == SF_Compute) ? ERHIAccess::UAVCompute : ERHIAccess::UAVGraphics;
@@ -1741,6 +1802,10 @@ inline int32 SetShaderResourcesFromBuffer_UAVPS(FD3D12CommandContext& CmdContext
 			const uint8 BindIndex = FRHIResourceTableEntry::GetBindIndex(ResourceInfo);
 
 			FRHIUnorderedAccessView* RHIUAV = (FRHIUnorderedAccessView*)(Resources[ResourceIndex].GetReference());
+			if (!RHIUAV)
+			{
+				UE_LOG(LogD3D12RHI, Fatal, TEXT("Null UAV (resource %d bind %d) on UB Layout %s"), ResourceIndex, BindIndex, LayoutName);
+			}
 
 			FD3D12UnorderedAccessView* D3D12Resource = CmdContext.RetrieveObject<FD3D12UnorderedAccessView>(RHIUAV);
 
@@ -1778,13 +1843,32 @@ uint32 FD3D12CommandContext::SetUAVPSResourcesFromTables(const ShaderType* RESTR
 		const int32 BufferIndex = FMath::FloorLog2(LowestBitMask); // todo: This has a branch on zero, we know it could never be zero...
 		DirtyBits ^= LowestBitMask;
 		FD3D12UniformBuffer* Buffer = BoundUniformBuffers[ShaderType::StaticFrequency][BufferIndex];
-		check(Buffer);
+
 		check(BufferIndex < Shader->ShaderResourceTable.ResourceTableLayoutHashes.Num());
+
+		if (!Buffer)
+		{
+			FString ShaderUB;
+#if RHI_INCLUDE_SHADER_DEBUG_DATA
+			if (BufferIndex < Shader->UniformBuffers.Num())
+			{
+				Shader->UniformBuffers[BufferIndex].ToString(ShaderUB);
+			}
+#endif
+			UE_LOG(LogD3D12RHI, Fatal, TEXT("Shader expected a uniform buffer at slot %u but got null instead (Shader='%s' UB='%s').  Rendering code needs to set a valid uniform buffer for this slot."), BufferIndex, Shader->GetShaderName(), *ShaderUB);
+		}
+
 		check(Buffer->GetLayout().GetHash() == Shader->ShaderResourceTable.ResourceTableLayoutHashes[BufferIndex]);
 
 		if ((EShaderFrequency)ShaderType::StaticFrequency == SF_Pixel)
 		{
-			NumChanged += SetShaderResourcesFromBuffer_UAVPS<(EShaderFrequency)ShaderType::StaticFrequency>(*this, Buffer, Shader->ShaderResourceTable.UnorderedAccessViewMap.GetData(), BufferIndex);
+#if RHI_INCLUDE_SHADER_DEBUG_DATA
+			const TCHAR* LayoutName = *Buffer->GetLayout().GetDebugName();
+#else
+			const TCHAR* LayoutName = nullptr;
+#endif
+
+			NumChanged += SetShaderResourcesFromBuffer_UAVPS<(EShaderFrequency)ShaderType::StaticFrequency>(*this, Buffer, Shader->ShaderResourceTable.UnorderedAccessViewMap.GetData(), BufferIndex, LayoutName);
 		}
 	}
 	return NumChanged;
diff --git a/Engine/Source/Runtime/D3D12RHI/Private/D3D12Shader.h b/Engine/Source/Runtime/D3D12RHI/Private/D3D12Shader.h
--- a/Engine/Source/Runtime/D3D12RHI/Private/D3D12Shader.h
+++ b/Engine/Source/Runtime/D3D12RHI/Private/D3D12Shader.h
@@ -43,6 +43,10 @@ struct FD3D12ShaderData
 
 	FD3D12ShaderResourceTable ShaderResourceTable;
 
+#if RHI_INCLUDE_SHADER_DEBUG_DATA
+	TArray<FName> UniformBuffers;
+#endif
+
 	FShaderCodePackedResourceCounts ResourceCounts{};
 
 	/** The static slot associated with the resource table index in ShaderResourceTable. */
diff --git a/Engine/Source/Runtime/D3D12RHI/Private/D3D12Shaders.cpp b/Engine/Source/Runtime/D3D12RHI/Private/D3D12Shaders.cpp
--- a/Engine/Source/Runtime/D3D12RHI/Private/D3D12Shaders.cpp
+++ b/Engine/Source/Runtime/D3D12RHI/Private/D3D12Shaders.cpp
@@ -16,13 +16,11 @@ static inline bool ReadShaderOptionalData(FShaderCodeReader& InShaderCode, TShad
 	}
 	OutShader.ResourceCounts = *PackedResourceCounts;
 
-#if !(UE_BUILD_SHIPPING || UE_BUILD_TEST)
+#if RHI_INCLUDE_SHADER_DEBUG_DATA
 	OutShader.ShaderName = InShaderCode.FindOptionalData(FShaderCodeName::Key);
 
-#if 0
 	int32 UniformBufferTableSize = 0;
-	const auto* UniformBufferData = InShaderCode.FindOptionalDataAndSize('u', UniformBufferTableSize);
-	//#todo-rco
+	const uint8* UniformBufferData = InShaderCode.FindOptionalDataAndSize(FShaderCodeUniformBuffers::Key, UniformBufferTableSize);
 	if (UniformBufferData && UniformBufferTableSize > 0)
 	{
 		FBufferReader UBReader((void*)UniformBufferData, UniformBufferTableSize, false);
@@ -35,7 +33,6 @@ static inline bool ReadShaderOptionalData(FShaderCodeReader& InShaderCode, TShad
 		}
 	}
 #endif
-#endif
 
 #if D3D12RHI_NEEDS_VENDOR_EXTENSIONS
 	int32 VendorExtensionTableSize = 0;
diff --git a/Engine/Source/Runtime/RHI/Public/RHIResources.h b/Engine/Source/Runtime/RHI/Public/RHIResources.h
--- a/Engine/Source/Runtime/RHI/Public/RHIResources.h
+++ b/Engine/Source/Runtime/RHI/Public/RHIResources.h
@@ -659,13 +659,19 @@ public:
 // Shaders
 //
 
+#if !(UE_BUILD_SHIPPING || UE_BUILD_TEST)
+	#define RHI_INCLUDE_SHADER_DEBUG_DATA 1
+#else
+	#define RHI_INCLUDE_SHADER_DEBUG_DATA 0
+#endif
+
 class FRHIShader : public FRHIResource
 {
 public:
 	void SetHash(FSHAHash InHash) { Hash = InHash; }
 	FSHAHash GetHash() const { return Hash; }
 
-#if !(UE_BUILD_SHIPPING || UE_BUILD_TEST)
+#if RHI_INCLUDE_SHADER_DEBUG_DATA
 	// for debugging only e.g. MaterialName:ShaderFile.usf or ShaderFile.usf/EntryFunc
 	FString ShaderName;
 	FORCEINLINE const TCHAR* GetShaderName() const { return *ShaderName; }
diff --git a/Engine/Source/Runtime/RenderCore/Public/ShaderCore.h b/Engine/Source/Runtime/RenderCore/Public/ShaderCore.h
--- a/Engine/Source/Runtime/RenderCore/Public/ShaderCore.h
+++ b/Engine/Source/Runtime/RenderCore/Public/ShaderCore.h
@@ -630,6 +630,12 @@ struct FShaderCodeName
 	// We store the straight ANSICHAR zero-terminated string
 };
 
+struct FShaderCodeUniformBuffers
+{
+	static const uint8 Key = 'u';
+	// We store an array of FString objects
+};
+
 // if this changes you need to make sure all shaders get invalidated
 struct FShaderCodeVendorExtension
 {
diff --git a/Engine/Source/Runtime/Windows/D3D11RHI/Private/D3D11Commands.cpp b/Engine/Source/Runtime/Windows/D3D11RHI/Private/D3D11Commands.cpp
--- a/Engine/Source/Runtime/Windows/D3D11RHI/Private/D3D11Commands.cpp
+++ b/Engine/Source/Runtime/Windows/D3D11RHI/Private/D3D11Commands.cpp
@@ -1359,14 +1359,16 @@ void FD3D11DynamicRHI::SetResourcesFromTables(const ShaderType* RESTRICT Shader)
 		if (!Buffer)
 		{
 			FString ShaderUB;
+#if RHI_INCLUDE_SHADER_DEBUG_DATA
 			if (BufferIndex < Shader->UniformBuffers.Num())
 			{
 				Shader->UniformBuffers[BufferIndex].ToString(ShaderUB);
 			}
+#endif
 			UE_LOG(LogD3D11RHI, Fatal, TEXT("Shader expected a uniform buffer at slot %u but got null instead (Shader='%s' UB='%s').  Rendering code needs to set a valid uniform buffer for this slot."), BufferIndex, Shader->GetShaderName(), *ShaderUB);
 		}
 
-#if !(UE_BUILD_SHIPPING || UE_BUILD_TEST)
+#if RHI_INCLUDE_SHADER_DEBUG_DATA
 		// to track down OR-7159 CRASH: Client crashed at start of match in D3D11Commands.cpp
 		{
 			const uint32 LayoutHash = Buffer->GetLayout().GetHash();
@@ -1402,7 +1404,7 @@ void FD3D11DynamicRHI::SetResourcesFromTables(const ShaderType* RESTRICT Shader)
 		}
 #endif
 
-#if !(UE_BUILD_SHIPPING || UE_BUILD_TEST)
+#if RHI_INCLUDE_SHADER_DEBUG_DATA
 		const TCHAR* LayoutName = *Buffer->GetLayout().GetDebugName();
 #else 
 		const TCHAR* LayoutName = nullptr;
@@ -1437,21 +1439,23 @@ int32 FD3D11DynamicRHI::SetUAVPSResourcesFromTables(const ShaderType* RESTRICT S
 		if (!Buffer)
 		{
 			FString ShaderUB;
+#if RHI_INCLUDE_SHADER_DEBUG_DATA
 			if (BufferIndex < Shader->UniformBuffers.Num())
 			{
 				Shader->UniformBuffers[BufferIndex].ToString(ShaderUB);
 			}
+#endif
 			UE_LOG(LogD3D11RHI, Fatal, TEXT("Shader expected a uniform buffer at slot %u but got null instead (Shader='%s' UB='%s').  Rendering code needs to set a valid uniform buffer for this slot."), BufferIndex, Shader->GetShaderName(), *ShaderUB);
 		}
 
-	#if !(UE_BUILD_SHIPPING || UE_BUILD_TEST)
-		const TCHAR* LayoutName = *Buffer->GetLayout().GetDebugName();
-	#else
-		const TCHAR* LayoutName = nullptr;
-	#endif
-
 		if ((EShaderFrequency)ShaderType::StaticFrequency == SF_Pixel)
 		{
+#if RHI_INCLUDE_SHADER_DEBUG_DATA
+			const TCHAR* LayoutName = *Buffer->GetLayout().GetDebugName();
+#else
+			const TCHAR* LayoutName = nullptr;
+#endif
+
 			NumChanged += SetShaderResourcesFromBufferUAVPS<(EShaderFrequency)ShaderType::StaticFrequency>(this, &StateCache, Buffer, Shader->ShaderResourceTable.UnorderedAccessViewMap.GetData(), BufferIndex, LayoutName);
 		}
 	}
diff --git a/Engine/Source/Runtime/Windows/D3D11RHI/Private/D3D11Shaders.cpp b/Engine/Source/Runtime/Windows/D3D11RHI/Private/D3D11Shaders.cpp
--- a/Engine/Source/Runtime/Windows/D3D11RHI/Private/D3D11Shaders.cpp
+++ b/Engine/Source/Runtime/Windows/D3D11RHI/Private/D3D11Shaders.cpp
@@ -24,10 +24,11 @@ static inline void ReadShaderOptionalData(FShaderCodeReader& InShaderCode, TShad
 	}
 	
 	OutShader.bShaderNeedsGlobalConstantBuffer = PackedResourceCounts->bGlobalUniformBufferUsed;
-#if !(UE_BUILD_SHIPPING || UE_BUILD_TEST)
+#if RHI_INCLUDE_SHADER_DEBUG_DATA
 	OutShader.ShaderName = InShaderCode.FindOptionalData(FShaderCodeName::Key);
+
 	int32 UniformBufferTableSize = 0;
-	auto* UniformBufferData = InShaderCode.FindOptionalDataAndSize('u', UniformBufferTableSize);
+	const uint8* UniformBufferData = InShaderCode.FindOptionalDataAndSize(FShaderCodeUniformBuffers::Key, UniformBufferTableSize);
 	if (UniformBufferData && UniformBufferTableSize > 0)
 	{
 		FBufferReader UBReader((void*)UniformBufferData, UniformBufferTableSize, false);
diff --git a/Engine/Source/Runtime/Windows/D3D11RHI/Public/D3D11Resources.h b/Engine/Source/Runtime/Windows/D3D11RHI/Public/D3D11Resources.h
--- a/Engine/Source/Runtime/Windows/D3D11RHI/Public/D3D11Resources.h
+++ b/Engine/Source/Runtime/Windows/D3D11RHI/Public/D3D11Resources.h
@@ -41,7 +41,9 @@ public:
 struct FD3D11ShaderData
 {
 	FD3D11ShaderResourceTable			ShaderResourceTable;
+#if RHI_INCLUDE_SHADER_DEBUG_DATA
 	TArray<FName>						UniformBuffers;
+#endif
 	TArray<FUniformBufferStaticSlot>	StaticSlots;
 	TArray<FShaderCodeVendorExtension>	VendorExtensions;
 	bool								bShaderNeedsGlobalConstantBuffer;
