Remove possible dangling pointer in vertex cache of default command context when buffer is destroyed - can only happen in default command list because that's the only command list which in practice perform flush & restore of state with a clear state in between

diff --git a/Engine/Source/Runtime/D3D12RHI/Private/D3D12Buffer.cpp b/Engine/Source/Runtime/D3D12RHI/Private/D3D12Buffer.cpp
--- a/Engine/Source/Runtime/D3D12RHI/Private/D3D12Buffer.cpp
+++ b/Engine/Source/Runtime/D3D12RHI/Private/D3D12Buffer.cpp
@@ -7,6 +7,22 @@ D3D12Buffer.cpp: D3D Common code for buffers.
 #include "D3D12RHIPrivate.h"
 #include "D3D12RHIBridge.h"
 
+FD3D12Buffer::~FD3D12Buffer()
+{
+	if (EnumHasAnyFlags(GetUsage(), EBufferUsageFlags::VertexBuffer) && GetParentDevice())
+	{
+		FD3D12CommandContext& DefaultContext = GetParentDevice()->GetDefaultCommandContext();
+		DefaultContext.StateCache.ClearVertexBuffer(&ResourceLocation);
+	}
+
+	int64 BufferSize = ResourceLocation.GetSize();
+	bool bTransient = ResourceLocation.IsTransient();
+	if (!bTransient)
+	{
+		UpdateBufferStats((EBufferUsageFlags)GetUsage(), -BufferSize);
+	}
+}
+
 struct FRHICommandUpdateBufferString
 {
 	static const TCHAR* TStr() { return TEXT("FRHICommandUpdateBuffer"); }
diff --git a/Engine/Source/Runtime/D3D12RHI/Private/D3D12StateCachePrivate.h b/Engine/Source/Runtime/D3D12RHI/Private/D3D12StateCachePrivate.h
--- a/Engine/Source/Runtime/D3D12RHI/Private/D3D12StateCachePrivate.h
+++ b/Engine/Source/Runtime/D3D12RHI/Private/D3D12StateCachePrivate.h
@@ -872,17 +872,15 @@ public:
 		return false;
 	}
 
-	D3D12_STATE_CACHE_INLINE bool IsStreamSource(const FD3D12ResourceLocation* VertexBufferLocation) const
+	D3D12_STATE_CACHE_INLINE void ClearVertexBuffer(const FD3D12ResourceLocation* VertexBufferLocation)
 	{
 		for (int32 index = 0; index <= PipelineState.Graphics.VBCache.MaxBoundVertexBufferIndex; ++index)
 		{
 			if (PipelineState.Graphics.VBCache.CurrentVertexBufferResources[index] == VertexBufferLocation)
 			{
-				return true;
+				PipelineState.Graphics.VBCache.CurrentVertexBufferResources[index] = nullptr;
 			}
 		}
-
-		return false;
 	}
 
 public:
diff --git a/Engine/Source/Runtime/D3D12RHI/Public/D3D12Resources.h b/Engine/Source/Runtime/D3D12RHI/Public/D3D12Resources.h
--- a/Engine/Source/Runtime/D3D12RHI/Public/D3D12Resources.h
+++ b/Engine/Source/Runtime/D3D12RHI/Public/D3D12Resources.h
@@ -949,15 +949,7 @@ public:
 		, LockedData(InParent)
 	{
 	}
-	virtual ~FD3D12Buffer()
-	{
-		int64 BufferSize = ResourceLocation.GetSize();
-		bool bTransient = ResourceLocation.IsTransient();
-		if (!bTransient)
-		{
-			UpdateBufferStats((EBufferUsageFlags)GetUsage(), -BufferSize);
-		}
-	}
+	virtual ~FD3D12Buffer();
 
 	virtual uint32 GetParentGPUIndex() const override;
 
