When creating contexts, take the reserve worker threads into account.

diff --git a/Engine/Source/Runtime/D3D12RHI/Private/D3D12Device.cpp b/Engine/Source/Runtime/D3D12RHI/Private/D3D12Device.cpp
--- a/Engine/Source/Runtime/D3D12RHI/Private/D3D12Device.cpp
+++ b/Engine/Source/Runtime/D3D12RHI/Private/D3D12Device.cpp
@@ -104,7 +104,18 @@ void FD3D12Device::CreateCommandContexts()
 	check(CommandContextArray.Num() == 0);
 	check(AsyncComputeContextArray.Num() == 0);
 
-	const uint32 NumContexts = FTaskGraphInterface::Get().GetNumWorkerThreads() + 1;
+	uint32 WorkerThreadCount = FTaskGraphInterface::Get().GetNumWorkerThreads();
+
+#if PLATFORM_WINDOWS
+	bool bEnableReserveWorkers = true; // by default
+	GConfig->GetBool(TEXT("TaskGraph"), TEXT("EnableReserveWorkers"), bEnableReserveWorkers, GEngineIni);
+	if (bEnableReserveWorkers)
+	{
+		WorkerThreadCount *= 2;
+	}
+#endif
+
+	const uint32 NumContexts = WorkerThreadCount + 1;
 	const uint32 NumAsyncComputeContexts = GEnableAsyncCompute ? 1 : 0;
 	
 	// We never make the default context free for allocation by the context containers
