Allow D3D12 and Vulkan to preload only shaders needed by PSOs.

- Reduces Shaders LLM greatly, at the expense of possible hitches if the cache isn't complete.

diff --git a/Engine/Source/Runtime/D3D12RHI/Private/D3D12RHI.cpp b/Engine/Source/Runtime/D3D12RHI/Private/D3D12RHI.cpp
--- a/Engine/Source/Runtime/D3D12RHI/Private/D3D12RHI.cpp
+++ b/Engine/Source/Runtime/D3D12RHI/Private/D3D12RHI.cpp
@@ -83,6 +83,10 @@ FD3D12DynamicRHI::FD3D12DynamicRHI(const TArray<TSharedPtr<FD3D12Adapter>>& Chos
 	GRHIMultiPipelineMergeableAccessMask = GRHIMergeableAccessMask;
 	EnumRemoveFlags(GRHIMultiPipelineMergeableAccessMask, ERHIAccess::UAVMask);
 
+	// This flag allows to skip preloading the shaders when loading a material and rely on the PSO cache to load them when precompiling the PSOs in background.
+	// If PSO cache isn't enabled or present, this flag has no effect.
+	GRHISupportsLazyShaderCodeLoading = true;
+
 	GPoolSizeVRAMPercentage = 0;
 	GTexturePoolSize = 0;
 	GConfig->GetInt(TEXT("TextureStreaming"), TEXT("PoolSizeVRAMPercentage"), GPoolSizeVRAMPercentage, GEngineIni);
diff --git a/Engine/Source/Runtime/RHI/Private/PipelineFileCache.cpp b/Engine/Source/Runtime/RHI/Private/PipelineFileCache.cpp
--- a/Engine/Source/Runtime/RHI/Private/PipelineFileCache.cpp
+++ b/Engine/Source/Runtime/RHI/Private/PipelineFileCache.cpp
@@ -140,7 +140,7 @@ static TAutoConsoleVariable<int32> CVarPSOFileCacheUserCacheUnusedElementCheckPe
 
 static TAutoConsoleVariable<int32> CVarLazyLoadShadersWhenPSOCacheIsPresent(
 	TEXT("r.ShaderPipelineCache.LazyLoadShadersWhenPSOCacheIsPresent"),
-	0,
+	1,
 	TEXT("Non-Zero: If we load a PSO cache, then lazy load from the shader code library. This assumes the PSO cache is more or less complete. This will only work on RHIs that support the library+Hash CreateShader API (GRHISupportsLazyShaderCodeLoading == true)."),
 	ECVF_RenderThreadSafe);
 
diff --git a/Engine/Source/Runtime/VulkanRHI/Private/VulkanRHI.cpp b/Engine/Source/Runtime/VulkanRHI/Private/VulkanRHI.cpp
--- a/Engine/Source/Runtime/VulkanRHI/Private/VulkanRHI.cpp
+++ b/Engine/Source/Runtime/VulkanRHI/Private/VulkanRHI.cpp
@@ -232,6 +232,10 @@ FVulkanDynamicRHI::FVulkanDynamicRHI()
 	GRHITransitionPrivateData_AlignInBytes = alignof(FVulkanPipelineBarrier);
 	GConfig->GetInt(TEXT("TextureStreaming"), TEXT("PoolSizeVRAMPercentage"), GPoolSizeVRAMPercentage, GEngineIni);
 
+	// This flag allows to skip preloading the shaders when loading a material and rely on the PSO cache to load them when precompiling the PSOs in background.
+	// If PSO cache isn't enabled or present, this flag has no effect.
+	GRHISupportsLazyShaderCodeLoading = true;
+
 	// Copy source requires its own image layout.
 	EnumRemoveFlags(GRHIMergeableAccessMask, ERHIAccess::CopySrc);
 }
