Adding support for attaching to PIX on startup. You can either pass -attachPIX on the command line or set r.D3D12.AutoAttachPIX=1 in ConsoleVariables.ini
The user still needs to manually hit "attach" inside PIX but it's a start.

Based on PR #7961 from:

diff --git a/Engine/Source/Runtime/D3D12RHI/Private/D3D12RHI.cpp b/Engine/Source/Runtime/D3D12RHI/Private/D3D12RHI.cpp
--- a/Engine/Source/Runtime/D3D12RHI/Private/D3D12RHI.cpp
+++ b/Engine/Source/Runtime/D3D12RHI/Private/D3D12RHI.cpp
@@ -286,6 +286,14 @@ void FD3D12DynamicRHI::Shutdown()
 	FMemory::Free(ZeroBuffer);
 	ZeroBuffer = NULL;
 	ZeroBufferSize = 0;
+
+#if D3D12RHI_SUPPORTS_WIN_PIX
+	if (WinPixGpuCapturerHandle)
+	{
+		FPlatformProcess::FreeDllHandle(WinPixGpuCapturerHandle);
+		WinPixGpuCapturerHandle = nullptr;
+	}
+#endif
 }
 
 FD3D12CommandContext* FD3D12DynamicRHI::CreateCommandContext(FD3D12Device* InParent, ED3D12CommandQueueType InQueueType, bool InIsDefaultContext)
diff --git a/Engine/Source/Runtime/D3D12RHI/Private/D3D12RHIPrivate.h b/Engine/Source/Runtime/D3D12RHI/Private/D3D12RHIPrivate.h
--- a/Engine/Source/Runtime/D3D12RHI/Private/D3D12RHIPrivate.h
+++ b/Engine/Source/Runtime/D3D12RHI/Private/D3D12RHIPrivate.h
@@ -1075,6 +1075,10 @@ protected:
 
 	TArray<TSharedPtr<FD3D12Adapter>> ChosenAdapters;
 
+#if D3D12RHI_SUPPORTS_WIN_PIX
+	void* WinPixGpuCapturerHandle = nullptr;
+#endif
+
 	/** Can pix events be used */
 	bool bPixEventEnabled = false;
 
@@ -1180,9 +1184,10 @@ public:
 
 private:
 
-#if USE_PIX && (PLATFORM_WINDOWS || PLATFORM_HOLOLENS)
+#if D3D12RHI_SUPPORTS_WIN_PIX
 	void* WindowsPixDllHandle = nullptr;
-#endif // USE_PIX && (PLATFORM_WINDOWS || PLATFORM_HOLOLENS)
+	void* WinPixGpuCapturerHandle = nullptr;
+#endif
 
 	TArray<TSharedPtr<FD3D12Adapter>> ChosenAdapters;
 
diff --git a/Engine/Source/Runtime/D3D12RHI/Private/Windows/WindowsD3D12Device.cpp b/Engine/Source/Runtime/D3D12RHI/Private/Windows/WindowsD3D12Device.cpp
--- a/Engine/Source/Runtime/D3D12RHI/Private/Windows/WindowsD3D12Device.cpp
+++ b/Engine/Source/Runtime/D3D12RHI/Private/Windows/WindowsD3D12Device.cpp
@@ -150,6 +150,16 @@ static TAutoConsoleVariable<int32> CVarExperimentalShaderModels(
 );
 #endif // !UE_BUILD_SHIPPING
 
+#if D3D12RHI_SUPPORTS_WIN_PIX
+int32 GAutoAttachPIX = 0;
+static FAutoConsoleVariableRef CVarAutoAttachPIX(
+	TEXT("r.D3D12.AutoAttachPIX"),
+	GAutoAttachPIX,
+	TEXT("Automatically attach PIX on startup"),
+	ECVF_ReadOnly | ECVF_RenderThreadSafe
+);
+#endif // D3D12RHI_SUPPORTS_WIN_PIX
+
 static inline int D3D12RHI_PreferAdapterVendor()
 {
 	if (FParse::Param(FCommandLine::Get(), TEXT("preferAMD")))
@@ -711,12 +721,11 @@ FDynamicRHI* FD3D12DynamicRHIModule::CreateRHI(ERHIFeatureLevel::Type RequestedF
 	GMaxRHIShaderPlatform = GShaderPlatformForFeatureLevel[GMaxRHIFeatureLevel];
 	check(GMaxRHIShaderPlatform != SP_NumPlatforms);
 
-#if USE_PIX
+#if D3D12RHI_SUPPORTS_WIN_PIX
 	bool bPixEventEnabled = (WindowsPixDllHandle != nullptr);
 #else
 	bool bPixEventEnabled = false;
 #endif // USE_PIX
-
 	
 	if (ChosenAdapters.Num() > 0 && ChosenAdapters[0].IsValid())
 	{
@@ -757,28 +766,29 @@ void FD3D12DynamicRHIModule::StartupModule()
 	}
 #endif
 
-#if USE_PIX
+#if D3D12RHI_SUPPORTS_WIN_PIX
 #if PLATFORM_CPU_ARM_FAMILY
-	static FString WindowsPixDllRelativePath = FPaths::Combine(*FPaths::EngineDir(), TEXT("Binaries/ThirdParty/Windows/WinPixEventRuntime/arm64"));
+	static FString WindowsPixDllRelativePath = FPaths::Combine(FPaths::EngineDir(), TEXT("Binaries/ThirdParty/Windows/WinPixEventRuntime/arm64"));
 #else
-	static FString WindowsPixDllRelativePath = FPaths::Combine(*FPaths::EngineDir(), TEXT("Binaries/ThirdParty/Windows/WinPixEventRuntime/x64"));
+	static FString WindowsPixDllRelativePath = FPaths::Combine(FPaths::EngineDir(), TEXT("Binaries/ThirdParty/Windows/WinPixEventRuntime/x64"));
 #endif
-	static FString WindowsPixDll("WinPixEventRuntime.dll");
-	UE_LOG(LogD3D12RHI, Log, TEXT("Loading %s for PIX profiling (from %s)."), WindowsPixDll.GetCharArray().GetData(), WindowsPixDllRelativePath.GetCharArray().GetData());
-	WindowsPixDllHandle = FPlatformProcess::GetDllHandle(*FPaths::Combine(*WindowsPixDllRelativePath, *WindowsPixDll));
+	static const TCHAR* WindowsPixDll = TEXT("WinPixEventRuntime.dll");
+
+	UE_LOG(LogD3D12RHI, Log, TEXT("Loading %s for PIX profiling (from %s)."), WindowsPixDll, *WindowsPixDllRelativePath);
+	WindowsPixDllHandle = FPlatformProcess::GetDllHandle(*FPaths::Combine(WindowsPixDllRelativePath, WindowsPixDll));
 	if (WindowsPixDllHandle == nullptr)
 	{
 		const int32 ErrorNum = FPlatformMisc::GetLastError();
 		TCHAR ErrorMsg[1024];
 		FPlatformMisc::GetSystemErrorMessage(ErrorMsg, 1024, ErrorNum);
-		UE_LOG(LogD3D12RHI, Error, TEXT("Failed to get %s handle: %s (%d)"), WindowsPixDll.GetCharArray().GetData(), ErrorMsg, ErrorNum);
+		UE_LOG(LogD3D12RHI, Error, TEXT("Failed to get %s handle: %s (%d)"), WindowsPixDll, ErrorMsg, ErrorNum);
 	}
 #endif
 }
 
 void FD3D12DynamicRHIModule::ShutdownModule()
 {
-#if USE_PIX
+#if D3D12RHI_SUPPORTS_WIN_PIX
 	if (WindowsPixDllHandle)
 	{
 		FPlatformProcess::FreeDllHandle(WindowsPixDllHandle);
@@ -835,6 +845,20 @@ static void DisableRayTracingSupport()
 
 void FD3D12DynamicRHI::Init()
 {
+#if D3D12RHI_SUPPORTS_WIN_PIX
+	// PIX GPU capture dll: makes PIX be able to attach to our process. !GetModuleHandle() is required because PIX may already be attached.
+	if (GAutoAttachPIX || FParse::Param(FCommandLine::Get(), TEXT("attachPIX")))
+	{
+		// If PIX is not already attached, load its dll to auto attach ourselves
+		if (!::GetModuleHandle(L"WinPixGpuCapturer.dll"))
+		{
+			// This should always be loaded from the installed PIX directory.
+			// This function does assume it's installed under Program Files so we may have to revisit for custom install locations.
+			WinPixGpuCapturerHandle = PIXLoadLatestWinPixGpuCapturerLibrary();
+		}
+	}
+#endif
+
 	for (TSharedPtr<FD3D12Adapter>& Adapter : ChosenAdapters)
 	{
 		Adapter->Initialize(this);
diff --git a/Engine/Source/Runtime/D3D12RHI/Public/D3D12RHI.h b/Engine/Source/Runtime/D3D12RHI/Public/D3D12RHI.h
--- a/Engine/Source/Runtime/D3D12RHI/Public/D3D12RHI.h
+++ b/Engine/Source/Runtime/D3D12RHI/Public/D3D12RHI.h
@@ -59,6 +59,7 @@
 	#define ASYNC_DEFERRED_DELETION					1
 	#define PIPELINE_STATE_FILE_LOCATION			FPaths::ProjectSavedDir()
 	#define USE_PIX									WITH_PIX_EVENT_RUNTIME
+	#define D3D12RHI_SUPPORTS_WIN_PIX				USE_PIX
 #else
 	#include "D3D12RHIPlatformPublic.h"
 #endif
