Fix build-break on HoloLens with D3D12/PIX path calling GetModuleHandle.

diff --git a/Engine/Source/Runtime/D3D12RHI/Private/Windows/WindowsD3D12Device.cpp b/Engine/Source/Runtime/D3D12RHI/Private/Windows/WindowsD3D12Device.cpp
--- a/Engine/Source/Runtime/D3D12RHI/Private/Windows/WindowsD3D12Device.cpp
+++ b/Engine/Source/Runtime/D3D12RHI/Private/Windows/WindowsD3D12Device.cpp
@@ -850,7 +850,7 @@ void FD3D12DynamicRHI::Init()
 	if (GAutoAttachPIX || FParse::Param(FCommandLine::Get(), TEXT("attachPIX")))
 	{
 		// If PIX is not already attached, load its dll to auto attach ourselves
-		if (!::GetModuleHandle(L"WinPixGpuCapturer.dll"))
+		if (!FPlatformProcess::GetDllHandle(L"WinPixGpuCapturer.dll"))
 		{
 			// This should always be loaded from the installed PIX directory.
 			// This function does assume it's installed under Program Files so we may have to revisit for custom install locations.
