Don't update the memory info on d3d12 error reporting - the call might fail on lost device under certain circumstances and using the last cached memory stats should be good enough as in indication of the memory load

diff --git a/Engine/Source/Runtime/D3D12RHI/Private/D3D12Util.cpp b/Engine/Source/Runtime/D3D12RHI/Private/D3D12Util.cpp
--- a/Engine/Source/Runtime/D3D12RHI/Private/D3D12Util.cpp
+++ b/Engine/Source/Runtime/D3D12RHI/Private/D3D12Util.cpp
@@ -668,8 +668,7 @@ void LogPageFaultData(FD3D12Adapter* InAdapter, D3D12_GPU_VIRTUAL_ADDRESS InPage
 } // namespace D3D12RHI
 
 void LogMemoryInfo(FD3D12Adapter* InAdapter)
-{
-	InAdapter->UpdateMemoryInfo();
+{	
 	const FD3D12MemoryInfo& MemoryInfo = InAdapter->GetMemoryInfo();
 
 	UE_LOG(LogD3D12RHI, Error, TEXT("Memory Info from frame ID %d:"), MemoryInfo.UpdateFrameNumber);
