Only reset the query flag when it's on occlussion query - only occlusion queries call RHIBeginRenderQuery (found via UDN)

diff --git a/Engine/Source/Runtime/D3D12RHI/Private/D3D12Commands.cpp b/Engine/Source/Runtime/D3D12RHI/Private/D3D12Commands.cpp
--- a/Engine/Source/Runtime/D3D12RHI/Private/D3D12Commands.cpp
+++ b/Engine/Source/Runtime/D3D12RHI/Private/D3D12Commands.cpp
@@ -1495,7 +1495,11 @@ void FD3D12CommandContext::RHIEndRenderQuery(FRHIRenderQuery* QueryRHI)
 	// Query data isn't ready until it has been resolved.
 	ensure(Query->bResultIsCached == false && Query->bResolved == false);
 
-	bIsDoingQuery = false;
+	if (Query->Type == RQT_Occlusion)
+	{
+		check(bIsDoingQuery);
+		bIsDoingQuery = false;
+	}
 }
 
 void FD3D12CommandContext::RHICalibrateTimers(FRHITimestampCalibrationQuery* CalibrationQuery)
