Workaround to make sure destruction of back buffers and viewport window association is not happening at the same time - known DXGI issue which MS will if in newer DXGI versions but these are not part of the AgilitySDK yet so need to workaround in code

diff --git a/Engine/Source/Runtime/D3D12RHI/Private/D3D12Resources.cpp b/Engine/Source/Runtime/D3D12RHI/Private/D3D12Resources.cpp
--- a/Engine/Source/Runtime/D3D12RHI/Private/D3D12Resources.cpp
+++ b/Engine/Source/Runtime/D3D12RHI/Private/D3D12Resources.cpp
@@ -313,6 +313,14 @@ FD3D12Resource::~FD3D12Resource()
 		GFSDK_Aftermath_DX12_UnregisterResource(AftermathHandle);
 	}
 #endif
+
+	if (bBackBuffer)
+	{
+		// Don't make the windows association call and release back buffer at the same time (see notes on critical section)
+		FScopeLock Lock(&FD3D12Viewport::DXGIBackBufferLock);
+		bBackBuffer = false;
+		Resource.SafeRelease();
+	}
 }
 
 ID3D12Pageable* FD3D12Resource::GetPageable()
diff --git a/Engine/Source/Runtime/D3D12RHI/Private/D3D12Viewport.cpp b/Engine/Source/Runtime/D3D12RHI/Private/D3D12Viewport.cpp
--- a/Engine/Source/Runtime/D3D12RHI/Private/D3D12Viewport.cpp
+++ b/Engine/Source/Runtime/D3D12RHI/Private/D3D12Viewport.cpp
@@ -106,6 +106,8 @@ namespace D3D12RHI
 }
 using namespace D3D12RHI;
 
+FCriticalSection FD3D12Viewport::DXGIBackBufferLock;
+
 #if WITH_MGPU
 FD3D12FramePacing::FD3D12FramePacing(FD3D12Adapter* Parent)
 	: FD3D12AdapterChild(Parent)
diff --git a/Engine/Source/Runtime/D3D12RHI/Private/Windows/WindowsD3D12Viewport.cpp b/Engine/Source/Runtime/D3D12RHI/Private/Windows/WindowsD3D12Viewport.cpp
--- a/Engine/Source/Runtime/D3D12RHI/Private/Windows/WindowsD3D12Viewport.cpp
+++ b/Engine/Source/Runtime/D3D12RHI/Private/Windows/WindowsD3D12Viewport.cpp
@@ -128,7 +128,7 @@ void FD3D12Viewport::Init()
 
 		VERIFYD3D12RESULT(Factory4->CreateSwapChainForCoreWindow(
 			pCommandQueue,
-			reinterpret_cast< IUnknown* >(Windows::UI::Core::CoreWindow::GetForCurrentThread()),
+			reinterpret_cast<IUnknown*>(Windows::UI::Core::CoreWindow::GetForCurrentThread()),
 			&SwapChainDesc,
 			NULL,
 			SwapChain1.GetInitReference()
@@ -223,7 +223,7 @@ void FD3D12Viewport::Init()
 
 					VERIFYD3D12RESULT(hr);
 				}
-				
+
 				VERIFYD3D12RESULT(SwapChain->QueryInterface(IID_PPV_ARGS(SwapChain1.GetInitReference())));
 
 				// Get a SwapChain4 if supported.
@@ -231,12 +231,17 @@ void FD3D12Viewport::Init()
 			}
 		}
 	}
-// Change from Microsoft for HoloLens support, may require further review: : BEGIN HoloLens support
+	// Change from Microsoft for HoloLens support, may require further review: : BEGIN HoloLens support
 #endif
 // : END HoloLens support
 
-	// Set the DXGI message hook to not change the window behind our back.
-	Adapter->GetDXGIFactory2()->MakeWindowAssociation(WindowHandle, DXGI_MWA_NO_WINDOW_CHANGES);
+	{
+		// Don't make the windows association call and release back buffer at the same time (see notes on critical section)
+		FScopeLock Lock(&DXGIBackBufferLock);
+
+		// Set the DXGI message hook to not change the window behind our back.
+		Adapter->GetDXGIFactory2()->MakeWindowAssociation(WindowHandle, DXGI_MWA_NO_WINDOW_CHANGES);
+	}
 
 	// Resize to setup mGPU correctly.
 	Resize(BufferDesc.Width, BufferDesc.Height, bIsFullscreen, PixelFormat);
diff --git a/Engine/Source/Runtime/D3D12RHI/Public/D3D12Viewport.h b/Engine/Source/Runtime/D3D12RHI/Public/D3D12Viewport.h
--- a/Engine/Source/Runtime/D3D12RHI/Public/D3D12Viewport.h
+++ b/Engine/Source/Runtime/D3D12RHI/Public/D3D12Viewport.h
@@ -64,6 +64,11 @@ class FD3D12Viewport : public FRHIViewport, public FD3D12AdapterChild
 {
 public:
 
+	// Lock viewport windows association and back buffer destruction because of possible crash inside DXGI factory during a call to MakeWindowAssociation
+	// Backbuffer release will wait on the call to MakeWindowAssociation while this will fail internally with 'The requested operation is not implemented.' in KernelBase.dll
+	// Reported & known problem in DXGI and will be fixed with future release but DXGI is not part of the Agility SDK so a code side fix is needed for now.
+	static FCriticalSection DXGIBackBufferLock;
+
 	FD3D12Viewport(class FD3D12Adapter* InParent, HWND InWindowHandle, uint32 InSizeX, uint32 InSizeY, bool bInIsFullscreen, EPixelFormat InPixelFormat);
 
 	void Init();
