Update buffer stats on force release as well via swap to empty buffer (make sure stats are correct)

diff --git a/Engine/Source/Runtime/D3D12RHI/Private/D3D12Buffer.cpp b/Engine/Source/Runtime/D3D12RHI/Private/D3D12Buffer.cpp
--- a/Engine/Source/Runtime/D3D12RHI/Private/D3D12Buffer.cpp
+++ b/Engine/Source/Runtime/D3D12RHI/Private/D3D12Buffer.cpp
@@ -474,6 +474,13 @@ void FD3D12Buffer::Swap(FD3D12Buffer& Other)
 
 void FD3D12Buffer::ReleaseUnderlyingResource()
 {
+	int64 BufferSize = ResourceLocation.GetSize();
+	bool bTransient = ResourceLocation.IsTransient();
+	if (!bTransient)
+	{
+		UpdateBufferStats((EBufferUsageFlags)GetUsage(), -BufferSize);
+	}
+
 	check(IsHeadLink());
 	for (FLinkedObjectIterator NextBuffer(this); NextBuffer; ++NextBuffer)
 	{
