- Add option to dump raytracing stats to CSV as command and from code
- Add BLAS memory stats for dynamic & static types next to total

diff --git a/Engine/Source/Runtime/D3D12RHI/Private/D3D12RayTracing.cpp b/Engine/Source/Runtime/D3D12RHI/Private/D3D12RayTracing.cpp
--- a/Engine/Source/Runtime/D3D12RHI/Private/D3D12RayTracing.cpp
+++ b/Engine/Source/Runtime/D3D12RHI/Private/D3D12RayTracing.cpp
@@ -111,7 +111,9 @@ DECLARE_DWORD_COUNTER_STAT(TEXT("Built BL AS (per frame)"), STAT_D3D12RayTracing
 DECLARE_DWORD_COUNTER_STAT(TEXT("Updated BL AS (per frame)"), STAT_D3D12RayTracingUpdatedBLAS, STATGROUP_D3D12RayTracing);
 DECLARE_DWORD_COUNTER_STAT(TEXT("Built TL AS (per frame)"), STAT_D3D12RayTracingBuiltTLAS, STATGROUP_D3D12RayTracing);
 
-DECLARE_MEMORY_STAT(TEXT("BL AS Memory"), STAT_D3D12RayTracingBLASMemory, STATGROUP_D3D12RayTracing);
+DECLARE_MEMORY_STAT(TEXT("Total BL AS Memory"), STAT_D3D12RayTracingBLASMemory, STATGROUP_D3D12RayTracing);
+DECLARE_MEMORY_STAT(TEXT("Static BL AS Memory"), STAT_D3D12RayTracingStaticBLASMemory, STATGROUP_D3D12RayTracing);
+DECLARE_MEMORY_STAT(TEXT("Dynamic BL AS Memory"), STAT_D3D12RayTracingDynamicBLASMemory, STATGROUP_D3D12RayTracing);
 DECLARE_MEMORY_STAT(TEXT("TL AS Memory"), STAT_D3D12RayTracingTLASMemory, STATGROUP_D3D12RayTracing);
 DECLARE_MEMORY_STAT(TEXT("Total Used Video Memory"), STAT_D3D12RayTracingUsedVideoMemory, STATGROUP_D3D12RayTracing);
 
@@ -132,18 +134,41 @@ inline void UnregisterD3D12RayTracingGeometry(FD3D12RayTracingGeometry* Geometry
 struct FD3D12RayTracingGeometryTracker
 {
 	TSet<FD3D12RayTracingGeometry*> Geometries;
+	uint64 TotalBLASSize = 0;
+	uint64 MaxTotalBLASSize = 0;
 	FCriticalSection CS;
 
+	uint64 GetGeometrySize(FD3D12RayTracingGeometry& Geometry)
+	{
+		if (Geometry.AccelerationStructureCompactedSize != 0)
+		{
+			return Geometry.AccelerationStructureCompactedSize;
+		}
+		else
+		{
+			return Geometry.SizeInfo.ResultSize;
+		}
+	}
+
 	void Add(FD3D12RayTracingGeometry* Geometry)
 	{
+		uint64 BLASSize = GetGeometrySize(*Geometry);
+
 		FScopeLock Lock(&CS);
 		Geometries.Add(Geometry);
+		TotalBLASSize += BLASSize;
+
+		MaxTotalBLASSize = FMath::Max(MaxTotalBLASSize, TotalBLASSize);
 	}
 
 	void Remove(FD3D12RayTracingGeometry* Geometry)
 	{
+		uint64 BLASSize = GetGeometrySize(*Geometry);
+
 		FScopeLock Lock(&CS);
 		Geometries.Remove(Geometry);
+
+		TotalBLASSize -= BLASSize;
 	}
 };
 
@@ -153,45 +178,17 @@ static FD3D12RayTracingGeometryTracker& GetD3D12RayTracingGeometryTracker()
 	return Instance;
 }
 
-static FAutoConsoleCommandWithWorldArgsAndOutputDevice GD3D12DumpRayTracingGeometriesCmd(
-	TEXT("D3D12.DumpRayTracingGeometries"),
-	TEXT("Dump memory allocations for ray tracing resources."),
-	FConsoleCommandWithWorldArgsAndOutputDeviceDelegate::CreateStatic([](const TArray<FString>& Args, UWorld*, FOutputDevice& OutputDevice)
+enum class EDumpRayTracingGeometryMode
+{
+	Top,
+	All,
+};
+
+static void DumpRayTracingGeometries(EDumpRayTracingGeometryMode Mode, int32 NumEntriesToShow, const FString& NameFilter, bool bCSV, FBufferedOutputDevice& BufferedOutput)
 {
 	FD3D12RayTracingGeometryTracker& Tracker = GetD3D12RayTracingGeometryTracker();
 	FScopeLock Lock(&Tracker.CS);
-
-	enum class EMode
-	{
-		Top,
-		All,
-	};
-
-	// Default: show top 50 largest objects.
-	EMode Mode = EMode::Top;
-	int32 NumEntriesToShow = 50;
-
-	FString NameFilter;
-
-	if (Args.Num())
-	{
-		if (Args[0] == TEXT("all"))
-		{
-			Mode = EMode::All;
-			NumEntriesToShow = -1;
-		}
-		else if (FCString::IsNumeric(*Args[0]))
-		{
-			Mode = EMode::Top;
-			LexFromString(NumEntriesToShow, *Args[0]);
-		}
-
-		if (Args.Num() > 1)
-		{
-			NameFilter = Args[1];
-		}
-	}
-
+	
 	auto GetGeometrySize = [](FD3D12RayTracingGeometry& Geometry)
 	{
 		if (Geometry.AccelerationStructureCompactedSize != 0)
@@ -210,7 +207,6 @@ static FAutoConsoleCommandWithWorldArgsAndOutputDevice GD3D12DumpRayTracingGeome
 		return GetGeometrySize(A) > GetGeometrySize(B);
 	});
 
-	FBufferedOutputDevice BufferedOutput;
 	FName CategoryName(TEXT("D3D12RayTracing"));
 	uint64 TotalSizeBytes = 0;
 	uint64 TopSizeBytes = 0;
@@ -244,6 +240,16 @@ static FAutoConsoleCommandWithWorldArgsAndOutputDevice GD3D12DumpRayTracingGeome
 		}
 	};
 
+	FArchive* CSVFile{ nullptr };
+	if (bCSV)
+	{
+		const FString Filename = FString::Printf(TEXT("%sd3d12DumpRayTracingGeometries-%s.csv"), *FPaths::ProfilingDir(), *FDateTime::Now().ToString());
+		CSVFile = IFileManager::Get().CreateFileWriter(*Filename, FILEWRITE_AllowRead);
+
+		const TCHAR* Header = TEXT("Name,Size (MBs),Prims,Segments,Compaction,Update,MarkedForDelete\n");
+		CSVFile->Serialize(TCHAR_TO_ANSI(Header), FPlatformString::Strlen(Header));
+	}
+
 	int32 ShownEntries = 0;
 	for (int32 i=0; i< Geometries.Num(); ++i)
 	{
@@ -252,13 +258,28 @@ static FAutoConsoleCommandWithWorldArgsAndOutputDevice GD3D12DumpRayTracingGeome
 
 		if (ShownEntries < NumEntriesToShow && ShouldShow(Geometry))
 		{
-			BufferedOutput.CategorizedLogf(CategoryName, ELogVerbosity::Log, TEXT("Name: %s - Size: %.3f MB - Prims: %d - Segments: %d -  Compaction: %d - Update: %d"),
-				Geometry->DebugName.IsValid() ? *Geometry->DebugName.ToString() : TEXT("*UNKNOWN*"),
-				SizeBytes / double(1 << 20),
-				Geometry->TotalPrimitiveCount,
-				Geometry->Segments.Num(),
-				(int32)EnumHasAllFlags(Geometry->BuildFlags, ERayTracingAccelerationStructureFlags::AllowCompaction),
-				(int32)EnumHasAllFlags(Geometry->BuildFlags, ERayTracingAccelerationStructureFlags::AllowUpdate));
+			if (bCSV)
+			{
+				const FString Row = FString::Printf(TEXT("%s,%.3f,%d,%d,%d,%d,%d\n"),
+					Geometry->DebugName.IsValid() ? *Geometry->DebugName.ToString() : TEXT("*UNKNOWN*"),
+					SizeBytes / double(1 << 20),
+					Geometry->TotalPrimitiveCount,
+					Geometry->Segments.Num(),
+					(int32)EnumHasAllFlags(Geometry->BuildFlags, ERayTracingAccelerationStructureFlags::AllowCompaction),
+					(int32)EnumHasAllFlags(Geometry->BuildFlags, ERayTracingAccelerationStructureFlags::AllowUpdate),
+					!Geometry->IsValid());
+				CSVFile->Serialize(TCHAR_TO_ANSI(*Row), Row.Len());
+			}
+			else
+			{
+				BufferedOutput.CategorizedLogf(CategoryName, ELogVerbosity::Log, TEXT("Name: %s - Size: %.3f MB - Prims: %d - Segments: %d -  Compaction: %d - Update: %d"),
+					Geometry->DebugName.IsValid() ? *Geometry->DebugName.ToString() : TEXT("*UNKNOWN*"),
+					SizeBytes / double(1 << 20),
+					Geometry->TotalPrimitiveCount,
+					Geometry->Segments.Num(),
+					(int32)EnumHasAllFlags(Geometry->BuildFlags, ERayTracingAccelerationStructureFlags::AllowCompaction),
+					(int32)EnumHasAllFlags(Geometry->BuildFlags, ERayTracingAccelerationStructureFlags::AllowUpdate));
+			}
 			TopSizeBytes += SizeBytes;
 			++ShownEntries;
 		}
@@ -266,20 +287,79 @@ static FAutoConsoleCommandWithWorldArgsAndOutputDevice GD3D12DumpRayTracingGeome
 		TotalSizeBytes += SizeBytes;
 	}
 
-	double TotalSizeF = double(TotalSizeBytes) / double(1 << 20);
-	double TopSizeF = double(TopSizeBytes) / double(1 << 20);
+	if (bCSV)
+	{
+		delete CSVFile;
+		CSVFile = nullptr;
+	}
+	else
+	{
+		double TotalSizeF = double(TotalSizeBytes) / double(1 << 20);
+		double TopSizeF = double(TopSizeBytes) / double(1 << 20);
+
+		if (ShownEntries != Geometries.Num() && ShownEntries)
+		{
+			BufferedOutput.CategorizedLogf(CategoryName, ELogVerbosity::Log,
+				TEXT("Use command `D3D12.DumpRayTracingGeometries all/N [name]` to dump all or N objects. ")
+				TEXT("Optionally add 'name' to filter entries, such as 'skm_'."));
+			BufferedOutput.CategorizedLogf(CategoryName, ELogVerbosity::Log, TEXT("Shown %d entries. Size: %.3f MB (%.2f%% of total)"),
+				ShownEntries, TopSizeF, 100.0 * TopSizeF / TotalSizeF);
+		}
 
-	if (ShownEntries != Geometries.Num() && ShownEntries)
+		BufferedOutput.CategorizedLogf(CategoryName, ELogVerbosity::Log, TEXT("Total size: %.3f MB"), TotalSizeF);
+	}
+}
+
+static FAutoConsoleCommandWithWorldArgsAndOutputDevice GD3D12DumpRayTracingGeometriesCmd(
+	TEXT("D3D12.DumpRayTracingGeometries"),
+	TEXT("Dump memory allocations for ray tracing resources."),
+	FConsoleCommandWithWorldArgsAndOutputDeviceDelegate::CreateStatic([](const TArray<FString>& Args, UWorld*, FOutputDevice& OutputDevice)
+{
+	// Default: show top 50 largest objects.
+	EDumpRayTracingGeometryMode Mode = EDumpRayTracingGeometryMode::Top;
+	int32 NumEntriesToShow = 50;
+	bool bCSV = false;
+
+	FString NameFilter;
+
+	if (Args.Num())
 	{
-		BufferedOutput.CategorizedLogf(CategoryName, ELogVerbosity::Log,
-			TEXT("Use command `D3D12.DumpRayTracingGeometries all/N [name]` to dump all or N objects. ")
-			TEXT("Optionally add 'name' to filter entries, such as 'skm_'."));
-		BufferedOutput.CategorizedLogf(CategoryName, ELogVerbosity::Log, TEXT("Shown %d entries. Size: %.3f MB (%.2f%% of total)"),
-			ShownEntries, TopSizeF, 100.0 * TopSizeF / TotalSizeF);
+		if (Args[0] == TEXT("all"))
+		{
+			Mode = EDumpRayTracingGeometryMode::All;
+			NumEntriesToShow = -1;
+		}
+		else if (FCString::IsNumeric(*Args[0]))
+		{
+			Mode = EDumpRayTracingGeometryMode::Top;
+			LexFromString(NumEntriesToShow, *Args[0]);
+		}
+
+		if (Args.Num() > 1)
+		{
+			NameFilter = Args[1];
+		}
 	}
 
-	BufferedOutput.CategorizedLogf(CategoryName, ELogVerbosity::Log, TEXT("Total size: %.3f MB"), TotalSizeF);
+	FBufferedOutputDevice BufferedOutput;
+	DumpRayTracingGeometries(Mode, NumEntriesToShow, NameFilter, bCSV, BufferedOutput);
+	BufferedOutput.RedirectTo(OutputDevice);
+}));
+
 
+static FAutoConsoleCommandWithWorldArgsAndOutputDevice GD3D12DumpRayTracingGeometriesToCSVCmd(
+	TEXT("D3D12.DumpRayTracingGeometriesToCSV"),
+	TEXT("Dump all memory allocations for ray tracing resources to a CSV file on disc."),
+	FConsoleCommandWithWorldArgsAndOutputDeviceDelegate::CreateStatic([](const TArray<FString>& Args, UWorld*, FOutputDevice& OutputDevice)
+{
+	// CSV dumps all entries
+	EDumpRayTracingGeometryMode Mode = EDumpRayTracingGeometryMode::All;
+	int32 NumEntriesToShow = -1;
+	bool bCSV = true;
+	FString NameFilter;
+
+	FBufferedOutputDevice BufferedOutput;
+	DumpRayTracingGeometries(Mode, NumEntriesToShow, NameFilter, bCSV, BufferedOutput);
 	BufferedOutput.RedirectTo(OutputDevice);
 }));
 
@@ -3408,6 +3488,14 @@ FD3D12RayTracingGeometry::FD3D12RayTracingGeometry(FD3D12Adapter* Adapter, const
 
 		INC_MEMORY_STAT_BY(STAT_D3D12RayTracingUsedVideoMemory, AccelerationStructureBuffers[GPUIndex]->GetSize());
 		INC_MEMORY_STAT_BY(STAT_D3D12RayTracingBLASMemory, AccelerationStructureBuffers[GPUIndex]->GetSize());
+		if (Initializer.bAllowUpdate)
+		{
+			INC_MEMORY_STAT_BY(STAT_D3D12RayTracingDynamicBLASMemory, AccelerationStructureBuffers[GPUIndex]->GetSize());
+		}
+		else
+		{
+			INC_MEMORY_STAT_BY(STAT_D3D12RayTracingStaticBLASMemory, AccelerationStructureBuffers[GPUIndex]->GetSize());
+		}
 	}
 
 	INC_DWORD_STAT_BY(STAT_D3D12RayTracingTrianglesBLAS, Initializer.TotalPrimitiveCount);
@@ -3470,6 +3558,14 @@ FD3D12RayTracingGeometry::~FD3D12RayTracingGeometry()
 		{
 			DEC_MEMORY_STAT_BY(STAT_D3D12RayTracingUsedVideoMemory, Buffer->GetSize());
 			DEC_MEMORY_STAT_BY(STAT_D3D12RayTracingBLASMemory, Buffer->GetSize());
+			if (EnumHasAllFlags(BuildFlags, ERayTracingAccelerationStructureFlags::AllowUpdate))
+			{
+				DEC_MEMORY_STAT_BY(STAT_D3D12RayTracingDynamicBLASMemory, Buffer->GetSize());
+			}
+			else
+			{
+				DEC_MEMORY_STAT_BY(STAT_D3D12RayTracingStaticBLASMemory, Buffer->GetSize());
+			}
 		}
 	}
 
@@ -3795,6 +3891,9 @@ void FD3D12RayTracingGeometry::CompactAccelerationStructure(FD3D12CommandContext
 
 	DEC_MEMORY_STAT_BY(STAT_D3D12RayTracingUsedVideoMemory, AccelerationStructureBuffers[InGPUIndex]->GetSize());
 	DEC_MEMORY_STAT_BY(STAT_D3D12RayTracingBLASMemory, AccelerationStructureBuffers[InGPUIndex]->GetSize());
+	DEC_MEMORY_STAT_BY(STAT_D3D12RayTracingStaticBLASMemory, AccelerationStructureBuffers[InGPUIndex]->GetSize());
+
+	UnregisterD3D12RayTracingGeometry(this);
 
 	// Move old AS into this temporary variable which gets released when this function returns
 	TRefCountPtr<FD3D12Buffer> OldAccelerationStructure = AccelerationStructureBuffers[InGPUIndex];
@@ -3803,6 +3902,7 @@ void FD3D12RayTracingGeometry::CompactAccelerationStructure(FD3D12CommandContext
 
 	INC_MEMORY_STAT_BY(STAT_D3D12RayTracingUsedVideoMemory, AccelerationStructureBuffers[InGPUIndex]->GetSize());
 	INC_MEMORY_STAT_BY(STAT_D3D12RayTracingBLASMemory, AccelerationStructureBuffers[InGPUIndex]->GetSize());
+	INC_MEMORY_STAT_BY(STAT_D3D12RayTracingStaticBLASMemory, AccelerationStructureBuffers[InGPUIndex]->GetSize());
 
 	OldAccelerationStructure->GetResource()->UpdateResidency(CommandContext.CommandListHandle);
 	AccelerationStructureBuffers[InGPUIndex]->GetResource()->UpdateResidency(CommandContext.CommandListHandle);
@@ -3814,6 +3914,8 @@ void FD3D12RayTracingGeometry::CompactAccelerationStructure(FD3D12CommandContext
 		D3D12_RAYTRACING_ACCELERATION_STRUCTURE_COPY_MODE_COMPACT);
 
 	AccelerationStructureCompactedSize = InSizeAfterCompaction;
+
+	RegisterD3D12RayTracingGeometry(this);
 }
 
 FD3D12RayTracingScene::FD3D12RayTracingScene(FD3D12Adapter* Adapter, FRayTracingSceneInitializer2 InInitializer, TResourceArray<D3D12_RAYTRACING_INSTANCE_DESC, 16> InInstances, TArray<uint32> InPerInstanceNumTransforms)
