Add d3d12 memory total stat

diff --git a/Engine/Source/Runtime/D3D12RHI/Private/D3D12DescriptorCache.cpp b/Engine/Source/Runtime/D3D12RHI/Private/D3D12DescriptorCache.cpp
--- a/Engine/Source/Runtime/D3D12RHI/Private/D3D12DescriptorCache.cpp
+++ b/Engine/Source/Runtime/D3D12RHI/Private/D3D12DescriptorCache.cpp
@@ -958,6 +958,7 @@ void FD3D12GlobalOnlineSamplerHeap::Init(uint32 TotalSize)
 
 	INC_DWORD_STAT(STAT_NumSamplerOnlineDescriptorHeaps);
 	INC_MEMORY_STAT_BY(STAT_SamplerOnlineDescriptorHeapMemory, Heap->GetMemorySize());
+	INC_MEMORY_STAT_BY(STAT_D3D12MemoryCurrentTotal, Heap->GetMemorySize());
 }
 
 
@@ -975,8 +976,8 @@ bool FD3D12GlobalOnlineSamplerHeap::RollOver()
 FD3D12SubAllocatedOnlineHeap::FD3D12SubAllocatedOnlineHeap(FD3D12DescriptorCache* InDescriptorCache)
 	: FD3D12OnlineHeap(nullptr, false)
 	, DescriptorCache(InDescriptorCache)
-{
-}
+	{
+	}
 
 FD3D12SubAllocatedOnlineHeap::~FD3D12SubAllocatedOnlineHeap() = default;
 
@@ -1094,6 +1095,7 @@ void FD3D12LocalOnlineHeap::Init(FD3D12Device* InParent, uint32 InNumDescriptors
 		INC_DWORD_STAT(STAT_NumSamplerOnlineDescriptorHeaps);
 		INC_MEMORY_STAT_BY(STAT_SamplerOnlineDescriptorHeapMemory, Heap->GetMemorySize());
 	}
+	INC_MEMORY_STAT_BY(STAT_D3D12MemoryCurrentTotal, Heap->GetMemorySize());
 }
 
 
@@ -1140,6 +1142,7 @@ bool FD3D12LocalOnlineHeap::RollOver()
 			INC_DWORD_STAT(STAT_NumSamplerOnlineDescriptorHeaps);
 			INC_MEMORY_STAT_BY(STAT_SamplerOnlineDescriptorHeapMemory, Heap->GetMemorySize());
 		}
+		INC_MEMORY_STAT_BY(STAT_D3D12MemoryCurrentTotal, Heap->GetMemorySize());
 
 		Entry.Heap = Heap;
 	}
diff --git a/Engine/Source/Runtime/D3D12RHI/Private/D3D12Stats.cpp b/Engine/Source/Runtime/D3D12RHI/Private/D3D12Stats.cpp
--- a/Engine/Source/Runtime/D3D12RHI/Private/D3D12Stats.cpp
+++ b/Engine/Source/Runtime/D3D12RHI/Private/D3D12Stats.cpp
@@ -321,6 +321,7 @@ void UpdateBufferStats(EBufferUsageFlags InUsageFlags, int64 RequestedSize)
 {
 	INC_MEMORY_STAT_BY_FName(GetRHIBufferStats(InUsageFlags), RequestedSize);
 	INC_MEMORY_STAT_BY_FName(GetD3D12BufferStat(InUsageFlags), RequestedSize);
+	INC_MEMORY_STAT_BY(STAT_D3D12MemoryCurrentTotal, RequestedSize);
 }
 
 #if NV_AFTERMATH
diff --git a/Engine/Source/Runtime/D3D12RHI/Private/D3D12Stats.h b/Engine/Source/Runtime/D3D12RHI/Private/D3D12Stats.h
--- a/Engine/Source/Runtime/D3D12RHI/Private/D3D12Stats.h
+++ b/Engine/Source/Runtime/D3D12RHI/Private/D3D12Stats.h
@@ -65,6 +65,7 @@ DECLARE_MEMORY_STAT_EXTERN(TEXT("Demoted Video Memory"), STAT_D3D12DemotedVideoM
 DECLARE_MEMORY_STAT_EXTERN(TEXT("Used Video Memory"), STAT_D3D12UsedVideoMemory, STATGROUP_D3D12Memory, );
 DECLARE_MEMORY_STAT_EXTERN(TEXT("Used System Memory"), STAT_D3D12UsedSystemMemory, STATGROUP_D3D12Memory, );
 
+DECLARE_MEMORY_STAT_EXTERN(TEXT("TOTAL"), STAT_D3D12MemoryCurrentTotal, STATGROUP_D3D12Resources, );
 DECLARE_MEMORY_STAT_EXTERN(TEXT("Render Targets"), STAT_D3D12RenderTargets, STATGROUP_D3D12Resources, );
 DECLARE_MEMORY_STAT_EXTERN(TEXT("UAV Textures"), STAT_D3D12UAVTextures, STATGROUP_D3D12Resources, );
 DECLARE_MEMORY_STAT_EXTERN(TEXT("Textures"), STAT_D3D12Textures, STATGROUP_D3D12Resources, );
diff --git a/Engine/Source/Runtime/D3D12RHI/Private/D3D12Texture.cpp b/Engine/Source/Runtime/D3D12RHI/Private/D3D12Texture.cpp
--- a/Engine/Source/Runtime/D3D12RHI/Private/D3D12Texture.cpp
+++ b/Engine/Source/Runtime/D3D12RHI/Private/D3D12Texture.cpp
@@ -370,6 +370,7 @@ void FD3D12TextureStats::UpdateD3D12TextureStats(TD3D12Texture& Texture, const D
 
 	INC_MEMORY_STAT_BY_FName(GetD3D12StatEnum(Desc.Flags).GetName(), TextureSize);
 	INC_MEMORY_STAT_BY_FName(GetRHIStatEnum(Desc.Flags, bCubeMap, b3D).GetName(), TextureSize);
+	INC_MEMORY_STAT_BY(STAT_D3D12MemoryCurrentTotal, TextureSize);
 
 	if (TextureSize > 0)
 	{
@@ -1029,11 +1030,13 @@ TD3D12Texture2D<BaseResourceType>* FD3D12DynamicRHI::CreateD3D12Texture2D(FRHICo
 			Flags,
 			CreateInfo.ClearValueBinding);		
 
+#if NAME_OBJECTS
 		if (CreateInfo.DebugName)
 		{
 			TRACE_CPUPROFILER_EVENT_SCOPE(D3D12RHI::SetDebugName);
 			NewTexture->SetName(CreateInfo.DebugName);
 		}
+#endif // NAME_OBJECTS
 
 		FD3D12ResourceLocation& Location = NewTexture->ResourceLocation;
 
diff --git a/Engine/Source/Runtime/D3D12RHI/Private/D3D12TransientResourceAllocator.cpp b/Engine/Source/Runtime/D3D12RHI/Private/D3D12TransientResourceAllocator.cpp
--- a/Engine/Source/Runtime/D3D12RHI/Private/D3D12TransientResourceAllocator.cpp
+++ b/Engine/Source/Runtime/D3D12RHI/Private/D3D12TransientResourceAllocator.cpp
@@ -89,6 +89,7 @@ FD3D12TransientHeap::FD3D12TransientHeap(const FInitializer& Initializer, FD3D12
 	SetGPUVirtualAddress(Heap->GetGPUVirtualAddress());
 
 	INC_MEMORY_STAT_BY(STAT_D3D12TransientHeaps, Desc.SizeInBytes);
+	INC_MEMORY_STAT_BY(STAT_D3D12MemoryCurrentTotal, Desc.SizeInBytes);
 }
 
 FD3D12TransientHeap::~FD3D12TransientHeap()
@@ -97,6 +98,7 @@ FD3D12TransientHeap::~FD3D12TransientHeap()
 	{
 		D3D12_HEAP_DESC Desc = Heap->GetHeapDesc();
 		DEC_MEMORY_STAT_BY(STAT_D3D12TransientHeaps, Desc.SizeInBytes);
+		DEC_MEMORY_STAT_BY(STAT_D3D12MemoryCurrentTotal, Desc.SizeInBytes);
 	}
 }
 
diff --git a/Engine/Source/Runtime/D3D12RHI/Private/D3D12Util.cpp b/Engine/Source/Runtime/D3D12RHI/Private/D3D12Util.cpp
--- a/Engine/Source/Runtime/D3D12RHI/Private/D3D12Util.cpp
+++ b/Engine/Source/Runtime/D3D12RHI/Private/D3D12Util.cpp
@@ -1637,6 +1637,7 @@ DEFINE_STAT(STAT_D3D12AvailableVideoMemory);
 DEFINE_STAT(STAT_D3D12DemotedVideoMemory);
 DEFINE_STAT(STAT_D3D12TotalVideoMemory);
 
+DEFINE_STAT(STAT_D3D12MemoryCurrentTotal);
 DEFINE_STAT(STAT_D3D12RenderTargets);
 DEFINE_STAT(STAT_D3D12UAVTextures);
 DEFINE_STAT(STAT_D3D12Textures);
