Rename tags with d3d12/ to d3d12- to remove error reporting from the log

diff --git a/Engine/Source/Runtime/D3D12RHI/Private/D3D12CommandList.cpp b/Engine/Source/Runtime/D3D12RHI/Private/D3D12CommandList.cpp
--- a/Engine/Source/Runtime/D3D12RHI/Private/D3D12CommandList.cpp
+++ b/Engine/Source/Runtime/D3D12RHI/Private/D3D12CommandList.cpp
@@ -212,7 +212,7 @@ void FD3D12CommandListHandle::FD3D12CommandListData::FlushResourceBarriers()
 
 void FD3D12CommandListHandle::FD3D12CommandListData::Reset(FD3D12CommandAllocator& CommandAllocator, bool bTrackExecTime)
 {
-	LLM_SCOPE_BYNAME(TEXT("D3D12/Commandlist"));
+	LLM_SCOPE_BYNAME(TEXT("D3D12Commandlist"));
 	VERIFYD3D12RESULT(CommandList->Reset(CommandAllocator, nullptr));
 
 	CurrentCommandAllocator = &CommandAllocator;
diff --git a/Engine/Source/Runtime/D3D12RHI/Private/D3D12Texture.cpp b/Engine/Source/Runtime/D3D12RHI/Private/D3D12Texture.cpp
--- a/Engine/Source/Runtime/D3D12RHI/Private/D3D12Texture.cpp
+++ b/Engine/Source/Runtime/D3D12RHI/Private/D3D12Texture.cpp
@@ -2164,7 +2164,7 @@ D3D12_RESOURCE_DESC FD3D12DynamicRHI::GetResourceDesc(const FRHITextureCreateInf
 
 void FD3D12TextureBase::UpdateTexture(uint32 MipIndex, uint32 DestX, uint32 DestY, uint32 DestZ, const D3D12_TEXTURE_COPY_LOCATION& SourceCopyLocation)
 {
-	LLM_SCOPE_BYNAME(TEXT("D3D12/CopyTextureRegion"));
+	LLM_SCOPE_BYNAME(TEXT("D3D12CopyTextureRegion"));
 	FD3D12CommandContext& DefaultContext = GetParentDevice()->GetDefaultCommandContext();
 	FD3D12CommandListHandle& hCommandList = DefaultContext.CommandListHandle;
 
