Allowing global constant buffer updates to be more dynamic. Removes the shadow data overhead from stages that aren't used in a context and also allows for global constant buffers larger than 4kb.

diff --git a/Engine/Source/Runtime/D3D12RHI/Private/D3D12ConstantBuffer.cpp b/Engine/Source/Runtime/D3D12RHI/Private/D3D12ConstantBuffer.cpp
--- a/Engine/Source/Runtime/D3D12RHI/Private/D3D12ConstantBuffer.cpp
+++ b/Engine/Source/Runtime/D3D12RHI/Private/D3D12ConstantBuffer.cpp
@@ -1,9 +1,5 @@
 // Copyright Epic Games, Inc. All Rights Reserved.
 
-/*=============================================================================
-D3D12ConstantBuffer.cpp: D3D Constant buffer RHI implementation.
-=============================================================================*/
-
 #include "D3D12RHIPrivate.h"
 
 DEFINE_STAT(STAT_D3D12GlobalConstantBufferUpdateTime);
@@ -16,7 +12,6 @@ FD3D12ConstantBuffer::FD3D12ConstantBuffer(FD3D12Device* InParent, FD3D12FastCon
 	bIsDirty(false),
 	Allocator(InAllocator)
 {
-	FMemory::Memset(ShadowData, 0);
 #if USE_STATIC_ROOT_SIGNATURE
 	View = new FD3D12ConstantBufferView(InParent);
 #endif
@@ -58,8 +53,8 @@ bool FD3D12ConstantBuffer::Version(FD3D12ResourceLocation& BufferOut, bool bDisc
 	// Get the next constant buffer
 	void* Data = Allocator.Allocate(TotalUpdateSize, BufferOut, ViewToUse);
 
-	check(TotalUpdateSize <= sizeof(ShadowData));
-	FMemory::Memcpy(Data, ShadowData, TotalUpdateSize);
+	check(TotalUpdateSize <= (uint32)ShadowData.Num());
+	FMemory::Memcpy(Data, ShadowData.GetData(), TotalUpdateSize);
 
 	bIsDirty = false;
 	return true;
diff --git a/Engine/Source/Runtime/D3D12RHI/Public/D3D12ConstantBuffer.h b/Engine/Source/Runtime/D3D12RHI/Public/D3D12ConstantBuffer.h
--- a/Engine/Source/Runtime/D3D12RHI/Public/D3D12ConstantBuffer.h
+++ b/Engine/Source/Runtime/D3D12RHI/Public/D3D12ConstantBuffer.h
@@ -1,13 +1,9 @@
 // Copyright Epic Games, Inc. All Rights Reserved.
 
-/*=============================================================================
-	D3D12ConstantBuffer.h: Public D3D Constant Buffer definitions.
-	=============================================================================*/
-
 #pragma once
 
-/** Size of the default constant buffer. */
-#define MAX_GLOBAL_CONSTANT_BUFFER_SIZE		4096
+// Chunk size for global constant buffer shadow data. Should always be a multiple of 256
+#define GLOBAL_CONSTANT_BUFFER_CHUNK_SIZE 512
 
 using namespace D3D12RHI;
 
@@ -28,12 +24,19 @@ public:
 	* @param Offset - The offset in the constant buffer to place the data at
 	* @param InSize - The size of the data being copied
 	*/
-	FORCEINLINE void UpdateConstant(const uint8* Data, uint16 Offset, uint16 InSize)
+	FORCEINLINE_DEBUGGABLE void UpdateConstant(const uint8* Data, uint16 Offset, uint16 InSize)
 	{
-		// Check that the data we are shadowing fits in the allocated shadow data
-		check((uint32)Offset + (uint32)InSize <= MAX_GLOBAL_CONSTANT_BUFFER_SIZE);
-		FMemory::Memcpy(ShadowData + Offset, Data, InSize);
-		CurrentUpdateSize = FMath::Max((uint32)(Offset + InSize), CurrentUpdateSize);
+		// Make sure we have enough memory in our shadow data
+		const uint32 RequiredSize = (uint32)Offset + (uint32)InSize;
+		if ((uint32)ShadowData.Num() < RequiredSize)
+		{
+			// Make sure we grow in chunks to prevent too many reallocs
+			const uint32 SizeToGrowTo = Align(RequiredSize, GLOBAL_CONSTANT_BUFFER_CHUNK_SIZE);
+			ShadowData.SetNumZeroed(SizeToGrowTo);
+		}
+
+		FMemory::Memcpy(ShadowData.GetData() + Offset, Data, InSize);
+		CurrentUpdateSize = FMath::Max(RequiredSize, CurrentUpdateSize);
 
 		bIsDirty = true;
 	}
@@ -51,7 +54,7 @@ protected:
 	FD3D12ConstantBufferView* View = nullptr;
 #endif
 
-	__declspec(align(16)) uint8 ShadowData[MAX_GLOBAL_CONSTANT_BUFFER_SIZE];
+	TArray<uint8> ShadowData;
 	
 	/** Size of all constants that has been updated since the last call to Commit. */
 	uint32	CurrentUpdateSize;
