Move per platform RayTracingGeometryInitializer and FRayTracingAccelerationStructureSize to FRHIRayTracingGeometry.
* Move RT initializers to RHIResource.
* Remove duplicated code from FD3D12RayTracingGeometry.

diff --git a/Engine/Source/Runtime/D3D12RHI/Private/D3D12RayTracing.cpp b/Engine/Source/Runtime/D3D12RHI/Private/D3D12RayTracing.cpp
--- a/Engine/Source/Runtime/D3D12RHI/Private/D3D12RayTracing.cpp
+++ b/Engine/Source/Runtime/D3D12RHI/Private/D3D12RayTracing.cpp
@@ -126,6 +126,7 @@ DECLARE_CYCLE_STAT(TEXT("BuildTopLevel"), STAT_D3D12BuildTLAS, STATGROUP_D3D12Ra
 DECLARE_CYCLE_STAT(TEXT("BuildBottomLevel"), STAT_D3D12BuildBLAS, STATGROUP_D3D12RayTracing);
 DECLARE_CYCLE_STAT(TEXT("DispatchRays"), STAT_D3D12DispatchRays, STATGROUP_D3D12RayTracing);
 
+static ERayTracingAccelerationStructureFlags GetRayTracingAccelerationStructureBuildFlags(const FRayTracingGeometryInitializer& Initializer);
 
 #if UE_BUILD_SHIPPING
 inline void RegisterD3D12RayTracingGeometry(FD3D12RayTracingGeometry* Geometry) {};
@@ -256,6 +257,8 @@ static void DumpRayTracingGeometries(EDumpRayTracingGeometryMode Mode, int32 Num
 		FD3D12RayTracingGeometry* Geometry = Geometries[i];
 		uint64 SizeBytes = GetGeometrySize(*Geometry);
 
+		ERayTracingAccelerationStructureFlags GeometryBuildFlags = GetRayTracingAccelerationStructureBuildFlags(Geometry->Initializer);
+
 		if (ShownEntries < NumEntriesToShow && ShouldShow(Geometry))
 		{
 			if (bCSV)
@@ -263,10 +266,10 @@ static void DumpRayTracingGeometries(EDumpRayTracingGeometryMode Mode, int32 Num
 				const FString Row = FString::Printf(TEXT("%s,%.3f,%d,%d,%d,%d,%d\n"),
 					Geometry->DebugName.IsValid() ? *Geometry->DebugName.ToString() : TEXT("*UNKNOWN*"),
 					SizeBytes / double(1 << 20),
-					Geometry->TotalPrimitiveCount,
-					Geometry->Segments.Num(),
-					(int32)EnumHasAllFlags(Geometry->BuildFlags, ERayTracingAccelerationStructureFlags::AllowCompaction),
-					(int32)EnumHasAllFlags(Geometry->BuildFlags, ERayTracingAccelerationStructureFlags::AllowUpdate),
+					Geometry->Initializer.TotalPrimitiveCount,
+					Geometry->Initializer.Segments.Num(),
+					(int32)EnumHasAllFlags(GeometryBuildFlags, ERayTracingAccelerationStructureFlags::AllowCompaction),
+					(int32)EnumHasAllFlags(GeometryBuildFlags, ERayTracingAccelerationStructureFlags::AllowUpdate),
 					!Geometry->IsValid());
 				CSVFile->Serialize(TCHAR_TO_ANSI(*Row), Row.Len());
 			}
@@ -275,10 +278,10 @@ static void DumpRayTracingGeometries(EDumpRayTracingGeometryMode Mode, int32 Num
 				BufferedOutput.CategorizedLogf(CategoryName, ELogVerbosity::Log, TEXT("Name: %s - Size: %.3f MB - Prims: %d - Segments: %d -  Compaction: %d - Update: %d"),
 					Geometry->DebugName.IsValid() ? *Geometry->DebugName.ToString() : TEXT("*UNKNOWN*"),
 					SizeBytes / double(1 << 20),
-					Geometry->TotalPrimitiveCount,
-					Geometry->Segments.Num(),
-					(int32)EnumHasAllFlags(Geometry->BuildFlags, ERayTracingAccelerationStructureFlags::AllowCompaction),
-					(int32)EnumHasAllFlags(Geometry->BuildFlags, ERayTracingAccelerationStructureFlags::AllowUpdate));
+					Geometry->Initializer.TotalPrimitiveCount,
+					Geometry->Initializer.Segments.Num(),
+					(int32)EnumHasAllFlags(GeometryBuildFlags, ERayTracingAccelerationStructureFlags::AllowCompaction),
+					(int32)EnumHasAllFlags(GeometryBuildFlags, ERayTracingAccelerationStructureFlags::AllowUpdate));
 			}
 			TopSizeBytes += SizeBytes;
 			++ShownEntries;
@@ -720,9 +723,10 @@ void FD3D12RayTracingCompactionRequestHandler::RequestCompact(FD3D12RayTracingGe
 {
 	uint32 GPUIndex = GetParentDevice()->GetGPUIndex();
 	check(InRTGeometry->AccelerationStructureBuffers[GPUIndex]);
-	check(EnumHasAllFlags(InRTGeometry->BuildFlags, ERayTracingAccelerationStructureFlags::AllowCompaction) &&
-		EnumHasAllFlags(InRTGeometry->BuildFlags, ERayTracingAccelerationStructureFlags::FastTrace) &&
-		!EnumHasAnyFlags(InRTGeometry->BuildFlags, ERayTracingAccelerationStructureFlags::AllowUpdate));
+	ERayTracingAccelerationStructureFlags GeometryBuildFlags = GetRayTracingAccelerationStructureBuildFlags(InRTGeometry->Initializer);
+	check(EnumHasAllFlags(GeometryBuildFlags, ERayTracingAccelerationStructureFlags::AllowCompaction) &&
+		EnumHasAllFlags(GeometryBuildFlags, ERayTracingAccelerationStructureFlags::FastTrace) &&
+		!EnumHasAnyFlags(GeometryBuildFlags, ERayTracingAccelerationStructureFlags::AllowUpdate));
 
 	FScopeLock Lock(&CS);
 	PendingRequests.Add(InRTGeometry);
@@ -3023,8 +3027,6 @@ void TranslateRayTracingGeometryDescs(const FRayTracingGeometryInitializer& Init
 
 	D3D12_RAYTRACING_GEOMETRY_TYPE GeometryType = TranslateRayTracingGeometryType(Initializer.GeometryType);
 
-	static constexpr uint32 IndicesPerPrimitive = 3; // Only triangle meshes are supported
-
 	uint32 ComputedPrimitiveCountForValidation = 0;
 
 	for (int32 SegmentIndex = 0; SegmentIndex < Initializer.Segments.Num(); ++SegmentIndex)
@@ -3056,7 +3058,7 @@ void TranslateRayTracingGeometryDescs(const FRayTracingGeometryInitializer& Init
 		{
 			uint32 IndexStride = Initializer.IndexBuffer->GetStride();
 			check(Initializer.IndexBuffer->GetSize() >=
-				(Segment.FirstPrimitive + Segment.NumPrimitives) * IndicesPerPrimitive * IndexStride + Initializer.IndexBufferOffset);
+				(Segment.FirstPrimitive + Segment.NumPrimitives) * FD3D12RayTracingGeometry::IndicesPerPrimitive * IndexStride + Initializer.IndexBufferOffset);
 		}
 
 		D3D12_RAYTRACING_GEOMETRY_DESC Desc = {};
@@ -3104,7 +3106,7 @@ void TranslateRayTracingGeometryDescs(const FRayTracingGeometryInitializer& Init
 				// In some cases the geometry is created with 16 bit index buffer, but it's 32 bit at build time.
 				// We conservatively set this to 32 bit to allocate acceleration structure memory.
 				Desc.Triangles.IndexFormat = DXGI_FORMAT_R32_UINT;
-				Desc.Triangles.IndexCount = Segment.NumPrimitives * IndicesPerPrimitive;
+				Desc.Triangles.IndexCount = Segment.NumPrimitives * FD3D12RayTracingGeometry::IndicesPerPrimitive;
 				Desc.Triangles.VertexCount = Segment.MaxVertices;
 			}
 			else
@@ -3432,13 +3434,12 @@ static TRefCountPtr<FD3D12Buffer> CreateRayTracingBuffer(FD3D12Adapter* Adapter,
 	return Result;
 }
 
-FD3D12RayTracingGeometry::FD3D12RayTracingGeometry(FD3D12Adapter* Adapter, const FRayTracingGeometryInitializer& Initializer)
-	: FD3D12AdapterChild(Adapter)
+FD3D12RayTracingGeometry::FD3D12RayTracingGeometry(FD3D12Adapter* Adapter, const FRayTracingGeometryInitializer& InInitializer)
+	: FRHIRayTracingGeometry(InInitializer), FD3D12AdapterChild(Adapter)
 {
 	INC_DWORD_STAT(STAT_D3D12RayTracingAllocatedBLAS);
 
 	DebugName = Initializer.DebugName.IsValid() ? Initializer.DebugName : FName(TEXT("BLAS"));
-	InitializerType = Initializer.Type;
 
 	FMemory::Memzero(bHasPendingCompactionRequests);
 	FMemory::Memzero(bRegisteredAsRenameListener);
@@ -3461,22 +3462,12 @@ FD3D12RayTracingGeometry::FD3D12RayTracingGeometry(FD3D12Adapter* Adapter, const
 		return;
 	}
 
-	IndexOffsetInBytes = Initializer.IndexBufferOffset;
-	TotalPrimitiveCount = Initializer.TotalPrimitiveCount;
-
-	GeometryType = TranslateRayTracingGeometryType(Initializer.GeometryType);
-
-	// #yuriy_todo: get flags directly through the initializer
-	BuildFlags = GetRayTracingAccelerationStructureBuildFlags(Initializer);
-
 	checkf(Initializer.Segments.Num() > 0, TEXT("Ray tracing geometry must be initialized with at least one segment."));
-	Segments = TArray<FRayTracingGeometrySegment>(Initializer.Segments.GetData(), Initializer.Segments.Num());
 
-	RHIIndexBuffer = Initializer.IndexBuffer;
-
-	static constexpr uint32 IndicesPerPrimitive = 3; // Only triangle meshes are supported
+	// #yuriy_todo: get flags directly through the initializer
+	ERayTracingAccelerationStructureFlags BuildFlags = GetRayTracingAccelerationStructureBuildFlags(Initializer);
 
-	GeometryDescs.SetNumUninitialized(Segments.Num());
+	GeometryDescs.SetNumUninitialized(Initializer.Segments.Num());
 	TranslateRayTracingGeometryDescs(Initializer, GeometryDescs);
 
 	SetDirty(FRHIGPUMask::All(), true);
@@ -3558,7 +3549,7 @@ FD3D12RayTracingGeometry::FD3D12RayTracingGeometry(FD3D12Adapter* Adapter, const
 
 		if (bOnAsyncThread)
 		{
-			check(InitializerType == ERayTracingGeometryInitializerType::StreamingSource);
+			check(Initializer.Type == ERayTracingGeometryInitializerType::StreamingSource);
 
 			FD3D12ResourceLocation* SrcResourceLoc_Heap = new FD3D12ResourceLocation(SrcResourceLoc.GetParentDevice());
 			FD3D12ResourceLocation::TransferOwnership(*SrcResourceLoc_Heap, SrcResourceLoc);
@@ -3630,7 +3621,7 @@ void FD3D12RayTracingGeometry::ReleaseUnderlyingResource()
 		UnregisterAsRenameListener(GPUIndex);
 	}
 
-	if (InitializerType != ERayTracingGeometryInitializerType::StreamingSource)
+	if (Initializer.Type != ERayTracingGeometryInitializerType::StreamingSource)
 	{
 		for (TRefCountPtr<FD3D12Buffer>& Buffer : AccelerationStructureBuffers)
 		{
@@ -3638,6 +3629,8 @@ void FD3D12RayTracingGeometry::ReleaseUnderlyingResource()
 			{
 				DEC_MEMORY_STAT_BY(STAT_D3D12RayTracingUsedVideoMemory, Buffer->GetSize());
 				DEC_MEMORY_STAT_BY(STAT_D3D12RayTracingBLASMemory, Buffer->GetSize());
+
+				ERayTracingAccelerationStructureFlags BuildFlags = GetRayTracingAccelerationStructureBuildFlags(Initializer);
 				if (EnumHasAllFlags(BuildFlags, ERayTracingAccelerationStructureFlags::AllowUpdate))
 				{
 					DEC_MEMORY_STAT_BY(STAT_D3D12RayTracingDynamicBLASMemory, Buffer->GetSize());
@@ -3649,7 +3642,7 @@ void FD3D12RayTracingGeometry::ReleaseUnderlyingResource()
 			}
 		}
 
-		DEC_DWORD_STAT_BY(STAT_D3D12RayTracingTrianglesBLAS, TotalPrimitiveCount);
+		DEC_DWORD_STAT_BY(STAT_D3D12RayTracingTrianglesBLAS, Initializer.TotalPrimitiveCount);
 		DEC_DWORD_STAT(STAT_D3D12RayTracingAllocatedBLAS);
 	}
 
@@ -3659,15 +3652,11 @@ void FD3D12RayTracingGeometry::ReleaseUnderlyingResource()
 		Buffer.SafeRelease();
 	}	
 
-	IndexOffsetInBytes = 0;
-	TotalPrimitiveCount = 0;
-	BuildFlags = ERayTracingAccelerationStructureFlags::None;
-	GeometryType = D3D12_RAYTRACING_GEOMETRY_TYPE_TRIANGLES;
-	RHIIndexBuffer = nullptr;
+	Initializer = {};
+
 	AccelerationStructureCompactedSize = 0;
 	GeometryDescs = {};
 	SizeInfo = {};
-	Segments.Empty();
 	for (TArray<FHitGroupSystemParameters>& HitGroupParametersForGPU : HitGroupSystemParameters)
 	{
 		HitGroupParametersForGPU.Empty();
@@ -3683,15 +3672,15 @@ void FD3D12RayTracingGeometry::RegisterAsRenameListener(uint32 InGPUIndex)
 {
 	check(!bRegisteredAsRenameListener[InGPUIndex]);
 
-	FD3D12Buffer* IndexBuffer = FD3D12DynamicRHI::ResourceCast(RHIIndexBuffer.GetReference(), InGPUIndex);
+	FD3D12Buffer* IndexBuffer = FD3D12DynamicRHI::ResourceCast(Initializer.IndexBuffer.GetReference(), InGPUIndex);
 	if (IndexBuffer)
 	{
 		IndexBuffer->AddRenameListener(this);
 	}
 
 	TArray<FD3D12Buffer*, TInlineAllocator<1>> UniqueVertexBuffers;
-	UniqueVertexBuffers.Reserve(Segments.Num());
-	for (const FRayTracingGeometrySegment& Segment : Segments)
+	UniqueVertexBuffers.Reserve(Initializer.Segments.Num());
+	for (const FRayTracingGeometrySegment& Segment : Initializer.Segments)
 	{
 		FD3D12Buffer* VertexBuffer = FD3D12DynamicRHI::ResourceCast(Segment.VertexBuffer.GetReference(), InGPUIndex);
 		if (VertexBuffer && !UniqueVertexBuffers.Contains(VertexBuffer))
@@ -3711,15 +3700,15 @@ void FD3D12RayTracingGeometry::UnregisterAsRenameListener(uint32 InGPUIndex)
 		return;
 	}
 
-	FD3D12Buffer* IndexBuffer = FD3D12DynamicRHI::ResourceCast(RHIIndexBuffer.GetReference(), InGPUIndex);
+	FD3D12Buffer* IndexBuffer = FD3D12DynamicRHI::ResourceCast(Initializer.IndexBuffer.GetReference(), InGPUIndex);
 	if (IndexBuffer)
 	{
 		IndexBuffer->RemoveRenameListener(this);
 	}
 
 	TArray<FD3D12Buffer*, TInlineAllocator<1>> UniqueVertexBuffers;
-	UniqueVertexBuffers.Reserve(Segments.Num());
-	for (const FRayTracingGeometrySegment& Segment : Segments)
+	UniqueVertexBuffers.Reserve(Initializer.Segments.Num());
+	for (const FRayTracingGeometrySegment& Segment : Initializer.Segments)
 	{
 		FD3D12Buffer* VertexBuffer = FD3D12DynamicRHI::ResourceCast(Segment.VertexBuffer.GetReference(), InGPUIndex);
 		if (VertexBuffer && !UniqueVertexBuffers.Contains(VertexBuffer))
@@ -3740,12 +3729,12 @@ void FD3D12RayTracingGeometry::ResourceRenamed(FD3D12BaseShaderResource* InRenam
 	// still referenced by the previous RT scene (but not used anymore). This only happens during freeing of RHI buffers during mesh streaming.
 	if (InNewResourceLocation == nullptr)
 	{	
-		FD3D12Buffer* IndexBuffer = FD3D12DynamicRHI::ResourceCast(RHIIndexBuffer.GetReference(), InRenamedResource->GetParentDevice()->GetGPUIndex());
+		FD3D12Buffer* IndexBuffer = FD3D12DynamicRHI::ResourceCast(Initializer.IndexBuffer.GetReference(), InRenamedResource->GetParentDevice()->GetGPUIndex());
 		if (IndexBuffer == InRenamedResource)
 		{
-			RHIIndexBuffer = nullptr;
+			Initializer.IndexBuffer = nullptr;
 		}
-		for (FRayTracingGeometrySegment& Segment : Segments)
+		for (FRayTracingGeometrySegment& Segment : Initializer.Segments)
 		{
 			FD3D12Buffer* VertexBuffer = FD3D12DynamicRHI::ResourceCast(Segment.VertexBuffer.GetReference(), InRenamedResource->GetParentDevice()->GetGPUIndex());
 			if (VertexBuffer == InRenamedResource)
@@ -3766,16 +3755,16 @@ void FD3D12RayTracingGeometry::ResourceRenamed(FD3D12BaseShaderResource* InRenam
 void FD3D12RayTracingGeometry::TransitionBuffers(FD3D12CommandContext& CommandContext)
 {
 	// Transition vertex and index resources..
-	if (RHIIndexBuffer)
+	if (Initializer.IndexBuffer)
 	{
-		FD3D12Buffer* IndexBuffer = CommandContext.RetrieveObject<FD3D12Buffer>(RHIIndexBuffer.GetReference());
+		FD3D12Buffer* IndexBuffer = CommandContext.RetrieveObject<FD3D12Buffer>(Initializer.IndexBuffer.GetReference());
 		if (IndexBuffer->GetResource()->RequiresResourceStateTracking())
 		{
 			FD3D12DynamicRHI::TransitionResource(CommandContext.CommandListHandle, IndexBuffer->GetResource(), D3D12_RESOURCE_STATE_TBD, D3D12_RESOURCE_STATE_NON_PIXEL_SHADER_RESOURCE, 0, FD3D12DynamicRHI::ETransitionMode::Apply);
 		}
 	}
 
-	for (const FRayTracingGeometrySegment& Segment : Segments)
+	for (const FRayTracingGeometrySegment& Segment : Initializer.Segments)
 	{
 		const FBufferRHIRef& RHIVertexBuffer = Segment.VertexBuffer;
 		FD3D12Buffer* VertexBuffer = CommandContext.RetrieveObject<FD3D12Buffer>(RHIVertexBuffer.GetReference());
@@ -3788,13 +3777,13 @@ void FD3D12RayTracingGeometry::TransitionBuffers(FD3D12CommandContext& CommandCo
 
 void FD3D12RayTracingGeometry::UpdateResidency(FD3D12CommandContext& CommandContext)
 {
-	if (RHIIndexBuffer)
+	if (Initializer.IndexBuffer)
 	{
-		FD3D12Buffer* IndexBuffer = CommandContext.RetrieveObject<FD3D12Buffer>(RHIIndexBuffer.GetReference());
+		FD3D12Buffer* IndexBuffer = CommandContext.RetrieveObject<FD3D12Buffer>(Initializer.IndexBuffer.GetReference());
 		IndexBuffer->GetResource()->UpdateResidency(CommandContext.CommandListHandle);		
 	}
 
-	for (const FRayTracingGeometrySegment& Segment : Segments)
+	for (const FRayTracingGeometrySegment& Segment : Initializer.Segments)
 	{
 		const FBufferRHIRef& RHIVertexBuffer = Segment.VertexBuffer;
 		FD3D12Buffer* VertexBuffer = CommandContext.RetrieveObject<FD3D12Buffer>(RHIVertexBuffer.GetReference());		
@@ -3805,9 +3794,11 @@ void FD3D12RayTracingGeometry::UpdateResidency(FD3D12CommandContext& CommandCont
 	AccelerationStructureBuffers[GPUIndex]->GetResource()->UpdateResidency(CommandContext.CommandListHandle);
 }
 
-void FD3D12RayTracingGeometry::SetInitializer(const FRayTracingGeometryInitializer& Initializer)
+void FD3D12RayTracingGeometry::SetInitializer(const FRayTracingGeometryInitializer& InInitializer)
 {
-	checkf(InitializerType == ERayTracingGeometryInitializerType::StreamingDestination, TEXT("Only FD3D12RayTracingGeometry that was created as StreamingDestination can update their initializer."));
+	checkf(Initializer.Type == ERayTracingGeometryInitializerType::StreamingDestination, TEXT("Only FD3D12RayTracingGeometry that was created as StreamingDestination can update their initializer."));
+	Initializer = InInitializer;
+	Initializer.Type = ERayTracingGeometryInitializerType::StreamingDestination; // Make sure this flag is set at all times so we can verify the resource state
 
 	for (uint32 GPUIndex = 0; GPUIndex < MAX_NUM_GPUS && GPUIndex < GNumExplicitGPUsForRendering; ++GPUIndex)
 	{
@@ -3816,19 +3807,9 @@ void FD3D12RayTracingGeometry::SetInitializer(const FRayTracingGeometryInitializ
 
 	DebugName = Initializer.DebugName.IsValid() ? Initializer.DebugName : FName(TEXT("BLAS"));
 
-	IndexOffsetInBytes = Initializer.IndexBufferOffset;
-	TotalPrimitiveCount = Initializer.TotalPrimitiveCount;
-
-	GeometryType = TranslateRayTracingGeometryType(Initializer.GeometryType);
-
-	BuildFlags = GetRayTracingAccelerationStructureBuildFlags(Initializer);
-
 	checkf(Initializer.Segments.Num() > 0, TEXT("Ray tracing geometry must be initialized with at least one segment."));
-	Segments = TArray<FRayTracingGeometrySegment>(Initializer.Segments.GetData(), Initializer.Segments.Num());
 
-	RHIIndexBuffer = Initializer.IndexBuffer;
-	
-	GeometryDescs.SetNumUninitialized(Segments.Num());
+	GeometryDescs.SetNumUninitialized(Initializer.Segments.Num());
 	TranslateRayTracingGeometryDescs(Initializer, GeometryDescs);
 	
 	FRHICommandListImmediate& RHICmdList = FRHICommandListExecutor::GetImmediateCommandList();
@@ -3846,12 +3827,14 @@ void FD3D12RayTracingGeometry::SetInitializer(const FRayTracingGeometryInitializ
 
 void FD3D12RayTracingGeometry::SetupHitGroupSystemParameters(uint32 InGPUIndex)
 {
+	D3D12_RAYTRACING_GEOMETRY_TYPE GeometryType = TranslateRayTracingGeometryType(Initializer.GeometryType);
+
 	TArray<FHitGroupSystemParameters>& HitGroupSystemParametersForThisGPU = HitGroupSystemParameters[InGPUIndex];
-	HitGroupSystemParametersForThisGPU.Reset(Segments.Num());
+	HitGroupSystemParametersForThisGPU.Reset(Initializer.Segments.Num());
 
-	FD3D12Buffer* IndexBuffer = FD3D12DynamicRHI::ResourceCast(RHIIndexBuffer.GetReference(), InGPUIndex);
+	FD3D12Buffer* IndexBuffer = FD3D12DynamicRHI::ResourceCast(Initializer.IndexBuffer.GetReference(), InGPUIndex);
 	const uint32 IndexStride = IndexBuffer ? IndexBuffer->GetStride() : 0;
-	for (const FRayTracingGeometrySegment& Segment : Segments)
+	for (const FRayTracingGeometrySegment& Segment : Initializer.Segments)
 	{
 		FD3D12Buffer* VertexBuffer = FD3D12DynamicRHI::ResourceCast(Segment.VertexBuffer.GetReference(), InGPUIndex);
 
@@ -3861,9 +3844,8 @@ void FD3D12RayTracingGeometry::SetupHitGroupSystemParameters(uint32 InGPUIndex)
 
 		if (GeometryType == D3D12_RAYTRACING_GEOMETRY_TYPE_TRIANGLES && IndexBuffer != nullptr)
 		{
-			static constexpr uint32 IndicesPerPrimitive = 3; // Only triangle meshes are supported
 			SystemParameters.IndexBuffer = IndexBuffer->ResourceLocation.GetGPUVirtualAddress();
-			SystemParameters.RootConstants.IndexBufferOffsetInBytes = IndexOffsetInBytes + IndexStride * Segment.FirstPrimitive * IndicesPerPrimitive;
+			SystemParameters.RootConstants.IndexBufferOffsetInBytes = Initializer.IndexBufferOffset + IndexStride * Segment.FirstPrimitive * FD3D12RayTracingGeometry::IndicesPerPrimitive;
 		}
 
 		HitGroupSystemParametersForThisGPU.Add(SystemParameters);
@@ -3872,30 +3854,32 @@ void FD3D12RayTracingGeometry::SetupHitGroupSystemParameters(uint32 InGPUIndex)
 
 void FD3D12RayTracingGeometry::CreateAccelerationStructureBuildDesc(FD3D12CommandContext& CommandContext, EAccelerationStructureBuildMode BuildMode, D3D12_GPU_VIRTUAL_ADDRESS ScratchBufferAddress, D3D12_BUILD_RAYTRACING_ACCELERATION_STRUCTURE_DESC& OutDesc, TArrayView<D3D12_RAYTRACING_GEOMETRY_DESC>& OutGeometryDescs) const
 {
-	if (RHIIndexBuffer)
+	if (Initializer.IndexBuffer)
 	{
-		checkf(RHIIndexBuffer->GetStride() == 2 || RHIIndexBuffer->GetStride() == 4, TEXT("Index buffer must be 16 or 32 bit."));
+		checkf(Initializer.IndexBuffer->GetStride() == 2 || Initializer.IndexBuffer->GetStride() == 4, TEXT("Index buffer must be 16 or 32 bit."));
 	}
 
 	const uint32 GPUIndex = CommandContext.GetGPUIndex();
-	const uint32 IndexStride = RHIIndexBuffer ? RHIIndexBuffer->GetStride() : 0;
+	const uint32 IndexStride = Initializer.IndexBuffer ? Initializer.IndexBuffer->GetStride() : 0;
 	const bool bIsUpdate = BuildMode == EAccelerationStructureBuildMode::Update;
 
 	// Use the pre-built descs as template and set the GPU resource pointers (current VB/IB).
 	check(OutGeometryDescs.Num() == GeometryDescs.Num());
 
-	FD3D12Buffer* IndexBuffer = CommandContext.RetrieveObject<FD3D12Buffer>(RHIIndexBuffer.GetReference());
+	FD3D12Buffer* IndexBuffer = CommandContext.RetrieveObject<FD3D12Buffer>(Initializer.IndexBuffer.GetReference());
 	FD3D12Buffer* NullTransformBufferD3D12 = CommandContext.RetrieveObject<FD3D12Buffer>(NullTransformBuffer.GetReference());
 
 	const TArray<FHitGroupSystemParameters>& HitGroupSystemParametersForThisGPU = HitGroupSystemParameters[GPUIndex];
-	check(HitGroupSystemParametersForThisGPU.Num() == Segments.Num());
+	check(HitGroupSystemParametersForThisGPU.Num() == Initializer.Segments.Num());
 
-	for (int32 SegmentIndex = 0; SegmentIndex < Segments.Num(); ++SegmentIndex)
+	ERayTracingAccelerationStructureFlags BuildFlags = GetRayTracingAccelerationStructureBuildFlags(Initializer);
+	D3D12_RAYTRACING_GEOMETRY_TYPE GeometryType = TranslateRayTracingGeometryType(Initializer.GeometryType);
+	for (int32 SegmentIndex = 0; SegmentIndex < Initializer.Segments.Num(); ++SegmentIndex)
 	{
 		D3D12_RAYTRACING_GEOMETRY_DESC& Desc = OutGeometryDescs[SegmentIndex];
 		Desc = GeometryDescs[SegmentIndex]; // Copy from template
 
-		const FRayTracingGeometrySegment& Segment = Segments[SegmentIndex];
+		const FRayTracingGeometrySegment& Segment = Initializer.Segments[SegmentIndex];
 		const FHitGroupSystemParameters& SystemParameters = HitGroupSystemParametersForThisGPU[SegmentIndex];
 
 		FD3D12Buffer* VertexBuffer = CommandContext.RetrieveObject<FD3D12Buffer>(Segment.VertexBuffer.GetReference());
@@ -3932,8 +3916,7 @@ void FD3D12RayTracingGeometry::CreateAccelerationStructureBuildDesc(FD3D12Comman
 
 			if (IndexBuffer)
 			{
-				static constexpr uint32 IndicesPerPrimitive = 3; // Only triangle meshes are supported
-				check(Desc.Triangles.IndexCount <= Segment.NumPrimitives * IndicesPerPrimitive);
+				check(Desc.Triangles.IndexCount <= Segment.NumPrimitives * FD3D12RayTracingGeometry::IndicesPerPrimitive);
 
 				Desc.Triangles.IndexFormat = (IndexStride == 4 ? DXGI_FORMAT_R32_UINT : DXGI_FORMAT_R16_UINT);
 				Desc.Triangles.IndexBuffer = SystemParameters.IndexBuffer + SystemParameters.RootConstants.IndexBufferOffsetInBytes;
@@ -3941,7 +3924,7 @@ void FD3D12RayTracingGeometry::CreateAccelerationStructureBuildDesc(FD3D12Comman
 			else
 			{
 				// Non-indexed geometry
-				checkf(Segments.Num() == 1, TEXT("Non-indexed geometry with multiple segments is not implemented."));
+				checkf(Initializer.Segments.Num() == 1, TEXT("Non-indexed geometry with multiple segments is not implemented."));
 				check(Desc.Triangles.IndexFormat == DXGI_FORMAT_UNKNOWN);
 				check(Desc.Triangles.IndexCount == 0);
 				check(Desc.Triangles.IndexBuffer == D3D12_GPU_VIRTUAL_ADDRESS(0));
@@ -4290,13 +4273,13 @@ void FD3D12RayTracingScene::BuildAccelerationStructure(FD3D12CommandContext& Com
 
 				AddResidencyHandleForResource(Geometry->AccelerationStructureBuffers[GPUIndex]->GetResource());
 
-				if (Geometry->RHIIndexBuffer)
+				if (Geometry->Initializer.IndexBuffer)
 				{
-					FD3D12Buffer* IndexBuffer = CommandContext.RetrieveObject<FD3D12Buffer>(Geometry->RHIIndexBuffer.GetReference());
+					FD3D12Buffer* IndexBuffer = CommandContext.RetrieveObject<FD3D12Buffer>(Geometry->Initializer.IndexBuffer.GetReference());
 					AddResidencyHandleForResource(IndexBuffer->GetResource());
 				}
 
-				for (const FRayTracingGeometrySegment& Segment : Geometry->Segments)
+				for (const FRayTracingGeometrySegment& Segment : Geometry->Initializer.Segments)
 				{
 					if (Segment.VertexBuffer)
 					{
@@ -4447,20 +4430,20 @@ void FD3D12CommandContext::RHIBuildAccelerationStructures(const TArrayView<const
 
 		if (P.Segments.Num())
 		{
-			checkf(P.Segments.Num() == Geometry->Segments.Num(),
+			checkf(P.Segments.Num() == Geometry->Initializer.Segments.Num(),
 				TEXT("If updated segments are provided, they must exactly match existing geometry segments. Only vertex buffer bindings may change."));
 
 			for (int32 i = 0; i < P.Segments.Num(); ++i)
 			{
-				checkf(P.Segments[i].MaxVertices <= Geometry->Segments[i].MaxVertices,
+				checkf(P.Segments[i].MaxVertices <= Geometry->Initializer.Segments[i].MaxVertices,
 					TEXT("Maximum number of vertices in a segment (%u) must not be smaller than what was declared during FRHIRayTracingGeometry creation (%u), as this controls BLAS memory allocation."),
-					P.Segments[i].MaxVertices, Geometry->Segments[i].MaxVertices
+					P.Segments[i].MaxVertices, Geometry->Initializer.Segments[i].MaxVertices
 				);
 
-				Geometry->Segments[i].VertexBuffer            = P.Segments[i].VertexBuffer;
-				Geometry->Segments[i].VertexBufferElementType = P.Segments[i].VertexBufferElementType;
-				Geometry->Segments[i].VertexBufferStride      = P.Segments[i].VertexBufferStride;
-				Geometry->Segments[i].VertexBufferOffset      = P.Segments[i].VertexBufferOffset;
+				Geometry->Initializer.Segments[i].VertexBuffer            = P.Segments[i].VertexBuffer;
+				Geometry->Initializer.Segments[i].VertexBufferElementType = P.Segments[i].VertexBufferElementType;
+				Geometry->Initializer.Segments[i].VertexBufferStride      = P.Segments[i].VertexBufferStride;
+				Geometry->Initializer.Segments[i].VertexBufferOffset      = P.Segments[i].VertexBufferOffset;
 			}
 		}
 	}
@@ -4568,7 +4551,8 @@ void FD3D12CommandContext::RHIBuildAccelerationStructures(const TArrayView<const
 
 		if (Geometry->IsDirty(GPUIndex))
 		{
-			if (ShouldCompactAfterBuild(Geometry->BuildFlags))
+			ERayTracingAccelerationStructureFlags GeometryBuildFlags = GetRayTracingAccelerationStructureBuildFlags(Geometry->Initializer);
+			if (ShouldCompactAfterBuild(GeometryBuildFlags))
 			{
 				GetParentDevice()->GetRayTracingCompactionRequestHandler()->RequestCompact(Geometry);
 				Geometry->bHasPendingCompactionRequests[GPUIndex] = true;
diff --git a/Engine/Source/Runtime/D3D12RHI/Private/D3D12RayTracing.h b/Engine/Source/Runtime/D3D12RHI/Private/D3D12RayTracing.h
--- a/Engine/Source/Runtime/D3D12RHI/Private/D3D12RayTracing.h
+++ b/Engine/Source/Runtime/D3D12RHI/Private/D3D12RayTracing.h
@@ -31,8 +31,6 @@ public:
 		checkf(IsInRHIThread() || !IsRunningRHIInSeparateThread(), TEXT("Acceleration structure addresses can only be accessed on RHI timeline due to compaction and defragmentation."));
 		return AccelerationStructureBuffers[GPUIndex]->ResourceLocation.GetGPUVirtualAddress();
 	}
-	virtual uint32 GetNumSegments() const final override { return Segments.Num(); }
-	virtual FRayTracingAccelerationStructureSize GetSizeInfo() const final override { return SizeInfo; }
 	virtual void SetInitializer(const FRayTracingGeometryInitializer& Initializer) final override;
 
 	void SetupHitGroupSystemParameters(uint32 InGPUIndex);
@@ -65,15 +63,11 @@ public:
 		return bIsAccelerationStructureDirty[GPUIndex];
 	}
 
-	uint32 IndexOffsetInBytes = 0;
-	uint32 TotalPrimitiveCount = 0; // Combined number of primitives in all mesh segments
+	using FRHIRayTracingGeometry::Initializer;
+	using FRHIRayTracingGeometry::SizeInfo;
 
-	D3D12_RAYTRACING_GEOMETRY_TYPE GeometryType = D3D12_RAYTRACING_GEOMETRY_TYPE_TRIANGLES;
+	static constexpr uint32 IndicesPerPrimitive = 3; // Triangle geometry only
 
-	TArray<FRayTracingGeometrySegment> Segments; // Defines addressable parts of the mesh that can be used for material assignment (one segment = one SBT record)
-	ERayTracingAccelerationStructureFlags BuildFlags = ERayTracingAccelerationStructureFlags::None;
-
-	FBufferRHIRef  RHIIndexBuffer;
 	static FBufferRHIRef NullTransformBuffer; // Null transform for hidden sections
 
 	TRefCountPtr<FD3D12Buffer> AccelerationStructureBuffers[MAX_NUM_GPUS];
@@ -92,9 +86,6 @@ public:
 	TArray<D3D12_RAYTRACING_GEOMETRY_DESC, TInlineAllocator<1>> GeometryDescs;
 
 	uint64 AccelerationStructureCompactedSize = 0;
-	FRayTracingAccelerationStructureSize SizeInfo = {};
-
-	ERayTracingGeometryInitializerType InitializerType;
 };
 
 class FD3D12RayTracingScene : public FRHIRayTracingScene, public FD3D12AdapterChild, public FNoncopyable
diff --git a/Engine/Source/Runtime/RHI/Public/DynamicRHI.h b/Engine/Source/Runtime/RHI/Public/DynamicRHI.h
--- a/Engine/Source/Runtime/RHI/Public/DynamicRHI.h
+++ b/Engine/Source/Runtime/RHI/Public/DynamicRHI.h
@@ -80,239 +80,6 @@ struct FRHIFlipDetails
 	{}
 };
 
-enum class ERayTracingInstanceFlags : uint8
-{
-	None                  = 0,
-	TriangleCullDisable   = 1 << 1, // No back face culling. Triangle is visible from both sides.
-	TriangleCullReverse   = 1 << 2, // Makes triangle front-facing if its vertices are counterclockwise from ray origin.
-	ForceOpaque           = 1 << 3, // Disable any-hit shader invocation for this instance.
-	ForceNonOpaque        = 1 << 4, // Force any-hit shader invocation even if geometries inside the instance were marked opaque.
-};
-ENUM_CLASS_FLAGS(ERayTracingInstanceFlags);
-
-/**
-* High level descriptor of one or more instances of a mesh in a ray tracing scene.
-* All instances covered by this descriptor will share shader bindings, but may have different transforms and user data.
-*/
-struct FRayTracingGeometryInstance
-{
-	// TODO: UE-130819 Ref counting is a temporary workaround for a very rare streaming crash.
-	FRayTracingGeometryRHIRef GeometryRHI = nullptr;
-
-	// A single physical mesh may be duplicated many times in the scene with different transforms and user data.
-	// All copies share the same shader binding table entries and therefore will have the same material and shader resources.
-	TArrayView<const FMatrix> Transforms;
-
-	TArrayView<const uint32> InstanceSceneDataOffsets;
-
-	// Optional buffer that stores GPU transforms. Used instead of CPU-side transform data.
-	FShaderResourceViewRHIRef GPUTransformsSRV = nullptr;
-
-	// Conservative number of instances. Some of the actual instances may be made inactive if GPU transforms are used.
-	// Must be less or equal to number of entries in Transforms view if CPU transform data is used.
-	// Must be less or equal to number of entries in GPUTransformsSRV if it is non-null.
-	uint32 NumTransforms = 0;
-
-	// Each geometry copy can receive a user-provided integer, which can be used to retrieve extra shader parameters or customize appearance.
-	// This data can be retrieved using GetInstanceUserData() in closest/any hit shaders.
-	// If UserData view is empty, then DefaultUserData value will be used for all instances.
-	// If UserData view is used, then it must have the same number of entries as NumInstances.
-	uint32 DefaultUserData = 0;
-	TArrayView<const uint32> UserData;
-
-	// Each geometry copy can have one bit to make it individually deactivated (removed from TLAS while maintaining hit group indexing). Useful for culling.
-	TArrayView<const uint32> ActivationMask;
-
-	// Mask that will be tested against one provided to TraceRay() in shader code.
-	// If binary AND of instance mask with ray mask is zero, then the instance is considered not intersected / invisible.
-	uint8 Mask = 0xFF;
-
-	// Flags to control triangle back face culling, whether to allow any-hit shaders, etc.
-	ERayTracingInstanceFlags Flags = ERayTracingInstanceFlags::None;
-};
-
-enum ERayTracingGeometryType
-{
-	// Indexed or non-indexed triangle list with fixed function ray intersection.
-	// Vertex buffer must contain vertex positions as VET_Float3.
-	// Vertex stride must be at least 12 bytes, but may be larger to support custom per-vertex data.
-	// Index buffer may be provided for indexed triangle lists. Implicit triangle list is assumed otherwise.
-	RTGT_Triangles,
-
-	// Custom primitive type that requires an intersection shader.
-	// Vertex buffer for procedural geometry must contain one AABB per primitive as {float3 MinXYZ, float3 MaxXYZ}.
-	// Vertex stride must be at least 24 bytes, but may be larger to support custom per-primitive data.
-	// Index buffers can't be used with procedural geometry.
-	RTGT_Procedural,
-};
-DECLARE_INTRINSIC_TYPE_LAYOUT(ERayTracingGeometryType);
-
-enum class ERayTracingGeometryInitializerType
-{
-	// Fully initializes the RayTracingGeometry object: creates underlying buffer and initializes shader parameters.
-	Rendering,
-
-	// Does not create underlying buffer or shader parameters. Used by the streaming system as an object that is streamed into. 
-	StreamingDestination,
-
-	// Creates buffers but does not create shader parameters. Used for intermediate objects in the streaming system.
-	StreamingSource, 
-};
-DECLARE_INTRINSIC_TYPE_LAYOUT(ERayTracingGeometryInitializerType);
-
-struct FRayTracingGeometrySegment
-{
-	DECLARE_TYPE_LAYOUT(FRayTracingGeometrySegment, NonVirtual);
-public:
-	LAYOUT_FIELD_INITIALIZED(FBufferRHIRef, VertexBuffer, nullptr);
-	LAYOUT_FIELD_INITIALIZED(EVertexElementType, VertexBufferElementType, VET_Float3);
-
-	// Offset in bytes from the base address of the vertex buffer.
-	LAYOUT_FIELD_INITIALIZED(uint32, VertexBufferOffset, 0);
-
-	// Number of bytes between elements of the vertex buffer (sizeof VET_Float3 by default).
-	// Must be equal or greater than the size of the position vector.
-	LAYOUT_FIELD_INITIALIZED(uint32, VertexBufferStride, 12);
-
-	// Number of vertices (positions) in VertexBuffer.
-	// If an index buffer is present, this must be at least the maximum index value in the index buffer + 1.
-	LAYOUT_FIELD_INITIALIZED(uint32, MaxVertices, 0);
-
-	// Primitive range for this segment.
-	LAYOUT_FIELD_INITIALIZED(uint32, FirstPrimitive, 0);
-	LAYOUT_FIELD_INITIALIZED(uint32, NumPrimitives, 0);
-
-	// Indicates whether any-hit shader could be invoked when hitting this geometry segment.
-	// Setting this to `false` turns off any-hit shaders, making the section "opaque" and improving ray tracing performance.
-	LAYOUT_FIELD_INITIALIZED(bool, bForceOpaque, false);
-
-	// Any-hit shader may be invoked multiple times for the same primitive during ray traversal.
-	// Setting this to `false` guarantees that only a single instance of any-hit shader will run per primitive, at some performance cost.
-	LAYOUT_FIELD_INITIALIZED(bool, bAllowDuplicateAnyHitShaderInvocation, true);
-
-	// Indicates whether this section is enabled and should be taken into account during acceleration structure creation
-	LAYOUT_FIELD_INITIALIZED(bool, bEnabled, true);
-};
-
-struct FRayTracingGeometryInitializer
-{
-	DECLARE_EXPORTED_TYPE_LAYOUT(FRayTracingGeometryInitializer, RHI_API, NonVirtual);
-public:
-	LAYOUT_FIELD_INITIALIZED(FBufferRHIRef, IndexBuffer, nullptr);
-
-	// Offset in bytes from the base address of the index buffer.
-	LAYOUT_FIELD_INITIALIZED(uint32, IndexBufferOffset, 0);
-	
-	LAYOUT_FIELD_INITIALIZED(ERayTracingGeometryType, GeometryType, RTGT_Triangles);
-
-	// Total number of primitives in all segments of the geometry. Only used for validation.
-	LAYOUT_FIELD_INITIALIZED(uint32, TotalPrimitiveCount, 0);
-
-	// Partitions of geometry to allow different shader and resource bindings.
-	// All ray tracing geometries must have at least one segment.
-	LAYOUT_FIELD(TMemoryImageArray<FRayTracingGeometrySegment>, Segments);
-
-	// Offline built geometry data. If null, the geometry will be built by the RHI at runtime.
-	LAYOUT_FIELD_INITIALIZED(FResourceArrayInterface*, OfflineData, nullptr);
-
-	// Pointer to an existing ray tracing geometry which the new geometry is built from.
-	LAYOUT_FIELD_INITIALIZED(FRHIRayTracingGeometry*, SourceGeometry, nullptr);
-
-	LAYOUT_FIELD_INITIALIZED(bool, bFastBuild, false);
-	LAYOUT_FIELD_INITIALIZED(bool, bAllowUpdate, false);
-	LAYOUT_FIELD_INITIALIZED(bool, bAllowCompaction, true);
-	LAYOUT_FIELD_INITIALIZED(ERayTracingGeometryInitializerType, Type, ERayTracingGeometryInitializerType::Rendering);
-
-	LAYOUT_FIELD(FName, DebugName);
-};
-
-enum ERayTracingSceneLifetime
-{
-	// Scene may only be used during the frame when it was created.
-	RTSL_SingleFrame,
-
-	// Scene may be constructed once and used in any number of later frames (not currently implemented).
-	// RTSL_MultiFrame,
-};
-
-enum class ERayTracingAccelerationStructureFlags
-{
-	None            = 0,
-	AllowUpdate     = 1 << 0,
-	AllowCompaction = 1 << 1,
-	FastTrace       = 1 << 2,
-	FastBuild       = 1 << 3,
-	MinimizeMemory  = 1 << 4,
-};
-ENUM_CLASS_FLAGS(ERayTracingAccelerationStructureFlags);
-
-struct FRayTracingSceneInitializer
-{
-	TArrayView<FRayTracingGeometryInstance> Instances;
-
-	// This value controls how many elements will be allocated in the shader binding table per geometry segment.
-	// Changing this value allows different hit shaders to be used for different effects.
-	// For example, setting this to 2 allows one hit shader for regular material evaluation and a different one for shadows.
-	// Desired hit shader can be selected by providing appropriate RayContributionToHitGroupIndex to TraceRay() function.
-	// Use ShaderSlot argument in SetRayTracingHitGroup() to assign shaders and resources for specific part of the shder binding table record.
-	uint32 ShaderSlotsPerGeometrySegment = 1;
-
-	// Defines how many different callable shaders with unique resource bindings can be bound to this scene.
-	// Shaders and resources are assigned to slots in the scene using SetRayTracingCallableShader().
-	uint32 NumCallableShaderSlots = 0;
-
-	// At least one miss shader must be present in a ray tracing scene.
-	// Default miss shader is always in slot 0. Default shader must not use local resources.
-	// Custom miss shaders can be bound to other slots using SetRayTracingMissShader().
-	uint32 NumMissShaderSlots = 1;
-
-	// Defines whether data in this scene should persist between frames.
-	// Currently only single-frame lifetime is supported.
-	ERayTracingSceneLifetime Lifetime = RTSL_SingleFrame;
-
-	FName DebugName;
-};
-
-struct FRayTracingSceneInitializer2
-{
-	// Unique list of geometries referenced by all instances in this scene.
-	// Any referenced geometry is kept alive while the scene is alive.
-	TArray<FRayTracingGeometryRHIRef> ReferencedGeometries;
-	// One entry per instance
-	TArray<FRHIRayTracingGeometry*> PerInstanceGeometries;
-	// Exclusive prefix sum of `Instance.NumTransforms` for all instances in this scene. Used to emulate SV_InstanceID in hit shaders.
-	TArray<uint32> BaseInstancePrefixSum;
-	// Exclusive prefix sum of instance geometry segments is used to calculate SBT record address from instance and segment indices.
-	TArray<uint32> SegmentPrefixSum;
-
-	// Total flattened number of ray tracing geometry instances (a single FRayTracingGeometryInstance may represent many).
-	uint32 NumNativeInstances = 0;
-
-	uint32 NumTotalSegments = 0;
-
-	// This value controls how many elements will be allocated in the shader binding table per geometry segment.
-	// Changing this value allows different hit shaders to be used for different effects.
-	// For example, setting this to 2 allows one hit shader for regular material evaluation and a different one for shadows.
-	// Desired hit shader can be selected by providing appropriate RayContributionToHitGroupIndex to TraceRay() function.
-	// Use ShaderSlot argument in SetRayTracingHitGroup() to assign shaders and resources for specific part of the shder binding table record.
-	uint32 ShaderSlotsPerGeometrySegment = 1;
-
-	// Defines how many different callable shaders with unique resource bindings can be bound to this scene.
-	// Shaders and resources are assigned to slots in the scene using SetRayTracingCallableShader().
-	uint32 NumCallableShaderSlots = 0;
-
-	// At least one miss shader must be present in a ray tracing scene.
-	// Default miss shader is always in slot 0. Default shader must not use local resources.
-	// Custom miss shaders can be bound to other slots using SetRayTracingMissShader().
-	uint32 NumMissShaderSlots = 1;
-
-	// Defines whether data in this scene should persist between frames.
-	// Currently only single-frame lifetime is supported.
-	ERayTracingSceneLifetime Lifetime = RTSL_SingleFrame;
-
-	FName DebugName;
-};
-
 struct FShaderResourceViewInitializer
 {
 	struct FBufferShaderResourceViewInitializer
diff --git a/Engine/Source/Runtime/RHI/Public/RHIResources.h b/Engine/Source/Runtime/RHI/Public/RHIResources.h
--- a/Engine/Source/Runtime/RHI/Public/RHIResources.h
+++ b/Engine/Source/Runtime/RHI/Public/RHIResources.h
@@ -1910,6 +1910,240 @@ typedef TRefCountPtr<FRHIRayTracingPipelineState> FRayTracingPipelineStateRHIRef
 // Ray tracing resources
 //
 
+enum class ERayTracingInstanceFlags : uint8
+{
+	None = 0,
+	TriangleCullDisable = 1 << 1, // No back face culling. Triangle is visible from both sides.
+	TriangleCullReverse = 1 << 2, // Makes triangle front-facing if its vertices are counterclockwise from ray origin.
+	ForceOpaque = 1 << 3, // Disable any-hit shader invocation for this instance.
+	ForceNonOpaque = 1 << 4, // Force any-hit shader invocation even if geometries inside the instance were marked opaque.
+};
+ENUM_CLASS_FLAGS(ERayTracingInstanceFlags);
+
+class FRHIRayTracingGeometry;
+/**
+* High level descriptor of one or more instances of a mesh in a ray tracing scene.
+* All instances covered by this descriptor will share shader bindings, but may have different transforms and user data.
+*/
+struct FRayTracingGeometryInstance
+{
+	// TODO: UE-130819 Ref counting is a temporary workaround for a very rare streaming crash.
+	TRefCountPtr<FRHIRayTracingGeometry> GeometryRHI = nullptr;
+
+	// A single physical mesh may be duplicated many times in the scene with different transforms and user data.
+	// All copies share the same shader binding table entries and therefore will have the same material and shader resources.
+	TArrayView<const FMatrix> Transforms;
+
+	TArrayView<const uint32> InstanceSceneDataOffsets;
+
+	// Optional buffer that stores GPU transforms. Used instead of CPU-side transform data.
+	FShaderResourceViewRHIRef GPUTransformsSRV = nullptr;
+
+	// Conservative number of instances. Some of the actual instances may be made inactive if GPU transforms are used.
+	// Must be less or equal to number of entries in Transforms view if CPU transform data is used.
+	// Must be less or equal to number of entries in GPUTransformsSRV if it is non-null.
+	uint32 NumTransforms = 0;
+
+	// Each geometry copy can receive a user-provided integer, which can be used to retrieve extra shader parameters or customize appearance.
+	// This data can be retrieved using GetInstanceUserData() in closest/any hit shaders.
+	// If UserData view is empty, then DefaultUserData value will be used for all instances.
+	// If UserData view is used, then it must have the same number of entries as NumInstances.
+	uint32 DefaultUserData = 0;
+	TArrayView<const uint32> UserData;
+
+	// Each geometry copy can have one bit to make it individually deactivated (removed from TLAS while maintaining hit group indexing). Useful for culling.
+	TArrayView<const uint32> ActivationMask;
+
+	// Mask that will be tested against one provided to TraceRay() in shader code.
+	// If binary AND of instance mask with ray mask is zero, then the instance is considered not intersected / invisible.
+	uint8 Mask = 0xFF;
+
+	// Flags to control triangle back face culling, whether to allow any-hit shaders, etc.
+	ERayTracingInstanceFlags Flags = ERayTracingInstanceFlags::None;
+};
+
+enum ERayTracingGeometryType
+{
+	// Indexed or non-indexed triangle list with fixed function ray intersection.
+	// Vertex buffer must contain vertex positions as VET_Float3.
+	// Vertex stride must be at least 12 bytes, but may be larger to support custom per-vertex data.
+	// Index buffer may be provided for indexed triangle lists. Implicit triangle list is assumed otherwise.
+	RTGT_Triangles,
+
+	// Custom primitive type that requires an intersection shader.
+	// Vertex buffer for procedural geometry must contain one AABB per primitive as {float3 MinXYZ, float3 MaxXYZ}.
+	// Vertex stride must be at least 24 bytes, but may be larger to support custom per-primitive data.
+	// Index buffers can't be used with procedural geometry.
+	RTGT_Procedural,
+};
+DECLARE_INTRINSIC_TYPE_LAYOUT(ERayTracingGeometryType);
+
+enum class ERayTracingGeometryInitializerType
+{
+	// Fully initializes the RayTracingGeometry object: creates underlying buffer and initializes shader parameters.
+	Rendering,
+
+	// Does not create underlying buffer or shader parameters. Used by the streaming system as an object that is streamed into. 
+	StreamingDestination,
+
+	// Creates buffers but does not create shader parameters. Used for intermediate objects in the streaming system.
+	StreamingSource,
+};
+DECLARE_INTRINSIC_TYPE_LAYOUT(ERayTracingGeometryInitializerType);
+
+struct FRayTracingGeometrySegment
+{
+	DECLARE_TYPE_LAYOUT(FRayTracingGeometrySegment, NonVirtual);
+public:
+	LAYOUT_FIELD_INITIALIZED(FBufferRHIRef, VertexBuffer, nullptr);
+	LAYOUT_FIELD_INITIALIZED(EVertexElementType, VertexBufferElementType, VET_Float3);
+
+	// Offset in bytes from the base address of the vertex buffer.
+	LAYOUT_FIELD_INITIALIZED(uint32, VertexBufferOffset, 0);
+
+	// Number of bytes between elements of the vertex buffer (sizeof VET_Float3 by default).
+	// Must be equal or greater than the size of the position vector.
+	LAYOUT_FIELD_INITIALIZED(uint32, VertexBufferStride, 12);
+
+	// Number of vertices (positions) in VertexBuffer.
+	// If an index buffer is present, this must be at least the maximum index value in the index buffer + 1.
+	LAYOUT_FIELD_INITIALIZED(uint32, MaxVertices, 0);
+
+	// Primitive range for this segment.
+	LAYOUT_FIELD_INITIALIZED(uint32, FirstPrimitive, 0);
+	LAYOUT_FIELD_INITIALIZED(uint32, NumPrimitives, 0);
+
+	// Indicates whether any-hit shader could be invoked when hitting this geometry segment.
+	// Setting this to `false` turns off any-hit shaders, making the section "opaque" and improving ray tracing performance.
+	LAYOUT_FIELD_INITIALIZED(bool, bForceOpaque, false);
+
+	// Any-hit shader may be invoked multiple times for the same primitive during ray traversal.
+	// Setting this to `false` guarantees that only a single instance of any-hit shader will run per primitive, at some performance cost.
+	LAYOUT_FIELD_INITIALIZED(bool, bAllowDuplicateAnyHitShaderInvocation, true);
+
+	// Indicates whether this section is enabled and should be taken into account during acceleration structure creation
+	LAYOUT_FIELD_INITIALIZED(bool, bEnabled, true);
+};
+
+struct FRayTracingGeometryInitializer
+{
+	DECLARE_EXPORTED_TYPE_LAYOUT(FRayTracingGeometryInitializer, RHI_API, NonVirtual);
+public:
+	LAYOUT_FIELD_INITIALIZED(FBufferRHIRef, IndexBuffer, nullptr);
+
+	// Offset in bytes from the base address of the index buffer.
+	LAYOUT_FIELD_INITIALIZED(uint32, IndexBufferOffset, 0);
+
+	LAYOUT_FIELD_INITIALIZED(ERayTracingGeometryType, GeometryType, RTGT_Triangles);
+
+	// Total number of primitives in all segments of the geometry. Only used for validation.
+	LAYOUT_FIELD_INITIALIZED(uint32, TotalPrimitiveCount, 0);
+
+	// Partitions of geometry to allow different shader and resource bindings.
+	// All ray tracing geometries must have at least one segment.
+	LAYOUT_FIELD(TMemoryImageArray<FRayTracingGeometrySegment>, Segments);
+
+	// Offline built geometry data. If null, the geometry will be built by the RHI at runtime.
+	LAYOUT_FIELD_INITIALIZED(FResourceArrayInterface*, OfflineData, nullptr);
+
+	// Pointer to an existing ray tracing geometry which the new geometry is built from.
+	LAYOUT_FIELD_INITIALIZED(FRHIRayTracingGeometry*, SourceGeometry, nullptr);
+
+	LAYOUT_FIELD_INITIALIZED(bool, bFastBuild, false);
+	LAYOUT_FIELD_INITIALIZED(bool, bAllowUpdate, false);
+	LAYOUT_FIELD_INITIALIZED(bool, bAllowCompaction, true);
+	LAYOUT_FIELD_INITIALIZED(ERayTracingGeometryInitializerType, Type, ERayTracingGeometryInitializerType::Rendering);
+
+	LAYOUT_FIELD(FName, DebugName);
+};
+
+enum ERayTracingSceneLifetime
+{
+	// Scene may only be used during the frame when it was created.
+	RTSL_SingleFrame,
+
+	// Scene may be constructed once and used in any number of later frames (not currently implemented).
+	// RTSL_MultiFrame,
+};
+
+enum class ERayTracingAccelerationStructureFlags
+{
+	None = 0,
+	AllowUpdate = 1 << 0,
+	AllowCompaction = 1 << 1,
+	FastTrace = 1 << 2,
+	FastBuild = 1 << 3,
+	MinimizeMemory = 1 << 4,
+};
+ENUM_CLASS_FLAGS(ERayTracingAccelerationStructureFlags);
+
+struct FRayTracingSceneInitializer
+{
+	TArrayView<FRayTracingGeometryInstance> Instances;
+
+	// This value controls how many elements will be allocated in the shader binding table per geometry segment.
+	// Changing this value allows different hit shaders to be used for different effects.
+	// For example, setting this to 2 allows one hit shader for regular material evaluation and a different one for shadows.
+	// Desired hit shader can be selected by providing appropriate RayContributionToHitGroupIndex to TraceRay() function.
+	// Use ShaderSlot argument in SetRayTracingHitGroup() to assign shaders and resources for specific part of the shder binding table record.
+	uint32 ShaderSlotsPerGeometrySegment = 1;
+
+	// Defines how many different callable shaders with unique resource bindings can be bound to this scene.
+	// Shaders and resources are assigned to slots in the scene using SetRayTracingCallableShader().
+	uint32 NumCallableShaderSlots = 0;
+
+	// At least one miss shader must be present in a ray tracing scene.
+	// Default miss shader is always in slot 0. Default shader must not use local resources.
+	// Custom miss shaders can be bound to other slots using SetRayTracingMissShader().
+	uint32 NumMissShaderSlots = 1;
+
+	// Defines whether data in this scene should persist between frames.
+	// Currently only single-frame lifetime is supported.
+	ERayTracingSceneLifetime Lifetime = RTSL_SingleFrame;
+
+	FName DebugName;
+};
+
+struct FRayTracingSceneInitializer2
+{
+	// Unique list of geometries referenced by all instances in this scene.
+	// Any referenced geometry is kept alive while the scene is alive.
+	TArray<TRefCountPtr<FRHIRayTracingGeometry>> ReferencedGeometries;
+	// One entry per instance
+	TArray<FRHIRayTracingGeometry*> PerInstanceGeometries;
+	// Exclusive prefix sum of `Instance.NumTransforms` for all instances in this scene. Used to emulate SV_InstanceID in hit shaders.
+	TArray<uint32> BaseInstancePrefixSum;
+	// Exclusive prefix sum of instance geometry segments is used to calculate SBT record address from instance and segment indices.
+	TArray<uint32> SegmentPrefixSum;
+
+	// Total flattened number of ray tracing geometry instances (a single FRayTracingGeometryInstance may represent many).
+	uint32 NumNativeInstances = 0;
+
+	uint32 NumTotalSegments = 0;
+
+	// This value controls how many elements will be allocated in the shader binding table per geometry segment.
+	// Changing this value allows different hit shaders to be used for different effects.
+	// For example, setting this to 2 allows one hit shader for regular material evaluation and a different one for shadows.
+	// Desired hit shader can be selected by providing appropriate RayContributionToHitGroupIndex to TraceRay() function.
+	// Use ShaderSlot argument in SetRayTracingHitGroup() to assign shaders and resources for specific part of the shder binding table record.
+	uint32 ShaderSlotsPerGeometrySegment = 1;
+
+	// Defines how many different callable shaders with unique resource bindings can be bound to this scene.
+	// Shaders and resources are assigned to slots in the scene using SetRayTracingCallableShader().
+	uint32 NumCallableShaderSlots = 0;
+
+	// At least one miss shader must be present in a ray tracing scene.
+	// Default miss shader is always in slot 0. Default shader must not use local resources.
+	// Custom miss shaders can be bound to other slots using SetRayTracingMissShader().
+	uint32 NumMissShaderSlots = 1;
+
+	// Defines whether data in this scene should persist between frames.
+	// Currently only single-frame lifetime is supported.
+	ERayTracingSceneLifetime Lifetime = RTSL_SingleFrame;
+
+	FName DebugName;
+};
+
 struct FRayTracingAccelerationStructureSize
 {
 	uint64 ResultSize = 0;
@@ -1927,23 +2161,41 @@ public:
 	FRHIRayTracingAccelerationStructure() : FRHIResource(RRT_RayTracingAccelerationStructure) {}
 };
 
-struct FRayTracingGeometryInitializer;
 using FRayTracingAccelerationStructureAddress = uint64;
 
 /** Bottom level ray tracing acceleration structure (contains triangles). */
 class FRHIRayTracingGeometry : public FRHIRayTracingAccelerationStructure
 {
 public:
+	FRHIRayTracingGeometry() = default;
+	FRHIRayTracingGeometry(const FRayTracingGeometryInitializer& InInitializer)
+		: Initializer(InInitializer)
+	{}
+
 	virtual FRayTracingAccelerationStructureAddress GetAccelerationStructureAddress(uint64 GPUIndex) const = 0;
-	virtual uint32 GetNumSegments() const = 0;
-	virtual FRayTracingAccelerationStructureSize GetSizeInfo() const = 0;
 	virtual void SetInitializer(const FRayTracingGeometryInitializer& Initializer) = 0;
+
+	const FRayTracingGeometryInitializer& GetInitializer() const
+	{
+		return Initializer;
+	}
+
+	uint32 GetNumSegments() const 
+	{ 
+		return Initializer.Segments.Num(); 
+	}
+
+	FRayTracingAccelerationStructureSize GetSizeInfo() const
+	{
+		return SizeInfo;
+	};
+protected:
+	FRayTracingAccelerationStructureSize SizeInfo = {};
+	FRayTracingGeometryInitializer Initializer = {};
 };
 
 typedef TRefCountPtr<FRHIRayTracingGeometry>     FRayTracingGeometryRHIRef;
 
-struct FRayTracingSceneInitializer2;
-
 /** Top level ray tracing acceleration structure (contains instances of meshes). */
 class FRHIRayTracingScene : public FRHIRayTracingAccelerationStructure
 {
diff --git a/Engine/Source/Runtime/VulkanRHI/Private/VulkanRayTracing.cpp b/Engine/Source/Runtime/VulkanRHI/Private/VulkanRayTracing.cpp
--- a/Engine/Source/Runtime/VulkanRHI/Private/VulkanRayTracing.cpp
+++ b/Engine/Source/Runtime/VulkanRHI/Private/VulkanRayTracing.cpp
@@ -287,8 +287,8 @@ static void GetBLASBuildData(
 		&BuildData.SizesInfo);
 }
 
-FVulkanRayTracingGeometry::FVulkanRayTracingGeometry(FRayTracingGeometryInitializer InInitializer, FVulkanDevice* InDevice)
-	: Device(InDevice), Initializer(MoveTemp(InInitializer))
+FVulkanRayTracingGeometry::FVulkanRayTracingGeometry(const FRayTracingGeometryInitializer& InInitializer, FVulkanDevice* InDevice)
+	: FRHIRayTracingGeometry(InInitializer), Device(InDevice)
 {
 	// Only supporting triangles initially
 	check(Initializer.GeometryType == ERayTracingGeometryType::RTGT_Triangles);
diff --git a/Engine/Source/Runtime/VulkanRHI/Private/VulkanRayTracing.h b/Engine/Source/Runtime/VulkanRHI/Private/VulkanRayTracing.h
--- a/Engine/Source/Runtime/VulkanRHI/Private/VulkanRayTracing.h
+++ b/Engine/Source/Runtime/VulkanRHI/Private/VulkanRayTracing.h
@@ -80,30 +80,22 @@ struct FVkRtBLASBuildData
 class FVulkanRayTracingGeometry : public FRHIRayTracingGeometry
 {
 public:
-	FVulkanRayTracingGeometry(FRayTracingGeometryInitializer Initializer, FVulkanDevice* InDevice);
+	FVulkanRayTracingGeometry(const FRayTracingGeometryInitializer& Initializer, FVulkanDevice* InDevice);
 	~FVulkanRayTracingGeometry();
 
-	virtual FRayTracingAccelerationStructureAddress GetAccelerationStructureAddress(uint64 GPUIndex) const final override { return Address; }
-	virtual uint32 GetNumSegments() const final override { return Initializer.Segments.Num(); }
-	virtual FRayTracingAccelerationStructureSize GetSizeInfo() const final override { return SizeInfo; }
+	virtual FRayTracingAccelerationStructureAddress GetAccelerationStructureAddress(uint64 GPUIndex) const final override { return Address; }	
 	virtual void SetInitializer(const FRayTracingGeometryInitializer& Initializer) final override;
 
 	void Swap(FVulkanRayTracingGeometry& Other);
 	void BuildAccelerationStructure(FVulkanCommandListContext& CommandContext, EAccelerationStructureBuildMode BuildMode);
 
-	const FRayTracingGeometryInitializer& GetInitializer() const { return Initializer; }
-
 private:
 	FVulkanDevice* const Device = nullptr;
 
-	FRayTracingGeometryInitializer Initializer;
-
 	VkAccelerationStructureKHR Handle = VK_NULL_HANDLE;
 	VkDeviceAddress Address = 0;
 	TRefCountPtr<FVulkanResourceMultiBuffer> AccelerationStructureBuffer;
 	TRefCountPtr<FVulkanResourceMultiBuffer> ScratchBuffer;
-	
-	FRayTracingAccelerationStructureSize SizeInfo = {};
 };
 
 class FVulkanRayTracingScene : public FRHIRayTracingScene
