Fixed nopch compile error due missing include

diff --git a/Engine/Source/Runtime/D3D12RHI/Private/D3D12RayTracing.cpp b/Engine/Source/Runtime/D3D12RHI/Private/D3D12RayTracing.cpp
--- a/Engine/Source/Runtime/D3D12RHI/Private/D3D12RayTracing.cpp
+++ b/Engine/Source/Runtime/D3D12RHI/Private/D3D12RayTracing.cpp
@@ -12,6 +12,7 @@
 #include "Hash/CityHash.h"
 #include "HAL/CriticalSection.h"
 #include "HAL/IConsoleManager.h"
+#include "HAL/FileManagerGeneric.h"
 #include "Misc/ScopeLock.h"
 #include "Async/ParallelFor.h"
 #include "Misc/BufferedOutputDevice.h"
