fix uav alias layout to always match source resource

diff --git a/Engine/Source/Runtime/D3D12RHI/Private/D3D12Texture.cpp b/Engine/Source/Runtime/D3D12RHI/Private/D3D12Texture.cpp
--- a/Engine/Source/Runtime/D3D12RHI/Private/D3D12Texture.cpp
+++ b/Engine/Source/Runtime/D3D12RHI/Private/D3D12Texture.cpp
@@ -822,6 +822,9 @@ void CreateUAVAliasResource(FD3D12Adapter* Adapter, D3D12_CLEAR_VALUE* ClearValu
 		AliasTextureDesc.Format = (DXGI_FORMAT)GPixelFormats[AliasTextureFormat].PlatformFormat;
 		AliasTextureDesc.Width = SourceDesc.Width / GPixelFormats[SourceFormat].BlockSizeX;
 		AliasTextureDesc.Height = SourceDesc.Height / GPixelFormats[SourceFormat].BlockSizeY;
+		// layout of UAV must match source resource
+		AliasTextureDesc.Layout = SourceResource->GetResource()->GetDesc().Layout;
+
 		EnumAddFlags(AliasTextureDesc.Flags, D3D12_RESOURCE_FLAG_ALLOW_UNORDERED_ACCESS);
 
 		AliasTextureDesc.UAVAliasPixelFormat = PF_Unknown;
