Forward heap properties to CPU & GPU acces checks to better support custom heap types (Github request)

diff --git a/Engine/Source/Runtime/D3D12RHI/Private/D3D12Resources.cpp b/Engine/Source/Runtime/D3D12RHI/Private/D3D12Resources.cpp
--- a/Engine/Source/Runtime/D3D12RHI/Private/D3D12Resources.cpp
+++ b/Engine/Source/Runtime/D3D12RHI/Private/D3D12Resources.cpp
@@ -513,7 +513,7 @@ HRESULT FD3D12Adapter::CreateCommittedResource(const FD3D12ResourceDesc& InDesc,
 		SetName(*ppOutResource, Name);
 
 		// Only track resources that cannot be accessed on the CPU.
-		if (IsGPUOnly(HeapProps.Type))
+		if (IsGPUOnly(HeapProps.Type, &HeapProps))
 		{
 			(*ppOutResource)->StartTrackingForResidency();
 		}
@@ -964,13 +964,13 @@ void FD3D12ResourceLocation::SetResource(FD3D12Resource* Value)
 }
 
 
-void FD3D12ResourceLocation::AsStandAlone(FD3D12Resource* Resource, uint64 InSize, bool bInIsTransient)
+void FD3D12ResourceLocation::AsStandAlone(FD3D12Resource* Resource, uint64 InSize, bool bInIsTransient, const D3D12_HEAP_PROPERTIES* CustomHeapProperties)
 {
 	SetType(FD3D12ResourceLocation::ResourceLocationType::eStandAlone);
 	SetResource(Resource);
 	SetSize(InSize);
 
-	if (IsCPUAccessible(Resource->GetHeapType()))
+	if (IsCPUAccessible(Resource->GetHeapType(), CustomHeapProperties))
 	{
 		D3D12_RANGE range = { 0, IsCPUWritable(Resource->GetHeapType()) ? 0 : InSize };
 		SetMappedBaseAddress(Resource->Map(&range));
diff --git a/Engine/Source/Runtime/D3D12RHI/Public/D3D12Resources.h b/Engine/Source/Runtime/D3D12RHI/Public/D3D12Resources.h
--- a/Engine/Source/Runtime/D3D12RHI/Public/D3D12Resources.h
+++ b/Engine/Source/Runtime/D3D12RHI/Public/D3D12Resources.h
@@ -595,7 +595,7 @@ public:
 
 	const inline bool IsValid() const { return Type != ResourceLocationType::eUndefined; }
 
-	void AsStandAlone(FD3D12Resource* Resource, uint64 InSize = 0, bool bInIsTransient = false);
+	void AsStandAlone(FD3D12Resource* Resource, uint64 InSize = 0, bool bInIsTransient = false, const D3D12_HEAP_PROPERTIES* CustomHeapProperties = nullptr);
 
 	inline void AsHeapAliased(FD3D12Resource* Resource)
 	{
