- D3D12 RHI resource listener fixes/refactor:
     * make sure the RT geometry object is registered on construction to the used resources so it's informed on stream & stream out events
     * stream out sends regular event with resource location which contains invalid allocated resource data, and nullptr for resource location is only used when the owning resource will be deleted (used by views to remove back reference - not needed by RT geometry objects because these keep refcounted pointers to the used resources and can never reach the rename event on delete - checked)
- Add to do for UE-139474 - validate and optionally recreate the DynamicRayTracingGeometries cached in the FStaticMeshSceneProxy on Stream in/out of the shared index buffer

diff --git a/Engine/Source/Runtime/D3D12RHI/Private/D3D12Buffer.cpp b/Engine/Source/Runtime/D3D12RHI/Private/D3D12Buffer.cpp
--- a/Engine/Source/Runtime/D3D12RHI/Private/D3D12Buffer.cpp
+++ b/Engine/Source/Runtime/D3D12RHI/Private/D3D12Buffer.cpp
@@ -466,9 +466,16 @@ void FD3D12Buffer::RenameLDAChain(FD3D12ResourceLocation& NewLocation)
 void FD3D12Buffer::Swap(FD3D12Buffer& Other)
 {
 	check(!LockedData.bLocked && !Other.LockedData.bLocked);
+
 	FRHIBuffer::Swap(Other);
 	FD3D12BaseShaderResource::Swap(Other);
 	FD3D12LinkedAdapterObject<FD3D12Buffer>::Swap(Other);
+
+	check(IsHeadLink());
+	for (FLinkedObjectIterator NextBuffer(this); NextBuffer; ++NextBuffer)
+	{
+		NextBuffer->ResourceRenamed(&NextBuffer->ResourceLocation);
+	}
 }
 
 void FD3D12Buffer::ReleaseUnderlyingResource()
@@ -485,7 +492,7 @@ void FD3D12Buffer::ReleaseUnderlyingResource()
 	{
 		check(!NextBuffer->LockedData.bLocked && NextBuffer->ResourceLocation.IsValid());
 		NextBuffer->ResourceLocation.Clear();
-		NextBuffer->RemoveAllRenameListeners();
+		NextBuffer->ResourceRenamed(&NextBuffer->ResourceLocation);
 	}
 }
 
diff --git a/Engine/Source/Runtime/D3D12RHI/Private/D3D12RayTracing.cpp b/Engine/Source/Runtime/D3D12RHI/Private/D3D12RayTracing.cpp
--- a/Engine/Source/Runtime/D3D12RHI/Private/D3D12RayTracing.cpp
+++ b/Engine/Source/Runtime/D3D12RHI/Private/D3D12RayTracing.cpp
@@ -3584,6 +3584,11 @@ FD3D12RayTracingGeometry::FD3D12RayTracingGeometry(FD3D12Adapter* Adapter, const
 		
 		Initializer.OfflineData->Discard();
 	}
+
+	for (uint32 GPUIndex = 0; GPUIndex < MAX_NUM_GPUS && GPUIndex < GNumExplicitGPUsForRendering; ++GPUIndex)
+	{
+		RegisterAsRenameListener(GPUIndex);
+	}
 }
 
 void FD3D12RayTracingGeometry::Swap(FD3D12RayTracingGeometry& Other)
@@ -3724,33 +3729,39 @@ void FD3D12RayTracingGeometry::UnregisterAsRenameListener(uint32 InGPUIndex)
 
 void FD3D12RayTracingGeometry::ResourceRenamed(FD3D12BaseShaderResource* InRenamedResource, FD3D12ResourceLocation* InNewResourceLocation)
 {
-	// If there is no valid resource location (released) then release  as listener to the resource and mark the raytracing geometry dirty 
-	// (can't be used anymore since one of its pending resources are invalid now).
-	// Ideally this should not happen with correct order of destruction but it can happen that the RT geometry gets deleted a bit later than the buffers because it's
-	// still referenced by the previous RT scene (but not used anymore). This only happens during freeing of RHI buffers during mesh streaming.
-	if (InNewResourceLocation == nullptr)
-	{	
-		FD3D12Buffer* IndexBuffer = FD3D12DynamicRHI::ResourceCast(Initializer.IndexBuffer.GetReference(), InRenamedResource->GetParentDevice()->GetGPUIndex());
-		if (IndexBuffer == InRenamedResource)
-		{
-			Initializer.IndexBuffer = nullptr;
-		}
-		for (FRayTracingGeometrySegment& Segment : Initializer.Segments)
+	// Empty resource location is used on destruction of the base shader resource but this
+	// shouldn't happen for RT Geometries because it keeps smart pointers to it's resources.
+	check(InNewResourceLocation != nullptr);
+
+	// Recreate the hit group parameters which cache the address to the index and vertex buffers directly if the geometry is fully valid
+	uint32 GPUIndex = InRenamedResource->GetParentDevice()->GetGPUIndex();
+	if (BuffersValid(GPUIndex))
+	{
+		SetupHitGroupSystemParameters(GPUIndex);
+	}
+}
+
+bool FD3D12RayTracingGeometry::BuffersValid(uint32 GPUIndex) const
+{
+	if (Initializer.IndexBuffer)
+	{
+		const FD3D12Buffer* IndexBuffer = FD3D12DynamicRHI::ResourceCast(Initializer.IndexBuffer.GetReference(), GPUIndex);
+		if (!IndexBuffer->ResourceLocation.IsValid())
 		{
-			FD3D12Buffer* VertexBuffer = FD3D12DynamicRHI::ResourceCast(Segment.VertexBuffer.GetReference(), InRenamedResource->GetParentDevice()->GetGPUIndex());
-			if (VertexBuffer == InRenamedResource)
-			{
-				Segment.VertexBuffer = nullptr;
-			}
+			return false;
 		}
-
-		SetDirty(FRHIGPUMask::FromIndex(InRenamedResource->GetParentDevice()->GetGPUIndex()), true);
 	}
-	else
+
+	for (const FRayTracingGeometrySegment& Segment : Segments)
 	{
-		// Recreate the hit group parameters which cache the address to the index and vertex buffers directly
-		SetupHitGroupSystemParameters(InRenamedResource->GetParentDevice()->GetGPUIndex());
+		const FD3D12Buffer* VertexBuffer = FD3D12DynamicRHI::ResourceCast(Segment.VertexBuffer.GetReference(), GPUIndex);
+		if (!VertexBuffer->ResourceLocation.IsValid())
+		{
+			return false;
+		}
 	}
+
+	return true;
 }
 
 void FD3D12RayTracingGeometry::TransitionBuffers(FD3D12CommandContext& CommandContext)
@@ -3833,6 +3844,7 @@ void FD3D12RayTracingGeometry::SetupHitGroupSystemParameters(uint32 InGPUIndex)
 	TArray<FHitGroupSystemParameters>& HitGroupSystemParametersForThisGPU = HitGroupSystemParameters[InGPUIndex];
 	HitGroupSystemParametersForThisGPU.Reset(Initializer.Segments.Num());
 
+	check(BuffersValid(InGPUIndex));
 	FD3D12Buffer* IndexBuffer = FD3D12DynamicRHI::ResourceCast(Initializer.IndexBuffer.GetReference(), InGPUIndex);
 	const uint32 IndexStride = IndexBuffer ? IndexBuffer->GetStride() : 0;
 	for (const FRayTracingGeometrySegment& Segment : Initializer.Segments)
@@ -3866,6 +3878,7 @@ void FD3D12RayTracingGeometry::CreateAccelerationStructureBuildDesc(FD3D12Comman
 
 	// Use the pre-built descs as template and set the GPU resource pointers (current VB/IB).
 	check(OutGeometryDescs.Num() == GeometryDescs.Num());
+	checkf(BuffersValid(GPUIndex), TEXT("Index & vertex buffers should be valid (not streamed out) when building the acceleration structure"));
 
 	FD3D12Buffer* IndexBuffer = CommandContext.RetrieveObject<FD3D12Buffer>(Initializer.IndexBuffer.GetReference());
 	FD3D12Buffer* NullTransformBufferD3D12 = CommandContext.RetrieveObject<FD3D12Buffer>(NullTransformBuffer.GetReference());
@@ -4271,6 +4284,8 @@ void FD3D12RayTracingScene::BuildAccelerationStructure(FD3D12CommandContext& Com
 
 				checkf(!Geometry->IsDirty(CommandContext.GetGPUIndex()),
 					TEXT("Acceleration structures for all geometries must be built before building the top level acceleration structure for the scene."));
+				checkf(Geometry->BuffersValid(CommandContext.GetGPUIndex()),
+					TEXT("Index & vertex buffers for all geometries must be valid (streamed in) when adding geometry to the top level acceleration structure for the scene"));
 
 				AddResidencyHandleForResource(Geometry->AccelerationStructureBuffers[GPUIndex]->GetResource());
 
diff --git a/Engine/Source/Runtime/D3D12RHI/Private/D3D12RayTracing.h b/Engine/Source/Runtime/D3D12RHI/Private/D3D12RayTracing.h
--- a/Engine/Source/Runtime/D3D12RHI/Private/D3D12RayTracing.h
+++ b/Engine/Source/Runtime/D3D12RHI/Private/D3D12RayTracing.h
@@ -62,6 +62,7 @@ public:
 	{
 		return bIsAccelerationStructureDirty[GPUIndex];
 	}
+	bool BuffersValid(uint32 GPUIndex) const;
 
 	using FRHIRayTracingGeometry::Initializer;
 	using FRHIRayTracingGeometry::SizeInfo;
diff --git a/Engine/Source/Runtime/D3D12RHI/Private/D3D12View.h b/Engine/Source/Runtime/D3D12RHI/Private/D3D12View.h
--- a/Engine/Source/Runtime/D3D12RHI/Private/D3D12View.h
+++ b/Engine/Source/Runtime/D3D12RHI/Private/D3D12View.h
@@ -702,6 +702,7 @@ protected:
 	void InitializeInternal(FD3D12BaseShaderResource* InBaseShaderResource, FD3D12ResourceLocation& InResourceLocation)
 	{
 		check(InBaseShaderResource);
+		checkf(BaseShaderResource == nullptr || BaseShaderResource == InBaseShaderResource, TEXT("Either BaseShaderResource is not set yet or when it is it can't change (after rename)"));
 
 		// Only register the first time - init can be called again during rename
 		if (BaseShaderResource == nullptr)
@@ -740,7 +741,18 @@ protected:
 		check(InRenamedResource == BaseShaderResource);
 		if (InNewResourceLocation)
 		{
-			RecreateView();
+			// Only recreate the view if new location is valid
+			if (InNewResourceLocation->IsValid())
+			{
+				RecreateView();
+			}
+			else
+			{
+#if DO_CHECK || USING_CODE_ANALYSIS
+				// Mark as invalid for usage
+				bInitialized = false;
+#endif
+			}
 		}
 		else
 		{
diff --git a/Engine/Source/Runtime/D3D12RHI/Public/D3D12Resources.h b/Engine/Source/Runtime/D3D12RHI/Public/D3D12Resources.h
--- a/Engine/Source/Runtime/D3D12RHI/Public/D3D12Resources.h
+++ b/Engine/Source/Runtime/D3D12RHI/Public/D3D12Resources.h
@@ -838,6 +838,7 @@ public:
 	void AddRenameListener(FD3D12ShaderResourceRenameListener* InRenameListener)
 	{
 		FScopeLock Lock(&RenameListenersCS);
+		check(!RenameListeners.Contains(InRenameListener));
 		RenameListeners.Add(InRenameListener);
 	}
 
@@ -858,7 +859,8 @@ public:
 		ResourceLocation.Swap(Other.ResourceLocation);
 		ResourceLocation.SetOwner(this);
 		::Swap(BufferAlignment, Other.BufferAlignment);
-		::Swap(RenameListeners, Other.RenameListeners);
+
+		// NOTE: Don't swap the rename listeners because these are still referencing the original BaseShaderResource
 	}
 
 	void RemoveAllRenameListeners()
diff --git a/Engine/Source/Runtime/Engine/Private/StaticMesh.cpp b/Engine/Source/Runtime/Engine/Private/StaticMesh.cpp
--- a/Engine/Source/Runtime/Engine/Private/StaticMesh.cpp
+++ b/Engine/Source/Runtime/Engine/Private/StaticMesh.cpp
@@ -1876,6 +1876,8 @@ void UStaticMesh::RequestUpdateCachedRenderState() const
 	{
 		CoarseMeshSM->RequestUpdateCachedRenderState(this);
 	}
+
+	// TODO: Need to mark all DynamicRayTracingGeometries used in FStaticMeshSceneProxy referencing this StaticMesh as either invalid or request a recreation (UE-139474)
 }
 
 
diff --git a/Engine/Source/Runtime/Engine/Private/StaticMeshRender.cpp b/Engine/Source/Runtime/Engine/Private/StaticMeshRender.cpp
--- a/Engine/Source/Runtime/Engine/Private/StaticMeshRender.cpp
+++ b/Engine/Source/Runtime/Engine/Private/StaticMeshRender.cpp
@@ -1723,6 +1723,7 @@ void FStaticMeshSceneProxy::GetDynamicRayTracingInstances(FRayTracingMaterialGat
 		}
 	}
 
+	// TODO: Need to validate that the DynamicRayTracingGeometries are still valid - they could contain streamed out IndexBuffers from the shared StaticMesh (UE-139474)
 	FRayTracingGeometry& Geometry = bEvaluateWPO? DynamicRayTracingGeometries[LODIndex] : RenderData->LODResources[LODIndex].RayTracingGeometry;
 	
 	if (LODModel.GetNumVertices() <= 0 || Geometry.Initializer.TotalPrimitiveCount <= 0)
