- Only add UAV barriers in d3d12 rhi internal transition tracking when the previous state is either unknown or explicitly UAV
- Add stat to track amount of UAV Barriers per frame

diff --git a/Engine/Source/Runtime/D3D12RHI/Private/D3D12Commands.cpp b/Engine/Source/Runtime/D3D12RHI/Private/D3D12Commands.cpp
--- a/Engine/Source/Runtime/D3D12RHI/Private/D3D12Commands.cpp
+++ b/Engine/Source/Runtime/D3D12RHI/Private/D3D12Commands.cpp
@@ -470,19 +470,20 @@ static void HandleResourceTransitions(FD3D12CommandContext& Context, const FD3D1
 			continue;
 		}
 
-		const bool bUAVAccess = EnumHasAnyFlags(Info.AccessAfter, ERHIAccess::UAVMask);
+		const bool bUAVAccessAfter = EnumHasAnyFlags(Info.AccessAfter, ERHIAccess::UAVMask);
 
 		// If targeting all pipelines and UAV access then add UAV barrier even with RHI based transitions
 		if (bDstAllPipelines)
 		{
-			bUAVBarrier |= bUAVAccess;
+			bUAVBarrier |= bUAVAccessAfter;
 		}
 
 		// Use the engine defined transitions?
 		if (GUseInternalTransitions)
 		{
 			// Only need to check for UAV barriers here, all other transitions are handled inside the RHI
-			bUAVBarrier |= bUAVAccess;
+			const bool bUAVAccessOrUnknownBefore = EnumHasAnyFlags(Info.AccessBefore, ERHIAccess::UAVMask) || Info.AccessBefore == ERHIAccess::Unknown;
+			bUAVBarrier |= (bUAVAccessOrUnknownBefore && bUAVAccessAfter);
 		}
 		else
 		{
@@ -490,7 +491,7 @@ static void HandleResourceTransitions(FD3D12CommandContext& Context, const FD3D1
 			if (!Info.Resource)
 			{
 				// Add a UAV barrier in case of no resource and after state is UAV - user probably requested a force UAV barrier here				
-				bUAVBarrier |= bUAVAccess;
+				bUAVBarrier |= bUAVAccessAfter;
 			}
 			else
 			{
diff --git a/Engine/Source/Runtime/D3D12RHI/Private/D3D12StateCache.cpp b/Engine/Source/Runtime/D3D12RHI/Private/D3D12StateCache.cpp
--- a/Engine/Source/Runtime/D3D12RHI/Private/D3D12StateCache.cpp
+++ b/Engine/Source/Runtime/D3D12RHI/Private/D3D12StateCache.cpp
@@ -195,6 +195,8 @@ void FD3D12StateCacheBase::FlushComputeShaderCache(bool bForce)
 	{
 		FD3D12CommandListHandle& CommandList = CmdContext->CommandListHandle;
 		CommandList.AddUAVBarrier();
+
+		INC_DWORD_STAT(STAT_D3D12UAVBarriers);
 	}
 }
 
diff --git a/Engine/Source/Runtime/D3D12RHI/Private/D3D12Stats.h b/Engine/Source/Runtime/D3D12RHI/Private/D3D12Stats.h
--- a/Engine/Source/Runtime/D3D12RHI/Private/D3D12Stats.h
+++ b/Engine/Source/Runtime/D3D12RHI/Private/D3D12Stats.h
@@ -29,6 +29,8 @@ DECLARE_CYCLE_STAT_EXTERN(TEXT("UnlockBuffer time"), STAT_D3D12UnlockBufferTime,
 DECLARE_CYCLE_STAT_EXTERN(TEXT("Commit transient resource time"), STAT_D3D12CommitTransientResourceTime, STATGROUP_D3D12RHI, );
 DECLARE_CYCLE_STAT_EXTERN(TEXT("Decommit transient resource time"), STAT_D3D12DecommitTransientResourceTime, STATGROUP_D3D12RHI, );
 
+DECLARE_DWORD_COUNTER_STAT_EXTERN(TEXT("UAV Barriers"), STAT_D3D12UAVBarriers, STATGROUP_D3D12RHI, );
+
 DECLARE_CYCLE_STAT_EXTERN(TEXT("CreateBoundShaderState time"), STAT_D3D12CreateBoundShaderStateTime, STATGROUP_D3D12RHI, );
 DECLARE_CYCLE_STAT_EXTERN(TEXT("New bound shader state time"), STAT_D3D12NewBoundShaderStateTime, STATGROUP_D3D12RHI, );
 DECLARE_DWORD_ACCUMULATOR_STAT_EXTERN(TEXT("Num bound shader states"), STAT_D3D12NumBoundShaderState, STATGROUP_D3D12RHI, );
diff --git a/Engine/Source/Runtime/D3D12RHI/Private/D3D12Util.cpp b/Engine/Source/Runtime/D3D12RHI/Private/D3D12Util.cpp
--- a/Engine/Source/Runtime/D3D12RHI/Private/D3D12Util.cpp
+++ b/Engine/Source/Runtime/D3D12RHI/Private/D3D12Util.cpp
@@ -1602,6 +1602,8 @@ DEFINE_STAT(STAT_D3D12UnlockBufferTime);
 DEFINE_STAT(STAT_D3D12CommitTransientResourceTime);
 DEFINE_STAT(STAT_D3D12DecommitTransientResourceTime);
 
+DEFINE_STAT(STAT_D3D12UAVBarriers);
+
 DEFINE_STAT(STAT_D3D12NewBoundShaderStateTime);
 DEFINE_STAT(STAT_D3D12CreateBoundShaderStateTime);
 DEFINE_STAT(STAT_D3D12NumBoundShaderState);
