Compile fix after robomerge

diff --git a/Engine/Source/Runtime/D3D12RHI/Private/D3D12RayTracing.cpp b/Engine/Source/Runtime/D3D12RHI/Private/D3D12RayTracing.cpp
--- a/Engine/Source/Runtime/D3D12RHI/Private/D3D12RayTracing.cpp
+++ b/Engine/Source/Runtime/D3D12RHI/Private/D3D12RayTracing.cpp
@@ -3752,7 +3752,7 @@ bool FD3D12RayTracingGeometry::BuffersValid(uint32 GPUIndex) const
 		}
 	}
 
-	for (const FRayTracingGeometrySegment& Segment : Segments)
+	for (const FRayTracingGeometrySegment& Segment : Initializer.Segments)
 	{
 		const FD3D12Buffer* VertexBuffer = FD3D12DynamicRHI::ResourceCast(Segment.VertexBuffer.GetReference(), GPUIndex);
 		if (!VertexBuffer->ResourceLocation.IsValid())
