CVar to crash on GPU OOM

diff --git a/Engine/Source/Runtime/D3D12RHI/Private/D3D12Util.cpp b/Engine/Source/Runtime/D3D12RHI/Private/D3D12Util.cpp
--- a/Engine/Source/Runtime/D3D12RHI/Private/D3D12Util.cpp
+++ b/Engine/Source/Runtime/D3D12RHI/Private/D3D12Util.cpp
@@ -680,7 +680,7 @@ void LogMemoryInfo(FD3D12Adapter* InAdapter)
 
 extern CORE_API bool GIsGPUCrashed;
 
-static void TerminateOnOutOfMemory(HRESULT D3DResult, bool bCreatingTextures)
+static void TerminateOnOutOfMemory(ID3D12Device* InDevice, HRESULT D3DResult, bool bCreatingTextures)
 {
 #if PLATFORM_WINDOWS
 	if (bCreatingTextures)
@@ -694,7 +694,29 @@ static void TerminateOnOutOfMemory(HRESULT D3DResult, bool bCreatingTextures)
 #if STATS
 	GetRendererModule().DebugLogOnCrash();
 #endif
-	FPlatformMisc::RequestExit(true);
+
+	static IConsoleVariable* GPUCrashOOM = IConsoleManager::Get().FindConsoleVariable(TEXT("r.GPUCrashOnOutOfMemory"));
+	if (GPUCrashOOM && GPUCrashOOM->GetInt())
+	{
+		FD3D12DynamicRHI* D3D12RHI = FD3D12DynamicRHI::GetD3DRHI();
+		// If no device provided then try and log the DRED status of each device
+		D3D12RHI->ForEachDevice(InDevice, [&](FD3D12Device* IterationDevice)
+			{
+				if (InDevice == nullptr || InDevice == IterationDevice->GetDevice())
+				{
+					FD3D12Adapter* Adapter = IterationDevice->GetParentAdapter();
+					LogMemoryInfo(Adapter);
+				}
+			});
+
+		UE_LOG(LogD3D12RHI, Fatal, TEXT("Out of video memory trying to allocate a rendering resource"));
+	}
+	else
+	{
+		// Exit silently without reporting a crash because an OOM is not necessarily our fault
+		FPlatformMisc::RequestExit(true);
+	}
+
 #else // PLATFORM_WINDOWS
 	UE_LOG(LogInit, Fatal, TEXT("Out of video memory trying to allocate a rendering resource"));
 #endif // !PLATFORM_WINDOWS
@@ -823,7 +845,7 @@ namespace D3D12RHI
 		
 		if (D3DResult == E_OUTOFMEMORY)
 		{
-			TerminateOnOutOfMemory(D3DResult, false);
+			TerminateOnOutOfMemory(Device, D3DResult, false);
 		}
 		else
 		{
@@ -868,7 +890,7 @@ namespace D3D12RHI
 		}
 		else if (D3DResult == E_OUTOFMEMORY)
 		{
-			TerminateOnOutOfMemory(D3DResult, true);
+			TerminateOnOutOfMemory(Device, D3DResult, true);
 
 #if STATS
 			GetRendererModule().DebugLogOnCrash();
diff --git a/Engine/Source/Runtime/RHI/Private/RHI.cpp b/Engine/Source/Runtime/RHI/Private/RHI.cpp
--- a/Engine/Source/Runtime/RHI/Private/RHI.cpp
+++ b/Engine/Source/Runtime/RHI/Private/RHI.cpp
@@ -1250,6 +1250,13 @@ static TAutoConsoleVariable<int32> CVarGPUCrashDump(
 	ECVF_ReadOnly
 );
 
+static TAutoConsoleVariable<int32> CVarGPUCrashOnOutOfMemory(
+	TEXT("r.GPUCrashOnOutOfMemory"),
+	0,
+	TEXT("Enable crash reporting on GPU OOM"),
+	ECVF_ReadOnly
+);
+
 static TAutoConsoleVariable<int32> CVarGPUCrashDebuggingAftermathMarkers(
 	TEXT("r.GPUCrashDebugging.Aftermath.Markers"),
 	0,
diff --git a/Engine/Source/Runtime/Windows/D3D11RHI/Private/D3D11RHIPrivate.h b/Engine/Source/Runtime/Windows/D3D11RHI/Private/D3D11RHIPrivate.h
--- a/Engine/Source/Runtime/Windows/D3D11RHI/Private/D3D11RHIPrivate.h
+++ b/Engine/Source/Runtime/Windows/D3D11RHI/Private/D3D11RHIPrivate.h
@@ -945,6 +945,8 @@ public:
 		GPUProfilingData.RegisterGPUDispatch(GroupCount);
 	}
 
+	inline const FD3D11Adapter& GetAdapter() const { return Adapter; }
+
 protected:
 	FD3DGPUProfiler GPUProfilingData;
 
diff --git a/Engine/Source/Runtime/Windows/D3D11RHI/Private/D3D11Util.cpp b/Engine/Source/Runtime/Windows/D3D11RHI/Private/D3D11Util.cpp
--- a/Engine/Source/Runtime/Windows/D3D11RHI/Private/D3D11Util.cpp
+++ b/Engine/Source/Runtime/Windows/D3D11RHI/Private/D3D11Util.cpp
@@ -8,6 +8,10 @@
 #include "EngineModule.h"
 #include "RendererInterface.h"
 
+THIRD_PARTY_INCLUDES_START
+#include "dxgi1_4.h"
+THIRD_PARTY_INCLUDES_END
+
 #define D3DERR(x) case x: ErrorCodeText = TEXT(#x); break;
 #define LOCTEXT_NAMESPACE "Developer.MessageLog"
 
@@ -191,6 +195,19 @@ static void TerminateOnDeviceRemoved(HRESULT D3DResult, ID3D11Device* Direct3DDe
 	}
 }
 
+void LogMemoryInfo(const FD3D11Adapter& InAdapter)
+{
+	TRefCountPtr<IDXGIAdapter3> Adapter3;
+	const HRESULT AdapterHR = InAdapter.DXGIAdapter->QueryInterface(IID_PPV_ARGS(Adapter3.GetInitReference()));
+	if (SUCCEEDED(AdapterHR))
+	{
+		DXGI_QUERY_VIDEO_MEMORY_INFO LocalMemoryInfo{};
+		VERIFYD3D11RESULT(Adapter3->QueryVideoMemoryInfo(0, DXGI_MEMORY_SEGMENT_GROUP_LOCAL, &LocalMemoryInfo));
+		UE_LOG(LogD3D11RHI, Error, TEXT("\tBudget:\t%7.2f MB"), LocalMemoryInfo.Budget / (1024.0f * 1024));
+		UE_LOG(LogD3D11RHI, Error, TEXT("\tUsed:\t%7.2f MB"), LocalMemoryInfo.CurrentUsage / (1024.0f * 1024));
+	}
+}
+
 static void TerminateOnOutOfMemory(HRESULT D3DResult, bool bCreatingTextures)
 {
 	if (D3DResult == E_OUTOFMEMORY)
@@ -206,7 +223,16 @@ static void TerminateOnOutOfMemory(HRESULT D3DResult, bool bCreatingTextures)
 #if STATS
 		GetRendererModule().DebugLogOnCrash();
 #endif
-		FPlatformMisc::RequestExit(true);
+		static IConsoleVariable* GPUCrashOOM = IConsoleManager::Get().FindConsoleVariable(TEXT("r.GPUCrashOnOutOfMemory"));
+		if (GPUCrashOOM && GPUCrashOOM->GetInt())
+		{
+			LogMemoryInfo(GD3D11RHI->GetAdapter());
+			UE_LOG(LogD3D11RHI, Fatal, TEXT("Out of video memory trying to allocate a rendering resource"));
+		}
+		else
+		{
+			FPlatformMisc::RequestExit(true);
+		}
 	}
 }
 
