Switching ERHIFeatureLevel:SM6 from SM6.5 to SM6.6 which allows us to use the built-in 64bit atomics.
AMD still has issues with SM6.6 but we can't wait on their fixes much longer.

diff --git a/Engine/Config/Windows/DataDrivenPlatformInfo.ini b/Engine/Config/Windows/DataDrivenPlatformInfo.ini
--- a/Engine/Config/Windows/DataDrivenPlatformInfo.ini
+++ b/Engine/Config/Windows/DataDrivenPlatformInfo.ini
@@ -119,7 +119,8 @@ bSupportsSSDIndirect=true
 bSupportsWaveOperations=true
 bSupportsDxc=true
 bSupportsDiaphragmDOF=true
-bRequiresVendorExtensionsForAtomics=true
+bRequiresVendorExtensionsForAtomics=false
+bSupportsUInt64ImageAtomics=true
 bSupportsAnisotropicMaterials=true
 bSupportsVariableRateShading=true
 bSupportsGen5TemporalAA=true
diff --git a/Engine/Source/Developer/Windows/ShaderFormatD3D/Private/ShaderFormatD3D.h b/Engine/Source/Developer/Windows/ShaderFormatD3D/Private/ShaderFormatD3D.h
--- a/Engine/Source/Developer/Windows/ShaderFormatD3D/Private/ShaderFormatD3D.h
+++ b/Engine/Source/Developer/Windows/ShaderFormatD3D/Private/ShaderFormatD3D.h
@@ -6,7 +6,7 @@
 #include "ShaderCompilerCommon.h"
 
 // TODO: Lock D3D12 to SM 6.6 min spec
-#define USE_SHADER_MODEL_6_6 0
+#define USE_SHADER_MODEL_6_6 1
 
 struct FShaderTarget;
 
diff --git a/Engine/Source/Runtime/D3D12RHI/Private/Windows/WindowsD3D12Device.cpp b/Engine/Source/Runtime/D3D12RHI/Private/Windows/WindowsD3D12Device.cpp
--- a/Engine/Source/Runtime/D3D12RHI/Private/Windows/WindowsD3D12Device.cpp
+++ b/Engine/Source/Runtime/D3D12RHI/Private/Windows/WindowsD3D12Device.cpp
@@ -672,7 +672,7 @@ static bool DoesAnyAdapterSupportSM6(const TArray<TSharedPtr<FD3D12Adapter>>& Ad
 	for (const TSharedPtr<FD3D12Adapter>& Adapter : Adapters)
 	{
 		// TODO: D3D_FEATURE_LEVEL_12_2
-		if (Adapter->GetFeatureLevel() >= D3D_FEATURE_LEVEL_12_1 && Adapter->GetHighestShaderModel() >= D3D_SHADER_MODEL_6_5)
+		if (Adapter->GetFeatureLevel() >= D3D_FEATURE_LEVEL_12_1 && Adapter->GetHighestShaderModel() >= D3D_SHADER_MODEL_6_6)
 		{
 			return true;
 		}
