- Add ERHITransitionAccessMode to let the command list know that all upcoming transitions will be fully defined and don't contain any unknown before states anymore (only then d3d12 can 100% rely on external transitions)
- Revive & fixup of external transition usage in d3d12 - can be conditionally enabled for certain scope (RDG). When swapping between external and internal transitions the command list is flushed to make sure all cached states are valid on the resource
- Fix some incorrect resource transition states for rhi validation

diff --git a/Engine/Source/Runtime/D3D12RHI/Private/D3D12Buffer.cpp b/Engine/Source/Runtime/D3D12RHI/Private/D3D12Buffer.cpp
--- a/Engine/Source/Runtime/D3D12RHI/Private/D3D12Buffer.cpp
+++ b/Engine/Source/Runtime/D3D12RHI/Private/D3D12Buffer.cpp
@@ -265,7 +265,7 @@ FD3D12SyncPoint FD3D12Buffer::UploadResourceDataViaCopyQueue(FResourceArrayInter
 	FD3D12Device* Device = SrcResourceLoc.GetParentDevice();
 	FD3D12CommandAllocatorManager& CommandAllocatorManager = Device->GetTextureStreamingCommandAllocatorManager();
 	FD3D12CommandAllocator* CurrentCommandAllocator = CommandAllocatorManager.ObtainCommandAllocator();
-	FD3D12CommandListHandle hCopyCommandList = Device->GetCopyCommandListManager().ObtainCommandList(*CurrentCommandAllocator);
+	FD3D12CommandListHandle hCopyCommandList = Device->GetCopyCommandListManager().ObtainCommandList(*CurrentCommandAllocator, ED3D12ResourceBarrierTransitionMode::External);
 
 	// Required for stat tracking ?!?
 	hCopyCommandList.SetCurrentOwningContext(&Device->GetDefaultCommandContext());
diff --git a/Engine/Source/Runtime/D3D12RHI/Private/D3D12CommandContext.cpp b/Engine/Source/Runtime/D3D12RHI/Private/D3D12CommandContext.cpp
--- a/Engine/Source/Runtime/D3D12RHI/Private/D3D12CommandContext.cpp
+++ b/Engine/Source/Runtime/D3D12RHI/Private/D3D12CommandContext.cpp
@@ -140,6 +140,9 @@ FD3D12CommandContext::FD3D12CommandContext(FD3D12Device* InParent, ED3D12Command
 	FMemory::Memzero(CurrentRenderTargets);
 	StateCache.Init(GetParentDevice(), this, nullptr);
 	GlobalUniformBuffers.AddZeroed(FUniformBufferStaticSlotRegistry::Get().GetSlotCount());
+
+	// Set the default transition access mode
+	TransitionAccessMode = ERHITransitionAccessMode::Default;
 }
 
 FD3D12CommandContext::~FD3D12CommandContext()
@@ -316,7 +319,8 @@ void FD3D12CommandContext::OpenCommandList()
 	ConditionalObtainCommandAllocator();
 
 	// Get a new command list
-	CommandListHandle = GetCommandListManager().ObtainCommandList(*CommandAllocator);
+	ED3D12ResourceBarrierTransitionMode TransitionMode = TransitionAccessMode == ERHITransitionAccessMode::DisallowUnknown ? ED3D12ResourceBarrierTransitionMode::External: ED3D12ResourceBarrierTransitionMode::Internal;
+	CommandListHandle = GetCommandListManager().ObtainCommandList(*CommandAllocator, TransitionMode);
 	CommandListHandle.SetCurrentOwningContext(this);
 
 	// Notify the descriptor cache about the new command list
@@ -425,6 +429,12 @@ FD3D12CommandListHandle FD3D12CommandContext::FlushCommands(bool WaitForCompleti
 			OpenCommandList();
 		}
 	}
+	else
+	{
+		// just update the transition access mode then if nothing has been enqueued yet (can also close and enqueue empty command but feels a lot less efficient)
+		ED3D12ResourceBarrierTransitionMode TransitionMode = TransitionAccessMode == ERHITransitionAccessMode::DisallowUnknown ? ED3D12ResourceBarrierTransitionMode::External : ED3D12ResourceBarrierTransitionMode::Internal;
+		CommandListHandle.SetTransitionMode(TransitionMode);
+	}
 
 	return CommandListHandle;
 }
diff --git a/Engine/Source/Runtime/D3D12RHI/Private/D3D12CommandContext.h b/Engine/Source/Runtime/D3D12RHI/Private/D3D12CommandContext.h
--- a/Engine/Source/Runtime/D3D12RHI/Private/D3D12CommandContext.h
+++ b/Engine/Source/Runtime/D3D12RHI/Private/D3D12CommandContext.h
@@ -138,6 +138,9 @@ public:
 	// Sync point with copy queue which needs to be checked before kicking this command lists
 	FD3D12SyncPoint CopyQueueSyncPoint;
 
+	// RHITransition access mode - use to set the ED3D12ResourceBarrierTransitionMode on the command list
+	ERHITransitionAccessMode TransitionAccessMode;
+
 	// Current GPU event stack
 	TArray<uint32> GPUEventStack;
 
@@ -262,6 +265,8 @@ public:
 
 	virtual void SetAsyncComputeBudgetInternal(EAsyncComputeBudget Budget) {}
 
+	virtual void RHISetTransitionAccessMode(ERHITransitionAccessMode Mode) final override;
+
 	void RHIBeginTransitionsWithoutFencing(TArrayView<const FRHITransition*> Transitions);
 	virtual void RHIBeginTransitions(TArrayView<const FRHITransition*> Transitions) final override;
 	virtual void RHIEndTransitions(TArrayView<const FRHITransition*> Transitions) final override;
@@ -523,6 +528,10 @@ public:
 	}
 
 	// Special implementation that only signal the fence once.
+	FORCEINLINE virtual void RHISetTransitionAccessMode(ERHITransitionAccessMode Mode) final override
+	{
+		ContextRedirect(RHISetTransitionAccessMode(Mode));
+	}
 	virtual void RHIBeginTransitions(TArrayView<const FRHITransition*> Transitions) final override;
 	virtual void RHIEndTransitions(TArrayView<const FRHITransition*> Transitions) final override;
 
diff --git a/Engine/Source/Runtime/D3D12RHI/Private/D3D12CommandList.cpp b/Engine/Source/Runtime/D3D12RHI/Private/D3D12CommandList.cpp
--- a/Engine/Source/Runtime/D3D12RHI/Private/D3D12CommandList.cpp
+++ b/Engine/Source/Runtime/D3D12RHI/Private/D3D12CommandList.cpp
@@ -95,6 +95,7 @@ FD3D12CommandListHandle::FD3D12CommandListData::FD3D12CommandListData(FD3D12Devi
 	, PendingResourceBarriers()
 	, ResidencySet(nullptr)
 	, CommandListID(GenerateCommandListID())
+	, ResourceBarrierTransitionMode(ED3D12ResourceBarrierTransitionMode::Internal)
 {
 	VERIFYD3D12RESULT(ParentDevice->GetDevice()->CreateCommandList(GetGPUMask().GetNative(), CommandListType, CommandAllocator, nullptr, IID_PPV_ARGS(CommandList.GetInitReference())));
 	INC_DWORD_STAT(STAT_D3D12NumCommandLists);
@@ -210,12 +211,13 @@ void FD3D12CommandListHandle::FD3D12CommandListData::FlushResourceBarriers()
 	ResourceBarrierBatcher.Flush(GetParentDevice(), CommandList, FD3D12DynamicRHI::GetResourceBarrierBatchSizeLimit());
 }
 
-void FD3D12CommandListHandle::FD3D12CommandListData::Reset(FD3D12CommandAllocator& CommandAllocator, bool bTrackExecTime)
+void FD3D12CommandListHandle::FD3D12CommandListData::Reset(FD3D12CommandAllocator& CommandAllocator, ED3D12ResourceBarrierTransitionMode bInTransitionMode, bool bTrackExecTime)
 {
 	LLM_SCOPE_BYNAME(TEXT("D3D12Commandlist"));
 	VERIFYD3D12RESULT(CommandList->Reset(CommandAllocator, nullptr));
 
 	CurrentCommandAllocator = &CommandAllocator;
+	ResourceBarrierTransitionMode = bInTransitionMode;
 	IsClosed = false;
 
 	// Indicate this command allocator is being used.
@@ -255,7 +257,8 @@ int32 FD3D12CommandListHandle::FD3D12CommandListData::CreateAndInsertTimestampQu
 {	
 	FD3D12LinearQueryHeap* QueryHeap = GetParentDevice()->GetCmdListExecTimeQueryHeap();
 	check(QueryHeap);
-	return QueryHeap->EndQuery(this);
+	FD3D12CommandListHandle TempHandle(this);
+	return QueryHeap->EndQuery(TempHandle);
 }
 
 void FD3D12CommandListHandle::FD3D12CommandListData::StartTrackingCommandListTime()
diff --git a/Engine/Source/Runtime/D3D12RHI/Private/D3D12CommandList.h b/Engine/Source/Runtime/D3D12RHI/Private/D3D12CommandList.h
--- a/Engine/Source/Runtime/D3D12RHI/Private/D3D12CommandList.h
+++ b/Engine/Source/Runtime/D3D12RHI/Private/D3D12CommandList.h
@@ -9,6 +9,13 @@ D3D12CommandList.h: Implementation of D3D12 Command List functions
 
 class FD3D12Device;
 
+enum class ED3D12ResourceBarrierTransitionMode
+{
+	Internal,		//< Use the RHI internal transition and resource tracking on the command list
+	External,		//< Use the RHITransition to push the resource transitions (optionally can validate the state via the internal transitions)
+					//< NOTE: only works when all RHITransitions are fully defined and don't contain any ERHIAccess::Unknown states
+};
+
 class FD3D12CommandAllocator : public FNoncopyable
 {
 public:
@@ -65,7 +72,7 @@ private:
 
 		// Reset the command list with a specified command allocator and optional initial state.
 		// Note: Command lists can be reset immediately after they are submitted for execution.
-		void Reset(FD3D12CommandAllocator& CommandAllocator, bool bTrackExecTime);
+		void Reset(FD3D12CommandAllocator& CommandAllocator, ED3D12ResourceBarrierTransitionMode bInTransitionMode, bool bTrackExecTime);
 
 		bool IsComplete(uint64 Generation)
 		{
@@ -251,6 +258,9 @@ private:
 		// Value should be only used for identity, not for synchronization. Valid values are guaranteed to be > 0.
 		uint64 CommandListID;
 
+		// Barrier transition mode used on the command list
+		ED3D12ResourceBarrierTransitionMode ResourceBarrierTransitionMode;
+
 		// Batches resource barriers together until it's explicitly flushed
 		FD3D12ResourceBarrierBatcher ResourceBarrierBatcher;
 
@@ -385,10 +395,10 @@ public:
 
 	// Reset the command list with a specified command allocator and optional initial state.
 	// Note: Command lists can be reset immediately after they are submitted for execution.
-	void Reset(FD3D12CommandAllocator& CommandAllocator, bool bTrackExecTime = false)
+	void Reset(FD3D12CommandAllocator& CommandAllocator, ED3D12ResourceBarrierTransitionMode bInTransitionMode, bool bTrackExecTime)
 	{
 		check(CommandListData);
-		CommandListData->Reset(CommandAllocator, bTrackExecTime);
+		CommandListData->Reset(CommandAllocator, bInTransitionMode, bTrackExecTime);
 	}
 	
 	ID3D12CommandList* CommandList() const
@@ -523,6 +533,18 @@ public:
 		return CommandListData->TrackedResourceState.GetResourceState(pResource);
 	}
 
+	ED3D12ResourceBarrierTransitionMode GetTransitionMode() const
+	{
+		check(CommandListData);
+		return CommandListData->ResourceBarrierTransitionMode;
+	}
+
+	void SetTransitionMode(ED3D12ResourceBarrierTransitionMode InMode)
+	{
+		check(CommandListData);
+		CommandListData->ResourceBarrierTransitionMode = InMode;
+	}
+
 	D3D12RHI_API void AddPendingResourceBarrier(FD3D12Resource* Resource, D3D12_RESOURCE_STATES State, uint32 SubResource);
 
 	TArray<FD3D12PendingResourceBarrier>& PendingResourceBarriers()
diff --git a/Engine/Source/Runtime/D3D12RHI/Private/D3D12Commands.cpp b/Engine/Source/Runtime/D3D12RHI/Private/D3D12Commands.cpp
--- a/Engine/Source/Runtime/D3D12RHI/Private/D3D12Commands.cpp
+++ b/Engine/Source/Runtime/D3D12RHI/Private/D3D12Commands.cpp
@@ -167,9 +167,10 @@ void FD3D12CommandContext::RHIDispatchIndirectComputeShader(FRHIBuffer* Argument
 
 	StateCache.ApplyState<D3D12PT_Compute>();
 
-	// Indirect args buffer can be a previously pending UAV, which becomes PS\Non-PS read. ApplyState will flush pending transitions, so enqueue the indirect
+	// Indirect args buffer can be a previously pending UAV with internal transitions, which becomes PS\Non-PS read. ApplyState will flush pending transitions, so enqueue the indirect
 	// arg transition and flush here.
-	FD3D12DynamicRHI::TransitionResource(CommandListHandle, Location.GetResource(), D3D12_RESOURCE_STATE_TBD, D3D12_RESOURCE_STATE_INDIRECT_ARGUMENT | D3D12_RESOURCE_STATE_NON_PIXEL_SHADER_RESOURCE, D3D12_RESOURCE_BARRIER_ALL_SUBRESOURCES, FD3D12DynamicRHI::ETransitionMode::Validate);
+	D3D12_RESOURCE_STATES IndirectState = (CommandListHandle.GetTransitionMode() == ED3D12ResourceBarrierTransitionMode::Internal) ? D3D12_RESOURCE_STATE_INDIRECT_ARGUMENT | D3D12_RESOURCE_STATE_NON_PIXEL_SHADER_RESOURCE : D3D12_RESOURCE_STATE_INDIRECT_ARGUMENT;
+	FD3D12DynamicRHI::TransitionResource(CommandListHandle, Location.GetResource(), D3D12_RESOURCE_STATE_TBD, IndirectState, D3D12_RESOURCE_BARRIER_ALL_SUBRESOURCES, FD3D12DynamicRHI::ETransitionMode::Validate);
 	CommandListHandle.FlushResourceBarriers();	// Must flush so the desired state is actually set.
 
 	FD3D12Adapter* Adapter = GetParentDevice()->GetParentAdapter();
@@ -478,8 +479,8 @@ static void HandleResourceTransitions(FD3D12CommandContext& Context, const FD3D1
 			bUAVBarrier |= bUAVAccessAfter;
 		}
 
-		// Use the engine defined transitions?
-		if (GUseInternalTransitions)
+		// Use the internal resource tracking
+		if (Context.CommandListHandle.GetTransitionMode() == ED3D12ResourceBarrierTransitionMode::Internal)
 		{
 			// Only need to check for UAV barriers here, all other transitions are handled inside the RHI
 			const bool bUAVAccessOrUnknownBefore = EnumHasAnyFlags(Info.AccessBefore, ERHIAccess::UAVMask) || Info.AccessBefore == ERHIAccess::Unknown;
@@ -502,6 +503,9 @@ static void HandleResourceTransitions(FD3D12CommandContext& Context, const FD3D1
 						return;
 					}
 
+					// All transition should be valid now
+					check(Info.AccessBefore != ERHIAccess::Unknown);
+
 					const bool bIsAsyncCompute = EnumHasAnyFlags(TransitionData->DstPipelines, ERHIPipeline::AsyncCompute);
 					const bool bKnownBeforeState = (Info.AccessBefore != ERHIAccess::Unknown && Info.AccessBefore != ERHIAccess::Discard);
 
@@ -576,11 +580,25 @@ void FD3D12CommandContext::RHIBeginTransitionsWithoutFencing(TArrayView<const FR
 	}
 }
 
+void FD3D12CommandContext::RHISetTransitionAccessMode(ERHITransitionAccessMode Mode)
+{
+	// Nothing changed then don't do anything
+	if (TransitionAccessMode == Mode)
+	{
+		return;
+	}
+
+	TransitionAccessMode = Mode;
+
+	// Flush the current command list and open a new one
+	FlushCommands();
+}
+
 void FD3D12CommandContext::RHIBeginTransitions(TArrayView<const FRHITransition*> Transitions)
 {
-	// Can't correctly handle RHITransitions from multiple threads at the samae time right now - only internal transitions work correctly then
+	// Can't correctly handle RHITransitions from multiple threads at the same time right now - only internal transitions work correctly then
 	// (will be fixed if unknown before state is not allowed anymore)
-	check(bIsDefaultContext || GUseInternalTransitions);
+	check(bIsDefaultContext || CommandListHandle.GetTransitionMode() == ED3D12ResourceBarrierTransitionMode::Internal);
 
 	RHIBeginTransitionsWithoutFencing(Transitions);
 	SignalTransitionFences(Transitions);
@@ -2020,8 +2038,6 @@ void FD3D12CommandContext::RHIDrawIndexedIndirect(FRHIBuffer* IndexBufferRHI, FR
 
 	StateCache.ApplyState<D3D12PT_Graphics>();
 
-	// Indirect args buffer can be a previously pending UAV, which becomes PS\Non-PS read. ApplyState will flush pending transitions, so enqueue the indirect
-	// arg transition and flush here.
 	FD3D12DynamicRHI::TransitionResource(CommandListHandle, Location.GetResource(), D3D12_RESOURCE_STATE_TBD, D3D12_RESOURCE_STATE_INDIRECT_ARGUMENT, D3D12_RESOURCE_BARRIER_ALL_SUBRESOURCES, FD3D12DynamicRHI::ETransitionMode::Validate);
 	CommandListHandle.FlushResourceBarriers();	// Must flush so the desired state is actually set.
 
@@ -2110,8 +2126,6 @@ void FD3D12CommandContext::RHIDrawIndexedPrimitiveIndirect(FRHIBuffer* IndexBuff
 
 	StateCache.ApplyState<D3D12PT_Graphics>();
 
-	// Indirect args buffer can be a previously pending UAV, which becomes PS\Non-PS read. ApplyState will flush pending transitions, so enqueue the indirect
-	// arg transition and flush here.
 	FD3D12DynamicRHI::TransitionResource(CommandListHandle, Location.GetResource(), D3D12_RESOURCE_STATE_TBD, D3D12_RESOURCE_STATE_INDIRECT_ARGUMENT, D3D12_RESOURCE_BARRIER_ALL_SUBRESOURCES, FD3D12DynamicRHI::ETransitionMode::Validate);
 	CommandListHandle.FlushResourceBarriers();	// Must flush so the desired state is actually set.
 
@@ -2176,8 +2190,6 @@ void FD3D12CommandContext::RHIDispatchIndirectMeshShader(FRHIBuffer* ArgumentBuf
 
 	StateCache.ApplyState<D3D12PT_Graphics>();
 
-	// Indirect args buffer can be a previously pending UAV, which becomes PS\Non-PS read. ApplyState will flush pending transitions, so enqueue the indirect
-	// arg transition and flush here.
 	FD3D12DynamicRHI::TransitionResource(CommandListHandle, Location.GetResource(), D3D12_RESOURCE_STATE_TBD, D3D12_RESOURCE_STATE_INDIRECT_ARGUMENT, D3D12_RESOURCE_BARRIER_ALL_SUBRESOURCES, FD3D12DynamicRHI::ETransitionMode::Validate);
 	CommandListHandle.FlushResourceBarriers();	// Must flush so the desired state is actually set.
 
@@ -2485,7 +2497,7 @@ void FD3D12CommandContext::RHIBroadcastTemporalEffect(const FName& InEffectName,
 	FD3D12CommandAllocatorManager& CopyCommandAllocatorManager = Device->GetTextureStreamingCommandAllocatorManager();
 	FD3D12CommandAllocator* CopyCommandAllocator = CopyCommandAllocatorManager.ObtainCommandAllocator();
 	FD3D12CommandListManager& CopyCommandListManager = Device->GetCopyCommandListManager();
-	FD3D12CommandListHandle hCopyCommandList = CopyCommandListManager.ObtainCommandList(*CopyCommandAllocator);
+	FD3D12CommandListHandle hCopyCommandList = CopyCommandListManager.ObtainCommandList(*CopyCommandAllocator, ED3D12ResourceBarrierTransitionMode::Internal);
 	hCopyCommandList.SetCurrentOwningContext(this);
 
 	for (TD3D12Resource* SrcResource : InResources)
diff --git a/Engine/Source/Runtime/D3D12RHI/Private/D3D12DescriptorCache.cpp b/Engine/Source/Runtime/D3D12RHI/Private/D3D12DescriptorCache.cpp
--- a/Engine/Source/Runtime/D3D12RHI/Private/D3D12DescriptorCache.cpp
+++ b/Engine/Source/Runtime/D3D12RHI/Private/D3D12DescriptorCache.cpp
@@ -284,15 +284,22 @@ void FD3D12DescriptorCache::SetVertexBuffers(FD3D12VertexBufferCache& Cache)
 	CmdContext->CommandListHandle.UpdateResidency(Cache.ResidencyHandles, Count);
 	CmdContext->CommandListHandle->IASetVertexBuffers(0, Count, Cache.CurrentVertexBufferViews);
 
-	for (uint32 i = 0; i < Count; ++i)
+	// Only set the state when using internal tracking (we have to trust high level state setup for external transitions)
+	// Problem is that Cache can contain buffers which have already transitioned to another state but are still left in the vertex buffer cache (nothing new bound on that slot yet)
+	// This is 'fine' as long as these vertex buffers are not used but will be bound via above call as vertex buffers with possible incorrect state.
+	// The vertex buffer state should be validated on each draw in RHIValidation to catch incorrect usage then.
+	if (CmdContext->CommandListHandle.GetTransitionMode() == ED3D12ResourceBarrierTransitionMode::Internal)
 	{
-		if (Cache.CurrentVertexBufferResources[i])
+		for (uint32 i = 0; i < Count; ++i)
 		{
-			FD3D12Resource* Resource = Cache.CurrentVertexBufferResources[i]->GetResource();
-			if (Resource && Resource->RequiresResourceStateTracking())
+			if (Cache.CurrentVertexBufferResources[i])
 			{
-				check(Resource->GetSubresourceCount() == 1);
-				FD3D12DynamicRHI::TransitionResource(CmdContext->CommandListHandle, Resource, D3D12_RESOURCE_STATE_TBD, D3D12_RESOURCE_STATE_VERTEX_AND_CONSTANT_BUFFER, D3D12_RESOURCE_BARRIER_ALL_SUBRESOURCES, FD3D12DynamicRHI::ETransitionMode::Validate);
+				FD3D12Resource* Resource = Cache.CurrentVertexBufferResources[i]->GetResource();
+				if (Resource && Resource->RequiresResourceStateTracking())
+				{
+					check(Resource->GetSubresourceCount() == 1);
+					FD3D12DynamicRHI::TransitionResource(CmdContext->CommandListHandle, Resource, D3D12_RESOURCE_STATE_TBD, D3D12_RESOURCE_STATE_VERTEX_AND_CONSTANT_BUFFER, D3D12_RESOURCE_BARRIER_ALL_SUBRESOURCES, FD3D12DynamicRHI::ETransitionMode::Validate);
+				}
 			}
 		}
 	}
@@ -570,7 +577,7 @@ void FD3D12DescriptorCache::SetSRVs(const FD3D12RootSignature* RootSignature, FD
 			D3D12_RESOURCE_STATES State = (ShaderStage == SF_Compute) ? D3D12_RESOURCE_STATE_NON_PIXEL_SHADER_RESOURCE : D3D12_RESOURCE_STATE_NON_PIXEL_SHADER_RESOURCE | D3D12_RESOURCE_STATE_PIXEL_SHADER_RESOURCE;
 			if (SRVs[SlotIndex]->IsDepthStencilResource())
 			{
-				if (GUseInternalTransitions)
+				if (CommandList.GetTransitionMode() == ED3D12ResourceBarrierTransitionMode::Internal)
 				{
 					State |= D3D12_RESOURCE_STATE_DEPTH_READ;
 				}
diff --git a/Engine/Source/Runtime/D3D12RHI/Private/D3D12DirectCommandListManager.cpp b/Engine/Source/Runtime/D3D12RHI/Private/D3D12DirectCommandListManager.cpp
--- a/Engine/Source/Runtime/D3D12RHI/Private/D3D12DirectCommandListManager.cpp
+++ b/Engine/Source/Runtime/D3D12RHI/Private/D3D12DirectCommandListManager.cpp
@@ -505,7 +505,7 @@ FGPUTimingCalibrationTimestamp FD3D12CommandListManager::GetCalibrationTimestamp
 	return Result;
 }
 
-FD3D12CommandListHandle FD3D12CommandListManager::ObtainCommandList(FD3D12CommandAllocator& CommandAllocator, bool bHasBackbufferWriteTransition)
+FD3D12CommandListHandle FD3D12CommandListManager::ObtainCommandList(FD3D12CommandAllocator& CommandAllocator, ED3D12ResourceBarrierTransitionMode bInTransitionMode, bool bHasBackbufferWriteTransition)
 {
 	FD3D12CommandListHandle List;
 	if (!ReadyLists.Dequeue(List))
@@ -515,7 +515,7 @@ FD3D12CommandListHandle FD3D12CommandListManager::ObtainCommandList(FD3D12Comman
 	}
 
 	check(List.GetCommandListType() == CommandListType);
-	List.Reset(CommandAllocator, ShouldTrackCommandListTime() && !(bHasBackbufferWriteTransition && bExcludeBackbufferWriteTransitionTime));
+	List.Reset(CommandAllocator, bInTransitionMode, ShouldTrackCommandListTime() && !(bHasBackbufferWriteTransition && bExcludeBackbufferWriteTransitionTime));
 	return List;
 }
 
@@ -1024,7 +1024,7 @@ uint32 FD3D12CommandListManager::GetResourceBarrierCommandList(FBarrierDescInfo&
 			ResourceBarrierCommandAllocator = ResourceBarrierCommandAllocatorManager.ObtainCommandAllocator();
 		}
 
-		hResourceBarrierList = ObtainCommandList(*ResourceBarrierCommandAllocator, InBarrierDescInfo.BackBufferBarrierDescs.Num() > 0);
+		hResourceBarrierList = ObtainCommandList(*ResourceBarrierCommandAllocator, ED3D12ResourceBarrierTransitionMode::External, InBarrierDescInfo.BackBufferBarrierDescs.Num() > 0);
 
 #if ENABLE_RESIDENCY_MANAGEMENT
 		hResourceBarrierList.UpdateResidency(InBarrierDescInfo.ResidencyHandles.GetData(), InBarrierDescInfo.ResidencyHandles.Num());
diff --git a/Engine/Source/Runtime/D3D12RHI/Private/D3D12DirectCommandListManager.h b/Engine/Source/Runtime/D3D12RHI/Private/D3D12DirectCommandListManager.h
--- a/Engine/Source/Runtime/D3D12RHI/Private/D3D12DirectCommandListManager.h
+++ b/Engine/Source/Runtime/D3D12RHI/Private/D3D12DirectCommandListManager.h
@@ -252,7 +252,7 @@ public:
 	// This use to also take an optional PSO parameter so that we could pass this directly to Create/Reset command lists,
 	// however this was removed as we generally can't actually predict what PSO we'll need until draw due to frequent
 	// state changes. We leave PSOs to always be resolved in ApplyState().
-	FD3D12CommandListHandle ObtainCommandList(FD3D12CommandAllocator& CommandAllocator, bool bHasBackbufferWriteTransition = false);
+	FD3D12CommandListHandle ObtainCommandList(FD3D12CommandAllocator& CommandAllocator, ED3D12ResourceBarrierTransitionMode bInTransitionMode, bool bHasBackbufferWriteTransition = false);
 	void ReleaseCommandList(FD3D12CommandListHandle& hList);
 
 	FD3D12SyncPoint ExecuteCommandListNoCopyQueueSync(FD3D12CommandListHandle& hList, bool WaitForCompletion);
diff --git a/Engine/Source/Runtime/D3D12RHI/Private/D3D12Query.cpp b/Engine/Source/Runtime/D3D12RHI/Private/D3D12Query.cpp
--- a/Engine/Source/Runtime/D3D12RHI/Private/D3D12Query.cpp
+++ b/Engine/Source/Runtime/D3D12RHI/Private/D3D12Query.cpp
@@ -572,7 +572,7 @@ FD3D12LinearQueryHeap::~FD3D12LinearQueryHeap()
 	ReleaseResources();
 }
 
-int32 FD3D12LinearQueryHeap::BeginQuery(FD3D12CommandListHandle CmdListHandle)
+int32 FD3D12LinearQueryHeap::BeginQuery(FD3D12CommandListHandle& CmdListHandle)
 {
 	const int32 SlotIdx = AllocateQueryHeapSlot();
 	const int32 HeapIdx = SlotIdx >> SlotToHeapIdxShift;
@@ -589,7 +589,7 @@ int32 FD3D12LinearQueryHeap::BeginQuery(FD3D12CommandListHandle CmdListHandle)
 	return SlotIdx - TailSlot;
 }
 
-int32 FD3D12LinearQueryHeap::EndQuery(FD3D12CommandListHandle CmdListHandle)
+int32 FD3D12LinearQueryHeap::EndQuery(FD3D12CommandListHandle& CmdListHandle)
 {
 	const int32 SlotIdx = AllocateQueryHeapSlot();
 	const int32 HeapIdx = SlotIdx >> SlotToHeapIdxShift;
diff --git a/Engine/Source/Runtime/D3D12RHI/Private/D3D12Query.h b/Engine/Source/Runtime/D3D12RHI/Private/D3D12Query.h
--- a/Engine/Source/Runtime/D3D12RHI/Private/D3D12Query.h
+++ b/Engine/Source/Runtime/D3D12RHI/Private/D3D12Query.h
@@ -187,14 +187,14 @@ public:
 	 * @param CmdListHandle - a handle to the command list where BeginQuery will be called
 	 * @return index of the allocated query
 	 */
-	int32 BeginQuery(FD3D12CommandListHandle CmdListHandle);
+	int32 BeginQuery(FD3D12CommandListHandle& CmdListHandle);
 
 	/**
 	* Allocate a slot on query heap and queue an EndQuery command to the given list
 	* @param CmdListHandle - a handle to the command list where EndQuery will be called
 	* @return index of the allocated query
 	*/
-	int32 EndQuery(FD3D12CommandListHandle CmdListHandle);
+	int32 EndQuery(FD3D12CommandListHandle& CmdListHandle);
 
 	/**
 	* Resolve new queries and get results for a query batch
diff --git a/Engine/Source/Runtime/D3D12RHI/Private/D3D12RHIPrivate.h b/Engine/Source/Runtime/D3D12RHI/Private/D3D12RHIPrivate.h
--- a/Engine/Source/Runtime/D3D12RHI/Private/D3D12RHIPrivate.h
+++ b/Engine/Source/Runtime/D3D12RHI/Private/D3D12RHIPrivate.h
@@ -149,8 +149,6 @@ typedef FD3D12StateCacheBase FD3D12StateCache;
 #define DEBUG_RHI_EXECUTE_COMMAND_LIST(scope) 
 #endif
 
-// Use the D3D12 RHI internal transitions to drive all resource transitions
-extern bool GUseInternalTransitions;
 // Use the D3D12 RHI internal transitions to validate the engine pushed RHI transitions
 extern bool GValidateInternalTransitions;
 
@@ -629,7 +627,7 @@ public:
 	{
 		// Early out if we are not using engine transitions and not validating them
 		check(InMode == ETransitionMode::Validate);
-		if (!GUseInternalTransitions && !GValidateInternalTransitions)
+		if (hCommandList.GetTransitionMode() == ED3D12ResourceBarrierTransitionMode::External && !GValidateInternalTransitions)
 			return;
 
 		FD3D12Resource* pResource = pView->GetResource();
@@ -664,7 +662,7 @@ public:
 	{
 		// Early out if we are not using engine transitions and not validating them
 		check(InMode == ETransitionMode::Validate);
-		if (!GUseInternalTransitions && !GValidateInternalTransitions)
+		if (hCommandList.GetTransitionMode() == ED3D12ResourceBarrierTransitionMode::External && !GValidateInternalTransitions)
 			return;
 
 		// Determine the required subresource states from the view desc
@@ -698,7 +696,7 @@ public:
 	{
 		// Early out if we are not using engine transitions and not validating them
 		check(InMode == ETransitionMode::Validate);
-		if (!GUseInternalTransitions && !GValidateInternalTransitions)
+		if (hCommandList.GetTransitionMode() == ED3D12ResourceBarrierTransitionMode::External && !GValidateInternalTransitions)
 			return;
 
 		FD3D12Resource* pResource = pView->GetResource();
@@ -740,7 +738,7 @@ public:
 	static inline void TransitionResource(FD3D12CommandListHandle& hCommandList, FD3D12UnorderedAccessView* pView, D3D12_RESOURCE_STATES after, ETransitionMode InMode)
 	{
 		// Early out if we are not using engine transitions and not validating them
-		if (!GUseInternalTransitions && !GValidateInternalTransitions)
+		if (InMode == ETransitionMode::Validate && hCommandList.GetTransitionMode() == ED3D12ResourceBarrierTransitionMode::External && !GValidateInternalTransitions)
 			return;
 
 		FD3D12Resource* pResource = pView->GetResource();
@@ -780,7 +778,7 @@ public:
 	static inline void TransitionResource(FD3D12CommandListHandle& hCommandList, FD3D12ShaderResourceView* pView, D3D12_RESOURCE_STATES after, ETransitionMode InMode)
 	{
 		// Early out if we are not using engine transitions and not validating them
-		if (!GUseInternalTransitions && !GValidateInternalTransitions)
+		if (InMode == ETransitionMode::Validate && hCommandList.GetTransitionMode() == ED3D12ResourceBarrierTransitionMode::External && !GValidateInternalTransitions)
 			return;
 
 		FD3D12Resource* pResource = pView->GetResource();
@@ -821,7 +819,7 @@ public:
 	static inline bool TransitionResource(FD3D12CommandListHandle& hCommandList, FD3D12Resource* pResource, D3D12_RESOURCE_STATES before, D3D12_RESOURCE_STATES after, uint32 subresource, ETransitionMode InMode)
 	{
 		// Early out if we are not using engine transitions and not validating them
-		if (InMode == ETransitionMode::Validate && !GUseInternalTransitions && !GValidateInternalTransitions)
+		if (InMode == ETransitionMode::Validate && hCommandList.GetTransitionMode() == ED3D12ResourceBarrierTransitionMode::External && !GValidateInternalTransitions)
 			return false;
 
 		return TransitionResourceWithTracking(hCommandList, pResource, before, after, subresource, InMode);
@@ -832,7 +830,7 @@ public:
 	static inline bool TransitionResource(FD3D12CommandListHandle& hCommandList, FD3D12Resource* pResource, D3D12_RESOURCE_STATES before, D3D12_RESOURCE_STATES after, const CViewSubresourceSubset& subresourceSubset, ETransitionMode InMode)
 	{
 		// Early out if we are not using engine transitions and not validating them
-		if (InMode == ETransitionMode::Validate && !GUseInternalTransitions && !GValidateInternalTransitions)
+		if (InMode == ETransitionMode::Validate && hCommandList.GetTransitionMode() == ED3D12ResourceBarrierTransitionMode::External && !GValidateInternalTransitions)
 			return false;
 
 		return TransitionResourceWithTracking(hCommandList, pResource, before, after, subresourceSubset, InMode);
@@ -935,7 +933,8 @@ public:
 	static bool ValidateAndSetResourceState(FD3D12CommandListHandle& InCommandList, FD3D12Resource* InResource, CResourceState& InResourceState, uint32 InSubresourceIndex, D3D12_RESOURCE_STATES InBeforeState, D3D12_RESOURCE_STATES InAfterState, bool bInForceAfterState, ETransitionMode InMode)
 	{
 		// Only validate the current state?
-		bool bValidateState = !GUseInternalTransitions && (InMode == ETransitionMode::Validate);
+		bool bUseInternalTransitions = InCommandList.GetTransitionMode() == ED3D12ResourceBarrierTransitionMode::Internal;
+		bool bValidateState = !bUseInternalTransitions && (InMode == ETransitionMode::Validate);
 
 		// Try and get the correct D3D before state for the transition
 		D3D12_RESOURCE_STATES TrackedState = InResourceState.GetSubresourceState(InSubresourceIndex);
@@ -954,7 +953,7 @@ public:
 					BeforeState = InResource->GetResourceState().GetSubresourceState(InSubresourceIndex);
 				}
 			}
-			else if (GUseInternalTransitions)
+			else if (bUseInternalTransitions)
 			{
 				// Already perform transition here if possible to skip patch up during command list execution
 				if (InBeforeState != D3D12_RESOURCE_STATE_TBD)
@@ -981,9 +980,7 @@ public:
 				if (BeforeState == D3D12_RESOURCE_STATE_TBD)
 				{
 					// If we don't have a valid before state, then we have to use the actual last stored state of the
-					// resource. We can sadly enough only correctly do this when parallel command lists don't perform
-					// any resource transition because then the current stored state might be invalid
-					// (Currently validated during begin/end transition in D3D12Commands)
+					// resource. This should only be used to transition from discard state (needs better validation)
 					BeforeState = InResource->GetResourceState().GetSubresourceState(InSubresourceIndex);
 				}
 
@@ -995,7 +992,7 @@ public:
 		bool bRequireUAVBarrier = false;
 
 		// Have a valid state now?
-		check(BeforeState != D3D12_RESOURCE_STATE_TBD || GUseInternalTransitions || bValidateState);
+		check(BeforeState != D3D12_RESOURCE_STATE_TBD || bUseInternalTransitions || bValidateState);
 		if (BeforeState != D3D12_RESOURCE_STATE_TBD)
 		{
 			// Make sure the before states match up or are unknown
@@ -1017,7 +1014,7 @@ public:
 				// We're not using IsTransitionNeeded() when bInForceAfterState is set because we do want to transition even if 'after' is a subset of 'before'
 				// This is so that we can ensure all subresources are in the same state, simplifying future barriers
 				// No state merging when using engine transitions - otherwise next before state might not match up anymore)
-				bool bAllowStateMerging = GUseInternalTransitions;
+				bool bAllowStateMerging = bUseInternalTransitions;
 				if ((bInForceAfterState && BeforeState != InAfterState) || IsTransitionNeeded(bAllowStateMerging, BeforeState, InAfterState))
 				{
 					InCommandList.AddTransitionBarrier(InResource, BeforeState, InAfterState, InSubresourceIndex);
@@ -1264,7 +1261,7 @@ public:
 		else
 		{
 			// If we are not using the internal transitions and need to apply the state change, then store the current state of restore
-			if (!GUseInternalTransitions && InTransitionMode == FD3D12DynamicRHI::ETransitionMode::Apply)
+			if (hCommandList.GetTransitionMode() == ED3D12ResourceBarrierTransitionMode::External && InTransitionMode == FD3D12DynamicRHI::ETransitionMode::Apply)
 			{
 				// try tracked state in command list
 				CResourceState& ResourceState = hCommandList.GetResourceState(pInResource);
diff --git a/Engine/Source/Runtime/D3D12RHI/Private/D3D12Resources.cpp b/Engine/Source/Runtime/D3D12RHI/Private/D3D12Resources.cpp
--- a/Engine/Source/Runtime/D3D12RHI/Private/D3D12Resources.cpp
+++ b/Engine/Source/Runtime/D3D12RHI/Private/D3D12Resources.cpp
@@ -1026,8 +1026,6 @@ bool FD3D12ResourceLocation::OnAllocationMoved(FRHIPoolAllocationData* InNewData
 			else
 			{
 				// Force in the readable state when there are different states (or use just state of subresource 0?)
-				// ideally restore all subresources in correct state? needed for when not using GUseInternalTransitions anymore because then it needs to match the set state from the engine
-				check(GUseInternalTransitions);
 				CreateState = CurrentResource->GetReadableState();
 			}
 			ResourceStateMode = ED3D12ResourceStateMode::MultiState;
diff --git a/Engine/Source/Runtime/D3D12RHI/Private/D3D12Texture.cpp b/Engine/Source/Runtime/D3D12RHI/Private/D3D12Texture.cpp
--- a/Engine/Source/Runtime/D3D12RHI/Private/D3D12Texture.cpp
+++ b/Engine/Source/Runtime/D3D12RHI/Private/D3D12Texture.cpp
@@ -1588,7 +1588,7 @@ FTexture2DRHIRef FD3D12DynamicRHI::RHIAsyncCreateTexture2D(uint32 SizeX, uint32
 
 			FD3D12CommandAllocatorManager& CommandAllocatorManager = Device->GetTextureStreamingCommandAllocatorManager();
 			FD3D12CommandAllocator* CurrentCommandAllocator = CommandAllocatorManager.ObtainCommandAllocator();
-			FD3D12CommandListHandle hCopyCommandList = Device->GetCopyCommandListManager().ObtainCommandList(*CurrentCommandAllocator);
+			FD3D12CommandListHandle hCopyCommandList = Device->GetCopyCommandListManager().ObtainCommandList(*CurrentCommandAllocator, ED3D12ResourceBarrierTransitionMode::Internal);
 			hCopyCommandList.SetCurrentOwningContext(&Device->GetDefaultCommandContext());
 
 			// NB: Do not increment numCopies because that will count as work on the direct
diff --git a/Engine/Source/Runtime/RHI/Private/RHIValidation.cpp b/Engine/Source/Runtime/RHI/Private/RHIValidation.cpp
--- a/Engine/Source/Runtime/RHI/Private/RHIValidation.cpp
+++ b/Engine/Source/Runtime/RHI/Private/RHIValidation.cpp
@@ -888,6 +888,25 @@ namespace RHIValidation
 			*GetRHIAccessName(CurrentStateFromRHI.Access));
 	}
 
+	static inline FString GetReasonString_UnallowedUnknownPreviousState(
+		FResource* Resource, FSubresourceIndex const& SubresourceIndex,
+		const FState& CurrentState,
+		const FState& CurrentStateFromRHI,
+		void* CreateTrace, void* BeginTrace)
+	{
+		FString DebugName = GetResourceDebugName(Resource, SubresourceIndex);
+		return FString::Printf(
+			BARRIER_TRACKER_LOG_PREFIX_RESNAME
+			TEXT("The transition for resource %s contains ERHIAcces::Unknown state while the command list does not allow unknown states at this point.\n See ERHITransitionAccessMode for more information.\n")
+			TEXT("    --- The tracked current previous state is: \"%s\"\n")
+			TEXT("%s")
+			BARRIER_TRACKER_LOG_SUFFIX,
+			*DebugName,
+			*DebugName,
+			*GetRHIAccessName(CurrentState.Access),
+			*GetReasonString_BeginBacktrace(CreateTrace, BeginTrace));
+	}
+
 	static inline FString GetReasonString_IncorrectPreviousState(
 		FResource* Resource, FSubresourceIndex const& SubresourceIndex,
 		const FState& CurrentState,
@@ -1101,7 +1120,7 @@ namespace RHIValidation
 		}
 	}
 
-	void FSubresourceState::BeginTransition(FResource* Resource, FSubresourceIndex const& SubresourceIndex, const FState& CurrentStateFromRHI, const FState& TargetState, EResourceTransitionFlags NewFlags, ERHIPipeline ExecutingPipeline, void* CreateTrace)
+	void FSubresourceState::BeginTransition(FResource* Resource, FSubresourceIndex const& SubresourceIndex, const FState& CurrentStateFromRHI, const FState& TargetState, EResourceTransitionFlags NewFlags, ERHIPipeline ExecutingPipeline, ERHITransitionAccessMode TransitionAccessMode, void* CreateTrace)
 	{
 		FPipelineState& State = States[ExecutingPipeline];
 
@@ -1133,6 +1152,9 @@ namespace RHIValidation
 			// Check for the correct pipeline
 			RHI_VALIDATION_CHECK(EnumHasAllFlags(CurrentStateFromRHI.Pipelines, ExecutingPipeline), *GetReasonString_WrongPipeline(Resource, SubresourceIndex, State.Current, TargetState));
 
+			// Check we don't have an unknown previous state when transition access mode does't allow it
+			RHI_VALIDATION_CHECK(TransitionAccessMode == ERHITransitionAccessMode::AllowUnknown || CurrentStateFromRHI.Access != ERHIAccess::Unknown, *GetReasonString_UnallowedUnknownPreviousState(Resource, SubresourceIndex, State.Previous, CurrentStateFromRHI, State.CreateTransitionBacktrace, BeginTrace));
+
 			// Check the current RHI state passed in matches the tracked state for the resource, or is "unknown".
 			RHI_VALIDATION_CHECK(CurrentStateFromRHI.Access == ERHIAccess::Unknown || (CurrentStateFromRHI.Access == State.Previous.Access && CurrentStateFromRHI.Pipelines == State.Previous.Pipelines),
 				*GetReasonString_IncorrectPreviousState(Resource, SubresourceIndex, State.Previous, CurrentStateFromRHI));
@@ -1300,7 +1322,7 @@ namespace RHIValidation
 		}
 	}
 
-	RHI_API EReplayStatus FOperation::Replay(ERHIPipeline Pipeline, bool& bAllowAllUAVsOverlap) const
+	RHI_API EReplayStatus FOperation::Replay(ERHIPipeline Pipeline, bool& bAllowAllUAVsOverlap, ERHITransitionAccessMode& TransitionAccessMode) const
 	{
 		switch (Type)
 		{
@@ -1311,7 +1333,7 @@ namespace RHIValidation
 			break;
 
 		case EOpType::BeginTransition:
-			Data_BeginTransition.Identity.Resource->EnumerateSubresources(Data_BeginTransition.Identity.SubresourceRange, [this, Pipeline](FSubresourceState& State, FSubresourceIndex const& SubresourceIndex)
+			Data_BeginTransition.Identity.Resource->EnumerateSubresources(Data_BeginTransition.Identity.SubresourceRange, [this, Pipeline, TransitionAccessMode](FSubresourceState& State, FSubresourceIndex const& SubresourceIndex)
 			{
 				State.BeginTransition(
 					Data_BeginTransition.Identity.Resource,
@@ -1320,6 +1342,7 @@ namespace RHIValidation
 					Data_BeginTransition.NextState,
 					Data_BeginTransition.Flags,
 					Pipeline,
+					TransitionAccessMode,
 					Data_BeginTransition.CreateBacktrace);
 
 			}, true);
@@ -1327,7 +1350,7 @@ namespace RHIValidation
 			break;
 
 		case EOpType::EndTransition:
-			Data_EndTransition.Identity.Resource->EnumerateSubresources(Data_EndTransition.Identity.SubresourceRange, [this, Pipeline](FSubresourceState& State, FSubresourceIndex const& SubresourceIndex)
+			Data_EndTransition.Identity.Resource->EnumerateSubresources(Data_EndTransition.Identity.SubresourceRange, [this, Pipeline, TransitionAccessMode](FSubresourceState& State, FSubresourceIndex const& SubresourceIndex)
 			{
 				State.EndTransition(
 					Data_EndTransition.Identity.Resource,
@@ -1411,6 +1434,10 @@ namespace RHIValidation
 			});
 			Data_SpecificUAVOverlap.Identity.Resource->ReleaseOpRef();
 			break;
+
+		case EOpType::SetTransitionAccessMode:
+			TransitionAccessMode = Data_SetTransitionAccessMode.TransitionAccessMode;
+			break;
 		}
 
 		return EReplayStatus::Normal;
@@ -1433,13 +1460,13 @@ namespace RHIValidation
 				FOpQueueState& CurrentQueue = OpQueues[CurrentIndex];
 				if (CurrentQueue.bWaiting)
 				{
-					Status = CurrentQueue.Ops.Replay(CurrentPipeline, CurrentQueue.bAllowAllUAVsOverlap);
+					Status = CurrentQueue.Ops.Replay(CurrentPipeline, CurrentQueue.bAllowAllUAVsOverlap, CurrentQueue.TransitionAccessMode);
 					if (!EnumHasAllFlags(Status, EReplayStatus::Waiting))
 					{
 						CurrentQueue.Ops.Reset();
 						if (CurrentIndex == DstOpQueueIndex && InOpsList.Incomplete())
 						{
-							Status |= InOpsList.Replay(CurrentPipeline, CurrentQueue.bAllowAllUAVsOverlap);
+							Status |= InOpsList.Replay(CurrentPipeline, CurrentQueue.bAllowAllUAVsOverlap, CurrentQueue.TransitionAccessMode);
 							CurrentQueue.bWaiting = InOpsList.Incomplete();
 						}
 						else
diff --git a/Engine/Source/Runtime/RHI/Public/RHI.h b/Engine/Source/Runtime/RHI/Public/RHI.h
--- a/Engine/Source/Runtime/RHI/Public/RHI.h
+++ b/Engine/Source/Runtime/RHI/Public/RHI.h
@@ -1917,6 +1917,18 @@ inline constexpr bool IsValidAccess(ERHIAccess Access)
 	return !IsInvalidAccess(Access);
 }
 
+enum class ERHITransitionAccessMode
+{
+	// All upcoming rhi transitions can still have Unkown before state - rhi might need to do internal tracking if needed
+	AllowUnknown = 0,
+
+	// All upcoming rhi transitions are fully defined - no internal RHI tracking required anymore
+	DisallowUnknown,
+
+	Default = AllowUnknown
+};
+ENUM_CLASS_FLAGS(ERHITransitionAccessMode);
+
 enum class ERHITransitionCreateFlags
 {
 	None = 0,
@@ -2646,7 +2658,7 @@ inline void RHIValidation::FTracker::AddOp(const RHIValidation::FOperation& Op)
 	if (GRHICommandList.Bypass() && CurrentList.Operations.Num() == 0)
 	{
 		auto& OpQueue = OpQueues[GetOpQueueIndex(Pipeline)];
-		if (!EnumHasAllFlags(Op.Replay(Pipeline, OpQueue.bAllowAllUAVsOverlap), EReplayStatus::Waiting))
+		if (!EnumHasAllFlags(Op.Replay(Pipeline, OpQueue.bAllowAllUAVsOverlap, OpQueue.TransitionAccessMode), EReplayStatus::Waiting))
 		{
 			return;
 		}
diff --git a/Engine/Source/Runtime/RHI/Public/RHICommandList.h b/Engine/Source/Runtime/RHI/Public/RHICommandList.h
--- a/Engine/Source/Runtime/RHI/Public/RHICommandList.h
+++ b/Engine/Source/Runtime/RHI/Public/RHICommandList.h
@@ -1504,6 +1504,17 @@ FRHICOMMAND_MACRO(FRHICommandClearUAVUint)
 	RHI_API void Execute(FRHICommandListBase& CmdList);
 };
 
+FRHICOMMAND_MACRO(FRHICommandSetTransitionAccessMode)
+{
+	ERHITransitionAccessMode TransitionAccessMode;
+
+	FORCEINLINE_DEBUGGABLE FRHICommandSetTransitionAccessMode(ERHITransitionAccessMode InTransitionAccessMode)
+		: TransitionAccessMode(InTransitionAccessMode)
+	{
+	}
+	RHI_API void Execute(FRHICommandListBase & CmdList);
+};
+
 FRHICOMMAND_MACRO(FRHICommandCopyToResolveTarget)
 {
 	FResolveParams ResolveParams;
@@ -2621,6 +2632,16 @@ public:
 		ALLOC_COMMAND(FRHICommandClearUAVUint)(UnorderedAccessViewRHI, Values);
 	}
 
+	FORCEINLINE_DEBUGGABLE void SetTransitionAccessMode(ERHITransitionAccessMode Mode)
+	{
+		if (Bypass())
+		{
+			GetComputeContext().RHISetTransitionAccessMode(Mode); 
+			return;
+		}
+		ALLOC_COMMAND(FRHICommandSetTransitionAccessMode)(Mode);
+	}
+
 	FORCEINLINE_DEBUGGABLE void BeginTransitions(TArrayView<const FRHITransition*> Transitions)
 	{
 		if (Bypass())
diff --git a/Engine/Source/Runtime/RHI/Public/RHICommandListCommandExecutes.inl b/Engine/Source/Runtime/RHI/Public/RHICommandListCommandExecutes.inl
--- a/Engine/Source/Runtime/RHI/Public/RHICommandListCommandExecutes.inl
+++ b/Engine/Source/Runtime/RHI/Public/RHICommandListCommandExecutes.inl
@@ -399,6 +399,12 @@ void FRHICommandResummarizeHTile::Execute(FRHICommandListBase& CmdList)
 	INTERNAL_DECORATOR(RHIResummarizeHTile)(DepthTexture);
 }
 
+void FRHICommandSetTransitionAccessMode::Execute(FRHICommandListBase& CmdList)
+{
+	RHISTAT(SetTransitionAccessMode);
+	INTERNAL_DECORATOR_COMPUTE(RHISetTransitionAccessMode)(TransitionAccessMode);
+}
+
 void FRHICommandBeginTransitions::Execute(FRHICommandListBase& CmdList)
 {
 	RHISTAT(BeginTransitions);
diff --git a/Engine/Source/Runtime/RHI/Public/RHIContext.h b/Engine/Source/Runtime/RHI/Public/RHIContext.h
--- a/Engine/Source/Runtime/RHI/Public/RHIContext.h
+++ b/Engine/Source/Runtime/RHI/Public/RHIContext.h
@@ -185,6 +185,8 @@ public:
 
 	virtual void RHISetAsyncComputeBudget(EAsyncComputeBudget Budget) {}
 
+	virtual void RHISetTransitionAccessMode(ERHITransitionAccessMode Mode) {}
+
 	virtual void RHIBeginTransitions(TArrayView<const FRHITransition*> Transitions) = 0;
 
 	virtual void RHIEndTransitions(TArrayView<const FRHITransition*> Transitions) = 0;
diff --git a/Engine/Source/Runtime/RHI/Public/RHIValidationCommon.h b/Engine/Source/Runtime/RHI/Public/RHIValidationCommon.h
--- a/Engine/Source/Runtime/RHI/Public/RHIValidationCommon.h
+++ b/Engine/Source/Runtime/RHI/Public/RHIValidationCommon.h
@@ -172,7 +172,7 @@ namespace RHIValidation
 
 		TRHIPipelineArray<FPipelineState> States;
 
-		void BeginTransition   (FResource* Resource, FSubresourceIndex const& SubresourceIndex, const FState& CurrentStateFromRHI, const FState& TargetState, EResourceTransitionFlags NewFlags, ERHIPipeline Pipeline, void* CreateTrace);
+		void BeginTransition   (FResource* Resource, FSubresourceIndex const& SubresourceIndex, const FState& CurrentStateFromRHI, const FState& TargetState, EResourceTransitionFlags NewFlags, ERHIPipeline Pipeline, ERHITransitionAccessMode TransitionAccessMode, void* CreateTrace);
 		void EndTransition     (FResource* Resource, FSubresourceIndex const& SubresourceIndex, const FState& CurrentStateFromRHI, const FState& TargetState, ERHIPipeline Pipeline, void* CreateTrace);
 		void Assert            (FResource* Resource, FSubresourceIndex const& SubresourceIndex, const FState& RequiredState, bool bAllowAllUAVsOverlap);
 		void SpecificUAVOverlap(FResource* Resource, FSubresourceIndex const& SubresourceIndex, ERHIPipeline Pipeline, bool bAllow);
@@ -551,7 +551,8 @@ namespace RHIValidation
 		Signal,
 		Wait,
 		AllUAVsOverlap,
-		SpecificUAVOverlap
+		SpecificUAVOverlap,
+		SetTransitionAccessMode,
 	};
 
 	struct FUniformBufferResource
@@ -642,9 +643,14 @@ namespace RHIValidation
 				FResourceIdentity Identity;
 				bool bAllow;
 			} Data_SpecificUAVOverlap;
+
+			struct
+			{
+				ERHITransitionAccessMode TransitionAccessMode;
+			} Data_SetTransitionAccessMode;
 		};
 
-		RHI_API EReplayStatus Replay(ERHIPipeline Pipeline, bool& bAllowAllUAVsOverlap) const;
+		RHI_API EReplayStatus Replay(ERHIPipeline Pipeline, bool& bAllowAllUAVsOverlap, ERHITransitionAccessMode& TransitionAccessMode) const;
 
 		static inline FOperation BeginTransitionResource(FResourceIdentity Identity, FState PreviousState, FState NextState, EResourceTransitionFlags Flags, void* CreateBacktrace)
 		{
@@ -784,6 +790,14 @@ namespace RHIValidation
 			Op.Data_SpecificUAVOverlap.bAllow = bAllow;
 			return MoveTemp(Op);
 		}
+
+		static inline FOperation SetTransitionAccessMode(ERHITransitionAccessMode Mode)
+		{
+			FOperation Op;
+			Op.Type = EOpType::SetTransitionAccessMode;
+			Op.Data_SetTransitionAccessMode.TransitionAccessMode = Mode;
+			return MoveTemp(Op);
+		}
 	};
 
 	struct FOperationsList
@@ -791,12 +805,12 @@ namespace RHIValidation
 		TArray<FOperation> Operations;
 		int32 OperationPos = 0;
 
-		inline EReplayStatus Replay(ERHIPipeline Pipeline, bool& bAllowAllUAVsOverlap)
+		inline EReplayStatus Replay(ERHIPipeline Pipeline, bool& bAllowAllUAVsOverlap, ERHITransitionAccessMode& TransitionAccessMode)
 		{
 			EReplayStatus Status = EReplayStatus::Normal;
 			for (; OperationPos < Operations.Num(); ++OperationPos)
 			{
-				Status |= Operations[OperationPos].Replay(Pipeline, bAllowAllUAVsOverlap);
+				Status |= Operations[OperationPos].Replay(Pipeline, bAllowAllUAVsOverlap, TransitionAccessMode);
 				if (EnumHasAllFlags(Status, EReplayStatus::Waiting))
 				{
 					break;
@@ -960,6 +974,11 @@ namespace RHIValidation
 			AddOp(FOperation::SpecificUAVOverlap(Identity, bAllow));
 		}
 
+		inline void SetTransitionAccessMode(ERHITransitionAccessMode Mode)
+		{
+			AddOp(FOperation::SetTransitionAccessMode(Mode));
+		}
+
 		inline void Dispatch()
 		{
 			UAVTrackers[int32(EUAVMode::Compute)].DrawOrDispatch(this, FState(ERHIAccess::UAVCompute, Pipeline));
@@ -1010,6 +1029,7 @@ namespace RHIValidation
 		{
 			bool bWaiting = false;
 			bool bAllowAllUAVsOverlap = false;
+			ERHITransitionAccessMode TransitionAccessMode = ERHITransitionAccessMode::Default;
 
 			FOperationsList Ops;
 		} static RHI_API OpQueues[int32(ERHIPipeline::Num)];
diff --git a/Engine/Source/Runtime/RHI/Public/RHIValidationContext.h b/Engine/Source/Runtime/RHI/Public/RHIValidationContext.h
--- a/Engine/Source/Runtime/RHI/Public/RHIValidationContext.h
+++ b/Engine/Source/Runtime/RHI/Public/RHIValidationContext.h
@@ -56,6 +56,12 @@ public:
 		RHIContext->RHISetAsyncComputeBudget(Budget);
 	}
 
+	virtual void RHISetTransitionAccessMode(ERHITransitionAccessMode Mode) override final
+	{
+		Tracker->SetTransitionAccessMode(Mode);
+		RHIContext->RHISetTransitionAccessMode(Mode);
+	}
+
 	virtual void RHIBeginTransitions(TArrayView<const FRHITransition*> Transitions) override final
 	{
 		for (const FRHITransition* Transition : Transitions)
@@ -389,6 +395,12 @@ public:
 		RHIContext->RHIResummarizeHTile(DepthTexture);
 	}
 
+	virtual void RHISetTransitionAccessMode(ERHITransitionAccessMode Mode) override final
+	{
+		Tracker->SetTransitionAccessMode(Mode);
+		RHIContext->RHISetTransitionAccessMode(Mode);
+	}
+
 	virtual void RHIBeginTransitions(TArrayView<const FRHITransition*> Transitions) override final
 	{
 		ensureMsgf(!State.bInsideBeginRenderPass, TEXT("Beginning a transition within a renderpass is not supported!"));
