Fix reading uninitialized memory during ray tracing sampler descriptor table deduplication.

When exhaustive sampler deduplication method is used, the CPU-side "shadow" copy of the sampler descriptor heap is used to search for matches.
However, the table is filled from multiple threads where an atomic linear allocator is used to give slots to worker threads to write into.
The table memory was not cleared and therefore it was possible for one of the worker threads to find a false positive match in the garbage just by chance (extremely unlikely, and yet it does happen due to memory reuse).

The solution is to always zero-initialize the descriptor memory before writing the final descriptor data into it.

diff --git a/Engine/Source/Runtime/D3D12RHI/Private/D3D12RayTracing.cpp b/Engine/Source/Runtime/D3D12RHI/Private/D3D12RayTracing.cpp
--- a/Engine/Source/Runtime/D3D12RHI/Private/D3D12RayTracing.cpp
+++ b/Engine/Source/Runtime/D3D12RHI/Private/D3D12RayTracing.cpp
@@ -1441,9 +1441,23 @@ struct FD3D12RayTracingDescriptorHeap : public FD3D12DeviceChild
 
 		DescriptorSize = GetParentDevice()->GetDevice()->GetDescriptorHandleIncrementSize(Type);
 
-	#if RAY_TRACING_DESCRIPTOR_CACHE_FULL_COMPARE
-		Descriptors.SetNum(MaxNumDescriptors);
-	#endif // RAY_TRACING_DESCRIPTOR_CACHE_FULL_COMPARE
+		bExhaustiveSamplerDeduplication = GD3D12RayTracingDeduplicateSamplers == 1;
+
+		if (bExhaustiveSamplerDeduplication && InType == D3D12_DESCRIPTOR_HEAP_TYPE_SAMPLER)
+		{
+			// When exhaustive descriptor deduplication is active, all shadow descriptor table entries must be initialized.
+			// Deduplication works by looping over all elements, but they may be written out of order by worker threads.
+			// Initializing descriptors to 0 avoids accidentally matching wrong descriptor entries.
+			Descriptors.SetNumZeroed(MaxNumDescriptors);
+		}
+#if RAY_TRACING_DESCRIPTOR_CACHE_FULL_COMPARE
+		else
+		{
+			// If exhaustive sampler deduplication is not used, then descriptor table copy is only used for validation.
+			// Entries are guaranteed to be written before they're read, so we don't have to spend the time on initialization.
+			Descriptors.SetNumUninitialized(MaxNumDescriptors);
+		}
+#endif // RAY_TRACING_DESCRIPTOR_CACHE_FULL_COMPARE
 	}
 
 	// Returns descriptor heap base index or -1 if allocation is not possible.
@@ -1482,15 +1496,12 @@ struct FD3D12RayTracingDescriptorHeap : public FD3D12DeviceChild
 	{
 		D3D12_CPU_DESCRIPTOR_HANDLE DestDescriptor = GetDescriptorCPU(BaseIndex);
 		GetParentDevice()->GetDevice()->CopyDescriptors(1, &DestDescriptor, &InNumDescriptors, InNumDescriptors, InDescriptors, nullptr, Type);
-	#if RAY_TRACING_DESCRIPTOR_CACHE_FULL_COMPARE
 		for (uint32 i = 0; i < InNumDescriptors; ++i)
 		{
 			Descriptors[BaseIndex + i].ptr = InDescriptors[i].ptr;
 		}
-	#endif // RAY_TRACING_DESCRIPTOR_CACHE_FULL_COMPARE
 	}
 
-#if RAY_TRACING_DESCRIPTOR_CACHE_FULL_COMPARE
 	bool CompareDescriptors(int32 BaseIndex, const D3D12_CPU_DESCRIPTOR_HANDLE* InDescriptors, uint32 InNumDescriptors)
 	{
 		for (uint32 i = 0; i < InNumDescriptors; ++i)
@@ -1502,7 +1513,6 @@ struct FD3D12RayTracingDescriptorHeap : public FD3D12DeviceChild
 		}
 		return true;
 	}
-#endif // RAY_TRACING_DESCRIPTOR_CACHE_FULL_COMPARE
 
 	D3D12_CPU_DESCRIPTOR_HANDLE GetDescriptorCPU(uint32 Index) const
 	{
@@ -1539,9 +1549,9 @@ struct FD3D12RayTracingDescriptorHeap : public FD3D12DeviceChild
 
 	FD3D12RayTracingDescriptorHeapCache::Entry HeapCacheEntry;
 
-#if RAY_TRACING_DESCRIPTOR_CACHE_FULL_COMPARE
 	TArray<D3D12_CPU_DESCRIPTOR_HANDLE> Descriptors;
-#endif // RAY_TRACING_DESCRIPTOR_CACHE_FULL_COMPARE
+
+	bool bExhaustiveSamplerDeduplication = false;
 };
 
 class FD3D12RayTracingDescriptorCache : public FD3D12DeviceChild
@@ -1609,9 +1619,7 @@ public:
 			}
 		}
 
-		const bool bExhaustiveSamplerDeduplication = GD3D12RayTracingDeduplicateSamplers == 1;
-
-		if (bExhaustiveSamplerDeduplication && Type == D3D12_DESCRIPTOR_HEAP_TYPE_SAMPLER)
+		if (Heap.bExhaustiveSamplerDeduplication && Type == D3D12_DESCRIPTOR_HEAP_TYPE_SAMPLER)
 		{
 			// Exhaustive search for a sampler table:
 			// We have to do this because sampler heap space is precious (hard limit of 2048 total entries).
@@ -1645,7 +1653,7 @@ public:
 
 		Heap.CopyDescriptors(DescriptorTableBaseIndex, Descriptors, NumDescriptors);
 
-		if (bExhaustiveSamplerDeduplication && Type == D3D12_DESCRIPTOR_HEAP_TYPE_SAMPLER)
+		if (Heap.bExhaustiveSamplerDeduplication && Type == D3D12_DESCRIPTOR_HEAP_TYPE_SAMPLER)
 		{
 			FPlatformAtomics::InterlockedAdd(&Heap.NumWrittenSamplerDescriptors, NumDescriptors);
 		}
