Refactor transient BLAS scratch allocation code:
* Add FRHIBufferRange struct for passing a buffer with offset and size.
* Update RHIBuildAccelerationStructures to use scratch buffer passed in FRHIBufferRange.
* If scratch buffer is not present, allocate one scratch buffer for all geometry instead of allocating per geometry.
* Add BUF_RayTracingScratch flag to be able to control alignment of buffers used as scratch memory for BLAS builds.

diff --git a/Engine/Source/Runtime/D3D12RHI/Private/D3D12Allocation.cpp b/Engine/Source/Runtime/D3D12RHI/Private/D3D12Allocation.cpp
--- a/Engine/Source/Runtime/D3D12RHI/Private/D3D12Allocation.cpp
+++ b/Engine/Source/Runtime/D3D12RHI/Private/D3D12Allocation.cpp
@@ -1412,6 +1412,11 @@ D3D12_RESOURCE_STATES FD3D12DefaultBufferAllocator::GetDefaultInitialResourceSta
 		check(InHeapType == D3D12_HEAP_TYPE_DEFAULT);
 		return D3D12_RESOURCE_STATE_RAYTRACING_ACCELERATION_STRUCTURE;
 	}
+	else if (EnumHasAnyFlags(InBufferFlags, BUF_RayTracingScratch))
+	{
+		check(InHeapType == D3D12_HEAP_TYPE_DEFAULT);
+		return D3D12_RESOURCE_STATE_UNORDERED_ACCESS;
+	}
 #endif // D3D12_RHI_RAYTRACING
 	else
 	{
diff --git a/Engine/Source/Runtime/D3D12RHI/Private/D3D12CommandContext.h b/Engine/Source/Runtime/D3D12RHI/Private/D3D12CommandContext.h
--- a/Engine/Source/Runtime/D3D12RHI/Private/D3D12CommandContext.h
+++ b/Engine/Source/Runtime/D3D12RHI/Private/D3D12CommandContext.h
@@ -388,7 +388,7 @@ public:
 	virtual void RHICopyBufferRegions(const TArrayView<const FCopyBufferRegionParams> Params) final override;
 	virtual void RHIBindAccelerationStructureMemory(FRHIRayTracingScene* Scene, FRHIBuffer* Buffer, uint32 BufferOffset) final override;
 	virtual void BuildAccelerationStructuresInternal(const TArrayView<D3D12_BUILD_RAYTRACING_ACCELERATION_STRUCTURE_DESC> BuildDesc);
-	virtual void RHIBuildAccelerationStructures(const TArrayView<const FRayTracingGeometryBuildParams> Params) final override;
+	virtual void RHIBuildAccelerationStructures(const TArrayView<const FRayTracingGeometryBuildParams> Params, const FRHIBufferRange& ScratchBufferRange) final override;
 	virtual void RHIBuildAccelerationStructure(const FRayTracingSceneBuildParams& SceneBuildParams) final override;
 	virtual void RHIClearRayTracingBindings(FRHIRayTracingScene* Scene) final override;
 	virtual void RHIRayTraceOcclusion(FRHIRayTracingScene* Scene,
@@ -743,9 +743,9 @@ public:
 		ContextRedirect(RHIEndRenderPass());
 	}
 
-	virtual void RHIBuildAccelerationStructures(const TArrayView<const FRayTracingGeometryBuildParams> Params) final override
+	virtual void RHIBuildAccelerationStructures(const TArrayView<const FRayTracingGeometryBuildParams> Params, const FRHIBufferRange& ScratchBufferRange) final override
 	{
-		ContextRedirect(RHIBuildAccelerationStructures(Params));
+		ContextRedirect(RHIBuildAccelerationStructures(Params, ScratchBufferRange));
 	}
 
 	virtual void RHIBuildAccelerationStructure(const FRayTracingSceneBuildParams& SceneBuildParams) final override
diff --git a/Engine/Source/Runtime/D3D12RHI/Private/D3D12RayTracing.cpp b/Engine/Source/Runtime/D3D12RHI/Private/D3D12RayTracing.cpp
--- a/Engine/Source/Runtime/D3D12RHI/Private/D3D12RayTracing.cpp
+++ b/Engine/Source/Runtime/D3D12RHI/Private/D3D12RayTracing.cpp
@@ -3199,6 +3199,10 @@ FRayTracingAccelerationStructureSize FD3D12DynamicRHI::RHICalcRayTracingGeometry
 		}
 	}
 
+	SizeInfo.ResultSize = Align(SizeInfo.ResultSize, GRHIRayTracingAccelerationStructureAlignment);
+	SizeInfo.BuildScratchSize = Align(SizeInfo.BuildScratchSize, GRHIRayTracingScratchBufferAlignment);
+	SizeInfo.UpdateScratchSize = Align(SizeInfo.UpdateScratchSize, GRHIRayTracingScratchBufferAlignment);
+
 	return SizeInfo;
 }
 
@@ -3466,11 +3470,6 @@ FD3D12RayTracingGeometry::FD3D12RayTracingGeometry(FD3D12Adapter* Adapter, const
 
 	RegisterD3D12RayTracingGeometry(this);
 
-	if (Initializer.Type == ERayTracingGeometryInitializerType::StreamingDestination)
-	{
-		return;
-	}
-
 	checkf(Initializer.Segments.Num() > 0, TEXT("Ray tracing geometry must be initialized with at least one segment."));
 
 	// #yuriy_todo: get flags directly through the initializer
@@ -3512,6 +3511,14 @@ FD3D12RayTracingGeometry::FD3D12RayTracingGeometry(FD3D12Adapter* Adapter, const
 		SizeInfo = RHICalcRayTracingGeometrySize(Initializer);
 	}
 
+	// If this RayTracingGeometry going to be used as streaming destination 
+	// we don't want to allocate its memory as it will be replaced later by streamed version
+	// but we still need correct SizeInfo as it is used to estimate its memory requirements outside of RHI.
+	if (Initializer.Type == ERayTracingGeometryInitializerType::StreamingDestination)
+	{
+		return;
+	}
+
 	// Allocate acceleration structure buffer
 	for (uint32 GPUIndex = 0; GPUIndex < MAX_NUM_GPUS && GPUIndex < GNumExplicitGPUsForRendering; ++GPUIndex)
 	{
@@ -4442,14 +4449,15 @@ void FD3D12CommandContext::BuildAccelerationStructuresInternal(const TArrayView<
 	}
 }
 
-void FD3D12CommandContext::RHIBuildAccelerationStructures(const TArrayView<const FRayTracingGeometryBuildParams> Params)
+void FD3D12CommandContext::RHIBuildAccelerationStructures(const TArrayView<const FRayTracingGeometryBuildParams> Params, const FRHIBufferRange& ScratchBufferRange)
 {
 	TRACE_CPUPROFILER_EVENT_SCOPE(BuildAccelerationStructure_BottomLevel);
 	SCOPE_CYCLE_COUNTER(STAT_D3D12BuildBLAS);
 	LLM_SCOPE_BYNAME(TEXT("FD3D12RT/BLAS"));
 
-	// Update geometry vertex buffers
+	checkf(ScratchBufferRange.Buffer != nullptr, TEXT("BuildAccelerationStructures requires valid scratch buffer"));
 
+	// Update geometry vertex buffers
 	for (const FRayTracingGeometryBuildParams& P : Params)
 	{
 		FD3D12RayTracingGeometry* Geometry = FD3D12DynamicRHI::ResourceCast(P.Geometry.GetReference());
@@ -4488,11 +4496,21 @@ void FD3D12CommandContext::RHIBuildAccelerationStructures(const TArrayView<const
 	const uint32 GPUIndex = GetGPUIndex();
 
 	// Then do all work
-	TArray<TRefCountPtr<FD3D12Buffer>, TInlineAllocator<32>> ScratchBuffers;
 	TArray<D3D12_BUILD_RAYTRACING_ACCELERATION_STRUCTURE_DESC, TInlineAllocator<32>> BuildDescs;
-
 	BuildDescs.Reserve(Params.Num());
-	ScratchBuffers.Reserve(Params.Num());
+
+	uint32 ScratchBufferSize = ScratchBufferRange.Size ? ScratchBufferRange.Size : ScratchBufferRange.Buffer->GetSize();
+
+	checkf(ScratchBufferSize + ScratchBufferRange.Offset <= ScratchBufferRange.Buffer->GetSize(),
+		TEXT("BLAS scratch buffer range size is %lld bytes with offset %lld, but the buffer only has %lld bytes. "),
+		ScratchBufferRange.Size, ScratchBufferRange.Offset, ScratchBufferRange.Buffer->GetSize());
+
+
+	const uint64 ScratchAlignment = GRHIRayTracingAccelerationStructureAlignment;
+	FD3D12Buffer* ScratchBuffer = FD3D12DynamicRHI::ResourceCast(ScratchBufferRange.Buffer, GPUIndex);
+	uint32 ScratchBufferOffset = ScratchBufferRange.Offset;
+
+	ScratchBuffer->GetResource()->UpdateResidency(CommandListHandle);
 
 	FMemMark Mark(FMemStack::Get());
 
@@ -4511,38 +4529,15 @@ void FD3D12CommandContext::RHIBuildAccelerationStructures(const TArrayView<const
 		Geometry->SetupHitGroupSystemParameters(GPUIndex);
 
 		if (Geometry->IsDirty(GPUIndex))
-		{			
-			D3D12_GPU_VIRTUAL_ADDRESS ScratchBufferAddress;
-			if (P.ScratchBuffer)
-			{
-				FD3D12Buffer* ScratchBuffer = FD3D12DynamicRHI::ResourceCast(P.ScratchBuffer, GPUIndex);
+		{
+			uint64 ScratchBufferRequiredSize = P.BuildMode == EAccelerationStructureBuildMode::Update ? Geometry->SizeInfo.UpdateScratchSize : Geometry->SizeInfo.BuildScratchSize;
+			checkf(ScratchBufferRequiredSize + ScratchBufferOffset <= ScratchBufferSize,
+				TEXT("BLAS scratch buffer size is %lld bytes with offset %lld (%lld bytes available), but the build requires %lld bytes. "),
+				ScratchBufferSize, ScratchBufferOffset, ScratchBufferSize - ScratchBufferOffset, ScratchBufferRequiredSize);
 
-				uint64 ScratchBufferRequiredSize = P.BuildMode == EAccelerationStructureBuildMode::Update ? Geometry->SizeInfo.UpdateScratchSize : Geometry->SizeInfo.BuildScratchSize;
-				checkf(ScratchBufferRequiredSize + P.ScratchBufferOffset <= ScratchBuffer->GetSize(),
-					TEXT("BLAS scratch buffer size is %lld bytes with offset %lld (%lld bytes available), but the build requires %lld bytes. "),
-					ScratchBuffer->GetSize(), P.ScratchBufferOffset, ScratchBuffer->GetSize() - P.ScratchBufferOffset, ScratchBufferRequiredSize);
+			D3D12_GPU_VIRTUAL_ADDRESS ScratchBufferAddress = ScratchBuffer->ResourceLocation.GetGPUVirtualAddress() + ScratchBufferOffset;
 
-				ScratchBufferAddress = ScratchBuffer->ResourceLocation.GetGPUVirtualAddress() + P.ScratchBufferOffset;
-				ScratchBuffer->GetResource()->UpdateResidency(CommandListHandle);
-			}
-			else
-			{
-				TRefCountPtr<FD3D12Buffer> ScratchBuffer;
-				if (P.BuildMode == EAccelerationStructureBuildMode::Update)
-				{
-					static const FName BufferName("UpdateScratchBLAS");
-					ScratchBuffer = CreateRayTracingBuffer(GetParentAdapter(), GPUIndex, Geometry->SizeInfo.UpdateScratchSize, ERayTracingBufferType::Scratch, BufferName);
-				}
-				else
-				{
-					static const FName BufferName("BuildScratchBLAS");
-					ScratchBuffer = CreateRayTracingBuffer(GetParentAdapter(), GPUIndex, Geometry->SizeInfo.BuildScratchSize, ERayTracingBufferType::Scratch, BufferName);
-				}
-
-				ScratchBufferAddress = ScratchBuffer->ResourceLocation.GetGPUVirtualAddress();
-				ScratchBuffers.Add(ScratchBuffer);
-				ScratchBuffer->GetResource()->UpdateResidency(CommandListHandle);
-			}
+			ScratchBufferOffset = Align(ScratchBufferOffset + ScratchBufferRequiredSize, ScratchAlignment);
 
 			checkf(ScratchBufferAddress % GRHIRayTracingAccelerationStructureAlignment == 0,
 				TEXT("BLAS scratch buffer (plus offset) must be aligned to %lld bytes."),
diff --git a/Engine/Source/Runtime/Engine/Private/RayTracingSkinnedGeometry.cpp b/Engine/Source/Runtime/Engine/Private/RayTracingSkinnedGeometry.cpp
--- a/Engine/Source/Runtime/Engine/Private/RayTracingSkinnedGeometry.cpp
+++ b/Engine/Source/Runtime/Engine/Private/RayTracingSkinnedGeometry.cpp
@@ -102,13 +102,29 @@ void FRayTracingSkinnedGeometryUpdateQueue::Commit(FRHICommandListImmediate & RH
 		BatchedBuildParams.Reserve(ToUpdate.Num());
 		BatchedUpdateParams.Reserve(ToUpdate.Num());
 
-		auto KickBatch = [&RHICmdList, &BatchedBuildParams, &BatchedUpdateParams, &PrimitivesToUpdates]()
+		auto KickBatch = [&RHICmdList, ScratchBuffer, &BatchedBuildParams, &BatchedUpdateParams, &PrimitivesToUpdates]()
 		{
+			// TODO compute correct offset and increment for the next call or just always use 0 as the offset since we know that 
+			// 2 calls to BuildAccelerationStructures won't overlap due to UAV barrier inside RHIBuildAccelerationStructures so scratch memory can be reused by the next call already
+			uint32 ScratchBLASOffset = 0;
+
 			if (BatchedBuildParams.Num())
 			{
 				SCOPED_GPU_STAT(RHICmdList, SkinnedGeometryBuildBLAS);
 				SCOPED_DRAW_EVENT(RHICmdList, SkinnedGeometryBuildBLAS);
-				RHICmdList.BuildAccelerationStructures(BatchedBuildParams);
+				
+				if (ScratchBuffer)
+				{
+					FRHIBufferRange ScratchBufferRange;
+					ScratchBufferRange.Buffer = ScratchBuffer;
+					ScratchBufferRange.Offset = ScratchBLASOffset;
+					RHICmdList.BuildAccelerationStructures(BatchedBuildParams, ScratchBufferRange);
+				}
+				else
+				{
+					RHICmdList.BuildAccelerationStructures(BatchedBuildParams);
+				}
+				
 				BatchedBuildParams.Empty(BatchedBuildParams.Max());
 			}
 
@@ -116,7 +132,18 @@ void FRayTracingSkinnedGeometryUpdateQueue::Commit(FRHICommandListImmediate & RH
 			{
 				SCOPED_GPU_STAT(RHICmdList, SkinnedGeometryUpdateBLAS);
 				SCOPED_DRAW_EVENT(RHICmdList, SkinnedGeometryUpdateBLAS);
-				RHICmdList.BuildAccelerationStructures(BatchedUpdateParams);
+
+				if (ScratchBuffer)
+				{
+					FRHIBufferRange ScratchBufferRange;
+					ScratchBufferRange.Buffer = ScratchBuffer;
+					ScratchBufferRange.Offset = ScratchBLASOffset;
+					RHICmdList.BuildAccelerationStructures(BatchedUpdateParams, ScratchBufferRange);
+				}
+				else
+				{
+					RHICmdList.BuildAccelerationStructures(BatchedUpdateParams);
+				}
 				BatchedUpdateParams.Empty(BatchedUpdateParams.Max());
 			}
 
@@ -124,7 +151,8 @@ void FRayTracingSkinnedGeometryUpdateQueue::Commit(FRHICommandListImmediate & RH
 		};
 
 		const uint64 ScratchAlignment = GRHIRayTracingAccelerationStructureAlignment;
-		uint32 ScratchBLASOffset = 0;
+		uint32 ScratchBLASCurrentOffset = 0;
+		uint32 ScratchBLASNextOffset = 0;
 
 		// Iterate all the geometries which need an update
 		for (TMap<FRayTracingGeometry*, FRayTracingUpdateInfo>::TRangedForIterator Iter = ToUpdate.begin(); Iter != ToUpdate.end(); ++Iter)
@@ -136,11 +164,9 @@ void FRayTracingSkinnedGeometryUpdateQueue::Commit(FRHICommandListImmediate & RH
 			BuildParams.Geometry = RayTracingGeometry->RayTracingGeometryRHI;
 			BuildParams.BuildMode = UpdateInfo.BuildMode;
 			BuildParams.Segments = RayTracingGeometry->Initializer.Segments;
-			BuildParams.ScratchBuffer = ScratchBuffer;
-			BuildParams.ScratchBufferOffset = ScratchBLASOffset;
 
 			// Update the offset
-			ScratchBLASOffset = Align(ScratchBLASOffset + UpdateInfo.ScratchSize, ScratchAlignment);
+			ScratchBLASNextOffset = Align(ScratchBLASNextOffset + UpdateInfo.ScratchSize, ScratchAlignment);
 
 			// Make 'Build' 10 times more expensive than 1 'Update' of the BVH
 			uint32 PrimitiveCount = RayTracingGeometry->Initializer.TotalPrimitiveCount;
@@ -189,7 +215,7 @@ void FRayTracingSkinnedGeometryUpdateQueue::Commit(FRDGBuilder& GraphBuilder)
 		
 		FRDGBufferDesc ScratchBufferDesc;
 		ScratchBufferDesc.UnderlyingType = FRDGBufferDesc::EUnderlyingType::StructuredBuffer;
-		ScratchBufferDesc.Usage = BUF_UnorderedAccess;
+		ScratchBufferDesc.Usage = BUF_RayTracingScratch;
 		ScratchBufferDesc.BytesPerElement = ScratchAlignment;
 		ScratchBufferDesc.NumElements = FMath::DivideAndRoundUp(BLASScratchSize, ScratchAlignment);
 
diff --git a/Engine/Source/Runtime/RHI/Private/RHI.cpp b/Engine/Source/Runtime/RHI/Private/RHI.cpp
--- a/Engine/Source/Runtime/RHI/Private/RHI.cpp
+++ b/Engine/Source/Runtime/RHI/Private/RHI.cpp
@@ -193,6 +193,7 @@ FString GetBufferUsageFlagsName(EBufferUsageFlags BufferUsage)
 		case BUF_FastVRAM:				return TEXT("BUF_FastVRAM");
 		case BUF_Shared:				return TEXT("BUF_Shared");
 		case BUF_AccelerationStructure:	return TEXT("BUF_AccelerationStructure");
+		case BUF_RayTracingScratch:		return TEXT("BUF_RayTracingScratch");
 		case BUF_VertexBuffer:			return TEXT("BUF_VertexBuffer");
 		case BUF_IndexBuffer:			return TEXT("BUF_IndexBuffer");
 		case BUF_StructuredBuffer:		return TEXT("BUF_StructuredBuffer");
diff --git a/Engine/Source/Runtime/RHI/Private/RHICommandList.cpp b/Engine/Source/Runtime/RHI/Private/RHICommandList.cpp
--- a/Engine/Source/Runtime/RHI/Private/RHICommandList.cpp
+++ b/Engine/Source/Runtime/RHI/Private/RHICommandList.cpp
@@ -2151,6 +2151,38 @@ void FRHIComputeCommandList::operator delete(void *RawMemory)
 	FMemory::Free(RawMemory);
 }
 
+#if RHI_RAYTRACING
+void FRHIComputeCommandList::BuildAccelerationStructure(FRHIRayTracingGeometry* Geometry)
+{
+	FRayTracingGeometryBuildParams Params;
+	Params.Geometry = Geometry;
+	Params.BuildMode = EAccelerationStructureBuildMode::Build;
+
+	FRHIBufferRange ScratchBufferRange{};
+	
+	FRHIResourceCreateInfo ScratchBufferCreateInfo(TEXT("RHIScratchBuffer"));
+	ScratchBufferRange.Buffer = RHICreateBuffer(Geometry->GetSizeInfo().BuildScratchSize, BUF_StructuredBuffer | BUF_RayTracingScratch, 0, ERHIAccess::UAVCompute, ScratchBufferCreateInfo);
+
+	BuildAccelerationStructures(MakeArrayView(&Params, 1), ScratchBufferRange);
+}
+
+void FRHIComputeCommandList::BuildAccelerationStructures(const TArrayView<const FRayTracingGeometryBuildParams> Params)
+{
+	uint64 TotalRequiredScratchMemorySize = 0;
+	for (const FRayTracingGeometryBuildParams& P : Params)
+	{
+		uint64 ScratchBufferRequiredSize = P.BuildMode == EAccelerationStructureBuildMode::Update ? P.Geometry->GetSizeInfo().UpdateScratchSize : P.Geometry->GetSizeInfo().BuildScratchSize;
+		TotalRequiredScratchMemorySize += ScratchBufferRequiredSize;
+	}
+
+	FRHIResourceCreateInfo ScratchBufferCreateInfo(TEXT("RHIScratchBuffer"));
+	FRHIBufferRange ScratchBufferRange{};	
+	ScratchBufferRange.Buffer = RHICreateBuffer(TotalRequiredScratchMemorySize, BUF_StructuredBuffer | BUF_RayTracingScratch, 0, ERHIAccess::UAVCompute, ScratchBufferCreateInfo);
+
+	BuildAccelerationStructures(Params, ScratchBufferRange);
+}
+#endif
+
 void* FRHICommandListBase::operator new(size_t Size)
 {
 	check(0); // you shouldn't be creating these
diff --git a/Engine/Source/Runtime/RHI/Public/RHI.h b/Engine/Source/Runtime/RHI/Public/RHI.h
--- a/Engine/Source/Runtime/RHI/Public/RHI.h
+++ b/Engine/Source/Runtime/RHI/Public/RHI.h
@@ -2453,6 +2453,13 @@ struct FUpdateTextureRegion3D
 	{}
 };
 
+struct FRHIBufferRange
+{
+	FRHIBuffer* Buffer{ nullptr };
+	uint64 Offset{ 0 };
+	uint64 Size{ 0 };
+};
+
 struct FRHIDispatchIndirectParameters
 {
 	uint32 ThreadGroupCountX;
diff --git a/Engine/Source/Runtime/RHI/Public/RHICommandList.h b/Engine/Source/Runtime/RHI/Public/RHICommandList.h
--- a/Engine/Source/Runtime/RHI/Public/RHICommandList.h
+++ b/Engine/Source/Runtime/RHI/Public/RHICommandList.h
@@ -2199,9 +2199,11 @@ FRHICOMMAND_MACRO(FRHICommandClearRayTracingBindings)
 struct FRHICommandBuildAccelerationStructures final : public FRHICommand<FRHICommandBuildAccelerationStructures>
 {
 	const TArrayView<const FRayTracingGeometryBuildParams> Params;
+	FRHIBufferRange ScratchBufferRange;
+	TRefCountPtr<FRHIBuffer> ScratchBuffer;
 
-	explicit FRHICommandBuildAccelerationStructures(const TArrayView<const FRayTracingGeometryBuildParams> InParams)
-		: Params(InParams)
+	explicit FRHICommandBuildAccelerationStructures(const TArrayView<const FRayTracingGeometryBuildParams> InParams, const FRHIBufferRange& ScratchBufferRange)
+		: Params(InParams), ScratchBufferRange(ScratchBufferRange), ScratchBuffer(ScratchBufferRange.Buffer)
 	{}
 
 	RHI_API void Execute(FRHICommandListBase& CmdList);
@@ -3003,19 +3005,14 @@ public:
 #endif // #if PLATFORM_REQUIRES_UAV_TO_RTV_TEXTURE_CACHE_FLUSH_WORKAROUND
 
 #if RHI_RAYTRACING
-	FORCEINLINE_DEBUGGABLE void BuildAccelerationStructure(FRHIRayTracingGeometry* Geometry)
-	{
-		FRayTracingGeometryBuildParams Params;
-		Params.Geometry = Geometry;
-		Params.BuildMode = EAccelerationStructureBuildMode::Build;
-		BuildAccelerationStructures(MakeArrayView(&Params, 1));
-	}
+	void BuildAccelerationStructure(FRHIRayTracingGeometry* Geometry);
+	void BuildAccelerationStructures(const TArrayView<const FRayTracingGeometryBuildParams> Params);
 
-	FORCEINLINE_DEBUGGABLE void BuildAccelerationStructures(const TArrayView<const FRayTracingGeometryBuildParams> Params)
+	FORCEINLINE_DEBUGGABLE void BuildAccelerationStructures(const TArrayView<const FRayTracingGeometryBuildParams> Params, const FRHIBufferRange& ScratchBufferRange)
 	{
 		if (Bypass())
 		{
-			GetComputeContext().RHIBuildAccelerationStructures(Params);
+			GetComputeContext().RHIBuildAccelerationStructures(Params, ScratchBufferRange);
 		}
 		else
 		{
@@ -3031,7 +3028,7 @@ public:
 					InlineParams[i].Segments = AllocArray(Params[i].Segments);
 				}
 			}
-			ALLOC_COMMAND(FRHICommandBuildAccelerationStructures)(MakeArrayView(InlineParams, Params.Num()));
+			ALLOC_COMMAND(FRHICommandBuildAccelerationStructures)(MakeArrayView(InlineParams, Params.Num()), ScratchBufferRange);
 		}
 	}
 
diff --git a/Engine/Source/Runtime/RHI/Public/RHICommandListCommandExecutes.inl b/Engine/Source/Runtime/RHI/Public/RHICommandListCommandExecutes.inl
--- a/Engine/Source/Runtime/RHI/Public/RHICommandListCommandExecutes.inl
+++ b/Engine/Source/Runtime/RHI/Public/RHICommandListCommandExecutes.inl
@@ -562,7 +562,7 @@ void FRHICommandClearRayTracingBindings::Execute(FRHICommandListBase& CmdList)
 void FRHICommandBuildAccelerationStructures::Execute(FRHICommandListBase& CmdList)
 {
 	RHISTAT(BuildAccelerationStructure);
-	INTERNAL_DECORATOR_COMPUTE(RHIBuildAccelerationStructures)(Params);
+	INTERNAL_DECORATOR_COMPUTE(RHIBuildAccelerationStructures)(Params, ScratchBufferRange);
 }
 
 void FRHICommandRayTraceOcclusion::Execute(FRHICommandListBase& CmdList)
diff --git a/Engine/Source/Runtime/RHI/Public/RHIContext.h b/Engine/Source/Runtime/RHI/Public/RHIContext.h
--- a/Engine/Source/Runtime/RHI/Public/RHIContext.h
+++ b/Engine/Source/Runtime/RHI/Public/RHIContext.h
@@ -331,7 +331,7 @@ public:
 		/* empty default implementation */
 	}
 
-	virtual void RHIBuildAccelerationStructures(const TArrayView<const FRayTracingGeometryBuildParams> Params)
+	virtual void RHIBuildAccelerationStructures(const TArrayView<const FRayTracingGeometryBuildParams> Params, const FRHIBufferRange& ScratchBufferRange)
 	{
 		checkNoEntry();
 	}
@@ -417,11 +417,6 @@ struct FRayTracingGeometryBuildParams
 	// Optional array of geometry segments that can be used to change per-segment vertex buffers.
 	// Only fields related to vertex buffer are used. If empty, then geometry vertex buffers are not changed.
 	TArrayView<const FRayTracingGeometrySegment> Segments;
-
-	// Optional explicit build scratch buffer. Must be in UAV state.
-	// Scratch buffer will be created automatically if this is not provided.
-	FRHIBuffer* ScratchBuffer = nullptr;
-	uint32 ScratchBufferOffset = 0;
 };
 
 struct FRayTracingSceneBuildParams
@@ -740,18 +735,19 @@ public:
 		checkNoEntry();
 	}
 
-	virtual void RHIBuildAccelerationStructures(const TArrayView<const FRayTracingGeometryBuildParams> Params)
+	virtual void RHIBuildAccelerationStructures(const TArrayView<const FRayTracingGeometryBuildParams> Params, const FRHIBufferRange& ScratchBufferRange)
 	{
 		checkNoEntry();
 	}
 
-	void RHIBuildAccelerationStructure(FRHIRayTracingGeometry* Geometry)
+	void RHIBuildAccelerationStructures(const TArrayView<const FRayTracingGeometryBuildParams> Params)
 	{
-		FRayTracingGeometryBuildParams Params;
-		Params.Geometry = Geometry;
-		Params.BuildMode = EAccelerationStructureBuildMode::Build;
+		checkNoEntry();
+	}
 
-		RHIBuildAccelerationStructures(MakeArrayView(&Params, 1));
+	void RHIBuildAccelerationStructure(FRHIRayTracingGeometry* Geometry)
+	{
+		checkNoEntry();
 	}
 
 	virtual void RHIBuildAccelerationStructure(const FRayTracingSceneBuildParams& SceneBuildParams)
diff --git a/Engine/Source/Runtime/RHI/Public/RHIDefinitions.h b/Engine/Source/Runtime/RHI/Public/RHIDefinitions.h
--- a/Engine/Source/Runtime/RHI/Public/RHIDefinitions.h
+++ b/Engine/Source/Runtime/RHI/Public/RHIDefinitions.h
@@ -1590,6 +1590,11 @@ enum class EBufferUsageFlags : uint32
 	 * as indirect arg buffers), and the other GPU doesn't actually care about the data.
 	*/
 	MultiGPUGraphIgnore		= 1 << 18,
+	
+	/** Allows buffer to be used as a scratch buffer for building ray tracing acceleration structure,
+	 * which implies unordered access. Only changes the buffer alignment and can be combined with other flags.
+	**/
+	RayTracingScratch = (1 << 19) | UnorderedAccess,
 
 	// Helper bit-masks
 	AnyDynamic = (Dynamic | Volatile),
@@ -1611,6 +1616,7 @@ ENUM_CLASS_FLAGS(EBufferUsageFlags);
 #define BUF_Transient              EBufferUsageFlags::Transient
 #define BUF_Shared                 EBufferUsageFlags::Shared
 #define BUF_AccelerationStructure  EBufferUsageFlags::AccelerationStructure
+#define BUF_RayTracingScratch	   EBufferUsageFlags::RayTracingScratch
 #define BUF_VertexBuffer           EBufferUsageFlags::VertexBuffer
 #define BUF_IndexBuffer            EBufferUsageFlags::IndexBuffer
 #define BUF_StructuredBuffer       EBufferUsageFlags::StructuredBuffer
diff --git a/Engine/Source/Runtime/RHI/Public/RHIValidationContext.h b/Engine/Source/Runtime/RHI/Public/RHIValidationContext.h
--- a/Engine/Source/Runtime/RHI/Public/RHIValidationContext.h
+++ b/Engine/Source/Runtime/RHI/Public/RHIValidationContext.h
@@ -207,10 +207,10 @@ public:
 
 	virtual void RHICopyToStagingBuffer(FRHIBuffer* SourceBufferRHI, FRHIStagingBuffer* DestinationStagingBufferRHI, uint32 InOffset, uint32 InNumBytes) override final;
 
-	virtual void RHIBuildAccelerationStructures(const TArrayView<const FRayTracingGeometryBuildParams> Params) override final
+	virtual void RHIBuildAccelerationStructures(const TArrayView<const FRayTracingGeometryBuildParams> Params, const FRHIBufferRange& ScratchBufferRange) override final
 	{
 		// #yuriy_todo: explicit transitions and state validation for BLAS
-		RHIContext->RHIBuildAccelerationStructures(Params);
+		RHIContext->RHIBuildAccelerationStructures(Params, ScratchBufferRange);
 	}
 
 	virtual void RHIBuildAccelerationStructure(const FRayTracingSceneBuildParams& SceneBuildParams) override final
@@ -977,10 +977,10 @@ public:
 		RHIContext->RHIClearRayTracingBindings(Scene);
 	}
 
-	virtual void RHIBuildAccelerationStructures(const TArrayView<const FRayTracingGeometryBuildParams> Params) override final
+	virtual void RHIBuildAccelerationStructures(const TArrayView<const FRayTracingGeometryBuildParams> Params, const FRHIBufferRange& ScratchBufferRange) override final
 	{
 		// #yuriy_todo: explicit transitions and state validation for BLAS
-		RHIContext->RHIBuildAccelerationStructures(Params);
+		RHIContext->RHIBuildAccelerationStructures(Params, ScratchBufferRange);
 	}
 
 	virtual void RHIBuildAccelerationStructure(const FRayTracingSceneBuildParams& SceneBuildParams) override final
diff --git a/Engine/Source/Runtime/Renderer/Private/RayTracing/RayTracingDynamicGeometry.cpp b/Engine/Source/Runtime/Renderer/Private/RayTracing/RayTracingDynamicGeometry.cpp
--- a/Engine/Source/Runtime/Renderer/Private/RayTracing/RayTracingDynamicGeometry.cpp
+++ b/Engine/Source/Runtime/Renderer/Private/RayTracing/RayTracingDynamicGeometry.cpp
@@ -392,9 +392,6 @@ void FRayTracingDynamicGeometryCollection::DispatchUpdates(FRHIComputeCommandLis
 			{
 				TRACE_CPUPROFILER_EVENT_SCOPE(SetupSegmentData);
 
-				const uint64 ScratchAlignment = GRHIRayTracingAccelerationStructureAlignment;
-				uint32 ScratchBLASOffset = 0;
-
 				// Setup the array views on final allocated segments array
 				FRayTracingGeometrySegment* SegmentData = Segments.GetData();
 				for (FRayTracingGeometryBuildParams& Param : BuildParams)
@@ -405,15 +402,6 @@ void FRayTracingDynamicGeometryCollection::DispatchUpdates(FRHIComputeCommandLis
 						Param.Segments = MakeArrayView(SegmentData, SegmentCount);
 						SegmentData += SegmentCount;
 					}
-
-					Param.ScratchBuffer = ScratchBuffer;
-					Param.ScratchBufferOffset = ScratchBLASOffset;
-
-					// Update the offset
-
-					const FRayTracingAccelerationStructureSize BLASSizeInfo = Param.Geometry->GetSizeInfo();
-					const uint64 ScratchSize = Param.BuildMode == EAccelerationStructureBuildMode::Build ? BLASSizeInfo.BuildScratchSize : BLASSizeInfo.UpdateScratchSize;
-					ScratchBLASOffset = Align(ScratchBLASOffset + ScratchSize, ScratchAlignment);
 				}
 			}
 
@@ -541,7 +529,11 @@ void FRayTracingDynamicGeometryCollection::DispatchUpdates(FRHIComputeCommandLis
 				// requests over to the RaytracingGeometry manager so they can be correctly scheduled
 				// with other build requests in the engine (see UE-106982)
 				SCOPED_DRAW_OR_COMPUTE_EVENT(ParentCmdList, Build);
-				ParentCmdList.BuildAccelerationStructures(BuildParams);
+
+				FRHIBufferRange ScratchBufferRange;
+				ScratchBufferRange.Buffer = ScratchBuffer;
+				ScratchBufferRange.Offset = 0;
+				ParentCmdList.BuildAccelerationStructures(BuildParams, ScratchBufferRange);
 			}
 
 		}
diff --git a/Engine/Source/Runtime/VulkanRHI/Private/VulkanContext.h b/Engine/Source/Runtime/VulkanRHI/Private/VulkanContext.h
--- a/Engine/Source/Runtime/VulkanRHI/Private/VulkanContext.h
+++ b/Engine/Source/Runtime/VulkanRHI/Private/VulkanContext.h
@@ -103,7 +103,7 @@ public:
 #if VULKAN_RHI_RAYTRACING
 	virtual void RHIClearRayTracingBindings(FRHIRayTracingScene* Scene) final override;
 	virtual void RHIBindAccelerationStructureMemory(FRHIRayTracingScene* Scene, FRHIBuffer* Buffer, uint32 BufferOffset) final override;
-	virtual void RHIBuildAccelerationStructures(const TArrayView<const FRayTracingGeometryBuildParams> Params) final override;
+	virtual void RHIBuildAccelerationStructures(const TArrayView<const FRayTracingGeometryBuildParams> Params, const FRHIBufferRange& ScratchBufferRange) final override;
 	virtual void RHIBuildAccelerationStructure(const FRayTracingSceneBuildParams& SceneBuildParams) final override;
 	virtual void RHIRayTraceOcclusion(FRHIRayTracingScene* Scene, FRHIShaderResourceView* Rays, FRHIUnorderedAccessView* Output, uint32 NumRays) final override;
 #endif
diff --git a/Engine/Source/Runtime/VulkanRHI/Private/VulkanRayTracing.cpp b/Engine/Source/Runtime/VulkanRHI/Private/VulkanRayTracing.cpp
--- a/Engine/Source/Runtime/VulkanRHI/Private/VulkanRayTracing.cpp
+++ b/Engine/Source/Runtime/VulkanRHI/Private/VulkanRayTracing.cpp
@@ -311,7 +311,7 @@ FVulkanRayTracingGeometry::FVulkanRayTracingGeometry(const FRayTracingGeometryIn
 	AccelerationStructureBuffer = ResourceCast(RHICreateBuffer(BuildData.SizesInfo.accelerationStructureSize, BUF_AccelerationStructure, 0, ERHIAccess::BVHWrite, BlasBufferCreateInfo).GetReference());
 
 	FRHIResourceCreateInfo ScratchBufferCreateInfo(TEXT("BuildScratchBLAS"));
-	ScratchBuffer = ResourceCast(RHICreateBuffer(BuildData.SizesInfo.buildScratchSize, BUF_UnorderedAccess | BUF_StructuredBuffer, 0, ERHIAccess::UAVCompute, ScratchBufferCreateInfo).GetReference());
+	ScratchBuffer = ResourceCast(RHICreateBuffer(BuildData.SizesInfo.buildScratchSize, BUF_StructuredBuffer | BUF_RayTracingScratch, 0, ERHIAccess::UAVCompute, ScratchBufferCreateInfo).GetReference());
 
 	VkDevice NativeDevice = Device->GetInstanceHandle();
 
@@ -849,7 +849,7 @@ void FVulkanCommandListContext::RHIBindAccelerationStructureMemory(FRHIRayTracin
 }
 
 // Todo: High level rhi call should have transitioned and verified vb and ib to read for each segment
-void FVulkanCommandListContext::RHIBuildAccelerationStructures(const TArrayView<const FRayTracingGeometryBuildParams> Params)
+void FVulkanCommandListContext::RHIBuildAccelerationStructures(const TArrayView<const FRayTracingGeometryBuildParams> Params, const FRHIBufferRange& ScratchBufferRange)
 {
 	for (const FRayTracingGeometryBuildParams& P : Params)
 	{
@@ -857,6 +857,7 @@ void FVulkanCommandListContext::RHIBuildAccelerationStructures(const TArrayView<
 
 		// Todo: Update geometry from params for each segment
 		// Todo: Can we do this only for an update?
+		// Todo: Use provided ScratchBuffer instead of allocating.
 
 		// Build as for each segment
 		Geometry->BuildAccelerationStructure(*this, P.BuildMode);
