MGPU perf:  micro-optimization of multi-GPU debug name initialization.  Debug name initialization was 6% of the total time in the render thread in development builds, due to dynamic allocations, and a name being generated per GPU, called many times a frame.  As a bonus optimization, if the name is unchanged (frequently happens for pooled resources used for the same purpose each frame), we skip setting the name.  Overall 3x perf improvement, dropping debug name overhead to 2%.

#rnx

diff --git a/Engine/Source/Runtime/D3D12RHI/Private/D3D12Buffer.cpp b/Engine/Source/Runtime/D3D12RHI/Private/D3D12Buffer.cpp
--- a/Engine/Source/Runtime/D3D12RHI/Private/D3D12Buffer.cpp
+++ b/Engine/Source/Runtime/D3D12RHI/Private/D3D12Buffer.cpp
@@ -872,16 +872,31 @@ void FD3D12DynamicRHI::RHIBindDebugLabelName(FRHIBuffer* BufferRHI, const TCHAR*
 
 	if (GNumExplicitGPUsForRendering > 1)
 	{
+		// Generate string of the form "Name (GPU #)" -- assumes GPU index is a single digit.  This is called many times
+		// a frame, so we want to avoid any string functions which dynamically allocate, to reduce perf overhead.
+		static_assert(MAX_NUM_GPUS <= 10);
+
+		static const TCHAR NameSuffix[] = TEXT(" (GPU #)");
+		constexpr int32 NameSuffixLengthWithTerminator = (int32)UE_ARRAY_COUNT(NameSuffix);
+		constexpr int32 NameBufferLength = 256;
+		constexpr int32 GPUIndexSuffixOffset = 6;		// Offset of '#' character
+
+		// Combine Name and suffix in our string buffer (clamping the length for bounds checking).  We'll replace the GPU index
+		// with the appropriate digit in the loop.
+		int32 NameLength = FMath::Min(FCString::Strlen(Name), NameBufferLength - NameSuffixLengthWithTerminator);
+		int32 GPUIndexOffset = NameLength + GPUIndexSuffixOffset;
+
+		TCHAR DebugName[NameBufferLength];
+		FMemory::Memcpy(&DebugName[0], Name, NameLength*sizeof(TCHAR));
+		FMemory::Memcpy(&DebugName[NameLength], NameSuffix, NameSuffixLengthWithTerminator*sizeof(TCHAR));
+
 		for (; BufferIt; ++BufferIt)
 		{
 			FD3D12Resource* Resource = BufferIt->GetResource();
 
-			TArray<FStringFormatArg> Args;
-			Args.Add(Name);
-			Args.Add(LexToString(BufferIt->GetParentDevice()->GetGPUIndex()));
+			DebugName[GPUIndexOffset] = TEXT('0') + BufferIt->GetParentDevice()->GetGPUIndex();
 
-			FString DebugName = FString::Format(TEXT("{0} (GPU {1})"), Args);
-			SetName(Resource, DebugName.GetCharArray().GetData());
+			SetName(Resource, DebugName);
 		}
 	}
 	else
diff --git a/Engine/Source/Runtime/D3D12RHI/Private/D3D12Texture.cpp b/Engine/Source/Runtime/D3D12RHI/Private/D3D12Texture.cpp
--- a/Engine/Source/Runtime/D3D12RHI/Private/D3D12Texture.cpp
+++ b/Engine/Source/Runtime/D3D12RHI/Private/D3D12Texture.cpp
@@ -3010,16 +3010,31 @@ void FD3D12DynamicRHI::RHIBindDebugLabelName(FRHITexture* TextureRHI, const TCHA
 
 	if (GNumExplicitGPUsForRendering > 1)
 	{
-		for (;BaseTexture; ++BaseTexture)
+		// Generate string of the form "Name (GPU #)" -- assumes GPU index is a single digit.  This is called many times
+		// a frame, so we want to avoid any string functions which dynamically allocate, to reduce perf overhead.
+		static_assert(MAX_NUM_GPUS <= 10);
+
+		static const TCHAR NameSuffix[] = TEXT(" (GPU #)");
+		constexpr int32 NameSuffixLengthWithTerminator = (int32)UE_ARRAY_COUNT(NameSuffix);
+		constexpr int32 NameBufferLength = 256;
+		constexpr int32 GPUIndexSuffixOffset = 6;		// Offset of '#' character
+
+		// Combine Name and suffix in our string buffer (clamping the length for bounds checking).  We'll replace the GPU index
+		// with the appropriate digit in the loop.
+		int32 NameLength = FMath::Min(FCString::Strlen(Name), NameBufferLength - NameSuffixLengthWithTerminator);
+		int32 GPUIndexOffset = NameLength + GPUIndexSuffixOffset;
+
+		TCHAR DebugName[NameBufferLength];
+		FMemory::Memcpy(&DebugName[0], Name, NameLength * sizeof(TCHAR));
+		FMemory::Memcpy(&DebugName[NameLength], NameSuffix, NameSuffixLengthWithTerminator * sizeof(TCHAR));
+
+		for (; BaseTexture; ++BaseTexture)
 		{
 			FD3D12Resource* Resource = BaseTexture->GetResource();
 
-			TArray<FStringFormatArg> Args;
-			Args.Add(Name);
-			Args.Add(LexToString(BaseTexture->GetParentDevice()->GetGPUIndex()));
+			DebugName[GPUIndexOffset] = TEXT('0') + BaseTexture->GetParentDevice()->GetGPUIndex();
 
-			FString DebugName = FString::Format(TEXT("{0} (GPU {1})"), Args);
-			SetName(Resource, DebugName.GetCharArray().GetData());
+			SetName(Resource, DebugName);
 		}
 	}
 	else
diff --git a/Engine/Source/Runtime/D3D12RHI/Public/D3D12Resources.h b/Engine/Source/Runtime/D3D12RHI/Public/D3D12Resources.h
--- a/Engine/Source/Runtime/D3D12RHI/Public/D3D12Resources.h
+++ b/Engine/Source/Runtime/D3D12RHI/Public/D3D12Resources.h
@@ -285,8 +285,13 @@ public:
 
 	void SetName(const TCHAR* Name)
 	{
-		DebugName = FName(Name);
-		::SetName(Resource, Name);
+		// Check name before setting it.  Saves FName lookup and driver call.  Names are frequently the same for pooled buffers
+		// that end up getting reused for the same purpose every frame (2/3 of calls to this function on a given frame).
+		if (DebugName != Name)
+		{
+			DebugName = FName(Name);
+			::SetName(Resource, Name);
+		}
 	}
 
 	FName GetName() const
