Bindless preperation for wider support
- Moving D3D12 descriptor allocator to RHICore.
- Removing GetDefaultShaderResourceView in favor of GetDefaultBindlessHandle.

diff --git a/Engine/Source/Runtime/D3D12RHI/Private/D3D12Descriptors.cpp b/Engine/Source/Runtime/D3D12RHI/Private/D3D12Descriptors.cpp
--- a/Engine/Source/Runtime/D3D12RHI/Private/D3D12Descriptors.cpp
+++ b/Engine/Source/Runtime/D3D12RHI/Private/D3D12Descriptors.cpp
@@ -50,167 +50,14 @@ FD3D12DescriptorHeap::FD3D12DescriptorHeap(FD3D12DescriptorHeap* SubAllocateSour
 
 FD3D12DescriptorHeap::~FD3D12DescriptorHeap() = default;
 
-////////////////////////////////////////////////////////////////////////////////////////////////////////////////////////
-// FDescriptorAllocator
-
-struct FDescriptorAllocatorRange
-{
-	FDescriptorAllocatorRange(uint32 InFirst, uint32 InLast) : First(InFirst), Last(InLast) {}
-	uint32 First;
-	uint32 Last;
-};
-
-FDescriptorAllocator::FDescriptorAllocator()
-{
-}
-
-FDescriptorAllocator::~FDescriptorAllocator()
-{
-}
-
-void FDescriptorAllocator::Init(uint32 InNumDescriptors)
-{
-	Capacity = InNumDescriptors;
-	Ranges.Emplace(0, InNumDescriptors - 1);
-}
-
-bool FDescriptorAllocator::Allocate(uint32 NumDescriptors, uint32& OutSlot)
-{
-	FScopeLock Lock(&CriticalSection);
-
-	if (const uint32 NumRanges = Ranges.Num(); NumRanges > 0)
-	{
-		uint32 Index = 0;
-		do
-		{
-			FDescriptorAllocatorRange& CurrentRange = Ranges[Index];
-			const uint32 Size = 1 + CurrentRange.Last - CurrentRange.First;
-			if (NumDescriptors <= Size)
-			{
-				uint32 First = CurrentRange.First;
-				if (NumDescriptors == Size && Index + 1 < NumRanges)
-				{
-					// Range is full and a new range exists, so move on to that one
-					Ranges.RemoveAt(Index);
-				}
-				else
-				{
-					CurrentRange.First += NumDescriptors;
-				}
-				OutSlot = First;
-				return true;
-			}
-			++Index;
-		} while (Index < NumRanges);
-	}
-
-	OutSlot = UINT_MAX;
-	return false;
-}
-
-void FDescriptorAllocator::Free(uint32 Offset, uint32 NumDescriptors)
-{
-	if (Offset == UINT_MAX || NumDescriptors == 0)
-	{
-		return;
-	}
-
-	FScopeLock Lock(&CriticalSection);
-
-	const uint32 End = Offset + NumDescriptors;
-	// Binary search of the range list
-	uint32 Index0 = 0;
-	uint32 Index1 = Ranges.Num() -1;
-	for (;;)
-	{
-		const uint32 Index = (Index0 + Index1) / 2;
-		if (Offset < Ranges[Index].First)
-		{
-			// Before current range, check if neighboring
-			if (End >= Ranges[Index].First)
-			{
-				check(End == Ranges[Index].First); // Can't overlap a range of free IDs
-				// Neighbor id, check if neighboring previous range too
-				if (Index > Index0 && Offset - 1 == Ranges[Index - 1].Last)
-				{
-					// Merge with previous range
-					Ranges[Index - 1].Last = Ranges[Index].Last;
-					Ranges.RemoveAt(Index);
-				}
-				else
-				{
-					// Just grow range
-					Ranges[Index].First = Offset;
-				}
-				return;
-			}
-			else
-			{
-				// Non-neighbor id
-				if (Index != Index0)
-				{
-					// Cull upper half of list
-					Index1 = Index - 1;
-				}
-				else
-				{
-					// Found our position in the list, insert the deleted range here
-					Ranges.EmplaceAt(Index, Offset, End -1);
-					return;
-				}
-			}
-		}
-		else if (Offset > Ranges[Index].Last)
-		{
-			// After current range, check if neighboring
-			if (Offset - 1 == Ranges[Index].Last)
-			{
-				// Neighbor id, check if neighboring next range too
-				if (Index < Index1 && End == Ranges[Index + 1].First)
-				{
-					// Merge with next range
-					Ranges[Index].Last = Ranges[Index + 1].Last;
-					Ranges.RemoveAt(Index);
-				}
-				else
-				{
-					// Just grow range
-					Ranges[Index].Last += NumDescriptors;
-				}
-				return;
-			}
-			else
-			{
-				// Non-neighbor id
-				if (Index != Index1)
-				{
-					// Cull bottom half of list
-					Index0 = Index + 1;
-				}
-				else
-				{
-					// Found our position in the list, insert the deleted range here
-					Ranges.EmplaceAt(Index + 1, Offset, End - 1);
-					return;
-				}
-			}
-		}
-		else
-		{
-			// Inside a free block, not a valid offset
-			checkNoEntry();
-		}
-	}
-}
-
 ////////////////////////////////////////////////////////////////////////////////////////////////////////////////////////
 // FD3D12DescriptorManager
 
 FD3D12DescriptorManager::FD3D12DescriptorManager(FD3D12Device* Device, FD3D12DescriptorHeap* InHeap)
 	: FD3D12DeviceChild(Device)
+	, FRHIHeapDescriptorAllocator(InHeap->GetType(), InHeap->GetNumDescriptors())
 	, Heap(InHeap)
 {
-	Allocator.Init(InHeap->GetNumDescriptors());
 }
 
 FD3D12DescriptorManager::~FD3D12DescriptorManager() = default;
@@ -226,7 +73,7 @@ void FD3D12DescriptorManager::Init(TCHAR* InName, ERHIDescriptorHeapType InType,
 			D3D12_DESCRIPTOR_HEAP_FLAG_SHADER_VISIBLE
 		);
 
-		Allocator.Init(InNumDescriptors);
+		FRHIDescriptorAllocator::Init(InNumDescriptors);
 	}
 }
 
@@ -239,32 +86,6 @@ void FD3D12DescriptorManager::Destroy()
 	}
 }
 
-FRHIDescriptorHandle FD3D12DescriptorManager::AllocateDescriptor()
-{
-	uint32 Index{};
-	if (Allocator.Allocate(1, Index))
-	{
-		return FRHIDescriptorHandle(Heap->GetType(), Index);
-	}
-	return FRHIDescriptorHandle();
-}
-
-bool FD3D12DescriptorManager::AllocateDescriptors(uint32 NumDescriptors, uint32& OutSlot)
-{
-	return Allocator.Allocate(NumDescriptors, OutSlot);
-}
-
-void FD3D12DescriptorManager::FreeDescriptor(FRHIDescriptorHandle InHandle)
-{
-	check(InHandle.GetType() == Heap->GetType());
-	Allocator.Free(InHandle.GetIndex(), 1);
-}
-
-void FD3D12DescriptorManager::FreeDescriptors(uint32 Slot, uint32 NumDescriptors)
-{
-	Allocator.Free(Slot, NumDescriptors);
-}
-
 ////////////////////////////////////////////////////////////////////////////////////////////////////////////////////////
 // FD3D12BindlessDescriptorManager
 
@@ -306,7 +127,7 @@ FRHIDescriptorHandle FD3D12BindlessDescriptorManager::AllocateDescriptor(ERHIDes
 	{
 		if (Manager.HandlesAllocation(InType))
 		{
-			return Manager.AllocateDescriptor();
+			return Manager.Allocate();
 		}
 	}
 	return FRHIDescriptorHandle();
@@ -318,7 +139,7 @@ void FD3D12BindlessDescriptorManager::FreeDescriptor(FRHIDescriptorHandle InHand
 	{
 		if (Manager.HandlesAllocation(InHandle.GetType()))
 		{
-			Manager.FreeDescriptor(InHandle);
+			Manager.Free(InHandle);
 			return;
 		}
 	}
@@ -662,10 +483,10 @@ FD3D12DescriptorHeap* FD3D12DescriptorHeapManager::AllocateHeap(const TCHAR* InD
 {
 	for (FD3D12DescriptorManager& GlobalHeap : GlobalHeaps)
 	{
-		if (GlobalHeap.HandlesAllocation(InHeapType, InHeapFlags))
+		if (GlobalHeap.HandlesAllocationWithFlags(InHeapType, InHeapFlags))
 		{
 			uint32 Offset = 0;
-			if (GlobalHeap.AllocateDescriptors(InNumDescriptors, Offset))
+			if (GlobalHeap.Allocate(InNumDescriptors, Offset))
 			{
 				return new FD3D12DescriptorHeap(GlobalHeap.GetHeap(), Offset, InNumDescriptors);
 			}
@@ -684,7 +505,7 @@ void FD3D12DescriptorHeapManager::FreeHeap(FD3D12DescriptorHeap* InHeap)
 	{
 		if (GlobalHeap.IsHeapAChild(InHeap))
 		{
-			GlobalHeap.FreeDescriptors(InHeap->GetOffset(), InHeap->GetNumDescriptors());
+			GlobalHeap.Free(InHeap->GetOffset(), InHeap->GetNumDescriptors());
 			return;
 		}
 	}
diff --git a/Engine/Source/Runtime/D3D12RHI/Private/D3D12Descriptors.h b/Engine/Source/Runtime/D3D12RHI/Private/D3D12Descriptors.h
--- a/Engine/Source/Runtime/D3D12RHI/Private/D3D12Descriptors.h
+++ b/Engine/Source/Runtime/D3D12RHI/Private/D3D12Descriptors.h
@@ -3,6 +3,7 @@
 #pragma once
 
 #include "D3D12CommandList.h"
+#include "RHIDescriptorAllocator.h"
 
 struct FD3D12OfflineHeapEntry;
 
@@ -60,28 +61,8 @@ private:
 };
 using FD3D12DescriptorHeapPtr = TRefCountPtr<FD3D12DescriptorHeap>;
 
-struct FDescriptorAllocatorRange;
-
-class FDescriptorAllocator
-{
-public:
-	FDescriptorAllocator();
-	~FDescriptorAllocator();
-
-	void Init(uint32 InNumDescriptors);
-
-	bool Allocate(uint32 NumDescriptors, uint32& OutSlot);
-	void Free(uint32 Slot, uint32 NumDescriptors);
-
-private:
-	TArray<FDescriptorAllocatorRange> Ranges;
-	uint32 Capacity = 0;
-
-	FCriticalSection CriticalSection;
-};
-
 /** Manager for resource descriptor allocations. */
-class FD3D12DescriptorManager : public FD3D12DeviceChild
+class FD3D12DescriptorManager : public FD3D12DeviceChild, public FRHIHeapDescriptorAllocator
 {
 public:
 	FD3D12DescriptorManager() = delete;
@@ -91,21 +72,10 @@ public:
 	void Init(TCHAR* InName, ERHIDescriptorHeapType InType, uint32 InNumDescriptors);
 	void Destroy();
 
-	FRHIDescriptorHandle AllocateDescriptor();
-	bool AllocateDescriptors(uint32 NumDescriptors, uint32& OutSlot);
-
-	void FreeDescriptor(FRHIDescriptorHandle InHandle);
-	void FreeDescriptors(uint32 Slot, uint32 NumDescriptors);
-
 	inline FD3D12DescriptorHeap* GetHeap() { return Heap.GetReference(); }
 	inline D3D12_CPU_DESCRIPTOR_HANDLE GetCpuDescriptorHandle(uint32 InSlot) const { return Heap->GetCPUSlotHandle(InSlot); }
 
-	inline bool HandlesAllocation(ERHIDescriptorHeapType InHeapType) const
-	{
-		return Heap && Heap->GetType() == InHeapType;
-	}
-
-	inline bool HandlesAllocation(ERHIDescriptorHeapType InHeapType, D3D12_DESCRIPTOR_HEAP_FLAGS InHeapFlags) const
+	inline bool HandlesAllocationWithFlags(ERHIDescriptorHeapType InHeapType, D3D12_DESCRIPTOR_HEAP_FLAGS InHeapFlags) const
 	{
 		return HandlesAllocation(InHeapType) && Heap->GetFlags() == InHeapFlags;
 	}
@@ -117,7 +87,6 @@ public:
 
 private:
 	FD3D12DescriptorHeapPtr Heap;
-	FDescriptorAllocator Allocator;
 };
 
 /** Manager for resource descriptors used in bindless rendering. */
diff --git a/Engine/Source/Runtime/D3D12RHI/Private/D3D12Texture.h b/Engine/Source/Runtime/D3D12RHI/Private/D3D12Texture.h
--- a/Engine/Source/Runtime/D3D12RHI/Private/D3D12Texture.h
+++ b/Engine/Source/Runtime/D3D12RHI/Private/D3D12Texture.h
@@ -327,9 +327,13 @@ public:
 
 	bool IsLastMipFirst() const { return bMipOrderDescending; }
 
-	virtual FRHIShaderResourceView* GetDefaultShaderResourceView() const final
+	virtual FRHIDescriptorHandle GetDefaultBindlessHandle() const final
 	{
-		return GetShaderResourceView();
+		if (FRHIShaderResourceView* View = GetShaderResourceView())
+		{
+			return View->GetBindlessHandle();
+		}
+		return FRHIDescriptorHandle();
 	}
 
 	/** FRHITexture override.  See FRHITexture::GetNativeResource() */
diff --git a/Engine/Source/Runtime/RHI/Public/RHIResources.h b/Engine/Source/Runtime/RHI/Public/RHIResources.h
--- a/Engine/Source/Runtime/RHI/Public/RHIResources.h
+++ b/Engine/Source/Runtime/RHI/Public/RHIResources.h
@@ -1253,7 +1253,7 @@ public:
 	virtual class FRHITextureCube* GetTextureCube() { return NULL; }
 	virtual class FRHITextureReference* GetTextureReference() { return NULL; }
 
-	virtual class FRHIShaderResourceView* GetDefaultShaderResourceView() const { return nullptr; }
+	virtual FRHIDescriptorHandle GetDefaultBindlessHandle() const { return FRHIDescriptorHandle(); }
 
 	// Slower method to get Size X, Y & Z information. Prefer sub-classes' GetSizeX(), etc
 	virtual FIntVector GetSizeXYZ() const = 0;
diff --git a/Engine/Source/Runtime/RHICore/Private/RHIDescriptorAllocator.cpp b/Engine/Source/Runtime/RHICore/Private/RHIDescriptorAllocator.cpp
new file mode 100644
--- /dev/null
+++ b/Engine/Source/Runtime/RHICore/Private/RHIDescriptorAllocator.cpp
@@ -0,0 +1,212 @@
+// Copyright Epic Games, Inc. All Rights Reserved.
+
+#include "RHIDescriptorAllocator.h"
+
+struct FRHIDescriptorAllocatorRange
+{
+	FRHIDescriptorAllocatorRange(uint32 InFirst, uint32 InLast) : First(InFirst), Last(InLast) {}
+	uint32 First;
+	uint32 Last;
+};
+
+FRHIDescriptorAllocator::FRHIDescriptorAllocator()
+{
+}
+
+FRHIDescriptorAllocator::FRHIDescriptorAllocator(uint32 InNumDescriptors)
+{
+	Init(InNumDescriptors);
+}
+
+FRHIDescriptorAllocator::~FRHIDescriptorAllocator()
+{
+}
+
+void FRHIDescriptorAllocator::Init(uint32 InNumDescriptors)
+{
+	Capacity = InNumDescriptors;
+	Ranges.Emplace(0, InNumDescriptors - 1);
+}
+
+void FRHIDescriptorAllocator::Shutdown()
+{
+	Ranges.Empty();
+	Capacity = 0;
+}
+
+FRHIDescriptorHandle FRHIDescriptorAllocator::Allocate(ERHIDescriptorHeapType InType)
+{
+	uint32 Index{};
+	if (Allocate(1, Index))
+	{
+		return FRHIDescriptorHandle(InType, Index);
+	}
+	return FRHIDescriptorHandle();
+}
+
+void FRHIDescriptorAllocator::Free(FRHIDescriptorHandle InHandle)
+{
+	if (InHandle.IsValid())
+	{
+		Free(InHandle.GetIndex(), 1);
+	}
+}
+
+bool FRHIDescriptorAllocator::Allocate(uint32 NumDescriptors, uint32& OutSlot)
+{
+	FScopeLock Lock(&CriticalSection);
+
+	if (const uint32 NumRanges = Ranges.Num(); NumRanges > 0)
+	{
+		uint32 Index = 0;
+		do
+		{
+			FRHIDescriptorAllocatorRange& CurrentRange = Ranges[Index];
+			const uint32 Size = 1 + CurrentRange.Last - CurrentRange.First;
+			if (NumDescriptors <= Size)
+			{
+				uint32 First = CurrentRange.First;
+				if (NumDescriptors == Size && Index + 1 < NumRanges)
+				{
+					// Range is full and a new range exists, so move on to that one
+					Ranges.RemoveAt(Index);
+				}
+				else
+				{
+					CurrentRange.First += NumDescriptors;
+				}
+				OutSlot = First;
+				return true;
+			}
+			++Index;
+		} while (Index < NumRanges);
+	}
+
+	OutSlot = UINT_MAX;
+	return false;
+}
+
+void FRHIDescriptorAllocator::Free(uint32 Offset, uint32 NumDescriptors)
+{
+	if (Offset == UINT_MAX || NumDescriptors == 0)
+	{
+		return;
+	}
+
+	FScopeLock Lock(&CriticalSection);
+
+	const uint32 End = Offset + NumDescriptors;
+	// Binary search of the range list
+	uint32 Index0 = 0;
+	uint32 Index1 = Ranges.Num() - 1;
+	for (;;)
+	{
+		const uint32 Index = (Index0 + Index1) / 2;
+		if (Offset < Ranges[Index].First)
+		{
+			// Before current range, check if neighboring
+			if (End >= Ranges[Index].First)
+			{
+				check(End == Ranges[Index].First); // Can't overlap a range of free IDs
+				// Neighbor id, check if neighboring previous range too
+				if (Index > Index0 && Offset - 1 == Ranges[Index - 1].Last)
+				{
+					// Merge with previous range
+					Ranges[Index - 1].Last = Ranges[Index].Last;
+					Ranges.RemoveAt(Index);
+				}
+				else
+				{
+					// Just grow range
+					Ranges[Index].First = Offset;
+				}
+				return;
+			}
+			else
+			{
+				// Non-neighbor id
+				if (Index != Index0)
+				{
+					// Cull upper half of list
+					Index1 = Index - 1;
+				}
+				else
+				{
+					// Found our position in the list, insert the deleted range here
+					Ranges.EmplaceAt(Index, Offset, End - 1);
+					return;
+				}
+			}
+		}
+		else if (Offset > Ranges[Index].Last)
+		{
+			// After current range, check if neighboring
+			if (Offset - 1 == Ranges[Index].Last)
+			{
+				// Neighbor id, check if neighboring next range too
+				if (Index < Index1 && End == Ranges[Index + 1].First)
+				{
+					// Merge with next range
+					Ranges[Index].Last = Ranges[Index + 1].Last;
+					Ranges.RemoveAt(Index);
+				}
+				else
+				{
+					// Just grow range
+					Ranges[Index].Last += NumDescriptors;
+				}
+				return;
+			}
+			else
+			{
+				// Non-neighbor id
+				if (Index != Index1)
+				{
+					// Cull bottom half of list
+					Index0 = Index + 1;
+				}
+				else
+				{
+					// Found our position in the list, insert the deleted range here
+					Ranges.EmplaceAt(Index + 1, Offset, End - 1);
+					return;
+				}
+			}
+		}
+		else
+		{
+			// Inside a free block, not a valid offset
+			checkNoEntry();
+		}
+	}
+}
+
+FRHIHeapDescriptorAllocator::FRHIHeapDescriptorAllocator(ERHIDescriptorHeapType InType, uint32 InDescriptorCount)
+	: FRHIDescriptorAllocator(InDescriptorCount)
+	, Type(InType)
+{
+}
+
+FRHIDescriptorHandle FRHIHeapDescriptorAllocator::Allocate()
+{
+	return FRHIDescriptorAllocator::Allocate(Type);
+}
+
+void FRHIHeapDescriptorAllocator::Free(FRHIDescriptorHandle InHandle)
+{
+	if (InHandle.IsValid())
+	{
+		check(InHandle.GetType() == Type);
+		FRHIDescriptorAllocator::Free(InHandle);
+	}
+}
+
+bool FRHIHeapDescriptorAllocator::Allocate(uint32 NumDescriptors, uint32& OutSlot)
+{
+	return FRHIDescriptorAllocator::Allocate(NumDescriptors, OutSlot);
+}
+
+void FRHIHeapDescriptorAllocator::Free(uint32 Slot, uint32 NumDescriptors)
+{
+	FRHIDescriptorAllocator::Free(Slot, NumDescriptors);
+}
\ No newline at end of file
diff --git a/Engine/Source/Runtime/RHICore/Public/RHIDescriptorAllocator.h b/Engine/Source/Runtime/RHICore/Public/RHIDescriptorAllocator.h
new file mode 100644
--- /dev/null
+++ b/Engine/Source/Runtime/RHICore/Public/RHIDescriptorAllocator.h
@@ -0,0 +1,51 @@
+// Copyright Epic Games, Inc. All Rights Reserved.
+#pragma once
+
+#include "HAL/CriticalSection.h"
+#include "RHIDefinitions.h"
+
+struct FRHIDescriptorAllocatorRange;
+
+class RHICORE_API FRHIDescriptorAllocator
+{
+public:
+	FRHIDescriptorAllocator();
+	FRHIDescriptorAllocator(uint32 InNumDescriptors);
+	~FRHIDescriptorAllocator();
+
+	void Init(uint32 InNumDescriptors);
+	void Shutdown();
+
+	FRHIDescriptorHandle Allocate(ERHIDescriptorHeapType InType);
+	void Free(FRHIDescriptorHandle InHandle);
+
+	bool Allocate(uint32 NumDescriptors, uint32& OutSlot);
+	void Free(uint32 Slot, uint32 NumDescriptors);
+
+	inline uint32 GetCapacity() const { return Capacity; }
+
+private:
+	TArray<FRHIDescriptorAllocatorRange> Ranges;
+	uint32 Capacity = 0;
+
+	FCriticalSection CriticalSection;
+};
+
+class RHICORE_API FRHIHeapDescriptorAllocator : protected FRHIDescriptorAllocator
+{
+public:
+	FRHIHeapDescriptorAllocator() = delete;
+	FRHIHeapDescriptorAllocator(ERHIDescriptorHeapType InType, uint32 InDescriptorCount);
+
+	FRHIDescriptorHandle Allocate();
+	void Free(FRHIDescriptorHandle InHandle);
+
+	bool Allocate(uint32 NumDescriptors, uint32& OutSlot);
+	void Free(uint32 Slot, uint32 NumDescriptors);
+
+	using FRHIDescriptorAllocator::GetCapacity;
+	inline bool HandlesAllocation(ERHIDescriptorHeapType InType) const { return Type == InType; }
+
+private:
+	ERHIDescriptorHeapType Type;
+};
