MGPU:  Fix for assert when "RHICopyBufferRegion" is called, due to no implementation being provided in FD3D12CommandContextRedirector.

While investigating this issue, I cleaned up the code a bit, noticing that the implementation of "RHICopyBufferRegion" was incorrectly wrapped by D3D12_RHI_RAYTRACING (which I had first suspected as the cause of the bug), which could cause issues if someone tried to compile out ray tracing.  The function isn't wrapped by a ray tracing define on any other RHI layer implementation, and isn't ray tracing specific.  RHICopyBufferRegions (plural) is only implemented on D3D12, and only used by ray tracing at the moment, so I left that one where it is for now (I figured having it there reduces the likelihood that someone tries to use it generically, when it only has an implementation on D3D12).

I fixed compile errors that result when ray tracing is compiled out (to test my code rearrangement), but there was a link error for PrepareSkyTexture_Internal, called from FSkyLightRenderState::PrepareSkyTexture.  To fix that, someone would need to come up with a fallback sky light implementation when ray tracing is disabled, which goes beyond the scope of what makes sense to solve in this CL.

#rnx

diff --git a/Engine/Plugins/Experimental/GPULightmass/Source/GPULightmass/Private/LightmapRenderer.cpp b/Engine/Plugins/Experimental/GPULightmass/Source/GPULightmass/Private/LightmapRenderer.cpp
--- a/Engine/Plugins/Experimental/GPULightmass/Source/GPULightmass/Private/LightmapRenderer.cpp
+++ b/Engine/Plugins/Experimental/GPULightmass/Source/GPULightmass/Private/LightmapRenderer.cpp
@@ -730,7 +730,7 @@ bool FSceneRenderState::SetupRayTracingScene(int32 LODIndex)
 
 	FRHICommandListImmediate& RHICmdList = FRHICommandListExecutor::GetImmediateCommandList();
 
-#ifdef RHI_RAYTRACING
+#if RHI_RAYTRACING
 	// Force build all the open build requests
 	bool bBuildAll = true;
 	GRayTracingGeometryManager.ProcessBuildRequests(RHICmdList, bBuildAll);
diff --git a/Engine/Plugins/Experimental/GPULightmass/Source/GPULightmass/Private/PathTracingLightParameters.inl b/Engine/Plugins/Experimental/GPULightmass/Source/GPULightmass/Private/PathTracingLightParameters.inl
--- a/Engine/Plugins/Experimental/GPULightmass/Source/GPULightmass/Private/PathTracingLightParameters.inl
+++ b/Engine/Plugins/Experimental/GPULightmass/Source/GPULightmass/Private/PathTracingLightParameters.inl
@@ -5,6 +5,8 @@
 #include "RayTracingDefinitions.h"
 #include "PathTracingDefinitions.h"
 
+#if RHI_RAYTRACING
+
 RENDERER_API FRDGTexture* PrepareIESAtlas(const TMap<FTexture*, int>& InIESLightProfilesMap, FRDGBuilder& GraphBuilder);
 
 RENDERER_API void PrepareLightGrid(FRDGBuilder& GraphBuilder, const FViewInfo& View, FPathTracingLightGrid* LightGridParameters, const FPathTracingLight* Lights, uint32 NumLights, uint32 NumInfiniteLights, FRDGBufferSRV* LightsSRV);
@@ -228,3 +230,5 @@ void SetupPathTracingLightParameters(
 
 	PrepareLightGrid(GraphBuilder, View, &PassParameters->LightGridParameters, Lights.GetData(), PassParameters->SceneLightCount, NumInfiniteLights, PassParameters->SceneLights);
 }
+
+#endif  // RHI_RAYTRACING
diff --git a/Engine/Source/Runtime/D3D12RHI/Private/D3D12CommandContext.h b/Engine/Source/Runtime/D3D12RHI/Private/D3D12CommandContext.h
--- a/Engine/Source/Runtime/D3D12RHI/Private/D3D12CommandContext.h
+++ b/Engine/Source/Runtime/D3D12RHI/Private/D3D12CommandContext.h
@@ -290,6 +290,7 @@ public:
 	virtual void RHIClearUAVUint(FRHIUnorderedAccessView* UnorderedAccessViewRHI, const FUintVector4& Values) final override;
 	virtual void RHICopyToResolveTarget(FRHITexture* SourceTexture, FRHITexture* DestTexture, const FResolveParams& ResolveParams) override;
 	virtual void RHICopyTexture(FRHITexture* SourceTexture, FRHITexture* DestTexture, const FRHICopyTextureInfo& CopyInfo) final override;
+	virtual void RHICopyBufferRegion(FRHIBuffer* DestBuffer, uint64 DstOffset, FRHIBuffer* SourceBuffer, uint64 SrcOffset, uint64 NumBytes) final override;
 	virtual void RHICopyToStagingBuffer(FRHIBuffer* SourceBuffer, FRHIStagingBuffer* DestinationStagingBuffer, uint32 Offset, uint32 NumBytes) final override;
 	virtual void RHIWriteGPUFence(FRHIGPUFence* Fence) final override;
 	virtual void RHIBeginRenderQuery(FRHIRenderQuery* RenderQuery) final override;
@@ -384,7 +385,6 @@ public:
 	virtual void RHIBroadcastTemporalEffect(const FName& InEffectName, const TArrayView<FRHIBuffer*> InBuffers) final AFR_API_OVERRIDE;
 
 #if D3D12_RHI_RAYTRACING
-	virtual void RHICopyBufferRegion(FRHIBuffer* DestBuffer, uint64 DstOffset, FRHIBuffer* SourceBuffer, uint64 SrcOffset, uint64 NumBytes) final override;
 	virtual void RHICopyBufferRegions(const TArrayView<const FCopyBufferRegionParams> Params) final override;
 	virtual void RHIBindAccelerationStructureMemory(FRHIRayTracingScene* Scene, FRHIBuffer* Buffer, uint32 BufferOffset) final override;
 	virtual void BuildAccelerationStructuresInternal(const TArrayView<D3D12_BUILD_RAYTRACING_ACCELERATION_STRUCTURE_DESC> BuildDesc);
@@ -602,6 +602,10 @@ public:
 	{
 		ContextRedirect(RHICopyTexture(SourceTextureRHI, DestTextureRHI, CopyInfo));
 	}
+	FORCEINLINE virtual void RHICopyBufferRegion(FRHIBuffer* DestBuffer, uint64 DstOffset, FRHIBuffer* SourceBuffer, uint64 SrcOffset, uint64 NumBytes) final override
+	{
+		ContextRedirect(RHICopyBufferRegion(DestBuffer, DstOffset, SourceBuffer, SrcOffset, NumBytes));
+	}
 	FORCEINLINE virtual void RHIBeginRenderQuery(FRHIRenderQuery* RenderQuery) final override
 	{
 		ContextRedirect(RHIBeginRenderQuery(RenderQuery));
