Query ray tracing shader identifiers from individual ray tracing PSO collection sub-objects when possible.

Add GD3D12WorkaroundFlags.bAllowGetShaderIdentifierOnCollectionSubObject to control whether the new code path can run (disabled on known bad old driver versions).

Code path is disabled on NVIDIA driver versions < 457.0 (earlier versions may work, but only this one was tested).

diff --git a/Engine/Source/Runtime/D3D12RHI/Private/D3D12RHI.cpp b/Engine/Source/Runtime/D3D12RHI/Private/D3D12RHI.cpp
--- a/Engine/Source/Runtime/D3D12RHI/Private/D3D12RHI.cpp
+++ b/Engine/Source/Runtime/D3D12RHI/Private/D3D12RHI.cpp
@@ -46,6 +46,8 @@ FD3D12DynamicRHI* FD3D12DynamicRHI::SingleD3DRHI = nullptr;
 extern int32 GGapRecorderUseBlockingCall;
 #endif
 
+FD3D12WorkaroundFlags GD3D12WorkaroundFlags;
+
 using namespace D3D12RHI;
 
 FD3D12DynamicRHI::FD3D12DynamicRHI(const TArray<TSharedPtr<FD3D12Adapter>>& ChosenAdaptersIn, bool bInPixEventEnabled) :
diff --git a/Engine/Source/Runtime/D3D12RHI/Private/D3D12RHIPrivate.h b/Engine/Source/Runtime/D3D12RHI/Private/D3D12RHIPrivate.h
--- a/Engine/Source/Runtime/D3D12RHI/Private/D3D12RHIPrivate.h
+++ b/Engine/Source/Runtime/D3D12RHI/Private/D3D12RHIPrivate.h
@@ -224,6 +224,18 @@ struct FD3D12RayTracingPipelineInfo
 	uint32 ScratchSize = 0;
 };
 
+struct FD3D12WorkaroundFlags
+{
+	/** 
+	* Certain drivers crash when GetShaderIdentifier() is called on a ray tracing pipeline collection.
+	* If we detect such driver, we have to fall back to the path that queries identifiers on full linked RTPSO.
+	* This is less efficient and can also trigger another known issue with D3D12 Agility version <= 4.
+	*/
+	bool bAllowGetShaderIdentifierOnCollectionSubObject = true;
+};
+
+extern FD3D12WorkaroundFlags GD3D12WorkaroundFlags;
+
 /** Forward declare the context for the AMD AGS utility library. */
 struct AGSContext;
 
diff --git a/Engine/Source/Runtime/D3D12RHI/Private/D3D12RayTracing.cpp b/Engine/Source/Runtime/D3D12RHI/Private/D3D12RayTracing.cpp
--- a/Engine/Source/Runtime/D3D12RHI/Private/D3D12RayTracing.cpp
+++ b/Engine/Source/Runtime/D3D12RHI/Private/D3D12RayTracing.cpp
@@ -1081,8 +1081,11 @@ public:
 			}
 
 			// Shader identifier can be queried immediately here per PSO collection, however this does not work on old NVIDIA drivers (430.00).
-			// Therefore shader identifiers need to be queried from the final linked pipeline (JIRA DH-2182).
-			// Entry.Identifier = GetShaderIdentifier(Entry.StateObject, Entry.GetPrimaryExportNameChars());
+			// Therefore shader identifiers need to be queried from the final linked pipeline (JIRA DH-2182) if a known bad driver is detected.
+			if (GD3D12WorkaroundFlags.bAllowGetShaderIdentifierOnCollectionSubObject)
+			{
+				Entry.Identifier = GetShaderIdentifier(Entry.StateObject, Entry.GetPrimaryExportNameChars());
+			}
 
 			CompileTimeCycles += FPlatformTime::Cycles64();
 
@@ -2350,7 +2353,6 @@ static void CreateSpecializedStateObjects(
 	}
 }
 
-
 class FD3D12RayTracingPipelineState : public FRHIRayTracingPipelineState
 {
 public:
diff --git a/Engine/Source/Runtime/D3D12RHI/Private/Windows/WindowsD3D12Device.cpp b/Engine/Source/Runtime/D3D12RHI/Private/Windows/WindowsD3D12Device.cpp
--- a/Engine/Source/Runtime/D3D12RHI/Private/Windows/WindowsD3D12Device.cpp
+++ b/Engine/Source/Runtime/D3D12RHI/Private/Windows/WindowsD3D12Device.cpp
@@ -1004,22 +1004,28 @@ void FD3D12DynamicRHI::Init()
 				TEXT("Please update to NVIDIA driver version 466.11 or newer."));
 		}
 
-	if (GRHISupportsRayTracing
-		&& IsRayTracingEmulated(AdapterDesc.DeviceId))
+		if (GRHISupportsRayTracing
+			&& IsRayTracingEmulated(AdapterDesc.DeviceId))
 		{
 #if DXR_ALLOW_EMULATED_RAYTRACING
-		if (!GAllowEmulatedRayTracing)
-		{
-			DisableRayTracingSupport();
-			UE_LOG(LogD3D12RHI, Warning, TEXT("Ray tracing is disabled for NVIDIA cards with the Pascal architecture. This can be overridden with the following CVar: r.D3D12.DXR.AllowEmulatedRayTracing=1"));
-		}
+			if (!GAllowEmulatedRayTracing)
+			{
+				DisableRayTracingSupport();
+				UE_LOG(LogD3D12RHI, Warning, TEXT("Ray tracing is disabled for NVIDIA cards with the Pascal architecture. This can be overridden with the following CVar: r.D3D12.DXR.AllowEmulatedRayTracing=1"));
+			}
 #else
-		DisableRayTracingSupport();
-		UE_LOG(LogD3D12RHI, Warning, TEXT("Ray tracing is disabled for NVIDIA cards with the Pascal architecture."));
+			DisableRayTracingSupport();
+			UE_LOG(LogD3D12RHI, Warning, TEXT("Ray tracing is disabled for NVIDIA cards with the Pascal architecture."));
 #endif // DXR_ALLOW_EMULATED_RAYTRACING
 		}
-	}
-#endif
+
+		if (GRHISupportsRayTracing && DriverVersion < 45700u)
+		{
+			GD3D12WorkaroundFlags.bAllowGetShaderIdentifierOnCollectionSubObject = false;
+			UE_LOG(LogD3D12RHI, Warning, TEXT("GD3D12WorkaroundFlags.bAllowGetShaderIdentifierOnCollectionSubObject is disabled due to a known issue with current driver version."));
+		}
+	} // if NVIDIA
+#endif // NV_API_ENABLE
 
 #if AMD_API_ENABLE
 	if (GRHISupportsRayTracing
