Remove deprecated FRayTracingSceneInitializer codepath.

diff --git a/Engine/Source/Runtime/D3D12RHI/Private/D3D12RHIPrivate.h b/Engine/Source/Runtime/D3D12RHI/Private/D3D12RHIPrivate.h
--- a/Engine/Source/Runtime/D3D12RHI/Private/D3D12RHIPrivate.h
+++ b/Engine/Source/Runtime/D3D12RHI/Private/D3D12RHIPrivate.h
@@ -507,7 +507,6 @@ public:
 	virtual FRayTracingAccelerationStructureSize RHICalcRayTracingGeometrySize(const FRayTracingGeometryInitializer& Initializer) final override;
 
 	virtual FRayTracingGeometryRHIRef RHICreateRayTracingGeometry(const FRayTracingGeometryInitializer& Initializer) final override;
-	virtual FRayTracingSceneRHIRef RHICreateRayTracingScene(const FRayTracingSceneInitializer& Initializer) final override;
 	virtual FRayTracingSceneRHIRef RHICreateRayTracingScene(FRayTracingSceneInitializer2 Initializer) final override;
 	virtual FRayTracingShaderRHIRef RHICreateRayTracingShader(TArrayView<const uint8> Code, const FSHAHash& Hash, EShaderFrequency ShaderFrequency) final override;
 	virtual FRayTracingPipelineStateRHIRef RHICreateRayTracingPipelineState(const FRayTracingPipelineStateInitializer& Initializer) final override;
diff --git a/Engine/Source/Runtime/D3D12RHI/Private/D3D12RayTracing.cpp b/Engine/Source/Runtime/D3D12RHI/Private/D3D12RayTracing.cpp
--- a/Engine/Source/Runtime/D3D12RHI/Private/D3D12RayTracing.cpp
+++ b/Engine/Source/Runtime/D3D12RHI/Private/D3D12RayTracing.cpp
@@ -3237,160 +3237,13 @@ void FD3D12DynamicRHI::RHITransferRayTracingGeometryUnderlyingResource(FRHIRayTr
 	}
 }
 
-// This code path is deprecated and will be removed in the next release. Please use FRayTracingSceneInitializer2 instead.
-FRayTracingSceneRHIRef FD3D12DynamicRHI::RHICreateRayTracingScene(const FRayTracingSceneInitializer& Initializer)
-{
-	TRACE_CPUPROFILER_EVENT_SCOPE(CreateRayTracingScene);
-
-	const uint32 NumSceneInstances = Initializer.Instances.Num();
-
-	FRayTracingSceneInitializer2 Initializer2;
-	Initializer2.DebugName = Initializer.DebugName;
-	Initializer2.ShaderSlotsPerGeometrySegment = Initializer.ShaderSlotsPerGeometrySegment;
-	Initializer2.NumMissShaderSlots = Initializer.NumMissShaderSlots;
-	Initializer2.PerInstanceGeometries.SetNumUninitialized(NumSceneInstances);
-	Initializer2.BaseInstancePrefixSum.SetNumUninitialized(NumSceneInstances);
-	Initializer2.SegmentPrefixSum.SetNumUninitialized(NumSceneInstances);
-	Initializer2.NumNativeInstances = 0;
-	Initializer2.NumTotalSegments = 0;
-
-	TArray<uint32> PerInstanceNumTransforms;
-	PerInstanceNumTransforms.SetNumUninitialized(NumSceneInstances);
-
-	Experimental::TSherwoodSet<FRHIRayTracingGeometry*> UniqueGeometries;
-
-	for (uint32 InstanceIndex = 0; InstanceIndex < NumSceneInstances; ++InstanceIndex)
-	{
-		const FRayTracingGeometryInstance& InstanceDesc = Initializer.Instances[InstanceIndex];
-
-		if (InstanceDesc.GPUTransformsSRV || !InstanceDesc.InstanceSceneDataOffsets.IsEmpty())
-		{
-			static bool bLogged = false; // Only log once
-			if (!bLogged)
-			{
-				bLogged = true;
-				UE_LOG(LogRHI, Warning,
-					TEXT("GPUScene and GPUTransformsSRV instances are not supported in FRayTracingSceneInitializer code path.\n")
-					TEXT("Use FRayTracingSceneInitializer2 and BuildRayTracingInstanceBuffer instead."));
-			}
-		}
-		else
-		{
-			checkf(InstanceDesc.NumTransforms <= uint32(InstanceDesc.Transforms.Num()),
-				TEXT("Expected at most %d ray tracing geometry instance transforms, but got %d."),
-				InstanceDesc.NumTransforms, InstanceDesc.Transforms.Num());
-		}
-
-		checkf(InstanceDesc.GeometryRHI, TEXT("Ray tracing instance must have a valid geometry."));
-
-		Initializer2.PerInstanceGeometries[InstanceIndex] = InstanceDesc.GeometryRHI;
-
-		// Compute geometry segment count prefix sum to be later used in GetHitRecordBaseIndex()
-		Initializer2.SegmentPrefixSum[InstanceIndex] = Initializer2.NumTotalSegments;
-		Initializer2.NumTotalSegments += InstanceDesc.GeometryRHI->GetNumSegments();
-
-		bool bIsAlreadyInSet = false;
-		UniqueGeometries.Add(InstanceDesc.GeometryRHI, &bIsAlreadyInSet);
-		if (!bIsAlreadyInSet)
-		{
-			Initializer2.ReferencedGeometries.Add(InstanceDesc.GeometryRHI);
-		}
-
-		Initializer2.BaseInstancePrefixSum[InstanceIndex] = Initializer2.NumNativeInstances;
-		Initializer2.NumNativeInstances += InstanceDesc.NumTransforms;
-
-		PerInstanceNumTransforms[InstanceIndex] = InstanceDesc.NumTransforms;
-	}
-
-	TResourceArray<D3D12_RAYTRACING_INSTANCE_DESC, 16> NativeInstances;
-	NativeInstances.SetAllowCPUAccess(true);
-	NativeInstances.SetNumUninitialized(Initializer2.NumNativeInstances);
-
-	const EParallelForFlags ParallelForFlags = EParallelForFlags::None; // set ForceSingleThread for testing
-	ParallelFor(NumSceneInstances, [&Initializer2, Instances = Initializer.Instances, &NativeInstances](int32 InstanceIndex)
-		{
-			const FRayTracingGeometryInstance& Instance = Instances[InstanceIndex];
-			FD3D12RayTracingGeometry* Geometry = FD3D12DynamicRHI::ResourceCast(Initializer2.PerInstanceGeometries[InstanceIndex]);
-
-			D3D12_RAYTRACING_INSTANCE_DESC InstanceDesc = {};
-			check(InstanceDesc.AccelerationStructure == 0);
-
-			InstanceDesc.InstanceMask = Instance.Mask;
-			InstanceDesc.InstanceContributionToHitGroupIndex = Initializer2.SegmentPrefixSum[InstanceIndex] * Initializer2.ShaderSlotsPerGeometrySegment;
-			InstanceDesc.Flags = TranslateRayTracingInstanceFlags(Instance.Flags);
-
-			const uint32 NumTransforms = Instance.NumTransforms;
-
-			checkf(Instance.UserData.Num() == 0 || Instance.UserData.Num() >= int32(NumTransforms),
-				TEXT("User data array must be either be empty (Instance.DefaultUserData is used), or contain one entry per entry in Transforms array."));
-
-			const bool bUseUniqueUserData = Instance.UserData.Num() != 0;
-
-			uint32 DescIndex = Initializer2.BaseInstancePrefixSum[InstanceIndex];
-
-			int32 NumInactiveDxrInstancesThisSceneInstance = 0;
-
-			for (uint32 TransformIndex = 0; TransformIndex < NumTransforms; ++TransformIndex)
-			{
-				InstanceDesc.InstanceID = bUseUniqueUserData ? Instance.UserData[TransformIndex] : Instance.DefaultUserData;
-
-				// Set flag for deactivated instances
-				if (!Instance.ActivationMask.IsEmpty())
-				{
-					if ((Instance.ActivationMask[TransformIndex / 32] & (1 << (TransformIndex % 32))) == 0)
-					{
-						InstanceDesc.AccelerationStructure = 0xFFFFFFFFFFFFFFFF;
-						NumInactiveDxrInstancesThisSceneInstance++;
-					}
-					else
-					{
-						InstanceDesc.AccelerationStructure = 0;
-					}
-				}
-
-				if (TransformIndex < (uint32)Instance.Transforms.Num())
-				{
-					// DXR uses a 3x4 transform matrix in row major layout
-
-					const FMatrix& Transform = Instance.Transforms[TransformIndex];
-
-					InstanceDesc.Transform[0][0] = Transform.M[0][0];
-					InstanceDesc.Transform[0][1] = Transform.M[1][0];
-					InstanceDesc.Transform[0][2] = Transform.M[2][0];
-					InstanceDesc.Transform[0][3] = Transform.M[3][0];
-
-					InstanceDesc.Transform[1][0] = Transform.M[0][1];
-					InstanceDesc.Transform[1][1] = Transform.M[1][1];
-					InstanceDesc.Transform[1][2] = Transform.M[2][1];
-					InstanceDesc.Transform[1][3] = Transform.M[3][1];
-
-					InstanceDesc.Transform[2][0] = Transform.M[0][2];
-					InstanceDesc.Transform[2][1] = Transform.M[1][2];
-					InstanceDesc.Transform[2][2] = Transform.M[2][2];
-					InstanceDesc.Transform[2][3] = Transform.M[3][2];
-				}
-				else
-				{
-					FMemory::Memset(&InstanceDesc.Transform, 0, sizeof(InstanceDesc.Transform));
-				}
-
-				NativeInstances[DescIndex] = InstanceDesc;
-				++DescIndex;
-			}
-		}, ParallelForFlags);
-
-	FD3D12Adapter& Adapter = GetAdapter();
-
-	return new FD3D12RayTracingScene(&Adapter, MoveTemp(Initializer2), MoveTemp(NativeInstances), MoveTemp(PerInstanceNumTransforms));
-}
-
 FRayTracingSceneRHIRef FD3D12DynamicRHI::RHICreateRayTracingScene(FRayTracingSceneInitializer2 Initializer)
 {
 	TRACE_CPUPROFILER_EVENT_SCOPE(CreateRayTracingScene);
 
 	FD3D12Adapter& Adapter = GetAdapter();
 
-	return new FD3D12RayTracingScene(&Adapter, MoveTemp(Initializer), {}, {});
+	return new FD3D12RayTracingScene(&Adapter, MoveTemp(Initializer));
 }
 
 FBufferRHIRef FD3D12RayTracingGeometry::NullTransformBuffer;
@@ -4077,8 +3930,8 @@ void FD3D12RayTracingGeometry::CompactAccelerationStructure(FD3D12CommandContext
 	RegisterD3D12RayTracingGeometry(this);
 }
 
-FD3D12RayTracingScene::FD3D12RayTracingScene(FD3D12Adapter* Adapter, FRayTracingSceneInitializer2 InInitializer, TResourceArray<D3D12_RAYTRACING_INSTANCE_DESC, 16> InInstances, TArray<uint32> InPerInstanceNumTransforms)
-	: FD3D12AdapterChild(Adapter), Initializer(MoveTemp(InInitializer)), Instances(MoveTemp(InInstances)), PerInstanceNumTransforms(MoveTemp(InPerInstanceNumTransforms))
+FD3D12RayTracingScene::FD3D12RayTracingScene(FD3D12Adapter* Adapter, FRayTracingSceneInitializer2 InInitializer)
+	: FD3D12AdapterChild(Adapter), Initializer(MoveTemp(InInitializer))
 {
 	INC_DWORD_STAT(STAT_D3D12RayTracingAllocatedTLAS);
 
@@ -4152,6 +4005,8 @@ void FD3D12RayTracingScene::BuildAccelerationStructure(FD3D12CommandContext& Com
 	FD3D12Buffer* ScratchBuffer, uint32 ScratchBufferOffset,
 	FD3D12Buffer* InstanceBuffer, uint32 InstanceBufferOffset)
 {
+	check(InstanceBuffer != nullptr);
+
 	TRACE_CPUPROFILER_EVENT_SCOPE(BuildAccelerationStructure_TopLevel);
 	SCOPE_CYCLE_COUNTER(STAT_D3D12BuildTLAS);
 
@@ -4195,8 +4050,6 @@ void FD3D12RayTracingScene::BuildAccelerationStructure(FD3D12CommandContext& Com
 			PrebuildInfo.ScratchDataSizeInBytes,
 			BuildInputs.NumDescs, Initializer.NumNativeInstances, Initializer.PerInstanceGeometries.Num());
 
-	TRefCountPtr<FD3D12Buffer> AutoInstanceBuffer;
-
 	if (BuildInputs.NumDescs)
 	{
 		// - Set up acceleration structure pointers and make them resident.
@@ -4206,8 +4059,6 @@ void FD3D12RayTracingScene::BuildAccelerationStructure(FD3D12CommandContext& Com
 		HitGroupSystemParametersCache.Empty(Initializer.NumTotalSegments);
 
 		{
-			int32 NativeInstanceIndex = 0;
-
 			const int32 NumSceneInstances = Initializer.PerInstanceGeometries.Num();
 			for (int32 InstanceIndex = 0; InstanceIndex < NumSceneInstances; ++InstanceIndex)
 			{
@@ -4216,59 +4067,9 @@ void FD3D12RayTracingScene::BuildAccelerationStructure(FD3D12CommandContext& Com
 				// make a copy of system parameters to they can optimized fetch during SBT building (only done for GPU0)
 				check(Geometry->HitGroupSystemParameters[0].Num() > 0);
 				HitGroupSystemParametersCache.Append(Geometry->HitGroupSystemParameters[0]);
-
-				if (InstanceBuffer == nullptr)
-				{
-					D3D12_GPU_VIRTUAL_ADDRESS BlasAddress = Geometry->AccelerationStructureBuffers[GPUIndex]->ResourceLocation.GetGPUVirtualAddress();
-
-					const int32 NumTransforms = PerInstanceNumTransforms[InstanceIndex];
-					for (int32 TransformIndex = 0; TransformIndex < NumTransforms; ++TransformIndex)
-					{
-						// If this instance has been deactivated, set BLAS address to 0
-						if (Instances[NativeInstanceIndex].AccelerationStructure == 0xFFFFFFFFFFFFFFFF)
-						{
-							Instances[NativeInstanceIndex].AccelerationStructure = 0;
-						}
-						else
-						{
-							Instances[NativeInstanceIndex].AccelerationStructure = BlasAddress;
-						}
-
-						NativeInstanceIndex++;
-					}
-				}
 			}
 		}
 
-		if (InstanceBuffer == nullptr)
-		{
-			check(Instances.Num() == BuildInputs.NumDescs);
-
-			D3D12_RESOURCE_DESC InstanceBufferDesc = CD3DX12_RESOURCE_DESC::Buffer(
-				sizeof(D3D12_RAYTRACING_INSTANCE_DESC) * BuildInputs.NumDescs,
-				D3D12_RESOURCE_FLAG_ALLOW_UNORDERED_ACCESS);
-
-			ED3D12ResourceTransientMode TransientMode = ED3D12ResourceTransientMode::NonTransient;
-			ID3D12ResourceAllocator* ResourceAllocator = nullptr;
-			FRHIGPUMask GPUMask = FRHIGPUMask::FromIndex(GPUIndex);
-			bool bHasInitialData = true;
-
-			AutoInstanceBuffer = Adapter->CreateRHIBuffer(
-				InstanceBufferDesc, D3D12_RAYTRACING_INSTANCE_DESCS_BYTE_ALIGNMENT,
-				sizeof(D3D12_RAYTRACING_INSTANCE_DESC), InstanceBufferDesc.Width, BUF_UnorderedAccess, ED3D12ResourceStateMode::MultiState, D3D12_RESOURCE_STATE_COPY_DEST,
-				bHasInitialData, GPUMask, TransientMode, ResourceAllocator, TEXT("InstanceBuffer"));
-
-			// Use copy queue for uploading the data
-			FD3D12SyncPoint CopyQueueSyncPoint = AutoInstanceBuffer->UploadResourceDataViaCopyQueue(&Instances);
-			CommandContext.CopyQueueSyncPoint.Merge(CopyQueueSyncPoint);
-
-			InstanceBuffer = AutoInstanceBuffer.GetReference();
-			InstanceBufferOffset = 0;
-
-			FD3D12DynamicRHI::TransitionResource(CommandContext.CommandListHandle, InstanceBuffer->GetResource(),
-				D3D12_RESOURCE_STATE_COPY_DEST, D3D12_RESOURCE_STATE_NON_PIXEL_SHADER_RESOURCE, 0, FD3D12DynamicRHI::ETransitionMode::Apply);
-		}
-
 		InstanceBuffer->GetResource()->UpdateResidency(CommandContext.CommandListHandle);
 
 		{
@@ -4342,7 +4143,7 @@ void FD3D12RayTracingScene::BuildAccelerationStructure(FD3D12CommandContext& Com
 
 	D3D12_BUILD_RAYTRACING_ACCELERATION_STRUCTURE_DESC BuildDesc = {};
 	BuildDesc.Inputs = BuildInputs;
-	BuildDesc.Inputs.InstanceDescs = InstanceBuffer ? InstanceBuffer->ResourceLocation.GetGPUVirtualAddress() + InstanceBufferOffset : D3D12_GPU_VIRTUAL_ADDRESS(0);
+	BuildDesc.Inputs.InstanceDescs = InstanceBuffer->ResourceLocation.GetGPUVirtualAddress() + InstanceBufferOffset;
 	BuildDesc.DestAccelerationStructureData = AccelerationStructureBuffer->ResourceLocation.GetGPUVirtualAddress() + BufferOffset;
 	BuildDesc.ScratchAccelerationStructureData = ScratchAddress;
 	BuildDesc.SourceAccelerationStructureData = D3D12_GPU_VIRTUAL_ADDRESS(0); // Null source TLAS as this is a build command
diff --git a/Engine/Source/Runtime/D3D12RHI/Private/D3D12RayTracing.h b/Engine/Source/Runtime/D3D12RHI/Private/D3D12RayTracing.h
--- a/Engine/Source/Runtime/D3D12RHI/Private/D3D12RayTracing.h
+++ b/Engine/Source/Runtime/D3D12RHI/Private/D3D12RayTracing.h
@@ -98,7 +98,7 @@ public:
 	// Scaling beyond 5 total threads does not yield any speedup in practice.
 	static constexpr uint32 MaxBindingWorkers = 5; // RHI thread + 4 parallel workers.
 
-	FD3D12RayTracingScene(FD3D12Adapter* Adapter, FRayTracingSceneInitializer2 Initializer, TResourceArray<D3D12_RAYTRACING_INSTANCE_DESC, 16> Instances, TArray<uint32> PerInstanceNumTransforms);
+	FD3D12RayTracingScene(FD3D12Adapter* Adapter, FRayTracingSceneInitializer2 Initializer);
 	~FD3D12RayTracingScene();
 
 	const FRayTracingSceneInitializer2& GetInitializer() const override final { return Initializer; }
@@ -122,9 +122,6 @@ public:
 
 	const FRayTracingSceneInitializer2 Initializer;
 
-	TResourceArray<D3D12_RAYTRACING_INSTANCE_DESC, 16> Instances;
-	TArray<uint32> PerInstanceNumTransforms;
-
 	// Scene keeps track of child acceleration structure buffers to ensure
 	// they are resident when any ray tracing work is dispatched.
 	TArray<FD3D12ResidencyHandle*> GeometryResidencyHandles[MAX_NUM_GPUS];
diff --git a/Engine/Source/Runtime/RHI/Public/DynamicRHI.h b/Engine/Source/Runtime/RHI/Public/DynamicRHI.h
--- a/Engine/Source/Runtime/RHI/Public/DynamicRHI.h
+++ b/Engine/Source/Runtime/RHI/Public/DynamicRHI.h
@@ -1289,12 +1289,6 @@ public:
 		checkNoEntry();
 		return nullptr;
 	}
-
-	virtual FRayTracingSceneRHIRef RHICreateRayTracingScene(const FRayTracingSceneInitializer& Initializer)
-	{
-		checkNoEntry();
-		return nullptr;
-	}
 	
 	virtual FRayTracingSceneRHIRef RHICreateRayTracingScene(FRayTracingSceneInitializer2 Initializer)
 	{
@@ -1621,12 +1615,6 @@ FORCEINLINE FRayTracingGeometryRHIRef RHICreateRayTracingGeometry(const FRayTrac
 	return GDynamicRHI->RHICreateRayTracingGeometry(Initializer);
 }
 
-UE_DEPRECATED(5.0, "This version of RHICreateRayTracingScene is deprecated. FRayTracingSceneInitializer2 should be used instead.")
-FORCEINLINE FRayTracingSceneRHIRef RHICreateRayTracingScene(const FRayTracingSceneInitializer& Initializer)
-{
-	return GDynamicRHI->RHICreateRayTracingScene(Initializer);
-}
-
 FORCEINLINE FRayTracingSceneRHIRef RHICreateRayTracingScene(FRayTracingSceneInitializer2 Initializer)
 {
 	return GDynamicRHI->RHICreateRayTracingScene(MoveTemp(Initializer));
diff --git a/Engine/Source/Runtime/RHI/Public/RHIResources.h b/Engine/Source/Runtime/RHI/Public/RHIResources.h
--- a/Engine/Source/Runtime/RHI/Public/RHIResources.h
+++ b/Engine/Source/Runtime/RHI/Public/RHIResources.h
@@ -2077,33 +2077,6 @@ enum class ERayTracingAccelerationStructureFlags
 };
 ENUM_CLASS_FLAGS(ERayTracingAccelerationStructureFlags);
 
-struct FRayTracingSceneInitializer
-{
-	TArrayView<FRayTracingGeometryInstance> Instances;
-
-	// This value controls how many elements will be allocated in the shader binding table per geometry segment.
-	// Changing this value allows different hit shaders to be used for different effects.
-	// For example, setting this to 2 allows one hit shader for regular material evaluation and a different one for shadows.
-	// Desired hit shader can be selected by providing appropriate RayContributionToHitGroupIndex to TraceRay() function.
-	// Use ShaderSlot argument in SetRayTracingHitGroup() to assign shaders and resources for specific part of the shder binding table record.
-	uint32 ShaderSlotsPerGeometrySegment = 1;
-
-	// Defines how many different callable shaders with unique resource bindings can be bound to this scene.
-	// Shaders and resources are assigned to slots in the scene using SetRayTracingCallableShader().
-	uint32 NumCallableShaderSlots = 0;
-
-	// At least one miss shader must be present in a ray tracing scene.
-	// Default miss shader is always in slot 0. Default shader must not use local resources.
-	// Custom miss shaders can be bound to other slots using SetRayTracingMissShader().
-	uint32 NumMissShaderSlots = 1;
-
-	// Defines whether data in this scene should persist between frames.
-	// Currently only single-frame lifetime is supported.
-	ERayTracingSceneLifetime Lifetime = RTSL_SingleFrame;
-
-	FName DebugName;
-};
-
 struct FRayTracingSceneInitializer2
 {
 	// Unique list of geometries referenced by all instances in this scene.
diff --git a/Engine/Source/Runtime/RHI/Public/RHIValidation.h b/Engine/Source/Runtime/RHI/Public/RHIValidation.h
--- a/Engine/Source/Runtime/RHI/Public/RHIValidation.h
+++ b/Engine/Source/Runtime/RHI/Public/RHIValidation.h
@@ -1689,13 +1689,6 @@ public:
 		return Result;
 	}
 
-	virtual FRayTracingSceneRHIRef RHICreateRayTracingScene(const FRayTracingSceneInitializer& Initializer) override final
-	{
-		FRayTracingSceneRHIRef Result = RHI->RHICreateRayTracingScene(Initializer);
-		Result->InitBarrierTracking(ERHIAccess::BVHWrite, *Initializer.DebugName.ToString()); // BVHs are always created in BVHWrite state
-		return Result;
-	}
-
 	virtual FRayTracingSceneRHIRef RHICreateRayTracingScene(FRayTracingSceneInitializer2 Initializer) override final
 	{
 		FName DebugName = Initializer.DebugName;
diff --git a/Engine/Source/Runtime/VulkanRHI/Private/VulkanRayTracing.cpp b/Engine/Source/Runtime/VulkanRHI/Private/VulkanRayTracing.cpp
--- a/Engine/Source/Runtime/VulkanRHI/Private/VulkanRayTracing.cpp
+++ b/Engine/Source/Runtime/VulkanRHI/Private/VulkanRayTracing.cpp
@@ -486,8 +486,8 @@ struct FVulkanRayTracingGeometryParameters
 	uint64 VertexBuffer = 0;
 };
 
-FVulkanRayTracingScene::FVulkanRayTracingScene(FRayTracingSceneInitializer2 InInitializer, FVulkanDevice* InDevice, FVulkanResourceMultiBuffer* InInstanceBuffer)
-	: Device(InDevice), Initializer(MoveTemp(InInitializer)), InstanceBuffer(InInstanceBuffer)
+FVulkanRayTracingScene::FVulkanRayTracingScene(FRayTracingSceneInitializer2 InInitializer, FVulkanDevice* InDevice)
+	: Device(InDevice), Initializer(MoveTemp(InInitializer))
 {
 	const ERayTracingAccelerationStructureFlags BuildFlags = ERayTracingAccelerationStructureFlags::FastTrace; // #yuriy_todo: pass this in
 	SizeInfo = RHICalcRayTracingSceneSize(Initializer.NumNativeInstances, BuildFlags);
@@ -525,19 +525,10 @@ void FVulkanRayTracingScene::BuildAccelerationStructure(
 	FVulkanResourceMultiBuffer* InInstanceBuffer, uint32 InInstanceOffset)
 {
 	check(AccelerationStructureBuffer.IsValid());
+	check(InInstanceBuffer != nullptr);
 	const bool bExternalScratchBuffer = InScratchBuffer != nullptr;
 
-	VkDeviceAddress InstanceBufferAddress = 0;
-
-	if (InInstanceBuffer != nullptr)
-	{
-		checkf(InstanceBuffer == nullptr, TEXT("High level instance buffer is only supported when using FRayTracingSceneInitializer2."));
-		InstanceBufferAddress = InInstanceBuffer->GetDeviceAddress() + InInstanceOffset;
-	}
-	else
-	{
-		InstanceBufferAddress = InstanceBuffer->GetDeviceAddress();
-	}
+	VkDeviceAddress InstanceBufferAddress = InInstanceBuffer->GetDeviceAddress() + InInstanceOffset;
 
 	// Build a metadata buffer	that contains VulkanRHI-specific per-geometry parameters that allow us to access
 	// vertex and index buffers from shaders that use inline ray tracing.
@@ -578,8 +569,6 @@ void FVulkanRayTracingScene::BuildAccelerationStructure(
 
 	CommandBufferManager.SubmitActiveCmdBuffer();
 	CommandBufferManager.PrepareForNewActiveCommandBuffer();
-
-	InstanceBuffer = nullptr;
 }
 
 void FVulkanRayTracingScene::BuildPerInstanceGeometryParameterBuffer()
@@ -694,154 +683,9 @@ FRayTracingAccelerationStructureSize FVulkanDynamicRHI::RHICalcRayTracingGeometr
 	return Result;
 }
 
-FRayTracingSceneRHIRef FVulkanDynamicRHI::RHICreateRayTracingScene(const FRayTracingSceneInitializer& Initializer)
-{
-	TRACE_CPUPROFILER_EVENT_SCOPE(CreateRayTracingScene);
-
-	const uint32 NumSceneInstances = Initializer.Instances.Num();
-
-	FRayTracingSceneInitializer2 Initializer2;
-	Initializer2.DebugName = Initializer.DebugName;
-	Initializer2.ShaderSlotsPerGeometrySegment = Initializer.ShaderSlotsPerGeometrySegment;
-	Initializer2.NumMissShaderSlots = Initializer.NumMissShaderSlots;
-	Initializer2.PerInstanceGeometries.SetNumUninitialized(NumSceneInstances);
-	Initializer2.BaseInstancePrefixSum.SetNumUninitialized(NumSceneInstances);
-	Initializer2.SegmentPrefixSum.SetNumUninitialized(NumSceneInstances);
-	Initializer2.NumNativeInstances = 0;
-	Initializer2.NumTotalSegments = 0;
-
-	TArray<uint32> PerInstanceNumTransforms;
-	PerInstanceNumTransforms.SetNumUninitialized(NumSceneInstances);
-
-	Experimental::TSherwoodSet<FRHIRayTracingGeometry*> UniqueGeometries;
-
-	for (uint32 InstanceIndex = 0; InstanceIndex < NumSceneInstances; ++InstanceIndex)
-	{
-		const FRayTracingGeometryInstance& InstanceDesc = Initializer.Instances[InstanceIndex];
-
-		if (InstanceDesc.GPUTransformsSRV || !InstanceDesc.InstanceSceneDataOffsets.IsEmpty())
-		{
-			static bool bLogged = false; // Only log once
-			if (!bLogged)
-			{
-				bLogged = true;
-				UE_LOG(LogRHI, Warning,
-					TEXT("GPUScene and GPUTransformsSRV instances are not supported in FRayTracingSceneInitializer code path.\n")
-					TEXT("Use FRayTracingSceneInitializer2 and BuildRayTracingInstanceBuffer instead."));
-			}
-		}
-		else
-		{
-			checkf(InstanceDesc.NumTransforms <= uint32(InstanceDesc.Transforms.Num()),
-				TEXT("Expected at most %d ray tracing geometry instance transforms, but got %d."),
-				InstanceDesc.NumTransforms, InstanceDesc.Transforms.Num());
-		}
-
-		checkf(InstanceDesc.GeometryRHI, TEXT("Ray tracing instance must have a valid geometry."));
-
-		Initializer2.PerInstanceGeometries[InstanceIndex] = InstanceDesc.GeometryRHI;
-
-		// Compute geometry segment count prefix sum to be later used in GetHitRecordBaseIndex()
-		Initializer2.SegmentPrefixSum[InstanceIndex] = Initializer2.NumTotalSegments;
-		Initializer2.NumTotalSegments += InstanceDesc.GeometryRHI->GetNumSegments();
-
-		bool bIsAlreadyInSet = false;
-		UniqueGeometries.Add(InstanceDesc.GeometryRHI, &bIsAlreadyInSet);
-		if (!bIsAlreadyInSet)
-		{
-			Initializer2.ReferencedGeometries.Add(InstanceDesc.GeometryRHI);
-		}
-
-		Initializer2.BaseInstancePrefixSum[InstanceIndex] = Initializer2.NumNativeInstances;
-		Initializer2.NumNativeInstances += InstanceDesc.NumTransforms;
-
-		PerInstanceNumTransforms[InstanceIndex] = InstanceDesc.NumTransforms;
-	}
-
-	TArray<VkAccelerationStructureInstanceKHR> NativeInstances;
-	NativeInstances.SetNumUninitialized(Initializer2.NumNativeInstances);
-
-	const EParallelForFlags ParallelForFlags = EParallelForFlags::None; // set ForceSingleThread for testing
-	ParallelFor(NumSceneInstances, [&Initializer2, Instances = Initializer.Instances, &NativeInstances](int32 InstanceIndex)
-	{
-		const FRayTracingGeometryInstance& RayTracingGeometryInstance = Instances[InstanceIndex];
-		FVulkanRayTracingGeometry* Geometry = ResourceCast(Initializer2.PerInstanceGeometries[InstanceIndex]);
-
-		const FRayTracingAccelerationStructureAddress AccelerationStructureAddress = Geometry->GetAccelerationStructureAddress(0);
-		check(AccelerationStructureAddress != 0);
-
-		VkAccelerationStructureInstanceKHR InstanceDesc = {};
-		InstanceDesc.mask = RayTracingGeometryInstance.Mask;
-		InstanceDesc.instanceShaderBindingTableRecordOffset = Initializer2.SegmentPrefixSum[InstanceIndex] * Initializer2.ShaderSlotsPerGeometrySegment; // TODO?
-		InstanceDesc.flags = TranslateRayTracingInstanceFlags(RayTracingGeometryInstance.Flags);
-
-		const uint32 NumTransforms = RayTracingGeometryInstance.NumTransforms;
-
-		checkf(RayTracingGeometryInstance.UserData.Num() == 0 || RayTracingGeometryInstance.UserData.Num() >= int32(NumTransforms),
-			TEXT("User data array must be either be empty (RayTracingGeometryInstance.DefaultUserData is used), or contain one entry per entry in Transforms array."));
-
-		const bool bUseUniqueUserData = RayTracingGeometryInstance.UserData.Num() != 0;
-
-		uint32 DescIndex = Initializer2.BaseInstancePrefixSum[InstanceIndex];
-
-		for (uint32 TransformIndex = 0; TransformIndex < NumTransforms; ++TransformIndex)
-		{
-			InstanceDesc.instanceCustomIndex = bUseUniqueUserData ? RayTracingGeometryInstance.UserData[TransformIndex] : RayTracingGeometryInstance.DefaultUserData;
-
-			InstanceDesc.accelerationStructureReference = AccelerationStructureAddress;
-
-			if (!RayTracingGeometryInstance.ActivationMask.IsEmpty() && (RayTracingGeometryInstance.ActivationMask[TransformIndex / 32] & (1 << (TransformIndex % 32))) == 0)
-			{
-				InstanceDesc.accelerationStructureReference = 0;
-			}
-
-			if (TransformIndex < (uint32)RayTracingGeometryInstance.Transforms.Num())
-			{
-				const FMatrix& Transform = RayTracingGeometryInstance.Transforms[TransformIndex];
-
-				InstanceDesc.transform.matrix[0][0] = Transform.M[0][0];
-				InstanceDesc.transform.matrix[0][1] = Transform.M[1][0];
-				InstanceDesc.transform.matrix[0][2] = Transform.M[2][0];
-				InstanceDesc.transform.matrix[0][3] = Transform.M[3][0];
-
-				InstanceDesc.transform.matrix[1][0] = Transform.M[0][1];
-				InstanceDesc.transform.matrix[1][1] = Transform.M[1][1];
-				InstanceDesc.transform.matrix[1][2] = Transform.M[2][1];
-				InstanceDesc.transform.matrix[1][3] = Transform.M[3][1];
-
-				InstanceDesc.transform.matrix[2][0] = Transform.M[0][2];
-				InstanceDesc.transform.matrix[2][1] = Transform.M[1][2];
-				InstanceDesc.transform.matrix[2][2] = Transform.M[2][2];
-				InstanceDesc.transform.matrix[2][3] = Transform.M[3][2];
-			}
-			else
-			{
-				FMemory::Memset(&InstanceDesc.transform, 0, sizeof(InstanceDesc.transform));
-			}
-
-			NativeInstances[DescIndex] = InstanceDesc;
-
-			++DescIndex;
-		}
-	}, ParallelForFlags);
-
-	// Allocate instance buffer
-	// TODO: VK_BUFFER_USAGE_ACCELERATION_STRUCTURE_BUILD_INPUT_READ_ONLY_BIT_KHR
-	const uint32 InstanceBufferByteSize = static_cast<uint32>(Initializer2.NumNativeInstances) * sizeof(VkAccelerationStructureInstanceKHR);
-	FRHIResourceCreateInfo TempRTInstanceBufferCreateInfo(TEXT("TempRTInstanceBuffer"));
-	TRefCountPtr<FVulkanResourceMultiBuffer> InstanceUploadBuffer = ResourceCast(RHICreateBuffer(InstanceBufferByteSize, BUF_Volatile, 0, ERHIAccess::SRVCompute, TempRTInstanceBufferCreateInfo).GetReference());
-
-	// Copy instance data
-	void* pMappedInstanceBufferMemory = ::RHILockBuffer(InstanceUploadBuffer, 0, InstanceBufferByteSize, RLM_WriteOnly);
-	FMemory::Memcpy(pMappedInstanceBufferMemory, NativeInstances.GetData(), InstanceBufferByteSize);
-	::RHIUnlockBuffer(InstanceUploadBuffer);
-
-	return new FVulkanRayTracingScene(MoveTemp(Initializer2), GetDevice(), InstanceUploadBuffer.GetReference());
-}
-
 FRayTracingSceneRHIRef FVulkanDynamicRHI::RHICreateRayTracingScene(FRayTracingSceneInitializer2 Initializer)
 {
-	return new FVulkanRayTracingScene(MoveTemp(Initializer), GetDevice(), {});
+	return new FVulkanRayTracingScene(MoveTemp(Initializer), GetDevice());
 }
 
 FRayTracingGeometryRHIRef FVulkanDynamicRHI::RHICreateRayTracingGeometry(const FRayTracingGeometryInitializer& Initializer)
diff --git a/Engine/Source/Runtime/VulkanRHI/Private/VulkanRayTracing.h b/Engine/Source/Runtime/VulkanRHI/Private/VulkanRayTracing.h
--- a/Engine/Source/Runtime/VulkanRHI/Private/VulkanRayTracing.h
+++ b/Engine/Source/Runtime/VulkanRHI/Private/VulkanRayTracing.h
@@ -102,7 +102,7 @@ private:
 class FVulkanRayTracingScene : public FRHIRayTracingScene
 {
 public:
-	FVulkanRayTracingScene(FRayTracingSceneInitializer2 Initializer, FVulkanDevice* InDevice, FVulkanResourceMultiBuffer* InInstanceBuffer);
+	FVulkanRayTracingScene(FRayTracingSceneInitializer2 Initializer, FVulkanDevice* InDevice);
 	~FVulkanRayTracingScene();
 
 	const FRayTracingSceneInitializer2& GetInitializer() const override final { return Initializer; }
@@ -125,8 +125,6 @@ private:
 
 	const FRayTracingSceneInitializer2 Initializer;
 
-	TRefCountPtr<FVulkanResourceMultiBuffer> InstanceBuffer;
-
 	// Native TLAS handles are owned by SRV objects in Vulkan RHI.
 	// D3D12 and other RHIs allow creating TLAS SRVs from any GPU address at any point
 	// and do not require them for operations such as build or update.
diff --git a/Engine/Source/Runtime/VulkanRHI/Public/VulkanDynamicRHI.h b/Engine/Source/Runtime/VulkanRHI/Public/VulkanDynamicRHI.h
--- a/Engine/Source/Runtime/VulkanRHI/Public/VulkanDynamicRHI.h
+++ b/Engine/Source/Runtime/VulkanRHI/Public/VulkanDynamicRHI.h
@@ -322,7 +322,6 @@ public:
 	virtual FRayTracingAccelerationStructureSize RHICalcRayTracingSceneSize(uint32 MaxInstances, ERayTracingAccelerationStructureFlags Flags) final override;
 	virtual FRayTracingAccelerationStructureSize RHICalcRayTracingGeometrySize(const FRayTracingGeometryInitializer& Initializer) final override;
 	virtual FRayTracingGeometryRHIRef RHICreateRayTracingGeometry(const FRayTracingGeometryInitializer& Initializer) final override;
-	virtual FRayTracingSceneRHIRef RHICreateRayTracingScene(const FRayTracingSceneInitializer& Initializer) final override;
 	virtual FRayTracingSceneRHIRef RHICreateRayTracingScene(FRayTracingSceneInitializer2 Initializer) final override;
 	virtual FRayTracingShaderRHIRef RHICreateRayTracingShader(TArrayView<const uint8> Code, const FSHAHash& Hash, EShaderFrequency ShaderFrequency) final override;
 	virtual void RHITransferRayTracingGeometryUnderlyingResource(FRHIRayTracingGeometry* DestGeometry, FRHIRayTracingGeometry* SrcGeometry) final override;
