Cleaning up D3D12 SRV and UAV implementations.
Primarily renaming the "Initialize" methods to be more verbose about when they're called in relation to when the view is created. This allows us to identify when a CreateView call is for a newly-allocated descriptor or one that's been around for a while.
Stubs added for deferred updates to bindless descriptors, not implemented yet.
Unified a few RHI commands in RHICreateShaderResourceView to use a single command.
Removed FD3D12ShaderResourceViewWithLocation and FD3D12UnorderedAccessViewWithLocation since they're no longer used.

diff --git a/Engine/Source/Runtime/D3D12RHI/Private/D3D12DescriptorCache.cpp b/Engine/Source/Runtime/D3D12RHI/Private/D3D12DescriptorCache.cpp
--- a/Engine/Source/Runtime/D3D12RHI/Private/D3D12DescriptorCache.cpp
+++ b/Engine/Source/Runtime/D3D12RHI/Private/D3D12DescriptorCache.cpp
@@ -97,7 +97,7 @@ void FD3D12DescriptorCache::Init(FD3D12Device* InParent, FD3D12CommandContext* I
 	SRVDesc.Texture2D.MostDetailedMip = 0;
 	SRVDesc.Texture2D.ResourceMinLODClamp = 0.0f;
 	NullSRV = new FD3D12ViewDescriptorHandle(GetParentDevice(), ERHIDescriptorHeapType::Standard);
-	NullSRV->CreateView(SRVDesc, nullptr);
+	NullSRV->CreateView(SRVDesc, nullptr, ED3D12DescriptorCreateReason::InitialCreate);
 
 	D3D12_RENDER_TARGET_VIEW_DESC RTVDesc = {};
 	RTVDesc.ViewDimension = D3D12_RTV_DIMENSION_TEXTURE2D;
@@ -111,7 +111,7 @@ void FD3D12DescriptorCache::Init(FD3D12Device* InParent, FD3D12CommandContext* I
 	UAVDesc.Format = DXGI_FORMAT_R8G8B8A8_UNORM;
 	UAVDesc.Texture2D.MipSlice = 0;
 	NullUAV = new FD3D12ViewDescriptorHandle(GetParentDevice(), ERHIDescriptorHeapType::Standard);
-	NullUAV->CreateView(UAVDesc, nullptr, nullptr);
+	NullUAV->CreateView(UAVDesc, nullptr, nullptr, ED3D12DescriptorCreateReason::InitialCreate);
 
 #if USE_STATIC_ROOT_SIGNATURE
 	NullCBV = new FD3D12ConstantBufferView(GetParentDevice());
diff --git a/Engine/Source/Runtime/D3D12RHI/Private/D3D12Descriptors.cpp b/Engine/Source/Runtime/D3D12RHI/Private/D3D12Descriptors.cpp
--- a/Engine/Source/Runtime/D3D12RHI/Private/D3D12Descriptors.cpp
+++ b/Engine/Source/Runtime/D3D12RHI/Private/D3D12Descriptors.cpp
@@ -86,6 +86,17 @@ void FD3D12DescriptorManager::Destroy()
 	}
 }
 
+void FD3D12DescriptorManager::UpdateImmediately(FRHIDescriptorHandle InHandle, D3D12_CPU_DESCRIPTOR_HANDLE InSourceCpuHandle)
+{
+	if (InHandle.IsValid())
+	{
+		const D3D12_CPU_DESCRIPTOR_HANDLE DestCpuHandle = Heap->GetCPUSlotHandle(InHandle.GetIndex());
+		const D3D12_DESCRIPTOR_HEAP_TYPE D3DHeapType = Translate(GetType());
+
+		GetParentDevice()->GetDevice()->CopyDescriptorsSimple(1, DestCpuHandle, InSourceCpuHandle, D3DHeapType);
+	}
+}
+
 ////////////////////////////////////////////////////////////////////////////////////////////////////////////////////////
 // FD3D12BindlessDescriptorManager
 
@@ -148,20 +159,25 @@ void FD3D12BindlessDescriptorManager::FreeDescriptor(FRHIDescriptorHandle InHand
 	checkNoEntry();
 }
 
-D3D12_CPU_DESCRIPTOR_HANDLE FD3D12BindlessDescriptorManager::GetCpuDescriptorHandle(FRHIDescriptorHandle InHandle) const
+void FD3D12BindlessDescriptorManager::UpdateImmediately(FRHIDescriptorHandle InHandle, D3D12_CPU_DESCRIPTOR_HANDLE InSourceCpuHandle)
 {
-	for (const FD3D12DescriptorManager& Manager : Managers)
+	for (FD3D12DescriptorManager& Manager : Managers)
 	{
 		if (Manager.HandlesAllocation(InHandle.GetType()))
 		{
-			return Manager.GetCpuDescriptorHandle(InHandle.GetIndex());
+			Manager.UpdateImmediately(InHandle, InSourceCpuHandle);
+			return;
 		}
 	}
 
+	// Bad configuration?
 	checkNoEntry();
+}
 
-	D3D12_CPU_DESCRIPTOR_HANDLE NullHandle{};
-	return NullHandle;
+void FD3D12BindlessDescriptorManager::UpdateDeferred(FRHIDescriptorHandle InHandle, D3D12_CPU_DESCRIPTOR_HANDLE InSourceCpuHandle)
+{
+	// TODO: implement deferred updates
+	UpdateImmediately(InHandle, InSourceCpuHandle);
 }
 
 ////////////////////////////////////////////////////////////////////////////////////////////////////////////////////////
diff --git a/Engine/Source/Runtime/D3D12RHI/Private/D3D12Descriptors.h b/Engine/Source/Runtime/D3D12RHI/Private/D3D12Descriptors.h
--- a/Engine/Source/Runtime/D3D12RHI/Private/D3D12Descriptors.h
+++ b/Engine/Source/Runtime/D3D12RHI/Private/D3D12Descriptors.h
@@ -72,8 +72,9 @@ public:
 	void Init(TCHAR* InName, ERHIDescriptorHeapType InType, uint32 InNumDescriptors);
 	void Destroy();
 
+	void UpdateImmediately(FRHIDescriptorHandle InHandle, D3D12_CPU_DESCRIPTOR_HANDLE InSourceCpuHandle);
+
 	inline FD3D12DescriptorHeap* GetHeap() { return Heap.GetReference(); }
-	inline D3D12_CPU_DESCRIPTOR_HANDLE GetCpuDescriptorHandle(uint32 InSlot) const { return Heap->GetCPUSlotHandle(InSlot); }
 
 	inline bool HandlesAllocationWithFlags(ERHIDescriptorHeapType InHeapType, D3D12_DESCRIPTOR_HEAP_FLAGS InHeapFlags) const
 	{
@@ -100,7 +101,8 @@ public:
 	FRHIDescriptorHandle AllocateDescriptor(ERHIDescriptorHeapType InType);
 	void FreeDescriptor(FRHIDescriptorHandle InHandle);
 
-	D3D12_CPU_DESCRIPTOR_HANDLE GetCpuDescriptorHandle(FRHIDescriptorHandle InHandle) const;
+	void UpdateImmediately(FRHIDescriptorHandle InHandle, D3D12_CPU_DESCRIPTOR_HANDLE InSourceCpuHandle);
+	void UpdateDeferred(FRHIDescriptorHandle InHandle, D3D12_CPU_DESCRIPTOR_HANDLE InSourceCpuHandle);
 
 private:
 	TArray<FD3D12DescriptorManager> Managers;
diff --git a/Engine/Source/Runtime/D3D12RHI/Private/D3D12State.cpp b/Engine/Source/Runtime/D3D12RHI/Private/D3D12State.cpp
--- a/Engine/Source/Runtime/D3D12RHI/Private/D3D12State.cpp
+++ b/Engine/Source/Runtime/D3D12RHI/Private/D3D12State.cpp
@@ -549,8 +549,7 @@ FD3D12SamplerState::FD3D12SamplerState(FD3D12Device* InParent, const D3D12_SAMPL
 
 	if (BindlessHandle.IsValid())
 	{
-		D3D12_CPU_DESCRIPTOR_HANDLE OnlineHandle = BindlessDescriptorManager.GetCpuDescriptorHandle(BindlessHandle);
-		GetParentDevice()->GetDevice()->CopyDescriptorsSimple(1, OnlineHandle, OfflineHandle, D3D12_DESCRIPTOR_HEAP_TYPE_SAMPLER);
+		GetParentDevice()->GetBindlessDescriptorManager().UpdateImmediately(BindlessHandle, OfflineHandle);
 	}
 }
 
diff --git a/Engine/Source/Runtime/D3D12RHI/Private/D3D12Texture.cpp b/Engine/Source/Runtime/D3D12RHI/Private/D3D12Texture.cpp
--- a/Engine/Source/Runtime/D3D12RHI/Private/D3D12Texture.cpp
+++ b/Engine/Source/Runtime/D3D12RHI/Private/D3D12Texture.cpp
@@ -1227,7 +1227,7 @@ TD3D12Texture2D<BaseResourceType>* FD3D12DynamicRHI::CreateD3D12Texture2D(FRHICo
 					}
 				}
 
-				NewTexture->SetShaderResourceView(new FD3D12ShaderResourceView(Device, SRVDesc, NewTexture));
+				NewTexture->SetShaderResourceView(new FD3D12ShaderResourceView(NewTexture, SRVDesc));
 			}
 		}
 
@@ -1347,7 +1347,7 @@ FD3D12Texture3D* FD3D12DynamicRHI::CreateD3D12Texture3D(FRHICommandListImmediate
 		SRVDesc.Texture3D.MipLevels = NumMips;
 		SRVDesc.Texture3D.MostDetailedMip = 0;
 
-		Texture3D->SetShaderResourceView(new FD3D12ShaderResourceView(Device, SRVDesc, Texture3D));
+		Texture3D->SetShaderResourceView(new FD3D12ShaderResourceView(Texture3D, SRVDesc));
 
 		return Texture3D;
 	}); 
@@ -1508,7 +1508,7 @@ FTexture2DRHIRef FD3D12DynamicRHI::RHIAsyncCreateTexture2D(uint32 SizeX, uint32
 		SRVDesc.Texture2D.PlaneSlice = GetPlaneSliceFromViewFormat(PlatformResourceFormat, SRVDesc.Format);
 
 		// Create a wrapper for the SRV and set it on the texture
-		NewTexture->SetShaderResourceView(new FD3D12ShaderResourceView(Device, SRVDesc, NewTexture));
+		NewTexture->SetShaderResourceView(new FD3D12ShaderResourceView(NewTexture, SRVDesc));
 
 		return NewTexture;
 	});
@@ -3244,7 +3244,7 @@ TD3D12Texture2D<BaseResourceType>* FD3D12DynamicRHI::CreateTextureFromResource(b
 	// Create a wrapper for the SRV and set it on the texture
 	if (bCreateShaderResource)
 	{
-		Texture2D->SetShaderResourceView(new FD3D12ShaderResourceView(Device, SRVDesc, Texture2D));
+		Texture2D->SetShaderResourceView(new FD3D12ShaderResourceView(Texture2D, SRVDesc));
 	}
 
 	FD3D12TextureStats::D3D12TextureAllocated(*Texture2D);
diff --git a/Engine/Source/Runtime/D3D12RHI/Private/D3D12UAV.cpp b/Engine/Source/Runtime/D3D12RHI/Private/D3D12UAV.cpp
--- a/Engine/Source/Runtime/D3D12RHI/Private/D3D12UAV.cpp
+++ b/Engine/Source/Runtime/D3D12RHI/Private/D3D12UAV.cpp
@@ -3,6 +3,38 @@
 #include "D3D12RHIPrivate.h"
 #include "ClearReplacementShaders.h"
 
+FD3D12UnorderedAccessView::FD3D12UnorderedAccessView(FD3D12Device* InParent)
+	: FD3D12View(InParent, ERHIDescriptorHeapType::Standard, ViewSubresourceSubsetFlags_None)
+{
+}
+
+FD3D12UnorderedAccessView::FD3D12UnorderedAccessView(FD3D12Device* InParent, const D3D12_UNORDERED_ACCESS_VIEW_DESC& InDesc, FD3D12BaseShaderResource* InBaseShaderResource, FD3D12Resource* InCounterResource)
+	: FD3D12View(InParent, ERHIDescriptorHeapType::Standard, ViewSubresourceSubsetFlags_None)
+	, CounterResource(InCounterResource)
+{
+	SetDesc(InDesc);
+	CreateView(InBaseShaderResource, InBaseShaderResource->ResourceLocation, InCounterResource, ED3D12DescriptorCreateReason::InitialCreate);
+}
+
+void FD3D12UnorderedAccessView::RecreateView()
+{
+	check(CounterResource == nullptr);
+	check(ResourceLocation->GetOffsetFromBaseOfResource() == 0);
+	CreateView(BaseShaderResource, *ResourceLocation, nullptr, ED3D12DescriptorCreateReason::UpdateOrRename);
+}
+
+void FD3D12UnorderedAccessView::CreateView(FD3D12BaseShaderResource* InBaseShaderResource, FD3D12ResourceLocation& InResourceLocation, FD3D12Resource* InCounterResource, ED3D12DescriptorCreateReason Reason)
+{
+	InitializeInternal(InBaseShaderResource, InResourceLocation);
+
+	if (Resource)
+	{
+		ID3D12Resource* D3DResource = Resource->GetUAVAccessResource() ? Resource->GetUAVAccessResource() : Resource->GetResource();
+		ID3D12Resource* D3DCounterResource = InCounterResource ? InCounterResource->GetResource() : nullptr;
+		Descriptor.CreateView(Desc, D3DResource, D3DCounterResource, Reason);
+	}
+}
+
 template<typename ResourceType>
 inline FD3D12UnorderedAccessView* CreateUAV(D3D12_UNORDERED_ACCESS_VIEW_DESC& Desc, ResourceType* Resource, bool bNeedsCounterResource)
 {
@@ -231,7 +263,7 @@ void FD3D12CommandContext::ClearUAV(TRHICommandList_RecursiveHazardous<FD3D12Com
 
 				// Scoped descriptor handle will free the offline CPU handle once we return
 				FD3D12ViewDescriptorHandle UAVHandle(ParentDevice, ERHIDescriptorHeapType::Standard);
-				UAVHandle.CreateView(R32UAVDesc, Resource, nullptr);
+				UAVHandle.CreateView(R32UAVDesc, Resource, nullptr, ED3D12DescriptorCreateReason::InitialCreate);
 
 				// Check if the view heap is full and needs to rollover.
 				if (!Context.StateCache.GetDescriptorCache()->GetCurrentViewHeap()->CanReserveSlots(1))
diff --git a/Engine/Source/Runtime/D3D12RHI/Private/D3D12View.cpp b/Engine/Source/Runtime/D3D12RHI/Private/D3D12View.cpp
--- a/Engine/Source/Runtime/D3D12RHI/Private/D3D12View.cpp
+++ b/Engine/Source/Runtime/D3D12RHI/Private/D3D12View.cpp
@@ -49,7 +49,7 @@ void FD3D12ViewDescriptorHandle::CreateView(const D3D12_CONSTANT_BUFFER_VIEW_DES
 	GetParentDevice()->GetDevice()->CreateConstantBufferView(&Desc, OfflineCpuHandle);
 }
 
-void FD3D12ViewDescriptorHandle::CreateView(const D3D12_SHADER_RESOURCE_VIEW_DESC& Desc, ID3D12Resource* Resource)
+void FD3D12ViewDescriptorHandle::CreateView(const D3D12_SHADER_RESOURCE_VIEW_DESC& Desc, ID3D12Resource* Resource, ED3D12DescriptorCreateReason Reason)
 {
 	check(HeapType == ERHIDescriptorHeapType::Standard);
 
@@ -63,22 +63,16 @@ void FD3D12ViewDescriptorHandle::CreateView(const D3D12_SHADER_RESOURCE_VIEW_DES
 
 	GetParentDevice()->GetDevice()->CreateShaderResourceView(Resource, &Desc, OfflineCpuHandle);
 
-	if (IsBindless())
-	{
-		GetParentDevice()->GetDevice()->CopyDescriptorsSimple(1, GetBindlessCpuDescriptorHandle(), GetOfflineCpuHandle(), D3D12_DESCRIPTOR_HEAP_TYPE_CBV_SRV_UAV);
-	}
+	UpdateBindlessSlot(Reason);
 }
 
-void FD3D12ViewDescriptorHandle::CreateView(const D3D12_UNORDERED_ACCESS_VIEW_DESC& Desc, ID3D12Resource* Resource, ID3D12Resource* CounterResource)
+void FD3D12ViewDescriptorHandle::CreateView(const D3D12_UNORDERED_ACCESS_VIEW_DESC& Desc, ID3D12Resource* Resource, ID3D12Resource* CounterResource, ED3D12DescriptorCreateReason Reason)
 {
 	check(HeapType == ERHIDescriptorHeapType::Standard);
 
 	GetParentDevice()->GetDevice()->CreateUnorderedAccessView(Resource, CounterResource, &Desc, OfflineCpuHandle);
 
-	if (IsBindless())
-	{
-		GetParentDevice()->GetDevice()->CopyDescriptorsSimple(1, GetBindlessCpuDescriptorHandle(), GetOfflineCpuHandle(), D3D12_DESCRIPTOR_HEAP_TYPE_CBV_SRV_UAV);
-	}
+	UpdateBindlessSlot(Reason);
 }
 
 void FD3D12ViewDescriptorHandle::AllocateDescriptorSlot()
@@ -103,7 +97,7 @@ void FD3D12ViewDescriptorHandle::FreeDescriptorSlot()
 		OfflineHeapIndex = UINT_MAX;
 		OfflineCpuHandle.ptr = 0;
 
-		if (IsBindless())
+		if (BindlessHandle.IsValid())
 		{
 			Device->GetBindlessDescriptorManager().FreeDescriptor(BindlessHandle);
 			BindlessHandle = FRHIDescriptorHandle();
@@ -112,14 +106,139 @@ void FD3D12ViewDescriptorHandle::FreeDescriptorSlot()
 	check(OfflineCpuHandle.ptr == 0);
 }
 
-D3D12_CPU_DESCRIPTOR_HANDLE FD3D12ViewDescriptorHandle::GetBindlessCpuDescriptorHandle() const
+void FD3D12ViewDescriptorHandle::UpdateBindlessSlot(ED3D12DescriptorCreateReason Reason)
 {
-	return GetParentDevice()->GetBindlessDescriptorManager().GetCpuDescriptorHandle(BindlessHandle);
+	if (BindlessHandle.IsValid())
+	{
+		FD3D12BindlessDescriptorManager& BindlessManager = GetParentDevice()->GetBindlessDescriptorManager();
+		if (Reason == ED3D12DescriptorCreateReason::InitialCreate)
+		{
+			BindlessManager.UpdateImmediately(BindlessHandle, GetOfflineCpuHandle());
+		}
+		else
+		{
+			BindlessManager.UpdateDeferred(BindlessHandle, GetOfflineCpuHandle());
+		}
+	}
 }
 
 ////////////////////////////////////////////////////////////////////////////////////////////////////////////////////////
 
-static FORCEINLINE D3D12_SHADER_RESOURCE_VIEW_DESC GetVertexBufferSRVDesc(FD3D12Buffer* VertexBuffer, uint32& CreationStride, uint8 Format, uint32 StartOffsetBytes, uint32 NumElements)
+FD3D12ShaderResourceView::FD3D12ShaderResourceView(FD3D12Device* InParent)
+	: FD3D12View(InParent, ERHIDescriptorHeapType::Standard, ViewSubresourceSubsetFlags_None)
+	, bContainsDepthPlane(false)
+	, bContainsStencilPlane(false)
+	, bSkipFastClearFinalize(false)
+	, bRequiresResourceStateTracking(false)
+{
+}
+
+FD3D12ShaderResourceView::FD3D12ShaderResourceView(FD3D12Buffer* InBuffer, const D3D12_SHADER_RESOURCE_VIEW_DESC& InDesc, uint32 InStride, uint32 InStartOffsetBytes)
+	: FD3D12ShaderResourceView(InBuffer->GetParentDevice())
+{
+	InitializeAfterCreate(InDesc, InBuffer, InBuffer->ResourceLocation, InStride, InStartOffsetBytes);
+}
+
+FD3D12ShaderResourceView::FD3D12ShaderResourceView(FD3D12TextureBase* InTexture, const D3D12_SHADER_RESOURCE_VIEW_DESC& InDesc, ETextureCreateFlags InTextureCreateFlags)
+	: FD3D12ShaderResourceView(InTexture->GetParentDevice())
+{
+	InitializeAfterCreate(InDesc, InTexture, InTexture->ResourceLocation, -1, 0, EnumHasAnyFlags(InTextureCreateFlags, ETextureCreateFlags::NoFastClearFinalize));
+}
+
+FD3D12ShaderResourceView::~FD3D12ShaderResourceView()
+{
+}
+
+void FD3D12ShaderResourceView::InitializeAfterCreate(const D3D12_SHADER_RESOURCE_VIEW_DESC& InDesc, FD3D12BaseShaderResource* InBaseShaderResource, FD3D12ResourceLocation& InResourceLocation, uint32 InStride, uint32 InStartOffsetBytes, bool InSkipFastClearFinalize)
+{
+	SetDesc(InDesc);
+
+	PreCreateView(InBaseShaderResource->ResourceLocation, InStride, InStartOffsetBytes, InSkipFastClearFinalize);
+	CreateView(InBaseShaderResource, InBaseShaderResource->ResourceLocation, ED3D12DescriptorCreateReason::InitialCreate);
+}
+
+void FD3D12ShaderResourceView::Update(FD3D12Buffer* InBuffer, const D3D12_SHADER_RESOURCE_VIEW_DESC& InDesc, uint32 InStride)
+{
+	FD3D12Device* InParent = InBuffer->GetParentDevice();
+	if (!this->GetParentDevice_Unsafe())
+	{
+		// This is a null SRV created without viewing on any resource
+		// We need to set its device and allocate a descriptor slot before moving forward
+		this->SetParentDevice(InParent);
+	}
+	check(GetParentDevice() == InParent);
+
+	SetDesc(InDesc);
+
+	PreCreateView(InBuffer->ResourceLocation, InStride, 0, false);
+	CreateView(InBuffer, InBuffer->ResourceLocation, ED3D12DescriptorCreateReason::UpdateOrRename);
+}
+
+void FD3D12ShaderResourceView::UpdateMinLODClamp(float ResourceMinLODClamp)
+{
+	check(bInitialized);
+	check(ResourceLocation);
+	check(Desc.ViewDimension == D3D12_SRV_DIMENSION_TEXTURE2D);
+
+	// Update the LODClamp, the reinitialize the SRV
+	Desc.Texture2D.ResourceMinLODClamp = ResourceMinLODClamp;
+	CreateView(BaseShaderResource, *ResourceLocation, ED3D12DescriptorCreateReason::UpdateOrRename);
+}
+
+void FD3D12ShaderResourceView::RecreateView()
+{
+	// Update the first element index, then reinitialize the SRV
+	if (Desc.ViewDimension == D3D12_SRV_DIMENSION_BUFFER)
+	{
+		uint32 StartElement = StartOffsetBytes / Stride;
+		Desc.Buffer.FirstElement = ResourceLocation->GetOffsetFromBaseOfResource() / Stride + StartElement;
+	}
+
+	PreCreateView(BaseShaderResource->ResourceLocation, Stride, StartOffsetBytes, bSkipFastClearFinalize);
+	CreateView(BaseShaderResource, BaseShaderResource->ResourceLocation, ED3D12DescriptorCreateReason::UpdateOrRename);
+}
+
+void FD3D12ShaderResourceView::PreCreateView(const FD3D12ResourceLocation& InResourceLocation, uint32 InStride, uint32 InStartOffsetBytes, bool InSkipFastClearFinalize)
+{
+	Stride = InStride;
+	StartOffsetBytes = InStartOffsetBytes;
+	bSkipFastClearFinalize = InSkipFastClearFinalize;
+
+	if (FD3D12Resource* NewResource = InResourceLocation.GetResource())
+	{
+		bContainsDepthPlane = NewResource->IsDepthStencilResource() && GetPlaneSliceFromViewFormat(NewResource->GetDesc().Format, Desc.Format) == 0;
+		bContainsStencilPlane = NewResource->IsDepthStencilResource() && GetPlaneSliceFromViewFormat(NewResource->GetDesc().Format, Desc.Format) == 1;
+		bRequiresResourceStateTracking = NewResource->RequiresResourceStateTracking();
+
+#if DO_CHECK
+		// Check the plane slice of the SRV matches the texture format
+		// Texture2DMS does not have explicit plane index (it's implied by the format)
+		if (Desc.ViewDimension == D3D12_SRV_DIMENSION_TEXTURE2D)
+		{
+			check(GetPlaneSliceFromViewFormat(NewResource->GetDesc().Format, Desc.Format) == Desc.Texture2D.PlaneSlice);
+		}
+#endif
+	}
+	else
+	{
+		bContainsDepthPlane = false;
+		bContainsStencilPlane = false;
+		bRequiresResourceStateTracking = false;
+	}
+}
+
+void FD3D12ShaderResourceView::CreateView(FD3D12BaseShaderResource* InBaseShaderResource, FD3D12ResourceLocation& InResourceLocation, ED3D12DescriptorCreateReason Reason)
+{
+	InitializeInternal(InBaseShaderResource, InResourceLocation);
+
+	if (ResourceLocation->GetResource())
+	{
+		ID3D12Resource* D3DResource = ResourceLocation->GetResource()->GetResource();
+		Descriptor.CreateView(Desc, D3DResource, Reason);
+	}
+}
+
+static FORCEINLINE D3D12_SHADER_RESOURCE_VIEW_DESC GetVertexBufferSRVDesc(const FD3D12Buffer* VertexBuffer, uint32& CreationStride, uint8 Format, uint32 StartOffsetBytes, uint32 NumElements)
 {
 	const uint32 BufferSize = VertexBuffer->GetSize();
 	const uint64 BufferOffset = VertexBuffer->ResourceLocation.GetOffsetFromBaseOfResource();
@@ -135,7 +254,7 @@ static FORCEINLINE D3D12_SHADER_RESOURCE_VIEW_DESC GetVertexBufferSRVDesc(FD3D12
 	SRVDesc.Shader4ComponentMapping = D3D12_DEFAULT_SHADER_4_COMPONENT_MAPPING;
 	SRVDesc.ViewDimension = D3D12_SRV_DIMENSION_BUFFER;
 
-	if (EnumHasAnyFlags(VertexBuffer->GetUsage(), BUF_ByteAddressBuffer))
+	if (EnumHasAnyFlags(VertexBuffer->GetUsage(), EBufferUsageFlags::ByteAddressBuffer))
 	{
 		SRVDesc.Format = DXGI_FORMAT_R32_TYPELESS;
 		SRVDesc.Buffer.Flags |= D3D12_BUFFER_SRV_FLAG_RAW;
@@ -153,20 +272,53 @@ static FORCEINLINE D3D12_SHADER_RESOURCE_VIEW_DESC GetVertexBufferSRVDesc(FD3D12
 	return SRVDesc;
 }
 
+static FORCEINLINE D3D12_SHADER_RESOURCE_VIEW_DESC GetStructuredBufferSRVDesc(const FD3D12Buffer* StructuredBuffer, uint32& CreationStride, uint8 Format, uint32 StartOffsetBytes, uint32 NumElements)
+{
+	const uint32 BufferSize = StructuredBuffer->ResourceLocation.GetSize();
+	const uint32 BufferStride = StructuredBuffer->GetStride();
+	const uint32 BufferOffset = StructuredBuffer->ResourceLocation.GetOffsetFromBaseOfResource();
+
+	const uint32 MaxElements = BufferSize / BufferStride;
+	const uint32 StartElement = FMath::Min<uint32>(StartOffsetBytes, BufferSize) / BufferStride;
+
+	D3D12_SHADER_RESOURCE_VIEW_DESC SRVDesc{};
+
+	SRVDesc.Shader4ComponentMapping = D3D12_DEFAULT_SHADER_4_COMPONENT_MAPPING;
+	SRVDesc.ViewDimension = D3D12_SRV_DIMENSION_BUFFER;
+
+	if (EnumHasAnyFlags(StructuredBuffer->GetUsage(), EBufferUsageFlags::ByteAddressBuffer))
+	{
+		SRVDesc.Format = DXGI_FORMAT_R32_TYPELESS;
+		SRVDesc.Buffer.Flags |= D3D12_BUFFER_SRV_FLAG_RAW;
+		CreationStride = 4;
+	}
+	else
+	{
+		SRVDesc.Format = DXGI_FORMAT_UNKNOWN;
+		SRVDesc.Buffer.StructureByteStride = BufferStride;
+		CreationStride = BufferStride;
+	}
+
+	SRVDesc.Buffer.NumElements = FMath::Min<uint32>(MaxElements - StartElement, NumElements);
+	SRVDesc.Buffer.FirstElement = (BufferOffset / CreationStride) + StartElement;
+
+	return SRVDesc;
+}
+
 static FORCEINLINE D3D12_SHADER_RESOURCE_VIEW_DESC GetIndexBufferSRVDesc(FD3D12Buffer* IndexBuffer, uint32 StartOffsetBytes, uint32 NumElements)
 {
 	const EBufferUsageFlags Usage = IndexBuffer->GetUsage();
 	const uint32 Width = IndexBuffer->GetSize();
 	const uint32 CreationStride = IndexBuffer->GetStride();
-	uint32 MaxElements = Width / CreationStride;
-	StartOffsetBytes = FMath::Min(StartOffsetBytes, Width);
-	uint32 StartElement = StartOffsetBytes / CreationStride;
+	const uint32 MaxElements = Width / CreationStride;
+	const uint32 StartElement = FMath::Min(StartOffsetBytes, Width) / CreationStride;
 	const FD3D12ResourceLocation& Location = IndexBuffer->ResourceLocation;
 
 	D3D12_SHADER_RESOURCE_VIEW_DESC SRVDesc{};
 	SRVDesc.Shader4ComponentMapping = D3D12_DEFAULT_SHADER_4_COMPONENT_MAPPING;
 	SRVDesc.ViewDimension = D3D12_SRV_DIMENSION_BUFFER;
-	if (EnumHasAnyFlags(Usage, BUF_ByteAddressBuffer))
+
+	if (EnumHasAnyFlags(Usage, EBufferUsageFlags::ByteAddressBuffer))
 	{
 		check(CreationStride == 4u);
 		SRVDesc.Format = DXGI_FORMAT_R32_TYPELESS;
@@ -193,7 +345,7 @@ static FORCEINLINE D3D12_SHADER_RESOURCE_VIEW_DESC GetIndexBufferSRVDesc(FD3D12B
 }
 
 template<typename TextureType>
-FD3D12ShaderResourceView* CreateSRV(TextureType* Texture, D3D12_SHADER_RESOURCE_VIEW_DESC& Desc)
+FD3D12ShaderResourceView* CreateSRV(TextureType* Texture, const D3D12_SHADER_RESOURCE_VIEW_DESC& Desc)
 {
 	if (Texture == nullptr)
 	{
@@ -297,6 +449,43 @@ uint64 FD3D12DynamicRHI::RHIGetMinimumAlignmentForBufferBackedSRV(EPixelFormat F
 	return GPixelFormats[Format].BlockBytes;
 }
 
+FRHICOMMAND_MACRO(FD3D12InitializeBufferSRVRHICommand)
+{
+	const FShaderResourceViewInitializer Initializer;
+	FD3D12ShaderResourceView* const View;
+
+	FD3D12InitializeBufferSRVRHICommand(const FShaderResourceViewInitializer& InInitializer, FD3D12ShaderResourceView* InView)
+		: Initializer(InInitializer)
+		, View(InView)
+	{
+	}
+
+	void Execute(FRHICommandListBase& RHICmdList)
+	{
+		const FShaderResourceViewInitializer::FBufferShaderResourceViewInitializer& Desc = Initializer.AsBufferSRV();
+		FD3D12Buffer* const Buffer = FD3D12DynamicRHI::ResourceCast(Desc.Buffer);
+
+		if (Initializer.GetType() == FShaderResourceViewInitializer::EType::VertexBufferSRV)
+		{
+			uint32 CreationStride = 0;
+			const D3D12_SHADER_RESOURCE_VIEW_DESC ViewDesc = GetVertexBufferSRVDesc(Buffer, CreationStride, Desc.Format, Desc.StartOffsetBytes, Desc.NumElements);
+
+			View->InitializeAfterCreate(ViewDesc, Buffer, Buffer->ResourceLocation, CreationStride, Desc.StartOffsetBytes);
+		}
+		else if (Initializer.GetType() == FShaderResourceViewInitializer::EType::StructuredBufferSRV)
+		{
+			uint32 CreationStride = 0;
+			const D3D12_SHADER_RESOURCE_VIEW_DESC ViewDesc = GetStructuredBufferSRVDesc(Buffer, CreationStride, Desc.Format, Desc.StartOffsetBytes, Desc.NumElements);
+
+			View->InitializeAfterCreate(ViewDesc, Buffer, Buffer->ResourceLocation, CreationStride, Desc.StartOffsetBytes);
+		}
+		else
+		{
+			checkNoEntry();
+		}
+	}
+};
+
 FShaderResourceViewRHIRef FD3D12DynamicRHI::RHICreateShaderResourceView(const FShaderResourceViewInitializer& Initializer)
 {
 	FShaderResourceViewInitializer::FBufferShaderResourceViewInitializer Desc = Initializer.AsBufferSRV();
@@ -314,175 +503,62 @@ FShaderResourceViewRHIRef FD3D12DynamicRHI::RHICreateShaderResourceView(const FS
 	switch (Initializer.GetType())
 	{
 		case FShaderResourceViewInitializer::EType::VertexBufferSRV:
-		{
-			struct FD3D12InitializeVertexBufferSRVRHICommand final : public FRHICommand<FD3D12InitializeVertexBufferSRVRHICommand>
-			{
-				FD3D12Buffer* VertexBuffer;
-				FD3D12ShaderResourceView* SRV;
-
-				uint32 StartOffsetBytes;
-				uint32 NumElements;
-				uint8 Format;
-
-				FD3D12InitializeVertexBufferSRVRHICommand(FD3D12Buffer* InVertexBuffer, FD3D12ShaderResourceView* InSRV, uint32 InStartOffsetBytes, uint32 InNumElements, uint8 InFormat)
-					: VertexBuffer(InVertexBuffer)
-					, SRV(InSRV)
-					, StartOffsetBytes(InStartOffsetBytes)
-					, NumElements(InNumElements)
-					, Format(InFormat)
-				{}
-
-				void Execute(FRHICommandListBase& RHICmdList)
-				{
-					uint32 CreationStride = GPixelFormats[Format].BlockBytes;
-					D3D12_SHADER_RESOURCE_VIEW_DESC SRVDesc = GetVertexBufferSRVDesc(VertexBuffer, CreationStride, Format, StartOffsetBytes, NumElements);
-
-					// Create a Shader Resource View
-					SRV->Initialize(SRVDesc, VertexBuffer, CreationStride, StartOffsetBytes);
-				}
-			};
-
+		case FShaderResourceViewInitializer::EType::StructuredBufferSRV:
 			return GetAdapter().CreateLinkedViews<FD3D12Buffer, FD3D12ShaderResourceView>(Buffer,
-				[Desc](FD3D12Buffer* Buffer)
+				[Initializer](FD3D12Buffer* Buffer)
 				{
 					check(Buffer);
 
 					FD3D12ShaderResourceView* ShaderResourceView = new FD3D12ShaderResourceView(Buffer->GetParentDevice());
 					
-					uint32 Stride = GPixelFormats[Desc.Format].BlockBytes;
 					FRHICommandListImmediate& RHICmdList = FRHICommandListExecutor::GetImmediateCommandList();
-					if (ShouldDeferBufferLockOperation(&RHICmdList) && EnumHasAnyFlags(Buffer->GetUsage(), BUF_AnyDynamic))
+					if (ShouldDeferBufferLockOperation(&RHICmdList) && EnumHasAnyFlags(Buffer->GetUsage(), EBufferUsageFlags::AnyDynamic))
 					{
 						// We have to defer the SRV initialization to the RHI thread if the buffer is dynamic (and RHI threading is enabled), as dynamic buffers can be renamed.
 						// Also insert an RHI thread fence to prevent parallel translate tasks running until this command has completed.
-						ALLOC_COMMAND_CL(RHICmdList, FD3D12InitializeVertexBufferSRVRHICommand)(Buffer, ShaderResourceView, Desc.StartOffsetBytes, Desc.NumElements, Desc.Format);
+						ALLOC_COMMAND_CL(RHICmdList, FD3D12InitializeBufferSRVRHICommand)(Initializer, ShaderResourceView);
 						RHICmdList.RHIThreadFence(true);
 					}
 					else
 					{
 						// Run the command directly if we're bypassing RHI command list recording, or the buffer is not dynamic.
-						FD3D12InitializeVertexBufferSRVRHICommand Command(Buffer, ShaderResourceView, Desc.StartOffsetBytes, Desc.NumElements, Desc.Format);
+						FD3D12InitializeBufferSRVRHICommand Command(Initializer, ShaderResourceView);
 						Command.Execute(RHICmdList);
 					}
 
 					return ShaderResourceView;
 				});
-		}
-
-		case FShaderResourceViewInitializer::EType::StructuredBufferSRV:
-		{
-			struct FD3D12InitializeStructuredBufferSRVRHICommand final : public FRHICommand<FD3D12InitializeStructuredBufferSRVRHICommand>
-			{
-				FD3D12Buffer* StructuredBuffer;
-				FD3D12ShaderResourceView* SRV;
-				uint32 StartOffsetBytes;
-				uint32 NumElements;
-
-				FD3D12InitializeStructuredBufferSRVRHICommand(FD3D12Buffer* InStructuredBuffer, FD3D12ShaderResourceView* InSRV, uint32 InStartOffsetBytes, uint32 InNumElements)
-					: StructuredBuffer(InStructuredBuffer)
-					, SRV(InSRV)
-					, StartOffsetBytes(InStartOffsetBytes)
-					, NumElements(InNumElements)
-				{}
-
-				void Execute(FRHICommandListBase& RHICmdList)
-				{
-					FD3D12ResourceLocation& Location = StructuredBuffer->ResourceLocation;
-
-					const uint64 Offset = Location.GetOffsetFromBaseOfResource();
-					const D3D12_RESOURCE_DESC& BufferDesc = Location.GetResource()->GetDesc();
-
-					// Create a Shader Resource View
-					D3D12_SHADER_RESOURCE_VIEW_DESC SRVDesc = {};
-					SRVDesc.Shader4ComponentMapping = D3D12_DEFAULT_SHADER_4_COMPONENT_MAPPING;
-
-					// BufferDesc.StructureByteStride  is not getting patched through the D3D resource DESC structs, so use the RHI version as a hack
-					uint32 Stride = StructuredBuffer->GetStride();
-					uint32 MaxElements = Location.GetSize() / Stride;
-					StartOffsetBytes = FMath::Min<uint32>(StartOffsetBytes, Location.GetSize());
-					uint32 StartElement = StartOffsetBytes / Stride;
-
-					if (EnumHasAnyFlags(StructuredBuffer->GetUsage(), BUF_ByteAddressBuffer))
-					{
-						SRVDesc.Buffer.Flags = D3D12_BUFFER_SRV_FLAG_RAW;
-						SRVDesc.Format = DXGI_FORMAT_R32_TYPELESS;
-						Stride = 4;
-					}
-					else
-					{
-						SRVDesc.Buffer.StructureByteStride = Stride;
-						SRVDesc.Format = DXGI_FORMAT_UNKNOWN;
-					}
-
-					SRVDesc.ViewDimension = D3D12_SRV_DIMENSION_BUFFER;
-					SRVDesc.Buffer.NumElements = FMath::Min<uint32>(MaxElements - StartElement, NumElements);
-					SRVDesc.Buffer.FirstElement = Offset / Stride + StartElement;
-
-					// Create a Shader Resource View
-					SRV->Initialize(SRVDesc, StructuredBuffer, Stride, StartOffsetBytes);
-				}
-			};
-
-			return GetAdapter().CreateLinkedViews<FD3D12Buffer,
-				FD3D12ShaderResourceView>(Buffer, [Desc](FD3D12Buffer* Buffer)
-					{
-						check(Buffer);
-
-						FD3D12ShaderResourceView* ShaderResourceView = new FD3D12ShaderResourceView(Buffer->GetParentDevice());
-
-						FRHICommandListImmediate& RHICmdList = FRHICommandListExecutor::GetImmediateCommandList();
-						if (ShouldDeferBufferLockOperation(&RHICmdList) && EnumHasAnyFlags(Buffer->GetUsage(), BUF_AnyDynamic))
-						{
-							// We have to defer the SRV initialization to the RHI thread if the buffer is dynamic (and RHI threading is enabled), as dynamic buffers can be renamed.
-							// Also insert an RHI thread fence to prevent parallel translate tasks running until this command has completed.
-							ALLOC_COMMAND_CL(RHICmdList, FD3D12InitializeStructuredBufferSRVRHICommand)(Buffer, ShaderResourceView, Desc.StartOffsetBytes, Desc.NumElements);
-							RHICmdList.RHIThreadFence(true);
-						}
-						else
-						{
-							// Run the command directly if we're bypassing RHI command list recording, or the buffer is not dynamic.
-							FD3D12InitializeStructuredBufferSRVRHICommand Command(Buffer, ShaderResourceView, Desc.StartOffsetBytes, Desc.NumElements);
-							Command.Execute(RHICmdList);
-						}
-
-						return ShaderResourceView;
-					});
-		}
 
 		case FShaderResourceViewInitializer::EType::IndexBufferSRV:
-		{
 			return GetAdapter().CreateLinkedViews<FD3D12Buffer, FD3D12ShaderResourceView>(Buffer,
-			[Desc](FD3D12Buffer* Buffer)
-			{
-				check(Buffer);
+				[Desc](FD3D12Buffer* Buffer)
+				{
+					check(Buffer);
 
-				const uint32 CreationStride = Buffer->GetStride();
-				D3D12_SHADER_RESOURCE_VIEW_DESC SRVDesc = GetIndexBufferSRVDesc(Buffer, Desc.StartOffsetBytes, Desc.NumElements);
+					const uint32 CreationStride = Buffer->GetStride();
+					const D3D12_SHADER_RESOURCE_VIEW_DESC SRVDesc = GetIndexBufferSRVDesc(Buffer, Desc.StartOffsetBytes, Desc.NumElements);
 
-				FD3D12ShaderResourceView* ShaderResourceView = new FD3D12ShaderResourceView(Buffer->GetParentDevice(), SRVDesc, Buffer, CreationStride);
-				return ShaderResourceView;
-			});
-		}
+					FD3D12ShaderResourceView* ShaderResourceView = new FD3D12ShaderResourceView(Buffer, SRVDesc, CreationStride);
+					return ShaderResourceView;
+				});
 
 	#if D3D12_RHI_RAYTRACING
 		case FShaderResourceViewInitializer::EType::AccelerationStructureSRV:
-		{
 			return GetAdapter().CreateLinkedViews<FD3D12Buffer, FD3D12ShaderResourceView>(Buffer,
 				[Desc](FD3D12Buffer* Buffer)
-			{
-				check(Buffer);
+				{
+					check(Buffer);
 
-				D3D12_SHADER_RESOURCE_VIEW_DESC SRVDesc = {};
+					D3D12_SHADER_RESOURCE_VIEW_DESC SRVDesc = {};
 
-				SRVDesc.Format = DXGI_FORMAT_UNKNOWN;
-				SRVDesc.ViewDimension = D3D12_SRV_DIMENSION_RAYTRACING_ACCELERATION_STRUCTURE;
-				SRVDesc.Shader4ComponentMapping = D3D12_DEFAULT_SHADER_4_COMPONENT_MAPPING;
-				SRVDesc.RaytracingAccelerationStructure.Location = Buffer->ResourceLocation.GetGPUVirtualAddress() + Desc.StartOffsetBytes;
+					SRVDesc.Format = DXGI_FORMAT_UNKNOWN;
+					SRVDesc.ViewDimension = D3D12_SRV_DIMENSION_RAYTRACING_ACCELERATION_STRUCTURE;
+					SRVDesc.Shader4ComponentMapping = D3D12_DEFAULT_SHADER_4_COMPONENT_MAPPING;
+					SRVDesc.RaytracingAccelerationStructure.Location = Buffer->ResourceLocation.GetGPUVirtualAddress() + Desc.StartOffsetBytes;
 
-				FD3D12ShaderResourceView* ShaderResourceView = new FD3D12ShaderResourceView(Buffer->GetParentDevice(), SRVDesc, Buffer, 4);
-				return ShaderResourceView;
-			});
-		}
+					FD3D12ShaderResourceView* ShaderResourceView = new FD3D12ShaderResourceView(Buffer, SRVDesc, 4);
+					return ShaderResourceView;
+				});
 	#endif // D3D12_RHI_RAYTRACING
 	}
 
@@ -497,7 +573,7 @@ void FD3D12DynamicRHI::RHIUpdateShaderResourceView(FRHIShaderResourceView* SRV,
 	{
 		FD3D12Buffer* Buffer = ResourceCast(BufferRHI);
 		FD3D12ShaderResourceView* SRVD3D12 = ResourceCast(SRV);
-		D3D12_SHADER_RESOURCE_VIEW_DESC SRVDesc = GetVertexBufferSRVDesc(Buffer, Stride, Format, 0, UINT32_MAX);
+		const D3D12_SHADER_RESOURCE_VIEW_DESC SRVDesc = GetVertexBufferSRVDesc(Buffer, Stride, Format, 0, UINT32_MAX);
 
 		// Rename the SRV to view on the new vertex buffer
 		for (auto It = MakeDualLinkedObjectIterator(Buffer, SRVD3D12); It; ++It)
@@ -505,8 +581,7 @@ void FD3D12DynamicRHI::RHIUpdateShaderResourceView(FRHIShaderResourceView* SRV,
 			Buffer = It.GetFirst();
 			SRVD3D12 = It.GetSecond();
 
-			FD3D12Device* ParentDevice = Buffer->GetParentDevice();
-			SRVD3D12->Initialize(ParentDevice, SRVDesc, Buffer, Stride);
+			SRVD3D12->Update(Buffer, SRVDesc, Stride);
 		}
 	}
 }
@@ -518,7 +593,7 @@ void FD3D12DynamicRHI::RHIUpdateShaderResourceView(FRHIShaderResourceView* SRV,
 	{
 		FD3D12Buffer* Buffer = ResourceCast(BufferRHI);
 		FD3D12ShaderResourceView* SRVD3D12 = ResourceCast(SRV);
-		D3D12_SHADER_RESOURCE_VIEW_DESC SRVDesc = GetIndexBufferSRVDesc(Buffer, 0, UINT32_MAX);
+		const D3D12_SHADER_RESOURCE_VIEW_DESC SRVDesc = GetIndexBufferSRVDesc(Buffer, 0, UINT32_MAX);
 		const uint32 Stride = Buffer->GetStride();
 
 		// Rename the SRV to view on the new index buffer
@@ -527,8 +602,7 @@ void FD3D12DynamicRHI::RHIUpdateShaderResourceView(FRHIShaderResourceView* SRV,
 			Buffer = It.GetFirst();
 			SRVD3D12 = It.GetSecond();
 
-			FD3D12Device* ParentDevice = Buffer->GetParentDevice();
-			SRVD3D12->Initialize(ParentDevice, SRVDesc, Buffer, Stride);
+			SRVD3D12->Update(Buffer, SRVDesc, Stride);
 		}
 	}
 }
diff --git a/Engine/Source/Runtime/D3D12RHI/Private/D3D12View.h b/Engine/Source/Runtime/D3D12RHI/Private/D3D12View.h
--- a/Engine/Source/Runtime/D3D12RHI/Private/D3D12View.h
+++ b/Engine/Source/Runtime/D3D12RHI/Private/D3D12View.h
@@ -4,6 +4,8 @@
 
 #include "D3D12Descriptors.h"
 
+class FD3D12TextureBase;
+
 enum ViewSubresourceSubsetFlags
 {
 	ViewSubresourceSubsetFlags_None = 0x0,
@@ -618,6 +620,12 @@ inline uint32 CViewSubresourceSubset::MaxSubresource() const
 	return (*(--end())).Value;
 }
 
+enum class ED3D12DescriptorCreateReason
+{
+	InitialCreate,
+	UpdateOrRename,
+};
+
 /** Manages descriptor allocations and view creation */
 class FD3D12ViewDescriptorHandle : public FD3D12DeviceChild
 {
@@ -632,8 +640,8 @@ public:
 	void CreateView(const D3D12_RENDER_TARGET_VIEW_DESC& Desc, ID3D12Resource* Resource);
 	void CreateView(const D3D12_DEPTH_STENCIL_VIEW_DESC& Desc, ID3D12Resource* Resource);
 	void CreateView(const D3D12_CONSTANT_BUFFER_VIEW_DESC& Desc);
-	void CreateView(const D3D12_SHADER_RESOURCE_VIEW_DESC& Desc, ID3D12Resource* Resource);
-	void CreateView(const D3D12_UNORDERED_ACCESS_VIEW_DESC& Desc, ID3D12Resource* Resource, ID3D12Resource* CounterResource);
+	void CreateView(const D3D12_SHADER_RESOURCE_VIEW_DESC& Desc, ID3D12Resource* Resource, ED3D12DescriptorCreateReason Reason);
+	void CreateView(const D3D12_UNORDERED_ACCESS_VIEW_DESC& Desc, ID3D12Resource* Resource, ID3D12Resource* CounterResource, ED3D12DescriptorCreateReason Reason);
 
 	inline D3D12_CPU_DESCRIPTOR_HANDLE GetOfflineCpuHandle() const { return OfflineCpuHandle; }
 	inline uint32                      GetOfflineHeapIndex() const { return OfflineHeapIndex; }
@@ -644,8 +652,7 @@ public:
 private:
 	void AllocateDescriptorSlot();
 	void FreeDescriptorSlot();
-
-	D3D12_CPU_DESCRIPTOR_HANDLE GetBindlessCpuDescriptorHandle() const;
+	void UpdateBindlessSlot(ED3D12DescriptorCreateReason Reason);
 
 	D3D12_CPU_DESCRIPTOR_HANDLE OfflineCpuHandle{};
 	uint32 OfflineHeapIndex{ UINT_MAX };
@@ -791,89 +798,19 @@ public:
 class FD3D12ShaderResourceView : public FRHIShaderResourceView, public FD3D12View<D3D12_SHADER_RESOURCE_VIEW_DESC>, public FD3D12LinkedAdapterObject<FD3D12ShaderResourceView>
 {
 public:
-	// Used for dynamic buffer SRVs, which can be renamed. Must be explicitly Initialize()'d before it can be used.
-	FD3D12ShaderResourceView(FD3D12Device* InParent)
-		: FD3D12View(InParent, ERHIDescriptorHeapType::Standard, ViewSubresourceSubsetFlags_None)
-		, bContainsDepthPlane(false)
-		, bContainsStencilPlane(false)
-		, bSkipFastClearFinalize(false)
-		, bRequiresResourceStateTracking(false)
-	{
-	}
-
-	// Used for all other SRV resource types. Initialization is immediate on the calling thread.
-	// Should not be used for dynamic resources which can be renamed.
-	FD3D12ShaderResourceView(FD3D12Device* InParent, D3D12_SHADER_RESOURCE_VIEW_DESC& InDesc, FD3D12BaseShaderResource* InBaseShaderResource, uint32 InStride = -1, uint32 InStartOffsetBytes = 0, bool InSkipFastClearFinalize = false)
-		: FD3D12ShaderResourceView(InParent)
-	{
-		Initialize(InDesc, InBaseShaderResource, InBaseShaderResource->ResourceLocation, InStride, InStartOffsetBytes, InSkipFastClearFinalize);
-	}
-
-	template<typename TextureType>
-	FD3D12ShaderResourceView(TextureType* InTexture, D3D12_SHADER_RESOURCE_VIEW_DESC& InDesc)
-		: FD3D12ShaderResourceView(InTexture->GetParentDevice(), InDesc, InTexture)
-	{
-	}
-
-	~FD3D12ShaderResourceView()
-	{
-	}
-
-	void Initialize(D3D12_SHADER_RESOURCE_VIEW_DESC& InDesc, FD3D12BaseShaderResource* InBaseShaderResource, FD3D12ResourceLocation& InResourceLocation, uint32 InStride, uint32 InStartOffsetBytes, bool InSkipFastClearFinalize)
-	{
-		SetDesc(InDesc);
-
-		PreCreateView(InResourceLocation, InStride, InStartOffsetBytes, InSkipFastClearFinalize);
-		CreateView(InBaseShaderResource, InResourceLocation);
-	}
+	// Used for dynamic buffer SRVs, which can be renamed. User must explicitly call InitializeAfterCreate before returning to callers or call Update if callers have requested updates.
+	FD3D12ShaderResourceView(FD3D12Device* InParent);
+	// Used for buffer SRVs, we usually only control the stride and offset
+	FD3D12ShaderResourceView(FD3D12Buffer* InBuffer, const D3D12_SHADER_RESOURCE_VIEW_DESC& InDesc, uint32 InStride, uint32 InStartOffsetBytes = 0);
+	// Used for texture SRVs, we don't control much other than disabling fast create
+	FD3D12ShaderResourceView(FD3D12TextureBase* InTexture, const D3D12_SHADER_RESOURCE_VIEW_DESC& InDesc, ETextureCreateFlags InTextureCreateFlags = ETextureCreateFlags::None);
+	~FD3D12ShaderResourceView();
 
-	void Initialize(D3D12_SHADER_RESOURCE_VIEW_DESC& InDesc, FD3D12BaseShaderResource* InBaseShaderResource, uint32 InStride, uint32 InStartOffsetBytes, bool InSkipFastClearFinalize = false)
-	{
-		Initialize(InDesc, InBaseShaderResource, InBaseShaderResource->ResourceLocation, InStride, InStartOffsetBytes, InSkipFastClearFinalize);
-	}
+	void InitializeAfterCreate(const D3D12_SHADER_RESOURCE_VIEW_DESC& InDesc, FD3D12BaseShaderResource* InBaseShaderResource, FD3D12ResourceLocation& InResourceLocation, uint32 InStride, uint32 InStartOffsetBytes = 0, bool InSkipFastClearFinalize = false);
+	void Update(FD3D12Buffer* InBuffer, const D3D12_SHADER_RESOURCE_VIEW_DESC& InDesc, uint32 InStride);
+	void UpdateMinLODClamp(float ResourceMinLODClamp);
 
-	void Initialize(FD3D12Device* InParent, D3D12_SHADER_RESOURCE_VIEW_DESC& InDesc, FD3D12BaseShaderResource* InBaseShaderResource, uint32 InStride, bool InSkipFastClearFinalize = false)
-	{
-		if (!this->GetParentDevice_Unsafe())
-		{
-			// This is a null SRV created without viewing on any resource
-			// We need to set its device and allocate a descriptor slot before moving forward
-			this->SetParentDevice(InParent);
-		}
-		check(GetParentDevice() == InParent);
-		Initialize(InDesc, InBaseShaderResource, InBaseShaderResource->ResourceLocation, InStride, 0, InSkipFastClearFinalize);
-	}
-
-	void Rename(FD3D12BaseShaderResource* InBaseShaderResource)
-	{
-		check(InBaseShaderResource == BaseShaderResource);
-		check(&InBaseShaderResource->ResourceLocation == ResourceLocation);
-
-		// Update the first element index, then reinitialize the SRV
-		if (Desc.ViewDimension == D3D12_SRV_DIMENSION_BUFFER)
-		{
-			uint32 StartElement = StartOffsetBytes / Stride;
-			Desc.Buffer.FirstElement = ResourceLocation->GetOffsetFromBaseOfResource() / Stride + StartElement;
-		}
-
-		Initialize(Desc, BaseShaderResource, Stride, StartOffsetBytes);
-	}
-
-	void Rename(float ResourceMinLODClamp)
-	{
-		check(bInitialized);
-		check(ResourceLocation);
-		check(Desc.ViewDimension == D3D12_SRV_DIMENSION_TEXTURE2D);
-
-		// Update the LODClamp, the reinitialize the SRV
-		Desc.Texture2D.ResourceMinLODClamp = ResourceMinLODClamp;
-		CreateView(BaseShaderResource, *ResourceLocation);
-	}
-
-	virtual void RecreateView() override
-	{
-		Rename(BaseShaderResource);
-	}
+	virtual void RecreateView() override;
 
 	FORCEINLINE bool IsDepthStencilResource()	const { return bContainsDepthPlane || bContainsStencilPlane; }
 	FORCEINLINE bool IsDepthPlaneResource()		const { return bContainsDepthPlane; }
@@ -884,44 +821,8 @@ public:
 	virtual FRHIDescriptorHandle GetBindlessHandle() const override { return Descriptor.GetBindlessHandle(); }
 
 protected:
-	void PreCreateView(FD3D12ResourceLocation& InResourceLocation, uint32 InStride, uint32 InStartOffsetBytes, bool InSkipFastClearFinalize)
-	{
-		Stride = InStride;
-		StartOffsetBytes = InStartOffsetBytes;
-		bSkipFastClearFinalize = InSkipFastClearFinalize;
-
-		if (InResourceLocation.GetResource())
-		{
-			bContainsDepthPlane = InResourceLocation.GetResource()->IsDepthStencilResource() && GetPlaneSliceFromViewFormat(InResourceLocation.GetResource()->GetDesc().Format, Desc.Format) == 0;
-			bContainsStencilPlane = InResourceLocation.GetResource()->IsDepthStencilResource() && GetPlaneSliceFromViewFormat(InResourceLocation.GetResource()->GetDesc().Format, Desc.Format) == 1;
-			bRequiresResourceStateTracking = InResourceLocation.GetResource()->RequiresResourceStateTracking();
-
-#if DO_CHECK
-			// Check the plane slice of the SRV matches the texture format
-			// Texture2DMS does not have explicit plane index (it's implied by the format)
-			if (Desc.ViewDimension == D3D12_SRV_DIMENSION_TEXTURE2D)
-			{
-				check(GetPlaneSliceFromViewFormat(InResourceLocation.GetResource()->GetDesc().Format, Desc.Format) == Desc.Texture2D.PlaneSlice);
-			}
-
-			if (Desc.ViewDimension == D3D12_SRV_DIMENSION_BUFFER)
-			{
-				//check(InResourceLocation.GetOffsetFromBaseOfResource() == Desc.Buffer.FirstElement * Stride);
-			}
-#endif
-		}
-	}
-
-	void CreateView(FD3D12BaseShaderResource* InBaseShaderResource, FD3D12ResourceLocation& InResourceLocation)
-	{
-		InitializeInternal(InBaseShaderResource, InResourceLocation);
-
-		if (ResourceLocation->GetResource())
-		{
-			ID3D12Resource* D3DResource = ResourceLocation->GetResource()->GetResource();
-			Descriptor.CreateView(Desc, D3DResource);
-		}
-	}
+	void PreCreateView(const FD3D12ResourceLocation& InResourceLocation, uint32 InStride, uint32 InStartOffsetBytes, bool InSkipFastClearFinalize);
+	void CreateView(FD3D12BaseShaderResource* InBaseShaderResource, FD3D12ResourceLocation& InResourceLocation, ED3D12DescriptorCreateReason Reason);
 
 	uint32 Stride{};
 	uint32 StartOffsetBytes{};
@@ -931,59 +832,14 @@ protected:
 	bool bRequiresResourceStateTracking : 1;
 };
 
-class FD3D12ShaderResourceViewWithLocation : public FD3D12ShaderResourceView
-{
-public:
-	FD3D12ShaderResourceViewWithLocation(FD3D12Device* InParent)
-		: FD3D12ShaderResourceView(InParent)
-		, ViewLocation(InParent)
-	{}
-
-	FD3D12ShaderResourceViewWithLocation(FD3D12Device* InParent, D3D12_SHADER_RESOURCE_VIEW_DESC& InDesc, FD3D12BaseShaderResource* InBaseShaderResource, uint32 InStride = -1, bool InSkipFastClearFinalize = false)
-		: FD3D12ShaderResourceView(InParent)
-		, ViewLocation(InParent)
-	{
-		// Should not be pool allocated
-		check(InBaseShaderResource->ResourceLocation.GetAllocatorType() != FD3D12ResourceLocation::EAllocatorType::AT_Pool);
-		Initialize(InDesc, InBaseShaderResource, ViewLocation, InStride, 0, InSkipFastClearFinalize);
-	}
-
-	virtual void RecreateView() override
-	{
-		// Currently unsupported because the view location resource would also need to be recreated then
-		check(false);
-	}
-
-	FD3D12ResourceLocation ViewLocation;
-};
-
 class FD3D12UnorderedAccessView : public FRHIUnorderedAccessView, public FD3D12View < D3D12_UNORDERED_ACCESS_VIEW_DESC >, public FD3D12LinkedAdapterObject<FD3D12UnorderedAccessView>
 {
+protected:
+	FD3D12UnorderedAccessView(FD3D12Device* InParent);
 public:
-	FD3D12UnorderedAccessView(FD3D12Device* InParent)
-	: FD3D12View(InParent, ERHIDescriptorHeapType::Standard, ViewSubresourceSubsetFlags_None)
-	{
-	}
+	FD3D12UnorderedAccessView(FD3D12Device* InParent, const D3D12_UNORDERED_ACCESS_VIEW_DESC& InDesc, FD3D12BaseShaderResource* InBaseShaderResource, FD3D12Resource* InCounterResource);
 
-	FD3D12UnorderedAccessView(FD3D12Device* InParent, D3D12_UNORDERED_ACCESS_VIEW_DESC& InDesc, FD3D12BaseShaderResource* InBaseShaderResource, FD3D12Resource* InCounterResource = nullptr)
-		: FD3D12View(InParent, ERHIDescriptorHeapType::Standard, ViewSubresourceSubsetFlags_None)
-		, CounterResource(InCounterResource)
-	{
-		Initialize(InDesc, InBaseShaderResource, InBaseShaderResource->ResourceLocation, InCounterResource);
-	}
-
-	void Initialize(D3D12_UNORDERED_ACCESS_VIEW_DESC& InDesc, FD3D12BaseShaderResource* InBaseShaderResource, FD3D12ResourceLocation& InResourceLocation, FD3D12Resource* InCounterResource = nullptr)
-	{
-		SetDesc(InDesc);
-		CreateView(InBaseShaderResource, InResourceLocation, InCounterResource);
-	}
-
-	virtual void RecreateView() override
-	{
-		check(CounterResource == nullptr);
-		check(ResourceLocation->GetOffsetFromBaseOfResource() == 0);
-		CreateView(BaseShaderResource, *ResourceLocation, nullptr);
-	}
+	virtual void RecreateView() override;
 
 	bool IsCounterResourceInitialized() const { return CounterResourceInitialized; }
 	void MarkCounterResourceInitialized() { CounterResourceInitialized = true; }
@@ -993,40 +849,12 @@ public:
 	virtual FRHIDescriptorHandle GetBindlessHandle() const override { return Descriptor.GetBindlessHandle(); }
 
 protected:
-	void CreateView(FD3D12BaseShaderResource* InBaseShaderResource, FD3D12ResourceLocation& InResourceLocation, FD3D12Resource* InCounterResource)
-	{
-		InitializeInternal(InBaseShaderResource, InResourceLocation);
-
-		if (Resource)
-		{
-			ID3D12Resource* D3DResource = Resource->GetUAVAccessResource() ? Resource->GetUAVAccessResource() : Resource->GetResource();
-			ID3D12Resource* D3DCounterResource = InCounterResource ? InCounterResource->GetResource() : nullptr;
-			Descriptor.CreateView(Desc, D3DResource, D3DCounterResource);
-		}
-	}
+	void CreateView(FD3D12BaseShaderResource* InBaseShaderResource, FD3D12ResourceLocation& InResourceLocation, FD3D12Resource* InCounterResource, ED3D12DescriptorCreateReason Reason);
 
 	TRefCountPtr<FD3D12Resource> CounterResource;
 	bool CounterResourceInitialized = false;
 };
 
-class FD3D12UnorderedAccessViewWithLocation : public FD3D12UnorderedAccessView
-{
-public:
-	FD3D12UnorderedAccessViewWithLocation(FD3D12Device* InParent)
-	: FD3D12UnorderedAccessView(InParent)
-	, ViewLocation(InParent)
-	{
-	}
-
-	virtual void RecreateView() override
-	{
-		// Currently unsupported because the view location resource would also need to be recreated then
-		check(false);
-	}
-
-	FD3D12ResourceLocation ViewLocation;
-};
-
 class FD3D12ConstantBufferView : public FD3D12DeviceChild
 {
 public:
diff --git a/Engine/Source/Runtime/D3D12RHI/Private/D3D12Viewport.cpp b/Engine/Source/Runtime/D3D12RHI/Private/D3D12Viewport.cpp
--- a/Engine/Source/Runtime/D3D12RHI/Private/D3D12Viewport.cpp
+++ b/Engine/Source/Runtime/D3D12RHI/Private/D3D12Viewport.cpp
@@ -333,7 +333,7 @@ FD3D12Texture2D* GetSwapChainSurface(FD3D12Device* Parent, EPixelFormat PixelFor
 		SRVDesc.Texture2D.MostDetailedMip = 0;
 		SRVDesc.Texture2D.MipLevels = 1;
 
-		FD3D12ShaderResourceView* WrappedShaderResourceView = new FD3D12ShaderResourceView(Device, SRVDesc, NewTexture);
+		FD3D12ShaderResourceView* WrappedShaderResourceView = new FD3D12ShaderResourceView(NewTexture, SRVDesc);
 		NewTexture->SetShaderResourceView(WrappedShaderResourceView);
 
 		return NewTexture;
diff --git a/Engine/Source/Runtime/RHICore/Public/RHIDescriptorAllocator.h b/Engine/Source/Runtime/RHICore/Public/RHIDescriptorAllocator.h
--- a/Engine/Source/Runtime/RHICore/Public/RHIDescriptorAllocator.h
+++ b/Engine/Source/Runtime/RHICore/Public/RHIDescriptorAllocator.h
@@ -44,7 +44,9 @@ public:
 	void Free(uint32 Slot, uint32 NumDescriptors);
 
 	using FRHIDescriptorAllocator::GetCapacity;
-	inline bool HandlesAllocation(ERHIDescriptorHeapType InType) const { return Type == InType; }
+	inline ERHIDescriptorHeapType GetType() const { return Type; }
+
+	inline bool HandlesAllocation(ERHIDescriptorHeapType InType) const { return GetType() == InType; }
 
 private:
 	ERHIDescriptorHeapType Type;
