Use correct stride when computing max element count with ByteAddressBuffers

diff --git a/Engine/Source/Runtime/D3D12RHI/Private/D3D12View.cpp b/Engine/Source/Runtime/D3D12RHI/Private/D3D12View.cpp
--- a/Engine/Source/Runtime/D3D12RHI/Private/D3D12View.cpp
+++ b/Engine/Source/Runtime/D3D12RHI/Private/D3D12View.cpp
@@ -278,9 +278,6 @@ static FORCEINLINE D3D12_SHADER_RESOURCE_VIEW_DESC GetStructuredBufferSRVDesc(co
 	const uint32 BufferStride = StructuredBuffer->GetStride();
 	const uint32 BufferOffset = StructuredBuffer->ResourceLocation.GetOffsetFromBaseOfResource();
 
-	const uint32 MaxElements = BufferSize / BufferStride;
-	const uint32 StartElement = FMath::Min<uint32>(StartOffsetBytes, BufferSize) / BufferStride;
-
 	D3D12_SHADER_RESOURCE_VIEW_DESC SRVDesc{};
 
 	SRVDesc.Shader4ComponentMapping = D3D12_DEFAULT_SHADER_4_COMPONENT_MAPPING;
@@ -299,6 +296,9 @@ static FORCEINLINE D3D12_SHADER_RESOURCE_VIEW_DESC GetStructuredBufferSRVDesc(co
 		CreationStride = BufferStride;
 	}
 
+	const uint32 MaxElements = BufferSize / CreationStride;
+	const uint32 StartElement = FMath::Min<uint32>(StartOffsetBytes, BufferSize) / CreationStride;
+
 	SRVDesc.Buffer.NumElements = FMath::Min<uint32>(MaxElements - StartElement, NumElements);
 	SRVDesc.Buffer.FirstElement = (BufferOffset / CreationStride) + StartElement;
 
