Fix D3D Validation Error: RHIScratchBuffer is invalid for use as a acceleration structure build scratch

diff --git a/Engine/Source/Runtime/D3D12RHI/Private/D3D12Allocation.cpp b/Engine/Source/Runtime/D3D12RHI/Private/D3D12Allocation.cpp
--- a/Engine/Source/Runtime/D3D12RHI/Private/D3D12Allocation.cpp
+++ b/Engine/Source/Runtime/D3D12RHI/Private/D3D12Allocation.cpp
@@ -1412,11 +1412,6 @@ D3D12_RESOURCE_STATES FD3D12DefaultBufferAllocator::GetDefaultInitialResourceSta
 		check(InHeapType == D3D12_HEAP_TYPE_DEFAULT);
 		return D3D12_RESOURCE_STATE_RAYTRACING_ACCELERATION_STRUCTURE;
 	}
-	else if (EnumHasAnyFlags(InBufferFlags, BUF_RayTracingScratch))
-	{
-		check(InHeapType == D3D12_HEAP_TYPE_DEFAULT);
-		return D3D12_RESOURCE_STATE_UNORDERED_ACCESS;
-	}
 #endif // D3D12_RHI_RAYTRACING
 	else
 	{
diff --git a/Engine/Source/Runtime/D3D12RHI/Private/D3D12RayTracing.cpp b/Engine/Source/Runtime/D3D12RHI/Private/D3D12RayTracing.cpp
--- a/Engine/Source/Runtime/D3D12RHI/Private/D3D12RayTracing.cpp
+++ b/Engine/Source/Runtime/D3D12RHI/Private/D3D12RayTracing.cpp
@@ -4291,6 +4291,14 @@ void FD3D12CommandContext::RHIBuildAccelerationStructures(const TArrayView<const
 		Geometry->TransitionBuffers(*this);
 	}
 
+	{
+		FD3D12Buffer* ScratchBuffer = FD3D12DynamicRHI::ResourceCast(ScratchBufferRange.Buffer, GetGPUIndex());
+		if (ScratchBuffer->GetResource()->RequiresResourceStateTracking())
+		{
+			FD3D12DynamicRHI::TransitionResource(CommandListHandle, ScratchBuffer->GetResource(), D3D12_RESOURCE_STATE_TBD, D3D12_RESOURCE_STATE_UNORDERED_ACCESS, 0, FD3D12DynamicRHI::ETransitionMode::Validate);
+		}
+	}
+
 	CommandListHandle.FlushResourceBarriers();
 		
 	const uint32 GPUIndex = GetGPUIndex();
