Fix race condition where multiple threads access a shared FD3D12RayTracingShaderTable::TransientCBVs array.

The solution is to give each worker thread its own version of this array, similar to TransitionSRVs and other such members of FD3D12RayTracingShaderTable.

diff --git a/Engine/Source/Runtime/D3D12RHI/Private/D3D12RayTracing.cpp b/Engine/Source/Runtime/D3D12RHI/Private/D3D12RayTracing.cpp
--- a/Engine/Source/Runtime/D3D12RHI/Private/D3D12RayTracing.cpp
+++ b/Engine/Source/Runtime/D3D12RHI/Private/D3D12RayTracing.cpp
@@ -1754,9 +1754,12 @@ public:
 	{
 		delete DescriptorCache;
 	#if USE_STATIC_ROOT_SIGNATURE
-		for (FD3D12ConstantBufferView* CBV : TransientCBVs)
+		for (const TArray<FD3D12ConstantBufferView*>& Container : TransientCBVs)
 		{
-			delete CBV;
+			for (FD3D12ConstantBufferView* CBV : Container)
+			{
+				delete CBV;
+			}
 		}
 	#endif // USE_STATIC_ROOT_SIGNATURE
 	}
@@ -2169,7 +2172,7 @@ public:
 	}
 
 #if USE_STATIC_ROOT_SIGNATURE
-	TArray<FD3D12ConstantBufferView*> TransientCBVs;
+	TArray<FD3D12ConstantBufferView*> TransientCBVs[MaxBindingWorkers];
 #endif // USE_STATIC_ROOT_SIGNATURE
 };
 
@@ -4539,7 +4542,7 @@ struct FD3D12RayTracingLocalResourceBinder
 	{
 		FD3D12ConstantBufferView* BufferView = new FD3D12ConstantBufferView(GetDevice());
 		BufferView->Create(Address, DataSize);
-		ShaderTable.TransientCBVs.Add(BufferView);
+		ShaderTable.TransientCBVs[WorkerIndex].Add(BufferView);
 		return BufferView->GetOfflineCpuHandle();
 	}
 #endif // USE_STATIC_ROOT_SIGNATURE
