D3D12 bindless descriptor deferred "deletion".
Took time to clean up FD3D12DeferredDeletionQueue, mainly hiding the internals of what it does from the public interface.

diff --git a/Engine/Source/Runtime/D3D12RHI/Private/D3D12Descriptors.cpp b/Engine/Source/Runtime/D3D12RHI/Private/D3D12Descriptors.cpp
--- a/Engine/Source/Runtime/D3D12RHI/Private/D3D12Descriptors.cpp
+++ b/Engine/Source/Runtime/D3D12RHI/Private/D3D12Descriptors.cpp
@@ -132,7 +132,7 @@ void FD3D12BindlessDescriptorManager::Init(uint32 InNumResourceDescriptors, uint
 	}
 }
 
-FRHIDescriptorHandle FD3D12BindlessDescriptorManager::AllocateDescriptor(ERHIDescriptorHeapType InType)
+FRHIDescriptorHandle FD3D12BindlessDescriptorManager::Allocate(ERHIDescriptorHeapType InType)
 {
 	for (FD3D12DescriptorManager& Manager : Managers)
 	{
@@ -144,7 +144,7 @@ FRHIDescriptorHandle FD3D12BindlessDescriptorManager::AllocateDescriptor(ERHIDes
 	return FRHIDescriptorHandle();
 }
 
-void FD3D12BindlessDescriptorManager::FreeDescriptor(FRHIDescriptorHandle InHandle)
+void FD3D12BindlessDescriptorManager::ImmediateFree(FRHIDescriptorHandle InHandle)
 {
 	for (FD3D12DescriptorManager& Manager : Managers)
 	{
@@ -159,6 +159,15 @@ void FD3D12BindlessDescriptorManager::FreeDescriptor(FRHIDescriptorHandle InHand
 	checkNoEntry();
 }
 
+void FD3D12BindlessDescriptorManager::DeferredFreeFromDestructor(FRHIDescriptorHandle InHandle)
+{
+	if (InHandle.IsValid())
+	{
+		FD3D12Fence& Fence = GetParentDevice()->GetCommandListManager().GetFence();
+		GetParentDevice()->GetParentAdapter()->GetDeferredDeletionQueue().EnqueueBindlessDescriptor(InHandle, &Fence, GetParentDevice()->GetGPUIndex());
+	}
+}
+
 void FD3D12BindlessDescriptorManager::UpdateImmediately(FRHIDescriptorHandle InHandle, D3D12_CPU_DESCRIPTOR_HANDLE InSourceCpuHandle)
 {
 	for (FD3D12DescriptorManager& Manager : Managers)
diff --git a/Engine/Source/Runtime/D3D12RHI/Private/D3D12Descriptors.h b/Engine/Source/Runtime/D3D12RHI/Private/D3D12Descriptors.h
--- a/Engine/Source/Runtime/D3D12RHI/Private/D3D12Descriptors.h
+++ b/Engine/Source/Runtime/D3D12RHI/Private/D3D12Descriptors.h
@@ -98,8 +98,9 @@ public:
 
 	void Init(uint32 InNumResourceDescriptors, uint32 InNumSamplerDescriptors);
 
-	FRHIDescriptorHandle AllocateDescriptor(ERHIDescriptorHeapType InType);
-	void FreeDescriptor(FRHIDescriptorHandle InHandle);
+	FRHIDescriptorHandle Allocate(ERHIDescriptorHeapType InType);
+	void ImmediateFree(FRHIDescriptorHandle InHandle);
+	void DeferredFreeFromDestructor(FRHIDescriptorHandle InHandle);
 
 	void UpdateImmediately(FRHIDescriptorHandle InHandle, D3D12_CPU_DESCRIPTOR_HANDLE InSourceCpuHandle);
 	void UpdateDeferred(FRHIDescriptorHandle InHandle, D3D12_CPU_DESCRIPTOR_HANDLE InSourceCpuHandle);
diff --git a/Engine/Source/Runtime/D3D12RHI/Private/D3D12Resources.cpp b/Engine/Source/Runtime/D3D12RHI/Private/D3D12Resources.cpp
--- a/Engine/Source/Runtime/D3D12RHI/Private/D3D12Resources.cpp
+++ b/Engine/Source/Runtime/D3D12RHI/Private/D3D12Resources.cpp
@@ -31,12 +31,106 @@ static FAutoConsoleVariableRef CVarAsyncDeferredDeletion(
 	ECVF_ReadOnly
 );
 
-/////////////////////////////////////////////////////////////////////
-//	FD3D12 Deferred Deletion Queue
-/////////////////////////////////////////////////////////////////////
+enum class ED3D12DeferredDeleteObjectType
+{
+	RHIObject,
+	D3DObject,
+	BindlessDescriptor,
+};
+
+// We can't store the complex FRHIDescriptorHandle type in the union. We also need to store the device index.
+struct FD3D12SimpleBindlessDescriptor
+{
+	uint32 Index;
+	uint16 Type;
+	uint16 DeviceIndex;
+};
+
+struct FD3D12DeferredDeleteObject
+{
+	union
+	{
+		FD3D12Resource*                RHIObject;
+		ID3D12Object*                  D3DObject;
+		FD3D12SimpleBindlessDescriptor BindlessDescriptor;
+	};
+	ED3D12DeferredDeleteObjectType Type;
+};
+
+struct FD3D12FencedDeleteObject : FD3D12DeferredDeleteObject
+{
+	using FFencePair = TPair<FD3D12Fence*, uint64>;
+	using FFenceList = TArray<FFencePair, TInlineAllocator<1>>;
 
-FD3D12DeferredDeletionQueue::FD3D12DeferredDeletionQueue(FD3D12Adapter* InParent) :
-	FD3D12AdapterChild(InParent) {}
+	FFenceList FenceList;
+};
+
+static bool IsFencedObjectComplete(const FD3D12FencedDeleteObject& FencedObject)
+{
+	for (const auto& FencePair : FencedObject.FenceList)
+	{
+		if (!FencePair.Key->IsFenceComplete(FencePair.Value))
+		{
+			return false;
+		}
+	}
+	return true;
+}
+
+////////////////////////////////////////////////////////////////////////////////////////////////////////////////////////
+// FD3D12AsyncDeletionWorker
+
+class FD3D12AsyncDeletionWorker : public FD3D12AdapterChild, public FNonAbandonableTask
+{
+public:
+	FD3D12AsyncDeletionWorker(FD3D12Adapter* Adapter, FThreadsafeQueue<FD3D12FencedDeleteObject>& DeletionQueue)
+		: FD3D12AdapterChild(Adapter)
+	{
+		DeletionQueue.BatchDequeue(ObjectsToDelete, &IsFencedObjectComplete, 4096);
+	}
+
+	void DoWork()
+	{
+		for (const FD3D12DeferredDeleteObject& ObjectToDelete : ObjectsToDelete)
+		{
+			switch (ObjectToDelete.Type)
+			{
+			case ED3D12DeferredDeleteObjectType::RHIObject:
+				// This should be a final release.
+				check(ObjectToDelete.RHIObject->GetRefCount() == 1);
+				ObjectToDelete.RHIObject->Release();
+				break;
+			case ED3D12DeferredDeleteObjectType::D3DObject:
+				ObjectToDelete.D3DObject->Release();
+				break;
+			case ED3D12DeferredDeleteObjectType::BindlessDescriptor:
+				{
+					const FRHIDescriptorHandle BindlessDescriptor(static_cast<ERHIDescriptorHeapType>(ObjectToDelete.BindlessDescriptor.Type), ObjectToDelete.BindlessDescriptor.Index);
+					GetParentAdapter()->GetDevice(ObjectToDelete.BindlessDescriptor.DeviceIndex)->GetBindlessDescriptorManager().ImmediateFree(BindlessDescriptor);
+				}
+				break;
+			default:
+				checkf(false, TEXT("Unknown ED3D12DeferredDeleteObjectType"));
+				break;
+			}
+		}
+	}
+
+	FORCEINLINE TStatId GetStatId() const
+	{
+		RETURN_QUICK_DECLARE_CYCLE_STAT(FD3D12AsyncDeletionWorker, STATGROUP_ThreadPoolAsyncTasks);
+	}
+
+private:
+	TArray<FD3D12DeferredDeleteObject> ObjectsToDelete;
+};
+
+////////////////////////////////////////////////////////////////////////////////////////////////////////////////////////
+
+FD3D12DeferredDeletionQueue::FD3D12DeferredDeletionQueue(FD3D12Adapter* InParent)
+	: FD3D12AdapterChild(InParent)
+{
+}
 
 FD3D12DeferredDeletionQueue::~FD3D12DeferredDeletionQueue()
 {
@@ -49,17 +143,24 @@ FD3D12DeferredDeletionQueue::~FD3D12DeferredDeletionQueue()
 	}
 }
 
-void FD3D12DeferredDeletionQueue::EnqueueResource(FD3D12Resource* pResource, FFenceList&& FenceList)
+void FD3D12DeferredDeletionQueue::EnqueueResource(FD3D12Resource* pResource, FRHIGPUMask GpuMask)
 {
 	check(pResource->ShouldDeferDelete());
 
 	// Useful message for identifying when resources are released on the rendering thread.
 	//UE_CLOG(IsInActualRenderingThread(), LogD3D12RHI, Display, TEXT("Rendering Thread: Deleting %#016llx when done with frame fence %llu"), pResource, Fence->GetCurrentFence());
 
-	FencedObjectType FencedObject;
-	FencedObject.RHIObject  = pResource;
-	FencedObject.FenceList  = MoveTemp(FenceList);
-	FencedObject.Type       = EObjectType::RHI;
+	FD3D12Adapter* Adapter = GetParentAdapter();
+
+	FD3D12FencedDeleteObject FencedObject;
+	for (uint32 GPUIndex : GpuMask)
+	{
+		FD3D12Fence* Fence = &Adapter->GetDevice(GPUIndex)->GetCommandListManager().GetFence();
+		FencedObject.FenceList.Emplace(Fence, Fence->GetCurrentFence());
+	}
+
+	FencedObject.RHIObject = pResource;
+	FencedObject.Type = ED3D12DeferredDeleteObjectType::RHIObject;
 	DeferredReleaseQueue.Enqueue(FencedObject);
 }
 
@@ -68,10 +169,21 @@ void FD3D12DeferredDeletionQueue::EnqueueResource(ID3D12Object* pResource, FD3D1
 	// Useful message for identifying when resources are released on the rendering thread.
 	//UE_CLOG(IsInActualRenderingThread(), LogD3D12RHI, Display, TEXT("Rendering Thread: Deleting %#016llx when done with frame fence %llu"), pResource, Fence->GetCurrentFence());
 
-	FencedObjectType FencedObject;
-	FencedObject.D3DObject  = pResource;
+	FD3D12FencedDeleteObject FencedObject;
+	FencedObject.D3DObject = pResource;
 	FencedObject.FenceList.Emplace(Fence, Fence->GetCurrentFence());
-	FencedObject.Type       = EObjectType::D3D;
+	FencedObject.Type = ED3D12DeferredDeleteObjectType::D3DObject;
+	DeferredReleaseQueue.Enqueue(FencedObject);
+}
+
+void FD3D12DeferredDeletionQueue::EnqueueBindlessDescriptor(FRHIDescriptorHandle InDescriptor, FD3D12Fence* InFence, uint32 InDeviceIndex)
+{
+	FD3D12FencedDeleteObject FencedObject;
+	FencedObject.BindlessDescriptor.Index = InDescriptor.GetIndex();
+	FencedObject.BindlessDescriptor.Type = static_cast<uint32>(InDescriptor.GetType());
+	FencedObject.BindlessDescriptor.DeviceIndex = InDeviceIndex;
+	FencedObject.FenceList.Emplace(InFence, InFence->GetCurrentFence());
+	FencedObject.Type = ED3D12DeferredDeleteObjectType::BindlessDescriptor;
 	DeferredReleaseQueue.Enqueue(FencedObject);
 }
 
@@ -107,7 +219,7 @@ bool FD3D12DeferredDeletionQueue::ReleaseResources(bool bDeleteImmediately, bool
 			}
 
 			// Create new delete task, which will only collect resources in the constructor for which the fence is complete, not the whole list!
-			DeleteTask = new FAsyncTask<FD3D12AsyncDeletionWorker>(Adapter, &DeferredReleaseQueue);
+			DeleteTask = new FAsyncTask<FD3D12AsyncDeletionWorker>(Adapter, DeferredReleaseQueue);
 
 			DeleteTask->StartBackgroundTask();
 			DeleteTasks.Enqueue(DeleteTask);
@@ -117,7 +229,7 @@ bool FD3D12DeferredDeletionQueue::ReleaseResources(bool bDeleteImmediately, bool
 		}
 	}
 
-	FencedObjectType FenceObject;
+	FD3D12FencedDeleteObject FenceObject;
 
 	if (bIsShutDown)
 	{
@@ -126,103 +238,70 @@ bool FD3D12DeferredDeletionQueue::ReleaseResources(bool bDeleteImmediately, bool
 
 		while (DeferredReleaseQueue.Dequeue(FenceObject))
 		{
-			if (FenceObject.Type == EObjectType::RHI)
+			switch (FenceObject.Type)
 			{
-				D3D12_RESOURCE_DESC Desc = FenceObject.RHIObject->GetDesc();
-				FString Name = FenceObject.RHIObject->GetName().ToString();
-				UE_LOG(LogD3D12RHI, Display, TEXT("D3D12 ReleaseResources: \"%s\", %llu x %u x %u, Mips: %u, Format: 0x%X, Flags: 0x%X"), *Name, Desc.Width, Desc.Height, Desc.DepthOrArraySize, Desc.MipLevels, Desc.Format, Desc.Flags);
-
-				uint32 RefCount = FenceObject.RHIObject->Release();
-				if (RefCount)
+			case ED3D12DeferredDeleteObjectType::RHIObject:
 				{
-					UE_LOG(LogD3D12RHI, Display, TEXT("RefCount was %u"), RefCount);
+					const D3D12_RESOURCE_DESC& Desc = FenceObject.RHIObject->GetDesc();
+					const FString Name = FenceObject.RHIObject->GetName().ToString();
+					UE_LOG(LogD3D12RHI, Display, TEXT("D3D12 ReleaseResources: \"%s\", %llu x %u x %u, Mips: %u, Format: 0x%X, Flags: 0x%X"), *Name, Desc.Width, Desc.Height, Desc.DepthOrArraySize, Desc.MipLevels, Desc.Format, Desc.Flags);
+
+					if (uint32 RefCount = FenceObject.RHIObject->Release())
+					{
+						UE_LOG(LogD3D12RHI, Display, TEXT("RefCount was %u"), RefCount);
+					}
 				}
-			}
-			else
-			{
-				UE_LOG(LogD3D12RHI, Display, TEXT("D3D12 ReleaseResources: 0x%llX"), FenceObject.D3DObject);
+				break;
+			case ED3D12DeferredDeleteObjectType::D3DObject:
+				{
+					UE_LOG(LogD3D12RHI, Display, TEXT("D3D12 ReleaseResources: 0x%llX"), FenceObject.D3DObject);
 
-				uint32 RefCount = FenceObject.D3DObject->Release();
-				if (RefCount)
+					if (uint32 RefCount = FenceObject.D3DObject->Release())
+					{
+						UE_LOG(LogD3D12RHI, Display, TEXT("RefCount was %u"), RefCount);
+					}
+				}
+				break;
+			case ED3D12DeferredDeleteObjectType::BindlessDescriptor:
 				{
-					UE_LOG(LogD3D12RHI, Display, TEXT("RefCount was %u"), RefCount);
+					const FRHIDescriptorHandle BindlessDescriptor(static_cast<ERHIDescriptorHeapType>(FenceObject.BindlessDescriptor.Type), FenceObject.BindlessDescriptor.Index);
+					GetParentAdapter()->GetDevice(FenceObject.BindlessDescriptor.DeviceIndex)->GetBindlessDescriptorManager().ImmediateFree(BindlessDescriptor);
 				}
+				break;
+			default:
+				checkf(false, TEXT("Unknown ED3D12DeferredDeleteObjectType"));
+				break;
 			}
 		}
 	}
 	else
 	{
-		struct FDequeueFenceObject
+		while (DeferredReleaseQueue.Dequeue(FenceObject, &IsFencedObjectComplete))
 		{
-			bool operator() (FencedObjectType FenceObject) const
-			{
-				for (auto& FencePair : FenceObject.FenceList)
-				{
-					if (!FencePair.Key->IsFenceComplete(FencePair.Value))
-					{
-						return false;
-					}
-				}
-				return true;
-			}
-		};
 
-		while (DeferredReleaseQueue.Dequeue(FenceObject, FDequeueFenceObject()))
-		{
-			if (FenceObject.Type == EObjectType::RHI)
+			switch (FenceObject.Type)
 			{
+			case ED3D12DeferredDeleteObjectType::RHIObject:
 				FenceObject.RHIObject->Release();
-			}
-			else
-			{
+				break;
+			case ED3D12DeferredDeleteObjectType::D3DObject:
 				FenceObject.D3DObject->Release();
-			}
-		}
-	}
-
-	return DeferredReleaseQueue.IsEmpty();
-}
-
-FD3D12DeferredDeletionQueue::FD3D12AsyncDeletionWorker::FD3D12AsyncDeletionWorker(FD3D12Adapter* Adapter, FThreadsafeQueue<FencedObjectType>* DeletionQueue)
-	: FD3D12AdapterChild(Adapter)
-{
-	struct FDequeueFenceObject
-	{
-		bool operator() (FencedObjectType FenceObject) const
-		{
-			for (auto& FencePair : FenceObject.FenceList)
-			{
-				if (!FencePair.Key->IsFenceComplete(FencePair.Value))
+				break;
+			case ED3D12DeferredDeleteObjectType::BindlessDescriptor:
 				{
-					return false;
+					const FRHIDescriptorHandle BindlessDescriptor(static_cast<ERHIDescriptorHeapType>(FenceObject.BindlessDescriptor.Type), FenceObject.BindlessDescriptor.Index);
+					GetParentAdapter()->GetDevice(FenceObject.BindlessDescriptor.DeviceIndex)->GetBindlessDescriptorManager().ImmediateFree(BindlessDescriptor);
 				}
+				break;
+			default:
+				checkf(false, TEXT("Unknown ED3D12DeferredDeleteObjectType"));
+				break;
 			}
-			return true;
-		}
-	};
-
-	DeletionQueue->BatchDequeue(&Queue, FDequeueFenceObject(), 4096);
-}
-
-void FD3D12DeferredDeletionQueue::FD3D12AsyncDeletionWorker::DoWork()
-{
-	FencedObjectType ResourceToDelete;
-
-	while (Queue.Dequeue(ResourceToDelete))
-	{
-		if (ResourceToDelete.Type == EObjectType::RHI)
-		{
-			// This should be a final release.
-			check(ResourceToDelete.RHIObject->GetRefCount() == 1);
-			ResourceToDelete.RHIObject->Release();
-		}
-		else
-		{
-			ResourceToDelete.D3DObject->Release();
 		}
 	}
-}
 
+	return DeferredReleaseQueue.IsEmpty();
+}
 
 /////////////////////////////////////////////////////////////////////
 //	ID3D12ResourceAllocator
@@ -377,22 +456,10 @@ void FD3D12Resource::DeferDelete()
 
 	// Upload heaps such as texture lock data can be referenced by multiple GPUs so we
 	// must wait for all of them to finish before releasing.
-	FD3D12DeferredDeletionQueue::FFenceList FenceList;
-	if (HeapType == D3D12_HEAP_TYPE_UPLOAD)
-	{
-		for (uint32 GPUIndex : FRHIGPUMask::All())
-		{
-			FD3D12Fence* Fence = &Adapter->GetDevice(GPUIndex)->GetCommandListManager().GetFence();
-			FenceList.Emplace(Fence, Fence->GetCurrentFence());
-		}
-	}
-	else
-	{
-		FD3D12Fence* Fence = &GetParentDevice()->GetCommandListManager().GetFence();
-		FenceList.Emplace(Fence, Fence->GetCurrentFence());
-	}
 
-	Adapter->GetDeferredDeletionQueue().EnqueueResource(this, MoveTemp(FenceList));
+	const FRHIGPUMask GpuMask = (HeapType == D3D12_HEAP_TYPE_UPLOAD) ? FRHIGPUMask::All() : GetParentDevice()->GetGPUMask();
+
+	Adapter->GetDeferredDeletionQueue().EnqueueResource(this, GpuMask);
 }
 
 /////////////////////////////////////////////////////////////////////
diff --git a/Engine/Source/Runtime/D3D12RHI/Private/D3D12State.cpp b/Engine/Source/Runtime/D3D12RHI/Private/D3D12State.cpp
--- a/Engine/Source/Runtime/D3D12RHI/Private/D3D12State.cpp
+++ b/Engine/Source/Runtime/D3D12RHI/Private/D3D12State.cpp
@@ -545,7 +545,7 @@ FD3D12SamplerState::FD3D12SamplerState(FD3D12Device* InParent, const D3D12_SAMPL
 	GetParentDevice()->CreateSamplerInternal(Desc, OfflineHandle);
 
 	FD3D12BindlessDescriptorManager& BindlessDescriptorManager = GetParentDevice()->GetBindlessDescriptorManager();
-	BindlessHandle = BindlessDescriptorManager.AllocateDescriptor(ERHIDescriptorHeapType::Sampler);
+	BindlessHandle = BindlessDescriptorManager.Allocate(ERHIDescriptorHeapType::Sampler);
 
 	if (BindlessHandle.IsValid())
 	{
@@ -562,7 +562,7 @@ FD3D12SamplerState::~FD3D12SamplerState()
 
 		if (BindlessHandle.IsValid())
 		{
-			GetParentDevice()->GetBindlessDescriptorManager().FreeDescriptor(BindlessHandle);
+			GetParentDevice()->GetBindlessDescriptorManager().DeferredFreeFromDestructor(BindlessHandle);
 		}
 	}
 }
diff --git a/Engine/Source/Runtime/D3D12RHI/Private/D3D12View.cpp b/Engine/Source/Runtime/D3D12RHI/Private/D3D12View.cpp
--- a/Engine/Source/Runtime/D3D12RHI/Private/D3D12View.cpp
+++ b/Engine/Source/Runtime/D3D12RHI/Private/D3D12View.cpp
@@ -84,7 +84,7 @@ void FD3D12ViewDescriptorHandle::AllocateDescriptorSlot()
 
 		if (HeapType == ERHIDescriptorHeapType::Standard)
 		{
-			BindlessHandle = Device->GetBindlessDescriptorManager().AllocateDescriptor(ERHIDescriptorHeapType::Standard);
+			BindlessHandle = Device->GetBindlessDescriptorManager().Allocate(ERHIDescriptorHeapType::Standard);
 		}
 	}
 }
@@ -99,7 +99,7 @@ void FD3D12ViewDescriptorHandle::FreeDescriptorSlot()
 
 		if (BindlessHandle.IsValid())
 		{
-			Device->GetBindlessDescriptorManager().FreeDescriptor(BindlessHandle);
+			Device->GetBindlessDescriptorManager().DeferredFreeFromDestructor(BindlessHandle);
 			BindlessHandle = FRHIDescriptorHandle();
 		}
 	}
diff --git a/Engine/Source/Runtime/D3D12RHI/Public/D3D12Resources.h b/Engine/Source/Runtime/D3D12RHI/Public/D3D12Resources.h
--- a/Engine/Source/Runtime/D3D12RHI/Public/D3D12Resources.h
+++ b/Engine/Source/Runtime/D3D12RHI/Public/D3D12Resources.h
@@ -723,64 +723,28 @@ struct ID3D12ResourceAllocator
 		D3D12_RESOURCE_STATES InCreateState, const D3D12_CLEAR_VALUE* InClearValue, const TCHAR* InName, FD3D12ResourceLocation& ResourceLocation) = 0;
 };
 
+struct FD3D12FencedDeleteObject;
+class FD3D12AsyncDeletionWorker;
+
 class FD3D12DeferredDeletionQueue : public FD3D12AdapterChild
 {
 public:
-	using FFencePair = TPair<FD3D12Fence*, uint64>;
-	using FFenceList = TArray<FFencePair, TInlineAllocator<1>>;
-
-private:
-	enum class EObjectType
-	{
-		RHI,
-		D3D,
-	};
-
-	struct FencedObjectType
-	{
-		union
-		{
-			FD3D12Resource* RHIObject;
-			ID3D12Object*   D3DObject;
-		};
-		FFenceList FenceList;
-		EObjectType Type;
-	};
-	FThreadsafeQueue<FencedObjectType> DeferredReleaseQueue;
-
-public:
-
-	inline const uint32 QueueSize() const { return DeferredReleaseQueue.GetSize(); }
-
-	void EnqueueResource(FD3D12Resource* pResource, FFenceList&& FenceList);
-	void EnqueueResource(ID3D12Object* pResource, FD3D12Fence* Fence);
-
-	bool ReleaseResources(bool bDeleteImmediately, bool bIsShutDown);
-
 	FD3D12DeferredDeletionQueue(FD3D12Adapter* InParent);
 	~FD3D12DeferredDeletionQueue();
 
-	class FD3D12AsyncDeletionWorker : public FD3D12AdapterChild, public FNonAbandonableTask
-	{
-	public:
-		FD3D12AsyncDeletionWorker(FD3D12Adapter* Adapter, FThreadsafeQueue<FencedObjectType>* DeletionQueue);
-
-		void DoWork();
+	void EnqueueResource(FD3D12Resource* pResource, FRHIGPUMask GpuMask);
+	void EnqueueResource(ID3D12Object* pResource, FD3D12Fence* Fence);
+	void EnqueueBindlessDescriptor(FRHIDescriptorHandle InDescriptor, FD3D12Fence* InFence, uint32 InDeviceIndex);
 
-		FORCEINLINE TStatId GetStatId() const
-		{
-			RETURN_QUICK_DECLARE_CYCLE_STAT(FD3D12AsyncDeletionWorker, STATGROUP_ThreadPoolAsyncTasks);
-		}
+	bool ReleaseResources(bool bDeleteImmediately, bool bIsShutDown);
 
-	private:
-		TQueue<FencedObjectType> Queue;
-	};
+	inline const uint32 QueueSize() const { return DeferredReleaseQueue.GetSize(); }
 
 private:
+	FThreadsafeQueue<FD3D12FencedDeleteObject> DeferredReleaseQueue;
 
 	FCriticalSection DeleteTaskCS;
 	TQueue<FAsyncTask<FD3D12AsyncDeletionWorker>*> DeleteTasks;
-
 };
 
 struct FD3D12LockedResource : public FD3D12DeviceChild
diff --git a/Engine/Source/Runtime/D3D12RHI/Public/D3D12Util.h b/Engine/Source/Runtime/D3D12RHI/Public/D3D12Util.h
--- a/Engine/Source/Runtime/D3D12RHI/Public/D3D12Util.h
+++ b/Engine/Source/Runtime/D3D12RHI/Public/D3D12Util.h
@@ -375,8 +375,8 @@ public:
 		return false;
 	}
 
-	template <typename CompareFunc>
-	bool BatchDequeue(TQueue<Type>* Result, const CompareFunc& Func, uint32 MaxItems)
+	template <typename ResultType, typename CompareFunc>
+	bool BatchDequeue(TArray<ResultType>& Result, const CompareFunc& Func, uint32 MaxItems)
 	{
 		FScopeLock ScopeLock(&SynchronizationObject);
 
@@ -388,7 +388,7 @@ public:
 			{
 				Items.Dequeue(Item);
 				Size--;
-				Result->Enqueue(Item);
+				Result.Emplace(Item);
 
 				i++;
 			}
