MGPU:  Fix for D3D12 validation errors related to simultaneous access of buffer resource on buffer initialization

With multiple GPU support, we need to issue staging buffer upload commands on the command context for the same device (GPU) that the resource is on.  So we always use the default command context per GPU, and ignore the command context passed in.  In practice, the caller will already be passing the default command context in, but if we run into a situation where that's not the case, it would require some sort of higher level refactor of the code (either moving the linked object iterator loop to a higher level, or introducing a cross GPU fence sync at the end of an initialization batch).  There is an assert to identify if we've encountered such a case, so we know we need to solve it.

#rnx

diff --git a/Engine/Source/Runtime/D3D12RHI/Private/D3D12Buffer.cpp b/Engine/Source/Runtime/D3D12RHI/Private/D3D12Buffer.cpp
--- a/Engine/Source/Runtime/D3D12RHI/Private/D3D12Buffer.cpp
+++ b/Engine/Source/Runtime/D3D12RHI/Private/D3D12Buffer.cpp
@@ -104,12 +104,35 @@ struct FD3D12RHICommandInitializeBuffer final : public FRHICommand<FD3D12RHIComm
 
 	void ExecuteOnCommandContext(FD3D12CommandContext& CommandContext)
 	{
+#if WITH_MGPU
+		// With multiple GPU support, we need to issue staging buffer upload commands on the command context for the same device (GPU) that the
+		// resource is on.  So we always use the default command context per GPU, and ignore the command context passed in.  In practice, the
+		// caller will already be passing the default command context in, but if we run into a situation where that's not the case, it would
+		// require some sort of higher level refactor of the code (for example, moving the linked object iterator loop to a higher level, or
+		// introducing a cross GPU fence sync at the end of an initialization batch).  This assert is to identify if we've encountered such a
+		// case, so we know we need to solve it.
+		//
+		// We only run the assert for resources that are on the first GPU, as certain callers (like GPU Lightmass) create single GPU resources,
+		// and don't attempt to pass in a specific GPU context.  The goal of the assert is to catch unexpected use cases where something other
+		// than the default command context is passed in, and it's good enough to catch that just on the first GPU, assuming any multi-GPU
+		// client will be using resources on all GPUs at some point.
+		if (Buffer->GetParentDevice()->GetGPUIndex() == 0)
+		{
+			check(&CommandContext == &Buffer->GetParentDevice()->GetDefaultCommandContext());
+		}
+#endif
+
 		for (FD3D12Buffer::FLinkedObjectIterator CurrentBuffer(Buffer); CurrentBuffer; ++CurrentBuffer)
 		{
 			FD3D12Resource* Destination = CurrentBuffer->ResourceLocation.GetResource();
 			FD3D12Device* Device = Destination->GetParentDevice();
+#if WITH_MGPU
+			FD3D12CommandContext& CurrentCommandContext = Device->GetDefaultCommandContext();
+#else
+			FD3D12CommandContext& CurrentCommandContext = CommandContext;
+#endif
 
-			FD3D12CommandListHandle& hCommandList = CommandContext.CommandListHandle;
+			FD3D12CommandListHandle& hCommandList = CurrentCommandContext.CommandListHandle;
 			// Copy from the temporary upload heap to the default resource
 			{
 				// if resource doesn't require state tracking then transition to copy dest here (could have been suballocated from shared resource) - not very optimal and should be batched
@@ -118,7 +141,7 @@ struct FD3D12RHICommandInitializeBuffer final : public FRHICommand<FD3D12RHIComm
 					hCommandList.AddTransitionBarrier(Destination, Destination->GetDefaultResourceState(), D3D12_RESOURCE_STATE_COPY_DEST, D3D12_RESOURCE_BARRIER_ALL_SUBRESOURCES);
 				}
 
-				CommandContext.numInitialResourceCopies++;
+				CurrentCommandContext.numInitialResourceCopies++;
 				hCommandList.FlushResourceBarriers();
 				hCommandList->CopyBufferRegion(
 					Destination->GetResource(),
@@ -149,7 +172,7 @@ struct FD3D12RHICommandInitializeBuffer final : public FRHICommand<FD3D12RHIComm
 
 				hCommandList.UpdateResidency(SrcResourceLoc.GetResource());
 
-				CommandContext.ConditionalFlushCommandList();
+				CurrentCommandContext.ConditionalFlushCommandList();
 			}
 
 			// Buffer is now written and ready, so unlock the block (locked after creation and can be defragmented if needed)
