Workaround for GPU crash on NVIDIA Ampere cards when running on Windows 11

The crash is caused by a driver bug with no known expected fix timeline. The only known workaround is to turn off transient memory aliasing.

- Add GD3D12WorkaroundFlags::bAllowTransientResourceAllocator to automatically turn off memory aliasing on known bad OS/driver/hardware configurations
- This workaround can be bypassed using '-rdgtransientallocator=1' command line (an existing argument already used in RenderGraphPrivate.cpp)

NOTE: This change requires CL 18843465 (adds GD3D12WorkaroundFlags) and 18884144 (adds support for detecting Windows 11)

diff --git a/Engine/Source/Runtime/D3D12RHI/Private/D3D12CommandContext.cpp b/Engine/Source/Runtime/D3D12RHI/Private/D3D12CommandContext.cpp
--- a/Engine/Source/Runtime/D3D12RHI/Private/D3D12CommandContext.cpp
+++ b/Engine/Source/Runtime/D3D12RHI/Private/D3D12CommandContext.cpp
@@ -1027,7 +1027,15 @@ void FD3D12DynamicRHI::RHIReleaseTransition(FRHITransition* Transition)
 
 IRHITransientResourceAllocator* FD3D12DynamicRHI::RHICreateTransientResourceAllocator()
 {
-	return new FD3D12TransientResourceHeapAllocator(GetAdapter().GetOrCreateTransientHeapCache());
+	// Some machine configurations have known issues when transient resource aliasing is used.
+	if (GD3D12WorkaroundFlags.bAllowTransientResourceAllocator)
+	{
+		return new FD3D12TransientResourceHeapAllocator(GetAdapter().GetOrCreateTransientHeapCache());
+	}
+	else
+	{
+		return nullptr;
+	}
 }
 
 //////////////////////////////////////////////////////////////////////////////////////////////////////////
diff --git a/Engine/Source/Runtime/D3D12RHI/Private/D3D12RHIPrivate.h b/Engine/Source/Runtime/D3D12RHI/Private/D3D12RHIPrivate.h
--- a/Engine/Source/Runtime/D3D12RHI/Private/D3D12RHIPrivate.h
+++ b/Engine/Source/Runtime/D3D12RHI/Private/D3D12RHIPrivate.h
@@ -232,6 +232,12 @@ struct FD3D12WorkaroundFlags
 	* This is less efficient and can also trigger another known issue with D3D12 Agility version <= 4.
 	*/
 	bool bAllowGetShaderIdentifierOnCollectionSubObject = true;
+
+	/**
+	* Some machine configurations have known issues when transient resource aliasing is used.
+	* If we detect such configuration, we can fall back to non-aliasing code path which is much less efficient.
+	*/
+	bool bAllowTransientResourceAllocator = true;
 };
 
 extern FD3D12WorkaroundFlags GD3D12WorkaroundFlags;
diff --git a/Engine/Source/Runtime/D3D12RHI/Private/Windows/WindowsD3D12Device.cpp b/Engine/Source/Runtime/D3D12RHI/Private/Windows/WindowsD3D12Device.cpp
--- a/Engine/Source/Runtime/D3D12RHI/Private/Windows/WindowsD3D12Device.cpp
+++ b/Engine/Source/Runtime/D3D12RHI/Private/Windows/WindowsD3D12Device.cpp
@@ -840,6 +840,91 @@ static bool IsRayTracingEmulated(uint32 DeviceId)
 	return false;
 }
 
+static bool IsNvidiaAmpereGPU(uint32 DeviceId)
+{
+	uint32 DeviceList[] =
+	{
+		0x2200,	// GA102
+		0x2204,	// GA102 - GeForce RTX 3090
+		0x2205,	// GA102 - GeForce RTX 3080 Ti 20GB
+		0x2206,	// GA102 - GeForce RTX 3080
+		0x2208,	// GA102 - GeForce RTX 3080 Ti
+		0x220a,	// GA102 - GeForce RTX 3080 12GB
+		0x220d,	// GA102 - CMP 90HX
+		0x2216,	// GA102 - GeForce RTX 3080 Lite Hash Rate
+		0x2230,	// GA102GL - RTX A6000
+		0x2231,	// GA102GL - RTX A5000
+		0x2232,	// GA102GL - RTX A4500
+		0x2233,	// GA102GL - RTX A5500
+		0x2235,	// GA102GL - A40
+		0x2236,	// GA102GL - A10
+		0x2237,	// GA102GL - A10G
+		0x2238,	// GA102GL - A10M
+		0x223f,	// GA102G
+		0x2302,	// GA103
+		0x2321,	// GA103
+		0x2414,	// GA103 - GeForce RTX 3060 Ti
+		0x2420,	// GA103M - GeForce RTX 3080 Ti Mobile
+		0x2460,	// GA103M - GeForce RTX 3080 Ti Laptop GPU
+		0x2482,	// GA104 - GeForce RTX 3070 Ti
+		0x2483,	// GA104
+		0x2484,	// GA104 - GeForce RTX 3070
+		0x2486,	// GA104 - GeForce RTX 3060 Ti
+		0x2487,	// GA104 - GeForce RTX 3060
+		0x2488,	// GA104 - GeForce RTX 3070 Lite Hash Rate
+		0x2489,	// GA104 - GeForce RTX 3060 Ti Lite Hash Rate
+		0x248a,	// GA104 - CMP 70HX
+		0x249c,	// GA104M - GeForce RTX 3080 Mobile / Max-Q 8GB/16GB
+		0x249f,	// GA104M
+		0x24a0,	// GA104 - Geforce RTX 3070 Ti Laptop GPU
+		0x24b0,	// GA104GL - RTX A4000
+		0x24b6,	// GA104GLM - RTX A5000 Mobile
+		0x24b7,	// GA104GLM - RTX A4000 Mobile
+		0x24b8,	// GA104GLM - RTX A3000 Mobile
+		0x24dc,	// GA104M - GeForce RTX 3080 Mobile / Max-Q 8GB/16GB
+		0x24dd,	// GA104M - GeForce RTX 3070 Mobile / Max-Q
+		0x24e0,	// GA104M - Geforce RTX 3070 Ti Laptop GPU
+		0x24fa,	// GA104 - RTX A4500 Embedded GPU 
+		0x2501,	// GA106 - GeForce RTX 3060
+		0x2503,	// GA106 - GeForce RTX 3060
+		0x2504,	// GA106 - GeForce RTX 3060 Lite Hash Rate
+		0x2505,	// GA106
+		0x2507,	// GA106 - Geforce RTX 3050
+		0x2520,	// GA106M - GeForce RTX 3060 Mobile / Max-Q
+		0x2523,	// GA106M - GeForce RTX 3050 Ti Mobile / Max-Q
+		0x2531,	// GA106 - RTX A2000
+		0x2560,	// GA106M - GeForce RTX 3060 Mobile / Max-Q
+		0x2563,	// GA106M - GeForce RTX 3050 Ti Mobile / Max-Q
+		0x2571,	// GA106 - RTX A2000 12GB
+		0x2583,	// GA107 - GeForce RTX 3050
+		0x25a0,	// GA107M - GeForce RTX 3050 Ti Mobile
+		0x25a2,	// GA107M - GeForce RTX 3050 Mobile
+		0x25a4,	// GA107
+		0x25a5,	// GA107M - GeForce RTX 3050 Mobile
+		0x25a6,	// GA107M - GeForce MX570
+		0x25a7,	// GA107M - GeForce MX570
+		0x25a9,	// GA107M - GeForce RTX 2050
+		0x25b5,	// GA107GLM - RTX A4 Mobile
+		0x25b8,	// GA107GLM - RTX A2000 Mobile
+		0x25b9,	// GA107GLM - RTX A1000 Laptop GPU
+		0x25e0,	// GA107BM - GeForce RTX 3050 Ti Mobile
+		0x25e2,	// GA107BM - GeForce RTX 3050 Mobile
+		0x25e5,	// GA107BM - GeForce RTX 3050 Mobile
+		0x25f9,	// GA107 - RTX A1000 Embedded GPU 
+		0x25fa,	// GA107 - RTX A2000 Embedded GPU
+	};
+
+	for (uint32 KnownDeviceId : DeviceList)
+	{
+		if (DeviceId == KnownDeviceId)
+		{
+			return true;
+		}
+	}
+
+	return false;
+}
+
 static void DisableRayTracingSupport()
 {
 	GRHISupportsRayTracing = false;
@@ -1004,6 +1089,29 @@ void FD3D12DynamicRHI::Init()
 				TEXT("Please update to NVIDIA driver version 466.11 or newer."));
 		}
 
+		// Enable a workaround for a known driver bug affecting Ampere GPUs on Windows 11 (Jira UE-132964)
+		// As of 2022-02-09, there is no expected fix date or driver version.
+		if (FPlatformMisc::VerifyWindowsVersion(10, 0, 22000) // Win11 is build number is 22000+
+			&& IsNvidiaAmpereGPU(AdapterDesc.DeviceId))
+		{
+			int32 TransientAllocatorValue = 0;
+			const bool bForceTransientAllocatorState = FParse::Value(FCommandLine::Get(), TEXT("rdgtransientallocator="), TransientAllocatorValue);
+			if (bForceTransientAllocatorState && TransientAllocatorValue == 1)
+			{
+				UE_LOG(LogD3D12RHI, Warning,
+					TEXT("Current machine configuration is known to be affected by a driver bug triggered by memory aliasing. ")
+					TEXT("Transient resource allocator will still be used because it is explicitly requested using '-rdgtransientallocator=1' command line."));
+			}
+			else
+			{
+				GD3D12WorkaroundFlags.bAllowTransientResourceAllocator = false;
+				UE_LOG(LogD3D12RHI, Warning, 
+					TEXT("GD3D12WorkaroundFlags.bAllowTransientResourceAllocator is disabled due to a known issue with current GPU, driver and OS version. ")
+					TEXT("This workaround can be bypassed using '-rdgtransientallocator=1' command line.")
+				);
+			}
+		}
+
 		if (GRHISupportsRayTracing
 			&& IsRayTracingEmulated(AdapterDesc.DeviceId))
 		{
