Skip max descriptor detection on Tier3 devices when RenderDoc is present.

RenderDoc doesn't like it when the call used to detect the heap size returns E_OUTOFMEMORY. Also changed the code which disables raytracing to be based on runtime detection of RD, rather than checking if the plug-in is loaded.

#rnx

diff --git a/Engine/Source/Runtime/D3D12RHI/Private/D3D12Adapter.cpp b/Engine/Source/Runtime/D3D12RHI/Private/D3D12Adapter.cpp
--- a/Engine/Source/Runtime/D3D12RHI/Private/D3D12Adapter.cpp
+++ b/Engine/Source/Runtime/D3D12RHI/Private/D3D12Adapter.cpp
@@ -918,6 +918,8 @@ void FD3D12Adapter::InitializeDevices()
 			}
 #endif
 
+			const bool bRenderDocPresent = D3D12RHI_IsRenderDocPresent(RootDevice);
+
 			D3D12_FEATURE_DATA_D3D12_OPTIONS D3D12Caps;
 			FMemory::Memzero(&D3D12Caps, sizeof(D3D12Caps));
 			VERIFYD3D12RESULT(RootDevice->CheckFeatureSupport(D3D12_FEATURE_D3D12_OPTIONS, &D3D12Caps, sizeof(D3D12Caps)));
@@ -940,48 +942,56 @@ void FD3D12Adapter::InitializeDevices()
 				//   it can support the request or fail the call. There is no cap exposed indicating how large of a descriptor 
 				//   heap the hardware could support – applications can just try what they want and fall back to 1000000 if 
 				//   larger doesn’t work.
-
-#if D3D12_SUPPORTS_INFO_QUEUE
-				// Temporarily silence CREATE_DESCRIPTOR_HEAP_LARGE_NUM_DESCRIPTORS since we know we might break on it
-				TRefCountPtr<ID3D12InfoQueue> InfoQueue;
-				RootDevice->QueryInterface(IID_PPV_ARGS(InfoQueue.GetInitReference()));
-				if (InfoQueue)
+				// RenderDoc seems to give up on subsequent API calls if one of them return E_OUTOFMEMORY, so we don't use this
+				// detection method if the RD plug-in is loaded.
+				if (!bRenderDocPresent)
 				{
-					D3D12_MESSAGE_ID MessageId = D3D12_MESSAGE_ID_CREATE_DESCRIPTOR_HEAP_LARGE_NUM_DESCRIPTORS;
+#if D3D12_SUPPORTS_INFO_QUEUE
+					// Temporarily silence CREATE_DESCRIPTOR_HEAP_LARGE_NUM_DESCRIPTORS since we know we might break on it
+					TRefCountPtr<ID3D12InfoQueue> InfoQueue;
+					RootDevice->QueryInterface(IID_PPV_ARGS(InfoQueue.GetInitReference()));
+					if (InfoQueue)
+					{
+						D3D12_MESSAGE_ID MessageId = D3D12_MESSAGE_ID_CREATE_DESCRIPTOR_HEAP_LARGE_NUM_DESCRIPTORS;
 
-					D3D12_INFO_QUEUE_FILTER NewFilter{};
-					NewFilter.DenyList.NumIDs = 1;
-					NewFilter.DenyList.pIDList = &MessageId;
+						D3D12_INFO_QUEUE_FILTER NewFilter{};
+						NewFilter.DenyList.NumIDs = 1;
+						NewFilter.DenyList.pIDList = &MessageId;
 
-					InfoQueue->PushStorageFilter(&NewFilter);
-				}
+						InfoQueue->PushStorageFilter(&NewFilter);
+					}
 #endif
 
-				// create an overly large heap and test for failure
-				D3D12_DESCRIPTOR_HEAP_DESC TempHeapDesc{};
-				TempHeapDesc.Type = D3D12_DESCRIPTOR_HEAP_TYPE_CBV_SRV_UAV;
-				TempHeapDesc.Flags = D3D12_DESCRIPTOR_HEAP_FLAG_SHADER_VISIBLE;
-				TempHeapDesc.NodeMask = FRHIGPUMask::All().GetNative();
-				TempHeapDesc.NumDescriptors = 2 * D3D12_MAX_SHADER_VISIBLE_DESCRIPTOR_HEAP_SIZE_TIER_2;
+					// create an overly large heap and test for failure
+					D3D12_DESCRIPTOR_HEAP_DESC TempHeapDesc{};
+					TempHeapDesc.Type = D3D12_DESCRIPTOR_HEAP_TYPE_CBV_SRV_UAV;
+					TempHeapDesc.Flags = D3D12_DESCRIPTOR_HEAP_FLAG_SHADER_VISIBLE;
+					TempHeapDesc.NodeMask = FRHIGPUMask::All().GetNative();
+					TempHeapDesc.NumDescriptors = 2 * D3D12_MAX_SHADER_VISIBLE_DESCRIPTOR_HEAP_SIZE_TIER_2;
 
-				TRefCountPtr<ID3D12DescriptorHeap> TempHeap;
-				HRESULT hr = RootDevice->CreateDescriptorHeap(&TempHeapDesc, IID_PPV_ARGS(TempHeap.GetInitReference()));
-				if (SUCCEEDED(hr))
-				{
-					MaxNonSamplerDescriptors = -1;
+					TRefCountPtr<ID3D12DescriptorHeap> TempHeap;
+					HRESULT hr = RootDevice->CreateDescriptorHeap(&TempHeapDesc, IID_PPV_ARGS(TempHeap.GetInitReference()));
+					if (SUCCEEDED(hr))
+					{
+						MaxNonSamplerDescriptors = -1;
+					}
+					else
+					{
+						MaxNonSamplerDescriptors = D3D12_MAX_SHADER_VISIBLE_DESCRIPTOR_HEAP_SIZE_TIER_2;
+					}
+
+#if D3D12_SUPPORTS_INFO_QUEUE
+					// Restore the info queue to its old state to ensure we get CREATE_DESCRIPTOR_HEAP_LARGE_NUM_DESCRIPTORS.
+					if (InfoQueue)
+					{
+						InfoQueue->PopStorageFilter();
+					}
+#endif
 				}
 				else
 				{
 					MaxNonSamplerDescriptors = D3D12_MAX_SHADER_VISIBLE_DESCRIPTOR_HEAP_SIZE_TIER_2;
 				}
-
-#if D3D12_SUPPORTS_INFO_QUEUE
-				// Restore the info queue to its old state to ensure we get CREATE_DESCRIPTOR_HEAP_LARGE_NUM_DESCRIPTORS.
-				if (InfoQueue)
-				{
-					InfoQueue->PopStorageFilter();
-				}
-#endif
 			}
 			else
 			{
@@ -1017,7 +1027,7 @@ void FD3D12Adapter::InitializeDevices()
 					}
 				}
 				else if (D3D12Caps5.RaytracingTier == D3D12_RAYTRACING_TIER_NOT_SUPPORTED 
-					&& FModuleManager::Get().IsModuleLoaded("RenderDocPlugin")
+					&& bRenderDocPresent
 					&& !FParse::Param(FCommandLine::Get(), TEXT("noraytracing")))
 				{
 					UE_LOG(LogD3D12RHI, Warning, TEXT("Ray Tracing is disabled because the RenderDoc plugin is currently not compatible with D3D12 ray tracing."));
diff --git a/Engine/Source/Runtime/D3D12RHI/Private/D3D12Device.cpp b/Engine/Source/Runtime/D3D12RHI/Private/D3D12Device.cpp
--- a/Engine/Source/Runtime/D3D12RHI/Private/D3D12Device.cpp
+++ b/Engine/Source/Runtime/D3D12RHI/Private/D3D12Device.cpp
@@ -256,17 +256,10 @@ void FD3D12Device::SetupAfterDeviceCreation()
 	bool bUnderGPUCapture = false;
 
 	// RenderDoc
+	if (D3D12RHI_IsRenderDocPresent(Direct3DDevice))
 	{
-		IID RenderDocID;
-		if (SUCCEEDED(IIDFromString(L"{A7AA6116-9C8D-4BBA-9083-B4D816B71B78}", &RenderDocID)))
-		{
-			TRefCountPtr<IUnknown> RenderDoc;
-			if (SUCCEEDED(Direct3DDevice->QueryInterface(RenderDocID, (void**)RenderDoc.GetInitReference())))
-			{
-				// Running under RenderDoc, so enable capturing mode
-				bUnderGPUCapture = true;
-			}
-		}
+		// Running under RenderDoc, so enable capturing mode
+		bUnderGPUCapture = true;
 	}
 
 	// Intel GPA
diff --git a/Engine/Source/Runtime/D3D12RHI/Private/D3D12RHIPrivate.h b/Engine/Source/Runtime/D3D12RHI/Private/D3D12RHIPrivate.h
--- a/Engine/Source/Runtime/D3D12RHI/Private/D3D12RHIPrivate.h
+++ b/Engine/Source/Runtime/D3D12RHI/Private/D3D12RHIPrivate.h
@@ -200,6 +200,21 @@ static bool D3D12RHI_ShouldForceCompatibility()
 	return bForceCompatibility;
 }
 
+static bool D3D12RHI_IsRenderDocPresent(ID3D12Device* Device)
+{
+	IID RenderDocID;
+	if (SUCCEEDED(IIDFromString(L"{A7AA6116-9C8D-4BBA-9083-B4D816B71B78}", &RenderDocID)))
+	{
+		TRefCountPtr<IUnknown> RenderDoc;
+		if (SUCCEEDED(Device->QueryInterface(RenderDocID, (void**)RenderDoc.GetInitReference())))
+		{
+			return true;
+		}
+	}
+
+	return false;
+}
+
 struct FD3D12UpdateTexture3DData
 {
 	FD3D12ResourceLocation* UploadHeapResourceLocation;
