MGPU:  Ray tracing crash after a few frames, caused by an attempt to unregister ray tracing geometry buffer listeners for additional GPUs after the buffer references have been modified when processing for the first GPU.  A separate pass is run to unregister buffer listeners for all GPUs, before rebuilding the ray tracing geometry and acceleration structures.  It's OK to unregister redundantly, as a flag is set to track whether a buffer is registered, and additional unregister calls do nothing.

#rnx

diff --git a/Engine/Source/Runtime/D3D12RHI/Private/D3D12CommandContext.h b/Engine/Source/Runtime/D3D12RHI/Private/D3D12CommandContext.h
--- a/Engine/Source/Runtime/D3D12RHI/Private/D3D12CommandContext.h
+++ b/Engine/Source/Runtime/D3D12RHI/Private/D3D12CommandContext.h
@@ -387,6 +387,10 @@ public:
 #if D3D12_RHI_RAYTRACING
 	virtual void RHIBindAccelerationStructureMemory(FRHIRayTracingScene* Scene, FRHIBuffer* Buffer, uint32 BufferOffset) final override;
 	virtual void BuildAccelerationStructuresInternal(const TArrayView<D3D12_BUILD_RAYTRACING_ACCELERATION_STRUCTURE_DESC> BuildDesc);
+#if WITH_MGPU
+	// Should be called before RHIBuildAccelerationStructures when multiple GPU support is present (for example, from FD3D12CommandContextRedirector::RHIBuildAccelerationStructures)
+	static void UnregisterAccelerationStructuresInternalMGPU(const TArrayView<const FRayTracingGeometryBuildParams> Params, FRHIGPUMask GPUMask);
+#endif
 	virtual void RHIBuildAccelerationStructures(const TArrayView<const FRayTracingGeometryBuildParams> Params, const FRHIBufferRange& ScratchBufferRange) final override;
 	virtual void RHIBuildAccelerationStructure(const FRayTracingSceneBuildParams& SceneBuildParams) final override;
 	virtual void RHIClearRayTracingBindings(FRHIRayTracingScene* Scene) final override;
@@ -748,6 +752,10 @@ public:
 
 	virtual void RHIBuildAccelerationStructures(const TArrayView<const FRayTracingGeometryBuildParams> Params, const FRHIBufferRange& ScratchBufferRange) final override
 	{
+#if WITH_MGPU
+		FD3D12CommandContext::UnregisterAccelerationStructuresInternalMGPU(Params, GPUMask);
+#endif 
+
 		ContextRedirect(RHIBuildAccelerationStructures(Params, ScratchBufferRange));
 	}
 
diff --git a/Engine/Source/Runtime/D3D12RHI/Private/D3D12RayTracing.cpp b/Engine/Source/Runtime/D3D12RHI/Private/D3D12RayTracing.cpp
--- a/Engine/Source/Runtime/D3D12RHI/Private/D3D12RayTracing.cpp
+++ b/Engine/Source/Runtime/D3D12RHI/Private/D3D12RayTracing.cpp
@@ -4252,6 +4252,24 @@ void FD3D12CommandContext::BuildAccelerationStructuresInternal(const TArrayView<
 	}
 }
 
+#if WITH_MGPU
+void FD3D12CommandContext::UnregisterAccelerationStructuresInternalMGPU(const TArrayView<const FRayTracingGeometryBuildParams> Params, FRHIGPUMask GPUMask)
+{
+	// We need to unregister rename listeners for all GPUs in a separate pass before running "RHIBuildAccelerationStructures", as the build process
+	// may modify the buffer references in the ray tracing geometry.  This leads to an assert where the code attempts to unregister the newer buffer
+	// references on the additional GPUs, rather than the original buffer references.  It's OK to unregister redundantly, as a flag is set to track
+	// whether a buffer is registered, and additional unregister calls do nothing.
+	for (uint32 GPUIndex : GPUMask)
+	{
+		for (const FRayTracingGeometryBuildParams& P : Params)
+		{
+			FD3D12RayTracingGeometry* Geometry = FD3D12DynamicRHI::ResourceCast(P.Geometry.GetReference());
+			Geometry->UnregisterAsRenameListener(GPUIndex);
+		}
+	}
+}
+#endif  // WITH_MGPU
+
 void FD3D12CommandContext::RHIBuildAccelerationStructures(const TArrayView<const FRayTracingGeometryBuildParams> Params, const FRHIBufferRange& ScratchBufferRange)
 {
 	TRACE_CPUPROFILER_EVENT_SCOPE(BuildAccelerationStructure_BottomLevel);
