Log out the Max RHIFeatureLevel on D3D12 RHI creation.

diff --git a/Engine/Source/Runtime/D3D12RHI/Private/Windows/WindowsD3D12Device.cpp b/Engine/Source/Runtime/D3D12RHI/Private/Windows/WindowsD3D12Device.cpp
--- a/Engine/Source/Runtime/D3D12RHI/Private/Windows/WindowsD3D12Device.cpp
+++ b/Engine/Source/Runtime/D3D12RHI/Private/Windows/WindowsD3D12Device.cpp
@@ -741,6 +741,9 @@ FDynamicRHI* FD3D12DynamicRHIModule::CreateRHI(ERHIFeatureLevel::Type RequestedF
 		GRHIDeviceIsIntegrated = ChosenAdapters[0].Get()->GetDesc().bIsIntegrated;
 	}
 
+	const FString FeatureLevelString = LexToString(GMaxRHIFeatureLevel);
+	UE_LOG(LogD3D12RHI, Display, TEXT("Creating D3D12 RHI with Max Feature Level %s"), *FeatureLevelString);
+
 	GD3D12RHI = new FD3D12DynamicRHI(ChosenAdapters, bPixEventEnabled);
 #if ENABLE_RHI_VALIDATION
 	if (FParse::Param(FCommandLine::Get(), TEXT("RHIValidation")))
