Adding proper lifetime management to resources used by FRHICommandD3D12UpdateUniformBuffer

diff --git a/Engine/Source/Runtime/D3D12RHI/Private/D3D12UniformBuffer.cpp b/Engine/Source/Runtime/D3D12RHI/Private/D3D12UniformBuffer.cpp
--- a/Engine/Source/Runtime/D3D12RHI/Private/D3D12UniformBuffer.cpp
+++ b/Engine/Source/Runtime/D3D12RHI/Private/D3D12UniformBuffer.cpp
@@ -89,28 +89,36 @@ FUniformBufferRHIRef FD3D12DynamicRHI::RHICreateUniformBuffer(const void* Conten
 	return UniformBufferOut;
 }
 
-struct FRHICommandD3D12UpdateUniformBufferString
+FRHICOMMAND_MACRO(FRHICommandD3D12UpdateUniformBuffer)
 {
-	static const TCHAR* TStr() { return TEXT("FRHICommandD3D12UpdateUniformBuffer"); }
-};
-struct FRHICommandD3D12UpdateUniformBuffer final : public FRHICommand<FRHICommandD3D12UpdateUniformBuffer, FRHICommandD3D12UpdateUniformBufferString>
-{
-	FD3D12UniformBuffer* UniformBuffer;
+	TRefCountPtr<FD3D12UniformBuffer> UniformBuffer;
 	FD3D12ResourceLocation UpdatedLocation;
-	FRHIResource** UpdatedResources;
-	int32 NumResources;
+	TArrayView<FRHIResource*> UpdatedResources;
+
 	FORCEINLINE_DEBUGGABLE FRHICommandD3D12UpdateUniformBuffer(FD3D12UniformBuffer* InUniformBuffer, FD3D12ResourceLocation& InUpdatedLocation, FRHIResource** InUpdatedResources, int32 InNumResources)
 		: UniformBuffer(InUniformBuffer)
 		, UpdatedLocation(InUpdatedLocation.GetParentDevice())
-		, UpdatedResources(InUpdatedResources)
-		, NumResources(InNumResources)
+		, UpdatedResources(InUpdatedResources, InNumResources)
 	{
 		FD3D12ResourceLocation::TransferOwnership(UpdatedLocation, InUpdatedLocation);
+
+		for (FRHIResource* Resource : UpdatedResources)
+		{
+			Resource->AddRef();
+		}
+	}
+
+	~FRHICommandD3D12UpdateUniformBuffer()
+	{
+		for (FRHIResource* Resource : UpdatedResources)
+		{
+			Resource->Release();
+		}
 	}
 
 	void Execute(FRHICommandListBase& CmdList)
 	{
-		for (int32 i = 0; i < NumResources; ++i)
+		for (int32 i = 0; i < UpdatedResources.Num(); ++i)
 		{
 			//check(UniformBuffer->ResourceTable[i]);
 			UniformBuffer->ResourceTable[i] = UpdatedResources[i];
