Make sure buffer SRV is removed from state cache when performing upload buffer copy to make sure it's gets transitioned to SRV again on next usage - could skip transition if same resources was used again causing GPU crash

diff --git a/Engine/Source/Runtime/D3D12RHI/Private/D3D12Buffer.cpp b/Engine/Source/Runtime/D3D12RHI/Private/D3D12Buffer.cpp
--- a/Engine/Source/Runtime/D3D12RHI/Private/D3D12Buffer.cpp
+++ b/Engine/Source/Runtime/D3D12RHI/Private/D3D12Buffer.cpp
@@ -45,7 +45,7 @@ struct FRHICommandUpdateBuffer final : public FRHICommand<FRHICommandUpdateBuffe
 
 	void Execute(FRHICommandListBase& CmdList)
 	{
-		FD3D12DynamicRHI::GetD3DRHI()->UpdateBuffer(Destination->GetResource(), Destination->GetOffsetFromBaseOfResource() + DestinationOffset, Source.GetResource(), Source.GetOffsetFromBaseOfResource(), NumBytes);
+		FD3D12DynamicRHI::GetD3DRHI()->UpdateBuffer(Destination, DestinationOffset, &Source, 0, NumBytes);
 	}
 };
 
@@ -781,11 +781,7 @@ void FD3D12DynamicRHI::UnlockBuffer(FRHICommandListImmediate* RHICmdList, FD3D12
 				}
 				else
 				{
-					UpdateBuffer(CurrentBuffer->ResourceLocation.GetResource(),
-						CurrentBuffer->ResourceLocation.GetOffsetFromBaseOfResource() + LockedData.LockedOffset,
-						LockedData.ResourceLocation.GetResource(),
-						LockedData.ResourceLocation.GetOffsetFromBaseOfResource(),
-						LockedData.LockedPitch);
+					UpdateBuffer(&CurrentBuffer->ResourceLocation, LockedData.LockedOffset, &LockedData.ResourceLocation, 0, LockedData.LockedPitch);
 				}
 			}
 		}
diff --git a/Engine/Source/Runtime/D3D12RHI/Private/D3D12RHI.cpp b/Engine/Source/Runtime/D3D12RHI/Private/D3D12RHI.cpp
--- a/Engine/Source/Runtime/D3D12RHI/Private/D3D12RHI.cpp
+++ b/Engine/Source/Runtime/D3D12RHI/Private/D3D12RHI.cpp
@@ -361,21 +361,30 @@ IRHIComputeContext* FD3D12DynamicRHI::RHIGetDefaultAsyncComputeContext()
 	return DefaultAsyncComputeContext;
 }
 
-void FD3D12DynamicRHI::UpdateBuffer(FD3D12Resource* Dest, uint32 DestOffset, FD3D12Resource* Source, uint32 SourceOffset, uint32 NumBytes)
+void FD3D12DynamicRHI::UpdateBuffer(FD3D12ResourceLocation* Dest, uint32 DestOffset, FD3D12ResourceLocation* Source, uint32 SourceOffset, uint32 NumBytes)
 {
-	FD3D12Device* Device = Dest->GetParentDevice();
+	FD3D12Resource* SourceResource = Source->GetResource();
+	uint32 SourceFullOffset = Source->GetOffsetFromBaseOfResource() + SourceOffset;
+
+	FD3D12Resource* DestResource = Dest->GetResource();
+	uint32 DestFullOffset = Dest->GetOffsetFromBaseOfResource() + DestOffset;
+
+	FD3D12Device* Device = DestResource->GetParentDevice();
 
 	FD3D12CommandContext& DefaultContext = Device->GetDefaultCommandContext();
 	FD3D12CommandListHandle& hCommandList = DefaultContext.CommandListHandle;
+	
+	// Clear the resource if still bound to make sure the SRVs are rebound again on next operation (and get correct resource transitions enqueued)
+	DefaultContext.ConditionalClearShaderResource(Dest);
 
-	FScopedResourceBarrier ScopeResourceBarrierDest(hCommandList, Dest, D3D12_RESOURCE_STATE_COPY_DEST, 0, FD3D12DynamicRHI::ETransitionMode::Apply);
+	FScopedResourceBarrier ScopeResourceBarrierDest(hCommandList, DestResource, D3D12_RESOURCE_STATE_COPY_DEST, 0, FD3D12DynamicRHI::ETransitionMode::Apply);
 	// Don't need to transition upload heaps
 
 	DefaultContext.numCopies++;
 	hCommandList.FlushResourceBarriers();
-	hCommandList->CopyBufferRegion(Dest->GetResource(), DestOffset, Source->GetResource(), SourceOffset, NumBytes);
-	hCommandList.UpdateResidency(Dest);
-	hCommandList.UpdateResidency(Source);
+	hCommandList->CopyBufferRegion(DestResource->GetResource(), DestFullOffset, SourceResource->GetResource(), SourceFullOffset, NumBytes);
+	hCommandList.UpdateResidency(DestResource);
+	hCommandList.UpdateResidency(SourceResource);
 	
 	DefaultContext.ConditionalFlushCommandList();
 
diff --git a/Engine/Source/Runtime/D3D12RHI/Private/D3D12RHIPrivate.h b/Engine/Source/Runtime/D3D12RHI/Private/D3D12RHIPrivate.h
--- a/Engine/Source/Runtime/D3D12RHI/Private/D3D12RHIPrivate.h
+++ b/Engine/Source/Runtime/D3D12RHI/Private/D3D12RHIPrivate.h
@@ -576,7 +576,7 @@ public:
 	FUpdateTexture3DData BeginUpdateTexture3D_Internal(FRHITexture3D* Texture, uint32 MipIndex, const struct FUpdateTextureRegion3D& UpdateRegion);
 	void EndUpdateTexture3D_Internal(FUpdateTexture3DData& UpdateData);
 
-	void UpdateBuffer(FD3D12Resource* Dest, uint32 DestOffset, FD3D12Resource* Source, uint32 SourceOffset, uint32 NumBytes);
+	void UpdateBuffer(FD3D12ResourceLocation* Dest, uint32 DestOffset, FD3D12ResourceLocation* Source, uint32 SourceOffset, uint32 NumBytes);
 
 #if UE_BUILD_DEBUG	
 	uint32 SubmissionLockStalls;
