Make sure instance buffer is in SRV state before building scene acceleration structure - fix validation error

diff --git a/Engine/Source/Runtime/D3D12RHI/Private/D3D12RayTracing.cpp b/Engine/Source/Runtime/D3D12RHI/Private/D3D12RayTracing.cpp
--- a/Engine/Source/Runtime/D3D12RHI/Private/D3D12RayTracing.cpp
+++ b/Engine/Source/Runtime/D3D12RHI/Private/D3D12RayTracing.cpp
@@ -4143,6 +4143,10 @@ void FD3D12RayTracingScene::BuildAccelerationStructure(FD3D12CommandContext& Com
 	AccelerationStructureBuffer->GetResource()->UpdateResidency(CommandContext.CommandListHandle);
 	ScratchBuffer->GetResource()->UpdateResidency(CommandContext.CommandListHandle);
 
+	// Enqueue transition to UAV/SRV
+	FD3D12DynamicRHI::TransitionResource(CommandContext.CommandListHandle, InstanceBuffer->GetResource(),
+		D3D12_RESOURCE_STATE_TBD, D3D12_RESOURCE_STATE_NON_PIXEL_SHADER_RESOURCE, 0, FD3D12DynamicRHI::ETransitionMode::Apply);
+
 	D3D12_BUILD_RAYTRACING_ACCELERATION_STRUCTURE_DESC BuildDesc = {};
 	BuildDesc.Inputs = BuildInputs;
 	BuildDesc.Inputs.InstanceDescs = InstanceBuffer->ResourceLocation.GetGPUVirtualAddress() + InstanceBufferOffset;
