Check for nullptr before accessing InstanceBuffer

diff --git a/Engine/Source/Runtime/D3D12RHI/Private/D3D12RayTracing.cpp b/Engine/Source/Runtime/D3D12RHI/Private/D3D12RayTracing.cpp
--- a/Engine/Source/Runtime/D3D12RHI/Private/D3D12RayTracing.cpp
+++ b/Engine/Source/Runtime/D3D12RHI/Private/D3D12RayTracing.cpp
@@ -4143,9 +4143,12 @@ void FD3D12RayTracingScene::BuildAccelerationStructure(FD3D12CommandContext& Com
 	AccelerationStructureBuffer->GetResource()->UpdateResidency(CommandContext.CommandListHandle);
 	ScratchBuffer->GetResource()->UpdateResidency(CommandContext.CommandListHandle);
 
-	// Enqueue transition to UAV/SRV
-	FD3D12DynamicRHI::TransitionResource(CommandContext.CommandListHandle, InstanceBuffer->GetResource(),
-		D3D12_RESOURCE_STATE_TBD, D3D12_RESOURCE_STATE_NON_PIXEL_SHADER_RESOURCE, 0, FD3D12DynamicRHI::ETransitionMode::Apply);
+	if (InstanceBuffer)
+	{
+		// Enqueue transition to UAV/SRV
+		FD3D12DynamicRHI::TransitionResource(CommandContext.CommandListHandle, InstanceBuffer->GetResource(),
+			D3D12_RESOURCE_STATE_TBD, D3D12_RESOURCE_STATE_NON_PIXEL_SHADER_RESOURCE, 0, FD3D12DynamicRHI::ETransitionMode::Apply);
+	}
 
 	D3D12_BUILD_RAYTRACING_ACCELERATION_STRUCTURE_DESC BuildDesc = {};
 	BuildDesc.Inputs = BuildInputs;
