Hololens: Missing HL shader platform selection in FD3D12DynamicRHIModule

Forgot to include the Hololens shader platform selection code in the D3D12 RHI bootstrapping. Didn't pop up in testing because D3D11 was the primary, but now D3D12 is so hot right now.
Original Jira was UE-133375 for the core Hololens SP work

diff --git a/Engine/Source/Runtime/D3D12RHI/Private/Windows/WindowsD3D12Device.cpp b/Engine/Source/Runtime/D3D12RHI/Private/Windows/WindowsD3D12Device.cpp
--- a/Engine/Source/Runtime/D3D12RHI/Private/Windows/WindowsD3D12Device.cpp
+++ b/Engine/Source/Runtime/D3D12RHI/Private/Windows/WindowsD3D12Device.cpp
@@ -691,10 +691,19 @@ static bool DoesAnyAdapterSupportSM6(const TArray<TSharedPtr<FD3D12Adapter>>& Ad
 
 FDynamicRHI* FD3D12DynamicRHIModule::CreateRHI(ERHIFeatureLevel::Type RequestedFeatureLevel)
 {
+#if PLATFORM_HOLOLENS
+	GMaxRHIFeatureLevel = ERHIFeatureLevel::ES3_1;
+	GMaxRHIShaderPlatform = SP_D3D_ES3_1_HOLOLENS;
+#endif
+
 	const bool bAllowSM6 = GAllowShaderModel6
 		|| (RequestedFeatureLevel != ERHIFeatureLevel::Num && RequestedFeatureLevel >= ERHIFeatureLevel::SM6);
 
+#if PLATFORM_HOLOLENS
+	GShaderPlatformForFeatureLevel[ERHIFeatureLevel::ES3_1] = SP_D3D_ES3_1_HOLOLENS;
+#else
 	GShaderPlatformForFeatureLevel[ERHIFeatureLevel::ES3_1] = SP_PCD3D_ES3_1;
+#endif
 	GShaderPlatformForFeatureLevel[ERHIFeatureLevel::SM5] = SP_PCD3D_SM5;
 	if (bAllowSM6)
 	{
