Restore UAV barriers on every transition to UAV when using internal transitions because UAV->SRV and then SRV->UAV would skip the UAV barrier and have zero synchronization because UAV->SRV was just skipped.

#rnx

diff --git a/Engine/Source/Runtime/D3D12RHI/Private/D3D12Commands.cpp b/Engine/Source/Runtime/D3D12RHI/Private/D3D12Commands.cpp
--- a/Engine/Source/Runtime/D3D12RHI/Private/D3D12Commands.cpp
+++ b/Engine/Source/Runtime/D3D12RHI/Private/D3D12Commands.cpp
@@ -474,8 +474,9 @@ static void HandleResourceTransitions(FD3D12CommandContext& Context, const FD3D1
 		if (GUseInternalTransitions)
 		{
 			// Only need to check for UAV barriers here, all other transitions are handled inside the RHI
-			const bool bUAVAccessOrUnknownBefore = EnumHasAnyFlags(Info.AccessBefore, ERHIAccess::UAVMask) || Info.AccessBefore == ERHIAccess::Unknown;
-			bUAVBarrier |= (bUAVAccessOrUnknownBefore && bUAVAccessAfter);
+			// Can't rely on before state to either be unknown or UAV because there could be a UAV->SRV transition already done (and ignored here)
+			// and the transition SRV->UAV needs a UAV barrier then to work correctly otherwise there is no synchronization at all
+			bUAVBarrier |= bUAVAccessAfter;
 		}
 		else
 		{
