Fixes for patchup transitions on resources with hidden UAV resource - was applying the wrong transitions at the wrong time

diff --git a/Engine/Source/Runtime/D3D12RHI/Private/D3D12CommandList.cpp b/Engine/Source/Runtime/D3D12RHI/Private/D3D12CommandList.cpp
--- a/Engine/Source/Runtime/D3D12RHI/Private/D3D12CommandList.cpp
+++ b/Engine/Source/Runtime/D3D12RHI/Private/D3D12CommandList.cpp
@@ -60,10 +60,10 @@ void FD3D12CommandListHandle::AddUAVBarrier()
 	}
 }
 
-void FD3D12CommandListHandle::AddAliasingBarrier(FD3D12Resource* InResourceBefore, FD3D12Resource* InResourceAfter)
+void FD3D12CommandListHandle::AddAliasingBarrier(ID3D12Resource* InResourceBefore, ID3D12Resource* InResourceAfter)
 {
 	check(CommandListData);
-	CommandListData->ResourceBarrierBatcher.AddAliasingBarrier(InResourceBefore ? InResourceBefore->GetResource() : nullptr, InResourceAfter->GetResource());
+	CommandListData->ResourceBarrierBatcher.AddAliasingBarrier(InResourceBefore, InResourceAfter);
 	CommandListData->CurrentOwningContext->numBarriers++;
 
 	if (!GD3D12BatchResourceBarriers)
diff --git a/Engine/Source/Runtime/D3D12RHI/Private/D3D12CommandList.h b/Engine/Source/Runtime/D3D12RHI/Private/D3D12CommandList.h
--- a/Engine/Source/Runtime/D3D12RHI/Private/D3D12CommandList.h
+++ b/Engine/Source/Runtime/D3D12RHI/Private/D3D12CommandList.h
@@ -594,7 +594,7 @@ public:
 	// Adds a UAV barrier to the barrier batch
 	void AddUAVBarrier();
 
-	void AddAliasingBarrier(FD3D12Resource* InResourceBefore, FD3D12Resource* InResourceAfter);
+	void AddAliasingBarrier(ID3D12Resource* InResourceBefore, ID3D12Resource* InResourceAfter);
 
 	// Flushes the batched resource barriers to the current command list
 	void FlushResourceBarriers()
diff --git a/Engine/Source/Runtime/D3D12RHI/Private/D3D12Commands.cpp b/Engine/Source/Runtime/D3D12RHI/Private/D3D12Commands.cpp
--- a/Engine/Source/Runtime/D3D12RHI/Private/D3D12Commands.cpp
+++ b/Engine/Source/Runtime/D3D12RHI/Private/D3D12Commands.cpp
@@ -413,7 +413,7 @@ static void HandleTransientAliasing(FD3D12CommandContext& Context, const FD3D12T
 
 			if (GD3D12TransientAllocatorFullAliasingBarrier != 0)
 			{
-				Context.CommandListHandle.AddAliasingBarrier(nullptr, Resource);
+				Context.CommandListHandle.AddAliasingBarrier(nullptr, Resource->GetResource());
 			}
 			else
 			{
@@ -432,7 +432,7 @@ static void HandleTransientAliasing(FD3D12CommandContext& Context, const FD3D12T
 					}
 
 					check(ResourceBefore);
-					Context.CommandListHandle.AddAliasingBarrier(ResourceBefore, Resource);
+					Context.CommandListHandle.AddAliasingBarrier(ResourceBefore->GetResource(), Resource->GetResource());
 				}
 			}
 		}
diff --git a/Engine/Source/Runtime/D3D12RHI/Private/D3D12DirectCommandListManager.cpp b/Engine/Source/Runtime/D3D12RHI/Private/D3D12DirectCommandListManager.cpp
--- a/Engine/Source/Runtime/D3D12RHI/Private/D3D12DirectCommandListManager.cpp
+++ b/Engine/Source/Runtime/D3D12RHI/Private/D3D12DirectCommandListManager.cpp
@@ -970,12 +970,12 @@ uint32 FD3D12CommandListManager::CollectInitialResourceBarriersAndUpdateFinalSta
 			CResourceState& ResourceState = PRB.Resource->GetResourceState();
 
 			const D3D12_RESOURCE_STATES Before = ResourceState.GetSubresourceState(PRB.SubResource);
+			check(Before != D3D12_RESOURCE_STATE_TBD && Before != D3D12_RESOURCE_STATE_CORRUPT);
 
 			// If state unknown then we don't enqueue a transition - only want to update the end state on the resource
 			// and not really enqueue a transition
 			const D3D12_RESOURCE_STATES After = (PRB.State != D3D12_RESOURCE_STATE_TBD) ? PRB.State : Before;
-
-			check(Before != D3D12_RESOURCE_STATE_TBD && Before != D3D12_RESOURCE_STATE_CORRUPT);
+		
 			if (Before != After)
 			{
 				if (IsDirectQueueExclusiveD3D12State(Before) || IsDirectQueueExclusiveD3D12State(After))
@@ -987,15 +987,38 @@ uint32 FD3D12CommandListManager::CollectInitialResourceBarriersAndUpdateFinalSta
 				{
 					AddTransitionBarrier(BarrierDescInfo.BackBufferBarrierDescs, PRB.Resource, Before, After, PRB.SubResource);
 				}
+				// Special case for UAV access resources transitioning from UAV (then they need to transition from the cache hidden state instead)
+				else if (PRB.Resource->GetUAVAccessResource() && EnumHasAnyFlags(Before | After, D3D12_RESOURCE_STATE_UNORDERED_ACCESS))
+				{
+					// After state should never be UAV here
+					check(!EnumHasAnyFlags(After, D3D12_RESOURCE_STATE_UNORDERED_ACCESS));
+					
+					// Add the aliasing barrier
+					BarrierDescInfo.BarrierDescs.AddUninitialized();
+					D3D12_RESOURCE_BARRIER& Barrier = BarrierDescInfo.BarrierDescs.Last();
+					Barrier.Type = D3D12_RESOURCE_BARRIER_TYPE_ALIASING;
+					Barrier.Flags = D3D12_RESOURCE_BARRIER_FLAG_NONE;
+					Barrier.Aliasing.pResourceBefore = PRB.Resource->GetUAVAccessResource();
+					Barrier.Aliasing.pResourceAfter = PRB.Resource->GetResource();
+
+					AddTransitionBarrier(BarrierDescInfo.BarrierDescs, PRB.Resource, ResourceState.GetUAVHiddenResourceState(), After, PRB.SubResource);
+				}
 				else
 				{
-					AddTransitionBarrierWithUAVAccessOverrides(BarrierDescInfo.BarrierDescs, PRB.Resource, Before, After, PRB.SubResource);
+					AddTransitionBarrier(BarrierDescInfo.BarrierDescs, PRB.Resource, Before, After, PRB.SubResource);
 				}
 			}
 
 			// Update the state to the what it will be after hList executes
-			const D3D12_RESOURCE_STATES CommandListState = InCommandListHandle.GetResourceState(PRB.Resource).GetSubresourceState(PRB.SubResource);
-			const D3D12_RESOURCE_STATES LastState = (CommandListState != D3D12_RESOURCE_STATE_TBD) ? CommandListState : After;
+			const CResourceState& CommandListState = InCommandListHandle.GetResourceState(PRB.Resource);
+			const D3D12_RESOURCE_STATES CommandListLastState = CommandListState.GetSubresourceState(PRB.SubResource);
+			const D3D12_RESOURCE_STATES LastState = (CommandListLastState != D3D12_RESOURCE_STATE_TBD) ? CommandListLastState : After;
+
+			// Copy over the hidden UAV state if the last state in the command list was UAV (needed for patch up transition on next command list)
+			if (PRB.Resource->GetUAVAccessResource() && EnumHasAnyFlags(LastState, D3D12_RESOURCE_STATE_UNORDERED_ACCESS))
+			{
+				ResourceState.SetUAVHiddenResourceState(CommandListState.GetUAVHiddenResourceState());
+			}
 
 			if (Before != LastState)
 			{
diff --git a/Engine/Source/Runtime/D3D12RHI/Private/D3D12RHIPrivate.h b/Engine/Source/Runtime/D3D12RHI/Private/D3D12RHIPrivate.h
--- a/Engine/Source/Runtime/D3D12RHI/Private/D3D12RHIPrivate.h
+++ b/Engine/Source/Runtime/D3D12RHI/Private/D3D12RHIPrivate.h
@@ -970,6 +970,9 @@ public:
 		D3D12_RESOURCE_STATES TrackedState = InResourceState.GetSubresourceState(InSubresourceIndex);
 		D3D12_RESOURCE_STATES BeforeState = TrackedState;
 
+		// Special case for UAV access resources
+		bool bHasUAVAccessResource = InResource->GetUAVAccessResource() != nullptr;
+
 		// Still untracked in this command list, then try and find out a before state to use
 		if (BeforeState == D3D12_RESOURCE_STATE_TBD)
 		{
@@ -996,9 +999,20 @@ public:
 				}
 				else
 				{
-					// We need a pending resource barrier so we can setup the state before this command list executes
-					InResourceState.SetSubresourceState(InSubresourceIndex, InAfterState);
-					InCommandList.AddPendingResourceBarrier(InResource, InAfterState, InSubresourceIndex);
+					// Special handling for UAVAccessResource and transition to UAV - don't want to
+					// enqueue pending resource to UAV because the actual resource won't transition
+					// Adding of patch up will only be done when transitioning to non UAV state
+					if (bHasUAVAccessResource && EnumHasAnyFlags(InAfterState, D3D12_RESOURCE_STATE_UNORDERED_ACCESS))
+					{
+						InCommandList.AddAliasingBarrier(InResource->GetResource(), InResource->GetUAVAccessResource());
+						InResourceState.SetSubresourceState(InSubresourceIndex, InAfterState);
+					}
+					else
+					{
+						// We need a pending resource barrier so we can setup the state before this command list executes
+						InResourceState.SetSubresourceState(InSubresourceIndex, InAfterState);
+						InCommandList.AddPendingResourceBarrier(InResource, InAfterState, InSubresourceIndex);
+					}
 				}
 			}
 			else
@@ -1037,25 +1051,65 @@ public:
 			}
 			else
 			{
+				bool bApplyTransitionBarrier = true;
+
 				// Require UAV barrier when before and after are UAV
 				if (BeforeState == D3D12_RESOURCE_STATE_UNORDERED_ACCESS && InAfterState == D3D12_RESOURCE_STATE_UNORDERED_ACCESS)
 				{
 					bRequireUAVBarrier = true;
 				}
-
-				// We're not using IsTransitionNeeded() when bInForceAfterState is set because we do want to transition even if 'after' is a subset of 'before'
-				// This is so that we can ensure all subresources are in the same state, simplifying future barriers
-				// No state merging when using engine transitions - otherwise next before state might not match up anymore)
-				bool bAllowStateMerging = GUseInternalTransitions;
-				if ((bInForceAfterState && BeforeState != InAfterState) || IsTransitionNeeded(bAllowStateMerging, BeforeState, InAfterState))
+				// Special case for UAV access resources
+				else if (bHasUAVAccessResource && EnumHasAnyFlags(BeforeState | InAfterState, D3D12_RESOURCE_STATE_UNORDERED_ACCESS))
 				{
-					InCommandList.AddTransitionBarrier(InResource, BeforeState, InAfterState, InSubresourceIndex);
-					InResourceState.SetSubresourceState(InSubresourceIndex, InAfterState);
+					// inject an aliasing barrier
+					const bool bFromUAV = EnumHasAnyFlags(BeforeState, D3D12_RESOURCE_STATE_UNORDERED_ACCESS);
+					const bool bToUAV = EnumHasAnyFlags(InAfterState, D3D12_RESOURCE_STATE_UNORDERED_ACCESS);
+					check(bFromUAV != bToUAV);
+
+					InCommandList.AddAliasingBarrier(
+						bFromUAV ? InResource->GetUAVAccessResource() : InResource->GetResource(),
+						bToUAV ? InResource->GetUAVAccessResource() : InResource->GetResource());
+
+					if (bToUAV)
+					{
+						InResourceState.SetUAVHiddenResourceState(BeforeState);
+						bApplyTransitionBarrier = false;
+					}
+					else
+					{
+						D3D12_RESOURCE_STATES HiddenState = InResourceState.GetUAVHiddenResourceState();
+
+						// Still unknown in this command list?
+						if (HiddenState == D3D12_RESOURCE_STATE_TBD)
+						{
+							InCommandList.AddPendingResourceBarrier(InResource, InAfterState, InSubresourceIndex);
+							InResourceState.SetSubresourceState(InSubresourceIndex, InAfterState);
+							bApplyTransitionBarrier = false;
+						}
+						else
+						{
+							// Use the hidden state as the before state on the resource
+							BeforeState = HiddenState;
+						}
+					}
 				}
-				// Force update the state when the tracked state is still unknown
-				else if (TrackedState == D3D12_RESOURCE_STATE_TBD)
+
+				if (bApplyTransitionBarrier)
 				{
-					InResourceState.SetSubresourceState(InSubresourceIndex, InAfterState);
+					// We're not using IsTransitionNeeded() when bInForceAfterState is set because we do want to transition even if 'after' is a subset of 'before'
+					// This is so that we can ensure all subresources are in the same state, simplifying future barriers
+					// No state merging when using engine transitions - otherwise next before state might not match up anymore)
+					bool bAllowStateMerging = GUseInternalTransitions;
+					if ((bInForceAfterState && BeforeState != InAfterState) || IsTransitionNeeded(bAllowStateMerging, BeforeState, InAfterState))
+					{
+						InCommandList.AddTransitionBarrier(InResource, BeforeState, InAfterState, InSubresourceIndex);
+						InResourceState.SetSubresourceState(InSubresourceIndex, InAfterState);
+					}
+					// Force update the state when the tracked state is still unknown
+					else if (TrackedState == D3D12_RESOURCE_STATE_TBD)
+					{
+						InResourceState.SetSubresourceState(InSubresourceIndex, InAfterState);
+					}
 				}
 			}
 		}
diff --git a/Engine/Source/Runtime/D3D12RHI/Private/D3D12Util.cpp b/Engine/Source/Runtime/D3D12RHI/Private/D3D12Util.cpp
--- a/Engine/Source/Runtime/D3D12RHI/Private/D3D12Util.cpp
+++ b/Engine/Source/Runtime/D3D12RHI/Private/D3D12Util.cpp
@@ -1371,6 +1371,9 @@ void CResourceState::Initialize(uint32 SubresourceCount)
 
 	// All subresources start out in an unknown state
 	SetResourceState(D3D12_RESOURCE_STATE_TBD);
+
+	// Unknown hidden resource state
+	UAVHiddenResourceState = D3D12_RESOURCE_STATE_TBD;
 }
 
 //----------------------------------------------------------------------------------------------------------------------------------
diff --git a/Engine/Source/Runtime/D3D12RHI/Public/D3D12Resources.h b/Engine/Source/Runtime/D3D12RHI/Public/D3D12Resources.h
--- a/Engine/Source/Runtime/D3D12RHI/Public/D3D12Resources.h
+++ b/Engine/Source/Runtime/D3D12RHI/Public/D3D12Resources.h
@@ -174,7 +174,6 @@ private:
 #ifdef PLATFORM_SUPPORTS_RESOURCE_COMPRESSION
 	D3D12_RESOURCE_STATES CompressedState{ D3D12_RESOURCE_STATE_CORRUPT };
 #endif
-	D3D12_RESOURCE_STATES UAVHiddenResourceState{ D3D12_RESOURCE_STATE_CORRUPT };
 
 	D3D12_HEAP_TYPE HeapType;
 	FName DebugName;
@@ -220,9 +219,6 @@ public:
 	ID3D12Resource* GetUAVAccessResource() const { return UAVAccessResource.GetReference(); }
 	void SetUAVAccessResource(ID3D12Resource* InUAVAccessResource) { UAVAccessResource = InUAVAccessResource; }
 
-	inline D3D12_RESOURCE_STATES GetUAVHiddenResourceState() const { return UAVHiddenResourceState; }
-	void SetUAVHiddenResourceState(D3D12_RESOURCE_STATES InUAVHiddenResourceState) { UAVHiddenResourceState = InUAVHiddenResourceState; }
-
 	inline void* Map(const D3D12_RANGE* ReadRange = nullptr)
 	{
 		if (NumMapCalls == 0)
@@ -973,37 +969,6 @@ inline void AddTransitionBarrier(TArray<D3D12_RESOURCE_BARRIER, InAllocatorType>
 	BarrierList.Add(CD3DX12_RESOURCE_BARRIER::Transition(pResource->GetResource(), Before, After, Subresource));
 }
 
-template<typename InAllocatorType>
-inline void AddTransitionBarrierWithUAVAccessOverrides(TArray<D3D12_RESOURCE_BARRIER, InAllocatorType>& BarrierList, FD3D12Resource* pResource, D3D12_RESOURCE_STATES Before, D3D12_RESOURCE_STATES After, uint32 Subresource)
-{
-	if (pResource->GetUAVAccessResource() && EnumHasAnyFlags(Before | After, D3D12_RESOURCE_STATE_UNORDERED_ACCESS))
-	{
-		// inject an aliasing barrier
-		const bool bFromUAV = EnumHasAnyFlags(Before, D3D12_RESOURCE_STATE_UNORDERED_ACCESS);
-		const bool bToUAV = EnumHasAnyFlags(After, D3D12_RESOURCE_STATE_UNORDERED_ACCESS);
-
-		check(bFromUAV != bToUAV);
-
-		BarrierList.Add(CD3DX12_RESOURCE_BARRIER::Aliasing(
-			bFromUAV ? pResource->GetUAVAccessResource() : pResource->GetResource(),
-			bToUAV ? pResource->GetUAVAccessResource() : pResource->GetResource()
-		));
-
-		if (bToUAV)
-		{
-			pResource->SetUAVHiddenResourceState(Before);
-		}
-		else if (D3D12_RESOURCE_STATES HiddenState = pResource->GetUAVHiddenResourceState(); HiddenState != After && HiddenState != D3D12_RESOURCE_STATE_CORRUPT)
-		{
-			AddTransitionBarrier(BarrierList, pResource, HiddenState, After, Subresource);
-		}
-	}
-	else
-	{
-		AddTransitionBarrier(BarrierList, pResource, Before, After, Subresource);
-	}
-}
-
 class FD3D12ResourceBarrierBatcher : public FNoncopyable
 {
 public:
@@ -1053,7 +1018,7 @@ public:
 		else
 #endif
 		{
-			AddTransitionBarrierWithUAVAccessOverrides(Barriers, pResource, Before, After, Subresource);
+			AddTransitionBarrier(Barriers, pResource, Before, After, Subresource);
 		}
 
 		return 1;
diff --git a/Engine/Source/Runtime/D3D12RHI/Public/D3D12Util.h b/Engine/Source/Runtime/D3D12RHI/Public/D3D12Util.h
--- a/Engine/Source/Runtime/D3D12RHI/Public/D3D12Util.h
+++ b/Engine/Source/Runtime/D3D12RHI/Public/D3D12Util.h
@@ -295,6 +295,9 @@ public:
 	void SetResourceState(D3D12_RESOURCE_STATES State);
 	void SetSubresourceState(uint32 SubresourceIndex, D3D12_RESOURCE_STATES State);
 
+	D3D12_RESOURCE_STATES GetUAVHiddenResourceState() const { return UAVHiddenResourceState; }
+	void SetUAVHiddenResourceState(D3D12_RESOURCE_STATES InUAVHiddenResourceState) { UAVHiddenResourceState = InUAVHiddenResourceState; }
+
 private:
 	// Only used if m_AllSubresourcesSame is 1.
 	// Bits defining the state of the full resource, bits are from D3D12_RESOURCE_STATES
@@ -304,6 +307,10 @@ private:
 	// Set to 0 if m_SubresourceState is valid.  In this case, each subresources may have a different state (or may be unknown)
 	uint32 m_AllSubresourcesSame : 1;
 
+	// Special resource state to track previous state before resource transitioned to UAV when the resource
+	// has a UAV aliasing resource so correct previous state can be found (only single state allowed)
+	D3D12_RESOURCE_STATES UAVHiddenResourceState;
+
 	// Only used if m_AllSubresourcesSame is 0.
 	// The state of each subresources.  Bits are from D3D12_RESOURCE_STATES.
 	TArray<D3D12_RESOURCE_STATES, TInlineAllocator<4>> m_SubresourceState;
