Support specifying MaxAttributeSizeInBytes in FRayTracingPipelineStateInitializer instead of always using RAY_TRACING_MAX_ALLOWED_ATTRIBUTE_SIZE.

diff --git a/Engine/Source/Runtime/D3D12RHI/Private/D3D12RayTracing.cpp b/Engine/Source/Runtime/D3D12RHI/Private/D3D12RayTracing.cpp
--- a/Engine/Source/Runtime/D3D12RHI/Private/D3D12RayTracing.cpp
+++ b/Engine/Source/Runtime/D3D12RHI/Private/D3D12RayTracing.cpp
@@ -517,6 +517,7 @@ static TRefCountPtr<ID3D12StateObject> CreateRayTracingStateObject(
 	ID3D12Device5* RayTracingDevice,
 	const TArrayView<const FDXILLibrary*>& ShaderLibraries,
 	const TArrayView<LPCWSTR>& Exports,
+	uint32 MaxAttributeSizeInBytes,
 	uint32 MaxPayloadSizeInBytes,
 	const TArrayView<const D3D12_HIT_GROUP_DESC>& HitGroups,
 	const ID3D12RootSignature* GlobalRootSignature,
@@ -566,7 +567,8 @@ static TRefCountPtr<ID3D12StateObject> CreateRayTracingStateObject(
 	// Shader config
 
 	D3D12_RAYTRACING_SHADER_CONFIG ShaderConfig = {};
-	ShaderConfig.MaxAttributeSizeInBytes = RAY_TRACING_MAX_ALLOWED_ATTRIBUTE_SIZE; // sizeof 2 floats (barycentrics)
+	ShaderConfig.MaxAttributeSizeInBytes = MaxAttributeSizeInBytes;
+	check(ShaderConfig.MaxAttributeSizeInBytes <= RAY_TRACING_MAX_ALLOWED_ATTRIBUTE_SIZE);
 	ShaderConfig.MaxPayloadSizeInBytes = MaxPayloadSizeInBytes;
 	check(ShaderConfig.MaxPayloadSizeInBytes <= RAY_TRACING_MAX_ALLOWED_PAYLOAD_SIZE);
 
@@ -878,6 +880,7 @@ public:
 	struct FKey
 	{
 		uint64 ShaderHash = 0;
+		uint32 MaxAttributeSizeInBytes = 0;
 		uint32 MaxPayloadSizeInBytes = 0;
 		ID3D12RootSignature* GlobalRootSignature = nullptr;
 		ID3D12RootSignature* LocalRootSignature = nullptr;
@@ -885,6 +888,7 @@ public:
 		bool operator == (const FKey& Other) const
 		{
 			return ShaderHash == Other.ShaderHash
+				&& MaxAttributeSizeInBytes == Other.MaxAttributeSizeInBytes
 				&& MaxPayloadSizeInBytes == Other.MaxPayloadSizeInBytes
 				&& GlobalRootSignature == Other.GlobalRootSignature
 				&& LocalRootSignature == Other.LocalRootSignature;
@@ -1008,6 +1012,7 @@ public:
 			ID3D12RootSignature* GlobalRootSignature = CacheKey.GlobalRootSignature;
 			ID3D12RootSignature* LocalRootSignature = CacheKey.LocalRootSignature;
 			const uint32 DefaultLocalRootSignatureIndex = 0;
+			uint32 MaxAttributeSizeInBytes = CacheKey.MaxAttributeSizeInBytes;
 			uint32 MaxPayloadSizeInBytes = CacheKey.MaxPayloadSizeInBytes;
 
 			D3D12_HIT_GROUP_DESC HitGroupDesc = {};
@@ -1067,6 +1072,7 @@ public:
 				RayTracingDevice,
 				MakeArrayView(&LibraryPtr, 1),
 				RenamedEntryPoints,
+				MaxAttributeSizeInBytes,
 				MaxPayloadSizeInBytes,
 				MakeArrayView(&HitGroupDesc, NumHitGroups),
 				GlobalRootSignature,
@@ -1114,6 +1120,7 @@ public:
 		FD3D12Device* Device,
 		FD3D12RayTracingShader* Shader,
 		ID3D12RootSignature* GlobalRootSignature,
+		uint32 MaxAttributeSizeInBytes,
 		uint32 MaxPayloadSizeInBytes,
 		ECollectionType CollectionType,
 		FGraphEventArray& CompletionList,
@@ -1137,6 +1144,7 @@ public:
 
 		FKey CacheKey;
 		CacheKey.ShaderHash = ShaderHash;
+		CacheKey.MaxAttributeSizeInBytes = MaxAttributeSizeInBytes;
 		CacheKey.MaxPayloadSizeInBytes = MaxPayloadSizeInBytes;
 		CacheKey.GlobalRootSignature = GlobalRootSignature;
 		CacheKey.LocalRootSignature = LocalRootSignature;
@@ -2228,6 +2236,7 @@ struct FD3D12RayTracingShaderLibrary
 static void CreateSpecializedStateObjects(
 	ID3D12Device5* RayTracingDevice,
 	ID3D12RootSignature* GlobalRootSignature,
+	uint32 MaxAttributeSizeInBytes,
 	uint32 MaxPayloadSizeInBytes,
 	const FD3D12RayTracingShaderLibrary& RayGenShaders,
 	const TArray<FD3D12RayTracingPipelineCache::FEntry*>& UniqueShaderCollections,
@@ -2341,6 +2350,7 @@ static void CreateSpecializedStateObjects(
 			RayTracingDevice,
 			{}, // Libraries,
 			{}, // LibraryExports,
+			MaxAttributeSizeInBytes,
 			MaxPayloadSizeInBytes,
 			{}, // HitGroups
 			GlobalRootSignature,
@@ -2438,6 +2448,7 @@ public:
 
 			FD3D12RayTracingPipelineCache::FEntry* ShaderCacheEntry = PipelineCache->GetOrCompileShader(
 				Device, Shader, GlobalRootSignature,
+				Initializer.MaxAttributeSizeInBytes,
 				Initializer.MaxPayloadSizeInBytes,
 				CollectionType, CompileCompletionList,
 				&bCacheHit);
@@ -2631,6 +2642,7 @@ public:
 				RayTracingDevice,
 				{}, // Libraries,
 				{}, // LibraryExports,
+				Initializer.MaxAttributeSizeInBytes,
 				Initializer.MaxPayloadSizeInBytes,
 				{}, // HitGroups
 				GlobalRootSignature,
@@ -2645,6 +2657,7 @@ public:
 			CreateSpecializedStateObjects(
 				RayTracingDevice,
 				GlobalRootSignature,
+				Initializer.MaxAttributeSizeInBytes,
 				Initializer.MaxPayloadSizeInBytes,
 				RayGenShaders,
 				UniqueShaderCollections,
diff --git a/Engine/Source/Runtime/RHI/Public/RHIResources.h b/Engine/Source/Runtime/RHI/Public/RHIResources.h
--- a/Engine/Source/Runtime/RHI/Public/RHIResources.h
+++ b/Engine/Source/Runtime/RHI/Public/RHIResources.h
@@ -3007,13 +3007,15 @@ class FRayTracingPipelineStateSignature
 {
 public:
 
+	uint32 MaxAttributeSizeInBytes = 8; // sizeof FDefaultAttributes declared in RayTracingCommon.ush
 	uint32 MaxPayloadSizeInBytes = 24; // sizeof FDefaultPayload declared in RayTracingCommon.ush
 	bool bAllowHitGroupIndexing = true;
 
 	// NOTE: GetTypeHash(const FRayTracingPipelineStateInitializer& Initializer) should also be updated when changing this function
 	bool operator==(const FRayTracingPipelineStateSignature& rhs) const
 	{
-		return MaxPayloadSizeInBytes == rhs.MaxPayloadSizeInBytes
+		return MaxAttributeSizeInBytes == rhs.MaxAttributeSizeInBytes
+			&& MaxPayloadSizeInBytes == rhs.MaxPayloadSizeInBytes
 			&& bAllowHitGroupIndexing == rhs.bAllowHitGroupIndexing
 			&& RayGenHash == rhs.RayGenHash
 			&& MissHash == rhs.MissHash
@@ -3023,7 +3025,8 @@ public:
 
 	friend uint32 GetTypeHash(const FRayTracingPipelineStateSignature& Initializer)
 	{
-		return GetTypeHash(Initializer.MaxPayloadSizeInBytes) ^
+		return GetTypeHash(Initializer.MaxAttributeSizeInBytes) ^
+			GetTypeHash(Initializer.MaxPayloadSizeInBytes) ^
 			GetTypeHash(Initializer.bAllowHitGroupIndexing) ^
 			GetTypeHash(Initializer.GetRayGenHash()) ^
 			GetTypeHash(Initializer.GetRayMissHash()) ^
