remove workaround for depth bounds tests that was introduced with 16843844. The new API has  DepthBoundsTestEnable set to false by default, and non-windows d3d12 platforms were not setting it. Doing this when D3D12_USE_DERIVED_PSO==0 solves the problem

diff --git a/Engine/Source/Runtime/D3D12RHI/Private/D3D12Pipelinestate.cpp b/Engine/Source/Runtime/D3D12RHI/Private/D3D12Pipelinestate.cpp
--- a/Engine/Source/Runtime/D3D12RHI/Private/D3D12Pipelinestate.cpp
+++ b/Engine/Source/Runtime/D3D12RHI/Private/D3D12Pipelinestate.cpp
@@ -79,7 +79,7 @@ FD3D12LowLevelGraphicsPipelineStateDesc GetLowLevelGraphicsPipelineStateDesc(con
 #undef EXT_SHADER
 #endif
 
-#if PLATFORM_WINDOWS
+#if !D3D12_USE_DERIVED_PSO
 	// TODO: [PSO API] For now, keep DBT enabled, if available, until it is added as part of a member to the Initializer's DepthStencilState
 	Desc.Desc.DepthStencilState.DepthBoundsTestEnable = GSupportsDepthBoundsTest && Initializer.bDepthBounds;
 #endif
