Moving D3D12RHI::ExecuteCodeWithCopyCommandQueueUsage into ID3D12DynamicRHI::RHIExecuteOnCopyCommandQueue.
Deprecating D3D12RHI::ExecuteCodeWithCopyCommandQueueUsage.

diff --git a/Engine/Source/Runtime/D3D12RHI/Private/D3D12Buffer.cpp b/Engine/Source/Runtime/D3D12RHI/Private/D3D12Buffer.cpp
--- a/Engine/Source/Runtime/D3D12RHI/Private/D3D12Buffer.cpp
+++ b/Engine/Source/Runtime/D3D12RHI/Private/D3D12Buffer.cpp
@@ -5,7 +5,6 @@ D3D12Buffer.cpp: D3D Common code for buffers.
 =============================================================================*/
 
 #include "D3D12RHIPrivate.h"
-#include "D3D12RHIBridge.h"
 
 FD3D12Buffer::~FD3D12Buffer()
 {
@@ -309,7 +308,7 @@ FD3D12SyncPoint FD3D12Buffer::UploadResourceDataViaCopyQueue(FResourceArrayInter
 	bool bWaitForCompletion = false;
 
 	FD3D12SyncPoint CopyQueueSyncPoint;
-	D3D12RHI::ExecuteCodeWithCopyCommandQueueUsage([Device, &hCopyCommandList, bWaitForCompletion, &CopyQueueSyncPoint](ID3D12CommandQueue* D3DCommandQueue) -> void
+	GetID3D12DynamicRHI()->RHIExecuteOnCopyCommandQueue([Device, &hCopyCommandList, bWaitForCompletion, &CopyQueueSyncPoint](ID3D12CommandQueue* D3DCommandQueue) -> void
 	{
 		CopyQueueSyncPoint = Device->GetCopyCommandListManager().ExecuteCommandListNoCopyQueueSync(hCopyCommandList, bWaitForCompletion);
 	});
diff --git a/Engine/Source/Runtime/D3D12RHI/Private/D3D12CommandList.cpp b/Engine/Source/Runtime/D3D12RHI/Private/D3D12CommandList.cpp
--- a/Engine/Source/Runtime/D3D12RHI/Private/D3D12CommandList.cpp
+++ b/Engine/Source/Runtime/D3D12RHI/Private/D3D12CommandList.cpp
@@ -2,7 +2,6 @@
 
 #include "D3D12RHIPrivate.h"
 #include "D3D12CommandList.h"
-#include "D3D12RHIBridge.h"
 #include "RHIValidation.h"
 
 static int32 GD3D12BatchResourceBarriers = 1;
@@ -332,8 +331,6 @@ void FD3D12CommandAllocator::Init(ID3D12Device* InDevice, const D3D12_COMMAND_LI
 
 namespace D3D12RHI
 {
-	static FCriticalSection CopyQueueCS;
-
 	void GetGfxCommandListAndQueue(FRHICommandList& RHICmdList, void*& OutGfxCmdList, void*& OutCommandQueue)
 	{
 		IRHICommandContext& RHICmdContext = RHICmdList.GetContext();
@@ -354,20 +351,4 @@ namespace D3D12RHI
 		ID3D12CommandQueue* CommandQueue = BaseCmdContext.GetParentAdapter()->GetDevice(0)->GetD3DCommandQueue();
 		OutCommandQueue = CommandQueue;
 	}
-
-	D3D12RHI_API void ExecuteCodeWithCopyCommandQueueUsage(TFunction<void(ID3D12CommandQueue*)>&& CodeToRun)
-	{
-		IRHICommandContext* Context = GDynamicRHI->RHIGetDefaultContext();
-		checkSlow(Context);
-		FD3D12CommandContextBase& BaseCmdContext = (FD3D12CommandContextBase&)Context->GetLowestLevelContext();
-
-		ID3D12CommandQueue* CommandQueue = BaseCmdContext.GetParentAdapter()->GetDevice(0)->GetD3DCommandQueue(ED3D12CommandQueueType::Copy);
-		checkSlow(CommandQueue);
-
-		{
-			FScopeLock Lock(&CopyQueueCS);
-			CodeToRun(CommandQueue);
-		}
-	}
-
 }
\ No newline at end of file
diff --git a/Engine/Source/Runtime/D3D12RHI/Private/D3D12Commands.cpp b/Engine/Source/Runtime/D3D12RHI/Private/D3D12Commands.cpp
--- a/Engine/Source/Runtime/D3D12RHI/Private/D3D12Commands.cpp
+++ b/Engine/Source/Runtime/D3D12RHI/Private/D3D12Commands.cpp
@@ -5,7 +5,6 @@ D3D12Commands.cpp: D3D RHI commands implementation.
 =============================================================================*/
 
 #include "D3D12RHIPrivate.h"
-#include "D3D12RHIBridge.h"
 #include "StaticBoundShaderState.h"
 #include "GlobalShader.h"
 #include "OneColorShader.h"
@@ -2491,7 +2490,7 @@ void FD3D12CommandContext::RHIBroadcastTemporalEffect(const FName& InEffectName,
 
 	bool bWaitForCompletion = false;
 
-	D3D12RHI::ExecuteCodeWithCopyCommandQueueUsage([bWaitForCompletion, &CopyCommandListManager, &hCopyCommandList](ID3D12CommandQueue* D3DCommandQueue) -> void
+	GetID3D12DynamicRHI()->RHIExecuteOnCopyCommandQueue([bWaitForCompletion, &CopyCommandListManager, &hCopyCommandList](ID3D12CommandQueue* D3DCommandQueue) -> void
 	{
 		CopyCommandListManager.ExecuteCommandListNoCopyQueueSync(hCopyCommandList, bWaitForCompletion);
 	});
diff --git a/Engine/Source/Runtime/D3D12RHI/Private/D3D12RHI.cpp b/Engine/Source/Runtime/D3D12RHI/Private/D3D12RHI.cpp
--- a/Engine/Source/Runtime/D3D12RHI/Private/D3D12RHI.cpp
+++ b/Engine/Source/Runtime/D3D12RHI/Private/D3D12RHI.cpp
@@ -523,6 +523,17 @@ void FD3D12DynamicRHI::RHIAddPendingBarrier(FRHITexture* InTexture, D3D12_RESOUR
 	CommandList.AddPendingResourceBarrier(D3D12Texture->GetResource(), InState, InSubResource);
 }
 
+void FD3D12DynamicRHI::RHIExecuteOnCopyCommandQueue(TFunction<void(ID3D12CommandQueue*)>&& CodeToRun)
+{
+	ID3D12CommandQueue* CommandQueue = GetAdapter().GetDevice(0)->GetD3DCommandQueue(ED3D12CommandQueueType::Copy);
+	checkSlow(CommandQueue);
+
+	{
+		FScopeLock Lock(&CopyQueueCS);
+		CodeToRun(CommandQueue);
+	}
+}
+
 bool ID3D12DynamicRHI::IsD3DDebugEnabled()
 {
 	return D3D12RHI_ShouldCreateWithD3DDebug();
diff --git a/Engine/Source/Runtime/D3D12RHI/Private/D3D12RHIPrivate.h b/Engine/Source/Runtime/D3D12RHI/Private/D3D12RHIPrivate.h
--- a/Engine/Source/Runtime/D3D12RHI/Private/D3D12RHIPrivate.h
+++ b/Engine/Source/Runtime/D3D12RHI/Private/D3D12RHIPrivate.h
@@ -454,6 +454,7 @@ public:
 	virtual void RHIFinishExternalComputeWork(uint32 InDeviceIndex, ID3D12GraphicsCommandList* InCommandList) final override;
 	virtual void RHIRegisterWork(uint32 InDeviceIndex, uint32 NumPrimitives) final override;
 	virtual void RHIAddPendingBarrier(FRHITexture* InTexture, D3D12_RESOURCE_STATES InState, uint32 InSubResource) final override;
+	virtual void RHIExecuteOnCopyCommandQueue(TFunction<void(ID3D12CommandQueue*)>&& CodeToRun) final override;
 
 	//
 	// The Following functions are the _RenderThread version of the above functions. They allow the RHI to control the thread synchronization for greater efficiency.
@@ -1198,6 +1199,9 @@ protected:
 	void* ZeroBuffer;
 	uint32 ZeroBufferSize;
 
+	/* Primary lock for RHIExecuteOnCopyCommandQueue */
+	FCriticalSection CopyQueueCS;
+
 public:
 
 	template<typename BaseResourceType>
diff --git a/Engine/Source/Runtime/D3D12RHI/Private/D3D12Texture.cpp b/Engine/Source/Runtime/D3D12RHI/Private/D3D12Texture.cpp
--- a/Engine/Source/Runtime/D3D12RHI/Private/D3D12Texture.cpp
+++ b/Engine/Source/Runtime/D3D12RHI/Private/D3D12Texture.cpp
@@ -5,7 +5,6 @@
 	=============================================================================*/
 
 #include "D3D12RHIPrivate.h"
-#include "D3D12RHIBridge.h"
 #include "TextureProfiler.h"
 
 int64 FD3D12GlobalStats::GDedicatedVideoMemory = 0;
@@ -1603,7 +1602,7 @@ FTexture2DRHIRef FD3D12DynamicRHI::RHIAsyncCreateTexture2D(uint32 SizeX, uint32
 
 			bool bWaitForCompletion = true;
 
-			D3D12RHI::ExecuteCodeWithCopyCommandQueueUsage([Device, bWaitForCompletion, &hCopyCommandList](ID3D12CommandQueue* D3DCommandQueue) -> void
+			RHIExecuteOnCopyCommandQueue([Device, bWaitForCompletion, &hCopyCommandList](ID3D12CommandQueue* D3DCommandQueue) -> void
 			{
 				Device->GetCopyCommandListManager().ExecuteCommandListNoCopyQueueSync(hCopyCommandList, bWaitForCompletion);
 			});
diff --git a/Engine/Source/Runtime/D3D12RHI/Public/D3D12RHIBridge.h b/Engine/Source/Runtime/D3D12RHI/Public/D3D12RHIBridge.h
--- a/Engine/Source/Runtime/D3D12RHI/Public/D3D12RHIBridge.h
+++ b/Engine/Source/Runtime/D3D12RHI/Public/D3D12RHIBridge.h
@@ -1,14 +1,20 @@
 // Copyright Epic Games, Inc. All Rights Reserved.
 
 #include "Templates/Function.h"
+#include "ID3D12DynamicRHI.h"
 
 namespace D3D12RHI
 {
 	// Use these functions with care!
 
 	// Returns ID3D12GraphicsCommandList* and ID3D12CommandQueue*
+	UE_DEPRECATED(5.1, "ID3D12DynamicRHI should be used to acquire D3D12 command lists and command queues")
 	D3D12RHI_API void GetGfxCommandListAndQueue(FRHICommandList& RHICmdList, void*& OutGfxCmdList, void*& OutCommandQueue);
 
 	// Execute code using the CopyCommandQueue interface in a thread safe way
-	D3D12RHI_API void ExecuteCodeWithCopyCommandQueueUsage(TFunction<void(ID3D12CommandQueue*)>&& CodeToRun);
+	UE_DEPRECATED(5.1, "ID3D12DynamicRHI::RHIExecuteOnCopyCommandQueue() should be used instead")
+	inline void ExecuteCodeWithCopyCommandQueueUsage(TFunction<void(ID3D12CommandQueue*)>&& CodeToRun)
+	{
+		GetID3D12DynamicRHI()->RHIExecuteOnCopyCommandQueue(MoveTemp(CodeToRun));
+	}
 }
diff --git a/Engine/Source/Runtime/D3D12RHI/Public/ID3D12DynamicRHI.h b/Engine/Source/Runtime/D3D12RHI/Public/ID3D12DynamicRHI.h
--- a/Engine/Source/Runtime/D3D12RHI/Public/ID3D12DynamicRHI.h
+++ b/Engine/Source/Runtime/D3D12RHI/Public/ID3D12DynamicRHI.h
@@ -2,7 +2,8 @@
 
 #pragma once
 
-#include "DynamicRHI.h"
+#include "RHI.h"
+#include "Templates/Function.h"
 
 #if PLATFORM_WINDOWS
 	#include "Windows/D3D12ThirdParty.h"
@@ -49,6 +50,8 @@ struct ID3D12DynamicRHI : public FDynamicRHI
 	virtual void                  RHIRegisterWork(uint32 InDeviceIndex, uint32 NumPrimitives) = 0;
 	virtual void                  RHIAddPendingBarrier(FRHITexture* InTexture, D3D12_RESOURCE_STATES InState, uint32 InSubResource) = 0;
 
+	virtual void                  RHIExecuteOnCopyCommandQueue(TFunction<void(ID3D12CommandQueue*)>&& CodeToRun) = 0;
+
 	static D3D12RHI_API bool      IsD3DDebugEnabled();
 };
 
