fix for r.setres on d3d12 consoles: the backbuffers were still referenced by the present queue, preventing their proper destruction

diff --git a/Engine/Source/Runtime/D3D12RHI/Private/D3D12Viewport.cpp b/Engine/Source/Runtime/D3D12RHI/Private/D3D12Viewport.cpp
--- a/Engine/Source/Runtime/D3D12RHI/Private/D3D12Viewport.cpp
+++ b/Engine/Source/Runtime/D3D12RHI/Private/D3D12Viewport.cpp
@@ -505,6 +505,8 @@ void FD3D12Viewport::Resize(uint32 InSizeX, uint32 InSizeY, bool bInIsFullscreen
 		check(SDRBackBuffers[i] == nullptr);
 	}
 
+	ClearPresentQueue();
+
 	// Flush the outstanding GPU work and wait for it to complete.
 	FlushRenderingCommands();
 	FRHICommandListExecutor::CheckNoOutstandingCmdLists();
diff --git a/Engine/Source/Runtime/D3D12RHI/Private/Windows/WindowsD3D12Viewport.cpp b/Engine/Source/Runtime/D3D12RHI/Private/Windows/WindowsD3D12Viewport.cpp
--- a/Engine/Source/Runtime/D3D12RHI/Private/Windows/WindowsD3D12Viewport.cpp
+++ b/Engine/Source/Runtime/D3D12RHI/Private/Windows/WindowsD3D12Viewport.cpp
@@ -259,6 +259,10 @@ void FD3D12Viewport::FinalDestroyInternal()
 {
 }
 
+void FD3D12Viewport::ClearPresentQueue()
+{
+}
+
 void FD3D12Viewport::ConditionalResetSwapChain(bool bIgnoreFocus)
 {
 	if (!bIsValid)
diff --git a/Engine/Source/Runtime/D3D12RHI/Public/D3D12Viewport.h b/Engine/Source/Runtime/D3D12RHI/Public/D3D12Viewport.h
--- a/Engine/Source/Runtime/D3D12RHI/Public/D3D12Viewport.h
+++ b/Engine/Source/Runtime/D3D12RHI/Public/D3D12Viewport.h
@@ -178,6 +178,7 @@ private:
 
 	void ResizeInternal();
 	void FinalDestroyInternal();
+	void ClearPresentQueue();
 
 	uint64 LastFlipTime;
 	uint64 LastFrameComplete;
