MGPU:  Fixes for editor GPU profile feature, when multi-GPU is enabled (via r.AllowMultiGPUInEditor=1 and -MaxGPUCount=2).  Need to call "FD3DGPUProfiler::EndFrame" for all devices (GPUs), to match the calls to "FD3DGPUProfiler::BeginFrame" for each device, to avoid an assert.  Also need to restrict dumping of the profile result tree to the first GPU, since the second GPU doesn't currently render anything in the editor, and will be empty (whatever gets dumped last is displayed in the visualizer).  If we ever start using the second GPU outside GPU Lightmass, we could upgrade the profiler to collate the event trees if necessary, but this works fine until then.

#rnx

diff --git a/Engine/Source/Runtime/D3D12RHI/Private/D3D12CommandContext.cpp b/Engine/Source/Runtime/D3D12RHI/Private/D3D12CommandContext.cpp
--- a/Engine/Source/Runtime/D3D12RHI/Private/D3D12CommandContext.cpp
+++ b/Engine/Source/Runtime/D3D12RHI/Private/D3D12CommandContext.cpp
@@ -657,16 +657,24 @@ void FD3D12CommandContextBase::RHIEndFrame()
 	UpdateMemoryStats();
 
 	// Stop Timing at the very last moment
-	D3D12RHI::FD3DGPUProfiler& GPUProfiler = Device->GetGPUProfiler();
-	GPUProfiler.EndFrame(ParentAdapter->GetOwningRHI());
+	for (uint32 GPUIndex : GPUMask)
+	{
+		Device = ParentAdapter->GetDevice(GPUIndex);
+
+		D3D12RHI::FD3DGPUProfiler& GPUProfiler = Device->GetGPUProfiler();
+		GPUProfiler.EndFrame(ParentAdapter->GetOwningRHI());
 
 #if WITH_MGPU
-	// Multi-GPU support : For now, set GGPUFrameTime to GPU 0's frame time to be
-	// consistent with code that calls RHIGetGPUFrameCycles and is not MGPU-aware.
-	// Perhaps we should change it to get the max frame time of all GPUs?
-	// Only use this code if running on a multi GPU configuration.
-	GGPUFrameTime = GPUProfiler.GetGPUFrameCycles(0);
+		if (GPUIndex == 0)
+		{
+			// Multi-GPU support : For now, set GGPUFrameTime to GPU 0's frame time to be
+			// consistent with code that calls RHIGetGPUFrameCycles and is not MGPU-aware.
+			// Perhaps we should change it to get the max frame time of all GPUs?
+			// Only use this code if running on a multi GPU configuration.
+			GGPUFrameTime = GPUProfiler.GetGPUFrameCycles(0);
+		}
 #endif
+	}
 }
 
 void FD3D12CommandContextBase::SignalTransitionFences(TArrayView<const FRHITransition*> Transitions)
diff --git a/Engine/Source/Runtime/D3D12RHI/Private/D3D12Stats.cpp b/Engine/Source/Runtime/D3D12RHI/Private/D3D12Stats.cpp
--- a/Engine/Source/Runtime/D3D12RHI/Private/D3D12Stats.cpp
+++ b/Engine/Source/Runtime/D3D12RHI/Private/D3D12Stats.cpp
@@ -120,14 +120,25 @@ void D3D12RHI::FD3DGPUProfiler::EndFrame(FD3D12DynamicRHI* InRHI)
 			DoPostProfileGPUWork();
 			UE_LOG(LogD3D12RHI, Log, TEXT(""));
 			UE_LOG(LogD3D12RHI, Log, TEXT(""));
-			CurrentEventNodeFrame->DumpEventTree();
 			GTriggerGPUProfile = false;
 			bLatchedGProfilingGPU = false;
 
-			if (RHIConfig::ShouldSaveScreenshotAfterProfilingGPU()
-				&& GEngine->GameViewport)
+			// Only dump the event tree and generate the screenshot for the first GPU.  Eventually, we may want to collate
+			// profiling data for all GPUs into a single tree, but the short term goal is to make profiling in the editor
+			// functional at all with "r.AllowMultiGPUInEditor=1" and "-MaxGPUCount=2" (settings required to enable multiple
+			// GPUs for GPU Lightmass).  In the editor, we don't actually render anything on the additional GPUs, but the
+			// editor's profile visualizer will pick up whatever event tree we dumped last, which will be the empty one from
+			// the last GPU, making the results useless without this code fix.  Unreal Insights would be preferred for
+			// multi-GPU profiling outside the editor.
+			if (GPUIndex == 0)
 			{
-				GEngine->GameViewport->Exec(NULL, TEXT("SCREENSHOT"), *GLog);
+				CurrentEventNodeFrame->DumpEventTree();
+
+				if (RHIConfig::ShouldSaveScreenshotAfterProfilingGPU()
+					&& GEngine->GameViewport)
+				{
+					GEngine->GameViewport->Exec(NULL, TEXT("SCREENSHOT"), *GLog);
+				}
 			}
 		}
 	}
