Reset UniqueTables instead of emptying it to reduce allocator churn.
Reset FCommandListResourceState::ResourceStates instead of emptying it to reduce allocator churn.

diff --git a/Engine/Source/Runtime/D3D12RHI/Private/D3D12CommandList.cpp b/Engine/Source/Runtime/D3D12RHI/Private/D3D12CommandList.cpp
--- a/Engine/Source/Runtime/D3D12RHI/Private/D3D12CommandList.cpp
+++ b/Engine/Source/Runtime/D3D12RHI/Private/D3D12CommandList.cpp
@@ -300,7 +300,7 @@ CResourceState& FD3D12CommandListHandle::FD3D12CommandListData::FCommandListReso
 
 void FD3D12CommandListHandle::FD3D12CommandListData::FCommandListResourceState::Empty()
 {
-	ResourceStates.Empty();
+	ResourceStates.Reset();
 }
 
 void FD3D12CommandListHandle::Execute(FD3D12SyncPoint& CopyQueueSyncPoint, bool WaitForCompletion)
diff --git a/Engine/Source/Runtime/D3D12RHI/Private/D3D12DescriptorCache.cpp b/Engine/Source/Runtime/D3D12RHI/Private/D3D12DescriptorCache.cpp
--- a/Engine/Source/Runtime/D3D12RHI/Private/D3D12DescriptorCache.cpp
+++ b/Engine/Source/Runtime/D3D12RHI/Private/D3D12DescriptorCache.cpp
@@ -205,7 +205,7 @@ void FD3D12DescriptorCache::GatherUniqueSamplerTables()
 	}
 
 	// Reset the tables as the next frame should inherit them from the global heap
-	UniqueTables.Empty();
+	UniqueTables.Reset();
 }
 
 bool FD3D12DescriptorCache::SetDescriptorHeaps()
