Copying //UE5/Dev-ParallelRendering to //UE5/Main

diff --git a/Engine/Source/Runtime/D3D12RHI/Private/D3D12RenderTarget.cpp b/Engine/Source/Runtime/D3D12RHI/Private/D3D12RenderTarget.cpp
--- a/Engine/Source/Runtime/D3D12RHI/Private/D3D12RenderTarget.cpp
+++ b/Engine/Source/Runtime/D3D12RHI/Private/D3D12RenderTarget.cpp
@@ -48,8 +48,8 @@ static FResolveRect GetDefaultRect(const FResolveRect& Rect, uint32 DefaultWidth
 template<typename TPixelShader>
 void FD3D12CommandContext::ResolveTextureUsingShader(
 	FRHICommandList_RecursiveHazardous& RHICmdList,
-	FD3D12Texture2D* SourceTexture,
-	FD3D12Texture2D* DestTexture,
+	FD3D12Texture* SourceTexture,
+	FD3D12Texture* DestTexture,
 	FD3D12RenderTargetView* DestTextureRTV,
 	FD3D12DepthStencilView* DestTextureDSV,
 	const D3D12_RESOURCE_DESC& ResolveTargetDesc,
@@ -202,96 +202,78 @@ void FD3D12CommandContext::RHICopyToResolveTarget(FRHITexture* SourceTextureRHI,
 		return;
 	}
 
-	// Could be back buffer reference texture, so get the correct D3D12 texture here
-	// We know already that it's a FD3D12Texture2D so cast is safe
-	if (EnumHasAnyFlags(SourceTextureRHI->GetFlags(), TexCreate_Presentable))
-	{
-		FD3D12BackBufferReferenceTexture2D* BufferBufferReferenceTexture = (FD3D12BackBufferReferenceTexture2D*)SourceTextureRHI;
-		SourceTextureRHI = BufferBufferReferenceTexture->GetBackBufferTexture();
-	}
-	if (EnumHasAnyFlags(DestTextureRHI->GetFlags(), TexCreate_Presentable))
-	{
-		FD3D12BackBufferReferenceTexture2D* BufferBufferReferenceTexture = (FD3D12BackBufferReferenceTexture2D*)DestTextureRHI;
-		DestTextureRHI = BufferBufferReferenceTexture->GetBackBufferTexture();
-	}
-
 	FRHICommandList_RecursiveHazardous RHICmdList(this, FRHIGPUMask::FromIndex(GetGPUIndex()));
 
-	FD3D12Texture2D* SourceTexture2D = RetrieveObject<FD3D12Texture2D>(SourceTextureRHI->GetTexture2D());
-	FD3D12Texture2D* DestTexture2D = RetrieveObject<FD3D12Texture2D>(DestTextureRHI->GetTexture2D());
-
-	FD3D12TextureCube* SourceTextureCube = RetrieveObject<FD3D12TextureCube>(SourceTextureRHI->GetTextureCube());
-	FD3D12TextureCube* DestTextureCube = RetrieveObject<FD3D12TextureCube>(DestTextureRHI->GetTextureCube());
+	FD3D12Texture* SourceTexture = GetD3D12TextureFromRHITexture(SourceTextureRHI);
+	FD3D12Texture* DestTexture = GetD3D12TextureFromRHITexture(DestTextureRHI);
 
-	FD3D12Texture3D* SourceTexture3D = RetrieveObject<FD3D12Texture3D>(SourceTextureRHI->GetTexture3D());
-	FD3D12Texture3D* DestTexture3D = RetrieveObject<FD3D12Texture3D>(DestTextureRHI->GetTexture3D());
+	ETextureDimension SourceDimension = SourceTexture->GetDesc().Dimension;
+	ETextureDimension DestDimension = DestTexture->GetDesc().Dimension;
 
-	if (SourceTexture2D && DestTexture2D)
+	if (SourceDimension == ETextureDimension::Texture2D && DestDimension == ETextureDimension::Texture2D)
 	{
 		const D3D_FEATURE_LEVEL FeatureLevel = GetParentDevice()->GetParentAdapter()->GetFeatureLevel();
-
-		check(!SourceTextureCube && !DestTextureCube);
-		if (SourceTexture2D != DestTexture2D)
+		if (SourceTexture != DestTexture)
 		{
 			if (IsDefaultContext())
 			{
 				GetParentDevice()->RegisterGPUWork();
 			}
 
-			if (DestTexture2D->GetDepthStencilView(FExclusiveDepthStencil::DepthWrite_StencilWrite)
+			if (SourceTexture->GetDepthStencilView(FExclusiveDepthStencil::DepthWrite_StencilWrite)
 				&& SourceTextureRHI->IsMultisampled()
 				&& !DestTextureRHI->IsMultisampled())
 			{
-				D3D12_RESOURCE_DESC const& ResolveTargetDesc = DestTexture2D->GetResource()->GetDesc();
+				D3D12_RESOURCE_DESC const& ResolveTargetDesc = DestTexture->GetResource()->GetDesc();
 
 				ResolveTextureUsingShader<FResolveDepthPS>(
 					RHICmdList,
-					SourceTexture2D,
-					DestTexture2D,
-					DestTexture2D->GetRenderTargetView(0, -1),
-					DestTexture2D->GetDepthStencilView(FExclusiveDepthStencil::DepthWrite_StencilWrite),
+					SourceTexture,
+					DestTexture,
+					DestTexture->GetRenderTargetView(0, -1),
+					DestTexture->GetDepthStencilView(FExclusiveDepthStencil::DepthWrite_StencilWrite),
 					ResolveTargetDesc,
-					GetDefaultRect(ResolveParams.Rect, DestTexture2D->GetSizeX(), DestTexture2D->GetSizeY()),
-					GetDefaultRect(ResolveParams.Rect, DestTexture2D->GetSizeX(), DestTexture2D->GetSizeY()),
+					GetDefaultRect(ResolveParams.Rect, DestTexture->GetSizeX(), DestTexture->GetSizeY()),
+					GetDefaultRect(ResolveParams.Rect, DestTexture->GetSizeX(), DestTexture->GetSizeY()),
 					FDummyResolveParameter()
 					);
 			}
 			else
 			{
 				DXGI_FORMAT SrcFmt = (DXGI_FORMAT)GPixelFormats[SourceTextureRHI->GetFormat()].PlatformFormat;
-				DXGI_FORMAT DstFmt = (DXGI_FORMAT)GPixelFormats[DestTexture2D->GetFormat()].PlatformFormat;
+				DXGI_FORMAT DstFmt = (DXGI_FORMAT)GPixelFormats[DestTexture->GetFormat()].PlatformFormat;
 
-				DXGI_FORMAT Fmt = ConvertTypelessToUnorm((DXGI_FORMAT)GPixelFormats[DestTexture2D->GetFormat()].PlatformFormat);
+				DXGI_FORMAT Fmt = ConvertTypelessToUnorm((DXGI_FORMAT)GPixelFormats[DestTexture->GetFormat()].PlatformFormat);
 
 				// Determine whether a MSAA resolve is needed, or just a copy.
-				if (SourceTextureRHI->IsMultisampled() && !DestTexture2D->IsMultisampled())
+				if (SourceTextureRHI->IsMultisampled() && !DestTexture->IsMultisampled())
 				{
-					FScopedResourceBarrier ConditionalScopeResourceBarrierDest(CommandListHandle, DestTexture2D->GetResource(), D3D12_RESOURCE_STATE_RESOLVE_DEST, ResolveParams.DestArrayIndex, FD3D12DynamicRHI::ETransitionMode::Validate);
-					FScopedResourceBarrier ConditionalScopeResourceBarrierSource(CommandListHandle, SourceTexture2D->GetResource(), D3D12_RESOURCE_STATE_RESOLVE_SOURCE, ResolveParams.SourceArrayIndex, FD3D12DynamicRHI::ETransitionMode::Validate);
+					FScopedResourceBarrier ConditionalScopeResourceBarrierDest(CommandListHandle, DestTexture->GetResource(), D3D12_RESOURCE_STATE_RESOLVE_DEST, ResolveParams.DestArrayIndex, FD3D12DynamicRHI::ETransitionMode::Validate);
+					FScopedResourceBarrier ConditionalScopeResourceBarrierSource(CommandListHandle, SourceTexture->GetResource(), D3D12_RESOURCE_STATE_RESOLVE_SOURCE, ResolveParams.SourceArrayIndex, FD3D12DynamicRHI::ETransitionMode::Validate);
 
 					otherWorkCounter++;
 					CommandListHandle.FlushResourceBarriers();
 					CommandListHandle->ResolveSubresource(
-						DestTexture2D->GetResource()->GetResource(),
+						DestTexture->GetResource()->GetResource(),
 						ResolveParams.DestArrayIndex,
-						SourceTexture2D->GetResource()->GetResource(),
+						SourceTexture->GetResource()->GetResource(),
 						ResolveParams.SourceArrayIndex,
 						Fmt
 						);
 
-					CommandListHandle.UpdateResidency(SourceTexture2D->GetResource());
-					CommandListHandle.UpdateResidency(DestTexture2D->GetResource());
+					CommandListHandle.UpdateResidency(SourceTexture->GetResource());
+					CommandListHandle.UpdateResidency(DestTexture->GetResource());
 				}
 				else
 				{
-					D3D12_RESOURCE_DESC const& srcDesc = SourceTexture2D->GetResource()->GetDesc();
-					D3D12_RESOURCE_DESC const& ResolveTargetDesc = DestTexture2D->GetResource()->GetDesc();
+					D3D12_RESOURCE_DESC const& srcDesc = SourceTexture->GetResource()->GetDesc();
+					D3D12_RESOURCE_DESC const& ResolveTargetDesc = DestTexture->GetResource()->GetDesc();
 					bool bCopySrcSubRect = ResolveParams.Rect.IsValid() && (ResolveParams.Rect.X1 != 0 || ResolveParams.Rect.Y1 != 0 || ResolveParams.Rect.X2 != srcDesc.Width || ResolveParams.Rect.Y2 != srcDesc.Height);
 					bool bCopyDstSubRect = ResolveParams.DestRect.IsValid() && (ResolveParams.DestRect.X1 != 0 || ResolveParams.DestRect.Y1 != 0 || ResolveParams.DestRect.X2 != ResolveTargetDesc.Width || ResolveParams.DestRect.Y2 != ResolveTargetDesc.Height);
 
 					if ((bCopySrcSubRect || bCopyDstSubRect)
 						&& !SourceTextureRHI->IsMultisampled()
-						&& !DestTexture2D->GetDepthStencilView(FExclusiveDepthStencil::DepthWrite_StencilWrite))
+						&& !DestTexture->GetDepthStencilView(FExclusiveDepthStencil::DepthWrite_StencilWrite))
 					{
 						// currently no support for readback buffers
 						check(ResolveTargetDesc.Dimension != D3D12_RESOURCE_DIMENSION_BUFFER);
@@ -308,11 +290,11 @@ void FD3D12CommandContext::RHICopyToResolveTarget(FRHITexture* SourceTextureRHI,
 
 						const FResolveRect& DestRect = ResolveParams.DestRect.IsValid() ? ResolveParams.DestRect : SrcRect;
 
-						FScopedResourceBarrier ConditionalScopeResourceBarrierDest(CommandListHandle, DestTexture2D->GetResource(), D3D12_RESOURCE_STATE_COPY_DEST, ResolveParams.DestArrayIndex, FD3D12DynamicRHI::ETransitionMode::Validate);
-						FScopedResourceBarrier ConditionalScopeResourceBarrierSource(CommandListHandle, SourceTexture2D->GetResource(), D3D12_RESOURCE_STATE_COPY_SOURCE, ResolveParams.SourceArrayIndex, FD3D12DynamicRHI::ETransitionMode::Validate);
+						FScopedResourceBarrier ConditionalScopeResourceBarrierDest(CommandListHandle, DestTexture->GetResource(), D3D12_RESOURCE_STATE_COPY_DEST, ResolveParams.DestArrayIndex, FD3D12DynamicRHI::ETransitionMode::Validate);
+						FScopedResourceBarrier ConditionalScopeResourceBarrierSource(CommandListHandle, SourceTexture->GetResource(), D3D12_RESOURCE_STATE_COPY_SOURCE, ResolveParams.SourceArrayIndex, FD3D12DynamicRHI::ETransitionMode::Validate);
 
-						CD3DX12_TEXTURE_COPY_LOCATION DestCopyLocation(DestTexture2D->GetResource()->GetResource(), ResolveParams.DestArrayIndex);
-						CD3DX12_TEXTURE_COPY_LOCATION SourceCopyLocation(SourceTexture2D->GetResource()->GetResource(), ResolveParams.SourceArrayIndex);
+						CD3DX12_TEXTURE_COPY_LOCATION DestCopyLocation(DestTexture->GetResource()->GetResource(), ResolveParams.DestArrayIndex);
+						CD3DX12_TEXTURE_COPY_LOCATION SourceCopyLocation(SourceTexture->GetResource()->GetResource(), ResolveParams.SourceArrayIndex);
 
 						numCopies++;
 						CommandListHandle.FlushResourceBarriers();
@@ -322,19 +304,19 @@ void FD3D12CommandContext::RHICopyToResolveTarget(FRHITexture* SourceTextureRHI,
 							&SourceCopyLocation,
 							&SrcBox);
 
-						CommandListHandle.UpdateResidency(SourceTexture2D->GetResource());
-						CommandListHandle.UpdateResidency(DestTexture2D->GetResource());
+						CommandListHandle.UpdateResidency(SourceTexture->GetResource());
+						CommandListHandle.UpdateResidency(DestTexture->GetResource());
 					}
 					else
 					{
-						FScopedResourceBarrier ConditionalScopeResourceBarrierSource(CommandListHandle, SourceTexture2D->GetResource(), D3D12_RESOURCE_STATE_COPY_SOURCE, ResolveParams.SourceArrayIndex, FD3D12DynamicRHI::ETransitionMode::Validate);
+						FScopedResourceBarrier ConditionalScopeResourceBarrierSource(CommandListHandle, SourceTexture->GetResource(), D3D12_RESOURCE_STATE_COPY_SOURCE, ResolveParams.SourceArrayIndex, FD3D12DynamicRHI::ETransitionMode::Validate);
 
 						// Resolve to a buffer.
 						if (ResolveTargetDesc.Dimension == D3D12_RESOURCE_DIMENSION_BUFFER)
 						{
 							check(IsDefaultContext());
 
-							const uint32 BlockBytes = GPixelFormats[SourceTexture2D->GetFormat()].BlockBytes;
+							const uint32 BlockBytes = GPixelFormats[SourceTexture->GetFormat()].BlockBytes;
 							const uint32 XBytes = (uint32)srcDesc.Width * BlockBytes;
 							const uint32 XBytesAligned = Align(XBytes, FD3D12_TEXTURE_DATA_PITCH_ALIGNMENT);
 
@@ -349,8 +331,8 @@ void FD3D12CommandContext::RHICopyToResolveTarget(FRHITexture* SourceTextureRHI,
 							placedTexture2D.Offset = 0;
 							placedTexture2D.Footprint = DestSubresource;
 
-							CD3DX12_TEXTURE_COPY_LOCATION DestCopyLocation(DestTexture2D->GetResource()->GetResource(), placedTexture2D);
-							CD3DX12_TEXTURE_COPY_LOCATION SourceCopyLocation(SourceTexture2D->GetResource()->GetResource(), ResolveParams.SourceArrayIndex);
+							CD3DX12_TEXTURE_COPY_LOCATION DestCopyLocation(DestTexture->GetResource()->GetResource(), placedTexture2D);
+							CD3DX12_TEXTURE_COPY_LOCATION SourceCopyLocation(SourceTexture->GetResource()->GetResource(), ResolveParams.SourceArrayIndex);
 
 							numCopies++;
 							CommandListHandle.FlushResourceBarriers();
@@ -360,11 +342,11 @@ void FD3D12CommandContext::RHICopyToResolveTarget(FRHITexture* SourceTextureRHI,
 								&SourceCopyLocation,
 								nullptr);
 
-							CommandListHandle.UpdateResidency(SourceTexture2D->GetResource());
-							CommandListHandle.UpdateResidency(DestTexture2D->GetResource());
+							CommandListHandle.UpdateResidency(SourceTexture->GetResource());
+							CommandListHandle.UpdateResidency(DestTexture->GetResource());
 
 							// Save the command list handle. This lets us check when this command list is complete. Note: This must be saved before we execute the command list
-							DestTexture2D->SetReadBackListHandle(CommandListHandle);
+							DestTexture->SetReadBackListHandle(CommandListHandle);
 
 							// Break up the command list here so that the wait on the previous frame's results don't block.
 							FlushCommands();
@@ -373,10 +355,10 @@ void FD3D12CommandContext::RHICopyToResolveTarget(FRHITexture* SourceTextureRHI,
 						else
 						{
 							// Transition to the copy dest state.
-							FScopedResourceBarrier ConditionalScopeResourceBarrierDest(CommandListHandle, DestTexture2D->GetResource(), D3D12_RESOURCE_STATE_COPY_DEST, 0, FD3D12DynamicRHI::ETransitionMode::Validate);
+							FScopedResourceBarrier ConditionalScopeResourceBarrierDest(CommandListHandle, DestTexture->GetResource(), D3D12_RESOURCE_STATE_COPY_DEST, 0, FD3D12DynamicRHI::ETransitionMode::Validate);
 
-							CD3DX12_TEXTURE_COPY_LOCATION DestCopyLocation(DestTexture2D->GetResource()->GetResource(), ResolveParams.DestArrayIndex);
-							CD3DX12_TEXTURE_COPY_LOCATION SourceCopyLocation(SourceTexture2D->GetResource()->GetResource(), ResolveParams.SourceArrayIndex);
+							CD3DX12_TEXTURE_COPY_LOCATION DestCopyLocation(DestTexture->GetResource()->GetResource(), ResolveParams.DestArrayIndex);
+							CD3DX12_TEXTURE_COPY_LOCATION SourceCopyLocation(SourceTexture->GetResource()->GetResource(), ResolveParams.SourceArrayIndex);
 
 							numCopies++;
 							CommandListHandle.FlushResourceBarriers();
@@ -386,19 +368,17 @@ void FD3D12CommandContext::RHICopyToResolveTarget(FRHITexture* SourceTextureRHI,
 								&SourceCopyLocation,
 								nullptr);
 
-							CommandListHandle.UpdateResidency(SourceTexture2D->GetResource());
-							CommandListHandle.UpdateResidency(DestTexture2D->GetResource());
+							CommandListHandle.UpdateResidency(SourceTexture->GetResource());
+							CommandListHandle.UpdateResidency(DestTexture->GetResource());
 						}
 					}
 				}
 			}
 		}
 	}
-	else if (SourceTextureCube && DestTextureCube)
+	else if (SourceDimension == ETextureDimension::TextureCube && DestDimension == ETextureDimension::TextureCube)
 	{
-		check(!SourceTexture2D && !DestTexture2D);
-
-		if (SourceTextureCube != DestTextureCube)
+		if (SourceTexture != DestTexture)
 		{
 			if (IsDefaultContext())
 			{
@@ -407,35 +387,35 @@ void FD3D12CommandContext::RHICopyToResolveTarget(FRHITexture* SourceTextureRHI,
 
 			// Determine the cubemap face being resolved.
 			const uint32 D3DFace = GetD3D12CubeFace(ResolveParams.CubeFace);
-			const uint32 SourceSubresource = CalcSubresource(ResolveParams.MipIndex, ResolveParams.SourceArrayIndex * 6 + D3DFace, SourceTextureCube->GetNumMips());
-			const uint32 DestSubresource = CalcSubresource(ResolveParams.MipIndex, ResolveParams.DestArrayIndex * 6 + D3DFace, DestTextureCube->GetNumMips());
+			const uint32 SourceSubresource = CalcSubresource(ResolveParams.MipIndex, ResolveParams.SourceArrayIndex * 6 + D3DFace, SourceTexture->GetNumMips());
+			const uint32 DestSubresource = CalcSubresource(ResolveParams.MipIndex, ResolveParams.DestArrayIndex * 6 + D3DFace, DestTexture->GetNumMips());
 
 			// Determine whether a MSAA resolve is needed, or just a copy.
-			if (SourceTextureRHI->IsMultisampled() && !DestTextureCube->IsMultisampled())
+			if (SourceTextureRHI->IsMultisampled() && !DestTexture->IsMultisampled())
 			{
-				FScopedResourceBarrier ConditionalScopeResourceBarrierDest(CommandListHandle, DestTextureCube->GetResource(), D3D12_RESOURCE_STATE_RESOLVE_DEST, DestSubresource, FD3D12DynamicRHI::ETransitionMode::Validate);
-				FScopedResourceBarrier ConditionalScopeResourceBarrierSource(CommandListHandle, SourceTextureCube->GetResource(), D3D12_RESOURCE_STATE_RESOLVE_SOURCE, SourceSubresource, FD3D12DynamicRHI::ETransitionMode::Validate);
+				FScopedResourceBarrier ConditionalScopeResourceBarrierDest(CommandListHandle, DestTexture->GetResource(), D3D12_RESOURCE_STATE_RESOLVE_DEST, DestSubresource, FD3D12DynamicRHI::ETransitionMode::Validate);
+				FScopedResourceBarrier ConditionalScopeResourceBarrierSource(CommandListHandle, SourceTexture->GetResource(), D3D12_RESOURCE_STATE_RESOLVE_SOURCE, SourceSubresource, FD3D12DynamicRHI::ETransitionMode::Validate);
 
 				otherWorkCounter++;
 				CommandListHandle.FlushResourceBarriers();
 				CommandListHandle->ResolveSubresource(
-					DestTextureCube->GetResource()->GetResource(),
+					DestTexture->GetResource()->GetResource(),
 					DestSubresource,
-					SourceTextureCube->GetResource()->GetResource(),
+					SourceTexture->GetResource()->GetResource(),
 					SourceSubresource,
-					(DXGI_FORMAT)GPixelFormats[DestTextureCube->GetFormat()].PlatformFormat
+					(DXGI_FORMAT)GPixelFormats[DestTexture->GetFormat()].PlatformFormat
 					);
 
-				CommandListHandle.UpdateResidency(SourceTextureCube->GetResource());
-				CommandListHandle.UpdateResidency(DestTextureCube->GetResource());
+				CommandListHandle.UpdateResidency(SourceTexture->GetResource());
+				CommandListHandle.UpdateResidency(DestTexture->GetResource());
 			}
 			else
 			{
-				CD3DX12_TEXTURE_COPY_LOCATION DestCopyLocation(DestTextureCube->GetResource()->GetResource(), DestSubresource);
-				CD3DX12_TEXTURE_COPY_LOCATION SourceCopyLocation(SourceTextureCube->GetResource()->GetResource(), SourceSubresource);
+				CD3DX12_TEXTURE_COPY_LOCATION DestCopyLocation(DestTexture->GetResource()->GetResource(), DestSubresource);
+				CD3DX12_TEXTURE_COPY_LOCATION SourceCopyLocation(SourceTexture->GetResource()->GetResource(), SourceSubresource);
 
-				FScopedResourceBarrier ConditionalScopeResourceBarrierDest(CommandListHandle, DestTextureCube->GetResource(), D3D12_RESOURCE_STATE_COPY_DEST, DestCopyLocation.SubresourceIndex, FD3D12DynamicRHI::ETransitionMode::Validate);
-				FScopedResourceBarrier ConditionalScopeResourceBarrierSource(CommandListHandle, SourceTextureCube->GetResource(), D3D12_RESOURCE_STATE_COPY_SOURCE, SourceCopyLocation.SubresourceIndex, FD3D12DynamicRHI::ETransitionMode::Validate);
+				FScopedResourceBarrier ConditionalScopeResourceBarrierDest(CommandListHandle, DestTexture->GetResource(), D3D12_RESOURCE_STATE_COPY_DEST, DestCopyLocation.SubresourceIndex, FD3D12DynamicRHI::ETransitionMode::Validate);
+				FScopedResourceBarrier ConditionalScopeResourceBarrierSource(CommandListHandle, SourceTexture->GetResource(), D3D12_RESOURCE_STATE_COPY_SOURCE, SourceCopyLocation.SubresourceIndex, FD3D12DynamicRHI::ETransitionMode::Validate);
 
 				numCopies++;
 				CommandListHandle.FlushResourceBarriers();
@@ -445,23 +425,23 @@ void FD3D12CommandContext::RHICopyToResolveTarget(FRHITexture* SourceTextureRHI,
 					&SourceCopyLocation,
 					nullptr);
 
-				CommandListHandle.UpdateResidency(SourceTextureCube->GetResource());
-				CommandListHandle.UpdateResidency(DestTextureCube->GetResource());
+				CommandListHandle.UpdateResidency(SourceTexture->GetResource());
+				CommandListHandle.UpdateResidency(DestTexture->GetResource());
 			}
 		}
 	}
-	else if (SourceTexture2D && DestTextureCube)
+	else if (SourceDimension == ETextureDimension::Texture2D && DestDimension == ETextureDimension::TextureCube)
 	{
 		// If source is 2D and Dest is a cube then copy the 2D texture to the specified cube face.
 		// Determine the cubemap face being resolved.
 		const uint32 D3DFace = GetD3D12CubeFace(ResolveParams.CubeFace);
 		const uint32 Subresource = CalcSubresource(0, D3DFace, 1);
 
-		CD3DX12_TEXTURE_COPY_LOCATION DestCopyLocation(DestTextureCube->GetResource()->GetResource(), Subresource);
-		CD3DX12_TEXTURE_COPY_LOCATION SourceCopyLocation(SourceTexture2D->GetResource()->GetResource(), 0);
+		CD3DX12_TEXTURE_COPY_LOCATION DestCopyLocation(DestTexture->GetResource()->GetResource(), Subresource);
+		CD3DX12_TEXTURE_COPY_LOCATION SourceCopyLocation(SourceTexture->GetResource()->GetResource(), 0);
 
-		FScopedResourceBarrier ConditionalScopeResourceBarrierDest(CommandListHandle, DestTextureCube->GetResource(), D3D12_RESOURCE_STATE_COPY_DEST, DestCopyLocation.SubresourceIndex, FD3D12DynamicRHI::ETransitionMode::Validate);
-		FScopedResourceBarrier ConditionalScopeResourceBarrierSource(CommandListHandle, SourceTexture2D->GetResource(), D3D12_RESOURCE_STATE_COPY_SOURCE, SourceCopyLocation.SubresourceIndex, FD3D12DynamicRHI::ETransitionMode::Validate);
+		FScopedResourceBarrier ConditionalScopeResourceBarrierDest(CommandListHandle, DestTexture->GetResource(), D3D12_RESOURCE_STATE_COPY_DEST, DestCopyLocation.SubresourceIndex, FD3D12DynamicRHI::ETransitionMode::Validate);
+		FScopedResourceBarrier ConditionalScopeResourceBarrierSource(CommandListHandle, SourceTexture->GetResource(), D3D12_RESOURCE_STATE_COPY_SOURCE, SourceCopyLocation.SubresourceIndex, FD3D12DynamicRHI::ETransitionMode::Validate);
 
 		numCopies++;
 		CommandListHandle.FlushResourceBarriers();
@@ -471,14 +451,15 @@ void FD3D12CommandContext::RHICopyToResolveTarget(FRHITexture* SourceTextureRHI,
 			&SourceCopyLocation,
 			nullptr);
 
-		CommandListHandle.UpdateResidency(SourceTexture2D->GetResource());
-		CommandListHandle.UpdateResidency(DestTextureCube->GetResource());
+		CommandListHandle.UpdateResidency(SourceTexture->GetResource());
+		CommandListHandle.UpdateResidency(DestTexture->GetResource());
 	}
-	else if (SourceTexture3D && DestTexture3D)
+
+	else if (SourceDimension == ETextureDimension::Texture3D && DestDimension == ETextureDimension::Texture3D)
 	{
 		// bit of a hack.  no one resolves slice by slice and 0 is the default value.  assume for the moment they are resolving the whole texture.
 		check(ResolveParams.SourceArrayIndex == 0);
-		check(SourceTexture3D == DestTexture3D);
+		check(SourceTexture == DestTexture);
 	}
 
 	// Force transtion to readable since RHI user expects both source and dest to be in readable state after calling RHICopyToResolveTarget
@@ -593,12 +574,11 @@ TRefCountPtr<FD3D12Resource> FD3D12DynamicRHI::GetStagingTexture(FRHITexture* Te
 	const FRHIGPUMask Node = Device->GetGPUMask();
 
 	FD3D12CommandListHandle& hCommandList = Device->GetDefaultCommandContext().CommandListHandle;
-	FD3D12TextureBase* Texture = GetD3D12TextureFromRHITexture(TextureRHI, GPUIndex);
+	FD3D12Texture* Texture = GetD3D12TextureFromRHITexture(TextureRHI, GPUIndex);
 	D3D12_RESOURCE_DESC const& SourceDesc = Texture->GetResource()->GetDesc();
 
 	// Ensure we're dealing with a Texture2D, which the rest of this function already assumes
 	check(TextureRHI->GetTexture2D());
-	FD3D12Texture2D* InTexture2D = static_cast<FD3D12Texture2D*>(Texture);
 
 	bool bRequiresTempStagingTexture = Texture->GetResource()->GetHeapType() != D3D12_HEAP_TYPE_READBACK;
 	if (bRequiresTempStagingTexture == false)
@@ -609,7 +589,7 @@ TRefCountPtr<FD3D12Resource> FD3D12DynamicRHI::GetStagingTexture(FRHITexture* Te
 
 		// Texture2Ds on the readback heap will have been flattened to 1D, so we need to retrieve pitch
 		// information from the original 2D version to correctly use sub-rects.
-		InTexture2D->GetReadBackHeapDesc(readbackHeapDesc, InFlags.GetMip());
+		Texture->GetReadBackHeapDesc(readbackHeapDesc, InFlags.GetMip());
 		StagingRectOUT = InRect;
 
 		return (Texture->GetResource());
@@ -641,7 +621,7 @@ TRefCountPtr<FD3D12Resource> FD3D12DynamicRHI::GetStagingTexture(FRHITexture* Te
 
 	// Copy the data to a staging resource.
 	uint32 Subresource = 0;
-	if (InTexture2D->IsCubemap())
+	if (Texture->GetDesc().IsTextureCube())
 	{
 		uint32 D3DFace = GetD3D12CubeFace(InFlags.GetCubeFace());
 		Subresource = CalcSubresource(InFlags.GetMip(), D3DFace, TextureRHI->GetNumMips());
@@ -699,7 +679,7 @@ TRefCountPtr<FD3D12Resource> FD3D12DynamicRHI::GetStagingTexture(FRHITexture* Te
 void FD3D12DynamicRHI::ReadSurfaceDataNoMSAARaw(FRHITexture* TextureRHI, FIntRect InRect, TArray<uint8>& OutData, FReadSurfaceDataFlags InFlags)
 {
 	const uint32 GPUIndex = InFlags.GetGPUIndex();
-	FD3D12TextureBase* Texture = GetD3D12TextureFromRHITexture(TextureRHI, GPUIndex);
+	FD3D12Texture* Texture = GetD3D12TextureFromRHITexture(TextureRHI, GPUIndex);
 
 	const uint32 SizeX = InRect.Width();
 	const uint32 SizeY = InRect.Height();
@@ -1109,7 +1089,7 @@ void FD3D12DynamicRHI::RHIReadSurfaceData(FRHITexture* TextureRHI, FIntRect InRe
 {
 	TArray<uint8> OutDataRaw;
 
-	FD3D12TextureBase* Texture = GetD3D12TextureFromRHITexture(TextureRHI, InFlags.GetGPUIndex());
+	FD3D12Texture* Texture = GetD3D12TextureFromRHITexture(TextureRHI, InFlags.GetGPUIndex());
 
 	// Check the format of the surface
 	D3D12_RESOURCE_DESC const& TextureDesc = Texture->GetResource()->GetDesc();
@@ -1164,10 +1144,9 @@ void FD3D12DynamicRHI::RHIReadSurfaceData(FRHITexture* InRHITexture, FIntRect In
 
 	// Retrieve the base texture
 	FD3D12CommandContext& CommandContext = GetRHIDevice(GPUIndex)->GetDefaultCommandContext();
-	FD3D12TextureBase* D3D12TextureBase = CommandContext.RetrieveTextureBase(InRHITexture);
+	FD3D12Texture* DestTexture2D = CommandContext.RetrieveTexture(InRHITexture);
 
 	// Wait for the command list if needed
-	FD3D12Texture2D* DestTexture2D = static_cast<FD3D12Texture2D*>(D3D12TextureBase);
 	FD3D12CLSyncPoint SyncPoint = DestTexture2D->GetReadBackSyncPoint();
 
 	if (!!SyncPoint)
@@ -1215,7 +1194,7 @@ void FD3D12DynamicRHI::RHIReadSurfaceData(FRHITexture* InRHITexture, FIntRect In
 }
 
 void FD3D12DynamicRHI::ReadSurfaceDataMSAARaw(FRHICommandList_RecursiveHazardous& RHICmdList, FRHITexture* TextureRHI, FIntRect InRect, TArray<uint8>& OutData, FReadSurfaceDataFlags InFlags)
-{
+{	
 	const uint32 GPUIndex = InFlags.GetGPUIndex();
 	FD3D12Device* Device = GetRHIDevice(GPUIndex);
 	FD3D12Adapter* Adapter = Device->GetParentAdapter();
@@ -1223,7 +1202,7 @@ void FD3D12DynamicRHI::ReadSurfaceDataMSAARaw(FRHICommandList_RecursiveHazardous
 
 	FD3D12CommandContext& DefaultContext = Device->GetDefaultCommandContext();
 	FD3D12CommandListHandle& hCommandList = DefaultContext.CommandListHandle;
-	FD3D12TextureBase* Texture = GetD3D12TextureFromRHITexture(TextureRHI, GPUIndex);
+	FD3D12Texture* Texture = GetD3D12TextureFromRHITexture(TextureRHI, GPUIndex);
 
 	const uint32 SizeX = InRect.Width();
 	const uint32 SizeY = InRect.Height();
@@ -1284,11 +1263,10 @@ void FD3D12DynamicRHI::ReadSurfaceDataMSAARaw(FRHICommandList_RecursiveHazardous
 
 	// Ensure we're dealing with a Texture2D, which the rest of this function already assumes
 	check(TextureRHI->GetTexture2D());
-	FD3D12Texture2D* InTexture2D = static_cast<FD3D12Texture2D*>(Texture);
 
 	// Determine the subresource index for cubemaps.
 	uint32 Subresource = 0;
-	if (InTexture2D->IsCubemap())
+	if (Texture->GetDesc().IsTextureCube())
 	{
 		uint32 D3DFace = GetD3D12CubeFace(InFlags.GetCubeFace());
 		Subresource = CalcSubresource(InFlags.GetMip(), D3DFace, TextureRHI->GetNumMips());
@@ -1377,8 +1355,8 @@ void FD3D12DynamicRHI::ReadSurfaceDataMSAARaw(FRHICommandList_RecursiveHazardous
 }
 
 void FD3D12DynamicRHI::RHIMapStagingSurface(FRHITexture* TextureRHI, FRHIGPUFence* FenceRHI, void*& OutData, int32& OutWidth, int32& OutHeight, uint32 GPUIndex)
-{
-	FD3D12Texture2D* DestTexture2D = ResourceCast(TextureRHI->GetTexture2D(), GPUIndex);
+{	
+	FD3D12Texture* DestTexture2D = ResourceCast(TextureRHI->GetTexture2D(), GPUIndex);
 
 	check(DestTexture2D);
 	FD3D12Resource* Texture = DestTexture2D->GetResource();
@@ -1438,8 +1416,8 @@ void FD3D12DynamicRHI::RHIMapStagingSurface(FRHITexture* TextureRHI, FRHIGPUFenc
 }
 
 void FD3D12DynamicRHI::RHIUnmapStagingSurface(FRHITexture* TextureRHI, uint32 GPUIndex)
-{
-	FD3D12Texture2D* DestTexture2D = ResourceCast(TextureRHI->GetTexture2D(), GPUIndex);
+{	
+	FD3D12Texture* DestTexture2D = ResourceCast(TextureRHI->GetTexture2D(), GPUIndex);
 
 	check(DestTexture2D);
 	ID3D12Resource* Texture = DestTexture2D->GetResource()->GetResource();
@@ -1464,7 +1442,7 @@ void FD3D12DynamicRHI::RHIReadSurfaceFloatData(FRHITexture* TextureRHI, FIntRect
 
 	FD3D12CommandContext& DefaultContext = Device->GetDefaultCommandContext();
 	FD3D12CommandListHandle& hCommandList = DefaultContext.CommandListHandle;
-	FD3D12TextureBase* Texture = GetD3D12TextureFromRHITexture(TextureRHI, GPUIndex);
+	FD3D12Texture* Texture = GetD3D12TextureFromRHITexture(TextureRHI, GPUIndex);
 
 	uint32 SizeX = InRect.Width();
 	uint32 SizeY = InRect.Height();
@@ -1494,29 +1472,8 @@ void FD3D12DynamicRHI::RHIReadSurfaceFloatData(FRHITexture* TextureRHI, FIntRect
 	VERIFYD3D12RESULT(Adapter->CreateBuffer(D3D12_HEAP_TYPE_READBACK, Node, Node, MipBytesAligned, TempTexture2D.GetInitReference(), nullptr));
 
 	// Ensure we're dealing with a Texture2D, which the rest of this function already assumes
-	bool bIsTextureCube = false;
-	check(TextureRHI->GetTexture2D() || TextureRHI->GetTexture2DArray() || TextureRHI->GetTextureCube());
-	FD3D12Texture2D* InTexture2D = static_cast<FD3D12Texture2D*>(Texture);
-	FD3D12Texture2DArray* InTexture2DArray = static_cast<FD3D12Texture2DArray*>(Texture);
-	FD3D12TextureCube* InTextureCube = static_cast<FD3D12TextureCube*>(Texture);
-	if (InTexture2D)
-	{
-		bIsTextureCube = InTexture2D->IsCubemap();
-	}
-	else if (InTexture2DArray)
-	{
-		bIsTextureCube = InTexture2DArray->IsCubemap();
-	}
-	else if (InTextureCube)
-	{
-		bIsTextureCube = InTextureCube->IsCubemap();
-		check(bIsTextureCube);
-	}
-	else
-	{
-		check(false);
-	}
-
+	bool bIsTextureCube = Texture->GetDesc().IsTextureCube();
+	
 	// Copy the data to a staging resource.
 	uint32 Subresource = 0;
 	if (bIsTextureCube)
@@ -1526,7 +1483,7 @@ void FD3D12DynamicRHI::RHIReadSurfaceFloatData(FRHITexture* TextureRHI, FIntRect
 	}
 	else
 	{
-		const bool bIsTextureArray = InTexture2DArray != nullptr;
+		const bool bIsTextureArray = Texture->GetDesc().IsTextureArray();
 		Subresource = CalcSubresource(InFlags.GetMip(), bIsTextureArray ? InFlags.GetArrayIndex() : 0, TextureDesc.MipLevels);
 	}
 
@@ -1603,7 +1560,7 @@ void FD3D12DynamicRHI::RHIRead3DSurfaceFloatData(FRHITexture* TextureRHI, FIntRe
 
 	FD3D12CommandContext& DefaultContext = Device->GetDefaultCommandContext();
 	FD3D12CommandListHandle& hCommandList = DefaultContext.CommandListHandle;
-	FD3D12TextureBase* Texture = GetD3D12TextureFromRHITexture(TextureRHI, GPUIndex);
+	FD3D12Texture* Texture = GetD3D12TextureFromRHITexture(TextureRHI, GPUIndex);
 
 	uint32 SizeX = InRect.Width();
 	uint32 SizeY = InRect.Height();
diff --git a/Engine/Source/Runtime/D3D12RHI/Private/D3D12Resources.cpp b/Engine/Source/Runtime/D3D12RHI/Private/D3D12Resources.cpp
--- a/Engine/Source/Runtime/D3D12RHI/Private/D3D12Resources.cpp
+++ b/Engine/Source/Runtime/D3D12RHI/Private/D3D12Resources.cpp
@@ -313,7 +313,7 @@ void ID3D12ResourceAllocator::AllocateTexture(uint32 GPUIndex, D3D12_HEAP_TYPE I
 {
 	// Check if texture can be 4K aligned
 	FD3D12ResourceDesc Desc = InDesc;
-	bool b4KAligment = TextureCanBe4KAligned(Desc, InUEFormat);
+	bool b4KAligment = FD3D12Texture::CanBe4KAligned(Desc, InUEFormat);
 	Desc.Alignment = b4KAligment ? D3D12_SMALL_RESOURCE_PLACEMENT_ALIGNMENT : D3D12_DEFAULT_RESOURCE_PLACEMENT_ALIGNMENT;
 
 	// Get the size and alignment for the allocation
@@ -739,6 +739,54 @@ HRESULT FD3D12Adapter::CreateBuffer(const D3D12_HEAP_PROPERTIES& HeapProps,
 		ppOutResource, Name);
 }
 
+void FD3D12Adapter::CreateUAVAliasResource(D3D12_CLEAR_VALUE* ClearValuePtr, const TCHAR* DebugName, FD3D12ResourceLocation& Location)
+{
+	FD3D12Resource* SourceResource = Location.GetResource();
+
+	const FD3D12ResourceDesc& SourceDesc = SourceResource->GetDesc();
+	const FD3D12Heap* const ResourceHeap = SourceResource->GetHeap();
+
+	const EPixelFormat SourceFormat = SourceDesc.PixelFormat;
+	const EPixelFormat AliasTextureFormat = SourceDesc.UAVAliasPixelFormat;
+
+	if (ensure(ResourceHeap != nullptr) && ensure(SourceFormat != PF_Unknown) && SourceFormat != AliasTextureFormat)
+	{
+		const uint64 SourceOffset = Location.GetOffsetFromBaseOfResource();
+
+		FD3D12ResourceDesc AliasTextureDesc = SourceDesc;
+		AliasTextureDesc.Format = (DXGI_FORMAT)GPixelFormats[AliasTextureFormat].PlatformFormat;
+		AliasTextureDesc.Width = SourceDesc.Width / GPixelFormats[SourceFormat].BlockSizeX;
+		AliasTextureDesc.Height = SourceDesc.Height / GPixelFormats[SourceFormat].BlockSizeY;
+		// layout of UAV must match source resource
+		AliasTextureDesc.Layout = SourceResource->GetResource()->GetDesc().Layout;
+
+		EnumAddFlags(AliasTextureDesc.Flags, D3D12_RESOURCE_FLAG_ALLOW_UNORDERED_ACCESS);
+
+		AliasTextureDesc.UAVAliasPixelFormat = PF_Unknown;
+
+		TRefCountPtr<ID3D12Resource> pAliasResource;
+		HRESULT AliasHR = GetD3DDevice()->CreatePlacedResource(
+			ResourceHeap->GetHeap(),
+			SourceOffset,
+			&AliasTextureDesc,
+			D3D12_RESOURCE_STATE_UNORDERED_ACCESS,
+			ClearValuePtr,
+			IID_PPV_ARGS(pAliasResource.GetInitReference()));
+
+		if (pAliasResource && DebugName)
+		{
+			TCHAR NameBuffer[512]{};
+			FCString::Snprintf(NameBuffer, UE_ARRAY_COUNT(NameBuffer), TEXT("%s UAVAlias"), DebugName);
+			SetName(pAliasResource, NameBuffer);
+		}
+
+		if (SUCCEEDED(AliasHR))
+		{
+			SourceResource->SetUAVAccessResource(pAliasResource);
+		}
+	}
+}
+
 /////////////////////////////////////////////////////////////////////
 //	FD3D12 Resource Location
 /////////////////////////////////////////////////////////////////////
