Copying //UE5/Dev-ParallelRendering to //UE5/Main

diff --git a/Engine/Source/Runtime/D3D12RHI/Private/D3D12Texture.cpp b/Engine/Source/Runtime/D3D12RHI/Private/D3D12Texture.cpp
--- a/Engine/Source/Runtime/D3D12RHI/Private/D3D12Texture.cpp
+++ b/Engine/Source/Runtime/D3D12RHI/Private/D3D12Texture.cpp
@@ -45,53 +45,33 @@ static TAutoConsoleVariable<bool> CVarTexturePoolOnlyAccountStreamableTexture(
 	ECVF_ReadOnly
 );
 
-// Forward Decls for template types
-template TD3D12Texture2D<FD3D12BaseTexture2D>::TD3D12Texture2D(class FD3D12Device* InParent, uint32 InSizeX, uint32 InSizeY, uint32 InSizeZ, uint32 InNumMips, uint32 InNumSamples,
-	EPixelFormat InFormat, bool bInCubemap, ETextureCreateFlags InFlags, const FClearValueBinding& InClearValue, const FD3D12TextureLayout* InTextureLayout
-#if PLATFORM_SUPPORTS_VIRTUAL_TEXTURES
-	, void* InRawTextureMemory
-#endif
-	);
-template TD3D12Texture2D<FD3D12BaseTexture2DArray>::TD3D12Texture2D(class FD3D12Device* InParent, uint32 InSizeX, uint32 InSizeY, uint32 InSizeZ, uint32 InNumMips, uint32 InNumSamples,
-	EPixelFormat InFormat, bool bInCubemap, ETextureCreateFlags InFlags, const FClearValueBinding& InClearValue, const FD3D12TextureLayout* InTextureLayout
-#if PLATFORM_SUPPORTS_VIRTUAL_TEXTURES
-	, void* InRawTextureMemory
-#endif
-	);
-template TD3D12Texture2D<FD3D12BaseTextureCube>::TD3D12Texture2D(class FD3D12Device* InParent, uint32 InSizeX, uint32 InSizeY, uint32 InSizeZ, uint32 InNumMips, uint32 InNumSamples,
-	EPixelFormat InFormat, bool bInCubemap, ETextureCreateFlags InFlags, const FClearValueBinding& InClearValue, const FD3D12TextureLayout* InTextureLayout
-#if PLATFORM_SUPPORTS_VIRTUAL_TEXTURES
-	, void* InRawTextureMemory
-#endif
-	);
-
-template TD3D12Texture2D<FD3D12BaseTexture2D>::~TD3D12Texture2D();
-template TD3D12Texture2D<FD3D12BaseTexture2DArray>::~TD3D12Texture2D();
-template TD3D12Texture2D<FD3D12BaseTextureCube>::~TD3D12Texture2D();
-
 
-/// @cond DOXYGEN_WARNINGS
+///////////////////////////////////////////////////////////////////////////////////////////
+// Texture Commands
+///////////////////////////////////////////////////////////////////////////////////////////
 
-template void FD3D12TextureStats::D3D12TextureAllocated(TD3D12Texture2D<FD3D12BaseTexture2D>& Texture, const D3D12_RESOURCE_DESC *Desc);
-template void FD3D12TextureStats::D3D12TextureAllocated(TD3D12Texture2D<FD3D12BaseTexture2DArray>& Texture, const D3D12_RESOURCE_DESC *Desc);
-template void FD3D12TextureStats::D3D12TextureAllocated(TD3D12Texture2D<FD3D12BaseTextureCube>& Texture, const D3D12_RESOURCE_DESC *Desc);
+static bool ShouldDeferCmdListOperation(FRHICommandList* RHICmdList)
+{
+	if (RHICmdList == nullptr)
+	{
+		return false;
+	}
 
-template void FD3D12TextureStats::D3D12TextureDeleted(TD3D12Texture2D<FD3D12BaseTexture2D>& Texture);
-template void FD3D12TextureStats::D3D12TextureDeleted(TD3D12Texture2D<FD3D12BaseTexture2DArray>& Texture);
-template void FD3D12TextureStats::D3D12TextureDeleted(TD3D12Texture2D<FD3D12BaseTextureCube>& Texture);
+	if (RHICmdList->Bypass() || !IsRunningRHIInSeparateThread())
+	{
+		return false;
+	}
 
-template void TD3D12Texture2D<FD3D12BaseTexture2D>::GetReadBackHeapDesc(D3D12_PLACED_SUBRESOURCE_FOOTPRINT& OutFootprint, uint32 Subresource) const;
-template void TD3D12Texture2D<FD3D12BaseTexture2DArray>::GetReadBackHeapDesc(D3D12_PLACED_SUBRESOURCE_FOOTPRINT& OutFootprint, uint32 Subresource) const;
-template void TD3D12Texture2D<FD3D12BaseTextureCube>::GetReadBackHeapDesc(D3D12_PLACED_SUBRESOURCE_FOOTPRINT& OutFootprint, uint32 Subresource) const;
+	return true;
+}
 
-/// @endcond
 struct FRHICommandUpdateTextureString
 {
 	static const TCHAR* TStr() { return TEXT("FRHICommandUpdateTexture"); }
 };
 struct FRHICommandUpdateTexture final : public FRHICommand<FRHICommandUpdateTexture, FRHICommandUpdateTextureString>
 {
-	FD3D12TextureBase* TextureBase;
+	FD3D12Texture* Texture;
 	uint32 MipIndex;
 	uint32 DestX;
 	uint32 DestY;
@@ -99,10 +79,10 @@ struct FRHICommandUpdateTexture final : public FRHICommand<FRHICommandUpdateText
 	D3D12_TEXTURE_COPY_LOCATION SourceCopyLocation;
 	FD3D12ResourceLocation Source;
 
-	FORCEINLINE_DEBUGGABLE FRHICommandUpdateTexture(FD3D12TextureBase* InTextureBase,
+	FORCEINLINE_DEBUGGABLE FRHICommandUpdateTexture(FD3D12Texture* InTexture,
 		uint32 InMipIndex, uint32 InDestX, uint32 InDestY, uint32 InDestZ,
 		const D3D12_TEXTURE_COPY_LOCATION& InSourceCopyLocation, FD3D12ResourceLocation* InSource)
-		: TextureBase(InTextureBase)
+		: Texture(InTexture)
 		, MipIndex(InMipIndex)
 		, DestX(InDestX)
 		, DestY(InDestY)
@@ -122,7 +102,7 @@ struct FRHICommandUpdateTexture final : public FRHICommand<FRHICommandUpdateText
 
 	void Execute(FRHICommandListBase& CmdList)
 	{
-		TextureBase->UpdateTexture(MipIndex, DestX, DestY, DestZ, SourceCopyLocation);
+		Texture->UpdateTexture(MipIndex, DestX, DestY, DestZ, SourceCopyLocation);
 	}
 };
 
@@ -132,13 +112,13 @@ struct FD3D12RHICommandInitializeTextureString
 };
 struct FD3D12RHICommandInitializeTexture final : public FRHICommand<FD3D12RHICommandInitializeTexture, FD3D12RHICommandInitializeTextureString>
 {
-	FD3D12TextureBase* TextureBase;
+	FD3D12Texture* Texture;
 	FD3D12ResourceLocation SrcResourceLoc;
 	uint32 NumSubresources;
 	D3D12_RESOURCE_STATES DestinationState;
 
-	FORCEINLINE_DEBUGGABLE FD3D12RHICommandInitializeTexture(FD3D12TextureBase* InTexture, FD3D12ResourceLocation& InSrcResourceLoc, uint32 InNumSubresources, D3D12_RESOURCE_STATES InDestinationState)
-		: TextureBase(InTexture)
+	FORCEINLINE_DEBUGGABLE FD3D12RHICommandInitializeTexture(FD3D12Texture* InTexture, FD3D12ResourceLocation& InSrcResourceLoc, uint32 InNumSubresources, D3D12_RESOURCE_STATES InDestinationState)
+		: Texture(InTexture)
 		, SrcResourceLoc(InSrcResourceLoc.GetParentDevice())
 		, NumSubresources(InNumSubresources)
 		, DestinationState(InDestinationState)
@@ -163,15 +143,15 @@ struct FD3D12RHICommandInitializeTexture final : public FRHICommand<FD3D12RHICom
 		UINT64* RowSizeInBytes = (UINT64*) (Rows + NumSubresources);
 
 		uint64 Size = 0;
-		const D3D12_RESOURCE_DESC& Desc = TextureBase->GetResource()->GetDesc();
-		TextureBase->GetParentDevice()->GetDevice()->GetCopyableFootprints(&Desc, 0, NumSubresources, SrcResourceLoc.GetOffsetFromBaseOfResource(), Footprints, Rows, RowSizeInBytes, &Size);
+		const D3D12_RESOURCE_DESC& Desc = Texture->GetResource()->GetDesc();
+		Texture->GetParentDevice()->GetDevice()->GetCopyableFootprints(&Desc, 0, NumSubresources, SrcResourceLoc.GetOffsetFromBaseOfResource(), Footprints, Rows, RowSizeInBytes, &Size);
 
 		D3D12_TEXTURE_COPY_LOCATION Src;
 		Src.pResource = SrcResourceLoc.GetResource()->GetResource();
 		Src.Type = D3D12_TEXTURE_COPY_TYPE_PLACED_FOOTPRINT;
 
 		// Initialize all the textures in the chain
-		for (FD3D12TextureBase& CurrentTexture : *TextureBase)
+		for (FD3D12Texture& CurrentTexture : *Texture)
 		{
 			FD3D12Device* Device = CurrentTexture.GetParentDevice();
 			FD3D12Resource* Resource = CurrentTexture.GetResource();
@@ -225,18 +205,105 @@ struct FD3D12RHICommandInitializeTexture final : public FRHICommand<FD3D12RHICom
 	}
 };
 
+struct FD3D12RHICommandAsyncReallocateTexture2DString
+{
+	static const TCHAR* TStr() { return TEXT("FD3D12RHICommandAsyncReallocateTexture2D"); }
+};
+struct FRHICommandD3D12AsyncReallocateTexture2D final : public FRHICommand<FRHICommandD3D12AsyncReallocateTexture2D, FD3D12RHICommandAsyncReallocateTexture2DString>
+{
+	FD3D12Texture* OldTexture;
+	FD3D12Texture* NewTexture;
+	int32 NewMipCount;
+	int32 NewSizeX;
+	int32 NewSizeY;
+	FThreadSafeCounter* RequestStatus;
+
+	FORCEINLINE_DEBUGGABLE FRHICommandD3D12AsyncReallocateTexture2D(FD3D12Texture* InOldTexture, FD3D12Texture* InNewTexture, int32 InNewMipCount, int32 InNewSizeX, int32 InNewSizeY, FThreadSafeCounter* InRequestStatus)
+		: OldTexture(InOldTexture)
+		, NewTexture(InNewTexture)
+		, NewMipCount(InNewMipCount)
+		, NewSizeX(InNewSizeX)
+		, NewSizeY(InNewSizeY)
+		, RequestStatus(InRequestStatus)
+	{
+	}
+
+	void Execute(FRHICommandListBase& RHICmdList)
+	{
+		CopyMips();
+	}
+
+	void CopyMips()
+	{
+		// Use the GPU to asynchronously copy the old mip-maps into the new texture.
+		const uint32 NumSharedMips = FMath::Min(OldTexture->GetNumMips(), NewTexture->GetNumMips());
+		const uint32 SourceMipOffset = OldTexture->GetNumMips() - NumSharedMips;
+		const uint32 DestMipOffset = NewTexture->GetNumMips() - NumSharedMips;
+
+		uint32 destSubresource = 0;
+		uint32 srcSubresource = 0;
+
+		for (FD3D12Texture::FDualLinkedObjectIterator It(OldTexture, NewTexture); It; ++It)
+		{
+			OldTexture = static_cast<FD3D12Texture*>(It.GetFirst());
+			NewTexture = static_cast<FD3D12Texture*>(It.GetSecond());
+
+			FD3D12Device* Device = OldTexture->GetParentDevice();
+
+			FD3D12CommandListHandle& hCommandList = Device->GetDefaultCommandContext().CommandListHandle;
+
+			FScopeResourceBarrier ScopeResourceBarrierDest(hCommandList, NewTexture->GetResource(), NewTexture->GetResource()->GetDefaultResourceState(), D3D12_RESOURCE_STATE_COPY_DEST, D3D12_RESOURCE_BARRIER_ALL_SUBRESOURCES);
+			FScopeResourceBarrier ScopeResourceBarrierSource(hCommandList, OldTexture->GetResource(), OldTexture->GetResource()->GetDefaultResourceState(), D3D12_RESOURCE_STATE_COPY_SOURCE, D3D12_RESOURCE_BARRIER_ALL_SUBRESOURCES);
+			hCommandList.FlushResourceBarriers();	// Must flush so the desired state is actually set.
+
+			for (uint32 MipIndex = 0; MipIndex < NumSharedMips; ++MipIndex)
+			{
+				// Use the GPU to copy between mip-maps.
+				// This is serialized with other D3D commands, so it isn't necessary to increment Counter to signal a pending asynchronous copy.
+
+				srcSubresource = CalcSubresource(MipIndex + SourceMipOffset, 0, OldTexture->GetNumMips());
+				destSubresource = CalcSubresource(MipIndex + DestMipOffset, 0, NewTexture->GetNumMips());
+
+				CD3DX12_TEXTURE_COPY_LOCATION DestCopyLocation(NewTexture->GetResource()->GetResource(), destSubresource);
+				CD3DX12_TEXTURE_COPY_LOCATION SourceCopyLocation(OldTexture->GetResource()->GetResource(), srcSubresource);
+
+				Device->GetDefaultCommandContext().numCopies++;
+				hCommandList->CopyTextureRegion(
+					&DestCopyLocation,
+					0, 0, 0,
+					&SourceCopyLocation,
+					nullptr);
+
+				hCommandList.UpdateResidency(NewTexture->GetResource());
+				hCommandList.UpdateResidency(OldTexture->GetResource());
+
+				Device->GetDefaultCommandContext().ConditionalFlushCommandList();
+
+				DEBUG_EXECUTE_COMMAND_CONTEXT(Device->GetDefaultCommandContext());
+			}
+		}
+
+		// Decrement the thread-safe counter used to track the completion of the reallocation, since D3D handles sequencing the
+		// async mip copies with other D3D calls.
+		RequestStatus->Decrement();
+	}
+};
+
+
+
 ///////////////////////////////////////////////////////////////////////////////////////////
 // Texture Stats
 ///////////////////////////////////////////////////////////////////////////////////////////
 
-bool FD3D12TextureStats::ShouldCountAsTextureMemory(D3D12_RESOURCE_FLAGS MiscFlags)
+
+static bool ShouldCountAsTextureMemory(D3D12_RESOURCE_FLAGS MiscFlags)
 {
 	// Shouldn't be used for DEPTH, RENDER TARGET, or UNORDERED ACCESS
 	return !EnumHasAnyFlags(MiscFlags, D3D12_RESOURCE_FLAG_ALLOW_DEPTH_STENCIL | D3D12_RESOURCE_FLAG_ALLOW_RENDER_TARGET | D3D12_RESOURCE_FLAG_ALLOW_UNORDERED_ACCESS);
 }
 
-// @param b3D true:3D, false:2D or cube map
-TStatId FD3D12TextureStats::GetRHIStatEnum(D3D12_RESOURCE_FLAGS MiscFlags, bool bCubeMap, bool b3D)
+
+static TStatId GetRHIStatEnum(D3D12_RESOURCE_FLAGS MiscFlags, bool bCubeMap, bool b3D)
 {
 #if STATS
 	if (ShouldCountAsTextureMemory(MiscFlags))
@@ -275,7 +342,8 @@ TStatId FD3D12TextureStats::GetRHIStatEnum(D3D12_RESOURCE_FLAGS MiscFlags, bool
 	return TStatId();
 }
 
-TStatId FD3D12TextureStats::GetD3D12StatEnum(D3D12_RESOURCE_FLAGS MiscFlags)
+
+static TStatId GetD3D12StatEnum(D3D12_RESOURCE_FLAGS MiscFlags)
 {
 #if STATS
 	if (EnumHasAnyFlags(MiscFlags, D3D12_RESOURCE_FLAG_ALLOW_RENDER_TARGET | D3D12_RESOURCE_FLAG_ALLOW_DEPTH_STENCIL))
@@ -294,15 +362,10 @@ TStatId FD3D12TextureStats::GetD3D12StatEnum(D3D12_RESOURCE_FLAGS MiscFlags)
 	return TStatId();
 }
 
-// Note: This function can be called from many different threads
-// @param TextureSize >0 to allocate, <0 to deallocate
-// @param b3D true:3D, false:2D or cube map
-template<typename TD3D12Texture>
-void FD3D12TextureStats::UpdateD3D12TextureStats(TD3D12Texture& Texture, const D3D12_RESOURCE_DESC& Desc, int64 TextureSize, bool b3D, bool bCubeMap, bool bStreamable, bool bNewTexture)
-{
 
+void FD3D12TextureStats::UpdateD3D12TextureStats(FD3D12Texture& Texture, const D3D12_RESOURCE_DESC& Desc, int64 TextureSize, bool b3D, bool bCubeMap, bool bStreamable, bool bNewTexture)
+{
 #if TEXTURE_PROFILER_ENABLED
-
 	if (!bNewTexture && 
 		!Texture.ResourceLocation.IsTransient() 
 		&& !EnumHasAnyFlags(Texture.GetFlags(), TexCreate_Virtual)
@@ -348,69 +411,64 @@ void FD3D12TextureStats::UpdateD3D12TextureStats(TD3D12Texture& Texture, const D
 	}
 }
 
-template<typename BaseResourceType>
-void FD3D12TextureStats::D3D12TextureAllocated(TD3D12Texture2D<BaseResourceType>& Texture, const D3D12_RESOURCE_DESC *Desc)
+
+void FD3D12TextureStats::D3D12TextureAllocated(FD3D12Texture& Texture, const D3D12_RESOURCE_DESC *Desc)
 {
-	FD3D12Resource* D3D12Texture2D = Texture.GetResource();
+	FD3D12Resource* D3D12Texture = Texture.GetResource();
 
-	if (D3D12Texture2D)
+	if (D3D12Texture)
 	{
+		if (!Desc)
+		{
+			Desc = &D3D12Texture->GetDesc();
+		}
+
 		// Don't update state for virtual or transient textures	
-		if (!EnumHasAnyFlags(Texture.Flags, TexCreate_Virtual) && !Texture.ResourceLocation.IsTransient())
+		ETextureCreateFlags CreateFlags = Texture.GetDesc().Flags;
+		if (!EnumHasAnyFlags(CreateFlags, TexCreate_Virtual | TexCreate_CPUReadback) && !Texture.ResourceLocation.IsTransient())
 		{
 			TRACE_CPUPROFILER_EVENT_SCOPE(D3D12RHI::UpdateTextureStats);
 
-			if (!Desc)
+			const int64 TextureSize = Texture.ResourceLocation.GetSize();			
+			bool bNewTexture = true;
+			bool bIsStreamable = EnumHasAnyFlags(Texture.GetDesc().Flags, ETextureCreateFlags::Streamable);
+			UpdateD3D12TextureStats(Texture, *Desc, TextureSize, Texture.GetDesc().IsTexture3D(), Texture.GetDesc().IsTextureCube(), bIsStreamable, bNewTexture);
+
+#if TEXTURE_PROFILER_ENABLED
+			if (Texture.ResourceLocation.GetType() != FD3D12ResourceLocation::ResourceLocationType::eAliased
+				&& Texture.ResourceLocation.GetType() != FD3D12ResourceLocation::ResourceLocationType::eHeapAliased)
 			{
-				Desc = &D3D12Texture2D->GetDesc();
+				size_t Size = Texture.ResourceLocation.GetSize();
+				uint32 Alignment = Desc->Alignment;
+				FTextureProfiler::Get()->AddTextureAllocation(&Texture, Size, Alignment, 0);
 			}
-
-			const D3D12_RESOURCE_ALLOCATION_INFO AllocationInfo = Texture.GetParentDevice()->GetDevice()->GetResourceAllocationInfo(0, 1, Desc);
-			const int64 TextureSize = AllocationInfo.SizeInBytes;
-
-			Texture.SetMemorySize(TextureSize);
-
-			UpdateD3D12TextureStats(Texture, *Desc, TextureSize, false, Texture.IsCubemap(), Texture.IsStreamable(), true);
-		}
-		else
-		{
-			Texture.SetMemorySize(Texture.ResourceLocation.GetSize());
+#endif
 		}
 
-#if TEXTURE_PROFILER_ENABLED
-		if (!EnumHasAnyFlags(Texture.GetFlags(), TexCreate_Virtual)
-			&& !Texture.ResourceLocation.IsTransient()
-			&& Texture.ResourceLocation.GetType() != FD3D12ResourceLocation::ResourceLocationType::eAliased
-			&& Texture.ResourceLocation.GetType() != FD3D12ResourceLocation::ResourceLocationType::eHeapAliased)
-		{
-			size_t Size = Texture.GetMemorySize();
-			uint32 Alignment = Desc->Alignment;
-			FTextureProfiler::Get()->AddTextureAllocation(&Texture, Size, Alignment, 0);
-		}
-#endif
 	}
 }
 
-template<typename BaseResourceType>
-void FD3D12TextureStats::D3D12TextureDeleted(TD3D12Texture2D<BaseResourceType>& Texture)
+
+void FD3D12TextureStats::D3D12TextureDeleted(FD3D12Texture& Texture)
 {
-	FD3D12Resource* D3D12Texture2D = Texture.GetResource();
+	FD3D12Resource* D3D12Texture = Texture.GetResource();
 
-	if (D3D12Texture2D)
+	if (D3D12Texture)
 	{
 		// Don't update state for transient textures	
-		if (!Texture.ResourceLocation.IsTransient())
+		ETextureCreateFlags CreateFlags = Texture.GetDesc().Flags;
+		if (!EnumHasAnyFlags(CreateFlags, TexCreate_Virtual | TexCreate_CPUReadback) && !Texture.ResourceLocation.IsTransient())
 		{
-			const D3D12_RESOURCE_DESC& Desc = D3D12Texture2D->GetDesc();
-			const int64 TextureSize = Texture.GetMemorySize();
-			ensure(TextureSize > 0 || EnumHasAnyFlags(Texture.Flags, TexCreate_Virtual) || Texture.GetAliasingSourceTexture() != nullptr);
+			const D3D12_RESOURCE_DESC& Desc = D3D12Texture->GetDesc();
+			const int64 TextureSize = Texture.ResourceLocation.GetSize();
+			ensure(TextureSize > 0 || Texture.GetAliasingSourceTexture() != nullptr);
 
-			UpdateD3D12TextureStats(Texture, Desc, -TextureSize, false, Texture.IsCubemap(), Texture.IsStreamable(), false);
+			bool bNewTexture = false;
+			bool bIsStreamable = EnumHasAnyFlags(Texture.GetDesc().Flags, ETextureCreateFlags::Streamable);
+			UpdateD3D12TextureStats(Texture, Desc, -TextureSize, Texture.GetDesc().IsTexture3D(), Texture.GetDesc().IsTextureCube(), bIsStreamable, bNewTexture);
 
 #if TEXTURE_PROFILER_ENABLED
-			if (!EnumHasAnyFlags(Texture.GetFlags(), TexCreate_Virtual)
-				&& !Texture.ResourceLocation.IsTransient()
-				&& Texture.ResourceLocation.GetType() != FD3D12ResourceLocation::ResourceLocationType::eAliased
+			if (Texture.ResourceLocation.GetType() != FD3D12ResourceLocation::ResourceLocationType::eAliased
 				&& Texture.ResourceLocation.GetType() != FD3D12ResourceLocation::ResourceLocationType::eHeapAliased)
 			{
 				FTextureProfiler::Get()->RemoveTextureAllocation(&Texture);
@@ -420,184 +478,175 @@ void FD3D12TextureStats::D3D12TextureDeleted(TD3D12Texture2D<BaseResourceType>&
 	}
 }
 
-void FD3D12TextureStats::D3D12TextureAllocated2D(FD3D12Texture2D& Texture)
+
+bool FD3D12Texture::CanBe4KAligned(const FD3D12ResourceDesc& Desc, EPixelFormat UEFormat)
 {
-	D3D12TextureAllocated(Texture);
+	// 4KB alignment is only available for read only textures
+	if (!EnumHasAnyFlags(Desc.Flags, D3D12_RESOURCE_FLAG_ALLOW_RENDER_TARGET | D3D12_RESOURCE_FLAG_ALLOW_DEPTH_STENCIL | D3D12_RESOURCE_FLAG_ALLOW_UNORDERED_ACCESS) &&
+		!Desc.NeedsUAVAliasWorkarounds() && // UAV aliased resources are secretly writable.
+		Desc.SampleDesc.Count == 1)
+	{
+		D3D12_TILE_SHAPE Tile = {};
+		Get4KTileShape(&Tile, Desc.Format, UEFormat, Desc.Dimension, Desc.SampleDesc.Count);
+
+		uint32 TilesNeeded = GetTilesNeeded(Desc.Width, Desc.Height, Desc.DepthOrArraySize, Tile);
+
+		constexpr uint32 NUM_4K_BLOCKS_PER_64K_PAGE = 16;
+		return TilesNeeded <= NUM_4K_BLOCKS_PER_64K_PAGE;
+	}
+	else
+	{
+		return false;
+	}
 }
 
-void FD3D12TextureStats::D3D12TextureAllocated(FD3D12Texture3D& Texture)
+///////////////////////////////////////////////////////////////////////////////////////////
+// FD3D12DynamicRHI Texture functions
+///////////////////////////////////////////////////////////////////////////////////////////
+
+using namespace D3D12RHI;
+
+FD3D12ResourceDesc FD3D12DynamicRHI::GetResourceDesc(const FRHITextureDesc& TextureDesc) const
 {
-	FD3D12Resource* D3D12Texture3D = Texture.GetResource();
+	FD3D12ResourceDesc ResourceDesc;
 
-	if (D3D12Texture3D)
-	{
-		const D3D12_RESOURCE_DESC& Desc = D3D12Texture3D->GetDesc();
-		// Don't update state for virtual or transient textures	
-		if (!EnumHasAnyFlags(Texture.GetFlags(), TexCreate_Virtual) && !Texture.ResourceLocation.IsTransient())
-		{
-			const D3D12_RESOURCE_ALLOCATION_INFO AllocationInfo = Texture.GetParentDevice()->GetDevice()->GetResourceAllocationInfo(0, 1, &Desc);
-			const int64 TextureSize = AllocationInfo.SizeInBytes;
+	check(TextureDesc.Extent.X > 0 && TextureDesc.Extent.Y > 0 && TextureDesc.NumMips > 0);
 
-			Texture.SetMemorySize(TextureSize);
+	const DXGI_FORMAT PlatformResourceFormat = GetPlatformTextureResourceFormat((DXGI_FORMAT)GPixelFormats[TextureDesc.Format].PlatformFormat, TextureDesc.Flags);
 
-			UpdateD3D12TextureStats(Texture, Desc, TextureSize, true, false, Texture.IsStreamable(), true);
+	bool bDenyShaderResource = false;
+	if (TextureDesc.Dimension != ETextureDimension::Texture3D)
+	{
+		if (TextureDesc.IsTextureCube())
+		{
+			check(TextureDesc.Extent.X <= (int32)GetMaxCubeTextureDimension());
+			check(TextureDesc.Extent.X == TextureDesc.Extent.Y);
 		}
 		else
 		{
-			Texture.SetMemorySize(Texture.ResourceLocation.GetSize());
+			check(TextureDesc.Extent.X <= (int32)GetMax2DTextureDimension());
+			check(TextureDesc.Extent.Y <= (int32)GetMax2DTextureDimension());
 		}
 
-#if TEXTURE_PROFILER_ENABLED
-		if (!EnumHasAnyFlags(Texture.GetFlags(), TexCreate_Virtual)
-			&& !Texture.ResourceLocation.IsTransient()
-			&& Texture.ResourceLocation.GetType() != FD3D12ResourceLocation::ResourceLocationType::eAliased
-			&& Texture.ResourceLocation.GetType() != FD3D12ResourceLocation::ResourceLocationType::eHeapAliased)
+		if (TextureDesc.IsTextureArray())
 		{
-			size_t Size = Texture.GetMemorySize();
-			uint32 Alignment = Desc.Alignment;
-			FTextureProfiler::Get()->AddTextureAllocation(&Texture, Size, Alignment, 0);
+			check(TextureDesc.ArraySize <= (int32)GetMaxTextureArrayLayers());
 		}
-#endif
-	}
-}
 
-void FD3D12TextureStats::D3D12TextureDeleted(FD3D12Texture3D& Texture)
-{
-	FD3D12Resource* D3D12Texture3D = Texture.GetResource();
+		uint32 ActualMSAACount = TextureDesc.NumSamples;
+		uint32 ActualMSAAQuality = GetMaxMSAAQuality(ActualMSAACount);
 
-	if (D3D12Texture3D)
-	{
-		// Don't update state for transient textures	
-		if (!Texture.ResourceLocation.IsTransient())
+		// 0xffffffff means not supported
+		if (ActualMSAAQuality == 0xffffffff || EnumHasAnyFlags(TextureDesc.Flags, TexCreate_Shared))
+		{
+			// no MSAA
+			ActualMSAACount = 1;
+			ActualMSAAQuality = 0;
+		}
+
+		ResourceDesc = CD3DX12_RESOURCE_DESC::Tex2D(
+			PlatformResourceFormat,
+			TextureDesc.Extent.X,
+			TextureDesc.Extent.Y,
+			TextureDesc.ArraySize * (TextureDesc.IsTextureCube() ? 6 : 1),  // Array size
+			TextureDesc.NumMips,
+			ActualMSAACount,
+			ActualMSAAQuality,
+			D3D12_RESOURCE_FLAG_NONE);  // Add misc flags later
+
+		if (EnumHasAnyFlags(TextureDesc.Flags, TexCreate_Shared))
+		{
+			ResourceDesc.Flags |= D3D12_RESOURCE_FLAG_ALLOW_SIMULTANEOUS_ACCESS;
+		}
+
+		if (EnumHasAnyFlags(TextureDesc.Flags, TexCreate_RenderTargetable))
+		{
+			check(!EnumHasAnyFlags(TextureDesc.Flags, TexCreate_DepthStencilTargetable | TexCreate_ResolveTargetable));
+			ResourceDesc.Flags |= D3D12_RESOURCE_FLAG_ALLOW_RENDER_TARGET;
+		}
+		else if (EnumHasAnyFlags(TextureDesc.Flags, TexCreate_DepthStencilTargetable))
+		{
+			check(!EnumHasAnyFlags(TextureDesc.Flags, TexCreate_RenderTargetable | TexCreate_ResolveTargetable));
+			ResourceDesc.Flags |= D3D12_RESOURCE_FLAG_ALLOW_DEPTH_STENCIL;
+		}
+		else if (EnumHasAnyFlags(TextureDesc.Flags, TexCreate_ResolveTargetable))
 		{
-			const D3D12_RESOURCE_DESC& Desc = D3D12Texture3D->GetDesc();
-			const int64 TextureSize = Texture.GetMemorySize();
-			if (TextureSize > 0)
+			check(!EnumHasAnyFlags(TextureDesc.Flags, TexCreate_RenderTargetable | TexCreate_DepthStencilTargetable));
+			if (TextureDesc.Format == PF_DepthStencil || TextureDesc.Format == PF_ShadowDepth || TextureDesc.Format == PF_D24)
 			{
-				UpdateD3D12TextureStats(Texture, Desc, -TextureSize, true, false, Texture.IsStreamable(), false);
+				ResourceDesc.Flags |= D3D12_RESOURCE_FLAG_ALLOW_DEPTH_STENCIL;
 			}
-
-#if TEXTURE_PROFILER_ENABLED
-			if (!EnumHasAnyFlags(Texture.GetFlags(), TexCreate_Virtual)
-				&& !Texture.ResourceLocation.IsTransient()
-				&& Texture.ResourceLocation.GetType() != FD3D12ResourceLocation::ResourceLocationType::eAliased
-				&& Texture.ResourceLocation.GetType() != FD3D12ResourceLocation::ResourceLocationType::eHeapAliased)
+			else
 			{
-				FTextureProfiler::Get()->RemoveTextureAllocation(&Texture);
+				ResourceDesc.Flags |= D3D12_RESOURCE_FLAG_ALLOW_RENDER_TARGET;
 			}
-#endif
+		}
+
+		// Only deny shader resources if it's a depth resource that will never be used as SRV
+		if (EnumHasAnyFlags(TextureDesc.Flags, TexCreate_DepthStencilTargetable) && !EnumHasAnyFlags(TextureDesc.Flags, TexCreate_ShaderResource))
+		{
+			bDenyShaderResource = true;
 		}
 	}
-}
+	else // ETextureDimension::Texture3D
+	{
+		check(TextureDesc.Dimension == ETextureDimension::Texture3D);
+		check(!EnumHasAnyFlags(TextureDesc.Flags, TexCreate_DepthStencilTargetable | TexCreate_ResolveTargetable));
 
-using namespace D3D12RHI;
+		ResourceDesc = CD3DX12_RESOURCE_DESC::Tex3D(
+			PlatformResourceFormat,
+			TextureDesc.Extent.X,
+			TextureDesc.Extent.Y,
+			TextureDesc.Depth,
+			TextureDesc.NumMips);
 
-template<typename BaseResourceType>
-TD3D12Texture2D<BaseResourceType>::~TD3D12Texture2D()
-{
-	if (IsHeadLink())
-	{
-		// Only call this once for a LDA chain
-		FD3D12TextureStats::D3D12TextureDeleted(*this);
+		if (EnumHasAnyFlags(TextureDesc.Flags, TexCreate_RenderTargetable))
+		{
+			ResourceDesc.Flags |= D3D12_RESOURCE_FLAG_ALLOW_RENDER_TARGET;
+		}
 	}
-#if PLATFORM_SUPPORTS_VIRTUAL_TEXTURES
-	GetParentDevice()->GetOwningRHI()->DestroyVirtualTexture(BaseResourceType::GetFlags(), GetRawTextureMemory(), GetRawTextureBlock(), GetMemorySize());
-#endif
-}
 
-FD3D12Texture3D::~FD3D12Texture3D()
-{
-	if (IsHeadLink())
+	if (bDenyShaderResource)
 	{
-		// Only call this once for a LDA chain
-		FD3D12TextureStats::D3D12TextureDeleted(*this);
+		ResourceDesc.Flags |= D3D12_RESOURCE_FLAG_DENY_SHADER_RESOURCE;
 	}
-}
-
-uint64 FD3D12DynamicRHI::RHICalcTexture2DPlatformSize(uint32 SizeX, uint32 SizeY, uint8 Format, uint32 NumMips, uint32 NumSamples, ETextureCreateFlags Flags, const FRHIResourceCreateInfo& CreateInfo, uint32& OutAlign)
-{
-	D3D12_RESOURCE_DESC Desc = {};
-	Desc.DepthOrArraySize = 1;
-	Desc.Dimension = D3D12_RESOURCE_DIMENSION_TEXTURE2D;
-	Desc.Format = (DXGI_FORMAT)GPixelFormats[Format].PlatformFormat;
-	Desc.Height = SizeY;
-	Desc.Layout = D3D12_TEXTURE_LAYOUT_UNKNOWN;
-	Desc.MipLevels = NumMips;
-	Desc.SampleDesc.Count = NumSamples;
-	Desc.Width = SizeX;
 
-	// Check if the 4K aligment is possible
-	Desc.Alignment = TextureCanBe4KAligned(Desc, (EPixelFormat)Format) ? D3D12_SMALL_RESOURCE_PLACEMENT_ALIGNMENT : 0;
+	bool bBCTextureNeedsUAVAlias = EnumHasAnyFlags(TextureDesc.Flags, TexCreate_UAV) && IsBlockCompressedFormat(TextureDesc.Format);
+	if (bBCTextureNeedsUAVAlias)
+	{
+		ResourceDesc.UAVAliasPixelFormat = GetBlockCompressedFormatUAVAliasFormat(TextureDesc.Format);
+	}
 
-	const D3D12_RESOURCE_ALLOCATION_INFO AllocationInfo = GetAdapter().GetD3DDevice()->GetResourceAllocationInfo(0, 1, &Desc);
-	OutAlign = static_cast<uint32>(AllocationInfo.Alignment);
+	if (EnumHasAnyFlags(TextureDesc.Flags, TexCreate_UAV) && !bBCTextureNeedsUAVAlias)
+	{
+		ResourceDesc.Flags |= D3D12_RESOURCE_FLAG_ALLOW_UNORDERED_ACCESS;
+	}
 
-	return AllocationInfo.SizeInBytes;
-}
+	ResourceDesc.PixelFormat = TextureDesc.Format;
 
-uint64 FD3D12DynamicRHI::RHICalcTexture2DArrayPlatformSize(uint32 SizeX, uint32 SizeY, uint32 ArraySize, uint8 Format, uint32 NumMips, uint32 NumSamples, ETextureCreateFlags Flags, const FRHIResourceCreateInfo& CreateInfo, uint32& OutAlign)
-{
-	D3D12_RESOURCE_DESC Desc = {};
-	Desc.Dimension = D3D12_RESOURCE_DIMENSION_TEXTURE2D;
-	Desc.Format = (DXGI_FORMAT)GPixelFormats[Format].PlatformFormat;
-	Desc.Height = SizeY;
-	Desc.DepthOrArraySize = ArraySize;
-	Desc.Layout = D3D12_TEXTURE_LAYOUT_UNKNOWN;
-	Desc.MipLevels = NumMips;
-	Desc.SampleDesc.Count = NumSamples;
-	Desc.Width = SizeX;
+#if D3D12RHI_NEEDS_VENDOR_EXTENSIONS
+	ResourceDesc.bRequires64BitAtomicSupport = EnumHasAnyFlags(TextureDesc.Flags, ETextureCreateFlags::Atomic64Compatible);
+#endif
 
 	// Check if the 4K aligment is possible
-	Desc.Alignment = TextureCanBe4KAligned(Desc, (EPixelFormat)Format) ? D3D12_SMALL_RESOURCE_PLACEMENT_ALIGNMENT : 0;
-
-	const D3D12_RESOURCE_ALLOCATION_INFO AllocationInfo = GetAdapter().GetD3DDevice()->GetResourceAllocationInfo(0, 1, &Desc);
-	OutAlign = static_cast<uint32>(AllocationInfo.Alignment);
+	ResourceDesc.Alignment = FD3D12Texture::CanBe4KAligned(ResourceDesc, (EPixelFormat)TextureDesc.Format) ? D3D12_SMALL_RESOURCE_PLACEMENT_ALIGNMENT : D3D12_DEFAULT_RESOURCE_PLACEMENT_ALIGNMENT;
 
-	return AllocationInfo.SizeInBytes;
+	return ResourceDesc;
 }
 
-uint64 FD3D12DynamicRHI::RHICalcTexture3DPlatformSize(uint32 SizeX, uint32 SizeY, uint32 SizeZ, uint8 Format, uint32 NumMips, ETextureCreateFlags Flags, const FRHIResourceCreateInfo& CreateInfo, uint32& OutAlign)
-{
-	D3D12_RESOURCE_DESC Desc = {};
-	Desc.DepthOrArraySize = SizeZ;
-	Desc.Dimension = D3D12_RESOURCE_DIMENSION_TEXTURE3D;
-	Desc.Format = (DXGI_FORMAT)GPixelFormats[Format].PlatformFormat;
-	Desc.Height = SizeY;
-	Desc.Layout = D3D12_TEXTURE_LAYOUT_UNKNOWN;
-	Desc.MipLevels = NumMips;
-	Desc.SampleDesc.Count = 1;
-	Desc.Width = SizeX;
-
-	// Check if the 4K aligment is possible
-	Desc.Alignment = TextureCanBe4KAligned(Desc, (EPixelFormat)Format) ? D3D12_SMALL_RESOURCE_PLACEMENT_ALIGNMENT : 0;
-
-	const D3D12_RESOURCE_ALLOCATION_INFO AllocationInfo = GetAdapter().GetD3DDevice()->GetResourceAllocationInfo(0, 1, &Desc);
-	OutAlign = static_cast<uint32>(AllocationInfo.Alignment);
-
-	return AllocationInfo.SizeInBytes;
-}
 
-uint64 FD3D12DynamicRHI::RHICalcTextureCubePlatformSize(uint32 Size, uint8 Format, uint32 NumMips, ETextureCreateFlags Flags, const FRHIResourceCreateInfo& CreateInfo, uint32& OutAlign)
+FDynamicRHI::FRHICalcTextureSizeResult FD3D12DynamicRHI::RHICalcTexturePlatformSize(const FRHITextureDesc& InTextureDesc, uint32 FirstMipIndex)
 {
-	D3D12_RESOURCE_DESC Desc = {};
-	Desc.DepthOrArraySize = 6;
-	Desc.Dimension = D3D12_RESOURCE_DIMENSION_TEXTURE2D;
-	Desc.Format = (DXGI_FORMAT)GPixelFormats[Format].PlatformFormat;
-	Desc.Height = Size;
-	Desc.Layout = D3D12_TEXTURE_LAYOUT_UNKNOWN;
-	Desc.MipLevels = NumMips;
-	Desc.SampleDesc.Count = 1;
-	Desc.Width = Size;
+	const D3D12_RESOURCE_DESC Desc = GetResourceDesc(InTextureDesc);
+	const D3D12_RESOURCE_ALLOCATION_INFO AllocationInfo = GetAdapter().GetDevice(0)->GetResourceAllocationInfo(Desc);
 
-	// Check if the 4K aligment is possible
-	Desc.Alignment = TextureCanBe4KAligned(Desc, (EPixelFormat)Format) ? D3D12_SMALL_RESOURCE_PLACEMENT_ALIGNMENT : 0;
-
-	const D3D12_RESOURCE_ALLOCATION_INFO AllocationInfo = GetAdapter().GetD3DDevice()->GetResourceAllocationInfo(0, 1, &Desc);
-	OutAlign = static_cast<uint32>(AllocationInfo.Alignment);
-
-	return AllocationInfo.SizeInBytes;
+	FDynamicRHI::FRHICalcTextureSizeResult Result;
+	Result.Size = AllocationInfo.SizeInBytes;
+	Result.Align = AllocationInfo.Alignment;
+	return Result;
 }
 
+
 /**
  * Retrieves texture memory stats.
  */
@@ -770,465 +819,109 @@ void SafeCreateTexture2D(FD3D12Device* pDevice,
 #endif // #if GUARDED_TEXTURE_CREATES
 }
 
-void CreateUAVAliasResource(FD3D12Adapter* Adapter, D3D12_CLEAR_VALUE* ClearValuePtr, const TCHAR* DebugName, FD3D12ResourceLocation& Location)
+FD3D12Texture* FD3D12DynamicRHI::CreateNewD3D12Texture(const FRHITextureCreateDesc& CreateDesc, FD3D12Device* Device)
 {
-	FD3D12Resource* SourceResource = Location.GetResource();
-
-	const FD3D12ResourceDesc& SourceDesc = SourceResource->GetDesc();
-	const FD3D12Heap* const ResourceHeap = SourceResource->GetHeap();
-
-	const EPixelFormat SourceFormat = SourceDesc.PixelFormat;
-	const EPixelFormat AliasTextureFormat = SourceDesc.UAVAliasPixelFormat;
-
-	if (ensure(ResourceHeap != nullptr) && ensure(SourceFormat != PF_Unknown) && SourceFormat != AliasTextureFormat)
-	{
-		const uint64 SourceOffset = Location.GetOffsetFromBaseOfResource();
-
-		FD3D12ResourceDesc AliasTextureDesc = SourceDesc;
-		AliasTextureDesc.Format = (DXGI_FORMAT)GPixelFormats[AliasTextureFormat].PlatformFormat;
-		AliasTextureDesc.Width = SourceDesc.Width / GPixelFormats[SourceFormat].BlockSizeX;
-		AliasTextureDesc.Height = SourceDesc.Height / GPixelFormats[SourceFormat].BlockSizeY;
-		// layout of UAV must match source resource
-		AliasTextureDesc.Layout = SourceResource->GetResource()->GetDesc().Layout;
-
-		EnumAddFlags(AliasTextureDesc.Flags, D3D12_RESOURCE_FLAG_ALLOW_UNORDERED_ACCESS);
-
-		AliasTextureDesc.UAVAliasPixelFormat = PF_Unknown;
-
-		TRefCountPtr<ID3D12Resource> pAliasResource;
-		HRESULT AliasHR = Adapter->GetD3DDevice()->CreatePlacedResource(
-			ResourceHeap->GetHeap(),
-			SourceOffset,
-			&AliasTextureDesc,
-			D3D12_RESOURCE_STATE_UNORDERED_ACCESS,
-			ClearValuePtr,
-			IID_PPV_ARGS(pAliasResource.GetInitReference()));
-
-		if (pAliasResource && DebugName)
-		{
-			TCHAR NameBuffer[512]{};
-			FCString::Snprintf(NameBuffer, UE_ARRAY_COUNT(NameBuffer), TEXT("%s UAVAlias"), DebugName);
-			SetName(pAliasResource, NameBuffer);
-		}
-
-		if (SUCCEEDED(AliasHR))
-		{
-			SourceResource->SetUAVAccessResource(pAliasResource);
-		}
-	}
+	return new FD3D12Texture(CreateDesc, Device);
 }
 
-static void DetermineTexture2DResourceFlagsAndLayout(uint32 SizeX, uint32 SizeY, uint32 SizeZ, uint32 NumMips, uint32 NumSamples, ETextureCreateFlags Flags, EPixelFormat Format, D3D12_RESOURCE_FLAGS& OutResourceFlags, D3D12_TEXTURE_LAYOUT& OutLayout, bool& bOutCreateRTV, bool& bOutCreateDSV, bool& bOutCreateSRV)
+FD3D12Texture* FD3D12DynamicRHI::CreateD3D12Texture(const FRHITextureCreateDesc& InCreateDesc, class FRHICommandListImmediate* RHICmdList, ID3D12ResourceAllocator* ResourceAllocator)
 {
-	TRACE_CPUPROFILER_EVENT_SCOPE(DetermineTexture2DResourceFlagsAndLayout);
+#if PLATFORM_WINDOWS || PLATFORM_HOLOLENS
+	
+	TRACE_CPUPROFILER_EVENT_SCOPE(D3D12RHI::CreateD3D12Texture);
 
-	OutResourceFlags = D3D12_RESOURCE_FLAG_NONE;
-	OutLayout = D3D12_TEXTURE_LAYOUT_UNKNOWN;
-	bOutCreateRTV = false;
-	bOutCreateDSV = false;
-	bOutCreateSRV = true;
+	// Make local copy of create desc because certain flags might be removed before actual texture creation
+	FRHITextureCreateDesc CreateDesc = InCreateDesc;
 
-	if (EnumHasAllFlags(Flags, TexCreate_CPUReadback))
-	{
-		check(!EnumHasAnyFlags(Flags, TexCreate_RenderTargetable | TexCreate_DepthStencilTargetable | TexCreate_ShaderResource));
-		bOutCreateSRV = false;
-	}
+	// Virtual textures currently not supported in default D3D12
+	check(!EnumHasAnyFlags(CreateDesc.Flags, TexCreate_Virtual));
 
-	if (EnumHasAnyFlags(Flags, TexCreate_DisableSRVCreation))
-	{
-		bOutCreateSRV = false;
-	}
+	// Get the resource desc
+	FD3D12ResourceDesc ResourceDesc = GetResourceDesc(CreateDesc);
 
-	if (EnumHasAnyFlags(Flags, TexCreate_Shared))
-	{
-		OutResourceFlags |= D3D12_RESOURCE_FLAG_ALLOW_SIMULTANEOUS_ACCESS;
-	}
+	const DXGI_FORMAT PlatformResourceFormat = GetPlatformTextureResourceFormat((DXGI_FORMAT)GPixelFormats[CreateDesc.Format].PlatformFormat, CreateDesc.Flags);
+	bool bCreateRTV = EnumHasAnyFlags(ResourceDesc.Flags, D3D12_RESOURCE_FLAG_ALLOW_RENDER_TARGET);
+	bool bCreateDSV = EnumHasAnyFlags(ResourceDesc.Flags, D3D12_RESOURCE_FLAG_ALLOW_DEPTH_STENCIL);
 
-	if (EnumHasAnyFlags(Flags, TexCreate_RenderTargetable))
+	// Setup the clear value	
+	D3D12_CLEAR_VALUE* ClearValuePtr = nullptr;
+	D3D12_CLEAR_VALUE ClearValue;
+	if (bCreateDSV && CreateDesc.ClearValue.ColorBinding == EClearBinding::EDepthStencilBound)
 	{
-		check(!EnumHasAnyFlags(Flags, TexCreate_DepthStencilTargetable | TexCreate_ResolveTargetable));
-		OutResourceFlags |= D3D12_RESOURCE_FLAG_ALLOW_RENDER_TARGET;
-		bOutCreateRTV = true;
+		const DXGI_FORMAT PlatformDepthStencilFormat = FindDepthStencilDXGIFormat(PlatformResourceFormat);
+		ClearValue = CD3DX12_CLEAR_VALUE(PlatformDepthStencilFormat, CreateDesc.ClearValue.Value.DSValue.Depth, (uint8)CreateDesc.ClearValue.Value.DSValue.Stencil);
+		ClearValuePtr = &ClearValue;
 	}
-	else if (EnumHasAnyFlags(Flags, TexCreate_DepthStencilTargetable))
+	else if (bCreateRTV && CreateDesc.ClearValue.ColorBinding == EClearBinding::EColorBound)
 	{
-		check(!EnumHasAnyFlags(Flags, TexCreate_RenderTargetable | TexCreate_ResolveTargetable));
-		OutResourceFlags |= D3D12_RESOURCE_FLAG_ALLOW_DEPTH_STENCIL;
-		bOutCreateDSV = true;
+		const bool bSRGB = EnumHasAnyFlags(CreateDesc.Flags, TexCreate_SRGB);
+		const DXGI_FORMAT PlatformRenderTargetFormat = FindShaderResourceDXGIFormat(PlatformResourceFormat, bSRGB);
+		ClearValue = CD3DX12_CLEAR_VALUE(PlatformRenderTargetFormat, CreateDesc.ClearValue.Value.Color);
+		ClearValuePtr = &ClearValue;
 	}
-	else if (EnumHasAnyFlags(Flags, TexCreate_ResolveTargetable))
+
+	// The state this resource will be in when it leaves this function
+	const FD3D12Resource::FD3D12ResourceTypeHelper Type(ResourceDesc, D3D12_HEAP_TYPE_DEFAULT);
+	const D3D12_RESOURCE_STATES InitialState = Type.GetOptimalInitialState(CreateDesc.InitialState, false);
+	const D3D12_RESOURCE_STATES CreateState = (CreateDesc.BulkData != nullptr) ? D3D12_RESOURCE_STATE_COPY_DEST : InitialState;
+
+	FD3D12Adapter* Adapter = &GetAdapter();
+	FD3D12Texture* D3D12TextureOut = Adapter->CreateLinkedObject<FD3D12Texture>(CreateDesc.GPUMask, [&](FD3D12Device* Device)
 	{
-		check(!EnumHasAnyFlags(Flags, TexCreate_RenderTargetable | TexCreate_DepthStencilTargetable));
-		if (Format == PF_DepthStencil || Format == PF_ShadowDepth || Format == PF_D24)
+		FD3D12Texture* NewTexture = CreateNewD3D12Texture(CreateDesc, Device);
+
+#if NAME_OBJECTS
+		if (CreateDesc.DebugName)
+		{
+			TRACE_CPUPROFILER_EVENT_SCOPE(D3D12RHI::SetDebugName);
+			NewTexture->SetName(CreateDesc.DebugName);
+		}
+#endif // NAME_OBJECTS
+
+		FD3D12ResourceLocation& Location = NewTexture->ResourceLocation;
+
+		if (ResourceAllocator)
+		{
+			const D3D12_HEAP_TYPE HeapType = EnumHasAnyFlags(CreateDesc.Flags, TexCreate_CPUReadback) ? D3D12_HEAP_TYPE_READBACK : D3D12_HEAP_TYPE_DEFAULT;
+			ResourceAllocator->AllocateTexture(Device->GetGPUIndex(), HeapType, ResourceDesc, (EPixelFormat)CreateDesc.Format, ED3D12ResourceStateMode::Default, CreateState, ClearValuePtr, CreateDesc.DebugName, Location);
+			Location.SetOwner(NewTexture);
+		}
+		else if (CreateDesc.IsTexture3D())
 		{
-			OutResourceFlags |= D3D12_RESOURCE_FLAG_ALLOW_DEPTH_STENCIL;
-			bOutCreateDSV = true;
+			VERIFYD3D12CREATETEXTURERESULT(Device->GetTextureAllocator().AllocateTexture(
+				ResourceDesc, ClearValuePtr, CreateDesc.Format, Location, CreateState, CreateDesc.DebugName), ResourceDesc, Device->GetDevice());
+			Location.SetOwner(NewTexture);
 		}
 		else
 		{
-			OutResourceFlags |= D3D12_RESOURCE_FLAG_ALLOW_RENDER_TARGET;
-			bOutCreateRTV = true;
+			SafeCreateTexture2D(Device,
+				Adapter,
+				ResourceDesc,
+				ClearValuePtr,
+				&Location,
+				NewTexture,
+				CreateDesc.Format,
+				CreateDesc.Flags,
+				CreateState,
+				CreateDesc.DebugName);
 		}
-	}
 
-	if (EnumHasAnyFlags(Flags, TexCreate_UAV))
-	{
-		OutResourceFlags |= D3D12_RESOURCE_FLAG_ALLOW_UNORDERED_ACCESS;
-	}
+		// Unlock immediately if no initial data
+		if (CreateDesc.BulkData == nullptr)
+		{
+			Location.UnlockPoolData();
+		}
 
-	if (bOutCreateDSV && !EnumHasAnyFlags(Flags, TexCreate_ShaderResource))
-	{
-		// Only deny shader resources if it's a depth resource that will never be used as SRV
-		OutResourceFlags |= D3D12_RESOURCE_FLAG_DENY_SHADER_RESOURCE;
-		bOutCreateSRV = false;
-	}
-}
+		check(Location.IsValid());
 
-template<typename BaseResourceType>
-TD3D12Texture2D<BaseResourceType>* FD3D12DynamicRHI::CreateD3D12Texture2D(FRHICommandListImmediate* RHICmdList, uint32 SizeX, uint32 SizeY, uint32 SizeZ, bool bTextureArray, bool bCubeTexture, EPixelFormat Format,
-	uint32 NumMips, uint32 NumSamples, ETextureCreateFlags Flags, ERHIAccess InResourceState, FRHIResourceCreateInfo& CreateInfo, ED3D12ResourceTransientMode TransientMode, ID3D12ResourceAllocator* ResourceAllocator)
-{
-#if PLATFORM_WINDOWS || PLATFORM_HOLOLENS
-	
-	TRACE_CPUPROFILER_EVENT_SCOPE(D3D12RHI::CreateD3D12Texture2D);
+		bool bBCTextureNeedsUAVAlias = EnumHasAnyFlags(CreateDesc.Flags, TexCreate_UAV) && IsBlockCompressedFormat(CreateDesc.Format);
+		if (bBCTextureNeedsUAVAlias)
+		{
+			check(ResourceDesc.UAVAliasPixelFormat != PF_Unknown);
+			Adapter->CreateUAVAliasResource(ClearValuePtr, CreateDesc.DebugName, Location);
+		}
 
-	check(SizeX > 0 && SizeY > 0 && NumMips > 0);
-
-	if (bCubeTexture)
-	{
-		check(SizeX <= GetMaxCubeTextureDimension());
-		check(SizeX == SizeY);
-	}
-	else
-	{
-		check(SizeX <= GetMax2DTextureDimension());
-		check(SizeY <= GetMax2DTextureDimension());
-	}
-
-	if (bTextureArray)
-	{
-		check(SizeZ > 0 && SizeZ <= GetMaxTextureArrayLayers());
-	}
-
-	SCOPE_CYCLE_COUNTER(STAT_D3D12CreateTextureTime);
-
-	const bool bSRGB = EnumHasAnyFlags(Flags, TexCreate_SRGB);
-
-	const DXGI_FORMAT PlatformResourceFormat = GetPlatformTextureResourceFormat((DXGI_FORMAT)GPixelFormats[Format].PlatformFormat, Flags);
-	const DXGI_FORMAT PlatformShaderResourceFormat = FindShaderResourceDXGIFormat(PlatformResourceFormat, bSRGB);
-	const DXGI_FORMAT PlatformRenderTargetFormat = FindShaderResourceDXGIFormat(PlatformResourceFormat, bSRGB);
-	const DXGI_FORMAT PlatformDepthStencilFormat = FindDepthStencilDXGIFormat(PlatformResourceFormat);
-
-	uint32 ActualMSAACount = NumSamples;
-	uint32 ActualMSAAQuality = GetMaxMSAAQuality(ActualMSAACount);
-
-	// 0xffffffff means not supported
-	if (ActualMSAAQuality == 0xffffffff || EnumHasAnyFlags(Flags, TexCreate_Shared))
-	{
-		// no MSAA
-		ActualMSAACount = 1;
-		ActualMSAAQuality = 0;
-	}
-	const bool bIsMultisampled = ActualMSAACount > 1;
-
-	// Describe the texture.
-	FD3D12ResourceDesc TextureDesc = CD3DX12_RESOURCE_DESC::Tex2D(
-		PlatformResourceFormat,
-		SizeX,
-		SizeY,
-		SizeZ,  // Array size
-		NumMips,
-		ActualMSAACount,
-		ActualMSAAQuality,
-		D3D12_RESOURCE_FLAG_NONE);  // Add misc flags later
-
-	TextureDesc.PixelFormat = Format;
-
-	bool bBCTextureNeedsUAVAlias = EnumHasAnyFlags(Flags, TexCreate_UAV) && IsBlockCompressedFormat(Format);
-	if (bBCTextureNeedsUAVAlias)
-	{
-		EnumRemoveFlags(Flags, TexCreate_UAV);
-		TextureDesc.UAVAliasPixelFormat = GetBlockCompressedFormatUAVAliasFormat(Format);
-	}
-
-#if D3D12RHI_NEEDS_VENDOR_EXTENSIONS
-	TextureDesc.bRequires64BitAtomicSupport = EnumHasAnyFlags(Flags, ETextureCreateFlags::Atomic64Compatible);
-#endif
-
-	// Set up the texture bind flags.
-	bool bCreateRTV;
-	bool bCreateDSV;
-	bool bCreateShaderResource;
-	DetermineTexture2DResourceFlagsAndLayout(SizeX, SizeY, SizeZ, NumMips, ActualMSAACount, Flags, Format, TextureDesc.Flags, TextureDesc.Layout, bCreateRTV, bCreateDSV, bCreateShaderResource);
-
-	// Virtual textures currently not supported in default D3D12
-	Flags &= ~TexCreate_Virtual;
-
-	FD3D12Adapter* Adapter = &GetAdapter();
-
-	D3D12_CLEAR_VALUE *ClearValuePtr = nullptr;
-	D3D12_CLEAR_VALUE ClearValue;
-	if (bCreateDSV && CreateInfo.ClearValueBinding.ColorBinding == EClearBinding::EDepthStencilBound)
-	{
-		ClearValue = CD3DX12_CLEAR_VALUE(PlatformDepthStencilFormat, CreateInfo.ClearValueBinding.Value.DSValue.Depth, (uint8)CreateInfo.ClearValueBinding.Value.DSValue.Stencil);
-		ClearValuePtr = &ClearValue;
-	}
-	else if (bCreateRTV && CreateInfo.ClearValueBinding.ColorBinding == EClearBinding::EColorBound)
-	{
-		ClearValue = CD3DX12_CLEAR_VALUE(PlatformRenderTargetFormat, CreateInfo.ClearValueBinding.Value.Color);
-		ClearValuePtr = &ClearValue;
-	}
-
-	if (Format == PF_NV12)
-	{
-		bCreateRTV = false;
-		bCreateShaderResource = false;
-	}
-
-	// The state this resource will be in when it leaves this function
-	const FD3D12Resource::FD3D12ResourceTypeHelper Type(TextureDesc, D3D12_HEAP_TYPE_DEFAULT);
-	const D3D12_RESOURCE_STATES InitialState = Type.GetOptimalInitialState(InResourceState, false);
-	const D3D12_RESOURCE_STATES CreateState = (CreateInfo.BulkData != nullptr) ? D3D12_RESOURCE_STATE_COPY_DEST : InitialState;
-
-	TD3D12Texture2D<BaseResourceType>* D3D12TextureOut = Adapter->CreateLinkedObject<TD3D12Texture2D<BaseResourceType>>(CreateInfo.GPUMask, [&](FD3D12Device* Device)
-	{
-		TD3D12Texture2D<BaseResourceType>* NewTexture = new TD3D12Texture2D<BaseResourceType>(Device,
-			SizeX,
-			SizeY,
-			SizeZ,
-			NumMips,
-			ActualMSAACount,
-			(EPixelFormat)Format,
-			bCubeTexture,
-			Flags,
-			CreateInfo.ClearValueBinding);		
-
-#if NAME_OBJECTS
-		if (CreateInfo.DebugName)
-		{
-			TRACE_CPUPROFILER_EVENT_SCOPE(D3D12RHI::SetDebugName);
-			NewTexture->SetName(CreateInfo.DebugName);
-		}
-#endif // NAME_OBJECTS
-
-		FD3D12ResourceLocation& Location = NewTexture->ResourceLocation;
-
-		if (ResourceAllocator)
-		{
-			const D3D12_HEAP_TYPE HeapType = EnumHasAnyFlags(Flags, TexCreate_CPUReadback) ? D3D12_HEAP_TYPE_READBACK : D3D12_HEAP_TYPE_DEFAULT;
-			ResourceAllocator->AllocateTexture(Device->GetGPUIndex(), HeapType, TextureDesc, (EPixelFormat)Format, ED3D12ResourceStateMode::Default, CreateState, ClearValuePtr, CreateInfo.DebugName, Location);
-			Location.SetOwner(NewTexture);
-		}
-		else
-		{
-			SafeCreateTexture2D(Device,
-				Adapter,
-				TextureDesc,
-				ClearValuePtr,
-				&Location,
-				NewTexture,
-				Format,
-				Flags,
-				CreateState,
-				CreateInfo.DebugName);
-		}
-
-		// Unlock immediately if no initial data
-		if (CreateInfo.BulkData == nullptr)
-		{
-			Location.UnlockPoolData();
-		}
-
-		check(Location.IsValid());
-
-		if (bBCTextureNeedsUAVAlias)
-		{
-			CreateUAVAliasResource(Adapter, ClearValuePtr, CreateInfo.DebugName, Location);
-		}
-
-		uint32 RTVIndex = 0;
-
-		{
-			TRACE_CPUPROFILER_EVENT_SCOPE(D3D12RHI::CreateViews);
-			if (bCreateRTV)
-			{
-				const bool bCreateRTVsPerSlice = EnumHasAnyFlags(Flags, TexCreate_TargetArraySlicesIndependently) && (bTextureArray || bCubeTexture);
-				NewTexture->SetNumRenderTargetViews(bCreateRTVsPerSlice ? NumMips * TextureDesc.DepthOrArraySize : NumMips);
-
-				// Create a render target view for each mip
-				for (uint32 MipIndex = 0; MipIndex < NumMips; MipIndex++)
-				{
-					if (bCreateRTVsPerSlice)
-					{
-						NewTexture->SetCreatedRTVsPerSlice(true, TextureDesc.DepthOrArraySize);
-
-						for (uint32 SliceIndex = 0; SliceIndex < TextureDesc.DepthOrArraySize; SliceIndex++)
-						{
-							D3D12_RENDER_TARGET_VIEW_DESC RTVDesc;
-							FMemory::Memzero(RTVDesc);
-
-							RTVDesc.Format = PlatformRenderTargetFormat;
-							RTVDesc.ViewDimension = D3D12_RTV_DIMENSION_TEXTURE2DARRAY;
-							RTVDesc.Texture2DArray.FirstArraySlice = SliceIndex;
-							RTVDesc.Texture2DArray.ArraySize = 1;
-							RTVDesc.Texture2DArray.MipSlice = MipIndex;
-							RTVDesc.Texture2DArray.PlaneSlice = GetPlaneSliceFromViewFormat(PlatformResourceFormat, RTVDesc.Format);
-
-							NewTexture->SetRenderTargetViewIndex(new FD3D12RenderTargetView(Device, RTVDesc, NewTexture), RTVIndex++);
-						}
-					}
-					else
-					{
-						D3D12_RENDER_TARGET_VIEW_DESC RTVDesc;
-						FMemory::Memzero(RTVDesc);
-
-						RTVDesc.Format = PlatformRenderTargetFormat;
-
-						if (bTextureArray || bCubeTexture)
-						{
-							if (bIsMultisampled)
-							{
-								RTVDesc.ViewDimension = D3D12_RTV_DIMENSION_TEXTURE2DMSARRAY;
-								RTVDesc.Texture2DMSArray.FirstArraySlice = 0;
-								RTVDesc.Texture2DMSArray.ArraySize = TextureDesc.DepthOrArraySize;
-							}
-							else
-							{
-								RTVDesc.ViewDimension = D3D12_RTV_DIMENSION_TEXTURE2DARRAY;
-								RTVDesc.Texture2DArray.FirstArraySlice = 0;
-								RTVDesc.Texture2DArray.ArraySize = TextureDesc.DepthOrArraySize;
-								RTVDesc.Texture2DArray.MipSlice = MipIndex;
-								RTVDesc.Texture2DArray.PlaneSlice = GetPlaneSliceFromViewFormat(PlatformResourceFormat, RTVDesc.Format);
-							}
-						}
-						else
-						{
-							if (bIsMultisampled)
-							{
-								RTVDesc.ViewDimension = D3D12_RTV_DIMENSION_TEXTURE2DMS;
-								// Nothing to set
-							}
-							else
-							{
-								RTVDesc.ViewDimension = D3D12_RTV_DIMENSION_TEXTURE2D;
-								RTVDesc.Texture2D.MipSlice = MipIndex;
-								RTVDesc.Texture2D.PlaneSlice = GetPlaneSliceFromViewFormat(PlatformResourceFormat, RTVDesc.Format);
-							}
-						}
-
-						NewTexture->SetRenderTargetViewIndex(new FD3D12RenderTargetView(Device, RTVDesc, NewTexture), RTVIndex++);
-					}
-				}
-			}
-
-			if (bCreateDSV)
-			{
-				// Create a depth-stencil-view for the texture.
-				D3D12_DEPTH_STENCIL_VIEW_DESC DSVDesc = {};
-				DSVDesc.Format = FindDepthStencilDXGIFormat(PlatformResourceFormat);
-				if (bTextureArray || bCubeTexture)
-				{
-					if (bIsMultisampled)
-					{
-						DSVDesc.ViewDimension = D3D12_DSV_DIMENSION_TEXTURE2DMSARRAY;
-						DSVDesc.Texture2DMSArray.FirstArraySlice = 0;
-						DSVDesc.Texture2DMSArray.ArraySize = TextureDesc.DepthOrArraySize;
-					}
-					else
-					{
-						DSVDesc.ViewDimension = D3D12_DSV_DIMENSION_TEXTURE2DARRAY;
-						DSVDesc.Texture2DArray.FirstArraySlice = 0;
-						DSVDesc.Texture2DArray.ArraySize = TextureDesc.DepthOrArraySize;
-						DSVDesc.Texture2DArray.MipSlice = 0;
-					}
-				}
-				else
-				{
-					if (bIsMultisampled)
-					{
-						DSVDesc.ViewDimension = D3D12_DSV_DIMENSION_TEXTURE2DMS;
-						// Nothing to set
-					}
-					else
-					{
-						DSVDesc.ViewDimension = D3D12_DSV_DIMENSION_TEXTURE2D;
-						DSVDesc.Texture2D.MipSlice = 0;
-					}
-				}
-
-				const bool HasStencil = HasStencilBits(DSVDesc.Format);
-				for (uint32 AccessType = 0; AccessType < FExclusiveDepthStencil::MaxIndex; ++AccessType)
-				{
-					// Create a read-only access views for the texture.
-					DSVDesc.Flags = (AccessType & FExclusiveDepthStencil::DepthRead_StencilWrite) ? D3D12_DSV_FLAG_READ_ONLY_DEPTH : D3D12_DSV_FLAG_NONE;
-					if (HasStencil)
-					{
-						DSVDesc.Flags |= (AccessType & FExclusiveDepthStencil::DepthWrite_StencilRead) ? D3D12_DSV_FLAG_READ_ONLY_STENCIL : D3D12_DSV_FLAG_NONE;
-					}
-
-					NewTexture->SetDepthStencilView(new FD3D12DepthStencilView(Device, DSVDesc, NewTexture, HasStencil), AccessType);
-				}
-			}
-
-			// Create a shader resource view for the texture.
-			if (bCreateShaderResource)
-			{
-				D3D12_SHADER_RESOURCE_VIEW_DESC SRVDesc = {};
-				SRVDesc.Shader4ComponentMapping = D3D12_DEFAULT_SHADER_4_COMPONENT_MAPPING;
-				SRVDesc.Format = PlatformShaderResourceFormat;
-
-				if (bCubeTexture && bTextureArray)
-				{
-					SRVDesc.ViewDimension = D3D12_SRV_DIMENSION_TEXTURECUBEARRAY;
-					SRVDesc.TextureCubeArray.MostDetailedMip = 0;
-					SRVDesc.TextureCubeArray.MipLevels = NumMips;
-					SRVDesc.TextureCubeArray.First2DArrayFace = 0;
-					SRVDesc.TextureCubeArray.NumCubes = SizeZ / 6;
-				}
-				else if (bCubeTexture)
-				{
-					SRVDesc.ViewDimension = D3D12_SRV_DIMENSION_TEXTURECUBE;
-					SRVDesc.TextureCube.MostDetailedMip = 0;
-					SRVDesc.TextureCube.MipLevels = NumMips;
-				}
-				else if (bTextureArray)
-				{
-					if (bIsMultisampled)
-					{
-						SRVDesc.ViewDimension = D3D12_SRV_DIMENSION_TEXTURE2DMSARRAY;
-						SRVDesc.Texture2DMSArray.FirstArraySlice = 0;
-						SRVDesc.Texture2DMSArray.ArraySize = TextureDesc.DepthOrArraySize;
-						//SRVDesc.Texture2DArray.PlaneSlice = GetPlaneSliceFromViewFormat(PlatformResourceFormat, SRVDesc.Format);
-					}
-					else
-					{
-						SRVDesc.ViewDimension = D3D12_SRV_DIMENSION_TEXTURE2DARRAY;
-						SRVDesc.Texture2DArray.MostDetailedMip = 0;
-						SRVDesc.Texture2DArray.MipLevels = NumMips;
-						SRVDesc.Texture2DArray.FirstArraySlice = 0;
-						SRVDesc.Texture2DArray.ArraySize = TextureDesc.DepthOrArraySize;
-						SRVDesc.Texture2DArray.PlaneSlice = GetPlaneSliceFromViewFormat(PlatformResourceFormat, SRVDesc.Format);
-					}
-				}
-				else
-				{
-					if (bIsMultisampled)
-					{
-						SRVDesc.ViewDimension = D3D12_SRV_DIMENSION_TEXTURE2DMS;
-						// Nothing to set
-					}
-					else
-					{
-						SRVDesc.ViewDimension = D3D12_SRV_DIMENSION_TEXTURE2D;
-						SRVDesc.Texture2D.MostDetailedMip = 0;
-						SRVDesc.Texture2D.MipLevels = NumMips;
-						SRVDesc.Texture2D.PlaneSlice = GetPlaneSliceFromViewFormat(PlatformResourceFormat, SRVDesc.Format);
-					}
-				}
-
-				NewTexture->SetShaderResourceView(new FD3D12ShaderResourceView(NewTexture, SRVDesc));
-			}
-		}
+		NewTexture->CreateViews();
 
 		return NewTexture;
 	});
@@ -1236,135 +929,10 @@ TD3D12Texture2D<BaseResourceType>* FD3D12DynamicRHI::CreateD3D12Texture2D(FRHICo
 	FD3D12TextureStats::D3D12TextureAllocated(*D3D12TextureOut);
 
 	// Initialize if data is given
-	if (CreateInfo.BulkData != nullptr)
-	{
-		D3D12TextureOut->InitializeTextureData(RHICmdList, CreateInfo.BulkData->GetResourceBulkData(), CreateInfo.BulkData->GetResourceBulkDataSize(), SizeX, SizeY, 1, SizeZ, NumMips, Format, InitialState);
-
-		CreateInfo.BulkData->Discard();
-	}
-
-	return D3D12TextureOut;
-#else
-	checkf(false, TEXT("XBOX_CODE_MERGE : Removed. The Xbox platform version should be used."));
-	return nullptr;
-#endif // PLATFORM_WINDOWS || PLATFORM_HOLOLENS
-}
-
-FD3D12Texture3D* FD3D12DynamicRHI::CreateD3D12Texture3D(FRHICommandListImmediate* RHICmdList, uint32 SizeX, uint32 SizeY, uint32 SizeZ, EPixelFormat Format, uint32 NumMips, ETextureCreateFlags Flags, ERHIAccess InResourceState, FRHIResourceCreateInfo& CreateInfo, ED3D12ResourceTransientMode TransientMode, ID3D12ResourceAllocator* ResourceAllocator)
-{
-#if PLATFORM_WINDOWS || PLATFORM_HOLOLENS
-	SCOPE_CYCLE_COUNTER(STAT_D3D12CreateTextureTime);
-
-	const bool bSRGB = EnumHasAnyFlags(Flags, TexCreate_SRGB);
-
-	const DXGI_FORMAT PlatformResourceFormat = (DXGI_FORMAT)GPixelFormats[Format].PlatformFormat;
-	const DXGI_FORMAT PlatformShaderResourceFormat = FindShaderResourceDXGIFormat(PlatformResourceFormat, bSRGB);
-	const DXGI_FORMAT PlatformRenderTargetFormat = FindShaderResourceDXGIFormat(PlatformResourceFormat, bSRGB);
-
-	// Describe the texture.
-	D3D12_RESOURCE_DESC TextureDesc = CD3DX12_RESOURCE_DESC::Tex3D(
-		PlatformResourceFormat,
-		SizeX,
-		SizeY,
-		SizeZ,
-		NumMips);
-
-	if (EnumHasAnyFlags(Flags, TexCreate_UAV))
-	{
-		TextureDesc.Flags |= D3D12_RESOURCE_FLAG_ALLOW_UNORDERED_ACCESS;
-	}
-
-	bool bCreateRTV = false;
-
-	if (EnumHasAnyFlags(Flags, TexCreate_RenderTargetable))
-	{
-		TextureDesc.Flags |= D3D12_RESOURCE_FLAG_ALLOW_RENDER_TARGET;
-		bCreateRTV = true;
-	}
-
-	// Set up the texture bind flags.
-	check(!EnumHasAnyFlags(Flags, TexCreate_DepthStencilTargetable | TexCreate_ResolveTargetable));
-	check(EnumHasAllFlags(Flags, TexCreate_ShaderResource));
-
-	D3D12_CLEAR_VALUE *ClearValuePtr = nullptr;
-	D3D12_CLEAR_VALUE ClearValue;
-	if (bCreateRTV && CreateInfo.ClearValueBinding.ColorBinding == EClearBinding::EColorBound)
-	{
-		ClearValue = CD3DX12_CLEAR_VALUE(PlatformRenderTargetFormat, CreateInfo.ClearValueBinding.Value.Color);
-		ClearValuePtr = &ClearValue;
-	}
-
-	// The state this resource will be in when it leaves this function
-	const FD3D12Resource::FD3D12ResourceTypeHelper Type(TextureDesc, D3D12_HEAP_TYPE_DEFAULT);
-	const D3D12_RESOURCE_STATES InitialState = Type.GetOptimalInitialState(InResourceState, false);
-
-	FD3D12Adapter* Adapter = &GetAdapter();
-	FD3D12Texture3D* D3D12TextureOut = Adapter->CreateLinkedObject<FD3D12Texture3D>(CreateInfo.GPUMask, [&](FD3D12Device* Device)
-	{
-		FD3D12Texture3D* Texture3D = new FD3D12Texture3D(Device, SizeX, SizeY, SizeZ, NumMips, Format, Flags, CreateInfo.ClearValueBinding);
-
-		if (CreateInfo.DebugName)
-		{
-			Texture3D->SetName(CreateInfo.DebugName);
-		}
-
-		if (ResourceAllocator)
-		{
-			ResourceAllocator->AllocateTexture(Device->GetGPUIndex(), D3D12_HEAP_TYPE_DEFAULT, TextureDesc, Format, ED3D12ResourceStateMode::Default, InitialState, ClearValuePtr, CreateInfo.DebugName, Texture3D->ResourceLocation);
-		}
-		else
-		{
-			VERIFYD3D12CREATETEXTURERESULT(Device->GetTextureAllocator().AllocateTexture(TextureDesc, ClearValuePtr, Format, Texture3D->ResourceLocation, (CreateInfo.BulkData != nullptr) ? D3D12_RESOURCE_STATE_COPY_DEST : InitialState, CreateInfo.DebugName), TextureDesc, Device->GetDevice());
-		}
-		Texture3D->ResourceLocation.SetOwner(Texture3D);
-
-		// Unlock immediately if no initial data
-		if (CreateInfo.BulkData == nullptr)
-		{
-			Texture3D->ResourceLocation.UnlockPoolData();
-		}
-
-		if (bCreateRTV)
-		{
-			// Create a render-target-view for the texture.
-			D3D12_RENDER_TARGET_VIEW_DESC RTVDesc;
-			FMemory::Memzero(&RTVDesc, sizeof(RTVDesc));
-			RTVDesc.Format = PlatformRenderTargetFormat;
-			RTVDesc.ViewDimension = D3D12_RTV_DIMENSION_TEXTURE3D;
-			RTVDesc.Texture3D.MipSlice = 0;
-			RTVDesc.Texture3D.FirstWSlice = 0;
-			RTVDesc.Texture3D.WSize = SizeZ;
-
-			Texture3D->SetRenderTargetView(new FD3D12RenderTargetView(Device, RTVDesc, Texture3D));
-		}
-
-		// Create a shader resource view for the texture.
-		D3D12_SHADER_RESOURCE_VIEW_DESC SRVDesc = {};
-		SRVDesc.Shader4ComponentMapping = D3D12_DEFAULT_SHADER_4_COMPONENT_MAPPING;
-		SRVDesc.Format = PlatformShaderResourceFormat;
-		SRVDesc.ViewDimension = D3D12_SRV_DIMENSION_TEXTURE3D;
-		SRVDesc.Texture3D.MipLevels = NumMips;
-		SRVDesc.Texture3D.MostDetailedMip = 0;
-
-		Texture3D->SetShaderResourceView(new FD3D12ShaderResourceView(Texture3D, SRVDesc));
-
-		return Texture3D;
-	}); 
-
-	// Intialize if data given
-	if (D3D12TextureOut)
+	if (CreateDesc.BulkData != nullptr)
 	{
-		if (CreateInfo.BulkData != nullptr)
-		{
-			D3D12TextureOut->InitializeTextureData(RHICmdList, CreateInfo.BulkData->GetResourceBulkData(), CreateInfo.BulkData->GetResourceBulkDataSize(), SizeX, SizeY, SizeZ, 1, NumMips, Format, InitialState);
-		}
-
-		FD3D12TextureStats::D3D12TextureAllocated(*D3D12TextureOut);
-	}
-
-	if (CreateInfo.BulkData != nullptr)
-	{
-		CreateInfo.BulkData->Discard();
+		D3D12TextureOut->InitializeTextureData(RHICmdList, CreateDesc, InitialState);
+		CreateDesc.BulkData->Discard();
 	}
 
 	return D3D12TextureOut;
@@ -1374,56 +942,37 @@ FD3D12Texture3D* FD3D12DynamicRHI::CreateD3D12Texture3D(FRHICommandListImmediate
 #endif // PLATFORM_WINDOWS || PLATFORM_HOLOLENS
 }
 
-FRHITexture* FD3D12DynamicRHI::CreateTexture(const FRHITextureCreateInfo& CreateInfo, const TCHAR* DebugName, ERHIAccess InitialState, ED3D12ResourceTransientMode TransientMode, ID3D12ResourceAllocator* ResourceAllocator)
-{
-	FRHIResourceCreateInfo ResourceCreateInfo(DebugName, CreateInfo.ClearValue);
-
-	const bool bTextureArray = CreateInfo.IsTextureArray();
-	const bool bTextureCube = CreateInfo.IsTextureCube();
-
-	switch (CreateInfo.Dimension)
-	{
-	case ETextureDimension::Texture2D:
-		return CreateD3D12Texture2D<FD3D12BaseTexture2D>(nullptr, CreateInfo.Extent.X, CreateInfo.Extent.Y, 1, bTextureArray, bTextureCube, CreateInfo.Format, CreateInfo.NumMips, CreateInfo.NumSamples, CreateInfo.Flags, InitialState, ResourceCreateInfo, TransientMode, ResourceAllocator);
-
-	case ETextureDimension::Texture2DArray:
-		return CreateD3D12Texture2D<FD3D12BaseTexture2DArray>(nullptr, CreateInfo.Extent.X, CreateInfo.Extent.Y, CreateInfo.ArraySize, bTextureArray, bTextureCube, CreateInfo.Format, CreateInfo.NumMips, CreateInfo.NumSamples, CreateInfo.Flags, InitialState, ResourceCreateInfo, TransientMode, ResourceAllocator);
-
-	case ETextureDimension::TextureCube:
-	case ETextureDimension::TextureCubeArray:
-		return CreateD3D12Texture2D<FD3D12BaseTextureCube>(nullptr, CreateInfo.Extent.X, CreateInfo.Extent.Y, 6 * CreateInfo.ArraySize, bTextureArray, bTextureCube, CreateInfo.Format, CreateInfo.NumMips, CreateInfo.NumSamples, CreateInfo.Flags, InitialState, ResourceCreateInfo, TransientMode, ResourceAllocator);
-
-	case ETextureDimension::Texture3D:
-		return CreateD3D12Texture3D(nullptr, CreateInfo.Extent.X, CreateInfo.Extent.Y, CreateInfo.Depth, CreateInfo.Format, CreateInfo.NumMips, CreateInfo.Flags, InitialState, ResourceCreateInfo, TransientMode, ResourceAllocator);
-
-	default:
-		checkNoEntry();
-	}
-
-	return nullptr;
-}
-
-/*-----------------------------------------------------------------------------
-	2D texture support.
-	-----------------------------------------------------------------------------*/
 
-FTexture2DRHIRef FD3D12DynamicRHI::RHICreateTexture2D_RenderThread(class FRHICommandListImmediate& RHICmdList, uint32 SizeX, uint32 SizeY, uint8 Format, uint32 NumMips, uint32 NumSamples, ETextureCreateFlags Flags, ERHIAccess InResourceState, FRHIResourceCreateInfo& CreateInfo)
+FTextureRHIRef FD3D12DynamicRHI::RHICreateTexture(const FRHITextureCreateDesc& CreateDesc)
 {
-	return CreateD3D12Texture2D<FD3D12BaseTexture2D>(&RHICmdList, SizeX, SizeY, 1, false, false, (EPixelFormat) Format, NumMips, NumSamples, Flags, InResourceState, CreateInfo);
+	FRHICommandListImmediate* RHIImmediateCmdList = nullptr;
+	ID3D12ResourceAllocator* ResourceAllocator = nullptr;
+	return CreateD3D12Texture(CreateDesc, RHIImmediateCmdList, ResourceAllocator);
 }
 
-FTexture2DRHIRef FD3D12DynamicRHI::RHICreateTexture2D(uint32 SizeX, uint32 SizeY, uint8 Format, uint32 NumMips, uint32 NumSamples, ETextureCreateFlags Flags, ERHIAccess InResourceState, FRHIResourceCreateInfo& CreateInfo)
+
+FTextureRHIRef FD3D12DynamicRHI::RHICreateTexture_RenderThread(class FRHICommandListImmediate& RHICmdList, const FRHITextureCreateDesc& CreateDesc)
 {
-	return CreateD3D12Texture2D<FD3D12BaseTexture2D>(nullptr, SizeX, SizeY, 1, false, false, (EPixelFormat) Format, NumMips, NumSamples, Flags, InResourceState, CreateInfo);
+	FRHICommandListImmediate* RHIImmediateCmdList = &RHICmdList;
+	ID3D12ResourceAllocator* ResourceAllocator = nullptr;
+	return CreateD3D12Texture(CreateDesc, RHIImmediateCmdList, ResourceAllocator);
 }
 
-FTexture2DRHIRef FD3D12DynamicRHI::RHIAsyncCreateTexture2D(uint32 SizeX, uint32 SizeY, uint8 Format, uint32 NumMips, ETextureCreateFlags Flags, ERHIAccess InResourceState, void** InitialMipData, uint32 NumInitialMips)
-{
+
+FTextureRHIRef FD3D12DynamicRHI::RHIAsyncCreateTexture2D(uint32 SizeX, uint32 SizeY, uint8 Format, uint32 NumMips, ETextureCreateFlags Flags, ERHIAccess InResourceState, void** InitialMipData, uint32 NumInitialMips)
+{	
 	check(GRHISupportsAsyncTextureCreation);
 	
 	const ETextureCreateFlags InvalidFlags = TexCreate_RenderTargetable | TexCreate_ResolveTargetable | TexCreate_DepthStencilTargetable | TexCreate_GenerateMipCapable | TexCreate_UAV | TexCreate_Presentable | TexCreate_CPUReadback;
 	check(!EnumHasAnyFlags(Flags, InvalidFlags));
 
+	FRHITextureDesc Desc = FRHITextureDesc::Create2D(FIntPoint(SizeX, SizeY), (EPixelFormat)Format, FClearValueBinding(), Flags, NumMips, 1);
+	FRHITextureCreateDesc CreateDesc(
+		Desc,
+		ERHIAccess::SRVMask,
+		TEXT("RHIAsyncCreateTexture2D")
+	);
+
 	const DXGI_FORMAT PlatformResourceFormat = (DXGI_FORMAT)GPixelFormats[Format].PlatformFormat;
 	const DXGI_FORMAT PlatformShaderResourceFormat = FindShaderResourceDXGIFormat(PlatformResourceFormat, EnumHasAnyFlags(Flags, TexCreate_SRGB));
 	const D3D12_RESOURCE_DESC TextureDesc = CD3DX12_RESOURCE_DESC::Tex2D(
@@ -1474,18 +1023,9 @@ FTexture2DRHIRef FD3D12DynamicRHI::RHIAsyncCreateTexture2D(uint32 SizeX, uint32
 	const D3D12_RESOURCE_STATES InitialState = D3D12_RESOURCE_STATE_COMMON;
 
 	FD3D12Adapter* Adapter = &GetAdapter();
-	FD3D12Texture2D* TextureOut = Adapter->CreateLinkedObject<FD3D12Texture2D>(FRHIGPUMask::All(), [&](FD3D12Device* Device)
-	{
-		FD3D12Texture2D* NewTexture = new FD3D12Texture2D(Device,
-			SizeX,
-			SizeY,
-			0,
-			NumMips,
-			/*ActualMSAACount=*/ 1,
-			(EPixelFormat)Format,
-			/*bInCubemap=*/ false,
-			Flags,
-			FClearValueBinding());
+	FD3D12Texture* TextureOut = Adapter->CreateLinkedObject<FD3D12Texture>(FRHIGPUMask::All(), [&](FD3D12Device* Device)
+	{
+		FD3D12Texture* NewTexture = CreateNewD3D12Texture(CreateDesc, Device);
 
 		SafeCreateTexture2D(Device,
 			Adapter,
@@ -1550,9 +1090,9 @@ FTexture2DRHIRef FD3D12DynamicRHI::RHIAsyncCreateTexture2D(uint32 SizeX, uint32
 		// on the GPU before returning here, this protection is safe, even if we end up straddling frame boundaries.
 		TempResourceLocation.GetResource()->AddRef();
 
-		for (FD3D12TextureBase& CurrentTextureBase : *TextureOut)
+		for (FD3D12Texture& CurrentTextureBase : *TextureOut)
 		{
-			FD3D12Texture2D& CurrentTexture = static_cast<FD3D12Texture2D&>(CurrentTextureBase);
+			FD3D12Texture& CurrentTexture = static_cast<FD3D12Texture&>(CurrentTextureBase);
 			FD3D12Device* Device = CurrentTexture.GetParentDevice();
 			FD3D12Resource* Resource = CurrentTexture.GetResource();
 
@@ -1628,14 +1168,14 @@ FTexture2DRHIRef FD3D12DynamicRHI::RHIAsyncCreateTexture2D(uint32 SizeX, uint32
 		FMemory::Free(TempBuffer);
 	}
 
-
 	return TextureOut;
 }
 
+
 void FD3D12DynamicRHI::RHICopySharedMips(FRHITexture2D* DestTexture2DRHI, FRHITexture2D* SrcTexture2DRHI)
 {
-	FD3D12Texture2D*  DestTexture2D = FD3D12DynamicRHI::ResourceCast(DestTexture2DRHI);
-	FD3D12Texture2D*  SrcTexture2D = FD3D12DynamicRHI::ResourceCast(SrcTexture2DRHI);
+	FD3D12Texture*  DestTexture2D = GetD3D12TextureFromRHITexture(DestTexture2DRHI);
+	FD3D12Texture*  SrcTexture2D = GetD3D12TextureFromRHITexture(SrcTexture2DRHI);
 
 	// Use the GPU to asynchronously copy the old mip-maps into the new texture.
 	const uint32 NumSharedMips = FMath::Min(DestTexture2D->GetNumMips(), SrcTexture2D->GetNumMips());
@@ -1647,10 +1187,10 @@ void FD3D12DynamicRHI::RHICopySharedMips(FRHITexture2D* DestTexture2DRHI, FRHITe
 
 	FD3D12Adapter* Adapter = &GetAdapter();
 
-	for (FD3D12Texture2D::FDualLinkedObjectIterator It(DestTexture2D, SrcTexture2D); It; ++It)
+	for (FD3D12Texture::FDualLinkedObjectIterator It(DestTexture2D, SrcTexture2D); It; ++It)
 	{
-		DestTexture2D = static_cast<FD3D12Texture2D*>(It.GetFirst());
-		SrcTexture2D = static_cast<FD3D12Texture2D*>(It.GetSecond());
+		DestTexture2D = static_cast<FD3D12Texture*>(It.GetFirst());
+		SrcTexture2D = static_cast<FD3D12Texture*>(It.GetSecond());
 
 		FD3D12Device* Device = DestTexture2D->GetParentDevice();
 
@@ -1691,194 +1231,456 @@ void FD3D12DynamicRHI::RHICopySharedMips(FRHITexture2D* DestTexture2DRHI, FRHITe
 	}
 }
 
-FTexture2DArrayRHIRef FD3D12DynamicRHI::RHICreateTexture2DArray_RenderThread(class FRHICommandListImmediate& RHICmdList, uint32 SizeX, uint32 SizeY, uint32 SizeZ, uint8 Format, uint32 NumMips, uint32 NumSamples, ETextureCreateFlags Flags, ERHIAccess InResourceState, FRHIResourceCreateInfo& CreateInfo)
+
+/**
+ * Computes the size in memory required by a given texture.
+ *
+ * @param	TextureRHI		- Texture we want to know the size of
+ * @return					- Size in Bytes
+ */
+uint32 FD3D12DynamicRHI::RHIComputeMemorySize(FRHITexture* TextureRHI)
 {
-	check(SizeZ >= 1);
+	if (!TextureRHI)
+	{
+		return 0;
+	}
 
-	return CreateD3D12Texture2D<FD3D12BaseTexture2DArray>(&RHICmdList, SizeX, SizeY, SizeZ, true, false, (EPixelFormat) Format, NumMips, NumSamples, Flags, InResourceState, CreateInfo);
+	FD3D12Texture* Texture = GetD3D12TextureFromRHITexture(TextureRHI);
+	return Texture->ResourceLocation.GetSize();
 }
 
-FTexture2DArrayRHIRef FD3D12DynamicRHI::RHICreateTexture2DArray(uint32 SizeX, uint32 SizeY, uint32 SizeZ, uint8 Format, uint32 NumMips, uint32 NumSamples, ETextureCreateFlags Flags, ERHIAccess InResourceState, FRHIResourceCreateInfo& CreateInfo)
+
+FTexture2DRHIRef FD3D12DynamicRHI::AsyncReallocateTexture2D_RenderThread(class FRHICommandListImmediate& RHICmdList, FRHITexture2D* Texture2DRHI, int32 NewMipCount, int32 NewSizeX, int32 NewSizeY, FThreadSafeCounter* RequestStatus)
 {
-	check(SizeZ >= 1);
+	if (RHICmdList.Bypass())
+	{
+		return FDynamicRHI::AsyncReallocateTexture2D_RenderThread(RHICmdList, Texture2DRHI, NewMipCount, NewSizeX, NewSizeY, RequestStatus);
+	}
+
+	FD3D12Texture* Texture2D = FD3D12DynamicRHI::ResourceCast(Texture2DRHI);
+	
+	FRHITextureDesc Desc = Texture2D->GetDesc();
+	Desc.Extent = FIntPoint(NewSizeX, NewSizeY);
+	Desc.NumMips = NewMipCount;
+
+	FRHITextureCreateDesc CreateDesc(
+		Desc,
+		RHIGetDefaultResourceState(Desc.Flags, false),
+		TEXT("AsyncReallocateTexture2D_RenderThread")
+	);
+
+	// Allocate a new texture.
+	FRHICommandListImmediate* RHIImmediateCmdList = nullptr;
+	ID3D12ResourceAllocator* ResourceAllocator = nullptr;
+	FD3D12Texture* NewTexture2D = CreateD3D12Texture(CreateDesc, RHIImmediateCmdList, ResourceAllocator);
+		
+	ALLOC_COMMAND_CL(RHICmdList, FRHICommandD3D12AsyncReallocateTexture2D)(Texture2D, NewTexture2D, NewMipCount, NewSizeX, NewSizeY, RequestStatus);
 
-	return CreateD3D12Texture2D<FD3D12BaseTexture2DArray>(nullptr, SizeX, SizeY, SizeZ, true, false, (EPixelFormat) Format, NumMips, NumSamples, Flags, InResourceState, CreateInfo);
+	return NewTexture2D;
 }
 
-FTexture3DRHIRef FD3D12DynamicRHI::RHICreateTexture3D_RenderThread(class FRHICommandListImmediate& RHICmdList, uint32 SizeX, uint32 SizeY, uint32 SizeZ, uint8 Format, uint32 NumMips, ETextureCreateFlags Flags, ERHIAccess InResourceState, FRHIResourceCreateInfo& CreateInfo)
+
+/**
+ * Starts an asynchronous texture reallocation. It may complete immediately if the reallocation
+ * could be performed without any reshuffling of texture memory, or if there isn't enough memory.
+ * The specified status counter will be decremented by 1 when the reallocation is complete (success or failure).
+ *
+ * Returns a new reference to the texture, which will represent the new mip count when the reallocation is complete.
+ * RHIGetAsyncReallocateTexture2DStatus() can be used to check the status of an ongoing or completed reallocation.
+ *
+ * @param Texture2D		- Texture to reallocate
+ * @param NewMipCount	- New number of mip-levels
+ * @param NewSizeX		- New width, in pixels
+ * @param NewSizeY		- New height, in pixels
+ * @param RequestStatus	- Will be decremented by 1 when the reallocation is complete (success or failure).
+ * @return				- New reference to the texture, or an invalid reference upon failure
+ */
+FTexture2DRHIRef FD3D12DynamicRHI::RHIAsyncReallocateTexture2D(FRHITexture2D* Texture2DRHI, int32 NewMipCount, int32 NewSizeX, int32 NewSizeY, FThreadSafeCounter* RequestStatus)
 {
-	return CreateD3D12Texture3D(&RHICmdList, SizeX, SizeY, SizeZ, (EPixelFormat) Format, NumMips, Flags, InResourceState, CreateInfo);
+	FD3D12Texture* Texture2D = FD3D12DynamicRHI::ResourceCast(Texture2DRHI);
+
+	FRHITextureDesc Desc = Texture2D->GetDesc();
+	Desc.Extent = FIntPoint(NewSizeX, NewSizeY);
+	Desc.NumMips = NewMipCount;
+
+	FRHITextureCreateDesc CreateDesc(
+		Desc,
+		RHIGetDefaultResourceState(Desc.Flags, false),
+		TEXT("RHIAsyncReallocateTexture2D")
+	);
+
+	// Allocate a new texture.
+	FRHICommandListImmediate* RHIImmediateCmdList = nullptr;
+	ID3D12ResourceAllocator* ResourceAllocator = nullptr;
+	FD3D12Texture* NewTexture2D = CreateD3D12Texture(CreateDesc, nullptr);
+	
+	FRHICommandD3D12AsyncReallocateTexture2D AsyncReallocateTexture2D(Texture2D, NewTexture2D, NewMipCount, NewSizeX, NewSizeY, RequestStatus);
+	AsyncReallocateTexture2D.CopyMips();
+
+	return NewTexture2D;
 }
 
-FTexture3DRHIRef FD3D12DynamicRHI::RHICreateTexture3D(uint32 SizeX, uint32 SizeY, uint32 SizeZ, uint8 Format, uint32 NumMips, ETextureCreateFlags Flags, ERHIAccess InResourceState, FRHIResourceCreateInfo& CreateInfo)
+ETextureReallocationStatus FD3D12DynamicRHI::RHIFinalizeAsyncReallocateTexture2D(FRHITexture2D* Texture2D, bool bBlockUntilCompleted)
 {
-	check(SizeZ >= 1);
-#if PLATFORM_WINDOWS || PLATFORM_HOLOLENS
-	return CreateD3D12Texture3D(nullptr, SizeX, SizeY, SizeZ, (EPixelFormat) Format, NumMips, Flags, InResourceState, CreateInfo);
-#else
-	checkf(false, TEXT("XBOX_CODE_MERGE : Removed. The Xbox platform version should be used."));
-	return nullptr;
-#endif // PLATFORM_WINDOWS || PLATFORM_HOLOLENS
+	return TexRealloc_Succeeded;
 }
 
-/**
- * Computes the size in memory required by a given texture.
- *
- * @param	TextureRHI		- Texture we want to know the size of
- * @return					- Size in Bytes
- */
-uint32 FD3D12DynamicRHI::RHIComputeMemorySize(FRHITexture* TextureRHI)
+ETextureReallocationStatus FD3D12DynamicRHI::RHICancelAsyncReallocateTexture2D(FRHITexture2D* Texture2D, bool bBlockUntilCompleted)
 {
-	if (!TextureRHI)
+	return TexRealloc_Succeeded;
+}
+
+
+///////////////////////////////////////////////////////////////////////////////////////////
+// FD3D12Texture
+///////////////////////////////////////////////////////////////////////////////////////////
+
+FD3D12Texture::~FD3D12Texture()
+{
+	if (IsHeadLink())
 	{
-		return 0;
+		// Only call this once for a LDA chain
+		FD3D12TextureStats::D3D12TextureDeleted(*this);
 	}
+}
 
-	FD3D12TextureBase* Texture = GetD3D12TextureFromRHITexture(TextureRHI);
-	return Texture->GetMemorySize();
+#if RHI_ENABLE_RESOURCE_INFO
+bool FD3D12Texture::GetResourceInfo(FRHIResourceInfo& OutResourceInfo) const
+{
+	OutResourceInfo = FRHIResourceInfo{};
+	OutResourceInfo.Name = GetName();
+	OutResourceInfo.Type = GetType();
+	OutResourceInfo.VRamAllocation.AllocationSize = ResourceLocation.GetSize();
+	OutResourceInfo.IsTransient = ResourceLocation.IsTransient();
+	return true;
 }
+#endif
 
+void* FD3D12Texture::GetNativeResource() const
+{
+	void* NativeResource = nullptr;
+	FD3D12Resource* Resource = GetResource();
+	if (Resource)
+	{
+		NativeResource = Resource->GetResource();
+	}
+	if (!NativeResource)
+	{
+		FD3D12Texture* Base = GetD3D12TextureFromRHITexture((FRHITexture*)this);
+		if (Base)
+		{
+			Resource = Base->GetResource();
+			if (Resource)
+			{
+				NativeResource = Resource->GetResource();
+			}
+		}
+	}
+	return NativeResource;
+}
 
-static void DoAsyncReallocateTexture2D(FD3D12Texture2D* Texture2D, FD3D12Texture2D* NewTexture2D, int32 NewMipCount, int32 NewSizeX, int32 NewSizeY, FThreadSafeCounter* RequestStatus)
+FRHIDescriptorHandle FD3D12Texture::GetDefaultBindlessHandle() const
 {
-	// Use the GPU to asynchronously copy the old mip-maps into the new texture.
-	const uint32 NumSharedMips = FMath::Min(Texture2D->GetNumMips(), NewTexture2D->GetNumMips());
-	const uint32 SourceMipOffset = Texture2D->GetNumMips() - NumSharedMips;
-	const uint32 DestMipOffset = NewTexture2D->GetNumMips() - NumSharedMips;
+	if (FRHIShaderResourceView* View = GetShaderResourceView())
+	{
+		return View->GetBindlessHandle();
+	}
+	return FRHIDescriptorHandle();
+}
 
-	uint32 destSubresource = 0;
-	uint32 srcSubresource = 0;
+void FD3D12Texture::SetDepthStencilView(FD3D12DepthStencilView* View, uint32 SubResourceIndex)
+{
+	if (SubResourceIndex < FExclusiveDepthStencil::MaxIndex)
+	{
+		DepthStencilViews[SubResourceIndex] = View;
+	}
+	else
+	{
+		check(false);
+	}
+}
 
-	for (FD3D12Texture2D::FDualLinkedObjectIterator It(Texture2D, NewTexture2D); It; ++It)
+void FD3D12Texture::SetRenderTargetViewIndex(FD3D12RenderTargetView* View, uint32 SubResourceIndex)
+{
+	if (SubResourceIndex < (uint32)RenderTargetViews.Num())
+	{
+		RenderTargetViews[SubResourceIndex] = View;
+	}
+	else
 	{
-		Texture2D = static_cast<FD3D12Texture2D*>(It.GetFirst());
-		NewTexture2D = static_cast<FD3D12Texture2D*>(It.GetSecond());
+		check(false);
+	}
+}
 
-		FD3D12Device* Device = Texture2D->GetParentDevice();
 
-		FD3D12CommandListHandle& hCommandList = Device->GetDefaultCommandContext().CommandListHandle;
+void FD3D12Texture::CreateViews()
+{
+	TRACE_CPUPROFILER_EVENT_SCOPE(D3D12RHI::CreateViews);
+
+	D3D12_RESOURCE_DESC ResourceDesc = ResourceLocation.GetResource()->GetDesc();
+	const FRHITextureDesc Desc = GetDesc();
+
+	const bool bSRGB = EnumHasAnyFlags(Desc.Flags, TexCreate_SRGB);
+	const DXGI_FORMAT PlatformResourceFormat = GetPlatformTextureResourceFormat((DXGI_FORMAT)GPixelFormats[Desc.Format].PlatformFormat, Desc.Flags);
+	const DXGI_FORMAT PlatformShaderResourceFormat = FindShaderResourceDXGIFormat(PlatformResourceFormat, bSRGB);
+	const DXGI_FORMAT PlatformRenderTargetFormat = FindShaderResourceDXGIFormat(PlatformResourceFormat, bSRGB);
+	const DXGI_FORMAT PlatformDepthStencilFormat = FindDepthStencilDXGIFormat(PlatformResourceFormat);
+
+	const bool bTexture2D = Desc.IsTexture2D();
+	const bool bTexture3D = Desc.IsTexture3D();
+	const bool bCubeTexture = Desc.IsTextureCube();
+	const bool bTextureArray = Desc.IsTextureArray();
+
+	// Set up the texture bind flags.
+	bool bCreateRTV = EnumHasAnyFlags(ResourceDesc.Flags, D3D12_RESOURCE_FLAG_ALLOW_RENDER_TARGET);
+	bool bCreateDSV = EnumHasAnyFlags(ResourceDesc.Flags, D3D12_RESOURCE_FLAG_ALLOW_DEPTH_STENCIL);
+	bool bCreateShaderResource = !EnumHasAnyFlags(ResourceDesc.Flags, D3D12_RESOURCE_FLAG_DENY_SHADER_RESOURCE);
+
+	if (EnumHasAllFlags(Desc.Flags, TexCreate_CPUReadback))
+	{
+		check(!EnumHasAnyFlags(Desc.Flags, TexCreate_RenderTargetable | TexCreate_DepthStencilTargetable | TexCreate_ShaderResource));
+		bCreateShaderResource = false;
+	}
+
+	if (EnumHasAnyFlags(Desc.Flags, TexCreate_DisableSRVCreation))
+	{
+		bCreateShaderResource = false;
+	}
+
+	if (Desc.Format == PF_NV12)
+	{
+		bCreateRTV = false;
+		bCreateShaderResource = false;
+	}
+
+	const bool bIsMultisampled = ResourceDesc.SampleDesc.Count > 1;
+
+	FD3D12Device* Device = GetParentDevice();
 
-		FScopeResourceBarrier ScopeResourceBarrierDest(hCommandList, NewTexture2D->GetResource(), NewTexture2D->GetResource()->GetDefaultResourceState(), D3D12_RESOURCE_STATE_COPY_DEST, D3D12_RESOURCE_BARRIER_ALL_SUBRESOURCES);
-		FScopeResourceBarrier ScopeResourceBarrierSource(hCommandList, Texture2D->GetResource(), Texture2D->GetResource()->GetDefaultResourceState(), D3D12_RESOURCE_STATE_COPY_SOURCE, D3D12_RESOURCE_BARRIER_ALL_SUBRESOURCES);
-		hCommandList.FlushResourceBarriers();	// Must flush so the desired state is actually set.
+	if (bCreateRTV)
+	{
+		if (bTexture3D)
+		{
+			// Create a single render-target-view for the texture.
+			D3D12_RENDER_TARGET_VIEW_DESC RTVDesc;
+			FMemory::Memzero(&RTVDesc, sizeof(RTVDesc));
+			RTVDesc.Format = PlatformRenderTargetFormat;
+			RTVDesc.ViewDimension = D3D12_RTV_DIMENSION_TEXTURE3D;
+			RTVDesc.Texture3D.MipSlice = 0;
+			RTVDesc.Texture3D.FirstWSlice = 0;
+			RTVDesc.Texture3D.WSize = Desc.Depth;
 
-		for (uint32 MipIndex = 0; MipIndex < NumSharedMips; ++MipIndex)
+			SetRenderTargetView(new FD3D12RenderTargetView(Device, RTVDesc, this));
+		}
+		else
 		{
-			// Use the GPU to copy between mip-maps.
-			// This is serialized with other D3D commands, so it isn't necessary to increment Counter to signal a pending asynchronous copy.
+			const bool bCreateRTVsPerSlice = EnumHasAnyFlags(Desc.Flags, TexCreate_TargetArraySlicesIndependently) && (bTextureArray || bCubeTexture);
+			SetNumRenderTargetViews(bCreateRTVsPerSlice ? Desc.NumMips * ResourceDesc.DepthOrArraySize : Desc.NumMips);
 
-			srcSubresource = CalcSubresource(MipIndex + SourceMipOffset, 0, Texture2D->GetNumMips());
-			destSubresource = CalcSubresource(MipIndex + DestMipOffset, 0, NewTexture2D->GetNumMips());
+			// Create a render target view for each mip
+			uint32 RTVIndex = 0;
+			for (uint32 MipIndex = 0; MipIndex < Desc.NumMips; MipIndex++)
+			{
+				if (bCreateRTVsPerSlice)
+				{
+					SetCreatedRTVsPerSlice(true, ResourceDesc.DepthOrArraySize);
 
-			CD3DX12_TEXTURE_COPY_LOCATION DestCopyLocation(NewTexture2D->GetResource()->GetResource(), destSubresource);
-			CD3DX12_TEXTURE_COPY_LOCATION SourceCopyLocation(Texture2D->GetResource()->GetResource(), srcSubresource);
+					for (uint32 SliceIndex = 0; SliceIndex < ResourceDesc.DepthOrArraySize; SliceIndex++)
+					{
+						D3D12_RENDER_TARGET_VIEW_DESC RTVDesc;
+						FMemory::Memzero(RTVDesc);
 
-			Device->GetDefaultCommandContext().numCopies++;
-			hCommandList->CopyTextureRegion(
-				&DestCopyLocation,
-				0, 0, 0,
-				&SourceCopyLocation,
-				nullptr);
+						RTVDesc.Format = PlatformRenderTargetFormat;
+						RTVDesc.ViewDimension = D3D12_RTV_DIMENSION_TEXTURE2DARRAY;
+						RTVDesc.Texture2DArray.FirstArraySlice = SliceIndex;
+						RTVDesc.Texture2DArray.ArraySize = 1;
+						RTVDesc.Texture2DArray.MipSlice = MipIndex;
+						RTVDesc.Texture2DArray.PlaneSlice = GetPlaneSliceFromViewFormat(PlatformResourceFormat, RTVDesc.Format);
 
-			hCommandList.UpdateResidency(NewTexture2D->GetResource());
-			hCommandList.UpdateResidency(Texture2D->GetResource());
+						SetRenderTargetViewIndex(new FD3D12RenderTargetView(Device, RTVDesc, this), RTVIndex++);
+					}
+				}
+				else
+				{
+					D3D12_RENDER_TARGET_VIEW_DESC RTVDesc;
+					FMemory::Memzero(RTVDesc);
 
-			Device->GetDefaultCommandContext().ConditionalFlushCommandList();
+					RTVDesc.Format = PlatformRenderTargetFormat;
 
-			DEBUG_EXECUTE_COMMAND_CONTEXT(Device->GetDefaultCommandContext());
+					if (bTextureArray || bCubeTexture)
+					{
+						if (bIsMultisampled)
+						{
+							RTVDesc.ViewDimension = D3D12_RTV_DIMENSION_TEXTURE2DMSARRAY;
+							RTVDesc.Texture2DMSArray.FirstArraySlice = 0;
+							RTVDesc.Texture2DMSArray.ArraySize = ResourceDesc.DepthOrArraySize;
+						}
+						else
+						{
+							RTVDesc.ViewDimension = D3D12_RTV_DIMENSION_TEXTURE2DARRAY;
+							RTVDesc.Texture2DArray.FirstArraySlice = 0;
+							RTVDesc.Texture2DArray.ArraySize = ResourceDesc.DepthOrArraySize;
+							RTVDesc.Texture2DArray.MipSlice = MipIndex;
+							RTVDesc.Texture2DArray.PlaneSlice = GetPlaneSliceFromViewFormat(PlatformResourceFormat, RTVDesc.Format);
+						}
+					}
+					else
+					{
+						if (bIsMultisampled)
+						{
+							RTVDesc.ViewDimension = D3D12_RTV_DIMENSION_TEXTURE2DMS;
+							// Nothing to set
+						}
+						else
+						{
+							RTVDesc.ViewDimension = D3D12_RTV_DIMENSION_TEXTURE2D;
+							RTVDesc.Texture2D.MipSlice = MipIndex;
+							RTVDesc.Texture2D.PlaneSlice = GetPlaneSliceFromViewFormat(PlatformResourceFormat, RTVDesc.Format);
+						}
+					}
+
+					SetRenderTargetViewIndex(new FD3D12RenderTargetView(Device, RTVDesc, this), RTVIndex++);
+				}
+			}
+		}
+	}
+
+	if (bCreateDSV)
+	{
+		check(!bTexture3D);
+
+		// Create a depth-stencil-view for the texture.
+		D3D12_DEPTH_STENCIL_VIEW_DESC DSVDesc = {};
+		DSVDesc.Format = FindDepthStencilDXGIFormat(PlatformResourceFormat);
+		if (bTextureArray || bCubeTexture)
+		{
+			if (bIsMultisampled)
+			{
+				DSVDesc.ViewDimension = D3D12_DSV_DIMENSION_TEXTURE2DMSARRAY;
+				DSVDesc.Texture2DMSArray.FirstArraySlice = 0;
+				DSVDesc.Texture2DMSArray.ArraySize = ResourceDesc.DepthOrArraySize;
+			}
+			else
+			{
+				DSVDesc.ViewDimension = D3D12_DSV_DIMENSION_TEXTURE2DARRAY;
+				DSVDesc.Texture2DArray.FirstArraySlice = 0;
+				DSVDesc.Texture2DArray.ArraySize = ResourceDesc.DepthOrArraySize;
+				DSVDesc.Texture2DArray.MipSlice = 0;
+			}
+		}
+		else
+		{
+			if (bIsMultisampled)
+			{
+				DSVDesc.ViewDimension = D3D12_DSV_DIMENSION_TEXTURE2DMS;
+				// Nothing to set
+			}
+			else
+			{
+				DSVDesc.ViewDimension = D3D12_DSV_DIMENSION_TEXTURE2D;
+				DSVDesc.Texture2D.MipSlice = 0;
+			}
+		}
+
+		const bool HasStencil = HasStencilBits(DSVDesc.Format);
+		for (uint32 AccessType = 0; AccessType < FExclusiveDepthStencil::MaxIndex; ++AccessType)
+		{
+			// Create a read-only access views for the texture.
+			DSVDesc.Flags = (AccessType & FExclusiveDepthStencil::DepthRead_StencilWrite) ? D3D12_DSV_FLAG_READ_ONLY_DEPTH : D3D12_DSV_FLAG_NONE;
+			if (HasStencil)
+			{
+				DSVDesc.Flags |= (AccessType & FExclusiveDepthStencil::DepthWrite_StencilRead) ? D3D12_DSV_FLAG_READ_ONLY_STENCIL : D3D12_DSV_FLAG_NONE;
+			}
+
+			SetDepthStencilView(new FD3D12DepthStencilView(Device, DSVDesc, this, HasStencil), AccessType);
+		}
+	}
+
+	// Create a shader resource view for the texture.
+	if (bCreateShaderResource)
+	{
+		D3D12_SHADER_RESOURCE_VIEW_DESC SRVDesc = {};
+		SRVDesc.Shader4ComponentMapping = D3D12_DEFAULT_SHADER_4_COMPONENT_MAPPING;
+		SRVDesc.Format = PlatformShaderResourceFormat;
+
+		if (bCubeTexture && bTextureArray)
+		{
+			SRVDesc.ViewDimension = D3D12_SRV_DIMENSION_TEXTURECUBEARRAY;
+			SRVDesc.TextureCubeArray.MostDetailedMip = 0;
+			SRVDesc.TextureCubeArray.MipLevels = Desc.NumMips;
+			SRVDesc.TextureCubeArray.First2DArrayFace = 0;
+			SRVDesc.TextureCubeArray.NumCubes = Desc.ArraySize;
+		}
+		else if (bCubeTexture)
+		{
+			SRVDesc.ViewDimension = D3D12_SRV_DIMENSION_TEXTURECUBE;
+			SRVDesc.TextureCube.MostDetailedMip = 0;
+			SRVDesc.TextureCube.MipLevels = Desc.NumMips;
+		}
+		else if (bTextureArray)
+		{
+			if (bIsMultisampled)
+			{
+				SRVDesc.ViewDimension = D3D12_SRV_DIMENSION_TEXTURE2DMSARRAY;
+				SRVDesc.Texture2DMSArray.FirstArraySlice = 0;
+				SRVDesc.Texture2DMSArray.ArraySize = ResourceDesc.DepthOrArraySize;
+				//SRVDesc.Texture2DArray.PlaneSlice = GetPlaneSliceFromViewFormat(PlatformResourceFormat, SRVDesc.Format);
+			}
+			else
+			{
+				SRVDesc.ViewDimension = D3D12_SRV_DIMENSION_TEXTURE2DARRAY;
+				SRVDesc.Texture2DArray.MostDetailedMip = 0;
+				SRVDesc.Texture2DArray.MipLevels = Desc.NumMips;
+				SRVDesc.Texture2DArray.FirstArraySlice = 0;
+				SRVDesc.Texture2DArray.ArraySize = ResourceDesc.DepthOrArraySize;
+				SRVDesc.Texture2DArray.PlaneSlice = GetPlaneSliceFromViewFormat(PlatformResourceFormat, SRVDesc.Format);
+			}
+		}
+		else if (bTexture3D)
+		{
+			SRVDesc.ViewDimension = D3D12_SRV_DIMENSION_TEXTURE3D;
+			SRVDesc.Texture3D.MipLevels = Desc.NumMips;
+			SRVDesc.Texture3D.MostDetailedMip = 0;
+		}
+		else
+		{
+			if (bIsMultisampled)
+			{
+				SRVDesc.ViewDimension = D3D12_SRV_DIMENSION_TEXTURE2DMS;
+				// Nothing to set
+			}
+			else
+			{
+				SRVDesc.ViewDimension = D3D12_SRV_DIMENSION_TEXTURE2D;
+				SRVDesc.Texture2D.MostDetailedMip = 0;
+				SRVDesc.Texture2D.MipLevels = Desc.NumMips;
+				SRVDesc.Texture2D.PlaneSlice = GetPlaneSliceFromViewFormat(PlatformResourceFormat, SRVDesc.Format);
+			}
 		}
-	}
 
-	// Decrement the thread-safe counter used to track the completion of the reallocation, since D3D handles sequencing the
-	// async mip copies with other D3D calls.
-	RequestStatus->Decrement();
+		SetShaderResourceView(new FD3D12ShaderResourceView(this, SRVDesc));
+	}
 }
 
 
-void FRHICommandD3D12AsyncReallocateTexture2D::Execute(FRHICommandListBase& RHICmdList)
+void FD3D12Texture::AliasResources(FD3D12Texture* Texture)
 {
-	DoAsyncReallocateTexture2D(OldTexture, NewTexture, NewMipCount, NewSizeX, NewSizeY, RequestStatus);
-}
+	// Alias the location, will perform an addref underneath
+	FD3D12ResourceLocation::Alias(ResourceLocation, Texture->ResourceLocation);
 
+	ShaderResourceView = Texture->ShaderResourceView;
 
-FTexture2DRHIRef FD3D12DynamicRHI::AsyncReallocateTexture2D_RenderThread(class FRHICommandListImmediate& RHICmdList, FRHITexture2D* Texture2DRHI, int32 NewMipCount, int32 NewSizeX, int32 NewSizeY, FThreadSafeCounter* RequestStatus)
-{
-	if (RHICmdList.Bypass())
+	for (uint32 Index = 0; Index < FExclusiveDepthStencil::MaxIndex; Index++)
 	{
-		return FDynamicRHI::AsyncReallocateTexture2D_RenderThread(RHICmdList, Texture2DRHI, NewMipCount, NewSizeX, NewSizeY, RequestStatus);
+		DepthStencilViews[Index] = Texture->DepthStencilViews[Index];
+	}
+	for (int32 Index = 0; Index < Texture->RenderTargetViews.Num(); Index++)
+	{
+		RenderTargetViews[Index] = Texture->RenderTargetViews[Index];
 	}
-
-	FD3D12Texture2D* Texture2D = FD3D12DynamicRHI::ResourceCast(Texture2DRHI);
-	
-	// Allocate a new texture.
-	FRHIResourceCreateInfo CreateInfo(TEXT("AsyncReallocateTexture2D_RenderThread"));
-	ERHIAccess RHIAccess = ERHIAccess::Unknown;
-	FD3D12Texture2D* NewTexture2D = CreateD3D12Texture2D<FD3D12BaseTexture2D>(nullptr, NewSizeX, NewSizeY, 1, false, false, Texture2DRHI->GetFormat(), NewMipCount, 1, Texture2DRHI->GetFlags(), RHIAccess, CreateInfo);
-	
-	ALLOC_COMMAND_CL(RHICmdList, FRHICommandD3D12AsyncReallocateTexture2D)(Texture2D, NewTexture2D, NewMipCount, NewSizeX, NewSizeY, RequestStatus);
-
-	return NewTexture2D;
-}
-
-
-/**
- * Starts an asynchronous texture reallocation. It may complete immediately if the reallocation
- * could be performed without any reshuffling of texture memory, or if there isn't enough memory.
- * The specified status counter will be decremented by 1 when the reallocation is complete (success or failure).
- *
- * Returns a new reference to the texture, which will represent the new mip count when the reallocation is complete.
- * RHIGetAsyncReallocateTexture2DStatus() can be used to check the status of an ongoing or completed reallocation.
- *
- * @param Texture2D		- Texture to reallocate
- * @param NewMipCount	- New number of mip-levels
- * @param NewSizeX		- New width, in pixels
- * @param NewSizeY		- New height, in pixels
- * @param RequestStatus	- Will be decremented by 1 when the reallocation is complete (success or failure).
- * @return				- New reference to the texture, or an invalid reference upon failure
- */
-FTexture2DRHIRef FD3D12DynamicRHI::RHIAsyncReallocateTexture2D(FRHITexture2D* Texture2DRHI, int32 NewMipCount, int32 NewSizeX, int32 NewSizeY, FThreadSafeCounter* RequestStatus)
-{
-	FD3D12Texture2D*  Texture2D = FD3D12DynamicRHI::ResourceCast(Texture2DRHI);
-
-	// Allocate a new texture.
-	FRHIResourceCreateInfo CreateInfo(TEXT("RHIAsyncReallocateTexture2D"));
-	ERHIAccess RHIAccess = ERHIAccess::Unknown;
-	FD3D12Texture2D* NewTexture2D = CreateD3D12Texture2D<FD3D12BaseTexture2D>(nullptr, NewSizeX, NewSizeY, 1, false, false, Texture2DRHI->GetFormat(), NewMipCount, 1, Texture2DRHI->GetFlags(), RHIAccess, CreateInfo);
-	
-	DoAsyncReallocateTexture2D(Texture2D, NewTexture2D, NewMipCount, NewSizeX, NewSizeY, RequestStatus);
-
-	return NewTexture2D;
-}
-
-/**
- * Returns the status of an ongoing or completed texture reallocation:
- *	TexRealloc_Succeeded	- The texture is ok, reallocation is not in progress.
- *	TexRealloc_Failed		- The texture is bad, reallocation is not in progress.
- *	TexRealloc_InProgress	- The texture is currently being reallocated async.
- *
- * @param Texture2D		- Texture to check the reallocation status for
- * @return				- Current reallocation status
- */
-ETextureReallocationStatus FD3D12DynamicRHI::RHIFinalizeAsyncReallocateTexture2D(FRHITexture2D* Texture2D, bool bBlockUntilCompleted)
-{
-	return TexRealloc_Succeeded;
-}
-
-/**
- * Cancels an async reallocation for the specified texture.
- * This should be called for the new texture, not the original.
- *
- * @param Texture				Texture to cancel
- * @param bBlockUntilCompleted	If true, blocks until the cancellation is fully completed
- * @return						Reallocation status
- */
-ETextureReallocationStatus FD3D12DynamicRHI::RHICancelAsyncReallocateTexture2D(FRHITexture2D* Texture2D, bool bBlockUntilCompleted)
-{
-	return TexRealloc_Succeeded;
 }
 
-template<typename RHIResourceType>
-void* TD3D12Texture2D<RHIResourceType>::Lock(class FRHICommandListImmediate* RHICmdList, uint32 MipIndex, uint32 ArrayIndex, EResourceLockMode LockMode, uint32& DestStride)
+void* FD3D12Texture::Lock(class FRHICommandListImmediate* RHICmdList, uint32 MipIndex, uint32 ArrayIndex, EResourceLockMode LockMode, uint32& DestStride)
 {
 	SCOPE_CYCLE_COUNTER(STAT_D3D12LockTextureTime);
 
@@ -1911,7 +1713,7 @@ void* TD3D12Texture2D<RHIResourceType>::Lock(class FRHICommandListImmediate* RHI
 
 	void* Data = nullptr;
 
-	if (GetParentDevice()->GetOwningRHI()->HandleSpecialLock(Data, MipIndex, ArrayIndex, RHIResourceType::GetFlags(), LockMode, GetTextureLayout(), RawTextureMemory, DestStride))
+	if (GetParentDevice()->GetOwningRHI()->HandleSpecialLock(Data, MipIndex, ArrayIndex, this, LockMode, DestStride))
 	{
 		// nothing left to do...
 		check(Data != nullptr);
@@ -2018,120 +1820,7 @@ void* TD3D12Texture2D<RHIResourceType>::Lock(class FRHICommandListImmediate* RHI
 	return Data;
 }
 
-D3D12_RESOURCE_DESC FD3D12DynamicRHI::GetResourceDesc(const FRHITextureCreateInfo& CreateInfo) const
-{
-	D3D12_RESOURCE_DESC Desc;
-
-	const DXGI_FORMAT Format = GetPlatformTextureResourceFormat((DXGI_FORMAT)GPixelFormats[CreateInfo.Format].PlatformFormat, CreateInfo.Flags);
-
-	if (CreateInfo.Dimension != ETextureDimension::Texture3D)
-	{
-		if (CreateInfo.IsTextureCube())
-		{
-			check(CreateInfo.Extent.X <= (int32)GetMaxCubeTextureDimension());
-			check(CreateInfo.Extent.X == CreateInfo.Extent.Y);
-		}
-		else
-		{
-			check(CreateInfo.Extent.X <= (int32)GetMax2DTextureDimension());
-			check(CreateInfo.Extent.Y <= (int32)GetMax2DTextureDimension());
-		}
-
-		if (CreateInfo.IsTextureArray())
-		{
-			check(CreateInfo.ArraySize <= (int32)GetMaxTextureArrayLayers());
-		}
-
-		uint32 ActualMSAACount = CreateInfo.NumSamples;
-		uint32 ActualMSAAQuality = GetMaxMSAAQuality(ActualMSAACount);
-
-		// 0xffffffff means not supported
-		if (ActualMSAAQuality == 0xffffffff || EnumHasAnyFlags(CreateInfo.Flags, TexCreate_Shared))
-		{
-			// no MSAA
-			ActualMSAACount = 1;
-			ActualMSAAQuality = 0;
-		}
-
-		Desc = CD3DX12_RESOURCE_DESC::Tex2D(
-			Format,
-			CreateInfo.Extent.X,
-			CreateInfo.Extent.Y,
-			CreateInfo.ArraySize * (CreateInfo.IsTextureCube() ? 6 : 1),  // Array size
-			CreateInfo.NumMips,
-			ActualMSAACount,
-			ActualMSAAQuality,
-			D3D12_RESOURCE_FLAG_NONE);  // Add misc flags later
-
-		if (EnumHasAnyFlags(CreateInfo.Flags, TexCreate_Shared))
-		{
-			Desc.Flags |= D3D12_RESOURCE_FLAG_ALLOW_SIMULTANEOUS_ACCESS;
-		}
-
-		if (EnumHasAnyFlags(CreateInfo.Flags, TexCreate_RenderTargetable))
-		{
-			check(!EnumHasAnyFlags(CreateInfo.Flags, TexCreate_DepthStencilTargetable | TexCreate_ResolveTargetable));
-			Desc.Flags |= D3D12_RESOURCE_FLAG_ALLOW_RENDER_TARGET;
-		}
-		else if (EnumHasAnyFlags(CreateInfo.Flags, TexCreate_DepthStencilTargetable))
-		{
-			check(!EnumHasAnyFlags(CreateInfo.Flags, TexCreate_RenderTargetable | TexCreate_ResolveTargetable));
-			Desc.Flags |= D3D12_RESOURCE_FLAG_ALLOW_DEPTH_STENCIL;
-		}
-		else if (EnumHasAnyFlags(CreateInfo.Flags, TexCreate_ResolveTargetable))
-		{
-			check(!EnumHasAnyFlags(CreateInfo.Flags, TexCreate_RenderTargetable | TexCreate_DepthStencilTargetable));
-			if (CreateInfo.Format == PF_DepthStencil || CreateInfo.Format == PF_ShadowDepth || CreateInfo.Format == PF_D24)
-			{
-				Desc.Flags |= D3D12_RESOURCE_FLAG_ALLOW_DEPTH_STENCIL;
-			}
-			else
-			{
-				Desc.Flags |= D3D12_RESOURCE_FLAG_ALLOW_RENDER_TARGET;
-			}
-		}
-
-		if (EnumHasAnyFlags(CreateInfo.Flags, TexCreate_UAV))
-		{
-			Desc.Flags |= D3D12_RESOURCE_FLAG_ALLOW_UNORDERED_ACCESS;
-		}
-
-		if (EnumHasAnyFlags(CreateInfo.Flags, TexCreate_DepthStencilTargetable) && !EnumHasAnyFlags(CreateInfo.Flags, TexCreate_ShaderResource))
-		{
-			// Only deny shader resources if it's a depth resource that will never be used as SRV
-			Desc.Flags |= D3D12_RESOURCE_FLAG_DENY_SHADER_RESOURCE;
-		}
-	}
-	else // ETextureDimension::Texture3D
-	{
-		check(CreateInfo.Dimension == ETextureDimension::Texture3D)
-		check(!EnumHasAnyFlags(CreateInfo.Flags, TexCreate_DepthStencilTargetable | TexCreate_ResolveTargetable));
-		check(EnumHasAnyFlags(CreateInfo.Flags, TexCreate_ShaderResource));
-
-		Desc = CD3DX12_RESOURCE_DESC::Tex3D(
-			Format,
-			CreateInfo.Extent.X,
-			CreateInfo.Extent.Y,
-			CreateInfo.Depth,
-			CreateInfo.NumMips);
-
-		if (EnumHasAnyFlags(CreateInfo.Flags, TexCreate_UAV))
-		{
-			Desc.Flags |= D3D12_RESOURCE_FLAG_ALLOW_UNORDERED_ACCESS;
-		}
-
-		if (EnumHasAnyFlags(CreateInfo.Flags, TexCreate_RenderTargetable))
-		{
-			Desc.Flags |= D3D12_RESOURCE_FLAG_ALLOW_RENDER_TARGET;
-		}
-	}
-
-	Desc.Alignment = D3D12_DEFAULT_RESOURCE_PLACEMENT_ALIGNMENT;
-
-	return Desc;
-}
-
-void FD3D12TextureBase::UpdateTexture(uint32 MipIndex, uint32 DestX, uint32 DestY, uint32 DestZ, const D3D12_TEXTURE_COPY_LOCATION& SourceCopyLocation)
+void FD3D12Texture::UpdateTexture(uint32 MipIndex, uint32 DestX, uint32 DestY, uint32 DestZ, const D3D12_TEXTURE_COPY_LOCATION& SourceCopyLocation)
 {
 	LLM_SCOPE_BYNAME(TEXT("D3D12CopyTextureRegion"));
 	FD3D12CommandContext& DefaultContext = GetParentDevice()->GetDefaultCommandContext();
@@ -2157,7 +1846,7 @@ void FD3D12TextureBase::UpdateTexture(uint32 MipIndex, uint32 DestX, uint32 Dest
 	DEBUG_EXECUTE_COMMAND_CONTEXT(DefaultContext);
 }
 
-void FD3D12TextureBase::CopyTextureRegion(uint32 DestX, uint32 DestY, uint32 DestZ, FD3D12TextureBase* SourceTexture, const D3D12_BOX& SourceBox)
+void FD3D12Texture::CopyTextureRegion(uint32 DestX, uint32 DestY, uint32 DestZ, FD3D12Texture* SourceTexture, const D3D12_BOX& SourceBox)
 {
 	FD3D12CommandContext& DefaultContext = GetParentDevice()->GetDefaultCommandContext();
 	FD3D12CommandListHandle& CommandListHandle = DefaultContext.CommandListHandle;
@@ -2179,10 +1868,15 @@ void FD3D12TextureBase::CopyTextureRegion(uint32 DestX, uint32 DestY, uint32 Des
 	CommandListHandle.UpdateResidency(GetResource());
 }
 
-void FD3D12TextureBase::InitializeTextureData(FRHICommandListImmediate* RHICmdList, const void* InitData, uint32 InitDataSize, uint32 SizeX, uint32 SizeY, uint32 SizeZ, uint32 NumSlices, uint32 NumMips, EPixelFormat Format, D3D12_RESOURCE_STATES DestinationState)
+void FD3D12Texture::InitializeTextureData(FRHICommandListImmediate* RHICmdList, const FRHITextureCreateDesc& CreateDesc, D3D12_RESOURCE_STATES DestinationState)
 {
 	// each mip of each array slice counts as a subresource
-	uint32 NumSubresources = NumMips * NumSlices;
+	uint16 ArraySize = CreateDesc.IsTextureArray() ? CreateDesc.ArraySize : 1;
+	if (CreateDesc.IsTextureCube())
+	{
+		ArraySize *= 6;
+	}
+	uint32 NumSubresources = CreateDesc.NumMips * ArraySize;
 
 	FD3D12Device* Device = GetParentDevice();
 
@@ -2204,7 +1898,7 @@ void FD3D12TextureBase::InitializeTextureData(FRHICommandListImmediate* RHICmdLi
 	FD3D12ResourceLocation SrcResourceLoc(Device);
 	uint8* DstDataBase = (uint8*) Device->GetDefaultFastAllocator().Allocate(Size, D3D12_TEXTURE_DATA_PLACEMENT_ALIGNMENT, &SrcResourceLoc);
 
-	const uint8* SrcData = (const uint8*) InitData;
+	const uint8* SrcData = (const uint8*)CreateDesc.BulkData->GetResourceBulkData();
 	for (uint32 Subresource = 0; Subresource < NumSubresources; Subresource++)
 	{
 		uint8* DstData = DstDataBase + Footprints[Subresource].Offset;
@@ -2231,7 +1925,7 @@ void FD3D12TextureBase::InitializeTextureData(FRHICommandListImmediate* RHICmdLi
 		}
 	}
 
-	check(SrcData == (uint8*) InitData + InitDataSize);
+	check(SrcData == (uint8*)CreateDesc.BulkData->GetResourceBulkData() + CreateDesc.BulkData->GetResourceBulkDataSize());
 
 	if (ShouldDeferCmdListOperation(RHICmdList))
 	{
@@ -2249,16 +1943,14 @@ void FD3D12TextureBase::InitializeTextureData(FRHICommandListImmediate* RHICmdLi
 	}
 }
 
-template<typename RHIResourceType>
-void TD3D12Texture2D<RHIResourceType>::Unlock(class FRHICommandListImmediate* RHICmdList, uint32 MipIndex, uint32 ArrayIndex)
+void FD3D12Texture::Unlock(class FRHICommandListImmediate* RHICmdList, uint32 MipIndex, uint32 ArrayIndex)
 {
 	SCOPE_CYCLE_COUNTER(STAT_D3D12UnlockTextureTime);
 
 	UnlockInternal(RHICmdList, ++FLinkedObjectIterator(this), MipIndex, ArrayIndex);
 }
 
-template<typename RHIResourceType>
-void TD3D12Texture2D<RHIResourceType>::UnlockInternal(class FRHICommandListImmediate* RHICmdList, FLinkedObjectIterator NextObject, uint32 MipIndex, uint32 ArrayIndex)
+void FD3D12Texture::UnlockInternal(class FRHICommandListImmediate* RHICmdList, FLinkedObjectIterator NextObject, uint32 MipIndex, uint32 ArrayIndex)
 {
 	// Calculate the subresource index corresponding to the specified mip-map.
 	const uint32 Subresource = CalcSubresource(MipIndex, ArrayIndex, this->GetNumMips());
@@ -2270,17 +1962,13 @@ void TD3D12Texture2D<RHIResourceType>::UnlockInternal(class FRHICommandListImmed
 	const uint32 MipSizeX = FMath::Max(this->GetSizeX() >> MipIndex, BlockSizeX);
 	const uint32 MipSizeY = FMath::Max(this->GetSizeY() >> MipIndex, BlockSizeY);
 
-	auto* FirstObject = static_cast<TD3D12Texture2D<RHIResourceType>*>(GetFirstLinkedObject());
+	auto* FirstObject = static_cast<FD3D12Texture*>(GetFirstLinkedObject());
 	TMap<uint32, FD3D12LockedResource*>& Map = FirstObject->LockedMap;
 	FD3D12LockedResource* LockedResource = Map[Subresource];
 
 	check(LockedResource);
 
-#if !PLATFORM_SUPPORTS_VIRTUAL_TEXTURES
-	void* RawTextureMemory = (void*)ResourceLocation.GetGPUVirtualAddress();
-#endif
-
-	if (GetParentDevice()->GetOwningRHI()->HandleSpecialUnlock(RHICmdList, MipIndex, RHIResourceType::GetFlags(), GetTextureLayout(), RawTextureMemory))
+	if (GetParentDevice()->GetOwningRHI()->HandleSpecialUnlock(RHICmdList, MipIndex, this))
 	{
 		// nothing left to do...
 	}
@@ -2325,7 +2013,7 @@ void TD3D12Texture2D<RHIResourceType>::UnlockInternal(class FRHICommandListImmed
 			if (NextObject)
 			{
 				// We pass the first link in the chain as that's the guy that got locked
-				((TD3D12Texture2D<RHIResourceType>*)NextObject.Get())->UnlockInternal(RHICmdList, ++NextObject, MipIndex, ArrayIndex);
+				((FD3D12Texture*)NextObject.Get())->UnlockInternal(RHICmdList, ++NextObject, MipIndex, ArrayIndex);
 			}
 		}
 	}
@@ -2338,8 +2026,7 @@ void TD3D12Texture2D<RHIResourceType>::UnlockInternal(class FRHICommandListImmed
 	}
 }
 
-template<typename RHIResourceType>
-void TD3D12Texture2D<RHIResourceType>::UpdateTexture2D(class FRHICommandListImmediate* RHICmdList, uint32 MipIndex, const FUpdateTextureRegion2D& UpdateRegion, uint32 SourcePitch, const uint8* SourceData)
+void FD3D12Texture::UpdateTexture2D(class FRHICommandListImmediate* RHICmdList, uint32 MipIndex, const FUpdateTextureRegion2D& UpdateRegion, uint32 SourcePitch, const uint8* SourceData)
 {
 	const FPixelFormatInfo& FormatInfo = GPixelFormats[this->GetFormat()];
 	check(UpdateRegion.Width  %	FormatInfo.BlockSizeX == 0);
@@ -2355,9 +2042,8 @@ void TD3D12Texture2D<RHIResourceType>::UpdateTexture2D(class FRHICommandListImme
 	const uint32 AlignedSourcePitch = Align(SourcePitch, FD3D12_TEXTURE_DATA_PITCH_ALIGNMENT);
 	const uint32 bufferSize = Align(HeightInBlocks*AlignedSourcePitch, D3D12_TEXTURE_DATA_PLACEMENT_ALIGNMENT);
 
-	for (FD3D12TextureBase& TextureBase : *this)
+	for (FD3D12Texture& Texture : *this)
 	{
-		FD3D12Texture2D& Texture = static_cast<FD3D12Texture2D&>(TextureBase);
 		FD3D12ResourceLocation UploadHeapResourceLocation(GetParentDevice());
 		void* pData = GetParentDevice()->GetDefaultFastAllocator().Allocate(bufferSize, D3D12_TEXTURE_DATA_PLACEMENT_ALIGNMENT, &UploadHeapResourceLocation);
 		check(nullptr != pData);
@@ -2412,10 +2098,9 @@ static void GetReadBackHeapDescImpl(D3D12_PLACED_SUBRESOURCE_FOOTPRINT& OutFootp
 	check(OutFootprint.Footprint.Width > 0 && OutFootprint.Footprint.Height > 0);
 }
 
-template<typename RHIResourceType>
-void TD3D12Texture2D<RHIResourceType>::GetReadBackHeapDesc(D3D12_PLACED_SUBRESOURCE_FOOTPRINT& OutFootprint, uint32 InSubresource) const
+void FD3D12Texture::GetReadBackHeapDesc(D3D12_PLACED_SUBRESOURCE_FOOTPRINT& OutFootprint, uint32 InSubresource) const
 {
-	check(EnumHasAnyFlags(RHIResourceType::GetFlags(), TexCreate_CPUReadback));
+	check(EnumHasAnyFlags(GetFlags(), TexCreate_CPUReadback));
 
 	if (InSubresource == 0 && FirstSubresourceFootprint)
 	{
@@ -2423,16 +2108,16 @@ void TD3D12Texture2D<RHIResourceType>::GetReadBackHeapDesc(D3D12_PLACED_SUBRESOU
 		return;
 	}
 
-	FIntVector TextureSize = RHIResourceType::GetSizeXYZ();
+	FIntVector TextureSize = GetSizeXYZ();
 
 	D3D12_RESOURCE_DESC Desc = {};
 	Desc.Dimension        = D3D12_RESOURCE_DIMENSION_TEXTURE2D;
 	Desc.Width            = TextureSize.X;
 	Desc.Height           = TextureSize.Y;
 	Desc.DepthOrArraySize = TextureSize.Z;
-	Desc.MipLevels        = RHIResourceType::GetNumMips();
-	Desc.Format           = (DXGI_FORMAT) GPixelFormats[RHIResourceType::GetFormat()].PlatformFormat;
-	Desc.SampleDesc.Count = RHIResourceType::GetNumSamples();
+	Desc.MipLevels        = GetNumMips();
+	Desc.Format           = (DXGI_FORMAT) GPixelFormats[GetFormat()].PlatformFormat;
+	Desc.SampleDesc.Count = GetNumSamples();
 
 	GetReadBackHeapDescImpl(OutFootprint, GetParentDevice()->GetDevice(), Desc, InSubresource);
 
@@ -2453,14 +2138,14 @@ void* FD3D12DynamicRHI::LockTexture2D_RenderThread(class FRHICommandListImmediat
 	}
 
 	check(TextureRHI);
-	FD3D12Texture2D* Texture = FD3D12DynamicRHI::ResourceCast(TextureRHI);
+	FD3D12Texture* Texture = FD3D12DynamicRHI::ResourceCast(TextureRHI);
 	return Texture->Lock(&RHICmdList, MipIndex, 0, LockMode, DestStride);
 }
 
 void* FD3D12DynamicRHI::RHILockTexture2D(FRHITexture2D* TextureRHI, uint32 MipIndex, EResourceLockMode LockMode, uint32& DestStride, bool bLockWithinMiptail)
 {
 	check(TextureRHI);
-	FD3D12Texture2D*  Texture = FD3D12DynamicRHI::ResourceCast(TextureRHI);
+	FD3D12Texture*  Texture = FD3D12DynamicRHI::ResourceCast(TextureRHI);
 	return Texture->Lock(nullptr, MipIndex, 0, LockMode, DestStride);
 }
 
@@ -2475,56 +2160,56 @@ void FD3D12DynamicRHI::UnlockTexture2D_RenderThread(class FRHICommandListImmedia
 	}
 
 	check(TextureRHI);
-	FD3D12Texture2D* Texture = FD3D12DynamicRHI::ResourceCast(TextureRHI);
+	FD3D12Texture* Texture = FD3D12DynamicRHI::ResourceCast(TextureRHI);
 	Texture->Unlock(&RHICmdList, MipIndex, 0);
 }
 
 void FD3D12DynamicRHI::RHIUnlockTexture2D(FRHITexture2D* TextureRHI, uint32 MipIndex, bool bLockWithinMiptail)
 {
 	check(TextureRHI);
-	FD3D12Texture2D*  Texture = FD3D12DynamicRHI::ResourceCast(TextureRHI);
+	FD3D12Texture*  Texture = FD3D12DynamicRHI::ResourceCast(TextureRHI);
 	Texture->Unlock(nullptr, MipIndex, 0);
 }
 
 void* FD3D12DynamicRHI::LockTexture2DArray_RenderThread(class FRHICommandListImmediate& RHICmdList, FRHITexture2DArray* TextureRHI, uint32 TextureIndex, uint32 MipIndex, EResourceLockMode LockMode, uint32& DestStride, bool bLockWithinMiptail)
 {
 	check(TextureRHI);
-	FD3D12Texture2DArray* Texture = FD3D12DynamicRHI::ResourceCast(TextureRHI);
+	FD3D12Texture* Texture = FD3D12DynamicRHI::ResourceCast(TextureRHI);
 	return Texture->Lock(&RHICmdList, MipIndex, TextureIndex, LockMode, DestStride);
 }
 
 void* FD3D12DynamicRHI::RHILockTexture2DArray(FRHITexture2DArray* TextureRHI, uint32 TextureIndex, uint32 MipIndex, EResourceLockMode LockMode, uint32& DestStride, bool bLockWithinMiptail)
 {
 	check(TextureRHI);
-	FD3D12Texture2DArray*  Texture = FD3D12DynamicRHI::ResourceCast(TextureRHI);
+	FD3D12Texture* Texture = FD3D12DynamicRHI::ResourceCast(TextureRHI);
 	return Texture->Lock(nullptr, MipIndex, TextureIndex, LockMode, DestStride);
 }
 
 void FD3D12DynamicRHI::UnlockTexture2DArray_RenderThread(class FRHICommandListImmediate& RHICmdList, FRHITexture2DArray* TextureRHI, uint32 TextureIndex, uint32 MipIndex, bool bLockWithinMiptail)
 {
 	check(TextureRHI);
-	FD3D12Texture2DArray* Texture = FD3D12DynamicRHI::ResourceCast(TextureRHI);
+	FD3D12Texture* Texture = FD3D12DynamicRHI::ResourceCast(TextureRHI);
 	Texture->Unlock(&RHICmdList, MipIndex, TextureIndex);
 }
 
 void FD3D12DynamicRHI::RHIUnlockTexture2DArray(FRHITexture2DArray* TextureRHI, uint32 TextureIndex, uint32 MipIndex, bool bLockWithinMiptail)
 {
 	check(TextureRHI);
-	FD3D12Texture2DArray*  Texture = FD3D12DynamicRHI::ResourceCast(TextureRHI);
+	FD3D12Texture* Texture = FD3D12DynamicRHI::ResourceCast(TextureRHI);
 	Texture->Unlock(nullptr, MipIndex, TextureIndex);
 }
 
 void FD3D12DynamicRHI::UpdateTexture2D_RenderThread(class FRHICommandListImmediate& RHICmdList, FRHITexture2D* TextureRHI, uint32 MipIndex, const struct FUpdateTextureRegion2D& UpdateRegion, uint32 SourcePitch, const uint8* SourceData)
 {
 	check(TextureRHI);
-	FD3D12Texture2D* Texture = FD3D12DynamicRHI::ResourceCast(TextureRHI);
+	FD3D12Texture* Texture = FD3D12DynamicRHI::ResourceCast(TextureRHI);
 	Texture->UpdateTexture2D(&RHICmdList, MipIndex, UpdateRegion, SourcePitch, SourceData);
 }
 
 void FD3D12DynamicRHI::RHIUpdateTexture2D(FRHITexture2D* TextureRHI, uint32 MipIndex, const FUpdateTextureRegion2D& UpdateRegion, uint32 SourcePitch, const uint8* SourceData)
 {
 	check(TextureRHI);
-	FD3D12Texture2D* Texture = FD3D12DynamicRHI::ResourceCast(TextureRHI);
+	FD3D12Texture* Texture = FD3D12DynamicRHI::ResourceCast(TextureRHI);
 	Texture->UpdateTexture2D(nullptr, MipIndex, UpdateRegion, SourcePitch, SourceData);
 }
 
@@ -2601,11 +2286,10 @@ public:
 
 	void Execute(FRHICommandListBase& RHICmdList)
 	{
-		FD3D12Texture3D* NativeTexture = FD3D12DynamicRHI::ResourceCast(DstTexture.GetReference());
+		FD3D12Texture* NativeTexture = FD3D12DynamicRHI::ResourceCast(DstTexture.GetReference());
 
-		for (FD3D12TextureBase& TextureLinkBase : *NativeTexture)
+		for (FD3D12Texture& TextureLink : *NativeTexture)
 		{
-			FD3D12Texture3D& TextureLink = static_cast<FD3D12Texture3D&>(TextureLinkBase);
 			FD3D12Device* Device = TextureLink.GetParentDevice();
 			FD3D12CommandListHandle& NativeCmdList = Device->GetDefaultCommandContext().CommandListHandle;
 
@@ -2720,7 +2404,7 @@ void FD3D12DynamicRHI::RHIUpdateTexture3D(FRHITexture3D* TextureRHI, uint32 MipI
 {
 	check(IsInRenderingThread());
 
-	FD3D12Texture3D* Texture = FD3D12DynamicRHI::ResourceCast(TextureRHI);
+	FD3D12Texture* Texture = FD3D12DynamicRHI::ResourceCast(TextureRHI);
 	const FPixelFormatInfo& FormatInfo = GPixelFormats[Texture->GetFormat()];
 
 	// Need to round up the height and with by block size.
@@ -2758,14 +2442,14 @@ FUpdateTexture3DData FD3D12DynamicRHI::BeginUpdateTexture3D_Internal(FRHITexture
 {
 	check(IsInRenderingThread());
 	FUpdateTexture3DData UpdateData(TextureRHI, MipIndex, UpdateRegion, 0, 0, nullptr, 0, GFrameNumberRenderThread);
-
+		
 	// Initialize the platform data
 	static_assert(sizeof(FD3D12UpdateTexture3DData) < sizeof(UpdateData.PlatformData), "Platform data in FUpdateTexture3DData too small to support D3D12");
 	FD3D12UpdateTexture3DData* UpdateDataD3D12 = new (&UpdateData.PlatformData[0]) FD3D12UpdateTexture3DData;
 	UpdateDataD3D12->bComputeShaderCopy = false;
 	UpdateDataD3D12->UploadHeapResourceLocation = nullptr;
 
-	FD3D12Texture3D* Texture = FD3D12DynamicRHI::ResourceCast(TextureRHI);
+	FD3D12Texture* Texture = FD3D12DynamicRHI::ResourceCast(TextureRHI);
 	const FPixelFormatInfo& FormatInfo = GPixelFormats[Texture->GetFormat()];
 	check(FormatInfo.BlockSizeZ == 1);
 
@@ -2777,8 +2461,7 @@ FUpdateTexture3DData FD3D12DynamicRHI::BeginUpdateTexture3D_Internal(FRHITexture
 	}
 
 	if (!bDoComputeShaderCopy)
-	{
-	
+	{	
 		const int32 NumBlockX = FMath::DivideAndRoundUp<int32>(UpdateRegion.Width, FormatInfo.BlockSizeX);
 		const int32 NumBlockY = FMath::DivideAndRoundUp<int32>(UpdateRegion.Height, FormatInfo.BlockSizeY);
 
@@ -2842,12 +2525,11 @@ public:
 
 	void Execute(FRHICommandListBase& RHICmdList)
 	{
-		FD3D12Texture3D* NativeTexture = FD3D12DynamicRHI::ResourceCast(DstTexture.GetReference());
+		FD3D12Texture* NativeTexture = FD3D12DynamicRHI::ResourceCast(DstTexture.GetReference());
 		FD3D12Resource* UploadBuffer = SrcResourceLocation->GetResource();
 
-		for (FD3D12TextureBase& TextureLinkBase : *NativeTexture)
+		for (FD3D12Texture& TextureLink : *NativeTexture)
 		{
-			FD3D12Texture3D& TextureLink = static_cast<FD3D12Texture3D&>(TextureLinkBase);
 			FD3D12Device* Device = TextureLink.GetParentDevice();
 			FD3D12CommandListHandle& NativeCmdList = Device->GetDefaultCommandContext().CommandListHandle;
 #if USE_PIX
@@ -2932,29 +2614,10 @@ void FD3D12DynamicRHI::EndUpdateTexture3D_Internal(FUpdateTexture3DData& UpdateD
 /*-----------------------------------------------------------------------------
 	Cubemap texture support.
 	-----------------------------------------------------------------------------*/
-FTextureCubeRHIRef FD3D12DynamicRHI::RHICreateTextureCube_RenderThread(class FRHICommandListImmediate& RHICmdList, uint32 Size, uint8 Format, uint32 NumMips, ETextureCreateFlags Flags, ERHIAccess InResourceState, FRHIResourceCreateInfo& CreateInfo)
-{
-	return CreateD3D12Texture2D<FD3D12BaseTextureCube>(&RHICmdList, Size, Size, 6, false, true, (EPixelFormat) Format, NumMips, 1, Flags, InResourceState, CreateInfo);
-}
-
-FTextureCubeRHIRef FD3D12DynamicRHI::RHICreateTextureCube(uint32 Size, uint8 Format, uint32 NumMips, ETextureCreateFlags Flags, ERHIAccess InResourceState, FRHIResourceCreateInfo& CreateInfo)
-{
-	return CreateD3D12Texture2D<FD3D12BaseTextureCube>(nullptr, Size, Size, 6, false, true, (EPixelFormat) Format, NumMips, 1, Flags, InResourceState, CreateInfo);
-}
-
-FTextureCubeRHIRef FD3D12DynamicRHI::RHICreateTextureCubeArray_RenderThread(class FRHICommandListImmediate& RHICmdList, uint32 Size, uint32 ArraySize, uint8 Format, uint32 NumMips, ETextureCreateFlags Flags, ERHIAccess InResourceState, FRHIResourceCreateInfo& CreateInfo)
-{
-	return CreateD3D12Texture2D<FD3D12BaseTextureCube>(&RHICmdList, Size, Size, 6 * ArraySize, true, true, (EPixelFormat) Format, NumMips, 1, Flags, InResourceState, CreateInfo);
-}
-
-FTextureCubeRHIRef FD3D12DynamicRHI::RHICreateTextureCubeArray(uint32 Size, uint32 ArraySize, uint8 Format, uint32 NumMips, ETextureCreateFlags Flags, ERHIAccess InResourceState, FRHIResourceCreateInfo& CreateInfo)
-{
-	return CreateD3D12Texture2D<FD3D12BaseTextureCube>(nullptr, Size, Size, 6 * ArraySize, true, true, (EPixelFormat) Format, NumMips, 1, Flags, InResourceState, CreateInfo);
-}
 
 void* FD3D12DynamicRHI::RHILockTextureCubeFace(FRHITextureCube* TextureCubeRHI, uint32 FaceIndex, uint32 ArrayIndex, uint32 MipIndex, EResourceLockMode LockMode, uint32& DestStride, bool bLockWithinMiptail)
 {
-	FD3D12TextureCube*  TextureCube = FD3D12DynamicRHI::ResourceCast(TextureCubeRHI);
+	FD3D12Texture*  TextureCube = FD3D12DynamicRHI::ResourceCast(TextureCubeRHI);
 	for (uint32 GPUIndex : TextureCube->GetLinkedObjectsGPUMask())
 	{
 		GetRHIDevice(GPUIndex)->GetDefaultCommandContext().ConditionalClearShaderResource(&TextureCube->ResourceLocation);
@@ -2964,7 +2627,7 @@ void* FD3D12DynamicRHI::RHILockTextureCubeFace(FRHITextureCube* TextureCubeRHI,
 }
 void FD3D12DynamicRHI::RHIUnlockTextureCubeFace(FRHITextureCube* TextureCubeRHI, uint32 FaceIndex, uint32 ArrayIndex, uint32 MipIndex, bool bLockWithinMiptail)
 {
-	FD3D12TextureCube*  TextureCube = FD3D12DynamicRHI::ResourceCast(TextureCubeRHI);
+	FD3D12Texture* TextureCube = FD3D12DynamicRHI::ResourceCast(TextureCubeRHI);
 	uint32 D3DFace = GetD3D12CubeFace((ECubeFace)FaceIndex);
 	TextureCube->Unlock(nullptr, MipIndex, D3DFace + ArrayIndex * 6);
 }
@@ -2972,7 +2635,7 @@ void FD3D12DynamicRHI::RHIUnlockTextureCubeFace(FRHITextureCube* TextureCubeRHI,
 void FD3D12DynamicRHI::RHIBindDebugLabelName(FRHITexture* TextureRHI, const TCHAR* Name)
 {
 #if NAME_OBJECTS
-	FD3D12TextureBase::FLinkedObjectIterator BaseTexture(GetD3D12TextureFromRHITexture(TextureRHI));
+	FD3D12Texture::FLinkedObjectIterator BaseTexture(GetD3D12TextureFromRHITexture(TextureRHI));
 
 	if (GNumExplicitGPUsForRendering > 1)
 	{
@@ -3014,7 +2677,7 @@ void FD3D12DynamicRHI::RHIBindDebugLabelName(FRHITexture* TextureRHI, const TCHA
 
 #if TEXTURE_PROFILER_ENABLED
 
-	FD3D12TextureBase* D3D12Texture = GetD3D12TextureFromRHITexture(TextureRHI);
+	FD3D12Texture* D3D12Texture = GetD3D12TextureFromRHITexture(TextureRHI);
 	
 	if (!EnumHasAnyFlags(TextureRHI->GetFlags(), TexCreate_Virtual)
 		&& !D3D12Texture->ResourceLocation.IsTransient()
@@ -3035,8 +2698,7 @@ void FD3D12DynamicRHI::RHIVirtualTextureSetFirstMipVisible(FRHITexture2D* Textur
 {
 }
 
-template<typename BaseResourceType>
-TD3D12Texture2D<BaseResourceType>* FD3D12DynamicRHI::CreateTextureFromResource(bool bTextureArray, bool bCubeTexture, EPixelFormat Format, ETextureCreateFlags TexCreateFlags, const FClearValueBinding& ClearValueBinding, ID3D12Resource* Resource)
+FD3D12Texture* FD3D12DynamicRHI::CreateTextureFromResource(bool bTextureArray, bool bCubeTexture, EPixelFormat Format, ETextureCreateFlags TexCreateFlags, const FClearValueBinding& ClearValueBinding, ID3D12Resource* Resource)
 {
 	check(Resource);
 	FD3D12Adapter* Adapter = &GetAdapter();
@@ -3062,24 +2724,6 @@ TD3D12Texture2D<BaseResourceType>* FD3D12DynamicRHI::CreateTextureFromResource(b
 
 	SCOPE_CYCLE_COUNTER(STAT_D3D12CreateTextureTime);
 
-	const bool bSRGB = EnumHasAnyFlags(TexCreateFlags, TexCreate_SRGB);
-
-	const DXGI_FORMAT PlatformResourceFormat = TextureDesc.Format;
-	const DXGI_FORMAT PlatformShaderResourceFormat = FindShaderResourceDXGIFormat(PlatformResourceFormat, bSRGB);
-	const DXGI_FORMAT PlatformRenderTargetFormat = FindShaderResourceDXGIFormat(PlatformResourceFormat, bSRGB);
-
-	// Set up the texture bind flags.
-	bool bCreateRTV = (TextureDesc.Flags & D3D12_RESOURCE_FLAG_ALLOW_RENDER_TARGET) != 0;
-	bool bCreateDSV = (TextureDesc.Flags & D3D12_RESOURCE_FLAG_ALLOW_DEPTH_STENCIL) != 0;
-	bool bCreateShaderResource = (TextureDesc.Flags & D3D12_RESOURCE_FLAG_DENY_SHADER_RESOURCE) == 0;
-
-	// DXGI_FORMAT_NV12 allows us to create RTV and SRV but only with other formats, so we should block creation here.
-	if (Format == PF_NV12)
-	{
-		bCreateRTV = false;
-		bCreateShaderResource = false;
-	}
-
 	// The state this resource will be in when it leaves this function
 	const FD3D12Resource::FD3D12ResourceTypeHelper Type(TextureDesc, D3D12_HEAP_TYPE_DEFAULT);
 	const D3D12_RESOURCE_STATES DestinationState = Type.GetOptimalInitialState(ERHIAccess::Unknown, !EnumHasAnyFlags(TexCreateFlags, TexCreate_Shared));
@@ -3088,9 +2732,11 @@ TD3D12Texture2D<BaseResourceType>* FD3D12DynamicRHI::CreateTextureFromResource(b
 	FD3D12Resource* TextureResource = new FD3D12Resource(Device, Device->GetGPUMask(), Resource, DestinationState, TextureDesc);
 	TextureResource->AddRef();
 
-	TD3D12Texture2D<BaseResourceType>* Texture2D = Adapter->CreateLinkedObject<TD3D12Texture2D<BaseResourceType>>(Device->GetGPUMask(), [&](FD3D12Device* Device)
+	FRHITextureCreateDesc CreateDesc(FRHITextureDesc::Create2D(FIntPoint(SizeX, SizeY), Format, ClearValueBinding, TexCreateFlags, NumMips, NumSamples), ERHIAccess::SRVMask, TEXT("TextureFromResource"));
+
+	FD3D12Texture* Texture2D = Adapter->CreateLinkedObject<FD3D12Texture>(Device->GetGPUMask(), [&](FD3D12Device* Device)
 	{
-		return new TD3D12Texture2D<BaseResourceType>(Device, SizeX, SizeY, SizeZ, NumMips, NumSamples, Format, false, TexCreateFlags, ClearValueBinding);
+		return CreateNewD3D12Texture(CreateDesc, Device);
 	});
 
 	FD3D12ResourceLocation& Location = Texture2D->ResourceLocation;
@@ -3098,146 +2744,7 @@ TD3D12Texture2D<BaseResourceType>* FD3D12DynamicRHI::CreateTextureFromResource(b
 	Location.SetResource(TextureResource);
 	Location.SetGPUVirtualAddress(TextureResource->GetGPUVirtualAddress());
 
-	uint32 RTVIndex = 0;
-
-	if (bCreateRTV)
-	{
-		const bool bCreateRTVsPerSlice = EnumHasAnyFlags(TexCreateFlags, TexCreate_TargetArraySlicesIndependently) && (bTextureArray || bCubeTexture);
-		Texture2D->SetNumRenderTargetViews(bCreateRTVsPerSlice ? NumMips * TextureDesc.DepthOrArraySize : NumMips);
-
-		// Create a render target view for each mip
-		for (uint32 MipIndex = 0; MipIndex < NumMips; MipIndex++)
-		{
-			if (bCreateRTVsPerSlice)
-			{
-				Texture2D->SetCreatedRTVsPerSlice(true, TextureDesc.DepthOrArraySize);
-
-				for (uint32 SliceIndex = 0; SliceIndex < TextureDesc.DepthOrArraySize; SliceIndex++)
-				{
-					D3D12_RENDER_TARGET_VIEW_DESC RTVDesc;
-					FMemory::Memzero(&RTVDesc, sizeof(RTVDesc));
-					RTVDesc.Format = PlatformRenderTargetFormat;
-					RTVDesc.ViewDimension = D3D12_RTV_DIMENSION_TEXTURE2DARRAY;
-					RTVDesc.Texture2DArray.FirstArraySlice = SliceIndex;
-					RTVDesc.Texture2DArray.ArraySize = 1;
-					RTVDesc.Texture2DArray.MipSlice = MipIndex;
-					RTVDesc.Texture2DArray.PlaneSlice = GetPlaneSliceFromViewFormat(PlatformResourceFormat, RTVDesc.Format);
-
-					Texture2D->SetRenderTargetViewIndex(new FD3D12RenderTargetView(Device, RTVDesc, Texture2D), RTVIndex++);
-				}
-			}
-			else
-			{
-				D3D12_RENDER_TARGET_VIEW_DESC RTVDesc;
-				FMemory::Memzero(&RTVDesc, sizeof(RTVDesc));
-				RTVDesc.Format = PlatformRenderTargetFormat;
-				if (bTextureArray || bCubeTexture)
-				{
-					RTVDesc.ViewDimension = D3D12_RTV_DIMENSION_TEXTURE2DARRAY;
-					RTVDesc.Texture2DArray.FirstArraySlice = 0;
-					RTVDesc.Texture2DArray.ArraySize = TextureDesc.DepthOrArraySize;
-					RTVDesc.Texture2DArray.MipSlice = MipIndex;
-					RTVDesc.Texture2DArray.PlaneSlice = GetPlaneSliceFromViewFormat(PlatformResourceFormat, RTVDesc.Format);
-				}
-				else if (NumSamples == 1)
-				{
-					RTVDesc.ViewDimension = D3D12_RTV_DIMENSION_TEXTURE2D;
-					RTVDesc.Texture2D.MipSlice = MipIndex;
-					RTVDesc.Texture2D.PlaneSlice = GetPlaneSliceFromViewFormat(PlatformResourceFormat, RTVDesc.Format);
-				}
-				else
-				{
-					RTVDesc.ViewDimension = D3D12_RTV_DIMENSION_TEXTURE2DMS;
-				}
-
-				Texture2D->SetRenderTargetViewIndex(new FD3D12RenderTargetView(Device, RTVDesc, Texture2D), RTVIndex++);
-			}
-		}
-	}
-
-	if (bCreateDSV)
-	{
-		// Create a depth-stencil-view for the texture.
-		D3D12_DEPTH_STENCIL_VIEW_DESC DSVDesc = {};
-		DSVDesc.Format = FindDepthStencilDXGIFormat(PlatformResourceFormat);
-		if (bTextureArray || bCubeTexture)
-		{
-			DSVDesc.ViewDimension = D3D12_DSV_DIMENSION_TEXTURE2DARRAY;
-			DSVDesc.Texture2DArray.FirstArraySlice = 0;
-			DSVDesc.Texture2DArray.ArraySize = TextureDesc.DepthOrArraySize;
-			DSVDesc.Texture2DArray.MipSlice = 0;
-		}
-		else if (NumSamples == 1)
-		{
-			DSVDesc.ViewDimension = D3D12_DSV_DIMENSION_TEXTURE2D;
-			DSVDesc.Texture2D.MipSlice = 0;
-		}
-		else
-		{
-			DSVDesc.ViewDimension = D3D12_DSV_DIMENSION_TEXTURE2DMS;
-		}
-
-		const bool HasStencil = HasStencilBits(DSVDesc.Format);
-		for (uint32 AccessType = 0; AccessType < FExclusiveDepthStencil::MaxIndex; ++AccessType)
-		{
-			// Create a read-only access views for the texture.
-			DSVDesc.Flags = (AccessType & FExclusiveDepthStencil::DepthRead_StencilWrite) ? D3D12_DSV_FLAG_READ_ONLY_DEPTH : D3D12_DSV_FLAG_NONE;
-			if (HasStencil)
-			{
-				DSVDesc.Flags |= (AccessType & FExclusiveDepthStencil::DepthWrite_StencilRead) ? D3D12_DSV_FLAG_READ_ONLY_STENCIL : D3D12_DSV_FLAG_NONE;
-			}
-
-			Texture2D->SetDepthStencilView(new FD3D12DepthStencilView(Device, DSVDesc, Texture2D, HasStencil), AccessType);
-		}
-	}
-
-	// Create a shader resource view for the texture.
-	D3D12_SHADER_RESOURCE_VIEW_DESC SRVDesc = {};
-	SRVDesc.Shader4ComponentMapping = D3D12_DEFAULT_SHADER_4_COMPONENT_MAPPING;
-	SRVDesc.Format = PlatformShaderResourceFormat;
-
-	if (bCubeTexture && bTextureArray)
-	{
-		SRVDesc.ViewDimension = D3D12_SRV_DIMENSION_TEXTURECUBEARRAY;
-		SRVDesc.TextureCubeArray.MostDetailedMip = 0;
-		SRVDesc.TextureCubeArray.MipLevels = NumMips;
-		SRVDesc.TextureCubeArray.ResourceMinLODClamp = 0.0f;
-		SRVDesc.TextureCubeArray.First2DArrayFace = 0;
-		SRVDesc.TextureCubeArray.NumCubes = SizeZ / 6;
-	}
-	else if (bCubeTexture)
-	{
-		SRVDesc.ViewDimension = D3D12_SRV_DIMENSION_TEXTURECUBE;
-		SRVDesc.TextureCube.MostDetailedMip = 0;
-		SRVDesc.TextureCube.MipLevels = NumMips;
-		SRVDesc.TextureCube.ResourceMinLODClamp = 0.0f;
-	}
-	else if (bTextureArray)
-	{
-		SRVDesc.ViewDimension = D3D12_SRV_DIMENSION_TEXTURE2DARRAY;
-		SRVDesc.Texture2DArray.MostDetailedMip = 0;
-		SRVDesc.Texture2DArray.MipLevels = NumMips;
-		SRVDesc.Texture2DArray.FirstArraySlice = 0;
-		SRVDesc.Texture2DArray.ArraySize = SizeZ;
-		SRVDesc.Texture2DArray.PlaneSlice = GetPlaneSliceFromViewFormat(PlatformResourceFormat, SRVDesc.Format);
-	}
-	else if (NumSamples == 1)
-	{
-		SRVDesc.ViewDimension = D3D12_SRV_DIMENSION_TEXTURE2D;
-		SRVDesc.Texture2D.MostDetailedMip = 0;
-		SRVDesc.Texture2D.MipLevels = NumMips;
-		SRVDesc.Texture2D.PlaneSlice = GetPlaneSliceFromViewFormat(PlatformResourceFormat, SRVDesc.Format);
-	}
-	else
-	{
-		SRVDesc.ViewDimension = D3D12_SRV_DIMENSION_TEXTURE2DMS;
-	}
-
-	// Create a wrapper for the SRV and set it on the texture
-	if (bCreateShaderResource)
-	{
-		Texture2D->SetShaderResourceView(new FD3D12ShaderResourceView(Texture2D, SRVDesc));
-	}
+	Texture2D->CreateViews();
 
 	FD3D12TextureStats::D3D12TextureAllocated(*Texture2D);
 
@@ -3246,39 +2753,38 @@ TD3D12Texture2D<BaseResourceType>* FD3D12DynamicRHI::CreateTextureFromResource(b
 
 FTexture2DRHIRef FD3D12DynamicRHI::RHICreateTexture2DFromResource(EPixelFormat Format, ETextureCreateFlags TexCreateFlags, const FClearValueBinding& ClearValueBinding, ID3D12Resource* Resource)
 {
-	return CreateTextureFromResource<FD3D12BaseTexture2D>(false, false, Format, TexCreateFlags, ClearValueBinding, Resource);
+	return CreateTextureFromResource(false, false, Format, TexCreateFlags, ClearValueBinding, Resource);
 }
 
 FTexture2DArrayRHIRef FD3D12DynamicRHI::RHICreateTexture2DArrayFromResource(EPixelFormat Format, ETextureCreateFlags TexCreateFlags, const FClearValueBinding& ClearValueBinding, ID3D12Resource* Resource)
 {
-	return CreateTextureFromResource<FD3D12BaseTexture2DArray>(true, false, Format, TexCreateFlags, ClearValueBinding, Resource);
+	return CreateTextureFromResource(true, false, Format, TexCreateFlags, ClearValueBinding, Resource);
 }
 
 FTextureCubeRHIRef FD3D12DynamicRHI::RHICreateTextureCubeFromResource(EPixelFormat Format, ETextureCreateFlags TexCreateFlags, const FClearValueBinding& ClearValueBinding, ID3D12Resource* Resource)
 {
-	return CreateTextureFromResource<FD3D12BaseTextureCube>(false, true, Format, TexCreateFlags, ClearValueBinding, Resource);
+	return CreateTextureFromResource(false, true, Format, TexCreateFlags, ClearValueBinding, Resource);
 }
 
 void FD3D12DynamicRHI::RHIAliasTextureResources(FTextureRHIRef& DestTextureRHI, FTextureRHIRef& SrcTextureRHI)
 {
-	FD3D12TextureBase* DestTexture = GetD3D12TextureFromRHITexture(DestTextureRHI);
-	FD3D12TextureBase* SrcTexture = GetD3D12TextureFromRHITexture(SrcTextureRHI);
+	FD3D12Texture* DestTexture = GetD3D12TextureFromRHITexture(DestTextureRHI);
+	FD3D12Texture* SrcTexture = GetD3D12TextureFromRHITexture(SrcTextureRHI);
 
 	// Make sure we keep a reference to the source texture we're aliasing, so we don't lose it if all other references
 	// go away but we're kept around.
 	DestTexture->SetAliasingSource(SrcTextureRHI);
 
-	for (FD3D12TextureBase::FDualLinkedObjectIterator It(DestTexture, SrcTexture); It; ++It)
+	for (FD3D12Texture::FDualLinkedObjectIterator It(DestTexture, SrcTexture); It; ++It)
 	{
-		FD3D12TextureBase* DestLinkedTexture = It.GetFirst();
-		FD3D12TextureBase* SrcLinkedTexture = It.GetSecond();
+		FD3D12Texture* DestLinkedTexture = It.GetFirst();
+		FD3D12Texture* SrcLinkedTexture = It.GetSecond();
 
 		DestLinkedTexture->AliasResources(SrcLinkedTexture);
 	}
 }
 
-template<typename BaseResourceType>
-TD3D12Texture2D<BaseResourceType>* FD3D12DynamicRHI::CreateAliasedD3D12Texture2D(TD3D12Texture2D<BaseResourceType>* SourceTexture)
+FD3D12Texture* FD3D12DynamicRHI::CreateAliasedD3D12Texture2D(FD3D12Texture* SourceTexture)
 {
 	FD3D12Adapter* Adapter = &GetAdapter();
 
@@ -3310,9 +2816,11 @@ TD3D12Texture2D<BaseResourceType>* FD3D12DynamicRHI::CreateAliasedD3D12Texture2D
 	const DXGI_FORMAT PlatformShaderResourceFormat = FindShaderResourceDXGIFormat(PlatformResourceFormat, bSRGB);
 	const DXGI_FORMAT PlatformRenderTargetFormat = FindShaderResourceDXGIFormat(PlatformResourceFormat, bSRGB);
 
-	TD3D12Texture2D<BaseResourceType>* Texture2D = Adapter->CreateLinkedObject<TD3D12Texture2D<BaseResourceType>>(Device->GetGPUMask(), [&](FD3D12Device* Device)
+	FRHITextureCreateDesc CreateDesc(SourceTexture->GetDesc(), ERHIAccess::SRVMask, TEXT("AliasedTexture"));
+
+	FD3D12Texture* Texture2D = Adapter->CreateLinkedObject<FD3D12Texture>(Device->GetGPUMask(), [&](FD3D12Device* Device)
 	{
-		return new TD3D12Texture2D<BaseResourceType>(Device, SizeX, SizeY, SizeZ, NumMips, NumSamples, SourceTexture->GetFormat(), false, SourceTexture->GetFlags(), SourceTexture->GetClearBinding());
+		return CreateNewD3D12Texture(CreateDesc, Device);
 	});
 
 	// Set up the texture bind flags.
@@ -3322,7 +2830,7 @@ TD3D12Texture2D<BaseResourceType>* FD3D12DynamicRHI::CreateAliasedD3D12Texture2D
 	const D3D12_RESOURCE_STATES State = D3D12_RESOURCE_STATE_COMMON;
 
 	bool bCreatedRTVPerSlice = false;
-	const bool bCubeTexture = SourceTexture->IsCubemap();
+	const bool bCubeTexture = CreateDesc.IsTextureCube();
 	const bool bTextureArray = !bCubeTexture && TextureDesc.DepthOrArraySize > 1;
 
 	if (bCreateRTV)
@@ -3334,7 +2842,7 @@ TD3D12Texture2D<BaseResourceType>* FD3D12DynamicRHI::CreateAliasedD3D12Texture2D
 		for (uint32 MipIndex = 0; MipIndex < TextureDesc.MipLevels; MipIndex++)
 		{
 			// These are null because we'll be aliasing them shortly.
-			if (EnumHasAnyFlags(SourceTexture->Flags, TexCreate_TargetArraySlicesIndependently) && (bTextureArray || bCubeTexture))
+			if (EnumHasAnyFlags(SourceTexture->GetDesc().Flags, TexCreate_TargetArraySlicesIndependently) && (bTextureArray || bCubeTexture))
 			{
 				bCreatedRTVPerSlice = true;
 
@@ -3366,37 +2874,27 @@ TD3D12Texture2D<BaseResourceType>* FD3D12DynamicRHI::CreateAliasedD3D12Texture2D
 
 FTextureRHIRef FD3D12DynamicRHI::RHICreateAliasedTexture(FTextureRHIRef& SourceTextureRHI)
 {
-	FD3D12TextureBase* SourceTexture = GetD3D12TextureFromRHITexture(SourceTextureRHI);
-	FTextureRHIRef ReturnTexture;
-	if (SourceTextureRHI->GetTexture2D() != nullptr)
-	{
-		ReturnTexture = CreateAliasedD3D12Texture2D<FD3D12BaseTexture2D>(static_cast<FD3D12Texture2D*>(SourceTextureRHI->GetTexture2D()));
-	}
-	else if (SourceTextureRHI->GetTexture2DArray() != nullptr)
-	{
-		ReturnTexture = CreateAliasedD3D12Texture2D<FD3D12BaseTexture2DArray>(static_cast<FD3D12Texture2DArray*>(SourceTextureRHI->GetTexture2DArray()));
-	}
-	else if (SourceTextureRHI->GetTextureCube() != nullptr)
-	{
-		ReturnTexture = CreateAliasedD3D12Texture2D<FD3D12BaseTextureCube>(static_cast<FD3D12TextureCube*>(SourceTextureRHI->GetTextureCube()));
-	}
-
+	FD3D12Texture* SourceTexture = GetD3D12TextureFromRHITexture(SourceTextureRHI);
+	FD3D12Texture* ReturnTexture = CreateAliasedD3D12Texture2D(SourceTexture);
 	if (ReturnTexture == nullptr)
 	{
 		UE_LOG(LogD3D12RHI, Error, TEXT("Currently FD3D12DynamicRHI::RHICreateAliasedTexture only supports 2D, 2D Array and Cube textures."));
 		return nullptr;
 	}
-
-	FD3D12TextureBase* DestTexture = GetD3D12TextureFromRHITexture(ReturnTexture);
-	DestTexture->SetAliasingSource(SourceTextureRHI);
+	ReturnTexture->SetAliasingSource(SourceTextureRHI);
 
 	return ReturnTexture;
 }
 
+
+///////////////////////////////////////////////////////////////////////////////////////////
+// FD3D12CommandContext Texture functions
+///////////////////////////////////////////////////////////////////////////////////////////
+
 void FD3D12CommandContext::RHICopyTexture(FRHITexture* SourceTextureRHI, FRHITexture* DestTextureRHI, const FRHICopyTextureInfo& CopyInfo)
 {
-	FD3D12TextureBase* SourceTexture = RetrieveTextureBase(SourceTextureRHI);
-	FD3D12TextureBase* DestTexture = RetrieveTextureBase(DestTextureRHI);
+	FD3D12Texture* SourceTexture = RetrieveTexture(SourceTextureRHI);
+	FD3D12Texture* DestTexture = RetrieveTexture(DestTextureRHI);
 
 	FScopedResourceBarrier ConditionalScopeResourceBarrierSource(CommandListHandle, SourceTexture->GetResource(), D3D12_RESOURCE_STATE_COPY_SOURCE, D3D12_RESOURCE_BARRIER_ALL_SUBRESOURCES, FD3D12DynamicRHI::ETransitionMode::Validate);
 	FScopedResourceBarrier ConditionalScopeResourceBarrierDest(CommandListHandle, DestTexture->GetResource(), D3D12_RESOURCE_STATE_COPY_DEST, D3D12_RESOURCE_BARRIER_ALL_SUBRESOURCES, FD3D12DynamicRHI::ETransitionMode::Validate);
@@ -3503,6 +3001,11 @@ void FD3D12CommandContext::RHICopyTexture(FRHITexture* SourceTextureRHI, FRHITex
 	DestTexture->SetReadBackListHandle(CommandListHandle);
 }
 
+
+///////////////////////////////////////////////////////////////////////////////////////////
+// FD3D12BackBufferReferenceTexture2D functions
+///////////////////////////////////////////////////////////////////////////////////////////
+
 FRHITexture* FD3D12BackBufferReferenceTexture2D::GetBackBufferTexture()
 {
 	return bIsSDR ? Viewport->GetSDRBackBuffer_RHIThread() : Viewport->GetBackBuffer_RHIThread();
