Copying //UE5/Dev-ParallelRendering to //UE5/Main

diff --git a/Engine/Source/Runtime/D3D12RHI/Private/D3D12Texture.h b/Engine/Source/Runtime/D3D12RHI/Private/D3D12Texture.h
--- a/Engine/Source/Runtime/D3D12RHI/Private/D3D12Texture.h
+++ b/Engine/Source/Runtime/D3D12RHI/Private/D3D12Texture.h
@@ -5,30 +5,13 @@ D3D12Texture.h: Implementation of D3D12 Texture
 =============================================================================*/
 #pragma once
 
+#include "D3D12Resources.h"
+#include "D3D12CommandList.h"
+
 /** If true, guard texture creates with SEH to log more information about a driver crash we are seeing during texture streaming. */
 #define GUARDED_TEXTURE_CREATES (PLATFORM_WINDOWS && !(UE_BUILD_SHIPPING || UE_BUILD_TEST || PLATFORM_COMPILER_CLANG))
 
-static bool TextureCanBe4KAligned(const FD3D12ResourceDesc& Desc, EPixelFormat UEFormat)
-{
-	// 4KB alignment is only available for read only textures
-	if (!EnumHasAnyFlags(Desc.Flags, D3D12_RESOURCE_FLAG_ALLOW_RENDER_TARGET | D3D12_RESOURCE_FLAG_ALLOW_DEPTH_STENCIL | D3D12_RESOURCE_FLAG_ALLOW_UNORDERED_ACCESS) &&
-		!Desc.NeedsUAVAliasWorkarounds() && // UAV aliased resources are secretly writable.
-		Desc.SampleDesc.Count == 1)
-	{
-		D3D12_TILE_SHAPE Tile = {};
-		Get4KTileShape(&Tile, Desc.Format, UEFormat, Desc.Dimension, Desc.SampleDesc.Count);
-
-		uint32 TilesNeeded = GetTilesNeeded(Desc.Width, Desc.Height, Desc.DepthOrArraySize, Tile);
-
-		constexpr uint32 NUM_4K_BLOCKS_PER_64K_PAGE = 16;
-		return TilesNeeded <= NUM_4K_BLOCKS_PER_64K_PAGE;
-	}
-	else
-	{
-		return false;
-	}
-}
-
+// @todo don't make global here!
 void SafeCreateTexture2D(FD3D12Device* pDevice, 
 	FD3D12Adapter* Adapter,
 	const FD3D12ResourceDesc& TextureDesc,
@@ -40,517 +23,159 @@ void SafeCreateTexture2D(FD3D12Device* pDevice,
 	D3D12_RESOURCE_STATES InitialState,
 	const TCHAR* Name);
 
-void CreateUAVAliasResource(
-	FD3D12Adapter* Adapter,
-	D3D12_CLEAR_VALUE* ClearValuePtr,
-	const TCHAR* DebugName,
-	FD3D12ResourceLocation& Location);
-
-/** Texture base class. */
-class FD3D12TextureBase : public FD3D12BaseShaderResource, public FD3D12LinkedAdapterObject<FD3D12TextureBase>
+/** D3D12 RHI Texture class */
+class FD3D12Texture : public FRHITexture, public FD3D12BaseShaderResource, public FD3D12LinkedAdapterObject<FD3D12Texture>
 {
 public:
 
-	FD3D12TextureBase(class FD3D12Device* InParent)
-		: FD3D12BaseShaderResource(InParent)
-	{
-	}
+	// Static helper functions
+	static bool CanBe4KAligned(const FD3D12ResourceDesc& Desc, EPixelFormat UEFormat);
+
+	FD3D12Texture() = delete;
 
-	virtual ~FD3D12TextureBase() {}
+	/** Initialization constructor. */
+	FD3D12Texture(
+		const FRHITextureCreateDesc& InDesc,
+		class FD3D12Device* InParent)
+		: FRHITexture(InDesc)
+		, FD3D12BaseShaderResource(InParent) {}
+	virtual ~FD3D12Texture();
+
+	// IRefCountedObject interface overrides from FD3D12BaseShaderResource
+	virtual uint32 AddRef() const override						{ return FRHIResource::AddRef();	}
+	virtual uint32 Release() const override						{ return FRHIResource::Release();	}
+	virtual uint32 GetRefCount() const override					{ return FRHIResource::GetRefCount(); }
+
+	// FRHIResource overrides 
+#if RHI_ENABLE_RESOURCE_INFO
+	bool GetResourceInfo(FRHIResourceInfo& OutResourceInfo) const override;
+#endif
+
+	// FRHITexture overrides
+	virtual void* GetTextureBaseRHI() override final			{ return this; }
+	virtual void* GetNativeResource() const override final;
+	virtual FRHIDescriptorHandle GetDefaultBindlessHandle() const final;
+			
+	// Accessors.
+	bool IsStreamable() const									{ return EnumHasAnyFlags(GetDesc().Flags, ETextureCreateFlags::Streamable); }
+	FD3D12ShaderResourceView* GetShaderResourceView() const		{ return ShaderResourceView; }
+	const FTextureRHIRef& GetAliasingSourceTexture() const		{ return AliasingSourceTexture; }
+	FD3D12CLSyncPoint GetReadBackSyncPoint() const				{ return ReadBackSyncPoint; }
+	bool HasRenderTargetViews() const							{ return (RenderTargetViews.Num() > 0); }
+	void GetReadBackHeapDesc(D3D12_PLACED_SUBRESOURCE_FOOTPRINT& OutFootprint, uint32 Subresource) const;
+	FD3D12RenderTargetView* GetRenderTargetView(int32 MipIndex, int32 ArraySliceIndex) const;
+	FD3D12DepthStencilView* GetDepthStencilView(FExclusiveDepthStencil AccessType) const { return DepthStencilViews[AccessType.GetIndex()]; }
+		
+	// Setters
+	void SetShaderResourceView(FD3D12ShaderResourceView* InShaderResourceView) { ShaderResourceView = InShaderResourceView; }
+	void SetReadBackListHandle(FD3D12CommandListHandle listToWaitFor) { ReadBackSyncPoint = listToWaitFor; }
 
-	inline void SetCreatedRTVsPerSlice(bool Value, int32 InRTVArraySize)
-	{ 
+	// Setup functionality
+	void InitializeTextureData(class FRHICommandListImmediate* RHICmdList, const FRHITextureCreateDesc& CreateDesc, D3D12_RESOURCE_STATES DestinationState);
+	void CreateViews();
+	void SetCreatedRTVsPerSlice(bool Value, int32 InRTVArraySize)
+	{
 		bCreatedRTVsPerSlice = Value;
-		RTVArraySize = InRTVArraySize;
+		RTVArraySizePerMip = InRTVArraySize;
 	}
-
 	void SetNumRenderTargetViews(int32 InNumViews)
 	{
 		RenderTargetViews.Empty(InNumViews);
 		RenderTargetViews.AddDefaulted(InNumViews);
 	}
-
-	void SetDepthStencilView(FD3D12DepthStencilView* View, uint32 SubResourceIndex)
-	{
-		if (SubResourceIndex < FExclusiveDepthStencil::MaxIndex)
-		{
-			DepthStencilViews[SubResourceIndex] = View;
-		}
-		else
-		{
-			check(false);
-		}
-	}
-
-	void SetRenderTargetViewIndex(FD3D12RenderTargetView* View, uint32 SubResourceIndex)
-	{
-		if (SubResourceIndex < (uint32)RenderTargetViews.Num())
-		{
-			RenderTargetViews[SubResourceIndex] = View;
-		}
-		else
-		{
-			check(false);
-		}
-	}
-
 	void SetRenderTargetView(FD3D12RenderTargetView* View)
 	{
 		RenderTargetViews.Empty(1);
 		RenderTargetViews.Add(View);
 	}
+	void SetDepthStencilView(FD3D12DepthStencilView* View, uint32 SubResourceIndex);
+	void SetRenderTargetViewIndex(FD3D12RenderTargetView* View, uint32 SubResourceIndex);
 
-	int64 GetMemorySize() const
-	{
-		return MemorySize;
-	}
-
-	void SetMemorySize(int64 InMemorySize)
-	{
-		check(InMemorySize >= 0);
-		MemorySize = InMemorySize;
-	}
-
-	void SetAliasingSource(FTextureRHIRef& SourceTextureRHI)
-	{
-		AliasingSourceTexture = SourceTextureRHI;
-	}
-
-	// Accessors.
-	FD3D12Resource* GetResource() const { return ResourceLocation.GetResource(); }
-	uint64 GetOffset() const { return ResourceLocation.GetOffsetFromBaseOfResource(); }
-	FD3D12ShaderResourceView* GetShaderResourceView() const { return ShaderResourceView; }
-	inline const FTextureRHIRef& GetAliasingSourceTexture() const { return AliasingSourceTexture; }
-
-	void SetShaderResourceView(FD3D12ShaderResourceView* InShaderResourceView) { ShaderResourceView = InShaderResourceView; }
-
-	static inline bool ShouldDeferCmdListOperation(FRHICommandList* RHICmdList)
-	{
-		if (RHICmdList == nullptr)
-		{
-			return false;
-		}
-
-		if (RHICmdList->Bypass() || !IsRunningRHIInSeparateThread())
-		{
-			return false;
-		}
-
-		return true;
-	}
-
-	void UpdateTexture(uint32 MipIndex, uint32 DestX, uint32 DestY, uint32 DestZ, const D3D12_TEXTURE_COPY_LOCATION& SourceCopyLocation);
-	void CopyTextureRegion(uint32 DestX, uint32 DestY, uint32 DestZ, FD3D12TextureBase* SourceTexture, const D3D12_BOX& SourceBox);
-	void InitializeTextureData(class FRHICommandListImmediate* RHICmdList, const void* InitData, uint32 InitDataSize, uint32 SizeX, uint32 SizeY, uint32 SizeZ, uint32 NumSlices, uint32 NumMips, EPixelFormat Format, D3D12_RESOURCE_STATES DestinationState);
-
-	/**
-	* Get the render target view for the specified mip and array slice.
-	* An array slice of -1 is used to indicate that no array slice should be required.
-	*/
-	FD3D12RenderTargetView* GetRenderTargetView(int32 MipIndex, int32 ArraySliceIndex) const
-	{
-		int32 ArrayIndex = MipIndex;
-
-		if (bCreatedRTVsPerSlice)
-		{
-			check(ArraySliceIndex >= 0);
-			ArrayIndex = MipIndex * RTVArraySize + ArraySliceIndex;
-			check(ArrayIndex < RenderTargetViews.Num());
-		}
-		else
-		{
-			// Catch attempts to use a specific slice without having created the texture to support it
-			check(ArraySliceIndex == -1 || ArraySliceIndex == 0);
-		}
-
-		if (ArrayIndex < RenderTargetViews.Num())
-		{
-			return RenderTargetViews[ArrayIndex];
-		}
-		return 0;
-	}
-	FD3D12DepthStencilView* GetDepthStencilView(FExclusiveDepthStencil AccessType) const
-	{
-		return DepthStencilViews[AccessType.GetIndex()];
-	}
-
-	inline bool HasRenderTargetViews() const
-	{
-		return (RenderTargetViews.Num() > 0);
-	}
-
-	void AliasResources(FD3D12TextureBase* Texture)
-	{
-		// Alias the location, will perform an addref underneath
-		FD3D12ResourceLocation::Alias(ResourceLocation, Texture->ResourceLocation);
-
-		ShaderResourceView = Texture->ShaderResourceView;
-
-		for (uint32 Index = 0; Index < FExclusiveDepthStencil::MaxIndex; Index++)
-		{
-			DepthStencilViews[Index] = Texture->DepthStencilViews[Index];
-		}
-		for (int32 Index = 0; Index < Texture->RenderTargetViews.Num(); Index++)
-		{
-			RenderTargetViews[Index] = Texture->RenderTargetViews[Index];
-		}
-	}
-
-	// Modifiers.
-	void SetReadBackListHandle(FD3D12CommandListHandle listToWaitFor) { ReadBackSyncPoint = listToWaitFor; }
-	FD3D12CLSyncPoint GetReadBackSyncPoint() const { return ReadBackSyncPoint; }
-
-	FD3D12CLSyncPoint ReadBackSyncPoint;
-
-protected:
-
-	/** Amount of memory allocated by this texture, in bytes. */
-	int64 MemorySize{};
-
-	/** A shader resource view of the texture. */
-	TRefCountPtr<FD3D12ShaderResourceView> ShaderResourceView;
-
-	/** A render targetable view of the texture. */
-	TArray<TRefCountPtr<FD3D12RenderTargetView>, TInlineAllocator<1>> RenderTargetViews;
-
-	/** A depth-stencil targetable view of the texture. */
-	TRefCountPtr<FD3D12DepthStencilView> DepthStencilViews[FExclusiveDepthStencil::MaxIndex];
-
-	TMap<uint32, FD3D12LockedResource*> LockedMap;
-
-	FTextureRHIRef AliasingSourceTexture;
-
-	int32 RTVArraySize{};
-
-	bool bCreatedRTVsPerSlice{ false };
-};
-
-#if PLATFORM_WINDOWS || PLATFORM_HOLOLENS
-struct FD3D12TextureLayout {};
-#endif
-
-/** 2D texture (vanilla, cubemap or 2D array) */
-template<typename BaseResourceType>
-class TD3D12Texture2D : public BaseResourceType, public FD3D12TextureBase
-{
-public:
-
-	/** Flags used when the texture was created */
-	ETextureCreateFlags Flags;
-
-	TD3D12Texture2D() = delete;
-
-	/** Initialization constructor. */
-	TD3D12Texture2D(
-		class FD3D12Device* InParent,
-		uint32 InSizeX,
-		uint32 InSizeY,
-		uint32 InSizeZ,
-		uint32 InNumMips,
-		uint32 InNumSamples,
-		EPixelFormat InFormat,
-		bool bInCubemap,
-		ETextureCreateFlags InFlags,
-		const FClearValueBinding& InClearValue,
-		const FD3D12TextureLayout* InTextureLayout = nullptr
-#if PLATFORM_SUPPORTS_VIRTUAL_TEXTURES
-		, void* InRawTextureMemory = nullptr
-#endif
-		)
-		: BaseResourceType(
-			InSizeX,
-			InSizeY,
-			InSizeZ,
-			InNumMips,
-			InNumSamples,
-			InFormat,
-			InFlags,
-			InClearValue
-			)
-		, FD3D12TextureBase(InParent)
-		, Flags(InFlags)
-		, bCubemap(bInCubemap)
-		, bStreamable(!!(InFlags & TexCreate_Streamable))
-		, bMipOrderDescending(false)
-#if PLATFORM_SUPPORTS_VIRTUAL_TEXTURES
-		, RawTextureMemory(InRawTextureMemory)
-#endif
-	{
-		if (InTextureLayout == nullptr)
-		{
-			FMemory::Memzero(TextureLayout);
-		}
-		else
-		{
-			TextureLayout = *InTextureLayout;
-#if PLATFORM_SUPPORTS_VIRTUAL_TEXTURES
-			bMipOrderDescending = InNumMips > 1u && TextureLayout.GetSubresourceOffset(0, 0, 0) > TextureLayout.GetSubresourceOffset(0, 1, 0);
-#endif
-		}
-	}
-
-	virtual ~TD3D12Texture2D();
-
-	// FRHIResource overrides
-#if RHI_ENABLE_RESOURCE_INFO
-	bool GetResourceInfo(FRHIResourceInfo& OutResourceInfo) const override
-	{
-		OutResourceInfo = FRHIResourceInfo{};
-		OutResourceInfo.Name = this->GetName();
-		OutResourceInfo.Type = this->GetType();
-		OutResourceInfo.VRamAllocation.AllocationSize = GetMemorySize();
-		OutResourceInfo.IsTransient = this->ResourceLocation.IsTransient();
-		return true;
-	}
-#endif
-
-	/**
-	* Locks one of the texture's mip-maps.
-	* @return A pointer to the specified texture data.
-	*/
+	// Locking/update functions
 	void* Lock(class FRHICommandListImmediate* RHICmdList, uint32 MipIndex, uint32 ArrayIndex, EResourceLockMode LockMode, uint32& DestStride);
-
-	/** Unlocks a previously locked mip-map. */
 	void Unlock(class FRHICommandListImmediate* RHICmdList, uint32 MipIndex, uint32 ArrayIndex);
-
-	//* Update the contents of the Texture2D using a Copy command */
 	void UpdateTexture2D(class FRHICommandListImmediate* RHICmdList, uint32 MipIndex, const struct FUpdateTextureRegion2D& UpdateRegion, uint32 SourcePitch, const uint8* SourceData);
+	void UpdateTexture(uint32 MipIndex, uint32 DestX, uint32 DestY, uint32 DestZ, const D3D12_TEXTURE_COPY_LOCATION& SourceCopyLocation);
+	void CopyTextureRegion(uint32 DestX, uint32 DestY, uint32 DestZ, FD3D12Texture* SourceTexture, const D3D12_BOX& SourceBox);
 
-	// Accessors.
-	FD3D12Resource* GetResource() const { return (FD3D12Resource*)FD3D12TextureBase::GetResource(); }
-
-	void GetReadBackHeapDesc(D3D12_PLACED_SUBRESOURCE_FOOTPRINT& OutFootprint, uint32 Subresource) const;
-
-	bool IsCubemap() const { return bCubemap; }
-
-	bool IsStreamable() const { return bStreamable; }
-
-	bool IsLastMipFirst() const { return bMipOrderDescending; }
-
-	virtual FRHIDescriptorHandle GetDefaultBindlessHandle() const final
-	{
-		if (FRHIShaderResourceView* View = GetShaderResourceView())
-		{
-			return View->GetBindlessHandle();
-		}
-		return FRHIDescriptorHandle();
-	}
-
-	/** FRHITexture override.  See FRHITexture::GetNativeResource() */
-	virtual void* GetNativeResource() const override final
-	{
-		void* NativeResource = nullptr;
-		FD3D12Resource* Resource = GetResource();
-		if (Resource)
-		{
-			NativeResource = Resource->GetResource();
-		}
-		if (!NativeResource)
-		{
-			FD3D12TextureBase* Base = GetD3D12TextureFromRHITexture((FRHITexture*)this);
-			if (Base)
-			{
-				Resource = Base->GetResource();
-				if (Resource)
-				{
-					NativeResource = Resource->GetResource();
-				}
-			}
-		}
-		return NativeResource;
-	}
-
-	virtual void* GetTextureBaseRHI() override final
-	{
-		return static_cast<FD3D12TextureBase*>(this);
-	}
-
-	// IRefCountedObject interface.
-	virtual uint32 AddRef() const
-	{
-		return FRHIResource::AddRef();
-	}
-	virtual uint32 Release() const
-	{
-		return FRHIResource::Release();
-	}
-	virtual uint32 GetRefCount() const
-	{
-		return FRHIResource::GetRefCount();
-	}
-#if PLATFORM_SUPPORTS_VIRTUAL_TEXTURES
-	void* GetRawTextureMemory() const
-	{
-		return RawTextureMemory;
-	}
-
-	void SetRawTextureMemory(void* Memory)
-	{
-		RawTextureMemory = Memory;
-	}
-	FPlatformMemory::FPlatformVirtualMemoryBlock& GetRawTextureBlock()
-	{
-		return RawTextureBlock;
-	}
-#endif
-
-	const FD3D12TextureLayout& GetTextureLayout() const { return TextureLayout; }
+	// Resource aliasing
+	void AliasResources(FD3D12Texture* Texture);
+	void SetAliasingSource(FTextureRHIRef& SourceTextureRHI)	{ AliasingSourceTexture = SourceTextureRHI; }
+	
+protected:
 
-private:
-	/** Unlocks a previously locked mip-map. */
+	// Lock helper functions
 	void UnlockInternal(class FRHICommandListImmediate* RHICmdList, FLinkedObjectIterator NextObject, uint32 MipIndex, uint32 ArrayIndex);
+		
+	// A shader resource view of the texture.
+	TRefCountPtr<FD3D12ShaderResourceView> ShaderResourceView;
 
-	/** Whether the texture is a cube-map. */
-	const uint32 bCubemap : 1;
+	// Set when RTVs are created for each slice - TexCreate_TargetArraySlicesIndependently for TextureArrays & Cubemaps
+	bool bCreatedRTVsPerSlice{ false };
 
-	/** Whether the texture has been created with flag TexCreate_Streamable */
-	const uint32 bStreamable : 1;
+	// Number of RTVs per mip when stored per slice
+	int32 RTVArraySizePerMip{};
 
-	/** Whether mips are ordered from the last to the first in memory */
-	uint32 bMipOrderDescending : 1;
+	// A render targetable view of the texture.
+	TArray<TRefCountPtr<FD3D12RenderTargetView>, TInlineAllocator<1>> RenderTargetViews;
 
-#if PLATFORM_SUPPORTS_VIRTUAL_TEXTURES
-	void* RawTextureMemory;
-	FPlatformMemory::FPlatformVirtualMemoryBlock RawTextureBlock;
-#endif
+	// A depth-stencil targetable view of the texture.
+	TRefCountPtr<FD3D12DepthStencilView> DepthStencilViews[FExclusiveDepthStencil::MaxIndex];
 
-	FD3D12TextureLayout TextureLayout;
+	// Data for each subresource while texture is locked
+	TMap<uint32, FD3D12LockedResource*> LockedMap;
 
+	// Cached footprint size of first resource - optimization
 	mutable TUniquePtr<D3D12_PLACED_SUBRESOURCE_FOOTPRINT> FirstSubresourceFootprint;
-};
-
-/** 3D Texture */
-class FD3D12Texture3D : public FRHITexture3D, public FD3D12TextureBase
-{
-public:
-	FD3D12Texture3D() = delete;
-	/** Initialization constructor. */
-	FD3D12Texture3D(
-	class FD3D12Device* InParent,
-		uint32 InSizeX,
-		uint32 InSizeY,
-		uint32 InSizeZ,
-		uint32 InNumMips,
-		EPixelFormat InFormat,
-		ETextureCreateFlags InFlags,
-		const FClearValueBinding& InClearValue
-		)
-		: FRHITexture3D(InSizeX, InSizeY, InSizeZ, InNumMips, InFormat, InFlags, InClearValue)
-		, FD3D12TextureBase(InParent)
-		, bStreamable(!!(InFlags & TexCreate_Streamable))
-	{
-	}
-
-	virtual ~FD3D12Texture3D();
-
-	// FRHIResource overrides
-#if RHI_ENABLE_RESOURCE_INFO
-	bool GetResourceInfo(FRHIResourceInfo& OutResourceInfo) const override
-	{
-		OutResourceInfo = FRHIResourceInfo{};
-		OutResourceInfo.Name = GetName();
-		OutResourceInfo.Type = GetType();
-		OutResourceInfo.VRamAllocation.AllocationSize = ResourceLocation.GetSize();
-		OutResourceInfo.IsTransient = ResourceLocation.IsTransient();
-		return true;
-	}
-#endif
-
-	/** FRHITexture override.  See FRHITexture::GetNativeResource() */
-	virtual void* GetNativeResource() const override final
-	{
-		FD3D12Resource* Resource = GetResource();
-		return (Resource == nullptr) ? nullptr : Resource->GetResource();
-	}
-
-	// Accessors.
-	FD3D12Resource* GetResource() const { return (FD3D12Resource*)FD3D12TextureBase::GetResource(); }
 
-	virtual void* GetTextureBaseRHI() override final
-	{
-		return static_cast<FD3D12TextureBase*>(this);
-	}
-
-	// IRefCountedObject interface.
-	virtual uint32 AddRef() const
-	{
-		return FRHIResource::AddRef();
-	}
-	virtual uint32 Release() const
-	{
-		return FRHIResource::Release();
-	}
-	virtual uint32 GetRefCount() const
-	{
-		return FRHIResource::GetRefCount();
-	}
-
-	bool IsStreamable() const { return bStreamable; }
+	// Source texture when aliased
+	FTextureRHIRef AliasingSourceTexture;
 
-private:
-	/** Whether the texture has been created with flag TexCreate_Streamable */
-	const uint32 bStreamable : 1;
+	// Syncpoint stored when texture is used for readback
+	FD3D12CLSyncPoint ReadBackSyncPoint;
 };
 
-class FD3D12BaseTexture2D : public FRHITexture2D, public FD3D12FastClearResource
+inline FD3D12RenderTargetView* FD3D12Texture::GetRenderTargetView(int32 MipIndex, int32 ArraySliceIndex) const
 {
-public:
-	FD3D12BaseTexture2D(uint32 InSizeX, uint32 InSizeY, uint32 InSizeZ, uint32 InNumMips, uint32 InNumSamples, EPixelFormat InFormat, ETextureCreateFlags InFlags, const FClearValueBinding& InClearValue)
-		: FRHITexture2D(InSizeX, InSizeY, InNumMips, InNumSamples, InFormat, InFlags, InClearValue)
-	{}
-	uint32 GetSizeZ() const { return 0; }
+	int32 ArrayIndex = MipIndex;
 
-	virtual void GetWriteMaskProperties(void*& OutData, uint32& OutSize) override final
+	if (bCreatedRTVsPerSlice)
 	{
-		FD3D12FastClearResource::GetWriteMaskProperties(OutData, OutSize);
+		check(ArraySliceIndex >= 0);
+		ArrayIndex = MipIndex * RTVArraySizePerMip + ArraySliceIndex;
+		check(ArrayIndex < RenderTargetViews.Num());
 	}
-};
-
-class FD3D12BaseTexture2DArray : public FRHITexture2DArray, public FD3D12FastClearResource
-{
-public:
-	FD3D12BaseTexture2DArray(uint32 InSizeX, uint32 InSizeY, uint32 InSizeZ, uint32 InNumMips, uint32 InNumSamples, EPixelFormat InFormat, ETextureCreateFlags InFlags, const FClearValueBinding& InClearValue)
-		: FRHITexture2DArray(InSizeX, InSizeY, InSizeZ, InNumMips, InNumSamples, InFormat, InFlags, InClearValue)
+	else
 	{
-		check(InNumSamples == 1);
+		// Catch attempts to use a specific slice without having created the texture to support it
+		check(ArraySliceIndex == -1 || ArraySliceIndex == 0);
 	}
-};
 
-class FD3D12BaseTextureCube : public FRHITextureCube, public FD3D12FastClearResource
-{
-public:
-	FD3D12BaseTextureCube(uint32 InSizeX, uint32 InSizeY, uint32 InSizeZ, uint32 InNumMips, uint32 InNumSamples, EPixelFormat InFormat, ETextureCreateFlags InFlags, const FClearValueBinding& InClearValue)
-		: FRHITextureCube(InSizeX, InNumMips, InFormat, InFlags, InClearValue)
-		, SliceCount(InSizeZ)
+	if (ArrayIndex < RenderTargetViews.Num())
 	{
-		check(InNumSamples == 1);
+		return RenderTargetViews[ArrayIndex];
 	}
-	uint32 GetSizeX() const { return GetSize(); }
-	uint32 GetSizeY() const { return GetSize(); }
-	uint32 GetSizeZ() const { return SliceCount; }
+	return 0;
+}
 
-private:
-	uint32 SliceCount;
+template<>
+struct TD3D12ResourceTraits<FRHITexture>
+{
+	typedef FD3D12Texture TConcreteType;
 };
 
-typedef TD3D12Texture2D<FD3D12BaseTexture2D>      FD3D12Texture2D;
-typedef TD3D12Texture2D<FD3D12BaseTexture2DArray> FD3D12Texture2DArray;
-typedef TD3D12Texture2D<FD3D12BaseTextureCube>    FD3D12TextureCube;
-
 class FD3D12Viewport;
 
-class FD3D12BackBufferReferenceTexture2D : public FD3D12Texture2D
+class FD3D12BackBufferReferenceTexture2D : public FD3D12Texture
 {
 public:
 
 	FD3D12BackBufferReferenceTexture2D(
+		const FRHITextureCreateDesc& InDesc,
 		FD3D12Viewport* InViewPort,
 		bool bInIsSDR,
-		FD3D12Device* InDevice,
-		uint32 InSizeX,
-		uint32 InSizeY,
-		EPixelFormat InFormat) :
-		FD3D12Texture2D(InDevice, InSizeX, InSizeY, 1, 1, 1, InFormat, false, TexCreate_RenderTargetable | TexCreate_Presentable, FClearValueBinding()),
+		FD3D12Device* InDevice) :
+		FD3D12Texture(InDesc, InDevice),
 		Viewport(InViewPort), bIsSDR(bInIsSDR)
 	{
 	}
@@ -566,8 +191,8 @@ private:
 	bool bIsSDR = false;
 };
 
-/** Given a pointer to a RHI texture that was created by the D3D12 RHI, returns a pointer to the FD3D12TextureBase it encapsulates. */
-FORCEINLINE FD3D12TextureBase* GetD3D12TextureFromRHITexture(FRHITexture* Texture)
+/** Given a pointer to a RHI texture that was created by the D3D12 RHI, returns a pointer to the FD3D12Texture it encapsulates. */
+FORCEINLINE FD3D12Texture* GetD3D12TextureFromRHITexture(FRHITexture* Texture)
 {
 	if (!Texture)
 	{
@@ -582,14 +207,14 @@ FORCEINLINE FD3D12TextureBase* GetD3D12TextureFromRHITexture(FRHITexture* Textur
 		RHITexture = BufferBufferReferenceTexture->GetBackBufferTexture();
 	}
 
-	FD3D12TextureBase* Result((FD3D12TextureBase*)RHITexture->GetTextureBaseRHI());
+	FD3D12Texture* Result((FD3D12Texture*)RHITexture->GetTextureBaseRHI());
 	check(Result);
 	return Result;
 }
 
-FORCEINLINE FD3D12TextureBase* GetD3D12TextureFromRHITexture(FRHITexture* Texture, uint32 GPUIndex)
+FORCEINLINE FD3D12Texture* GetD3D12TextureFromRHITexture(FRHITexture* Texture, uint32 GPUIndex)
 {
-	FD3D12TextureBase* Result = GetD3D12TextureFromRHITexture(Texture);
+	FD3D12Texture* Result = GetD3D12TextureFromRHITexture(Texture);
 	if (Result != nullptr)
 	{
 		Result = Result->GetLinkedObject(GPUIndex);
@@ -606,74 +231,12 @@ class FD3D12TextureStats
 {
 public:
 
-	static bool ShouldCountAsTextureMemory(D3D12_RESOURCE_FLAGS MiscFlags);
-
-	// @param b3D true:3D, false:2D or cube map
-	static TStatId GetRHIStatEnum(D3D12_RESOURCE_FLAGS MiscFlags, bool bCubeMap, bool b3D);
-	static TStatId GetD3D12StatEnum(D3D12_RESOURCE_FLAGS MiscFlags);
-
 	// Note: This function can be called from many different threads
 	// @param TextureSize >0 to allocate, <0 to deallocate
 	// @param b3D true:3D, false:2D or cube map
 	// @param bStreamable true:Streamable, false:not streamable
-	template<typename TD3D12Texture>
-	static void UpdateD3D12TextureStats(TD3D12Texture& Texture, const D3D12_RESOURCE_DESC& Desc, int64 TextureSize, bool b3D, bool bCubeMap, bool bStreamable, bool bNewTexture = false);
-
-	template<typename BaseResourceType>
-	static void D3D12TextureAllocated(TD3D12Texture2D<BaseResourceType>& Texture, const D3D12_RESOURCE_DESC *Desc = nullptr);
-
-	template<typename BaseResourceType>
-	static void D3D12TextureDeleted(TD3D12Texture2D<BaseResourceType>& Texture);
-
-	static void D3D12TextureAllocated2D(FD3D12Texture2D& Texture);
-
-	static void D3D12TextureAllocated(FD3D12Texture3D& Texture);
-
-	static void D3D12TextureDeleted(FD3D12Texture3D& Texture);
-};
+	static void UpdateD3D12TextureStats(FD3D12Texture& Texture, const D3D12_RESOURCE_DESC& Desc, int64 TextureSize, bool b3D, bool bCubeMap, bool bStreamable, bool bNewTexture);
 
-template<>
-struct TD3D12ResourceTraits<FRHITexture3D>
-{
-	typedef FD3D12Texture3D TConcreteType;
-};
-
-template<>
-struct TD3D12ResourceTraits<FRHITexture2D>
-{
-	typedef FD3D12Texture2D TConcreteType;
+	static void D3D12TextureAllocated(FD3D12Texture& Texture, const D3D12_RESOURCE_DESC *Desc = nullptr);
+	static void D3D12TextureDeleted(FD3D12Texture& Texture);
 };
-
-template<>
-struct TD3D12ResourceTraits<FRHITexture2DArray>
-{
-	typedef FD3D12Texture2DArray TConcreteType;
-};
-
-template<>
-struct TD3D12ResourceTraits<FRHITextureCube>
-{
-	typedef FD3D12TextureCube TConcreteType;
-};
-
-struct FRHICommandD3D12AsyncReallocateTexture2D final : public FRHICommand<FRHICommandD3D12AsyncReallocateTexture2D>
-{
-	FD3D12Texture2D* OldTexture;
-	FD3D12Texture2D* NewTexture;
-	int32 NewMipCount;
-	int32 NewSizeX;
-	int32 NewSizeY;
-	FThreadSafeCounter* RequestStatus;
-
-	FORCEINLINE_DEBUGGABLE FRHICommandD3D12AsyncReallocateTexture2D(FD3D12Texture2D* InOldTexture, FD3D12Texture2D* InNewTexture, int32 InNewMipCount, int32 InNewSizeX, int32 InNewSizeY, FThreadSafeCounter* InRequestStatus)
-		: OldTexture(InOldTexture)
-		, NewTexture(InNewTexture)
-		, NewMipCount(InNewMipCount)
-		, NewSizeX(InNewSizeX)
-		, NewSizeY(InNewSizeY)
-		, RequestStatus(InRequestStatus)
-	{
-	}
-
-	void Execute(FRHICommandListBase& RHICmdList);
-};
\ No newline at end of file
diff --git a/Engine/Source/Runtime/D3D12RHI/Private/D3D12TransientResourceAllocator.cpp b/Engine/Source/Runtime/D3D12RHI/Private/D3D12TransientResourceAllocator.cpp
--- a/Engine/Source/Runtime/D3D12RHI/Private/D3D12TransientResourceAllocator.cpp
+++ b/Engine/Source/Runtime/D3D12RHI/Private/D3D12TransientResourceAllocator.cpp
@@ -144,20 +144,10 @@ FRHITransientTexture* FD3D12TransientResourceHeapAllocator::CreateTexture(const
 	return CreateTextureInternal(InCreateInfo, InDebugName, InPassIndex, Info.SizeInBytes, Info.Alignment,
 		[&](const FRHITransientHeap::FResourceInitializer& Initializer)
 	{
-		ERHIAccess InitialState = ERHIAccess::UAVMask;
-
-		if (EnumHasAnyFlags(InCreateInfo.Flags, TexCreate_RenderTargetable | TexCreate_ResolveTargetable))
-		{
-			InitialState = ERHIAccess::RTV;
-		}
-		else if (EnumHasAnyFlags(InCreateInfo.Flags, TexCreate_DepthStencilTargetable))
-		{
-			InitialState = ERHIAccess::DSVWrite;
-		}
-
-		ED3D12ResourceTransientMode TransientMode = ED3D12ResourceTransientMode::Transient;
 		FResourceAllocatorAdapter ResourceAllocatorAdapter(GetParentAdapter(), static_cast<FD3D12TransientHeap&>(Initializer.Heap), Initializer.Allocation, Desc);
-		FRHITexture* Texture = DynamicRHI->CreateTexture(InCreateInfo, InDebugName, InitialState, TransientMode, &ResourceAllocatorAdapter);
+
+		FRHITextureCreateDesc CreateDesc(InCreateInfo, ERHIAccess::Discard, InDebugName);
+		FRHITexture* Texture = DynamicRHI->CreateD3D12Texture(CreateDesc, nullptr, &ResourceAllocatorAdapter);
 		return new FRHITransientTexture(Texture, ResourceAllocatorAdapter.GpuVirtualAddress, Initializer.Hash, Info.SizeInBytes, ERHITransientAllocationType::Heap, InCreateInfo);
 	});
 }
@@ -208,9 +198,8 @@ FRHITransientBuffer* FD3D12TransientResourceHeapAllocator::CreateBuffer(const FR
 	return CreateBufferInternal(InCreateInfo, InDebugName, InPassIndex, Size, Alignment,
 		[&](const FRHITransientHeap::FResourceInitializer& Initializer)
 	{
-		ED3D12ResourceTransientMode TransientMode = ED3D12ResourceTransientMode::Transient;
 		FResourceAllocatorAdapter ResourceAllocatorAdapter(GetParentAdapter(), static_cast<FD3D12TransientHeap&>(Initializer.Heap), Initializer.Allocation, Desc);
-		FRHIBuffer* Buffer = FD3D12DynamicRHI::GetD3DRHI()->CreateBuffer(InCreateInfo, InDebugName, ERHIAccess::UAVMask, TransientMode, &ResourceAllocatorAdapter);
+		FRHIBuffer* Buffer = FD3D12DynamicRHI::GetD3DRHI()->CreateBuffer(InCreateInfo, InDebugName, ERHIAccess::Discard, &ResourceAllocatorAdapter);
 		return new FRHITransientBuffer(Buffer, ResourceAllocatorAdapter.GpuVirtualAddress, Initializer.Hash, Size, ERHITransientAllocationType::Heap, InCreateInfo);
 	});
 }
\ No newline at end of file
diff --git a/Engine/Source/Runtime/D3D12RHI/Private/D3D12UAV.cpp b/Engine/Source/Runtime/D3D12RHI/Private/D3D12UAV.cpp
--- a/Engine/Source/Runtime/D3D12RHI/Private/D3D12UAV.cpp
+++ b/Engine/Source/Runtime/D3D12RHI/Private/D3D12UAV.cpp
@@ -101,53 +101,63 @@ FUnorderedAccessViewRHIRef FD3D12DynamicRHI::RHICreateUnorderedAccessView(FRHIBu
 
 FUnorderedAccessViewRHIRef FD3D12DynamicRHI::RHICreateUnorderedAccessView(FRHITexture* TextureRHI, uint32 MipLevel, uint8 Format, uint16 FirstArraySlice, uint16 NumArraySlices)
 {
-	FD3D12TextureBase* Texture = GetD3D12TextureFromRHITexture(TextureRHI);
+	FD3D12Texture* Texture = GetD3D12TextureFromRHITexture(TextureRHI);
+	ETextureDimension Dimension = TextureRHI->GetDesc().Dimension;
 
 	D3D12_UNORDERED_ACCESS_VIEW_DESC UAVDesc = {};
 
 	const DXGI_FORMAT PlatformResourceFormat = (DXGI_FORMAT)GPixelFormats[Format].PlatformFormat;
 	UAVDesc.Format = FindShaderResourceDXGIFormat(PlatformResourceFormat, false);
 
-	if (TextureRHI->GetTexture3D() != NULL)
+	switch (Dimension)
+	{
+	case ETextureDimension::Texture3D:
 	{
-		FD3D12Texture3D* Texture3D = (FD3D12Texture3D*)Texture;
 		UAVDesc.ViewDimension = D3D12_UAV_DIMENSION_TEXTURE3D;
 		UAVDesc.Texture3D.MipSlice = MipLevel;
 		UAVDesc.Texture3D.FirstWSlice = 0;
-		UAVDesc.Texture3D.WSize = Texture3D->GetSizeZ() >> MipLevel;
+		UAVDesc.Texture3D.WSize = Texture->GetDesc().Depth >> MipLevel;
 
-		return CreateUAV(UAVDesc, Texture3D, false);
+		return CreateUAV(UAVDesc, Texture, false);
 	}
-	else if (TextureRHI->GetTexture2DArray() != NULL)
+	case ETextureDimension::Texture2DArray:
 	{
-		FD3D12Texture2DArray* Texture2DArray = (FD3D12Texture2DArray*)Texture;
 		UAVDesc.ViewDimension = D3D12_UAV_DIMENSION_TEXTURE2DARRAY;
 		UAVDesc.Texture2DArray.MipSlice = MipLevel;
 		UAVDesc.Texture2DArray.FirstArraySlice = NumArraySlices == 0 ? 0 : FirstArraySlice;
-		UAVDesc.Texture2DArray.ArraySize = NumArraySlices == 0 ? Texture2DArray->GetSizeZ() : NumArraySlices;
+		UAVDesc.Texture2DArray.ArraySize = NumArraySlices == 0 ? Texture->GetDesc().ArraySize : NumArraySlices;
 		UAVDesc.Texture2DArray.PlaneSlice = GetPlaneSliceFromViewFormat(PlatformResourceFormat, UAVDesc.Format);
 
-		return CreateUAV(UAVDesc, Texture2DArray, false);
+		return CreateUAV(UAVDesc, Texture, false);
 	}
-	else if (TextureRHI->GetTextureCube() != NULL)
+	case ETextureDimension::TextureCube:
 	{
-		FD3D12TextureCube* TextureCube = (FD3D12TextureCube*)Texture;
 		UAVDesc.ViewDimension = D3D12_UAV_DIMENSION_TEXTURE2DARRAY;
 		UAVDesc.Texture2DArray.MipSlice = MipLevel;
 		UAVDesc.Texture2DArray.FirstArraySlice = 0;
-		UAVDesc.Texture2DArray.ArraySize = TextureCube->GetSizeZ();
+		UAVDesc.Texture2DArray.ArraySize = 6;
 		UAVDesc.Texture2DArray.PlaneSlice = GetPlaneSliceFromViewFormat(PlatformResourceFormat, UAVDesc.Format);
 
-		return CreateUAV(UAVDesc, TextureCube, false);
+		return CreateUAV(UAVDesc, Texture, false);
 	}
-	else
+	case ETextureDimension::TextureCubeArray:
+	{
+		UAVDesc.ViewDimension = D3D12_UAV_DIMENSION_TEXTURE2DARRAY;
+		UAVDesc.Texture2DArray.MipSlice = MipLevel;
+		UAVDesc.Texture2DArray.FirstArraySlice = NumArraySlices == 0 ? 0 : FirstArraySlice * 6;
+		UAVDesc.Texture2DArray.ArraySize = NumArraySlices == 0 ? Texture->GetDesc().ArraySize * 6 : NumArraySlices * 6;
+		UAVDesc.Texture2DArray.PlaneSlice = GetPlaneSliceFromViewFormat(PlatformResourceFormat, UAVDesc.Format);
+
+		return CreateUAV(UAVDesc, Texture, false);
+	}
+	default:
 	{
-		FD3D12Texture2D* Texture2D = (FD3D12Texture2D*)Texture;
 		UAVDesc.ViewDimension = D3D12_UAV_DIMENSION_TEXTURE2D;
 		UAVDesc.Texture2D.MipSlice = MipLevel;
 		UAVDesc.Texture2D.PlaneSlice = GetPlaneSliceFromViewFormat(PlatformResourceFormat, UAVDesc.Format);
 
-		return CreateUAV(UAVDesc, Texture2D, false);
+		return CreateUAV(UAVDesc, Texture, false);
+	}
 	}
 }
 
diff --git a/Engine/Source/Runtime/D3D12RHI/Private/D3D12View.cpp b/Engine/Source/Runtime/D3D12RHI/Private/D3D12View.cpp
--- a/Engine/Source/Runtime/D3D12RHI/Private/D3D12View.cpp
+++ b/Engine/Source/Runtime/D3D12RHI/Private/D3D12View.cpp
@@ -139,7 +139,7 @@ FD3D12ShaderResourceView::FD3D12ShaderResourceView(FD3D12Buffer* InBuffer, const
 	InitializeAfterCreate(InDesc, InBuffer, InBuffer->ResourceLocation, InStride, InStartOffsetBytes);
 }
 
-FD3D12ShaderResourceView::FD3D12ShaderResourceView(FD3D12TextureBase* InTexture, const D3D12_SHADER_RESOURCE_VIEW_DESC& InDesc, ETextureCreateFlags InTextureCreateFlags)
+FD3D12ShaderResourceView::FD3D12ShaderResourceView(FD3D12Texture* InTexture, const D3D12_SHADER_RESOURCE_VIEW_DESC& InDesc, ETextureCreateFlags InTextureCreateFlags)
 	: FD3D12ShaderResourceView(InTexture->GetParentDevice())
 {
 	InitializeAfterCreate(InDesc, InTexture, InTexture->ResourceLocation, -1, 0, EnumHasAnyFlags(InTextureCreateFlags, ETextureCreateFlags::NoFastClearFinalize));
@@ -360,49 +360,53 @@ FD3D12ShaderResourceView* CreateSRV(TextureType* Texture, const D3D12_SHADER_RES
 	});
 }
 
-FShaderResourceViewRHIRef FD3D12DynamicRHI::RHICreateShaderResourceView(FRHITexture* Texture, const FRHITextureSRVCreateInfo& CreateInfo)
+FShaderResourceViewRHIRef FD3D12DynamicRHI::RHICreateShaderResourceView(FRHITexture* RHITexture, const FRHITextureSRVCreateInfo& CreateInfo)
 {
+	FD3D12Texture* Texture = GetD3D12TextureFromRHITexture(RHITexture);
+	ETextureDimension Dimension = RHITexture->GetDesc().Dimension;
+
 	D3D12_SHADER_RESOURCE_VIEW_DESC SRVDesc = {};
 	SRVDesc.Shader4ComponentMapping = D3D12_DEFAULT_SHADER_4_COMPONENT_MAPPING;
 
-	DXGI_FORMAT BaseTextureFormat = DXGI_FORMAT_UNKNOWN;
-	FD3D12TextureBase* BaseTexture = nullptr;
-	if (FD3D12Texture3D* Texture3D = FD3D12DynamicRHI::ResourceCast(Texture->GetTexture3D()))
+	const D3D12_RESOURCE_DESC& TextureDesc = Texture->GetResource()->GetDesc();
+	DXGI_FORMAT BaseTextureFormat = TextureDesc.Format;
+
+	switch (Dimension)
+	{
+	case ETextureDimension::Texture3D:
 	{
-		const D3D12_RESOURCE_DESC& TextureDesc = Texture3D->GetResource()->GetDesc();
-		BaseTextureFormat = TextureDesc.Format;
-		BaseTexture = Texture3D;
 		SRVDesc.ViewDimension = D3D12_SRV_DIMENSION_TEXTURE3D;
 		SRVDesc.Texture3D.MipLevels = CreateInfo.NumMipLevels;
 		SRVDesc.Texture3D.MostDetailedMip = CreateInfo.MipLevel;
+		break;
 	}
-	else if (FD3D12Texture2DArray* Texture2DArray = FD3D12DynamicRHI::ResourceCast(Texture->GetTexture2DArray()))
+	case ETextureDimension::Texture2DArray:
 	{
-		const D3D12_RESOURCE_DESC& TextureDesc = Texture2DArray->GetResource()->GetDesc();
-		BaseTextureFormat = TextureDesc.Format;
-		BaseTexture = Texture2DArray;
 		SRVDesc.ViewDimension = D3D12_SRV_DIMENSION_TEXTURE2DARRAY;
 		SRVDesc.Texture2DArray.ArraySize = (CreateInfo.NumArraySlices == 0 ? TextureDesc.DepthOrArraySize : CreateInfo.NumArraySlices);
 		SRVDesc.Texture2DArray.FirstArraySlice = CreateInfo.FirstArraySlice;
 		SRVDesc.Texture2DArray.MipLevels = CreateInfo.NumMipLevels;
 		SRVDesc.Texture2DArray.MostDetailedMip = CreateInfo.MipLevel;
+		break;
 	}
-	else if (FD3D12TextureCube* TextureCube = FD3D12DynamicRHI::ResourceCast(Texture->GetTextureCube()))
+	case ETextureDimension::TextureCube:
 	{
-		const D3D12_RESOURCE_DESC& TextureDesc = TextureCube->GetResource()->GetDesc();
-		BaseTextureFormat = TextureDesc.Format;
-		BaseTexture = TextureCube;
 		SRVDesc.ViewDimension = D3D12_SRV_DIMENSION_TEXTURECUBE;
 		SRVDesc.TextureCube.MipLevels = CreateInfo.NumMipLevels;
 		SRVDesc.TextureCube.MostDetailedMip = CreateInfo.MipLevel;
+		break;
 	}
-	else
+	case ETextureDimension::TextureCubeArray:
 	{
-		FD3D12Texture2D* Texture2D = FD3D12DynamicRHI::ResourceCast(Texture->GetTexture2D());
-		const D3D12_RESOURCE_DESC& TextureDesc = Texture2D->GetResource()->GetDesc();
-		BaseTextureFormat = TextureDesc.Format;
-		BaseTexture = Texture2D;
-
+		SRVDesc.ViewDimension = D3D12_SRV_DIMENSION_TEXTURECUBEARRAY;
+		SRVDesc.TextureCubeArray.First2DArrayFace = CreateInfo.FirstArraySlice;
+		SRVDesc.TextureCubeArray.NumCubes = (CreateInfo.NumArraySlices == 0 ? TextureDesc.DepthOrArraySize / 6 : CreateInfo.NumArraySlices / 6);
+		SRVDesc.TextureCubeArray.MipLevels = CreateInfo.NumMipLevels;
+		SRVDesc.TextureCubeArray.MostDetailedMip = CreateInfo.MipLevel;
+		break;
+	}
+	default:
+	{		
 		if (TextureDesc.SampleDesc.Count > 1)
 		{
 			// MS textures can't have mips apparently, so nothing else to set.
@@ -414,10 +418,12 @@ FShaderResourceViewRHIRef FD3D12DynamicRHI::RHICreateShaderResourceView(FRHIText
 			SRVDesc.Texture2D.MipLevels = CreateInfo.NumMipLevels;
 			SRVDesc.Texture2D.MostDetailedMip = CreateInfo.MipLevel;
 		}
+		break;
+	}
 	}
 
 	// Allow input CreateInfo to override SRGB and/or format
-	const bool bBaseSRGB = EnumHasAnyFlags(Texture->GetFlags(), TexCreate_SRGB);
+	const bool bBaseSRGB = EnumHasAnyFlags(RHITexture->GetFlags(), TexCreate_SRGB);
 	const bool bSRGB = CreateInfo.SRGBOverride != SRGBO_ForceDisable && bBaseSRGB;
 	const DXGI_FORMAT ViewTextureFormat = (CreateInfo.Format == PF_Unknown) ? BaseTextureFormat : (DXGI_FORMAT)GPixelFormats[CreateInfo.Format].PlatformFormat;
 	SRVDesc.Format = FindShaderResourceDXGIFormat(ViewTextureFormat, bSRGB);
@@ -429,8 +435,8 @@ FShaderResourceViewRHIRef FD3D12DynamicRHI::RHICreateShaderResourceView(FRHIText
 	default: break; // other view types don't support PlaneSlice
 	}
 
-	check(BaseTexture);
-	return CreateSRV(BaseTexture, SRVDesc);
+	check(Texture);
+	return CreateSRV(Texture, SRVDesc);
 }
 
 FShaderResourceViewRHIRef FD3D12DynamicRHI::RHICreateShaderResourceView(FRHIBuffer* BufferRHI)
diff --git a/Engine/Source/Runtime/D3D12RHI/Private/D3D12View.h b/Engine/Source/Runtime/D3D12RHI/Private/D3D12View.h
--- a/Engine/Source/Runtime/D3D12RHI/Private/D3D12View.h
+++ b/Engine/Source/Runtime/D3D12RHI/Private/D3D12View.h
@@ -4,7 +4,7 @@
 
 #include "D3D12Descriptors.h"
 
-class FD3D12TextureBase;
+class FD3D12Texture;
 
 enum ViewSubresourceSubsetFlags
 {
@@ -803,7 +803,7 @@ public:
 	// Used for buffer SRVs, we usually only control the stride and offset
 	FD3D12ShaderResourceView(FD3D12Buffer* InBuffer, const D3D12_SHADER_RESOURCE_VIEW_DESC& InDesc, uint32 InStride, uint32 InStartOffsetBytes = 0);
 	// Used for texture SRVs, we don't control much other than disabling fast create
-	FD3D12ShaderResourceView(FD3D12TextureBase* InTexture, const D3D12_SHADER_RESOURCE_VIEW_DESC& InDesc, ETextureCreateFlags InTextureCreateFlags = ETextureCreateFlags::None);
+	FD3D12ShaderResourceView(FD3D12Texture* InTexture, const D3D12_SHADER_RESOURCE_VIEW_DESC& InDesc, ETextureCreateFlags InTextureCreateFlags = ETextureCreateFlags::None);
 	~FD3D12ShaderResourceView();
 
 	void InitializeAfterCreate(const D3D12_SHADER_RESOURCE_VIEW_DESC& InDesc, FD3D12BaseShaderResource* InBaseShaderResource, FD3D12ResourceLocation& InResourceLocation, uint32 InStride, uint32 InStartOffsetBytes = 0, bool InSkipFastClearFinalize = false);
diff --git a/Engine/Source/Runtime/D3D12RHI/Private/D3D12Viewport.cpp b/Engine/Source/Runtime/D3D12RHI/Private/D3D12Viewport.cpp
--- a/Engine/Source/Runtime/D3D12RHI/Private/D3D12Viewport.cpp
+++ b/Engine/Source/Runtime/D3D12RHI/Private/D3D12Viewport.cpp
@@ -202,7 +202,7 @@ bool bNeedSwapChain = true;
 /**
  * Creates a FD3D12Surface to represent a swap chain's back buffer.
  */
-FD3D12Texture2D* GetSwapChainSurface(FD3D12Device* Parent, EPixelFormat PixelFormat, uint32 SizeX, uint32 SizeY, IDXGISwapChain* SwapChain, uint32 BackBufferIndex, TRefCountPtr<ID3D12Resource> BackBufferResourceOverride)
+FD3D12Texture* GetSwapChainSurface(FD3D12Device* Parent, EPixelFormat PixelFormat, uint32 SizeX, uint32 SizeY, IDXGISwapChain* SwapChain, uint32 BackBufferIndex, TRefCountPtr<ID3D12Resource> BackBufferResourceOverride)
 {
 	verify(D3D12_VIEWPORT_EXPOSES_SWAP_CHAIN || SwapChain == nullptr);
 
@@ -245,18 +245,14 @@ FD3D12Texture2D* GetSwapChainSurface(FD3D12Device* Parent, EPixelFormat PixelFor
 
 	D3D12_RESOURCE_DESC BackBufferDesc = BackBufferResource->GetDesc();
 
-	FD3D12Texture2D* SwapChainTexture = Adapter->CreateLinkedObject<FD3D12Texture2D>(FRHIGPUMask::All(), [&](FD3D12Device* Device)
+	FString Name = FString::Printf(TEXT("BackBuffer%d"), BackBufferIndex);
+	FRHITextureCreateDesc CreateDesc(FRHITextureDesc::Create2D(FIntPoint((uint32)BackBufferDesc.Width, BackBufferDesc.Height), PixelFormat, FClearValueBinding(), TexCreate_RenderTargetable, 1, 1), ERHIAccess::Present, *Name);
+
+	FD3D12DynamicRHI* DynamicRHI = FD3D12DynamicRHI::GetD3DRHI();
+
+	FD3D12Texture* SwapChainTexture = Adapter->CreateLinkedObject<FD3D12Texture>(FRHIGPUMask::All(), [&](FD3D12Device* Device)
 	{
-		FD3D12Texture2D* NewTexture = new FD3D12Texture2D(Device,
-			(uint32)BackBufferDesc.Width,
-			BackBufferDesc.Height,
-			1,
-			1,
-			1,
-			PixelFormat,
-			false,
-			TexCreate_RenderTargetable,
-			FClearValueBinding());
+		FD3D12Texture* NewTexture = DynamicRHI->CreateNewD3D12Texture(CreateDesc, Device);
 
 		const D3D12_RESOURCE_STATES InitialState = D3D12_RESOURCE_STATE_COMMON;
 
@@ -339,12 +335,14 @@ FD3D12Texture2D* GetSwapChainSurface(FD3D12Device* Parent, EPixelFormat PixelFor
 		return NewTexture;
 	});
 
-	FString Name = FString::Printf(TEXT("BackBuffer%d"), BackBufferIndex);
 	SetName(SwapChainTexture->GetResource(), *Name);
 
 	SwapChainTexture->GetResource()->SetIsBackBuffer(true);
 
-	FD3D12TextureStats::D3D12TextureAllocated2D(*SwapChainTexture);
+	const D3D12_RESOURCE_ALLOCATION_INFO AllocationInfo = Parent->GetDevice()->GetResourceAllocationInfo(0, 1, &SwapChainTexture->GetResource()->GetDesc());
+	SwapChainTexture->ResourceLocation.SetSize(AllocationInfo.SizeInBytes);
+
+	FD3D12TextureStats::D3D12TextureAllocated(*SwapChainTexture);
 	return SwapChainTexture;
 }
 
@@ -352,11 +350,13 @@ FD3D12Texture2D* GetSwapChainSurface(FD3D12Device* Parent, EPixelFormat PixelFor
 /**
 * Create the dummy back buffer textures - They don't have actual D3D resource but are used to always reference the current back buffer index on the RHI thread
 */
-FD3D12Texture2D* FD3D12Viewport::CreateDummyBackBufferTextures(FD3D12Adapter* InAdapter, EPixelFormat InPixelFormat, uint32 InSizeX, uint32 InSizeY, bool bInIsSDR)
+FD3D12Texture* FD3D12Viewport::CreateDummyBackBufferTextures(FD3D12Adapter* InAdapter, EPixelFormat InPixelFormat, uint32 InSizeX, uint32 InSizeY, bool bInIsSDR)
 {
-	FD3D12Texture2D* Result = InAdapter->CreateLinkedObject<FD3D12Texture2D>(FRHIGPUMask::All(), [&](FD3D12Device* Device)
+	FRHITextureCreateDesc CreateDesc(FRHITextureDesc::Create2D(FIntPoint(InSizeX, InSizeY), InPixelFormat, FClearValueBinding(), TexCreate_RenderTargetable | TexCreate_Presentable, 1, 1), ERHIAccess::Present, TEXT("BackBufferReference"));
+
+	FD3D12Texture* Result = InAdapter->CreateLinkedObject<FD3D12Texture>(FRHIGPUMask::All(), [&](FD3D12Device* Device)
 	{
-		FD3D12Texture2D* NewTexture = new FD3D12BackBufferReferenceTexture2D(this, bInIsSDR, Device, InSizeX, InSizeY, PixelFormat);
+		FD3D12Texture* NewTexture = new FD3D12BackBufferReferenceTexture2D(CreateDesc, this, bInIsSDR, Device);
 		return NewTexture;
 	});
 	return Result;
@@ -478,7 +478,7 @@ void FD3D12Viewport::Resize(uint32 InSizeX, uint32 InSizeY, bool bInIsFullscreen
 			}
 			check(BackBuffers[i]->GetRefCount() == 1);
 
-			for (FD3D12TextureBase& Tex : *BackBuffers[i])
+			for (FD3D12Texture& Tex : *BackBuffers[i])
 			{
 				Tex.GetResource()->DoNotDeferDelete();
 			}
@@ -495,7 +495,7 @@ void FD3D12Viewport::Resize(uint32 InSizeX, uint32 InSizeY, bool bInIsFullscreen
 			}
 			check(SDRBackBuffers[i]->GetRefCount() == 1);
 
-			for (FD3D12TextureBase& Tex : *SDRBackBuffers[i])
+			for (FD3D12Texture& Tex : *SDRBackBuffers[i])
 			{
 				Tex.GetResource()->DoNotDeferDelete();
 			}
@@ -817,8 +817,8 @@ bool FD3D12Viewport::Present(bool bLockToVsync)
 		FD3D12CommandContext& DefaultContext = Device->GetDefaultCommandContext();
 
 		// Those are not necessarily the swap chain back buffer in case of multi-gpu
-		FD3D12Texture2D* DeviceBackBuffer = DefaultContext.RetrieveObject<FD3D12Texture2D, FRHITexture2D*>(GetBackBuffer_RHIThread());
-		FD3D12Texture2D* DeviceSDRBackBuffer = DefaultContext.RetrieveObject<FD3D12Texture2D, FRHITexture2D*>(GetSDRBackBuffer_RHIThread());
+		FD3D12Texture* DeviceBackBuffer = DefaultContext.RetrieveObject<FD3D12Texture, FRHITexture2D*>(GetBackBuffer_RHIThread());
+		FD3D12Texture* DeviceSDRBackBuffer = DefaultContext.RetrieveObject<FD3D12Texture, FRHITexture2D*>(GetSDRBackBuffer_RHIThread());
 
 		FD3D12DynamicRHI::TransitionResource(DefaultContext.CommandListHandle, DeviceBackBuffer->GetShaderResourceView()->GetResource(), D3D12_RESOURCE_STATE_TBD, D3D12_RESOURCE_STATE_PRESENT, 0, FD3D12DynamicRHI::ETransitionMode::Apply);
 		if (SDRBackBuffer_RHIThread != nullptr)
diff --git a/Engine/Source/Runtime/D3D12RHI/Private/Windows/WindowsD3D12Viewport.cpp b/Engine/Source/Runtime/D3D12RHI/Private/Windows/WindowsD3D12Viewport.cpp
--- a/Engine/Source/Runtime/D3D12RHI/Private/Windows/WindowsD3D12Viewport.cpp
+++ b/Engine/Source/Runtime/D3D12RHI/Private/Windows/WindowsD3D12Viewport.cpp
@@ -14,7 +14,7 @@
 
 static const uint32 WindowsDefaultNumBackBuffers = 3;
 
-extern FD3D12Texture2D* GetSwapChainSurface(FD3D12Device* Parent, EPixelFormat PixelFormat, uint32 SizeX, uint32 SizeY, IDXGISwapChain* SwapChain, uint32 BackBufferIndex, TRefCountPtr<ID3D12Resource> BackBufferResourceOverride);
+extern FD3D12Texture* GetSwapChainSurface(FD3D12Device* Parent, EPixelFormat PixelFormat, uint32 SizeX, uint32 SizeY, IDXGISwapChain* SwapChain, uint32 BackBufferIndex, TRefCountPtr<ID3D12Resource> BackBufferResourceOverride);
 
 static int32 GD3D12UseAllowTearing = 1;
 static FAutoConsoleVariableRef CVarD3DUseAllowTearing(
diff --git a/Engine/Source/Runtime/D3D12RHI/Public/D3D12Resources.h b/Engine/Source/Runtime/D3D12RHI/Public/D3D12Resources.h
--- a/Engine/Source/Runtime/D3D12RHI/Public/D3D12Resources.h
+++ b/Engine/Source/Runtime/D3D12RHI/Public/D3D12Resources.h
@@ -6,8 +6,12 @@
 
 #pragma once
 
+#include "D3D12RHIPrivate.h"
 #include "BoundShaderStateCache.h"
 #include "D3D12ShaderResources.h"
+#include "D3D12Residency.h"
+#include "D3D12Util.h"
+#include "D3D12State.h"
 #include "RHIPoolAllocator.h"
 
 constexpr D3D12_RESOURCE_STATES BackBufferBarrierWriteTransitionTargets = D3D12_RESOURCE_STATES(
@@ -42,12 +46,6 @@ enum class ED3D12ResourceStateMode
 	MultiState,					//< Force enable state tracking of resource
 };
 
-enum class ED3D12ResourceTransientMode
-{
-	NonTransient,				//< Resource is not transient
-	Transient,					//< Resource is transient
-};
-
 class FD3D12PendingResourceBarrier
 {
 public:
@@ -782,19 +780,6 @@ protected:
 };
 
 
-#if PLATFORM_WINDOWS || PLATFORM_HOLOLENS
-class FD3D12FastClearResource
-{
-public:
-	inline void GetWriteMaskProperties(void*& OutData, uint32& OutSize)
-	{
-		OutData = nullptr;
-		OutSize = 0;
-	}
-};
-#endif
-
-
 /** The base class of resources that may be bound as shader resources (texture or buffer). */
 class FD3D12BaseShaderResource : public FD3D12DeviceChild, public IRefCountedObject
 {
diff --git a/Engine/Source/Runtime/D3D12RHI/Public/D3D12Viewport.h b/Engine/Source/Runtime/D3D12RHI/Public/D3D12Viewport.h
--- a/Engine/Source/Runtime/D3D12RHI/Public/D3D12Viewport.h
+++ b/Engine/Source/Runtime/D3D12RHI/Public/D3D12Viewport.h
@@ -92,10 +92,10 @@ public:
 	// Accessors.
 	FIntPoint GetSizeXY() const { return FIntPoint(SizeX, SizeY); }
 
-	FD3D12Texture2D* GetDummyBackBuffer_RenderThread(bool bInIsSDR) const { return bInIsSDR ? SDRDummyBackBuffer_RenderThread : DummyBackBuffer_RenderThread; }
+	FD3D12Texture* GetDummyBackBuffer_RenderThread(bool bInIsSDR) const { return bInIsSDR ? SDRDummyBackBuffer_RenderThread : DummyBackBuffer_RenderThread; }
 
-	FD3D12Texture2D* GetBackBuffer_RHIThread() const { return BackBuffer_RHIThread; }
-	FD3D12Texture2D* GetSDRBackBuffer_RHIThread() const { return (PixelFormat == SDRPixelFormat) ? GetBackBuffer_RHIThread() : SDRBackBuffer_RHIThread; }
+	FD3D12Texture* GetBackBuffer_RHIThread() const { return BackBuffer_RHIThread; }
+	FD3D12Texture* GetSDRBackBuffer_RHIThread() const { return (PixelFormat == SDRPixelFormat) ? GetBackBuffer_RHIThread() : SDRBackBuffer_RHIThread; }
 
 	virtual void WaitForFrameEventCompletion() override;
 	virtual void IssueFrameEvent() override;
@@ -159,7 +159,7 @@ private:
 	/**
 	 * Create the dummy back buffer textures
 	 */
-	FD3D12Texture2D* CreateDummyBackBufferTextures(FD3D12Adapter* InAdapter, EPixelFormat InPixelFormat, uint32 InSizeX, uint32 InSizeY, bool bInIsSDR);
+	FD3D12Texture* CreateDummyBackBufferTextures(FD3D12Adapter* InAdapter, EPixelFormat InPixelFormat, uint32 InSizeX, uint32 InSizeY, bool bInIsSDR);
 
 	/** Presents the frame synchronizing with DWM. */
 	void PresentWithVsyncDWM();
@@ -205,12 +205,12 @@ private:
 	TRefCountPtr<IDXGISwapChain1> SDRSwapChain1;
 #endif // D3D12_VIEWPORT_EXPOSES_SWAP_CHAIN
 
-	TArray<TRefCountPtr<FD3D12Texture2D>> BackBuffers;
+	TArray<TRefCountPtr<FD3D12Texture>> BackBuffers;
 	uint32 NumBackBuffers;
 
-	TRefCountPtr<FD3D12Texture2D> DummyBackBuffer_RenderThread; // Dummy back buffer texture which always references the current back buffer on the RHI thread
+	TRefCountPtr<FD3D12Texture> DummyBackBuffer_RenderThread; // Dummy back buffer texture which always references the current back buffer on the RHI thread
 	uint32 CurrentBackBufferIndex_RHIThread;
-	FD3D12Texture2D* BackBuffer_RHIThread;
+	FD3D12Texture* BackBuffer_RHIThread;
 
 #if WITH_MGPU
 	int32 BackbufferMultiGPUBinding; // where INDEX_NONE cycles through the GPU, otherwise the GPU index.
@@ -222,9 +222,9 @@ private:
 	/** 
 	 * When HDR is enabled, SDR backbuffers may be required on some architectures for game DVR or broadcasting
 	 */
-	TArray<TRefCountPtr<FD3D12Texture2D>> SDRBackBuffers;
-	TRefCountPtr<FD3D12Texture2D> SDRDummyBackBuffer_RenderThread;
-	FD3D12Texture2D* SDRBackBuffer_RHIThread;
+	TArray<TRefCountPtr<FD3D12Texture>> SDRBackBuffers;
+	TRefCountPtr<FD3D12Texture> SDRDummyBackBuffer_RenderThread;
+	FD3D12Texture* SDRBackBuffer_RHIThread;
 	EPixelFormat SDRPixelFormat;
 
 	/** A fence value used to track the GPU's progress. */
diff --git a/Engine/Source/Runtime/Engine/Private/Streaming/VolumeTextureStreaming.cpp b/Engine/Source/Runtime/Engine/Private/Streaming/VolumeTextureStreaming.cpp
--- a/Engine/Source/Runtime/Engine/Private/Streaming/VolumeTextureStreaming.cpp
+++ b/Engine/Source/Runtime/Engine/Private/Streaming/VolumeTextureStreaming.cpp
@@ -8,51 +8,6 @@ VolumeTextureStreaming.cpp: Helpers to stream in and out volume texture LODs.
 
 extern RHI_API bool GUseTexture3DBulkDataRHI;
 
-//*****************************************************************************
-//***************************** Global Definitions ****************************
-//*****************************************************************************
-
-void RHICopySharedMips(FRHICommandList& RHICmdList, FRHITexture3D* DestTexture, FRHITexture3D* SrcTexture)
-{
-	// Transition to copy source and dest
-	{
-		FRHITransitionInfo TransitionsBefore[] = { FRHITransitionInfo(SrcTexture, ERHIAccess::SRVMask, ERHIAccess::CopySrc), FRHITransitionInfo(DestTexture, ERHIAccess::SRVMask, ERHIAccess::CopyDest) };
-		RHICmdList.Transition(MakeArrayView(TransitionsBefore, UE_ARRAY_COUNT(TransitionsBefore)));
-	}
-
-	// Copy 
-	{
-		FRHICopyTextureInfo CopyInfo;
-
-		auto SetCopyInfo = [&](FRHITexture3D* Texture3D)
-		{
-			CopyInfo.Size.X = Texture3D->GetSizeX();
-			CopyInfo.Size.Y = Texture3D->GetSizeY();
-			CopyInfo.Size.Z = Texture3D->GetSizeZ();
-			CopyInfo.NumMips = Texture3D->GetNumMips();
-		};
-
-		if (DestTexture->GetNumMips() < SrcTexture->GetNumMips())
-		{
-			SetCopyInfo(DestTexture);
-		}
-		else
-		{
-			SetCopyInfo(SrcTexture);
-		}
-
-		CopyInfo.SourceMipIndex = SrcTexture->GetNumMips() - CopyInfo.NumMips;
-		CopyInfo.DestMipIndex = DestTexture->GetNumMips() - CopyInfo.NumMips;
-		RHICmdList.CopyTexture(SrcTexture, DestTexture, CopyInfo);
-	}
-
-	// Transition to SRV
-	{
-		FRHITransitionInfo TransitionsAfter[] = { FRHITransitionInfo(SrcTexture, ERHIAccess::CopySrc, ERHIAccess::SRVMask), FRHITransitionInfo(DestTexture, ERHIAccess::CopyDest, ERHIAccess::SRVMask) };
-		RHICmdList.Transition(MakeArrayView(TransitionsAfter, UE_ARRAY_COUNT(TransitionsAfter)));
-	}
-}
-
 //*****************************************************************************
 //******************** FVolumeTextureMipAllocator_Reallocate ******************
 //*****************************************************************************
@@ -135,7 +90,7 @@ bool FVolumeTextureMipAllocator_Reallocate::FinalizeMips(const FTextureUpdateCon
 		ENQUEUE_RENDER_COMMAND(FCopySharedMipsForTexture3D)(
 			[&](FRHICommandListImmediate& RHICmdList)
 		{
-			RHICopySharedMips(RHICmdList, IntermediateTextureRHI.GetReference(), Context.Resource->GetTexture3DRHI());
+			RHICmdList.CopySharedMips(IntermediateTextureRHI.GetReference(), Context.Resource->GetTexture3DRHI());
 			bCopySharedMipsDone = true;
 		});
 		// Expected to execute immediately since ran on the renderthread.
diff --git a/Engine/Source/Runtime/Engine/Private/UnrealEngine.cpp b/Engine/Source/Runtime/Engine/Private/UnrealEngine.cpp
--- a/Engine/Source/Runtime/Engine/Private/UnrealEngine.cpp
+++ b/Engine/Source/Runtime/Engine/Private/UnrealEngine.cpp
@@ -3538,7 +3538,7 @@ public:
 	{
 	}
 
-	virtual void RenderTexture_RenderThread(FRHICommandListImmediate& RHICmdList, FRHITexture2D* BackBuffer, FRHITexture2D* SrcTexture, FVector2D WindowSize) const override
+	virtual void RenderTexture_RenderThread(FRHICommandListImmediate& RHICmdList, FRHITexture* BackBuffer, FRHITexture* SrcTexture, FVector2D WindowSize) const override
 	{
 		check(IsInRenderingThread());
 
diff --git a/Engine/Source/Runtime/Engine/Public/StereoRendering.h b/Engine/Source/Runtime/Engine/Public/StereoRendering.h
--- a/Engine/Source/Runtime/Engine/Public/StereoRendering.h
+++ b/Engine/Source/Runtime/Engine/Public/StereoRendering.h
@@ -158,7 +158,7 @@ public:
 
 	// Renders texture into a backbuffer. Could be empty if no rendertarget texture is used, or if direct-rendering 
 	// through RHI bridge is implemented. 
-	virtual void RenderTexture_RenderThread(class FRHICommandListImmediate& RHICmdList, class FRHITexture2D* BackBuffer, class FRHITexture2D* SrcTexture, FVector2D WindowSize) const {}
+	virtual void RenderTexture_RenderThread(class FRHICommandListImmediate& RHICmdList, class FRHITexture* BackBuffer, class FRHITexture* SrcTexture, FVector2D WindowSize) const {}
 
 	/**
 	 * Returns currently active render target manager.
diff --git a/Engine/Source/Runtime/GameplayMediaEncoder/Private/GameplayMediaEncoder.cpp b/Engine/Source/Runtime/GameplayMediaEncoder/Private/GameplayMediaEncoder.cpp
--- a/Engine/Source/Runtime/GameplayMediaEncoder/Private/GameplayMediaEncoder.cpp
+++ b/Engine/Source/Runtime/GameplayMediaEncoder/Private/GameplayMediaEncoder.cpp
@@ -586,7 +586,7 @@ TSharedPtr<AVEncoder::FVideoEncoderInputFrame> FGameplayMediaEncoder::ObtainInpu
 		if (RHIType == ERHIInterfaceType::D3D11)
 		{
 			FRHIResourceCreateInfo CreateInfo(TEXT("VideoCapturerBackBuffer"));
-			FTexture2DRHIRef Texture = GDynamicRHI->RHICreateTexture2D(VideoConfig.Width, VideoConfig.Height, EPixelFormat::PF_B8G8R8A8, 1, 1,
+			FTexture2DRHIRef Texture = RHICreateTexture2D(VideoConfig.Width, VideoConfig.Height, EPixelFormat::PF_B8G8R8A8, 1, 1,
 			                                                           TexCreate_Shared | TexCreate_RenderTargetable | TexCreate_UAV, ERHIAccess::CopyDest, CreateInfo);
 			InputFrame->SetTexture((ID3D11Texture2D*)Texture->GetNativeResource(), [&, InputFrame](ID3D11Texture2D* NativeTexture) { BackBuffers.Remove(InputFrame); });
 			BackBuffers.Add(InputFrame, Texture);
@@ -594,7 +594,7 @@ TSharedPtr<AVEncoder::FVideoEncoderInputFrame> FGameplayMediaEncoder::ObtainInpu
 		else if (RHIType == ERHIInterfaceType::D3D12)
 		{
 			FRHIResourceCreateInfo CreateInfo(TEXT("VideoCapturerBackBuffer"));
-			FTexture2DRHIRef Texture = GDynamicRHI->RHICreateTexture2D(VideoConfig.Width, VideoConfig.Height, EPixelFormat::PF_B8G8R8A8, 1, 1,
+			FTexture2DRHIRef Texture = RHICreateTexture2D(VideoConfig.Width, VideoConfig.Height, EPixelFormat::PF_B8G8R8A8, 1, 1,
 			                                                           TexCreate_Shared | TexCreate_RenderTargetable | TexCreate_UAV, ERHIAccess::CopyDest, CreateInfo);
 			InputFrame->SetTexture((ID3D12Resource*)Texture->GetNativeResource(), [&, InputFrame](ID3D12Resource* NativeTexture) { BackBuffers.Remove(InputFrame); });
 			BackBuffers.Add(InputFrame, Texture);
diff --git a/Engine/Source/Runtime/NullDrv/Public/NullRHI.h b/Engine/Source/Runtime/NullDrv/Public/NullRHI.h
--- a/Engine/Source/Runtime/NullDrv/Public/NullRHI.h
+++ b/Engine/Source/Runtime/NullDrv/Public/NullRHI.h
@@ -187,23 +187,9 @@ public:
 
 	}
 
-	virtual uint64 RHICalcTexture2DPlatformSize(uint32 SizeX, uint32 SizeY, uint8 Format, uint32 NumMips, uint32 NumSamples, ETextureCreateFlags Flags, const FRHIResourceCreateInfo& CreateInfo, uint32& OutAlign) final override
-	{ 
-		OutAlign = 0; 
-		return 0; 
-	}
-
-
-	virtual uint64 RHICalcTexture3DPlatformSize(uint32 SizeX, uint32 SizeY, uint32 SizeZ, uint8 Format, uint32 NumMips, ETextureCreateFlags Flags, const FRHIResourceCreateInfo& CreateInfo, uint32& OutAlign) final override
-	{ 
-		OutAlign = 0; 
-		return 0; 
-	}
-
-	virtual uint64 RHICalcTextureCubePlatformSize(uint32 Size, uint8 Format, uint32 NumMips, ETextureCreateFlags Flags, const FRHIResourceCreateInfo& CreateInfo, uint32& OutAlign) final override
-	{ 
-		OutAlign = 0; 
-		return 0; 
+	virtual FRHICalcTextureSizeResult RHICalcTexturePlatformSize(FRHITextureDesc const& Desc, uint32 FirstMipIndex) override final
+	{
+		return {};
 	}
 
 	virtual void RHIGetTextureMemoryStats(FTextureMemoryStats& OutStats) final override
@@ -216,29 +202,42 @@ public:
 		return false; 
 	}
 
+	class FNullTexture : public FRHITexture
+	{
+	public:
+		FNullTexture(const FRHITextureCreateDesc& InDesc)
+			: FRHITexture(InDesc)
+		{}
+	};
 
-	virtual FTexture2DRHIRef RHICreateTexture2D(uint32 SizeX, uint32 SizeY, uint8 Format, uint32 NumMips, uint32 NumSamples, ETextureCreateFlags Flags, ERHIAccess InResourceState, FRHIResourceCreateInfo& CreateInfo) final override
-	{ 
-		return new FRHITexture2D(SizeX,SizeY,NumMips,NumSamples,(EPixelFormat)Format,Flags, CreateInfo.ClearValueBinding); 
+	virtual FTextureRHIRef RHICreateTexture(const FRHITextureCreateDesc& CreateDesc) final override
+	{
+		return new FNullTexture(CreateDesc);
 	}
 
-	virtual FTexture2DRHIRef RHIAsyncCreateTexture2D(uint32 SizeX, uint32 SizeY, uint8 Format, uint32 NumMips, ETextureCreateFlags Flags, ERHIAccess InResourceState, void** InitialMipData, uint32 NumInitialMips) final override
+	virtual FTextureRHIRef RHIAsyncCreateTexture2D(uint32 SizeX, uint32 SizeY, uint8 Format, uint32 NumMips, ETextureCreateFlags Flags, ERHIAccess InResourceState, void** InitialMipData, uint32 NumInitialMips) final override
 	{ 
-		return FTexture2DRHIRef(); 
+		const uint32 NumSamples = 1;
+		const uint32 ExtData = 0;
+
+		return this->RHICreateTexture(
+			FRHITextureCreateDesc::Create2D(
+				  TEXT("FNullDynamicRHI::RHIAsyncCreateTexture2D")
+				, { (int32)SizeX, (int32)SizeY}
+				, (EPixelFormat)Format
+				, FClearValueBinding::None
+				, Flags
+				, NumMips
+				, NumSamples
+				, ExtData
+				, InResourceState
+			)
+		);
 	}
 
 	virtual void RHICopySharedMips(FRHITexture2D* DestTexture2D, FRHITexture2D* SrcTexture2D) final override
 	{
 	}
-	virtual FTexture2DArrayRHIRef RHICreateTexture2DArray(uint32 SizeX, uint32 SizeY, uint32 SizeZ, uint8 Format, uint32 NumMips, uint32 NumSamples, ETextureCreateFlags Flags, ERHIAccess InResourceState, FRHIResourceCreateInfo& CreateInfo) final override
-	{ 
-		return new FRHITexture2DArray(SizeX,SizeY,SizeZ,NumMips,NumSamples,(EPixelFormat)Format,Flags, CreateInfo.ClearValueBinding); 
-	}
-
-	virtual FTexture3DRHIRef RHICreateTexture3D(uint32 SizeX, uint32 SizeY, uint32 SizeZ, uint8 Format, uint32 NumMips, ETextureCreateFlags Flags, ERHIAccess InResourceState, FRHIResourceCreateInfo& CreateInfo) final override
-	{ 
-		return new FRHITexture3D(SizeX, SizeY, SizeZ, NumMips, (EPixelFormat)Format, Flags, CreateInfo.ClearValueBinding);
-	}
 
 	virtual FShaderResourceViewRHIRef RHICreateShaderResourceView(FRHITexture* Texture, const FRHITextureSRVCreateInfo& CreateInfo) final override
 	{ 
@@ -255,7 +254,17 @@ public:
 	}
 	virtual FTexture2DRHIRef RHIAsyncReallocateTexture2D(FRHITexture2D* Texture2D, int32 NewMipCount, int32 NewSizeX, int32 NewSizeY, FThreadSafeCounter* RequestStatus) final override
 	{ 
-		return new FRHITexture2D(NewSizeX,NewSizeY,NewMipCount,1,Texture2D->GetFormat(),Texture2D->GetFlags(), Texture2D->GetClearBinding());
+		return this->RHICreateTexture(
+			FRHITextureCreateDesc::Create2D(
+				  TEXT("FNullDynamicRHI::RHIAsyncReallocateTexture2D")
+				, { NewSizeX, NewSizeY }
+				, Texture2D->GetFormat()
+				, Texture2D->GetClearBinding()
+				, Texture2D->GetFlags()
+				, NewMipCount
+				, Texture2D->GetNumSamples()
+			)
+		);
 	}
 	virtual ETextureReallocationStatus RHIFinalizeAsyncReallocateTexture2D(FRHITexture2D* Texture2D, bool bBlockUntilCompleted) final override
 	{ 
@@ -291,14 +300,7 @@ public:
 	{
 
 	}
-	virtual FTextureCubeRHIRef RHICreateTextureCube(uint32 Size, uint8 Format, uint32 NumMips, ETextureCreateFlags Flags, ERHIAccess InResourceState, FRHIResourceCreateInfo& CreateInfo) final override
-	{ 
-		return new FRHITextureCube(Size, NumMips, (EPixelFormat)Format, Flags, CreateInfo.ClearValueBinding);
-	}
-	virtual FTextureCubeRHIRef RHICreateTextureCubeArray(uint32 Size, uint32 ArraySize, uint8 Format, uint32 NumMips, ETextureCreateFlags Flags, ERHIAccess InResourceState, FRHIResourceCreateInfo& CreateInfo) final override
-	{ 
-		return new FRHITextureCube(Size, NumMips, (EPixelFormat)Format, Flags, CreateInfo.ClearValueBinding);
-	}
+
 	virtual void* RHILockTextureCubeFace(FRHITextureCube* Texture, uint32 FaceIndex, uint32 ArrayIndex, uint32 MipIndex, EResourceLockMode LockMode, uint32& DestStride, bool bLockWithinMiptail) final override
 	{ 
 		DestStride = 0; 
@@ -385,8 +387,16 @@ public:
 	}
 
 	virtual FTexture2DRHIRef RHIGetViewportBackBuffer(FRHIViewport* Viewport) final override
-	{ 
-		return new FRHITexture2D(1,1,1,1,PF_B8G8R8A8,TexCreate_RenderTargetable, FClearValueBinding()); 
+	{
+		return this->RHICreateTexture(
+			FRHITextureCreateDesc::Create2D(
+				  TEXT("FNullDynamicRHI::RHIGetViewportBackBuffer")
+				, { 1, 1 }
+				, PF_B8G8R8A8
+				, FClearValueBinding::None
+				, ETextureCreateFlags::RenderTargetable
+			)
+		);
 	}
 
 	virtual void RHIBeginFrame() final override
diff --git a/Engine/Source/Runtime/OpenGLDrv/Private/Android/AndroidEGL.cpp b/Engine/Source/Runtime/OpenGLDrv/Private/Android/AndroidEGL.cpp
--- a/Engine/Source/Runtime/OpenGLDrv/Private/Android/AndroidEGL.cpp
+++ b/Engine/Source/Runtime/OpenGLDrv/Private/Android/AndroidEGL.cpp
@@ -86,7 +86,6 @@ struct AndroidESPImpl
 	ANativeWindow* Window;
 	bool Initalized ;
 	EOpenGLCurrentContext CurrentContextType;
-	GLuint OnScreenColorRenderBuffer;
 	GLuint ResolveFrameBuffer;
 	FPlatformRect CachedWindowRect;
 
@@ -525,7 +524,6 @@ eglDisplay(EGL_NO_DISPLAY)
 	,DepthSize(0)
 	,Window(NULL)
 	,Initalized(false)
-	,OnScreenColorRenderBuffer(0)
 	,ResolveFrameBuffer(0)
 	,NativeVisualID(0)
 	,CurrentContextType(CONTEXT_Invalid)
@@ -614,11 +612,6 @@ void AndroidEGL::DestroyBackBuffer()
 		glDeleteFramebuffers(1, &PImplData->ResolveFrameBuffer);
 		PImplData->ResolveFrameBuffer = 0 ;
 	}
-	if(PImplData->OnScreenColorRenderBuffer)
-	{
-		glDeleteRenderbuffers(1, &(PImplData->OnScreenColorRenderBuffer));
-		PImplData->OnScreenColorRenderBuffer = 0;
-	}
 }
 
 void AndroidEGL::InitBackBuffer()
@@ -633,7 +626,6 @@ void AndroidEGL::InitBackBuffer()
 	{
 		PImplData->ResolveFrameBuffer = 0;
 	}
-	PImplData->OnScreenColorRenderBuffer = 0;
 	OnScreenWidth = PImplData->eglWidth;
 	OnScreenHeight = PImplData->eglHeight;
 
@@ -832,11 +824,6 @@ bool AndroidEGL::IsInitialized()
 	return PImplData->Initalized;
 }
 
-GLuint AndroidEGL::GetOnScreenColorRenderBuffer()
-{
-	return PImplData->OnScreenColorRenderBuffer;
-}
-
 GLuint AndroidEGL::GetResolveFrameBuffer()
 {
 	return PImplData->ResolveFrameBuffer;
diff --git a/Engine/Source/Runtime/OpenGLDrv/Private/Android/AndroidEGL.h b/Engine/Source/Runtime/OpenGLDrv/Private/Android/AndroidEGL.h
--- a/Engine/Source/Runtime/OpenGLDrv/Private/Android/AndroidEGL.h
+++ b/Engine/Source/Runtime/OpenGLDrv/Private/Android/AndroidEGL.h
@@ -95,7 +95,6 @@ public:
 	void AcquireCurrentRenderingContext();
 	void ReleaseContextOwnership();
 
-	GLuint GetOnScreenColorRenderBuffer();
 	GLuint GetResolveFrameBuffer();
 	bool IsCurrentContextValid();
 	EGLContext  GetCurrentContext(  );
diff --git a/Engine/Source/Runtime/OpenGLDrv/Private/Android/AndroidOpenGL.cpp b/Engine/Source/Runtime/OpenGLDrv/Private/Android/AndroidOpenGL.cpp
--- a/Engine/Source/Runtime/OpenGLDrv/Private/Android/AndroidOpenGL.cpp
+++ b/Engine/Source/Runtime/OpenGLDrv/Private/Android/AndroidOpenGL.cpp
@@ -444,17 +444,22 @@ void PlatformDestroyOpenGLContext(FPlatformOpenGLDevice* Device, FPlatformOpenGL
 {
 }
 
-FRHITexture* PlatformCreateBuiltinBackBuffer(FOpenGLDynamicRHI* OpenGLRHI, uint32 SizeX, uint32 SizeY)
+FOpenGLTexture* PlatformCreateBuiltinBackBuffer(FOpenGLDynamicRHI* OpenGLRHI, uint32 SizeX, uint32 SizeY)
 {
+	check(IsInRenderingThread());
 	// Create the built-in back buffer if we disable backbuffer sampling.
 	// Otherwise return null and we will create an off-screen surface afterward.
 	if (!FPlatformMisc::SupportsBackbufferSampling())
 	{
-		ETextureCreateFlags Flags = TexCreate_RenderTargetable;
-		FOpenGLTexture2D* Texture2D = new FOpenGLTexture2D(OpenGLRHI, AndroidEGL::GetInstance()->GetOnScreenColorRenderBuffer(), GL_RENDERBUFFER, GL_COLOR_ATTACHMENT0, SizeX, SizeY, 0, 1, 1, 1, 1, PF_B8G8R8A8, false, false, Flags, FClearValueBinding::Transparent);
-		OpenGLTextureAllocated(Texture2D, Flags);
-
-		return Texture2D;
+		FOpenGLTextureCreateDesc CreateDesc = FOpenGLTextureCreateDesc::Create2D(
+			TEXT("PlatformCreateBuiltinBackBuffer"),
+			{ (int32)SizeX, (int32)SizeY },
+			PF_B8G8R8A8,
+			FClearValueBinding::Transparent,
+			TexCreate_RenderTargetable | TexCreate_Presentable
+		);
+
+		return new FOpenGLTexture(CreateDesc);
 	}
 	else
 	{
@@ -465,7 +470,8 @@ FRHITexture* PlatformCreateBuiltinBackBuffer(FOpenGLDynamicRHI* OpenGLRHI, uint3
 void PlatformResizeGLContext( FPlatformOpenGLDevice* Device, FPlatformOpenGLContext* Context, uint32 SizeX, uint32 SizeY, bool bFullscreen, bool bWasFullscreen, GLenum BackBufferTarget, GLuint BackBufferResource)
 {
 	check(Context);
-	
+	VERIFY_GL_SCOPE();
+
 	Context->BackBufferResource = BackBufferResource;
 	Context->BackBufferTarget = BackBufferTarget;
 
@@ -543,13 +549,6 @@ void FPlatformOpenGLDevice::SetupCurrentContext()
 
 void PlatformLabelObjects()
 {
-	// @todo: Check that there is a valid id (non-zero) as LabelObject will fail otherwise
-	GLuint RenderBuffer = AndroidEGL::GetInstance()->GetOnScreenColorRenderBuffer();
-	if (RenderBuffer != 0)
-	{
-		FOpenGL::LabelObject(GL_RENDERBUFFER, RenderBuffer, "OnScreenColorRB");
-	}
-
 	GLuint FrameBuffer = AndroidEGL::GetInstance()->GetResolveFrameBuffer();
 	if (FrameBuffer != 0)
 	{
diff --git a/Engine/Source/Runtime/OpenGLDrv/Private/Linux/OpenGLLinux.cpp b/Engine/Source/Runtime/OpenGLDrv/Private/Linux/OpenGLLinux.cpp
--- a/Engine/Source/Runtime/OpenGLDrv/Private/Linux/OpenGLLinux.cpp
+++ b/Engine/Source/Runtime/OpenGLDrv/Private/Linux/OpenGLLinux.cpp
@@ -985,7 +985,7 @@ bool PlatformContextIsCurrent( uint64 QueryContext )
 	return (uint64)Linux_GetCurrentContext() == QueryContext;
 }
 
-FRHITexture* PlatformCreateBuiltinBackBuffer( FOpenGLDynamicRHI* OpenGLRHI, uint32 SizeX, uint32 SizeY )
+FOpenGLTexture* PlatformCreateBuiltinBackBuffer( FOpenGLDynamicRHI* OpenGLRHI, uint32 SizeX, uint32 SizeY )
 {
 	return nullptr;
 }
diff --git a/Engine/Source/Runtime/OpenGLDrv/Private/OpenGLCommands.cpp b/Engine/Source/Runtime/OpenGLDrv/Private/OpenGLCommands.cpp
--- a/Engine/Source/Runtime/OpenGLDrv/Private/OpenGLCommands.cpp
+++ b/Engine/Source/Runtime/OpenGLDrv/Private/OpenGLCommands.cpp
@@ -449,7 +449,7 @@ void FOpenGLDynamicRHI::RHISetUAVParameter(FRHIComputeShader* ComputeShaderRHI,u
 	check(0);
 }
 
-void FOpenGLDynamicRHI::InternalSetShaderTexture(FOpenGLTextureBase* Texture, FOpenGLShaderResourceView* SRV, GLint TextureIndex, GLenum Target, GLuint Resource, int NumMips, int LimitMip)
+void FOpenGLDynamicRHI::InternalSetShaderTexture(FOpenGLTexture* Texture, FOpenGLShaderResourceView* SRV, GLint TextureIndex, GLenum Target, GLuint Resource, int NumMips, int LimitMip)
 {
 	auto& PendingTextureState = PendingState.Textures[TextureIndex];
 	PendingTextureState.Texture = Texture;
@@ -896,19 +896,19 @@ void FOpenGLDynamicRHI::UpdateSRV(FOpenGLShaderResourceView* SRV)
 	check(SRV);
 	// For Depth/Stencil textures whose Stencil component we wish to sample we must blit the stencil component out to an intermediate texture when we 'Store' the texture.
 #if PLATFORM_DESKTOP
-	if (FOpenGL::GetFeatureLevel() >= ERHIFeatureLevel::SM5 && IsValidRef(SRV->Texture2D))
+	if (FOpenGL::GetFeatureLevel() >= ERHIFeatureLevel::SM5 && IsValidRef(SRV->Texture))
 	{
-		FOpenGLTexture2D* Texture2D = ResourceCast(SRV->Texture2D.GetReference());
+		FOpenGLTexture* Texture = SRV->Texture;
 		
 		uint32 ArrayIndices = 0;
 		uint32 MipmapLevels = 0;
 		
-		GLuint SourceFBO = GetOpenGLFramebuffer(0, nullptr, &ArrayIndices, &MipmapLevels, (FOpenGLTextureBase*)Texture2D);
+		GLuint SourceFBO = GetOpenGLFramebuffer(0, nullptr, &ArrayIndices, &MipmapLevels, Texture);
 		
 		glBindFramebuffer(GL_FRAMEBUFFER, SourceFBO);
 		
-		uint32 SizeX = Texture2D->GetSizeX();
-		uint32 SizeY = Texture2D->GetSizeY();
+		uint32 SizeX = Texture->GetSizeX();
+		uint32 SizeY = Texture->GetSizeY();
 		
 		uint32 MipBytes = SizeX * SizeY;
 		TRefCountPtr<FOpenGLPixelBuffer> PixelBuffer = new FOpenGLPixelBuffer(GL_PIXEL_UNPACK_BUFFER, 0, MipBytes, BUF_Dynamic, nullptr);
@@ -1013,7 +1013,7 @@ void FOpenGLDynamicRHI::RHISetShaderResourceViewParameter(FRHIComputeShader* Com
 void FOpenGLDynamicRHI::RHISetShaderTexture(FRHIGraphicsShader* ShaderRHI,uint32 TextureIndex, FRHITexture* NewTextureRHI)
 {
 	VERIFY_GL_SCOPE();
-	FOpenGLTextureBase* NewTexture = GetOpenGLTextureFromRHITexture(NewTextureRHI);
+	FOpenGLTexture* NewTexture = GetOpenGLTextureFromRHITexture(NewTextureRHI);
 
 	GLint Index = 0;
 	GLint MaxUnits = 0;
@@ -1045,7 +1045,7 @@ void FOpenGLDynamicRHI::RHISetShaderSampler(FRHIGraphicsShader* ShaderRHI,uint32
 void FOpenGLDynamicRHI::RHISetShaderTexture(FRHIComputeShader* ComputeShaderRHI,uint32 TextureIndex, FRHITexture* NewTextureRHI)
 {
 	VERIFY_GL_SCOPE();
-	FOpenGLTextureBase* NewTexture = GetOpenGLTextureFromRHITexture(NewTextureRHI);
+	FOpenGLTexture* NewTexture = GetOpenGLTextureFromRHITexture(NewTextureRHI);
 	ensureMsgf((int32)TextureIndex < FOpenGL::GetMaxComputeTextureImageUnits(), TEXT("Using more compute texture units (%d) than allowed (%d)!"), TextureIndex, FOpenGL::GetMaxComputeTextureImageUnits());
 	if (NewTexture)
 	{
@@ -1281,8 +1281,8 @@ void FOpenGLDynamicRHI::SetPendingBlendStateForActiveRenderTargets( FOpenGLConte
 		}
 		else if (RenderTargetIndex == 0)
 		{
-			FOpenGLTexture2D* RenderTarget2D = (FOpenGLTexture2D*)PendingState.RenderTargets[RenderTargetIndex];
-			bMSAAEnabled = RenderTarget2D->GetNumSamples() > 1 || RenderTarget2D->GetNumSamplesTileMem() > 1;
+			FOpenGLTexture* RenderTarget2D = PendingState.RenderTargets[RenderTargetIndex];
+			bMSAAEnabled = RenderTarget2D->GetNumSamplesRendered() > 1;
 		}
 
 		const FOpenGLBlendStateData::FRenderTarget& RenderTargetBlendState = PendingState.BlendState.RenderTargets[RenderTargetIndex];
@@ -1519,12 +1519,12 @@ void FOpenGLDynamicRHI::SetRenderTargets(
 		}
 	}
 
-	FOpenGLTextureBase* NewDepthStencilRT = GetOpenGLTextureFromRHITexture(NewDepthStencilTargetRHI ? NewDepthStencilTargetRHI->Texture : nullptr);
+	FOpenGLTexture* NewDepthStencilRT = GetOpenGLTextureFromRHITexture(NewDepthStencilTargetRHI ? NewDepthStencilTargetRHI->Texture : nullptr);
 	
 	PendingState.DepthStencil = NewDepthStencilRT;
 	PendingState.StencilStoreAction = NewDepthStencilTargetRHI ? NewDepthStencilTargetRHI->GetStencilStoreAction() : ERenderTargetStoreAction::ENoAction;
-	PendingState.DepthTargetWidth = NewDepthStencilTargetRHI ? GetOpenGLTextureSizeXFromRHITexture(NewDepthStencilTargetRHI->Texture) : 0u;
-	PendingState.DepthTargetHeight = NewDepthStencilTargetRHI ? GetOpenGLTextureSizeYFromRHITexture(NewDepthStencilTargetRHI->Texture) : 0u;
+	PendingState.DepthTargetWidth   = NewDepthStencilRT ? NewDepthStencilRT->GetDesc().Extent.X : 0u;
+	PendingState.DepthTargetHeight  = NewDepthStencilRT ? NewDepthStencilRT->GetDesc().Extent.Y : 0u;
 	
 	if (PendingState.FirstNonzeroRenderTarget == -1 && !PendingState.DepthStencil)
 	{
@@ -1544,63 +1544,22 @@ void FOpenGLDynamicRHI::SetRenderTargets(
 		PendingState.Viewport.Min.X = 0;
 		PendingState.Viewport.Min.Y = 0;
 
-		uint32 Width = 0;
-		uint32 Height = 0;
+		const FRHITextureDesc& Desc = NewRenderTargetsRHI[PendingState.FirstNonzeroRenderTarget].Texture->GetDesc();
 
-		FOpenGLTexture2D* NewRenderTarget2D = (FOpenGLTexture2D*)NewRenderTargetsRHI[PendingState.FirstNonzeroRenderTarget].Texture->GetTexture2D();
-		if(NewRenderTarget2D)
-		{
-			Width = NewRenderTarget2D->GetSizeX();
-			Height = NewRenderTarget2D->GetSizeY();
-		}
-		else
-		{
-			FOpenGLTextureCube* NewRenderTargetCube = (FOpenGLTextureCube*)NewRenderTargetsRHI[PendingState.FirstNonzeroRenderTarget].Texture->GetTextureCube();
-			if(NewRenderTargetCube)
-			{
-				Width = NewRenderTargetCube->GetSize();
-				Height = NewRenderTargetCube->GetSize();
-			}
-			else
-			{
-				FOpenGLTexture3D* NewRenderTarget3D = (FOpenGLTexture3D*)NewRenderTargetsRHI[PendingState.FirstNonzeroRenderTarget].Texture->GetTexture3D();
-				if(NewRenderTarget3D)
-				{
-					Width = NewRenderTarget3D->GetSizeX();
-					Height = NewRenderTarget3D->GetSizeY();
-				}
-				else
-				{
-					FOpenGLTexture2DArray* NewRenderTarget2DArray = (FOpenGLTexture2DArray*)NewRenderTargetsRHI[PendingState.FirstNonzeroRenderTarget].Texture->GetTexture2DArray();
-					if(NewRenderTarget2DArray)
-					{
-						Width = NewRenderTarget2DArray->GetSizeX();
-						Height = NewRenderTarget2DArray->GetSizeY();
-					}
-					else
-					{
-						check(0);
-					}
-				}
-			}
-		}
-
-		{
-			uint32 MipIndex = NewRenderTargetsRHI[PendingState.FirstNonzeroRenderTarget].MipIndex;
-			Width = FMath::Max<uint32>(1,(Width >> MipIndex));
-			Height = FMath::Max<uint32>(1,(Height >> MipIndex));
-		}
+		uint32 MipIndex = NewRenderTargetsRHI[PendingState.FirstNonzeroRenderTarget].MipIndex;
+		uint32 Width    = FMath::Max<uint32>(1, Desc.Extent.X >> MipIndex);
+		uint32 Height   = FMath::Max<uint32>(1, Desc.Extent.Y >> MipIndex);
 
 		PendingState.Viewport.Max.X = PendingState.RenderTargetWidth = Width;
 		PendingState.Viewport.Max.Y = PendingState.RenderTargetHeight = Height;
 	}
-	else if( NewDepthStencilTargetRHI )
+	else if (NewDepthStencilTargetRHI)
 	{
 		// Set viewport size to new depth target size.
 		PendingState.Viewport.Min.X = 0;
 		PendingState.Viewport.Min.Y = 0;
-		PendingState.Viewport.Max.X = GetOpenGLTextureSizeXFromRHITexture(NewDepthStencilTargetRHI->Texture);
-		PendingState.Viewport.Max.Y = GetOpenGLTextureSizeYFromRHITexture(NewDepthStencilTargetRHI->Texture);
+		PendingState.Viewport.Max.X = NewDepthStencilTargetRHI->Texture->GetDesc().Extent.X;
+		PendingState.Viewport.Max.Y = NewDepthStencilTargetRHI->Texture->GetDesc().Extent.Y;
 	}
 }
 
@@ -2145,14 +2104,14 @@ template <> FORCEINLINE uint32 GetNumUAVUnits<SF_Vertex>()	{ return FOpenGL::Get
 template <EShaderFrequency Frequency>
 void SetResource(FOpenGLDynamicRHI* RESTRICT OpenGLRHI, uint32 BindIndex, FRHITexture* RESTRICT TextureRHI)
 {
-	FOpenGLTextureBase* Texture = GetOpenGLTextureFromRHITexture(TextureRHI);
+	FOpenGLTexture* Texture = GetOpenGLTextureFromRHITexture(TextureRHI);
 #if !UE_BUILD_TEST && !UE_BUILD_SHIPPING
 	ensureMsgf(BindIndex < GetNumTextureUnits<Frequency>(), TEXT("Using more %s texture units (%d) than allowed (%d) on a shader unit!"), GetShaderFrequencyString(Frequency, false), BindIndex, GetNumTextureUnits<Frequency>());
 #endif
 	if (Texture)
 	{
 		TextureRHI->SetLastRenderTime(FApp::GetCurrentTime());
-		OpenGLRHI->InternalSetShaderTexture(Texture, nullptr, GetFirstTextureUnit<Frequency>() + BindIndex, Texture->Target, Texture->GetResource(), Texture->NumMips, -1);
+		OpenGLRHI->InternalSetShaderTexture(Texture, nullptr, GetFirstTextureUnit<Frequency>() + BindIndex, Texture->Target, Texture->GetResource(), Texture->GetNumMips(), -1);
 	}
 	else
 	{
diff --git a/Engine/Source/Runtime/OpenGLDrv/Private/OpenGLDevice.cpp b/Engine/Source/Runtime/OpenGLDrv/Private/OpenGLDevice.cpp
--- a/Engine/Source/Runtime/OpenGLDrv/Private/OpenGLDevice.cpp
+++ b/Engine/Source/Runtime/OpenGLDrv/Private/OpenGLDevice.cpp
@@ -1042,6 +1042,9 @@ FOpenGLDynamicRHI::FOpenGLDynamicRHI()
 ,	PlatformDevice(NULL)
 ,	GPUProfilingData(this)
 {
+	check(Singleton == nullptr);
+	Singleton = this;
+
 	// This should be called once at the start
 	check( IsInGameThread() );
 	check( !GIsThreadedRendering );
diff --git a/Engine/Source/Runtime/OpenGLDrv/Private/OpenGLDrv.cpp b/Engine/Source/Runtime/OpenGLDrv/Private/OpenGLDrv.cpp
--- a/Engine/Source/Runtime/OpenGLDrv/Private/OpenGLDrv.cpp
+++ b/Engine/Source/Runtime/OpenGLDrv/Private/OpenGLDrv.cpp
@@ -46,7 +46,7 @@ bool FOpenGLDynamicRHI::RHISupportsFramebufferSRGBEnable() const
 
 GLuint FOpenGLDynamicRHI::RHIGetResource(FRHITexture* InTexture) const
 {
-	FOpenGLTextureBase* GLTexture = GetOpenGLTextureFromRHITexture(InTexture);
+	FOpenGLTexture* GLTexture = GetOpenGLTextureFromRHITexture(InTexture);
 	return GLTexture->GetResource();
 }
 
diff --git a/Engine/Source/Runtime/OpenGLDrv/Private/OpenGLDrvPrivate.h b/Engine/Source/Runtime/OpenGLDrv/Private/OpenGLDrvPrivate.h
--- a/Engine/Source/Runtime/OpenGLDrv/Private/OpenGLDrvPrivate.h
+++ b/Engine/Source/Runtime/OpenGLDrv/Private/OpenGLDrvPrivate.h
@@ -245,7 +245,7 @@ void PlatformSharedContextSetup(FPlatformOpenGLDevice* Device);
 void PlatformNULLContextSetup();
 
 // Creates a platform-specific back buffer
-FRHITexture* PlatformCreateBuiltinBackBuffer(FOpenGLDynamicRHI* OpenGLRHI, uint32 SizeX, uint32 SizeY);
+class FOpenGLTexture* PlatformCreateBuiltinBackBuffer(FOpenGLDynamicRHI* OpenGLRHI, uint32 SizeX, uint32 SizeY);
 
 /**
  * Main function for transferring data to on-screen buffers.
@@ -290,26 +290,14 @@ void* PlatformGetWindow(FPlatformOpenGLContext* Context, void** AddParam);
 struct FOpenGLTextureFormat
 {
 	// [0]: without sRGB, [1]: with sRGB
-	GLenum InternalFormat[2];
-	// [0]: without sRGB, [1]: with sRGB
-	GLenum SizedInternalFormat[2];
-	GLenum Format;
-	GLenum Type;
-	bool bCompressed;
+	GLenum InternalFormat[2] = { GL_NONE, GL_NONE };
+	GLenum Format = GL_NONE;
+	GLenum Type = GL_NONE;
+	bool bCompressed = false;
 	// Reorder to B and R elements using texture swizzle
-	bool bBGRA;
+	bool bBGRA = false;
 
-	FOpenGLTextureFormat()
-	{
-		InternalFormat[0] = GL_NONE;
-		InternalFormat[1] = GL_NONE;
-		SizedInternalFormat[0] = GL_NONE;
-		SizedInternalFormat[1] = GL_NONE;
-		Format = GL_NONE;
-		Type = GL_NONE;
-		bCompressed = false;
-		bBGRA = false;
-	}
+	FOpenGLTextureFormat() = default;
 
 	FOpenGLTextureFormat(
 		GLenum InInternalFormat, GLenum InInternalFormatSRGB, GLenum InFormat,
@@ -317,30 +305,11 @@ struct FOpenGLTextureFormat
 	{
 		InternalFormat[0] = InInternalFormat;
 		InternalFormat[1] = InInternalFormatSRGB;
-		SizedInternalFormat[0] = InInternalFormat;
-		SizedInternalFormat[1] = InInternalFormatSRGB;
 		Format = InFormat;
 		Type = InType;
 		bCompressed = bInCompressed;
 		bBGRA = bInBGRA;
 	}
-
-	FOpenGLTextureFormat(
-		GLenum InInternalFormat, GLenum InInternalFormatSRGB, 
-		GLenum InSizedInternalFormat, GLenum InSizedInternalFormatSRGB,
-		GLenum InFormat,
-		GLenum InType, bool bInCompressed, bool bInBGRA)
-	{
-		InternalFormat[0] = InInternalFormat;
-		InternalFormat[1] = InInternalFormatSRGB;
-		SizedInternalFormat[0] = InSizedInternalFormat;
-		SizedInternalFormat[1] = InSizedInternalFormatSRGB;
-		Format = InFormat;
-		Type = InType;
-		bCompressed = bInCompressed;
-		bBGRA = bInBGRA;
-	}
-
 };
 
 extern FOpenGLTextureFormat OPENGLDRV_API GOpenGLTextureFormats[PF_MAX];
diff --git a/Engine/Source/Runtime/OpenGLDrv/Private/OpenGLQuery.cpp b/Engine/Source/Runtime/OpenGLDrv/Private/OpenGLQuery.cpp
--- a/Engine/Source/Runtime/OpenGLDrv/Private/OpenGLQuery.cpp
+++ b/Engine/Source/Runtime/OpenGLDrv/Private/OpenGLQuery.cpp
@@ -711,7 +711,6 @@ void FOpenGLEventQuery::IssueEvent()
 	FOpenGL::Flush();
 
 	checkSlow(FOpenGL::IsSync(Sync));
-
 }
 
 void FOpenGLEventQuery::WaitForCompletion()
@@ -720,10 +719,8 @@ void FOpenGLEventQuery::WaitForCompletion()
 
 	QUICK_SCOPE_CYCLE_COUNTER(STAT_FOpenGLEventQuery_WaitForCompletion);
 
-
 	checkSlow(FOpenGL::IsSync(Sync));
 
-
 	// Wait up to 1/2 second for sync execution
 	FOpenGL::EFenceResult Status = FOpenGL::ClientWaitSync( Sync, 0, 500*1000*1000);
 
@@ -748,27 +745,20 @@ void FOpenGLEventQuery::WaitForCompletion()
 	}	
 }
 
-void FOpenGLEventQuery::InitDynamicRHI()
+FOpenGLEventQuery::FOpenGLEventQuery()
 {
-	FRHICommandListImmediate& RHICmdList = FRHICommandListExecutor::GetImmediateCommandList();
-
-	RHITHREAD_GLCOMMAND_PROLOGUE();
 	VERIFY_GL_SCOPE();
+
 	// Initialize the query by issuing an initial event.
 	IssueEvent();
 
 	check(FOpenGL::IsSync(Sync));
-	RHITHREAD_GLCOMMAND_EPILOGUE();
 }
 
-void FOpenGLEventQuery::ReleaseDynamicRHI()
+FOpenGLEventQuery::~FOpenGLEventQuery()
 {
-	FRHICommandListImmediate& RHICmdList = FRHICommandListExecutor::GetImmediateCommandList();
-
-	RHITHREAD_GLCOMMAND_PROLOGUE();
 	VERIFY_GL_SCOPE();
 	FOpenGL::DeleteSync(Sync);
-	RHITHREAD_GLCOMMAND_EPILOGUE();
 }
 
 /*=============================================================================
