Fixing bad state handling in RHICopySharedMips

diff --git a/Engine/Source/Runtime/D3D12RHI/Private/D3D12Resources.cpp b/Engine/Source/Runtime/D3D12RHI/Private/D3D12Resources.cpp
--- a/Engine/Source/Runtime/D3D12RHI/Private/D3D12Resources.cpp
+++ b/Engine/Source/Runtime/D3D12RHI/Private/D3D12Resources.cpp
@@ -1230,6 +1230,14 @@ void FD3D12ResourceLocation::UnlockPoolData()
 	}
 }
 
+bool FD3D12ResourceLocation::IsPoolDataLocked()
+{
+	if (AllocatorType == AT_Pool)
+	{
+		return GetPoolAllocatorPrivateData().PoolData.IsLocked();
+	}
+	return false;
+}
 
 /////////////////////////////////////////////////////////////////////
 //	FD3D12 Resource Barrier Batcher
diff --git a/Engine/Source/Runtime/D3D12RHI/Private/D3D12Texture.cpp b/Engine/Source/Runtime/D3D12RHI/Private/D3D12Texture.cpp
--- a/Engine/Source/Runtime/D3D12RHI/Private/D3D12Texture.cpp
+++ b/Engine/Source/Runtime/D3D12RHI/Private/D3D12Texture.cpp
@@ -1222,8 +1222,11 @@ void FD3D12DynamicRHI::RHICopySharedMips(FRHITexture2D* DestTexture2DRHI, FRHITe
 			}
 		}
 
-		// unlock the pool allocated resource because all data has been written
-		DestTexture2D->ResourceLocation.UnlockPoolData();
+		if (DestTexture2D->ResourceLocation.IsPoolDataLocked())
+		{
+			// unlock the pool allocated resource because all data has been written
+			DestTexture2D->ResourceLocation.UnlockPoolData();
+		}
 
 		Device->GetDefaultCommandContext().ConditionalFlushCommandList();
 
diff --git a/Engine/Source/Runtime/D3D12RHI/Public/D3D12Resources.h b/Engine/Source/Runtime/D3D12RHI/Public/D3D12Resources.h
--- a/Engine/Source/Runtime/D3D12RHI/Public/D3D12Resources.h
+++ b/Engine/Source/Runtime/D3D12RHI/Public/D3D12Resources.h
@@ -596,6 +596,7 @@ public:
 	// Pool allocation specific functions
 	bool OnAllocationMoved(FRHIPoolAllocationData* InNewData);
 	void UnlockPoolData();
+	bool IsPoolDataLocked();
 
 	const inline bool IsValid() const { return Type != ResourceLocationType::eUndefined; }
 
