Unlock the pool data for async created texture after the copy queue uploads are done - texture can now be picked up for defrag if needed

diff --git a/Engine/Source/Runtime/D3D12RHI/Private/D3D12Resources.cpp b/Engine/Source/Runtime/D3D12RHI/Private/D3D12Resources.cpp
--- a/Engine/Source/Runtime/D3D12RHI/Private/D3D12Resources.cpp
+++ b/Engine/Source/Runtime/D3D12RHI/Private/D3D12Resources.cpp
@@ -1230,15 +1230,6 @@ void FD3D12ResourceLocation::UnlockPoolData()
 	}
 }
 
-bool FD3D12ResourceLocation::IsPoolDataLocked()
-{
-	if (AllocatorType == AT_Pool)
-	{
-		return GetPoolAllocatorPrivateData().PoolData.IsLocked();
-	}
-	return false;
-}
-
 /////////////////////////////////////////////////////////////////////
 //	FD3D12 Resource Barrier Batcher
 /////////////////////////////////////////////////////////////////////
diff --git a/Engine/Source/Runtime/D3D12RHI/Private/D3D12Texture.cpp b/Engine/Source/Runtime/D3D12RHI/Private/D3D12Texture.cpp
--- a/Engine/Source/Runtime/D3D12RHI/Private/D3D12Texture.cpp
+++ b/Engine/Source/Runtime/D3D12RHI/Private/D3D12Texture.cpp
@@ -1090,9 +1090,8 @@ FTextureRHIRef FD3D12DynamicRHI::RHIAsyncCreateTexture2D(uint32 SizeX, uint32 Si
 		// on the GPU before returning here, this protection is safe, even if we end up straddling frame boundaries.
 		TempResourceLocation.GetResource()->AddRef();
 
-		for (FD3D12Texture& CurrentTextureBase : *TextureOut)
+		for (FD3D12Texture& CurrentTexture : *TextureOut)
 		{
-			FD3D12Texture& CurrentTexture = static_cast<FD3D12Texture&>(CurrentTextureBase);
 			FD3D12Device* Device = CurrentTexture.GetParentDevice();
 			FD3D12Resource* Resource = CurrentTexture.GetResource();
 
@@ -1148,6 +1147,9 @@ FTextureRHIRef FD3D12DynamicRHI::RHIAsyncCreateTexture2D(uint32 SizeX, uint32 Si
 			});
 
 			CommandAllocatorManager.ReleaseCommandAllocator(CurrentCommandAllocator);
+
+			// Blocking update is done
+			CurrentTexture.ResourceLocation.UnlockPoolData();
 		}
 
 		FD3D12TextureStats::D3D12TextureAllocated(*TextureOut);
@@ -1222,12 +1224,6 @@ void FD3D12DynamicRHI::RHICopySharedMips(FRHITexture2D* DestTexture2DRHI, FRHITe
 			}
 		}
 
-		if (DestTexture2D->ResourceLocation.IsPoolDataLocked())
-		{
-			// unlock the pool allocated resource because all data has been written
-			DestTexture2D->ResourceLocation.UnlockPoolData();
-		}
-
 		Device->GetDefaultCommandContext().ConditionalFlushCommandList();
 
 		DEBUG_EXECUTE_COMMAND_CONTEXT(Device->GetDefaultCommandContext());
diff --git a/Engine/Source/Runtime/D3D12RHI/Public/D3D12Resources.h b/Engine/Source/Runtime/D3D12RHI/Public/D3D12Resources.h
--- a/Engine/Source/Runtime/D3D12RHI/Public/D3D12Resources.h
+++ b/Engine/Source/Runtime/D3D12RHI/Public/D3D12Resources.h
@@ -596,7 +596,6 @@ public:
 	// Pool allocation specific functions
 	bool OnAllocationMoved(FRHIPoolAllocationData* InNewData);
 	void UnlockPoolData();
-	bool IsPoolDataLocked();
 
 	const inline bool IsValid() const { return Type != ResourceLocationType::eUndefined; }
 
