Add missing RGB565 and RGBA5551 formats to byte size calculation

diff --git a/Engine/Source/Runtime/D3D12RHI/Private/D3D12RenderTarget.cpp b/Engine/Source/Runtime/D3D12RHI/Private/D3D12RenderTarget.cpp
--- a/Engine/Source/Runtime/D3D12RHI/Private/D3D12RenderTarget.cpp
+++ b/Engine/Source/Runtime/D3D12RHI/Private/D3D12RenderTarget.cpp
@@ -495,6 +495,8 @@ static uint32 ComputeBytesPerPixel(DXGI_FORMAT Format)
 	case DXGI_FORMAT_R16_UINT:
 	case DXGI_FORMAT_R16_SNORM:
 	case DXGI_FORMAT_R16_SINT:
+	case DXGI_FORMAT_B5G6R5_UNORM:
+	case DXGI_FORMAT_B5G5R5A1_UNORM:
 		BytesPerPixel = 2;
 		break;
 	case DXGI_FORMAT_B8G8R8A8_TYPELESS:
