Cleaning up dxgi includes in D3D12RHI

diff --git a/Engine/Source/Runtime/D3D12RHI/Private/D3D12Residency.h b/Engine/Source/Runtime/D3D12RHI/Private/D3D12Residency.h
--- a/Engine/Source/Runtime/D3D12RHI/Private/D3D12Residency.h
+++ b/Engine/Source/Runtime/D3D12RHI/Private/D3D12Residency.h
@@ -26,7 +26,6 @@ namespace D3DX12Residency
 #include "Windows/AllowWindowsPlatformAtomics.h"
 
 THIRD_PARTY_INCLUDES_START
-#include "dxgi1_6.h"
 #pragma warning(push)
 #pragma warning(disable: 6031)
 	#include <D3DX12Residency.h>
diff --git a/Engine/Source/Runtime/D3D12RHI/Public/HoloLens/D3D12ThirdParty.h b/Engine/Source/Runtime/D3D12RHI/Public/HoloLens/D3D12ThirdParty.h
--- a/Engine/Source/Runtime/D3D12RHI/Public/HoloLens/D3D12ThirdParty.h
+++ b/Engine/Source/Runtime/D3D12RHI/Public/HoloLens/D3D12ThirdParty.h
@@ -16,7 +16,7 @@
 THIRD_PARTY_INCLUDES_START
 	#include <d3d12.h>
 	#include <d3dx12.h>
-	#include <dxgi.h>
+	#include <dxgi1_6.h>
 	#include <dxgidebug.h>
 THIRD_PARTY_INCLUDES_END
 #include "Microsoft/HideMicrosoftPlatformTypes.h"
diff --git a/Engine/Source/Runtime/D3D12RHI/Public/Windows/D3D12ThirdParty.h b/Engine/Source/Runtime/D3D12RHI/Public/Windows/D3D12ThirdParty.h
--- a/Engine/Source/Runtime/D3D12RHI/Public/Windows/D3D12ThirdParty.h
+++ b/Engine/Source/Runtime/D3D12RHI/Public/Windows/D3D12ThirdParty.h
@@ -16,7 +16,7 @@
 THIRD_PARTY_INCLUDES_START
 	#include <d3d12.h>
 	#include <d3dx12.h>
-	#include <dxgi.h>
+	#include <dxgi1_6.h>
 	#include <dxgidebug.h>
 THIRD_PARTY_INCLUDES_END
 #include "Microsoft/HideMicrosoftPlatformTypes.h"
