Comments out ensure() when using ProfileGPU

diff --git a/Engine/Source/Runtime/D3D12RHI/Private/D3D12Stats.cpp b/Engine/Source/Runtime/D3D12RHI/Private/D3D12Stats.cpp
--- a/Engine/Source/Runtime/D3D12RHI/Private/D3D12Stats.cpp
+++ b/Engine/Source/Runtime/D3D12RHI/Private/D3D12Stats.cpp
@@ -514,7 +514,7 @@ void D3D12RHI::FD3DGPUProfiler::DoPostProfileGPUWork()
 		{
 			const FResolvedCmdListExecTime& Prev = CmdListExecTimes[Idx - 1];
 			const FResolvedCmdListExecTime& Cur = CmdListExecTimes[Idx];
-			ensure(Cur.StartTimestamp >= Prev.EndTimestamp);
+			//ensure(Cur.StartTimestamp >= Prev.EndTimestamp);
 			CmdListStartTimestamps.Add(Cur.StartTimestamp);
 			CmdListEndTimestamps.Add(Cur.EndTimestamp);
 			const uint64 Bubble = Cur.StartTimestamp >= Prev.EndTimestamp ? Cur.StartTimestamp - Prev.EndTimestamp : 0;
