- State tracking fixes when using external transitions in d3d12 RHI
- Restore correct transient texture resource creation state (otherwise discard would happen on incorrect state)

diff --git a/Engine/Source/Runtime/D3D12RHI/Private/D3D12Commands.cpp b/Engine/Source/Runtime/D3D12RHI/Private/D3D12Commands.cpp
--- a/Engine/Source/Runtime/D3D12RHI/Private/D3D12Commands.cpp
+++ b/Engine/Source/Runtime/D3D12RHI/Private/D3D12Commands.cpp
@@ -160,7 +160,8 @@ void FD3D12CommandContext::RHIDispatchIndirectComputeShader(FRHIBuffer* Argument
 
 	// Indirect args buffer can be a previously pending UAV, which becomes PS\Non-PS read. ApplyState will flush pending transitions, so enqueue the indirect
 	// arg transition and flush here.
-	FD3D12DynamicRHI::TransitionResource(CommandListHandle, Location.GetResource(), D3D12_RESOURCE_STATE_TBD, D3D12_RESOURCE_STATE_INDIRECT_ARGUMENT | D3D12_RESOURCE_STATE_NON_PIXEL_SHADER_RESOURCE, D3D12_RESOURCE_BARRIER_ALL_SUBRESOURCES, FD3D12DynamicRHI::ETransitionMode::Validate);
+	D3D12_RESOURCE_STATES IndirectState = GUseInternalTransitions ? D3D12_RESOURCE_STATE_INDIRECT_ARGUMENT | D3D12_RESOURCE_STATE_NON_PIXEL_SHADER_RESOURCE : D3D12_RESOURCE_STATE_INDIRECT_ARGUMENT;
+	FD3D12DynamicRHI::TransitionResource(CommandListHandle, Location.GetResource(), D3D12_RESOURCE_STATE_TBD, IndirectState, D3D12_RESOURCE_BARRIER_ALL_SUBRESOURCES, FD3D12DynamicRHI::ETransitionMode::Validate);
 	CommandListHandle.FlushResourceBarriers();	// Must flush so the desired state is actually set.
 
 	FD3D12Adapter* Adapter = GetParentDevice()->GetParentAdapter();
@@ -439,9 +440,12 @@ static void HandleTransientAliasing(FD3D12CommandContext& Context, const FD3D12T
 		{
 			TRACE_CPUPROFILER_EVENT_SCOPE(D3D12RHI::DiscardTransient);
 
-			// Restore the resource back to the initial state when done
-			D3D12_RESOURCE_STATES FinalState = GetInitialResourceState(Resource->GetDesc());
-			FD3D12DynamicRHI::TransitionResource(Context.CommandListHandle, Resource, D3D12_RESOURCE_STATE_TBD, FinalState, D3D12_RESOURCE_BARRIER_ALL_SUBRESOURCES, FD3D12DynamicRHI::ETransitionMode::Apply);
+			if (GUseInternalTransitions)
+			{
+				// Restore the resource back to the initial state when done
+				D3D12_RESOURCE_STATES FinalState = GetInitialResourceState(Resource->GetDesc());
+				FD3D12DynamicRHI::TransitionResource(Context.CommandListHandle, Resource, D3D12_RESOURCE_STATE_TBD, FinalState, D3D12_RESOURCE_BARRIER_ALL_SUBRESOURCES, FD3D12DynamicRHI::ETransitionMode::Apply);
+			}
 
 			// Remove from caches
 			Context.ConditionalClearShaderResource(&BaseShaderResource->ResourceLocation);
@@ -500,7 +504,7 @@ static void HandleResourceTransitions(FD3D12CommandContext& Context, const FD3D1
 					// Derive or get the correct before & after state
 					D3D12_RESOURCE_STATES BeforeState = bKnownBeforeState ? GetD3D12ResourceState(Info.AccessBefore, bIsAsyncCompute) : D3D12_RESOURCE_STATE_TBD;
 					D3D12_RESOURCE_STATES AfterState = GetD3D12ResourceState(Info.AccessAfter, bIsAsyncCompute);
-
+					
 					// Add the compression flags if needed
 					if (EnumHasAnyFlags(Info.Flags, EResourceTransitionFlags::MaintainCompression))
 					{
diff --git a/Engine/Source/Runtime/D3D12RHI/Private/D3D12DescriptorCache.cpp b/Engine/Source/Runtime/D3D12RHI/Private/D3D12DescriptorCache.cpp
--- a/Engine/Source/Runtime/D3D12RHI/Private/D3D12DescriptorCache.cpp
+++ b/Engine/Source/Runtime/D3D12RHI/Private/D3D12DescriptorCache.cpp
@@ -284,15 +284,21 @@ void FD3D12DescriptorCache::SetVertexBuffers(FD3D12VertexBufferCache& Cache)
 	CmdContext->CommandListHandle.UpdateResidency(Cache.ResidencyHandles, Count);
 	CmdContext->CommandListHandle->IASetVertexBuffers(0, Count, Cache.CurrentVertexBufferViews);
 
-	for (uint32 i = 0; i < Count; ++i)
+	// If using external transition then don't bother to Validate the state because
+	// resource could have already transitioned to new state but left in the cache without being actually used by the GPU
+	// It's still wrong but needs a bigger high level refactor to fix this problem.
+	if (GUseInternalTransitions)
 	{
-		if (Cache.CurrentVertexBufferResources[i])
+		for (uint32 i = 0; i < Count; ++i)
 		{
-			FD3D12Resource* Resource = Cache.CurrentVertexBufferResources[i]->GetResource();
-			if (Resource && Resource->RequiresResourceStateTracking())
+			if (Cache.CurrentVertexBufferResources[i])
 			{
-				check(Resource->GetSubresourceCount() == 1);
-				FD3D12DynamicRHI::TransitionResource(CmdContext->CommandListHandle, Resource, D3D12_RESOURCE_STATE_TBD, D3D12_RESOURCE_STATE_VERTEX_AND_CONSTANT_BUFFER, D3D12_RESOURCE_BARRIER_ALL_SUBRESOURCES, FD3D12DynamicRHI::ETransitionMode::Validate);
+				FD3D12Resource* Resource = Cache.CurrentVertexBufferResources[i]->GetResource();
+				if (Resource && Resource->RequiresResourceStateTracking())
+				{
+					check(Resource->GetSubresourceCount() == 1);
+					FD3D12DynamicRHI::TransitionResource(CmdContext->CommandListHandle, Resource, D3D12_RESOURCE_STATE_TBD, D3D12_RESOURCE_STATE_VERTEX_AND_CONSTANT_BUFFER, D3D12_RESOURCE_BARRIER_ALL_SUBRESOURCES, FD3D12DynamicRHI::ETransitionMode::Validate);
+				}
 			}
 		}
 	}
diff --git a/Engine/Source/Runtime/D3D12RHI/Private/D3D12PoolAllocator.cpp b/Engine/Source/Runtime/D3D12RHI/Private/D3D12PoolAllocator.cpp
--- a/Engine/Source/Runtime/D3D12RHI/Private/D3D12PoolAllocator.cpp
+++ b/Engine/Source/Runtime/D3D12RHI/Private/D3D12PoolAllocator.cpp
@@ -754,44 +754,21 @@ void FD3D12PoolAllocator::FlushPendingCopyOps(FD3D12CommandContext& InCommandCon
 	FD3D12Fence& FrameFence = Adapter->GetFrameFence();
 
 	for (FD3D12VRAMCopyOperation& CopyOperation : PendingCopyOps)
-	{		
+	{
 		// Cleared copy op?
 		if (CopyOperation.SourceResource == nullptr || CopyOperation.DestResource == nullptr)
 		{
 			continue;
 		}
 
-		bool bRTAccelerationStructure = false;
-		if (CopyOperation.SourceResource->RequiresResourceStateTracking())
-		{
-			check(CopyOperation.DestResource->RequiresResourceStateTracking());
-			FD3D12DynamicRHI::TransitionResource(CommandListHandle, CopyOperation.SourceResource, D3D12_RESOURCE_STATE_TBD, D3D12_RESOURCE_STATE_COPY_SOURCE, D3D12_RESOURCE_BARRIER_ALL_SUBRESOURCES, FD3D12DynamicRHI::ETransitionMode::Apply);
-			FD3D12DynamicRHI::TransitionResource(CommandListHandle, CopyOperation.DestResource, D3D12_RESOURCE_STATE_TBD, D3D12_RESOURCE_STATE_COPY_DEST, D3D12_RESOURCE_BARRIER_ALL_SUBRESOURCES, FD3D12DynamicRHI::ETransitionMode::Apply);
-		}
 #if D3D12_RHI_RAYTRACING
-		else if (CopyOperation.SourceResource->GetDefaultResourceState() == D3D12_RESOURCE_STATE_RAYTRACING_ACCELERATION_STRUCTURE)
+		if (!CopyOperation.SourceResource->RequiresResourceStateTracking() && CopyOperation.SourceResource->GetDefaultResourceState() == D3D12_RESOURCE_STATE_RAYTRACING_ACCELERATION_STRUCTURE)
 		{
 			// can't make state changes to RT resources
-			bRTAccelerationStructure = true;
 			check(CopyOperation.DestResource->GetDefaultResourceState() == D3D12_RESOURCE_STATE_RAYTRACING_ACCELERATION_STRUCTURE);
-		}
-#endif // D3D12_RHI_RAYTRACING
-		else
-		{
-			check(!CopyOperation.DestResource->RequiresResourceStateTracking());
-			CommandListHandle.AddTransitionBarrier(CopyOperation.SourceResource, CopyOperation.SourceResource->GetDefaultResourceState(), D3D12_RESOURCE_STATE_COPY_SOURCE, D3D12_RESOURCE_BARRIER_ALL_SUBRESOURCES);
-			CommandListHandle.AddTransitionBarrier(CopyOperation.DestResource, CopyOperation.DestResource->GetDefaultResourceState(), D3D12_RESOURCE_STATE_COPY_DEST, D3D12_RESOURCE_BARRIER_ALL_SUBRESOURCES);
-		}
 
-		// Enable to log all defrag copy ops
-		//UE_LOG(LogD3D12RHI, Log, TEXT("Running defrag copy op for: %s at Frame: %d (Old: %#016llx New: %#016llx)"), *CopyOperation.SourceResource->GetName().ToString(), FrameFence.GetCurrentFence(), CopyOperation.SourceResource, CopyOperation.DestResource);
+			InCommandContext.numCopies++;
 
-		InCommandContext.numCopies++;
-		CommandListHandle.FlushResourceBarriers();
-
-#if D3D12_RHI_RAYTRACING
-		if (bRTAccelerationStructure)
-		{
 			D3D12_GPU_VIRTUAL_ADDRESS SrcAddress = CopyOperation.SourceResource->GetResource()->GetGPUVirtualAddress() + CopyOperation.SourceOffset;
 			D3D12_GPU_VIRTUAL_ADDRESS DestAddress = CopyOperation.DestResource->GetResource()->GetGPUVirtualAddress() + CopyOperation.DestOffset;
 			CommandListHandle.RayTracingCommandList()->CopyRaytracingAccelerationStructure(DestAddress, SrcAddress, D3D12_RAYTRACING_ACCELERATION_STRUCTURE_COPY_MODE_CLONE);
@@ -799,6 +776,12 @@ void FD3D12PoolAllocator::FlushPendingCopyOps(FD3D12CommandContext& InCommandCon
 		else
 #endif // D3D12_RHI_RAYTRACING
 		{
+			FScopedResourceBarrier SourceScopedResourceBarrier(CommandListHandle, CopyOperation.SourceResource, D3D12_RESOURCE_STATE_COPY_SOURCE, D3D12_RESOURCE_BARRIER_ALL_SUBRESOURCES, FD3D12DynamicRHI::ETransitionMode::Apply);
+			FScopedResourceBarrier DestScopedResourceBarrier(CommandListHandle, CopyOperation.DestResource, D3D12_RESOURCE_STATE_COPY_DEST, D3D12_RESOURCE_BARRIER_ALL_SUBRESOURCES, FD3D12DynamicRHI::ETransitionMode::Apply);
+
+			InCommandContext.numCopies++;
+			CommandListHandle.FlushResourceBarriers();
+
 			switch (CopyOperation.CopyType)
 			{
 			case FD3D12VRAMCopyOperation::BufferRegion:
@@ -817,12 +800,6 @@ void FD3D12PoolAllocator::FlushPendingCopyOps(FD3D12CommandContext& InCommandCon
 
 		CommandListHandle.UpdateResidency(CopyOperation.SourceResource);
 		CommandListHandle.UpdateResidency(CopyOperation.DestResource);
-
-		if (!bRTAccelerationStructure && !CopyOperation.SourceResource->RequiresResourceStateTracking())
-		{
-			CommandListHandle.AddTransitionBarrier(CopyOperation.SourceResource, D3D12_RESOURCE_STATE_COPY_SOURCE, CopyOperation.SourceResource->GetDefaultResourceState(), D3D12_RESOURCE_BARRIER_ALL_SUBRESOURCES);
-			CommandListHandle.AddTransitionBarrier(CopyOperation.DestResource, D3D12_RESOURCE_STATE_COPY_DEST, CopyOperation.DestResource->GetDefaultResourceState(), D3D12_RESOURCE_BARRIER_ALL_SUBRESOURCES);
-		}
 	}
 
 	InCommandContext.ConditionalFlushCommandList();
diff --git a/Engine/Source/Runtime/D3D12RHI/Private/D3D12RHIPrivate.h b/Engine/Source/Runtime/D3D12RHI/Private/D3D12RHIPrivate.h
--- a/Engine/Source/Runtime/D3D12RHI/Private/D3D12RHIPrivate.h
+++ b/Engine/Source/Runtime/D3D12RHI/Private/D3D12RHIPrivate.h
@@ -623,7 +623,7 @@ public:
 		check(Before != D3D12_RESOURCE_STATE_TBD && After != D3D12_RESOURCE_STATE_TBD);
 
 		// Depth write is actually a suitable for read operations as a "normal" depth buffer.
-		if ((Before == D3D12_RESOURCE_STATE_DEPTH_WRITE) && (After == D3D12_RESOURCE_STATE_DEPTH_READ))
+		if (bInAllowStateMerging && (Before == D3D12_RESOURCE_STATE_DEPTH_WRITE) && (After == D3D12_RESOURCE_STATE_DEPTH_READ))
 		{
 			return false;
 		}
diff --git a/Engine/Source/Runtime/D3D12RHI/Private/D3D12TransientResourceAllocator.cpp b/Engine/Source/Runtime/D3D12RHI/Private/D3D12TransientResourceAllocator.cpp
--- a/Engine/Source/Runtime/D3D12RHI/Private/D3D12TransientResourceAllocator.cpp
+++ b/Engine/Source/Runtime/D3D12RHI/Private/D3D12TransientResourceAllocator.cpp
@@ -144,9 +144,20 @@ FRHITransientTexture* FD3D12TransientResourceHeapAllocator::CreateTexture(const
 	return CreateTextureInternal(InCreateInfo, InDebugName, InPassIndex, Info.SizeInBytes, Info.Alignment,
 		[&](const FRHITransientHeap::FResourceInitializer& Initializer)
 	{
+		ERHIAccess InitialState = ERHIAccess::UAVMask;
+
+		if (EnumHasAnyFlags(InCreateInfo.Flags, TexCreate_RenderTargetable | TexCreate_ResolveTargetable))
+		{
+			InitialState = ERHIAccess::RTV;
+		}
+		else if (EnumHasAnyFlags(InCreateInfo.Flags, TexCreate_DepthStencilTargetable))
+		{
+			InitialState = ERHIAccess::DSVWrite;
+		}
+
 		FResourceAllocatorAdapter ResourceAllocatorAdapter(GetParentAdapter(), static_cast<FD3D12TransientHeap&>(Initializer.Heap), Initializer.Allocation, Desc);
 
-		FRHITextureCreateDesc CreateDesc(InCreateInfo, ERHIAccess::Discard, InDebugName);
+		FRHITextureCreateDesc CreateDesc(InCreateInfo, InitialState, InDebugName);
 		FRHITexture* Texture = DynamicRHI->CreateD3D12Texture(CreateDesc, nullptr, &ResourceAllocatorAdapter);
 		return new FRHITransientTexture(Texture, ResourceAllocatorAdapter.GpuVirtualAddress, Initializer.Hash, Info.SizeInBytes, ERHITransientAllocationType::Heap, InCreateInfo);
 	});
