Fix for invalid D3D resource states due to storage in signed bit field

Valid D3D12_RESOURCE_STATES values now cover the the lower 31 bits.  CResourceState utilized a D3D12_RESOURCE_STATES 31 bit bit field to store a D3D12_RESOURCE_STATES.  However, the D3D12_RESOURCE_STATES enum is not explicitly typed so the compiler can interpret it as a signed type.  This allows the resource state value to overflow the supported range of the siigned 31 bit bit field rather than exactly retain the lower 31 bits as we desire.

CResourceState now stores the resource state in an unsigned 31 bit bit field.  This guarentees we retain the lower 31 bits of the resoure state exactly.  D3D12_RESOURCE_STATE_TBD and D3D12_RESOURCE_STATE_CORRUPT have been updated accordingly so that they do not utilize the top bit.  They were previously relying on the signed conversion to retain their negative values which is no longer possible.

An assertion has been added so that if a D3D12_RESOURCE_STATES is encountered that utilzes the top bit in the future, and hence cannot be correctly stored in our 31 bit bit field, we will catch it.

diff --git a/Engine/Source/Runtime/D3D12RHI/Private/D3D12Util.cpp b/Engine/Source/Runtime/D3D12RHI/Private/D3D12Util.cpp
--- a/Engine/Source/Runtime/D3D12RHI/Private/D3D12Util.cpp
+++ b/Engine/Source/Runtime/D3D12RHI/Private/D3D12Util.cpp
@@ -1416,7 +1416,7 @@ D3D12_RESOURCE_STATES CResourceState::GetSubresourceState(uint32 SubresourceInde
 {
 	if (m_AllSubresourcesSame)
 	{
-		return m_ResourceState;
+		return static_cast<D3D12_RESOURCE_STATES>(m_ResourceState);
 	}
 	else
 	{
@@ -1459,7 +1459,10 @@ void CResourceState::SetResourceState(D3D12_RESOURCE_STATES State)
 {
 	m_AllSubresourcesSame = 1;
 
-	m_ResourceState = State;
+	// m_ResourceState is restricted to 31 bits.  Ensure State can be properly represented.
+	check((State & (1 << 31)) == 0);
+
+	m_ResourceState = *reinterpret_cast<uint32*>(&State);
 
 	// State is now tracked per-resource, so m_SubresourceState should not be read.
 #if UE_BUILD_DEBUG
@@ -1490,7 +1493,7 @@ void CResourceState::SetSubresourceState(uint32 SubresourceIndex, D3D12_RESOURCE
 			const uint32 numSubresourceStates = m_SubresourceState.Num();
 			for (uint32 i = 0; i < numSubresourceStates; i++)
 			{
-				m_SubresourceState[i] = m_ResourceState;
+				m_SubresourceState[i] = static_cast<D3D12_RESOURCE_STATES>(m_ResourceState);
 			}
 
 			m_AllSubresourcesSame = 0;
diff --git a/Engine/Source/Runtime/D3D12RHI/Public/D3D12Util.h b/Engine/Source/Runtime/D3D12RHI/Public/D3D12Util.h
--- a/Engine/Source/Runtime/D3D12RHI/Public/D3D12Util.h
+++ b/Engine/Source/Runtime/D3D12RHI/Public/D3D12Util.h
@@ -263,8 +263,8 @@ void StallRHIThreadAndForceFlush(FD3D12Device* InDevice);
 
 // Custom resource states
 // To Be Determined (TBD) means we need to fill out a resource barrier before the command list is executed.
-#define D3D12_RESOURCE_STATE_TBD (D3D12_RESOURCE_STATES)-1
-#define D3D12_RESOURCE_STATE_CORRUPT (D3D12_RESOURCE_STATES)-2
+#define D3D12_RESOURCE_STATE_TBD D3D12_RESOURCE_STATES(-1 ^ (1 << 31))
+#define D3D12_RESOURCE_STATE_CORRUPT D3D12_RESOURCE_STATES(-2 ^ (1 << 31))
 
 static bool IsValidD3D12ResourceState(D3D12_RESOURCE_STATES InState)
 {
@@ -301,7 +301,7 @@ public:
 private:
 	// Only used if m_AllSubresourcesSame is 1.
 	// Bits defining the state of the full resource, bits are from D3D12_RESOURCE_STATES
-	D3D12_RESOURCE_STATES m_ResourceState : 31;
+	uint32 m_ResourceState : 31;
 
 	// Set to 1 if m_ResourceState is valid.  In this case, all subresources have the same state
 	// Set to 0 if m_SubresourceState is valid.  In this case, each subresources may have a different state (or may be unknown)
