Add support for breadcrumbs on XB consoles just like we have on PC

diff --git a/Engine/Source/Runtime/D3D12RHI/Private/D3D12Adapter.cpp b/Engine/Source/Runtime/D3D12RHI/Private/D3D12Adapter.cpp
--- a/Engine/Source/Runtime/D3D12RHI/Private/D3D12Adapter.cpp
+++ b/Engine/Source/Runtime/D3D12RHI/Private/D3D12Adapter.cpp
@@ -303,64 +303,7 @@
 
 #if PLATFORM_WINDOWS || (PLATFORM_HOLOLENS && !UE_BUILD_SHIPPING && WITH_PIX_EVENT_RUNTIME)
 	
-	// Multiple ways to enable the different D3D12 crash debugging modes:
-	// - via RHI independent r.GPUCrashDebugging cvar: by default enable low overhead breadcrumbs and NvAftermath are enabled
-	// - via 'gpucrashdebugging' command line argument: enable all possible GPU crash debug modes (minor performance impact)
-	// - via 'r.D3D12.BreadCrumbs', 'r.D3D12.AfterMath' or 'r.D3D12.Dred' each type of GPU crash debugging mode can be enabled
-	// - via '-gpubreadcrumbs(=0)', '-nvaftermath(=0)' or '-dred(=0)' command line argument: each type of gpu crash debugging mode can enabled/disabled
-	if (FParse::Param(FCommandLine::Get(), TEXT("gpucrashdebugging")))
-	{
-		GPUCrashDebuggingModes = ED3D12GPUCrashDebuggingModes::All;
-	}
-	else
-	{
-		// Parse the specific GPU crash debugging cvars and enable the different modes
-		const auto ParseCVar = [this](const TCHAR* CVarName, ED3D12GPUCrashDebuggingModes DebuggingMode)
-		{
-			IConsoleVariable* ConsoleVariable = IConsoleManager::Get().FindConsoleVariable(CVarName);
-			if (ConsoleVariable && ConsoleVariable->GetInt() > 0)
-			{
-				EnumAddFlags(GPUCrashDebuggingModes, DebuggingMode);
-			}
-		};
-		ParseCVar(TEXT("r.GPUCrashDebugging"), ED3D12GPUCrashDebuggingModes((int)ED3D12GPUCrashDebuggingModes::BreadCrumbs | (int)ED3D12GPUCrashDebuggingModes::NvAftermath));
-		ParseCVar(TEXT("r.D3D12.BreadCrumbs"), ED3D12GPUCrashDebuggingModes::BreadCrumbs);
-		ParseCVar(TEXT("r.D3D12.NvAfterMath"), ED3D12GPUCrashDebuggingModes::NvAftermath);
-		ParseCVar(TEXT("r.D3D12.DRED"), ED3D12GPUCrashDebuggingModes::DRED);
+    SetupGPUCrashDebuggingModesCommon();
 
-		// Enable/disable specific crash debugging modes if requested via command line argument
-		const auto ParseCommandLine = [this](const TCHAR* CommandLineArgument, ED3D12GPUCrashDebuggingModes DebuggingMode)
-		{
-			int32 Value = 0;
-			if (FParse::Value(FCommandLine::Get(), *FString::Printf(TEXT("%s="), CommandLineArgument), Value))
-			{
-				if (Value > 0)
-				{
-					EnumAddFlags(GPUCrashDebuggingModes, DebuggingMode);
-				}
-				else
-				{
-					EnumRemoveFlags(GPUCrashDebuggingModes, DebuggingMode);
-				}
-			}
-			else  if (FParse::Param(FCommandLine::Get(), CommandLineArgument))
-			{
-				EnumAddFlags(GPUCrashDebuggingModes, DebuggingMode);
-			}
-		};
-		ParseCommandLine(TEXT("gpubreadcrumbs"), ED3D12GPUCrashDebuggingModes::BreadCrumbs);
-		ParseCommandLine(TEXT("nvaftermath"), ED3D12GPUCrashDebuggingModes::NvAftermath);
-		ParseCommandLine(TEXT("dred"), ED3D12GPUCrashDebuggingModes::DRED);
-	}
-
-	// Submit draw events when any crash debugging mode is enabled
-	if (GPUCrashDebuggingModes != ED3D12GPUCrashDebuggingModes::None)
-	{
-		SetEmitDrawEvents(true);
-	}
-
-	bool bBreadcrumbs = EnumHasAnyFlags(GPUCrashDebuggingModes, ED3D12GPUCrashDebuggingModes::BreadCrumbs);
-	FGenericCrashContext::SetEngineData(TEXT("RHI.Breadcrumbs"), bBreadcrumbs ? TEXT("true") : TEXT("false"));
-
 #if NV_AFTERMATH
 	if (IsRHIDeviceNVIDIA() && GDX12NVAfterMathModuleLoaded)
@@ -1207,6 +1150,69 @@
 	checkf(DispatchIndirectComputeCommandSignature.IsValid(), TEXT("Indirect compute dispatch command signature is expected to be created by platform-specific D3D12 adapter implementation."))
 }
 
+void FD3D12Adapter::SetupGPUCrashDebuggingModesCommon()
+{
+	// Multiple ways to enable the different D3D12 crash debugging modes:
+	// - via RHI independent r.GPUCrashDebugging cvar: by default enable low overhead breadcrumbs and NvAftermath are enabled
+	// - via 'gpucrashdebugging' command line argument: enable all possible GPU crash debug modes (minor performance impact)
+	// - via 'r.D3D12.BreadCrumbs', 'r.D3D12.AfterMath' or 'r.D3D12.Dred' each type of GPU crash debugging mode can be enabled
+	// - via '-gpubreadcrumbs(=0)', '-nvaftermath(=0)' or '-dred(=0)' command line argument: each type of gpu crash debugging mode can enabled/disabled
+	if (FParse::Param(FCommandLine::Get(), TEXT("gpucrashdebugging")))
+	{
+		GPUCrashDebuggingModes = ED3D12GPUCrashDebuggingModes::All;
+	}
+	else
+	{
+		// Parse the specific GPU crash debugging cvars and enable the different modes
+		const auto ParseCVar = [this](const TCHAR* CVarName, ED3D12GPUCrashDebuggingModes DebuggingMode)
+		{
+			IConsoleVariable* ConsoleVariable = IConsoleManager::Get().FindConsoleVariable(CVarName);
+			if (ConsoleVariable && ConsoleVariable->GetInt() > 0)
+			{
+				EnumAddFlags(GPUCrashDebuggingModes, DebuggingMode);
+			}
+		};
+		ParseCVar(TEXT("r.GPUCrashDebugging"), ED3D12GPUCrashDebuggingModes((int)ED3D12GPUCrashDebuggingModes::BreadCrumbs | (int)ED3D12GPUCrashDebuggingModes::NvAftermath));
+		ParseCVar(TEXT("r.D3D12.BreadCrumbs"), ED3D12GPUCrashDebuggingModes::BreadCrumbs);
+		ParseCVar(TEXT("r.D3D12.NvAfterMath"), ED3D12GPUCrashDebuggingModes::NvAftermath);
+		ParseCVar(TEXT("r.D3D12.DRED"), ED3D12GPUCrashDebuggingModes::DRED);
 
+		// Enable/disable specific crash debugging modes if requested via command line argument
+		const auto ParseCommandLine = [this](const TCHAR* CommandLineArgument, ED3D12GPUCrashDebuggingModes DebuggingMode)
+		{
+			int32 Value = 0;
+			if (FParse::Value(FCommandLine::Get(), *FString::Printf(TEXT("%s="), CommandLineArgument), Value))
+			{
+				if (Value > 0)
+				{
+					EnumAddFlags(GPUCrashDebuggingModes, DebuggingMode);
+				}
+				else
+				{
+					EnumRemoveFlags(GPUCrashDebuggingModes, DebuggingMode);
+				}
+			}
+			else  if (FParse::Param(FCommandLine::Get(), CommandLineArgument))
+			{
+				EnumAddFlags(GPUCrashDebuggingModes, DebuggingMode);
+			}
+		};
+		ParseCommandLine(TEXT("gpubreadcrumbs"), ED3D12GPUCrashDebuggingModes::BreadCrumbs);
+		ParseCommandLine(TEXT("nvaftermath"), ED3D12GPUCrashDebuggingModes::NvAftermath);
+		ParseCommandLine(TEXT("dred"), ED3D12GPUCrashDebuggingModes::DRED);
+	}
+
+	// Submit draw events when any crash debugging mode is enabled
+	if (GPUCrashDebuggingModes != ED3D12GPUCrashDebuggingModes::None)
+	{
+		SetEmitDrawEvents(true);
+	}
+
+	bool bBreadcrumbs = EnumHasAnyFlags(GPUCrashDebuggingModes, ED3D12GPUCrashDebuggingModes::BreadCrumbs);
+	FGenericCrashContext::SetEngineData(TEXT("RHI.Breadcrumbs"), bBreadcrumbs ? TEXT("true") : TEXT("false"));
+
+}
+
+
 void FD3D12Adapter::Cleanup()
 {
diff --git a/Engine/Source/Runtime/D3D12RHI/Private/D3D12Adapter.h b/Engine/Source/Runtime/D3D12RHI/Private/D3D12Adapter.h
--- a/Engine/Source/Runtime/D3D12RHI/Private/D3D12Adapter.h
+++ b/Engine/Source/Runtime/D3D12RHI/Private/D3D12Adapter.h
@@ -448,6 +448,8 @@ protected:
 	// Creates default root and execute indirect signatures
 	virtual void CreateCommandSignatures();
 
+	void SetupGPUCrashDebuggingModesCommon();
+
 	FD3D12DynamicRHI* OwningRHI;
 
 	// LDA setups have one ID3D12Device
diff --git a/Engine/Source/Runtime/D3D12RHI/Private/D3D12CommandContext.cpp b/Engine/Source/Runtime/D3D12RHI/Private/D3D12CommandContext.cpp
--- a/Engine/Source/Runtime/D3D12RHI/Private/D3D12CommandContext.cpp
+++ b/Engine/Source/Runtime/D3D12RHI/Private/D3D12CommandContext.cpp
@@ -151,10 +151,8 @@ FD3D12CommandContext::~FD3D12CommandContext()
 /** Write out the event stack to the bread crumb resource if available */
 void FD3D12CommandContext::WriteGPUEventStackToBreadCrumbData(bool bBeginEvent)
 {
-	// Only in Windows for now, could be made available on Xbox as well
-#if PLATFORM_WINDOWS
 	// Write directly to command list if breadcrumb resource is available
-	FD3D12Resource* BreadCrumbResource = CommandListHandle.GetCommandListManager()->GetBreadCrumbResource();
+	FD3D12Resource* BreadCrumbResource = GetCommandListManager().GetBreadCrumbResource();
 	ID3D12GraphicsCommandList2* CommandList2 = CommandListHandle.GraphicsCommandList2();
 	if (BreadCrumbResource && CommandList2)
 	{
@@ -200,7 +198,6 @@ void FD3D12CommandContext::WriteGPUEventStackToBreadCrumbData(bool bBeginEvent)
 			FMemory::Free(Mem);
 		}
 	}
-#endif // PLATFORM_WINDOWS
 }
 
 
diff --git a/Engine/Source/Runtime/D3D12RHI/Private/D3D12DirectCommandListManager.cpp b/Engine/Source/Runtime/D3D12RHI/Private/D3D12DirectCommandListManager.cpp
--- a/Engine/Source/Runtime/D3D12RHI/Private/D3D12DirectCommandListManager.cpp
+++ b/Engine/Source/Runtime/D3D12RHI/Private/D3D12DirectCommandListManager.cpp
@@ -329,7 +329,7 @@ FD3D12CommandListManager::FD3D12CommandListManager(FD3D12Device* InParent, D3D12
 	, CommandListFence(nullptr)
 	, CommandListType(InCommandListType)
 	, QueueType(InQueueType)
-	, BreadCrumbResourceAddress(nullptr)
+	, DiagnosticBuffer({ nullptr, nullptr, nullptr, 0 })
 	, bExcludeBackbufferWriteTransitionTime(false)
 #if WITH_PROFILEGPU || D3D12_SUBMISSION_GAP_RECORDER
 	, bShouldTrackCmdListTime(false)
@@ -369,14 +369,7 @@ void FD3D12CommandListManager::Destroy()
 		CommandListFence.SafeRelease();
 	}
 
-	if (BreadCrumbResource)
-	{
-		BreadCrumbResource.SafeRelease();
-		BreadCrumbHeap.SafeRelease();
-
-		VirtualFree(BreadCrumbResourceAddress, 0, MEM_RELEASE);
-		BreadCrumbResourceAddress = nullptr;
-	}
+	DestroyDiagnosticBuffer(DiagnosticBuffer);
 }
 
 void FD3D12CommandListManager::Create(const TCHAR* Name, uint32 NumCommandLists, uint32 Priority)
@@ -435,49 +428,12 @@ void FD3D12CommandListManager::Create(const TCHAR* Name, uint32 NumCommandLists,
 
 			DiagnosticBufferOffset = EventBufferSize;
 
-			BreadCrumbResourceAddress = VirtualAlloc(nullptr, BreadCrumbBufferSize, MEM_COMMIT, PAGE_READWRITE);
-			if (BreadCrumbResourceAddress)
-			{
-				// Create non refcounted heap because SetHeap function will take ownership without perform AddRef
-				ID3D12Heap* D3D12Heap = nullptr;
-				hr = D3D12Device3->OpenExistingHeapFromAddress(BreadCrumbResourceAddress, IID_PPV_ARGS(&D3D12Heap));
-				if (SUCCEEDED(hr))
-				{
-					BreadCrumbHeap = new FD3D12Heap(Device, GetVisibilityMask());
-					BreadCrumbHeap->SetHeap(D3D12Heap, TEXT("BreadCrumbHeap"));
-
-					TCHAR TempStr[MAX_SPRINTF] = TEXT("");
-					FCString::Sprintf(TempStr, TEXT("BreadCrumbResource_%s"), Name);
-
-					const D3D12_RESOURCE_DESC BufferDesc = CD3DX12_RESOURCE_DESC::Buffer(BreadCrumbBufferSize, D3D12_RESOURCE_FLAG_ALLOW_CROSS_ADAPTER);
-					hr = Adapter->CreatePlacedResource(BufferDesc, BreadCrumbHeap.GetReference(), 0, D3D12_RESOURCE_STATE_COPY_DEST, nullptr, BreadCrumbResource.GetInitReference(), TempStr, false);
-					if (SUCCEEDED(hr))
-					{
-						UE_LOG(LogD3D12RHI, Log, TEXT("[GPUBreadCrumb] Successfully setup breadcrumb resource for %s"), Name);
-						BreadCrumbResourceGPUAddress = BreadCrumbResource->GetGPUVirtualAddress();
-					}
-					else
-					{
-						BreadCrumbHeap.SafeRelease();
-
-						VirtualFree(BreadCrumbResourceAddress, 0, MEM_RELEASE);
-						BreadCrumbResourceAddress = nullptr;
-
-						UE_LOG(LogD3D12RHI, Warning, TEXT("[GPUBreadCrumb] Failed to CreatePlacedResource, error: %x"), hr);
-					}
-				}
-				else
-				{
-					VirtualFree(BreadCrumbResourceAddress, 0, MEM_RELEASE);
-					BreadCrumbResourceAddress = nullptr;
+			// Create the platform-specific diagnostic buffer
+			TCHAR TempStr[MAX_SPRINTF] = TEXT("");
+			FCString::Sprintf(TempStr, TEXT("BreadCrumbResource_%s"), Name);
 
-					UE_LOG(LogD3D12RHI, Warning, TEXT("[GPUBreadCrumb] Failed to OpenExistingHeapFromAddress, error: %x"), hr);
-				}
-			}
-		}
-		else
-		{
-			UE_LOG(LogD3D12RHI, Warning, TEXT("[GPUBreadCrumb] ID3D12Device3 not available (only available on Windows 10 1709+), error: %x"), hr);
+			const D3D12_RESOURCE_DESC BufferDesc = CD3DX12_RESOURCE_DESC::Buffer(BreadCrumbBufferSize, D3D12_RESOURCE_FLAG_ALLOW_CROSS_ADAPTER);
+			DiagnosticBuffer = CreateDiagnosticBuffer(Adapter, Device, BufferDesc, Name);
 		}
 	}
 }
@@ -1154,6 +1110,66 @@ bool FD3D12CommandListManager::ShouldTrackCommandListTime() const
 #endif
 }
 
+FD3D12CommandListManager::FDiagnosticBuffer FD3D12CommandListManager::CreateDiagnosticBuffer(FD3D12Adapter *Adapter, FD3D12Device *Device, const D3D12_RESOURCE_DESC& Desc, const TCHAR* Name)
+{
+	TRefCountPtr<ID3D12Device3> D3D12Device3;
+	HRESULT hr = Device->GetDevice()->QueryInterface(IID_PPV_ARGS(D3D12Device3.GetInitReference()));
+	if (SUCCEEDED(hr))
+	{
+		void* BreadCrumbResourceAddress = VirtualAlloc(nullptr, Desc.Width, MEM_COMMIT, PAGE_READWRITE);
+		if (BreadCrumbResourceAddress)
+		{
+			ID3D12Heap* D3D12Heap = nullptr;
+			hr = D3D12Device3->OpenExistingHeapFromAddress(BreadCrumbResourceAddress, IID_PPV_ARGS(&D3D12Heap));
+			if (SUCCEEDED(hr))
+			{
+				TRefCountPtr<FD3D12Heap> BreadCrumbHeap = new FD3D12Heap(Device, GetVisibilityMask());
+				BreadCrumbHeap->SetHeap(D3D12Heap, TEXT("DiagnosticBuffer"));
+
+				TRefCountPtr<FD3D12Resource> BreadCrumbResource;
+				hr = Adapter->CreatePlacedResource(Desc, BreadCrumbHeap.GetReference(), 0, D3D12_RESOURCE_STATE_COPY_DEST, nullptr, BreadCrumbResource.GetInitReference(), Name, false);
+				if (SUCCEEDED(hr))
+				{
+					UE_LOG(LogD3D12RHI, Log, TEXT("[GPUBreadCrumb] Successfully setup breadcrumb resource for %s"), Name);
+
+					return { BreadCrumbHeap, BreadCrumbResource, BreadCrumbResourceAddress, BreadCrumbResource->GetGPUVirtualAddress() };
+				}
+				else
+				{
+					BreadCrumbHeap.SafeRelease();
+					VirtualFree(BreadCrumbResourceAddress, 0, MEM_RELEASE);
+					UE_LOG(LogD3D12RHI, Warning, TEXT("[GPUBreadCrumb] Failed to CreatePlacedResource, error: %x"), hr);
+				}
+			}
+			else
+			{
+				VirtualFree(BreadCrumbResourceAddress, 0, MEM_RELEASE);
+				UE_LOG(LogD3D12RHI, Warning, TEXT("[GPUBreadCrumb] Failed to OpenExistingHeapFromAddress, error: %x"), hr);
+			}
+		}
+		else
+		{
+			UE_LOG(LogD3D12RHI, Warning, TEXT("[GPUBreadCrumb] Failed to VirtualAlloc resource memory"));
+		}
+	}
+	else
+	{
+		UE_LOG(LogD3D12RHI, Warning, TEXT("[GPUBreadCrumb] ID3D12Device3 not available (only available on Windows 10 1709+), error: %x"), hr);
+	}
+
+	return { nullptr, nullptr, nullptr, 0 };
+}
+
+void FD3D12CommandListManager::DestroyDiagnosticBuffer(FDiagnosticBuffer& Buffer)
+{
+	Buffer.BreadCrumbResource.SafeRelease();
+	Buffer.BreadCrumbHeap.SafeRelease();
+
+	VirtualFree(Buffer.BreadCrumbResourceAddress, 0, MEM_RELEASE);
+	Buffer.BreadCrumbResourceAddress = nullptr;
+	Buffer.BreadCrumbResourceGPUAddress = 0;
+}
+
 FD3D12FenceCore* FD3D12FenceCorePool::ObtainFenceCore(uint32 GPUIndex)
 {
 	{
@@ -1202,4 +1218,4 @@ void FD3D12CommandListPayload::Append(ID3D12CommandList* CommandList, FD3D12Resi
 	CommandLists[NumCommandLists] = CommandList;
 	ResidencySets[NumCommandLists] = Set;
 	NumCommandLists++;
-}
+}
\ No newline at end of file
diff --git a/Engine/Source/Runtime/D3D12RHI/Private/D3D12DirectCommandListManager.h b/Engine/Source/Runtime/D3D12RHI/Private/D3D12DirectCommandListManager.h
--- a/Engine/Source/Runtime/D3D12RHI/Private/D3D12DirectCommandListManager.h
+++ b/Engine/Source/Runtime/D3D12RHI/Private/D3D12DirectCommandListManager.h
@@ -280,18 +280,18 @@ public:
 	FORCEINLINE FD3D12Fence& GetFence() { check(CommandListFence); return *CommandListFence; }
 
 	/** Get the breadcrumb resource which is written during command context recording */
-	FD3D12Resource* GetBreadCrumbResource() { return BreadCrumbResource.GetReference(); }
+	FD3D12Resource* GetBreadCrumbResource() { return DiagnosticBuffer.BreadCrumbResource.GetReference(); }
 
 	/** Get the CPU readable data from the breadcrumb data - this data is still valid after the Device is Lost */
-	const void* GetBreadCrumbResourceAddress() const { return BreadCrumbResourceAddress; }
+	const void* GetBreadCrumbResourceAddress() const { return DiagnosticBuffer.BreadCrumbResourceAddress; }
 	const FD3D12DiagnosticBufferData* GetDiagnosticBufferData() const
 	{ 
-		const uint8* Address = BreadCrumbResourceAddress ? reinterpret_cast<const uint8*>(GetBreadCrumbResourceAddress()) + DiagnosticBufferOffset : nullptr;
+		const uint8* Address = DiagnosticBuffer.BreadCrumbResourceAddress ? reinterpret_cast<const uint8*>(GetBreadCrumbResourceAddress()) + DiagnosticBufferOffset : nullptr;
 		return reinterpret_cast<const FD3D12DiagnosticBufferData*>(Address);
 	}
 	const D3D12_GPU_VIRTUAL_ADDRESS GetDiagnosticBufferGPUAddress() const
 	{ 
-		return BreadCrumbResourceGPUAddress ? BreadCrumbResourceGPUAddress + DiagnosticBufferOffset : 0;
+		return DiagnosticBuffer.BreadCrumbResourceGPUAddress ? DiagnosticBuffer.BreadCrumbResourceGPUAddress + DiagnosticBufferOffset : 0;
 	}
 
 	void WaitForCommandQueueFlush();
@@ -379,10 +379,19 @@ protected:
 	TArray<FD3D12CommandListHandle>			ExecuteCommandListHandles;
 
 	// Helper data used to track GPU progress on this command queue
-	void* BreadCrumbResourceAddress;
-	TRefCountPtr<FD3D12Heap> BreadCrumbHeap;
-	TRefCountPtr<FD3D12Resource> BreadCrumbResource;
-	D3D12_GPU_VIRTUAL_ADDRESS BreadCrumbResourceGPUAddress = {};
+	struct FDiagnosticBuffer
+	{
+		TRefCountPtr<FD3D12Heap> BreadCrumbHeap;
+		TRefCountPtr<FD3D12Resource> BreadCrumbResource;
+
+		void* BreadCrumbResourceAddress;
+		D3D12_GPU_VIRTUAL_ADDRESS BreadCrumbResourceGPUAddress;
+	};
+
+	virtual FDiagnosticBuffer CreateDiagnosticBuffer(FD3D12Adapter *Adapter, FD3D12Device *Device, const D3D12_RESOURCE_DESC& Desc, const TCHAR* Name);
+	void DestroyDiagnosticBuffer(FDiagnosticBuffer& Buffer);
+
+	FDiagnosticBuffer DiagnosticBuffer;
 	uint32 DiagnosticBufferOffset = 0;
 
 #if WITH_PROFILEGPU || D3D12_SUBMISSION_GAP_RECORDER
