Changing SM6.6 minimum spec to require wave ops and 64bit atomics.
Adding correct support check for GRHISupportsAtomicUInt64 on D3D12 Windows.
Minor cleanup of FD3D12AdapterDesc.

diff --git a/Engine/Source/Runtime/D3D12RHI/Private/D3D12Adapter.cpp b/Engine/Source/Runtime/D3D12RHI/Private/D3D12Adapter.cpp
--- a/Engine/Source/Runtime/D3D12RHI/Private/D3D12Adapter.cpp
+++ b/Engine/Source/Runtime/D3D12RHI/Private/D3D12Adapter.cpp
@@ -202,6 +202,41 @@ static LONG __stdcall D3DVectoredExceptionHandler(EXCEPTION_POINTERS* InInfo)
 #endif // #if PLATFORM_WINDOWS || PLATFORM_HOLOLENS
 
 
+FD3D12AdapterDesc::FD3D12AdapterDesc() = default;
+
+FD3D12AdapterDesc::FD3D12AdapterDesc(const DXGI_ADAPTER_DESC& InDesc, int32 InAdapterIndex, D3D_FEATURE_LEVEL InMaxSupportedFeatureLevel, D3D_SHADER_MODEL InMaxSupportedShaderModel, ERHIFeatureLevel::Type InMaxRHIFeatureLevel)
+	: Desc(InDesc)
+	, AdapterIndex(InAdapterIndex)
+	, MaxSupportedFeatureLevel(InMaxSupportedFeatureLevel)
+	, MaxSupportedShaderModel(InMaxSupportedShaderModel)
+	, MaxRHIFeatureLevel(InMaxRHIFeatureLevel)
+{
+}
+
+bool FD3D12AdapterDesc::IsValid() const
+{
+	return MaxSupportedFeatureLevel != (D3D_FEATURE_LEVEL)0 && AdapterIndex >= 0;
+}
+
+#if PLATFORM_WINDOWS || PLATFORM_HOLOLENS
+HRESULT FD3D12AdapterDesc::EnumAdapters(int32 AdapterIndex, DXGI_GPU_PREFERENCE GpuPreference, IDXGIFactory* DxgiFactory, IDXGIFactory6* DxgiFactory6, IDXGIAdapter** TempAdapter)
+{
+	if (!DxgiFactory6 || GpuPreference == DXGI_GPU_PREFERENCE_UNSPECIFIED)
+	{
+		return DxgiFactory->EnumAdapters(AdapterIndex, TempAdapter);
+	}
+	else
+	{
+		return DxgiFactory6->EnumAdapterByGpuPreference(AdapterIndex, GpuPreference, IID_PPV_ARGS(TempAdapter));
+	}
+}
+
+HRESULT FD3D12AdapterDesc::EnumAdapters(IDXGIFactory* DxgiFactory, IDXGIFactory6* DxgiFactory6, IDXGIAdapter** TempAdapter) const
+{
+	return EnumAdapters(AdapterIndex, GpuPreference, DxgiFactory, DxgiFactory6, TempAdapter);
+}
+#endif
+
 FD3D12Adapter::FD3D12Adapter(FD3D12AdapterDesc& DescIn)
 	: OwningRHI(nullptr)
 	, bDepthBoundsTestSupported(false)
@@ -1003,6 +1038,7 @@ void FD3D12Adapter::InitializeDevices()
 				if (D3D12Caps9.AtomicInt64OnTypedResourceSupported)
 				{
 					UE_LOG(LogD3D12RHI, Log, TEXT("AtomicInt64OnTypedResource is supported"));
+					GRHISupportsAtomicUInt64 = true;
 				}
 
 				if (D3D12Caps9.AtomicInt64OnGroupSharedSupported)
diff --git a/Engine/Source/Runtime/D3D12RHI/Private/D3D12Adapter.h b/Engine/Source/Runtime/D3D12RHI/Private/D3D12Adapter.h
--- a/Engine/Source/Runtime/D3D12RHI/Private/D3D12Adapter.h
+++ b/Engine/Source/Runtime/D3D12RHI/Private/D3D12Adapter.h
@@ -35,67 +35,35 @@ This structure allows a single RHI to control several different hardware setups.
 
 #pragma once
 
-THIRD_PARTY_INCLUDES_START
-#if D3D12_SUPPORTS_DXGI_DEBUG
-	#include "dxgidebug.h"
-#endif
-THIRD_PARTY_INCLUDES_END
-
 class FD3D12DynamicRHI;
 
 struct FD3D12AdapterDesc
 {
-	FD3D12AdapterDesc() = default;
-
-	FD3D12AdapterDesc(const DXGI_ADAPTER_DESC& DescIn , int32 InAdapterIndex, D3D_FEATURE_LEVEL InMaxSupportedFeatureLevel, D3D_SHADER_MODEL InMaxSupportedShaderModel, uint32 NumNodes, bool InIsIntegrated
-#if PLATFORM_WINDOWS || PLATFORM_HOLOLENS
-		, DXGI_GPU_PREFERENCE InGpuPreference
-#endif
-	)
-		: AdapterIndex(InAdapterIndex)
-		, MaxSupportedFeatureLevel(InMaxSupportedFeatureLevel)
-		, MaxSupportedShaderModel(InMaxSupportedShaderModel)
-		, Desc(DescIn)
-		, NumDeviceNodes(NumNodes)
-		, bIsIntegrated(InIsIntegrated)
-#if PLATFORM_WINDOWS || PLATFORM_HOLOLENS
-		, GpuPreference(InGpuPreference)
-#endif
-	{
-	}
+	FD3D12AdapterDesc();
+	FD3D12AdapterDesc(const DXGI_ADAPTER_DESC& InDesc, int32 InAdapterIndex, D3D_FEATURE_LEVEL InMaxSupportedFeatureLevel, D3D_SHADER_MODEL InMaxSupportedShaderModel, ERHIFeatureLevel::Type InMaxRHIFeatureLevel);
 
-	bool IsValid() const { return MaxSupportedFeatureLevel != (D3D_FEATURE_LEVEL)0 && AdapterIndex >= 0; }
+	bool IsValid() const;
 
 #if PLATFORM_WINDOWS || PLATFORM_HOLOLENS
-	static HRESULT EnumAdapters(int32 AdapterIndex, DXGI_GPU_PREFERENCE GpuPreference, IDXGIFactory* DxgiFactory, IDXGIFactory6* DxgiFactory6, IDXGIAdapter** TempAdapter)
-	{
-		if (!DxgiFactory6 || GpuPreference == DXGI_GPU_PREFERENCE_UNSPECIFIED)
-		{
-			return DxgiFactory->EnumAdapters(AdapterIndex, TempAdapter);
-		}
-		else
-		{
-			return DxgiFactory6->EnumAdapterByGpuPreference(AdapterIndex, GpuPreference, IID_PPV_ARGS(TempAdapter));
-		}
-	}
-
-	HRESULT EnumAdapters(IDXGIFactory* DxgiFactory, IDXGIFactory6* DxgiFactory6, IDXGIAdapter** TempAdapter) const
-	{
-		return EnumAdapters(AdapterIndex, GpuPreference, DxgiFactory, DxgiFactory6, TempAdapter);
-	}
+	static HRESULT EnumAdapters(int32 AdapterIndex, DXGI_GPU_PREFERENCE GpuPreference, IDXGIFactory* DxgiFactory, IDXGIFactory6* DxgiFactory6, IDXGIAdapter** TempAdapter);
+	HRESULT EnumAdapters(IDXGIFactory* DxgiFactory, IDXGIFactory6* DxgiFactory6, IDXGIAdapter** TempAdapter) const;
 #endif
 
-	/** -1 if not supported or FindAdpater() wasn't called. Ideally we would store a pointer to IDXGIAdapter but it's unlikely the adpaters change during engine init. */
-	int32 AdapterIndex{ -1 };
+	DXGI_ADAPTER_DESC Desc{};
+
+	/** -1 if not supported or FindAdapter() wasn't called. Ideally we would store a pointer to IDXGIAdapter but it's unlikely the adpaters change during engine init. */
+	int32 AdapterIndex = -1;
 
-	/** The maximum D3D12 feature level supported. 0 if not supported or FindAdpater() wasn't called */
-	D3D_FEATURE_LEVEL MaxSupportedFeatureLevel{ (D3D_FEATURE_LEVEL)0 };
+	/** The maximum D3D12 feature level supported. 0 if not supported or FindAdapter() wasn't called */
+	D3D_FEATURE_LEVEL MaxSupportedFeatureLevel = (D3D_FEATURE_LEVEL)0;
 
-	D3D_SHADER_MODEL MaxSupportedShaderModel{ D3D_SHADER_MODEL_5_1 };
+	/** The maximum Shader Model supported. 0 if not supported or FindAdpater() wasn't called */
+	D3D_SHADER_MODEL MaxSupportedShaderModel = (D3D_SHADER_MODEL)0;
 
-	DXGI_ADAPTER_DESC Desc{};
+	ERHIFeatureLevel::Type MaxRHIFeatureLevel = ERHIFeatureLevel::Num;
 
-	uint32 NumDeviceNodes{ 0 };
+	/** Number of device nodes (read: GPUs) */
+	uint32 NumDeviceNodes = 1;
 
 	/** Whether the GPU is integrated or discrete. */
 	bool bIsIntegrated = false;
diff --git a/Engine/Source/Runtime/D3D12RHI/Private/Windows/WindowsD3D12Device.cpp b/Engine/Source/Runtime/D3D12RHI/Private/Windows/WindowsD3D12Device.cpp
--- a/Engine/Source/Runtime/D3D12RHI/Private/Windows/WindowsD3D12Device.cpp
+++ b/Engine/Source/Runtime/D3D12RHI/Private/Windows/WindowsD3D12Device.cpp
@@ -240,13 +240,41 @@ static D3D_FEATURE_LEVEL GetRequiredD3DFeatureLevel()
 	return D3D_FEATURE_LEVEL_11_0;
 }
 
+static D3D_FEATURE_LEVEL FindHighestFeatureLevel(ID3D12Device* Device, D3D_FEATURE_LEVEL MinFeatureLevel)
+{
+	const D3D_FEATURE_LEVEL FeatureLevels[] =
+	{
+		// Add new feature levels that the app supports here.
+#if D3D12_CORE_ENABLED
+		D3D_FEATURE_LEVEL_12_2,
+#endif
+		D3D_FEATURE_LEVEL_12_1,
+		D3D_FEATURE_LEVEL_12_0,
+		D3D_FEATURE_LEVEL_11_1,
+		D3D_FEATURE_LEVEL_11_0
+	};
+
+	// Determine the max feature level supported by the driver and hardware.
+	D3D12_FEATURE_DATA_FEATURE_LEVELS FeatureLevelCaps{};
+	FeatureLevelCaps.pFeatureLevelsRequested = FeatureLevels;
+	FeatureLevelCaps.NumFeatureLevels = UE_ARRAY_COUNT(FeatureLevels);
+
+	if (SUCCEEDED(Device->CheckFeatureSupport(D3D12_FEATURE_FEATURE_LEVELS, &FeatureLevelCaps, sizeof(FeatureLevelCaps))))
+	{
+		return FeatureLevelCaps.MaxSupportedFeatureLevel;
+	}
+
+	return MinFeatureLevel;
+}
+
 static D3D_SHADER_MODEL FindHighestShaderModel(ID3D12Device* Device)
 {
 	// Because we can't guarantee older Windows versions will know about newer shader models, we need to check them all
 	// in descending order and return the first result that succeeds.
 	const D3D_SHADER_MODEL ShaderModelsToCheck[] =
 	{
-#if !PLATFORM_HOLOLENS
+#if D3D12_CORE_ENABLED
+		D3D_SHADER_MODEL_6_7,
 		D3D_SHADER_MODEL_6_6,
 #endif
 		D3D_SHADER_MODEL_6_5,
@@ -272,42 +300,58 @@ static D3D_SHADER_MODEL FindHighestShaderModel(ID3D12Device* Device)
 	return D3D_SHADER_MODEL_5_1;
 }
 
+inline ERHIFeatureLevel::Type FindMaxRHIFeatureLevel(ID3D12Device* Device, D3D_FEATURE_LEVEL InMaxFeatureLevel, D3D_SHADER_MODEL InMaxShaderModel)
+{
+	ERHIFeatureLevel::Type MaxRHIFeatureLevel = ERHIFeatureLevel::Num;
+
+	if (InMaxFeatureLevel >= D3D_FEATURE_LEVEL_12_1 && InMaxShaderModel >= D3D_SHADER_MODEL_6_6)
+	{
+		D3D12_FEATURE_DATA_D3D12_OPTIONS1 D3D12Caps1{};
+		Device->CheckFeatureSupport(D3D12_FEATURE_D3D12_OPTIONS1, &D3D12Caps1, sizeof(D3D12Caps1));
+
+		D3D12_FEATURE_DATA_D3D12_OPTIONS9 D3D12Caps9{};
+		Device->CheckFeatureSupport(D3D12_FEATURE_D3D12_OPTIONS9, &D3D12Caps9, sizeof(D3D12Caps9));
+
+		if (D3D12Caps1.WaveOps && D3D12Caps9.AtomicInt64OnTypedResourceSupported)
+		{
+			MaxRHIFeatureLevel = ERHIFeatureLevel::SM6;
+		}
+	}
+
+	if (MaxRHIFeatureLevel == ERHIFeatureLevel::Num && InMaxFeatureLevel >= D3D_FEATURE_LEVEL_11_0)
+	{
+		MaxRHIFeatureLevel = ERHIFeatureLevel::SM5;
+	}
+
+	return MaxRHIFeatureLevel;
+}
+
+struct FD3D12DeviceBasicInfo
+{
+	D3D_FEATURE_LEVEL MaxFeatureLevel;
+	D3D_SHADER_MODEL  MaxShaderModel;
+	uint32            NumDeviceNodes;
+
+	ERHIFeatureLevel::Type MaxRHIFeatureLevel;
+};
+
 /**
  * Attempts to create a D3D12 device for the adapter using at minimum MinFeatureLevel.
  * If creation is successful, true is returned and the max supported feature level is set in OutMaxFeatureLevel.
  */
-static bool SafeTestD3D12CreateDevice(IDXGIAdapter* Adapter, D3D_FEATURE_LEVEL MinFeatureLevel, D3D_FEATURE_LEVEL& OutMaxFeatureLevel, D3D_SHADER_MODEL& OutMaxShaderModel, uint32& OutNumDeviceNodes)
+static bool SafeTestD3D12CreateDevice(IDXGIAdapter* Adapter, D3D_FEATURE_LEVEL MinFeatureLevel, FD3D12DeviceBasicInfo& OutInfo)
 {
-	const D3D_FEATURE_LEVEL FeatureLevels[] =
-	{
-		// Add new feature levels that the app supports here.
-		D3D_FEATURE_LEVEL_12_1,
-		D3D_FEATURE_LEVEL_12_0,
-		D3D_FEATURE_LEVEL_11_1,
-		D3D_FEATURE_LEVEL_11_0
-	};
-
 	__try
 	{
-		ID3D12Device* pDevice = nullptr;
-		if (SUCCEEDED(D3D12CreateDevice(Adapter, MinFeatureLevel, IID_PPV_ARGS(&pDevice))))
+		ID3D12Device* Device = nullptr;
+		if (SUCCEEDED(D3D12CreateDevice(Adapter, MinFeatureLevel, IID_PPV_ARGS(&Device))))
 		{
-			// Determine the max feature level supported by the driver and hardware.
-			D3D_FEATURE_LEVEL MaxFeatureLevel = MinFeatureLevel;
-			D3D12_FEATURE_DATA_FEATURE_LEVELS FeatureLevelCaps = {};
-			FeatureLevelCaps.pFeatureLevelsRequested = FeatureLevels;
-			FeatureLevelCaps.NumFeatureLevels = UE_ARRAY_COUNT(FeatureLevels);	
-			if (SUCCEEDED(pDevice->CheckFeatureSupport(D3D12_FEATURE_FEATURE_LEVELS, &FeatureLevelCaps, sizeof(FeatureLevelCaps))))
-			{
-				MaxFeatureLevel = FeatureLevelCaps.MaxSupportedFeatureLevel;
-			}
+			OutInfo.MaxFeatureLevel = FindHighestFeatureLevel(Device, MinFeatureLevel);
+			OutInfo.MaxShaderModel = FindHighestShaderModel(Device);
+			OutInfo.NumDeviceNodes = Device->GetNodeCount();
+			OutInfo.MaxRHIFeatureLevel = FindMaxRHIFeatureLevel(Device, OutInfo.MaxFeatureLevel, OutInfo.MaxShaderModel);
 
-			OutMaxShaderModel = FindHighestShaderModel(pDevice);
-
-			OutMaxFeatureLevel = MaxFeatureLevel;
-			OutNumDeviceNodes = pDevice->GetNodeCount();
-
-			pDevice->Release();
+			Device->Release();
 			return true;
 		}
 	}
@@ -404,19 +448,7 @@ static bool IsAdapterSupported(FD3D12Adapter* InAdapter, ERHIFeatureLevel::Type
 	{
 		if (const FD3D12AdapterDesc& Desc = InAdapter->GetDesc(); Desc.IsValid())
 		{
-			ERHIFeatureLevel::Type MaxSupportedFeatureLevel = ERHIFeatureLevel::Num;
-
-			// TODO: D3D_FEATURE_LEVEL_12_2?
-			if (Desc.MaxSupportedFeatureLevel >= D3D_FEATURE_LEVEL_12_1 && Desc.MaxSupportedShaderModel >= D3D_SHADER_MODEL_6_6)
-			{
-				MaxSupportedFeatureLevel = ERHIFeatureLevel::SM6;
-			}
-			else if (Desc.MaxSupportedFeatureLevel >= D3D_FEATURE_LEVEL_11_0)
-			{
-				MaxSupportedFeatureLevel = ERHIFeatureLevel::SM5;
-			}
-
-			return MaxSupportedFeatureLevel != ERHIFeatureLevel::Num && MaxSupportedFeatureLevel >= InRequestedFeatureLevel;
+			return Desc.MaxRHIFeatureLevel != ERHIFeatureLevel::Num && Desc.MaxRHIFeatureLevel >= InRequestedFeatureLevel;
 		}
 	}
 
@@ -459,7 +491,7 @@ namespace D3D12RHI
 		case D3D_FEATURE_LEVEL_11_1:	return TEXT("11_1");
 		case D3D_FEATURE_LEVEL_12_0:	return TEXT("12_0");
 		case D3D_FEATURE_LEVEL_12_1:	return TEXT("12_1");
-#if !PLATFORM_HOLOLENS
+#if D3D12_CORE_ENABLED
 		case D3D_FEATURE_LEVEL_12_2:	return TEXT("12_2");
 #endif
 		}
@@ -538,9 +570,6 @@ void FD3D12DynamicRHIModule::FindAdapter()
 	FD3D12AdapterDesc FirstDiscreteAdapter;
 	FD3D12AdapterDesc FirstAdapter;
 
-	bool bIsAnyAMD = false;
-	bool bIsAnyIntel = false;
-	bool bIsAnyNVIDIA = false;
 	bool bRequestedWARP = D3D12RHI_ShouldCreateWithWarp();
 	bool bAllowSoftwareRendering = D3D12RHI_AllowSoftwareFallback();
 
@@ -562,13 +591,10 @@ void FD3D12DynamicRHIModule::FindAdapter()
 		// Check that if adapter supports D3D12.
 		if (TempAdapter)
 		{
-			D3D_FEATURE_LEVEL MaxSupportedFeatureLevel = static_cast<D3D_FEATURE_LEVEL>(0);
-			D3D_SHADER_MODEL MaxSupportedShaderModel = D3D_SHADER_MODEL_5_1;
-			uint32 NumNodes = 0;
-
-			if (SafeTestD3D12CreateDevice(TempAdapter, MinRequiredFeatureLevel, MaxSupportedFeatureLevel, MaxSupportedShaderModel, NumNodes))
+			FD3D12DeviceBasicInfo DeviceInfo;
+			if (SafeTestD3D12CreateDevice(TempAdapter, MinRequiredFeatureLevel, DeviceInfo))
 			{
-				check(NumNodes > 0);
+				check(DeviceInfo.NumDeviceNodes > 0);
 				// Log some information about the available D3D12 adapters.
 				DXGI_ADAPTER_DESC AdapterDesc;
 				VERIFYD3D12RESULT(TempAdapter->GetDesc(&AdapterDesc));
@@ -578,27 +604,22 @@ void FD3D12DynamicRHIModule::FindAdapter()
 					TEXT("Found D3D12 adapter %u: %s (Max supported Feature Level %s, shader model %d.%d)"),
 					AdapterIndex,
 					AdapterDesc.Description,
-					GetFeatureLevelString(MaxSupportedFeatureLevel),
-					(MaxSupportedShaderModel>>4), (MaxSupportedShaderModel & 0xF)
-					);
+					GetFeatureLevelString(DeviceInfo.MaxFeatureLevel),
+					(DeviceInfo.MaxShaderModel >> 4), (DeviceInfo.MaxShaderModel & 0xF)
+				);
 				UE_LOG(LogD3D12RHI, Log,
 					TEXT("Adapter has %uMB of dedicated video memory, %uMB of dedicated system memory, and %uMB of shared system memory, %d output[s]"),
-					(uint32)(AdapterDesc.DedicatedVideoMemory / (1024*1024)),
-					(uint32)(AdapterDesc.DedicatedSystemMemory / (1024*1024)),
-					(uint32)(AdapterDesc.SharedSystemMemory / (1024*1024)),
+					(uint32)(AdapterDesc.DedicatedVideoMemory / (1024 * 1024)),
+					(uint32)(AdapterDesc.DedicatedSystemMemory / (1024 * 1024)),
+					(uint32)(AdapterDesc.SharedSystemMemory / (1024 * 1024)),
 					OutputCount
-					);
-
+				);
 
 				bool bIsAMD = AdapterDesc.VendorId == 0x1002;
 				bool bIsIntel = AdapterDesc.VendorId == 0x8086;
 				bool bIsNVIDIA = AdapterDesc.VendorId == 0x10DE;
 				bool bIsWARP = AdapterDesc.VendorId == 0x1414;
 
-				if (bIsAMD) bIsAnyAMD = true;
-				if (bIsIntel) bIsAnyIntel = true;
-				if (bIsNVIDIA) bIsAnyNVIDIA = true;
-
 				// Simple heuristic but without profiling it's hard to do better
 				bool bIsNonLocalMemoryPresent = false;
 				if (bIsIntel)
@@ -616,8 +637,12 @@ void FD3D12DynamicRHIModule::FindAdapter()
 				// PerfHUD is for performance profiling
 				const bool bIsPerfHUD = !FCString::Stricmp(AdapterDesc.Description, TEXT("NVIDIA PerfHUD"));
 
-				FD3D12AdapterDesc CurrentAdapter(AdapterDesc, AdapterIndex, MaxSupportedFeatureLevel, MaxSupportedShaderModel, NumNodes, bIsIntegrated, GpuPreference);
-				
+				FD3D12AdapterDesc CurrentAdapter(AdapterDesc, AdapterIndex, DeviceInfo.MaxFeatureLevel, DeviceInfo.MaxShaderModel, DeviceInfo.MaxRHIFeatureLevel);
+
+				CurrentAdapter.NumDeviceNodes = DeviceInfo.NumDeviceNodes;
+				CurrentAdapter.GpuPreference = GpuPreference;
+				CurrentAdapter.bIsIntegrated = bIsIntegrated;
+
 				// If requested WARP, then reject all other adapters. If WARP not requested, then reject the WARP device if software rendering support is disallowed
 				const bool bSkipWARP = (bRequestedWARP && !bIsWARP) || (!bRequestedWARP && bIsWARP && !bAllowSoftwareRendering);
 
@@ -1025,6 +1050,8 @@ void FD3D12DynamicRHI::Init()
 		Adapter->InitializeDevices();
 	}
 
+	bool bHasVendorSupportForAtomic64 = false;
+
 #if AMD_API_ENABLE
 	// Warn if we are trying to use RGP frame markers but are either running on a non-AMD device
 	// or using an older AMD driver without RGP marker support
@@ -1041,8 +1068,8 @@ void FD3D12DynamicRHI::Init()
 			    UE_LOG(LogD3D12RHI, Warning, TEXT("Attempting to use RGP frame markers without driver support. Update AMD driver."));
 		    }
 
-			GRHISupportsAtomicUInt64 = AMDSupportedExtensions.intrinsics19 != 0;  // "intrinsics19" includes AtomicU64
-			GRHISupportsAtomicUInt64 = GRHISupportsAtomicUInt64 && AMDSupportedExtensions.UAVBindSlot != 0;
+			bHasVendorSupportForAtomic64 = AMDSupportedExtensions.intrinsics19 != 0;  // "intrinsics19" includes AtomicU64
+			bHasVendorSupportForAtomic64 = bHasVendorSupportForAtomic64 && AMDSupportedExtensions.UAVBindSlot != 0;
 		}
 	}
 	else if (GEmitRgpFrameMarkers)
@@ -1082,7 +1109,7 @@ void FD3D12DynamicRHI::Init()
 		const NvAPI_Status NvStatus = NvAPI_Initialize();
 		if (NvStatus == NVAPI_OK)
 		{
-			NvAPI_Status NvStatusAtomicU64 = NvAPI_D3D12_IsNvShaderExtnOpCodeSupported(GetAdapter().GetD3DDevice(), NV_EXTN_OP_UINT64_ATOMIC, &GRHISupportsAtomicUInt64);
+			NvAPI_Status NvStatusAtomicU64 = NvAPI_D3D12_IsNvShaderExtnOpCodeSupported(GetAdapter().GetD3DDevice(), NV_EXTN_OP_UINT64_ATOMIC, &bHasVendorSupportForAtomic64);
 			if (NvStatusAtomicU64 != NVAPI_OK)
 			{
 				UE_LOG(LogD3D12RHI, Warning, TEXT("Failed to query support for 64 bit atomics"));
@@ -1226,7 +1253,7 @@ void FD3D12DynamicRHI::Init()
 						SupportedExtensionsVersions[i].Revision);
 
 					INTCExtensionInfo.RequestedExtensionVersion = SupportedExtensionsVersions[i];
-					GRHISupportsAtomicUInt64 = true;
+					bHasVendorSupportForAtomic64 = true;
 					break;
 				}
 			}
@@ -1364,6 +1391,12 @@ void FD3D12DynamicRHI::Init()
 	{
 		UE_LOG(LogD3D12RHI, Log, TEXT("RHI has support for 64 bit atomics"));
 	}
+	else if (bHasVendorSupportForAtomic64)
+	{
+		GRHISupportsAtomicUInt64 = true;
+
+		UE_LOG(LogD3D12RHI, Log, TEXT("RHI has vendor support for 64 bit atomics"));
+	}
 	else
 	{
 		UE_LOG(LogD3D12RHI, Log, TEXT("RHI does not have support for 64 bit atomics"));
