Restore MRT optimizations on XB consoles

diff --git a/Engine/Source/Runtime/D3D12RHI/Private/D3D12Resources.cpp b/Engine/Source/Runtime/D3D12RHI/Private/D3D12Resources.cpp
--- a/Engine/Source/Runtime/D3D12RHI/Private/D3D12Resources.cpp
+++ b/Engine/Source/Runtime/D3D12RHI/Private/D3D12Resources.cpp
@@ -577,6 +577,10 @@ HRESULT FD3D12Adapter::CreateCommittedResource(const FD3D12ResourceDesc& InDesc,
 	}
 #endif // D3D12_RHI_RAYTRACING
 
+#if D3D12_WITH_CUSTOM_TEXTURE_LAYOUT
+	ApplyCustomTextureLayout(LocalDesc, *this);
+#endif
+
 	HRESULT hr = S_OK;
 #if INTEL_EXTENSIONS
 	if (InDesc.bRequires64BitAtomicSupport && IsRHIDeviceIntel() && GRHISupportsAtomicUInt64)
diff --git a/Engine/Source/Runtime/D3D12RHI/Public/D3D12Resources.h b/Engine/Source/Runtime/D3D12RHI/Public/D3D12Resources.h
--- a/Engine/Source/Runtime/D3D12RHI/Public/D3D12Resources.h
+++ b/Engine/Source/Runtime/D3D12RHI/Public/D3D12Resources.h
@@ -30,6 +30,7 @@ class FD3D12CommandListHandle;
 class FD3D12SegListAllocator;
 class FD3D12PoolAllocator;
 struct FD3D12GraphicsPipelineState;
+struct FD3D12ResourceDesc;
 typedef FD3D12StateCacheBase FD3D12StateCache;
 
 #if D3D12_RHI_RAYTRACING
@@ -39,6 +40,14 @@ class FD3D12RayTracingPipelineState;
 class FD3D12RayTracingShader;
 #endif // D3D12_RHI_RAYTRACING
 
+#ifndef D3D12_WITH_CUSTOM_TEXTURE_LAYOUT
+#define D3D12_WITH_CUSTOM_TEXTURE_LAYOUT 0
+#endif
+
+#if D3D12_WITH_CUSTOM_TEXTURE_LAYOUT
+extern void ApplyCustomTextureLayout(FD3D12ResourceDesc& TextureLayout, FD3D12Adapter& Adapter);
+#endif
+
 enum class ED3D12ResourceStateMode
 {
 	Default,					//< Decide if tracking is required based on flags
