Fix for build break.

diff --git a/Engine/Source/Runtime/D3D12RHI/Private/D3D12RenderTarget.cpp b/Engine/Source/Runtime/D3D12RHI/Private/D3D12RenderTarget.cpp
--- a/Engine/Source/Runtime/D3D12RHI/Private/D3D12RenderTarget.cpp
+++ b/Engine/Source/Runtime/D3D12RHI/Private/D3D12RenderTarget.cpp
@@ -465,11 +465,11 @@ void FD3D12CommandContext::RHICopyToResolveTarget(FRHITexture* SourceTextureRHI,
 	// Force transtion to readable since RHI user expects both source and dest to be in readable state after calling RHICopyToResolveTarget
 	if (SourceTextureRHI && ResolveParams.SourceAccessFinal != ERHIAccess::Unknown)
 	{
-		RHICmdList.TransitionResource(ResolveParams.SourceAccessFinal, SourceTextureRHI);
+		RHICmdList.Transition(FRHITransitionInfo(SourceTextureRHI, ERHIAccess::Unknown, ResolveParams.SourceAccessFinal));
 	}
 	if (DestTextureRHI && SourceTextureRHI != DestTextureRHI && ResolveParams.DestAccessFinal != ERHIAccess::Unknown)
 	{
-		RHICmdList.TransitionResource(ResolveParams.DestAccessFinal, DestTextureRHI);
+		RHICmdList.Transition(FRHITransitionInfo(DestTextureRHI, ERHIAccess::Unknown, ResolveParams.DestAccessFinal));
 	}
 
 	ConditionalFlushCommandList();
