Reduce redundant D3D12 PSOs by ignoring vertex stream strides when computing the PSO cache key.

D3D12_INPUT_LAYOUT_DESC part of D3D12 graphics PSO descriptor only includes the format, slot/semantic index, etc.
The stride is set at draw time, via ID3D12GraphicsCommandList::IASetVertexBuffers().

This saves a small percentage of total created low-level PSOs.

diff --git a/Engine/Source/Runtime/D3D12RHI/Private/D3D12Pipelinestate.cpp b/Engine/Source/Runtime/D3D12RHI/Private/D3D12Pipelinestate.cpp
--- a/Engine/Source/Runtime/D3D12RHI/Private/D3D12Pipelinestate.cpp
+++ b/Engine/Source/Runtime/D3D12RHI/Private/D3D12Pipelinestate.cpp
@@ -46,7 +46,7 @@ FD3D12LowLevelGraphicsPipelineStateDesc GetLowLevelGraphicsPipelineStateDesc(con
 	{
 		Desc.Desc.InputLayout.NumElements        = InputLayout->VertexElements.Num();
 		Desc.Desc.InputLayout.pInputElementDescs = InputLayout->VertexElements.GetData();
-		Desc.InputLayoutHash = InputLayout->Hash;
+		Desc.InputLayoutHash = InputLayout->HashNoStrides; // Vertex stream stride does not affect the D3D12 PSO
 	}
 
 #define COPY_SHADER(Initial, Name) \
diff --git a/Engine/Source/Runtime/D3D12RHI/Private/D3D12Shader.h b/Engine/Source/Runtime/D3D12RHI/Private/D3D12Shader.h
--- a/Engine/Source/Runtime/D3D12RHI/Private/D3D12Shader.h
+++ b/Engine/Source/Runtime/D3D12RHI/Private/D3D12Shader.h
@@ -24,11 +24,14 @@ public:
 
 	uint16 StreamStrides[MaxVertexElementCount];
 	uint32 Hash;
+	uint32 HashNoStrides;
+
 
 	/** Initialization constructor. */
-	explicit FD3D12VertexDeclaration(const FD3D12VertexElements& InElements, const uint16* InStrides, const uint32 InHash)
+	explicit FD3D12VertexDeclaration(const FD3D12VertexElements& InElements, const uint16* InStrides, const uint32 InHash, const uint32 InHashNoStrides)
 		: VertexElements(InElements)
 		, Hash(InHash)
+		, HashNoStrides(InHashNoStrides)
 	{
 		FMemory::Memcpy(StreamStrides, InStrides, sizeof(StreamStrides));
 	}
diff --git a/Engine/Source/Runtime/D3D12RHI/Private/D3D12VertexDeclaration.cpp b/Engine/Source/Runtime/D3D12RHI/Private/D3D12VertexDeclaration.cpp
--- a/Engine/Source/Runtime/D3D12RHI/Private/D3D12VertexDeclaration.cpp
+++ b/Engine/Source/Runtime/D3D12RHI/Private/D3D12VertexDeclaration.cpp
@@ -16,6 +16,9 @@ struct FD3D12VertexDeclarationKey
 	/** Hash of the vertex elements. */
 	uint32 Hash;
 
+	/** Hash of the vertex elements, without strides. */
+	uint32 HashNoStrides;
+
 	uint16 StreamStrides[MaxVertexElementCount];
 
 	/** Initialization constructor. */
@@ -88,6 +91,8 @@ struct FD3D12VertexDeclarationKey
 
 		// Hash once.
 		Hash = FCrc::MemCrc_DEPRECATED(VertexElements.GetData(), VertexElements.Num()*sizeof(D3D12_INPUT_ELEMENT_DESC));
+		HashNoStrides = Hash;
+
 		Hash = FCrc::MemCrc_DEPRECATED(StreamStrides, sizeof(StreamStrides), Hash);
 
 		// Assign all the SemanticName after hashing. It's a constant string, always the same, so no need to hash the data.
@@ -133,7 +138,7 @@ struct FVertexDeclarationCache
 		FVertexDeclarationRHIRef* VertexDeclarationRefPtr = Cache.Find(InKey);
 		if (VertexDeclarationRefPtr == nullptr)
 		{
-			VertexDeclarationRefPtr = &Cache.Add(InKey, new FD3D12VertexDeclaration(InKey.VertexElements, InKey.StreamStrides, InKey.Hash));
+			VertexDeclarationRefPtr = &Cache.Add(InKey, new FD3D12VertexDeclaration(InKey.VertexElements, InKey.StreamStrides, InKey.Hash, InKey.HashNoStrides));
 		}
 
 		return VertexDeclarationRefPtr;
