Adding GRHISupportsDynamicResources.
D3D12 shader loading will validate this flag is set when loading shaders using dynamic resources.
Moving around D3D12 detection of resource tiers.
Cleaning up D3D12 RHI feature level argument handling; SM6 is a fully supported FeatureLevel that is validated before creating the RHI.
Restricting HoloLens to ES3.1 per team's suggestion.

diff --git a/Engine/Source/Runtime/D3D12RHI/Private/D3D12Adapter.cpp b/Engine/Source/Runtime/D3D12RHI/Private/D3D12Adapter.cpp
--- a/Engine/Source/Runtime/D3D12RHI/Private/D3D12Adapter.cpp
+++ b/Engine/Source/Runtime/D3D12RHI/Private/D3D12Adapter.cpp
@@ -204,12 +204,14 @@ static LONG __stdcall D3DVectoredExceptionHandler(EXCEPTION_POINTERS* InInfo)
 
 FD3D12AdapterDesc::FD3D12AdapterDesc() = default;
 
-FD3D12AdapterDesc::FD3D12AdapterDesc(const DXGI_ADAPTER_DESC& InDesc, int32 InAdapterIndex, D3D_FEATURE_LEVEL InMaxSupportedFeatureLevel, D3D_SHADER_MODEL InMaxSupportedShaderModel, ERHIFeatureLevel::Type InMaxRHIFeatureLevel)
+FD3D12AdapterDesc::FD3D12AdapterDesc(const DXGI_ADAPTER_DESC& InDesc, int32 InAdapterIndex, const FD3D12DeviceBasicInfo& DeviceInfo)
 	: Desc(InDesc)
 	, AdapterIndex(InAdapterIndex)
-	, MaxSupportedFeatureLevel(InMaxSupportedFeatureLevel)
-	, MaxSupportedShaderModel(InMaxSupportedShaderModel)
-	, MaxRHIFeatureLevel(InMaxRHIFeatureLevel)
+	, MaxSupportedFeatureLevel(DeviceInfo.MaxFeatureLevel)
+	, MaxSupportedShaderModel(DeviceInfo.MaxShaderModel)
+	, ResourceBindingTier(DeviceInfo.ResourceBindingTier)
+	, ResourceHeapTier(DeviceInfo.ResourceHeapTier)
+	, MaxRHIFeatureLevel(DeviceInfo.MaxRHIFeatureLevel)
 {
 }
 
@@ -898,21 +900,15 @@ void FD3D12Adapter::InitializeDevices()
 
 			const bool bRenderDocPresent = D3D12RHI_IsRenderDocPresent(RootDevice);
 
-			D3D12_FEATURE_DATA_D3D12_OPTIONS D3D12Caps;
-			FMemory::Memzero(&D3D12Caps, sizeof(D3D12Caps));
-			VERIFYD3D12RESULT(RootDevice->CheckFeatureSupport(D3D12_FEATURE_D3D12_OPTIONS, &D3D12Caps, sizeof(D3D12Caps)));
-			ResourceHeapTier = D3D12Caps.ResourceHeapTier;
-			ResourceBindingTier = D3D12Caps.ResourceBindingTier;
-
-			if (ResourceBindingTier == D3D12_RESOURCE_BINDING_TIER_1)
+			if (GetResourceBindingTier() == D3D12_RESOURCE_BINDING_TIER_1)
 			{
 				MaxNonSamplerDescriptors = D3D12_MAX_SHADER_VISIBLE_DESCRIPTOR_HEAP_SIZE_TIER_1;
 			}
-			else if (ResourceBindingTier == D3D12_RESOURCE_BINDING_TIER_2)
+			else if (GetResourceBindingTier() == D3D12_RESOURCE_BINDING_TIER_2)
 			{
 				MaxNonSamplerDescriptors = D3D12_MAX_SHADER_VISIBLE_DESCRIPTOR_HEAP_SIZE_TIER_2;
 			}
-			else if (ResourceBindingTier == D3D12_RESOURCE_BINDING_TIER_3)
+			else if (GetResourceBindingTier() == D3D12_RESOURCE_BINDING_TIER_3)
 			{
 				// From: https://microsoft.github.io/DirectX-Specs/d3d/ResourceBinding.html#levels-of-hardware-support
 				//   For Tier 3, the max # descriptors is listed as 1000000+. The + indicates that the runtime allows applications
@@ -980,12 +976,20 @@ void FD3D12Adapter::InitializeDevices()
 
 			MaxSamplerDescriptors = D3D12_MAX_SHADER_VISIBLE_SAMPLER_HEAP_SIZE;
 
+			// From: https://microsoft.github.io/DirectX-Specs/d3d/HLSL_SM_6_6_DynamicResources.html
+			//     ResourceDescriptorHeap/SamplerDescriptorHeap must be supported on devices that support both D3D12_RESOURCE_BINDING_TIER_3 and D3D_SHADER_MODEL_6_6
+			if (GetHighestShaderModel() >= D3D_SHADER_MODEL_6_6 && GetResourceBindingTier() >= D3D12_RESOURCE_BINDING_TIER_3)
+			{
+				GRHISupportsBindless = true;
+				UE_LOG(LogD3D12RHI, Log, TEXT("Bindless resources are supported"));
+			}
+
 #if D3D12_RHI_RAYTRACING
 			D3D12_FEATURE_DATA_D3D12_OPTIONS5 D3D12Caps5 = {};
 			if (SUCCEEDED(RootDevice->CheckFeatureSupport(D3D12_FEATURE_D3D12_OPTIONS5, &D3D12Caps5, sizeof(D3D12Caps5))))
 			{
 				if (D3D12Caps5.RaytracingTier >= D3D12_RAYTRACING_TIER_1_0
-					&& D3D12Caps.ResourceBindingTier >= D3D12_RESOURCE_BINDING_TIER_2
+					&& GetResourceBindingTier() >= D3D12_RESOURCE_BINDING_TIER_2
 					&& RootDevice5
 					&& FDataDrivenShaderPlatformInfo::GetSupportsRayTracing(GMaxRHIShaderPlatform)
 					&& !FParse::Param(FCommandLine::Get(), TEXT("noraytracing")))
@@ -1019,7 +1023,7 @@ void FD3D12Adapter::InitializeDevices()
 
 #endif // D3D12_RHI_RAYTRACING
 
-#if PLATFORM_WINDOWS && D3D12_CORE_ENABLED
+#if PLATFORM_WINDOWS
 			{
 				D3D12_FEATURE_DATA_D3D12_OPTIONS7 D3D12Caps7 = {};
 				RootDevice->CheckFeatureSupport(D3D12_FEATURE_D3D12_OPTIONS7, &D3D12Caps7, sizeof(D3D12Caps7));
@@ -1061,7 +1065,7 @@ void FD3D12Adapter::InitializeDevices()
 					UE_LOG(LogD3D12RHI, Log, TEXT("Shader Model 6.6 atomic64 is not supported"));
 				}
 			}
-#endif // PLATFORM_WINDOWS && D3D12_CORE_ENABLED
+#endif // PLATFORM_WINDOWS
 		}
 
 #if PLATFORM_WINDOWS || PLATFORM_HOLOLENS
@@ -1091,7 +1095,7 @@ void FD3D12Adapter::InitializeDevices()
 #if TRACK_RESOURCE_ALLOCATIONS
 		// Set flag if we want to track all allocations - comes with some overhead and only possible when Tier 2 is available
 		// (because we will create placed buffers for texture allocation to retrieve the GPU virtual addresses)
-		bTrackAllAllocation = (GD3D12TrackAllAlocations || GPUCrashDebuggingModes == ED3D12GPUCrashDebuggingModes::All) && (ResourceHeapTier == D3D12_RESOURCE_HEAP_TIER_2);
+		bTrackAllAllocation = (GD3D12TrackAllAlocations || GPUCrashDebuggingModes == ED3D12GPUCrashDebuggingModes::All) && (GetResourceHeapTier() == D3D12_RESOURCE_HEAP_TIER_2);
 #endif 
 
 		CreateCommandSignatures();
diff --git a/Engine/Source/Runtime/D3D12RHI/Private/D3D12Adapter.h b/Engine/Source/Runtime/D3D12RHI/Private/D3D12Adapter.h
--- a/Engine/Source/Runtime/D3D12RHI/Private/D3D12Adapter.h
+++ b/Engine/Source/Runtime/D3D12RHI/Private/D3D12Adapter.h
@@ -37,10 +37,21 @@ This structure allows a single RHI to control several different hardware setups.
 
 class FD3D12DynamicRHI;
 
+struct FD3D12DeviceBasicInfo
+{
+	D3D_FEATURE_LEVEL           MaxFeatureLevel;
+	D3D_SHADER_MODEL            MaxShaderModel;
+	D3D12_RESOURCE_BINDING_TIER ResourceBindingTier;
+	D3D12_RESOURCE_HEAP_TIER    ResourceHeapTier;
+	uint32                      NumDeviceNodes;
+
+	ERHIFeatureLevel::Type      MaxRHIFeatureLevel;
+};
+
 struct FD3D12AdapterDesc
 {
 	FD3D12AdapterDesc();
-	FD3D12AdapterDesc(const DXGI_ADAPTER_DESC& InDesc, int32 InAdapterIndex, D3D_FEATURE_LEVEL InMaxSupportedFeatureLevel, D3D_SHADER_MODEL InMaxSupportedShaderModel, ERHIFeatureLevel::Type InMaxRHIFeatureLevel);
+	FD3D12AdapterDesc(const DXGI_ADAPTER_DESC& InDesc, int32 InAdapterIndex, const FD3D12DeviceBasicInfo& DeviceInfo);
 
 	bool IsValid() const;
 
@@ -60,6 +71,10 @@ struct FD3D12AdapterDesc
 	/** The maximum Shader Model supported. 0 if not supported or FindAdpater() wasn't called */
 	D3D_SHADER_MODEL MaxSupportedShaderModel = (D3D_SHADER_MODEL)0;
 
+	D3D12_RESOURCE_BINDING_TIER ResourceBindingTier = D3D12_RESOURCE_BINDING_TIER_1;
+
+	D3D12_RESOURCE_HEAP_TIER ResourceHeapTier = D3D12_RESOURCE_HEAP_TIER_1;
+
 	ERHIFeatureLevel::Type MaxRHIFeatureLevel = ERHIFeatureLevel::Num;
 
 	/** Number of device nodes (read: GPUs) */
@@ -136,8 +151,8 @@ public:
 	FORCEINLINE const bool IsDebugDevice() const { return bDebugDevice; }
 	FORCEINLINE const ED3D12GPUCrashDebuggingModes GetGPUCrashDebuggingModes() const { return GPUCrashDebuggingModes; }
 	FORCEINLINE FD3D12DynamicRHI* GetOwningRHI() { return OwningRHI; }
-	FORCEINLINE const D3D12_RESOURCE_HEAP_TIER GetResourceHeapTier() const { return ResourceHeapTier; }
-	FORCEINLINE const D3D12_RESOURCE_BINDING_TIER GetResourceBindingTier() const { return ResourceBindingTier; }
+	FORCEINLINE const D3D12_RESOURCE_HEAP_TIER GetResourceHeapTier() const { return Desc.ResourceHeapTier; }
+	FORCEINLINE const D3D12_RESOURCE_BINDING_TIER GetResourceBindingTier() const { return Desc.ResourceBindingTier; }
 	FORCEINLINE const D3D_ROOT_SIGNATURE_VERSION GetRootSignatureVersion() const { return RootSignatureVersion; }
 	FORCEINLINE const bool IsDepthBoundsTestSupported() const { return bDepthBoundsTestSupported; }
 	FORCEINLINE const bool IsHeapNotZeroedSupported() const { return bHeapNotZeroedSupported; }
@@ -449,8 +464,6 @@ protected:
 	HANDLE ExceptionHandlerHandle = INVALID_HANDLE_VALUE;
 #endif
 
-	D3D12_RESOURCE_HEAP_TIER ResourceHeapTier;
-	D3D12_RESOURCE_BINDING_TIER ResourceBindingTier;
 	D3D_ROOT_SIGNATURE_VERSION RootSignatureVersion;
 	bool bDepthBoundsTestSupported;
 	bool bHeapNotZeroedSupported;
diff --git a/Engine/Source/Runtime/D3D12RHI/Private/D3D12Shaders.cpp b/Engine/Source/Runtime/D3D12RHI/Private/D3D12Shaders.cpp
--- a/Engine/Source/Runtime/D3D12RHI/Private/D3D12Shaders.cpp
+++ b/Engine/Source/Runtime/D3D12RHI/Private/D3D12Shaders.cpp
@@ -61,6 +61,11 @@ static bool ValidateShaderIsUsable(FD3D12ShaderData* InShader)
 	{
 		return false;
 	}
+
+	if (EnumHasAnyFlags(InShader->Features, EShaderCodeFeatures::BindlessResources | EShaderCodeFeatures::BindlessSamplers) && !GRHISupportsBindless)
+	{
+		return false;
+	}
 #endif
 
 	return true;
diff --git a/Engine/Source/Runtime/D3D12RHI/Private/Windows/WindowsD3D12Device.cpp b/Engine/Source/Runtime/D3D12RHI/Private/Windows/WindowsD3D12Device.cpp
--- a/Engine/Source/Runtime/D3D12RHI/Private/Windows/WindowsD3D12Device.cpp
+++ b/Engine/Source/Runtime/D3D12RHI/Private/Windows/WindowsD3D12Device.cpp
@@ -141,14 +141,6 @@ static FAutoConsoleVariableRef CVarAllowAsyncCompute(
 	ECVF_ReadOnly | ECVF_RenderThreadSafe
 );
 
-int32 GAllowShaderModel6 = 0;
-static FAutoConsoleVariableRef CVarAllowShaderModel6(
-	TEXT("r.D3D12.AllowShaderModel6"),
-	GAllowShaderModel6,
-	TEXT("Allows the usage of SM6 feature level."),
-	ECVF_ReadOnly | ECVF_RenderThreadSafe
-);
-
 #if !UE_BUILD_SHIPPING
 static TAutoConsoleVariable<int32> CVarExperimentalShaderModels(
 	TEXT("r.D3D12.ExperimentalShaderModels"),
@@ -326,14 +318,14 @@ inline ERHIFeatureLevel::Type FindMaxRHIFeatureLevel(ID3D12Device* Device, D3D_F
 	return MaxRHIFeatureLevel;
 }
 
-struct FD3D12DeviceBasicInfo
+inline void GetResourceTiers(ID3D12Device* Device, D3D12_RESOURCE_BINDING_TIER& OutResourceBindingTier, D3D12_RESOURCE_HEAP_TIER& OutResourceHeapTier)
 {
-	D3D_FEATURE_LEVEL MaxFeatureLevel;
-	D3D_SHADER_MODEL  MaxShaderModel;
-	uint32            NumDeviceNodes;
+	D3D12_FEATURE_DATA_D3D12_OPTIONS D3D12Caps{};
+	Device->CheckFeatureSupport(D3D12_FEATURE_D3D12_OPTIONS, &D3D12Caps, sizeof(D3D12Caps));
 
-	ERHIFeatureLevel::Type MaxRHIFeatureLevel;
-};
+	OutResourceBindingTier = D3D12Caps.ResourceBindingTier;
+	OutResourceHeapTier = D3D12Caps.ResourceHeapTier;
+}
 
 /**
  * Attempts to create a D3D12 device for the adapter using at minimum MinFeatureLevel.
@@ -348,6 +340,7 @@ static bool SafeTestD3D12CreateDevice(IDXGIAdapter* Adapter, D3D_FEATURE_LEVEL M
 		{
 			OutInfo.MaxFeatureLevel = FindHighestFeatureLevel(Device, MinFeatureLevel);
 			OutInfo.MaxShaderModel = FindHighestShaderModel(Device);
+			GetResourceTiers(Device, OutInfo.ResourceBindingTier, OutInfo.ResourceHeapTier);
 			OutInfo.NumDeviceNodes = Device->GetNodeCount();
 			OutInfo.MaxRHIFeatureLevel = FindMaxRHIFeatureLevel(Device, OutInfo.MaxFeatureLevel, OutInfo.MaxShaderModel);
 
@@ -470,7 +463,6 @@ bool FD3D12DynamicRHIModule::IsSupported(ERHIFeatureLevel::Type RequestedFeature
 		FindAdapter();
 	}
 
-	// The hardware must support at least 11.0.
 	return ChosenAdapters.Num() > 0
 		&& !IsAdapterBlocked(ChosenAdapters[0].Get())
 		&& IsAdapterSupported(ChosenAdapters[0].Get(), RequestedFeatureLevel);
@@ -637,7 +629,7 @@ void FD3D12DynamicRHIModule::FindAdapter()
 				// PerfHUD is for performance profiling
 				const bool bIsPerfHUD = !FCString::Stricmp(AdapterDesc.Description, TEXT("NVIDIA PerfHUD"));
 
-				FD3D12AdapterDesc CurrentAdapter(AdapterDesc, AdapterIndex, DeviceInfo.MaxFeatureLevel, DeviceInfo.MaxShaderModel, DeviceInfo.MaxRHIFeatureLevel);
+				FD3D12AdapterDesc CurrentAdapter(AdapterDesc, AdapterIndex, DeviceInfo);
 
 				CurrentAdapter.NumDeviceNodes = DeviceInfo.NumDeviceNodes;
 				CurrentAdapter.GpuPreference = GpuPreference;
@@ -739,23 +731,20 @@ static bool DoesAnyAdapterSupportSM6(const TArray<TSharedPtr<FD3D12Adapter>>& Ad
 FDynamicRHI* FD3D12DynamicRHIModule::CreateRHI(ERHIFeatureLevel::Type RequestedFeatureLevel)
 {
 #if PLATFORM_HOLOLENS
+	check(RequestedFeatureLevel == ERHIFeatureLevel::ES3_1);
+
 	GMaxRHIFeatureLevel = ERHIFeatureLevel::ES3_1;
 	GMaxRHIShaderPlatform = SP_D3D_ES3_1_HOLOLENS;
-#endif
 
-	const bool bAllowSM6 = GAllowShaderModel6
-		|| (RequestedFeatureLevel != ERHIFeatureLevel::Num && RequestedFeatureLevel >= ERHIFeatureLevel::SM6);
-
-#if PLATFORM_HOLOLENS
 	GShaderPlatformForFeatureLevel[ERHIFeatureLevel::ES3_1] = SP_D3D_ES3_1_HOLOLENS;
 #else
 	GShaderPlatformForFeatureLevel[ERHIFeatureLevel::ES3_1] = SP_PCD3D_ES3_1;
-#endif
 	GShaderPlatformForFeatureLevel[ERHIFeatureLevel::SM5] = SP_PCD3D_SM5;
-	if (bAllowSM6)
+	if (DoesAnyAdapterSupportSM6(ChosenAdapters))
 	{
 		GShaderPlatformForFeatureLevel[ERHIFeatureLevel::SM6] = SP_PCD3D_SM6;
 	}
+#endif
 
 	ERHIFeatureLevel::Type PreviewFeatureLevel;
 	if (!GIsEditor && RHIGetPreviewFeatureLevel(PreviewFeatureLevel))
@@ -767,21 +756,14 @@ FDynamicRHI* FD3D12DynamicRHIModule::CreateRHI(ERHIFeatureLevel::Type RequestedF
 	}
 	else
 	{
-		GMaxRHIFeatureLevel = ERHIFeatureLevel::SM5;
-		if (RequestedFeatureLevel < ERHIFeatureLevel::SM6)	
-		{
-			GMaxRHIFeatureLevel = RequestedFeatureLevel;
-		}
-		else if (bAllowSM6 && DoesAnyAdapterSupportSM6(ChosenAdapters))
-		{
-			GMaxRHIFeatureLevel = ERHIFeatureLevel::SM6;
-		}
+		GMaxRHIFeatureLevel = RequestedFeatureLevel;
 	}
 
 	if (!ensure(GMaxRHIFeatureLevel < ERHIFeatureLevel::Num))
 	{
 		GMaxRHIFeatureLevel = ERHIFeatureLevel::SM5;
 	}
+
 	GMaxRHIShaderPlatform = GShaderPlatformForFeatureLevel[GMaxRHIFeatureLevel];
 	check(GMaxRHIShaderPlatform != SP_NumPlatforms);
 
diff --git a/Engine/Source/Runtime/RHI/Private/HoloLens/HoloLensDynamicRHI.cpp b/Engine/Source/Runtime/RHI/Private/HoloLens/HoloLensDynamicRHI.cpp
--- a/Engine/Source/Runtime/RHI/Private/HoloLens/HoloLensDynamicRHI.cpp
+++ b/Engine/Source/Runtime/RHI/Private/HoloLens/HoloLensDynamicRHI.cpp
@@ -109,7 +109,7 @@ FDynamicRHI* PlatformCreateDynamicRHI()
 	}
 	else
 	{
-		DynamicRHI = DynamicRHIModule->CreateRHI();
+		DynamicRHI = DynamicRHIModule->CreateRHI(ERHIFeatureLevel::ES3_1);
 	}
 
 	GMaxRHIFeatureLevel = ERHIFeatureLevel::ES3_1;
diff --git a/Engine/Source/Runtime/RHI/Private/RHI.cpp b/Engine/Source/Runtime/RHI/Private/RHI.cpp
--- a/Engine/Source/Runtime/RHI/Private/RHI.cpp
+++ b/Engine/Source/Runtime/RHI/Private/RHI.cpp
@@ -1477,6 +1477,8 @@ bool GRHISupportsAsyncPipelinePrecompile = true;
 
 bool GRHISupportsMapWriteNoOverwrite = false;
 
+bool GRHISupportsBindless = false;
+
 FVertexElementTypeSupportInfo GVertexElementTypeSupport;
 
 RHI_API int32 volatile GCurrentTextureMemorySize = 0;
diff --git a/Engine/Source/Runtime/RHI/Public/RHI.h b/Engine/Source/Runtime/RHI/Public/RHI.h
--- a/Engine/Source/Runtime/RHI/Public/RHI.h
+++ b/Engine/Source/Runtime/RHI/Public/RHI.h
@@ -771,6 +771,9 @@ extern RHI_API FVector2f GRHIDefaultMSAASampleOffsets[1 + 2 + 4 + 8 + 16];
 /** True if the RHI supports pipeline precompiling from any thread. */
 extern RHI_API bool GRHISupportsAsyncPipelinePrecompile;
 
+/** Whether dynamic (bindless) resources are supported */
+extern RHI_API bool GRHISupportsBindless;
+
 // Calculate the index of the sample in GRHIDefaultMSAASampleOffsets
 extern RHI_API int32 CalculateMSAASampleArrayIndex(int32 NumSamples, int32 SampleIndex);
 
