Add support for creating raw buffer SRVs for any buffer in D3D12

In some RHIs, a raw buffer view can only be created if the buffer itself was originally created with a corresponding flag. However, modern RHIs such as D3D12 allow creating raw views for any buffer.
This changelist adds FRawBufferShaderResourceViewInitializer to create raw views of buffers when GRHISupportsRawViewsForAnyBuffer is true.

diff --git a/Engine/Source/Runtime/D3D12RHI/Private/D3D12RHI.cpp b/Engine/Source/Runtime/D3D12RHI/Private/D3D12RHI.cpp
--- a/Engine/Source/Runtime/D3D12RHI/Private/D3D12RHI.cpp
+++ b/Engine/Source/Runtime/D3D12RHI/Private/D3D12RHI.cpp
@@ -208,6 +208,8 @@ FD3D12DynamicRHI::FD3D12DynamicRHI(const TArray<TSharedPtr<FD3D12Adapter>>& Chos
 
 	GRHISupportsParallelRHIExecute = D3D12_SUPPORTS_PARALLEL_RHI_EXECUTE;
 
+	GRHISupportsRawViewsForAnyBuffer = true;
+
 	GSupportsTimestampRenderQueries = true;
 	GSupportsParallelOcclusionQueries = true;
 
diff --git a/Engine/Source/Runtime/D3D12RHI/Private/D3D12View.cpp b/Engine/Source/Runtime/D3D12RHI/Private/D3D12View.cpp
--- a/Engine/Source/Runtime/D3D12RHI/Private/D3D12View.cpp
+++ b/Engine/Source/Runtime/D3D12RHI/Private/D3D12View.cpp
@@ -239,6 +239,37 @@ void FD3D12ShaderResourceView::CreateView(FD3D12BaseShaderResource* InBaseShader
 	}
 }
 
+static D3D12_SHADER_RESOURCE_VIEW_DESC GetRawBufferSRVDesc(FD3D12Buffer* Buffer, uint32 StartOffsetBytes, uint32 NumElements)
+{
+	const uint32 Stride = 4;
+	checkf(StartOffsetBytes % Stride == 0, TEXT("Raw buffer offset must be DWORD-aligned"));
+
+	const uint32 Width = Buffer->GetSize();
+	const uint32 MaxElements = Width / Stride;
+	const uint32 StartElement = FMath::Min(StartOffsetBytes, Width) / Stride;
+	const FD3D12ResourceLocation& Location = Buffer->ResourceLocation;
+
+	D3D12_SHADER_RESOURCE_VIEW_DESC SRVDesc{};
+	SRVDesc.Shader4ComponentMapping = D3D12_DEFAULT_SHADER_4_COMPONENT_MAPPING;
+	SRVDesc.ViewDimension = D3D12_SRV_DIMENSION_BUFFER;
+	SRVDesc.Format = DXGI_FORMAT_R32_TYPELESS;
+	SRVDesc.Buffer.Flags = D3D12_BUFFER_SRV_FLAG_RAW;
+	SRVDesc.Buffer.NumElements = FMath::Min(MaxElements - StartElement, NumElements);
+
+	if (Location.GetResource())
+	{
+		// Create a Shader Resource View
+		SRVDesc.Buffer.FirstElement = Location.GetOffsetFromBaseOfResource() / 4 + StartElement;
+	}
+	else
+	{
+		// Null underlying D3D12 resource should only be the case for dynamic resources
+		check(EnumHasAnyFlags(Buffer->GetUsage(), BUF_AnyDynamic));
+	}
+
+	return SRVDesc;
+}
+
 static FORCEINLINE D3D12_SHADER_RESOURCE_VIEW_DESC GetVertexBufferSRVDesc(const FD3D12Buffer* VertexBuffer, uint32& CreationStride, uint8 Format, uint32 StartOffsetBytes, uint32 NumElements)
 {
 	const uint32 BufferSize = VertexBuffer->GetSize();
@@ -549,6 +580,16 @@ FShaderResourceViewRHIRef FD3D12DynamicRHI::RHICreateShaderResourceView(const FS
 					return ShaderResourceView;
 				});
 
+		case FShaderResourceViewInitializer::EType::RawBufferSRV:
+			return GetAdapter().CreateLinkedViews<FD3D12Buffer, FD3D12ShaderResourceView>(Buffer,
+				[Desc](FD3D12Buffer* Buffer)
+				{
+					check(Buffer);
+					const D3D12_SHADER_RESOURCE_VIEW_DESC SRVDesc = GetRawBufferSRVDesc(Buffer, Desc.StartOffsetBytes, Desc.NumElements);
+					FD3D12ShaderResourceView* ShaderResourceView = new FD3D12ShaderResourceView(Buffer, SRVDesc, 4);
+					return ShaderResourceView;
+				});
+
 	#if D3D12_RHI_RAYTRACING
 		case FShaderResourceViewInitializer::EType::AccelerationStructureSRV:
 			return GetAdapter().CreateLinkedViews<FD3D12Buffer, FD3D12ShaderResourceView>(Buffer,
diff --git a/Engine/Source/Runtime/RHI/Private/DynamicRHI.cpp b/Engine/Source/Runtime/RHI/Private/DynamicRHI.cpp
--- a/Engine/Source/Runtime/RHI/Private/DynamicRHI.cpp
+++ b/Engine/Source/Runtime/RHI/Private/DynamicRHI.cpp
@@ -634,6 +634,19 @@ FShaderResourceViewInitializer::FShaderResourceViewInitializer(FRHIBuffer* InBuf
 	InitType();
 }
 
+FRawBufferShaderResourceViewInitializer::FRawBufferShaderResourceViewInitializer(FRHIBuffer* InBuffer)
+	: FShaderResourceViewInitializer(nullptr)
+{
+	check(GRHISupportsRawViewsForAnyBuffer);
+
+	Type = EType::RawBufferSRV;
+
+	BufferInitializer.Buffer = InBuffer;
+	BufferInitializer.Format = PF_Unknown;
+	BufferInitializer.StartOffsetBytes = 0;
+	BufferInitializer.NumElements = UINT32_MAX; // Whole resource
+}
+
 void FShaderResourceViewInitializer::InitType()
 {
 	if (BufferInitializer.Buffer)
diff --git a/Engine/Source/Runtime/RHI/Private/RHI.cpp b/Engine/Source/Runtime/RHI/Private/RHI.cpp
--- a/Engine/Source/Runtime/RHI/Private/RHI.cpp
+++ b/Engine/Source/Runtime/RHI/Private/RHI.cpp
@@ -1385,6 +1385,7 @@ TRHIGlobal<int64> GMaxComputeSharedMemory(1<<15);
 TRHIGlobal<int32> GMaxVolumeTextureDimensions(2048);
 TRHIGlobal<int32> GMaxCubeTextureDimensions(2048);
 TRHIGlobal<int32> GMaxWorkGroupInvocations(1024);
+bool GRHISupportsRawViewsForAnyBuffer = false;
 bool GRHISupportsRWTextureBuffers = true;
 bool GRHISupportsVRS = false;
 bool GRHISupportsLateVRSUpdate = false;
diff --git a/Engine/Source/Runtime/RHI/Public/DynamicRHI.h b/Engine/Source/Runtime/RHI/Public/DynamicRHI.h
--- a/Engine/Source/Runtime/RHI/Public/DynamicRHI.h
+++ b/Engine/Source/Runtime/RHI/Public/DynamicRHI.h
@@ -82,6 +82,8 @@ struct FRHIFlipDetails
 
 struct FShaderResourceViewInitializer
 {
+	friend struct FRawBufferShaderResourceViewInitializer;
+
 	struct FBufferShaderResourceViewInitializer
 	{
 		FRHIBuffer* Buffer;
@@ -134,6 +136,7 @@ struct FShaderResourceViewInitializer
 		StructuredBufferSRV,
 		IndexBufferSRV,
 		AccelerationStructureSRV,
+		RawBufferSRV, 
 	};
 
 	const EType GetType() const
@@ -149,6 +152,17 @@ private:
 	EType Type;
 };
 
+
+/*
+* FRawBufferShaderResourceViewInitializer can be used to explicitly create a raw view for any buffer,
+* even if it was not created with EBufferUsageFlags::ByteAddressBuffer flag.
+* Can only be used if GRHISupportsRawViewsForAnyBuffer is set.
+*/
+struct FRawBufferShaderResourceViewInitializer : FShaderResourceViewInitializer
+{
+	RHI_API FRawBufferShaderResourceViewInitializer(FRHIBuffer* InBuffer);
+};
+
 class FDynamicRHI;
 
 class RHI_API FDefaultRHIRenderQueryPool final : public FRHIRenderQueryPool
diff --git a/Engine/Source/Runtime/RHI/Public/RHI.h b/Engine/Source/Runtime/RHI/Public/RHI.h
--- a/Engine/Source/Runtime/RHI/Public/RHI.h
+++ b/Engine/Source/Runtime/RHI/Public/RHI.h
@@ -497,6 +497,9 @@ extern RHI_API TRHIGlobal<int32> GMaxVolumeTextureDimensions;
 /** Whether RW texture buffers are supported */
 extern RHI_API bool GRHISupportsRWTextureBuffers;
 
+/** Whether a raw (ByteAddress) buffer view can be created for any buffer, regardless of its EBufferUsageFlags::ByteAddressBuffer flag. */
+extern RHI_API bool GRHISupportsRawViewsForAnyBuffer;
+
 FORCEINLINE uint64 GetMaxBufferDimension()
 {
 	return GMaxBufferDimensions;
