Fix for double signal of gpu fence in DX12 when performing All -> Graphics transition.

diff --git a/Engine/Source/Runtime/D3D12RHI/Private/D3D12CommandContext.cpp b/Engine/Source/Runtime/D3D12RHI/Private/D3D12CommandContext.cpp
--- a/Engine/Source/Runtime/D3D12RHI/Private/D3D12CommandContext.cpp
+++ b/Engine/Source/Runtime/D3D12RHI/Private/D3D12CommandContext.cpp
@@ -677,18 +677,27 @@ void FD3D12CommandContextBase::RHIEndFrame()
 void FD3D12CommandContextBase::SignalTransitionFences(TArrayView<const FRHITransition*> Transitions)
 {
 	bool bSubmitted = false;
+
+	const ERHIPipeline SourcePipeline = IsAsyncComputeContext() ? ERHIPipeline::AsyncCompute : ERHIPipeline::Graphics;
+
 	for (const FRHITransition* Transition : Transitions)
 	{
 		const auto* Data = Transition->GetPrivateData<FD3D12TransitionData>();
 		const auto& Fence = Data->Fence;
-		if (Fence)
+
+		if (!Fence)
+		{
+			continue;
+		}
+
+		if (Data->SrcPipelines != ERHIPipeline::All || !EnumHasAnyFlags(Data->DstPipelines, SourcePipeline))
 		{
 			if (!bSubmitted)
 			{
 				RHISubmitCommandsHint();
 				bSubmitted = true;
 			}
-			Fence->Signal(IsAsyncComputeContext() ? ED3D12CommandQueueType::Async : ED3D12CommandQueueType::Direct);
+			Fence->Signal(GetCommandQueueType());
 		}
 	}
 }
@@ -696,18 +705,27 @@ void FD3D12CommandContextBase::SignalTransitionFences(TArrayView<const FRHITrans
 void FD3D12CommandContextBase::WaitForTransitionFences(TArrayView<const FRHITransition*> Transitions)
 {
 	bool bSubmitted = false;
+
+	const ERHIPipeline DstPipeline = IsAsyncComputeContext() ? ERHIPipeline::AsyncCompute : ERHIPipeline::Graphics;
+
 	for (const FRHITransition* Transition : Transitions)
 	{
 		const auto* Data = Transition->GetPrivateData<FD3D12TransitionData>();
 		const auto& Fence = Data->Fence;
-		if (Fence)
+
+		if (!Fence)
+		{
+			continue;
+		}
+
+		if (Data->DstPipelines != ERHIPipeline::All || !EnumHasAnyFlags(Data->SrcPipelines, DstPipeline))
 		{
 			if (!bSubmitted)
 			{
 				RHISubmitCommandsHint();
 				bSubmitted = true;
 			}
-			Fence->GpuWait(IsAsyncComputeContext() ? ED3D12CommandQueueType::Async : ED3D12CommandQueueType::Direct, Fence->GetLastSignaledFence());
+			Fence->GpuWait(GetCommandQueueType(), Fence->GetLastSignaledFence());
 		}
 	}
 }
diff --git a/Engine/Source/Runtime/RHI/Private/RHICommandList.cpp b/Engine/Source/Runtime/RHI/Private/RHICommandList.cpp
--- a/Engine/Source/Runtime/RHI/Private/RHICommandList.cpp
+++ b/Engine/Source/Runtime/RHI/Private/RHICommandList.cpp
@@ -1055,7 +1055,7 @@ void FRHICommandListExecutor::Transition(TArrayView<const FRHITransitionInfo> In
 	}
 #endif
 
-	if (!GSupportsEfficientAsyncCompute)
+	if (!GSupportsEfficientAsyncCompute || RHICmdListAsyncCompute.Bypass())
 	{
 		checkf(SrcPipelines != ERHIPipeline::AsyncCompute, TEXT("Async compute is disabled. Cannot transition from it."));
 		checkf(DstPipelines != ERHIPipeline::AsyncCompute, TEXT("Async compute is disabled. Cannot transition to it."));
