Adding ability to create a per-platform public interface for D3D12RHI.

diff --git a/Engine/Source/Runtime/D3D12RHI/D3D12RHI.Build.cs b/Engine/Source/Runtime/D3D12RHI/D3D12RHI.Build.cs
--- a/Engine/Source/Runtime/D3D12RHI/D3D12RHI.Build.cs
+++ b/Engine/Source/Runtime/D3D12RHI/D3D12RHI.Build.cs
@@ -43,6 +43,8 @@ public class D3D12RHI : ModuleRules
         if (Target.Platform.IsInGroup(UnrealPlatformGroup.Windows) ||
             Target.Platform == UnrealTargetPlatform.HoloLens)
 		{
+			PublicDefinitions.Add("D3D12RHI_PLATFORM_HAS_CUSTOM_INTERFACE=0");
+
 			AddEngineThirdPartyPrivateStaticDependencies(Target, "DX12");
 			AddEngineThirdPartyPrivateStaticDependencies(Target, "AMD_AGS");
 
diff --git a/Engine/Source/Runtime/D3D12RHI/Private/D3D12RHIPrivate.h b/Engine/Source/Runtime/D3D12RHI/Private/D3D12RHIPrivate.h
--- a/Engine/Source/Runtime/D3D12RHI/Private/D3D12RHIPrivate.h
+++ b/Engine/Source/Runtime/D3D12RHI/Private/D3D12RHIPrivate.h
@@ -265,7 +265,7 @@ struct INTCExtensionContext;
 PRAGMA_DISABLE_DEPRECATION_WARNINGS
 
 /** The interface which is implemented by the dynamically bound RHI. */
-class FD3D12DynamicRHI : public ID3D12DynamicRHI
+class FD3D12DynamicRHI : public ID3D12PlatformDynamicRHI
 {
 	friend class FD3D12CommandContext;
 
diff --git a/Engine/Source/Runtime/D3D12RHI/Public/ID3D12DynamicRHI.h b/Engine/Source/Runtime/D3D12RHI/Public/ID3D12DynamicRHI.h
--- a/Engine/Source/Runtime/D3D12RHI/Public/ID3D12DynamicRHI.h
+++ b/Engine/Source/Runtime/D3D12RHI/Public/ID3D12DynamicRHI.h
@@ -65,3 +65,15 @@ inline ID3D12DynamicRHI* GetID3D12DynamicRHI()
 	check(GDynamicRHI->GetInterfaceType() == ERHIInterfaceType::D3D12);
 	return GetDynamicRHI<ID3D12DynamicRHI>();
 }
+
+#if D3D12RHI_PLATFORM_HAS_CUSTOM_INTERFACE
+	#include "ID3D12PlatformDynamicRHI.h"
+#else
+	using ID3D12PlatformDynamicRHI = ID3D12DynamicRHI;
+#endif
+
+inline ID3D12PlatformDynamicRHI* GetID3D12PlatformDynamicRHI()
+{
+	check(GDynamicRHI->GetInterfaceType() == ERHIInterfaceType::D3D12);
+	return GetDynamicRHI<ID3D12PlatformDynamicRHI>();
+}
