- Use shared preallocated resource for DispatchRaysDescBuffer per FD3D12Device (reduce allocation overhead)
- Copy SBT data part of D3D12_DISPATCH_RAYS_DESC from upload memory resource and indirect dimension from ArgumentBuffer which is build on the GPU timeline

diff --git a/Engine/Source/Runtime/D3D12RHI/Private/D3D12Device.cpp b/Engine/Source/Runtime/D3D12RHI/Private/D3D12Device.cpp
--- a/Engine/Source/Runtime/D3D12RHI/Private/D3D12Device.cpp
+++ b/Engine/Source/Runtime/D3D12RHI/Private/D3D12Device.cpp
@@ -49,6 +49,9 @@ FD3D12Device::~FD3D12Device()
 	delete RayTracingCompactionRequestHandler;
 	RayTracingCompactionRequestHandler = nullptr;
 
+	delete RayTracingDispatchRaysDescBuffer;
+	RayTracingDispatchRaysDescBuffer = nullptr;
+
 	DestroyRayTracingDescriptorCache(); // #dxr_todo UE-72158: unify RT descriptor cache with main FD3D12DescriptorCache
 #endif
 
@@ -378,6 +381,13 @@ void FD3D12Device::SetupAfterDeviceCreation()
 #if D3D12_RHI_RAYTRACING
 	check(RayTracingCompactionRequestHandler == nullptr);
 	RayTracingCompactionRequestHandler = new FD3D12RayTracingCompactionRequestHandler(this);
+
+	D3D12_RESOURCE_DESC DispatchRaysDescBufferDesc = CD3DX12_RESOURCE_DESC::Buffer(sizeof(D3D12_DISPATCH_RAYS_DESC), D3D12_RESOURCE_FLAG_NONE);
+	RayTracingDispatchRaysDescBuffer = GetParentAdapter()->CreateRHIBuffer(
+		DispatchRaysDescBufferDesc, 256,
+		0, DispatchRaysDescBufferDesc.Width, BUF_DrawIndirect,
+		ED3D12ResourceStateMode::MultiState, D3D12_RESOURCE_STATE_INDIRECT_ARGUMENT, false /*bInitialData*/,
+		GetGPUMask(), nullptr /*ResourceAllocator*/, TEXT("DispatchRaysDescBuffer"));
 #endif // D3D12_RHI_RAYTRACING
 
 	GPUProfilingData.Init();
diff --git a/Engine/Source/Runtime/D3D12RHI/Private/D3D12Device.h b/Engine/Source/Runtime/D3D12RHI/Private/D3D12Device.h
--- a/Engine/Source/Runtime/D3D12RHI/Private/D3D12Device.h
+++ b/Engine/Source/Runtime/D3D12RHI/Private/D3D12Device.h
@@ -11,6 +11,7 @@ D3D12Device.h: D3D12 Device Interfaces
 
 class FD3D12DynamicRHI;
 class FD3D12BasicRayTracingPipeline;
+class FD3D12Buffer;
 class FD3D12RayTracingDescriptorHeapCache;
 class FD3D12RayTracingPipelineCache;
 class FD3D12RayTracingCompactionRequestHandler;
@@ -55,6 +56,7 @@ public:
 	const FD3D12BasicRayTracingPipeline*	GetBasicRayTracingPipeline();
 	FD3D12RayTracingDescriptorHeapCache*	GetRayTracingDescriptorHeapCache() { return RayTracingDescriptorHeapCache; }
 	FD3D12RayTracingPipelineCache*			GetRayTracingPipelineCache() { return RayTracingPipelineCache; }
+	FD3D12Buffer*							GetRayTracingDispatchRaysDescBuffer() { return RayTracingDispatchRaysDescBuffer; }
 	FD3D12RayTracingCompactionRequestHandler* GetRayTracingCompactionRequestHandler() { return RayTracingCompactionRequestHandler; }
 	TRefCountPtr<ID3D12StateObject>			DeserializeRayTracingStateObject(D3D12_SHADER_BYTECODE Bytecode, ID3D12RootSignature* RootSignature);
 
@@ -200,7 +202,8 @@ protected:
 	FD3D12BasicRayTracingPipeline* BasicRayTracingPipeline = nullptr;
 	FD3D12RayTracingPipelineCache* RayTracingPipelineCache = nullptr;
 	FD3D12RayTracingCompactionRequestHandler* RayTracingCompactionRequestHandler = nullptr;
-	// #dxr_todo UE-72158: unify RT descriptor cache with main FD3D12DescriptorCache
+	FD3D12Buffer* RayTracingDispatchRaysDescBuffer = nullptr;
+// #dxr_todo UE-72158: unify RT descriptor cache with main FD3D12DescriptorCache
 	FD3D12RayTracingDescriptorHeapCache* RayTracingDescriptorHeapCache = nullptr;
 	void DestroyRayTracingDescriptorCache();
 #endif
diff --git a/Engine/Source/Runtime/D3D12RHI/Private/D3D12RayTracing.cpp b/Engine/Source/Runtime/D3D12RHI/Private/D3D12RayTracing.cpp
--- a/Engine/Source/Runtime/D3D12RHI/Private/D3D12RayTracing.cpp
+++ b/Engine/Source/Runtime/D3D12RHI/Private/D3D12RayTracing.cpp
@@ -5082,9 +5082,9 @@ static void DispatchRays(FD3D12CommandContext& CommandContext,
 {
 	SCOPE_CYCLE_COUNTER(STAT_D3D12DispatchRays);
 
-	FD3D12Adapter* Adapter = CommandContext.GetParentDevice()->GetParentAdapter();
+	FD3D12Device* Device = CommandContext.GetParentDevice();
+	FD3D12Adapter* Adapter = Device->GetParentAdapter();
 
-	FBufferRHIRef DispatchRaysDescBufferRHI;
 	FD3D12Buffer* DispatchRaysDescBuffer = nullptr;
 
 	if (ArgumentBuffer)
@@ -5094,27 +5094,37 @@ static void DispatchRays(FD3D12CommandContext& CommandContext,
 		// Source indirect argument buffer only contains the dispatch dimensions, however D3D12 requires a full D3D12_DISPATCH_RAYS_DESC structure.
 		// We create a new buffer, fill the SBT pointers on CPU and copy the dispatch dimensions into the right place.
 
-		D3D12_RESOURCE_DESC BufferDesc = CD3DX12_RESOURCE_DESC::Buffer(
-			sizeof(D3D12_DISPATCH_RAYS_DESC),
-			D3D12_RESOURCE_FLAG_ALLOW_UNORDERED_ACCESS,
-			D3D12_DEFAULT_RESOURCE_PLACEMENT_ALIGNMENT);
+		DispatchRaysDescBuffer = Device->GetRayTracingDispatchRaysDescBuffer();
+		FD3D12Resource* DispatchRaysDescBufferResource = DispatchRaysDescBuffer->GetResource();
 
-		FRHIResourceCreateInfo DispatchRaysDescBufferCreateInfo(TEXT("DispatchRaysDescBuffer"));
-		DispatchRaysDescBufferCreateInfo.GPUMask = FRHIGPUMask::FromIndex(CommandContext.GetGPUIndex());
+		FD3D12DynamicRHI::TransitionResource(CommandListHandle, DispatchRaysDescBufferResource, D3D12_RESOURCE_STATE_TBD, D3D12_RESOURCE_STATE_COPY_DEST, 0, FD3D12DynamicRHI::ETransitionMode::Apply);
+		FD3D12DynamicRHI::TransitionResource(CommandListHandle, ArgumentBuffer->GetResource(), D3D12_RESOURCE_STATE_TBD, D3D12_RESOURCE_STATE_COPY_SOURCE, 0, FD3D12DynamicRHI::ETransitionMode::Apply);
+		CommandListHandle.FlushResourceBarriers();
 
-		DispatchRaysDescBufferRHI = ::RHICreateVertexBuffer(sizeof(D3D12_DISPATCH_RAYS_DESC), BUF_UnorderedAccess | BUF_DrawIndirect, DispatchRaysDescBufferCreateInfo);
-		DispatchRaysDescBuffer = FD3D12DynamicRHI::ResourceCast(DispatchRaysDescBufferRHI.GetReference());
+		// Compute the allocation & copy sizes
+		uint32 DispatchRayDescSize = sizeof(D3D12_DISPATCH_RAYS_DESC);
+		uint32 SBTPartSize = offsetof(D3D12_DISPATCH_RAYS_DESC, Width);
+		uint32 IndirectDimensionSize = DispatchRayDescSize - SBTPartSize;
+		static_assert((sizeof(D3D12_DISPATCH_RAYS_DESC) - offsetof(D3D12_DISPATCH_RAYS_DESC, Width)) == sizeof(uint32) * 4, "Assume 4 uints at the end of the struct to store the dimension + alignment overhead");
 
-		FD3D12Resource* DispatchRaysDescBufferResource = DispatchRaysDescBuffer->GetResource();
+		uint32 BaseRayDescBufferOffset = DispatchRaysDescBuffer->ResourceLocation.GetOffsetFromBaseOfResource();
 
-		{
-			TRHICommandList_RecursiveHazardous<FD3D12CommandContext> RHICmdList(&CommandContext);
-			FShaderResourceViewRHIRef SRV = RHICreateShaderResourceView(ArgumentBuffer, 4, PF_R32_UINT);
-			FUnorderedAccessViewRHIRef UAV = RHICreateUnorderedAccessView(DispatchRaysDescBuffer, PF_R32_UINT);
-			FRayTracingDispatchDescCS::Dispatch(RHICmdList, &DispatchDesc, sizeof(DispatchDesc), offsetof(D3D12_DISPATCH_RAYS_DESC, Width), SRV, ArgumentOffset, UAV);
-		}
+		// Copy SBT data part of the dispatch desc to upload memory
+		FD3D12ResourceLocation UploadResourceLocation(Device);
+		void* Data = Device->GetDefaultFastAllocator().Allocate(DispatchRayDescSize, 256, &UploadResourceLocation);
+		FMemory::Memcpy(Data, &DispatchDesc, SBTPartSize);
+
+		// Copy SBT data part to resource
+		CommandListHandle->CopyBufferRegion(DispatchRaysDescBufferResource->GetResource(), BaseRayDescBufferOffset,
+			UploadResourceLocation.GetResource()->GetResource(), UploadResourceLocation.GetOffsetFromBaseOfResource(), SBTPartSize);
+
+		// Copy GPU computed indirect args to resource
+		CommandListHandle->CopyBufferRegion(DispatchRaysDescBufferResource->GetResource(), BaseRayDescBufferOffset + SBTPartSize,
+			ArgumentBuffer->GetResource()->GetResource(), ArgumentBuffer->ResourceLocation.GetOffsetFromBaseOfResource() + ArgumentOffset, IndirectDimensionSize);
+
+		FD3D12DynamicRHI::TransitionResource(CommandListHandle, DispatchRaysDescBufferResource, D3D12_RESOURCE_STATE_COPY_DEST, D3D12_RESOURCE_STATE_INDIRECT_ARGUMENT, 0, FD3D12DynamicRHI::ETransitionMode::Apply);
 
-		CommandListHandle.AddTransitionBarrier(DispatchRaysDescBufferResource, D3D12_RESOURCE_STATE_UNORDERED_ACCESS, D3D12_RESOURCE_STATE_INDIRECT_ARGUMENT, 0);
+		CommandListHandle.FlushResourceBarriers();
 	}
 
 	// Setup state for RT dispatch
