D3D12RHI: Updated check for aliased textures in FD3D12TextureStats

Previously, FD3D12TextureStats::D3D12TextureDeleted would only allow a TextureSize of 0 if there was an AliasingSourceTexture set. That only applied to explicitly aliased textures via RHICreateAliasedTexture, created by aliasing a FTextureRHIRef.
There is another class of aliased texture: textures created from externally sourced ID3D12Resources (CreateTextureFromResource). In this case, we do mark the ResourceLocation as aliased, but the AliasingSourceTexture cannot be set as we do not have a source FTextureRHIRef.

Also, added aliased helper to FD3D12ResourceLocation.

diff --git a/Engine/Source/Runtime/D3D12RHI/Private/D3D12Texture.cpp b/Engine/Source/Runtime/D3D12RHI/Private/D3D12Texture.cpp
--- a/Engine/Source/Runtime/D3D12RHI/Private/D3D12Texture.cpp
+++ b/Engine/Source/Runtime/D3D12RHI/Private/D3D12Texture.cpp
@@ -366,13 +366,12 @@ static TStatId GetD3D12StatEnum(D3D12_RESOURCE_FLAGS MiscFlags)
 void FD3D12TextureStats::UpdateD3D12TextureStats(FD3D12Texture& Texture, const D3D12_RESOURCE_DESC& Desc, int64 TextureSize, bool b3D, bool bCubeMap, bool bStreamable, bool bNewTexture)
 {
 #if TEXTURE_PROFILER_ENABLED
-	if (!bNewTexture && 
-		!Texture.ResourceLocation.IsTransient() 
+	if (!bNewTexture
+		&& !Texture.ResourceLocation.IsTransient() 
 		&& !EnumHasAnyFlags(Texture.GetFlags(), TexCreate_Virtual)
-		&& Texture.ResourceLocation.GetType() != FD3D12ResourceLocation::ResourceLocationType::eAliased
-		&& Texture.ResourceLocation.GetType() != FD3D12ResourceLocation::ResourceLocationType::eHeapAliased)
+		&& !Texture.ResourceLocation.IsAliased())
 	{
-		uint64 SafeSize = (uint64)(TextureSize >= 0 ? TextureSize : 0);
+		const uint64 SafeSize = (uint64)(TextureSize >= 0 ? TextureSize : 0);
 		FTextureProfiler::Get()->UpdateTextureAllocation(&Texture, SafeSize, Desc.Alignment, 0);
 	}
 #endif
@@ -385,7 +384,7 @@ void FD3D12TextureStats::UpdateD3D12TextureStats(FD3D12Texture& Texture, const D
 	const int64 AlignedSize = (TextureSize > 0) ? Align(TextureSize, 1024) / 1024 : -(Align(-TextureSize, 1024) / 1024);
 	if (ShouldCountAsTextureMemory(Desc.Flags))
 	{
-		bool bOnlyStreamableTextureAccounted = CVarTexturePoolOnlyAccountStreamableTexture.GetValueOnAnyThread();
+		const bool bOnlyStreamableTextureAccounted = CVarTexturePoolOnlyAccountStreamableTexture.GetValueOnAnyThread();
 
 		if (!bOnlyStreamableTextureAccounted || bStreamable)
 		{
@@ -424,23 +423,21 @@ void FD3D12TextureStats::D3D12TextureAllocated(FD3D12Texture& Texture, const D3D
 		}
 
 		// Don't update state for virtual or transient textures	
-		ETextureCreateFlags CreateFlags = Texture.GetDesc().Flags;
+		const ETextureCreateFlags CreateFlags = Texture.GetDesc().Flags;
 		if (!EnumHasAnyFlags(CreateFlags, TexCreate_Virtual | TexCreate_CPUReadback) && !Texture.ResourceLocation.IsTransient())
 		{
 			TRACE_CPUPROFILER_EVENT_SCOPE(D3D12RHI::UpdateTextureStats);
 
 			const int64 TextureSize = Texture.ResourceLocation.GetSize();			
-			bool bNewTexture = true;
-			bool bIsStreamable = EnumHasAnyFlags(Texture.GetDesc().Flags, ETextureCreateFlags::Streamable);
+			const bool bNewTexture = true;
+			const bool bIsStreamable = EnumHasAnyFlags(CreateFlags, ETextureCreateFlags::Streamable);
 			UpdateD3D12TextureStats(Texture, *Desc, TextureSize, Texture.GetDesc().IsTexture3D(), Texture.GetDesc().IsTextureCube(), bIsStreamable, bNewTexture);
 
 #if TEXTURE_PROFILER_ENABLED
-			if (Texture.ResourceLocation.GetType() != FD3D12ResourceLocation::ResourceLocationType::eAliased
-				&& Texture.ResourceLocation.GetType() != FD3D12ResourceLocation::ResourceLocationType::eHeapAliased)
+			if (!Texture.ResourceLocation.IsAliased())
 			{
-				size_t Size = Texture.ResourceLocation.GetSize();
-				uint32 Alignment = Desc->Alignment;
-				FTextureProfiler::Get()->AddTextureAllocation(&Texture, Size, Alignment, 0);
+				const uint32 Alignment = Desc->Alignment;
+				FTextureProfiler::Get()->AddTextureAllocation(&Texture, TextureSize, Alignment, 0);
 			}
 #endif
 		}
@@ -456,20 +453,19 @@ void FD3D12TextureStats::D3D12TextureDeleted(FD3D12Texture& Texture)
 	if (D3D12Texture)
 	{
 		// Don't update state for transient textures	
-		ETextureCreateFlags CreateFlags = Texture.GetDesc().Flags;
+		const ETextureCreateFlags CreateFlags = Texture.GetDesc().Flags;
 		if (!EnumHasAnyFlags(CreateFlags, TexCreate_Virtual | TexCreate_CPUReadback) && !Texture.ResourceLocation.IsTransient())
 		{
 			const D3D12_RESOURCE_DESC& Desc = D3D12Texture->GetDesc();
 			const int64 TextureSize = Texture.ResourceLocation.GetSize();
-			ensure(TextureSize > 0 || Texture.GetAliasingSourceTexture() != nullptr);
+			ensure(TextureSize > 0 || Texture.ResourceLocation.IsAliased());
 
-			bool bNewTexture = false;
-			bool bIsStreamable = EnumHasAnyFlags(Texture.GetDesc().Flags, ETextureCreateFlags::Streamable);
+			const bool bNewTexture = false;
+			const bool bIsStreamable = EnumHasAnyFlags(CreateFlags, ETextureCreateFlags::Streamable);
 			UpdateD3D12TextureStats(Texture, Desc, -TextureSize, Texture.GetDesc().IsTexture3D(), Texture.GetDesc().IsTextureCube(), bIsStreamable, bNewTexture);
 
 #if TEXTURE_PROFILER_ENABLED
-			if (Texture.ResourceLocation.GetType() != FD3D12ResourceLocation::ResourceLocationType::eAliased
-				&& Texture.ResourceLocation.GetType() != FD3D12ResourceLocation::ResourceLocationType::eHeapAliased)
+			if (!Texture.ResourceLocation.IsAliased())
 			{
 				FTextureProfiler::Get()->RemoveTextureAllocation(&Texture);
 			}
@@ -2680,8 +2676,7 @@ void FD3D12DynamicRHI::RHIBindDebugLabelName(FRHITexture* TextureRHI, const TCHA
 	
 	if (!EnumHasAnyFlags(TextureRHI->GetFlags(), TexCreate_Virtual)
 		&& !D3D12Texture->ResourceLocation.IsTransient()
-		&& D3D12Texture->ResourceLocation.GetType() != FD3D12ResourceLocation::ResourceLocationType::eAliased
-		&& D3D12Texture->ResourceLocation.GetType() != FD3D12ResourceLocation::ResourceLocationType::eHeapAliased)
+		&& !D3D12Texture->ResourceLocation.IsAliased())
 	{
 		FTextureProfiler::Get()->UpdateTextureName(TextureRHI);
 	}
@@ -2888,7 +2883,6 @@ FTextureRHIRef FD3D12DynamicRHI::RHICreateAliasedTexture(FTextureRHIRef& SourceT
 		UE_LOG(LogD3D12RHI, Error, TEXT("Currently FD3D12DynamicRHI::RHICreateAliasedTexture only supports 2D, 2D Array and Cube textures."));
 		return nullptr;
 	}
-	ReturnTexture->SetAliasingSource(SourceTextureRHI);
 
 	return ReturnTexture;
 }
diff --git a/Engine/Source/Runtime/D3D12RHI/Public/D3D12Resources.h b/Engine/Source/Runtime/D3D12RHI/Public/D3D12Resources.h
--- a/Engine/Source/Runtime/D3D12RHI/Public/D3D12Resources.h
+++ b/Engine/Source/Runtime/D3D12RHI/Public/D3D12Resources.h
@@ -549,7 +549,7 @@ public:
 		eSubAllocation,
 		eFastAllocation,
 		eMultiFrameFastAllocation,
-		eAliased, // Oculus is the only API that uses this
+		eAliased, // XR HMDs are the only use cases
 		eNodeReference,
 		eHeapAliased, 
 	};
@@ -649,11 +649,16 @@ public:
 		SetGPUVirtualAddress(GPUBase + Offset);
 	}
 
-	// Oculus API Aliases textures so this allows 2+ resource locations to reference the same underlying
+	// XR plugins alias textures so this allows 2+ resource locations to reference the same underlying
 	// resource. We should avoid this as much as possible as it requires expensive reference counting and
 	// it complicates the resource ownership model.
 	static void Alias(FD3D12ResourceLocation& Destination, FD3D12ResourceLocation& Source);
 	static void ReferenceNode(FD3D12Device* NodeDevice, FD3D12ResourceLocation& Destination, FD3D12ResourceLocation& Source);
+	bool IsAliased() const
+	{
+		return (Type == ResourceLocationType::eAliased) 
+			|| (Type == ResourceLocationType::eHeapAliased);
+	}
 
 	void SetTransient(bool bInTransient)
 	{
