Destroy cached RayTracingDispatchRaysDescBuffer during CleanupRaytracing instead of destructor

diff --git a/Engine/Source/Runtime/D3D12RHI/Private/D3D12Device.cpp b/Engine/Source/Runtime/D3D12RHI/Private/D3D12Device.cpp
--- a/Engine/Source/Runtime/D3D12RHI/Private/D3D12Device.cpp
+++ b/Engine/Source/Runtime/D3D12RHI/Private/D3D12Device.cpp
@@ -49,9 +49,6 @@ FD3D12Device::~FD3D12Device()
 	delete RayTracingCompactionRequestHandler;
 	RayTracingCompactionRequestHandler = nullptr;
 
-	delete RayTracingDispatchRaysDescBuffer;
-	RayTracingDispatchRaysDescBuffer = nullptr;
-
 	DestroyRayTracingDescriptorCache(); // #dxr_todo UE-72158: unify RT descriptor cache with main FD3D12DescriptorCache
 #endif
 
diff --git a/Engine/Source/Runtime/D3D12RHI/Private/D3D12RayTracing.cpp b/Engine/Source/Runtime/D3D12RHI/Private/D3D12RayTracing.cpp
--- a/Engine/Source/Runtime/D3D12RHI/Private/D3D12RayTracing.cpp
+++ b/Engine/Source/Runtime/D3D12RHI/Private/D3D12RayTracing.cpp
@@ -2961,6 +2961,9 @@ void FD3D12Device::CleanupRayTracing()
 	delete RayTracingPipelineCache;
 	RayTracingPipelineCache = nullptr;
 
+	delete RayTracingDispatchRaysDescBuffer;
+	RayTracingDispatchRaysDescBuffer = nullptr;
+
 	// Note: RayTracingDescriptorHeapCache is destroyed in ~FD3D12Device, after all deferred deletion is processed
 }
 
