Reapply change to report TexCreate_Virtual to D3D12 texture stats.

diff --git a/Engine/Source/Runtime/D3D12RHI/Private/D3D12Texture.cpp b/Engine/Source/Runtime/D3D12RHI/Private/D3D12Texture.cpp
--- a/Engine/Source/Runtime/D3D12RHI/Private/D3D12Texture.cpp
+++ b/Engine/Source/Runtime/D3D12RHI/Private/D3D12Texture.cpp
@@ -422,7 +422,7 @@ void FD3D12TextureStats::D3D12TextureAllocated(FD3D12Texture& Texture, const D3D
 			Desc = &D3D12Texture->GetDesc();
 		}
 
-		// Don't update state for virtual or transient textures	
+		// Don't update state for readback, virtual, or transient textures	
 		const ETextureCreateFlags CreateFlags = Texture.GetDesc().Flags;
 		if (!EnumHasAnyFlags(CreateFlags, TexCreate_Virtual | TexCreate_CPUReadback) && !Texture.ResourceLocation.IsTransient())
 		{
@@ -452,20 +452,20 @@ void FD3D12TextureStats::D3D12TextureDeleted(FD3D12Texture& Texture)
 
 	if (D3D12Texture)
 	{
-		// Don't update state for transient textures	
+		// Don't update state for readback or transient textures, but virtual textures need to have their size deducted from calls to RHIVirtualTextureSetFirstMipInMemory.
 		const ETextureCreateFlags CreateFlags = Texture.GetDesc().Flags;
-		if (!EnumHasAnyFlags(CreateFlags, TexCreate_Virtual | TexCreate_CPUReadback) && !Texture.ResourceLocation.IsTransient())
+		if (!EnumHasAnyFlags(CreateFlags, TexCreate_CPUReadback) && !Texture.ResourceLocation.IsTransient())
 		{
 			const D3D12_RESOURCE_DESC& Desc = D3D12Texture->GetDesc();
 			const int64 TextureSize = Texture.ResourceLocation.GetSize();
-			ensure(TextureSize > 0 || Texture.ResourceLocation.IsAliased());
+			ensure(TextureSize > 0 || EnumHasAnyFlags(CreateFlags, TexCreate_Virtual) || Texture.ResourceLocation.IsAliased());
 
 			const bool bNewTexture = false;
 			const bool bIsStreamable = EnumHasAnyFlags(CreateFlags, ETextureCreateFlags::Streamable);
 			UpdateD3D12TextureStats(Texture, Desc, -TextureSize, Texture.GetDesc().IsTexture3D(), Texture.GetDesc().IsTextureCube(), bIsStreamable, bNewTexture);
 
 #if TEXTURE_PROFILER_ENABLED
-			if (!Texture.ResourceLocation.IsAliased())
+			if (!EnumHasAnyFlags(CreateFlags, TexCreate_Virtual) && !Texture.ResourceLocation.IsAliased())
 			{
 				FTextureProfiler::Get()->RemoveTextureAllocation(&Texture);
 			}
