Only set the D3D12 residency priority for acceleration structure resources on pool creation, rather than on individual BLAS creation.

Saves significant CPU time when creating transient (short-lived) acceleration structures.

diff --git a/Engine/Source/Runtime/D3D12RHI/Private/D3D12PoolAllocator.cpp b/Engine/Source/Runtime/D3D12RHI/Private/D3D12PoolAllocator.cpp
--- a/Engine/Source/Runtime/D3D12RHI/Private/D3D12PoolAllocator.cpp
+++ b/Engine/Source/Runtime/D3D12RHI/Private/D3D12PoolAllocator.cpp
@@ -94,6 +94,22 @@ void FD3D12MemoryPool::Init()
 		}
 	}	
 
+#if D3D12_RHI_RAYTRACING
+	// Elevates the raytracing acceleration structure heap priority, which may help performance / stability in low memory conditions.
+	if (IsGPUOnly(InitConfig.HeapType) && InitConfig.InitialResourceState == D3D12_RESOURCE_STATE_RAYTRACING_ACCELERATION_STRUCTURE)
+	{
+		for (uint32 GPUIndex : GetGPUMask())
+		{
+			ID3D12Pageable* HeapResource = AllocationStrategy == EResourceAllocationStrategy::kPlacedResource
+				? BackingHeap->GetHeap()
+				: BackingResource->GetPageable();
+			D3D12_RESIDENCY_PRIORITY HeapPriority = D3D12_RESIDENCY_PRIORITY_HIGH;
+			FD3D12Device* NodeDevice = Adapter->GetDevice(GPUIndex);
+			NodeDevice->GetDevice5()->SetResidencyPriority(1, &HeapResource, &HeapPriority);
+		}
+	}
+#endif // D3D12_RHI_RAYTRACING
+
 	FRHIMemoryPool::Init();
 }
 
@@ -375,8 +391,10 @@ void FD3D12PoolAllocator::AllocateResource(uint32 GPUIndex, D3D12_HEAP_TYPE InHe
 	}
 
 #if D3D12_RHI_RAYTRACING
-	// Elevates the raytracing acceleration structure heap priority, which may help performance / stability in low memory conditions
-	if (InCreateState == D3D12_RESOURCE_STATE_RAYTRACING_ACCELERATION_STRUCTURE)
+	// Elevates the raytracing acceleration structure heap priority, which may help performance / stability in low memory conditions.
+	// When allocation is pooled, the entire pool priority is boosted on creation.
+	if (InCreateState == D3D12_RESOURCE_STATE_RAYTRACING_ACCELERATION_STRUCTURE
+		&& ResourceLocation.GetAllocatorType() != FD3D12ResourceLocation::AT_Pool)
 	{
 		ID3D12Pageable* HeapResource = ResourceLocation.GetResource()->GetPageable();
 		D3D12_RESIDENCY_PRIORITY HeapPriority = D3D12_RESIDENCY_PRIORITY_HIGH;
