Workaround for platforms which cannot utilize DiscardResource on fast clear RTs with typeless formats

- support doing a ClearRenderTargetView instead of a DiscardResource on typeless RTs as required by some platforms

diff --git a/Engine/Source/Runtime/Core/Public/HAL/Platform.h b/Engine/Source/Runtime/Core/Public/HAL/Platform.h
--- a/Engine/Source/Runtime/Core/Public/HAL/Platform.h
+++ b/Engine/Source/Runtime/Core/Public/HAL/Platform.h
@@ -617,7 +617,7 @@
 
 #ifndef PLATFORM_REQUIRES_UAV_TO_RTV_TEXTURE_CACHE_FLUSH_WORKAROUND
 	#define PLATFORM_REQUIRES_UAV_TO_RTV_TEXTURE_CACHE_FLUSH_WORKAROUND 0
-#endif // #ifndef PLATFORM_REQUIRES_UAV_TO_RTV_TEXTURE_CACHE_FLUSH_WORKAROUND
+#endif
 
 #ifndef PLATFORM_NEEDS_GPU_UAV_RESOURCE_INIT_WORKAROUND
 	#define PLATFORM_NEEDS_GPU_UAV_RESOURCE_INIT_WORKAROUND 0
@@ -639,6 +639,10 @@
 	#define PLATFORM_HAS_FENV_H 1
 #endif
 
+#ifndef PLATFORM_REQUIRES_TYPELESS_RESOURCE_DISCARD_WORKAROUND
+	#define PLATFORM_REQUIRES_TYPELESS_RESOURCE_DISCARD_WORKAROUND 0
+#endif
+
 // deprecated, do not use
 #define PLATFORM_HAS_THREADSAFE_RHIGetRenderQueryResult	#
 #define PLATFORM_SUPPORTS_RHI_THREAD #
diff --git a/Engine/Source/Runtime/D3D12RHI/Private/D3D12Commands.cpp b/Engine/Source/Runtime/D3D12RHI/Private/D3D12Commands.cpp
--- a/Engine/Source/Runtime/D3D12RHI/Private/D3D12Commands.cpp
+++ b/Engine/Source/Runtime/D3D12RHI/Private/D3D12Commands.cpp
@@ -184,7 +184,7 @@ void FD3D12CommandContext::RHIDispatchIndirectComputeShader(FRHIBuffer* Argument
 }
 
 template <typename FunctionType>
-void EnumerateSubresources(FD3D12Resource* Resource, const FRHITransitionInfo& Info, FunctionType Function)
+void EnumerateSubresources(FD3D12Resource* Resource, const FRHITransitionInfo& Info, FD3D12Texture* Texture, FunctionType Function)
 {
 	uint32 FirstMipSlice = 0;
 	uint32 FirstArraySlice = 0;
@@ -223,7 +223,14 @@ void EnumerateSubresources(FD3D12Resource* Resource, const FRHITransitionInfo& I
 			for (uint32 MipSlice = FirstMipSlice; MipSlice < FirstMipSlice + IterationMipCount; ++MipSlice)
 			{
 				const uint32 Subresource = D3D12CalcSubresource(MipSlice, ArraySlice, PlaneSlice, MipCount, ArraySize);
-				Function(Subresource);
+				const FD3D12RenderTargetView* RTV = nullptr;
+#if PLATFORM_REQUIRES_TYPELESS_RESOURCE_DISCARD_WORKAROUND
+				if (Texture)
+				{
+					RTV = Texture->GetRenderTargetView(MipSlice, ArraySlice);
+				}
+#endif // PLATFORM_REQUIRES_TYPELESS_RESOURCE_DISCARD_WORKAROUND
+				Function(Subresource, RTV);
 			}
 		}
 	}
@@ -255,7 +262,14 @@ void ProcessResource(FD3D12CommandContext& Context, const FRHITransitionInfo& In
 	{
 		FD3D12Texture* Texture = Context.RetrieveTexture(Info.Texture);
 		check(Texture);
-		Function(Info, Texture->GetResource());
+		FD3D12Texture* TextureOut = nullptr;
+#if PLATFORM_REQUIRES_TYPELESS_RESOURCE_DISCARD_WORKAROUND
+		if (Texture->GetRequiresTypelessResourceDiscardWorkaround())
+		{
+			TextureOut = Texture;
+		}
+#endif // #if PLATFORM_REQUIRES_TYPELESS_RESOURCE_DISCARD_WORKAROUND
+		Function(Info, Texture->GetResource(), TextureOut);
 		break;
 	}
 	default:
@@ -272,15 +286,23 @@ static bool ProcessTransitionDuringBegin(const FD3D12TransitionData* Data)
 
 struct FD3D12DiscardResource
 {
-	FD3D12DiscardResource(FD3D12Resource* InResource, EResourceTransitionFlags InFlags, uint32 InSubresource = D3D12_RESOURCE_BARRIER_ALL_SUBRESOURCES)
+	FD3D12DiscardResource(FD3D12Resource* InResource, EResourceTransitionFlags InFlags, uint32 InSubresource = D3D12_RESOURCE_BARRIER_ALL_SUBRESOURCES, const FD3D12Texture* InTexture = nullptr, const FD3D12RenderTargetView* InRTV = nullptr)
 		: Resource(InResource)
 		, Flags(InFlags)
 		, Subresource(InSubresource)
+#if PLATFORM_REQUIRES_TYPELESS_RESOURCE_DISCARD_WORKAROUND
+		, Texture(InTexture)
+		, RTV(InRTV)
+#endif // #if PLATFORM_REQUIRES_TYPELESS_RESOURCE_DISCARD_WORKAROUND
 	{}
 
 	FD3D12Resource* Resource;
 	EResourceTransitionFlags Flags;
 	uint32 Subresource;
+#if PLATFORM_REQUIRES_TYPELESS_RESOURCE_DISCARD_WORKAROUND
+	const FD3D12Texture* Texture = nullptr;
+	const FD3D12RenderTargetView* RTV = nullptr;
+#endif // #if PLATFORM_REQUIRES_TYPELESS_RESOURCE_DISCARD_WORKAROUND
 };
 
 using FD3D12DiscardResourceArray = TArray<FD3D12DiscardResource>;
@@ -298,7 +320,7 @@ static void HandleResourceDiscardTransitions(
 			continue;
 		}
 
-		ProcessResource(Context, Info, [&](const FRHITransitionInfo& Info, FD3D12Resource* Resource)
+		ProcessResource(Context, Info, [&](const FRHITransitionInfo& Info, FD3D12Resource* Resource, FD3D12Texture* Texture = nullptr)
 		{
 			if (!Resource->RequiresResourceStateTracking())
 			{
@@ -311,14 +333,22 @@ static void HandleResourceDiscardTransitions(
 			if (Info.IsWholeResource() || Resource->GetSubresourceCount() == 1)
 			{
 				FD3D12DynamicRHI::TransitionResource(Context.CommandListHandle, Resource, InitialState, InitialState, D3D12_RESOURCE_BARRIER_ALL_SUBRESOURCES, FD3D12DynamicRHI::ETransitionMode::Apply);
-				ResourcesToDiscard.Emplace(Resource, Info.Flags, D3D12_RESOURCE_BARRIER_ALL_SUBRESOURCES);
+
+				const FD3D12RenderTargetView* RTV = nullptr;
+#if PLATFORM_REQUIRES_TYPELESS_RESOURCE_DISCARD_WORKAROUND
+				if (Texture)
+				{
+					RTV = Texture->GetRenderTargetView(0, -1);
+				}
+#endif // #if PLATFORM_REQUIRES_TYPELESS_RESOURCE_DISCARD_WORKAROUND				
+				ResourcesToDiscard.Emplace(Resource, Info.Flags, D3D12_RESOURCE_BARRIER_ALL_SUBRESOURCES, Texture, RTV);
 			}
 			else
 			{
-				EnumerateSubresources(Resource, Info, [&](uint32 Subresource)
+				EnumerateSubresources(Resource, Info, Texture, [&](uint32 Subresource, const FD3D12RenderTargetView* RTV)
 				{
 					FD3D12DynamicRHI::TransitionResource(Context.CommandListHandle, Resource, InitialState, InitialState, Subresource, FD3D12DynamicRHI::ETransitionMode::Apply);
-					ResourcesToDiscard.Emplace(Resource, Info.Flags, Subresource);
+					ResourcesToDiscard.Emplace(Resource, Info.Flags, Subresource, Texture, RTV);
 				});
 			}
 		});
@@ -362,19 +392,30 @@ static void HandleDiscardResources(
 		}
 		else if (EnumHasAnyFlags(DiscardResource.Flags, EResourceTransitionFlags::Discard))
 		{
-			if (DiscardResource.Subresource == D3D12_RESOURCE_BARRIER_ALL_SUBRESOURCES)
+#if PLATFORM_REQUIRES_TYPELESS_RESOURCE_DISCARD_WORKAROUND
+			if (DiscardResource.Texture && DiscardResource.RTV)
 			{
-				Context.CommandListHandle->DiscardResource(DiscardResource.Resource->GetResource(), nullptr);
+				FLinearColor ClearColor = DiscardResource.Texture->GetClearColor();
+				Context.CommandListHandle->ClearRenderTargetView(DiscardResource.RTV->GetView(), reinterpret_cast<float*>(&ClearColor), 0, nullptr);
+				Context.CommandListHandle.UpdateResidency(DiscardResource.RTV->GetResource());
 			}
 			else
+#endif // #if PLATFORM_REQUIRES_TYPELESS_RESOURCE_DISCARD_WORKAROUND
 			{
-				D3D12_DISCARD_REGION Region;
-				Region.NumRects = 0;
-				Region.pRects = nullptr;
-				Region.FirstSubresource = DiscardResource.Subresource;
-				Region.NumSubresources = 1;
+				if (DiscardResource.Subresource == D3D12_RESOURCE_BARRIER_ALL_SUBRESOURCES)
+				{
+					Context.CommandListHandle->DiscardResource(DiscardResource.Resource->GetResource(), nullptr);
+				}
+				else
+				{
+					D3D12_DISCARD_REGION Region;
+					Region.NumRects = 0;
+					Region.pRects = nullptr;
+					Region.FirstSubresource = DiscardResource.Subresource;
+					Region.NumSubresources = 1;
 
-				Context.CommandListHandle->DiscardResource(DiscardResource.Resource->GetResource(), &Region);
+					Context.CommandListHandle->DiscardResource(DiscardResource.Resource->GetResource(), &Region);
+				}
 			}
 		}
 	}
@@ -485,7 +526,7 @@ static void HandleResourceTransitions(FD3D12CommandContext& Context, const FD3D1
 			}
 			else
 			{
-				ProcessResource(Context, Info, [&](const FRHITransitionInfo& Info, FD3D12Resource* Resource)
+				ProcessResource(Context, Info, [&](const FRHITransitionInfo& Info, FD3D12Resource* Resource, FD3D12Texture* UnusedTexture = nullptr)
 				{
 					if (!Resource->RequiresResourceStateTracking())
 					{
@@ -523,7 +564,7 @@ static void HandleResourceTransitions(FD3D12CommandContext& Context, const FD3D1
 					}
 					else
 					{
-						EnumerateSubresources(Resource, Info, [&](uint32 Subresource)
+						EnumerateSubresources(Resource, Info, nullptr, [&](uint32 Subresource, const FD3D12RenderTargetView* UnusedRTV = nullptr)
 						{
 							if (FD3D12DynamicRHI::TransitionResource(Context.CommandListHandle, Resource, StateBefore, StateAfter, Subresource, FD3D12DynamicRHI::ETransitionMode::Apply))
 							{
diff --git a/Engine/Source/Runtime/D3D12RHI/Private/D3D12Texture.h b/Engine/Source/Runtime/D3D12RHI/Private/D3D12Texture.h
--- a/Engine/Source/Runtime/D3D12RHI/Private/D3D12Texture.h
+++ b/Engine/Source/Runtime/D3D12RHI/Private/D3D12Texture.h
@@ -65,6 +65,9 @@ public:
 	void GetReadBackHeapDesc(D3D12_PLACED_SUBRESOURCE_FOOTPRINT& OutFootprint, uint32 Subresource) const;
 	FD3D12RenderTargetView* GetRenderTargetView(int32 MipIndex, int32 ArraySliceIndex) const;
 	FD3D12DepthStencilView* GetDepthStencilView(FExclusiveDepthStencil AccessType) const { return DepthStencilViews[AccessType.GetIndex()]; }
+#if PLATFORM_REQUIRES_TYPELESS_RESOURCE_DISCARD_WORKAROUND
+	bool GetRequiresTypelessResourceDiscardWorkaround() const { return bRequiresTypelessResourceDiscardWorkaround; }
+#endif // #if PLATFORM_REQUIRES_TYPELESS_RESOURCE_DISCARD_WORKAROUND
 		
 	// Setters
 	void SetShaderResourceView(FD3D12ShaderResourceView* InShaderResourceView) { ShaderResourceView = InShaderResourceView; }
@@ -113,6 +116,10 @@ protected:
 	// Set when RTVs are created for each slice - TexCreate_TargetArraySlicesIndependently for TextureArrays & Cubemaps
 	bool bCreatedRTVsPerSlice{ false };
 
+#if PLATFORM_REQUIRES_TYPELESS_RESOURCE_DISCARD_WORKAROUND
+	bool bRequiresTypelessResourceDiscardWorkaround = false;
+#endif // #if PLATFORM_REQUIRES_TYPELESS_RESOURCE_DISCARD_WORKAROUND
+
 	// Number of RTVs per mip when stored per slice
 	int32 RTVArraySizePerMip{};
 
