Fix for D3D validation error due to creation of RT indirect buffer without associated resource creation flag to allow indirect buffer state

- RT indirect buffer creation bypasses the logic in GetResourceDescAndAlignment() and hence does not get D3D12RHI_RESOURCE_FLAG_ALLOW_INDIRECT_BUFFER added to the resource description flags
  - explicitly add the flag in this instance to resolve the issue

diff --git a/Engine/Source/Runtime/D3D12RHI/Private/D3D12Device.cpp b/Engine/Source/Runtime/D3D12RHI/Private/D3D12Device.cpp
--- a/Engine/Source/Runtime/D3D12RHI/Private/D3D12Device.cpp
+++ b/Engine/Source/Runtime/D3D12RHI/Private/D3D12Device.cpp
@@ -379,7 +379,7 @@ void FD3D12Device::SetupAfterDeviceCreation()
 	check(RayTracingCompactionRequestHandler == nullptr);
 	RayTracingCompactionRequestHandler = new FD3D12RayTracingCompactionRequestHandler(this);
 
-	D3D12_RESOURCE_DESC DispatchRaysDescBufferDesc = CD3DX12_RESOURCE_DESC::Buffer(sizeof(D3D12_DISPATCH_RAYS_DESC), D3D12_RESOURCE_FLAG_NONE);
+	D3D12_RESOURCE_DESC DispatchRaysDescBufferDesc = CD3DX12_RESOURCE_DESC::Buffer(sizeof(D3D12_DISPATCH_RAYS_DESC), D3D12RHI_RESOURCE_FLAG_ALLOW_INDIRECT_BUFFER);
 	RayTracingDispatchRaysDescBuffer = GetParentAdapter()->CreateRHIBuffer(
 		DispatchRaysDescBufferDesc, 256,
 		0, DispatchRaysDescBufferDesc.Width, BUF_DrawIndirect,
