Fix buffer overflow if FD3D12CommandListManager::ExecuteCommandListInternal is called with more than 128 command lists which require transition patch-up.

The BarrierCommandList must have the same maximum size as the Lists array, which is sized to FD3D12CommandListPayload::MaxCommandListsPerPayload (currently 256).

#rnx

diff --git a/Engine/Source/Runtime/D3D12RHI/Private/D3D12DirectCommandListManager.cpp b/Engine/Source/Runtime/D3D12RHI/Private/D3D12DirectCommandListManager.cpp
--- a/Engine/Source/Runtime/D3D12RHI/Private/D3D12DirectCommandListManager.cpp
+++ b/Engine/Source/Runtime/D3D12RHI/Private/D3D12DirectCommandListManager.cpp
@@ -658,7 +658,7 @@ FD3D12SyncPoint FD3D12CommandListManager::ExecuteCommandListInternal(TArray<FD3D
 	FD3D12CommandListPayload CurrentCommandListPayload;
 
 	check(Lists.Num() <= FD3D12CommandListPayload::MaxCommandListsPerPayload);
-	FD3D12CommandListHandle BarrierCommandList[128];
+	FD3D12CommandListHandle BarrierCommandList[FD3D12CommandListPayload::MaxCommandListsPerPayload];
 	if (bNeedsResourceBarriers)
 	{
 		// should not be copy queue
