Allow typed atomics on Intel devices for Dx12 that report support for SM6.6 through extensions

(Courtesy of [REDACTED_PERSON] at Intel)

diff --git a/Engine/Source/Runtime/D3D12RHI/Private/D3D12Adapter.cpp b/Engine/Source/Runtime/D3D12RHI/Private/D3D12Adapter.cpp
--- a/Engine/Source/Runtime/D3D12RHI/Private/D3D12Adapter.cpp
+++ b/Engine/Source/Runtime/D3D12RHI/Private/D3D12Adapter.cpp
@@ -1056,6 +1056,10 @@ void FD3D12Adapter::InitializeDevices()
 				if (D3D12Caps9.AtomicInt64OnTypedResourceSupported && D3D12Caps11.AtomicInt64OnDescriptorHeapResourceSupported)
 				{
 					GRHISupportsDX12AtomicUInt64 = true;
+				}
+
+				if (GRHISupportsDX12AtomicUInt64)
+				{
 					UE_LOG(LogD3D12RHI, Log, TEXT("Shader Model 6.6 atomic64 is supported"));
 				}
 				else
diff --git a/Engine/Source/Runtime/D3D12RHI/Private/Windows/WindowsD3D12Device.cpp b/Engine/Source/Runtime/D3D12RHI/Private/Windows/WindowsD3D12Device.cpp
--- a/Engine/Source/Runtime/D3D12RHI/Private/Windows/WindowsD3D12Device.cpp
+++ b/Engine/Source/Runtime/D3D12RHI/Private/Windows/WindowsD3D12Device.cpp
@@ -1206,7 +1206,7 @@ void FD3D12DynamicRHI::Init()
 #if INTEL_EXTENSIONS
 	if (IsRHIDeviceIntel() && bAllowVendorDevice)
 	{
-		const INTCExtensionVersion AtomicsRequiredVersion = { 3, 4, 1 }; //version 3.4.1
+		const INTCExtensionVersion AtomicsRequiredVersion = { 4, 4, 0 }; //version 4.4.0 Alchemist+
 		INTCExtensionVersion* SupportedExtensionsVersions = nullptr;
 		uint32_t SupportedExtensionsVersionCount = 0;
 		INTCExtensionInfo INTCExtensionInfo{};
@@ -1241,6 +1241,18 @@ void FD3D12DynamicRHI::Init()
 			}
 		}
 
+	#if D3D12_CORE_ENABLED
+		if (GRHISupportsAtomicUInt64)
+		{
+			D3D12_FEATURE_DATA_SHADER_MODEL D3D12ShaderModelLevel = { D3D_SHADER_MODEL_6_6 };
+			if (GetAdapter().GetD3DDevice()->CheckFeatureSupport(D3D12_FEATURE_SHADER_MODEL, &D3D12ShaderModelLevel, sizeof(D3D12ShaderModelLevel)) == S_OK &&
+				D3D12ShaderModelLevel.HighestShaderModel >= D3D_SHADER_MODEL_6_6)
+			{
+				GRHISupportsDX12AtomicUInt64 = true;
+			}
+		}
+	#endif
+
 		check(IntelExtensionContext == nullptr);
 		INTCExtensionAppInfo AppInfo{};
 		AppInfo.pEngineName = TEXT("Unreal Engine");
