Release underlying resource data during adpater cleanup for all thread local transient uniform buffer allocator to make sure all d3d resources are cleaned up before adapter is destroyed

#rnx

diff --git a/Engine/Source/Runtime/D3D12RHI/Private/D3D12Adapter.cpp b/Engine/Source/Runtime/D3D12RHI/Private/D3D12Adapter.cpp
--- a/Engine/Source/Runtime/D3D12RHI/Private/D3D12Adapter.cpp
+++ b/Engine/Source/Runtime/D3D12RHI/Private/D3D12Adapter.cpp
@@ -1293,6 +1293,12 @@ void FD3D12Adapter::Cleanup()
 		RHICmdList.ImmediateFlush(EImmediateFlushType::FlushRHIThread);
 	}
 
+	// Release allocation data of all thread local transient uniform buffer allocators
+	for (FD3D12FastConstantAllocator* Allocator : TransientUniformBufferAllocators)
+	{
+		Allocator->Cleanup();
+	}
+
 	// Cleanup resources
 	DeferredDeletionQueue.ReleaseResources(true, true);
 
@@ -1541,17 +1547,37 @@ FD3D12FastConstantAllocator& FD3D12Adapter::GetTransientUniformBufferAllocator()
 {
 	class FTransientUniformBufferAllocator : public FD3D12FastConstantAllocator, public TThreadSingleton<FTransientUniformBufferAllocator>
 	{
-		using FD3D12FastConstantAllocator::FD3D12FastConstantAllocator;
+	public:
+		FTransientUniformBufferAllocator(FD3D12Adapter* InAdapter, FD3D12Device* Parent, FRHIGPUMask VisibiltyMask) : FD3D12FastConstantAllocator(Parent, VisibiltyMask), Adapter(InAdapter) {}
+		~FTransientUniformBufferAllocator()
+		{
+			Adapter->ReleaseTransientUniformBufferAllocator(this);
+		}
+	private:
+		FD3D12Adapter* Adapter;
 	};
 
 	// Multi-GPU support : is using device 0 always appropriate here?
 	return FTransientUniformBufferAllocator::Get([this]() -> FTransientUniformBufferAllocator*
 	{
-		FTransientUniformBufferAllocator* Alloc = new FTransientUniformBufferAllocator(Devices[0], FRHIGPUMask::All());
+		FTransientUniformBufferAllocator* Alloc = new FTransientUniformBufferAllocator(this, Devices[0], FRHIGPUMask::All());
+
+		// Register so the underlying resource location can be freed during adapter cleanup instead of when thread local allocation is destroyed
+		{
+			FScopeLock Lock(&TransientUniformBufferAllocatorsCS);
+			TransientUniformBufferAllocators.Add(Alloc);
+		}
+
 		return Alloc;
 	});
 }
 
+void FD3D12Adapter::ReleaseTransientUniformBufferAllocator(FD3D12FastConstantAllocator* InAllocator)
+{
+	FScopeLock Lock(&TransientUniformBufferAllocatorsCS);
+	verify(TransientUniformBufferAllocators.Remove(InAllocator) == 1);
+}
+
 void FD3D12Adapter::UpdateMemoryInfo()
 {
 #if PLATFORM_WINDOWS || PLATFORM_HOLOLENS
diff --git a/Engine/Source/Runtime/D3D12RHI/Private/D3D12Adapter.h b/Engine/Source/Runtime/D3D12RHI/Private/D3D12Adapter.h
--- a/Engine/Source/Runtime/D3D12RHI/Private/D3D12Adapter.h
+++ b/Engine/Source/Runtime/D3D12RHI/Private/D3D12Adapter.h
@@ -391,6 +391,7 @@ public:
 	FD3D12TemporalEffect* GetTemporalEffect(const FName& EffectName);
 
 	FD3D12FastConstantAllocator& GetTransientUniformBufferAllocator();
+	void ReleaseTransientUniformBufferAllocator(FD3D12FastConstantAllocator* InAllocator);
 
 	void BlockUntilIdle();
 
@@ -587,6 +588,9 @@ protected:
 	TMap<FName, FD3D12TemporalEffect> TemporalEffectMap;
 #endif
 
+	TArray<FD3D12FastConstantAllocator*> TransientUniformBufferAllocators;
+	FCriticalSection TransientUniformBufferAllocatorsCS;
+
 	TUniquePtr<IRHITransientMemoryCache> TransientMemoryCache;
 
 	// Each of these devices represents a physical GPU 'Node'
diff --git a/Engine/Source/Runtime/D3D12RHI/Private/D3D12Allocation.h b/Engine/Source/Runtime/D3D12RHI/Private/D3D12Allocation.h
--- a/Engine/Source/Runtime/D3D12RHI/Private/D3D12Allocation.h
+++ b/Engine/Source/Runtime/D3D12RHI/Private/D3D12Allocation.h
@@ -522,6 +522,7 @@ public:
 	FD3D12FastConstantAllocator(FD3D12Device* Parent, FRHIGPUMask VisibiltyMask);
 
 	void* Allocate(uint32 Bytes, class FD3D12ResourceLocation& OutLocation, FD3D12ConstantBufferView* OutCBView);
+	void Cleanup() { UnderlyingResource.Clear(); }
 
 private:
 	FD3D12ResourceLocation UnderlyingResource;
