Break up Timestamp query batches when going over the requested NumQueriesInBatch because it could otherwise overflow the query heap if it's overflowing at the end of the heap

#rnx

diff --git a/Engine/Source/Runtime/D3D12RHI/Private/D3D12Query.cpp b/Engine/Source/Runtime/D3D12RHI/Private/D3D12Query.cpp
--- a/Engine/Source/Runtime/D3D12RHI/Private/D3D12Query.cpp
+++ b/Engine/Source/Runtime/D3D12RHI/Private/D3D12Query.cpp
@@ -310,33 +310,36 @@ uint32 FD3D12QueryHeap::AllocQuery(FD3D12CommandContext& CmdContext)
 
 	if (QueryType == D3D12_QUERY_TYPE_OCCLUSION)
 	{
-		check(CurrentQueryBatch.bOpen);
+		checkf(CurrentQueryBatch.bOpen && CurrentQueryBatch.ElementCount < CurrentQueryBatch.NumQueriesInBatch, TEXT("The reserved query batch is too small. Increase NumQueriesInBatch when calling StartQueryBatch."));
 	}
 	else
 	{
+		check(!CurrentQueryBatch.bOpen || CurrentQueryBatch.ElementCount <= CurrentQueryBatch.NumQueriesInBatch);
+
 		// Queries per batch gives a hint for resizing the query heap when starting a batch
 		// The code will still 'work' if it uses more than this value, but it might 'overwrite'
 		// previous batches (this isn't checked anywhere). If a single batch uses more queries
 		// than available in the heap than resolve will happen in the middle (see below) 
-		static const uint32 DefaultMaxTimeStampQueriesPerBat = 256;
+		static const uint32 DefaultMaxTimeStampQueriesPerBatch = 256;
 
 		if (!CurrentQueryBatch.bOpen)
 		{
-			StartQueryBatch(CmdContext, DefaultMaxTimeStampQueriesPerBat);
+			StartQueryBatch(CmdContext, DefaultMaxTimeStampQueriesPerBatch);
 			check(CurrentQueryBatch.bOpen && CurrentQueryBatch.ElementCount == 0);
 		}
 
-		if (CurrentQueryBatch.StartElement > CurrentElement)
-		{
-			// We're in the middle of a batch, but we're at the end of the heap
-			// We need to split the batch in two and resolve the first piece
+		// We're in the middle of a batch, but we're at the end of the heap or we are going over the max time stamp queries
+		// We need to split the batch in two and resolve the first piece - either we had a wrap around or we might need
+		// to increase the heap size
+		if (CurrentQueryBatch.StartElement > CurrentElement || CurrentQueryBatch.ElementCount == CurrentQueryBatch.NumQueriesInBatch)
+		{						
 			EndQueryBatchAndResolveQueryData(CmdContext);
 		}
 
 		// check for the the batch being closed due to wrap and open a new one
 		if (!CurrentQueryBatch.bOpen)
 		{
-			StartQueryBatch(CmdContext, DefaultMaxTimeStampQueriesPerBat);
+			StartQueryBatch(CmdContext, DefaultMaxTimeStampQueriesPerBatch);
 			check(CurrentQueryBatch.bOpen && CurrentQueryBatch.ElementCount == 0);
 		}
 	}
@@ -372,6 +375,7 @@ void FD3D12QueryHeap::StartQueryBatch(FD3D12CommandContext& CmdContext, uint32 N
 
 	// Start a new batch
 	CurrentQueryBatch.StartElement = GetNextElement(LastAllocatedElement);
+	CurrentQueryBatch.NumQueriesInBatch = NumQueriesInBatch;
 	CurrentQueryBatch.UsedQueryHeap = ActiveQueryHeap;
 	CurrentQueryBatch.UsedResultBuffer = ActiveResultBuffer;
 	CurrentQueryBatch.bOpen = true;
@@ -392,8 +396,6 @@ void FD3D12QueryHeap::EndQueryBatchAndResolveQueryData(FD3D12CommandContext& Cmd
 		return;
 	}
 
-	check(CurrentQueryBatch.bOpen);
-
 	// Close the current batch
 	CurrentQueryBatch.bOpen = false;
 
diff --git a/Engine/Source/Runtime/D3D12RHI/Private/D3D12Query.h b/Engine/Source/Runtime/D3D12RHI/Private/D3D12Query.h
--- a/Engine/Source/Runtime/D3D12RHI/Private/D3D12Query.h
+++ b/Engine/Source/Runtime/D3D12RHI/Private/D3D12Query.h
@@ -88,6 +88,7 @@ private:
 	public:
 		uint32 StartElement;    							// The first element in the batch (inclusive)
 		uint32 ElementCount;    							// The number of elements in the batch
+		uint32 NumQueriesInBatch;							// Maximum number of queries which are reserved for the batch
 		bool bOpen;             							// Is the batch still open for more begin/end queries?
 
 		TRefCountPtr<ID3D12QueryHeap> UsedQueryHeap;		// The query heap where all elements reside
@@ -107,6 +108,7 @@ private:
 		{
 			StartElement = 0;
 			ElementCount = 0;
+			NumQueriesInBatch = 0;
 			bOpen = false;
 			RenderQueries.Reset();
 
