Make sure the Adapter is cleared after the adapter shutdown function is done on the TLS allocator (fix crash when TLS allocator is destroyed after adapter)

#rnx

diff --git a/Engine/Source/Runtime/D3D12RHI/Private/D3D12Adapter.cpp b/Engine/Source/Runtime/D3D12RHI/Private/D3D12Adapter.cpp
--- a/Engine/Source/Runtime/D3D12RHI/Private/D3D12Adapter.cpp
+++ b/Engine/Source/Runtime/D3D12RHI/Private/D3D12Adapter.cpp
@@ -202,6 +202,20 @@ static LONG __stdcall D3DVectoredExceptionHandler(EXCEPTION_POINTERS* InInfo)
 #endif // #if PLATFORM_WINDOWS || PLATFORM_HOLOLENS
 
 
+FTransientUniformBufferAllocator::~FTransientUniformBufferAllocator()
+{
+	if (Adapter)
+	{
+		Adapter->ReleaseTransientUniformBufferAllocator(this);
+	}
+}
+
+void FTransientUniformBufferAllocator::Cleanup()
+{
+	ClearResource();
+	Adapter = nullptr;
+}
+
 FD3D12AdapterDesc::FD3D12AdapterDesc() = default;
 
 FD3D12AdapterDesc::FD3D12AdapterDesc(const DXGI_ADAPTER_DESC& InDesc, int32 InAdapterIndex, const FD3D12DeviceBasicInfo& DeviceInfo)
@@ -1294,10 +1308,11 @@ void FD3D12Adapter::Cleanup()
 	}
 
 	// Release allocation data of all thread local transient uniform buffer allocators
-	for (FD3D12FastConstantAllocator* Allocator : TransientUniformBufferAllocators)
+	for (FTransientUniformBufferAllocator* Allocator : TransientUniformBufferAllocators)
 	{
 		Allocator->Cleanup();
 	}
+	TransientUniformBufferAllocators.Empty();
 
 	// Cleanup resources
 	DeferredDeletionQueue.ReleaseResources(true, true);
@@ -1545,18 +1560,6 @@ FD3D12TemporalEffect* FD3D12Adapter::GetTemporalEffect(const FName& EffectName)
 
 FD3D12FastConstantAllocator& FD3D12Adapter::GetTransientUniformBufferAllocator()
 {
-	class FTransientUniformBufferAllocator : public FD3D12FastConstantAllocator, public TThreadSingleton<FTransientUniformBufferAllocator>
-	{
-	public:
-		FTransientUniformBufferAllocator(FD3D12Adapter* InAdapter, FD3D12Device* Parent, FRHIGPUMask VisibiltyMask) : FD3D12FastConstantAllocator(Parent, VisibiltyMask), Adapter(InAdapter) {}
-		~FTransientUniformBufferAllocator()
-		{
-			Adapter->ReleaseTransientUniformBufferAllocator(this);
-		}
-	private:
-		FD3D12Adapter* Adapter;
-	};
-
 	// Multi-GPU support : is using device 0 always appropriate here?
 	return FTransientUniformBufferAllocator::Get([this]() -> FTransientUniformBufferAllocator*
 	{
@@ -1572,7 +1575,7 @@ FD3D12FastConstantAllocator& FD3D12Adapter::GetTransientUniformBufferAllocator()
 	});
 }
 
-void FD3D12Adapter::ReleaseTransientUniformBufferAllocator(FD3D12FastConstantAllocator* InAllocator)
+void FD3D12Adapter::ReleaseTransientUniformBufferAllocator(FTransientUniformBufferAllocator* InAllocator)
 {
 	FScopeLock Lock(&TransientUniformBufferAllocatorsCS);
 	verify(TransientUniformBufferAllocators.Remove(InAllocator) == 1);
diff --git a/Engine/Source/Runtime/D3D12RHI/Private/D3D12Adapter.h b/Engine/Source/Runtime/D3D12RHI/Private/D3D12Adapter.h
--- a/Engine/Source/Runtime/D3D12RHI/Private/D3D12Adapter.h
+++ b/Engine/Source/Runtime/D3D12RHI/Private/D3D12Adapter.h
@@ -104,6 +104,17 @@ struct FD3D12MemoryInfo
 	}
 };
 
+class FTransientUniformBufferAllocator : public FD3D12FastConstantAllocator, public TThreadSingleton<FTransientUniformBufferAllocator>
+{
+public:
+	FTransientUniformBufferAllocator(FD3D12Adapter* InAdapter, FD3D12Device* Parent, FRHIGPUMask VisibiltyMask) : FD3D12FastConstantAllocator(Parent, VisibiltyMask), Adapter(InAdapter) {}
+	~FTransientUniformBufferAllocator();
+
+	void Cleanup();
+private:
+	FD3D12Adapter* Adapter = nullptr;
+};
+
 enum class ED3D12GPUCrashDebuggingModes
 {
 	None				= 0x0,
@@ -391,7 +402,7 @@ public:
 	FD3D12TemporalEffect* GetTemporalEffect(const FName& EffectName);
 
 	FD3D12FastConstantAllocator& GetTransientUniformBufferAllocator();
-	void ReleaseTransientUniformBufferAllocator(FD3D12FastConstantAllocator* InAllocator);
+	void ReleaseTransientUniformBufferAllocator(FTransientUniformBufferAllocator* InAllocator);
 
 	void BlockUntilIdle();
 
@@ -588,7 +599,7 @@ protected:
 	TMap<FName, FD3D12TemporalEffect> TemporalEffectMap;
 #endif
 
-	TArray<FD3D12FastConstantAllocator*> TransientUniformBufferAllocators;
+	TArray<FTransientUniformBufferAllocator*> TransientUniformBufferAllocators;
 	FCriticalSection TransientUniformBufferAllocatorsCS;
 
 	TUniquePtr<IRHITransientMemoryCache> TransientMemoryCache;
diff --git a/Engine/Source/Runtime/D3D12RHI/Private/D3D12Allocation.h b/Engine/Source/Runtime/D3D12RHI/Private/D3D12Allocation.h
--- a/Engine/Source/Runtime/D3D12RHI/Private/D3D12Allocation.h
+++ b/Engine/Source/Runtime/D3D12RHI/Private/D3D12Allocation.h
@@ -522,7 +522,7 @@ public:
 	FD3D12FastConstantAllocator(FD3D12Device* Parent, FRHIGPUMask VisibiltyMask);
 
 	void* Allocate(uint32 Bytes, class FD3D12ResourceLocation& OutLocation, FD3D12ConstantBufferView* OutCBView);
-	void Cleanup() { UnderlyingResource.Clear(); }
+	void ClearResource() { UnderlyingResource.Clear(); }
 
 private:
 	FD3D12ResourceLocation UnderlyingResource;
