Use D3D12_WRITEBUFFERIMMEDIATE_MODE_MARKER_IN / OUT instead of just IN for GPU bread crumbs.

This is more accurate when there are multiple GPU commands back-to-back without barriers, however it may be slightly more expensive when many breadcrumbs are inserted.

Also fix MaxParameterCount in  FD3D12CommandContext::WriteGPUEventStackToBreadCrumbData and slightly refactor/tidy FD3D12CommandListManager::FDiagnosticBuffer.

diff --git a/Engine/Source/Runtime/D3D12RHI/Private/D3D12CommandContext.cpp b/Engine/Source/Runtime/D3D12RHI/Private/D3D12CommandContext.cpp
--- a/Engine/Source/Runtime/D3D12RHI/Private/D3D12CommandContext.cpp
+++ b/Engine/Source/Runtime/D3D12RHI/Private/D3D12CommandContext.cpp
@@ -152,15 +152,15 @@ FD3D12CommandContext::~FD3D12CommandContext()
 void FD3D12CommandContext::WriteGPUEventStackToBreadCrumbData(bool bBeginEvent)
 {
 	// Write directly to command list if breadcrumb resource is available
-	FD3D12Resource* BreadCrumbResource = GetCommandListManager().GetBreadCrumbResource();
+	FD3D12Resource* BreadCrumbResource = GetCommandListManager().GetDiagnosticBufferResource();
 	ID3D12GraphicsCommandList2* CommandList2 = CommandListHandle.GraphicsCommandList2();
 	if (BreadCrumbResource && CommandList2)
 	{
 		// Find the max parameter count from the resource
-		int MaxParameterCount = BreadCrumbResource->GetDesc().Width / sizeof(uint32);
+		const int32 MaxParameterCount = GetCommandListManager().GetBreadCrumbsBufferSize() / sizeof(uint32);
 
 		// allocate the parameters on the stack if smaller than 4K
-		int ParameterCount = GPUEventStack.Num() < (MaxParameterCount - 2) ? GPUEventStack.Num() + 2 : MaxParameterCount;
+		int32 ParameterCount = GPUEventStack.Num() < (MaxParameterCount - 2) ? GPUEventStack.Num() + 2 : MaxParameterCount;
 		size_t MemSize = ParameterCount * (sizeof(D3D12_WRITEBUFFERIMMEDIATE_PARAMETER) + sizeof(D3D12_WRITEBUFFERIMMEDIATE_PARAMETER));
 		const bool bAllocateOnStack = (MemSize < 4096);
 		void* Mem = bAllocateOnStack ? FMemory_Alloca(MemSize) : FMemory::Malloc(MemSize);
@@ -172,7 +172,23 @@ void FD3D12CommandContext::WriteGPUEventStackToBreadCrumbData(bool bBeginEvent)
 			for (int i = 0; i < ParameterCount; ++i)
 			{
 				Parameters[i].Dest = BreadCrumbResource->GetGPUVirtualAddress() + 4 * i;
+
+			#if 1 // This is more accurate, but may have a higher theoretical GPU overhead
+				if (bBeginEvent)
+				{
+					// The write operation is guaranteed to occur after all preceding commands in the command stream have started, including previous.
+					// We use this mode because we want to know which ops have started on the GPU.
+					Modes[i] = D3D12_WRITEBUFFERIMMEDIATE_MODE_MARKER_IN;
+				}
+				else
+				{
+					// The write operation is deferred until all previous commands in the command stream have completed through the GPU pipeline, including previous WriteBufferImmediate operations.
+					// We want this mode when ending breadcrumb scopes because the GPU might start the next scope before we hit a problem in the current one in some cases (when there are no barriers).
+					Modes[i] = D3D12_WRITEBUFFERIMMEDIATE_MODE_MARKER_OUT;
+				}
+			#else // This is less accurate (we don't know for sure if the event has finished), but has lower theoretical GPU overhead
 				Modes[i] = D3D12_WRITEBUFFERIMMEDIATE_MODE_MARKER_IN;
+			#endif
 
 				// Write event stack count first
 				if (i == 0)
diff --git a/Engine/Source/Runtime/D3D12RHI/Private/D3D12DirectCommandListManager.cpp b/Engine/Source/Runtime/D3D12RHI/Private/D3D12DirectCommandListManager.cpp
--- a/Engine/Source/Runtime/D3D12RHI/Private/D3D12DirectCommandListManager.cpp
+++ b/Engine/Source/Runtime/D3D12RHI/Private/D3D12DirectCommandListManager.cpp
@@ -424,16 +424,21 @@ void FD3D12CommandListManager::Create(const TCHAR* Name, uint32 NumCommandLists,
 			// Allocate persistent CPU readable memory which will still be valid after a device lost and wrap this data in a placed resource
 			// so the GPU command list can write to it
 			const uint32 EventBufferSize = MaxEventCount * sizeof(uint32);
-			const uint32 BreadCrumbBufferSize = EventBufferSize + ShaderDiagnosticBufferSize;
-
-			DiagnosticBufferOffset = EventBufferSize;
+			const uint32 TotalBufferSize = EventBufferSize + ShaderDiagnosticBufferSize;
 
 			// Create the platform-specific diagnostic buffer
 			TCHAR TempStr[MAX_SPRINTF] = TEXT("");
-			FCString::Sprintf(TempStr, TEXT("BreadCrumbResource_%s"), Name);
+			FCString::Sprintf(TempStr, TEXT("DiagnosticBuffer (%s)"), Name);
+
+			const D3D12_RESOURCE_DESC BufferDesc = CD3DX12_RESOURCE_DESC::Buffer(TotalBufferSize, D3D12_RESOURCE_FLAG_ALLOW_CROSS_ADAPTER);
+			DiagnosticBuffer = CreateDiagnosticBuffer(Adapter, Device, BufferDesc, TempStr);
+
+			// Diagnostic buffer is split between breadcrumb events and diagnostic messages.
+			DiagnosticBuffer.BreadCrumbsOffset = 0;
+			DiagnosticBuffer.BreadCrumbsSize = EventBufferSize;
 
-			const D3D12_RESOURCE_DESC BufferDesc = CD3DX12_RESOURCE_DESC::Buffer(BreadCrumbBufferSize, D3D12_RESOURCE_FLAG_ALLOW_CROSS_ADAPTER);
-			DiagnosticBuffer = CreateDiagnosticBuffer(Adapter, Device, BufferDesc, Name);
+			DiagnosticBuffer.DiagnosticsOffset = DiagnosticBuffer.BreadCrumbsOffset + DiagnosticBuffer.BreadCrumbsSize;
+			DiagnosticBuffer.DiagnosticsSize = ShaderDiagnosticBufferSize;
 		}
 	}
 }
@@ -1162,12 +1167,12 @@ FD3D12CommandListManager::FDiagnosticBuffer FD3D12CommandListManager::CreateDiag
 
 void FD3D12CommandListManager::DestroyDiagnosticBuffer(FDiagnosticBuffer& Buffer)
 {
-	Buffer.BreadCrumbResource.SafeRelease();
-	Buffer.BreadCrumbHeap.SafeRelease();
+	Buffer.Resource.SafeRelease();
+	Buffer.Heap.SafeRelease();
 
-	VirtualFree(Buffer.BreadCrumbResourceAddress, 0, MEM_RELEASE);
-	Buffer.BreadCrumbResourceAddress = nullptr;
-	Buffer.BreadCrumbResourceGPUAddress = 0;
+	VirtualFree(Buffer.CpuAddress, 0, MEM_RELEASE);
+	Buffer.CpuAddress = nullptr;
+	Buffer.GpuAddress = 0;
 }
 
 FD3D12FenceCore* FD3D12FenceCorePool::ObtainFenceCore(uint32 GPUIndex)
diff --git a/Engine/Source/Runtime/D3D12RHI/Private/D3D12DirectCommandListManager.h b/Engine/Source/Runtime/D3D12RHI/Private/D3D12DirectCommandListManager.h
--- a/Engine/Source/Runtime/D3D12RHI/Private/D3D12DirectCommandListManager.h
+++ b/Engine/Source/Runtime/D3D12RHI/Private/D3D12DirectCommandListManager.h
@@ -280,18 +280,31 @@ public:
 	FORCEINLINE FD3D12Fence& GetFence() { check(CommandListFence); return *CommandListFence; }
 
 	/** Get the breadcrumb resource which is written during command context recording */
-	FD3D12Resource* GetBreadCrumbResource() { return DiagnosticBuffer.BreadCrumbResource.GetReference(); }
+	FD3D12Resource* GetDiagnosticBufferResource() { return DiagnosticBuffer.Resource.GetReference(); }
+
+	uint32 GetBreadCrumbsBufferSize() const { return DiagnosticBuffer.BreadCrumbsSize; }
 
 	/** Get the CPU readable data from the breadcrumb data - this data is still valid after the Device is Lost */
-	const void* GetBreadCrumbResourceAddress() const { return DiagnosticBuffer.BreadCrumbResourceAddress; }
+	const void* GetBreadCrumbBufferData() const
+	{
+		return DiagnosticBuffer.CpuAddress
+			? reinterpret_cast<const uint8*>(DiagnosticBuffer.CpuAddress) + DiagnosticBuffer.BreadCrumbsOffset
+			: nullptr;
+	}
+
 	const FD3D12DiagnosticBufferData* GetDiagnosticBufferData() const
 	{ 
-		const uint8* Address = DiagnosticBuffer.BreadCrumbResourceAddress ? reinterpret_cast<const uint8*>(GetBreadCrumbResourceAddress()) + DiagnosticBufferOffset : nullptr;
+		const uint8* Address = DiagnosticBuffer.CpuAddress
+			? reinterpret_cast<const uint8*>(DiagnosticBuffer.CpuAddress) + DiagnosticBuffer.DiagnosticsOffset
+			: nullptr;
 		return reinterpret_cast<const FD3D12DiagnosticBufferData*>(Address);
 	}
+
 	const D3D12_GPU_VIRTUAL_ADDRESS GetDiagnosticBufferGPUAddress() const
 	{ 
-		return DiagnosticBuffer.BreadCrumbResourceGPUAddress ? DiagnosticBuffer.BreadCrumbResourceGPUAddress + DiagnosticBufferOffset : 0;
+		return DiagnosticBuffer.GpuAddress
+			? DiagnosticBuffer.GpuAddress + DiagnosticBuffer.DiagnosticsOffset
+			: 0;
 	}
 
 	void WaitForCommandQueueFlush();
@@ -381,18 +394,23 @@ protected:
 	// Helper data used to track GPU progress on this command queue
 	struct FDiagnosticBuffer
 	{
-		TRefCountPtr<FD3D12Heap> BreadCrumbHeap;
-		TRefCountPtr<FD3D12Resource> BreadCrumbResource;
+		TRefCountPtr<FD3D12Heap> Heap;
+		TRefCountPtr<FD3D12Resource> Resource;
+
+		void* CpuAddress;
+		D3D12_GPU_VIRTUAL_ADDRESS GpuAddress;
+
+		uint32 BreadCrumbsOffset = 0;
+		uint32 BreadCrumbsSize = 0;
 
-		void* BreadCrumbResourceAddress;
-		D3D12_GPU_VIRTUAL_ADDRESS BreadCrumbResourceGPUAddress;
+		uint32 DiagnosticsOffset = 0;
+		uint32 DiagnosticsSize = 0;
 	};
 
 	virtual FDiagnosticBuffer CreateDiagnosticBuffer(FD3D12Adapter *Adapter, FD3D12Device *Device, const D3D12_RESOURCE_DESC& Desc, const TCHAR* Name);
 	void DestroyDiagnosticBuffer(FDiagnosticBuffer& Buffer);
 
 	FDiagnosticBuffer DiagnosticBuffer;
-	uint32 DiagnosticBufferOffset = 0;
 
 #if WITH_PROFILEGPU || D3D12_SUBMISSION_GAP_RECORDER
 	uint64 CmdListTimingQueryBatchTokens[2];
diff --git a/Engine/Source/Runtime/D3D12RHI/Private/D3D12Util.cpp b/Engine/Source/Runtime/D3D12RHI/Private/D3D12Util.cpp
--- a/Engine/Source/Runtime/D3D12RHI/Private/D3D12Util.cpp
+++ b/Engine/Source/Runtime/D3D12RHI/Private/D3D12Util.cpp
@@ -209,7 +209,7 @@ static FString GetD3D12TextureFlagString(uint32 TextureFlags)
 /** Log the GPU progress of the given CommandListManager to the Error log if breadcrumb data is available */
 static bool LogBreadcrumbData(D3D12RHI::FD3DGPUProfiler& GPUProfiler, FD3D12CommandListManager& CommandListManager)
 {
-	uint32* BreadCrumbData = (uint32*)CommandListManager.GetBreadCrumbResourceAddress();
+	uint32* BreadCrumbData = (uint32*)CommandListManager.GetBreadCrumbBufferData();
 	if (BreadCrumbData == nullptr)
 	{
 		return false;
@@ -219,29 +219,29 @@ static bool LogBreadcrumbData(D3D12RHI::FD3DGPUProfiler& GPUProfiler, FD3D12Comm
 	bool bBeginEvent = BreadCrumbData[1] > 0;
 	check(EventCount >= 0 && EventCount < (MAX_GPU_BREADCRUMB_DEPTH - 2));
 
-	FString gpu_progress = FString::Printf(TEXT("[GPUBreadCrumb]\t%s Queue %d - %s"), GetD3DCommandQueueTypeName(CommandListManager.GetQueueType()), 
+	FString GpuProgress = FString::Printf(TEXT("[GPUBreadCrumb]\t%s Queue %d - %s"), GetD3DCommandQueueTypeName(CommandListManager.GetQueueType()), 
 		CommandListManager.GetGPUIndex(), EventCount == 0 ? TEXT("No Data") : (bBeginEvent ? TEXT("Begin: ") : TEXT("End: ")));
 	for (uint32 EventIndex = 0; EventIndex < EventCount; ++EventIndex)
 	{
 		if (EventIndex > 0)
 		{
-			gpu_progress.Append(TEXT(" - "));
+			GpuProgress.Append(TEXT(" - "));
 		}
 
 		// get the crc and try and translate back into a string
-		uint32 event_crc = BreadCrumbData[EventIndex + 2];
-		const FString* event_name = GPUProfiler.FindEventString(event_crc);
-		if (event_name)
+		uint32 EventCrc = BreadCrumbData[EventIndex + 2];
+		const FString* EventName = GPUProfiler.FindEventString(EventCrc);
+		if (EventName)
 		{
-			gpu_progress.Append(*event_name);
+			GpuProgress.Append(*EventName);
 		}
 		else
 		{
-			gpu_progress.Append(TEXT("Unknown Event"));
+			GpuProgress.Append(TEXT("Unknown Event"));
 		}
 	}
 
-	UE_LOG(LogD3D12RHI, Error, TEXT("%s"), *gpu_progress);
+	UE_LOG(LogD3D12RHI, Error, TEXT("%s"), *GpuProgress);
 
 	const FD3D12DiagnosticBufferData* DiagnosticData = CommandListManager.GetDiagnosticBufferData();
 	if (DiagnosticData && DiagnosticData->Counter)
