Fix for assertion on occlusion query batch overrun due to use of outer occlusion query

- CL 19848534 added a query count per batch and asserts that we don't exceed the expected count
  - however, it failed to account for the outer occulsion query option (GInsertOuterOcclusionQuery)
  - increase the query count in the batch to account for the additional outer query when GInsertOuterOcclusionQuery is enabled

- testing
  - ran EngineTest and ShooterGame on an impacted D3D12 platform with GInsertOuterOcclusionQuery enabed

diff --git a/Engine/Source/Runtime/D3D12RHI/Private/D3D12Query.cpp b/Engine/Source/Runtime/D3D12RHI/Private/D3D12Query.cpp
--- a/Engine/Source/Runtime/D3D12RHI/Private/D3D12Query.cpp
+++ b/Engine/Source/Runtime/D3D12RHI/Private/D3D12Query.cpp
@@ -201,6 +201,7 @@ bool FD3D12Device::GetQueryData(FD3D12RenderQuery& Query, bool bWait)
 
 void FD3D12CommandContext::RHIBeginOcclusionQueryBatch(uint32 NumQueriesInBatch)
 {
+	NumQueriesInBatch = RHIConsoleVariables::GInsertOuterOcclusionQuery ? NumQueriesInBatch + 1 : NumQueriesInBatch;	// account for additional outer occlusion query if utilized
 	GetParentDevice()->GetOcclusionQueryHeap(CommandQueueType)->StartQueryBatch(*this, NumQueriesInBatch);
 	if (RHIConsoleVariables::GInsertOuterOcclusionQuery)
 	{
