Skip redundant calls to ID3D12Device5::SetResidencyPriority(), which are extremely expensive

The current priority can be stored in the D3D12 resource private data and redundant priority setting can be skipped.
This is important for buffer allocations coming from CreateRayTracingBuffer().

diff --git a/Engine/Source/Runtime/D3D12RHI/Private/D3D12Adapter.cpp b/Engine/Source/Runtime/D3D12RHI/Private/D3D12Adapter.cpp
--- a/Engine/Source/Runtime/D3D12RHI/Private/D3D12Adapter.cpp
+++ b/Engine/Source/Runtime/D3D12RHI/Private/D3D12Adapter.cpp
@@ -1957,3 +1957,25 @@ void FD3D12Adapter::DumpTrackedAllocationData(FOutputDevice& OutputDevice, bool
 }
 
 #endif // TRACK_RESOURCE_ALLOCATIONS
+
+void FD3D12Adapter::SetResidencyPriority(ID3D12Pageable* Pageable, D3D12_RESIDENCY_PRIORITY HeapPriority, uint32 GPUIndex)
+{
+#if D3D12_RHI_RAYTRACING // ID3D12Device5 is currently tied to DXR support
+	// GUID for our custom private data that holds the current priority
+	static const GUID DataGuid = GUID{ 0xB365961D, 0xA209, 0x4475, { 0x84, 0x9B, 0xDA, 0xFF, 0x90, 0x91, 0x2E, 0xB1 } };
+
+	const uint32 DataSize = uint32(sizeof(HeapPriority));
+	uint32 ExistingDataSize = DataSize;
+	D3D12_RESIDENCY_PRIORITY ExistingPriority = (D3D12_RESIDENCY_PRIORITY)0;
+
+	HRESULT Result = Pageable->GetPrivateData(DataGuid, &ExistingDataSize, &ExistingPriority);
+
+	if (!SUCCEEDED(Result) || ExistingDataSize != DataSize || ExistingPriority != HeapPriority)
+	{
+		FD3D12Device* NodeDevice = GetDevice(GPUIndex);
+		NodeDevice->GetDevice5()->SetResidencyPriority(1, &Pageable, &HeapPriority);
+		Pageable->SetPrivateData(DataGuid, DataSize, &HeapPriority);
+	}
+#endif // D3D12_RHI_RAYTRACING
+}
+
diff --git a/Engine/Source/Runtime/D3D12RHI/Private/D3D12Adapter.h b/Engine/Source/Runtime/D3D12RHI/Private/D3D12Adapter.h
--- a/Engine/Source/Runtime/D3D12RHI/Private/D3D12Adapter.h
+++ b/Engine/Source/Runtime/D3D12RHI/Private/D3D12Adapter.h
@@ -446,6 +446,8 @@ public:
 	void SubmitGapRecorderTimestamps();
 #endif
 
+	void SetResidencyPriority(ID3D12Pageable* Pageable, D3D12_RESIDENCY_PRIORITY HeapPriority, uint32 GPUIndex);
+
 protected:
 
 	virtual void CreateRootDevice(bool bWithDebug);
diff --git a/Engine/Source/Runtime/D3D12RHI/Private/D3D12PoolAllocator.cpp b/Engine/Source/Runtime/D3D12RHI/Private/D3D12PoolAllocator.cpp
--- a/Engine/Source/Runtime/D3D12RHI/Private/D3D12PoolAllocator.cpp
+++ b/Engine/Source/Runtime/D3D12RHI/Private/D3D12PoolAllocator.cpp
@@ -397,9 +397,7 @@ void FD3D12PoolAllocator::AllocateResource(uint32 GPUIndex, D3D12_HEAP_TYPE InHe
 		&& ResourceLocation.GetAllocatorType() != FD3D12ResourceLocation::AT_Pool)
 	{
 		ID3D12Pageable* HeapResource = ResourceLocation.GetResource()->GetPageable();
-		D3D12_RESIDENCY_PRIORITY HeapPriority = D3D12_RESIDENCY_PRIORITY_HIGH;
-		FD3D12Device* NodeDevice = Adapter->GetDevice(GPUIndex);
-		NodeDevice->GetDevice5()->SetResidencyPriority(1, &HeapResource, &HeapPriority);
+		Adapter->SetResidencyPriority(HeapResource, D3D12_RESIDENCY_PRIORITY_HIGH, GPUIndex);
 	}
 #endif // D3D12_RHI_RAYTRACING
 }
diff --git a/Engine/Source/Runtime/D3D12RHI/Private/D3D12RayTracing.cpp b/Engine/Source/Runtime/D3D12RHI/Private/D3D12RayTracing.cpp
--- a/Engine/Source/Runtime/D3D12RHI/Private/D3D12RayTracing.cpp
+++ b/Engine/Source/Runtime/D3D12RHI/Private/D3D12RayTracing.cpp
@@ -3313,9 +3313,7 @@ static TRefCountPtr<FD3D12Buffer> CreateRayTracingBuffer(FD3D12Adapter* Adapter,
 		// Elevates the scratch buffer heap priority, which may help performance / stability in low memory conditions 
 		// (Acceleration structure already boosted from allocation side)
 		ID3D12Pageable* HeapResource = Result->GetResource()->GetPageable();
-		D3D12_RESIDENCY_PRIORITY HeapPriority = D3D12_RESIDENCY_PRIORITY_HIGH;
-		FD3D12Device* NodeDevice = Adapter->GetDevice(GPUIndex);
-		NodeDevice->GetDevice5()->SetResidencyPriority(1, &HeapResource, &HeapPriority);
+		Adapter->SetResidencyPriority(HeapResource, D3D12_RESIDENCY_PRIORITY_HIGH, GPUIndex);
 	}
 	else
 	{
diff --git a/Engine/Source/Runtime/D3D12RHI/Private/D3D12TransientResourceAllocator.cpp b/Engine/Source/Runtime/D3D12RHI/Private/D3D12TransientResourceAllocator.cpp
--- a/Engine/Source/Runtime/D3D12RHI/Private/D3D12TransientResourceAllocator.cpp
+++ b/Engine/Source/Runtime/D3D12RHI/Private/D3D12TransientResourceAllocator.cpp
@@ -75,9 +75,7 @@ FD3D12TransientHeap::FD3D12TransientHeap(const FInitializer& Initializer, FD3D12
 		TRefCountPtr<ID3D12Device5> D3DDevice5;
 		if (SUCCEEDED(D3DDevice->QueryInterface(IID_PPV_ARGS(D3DDevice5.GetInitReference()))))
 		{
-			ID3D12Pageable* Pageable = D3DHeap;
-			D3D12_RESIDENCY_PRIORITY HeapPriority = D3D12_RESIDENCY_PRIORITY_HIGH;
-			D3DDevice5->SetResidencyPriority(1, &Pageable, &HeapPriority);
+			Adapter->SetResidencyPriority(D3DHeap, D3D12_RESIDENCY_PRIORITY_HIGH, Device->GetGPUIndex());
 		}
 #endif // PLATFORM_WINDOWS
 	}
