Only calculate hit group parameters cache on command context for GPU 0.

diff --git a/Engine/Source/Runtime/D3D12RHI/Private/D3D12RayTracing.cpp b/Engine/Source/Runtime/D3D12RHI/Private/D3D12RayTracing.cpp
--- a/Engine/Source/Runtime/D3D12RHI/Private/D3D12RayTracing.cpp
+++ b/Engine/Source/Runtime/D3D12RHI/Private/D3D12RayTracing.cpp
@@ -803,11 +803,11 @@ void FD3D12RayTracingCompactionRequestHandler::Update(FD3D12CommandContext& InCo
 	}
 
 	// build a new set of build requests to extract the build data	
-	for (FD3D12RayTracingGeometry* RTGeomtry : PendingRequests)
+	for (FD3D12RayTracingGeometry* RTGeometry : PendingRequests)
 	{
-		ActiveRequests.Add(RTGeomtry);
+		ActiveRequests.Add(RTGeometry);
 
-		FD3D12ResourceLocation& ResourceLocation = RTGeomtry->AccelerationStructureBuffers[GPUIndex].GetReference()->ResourceLocation;
+		FD3D12ResourceLocation& ResourceLocation = RTGeometry->AccelerationStructureBuffers[GPUIndex].GetReference()->ResourceLocation;
 		ActiveBLASGPUAddresses.Add(ResourceLocation.GetGPUVirtualAddress());
 
 		ResourceLocation.GetResource()->UpdateResidency(InCommandContext.CommandListHandle);
@@ -3995,6 +3995,8 @@ void FD3D12RayTracingScene::ReleaseBuffer()
 			DEC_MEMORY_STAT_BY(STAT_D3D12RayTracingUsedVideoMemory, AccelerationStructureBuffer->GetSize());
 			DEC_MEMORY_STAT_BY(STAT_D3D12RayTracingTLASMemory, AccelerationStructureBuffer->GetSize());
 		}
+
+		AccelerationStructureBuffer = nullptr;
 	}
 	ShaderResourceView = nullptr;
 }
@@ -4076,18 +4078,19 @@ void FD3D12RayTracingScene::BuildAccelerationStructure(FD3D12CommandContext& Com
 		// - Set up acceleration structure pointers and make them resident.
 		// - Generate HitGroupSystemParametersCache.
 
-		HitGroupSystemParametersCache.Reserve(Initializer.NumTotalSegments);
-		HitGroupSystemParametersCache.Empty(Initializer.NumTotalSegments);
-
+		// make a copy of system parameters to they can optimized fetch during SBT building (only done for GPU0)
+		if(GPUIndex == 0)
 		{
+			HitGroupSystemParametersCache.Reserve(Initializer.NumTotalSegments);
+			HitGroupSystemParametersCache.Empty(Initializer.NumTotalSegments);
+
 			const int32 NumSceneInstances = Initializer.PerInstanceGeometries.Num();
 			for (int32 InstanceIndex = 0; InstanceIndex < NumSceneInstances; ++InstanceIndex)
 			{
 				FD3D12RayTracingGeometry* Geometry = FD3D12DynamicRHI::ResourceCast(Initializer.PerInstanceGeometries[InstanceIndex]);
 
-				// make a copy of system parameters to they can optimized fetch during SBT building (only done for GPU0)
-				check(Geometry->HitGroupSystemParameters[0].Num() > 0);
-				HitGroupSystemParametersCache.Append(Geometry->HitGroupSystemParameters[0]);
+				check(Geometry->HitGroupSystemParameters[GPUIndex].Num() > 0);
+				HitGroupSystemParametersCache.Append(Geometry->HitGroupSystemParameters[GPUIndex]);
 			}
 		}
 
