Support multiple layers in FRHIRayTracingScene.
- Each layer corresponds to a separate TLAS.
- Values returned by FRHIRayTracingScene::GetSizeInfo() are large enough to contain all layers.
- Higher level code can use FRHIRayTracingScene::GetLayerBufferOffset(...) to query the offset of individual layer TLAS when building SRVs.
- Updated FRayTracingSceneInitializer2 so higher level code specifies how many instances are in each layer.

diff --git a/Engine/Plugins/Experimental/GPULightmass/Source/GPULightmass/Private/LightmapRenderer.cpp b/Engine/Plugins/Experimental/GPULightmass/Source/GPULightmass/Private/LightmapRenderer.cpp
--- a/Engine/Plugins/Experimental/GPULightmass/Source/GPULightmass/Private/LightmapRenderer.cpp
+++ b/Engine/Plugins/Experimental/GPULightmass/Source/GPULightmass/Private/LightmapRenderer.cpp
@@ -1070,7 +1070,7 @@ bool FSceneRenderState::SetupRayTracingScene(int32 LODIndex)
 
 			const FRayTracingSceneInitializer2& SceneInitializer = RayTracingScene->GetInitializer();
 
-			FRayTracingAccelerationStructureSize SizeInfo = RHICalcRayTracingSceneSize(SceneInitializer.NumNativeInstances, ERayTracingAccelerationStructureFlags::FastTrace);
+			FRayTracingAccelerationStructureSize SizeInfo = RHICalcRayTracingSceneSize(SceneInitializer.NumNativeInstancesPerLayer[0], ERayTracingAccelerationStructureFlags::FastTrace);
 			FRHIResourceCreateInfo BufferCreateInfo(TEXT("LightmassRayTracingSceneBuffer"));
 			RayTracingSceneBuffer = RHICreateBuffer(uint32(SizeInfo.ResultSize), BUF_AccelerationStructure, 0, ERHIAccess::BVHWrite, BufferCreateInfo);
 			RayTracingSceneSRV = RHICreateShaderResourceView(RayTracingSceneBuffer);
@@ -1082,7 +1082,7 @@ bool FSceneRenderState::SetupRayTracingScene(int32 LODIndex)
 				ScratchBufferCreateInfo);
 
 			FRWBufferStructured InstanceBuffer;
-			InstanceBuffer.Initialize(TEXT("LightmassRayTracingInstanceBuffer"), GRHIRayTracingInstanceDescriptorSize, SceneInitializer.NumNativeInstances);
+			InstanceBuffer.Initialize(TEXT("LightmassRayTracingInstanceBuffer"), GRHIRayTracingInstanceDescriptorSize, SceneInitializer.NumNativeInstancesPerLayer[0]);
 
 			// Need to pass "BUF_MultiGPUAllocate", as acceleration structure virtual addresses are different per GPU
 			FByteAddressBuffer AccelerationStructureAddressesBuffer;
@@ -1091,7 +1091,7 @@ bool FSceneRenderState::SetupRayTracingScene(int32 LODIndex)
 				SceneInitializer.ReferencedGeometries.Num() * sizeof(FRayTracingAccelerationStructureAddress),
 				BUF_Volatile | BUF_MultiGPUAllocate);
 
-			const uint32 InstanceUploadBufferSize = SceneInitializer.NumNativeInstances * sizeof(FRayTracingInstanceDescriptorInput);
+			const uint32 InstanceUploadBufferSize = SceneInitializer.NumNativeInstancesPerLayer[0] * sizeof(FRayTracingInstanceDescriptorInput);
 			FBufferRHIRef InstanceUploadBuffer;
 			FShaderResourceViewRHIRef InstanceUploadSRV;
 			{
@@ -1120,7 +1120,7 @@ bool FSceneRenderState::SetupRayTracingScene(int32 LODIndex)
 					SceneWithGeometryInstances.BaseUploadBufferOffsets,
 					SceneWithGeometryInstances.NumNativeGPUSceneInstances,
 					SceneWithGeometryInstances.NumNativeCPUInstances,
-					MakeArrayView(InstanceUploadData, SceneInitializer.NumNativeInstances),
+					MakeArrayView(InstanceUploadData, SceneInitializer.NumNativeInstancesPerLayer[0]),
 					MakeArrayView(TransformUploadData, SceneWithGeometryInstances.NumNativeCPUInstances * 3));
 				RHICmdList.UnlockBuffer(TransformUploadBuffer);
 				RHICmdList.UnlockBuffer(InstanceUploadBuffer);
diff --git a/Engine/Source/Runtime/D3D12RHI/Private/D3D12RayTracing.cpp b/Engine/Source/Runtime/D3D12RHI/Private/D3D12RayTracing.cpp
--- a/Engine/Source/Runtime/D3D12RHI/Private/D3D12RayTracing.cpp
+++ b/Engine/Source/Runtime/D3D12RHI/Private/D3D12RayTracing.cpp
@@ -3959,16 +3959,34 @@ FD3D12RayTracingScene::FD3D12RayTracingScene(FD3D12Adapter* Adapter, FRayTracing
 	checkf(Initializer.NumMissShaderSlots >= 1, TEXT("Need at least 1 miss shader slot."));
 	checkf(Initializer.Lifetime == RTSL_SingleFrame, TEXT("Only single-frame ray tracing scenes are currently implemented."));
 
-	ERayTracingAccelerationStructureFlags BuildFlags = ERayTracingAccelerationStructureFlags::FastTrace; // #yuriy_todo: pass this in
+	const ERayTracingAccelerationStructureFlags BuildFlags = ERayTracingAccelerationStructureFlags::FastTrace; // #yuriy_todo: pass this in
 
-	BuildInputs = {};
-	BuildInputs.Type = D3D12_RAYTRACING_ACCELERATION_STRUCTURE_TYPE_TOP_LEVEL;
-	BuildInputs.DescsLayout = D3D12_ELEMENTS_LAYOUT_ARRAY;
-	BuildInputs.NumDescs = Initializer.NumNativeInstances;
-	BuildInputs.Flags = TranslateRayTracingAccelerationStructureFlags(BuildFlags);
+	SizeInfo = {};
+
+	const uint32 NumLayers = Initializer.NumNativeInstancesPerLayer.Num();
+	check(NumLayers > 0);
+
+	Layers.SetNum(NumLayers);
+
+	for (uint32 LayerIndex = 0; LayerIndex < NumLayers; ++LayerIndex)
+	{
+		FLayerData& Layer = Layers[LayerIndex];
 
-	// Get maximum buffer sizes for all GPUs in the system
-	SizeInfo = RHICalcRayTracingSceneSize(BuildInputs.NumDescs, BuildFlags);
+		Layer.BuildInputs = {};
+		Layer.BuildInputs.Type = D3D12_RAYTRACING_ACCELERATION_STRUCTURE_TYPE_TOP_LEVEL;
+		Layer.BuildInputs.DescsLayout = D3D12_ELEMENTS_LAYOUT_ARRAY;
+		Layer.BuildInputs.NumDescs = Initializer.NumNativeInstancesPerLayer[LayerIndex];
+		Layer.BuildInputs.Flags = TranslateRayTracingAccelerationStructureFlags(BuildFlags);
+
+		// Get maximum buffer sizes for all GPUs in the system
+		Layer.SizeInfo = RHICalcRayTracingSceneSize(Layer.BuildInputs.NumDescs, BuildFlags);
+
+		Layer.BufferOffset = Align(SizeInfo.ResultSize, GRHIRayTracingAccelerationStructureAlignment);
+		Layer.ScratchBufferOffset = Align(SizeInfo.BuildScratchSize, GRHIRayTracingScratchBufferAlignment);
+
+		SizeInfo.ResultSize = Layer.BufferOffset + Layer.SizeInfo.ResultSize;
+		SizeInfo.BuildScratchSize = Layer.BufferOffset + Layer.SizeInfo.BuildScratchSize;
+	}
 };
 
 FD3D12RayTracingScene::~FD3D12RayTracingScene()
@@ -3998,7 +4016,11 @@ void FD3D12RayTracingScene::ReleaseBuffer()
 
 		AccelerationStructureBuffer = nullptr;
 	}
-	ShaderResourceView = nullptr;
+
+	for (auto& Layer : Layers)
+	{
+		Layer.ShaderResourceView = nullptr;
+	}
 }
 
 void FD3D12RayTracingScene::BindBuffer(FRHIBuffer* InBuffer, uint32 InBufferOffset)
@@ -4020,8 +4042,14 @@ void FD3D12RayTracingScene::BindBuffer(FRHIBuffer* InBuffer, uint32 InBufferOffs
 	}
 	BufferOffset = InBufferOffset;
 
-	FShaderResourceViewInitializer ViewInitializer(InBuffer, InBufferOffset, 0);
-	ShaderResourceView = RHICreateShaderResourceView(ViewInitializer);
+	for (auto& Layer : Layers)
+	{
+		const uint32 LayerOffset = InBufferOffset + Layer.BufferOffset;
+		check(LayerOffset % GRHIRayTracingAccelerationStructureAlignment == 0);
+
+		FShaderResourceViewInitializer ViewInitializer(InBuffer, LayerOffset, 0);
+		Layer.ShaderResourceView = RHICreateShaderResourceView(ViewInitializer);
+	}
 }
 
 void FD3D12RayTracingScene::BuildAccelerationStructure(FD3D12CommandContext& CommandContext,
@@ -4037,43 +4065,43 @@ void FD3D12RayTracingScene::BuildAccelerationStructure(FD3D12CommandContext& Com
 	FD3D12Adapter* Adapter = CommandContext.GetParentAdapter();
 	ID3D12Device5* RayTracingDevice = CommandContext.GetParentDevice()->GetDevice5();
 
-	D3D12_RAYTRACING_ACCELERATION_STRUCTURE_PREBUILD_INFO PrebuildInfo = {};
-
-	RayTracingDevice->GetRaytracingAccelerationStructurePrebuildInfo(&BuildInputs, &PrebuildInfo);
-
-	checkf(PrebuildInfo.ResultDataMaxSizeInBytes <= SizeInfo.ResultSize,
-		TEXT("TLAS build result buffer now requires %lld bytes, but only %lld was calculated in the constructor."),
-			PrebuildInfo.ResultDataMaxSizeInBytes, SizeInfo.ResultSize);
-
-	checkf(PrebuildInfo.ScratchDataSizeInBytes <= SizeInfo.BuildScratchSize,
-		TEXT("TLAS build scratch buffer now requires %lld bytes, but only %lld was calculated in the constructor."),
-			PrebuildInfo.ScratchDataSizeInBytes, SizeInfo.BuildScratchSize);
-
 	TRefCountPtr<FD3D12Buffer> AutoScratchBuffer;
 	if (ScratchBuffer == nullptr)
 	{
 		static const FName ScratchBufferName("AutoBuildScratchTLAS");
-		AutoScratchBuffer = CreateRayTracingBuffer(Adapter, GPUIndex, PrebuildInfo.ScratchDataSizeInBytes, ERayTracingBufferType::Scratch, ScratchBufferName);
+		AutoScratchBuffer = CreateRayTracingBuffer(Adapter, GPUIndex, SizeInfo.BuildScratchSize, ERayTracingBufferType::Scratch, ScratchBufferName);
 		ScratchBuffer = AutoScratchBuffer.GetReference();
 		ScratchBufferOffset = 0;
 	}
 
-	checkf(ScratchBuffer, TEXT("TLAS build requires scratch buffer of at least %lld bytes."), PrebuildInfo.ScratchDataSizeInBytes);
+	checkf(ScratchBuffer, TEXT("TLAS build requires scratch buffer of at least %lld bytes."), SizeInfo.BuildScratchSize);
 
-	const D3D12_GPU_VIRTUAL_ADDRESS ScratchAddress = ScratchBuffer->ResourceLocation.GetGPUVirtualAddress() + ScratchBufferOffset;
+	const uint32 NumLayers = Initializer.NumNativeInstancesPerLayer.Num();
 
-	checkf(ScratchAddress % GRHIRayTracingAccelerationStructureAlignment == 0,
-		TEXT("TLAS scratch buffer (plus offset) must be aligned to %lld bytes."),
-		GRHIRayTracingAccelerationStructureAlignment);
+	for (uint32 LayerIndex = 0; LayerIndex < NumLayers; ++LayerIndex)
+	{
+		const FLayerData& Layer = Layers[LayerIndex];
 
-	checkf(PrebuildInfo.ScratchDataSizeInBytes + ScratchBufferOffset <= ScratchBuffer->GetSize(),
-		TEXT("TLAS scratch buffer size is %lld bytes with offset %lld (%lld bytes available), but the build requires %lld bytes. ")
-		TEXT("BuildInputs.NumDescs = %d, Instances.Num = %d, PerInstanceGeometries.Num = %d."),
-			ScratchBuffer->GetSize(), ScratchBufferOffset, ScratchBuffer->GetSize() - ScratchBufferOffset, 
+		D3D12_RAYTRACING_ACCELERATION_STRUCTURE_PREBUILD_INFO PrebuildInfo = {};
+
+		RayTracingDevice->GetRaytracingAccelerationStructurePrebuildInfo(&Layer.BuildInputs, &PrebuildInfo);
+
+		checkf(PrebuildInfo.ResultDataMaxSizeInBytes <= Layer.SizeInfo.ResultSize,
+			TEXT("TLAS build result buffer now requires %lld bytes, but only %lld was calculated in the constructor."),
+			PrebuildInfo.ResultDataMaxSizeInBytes, Layer.SizeInfo.ResultSize);
+
+		checkf(PrebuildInfo.ScratchDataSizeInBytes <= Layer.SizeInfo.BuildScratchSize,
+			TEXT("TLAS build scratch buffer now requires %lld bytes, but only %lld was calculated in the constructor."),
+			PrebuildInfo.ScratchDataSizeInBytes, Layer.SizeInfo.BuildScratchSize);
+
+		checkf(PrebuildInfo.ScratchDataSizeInBytes + ScratchBufferOffset <= ScratchBuffer->GetSize(),
+			TEXT("TLAS scratch buffer size is %lld bytes with offset %lld (%lld bytes available), but the build requires %lld bytes. ")
+			TEXT("BuildInputs.NumDescs = %d, Instances.Num = %d, PerInstanceGeometries.Num = %d."),
+			ScratchBuffer->GetSize(), ScratchBufferOffset, ScratchBuffer->GetSize() - ScratchBufferOffset,
 			PrebuildInfo.ScratchDataSizeInBytes,
-			BuildInputs.NumDescs, Initializer.NumNativeInstances, Initializer.PerInstanceGeometries.Num());
+			Layer.BuildInputs.NumDescs, Initializer.NumNativeInstancesPerLayer[LayerIndex], Initializer.PerInstanceGeometries.Num());
+	}
 
-	if (BuildInputs.NumDescs)
 	{
 		// - Set up acceleration structure pointers and make them resident.
 		// - Generate HitGroupSystemParametersCache.
@@ -4165,40 +4193,59 @@ void FD3D12RayTracingScene::BuildAccelerationStructure(FD3D12CommandContext& Com
 	AccelerationStructureBuffer->GetResource()->UpdateResidency(CommandContext.CommandListHandle);
 	ScratchBuffer->GetResource()->UpdateResidency(CommandContext.CommandListHandle);
 
-	if (InstanceBuffer)
-	{
-		// Enqueue transition to UAV/SRV
-		FD3D12DynamicRHI::TransitionResource(CommandContext.CommandListHandle, InstanceBuffer->GetResource(),
-			D3D12_RESOURCE_STATE_TBD, D3D12_RESOURCE_STATE_NON_PIXEL_SHADER_RESOURCE, 0, FD3D12DynamicRHI::ETransitionMode::Apply);
-	}
-
-	D3D12_BUILD_RAYTRACING_ACCELERATION_STRUCTURE_DESC BuildDesc = {};
-	BuildDesc.Inputs = BuildInputs;
-	BuildDesc.Inputs.InstanceDescs = InstanceBuffer->ResourceLocation.GetGPUVirtualAddress() + InstanceBufferOffset;
-	BuildDesc.DestAccelerationStructureData = AccelerationStructureBuffer->ResourceLocation.GetGPUVirtualAddress() + BufferOffset;
-	BuildDesc.ScratchAccelerationStructureData = ScratchAddress;
-	BuildDesc.SourceAccelerationStructureData = D3D12_GPU_VIRTUAL_ADDRESS(0); // Null source TLAS as this is a build command
-
-	checkf(BuildDesc.ScratchAccelerationStructureData % 256 == 0, TEXT("TLAS build scratch buffer must have 256 byte alignment."));
+	// Enqueue transition to UAV/SRV
+	FD3D12DynamicRHI::TransitionResource(CommandContext.CommandListHandle, InstanceBuffer->GetResource(),
+		D3D12_RESOURCE_STATE_TBD, D3D12_RESOURCE_STATE_NON_PIXEL_SHADER_RESOURCE, 0, FD3D12DynamicRHI::ETransitionMode::Apply);
 
 	if (GD3D12RayTracingGPUValidation)
 	{
-		TRHICommandList_RecursiveHazardous<FD3D12CommandContext> RHICmdList(&CommandContext, CommandContext.GetGPUMask());
-		uint32 InstanceBufferStride = uint32(sizeof(D3D12_RAYTRACING_INSTANCE_DESC));
-		uint32 TotalHitGroupSlots = Initializer.NumTotalSegments * Initializer.ShaderSlotsPerGeometrySegment;
-		FRayTracingValidateSceneBuildParamsCS::Dispatch(RHICmdList, 
-			TotalHitGroupSlots, BuildInputs.NumDescs,
-			InstanceBuffer, InstanceBufferOffset, InstanceBufferStride);
+		for (uint32 LayerIndex = 0; LayerIndex < NumLayers; ++LayerIndex)
+		{
+			const FLayerData& Layer = Layers[LayerIndex];
+			TRHICommandList_RecursiveHazardous<FD3D12CommandContext> RHICmdList(&CommandContext, CommandContext.GetGPUMask());
+			uint32 InstanceBufferStride = GRHIRayTracingInstanceDescriptorSize;
+			uint32 TotalHitGroupSlots = Initializer.NumTotalSegments * Initializer.ShaderSlotsPerGeometrySegment;
+			FRayTracingValidateSceneBuildParamsCS::Dispatch(RHICmdList,
+				TotalHitGroupSlots, Layer.BuildInputs.NumDescs,
+				InstanceBuffer, InstanceBufferOffset, InstanceBufferStride);
+		}
 	}
 
 	// UAV barrier is used here to ensure that all bottom level acceleration structures are built
 	CommandContext.CommandListHandle.AddUAVBarrier();
 	CommandContext.CommandListHandle.FlushResourceBarriers();
 
-	ID3D12GraphicsCommandList4* RayTracingCommandList = CommandContext.CommandListHandle.RayTracingCommandList();
-	RayTracingCommandList->BuildRaytracingAccelerationStructure(&BuildDesc, 0, nullptr);
+	uint32 InstanceBaseOffset = 0;
+
+	for(uint32 LayerIndex = 0; LayerIndex < NumLayers; ++LayerIndex)
+	{
+		const FLayerData& Layer = Layers[LayerIndex];
+
+		const D3D12_GPU_VIRTUAL_ADDRESS BufferAddress = AccelerationStructureBuffer->ResourceLocation.GetGPUVirtualAddress() + BufferOffset + Layer.BufferOffset;
+		const D3D12_GPU_VIRTUAL_ADDRESS ScratchAddress = ScratchBuffer->ResourceLocation.GetGPUVirtualAddress() + ScratchBufferOffset + Layer.ScratchBufferOffset;
+
+		checkf(BufferAddress % GRHIRayTracingAccelerationStructureAlignment == 0,
+			TEXT("TLAS  buffer (plus offset) must be aligned to %lld bytes."),
+			GRHIRayTracingAccelerationStructureAlignment);
+
+		checkf(ScratchAddress % GRHIRayTracingScratchBufferAlignment == 0,
+			TEXT("TLAS scratch buffer (plus offset) must be aligned to %lld bytes."),
+			GRHIRayTracingScratchBufferAlignment);
 
-	INC_DWORD_STAT(STAT_D3D12RayTracingBuiltTLAS);
+		D3D12_BUILD_RAYTRACING_ACCELERATION_STRUCTURE_DESC BuildDesc = {};
+		BuildDesc.Inputs = Layer.BuildInputs;
+		BuildDesc.Inputs.InstanceDescs = InstanceBuffer->ResourceLocation.GetGPUVirtualAddress() + InstanceBufferOffset + InstanceBaseOffset * GRHIRayTracingInstanceDescriptorSize;
+		BuildDesc.DestAccelerationStructureData = BufferAddress;
+		BuildDesc.ScratchAccelerationStructureData = ScratchAddress;
+		BuildDesc.SourceAccelerationStructureData = D3D12_GPU_VIRTUAL_ADDRESS(0); // Null source TLAS as this is a build command
+
+		ID3D12GraphicsCommandList4* RayTracingCommandList = CommandContext.CommandListHandle.RayTracingCommandList();
+		RayTracingCommandList->BuildRaytracingAccelerationStructure(&BuildDesc, 0, nullptr);
+
+		INC_DWORD_STAT(STAT_D3D12RayTracingBuiltTLAS);
+
+		InstanceBaseOffset += Initializer.NumNativeInstancesPerLayer[LayerIndex];
+	}
 
 	// UAV barrier is used here to ensure that the acceleration structure build is complete before any rays are traced
 	// #dxr_todo: these barriers should ideally be inserted by the high level code to allow more overlapped execution
@@ -5283,7 +5330,7 @@ void FD3D12CommandContext::RHIRayTraceOcclusion(FRHIRayTracingScene* InScene,
 	DispatchDesc.Depth = 1;
 
 	FRayTracingShaderBindings Bindings;
-	Bindings.SRVs[0] = Scene->ShaderResourceView;
+	Bindings.SRVs[0] = Scene->Layers[0].ShaderResourceView; // TODO: Expose layers
 	Bindings.SRVs[1] = Rays;
 	Bindings.UAVs[0] = Output;
 
@@ -5320,7 +5367,7 @@ void FD3D12CommandContext::RHIRayTraceIntersection(FRHIRayTracingScene* InScene,
 	DispatchDesc.Depth = 1;
 
 	FRayTracingShaderBindings Bindings;
-	Bindings.SRVs[0] = Scene->ShaderResourceView;
+	Bindings.SRVs[0] = Scene->Layers[0].ShaderResourceView; // TODO: Expose layers
 	Bindings.SRVs[1] = Rays;
 	// #dxr_todo: intersection and occlusion shaders should be split into separate files to avoid resource slot collisions.
 	// Workaround for now is to bind a valid UAV to slots 0 and 1, even though only slot 1 is referenced.
diff --git a/Engine/Source/Runtime/D3D12RHI/Private/D3D12RayTracing.h b/Engine/Source/Runtime/D3D12RHI/Private/D3D12RayTracing.h
--- a/Engine/Source/Runtime/D3D12RHI/Private/D3D12RayTracing.h
+++ b/Engine/Source/Runtime/D3D12RHI/Private/D3D12RayTracing.h
@@ -102,6 +102,7 @@ public:
 	~FD3D12RayTracingScene();
 
 	const FRayTracingSceneInitializer2& GetInitializer() const override final { return Initializer; }
+	uint32 GetLayerBufferOffset(uint32 LayerIndex) const override final { return Layers[LayerIndex].BufferOffset; }
 
 	void BindBuffer(FRHIBuffer* Buffer, uint32 BufferOffset);
 	void ReleaseBuffer();
@@ -111,15 +112,20 @@ public:
 		FD3D12Buffer* InstanceBuffer, uint32 InstanceBufferOffset
 	);
 
-	FShaderResourceViewRHIRef ShaderResourceView;
+	struct FLayerData
+	{
+		FShaderResourceViewRHIRef ShaderResourceView;
+		D3D12_BUILD_RAYTRACING_ACCELERATION_STRUCTURE_INPUTS BuildInputs = {};
+		FRayTracingAccelerationStructureSize SizeInfo = {};
+		uint32 BufferOffset;
+		uint32 ScratchBufferOffset;
+	};
+
+	TArray<FLayerData> Layers;
 
 	TRefCountPtr<FD3D12Buffer> AccelerationStructureBuffers[MAX_NUM_GPUS];
 	uint32 BufferOffset = 0;
 
-	D3D12_BUILD_RAYTRACING_ACCELERATION_STRUCTURE_INPUTS BuildInputs = {};
-
-	FRayTracingAccelerationStructureSize SizeInfo = {};
-
 	const FRayTracingSceneInitializer2 Initializer;
 
 	// Scene keeps track of child acceleration structure buffers to ensure
diff --git a/Engine/Source/Runtime/RHI/Public/DynamicRHI.h b/Engine/Source/Runtime/RHI/Public/DynamicRHI.h
--- a/Engine/Source/Runtime/RHI/Public/DynamicRHI.h
+++ b/Engine/Source/Runtime/RHI/Public/DynamicRHI.h
@@ -1647,6 +1647,13 @@ FORCEINLINE FRayTracingGeometryRHIRef RHICreateRayTracingGeometry(const FRayTrac
 
 FORCEINLINE FRayTracingSceneRHIRef RHICreateRayTracingScene(FRayTracingSceneInitializer2 Initializer)
 {
+	PRAGMA_DISABLE_DEPRECATION_WARNINGS
+	if(Initializer.NumNativeInstances > 0 || Initializer.NumNativeInstancesPerLayer.IsEmpty())
+	{
+		checkf(Initializer.NumNativeInstancesPerLayer.IsEmpty(), TEXT("Can't specify NumNativeInstances and NumNativeInstancesPerLayer simultaneously."));
+		Initializer.NumNativeInstancesPerLayer.Add(Initializer.NumNativeInstances);
+	}
+	PRAGMA_ENABLE_DEPRECATION_WARNINGS
 	return GDynamicRHI->RHICreateRayTracingScene(MoveTemp(Initializer));
 }
 
diff --git a/Engine/Source/Runtime/RHI/Public/RHIResources.h b/Engine/Source/Runtime/RHI/Public/RHIResources.h
--- a/Engine/Source/Runtime/RHI/Public/RHIResources.h
+++ b/Engine/Source/Runtime/RHI/Public/RHIResources.h
@@ -2687,7 +2687,10 @@ struct FRayTracingSceneInitializer2
 	// Exclusive prefix sum of instance geometry segments is used to calculate SBT record address from instance and segment indices.
 	TArray<uint32> SegmentPrefixSum;
 
-	// Total flattened number of ray tracing geometry instances (a single FRayTracingGeometryInstance may represent many).
+	// Total flattened number of ray tracing geometry instances (a single FRayTracingGeometryInstance may represent many) per layer.
+	TArray<uint32> NumNativeInstancesPerLayer;
+
+	UE_DEPRECATED(5.1, "Use NumNativeInstancesPerLayer instead.")
 	uint32 NumNativeInstances = 0;
 
 	uint32 NumTotalSegments = 0;
@@ -2713,6 +2716,12 @@ struct FRayTracingSceneInitializer2
 	ERayTracingSceneLifetime Lifetime = RTSL_SingleFrame;
 
 	FName DebugName;
+
+	PRAGMA_DISABLE_DEPRECATION_WARNINGS
+	FRayTracingSceneInitializer2() = default;
+	FRayTracingSceneInitializer2(FRayTracingSceneInitializer2&&) = default;
+	FRayTracingSceneInitializer2& operator=(FRayTracingSceneInitializer2&&) = default;
+	PRAGMA_ENABLE_DEPRECATION_WARNINGS
 };
 
 struct FRayTracingAccelerationStructureSize
@@ -2730,6 +2739,14 @@ class FRHIRayTracingAccelerationStructure
 {
 public:
 	FRHIRayTracingAccelerationStructure() : FRHIResource(RRT_RayTracingAccelerationStructure) {}
+
+	FRayTracingAccelerationStructureSize GetSizeInfo() const
+	{
+		return SizeInfo;
+	};
+
+protected:
+	FRayTracingAccelerationStructureSize SizeInfo = {};
 };
 
 using FRayTracingAccelerationStructureAddress = uint64;
@@ -2756,13 +2773,7 @@ public:
 	{ 
 		return Initializer.Segments.Num(); 
 	}
-
-	FRayTracingAccelerationStructureSize GetSizeInfo() const
-	{
-		return SizeInfo;
-	};
 protected:
-	FRayTracingAccelerationStructureSize SizeInfo = {};
 	FRayTracingGeometryInitializer Initializer = {};
 	ERayTracingGeometryInitializerType InitializedType = ERayTracingGeometryInitializerType::Rendering;
 };
@@ -2782,6 +2793,8 @@ public:
 	{
 		return nullptr;
 	}
+
+	virtual uint32 GetLayerBufferOffset(uint32 LayerIndex) const = 0;
 };
 
 typedef TRefCountPtr<FRHIRayTracingScene>        FRayTracingSceneRHIRef;
diff --git a/Engine/Source/Runtime/Renderer/Private/DeferredShadingRenderer.cpp b/Engine/Source/Runtime/Renderer/Private/DeferredShadingRenderer.cpp
--- a/Engine/Source/Runtime/Renderer/Private/DeferredShadingRenderer.cpp
+++ b/Engine/Source/Runtime/Renderer/Private/DeferredShadingRenderer.cpp
@@ -1618,7 +1618,7 @@ bool FDeferredShadingSceneRenderer::DispatchRayTracingWorldUpdates(FRDGBuilder&
 		#if DO_CHECK
 			{
 				FRayTracingAccelerationStructureSize SizeInfo = RHICalcRayTracingSceneSize(
-					RayTracingSceneRHI->GetInitializer().NumNativeInstances, ERayTracingAccelerationStructureFlags::FastTrace);
+					RayTracingSceneRHI->GetInitializer().NumNativeInstancesPerLayer[0], ERayTracingAccelerationStructureFlags::FastTrace);
 
 				check(SizeInfo.ResultSize <= Scene->RayTracingScene.SizeInfo.ResultSize);
 				check(SizeInfo.BuildScratchSize <= Scene->RayTracingScene.SizeInfo.BuildScratchSize);
diff --git a/Engine/Source/Runtime/Renderer/Private/RayTracing/RayTracingInstanceBufferUtil.cpp b/Engine/Source/Runtime/Renderer/Private/RayTracing/RayTracingInstanceBufferUtil.cpp
--- a/Engine/Source/Runtime/Renderer/Private/RayTracing/RayTracingInstanceBufferUtil.cpp
+++ b/Engine/Source/Runtime/Renderer/Private/RayTracing/RayTracingInstanceBufferUtil.cpp
@@ -40,7 +40,7 @@ FRayTracingSceneWithGeometryInstances CreateRayTracingSceneWithGeometryInstances
 	Initializer.PerInstanceGeometries.SetNumUninitialized(NumSceneInstances);
 	Initializer.BaseInstancePrefixSum.SetNumUninitialized(NumSceneInstances);
 	Initializer.SegmentPrefixSum.SetNumUninitialized(NumSceneInstances);
-	Initializer.NumNativeInstances = 0;
+	Initializer.NumNativeInstancesPerLayer.SetNumZeroed(1);
 	Initializer.NumTotalSegments = 0;
 
 	Experimental::TSherwoodMap<FRHIRayTracingGeometry*, uint32> UniqueGeometries;
@@ -105,8 +105,8 @@ FRayTracingSceneWithGeometryInstances CreateRayTracingSceneWithGeometryInstances
 			Output.NumNativeGPUInstances += InstanceDesc.NumTransforms;
 		}
 
-		Initializer.BaseInstancePrefixSum[InstanceIndex] = Initializer.NumNativeInstances;
-		Initializer.NumNativeInstances += InstanceDesc.NumTransforms;
+		Initializer.BaseInstancePrefixSum[InstanceIndex] = Initializer.NumNativeInstancesPerLayer[0];
+		Initializer.NumNativeInstancesPerLayer[0] += InstanceDesc.NumTransforms;
 	}
 
 	Output.Scene = RHICreateRayTracingScene(MoveTemp(Initializer));
@@ -231,7 +231,7 @@ void FillRayTracingInstanceUploadBuffer(
 #endif
 		});
 
-	SET_DWORD_STAT(STAT_RayTracingInstances, SceneInitializer.NumNativeInstances - NumInactiveNativeInstances);
+	SET_DWORD_STAT(STAT_RayTracingInstances, SceneInitializer.NumNativeInstancesPerLayer[0] - NumInactiveNativeInstances);
 }
 
 struct FRayTracingBuildInstanceBufferCS : public FGlobalShader
diff --git a/Engine/Source/Runtime/Renderer/Private/RayTracing/RayTracingScene.cpp b/Engine/Source/Runtime/Renderer/Private/RayTracing/RayTracingScene.cpp
--- a/Engine/Source/Runtime/Renderer/Private/RayTracing/RayTracingScene.cpp
+++ b/Engine/Source/Runtime/Renderer/Private/RayTracing/RayTracingScene.cpp
@@ -41,10 +41,10 @@ void FRayTracingScene::Create(FRDGBuilder& GraphBuilder, const FGPUScene& GPUSce
 
 	// Round up number of instances to some multiple to avoid pathological growth reallocations.
 	static constexpr uint32 AllocationGranularity = 8 * 1024;
-	const uint32 NumNativeInstancesAligned = FMath::DivideAndRoundUp(FMath::Max(SceneInitializer.NumNativeInstances, 1U), AllocationGranularity) * AllocationGranularity;
+	const uint32 NumNativeInstancesAligned = FMath::DivideAndRoundUp(FMath::Max(SceneInitializer.NumNativeInstancesPerLayer[0], 1U), AllocationGranularity) * AllocationGranularity;
 	const uint32 NumTransformsAligned = FMath::DivideAndRoundUp(FMath::Max(SceneWithGeometryInstances.NumNativeCPUInstances, 1U), AllocationGranularity) * AllocationGranularity;
 
-	SizeInfo = RHICalcRayTracingSceneSize(SceneInitializer.NumNativeInstances, ERayTracingAccelerationStructureFlags::FastTrace);
+	SizeInfo = RHICalcRayTracingSceneSize(SceneInitializer.NumNativeInstancesPerLayer[0], ERayTracingAccelerationStructureFlags::FastTrace);
 	FRayTracingAccelerationStructureSize SizeInfoAligned = RHICalcRayTracingSceneSize(NumNativeInstancesAligned, ERayTracingAccelerationStructureFlags::FastTrace);
 	SizeInfo.ResultSize = FMath::Max(SizeInfo.ResultSize, SizeInfoAligned.ResultSize);
 	SizeInfo.BuildScratchSize = FMath::Max(SizeInfo.BuildScratchSize, SizeInfoAligned.BuildScratchSize);
@@ -126,9 +126,9 @@ void FRayTracingScene::Create(FRDGBuilder& GraphBuilder, const FGPUScene& GPUSce
 		}
 	}
 
-	if (SceneInitializer.NumNativeInstances > 0)
+	if (SceneInitializer.NumNativeInstancesPerLayer[0] > 0)
 	{
-		const uint32 InstanceUploadBytes = SceneInitializer.NumNativeInstances * sizeof(FRayTracingInstanceDescriptorInput);
+		const uint32 InstanceUploadBytes = SceneInitializer.NumNativeInstancesPerLayer[0] * sizeof(FRayTracingInstanceDescriptorInput);
 		const uint32 TransformUploadBytes = SceneWithGeometryInstances.NumNativeCPUInstances * 3 * sizeof(FVector4f);
 
 		FRayTracingInstanceDescriptorInput* InstanceUploadData = (FRayTracingInstanceDescriptorInput*)RHILockBuffer(InstanceUploadBuffer, 0, InstanceUploadBytes, RLM_WriteOnly);
@@ -136,7 +136,7 @@ void FRayTracingScene::Create(FRDGBuilder& GraphBuilder, const FGPUScene& GPUSce
 
 		// Fill instance upload buffer on separate thread since results are only needed in RHI thread
 		FillInstanceUploadBufferTask = FFunctionGraphTask::CreateAndDispatchWhenReady(
-			[InstanceUploadData = MakeArrayView(InstanceUploadData, SceneInitializer.NumNativeInstances),
+			[InstanceUploadData = MakeArrayView(InstanceUploadData, SceneInitializer.NumNativeInstancesPerLayer[0]),
 			TransformUploadData = MakeArrayView(TransformUploadData, SceneWithGeometryInstances.NumNativeCPUInstances * 3),
 			NumNativeGPUSceneInstances = SceneWithGeometryInstances.NumNativeGPUSceneInstances,
 			NumNativeCPUInstances = SceneWithGeometryInstances.NumNativeCPUInstances,
diff --git a/Engine/Source/Runtime/Renderer/Private/Tests/RayTracingTestbed.cpp b/Engine/Source/Runtime/Renderer/Private/Tests/RayTracingTestbed.cpp
--- a/Engine/Source/Runtime/Renderer/Private/Tests/RayTracingTestbed.cpp
+++ b/Engine/Source/Runtime/Renderer/Private/Tests/RayTracingTestbed.cpp
@@ -138,12 +138,12 @@ bool RunRayTracingTestbed_RenderThread(const FString& Parameters)
 		ScratchBufferCreateInfo);
 
 	FRWBufferStructured InstanceBuffer;
-	InstanceBuffer.Initialize(TEXT("RayTracingTestBedInstanceBuffer"), GRHIRayTracingInstanceDescriptorSize, SceneInitializer.NumNativeInstances);
+	InstanceBuffer.Initialize(TEXT("RayTracingTestBedInstanceBuffer"), GRHIRayTracingInstanceDescriptorSize, SceneInitializer.NumNativeInstancesPerLayer[0]);
 
 	FByteAddressBuffer AccelerationStructureAddressesBuffer;
 	AccelerationStructureAddressesBuffer.Initialize(TEXT("RayTracingTestBedAccelerationStructureAddressesBuffer"), sizeof(FRayTracingAccelerationStructureAddress), BUF_Volatile);
 
-	const uint32 InstanceUploadBufferSize = SceneInitializer.NumNativeInstances * sizeof(FRayTracingInstanceDescriptorInput);
+	const uint32 InstanceUploadBufferSize = SceneInitializer.NumNativeInstancesPerLayer[0] * sizeof(FRayTracingInstanceDescriptorInput);
 	FBufferRHIRef InstanceUploadBuffer;
 	FShaderResourceViewRHIRef InstanceUploadSRV;
 	{
@@ -174,7 +174,7 @@ bool RunRayTracingTestbed_RenderThread(const FString& Parameters)
 			RayTracingScene.BaseUploadBufferOffsets,
 			RayTracingScene.NumNativeGPUSceneInstances,
 			RayTracingScene.NumNativeCPUInstances,
-			MakeArrayView(InstanceUploadData, SceneInitializer.NumNativeInstances),
+			MakeArrayView(InstanceUploadData, SceneInitializer.NumNativeInstancesPerLayer[0]),
 			MakeArrayView(TransformUploadData, RayTracingScene.NumNativeCPUInstances * 3));
 		RHICmdList.UnlockBuffer(TransformUploadBuffer);
 		RHICmdList.UnlockBuffer(InstanceUploadBuffer);
diff --git a/Engine/Source/Runtime/VulkanRHI/Private/VulkanRayTracing.cpp b/Engine/Source/Runtime/VulkanRHI/Private/VulkanRayTracing.cpp
--- a/Engine/Source/Runtime/VulkanRHI/Private/VulkanRayTracing.cpp
+++ b/Engine/Source/Runtime/VulkanRHI/Private/VulkanRayTracing.cpp
@@ -397,7 +397,26 @@ FVulkanRayTracingScene::FVulkanRayTracingScene(FRayTracingSceneInitializer2 InIn
 	: Device(InDevice), Initializer(MoveTemp(InInitializer))
 {
 	const ERayTracingAccelerationStructureFlags BuildFlags = ERayTracingAccelerationStructureFlags::FastTrace; // #yuriy_todo: pass this in
-	SizeInfo = RHICalcRayTracingSceneSize(Initializer.NumNativeInstances, BuildFlags);
+
+	SizeInfo = {};
+
+	const uint32 NumLayers = Initializer.NumNativeInstancesPerLayer.Num();
+	check(NumLayers > 0);
+
+	Layers.SetNum(NumLayers);
+
+	for (uint32 LayerIndex = 0; LayerIndex < NumLayers; ++LayerIndex)
+	{
+		FLayerData& Layer = Layers[LayerIndex];
+
+		FRayTracingAccelerationStructureSize LayerSizeInfo = RHICalcRayTracingSceneSize(Initializer.NumNativeInstancesPerLayer[LayerIndex], BuildFlags);
+
+		Layer.BufferOffset = Align(SizeInfo.ResultSize, GRHIRayTracingAccelerationStructureAlignment);
+		Layer.ScratchBufferOffset = Align(SizeInfo.BuildScratchSize, GRHIRayTracingScratchBufferAlignment);
+
+		SizeInfo.ResultSize = Layer.BufferOffset + LayerSizeInfo.ResultSize;
+		SizeInfo.BuildScratchSize = Layer.BufferOffset + LayerSizeInfo.BuildScratchSize;
+	}
 
 	const uint32 ParameterBufferSize = FMath::Max<uint32>(1, Initializer.NumTotalSegments) * sizeof(FVulkanRayTracingGeometryParameters);
 	FRHIResourceCreateInfo ParameterBufferCreateInfo(TEXT("RayTracingSceneMetadata"));
@@ -414,16 +433,21 @@ FVulkanRayTracingScene::~FVulkanRayTracingScene()
 
 void FVulkanRayTracingScene::BindBuffer(FRHIBuffer* InBuffer, uint32 InBufferOffset)
 {
-	checkf(AccelerationStructureView == nullptr, TEXT("Binding multiple buffers is not currently supported."));
-
 	check(IsInRHIThread() || !IsRunningRHIInSeparateThread());
 
 	check(SizeInfo.ResultSize + InBufferOffset <= InBuffer->GetSize());
-	check(InBufferOffset % 256 == 0); // Spec requires offset to be a multiple of 256
+	
 	AccelerationStructureBuffer = ResourceCast(InBuffer);
 
-	FShaderResourceViewInitializer ViewInitializer(InBuffer, InBufferOffset, 0);
-	AccelerationStructureView = new FVulkanShaderResourceView(Device, AccelerationStructureBuffer, InBufferOffset);
+	for (auto& Layer : Layers)
+	{
+		checkf(Layer.ShaderResourceView == nullptr, TEXT("Binding multiple buffers is not currently supported."));
+
+		const uint32 LayerOffset = InBufferOffset + Layer.BufferOffset;
+		check(LayerOffset % GRHIRayTracingAccelerationStructureAlignment == 0);
+
+		Layer.ShaderResourceView = new FVulkanShaderResourceView(Device, AccelerationStructureBuffer, LayerOffset);
+	}
 }
 
 void FVulkanRayTracingScene::BuildAccelerationStructure(
@@ -431,48 +455,68 @@ void FVulkanRayTracingScene::BuildAccelerationStructure(
 	FVulkanResourceMultiBuffer* InScratchBuffer, uint32 InScratchOffset,
 	FVulkanResourceMultiBuffer* InInstanceBuffer, uint32 InInstanceOffset)
 {
-	check(AccelerationStructureBuffer.IsValid());
-	check(InInstanceBuffer != nullptr);
-	const bool bExternalScratchBuffer = InScratchBuffer != nullptr;
-
-	VkDeviceAddress InstanceBufferAddress = InInstanceBuffer->GetDeviceAddress() + InInstanceOffset;
-
 	// Build a metadata buffer	that contains VulkanRHI-specific per-geometry parameters that allow us to access
 	// vertex and index buffers from shaders that use inline ray tracing.
 	BuildPerInstanceGeometryParameterBuffer();
 
-	FVkRtTLASBuildData BuildData;
-	GetTLASBuildData(Device->GetInstanceHandle(), Initializer.NumNativeInstances, InstanceBufferAddress, BuildData);
+	check(AccelerationStructureBuffer.IsValid());
+	check(InInstanceBuffer != nullptr);
 
 	TRefCountPtr<FVulkanResourceMultiBuffer> ScratchBuffer;
 
-	if (!bExternalScratchBuffer)
+	if (InScratchBuffer == nullptr)
 	{
 		FRHIResourceCreateInfo ScratchBufferCreateInfo(TEXT("BuildScratchTLAS"));
-		ScratchBuffer = ResourceCast(RHICreateBuffer(BuildData.SizesInfo.buildScratchSize, BUF_UnorderedAccess | BUF_StructuredBuffer, 0, ERHIAccess::UAVCompute, ScratchBufferCreateInfo).GetReference());
+		ScratchBuffer = ResourceCast(RHICreateBuffer(SizeInfo.BuildScratchSize, BUF_UnorderedAccess | BUF_StructuredBuffer, 0, ERHIAccess::UAVCompute, ScratchBufferCreateInfo).GetReference());
 		InScratchBuffer = ScratchBuffer.GetReference();
+		InScratchOffset = 0;
 	}
 
-	checkf(AccelerationStructureView, TEXT("A buffer must be bound to the ray tracing scene before it can be built."));
-	BuildData.GeometryInfo.dstAccelerationStructure = AccelerationStructureView->AccelerationStructureHandle;
+	const uint32 NumLayers = Initializer.NumNativeInstancesPerLayer.Num();
+
+	TArray<FVkRtTLASBuildData> BuildDatas;
+	BuildDatas.SetNum(NumLayers);
+
+	TArray<VkAccelerationStructureBuildGeometryInfoKHR> GeometryInfos;
+	GeometryInfos.SetNum(NumLayers);
+
+	TArray<VkAccelerationStructureBuildRangeInfoKHR> BuildRanges;
+	BuildRanges.SetNum(NumLayers);
+
+	TArray<VkAccelerationStructureBuildRangeInfoKHR*> pBuildRanges;
+	pBuildRanges.SetNum(NumLayers);
 
-	BuildData.GeometryInfo.scratchData.deviceAddress = InScratchBuffer->GetDeviceAddress();
-	if (bExternalScratchBuffer)
+	uint32 InstanceBaseOffset = 0;
+
+	for (uint32 LayerIndex = 0; LayerIndex < NumLayers; ++LayerIndex)
 	{
-		BuildData.GeometryInfo.scratchData.deviceAddress += InScratchOffset;
-	}
+		const FLayerData& Layer = Layers[LayerIndex];
+
+		VkDeviceAddress InstanceBufferAddress = InInstanceBuffer->GetDeviceAddress() + InInstanceOffset;
+
+		FVkRtTLASBuildData& BuildData = BuildDatas[LayerIndex];
+		GetTLASBuildData(Device->GetInstanceHandle(), Initializer.NumNativeInstancesPerLayer[LayerIndex], InstanceBufferAddress, BuildData);
 
-	VkAccelerationStructureBuildRangeInfoKHR TLASBuildRangeInfo;
-	TLASBuildRangeInfo.primitiveCount = Initializer.NumNativeInstances;
-	TLASBuildRangeInfo.primitiveOffset = 0;
-	TLASBuildRangeInfo.transformOffset = 0;
-	TLASBuildRangeInfo.firstVertex = 0;
+		checkf(Layer.ShaderResourceView, TEXT("A buffer must be bound to the ray tracing scene before it can be built."));
+		BuildData.GeometryInfo.dstAccelerationStructure = Layer.ShaderResourceView->AccelerationStructureHandle;
+		BuildData.GeometryInfo.scratchData.deviceAddress = InScratchBuffer->GetDeviceAddress() + InScratchOffset + Layer.ScratchBufferOffset;
 
-	VkAccelerationStructureBuildRangeInfoKHR* const pBuildRanges = &TLASBuildRangeInfo;
+		GeometryInfos[LayerIndex] = BuildData.GeometryInfo;
+
+		VkAccelerationStructureBuildRangeInfoKHR& TLASBuildRangeInfo = BuildRanges[LayerIndex];
+		TLASBuildRangeInfo.primitiveCount = Initializer.NumNativeInstancesPerLayer[LayerIndex];
+		TLASBuildRangeInfo.primitiveOffset = InstanceBaseOffset;
+		TLASBuildRangeInfo.transformOffset = 0;
+		TLASBuildRangeInfo.firstVertex = 0;
+
+		pBuildRanges[LayerIndex] = &TLASBuildRangeInfo;
+
+		InstanceBaseOffset += Initializer.NumNativeInstancesPerLayer[LayerIndex];
+	}
 
 	FVulkanCommandBufferManager& CommandBufferManager = *CommandContext.GetCommandBufferManager();
 	FVulkanCmdBuffer* const CmdBuffer = CommandBufferManager.GetActiveCmdBuffer();
-	VulkanDynamicAPI::vkCmdBuildAccelerationStructuresKHR(CmdBuffer->GetHandle(), 1, &BuildData.GeometryInfo, &pBuildRanges);
+	VulkanDynamicAPI::vkCmdBuildAccelerationStructuresKHR(CmdBuffer->GetHandle(), NumLayers, GeometryInfos.GetData(), pBuildRanges.GetData());
 
 	CommandBufferManager.SubmitActiveCmdBuffer();
 	CommandBufferManager.PrepareForNewActiveCommandBuffer();
diff --git a/Engine/Source/Runtime/VulkanRHI/Private/VulkanRayTracing.h b/Engine/Source/Runtime/VulkanRHI/Private/VulkanRayTracing.h
--- a/Engine/Source/Runtime/VulkanRHI/Private/VulkanRayTracing.h
+++ b/Engine/Source/Runtime/VulkanRHI/Private/VulkanRayTracing.h
@@ -105,6 +105,7 @@ public:
 	~FVulkanRayTracingScene();
 
 	const FRayTracingSceneInitializer2& GetInitializer() const override final { return Initializer; }
+	uint32 GetLayerBufferOffset(uint32 LayerIndex) const override final { return Layers[LayerIndex].BufferOffset; }
 
 	void BindBuffer(FRHIBuffer* InBuffer, uint32 InBufferOffset);
 	void BuildAccelerationStructure(
@@ -112,8 +113,6 @@ public:
 		FVulkanResourceMultiBuffer* ScratchBuffer, uint32 ScratchOffset, 
 		FVulkanResourceMultiBuffer* InstanceBuffer, uint32 InstanceOffset);
 
-	FRayTracingAccelerationStructureSize SizeInfo;
-
 	virtual FRHIShaderResourceView* GetMetadataBufferSRV() const override final
 	{
 		return PerInstanceGeometryParameterSRV.GetReference();
@@ -131,7 +130,15 @@ private:
 	// we allow TLAS memory to be allocated using transient resource allocator and 
 	// the lifetime of the scene object may be different from the lifetime of the buffer.
 	// Many VkAccelerationStructureKHR-s may be created, pointing at the same buffer.
-	TRefCountPtr<FVulkanShaderResourceView> AccelerationStructureView;
+
+	struct FLayerData
+	{
+		TRefCountPtr<FVulkanShaderResourceView> ShaderResourceView;
+		uint32 BufferOffset;
+		uint32 ScratchBufferOffset;
+	};
+
+	TArray<FLayerData> Layers;
 	
 	TRefCountPtr<FVulkanResourceMultiBuffer> AccelerationStructureBuffer;
 
