Free dxgi.dll after using it, not before.

diff --git a/Engine/Source/Runtime/D3D12RHI/Private/D3D12Adapter.cpp b/Engine/Source/Runtime/D3D12RHI/Private/D3D12Adapter.cpp
--- a/Engine/Source/Runtime/D3D12RHI/Private/D3D12Adapter.cpp
+++ b/Engine/Source/Runtime/D3D12RHI/Private/D3D12Adapter.cpp
@@ -301,6 +301,17 @@ FD3D12Adapter::FD3D12Adapter(FD3D12AdapterDesc& DescIn)
 	}
 }
 
+FD3D12Adapter::~FD3D12Adapter()
+{
+#if D3D12_SUPPORTS_DXGI_DEBUG
+	FPlatformProcess::FreeDllHandle(DxgiDebugDllHandle);
+#endif
+
+#if PLATFORM_WINDOWS
+	FPlatformProcess::FreeDllHandle(DxgiDllHandle);
+#endif
+}
+
 void FD3D12Adapter::Initialize(FD3D12DynamicRHI* RHI)
 {
 	OwningRHI = RHI;
@@ -617,20 +628,15 @@ void FD3D12Adapter::CreateRootDevice(bool bWithDebug)
 	if (bWithDebug)
 	{
 		// Manually load dxgi debug if available
-		HMODULE DxgiDebugDLL = (HMODULE)FPlatformProcess::GetDllHandle(TEXT("dxgidebug.dll"));
-		if (DxgiDebugDLL)
+		DxgiDebugDllHandle = (HMODULE)FPlatformProcess::GetDllHandle(TEXT("dxgidebug.dll"));
+		if (DxgiDebugDllHandle)
 		{
 			typedef HRESULT(WINAPI* FDXGIGetDebugInterface)(REFIID, void**);
-#pragma warning(push)
-#pragma warning(disable: 4191) // disable the "unsafe conversion from 'FARPROC' to 'blah'" warning
-			FDXGIGetDebugInterface DXGIGetDebugInterfaceFnPtr = (FDXGIGetDebugInterface)(GetProcAddress(DxgiDebugDLL, "DXGIGetDebugInterface"));
-#pragma warning(pop)
+			FDXGIGetDebugInterface DXGIGetDebugInterfaceFnPtr = (FDXGIGetDebugInterface)(void*)(GetProcAddress(DxgiDebugDllHandle, "DXGIGetDebugInterface"));
 			if (DXGIGetDebugInterfaceFnPtr != nullptr)
 			{
-				DXGIGetDebugInterfaceFnPtr(__uuidof(IDXGIDebug), (void**)DXGIDebug.GetInitReference());
+				DXGIGetDebugInterfaceFnPtr(IID_PPV_ARGS(DXGIDebug.GetInitReference()));
 			}
-
-			FPlatformProcess::FreeDllHandle(DxgiDebugDLL);
 		}
 	}
 #endif //  (PLATFORM_WINDOWS || PLATFORM_HOLOLENS)
@@ -1394,12 +1400,10 @@ void FD3D12Adapter::CreateDXGIFactory(bool bWithDebug)
 
 #if PLATFORM_WINDOWS
 	// Dynamically load this otherwise Win7 fails to boot as it's missing on that DLL
-	HMODULE DxgiDLL = (HMODULE)FPlatformProcess::GetDllHandle(TEXT("dxgi.dll"));
-	check(DxgiDLL);
-
-	FCreateDXGIFactory2* CreateDXGIFactory2FnPtr = (FCreateDXGIFactory2*)(void*)::GetProcAddress(DxgiDLL, "CreateDXGIFactory2");
+	DxgiDllHandle = (HMODULE)FPlatformProcess::GetDllHandle(TEXT("dxgi.dll"));
+	check(DxgiDllHandle);
 
-	FPlatformProcess::FreeDllHandle(DxgiDLL);
+	FCreateDXGIFactory2* CreateDXGIFactory2FnPtr = (FCreateDXGIFactory2*)(void*)::GetProcAddress(DxgiDllHandle, "CreateDXGIFactory2");
 #else
 	FCreateDXGIFactory2* CreateDXGIFactory2FnPtr = &CreateDXGIFactory2;
 #endif
diff --git a/Engine/Source/Runtime/D3D12RHI/Private/D3D12Adapter.h b/Engine/Source/Runtime/D3D12RHI/Private/D3D12Adapter.h
--- a/Engine/Source/Runtime/D3D12RHI/Private/D3D12Adapter.h
+++ b/Engine/Source/Runtime/D3D12RHI/Private/D3D12Adapter.h
@@ -134,7 +134,7 @@ class FD3D12Adapter : public FNoncopyable
 public:
 
 	FD3D12Adapter(FD3D12AdapterDesc& DescIn);
-	virtual ~FD3D12Adapter() { }
+	virtual ~FD3D12Adapter();
 
 	void Initialize(FD3D12DynamicRHI* RHI);
 	void InitializeDevices();
@@ -510,10 +510,15 @@ protected:
 #endif
 
 #if D3D12_SUPPORTS_DXGI_DEBUG
+	HMODULE DxgiDebugDllHandle{};
 	TRefCountPtr<IDXGIDebug> DXGIDebug;
 	HANDLE ExceptionHandlerHandle = INVALID_HANDLE_VALUE;
 #endif
 
+#if PLATFORM_WINDOWS
+	HMODULE DxgiDllHandle{};
+#endif
+
 	D3D_ROOT_SIGNATURE_VERSION RootSignatureVersion;
 	bool bDepthBoundsTestSupported;
 	bool bHeapNotZeroedSupported;
