Addressing crashes when rendering offscreen in D3D12.

diff --git a/Engine/Source/Runtime/D3D12RHI/Private/Windows/WindowsD3D12Viewport.cpp b/Engine/Source/Runtime/D3D12RHI/Private/Windows/WindowsD3D12Viewport.cpp
--- a/Engine/Source/Runtime/D3D12RHI/Private/Windows/WindowsD3D12Viewport.cpp
+++ b/Engine/Source/Runtime/D3D12RHI/Private/Windows/WindowsD3D12Viewport.cpp
@@ -200,13 +200,16 @@ void FD3D12Viewport::Init()
 	}
 #endif
 
-	SwapChain1->QueryInterface(IID_PPV_ARGS(SwapChain2.GetInitReference()));
+	if (SwapChain1)
+	{
+		SwapChain1->QueryInterface(IID_PPV_ARGS(SwapChain2.GetInitReference()));
 #if DXGI_MAX_SWAPCHAIN_INTERFACE >= 3
-	SwapChain1->QueryInterface(IID_PPV_ARGS(SwapChain3.GetInitReference()));
+		SwapChain1->QueryInterface(IID_PPV_ARGS(SwapChain3.GetInitReference()));
 #endif
 #if DXGI_MAX_SWAPCHAIN_INTERFACE >= 4
-	SwapChain1->QueryInterface(IID_PPV_ARGS(SwapChain4.GetInitReference()));
+		SwapChain1->QueryInterface(IID_PPV_ARGS(SwapChain4.GetInitReference()));
 #endif
+	}
 
 	{
 		// Don't make the windows association call and release back buffer at the same time (see notes on critical section)
