Moving the WaveOps feature check so it can't be missed on dx12.
Done by christopher.waters and tested locally.

diff --git a/Engine/Source/Runtime/D3D12RHI/Private/D3D12Adapter.cpp b/Engine/Source/Runtime/D3D12RHI/Private/D3D12Adapter.cpp
--- a/Engine/Source/Runtime/D3D12RHI/Private/D3D12Adapter.cpp
+++ b/Engine/Source/Runtime/D3D12RHI/Private/D3D12Adapter.cpp
@@ -526,15 +526,6 @@ void FD3D12Adapter::CreateRootDevice(bool bWithDebug)
 		));
 	}
 
-	// Detect availability of shader model 6.0 wave operations
-	{
-		D3D12_FEATURE_DATA_D3D12_OPTIONS1 Features = {};
-		RootDevice->CheckFeatureSupport(D3D12_FEATURE_D3D12_OPTIONS1, &Features, sizeof(Features));
-		GRHISupportsWaveOperations = Features.WaveOps;
-		GRHIMinimumWaveSize = Features.WaveLaneCountMin;
-		GRHIMaximumWaveSize = Features.WaveLaneCountMax;
-	}
-
 #if ENABLE_RESIDENCY_MANAGEMENT
 	if (!CVarResidencyManagement.GetValueOnAnyThread())
 	{
@@ -1002,6 +993,15 @@ void FD3D12Adapter::InitializeDevices()
 				UE_LOG(LogD3D12RHI, Log, TEXT("Bindless resources are supported"));
 			}
 
+			// Detect availability of shader model 6.0 wave operations
+			{
+				D3D12_FEATURE_DATA_D3D12_OPTIONS1 Features{};
+				RootDevice->CheckFeatureSupport(D3D12_FEATURE_D3D12_OPTIONS1, &Features, sizeof(Features));
+				GRHISupportsWaveOperations = Features.WaveOps;
+				GRHIMinimumWaveSize = Features.WaveLaneCountMin;
+				GRHIMaximumWaveSize = Features.WaveLaneCountMax;
+			}
+
 #if D3D12_RHI_RAYTRACING
 			D3D12_FEATURE_DATA_D3D12_OPTIONS5 D3D12Caps5 = {};
 			if (SUCCEEDED(RootDevice->CheckFeatureSupport(D3D12_FEATURE_D3D12_OPTIONS5, &D3D12Caps5, sizeof(D3D12Caps5))))
