Make D3D12 use async compute for nVidia hardware.

diff --git a/Engine/Source/Runtime/D3D12RHI/Private/Windows/WindowsD3D12Device.cpp b/Engine/Source/Runtime/D3D12RHI/Private/Windows/WindowsD3D12Device.cpp
--- a/Engine/Source/Runtime/D3D12RHI/Private/Windows/WindowsD3D12Device.cpp
+++ b/Engine/Source/Runtime/D3D12RHI/Private/Windows/WindowsD3D12Device.cpp
@@ -1396,7 +1396,7 @@ void FD3D12DynamicRHI::Init()
 		UE_LOG(LogD3D12RHI, Log, TEXT("RHI does not have support for 64 bit atomics"));
 	}
 
-	GSupportsEfficientAsyncCompute = GAllowAsyncCompute && (FParse::Param(FCommandLine::Get(), TEXT("ForceAsyncCompute")) || (GRHISupportsParallelRHIExecute && IsRHIDeviceAMD()));
+	GSupportsEfficientAsyncCompute = GAllowAsyncCompute && (FParse::Param(FCommandLine::Get(), TEXT("ForceAsyncCompute")) || (GRHISupportsParallelRHIExecute && (IsRHIDeviceAMD() || IsRHIDeviceNVIDIA())));
 
 	GSupportsDepthBoundsTest = SupportsDepthBoundsTest(this);
 
