Batched TLAS builds.

diff --git a/Engine/Source/Runtime/D3D12RHI/Private/D3D12RayTracing.cpp b/Engine/Source/Runtime/D3D12RHI/Private/D3D12RayTracing.cpp
--- a/Engine/Source/Runtime/D3D12RHI/Private/D3D12RayTracing.cpp
+++ b/Engine/Source/Runtime/D3D12RHI/Private/D3D12RayTracing.cpp
@@ -4215,6 +4215,9 @@ void FD3D12RayTracingScene::BuildAccelerationStructure(FD3D12CommandContext& Com
 	CommandContext.CommandListHandle.AddUAVBarrier();
 	CommandContext.CommandListHandle.FlushResourceBarriers();
 
+	TArray<D3D12_BUILD_RAYTRACING_ACCELERATION_STRUCTURE_DESC, TInlineAllocator<32>> BuildDescs;
+	BuildDescs.Reserve(NumLayers);
+
 	uint32 InstanceBaseOffset = 0;
 
 	for(uint32 LayerIndex = 0; LayerIndex < NumLayers; ++LayerIndex)
@@ -4232,21 +4235,20 @@ void FD3D12RayTracingScene::BuildAccelerationStructure(FD3D12CommandContext& Com
 			TEXT("TLAS scratch buffer (plus offset) must be aligned to %lld bytes."),
 			GRHIRayTracingScratchBufferAlignment);
 
-		D3D12_BUILD_RAYTRACING_ACCELERATION_STRUCTURE_DESC BuildDesc = {};
+		D3D12_BUILD_RAYTRACING_ACCELERATION_STRUCTURE_DESC& BuildDesc = BuildDescs.AddDefaulted_GetRef();
 		BuildDesc.Inputs = Layer.BuildInputs;
 		BuildDesc.Inputs.InstanceDescs = InstanceBuffer->ResourceLocation.GetGPUVirtualAddress() + InstanceBufferOffset + InstanceBaseOffset * GRHIRayTracingInstanceDescriptorSize;
 		BuildDesc.DestAccelerationStructureData = BufferAddress;
 		BuildDesc.ScratchAccelerationStructureData = ScratchAddress;
 		BuildDesc.SourceAccelerationStructureData = D3D12_GPU_VIRTUAL_ADDRESS(0); // Null source TLAS as this is a build command
 
-		ID3D12GraphicsCommandList4* RayTracingCommandList = CommandContext.CommandListHandle.RayTracingCommandList();
-		RayTracingCommandList->BuildRaytracingAccelerationStructure(&BuildDesc, 0, nullptr);
-
 		INC_DWORD_STAT(STAT_D3D12RayTracingBuiltTLAS);
 
 		InstanceBaseOffset += Initializer.NumNativeInstancesPerLayer[LayerIndex];
 	}
 
+	CommandContext.BuildAccelerationStructuresInternal(BuildDescs);
+
 	// UAV barrier is used here to ensure that the acceleration structure build is complete before any rays are traced
 	// #dxr_todo: these barriers should ideally be inserted by the high level code to allow more overlapped execution
 	CommandContext.CommandListHandle.AddUAVBarrier();
@@ -4488,8 +4490,7 @@ void FD3D12CommandContext::RHIBuildAccelerationStructures(const TArrayView<const
 		}
 	}
 
-	TArrayView<D3D12_BUILD_RAYTRACING_ACCELERATION_STRUCTURE_DESC> LocalBuildDescs = MakeArrayView(BuildDescs.GetData(), BuildDescs.Num());
-	BuildAccelerationStructuresInternal(LocalBuildDescs);
+	BuildAccelerationStructuresInternal(BuildDescs);
 
 	for (const FRayTracingGeometryBuildParams& P : Params)
 	{
