Support batching bindings for callable and miss shaders
- Not supporting r.D3D12.RayTracing.CacheShaderRecords at the moment.

diff --git a/Engine/Source/Runtime/D3D12RHI/Private/D3D12CommandContext.h b/Engine/Source/Runtime/D3D12RHI/Private/D3D12CommandContext.h
--- a/Engine/Source/Runtime/D3D12RHI/Private/D3D12CommandContext.h
+++ b/Engine/Source/Runtime/D3D12RHI/Private/D3D12CommandContext.h
@@ -409,9 +409,9 @@ public:
 		FRHIRayTracingScene* Scene,
 		const FRayTracingShaderBindings& GlobalResourceBindings,
 		FRHIBuffer* ArgumentBuffer, uint32 ArgumentOffset) final override;
-	virtual void RHISetRayTracingHitGroups(
+	virtual void RHISetRayTracingBindings(
 		FRHIRayTracingScene* Scene, FRHIRayTracingPipelineState* Pipeline,
-		uint32 NumBindings, const FRayTracingLocalShaderBindings* Bindings) final override;
+		uint32 NumBindings, const FRayTracingLocalShaderBindings* Bindings, ERayTracingBindingType BindingType) final override;
 	virtual void RHISetRayTracingHitGroup(
 		FRHIRayTracingScene* Scene, uint32 InstanceIndex, uint32 SegmentIndex, uint32 ShaderSlot,
 		FRHIRayTracingPipelineState* Pipeline, uint32 HitGroupIndex,
@@ -802,9 +802,9 @@ public:
 		ContextRedirect(RHISetRayTracingHitGroup(Scene, InstanceIndex, SegmentIndex, ShaderSlot, Pipeline, HitGroupIndex, NumUniformBuffers, UniformBuffers, LooseParameterDataSize, LooseParameterData, UserData));
 	}
 
-	virtual void RHISetRayTracingHitGroups(FRHIRayTracingScene* Scene, FRHIRayTracingPipelineState* Pipeline, uint32 NumBindings, const FRayTracingLocalShaderBindings* Bindings) final override
+	virtual void RHISetRayTracingBindings(FRHIRayTracingScene* Scene, FRHIRayTracingPipelineState* Pipeline, uint32 NumBindings, const FRayTracingLocalShaderBindings* Bindings, ERayTracingBindingType BindingType) final override
 	{
-		ContextRedirect(RHISetRayTracingHitGroups(Scene, Pipeline, NumBindings, Bindings));
+		ContextRedirect(RHISetRayTracingBindings(Scene, Pipeline, NumBindings, Bindings, BindingType));
 	}
 
 	virtual void RHISetRayTracingCallableShader(
diff --git a/Engine/Source/Runtime/D3D12RHI/Private/D3D12RayTracing.cpp b/Engine/Source/Runtime/D3D12RHI/Private/D3D12RayTracing.cpp
--- a/Engine/Source/Runtime/D3D12RHI/Private/D3D12RayTracing.cpp
+++ b/Engine/Source/Runtime/D3D12RHI/Private/D3D12RayTracing.cpp
@@ -129,7 +129,7 @@ DECLARE_MEMORY_STAT(TEXT("Total Used Video Memory"), STAT_D3D12RayTracingUsedVid
 DECLARE_CYCLE_STAT(TEXT("RTPSO Compile Shader"), STAT_RTPSO_CompileShader, STATGROUP_D3D12RayTracing);
 DECLARE_CYCLE_STAT(TEXT("RTPSO Create Pipeline"), STAT_RTPSO_CreatePipeline, STATGROUP_D3D12RayTracing);
 
-DECLARE_CYCLE_STAT(TEXT("SetRayTracingHitGroups"), STAT_D3D12SetRayTracingHitGroups, STATGROUP_D3D12RayTracing);
+DECLARE_CYCLE_STAT(TEXT("SetRayTracingBindings"), STAT_D3D12SetRayTracingBindings, STATGROUP_D3D12RayTracing);
 DECLARE_CYCLE_STAT(TEXT("CreateShaderTable"), STAT_D3D12CreateShaderTable, STATGROUP_D3D12RayTracing);
 DECLARE_CYCLE_STAT(TEXT("BuildTopLevel"), STAT_D3D12BuildTLAS, STATGROUP_D3D12RayTracing);
 DECLARE_CYCLE_STAT(TEXT("BuildBottomLevel"), STAT_D3D12BuildBLAS, STATGROUP_D3D12RayTracing);
@@ -1988,7 +1988,7 @@ public:
 		Desc.RayGenerationShaderRecord.StartAddress = ShaderTableAddress + RayGenShaderTableOffset + RayGenShaderIndex * RayGenRecordStride;
 		Desc.RayGenerationShaderRecord.SizeInBytes = RayGenRecordStride;
 
-		Desc.MissShaderTable.StartAddress = ShaderTableAddress + MissShaderTableOffset;
+		Desc.MissShaderTable.StartAddress = ShaderTableAddress + MissShaderTableOffset + MissShaderBaseIndex * LocalRecordStride;
 		Desc.MissShaderTable.StrideInBytes = LocalRecordStride;
 		Desc.MissShaderTable.SizeInBytes = LocalRecordStride * NumMissRecords;
 
@@ -5522,12 +5522,85 @@ static void SetRayTracingHitGroup(
 		: FD3D12ShaderIdentifier::Null);
 }
 
-void FD3D12CommandContext::RHISetRayTracingHitGroups(
+static void SetRayTracingCallableShader(
+	FD3D12Device* Device,
+	FD3D12RayTracingShaderTable* ShaderTable,
+	FD3D12RayTracingScene* Scene,
+	FD3D12RayTracingPipelineState* Pipeline,
+	uint32 ShaderSlotInScene, uint32 ShaderIndexInPipeline,
+	uint32 NumUniformBuffers, FRHIUniformBuffer* const* UniformBuffers,
+	uint32 LooseParameterDataSize, const void* LooseParameterData,
+	uint32 UserData,
+	uint32 WorkerIndex)
+{
+	checkf(ShaderSlotInScene < Scene->Initializer.NumCallableShaderSlots, TEXT("Shader slot is invalid. Make sure that NumCallableShaderSlots is correct on FRayTracingSceneInitializer."));
+
+	const uint32 RecordIndex = ShaderTable->CallableShaderRecordIndexOffset + ShaderSlotInScene;
+
+	const uint32 UserDataOffset = offsetof(FHitGroupSystemParameters, RootConstants) + offsetof(FHitGroupSystemRootConstants, UserData);
+	ShaderTable->SetLocalShaderParameters(RecordIndex, UserDataOffset, UserData);
+
+	const FD3D12RayTracingShader* Shader = Pipeline->CallableShaders.Shaders[ShaderIndexInPipeline];
+
+	FD3D12RayTracingLocalResourceBinder ResourceBinder(*Device, *ShaderTable, *(Shader->pRootSignature), RecordIndex, WorkerIndex);
+	const bool bResourcesBound = SetRayTracingShaderResources(Shader,
+		0, nullptr, // Textures
+		0, nullptr, // SRVs
+		NumUniformBuffers, UniformBuffers,
+		0, nullptr, // Samplers
+		0, nullptr, // UAVs
+		LooseParameterDataSize, LooseParameterData, // Loose parameters
+		ResourceBinder);
+
+	ShaderTable->SetLocalShaderIdentifier(RecordIndex,
+		bResourcesBound
+		? Pipeline->CallableShaders.Identifiers[ShaderIndexInPipeline]
+		: FD3D12ShaderIdentifier::Null);
+}
+
+static void SetRayTracingMissShader(
+	FD3D12Device* Device,
+	FD3D12RayTracingShaderTable* ShaderTable,
+	FD3D12RayTracingScene* Scene,
+	FD3D12RayTracingPipelineState* Pipeline,
+	uint32 ShaderSlotInScene, uint32 ShaderIndexInPipeline,
+	uint32 NumUniformBuffers, FRHIUniformBuffer* const* UniformBuffers,
+	uint32 LooseParameterDataSize, const void* LooseParameterData,
+	uint32 UserData,
+	uint32 WorkerIndex)
+{
+	checkf(ShaderSlotInScene < Scene->Initializer.NumMissShaderSlots, TEXT("Shader slot is invalid. Make sure that NumMissShaderSlots is correct on FRayTracingSceneInitializer."));
+
+	const uint32 RecordIndex = ShaderTable->MissShaderRecordIndexOffset + ShaderSlotInScene;
+
+	const uint32 UserDataOffset = offsetof(FHitGroupSystemParameters, RootConstants) + offsetof(FHitGroupSystemRootConstants, UserData);
+	ShaderTable->SetLocalShaderParameters(RecordIndex, UserDataOffset, UserData);
+
+	const FD3D12RayTracingShader* Shader = Pipeline->CallableShaders.Shaders[ShaderIndexInPipeline];
+
+	FD3D12RayTracingLocalResourceBinder ResourceBinder(*Device, *ShaderTable, *(Shader->pRootSignature), RecordIndex, WorkerIndex);
+	const bool bResourcesBound = SetRayTracingShaderResources(Shader,
+		0, nullptr, // Textures
+		0, nullptr, // SRVs
+		NumUniformBuffers, UniformBuffers,
+		0, nullptr, // Samplers
+		0, nullptr, // UAVs
+		LooseParameterDataSize, LooseParameterData, // Loose parameters
+		ResourceBinder);
+
+	ShaderTable->SetLocalShaderIdentifier(RecordIndex,
+		bResourcesBound
+		? Pipeline->MissShaders.Identifiers[ShaderIndexInPipeline]
+		: FD3D12ShaderIdentifier::Null);
+}
+
+void FD3D12CommandContext::RHISetRayTracingBindings(
 	FRHIRayTracingScene* InScene, FRHIRayTracingPipelineState* InPipeline,
-	uint32 NumBindings, const FRayTracingLocalShaderBindings* Bindings)
+	uint32 NumBindings, const FRayTracingLocalShaderBindings* Bindings,
+	ERayTracingBindingType BindingType)
 {
-	TRACE_CPUPROFILER_EVENT_SCOPE(SetRayTracingHitGroups);
-	SCOPE_CYCLE_COUNTER(STAT_D3D12SetRayTracingHitGroups);
+	TRACE_CPUPROFILER_EVENT_SCOPE(RHISetRayTracingBindings);
+	SCOPE_CYCLE_COUNTER(STAT_D3D12SetRayTracingBindings);
 
 	FD3D12RayTracingScene* Scene = FD3D12DynamicRHI::ResourceCast(InScene);
 	FD3D12RayTracingPipelineState* Pipeline = FD3D12DynamicRHI::ResourceCast(InPipeline);
@@ -5550,7 +5623,7 @@ void FD3D12CommandContext::RHISetRayTracingHitGroups(
 	volatile int32 SharedCounter = 0;
 	const int32 IncrementCount = 128;
 
-	auto BindingTask = [Bindings, Device, ShaderTable, Scene, Pipeline, &SharedCounter, IncrementCount, NumBindings](uint32 WorkerIndex)
+	auto BindingTask = [Bindings, Device, ShaderTable, Scene, Pipeline, &SharedCounter, IncrementCount, NumBindings, BindingType](uint32 WorkerIndex)
 	{
 		TRACE_CPUPROFILER_EVENT_SCOPE(SetRayTracingHitGroups_Task);
 	 
@@ -5563,17 +5636,49 @@ void FD3D12CommandContext::RHISetRayTracingHitGroups(
 			for (int32 CurrentIndex = BaseIndex; CurrentIndex < EndIndex; ++CurrentIndex)
 			{
 				const FRayTracingLocalShaderBindings& Binding = Bindings[CurrentIndex];
-				SetRayTracingHitGroup(Device, ShaderTable, Scene, Pipeline,
-					Binding.InstanceIndex,
-					Binding.SegmentIndex,
-					Binding.ShaderSlot,
-					Binding.ShaderIndexInPipeline,
-					Binding.NumUniformBuffers,
-					Binding.UniformBuffers,
-					Binding.LooseParameterDataSize,
-					Binding.LooseParameterData,
-					Binding.UserData,
-					WorkerIndex);
+
+				if (BindingType == ERayTracingBindingType::HitGroup)
+				{
+					SetRayTracingHitGroup(Device, ShaderTable, Scene, Pipeline,
+						Binding.InstanceIndex,
+						Binding.SegmentIndex,
+						Binding.ShaderSlot,
+						Binding.ShaderIndexInPipeline,
+						Binding.NumUniformBuffers,
+						Binding.UniformBuffers,
+						Binding.LooseParameterDataSize,
+						Binding.LooseParameterData,
+						Binding.UserData,
+						WorkerIndex);
+				}
+				else if (BindingType == ERayTracingBindingType::CallableShader)
+				{
+					SetRayTracingCallableShader(Device, ShaderTable, Scene, Pipeline,
+						Binding.ShaderSlot,
+						Binding.ShaderIndexInPipeline,
+						Binding.NumUniformBuffers,
+						Binding.UniformBuffers,
+						Binding.LooseParameterDataSize,
+						Binding.LooseParameterData,
+						Binding.UserData,
+						WorkerIndex);
+				}
+				else if (BindingType == ERayTracingBindingType::MissShader)
+				{
+					SetRayTracingMissShader(Device, ShaderTable, Scene, Pipeline,
+						Binding.ShaderSlot,
+						Binding.ShaderIndexInPipeline,
+						Binding.NumUniformBuffers,
+						Binding.UniformBuffers,
+						Binding.LooseParameterDataSize,
+						Binding.LooseParameterData,
+						Binding.UserData,
+						WorkerIndex);
+				}
+				else
+				{
+					checkNoEntry();
+				}
 			}
 		}
 	};
@@ -5638,29 +5743,9 @@ void FD3D12CommandContext::RHISetRayTracingCallableShader(
 	checkf(ShaderSlotInScene < Scene->Initializer.NumCallableShaderSlots, TEXT("Shader slot is invalid. Make sure that NumCallableShaderSlots is correct on FRayTracingSceneInitializer."));
 
 	FD3D12RayTracingShaderTable* ShaderTable = Scene->FindOrCreateShaderTable(Pipeline, GetParentDevice());
-
-	const uint32 RecordIndex = ShaderTable->CallableShaderRecordIndexOffset + ShaderSlotInScene;
-
-	const uint32 UserDataOffset = offsetof(FHitGroupSystemParameters, RootConstants) + offsetof(FHitGroupSystemRootConstants, UserData);
-	ShaderTable->SetLocalShaderParameters(RecordIndex, UserDataOffset, UserData);
-
-	const FD3D12RayTracingShader* Shader = Pipeline->CallableShaders.Shaders[ShaderIndexInPipeline];
-
 	const uint32 WorkerIndex = 0;
-	FD3D12RayTracingLocalResourceBinder ResourceBinder(*GetParentDevice(), *ShaderTable, *(Shader->pRootSignature), RecordIndex, WorkerIndex);
-	const bool bResourcesBound = SetRayTracingShaderResources(Shader,
-		0, nullptr, // Textures
-		0, nullptr, // SRVs
-		NumUniformBuffers, UniformBuffers,
-		0, nullptr, // Samplers
-		0, nullptr, // UAVs
-		0, nullptr, // Loose parameters
-		ResourceBinder);
 
-	ShaderTable->SetLocalShaderIdentifier(RecordIndex, 
-		bResourcesBound 
-		? Pipeline->CallableShaders.Identifiers[ShaderIndexInPipeline]
-		: FD3D12ShaderIdentifier::Null);
+	SetRayTracingCallableShader(GetParentDevice(), ShaderTable, Scene, Pipeline, ShaderSlotInScene, ShaderIndexInPipeline, NumUniformBuffers, UniformBuffers, 0, nullptr, UserData, WorkerIndex);
 }
 
 void FD3D12CommandContext::RHISetRayTracingMissShader(
diff --git a/Engine/Source/Runtime/RHI/Public/RHICommandList.h b/Engine/Source/Runtime/RHI/Public/RHICommandList.h
--- a/Engine/Source/Runtime/RHI/Public/RHICommandList.h
+++ b/Engine/Source/Runtime/RHI/Public/RHICommandList.h
@@ -143,6 +143,13 @@ struct FRayTracingLocalShaderBindings
 	uint8* LooseParameterData = nullptr;
 };
 
+enum class ERayTracingBindingType : uint8
+{
+	HitGroup,
+	CallableShader,
+	MissShader,
+};
+
 // C++ counter-part of FBasicRayData declared in RayTracingCommon.ush
 struct FBasicRayData
 {
@@ -2281,16 +2288,8 @@ FRHICOMMAND_MACRO(FRHICommandRayTraceDispatch)
 
 FRHICOMMAND_MACRO(FRHICommandSetRayTracingBindings)
 {
-	enum EBindingType
-	{
-		EBindingType_HitGroup,
-		EBindingType_CallableShader,
-		EBindingType_MissShader,
-		EBindingType_HitGroupBatch,
-	};
-
 	FRHIRayTracingScene* Scene = nullptr;
-	EBindingType BindingType = EBindingType_HitGroup;
+	ERayTracingBindingType BindingType = ERayTracingBindingType::HitGroup;
 	uint32 InstanceIndex = 0;
 	uint32 SegmentIndex = 0;
 	uint32 ShaderSlot = 0;
@@ -2303,7 +2302,7 @@ FRHICOMMAND_MACRO(FRHICommandSetRayTracingBindings)
 	uint32 UserData = 0;
 
 	// Batched bindings
-	uint32 NumBindings = 0;
+	int32 NumBindings = -1;
 	const FRayTracingLocalShaderBindings* Bindings = nullptr;
 
 	// Hit group bindings
@@ -2312,7 +2311,7 @@ FRHICOMMAND_MACRO(FRHICommandSetRayTracingBindings)
 		uint32 InLooseParameterDataSize, const void* InLooseParameterData,
 		uint32 InUserData)
 		: Scene(InScene)
-		, BindingType(EBindingType_HitGroup)
+		, BindingType(ERayTracingBindingType::HitGroup)
 		, InstanceIndex(InInstanceIndex)
 		, SegmentIndex(InSegmentIndex)
 		, ShaderSlot(InShaderSlot)
@@ -2326,21 +2325,11 @@ FRHICOMMAND_MACRO(FRHICommandSetRayTracingBindings)
 	{
 	}
 
-	// Batched hit group bindings
-	FRHICommandSetRayTracingBindings(FRHIRayTracingScene* InScene, FRayTracingPipelineState* InPipeline, uint32 InNumBindings, const FRayTracingLocalShaderBindings* InBindings)
-		: Scene(InScene)
-		, BindingType(EBindingType_HitGroupBatch)
-		, Pipeline(InPipeline)
-		, NumBindings(InNumBindings)
-		, Bindings(InBindings)
-	{
-	}
-
 	// Callable and Miss shader bindings
 	FRHICommandSetRayTracingBindings(FRHIRayTracingScene* InScene, uint32 InShaderSlot,
 		FRayTracingPipelineState* InPipeline, uint32 InShaderIndex,
 		uint32 InNumUniformBuffers, FRHIUniformBuffer* const* InUniformBuffers,
-		uint32 InUserData, EBindingType InBindingType)
+		uint32 InUserData, ERayTracingBindingType InBindingType)
 		: Scene(InScene)
 		, BindingType(InBindingType)
 		, InstanceIndex(0)
@@ -2352,6 +2341,18 @@ FRHICOMMAND_MACRO(FRHICommandSetRayTracingBindings)
 		, UniformBuffers(InUniformBuffers)
 		, UserData(InUserData)
 	{
+		checkf(InBindingType != ERayTracingBindingType::HitGroup, TEXT("Hit group bindings must specify Instance and Segment Index."));
+	}
+
+	// Bindings Batch
+	FRHICommandSetRayTracingBindings(FRHIRayTracingScene* InScene, FRayTracingPipelineState* InPipeline, uint32 InNumBindings, const FRayTracingLocalShaderBindings* InBindings, ERayTracingBindingType InBindingType)
+		: Scene(InScene)
+		, BindingType(InBindingType)
+		, Pipeline(InPipeline)
+		, NumBindings(InNumBindings)
+		, Bindings(InBindings)
+	{
+
 	}
 
 	RHI_API void Execute(FRHICommandListBase& CmdList);
@@ -3836,15 +3837,16 @@ public:
 		}
 	}
 
-	FORCEINLINE_DEBUGGABLE void SetRayTracingHitGroups(
+	FORCEINLINE_DEBUGGABLE void SetRayTracingBindings(
 		FRHIRayTracingScene* Scene, FRayTracingPipelineState* Pipeline,
 		uint32 NumBindings, const FRayTracingLocalShaderBindings* Bindings,
+		ERayTracingBindingType BindingType,
 		bool bCopyDataToInlineStorage = true)
 	{
 		if (Bypass())
 		{
 			extern RHI_API FRHIRayTracingPipelineState* GetRHIRayTracingPipelineState(FRayTracingPipelineState*);
-			GetContext().RHISetRayTracingHitGroups(Scene, GetRHIRayTracingPipelineState(Pipeline), NumBindings, Bindings);
+			GetContext().RHISetRayTracingBindings(Scene, GetRHIRayTracingPipelineState(Pipeline), NumBindings, Bindings, BindingType);
 		}
 		else
 		{
@@ -3879,15 +3881,39 @@ public:
 					}
 				}
 
-				ALLOC_COMMAND(FRHICommandSetRayTracingBindings)(Scene, Pipeline, NumBindings, InlineBindings);
+				ALLOC_COMMAND(FRHICommandSetRayTracingBindings)(Scene, Pipeline, NumBindings, InlineBindings, BindingType);
 			}
 			else
 			{
-				ALLOC_COMMAND(FRHICommandSetRayTracingBindings)(Scene, Pipeline, NumBindings, Bindings);
+				ALLOC_COMMAND(FRHICommandSetRayTracingBindings)(Scene, Pipeline, NumBindings, Bindings, BindingType);
 			}
 		}
 	}
 
+	FORCEINLINE_DEBUGGABLE void SetRayTracingHitGroups(
+		FRHIRayTracingScene* Scene, FRayTracingPipelineState* Pipeline,
+		uint32 NumBindings, const FRayTracingLocalShaderBindings* Bindings,
+		bool bCopyDataToInlineStorage = true)
+	{
+		SetRayTracingBindings(Scene, Pipeline, NumBindings, Bindings, ERayTracingBindingType::HitGroup, bCopyDataToInlineStorage);
+	}
+
+	FORCEINLINE_DEBUGGABLE void SetRayTracingCallableShaders(
+		FRHIRayTracingScene* Scene, FRayTracingPipelineState* Pipeline,
+		uint32 NumBindings, const FRayTracingLocalShaderBindings* Bindings,
+		bool bCopyDataToInlineStorage = true)
+	{
+		SetRayTracingBindings(Scene, Pipeline, NumBindings, Bindings, ERayTracingBindingType::CallableShader, bCopyDataToInlineStorage);
+	}
+
+	FORCEINLINE_DEBUGGABLE void SetRayTracingMissShaders(
+		FRHIRayTracingScene* Scene, FRayTracingPipelineState* Pipeline,
+		uint32 NumBindings, const FRayTracingLocalShaderBindings* Bindings,
+		bool bCopyDataToInlineStorage = true)
+	{
+		SetRayTracingBindings(Scene, Pipeline, NumBindings, Bindings, ERayTracingBindingType::MissShader, bCopyDataToInlineStorage);
+	}
+
 	FORCEINLINE_DEBUGGABLE void SetRayTracingHitGroup(
 		FRHIRayTracingScene* Scene, uint32 InstanceIndex, uint32 SegmentIndex, uint32 ShaderSlot,
 		FRayTracingPipelineState* Pipeline, uint32 HitGroupIndex,
@@ -3952,7 +3978,7 @@ public:
 				}
 			}
 
-			ALLOC_COMMAND(FRHICommandSetRayTracingBindings)(Scene, ShaderSlotInScene, Pipeline, ShaderIndexInPipeline, NumUniformBuffers, InlineUniformBuffers, UserData, FRHICommandSetRayTracingBindings::EBindingType::EBindingType_CallableShader);
+			ALLOC_COMMAND(FRHICommandSetRayTracingBindings)(Scene, ShaderSlotInScene, Pipeline, ShaderIndexInPipeline, NumUniformBuffers, InlineUniformBuffers, UserData, ERayTracingBindingType::CallableShader);
 		}
 	}
 
@@ -3979,7 +4005,7 @@ public:
 				}
 			}
 
-			ALLOC_COMMAND(FRHICommandSetRayTracingBindings)(Scene, ShaderSlotInScene, Pipeline, ShaderIndexInPipeline, NumUniformBuffers, InlineUniformBuffers, UserData, FRHICommandSetRayTracingBindings::EBindingType::EBindingType_MissShader);
+			ALLOC_COMMAND(FRHICommandSetRayTracingBindings)(Scene, ShaderSlotInScene, Pipeline, ShaderIndexInPipeline, NumUniformBuffers, InlineUniformBuffers, UserData, ERayTracingBindingType::MissShader);
 		}
 	}
 
diff --git a/Engine/Source/Runtime/RHI/Public/RHICommandListCommandExecutes.inl b/Engine/Source/Runtime/RHI/Public/RHICommandListCommandExecutes.inl
--- a/Engine/Source/Runtime/RHI/Public/RHICommandListCommandExecutes.inl
+++ b/Engine/Source/Runtime/RHI/Public/RHICommandListCommandExecutes.inl
@@ -607,24 +607,29 @@ void FRHICommandSetRayTracingBindings::Execute(FRHICommandListBase& CmdList)
 {
 	RHISTAT(SetRayTracingHitGroup);
 	extern RHI_API FRHIRayTracingPipelineState* GetRHIRayTracingPipelineState(FRayTracingPipelineState*);
-	if (BindingType == EBindingType_HitGroup)
+	if (NumBindings >= 0)
+	{
+		check(Bindings != nullptr);
+		INTERNAL_DECORATOR(RHISetRayTracingBindings)(Scene, GetRHIRayTracingPipelineState(Pipeline), NumBindings, Bindings, BindingType);
+	}
+	else if (BindingType == ERayTracingBindingType::HitGroup)
 	{
 		INTERNAL_DECORATOR(RHISetRayTracingHitGroup)(Scene, InstanceIndex, SegmentIndex, ShaderSlot, GetRHIRayTracingPipelineState(Pipeline), ShaderIndex,
 			NumUniformBuffers, UniformBuffers,
 			LooseParameterDataSize, LooseParameterData,
 			UserData);
 	}
-	else if (BindingType == EBindingType_HitGroupBatch)
+	else if (BindingType == ERayTracingBindingType::CallableShader)
 	{
-		INTERNAL_DECORATOR(RHISetRayTracingHitGroups)(Scene, GetRHIRayTracingPipelineState(Pipeline), NumBindings, Bindings);
+		INTERNAL_DECORATOR(RHISetRayTracingCallableShader)(Scene, ShaderSlot, GetRHIRayTracingPipelineState(Pipeline), ShaderIndex, NumUniformBuffers, UniformBuffers, UserData);
 	}
-	else if (BindingType == EBindingType_CallableShader)
+	else if(BindingType == ERayTracingBindingType::MissShader)
 	{
-		INTERNAL_DECORATOR(RHISetRayTracingCallableShader)(Scene, ShaderSlot, GetRHIRayTracingPipelineState(Pipeline), ShaderIndex, NumUniformBuffers, UniformBuffers, UserData);
+		INTERNAL_DECORATOR(RHISetRayTracingMissShader)(Scene, ShaderSlot, GetRHIRayTracingPipelineState(Pipeline), ShaderIndex, NumUniformBuffers, UniformBuffers, UserData);
 	}
 	else
 	{
-		INTERNAL_DECORATOR(RHISetRayTracingMissShader)(Scene, ShaderSlot, GetRHIRayTracingPipelineState(Pipeline), ShaderIndex, NumUniformBuffers, UniformBuffers, UserData);
+		checkNoEntry();
 	}
 }
 
diff --git a/Engine/Source/Runtime/RHI/Public/RHIContext.h b/Engine/Source/Runtime/RHI/Public/RHIContext.h
--- a/Engine/Source/Runtime/RHI/Public/RHIContext.h
+++ b/Engine/Source/Runtime/RHI/Public/RHIContext.h
@@ -29,6 +29,7 @@ struct FRayTracingGeometrySegment;
 struct FRayTracingGeometryBuildParams;
 struct FRayTracingSceneBuildParams;
 struct FRayTracingLocalShaderBindings;
+enum class ERayTracingBindingType : uint8;
 enum class EAsyncComputeBudget;
 
 #define VALIDATE_UNIFORM_BUFFER_STATIC_BINDINGS (!UE_BUILD_SHIPPING && !UE_BUILD_TEST)
@@ -827,7 +828,7 @@ public:
 		checkNoEntry();
 	}
 
-	virtual void RHISetRayTracingHitGroups(FRHIRayTracingScene* Scene, FRHIRayTracingPipelineState* Pipeline, uint32 NumBindings, const FRayTracingLocalShaderBindings* Bindings)
+	virtual void RHISetRayTracingBindings(FRHIRayTracingScene* Scene, FRHIRayTracingPipelineState* Pipeline, uint32 NumBindings, const FRayTracingLocalShaderBindings* Bindings, ERayTracingBindingType BindingType)
 	{
 		checkNoEntry();
 	}
diff --git a/Engine/Source/Runtime/RHI/Public/RHIValidationContext.h b/Engine/Source/Runtime/RHI/Public/RHIValidationContext.h
--- a/Engine/Source/Runtime/RHI/Public/RHIValidationContext.h
+++ b/Engine/Source/Runtime/RHI/Public/RHIValidationContext.h
@@ -1026,9 +1026,9 @@ public:
 		RHIContext->RHIRayTraceDispatchIndirect(RayTracingPipelineState, RayGenShader, Scene, GlobalResourceBindings, ArgumentBuffer, ArgumentOffset);
 	}
 
-	virtual void RHISetRayTracingHitGroups(FRHIRayTracingScene* Scene, FRHIRayTracingPipelineState* Pipeline, uint32 NumBindings, const FRayTracingLocalShaderBindings* Bindings) override final
+	virtual void RHISetRayTracingBindings(FRHIRayTracingScene* Scene, FRHIRayTracingPipelineState* Pipeline, uint32 NumBindings, const FRayTracingLocalShaderBindings* Bindings, ERayTracingBindingType BindingType) override final
 	{
-		RHIContext->RHISetRayTracingHitGroups(Scene, Pipeline, NumBindings, Bindings);
+		RHIContext->RHISetRayTracingBindings(Scene, Pipeline, NumBindings, Bindings, BindingType);
 	}
 
 	virtual void RHISetRayTracingHitGroup(
