Removing the majority of D3D12RHI_API uses now that D3D12 has a more official public API.

diff --git a/Engine/Source/Runtime/D3D12RHI/Private/D3D12CommandList.h b/Engine/Source/Runtime/D3D12RHI/Private/D3D12CommandList.h
--- a/Engine/Source/Runtime/D3D12RHI/Private/D3D12CommandList.h
+++ b/Engine/Source/Runtime/D3D12RHI/Private/D3D12CommandList.h
@@ -523,7 +523,7 @@ public:
 		return CommandListData->TrackedResourceState.GetResourceState(pResource);
 	}
 
-	D3D12RHI_API void AddPendingResourceBarrier(FD3D12Resource* Resource, D3D12_RESOURCE_STATES State, uint32 SubResource);
+	void AddPendingResourceBarrier(FD3D12Resource* Resource, D3D12_RESOURCE_STATES State, uint32 SubResource);
 
 	TArray<FD3D12PendingResourceBarrier>& PendingResourceBarriers()
 	{
diff --git a/Engine/Source/Runtime/D3D12RHI/Private/D3D12DescriptorCache.cpp b/Engine/Source/Runtime/D3D12RHI/Private/D3D12DescriptorCache.cpp
--- a/Engine/Source/Runtime/D3D12RHI/Private/D3D12DescriptorCache.cpp
+++ b/Engine/Source/Runtime/D3D12RHI/Private/D3D12DescriptorCache.cpp
@@ -258,7 +258,7 @@ bool FD3D12DescriptorCache::SetDescriptorHeaps()
 }
 
 
-D3D12RHI_API void FD3D12DescriptorCache::SetCurrentCommandList(const FD3D12CommandListHandle& CommandListHandle)
+void FD3D12DescriptorCache::SetCurrentCommandList(const FD3D12CommandListHandle& CommandListHandle)
 {
 	// Clear the previous heap pointers (since it's a new command list) and then set the current descriptor heaps.
 	pPreviousViewHeap = nullptr;
diff --git a/Engine/Source/Runtime/D3D12RHI/Private/D3D12DescriptorCache.h b/Engine/Source/Runtime/D3D12RHI/Private/D3D12DescriptorCache.h
--- a/Engine/Source/Runtime/D3D12RHI/Private/D3D12DescriptorCache.h
+++ b/Engine/Source/Runtime/D3D12RHI/Private/D3D12DescriptorCache.h
@@ -343,7 +343,7 @@ public:
 	// Notify the descriptor cache every time you start recording a command list.
 	// This sets descriptor heaps on the command list and indicates the current fence value which allows
 	// us to avoid querying DX12 for that value thousands of times per frame, which can be costly.
-	D3D12RHI_API void SetCurrentCommandList(const FD3D12CommandListHandle& CommandListHandle);
+	void SetCurrentCommandList(const FD3D12CommandListHandle& CommandListHandle);
 
 	// ------------------------------------------------------
 	// end Descriptor Slot Reservation stuff
diff --git a/Engine/Source/Runtime/D3D12RHI/Private/D3D12Device.h b/Engine/Source/Runtime/D3D12RHI/Private/D3D12Device.h
--- a/Engine/Source/Runtime/D3D12RHI/Private/D3D12Device.h
+++ b/Engine/Source/Runtime/D3D12RHI/Private/D3D12Device.h
@@ -123,8 +123,8 @@ public:
 
 	TArray<FD3D12CommandListHandle> PendingCommandLists;
 
-	D3D12RHI_API void RegisterGPUWork(uint32 NumPrimitives = 0, uint32 NumVertices = 0);
-	D3D12RHI_API void RegisterGPUDispatch(FIntVector GroupCount);
+	void RegisterGPUWork(uint32 NumPrimitives = 0, uint32 NumVertices = 0);
+	void RegisterGPUDispatch(FIntVector GroupCount);
 	
 	FD3D12SamplerState* CreateSampler(const FSamplerStateInitializerRHI& Initializer);
 	void CreateSamplerInternal(const D3D12_SAMPLER_DESC& Desc, D3D12_CPU_DESCRIPTOR_HANDLE Descriptor);
diff --git a/Engine/Source/Runtime/D3D12RHI/Private/D3D12Resources.cpp b/Engine/Source/Runtime/D3D12RHI/Private/D3D12Resources.cpp
--- a/Engine/Source/Runtime/D3D12RHI/Private/D3D12Resources.cpp
+++ b/Engine/Source/Runtime/D3D12RHI/Private/D3D12Resources.cpp
@@ -16,7 +16,7 @@ D3D12Resources.cpp: D3D RHI utility implementation.
 	THIRD_PARTY_INCLUDES_END
 #endif
 
-D3D12RHI_API int32 GD3D12AsyncDeferredDeletion = ASYNC_DEFERRED_DELETION;
+int32 GD3D12AsyncDeferredDeletion = ASYNC_DEFERRED_DELETION;
 
 static FAutoConsoleVariableRef CVarAsyncDeferredDeletion(
 	TEXT("D3D12.AsyncDeferredDeletion"),
diff --git a/Engine/Source/Runtime/D3D12RHI/Private/D3D12Texture.h b/Engine/Source/Runtime/D3D12RHI/Private/D3D12Texture.h
--- a/Engine/Source/Runtime/D3D12RHI/Private/D3D12Texture.h
+++ b/Engine/Source/Runtime/D3D12RHI/Private/D3D12Texture.h
@@ -190,7 +190,7 @@ public:
 	FD3D12Viewport* GetViewPort() { return Viewport; }
 	bool IsSDR() const { return bIsSDR; }
 
-	D3D12RHI_API FRHITexture* GetBackBufferTexture();
+	FRHITexture* GetBackBufferTexture();
 
 private:
 
diff --git a/Engine/Source/Runtime/D3D12RHI/Private/Windows/WindowsD3D12Device.cpp b/Engine/Source/Runtime/D3D12RHI/Private/Windows/WindowsD3D12Device.cpp
--- a/Engine/Source/Runtime/D3D12RHI/Private/Windows/WindowsD3D12Device.cpp
+++ b/Engine/Source/Runtime/D3D12RHI/Private/Windows/WindowsD3D12Device.cpp
@@ -679,7 +679,7 @@ void FD3D12DynamicRHIModule::FindAdapter()
 	{
 		UE_LOG(LogD3D12RHI, Log, TEXT("Forcing D3D12.AsyncDeferredDeletion=0 as a workaround for a deadlock on older versions of Windows."));
 
-		extern D3D12RHI_API int32 GD3D12AsyncDeferredDeletion;
+		extern int32 GD3D12AsyncDeferredDeletion;
 		GD3D12AsyncDeferredDeletion = 0;
 	}
 #endif // PLATFORM_DESKTOP
diff --git a/Engine/Source/Runtime/D3D12RHI/Public/D3D12Resources.h b/Engine/Source/Runtime/D3D12RHI/Public/D3D12Resources.h
--- a/Engine/Source/Runtime/D3D12RHI/Public/D3D12Resources.h
+++ b/Engine/Source/Runtime/D3D12RHI/Public/D3D12Resources.h
@@ -1093,7 +1093,7 @@ private:
 	uint32 ShadowBufferSize;
 };
 
-class D3D12RHI_API FD3D12GPUFence : public FRHIGPUFence
+class FD3D12GPUFence : public FRHIGPUFence
 {
 public:
 	FD3D12GPUFence(FName InName, FD3D12Fence* InFence)
