Implemented custom resolve codepath for D3D12.

diff --git a/Engine/Source/Runtime/D3D12RHI/Private/D3D12CommandContext.h b/Engine/Source/Runtime/D3D12RHI/Private/D3D12CommandContext.h
--- a/Engine/Source/Runtime/D3D12RHI/Private/D3D12CommandContext.h
+++ b/Engine/Source/Runtime/D3D12RHI/Private/D3D12CommandContext.h
@@ -352,13 +352,18 @@ public:
 			RHIEndOcclusionQueryBatch();
 		}
 
-		UE::RHICore::ResolveRenderPassTargets(*this, RenderPassInfo);
+		UE::RHICore::ResolveRenderPassTargets(RenderPassInfo, [this](UE::RHICore::FResolveTextureInfo Info)
+		{
+			ResolveTexture(Info);
+		});
 
 		FRHIRenderTargetView RTV(nullptr, ERenderTargetLoadAction::ENoAction);
 		FRHIDepthRenderTargetView DepthRTV(nullptr, ERenderTargetLoadAction::ENoAction, ERenderTargetStoreAction::ENoAction);
 		SetRenderTargets(1, &RTV, &DepthRTV);
 	}
 
+	void ResolveTexture(UE::RHICore::FResolveTextureInfo Info);
+
 	// When using Alternate Frame Rendering some temporal effects i.e. effects which consume GPU work from previous frames must synchronize their resources
 	// to prevent visual corruption.
 
diff --git a/Engine/Source/Runtime/D3D12RHI/Private/D3D12RenderTarget.cpp b/Engine/Source/Runtime/D3D12RHI/Private/D3D12RenderTarget.cpp
--- a/Engine/Source/Runtime/D3D12RHI/Private/D3D12RenderTarget.cpp
+++ b/Engine/Source/Runtime/D3D12RHI/Private/D3D12RenderTarget.cpp
@@ -478,6 +478,69 @@ void FD3D12CommandContext::RHICopyToResolveTarget(FRHITexture* SourceTextureRHI,
 	DEBUG_EXECUTE_COMMAND_LIST(this);
 }
 
+void FD3D12CommandContext::ResolveTexture(UE::RHICore::FResolveTextureInfo Info)
+{
+	uint32 GPUIndex = GetGPUIndex();
+	FRHICommandList_RecursiveHazardous RHICmdList(this, FRHIGPUMask::FromIndex(GPUIndex));
+
+	if (IsDefaultContext())
+	{
+		GetParentDevice()->RegisterGPUWork();
+	}
+
+	FD3D12Texture* SourceTexture         = GetD3D12TextureFromRHITexture(Info.SourceTexture, GPUIndex);
+	FD3D12Resource* SourceResource       = SourceTexture->GetResource();
+	const FRHITextureDesc& SourceDesc    = SourceTexture->GetDesc();
+
+	FD3D12Texture* DestTexture           = GetD3D12TextureFromRHITexture(Info.DestTexture, GPUIndex);
+	FD3D12Resource* DestResource         = DestTexture->GetResource();
+	const FRHITextureDesc& DestDesc      = DestTexture->GetDesc();
+
+	if (SourceDesc.Format == PF_DepthStencil)
+	{
+		ResolveTextureUsingShader<FResolveDepthPS>(
+			RHICmdList,
+			SourceTexture,
+			DestTexture,
+			DestTexture->GetRenderTargetView(0, -1),
+			DestTexture->GetDepthStencilView(FExclusiveDepthStencil::DepthWrite_StencilWrite),
+			DestResource->GetDesc(),
+			GetDefaultRect(Info.ResolveRect, SourceDesc.Extent.X, SourceDesc.Extent.Y),
+			GetDefaultRect(Info.ResolveRect, DestDesc.Extent.X, DestDesc.Extent.Y),
+			FDummyResolveParameter()
+		);
+	}
+	else
+	{
+		const DXGI_FORMAT DestFormatTypeless = ConvertTypelessToUnorm((DXGI_FORMAT)GPixelFormats[DestDesc.Format].PlatformFormat);
+
+		int32 ArraySliceBegin = Info.ArraySlice;
+		int32 ArraySliceEnd   = Info.ArraySlice + 1;
+
+		if (Info.ArraySlice < 0)
+		{
+			ArraySliceBegin = 0;
+			ArraySliceEnd   = SourceDesc.ArraySize;
+		}
+
+		for (int32 ArraySlice = ArraySliceBegin; ArraySlice < ArraySliceEnd; ArraySlice++)
+		{
+			int32 DestSubresource   = CalcSubresource(Info.MipLevel, ArraySlice, DestDesc.ArraySize);
+			int32 SourceSubresource = CalcSubresource(Info.MipLevel, ArraySlice, SourceDesc.ArraySize);
+
+			CommandListHandle->ResolveSubresource(DestResource->GetResource(), DestSubresource, SourceTexture->GetResource()->GetResource(), SourceSubresource, DestFormatTypeless);
+			otherWorkCounter++;
+		}
+	}
+
+	CommandListHandle.UpdateResidency(SourceTexture->GetResource());
+	CommandListHandle.UpdateResidency(DestTexture->GetResource());
+
+	ConditionalFlushCommandList();
+
+	DEBUG_EXECUTE_COMMAND_LIST(this);
+}
+
 static uint32 ComputeBytesPerPixel(DXGI_FORMAT Format)
 {
 	uint32 BytesPerPixel = 0;
diff --git a/Engine/Source/Runtime/RHICore/Private/RHICore.cpp b/Engine/Source/Runtime/RHICore/Private/RHICore.cpp
--- a/Engine/Source/Runtime/RHICore/Private/RHICore.cpp
+++ b/Engine/Source/Runtime/RHICore/Private/RHICore.cpp
@@ -2,9 +2,6 @@
 
 #include "RHICore.h"
 #include "Modules/ModuleManager.h"
-#include "RHI.h"
-#include "RHIResources.h"
-#include "RHIContext.h"
 
 DEFINE_LOG_CATEGORY(LogRHICore);
 IMPLEMENT_MODULE(FDefaultModuleImpl, RHICore);
@@ -63,5 +60,36 @@ void ResolveRenderPassTargets(IRHICommandContext& Context, const FRHIRenderPassI
 	ResolveTexture(DSV.DepthStencilTarget, DSV.ResolveTarget, 0, 0);
 }
 
+void ResolveRenderPassTargets(const FRHIRenderPassInfo& RenderPassInfo, TFunction<void(FResolveTextureInfo)> ResolveFunction)
+{
+	const auto ResolveTexture = [&](FResolveTextureInfo ResolveInfo)
+	{
+		if (!ResolveInfo.SourceTexture || !ResolveInfo.DestTexture || ResolveInfo.SourceTexture == ResolveInfo.DestTexture)
+		{
+			return;
+		}
+
+		const FRHITextureDesc& SourceDesc = ResolveInfo.SourceTexture->GetDesc();
+		const FRHITextureDesc& DestDesc   = ResolveInfo.DestTexture->GetDesc();
+
+		check(SourceDesc.Format == DestDesc.Format);
+		check(SourceDesc.Extent == DestDesc.Extent);
+		check(SourceDesc.IsMultisample() && !DestDesc.IsMultisample());
+		check(SourceDesc.Format != PF_DepthStencil || (SourceDesc.IsTexture2D() && DestDesc.IsTexture2D()));
+		check(!SourceDesc.IsTexture3D() && !DestDesc.IsTexture3D());
+
+		ResolveFunction(ResolveInfo);
+	};
+
+	for (int32 Index = 0; Index < MaxSimultaneousRenderTargets; ++Index)
+	{
+		const auto& RTV = RenderPassInfo.ColorRenderTargets[Index];
+		ResolveTexture({ RTV.RenderTarget, RTV.ResolveTarget, RTV.MipIndex, RTV.ArraySlice, RenderPassInfo.ResolveRect });
+	}
+
+	const auto& DSV = RenderPassInfo.DepthStencilRenderTarget;
+	ResolveTexture({ DSV.DepthStencilTarget, DSV.ResolveTarget, 0, 0, RenderPassInfo.ResolveRect });
+}
+
 } //! RHICore
 } //! UE
\ No newline at end of file
diff --git a/Engine/Source/Runtime/RHICore/Public/RHICore.h b/Engine/Source/Runtime/RHICore/Public/RHICore.h
--- a/Engine/Source/Runtime/RHICore/Public/RHICore.h
+++ b/Engine/Source/Runtime/RHICore/Public/RHICore.h
@@ -3,6 +3,9 @@
 #pragma once
 
 #include "Logging/LogMacros.h"
+#include "RHI.h"
+#include "RHIResources.h"
+#include "RHIContext.h"
 
 RHICORE_API DECLARE_LOG_CATEGORY_EXTERN(LogRHICore, Log, VeryVerbose);
 
@@ -17,5 +20,30 @@ namespace RHICore
 
 RHICORE_API void ResolveRenderPassTargets(IRHICommandContext& Context, const FRHIRenderPassInfo& Info);
 
+struct FResolveTextureInfo
+{
+	FResolveTextureInfo() = default;
+	
+	FResolveTextureInfo(
+		FRHITexture* InSourceTexture,
+		FRHITexture* InDestTexture,
+		uint8 InMipLevel,
+		int32 InArraySlice,
+		FResolveRect InResolveRect)
+		: SourceTexture(InSourceTexture)
+		, DestTexture(InDestTexture)
+		, MipLevel(InMipLevel)
+		, ArraySlice(InArraySlice)
+	{}
+
+	FRHITexture* SourceTexture = nullptr;
+	FRHITexture* DestTexture = nullptr;
+	uint8 MipLevel = 0;
+	int32 ArraySlice = -1;
+	FResolveRect ResolveRect;
+};
+
+RHICORE_API void ResolveRenderPassTargets(const FRHIRenderPassInfo& Info, TFunction<void(FResolveTextureInfo)> Function);
+
 } //! RHICore
 } //! UE
\ No newline at end of file
