fix incorrect endif comment

#rnx

diff --git a/Engine/Source/Runtime/D3D12RHI/Private/D3D12Adapter.cpp b/Engine/Source/Runtime/D3D12RHI/Private/D3D12Adapter.cpp
--- a/Engine/Source/Runtime/D3D12RHI/Private/D3D12Adapter.cpp
+++ b/Engine/Source/Runtime/D3D12RHI/Private/D3D12Adapter.cpp
@@ -199,7 +199,7 @@ static LONG __stdcall D3DVectoredExceptionHandler(EXCEPTION_POINTERS* InInfo)
 	return EXCEPTION_CONTINUE_SEARCH;
 }
 
-#endif // #if PLATFORM_WINDOWS || PLATFORM_HOLOLENS
+#endif // #if PLATFORM_WINDOWS
 
 
 FTransientUniformBufferAllocator::~FTransientUniformBufferAllocator()
