Include Callable Shader Library when building RTPSO.
- Using a basic callable shader as placeholder of invalid entries in library (Library array is sparse so that shaders can cache their index in RTPSO).
- Support NULL entries in callable shader SBT. Results in no-op shader. Useful when using r.RayTracing.NonBlockingPipelineCreation.

diff --git a/Engine/Shaders/Private/RayTracing/RayTracingMaterialDefaultHitShaders.usf b/Engine/Shaders/Private/RayTracing/RayTracingMaterialDefaultHitShaders.usf
--- a/Engine/Shaders/Private/RayTracing/RayTracingMaterialDefaultHitShaders.usf
+++ b/Engine/Shaders/Private/RayTracing/RayTracingMaterialDefaultHitShaders.usf
@@ -21,4 +21,16 @@ RAY_TRACING_ENTRY_ANY_HIT(HiddenMaterialAHS,
 	FRayTracingIntersectionAttributes, Attributes)
 {
 	IgnoreHit();
-}  
+}
+
+struct FDefaultCallableShaderParams
+{
+	int Param;
+};
+
+RAY_TRACING_ENTRY_CALLABLE(DefaultCallableShader,
+	FDefaultCallableShaderParams, Params)
+{
+	// This shader serves only has a placeholder in RTPSO
+	// If RTPSO doesn't include a required shader, a NULL SBT entry should be used.
+}
diff --git a/Engine/Source/Runtime/D3D12RHI/Private/D3D12RayTracing.cpp b/Engine/Source/Runtime/D3D12RHI/Private/D3D12RayTracing.cpp
--- a/Engine/Source/Runtime/D3D12RHI/Private/D3D12RayTracing.cpp
+++ b/Engine/Source/Runtime/D3D12RHI/Private/D3D12RayTracing.cpp
@@ -5540,22 +5540,29 @@ static void SetRayTracingCallableShader(
 	const uint32 UserDataOffset = offsetof(FHitGroupSystemParameters, RootConstants) + offsetof(FHitGroupSystemRootConstants, UserData);
 	ShaderTable->SetLocalShaderParameters(RecordIndex, UserDataOffset, UserData);
 
-	const FD3D12RayTracingShader* Shader = Pipeline->CallableShaders.Shaders[ShaderIndexInPipeline];
+	const FD3D12ShaderIdentifier* ShaderIdentifier = &FD3D12ShaderIdentifier::Null;
 
-	FD3D12RayTracingLocalResourceBinder ResourceBinder(*Device, *ShaderTable, *(Shader->pRootSignature), RecordIndex, WorkerIndex);
-	const bool bResourcesBound = SetRayTracingShaderResources(Shader,
-		0, nullptr, // Textures
-		0, nullptr, // SRVs
-		NumUniformBuffers, UniformBuffers,
-		0, nullptr, // Samplers
-		0, nullptr, // UAVs
-		LooseParameterDataSize, LooseParameterData, // Loose parameters
-		ResourceBinder);
+	if (ShaderIndexInPipeline != INDEX_NONE)
+	{
+		const FD3D12RayTracingShader* Shader = Pipeline->CallableShaders.Shaders[ShaderIndexInPipeline];
 
-	ShaderTable->SetLocalShaderIdentifier(RecordIndex,
-		bResourcesBound
-		? Pipeline->CallableShaders.Identifiers[ShaderIndexInPipeline]
-		: FD3D12ShaderIdentifier::Null);
+		FD3D12RayTracingLocalResourceBinder ResourceBinder(*Device, *ShaderTable, *(Shader->pRootSignature), RecordIndex, WorkerIndex);
+		const bool bResourcesBound = SetRayTracingShaderResources(Shader,
+			0, nullptr, // Textures
+			0, nullptr, // SRVs
+			NumUniformBuffers, UniformBuffers,
+			0, nullptr, // Samplers
+			0, nullptr, // UAVs
+			LooseParameterDataSize, LooseParameterData, // Loose parameters
+			ResourceBinder);
+
+		if (bResourcesBound)
+		{
+			ShaderIdentifier = &Pipeline->CallableShaders.Identifiers[ShaderIndexInPipeline];
+		}
+	}
+
+	ShaderTable->SetLocalShaderIdentifier(RecordIndex, *ShaderIdentifier);
 }
 
 static void SetRayTracingMissShader(
diff --git a/Engine/Source/Runtime/Renderer/Private/RayTracing/RayTracingMaterialHitShaders.cpp b/Engine/Source/Runtime/Renderer/Private/RayTracing/RayTracingMaterialHitShaders.cpp
--- a/Engine/Source/Runtime/Renderer/Private/RayTracing/RayTracingMaterialHitShaders.cpp
+++ b/Engine/Source/Runtime/Renderer/Private/RayTracing/RayTracingMaterialHitShaders.cpp
@@ -235,6 +235,7 @@ IMPLEMENT_MATERIALCHS_TYPE(TUniformLightMapPolicy<LMP_DISTANCE_FIELD_SHADOWS_AND
 
 IMPLEMENT_GLOBAL_SHADER(FHiddenMaterialHitGroup, "/Engine/Private/RayTracing/RayTracingMaterialDefaultHitShaders.usf", "closesthit=HiddenMaterialCHS anyhit=HiddenMaterialAHS", SF_RayHitGroup);
 IMPLEMENT_GLOBAL_SHADER(FOpaqueShadowHitGroup, "/Engine/Private/RayTracing/RayTracingMaterialDefaultHitShaders.usf", "closesthit=OpaqueShadowCHS", SF_RayHitGroup);
+IMPLEMENT_GLOBAL_SHADER(FDefaultCallableShader, "/Engine/Private/RayTracing/RayTracingMaterialDefaultHitShaders.usf", "DefaultCallableShader", SF_RayCallable);
 
 // Select TextureLOD
 template<typename LightMapPolicyType, bool bUseAnyHitShader, bool bUseIntersectionShader>
@@ -641,6 +642,16 @@ FRayTracingPipelineState* FDeferredShadingSceneRenderer::BindRayTracingMaterialP
 
 	Initializer.SetHitGroupTable(RayTracingHitGroupLibrary);
 
+	TArray<FRHIRayTracingShader*> RayTracingCallableShaderLibrary;
+
+	if (bEnableMaterials)
+	{
+		FRHIRayTracingShader* DefaultCallableShader = View.ShaderMap->GetShader<FDefaultCallableShader>().GetRayTracingShader();
+		FShaderMapResource::GetRayTracingCallableShaderLibrary(RayTracingCallableShaderLibrary, DefaultCallableShader);
+	}
+
+	Initializer.SetCallableTable(RayTracingCallableShaderLibrary);
+
 	FRayTracingPipelineState* FallbackPipelineState = GRayTracingNonBlockingPipelineCreation && View.ViewState
 		? PipelineStateCache::GetRayTracingPipelineState(View.ViewState->LastRayTracingMaterialPipelineSignature)
 		: nullptr;
diff --git a/Engine/Source/Runtime/Renderer/Private/RayTracing/RayTracingMaterialHitShaders.h b/Engine/Source/Runtime/Renderer/Private/RayTracing/RayTracingMaterialHitShaders.h
--- a/Engine/Source/Runtime/Renderer/Private/RayTracing/RayTracingMaterialHitShaders.h
+++ b/Engine/Source/Runtime/Renderer/Private/RayTracing/RayTracingMaterialHitShaders.h
@@ -154,6 +154,19 @@ class RENDERER_API FOpaqueShadowHitGroup : public FGlobalShader
 	using FParameters = FEmptyShaderParameters;
 };
 
+class RENDERER_API FDefaultCallableShader : public FGlobalShader
+{
+	DECLARE_GLOBAL_SHADER(FDefaultCallableShader)
+	SHADER_USE_ROOT_PARAMETER_STRUCT(FDefaultCallableShader, FGlobalShader)
+
+	static bool ShouldCompilePermutation(const FGlobalShaderPermutationParameters& Parameters)
+	{
+		return ShouldCompileRayTracingCallableShadersForProject(Parameters.Platform);
+	}
+
+	using FParameters = FEmptyShaderParameters;
+};
+
 class FRayTracingLocalShaderBindingWriter
 {
 public:
