D3D12 Bindless Heaps
- When switching to a RootSignature that uses bindless resources/samplers, override the currently bound descriptor heaps with the bindless versions.
- If we are using bindless resources or samplers, completely skip the normal non-bindless descriptor setup. This required some minor reorganizing of FD3D12StateCacheBase.
- Adding temporary logging to QuantizeBoundShaderState to check that all shader stages can support one ore more other concurrent stages in using bindless resources/samplers. This should be replaced with a higher level "can this PSO be created from a given Desc" logic.

diff --git a/Engine/Source/Runtime/D3D12RHI/Private/D3D12DescriptorCache.cpp b/Engine/Source/Runtime/D3D12RHI/Private/D3D12DescriptorCache.cpp
--- a/Engine/Source/Runtime/D3D12RHI/Private/D3D12DescriptorCache.cpp
+++ b/Engine/Source/Runtime/D3D12RHI/Private/D3D12DescriptorCache.cpp
@@ -56,18 +56,18 @@ void FD3D12DescriptorCache::HeapLoopedAround(ERHIDescriptorHeapType InHeapType)
 FD3D12DescriptorCache::FD3D12DescriptorCache(FRHIGPUMask Node)
 	: FD3D12DeviceChild(nullptr)
 	, FD3D12SingleNodeGPUObject(Node)
-	, pPreviousViewHeap(nullptr)
-	, pPreviousSamplerHeap(nullptr)
-	, CurrentViewHeap(nullptr)
-	, CurrentSamplerHeap(nullptr)
-	, LocalViewHeap(nullptr)
 	, LocalSamplerHeap(this)
 	, SubAllocatedViewHeap(this)
 	, SamplerMap(271) // Prime numbers for better hashing
-	, bUsingGlobalSamplerHeap(false)
-	, NumLocalViewDescriptors(0)
 {
-	CmdContext = nullptr;
+}
+
+FD3D12DescriptorCache::~FD3D12DescriptorCache()
+{
+	if (LocalViewHeap)
+	{
+		delete LocalViewHeap;
+	}
 }
 
 void FD3D12DescriptorCache::Init(FD3D12Device* InParent, FD3D12CommandContext* InCmdContext, uint32 InNumLocalViewDescriptors, uint32 InNumSamplerDescriptors)
@@ -88,6 +88,9 @@ void FD3D12DescriptorCache::Init(FD3D12Device* InParent, FD3D12CommandContext* I
 	CurrentSamplerHeap = &LocalSamplerHeap;
 	bUsingGlobalSamplerHeap = false;
 
+	OverrideViewHeap = nullptr;
+	OverrideSamplerHeap = nullptr;
+
 	// Create default views
 	D3D12_SHADER_RESOURCE_VIEW_DESC SRVDesc = {};
 	SRVDesc.Shader4ComponentMapping = D3D12_DEFAULT_SHADER_4_COMPONENT_MAPPING;
@@ -219,41 +222,50 @@ bool FD3D12DescriptorCache::SetDescriptorHeaps()
 
 	// See if the descriptor heaps changed.
 	bool bHeapChanged = false;
-	ID3D12DescriptorHeap* const pCurrentViewHeap = CurrentViewHeap->GetHeap();
-	if (pPreviousViewHeap != pCurrentViewHeap)
+
+	ID3D12DescriptorHeap* PendingViewHeap = OverrideViewHeap ? OverrideViewHeap->GetHeap() : CurrentViewHeap->GetHeap();
+	if (LastSetViewHeap != PendingViewHeap)
 	{
 		// The view heap changed, so dirty the descriptor tables.
 		bHeapChanged = true;
-		CmdContext->StateCache.DirtyViewDescriptorTables();
 
-		INC_DWORD_STAT_BY(STAT_ViewHeapChanged, pPreviousViewHeap == nullptr ? 0 : 1);	// Don't count the initial set on a command list.
+		if (!OverrideViewHeap)
+		{
+			CmdContext->StateCache.DirtyViewDescriptorTables();
+		}
+
+		INC_DWORD_STAT_BY(STAT_ViewHeapChanged, LastSetViewHeap == nullptr ? 0 : 1);	// Don't count the initial set on a command list.
 	}
 
-	ID3D12DescriptorHeap* const pCurrentSamplerHeap = CurrentSamplerHeap->GetHeap();
-	if (pPreviousSamplerHeap != pCurrentSamplerHeap)
+	ID3D12DescriptorHeap* PendingSamplerHeap = OverrideSamplerHeap ? OverrideSamplerHeap->GetHeap() : CurrentSamplerHeap->GetHeap();
+	if (LastSetSamplerHeap != PendingSamplerHeap)
 	{
 		// The sampler heap changed, so dirty the descriptor tables.
 		bHeapChanged = true;
-		CmdContext->StateCache.DirtySamplerDescriptorTables();
 
-		// Reset the sampler map since it will have invalid entries for the new heap.
-		SamplerMap.Reset();
+		if (!OverrideSamplerHeap)
+		{
+			CmdContext->StateCache.DirtySamplerDescriptorTables();
 
-		INC_DWORD_STAT_BY(STAT_SamplerHeapChanged, pPreviousSamplerHeap == nullptr ? 0 : 1);	// Don't count the initial set on a command list.
+			// Reset the sampler map since it will have invalid entries for the new heap.
+			SamplerMap.Reset();
+		}
+
+		INC_DWORD_STAT_BY(STAT_SamplerHeapChanged, LastSetSamplerHeap == nullptr ? 0 : 1);	// Don't count the initial set on a command list.
 	}
 
 	// Set the descriptor heaps.
 	if (bHeapChanged)
 	{
-		ID3D12DescriptorHeap* /*const*/ ppHeaps[] = { pCurrentViewHeap, pCurrentSamplerHeap };
+		ID3D12DescriptorHeap* ppHeaps[] = { PendingViewHeap, PendingSamplerHeap };
 		CmdContext->CommandListHandle->SetDescriptorHeaps(UE_ARRAY_COUNT(ppHeaps), ppHeaps);
 
-		pPreviousViewHeap = pCurrentViewHeap;
-		pPreviousSamplerHeap = pCurrentSamplerHeap;
+		LastSetViewHeap = PendingViewHeap;
+		LastSetSamplerHeap = PendingSamplerHeap;
 	}
 
-	check(pPreviousSamplerHeap == pCurrentSamplerHeap);
-	check(pPreviousViewHeap == pCurrentViewHeap);
+	check(LastSetSamplerHeap == PendingSamplerHeap);
+	check(LastSetViewHeap == PendingViewHeap);
 	return bHeapChanged;
 }
 
@@ -261,8 +273,8 @@ bool FD3D12DescriptorCache::SetDescriptorHeaps()
 void FD3D12DescriptorCache::SetCurrentCommandList(const FD3D12CommandListHandle& CommandListHandle)
 {
 	// Clear the previous heap pointers (since it's a new command list) and then set the current descriptor heaps.
-	pPreviousViewHeap = nullptr;
-	pPreviousSamplerHeap = nullptr;
+	LastSetViewHeap = nullptr;
+	LastSetSamplerHeap = nullptr;
 
 	CurrentViewHeap->SetCurrentCommandList(CommandListHandle);
 
@@ -801,6 +813,20 @@ bool FD3D12DescriptorCache::SwitchToGlobalSamplerHeap()
 	return bDescriptorHeapsChanged;
 }
 
+bool FD3D12DescriptorCache::SetHeapOverrides(FD3D12DescriptorHeap* InOverrideViewHeap, FD3D12DescriptorHeap* InOverrideSamplerHeap)
+{
+	bool bDescriptorHeapsChanged = false;
+
+	if (OverrideViewHeap != InOverrideViewHeap || OverrideSamplerHeap != InOverrideSamplerHeap)
+	{
+		OverrideViewHeap = InOverrideViewHeap;
+		OverrideSamplerHeap = InOverrideSamplerHeap;
+
+		bDescriptorHeapsChanged = SetDescriptorHeaps();
+	}
+
+	return bDescriptorHeapsChanged;
+}
 
 ////////////////////////////////////////////////////////////////////////////////////////////////////////////////////////
 // FD3D12OnlineHeap
diff --git a/Engine/Source/Runtime/D3D12RHI/Private/D3D12DescriptorCache.h b/Engine/Source/Runtime/D3D12RHI/Private/D3D12DescriptorCache.h
--- a/Engine/Source/Runtime/D3D12RHI/Private/D3D12DescriptorCache.h
+++ b/Engine/Source/Runtime/D3D12RHI/Private/D3D12DescriptorCache.h
@@ -320,24 +320,22 @@ private:
 
 class FD3D12DescriptorCache : public FD3D12DeviceChild, public FD3D12SingleNodeGPUObject
 {
-protected:
-	FD3D12CommandContext* CmdContext;
-
 public:
-	FD3D12OnlineHeap* GetCurrentViewHeap() { return CurrentViewHeap; }
-	FD3D12OnlineHeap* GetCurrentSamplerHeap() { return CurrentSamplerHeap; }
-
+	FD3D12DescriptorCache() = delete;
 	FD3D12DescriptorCache(FRHIGPUMask Node);
 
-	~FD3D12DescriptorCache()
-	{
-		if (LocalViewHeap) { delete(LocalViewHeap); }
-	}
+	~FD3D12DescriptorCache();
+
+	inline bool IsViewHeapOverridden() const { return OverrideViewHeap != nullptr; }
+	inline bool IsSamplerHeapOverridden() const { return OverrideViewHeap != nullptr; }
+
+	inline FD3D12OnlineHeap* GetCurrentViewHeap() const { return CurrentViewHeap; }
+	inline FD3D12OnlineHeap* GetCurrentSamplerHeap() const { return CurrentSamplerHeap; }
 
 	// Checks if the specified descriptor heap has been set on the current command list.
 	bool IsHeapSet(ID3D12DescriptorHeap* const pHeap) const
 	{
-		return (pHeap == pPreviousViewHeap) || (pHeap == pPreviousSamplerHeap);
+		return (pHeap == LastSetViewHeap) || (pHeap == LastSetSamplerHeap);
 	}
 
 	// Notify the descriptor cache every time you start recording a command list.
@@ -348,17 +346,6 @@ public:
 	// ------------------------------------------------------
 	// end Descriptor Slot Reservation stuff
 
-	// null views
-
-	FD3D12ViewDescriptorHandle* NullSRV{};
-	FD3D12ViewDescriptorHandle* NullRTV{};
-	FD3D12ViewDescriptorHandle* NullUAV{};
-
-#if USE_STATIC_ROOT_SIGNATURE
-	FD3D12ConstantBufferView* NullCBV{};
-#endif
-	TRefCountPtr<FD3D12SamplerState> DefaultSampler;
-
 	void SetVertexBuffers(FD3D12VertexBufferCache& Cache);
 	void SetRenderTargets(FD3D12RenderTargetView** RenderTargetViewArray, uint32 Count, FD3D12DepthStencilView* DepthStencilTarget);
 
@@ -392,24 +379,42 @@ public:
 	bool SwitchToContextLocalSamplerHeap();
 	bool SwitchToGlobalSamplerHeap();
 
+	bool SetHeapOverrides(FD3D12DescriptorHeap* InOverrideViewHeap, FD3D12DescriptorHeap* InOverrideSamplerHeap);
+
 	TArray<FD3D12UniqueSamplerTable>& GetUniqueTables() { return UniqueTables; }
 
 	inline bool UsingGlobalSamplerHeap() const { return bUsingGlobalSamplerHeap; }
 	FD3D12SamplerSet& GetLocalSamplerSet() { return LocalSamplerSet; }
 
+	// null views
+	FD3D12ViewDescriptorHandle* NullSRV = nullptr;
+	FD3D12ViewDescriptorHandle* NullRTV = nullptr;
+	FD3D12ViewDescriptorHandle* NullUAV = nullptr;
+
+#if USE_STATIC_ROOT_SIGNATURE
+	FD3D12ConstantBufferView* NullCBV = nullptr;
+#endif
+	TRefCountPtr<FD3D12SamplerState> DefaultSampler;
+
+protected:
+	FD3D12CommandContext* CmdContext = nullptr;
+
 private:
 	// Sets the current descriptor tables on the command list and marks any descriptor tables as dirty if necessary.
 	// Returns true if one of the heaps actually changed, false otherwise.
 	bool SetDescriptorHeaps();
 
 	// The previous view and sampler heaps set on the current command list.
-	ID3D12DescriptorHeap* pPreviousViewHeap;
-	ID3D12DescriptorHeap* pPreviousSamplerHeap;
+	ID3D12DescriptorHeap* LastSetViewHeap = nullptr;
+	ID3D12DescriptorHeap* LastSetSamplerHeap = nullptr;
+
+	FD3D12OnlineHeap* CurrentViewHeap = nullptr;
+	FD3D12OnlineHeap* CurrentSamplerHeap = nullptr;
 
-	FD3D12OnlineHeap* CurrentViewHeap;
-	FD3D12OnlineHeap* CurrentSamplerHeap;
+	FD3D12DescriptorHeap* OverrideViewHeap = nullptr;
+	FD3D12DescriptorHeap* OverrideSamplerHeap = nullptr;
 
-	FD3D12LocalOnlineHeap* LocalViewHeap;
+	FD3D12LocalOnlineHeap* LocalViewHeap = nullptr;
 	FD3D12LocalOnlineHeap LocalSamplerHeap;
 	FD3D12SubAllocatedOnlineHeap SubAllocatedViewHeap;
 
@@ -418,7 +423,7 @@ private:
 	TArray<FD3D12UniqueSamplerTable> UniqueTables;
 
 	FD3D12SamplerSet LocalSamplerSet;
-	bool bUsingGlobalSamplerHeap;
+	bool bUsingGlobalSamplerHeap = false;
 
-	uint32 NumLocalViewDescriptors;
+	uint32 NumLocalViewDescriptors = 0;
 };
\ No newline at end of file
diff --git a/Engine/Source/Runtime/D3D12RHI/Private/D3D12Descriptors.cpp b/Engine/Source/Runtime/D3D12RHI/Private/D3D12Descriptors.cpp
--- a/Engine/Source/Runtime/D3D12RHI/Private/D3D12Descriptors.cpp
+++ b/Engine/Source/Runtime/D3D12RHI/Private/D3D12Descriptors.cpp
@@ -189,6 +189,19 @@ void FD3D12BindlessDescriptorManager::UpdateDeferred(FRHIDescriptorHandle InHand
 	UpdateImmediately(InHandle, InSourceCpuHandle);
 }
 
+FD3D12DescriptorHeap* FD3D12BindlessDescriptorManager::GetHeapForType(ERHIDescriptorHeapType InType)
+{
+	for (FD3D12DescriptorManager& Manager : Managers)
+	{
+		if (Manager.HandlesAllocation(InType))
+		{
+			return Manager.GetHeap();
+		}
+	}
+
+	return nullptr;
+}
+
 ////////////////////////////////////////////////////////////////////////////////////////////////////////////////////////
 // FD3D12OnlineDescriptorManager
 
diff --git a/Engine/Source/Runtime/D3D12RHI/Private/D3D12Descriptors.h b/Engine/Source/Runtime/D3D12RHI/Private/D3D12Descriptors.h
--- a/Engine/Source/Runtime/D3D12RHI/Private/D3D12Descriptors.h
+++ b/Engine/Source/Runtime/D3D12RHI/Private/D3D12Descriptors.h
@@ -105,6 +105,8 @@ public:
 	void UpdateImmediately(FRHIDescriptorHandle InHandle, D3D12_CPU_DESCRIPTOR_HANDLE InSourceCpuHandle);
 	void UpdateDeferred(FRHIDescriptorHandle InHandle, D3D12_CPU_DESCRIPTOR_HANDLE InSourceCpuHandle);
 
+	FD3D12DescriptorHeap* GetHeapForType(ERHIDescriptorHeapType InType);
+
 private:
 	TArray<FD3D12DescriptorManager> Managers;
 };
diff --git a/Engine/Source/Runtime/D3D12RHI/Private/D3D12RootSignature.cpp b/Engine/Source/Runtime/D3D12RHI/Private/D3D12RootSignature.cpp
--- a/Engine/Source/Runtime/D3D12RHI/Private/D3D12RootSignature.cpp
+++ b/Engine/Source/Runtime/D3D12RHI/Private/D3D12RootSignature.cpp
@@ -589,14 +589,15 @@ void FD3D12RootSignature::InternalAnalyzeSignature(const RootSignatureDescType&
 		bHasUAVs = false;
 		bHasSRVs = false;
 		bHasCBVs = false;
-		bHasRDTCBVs = false;
-		bHasRDCBVs = false;
 		bHasSamplers = false;
 
 		FMemory::Memset(BindSlotOffsetsInDWORDs, 0, sizeof(BindSlotOffsetsInDWORDs));
 		TotalRootSignatureSizeInDWORDs = 0;
 	}
 
+	bUsesDynamicResources = (Desc.Flags & D3D12_ROOT_SIGNATURE_FLAG_CBV_SRV_UAV_HEAP_DIRECTLY_INDEXED) != 0;
+	bUsesDynamicSamplers = (Desc.Flags & D3D12_ROOT_SIGNATURE_FLAG_SAMPLER_HEAP_DIRECTLY_INDEXED) != 0;
+
 	const bool bDenyVS = (Desc.Flags & D3D12_ROOT_SIGNATURE_FLAG_DENY_VERTEX_SHADER_ROOT_ACCESS) != 0;
 	const bool bDenyGS = (Desc.Flags & D3D12_ROOT_SIGNATURE_FLAG_DENY_GEOMETRY_SHADER_ROOT_ACCESS) != 0;
 	const bool bDenyPS = (Desc.Flags & D3D12_ROOT_SIGNATURE_FLAG_DENY_PIXEL_SHADER_ROOT_ACCESS) != 0;
diff --git a/Engine/Source/Runtime/D3D12RHI/Private/D3D12RootSignature.h b/Engine/Source/Runtime/D3D12RHI/Private/D3D12RootSignature.h
--- a/Engine/Source/Runtime/D3D12RHI/Private/D3D12RootSignature.h
+++ b/Engine/Source/Runtime/D3D12RHI/Private/D3D12RootSignature.h
@@ -66,23 +66,13 @@ private:
 	// Struct for all the useful info we want per shader stage.
 	struct ShaderStage
 	{
-		ShaderStage()
-			: MaxCBVCount(0u)
-			, MaxSRVCount(0u)
-			, MaxSamplerCount(0u)
-			, MaxUAVCount(0u)
-			, CBVRegisterMask(0u)
-			, bVisible(false)
-		{
-		}
-
 		// TODO: Make these arrays and index into them by type instead of individual variables.
-		uint8 MaxCBVCount;
-		uint8 MaxSRVCount;
-		uint8 MaxSamplerCount;
-		uint8 MaxUAVCount;
-		CBVSlotMask CBVRegisterMask;
-		bool bVisible;
+		uint8 MaxCBVCount = 0;
+		uint8 MaxSRVCount = 0;
+		uint8 MaxSamplerCount = 0;
+		uint8 MaxUAVCount = 0;
+		CBVSlotMask CBVRegisterMask = 0;
+		bool bVisible = false;
 	};
 
 public:
@@ -194,6 +184,9 @@ public:
 	inline bool HasSRVs() const { return bHasSRVs; }
 	inline bool HasCBVs() const { return bHasCBVs; }
 	inline bool HasSamplers() const { return bHasSamplers; }
+	inline bool UsesDynamicResources() const { return bUsesDynamicResources; }
+	inline bool UsesDynamicSamplers() const { return bUsesDynamicSamplers; }
+
 	inline bool HasVS() const { return Stage[SF_Vertex].bVisible; }
 	inline bool HasMS() const { return Stage[SF_Mesh].bVisible; }
 	inline bool HasAS() const { return Stage[SF_Amplification].bVisible; }
@@ -293,7 +286,6 @@ private:
 		*pBindSlot = RootParameterIndex;
 
 		bHasCBVs = true;
-		bHasRDTCBVs = true;
 	}
 
 	inline void SetCBVRDBindSlot(EShaderFrequency SF, uint8 RootParameterIndex)
@@ -318,7 +310,6 @@ private:
 		*pBindSlot = RootParameterIndex;
 
 		bHasCBVs = true;
-		bHasRDCBVs = true;
 	}
 
 	inline void SetUAVRDTBindSlot(EShaderFrequency SF, uint8 RootParameterIndex)
@@ -459,10 +450,10 @@ private:
 	uint8 bHasUAVs : 1;
 	uint8 bHasSRVs : 1;
 	uint8 bHasCBVs : 1;
-	uint8 bHasRDTCBVs : 1;
-	uint8 bHasRDCBVs : 1;
 	uint8 bHasSamplers : 1;
 	uint8 bHasVendorExtensionSpace : 1;
+	uint8 bUsesDynamicResources : 1;
+	uint8 bUsesDynamicSamplers : 1;
 	uint8 bUnused : 1;
 };
 
diff --git a/Engine/Source/Runtime/D3D12RHI/Private/D3D12StateCache.cpp b/Engine/Source/Runtime/D3D12RHI/Private/D3D12StateCache.cpp
--- a/Engine/Source/Runtime/D3D12RHI/Private/D3D12StateCache.cpp
+++ b/Engine/Source/Runtime/D3D12RHI/Private/D3D12StateCache.cpp
@@ -409,6 +409,47 @@
 	return ((Stage == SF_Mesh || Stage == SF_Amplification) && !GRHISupportsMeshShadersTier0);
 }
 
+bool FD3D12StateCacheBase::InternalSetRootSignature(ED3D12PipelineType InPipelineType, const FD3D12RootSignature* InRootSignature)
+{
+	bool bWasRootSignatureChanged = false;
+
+	if (InPipelineType == ED3D12PipelineType::Compute)
+	{
+		if (PipelineState.Compute.bNeedSetRootSignature)
+		{
+			CmdContext->CommandListHandle->SetComputeRootSignature(InRootSignature->GetRootSignature());
+			PipelineState.Compute.bNeedSetRootSignature = false;
+
+			// After setting a root signature, all root parameters are undefined and must be set again.
+			PipelineState.Common.SRVCache.DirtyCompute();
+			PipelineState.Common.UAVCache.DirtyCompute();
+			PipelineState.Common.SamplerCache.DirtyCompute();
+			PipelineState.Common.CBVCache.DirtyCompute();
+
+			bWasRootSignatureChanged = true;
+		}
+	}
+	else if (InPipelineType == ED3D12PipelineType::Graphics)
+	{
+		// See if we need to set a graphics root signature
+		if (PipelineState.Graphics.bNeedSetRootSignature)
+		{
+			CmdContext->CommandListHandle->SetGraphicsRootSignature(InRootSignature->GetRootSignature());
+			PipelineState.Graphics.bNeedSetRootSignature = false;
+
+			// After setting a root signature, all root parameters are undefined and must be set again.
+			PipelineState.Common.SRVCache.DirtyGraphics();
+			PipelineState.Common.UAVCache.DirtyGraphics();
+			PipelineState.Common.SamplerCache.DirtyGraphics();
+			PipelineState.Common.CBVCache.DirtyGraphics();
+
+			bWasRootSignatureChanged = true;
+		}
+	}
+
+	return bWasRootSignatureChanged;
+}
+
 void FD3D12StateCacheBase::InternalSetPipelineState(FD3D12PipelineState* InPipelineState)
 {
 	// See if we need to set our PSO:
@@ -444,48 +485,44 @@
 #endif
 
 	FD3D12CommandListHandle& CommandList = CmdContext->CommandListHandle;
-	const FD3D12RootSignature* pRootSignature = nullptr;
 
+	FD3D12PipelineStateCommonData* PSOCommonData = nullptr;
+
 	// PSO
 	if (PipelineType == ED3D12PipelineType::Compute)
 	{
-		pRootSignature = PipelineState.Compute.CurrentPipelineStateObject->RootSignature;
-
-		// See if we need to set a compute root signature
-		if (PipelineState.Compute.bNeedSetRootSignature)
-		{
-			CommandList->SetComputeRootSignature(pRootSignature->GetRootSignature());
-			PipelineState.Compute.bNeedSetRootSignature = false;
-
-			// After setting a root signature, all root parameters are undefined and must be set again.
-			PipelineState.Common.SRVCache.DirtyCompute();
-			PipelineState.Common.UAVCache.DirtyCompute();
-			PipelineState.Common.SamplerCache.DirtyCompute();
-			PipelineState.Common.CBVCache.DirtyCompute();
-		}
-
-		// Ensure the correct compute PSO is set.
-		InternalSetPipelineState(GetComputePipelineState()->PipelineState);
+		PSOCommonData = GetComputePipelineState();
 	}
 	else if (PipelineType == ED3D12PipelineType::Graphics)
 	{
-		pRootSignature = PipelineState.Graphics.CurrentPipelineStateObject->RootSignature;
-
-		// See if we need to set a graphics root signature
-		if (PipelineState.Graphics.bNeedSetRootSignature)
-		{
-			CommandList->SetGraphicsRootSignature(pRootSignature->GetRootSignature());
-			PipelineState.Graphics.bNeedSetRootSignature = false;
-
-			// After setting a root signature, all root parameters are undefined and must be set again.
-			PipelineState.Common.SRVCache.DirtyGraphics();
-			PipelineState.Common.UAVCache.DirtyGraphics();
-			PipelineState.Common.SamplerCache.DirtyGraphics();
-			PipelineState.Common.CBVCache.DirtyGraphics();
-		}
-
-		// Ensure the correct graphics PSO is set.
-		InternalSetPipelineState(GetGraphicsPipelineState()->PipelineState);
+		PSOCommonData = GetGraphicsPipelineState();
 	}
 
+	const bool bRootSignatureChanged = InternalSetRootSignature(PipelineType, PSOCommonData->RootSignature);
+
+	// Ensure the correct graphics PSO is set.
+	InternalSetPipelineState(PSOCommonData->PipelineState);
+
+	bool bBindlessResources = false;
+	bool bBindlessSamplers = false;
+
+	if (bRootSignatureChanged)
+	{
+		bBindlessResources = PSOCommonData->RootSignature->UsesDynamicResources();
+		bBindlessSamplers = PSOCommonData->RootSignature->UsesDynamicSamplers();
+		
+		FD3D12BindlessDescriptorManager& BindlessManager = GetParentDevice()->GetBindlessDescriptorManager();
+
+		FD3D12DescriptorHeap* ResourceHeap = BindlessManager.GetHeapForType(ERHIDescriptorHeapType::Standard);
+		FD3D12DescriptorHeap* SamplerHeap = BindlessManager.GetHeapForType(ERHIDescriptorHeapType::Sampler);
+
+		checkf(!bBindlessResources || ResourceHeap != nullptr, TEXT("Using dynamic samplers without the bindless sampler heap configured. Please check your configuration."));
+		checkf(!bBindlessSamplers || SamplerHeap != nullptr, TEXT("Using dynamic samplers without the bindless sampler heap configured. Please check your configuration."));
+		
+		DescriptorCache.SetHeapOverrides(
+			bBindlessResources ? BindlessManager.GetHeapForType(ERHIDescriptorHeapType::Standard) : nullptr,
+			bBindlessSamplers ? BindlessManager.GetHeapForType(ERHIDescriptorHeapType::Sampler) : nullptr
+		);
+	}
+
 	// Need to cache compute budget, as we need to reset after PSO changes
@@ -543,7 +580,7 @@
 			bNeedSetDepthBounds = false;
 			CmdContext->SetDepthBounds(PipelineState.Graphics.MinDepth, PipelineState.Graphics.MaxDepth);
 		}
-		
+
 		if (bNeedSetShadingRate)
 		{
 			bNeedSetShadingRate = false;
@@ -555,16 +592,33 @@
 	const uint32 StartStage = PipelineType == ED3D12PipelineType::Graphics ? 0 : SF_Compute;
 	const uint32 EndStage = PipelineType == ED3D12PipelineType::Graphics ? SF_Compute : SF_NumStandardFrequencies;
 
-	const EShaderFrequency UAVStage = PipelineType == ED3D12PipelineType::Graphics ? SF_Pixel : SF_Compute;
-
 	//
 	// Reserve space in descriptor heaps
 	// Since this can cause heap rollover (which causes old bindings to become invalid), the reserve must be done atomically
 	//
 
 	// Samplers
-	ApplySamplers(pRootSignature, StartStage, EndStage);
+	if (!bBindlessSamplers)
+	{
+		ApplySamplers(PSOCommonData->RootSignature, StartStage, EndStage);
+	}
 
+	if (!bBindlessResources)
+	{
+		ApplyResources(PSOCommonData->RootSignature, StartStage, EndStage);
+	}
+
+	// Flush any needed resource barriers
+	CommandList.FlushResourceBarriers();
+
+#if ASSERT_RESOURCE_STATES
+	bool bSucceeded = AssertResourceStates(PipelineType);
+	check(bSucceeded);
+#endif
+}
+
+void FD3D12StateCacheBase::ApplyResources(const FD3D12RootSignature* const pRootSignature, uint32 StartStage, uint32 EndStage)
+{
 	// Determine what resource bind slots are dirty for the current shaders and how many descriptor table slots we need.
 	// We only set dirty resources that can be used for the upcoming Draw/Dispatch.
 	SRVSlotMask CurrentShaderDirtySRVSlots[SF_NumStandardFrequencies] = {};
@@ -576,6 +630,9 @@
 	uint32 NumCBVs[SF_NumStandardFrequencies] ={};
 #endif
 	uint32 NumViews = 0;
+
+	const EShaderFrequency UAVStage = StartStage == SF_Compute ? SF_Compute : SF_Pixel;
+
 	for (uint32 iTries = 0; iTries < 2; ++iTries)
 	{
 		const UAVSlotMask CurrentShaderUAVRegisterMask = BitMask<UAVSlotMask>(PipelineState.Common.CurrentShaderUAVCounts[UAVStage]);
@@ -666,7 +723,14 @@
 	if (CurrentShaderDirtyUAVSlots)
 	{
 		SCOPE_CYCLE_COUNTER(STAT_D3D12ApplyStateSetUAVTime);
-		DescriptorCache.SetUAVs<UAVStage>(pRootSignature, PipelineState.Common.UAVCache, CurrentShaderDirtyUAVSlots, NumUAVs, ViewHeapSlot);
+		if (UAVStage == SF_Pixel)
+		{
+			DescriptorCache.SetUAVs<SF_Pixel>(pRootSignature, PipelineState.Common.UAVCache, CurrentShaderDirtyUAVSlots, NumUAVs, ViewHeapSlot);
+		}
+		else
+		{
+			DescriptorCache.SetUAVs<SF_Compute>(pRootSignature, PipelineState.Common.UAVCache, CurrentShaderDirtyUAVSlots, NumUAVs, ViewHeapSlot);
+		}
 	}
 
 	// Shader resource views
@@ -680,7 +744,12 @@
 			DescriptorCache.SetSRVs<##Shader>(pRootSignature, SRVCache, CurrentShaderDirtySRVSlots[##Shader], NumSRVs[##Shader], ViewHeapSlot); \
 		}
 
-		if (PipelineType == ED3D12PipelineType::Graphics)
+		if (StartStage == SF_Compute)
 		{
+			// Note that ray tracing pipeline shares state with compute
+			CONDITIONAL_SET_SRVS(SF_Compute);
+		}
+		else
+		{
 			CONDITIONAL_SET_SRVS(SF_Vertex);
 			CONDITIONAL_SET_SRVS(SF_Mesh);
@@ -688,11 +757,6 @@
 			CONDITIONAL_SET_SRVS(SF_Geometry);
 			CONDITIONAL_SET_SRVS(SF_Pixel);
 		}
-		else
-		{
-			// Note that ray tracing pipeline shares state with compute
-			CONDITIONAL_SET_SRVS(SF_Compute);
-		}
 #undef CONDITIONAL_SET_SRVS
 	}
 
@@ -714,8 +778,12 @@
 			DescriptorCache.SetConstantBuffers<##Shader>(pRootSignature, CBVCache, CurrentShaderDirtyCBVSlots[##Shader]); \
 		}
 #endif
-
-		if (PipelineType == ED3D12PipelineType::Graphics)
+		if (StartStage == SF_Compute)
 		{
+			// Note that ray tracing pipeline shares state with compute
+			CONDITIONAL_SET_CBVS(SF_Compute);
+		}
+		else
+		{
 			CONDITIONAL_SET_CBVS(SF_Vertex);
 			CONDITIONAL_SET_CBVS(SF_Mesh);
@@ -723,21 +791,8 @@
 			CONDITIONAL_SET_CBVS(SF_Geometry);
 			CONDITIONAL_SET_CBVS(SF_Pixel);
 		}
-		else
-		{
-			// Note that ray tracing pipeline shares state with compute
-			CONDITIONAL_SET_CBVS(SF_Compute);
-		}
 #undef CONDITIONAL_SET_CBVS
 	}
-
-	// Flush any needed resource barriers
-	CommandList.FlushResourceBarriers();
-
-#if ASSERT_RESOURCE_STATES
-	bool bSucceeded = AssertResourceStates(PipelineType);
-	check(bSucceeded);
-#endif
 }
 
 void FD3D12StateCacheBase::ApplySamplers(const FD3D12RootSignature* const pRootSignature, uint32 StartStage, uint32 EndStage)
diff --git a/Engine/Source/Runtime/D3D12RHI/Private/D3D12StateCachePrivate.h b/Engine/Source/Runtime/D3D12RHI/Private/D3D12StateCachePrivate.h
--- a/Engine/Source/Runtime/D3D12RHI/Private/D3D12StateCachePrivate.h
+++ b/Engine/Source/Runtime/D3D12RHI/Private/D3D12StateCachePrivate.h
@@ -450,6 +450,7 @@ protected:
 #endif
 #undef DECLARE_SHADER_TRAITS
 
+	bool InternalSetRootSignature(ED3D12PipelineType InPipelineType, const FD3D12RootSignature* InRootSignature);
 	void InternalSetPipelineState(FD3D12PipelineState* InPipelineState);
 
 private:
@@ -773,6 +774,7 @@ public:
 	template <ED3D12PipelineType PipelineType> 
 	void ApplyState();
 	void ApplySamplers(const FD3D12RootSignature* const pRootSignature, uint32 StartStage, uint32 EndStage);
+	void ApplyResources(const FD3D12RootSignature* const pRootSignature, uint32 StartStage, uint32 EndStage);
 	void DirtyStateForNewCommandList();
 	void DirtyState();
 	void DirtyViewDescriptorTables();
diff --git a/Engine/Source/Runtime/D3D12RHI/Private/D3D12Util.cpp b/Engine/Source/Runtime/D3D12RHI/Private/D3D12Util.cpp
--- a/Engine/Source/Runtime/D3D12RHI/Private/D3D12Util.cpp
+++ b/Engine/Source/Runtime/D3D12RHI/Private/D3D12Util.cpp
@@ -1023,6 +1023,26 @@ static void QuantizeBoundShaderStateCommon(
 	SetBoundShaderStateFlags(OutQBSS, ShaderData);
 }
 
+static bool IsCompatibleWithBindlessSamplers(const FD3D12ShaderData* ShaderData)
+{
+	if (ensure(ShaderData))
+	{
+		return ShaderData->UsesBindlessSamplers()
+			|| ShaderData->ResourceCounts.NumSamplers == 0;
+	}
+	return true;
+}
+
+static bool IsCompatibleWithBindlessResources(const FD3D12ShaderData* ShaderData)
+{
+	if (ensure(ShaderData))
+	{
+		return ShaderData->UsesBindlessResources()
+			|| (ShaderData->ResourceCounts.NumSRVs + ShaderData->ResourceCounts.NumUAVs) == 0;
+	}
+	return true;
+}
+
 void QuantizeBoundShaderState(
 	const D3D12_RESOURCE_BINDING_TIER& ResourceBindingTier,
 	const FD3D12BoundShaderState* const BSS,
@@ -1043,6 +1063,42 @@ void QuantizeBoundShaderState(
 #endif
 	QuantizeBoundShaderStateCommon(QBSS, BSS->GetPixelShader(),         ResourceBindingTier, SV_Pixel, true /*bAllowUAVs*/);
 	QuantizeBoundShaderStateCommon(QBSS, BSS->GetGeometryShader(),      ResourceBindingTier, SV_Geometry);
+
+#if DO_CHECK
+	if (QBSS.bUseDirectlyIndexedResourceHeap || QBSS.bUseDirectlyIndexedSamplerHeap)
+	{
+		struct FGenericShaderPair
+		{
+			const FRHIGraphicsShader* RHI;
+			const FD3D12ShaderData* Data;
+		};
+		const FGenericShaderPair ShaderDatas[] =
+		{
+			{ BSS->GetVertexShader(), BSS->GetVertexShader() },
+#if PLATFORM_SUPPORTS_MESH_SHADERS
+			{ BSS->GetMeshShader(), BSS->GetMeshShader() },
+			{ BSS->GetAmplificationShader(), BSS->GetAmplificationShader() },
+#endif
+			{ BSS->GetPixelShader(), BSS->GetPixelShader() },
+			{ BSS->GetGeometryShader(), BSS->GetGeometryShader() },
+		};
+
+		for (const FGenericShaderPair& ShaderPair : ShaderDatas)
+		{
+			if (ShaderPair.RHI)
+			{
+				if (QBSS.bUseDirectlyIndexedResourceHeap)
+				{
+					checkf(IsCompatibleWithBindlessResources(ShaderPair.Data), TEXT("Mismatched dynamic resource usage. %s doesn't support binding with stages that use dynamic resources"), ShaderPair.RHI->GetShaderName());
+				}
+				if (QBSS.bUseDirectlyIndexedSamplerHeap)
+				{
+					checkf(IsCompatibleWithBindlessSamplers(ShaderPair.Data), TEXT("Mismatched dynamic resource usage. %s doesn't support binding with stages that use dynamic samplers"), ShaderPair.RHI->GetShaderName());
+				}
+			}
+		}
+	}
+#endif
 }
 
 void QuantizeBoundShaderState(
