Removing the dynamic samplers flag from Raytracing Root Signatures until we figure out how to set it dynamically.

diff --git a/Engine/Source/Runtime/D3D12RHI/Private/D3D12Util.cpp b/Engine/Source/Runtime/D3D12RHI/Private/D3D12Util.cpp
--- a/Engine/Source/Runtime/D3D12RHI/Private/D3D12Util.cpp
+++ b/Engine/Source/Runtime/D3D12RHI/Private/D3D12Util.cpp
@@ -1128,7 +1128,6 @@ FD3D12QuantizedBoundShaderState GetRayTracingGlobalRootSignatureDesc()
 
 	OutQBSS.RootSignatureType = RS_RayTracingGlobal;
 	OutQBSS.bUseDiagnosticBuffer = true;
-	OutQBSS.bUseDirectlyIndexedSamplerHeap = true;
 
 	QBSSRegisterCounts.SamplerCount = MAX_SAMPLERS;
 	QBSSRegisterCounts.ShaderResourceCount = MAX_SRVS;
