Make RHICreateUniformBuffer thread safe.

diff --git a/Engine/Source/Runtime/Apple/AGXRHI/Private/AGXUniformBuffer.cpp b/Engine/Source/Runtime/Apple/AGXRHI/Private/AGXUniformBuffer.cpp
--- a/Engine/Source/Runtime/Apple/AGXRHI/Private/AGXUniformBuffer.cpp
+++ b/Engine/Source/Runtime/Apple/AGXRHI/Private/AGXUniformBuffer.cpp
@@ -62,8 +62,6 @@ void FAGXSuballocatedUniformBuffer::Update(const void* Contents, TArray<TRefCoun
 // The amount of data read from Contents is given by the Layout
 void FAGXSuballocatedUniformBuffer::PushToGPUBacking(const void* Contents)
 {
-    check(IsInRenderingThread() ^ IsRunningRHIInSeparateThread());
-    
     FAGXDeviceContext& DeviceContext = GetAGXDeviceContext();
     
     FAGXFrameAllocator* Allocator = DeviceContext.GetUniformAllocator();
@@ -88,7 +86,7 @@ void FAGXSuballocatedUniformBuffer::PrepareToBind()
     }
 }
 
-void FAGXSuballocatedUniformBuffer::CopyResourceTable_RenderThread(const void* Contents, TArray<TRefCountPtr<FRHIResource> >& OutResourceTable)
+void FAGXSuballocatedUniformBuffer::CopyResourceTable(const void* Contents, TArray<TRefCountPtr<FRHIResource> >& OutResourceTable) const
 {
 #if METAL_UNIFORM_BUFFER_VALIDATION
 	if (Validation == EUniformBufferValidation::ValidateResources)
diff --git a/Engine/Source/Runtime/Apple/AGXRHI/Private/AGXUniformBufferCommands.cpp b/Engine/Source/Runtime/Apple/AGXRHI/Private/AGXUniformBufferCommands.cpp
--- a/Engine/Source/Runtime/Apple/AGXRHI/Private/AGXUniformBufferCommands.cpp
+++ b/Engine/Source/Runtime/Apple/AGXRHI/Private/AGXUniformBufferCommands.cpp
@@ -14,7 +14,7 @@ static void DoUpdateUniformBuffer(FAGXUniformBuffer* UB, const void* Contents)
 
 	TArray<TRefCountPtr<FRHIResource> > ResourceTable;
 
-	UB->CopyResourceTable_RenderThread(Contents, ResourceTable);
+	UB->CopyResourceTable(Contents, ResourceTable);
 
     if(bUpdateImmediately)
     {
@@ -37,14 +37,12 @@ static void DoUpdateUniformBuffer(FAGXUniformBuffer* UB, const void* Contents)
 
 FUniformBufferRHIRef FAGXDynamicRHI::RHICreateUniformBuffer(const void* Contents, const FRHIUniformBufferLayout* Layout, EUniformBufferUsage Usage, EUniformBufferValidation Validation)
 {
-    FAGXDeviceContext& DeviceContext = (FAGXDeviceContext&)GetAGXDeviceContext();
-    FAGXFrameAllocator* UniformAllocator = DeviceContext.GetUniformAllocator();
-    
-    FAGXUniformBuffer* UB = new FAGXUniformBuffer(Layout, Usage, Validation);
-    
-    DoUpdateUniformBuffer(UB, Contents);
-    
-    return UB;
+	FAGXUniformBuffer* UB = new FAGXUniformBuffer(Layout, Usage, Validation);
+
+	TArray<TRefCountPtr<FRHIResource>> ResourceTable;
+	UB->CopyResourceTable(Contents, ResourceTable);
+	UB->Update(Contents, ResourceTable);
+	return UB;
 }
 
 void FAGXDynamicRHI::RHIUpdateUniformBuffer(FRHIUniformBuffer* UniformBufferRHI, const void* Contents)
diff --git a/Engine/Source/Runtime/Apple/AGXRHI/Public/AGXUniformBuffer.h b/Engine/Source/Runtime/Apple/AGXRHI/Public/AGXUniformBuffer.h
--- a/Engine/Source/Runtime/Apple/AGXRHI/Public/AGXUniformBuffer.h
+++ b/Engine/Source/Runtime/Apple/AGXRHI/Public/AGXUniformBuffer.h
@@ -50,7 +50,7 @@ public:
     void PrepareToBind();
 
 	// Copies the RDG resources to a resource table for a deferred update on the RHI thread.
-	void CopyResourceTable_RenderThread(const void* Contents, TArray<TRefCountPtr<FRHIResource> >& OutResourceTable);
+	void CopyResourceTable(const void* Contents, TArray<TRefCountPtr<FRHIResource> >& OutResourceTable) const;
 
 private:
     
diff --git a/Engine/Source/Runtime/Apple/MetalRHI/Private/MetalUniformBuffer.cpp b/Engine/Source/Runtime/Apple/MetalRHI/Private/MetalUniformBuffer.cpp
--- a/Engine/Source/Runtime/Apple/MetalRHI/Private/MetalUniformBuffer.cpp
+++ b/Engine/Source/Runtime/Apple/MetalRHI/Private/MetalUniformBuffer.cpp
@@ -62,8 +62,6 @@ void FMetalSuballocatedUniformBuffer::Update(const void* Contents, TArray<TRefCo
 // The amount of data read from Contents is given by the Layout
 void FMetalSuballocatedUniformBuffer::PushToGPUBacking(const void* Contents)
 {
-    check(IsInRenderingThread() ^ IsRunningRHIInSeparateThread());
-    
     FMetalDeviceContext& DeviceContext = GetMetalDeviceContext();
     
     FMetalFrameAllocator* Allocator = DeviceContext.GetUniformAllocator();
@@ -88,7 +86,7 @@ void FMetalSuballocatedUniformBuffer::PrepareToBind()
     }
 }
 
-void FMetalSuballocatedUniformBuffer::CopyResourceTable_RenderThread(const void* Contents, TArray<TRefCountPtr<FRHIResource> >& OutResourceTable)
+void FMetalSuballocatedUniformBuffer::CopyResourceTable(const void* Contents, TArray<TRefCountPtr<FRHIResource> >& OutResourceTable) const
 {
 #if METAL_UNIFORM_BUFFER_VALIDATION
 	if (Validation == EUniformBufferValidation::ValidateResources)
diff --git a/Engine/Source/Runtime/Apple/MetalRHI/Private/MetalUniformBufferCommands.cpp b/Engine/Source/Runtime/Apple/MetalRHI/Private/MetalUniformBufferCommands.cpp
--- a/Engine/Source/Runtime/Apple/MetalRHI/Private/MetalUniformBufferCommands.cpp
+++ b/Engine/Source/Runtime/Apple/MetalRHI/Private/MetalUniformBufferCommands.cpp
@@ -14,7 +14,7 @@ static void DoUpdateUniformBuffer(FMetalUniformBuffer* UB, const void* Contents)
 
 	TArray<TRefCountPtr<FRHIResource> > ResourceTable;
 
-	UB->CopyResourceTable_RenderThread(Contents, ResourceTable);
+	UB->CopyResourceTable(Contents, ResourceTable);
 
     if(bUpdateImmediately)
     {
@@ -37,14 +37,12 @@ static void DoUpdateUniformBuffer(FMetalUniformBuffer* UB, const void* Contents)
 
 FUniformBufferRHIRef FMetalDynamicRHI::RHICreateUniformBuffer(const void* Contents, const FRHIUniformBufferLayout* Layout, EUniformBufferUsage Usage, EUniformBufferValidation Validation)
 {
-    FMetalDeviceContext& DeviceContext = (FMetalDeviceContext&)GetMetalDeviceContext();
-    FMetalFrameAllocator* UniformAllocator = DeviceContext.GetUniformAllocator();
-    
-    FMetalUniformBuffer* UB = new FMetalUniformBuffer(Layout, Usage, Validation);
-    
-    DoUpdateUniformBuffer(UB, Contents);
-    
-    return UB;
+	FMetalUniformBuffer* UB = new FMetalUniformBuffer(Layout, Usage, Validation);
+
+	TArray<TRefCountPtr<FRHIResource>> ResourceTable;
+	UB->CopyResourceTable(Contents, ResourceTable);
+	UB->Update(Contents, ResourceTable);
+	return UB;
 }
 
 void FMetalDynamicRHI::RHIUpdateUniformBuffer(FRHIUniformBuffer* UniformBufferRHI, const void* Contents)
diff --git a/Engine/Source/Runtime/Apple/MetalRHI/Public/MetalUniformBuffer.h b/Engine/Source/Runtime/Apple/MetalRHI/Public/MetalUniformBuffer.h
--- a/Engine/Source/Runtime/Apple/MetalRHI/Public/MetalUniformBuffer.h
+++ b/Engine/Source/Runtime/Apple/MetalRHI/Public/MetalUniformBuffer.h
@@ -50,7 +50,7 @@ public:
     void PrepareToBind();
 
 	// Copies the RDG resources to a resource table for a deferred update on the RHI thread.
-	void CopyResourceTable_RenderThread(const void* Contents, TArray<TRefCountPtr<FRHIResource> >& OutResourceTable);
+	void CopyResourceTable(const void* Contents, TArray<TRefCountPtr<FRHIResource> >& OutResourceTable) const;
 
 private:
     
diff --git a/Engine/Source/Runtime/D3D12RHI/Private/D3D12Allocation.cpp b/Engine/Source/Runtime/D3D12RHI/Private/D3D12Allocation.cpp
--- a/Engine/Source/Runtime/D3D12RHI/Private/D3D12Allocation.cpp
+++ b/Engine/Source/Runtime/D3D12RHI/Private/D3D12Allocation.cpp
@@ -1077,6 +1077,8 @@ void* FD3D12UploadHeapAllocator::AllocUploadResource(uint32 InSize, uint32 InAli
 	}
 	else
 	{
+		FD3D12ScopeLock Lock(&BigBlockCS);
+
 		// Forward to the big block allocator
 		const D3D12_RESOURCE_DESC ResourceDesc = CD3DX12_RESOURCE_DESC::Buffer(InSize, D3D12_RESOURCE_FLAG_NONE);
 		BigBlockAllocator.AllocateResource(GetParentDevice()->GetGPUIndex(), D3D12_HEAP_TYPE_UPLOAD, ResourceDesc, InSize, InAlignment, ED3D12ResourceStateMode::SingleState, D3D12_RESOURCE_STATE_GENERIC_READ, nullptr, nullptr, ResourceLocation);
@@ -1101,7 +1103,10 @@ void* FD3D12UploadHeapAllocator::AllocFastConstantAllocationPage(uint32 InSize,
 void FD3D12UploadHeapAllocator::CleanUpAllocations(uint64 InFrameLag)
 {
 	SmallBlockAllocator.CleanUpAllocations(InFrameLag);
-	BigBlockAllocator.CleanUpAllocations(InFrameLag);
+	{
+		FD3D12ScopeLock Lock(&BigBlockCS);
+		BigBlockAllocator.CleanUpAllocations(InFrameLag);
+	}
 	FastConstantPageAllocator.CleanUpAllocations(InFrameLag);
 }
 
@@ -1118,7 +1123,10 @@ void FD3D12UploadHeapAllocator::UpdateMemoryStats()
 
 #if defined(D3D12RHI_TRACK_DETAILED_STATS)
 	SmallBlockAllocator.UpdateMemoryStats(MemoryAllocated, MemoryUsed, FreeMemory, AlignmentWaste, AllocatedPageCount, FullPageCount);
-	BigBlockAllocator.UpdateMemoryStats(MemoryAllocated, MemoryUsed, FreeMemory, EndFreeMemory, AlignmentWaste, AllocatedPageCount, FullPageCount);
+	{
+		FD3D12ScopeLock Lock(&BigBlockCS);
+		BigBlockAllocator.UpdateMemoryStats(MemoryAllocated, MemoryUsed, FreeMemory, EndFreeMemory, AlignmentWaste, AllocatedPageCount, FullPageCount);
+	}
 	FastConstantPageAllocator.UpdateMemoryStats(MemoryAllocated, MemoryUsed, FreeMemory, AlignmentWaste, AllocatedPageCount, FullPageCount);
 #endif
 
diff --git a/Engine/Source/Runtime/D3D12RHI/Private/D3D12Allocation.h b/Engine/Source/Runtime/D3D12RHI/Private/D3D12Allocation.h
--- a/Engine/Source/Runtime/D3D12RHI/Private/D3D12Allocation.h
+++ b/Engine/Source/Runtime/D3D12RHI/Private/D3D12Allocation.h
@@ -378,12 +378,13 @@ public:
 private:
 
 	// Buddy allocator used for all 'small' allocation - fast but aligns to power of 2
-	FD3D12MultiBuddyAllocator SmallBlockAllocator;
+	FD3D12MultiBuddyAllocator	SmallBlockAllocator;
 	// Pool allocator for all bigger allocations - less fast but less alignment waste
-	FD3D12PoolAllocator		  BigBlockAllocator;
+	FCriticalSection			BigBlockCS;
+	FD3D12PoolAllocator			BigBlockAllocator;
 	// Seperate buddy allocator used for the fast constant allocator pages which get always freed within the same frame by default
 	// (different allocator to avoid fragmentation with the other pools - always the same size allocations)
-	FD3D12MultiBuddyAllocator FastConstantPageAllocator;
+	FD3D12MultiBuddyAllocator	FastConstantPageAllocator;
 };
 
 //-----------------------------------------------------------------------------
diff --git a/Engine/Source/Runtime/RHI/Public/RHIValidation.h b/Engine/Source/Runtime/RHI/Public/RHIValidation.h
--- a/Engine/Source/Runtime/RHI/Public/RHIValidation.h
+++ b/Engine/Source/Runtime/RHI/Public/RHIValidation.h
@@ -365,7 +365,8 @@ public:
 		// RHICreateUniformBuffer can be called from the RenderThread (normal use case) or the RHI Thread (FLocalUniformBuffer)
 		check(Layout->Resources.Num() > 0 || Layout->ConstantBufferSize > 0);
 		FUniformBufferRHIRef UniformBuffer = RHI->RHICreateUniformBuffer(Contents, Layout, Usage, Validation);
-		UniformBuffer->InitLifetimeTracking(IsInRenderingThread() ? RenderThreadFrameID : RHIThreadFrameID, Usage);
+		// Use the render thread frame ID for any non RHI thread allocations.
+		UniformBuffer->InitLifetimeTracking(IsInRHIThread() ? RHIThreadFrameID : RenderThreadFrameID, Usage);
 		return UniformBuffer;
 	}
 
diff --git a/Engine/Source/Runtime/VulkanRHI/Private/VulkanUniformBuffer.cpp b/Engine/Source/Runtime/VulkanRHI/Private/VulkanUniformBuffer.cpp
--- a/Engine/Source/Runtime/VulkanRHI/Private/VulkanUniformBuffer.cpp
+++ b/Engine/Source/Runtime/VulkanRHI/Private/VulkanUniformBuffer.cpp
@@ -121,14 +121,17 @@ FVulkanUniformBuffer::FVulkanUniformBuffer(FVulkanDevice& Device, const FRHIUnif
 
 	if (InLayout->ConstantBufferSize > 0)
 	{
-		if (UseRingBuffer(InUsage))
+		const bool bInRenderingThread = IsInRenderingThread();
+		const bool bInRHIThread = IsInRHIThread();
+
+		if (UseRingBuffer(InUsage) && (bInRenderingThread || bInRHIThread))
 		{
 			FRHICommandListImmediate& RHICmdList = FRHICommandListExecutor::GetImmediateCommandList();
 			FVulkanUniformBuffer* UniformBuffer = this;
 			int32 DataSize = InLayout->ConstantBufferSize;
 			
 			// make sure we allocate from RingBuffer on RHIT
-			const bool bCanAllocOnThisThread = RHICmdList.Bypass() || (!IsRunningRHIInSeparateThread() && IsInRenderingThread()) || IsInRHIThread();
+			const bool bCanAllocOnThisThread = RHICmdList.Bypass() || (!IsRunningRHIInSeparateThread() && bInRenderingThread) || bInRHIThread;
 			if (bCanAllocOnThisThread)
 			{
 				FVulkanCommandListContextImmediate& Context = Device.GetImmediateContext();
diff --git a/Engine/Source/Runtime/Windows/D3D11RHI/Private/D3D11UniformBuffer.cpp b/Engine/Source/Runtime/Windows/D3D11RHI/Private/D3D11UniformBuffer.cpp
--- a/Engine/Source/Runtime/Windows/D3D11RHI/Private/D3D11UniformBuffer.cpp
+++ b/Engine/Source/Runtime/Windows/D3D11RHI/Private/D3D11UniformBuffer.cpp
@@ -166,8 +166,6 @@ static TRefCountPtr<ID3D11Buffer> CreateAndUpdatePooledUniformBuffer(
 
 FUniformBufferRHIRef FD3D11DynamicRHI::RHICreateUniformBuffer(const void* Contents, const FRHIUniformBufferLayout* Layout, EUniformBufferUsage Usage, EUniformBufferValidation Validation)
 {
-	check(IsInRenderingThread() || IsInRHIThread());
-
 	if (Validation == EUniformBufferValidation::ValidateResources)
 	{
 		ValidateShaderParameterResourcesRHI(Contents, *Layout);
@@ -185,7 +183,7 @@ FUniformBufferRHIRef FD3D11DynamicRHI::RHICreateUniformBuffer(const void* Conten
 
 		SCOPE_CYCLE_COUNTER(STAT_D3D11UpdateUniformBufferTime);
 
-		if (IsPoolingEnabled())
+		if (IsPoolingEnabled() && (IsInRenderingThread() || IsInRHIThread()))
 		{
 			if (ShouldNotEnqueueRHICommand())
 			{
