Fixed missing barrier in ResolveTexture on D3D12.

diff --git a/Engine/Source/Runtime/D3D12RHI/Private/D3D12RenderTarget.cpp b/Engine/Source/Runtime/D3D12RHI/Private/D3D12RenderTarget.cpp
--- a/Engine/Source/Runtime/D3D12RHI/Private/D3D12RenderTarget.cpp
+++ b/Engine/Source/Runtime/D3D12RHI/Private/D3D12RenderTarget.cpp
@@ -529,6 +529,10 @@ void FD3D12CommandContext::ResolveTexture(UE::RHICore::FResolveTextureInfo Info)
 			int32 DestSubresource   = CalcSubresource(Info.MipLevel, ArraySlice, DestDesc.ArraySize);
 			int32 SourceSubresource = CalcSubresource(Info.MipLevel, ArraySlice, SourceDesc.ArraySize);
 
+			FScopedResourceBarrier ConditionalScopeResourceBarrierDest(CommandListHandle, DestTexture->GetResource(), D3D12_RESOURCE_STATE_RESOLVE_DEST, DestSubresource, FD3D12DynamicRHI::ETransitionMode::Validate);
+			FScopedResourceBarrier ConditionalScopeResourceBarrierSource(CommandListHandle, SourceTexture->GetResource(), D3D12_RESOURCE_STATE_RESOLVE_SOURCE, SourceSubresource, FD3D12DynamicRHI::ETransitionMode::Validate);
+
+			CommandListHandle.FlushResourceBarriers();
 			CommandListHandle->ResolveSubresource(DestResource->GetResource(), DestSubresource, SourceTexture->GetResource()->GetResource(), SourceSubresource, DestFormatTypeless);
 			otherWorkCounter++;
 		}
