Adding RHI cvars to control if bindless resources or samplers are enabled. If an RHI is configured with Bindless support, these configs will turn on the bindless defines and should control bindless heap creation.
The previous control was just setting the D3D12 bindless heap sizes; now those sizes can be set and will only be used when the RHI controls are enabled.

diff --git a/Engine/Source/Runtime/D3D12RHI/Private/D3D12Device.cpp b/Engine/Source/Runtime/D3D12RHI/Private/D3D12Device.cpp
--- a/Engine/Source/Runtime/D3D12RHI/Private/D3D12Device.cpp
+++ b/Engine/Source/Runtime/D3D12RHI/Private/D3D12Device.cpp
@@ -344,7 +344,13 @@ void FD3D12Device::SetupAfterDeviceCreation()
 	check(GBindlessSamplerDescriptorHeapSize <= GGlobalSamplerDescriptorHeapSize);
 
 	DescriptorHeapManager.Init(GGlobalResourceDescriptorHeapSize, GGlobalSamplerDescriptorHeapSize);
-	BindlessDescriptorManager.Init(GBindlessResourceDescriptorHeapSize, GBindlessSamplerDescriptorHeapSize);
+
+	const bool bBindlessResources = RHIGetBindlessResourcesConfiguration(GMaxRHIShaderPlatform) != ERHIBindlessConfiguration::Disabled;
+	const bool bBindlessSamplers = RHIGetBindlessSamplersConfiguration(GMaxRHIShaderPlatform) != ERHIBindlessConfiguration::Disabled;
+	if (bBindlessResources || bBindlessSamplers)
+	{
+		BindlessDescriptorManager.Init(bBindlessResources ? GBindlessResourceDescriptorHeapSize : 0, bBindlessSamplers ? GBindlessSamplerDescriptorHeapSize : 0);
+	}
 
 	// Init offline descriptor managers
 	for (uint32 Index = 0; Index < static_cast<uint32>(ERHIDescriptorHeapType::count); Index++)
@@ -352,7 +358,7 @@ void FD3D12Device::SetupAfterDeviceCreation()
 		OfflineDescriptorManagers[Index].Init(static_cast<ERHIDescriptorHeapType>(Index));
 	}
 
-	GlobalSamplerHeap.Init(GGlobalSamplerHeapSize - GBindlessSamplerDescriptorHeapSize);
+	GlobalSamplerHeap.Init(GGlobalSamplerHeapSize);
 
 	OnlineDescriptorManager.Init(GOnlineDescriptorHeapSize, GOnlineDescriptorHeapBlockSize);
 
diff --git a/Engine/Source/Runtime/D3D12RHI/Private/D3D12StateCache.cpp b/Engine/Source/Runtime/D3D12RHI/Private/D3D12StateCache.cpp
--- a/Engine/Source/Runtime/D3D12RHI/Private/D3D12StateCache.cpp
+++ b/Engine/Source/Runtime/D3D12RHI/Private/D3D12StateCache.cpp
@@ -29,7 +29,7 @@ static FAutoConsoleVariableRef CVarGlobalSamplerDescriptorHeapSize(
 	ECVF_ReadOnly
 );
 
-int32 GBindlessResourceDescriptorHeapSize = 0;
+int32 GBindlessResourceDescriptorHeapSize = 1000 * 1000;
 static FAutoConsoleVariableRef CVarBindlessResourceDescriptorHeapSize(
 	TEXT("D3D12.Bindless.ResourceDescriptorHeapSize"),
 	GBindlessResourceDescriptorHeapSize,
@@ -37,7 +37,7 @@ static FAutoConsoleVariableRef CVarBindlessResourceDescriptorHeapSize(
 	ECVF_ReadOnly
 );
 
-int32 GBindlessSamplerDescriptorHeapSize = 0;
+int32 GBindlessSamplerDescriptorHeapSize = 2048;
 static FAutoConsoleVariableRef CVarBindlessSamplerDescriptorHeapSize(
 	TEXT("D3D12.Bindless.SamplerDescriptorHeapSize"),
 	GBindlessSamplerDescriptorHeapSize,
diff --git a/Engine/Source/Runtime/D3D12RHI/Private/Windows/WindowsD3D12Device.cpp b/Engine/Source/Runtime/D3D12RHI/Private/Windows/WindowsD3D12Device.cpp
--- a/Engine/Source/Runtime/D3D12RHI/Private/Windows/WindowsD3D12Device.cpp
+++ b/Engine/Source/Runtime/D3D12RHI/Private/Windows/WindowsD3D12Device.cpp
@@ -292,7 +292,13 @@ static D3D_SHADER_MODEL FindHighestShaderModel(ID3D12Device* Device)
 	return D3D_SHADER_MODEL_5_1;
 }
 
-inline ERHIFeatureLevel::Type FindMaxRHIFeatureLevel(ID3D12Device* Device, D3D_FEATURE_LEVEL InMaxFeatureLevel, D3D_SHADER_MODEL InMaxShaderModel)
+inline bool ShouldCheckBindlessSupport(EShaderPlatform ShaderPlatform)
+{
+	return RHIGetBindlessResourcesConfiguration(ShaderPlatform) != ERHIBindlessConfiguration::Disabled
+		|| RHIGetBindlessSamplersConfiguration(ShaderPlatform) != ERHIBindlessConfiguration::Disabled;
+}
+
+inline ERHIFeatureLevel::Type FindMaxRHIFeatureLevel(ID3D12Device* Device, D3D_FEATURE_LEVEL InMaxFeatureLevel, D3D_SHADER_MODEL InMaxShaderModel, D3D12_RESOURCE_BINDING_TIER ResourceBindingTier)
 {
 	ERHIFeatureLevel::Type MaxRHIFeatureLevel = ERHIFeatureLevel::Num;
 
@@ -304,7 +310,13 @@ inline ERHIFeatureLevel::Type FindMaxRHIFeatureLevel(ID3D12Device* Device, D3D_F
 		D3D12_FEATURE_DATA_D3D12_OPTIONS9 D3D12Caps9{};
 		Device->CheckFeatureSupport(D3D12_FEATURE_D3D12_OPTIONS9, &D3D12Caps9, sizeof(D3D12Caps9));
 
-		if (D3D12Caps1.WaveOps && D3D12Caps9.AtomicInt64OnTypedResourceSupported)
+		bool bHighEnoughBindingTier = true;
+		if (ShouldCheckBindlessSupport(SP_PCD3D_SM6))
+		{
+			bHighEnoughBindingTier = ResourceBindingTier >= D3D12_RESOURCE_BINDING_TIER_3;
+		}
+
+		if (D3D12Caps1.WaveOps && D3D12Caps9.AtomicInt64OnTypedResourceSupported && bHighEnoughBindingTier)
 		{
 			MaxRHIFeatureLevel = ERHIFeatureLevel::SM6;
 		}
@@ -342,7 +354,7 @@ static bool SafeTestD3D12CreateDevice(IDXGIAdapter* Adapter, D3D_FEATURE_LEVEL M
 			OutInfo.MaxShaderModel = FindHighestShaderModel(Device);
 			GetResourceTiers(Device, OutInfo.ResourceBindingTier, OutInfo.ResourceHeapTier);
 			OutInfo.NumDeviceNodes = Device->GetNodeCount();
-			OutInfo.MaxRHIFeatureLevel = FindMaxRHIFeatureLevel(Device, OutInfo.MaxFeatureLevel, OutInfo.MaxShaderModel);
+			OutInfo.MaxRHIFeatureLevel = FindMaxRHIFeatureLevel(Device, OutInfo.MaxFeatureLevel, OutInfo.MaxShaderModel, OutInfo.ResourceBindingTier);
 
 			Device->Release();
 			return true;
diff --git a/Engine/Source/Runtime/Engine/Private/ShaderCompiler/ShaderCompiler.cpp b/Engine/Source/Runtime/Engine/Private/ShaderCompiler/ShaderCompiler.cpp
--- a/Engine/Source/Runtime/Engine/Private/ShaderCompiler/ShaderCompiler.cpp
+++ b/Engine/Source/Runtime/Engine/Private/ShaderCompiler/ShaderCompiler.cpp
@@ -5996,6 +5996,20 @@ void GlobalBeginCompileShader(
 	Input.Environment.SetDefine(TEXT("PLATFORM_SUPPORTS_MESH_SHADERS_TIER1"), RHISupportsMeshShadersTier1(EShaderPlatform(Target.Platform)) ? 1 : 0);
 	Input.Environment.SetDefine(TEXT("PLATFORM_SUPPORTS_BINDLESS"), RHISupportsBindless(EShaderPlatform(Target.Platform)) ? 1 : 0);
 
+	if (RHISupportsBindless(EShaderPlatform(Target.Platform)))
+	{
+		const bool bIsRaytracingShader = IsRayTracingShaderFrequency(Input.Target.GetFrequency());
+
+		const ERHIBindlessConfiguration ResourcesConfig = RHIGetBindlessResourcesConfiguration(EShaderPlatform(Target.Platform));
+		const ERHIBindlessConfiguration SamplersConfig = RHIGetBindlessSamplersConfiguration(EShaderPlatform(Target.Platform));
+
+		Input.Environment.SetDefine(TEXT("ENABLE_BINDLESS_RESOURCES"),
+			ResourcesConfig == ERHIBindlessConfiguration::AllShaders || (ResourcesConfig == ERHIBindlessConfiguration::RayTracingShaders && bIsRaytracingShader));
+
+		Input.Environment.SetDefine(TEXT("ENABLE_BINDLESS_SAMPLERS"),
+			SamplersConfig == ERHIBindlessConfiguration::AllShaders || (SamplersConfig == ERHIBindlessConfiguration::RayTracingShaders && bIsRaytracingShader));
+	}
+
 	{
 		static const auto CVar = IConsoleManager::Get().FindConsoleVariable(TEXT("r.VT.AnisotropicFiltering"));
 		Input.Environment.SetDefine(TEXT("VIRTUAL_TEXTURE_ANISOTROPIC_FILTERING"), CVar ? (CVar->GetInt() != 0) : 0);
diff --git a/Engine/Source/Runtime/RHI/Private/RHI.cpp b/Engine/Source/Runtime/RHI/Private/RHI.cpp
--- a/Engine/Source/Runtime/RHI/Private/RHI.cpp
+++ b/Engine/Source/Runtime/RHI/Private/RHI.cpp
@@ -1300,6 +1300,43 @@ static FAutoConsoleVariableRef CVarEnableAttachmentVariableRateShading(
 	TEXT("Toggle to enable image-based Variable Rate Shading."),
 	ECVF_RenderThreadSafe);
 
+
+int32 GRHIBindlessResourceConfiguration = 0;
+static FAutoConsoleVariableRef CVarEnableBindlessResources(
+	TEXT("rhi.Bindless.Resources"),
+	GRHIBindlessResourceConfiguration,
+	TEXT("Set to 1 to enable for all shader types. Set to 2 to restrict to Raytracing shaders."),
+	ECVF_ReadOnly
+);
+
+int32 GRHIBindlessSamplerConfiguration = 0;
+static FAutoConsoleVariableRef CVarEnableBindlessSamplers(
+	TEXT("rhi.Bindless.Samplers"),
+	GRHIBindlessSamplerConfiguration,
+	TEXT("Set to 1 to enable for all shader types. Set to 2 to restrict to Raytracing shaders."),
+	ECVF_ReadOnly
+);
+
+static ERHIBindlessConfiguration DetermineBindlessConfiguration(EShaderPlatform Platform, int32 BindlessConfigSetting)
+{
+	if (!RHISupportsBindless(Platform) || BindlessConfigSetting == 0)
+	{
+		return ERHIBindlessConfiguration::Disabled;
+	}
+
+	return BindlessConfigSetting == 2 ? ERHIBindlessConfiguration::RayTracingShaders : ERHIBindlessConfiguration::AllShaders;
+}
+
+ERHIBindlessConfiguration RHIGetBindlessResourcesConfiguration(EShaderPlatform Platform)
+{
+	return DetermineBindlessConfiguration(Platform, GRHIBindlessResourceConfiguration);
+}
+
+ERHIBindlessConfiguration RHIGetBindlessSamplersConfiguration(EShaderPlatform Platform)
+{
+	return DetermineBindlessConfiguration(Platform, GRHIBindlessSamplerConfiguration);
+}
+
 namespace RHIConfig
 {
 	bool ShouldSaveScreenshotAfterProfilingGPU()
diff --git a/Engine/Source/Runtime/RHI/Public/RHIDefinitions.h b/Engine/Source/Runtime/RHI/Public/RHIDefinitions.h
--- a/Engine/Source/Runtime/RHI/Public/RHIDefinitions.h
+++ b/Engine/Source/Runtime/RHI/Public/RHIDefinitions.h
@@ -2602,6 +2602,15 @@ inline ERHIResourceType GetRHIResourceType(ETextureDimension Dimension)
 	return ERHIResourceType::RRT_None;
 }
 
+enum class ERHIBindlessConfiguration
+{
+	Disabled,
+	AllShaders,
+	RayTracingShaders,
+};
+
+RHI_API ERHIBindlessConfiguration RHIGetBindlessResourcesConfiguration(EShaderPlatform Platform);
+RHI_API ERHIBindlessConfiguration RHIGetBindlessSamplersConfiguration(EShaderPlatform Platform);
 
 #if PLATFORM_SUPPORTS_GEOMETRY_SHADERS
 	#define GEOMETRY_SHADER(GeometryShader)	(GeometryShader)
diff --git a/Engine/Source/Runtime/RenderCore/Private/Shader.cpp b/Engine/Source/Runtime/RenderCore/Private/Shader.cpp
--- a/Engine/Source/Runtime/RenderCore/Private/Shader.cpp
+++ b/Engine/Source/Runtime/RenderCore/Private/Shader.cpp
@@ -1925,6 +1925,19 @@ void ShaderMapAppendKeyString(EShaderPlatform Platform, FString& KeyString)
 	if (RHISupportsBindless(Platform))
 	{
 		KeyString += TEXT("_BNDLS");
+
+		const ERHIBindlessConfiguration ResourcesConfig = RHIGetBindlessResourcesConfiguration(Platform);
+		const ERHIBindlessConfiguration SamplersConfig = RHIGetBindlessSamplersConfiguration(Platform);
+
+		if (ResourcesConfig != ERHIBindlessConfiguration::Disabled)
+		{
+			KeyString += ResourcesConfig == ERHIBindlessConfiguration::RayTracingShaders ? TEXT("_RTRES") : TEXT("ALLRES");
+		}
+
+		if (SamplersConfig != ERHIBindlessConfiguration::Disabled)
+		{
+			KeyString += SamplersConfig == ERHIBindlessConfiguration::RayTracingShaders ? TEXT("_RTSAM") : TEXT("ALLSAM");
+		}
 	}
 
 	if (ShouldCompileRayTracingShadersForProject(Platform))
