D3D12 RHI fence logic fixes:
- Stop the RHI thread holiding a critical section in FD3D12CommandListData::WaitForCompletion while we're waiting on the GPU. This has a very bad effect on RHI thread performance in GPU-bound situations.

- Simplify FD3D12CommandListData::IsComplete to fix a longstanding bug in D3D12 Fence Logic. This avoids the need for holding a lock while we wait for fences to complete.

Specifically, FD3D12CommandListData::IsComplete had over-complex logic that relied on reading the contents of the ActiveGenerations queue to determine if a generation had completed. This isn't correct, since items can be dequeued before the corresponding GPU fence is hit. Now we rely only on the LastCompleteGeneration counter (similar to the existing logic in WaitForCompletion).

diff --git a/Engine/Source/Runtime/D3D12RHI/Private/D3D12CommandList.h b/Engine/Source/Runtime/D3D12RHI/Private/D3D12CommandList.h
--- a/Engine/Source/Runtime/D3D12RHI/Private/D3D12CommandList.h
+++ b/Engine/Source/Runtime/D3D12RHI/Private/D3D12CommandList.h
@@ -76,35 +76,20 @@ private:
 			}
 
 			check(Generation < CurrentGeneration);
-			if (Generation > LastCompleteGeneration)
-			{
-				FScopeLock Lock(&ActiveGenerationsCS);
-				GenerationSyncPointPair GenerationSyncPoint;
-				if (ActiveGenerations.Peek(GenerationSyncPoint))
-				{
-					if (Generation < GenerationSyncPoint.Key)
-					{
-						// The requested generation is older than the oldest tracked generation, so it must be complete.
-						return true;
-					}
-					else
-					{
-						if (GenerationSyncPoint.Value.IsComplete())
-						{
-							// Oldest tracked generation is done so clean the queue and try again.
-							CleanupActiveGenerations();
-							return IsComplete(Generation);
-						}
-						else
-						{
-							// The requested generation is newer than the older track generation but the old one isn't done.
-							return false;
-						}
-					}
-				}
-			}
-
-			return true;
+		    if (Generation <= LastCompleteGeneration)
+		    {
+			    return true;
+		    }
+		    FScopeLock Lock(&ActiveGenerationsCS);
+		    if (Generation <= LastCompleteGeneration)
+		    {
+			    // This generation completed while we waited for the lock
+			    return true;
+		    }
+    
+		    // Cleanup and update LastCompleteGeneration
+		    CleanupActiveGenerations();
+			return Generation <= LastCompleteGeneration;
 		}
 
 		void WaitForCompletion(uint64 Generation)
@@ -114,7 +99,10 @@ private:
 				CleanupActiveGenerations();
 				if (Generation > LastCompleteGeneration)
 				{
-					FScopeLock Lock(&ActiveGenerationsCS);
+					// We need an additional WaitForCompletionCS lock here to prevent concurrent access to this loop while we're waiting for the GPU 
+					FScopeLock Lock(&WaitForCompletionCS); 
+					FScopeLock Lock2(&ActiveGenerationsCS);
+
 					ensureMsgf(Generation < CurrentGeneration, TEXT("You can't wait for an unsubmitted command list to complete.  Kick first!"));
 					GenerationSyncPointPair GenerationSyncPoint;
 					while (ActiveGenerations.Peek(GenerationSyncPoint) && (Generation > LastCompleteGeneration))
@@ -122,10 +110,15 @@ private:
 						check(Generation >= GenerationSyncPoint.Key);
 						ActiveGenerations.Dequeue(GenerationSyncPoint);
 
-						GenerationSyncPoint.Value.WaitForCompletion();
+						// Allow other threads to run while we wait for the GPU
+					    ActiveGenerationsCS.Unlock();
+    
+					    GenerationSyncPoint.Value.WaitForCompletion();
+   
+					    ActiveGenerationsCS.Lock();
 
-						check(GenerationSyncPoint.Key > LastCompleteGeneration);
-						LastCompleteGeneration = GenerationSyncPoint.Key;
+						// It's possible that LastCompleteGeneration was modified in CleanupActiveGenerations by the time we get here, so we take the max
+					    LastCompleteGeneration = FMath::Max(LastCompleteGeneration, GenerationSyncPoint.Key);
 					}
 				}
 			}
@@ -221,6 +214,7 @@ private:
 		typedef TPair<uint64, FD3D12SyncPoint>	GenerationSyncPointPair;	// Pair of command list generation to a sync point
 		TQueue<GenerationSyncPointPair>			ActiveGenerations;	// Queue of active command list generations and their sync points. Used to determine what command lists have been completed on the GPU.
 		FCriticalSection						ActiveGenerationsCS;	// While only a single thread can record to a command list at any given time, multiple threads can ask for the state of a given command list. So the associated tracking must be thread-safe.
+		FCriticalSection						WaitForCompletionCS;	// Prevent concurrent calls to WaitForCompletion
 
 		// Array of resources who's state needs to be synced between submits.
 		TArray<FD3D12PendingResourceBarrier>	PendingResourceBarriers;
