- only break on d3d errors when debugger is attached
- disable the following d3d debug warning & errors:
	D3D12_MESSAGE_ID_COMMAND_LIST_DRAW_VERTEX_BUFFER_NOT_SET
	D3D12_MESSAGE_ID_COMMAND_LIST_DRAW_VERTEX_BUFFER_TOO_SMALL
	D3D12_MESSAGE_ID_CREATERESOURCE_STATE_IGNORED

diff --git a/Engine/Source/Runtime/D3D12RHI/Private/D3D12Adapter.cpp b/Engine/Source/Runtime/D3D12RHI/Private/D3D12Adapter.cpp
--- a/Engine/Source/Runtime/D3D12RHI/Private/D3D12Adapter.cpp
+++ b/Engine/Source/Runtime/D3D12RHI/Private/D3D12Adapter.cpp
@@ -187,8 +187,12 @@ static LONG __stdcall D3DVectoredExceptionHandler(EXCEPTION_POINTERS* InInfo)
 	{
 		if (CheckD3DStoredMessages())
 		{
-			// when we get here, then it means that BreakOnSeverity was set for this error message, so request the debug break here as well
-			UE_DEBUG_BREAK();
+			if (FPlatformMisc::IsDebuggerPresent())
+			{
+				// when we get here, then it means that BreakOnSeverity was set for this error message, so request the debug break here as well
+				// when the debugger is attached
+				UE_DEBUG_BREAK();
+			}
 		}
 
 		// Handles the exception
@@ -703,10 +707,9 @@ void FD3D12Adapter::CreateRootDevice(bool bWithDebug)
 				//		We intentionally keep the readback resources persistently mapped.
 				D3D12_MESSAGE_ID_EXECUTECOMMANDLISTS_GPU_WRITTEN_READBACK_RESOURCE_MAPPED,
 
-				// Note message ID doesn't exist in the current header (yet, should be available in the RS2 header) for now just mute by the ID number.
-				// RESOURCE_BARRIER_DUPLICATE_SUBRESOURCE_TRANSITIONS - This shows up a lot and is very noisy. It would require changes to the resource tracking system
+				// This shows up a lot and is very noisy. It would require changes to the resource tracking system
 				// but will hopefully be resolved when the RHI switches to use the engine's resource tracking system.
-				(D3D12_MESSAGE_ID)1008,
+				D3D12_MESSAGE_ID_RESOURCE_BARRIER_DUPLICATE_SUBRESOURCE_TRANSITIONS,
 
 				// This error gets generated on the first run when you install a new driver. The code handles this error properly and resets the PipelineLibrary,
 				// so we can safely ignore this message. It could possibly be avoided by adding driver version to the PSO cache filename, but an average user is unlikely
@@ -717,6 +720,16 @@ void FD3D12Adapter::CreateRootDevice(bool bWithDebug)
 				// it ignored the aliasing barriers to validate, and probably can't check them when called from IASetVertexBuffers because it only has GPU Virtual Addresses then
 				D3D12_MESSAGE_ID_HEAP_ADDRESS_RANGE_INTERSECTS_MULTIPLE_BUFFERS,
 
+				// Ignore draw vertex buffer not set or too small - these are warnings and if the shader doesn't read from it it's fine. This happens because vertex
+				// buffers are not removed from the cache, but only get removed when another buffer is set at the same slot or when the buffer gets destroyed.
+				D3D12_MESSAGE_ID_COMMAND_LIST_DRAW_VERTEX_BUFFER_NOT_SET,
+				D3D12_MESSAGE_ID_COMMAND_LIST_DRAW_VERTEX_BUFFER_TOO_SMALL,
+
+				// D3D12 complains when a buffer is created with a specific initial resource state while all buffers are currently created in COMMON state. The 
+				// next transition is then done use state promotion. It's just a warning and we need to keep track of the correct initial state as well for upcoming
+				// internal transitions.
+				D3D12_MESSAGE_ID_CREATERESOURCE_STATE_IGNORED,
+
 #if ENABLE_RESIDENCY_MANAGEMENT
 				// TODO: Remove this when the debug layers work for executions which are guarded by a fence
 				D3D12_MESSAGE_ID_INVALID_USE_OF_NON_RESIDENT_RESOURCE,
