Addressing issues caused by not handling multi-plane texture formats correctly.
Making an intentional decision to only copy plane 0 for now since our readback buffers used in this manner are still only sized for that plane.

diff --git a/Engine/Source/Runtime/D3D12RHI/Private/D3D12RenderTarget.cpp b/Engine/Source/Runtime/D3D12RHI/Private/D3D12RenderTarget.cpp
--- a/Engine/Source/Runtime/D3D12RHI/Private/D3D12RenderTarget.cpp
+++ b/Engine/Source/Runtime/D3D12RHI/Private/D3D12RenderTarget.cpp
@@ -15,6 +15,98 @@
 #include "RHISurfaceDataConversion.h"
 #include "CommonRenderResources.h"
 
+static uint32 ComputeBytesPerPixel(DXGI_FORMAT Format)
+{
+	uint32 BytesPerPixel = 0;
+
+	switch (Format)
+	{
+	case DXGI_FORMAT_R8G8_TYPELESS:
+	case DXGI_FORMAT_R8G8_UNORM:
+	case DXGI_FORMAT_R8G8_UINT:
+	case DXGI_FORMAT_R8G8_SNORM:
+	case DXGI_FORMAT_R8G8_SINT:
+	case DXGI_FORMAT_R16_TYPELESS:
+	case DXGI_FORMAT_R16_FLOAT:
+	case DXGI_FORMAT_D16_UNORM:
+	case DXGI_FORMAT_R16_UNORM:
+	case DXGI_FORMAT_R16_UINT:
+	case DXGI_FORMAT_R16_SNORM:
+	case DXGI_FORMAT_R16_SINT:
+	case DXGI_FORMAT_B5G6R5_UNORM:
+	case DXGI_FORMAT_B5G5R5A1_UNORM:
+		BytesPerPixel = 2;
+		break;
+	case DXGI_FORMAT_B8G8R8A8_TYPELESS:
+	case DXGI_FORMAT_B8G8R8A8_UNORM:
+	case DXGI_FORMAT_R8G8B8A8_TYPELESS:
+	case DXGI_FORMAT_R8G8B8A8_UNORM:
+	case DXGI_FORMAT_R24G8_TYPELESS:
+	case DXGI_FORMAT_R10G10B10A2_UNORM:
+	case DXGI_FORMAT_R11G11B10_FLOAT:
+	case DXGI_FORMAT_R16G16_UNORM:
+	case DXGI_FORMAT_R32_UINT:
+	case DXGI_FORMAT_R32_TYPELESS:
+	case DXGI_FORMAT_R32_FLOAT:
+	case DXGI_FORMAT_B8G8R8A8_UNORM_SRGB:
+	case DXGI_FORMAT_B8G8R8X8_TYPELESS:
+	case DXGI_FORMAT_B8G8R8X8_UNORM_SRGB:
+	case DXGI_FORMAT_B8G8R8X8_UNORM:
+	case DXGI_FORMAT_R10G10B10A2_TYPELESS:
+	case DXGI_FORMAT_R10G10B10A2_UINT:
+	case DXGI_FORMAT_R8G8B8A8_UNORM_SRGB:
+	case DXGI_FORMAT_R8G8B8A8_UINT:
+	case DXGI_FORMAT_R8G8B8A8_SNORM:
+	case DXGI_FORMAT_R8G8B8A8_SINT:
+	case DXGI_FORMAT_R16G16_TYPELESS:
+	case DXGI_FORMAT_R16G16_FLOAT:
+	case DXGI_FORMAT_R16G16_UINT:
+	case DXGI_FORMAT_R16G16_SNORM:
+	case DXGI_FORMAT_R16G16_SINT:
+	case DXGI_FORMAT_D32_FLOAT:
+	case DXGI_FORMAT_R32_SINT:
+		BytesPerPixel = 4;
+		break;
+	case DXGI_FORMAT_R16G16B16A16_FLOAT:
+	case DXGI_FORMAT_R16G16B16A16_UNORM:
+	case DXGI_FORMAT_R16G16B16A16_TYPELESS:
+	case DXGI_FORMAT_R16G16B16A16_UINT:
+	case DXGI_FORMAT_R16G16B16A16_SNORM:
+	case DXGI_FORMAT_R16G16B16A16_SINT:
+	case DXGI_FORMAT_R32G32_TYPELESS:
+	case DXGI_FORMAT_R32G32_FLOAT:
+	case DXGI_FORMAT_R32G32_UINT:
+	case DXGI_FORMAT_R32G32_SINT:
+		BytesPerPixel = 8;
+		break;
+	case DXGI_FORMAT_R32G8X24_TYPELESS:
+	case DXGI_FORMAT_D32_FLOAT_S8X24_UINT:
+	case DXGI_FORMAT_R32_FLOAT_X8X24_TYPELESS:
+		BytesPerPixel = 8;
+		break;
+	case DXGI_FORMAT_R8_TYPELESS:
+	case DXGI_FORMAT_R8_UNORM:
+	case DXGI_FORMAT_R8_UINT:
+	case DXGI_FORMAT_R8_SNORM:
+	case DXGI_FORMAT_R8_SINT:
+	case DXGI_FORMAT_A8_UNORM:
+	case DXGI_FORMAT_R1_UNORM:
+		BytesPerPixel = 1;
+		break;
+	case DXGI_FORMAT_R32G32B32A32_UINT:
+	case DXGI_FORMAT_R32G32B32A32_TYPELESS:
+	case DXGI_FORMAT_R32G32B32A32_FLOAT:
+	case DXGI_FORMAT_R32G32B32A32_SINT:
+		BytesPerPixel = 16;
+		break;
+	}
+
+	// format not supported yet
+	check(BytesPerPixel);
+
+	return BytesPerPixel;
+}
+
 static inline DXGI_FORMAT ConvertTypelessToUnorm(DXGI_FORMAT Format)
 {
 	// required to prevent 
@@ -189,6 +281,23 @@
 	}
 }
 
+static DXGI_FORMAT GetPlaneFormat(DXGI_FORMAT InFormat, uint32 InPlaneSlice)
+{
+	if (InFormat == DXGI_FORMAT_R32G8X24_TYPELESS || InFormat == DXGI_FORMAT_R24G8_TYPELESS)
+	{
+		if (InPlaneSlice == 0)
+		{
+			return DXGI_FORMAT_R32_TYPELESS;
+		}
+		if (InPlaneSlice == 1)
+		{
+			return DXGI_FORMAT_R8_TYPELESS;
+		}
+	}
+
+	return InFormat;
+}
+
 /**
 * Copies the contents of the given surface to its resolve target texture.
 * @param SourceSurface - surface with a resolve texture to copy to
@@ -317,23 +426,24 @@
 						{
 							check(IsDefaultContext());
 
-							const uint32 BlockBytes = GPixelFormats[SourceTexture->GetFormat()].BlockBytes;
+							const uint32 PlaneSlice = 0;
+							const DXGI_FORMAT DestFormat = GetPlaneFormat(srcDesc.Format, PlaneSlice);
+
+							const uint32 BlockBytes = ComputeBytesPerPixel(DestFormat);
 							const uint32 XBytes = (uint32)srcDesc.Width * BlockBytes;
 							const uint32 XBytesAligned = Align(XBytes, FD3D12_TEXTURE_DATA_PITCH_ALIGNMENT);
 
-							D3D12_SUBRESOURCE_FOOTPRINT DestSubresource;
-							DestSubresource.Depth = 1;
-							DestSubresource.Height = srcDesc.Height;
-							DestSubresource.Width = srcDesc.Width;
-							DestSubresource.Format = srcDesc.Format;
-							DestSubresource.RowPitch = XBytesAligned;
+							D3D12_PLACED_SUBRESOURCE_FOOTPRINT DestFootprint{};
+							DestFootprint.Footprint.Depth = 1;
+							DestFootprint.Footprint.Height = srcDesc.Height;
+							DestFootprint.Footprint.Width = srcDesc.Width;
+							DestFootprint.Footprint.Format = DestFormat;
+							DestFootprint.Footprint.RowPitch = XBytesAligned;
 
-							D3D12_PLACED_SUBRESOURCE_FOOTPRINT placedTexture2D = {0};
-							placedTexture2D.Offset = 0;
-							placedTexture2D.Footprint = DestSubresource;
+							const UINT SourceSubresource = D3D12CalcSubresource(ResolveParams.MipIndex, ResolveParams.SourceArrayIndex, PlaneSlice, srcDesc.MipLevels, srcDesc.DepthOrArraySize);
 
-							CD3DX12_TEXTURE_COPY_LOCATION DestCopyLocation(DestTexture->GetResource()->GetResource(), placedTexture2D);
-							CD3DX12_TEXTURE_COPY_LOCATION SourceCopyLocation(SourceTexture->GetResource()->GetResource(), ResolveParams.SourceArrayIndex);
+							CD3DX12_TEXTURE_COPY_LOCATION DestCopyLocation(DestTexture->GetResource()->GetResource(), DestFootprint);
+							CD3DX12_TEXTURE_COPY_LOCATION SourceCopyLocation(SourceTexture->GetResource()->GetResource(), SourceSubresource);
 
 							numCopies++;
 							CommandListHandle.FlushResourceBarriers();
@@ -546,98 +656,6 @@
 	DEBUG_EXECUTE_COMMAND_LIST(this);
 }
 
-static uint32 ComputeBytesPerPixel(DXGI_FORMAT Format)
-{
-	uint32 BytesPerPixel = 0;
-
-	switch (Format)
-	{
-	case DXGI_FORMAT_R8G8_TYPELESS:
-	case DXGI_FORMAT_R8G8_UNORM:
-	case DXGI_FORMAT_R8G8_UINT:
-	case DXGI_FORMAT_R8G8_SNORM:
-	case DXGI_FORMAT_R8G8_SINT:
-	case DXGI_FORMAT_R16_TYPELESS:
-	case DXGI_FORMAT_R16_FLOAT:
-	case DXGI_FORMAT_D16_UNORM:
-	case DXGI_FORMAT_R16_UNORM:
-	case DXGI_FORMAT_R16_UINT:
-	case DXGI_FORMAT_R16_SNORM:
-	case DXGI_FORMAT_R16_SINT:
-	case DXGI_FORMAT_B5G6R5_UNORM:
-	case DXGI_FORMAT_B5G5R5A1_UNORM:
-		BytesPerPixel = 2;
-		break;
-	case DXGI_FORMAT_B8G8R8A8_TYPELESS:
-	case DXGI_FORMAT_B8G8R8A8_UNORM:
-	case DXGI_FORMAT_R8G8B8A8_TYPELESS:
-	case DXGI_FORMAT_R8G8B8A8_UNORM:
-	case DXGI_FORMAT_R24G8_TYPELESS:
-	case DXGI_FORMAT_R10G10B10A2_UNORM:
-	case DXGI_FORMAT_R11G11B10_FLOAT:
-	case DXGI_FORMAT_R16G16_UNORM:
-	case DXGI_FORMAT_R32_UINT:
-	case DXGI_FORMAT_R32_TYPELESS:
-	case DXGI_FORMAT_R32_FLOAT:
-	case DXGI_FORMAT_B8G8R8A8_UNORM_SRGB:
-	case DXGI_FORMAT_B8G8R8X8_TYPELESS:
-	case DXGI_FORMAT_B8G8R8X8_UNORM_SRGB:
-	case DXGI_FORMAT_B8G8R8X8_UNORM:
-	case DXGI_FORMAT_R10G10B10A2_TYPELESS:
-	case DXGI_FORMAT_R10G10B10A2_UINT:
-	case DXGI_FORMAT_R8G8B8A8_UNORM_SRGB:
-	case DXGI_FORMAT_R8G8B8A8_UINT:
-	case DXGI_FORMAT_R8G8B8A8_SNORM:
-	case DXGI_FORMAT_R8G8B8A8_SINT:
-	case DXGI_FORMAT_R16G16_TYPELESS:
-	case DXGI_FORMAT_R16G16_FLOAT:
-	case DXGI_FORMAT_R16G16_UINT:
-	case DXGI_FORMAT_R16G16_SNORM:
-	case DXGI_FORMAT_R16G16_SINT:
-	case DXGI_FORMAT_D32_FLOAT:
-	case DXGI_FORMAT_R32_SINT:
-		BytesPerPixel = 4;
-		break;
-	case DXGI_FORMAT_R16G16B16A16_FLOAT:
-	case DXGI_FORMAT_R16G16B16A16_UNORM:
-	case DXGI_FORMAT_R16G16B16A16_TYPELESS:
-	case DXGI_FORMAT_R16G16B16A16_UINT:
-	case DXGI_FORMAT_R16G16B16A16_SNORM:
-	case DXGI_FORMAT_R16G16B16A16_SINT:
-	case DXGI_FORMAT_R32G32_TYPELESS:
-	case DXGI_FORMAT_R32G32_FLOAT:
-	case DXGI_FORMAT_R32G32_UINT:
-	case DXGI_FORMAT_R32G32_SINT:
-		BytesPerPixel = 8;
-		break;
-	case DXGI_FORMAT_R32G8X24_TYPELESS:
-	case DXGI_FORMAT_D32_FLOAT_S8X24_UINT:
-	case DXGI_FORMAT_R32_FLOAT_X8X24_TYPELESS:
-		BytesPerPixel = 8;
-		break;
-	case DXGI_FORMAT_R8_TYPELESS:
-	case DXGI_FORMAT_R8_UNORM:
-	case DXGI_FORMAT_R8_UINT:
-	case DXGI_FORMAT_R8_SNORM:
-	case DXGI_FORMAT_R8_SINT:
-	case DXGI_FORMAT_A8_UNORM:
-	case DXGI_FORMAT_R1_UNORM:
-		BytesPerPixel = 1;
-		break;
-	case DXGI_FORMAT_R32G32B32A32_UINT:
-	case DXGI_FORMAT_R32G32B32A32_TYPELESS:
-	case DXGI_FORMAT_R32G32B32A32_FLOAT:
-	case DXGI_FORMAT_R32G32B32A32_SINT:
-		BytesPerPixel = 16;
-		break;
-	}
-
-	// format not supported yet
-	check(BytesPerPixel);
-
-	return BytesPerPixel;
-}
-
 TRefCountPtr<FD3D12Resource> FD3D12DynamicRHI::GetStagingTexture(FRHITexture* TextureRHI, FIntRect InRect, FIntRect& StagingRectOUT, FReadSurfaceDataFlags InFlags, D3D12_PLACED_SUBRESOURCE_FOOTPRINT &readbackHeapDesc, uint32 GPUIndex)
 {
 	FD3D12Device* Device = GetRHIDevice(GPUIndex);
@@ -681,7 +699,10 @@
 	// create a temp 2d texture to copy render target to
 	TRefCountPtr<FD3D12Resource> TempTexture2D;
 
-	const uint32 BlockBytes = GPixelFormats[TextureRHI->GetFormat()].BlockBytes;
+	const uint32 PlaneSlice = 0;
+	const DXGI_FORMAT DestFormat = GetPlaneFormat(SourceDesc.Format, PlaneSlice);
+
+	const uint32 BlockBytes = ComputeBytesPerPixel(DestFormat);
 	const uint32 XBytesAligned = Align((uint32)SourceDesc.Width * BlockBytes, FD3D12_TEXTURE_DATA_PITCH_ALIGNMENT);
 	const uint32 MipBytesAligned = XBytesAligned * SourceDesc.Height;
 	VERIFYD3D12RESULT(Adapter->CreateBuffer(D3D12_HEAP_TYPE_READBACK, Node, Node, MipBytesAligned, TempTexture2D.GetInitReference(), nullptr));
@@ -709,20 +730,15 @@
 		RectPtr = &Rect;
 	}
 
-	uint32 BytesPerPixel = ComputeBytesPerPixel(SourceDesc.Format);
-	D3D12_SUBRESOURCE_FOOTPRINT DestSubresource;
-	DestSubresource.Depth = 1;
-	DestSubresource.Height = SourceDesc.Height;
-	DestSubresource.Width = SourceDesc.Width;
-	DestSubresource.Format = SourceDesc.Format;
-	DestSubresource.RowPitch = XBytesAligned;
-	check(DestSubresource.RowPitch % FD3D12_TEXTURE_DATA_PITCH_ALIGNMENT == 0);	// Make sure we align correctly.
+	D3D12_PLACED_SUBRESOURCE_FOOTPRINT DestFootprint{};
+	DestFootprint.Footprint.Depth = 1;
+	DestFootprint.Footprint.Height = SourceDesc.Height;
+	DestFootprint.Footprint.Width = SourceDesc.Width;
+	DestFootprint.Footprint.Format = DestFormat;
+	DestFootprint.Footprint.RowPitch = XBytesAligned;
+	check((DestFootprint.Footprint.RowPitch % FD3D12_TEXTURE_DATA_PITCH_ALIGNMENT) == 0);	// Make sure we align correctly.
 
-	D3D12_PLACED_SUBRESOURCE_FOOTPRINT placedTexture2D = { 0 };
-	placedTexture2D.Offset = 0;
-	placedTexture2D.Footprint = DestSubresource;
-
-	CD3DX12_TEXTURE_COPY_LOCATION DestCopyLocation(TempTexture2D->GetResource(), placedTexture2D);
+	CD3DX12_TEXTURE_COPY_LOCATION DestCopyLocation(TempTexture2D->GetResource(), DestFootprint);
 	CD3DX12_TEXTURE_COPY_LOCATION SourceCopyLocation(Texture->GetResource()->GetResource(), Subresource);
 
 	FScopedResourceBarrier ScopeResourceBarrierSource(hCommandList, Texture->GetResource(), D3D12_RESOURCE_STATE_COPY_SOURCE, SourceCopyLocation.SubresourceIndex, FD3D12DynamicRHI::ETransitionMode::Apply);
@@ -739,7 +755,7 @@
 	hCommandList.UpdateResidency(Texture->GetResource());
 
 	// Remember the width, height, pitch, etc...
-	readbackHeapDesc = placedTexture2D;
+	readbackHeapDesc = DestFootprint;
 
 	// We need to execute the command list so we can read the data from readback heap
 	Device->GetDefaultCommandContext().FlushCommands(true);
@@ -1327,7 +1343,11 @@
 
 	// Create a CPU-accessible staging texture to copy the resolved sample data to.
 	TRefCountPtr<FD3D12Resource> StagingTexture2D;
-	const uint32 BlockBytes = GPixelFormats[TextureRHI->GetFormat()].BlockBytes;
+
+	const uint32 PlaneSlice = 0;
+	const DXGI_FORMAT DestFormat = GetPlaneFormat(TextureDesc.Format, PlaneSlice);
+
+	const uint32 BlockBytes = ComputeBytesPerPixel(DestFormat);
 	const uint32 XBytesAligned = Align(SizeX * BlockBytes, FD3D12_TEXTURE_DATA_PITCH_ALIGNMENT);
 	const uint32 MipBytesAligned = XBytesAligned * SizeY;
 	VERIFYD3D12RESULT(Adapter->CreateBuffer(D3D12_HEAP_TYPE_READBACK, NodeMask, NodeMask, MipBytesAligned, StagingTexture2D.GetInitReference(), nullptr));
@@ -1348,19 +1368,16 @@
 	}
 
 	// Setup the descriptions for the copy to the readback heap.
-	D3D12_SUBRESOURCE_FOOTPRINT DestSubresource;
-	DestSubresource.Depth = 1;
-	DestSubresource.Height = SizeY;
-	DestSubresource.Width = SizeX;
-	DestSubresource.Format = TextureDesc.Format;
-	DestSubresource.RowPitch = XBytesAligned;
-	check(DestSubresource.RowPitch % FD3D12_TEXTURE_DATA_PITCH_ALIGNMENT == 0);	// Make sure we align correctly.
 
-	D3D12_PLACED_SUBRESOURCE_FOOTPRINT placedTexture2D = { 0 };
-	placedTexture2D.Offset = 0;
-	placedTexture2D.Footprint = DestSubresource;
+	D3D12_PLACED_SUBRESOURCE_FOOTPRINT DestFootprint{};
+	DestFootprint.Footprint.Depth = 1;
+	DestFootprint.Footprint.Height = SizeY;
+	DestFootprint.Footprint.Width = SizeX;
+	DestFootprint.Footprint.Format = DestFormat;
+	DestFootprint.Footprint.RowPitch = XBytesAligned;
+	check((DestFootprint.Footprint.RowPitch % FD3D12_TEXTURE_DATA_PITCH_ALIGNMENT) == 0);	// Make sure we align correctly.
 
-	CD3DX12_TEXTURE_COPY_LOCATION DestCopyLocation(StagingTexture2D->GetResource(), placedTexture2D);
+	CD3DX12_TEXTURE_COPY_LOCATION DestCopyLocation(StagingTexture2D->GetResource(), DestFootprint);
 	CD3DX12_TEXTURE_COPY_LOCATION SourceCopyLocation(NonMSAATexture2D->GetResource(), Subresource);
 
 	// Allocate the output buffer.
@@ -1537,7 +1554,11 @@
 
 	// create a temp 2d texture to copy render target to
 	TRefCountPtr<FD3D12Resource> TempTexture2D;
-	const uint32 BlockBytes = GPixelFormats[TextureRHI->GetFormat()].BlockBytes;
+
+	const uint32 PlaneSlice = 0;
+	const DXGI_FORMAT DestFormat = GetPlaneFormat(TextureDesc.Format, PlaneSlice);
+
+	const uint32 BlockBytes = ComputeBytesPerPixel(DestFormat);
 	const uint32 XBytesAligned = Align((uint32)TextureDesc.Width * BlockBytes, FD3D12_TEXTURE_DATA_PITCH_ALIGNMENT);
 	const uint32 MipBytesAligned = XBytesAligned * TextureDesc.Height;
 	VERIFYD3D12RESULT(Adapter->CreateBuffer(D3D12_HEAP_TYPE_READBACK, Node, Node, MipBytesAligned, TempTexture2D.GetInitReference(), nullptr));
@@ -1558,20 +1579,15 @@
 		Subresource = CalcSubresource(InFlags.GetMip(), bIsTextureArray ? InFlags.GetArrayIndex() : 0, TextureDesc.MipLevels);
 	}
 
-	uint32 BytesPerPixel = ComputeBytesPerPixel(TextureDesc.Format);
-	D3D12_SUBRESOURCE_FOOTPRINT DestSubresource;
-	DestSubresource.Depth = 1;
-	DestSubresource.Height = TextureDesc.Height;
-	DestSubresource.Width = TextureDesc.Width;
-	DestSubresource.Format = TextureDesc.Format;
-	DestSubresource.RowPitch = XBytesAligned;
-	check(DestSubresource.RowPitch % FD3D12_TEXTURE_DATA_PITCH_ALIGNMENT == 0);	// Make sure we align correctly.
+	D3D12_PLACED_SUBRESOURCE_FOOTPRINT DestFootprint{};
+	DestFootprint.Footprint.Depth = 1;
+	DestFootprint.Footprint.Height = TextureDesc.Height;
+	DestFootprint.Footprint.Width = TextureDesc.Width;
+	DestFootprint.Footprint.Format = DestFormat;
+	DestFootprint.Footprint.RowPitch = XBytesAligned;
+	check((DestFootprint.Footprint.RowPitch % FD3D12_TEXTURE_DATA_PITCH_ALIGNMENT) == 0);	// Make sure we align correctly.
 
-	D3D12_PLACED_SUBRESOURCE_FOOTPRINT placedTexture2D = { 0 };
-	placedTexture2D.Offset = 0;
-	placedTexture2D.Footprint = DestSubresource;
-
-	CD3DX12_TEXTURE_COPY_LOCATION DestCopyLocation(TempTexture2D->GetResource(), placedTexture2D);
+	CD3DX12_TEXTURE_COPY_LOCATION DestCopyLocation(TempTexture2D->GetResource(), DestFootprint);
 	CD3DX12_TEXTURE_COPY_LOCATION SourceCopyLocation(Texture->GetResource()->GetResource(), Subresource);
 
 	{
