Add validation for acceleration structures having been built before the shader binding table is populated

diff --git a/Engine/Source/Runtime/D3D12RHI/Private/D3D12RayTracing.cpp b/Engine/Source/Runtime/D3D12RHI/Private/D3D12RayTracing.cpp
--- a/Engine/Source/Runtime/D3D12RHI/Private/D3D12RayTracing.cpp
+++ b/Engine/Source/Runtime/D3D12RHI/Private/D3D12RayTracing.cpp
@@ -4252,6 +4252,8 @@ void FD3D12RayTracingScene::BuildAccelerationStructure(FD3D12CommandContext& Com
 	// UAV barrier is used here to ensure that the acceleration structure build is complete before any rays are traced
 	// #dxr_todo: these barriers should ideally be inserted by the high level code to allow more overlapped execution
 	CommandContext.CommandListHandle.AddUAVBarrier();
+
+	bBuilt = true;
 }
 
 void FD3D12RayTracingScene::UpdateResidency(FD3D12CommandContext& CommandContext)
@@ -5613,6 +5615,8 @@ void FD3D12CommandContext::RHISetRayTracingBindings(
 	FD3D12RayTracingPipelineState* Pipeline = FD3D12DynamicRHI::ResourceCast(InPipeline);
 	FD3D12Device* Device = GetParentDevice();
 
+	checkf(Scene->bBuilt, TEXT("Ray tracing scene must be built before any shaders can be bound to it. Make sure that RHIBuildAccelerationStructure() command has been executed."));
+
 	FD3D12RayTracingShaderTable* ShaderTable = Scene->FindOrCreateShaderTable(Pipeline, Device);
 	checkf(ShaderTable->LocalShaderTableOffset == ShaderTable->HitGroupShaderTableOffset,
 		TEXT("Hit shader records are assumed to be at the beginning of local shader table"));
diff --git a/Engine/Source/Runtime/D3D12RHI/Private/D3D12RayTracing.h b/Engine/Source/Runtime/D3D12RHI/Private/D3D12RayTracing.h
--- a/Engine/Source/Runtime/D3D12RHI/Private/D3D12RayTracing.h
+++ b/Engine/Source/Runtime/D3D12RHI/Private/D3D12RayTracing.h
@@ -148,6 +148,8 @@ public:
 	TMap<const FD3D12RayTracingPipelineState*, FD3D12RayTracingShaderTable*> ShaderTables[MAX_NUM_GPUS];
 
 	uint64 LastCommandListID = 0;
+
+	bool bBuilt = false;
 };
 
 // Manages all the pending BLAS compaction requests
