Separate backbuffer transitions

diff --git a/Engine/Source/Runtime/Core/Public/HAL/Platform.h b/Engine/Source/Runtime/Core/Public/HAL/Platform.h
--- a/Engine/Source/Runtime/Core/Public/HAL/Platform.h
+++ b/Engine/Source/Runtime/Core/Public/HAL/Platform.h
@@ -613,6 +613,10 @@
 	#define PLATFORM_USE_BACKBUFFER_WRITE_TRANSITION_TRACKING 0
 #endif
 
+#ifndef PLATFORM_USE_SEPARATE_BACKBUFFER_WRITE_TRANSITION
+	#define PLATFORM_USE_SEPARATE_BACKBUFFER_WRITE_TRANSITION 0
+#endif
+
 #ifndef PLATFORM_USE_PLATFORM_FILE_MANAGED_STORAGE_WRAPPER
 	#define PLATFORM_USE_PLATFORM_FILE_MANAGED_STORAGE_WRAPPER	0
 #endif
diff --git a/Engine/Source/Runtime/D3D12RHI/Private/D3D12Resources.cpp b/Engine/Source/Runtime/D3D12RHI/Private/D3D12Resources.cpp
--- a/Engine/Source/Runtime/D3D12RHI/Private/D3D12Resources.cpp
+++ b/Engine/Source/Runtime/D3D12RHI/Private/D3D12Resources.cpp
@@ -1267,6 +1267,12 @@ static void RecordResourceBarriersToCommandList(
 	}
 }
 
+#if (PLATFORM_USE_BACKBUFFER_WRITE_TRANSITION_TRACKING == 0) && (PLATFORM_USE_SEPARATE_BACKBUFFER_WRITE_TRANSITION == 1)
+#define LOCAL_USE_SEPARATE_BACKBUFFER_WRITE_TRANSITION 1
+#else
+#define LOCAL_USE_SEPARATE_BACKBUFFER_WRITE_TRANSITION 0
+#endif
+
 void ResourceBarriersSeparateRTV2SRV(
 	ID3D12GraphicsCommandList* pCommandList,
 	const TArray<D3D12_RESOURCE_BARRIER>& Barriers,
@@ -1274,11 +1280,44 @@ void ResourceBarriersSeparateRTV2SRV(
 {
 	if (!GD3D12SeparateRTV2SRVTransitions)
 	{
+#if LOCAL_USE_SEPARATE_BACKBUFFER_WRITE_TRANSITION
+		TArray<D3D12_RESOURCE_BARRIER, TInlineAllocator<4>> BackBufferBarriers;
+		TArray<D3D12_RESOURCE_BARRIER, TInlineAllocator<8>> OtherBarriers;
+
+		for (int32 Index = 0; Index < Barriers.Num(); ++Index)
+		{
+			const D3D12_RESOURCE_BARRIER& Barrier = Barriers[Index];
+			if (Barrier.Type == D3D12_RESOURCE_BARRIER_TYPE_TRANSITION
+				&& Barrier.Transition.StateBefore == D3D12_RESOURCE_STATE_PRESENT // can also be displayed as D3D12_RESOURCE_STATE_COMMON in pix
+				&& Barrier.Transition.StateAfter == D3D12_RESOURCE_STATE_RENDER_TARGET)
+			{
+				BackBufferBarriers.Add(Barrier);
+			}
+			else
+			{
+				OtherBarriers.Add(Barrier);
+			}
+		}
+
+		if (BackBufferBarriers.Num() > 0)
+		{
+			RecordResourceBarriersToCommandList(pCommandList, BackBufferBarriers.GetData(), BackBufferBarriers.Num(), BarrierBatchMax);
+		}
+
+		if (OtherBarriers.Num() > 0)
+		{
+			RecordResourceBarriersToCommandList(pCommandList, OtherBarriers.GetData(), OtherBarriers.Num(), BarrierBatchMax);
+		}
+#else
 		RecordResourceBarriersToCommandList(pCommandList, Barriers.GetData(), Barriers.Num(), BarrierBatchMax);
+#endif
 	}
 	else
 	{
 		TArray<D3D12_RESOURCE_BARRIER, TInlineAllocator<4>> RTV2SRVBarriers;
+#if LOCAL_USE_SEPARATE_BACKBUFFER_WRITE_TRANSITION
+		TArray<D3D12_RESOURCE_BARRIER, TInlineAllocator<4>> BackBufferBarriers;
+#endif
 		TArray<D3D12_RESOURCE_BARRIER, TInlineAllocator<8>> OtherBarriers;
 
 		for (int32 Index = 0; Index < Barriers.Num(); ++Index)
@@ -1291,6 +1330,14 @@ void ResourceBarriersSeparateRTV2SRV(
 			{
 				RTV2SRVBarriers.Add(Barrier);
 			}
+#if LOCAL_USE_SEPARATE_BACKBUFFER_WRITE_TRANSITION
+			else if (Barrier.Type == D3D12_RESOURCE_BARRIER_TYPE_TRANSITION
+				&& Barrier.Transition.StateBefore == D3D12_RESOURCE_STATE_PRESENT // can also be displayed as D3D12_RESOURCE_STATE_COMMON in pix
+				&& Barrier.Transition.StateAfter == D3D12_RESOURCE_STATE_RENDER_TARGET)
+			{
+				BackBufferBarriers.Add(Barrier);
+			}
+#endif
 			else
 			{
 				OtherBarriers.Add(Barrier);
@@ -1302,12 +1349,19 @@ void ResourceBarriersSeparateRTV2SRV(
 			RecordResourceBarriersToCommandList(pCommandList, RTV2SRVBarriers.GetData(), RTV2SRVBarriers.Num(), BarrierBatchMax);
 		}
 
+#if LOCAL_USE_SEPARATE_BACKBUFFER_WRITE_TRANSITION
+		if (BackBufferBarriers.Num() > 0)
+		{
+			RecordResourceBarriersToCommandList(pCommandList, BackBufferBarriers.GetData(), BackBufferBarriers.Num(), BarrierBatchMax);
+		}
+#endif
 		if (OtherBarriers.Num() > 0)
 		{
 			RecordResourceBarriersToCommandList(pCommandList, OtherBarriers.GetData(), OtherBarriers.Num(), BarrierBatchMax);
 		}
 	}
 }
+#undef LOCAL_USE_SEPARATE_BACKBUFFER_WRITE_TRANSITION
 
 void FD3D12ResourceBarrierBatcher::Flush(FD3D12Device* Device, ID3D12GraphicsCommandList* pCommandList, int32 BarrierBatchMax)
 {
