Raising D3D12RHI minimum Windows 10 version to 15063 aka Windows 1703 aka "Windows Creator Update" as this is the version ID3D12Device2 was introduced.
Adding fatal error on D3D12RHI init if the minimum device types aren't supported (ID3D12Device1 and ID3D12Device2).

#rnx

diff --git a/Engine/Source/Runtime/D3D12RHI/Private/D3D12Adapter.cpp b/Engine/Source/Runtime/D3D12RHI/Private/D3D12Adapter.cpp
--- a/Engine/Source/Runtime/D3D12RHI/Private/D3D12Adapter.cpp
+++ b/Engine/Source/Runtime/D3D12RHI/Private/D3D12Adapter.cpp
@@ -888,6 +888,13 @@ void FD3D12Adapter::InitializeDevices()
 			{
 				UE_LOG(LogD3D12RHI, Log, TEXT("ID3D12Device2 is supported."));
 			}
+
+			if (RootDevice1 == nullptr || RootDevice2 == nullptr)
+			{
+				// Note: we require Windows 1703 in FD3D12DynamicRHIModule::IsSupported()
+				// If we still lack support, the user's drivers could be out of date.
+				UE_LOG(LogD3D12RHI, Fatal, TEXT("Missing full support for Direct3D 12. Please update to the latest drivers."));
+			}
 #endif
 #if D3D12_MAX_DEVICE_INTERFACE >= 3
 			if (SUCCEEDED(RootDevice->QueryInterface(IID_PPV_ARGS(RootDevice3.GetInitReference()))))
diff --git a/Engine/Source/Runtime/D3D12RHI/Private/Windows/WindowsD3D12Device.cpp b/Engine/Source/Runtime/D3D12RHI/Private/Windows/WindowsD3D12Device.cpp
--- a/Engine/Source/Runtime/D3D12RHI/Private/Windows/WindowsD3D12Device.cpp
+++ b/Engine/Source/Runtime/D3D12RHI/Private/Windows/WindowsD3D12Device.cpp
@@ -470,8 +470,11 @@ static bool IsAdapterSupported(FD3D12Adapter* InAdapter, ERHIFeatureLevel::Type
 bool FD3D12DynamicRHIModule::IsSupported(ERHIFeatureLevel::Type RequestedFeatureLevel)
 {
 #if !PLATFORM_HOLOLENS
-	if (!FPlatformMisc::VerifyWindowsVersion(10, 0))
+	// Windows version 15063 is Windows 1703 aka "Windows Creator Update"
+	// This is the first version that supports ID3D12Device2 which is our minimum runtime device version.
+	if (!FPlatformMisc::VerifyWindowsVersion(10, 0, 15063))
 	{
+		UE_LOG(LogD3D12RHI, Warning, TEXT("Missing full support for Direct3D 12. Update to Windows 1703 or newer for D3D12 support."));
 		return false;
 	}
 #endif
