Fix flags that were not provided when creating a new shader resource view. ETextureCreateFlags::NoFastClearFinalize / TexCreate_NoFastClearFinalize was ignored  which was always flushing metadata during UI composition in SlateUI_Composition

diff --git a/Engine/Source/Runtime/D3D12RHI/Private/D3D12Texture.cpp b/Engine/Source/Runtime/D3D12RHI/Private/D3D12Texture.cpp
--- a/Engine/Source/Runtime/D3D12RHI/Private/D3D12Texture.cpp
+++ b/Engine/Source/Runtime/D3D12RHI/Private/D3D12Texture.cpp
@@ -1653,7 +1653,7 @@ void FD3D12Texture::CreateViews()
 			}
 		}
 
-		SetShaderResourceView(new FD3D12ShaderResourceView(this, SRVDesc));
+		SetShaderResourceView(new FD3D12ShaderResourceView(this, SRVDesc, Desc.Flags));
 	}
 }
 
