Fix for ensure in D3D12 async buffer creation due to IsInParallelRenderingThread being outside of a task scope.

diff --git a/Engine/Source/Runtime/D3D12RHI/Private/D3D12Buffer.cpp b/Engine/Source/Runtime/D3D12RHI/Private/D3D12Buffer.cpp
--- a/Engine/Source/Runtime/D3D12RHI/Private/D3D12Buffer.cpp
+++ b/Engine/Source/Runtime/D3D12RHI/Private/D3D12Buffer.cpp
@@ -197,7 +197,7 @@ void FD3D12Buffer::UploadResourceData(FRHICommandListBase& RHICmdList, FResource
 	}
 	else
 	{
-		const bool bOnAsyncThread = !IsInRHIThread() && !IsInParallelRenderingThread();
+		const bool bOnAsyncThread = !IsInRHIThread() && !IsInRenderingThread();
 
 		// Get an upload heap and initialize data
 		FD3D12ResourceLocation SrcResourceLoc(GetParentDevice());
