Fix for ProfileGPU scoping error in D3D12. D3D12 was erroneously including async compute scopes in the device GPU profiler stack.

diff --git a/Engine/Source/Runtime/D3D12RHI/Private/D3D12CommandContext.cpp b/Engine/Source/Runtime/D3D12RHI/Private/D3D12CommandContext.cpp
--- a/Engine/Source/Runtime/D3D12RHI/Private/D3D12CommandContext.cpp
+++ b/Engine/Source/Runtime/D3D12RHI/Private/D3D12CommandContext.cpp
@@ -222,7 +222,7 @@ void FD3D12CommandContext::RHIPushEvent(const TCHAR* Name, FColor Color)
 	D3D12RHI::FD3DGPUProfiler& GPUProfiler = GetParentDevice()->GetGPUProfiler();
 
 	// forward event to profiler if it's the default context
-	if (IsDefaultContext())
+	if (IsDefaultContext() && !IsAsyncComputeContext())
 	{
 		GPUProfiler.PushEvent(Name, Color);
 	}
@@ -264,7 +264,7 @@ void FD3D12CommandContext::RHIPopEvent()
 {
 	D3D12RHI::FD3DGPUProfiler& GPUProfiler = GetParentDevice()->GetGPUProfiler();
 
-	if (IsDefaultContext())
+	if (IsDefaultContext() && !IsAsyncComputeContext())
 	{
 		GPUProfiler.PopEvent();
 	}
