Use ParallelForWithExistingTaskContext() in FD3D12CommandContext::RHISetRayTracingBindings() instead of adhoc tasks load-balancing scheme

This avoids a pathological task scheduling case where RHI thread may process all the work before the helper tasks even spawn.
Regular ParallelFor handles this case, but previously it did not have a solid way to guarantee a maximum number of worker tasks which is needed here.

diff --git a/Engine/Source/Runtime/D3D12RHI/Private/D3D12RayTracing.cpp b/Engine/Source/Runtime/D3D12RHI/Private/D3D12RayTracing.cpp
--- a/Engine/Source/Runtime/D3D12RHI/Private/D3D12RayTracing.cpp
+++ b/Engine/Source/Runtime/D3D12RHI/Private/D3D12RayTracing.cpp
@@ -5628,88 +5628,71 @@ void FD3D12CommandContext::RHISetRayTracingBindings(
 		? FMath::Min<uint32>(NumWorkerThreads, FD3D12RayTracingScene::MaxBindingWorkers)
 		: 1;
 
-	const uint32 MinItemsPerTask = 1024u;
-	const uint32 NumTasks = FMath::Min(MaxTasks, FMath::Max(1u, FMath::DivideAndRoundUp(NumBindings, MinItemsPerTask)));
-
-	volatile int32 SharedCounter = 0;
-	const int32 IncrementCount = 128;
-
-	auto BindingTask = [Bindings, Device, ShaderTable, Scene, Pipeline, &SharedCounter, IncrementCount, NumBindings, BindingType](uint32 WorkerIndex)
+	struct FTaskContext
 	{
-		TRACE_CPUPROFILER_EVENT_SCOPE(SetRayTracingBindings_Task);
-	 
-		int32 BaseIndex = 0;
-		while (BaseIndex < (int32)NumBindings)
-		{
-			BaseIndex = FPlatformAtomics::InterlockedAdd(&SharedCounter, IncrementCount);
-			int32 EndIndex = FMath::Min(BaseIndex + IncrementCount, (int32)NumBindings);
-
-			for (int32 CurrentIndex = BaseIndex; CurrentIndex < EndIndex; ++CurrentIndex)
-			{
-				const FRayTracingLocalShaderBindings& Binding = Bindings[CurrentIndex];
-
-				if (BindingType == ERayTracingBindingType::HitGroup)
-				{
-					SetRayTracingHitGroup(Device, ShaderTable, Scene, Pipeline,
-						Binding.InstanceIndex,
-						Binding.SegmentIndex,
-						Binding.ShaderSlot,
-						Binding.ShaderIndexInPipeline,
-						Binding.NumUniformBuffers,
-						Binding.UniformBuffers,
-						Binding.LooseParameterDataSize,
-						Binding.LooseParameterData,
-						Binding.UserData,
-						WorkerIndex);
-				}
-				else if (BindingType == ERayTracingBindingType::CallableShader)
-				{
-					SetRayTracingCallableShader(Device, ShaderTable, Scene, Pipeline,
-						Binding.ShaderSlot,
-						Binding.ShaderIndexInPipeline,
-						Binding.NumUniformBuffers,
-						Binding.UniformBuffers,
-						Binding.LooseParameterDataSize,
-						Binding.LooseParameterData,
-						Binding.UserData,
-						WorkerIndex);
-				}
-				else if (BindingType == ERayTracingBindingType::MissShader)
-				{
-					SetRayTracingMissShader(Device, ShaderTable, Scene, Pipeline,
-						Binding.ShaderSlot,
-						Binding.ShaderIndexInPipeline,
-						Binding.NumUniformBuffers,
-						Binding.UniformBuffers,
-						Binding.LooseParameterDataSize,
-						Binding.LooseParameterData,
-						Binding.UserData,
-						WorkerIndex);
-				}
-				else
-				{
-					checkNoEntry();
-				}
-			}
-		}
+		uint32 WorkerIndex = 0;
 	};
 
-	TaskList.Reserve(NumTasks);
+	TArray<FTaskContext, TInlineAllocator<FD3D12RayTracingScene::MaxBindingWorkers>> TaskContexts;
+	for (uint32 WorkerIndex = 0; WorkerIndex < MaxTasks; ++WorkerIndex)
+	{
+		TaskContexts.Add(FTaskContext{WorkerIndex});
+	}
 
-	// Create and kick any parallel tasks first and run task 0 on this thread
-	for (uint32 TaskIndex = 1; TaskIndex < NumTasks; ++TaskIndex)
+	auto BindingTask = [Bindings, Device, ShaderTable, Scene, Pipeline, BindingType](const FTaskContext& Context, int32 CurrentIndex)
 	{
-		TaskList.Add(FFunctionGraphTask::CreateAndDispatchWhenReady(
-			[TaskIndex, &BindingTask]()
+		const FRayTracingLocalShaderBindings& Binding = Bindings[CurrentIndex];
+
+		if (BindingType == ERayTracingBindingType::HitGroup)
 		{
-			BindingTask(TaskIndex);
-		}, TStatId(), nullptr, ENamedThreads::AnyThread));
-	}
+			SetRayTracingHitGroup(Device, ShaderTable, Scene, Pipeline,
+				Binding.InstanceIndex,
+				Binding.SegmentIndex,
+				Binding.ShaderSlot,
+				Binding.ShaderIndexInPipeline,
+				Binding.NumUniformBuffers,
+				Binding.UniformBuffers,
+				Binding.LooseParameterDataSize,
+				Binding.LooseParameterData,
+				Binding.UserData,
+				Context.WorkerIndex);
+		}
+		else if (BindingType == ERayTracingBindingType::CallableShader)
+		{
+			SetRayTracingCallableShader(Device, ShaderTable, Scene, Pipeline,
+				Binding.ShaderSlot,
+				Binding.ShaderIndexInPipeline,
+				Binding.NumUniformBuffers,
+				Binding.UniformBuffers,
+				Binding.LooseParameterDataSize,
+				Binding.LooseParameterData,
+				Binding.UserData,
+				Context.WorkerIndex);
+		}
+		else if (BindingType == ERayTracingBindingType::MissShader)
+		{
+			SetRayTracingMissShader(Device, ShaderTable, Scene, Pipeline,
+				Binding.ShaderSlot,
+				Binding.ShaderIndexInPipeline,
+				Binding.NumUniformBuffers,
+				Binding.UniformBuffers,
+				Binding.LooseParameterDataSize,
+				Binding.LooseParameterData,
+				Binding.UserData,
+				Context.WorkerIndex);
+		}
+		else
+		{
+			checkNoEntry();
+		}
+	};
 
-	// Run task 0 on this thread, in parallel with any workers
-	BindingTask(0);
+	// One helper worker task will be created at most per this many work items, plus one worker for current thread (unless running on a task thread),
+	// up to a hard maximum of FD3D12RayTracingScene::MaxBindingWorkers.
+	// Internally, parallel for tasks still subdivide the work into smaller chunks and perform fine-grained load-balancing.
+	const int32 ItemsPerTask = 1024;
 
-	FTaskGraphInterface::Get().WaitUntilTasksComplete(TaskList, ENamedThreads::AnyThread);
+	ParallelForWithExistingTaskContext(TEXT("SetRayTracingBindings"), MakeArrayView(TaskContexts), NumBindings, ItemsPerTask, BindingTask);
 }
 
 void FD3D12CommandContext::RHISetRayTracingHitGroup(
