DX12 CopyBufferRegion: Fixed a crash that happens when two buffers are copied and upload heap state is changed to D3D12_RESOURCE_STATE_COPY_SOURCE from D3D12_RESOURCE_STATE_GENERIC_READ.
Since D3D12_RESOURCE_STATE_GENERIC_READ already includes D3D12_RESOURCE_STATE_COPY_SOURCE we don't need resource transition in that case. Otherwise the transition causes a hard crash in DX12.
If both Before and after states are the same, we shouldn't be catching ensure messages.
In addition to that CopyBuffer method was checking lower level buffer width instead of wrapper property Size. Some buffers occupy an offset region in a bigger buffer therefore we should check the size not the lower level property - Width and operate on the expected buffer area.

diff --git a/Engine/Source/Runtime/D3D12RHI/Private/D3D12Buffer.cpp b/Engine/Source/Runtime/D3D12RHI/Private/D3D12Buffer.cpp
--- a/Engine/Source/Runtime/D3D12RHI/Private/D3D12Buffer.cpp
+++ b/Engine/Source/Runtime/D3D12RHI/Private/D3D12Buffer.cpp
@@ -822,11 +822,19 @@ void FD3D12DynamicRHI::RHICopyBuffer(FRHIBuffer* SourceBufferRHI, FRHIBuffer* De
 		FD3D12Resource* pDestResource = DestBuffer->ResourceLocation.GetResource();
 		D3D12_RESOURCE_DESC const& DestBufferDesc = pDestResource->GetDesc();
 
-		check(SourceBufferDesc.Width == DestBufferDesc.Width);
+		check(SourceBuffer->GetSize() == DestBuffer->GetSize());
 
 		FD3D12CommandContext& Context = Device->GetDefaultCommandContext();
 		Context.numCopies++;
-		Context.CommandListHandle->CopyResource(pDestResource->GetResource(), pSourceResource->GetResource());
+
+		// The underlying D3D12 buffer can be larger than the RHI buffer due to pooling.
+		Context.CommandListHandle->CopyBufferRegion(
+			pDestResource->GetResource(), 
+			DestBuffer->ResourceLocation.GetOffsetFromBaseOfResource(), 
+			pSourceResource->GetResource(), 
+			SourceBuffer->ResourceLocation.GetOffsetFromBaseOfResource(), 
+			SourceBufferRHI->GetSize());
+
 		Context.CommandListHandle.UpdateResidency(pDestResource);
 		Context.CommandListHandle.UpdateResidency(pSourceResource);
 
diff --git a/Engine/Source/Runtime/D3D12RHI/Private/D3D12RHIPrivate.h b/Engine/Source/Runtime/D3D12RHI/Private/D3D12RHIPrivate.h
--- a/Engine/Source/Runtime/D3D12RHI/Private/D3D12RHIPrivate.h
+++ b/Engine/Source/Runtime/D3D12RHI/Private/D3D12RHIPrivate.h
@@ -1343,7 +1343,8 @@ public:
 		if (!pResource->RequiresResourceStateTracking())
 		{
 			CurrentState = pResource->GetDefaultResourceState();
-			if (CurrentState != DesiredState)
+			// Some states such as D3D12_RESOURCE_STATE_GENERIC_READ already includes D3D12_RESOURCE_STATE_COPY_SOURCE as well as other states, therefore transition isn't required.
+			if (CurrentState != DesiredState && !EnumHasAllFlags(CurrentState, InDesiredState))
 			{
 				// we will add a transition, we need to transition back to the default state when the scoped object dies : 
 				bRestoreState = true;
