ERHIFeatureLevel::SM6 should require DirectX Feature Level 12_0, not DirectX Feature Level 12_1.

diff --git a/Engine/Source/Runtime/D3D12RHI/Private/Windows/WindowsD3D12Device.cpp b/Engine/Source/Runtime/D3D12RHI/Private/Windows/WindowsD3D12Device.cpp
--- a/Engine/Source/Runtime/D3D12RHI/Private/Windows/WindowsD3D12Device.cpp
+++ b/Engine/Source/Runtime/D3D12RHI/Private/Windows/WindowsD3D12Device.cpp
@@ -302,7 +302,7 @@ inline ERHIFeatureLevel::Type FindMaxRHIFeatureLevel(ID3D12Device* Device, D3D_F
 {
 	ERHIFeatureLevel::Type MaxRHIFeatureLevel = ERHIFeatureLevel::Num;
 
-	if (InMaxFeatureLevel >= D3D_FEATURE_LEVEL_12_1 && InMaxShaderModel >= D3D_SHADER_MODEL_6_6)
+	if (InMaxFeatureLevel >= D3D_FEATURE_LEVEL_12_0 && InMaxShaderModel >= D3D_SHADER_MODEL_6_6)
 	{
 		D3D12_FEATURE_DATA_D3D12_OPTIONS1 D3D12Caps1{};
 		Device->CheckFeatureSupport(D3D12_FEATURE_D3D12_OPTIONS1, &D3D12Caps1, sizeof(D3D12Caps1));
