Remove support for PLATFORM_USE_BACKBUFFER_WRITE_TRANSITION_TRACKING

- previously disabled due to change in platform behavior
- removing dead code with this change

- testing
  - ran ShooterGame on platform which previously used the system and on Win64 which did not

#rnx

diff --git a/Engine/Source/Runtime/Core/Public/HAL/Platform.h b/Engine/Source/Runtime/Core/Public/HAL/Platform.h
--- a/Engine/Source/Runtime/Core/Public/HAL/Platform.h
+++ b/Engine/Source/Runtime/Core/Public/HAL/Platform.h
@@ -606,10 +606,6 @@
 	#define PLATFORM_SUPPORTS_COLORIZED_OUTPUT_DEVICE PLATFORM_DESKTOP
 #endif
 
-#ifndef PLATFORM_USE_BACKBUFFER_WRITE_TRANSITION_TRACKING
-	#define PLATFORM_USE_BACKBUFFER_WRITE_TRANSITION_TRACKING 0
-#endif
-
 #ifndef PLATFORM_USE_SEPARATE_BACKBUFFER_WRITE_TRANSITION
 	#define PLATFORM_USE_SEPARATE_BACKBUFFER_WRITE_TRANSITION 0
 #endif
diff --git a/Engine/Source/Runtime/D3D12RHI/Private/D3D12CommandContext.cpp b/Engine/Source/Runtime/D3D12RHI/Private/D3D12CommandContext.cpp
--- a/Engine/Source/Runtime/D3D12RHI/Private/D3D12CommandContext.cpp
+++ b/Engine/Source/Runtime/D3D12RHI/Private/D3D12CommandContext.cpp
@@ -504,11 +504,6 @@ void FD3D12CommandContextBase::RHIBeginFrame()
 		if (ensure(ContextAtIndex))
 		{
 			Device->GetTimestampQueryHeap(CommandQueueType)->EndQueryBatchAndResolveQueryData(*ContextAtIndex);
-#if PLATFORM_USE_BACKBUFFER_WRITE_TRANSITION_TRACKING
-			uint64 TimeStampFrequency = 0;
-			VERIFYD3D12RESULT(Device->GetCommandListManager(CommandQueueType)->GetTimestampFrequency(&TimeStampFrequency));
-			Device->GetBackBufferWriteBarrierTracker()->ResolveBatches(TimeStampFrequency, false);
-#endif // #if PLATFORM_USE_BACKBUFFER_WRITE_TRANSITION_TRACKING
 		}
 
 		FD3D12GlobalOnlineSamplerHeap& SamplerHeap = Device->GetGlobalSamplerHeap();
@@ -619,11 +614,6 @@ void FD3D12CommandContextBase::RHIEndFrame()
 
 		FD3D12CommandContext& DefaultContext = Device->GetDefaultCommandContext();
 		DefaultContext.CommandListHandle.FlushResourceBarriers();
-
-#if PLATFORM_USE_BACKBUFFER_WRITE_TRANSITION_TRACKING
-		Device->GetBackBufferWriteBarrierTracker()->EndBatch(DefaultContext);
-#endif // #if PLATFORM_USE_BACKBUFFER_WRITE_TRANSITION_TRACKING
-
 		DefaultContext.ReleaseCommandAllocator();
 		DefaultContext.ClearState();
 		DefaultContext.FlushCommands();
@@ -792,13 +782,6 @@ void FD3D12CommandContext::RHIEndScene()
 	bDrawingScene = false;
 }
 
-#if PLATFORM_USE_BACKBUFFER_WRITE_TRANSITION_TRACKING
-void FD3D12CommandContext::RHIBackBufferWaitTrackingBeginFrame(uint64 FrameToken, bool bDeferred)
-{
-	GetParentDevice()->GetBackBufferWriteBarrierTracker()->BeginBatch(FrameToken, bDeferred);
-}
-#endif // #if PLATFORM_USE_BACKBUFFER_WRITE_TRANSITION_TRACKING
-
 #if D3D12_SUPPORTS_PARALLEL_RHI_EXECUTE
 
 //todo recycle these to avoid alloc
diff --git a/Engine/Source/Runtime/D3D12RHI/Private/D3D12CommandContext.h b/Engine/Source/Runtime/D3D12RHI/Private/D3D12CommandContext.h
--- a/Engine/Source/Runtime/D3D12RHI/Private/D3D12CommandContext.h
+++ b/Engine/Source/Runtime/D3D12RHI/Private/D3D12CommandContext.h
@@ -324,10 +324,6 @@ public:
 #endif
 	virtual void RHISetDepthBounds(float MinDepth, float MaxDepth) final override;
     virtual void RHISetShadingRate(EVRSShadingRate ShadingRate, EVRSRateCombiner Combiner) final override;
-#if PLATFORM_USE_BACKBUFFER_WRITE_TRANSITION_TRACKING
-	virtual void RHIBackBufferWaitTrackingBeginFrame(uint64 FrameToken, bool bDeferred) final override;
-#endif // #if PLATFORM_USE_BACKBUFFER_WRITE_TRANSITION_TRACKING
-
 	virtual void RHIClearMRTImpl(bool* bClearColorArray, int32 NumClearColors, const FLinearColor* ColorArray, bool bClearDepth, float Depth, bool bClearStencil, uint32 Stencil);
 
 
@@ -844,13 +840,6 @@ public:
 		ContextRedirect(RHIBindAccelerationStructureMemory(Scene, Buffer, BufferOffset));
 	}
 
-#if PLATFORM_USE_BACKBUFFER_WRITE_TRANSITION_TRACKING
-	virtual void RHIBackBufferWaitTrackingBeginFrame(uint64 FrameToken, bool bDeferred) final override
-	{
-		ContextRedirect(RHIBackBufferWaitTrackingBeginFrame(FrameToken, bDeferred));
-	}
-#endif // #if PLATFORM_USE_BACKBUFFER_WRITE_TRANSITION_TRACKING
-
 	FORCEINLINE void SetPhysicalContext(FD3D12CommandContext* Context)
 	{
 		check(Context);
diff --git a/Engine/Source/Runtime/D3D12RHI/Private/D3D12CommandList.cpp b/Engine/Source/Runtime/D3D12RHI/Private/D3D12CommandList.cpp
--- a/Engine/Source/Runtime/D3D12RHI/Private/D3D12CommandList.cpp
+++ b/Engine/Source/Runtime/D3D12RHI/Private/D3D12CommandList.cpp
@@ -197,13 +197,6 @@ void FD3D12CommandListHandle::FD3D12CommandListData::FlushResourceBarriers()
 	{
 		ResourceBarriers.Append(Barriers.GetData(), Barriers.Num());
 	}
-#if PLATFORM_USE_BACKBUFFER_WRITE_TRANSITION_TRACKING
-	const TArray<D3D12_RESOURCE_BARRIER>& BackBufferBarriers = ResourceBarrierBatcher.GetBackBufferBarriers();
-	if (BackBufferBarriers.Num())
-	{
-		ResourceBarriers.Append(BackBufferBarriers.GetData(), BackBufferBarriers.Num());
-	}
-#endif // #if PLATFORM_USE_BACKBUFFER_WRITE_TRANSITION_TRACKING
 #endif // #if DEBUG_RESOURCE_STATES
 
 	ResourceBarrierBatcher.Flush(GetParentDevice(), CommandList, FD3D12DynamicRHI::GetResourceBarrierBatchSizeLimit());
@@ -234,9 +227,6 @@ void FD3D12CommandListHandle::FD3D12CommandListData::Reset(FD3D12CommandAllocato
 
 	// If this fails then some previous resource barriers were never submitted.
 	check(ResourceBarrierBatcher.GetBarriers().Num() == 0);
-#if PLATFORM_USE_BACKBUFFER_WRITE_TRANSITION_TRACKING
-	check(ResourceBarrierBatcher.GetBackBufferBarriers().Num() == 0);
-#endif // #if PLATFORM_USE_BACKBUFFER_WRITE_TRANSITION_TRACKING
 
 #if DEBUG_RESOURCE_STATES
 	ResourceBarriers.Reset();
diff --git a/Engine/Source/Runtime/D3D12RHI/Private/D3D12Device.cpp b/Engine/Source/Runtime/D3D12RHI/Private/D3D12Device.cpp
--- a/Engine/Source/Runtime/D3D12RHI/Private/D3D12Device.cpp
+++ b/Engine/Source/Runtime/D3D12RHI/Private/D3D12Device.cpp
@@ -63,9 +63,6 @@ FD3D12Device::~FD3D12Device()
 	delete CommandListManager;
 	delete CopyCommandListManager;
 	delete AsyncCommandListManager;
-#if PLATFORM_USE_BACKBUFFER_WRITE_TRANSITION_TRACKING
-	delete BackBufferWriteBarrierTracker;
-#endif // #if PLATFORM_USE_BACKBUFFER_WRITE_TRANSITION_TRACKING
 }
 
 ID3D12Device* FD3D12Device::GetDevice()
diff --git a/Engine/Source/Runtime/D3D12RHI/Private/D3D12Device.h b/Engine/Source/Runtime/D3D12RHI/Private/D3D12Device.h
--- a/Engine/Source/Runtime/D3D12RHI/Private/D3D12Device.h
+++ b/Engine/Source/Runtime/D3D12RHI/Private/D3D12Device.h
@@ -117,9 +117,6 @@ public:
 	inline FD3D12FastAllocator& GetDefaultFastAllocator() { return DefaultFastAllocator; }
 	inline FD3D12TextureAllocatorPool& GetTextureAllocator() { return TextureAllocator; }
 	inline FD3D12ResidencyManager& GetResidencyManager() { return ResidencyManager; }
-#if PLATFORM_USE_BACKBUFFER_WRITE_TRANSITION_TRACKING
-	inline FD3D12TimedIntervalQueryTracker*	GetBackBufferWriteBarrierTracker() { return BackBufferWriteBarrierTracker; }
-#endif // #if PLATFORM_USE_BACKBUFFER_WRITE_TRANSITION_TRACKING
 
 	TArray<FD3D12CommandListHandle> PendingCommandLists;
 
@@ -207,9 +204,6 @@ protected:
 	FD3D12RayTracingDescriptorHeapCache* RayTracingDescriptorHeapCache = nullptr;
 	void DestroyRayTracingDescriptorCache();
 #endif
-#if PLATFORM_USE_BACKBUFFER_WRITE_TRANSITION_TRACKING
-	FD3D12TimedIntervalQueryTracker* BackBufferWriteBarrierTracker = nullptr;
-#endif // #if PLATFORM_USE_BACKBUFFER_WRITE_TRANSITION_TRACKING
 
 	D3D12RHI::FD3DGPUProfiler GPUProfilingData;
 };
diff --git a/Engine/Source/Runtime/D3D12RHI/Private/D3D12DirectCommandListManager.cpp b/Engine/Source/Runtime/D3D12RHI/Private/D3D12DirectCommandListManager.cpp
--- a/Engine/Source/Runtime/D3D12RHI/Private/D3D12DirectCommandListManager.cpp
+++ b/Engine/Source/Runtime/D3D12RHI/Private/D3D12DirectCommandListManager.cpp
@@ -1019,39 +1019,12 @@ uint32 FD3D12CommandListManager::GetResourceBarrierCommandList(FBarrierDescInfo&
 		LogResourceBarriers(InBarrierDescInfo.BackBufferBarrierDescs.Num(), InBarrierDescInfo.BackBufferBarrierDescs.GetData(), hResourceBarrierList.CommandList());
 #endif // #if DEBUG_RESOURCE_STATES
 
-		const int32 BarrierBatchMax = FD3D12DynamicRHI::GetResourceBarrierBatchSizeLimit();
-
-#if PLATFORM_USE_BACKBUFFER_WRITE_TRANSITION_TRACKING
-		if (InBarrierDescInfo.BackBufferBarrierDescs.Num())
-		{
-			FD3D12ScopedTimedIntervalQuery BarrierScopeTimer(GetParentDevice()->GetBackBufferWriteBarrierTracker(), hResourceBarrierList.GraphicsCommandList());
-			if (InBarrierDescInfo.BackBufferBarrierDescs.Num() > BarrierBatchMax)
-			{
-				int Num = InBarrierDescInfo.BackBufferBarrierDescs.Num();
-				D3D12_RESOURCE_BARRIER* Ptr = InBarrierDescInfo.BackBufferBarrierDescs.GetData();
-				while (Num > 0)
-				{
-					const int DispatchNum = FMath::Min(Num, BarrierBatchMax);
-					hResourceBarrierList->ResourceBarrier(DispatchNum, Ptr);
-					Ptr += BarrierBatchMax;
-					Num -= BarrierBatchMax;
-				}
-			}
-			else
-			{
-				hResourceBarrierList->ResourceBarrier(InBarrierDescInfo.BackBufferBarrierDescs.Num(), InBarrierDescInfo.BackBufferBarrierDescs.GetData());
-			}
-		}
-
-		if (InBarrierDescInfo.BarrierDescs.Num())
-#else
 		InBarrierDescInfo.BarrierDescs.Append(InBarrierDescInfo.BackBufferBarrierDescs);
 		InBarrierDescInfo.BackBufferBarrierDescs.Empty();
-#endif // #if PLATFORM_USE_BACKBUFFER_WRITE_TRANSITION_TRACKING
-		{
-			extern void ResourceBarriersSeparateRTV2SRV(ID3D12GraphicsCommandList*, const TArray<D3D12_RESOURCE_BARRIER>&, int32);
-			ResourceBarriersSeparateRTV2SRV(hResourceBarrierList.GraphicsCommandList(), InBarrierDescInfo.BarrierDescs, BarrierBatchMax);
-		}
+
+		const int32 BarrierBatchMax = FD3D12DynamicRHI::GetResourceBarrierBatchSizeLimit();
+		extern void ResourceBarriersSeparateRTV2SRV(ID3D12GraphicsCommandList*, const TArray<D3D12_RESOURCE_BARRIER>&, int32);
+		ResourceBarriersSeparateRTV2SRV(hResourceBarrierList.GraphicsCommandList(), InBarrierDescInfo.BarrierDescs, BarrierBatchMax);
 	}
 
 	return BarrierCount;
diff --git a/Engine/Source/Runtime/D3D12RHI/Private/D3D12RHI.cpp b/Engine/Source/Runtime/D3D12RHI/Private/D3D12RHI.cpp
--- a/Engine/Source/Runtime/D3D12RHI/Private/D3D12RHI.cpp
+++ b/Engine/Source/Runtime/D3D12RHI/Private/D3D12RHI.cpp
@@ -747,7 +747,7 @@ uint64 FD3D12SubmissionGapRecorder::SubmitSubmissionTimestampsForFrame(uint32 Fr
 	{
 		static TConsoleVariableData<int32>* VSyncIntervalCVar = IConsoleManager::Get().FindTConsoleVariableDataInt(TEXT("rhi.syncinterval"));
 
-		if (VSyncIntervalCVar && VSyncIntervalCVar->GetValueOnRenderThread() > 0 && !PLATFORM_USE_BACKBUFFER_WRITE_TRANSITION_TRACKING)
+		if (VSyncIntervalCVar && VSyncIntervalCVar->GetValueOnRenderThread() > 0)
 		{
 			int32 offset = PrevFrameBeginSubmissionTimestamps.Num() - (EndFrameSlotIdx - (PresentSlotIdx + 2));
 			if (PrevFrameBeginSubmissionTimestamps.IsValidIndex(offset))
diff --git a/Engine/Source/Runtime/D3D12RHI/Private/D3D12RHIPrivate.h b/Engine/Source/Runtime/D3D12RHI/Private/D3D12RHIPrivate.h
--- a/Engine/Source/Runtime/D3D12RHI/Private/D3D12RHIPrivate.h
+++ b/Engine/Source/Runtime/D3D12RHI/Private/D3D12RHIPrivate.h
@@ -76,7 +76,6 @@ extern int32 GDX12NVAfterMathMarkers;
 #include "../Public/D3D12Viewport.h"
 #include "../Public/D3D12ConstantBuffer.h"
 #include "D3D12Query.h"
-#include "D3D12TimedIntervalQuery.h"
 #include "D3D12DescriptorCache.h"
 #include "D3D12StateCachePrivate.h"
 typedef FD3D12StateCacheBase FD3D12StateCache;
diff --git a/Engine/Source/Runtime/D3D12RHI/Private/D3D12Resources.cpp b/Engine/Source/Runtime/D3D12RHI/Private/D3D12Resources.cpp
--- a/Engine/Source/Runtime/D3D12RHI/Private/D3D12Resources.cpp
+++ b/Engine/Source/Runtime/D3D12RHI/Private/D3D12Resources.cpp
@@ -1267,11 +1267,6 @@ static void RecordResourceBarriersToCommandList(
 	}
 }
 
-#if (PLATFORM_USE_BACKBUFFER_WRITE_TRANSITION_TRACKING == 0) && (PLATFORM_USE_SEPARATE_BACKBUFFER_WRITE_TRANSITION == 1)
-#define LOCAL_USE_SEPARATE_BACKBUFFER_WRITE_TRANSITION 1
-#else
-#define LOCAL_USE_SEPARATE_BACKBUFFER_WRITE_TRANSITION 0
-#endif
 
 void ResourceBarriersSeparateRTV2SRV(
 	ID3D12GraphicsCommandList* pCommandList,
@@ -1280,7 +1275,7 @@ void ResourceBarriersSeparateRTV2SRV(
 {
 	if (!GD3D12SeparateRTV2SRVTransitions)
 	{
-#if LOCAL_USE_SEPARATE_BACKBUFFER_WRITE_TRANSITION
+#if PLATFORM_USE_SEPARATE_BACKBUFFER_WRITE_TRANSITION
 		TArray<D3D12_RESOURCE_BARRIER, TInlineAllocator<4>> BackBufferBarriers;
 		TArray<D3D12_RESOURCE_BARRIER, TInlineAllocator<8>> OtherBarriers;
 
@@ -1315,7 +1310,7 @@ void ResourceBarriersSeparateRTV2SRV(
 	else
 	{
 		TArray<D3D12_RESOURCE_BARRIER, TInlineAllocator<4>> RTV2SRVBarriers;
-#if LOCAL_USE_SEPARATE_BACKBUFFER_WRITE_TRANSITION
+#if PLATFORM_USE_SEPARATE_BACKBUFFER_WRITE_TRANSITION
 		TArray<D3D12_RESOURCE_BARRIER, TInlineAllocator<4>> BackBufferBarriers;
 #endif
 		TArray<D3D12_RESOURCE_BARRIER, TInlineAllocator<8>> OtherBarriers;
@@ -1330,7 +1325,7 @@ void ResourceBarriersSeparateRTV2SRV(
 			{
 				RTV2SRVBarriers.Add(Barrier);
 			}
-#if LOCAL_USE_SEPARATE_BACKBUFFER_WRITE_TRANSITION
+#if PLATFORM_USE_SEPARATE_BACKBUFFER_WRITE_TRANSITION
 			else if (Barrier.Type == D3D12_RESOURCE_BARRIER_TYPE_TRANSITION
 				&& Barrier.Transition.StateBefore == D3D12_RESOURCE_STATE_PRESENT // can also be displayed as D3D12_RESOURCE_STATE_COMMON in pix
 				&& Barrier.Transition.StateAfter == D3D12_RESOURCE_STATE_RENDER_TARGET)
@@ -1349,7 +1344,7 @@ void ResourceBarriersSeparateRTV2SRV(
 			RecordResourceBarriersToCommandList(pCommandList, RTV2SRVBarriers.GetData(), RTV2SRVBarriers.Num(), BarrierBatchMax);
 		}
 
-#if LOCAL_USE_SEPARATE_BACKBUFFER_WRITE_TRANSITION
+#if PLATFORM_USE_SEPARATE_BACKBUFFER_WRITE_TRANSITION
 		if (BackBufferBarriers.Num() > 0)
 		{
 			RecordResourceBarriersToCommandList(pCommandList, BackBufferBarriers.GetData(), BackBufferBarriers.Num(), BarrierBatchMax);
@@ -1361,7 +1356,7 @@ void ResourceBarriersSeparateRTV2SRV(
 		}
 	}
 }
-#undef LOCAL_USE_SEPARATE_BACKBUFFER_WRITE_TRANSITION
+
 
 void FD3D12ResourceBarrierBatcher::Flush(FD3D12Device* Device, ID3D12GraphicsCommandList* pCommandList, int32 BarrierBatchMax)
 {
@@ -1371,15 +1366,6 @@ void FD3D12ResourceBarrierBatcher::Flush(FD3D12Device* Device, ID3D12GraphicsCom
 		ResourceBarriersSeparateRTV2SRV(pCommandList, Barriers, BarrierBatchMax);
 	}
 
-#if PLATFORM_USE_BACKBUFFER_WRITE_TRANSITION_TRACKING
-	if (BackBufferBarriers.Num())
-	{
-		check(pCommandList);
-		FD3D12ScopedTimedIntervalQuery BarrierScopeTimer(Device->GetBackBufferWriteBarrierTracker(), pCommandList);
-		RecordResourceBarriersToCommandList(pCommandList, BackBufferBarriers.GetData(), BackBufferBarriers.Num(), BarrierBatchMax);
-	}
-#endif // #if PLATFORM_USE_BACKBUFFER_WRITE_TRANSITION_TRACKING
-
 	Reset();
 }
 
diff --git a/Engine/Source/Runtime/D3D12RHI/Private/D3D12TimedIntervalQuery.cpp b/Engine/Source/Runtime/D3D12RHI/Private/D3D12TimedIntervalQuery.cpp
deleted file mode 100644
diff --git a/Engine/Source/Runtime/D3D12RHI/Private/D3D12TimedIntervalQuery.h b/Engine/Source/Runtime/D3D12RHI/Private/D3D12TimedIntervalQuery.h
deleted file mode 100644
diff --git a/Engine/Source/Runtime/D3D12RHI/Public/D3D12Resources.h b/Engine/Source/Runtime/D3D12RHI/Public/D3D12Resources.h
--- a/Engine/Source/Runtime/D3D12RHI/Public/D3D12Resources.h
+++ b/Engine/Source/Runtime/D3D12RHI/Public/D3D12Resources.h
@@ -1008,17 +1008,7 @@ public:
 		}
 
 		check(IsValidD3D12ResourceState(Before) && IsValidD3D12ResourceState(After));
-
-#if PLATFORM_USE_BACKBUFFER_WRITE_TRANSITION_TRACKING
-		if (pResource->IsBackBuffer() && EnumHasAnyFlags(After, BackBufferBarrierWriteTransitionTargets))
-		{
-			AddTransitionBarrier(BackBufferBarriers, pResource, Before, After, Subresource);
-		}
-		else
-#endif
-		{
-			AddTransitionBarrier(Barriers, pResource, Before, After, Subresource);
-		}
+		AddTransitionBarrier(Barriers, pResource, Before, After, Subresource);
 
 		return 1;
 	}
@@ -1042,10 +1032,6 @@ public:
 		// Reset the arrays without shrinking (Does not destruct items, does not de-allocate memory).
 		Barriers.SetNumUnsafeInternal(0);
 		check(Barriers.Num() == 0);
-#if PLATFORM_USE_BACKBUFFER_WRITE_TRANSITION_TRACKING
-		BackBufferBarriers.SetNumUnsafeInternal(0);
-		check(BackBufferBarriers.Num() == 0);
-#endif // #if PLATFORM_USE_BACKBUFFER_WRITE_TRANSITION_TRACKING
 	}
 
 	const TArray<D3D12_RESOURCE_BARRIER>& GetBarriers() const
@@ -1053,18 +1039,8 @@ public:
 		return Barriers;
 	}
 
-#if PLATFORM_USE_BACKBUFFER_WRITE_TRANSITION_TRACKING
-	const TArray<D3D12_RESOURCE_BARRIER>& GetBackBufferBarriers() const
-	{
-		return BackBufferBarriers;
-	}
-#endif // #if PLATFORM_USE_BACKBUFFER_WRITE_TRANSITION_TRACKING
-
 private:
 	TArray<D3D12_RESOURCE_BARRIER> Barriers;
-#if PLATFORM_USE_BACKBUFFER_WRITE_TRANSITION_TRACKING
-	TArray<D3D12_RESOURCE_BARRIER> BackBufferBarriers;
-#endif // #if PLATFORM_USE_BACKBUFFER_WRITE_TRANSITION_TRACKING
 };
 
 class FD3D12StagingBuffer final : public FRHIStagingBuffer
diff --git a/Engine/Source/Runtime/RHI/Public/RHICommandList.h b/Engine/Source/Runtime/RHI/Public/RHICommandList.h
--- a/Engine/Source/Runtime/RHI/Public/RHICommandList.h
+++ b/Engine/Source/Runtime/RHI/Public/RHICommandList.h
@@ -2183,21 +2183,6 @@ FRHICOMMAND_MACRO(FRHICommandUpdateRHIResources)
 	RHI_API void Execute(FRHICommandListBase& CmdList);
 };
 
-#if PLATFORM_USE_BACKBUFFER_WRITE_TRANSITION_TRACKING
-FRHICOMMAND_MACRO(FRHICommandBackBufferWaitTrackingBeginFrame)
-{
-	uint64	FrameToken;
-	bool	bDeferred;
-
-	FORCEINLINE_DEBUGGABLE FRHICommandBackBufferWaitTrackingBeginFrame(uint64 FrameTokenIn, bool bDeferredIn)
-		:	FrameToken(FrameTokenIn),
-			bDeferred(bDeferredIn)
-	{}
-	
-	RHI_API void Execute(FRHICommandListBase& CmdList);
-};
-#endif // #if PLATFORM_USE_BACKBUFFER_WRITE_TRANSITION_TRACKING
-
 FRHICOMMAND_MACRO(FRHICommandCopyBufferRegion)
 {
 	FRHIBuffer* DestBuffer;
@@ -3771,18 +3756,6 @@ public:
 		}
 		ALLOC_COMMAND(FRHICommandDiscardRenderTargets)(Depth, Stencil, ColorBitMask);
 	}
-
-#if PLATFORM_USE_BACKBUFFER_WRITE_TRANSITION_TRACKING
-	FORCEINLINE_DEBUGGABLE void RHIBackBufferWaitTrackingBeginFrame(uint64 FrameToken, bool bDeferred)
-	{
-		if (Bypass())
-		{
-			GetContext().RHIBackBufferWaitTrackingBeginFrame(FrameToken, bDeferred);
-			return;
-		}
-		ALLOC_COMMAND(FRHICommandBackBufferWaitTrackingBeginFrame)(FrameToken, bDeferred);
-	}
-#endif // #if PLATFORM_USE_BACKBUFFER_WRITE_TRANSITION_TRACKING
 	
 	FORCEINLINE_DEBUGGABLE void CopyBufferRegion(FRHIBuffer* DestBuffer, uint64 DstOffset, FRHIBuffer* SourceBuffer, uint64 SrcOffset, uint64 NumBytes)
 	{
diff --git a/Engine/Source/Runtime/RHI/Public/RHICommandListCommandExecutes.inl b/Engine/Source/Runtime/RHI/Public/RHICommandListCommandExecutes.inl
--- a/Engine/Source/Runtime/RHI/Public/RHICommandListCommandExecutes.inl
+++ b/Engine/Source/Runtime/RHI/Public/RHICommandListCommandExecutes.inl
@@ -783,11 +783,3 @@ void FRHICommandDiscardRenderTargets::Execute(FRHICommandListBase& CmdList)
 	RHISTAT(RHIDiscardRenderTargets);
 	INTERNAL_DECORATOR(RHIDiscardRenderTargets)(Depth, Stencil, ColorBitMask);
 }
-
-#if PLATFORM_USE_BACKBUFFER_WRITE_TRANSITION_TRACKING
-void FRHICommandBackBufferWaitTrackingBeginFrame::Execute(FRHICommandListBase& CmdList)
-{
-	RHISTAT(FRHICommandBackBufferWaitTrackingBeginFrame);
-	INTERNAL_DECORATOR(RHIBackBufferWaitTrackingBeginFrame)(FrameToken, bDeferred);
-}
-#endif // #if PLATFORM_USE_BACKBUFFER_WRITE_TRANSITION_TRACKING
diff --git a/Engine/Source/Runtime/RHI/Public/RHIContext.h b/Engine/Source/Runtime/RHI/Public/RHIContext.h
--- a/Engine/Source/Runtime/RHI/Public/RHIContext.h
+++ b/Engine/Source/Runtime/RHI/Public/RHIContext.h
@@ -834,13 +834,6 @@ public:
 		checkNoEntry();
 	}
 
-#if PLATFORM_USE_BACKBUFFER_WRITE_TRANSITION_TRACKING
-	virtual void RHIBackBufferWaitTrackingBeginFrame(uint64 FrameToken, bool bDeferred)
-	{
-		checkNoEntry();
-	}
-#endif // #if PLATFORM_USE_BACKBUFFER_WRITE_TRANSITION_TRACKING
-
 	protected:
 		FRHIRenderPassInfo RenderPassInfo;
 };
diff --git a/Engine/Source/Runtime/RHI/Public/RHIValidationContext.h b/Engine/Source/Runtime/RHI/Public/RHIValidationContext.h
--- a/Engine/Source/Runtime/RHI/Public/RHIValidationContext.h
+++ b/Engine/Source/Runtime/RHI/Public/RHIValidationContext.h
@@ -1088,13 +1088,6 @@ public:
 		RHIContext->RHISetRayTracingMissShader(Scene, ShaderSlotInScene, Pipeline, ShaderIndexInPipeline, NumUniformBuffers, UniformBuffers, UserData);
 	}
 
-#if PLATFORM_USE_BACKBUFFER_WRITE_TRANSITION_TRACKING
-	virtual void RHIBackBufferWaitTrackingBeginFrame(uint64 FrameToken, bool bDeferred) override final
-	{
-		RHIContext->RHIBackBufferWaitTrackingBeginFrame(FrameToken, bDeferred);
-	}
-#endif // #if PLATFORM_USE_BACKBUFFER_WRITE_TRANSITION_TRACKING
-
 	void SetupDrawing()
 	{
 		// nothing to validate right now
