MGPU:  Trivial fix for spurious assert when -rhivalidation is enabled.

Need to do a second level of command context indirection, to potentially go from FD3D12CommandContextRedirector to FD3D12CommandContext.  The issue arises because GetLowestLevelContext on FValidationContext returns the contained FD3D12CommandContextRedirector, and doesn't attempt to call GetLowestLevelContext on that to get the physical FD3D12CommandContext required by ExecuteOnCommandContext.  Additional comments in code...

#rnx

diff --git a/Engine/Source/Runtime/D3D12RHI/Private/D3D12Buffer.cpp b/Engine/Source/Runtime/D3D12RHI/Private/D3D12Buffer.cpp
--- a/Engine/Source/Runtime/D3D12RHI/Private/D3D12Buffer.cpp
+++ b/Engine/Source/Runtime/D3D12RHI/Private/D3D12Buffer.cpp
@@ -98,7 +98,21 @@ struct FD3D12RHICommandInitializeBuffer final : public FRHICommand<FD3D12RHIComm
 	void Execute(FRHICommandListBase& CmdList)
 	{
 		FD3D12CommandContext& CommandContext = (FD3D12CommandContext&)(CmdList.IsImmediateAsyncCompute() ? CmdList.GetComputeContext().GetLowestLevelContext() : CmdList.GetContext().GetLowestLevelContext());
+#if ENABLE_RHI_VALIDATION && WITH_MGPU
+		// Need to do a second level of indirection, to potentially go from FD3D12CommandContextRedirector to FD3D12CommandContext, when
+		// -rhivalidation is enabled.  The issue arises because GetLowestLevelContext on FValidationContext returns the contained
+		// FD3D12CommandContextRedirector, and doesn't attempt to call GetLowestLevelContext on that to get the physical FD3D12CommandContext
+		// required here.  Only when both validation and MGPU are present is it possible for this second level of indirection to be required.
+		// The call to "GetLowestLevelContext" is a nop (returns *this) if already a leaf command context.
+		//
+		// I thought about changing the original FValidationContext::GetLowestLevelContext implementation to attempt a second indirection there,
+		// but I wasn't sure what other side effects there might be for that change.  Here, the change is extremely safe, because the
+		// CommandContext passed to ExecuteOnCommandContext is solely used for a validation assert when WITH_MGPU is true, not actual rendering.
+		// So it can't affect behavior, it's just to avoid a spurious assert.
+		ExecuteOnCommandContext((FD3D12CommandContext&)CommandContext.GetLowestLevelContext());
+#else
 		ExecuteOnCommandContext(CommandContext);
+#endif
 	}
 
 	void ExecuteOnCommandContext(FD3D12CommandContext& CommandContext)
