Fix RaytracingTestbed test:
* Add miss shader to default shader tables to be used by RayTraceOcclusion/Intersection
* Compile CS for building instance buffer if RT is supported.

diff --git a/Engine/Source/Runtime/D3D12RHI/Private/D3D12RayTracing.cpp b/Engine/Source/Runtime/D3D12RHI/Private/D3D12RayTracing.cpp
--- a/Engine/Source/Runtime/D3D12RHI/Private/D3D12RayTracing.cpp
+++ b/Engine/Source/Runtime/D3D12RHI/Private/D3D12RayTracing.cpp
@@ -2745,6 +2745,11 @@ public:
 				NodeDevice,
 				RayGenShaders.Identifiers,
 				HitGroupShaders.Identifiers[0]);
+
+			if (MissShaders.Identifiers.Num())
+			{
+				DefaultShaderTables[GPUIndex].SetMissIdentifier(0, MissShaders.Identifiers[0]);
+			}			
 		}
 
 		PipelineStackSize = PipelineProperties->GetPipelineStackSize();
@@ -2798,6 +2803,7 @@ public:
 
 	// Shader table that can be used to dispatch ray tracing work that doesn't require real SBT bindings.
 	// This is useful for the case where user only provides default RayGen, Miss and HitGroup shaders.
+	// Currently only used by RayTraceOcclusion/Intersection.
 	FD3D12RayTracingShaderTable DefaultShaderTables[MAX_NUM_GPUS];
 
 	ID3D12RootSignature* GlobalRootSignature = nullptr;
diff --git a/Engine/Source/Runtime/Renderer/Private/RayTracing/RayTracingInstanceBufferUtil.cpp b/Engine/Source/Runtime/Renderer/Private/RayTracing/RayTracingInstanceBufferUtil.cpp
--- a/Engine/Source/Runtime/Renderer/Private/RayTracing/RayTracingInstanceBufferUtil.cpp
+++ b/Engine/Source/Runtime/Renderer/Private/RayTracing/RayTracingInstanceBufferUtil.cpp
@@ -312,7 +312,7 @@ struct FRayTracingBuildInstanceBufferCS : public FGlobalShader
 
 	static bool ShouldCompilePermutation(const FGlobalShaderPermutationParameters& Parameters)
 	{
-		return IsRayTracingEnabledForProject(Parameters.Platform) && RHISupportsComputeShaders(Parameters.Platform);
+		return RHISupportsRayTracing(Parameters.Platform) && RHISupportsComputeShaders(Parameters.Platform);
 	}
 };
 
