Use a fault-tolerant method to query ray tracing shader IDs from partial RTPSOs.

* Some DXR implementations (such as VKD3D) may not support shader ID queries from collections. If shader ID query fails, we will attempt to query it fromt he final linked RTPSO later.

diff --git a/Engine/Source/Runtime/D3D12RHI/Private/D3D12RayTracing.cpp b/Engine/Source/Runtime/D3D12RHI/Private/D3D12RayTracing.cpp
--- a/Engine/Source/Runtime/D3D12RHI/Private/D3D12RayTracing.cpp
+++ b/Engine/Source/Runtime/D3D12RHI/Private/D3D12RayTracing.cpp
@@ -714,6 +714,25 @@ static FD3D12ShaderIdentifier GetShaderIdentifier(ID3D12StateObject* StateObject
 	return GetShaderIdentifier(PipelineProperties, ExportName);
 }
 
+static bool TryGetShaderIdentifier(ID3D12StateObject* StateObject, const TCHAR* ExportName, FD3D12ShaderIdentifier& OutShaderIdentifier)
+{
+	TRefCountPtr<ID3D12StateObjectProperties> PipelineProperties;
+	HRESULT QueryInterfaceResult = StateObject->QueryInterface(IID_PPV_ARGS(PipelineProperties.GetInitReference()));
+	if (!SUCCEEDED(QueryInterfaceResult))
+	{
+		return false;
+	}
+	const void* ShaderIdData = PipelineProperties->GetShaderIdentifier(ExportName);
+	if (ShaderIdData)
+	{
+		OutShaderIdentifier.SetData(ShaderIdData);
+		return true;
+	}
+	else
+	{
+		return false;
+	}
+}
 
 FD3D12RayTracingCompactionRequestHandler::FD3D12RayTracingCompactionRequestHandler(FD3D12Device* Device) : FD3D12DeviceChild(Device), PostBuildInfoBufferReadbackFence(0)
 {
@@ -1097,7 +1116,9 @@ public:
 			// Therefore shader identifiers need to be queried from the final linked pipeline (JIRA DH-2182) if a known bad driver is detected.
 			if (GD3D12WorkaroundFlags.bAllowGetShaderIdentifierOnCollectionSubObject)
 			{
-				Entry.Identifier = GetShaderIdentifier(Entry.StateObject, Entry.GetPrimaryExportNameChars());
+				// Some DXR implementations (such as VKD3D) may not support shader ID queries from collections.
+				// If shader ID query fails, we will attempt to query it from the final linked RTPSO later.
+				TryGetShaderIdentifier(Entry.StateObject, Entry.GetPrimaryExportNameChars(), Entry.Identifier);
 			}
 
 			CompileTimeCycles += FPlatformTime::Cycles64();
