VRS: D3D12RHI updates - Defer ShadingRateImage binding to State Cache

Previously, we would bind the shading rate image + related combiners as part of SetRenderTargetsAndClear.
Instead, VRS images can now use the state cache infrastructure to reduce binding overhead.

Inspiration CL shared by [REDACTED_PERSON] (The Coalition)

diff --git a/Engine/Source/Runtime/D3D12RHI/Private/D3D12CommandContext.h b/Engine/Source/Runtime/D3D12RHI/Private/D3D12CommandContext.h
--- a/Engine/Source/Runtime/D3D12RHI/Private/D3D12CommandContext.h
+++ b/Engine/Source/Runtime/D3D12RHI/Private/D3D12CommandContext.h
@@ -174,6 +174,7 @@ public:
 
 	/** Used by variable rate shading to cache the current state of the combiners and the constant shading rate*/
 #if PLATFORM_SUPPORTS_VARIABLE_RATE_SHADING
+	static_assert(D3D12_RS_SET_SHADING_RATE_COMBINER_COUNT == ED3D12VRSCombinerStages::Num);
 	D3D12_SHADING_RATE_COMBINER		VRSCombiners[D3D12_RS_SET_SHADING_RATE_COMBINER_COUNT];
 	D3D12_SHADING_RATE				VRSShadingRate;
 #endif
@@ -259,7 +260,8 @@ public:
 		);
 
 	virtual void SetDepthBounds(float MinDepth, float MaxDepth);
-	virtual void SetShadingRate(EVRSShadingRate ShadingRate, EVRSRateCombiner Combiner);
+	virtual void SetShadingRate(EVRSShadingRate ShadingRate, const TStaticArray<EVRSRateCombiner, ED3D12VRSCombinerStages::Num>& Combiners);
+	virtual void SetShadingRateImage(FD3D12Resource* ShadingRateImage);
 
 	virtual void SetAsyncComputeBudgetInternal(EAsyncComputeBudget Budget) {}
 
diff --git a/Engine/Source/Runtime/D3D12RHI/Private/D3D12Commands.cpp b/Engine/Source/Runtime/D3D12RHI/Private/D3D12Commands.cpp
--- a/Engine/Source/Runtime/D3D12RHI/Private/D3D12Commands.cpp
+++ b/Engine/Source/Runtime/D3D12RHI/Private/D3D12Commands.cpp
@@ -1451,29 +1451,35 @@ void FD3D12CommandContext::SetRenderTargetsAndClear(const FRHISetRenderTargetsIn
 	}
 
 #if PLATFORM_SUPPORTS_VARIABLE_RATE_SHADING
-	// If we support tier 2, we will support tier 1.
-	if (GRHISupportsPipelineVariableRateShading && GRHIVariableRateShadingEnabled && CommandListHandle.GraphicsCommandList5() != nullptr)
+	if (GRHISupportsPipelineVariableRateShading && GRHIVariableRateShadingEnabled)
 	{
 		if (GRHISupportsAttachmentVariableRateShading && GRHIAttachmentVariableRateShadingEnabled)
 		{
-			VRSCombiners[1] = ConvertShadingRateCombiner(RenderTargetsInfo.ShadingRateTextureCombiner); // Combiner 1 is used to mix rates from a texture and the previous combiner
 			if (RenderTargetsInfo.ShadingRateTexture != nullptr)
 			{
 				FD3D12Resource* Resource = RetrieveTexture(RenderTargetsInfo.ShadingRateTexture)->GetResource();
-				CommandListHandle.GraphicsCommandList5()->RSSetShadingRateImage(Resource->GetResource());
+
+				FD3D12DynamicRHI::TransitionResource(
+					CommandListHandle,
+					Resource,
+					D3D12_RESOURCE_STATE_TBD,
+					D3D12_RESOURCE_STATE_SHADING_RATE_SOURCE,
+					D3D12_RESOURCE_BARRIER_ALL_SUBRESOURCES,
+					FD3D12DynamicRHI::ETransitionMode::Validate);
+
+
+				StateCache.SetShadingRateImage(Resource, RenderTargetsInfo.ShadingRateTextureCombiner);
 			}
 			else
 			{
-				CommandListHandle.GraphicsCommandList5()->RSSetShadingRateImage(nullptr);
+				StateCache.SetShadingRateImage(nullptr, EVRSRateCombiner::VRSRB_Passthrough);
 			}
 		}
 		else
 		{
 			// Ensure this is set appropriate if image-based VRS not supported or not enabled.
-			VRSCombiners[1] = D3D12_SHADING_RATE_COMBINER_PASSTHROUGH;
+			StateCache.SetShadingRateImage(nullptr, EVRSRateCombiner::VRSRB_Passthrough);
 		}
-
-		CommandListHandle.GraphicsCommandList5()->RSSetShadingRate(VRSShadingRate, VRSCombiners);
 	}
 #endif
 }
@@ -2421,19 +2427,42 @@ void FD3D12CommandContext::RHISetShadingRate(EVRSShadingRate ShadingRate, EVRSRa
 #endif
 }
 
- void FD3D12CommandContext::SetShadingRate(EVRSShadingRate ShadingRate, EVRSRateCombiner Combiner)
- {
+void FD3D12CommandContext::SetShadingRate(EVRSShadingRate ShadingRate, const TStaticArray<EVRSRateCombiner, ED3D12VRSCombinerStages::Num>& Combiners)
+{
  #if PLATFORM_SUPPORTS_VARIABLE_RATE_SHADING
- 	if (GRHISupportsPipelineVariableRateShading && GRHIVariableRateShadingEnabled && CommandListHandle.GraphicsCommandList5())
+ 	if (GRHISupportsPipelineVariableRateShading && GRHIVariableRateShadingEnabled && CommandListHandle.GraphicsCommandList5() != nullptr)
  	{
- 		VRSCombiners[0] = ConvertShadingRateCombiner(Combiner);	// Combiner 0 is used to mix per draw and per VS/GS rates
+		for (int32 CombinerIndex = 0; CombinerIndex < Combiners.Num(); ++CombinerIndex)
+		{
+			VRSCombiners[CombinerIndex] = ConvertShadingRateCombiner(Combiners[CombinerIndex]);
+		}
  		VRSShadingRate = static_cast<D3D12_SHADING_RATE>(ShadingRate);
  		CommandListHandle.GraphicsCommandList5()->RSSetShadingRate(VRSShadingRate, VRSCombiners);
  	}
 	else
 	{
 		// Ensure we're at a reasonable default in the case we're not supporting VRS.
-		VRSCombiners[0] = VRSCombiners[1] = D3D12_SHADING_RATE_COMBINER_PASSTHROUGH;
+		for (int32 CombinerIndex = 0; CombinerIndex < Combiners.Num(); ++CombinerIndex)
+		{
+			VRSCombiners[CombinerIndex] = D3D12_SHADING_RATE_COMBINER_PASSTHROUGH;
+		}
+	}
+ #endif
+}
+
+void FD3D12CommandContext::SetShadingRateImage(FD3D12Resource* RateImageTexture)
+{
+#if PLATFORM_SUPPORTS_VARIABLE_RATE_SHADING
+	if (GRHISupportsAttachmentVariableRateShading && GRHIAttachmentVariableRateShadingEnabled && CommandListHandle.GraphicsCommandList5() != nullptr)
+	{
+		if (RateImageTexture)
+		{
+			CommandListHandle.GraphicsCommandList5()->RSSetShadingRateImage(RateImageTexture->GetResource());
+		}
+		else
+		{
+			CommandListHandle.GraphicsCommandList5()->RSSetShadingRateImage(nullptr);
+		}
 	}
  #endif
 }
diff --git a/Engine/Source/Runtime/D3D12RHI/Private/D3D12StateCache.cpp b/Engine/Source/Runtime/D3D12RHI/Private/D3D12StateCache.cpp
--- a/Engine/Source/Runtime/D3D12RHI/Private/D3D12StateCache.cpp
+++ b/Engine/Source/Runtime/D3D12RHI/Private/D3D12StateCache.cpp
@@ -253,8 +253,12 @@ void FD3D12StateCacheBase::ClearState()
 
 	PipelineState.Graphics.MinDepth = 0.0f;
 	PipelineState.Graphics.MaxDepth = 1.0f;
-	PipelineState.Graphics.Combiner = EVRSRateCombiner::VRSRB_Passthrough;
+	for (auto& Combiner : PipelineState.Graphics.Combiners)
+	{
+		Combiner = EVRSRateCombiner::VRSRB_Passthrough;
+	}
 	PipelineState.Graphics.DrawShadingRate = EVRSShadingRate::VRSSR_1x1;
+	PipelineState.Graphics.ShadingRateImage = nullptr;
 }
 
 void FD3D12StateCacheBase::DirtyStateForNewCommandList()
@@ -298,6 +302,8 @@ void FD3D12StateCacheBase::DirtyStateForNewCommandList()
 	
 	bNeedSetShadingRate = GRHISupportsPipelineVariableRateShading && GRHIVariableRateShadingEnabled;
 	
+	bNeedSetShadingRateImage = GRHISupportsAttachmentVariableRateShading && GRHIAttachmentVariableRateShadingEnabled;
+
 	// Always dirty View and Sampler bindings. We detect the slots that are actually used at Draw/Dispatch time.
 	PipelineState.Common.SRVCache.DirtyAll();
 	PipelineState.Common.UAVCache.DirtyAll();
@@ -321,6 +327,7 @@ void FD3D12StateCacheBase::DirtyState()
 	bNeedSetStencilRef = true;
 	bNeedSetDepthBounds = GSupportsDepthBoundsTest;
 	bNeedSetShadingRate = GRHISupportsPipelineVariableRateShading && GRHIVariableRateShadingEnabled;
+	bNeedSetShadingRateImage = GRHISupportsAttachmentVariableRateShading && GRHIAttachmentVariableRateShadingEnabled;
 	PipelineState.Common.SRVCache.DirtyAll();
 	PipelineState.Common.UAVCache.DirtyAll();
 	PipelineState.Common.CBVCache.DirtyAll();
@@ -584,7 +591,13 @@ void FD3D12StateCacheBase::ApplyState()
 		if (bNeedSetShadingRate)
 		{
 			bNeedSetShadingRate = false;
-			CmdContext->SetShadingRate(PipelineState.Graphics.DrawShadingRate, PipelineState.Graphics.Combiner);
+			CmdContext->SetShadingRate(PipelineState.Graphics.DrawShadingRate, PipelineState.Graphics.Combiners);
+		}
+
+		if (bNeedSetShadingRateImage)
+		{
+			bNeedSetShadingRateImage = false;
+			CmdContext->SetShadingRateImage(PipelineState.Graphics.ShadingRateImage);
 		}
 	}
 
diff --git a/Engine/Source/Runtime/D3D12RHI/Private/D3D12StateCachePrivate.h b/Engine/Source/Runtime/D3D12RHI/Private/D3D12StateCachePrivate.h
--- a/Engine/Source/Runtime/D3D12RHI/Private/D3D12StateCachePrivate.h
+++ b/Engine/Source/Runtime/D3D12RHI/Private/D3D12StateCachePrivate.h
@@ -57,6 +57,12 @@ enum class ED3D12PipelineType : uint8
 	RayTracing,
 };
 
+namespace ED3D12VRSCombinerStages
+{
+	constexpr int32 PerPrimitive	= 0;
+	constexpr int32 ScreenSpace		= PerPrimitive + 1;
+	constexpr int32 Num				= ScreenSpace + 1;
+};
 
 #define MAX_VBS			D3D12_IA_VERTEX_INPUT_RESOURCE_SLOT_COUNT
 
@@ -335,6 +341,7 @@ protected:
 	bool bNeedSetStencilRef;
 	bool bNeedSetDepthBounds;
 	bool bNeedSetShadingRate;
+	bool bNeedSetShadingRateImage;
 	D3D12_RESOURCE_BINDING_TIER ResourceBindingTier;
 
 	struct
@@ -392,7 +399,10 @@ protected:
 			float MaxDepth;
 
 			EVRSShadingRate  DrawShadingRate;
-			EVRSRateCombiner Combiner;
+
+			TStaticArray<EVRSRateCombiner, ED3D12VRSCombinerStages::Num> Combiners;
+
+			FD3D12Resource*  ShadingRateImage;
 		} Graphics;
 
 		struct
@@ -815,14 +825,32 @@ public:
 
 	void SetShadingRate(EVRSShadingRate ShadingRate, EVRSRateCombiner Combiner)
 	{
-		if (PipelineState.Graphics.DrawShadingRate != ShadingRate || PipelineState.Graphics.Combiner != Combiner)
+		if (PipelineState.Graphics.DrawShadingRate != ShadingRate || PipelineState.Graphics.Combiners[ED3D12VRSCombinerStages::PerPrimitive] != Combiner)
 		{
 			PipelineState.Graphics.DrawShadingRate = ShadingRate;
-			PipelineState.Graphics.Combiner = Combiner;
+			PipelineState.Graphics.Combiners[ED3D12VRSCombinerStages::PerPrimitive] = Combiner;
 			bNeedSetShadingRate = GRHISupportsPipelineVariableRateShading && GRHIVariableRateShadingEnabled;
 		}
 	}
 
+	void SetShadingRateImage(FD3D12Resource* ShadingRateImage, EVRSRateCombiner Combiner)
+	{
+		if (PipelineState.Graphics.ShadingRateImage != ShadingRateImage)
+		{
+			PipelineState.Graphics.ShadingRateImage = ShadingRateImage;
+			bNeedSetShadingRateImage = GRHISupportsAttachmentVariableRateShading && GRHIAttachmentVariableRateShadingEnabled;
+		}
+
+		// If we aren't provided a ShadingRateImage, we should request the passthrough combiner
+		ensure((PipelineState.Graphics.ShadingRateImage != nullptr) || (Combiner == EVRSRateCombiner::VRSRB_Passthrough));
+
+		if (PipelineState.Graphics.Combiners[ED3D12VRSCombinerStages::ScreenSpace] != Combiner)
+		{
+			PipelineState.Graphics.Combiners[ED3D12VRSCombinerStages::ScreenSpace] = Combiner;
+			bNeedSetShadingRate = GRHISupportsAttachmentVariableRateShading && GRHIAttachmentVariableRateShadingEnabled;
+		}
+	}
+
 	void SetComputeBudget(EAsyncComputeBudget ComputeBudget)
 	{
 		PipelineState.Compute.ComputeBudget = ComputeBudget;
diff --git a/Engine/Source/Runtime/D3D12RHI/Private/D3D12Util.cpp b/Engine/Source/Runtime/D3D12RHI/Private/D3D12Util.cpp
--- a/Engine/Source/Runtime/D3D12RHI/Private/D3D12Util.cpp
+++ b/Engine/Source/Runtime/D3D12RHI/Private/D3D12Util.cpp
@@ -1400,7 +1400,7 @@ D3D12_RESOURCE_STATES GetD3D12ResourceState(ERHIAccess InRHIAccess, bool InIsAsy
 				State |= D3D12_RESOURCE_STATE_DEPTH_READ;
 			}
 #if PLATFORM_SUPPORTS_VARIABLE_RATE_SHADING
-			if (EnumHasAnyFlags(InRHIAccess, ERHIAccess::ShadingRateSource) && GRHISupportsPipelineVariableRateShading)
+			if (EnumHasAnyFlags(InRHIAccess, ERHIAccess::ShadingRateSource) && GRHISupportsAttachmentVariableRateShading)
 			{
 				State |= D3D12_RESOURCE_STATE_SHADING_RATE_SOURCE;
 			}
