Require DXR Tier 1.1 for ray tracing and use shader model 6.5 for all ray tracing shaders in SM5 mode.

diff --git a/Engine/Source/Developer/Windows/ShaderFormatD3D/Private/D3DShaderCompiler.cpp b/Engine/Source/Developer/Windows/ShaderFormatD3D/Private/D3DShaderCompiler.cpp
--- a/Engine/Source/Developer/Windows/ShaderFormatD3D/Private/D3DShaderCompiler.cpp
+++ b/Engine/Source/Developer/Windows/ShaderFormatD3D/Private/D3DShaderCompiler.cpp
@@ -129,7 +129,7 @@ static const TCHAR* GetShaderProfileName(FShaderTarget Target, bool bForceSM6)
 		case SF_RayMiss:
 		case SF_RayHitGroup:
 		case SF_RayCallable:
-			return TEXT("lib_6_3");
+			return TEXT("lib_6_5");
 		}
 	}
 	else if ((Target.Platform == SP_PCD3D_ES3_1) || (Target.Platform == SP_D3D_ES3_1_HOLOLENS))
diff --git a/Engine/Source/Runtime/D3D12RHI/Private/D3D12Adapter.cpp b/Engine/Source/Runtime/D3D12RHI/Private/D3D12Adapter.cpp
--- a/Engine/Source/Runtime/D3D12RHI/Private/D3D12Adapter.cpp
+++ b/Engine/Source/Runtime/D3D12RHI/Private/D3D12Adapter.cpp
@@ -1040,19 +1040,21 @@ void FD3D12Adapter::InitializeDevices()
 					&& FDataDrivenShaderPlatformInfo::GetSupportsRayTracing(GMaxRHIShaderPlatform)
 					&& !FParse::Param(FCommandLine::Get(), TEXT("noraytracing")))
 				{
-					UE_LOG(LogD3D12RHI, Log, TEXT("D3D12 ray tracing tier 1.0 is supported."));
-
-					GRHISupportsRayTracing = RHISupportsRayTracing(GMaxRHIShaderPlatform);
-					GRHISupportsRayTracingShaders = GRHISupportsRayTracing && RHISupportsRayTracingShaders(GMaxRHIShaderPlatform);
-
 					if (D3D12Caps5.RaytracingTier >= D3D12_RAYTRACING_TIER_1_1
 						&& RootDevice7)
 					{
 						UE_LOG(LogD3D12RHI, Log, TEXT("D3D12 ray tracing tier 1.1 is supported."));
 
+						GRHISupportsRayTracing = RHISupportsRayTracing(GMaxRHIShaderPlatform);
+						GRHISupportsRayTracingShaders = GRHISupportsRayTracing && RHISupportsRayTracingShaders(GMaxRHIShaderPlatform);
+
 						GRHISupportsRayTracingPSOAdditions = true;
 						GRHISupportsInlineRayTracing = GRHISupportsRayTracing && RHISupportsInlineRayTracing(GMaxRHIShaderPlatform) && (GMaxRHIFeatureLevel >= ERHIFeatureLevel::SM6);
 					}
+					else
+					{
+						UE_LOG(LogD3D12RHI, Log, TEXT("Ray tracing is disabled because D3D12 ray tracing tier 1.1 is required but only tier 1.0 is supported."));
+					}
 				}
 				else if (D3D12Caps5.RaytracingTier == D3D12_RAYTRACING_TIER_NOT_SUPPORTED 
 					&& bRenderDocPresent
