Fixing DescriptorCache handling of Bindless heaps on new frames.

diff --git a/Engine/Source/Runtime/D3D12RHI/Private/D3D12DescriptorCache.cpp b/Engine/Source/Runtime/D3D12RHI/Private/D3D12DescriptorCache.cpp
--- a/Engine/Source/Runtime/D3D12RHI/Private/D3D12DescriptorCache.cpp
+++ b/Engine/Source/Runtime/D3D12RHI/Private/D3D12DescriptorCache.cpp
@@ -800,8 +800,10 @@ bool FD3D12DescriptorCache::SwitchToContextLocalSamplerHeap()
 bool FD3D12DescriptorCache::SwitchToGlobalSamplerHeap()
 {
 	bool bDescriptorHeapsChanged = false;
-	if (!UsingGlobalSamplerHeap())
+	// TODO: should we be clearing the override here?
+	if (OverrideSamplerHeap || !UsingGlobalSamplerHeap())
 	{
+		OverrideSamplerHeap = nullptr;
 		bUsingGlobalSamplerHeap = true;
 		CurrentSamplerHeap = &GetParentDevice()->GetGlobalSamplerHeap();
 		bDescriptorHeapsChanged = SetDescriptorHeaps();
