Fixed: Broken SkinCache in a mobile preview (D3D12 only)

diff --git a/Engine/Source/Runtime/D3D12RHI/Private/D3D12Buffer.cpp b/Engine/Source/Runtime/D3D12RHI/Private/D3D12Buffer.cpp
--- a/Engine/Source/Runtime/D3D12RHI/Private/D3D12Buffer.cpp
+++ b/Engine/Source/Runtime/D3D12RHI/Private/D3D12Buffer.cpp
@@ -500,13 +500,6 @@ void FD3D12Buffer::GetResourceDescAndAlignment(uint64 InSize, uint32 InStride, E
 	if (EnumHasAnyFlags(InUsage, BUF_UnorderedAccess))
 	{
 		ResourceDesc.Flags |= D3D12_RESOURCE_FLAG_ALLOW_UNORDERED_ACCESS;
-
-		static bool bRequiresRawView = (GMaxRHIFeatureLevel < ERHIFeatureLevel::SM5);
-		if (bRequiresRawView)
-		{
-			// Force the buffer to be a raw, byte address buffer
-			InUsage |= BUF_ByteAddressBuffer;
-		}
 	}
 
 	if (!EnumHasAnyFlags(InUsage, BUF_ShaderResource))
