When binding a shader with zero SamplerStates and bindless samplers enabled, asserts were triggering because FD3D12StateCacheBase::ApplySamplers was being called when it shouldn't.
- When counting shader sampler/resource usage to determine the number of descriptors in the Root Signature, set the count to 0 instead of MAX when the shader has zero.
- When binding/switching PSOs/RootSigs, only apply samplers/resources when the shader (read: RootSig) has any to apply.

diff --git a/Engine/Source/Runtime/D3D12RHI/Private/D3D12RootSignature.h b/Engine/Source/Runtime/D3D12RHI/Private/D3D12RootSignature.h
--- a/Engine/Source/Runtime/D3D12RHI/Private/D3D12RootSignature.h
+++ b/Engine/Source/Runtime/D3D12RHI/Private/D3D12RootSignature.h
@@ -183,6 +183,7 @@ public:
 	inline bool HasUAVs() const { return bHasUAVs; }
 	inline bool HasSRVs() const { return bHasSRVs; }
 	inline bool HasCBVs() const { return bHasCBVs; }
+	inline bool HasResources() const { return bHasUAVs || bHasSRVs || bHasCBVs; }
 	inline bool HasSamplers() const { return bHasSamplers; }
 	inline bool UsesDynamicResources() const { return bUsesDynamicResources; }
 	inline bool UsesDynamicSamplers() const { return bUsesDynamicSamplers; }
diff --git a/Engine/Source/Runtime/D3D12RHI/Private/D3D12StateCache.cpp b/Engine/Source/Runtime/D3D12RHI/Private/D3D12StateCache.cpp
--- a/Engine/Source/Runtime/D3D12RHI/Private/D3D12StateCache.cpp
+++ b/Engine/Source/Runtime/D3D12RHI/Private/D3D12StateCache.cpp
@@ -510,14 +510,14 @@ void FD3D12StateCacheBase::ApplyState()
 	// Ensure the correct graphics PSO is set.
 	InternalSetPipelineState(PSOCommonData->PipelineState);
 
-	bool bBindlessResources = false;
-	bool bBindlessSamplers = false;
+	bool bBindlessResources = PSOCommonData->RootSignature->UsesDynamicResources();
+	bool bBindlessSamplers = PSOCommonData->RootSignature->UsesDynamicSamplers();
+
+	const bool bApplyResources = !bBindlessResources && PSOCommonData->RootSignature->HasResources();
+	const bool bApplySamplers = !bBindlessSamplers && PSOCommonData->RootSignature->HasSamplers();
 
 	if (bRootSignatureChanged)
 	{
-		bBindlessResources = PSOCommonData->RootSignature->UsesDynamicResources();
-		bBindlessSamplers = PSOCommonData->RootSignature->UsesDynamicSamplers();
-		
 		FD3D12BindlessDescriptorManager& BindlessManager = GetParentDevice()->GetBindlessDescriptorManager();
 
 		FD3D12DescriptorHeap* ResourceHeap = BindlessManager.GetHeapForType(ERHIDescriptorHeapType::Standard);
@@ -611,12 +611,12 @@ void FD3D12StateCacheBase::ApplyState()
 	//
 
 	// Samplers
-	if (!bBindlessSamplers)
+	if (bApplySamplers)
 	{
 		ApplySamplers(PSOCommonData->RootSignature, StartStage, EndStage);
 	}
 
-	if (!bBindlessResources)
+	if (bApplyResources)
 	{
 		ApplyResources(PSOCommonData->RootSignature, StartStage, EndStage);
 	}
diff --git a/Engine/Source/Runtime/D3D12RHI/Private/D3D12Util.cpp b/Engine/Source/Runtime/D3D12RHI/Private/D3D12Util.cpp
--- a/Engine/Source/Runtime/D3D12RHI/Private/D3D12Util.cpp
+++ b/Engine/Source/Runtime/D3D12RHI/Private/D3D12Util.cpp
@@ -985,8 +985,8 @@ void FD3D12QuantizedBoundShaderState::InitShaderRegisterCounts(const D3D12_RESOU
 	}
 	else
 	{
-		Shader.SamplerCount = MaxSamplerCount;
-		Shader.ShaderResourceCount = MaxShaderResourceCount;
+		Shader.SamplerCount = Counts.NumSamplers > 0 ? MaxSamplerCount : 0;
+		Shader.ShaderResourceCount = Counts.NumSRVs > 0 ? MaxShaderResourceCount : 0;
 	}
 
 	if (ResourceBindingTier <= D3D12_RESOURCE_BINDING_TIER_2)
@@ -997,7 +997,7 @@ void FD3D12QuantizedBoundShaderState::InitShaderRegisterCounts(const D3D12_RESOU
 	else
 	{
 		Shader.ConstantBufferCount = (Counts.NumCBs > MAX_ROOT_CBVS) ? MaxConstantBufferCount : Counts.NumCBs;
-		Shader.UnorderedAccessCount = (bAllowUAVs) ? MaxUnorderedAccessCount : 0;
+		Shader.UnorderedAccessCount = Counts.NumUAVs > 0 && bAllowUAVs ? MaxUnorderedAccessCount : 0;
 	}
 }
 
