Fixed bad bIs24BitUnormDepthStencil on some platforms to fix  manual depth test.

diff --git a/Engine/Source/Runtime/Apple/MetalRHI/Private/MetalRHI.cpp b/Engine/Source/Runtime/Apple/MetalRHI/Private/MetalRHI.cpp
--- a/Engine/Source/Runtime/Apple/MetalRHI/Private/MetalRHI.cpp
+++ b/Engine/Source/Runtime/Apple/MetalRHI/Private/MetalRHI.cpp
@@ -880,6 +880,7 @@ FMetalDynamicRHI::FMetalDynamicRHI(ERHIFeatureLevel::Type RequestedFeatureLevel)
 	else
 	{
 		GPixelFormats[PF_DepthStencil	].PlatformFormat	= (uint32)mtlpp::PixelFormat::Depth32Float_Stencil8;
+		GPixelFormats[PF_DepthStencil	].bIs24BitUnormDepthStencil = false;
 	}
 	GPixelFormats[PF_DepthStencil		].BlockBytes		= 4;
 	GPixelFormats[PF_DepthStencil		].Supported			= true;
diff --git a/Engine/Source/Runtime/D3D12RHI/Private/D3D12RHI.cpp b/Engine/Source/Runtime/D3D12RHI/Private/D3D12RHI.cpp
--- a/Engine/Source/Runtime/D3D12RHI/Private/D3D12RHI.cpp
+++ b/Engine/Source/Runtime/D3D12RHI/Private/D3D12RHI.cpp
@@ -118,6 +118,7 @@ FD3D12DynamicRHI::FD3D12DynamicRHI(const TArray<TSharedPtr<FD3D12Adapter>>& Chos
 		GPixelFormats[PF_DepthStencil].PlatformFormat = DXGI_FORMAT_R32G8X24_TYPELESS;
 		GPixelFormats[PF_DepthStencil].BlockBytes = 5;
 		GPixelFormats[PF_DepthStencil].Supported = true;
+		GPixelFormats[PF_DepthStencil].bIs24BitUnormDepthStencil = false;
 		GPixelFormats[PF_X24_G8].PlatformFormat = DXGI_FORMAT_X32_TYPELESS_G8X24_UINT;
 		GPixelFormats[PF_X24_G8].BlockBytes = 5;
 	}
diff --git a/Engine/Source/Runtime/Windows/D3D11RHI/Private/D3D11Device.cpp b/Engine/Source/Runtime/Windows/D3D11RHI/Private/D3D11Device.cpp
--- a/Engine/Source/Runtime/Windows/D3D11RHI/Private/D3D11Device.cpp
+++ b/Engine/Source/Runtime/Windows/D3D11RHI/Private/D3D11Device.cpp
@@ -129,6 +129,7 @@ FD3D11DynamicRHI::FD3D11DynamicRHI(IDXGIFactory1* InDXGIFactory1, D3D_FEATURE_LE
 	{
 		GPixelFormats[PF_DepthStencil].PlatformFormat = DXGI_FORMAT_R32G8X24_TYPELESS;
 		GPixelFormats[PF_DepthStencil].BlockBytes = 5;
+		GPixelFormats[PF_DepthStencil].bIs24BitUnormDepthStencil = false;
 		GPixelFormats[PF_X24_G8].PlatformFormat = DXGI_FORMAT_X32_TYPELESS_G8X24_UINT;
 		GPixelFormats[PF_X24_G8].BlockBytes = 5;
 	}
