D3D12 Bindless Work
- Cleaning up QuantizeBoundShaderState implemenations to be named better and to return the struct instead of taking it as an argument.
- Fixing RayTracing global RootSignature generation to read the bindless configuration.
- Fixing RayTracing restoration of descriptor heaps.
- Removing some dead code from D3D12Descriptors

diff --git a/Engine/Source/Runtime/D3D12RHI/Private/D3D12Adapter.cpp b/Engine/Source/Runtime/D3D12RHI/Private/D3D12Adapter.cpp
--- a/Engine/Source/Runtime/D3D12RHI/Private/D3D12Adapter.cpp
+++ b/Engine/Source/Runtime/D3D12RHI/Private/D3D12Adapter.cpp
@@ -319,6 +319,14 @@ FD3D12Adapter::~FD3D12Adapter()
 void FD3D12Adapter::Initialize(FD3D12DynamicRHI* RHI)
 {
 	OwningRHI = RHI;
+
+#if PLATFORM_SUPPORTS_BINDLESS_RENDERING
+	if (Desc.MaxSupportedFeatureLevel >= D3D_FEATURE_LEVEL_12_0 && Desc.MaxSupportedShaderModel >= D3D_SHADER_MODEL_6_6 && Desc.ResourceBindingTier >= D3D12_RESOURCE_BINDING_TIER_3)
+	{
+		bBindlessResourcesAllowed = (RHIGetBindlessResourcesConfiguration(SP_PCD3D_SM6) != ERHIBindlessConfiguration::Disabled);
+		bBindlessSamplersAllowed = (RHIGetBindlessSamplersConfiguration(SP_PCD3D_SM6) != ERHIBindlessConfiguration::Disabled);
+	}
+#endif
 }
 
 /** Callback function called when the GPU crashes, when Aftermath is enabled */
diff --git a/Engine/Source/Runtime/D3D12RHI/Private/D3D12Adapter.h b/Engine/Source/Runtime/D3D12RHI/Private/D3D12Adapter.h
--- a/Engine/Source/Runtime/D3D12RHI/Private/D3D12Adapter.h
+++ b/Engine/Source/Runtime/D3D12RHI/Private/D3D12Adapter.h
@@ -182,6 +182,11 @@ public:
 	FORCEINLINE IDXGIFactory6* GetDXGIFactory6() const { return DxgiFactory6; }
 #endif
 
+#if PLATFORM_SUPPORTS_BINDLESS_RENDERING
+	FORCEINLINE bool AreBindlessResourcesAllowed() const { return bBindlessResourcesAllowed; }
+	FORCEINLINE bool AreBindlessSamplersAllowed() const { return bBindlessSamplersAllowed; }
+#endif
+
 	FORCEINLINE void SetDeviceRemoved(bool value) { bDeviceRemoved = value; }
 	FORCEINLINE const bool IsDeviceRemoved() const { return bDeviceRemoved; }
 	FORCEINLINE const bool IsDebugDevice() const { return bDebugDevice; }
@@ -547,6 +552,12 @@ protected:
 	bool bDeviceRemoved;
 
 	FD3D12AdapterDesc Desc;
+
+#if PLATFORM_SUPPORTS_BINDLESS_RENDERING
+	bool bBindlessResourcesAllowed = false;
+	bool bBindlessSamplersAllowed = false;
+#endif
+
 	TRefCountPtr<IDXGIAdapter> DxgiAdapter;
 
 	FD3D12RootSignatureManager RootSignatureManager;
diff --git a/Engine/Source/Runtime/D3D12RHI/Private/D3D12DescriptorCache.cpp b/Engine/Source/Runtime/D3D12RHI/Private/D3D12DescriptorCache.cpp
--- a/Engine/Source/Runtime/D3D12RHI/Private/D3D12DescriptorCache.cpp
+++ b/Engine/Source/Runtime/D3D12RHI/Private/D3D12DescriptorCache.cpp
@@ -830,6 +830,25 @@ bool FD3D12DescriptorCache::SetHeapOverrides(FD3D12DescriptorHeap* InOverrideVie
 	return bDescriptorHeapsChanged;
 }
 
+void FD3D12DescriptorCache::OverrideLastSetHeaps(ID3D12DescriptorHeap* ViewHeap, ID3D12DescriptorHeap* SamplerHeap)
+{
+	LastSetViewHeap = ViewHeap;
+	LastSetSamplerHeap = SamplerHeap;
+
+	ID3D12DescriptorHeap* ppHeaps[] = { ViewHeap, SamplerHeap };
+	CmdContext->CommandListHandle->SetDescriptorHeaps(UE_ARRAY_COUNT(ppHeaps), ppHeaps);
+
+	bHeapsOverridden = true;
+}
+
+void FD3D12DescriptorCache::RestoreAfterExternalHeapsSet()
+{
+	if (bHeapsOverridden)
+	{
+		SetDescriptorHeaps();
+	}
+}
+
 ////////////////////////////////////////////////////////////////////////////////////////////////////////////////////////
 // FD3D12OnlineHeap
 ////////////////////////////////////////////////////////////////////////////////////////////////////////////////////////
diff --git a/Engine/Source/Runtime/D3D12RHI/Private/D3D12DescriptorCache.h b/Engine/Source/Runtime/D3D12RHI/Private/D3D12DescriptorCache.h
--- a/Engine/Source/Runtime/D3D12RHI/Private/D3D12DescriptorCache.h
+++ b/Engine/Source/Runtime/D3D12RHI/Private/D3D12DescriptorCache.h
@@ -381,6 +381,9 @@ public:
 
 	bool SetHeapOverrides(FD3D12DescriptorHeap* InOverrideViewHeap, FD3D12DescriptorHeap* InOverrideSamplerHeap);
 
+	void OverrideLastSetHeaps(ID3D12DescriptorHeap* ViewHeap, ID3D12DescriptorHeap* SamplerHeap);
+	void RestoreAfterExternalHeapsSet();
+
 	TArray<FD3D12UniqueSamplerTable>& GetUniqueTables() { return UniqueTables; }
 
 	inline bool UsingGlobalSamplerHeap() const { return bUsingGlobalSamplerHeap; }
@@ -424,6 +427,7 @@ private:
 
 	FD3D12SamplerSet LocalSamplerSet;
 	bool bUsingGlobalSamplerHeap = false;
+	bool bHeapsOverridden = false;
 
 	uint32 NumLocalViewDescriptors = 0;
 };
\ No newline at end of file
diff --git a/Engine/Source/Runtime/D3D12RHI/Private/D3D12Descriptors.cpp b/Engine/Source/Runtime/D3D12RHI/Private/D3D12Descriptors.cpp
--- a/Engine/Source/Runtime/D3D12RHI/Private/D3D12Descriptors.cpp
+++ b/Engine/Source/Runtime/D3D12RHI/Private/D3D12Descriptors.cpp
@@ -62,30 +62,6 @@ FD3D12DescriptorManager::FD3D12DescriptorManager(FD3D12Device* Device, FD3D12Des
 
 FD3D12DescriptorManager::~FD3D12DescriptorManager() = default;
 
-void FD3D12DescriptorManager::Init(TCHAR* InName, ERHIDescriptorHeapType InType, uint32 InNumDescriptors)
-{
-	if (InNumDescriptors > 0)
-	{
-		Heap = GetParentDevice()->GetDescriptorHeapManager().AllocateHeap(
-			InName,
-			InType,
-			InNumDescriptors,
-			D3D12_DESCRIPTOR_HEAP_FLAG_SHADER_VISIBLE
-		);
-
-		FRHIDescriptorAllocator::Init(InNumDescriptors);
-	}
-}
-
-void FD3D12DescriptorManager::Destroy()
-{
-	if (Heap)
-	{
-		GetParentDevice()->GetDescriptorHeapManager().FreeHeap(Heap);
-		Heap.SafeRelease();
-	}
-}
-
 void FD3D12DescriptorManager::UpdateImmediately(FRHIDescriptorHandle InHandle, D3D12_CPU_DESCRIPTOR_HANDLE InSourceCpuHandle)
 {
 	if (InHandle.IsValid())
@@ -105,6 +81,8 @@ FD3D12BindlessDescriptorManager::FD3D12BindlessDescriptorManager(FD3D12Device* D
 {
 }
 
+FD3D12BindlessDescriptorManager::~FD3D12BindlessDescriptorManager() = default;
+
 void FD3D12BindlessDescriptorManager::Init(uint32 InNumResourceDescriptors, uint32 InNumSamplerDescriptors)
 {
 	FD3D12DescriptorHeapManager& HeapManager = GetParentDevice()->GetDescriptorHeapManager();
diff --git a/Engine/Source/Runtime/D3D12RHI/Private/D3D12Descriptors.h b/Engine/Source/Runtime/D3D12RHI/Private/D3D12Descriptors.h
--- a/Engine/Source/Runtime/D3D12RHI/Private/D3D12Descriptors.h
+++ b/Engine/Source/Runtime/D3D12RHI/Private/D3D12Descriptors.h
@@ -69,9 +69,6 @@ public:
 	FD3D12DescriptorManager(FD3D12Device* Device, FD3D12DescriptorHeap* InHeap);
 	~FD3D12DescriptorManager();
 
-	void Init(TCHAR* InName, ERHIDescriptorHeapType InType, uint32 InNumDescriptors);
-	void Destroy();
-
 	void UpdateImmediately(FRHIDescriptorHandle InHandle, D3D12_CPU_DESCRIPTOR_HANDLE InSourceCpuHandle);
 
 	inline FD3D12DescriptorHeap* GetHeap() { return Heap.GetReference(); }
@@ -95,6 +92,7 @@ class FD3D12BindlessDescriptorManager : public FD3D12DeviceChild
 {
 public:
 	FD3D12BindlessDescriptorManager(FD3D12Device* Device);
+	~FD3D12BindlessDescriptorManager();
 
 	void Init(uint32 InNumResourceDescriptors, uint32 InNumSamplerDescriptors);
 
diff --git a/Engine/Source/Runtime/D3D12RHI/Private/D3D12Device.cpp b/Engine/Source/Runtime/D3D12RHI/Private/D3D12Device.cpp
--- a/Engine/Source/Runtime/D3D12RHI/Private/D3D12Device.cpp
+++ b/Engine/Source/Runtime/D3D12RHI/Private/D3D12Device.cpp
@@ -342,12 +342,14 @@ void FD3D12Device::SetupAfterDeviceCreation()
 
 	DescriptorHeapManager.Init(GGlobalResourceDescriptorHeapSize, GGlobalSamplerDescriptorHeapSize);
 
+#if PLATFORM_SUPPORTS_BINDLESS_RENDERING
 	const bool bBindlessResources = RHIGetBindlessResourcesConfiguration(GMaxRHIShaderPlatform) != ERHIBindlessConfiguration::Disabled;
 	const bool bBindlessSamplers = RHIGetBindlessSamplersConfiguration(GMaxRHIShaderPlatform) != ERHIBindlessConfiguration::Disabled;
 	if (bBindlessResources || bBindlessSamplers)
 	{
 		BindlessDescriptorManager.Init(bBindlessResources ? GBindlessResourceDescriptorHeapSize : 0, bBindlessSamplers ? GBindlessSamplerDescriptorHeapSize : 0);
 	}
+#endif
 
 	// Init offline descriptor managers
 	for (uint32 Index = 0; Index < static_cast<uint32>(ERHIDescriptorHeapType::count); Index++)
diff --git a/Engine/Source/Runtime/D3D12RHI/Private/D3D12RayTracing.cpp b/Engine/Source/Runtime/D3D12RHI/Private/D3D12RayTracing.cpp
--- a/Engine/Source/Runtime/D3D12RHI/Private/D3D12RayTracing.cpp
+++ b/Engine/Source/Runtime/D3D12RHI/Private/D3D12RayTracing.cpp
@@ -390,6 +390,8 @@ inline void UnregisterD3D12RayTracingGeometry(FD3D12RayTracingGeometry* Geometry
 #define RAY_TRACING_DESCRIPTOR_CACHE_FULL_COMPARE 1
 #endif // RAY_TRACING_DESCRIPTOR_CACHE_FULL_COMPARE
 
+const FD3D12RootSignature* GetGlobalRayTracingRootSignature(FD3D12Adapter& Adapter);
+
 struct FD3D12ShaderIdentifier
 {
 	uint64 Data[4] = {~0ull, ~0ull, ~0ull, ~0ull};
@@ -1259,13 +1261,8 @@ public:
 	ID3D12RootSignature* GetGlobalRootSignature()
 	{
 		FD3D12Adapter* Adapter = GetParentDevice()->GetParentAdapter();
-
-	#if USE_STATIC_ROOT_SIGNATURE
-		return Adapter->GetStaticRayTracingGlobalRootSignature()->GetRootSignature();
-	#else // USE_STATIC_ROOT_SIGNATURE
-		FD3D12QuantizedBoundShaderState QBSS = GetRayTracingGlobalRootSignatureDesc();
-		return Adapter->GetRootSignature(QBSS)->GetRootSignature();
-	#endif // USE_STATIC_ROOT_SIGNATURE
+		const FD3D12RootSignature* RootSignature = GetGlobalRayTracingRootSignature(*Adapter);
+		return RootSignature->GetRootSignature();
 	}
 
 private:
@@ -1622,16 +1619,10 @@ public:
 	{
 		UpdateSyncPoint();
 
-		ID3D12DescriptorHeap* Heaps[2] =
-		{
-			ViewHeap.D3D12Heap,
-			SamplerHeap.D3D12Heap
-		};
-
 		check(ViewHeap.GetParentDevice() == CommandContext.GetParentDevice());
 		check(SamplerHeap.GetParentDevice() == CommandContext.GetParentDevice());
 
-		CommandContext.CommandListHandle.GraphicsCommandList()->SetDescriptorHeaps(2, Heaps);
+		CommandContext.StateCache.GetDescriptorCache()->OverrideLastSetHeaps(ViewHeap.D3D12Heap, SamplerHeap.D3D12Heap);
 	}
 
 	// Returns descriptor heap base index for this descriptor table allocation or -1 if allocation failed.
@@ -5222,15 +5213,6 @@ static void DispatchRays(FD3D12CommandContext& CommandContext,
 	}
 
 	// Setup state for RT dispatch
-	
-	// #dxr_todo UE-72158: RT and non-RT descriptors should use the same global heap that's dynamically sub-allocated.
-	// This requires a major refactor of descriptor heap management. In the short term, RT work uses a dedicated heap
-	// that's temporarily set for the duration of RT dispatch.
-	ID3D12DescriptorHeap* PreviousHeaps[2] =
-	{
-		CommandContext.StateCache.GetDescriptorCache()->GetCurrentViewHeap()->GetHeap(),
-		CommandContext.StateCache.GetDescriptorCache()->GetCurrentSamplerHeap()->GetHeap(),
-	};
 
 	// Invalidate state cache to ensure all root parameters for regular shaders are reset when non-RT work is dispatched later.
 	CommandContext.StateCache.TransitionComputeState(ED3D12PipelineType::RayTracing);
@@ -5239,12 +5221,7 @@ static void DispatchRays(FD3D12CommandContext& CommandContext,
 
 	// Bind diagnostic buffer to allow asserts in ray generation shaders
 	{
-#if USE_STATIC_ROOT_SIGNATURE
-		const FD3D12RootSignature* RootSignature = Adapter->GetStaticRayTracingGlobalRootSignature();
-#else // USE_STATIC_ROOT_SIGNATURE
-		FD3D12QuantizedBoundShaderState RSDesc = GetRayTracingGlobalRootSignatureDesc();
-		const FD3D12RootSignature* RootSignature = Adapter->GetRootSignature(RSDesc);
-#endif // USE_STATIC_ROOT_SIGNATURE
+		const FD3D12RootSignature* RootSignature = GetGlobalRayTracingRootSignature(*Adapter);
 
 		const int8 DiagnosticBufferSlot = RootSignature->GetDiagnosticBufferSlot();
 		D3D12_GPU_VIRTUAL_ADDRESS DiagnosticBufferAddress = CommandContext.GetParentDevice()->GetCommandListManager().GetDiagnosticBufferGPUAddress();
@@ -5333,7 +5310,7 @@ static void DispatchRays(FD3D12CommandContext& CommandContext,
 	}
 
 	// Restore old global descriptor heaps
-	CommandContext.CommandListHandle.GraphicsCommandList()->SetDescriptorHeaps(2, PreviousHeaps);
+	CommandContext.StateCache.GetDescriptorCache()->RestoreAfterExternalHeapsSet();
 }
 
 
diff --git a/Engine/Source/Runtime/D3D12RHI/Private/D3D12Shaders.cpp b/Engine/Source/Runtime/D3D12RHI/Private/D3D12Shaders.cpp
--- a/Engine/Source/Runtime/D3D12RHI/Private/D3D12Shaders.cpp
+++ b/Engine/Source/Runtime/D3D12RHI/Private/D3D12Shaders.cpp
@@ -152,9 +152,7 @@ FComputeShaderRHIRef FD3D12DynamicRHI::RHICreateComputeShader(TArrayView<const u
 #if USE_STATIC_ROOT_SIGNATURE
 		Shader->RootSignature = Adapter.GetStaticComputeRootSignature();
 #else
-		const D3D12_RESOURCE_BINDING_TIER Tier = Adapter.GetResourceBindingTier();
-		FD3D12QuantizedBoundShaderState QBSS;
-		QuantizeBoundShaderState(Tier, Shader, QBSS);
+		const FD3D12QuantizedBoundShaderState QBSS = QuantizeBoundComputeShaderState(Adapter, Shader);
 		Shader->RootSignature = Adapter.GetRootSignature(QBSS);
 #endif
 	}
@@ -212,9 +210,7 @@ FRayTracingShaderRHIRef FD3D12DynamicRHI::RHICreateRayTracingShader(TArrayView<c
 		checkNoEntry(); // Unexpected shader target frequency
 	}
 #else // USE_STATIC_ROOT_SIGNATURE
-	const D3D12_RESOURCE_BINDING_TIER Tier = Adapter.GetResourceBindingTier();
-	FD3D12QuantizedBoundShaderState QBSS;
-	QuantizeBoundShaderState(ShaderFrequency, Tier, Shader, QBSS);
+	const FD3D12QuantizedBoundShaderState QBSS = QuantizeBoundRayTracingShaderState(Adapter, ShaderFrequency, Shader);
 	Shader->pRootSignature = Adapter.GetRootSignature(QBSS);
 #endif // USE_STATIC_ROOT_SIGNATURE
 
@@ -243,9 +239,7 @@ FD3D12BoundShaderState::FD3D12BoundShaderState(
 #if USE_STATIC_ROOT_SIGNATURE
 	pRootSignature = InAdapter->GetStaticGraphicsRootSignature();
 #else
-	const D3D12_RESOURCE_BINDING_TIER Tier = InAdapter->GetResourceBindingTier();
-	FD3D12QuantizedBoundShaderState QuantizedBoundShaderState;
-	QuantizeBoundShaderState(Tier, this, QuantizedBoundShaderState);
+	const FD3D12QuantizedBoundShaderState QuantizedBoundShaderState = QuantizeBoundGraphicsShaderState(*InAdapter, this);
 	pRootSignature = InAdapter->GetRootSignature(QuantizedBoundShaderState);
 #endif
 
@@ -268,9 +262,7 @@ FD3D12BoundShaderState::FD3D12BoundShaderState(
 #if USE_STATIC_ROOT_SIGNATURE
 	pRootSignature = InAdapter->GetStaticGraphicsRootSignature();
 #else
-	const D3D12_RESOURCE_BINDING_TIER Tier = InAdapter->GetResourceBindingTier();
-	FD3D12QuantizedBoundShaderState QuantizedBoundShaderState;
-	QuantizeBoundShaderState(Tier, this, QuantizedBoundShaderState);
+	const FD3D12QuantizedBoundShaderState QuantizedBoundShaderState = QuantizeBoundGraphicsShaderState(*InAdapter, this);
 	pRootSignature = InAdapter->GetRootSignature(QuantizedBoundShaderState);
 #endif
 
diff --git a/Engine/Source/Runtime/D3D12RHI/Private/D3D12Util.cpp b/Engine/Source/Runtime/D3D12RHI/Private/D3D12Util.cpp
--- a/Engine/Source/Runtime/D3D12RHI/Private/D3D12Util.cpp
+++ b/Engine/Source/Runtime/D3D12RHI/Private/D3D12Util.cpp
@@ -964,7 +964,7 @@ namespace D3D12RHI
 	}
 }
 
-void FD3D12QuantizedBoundShaderState::InitShaderRegisterCounts(const D3D12_RESOURCE_BINDING_TIER& ResourceBindingTier, const FShaderCodePackedResourceCounts& Counts, FShaderRegisterCounts& Shader, bool bAllowUAVs)
+void FD3D12QuantizedBoundShaderState::InitShaderRegisterCounts(D3D12_RESOURCE_BINDING_TIER ResourceBindingTier, const FShaderCodePackedResourceCounts& Counts, FShaderRegisterCounts& Shader, bool bAllowUAVs)
 {
 	static const uint32 MaxSamplerCount = MAX_SAMPLERS;
 	static const uint32 MaxConstantBufferCount = MAX_CBS;
@@ -1022,8 +1022,10 @@ static void SetBoundShaderStateFlags(FD3D12QuantizedBoundShaderState& OutQBSS, c
 	if (ShaderData)
 	{
 		OutQBSS.bUseDiagnosticBuffer |= ShaderData->UsesDiagnosticBuffer();
+#if PLATFORM_SUPPORTS_BINDLESS_RENDERING
 		OutQBSS.bUseDirectlyIndexedResourceHeap |= ShaderData->UsesBindlessResources();
 		OutQBSS.bUseDirectlyIndexedSamplerHeap |= ShaderData->UsesBindlessSamplers();
+#endif
 	}
 }
 
@@ -1064,19 +1066,17 @@ static bool IsCompatibleWithBindlessResources(const FD3D12ShaderData* ShaderData
 	return true;
 }
 
-void QuantizeBoundShaderState(
-	const D3D12_RESOURCE_BINDING_TIER& ResourceBindingTier,
-	const FD3D12BoundShaderState* const BSS,
-	FD3D12QuantizedBoundShaderState &QBSS
-	)
+FD3D12QuantizedBoundShaderState QuantizeBoundGraphicsShaderState(FD3D12Adapter& Adapter, const FD3D12BoundShaderState* const BSS)
 {
 	// BSS quantizer. There is a 1:1 mapping of quantized bound shader state objects to root signatures.
 	// The objective is to allow a single root signature to represent many bound shader state objects.
 	// The bigger the quantization step sizes, the fewer the root signatures.
-	FMemory::Memzero(&QBSS, sizeof(QBSS));
+	FD3D12QuantizedBoundShaderState QBSS{};
 
 	QBSS.bAllowIAInputLayout = BSS->GetVertexDeclaration() != nullptr;	// Does the root signature need access to vertex buffers?
 
+	const D3D12_RESOURCE_BINDING_TIER ResourceBindingTier = Adapter.GetResourceBindingTier();
+
 	QuantizeBoundShaderStateCommon(QBSS, BSS->GetVertexShader(),        ResourceBindingTier, SV_Vertex);
 #if PLATFORM_SUPPORTS_MESH_SHADERS
 	QuantizeBoundShaderStateCommon(QBSS, BSS->GetMeshShader(),          ResourceBindingTier, SV_Mesh);
@@ -1085,7 +1085,7 @@ void QuantizeBoundShaderState(
 	QuantizeBoundShaderStateCommon(QBSS, BSS->GetPixelShader(),         ResourceBindingTier, SV_Pixel, true /*bAllowUAVs*/);
 	QuantizeBoundShaderStateCommon(QBSS, BSS->GetGeometryShader(),      ResourceBindingTier, SV_Geometry);
 
-#if DO_CHECK
+#if DO_CHECK && PLATFORM_SUPPORTS_BINDLESS_RENDERING
 	if (QBSS.bUseDirectlyIndexedResourceHeap || QBSS.bUseDirectlyIndexedSamplerHeap)
 	{
 		struct FGenericShaderPair
@@ -1120,62 +1120,71 @@ void QuantizeBoundShaderState(
 		}
 	}
 #endif
+
+	return QBSS;
 }
 
-void QuantizeBoundShaderState(
-	const D3D12_RESOURCE_BINDING_TIER& ResourceBindingTier,
-	const FD3D12ComputeShader* const ComputeShader,
-	FD3D12QuantizedBoundShaderState &OutQBSS
-	)
+FD3D12QuantizedBoundShaderState QuantizeBoundComputeShaderState(FD3D12Adapter& Adapter, const FD3D12ComputeShader* const ComputeShader)
 {
 	check(ComputeShader);
 
 	// BSS quantizer. There is a 1:1 mapping of quantized bound shader state objects to root signatures.
 	// The objective is to allow a single root signature to represent many bound shader state objects.
 	// The bigger the quantization step sizes, the fewer the root signatures.
-	FMemory::Memzero(&OutQBSS, sizeof(OutQBSS));
+	FD3D12QuantizedBoundShaderState QBSS{};
 
-	QuantizeBoundShaderStateCommon(OutQBSS, ComputeShader, ResourceBindingTier, SV_All, true /*bAllowUAVs*/);
+	QuantizeBoundShaderStateCommon(QBSS, ComputeShader, Adapter.GetResourceBindingTier(), SV_All, true /*bAllowUAVs*/);
 
-	check(OutQBSS.bAllowIAInputLayout == false); // No access to vertex buffers needed
+	check(QBSS.bAllowIAInputLayout == false); // No access to vertex buffers needed
+	return QBSS;
 }
 
 #if D3D12_RHI_RAYTRACING
 
-FD3D12QuantizedBoundShaderState GetRayTracingGlobalRootSignatureDesc()
+FD3D12QuantizedBoundShaderState GetRayTracingGlobalRootSignatureDesc(const FD3D12Adapter& Adapter)
 {
-	FD3D12QuantizedBoundShaderState OutQBSS = {};
-	FShaderRegisterCounts& QBSSRegisterCounts = OutQBSS.RegisterCounts[SV_All];
+	FD3D12QuantizedBoundShaderState QBSS{};
+	FShaderRegisterCounts& QBSSRegisterCounts = QBSS.RegisterCounts[SV_All];
 
-	OutQBSS.RootSignatureType = RS_RayTracingGlobal;
-	OutQBSS.bUseDiagnosticBuffer = true;
+	QBSS.RootSignatureType = RS_RayTracingGlobal;
+	QBSS.bUseDiagnosticBuffer = true;
+
+#if PLATFORM_SUPPORTS_BINDLESS_RENDERING
+	QBSS.bUseDirectlyIndexedResourceHeap = Adapter.AreBindlessResourcesAllowed();
+	QBSS.bUseDirectlyIndexedSamplerHeap = Adapter.AreBindlessSamplersAllowed();
+#endif
 
 	QBSSRegisterCounts.SamplerCount = MAX_SAMPLERS;
 	QBSSRegisterCounts.ShaderResourceCount = MAX_SRVS;
 	QBSSRegisterCounts.ConstantBufferCount = MAX_CBS;
 	QBSSRegisterCounts.UnorderedAccessCount = MAX_UAVS;
 
-	return OutQBSS;
+	return QBSS;
 }
 
-void QuantizeBoundShaderState(
-	EShaderFrequency ShaderFrequency,
-	const D3D12_RESOURCE_BINDING_TIER& ResourceBindingTier,
-	const FD3D12RayTracingShader* const RayTracingShader,
-	FD3D12QuantizedBoundShaderState &OutQBSS
+const FD3D12RootSignature* GetGlobalRayTracingRootSignature(FD3D12Adapter& Adapter)
+{
+#if USE_STATIC_ROOT_SIGNATURE
+	return Adapter.GetStaticRayTracingGlobalRootSignature();
+#else // USE_STATIC_ROOT_SIGNATURE
+	FD3D12QuantizedBoundShaderState RSDesc = GetRayTracingGlobalRootSignatureDesc(Adapter);
+	return Adapter.GetRootSignature(RSDesc);
+#endif // USE_STATIC_ROOT_SIGNATURE
+}
+
+FD3D12QuantizedBoundShaderState QuantizeBoundRayTracingShaderState(FD3D12Adapter& Adapter, EShaderFrequency ShaderFrequency, const FD3D12RayTracingShader* const RayTracingShader
 )
 {
-	FMemory::Memzero(&OutQBSS, sizeof(OutQBSS));
-	FShaderRegisterCounts& QBSSRegisterCounts = OutQBSS.RegisterCounts[SV_All];
+	FD3D12QuantizedBoundShaderState QBSS{};
+
+	FShaderRegisterCounts& QBSSRegisterCounts = QBSS.RegisterCounts[SV_All];
 
 	switch (ShaderFrequency)
 	{
 	case SF_RayGen:
 	{
 		// Shared conservative root signature layout is used for all raygen and miss shaders.
-
-		OutQBSS = GetRayTracingGlobalRootSignatureDesc();
-
+		QBSS = GetRayTracingGlobalRootSignatureDesc(Adapter);
 		break;
 	}
 
@@ -1184,11 +1193,10 @@ void QuantizeBoundShaderState(
 	case SF_RayMiss:
 	{
 		// Local root signature is used for hit group shaders, using the exact number of resources to minimize shader binding table record size.
-
 		check(RayTracingShader);
 		const FShaderCodePackedResourceCounts& Counts = RayTracingShader->ResourceCounts;
 
-		OutQBSS.RootSignatureType = RS_RayTracingLocal;
+		QBSS.RootSignatureType = RS_RayTracingLocal;
 
 		QBSSRegisterCounts.SamplerCount = Counts.NumSamplers;
 		QBSSRegisterCounts.ShaderResourceCount = Counts.NumSRVs;
@@ -1206,7 +1214,9 @@ void QuantizeBoundShaderState(
 		checkNoEntry(); // Unexpected shader target frequency
 	}
 
-	SetBoundShaderStateFlags(OutQBSS, RayTracingShader);
+	SetBoundShaderStateFlags(QBSS, RayTracingShader);
+
+	return QBSS;
 }
 #endif // D3D12_RHI_RAYTRACING
 
diff --git a/Engine/Source/Runtime/D3D12RHI/Public/D3D12Util.h b/Engine/Source/Runtime/D3D12RHI/Public/D3D12Util.h
--- a/Engine/Source/Runtime/D3D12RHI/Public/D3D12Util.h
+++ b/Engine/Source/Runtime/D3D12RHI/Public/D3D12Util.h
@@ -112,7 +112,7 @@ struct FD3D12QuantizedBoundShaderState
 
 	friend uint32 GetTypeHash(const FD3D12QuantizedBoundShaderState& Key);
 
-	static void InitShaderRegisterCounts(const D3D12_RESOURCE_BINDING_TIER& ResourceBindingTier, const FShaderCodePackedResourceCounts& Counts, FShaderRegisterCounts& Shader, bool bAllowUAVs = false);
+	static void InitShaderRegisterCounts(D3D12_RESOURCE_BINDING_TIER ResourceBindingTier, const FShaderCodePackedResourceCounts& Counts, FShaderRegisterCounts& Shader, bool bAllowUAVs = false);
 };
 
 /**
@@ -120,25 +120,14 @@ struct FD3D12QuantizedBoundShaderState
 */
 
 class FD3D12BoundShaderState;
-extern void QuantizeBoundShaderState(
-	const D3D12_RESOURCE_BINDING_TIER& ResourceBindingTier,
-	const FD3D12BoundShaderState* const BSS,
-	FD3D12QuantizedBoundShaderState &OutQBSS);
-
 class FD3D12ComputeShader;
-extern void QuantizeBoundShaderState(
-	const D3D12_RESOURCE_BINDING_TIER& ResourceBindingTier,
-	const FD3D12ComputeShader* const ComputeShader,
-	FD3D12QuantizedBoundShaderState &OutQBSS);
+
+extern FD3D12QuantizedBoundShaderState QuantizeBoundGraphicsShaderState(FD3D12Adapter& Adapter, const FD3D12BoundShaderState* BSS);
+extern FD3D12QuantizedBoundShaderState QuantizeBoundComputeShaderState(FD3D12Adapter& Adapter, const FD3D12ComputeShader* ComputeShader);
 
 #if D3D12_RHI_RAYTRACING
 class FD3D12RayTracingShader;
-FD3D12QuantizedBoundShaderState GetRayTracingGlobalRootSignatureDesc();
-extern void QuantizeBoundShaderState(
-	EShaderFrequency ShaderFrequency,
-	const D3D12_RESOURCE_BINDING_TIER& ResourceBindingTier,
-	const FD3D12RayTracingShader* const Shader,
-	FD3D12QuantizedBoundShaderState &OutQBSS);
+extern FD3D12QuantizedBoundShaderState QuantizeBoundRayTracingShaderState(FD3D12Adapter& Adapter, EShaderFrequency ShaderFrequency, const FD3D12RayTracingShader* Shader);
 #endif
 
 /**
