D3D12 Bindless
When switching between Compute and Graphics pipelines, we aren't guaranteed to change the respective pipeline's Root Signature and since bindless heap setting was based on a RS change, we weren't always setting bindless heaps correctly.
FD3D12DescriptorCache::IsSamplerHeapOverridden wasn't checking the correct member.

diff --git a/Engine/Source/Runtime/D3D12RHI/Private/D3D12DescriptorCache.h b/Engine/Source/Runtime/D3D12RHI/Private/D3D12DescriptorCache.h
--- a/Engine/Source/Runtime/D3D12RHI/Private/D3D12DescriptorCache.h
+++ b/Engine/Source/Runtime/D3D12RHI/Private/D3D12DescriptorCache.h
@@ -327,7 +327,7 @@ public:
 	~FD3D12DescriptorCache();
 
 	inline bool IsViewHeapOverridden() const { return OverrideViewHeap != nullptr; }
-	inline bool IsSamplerHeapOverridden() const { return OverrideViewHeap != nullptr; }
+	inline bool IsSamplerHeapOverridden() const { return OverrideSamplerHeap != nullptr; }
 
 	inline FD3D12OnlineHeap* GetCurrentViewHeap() const { return CurrentViewHeap; }
 	inline FD3D12OnlineHeap* GetCurrentSamplerHeap() const { return CurrentSamplerHeap; }
diff --git a/Engine/Source/Runtime/D3D12RHI/Private/D3D12StateCache.cpp b/Engine/Source/Runtime/D3D12RHI/Private/D3D12StateCache.cpp
--- a/Engine/Source/Runtime/D3D12RHI/Private/D3D12StateCache.cpp
+++ b/Engine/Source/Runtime/D3D12RHI/Private/D3D12StateCache.cpp
@@ -516,7 +516,10 @@ void FD3D12StateCacheBase::ApplyState()
 	const bool bApplyResources = !bBindlessResources && PSOCommonData->RootSignature->HasResources();
 	const bool bApplySamplers = !bBindlessSamplers && PSOCommonData->RootSignature->HasSamplers();
 
-	if (bRootSignatureChanged)
+	const bool bBindlessResourcesAlreadyEnabled = DescriptorCache.IsViewHeapOverridden();
+	const bool bBindlessSamplersAlreadyEnabled = DescriptorCache.IsSamplerHeapOverridden();
+
+	if (bRootSignatureChanged || bBindlessResources != bBindlessResourcesAlreadyEnabled || bBindlessSamplers != bBindlessSamplersAlreadyEnabled)
 	{
 		FD3D12BindlessDescriptorManager& BindlessManager = GetParentDevice()->GetBindlessDescriptorManager();
 
