Adding more bindless support to the D3D12 root signature creators.

diff --git a/Engine/Source/Runtime/D3D12RHI/Public/D3D12RootSignatureDefinitions.h b/Engine/Source/Runtime/D3D12RHI/Public/D3D12RootSignatureDefinitions.h
--- a/Engine/Source/Runtime/D3D12RHI/Public/D3D12RootSignatureDefinitions.h
+++ b/Engine/Source/Runtime/D3D12RHI/Public/D3D12RootSignatureDefinitions.h
@@ -146,7 +146,7 @@ namespace D3D12ShaderUtils
 				break;
 			}
 
-			return TEXT("0");
+			return TEXT("");
 		};
 	};
 
@@ -296,7 +296,7 @@ namespace D3D12ShaderUtils
 			return *this;
 		}
 
-		FString CreateAndGenerateString(EShaderFrequency Freq)
+		FString CreateAndGenerateString(EShaderFrequency Freq, bool bBindlessResources, bool bBindlessSamplers)
 		{
 			if (Freq < SF_NumGraphicsFrequencies)
 			{
@@ -311,6 +311,16 @@ namespace D3D12ShaderUtils
 				D3D12ShaderUtils::CreateComputeRootSignature(*this);
 			}
 
+			if (bBindlessResources)
+			{
+				AddRootFlag(D3D12_ROOT_SIGNATURE_FLAG_CBV_SRV_UAV_HEAP_DIRECTLY_INDEXED);
+			}
+
+			if (bBindlessSamplers)
+			{
+				AddRootFlag(D3D12_ROOT_SIGNATURE_FLAG_SAMPLER_HEAP_DIRECTLY_INDEXED);
+			}
+
 			//FString String = FString::Printf(TEXT("#define TheRootSignature \"RootFlags(%s),%s\"\n[RootSignature(TheRootSignature)\n"),
 			FString String = FString::Printf(TEXT("\"RootFlags(%s),%s\""),
 				Flags.Len() == 0 ? TEXT("0") : *Flags,
@@ -322,12 +332,12 @@ namespace D3D12ShaderUtils
 		FString Table;
 	};
 
-	inline FString GenerateRootSignatureString(EShaderFrequency InFrequency)
+	inline FString GenerateRootSignatureString(EShaderFrequency InFrequency, bool bBindlessResources = false, bool bBindlessSamplers = false)
 	{
 		if (InFrequency < SF_NumStandardFrequencies)
 		{
 			FTextRootSignatureCreator Creator;
-			return Creator.CreateAndGenerateString(InFrequency);
+			return Creator.CreateAndGenerateString(InFrequency, bBindlessResources, bBindlessSamplers);
 		}
 
 		return TEXT("");
