D3D12 Root Signature changes
- Reorganizing the static root signature code to be initialized once on startup instead of relying on function-local statics.
- Removing the unused root signature arguments from FD3D12PipelineStateCache::RebuildFromDiskCache so the static root signatures can be isolated to when USE_STATIC_ROOT_SIGNATURE is defined.

diff --git a/Engine/Source/Runtime/D3D12RHI/Private/D3D12Adapter.cpp b/Engine/Source/Runtime/D3D12RHI/Private/D3D12Adapter.cpp
--- a/Engine/Source/Runtime/D3D12RHI/Private/D3D12Adapter.cpp
+++ b/Engine/Source/Runtime/D3D12RHI/Private/D3D12Adapter.cpp
@@ -274,6 +274,12 @@ FD3D12Adapter::FD3D12Adapter(FD3D12AdapterDesc& DescIn)
 	, DefaultAsyncComputeContextRedirector(this, ED3D12CommandQueueType::Async, true)
 	, FrameCounter(0)
 	, DebugFlags(0)
+#if USE_STATIC_ROOT_SIGNATURE
+	, StaticGraphicsRootSignature(this)
+	, StaticComputeRootSignature(this)
+	, StaticRayTracingGlobalRootSignature(this)
+	, StaticRayTracingLocalRootSignature(this)
+#endif
 {
 	FMemory::Memzero(&UploadHeapAllocator, sizeof(UploadHeapAllocator));
 	FMemory::Memzero(&Devices, sizeof(Devices));
@@ -319,14 +325,6 @@ FD3D12Adapter::~FD3D12Adapter()
 void FD3D12Adapter::Initialize(FD3D12DynamicRHI* RHI)
 {
 	OwningRHI = RHI;
-
-#if PLATFORM_SUPPORTS_BINDLESS_RENDERING
-	if (Desc.MaxSupportedFeatureLevel >= D3D_FEATURE_LEVEL_12_0 && Desc.MaxSupportedShaderModel >= D3D_SHADER_MODEL_6_6 && Desc.ResourceBindingTier >= D3D12_RESOURCE_BINDING_TIER_3)
-	{
-		bBindlessResourcesAllowed = (RHIGetBindlessResourcesConfiguration(SP_PCD3D_SM6) != ERHIBindlessConfiguration::Disabled);
-		bBindlessSamplersAllowed = (RHIGetBindlessSamplersConfiguration(SP_PCD3D_SM6) != ERHIBindlessConfiguration::Disabled);
-	}
-#endif
 }
 
 /** Callback function called when the GPU crashes, when Aftermath is enabled */
@@ -1206,11 +1204,34 @@ void FD3D12Adapter::InitializeDevices()
 		FString DriverBlobFilename = PIPELINE_STATE_FILE_LOCATION / FString::Printf(TEXT("D3DDriverByteCodeBlob_%s"), *UniqueDeviceCachePath);
 
 		PipelineStateCache.Init(GraphicsCacheFile, ComputeCacheFile, DriverBlobFilename);
+		PipelineStateCache.RebuildFromDiskCache();
 
-		ID3D12RootSignature* StaticGraphicsRS = (GetStaticGraphicsRootSignature()) ? GetStaticGraphicsRootSignature()->GetRootSignature() : nullptr;
-		ID3D12RootSignature* StaticComputeRS = (GetStaticComputeRootSignature()) ? GetStaticComputeRootSignature()->GetRootSignature() : nullptr;
+		ERHIBindlessConfiguration BindlessResourcesConfig = ERHIBindlessConfiguration::Disabled;
+		ERHIBindlessConfiguration BindlessSamplersConfig = ERHIBindlessConfiguration::Disabled;
 
-		PipelineStateCache.RebuildFromDiskCache(StaticGraphicsRS, StaticComputeRS);
+#if PLATFORM_SUPPORTS_BINDLESS_RENDERING
+		if (Desc.MaxSupportedFeatureLevel >= D3D_FEATURE_LEVEL_12_0 && Desc.MaxSupportedShaderModel >= D3D_SHADER_MODEL_6_6 && Desc.ResourceBindingTier >= D3D12_RESOURCE_BINDING_TIER_3)
+		{
+			BindlessResourcesConfig = RHIGetBindlessResourcesConfiguration(GMaxRHIShaderPlatform);
+			BindlessSamplersConfig = RHIGetBindlessSamplersConfiguration(GMaxRHIShaderPlatform);
+
+			bBindlessResourcesAllowed = (BindlessResourcesConfig != ERHIBindlessConfiguration::Disabled);
+			bBindlessSamplersAllowed = (BindlessSamplersConfig != ERHIBindlessConfiguration::Disabled);
+		}
+#endif
+
+#if USE_STATIC_ROOT_SIGNATURE
+		const bool bBindlessResourcesGraphics = (BindlessResourcesConfig == ERHIBindlessConfiguration::AllShaders);
+		const bool bBindlessSamplersGraphics = (BindlessSamplersConfig == ERHIBindlessConfiguration::AllShaders);
+
+		const bool bBindlessResourcesRayTracing = (BindlessResourcesConfig != ERHIBindlessConfiguration::Disabled);
+		const bool bBindlessSamplersRayTracing = (BindlessSamplersConfig != ERHIBindlessConfiguration::Disabled);
+
+		StaticGraphicsRootSignature.InitStaticGraphicsRootSignature(bBindlessResourcesGraphics, bBindlessSamplersGraphics);
+		StaticComputeRootSignature.InitStaticComputeRootSignatureDesc(bBindlessResourcesGraphics, bBindlessSamplersGraphics);
+		StaticRayTracingGlobalRootSignature.InitStaticRayTracingGlobalRootSignatureDesc(bBindlessResourcesRayTracing, bBindlessSamplersRayTracing);
+		StaticRayTracingLocalRootSignature.InitStaticRayTracingLocalRootSignatureDesc();
+#endif
 	}
 }
 
diff --git a/Engine/Source/Runtime/D3D12RHI/Private/D3D12Adapter.h b/Engine/Source/Runtime/D3D12RHI/Private/D3D12Adapter.h
--- a/Engine/Source/Runtime/D3D12RHI/Private/D3D12Adapter.h
+++ b/Engine/Source/Runtime/D3D12RHI/Private/D3D12Adapter.h
@@ -222,32 +222,20 @@ public:
 #if USE_STATIC_ROOT_SIGNATURE
 	FORCEINLINE const FD3D12RootSignature* GetStaticGraphicsRootSignature()
 	{
-		static const FD3D12RootSignature StaticGraphicsRootSignature(this, FD3D12RootSignatureDesc::GetStaticGraphicsRootSignatureDesc());
 		return &StaticGraphicsRootSignature;
 	}
 	FORCEINLINE const FD3D12RootSignature* GetStaticComputeRootSignature()
 	{
-		static const FD3D12RootSignature StaticComputeRootSignature(this, FD3D12RootSignatureDesc::GetStaticComputeRootSignatureDesc());
 		return &StaticComputeRootSignature;
 	}
 	FORCEINLINE const FD3D12RootSignature* GetStaticRayTracingGlobalRootSignature()
 	{
-		static const FD3D12RootSignature StaticRootSignature(this, FD3D12RootSignatureDesc::GetStaticRayTracingGlobalRootSignatureDesc(), 1 /*RAY_TRACING_REGISTER_SPACE_GLOBAL*/);
-		return &StaticRootSignature;
+		return &StaticRayTracingGlobalRootSignature;
 	}
 	FORCEINLINE const FD3D12RootSignature* GetStaticRayTracingLocalRootSignature()
 	{
-		static const FD3D12RootSignature StaticRootSignature(this, FD3D12RootSignatureDesc::GetStaticRayTracingLocalRootSignatureDesc(), 0 /*RAY_TRACING_REGISTER_SPACE_LOCAL*/);
-		return &StaticRootSignature;
+		return &StaticRayTracingLocalRootSignature;
 	}
-#else // USE_STATIC_ROOT_SIGNATURE
-	FORCEINLINE const FD3D12RootSignature* GetStaticGraphicsRootSignature() { return nullptr; }
-#if PLATFORM_SUPPORTS_MESH_SHADERS
-	FORCEINLINE const FD3D12RootSignature* GetStaticMeshRootSignature() { return nullptr; }
-#endif
-	FORCEINLINE const FD3D12RootSignature* GetStaticComputeRootSignature() { return nullptr; }
-	FORCEINLINE const FD3D12RootSignature* GetStaticRayTracingGlobalRootSignature() { return nullptr; }
-	FORCEINLINE const FD3D12RootSignature* GetStaticRayTracingLocalRootSignature() { return nullptr; }
 #endif // USE_STATIC_ROOT_SIGNATURE
 
 	FORCEINLINE FD3D12RootSignature* GetRootSignature(const FD3D12QuantizedBoundShaderState& QBSS) 
@@ -635,4 +623,11 @@ protected:
 	FD3D12Device* Devices[MAX_NUM_GPUS];
 
 	uint32 DebugFlags;
+
+#if USE_STATIC_ROOT_SIGNATURE
+	FD3D12RootSignature StaticGraphicsRootSignature;
+	FD3D12RootSignature StaticComputeRootSignature;
+	FD3D12RootSignature StaticRayTracingGlobalRootSignature;
+	FD3D12RootSignature StaticRayTracingLocalRootSignature;
+#endif
 };
diff --git a/Engine/Source/Runtime/D3D12RHI/Private/D3D12RootSignature.cpp b/Engine/Source/Runtime/D3D12RHI/Private/D3D12RootSignature.cpp
--- a/Engine/Source/Runtime/D3D12RHI/Private/D3D12RootSignature.cpp
+++ b/Engine/Source/Runtime/D3D12RHI/Private/D3D12RootSignature.cpp
@@ -388,10 +388,10 @@ FD3D12RootSignatureDesc::FD3D12RootSignatureDesc(const FD3D12QuantizedBoundShade
 
 }
 
-const D3D12_VERSIONED_ROOT_SIGNATURE_DESC& FD3D12RootSignatureDesc::GetStaticGraphicsRootSignatureDesc()
+void FD3D12RootSignature::InitStaticGraphicsRootSignature(bool bBindlessResources, bool bBindlessSamplers)
 {
 	// TODO: Support vendor extensions for static root signatures?
-	static struct
+	struct
 	{
 		D3D12_SHADER_VISIBILITY Vis;
 		D3D12_DESCRIPTOR_RANGE_TYPE Type;
@@ -425,9 +425,9 @@ const D3D12_VERSIONED_ROOT_SIGNATURE_DESC& FD3D12RootSignatureDesc::GetStaticGra
 		{ D3D12_SHADER_VISIBILITY_ALL, D3D12_DESCRIPTOR_RANGE_TYPE_UAV, MAX_UAVS, 0, D3D12ShaderUtils::StaticRootSignatureConstants::UAVDescriptorRangeFlags },
 	};
 
-	static const uint32 DescriptorTableCount = UE_ARRAY_COUNT(RangeDesc);
-	static CD3DX12_ROOT_PARAMETER1 TableSlots[DescriptorTableCount];
-	static CD3DX12_DESCRIPTOR_RANGE1 DescriptorRanges[DescriptorTableCount];
+	const uint32 DescriptorTableCount = UE_ARRAY_COUNT(RangeDesc);
+	CD3DX12_ROOT_PARAMETER1 TableSlots[DescriptorTableCount];
+	CD3DX12_DESCRIPTOR_RANGE1 DescriptorRanges[DescriptorTableCount];
 
 	for (uint32 i = 0; i < DescriptorTableCount; i++)
 	{
@@ -442,17 +442,27 @@ const D3D12_VERSIONED_ROOT_SIGNATURE_DESC& FD3D12RootSignatureDesc::GetStaticGra
 		TableSlots[i].InitAsDescriptorTable(1, &DescriptorRanges[i], RangeDesc[i].Vis);
 	}
 
-	static CD3DX12_VERSIONED_ROOT_SIGNATURE_DESC RootDesc(DescriptorTableCount, TableSlots, 0, nullptr, D3D12_ROOT_SIGNATURE_FLAG_ALLOW_INPUT_ASSEMBLER_INPUT_LAYOUT);
-	return RootDesc;
+	D3D12_ROOT_SIGNATURE_FLAGS Flags = D3D12_ROOT_SIGNATURE_FLAG_ALLOW_INPUT_ASSEMBLER_INPUT_LAYOUT;
+	if (bBindlessResources)
+	{
+		Flags |= D3D12_ROOT_SIGNATURE_FLAG_CBV_SRV_UAV_HEAP_DIRECTLY_INDEXED;
+	}
+	if (bBindlessSamplers)
+	{
+		Flags |= D3D12_ROOT_SIGNATURE_FLAG_SAMPLER_HEAP_DIRECTLY_INDEXED;
+	}
+
+	CD3DX12_VERSIONED_ROOT_SIGNATURE_DESC RootDesc(DescriptorTableCount, TableSlots, 0, nullptr, Flags);
+	Init(RootDesc);
 }
 
-const D3D12_VERSIONED_ROOT_SIGNATURE_DESC& FD3D12RootSignatureDesc::GetStaticComputeRootSignatureDesc()
+void FD3D12RootSignature::InitStaticComputeRootSignatureDesc(bool bBindlessResources, bool bBindlessSamplers)
 {
 	// TODO: Support vendor extensions for static root signatures?
 
-	static const uint32 DescriptorTableCount = 4;
-	static CD3DX12_ROOT_PARAMETER1 TableSlots[DescriptorTableCount];
-	static CD3DX12_DESCRIPTOR_RANGE1 DescriptorRanges[DescriptorTableCount];
+	const uint32 DescriptorTableCount = 4;
+	CD3DX12_ROOT_PARAMETER1 TableSlots[DescriptorTableCount];
+	CD3DX12_DESCRIPTOR_RANGE1 DescriptorRanges[DescriptorTableCount];
 
 	uint32 RangeIndex = 0;
 	DescriptorRanges[RangeIndex].Init(D3D12_DESCRIPTOR_RANGE_TYPE_SRV, MAX_SRVS, 0, 0, D3D12ShaderUtils::StaticRootSignatureConstants::SRVDescriptorRangeFlags);
@@ -468,30 +478,40 @@ const D3D12_VERSIONED_ROOT_SIGNATURE_DESC& FD3D12RootSignatureDesc::GetStaticCom
 	TableSlots[RangeIndex].InitAsDescriptorTable(1, &DescriptorRanges[RangeIndex], D3D12_SHADER_VISIBILITY_ALL);
 	++RangeIndex;
 
-	static CD3DX12_VERSIONED_ROOT_SIGNATURE_DESC RootDesc(RangeIndex, TableSlots, 0, nullptr, D3D12_ROOT_SIGNATURE_FLAG_NONE);
-	return RootDesc;
+	D3D12_ROOT_SIGNATURE_FLAGS Flags = D3D12_ROOT_SIGNATURE_FLAG_NONE;
+	if (bBindlessResources)
+	{
+		Flags |= D3D12_ROOT_SIGNATURE_FLAG_CBV_SRV_UAV_HEAP_DIRECTLY_INDEXED;
+	}
+	if (bBindlessSamplers)
+	{
+		Flags |= D3D12_ROOT_SIGNATURE_FLAG_SAMPLER_HEAP_DIRECTLY_INDEXED;
+	}
+
+	CD3DX12_VERSIONED_ROOT_SIGNATURE_DESC RootDesc(RangeIndex, TableSlots, 0, nullptr, Flags);
+	Init(RootDesc);
 }
 
-const D3D12_VERSIONED_ROOT_SIGNATURE_DESC& FD3D12RootSignatureDesc::GetStaticRayTracingGlobalRootSignatureDesc()
+void FD3D12RootSignature::InitStaticRayTracingGlobalRootSignatureDesc(bool bBindlessResources, bool bBindlessSamplers)
 {
 #if D3D12_RHI_RAYTRACING
-	return D3D12ShaderUtils::GetStaticRayTracingRootSignatureDesc<false, FD3D12_ROOT_SIGNATURE_FLAG_GLOBAL_ROOT_SIGNATURE>();
-#else // D3D12_RHI_RAYTRACING
+	D3D12ShaderUtils::FStaticRayTracingRootSignatureDesc Desc{};
+	D3D12ShaderUtils::InitStaticRayTracingRootSignatureDesc(Desc, false, FD3D12_ROOT_SIGNATURE_FLAG_GLOBAL_ROOT_SIGNATURE, bBindlessResources, bBindlessSamplers);
+	Init(Desc.RootDesc, 1 /*RAY_TRACING_REGISTER_SPACE_GLOBAL*/);
+#else
 	checkNoEntry();
-	static const D3D12_VERSIONED_ROOT_SIGNATURE_DESC Desc = {};
-	return Desc;
-#endif // D3D12_RHI_RAYTRACING
+#endif
 }
 
-const D3D12_VERSIONED_ROOT_SIGNATURE_DESC& FD3D12RootSignatureDesc::GetStaticRayTracingLocalRootSignatureDesc()
+void FD3D12RootSignature::InitStaticRayTracingLocalRootSignatureDesc()
 {
 #if D3D12_RHI_RAYTRACING
-	return D3D12ShaderUtils::GetStaticRayTracingRootSignatureDesc<true, D3D12_ROOT_SIGNATURE_FLAG_LOCAL_ROOT_SIGNATURE>();
-#else // D3D12_RHI_RAYTRACING
+	D3D12ShaderUtils::FStaticRayTracingRootSignatureDesc Desc{};
+	D3D12ShaderUtils::InitStaticRayTracingRootSignatureDesc(Desc, true, D3D12_ROOT_SIGNATURE_FLAG_LOCAL_ROOT_SIGNATURE, false, false);
+	Init(Desc.RootDesc, 0 /*RAY_TRACING_REGISTER_SPACE_LOCAL*/);
+#else
 	checkNoEntry();
-	static const D3D12_VERSIONED_ROOT_SIGNATURE_DESC Desc = {};
-	return Desc;
-#endif // D3D12_RHI_RAYTRACING
+#endif
 }
 
 void FD3D12RootSignature::Init(const FD3D12QuantizedBoundShaderState& InQBSS)
diff --git a/Engine/Source/Runtime/D3D12RHI/Private/D3D12RootSignature.h b/Engine/Source/Runtime/D3D12RHI/Private/D3D12RootSignature.h
--- a/Engine/Source/Runtime/D3D12RHI/Private/D3D12RootSignature.h
+++ b/Engine/Source/Runtime/D3D12RHI/Private/D3D12RootSignature.h
@@ -35,6 +35,8 @@ enum ERootParameterKeys
 	RPK_RootParameterKeyCount,
 };
 
+class FD3D12RootSignature;
+
 class FD3D12RootSignatureDesc
 {
 public:
@@ -42,11 +44,6 @@ public:
 
 	inline const D3D12_VERSIONED_ROOT_SIGNATURE_DESC& GetDesc() const { return RootDesc; }
 
-	static const D3D12_VERSIONED_ROOT_SIGNATURE_DESC& GetStaticGraphicsRootSignatureDesc();
-	static const D3D12_VERSIONED_ROOT_SIGNATURE_DESC& GetStaticComputeRootSignatureDesc();
-	static const D3D12_VERSIONED_ROOT_SIGNATURE_DESC& GetStaticRayTracingGlobalRootSignatureDesc();
-	static const D3D12_VERSIONED_ROOT_SIGNATURE_DESC& GetStaticRayTracingLocalRootSignatureDesc();
-
 	static constexpr uint32 MaxRootParameters = 32;	// Arbitrary max, increase as needed.
 
 	inline int8 GetDiagnosticBufferSlot() const { return DiagnosticBufferSlot; }
@@ -99,6 +96,11 @@ public:
 	void Init(const D3D12_VERSIONED_ROOT_SIGNATURE_DESC& InDesc, uint32 BindingSpace = 0);
 	void Init(ID3DBlob* const InBlob, uint32 BindingSpace = 0);
 
+	void InitStaticGraphicsRootSignature(bool bBindlessResources, bool bBindlessSamplers);
+	void InitStaticComputeRootSignatureDesc(bool bBindlessResources, bool bBindlessSamplers);
+	void InitStaticRayTracingGlobalRootSignatureDesc(bool bBindlessResources, bool bBindlessSamplers);
+	void InitStaticRayTracingLocalRootSignatureDesc();
+
 	ID3D12RootSignature* GetRootSignature() const { return RootSignature.GetReference(); }
 	ID3DBlob* GetRootSignatureBlob() const { return RootSignatureBlob.GetReference(); }
 
diff --git a/Engine/Source/Runtime/D3D12RHI/Private/Windows/WindowsD3D12PipelineState.cpp b/Engine/Source/Runtime/D3D12RHI/Private/Windows/WindowsD3D12PipelineState.cpp
--- a/Engine/Source/Runtime/D3D12RHI/Private/Windows/WindowsD3D12PipelineState.cpp
+++ b/Engine/Source/Runtime/D3D12RHI/Private/Windows/WindowsD3D12PipelineState.cpp
@@ -176,13 +176,10 @@ void FD3D12PipelineStateCache::OnPSOCreated(FD3D12PipelineState* PipelineState,
 	}
 }
 
-void FD3D12PipelineStateCache::RebuildFromDiskCache(ID3D12RootSignature* GraphicsRootSignature, ID3D12RootSignature* ComputeRootSignature)
+void FD3D12PipelineStateCache::RebuildFromDiskCache()
 {
 	FRWScopeLock Lock(DiskCachesCS, SLT_Write);
 
-	UNREFERENCED_PARAMETER(GraphicsRootSignature);
-	UNREFERENCED_PARAMETER(ComputeRootSignature);
-
 	if (IsInErrorState())
 	{
 		// TODO: Make sure we clear the disk caches that are in error.
diff --git a/Engine/Source/Runtime/D3D12RHI/Private/Windows/WindowsD3D12PipelineState.h b/Engine/Source/Runtime/D3D12RHI/Private/Windows/WindowsD3D12PipelineState.h
--- a/Engine/Source/Runtime/D3D12RHI/Private/Windows/WindowsD3D12PipelineState.h
+++ b/Engine/Source/Runtime/D3D12RHI/Private/Windows/WindowsD3D12PipelineState.h
@@ -207,7 +207,7 @@ public:
 	using FD3D12PipelineStateCacheBase::FindInLoadedCache;
 	using FD3D12PipelineStateCacheBase::CreateAndAdd;
 #endif
-	void RebuildFromDiskCache(ID3D12RootSignature* GraphicsRootSignature, ID3D12RootSignature* ComputeRootSignature);
+	void RebuildFromDiskCache();
 
 	void Close();
 
diff --git a/Engine/Source/Runtime/D3D12RHI/Public/D3D12RayTracingRootSignature.h b/Engine/Source/Runtime/D3D12RHI/Public/D3D12RayTracingRootSignature.h
--- a/Engine/Source/Runtime/D3D12RHI/Public/D3D12RayTracingRootSignature.h
+++ b/Engine/Source/Runtime/D3D12RHI/Public/D3D12RayTracingRootSignature.h
@@ -7,42 +7,54 @@
 
 namespace D3D12ShaderUtils
 {
-	template <bool bLocalRootSignature, D3D12_ROOT_SIGNATURE_FLAGS Flags>
-	const D3D12_VERSIONED_ROOT_SIGNATURE_DESC& GetStaticRayTracingRootSignatureDesc()
+	struct FStaticRayTracingRootSignatureDesc
 	{
-		static CD3DX12_ROOT_PARAMETER1 TableSlots[7];
-		static CD3DX12_DESCRIPTOR_RANGE1 DescriptorRanges[4];
-		static CD3DX12_VERSIONED_ROOT_SIGNATURE_DESC RootDesc;
+		CD3DX12_ROOT_PARAMETER1 TableSlots[7];
+		CD3DX12_DESCRIPTOR_RANGE1 DescriptorRanges[4];
+		CD3DX12_VERSIONED_ROOT_SIGNATURE_DESC RootDesc;
+	};
 
+	void InitStaticRayTracingRootSignatureDesc(FStaticRayTracingRootSignatureDesc& Desc, bool bLocalRootSignature, D3D12_ROOT_SIGNATURE_FLAGS BaseFlags, bool bBindlessResources, bool bBindlessSamplers)
+	{
 		uint32 SlotIndex = 0;
 
 		if (bLocalRootSignature)
 		{
 			uint32 NumConstants = 4; // sizeof(FHitGroupSystemRootConstants) / sizeof(uint32);
 
-			TableSlots[SlotIndex++].InitAsShaderResourceView(RAY_TRACING_SYSTEM_INDEXBUFFER_REGISTER, RAY_TRACING_REGISTER_SPACE_SYSTEM);
-			TableSlots[SlotIndex++].InitAsShaderResourceView(RAY_TRACING_SYSTEM_VERTEXBUFFER_REGISTER, RAY_TRACING_REGISTER_SPACE_SYSTEM);
-			TableSlots[SlotIndex++].InitAsConstants(NumConstants, RAY_TRACING_SYSTEM_ROOTCONSTANT_REGISTER, RAY_TRACING_REGISTER_SPACE_SYSTEM);
+			Desc.TableSlots[SlotIndex++].InitAsShaderResourceView(RAY_TRACING_SYSTEM_INDEXBUFFER_REGISTER, RAY_TRACING_REGISTER_SPACE_SYSTEM);
+			Desc.TableSlots[SlotIndex++].InitAsShaderResourceView(RAY_TRACING_SYSTEM_VERTEXBUFFER_REGISTER, RAY_TRACING_REGISTER_SPACE_SYSTEM);
+			Desc.TableSlots[SlotIndex++].InitAsConstants(NumConstants, RAY_TRACING_SYSTEM_ROOTCONSTANT_REGISTER, RAY_TRACING_REGISTER_SPACE_SYSTEM);
 		}
 
 		uint32 BindingSpace = bLocalRootSignature ? RAY_TRACING_REGISTER_SPACE_LOCAL : RAY_TRACING_REGISTER_SPACE_GLOBAL;
 
 		// Table ranges
 
-		DescriptorRanges[0].Init(D3D12_DESCRIPTOR_RANGE_TYPE_SRV, MAX_SRVS, 0, BindingSpace, D3D12_DESCRIPTOR_RANGE_FLAG_DATA_STATIC_WHILE_SET_AT_EXECUTE | D3D12_DESCRIPTOR_RANGE_FLAG_DESCRIPTORS_VOLATILE);
-		DescriptorRanges[1].Init(D3D12_DESCRIPTOR_RANGE_TYPE_CBV, MAX_CBS, 0, BindingSpace, D3D12_DESCRIPTOR_RANGE_FLAG_DATA_STATIC_WHILE_SET_AT_EXECUTE | D3D12_DESCRIPTOR_RANGE_FLAG_DESCRIPTORS_VOLATILE);
-		DescriptorRanges[2].Init(D3D12_DESCRIPTOR_RANGE_TYPE_SAMPLER, MAX_SAMPLERS, 0, BindingSpace, D3D12_DESCRIPTOR_RANGE_FLAG_DESCRIPTORS_VOLATILE);
-		DescriptorRanges[3].Init(D3D12_DESCRIPTOR_RANGE_TYPE_UAV, MAX_UAVS, 0, BindingSpace, D3D12_DESCRIPTOR_RANGE_FLAG_DATA_STATIC_WHILE_SET_AT_EXECUTE | D3D12_DESCRIPTOR_RANGE_FLAG_DESCRIPTORS_VOLATILE);
+		Desc.DescriptorRanges[0].Init(D3D12_DESCRIPTOR_RANGE_TYPE_SRV, MAX_SRVS, 0, BindingSpace, D3D12_DESCRIPTOR_RANGE_FLAG_DATA_STATIC_WHILE_SET_AT_EXECUTE | D3D12_DESCRIPTOR_RANGE_FLAG_DESCRIPTORS_VOLATILE);
+		Desc.DescriptorRanges[1].Init(D3D12_DESCRIPTOR_RANGE_TYPE_CBV, MAX_CBS, 0, BindingSpace, D3D12_DESCRIPTOR_RANGE_FLAG_DATA_STATIC_WHILE_SET_AT_EXECUTE | D3D12_DESCRIPTOR_RANGE_FLAG_DESCRIPTORS_VOLATILE);
+		Desc.DescriptorRanges[2].Init(D3D12_DESCRIPTOR_RANGE_TYPE_SAMPLER, MAX_SAMPLERS, 0, BindingSpace, D3D12_DESCRIPTOR_RANGE_FLAG_DESCRIPTORS_VOLATILE);
+		Desc.DescriptorRanges[3].Init(D3D12_DESCRIPTOR_RANGE_TYPE_UAV, MAX_UAVS, 0, BindingSpace, D3D12_DESCRIPTOR_RANGE_FLAG_DATA_STATIC_WHILE_SET_AT_EXECUTE | D3D12_DESCRIPTOR_RANGE_FLAG_DESCRIPTORS_VOLATILE);
 
 		// Table slots
 
-		TableSlots[SlotIndex++].InitAsDescriptorTable(1, &DescriptorRanges[0], D3D12_SHADER_VISIBILITY_ALL); // SRV
-		TableSlots[SlotIndex++].InitAsDescriptorTable(1, &DescriptorRanges[1], D3D12_SHADER_VISIBILITY_ALL); // CBV
-		TableSlots[SlotIndex++].InitAsDescriptorTable(1, &DescriptorRanges[2], D3D12_SHADER_VISIBILITY_ALL); // Sampler
-		TableSlots[SlotIndex++].InitAsDescriptorTable(1, &DescriptorRanges[3], D3D12_SHADER_VISIBILITY_ALL); // UAV
+		Desc.TableSlots[SlotIndex++].InitAsDescriptorTable(1, &Desc.DescriptorRanges[0], D3D12_SHADER_VISIBILITY_ALL); // SRV
+		Desc.TableSlots[SlotIndex++].InitAsDescriptorTable(1, &Desc.DescriptorRanges[1], D3D12_SHADER_VISIBILITY_ALL); // CBV
+		Desc.TableSlots[SlotIndex++].InitAsDescriptorTable(1, &Desc.DescriptorRanges[2], D3D12_SHADER_VISIBILITY_ALL); // Sampler
+		Desc.TableSlots[SlotIndex++].InitAsDescriptorTable(1, &Desc.DescriptorRanges[3], D3D12_SHADER_VISIBILITY_ALL); // UAV
+
+		D3D12_ROOT_SIGNATURE_FLAGS Flags = BaseFlags;
 
-		RootDesc.Init_1_1(SlotIndex, TableSlots, 0, nullptr, Flags);
+		if (!bLocalRootSignature && bBindlessResources)
+		{
+			Flags |= D3D12_ROOT_SIGNATURE_FLAG_CBV_SRV_UAV_HEAP_DIRECTLY_INDEXED;
+		}
+
+		if (!bLocalRootSignature && bBindlessSamplers)
+		{
+			Flags |= D3D12_ROOT_SIGNATURE_FLAG_SAMPLER_HEAP_DIRECTLY_INDEXED;
+		}
 
-		return RootDesc;
+		Desc.RootDesc.Init_1_1(SlotIndex, Desc.TableSlots, 0, nullptr, Flags);
 	}
 }
