Addressing static analysis issues in D3D12.

diff --git a/Engine/Source/Runtime/D3D12RHI/Private/D3D12RootSignature.cpp b/Engine/Source/Runtime/D3D12RHI/Private/D3D12RootSignature.cpp
--- a/Engine/Source/Runtime/D3D12RHI/Private/D3D12RootSignature.cpp
+++ b/Engine/Source/Runtime/D3D12RHI/Private/D3D12RootSignature.cpp
@@ -429,6 +429,9 @@ void FD3D12RootSignature::InitStaticGraphicsRootSignature(bool bBindlessResource
 	CD3DX12_ROOT_PARAMETER1 TableSlots[DescriptorTableCount];
 	CD3DX12_DESCRIPTOR_RANGE1 DescriptorRanges[DescriptorTableCount];
 
+	FMemory::Memzero(TableSlots);
+	FMemory::Memzero(DescriptorRanges);
+
 	for (uint32 i = 0; i < DescriptorTableCount; i++)
 	{
 		DescriptorRanges[i].Init(
diff --git a/Engine/Source/Runtime/D3D12RHI/Private/Windows/WindowsD3D12Device.cpp b/Engine/Source/Runtime/D3D12RHI/Private/Windows/WindowsD3D12Device.cpp
--- a/Engine/Source/Runtime/D3D12RHI/Private/Windows/WindowsD3D12Device.cpp
+++ b/Engine/Source/Runtime/D3D12RHI/Private/Windows/WindowsD3D12Device.cpp
@@ -326,6 +326,7 @@ static INTCExtensionContext* CreateIntelExtensionsContext(ID3D12Device* Device,
 	{
 		for (uint32_t i = 0; i < SupportedExtensionsVersionCount; i++)
 		{
+			CA_SUPPRESS(6385);
 			if ((SupportedExtensionsVersions[i].HWFeatureLevel >= AtomicsRequiredVersion.HWFeatureLevel) &&
 				(SupportedExtensionsVersions[i].APIVersion >= AtomicsRequiredVersion.APIVersion) &&
 				(SupportedExtensionsVersions[i].Revision >= AtomicsRequiredVersion.Revision))
