Static analysis fix.

diff --git a/Engine/Source/Runtime/D3D12RHI/Private/D3D12RootSignature.cpp b/Engine/Source/Runtime/D3D12RHI/Private/D3D12RootSignature.cpp
--- a/Engine/Source/Runtime/D3D12RHI/Private/D3D12RootSignature.cpp
+++ b/Engine/Source/Runtime/D3D12RHI/Private/D3D12RootSignature.cpp
@@ -467,6 +467,9 @@ void FD3D12RootSignature::InitStaticComputeRootSignatureDesc(bool bBindlessResou
 	CD3DX12_ROOT_PARAMETER1 TableSlots[DescriptorTableCount];
 	CD3DX12_DESCRIPTOR_RANGE1 DescriptorRanges[DescriptorTableCount];
 
+	FMemory::Memzero(TableSlots);
+	FMemory::Memzero(DescriptorRanges);
+
 	uint32 RangeIndex = 0;
 	DescriptorRanges[RangeIndex].Init(D3D12_DESCRIPTOR_RANGE_TYPE_SRV, MAX_SRVS, 0, 0, D3D12ShaderUtils::StaticRootSignatureConstants::SRVDescriptorRangeFlags);
 	TableSlots[RangeIndex].InitAsDescriptorTable(1, &DescriptorRanges[RangeIndex], D3D12_SHADER_VISIBILITY_ALL);
