Fixing asserts on platforms that don't support RayTracing.

diff --git a/Engine/Source/Runtime/D3D12RHI/Private/D3D12Adapter.cpp b/Engine/Source/Runtime/D3D12RHI/Private/D3D12Adapter.cpp
--- a/Engine/Source/Runtime/D3D12RHI/Private/D3D12Adapter.cpp
+++ b/Engine/Source/Runtime/D3D12RHI/Private/D3D12Adapter.cpp
@@ -1229,8 +1229,10 @@ void FD3D12Adapter::InitializeDevices()
 
 		StaticGraphicsRootSignature.InitStaticGraphicsRootSignature(bBindlessResourcesGraphics, bBindlessSamplersGraphics);
 		StaticComputeRootSignature.InitStaticComputeRootSignatureDesc(bBindlessResourcesGraphics, bBindlessSamplersGraphics);
+#if D3D12_RHI_RAYTRACING
 		StaticRayTracingGlobalRootSignature.InitStaticRayTracingGlobalRootSignatureDesc(bBindlessResourcesRayTracing, bBindlessSamplersRayTracing);
 		StaticRayTracingLocalRootSignature.InitStaticRayTracingLocalRootSignatureDesc();
+#endif
 #endif
 	}
 }
