D3D12 Root Signature creation unification
- Static root signatures were defined in two different code paths that are both very hard-coded
- Refactored FRootSignatureCreator classes to fully support our text based and D3D struct based Root Signature creation.
- FRootSignatureCreator classes improved to also support removal of table types when certain bindless modes are enabled.
- RayTracing Root Signatures also refactored to go down the FRootSignatureCreator paths.

diff --git a/Engine/Source/Runtime/D3D12RHI/Private/D3D12Adapter.cpp b/Engine/Source/Runtime/D3D12RHI/Private/D3D12Adapter.cpp
--- a/Engine/Source/Runtime/D3D12RHI/Private/D3D12Adapter.cpp
+++ b/Engine/Source/Runtime/D3D12RHI/Private/D3D12Adapter.cpp
@@ -1221,19 +1221,37 @@ void FD3D12Adapter::InitializeDevices()
 #endif
 
 #if USE_STATIC_ROOT_SIGNATURE
-		const bool bBindlessResourcesGraphics = (BindlessResourcesConfig == ERHIBindlessConfiguration::AllShaders);
-		const bool bBindlessSamplersGraphics = (BindlessSamplersConfig == ERHIBindlessConfiguration::AllShaders);
+		ED3D12RootSignatureFlags GraphicsFlags{};
+#if PLATFORM_SUPPORTS_MESH_SHADERS
+		EnumAddFlags(GraphicsFlags, ED3D12RootSignatureFlags::AllowMeshShaders);
+#endif
+		if (BindlessResourcesConfig == ERHIBindlessConfiguration::AllShaders)
+		{
+			EnumAddFlags(GraphicsFlags, ED3D12RootSignatureFlags::BindlessResources);
+		}
+		if (BindlessSamplersConfig == ERHIBindlessConfiguration::AllShaders)
+		{
+			EnumAddFlags(GraphicsFlags, ED3D12RootSignatureFlags::BindlessSamplers);
+		}
 
-		const bool bBindlessResourcesRayTracing = (BindlessResourcesConfig != ERHIBindlessConfiguration::Disabled);
-		const bool bBindlessSamplersRayTracing = (BindlessSamplersConfig != ERHIBindlessConfiguration::Disabled);
+		StaticGraphicsRootSignature.InitStaticGraphicsRootSignature(GraphicsFlags);
+		StaticComputeRootSignature.InitStaticComputeRootSignatureDesc(GraphicsFlags);
 
-		StaticGraphicsRootSignature.InitStaticGraphicsRootSignature(bBindlessResourcesGraphics, bBindlessSamplersGraphics);
-		StaticComputeRootSignature.InitStaticComputeRootSignatureDesc(bBindlessResourcesGraphics, bBindlessSamplersGraphics);
 #if D3D12_RHI_RAYTRACING
-		StaticRayTracingGlobalRootSignature.InitStaticRayTracingGlobalRootSignatureDesc(bBindlessResourcesRayTracing, bBindlessSamplersRayTracing);
+		ED3D12RootSignatureFlags RayTracingFlags{};
+		if (BindlessResourcesConfig != ERHIBindlessConfiguration::Disabled)
+		{
+			EnumAddFlags(RayTracingFlags, ED3D12RootSignatureFlags::BindlessResources);
+		}
+		if (BindlessSamplersConfig != ERHIBindlessConfiguration::Disabled)
+		{
+			EnumAddFlags(RayTracingFlags, ED3D12RootSignatureFlags::BindlessSamplers);
+		}
+
+		StaticRayTracingGlobalRootSignature.InitStaticRayTracingGlobalRootSignatureDesc(RayTracingFlags);
 		StaticRayTracingLocalRootSignature.InitStaticRayTracingLocalRootSignatureDesc();
 #endif
-#endif
+#endif // USE_STATIC_ROOT_SIGNATURE
 	}
 }
 
diff --git a/Engine/Source/Runtime/D3D12RHI/Private/D3D12RootSignature.cpp b/Engine/Source/Runtime/D3D12RHI/Private/D3D12RootSignature.cpp
--- a/Engine/Source/Runtime/D3D12RHI/Private/D3D12RootSignature.cpp
+++ b/Engine/Source/Runtime/D3D12RHI/Private/D3D12RootSignature.cpp
@@ -8,10 +8,6 @@
 #include "D3D12RootSignatureDefinitions.h"
 #include "RayTracingBuiltInResources.h"
 
-#if D3D12_RHI_RAYTRACING
-#include "D3D12RayTracingRootSignature.h"
-#endif // D3D12_RHI_RAYTRACING
-
 #ifndef FD3D12_ROOT_SIGNATURE_FLAG_GLOBAL_ROOT_SIGNATURE
 #define FD3D12_ROOT_SIGNATURE_FLAG_GLOBAL_ROOT_SIGNATURE D3D12_ROOT_SIGNATURE_FLAG_NONE
 #endif
@@ -388,137 +384,35 @@ FD3D12RootSignatureDesc::FD3D12RootSignatureDesc(const FD3D12QuantizedBoundShade
 
 }
 
-void FD3D12RootSignature::InitStaticGraphicsRootSignature(bool bBindlessResources, bool bBindlessSamplers)
+void FD3D12RootSignature::InitStaticGraphicsRootSignature(ED3D12RootSignatureFlags InFlags)
 {
-	// TODO: Support vendor extensions for static root signatures?
-	struct
-	{
-		D3D12_SHADER_VISIBILITY Vis;
-		D3D12_DESCRIPTOR_RANGE_TYPE Type;
-		uint32 Count;
-		uint32 BaseShaderReg;
-		D3D12_DESCRIPTOR_RANGE_FLAGS Flags;
-	} RangeDesc[] =
-	{
-		{ D3D12_SHADER_VISIBILITY_PIXEL, D3D12_DESCRIPTOR_RANGE_TYPE_SRV, MAX_SRVS, 0, D3D12ShaderUtils::StaticRootSignatureConstants::SRVDescriptorRangeFlags },
-		{ D3D12_SHADER_VISIBILITY_PIXEL, D3D12_DESCRIPTOR_RANGE_TYPE_CBV, MAX_CBS, 0, D3D12ShaderUtils::StaticRootSignatureConstants::CBVDescriptorRangeFlags },
-		{ D3D12_SHADER_VISIBILITY_PIXEL, D3D12_DESCRIPTOR_RANGE_TYPE_SAMPLER, MAX_SAMPLERS, 0, D3D12ShaderUtils::StaticRootSignatureConstants::SamplerDescriptorRangeFlags },
-
-		{ D3D12_SHADER_VISIBILITY_VERTEX, D3D12_DESCRIPTOR_RANGE_TYPE_SRV, MAX_SRVS, 0, D3D12ShaderUtils::StaticRootSignatureConstants::SRVDescriptorRangeFlags },
-		{ D3D12_SHADER_VISIBILITY_VERTEX, D3D12_DESCRIPTOR_RANGE_TYPE_CBV, MAX_CBS, 0, D3D12ShaderUtils::StaticRootSignatureConstants::CBVDescriptorRangeFlags },
-		{ D3D12_SHADER_VISIBILITY_VERTEX, D3D12_DESCRIPTOR_RANGE_TYPE_SAMPLER, MAX_SAMPLERS, 0, D3D12ShaderUtils::StaticRootSignatureConstants::SamplerDescriptorRangeFlags },
-
-		{ D3D12_SHADER_VISIBILITY_GEOMETRY, D3D12_DESCRIPTOR_RANGE_TYPE_SRV, MAX_SRVS, 0, D3D12ShaderUtils::StaticRootSignatureConstants::SRVDescriptorRangeFlags },
-		{ D3D12_SHADER_VISIBILITY_GEOMETRY, D3D12_DESCRIPTOR_RANGE_TYPE_CBV, MAX_CBS, 0, D3D12ShaderUtils::StaticRootSignatureConstants::CBVDescriptorRangeFlags },
-		{ D3D12_SHADER_VISIBILITY_GEOMETRY, D3D12_DESCRIPTOR_RANGE_TYPE_SAMPLER, MAX_SAMPLERS, 0, D3D12ShaderUtils::StaticRootSignatureConstants::SamplerDescriptorRangeFlags },
-
-#if PLATFORM_SUPPORTS_MESH_SHADERS
-		{ D3D12_SHADER_VISIBILITY_MESH, D3D12_DESCRIPTOR_RANGE_TYPE_SRV, MAX_SRVS, 0, D3D12ShaderUtils::StaticRootSignatureConstants::SRVDescriptorRangeFlags },
-		{ D3D12_SHADER_VISIBILITY_MESH, D3D12_DESCRIPTOR_RANGE_TYPE_CBV, MAX_CBS, 0, D3D12ShaderUtils::StaticRootSignatureConstants::CBVDescriptorRangeFlags },
-		{ D3D12_SHADER_VISIBILITY_MESH, D3D12_DESCRIPTOR_RANGE_TYPE_SAMPLER, MAX_SAMPLERS, 0, D3D12ShaderUtils::StaticRootSignatureConstants::SamplerDescriptorRangeFlags },
-
-		{ D3D12_SHADER_VISIBILITY_AMPLIFICATION, D3D12_DESCRIPTOR_RANGE_TYPE_SRV, MAX_SRVS, 0, D3D12ShaderUtils::StaticRootSignatureConstants::SRVDescriptorRangeFlags },
-		{ D3D12_SHADER_VISIBILITY_AMPLIFICATION, D3D12_DESCRIPTOR_RANGE_TYPE_CBV, MAX_CBS, 0, D3D12ShaderUtils::StaticRootSignatureConstants::CBVDescriptorRangeFlags },
-		{ D3D12_SHADER_VISIBILITY_AMPLIFICATION, D3D12_DESCRIPTOR_RANGE_TYPE_SAMPLER, MAX_SAMPLERS, 0, D3D12ShaderUtils::StaticRootSignatureConstants::SamplerDescriptorRangeFlags },
-#endif // PLATFORM_SUPPORTS_MESH_SHADERS
-
-		{ D3D12_SHADER_VISIBILITY_ALL, D3D12_DESCRIPTOR_RANGE_TYPE_UAV, MAX_UAVS, 0, D3D12ShaderUtils::StaticRootSignatureConstants::UAVDescriptorRangeFlags },
-	};
-
-	const uint32 DescriptorTableCount = UE_ARRAY_COUNT(RangeDesc);
-	CD3DX12_ROOT_PARAMETER1 TableSlots[DescriptorTableCount];
-	CD3DX12_DESCRIPTOR_RANGE1 DescriptorRanges[DescriptorTableCount];
-
-	FMemory::Memzero(TableSlots);
-	FMemory::Memzero(DescriptorRanges);
-
-	for (uint32 i = 0; i < DescriptorTableCount; i++)
-	{
-		DescriptorRanges[i].Init(
-			RangeDesc[i].Type,
-			RangeDesc[i].Count,
-			RangeDesc[i].BaseShaderReg,
-			0u,
-			RangeDesc[i].Flags
-		);
-
-		TableSlots[i].InitAsDescriptorTable(1, &DescriptorRanges[i], RangeDesc[i].Vis);
-	}
-
-	D3D12_ROOT_SIGNATURE_FLAGS Flags = D3D12_ROOT_SIGNATURE_FLAG_ALLOW_INPUT_ASSEMBLER_INPUT_LAYOUT;
-	if (bBindlessResources)
-	{
-		Flags |= D3D12_ROOT_SIGNATURE_FLAG_CBV_SRV_UAV_HEAP_DIRECTLY_INDEXED;
-	}
-	if (bBindlessSamplers)
-	{
-		Flags |= D3D12_ROOT_SIGNATURE_FLAG_SAMPLER_HEAP_DIRECTLY_INDEXED;
-	}
-
-	CD3DX12_VERSIONED_ROOT_SIGNATURE_DESC RootDesc(DescriptorTableCount, TableSlots, 0, nullptr, Flags);
-	Init(RootDesc);
+	D3D12ShaderUtils::FBinaryRootSignatureCreator Creator;
+	D3D12ShaderUtils::CreateGfxRootSignature(Creator, InFlags);
+	Init(Creator.Finalize());
 }
 
-void FD3D12RootSignature::InitStaticComputeRootSignatureDesc(bool bBindlessResources, bool bBindlessSamplers)
+void FD3D12RootSignature::InitStaticComputeRootSignatureDesc(ED3D12RootSignatureFlags InFlags)
 {
-	// TODO: Support vendor extensions for static root signatures?
-
-	const uint32 DescriptorTableCount = 4;
-	CD3DX12_ROOT_PARAMETER1 TableSlots[DescriptorTableCount];
-	CD3DX12_DESCRIPTOR_RANGE1 DescriptorRanges[DescriptorTableCount];
-
-	FMemory::Memzero(TableSlots);
-	FMemory::Memzero(DescriptorRanges);
-
-	uint32 RangeIndex = 0;
-	DescriptorRanges[RangeIndex].Init(D3D12_DESCRIPTOR_RANGE_TYPE_SRV, MAX_SRVS, 0, 0, D3D12ShaderUtils::StaticRootSignatureConstants::SRVDescriptorRangeFlags);
-	TableSlots[RangeIndex].InitAsDescriptorTable(1, &DescriptorRanges[RangeIndex], D3D12_SHADER_VISIBILITY_ALL);
-	++RangeIndex;
-	DescriptorRanges[RangeIndex].Init(D3D12_DESCRIPTOR_RANGE_TYPE_CBV, MAX_CBS, 0, 0, D3D12ShaderUtils::StaticRootSignatureConstants::CBVDescriptorRangeFlags);
-	TableSlots[RangeIndex].InitAsDescriptorTable(1, &DescriptorRanges[RangeIndex], D3D12_SHADER_VISIBILITY_ALL);
-	++RangeIndex;
-	DescriptorRanges[RangeIndex].Init(D3D12_DESCRIPTOR_RANGE_TYPE_SAMPLER, MAX_SAMPLERS, 0, 0, D3D12ShaderUtils::StaticRootSignatureConstants::SamplerDescriptorRangeFlags);
-	TableSlots[RangeIndex].InitAsDescriptorTable(1, &DescriptorRanges[RangeIndex], D3D12_SHADER_VISIBILITY_ALL);
-	++RangeIndex;
-	DescriptorRanges[RangeIndex].Init(D3D12_DESCRIPTOR_RANGE_TYPE_UAV, MAX_UAVS, 0, 0, D3D12ShaderUtils::StaticRootSignatureConstants::UAVDescriptorRangeFlags);
-	TableSlots[RangeIndex].InitAsDescriptorTable(1, &DescriptorRanges[RangeIndex], D3D12_SHADER_VISIBILITY_ALL);
-	++RangeIndex;
-
-	D3D12_ROOT_SIGNATURE_FLAGS Flags = D3D12_ROOT_SIGNATURE_FLAG_NONE;
-	if (bBindlessResources)
-	{
-		Flags |= D3D12_ROOT_SIGNATURE_FLAG_CBV_SRV_UAV_HEAP_DIRECTLY_INDEXED;
-	}
-	if (bBindlessSamplers)
-	{
-		Flags |= D3D12_ROOT_SIGNATURE_FLAG_SAMPLER_HEAP_DIRECTLY_INDEXED;
-	}
-
-	CD3DX12_VERSIONED_ROOT_SIGNATURE_DESC RootDesc(RangeIndex, TableSlots, 0, nullptr, Flags);
-	Init(RootDesc);
+	D3D12ShaderUtils::FBinaryRootSignatureCreator Creator;
+	D3D12ShaderUtils::CreateComputeRootSignature(Creator, InFlags);
+	Init(Creator.Finalize());
 }
 
-void FD3D12RootSignature::InitStaticRayTracingGlobalRootSignatureDesc(bool bBindlessResources, bool bBindlessSamplers)
-{
 #if D3D12_RHI_RAYTRACING
-	D3D12ShaderUtils::FStaticRayTracingRootSignatureDesc Desc{};
-	D3D12ShaderUtils::InitStaticRayTracingRootSignatureDesc(Desc, false, FD3D12_ROOT_SIGNATURE_FLAG_GLOBAL_ROOT_SIGNATURE, bBindlessResources, bBindlessSamplers);
-	Init(Desc.RootDesc, 1 /*RAY_TRACING_REGISTER_SPACE_GLOBAL*/);
-#else
-	checkNoEntry();
-#endif
+void FD3D12RootSignature::InitStaticRayTracingGlobalRootSignatureDesc(ED3D12RootSignatureFlags InFlags)
+{
+	D3D12ShaderUtils::FBinaryRootSignatureCreator Creator;
+	D3D12ShaderUtils::CreateRayTracingSignature(Creator, false, FD3D12_ROOT_SIGNATURE_FLAG_GLOBAL_ROOT_SIGNATURE, InFlags);
+	Init(Creator.Finalize(), RAY_TRACING_REGISTER_SPACE_GLOBAL);
 }
 
 void FD3D12RootSignature::InitStaticRayTracingLocalRootSignatureDesc()
 {
-#if D3D12_RHI_RAYTRACING
-	D3D12ShaderUtils::FStaticRayTracingRootSignatureDesc Desc{};
-	D3D12ShaderUtils::InitStaticRayTracingRootSignatureDesc(Desc, true, D3D12_ROOT_SIGNATURE_FLAG_LOCAL_ROOT_SIGNATURE, false, false);
-	Init(Desc.RootDesc, 0 /*RAY_TRACING_REGISTER_SPACE_LOCAL*/);
-#else
-	checkNoEntry();
-#endif
+	D3D12ShaderUtils::FBinaryRootSignatureCreator Creator;
+	D3D12ShaderUtils::CreateRayTracingSignature(Creator, true, D3D12_ROOT_SIGNATURE_FLAG_LOCAL_ROOT_SIGNATURE, ED3D12RootSignatureFlags::None);
+	Init(Creator.Finalize(), RAY_TRACING_REGISTER_SPACE_LOCAL);
 }
+#endif // D3D12_RHI_RAYTRACING
 
 void FD3D12RootSignature::Init(const FD3D12QuantizedBoundShaderState& InQBSS)
 {
@@ -567,27 +461,6 @@ void FD3D12RootSignature::Init(const D3D12_VERSIONED_ROOT_SIGNATURE_DESC& InDesc
 	// TODO: Analyze vendor extension space?
 }
 
-void FD3D12RootSignature::Init(ID3DBlob* const InBlob, uint32 BindingSpace)
-{
-	ID3D12Device* Device = GetParentAdapter()->GetD3DDevice();
-
-	// Save the blob
-	RootSignatureBlob = InBlob;
-
-	// Deserialize to get the desc.
-	TRefCountPtr<ID3D12VersionedRootSignatureDeserializer> Deserializer;
-	VERIFYD3D12RESULT(D3D12CreateVersionedRootSignatureDeserializer(RootSignatureBlob->GetBufferPointer(), RootSignatureBlob->GetBufferSize(), IID_PPV_ARGS(Deserializer.GetInitReference())));
-
-	// Create and analyze the root signature.
-	VERIFYD3D12RESULT(Device->CreateRootSignature(FRHIGPUMask::All().GetNative(),
-		RootSignatureBlob->GetBufferPointer(),
-		RootSignatureBlob->GetBufferSize(),
-		IID_PPV_ARGS(RootSignature.GetInitReference())));
-
-	AnalyzeSignature(*Deserializer->GetUnconvertedRootSignatureDesc(), BindingSpace);
-	// TODO: Analyze vendor extension space?
-}
-
 void FD3D12RootSignature::AnalyzeSignature(const D3D12_VERSIONED_ROOT_SIGNATURE_DESC& Desc, uint32 BindingSpace)
 {
 	switch (Desc.Version)
diff --git a/Engine/Source/Runtime/D3D12RHI/Private/D3D12RootSignature.h b/Engine/Source/Runtime/D3D12RHI/Private/D3D12RootSignature.h
--- a/Engine/Source/Runtime/D3D12RHI/Private/D3D12RootSignature.h
+++ b/Engine/Source/Runtime/D3D12RHI/Private/D3D12RootSignature.h
@@ -4,6 +4,8 @@
 	D3D12RootSignature.h: D3D12 Root Signatures
 =============================================================================*/
 
+#include "D3D12RootSignatureDefinitions.h"
+
 // Root parameter keys grouped by visibility.
 enum ERootParameterKeys
 {
@@ -81,25 +83,16 @@ public:
 	{
 		Init(InQBSS);
 	}
-	explicit FD3D12RootSignature(FD3D12Adapter* InParent, const D3D12_VERSIONED_ROOT_SIGNATURE_DESC& InDesc, uint32 BindingSpace = 0)
-		: FD3D12AdapterChild(InParent)
-	{
-		Init(InDesc, BindingSpace);
-	}
-	explicit FD3D12RootSignature(FD3D12Adapter* InParent, ID3DBlob* const InBlob, uint32 BindingSpace = 0)
-		: FD3D12AdapterChild(InParent)
-	{
-		Init(InBlob, BindingSpace);
-	}
 
 	void Init(const FD3D12QuantizedBoundShaderState& InQBSS);
 	void Init(const D3D12_VERSIONED_ROOT_SIGNATURE_DESC& InDesc, uint32 BindingSpace = 0);
-	void Init(ID3DBlob* const InBlob, uint32 BindingSpace = 0);
 
-	void InitStaticGraphicsRootSignature(bool bBindlessResources, bool bBindlessSamplers);
-	void InitStaticComputeRootSignatureDesc(bool bBindlessResources, bool bBindlessSamplers);
-	void InitStaticRayTracingGlobalRootSignatureDesc(bool bBindlessResources, bool bBindlessSamplers);
+	void InitStaticGraphicsRootSignature(ED3D12RootSignatureFlags InFlags);
+	void InitStaticComputeRootSignatureDesc(ED3D12RootSignatureFlags InFlags);
+#if D3D12_RHI_RAYTRACING
+	void InitStaticRayTracingGlobalRootSignatureDesc(ED3D12RootSignatureFlags InFlags);
 	void InitStaticRayTracingLocalRootSignatureDesc();
+#endif
 
 	ID3D12RootSignature* GetRootSignature() const { return RootSignature.GetReference(); }
 	ID3DBlob* GetRootSignatureBlob() const { return RootSignatureBlob.GetReference(); }
diff --git a/Engine/Source/Runtime/D3D12RHI/Public/D3D12RayTracingRootSignature.h b/Engine/Source/Runtime/D3D12RHI/Public/D3D12RayTracingRootSignature.h
deleted file mode 100644
diff --git a/Engine/Source/Runtime/D3D12RHI/Public/D3D12RootSignatureDefinitions.h b/Engine/Source/Runtime/D3D12RHI/Public/D3D12RootSignatureDefinitions.h
--- a/Engine/Source/Runtime/D3D12RHI/Public/D3D12RootSignatureDefinitions.h
+++ b/Engine/Source/Runtime/D3D12RHI/Public/D3D12RootSignatureDefinitions.h
@@ -1,11 +1,19 @@
 // Copyright Epic Games, Inc. All Rights Reserved.
 
-/*=============================================================================
-D3D12RootSignatureDefinitions.h: D3D12 utilities for Root signatures.
-=============================================================================*/
-
 #pragma once
 
+#include "RayTracingBuiltInResources.h"
+
+enum class ED3D12RootSignatureFlags
+{
+	None = 0,
+	AllowMeshShaders = 1 << 0,
+	InputAssembler = 1 << 1,
+	BindlessResources = 1 << 2,
+	BindlessSamplers = 1 << 3,
+};
+ENUM_CLASS_FLAGS(ED3D12RootSignatureFlags)
+
 namespace D3D12ShaderUtils
 {
 	namespace StaticRootSignatureConstants
@@ -17,7 +25,7 @@ namespace D3D12ShaderUtils
 		const D3D12_DESCRIPTOR_RANGE_FLAGS SamplerDescriptorRangeFlags = D3D12_DESCRIPTOR_RANGE_FLAG_DESCRIPTORS_VOLATILE;
 	}
 
-	enum class ERootSignatureType
+	enum class ERootSignatureRangeType
 	{
 		CBV,
 		SRV,
@@ -25,321 +33,379 @@ namespace D3D12ShaderUtils
 		Sampler,
 	};
 
-	// Simple base class to help write out a root signature (subclass to generate either to a binary struct or a #define)
-	struct FRootSignatureCreator
+	enum class ERootSignatureVisibility
 	{
-		virtual ~FRootSignatureCreator() { }
-
-		virtual FRootSignatureCreator& Reset() = 0;
-		virtual FRootSignatureCreator& AddRootFlag(D3D12_ROOT_SIGNATURE_FLAGS Flag) = 0;
-		virtual FRootSignatureCreator& AddTable(EShaderFrequency Stage, ERootSignatureType Type, int32 NumDescriptors/*, bool bAppend*/) = 0;
+		Vertex,
+		Pixel,
+		Geometry,
+		Mesh,
+		Amplification,
+		All
+	};
 
-		inline D3D12_DESCRIPTOR_RANGE_TYPE GetD3D12DescriptorRangeType(ERootSignatureType Type)
+	inline D3D12_DESCRIPTOR_RANGE_TYPE GetD3D12DescriptorRangeType(ERootSignatureRangeType Type)
+	{
+		switch (Type)
 		{
-			switch (Type)
-			{
-			case ERootSignatureType::SRV:
-				return D3D12_DESCRIPTOR_RANGE_TYPE_SRV;
-			case ERootSignatureType::UAV:
-				return D3D12_DESCRIPTOR_RANGE_TYPE_UAV;
-			case ERootSignatureType::Sampler:
-				return D3D12_DESCRIPTOR_RANGE_TYPE_SAMPLER;
-			default:
-				return D3D12_DESCRIPTOR_RANGE_TYPE_CBV;
-			}
+		case ERootSignatureRangeType::SRV:
+			return D3D12_DESCRIPTOR_RANGE_TYPE_SRV;
+		case ERootSignatureRangeType::UAV:
+			return D3D12_DESCRIPTOR_RANGE_TYPE_UAV;
+		case ERootSignatureRangeType::Sampler:
+			return D3D12_DESCRIPTOR_RANGE_TYPE_SAMPLER;
+		default:
+			return D3D12_DESCRIPTOR_RANGE_TYPE_CBV;
 		}
+	}
 
-		inline D3D12_DESCRIPTOR_RANGE_FLAGS GetD3D12DescriptorRangeFlags(ERootSignatureType Type)
+	inline D3D12_DESCRIPTOR_RANGE_FLAGS GetD3D12DescriptorRangeFlags(ERootSignatureRangeType Type)
+	{
+		switch (Type)
 		{
-			switch (Type)
-			{
-			case ERootSignatureType::SRV:
-				return D3D12ShaderUtils::StaticRootSignatureConstants::SRVDescriptorRangeFlags;
-			case ERootSignatureType::CBV:
-				return D3D12ShaderUtils::StaticRootSignatureConstants::CBVDescriptorRangeFlags;
-			case ERootSignatureType::UAV:
-				return D3D12ShaderUtils::StaticRootSignatureConstants::UAVDescriptorRangeFlags;
-			case ERootSignatureType::Sampler:
-				return D3D12ShaderUtils::StaticRootSignatureConstants::SamplerDescriptorRangeFlags;
-			default:
-				return D3D12_DESCRIPTOR_RANGE_FLAG_NONE;
-			}
+		case ERootSignatureRangeType::SRV:
+			return D3D12ShaderUtils::StaticRootSignatureConstants::SRVDescriptorRangeFlags;
+		case ERootSignatureRangeType::CBV:
+			return D3D12ShaderUtils::StaticRootSignatureConstants::CBVDescriptorRangeFlags;
+		case ERootSignatureRangeType::UAV:
+			return D3D12ShaderUtils::StaticRootSignatureConstants::UAVDescriptorRangeFlags;
+		case ERootSignatureRangeType::Sampler:
+			return D3D12ShaderUtils::StaticRootSignatureConstants::SamplerDescriptorRangeFlags;
+		default:
+			return D3D12_DESCRIPTOR_RANGE_FLAG_NONE;
 		}
+	}
 
-		inline D3D12_SHADER_VISIBILITY GetD3D12ShaderVisibility(EShaderFrequency Stage)
+	inline D3D12_SHADER_VISIBILITY GetD3D12ShaderVisibility(ERootSignatureVisibility Visibility)
+	{
+		switch (Visibility)
 		{
-			switch (Stage)
-			{
-			case SF_Vertex:
-				return D3D12_SHADER_VISIBILITY_VERTEX;
-			case SF_Pixel:
-				return D3D12_SHADER_VISIBILITY_PIXEL;
-			case SF_Geometry:
-				return D3D12_SHADER_VISIBILITY_GEOMETRY;
-			default:
-				return D3D12_SHADER_VISIBILITY_ALL;
-			}
+		case ERootSignatureVisibility::Vertex:
+			return D3D12_SHADER_VISIBILITY_VERTEX;
+		case ERootSignatureVisibility::Pixel:
+			return D3D12_SHADER_VISIBILITY_PIXEL;
+		case ERootSignatureVisibility::Geometry:
+			return D3D12_SHADER_VISIBILITY_GEOMETRY;
+#if !defined(D3D12RHI_TOOLS_MESH_SHADERS_UNSUPPORTED)
+		case ERootSignatureVisibility::Mesh:
+			return D3D12_SHADER_VISIBILITY_MESH;
+		case ERootSignatureVisibility::Amplification:
+			return D3D12_SHADER_VISIBILITY_AMPLIFICATION;
+#endif
+		default:
+			return D3D12_SHADER_VISIBILITY_ALL;
 		}
+	}
 
-		inline const TCHAR* GetVisibilityFlag(EShaderFrequency Stage)
+	inline const TCHAR* GetVisibilityFlag(ERootSignatureVisibility Visibility)
+	{
+		switch (Visibility)
 		{
-			switch (Stage)
-			{
-			case SF_Vertex:
-				return TEXT("SHADER_VISIBILITY_VERTEX");
-			case SF_Geometry:
-				return TEXT("SHADER_VISIBILITY_GEOMETRY");
-			case SF_Pixel:
-				return TEXT("SHADER_VISIBILITY_PIXEL");
-			case SF_Mesh:
-				return TEXT("SHADER_VISIBILITY_MESH");
-			case SF_Amplification:
-				return TEXT("SHADER_VISIBILITY_AMPLIFICATION");
-			default:
-				return TEXT("SHADER_VISIBILITY_ALL");
-			}
-		};
+		case ERootSignatureVisibility::Vertex:
+			return TEXT("SHADER_VISIBILITY_VERTEX");
+		case ERootSignatureVisibility::Geometry:
+			return TEXT("SHADER_VISIBILITY_GEOMETRY");
+		case ERootSignatureVisibility::Pixel:
+			return TEXT("SHADER_VISIBILITY_PIXEL");
+#if !defined(D3D12RHI_TOOLS_MESH_SHADERS_UNSUPPORTED)
+		case ERootSignatureVisibility::Mesh:
+			return TEXT("SHADER_VISIBILITY_MESH");
+		case ERootSignatureVisibility::Amplification:
+			return TEXT("SHADER_VISIBILITY_AMPLIFICATION");
+#endif
+		default:
+			return TEXT("SHADER_VISIBILITY_ALL");
+		}
+	};
 
-		inline const TCHAR* GetTypePrefix(ERootSignatureType Type)
+	inline const TCHAR* GetTypePrefix(ERootSignatureRangeType Type)
+	{
+		switch (Type)
 		{
-			switch (Type)
-			{
-			case ERootSignatureType::SRV:
-				return TEXT("SRV(t");
-			case ERootSignatureType::UAV:
-				return TEXT("UAV(u");
-			case ERootSignatureType::Sampler:
-				return TEXT("Sampler(s");
-			default:
-				return TEXT("CBV(b");
-			}
+		case ERootSignatureRangeType::SRV:
+			return TEXT("SRV(t");
+		case ERootSignatureRangeType::UAV:
+			return TEXT("UAV(u");
+		case ERootSignatureRangeType::Sampler:
+			return TEXT("Sampler(s");
+		default:
+			return TEXT("CBV(b");
 		}
+	}
 
-		inline const TCHAR* GetFlagName(D3D12_ROOT_SIGNATURE_FLAGS Flag)
+	inline const TCHAR* GetFlagName(D3D12_ROOT_SIGNATURE_FLAGS Flag)
+	{
+		switch (Flag)
 		{
-			switch (Flag)
-			{
-			case D3D12_ROOT_SIGNATURE_FLAG_ALLOW_INPUT_ASSEMBLER_INPUT_LAYOUT:
-				return TEXT("ALLOW_INPUT_ASSEMBLER_INPUT_LAYOUT");
-			case D3D12_ROOT_SIGNATURE_FLAG_DENY_VERTEX_SHADER_ROOT_ACCESS:
-				return TEXT("DENY_VERTEX_SHADER_ROOT_ACCESS");
-			case D3D12_ROOT_SIGNATURE_FLAG_DENY_GEOMETRY_SHADER_ROOT_ACCESS:
-				return TEXT("DENY_GEOMETRY_SHADER_ROOT_ACCESS");
-			case D3D12_ROOT_SIGNATURE_FLAG_DENY_PIXEL_SHADER_ROOT_ACCESS:
-				return TEXT("DENY_PIXEL_SHADER_ROOT_ACCESS");
-			case D3D12_ROOT_SIGNATURE_FLAG_ALLOW_STREAM_OUTPUT:
-				return TEXT("ALLOW_STREAM_OUTPUT");
+		case D3D12_ROOT_SIGNATURE_FLAG_ALLOW_INPUT_ASSEMBLER_INPUT_LAYOUT:
+			return TEXT("ALLOW_INPUT_ASSEMBLER_INPUT_LAYOUT");
+		case D3D12_ROOT_SIGNATURE_FLAG_DENY_VERTEX_SHADER_ROOT_ACCESS:
+			return TEXT("DENY_VERTEX_SHADER_ROOT_ACCESS");
+		case D3D12_ROOT_SIGNATURE_FLAG_DENY_GEOMETRY_SHADER_ROOT_ACCESS:
+			return TEXT("DENY_GEOMETRY_SHADER_ROOT_ACCESS");
+		case D3D12_ROOT_SIGNATURE_FLAG_DENY_PIXEL_SHADER_ROOT_ACCESS:
+			return TEXT("DENY_PIXEL_SHADER_ROOT_ACCESS");
+		case D3D12_ROOT_SIGNATURE_FLAG_ALLOW_STREAM_OUTPUT:
+			return TEXT("ALLOW_STREAM_OUTPUT");
 
 #if !defined(D3D12RHI_TOOLS_MESH_SHADERS_UNSUPPORTED)
-			case D3D12_ROOT_SIGNATURE_FLAG_DENY_AMPLIFICATION_SHADER_ROOT_ACCESS:
-				return TEXT("DENY_AMPLIFICATION_SHADER_ROOT_ACCESS");
-			case D3D12_ROOT_SIGNATURE_FLAG_DENY_MESH_SHADER_ROOT_ACCESS:
-				return TEXT("DENY_MESH_SHADER_ROOT_ACCESS");
+		case D3D12_ROOT_SIGNATURE_FLAG_DENY_AMPLIFICATION_SHADER_ROOT_ACCESS:
+			return TEXT("DENY_AMPLIFICATION_SHADER_ROOT_ACCESS");
+		case D3D12_ROOT_SIGNATURE_FLAG_DENY_MESH_SHADER_ROOT_ACCESS:
+			return TEXT("DENY_MESH_SHADER_ROOT_ACCESS");
 #endif
 
-			case D3D12_ROOT_SIGNATURE_FLAG_CBV_SRV_UAV_HEAP_DIRECTLY_INDEXED:
-				return TEXT("CBV_SRV_UAV_HEAP_DIRECTLY_INDEXED");
-			case D3D12_ROOT_SIGNATURE_FLAG_SAMPLER_HEAP_DIRECTLY_INDEXED:
-				return TEXT("SAMPLER_HEAP_DIRECTLY_INDEXED");
+		case D3D12_ROOT_SIGNATURE_FLAG_CBV_SRV_UAV_HEAP_DIRECTLY_INDEXED:
+			return TEXT("CBV_SRV_UAV_HEAP_DIRECTLY_INDEXED");
+		case D3D12_ROOT_SIGNATURE_FLAG_SAMPLER_HEAP_DIRECTLY_INDEXED:
+			return TEXT("SAMPLER_HEAP_DIRECTLY_INDEXED");
 
-			default:
-				break;
-			}
+		default:
+			break;
+		}
 
-			return TEXT("");
-		};
+		return TEXT("");
 	};
 
-	// Fat/Static Gfx Root Signature
-	inline void CreateGfxRootSignature(FRootSignatureCreator& Creator, bool bAllowMeshShaders)
+	// Simple base class to help write out a root signature (subclass to generate either to a binary struct or a #define)
+	struct FRootSignatureCreator
 	{
-		// Ensure the creator starts in a clean state (in cases of creator reuse, etc.).
-		Creator
-			.Reset()
-			.AddRootFlag(D3D12_ROOT_SIGNATURE_FLAG_ALLOW_INPUT_ASSEMBLER_INPUT_LAYOUT)
-			.AddTable(SF_Pixel, ERootSignatureType::SRV, MAX_SRVS)
-			.AddTable(SF_Pixel, ERootSignatureType::CBV, MAX_CBS)
-			.AddTable(SF_Pixel, ERootSignatureType::Sampler, MAX_SAMPLERS)
-			.AddTable(SF_Vertex, ERootSignatureType::SRV, MAX_SRVS)
-			.AddTable(SF_Vertex, ERootSignatureType::CBV, MAX_CBS)
-			.AddTable(SF_Vertex, ERootSignatureType::Sampler, MAX_SAMPLERS)
-			.AddTable(SF_Geometry, ERootSignatureType::SRV, MAX_SRVS)
-			.AddTable(SF_Geometry, ERootSignatureType::CBV, MAX_CBS)
-			.AddTable(SF_Geometry, ERootSignatureType::Sampler, MAX_SAMPLERS);
-		if (bAllowMeshShaders)
+		ED3D12RootSignatureFlags Flags = ED3D12RootSignatureFlags::None;
+
+		virtual ~FRootSignatureCreator() { }
+
+		virtual void AddRootFlag(D3D12_ROOT_SIGNATURE_FLAGS Flag) = 0;
+		virtual void AddTable(ERootSignatureVisibility Visibility, ERootSignatureRangeType Type, int32 NumDescriptors, D3D12_DESCRIPTOR_RANGE_FLAGS FlagsOverride = D3D12_DESCRIPTOR_RANGE_FLAG_NONE) = 0;
+
+		void SetFlags(ED3D12RootSignatureFlags InFlags)
 		{
-			Creator
-				.AddTable(SF_Mesh, ERootSignatureType::SRV, MAX_SRVS)
-				.AddTable(SF_Mesh, ERootSignatureType::CBV, MAX_CBS)
-				.AddTable(SF_Mesh, ERootSignatureType::Sampler, MAX_SAMPLERS)
-				.AddTable(SF_Amplification, ERootSignatureType::SRV, MAX_SRVS)
-				.AddTable(SF_Amplification, ERootSignatureType::CBV, MAX_CBS)
-				.AddTable(SF_Amplification, ERootSignatureType::Sampler, MAX_SAMPLERS);
+			Flags = InFlags;
+
+			if (EnumHasAnyFlags(InFlags, ED3D12RootSignatureFlags::InputAssembler))
+			{
+				AddRootFlag(D3D12_ROOT_SIGNATURE_FLAG_ALLOW_INPUT_ASSEMBLER_INPUT_LAYOUT);
+			}
+
+			if (EnumHasAnyFlags(InFlags, ED3D12RootSignatureFlags::BindlessResources))
+			{
+				AddRootFlag(D3D12_ROOT_SIGNATURE_FLAG_CBV_SRV_UAV_HEAP_DIRECTLY_INDEXED);
+			}
+
+			if (EnumHasAnyFlags(InFlags, ED3D12RootSignatureFlags::BindlessSamplers))
+			{
+				AddRootFlag(D3D12_ROOT_SIGNATURE_FLAG_SAMPLER_HEAP_DIRECTLY_INDEXED);
+			}
 		}
-		Creator.AddTable(SF_NumFrequencies, ERootSignatureType::UAV, MAX_UAVS);
-	}
 
-	// Fat/Static Compute Root Signature
-	inline void CreateComputeRootSignature(FRootSignatureCreator& Creator)
-	{
-		// Ensure the creator starts in a clean state (in cases of creator reuse, etc.).
-		Creator
-			.Reset()
-			.AddRootFlag(D3D12_ROOT_SIGNATURE_FLAG_NONE)
-			.AddTable(SF_NumFrequencies, ERootSignatureType::SRV, MAX_SRVS)
-			.AddTable(SF_NumFrequencies, ERootSignatureType::CBV, MAX_CBS)
-			.AddTable(SF_NumFrequencies, ERootSignatureType::Sampler, MAX_SAMPLERS)
-			.AddTable(SF_NumFrequencies, ERootSignatureType::UAV, MAX_UAVS);
-	}
+		inline bool ShouldSkipType(ERootSignatureRangeType Type)
+		{
+			if (Type == ERootSignatureRangeType::SRV || Type == ERootSignatureRangeType::UAV)
+			{
+				return EnumHasAnyFlags(Flags, ED3D12RootSignatureFlags::BindlessResources);
+			}
+
+			if (Type == ERootSignatureRangeType::Sampler)
+			{
+				return EnumHasAnyFlags(Flags, ED3D12RootSignatureFlags::BindlessSamplers);
+			}
+
+			return false;
+		}
+	};
 
 	struct FBinaryRootSignatureCreator : public FRootSignatureCreator
 	{
-		TArray<D3D12_DESCRIPTOR_RANGE1> DescriptorRanges;
-		TArray<D3D12_ROOT_PARAMETER1> Parameters;
-		TMap<uint32, uint32> ParameterToRangeMap;
+		D3D12_ROOT_SIGNATURE_FLAGS RootFlags = D3D12_ROOT_SIGNATURE_FLAG_NONE;
+		uint32 RegisterSpace = 0;
+		TArray<CD3DX12_DESCRIPTOR_RANGE1> DescriptorRanges;
+		TArray<CD3DX12_ROOT_PARAMETER1> Parameters;
+		CD3DX12_VERSIONED_ROOT_SIGNATURE_DESC RootDesc;
 
-		D3D12_ROOT_SIGNATURE_FLAGS Flags = D3D12_ROOT_SIGNATURE_FLAG_NONE;
+		struct FPendingTable
+		{
+			ERootSignatureVisibility Visibility;
+			ERootSignatureRangeType Type;
+			int32 NumDescriptors;
+			D3D12_DESCRIPTOR_RANGE_FLAGS FlagsOverride;
+		};
+		TArray<FPendingTable> PendingTables;
 
-		FRootSignatureCreator& Reset() override
+		void SetRegisterSpace(uint32 InSpace)
 		{
-			DescriptorRanges.Empty();
-			Parameters.Empty();
-			ParameterToRangeMap.Empty();
-			Flags = D3D12_ROOT_SIGNATURE_FLAG_NONE;
-			return *this;
+			RegisterSpace = InSpace;
 		}
 
-		FRootSignatureCreator& AddRootFlag(D3D12_ROOT_SIGNATURE_FLAGS Flag) override
+		void AddRootFlag(D3D12_ROOT_SIGNATURE_FLAGS RootFlag) override
 		{
-			Flags |= Flag;
-			return *this;
+			RootFlags |= RootFlag;
 		}
 
-		FRootSignatureCreator& AddTable(EShaderFrequency Stage, ERootSignatureType Type, int32 NumDescriptors) override
+		void AddShaderResourceViewParameter(uint32 Register, uint32 Space)
 		{
-			int32 ParameterIndex = Parameters.AddZeroed();
-			int32 RangeIndex = DescriptorRanges.AddZeroed();
-
-			D3D12_ROOT_PARAMETER1& Parameter = Parameters[ParameterIndex];
-			D3D12_DESCRIPTOR_RANGE1& Range = DescriptorRanges[ParameterIndex];
-
-			Parameter.ParameterType = D3D12_ROOT_PARAMETER_TYPE_DESCRIPTOR_TABLE;
-			Parameter.DescriptorTable.NumDescriptorRanges = 1;
-
-			// Pointer will be filled in last
-			//Parameter.DescriptorTable.pDescriptorRanges = &DescriptorRanges[0];
-			ParameterToRangeMap.Add(ParameterIndex, RangeIndex);
-
-			Range.RangeType = GetD3D12DescriptorRangeType(Type);
-			Range.NumDescriptors = NumDescriptors;
-			Range.Flags = GetD3D12DescriptorRangeFlags(Type);
-			//Range.OffsetInDescriptorsFromTableStart = UINT32_MAX;
-			Parameter.ShaderVisibility =  GetD3D12ShaderVisibility(Stage);
+			CD3DX12_ROOT_PARAMETER1& Parameter = Parameters.AddZeroed_GetRef();
+			Parameter.InitAsShaderResourceView(Register, Space);
+		}
 
-			return *this;
+		void AddConstantsParameter(uint32 Num32BitValues, uint32 Register, uint32 Space)
+		{
+			CD3DX12_ROOT_PARAMETER1& Parameter = Parameters.AddZeroed_GetRef();
+			Parameter.InitAsConstants(Num32BitValues, Register, Space);
 		}
 
-		void Compile(EShaderFrequency Freq)
+		void AddTable(ERootSignatureVisibility Visibility, ERootSignatureRangeType Type, int32 NumDescriptors, D3D12_DESCRIPTOR_RANGE_FLAGS FlagsOverride = D3D12_DESCRIPTOR_RANGE_FLAG_NONE) override
 		{
-			if (Freq < SF_NumGraphicsFrequencies)
+			if (!ShouldSkipType(Type))
 			{
-				bool bAllowMeshShaders = false;
-#if !defined(D3D12RHI_TOOLS_MESH_SHADERS_UNSUPPORTED)
-				bAllowMeshShaders = true;
-#endif
-				D3D12ShaderUtils::CreateGfxRootSignature(*this, bAllowMeshShaders);
+				const FPendingTable NewTable{ Visibility, Type, NumDescriptors, FlagsOverride };
+				PendingTables.Emplace(NewTable);
 			}
-			else //if (Freq == SF_Compute)
+		}
+
+		const D3D12_VERSIONED_ROOT_SIGNATURE_DESC& Finalize()
+		{
+#if !defined(D3D12RHI_TOOLS_RAYTRACING_SHADERS_UNSUPPORTED)
+			if (RootFlags & D3D12_ROOT_SIGNATURE_FLAG_LOCAL_ROOT_SIGNATURE)
 			{
-				D3D12ShaderUtils::CreateComputeRootSignature(*this);
+				// if set, D3D12_ROOT_SIGNATURE_FLAG_LOCAL_ROOT_SIGNATURE should be the only flag set
+				RootFlags = D3D12_ROOT_SIGNATURE_FLAG_LOCAL_ROOT_SIGNATURE;
 			}
+#endif
 
-			// Patch pointers
-			for (auto& Pair : ParameterToRangeMap)
+			DescriptorRanges.SetNum(PendingTables.Num());
+
+			for (int32 Index = 0; Index < PendingTables.Num(); Index++)
 			{
-				Parameters[Pair.Key].DescriptorTable.pDescriptorRanges = &DescriptorRanges[Pair.Value];
+				const FPendingTable& PendingTable = PendingTables[Index];
+
+				const D3D12_DESCRIPTOR_RANGE_FLAGS RangeFlags = PendingTable.FlagsOverride ? PendingTable.FlagsOverride : GetD3D12DescriptorRangeFlags(PendingTable.Type);
+
+				DescriptorRanges[Index].Init(GetD3D12DescriptorRangeType(PendingTable.Type), PendingTable.NumDescriptors, 0, RegisterSpace, RangeFlags);
+
+				CD3DX12_ROOT_PARAMETER1& Parameter = Parameters.AddZeroed_GetRef();
+				Parameter.InitAsDescriptorTable(1, &DescriptorRanges[Index], GetD3D12ShaderVisibility(PendingTable.Visibility));
 			}
+
+			RootDesc.Init_1_1(Parameters.Num(), Parameters.GetData(), 0, nullptr, RootFlags);
+
+			return RootDesc;
 		}
 	};
 
 	/* Root signature generator for DXC */
 	struct FTextRootSignatureCreator : public FRootSignatureCreator
 	{
-		FRootSignatureCreator& AddRootFlag(D3D12_ROOT_SIGNATURE_FLAGS InFlag) override
+		void AddRootFlag(D3D12_ROOT_SIGNATURE_FLAGS InFlag) override
 		{
-			if (Flags.Len() > 0)
+			if (RootFlags.Len() > 0)
 			{
-				Flags += "|";
+				RootFlags += "|";
 			}
-			Flags += GetFlagName(InFlag);
-			return *this;
+			RootFlags += GetFlagName(InFlag);
 		}
 
-		FRootSignatureCreator& AddTable(EShaderFrequency InStage, ERootSignatureType Type, int32 NumDescriptors) override
+		void AddTable(ERootSignatureVisibility Visibility, ERootSignatureRangeType Type, int32 NumDescriptors, D3D12_DESCRIPTOR_RANGE_FLAGS FlagsOverride = D3D12_DESCRIPTOR_RANGE_FLAG_NONE) override
 		{
-			FString Line = FString::Printf(TEXT("DescriptorTable(visibility=%s, %s0, numDescriptors=%d%s"/*, flags = DESCRIPTORS_VOLATILE*/"))"),
-				GetVisibilityFlag(InStage), GetTypePrefix(Type), NumDescriptors,
-				TEXT("")
-			);
-			if (Table.Len() > 0)
+			if (!ShouldSkipType(Type))
 			{
-				Table += ",";
+				FString Line = FString::Printf(TEXT("DescriptorTable(visibility=%s, %s0, numDescriptors=%d%s"/*, flags = DESCRIPTORS_VOLATILE*/"))"),
+					GetVisibilityFlag(Visibility), GetTypePrefix(Type), NumDescriptors,
+					TEXT("")
+				);
+				if (Table.Len() > 0)
+				{
+					Table += ",";
+				}
+				Table += Line;
 			}
-			Table += Line;
-			return *this;
 		}
 
-		FRootSignatureCreator& Reset() override
+		FString GenerateString() const
 		{
-			Flags.Empty();
-			Table.Empty();
-			return *this;
+			FString String = FString::Printf(TEXT("\"RootFlags(%s),%s\""),
+				RootFlags.Len() == 0 ? TEXT("0") : *RootFlags,
+				*Table);
+			return String;
 		}
 
-		FString CreateAndGenerateString(EShaderFrequency Freq, bool bBindlessResources, bool bBindlessSamplers)
-		{
-			if (Freq < SF_NumGraphicsFrequencies)
-			{
-				bool bAllowMeshShaders = false;
+		FString RootFlags;
+		FString Table;
+	};
+
+	inline void AddAllStandardTablesForVisibility(FRootSignatureCreator& Creator, ERootSignatureVisibility Visibility)
+	{
+		Creator.AddTable(Visibility, ERootSignatureRangeType::SRV, MAX_SRVS);
+		Creator.AddTable(Visibility, ERootSignatureRangeType::CBV, MAX_CBS);
+		Creator.AddTable(Visibility, ERootSignatureRangeType::Sampler, MAX_SAMPLERS);
+	}
+
+	// Fat/Static Gfx Root Signature
+	inline void CreateGfxRootSignature(FRootSignatureCreator& Creator, ED3D12RootSignatureFlags InFlags)
+	{
+		// Ensure the creator starts in a clean state (in cases of creator reuse, etc.).
+		Creator.SetFlags(InFlags | ED3D12RootSignatureFlags::InputAssembler);
+
+		AddAllStandardTablesForVisibility(Creator, ERootSignatureVisibility::Pixel);
+		AddAllStandardTablesForVisibility(Creator, ERootSignatureVisibility::Vertex);
+		AddAllStandardTablesForVisibility(Creator, ERootSignatureVisibility::Geometry);
+
 #if !defined(D3D12RHI_TOOLS_MESH_SHADERS_UNSUPPORTED)
-				bAllowMeshShaders = true;
+		if (EnumHasAnyFlags(InFlags, ED3D12RootSignatureFlags::AllowMeshShaders))
+		{
+			AddAllStandardTablesForVisibility(Creator, ERootSignatureVisibility::Mesh);
+			AddAllStandardTablesForVisibility(Creator, ERootSignatureVisibility::Amplification);
+		}
 #endif
-				D3D12ShaderUtils::CreateGfxRootSignature(*this, bAllowMeshShaders);
-			}
-			else //if (Freq == SF_Compute)
-			{
-				D3D12ShaderUtils::CreateComputeRootSignature(*this);
-			}
 
-			if (bBindlessResources)
-			{
-				AddRootFlag(D3D12_ROOT_SIGNATURE_FLAG_CBV_SRV_UAV_HEAP_DIRECTLY_INDEXED);
-			}
+		Creator.AddTable(ERootSignatureVisibility::All, ERootSignatureRangeType::UAV, MAX_UAVS);
+	}
 
-			if (bBindlessSamplers)
-			{
-				AddRootFlag(D3D12_ROOT_SIGNATURE_FLAG_SAMPLER_HEAP_DIRECTLY_INDEXED);
-			}
+	// Fat/Static Compute Root Signature
+	inline void CreateComputeRootSignature(FRootSignatureCreator& Creator, ED3D12RootSignatureFlags InFlags)
+	{
+		// Ensure the creator starts in a clean state (in cases of creator reuse, etc.).
+		Creator.SetFlags(InFlags);
+		AddAllStandardTablesForVisibility(Creator, ERootSignatureVisibility::All);
+		Creator.AddTable(ERootSignatureVisibility::All, ERootSignatureRangeType::UAV, MAX_UAVS);
+	}
 
-			//FString String = FString::Printf(TEXT("#define TheRootSignature \"RootFlags(%s),%s\"\n[RootSignature(TheRootSignature)\n"),
-			FString String = FString::Printf(TEXT("\"RootFlags(%s),%s\""),
-				Flags.Len() == 0 ? TEXT("0") : *Flags,
-				*Table);
-			return String;
+#if !defined(D3D12RHI_TOOLS_RAYTRACING_SHADERS_UNSUPPORTED)
+	inline void CreateRayTracingSignature(FBinaryRootSignatureCreator& Creator, bool bLocalRootSignature, D3D12_ROOT_SIGNATURE_FLAGS BaseRootFlags, ED3D12RootSignatureFlags InFlags)
+	{
+		Creator.SetFlags(InFlags);
+		Creator.AddRootFlag(BaseRootFlags);
+		Creator.SetRegisterSpace(bLocalRootSignature ? RAY_TRACING_REGISTER_SPACE_LOCAL : RAY_TRACING_REGISTER_SPACE_GLOBAL);
+
+		if (bLocalRootSignature)
+		{
+			Creator.AddShaderResourceViewParameter(RAY_TRACING_SYSTEM_INDEXBUFFER_REGISTER, RAY_TRACING_REGISTER_SPACE_SYSTEM);
+			Creator.AddShaderResourceViewParameter(RAY_TRACING_SYSTEM_VERTEXBUFFER_REGISTER, RAY_TRACING_REGISTER_SPACE_SYSTEM);
+
+			uint32 NumConstants = 4; // sizeof(FHitGroupSystemRootConstants) / sizeof(uint32);
+			Creator.AddConstantsParameter(NumConstants, RAY_TRACING_SYSTEM_ROOTCONSTANT_REGISTER, RAY_TRACING_REGISTER_SPACE_SYSTEM);
 		}
 
-		FString Flags;
-		FString Table;
-	};
+		AddAllStandardTablesForVisibility(Creator, ERootSignatureVisibility::All);
 
-	inline FString GenerateRootSignatureString(EShaderFrequency InFrequency, bool bBindlessResources = false, bool bBindlessSamplers = false)
+		Creator.AddTable(ERootSignatureVisibility::All, ERootSignatureRangeType::UAV, MAX_UAVS, D3D12_DESCRIPTOR_RANGE_FLAG_DATA_STATIC_WHILE_SET_AT_EXECUTE | D3D12_DESCRIPTOR_RANGE_FLAG_DESCRIPTORS_VOLATILE);
+	}
+#endif //!defined(D3D12RHI_TOOLS_RAYTRACING_SHADERS_UNSUPPORTED)
+
+	inline FString GenerateRootSignatureString(EShaderFrequency InFrequency, ED3D12RootSignatureFlags InFlags)
 	{
-		if (InFrequency < SF_NumStandardFrequencies)
+		FTextRootSignatureCreator Creator;
+
+		if (InFrequency < SF_NumGraphicsFrequencies)
+		{
+			D3D12ShaderUtils::CreateGfxRootSignature(Creator, InFlags);
+		}
+		else if (InFrequency == SF_Compute)
 		{
-			FTextRootSignatureCreator Creator;
-			return Creator.CreateAndGenerateString(InFrequency, bBindlessResources, bBindlessSamplers);
+			D3D12ShaderUtils::CreateComputeRootSignature(Creator, InFlags);
+		}
+		else
+		{
+			return TEXT("");
 		}
 
-		return TEXT("");
+		return Creator.GenerateString();
 	}
 }
