Optimized RHI transition allocations across platforms.

diff --git a/Engine/Source/Runtime/D3D12RHI/Private/D3D12CommandContext.h b/Engine/Source/Runtime/D3D12RHI/Private/D3D12CommandContext.h
--- a/Engine/Source/Runtime/D3D12RHI/Private/D3D12CommandContext.h
+++ b/Engine/Source/Runtime/D3D12RHI/Private/D3D12CommandContext.h
@@ -924,9 +924,9 @@ struct FD3D12TransitionData
 	ERHIPipeline SrcPipelines, DstPipelines;
 	ERHITransitionCreateFlags CreateFlags = ERHITransitionCreateFlags::None;
 
-	TArray<FRHITransitionInfo, TInlineAllocator<4>> TransitionInfos;
-	TArray<FRHITransientAliasingInfo, TInlineAllocator<4>> AliasingInfos;
-	TArray<FRHITransientAliasingOverlap, TInlineAllocator<4>> AliasingOverlaps;
+	TArray<FRHITransitionInfo, TInlineAllocator<4, FConcurrentLinearArrayAllocator>> TransitionInfos;
+	TArray<FRHITransientAliasingInfo, TInlineAllocator<4, FConcurrentLinearArrayAllocator>> AliasingInfos;
+	TArray<FRHITransientAliasingOverlap, TInlineAllocator<4, FConcurrentLinearArrayAllocator>> AliasingOverlaps;
 	TRefCountPtr<FD3D12Fence> Fence;
 
 	bool bCrossPipeline = false;
diff --git a/Engine/Source/Runtime/RHI/Public/DynamicRHI.h b/Engine/Source/Runtime/RHI/Public/DynamicRHI.h
--- a/Engine/Source/Runtime/RHI/Public/DynamicRHI.h
+++ b/Engine/Source/Runtime/RHI/Public/DynamicRHI.h
@@ -1609,7 +1609,7 @@ RHI_API FRenderQueryPoolRHIRef RHICreateRenderQueryPool(ERenderQueryType QueryTy
 FORCEINLINE const FRHITransition* RHICreateTransition(const FRHITransitionCreateInfo& CreateInfo)
 {
 	// Placement create the transition on the heap.
-	FRHITransition* Transition = new (FMemory::Malloc(FRHITransition::GetTotalAllocationSize(), (uint32)FRHITransition::GetAlignment())) FRHITransition(CreateInfo.SrcPipelines, CreateInfo.DstPipelines);
+	FRHITransition* Transition = new (FConcurrentLinearAllocator::Malloc(FRHITransition::GetTotalAllocationSize(), (uint32)FRHITransition::GetAlignment())) FRHITransition(CreateInfo.SrcPipelines, CreateInfo.DstPipelines);
 	GDynamicRHI->RHICreateTransition(Transition, CreateInfo);
 	return Transition;
 }
diff --git a/Engine/Source/Runtime/RHI/Public/RHI.h b/Engine/Source/Runtime/RHI/Public/RHI.h
--- a/Engine/Source/Runtime/RHI/Public/RHI.h
+++ b/Engine/Source/Runtime/RHI/Public/RHI.h
@@ -2762,7 +2762,7 @@ inline void FRHITransition::Cleanup() const
 
 	// Explicit destruction of the transition.
 	Transition->~FRHITransition();
-	FMemory::Free(Transition);
+	FConcurrentLinearAllocator::Free(Transition);
 }
 
 #if ENABLE_RHI_VALIDATION
