buildfix

diff --git a/Engine/Source/Runtime/D3D12RHI/Public/D3D12RHI.h b/Engine/Source/Runtime/D3D12RHI/Public/D3D12RHI.h
--- a/Engine/Source/Runtime/D3D12RHI/Public/D3D12RHI.h
+++ b/Engine/Source/Runtime/D3D12RHI/Public/D3D12RHI.h
@@ -64,7 +64,7 @@
 	#include "D3D12RHIPlatformPublic.h"
 #endif
 
-#if PLATFORM_WINDOWS
+#if PLATFORM_WINDOWS || PLATFORM_HOLOLENS
 #define D3D12_USE_DUMMY_BACKBUFFER 1
 #else
 #define D3D12_USE_DUMMY_BACKBUFFER 0
