MGPU:  Remove "r.AllowMultiGPUInEditor" CVar.  This CVar is redundant, because the command line -MaxGPUCount setting is required anyway to enable MGPU.

diff --git a/Engine/Source/Runtime/D3D12RHI/Private/D3D12Adapter.cpp b/Engine/Source/Runtime/D3D12RHI/Private/D3D12Adapter.cpp
--- a/Engine/Source/Runtime/D3D12RHI/Private/D3D12Adapter.cpp
+++ b/Engine/Source/Runtime/D3D12RHI/Private/D3D12Adapter.cpp
@@ -762,18 +762,8 @@ void FD3D12Adapter::CreateRootDevice(bool bWithDebug)
 	GNumExplicitGPUsForRendering = 1;
 	if (Desc.NumDeviceNodes > 1)
 	{
-		// Can't access GAllowMultiGPUInEditor directly as its value is cached but hasn't been set by console manager due to module loading order
-		static IConsoleVariable* AllowMultiGPUInEditor = IConsoleManager::Get().FindConsoleVariable(TEXT("r.AllowMultiGPUInEditor"));
-
-		if (GIsEditor && AllowMultiGPUInEditor->GetInt() == 0)
-		{
-			UE_LOG(LogD3D12RHI, Log, TEXT("Multi-GPU is available, but skipping due to editor mode."));
-		}
-		else
-		{
-			GNumExplicitGPUsForRendering = Desc.NumDeviceNodes;
-			UE_LOG(LogD3D12RHI, Log, TEXT("Enabling multi-GPU with %d nodes"), Desc.NumDeviceNodes);
-		}
+		GNumExplicitGPUsForRendering = Desc.NumDeviceNodes;
+		UE_LOG(LogD3D12RHI, Log, TEXT("Enabling multi-GPU with %d nodes"), Desc.NumDeviceNodes);
 	}
 
 	// Viewport ignores AFR if PresentGPU is specified.
diff --git a/Engine/Source/Runtime/D3D12RHI/Private/D3D12Stats.cpp b/Engine/Source/Runtime/D3D12RHI/Private/D3D12Stats.cpp
--- a/Engine/Source/Runtime/D3D12RHI/Private/D3D12Stats.cpp
+++ b/Engine/Source/Runtime/D3D12RHI/Private/D3D12Stats.cpp
@@ -125,11 +125,10 @@ void D3D12RHI::FD3DGPUProfiler::EndFrame(FD3D12DynamicRHI* InRHI)
 
 			// Only dump the event tree and generate the screenshot for the first GPU.  Eventually, we may want to collate
 			// profiling data for all GPUs into a single tree, but the short term goal is to make profiling in the editor
-			// functional at all with "r.AllowMultiGPUInEditor=1" and "-MaxGPUCount=2" (settings required to enable multiple
-			// GPUs for GPU Lightmass).  In the editor, we don't actually render anything on the additional GPUs, but the
-			// editor's profile visualizer will pick up whatever event tree we dumped last, which will be the empty one from
-			// the last GPU, making the results useless without this code fix.  Unreal Insights would be preferred for
-			// multi-GPU profiling outside the editor.
+			// functional at all with "-MaxGPUCount=2" (required to enable multiple GPUs for GPU Lightmass).  In the editor,
+			// we don't actually render anything on the additional GPUs, but the editor's profile visualizer will pick up
+			// whatever event tree we dumped last, which will be the empty one from the last GPU, making the results
+			// useless without this code fix.  Unreal Insights would be preferred for multi-GPU profiling outside the editor.
 			if (GPUIndex == 0)
 			{
 				CurrentEventNodeFrame->DumpEventTree();
diff --git a/Engine/Source/Runtime/RHI/Private/MultiGPU.cpp b/Engine/Source/Runtime/RHI/Private/MultiGPU.cpp
--- a/Engine/Source/Runtime/RHI/Private/MultiGPU.cpp
+++ b/Engine/Source/Runtime/RHI/Private/MultiGPU.cpp
@@ -11,18 +11,10 @@
 uint32 GNumAlternateFrameRenderingGroups = 1;
 uint32 GNumExplicitGPUsForRendering = 1;
 uint32 GVirtualMGPU = 0;
-int32 GAllowMultiGPUInEditor = 0;
 #endif
 
 #if WITH_MGPU
 
-TAutoConsoleVariable<int32> CVarAllowMultiGPUInEditor(
-	TEXT("r.AllowMultiGPUInEditor"),
-	GAllowMultiGPUInEditor,
-	TEXT("Enable experimental multi-GPU support in editor"),
-	ECVF_ReadOnly
-);
-
 RHI_API TArray<FRHIGPUMask, TFixedAllocator<MAX_NUM_GPUS>> AFRUtils::GroupMasks;
 RHI_API TArray<FRHIGPUMask, TFixedAllocator<MAX_NUM_GPUS>> AFRUtils::SiblingMasks;
 
diff --git a/Engine/Source/Runtime/RHI/Public/MultiGPU.h b/Engine/Source/Runtime/RHI/Public/MultiGPU.h
--- a/Engine/Source/Runtime/RHI/Public/MultiGPU.h
+++ b/Engine/Source/Runtime/RHI/Public/MultiGPU.h
@@ -16,7 +16,6 @@
 extern RHI_API uint32 GNumExplicitGPUsForRendering;
 extern RHI_API uint32 GNumAlternateFrameRenderingGroups;
 extern RHI_API uint32 GVirtualMGPU;
-extern RHI_API int32 GAllowMultiGPUInEditor;
 #else
 #define WITH_SLI 0
 #define WITH_MGPU 0
@@ -24,7 +23,6 @@ extern RHI_API int32 GAllowMultiGPUInEditor;
 #define GNumExplicitGPUsForRendering 1
 #define GNumAlternateFrameRenderingGroups 1
 #define GVirtualMGPU 0
-#define GAllowMultiGPUInEditor 0
 #endif
 
 /** A mask where each bit is a GPU index. Can not be empty so that non SLI platforms can optimize it to be always 1.  */
diff --git a/Engine/Source/Runtime/Renderer/Private/PathTracing.cpp b/Engine/Source/Runtime/Renderer/Private/PathTracing.cpp
--- a/Engine/Source/Runtime/Renderer/Private/PathTracing.cpp
+++ b/Engine/Source/Runtime/Renderer/Private/PathTracing.cpp
@@ -213,7 +213,7 @@ TAutoConsoleVariable<int32> CVarPathTracingMultiGPU(
 	TEXT("r.PathTracing.MultiGPU"),
 	0,
 	TEXT("Run the path tracer using all available GPUs when enabled (default = 0)\n")
-	TEXT("Using this functionality in the editor requires r.AllowMultiGPUInEditor=1"),
+	TEXT("Using this functionality in the editor requires -MaxGPUCount=N setting on the command line"),
 	ECVF_RenderThreadSafe
 );
 
