Only rename the d3d12 resource if it's not suballocated (requires resource tracking or it's a committed resource)

diff --git a/Engine/Source/Runtime/D3D12RHI/Private/D3D12Buffer.cpp b/Engine/Source/Runtime/D3D12RHI/Private/D3D12Buffer.cpp
--- a/Engine/Source/Runtime/D3D12RHI/Private/D3D12Buffer.cpp
+++ b/Engine/Source/Runtime/D3D12RHI/Private/D3D12Buffer.cpp
@@ -855,43 +855,52 @@ void FD3D12DynamicRHI::RHICopyBuffer(FRHIBuffer* SourceBufferRHI, FRHIBuffer* De
 
 void FD3D12DynamicRHI::RHIBindDebugLabelName(FRHIBuffer* BufferRHI, const TCHAR* Name)
 {
+	if (BufferRHI == nullptr)
+	{
+		return;
+	}
+
 #if NAME_OBJECTS
 	FD3D12Buffer* Buffer = FD3D12DynamicRHI::ResourceCast(BufferRHI);
 
-	FD3D12Buffer::FLinkedObjectIterator BufferIt(Buffer);
-
-	if (GNumExplicitGPUsForRendering > 1)
+	// only rename the underlying d3d12 resource if it's not sub allocated (requires resource state tracking or stand alone allocated)
+	if (Buffer->GetResource()->RequiresResourceStateTracking() || Buffer->ResourceLocation.GetType() == FD3D12ResourceLocation::ResourceLocationType::eStandAlone)
 	{
-		// Generate string of the form "Name (GPU #)" -- assumes GPU index is a single digit.  This is called many times
-		// a frame, so we want to avoid any string functions which dynamically allocate, to reduce perf overhead.
-		static_assert(MAX_NUM_GPUS <= 10);
+		FD3D12Buffer::FLinkedObjectIterator BufferIt(Buffer);
+
+		if (GNumExplicitGPUsForRendering > 1)
+		{
+			// Generate string of the form "Name (GPU #)" -- assumes GPU index is a single digit.  This is called many times
+			// a frame, so we want to avoid any string functions which dynamically allocate, to reduce perf overhead.
+			static_assert(MAX_NUM_GPUS <= 10);
 
-		static const TCHAR NameSuffix[] = TEXT(" (GPU #)");
-		constexpr int32 NameSuffixLengthWithTerminator = (int32)UE_ARRAY_COUNT(NameSuffix);
-		constexpr int32 NameBufferLength = 256;
-		constexpr int32 GPUIndexSuffixOffset = 6;		// Offset of '#' character
+			static const TCHAR NameSuffix[] = TEXT(" (GPU #)");
+			constexpr int32 NameSuffixLengthWithTerminator = (int32)UE_ARRAY_COUNT(NameSuffix);
+			constexpr int32 NameBufferLength = 256;
+			constexpr int32 GPUIndexSuffixOffset = 6;		// Offset of '#' character
 
-		// Combine Name and suffix in our string buffer (clamping the length for bounds checking).  We'll replace the GPU index
-		// with the appropriate digit in the loop.
-		int32 NameLength = FMath::Min(FCString::Strlen(Name), NameBufferLength - NameSuffixLengthWithTerminator);
-		int32 GPUIndexOffset = NameLength + GPUIndexSuffixOffset;
+			// Combine Name and suffix in our string buffer (clamping the length for bounds checking).  We'll replace the GPU index
+			// with the appropriate digit in the loop.
+			int32 NameLength = FMath::Min(FCString::Strlen(Name), NameBufferLength - NameSuffixLengthWithTerminator);
+			int32 GPUIndexOffset = NameLength + GPUIndexSuffixOffset;
 
-		TCHAR DebugName[NameBufferLength];
-		FMemory::Memcpy(&DebugName[0], Name, NameLength*sizeof(TCHAR));
-		FMemory::Memcpy(&DebugName[NameLength], NameSuffix, NameSuffixLengthWithTerminator*sizeof(TCHAR));
+			TCHAR DebugName[NameBufferLength];
+			FMemory::Memcpy(&DebugName[0], Name, NameLength * sizeof(TCHAR));
+			FMemory::Memcpy(&DebugName[NameLength], NameSuffix, NameSuffixLengthWithTerminator * sizeof(TCHAR));
 
-		for (; BufferIt; ++BufferIt)
-		{
-			FD3D12Resource* Resource = BufferIt->GetResource();
+			for (; BufferIt; ++BufferIt)
+			{
+				FD3D12Resource* Resource = BufferIt->GetResource();
 
-			DebugName[GPUIndexOffset] = TEXT('0') + BufferIt->GetParentDevice()->GetGPUIndex();
+				DebugName[GPUIndexOffset] = TEXT('0') + BufferIt->GetParentDevice()->GetGPUIndex();
 
-			SetName(Resource, DebugName);
+				SetName(Resource, DebugName);
+			}
+		}
+		else
+		{
+			SetName(Buffer->GetResource(), Name);
 		}
-	}
-	else
-	{
-		SetName(Buffer->GetResource(), Name);
 	}
 #endif
 
