- Add support for shared buffer allocation as well (textures were already supported)
- Always perform committed resource code path for shared resources in the pool allocator

diff --git a/Engine/Source/Runtime/D3D12RHI/Private/D3D12Buffer.cpp b/Engine/Source/Runtime/D3D12RHI/Private/D3D12Buffer.cpp
--- a/Engine/Source/Runtime/D3D12RHI/Private/D3D12Buffer.cpp
+++ b/Engine/Source/Runtime/D3D12RHI/Private/D3D12Buffer.cpp
@@ -512,6 +512,11 @@ void FD3D12Buffer::GetResourceDescAndAlignment(uint64 InSize, uint32 InStride, E
 		ResourceDesc.Flags |= D3D12RHI_RESOURCE_FLAG_ALLOW_INDIRECT_BUFFER;
 	}
 
+	if (EnumHasAnyFlags(InUsage, BUF_Shared))
+	{
+		ResourceDesc.Flags |= D3D12_RESOURCE_FLAG_ALLOW_SIMULTANEOUS_ACCESS;
+	}
+
 	// Structured buffers, non-ByteAddress buffers, need to be aligned to their stride to ensure that they can be addressed correctly with element based offsets.
 	Alignment = (InStride > 0) && (EnumHasAnyFlags(InUsage, BUF_StructuredBuffer) || !EnumHasAnyFlags(InUsage, BUF_ByteAddressBuffer | BUF_DrawIndirect)) ? InStride : 4;
 }
diff --git a/Engine/Source/Runtime/D3D12RHI/Private/D3D12PoolAllocator.cpp b/Engine/Source/Runtime/D3D12RHI/Private/D3D12PoolAllocator.cpp
--- a/Engine/Source/Runtime/D3D12RHI/Private/D3D12PoolAllocator.cpp
+++ b/Engine/Source/Runtime/D3D12RHI/Private/D3D12PoolAllocator.cpp
@@ -276,7 +276,8 @@ void FD3D12PoolAllocator::AllocateResource(uint32 GPUIndex, D3D12_HEAP_TYPE InHe
 
 	check(GPUIndex == Device->GetGPUIndex());
 
-	const bool PoolResource = InSize <= MaxAllocationSize;
+	const bool bSharedResource = EnumHasAnyFlags(InDesc.Flags, D3D12_RESOURCE_FLAG_ALLOW_SIMULTANEOUS_ACCESS);
+	const bool PoolResource = InSize <= MaxAllocationSize && !bSharedResource;
 	if (PoolResource)
 	{
 		const bool bPlacedResource = (AllocationStrategy == EResourceAllocationStrategy::kPlacedResource);
