Only enable async compute by default on Ampere and newer cards on nVidia hardware

diff --git a/Engine/Source/Runtime/D3D12RHI/Private/Windows/WindowsD3D12Device.cpp b/Engine/Source/Runtime/D3D12RHI/Private/Windows/WindowsD3D12Device.cpp
--- a/Engine/Source/Runtime/D3D12RHI/Private/Windows/WindowsD3D12Device.cpp
+++ b/Engine/Source/Runtime/D3D12RHI/Private/Windows/WindowsD3D12Device.cpp
@@ -1467,7 +1467,18 @@ void FD3D12DynamicRHI::Init()
 		UE_LOG(LogD3D12RHI, Log, TEXT("RHI does not have support for 64 bit atomics"));
 	}
 
-	GSupportsEfficientAsyncCompute = GAllowAsyncCompute && (FParse::Param(FCommandLine::Get(), TEXT("ForceAsyncCompute")) || (GRHISupportsParallelRHIExecute && (IsRHIDeviceAMD() || IsRHIDeviceNVIDIA())));
+	D3D12_FEATURE_DATA_D3D12_OPTIONS6 options = {};
+	HRESULT Options6HR = GetAdapter().GetD3DDevice()->CheckFeatureSupport(D3D12_FEATURE_D3D12_OPTIONS6, &options, sizeof(options));
+
+	// Allow async compute by default on nVidia cards which support PerPrimitiveShadingRateSupportedWithViewportIndexing 
+	// this should be a good metric according to nVidia itself (this is set for Ampere and newer cards)
+	bool bnVidiaAsyncComputeSupported = false;
+	if (IsRHIDeviceNVIDIA() && Options6HR == S_OK && options.PerPrimitiveShadingRateSupportedWithViewportIndexing)
+	{
+		bnVidiaAsyncComputeSupported = true;
+	}
+
+	GSupportsEfficientAsyncCompute = GAllowAsyncCompute && (FParse::Param(FCommandLine::Get(), TEXT("ForceAsyncCompute")) || (GRHISupportsParallelRHIExecute && (IsRHIDeviceAMD() || bnVidiaAsyncComputeSupported)));
 
 	GSupportsDepthBoundsTest = SupportsDepthBoundsTest(this);
 
@@ -1503,9 +1514,7 @@ void FD3D12DynamicRHI::Init()
 	// - Standalones are added to the deferred deletion queue of its parent FD3D12Adapter
 	GRHIForceNoDeletionLatencyForStreamingTextures = !!PLATFORM_WINDOWS;
 
-	D3D12_FEATURE_DATA_D3D12_OPTIONS6 options = {};
-	HRESULT hr = GetAdapter().GetD3DDevice()->CheckFeatureSupport(D3D12_FEATURE_D3D12_OPTIONS6, &options, sizeof(options));
-	if(hr == S_OK && options.VariableShadingRateTier != D3D12_VARIABLE_SHADING_RATE_TIER_NOT_SUPPORTED)
+	if(Options6HR == S_OK && options.VariableShadingRateTier != D3D12_VARIABLE_SHADING_RATE_TIER_NOT_SUPPORTED)
 	{
 		GRHISupportsPipelineVariableRateShading = true;		// We have at least tier 1.
 		GRHISupportsLargerVariableRateShadingSizes = (options.AdditionalShadingRatesSupported != 0);
