Add more details to error messages when acceleration structure size calculation returns 0

diff --git a/Engine/Source/Runtime/D3D12RHI/Private/D3D12RayTracing.cpp b/Engine/Source/Runtime/D3D12RHI/Private/D3D12RayTracing.cpp
--- a/Engine/Source/Runtime/D3D12RHI/Private/D3D12RayTracing.cpp
+++ b/Engine/Source/Runtime/D3D12RHI/Private/D3D12RayTracing.cpp
@@ -3305,11 +3305,14 @@ enum class ERayTracingBufferType
 
 static TRefCountPtr<FD3D12Buffer> CreateRayTracingBuffer(FD3D12Adapter* Adapter, uint32 GPUIndex, uint64 Size, ERayTracingBufferType Type, const FDebugName& DebugName)
 {
-	check(Size);
+	FString DebugNameString = DebugName.ToString();
+
+	checkf(Size != 0, TEXT("Attempting to create ray tracing %s buffer of zero size. Debug name: %s"),
+		Type == ERayTracingBufferType::AccelerationStructure ? TEXT("AccelerationStructure") : TEXT("Scratch"),
+		*DebugNameString);
 
 	TRefCountPtr<FD3D12Buffer> Result;
 
-	FString DebugNameString = DebugName.ToString();
 	ID3D12ResourceAllocator* ResourceAllocator = nullptr;
 	FRHIGPUMask GPUMask = FRHIGPUMask::FromIndex(GPUIndex);
 	bool bHasInitialData = false;
@@ -3347,6 +3350,30 @@ static TRefCountPtr<FD3D12Buffer> CreateRayTracingBuffer(FD3D12Adapter* Adapter,
 	return Result;
 }
 
+FString GetGeometryInitializerDebugString(const FRayTracingGeometryInitializer& Initializer)
+{
+	TStringBuilder<128> Result;
+
+	Result << "DebugName=" << Initializer.DebugName.ToString();
+	Result << " NumSegments=" << Initializer.Segments.Num();
+	Result << " NumPrims=" << Initializer.TotalPrimitiveCount;
+	if (Initializer.IndexBuffer)
+	{
+		Result << " IndexStride=" << Initializer.IndexBuffer->GetStride();
+	}
+	else
+	{
+		Result << " NonIndexed";
+	}
+
+	if (Initializer.OfflineData)
+	{
+		Result << " HasOfflineData";
+	}
+
+	return Result.ToString();
+}
+
 FD3D12RayTracingGeometry::FD3D12RayTracingGeometry(FD3D12Adapter* Adapter, const FRayTracingGeometryInitializer& InInitializer)
 	: FRHIRayTracingGeometry(InInitializer), FD3D12AdapterChild(Adapter)
 {
@@ -3411,6 +3438,10 @@ FD3D12RayTracingGeometry::FD3D12RayTracingGeometry(FD3D12Adapter* Adapter, const
 		SizeInfo = RHICalcRayTracingGeometrySize(Initializer);
 	}
 
+	checkf(SizeInfo.ResultSize != 0,
+		TEXT("Unexpected acceleration structure buffer size (0).\nGeometry initializer details:\n%s"),
+		*GetGeometryInitializerDebugString(Initializer));
+
 	// If this RayTracingGeometry going to be used as streaming destination 
 	// we don't want to allocate its memory as it will be replaced later by streamed version
 	// but we still need correct SizeInfo as it is used to estimate its memory requirements outside of RHI.
