Fix compare operator for RHIMatchPrecachePSOInitializers where LHS ImmutableSamplerState was checked on both sides

diff --git a/Engine/Source/Runtime/D3D12RHI/Private/D3D12State.cpp b/Engine/Source/Runtime/D3D12RHI/Private/D3D12State.cpp
--- a/Engine/Source/Runtime/D3D12RHI/Private/D3D12State.cpp
+++ b/Engine/Source/Runtime/D3D12RHI/Private/D3D12State.cpp
@@ -501,7 +501,7 @@ uint64 FD3D12DynamicRHI::RHIComputePrecachePSOHash(const FGraphicsPipelineStateI
 bool FD3D12DynamicRHI::RHIMatchPrecachePSOInitializers(const FGraphicsPipelineStateInitializer& LHS, const FGraphicsPipelineStateInitializer& RHS)
 {
 	// first check non pointer objects
-	if (LHS.ImmutableSamplerState != LHS.ImmutableSamplerState ||
+	if (LHS.ImmutableSamplerState != RHS.ImmutableSamplerState ||
 		LHS.PrimitiveType != RHS.PrimitiveType ||
 		LHS.bDepthBounds != RHS.bDepthBounds ||
 		LHS.MultiViewCount != RHS.MultiViewCount ||
diff --git a/Engine/Source/Runtime/RHI/Private/DynamicRHI.cpp b/Engine/Source/Runtime/RHI/Private/DynamicRHI.cpp
--- a/Engine/Source/Runtime/RHI/Private/DynamicRHI.cpp
+++ b/Engine/Source/Runtime/RHI/Private/DynamicRHI.cpp
@@ -605,7 +605,7 @@ uint64 FDynamicRHI::RHIComputePrecachePSOHash(const FGraphicsPipelineStateInitia
 bool FDynamicRHI::RHIMatchPrecachePSOInitializers(const FGraphicsPipelineStateInitializer& LHS, const FGraphicsPipelineStateInitializer& RHS)
 {
 	// first check non pointer objects
-	if (LHS.ImmutableSamplerState != LHS.ImmutableSamplerState ||
+	if (LHS.ImmutableSamplerState != RHS.ImmutableSamplerState ||
 		LHS.PrimitiveType != RHS.PrimitiveType ||
 		LHS.bDepthBounds != RHS.bDepthBounds ||
 		LHS.MultiViewCount != RHS.MultiViewCount ||
