Check for nullptr on buffer resource buffer trying to set a debug name on it

diff --git a/Engine/Source/Runtime/D3D12RHI/Private/D3D12Buffer.cpp b/Engine/Source/Runtime/D3D12RHI/Private/D3D12Buffer.cpp
--- a/Engine/Source/Runtime/D3D12RHI/Private/D3D12Buffer.cpp
+++ b/Engine/Source/Runtime/D3D12RHI/Private/D3D12Buffer.cpp
@@ -869,7 +869,7 @@ void FD3D12DynamicRHI::RHIBindDebugLabelName(FRHIBuffer* BufferRHI, const TCHAR*
 	FD3D12Buffer* Buffer = FD3D12DynamicRHI::ResourceCast(BufferRHI);
 
 	// only rename the underlying d3d12 resource if it's not sub allocated (requires resource state tracking or stand alone allocated)
-	if (Buffer->GetResource()->RequiresResourceStateTracking() || Buffer->ResourceLocation.GetType() == FD3D12ResourceLocation::ResourceLocationType::eStandAlone)
+	if (Buffer->GetResource() != nullptr && (Buffer->GetResource()->RequiresResourceStateTracking() || Buffer->ResourceLocation.GetType() == FD3D12ResourceLocation::ResourceLocationType::eStandAlone))
 	{
 		FD3D12Buffer::FLinkedObjectIterator BufferIt(Buffer);
 
