Fix for assert in D3D12 due to static uniform buffers being cleared when resizing viewports or submitting command lists within a pass.

diff --git a/Engine/Source/Runtime/D3D12RHI/Private/D3D12CommandContext.cpp b/Engine/Source/Runtime/D3D12RHI/Private/D3D12CommandContext.cpp
--- a/Engine/Source/Runtime/D3D12RHI/Private/D3D12CommandContext.cpp
+++ b/Engine/Source/Runtime/D3D12RHI/Private/D3D12CommandContext.cpp
@@ -542,7 +542,7 @@ void FD3D12CommandContextBase::RHIBeginFrame()
 	}
 }
 
-void FD3D12CommandContext::ClearState()
+void FD3D12CommandContext::ClearState(EClearStateFlags Flags)
 {
 	StateCache.ClearState();
 
@@ -551,7 +551,11 @@ void FD3D12CommandContext::ClearState()
 
 	FMemory::Memzero(BoundUniformBuffers, sizeof(BoundUniformBuffers));
 	FMemory::Memzero(DirtyUniformBuffers, sizeof(DirtyUniformBuffers));
-	FMemory::Memzero(GlobalUniformBuffers.GetData(), GlobalUniformBuffers.Num() * sizeof(FRHIUniformBuffer*));
+
+	if (EnumHasAnyFlags(Flags, EClearStateFlags::StaticUniformBuffers))
+	{
+		FMemory::Memzero(GlobalUniformBuffers.GetData(), GlobalUniformBuffers.Num() * sizeof(FRHIUniformBuffer*));
+	}
 
 	for (int i = 0; i < UE_ARRAY_COUNT(BoundUniformBufferRefs); i++)
 	{
diff --git a/Engine/Source/Runtime/D3D12RHI/Private/D3D12CommandContext.h b/Engine/Source/Runtime/D3D12RHI/Private/D3D12CommandContext.h
--- a/Engine/Source/Runtime/D3D12RHI/Private/D3D12CommandContext.h
+++ b/Engine/Source/Runtime/D3D12RHI/Private/D3D12CommandContext.h
@@ -123,7 +123,14 @@ public:
 
 	void Finish(TArray<FD3D12CommandListHandle>& CommandLists);
 
-	void ClearState();
+	enum class EClearStateFlags
+	{
+		None = 0,
+		StaticUniformBuffers = 1 << 0,
+		All = StaticUniformBuffers
+	};
+
+	void ClearState(EClearStateFlags Flags = EClearStateFlags::All);
 	void ConditionalClearShaderResource(FD3D12ResourceLocation* Resource);
 	void ClearAllShaderResources();
 
diff --git a/Engine/Source/Runtime/D3D12RHI/Private/D3D12Commands.cpp b/Engine/Source/Runtime/D3D12RHI/Private/D3D12Commands.cpp
--- a/Engine/Source/Runtime/D3D12RHI/Private/D3D12Commands.cpp
+++ b/Engine/Source/Runtime/D3D12RHI/Private/D3D12Commands.cpp
@@ -2509,7 +2509,7 @@ void FD3D12CommandContext::RHISubmitCommandsHint()
 	// Clear the state on requested command hint flush and not in scene or viewport rendering
 	if (!IsDrawingSceneOrViewport())
 	{
-		ClearState();
+		ClearState(EClearStateFlags::None);
 	}
 }
 
diff --git a/Engine/Source/Runtime/D3D12RHI/Private/D3D12Viewport.cpp b/Engine/Source/Runtime/D3D12RHI/Private/D3D12Viewport.cpp
--- a/Engine/Source/Runtime/D3D12RHI/Private/D3D12Viewport.cpp
+++ b/Engine/Source/Runtime/D3D12RHI/Private/D3D12Viewport.cpp
@@ -474,10 +474,10 @@ void FD3D12Viewport::Resize(uint32 InSizeX, uint32 InSizeY, bool bInIsFullscreen
 	for (uint32 GPUIndex : FRHIGPUMask::All())
 	{
 		FD3D12Device* Device = Adapter->GetDevice(GPUIndex);
-		Device->GetDefaultCommandContext().ClearState();
+		Device->GetDefaultCommandContext().ClearState(FD3D12CommandContext::EClearStateFlags::None);
 		if (GEnableAsyncCompute)
 		{
-			Device->GetDefaultAsyncComputeContext().ClearState();
+			Device->GetDefaultAsyncComputeContext().ClearState(FD3D12CommandContext::EClearStateFlags::None);
 		}
 	}
 
