Add manual LLM tracking to D3D12TextureStats and the backing heap of the D3D12TransientAllocator so that stat LLMPlatform shows memory usage for the Graphics tag comparable to that of D3D11.

diff --git a/Engine/Source/Runtime/D3D12RHI/Private/D3D12Texture.cpp b/Engine/Source/Runtime/D3D12RHI/Private/D3D12Texture.cpp
--- a/Engine/Source/Runtime/D3D12RHI/Private/D3D12Texture.cpp
+++ b/Engine/Source/Runtime/D3D12RHI/Private/D3D12Texture.cpp
@@ -402,10 +402,20 @@ void FD3D12TextureStats::UpdateD3D12TextureStats(FD3D12Texture& Texture, const D
 
 	if (TextureSize > 0)
 	{
+#if PLATFORM_WINDOWS
+		// On Windows there is no way to hook into the low level d3d allocations and frees.
+		// This means that we must manually add the tracking here.
+		LLM(FLowLevelMemTracker::Get().OnLowLevelAlloc(ELLMTracker::Platform, Texture.GetResource(), TextureSize, ELLMTag::GraphicsPlatform));
+		LLM(FLowLevelMemTracker::Get().OnLowLevelAlloc(ELLMTracker::Default, Texture.GetResource(), TextureSize, ELLMTag::Textures));
+#endif
 		INC_DWORD_STAT(STAT_D3D12TexturesAllocated);
 	}
 	else
 	{
+#if PLATFORM_WINDOWS
+		LLM(FLowLevelMemTracker::Get().OnLowLevelFree(ELLMTracker::Platform, Texture.GetResource()));
+		LLM(FLowLevelMemTracker::Get().OnLowLevelFree(ELLMTracker::Default, Texture.GetResource()));
+#endif
 		INC_DWORD_STAT(STAT_D3D12TexturesReleased);
 	}
 }
diff --git a/Engine/Source/Runtime/D3D12RHI/Private/D3D12TransientResourceAllocator.cpp b/Engine/Source/Runtime/D3D12RHI/Private/D3D12TransientResourceAllocator.cpp
--- a/Engine/Source/Runtime/D3D12RHI/Private/D3D12TransientResourceAllocator.cpp
+++ b/Engine/Source/Runtime/D3D12RHI/Private/D3D12TransientResourceAllocator.cpp
@@ -71,6 +71,9 @@ FD3D12TransientHeap::FD3D12TransientHeap(const FInitializer& Initializer, FD3D12
 		VERIFYD3D12RESULT(D3DDevice->CreateHeap(&Desc, IID_PPV_ARGS(&D3DHeap)));
 
 #if PLATFORM_WINDOWS
+		// On Windows there is no way to hook into the low level d3d allocations and frees.
+		// This means that we must manually add the tracking here.
+		LLM(FLowLevelMemTracker::Get().OnLowLevelAlloc(ELLMTracker::Platform, D3DHeap, Desc.SizeInBytes, ELLMTag::GraphicsPlatform));
 		// Boost priority to make sure it's not paged out
 		TRefCountPtr<ID3D12Device5> D3DDevice5;
 		if (SUCCEEDED(D3DDevice->QueryInterface(IID_PPV_ARGS(D3DDevice5.GetInitReference()))))
@@ -97,6 +100,9 @@ FD3D12TransientHeap::~FD3D12TransientHeap()
 		D3D12_HEAP_DESC Desc = Heap->GetHeapDesc();
 		DEC_MEMORY_STAT_BY(STAT_D3D12TransientHeaps, Desc.SizeInBytes);
 		DEC_MEMORY_STAT_BY(STAT_D3D12MemoryCurrentTotal, Desc.SizeInBytes);
+#if PLATFORM_WINDOWS
+		LLM(FLowLevelMemTracker::Get().OnLowLevelFree(ELLMTracker::Platform, Heap->GetHeap()));
+#endif
 	}
 }
 
