Add CVar r.RayTracing.RequireSM6 (default=0) to explicitly control when ray tracing shaders/features should be available, instead of relying on project target platform shader model settings. This is required because target platform settings are not available when running a packaged game.

The engine uses GRayTracingPlaformMask to compute ShouldCompilePermutation() for shaders. This is filled based on platform settings and capabilities when running the editor. When running a standalone game, the mask is filled purely based on whether ray tracing is enabled for the project (r.RayTracing) and the current runtime environment supports it (GRHISupportsRayTracing).

The new CVar is meant to be consistent between the cook time and run time. This is achieved by setting GRHISupportsRayTracing=false when running on SM5 if r.RayTracing.RequireSM6 is 1.

diff --git a/Engine/Source/Developer/Windows/WindowsTargetPlatform/Private/GenericWindowsTargetPlatform.h b/Engine/Source/Developer/Windows/WindowsTargetPlatform/Private/GenericWindowsTargetPlatform.h
--- a/Engine/Source/Developer/Windows/WindowsTargetPlatform/Private/GenericWindowsTargetPlatform.h
+++ b/Engine/Source/Developer/Windows/WindowsTargetPlatform/Private/GenericWindowsTargetPlatform.h
@@ -21,7 +21,7 @@
 	#include "AudioCompressionSettings.h"
 #endif // WITH_ENGINE
 #include "Windows/WindowsPlatformProperties.h"
-
+#include "HAL/IConsoleManager.h"
 
 #define LOCTEXT_NAMESPACE "TGenericWindowsTargetPlatform"
 
@@ -271,18 +271,22 @@ public:
 	{
 		if (UsesRayTracing())
 		{
-			// Only D3D12 and Vulkan can possibly support ray tracing
 			TCHAR const* RelevantSettings[] = 
 			{
-				TEXT("TargetedRHIs"),
-				TEXT("D3D12TargetedShaderFormats"),
 				TEXT("VulkanTargetedShaderFormats")
 			};
 
 			GetAllTargetedShaderFormatsInternal(RelevantSettings, OutFormats);
 
-			// We always support ray tracing shaders when cooking for SM6, however we may skip them for SM5 based on project settings.
+			// We always support ray tracing shaders when cooking for D3D12 SM6, however we may skip them for SM5 based on project settings.
 			OutFormats.AddUnique(FName(TEXT("PCD3D_SM6")));
+
+			static IConsoleVariable* RequireSM6CVar = IConsoleManager::Get().FindConsoleVariable(TEXT("r.RayTracing.RequireSM6"));
+			const bool bRequireSM6 = RequireSM6CVar && RequireSM6CVar->GetBool();
+			if (!bRequireSM6)
+			{
+				OutFormats.AddUnique(FName(TEXT("PCD3D_SM5")));
+			}
 		}
 	}
 
diff --git a/Engine/Source/Runtime/D3D12RHI/Private/D3D12Adapter.cpp b/Engine/Source/Runtime/D3D12RHI/Private/D3D12Adapter.cpp
--- a/Engine/Source/Runtime/D3D12RHI/Private/D3D12Adapter.cpp
+++ b/Engine/Source/Runtime/D3D12RHI/Private/D3D12Adapter.cpp
@@ -1030,6 +1030,11 @@ void FD3D12Adapter::InitializeDevices()
 			D3D12_FEATURE_DATA_D3D12_OPTIONS5 D3D12Caps5 = {};
 			if (SUCCEEDED(RootDevice->CheckFeatureSupport(D3D12_FEATURE_D3D12_OPTIONS5, &D3D12Caps5, sizeof(D3D12Caps5))))
 			{
+				static IConsoleVariable* RequireSM6CVar = IConsoleManager::Get().FindConsoleVariable(TEXT("r.RayTracing.RequireSM6"));
+				const bool bRequireSM6 = RequireSM6CVar && RequireSM6CVar->GetBool();
+
+				const bool bRayTracingAllowedOnCurrentShaderPlatform = (bRequireSM6 == false) || (GMaxRHIShaderPlatform == SP_PCD3D_SM6);
+
 				if (D3D12Caps5.RaytracingTier >= D3D12_RAYTRACING_TIER_1_0
 					&& GetResourceBindingTier() >= D3D12_RESOURCE_BINDING_TIER_2
 					&& RootDevice5
@@ -1039,13 +1044,20 @@ void FD3D12Adapter::InitializeDevices()
 					if (D3D12Caps5.RaytracingTier >= D3D12_RAYTRACING_TIER_1_1
 						&& RootDevice7)
 					{
-						UE_LOG(LogD3D12RHI, Log, TEXT("D3D12 ray tracing tier 1.1 is supported."));
-
-						GRHISupportsRayTracing = RHISupportsRayTracing(GMaxRHIShaderPlatform);
-						GRHISupportsRayTracingShaders = GRHISupportsRayTracing && RHISupportsRayTracingShaders(GMaxRHIShaderPlatform);
-
-						GRHISupportsRayTracingPSOAdditions = true;
-						GRHISupportsInlineRayTracing = GRHISupportsRayTracing && RHISupportsInlineRayTracing(GMaxRHIShaderPlatform) && (GMaxRHIFeatureLevel >= ERHIFeatureLevel::SM6);
+						if (bRayTracingAllowedOnCurrentShaderPlatform)
+						{
+							UE_LOG(LogD3D12RHI, Log, TEXT("D3D12 ray tracing tier 1.1 is supported."));
+
+							GRHISupportsRayTracing = RHISupportsRayTracing(GMaxRHIShaderPlatform);
+							GRHISupportsRayTracingShaders = GRHISupportsRayTracing && RHISupportsRayTracingShaders(GMaxRHIShaderPlatform);
+
+							GRHISupportsRayTracingPSOAdditions = true;
+							GRHISupportsInlineRayTracing = GRHISupportsRayTracing && RHISupportsInlineRayTracing(GMaxRHIShaderPlatform) && (GMaxRHIFeatureLevel >= ERHIFeatureLevel::SM6);
+						}
+						else
+						{
+ 							UE_LOG(LogD3D12RHI, Log, TEXT("Ray tracing is disabled because SM6 shader platform is required (r.RayTracing.RequireSM6=1)."));
+						}
 					}
 					else
 					{
diff --git a/Engine/Source/Runtime/RenderCore/Private/RenderUtils.cpp b/Engine/Source/Runtime/RenderCore/Private/RenderUtils.cpp
--- a/Engine/Source/Runtime/RenderCore/Private/RenderUtils.cpp
+++ b/Engine/Source/Runtime/RenderCore/Private/RenderUtils.cpp
@@ -45,6 +45,15 @@ FAutoConsoleVariableRef CVarRayTracingEnableInEditor(
 	ECVF_ReadOnly
 );
 
+static int32 GRayTracingRequireSM6 = 0;
+static FAutoConsoleVariableRef CVarRayTracingSM5(
+	TEXT("r.RayTracing.RequireSM6"),
+	GRayTracingRequireSM6,
+	TEXT("Whether ray tracing shaders and features should only be available when targetting and running SM6. If disabled, ray tracing shaders will also be available when running in SM5 mode. (default = 0, allow SM5 and SM6)"),
+	ECVF_RenderThreadSafe | ECVF_ReadOnly
+);
+
+
 FBufferWithRDG::FBufferWithRDG() = default;
 FBufferWithRDG::FBufferWithRDG(const FBufferWithRDG & Other) = default;
 FBufferWithRDG& FBufferWithRDG::operator=(const FBufferWithRDG & Other) = default;
