Store segment FirstPrimitive in FHitGroupSystemRootConstants.

diff --git a/Engine/Shaders/Shared/RayTracingBuiltInResources.h b/Engine/Shaders/Shared/RayTracingBuiltInResources.h
--- a/Engine/Shaders/Shared/RayTracingBuiltInResources.h
+++ b/Engine/Shaders/Shared/RayTracingBuiltInResources.h
@@ -41,6 +41,9 @@ struct FHitGroupSystemRootConstants
 	// Offset into HitGroupSystemIndexBuffer
 	UINT_TYPE IndexBufferOffsetInBytes;
 
+	// First primitive of the segment (as set in FRayTracingGeometrySegment)
+	UINT_TYPE FirstPrimitive;
+
 	// User-provided constant assigned to the hit group
 	UINT_TYPE UserData;
 
@@ -48,6 +51,8 @@ struct FHitGroupSystemRootConstants
 	// Can be used to emulate SV_InstanceID in ray tracing shaders.
 	UINT_TYPE BaseInstanceIndex;
 
+	UINT_TYPE Pad0;
+
 	// Helper functions
 
 	UINT_TYPE GetIndexStride()
diff --git a/Engine/Source/Runtime/D3D12RHI/Private/D3D12RayTracing.cpp b/Engine/Source/Runtime/D3D12RHI/Private/D3D12RayTracing.cpp
--- a/Engine/Source/Runtime/D3D12RHI/Private/D3D12RayTracing.cpp
+++ b/Engine/Source/Runtime/D3D12RHI/Private/D3D12RayTracing.cpp
@@ -3805,6 +3805,7 @@ void FD3D12RayTracingGeometry::SetupHitGroupSystemParameters(uint32 InGPUIndex)
 		{
 			SystemParameters.IndexBuffer = IndexBuffer->ResourceLocation.GetGPUVirtualAddress();
 			SystemParameters.RootConstants.IndexBufferOffsetInBytes = Initializer.IndexBufferOffset + IndexStride * Segment.FirstPrimitive * FD3D12RayTracingGeometry::IndicesPerPrimitive;
+			SystemParameters.RootConstants.FirstPrimitive = Segment.FirstPrimitive;
 		}
 
 		HitGroupSystemParametersForThisGPU.Add(SystemParameters);
@@ -4433,7 +4434,7 @@ void FD3D12CommandContext::RHIBuildAccelerationStructures(const TArrayView<const
 			for (int32 i = 0; i < P.Segments.Num(); ++i)
 			{
 				checkf(P.Segments[i].MaxVertices <= Geometry->Initializer.Segments[i].MaxVertices,
-					TEXT("Maximum number of vertices in a segment (%u) must not be smaller than what was declared during FRHIRayTracingGeometry creation (%u), as this controls BLAS memory allocation."),
+					TEXT("Maximum number of vertices in a segment (%u) must not be larger than what was declared during FRHIRayTracingGeometry creation (%u), as this controls BLAS memory allocation."),
 					P.Segments[i].MaxVertices, Geometry->Initializer.Segments[i].MaxVertices
 				);
 
diff --git a/Engine/Source/Runtime/D3D12RHI/Private/D3D12RootSignature.cpp b/Engine/Source/Runtime/D3D12RHI/Private/D3D12RootSignature.cpp
--- a/Engine/Source/Runtime/D3D12RHI/Private/D3D12RootSignature.cpp
+++ b/Engine/Source/Runtime/D3D12RHI/Private/D3D12RootSignature.cpp
@@ -170,10 +170,12 @@ FD3D12RootSignatureDesc::FD3D12RootSignatureDesc(const FD3D12QuantizedBoundShade
 		//          FHitGroupSystemRootConstants:
 		// 4 bytes: index/vertex fetch configuration as root constant (bitfield defining index and vertex formats)
 		// 4 bytes: index buffer offset in bytes
+		// 4 bytes: first primitive of the segment (as set in FRayTracingGeometrySegment)
 		// 4 bytes: hit group user data
+		// 4 bytes: index of the first instance that belongs to the current batch.
 		// 4 bytes: unused padding to ensure the next parameter is aligned to 8-byte boundary
 		// -----------
-		// 32 bytes
+		// 40 bytes
 
 		check(RootParameterCount == 0 && RootParametersSize == 0); // We expect system RT parameters to come first
 
