Removing workaround which disables the transient allocator for certain GPUs on Windows 11, as the driver bug has been fixed and we've set the minimum required driver version to the release which contains the fix.

#rnx

diff --git a/Engine/Source/Runtime/D3D12RHI/Private/D3D12CommandContext.cpp b/Engine/Source/Runtime/D3D12RHI/Private/D3D12CommandContext.cpp
--- a/Engine/Source/Runtime/D3D12RHI/Private/D3D12CommandContext.cpp
+++ b/Engine/Source/Runtime/D3D12RHI/Private/D3D12CommandContext.cpp
@@ -1053,15 +1053,7 @@ void FD3D12DynamicRHI::RHIReleaseTransition(FRHITransition* Transition)
 
 IRHITransientResourceAllocator* FD3D12DynamicRHI::RHICreateTransientResourceAllocator()
 {
-	// Some machine configurations have known issues when transient resource aliasing is used.
-	if (GD3D12WorkaroundFlags.bAllowTransientResourceAllocator)
-	{
-		return new FD3D12TransientResourceHeapAllocator(GetAdapter().GetOrCreateTransientHeapCache());
-	}
-	else
-	{
-		return nullptr;
-	}
+	return new FD3D12TransientResourceHeapAllocator(GetAdapter().GetOrCreateTransientHeapCache());
 }
 
 //////////////////////////////////////////////////////////////////////////////////////////////////////////
diff --git a/Engine/Source/Runtime/D3D12RHI/Private/D3D12RHIPrivate.h b/Engine/Source/Runtime/D3D12RHI/Private/D3D12RHIPrivate.h
--- a/Engine/Source/Runtime/D3D12RHI/Private/D3D12RHIPrivate.h
+++ b/Engine/Source/Runtime/D3D12RHI/Private/D3D12RHIPrivate.h
@@ -245,12 +245,6 @@ struct FD3D12WorkaroundFlags
 	* This is less efficient and can also trigger another known issue with D3D12 Agility version <= 4.
 	*/
 	bool bAllowGetShaderIdentifierOnCollectionSubObject = true;
-
-	/**
-	* Some machine configurations have known issues when transient resource aliasing is used.
-	* If we detect such configuration, we can fall back to non-aliasing code path which is much less efficient.
-	*/
-	bool bAllowTransientResourceAllocator = true;
 };
 
 extern FD3D12WorkaroundFlags GD3D12WorkaroundFlags;
diff --git a/Engine/Source/Runtime/D3D12RHI/Private/Windows/WindowsD3D12Device.cpp b/Engine/Source/Runtime/D3D12RHI/Private/Windows/WindowsD3D12Device.cpp
--- a/Engine/Source/Runtime/D3D12RHI/Private/Windows/WindowsD3D12Device.cpp
+++ b/Engine/Source/Runtime/D3D12RHI/Private/Windows/WindowsD3D12Device.cpp
@@ -1283,29 +1283,6 @@ void FD3D12DynamicRHI::Init()
 				TEXT("Please update to NVIDIA driver version 466.11 or newer."));
 		}
 
-		// Enable a workaround for a known driver bug affecting Ampere GPUs on Windows 11 (Jira UE-132964)
-		// As of 2022-02-09, there is no expected fix date or driver version.
-		if (FPlatformMisc::VerifyWindowsVersion(10, 0, 22000) // Win11 is build number is 22000+
-			&& IsNvidiaAmpereGPU(AdapterDesc.DeviceId))
-		{
-			int32 TransientAllocatorValue = 0;
-			const bool bForceTransientAllocatorState = FParse::Value(FCommandLine::Get(), TEXT("rdgtransientallocator="), TransientAllocatorValue);
-			if (bForceTransientAllocatorState && TransientAllocatorValue == 1)
-			{
-				UE_LOG(LogD3D12RHI, Warning,
-					TEXT("Current machine configuration is known to be affected by a driver bug triggered by memory aliasing. ")
-					TEXT("Transient resource allocator will still be used because it is explicitly requested using '-rdgtransientallocator=1' command line."));
-			}
-			else
-			{
-				GD3D12WorkaroundFlags.bAllowTransientResourceAllocator = false;
-				UE_LOG(LogD3D12RHI, Warning, 
-					TEXT("GD3D12WorkaroundFlags.bAllowTransientResourceAllocator is disabled due to a known issue with current GPU, driver and OS version. ")
-					TEXT("This workaround can be bypassed using '-rdgtransientallocator=1' command line.")
-				);
-			}
-		}
-
 		if (GRHISupportsRayTracing
 			&& IsRayTracingEmulated(AdapterDesc.DeviceId))
 		{
