Fix backbuffer going out of sync between RHI thread and render thread

diff --git a/Engine/Source/Runtime/D3D12RHI/Private/D3D12Viewport.cpp b/Engine/Source/Runtime/D3D12RHI/Private/D3D12Viewport.cpp
--- a/Engine/Source/Runtime/D3D12RHI/Private/D3D12Viewport.cpp
+++ b/Engine/Source/Runtime/D3D12RHI/Private/D3D12Viewport.cpp
@@ -655,7 +655,8 @@ void FD3D12Viewport::AdvanceExpectedBackBufferIndex_RenderThread()
 {
 	bool bNeedsNativePresent = IsValidRef(CustomPresent) ? 
 		CustomPresent->NeedsNativePresent() || CustomPresent->NeedsAdvanceBackbuffer() : true;
-	if (bNeedsNativePresent)
+
+	if (bNeedsNativePresent && IsPresentAllowed())
 	{
 #if WITH_MGPU
 		FScopeLock Lock(&ExpectedBackBufferIndexLock);
@@ -878,6 +879,11 @@ void FD3D12Viewport::PresentWithVsyncDWM()
 
 bool FD3D12Viewport::Present(bool bLockToVsync)
 {
+	if (!IsPresentAllowed())
+	{
+		return false;
+	}
+
 	FD3D12Adapter* Adapter = GetParentAdapter();
 	
 	for (uint32 GPUIndex : FRHIGPUMask::All())
@@ -1097,12 +1103,8 @@ void FD3D12CommandContextBase::RHIEndDrawingViewport(FRHIViewport* ViewportRHI,
 	ParentAdapter->SubmissionGapRecorder.SetPresentSlotIdx(CurrentSlotIdx);
 #endif
 
-	bool bNativelyPresented = false; 
-	if (!RHI.RHIIsRenderingSuspended())
-	{
-		bNativelyPresented = Viewport->Present(bLockToVsync);
-	}
-
+	bool bNativelyPresented = Viewport->Present(bLockToVsync);
+	
 	// Multi-GPU support : here each GPU wait's for it's own frame completion. Note that even in AFR, each GPU renders an (empty) frame.
 	if (bNativelyPresented)
 	{
diff --git a/Engine/Source/Runtime/D3D12RHI/Private/Windows/WindowsD3D12Viewport.cpp b/Engine/Source/Runtime/D3D12RHI/Private/Windows/WindowsD3D12Viewport.cpp
--- a/Engine/Source/Runtime/D3D12RHI/Private/Windows/WindowsD3D12Viewport.cpp
+++ b/Engine/Source/Runtime/D3D12RHI/Private/Windows/WindowsD3D12Viewport.cpp
@@ -630,6 +630,12 @@ FD3D12Texture* FD3D12Viewport::GetBackBuffer_RenderThread() const
 	return DummyBackBuffer;
 }
 
+bool FD3D12Viewport::IsPresentAllowed()
+{
+	FD3D12DynamicRHI& RHI = *ParentAdapter->GetOwningRHI();
+	return !RHI.RHIIsRenderingSuspended();
+}
+
 void FD3D12DynamicRHI::RHIGetDisplaysInformation(FDisplayInformationArray& OutDisplayInformation)
 {
 	OutDisplayInformation.Append(DisplayList);
diff --git a/Engine/Source/Runtime/D3D12RHI/Public/D3D12Viewport.h b/Engine/Source/Runtime/D3D12RHI/Public/D3D12Viewport.h
--- a/Engine/Source/Runtime/D3D12RHI/Public/D3D12Viewport.h
+++ b/Engine/Source/Runtime/D3D12RHI/Public/D3D12Viewport.h
@@ -167,6 +167,7 @@ public:
 	}
 
 private:
+	bool IsPresentAllowed();
 
 	/**
 	 * Create the dummy back buffer textures
