change log verbosity for SetColorSpace1 since it can happen frequently in editor

diff --git a/Engine/Source/Runtime/D3D12RHI/Private/Windows/WindowsD3D12Viewport.cpp b/Engine/Source/Runtime/D3D12RHI/Private/Windows/WindowsD3D12Viewport.cpp
--- a/Engine/Source/Runtime/D3D12RHI/Private/Windows/WindowsD3D12Viewport.cpp
+++ b/Engine/Source/Runtime/D3D12RHI/Private/Windows/WindowsD3D12Viewport.cpp
@@ -552,7 +552,7 @@ void FD3D12Viewport::EnsureColorSpace(EDisplayColorGamut DisplayGamut, EDisplayO
 		if (SUCCEEDED(hr) && ((ColorSpaceSupport & DXGI_SWAP_CHAIN_COLOR_SPACE_SUPPORT_FLAG_PRESENT) != 0))
 		{
 			VERIFYD3D12RESULT(SwapChain4->SetColorSpace1(NewColorSpace));
-			UE_LOG(LogD3D12RHI, Log, TEXT("Setting color space on swap chain (%#016llx): %s"), SwapChain4.GetReference(), *NewColorSpaceName);
+			UE_LOG(LogD3D12RHI, Verbose, TEXT("Setting color space on swap chain (%#016llx): %s"), SwapChain4.GetReference(), *NewColorSpaceName);
 			ColorSpace = NewColorSpace;
 		}
 		else
diff --git a/Engine/Source/Runtime/Windows/D3D11RHI/Private/Windows/WindowsD3D11Viewport.cpp b/Engine/Source/Runtime/Windows/D3D11RHI/Private/Windows/WindowsD3D11Viewport.cpp
--- a/Engine/Source/Runtime/Windows/D3D11RHI/Private/Windows/WindowsD3D11Viewport.cpp
+++ b/Engine/Source/Runtime/Windows/D3D11RHI/Private/Windows/WindowsD3D11Viewport.cpp
@@ -343,7 +343,7 @@ inline void EnsureColorSpace(IDXGISwapChain* SwapChain, EDisplayColorGamut Displ
 		if (SUCCEEDED(hr) && (ColorSpaceSupport & DXGI_SWAP_CHAIN_COLOR_SPACE_SUPPORT_FLAG_PRESENT))
 		{
 			swapChain3->SetColorSpace1(NewColorSpace);
-			UE_LOG(LogD3D11RHI, Log, TEXT("Setting color space on swap chain: %s"), *ColorSpaceName);
+			UE_LOG(LogD3D11RHI, Verbose, TEXT("Setting color space on swap chain: %s"), *ColorSpaceName);
 		}
 		else
 		{
