Add CVar r.RayTracing.RequireSM6 (default=0) to explicitly control when ray tracing shaders/features should be available, instead of relying on project target platform shader model settings. This is required because target platform settings are not available when running a packaged game.

The engine uses GRayTracingPlaformMask to compute ShouldCompilePermutation() for shaders. This is filled based on platform settings and capabilities when running the editor. When running a standalone game, the mask is filled purely based on whether ray tracing is enabled for the project (r.RayTracing) and the current runtime environment supports it (GRHISupportsRayTracing).

The new CVar is meant to be consistent between the cook time and run time. This is achieved by setting GRHISupportsRayTracing=false when running on SM5 if r.RayTracing.RequireSM6 is 1.

diff --git a/Engine/Source/Runtime/D3D12RHI/Private/D3D12Adapter.cpp b/Engine/Source/Runtime/D3D12RHI/Private/D3D12Adapter.cpp
--- a/Engine/Source/Runtime/D3D12RHI/Private/D3D12Adapter.cpp
+++ b/Engine/Source/Runtime/D3D12RHI/Private/D3D12Adapter.cpp
@@ -765,6 +765,7 @@ void FD3D12Adapter::CreateRootDevice(bool bWithDebug)
 		GNumExplicitGPUsForRendering = Desc.NumDeviceNodes;
 		UE_LOG(LogD3D12RHI, Log, TEXT("Enabling multi-GPU with %d nodes"), Desc.NumDeviceNodes);
 	}
+	}
 
 	// Viewport ignores AFR if PresentGPU is specified.
 	int32 Dummy;
