Fix compile error due to bad merge

diff --git a/Engine/Source/Runtime/D3D12RHI/Private/D3D12Adapter.cpp b/Engine/Source/Runtime/D3D12RHI/Private/D3D12Adapter.cpp
--- a/Engine/Source/Runtime/D3D12RHI/Private/D3D12Adapter.cpp
+++ b/Engine/Source/Runtime/D3D12RHI/Private/D3D12Adapter.cpp
@@ -765,7 +765,6 @@ void FD3D12Adapter::CreateRootDevice(bool bWithDebug)
 		GNumExplicitGPUsForRendering = Desc.NumDeviceNodes;
 		UE_LOG(LogD3D12RHI, Log, TEXT("Enabling multi-GPU with %d nodes"), Desc.NumDeviceNodes);
 	}
-	}
 
 	// Viewport ignores AFR if PresentGPU is specified.
 	int32 Dummy;
