Remove SetHDRMetaData: EnsureColorSpace1 is our interface with the OS which will communicate the information to the display

diff --git a/Engine/Source/Runtime/D3D12RHI/Private/Windows/WindowsD3D12Viewport.cpp b/Engine/Source/Runtime/D3D12RHI/Private/Windows/WindowsD3D12Viewport.cpp
--- a/Engine/Source/Runtime/D3D12RHI/Private/Windows/WindowsD3D12Viewport.cpp
+++ b/Engine/Source/Runtime/D3D12RHI/Private/Windows/WindowsD3D12Viewport.cpp
@@ -39,7 +39,6 @@ FD3D12Viewport::FD3D12Viewport(class FD3D12Adapter* InParent, HWND InWindowHandl
 	bFullscreenLost(false),
 	PixelFormat(InPreferredPixelFormat),
 	bIsValid(true),
-	bHDRMetaDataSet(false),
 	ColorSpace(DXGI_COLOR_SPACE_RGB_FULL_G22_NONE_P709),
 	NumBackBuffers(WindowsDefaultNumBackBuffers),
 	DummyBackBuffer_RenderThread(nullptr),
@@ -407,20 +406,6 @@ void FD3D12Viewport::EnableHDR()
 {
 	if ( GRHISupportsHDROutput && IsHDREnabled() )
 	{
-		const float DisplayMaxOutputNits = HDRGetDisplayMaximumLuminance();
-		const float DisplayMinOutputNits = 0.0f;	// Min output of the display
-		const float DisplayMaxCLL = 0.0f;			// Max content light level in lumens (0.0 == unknown)
-		const float DisplayFALL = 0.0f;				// Frame average light level (0.0 == unknown)
-
-		// Ideally we can avoid setting TV meta data and instead the engine can do tone mapping based on the
-		// actual current display properties (display mapping).
-		SetHDRTVMode(true,
-			DisplayColorGamut,
-			DisplayMaxOutputNits,
-			DisplayMinOutputNits,
-			DisplayMaxCLL,
-			DisplayFALL);
-
 		// Ensure we have the correct color space set.
 		EnsureColorSpace(DisplayColorGamut, DisplayOutputFormat);
 	}
@@ -430,28 +415,8 @@ void FD3D12Viewport::ShutdownHDR()
 {
 	if (GRHISupportsHDROutput)
 	{
-		// Default SDR display data
-		const EDisplayColorGamut DisplayGamut = EDisplayColorGamut::sRGB_D65;
-		const EDisplayOutputFormat OutputDevice = EDisplayOutputFormat::SDR_sRGB;
-
-		// Note: These values aren't actually used.
-		const float DisplayMaxOutputNits = 100.0f;	// Max output of the display
-		const float DisplayMinOutputNits = 0.0f;	// Min output of the display
-		const float DisplayMaxCLL = 100.0f;			// Max content light level in lumens
-		const float DisplayFALL = 20.0f;			// Frame average light level
-
-		// Ideally we can avoid setting TV meta data and instead the engine can do tone mapping based on the
-		// actual current display properties (display mapping).
-		SetHDRTVMode(
-			false,
-			DisplayGamut,
-			DisplayMaxOutputNits,
-			DisplayMinOutputNits,
-			DisplayMaxCLL,
-			DisplayFALL);
-
 		// Ensure we have the correct color space set.
-		EnsureColorSpace(DisplayGamut, OutputDevice);
+		EnsureColorSpace(EDisplayColorGamut::sRGB_D65, EDisplayOutputFormat::SDR_sRGB);
 	}
 }
 
@@ -562,59 +527,6 @@ void FD3D12Viewport::EnsureColorSpace(EDisplayColorGamut DisplayGamut, EDisplayO
 	}
 }
 
-void FD3D12Viewport::SetHDRTVMode(bool bEnableHDR, EDisplayColorGamut DisplayGamut, float MaxOutputNits, float MinOutputNits, float MaxCLL, float MaxFALL)
-{
-	ensure(SwapChain4.GetReference());
-
-	static const DisplayChromacities DisplayChromacityList[] =
-	{
-		{ 0.64000f, 0.33000f, 0.30000f, 0.60000f, 0.15000f, 0.06000f, 0.31270f, 0.32900f }, // EDisplayColorGamut::sRGB_D65
-		{ 0.68000f, 0.32000f, 0.26500f, 0.69000f, 0.15000f, 0.06000f, 0.31270f, 0.32900f }, // EDisplayColorGamut::DCIP3_D65
-		{ 0.70800f, 0.29200f, 0.17000f, 0.79700f, 0.13100f, 0.04600f, 0.31270f, 0.32900f }, // EDisplayColorGamut::Rec2020_D65
-		{ 0.73470f, 0.26530f, 0.00000f, 1.00000f, 0.00010f,-0.07700f, 0.32168f, 0.33767f }, // EDisplayColorGamut::ACES_D60
-		{ 0.71300f, 0.29300f, 0.16500f, 0.83000f, 0.12800f, 0.04400f, 0.32168f, 0.33767f }, // EDisplayColorGamut::ACEScg_D60
-	};
-
-	if (bEnableHDR)
-	{
-		const DisplayChromacities& Chroma = DisplayChromacityList[(int32)DisplayGamut];
-
-		// Set HDR meta data
-		DXGI_HDR_METADATA_HDR10 HDR10MetaData = {};
-		HDR10MetaData.RedPrimary[0] = static_cast<uint16>(Chroma.RedX * 50000.0f);
-		HDR10MetaData.RedPrimary[1] = static_cast<uint16>(Chroma.RedY * 50000.0f);
-		HDR10MetaData.GreenPrimary[0] = static_cast<uint16>(Chroma.GreenX * 50000.0f);
-		HDR10MetaData.GreenPrimary[1] = static_cast<uint16>(Chroma.GreenY * 50000.0f);
-		HDR10MetaData.BluePrimary[0] = static_cast<uint16>(Chroma.BlueX * 50000.0f);
-		HDR10MetaData.BluePrimary[1] = static_cast<uint16>(Chroma.BlueY * 50000.0f);
-		HDR10MetaData.WhitePoint[0] = static_cast<uint16>(Chroma.WpX * 50000.0f);
-		HDR10MetaData.WhitePoint[1] = static_cast<uint16>(Chroma.WpY * 50000.0f);
-		HDR10MetaData.MaxMasteringLuminance = static_cast<uint32>(MaxOutputNits * 10000.0f);
-		HDR10MetaData.MinMasteringLuminance = static_cast<uint32>(MinOutputNits * 10000.0f);
-		HDR10MetaData.MaxContentLightLevel = static_cast<uint16>(MaxCLL);
-		HDR10MetaData.MaxFrameAverageLightLevel = static_cast<uint16>(MaxFALL);
-
-		VERIFYD3D12RESULT(SwapChain4->SetHDRMetaData(DXGI_HDR_METADATA_TYPE_HDR10, sizeof(HDR10MetaData), &HDR10MetaData));
-		UE_LOG(LogD3D12RHI, Log, TEXT("Setting HDR meta data on swap chain (%#016llx) using DisplayGamut %u:"), SwapChain4.GetReference(), static_cast<uint32>(DisplayGamut));
-		UE_LOG(LogD3D12RHI, Log, TEXT("\t\tMaxMasteringLuminance = %.4f nits"), HDR10MetaData.MaxMasteringLuminance * .0001f);
-		UE_LOG(LogD3D12RHI, Log, TEXT("\t\tMinMasteringLuminance = %.4f nits"), HDR10MetaData.MinMasteringLuminance * .0001f);
-		UE_LOG(LogD3D12RHI, Log, TEXT("\t\tMaxContentLightLevel = %u nits"), HDR10MetaData.MaxContentLightLevel);
-		UE_LOG(LogD3D12RHI, Log, TEXT("\t\tMaxFrameAverageLightLevel %u = nits"), HDR10MetaData.MaxFrameAverageLightLevel);
-		bHDRMetaDataSet = true;
-	}
-	else
-	{
-		if (bHDRMetaDataSet)
-		{
-			// Clear meta data.
-			VERIFYD3D12RESULT(SwapChain4->SetHDRMetaData(DXGI_HDR_METADATA_TYPE_NONE, 0, nullptr));
-			UE_LOG(LogD3D12RHI, Log, TEXT("Clearing HDR meta data on swap chain (%#016llx)."), SwapChain4.GetReference());
-			bHDRMetaDataSet = false;
-		}
-	}
-}
-
-
 void FD3D12Viewport::OnResumeRendering()
 {}
 
diff --git a/Engine/Source/Runtime/D3D12RHI/Public/D3D12Viewport.h b/Engine/Source/Runtime/D3D12RHI/Public/D3D12Viewport.h
--- a/Engine/Source/Runtime/D3D12RHI/Public/D3D12Viewport.h
+++ b/Engine/Source/Runtime/D3D12RHI/Public/D3D12Viewport.h
@@ -221,7 +221,6 @@ private:
 #endif
 
 #if PLATFORM_WINDOWS || PLATFORM_HOLOLENS
-	bool bHDRMetaDataSet;
 	DXGI_COLOR_SPACE_TYPE ColorSpace;
 #endif
 #endif // D3D12_VIEWPORT_EXPOSES_SWAP_CHAIN
@@ -290,13 +289,6 @@ private:
 #if PLATFORM_WINDOWS || PLATFORM_HOLOLENS
 	/** Ensure the correct color space is set on the swap chain */
 	void EnsureColorSpace(EDisplayColorGamut DisplayGamut, EDisplayOutputFormat OutputDevice);
-
-	/** 
-	 * Set HDR meta data. 
-	 * Note: Meta data should only be provided for TVs, not monitors. 
-	 * This is because the TV is doing the work to display the colors correctly.
-	 */
-	void SetHDRTVMode(bool bEnableHDR, EDisplayColorGamut DisplayGamut, float MaxOutputNits, float MinOutputNits, float MaxCLL, float MaxFALL);
 #endif
 };
 
