Disable async compute on nVidia by default because of Async compute GPU crashes

#rnx

diff --git a/Engine/Source/Runtime/D3D12RHI/Private/Windows/WindowsD3D12Device.cpp b/Engine/Source/Runtime/D3D12RHI/Private/Windows/WindowsD3D12Device.cpp
--- a/Engine/Source/Runtime/D3D12RHI/Private/Windows/WindowsD3D12Device.cpp
+++ b/Engine/Source/Runtime/D3D12RHI/Private/Windows/WindowsD3D12Device.cpp
@@ -1450,10 +1450,13 @@ void FD3D12DynamicRHI::Init()
 	// Allow async compute by default on nVidia cards which support PerPrimitiveShadingRateSupportedWithViewportIndexing 
 	// this should be a good metric according to nVidia itself (this is set for Ampere and newer cards)
 	bool bnVidiaAsyncComputeSupported = false;
+	// Disable async compute on nVidia by default because of Async compute GPU crashes (UE-163646)
+	/*
 	if (IsRHIDeviceNVIDIA() && Options6HR == S_OK && options.PerPrimitiveShadingRateSupportedWithViewportIndexing)
 	{
 		bnVidiaAsyncComputeSupported = true;
 	}
+	*/
 
 	GSupportsEfficientAsyncCompute = GAllowAsyncCompute && (FParse::Param(FCommandLine::Get(), TEXT("ForceAsyncCompute")) || (GRHISupportsParallelRHIExecute && (IsRHIDeviceAMD() || bnVidiaAsyncComputeSupported)));
 
