- Fixing D3D12 shared buffer creation failing with simultaneous access flag present
- Adding D3D12 resource utils to buffers (already present for textures)

diff --git a/Engine/Source/Runtime/D3D12RHI/Private/D3D12RHI.cpp b/Engine/Source/Runtime/D3D12RHI/Private/D3D12RHI.cpp
--- a/Engine/Source/Runtime/D3D12RHI/Private/D3D12RHI.cpp
+++ b/Engine/Source/Runtime/D3D12RHI/Private/D3D12RHI.cpp
@@ -562,6 +562,18 @@ uint32 FD3D12DynamicRHI::RHIGetResourceDeviceIndex(FRHIBuffer* InBuffer) const
 	return D3D12Buffer->GetParentDevice()->GetGPUIndex();
 }
 
+int64 FD3D12DynamicRHI::RHIGetResourceMemorySize(FRHIBuffer* InBuffer) const
+{
+	FD3D12Buffer* D3D12Buffer = ResourceCast(InBuffer);
+	return D3D12Buffer->ResourceLocation.GetSize();
+}
+
+bool FD3D12DynamicRHI::RHIIsResourcePlaced(FRHIBuffer* InBuffer) const
+{
+	FD3D12Buffer* D3D12Buffer = ResourceCast(InBuffer);
+	return D3D12Buffer->GetResource()->IsPlacedResource();
+}
+
 ID3D12Resource* FD3D12DynamicRHI::RHIGetResource(FRHITexture* InTexture) const
 {
 	return (ID3D12Resource*)InTexture->GetNativeResource();
diff --git a/Engine/Source/Runtime/D3D12RHI/Private/D3D12RHIPrivate.h b/Engine/Source/Runtime/D3D12RHI/Private/D3D12RHIPrivate.h
--- a/Engine/Source/Runtime/D3D12RHI/Private/D3D12RHIPrivate.h
+++ b/Engine/Source/Runtime/D3D12RHI/Private/D3D12RHIPrivate.h
@@ -551,6 +551,8 @@ public:
 	virtual FTextureCubeRHIRef RHICreateTextureCubeFromResource(EPixelFormat Format, ETextureCreateFlags TexCreateFlags, const FClearValueBinding& ClearValueBinding, ID3D12Resource* Resource) final override;
 	virtual ID3D12Resource* RHIGetResource(FRHIBuffer* InBuffer) const final override;
 	virtual uint32 RHIGetResourceDeviceIndex(FRHIBuffer* InBuffer) const final override;
+	virtual int64 RHIGetResourceMemorySize(FRHIBuffer* InBuffer) const final override;
+	virtual bool RHIIsResourcePlaced(FRHIBuffer* InBuffer) const final override;
 	virtual ID3D12Resource* RHIGetResource(FRHITexture* InTexture) const final override;
 	virtual uint32 RHIGetResourceDeviceIndex(FRHITexture* InTexture) const final override;
 	virtual int64 RHIGetResourceMemorySize(FRHITexture* InTexture) const final override;
diff --git a/Engine/Source/Runtime/D3D12RHI/Private/D3D12Resources.cpp b/Engine/Source/Runtime/D3D12RHI/Private/D3D12Resources.cpp
--- a/Engine/Source/Runtime/D3D12RHI/Private/D3D12Resources.cpp
+++ b/Engine/Source/Runtime/D3D12RHI/Private/D3D12Resources.cpp
@@ -244,12 +244,17 @@ HRESULT FD3D12Adapter::CreateCommittedResource(const FD3D12ResourceDesc& InDesc,
 	TRefCountPtr<ID3D12Resource> pResource;
 	const bool bRequiresInitialization = (InDesc.Flags & (D3D12_RESOURCE_FLAG_ALLOW_RENDER_TARGET | D3D12_RESOURCE_FLAG_ALLOW_DEPTH_STENCIL)) != 0;
 	D3D12_HEAP_FLAGS HeapFlags = (bHeapNotZeroedSupported && !bRequiresInitialization) ? FD3D12_HEAP_FLAG_CREATE_NOT_ZEROED : D3D12_HEAP_FLAG_NONE;
+	FD3D12ResourceDesc LocalDesc = InDesc;
 	if (InDesc.Flags & D3D12_RESOURCE_FLAG_ALLOW_SIMULTANEOUS_ACCESS)
 	{
 		HeapFlags |= D3D12_HEAP_FLAG_SHARED;
-	}
 
-	FD3D12ResourceDesc LocalDesc = InDesc;
+		// Simultaneous access flag is used to detect shared heap requirement but can't be used when allocating buffer resource
+		if (InDesc.Dimension == D3D12_RESOURCE_DIMENSION_BUFFER)
+		{
+			LocalDesc.Flags &= ~D3D12_RESOURCE_FLAG_ALLOW_SIMULTANEOUS_ACCESS;
+		}
+	}
 
 #if D3D12_RHI_RAYTRACING
 	if (InDefaultState == D3D12_RESOURCE_STATE_RAYTRACING_ACCELERATION_STRUCTURE)
@@ -736,6 +741,12 @@ void FD3D12ResourceLocation::UpdateStandAloneStats(bool bIncrement)
 		bool bIsBuffer = (Desc.Dimension == D3D12_RESOURCE_DIMENSION_BUFFER);
 		bool bIsRenderTarget = (Desc.Flags & D3D12_RESOURCE_FLAG_ALLOW_RENDER_TARGET || Desc.Flags & D3D12_RESOURCE_FLAG_ALLOW_DEPTH_STENCIL);
 		bool bIsUAV = (Desc.Flags & D3D12_RESOURCE_FLAG_ALLOW_UNORDERED_ACCESS) > 0;
+		
+		if (bIsBuffer)
+		{
+			// Simultaneous access flag is used to detect shared heap requirement but can't be used for buffers on device calls
+			Desc.Flags &= ~D3D12_RESOURCE_FLAG_ALLOW_SIMULTANEOUS_ACCESS;
+		}
 
 		// Get the desired size and allocated size for stand alone resources - allocated are very slow anyway
 		D3D12_RESOURCE_ALLOCATION_INFO Info = UnderlyingResource->GetParentDevice()->GetDevice()->GetResourceAllocationInfo(0, 1, &Desc);
diff --git a/Engine/Source/Runtime/D3D12RHI/Public/ID3D12DynamicRHI.h b/Engine/Source/Runtime/D3D12RHI/Public/ID3D12DynamicRHI.h
--- a/Engine/Source/Runtime/D3D12RHI/Public/ID3D12DynamicRHI.h
+++ b/Engine/Source/Runtime/D3D12RHI/Public/ID3D12DynamicRHI.h
@@ -36,6 +36,8 @@ struct ID3D12DynamicRHI : public FDynamicRHI
 
 	virtual ID3D12Resource*       RHIGetResource(FRHIBuffer* InBuffer) const = 0;
 	virtual uint32                RHIGetResourceDeviceIndex(FRHIBuffer* InBuffer) const = 0;
+	virtual int64                 RHIGetResourceMemorySize(FRHIBuffer* InBuffer) const = 0;
+	virtual bool                  RHIIsResourcePlaced(FRHIBuffer* InBuffer) const = 0;
 	virtual ID3D12Resource*       RHIGetResource(FRHITexture* InTexture) const = 0;
 	virtual uint32                RHIGetResourceDeviceIndex(FRHITexture* InTexture) const = 0;
 	virtual int64                 RHIGetResourceMemorySize(FRHITexture* InTexture) const = 0;
