Fix clang compile error

diff --git a/Engine/Source/Runtime/D3D12RHI/Private/D3D12RenderTarget.cpp b/Engine/Source/Runtime/D3D12RHI/Private/D3D12RenderTarget.cpp
--- a/Engine/Source/Runtime/D3D12RHI/Private/D3D12RenderTarget.cpp
+++ b/Engine/Source/Runtime/D3D12RHI/Private/D3D12RenderTarget.cpp
@@ -265,7 +265,7 @@ void FD3D12CommandContext::ResolveTextureUsingShader(
 
 			const uint32 TextureIndex = ResolvePixelShader->UnresolvedSurface.GetBaseIndex();
 
-			RHICmdList.RunOnContext([TextureIndex, SourceTexture](auto& Context)
+			RHICmdList.RunOnContext([TextureIndex, SourceTexture](FD3D12CommandContext& Context)
 			{
 				// Set the source texture.
 				Context.StateCache.SetShaderResourceView<SF_Pixel>(SourceTexture->GetShaderResourceView(), TextureIndex);
diff --git a/Engine/Source/Runtime/Windows/D3D11RHI/Private/D3D11RenderTarget.cpp b/Engine/Source/Runtime/Windows/D3D11RHI/Private/D3D11RenderTarget.cpp
--- a/Engine/Source/Runtime/Windows/D3D11RHI/Private/D3D11RenderTarget.cpp
+++ b/Engine/Source/Runtime/Windows/D3D11RHI/Private/D3D11RenderTarget.cpp
@@ -160,7 +160,7 @@ void FD3D11DynamicRHI::ResolveTextureUsingShader(
 
 		if (SourceTexture)
 		{
-			RHICmdList.RunOnContext([SourceTexture, TextureIndex](auto& Context)
+			RHICmdList.RunOnContext([SourceTexture, TextureIndex](FD3D11DynamicRHI& Context)
 			{
 				Context.SetShaderResourceView<SF_Pixel>(SourceTexture, SourceTexture->GetShaderResourceView(), TextureIndex);
 			});
