Adding log messages when D3D12CreateDevice fails.

diff --git a/Engine/Source/Runtime/D3D12RHI/Private/Windows/WindowsD3D12Device.cpp b/Engine/Source/Runtime/D3D12RHI/Private/Windows/WindowsD3D12Device.cpp
--- a/Engine/Source/Runtime/D3D12RHI/Private/Windows/WindowsD3D12Device.cpp
+++ b/Engine/Source/Runtime/D3D12RHI/Private/Windows/WindowsD3D12Device.cpp
@@ -481,7 +481,8 @@ static bool SafeTestD3D12CreateDevice(IDXGIAdapter* Adapter, D3D_FEATURE_LEVEL M
 	__try
 	{
 		ID3D12Device* Device = nullptr;
-		if (SUCCEEDED(D3D12CreateDevice(Adapter, MinFeatureLevel, IID_PPV_ARGS(&Device))))
+		const HRESULT D3D12CreateDeviceResult = D3D12CreateDevice(Adapter, MinFeatureLevel, IID_PPV_ARGS(&Device));
+		if (SUCCEEDED(D3D12CreateDeviceResult))
 		{
 			OutInfo.MaxFeatureLevel = FindHighestFeatureLevel(Device, MinFeatureLevel);
 			OutInfo.MaxShaderModel = FindHighestShaderModel(Device);
@@ -492,6 +493,10 @@ static bool SafeTestD3D12CreateDevice(IDXGIAdapter* Adapter, D3D_FEATURE_LEVEL M
 			Device->Release();
 			return true;
 		}
+		else
+		{
+			UE_LOG(LogD3D12RHI, Log, TEXT("D3D12CreateDevice failed with code 0x%08X"), static_cast<int32>(D3D12CreateDeviceResult));
+		}
 	}
 	__except (IsDelayLoadException(GetExceptionInformation()))
 	{
