Buffer locking must always defer the state cache clearing operation, regardless of the lock mode.

#rnx

diff --git a/Engine/Source/Runtime/D3D12RHI/Private/D3D12Buffer.cpp b/Engine/Source/Runtime/D3D12RHI/Private/D3D12Buffer.cpp
--- a/Engine/Source/Runtime/D3D12RHI/Private/D3D12Buffer.cpp
+++ b/Engine/Source/Runtime/D3D12RHI/Private/D3D12Buffer.cpp
@@ -62,16 +62,32 @@ struct FRHICommandRenameUploadBuffer final : public FRHICommand<FRHICommandRenam
 
 	FORCEINLINE_DEBUGGABLE FRHICommandRenameUploadBuffer(FD3D12Buffer* InResource, FD3D12Device* Device)
 		: Resource(InResource)
-		, NewLocation(Device) 
+		, NewLocation(Device)
 	{}
 
 	void Execute(FRHICommandListBase& CmdList)
 	{
-		// Clear the resource if still bound to make sure the SRVs are rebound again on next operation
-		FD3D12CommandContext& Context = (FD3D12CommandContext&)(CmdList.GetComputeContext().GetLowestLevelContext());
-		Context.ConditionalClearShaderResource(&Resource->ResourceLocation);
+		// Make sure we have an active pipeline when running at the top of the pipe, so the lambda below can obtain a context.
+		TOptional<ERHIPipeline> PreviousPipeline;
+		if (CmdList.IsTopOfPipe() && CmdList.GetPipeline() == ERHIPipeline::None)
+		{
+			PreviousPipeline = CmdList.SwitchPipeline(ERHIPipeline::Graphics);
+		}
+
+		// Clear the resource if still bound to make sure the SRVs are rebound again on next operation. This needs to happen
+		// on the RHI timeline when this command runs at the top of the pipe (which can happen when locking buffers in
+		// RLM_WriteOnly_NoOverwrite mode).
+		CmdList.EnqueueLambda([ResourceLocation = &Resource->ResourceLocation](FRHICommandListBase& RHICmdList) {
+			FD3D12CommandContext& Context = (FD3D12CommandContext&)(RHICmdList.GetComputeContext().GetLowestLevelContext());
+			Context.ConditionalClearShaderResource(ResourceLocation);
+		});
 
 		Resource->RenameLDAChain(NewLocation);
+
+		if (PreviousPipeline.IsSet())
+		{
+			CmdList.SwitchPipeline(PreviousPipeline.GetValue());
+		}
 	}
 };
 
