Fix crash in FD3D12Buffer::UploadResourceData
 - Parallel command lists may not have an active pipe.
 - Push/pop the graphics pipe so the upload command can obtain the correct command context.

diff --git a/Engine/Source/Runtime/D3D12RHI/Private/D3D12Buffer.cpp b/Engine/Source/Runtime/D3D12RHI/Private/D3D12Buffer.cpp
--- a/Engine/Source/Runtime/D3D12RHI/Private/D3D12Buffer.cpp
+++ b/Engine/Source/Runtime/D3D12RHI/Private/D3D12Buffer.cpp
@@ -240,7 +240,15 @@ void FD3D12Buffer::UploadResourceData(FRHICommandListBase& RHICmdList, FResource
 		}
 		else
 		{
+			// @todo d3d12 - Refactor contexts so pipe switching is not required.
+			TOptional<ERHIPipeline> PreviousPipeline;
+			if (RHICmdList.GetPipeline() == ERHIPipeline::None)
+				PreviousPipeline = RHICmdList.SwitchPipeline(ERHIPipeline::Graphics);
+
 			new (RHICmdList.AllocCommand<FD3D12RHICommandInitializeBuffer>()) FD3D12RHICommandInitializeBuffer(this, SrcResourceLoc, BufferSize, InDestinationState);
+
+			if (PreviousPipeline.IsSet())
+				RHICmdList.SwitchPipeline(PreviousPipeline.GetValue());
 		}
 	}
 
