Silence bogus static analyzer warning.

#rnx

diff --git a/Engine/Source/Runtime/D3D12RHI/Private/D3D12Submission.cpp b/Engine/Source/Runtime/D3D12RHI/Private/D3D12Submission.cpp
--- a/Engine/Source/Runtime/D3D12RHI/Private/D3D12Submission.cpp
+++ b/Engine/Source/Runtime/D3D12RHI/Private/D3D12Submission.cpp
@@ -885,9 +885,9 @@ FD3D12DynamicRHI::FProcessResult FD3D12DynamicRHI::ProcessInterruptQueue()
 							break;
 
 						case ED3D12QueryType::CommandListEnd:
-							check(ListBegin && !IdleBegin);
+							check(ListBegin != nullptr && IdleBegin == nullptr);
 							// Accumulate the number of ticks that have elapsed between the end of the previous command list, and the start of this one.
-							CurrentQueue.CumulativeIdleTicks += (CurrentQueue.LastEndTime != 0) ? ListBegin->GetResult() - CurrentQueue.LastEndTime : 0;
+							CurrentQueue.CumulativeIdleTicks += (CurrentQueue.LastEndTime != 0) ? ListBegin->GetResult() - CurrentQueue.LastEndTime : 0; //-V522
 							CurrentQueue.LastEndTime = Query.GetResult();
 							ListBegin = nullptr;
 							break;
@@ -898,9 +898,9 @@ FD3D12DynamicRHI::FProcessResult FD3D12DynamicRHI::ProcessInterruptQueue()
 							break;
 
 						case ED3D12QueryType::IdleEnd:
-							check(ListBegin && IdleBegin);
+							check(ListBegin != nullptr && IdleBegin != nullptr);
 							// Accumulate the time this pipe spent in an idle scope. This includes vsync and waiting on other pipes.
-							CurrentQueue.CumulativeIdleTicks += Query.GetResult() - IdleBegin->GetResult();
+							CurrentQueue.CumulativeIdleTicks += Query.GetResult() - IdleBegin->GetResult(); //-V522
 							IdleBegin = nullptr;
 							break;
 
