Fixing issues caused by FD3D12DescriptorCache recreating views for every command context. Moving the "null"/"default" view creation into FD3D12Device so that we only create these views once.

diff --git a/Engine/Source/Runtime/D3D12RHI/Private/D3D12DescriptorCache.cpp b/Engine/Source/Runtime/D3D12RHI/Private/D3D12DescriptorCache.cpp
--- a/Engine/Source/Runtime/D3D12RHI/Private/D3D12DescriptorCache.cpp
+++ b/Engine/Source/Runtime/D3D12RHI/Private/D3D12DescriptorCache.cpp
@@ -57,6 +57,7 @@ FD3D12DescriptorCache::FD3D12DescriptorCache(FD3D12CommandContext& Context, FRHI
 	: FD3D12DeviceChild(Context.Device)
 	, FD3D12SingleNodeGPUObject(Node)
 	, Context(Context)
+	, DefaultViews(Context.Device->GetDefaultViews())
 	, LocalSamplerHeap(*this, Context)
 	, SubAllocatedViewHeap(*this, Context)
 	, SamplerMap(271) // Prime numbers for better hashing
@@ -85,64 +86,6 @@ void FD3D12DescriptorCache::Init(uint32 InNumLocalViewDescriptors, uint32 InNumS
 
 	OverrideViewHeap = nullptr;
 	OverrideSamplerHeap = nullptr;
-
-	// Create default views
-	D3D12_SHADER_RESOURCE_VIEW_DESC SRVDesc = {};
-	SRVDesc.Shader4ComponentMapping = D3D12_DEFAULT_SHADER_4_COMPONENT_MAPPING;
-	SRVDesc.ViewDimension = D3D12_SRV_DIMENSION_TEXTURE2D;
-	SRVDesc.Format = DXGI_FORMAT_B8G8R8A8_UNORM;
-	SRVDesc.Texture2D.MipLevels = 1;
-	SRVDesc.Texture2D.MostDetailedMip = 0;
-	SRVDesc.Texture2D.ResourceMinLODClamp = 0.0f;
-	NullSRV = new FD3D12ViewDescriptorHandle(GetParentDevice(), ERHIDescriptorHeapType::Standard);
-	NullSRV->CreateView(SRVDesc, nullptr, ED3D12DescriptorCreateReason::InitialCreate);
-
-	D3D12_RENDER_TARGET_VIEW_DESC RTVDesc = {};
-	RTVDesc.ViewDimension = D3D12_RTV_DIMENSION_TEXTURE2D;
-	RTVDesc.Format = DXGI_FORMAT_B8G8R8A8_UNORM;
-	RTVDesc.Texture2D.MipSlice = 0;
-	NullRTV = new FD3D12ViewDescriptorHandle(GetParentDevice(), ERHIDescriptorHeapType::RenderTarget);
-	NullRTV->CreateView(RTVDesc, nullptr);
-
-	D3D12_UNORDERED_ACCESS_VIEW_DESC UAVDesc = {};
-	UAVDesc.ViewDimension = D3D12_UAV_DIMENSION_TEXTURE2D;
-	UAVDesc.Format = DXGI_FORMAT_R8G8B8A8_UNORM;
-	UAVDesc.Texture2D.MipSlice = 0;
-	NullUAV = new FD3D12ViewDescriptorHandle(GetParentDevice(), ERHIDescriptorHeapType::Standard);
-	NullUAV->CreateView(UAVDesc, nullptr, nullptr, ED3D12DescriptorCreateReason::InitialCreate);
-
-#if USE_STATIC_ROOT_SIGNATURE
-	NullCBV = new FD3D12ConstantBufferView(GetParentDevice());
-#endif
-
-	const FSamplerStateInitializerRHI SamplerDesc(
-		SF_Trilinear,
-		AM_Clamp,
-		AM_Clamp,
-		AM_Clamp,
-		0,
-		0,
-		0,
-		FLT_MAX
-		);
-
-	FSamplerStateRHIRef Sampler = Parent->CreateSampler(SamplerDesc);
-
-	DefaultSampler = static_cast<FD3D12SamplerState*>(Sampler.GetReference());
-
-	// The default sampler must have ID=0
-	// DescriptorCache::SetSamplers relies on this
-	check(DefaultSampler->ID == 0);
-}
-
-void FD3D12DescriptorCache::Clear()
-{
-	delete NullSRV; NullSRV = nullptr;
-	delete NullUAV; NullUAV = nullptr;
-	delete NullRTV; NullRTV = nullptr;
-#if USE_STATIC_ROOT_SIGNATURE
-	delete NullCBV; NullCBV = nullptr;
-#endif
 }
 
 bool FD3D12DescriptorCache::SetDescriptorHeaps()
@@ -279,7 +222,7 @@ void FD3D12DescriptorCache::SetUAVs(const FD3D12RootSignature* RootSignature, FD
 	{
 		if ((SlotIndex < UAVStartSlot) || (UAVs[SlotIndex] == nullptr))
 		{
-			SrcDescriptors[SlotIndex] = NullUAV->GetOfflineCpuHandle();
+			SrcDescriptors[SlotIndex] = DefaultViews.NullUAV->GetOfflineCpuHandle();
 		}
 		else
 		{
@@ -360,7 +303,7 @@ void FD3D12DescriptorCache::SetRenderTargets(FD3D12RenderTargetView** RenderTarg
 		}
 		else
 		{
-			RTVDescriptors[i] = NullRTV->GetOfflineCpuHandle();
+			RTVDescriptors[i] = DefaultViews.NullRTV->GetOfflineCpuHandle();
 		}
 	}
 
@@ -449,7 +392,7 @@ void FD3D12DescriptorCache::SetSamplers(const FD3D12RootSignature* RootSignature
 			}
 			else
 			{
-				SrcDescriptors[SlotIndex] = DefaultSampler->OfflineHandle;
+				SrcDescriptors[SlotIndex] = DefaultViews.DefaultSampler->OfflineHandle;
 			}
 		}
 		FD3D12SamplerStateCache::CleanSlots(CurrentDirtySlotMask, SlotsNeeded);
@@ -535,7 +478,7 @@ void FD3D12DescriptorCache::SetSRVs(const FD3D12RootSignature* RootSignature, FD
 		}
 		else
 		{
-			SrcDescriptors[SlotIndex] = NullSRV->GetOfflineCpuHandle();
+			SrcDescriptors[SlotIndex] = DefaultViews.NullSRV->GetOfflineCpuHandle();
 		}
 		check(SrcDescriptors[SlotIndex].ptr != 0);
 	}
@@ -614,7 +557,7 @@ void FD3D12DescriptorCache::SetConstantBuffers(const FD3D12RootSignature* RootSi
 		}
 		else
 		{
-			Device->CopyDescriptorsSimple(1, DestDescriptor, NullCBV->GetOfflineCpuHandle(), D3D12_DESCRIPTOR_HEAP_TYPE_CBV_SRV_UAV);
+			Device->CopyDescriptorsSimple(1, DestDescriptor, DefaultViews.NullCBV->GetOfflineCpuHandle(), D3D12_DESCRIPTOR_HEAP_TYPE_CBV_SRV_UAV);
 		}
 
 		DestDescriptorSlot++;
diff --git a/Engine/Source/Runtime/D3D12RHI/Private/D3D12DescriptorCache.h b/Engine/Source/Runtime/D3D12RHI/Private/D3D12DescriptorCache.h
--- a/Engine/Source/Runtime/D3D12RHI/Private/D3D12DescriptorCache.h
+++ b/Engine/Source/Runtime/D3D12RHI/Private/D3D12DescriptorCache.h
@@ -3,6 +3,7 @@
 #pragma once
 
 class FD3D12DynamicRHI;
+struct FD3D12DefaultViews;
 class FD3D12DescriptorCache;
 struct FD3D12VertexBufferCache;
 struct FD3D12IndexBufferCache;
@@ -359,7 +360,6 @@ public:
 	bool HeapRolledOver(ERHIDescriptorHeapType InHeapType);
 	void HeapLoopedAround(ERHIDescriptorHeapType InHeapType);
 	void Init(uint32 InNumLocalViewDescriptors, uint32 InNumSamplerDescriptors);
-	void Clear();
 
 	bool SwitchToContextLocalViewHeap();
 	bool SwitchToContextLocalSamplerHeap();
@@ -377,18 +377,9 @@ public:
 	// Returns true if one of the heaps actually changed, false otherwise.
 	bool SetDescriptorHeaps();
 
-	// null views
-	FD3D12ViewDescriptorHandle* NullSRV = nullptr;
-	FD3D12ViewDescriptorHandle* NullRTV = nullptr;
-	FD3D12ViewDescriptorHandle* NullUAV = nullptr;
-
-#if USE_STATIC_ROOT_SIGNATURE
-	FD3D12ConstantBufferView* NullCBV = nullptr;
-#endif
-	TRefCountPtr<FD3D12SamplerState> DefaultSampler;
-
 protected:
 	FD3D12CommandContext& Context;
+	const FD3D12DefaultViews& DefaultViews;
 
 private:
 	// The previous view and sampler heaps set on the current command list.
diff --git a/Engine/Source/Runtime/D3D12RHI/Private/D3D12Device.cpp b/Engine/Source/Runtime/D3D12RHI/Private/D3D12Device.cpp
--- a/Engine/Source/Runtime/D3D12RHI/Private/D3D12Device.cpp
+++ b/Engine/Source/Runtime/D3D12RHI/Private/D3D12Device.cpp
@@ -11,6 +11,17 @@ static TAutoConsoleVariable<int32> CVarD3D12GPUTimeout(
 	ECVF_ReadOnly
 );
 
+static uint32 GetQueryHeapPoolIndex(D3D12_QUERY_HEAP_TYPE HeapType)
+{
+	switch (HeapType)
+	{
+	default: checkNoEntry(); [[fallthrough]];
+	case D3D12_QUERY_HEAP_TYPE_OCCLUSION:            return 0;
+	case D3D12_QUERY_HEAP_TYPE_TIMESTAMP:            return 1;
+	case D3D12_QUERY_HEAP_TYPE_COPY_QUEUE_TIMESTAMP: return 2;
+	}
+}
+
 FD3D12Queue::FD3D12Queue(FD3D12Device* Device, ED3D12QueueType QueueType)
 	: Device(Device)
 	, QueueType(QueueType)
@@ -362,6 +373,9 @@ void FD3D12Device::SetupAfterDeviceCreation()
 
 	OnlineDescriptorManager.Init(GOnlineDescriptorHeapSize, GOnlineDescriptorHeapBlockSize);
 
+	// Make sure we create the default views before the first command context
+	CreateDefaultViews();
+
 	// Needs to be called before creating command contexts
 	UpdateConstantBufferPageProperties();
 
@@ -383,6 +397,55 @@ void FD3D12Device::SetupAfterDeviceCreation()
 	ImmediateCommandContext = FD3D12DynamicRHI::GetD3DRHI()->CreateCommandContext(this, ED3D12QueueType::Direct, true);
 }
 
+void FD3D12Device::CreateDefaultViews()
+{
+	{
+		D3D12_SHADER_RESOURCE_VIEW_DESC SRVDesc{};
+		SRVDesc.Format = DXGI_FORMAT_B8G8R8A8_UNORM;
+		SRVDesc.ViewDimension = D3D12_SRV_DIMENSION_TEXTURE2D;
+		SRVDesc.Shader4ComponentMapping = D3D12_DEFAULT_SHADER_4_COMPONENT_MAPPING;
+		SRVDesc.Texture2D.MipLevels = 1;
+		SRVDesc.Texture2D.MostDetailedMip = 0;
+		SRVDesc.Texture2D.ResourceMinLODClamp = 0.0f;
+
+		DefaultViews.NullSRV = new FD3D12ViewDescriptorHandle(this, ERHIDescriptorHeapType::Standard);
+		DefaultViews.NullSRV->CreateView(SRVDesc, nullptr, ED3D12DescriptorCreateReason::InitialCreate);
+	}
+
+	{
+		D3D12_RENDER_TARGET_VIEW_DESC RTVDesc{};
+		RTVDesc.Format = DXGI_FORMAT_B8G8R8A8_UNORM;
+		RTVDesc.ViewDimension = D3D12_RTV_DIMENSION_TEXTURE2D;
+		RTVDesc.Texture2D.MipSlice = 0;
+
+		DefaultViews.NullRTV = new FD3D12ViewDescriptorHandle(this, ERHIDescriptorHeapType::RenderTarget);
+		DefaultViews.NullRTV->CreateView(RTVDesc, nullptr);
+	}
+
+	{
+		D3D12_UNORDERED_ACCESS_VIEW_DESC UAVDesc{};
+		UAVDesc.Format = DXGI_FORMAT_R8G8B8A8_UNORM;
+		UAVDesc.ViewDimension = D3D12_UAV_DIMENSION_TEXTURE2D;
+		UAVDesc.Texture2D.MipSlice = 0;
+
+		DefaultViews.NullUAV = new FD3D12ViewDescriptorHandle(this, ERHIDescriptorHeapType::Standard);
+		DefaultViews.NullUAV->CreateView(UAVDesc, nullptr, nullptr, ED3D12DescriptorCreateReason::InitialCreate);
+	}
+
+#if USE_STATIC_ROOT_SIGNATURE
+	DefaultViews.NullCBV = new FD3D12ConstantBufferView(this);
+#endif
+
+	{
+		const FSamplerStateInitializerRHI SamplerDesc(SF_Trilinear, AM_Clamp, AM_Clamp, AM_Clamp);
+		DefaultViews.DefaultSampler = CreateSampler(SamplerDesc);
+
+		// The default sampler must have ID=0
+		// FD3D12DescriptorCache::SetSamplers relies on this
+		check(DefaultViews.DefaultSampler->ID == 0);
+	}
+}
+
 void FD3D12Device::UpdateConstantBufferPageProperties()
 {
 	//In genera, constant buffers should use write-combine memory (i.e. upload heaps) for optimal performance
diff --git a/Engine/Source/Runtime/D3D12RHI/Private/D3D12Device.h b/Engine/Source/Runtime/D3D12RHI/Private/D3D12Device.h
--- a/Engine/Source/Runtime/D3D12RHI/Private/D3D12Device.h
+++ b/Engine/Source/Runtime/D3D12RHI/Private/D3D12Device.h
@@ -242,6 +242,19 @@ public:
 	void SetupAfterDeviceCreation();
 };
 
+struct FD3D12DefaultViews
+{
+	FD3D12ViewDescriptorHandle* NullSRV = nullptr;
+	FD3D12ViewDescriptorHandle* NullRTV = nullptr;
+	FD3D12ViewDescriptorHandle* NullUAV = nullptr;
+
+#if USE_STATIC_ROOT_SIGNATURE
+	FD3D12ConstantBufferView* NullCBV = nullptr;
+#endif
+
+	TRefCountPtr<FD3D12SamplerState> DefaultSampler;
+};
+
 class FD3D12Device final : public FD3D12SingleNodeGPUObject, public FNoncopyable, public FD3D12AdapterChild
 {
 public:
@@ -299,6 +312,8 @@ public:
 		return OfflineDescriptorManagers[static_cast<int>(InType)];
 	}
 
+	const FD3D12DefaultViews& GetDefaultViews() const { return DefaultViews; }
+
 	// Memory Allocators
 	inline FD3D12DefaultBufferAllocator& GetDefaultBufferAllocator() { return DefaultBufferAllocator; }
 	inline FD3D12FastAllocator&          GetDefaultFastAllocator  () { return DefaultFastAllocator;   }
@@ -341,6 +356,7 @@ public:
 
 private:
 	// called by SetupAfterDeviceCreation() when the device gets initialized
+	void CreateDefaultViews();
 	void UpdateMSAASettings();
 	void UpdateConstantBufferPageProperties();
 
@@ -359,17 +375,9 @@ private:
 	FD3D12GlobalOnlineSamplerHeap GlobalSamplerHeap;
 	FD3D12OnlineDescriptorManager OnlineDescriptorManager;
 
+	FD3D12DefaultViews DefaultViews;
+
 	TD3D12ObjectPool<FD3D12QueryHeap> QueryHeapPool[3];
-	static uint32 GetQueryHeapPoolIndex(D3D12_QUERY_HEAP_TYPE HeapType)
-	{
-		switch (HeapType)
-		{
-		default: checkNoEntry(); [[fallthrough]];
-		case D3D12_QUERY_HEAP_TYPE_OCCLUSION:            return 0;
-		case D3D12_QUERY_HEAP_TYPE_TIMESTAMP:            return 1;
-		case D3D12_QUERY_HEAP_TYPE_COPY_QUEUE_TIMESTAMP: return 2;
-		}
-	}
 	
 	FD3D12CommandContext* ImmediateCommandContext = nullptr;
 
