Add bindless resource support as a requirement for ray tracing on PC (SM 6.6 + Resource Binding Tier 3).

Compile DXR shaders using the same profile when targeting PCD3D_SM5 and PCD3D_SM6 (i.e. lib_6_6 or lib_6_5 based on USE_SHADER_MODEL_6_6 define).

diff --git a/Engine/Shaders/Shared/RayTracingDefinitions.h b/Engine/Shaders/Shared/RayTracingDefinitions.h
--- a/Engine/Shaders/Shared/RayTracingDefinitions.h
+++ b/Engine/Shaders/Shared/RayTracingDefinitions.h
@@ -9,7 +9,7 @@
 
 // Change this to force recompilation of all ray tracing shaders (use https://www.random.org/cgi-bin/randbyte?nbytes=4&format=h)
 // This avoids changing the global ShaderVersion.ush and forcing recompilation of all shaders in the engine (only RT shaders will be affected)
-#define RAY_TRACING_SHADER_VERSION 0x11241B50
+#define RAY_TRACING_SHADER_VERSION 0x4502E713
 
 #define RAY_TRACING_REGISTER_SPACE_LOCAL  0 // default register space for hit group (closest hit, any hit, intersection) shader resources
 #define RAY_TRACING_REGISTER_SPACE_GLOBAL 1 // register space for ray generation and miss shaders 
diff --git a/Engine/Source/Developer/Windows/ShaderFormatD3D/Private/D3DShaderCompiler.cpp b/Engine/Source/Developer/Windows/ShaderFormatD3D/Private/D3DShaderCompiler.cpp
--- a/Engine/Source/Developer/Windows/ShaderFormatD3D/Private/D3DShaderCompiler.cpp
+++ b/Engine/Source/Developer/Windows/ShaderFormatD3D/Private/D3DShaderCompiler.cpp
@@ -85,7 +85,7 @@ static void D3D11FilterShaderCompileWarnings(const FString& CompileWarnings, TAr
 static const TCHAR* GetShaderProfileName(ELanguage Language, uint32 Frequency, bool bForceSM6)
 {
 
-	if (Language == ELanguage::SM6)
+	if (Language == ELanguage::SM6 || IsRayTracingShaderFrequency(EShaderFrequency(Frequency)))
 	{
 		switch (Frequency)
 		{
@@ -131,6 +131,7 @@ static const TCHAR* GetShaderProfileName(ELanguage Language, uint32 Frequency, b
 		case SF_RayMiss:
 		case SF_RayHitGroup:
 		case SF_RayCallable:
+			checkNoEntry();
 			return TEXT("lib_6_5");
 		}
 	}
diff --git a/Engine/Source/Developer/Windows/ShaderFormatD3D/Private/ShaderFormatD3D.cpp b/Engine/Source/Developer/Windows/ShaderFormatD3D/Private/ShaderFormatD3D.cpp
--- a/Engine/Source/Developer/Windows/ShaderFormatD3D/Private/ShaderFormatD3D.cpp
+++ b/Engine/Source/Developer/Windows/ShaderFormatD3D/Private/ShaderFormatD3D.cpp
@@ -116,7 +116,7 @@ public:
 	void ModifyShaderCompilerInput(FShaderCompilerInput& Input) const final
 	{
 		CheckFormat(Input.ShaderFormat);
-		if (Input.ShaderFormat == NAME_PCD3D_SM6)
+		if (Input.ShaderFormat == NAME_PCD3D_SM6 || Input.IsRayTracingShader())
 		{
 			Input.Environment.SetDefine(TEXT("SM6_PROFILE"), 1);
 			Input.Environment.SetDefine(TEXT("COMPILER_DXC"), 1);
@@ -135,14 +135,13 @@ public:
 		{
 			Input.Environment.SetDefine(TEXT("SM5_PROFILE"), 1);
 			const bool bUseDXC =
-				Input.IsRayTracingShader()
-				|| Input.Environment.CompilerFlags.Contains(CFLAG_WaveOperations)
+				Input.Environment.CompilerFlags.Contains(CFLAG_WaveOperations)
 				|| Input.Environment.CompilerFlags.Contains(CFLAG_ForceDXC);
 			Input.Environment.SetDefine(TEXT("COMPILER_DXC"), bUseDXC);
 
 			if (bUseDXC)
 			{
-				AddShaderTargetDefines(Input, 6, Input.IsRayTracingShader() ? 5 : 0);
+				AddShaderTargetDefines(Input, 6, 0);
 			}
 			else
 			{
diff --git a/Engine/Source/Runtime/D3D12RHI/Private/D3D12Adapter.cpp b/Engine/Source/Runtime/D3D12RHI/Private/D3D12Adapter.cpp
--- a/Engine/Source/Runtime/D3D12RHI/Private/D3D12Adapter.cpp
+++ b/Engine/Source/Runtime/D3D12RHI/Private/D3D12Adapter.cpp
@@ -969,11 +969,12 @@ void FD3D12Adapter::InitializeDevices()
 					&& !FParse::Param(FCommandLine::Get(), TEXT("noraytracing")))
 				{
 					if (D3D12Caps5.RaytracingTier >= D3D12_RAYTRACING_TIER_1_1
+						&& GRHISupportsBindless
 						&& RootDevice7)
 					{
 						if (bRayTracingAllowedOnCurrentShaderPlatform)
 						{
-							UE_LOG(LogD3D12RHI, Log, TEXT("D3D12 ray tracing tier 1.1 is supported."));
+							UE_LOG(LogD3D12RHI, Log, TEXT("D3D12 ray tracing tier 1.1 and bindless resources are supported."));
 
 							GRHISupportsRayTracing = RHISupportsRayTracing(GMaxRHIShaderPlatform);
 							GRHISupportsRayTracingShaders = GRHISupportsRayTracing && RHISupportsRayTracingShaders(GMaxRHIShaderPlatform);
@@ -986,6 +987,10 @@ void FD3D12Adapter::InitializeDevices()
  							UE_LOG(LogD3D12RHI, Log, TEXT("Ray tracing is disabled because SM6 shader platform is required (r.RayTracing.RequireSM6=1)."));
 						}
 					}
+					else if (!GRHISupportsBindless)
+					{
+						UE_LOG(LogD3D12RHI, Log, TEXT("Ray tracing is disabled because bindless resources are not supported (Shader Model 6.6 and Resource Binding Tier 3 are required)."));
+					}
 					else
 					{
 						UE_LOG(LogD3D12RHI, Log, TEXT("Ray tracing is disabled because D3D12 ray tracing tier 1.1 is required but only tier 1.0 is supported."));
