MGPU:  Support for specifying GPU mask for RDG textures, to facilitate a hack fix for GPU Lightmass MGPU transition validation assert.

The assert is due to an architectural bug in RDG where resource transitions for a given resource occur at an arbitrary time in the graph execution, due to batching.  The GPU mask can vary, based on which Pass is executing when a transition for a given resource occurs.  This leads to the internal state used for validation being inconsistent on one of the GPUs.  The hack fix is to make it so the offending resource (PathTracer.LightGrid) only exists on one GPU.  This works because the inconsistent transition state happens on a GPU that the resource isn't actually used on, so removing it from that GPU prevents the assert.  The RDG transition bug doesn't come up often, because most graphs run all-GPU work first (GPU scene update), then run to termination on a single mask, but GPU Lightmass switches back and forth.

The longer term plan (for 5.2) is to scope ALL resources so they only exist on the GPUs they are needed on, then make it so transitions ignore the active GPU mask, and instead apply transitions for all GPUs where a given resource is present.  So there will never be a mismatch in transition state across GPUs.  Scoping resources to specific GPUs saves GPU memory and the CPU performance cost of maintaining the resources, so it would be worth doing on its own.  And it provides other advantages in terms of improved validation and potential automation of cross GPU transfers through RDG.  But this was too big a refactor to do for 5.1, especially with Parallel Rendering work in flight, so I went with the hack fix to avoid the crash.

#rnx

diff --git a/Engine/Source/Runtime/D3D12RHI/Private/D3D12Commands.cpp b/Engine/Source/Runtime/D3D12RHI/Private/D3D12Commands.cpp
--- a/Engine/Source/Runtime/D3D12RHI/Private/D3D12Commands.cpp
+++ b/Engine/Source/Runtime/D3D12RHI/Private/D3D12Commands.cpp
@@ -240,23 +240,31 @@ void ProcessResource(FD3D12CommandContext& Context, const FRHITransitionInfo& In
 	}
 	case FRHITransitionInfo::EType::Buffer:
 	{
+		// Resource may be null if this is a multi-GPU resource not present on the current GPU
 		FD3D12Buffer* Buffer = Context.RetrieveObject<FD3D12Buffer>(Info.Buffer);
-		check(Buffer);
-		Function(Info, Buffer->GetResource());
+		check(Buffer || GNumExplicitGPUsForRendering > 1);
+		if (Buffer)
+		{
+			Function(Info, Buffer->GetResource());
+		}
 		break;
 	}
 	case FRHITransitionInfo::EType::Texture:
 	{
+		// Resource may be null if this is a multi-GPU resource not present on the current GPU
 		FD3D12Texture* Texture = Context.RetrieveTexture(Info.Texture);
-		check(Texture);
-		FD3D12Texture* TextureOut = nullptr;
-#if PLATFORM_REQUIRES_TYPELESS_RESOURCE_DISCARD_WORKAROUND
-		if (Texture->GetRequiresTypelessResourceDiscardWorkaround())
+		check(Texture || GNumExplicitGPUsForRendering > 1);
+		if (Texture)
 		{
-			TextureOut = Texture;
-		}
+			FD3D12Texture* TextureOut = nullptr;
+#if PLATFORM_REQUIRES_TYPELESS_RESOURCE_DISCARD_WORKAROUND
+			if (Texture->GetRequiresTypelessResourceDiscardWorkaround())
+			{
+				TextureOut = Texture;
+			}
 #endif // #if PLATFORM_REQUIRES_TYPELESS_RESOURCE_DISCARD_WORKAROUND
-		Function(Info, Texture->GetResource(), TextureOut);
+			Function(Info, Texture->GetResource(), TextureOut);
+		}
 		break;
 	}
 	default:
@@ -411,15 +419,17 @@ void FD3D12CommandContext::HandleTransientAliasing(const FD3D12TransitionData* T
 		{
 		case FRHITransientAliasingInfo::EType::Buffer:
 		{
+			// Resource may be null if this is a multi-GPU resource not present on the current GPU
 			FD3D12Buffer* Buffer = RetrieveObject<FD3D12Buffer>(Info.Buffer);
-			check(Buffer);
+			check(Buffer || GNumExplicitGPUsForRendering > 1);
 			BaseShaderResource = Buffer;
 			break;
 		}
 		case FRHITransientAliasingInfo::EType::Texture:
 		{
+			// Resource may be null if this is a multi-GPU resource not present on the current GPU
 			FD3D12Texture* Texture = RetrieveTexture(Info.Texture);
-			check(Texture);
+			check(Texture || GNumExplicitGPUsForRendering > 1);
 			BaseShaderResource = Texture;
 			break;
 		}
@@ -428,6 +438,12 @@ void FD3D12CommandContext::HandleTransientAliasing(const FD3D12TransitionData* T
 			break;
 		}
 
+		// Resource may be null if this is a multi-GPU resource not present on the current GPU
+		if (!BaseShaderResource)
+		{
+			continue;
+		}
+
 		FD3D12Resource* Resource = BaseShaderResource->ResourceLocation.GetResource();
 		if (Info.Action == FRHITransientAliasingInfo::EAction::Acquire)
 		{
@@ -446,15 +462,31 @@ void FD3D12CommandContext::HandleTransientAliasing(const FD3D12TransitionData* T
 					switch (Overlap.Type)
 					{
 					case FRHITransientAliasingOverlap::EType::Texture:
-						ResourceBefore = RetrieveTexture(Overlap.Texture)->GetResource();
+						{
+							const FD3D12Texture* Texture = RetrieveTexture(Overlap.Texture);
+							if (Texture)
+							{
+								ResourceBefore = Texture->GetResource();
+							}
+						}
 						break;
 					case FRHITransientAliasingOverlap::EType::Buffer:
-						ResourceBefore = RetrieveObject<FD3D12Buffer>(Overlap.Buffer)->GetResource();
+						{
+							const FD3D12Buffer* Buffer = RetrieveObject<FD3D12Buffer>(Overlap.Buffer);
+							if (Buffer)
+							{
+								ResourceBefore = Buffer->GetResource();
+							}
+						}
 						break;
 					}
 
-					check(ResourceBefore);
-					AddAliasingBarrier(ResourceBefore->GetResource(), Resource->GetResource());
+					// Resource may be null if this is a multi-GPU resource not present on the current GPU
+					check(ResourceBefore || GNumExplicitGPUsForRendering > 1);
+					if (ResourceBefore)
+					{
+						AddAliasingBarrier(ResourceBefore->GetResource(), Resource->GetResource());
+					}
 				}
 			}
 		}
diff --git a/Engine/Source/Runtime/RHI/Public/RHIResources.h b/Engine/Source/Runtime/RHI/Public/RHIResources.h
--- a/Engine/Source/Runtime/RHI/Public/RHIResources.h
+++ b/Engine/Source/Runtime/RHI/Public/RHIResources.h
@@ -1499,6 +1499,9 @@ struct RHI_API FRHITextureDesc
 	/** Resource memory percentage which should be allocated onto fast VRAM (hint-only). (encoding into 8bits, 0..255 -> 0%..100%) */
 	uint8 FastVRAMPercentage = 0xFF;
 
+	/* A mask representing which GPUs to create the resource on, in a multi-GPU system. */
+	FRHIGPUMask GPUMask = FRHIGPUMask::All();
+
 	/** Check the validity. */
 	static bool CheckValidity(const FRHITextureDesc& Desc, const TCHAR* Name)
 	{
@@ -1638,12 +1641,10 @@ struct FRHITextureCreateDesc : public FRHITextureDesc
 		, ERHIAccess                  InInitialState
 		, TCHAR const*                InDebugName
 		, FResourceBulkDataInterface* InBulkData     = nullptr
-		, FRHIGPUMask                 InGPUMask      = FRHIGPUMask::All()
 		)
 		: FRHITextureDesc(InDesc)
 		, InitialState   (InInitialState)
 		, DebugName      (InDebugName)
-		, GPUMask        (InGPUMask)
 		, BulkData       (InBulkData)
 	{}
 
@@ -1681,9 +1682,6 @@ struct FRHITextureCreateDesc : public FRHITextureDesc
 	/* A friendly name for the resource. */
 	const TCHAR* DebugName = nullptr;
 
-	/* A mask representing which GPUs to create the resource on, in a multi-GPU system. */
-	FRHIGPUMask GPUMask = FRHIGPUMask::All();
-
 	/* Optional initial data to fill the resource with. */
 	FResourceBulkDataInterface* BulkData = nullptr;
 };
diff --git a/Engine/Source/Runtime/Renderer/Private/PathTracing.cpp b/Engine/Source/Runtime/Renderer/Private/PathTracing.cpp
--- a/Engine/Source/Runtime/Renderer/Private/PathTracing.cpp
+++ b/Engine/Source/Runtime/Renderer/Private/PathTracing.cpp
@@ -1550,6 +1550,15 @@ RENDERER_API void PrepareLightGrid(FRDGBuilder& GraphBuilder, ERHIFeatureLevel::
 			PF_R32_UINT,
 			FClearValueBinding::None,
 			TexCreate_ShaderResource | TexCreate_UAV);
+
+		// jhoerner TODO 9/30/2022: Hack to work around MGPU resource transition architectural bug in RDG.  Mask PathTracer.LightGrid texture
+		// to only be present on current GPU.  The bug is that RDG batches transitions, but the execution of batched transitions uses the
+		// GPU Mask of the current Pass that's executing, not the GPU Mask that's relevant to the Passes where a given resource is used.  This
+		// causes an assert due to a mismatch in the expected transition state on a specific GPU, when an intermediate transition was skipped
+		// on that GPU, due to the arbitrary nature of the GPU mask when a transition batch is flushed.  The hack works by removing the
+		// resource from GPUs it's not actually used on, where the intermediate transition gets skipped.
+		LightGridDesc.GPUMask = GraphBuilder.RHICmdList.GetGPUMask();
+
 		FRDGTexture* LightGridTexture = GraphBuilder.CreateTexture(LightGridDesc, TEXT("PathTracer.LightGrid"), ERDGTextureFlags::None);
 		LightGridPassParameters->RWLightGrid = GraphBuilder.CreateUAV(LightGridTexture);
 
