[Insights] Added memory tracing for D3D11 and D3D12 (using the "video memory" root heap / address space) to match LLM stats.

diff --git a/Engine/Source/Runtime/D3D12RHI/Private/D3D12Texture.cpp b/Engine/Source/Runtime/D3D12RHI/Private/D3D12Texture.cpp
--- a/Engine/Source/Runtime/D3D12RHI/Private/D3D12Texture.cpp
+++ b/Engine/Source/Runtime/D3D12RHI/Private/D3D12Texture.cpp
@@ -6,6 +6,7 @@
 
 #include "D3D12RHIPrivate.h"
 #include "TextureProfiler.h"
+#include "ProfilingDebugging/MemoryTrace.h"
 
 int64 FD3D12GlobalStats::GDedicatedVideoMemory = 0;
 int64 FD3D12GlobalStats::GDedicatedSystemMemory = 0;
@@ -390,6 +391,10 @@ void FD3D12TextureStats::UpdateD3D12TextureStats(FD3D12Texture& Texture, const D
 		// This means that we must manually add the tracking here.
 		LLM(FLowLevelMemTracker::Get().OnLowLevelAlloc(ELLMTracker::Platform, Texture.GetResource(), TextureSize, ELLMTag::GraphicsPlatform));
 		LLM(FLowLevelMemTracker::Get().OnLowLevelAlloc(ELLMTracker::Default, Texture.GetResource(), TextureSize, ELLMTag::Textures));
+		{
+			UE_MEMSCOPE_DEFAULT(ELLMTag::Textures);
+			MemoryTrace_Alloc((uint64)Texture.GetResource(), TextureSize, 1024, EMemoryTraceRootHeap::VideoMemory);
+		}
 #endif
 		INC_DWORD_STAT(STAT_D3D12TexturesAllocated);
 	}
@@ -398,6 +403,7 @@ void FD3D12TextureStats::UpdateD3D12TextureStats(FD3D12Texture& Texture, const D
 #if PLATFORM_WINDOWS
 		LLM(FLowLevelMemTracker::Get().OnLowLevelFree(ELLMTracker::Platform, Texture.GetResource()));
 		LLM(FLowLevelMemTracker::Get().OnLowLevelFree(ELLMTracker::Default, Texture.GetResource()));
+		MemoryTrace_Free((uint64)Texture.GetResource(), EMemoryTraceRootHeap::VideoMemory);
 #endif
 		INC_DWORD_STAT(STAT_D3D12TexturesReleased);
 	}
diff --git a/Engine/Source/Runtime/D3D12RHI/Private/D3D12TransientResourceAllocator.cpp b/Engine/Source/Runtime/D3D12RHI/Private/D3D12TransientResourceAllocator.cpp
--- a/Engine/Source/Runtime/D3D12RHI/Private/D3D12TransientResourceAllocator.cpp
+++ b/Engine/Source/Runtime/D3D12RHI/Private/D3D12TransientResourceAllocator.cpp
@@ -3,6 +3,7 @@
 #include "D3D12RHIPrivate.h"
 #include "D3D12TransientResourceAllocator.h"
 #include "D3D12Stats.h"
+#include "ProfilingDebugging/MemoryTrace.h"
 
 D3D12_RESOURCE_STATES GetInitialResourceState(const D3D12_RESOURCE_DESC& InDesc)
 {
@@ -74,6 +75,7 @@ FD3D12TransientHeap::FD3D12TransientHeap(const FInitializer& Initializer, FD3D12
 		// On Windows there is no way to hook into the low level d3d allocations and frees.
 		// This means that we must manually add the tracking here.
 		LLM(FLowLevelMemTracker::Get().OnLowLevelAlloc(ELLMTracker::Platform, D3DHeap, Desc.SizeInBytes, ELLMTag::GraphicsPlatform));
+		MemoryTrace_Alloc((uint64)D3DHeap, Desc.SizeInBytes, 0, EMemoryTraceRootHeap::VideoMemory);
 		// Boost priority to make sure it's not paged out
 		TRefCountPtr<ID3D12Device5> D3DDevice5;
 		if (SUCCEEDED(D3DDevice->QueryInterface(IID_PPV_ARGS(D3DDevice5.GetInitReference()))))
@@ -102,6 +104,7 @@ FD3D12TransientHeap::~FD3D12TransientHeap()
 		DEC_MEMORY_STAT_BY(STAT_D3D12MemoryCurrentTotal, Desc.SizeInBytes);
 #if PLATFORM_WINDOWS
 		LLM(FLowLevelMemTracker::Get().OnLowLevelFree(ELLMTracker::Platform, Heap->GetHeap()));
+		MemoryTrace_Free((uint64)Heap->GetHeap(), EMemoryTraceRootHeap::VideoMemory);
 #endif
 	}
 }
diff --git a/Engine/Source/Runtime/Windows/D3D11RHI/Private/D3D11RHI.cpp b/Engine/Source/Runtime/Windows/D3D11RHI/Private/D3D11RHI.cpp
--- a/Engine/Source/Runtime/Windows/D3D11RHI/Private/D3D11RHI.cpp
+++ b/Engine/Source/Runtime/Windows/D3D11RHI/Private/D3D11RHI.cpp
@@ -8,6 +8,7 @@
 #include "RHIStaticStates.h"
 #include "StaticBoundShaderState.h"
 #include "Engine/GameViewportClient.h"
+#include "ProfilingDebugging/MemoryTrace.h"
 
 #if WITH_DX_PERF
 	// For perf events
@@ -728,6 +729,7 @@ void UpdateBufferStats(TRefCountPtr<ID3D11Buffer> Buffer, bool bAllocating)
 		// to hook the actual d3d allocations we can't track the memory in the normal way.
 		// Instead we simply tell LLM the size of these resources.
 		LLM_SCOPED_PAUSE_TRACKING_WITH_ENUM_AND_AMOUNT(ELLMTag::GraphicsPlatform, Desc.ByteWidth, ELLMTracker::Platform, ELLMAllocType::None);
+		MemoryTrace_Alloc((uint64)Buffer.GetReference(), Desc.ByteWidth, 0, EMemoryTraceRootHeap::VideoMemory);
 	}
 	else
 	{ //-V523
@@ -752,6 +754,7 @@ void UpdateBufferStats(TRefCountPtr<ID3D11Buffer> Buffer, bool bAllocating)
 		// to hook the actual d3d allocations we can't track the memory in the normal way.
 		// Instead we simply tell LLM the size of these resources.
 		LLM_SCOPED_PAUSE_TRACKING_WITH_ENUM_AND_AMOUNT(ELLMTag::GraphicsPlatform, -(int64)Desc.ByteWidth, ELLMTracker::Platform, ELLMAllocType::None);
+		MemoryTrace_Free((uint64)Buffer.GetReference(), EMemoryTraceRootHeap::VideoMemory);
 	}
 }
 
diff --git a/Engine/Source/Runtime/Windows/D3D11RHI/Private/D3D11Texture.cpp b/Engine/Source/Runtime/Windows/D3D11RHI/Private/D3D11Texture.cpp
--- a/Engine/Source/Runtime/Windows/D3D11RHI/Private/D3D11Texture.cpp
+++ b/Engine/Source/Runtime/Windows/D3D11RHI/Private/D3D11Texture.cpp
@@ -15,6 +15,7 @@
 #endif
 
 #include "HAL/LowLevelMemTracker.h"
+#include "ProfilingDebugging/MemoryTrace.h"
 
 int64 FD3D11GlobalStats::GDedicatedVideoMemory = 0;
 int64 FD3D11GlobalStats::GDedicatedSystemMemory = 0;
@@ -129,11 +130,16 @@ void UpdateD3D11TextureStats(FD3D11Texture& Texture, bool bAllocating)
 	{
 		LLM(FLowLevelMemTracker::Get().OnLowLevelAlloc(ELLMTracker::Platform, Texture.GetResource(), TextureSize, ELLMTag::GraphicsPlatform));
 		LLM(FLowLevelMemTracker::Get().OnLowLevelAlloc(ELLMTracker::Default , Texture.GetResource(), TextureSize, ELLMTag::Textures));
+		{
+			UE_MEMSCOPE_DEFAULT(ELLMTag::Textures);
+			MemoryTrace_Alloc((uint64)Texture.GetResource(), TextureSize, 1024, EMemoryTraceRootHeap::VideoMemory);
+		}
 	}
 	else
 	{
 		LLM(FLowLevelMemTracker::Get().OnLowLevelFree(ELLMTracker::Platform, Texture.GetResource()));
 		LLM(FLowLevelMemTracker::Get().OnLowLevelFree(ELLMTracker::Default , Texture.GetResource()));
+		MemoryTrace_Free((uint64)Texture.GetResource(), EMemoryTraceRootHeap::VideoMemory);
 	}
 #endif
 }
