Disable the submission thread in D3D12 to work around GPU hangs when playing videos.

#rnx

diff --git a/Engine/Source/Runtime/D3D12RHI/Private/D3D12Submission.cpp b/Engine/Source/Runtime/D3D12RHI/Private/D3D12Submission.cpp
--- a/Engine/Source/Runtime/D3D12RHI/Private/D3D12Submission.cpp
+++ b/Engine/Source/Runtime/D3D12RHI/Private/D3D12Submission.cpp
@@ -9,7 +9,7 @@
 #endif
 
 // These defines control which threads are enabled in the GPU submission pipeline.
-#define D3D12_USE_SUBMISSION_THREAD (1 && (WITH_MGPU == 0)) // @todo mgpu - fix crashes when submission thread is enabled
+#define D3D12_USE_SUBMISSION_THREAD (0 && (WITH_MGPU == 0)) // @todo mgpu - fix crashes when submission thread is enabled
 #define D3D12_USE_INTERRUPT_THREAD  (1 && D3D12_PLATFORM_SUPPORTS_BLOCKING_FENCES)
 
 // When enabled, GPU timestamp queries are adjusted to remove idle time caused by CPU bubbles.
