MGPU:  Hack to disable lockstep for cross GPU transfers to improve performance for Virtual Production.  Effectively reverts code to match what it was before the Parallel Rendering integration, in three ways:  1. Disable the lockstep.  2. Enable pre-transfer synchronization only if doing pull transfers.  3. Run transfers on graphics queue, rather than copy queue (stops current GPU from overwriting or transitioning the buffer before the transfer finishes).   There is a CVar to enable the lockstep again (D3D12.UnsafeCrossGPUTransfers=0) if desired for A/B testing or investigating other D3D bugs where you need validation to run.  A proper solution for safe and efficient cross GPU resource transfers will go into 5.2.

Also disabling D3D12_ENABLE_ADJUSTED_TIMESTAMPS, necessary for GPU timeline events to line up with the CPU in Unreal Insights.  If we want to support adjusted timestamps, we would need to expose separate unadjusted timestamps through the RHI API.  This reverts behavior to before the Parallel Rendering integration, where timestamps weren't adjusted.

#rnx

diff --git a/Engine/Source/Runtime/D3D12RHI/Private/D3D12CommandContext.cpp b/Engine/Source/Runtime/D3D12RHI/Private/D3D12CommandContext.cpp
--- a/Engine/Source/Runtime/D3D12RHI/Private/D3D12CommandContext.cpp
+++ b/Engine/Source/Runtime/D3D12RHI/Private/D3D12CommandContext.cpp
@@ -33,6 +33,21 @@ static FAutoConsoleVariableRef CVarEmitRgpFrameMarkers(
 	ECVF_ReadOnly | ECVF_RenderThreadSafe
 );
 
+// jhoerner TODO 10/4/2022:  This setting is a hack to improve performance by reverting cross GPU transfer synchronization behavior to
+// what it was in 5.0, at a cost in validation correctness (D3D debug errors related to using a cross GPU transferred resource in an
+// incorrect transition state, or when possibly still being written).  In practice, these errors haven't caused artifacts or stability
+// issues, but if you run into an artifact suspected to be related to a cross GPU transfer, or want to run with validation for
+// debugging, you can disable the hack.  A future refactor in 5.2 will clean this up and provide validation correctness without any
+// performance loss.
+//
+bool GD3D12UnsafeCrossGPUTransfers = true;
+static FAutoConsoleVariableRef CVarD3D12UnsafeCrossGPUTransfers(
+	TEXT("D3D12.UnsafeCrossGPUTransfers"),
+	GD3D12UnsafeCrossGPUTransfers,
+	TEXT("Disables cross GPU synchronization correctness, for a gain in performance (Default: true)."),
+	ECVF_RenderThreadSafe
+);
+
 FD3D12CommandContextBase::FD3D12CommandContextBase(class FD3D12Adapter* InParentAdapter, FRHIGPUMask InGPUMask)
 	: FD3D12AdapterChild(InParentAdapter)
 	, GPUMask(InGPUMask)
@@ -902,7 +917,7 @@ void FD3D12CommandContextRedirector::RHITransferResources(const TArrayView<const
 
 	// Note that by default it is not empty, but GPU0
 	FRHIGPUMask SrcMask, DstMask;
-	const bool bLockstep = true; // @todo mgpu - fix synchronization
+	bool bLockstep = GD3D12UnsafeCrossGPUTransfers == false; // @todo mgpu - fix synchronization
 	bool bDelayFence = false;
 
 	{
@@ -917,7 +932,7 @@ void FD3D12CommandContextRedirector::RHITransferResources(const TArrayView<const
 			}
 
 			// @todo mgpu - fix synchronization
-			//bLockstep |= Param.bLockStepGPUs;
+			bLockstep |= Param.bLockStepGPUs;
 
 			// If it's the first time we set the mask.
 			if (bFirst)
@@ -986,8 +1001,15 @@ void FD3D12CommandContextRedirector::RHITransferResources(const TArrayView<const
 	}
 	else
 	{
-		// Destination GPUs wait for source GPUs
-		MGPUSync(SrcMask, DstMask);
+		for (const FTransferResourceParams& Param : Params)
+		{
+			if (Param.bPullData)
+			{
+				// Destination GPUs wait for source GPUs
+				MGPUSync(SrcMask, DstMask);
+				break;
+			}
+		}
 	}
 
 	// Enqueue the copy work
@@ -1010,7 +1032,7 @@ void FD3D12CommandContextRedirector::RHITransferResources(const TArrayView<const
 			// If the texture size is zero (Max.Z == 0, set in the constructor), copy the whole resource
 			if (Param.Max.Z == 0)
 			{
-				CopyContext->CopyCommandList()->CopyResource(DstTexture->GetResource()->GetResource(), SrcTexture->GetResource()->GetResource());
+				CopyContext->GraphicsCommandList()->CopyResource(DstTexture->GetResource()->GetResource(), SrcTexture->GetResource()->GetResource());
 			}
 			else
 			{
@@ -1027,7 +1049,7 @@ void FD3D12CommandContextRedirector::RHITransferResources(const TArrayView<const
 				CD3DX12_TEXTURE_COPY_LOCATION SrcLocation(SrcTexture->GetResource()->GetResource(), 0);
 				CD3DX12_TEXTURE_COPY_LOCATION DstLocation(DstTexture->GetResource()->GetResource(), 0);
 
-				CopyContext->CopyCommandList()->CopyTextureRegion(&DstLocation, Box.left, Box.top, Box.front, &SrcLocation, &Box);
+				CopyContext->GraphicsCommandList()->CopyTextureRegion(&DstLocation, Box.left, Box.top, Box.front, &SrcLocation, &Box);
 			}
 		}
 		else
@@ -1035,7 +1057,7 @@ void FD3D12CommandContextRedirector::RHITransferResources(const TArrayView<const
 			FD3D12Resource* SrcResource = FD3D12DynamicRHI::ResourceCast(Param.Buffer.GetReference(), Param.SrcGPUIndex)->GetResource();
 			FD3D12Resource* DstResource = FD3D12DynamicRHI::ResourceCast(Param.Buffer.GetReference(), Param.DestGPUIndex)->GetResource();
 
-			CopyContext->CopyCommandList()->CopyResource(DstResource->GetResource(), SrcResource->GetResource());
+			CopyContext->GraphicsCommandList()->CopyResource(DstResource->GetResource(), SrcResource->GetResource());
 		}
 	}
 
