D3D12:  Disable time stamp adjustment to fix Unreal Insights GPU timing.

#rnx

diff --git a/Engine/Source/Runtime/D3D12RHI/Private/D3D12Submission.cpp b/Engine/Source/Runtime/D3D12RHI/Private/D3D12Submission.cpp
--- a/Engine/Source/Runtime/D3D12RHI/Private/D3D12Submission.cpp
+++ b/Engine/Source/Runtime/D3D12RHI/Private/D3D12Submission.cpp
@@ -13,7 +13,7 @@
 #define D3D12_USE_INTERRUPT_THREAD  (1 && D3D12_PLATFORM_SUPPORTS_BLOCKING_FENCES)
 
 // When enabled, GPU timestamp queries are adjusted to remove idle time caused by CPU bubbles.
-#define D3D12_ENABLE_ADJUSTED_TIMESTAMPS 1
+#define D3D12_ENABLE_ADJUSTED_TIMESTAMPS 0		// Adjusted timestamps break Unreal Insights, disabling for now
 
 DECLARE_CYCLE_STAT(TEXT("Submit"), STAT_D3D12Submit, STATGROUP_D3D12RHI);
 
