Move DiagnosticBuffer binding out of the per-stage UAV binding function (SetUAVs) and make it only bind on RootSignature change. Binding in SetUAVs would also have prevented it from being bound correctly with bindless resources.

diff --git a/Engine/Source/Runtime/D3D12RHI/Private/D3D12DescriptorCache.cpp b/Engine/Source/Runtime/D3D12RHI/Private/D3D12DescriptorCache.cpp
--- a/Engine/Source/Runtime/D3D12RHI/Private/D3D12DescriptorCache.cpp
+++ b/Engine/Source/Runtime/D3D12RHI/Private/D3D12DescriptorCache.cpp
@@ -263,22 +263,6 @@ void FD3D12DescriptorCache::SetUAVs(const FD3D12RootSignature* RootSignature, FD
 #ifdef VERBOSE_DESCRIPTOR_HEAP_DEBUG
 	FMsg::Logf(__FILE__, __LINE__, TEXT("DescriptorCache"), ELogVerbosity::Log, TEXT("SetUnorderedAccessViewTable [STAGE %d] to slots %d - %d"), (int32)ShaderStage, FirstSlotIndex, FirstSlotIndex + SlotsNeeded - 1);
 #endif
-
-	const int8 DiagnosticBufferSlot = RootSignature->GetDiagnosticBufferSlot();
-	FD3D12Queue& Queue = Context.GetParentDevice()->GetQueue(Context.QueueType);
-	D3D12_GPU_VIRTUAL_ADDRESS DiagnosticBufferAddress = Queue.GetDiagnosticBufferGPUAddress();
-
-	if (DiagnosticBufferSlot >= 0 && DiagnosticBufferAddress)
-	{
-		if (ShaderStage == SF_Compute)
-		{
-			Context.GraphicsCommandList()->SetComputeRootUnorderedAccessView(DiagnosticBufferSlot, DiagnosticBufferAddress);
-		}
-		else
-		{
-			Context.GraphicsCommandList()->SetGraphicsRootUnorderedAccessView(DiagnosticBufferSlot, DiagnosticBufferAddress);
-		}
-	}
 }
 
 void FD3D12DescriptorCache::SetRenderTargets(FD3D12RenderTargetView** RenderTargetViewArray, uint32 Count, FD3D12DepthStencilView* DepthStencilTarget)
diff --git a/Engine/Source/Runtime/D3D12RHI/Private/D3D12StateCache.cpp b/Engine/Source/Runtime/D3D12RHI/Private/D3D12StateCache.cpp
--- a/Engine/Source/Runtime/D3D12RHI/Private/D3D12StateCache.cpp
+++ b/Engine/Source/Runtime/D3D12RHI/Private/D3D12StateCache.cpp
@@ -450,6 +450,25 @@ void FD3D12StateCache::ApplyState()
 		);
 	}
 
+	if (bRootSignatureChanged)
+	{
+		const int8 DiagnosticBufferSlot = PSOCommonData->RootSignature->GetDiagnosticBufferSlot();
+		const FD3D12Queue& Queue = GetParentDevice()->GetQueue(CmdContext.QueueType);
+		const D3D12_GPU_VIRTUAL_ADDRESS DiagnosticBufferAddress = Queue.GetDiagnosticBufferGPUAddress();
+
+		if (DiagnosticBufferSlot >= 0 && DiagnosticBufferAddress)
+		{
+			if (PipelineType == ED3D12PipelineType::Compute)
+			{
+				CmdContext.GraphicsCommandList()->SetComputeRootUnorderedAccessView(DiagnosticBufferSlot, DiagnosticBufferAddress);
+			}
+			else
+			{
+				CmdContext.GraphicsCommandList()->SetGraphicsRootUnorderedAccessView(DiagnosticBufferSlot, DiagnosticBufferAddress);
+			}
+		}
+	}
+
 	// Need to cache compute budget, as we need to reset after PSO changes
 	if (PipelineType == ED3D12PipelineType::Compute && CmdContext.IsAsyncComputeContext())
 	{
