Fix crash when uploading resources via the copy queue.

This works around an apparent bug in D3D where the QueryInterface call fails for copy queue command lists, but incorrectly succeeds for the graphics command list interface on the same object. For copy command lists, we need to store the pointer returned by CreateCommandList directly in the CopyCommandList member to avoid the bad QI call, and we should not populate the GraphicsCommandListN members, to avoid calling methods on this object with an incorrect vtable.

#rnx

diff --git a/Engine/Source/Runtime/D3D12RHI/Private/D3D12Buffer.cpp b/Engine/Source/Runtime/D3D12RHI/Private/D3D12Buffer.cpp
--- a/Engine/Source/Runtime/D3D12RHI/Private/D3D12Buffer.cpp
+++ b/Engine/Source/Runtime/D3D12RHI/Private/D3D12Buffer.cpp
@@ -280,7 +280,7 @@ FD3D12SyncPointRef FD3D12Buffer::UploadResourceDataViaCopyQueue(FResourceArrayIn
 		SyncPoint = CopyScope.GetSyncPoint();
 
 		// Perform actual copy op
-		CopyScope.Context.GraphicsCommandList()->CopyBufferRegion(
+		CopyScope.Context.CopyCommandList()->CopyBufferRegion(
 			ResourceLocation.GetResource()->GetResource(),
 			ResourceLocation.GetOffsetFromBaseOfResource(),
 			SrcResourceLoc.GetResource()->GetResource(),
diff --git a/Engine/Source/Runtime/D3D12RHI/Private/D3D12CommandList.cpp b/Engine/Source/Runtime/D3D12RHI/Private/D3D12CommandList.cpp
--- a/Engine/Source/Runtime/D3D12RHI/Private/D3D12CommandList.cpp
+++ b/Engine/Source/Runtime/D3D12RHI/Private/D3D12CommandList.cpp
@@ -173,46 +173,51 @@ FD3D12CommandList::FD3D12CommandList(FD3D12CommandAllocator* CommandAllocator, F
 			IID_PPV_ARGS(Interfaces.GraphicsCommandList.GetInitReference())
 		));
 		Interfaces.CommandList = Interfaces.GraphicsCommandList;
-		break;
 
-	default:
-		VERIFYD3D12RESULT(Device->GetDevice()->CreateCommandList(
-			Device->GetGPUMask().GetNative(),
-			GetD3DCommandListType(QueueType),
-			*CommandAllocator,
-			nullptr,
-			IID_PPV_ARGS(Interfaces.CommandList.GetInitReference())
-		));
+		Interfaces.CommandList->QueryInterface(IID_PPV_ARGS(Interfaces.CopyCommandList.GetInitReference()));
 
-		Interfaces.CommandList->QueryInterface(IID_PPV_ARGS(Interfaces.GraphicsCommandList.GetInitReference()));
-		break;
-	}
-
-	INC_DWORD_STAT(STAT_D3D12NumCommandLists);
-
-	// Optionally obtain the versioned ID3D12GraphicsCommandList[0-9]+ interfaces, we don't check the HRESULT.
-	Interfaces.CommandList->QueryInterface(IID_PPV_ARGS(Interfaces.CopyCommandList.GetInitReference()));
+		// Optionally obtain the versioned ID3D12GraphicsCommandList[0-9]+ interfaces, we don't check the HRESULT.
 #if D3D12_MAX_COMMANDLIST_INTERFACE >= 1
-	Interfaces.CommandList->QueryInterface(IID_PPV_ARGS(Interfaces.GraphicsCommandList1.GetInitReference()));
+		Interfaces.CommandList->QueryInterface(IID_PPV_ARGS(Interfaces.GraphicsCommandList1.GetInitReference()));
 #endif
 #if D3D12_MAX_COMMANDLIST_INTERFACE >= 2
-	Interfaces.CommandList->QueryInterface(IID_PPV_ARGS(Interfaces.GraphicsCommandList2.GetInitReference()));
+		Interfaces.CommandList->QueryInterface(IID_PPV_ARGS(Interfaces.GraphicsCommandList2.GetInitReference()));
 #endif
 #if D3D12_MAX_COMMANDLIST_INTERFACE >= 3
-	Interfaces.CommandList->QueryInterface(IID_PPV_ARGS(Interfaces.GraphicsCommandList3.GetInitReference()));
+		Interfaces.CommandList->QueryInterface(IID_PPV_ARGS(Interfaces.GraphicsCommandList3.GetInitReference()));
 #endif
 #if D3D12_MAX_COMMANDLIST_INTERFACE >= 4
-	Interfaces.CommandList->QueryInterface(IID_PPV_ARGS(Interfaces.GraphicsCommandList4.GetInitReference()));
+		Interfaces.CommandList->QueryInterface(IID_PPV_ARGS(Interfaces.GraphicsCommandList4.GetInitReference()));
 #endif
 #if D3D12_MAX_COMMANDLIST_INTERFACE >= 5
-	Interfaces.CommandList->QueryInterface(IID_PPV_ARGS(Interfaces.GraphicsCommandList5.GetInitReference()));
+		Interfaces.CommandList->QueryInterface(IID_PPV_ARGS(Interfaces.GraphicsCommandList5.GetInitReference()));
 #endif
 #if D3D12_MAX_COMMANDLIST_INTERFACE >= 6
-	Interfaces.CommandList->QueryInterface(IID_PPV_ARGS(Interfaces.GraphicsCommandList6.GetInitReference()));
+		Interfaces.CommandList->QueryInterface(IID_PPV_ARGS(Interfaces.GraphicsCommandList6.GetInitReference()));
 #endif
 #if D3D12_PLATFORM_SUPPORTS_ASSERTRESOURCESTATES
-	Interfaces.CommandList->QueryInterface(IID_PPV_ARGS(Interfaces.DebugCommandList.GetInitReference()));
+		Interfaces.CommandList->QueryInterface(IID_PPV_ARGS(Interfaces.DebugCommandList.GetInitReference()));
 #endif
+		break;
+
+	case ED3D12QueueType::Copy:
+		VERIFYD3D12RESULT(Device->GetDevice()->CreateCommandList(
+			Device->GetGPUMask().GetNative(),
+			GetD3DCommandListType(QueueType),
+			*CommandAllocator,
+			nullptr,
+			IID_PPV_ARGS(Interfaces.CopyCommandList.GetInitReference())
+		));
+		Interfaces.CommandList = Interfaces.CopyCommandList;
+
+		break;
+
+	default:
+		checkNoEntry();
+		return;
+	}
+
+	INC_DWORD_STAT(STAT_D3D12NumCommandLists);
 
 #if NV_AFTERMATH
 	if (GDX12NVAfterMathEnabled)
