Fix GPU hangs when playing videos with the submission thread enabled.

The fence waits that the video player performs must be performed by the submission thread, rather than by the RHI thread. Added an RHIWaitManualFence() method for enabling this.
Also, when we enqueue the lambda which copies the decoded frame to a texture, we must wait for the submission thread to execute it before continuing running the decoder thread, to prevent a race condition where the internal fence value in the media framework sample object is incremented before the wait is added to the copy queue, which effectively makes the GPU wait for a value which is never signaled.

Also re-enabled the submission thread, which was disabled before as a workaround for this GPU hang.

#rnx

diff --git a/Engine/Source/Runtime/D3D12RHI/Private/D3D12CommandContext.cpp b/Engine/Source/Runtime/D3D12RHI/Private/D3D12CommandContext.cpp
--- a/Engine/Source/Runtime/D3D12RHI/Private/D3D12CommandContext.cpp
+++ b/Engine/Source/Runtime/D3D12RHI/Private/D3D12CommandContext.cpp
@@ -303,6 +303,12 @@ void FD3D12ContextCommon::SignalManualFence(ID3D12Fence* Fence, uint64 Value)
 	GetPayload(EPhase::Signal)->FencesToSignal.Emplace(Fence, Value);
 }
 
+void FD3D12ContextCommon::WaitManualFence(ID3D12Fence* Fence, uint64 Value)
+{
+	checkf(!CommandList, TEXT("Command list must be closed before signaling a manual fence."));
+	GetPayload(EPhase::Wait)->FencesToWait.Emplace(Fence, Value);
+}
+
 FD3D12QueryLocation FD3D12ContextCommon::AllocateQuery(ED3D12QueryType Type, uint64* Target)
 {
 	checkf(CommandList, TEXT("Command list must be open to insert a timestamp."));
diff --git a/Engine/Source/Runtime/D3D12RHI/Private/D3D12CommandContext.h b/Engine/Source/Runtime/D3D12RHI/Private/D3D12CommandContext.h
--- a/Engine/Source/Runtime/D3D12RHI/Private/D3D12CommandContext.h
+++ b/Engine/Source/Runtime/D3D12RHI/Private/D3D12CommandContext.h
@@ -124,6 +124,9 @@ public:
 	// Inserts a command that signals the specified D3D12 fence object.
 	void SignalManualFence(ID3D12Fence* Fence, uint64 Value);
 
+	// Inserts a command that waits the specified D3D12 fence object.
+	void WaitManualFence(ID3D12Fence* Fence, uint64 Value);
+
 	// Inserts a timestamp query command. "Target" specifies the optional 
 	// location the result will be written to by the interrupt handler thread.
 	FD3D12QueryLocation InsertTimestamp(ED3D12Units Units, uint64* Target);
diff --git a/Engine/Source/Runtime/D3D12RHI/Private/D3D12DirectCommandListManager.cpp b/Engine/Source/Runtime/D3D12RHI/Private/D3D12DirectCommandListManager.cpp
--- a/Engine/Source/Runtime/D3D12RHI/Private/D3D12DirectCommandListManager.cpp
+++ b/Engine/Source/Runtime/D3D12RHI/Private/D3D12DirectCommandListManager.cpp
@@ -112,6 +112,7 @@ FD3D12ManualFence::FD3D12ManualFence(FD3D12Adapter* InParent)
 		Pair.Context = &Device->GetDefaultCommandContext();
 
 		VERIFYD3D12RESULT(Parent->GetD3DDevice()->CreateFence(0, D3D12_FENCE_FLAG_NONE, IID_PPV_ARGS(Pair.Fence.GetInitReference())));
+		Pair.Fence->SetName(TEXT("Manual Fence"));
 	}
 }
 
diff --git a/Engine/Source/Runtime/D3D12RHI/Private/D3D12RHI.cpp b/Engine/Source/Runtime/D3D12RHI/Private/D3D12RHI.cpp
--- a/Engine/Source/Runtime/D3D12RHI/Private/D3D12RHI.cpp
+++ b/Engine/Source/Runtime/D3D12RHI/Private/D3D12RHI.cpp
@@ -620,6 +620,24 @@ void FD3D12DynamicRHI::RHITransitionResource(FRHICommandList& RHICmdList, FRHITe
 	FD3D12CommandContext::Get(RHICmdList).TransitionResource(D3D12Texture->GetResource(), D3D12_RESOURCE_STATE_TBD, InState, InSubResource);
 }
 
+void FD3D12DynamicRHI::RHISignalManualFence(FRHICommandList& RHICmdList, ID3D12Fence* Fence, uint64 Value)
+{
+	checkf(FRHIGPUMask::All() == FRHIGPUMask::GPU0(), TEXT("RHISignalManualFence cannot be used by multi-GPU code"));
+	FD3D12CommandContext& Context = FD3D12CommandContext::Get(RHICmdList);
+	Context.CloseCommandList(false);
+	Context.SignalManualFence(Fence, Value);
+	Context.OpenCommandList();
+}
+
+void FD3D12DynamicRHI::RHIWaitManualFence(FRHICommandList& RHICmdList, ID3D12Fence* Fence, uint64 Value)
+{
+	checkf(FRHIGPUMask::All() == FRHIGPUMask::GPU0(), TEXT("RHIWaitManualFence cannot be used by multi-GPU code"));
+	FD3D12CommandContext& Context = FD3D12CommandContext::Get(RHICmdList);
+	Context.CloseCommandList(false);
+	Context.WaitManualFence(Fence, Value);
+	Context.OpenCommandList();
+}
+
 bool ID3D12DynamicRHI::IsD3DDebugEnabled()
 {
 	return D3D12RHI_ShouldCreateWithD3DDebug();
diff --git a/Engine/Source/Runtime/D3D12RHI/Private/D3D12RHIPrivate.h b/Engine/Source/Runtime/D3D12RHI/Private/D3D12RHIPrivate.h
--- a/Engine/Source/Runtime/D3D12RHI/Private/D3D12RHIPrivate.h
+++ b/Engine/Source/Runtime/D3D12RHI/Private/D3D12RHIPrivate.h
@@ -560,6 +560,8 @@ public:
 	virtual D3D12_CPU_DESCRIPTOR_HANDLE RHIGetRenderTargetView(FRHITexture* InTexture, int32 InMipIndex = 0, int32 InArraySliceIndex = 0) const final override;
 	virtual void RHIFinishExternalComputeWork(uint32 InDeviceIndex, ID3D12GraphicsCommandList* InCommandList) final override;
 	virtual void RHITransitionResource(FRHICommandList& RHICmdList, FRHITexture* InTexture, D3D12_RESOURCE_STATES InState, uint32 InSubResource) final override;
+	virtual void RHISignalManualFence(FRHICommandList& RHICmdList, ID3D12Fence* Fence, uint64 Value) final override;
+	virtual void RHIWaitManualFence(FRHICommandList& RHICmdList, ID3D12Fence* Fence, uint64 Value) final override;
 
 	//
 	// The Following functions are the _RenderThread version of the above functions. They allow the RHI to control the thread synchronization for greater efficiency.
diff --git a/Engine/Source/Runtime/D3D12RHI/Private/D3D12Submission.cpp b/Engine/Source/Runtime/D3D12RHI/Private/D3D12Submission.cpp
--- a/Engine/Source/Runtime/D3D12RHI/Private/D3D12Submission.cpp
+++ b/Engine/Source/Runtime/D3D12RHI/Private/D3D12Submission.cpp
@@ -9,7 +9,7 @@
 #endif
 
 // These defines control which threads are enabled in the GPU submission pipeline.
-#define D3D12_USE_SUBMISSION_THREAD (0 && (WITH_MGPU == 0)) // @todo mgpu - fix crashes when submission thread is enabled
+#define D3D12_USE_SUBMISSION_THREAD (1 && (WITH_MGPU == 0)) // @todo mgpu - fix crashes when submission thread is enabled
 #define D3D12_USE_INTERRUPT_THREAD  (1 && D3D12_PLATFORM_SUPPORTS_BLOCKING_FENCES)
 
 // When enabled, GPU timestamp queries are adjusted to remove idle time caused by CPU bubbles.
@@ -607,6 +607,13 @@ uint64 FD3D12Queue::ExecutePayload()
 	TRACE_CPUPROFILER_EVENT_SCOPE(ExecuteCommandList);
 	LLM_SCOPE_BYNAME(TEXT("RHIMisc/ExecuteCommandLists"));
 	check(PayloadToSubmit && this == &PayloadToSubmit->Queue);
+
+	// Wait for manual fences.
+	for (auto& [ManualFence, Value] : PayloadToSubmit->FencesToWait)
+	{
+		VERIFYD3D12RESULT(D3DCommandQueue->Wait(ManualFence, Value));
+	}
+
 	PayloadToSubmit->PreExecute();
 
 	if (const int32 NumCommandLists = PayloadToSubmit->CommandListsToExecute.Num())
diff --git a/Engine/Source/Runtime/D3D12RHI/Private/D3D12Submission.h b/Engine/Source/Runtime/D3D12RHI/Private/D3D12Submission.h
--- a/Engine/Source/Runtime/D3D12RHI/Private/D3D12Submission.h
+++ b/Engine/Source/Runtime/D3D12RHI/Private/D3D12Submission.h
@@ -193,6 +193,9 @@ struct FD3D12PayloadBase
 
 	virtual void PreExecute() {}
 
+	// Wait
+	TArray<FManualFence> FencesToWait;
+
 	// Execute
 	bool bResolveQueries = false;
 	TArray<class FD3D12CommandList*> CommandListsToExecute;
diff --git a/Engine/Source/Runtime/D3D12RHI/Public/ID3D12DynamicRHI.h b/Engine/Source/Runtime/D3D12RHI/Public/ID3D12DynamicRHI.h
--- a/Engine/Source/Runtime/D3D12RHI/Public/ID3D12DynamicRHI.h
+++ b/Engine/Source/Runtime/D3D12RHI/Public/ID3D12DynamicRHI.h
@@ -48,6 +48,9 @@ struct ID3D12DynamicRHI : public FDynamicRHI
 	virtual void                  RHIFinishExternalComputeWork(uint32 InDeviceIndex, ID3D12GraphicsCommandList* InCommandList) = 0;
 	virtual void                  RHITransitionResource(FRHICommandList& RHICmdList, FRHITexture* InTexture, D3D12_RESOURCE_STATES InState, uint32 InSubResource) = 0;
 
+	virtual void						RHISignalManualFence(FRHICommandList& RHICmdList, ID3D12Fence* Fence, uint64 Value) = 0;
+	virtual void						RHIWaitManualFence(FRHICommandList& RHICmdList, ID3D12Fence* Fence, uint64 Value) = 0;
+
 	static D3D12RHI_API bool      IsD3DDebugEnabled();
 };
 
