Only touch the pool allocation data while the pool lock is taken - the temporary stored allocation data could have been changed by an allocation done on the pool on another thread while the old swap is happening

diff --git a/Engine/Source/Runtime/D3D12RHI/Private/D3D12RHICommon.h b/Engine/Source/Runtime/D3D12RHI/Private/D3D12RHICommon.h
--- a/Engine/Source/Runtime/D3D12RHI/Private/D3D12RHICommon.h
+++ b/Engine/Source/Runtime/D3D12RHI/Private/D3D12RHICommon.h
@@ -64,6 +64,11 @@ public:
 	{
 		return Parent;
 	}
+
+	void Swap(FD3D12DeviceChild& Other)
+	{
+		::Swap(*this, Other);
+	}
 };
 
 class FD3D12GPUObject
diff --git a/Engine/Source/Runtime/D3D12RHI/Private/D3D12Resources.cpp b/Engine/Source/Runtime/D3D12RHI/Private/D3D12Resources.cpp
--- a/Engine/Source/Runtime/D3D12RHI/Private/D3D12Resources.cpp
+++ b/Engine/Source/Runtime/D3D12RHI/Private/D3D12Resources.cpp
@@ -571,23 +571,27 @@ void FD3D12ResourceLocation::Swap(FD3D12ResourceLocation& Other)
 	if (Other.GetAllocatorType() == FD3D12ResourceLocation::AT_Pool)
 	{
 		check(GetAllocatorType() != FD3D12ResourceLocation::AT_Pool);
-		
-		// Cache the allocator data and reset before swap
-		FD3D12PoolAllocatorPrivateData& PoolData = Other.GetPoolAllocatorPrivateData();
-		FD3D12PoolAllocatorPrivateData TmpPoolData = PoolData;
-		PoolData.Init();
 
-		// Perform swap
-		::Swap(*this, Other);
+		// Swap all members except the pool data because that needs to happen while the pool is taken
+		FD3D12DeviceChild::Swap(Other);
 
-		// Restore allocator data and perform pool aware swap
-		PoolData = TmpPoolData;
+		::Swap(Owner, Other.Owner);
+		::Swap(UnderlyingResource, Other.UnderlyingResource);
+		::Swap(ResidencyHandle, Other.ResidencyHandle);
+				
+		::Swap(MappedBaseAddress, Other.MappedBaseAddress);
+		::Swap(GPUVirtualAddress, Other.GPUVirtualAddress);
+		::Swap(OffsetFromBaseOfResource, Other.OffsetFromBaseOfResource);
 
-		// Reset the tmp pool data again, because it's not needed anymore - all data copied over
-		TmpPoolData.Init();
+		::Swap(Type, Other.Type);
+		::Swap(Size, Other.Size);
+		::Swap(bTransient, Other.bTransient);
+				
+		// We know this allocation is empty so can use transfer instead of full swap
+		Other.GetPoolAllocator()->TransferOwnership(Other, *this);
 
-		Other.SetPoolAllocator(GetPoolAllocator());
-		GetPoolAllocator()->TransferOwnership(Other, *this);
+		::Swap(AllocatorType, Other.AllocatorType);
+		::Swap(PoolAllocator, Other.PoolAllocator);
 	}
 	else
 	{
