Fixed an issue where swapping the resource would not properly update the resource tracking.

Submitting on behalf of Bryan Sefcik due to preflight not submitting

Preflight fail due to unrelated WorldPartition issue

diff --git a/Engine/Source/Runtime/D3D12RHI/Private/D3D12Resources.cpp b/Engine/Source/Runtime/D3D12RHI/Private/D3D12Resources.cpp
--- a/Engine/Source/Runtime/D3D12RHI/Private/D3D12Resources.cpp
+++ b/Engine/Source/Runtime/D3D12RHI/Private/D3D12Resources.cpp
@@ -568,6 +568,12 @@ void FD3D12ResourceLocation::Swap(FD3D12ResourceLocation& Other)
 	}
 #endif
 
+	if (Other.Type == ResourceLocationType::eStandAlone)
+	{
+		bool bIncrement = false;
+		Other.UpdateStandAloneStats(bIncrement);
+	}
+
 	if (Other.GetAllocatorType() == FD3D12ResourceLocation::AT_Pool)
 	{
 		check(GetAllocatorType() != FD3D12ResourceLocation::AT_Pool);
@@ -597,6 +603,12 @@ void FD3D12ResourceLocation::Swap(FD3D12ResourceLocation& Other)
 	{
 		::Swap(*this, Other);
 	}
+
+	if (Type == ResourceLocationType::eStandAlone)
+	{
+		bool bIncrement = true;
+		UpdateStandAloneStats(bIncrement);
+	}
 }
 
 void FD3D12ResourceLocation::Alias(FD3D12ResourceLocation & Destination, FD3D12ResourceLocation & Source)
