Update HDR display info when display changed

Only handling the display change event is not enough though: the DXGI factory might not be contain up-to-date display information in case it's not current anymore. In that case, we need to rely on a newly created factory, used exclusively to retrieve the display info

diff --git a/Engine/Source/Runtime/D3D12RHI/Private/D3D12Adapter.cpp b/Engine/Source/Runtime/D3D12RHI/Private/D3D12Adapter.cpp
--- a/Engine/Source/Runtime/D3D12RHI/Private/D3D12Adapter.cpp
+++ b/Engine/Source/Runtime/D3D12RHI/Private/D3D12Adapter.cpp
@@ -1392,16 +1392,12 @@ FD3D12Adapter::~FD3D12Adapter()
 #endif
 }
 
-void FD3D12Adapter::CreateDXGIFactory(bool bWithDebug)
+void FD3D12Adapter::CreateDXGIFactory(TRefCountPtr<IDXGIFactory2>& DxgiFactory2, bool bWithDebug, HMODULE DxgiDllHandle)
 {
 #if PLATFORM_WINDOWS || PLATFORM_HOLOLENS
 	typedef HRESULT(WINAPI FCreateDXGIFactory2)(UINT, REFIID, void**);
 
 #if PLATFORM_WINDOWS
-	// Dynamically load this otherwise Win7 fails to boot as it's missing on that DLL
-	DxgiDllHandle = (HMODULE)FPlatformProcess::GetDllHandle(TEXT("dxgi.dll"));
-	check(DxgiDllHandle);
-
 	FCreateDXGIFactory2* CreateDXGIFactory2FnPtr = (FCreateDXGIFactory2*)(void*)::GetProcAddress(DxgiDllHandle, "CreateDXGIFactory2");
 #else
 	FCreateDXGIFactory2* CreateDXGIFactory2FnPtr = &CreateDXGIFactory2;
@@ -1411,6 +1407,21 @@ void FD3D12Adapter::CreateDXGIFactory(bool bWithDebug)
 
 	const uint32 Flags = bWithDebug ? DXGI_CREATE_FACTORY_DEBUG : 0;
 	VERIFYD3D12RESULT(CreateDXGIFactory2FnPtr(Flags, IID_PPV_ARGS(DxgiFactory2.GetInitReference())));
+#endif // #if PLATFORM_WINDOWS || PLATFORM_HOLOLENS
+}
+
+void FD3D12Adapter::CreateDXGIFactory(bool bWithDebug)
+{
+#if PLATFORM_WINDOWS || PLATFORM_HOLOLENS
+	HMODULE UsedDxgiDllHandle = (HMODULE)0;
+#if PLATFORM_WINDOWS
+	// Dynamically load this otherwise Win7 fails to boot as it's missing on that DLL
+	DxgiDllHandle = (HMODULE)FPlatformProcess::GetDllHandle(TEXT("dxgi.dll"));
+	check(DxgiDllHandle);
+	UsedDxgiDllHandle = DxgiDllHandle;
+#endif
+
+	CreateDXGIFactory(DxgiFactory2, bWithDebug, UsedDxgiDllHandle);
 
 	InitDXGIFactoryVariants(DxgiFactory2);
 #endif // #if PLATFORM_WINDOWS || PLATFORM_HOLOLENS
diff --git a/Engine/Source/Runtime/D3D12RHI/Private/D3D12Adapter.h b/Engine/Source/Runtime/D3D12RHI/Private/D3D12Adapter.h
--- a/Engine/Source/Runtime/D3D12RHI/Private/D3D12Adapter.h
+++ b/Engine/Source/Runtime/D3D12RHI/Private/D3D12Adapter.h
@@ -263,6 +263,7 @@ public:
 	FORCEINLINE uint32 GetVRSTileSize() const { return VRSTileSize; }
 
 	void CreateDXGIFactory(bool bWithDebug);
+	static void CreateDXGIFactory(TRefCountPtr<IDXGIFactory2>& DxgiFactory2, bool bWithDebug, HMODULE DxgiDllHandle);
 	void InitDXGIFactoryVariants(IDXGIFactory2* InDxgiFactory2);
 	HRESULT EnumAdapters(IDXGIAdapter** TempAdapter) const;
 
@@ -433,6 +434,10 @@ public:
 
 	void SetResidencyPriority(ID3D12Pageable* Pageable, D3D12_RESIDENCY_PRIORITY HeapPriority, uint32 GPUIndex);
 
+#if PLATFORM_WINDOWS
+	HMODULE GetDxgiDllHandle() const { return DxgiDllHandle; };
+#endif
+
 protected:
 
 	virtual void CreateRootDevice(bool bWithDebug);
diff --git a/Engine/Source/Runtime/D3D12RHI/Private/D3D12RHIPrivate.h b/Engine/Source/Runtime/D3D12RHI/Private/D3D12RHIPrivate.h
--- a/Engine/Source/Runtime/D3D12RHI/Private/D3D12RHIPrivate.h
+++ b/Engine/Source/Runtime/D3D12RHI/Private/D3D12RHIPrivate.h
@@ -516,6 +516,9 @@ public:
 	virtual FViewportRHIRef RHICreateViewport(void* WindowHandle, uint32 SizeX, uint32 SizeY, bool bIsFullscreen, EPixelFormat PreferredPixelFormat) final override;
 	virtual void RHIResizeViewport(FRHIViewport* Viewport, uint32 SizeX, uint32 SizeY, bool bIsFullscreen) final override;
 	virtual void RHIResizeViewport(FRHIViewport* ViewportRHI, uint32 SizeX, uint32 SizeY, bool bIsFullscreen, EPixelFormat PreferredPixelFormat) final override;
+#if PLATFORM_WINDOWS
+	virtual void RHIHandleDisplayChange() final override;
+#endif
 	virtual void RHITick(float DeltaTime) final override;
 	virtual void RHIBlockUntilGPUIdle() final override;
 	virtual void RHISubmitCommandsAndFlushGPU() final override;
@@ -762,6 +765,10 @@ protected:
 	void* ZeroBuffer;
 	uint32 ZeroBufferSize;
 
+#if PLATFORM_WINDOWS
+	TRefCountPtr<IDXGIFactory2> DXGIFactoryForDisplayList;
+#endif
+
 public:
 
 	virtual FD3D12ResourceDesc GetResourceDesc(const FRHITextureDesc& CreateInfo) const;
diff --git a/Engine/Source/Runtime/D3D12RHI/Private/Windows/WindowsD3D12Device.cpp b/Engine/Source/Runtime/D3D12RHI/Private/Windows/WindowsD3D12Device.cpp
--- a/Engine/Source/Runtime/D3D12RHI/Private/Windows/WindowsD3D12Device.cpp
+++ b/Engine/Source/Runtime/D3D12RHI/Private/Windows/WindowsD3D12Device.cpp
@@ -520,14 +520,49 @@ bool FD3D12DynamicRHI::SetupDisplayHDRMetaData()
 	// Determines if any displays support HDR
 	check(GetNumAdapters() >= 1);
 
-	DisplayList.Empty();
+	TArray<TRefCountPtr<IDXGIAdapter> > DXGIAdapters;
+	bool bStaleAdapters = false;
 
-	bool bSupportsHDROutput = false;
 	const int32 NumAdapters = GetNumAdapters();
 	for (int32 AdapterIndex = 0; AdapterIndex < NumAdapters; ++AdapterIndex)
 	{
-		FD3D12Adapter& Adapter = GetAdapter(AdapterIndex);
-		IDXGIAdapter* DXGIAdapter = Adapter.GetAdapter();
+		FD3D12Adapter& CurrentAdapter = GetAdapter(AdapterIndex);
+		DXGIAdapters.Add(GetAdapter(AdapterIndex).GetAdapter());
+		if (CurrentAdapter.GetDXGIFactory2() != nullptr && !CurrentAdapter.GetDXGIFactory2()->IsCurrent())
+		{
+			bStaleAdapters = true;
+		}
+	}
+
+#if PLATFORM_WINDOWS
+	// if we found that the list of adapters is stale (changed windows HDR setting), try to update it with the new list
+	if (bStaleAdapters)
+	{
+		if (!DXGIFactoryForDisplayList.IsValid() || !DXGIFactoryForDisplayList->IsCurrent())
+		{
+			FD3D12Adapter::CreateDXGIFactory(DXGIFactoryForDisplayList, false, GetAdapter(0).GetDxgiDllHandle());
+		}
+
+		if (DXGIFactoryForDisplayList.IsValid() && DXGIFactoryForDisplayList->IsCurrent())
+		{
+			DXGIAdapters.Empty();
+			TRefCountPtr<IDXGIAdapter> TempAdapter;
+			for (uint32 AdapterIndex = 0; DXGIFactoryForDisplayList->EnumAdapters(AdapterIndex, TempAdapter.GetInitReference()) != DXGI_ERROR_NOT_FOUND; ++AdapterIndex)
+			{
+				DXGIAdapters.Add(TempAdapter);
+			}
+		}
+
+	}
+#endif
+
+	DisplayList.Empty();
+
+	bool bSupportsHDROutput = false;
+	const int32 NumDXGIAdapters = DXGIAdapters.Num();
+	for (int32 AdapterIndex = 0; AdapterIndex < NumDXGIAdapters; ++AdapterIndex)
+	{
+		IDXGIAdapter* DXGIAdapter = DXGIAdapters[AdapterIndex];
 
 		for (uint32 DisplayIndex = 0; true; ++DisplayIndex)
 		{
@@ -566,6 +601,15 @@ bool FD3D12DynamicRHI::SetupDisplayHDRMetaData()
 	return bSupportsHDROutput;
 }
 
+extern void HDRSettingChangedSinkCallback();
+void FD3D12DynamicRHI::RHIHandleDisplayChange()
+{
+	RHIBlockUntilGPUIdle();
+	GRHISupportsHDROutput = SetupDisplayHDRMetaData();
+	// make sure CVars are being updated properly
+	HDRSettingChangedSinkCallback();
+}
+
 static bool IsAdapterBlocked(FD3D12Adapter* InAdapter)
 {
 #if !UE_BUILD_SHIPPING
diff --git a/Engine/Source/Runtime/D3D12RHI/Private/Windows/WindowsD3D12Viewport.cpp b/Engine/Source/Runtime/D3D12RHI/Private/Windows/WindowsD3D12Viewport.cpp
--- a/Engine/Source/Runtime/D3D12RHI/Private/Windows/WindowsD3D12Viewport.cpp
+++ b/Engine/Source/Runtime/D3D12RHI/Private/Windows/WindowsD3D12Viewport.cpp
@@ -403,7 +403,8 @@ void FD3D12Viewport::EnableHDR()
 
 void FD3D12Viewport::ShutdownHDR()
 {
-	if (GRHISupportsHDROutput)
+	// Make sure to set the appropriate color space even if GRHISupportsHDROutput is false because we 
+	// might have toggled HDR on and off in the windows settings
 	{
 		// Ensure we have the correct color space set.
 		EnsureColorSpace(EDisplayColorGamut::sRGB_D65, EDisplayOutputFormat::SDR_sRGB);
diff --git a/Engine/Source/Runtime/RHI/Public/DynamicRHI.h b/Engine/Source/Runtime/RHI/Public/DynamicRHI.h
--- a/Engine/Source/Runtime/RHI/Public/DynamicRHI.h
+++ b/Engine/Source/Runtime/RHI/Public/DynamicRHI.h
@@ -894,6 +894,7 @@ public:
 
 	// Tests the viewport to see if its HDR status has changed. This is usually tested after a window has been moved
 	virtual void RHICheckViewportHDRStatus(FRHIViewport* Viewport);
+	virtual void RHIHandleDisplayChange() {}
 
 	//  must be called from the main thread.
 	// FlushType: Thread safe
@@ -1524,6 +1525,11 @@ FORCEINLINE void RHICheckViewportHDRStatus(FRHIViewport* Viewport)
 	GDynamicRHI->RHICheckViewportHDRStatus(Viewport);
 }
 
+FORCEINLINE void RHIHandleDisplayChange()
+{
+	GDynamicRHI->RHIHandleDisplayChange();
+}
+
 FORCEINLINE void RHITick(float DeltaTime)
 {
 	GDynamicRHI->RHITick(DeltaTime);
diff --git a/Engine/Source/Runtime/RHI/Public/RHIValidation.h b/Engine/Source/Runtime/RHI/Public/RHIValidation.h
--- a/Engine/Source/Runtime/RHI/Public/RHIValidation.h
+++ b/Engine/Source/Runtime/RHI/Public/RHIValidation.h
@@ -1113,6 +1113,11 @@ public:
 		RHI->RHICheckViewportHDRStatus(Viewport);
 	}
 
+	virtual void RHIHandleDisplayChange() override final
+	{
+		RHI->RHIHandleDisplayChange();
+	}
+
 	//  must be called from the main thread.
 	// FlushType: Thread safe
 	virtual void RHITick(float DeltaTime) override final
diff --git a/Engine/Source/Runtime/Slate/Private/Framework/Application/SlateApplication.cpp b/Engine/Source/Runtime/Slate/Private/Framework/Application/SlateApplication.cpp
--- a/Engine/Source/Runtime/Slate/Private/Framework/Application/SlateApplication.cpp
+++ b/Engine/Source/Runtime/Slate/Private/Framework/Application/SlateApplication.cpp
@@ -6974,6 +6974,8 @@ void FSlateApplication::OnVirtualDesktopSizeChanged(const FDisplayMetrics& NewDi
 		VirtualDisplayRect.Top,
 		VirtualDisplayRect.Right,
 		VirtualDisplayRect.Bottom);
+
+	Renderer->OnVirtualDesktopSizeChanged(NewDisplayMetric);
 }
 
 
diff --git a/Engine/Source/Runtime/SlateCore/Public/Rendering/SlateRenderer.h b/Engine/Source/Runtime/SlateCore/Public/Rendering/SlateRenderer.h
--- a/Engine/Source/Runtime/SlateCore/Public/Rendering/SlateRenderer.h
+++ b/Engine/Source/Runtime/SlateCore/Public/Rendering/SlateRenderer.h
@@ -572,6 +572,8 @@ public:
 	virtual void AddWidgetRendererUpdate(const struct FRenderThreadUpdateContext& Context, bool bDeferredRenderTargetUpdate) {}
 
 	virtual EPixelFormat GetSlateRecommendedColorFormat() { return PF_B8G8R8A8; }
+
+	virtual void OnVirtualDesktopSizeChanged(const FDisplayMetrics& NewDisplayMetric) {}
 private:
 
 	// Non-copyable
diff --git a/Engine/Source/Runtime/SlateRHIRenderer/Private/SlateRHIRenderer.cpp b/Engine/Source/Runtime/SlateRHIRenderer/Private/SlateRHIRenderer.cpp
--- a/Engine/Source/Runtime/SlateRHIRenderer/Private/SlateRHIRenderer.cpp
+++ b/Engine/Source/Runtime/SlateRHIRenderer/Private/SlateRHIRenderer.cpp
@@ -201,6 +201,7 @@ FSlateRHIRenderer::FSlateRHIRenderer(TSharedRef<FSlateFontServices> InSlateFontS
 	, EnqueuedWindowDrawBuffer(NULL)
 	, FreeBufferIndex(0)
 	, FastPathRenderingDataCleanupList(nullptr)
+	, bUpdateHDRDisplayInformation(false)
 	, CurrentSceneIndex(-1)
 	, ResourceVersion(0)
 {
@@ -499,6 +500,14 @@ void FSlateRHIRenderer::ConditionalResizeViewport(FViewportInfo* ViewInfo, uint3
 	}
 }
 
+void FSlateRHIRenderer::OnVirtualDesktopSizeChanged(const FDisplayMetrics& NewDisplayMetric)
+{
+	// Defer the update to as we need to call FlushRenderingCommands() before sending the event to the RHI. 
+	// FlushRenderingCommands -> FRenderCommandFence::IsFenceComplete -> CheckRenderingThreadHealth -> FPlatformApplicationMisc::PumpMessages
+	// The Display change event is not been consumed yet, and we do BroadcastDisplayMetricsChanged -> OnVirtualDesktopSizeChanged again
+	bUpdateHDRDisplayInformation = true;
+}
+
 void FSlateRHIRenderer::UpdateFullscreenState(const TSharedRef<SWindow> Window, uint32 OverrideResX, uint32 OverrideResY)
 {
 	FViewportInfo* ViewInfo = WindowToViewportInfo.FindRef(&Window.Get());
@@ -1583,6 +1592,12 @@ void FSlateRHIRenderer::DrawWindows_Private(FSlateDrawBuffer& WindowDrawBuffer)
 {
 	checkSlow(IsThreadSafeForSlateRendering());
 
+	if (bUpdateHDRDisplayInformation)
+	{
+		FlushRenderingCommands();
+		RHIHandleDisplayChange();
+		bUpdateHDRDisplayInformation = false;
+	}
 
 	FSlateRHIRenderingPolicy* Policy = RenderingPolicy.Get();
 	ENQUEUE_RENDER_COMMAND(SlateBeginDrawingWindowsCommand)(
diff --git a/Engine/Source/Runtime/SlateRHIRenderer/Private/SlateRHIRenderer.h b/Engine/Source/Runtime/SlateRHIRenderer/Private/SlateRHIRenderer.h
--- a/Engine/Source/Runtime/SlateRHIRenderer/Private/SlateRHIRenderer.h
+++ b/Engine/Source/Runtime/SlateRHIRenderer/Private/SlateRHIRenderer.h
@@ -269,6 +269,8 @@ private:
 	 */
 	void CleanUpdatableTextures();
 
+	virtual void OnVirtualDesktopSizeChanged(const FDisplayMetrics& NewDisplayMetric);
+
 private:
 	/** A mapping of SWindows to their RHI implementation */
 	TMap< const SWindow*, FViewportInfo*> WindowToViewportInfo;
@@ -302,6 +304,7 @@ private:
 
 	bool bIsStandaloneStereoOnlyDevice;
 	bool bTakingAScreenShot;
+	bool bUpdateHDRDisplayInformation;
 	FIntRect ScreenshotRect;
 	FViewportInfo* ScreenshotViewportInfo;
 	TArray<FColor>* OutScreenshotData;
