build fix

diff --git a/Engine/Source/Runtime/D3D12RHI/Private/Windows/WindowsD3D12Device.cpp b/Engine/Source/Runtime/D3D12RHI/Private/Windows/WindowsD3D12Device.cpp
--- a/Engine/Source/Runtime/D3D12RHI/Private/Windows/WindowsD3D12Device.cpp
+++ b/Engine/Source/Runtime/D3D12RHI/Private/Windows/WindowsD3D12Device.cpp
@@ -601,6 +601,7 @@ bool FD3D12DynamicRHI::SetupDisplayHDRMetaData()
 	return bSupportsHDROutput;
 }
 
+#if PLATFORM_WINDOWS
 extern void HDRSettingChangedSinkCallback();
 void FD3D12DynamicRHI::RHIHandleDisplayChange()
 {
@@ -609,6 +610,7 @@ void FD3D12DynamicRHI::RHIHandleDisplayChange()
 	// make sure CVars are being updated properly
 	HDRSettingChangedSinkCallback();
 }
+#endif
 
 static bool IsAdapterBlocked(FD3D12Adapter* InAdapter)
 {
