Add support for full swap during RHIUpdateRHIResources in d3d12 rhi. It before always assumed that either target was empty buffer

diff --git a/Engine/Source/Runtime/D3D12RHI/Private/D3D12PoolAllocator.cpp b/Engine/Source/Runtime/D3D12RHI/Private/D3D12PoolAllocator.cpp
--- a/Engine/Source/Runtime/D3D12RHI/Private/D3D12PoolAllocator.cpp
+++ b/Engine/Source/Runtime/D3D12RHI/Private/D3D12PoolAllocator.cpp
@@ -721,6 +721,51 @@ void FD3D12PoolAllocator::TransferOwnership(FD3D12ResourceLocation& InSource, FD
 }
 
 
+void FD3D12PoolAllocator::Swap(FD3D12ResourceLocation& InLHS, FD3D12ResourceLocation& InRHS)
+{
+	FScopeLock Lock(&CS);
+
+	check(IsOwner(InLHS) && IsOwner(InRHS));
+
+	FRHIPoolAllocationData& LHSPoolData = InLHS.GetPoolAllocatorPrivateData().PoolData;
+	FRHIPoolAllocationData& RHSPoolData = InRHS.GetPoolAllocatorPrivateData().PoolData;
+
+	// If locked then assume the block is currently being defragged and then the frame fenced operation
+	// also needs to be updated
+	if (LHSPoolData.IsLocked() || RHSPoolData.IsLocked())
+	{
+		int32 bFoundRequired = 0;
+		bFoundRequired += LHSPoolData.IsLocked() ? 1 : 0;
+		bFoundRequired += RHSPoolData.IsLocked() ? 1 : 0;
+		for (FrameFencedAllocationData& Operation : FrameFencedOperations)
+		{
+			if (Operation.AllocationData == &LHSPoolData || Operation.AllocationData == &RHSPoolData)
+			{
+				check(Operation.Operation == FrameFencedAllocationData::EOperation::Unlock);
+				Operation.AllocationData = Operation.AllocationData == &LHSPoolData ? &RHSPoolData : &LHSPoolData;
+				bFoundRequired--;
+				if (bFoundRequired == 0)
+				{
+					break;
+				}
+			}
+		}
+		check(bFoundRequired == 0);
+	}
+
+	// Perform swap with data & locking state but keep ownership
+	FRHIPoolAllocationData TmpAllocationData = LHSPoolData;
+	bool TmpIsLocked = LHSPoolData.IsLocked();
+	FRHIPoolResource* TmpOwnerLHS = LHSPoolData.GetOwner();
+	FRHIPoolResource* TmpOwnerRHS = RHSPoolData.GetOwner();
+
+	LHSPoolData.MoveFrom(RHSPoolData, RHSPoolData.IsLocked());
+	LHSPoolData.SetOwner(TmpOwnerLHS);
+	RHSPoolData.MoveFrom(TmpAllocationData, TmpIsLocked);
+	RHSPoolData.SetOwner(TmpOwnerRHS);
+}
+
+
 FD3D12Resource* FD3D12PoolAllocator::GetBackingResource(FD3D12ResourceLocation& InResourceLocation) const
 {
 	check(IsOwner(InResourceLocation));
diff --git a/Engine/Source/Runtime/D3D12RHI/Private/D3D12PoolAllocator.h b/Engine/Source/Runtime/D3D12RHI/Private/D3D12PoolAllocator.h
--- a/Engine/Source/Runtime/D3D12RHI/Private/D3D12PoolAllocator.h
+++ b/Engine/Source/Runtime/D3D12RHI/Private/D3D12PoolAllocator.h
@@ -144,6 +144,7 @@ public:
 	void FlushPendingCopyOps(FD3D12CommandContext& InCommandContext);
 
 	void TransferOwnership(FD3D12ResourceLocation& InSource, FD3D12ResourceLocation& InDest);
+	void Swap(FD3D12ResourceLocation& InLHS, FD3D12ResourceLocation& InRHS);
 	bool IsOwner(FD3D12ResourceLocation& ResourceLocation) const { return ResourceLocation.GetPoolAllocator() == this; }
 
 	EResourceAllocationStrategy GetAllocationStrategy() const { return AllocationStrategy; }
diff --git a/Engine/Source/Runtime/D3D12RHI/Private/D3D12Resources.cpp b/Engine/Source/Runtime/D3D12RHI/Private/D3D12Resources.cpp
--- a/Engine/Source/Runtime/D3D12RHI/Private/D3D12Resources.cpp
+++ b/Engine/Source/Runtime/D3D12RHI/Private/D3D12Resources.cpp
@@ -576,8 +576,6 @@ void FD3D12ResourceLocation::Swap(FD3D12ResourceLocation& Other)
 
 	if (Other.GetAllocatorType() == FD3D12ResourceLocation::AT_Pool)
 	{
-		check(GetAllocatorType() != FD3D12ResourceLocation::AT_Pool);
-
 		// Swap all members except the pool data because that needs to happen while the pool is taken
 		FD3D12DeviceChild::Swap(Other);
 
@@ -593,11 +591,23 @@ void FD3D12ResourceLocation::Swap(FD3D12ResourceLocation& Other)
 		::Swap(Size, Other.Size);
 		::Swap(bTransient, Other.bTransient);
 				
-		// We know this allocation is empty so can use transfer instead of full swap
-		Other.GetPoolAllocator()->TransferOwnership(Other, *this);
+		// Also pool allocated, then perform full swap with lock
+		if (GetAllocatorType() == EAllocatorType::AT_Pool)
+		{
+			check(PoolAllocator == Other.PoolAllocator);
+			GetPoolAllocator()->Swap(*this, Other);
+		}
+		else
+		{
+			// Assume unallocated
+			check(GetAllocatorType() == FD3D12ResourceLocation::AT_Unknown);
+
+			// We know this allocation is empty so can use transfer instead of full swap
+			Other.GetPoolAllocator()->TransferOwnership(Other, *this);
 
-		::Swap(AllocatorType, Other.AllocatorType);
-		::Swap(PoolAllocator, Other.PoolAllocator);
+			::Swap(AllocatorType, Other.AllocatorType);
+			::Swap(PoolAllocator, Other.PoolAllocator);
+		}
 	}
 	else
 	{
