#UE Allow for a graceful shutdown when using D3D11, but a FD3D12Adapter is constructed

diff --git a/Engine/Source/Runtime/D3D12RHI/Private/D3D12Adapter.cpp b/Engine/Source/Runtime/D3D12RHI/Private/D3D12Adapter.cpp
--- a/Engine/Source/Runtime/D3D12RHI/Private/D3D12Adapter.cpp
+++ b/Engine/Source/Runtime/D3D12RHI/Private/D3D12Adapter.cpp
@@ -1336,16 +1336,22 @@ FD3D12Adapter::~FD3D12Adapter()
 
 	for (uint32 GPUIndex : FRHIGPUMask::All())
 	{
-		UploadHeapAllocator[GPUIndex]->Destroy();
-		delete(UploadHeapAllocator[GPUIndex]);
-		UploadHeapAllocator[GPUIndex] = nullptr;
+		if (UploadHeapAllocator[GPUIndex])
+		{
+			UploadHeapAllocator[GPUIndex]->Destroy();
+			delete(UploadHeapAllocator[GPUIndex]);
+			UploadHeapAllocator[GPUIndex] = nullptr;
+		}
 	}
 
 	FrameFence = {};
 
 	TransientMemoryCache.Reset();
 
-	PipelineStateCache.Close();
+	if (RootDevice)
+	{
+		PipelineStateCache.Close();
+	}
 	RootSignatureManager.Destroy();
 
 	DrawIndirectCommandSignature.SafeRelease();
