Fixing heap mismatch when running a UAV clear on structured buffers. This may need a revisit later on to address the mid-frame descriptor allocation.

diff --git a/Engine/Source/Runtime/D3D12RHI/Private/D3D12Descriptors.cpp b/Engine/Source/Runtime/D3D12RHI/Private/D3D12Descriptors.cpp
--- a/Engine/Source/Runtime/D3D12RHI/Private/D3D12Descriptors.cpp
+++ b/Engine/Source/Runtime/D3D12RHI/Private/D3D12Descriptors.cpp
@@ -179,6 +179,18 @@ FD3D12DescriptorHeap* FD3D12BindlessDescriptorManager::GetHeapForType(ERHIDescri
 	return nullptr;
 }
 
+D3D12_GPU_DESCRIPTOR_HANDLE FD3D12BindlessDescriptorManager::GetGpuHandle(FRHIDescriptorHandle InHandle) const
+{
+	for (const FD3D12DescriptorManager& Manager : Managers)
+	{
+		if (Manager.HandlesAllocation(InHandle.GetType()))
+		{
+			return Manager.GetHeap()->GetGPUSlotHandle(InHandle.GetIndex());
+		}
+	}
+	return D3D12_GPU_DESCRIPTOR_HANDLE{};
+}
+
 ////////////////////////////////////////////////////////////////////////////////////////////////////////////////////////
 // FD3D12OnlineDescriptorManager
 
diff --git a/Engine/Source/Runtime/D3D12RHI/Private/D3D12Descriptors.h b/Engine/Source/Runtime/D3D12RHI/Private/D3D12Descriptors.h
--- a/Engine/Source/Runtime/D3D12RHI/Private/D3D12Descriptors.h
+++ b/Engine/Source/Runtime/D3D12RHI/Private/D3D12Descriptors.h
@@ -70,7 +70,8 @@ public:
 
 	void UpdateImmediately(FRHIDescriptorHandle InHandle, D3D12_CPU_DESCRIPTOR_HANDLE InSourceCpuHandle);
 
-	inline FD3D12DescriptorHeap* GetHeap() { return Heap.GetReference(); }
+	inline       FD3D12DescriptorHeap* GetHeap()       { return Heap.GetReference(); }
+	inline const FD3D12DescriptorHeap* GetHeap() const { return Heap.GetReference(); }
 
 	inline bool HandlesAllocationWithFlags(ERHIDescriptorHeapType InHeapType, D3D12_DESCRIPTOR_HEAP_FLAGS InHeapFlags) const
 	{
@@ -104,6 +105,8 @@ public:
 
 	FD3D12DescriptorHeap* GetHeapForType(ERHIDescriptorHeapType InType);
 
+	D3D12_GPU_DESCRIPTOR_HANDLE GetGpuHandle(FRHIDescriptorHandle InHandle) const;
+
 private:
 	TArray<FD3D12DescriptorManager> Managers;
 };
diff --git a/Engine/Source/Runtime/D3D12RHI/Private/D3D12UAV.cpp b/Engine/Source/Runtime/D3D12RHI/Private/D3D12UAV.cpp
--- a/Engine/Source/Runtime/D3D12RHI/Private/D3D12UAV.cpp
+++ b/Engine/Source/Runtime/D3D12RHI/Private/D3D12UAV.cpp
@@ -277,18 +277,27 @@ void FD3D12CommandContext::ClearUAV(TRHICommandList_RecursiveHazardous<FD3D12Com
 				FD3D12ViewDescriptorHandle UAVHandle(ParentDevice, ERHIDescriptorHeapType::Standard);
 				UAVHandle.CreateView(R32UAVDesc, Resource, nullptr, ED3D12DescriptorCreateReason::InitialCreate);
 
-				// Check if the view heap is full and needs to rollover.
-				if (!Context.StateCache.GetDescriptorCache()->GetCurrentViewHeap()->CanReserveSlots(1))
+				const D3D12_CPU_DESCRIPTOR_HANDLE CPUHandle = UAVHandle.GetOfflineCpuHandle();
+
+				D3D12_GPU_DESCRIPTOR_HANDLE GPUHandle;
+				if (!Context.StateCache.GetDescriptorCache()->IsViewHeapOverridden())
 				{
-					Context.StateCache.GetDescriptorCache()->GetCurrentViewHeap()->RollOver();
-				}
+					// Check if the view heap is full and needs to rollover.
+					if (!Context.StateCache.GetDescriptorCache()->GetCurrentViewHeap()->CanReserveSlots(1))
+					{
+						Context.StateCache.GetDescriptorCache()->GetCurrentViewHeap()->RollOver();
+					}
 
-				uint32 ReservedSlot = Context.StateCache.GetDescriptorCache()->GetCurrentViewHeap()->ReserveSlots(1);
-				D3D12_CPU_DESCRIPTOR_HANDLE CPUHandle = UAVHandle.GetOfflineCpuHandle();
-				D3D12_CPU_DESCRIPTOR_HANDLE DestSlot = Context.StateCache.GetDescriptorCache()->GetCurrentViewHeap()->GetCPUSlotHandle(ReservedSlot);
-				D3D12_GPU_DESCRIPTOR_HANDLE GPUHandle = Context.StateCache.GetDescriptorCache()->GetCurrentViewHeap()->GetGPUSlotHandle(ReservedSlot);
+					uint32 ReservedSlot = Context.StateCache.GetDescriptorCache()->GetCurrentViewHeap()->ReserveSlots(1);
+					D3D12_CPU_DESCRIPTOR_HANDLE DestSlot = Context.StateCache.GetDescriptorCache()->GetCurrentViewHeap()->GetCPUSlotHandle(ReservedSlot);
+					GPUHandle = Context.StateCache.GetDescriptorCache()->GetCurrentViewHeap()->GetGPUSlotHandle(ReservedSlot);
 
-				Device->CopyDescriptorsSimple(1, DestSlot, CPUHandle, D3D12_DESCRIPTOR_HEAP_TYPE_CBV_SRV_UAV);
+					Device->CopyDescriptorsSimple(1, DestSlot, CPUHandle, D3D12_DESCRIPTOR_HEAP_TYPE_CBV_SRV_UAV);
+				}
+				else
+				{
+					GPUHandle = ParentDevice->GetBindlessDescriptorManager().GetGpuHandle(UAVHandle.GetBindlessHandle());
+				}
 
 				Context.TransitionResource(UnorderedAccessView, D3D12_RESOURCE_STATE_UNORDERED_ACCESS);
 
