Suppress deprecation warnings around RHIRayTraceOcclusion and RHIRayTraceIntersection calls

diff --git a/Engine/Source/Runtime/D3D12RHI/Private/D3D12CommandContext.h b/Engine/Source/Runtime/D3D12RHI/Private/D3D12CommandContext.h
--- a/Engine/Source/Runtime/D3D12RHI/Private/D3D12CommandContext.h
+++ b/Engine/Source/Runtime/D3D12RHI/Private/D3D12CommandContext.h
@@ -1044,7 +1044,9 @@ public:
 		FRHIUnorderedAccessView* Output,
 		uint32 NumRays) final override
 	{
+		PRAGMA_DISABLE_DEPRECATION_WARNINGS
 		ContextRedirect(RHIRayTraceOcclusion(Scene, Rays, Output, NumRays));
+		PRAGMA_ENABLE_DEPRECATION_WARNINGS
 	}
 
 	virtual void RHIRayTraceIntersection(FRHIRayTracingScene* Scene,
@@ -1052,7 +1054,9 @@ public:
 		FRHIUnorderedAccessView* Output,
 		uint32 NumRays) final override
 	{
+		PRAGMA_DISABLE_DEPRECATION_WARNINGS
 		ContextRedirect(RHIRayTraceIntersection(Scene, Rays, Output, NumRays));
+		PRAGMA_ENABLE_DEPRECATION_WARNINGS
 	}
 
 	virtual void RHIRayTraceDispatch(FRHIRayTracingPipelineState* RayTracingPipelineState, FRHIRayTracingShader* RayGenShader,
