Delay opening a new command list until there's work that needs to be recorded to it.

This prevents submitting empty command lists, which wastes CPU and GPU time.

Submitting on behalf of Luke.Thatcher.

#rnx

diff --git a/Engine/Source/Runtime/D3D12RHI/Private/D3D12CommandContext.cpp b/Engine/Source/Runtime/D3D12RHI/Private/D3D12CommandContext.cpp
--- a/Engine/Source/Runtime/D3D12RHI/Private/D3D12CommandContext.cpp
+++ b/Engine/Source/Runtime/D3D12RHI/Private/D3D12CommandContext.cpp
@@ -111,8 +111,6 @@ FD3D12CommandContext::FD3D12CommandContext(FD3D12Device* InParent, ED3D12QueueTy
 	, CSConstantBuffer(InParent, ConstantsAllocator)
 {
 	StaticUniformBuffers.AddZeroed(FUniformBufferStaticSlotRegistry::Get().GetSlotCount());
-
-	OpenCommandList();
 	ClearState();
 }
 
@@ -283,23 +281,31 @@ FD3D12ContextCommon::FD3D12ContextCommon(FD3D12Device* Device, ED3D12QueueType Q
 
 void FD3D12ContextCommon::WaitSyncPoint(FD3D12SyncPoint* SyncPoint)
 {
-	checkf(!CommandList, TEXT("Command list must be closed before waiting on a sync point."));
+	if (IsOpen())
+	{
+		CloseCommandList();
+	}
+
 	GetPayload(EPhase::Wait)->SyncPointsToWait.Add(SyncPoint);
 }
 
 void FD3D12ContextCommon::SignalSyncPoint(FD3D12SyncPoint* SyncPoint)
 {
-	checkf(!CommandList, TEXT("Command list must be closed before signaling a sync point."));
-
-	checkf(SyncPoint->GetType() == ED3D12SyncPointType::GPUOnly || EnumHasAllFlags(State, EState::AllowFinalize),
-		TEXT("Cannot signal a CPU-accessible sync point if the previous command list did not resolve queries when it was closed."));
+	if (IsOpen())
+	{
+		CloseCommandList();
+	}
 
 	GetPayload(EPhase::Signal)->SyncPointsToSignal.Add(SyncPoint);
 }
 
 void FD3D12ContextCommon::SignalManualFence(ID3D12Fence* Fence, uint64 Value)
 {
-	checkf(!CommandList, TEXT("Command list must be closed before signaling a manual fence."));
+	if (IsOpen())
+	{
+		CloseCommandList();
+	}
+
 	GetPayload(EPhase::Signal)->FencesToSignal.Emplace(Fence, Value);
 }
 
@@ -311,7 +317,6 @@ void FD3D12ContextCommon::WaitManualFence(ID3D12Fence* Fence, uint64 Value)
 
 FD3D12QueryLocation FD3D12ContextCommon::AllocateQuery(ED3D12QueryType Type, uint64* Target)
 {
-	checkf(CommandList, TEXT("Command list must be open to insert a timestamp."));
 	switch (Type)
 	{
 	default:
@@ -320,10 +325,10 @@ FD3D12QueryLocation FD3D12ContextCommon::AllocateQuery(ED3D12QueryType Type, uin
 
 	case ED3D12QueryType::AdjustedRaw:
 	case ED3D12QueryType::AdjustedMicroseconds:
-		return CommandList->State.TimestampQueries.Emplace_GetRef(TimestampQueries.Allocate(Type, Target));
+		return TimestampQueries.Allocate(Type, Target);
 
 	case ED3D12QueryType::Occlusion:
-		return CommandList->State.OcclusionQueries.Emplace_GetRef(OcclusionQueries.Allocate(Type, Target));
+		return OcclusionQueries.Allocate(Type, Target);
 	}
 }
 
@@ -341,19 +346,15 @@ FD3D12QueryLocation FD3D12ContextCommon::InsertTimestamp(ED3D12Units Units, uint
 	}
 
 	FD3D12QueryLocation Location = AllocateQuery(Type, Target);
-	CommandList->WriteTimestamp(Location);
+	EndQuery(Location);
+
 	return Location;
 }
 
 void FD3D12ContextCommon::OpenCommandList()
 {
 	LLM_SCOPE_BYNAME(TEXT("RHIMisc/OpenCommandList"));
-
-#if DO_CHECK
-	checkf(!CommandList, TEXT("Command list is already open."));
-	checkf(EnumHasAllFlags(State, EState::AllowOpen), TEXT("Cannot open a new command list until the context has been finalized."));
-	State = EState::AllowClose;
-#endif
+	checkf(!IsOpen(), TEXT("Command list is already open."));
 
 	if (CommandAllocator == nullptr)
 	{
@@ -363,6 +364,9 @@ void FD3D12ContextCommon::OpenCommandList()
 
 	// Get a new command list
 	CommandList = Device->ObtainCommandList(CommandAllocator, &TimestampQueries);
+	GetPayload(EPhase::Execute)->CommandListsToExecute.Add(CommandList);
+
+	check(ActiveQueries == 0);
 }
 
 void FD3D12CommandContext::OpenCommandList()
@@ -372,72 +376,47 @@ void FD3D12CommandContext::OpenCommandList()
 	// Notify the descriptor cache about the new command list
 	// This will set the descriptor cache's current heaps on the new command list.
 	StateCache.GetDescriptorCache()->OpenCommandList();
-
-	// Go through the state and find bits that differ from command list defaults.
-	// Mark state as dirty so next time ApplyState is called, it will set all state on this new command list
-	StateCache.DirtyStateForNewCommandList();
-
-	check(ActiveQueries == 0);
 }
 
-void FD3D12ContextCommon::CloseCommandList(bool bResolveQueries)
+void FD3D12ContextCommon::CloseCommandList()
 {
-#if DO_CHECK
-	checkf(EnumHasAllFlags(State, EState::AllowClose), TEXT("Command list is already closed."));
-	// Finalize is only allowed when queries have been resolved.
-	State = bResolveQueries
-		? EState::AllowOpen | EState::AllowFinalize
-		: EState::AllowOpen;
-#endif
+	checkf(IsOpen(), TEXT("Command list is not open."));
+	checkf(Payloads.Num() && CurrentPhase == EPhase::Execute, TEXT("Expected the current payload to be in the execute phase."));
+	
+	checkf(ActiveQueries == 0, TEXT("All queries must be completed before the command list is closed."));
+
+	FD3D12Payload* Payload = GetPayload(EPhase::Execute);
 
 	// Do this before we insert the final timestamp to ensure we're timing all the work on the command list.
 	FlushResourceBarriers();
 
-	// Don't close the command list if we're resolving queries.
-	// The submission thread will close it after appending a ResolveQueryData() command.
-	if (!bResolveQueries)
-	{
-		CommandList->Close();
-	}
-
-	GetPayload(EPhase::Execute)->CommandListsToExecute.Add(CommandList);
-	GetPayload(EPhase::Execute)->bResolveQueries = bResolveQueries;
+	CommandList->Close();
 	CommandList = nullptr;
 
-	if (bResolveQueries)
-	{
-		TimestampQueries.CloseAndReset(GetPayload(EPhase::Execute)->QueryRanges);
-		OcclusionQueries.CloseAndReset(GetPayload(EPhase::Execute)->QueryRanges);
-
-		// When resolving queries, we're leaving the command list open,
-		// which means we can't reuse this allocator for the next command list.
-		check(CommandAllocator);
-		GetPayload(EPhase::Signal)->AllocatorsToRelease.Add(CommandAllocator);
-		CommandAllocator = nullptr;
-	}
+	TimestampQueries.CloseAndReset(Payload->QueryRanges);
+	OcclusionQueries.CloseAndReset(Payload->QueryRanges);
 }
 
-void FD3D12CommandContext::CloseCommandList(bool bResolveQueries)
+void FD3D12CommandContext::CloseCommandList()
 {
-	checkf(ActiveQueries == 0, TEXT("All queries must be completed before the command list is closed."));
-
 	StateCache.GetDescriptorCache()->CloseCommandList();
-
-	FD3D12ContextCommon::CloseCommandList(bResolveQueries);
+	FD3D12ContextCommon::CloseCommandList();
+	// Mark state as dirty now, because ApplyState may be called before OpenCommandList(), and it needs to know that the state has
+	// become invalid, so it can set it up again (which opens a new command list if necessary).
+	StateCache.DirtyStateForNewCommandList();
 }
 
 void FD3D12ContextCommon::Finalize(TArray<FD3D12Payload*>& OutPayloads)
 {
-#if DO_CHECK
-	checkf(!CommandList, TEXT("Command list must be closed before finalizing."));
-	checkf(EnumHasAllFlags(State, EState::AllowFinalize), TEXT("The last call to CloseCommandList() did not resolve queries."));
-	State = EState::AllowOpen | EState::AllowFinalize;
-#endif
+	if (IsOpen())
+	{
+		CloseCommandList();
+	}
 
 	// Collect the context's batch of sync points to wait/signal
 	if (BatchedSyncPoints.ToWait.Num())
 	{
-		FD3D12Payload* Payload = Payloads.Num() 
+		FD3D12Payload* Payload = Payloads.Num()
 			? Payloads[0]
 			: GetPayload(EPhase::Wait);
 
@@ -533,7 +512,6 @@ FD3D12CopyScope::~FD3D12CopyScope()
 {
 	checkf(bSyncPointRetrieved, TEXT("The copy sync point must be retrieved before the end of the scope."));
 
-	Context.CloseCommandList(true);
 	Context.SignalSyncPoint(SyncPoint);
 
 	TArray<FD3D12Payload*> Payloads;
@@ -561,7 +539,7 @@ void FD3D12ContextCommon::FlushCommands(ED3D12FlushFlags FlushFlags)
 
 	if (IsOpen())
 	{
-		CloseCommandList(true);
+		CloseCommandList();
 	}
 
 	FD3D12SyncPointRef SyncPoint;
@@ -597,21 +575,15 @@ void FD3D12ContextCommon::FlushCommands(ED3D12FlushFlags FlushFlags)
 		FRenderThreadIdleScope Scope(ERenderThreadIdleTypes::WaitingForAllOtherSleep);
 		SubmissionEvent->Wait();
 	}
-
-	if (!EnumHasAnyFlags(FlushFlags, ED3D12FlushFlags::NoOpen))
-	{
-		OpenCommandList();
-	}
 }
 
-void FD3D12CommandContext::ConditionalSplitCommandList()
+void FD3D12ContextCommon::ConditionalSplitCommandList()
 {
 	// Start a new command list if the total number of commands exceeds the threshold. Too many commands in a single command list can cause TDRs.
-	if (IsOpen() && ActiveQueries == 0 && GD3D12MaxCommandsPerCommandList > 0 && GetNumCommands() > (uint32)GD3D12MaxCommandsPerCommandList)
+	if (IsOpen() && ActiveQueries == 0 && GD3D12MaxCommandsPerCommandList > 0 && CommandList->State.NumCommands > (uint32)GD3D12MaxCommandsPerCommandList)
 	{
-		UE_LOG(LogD3D12RHI, Verbose, TEXT("Splitting command lists because too many commands have been enqueued already (%d commands)"), GetNumCommands());
-		CloseCommandList(false);
-		OpenCommandList();
+		UE_LOG(LogD3D12RHI, Verbose, TEXT("Splitting command lists because too many commands have been enqueued already (%d commands)"), CommandList->State.NumCommands);
+		CloseCommandList();
 	}
 }
 
@@ -890,10 +862,6 @@ void FD3D12CommandContextRedirector::RHITransferResources(const TArrayView<const
 			CombinedMask |= WaitMask.GetValue();
 		}
 
-		// Close all relevant contexts to allow signalling/waiting of sync points
-		for (uint32 GPUIndex : CombinedMask)
-			PhysicalContexts[GPUIndex]->CloseCommandList(false);
-
 		// Signal a sync point on each source GPU
 		TStaticArray<FD3D12SyncPointRef, MAX_NUM_GPUS> SyncPoints;
 		for (uint32 GPUIndex : SignalMask)
@@ -914,10 +882,6 @@ void FD3D12CommandContextRedirector::RHITransferResources(const TArrayView<const
 			}
 		}
 
-		// Re-open contexts
-		for (uint32 GPUIndex : CombinedMask)
-			PhysicalContexts[GPUIndex]->OpenCommandList();
-
 		return SyncPoints;
 	};
 
@@ -988,9 +952,7 @@ void FD3D12CommandContextRedirector::RHITransferResources(const TArrayView<const
 			{
 				FD3D12SyncPoint* SyncPoint = static_cast<FD3D12SyncPoint*>(FenceData->SyncPoints[GPUIndex]);
 
-				PhysicalContexts[GPUIndex]->CloseCommandList(false);
 				PhysicalContexts[GPUIndex]->WaitSyncPoint(SyncPoint);
-				PhysicalContexts[GPUIndex]->OpenCommandList();
 
 				SyncPoint->Release();
 			}
@@ -1117,9 +1079,7 @@ void FD3D12CommandContextRedirector::RHITransferResourceSignal(const TArrayView<
 		FD3D12SyncPointRef SyncPoint = FD3D12SyncPoint::Create(ED3D12SyncPointType::GPUOnly);
 		SyncPoint->AddRef();
 
-		PhysicalContexts[SrcGPUIndex]->CloseCommandList(false);
 		PhysicalContexts[SrcGPUIndex]->SignalSyncPoint(SyncPoint);
-		PhysicalContexts[SrcGPUIndex]->OpenCommandList();
 
 		FTransferResourceFenceData* FenceData = FenceDatas[FenceIndex++];
 		FenceData->Mask = FRHIGPUMask::FromIndex(SrcGPUIndex);
@@ -1138,10 +1098,6 @@ void FD3D12CommandContextRedirector::RHITransferResourceWait(const TArrayView<FT
 			: FenceDatas[Index]->Mask | AllMasks;
 	}
 
-	// Close all relevant contexts to allow signalling/waiting of sync points
-	for (uint32 GPUIndex : AllMasks)
-		PhysicalContexts[GPUIndex]->CloseCommandList(false);
-
 	for (FTransferResourceFenceData* FenceData : FenceDatas)
 	{
 		// Wait for sync points
@@ -1169,9 +1125,5 @@ void FD3D12CommandContextRedirector::RHITransferResourceWait(const TArrayView<FT
 		delete FenceData;
 	}
 
-	// Re-open contexts
-	for (uint32 GPUIndex : AllMasks)
-		PhysicalContexts[GPUIndex]->OpenCommandList();
-
 #endif // WITH_MGPU
 }
diff --git a/Engine/Source/Runtime/D3D12RHI/Private/D3D12CommandContext.h b/Engine/Source/Runtime/D3D12RHI/Private/D3D12CommandContext.h
--- a/Engine/Source/Runtime/D3D12RHI/Private/D3D12CommandContext.h
+++ b/Engine/Source/Runtime/D3D12RHI/Private/D3D12CommandContext.h
@@ -83,10 +83,7 @@ enum class ED3D12FlushFlags
 	WaitForSubmission = 1,
 
 	// Both the calling thread until the GPU has signaled completion of all dispatched work.
-	WaitForCompletion = 2,
-
-	// Don't reopen the command list after the flush.
-	NoOpen = 4
+	WaitForCompletion = 2
 };
 ENUM_CLASS_FLAGS(ED3D12FlushFlags)
 
@@ -104,9 +101,11 @@ protected:
 public:
 	virtual ~FD3D12ContextCommon() = default;
 
+protected:
 	virtual void OpenCommandList();
-	virtual void CloseCommandList(bool bResolveQueries);
+	virtual void CloseCommandList();
 
+public:
 	enum class EClearStateMode
 	{
 		TransientOnly,
@@ -152,11 +151,10 @@ public:
 	bool IsOpen() const { return CommandList != nullptr; }
 
 	// Returns unique identity that can be used to distinguish between command lists even after they were recycled.
-	uint64 GetCommandListID() const	{ return GetCommandList().State.CommandListID; }
+	uint64 GetCommandListID() { return GetCommandList().State.CommandListID; }
 
 	FD3D12SyncPoint* GetContextSyncPoint()
 	{
-		checkf(CommandList, TEXT("Command list is not open."));
 		if (!ContextSyncPoint)
 		{
 			ContextSyncPoint = FD3D12SyncPoint::Create(ED3D12SyncPointType::GPUAndCPU);
@@ -196,20 +194,14 @@ private:
 	// A sync point signaled when all payloads in this context have completed.
 	FD3D12SyncPointRef ContextSyncPoint;
 
-#if DO_CHECK
-	// Flags used to validate correct use of Open/Close/Finalize functions.
-	enum class EState
-	{
-		AllowOpen     = 1,
-		AllowClose    = 2,
-		AllowFinalize = 4
-	} State = EState::AllowOpen;
-	FRIEND_ENUM_CLASS_FLAGS(EState)
-#endif
-
-	FD3D12CommandList& GetCommandList() const
+	// Returns the current command list (or creates a new one if the command list was not open).
+	FD3D12CommandList& GetCommandList()
 	{
-		checkf(CommandList, TEXT("Command list is not open."));
+		if (!CommandList)
+		{
+			OpenCommandList();
+		}
+
 		return *CommandList;
 	}
 
@@ -230,96 +222,54 @@ protected:
 		return Payloads.Last();
 	}
 
-	uint32 GetNumCommands() const
-	{
-		return GetCommandList().State.NumCommands;
-	}
+	uint32 ActiveQueries = 0;
 
 public:
 	// Flushes any pending commands in this context to the GPU.
 	void FlushCommands(ED3D12FlushFlags FlushFlags = ED3D12FlushFlags::None);
 
-	//
-	// Wrapper type to prevent l-value use of the returned command list interfaces.
-	// A context's command list may be swapped out during recording. Users should access the command
-	// list via the context itself, to ensure they always have the correct command list instance.
-	//
-	template <typename T>
-	class TRValuePtr
-	{
-		friend FD3D12ContextCommon;
-
-		FD3D12CommandList& CommandList;
-		T* Ptr;
-
-		TRValuePtr(FD3D12CommandList& CommandList, T* Ptr)
-			: CommandList(CommandList)
-			, Ptr(Ptr)
-		{}
-
-	public:
-		TRValuePtr() = delete;
-
-		TRValuePtr(TRValuePtr const&) = delete;
-		TRValuePtr(TRValuePtr&&)      = delete;
-
-		TRValuePtr& operator= (TRValuePtr const&) = delete;
-		TRValuePtr& operator= (TRValuePtr&&)      = delete;
-
-		operator bool () const&& { return !!Ptr; }
-		bool operator!() const&& { return  !Ptr; }
-
-		// These accessor functions count useful work on command lists
-		T* operator ->  () && { CommandList.State.NumCommands++; return Ptr; }
-		T* Get          () && { CommandList.State.NumCommands++; return Ptr; }
-
-		T* GetNoRefCount() && { return Ptr; }
-	};
-
-private:
-	template <typename FInterfaceType>
-	auto BuildRValuePtr(TRefCountPtr<FInterfaceType> FD3D12CommandList::FInterfaces::* Member) const
-	{
-		FD3D12CommandList& CmdList = GetCommandList();
-		return TRValuePtr<FInterfaceType>(CmdList, CmdList.Interfaces.*Member);
-	}
+	// Closes the current command list if the number of enqueued commands exceeds
+	// the threshold defined by the "D3D12.MaxCommandsPerCommandList" cvar.
+	void ConditionalSplitCommandList();
 
-public:
-	auto BaseCommandList      () const { return BuildRValuePtr(&FD3D12CommandList::FInterfaces::CommandList         ); }
-	auto CopyCommandList      () const { return BuildRValuePtr(&FD3D12CommandList::FInterfaces::CopyCommandList     ); }
-	auto GraphicsCommandList  () const { return BuildRValuePtr(&FD3D12CommandList::FInterfaces::GraphicsCommandList ); }
+	auto BaseCommandList      () { return GetCommandList().BaseCommandList(); }
+	auto CopyCommandList      () { return GetCommandList().CopyCommandList(); }
+	auto GraphicsCommandList  () { return GetCommandList().GraphicsCommandList(); }
 #if D3D12_MAX_COMMANDLIST_INTERFACE >= 1					    
-	auto GraphicsCommandList1 () const { return BuildRValuePtr(&FD3D12CommandList::FInterfaces::GraphicsCommandList1); }
+	auto GraphicsCommandList1 () { return GetCommandList().GraphicsCommandList1(); }
 #endif														    
 #if D3D12_MAX_COMMANDLIST_INTERFACE >= 2					    
-	auto GraphicsCommandList2 () const { return BuildRValuePtr(&FD3D12CommandList::FInterfaces::GraphicsCommandList2); }
+	auto GraphicsCommandList2 () { return GetCommandList().GraphicsCommandList2(); }
 #endif														    
 #if D3D12_MAX_COMMANDLIST_INTERFACE >= 3					    
-	auto GraphicsCommandList3 () const { return BuildRValuePtr(&FD3D12CommandList::FInterfaces::GraphicsCommandList3); }
+	auto GraphicsCommandList3 () { return GetCommandList().GraphicsCommandList3(); }
 #endif														    
 #if D3D12_MAX_COMMANDLIST_INTERFACE >= 4					    
-	auto GraphicsCommandList4 () const { return BuildRValuePtr(&FD3D12CommandList::FInterfaces::GraphicsCommandList4); }
+	auto GraphicsCommandList4 () { return GetCommandList().GraphicsCommandList4(); }
 #endif														    
 #if D3D12_MAX_COMMANDLIST_INTERFACE >= 5					    
-	auto GraphicsCommandList5 () const { return BuildRValuePtr(&FD3D12CommandList::FInterfaces::GraphicsCommandList5); }
+	auto GraphicsCommandList5 () { return GetCommandList().GraphicsCommandList5(); }
 #endif														    
 #if D3D12_MAX_COMMANDLIST_INTERFACE >= 6					    
-	auto GraphicsCommandList6 () const { return BuildRValuePtr(&FD3D12CommandList::FInterfaces::GraphicsCommandList6); }
+	auto GraphicsCommandList6 () { return GetCommandList().GraphicsCommandList6(); }
 #endif														    
 #if D3D12_PLATFORM_SUPPORTS_ASSERTRESOURCESTATES			    
-	auto DebugCommandList     () const { return BuildRValuePtr(&FD3D12CommandList::FInterfaces::DebugCommandList    ); }
+	auto DebugCommandList     () { return GetCommandList().DebugCommandList(); }
 #endif
 #if D3D12_RHI_RAYTRACING
-	auto RayTracingCommandList() const { return BuildRValuePtr(&FD3D12CommandList::FInterfaces::GraphicsCommandList4); }
+	auto RayTracingCommandList() { return GetCommandList().RayTracingCommandList(); }
 #endif
 #if NV_AFTERMATH
-	auto AftermathHandle      () const { return GetCommandList().Interfaces.AftermathHandle; }
+	auto AftermathHandle      () { return GetCommandList().AftermathHandle(); }
 #endif
 
+	void BeginQuery(FD3D12QueryLocation const& Location) { GetCommandList().BeginQuery(Location); }
+	void EndQuery  (FD3D12QueryLocation const& Location) { GetCommandList().EndQuery  (Location); }
+
 #if ENABLE_RESIDENCY_MANAGEMENT
 	void UpdateResidency(TConstArrayView<FD3D12ResidencyHandle*> Handles) { GetCommandList().UpdateResidency(Handles); }
 	void UpdateResidency(FD3D12ResidencyHandle& Handle                  ) { GetCommandList().UpdateResidency({ &Handle }); }
-	void UpdateResidency(FD3D12ResidencyHandle* Handle                  ) { check(Handle); GetCommandList().UpdateResidency({ Handle }); }
+	void UpdateResidency(FD3D12ResidencyHandle* Handle                  ) { check(Handle  ); GetCommandList().UpdateResidency({ Handle }); }
 	void UpdateResidency(FD3D12Resource* Resource                       ) { check(Resource); GetCommandList().UpdateResidency({ &Resource->GetResidencyHandle() }); }
 #else
 	void UpdateResidency(TConstArrayView<FD3D12ResidencyHandle*> Handles) { }
@@ -357,10 +307,6 @@ private:
 	bool TransitionResource(FD3D12Resource* InResource, CResourceState& ResourceState_OnCommandList, uint32 InSubresourceIndex, D3D12_RESOURCE_STATES InBeforeState, D3D12_RESOURCE_STATES InAfterState, bool bInForceAfterState);
 };
 
-#if DO_CHECK
-	ENUM_CLASS_FLAGS(FD3D12ContextCommon::EState)
-#endif
-
 //
 // Context for the copy queue. Doesn't implement an RHI interface 
 // since the copy queue is not directly exposed to the renderer.
@@ -370,9 +316,7 @@ class FD3D12ContextCopy final : public FD3D12ContextCommon
 public:
 	FD3D12ContextCopy(FD3D12Device* Device)
 		: FD3D12ContextCommon(Device, ED3D12QueueType::Copy, false)
-	{
-		OpenCommandList();
-	}
+	{}
 };
 
 //
@@ -444,7 +388,7 @@ public:
 	}
 
 	virtual void OpenCommandList() override final;
-	virtual void CloseCommandList(bool bResolveQueries) override final;
+	virtual void CloseCommandList() override final;
 
 	virtual ERHIPipeline GetPipeline() const override
 	{
@@ -453,8 +397,6 @@ public:
 			: ERHIPipeline::AsyncCompute;
 	}
 
-	void ConditionalSplitCommandList();
-
 	virtual void ClearState(EClearStateMode ClearStateMode = EClearStateMode::All) override final;
 	void ConditionalClearShaderResource(FD3D12ResourceLocation* Resource);
 	void ClearAllShaderResources();
@@ -510,8 +452,6 @@ public:
 	virtual void FlushTextureCache() {};
 #endif
 
-	uint32 ActiveQueries = 0;
-
 	/** Constant buffers for Set*ShaderParameter calls. */
 	FD3D12ConstantBuffer VSConstantBuffer;
 	FD3D12ConstantBuffer MSConstantBuffer;
diff --git a/Engine/Source/Runtime/D3D12RHI/Private/D3D12CommandList.cpp b/Engine/Source/Runtime/D3D12RHI/Private/D3D12CommandList.cpp
--- a/Engine/Source/Runtime/D3D12RHI/Private/D3D12CommandList.cpp
+++ b/Engine/Source/Runtime/D3D12RHI/Private/D3D12CommandList.cpp
@@ -294,7 +294,7 @@ void FD3D12CommandList::WriteBeginTimestamp()
 {
 	if (!State.bBeginTimestampWritten && State.BeginTimestamp)
 	{
-		WriteTimestamp(State.BeginTimestamp);
+		EndQuery(State.BeginTimestamp);
 		State.bBeginTimestampWritten = true;
 	}
 }
@@ -303,15 +303,58 @@ void FD3D12CommandList::WriteEndTimestamp()
 {
 	if (!State.bEndTimestampWritten && State.EndTimestamp)
 	{
-		WriteTimestamp(State.EndTimestamp);
+		EndQuery(State.EndTimestamp);
 		State.bEndTimestampWritten = true;
 	}
 }
 
+void FD3D12CommandList::BeginQuery(FD3D12QueryLocation const& Location)
+{
+	check(Location);
+	check(Location.Heap->QueryType == D3D12_QUERY_TYPE_OCCLUSION);
+
+	GraphicsCommandList()->BeginQuery(
+		Location.Heap->GetD3DQueryHeap(),
+		Location.Heap->QueryType,
+		Location.Index
+	);
+}
+
+void FD3D12CommandList::EndQuery(FD3D12QueryLocation const& Location)
+{
+	check(Location);
+	switch (Location.Heap->QueryType)
+	{
+	default:
+		checkNoEntry();
+		break;
+
+	case D3D12_QUERY_TYPE_OCCLUSION:
+		GraphicsCommandList()->EndQuery(
+			Location.Heap->GetD3DQueryHeap(),
+			Location.Heap->QueryType,
+			Location.Index
+		);
+		State.OcclusionQueries.Add(Location);
+		break;
+
+	case D3D12_QUERY_TYPE_TIMESTAMP:
+		WriteTimestamp(Location);
+
+		// Command list begin/end timestamps are handled separately by the 
+		// submission thread, so shouldn't be in the TimestampQueries array.
+		if (Location.Type != ED3D12QueryType::CommandListBegin && Location.Type != ED3D12QueryType::CommandListEnd)
+		{
+			State.TimestampQueries.Add(Location);
+		}
+		break;
+	}
+}
+
 #if D3D12RHI_PLATFORM_USES_TIMESTAMP_QUERIES
-void FD3D12CommandList::WriteTimestamp(FD3D12QueryLocation& Location)
+void FD3D12CommandList::WriteTimestamp(FD3D12QueryLocation const& Location)
 {
-	Interfaces.GraphicsCommandList->EndQuery(
+	GraphicsCommandList()->EndQuery(
 		Location.Heap->GetD3DQueryHeap(),
 		Location.Heap->QueryType,
 		Location.Index
diff --git a/Engine/Source/Runtime/D3D12RHI/Private/D3D12CommandList.h b/Engine/Source/Runtime/D3D12RHI/Private/D3D12CommandList.h
--- a/Engine/Source/Runtime/D3D12RHI/Private/D3D12CommandList.h
+++ b/Engine/Source/Runtime/D3D12RHI/Private/D3D12CommandList.h
@@ -67,7 +67,8 @@ private:
 public:
 	~FD3D12CommandList();
 
-	void WriteTimestamp(FD3D12QueryLocation& Location);
+	void BeginQuery(FD3D12QueryLocation const& Location);
+	void EndQuery  (FD3D12QueryLocation const& Location);
 
 	void Reset(FD3D12CommandAllocator* CommandAllocator, FD3D12QueryAllocator* TimestampAllocator);
 	void Close();
@@ -75,6 +76,8 @@ public:
 	bool IsOpen  () const { return !State.IsClosed; }
 	bool IsClosed() const { return  State.IsClosed; }
 
+	uint32 GetNumCommands() const { return State.NumCommands; }
+
 	FD3D12Device*       const Device;
 	ED3D12QueueType     const QueueType;
 	FD3D12ResidencySet* const ResidencySet;
@@ -118,6 +121,86 @@ private:
 #endif
 	} Interfaces;
 
+public:
+	//
+	// Wrapper type to prevent l-value use of the returned command list interfaces.
+	// A context's command list may be swapped out during recording. Users should access the command
+	// list via the context itself, to ensure they always have the correct command list instance.
+	//
+	template <typename T>
+	class TRValuePtr
+	{
+		friend FD3D12CommandList;
+
+		FD3D12CommandList& CommandList;
+		T* Ptr;
+
+		TRValuePtr(FD3D12CommandList& CommandList, T* Ptr)
+			: CommandList(CommandList)
+			, Ptr(Ptr)
+		{}
+
+	public:
+		TRValuePtr() = delete;
+
+		TRValuePtr(TRValuePtr const&) = delete;
+		TRValuePtr(TRValuePtr&&)      = delete;
+
+		TRValuePtr& operator= (TRValuePtr const&) = delete;
+		TRValuePtr& operator= (TRValuePtr&&)      = delete;
+
+		operator bool () const&& { return !!Ptr; }
+		bool operator!() const&& { return  !Ptr; }
+
+		// These accessor functions count useful work on command lists
+		T* operator ->  () && { CommandList.State.NumCommands++; return Ptr; }
+		T* Get          () && { CommandList.State.NumCommands++; return Ptr; }
+
+		T* GetNoRefCount() && { return Ptr; }
+	};
+
+private:
+	template <typename FInterfaceType>
+	auto BuildRValuePtr(TRefCountPtr<FInterfaceType> FInterfaces::* Member)
+	{
+		return TRValuePtr<FInterfaceType>(*this, Interfaces.*Member);
+	}
+
+public:
+	auto BaseCommandList      () { return BuildRValuePtr(&FInterfaces::CommandList         ); }
+	auto CopyCommandList      () { return BuildRValuePtr(&FInterfaces::CopyCommandList     ); }
+	auto GraphicsCommandList  () { return BuildRValuePtr(&FInterfaces::GraphicsCommandList ); }
+#if D3D12_MAX_COMMANDLIST_INTERFACE >= 1
+	auto GraphicsCommandList1 () { return BuildRValuePtr(&FInterfaces::GraphicsCommandList1); }
+#endif
+#if D3D12_MAX_COMMANDLIST_INTERFACE >= 2
+	auto GraphicsCommandList2 () { return BuildRValuePtr(&FInterfaces::GraphicsCommandList2); }
+#endif
+#if D3D12_MAX_COMMANDLIST_INTERFACE >= 3
+	auto GraphicsCommandList3 () { return BuildRValuePtr(&FInterfaces::GraphicsCommandList3); }
+#endif
+#if D3D12_MAX_COMMANDLIST_INTERFACE >= 4
+	auto GraphicsCommandList4 () { return BuildRValuePtr(&FInterfaces::GraphicsCommandList4); }
+#endif
+#if D3D12_MAX_COMMANDLIST_INTERFACE >= 5
+	auto GraphicsCommandList5 () { return BuildRValuePtr(&FInterfaces::GraphicsCommandList5); }
+#endif
+#if D3D12_MAX_COMMANDLIST_INTERFACE >= 6
+	auto GraphicsCommandList6 () { return BuildRValuePtr(&FInterfaces::GraphicsCommandList6); }
+#endif
+#if D3D12_PLATFORM_SUPPORTS_ASSERTRESOURCESTATES
+	auto DebugCommandList     () { return BuildRValuePtr(&FInterfaces::DebugCommandList    ); }
+#endif
+#if D3D12_RHI_RAYTRACING
+	auto RayTracingCommandList() { return BuildRValuePtr(&FInterfaces::GraphicsCommandList4); }
+#endif
+#if NV_AFTERMATH
+	auto AftermathHandle      () { return Interfaces.AftermathHandle; } // @todo - should this increment NumCommands?
+#endif
+
+private:
+	void WriteTimestamp(FD3D12QueryLocation const& Location);
+
 	// Contents of the state struct are reset when the command list is recycled
 	struct FState
 	{
diff --git a/Engine/Source/Runtime/D3D12RHI/Private/D3D12Commands.cpp b/Engine/Source/Runtime/D3D12RHI/Private/D3D12Commands.cpp
--- a/Engine/Source/Runtime/D3D12RHI/Private/D3D12Commands.cpp
+++ b/Engine/Source/Runtime/D3D12RHI/Private/D3D12Commands.cpp
@@ -562,7 +562,6 @@ void FD3D12CommandContext::RHIBeginTransitions(TArrayView<const FRHITransition*>
 	}
 
 	// Signal fences
-	bool bClosed = false;
 	const ERHIPipeline SourcePipeline = GetPipeline();
 	for (const FRHITransition* Transition : Transitions)
 	{
@@ -573,28 +572,16 @@ void FD3D12CommandContext::RHIBeginTransitions(TArrayView<const FRHITransition*>
 
 			if (DeviceSyncPoints[SourcePipeline])
 			{
-				if (!bClosed)
-				{
-					CloseCommandList(false);
-					bClosed = true;
-				}
-
 				SignalSyncPoint(DeviceSyncPoints[SourcePipeline]);
 			}
 		}
 	}
-
-	if (bClosed)
-	{
-		OpenCommandList();
-	}
 }
 
 void FD3D12CommandContext::RHIEndTransitions(TArrayView<const FRHITransition*> Transitions)
 {
 	// Wait for fences
 	{
-		bool bClosed = false;
 		const ERHIPipeline DstPipeline = GetPipeline();
 		for (const FRHITransition* Transition : Transitions)
 		{
@@ -607,22 +594,11 @@ void FD3D12CommandContext::RHIEndTransitions(TArrayView<const FRHITransition*> T
 				{
 					if (SrcPipeline != DstPipeline && DeviceSyncPoints[SrcPipeline])
 					{
-						if (!bClosed)
-						{
-							CloseCommandList(false);
-							bClosed = true;
-						}
-
 						WaitSyncPoint(DeviceSyncPoints[SrcPipeline]);
 					}
 				});
 			}
 		}
-
-		if (bClosed)
-		{
-			OpenCommandList();
-		}
 	}
 
 	static IConsoleVariable* CVarShowTransitions = IConsoleManager::Get().FindConsoleVariable(TEXT("r.ProfileGPU.ShowTransitions"));
@@ -2349,10 +2325,7 @@ void FD3D12CommandContext::BroadcastTemporalEffect(const FName& InEffectName, co
 	// This will be used to fence the copy from starting until all prior work on this queue is complete.
 	// @todo mgpu - RHIBroadcastTemporalEffect does not take async compute into consideration.
 	FD3D12SyncPointRef WaitForGraphics_SyncPoint = FD3D12SyncPoint::Create(ED3D12SyncPointType::GPUOnly);
-
-	CloseCommandList(false);
 	SignalSyncPoint(WaitForGraphics_SyncPoint);
-	OpenCommandList();
 
 	// Start the copy work
 	FD3D12SyncPointRef& CopySyncPoint = Device->GetParentAdapter()->GetTemporalEffect(InEffectName)[SiblingGPUIndex];
@@ -2448,10 +2421,7 @@ void FD3D12CommandContext::RHIWaitForTemporalEffect(const FName& InEffectName)
 
 	if (SyncPoint)
 	{
-		CloseCommandList(false);
 		WaitSyncPoint(SyncPoint);
-		OpenCommandList();
-
 		SyncPoint = nullptr;
 	}
 }
diff --git a/Engine/Source/Runtime/D3D12RHI/Private/D3D12Device.cpp b/Engine/Source/Runtime/D3D12RHI/Private/D3D12Device.cpp
--- a/Engine/Source/Runtime/D3D12RHI/Private/D3D12Device.cpp
+++ b/Engine/Source/Runtime/D3D12RHI/Private/D3D12Device.cpp
@@ -558,10 +558,6 @@ FD3D12ContextCommon* FD3D12Device::ObtainContext(ED3D12QueueType QueueType)
 		case ED3D12QueueType::Copy  : Context = new FD3D12ContextCopy   (this);                   break;
 		}
 	}
-	else
-	{
-		Context->OpenCommandList();
-	}
 
 	check(Context);
 	return Context;
diff --git a/Engine/Source/Runtime/D3D12RHI/Private/D3D12Device.h b/Engine/Source/Runtime/D3D12RHI/Private/D3D12Device.h
--- a/Engine/Source/Runtime/D3D12RHI/Private/D3D12Device.h
+++ b/Engine/Source/Runtime/D3D12RHI/Private/D3D12Device.h
@@ -117,7 +117,6 @@ public:
 
 	FD3D12Payload*          PayloadToSubmit    = nullptr;
 	FD3D12CommandAllocator* BarrierAllocator   = nullptr;
-	FD3D12CommandList*      BarrierCommandList = nullptr;
 	FD3D12QueryAllocator    BarrierTimestamps;
 
 	uint32 NumCommandListsInBatch = 0;
diff --git a/Engine/Source/Runtime/D3D12RHI/Private/D3D12DirectCommandListManager.cpp b/Engine/Source/Runtime/D3D12RHI/Private/D3D12DirectCommandListManager.cpp
--- a/Engine/Source/Runtime/D3D12RHI/Private/D3D12DirectCommandListManager.cpp
+++ b/Engine/Source/Runtime/D3D12RHI/Private/D3D12DirectCommandListManager.cpp
@@ -67,9 +67,7 @@ void FD3D12CommandContext::RHIWriteGPUFence(FRHIGPUFence* FenceRHI)
 	check(Fence);
 	check(Fence->SyncPoints[GetGPUIndex()]);
 
-	CloseCommandList(true);
 	SignalSyncPoint(Fence->SyncPoints[GetGPUIndex()]);
-	OpenCommandList();
 }
 
 
@@ -143,19 +141,9 @@ void FD3D12ManualFence::AdvanceFrame()
 	const uint64 NextValue = NextFenceValue.Increment();
 	FRHICommandListExecutor::GetImmediateCommandList().EnqueueLambda([this, NextValue](FRHICommandListImmediate&)
 	{
-		for (FFencePair& Pair : FencePairs)
-		{
-			Pair.Context->CloseCommandList(false);
-		}
-
 		for (FFencePair& Pair : FencePairs)
 		{
 			Pair.Context->SignalManualFence(Pair.Fence, NextValue);
 		}
-
-		for (FFencePair& Pair : FencePairs)
-		{
-			Pair.Context->OpenCommandList();
-		}
 	});
 }
diff --git a/Engine/Source/Runtime/D3D12RHI/Private/D3D12Query.cpp b/Engine/Source/Runtime/D3D12RHI/Private/D3D12Query.cpp
--- a/Engine/Source/Runtime/D3D12RHI/Private/D3D12Query.cpp
+++ b/Engine/Source/Runtime/D3D12RHI/Private/D3D12Query.cpp
@@ -191,12 +191,7 @@ void FD3D12CommandContext::RHIBeginRenderQuery(FRHIRenderQuery* QueryRHI)
 	checkf(Query->Type == RQT_Occlusion, TEXT("Only occlusion queries support RHIBeginRenderQuery()."));
 
 	Query->ActiveLocation = AllocateQuery(ED3D12QueryType::Occlusion, &Query->Result);
-
-	GraphicsCommandList()->BeginQuery(
-		Query->ActiveLocation.Heap->GetD3DQueryHeap(),
-		Query->ActiveLocation.Heap->QueryType,
-		Query->ActiveLocation.Index
-	);
+	BeginQuery(Query->ActiveLocation);
 
 	ActiveQueries++;
 }
@@ -241,12 +236,7 @@ void FD3D12CommandContext::RHIEndRenderQuery(FRHIRenderQuery* QueryRHI)
 		check(ActiveQueries > 0);
 		ActiveQueries--;
 
-		check(Query->ActiveLocation.Heap);
-		GraphicsCommandList()->EndQuery(
-			Query->ActiveLocation.Heap->GetD3DQueryHeap(),
-			Query->ActiveLocation.Heap->QueryType,
-			Query->ActiveLocation.Index
-		);
+		EndQuery(Query->ActiveLocation);
 		Query->ActiveLocation = {};
 		break;
 
diff --git a/Engine/Source/Runtime/D3D12RHI/Private/D3D12RHI.cpp b/Engine/Source/Runtime/D3D12RHI/Private/D3D12RHI.cpp
--- a/Engine/Source/Runtime/D3D12RHI/Private/D3D12RHI.cpp
+++ b/Engine/Source/Runtime/D3D12RHI/Private/D3D12RHI.cpp
@@ -408,7 +408,6 @@ void FD3D12DynamicRHI::EnqueueEndOfPipeTask(TUniqueFunction<void()> TaskFunc, TU
 
 		FD3D12SyncPointRef SyncPoint = FD3D12SyncPoint::Create(ED3D12SyncPointType::GPUAndCPU);
 		Payload->SyncPointsToSignal.Emplace(SyncPoint);
-		Payload->bResolveQueries = true;
 		Prereqs.Add(SyncPoint->GetGraphEvent());
 
 		if (ModifyPayloadCallback)
@@ -624,18 +623,14 @@ void FD3D12DynamicRHI::RHISignalManualFence(FRHICommandList& RHICmdList, ID3D12F
 {
 	checkf(FRHIGPUMask::All() == FRHIGPUMask::GPU0(), TEXT("RHISignalManualFence cannot be used by multi-GPU code"));
 	FD3D12CommandContext& Context = FD3D12CommandContext::Get(RHICmdList);
-	Context.CloseCommandList(false);
 	Context.SignalManualFence(Fence, Value);
-	Context.OpenCommandList();
 }
 
 void FD3D12DynamicRHI::RHIWaitManualFence(FRHICommandList& RHICmdList, ID3D12Fence* Fence, uint64 Value)
 {
 	checkf(FRHIGPUMask::All() == FRHIGPUMask::GPU0(), TEXT("RHIWaitManualFence cannot be used by multi-GPU code"));
 	FD3D12CommandContext& Context = FD3D12CommandContext::Get(RHICmdList);
-	Context.CloseCommandList(false);
 	Context.WaitManualFence(Fence, Value);
-	Context.OpenCommandList();
 }
 
 bool ID3D12DynamicRHI::IsD3DDebugEnabled()
diff --git a/Engine/Source/Runtime/D3D12RHI/Private/D3D12Resources.cpp b/Engine/Source/Runtime/D3D12RHI/Private/D3D12Resources.cpp
--- a/Engine/Source/Runtime/D3D12RHI/Private/D3D12Resources.cpp
+++ b/Engine/Source/Runtime/D3D12RHI/Private/D3D12Resources.cpp
@@ -1006,10 +1006,7 @@ void FD3D12ResourceBarrierBatcher::FlushIntoCommandList(FD3D12CommandList& Comma
 
 	auto InsertTimestamp = [&](ED3D12QueryType Type)
 	{
-		auto Timestamp = TimestampAllocator.Allocate(Type, nullptr);
-		CommandList.State.TimestampQueries.Emplace_GetRef(Timestamp);
-		CommandList.WriteTimestamp(Timestamp);
-		CommandList.State.NumCommands++;
+		CommandList.EndQuery(TimestampAllocator.Allocate(Type, nullptr));
 	};
 
 	for (int32 BatchStart = 0, BatchEnd = 0; BatchStart < Barriers.Num(); BatchStart = BatchEnd)
@@ -1031,8 +1028,7 @@ void FD3D12ResourceBarrierBatcher::FlushIntoCommandList(FD3D12CommandList& Comma
 			InsertTimestamp(ED3D12QueryType::IdleBegin);
 		}
 
-		CommandList.Interfaces.GraphicsCommandList->ResourceBarrier(BatchEnd - BatchStart, &Barriers[BatchStart]);
-		CommandList.State.NumCommands++;
+		CommandList.GraphicsCommandList()->ResourceBarrier(BatchEnd - BatchStart, &Barriers[BatchStart]);
 
 		if (bIdle)
 		{
diff --git a/Engine/Source/Runtime/D3D12RHI/Private/D3D12Submission.cpp b/Engine/Source/Runtime/D3D12RHI/Private/D3D12Submission.cpp
--- a/Engine/Source/Runtime/D3D12RHI/Private/D3D12Submission.cpp
+++ b/Engine/Source/Runtime/D3D12RHI/Private/D3D12Submission.cpp
@@ -125,14 +125,9 @@ IRHIPlatformCommandList* FD3D12DynamicRHI::RHIFinalizeContext(IRHIComputeContext
 	FD3D12FinalizedCommands Result;
 	auto FinalizeContext = [&](FD3D12CommandContext* CmdContext)
 	{
-		CmdContext->CloseCommandList(true);
 		CmdContext->Finalize(Result);
 
-		if (CmdContext->IsDefaultContext())
-		{
-			CmdContext->OpenCommandList();
-		}
-		else
+		if (!CmdContext->IsDefaultContext())
 		{
 			CmdContext->ClearState();
 			CmdContext->GetParentDevice()->ReleaseContext(CmdContext);
@@ -288,9 +283,6 @@ FD3D12DynamicRHI::FProcessResult FD3D12DynamicRHI::ProcessSubmissionQueue()
 				// executing on a different queue (e.g. graphics-only transitions required 
 				// before async compute work), so we gather potential work across all
 				// queues for this device.
-				// 
-				// The last command list from each finalized context will still be open, 
-				// allowing us to use it to append any necessary ResolveQueryData commands.
 				//
 				const uint32 MaxBatchSize = GetMaxExecuteBatchSize();
 				auto AccumulateQueries = [&](FD3D12CommandList* CommandList)
@@ -333,10 +325,6 @@ FD3D12DynamicRHI::FProcessResult FD3D12DynamicRHI::ProcessSubmissionQueue()
 						if (&BarrierQueue == &CurrentQueue)
 						{
 							// Barrier command list will run on the current queue.
-
-							// There's always another command list after the barrier one, so close this immediately.
-							BarrierCommandList->Close();
-
 							// Insert it immediately before the corresponding command list that generated it.
 							check(BarrierQueue.PayloadToSubmit);
 							BarrierQueue.PayloadToSubmit->CommandListsToExecute.Insert(BarrierCommandList, Index++);
@@ -370,66 +358,88 @@ FD3D12DynamicRHI::FProcessResult FD3D12DynamicRHI::ProcessSubmissionQueue()
 				check(Payload->SyncPointsToWait.Num() == 0);
 
 				Queue->BarrierTimestamps.CloseAndReset(Queue->PendingQueryRanges);
-				Queue->BarrierCommandList = nullptr;
 				Queue->NumCommandListsInBatch = 0;
 
 				check(Payload->QueryRanges.Num() == 0);
 				check(Payload->TimestampQueries.Num() == 0);
 				check(Payload->OcclusionQueries.Num() == 0);
 
-				FD3D12CommandList* FinalCommandList = Payload->CommandListsToExecute.Num() > 0
-					? Payload->CommandListsToExecute.Last()
-					: nullptr;
-
-				if (Payload->bResolveQueries && Queue->PendingQueryRanges.Num())
+				if (Queue->PendingQueryRanges.Num())
 				{
-					if (!FinalCommandList)
+					// If this payload will signal a CPU-visible sync point, we need to resolve queries.
+					// This makes sure that the query data has reached the CPU before the sync point the CPU is waiting on is signaled.
+					bool bResolveQueries = false;
+					for (FD3D12SyncPoint* SyncPoint : Payload->SyncPointsToSignal)
 					{
-						// We've got queries to resolve, but no command list to use. Allocate one.
-						if (!Queue->BarrierAllocator)
+						if (SyncPoint->GetType() == ED3D12SyncPointType::GPUAndCPU)
 						{
-							Queue->BarrierAllocator = Queue->Device->ObtainCommandAllocator(Queue->QueueType);
+							bResolveQueries = true;
+							break;
 						}
-
-						FinalCommandList = Queue->Device->ObtainCommandList(Queue->BarrierAllocator, nullptr);
-						Payload->CommandListsToExecute.Add(FinalCommandList);
 					}
 
-					check(FinalCommandList->IsOpen());
-					FinalCommandList->WriteEndTimestamp();
-
-					for (FD3D12QueryRange const& Range : Queue->PendingQueryRanges)
+					if (bResolveQueries)
 					{
-						check(Range.End > Range.Start);
-
-						FD3D12Resource* ResultBuffer = Range.Heap->GetResultBuffer();
-						FinalCommandList->UpdateResidency(
+						FD3D12CommandList* ResolveCommandList = nullptr;
 						{
-							&Range.Heap->GetHeapResidencyHandle(),
-							&ResultBuffer->GetResidencyHandle()
-						});
+							auto GetResolveCommandList = [&]() -> FD3D12CommandList*
+							{
+								if (ResolveCommandList)
+									return ResolveCommandList;
 
-						if (Range.Heap->GetD3DQueryHeap())
+								if (!Queue->BarrierAllocator)
+									Queue->BarrierAllocator = Queue->Device->ObtainCommandAllocator(Queue->QueueType); 
+
+								return ResolveCommandList = Queue->Device->ObtainCommandList(Queue->BarrierAllocator, nullptr);
+							};
+
+							// We've got queries to resolve. Allocate a command list.
+							for (FD3D12QueryRange const& Range : Queue->PendingQueryRanges)
+							{
+								check(Range.End > Range.Start);
+
+	#if ENABLE_RESIDENCY_MANAGEMENT
+								FD3D12ResidencyHandle* ResidencyHandles[] =
+								{
+									&Range.Heap->GetHeapResidencyHandle(),
+									&Range.Heap->GetResultBuffer()->GetResidencyHandle()
+								};
+
+								for (FD3D12ResidencyHandle* Handle : ResidencyHandles)
+								{
+									if (Handle->IsInitialized())
+									{
+										GetResolveCommandList()->UpdateResidency({ Handle });
+									}
+								}
+	#endif // ENABLE_RESIDENCY_MANAGEMENT
+
+								if (Range.Heap->GetD3DQueryHeap())
+								{
+									GetResolveCommandList()->GraphicsCommandList()->ResolveQueryData(
+										Range.Heap->GetD3DQueryHeap(),
+										Range.Heap->QueryType,
+										Range.Start,
+										Range.End - Range.Start,
+										Range.Heap->GetResultBuffer()->GetResource(),
+										Range.Start * FD3D12QueryHeap::ResultSize
+									);
+								}
+							}
+						}
+
+						Payload->QueryRanges      = MoveTemp(Queue->PendingQueryRanges     );
+						Payload->TimestampQueries = MoveTemp(Queue->PendingTimestampQueries);
+						Payload->OcclusionQueries = MoveTemp(Queue->PendingOcclusionQueries);
+
+						if (ResolveCommandList)
 						{
-							FinalCommandList->Interfaces.GraphicsCommandList->ResolveQueryData(
-								Range.Heap->GetD3DQueryHeap(),
-								Range.Heap->QueryType,
-								Range.Start,
-								Range.End - Range.Start,
-								ResultBuffer->GetResource(),
-								Range.Start * FD3D12QueryHeap::ResultSize
-							);
+							ResolveCommandList->Close();
+							Payload->CommandListsToExecute.Add(ResolveCommandList);
 						}
 					}
-
-					Payload->QueryRanges      = MoveTemp(Queue->PendingQueryRanges     );
-					Payload->TimestampQueries = MoveTemp(Queue->PendingTimestampQueries);
-					Payload->OcclusionQueries = MoveTemp(Queue->PendingOcclusionQueries);
 				}
 
-				if (FinalCommandList && FinalCommandList->IsOpen())
-					FinalCommandList->Close();
-
 				if (Queue->BarrierAllocator)
 				{
 					Payload->AllocatorsToRelease.Add(Queue->BarrierAllocator);
@@ -582,24 +592,21 @@ FD3D12CommandList* FD3D12DynamicRHI::GenerateBarrierCommandListAndUpdateState(FD
 			: SourceCommandList->QueueType
 	);
 
-	// Close the previous barrier command list, since we're going to reuse its allocator.
-	if (Queue.BarrierCommandList && Queue.BarrierCommandList->IsOpen())
-		Queue.BarrierCommandList->Close();
-
 	// Get an allocator if we don't have one
 	if (!Queue.BarrierAllocator)
 		Queue.BarrierAllocator = Queue.Device->ObtainCommandAllocator(Queue.QueueType);
 
 	// Get a new command list
-	Queue.BarrierCommandList = Queue.Device->ObtainCommandList(Queue.BarrierAllocator, &Queue.BarrierTimestamps);
+	FD3D12CommandList* BarrierCommandList = Queue.Device->ObtainCommandList(Queue.BarrierAllocator, &Queue.BarrierTimestamps);
 
 #if ENABLE_RESIDENCY_MANAGEMENT
-	Queue.BarrierCommandList->UpdateResidency(ResidencyHandles);
+	BarrierCommandList->UpdateResidency(ResidencyHandles);
 #endif
 
-	Batcher.FlushIntoCommandList(*Queue.BarrierCommandList, Queue.BarrierTimestamps);
+	Batcher.FlushIntoCommandList(*BarrierCommandList, Queue.BarrierTimestamps);
+	BarrierCommandList->Close();
 
-	return Queue.BarrierCommandList;
+	return BarrierCommandList;
 }
 
 uint64 FD3D12Queue::ExecutePayload()
diff --git a/Engine/Source/Runtime/D3D12RHI/Private/D3D12Submission.h b/Engine/Source/Runtime/D3D12RHI/Private/D3D12Submission.h
--- a/Engine/Source/Runtime/D3D12RHI/Private/D3D12Submission.h
+++ b/Engine/Source/Runtime/D3D12RHI/Private/D3D12Submission.h
@@ -197,7 +197,6 @@ struct FD3D12PayloadBase
 	TArray<FManualFence> FencesToWait;
 
 	// Execute
-	bool bResolveQueries = false;
 	TArray<class FD3D12CommandList*> CommandListsToExecute;
 
 	// Signal
diff --git a/Engine/Source/Runtime/D3D12RHI/Private/D3D12Viewport.cpp b/Engine/Source/Runtime/D3D12RHI/Private/D3D12Viewport.cpp
--- a/Engine/Source/Runtime/D3D12RHI/Private/D3D12Viewport.cpp
+++ b/Engine/Source/Runtime/D3D12RHI/Private/D3D12Viewport.cpp
@@ -702,7 +702,6 @@ bool FD3D12Viewport::Present(bool bLockToVsync)
 		// system present queue out of order.
 		if (LastFrameSyncPoint)
 		{
-			PresentContext->CloseCommandList(true);
 			PresentContext->WaitSyncPoint(LastFrameSyncPoint);
 		}
 	}
@@ -733,12 +732,7 @@ bool FD3D12Viewport::Present(bool bLockToVsync)
 		// Currently, the swap chain Present() is called directly by the RHI thread.
 		// We need to submit the above commands and wait for the submission thread to process everything before we can continue.
 		FD3D12CommandContext& DefaultContext = Adapter->GetDevice(GPUIndex)->GetDefaultCommandContext();
-
-		ED3D12FlushFlags FlushMode = &DefaultContext == PresentContext
-			? ED3D12FlushFlags::WaitForSubmission | ED3D12FlushFlags::NoOpen
-			: ED3D12FlushFlags::WaitForSubmission;
-
-		DefaultContext.FlushCommands(FlushMode);
+		DefaultContext.FlushCommands(ED3D12FlushFlags::WaitForSubmission);
 	}
 
 	const int32 SyncInterval = bLockToVsync ? RHIGetSyncInterval() : 0;
@@ -800,10 +794,8 @@ void FD3D12Viewport::IssueFrameEvent()
 
 		FD3D12SyncPointRef SyncPoint = FD3D12SyncPoint::Create(ED3D12SyncPointType::GPUAndCPU);
 
-		Context.CloseCommandList(true);
 		Context.SignalSyncPoint(SyncPoint);
 		Context.Finalize(Payloads);
-		Context.OpenCommandList();
 
 		FrameSyncPoints.Emplace(MoveTemp(SyncPoint));
 	}
