Reorganizing D3D12 constant buffer binding code
- Instead of a single "SetConstantBuffers" function that can set CBVs or Root CBs, group the CBVs with other resources and have Root CBs bind on their own.
- Binding Root CBs on their own allows the bindless resources path to actually bind Root CBs.
- Remove the overengineered templates from SetConstantBuffers

diff --git a/Engine/Source/Runtime/D3D12RHI/Private/D3D12DescriptorCache.cpp b/Engine/Source/Runtime/D3D12RHI/Private/D3D12DescriptorCache.cpp
--- a/Engine/Source/Runtime/D3D12RHI/Private/D3D12DescriptorCache.cpp
+++ b/Engine/Source/Runtime/D3D12RHI/Private/D3D12DescriptorCache.cpp
@@ -6,22 +6,6 @@
 #include "D3D12RHIPrivate.h"
 
 // Define template functions that are only declared in the header.
-#if USE_STATIC_ROOT_SIGNATURE
-template void FD3D12DescriptorCache::SetConstantBuffers<SF_Vertex>(const FD3D12RootSignature* RootSignature, FD3D12ConstantBufferCache& Cache, const CBVSlotMask& SlotsNeededMask, uint32 Count, uint32& HeapSlot);
-template void FD3D12DescriptorCache::SetConstantBuffers<SF_Mesh>(const FD3D12RootSignature* RootSignature, FD3D12ConstantBufferCache& Cache, const CBVSlotMask& SlotsNeededMask, uint32 Count, uint32& HeapSlot);
-template void FD3D12DescriptorCache::SetConstantBuffers<SF_Amplification>(const FD3D12RootSignature* RootSignature, FD3D12ConstantBufferCache& Cache, const CBVSlotMask& SlotsNeededMask, uint32 Count, uint32& HeapSlot);
-template void FD3D12DescriptorCache::SetConstantBuffers<SF_Geometry>(const FD3D12RootSignature* RootSignature, FD3D12ConstantBufferCache& Cache, const CBVSlotMask& SlotsNeededMask, uint32 Count, uint32& HeapSlot);
-template void FD3D12DescriptorCache::SetConstantBuffers<SF_Pixel>(const FD3D12RootSignature* RootSignature, FD3D12ConstantBufferCache& Cache, const CBVSlotMask& SlotsNeededMask, uint32 Count, uint32& HeapSlot);
-template void FD3D12DescriptorCache::SetConstantBuffers<SF_Compute>(const FD3D12RootSignature* RootSignature, FD3D12ConstantBufferCache& Cache, const CBVSlotMask& SlotsNeededMask, uint32 Count, uint32& HeapSlot);
-#else
-template void FD3D12DescriptorCache::SetConstantBuffers<SF_Vertex>(const FD3D12RootSignature* RootSignature, FD3D12ConstantBufferCache& Cache, const CBVSlotMask& SlotsNeededMask);
-template void FD3D12DescriptorCache::SetConstantBuffers<SF_Mesh>(const FD3D12RootSignature* RootSignature, FD3D12ConstantBufferCache& Cache, const CBVSlotMask& SlotsNeededMask);
-template void FD3D12DescriptorCache::SetConstantBuffers<SF_Amplification>(const FD3D12RootSignature* RootSignature, FD3D12ConstantBufferCache& Cache, const CBVSlotMask& SlotsNeededMask);
-template void FD3D12DescriptorCache::SetConstantBuffers<SF_Geometry>(const FD3D12RootSignature* RootSignature, FD3D12ConstantBufferCache& Cache, const CBVSlotMask& SlotsNeededMask);
-template void FD3D12DescriptorCache::SetConstantBuffers<SF_Pixel>(const FD3D12RootSignature* RootSignature, FD3D12ConstantBufferCache& Cache, const CBVSlotMask& SlotsNeededMask);
-template void FD3D12DescriptorCache::SetConstantBuffers<SF_Compute>(const FD3D12RootSignature* RootSignature, FD3D12ConstantBufferCache& Cache, const CBVSlotMask& SlotsNeededMask);
-#endif
-
 template void FD3D12DescriptorCache::SetSRVs<SF_Vertex>(const FD3D12RootSignature* RootSignature, FD3D12ShaderResourceViewCache& Cache, const SRVSlotMask& SlotsNeededMask, uint32 SlotsNeeded, uint32& HeapSlot);
 template void FD3D12DescriptorCache::SetSRVs<SF_Mesh>(const FD3D12RootSignature* RootSignature, FD3D12ShaderResourceViewCache& Cache, const SRVSlotMask& SlotsNeededMask, uint32 SlotsNeeded, uint32& HeapSlot);
 template void FD3D12DescriptorCache::SetSRVs<SF_Amplification>(const FD3D12RootSignature* RootSignature, FD3D12ShaderResourceViewCache& Cache, const SRVSlotMask& SlotsNeededMask, uint32 SlotsNeeded, uint32& HeapSlot);
@@ -498,15 +482,9 @@ void FD3D12DescriptorCache::SetSRVs(const FD3D12RootSignature* RootSignature, FD
 #endif
 }
 
-template <EShaderFrequency ShaderStage>
-#if USE_STATIC_ROOT_SIGNATURE
-void FD3D12DescriptorCache::SetConstantBuffers(const FD3D12RootSignature* RootSignature, FD3D12ConstantBufferCache& Cache, const CBVSlotMask& SlotsNeededMask, uint32 SlotsNeeded, uint32& HeapSlot)
-#else
-void FD3D12DescriptorCache::SetConstantBuffers(const FD3D12RootSignature* RootSignature, FD3D12ConstantBufferCache& Cache, const CBVSlotMask& SlotsNeededMask)
-#endif
+void FD3D12DescriptorCache::SetConstantBufferViews(EShaderFrequency ShaderStage, const FD3D12RootSignature* RootSignature, FD3D12ConstantBufferCache& Cache, CBVSlotMask SlotsNeededMask, uint32 SlotsNeeded, uint32& HeapSlot)
 {
-	static_assert(ShaderStage < SF_NumStandardFrequencies, "Unexpected shader frequency.");
-
+#if USE_STATIC_ROOT_SIGNATURE
 	CBVSlotMask& CurrentDirtySlotMask = Cache.DirtySlotMask[ShaderStage];
 	check(CurrentDirtySlotMask != 0);	// All dirty slots for the current shader stage.
 	check(SlotsNeededMask != 0);		// All dirty slots for the current shader stage AND used by the current shader stage.
@@ -517,7 +495,6 @@ void FD3D12DescriptorCache::SetConstantBuffers(const FD3D12RootSignature* RootSi
 	const CBVSlotMask RDCBVSlotsNeededMask = GRootCBVSlotMask & SlotsNeededMask;
 	check(RDCBVSlotsNeededMask); // Check this wasn't a wasted call.
 
-#if USE_STATIC_ROOT_SIGNATURE
 	// Now desc table with CBV
 	auto& CBVHandles = Cache.CBHandles[ShaderStage];
 
@@ -576,7 +553,19 @@ void FD3D12DescriptorCache::SetConstantBuffers(const FD3D12RootSignature* RootSi
 #ifdef VERBOSE_DESCRIPTOR_HEAP_DEBUG
 	FMsg::Logf(__FILE__, __LINE__, TEXT("DescriptorCache"), ELogVerbosity::Log, TEXT("SetShaderResourceViewTable [STAGE %d] to slots %d - %d"), (int32)ShaderStage, FirstSlotIndex, FirstSlotIndex + SlotsNeeded - 1);
 #endif
-#else
+
+#endif // USE_STATIC_ROOT_SIGNATURE
+}
+
+void FD3D12DescriptorCache::SetRootConstantBuffers(EShaderFrequency ShaderStage, const FD3D12RootSignature* RootSignature, FD3D12ConstantBufferCache& Cache, CBVSlotMask SlotsNeededMask)
+{
+	CBVSlotMask& CurrentDirtySlotMask = Cache.DirtySlotMask[ShaderStage];
+	check(CurrentDirtySlotMask != 0);	// All dirty slots for the current shader stage.
+	check(SlotsNeededMask != 0);		// All dirty slots for the current shader stage AND used by the current shader stage.
+
+	// Process root CBV
+	const CBVSlotMask RDCBVSlotsNeededMask = GRootCBVSlotMask & SlotsNeededMask;
+	check(RDCBVSlotsNeededMask); // Check this wasn't a wasted call.
 
 	// Set root descriptors.
 	// At least one needed root descriptor is dirty.
@@ -589,7 +578,7 @@ void FD3D12DescriptorCache::SetConstantBuffers(const FD3D12RootSignature* RootSi
 		// Only set the root descriptor if it's dirty and we need to set it (it can be used by the shader).
 		if (FD3D12ConstantBufferCache::IsSlotDirty(RDCBVSlotsNeededMask, SlotIndex))
 		{
-			const D3D12_GPU_VIRTUAL_ADDRESS& CurrentGPUVirtualAddress = Cache.CurrentGPUVirtualAddress[ShaderStage][SlotIndex];
+			const D3D12_GPU_VIRTUAL_ADDRESS CurrentGPUVirtualAddress = Cache.CurrentGPUVirtualAddress[ShaderStage][SlotIndex];
 			check(CurrentGPUVirtualAddress != 0);
 			if (ShaderStage == SF_Compute)
 			{
@@ -610,7 +599,6 @@ void FD3D12DescriptorCache::SetConstantBuffers(const FD3D12RootSignature* RootSi
 	check((CurrentDirtySlotMask & RDCBVSlotsNeededMask) == 0);	// Check all slots that needed to be set, were set.
 
 	static_assert(GDescriptorTableCBVSlotMask == 0, "FD3D12DescriptorCache::SetConstantBuffers needs to be updated to handle descriptor tables.");	// Check that all CBVs slots are controlled by root descriptors.
-#endif	
 }
 
 bool FD3D12DescriptorCache::SwitchToContextLocalViewHeap()
diff --git a/Engine/Source/Runtime/D3D12RHI/Private/D3D12DescriptorCache.h b/Engine/Source/Runtime/D3D12RHI/Private/D3D12DescriptorCache.h
--- a/Engine/Source/Runtime/D3D12RHI/Private/D3D12DescriptorCache.h
+++ b/Engine/Source/Runtime/D3D12RHI/Private/D3D12DescriptorCache.h
@@ -348,12 +348,9 @@ public:
 	template <EShaderFrequency ShaderStage>
 	void SetSRVs(const FD3D12RootSignature* RootSignature, FD3D12ShaderResourceViewCache& Cache, const SRVSlotMask& SlotsNeededMask, uint32 Count, uint32& HeapSlot);
 
-	template <EShaderFrequency ShaderStage> 
-#if USE_STATIC_ROOT_SIGNATURE
-	void SetConstantBuffers(const FD3D12RootSignature* RootSignature, FD3D12ConstantBufferCache& Cache, const CBVSlotMask& SlotsNeededMask, uint32 Count, uint32& HeapSlot);
-#else
-	void SetConstantBuffers(const FD3D12RootSignature* RootSignature, FD3D12ConstantBufferCache& Cache, const CBVSlotMask& SlotsNeededMask);
-#endif
+	void SetConstantBufferViews(EShaderFrequency ShaderStage, const FD3D12RootSignature* RootSignature, FD3D12ConstantBufferCache& Cache, CBVSlotMask SlotsNeededMask, uint32 Count, uint32& HeapSlot);
+
+	void SetRootConstantBuffers(EShaderFrequency ShaderStage, const FD3D12RootSignature* RootSignature, FD3D12ConstantBufferCache& Cache, CBVSlotMask SlotsNeededMask);
 
 	void SetStreamOutTargets(FD3D12Resource **Buffers, uint32 Count, const uint32* Offsets);
 
diff --git a/Engine/Source/Runtime/D3D12RHI/Private/D3D12StateCache.cpp b/Engine/Source/Runtime/D3D12RHI/Private/D3D12StateCache.cpp
--- a/Engine/Source/Runtime/D3D12RHI/Private/D3D12StateCache.cpp
+++ b/Engine/Source/Runtime/D3D12RHI/Private/D3D12StateCache.cpp
@@ -558,6 +558,8 @@ void FD3D12StateCache::ApplyState()
 		ApplyResources(PSOCommonData->RootSignature, StartStage, EndStage);
 	}
 
+	ApplyConstants(PSOCommonData->RootSignature, StartStage, EndStage);
+
 	// Flush any needed resource barriers
 	CmdContext.FlushResourceBarriers();
 
@@ -572,7 +574,9 @@ void FD3D12StateCache::ApplyResources(const FD3D12RootSignature* const pRootSign
 	// Determine what resource bind slots are dirty for the current shaders and how many descriptor table slots we need.
 	// We only set dirty resources that can be used for the upcoming Draw/Dispatch.
 	SRVSlotMask CurrentShaderDirtySRVSlots[SF_NumStandardFrequencies] = {};
+#if USE_STATIC_ROOT_SIGNATURE
 	CBVSlotMask CurrentShaderDirtyCBVSlots[SF_NumStandardFrequencies] = {};
+#endif
 	UAVSlotMask CurrentShaderDirtyUAVSlots = 0;
 	uint32 NumUAVs = 0;
 	uint32 NumSRVs[SF_NumStandardFrequencies] = {};
@@ -629,9 +633,9 @@ void FD3D12StateCache::ApplyResources(const FD3D12RootSignature* const pRootSign
 				NumViews += NumSRVs[Stage];
 			}
 
+#if USE_STATIC_ROOT_SIGNATURE
 			const CBVSlotMask CurrentShaderCBVRegisterMask = BitMask<CBVSlotMask>(PipelineState.Common.CurrentShaderCBCounts[Stage]);
 			CurrentShaderDirtyCBVSlots[Stage] = CurrentShaderCBVRegisterMask & PipelineState.Common.CBVCache.DirtySlotMask[Stage];
-#if USE_STATIC_ROOT_SIGNATURE
 			if (CurrentShaderDirtyCBVSlots[Stage])
 			{
 				if (ResourceBindingTier == D3D12_RESOURCE_BINDING_TIER_1)
@@ -710,24 +714,66 @@ void FD3D12StateCache::ApplyResources(const FD3D12RootSignature* const pRootSign
 #undef CONDITIONAL_SET_SRVS
 	}
 
+#if USE_STATIC_ROOT_SIGNATURE
 	// Constant buffers
 	{
 		//SCOPE_CYCLE_COUNTER(STAT_D3D12ApplyStateSetConstantBufferTime);
 		FD3D12ConstantBufferCache& CBVCache = PipelineState.Common.CBVCache;
 
-#if USE_STATIC_ROOT_SIGNATURE
 	#define CONDITIONAL_SET_CBVS(Shader) \
 		if (CurrentShaderDirtyCBVSlots[##Shader]) \
 		{ \
-			DescriptorCache.SetConstantBuffers<##Shader>(pRootSignature, CBVCache, CurrentShaderDirtyCBVSlots[##Shader], NumCBVs[##Shader], ViewHeapSlot); \
+			DescriptorCache.SetConstantBufferViews(##Shader, pRootSignature, CBVCache, CurrentShaderDirtyCBVSlots[##Shader], NumCBVs[##Shader], ViewHeapSlot); \
 		}
-#else
-	#define CONDITIONAL_SET_CBVS(Shader) \
+
+		if (StartStage == SF_Compute)
+		{
+			// Note that ray tracing pipeline shares state with compute
+			CONDITIONAL_SET_CBVS(SF_Compute);
+		}
+		else
+		{
+			CONDITIONAL_SET_CBVS(SF_Vertex);
+			CONDITIONAL_SET_CBVS(SF_Mesh);
+			CONDITIONAL_SET_CBVS(SF_Amplification);
+			CONDITIONAL_SET_CBVS(SF_Geometry);
+			CONDITIONAL_SET_CBVS(SF_Pixel);
+		}
+#undef CONDITIONAL_SET_CBVS
+	}
+#endif // USE_STATIC_ROOT_SIGNATURE
+}
+
+void FD3D12StateCache::ApplyConstants(const FD3D12RootSignature* const pRootSignature, uint32 StartStage, uint32 EndStage)
+{
+#if !USE_STATIC_ROOT_SIGNATURE
+	// Determine what resource bind slots are dirty for the current shaders and how many descriptor table slots we need.
+	// We only set dirty resources that can be used for the upcoming Draw/Dispatch.
+	CBVSlotMask CurrentShaderDirtyCBVSlots[SF_NumStandardFrequencies] = {};
+
+	for (uint32 Stage = StartStage; Stage < EndStage; ++Stage)
+	{
+		if (ShouldSkipStage(Stage))
+		{
+			continue;
+		}
+
+		const CBVSlotMask CurrentShaderCBVRegisterMask = BitMask<CBVSlotMask>(PipelineState.Common.CurrentShaderCBCounts[Stage]);
+		CurrentShaderDirtyCBVSlots[Stage] = CurrentShaderCBVRegisterMask & PipelineState.Common.CBVCache.DirtySlotMask[Stage];
+		// Note: CBVs don't currently use descriptor tables but we still need to know what resource point slots are dirty.
+	}
+
+	// Constant buffers
+	{
+		//SCOPE_CYCLE_COUNTER(STAT_D3D12ApplyStateSetConstantBufferTime);
+		FD3D12ConstantBufferCache& CBVCache = PipelineState.Common.CBVCache;
+
+#define CONDITIONAL_SET_CBVS(Shader) \
 		if (CurrentShaderDirtyCBVSlots[##Shader]) \
 		{ \
-			DescriptorCache.SetConstantBuffers<##Shader>(pRootSignature, CBVCache, CurrentShaderDirtyCBVSlots[##Shader]); \
+			DescriptorCache.SetRootConstantBuffers(##Shader, pRootSignature, CBVCache, CurrentShaderDirtyCBVSlots[##Shader]); \
 		}
-#endif
+
 		if (StartStage == SF_Compute)
 		{
 			// Note that ray tracing pipeline shares state with compute
@@ -743,8 +789,10 @@ void FD3D12StateCache::ApplyResources(const FD3D12RootSignature* const pRootSign
 		}
 #undef CONDITIONAL_SET_CBVS
 	}
+#endif // !USE_STATIC_ROOT_SIGNATURE
 }
 
+
 void FD3D12StateCache::ApplySamplers(const FD3D12RootSignature* const pRootSignature, uint32 StartStage, uint32 EndStage)
 {
 	bool HighLevelCacheMiss = false;
diff --git a/Engine/Source/Runtime/D3D12RHI/Private/D3D12StateCachePrivate.h b/Engine/Source/Runtime/D3D12RHI/Private/D3D12StateCachePrivate.h
--- a/Engine/Source/Runtime/D3D12RHI/Private/D3D12StateCachePrivate.h
+++ b/Engine/Source/Runtime/D3D12RHI/Private/D3D12StateCachePrivate.h
@@ -21,7 +21,7 @@
 #define NUM_SAMPLER_DESCRIPTORS D3D12_MAX_SHADER_VISIBLE_SAMPLER_HEAP_SIZE
 
 // Keep set state functions inline to reduce call overhead
-#define D3D12_STATE_CACHE_INLINE FORCEINLINE
+#define D3D12_STATE_CACHE_INLINE FORCEINLINE_DEBUGGABLE
 
 #if D3D12_STATE_CACHE_RUNTIME_TOGGLE
 extern bool GD3D12SkipStateCaching;
@@ -725,6 +725,7 @@ public:
 	void ApplyState();
 	void ApplySamplers(const FD3D12RootSignature* const pRootSignature, uint32 StartStage, uint32 EndStage);
 	void ApplyResources(const FD3D12RootSignature* const pRootSignature, uint32 StartStage, uint32 EndStage);
+	void ApplyConstants(const FD3D12RootSignature* const pRootSignature, uint32 StartStage, uint32 EndStage);
 	void DirtyStateForNewCommandList();
 	void DirtyState();
 	void DirtyViewDescriptorTables();
