Fix initialization order to make sure we check ray tracing support first, then create ray tracing root signature and then create dispatch rays indirect command signature.

diff --git a/Engine/Source/Runtime/D3D12RHI/Private/D3D12Adapter.cpp b/Engine/Source/Runtime/D3D12RHI/Private/D3D12Adapter.cpp
--- a/Engine/Source/Runtime/D3D12RHI/Private/D3D12Adapter.cpp
+++ b/Engine/Source/Runtime/D3D12RHI/Private/D3D12Adapter.cpp
@@ -1105,8 +1105,6 @@ PRAGMA_ENABLE_DEPRECATION_WARNINGS
 		bTrackAllAllocation = (GD3D12TrackAllAlocations || GPUCrashDebuggingModes == ED3D12GPUCrashDebuggingModes::All) && (GetResourceHeapTier() == D3D12_RESOURCE_HEAP_TIER_2);
 #endif 
 
-		CreateCommandSignatures();
-
 		// Context redirectors allow RHI commands to be executed on multiple GPUs at the
 		// same time in a multi-GPU system. Redirectors have a physical mask for the GPUs
 		// they can support and an active mask which restricts commands to operate on a
@@ -1193,6 +1191,9 @@ PRAGMA_ENABLE_DEPRECATION_WARNINGS
 		StaticRayTracingLocalRootSignature.InitStaticRayTracingLocalRootSignatureDesc();
 #endif
 #endif // USE_STATIC_ROOT_SIGNATURE
+
+		// Creating command signatures relies on static ray tracing root signatures.
+		CreateCommandSignatures();
 	}
 }
 
