Check for nullptr before trying to destroy buddy allocator during shutdown

#rnx

diff --git a/Engine/Source/Runtime/D3D12RHI/Private/D3D12Allocation.cpp b/Engine/Source/Runtime/D3D12RHI/Private/D3D12Allocation.cpp
--- a/Engine/Source/Runtime/D3D12RHI/Private/D3D12Allocation.cpp
+++ b/Engine/Source/Runtime/D3D12RHI/Private/D3D12Allocation.cpp
@@ -783,8 +783,11 @@ void FD3D12MultiBuddyAllocator::ReleaseAllResources()
 {
 	for (int32 i = (Allocators.Num() - 1); i >= 0; i--)
 	{
-		Allocators[i]->Destroy();
-		delete(Allocators[i]);
+		if (Allocators[i])
+		{
+			Allocators[i]->Destroy();
+			delete(Allocators[i]);
+		}
 	}
 
 	Allocators.Empty();
