- Only perform forced patch command list on graphics for async compute when there is an actual transition which needs it (skip for SRV_MASK -> SRV_COMPUTE when there are no other transitions done on the resource in that command list for example)
- Process RHI transitions which are forced during begin because those contain transition from Graphics to Compute and should help remove forced patch up command lists for async compute to run on the graphics queue

Fixes all forced graphics sync from async compute for patch up command lists when running CitySample

diff --git a/Engine/Source/Runtime/D3D12RHI/Private/D3D12CommandList.cpp b/Engine/Source/Runtime/D3D12RHI/Private/D3D12CommandList.cpp
--- a/Engine/Source/Runtime/D3D12RHI/Private/D3D12CommandList.cpp
+++ b/Engine/Source/Runtime/D3D12RHI/Private/D3D12CommandList.cpp
@@ -32,9 +32,9 @@ void FD3D12ContextCommon::AddPendingResourceBarrier(FD3D12Resource* Resource, D3
 	check(After != D3D12_RESOURCE_STATE_TBD);
 	check(Resource->RequiresResourceStateTracking());
 	check(&GetCommandList().GetResourceState_OnCommandList(Resource) == &ResourceState_OnCommandList);
-
+	
 	GetCommandList().State.PendingResourceBarriers.Emplace(Resource, After, SubResource);
-	ResourceState_OnCommandList.SetSubresourceState(SubResource, After);
+	ResourceState_OnCommandList.SetSubresourceState(SubResource, After);	
 }
 
 void FD3D12ContextCommon::AddTransitionBarrier(FD3D12Resource* pResource, D3D12_RESOURCE_STATES Before, D3D12_RESOURCE_STATES After, uint32 Subresource, CResourceState* ResourceState_OnCommandList)
@@ -68,6 +68,7 @@ void FD3D12ContextCommon::AddTransitionBarrier(FD3D12Resource* pResource, D3D12_
 		}
 
 		ResourceState_OnCommandList->SetSubresourceState(Subresource, After);
+		ResourceState_OnCommandList->SetHasInternalTransition();
 	}
 }
 
diff --git a/Engine/Source/Runtime/D3D12RHI/Private/D3D12Commands.cpp b/Engine/Source/Runtime/D3D12RHI/Private/D3D12Commands.cpp
--- a/Engine/Source/Runtime/D3D12RHI/Private/D3D12Commands.cpp
+++ b/Engine/Source/Runtime/D3D12RHI/Private/D3D12Commands.cpp
@@ -522,6 +522,39 @@ void FD3D12CommandContext::HandleResourceTransitions(const FD3D12TransitionData*
 		// Can't rely on before state to either be unknown or UAV because there could be a UAV->SRV transition already done (and ignored here)
 		// and the transition SRV->UAV needs a UAV barrier then to work correctly otherwise there is no synchronization at all
 		bUAVBarrier |= bUAVAccessAfter;
+
+		// Process transitions which are forced during begin because those contain transition from Graphics to Compute and should
+		// help remove forced patch up command lists for async compute to run on the graphics queue
+		if (Info.Resource && ProcessTransitionDuringBegin(TransitionData))
+		{
+			ProcessResource(*this, Info, [&](const FRHITransitionInfo& Info, FD3D12Resource* Resource, FD3D12Texture* UnusedTexture = nullptr)
+			{
+				if (!Resource->RequiresResourceStateTracking())
+				{
+					return;
+				}
+
+				const bool bIsAsyncCompute = EnumHasAnyFlags(TransitionData->DstPipelines, ERHIPipeline::AsyncCompute);
+
+				// Use D3D12_RESOURCE_STATE_TBD as before state for now and don't use provided before state because there is not validation nor handling if the provided state
+				// doesn't match up with the tracked state - this needs to be improved and checked
+				D3D12_RESOURCE_STATES StateBefore = D3D12_RESOURCE_STATE_TBD;
+				D3D12_RESOURCE_STATES StateAfter = Info.AccessAfter == ERHIAccess::Discard ? GetInitialResourceState(Resource->GetDesc()) : GetD3D12ResourceState(Info.AccessAfter, bIsAsyncCompute);
+
+				// enqueue the correct transitions
+				if (Info.IsWholeResource() || Resource->GetSubresourceCount() == 1)
+				{
+					TransitionResource(Resource, StateBefore, StateAfter, D3D12_RESOURCE_BARRIER_ALL_SUBRESOURCES);
+				}
+				else
+				{
+					EnumerateSubresources(Resource, Info, nullptr, [&](uint32 Subresource, const FD3D12RenderTargetView* UnusedRTV = nullptr)
+					{
+						TransitionResource(Resource, StateBefore, StateAfter, Subresource);
+					});
+				}
+			});
+		}
 	}
 }
 
diff --git a/Engine/Source/Runtime/D3D12RHI/Private/D3D12Submission.cpp b/Engine/Source/Runtime/D3D12RHI/Private/D3D12Submission.cpp
--- a/Engine/Source/Runtime/D3D12RHI/Private/D3D12Submission.cpp
+++ b/Engine/Source/Runtime/D3D12RHI/Private/D3D12Submission.cpp
@@ -512,9 +512,30 @@ FD3D12CommandList* FD3D12DynamicRHI::GenerateBarrierCommandListAndUpdateState(FD
 		
 		if (Before != After)
 		{
-			if (IsDirectQueueExclusiveD3D12State(Before) || IsDirectQueueExclusiveD3D12State(After))
+			if (SourceCommandList->QueueType != ED3D12QueueType::Direct)
 			{
-				bHasGraphicStates = true;
+				check(!IsDirectQueueExclusiveD3D12State(After));
+				if (IsDirectQueueExclusiveD3D12State(Before))
+				{
+					// If we are transitioning to a subset of the already set state on the resource (SRV_MASK -> SRV_COMPUTE for example)
+					// and there are no other transitions done on this resource in the command list then this transition can be ignored
+					// (otherwise a transition from SRV_COMPUTE as before state might already be recorded in the command list)
+					CResourceState& ResourceState_OnCommandList = SourceCommandList->GetResourceState_OnCommandList(PRB.Resource);
+					if (ResourceState_OnCommandList.HasInternalTransition() || !EnumHasAllFlags(Before, After))
+					{
+						bHasGraphicStates = true;
+					}
+					else
+					{
+						// should be the same final state as well
+						check(ResourceState_OnCommandList.AreAllSubresourcesSame());
+						check(After == ResourceState_OnCommandList.GetSubresourceState(PRB.SubResource));
+						
+						// Force the original state again if we're skipping this transition
+						ResourceState_OnCommandList.SetResourceState(Before);
+						continue;
+					}
+				}
 			}
 
 			if (PRB.Resource->IsBackBuffer() && EnumHasAnyFlags(After, BackBufferBarrierWriteTransitionTargets))
diff --git a/Engine/Source/Runtime/D3D12RHI/Private/D3D12Util.cpp b/Engine/Source/Runtime/D3D12RHI/Private/D3D12Util.cpp
--- a/Engine/Source/Runtime/D3D12RHI/Private/D3D12Util.cpp
+++ b/Engine/Source/Runtime/D3D12RHI/Private/D3D12Util.cpp
@@ -1438,11 +1438,14 @@ void CResourceState::Initialize(uint32 SubresourceCount)
 	m_SubresourceState.SetNumUninitialized(SubresourceCount);
 	check(m_SubresourceState.Num() == SubresourceCount);
 
+	// No internal transition yet
+	bHasInternalTransition = 0;
+
 	// All subresources start out in an unknown state
 	SetResourceState(D3D12_RESOURCE_STATE_TBD);
 
 	// Unknown hidden resource state
-	UAVHiddenResourceState = D3D12_RESOURCE_STATE_TBD;
+	SetUAVHiddenResourceState(D3D12_RESOURCE_STATE_TBD);
 }
 
 //----------------------------------------------------------------------------------------------------------------------------------
diff --git a/Engine/Source/Runtime/D3D12RHI/Public/D3D12Util.h b/Engine/Source/Runtime/D3D12RHI/Public/D3D12Util.h
--- a/Engine/Source/Runtime/D3D12RHI/Public/D3D12Util.h
+++ b/Engine/Source/Runtime/D3D12RHI/Public/D3D12Util.h
@@ -284,15 +284,25 @@ public:
 
 	D3D12_RESOURCE_STATES GetUAVHiddenResourceState() const
 	{
-		return UAVHiddenResourceState;	
+		return static_cast<D3D12_RESOURCE_STATES>(UAVHiddenResourceState);
 	}
 
 	void SetUAVHiddenResourceState(D3D12_RESOURCE_STATES InUAVHiddenResourceState)
 	{
 		// The hidden state can never include UAV
 		check(InUAVHiddenResourceState == D3D12_RESOURCE_STATE_TBD || !EnumHasAnyFlags(InUAVHiddenResourceState, D3D12_RESOURCE_STATE_UNORDERED_ACCESS));
+		check((InUAVHiddenResourceState & (1 << 31)) == 0);
 
-		UAVHiddenResourceState = InUAVHiddenResourceState;
+		UAVHiddenResourceState = (uint32)InUAVHiddenResourceState;
+	}
+
+	void SetHasInternalTransition()
+	{
+		bHasInternalTransition = 1;
+	}
+	bool HasInternalTransition() const
+	{
+		return bHasInternalTransition != 0;
 	}
 
 private:
@@ -306,7 +316,10 @@ private:
 
 	// Special resource state to track previous state before resource transitioned to UAV when the resource
 	// has a UAV aliasing resource so correct previous state can be found (only single state allowed)
-	D3D12_RESOURCE_STATES UAVHiddenResourceState;
+	uint32 UAVHiddenResourceState : 31;
+
+	// Was the resource used for another transition than the pending transition
+	uint32 bHasInternalTransition : 1;
 
 	// Only used if m_AllSubresourcesSame is 0.
 	// The state of each subresources.  Bits are from D3D12_RESOURCE_STATES.
