Removes GRHISupportsDynamicResolution=true from desktop D3D12 RHI

diff --git a/Engine/Source/Runtime/D3D12RHI/Private/D3D12RHI.cpp b/Engine/Source/Runtime/D3D12RHI/Private/D3D12RHI.cpp
--- a/Engine/Source/Runtime/D3D12RHI/Private/D3D12RHI.cpp
+++ b/Engine/Source/Runtime/D3D12RHI/Private/D3D12RHI.cpp
@@ -247,7 +247,6 @@ FD3D12DynamicRHI::FD3D12DynamicRHI(const TArray<TSharedPtr<FD3D12Adapter>>& Chos
 
 	GRHISupportsMapWriteNoOverwrite = true;
 
-	GRHISupportsDynamicResolution = true;
 	GRHISupportsFrameCyclesBubblesRemoval = true;
 	GRHISupportsGPUTimestampBubblesRemoval = true;
 	GRHISupportsRHIOnTaskThread = true;
