Handle different states per subresource when checking if patch up transition can be skipped

diff --git a/Engine/Source/Runtime/D3D12RHI/Private/D3D12Submission.cpp b/Engine/Source/Runtime/D3D12RHI/Private/D3D12Submission.cpp
--- a/Engine/Source/Runtime/D3D12RHI/Private/D3D12Submission.cpp
+++ b/Engine/Source/Runtime/D3D12RHI/Private/D3D12Submission.cpp
@@ -528,11 +528,10 @@ FD3D12CommandList* FD3D12DynamicRHI::GenerateBarrierCommandListAndUpdateState(FD
 					else
 					{
 						// should be the same final state as well
-						check(ResourceState_OnCommandList.AreAllSubresourcesSame());
 						check(After == ResourceState_OnCommandList.GetSubresourceState(PRB.SubResource));
 						
 						// Force the original state again if we're skipping this transition
-						ResourceState_OnCommandList.SetResourceState(Before);
+						ResourceState_OnCommandList.SetSubresourceState(PRB.SubResource, Before);
 						continue;
 					}
 				}
