Fix compilation error when enabling DEBUG_RESOURCE_STATES. Resource barrier tracking relied on functions that no longer exist: GetBarriers and GetBackBufferBarriers. Rewrote the resource barrier tracking to work with the new version of ResourceBarrierBatcher.

diff --git a/Engine/Source/Runtime/D3D12RHI/Private/D3D12CommandList.cpp b/Engine/Source/Runtime/D3D12RHI/Private/D3D12CommandList.cpp
--- a/Engine/Source/Runtime/D3D12RHI/Private/D3D12CommandList.cpp
+++ b/Engine/Source/Runtime/D3D12RHI/Private/D3D12CommandList.cpp
@@ -95,22 +95,6 @@ void FD3D12ContextCommon::AddAliasingBarrier(ID3D12Resource* InResourceBefore, I
 
 void FD3D12ContextCommon::FlushResourceBarriers()
 {
-#if DEBUG_RESOURCE_STATES
-	// Keep track of all the resource barriers that have been submitted to the current command list.
-	const TArray<D3D12_RESOURCE_BARRIER>& Barriers = ResourceBarrierBatcher.GetBarriers();
-	if (Barriers.Num())
-	{
-		ResourceBarriers.Append(Barriers.GetData(), Barriers.Num());
-	}
-
-	const TArray<D3D12_RESOURCE_BARRIER>& BackBufferBarriers = ResourceBarrierBatcher.GetBackBufferBarriers();
-	if (BackBufferBarriers.Num())
-	{
-		ResourceBarriers.Append(BackBufferBarriers.GetData(), BackBufferBarriers.Num());
-	}
-
-#endif // #if DEBUG_RESOURCE_STATES
-
 	if (ResourceBarrierBatcher.Num())
 	{
 		ResourceBarrierBatcher.FlushIntoCommandList(GetCommandList(), TimestampQueries);
diff --git a/Engine/Source/Runtime/D3D12RHI/Private/D3D12Resources.cpp b/Engine/Source/Runtime/D3D12RHI/Private/D3D12Resources.cpp
--- a/Engine/Source/Runtime/D3D12RHI/Private/D3D12Resources.cpp
+++ b/Engine/Source/Runtime/D3D12RHI/Private/D3D12Resources.cpp
@@ -1029,6 +1029,13 @@ void FD3D12ResourceBarrierBatcher::FlushIntoCommandList(FD3D12CommandList& Comma
 		}
 
 		CommandList.GraphicsCommandList()->ResourceBarrier(BatchEnd - BatchStart, &Barriers[BatchStart]);
+#if DEBUG_RESOURCE_STATES
+		// Keep track of all the resource barriers that have been submitted to the current command list.
+		for(int i = BatchStart; i < BatchEnd - BatchStart - 1; i++)
+		{
+			CommandList.State.ResourceBarriers.Emplace(Barriers[i]);
+		}
+#endif // #if DEBUG_RESOURCE_STATES
 
 		if (bIdle)
 		{
