Fix crash when playing a video with Electra player.
 - D3D12 WaitManualFence function should close the command list if currently open.

diff --git a/Engine/Source/Runtime/D3D12RHI/Private/D3D12CommandContext.cpp b/Engine/Source/Runtime/D3D12RHI/Private/D3D12CommandContext.cpp
--- a/Engine/Source/Runtime/D3D12RHI/Private/D3D12CommandContext.cpp
+++ b/Engine/Source/Runtime/D3D12RHI/Private/D3D12CommandContext.cpp
@@ -311,7 +311,11 @@ void FD3D12ContextCommon::SignalManualFence(ID3D12Fence* Fence, uint64 Value)
 
 void FD3D12ContextCommon::WaitManualFence(ID3D12Fence* Fence, uint64 Value)
 {
-	checkf(!CommandList, TEXT("Command list must be closed before signaling a manual fence."));
+	if (IsOpen())
+	{
+		CloseCommandList();
+	}
+
 	GetPayload(EPhase::Wait)->FencesToWait.Emplace(Fence, Value);
 }
 
