Fix memory stomp caused by D3D12 GPU render queries
 - The interrupt thread writes the resulting query data into a "Target" pointer when the GPU signals completion. That pointer was pointing to the Result member of FD3D12RenderQuery, which may have been deleted already (RHI objects are deleted on the RHI thread, which is further up the pipe than the interrupt thread).
 - Moved the Result member to a heap allocation which is deferred-deleted at the end-of-pipe.

diff --git a/Engine/Source/Runtime/D3D12RHI/Private/D3D12CommandContext.h b/Engine/Source/Runtime/D3D12RHI/Private/D3D12CommandContext.h
--- a/Engine/Source/Runtime/D3D12RHI/Private/D3D12CommandContext.h
+++ b/Engine/Source/Runtime/D3D12RHI/Private/D3D12CommandContext.h
@@ -39,6 +39,7 @@ struct FD3D12DeferredDeleteObject
 		RHIObject,
 		D3DObject,
 		BindlessDescriptor,
+		CPUAllocation
 	} Type;
 
 	union
@@ -51,6 +52,8 @@ struct FD3D12DeferredDeleteObject
 			FRHIDescriptorHandle Handle;
 			FD3D12Device* Device;
 		} BindlessDescriptor;
+
+		void* CPUAllocation;
 	};
 
 	explicit FD3D12DeferredDeleteObject(FD3D12Resource* RHIObject)
@@ -67,6 +70,13 @@ struct FD3D12DeferredDeleteObject
 		: Type(EType::BindlessDescriptor)
 		, BindlessDescriptor({ Handle, Device })
 	{}
+
+	explicit FD3D12DeferredDeleteObject(void* Ptr, EType Type)
+		: Type(Type)
+		, CPUAllocation(Ptr)
+	{
+		check(Type == EType::CPUAllocation);
+	}
 };
 
 enum class ED3D12Units
diff --git a/Engine/Source/Runtime/D3D12RHI/Private/D3D12Query.cpp b/Engine/Source/Runtime/D3D12RHI/Private/D3D12Query.cpp
--- a/Engine/Source/Runtime/D3D12RHI/Private/D3D12Query.cpp
+++ b/Engine/Source/Runtime/D3D12RHI/Private/D3D12Query.cpp
@@ -32,6 +32,18 @@ namespace D3D12RHI
 }
 using namespace D3D12RHI;
 
+FD3D12RenderQuery::FD3D12RenderQuery(FD3D12Device* Parent, ERenderQueryType InQueryType)
+	: FD3D12DeviceChild(Parent)
+	, Type(InQueryType)
+	, Result((uint64*)FMemory::Malloc(sizeof(uint64), alignof(uint64)))
+{}
+
+FD3D12RenderQuery::~FD3D12RenderQuery()
+{
+	FD3D12DynamicRHI::GetD3DRHI()->DeferredDelete(Result, FD3D12DeferredDeleteObject::EType::CPUAllocation);
+	Result = nullptr;
+}
+
 FD3D12QueryHeap::FD3D12QueryHeap(FD3D12Device* Device, D3D12_QUERY_TYPE QueryType, D3D12_QUERY_HEAP_TYPE HeapType, uint32 NumQueries)
 	: FD3D12SingleNodeGPUObject(Device->GetGPUMask())
 	, Device(Device)
@@ -190,7 +202,7 @@ void FD3D12CommandContext::RHIBeginRenderQuery(FRHIRenderQuery* QueryRHI)
 	FD3D12RenderQuery* Query = RetrieveObject<FD3D12RenderQuery>(QueryRHI);
 	checkf(Query->Type == RQT_Occlusion, TEXT("Only occlusion queries support RHIBeginRenderQuery()."));
 
-	Query->ActiveLocation = AllocateQuery(ED3D12QueryType::Occlusion, &Query->Result);
+	Query->ActiveLocation = AllocateQuery(ED3D12QueryType::Occlusion, Query->Result);
 	BeginQuery(Query->ActiveLocation);
 
 	ActiveQueries++;
@@ -241,7 +253,7 @@ void FD3D12CommandContext::RHIEndRenderQuery(FRHIRenderQuery* QueryRHI)
 		break;
 
 	case RQT_AbsoluteTime:
-		InsertTimestamp(ED3D12Units::Microseconds, &Query->Result);
+		InsertTimestamp(ED3D12Units::Microseconds, Query->Result);
 		break;
 	}
 }
@@ -297,7 +309,7 @@ bool FD3D12DynamicRHI::RHIGetRenderQueryResult(FRHIRenderQuery* QueryRHI, uint64
 		}
 	}
 
-	OutResult = Query->Result;
+	OutResult = *Query->Result;
 	return true;
 }
 
diff --git a/Engine/Source/Runtime/D3D12RHI/Private/D3D12Query.h b/Engine/Source/Runtime/D3D12RHI/Private/D3D12Query.h
--- a/Engine/Source/Runtime/D3D12RHI/Private/D3D12Query.h
+++ b/Engine/Source/Runtime/D3D12RHI/Private/D3D12Query.h
@@ -9,19 +9,17 @@ D3D12Query.h: Implementation of D3D12 Query
 class FD3D12RenderQuery : public FRHIRenderQuery, public FD3D12DeviceChild, public FD3D12LinkedAdapterObject<FD3D12RenderQuery>
 {
 public:
-	/** Initialization constructor. */
-	FD3D12RenderQuery(FD3D12Device* Parent, ERenderQueryType InQueryType)
-		: FD3D12DeviceChild(Parent)
-		, Type(InQueryType)
-	{}
+	FD3D12RenderQuery(FD3D12Device* Parent, ERenderQueryType InQueryType);
+	~FD3D12RenderQuery();
 
 	ERenderQueryType const Type;
 
 	// Signaled when the result is available. Nullptr if the query has never been used.
 	FD3D12SyncPointRef SyncPoint;
 
-	// The query result, read from the GPU on the interrupt thread.
-	uint64 Result = 0;
+	// The query result, read from the GPU. Heap allocated since it is
+	// accessed by the interrupt thread, and needs to outlive the RHI object.
+	uint64* Result;
 	
 	// The current query location for occlusion queries.
 	FD3D12QueryLocation ActiveLocation;
diff --git a/Engine/Source/Runtime/D3D12RHI/Private/D3D12RHI.cpp b/Engine/Source/Runtime/D3D12RHI/Private/D3D12RHI.cpp
--- a/Engine/Source/Runtime/D3D12RHI/Private/D3D12RHI.cpp
+++ b/Engine/Source/Runtime/D3D12RHI/Private/D3D12RHI.cpp
@@ -482,6 +482,9 @@ void FD3D12DynamicRHI::RHIPerFrameRHIFlushComplete()
 				case FD3D12DeferredDeleteObject::EType::BindlessDescriptor:
 					ObjectToDelete.BindlessDescriptor.Device->GetBindlessDescriptorManager().ImmediateFree(ObjectToDelete.BindlessDescriptor.Handle);
 					break;
+				case FD3D12DeferredDeleteObject::EType::CPUAllocation:
+					FMemory::Free(ObjectToDelete.CPUAllocation);
+					break;
 				default:
 					checkf(false, TEXT("Unknown ED3D12DeferredDeleteObjectType"));
 					break;
