Virtual Production:  Optimizations to help with D3D12 CPU and GPU performance regressions following Parallel Rendering integration:
* Feature that merges command lists from QueueAsyncCommandListSubmit into a single Payload where possible (r.D3D12.AllowPayloadMerge, default enabled).  Saves around 0.5 ms per async command list batch, with total savings from 1 to 4 ms on VP test scenes.
* Remove unnecessary command flush before Present on D3D12.  Added a virtual function "NeedFlushBeforeEndDrawing" that can return false if the platform function RHIEndDrawingViewport function flushes commands, and the calling function doesn't need to.
* Wrapped three scene rendering flushes, plus BeginFlushResourcesRHI with "RHIIncludeOptionalFlushes()", which returns false on D3D12.  These were a net perf loss on D3D12, as they wrap hardly any rendering.  I didn't want to change behavior on other platforms, hence the conditional.
* Overall the above changes remove 8 command list flushes in a two view scene, which typically cost in the range of 0.06 ms each, saving 0.48 ms CPU total.  Perf win increases with more views, and removing the flushes also helps with GPU bubbles.
* Added a flush in a strategic location at the end of Shadows and Lumen.  In a sample scene, saved 1.5 ms GPU eliminating a bubble.

#rnx

diff --git a/Engine/Source/Runtime/D3D12RHI/Private/D3D12RHIPrivate.h b/Engine/Source/Runtime/D3D12RHI/Private/D3D12RHIPrivate.h
--- a/Engine/Source/Runtime/D3D12RHI/Private/D3D12RHIPrivate.h
+++ b/Engine/Source/Runtime/D3D12RHI/Private/D3D12RHIPrivate.h
@@ -658,6 +658,7 @@ public:
 	virtual void HandleGpuTimeout(FD3D12Payload* Payload, double SecondsSinceSubmission);
 
 	bool RHIRequiresComputeGenerateMips() const override { return true; };
+	bool RHIIncludeOptionalFlushes() const override { return false; }
 
 	bool IsQuadBufferStereoEnabled() const;
 	void DisableQuadBufferStereo();
diff --git a/Engine/Source/Runtime/D3D12RHI/Private/D3D12Submission.cpp b/Engine/Source/Runtime/D3D12RHI/Private/D3D12Submission.cpp
--- a/Engine/Source/Runtime/D3D12RHI/Private/D3D12Submission.cpp
+++ b/Engine/Source/Runtime/D3D12RHI/Private/D3D12Submission.cpp
@@ -15,6 +15,14 @@
 // When enabled, GPU timestamp queries are adjusted to remove idle time caused by CPU bubbles.
 #define D3D12_ENABLE_ADJUSTED_TIMESTAMPS 0		// Adjusted timestamps break Unreal Insights, disabling for now
 
+static bool GD3D12AsyncPayloadMerge = true;
+static FAutoConsoleVariableRef CVarD3D12AsyncPayloadMerge(
+	TEXT("r.D3D12.AllowPayloadMerge"),
+	GD3D12AsyncPayloadMerge,
+	TEXT("Whether to attempt to merge command lists into a single payload, saving perf.  Mainly applies to QueueAsyncCommandListSubmit.  (Default = 1)\n"),
+	ECVF_RenderThreadSafe
+);
+
 DECLARE_CYCLE_STAT(TEXT("Submit"), STAT_D3D12Submit, STATGROUP_D3D12RHI);
 
 DECLARE_CYCLE_STAT(TEXT("GPU Total Time [All Queues]"), STAT_RHI_GPUTotalTime, STATGROUP_D3D12RHI);
@@ -170,6 +178,85 @@ void FD3D12DynamicRHI::SubmitCommands(TConstArrayView<FD3D12FinalizedCommands*>
 {
 	SCOPED_NAMED_EVENT_TEXT("CommandList_Submit", FColor::Magenta);
 
+	if ((Commands.Num() > 1) && GD3D12AsyncPayloadMerge)
+	{
+		// See if we can merge the payloads.  All need to be on the direct queue, have no waits, no fence signals
+		// (SyncPointsToSignal are OK), and not involve events or timing.  Trivial cases where the payload
+		// contents can't be dependent on one another.  Intended to optimize groups of command lists generated by
+		// QueueAsyncCommandListSubmit.
+		bool bCanMergePayloads = true;
+		for (IRHIPlatformCommandList* Ptr : Commands)
+		{
+			FD3D12FinalizedCommands* Payloads = static_cast<FD3D12FinalizedCommands*>(Ptr);
+			for (FD3D12Payload* Payload : *Payloads)
+			{
+				if (Payload->Queue.QueueType != ED3D12QueueType::Direct ||
+					Payload->SyncPointsToWait.Num() ||
+					Payload->FencesToWait.Num() ||
+					Payload->FencesToSignal.Num() ||
+					Payload->Timing.IsSet() ||
+					Payload->CompletionFenceValue ||
+					Payload->SubmissionEvent ||
+					Payload->SubmissionTime.IsSet() ||
+					Payload->TimestampQueries.Num() ||
+					Payload->OcclusionQueries.Num())
+				{
+					bCanMergePayloads = false;
+					break;
+				}
+			}
+		}
+
+		if (bCanMergePayloads)
+		{
+			TArray<FD3D12Payload*> MergedPayloadPerGPU;
+			MergedPayloadPerGPU.AddZeroed(GNumExplicitGPUsForRendering);
+
+			for (IRHIPlatformCommandList* Ptr : Commands)
+			{
+				FD3D12FinalizedCommands* Payloads = static_cast<FD3D12FinalizedCommands*>(Ptr);
+				for (FD3D12Payload* Payload : *Payloads)
+				{
+					uint32 GPUIndex = Payload->Queue.Device->GetGPUIndex();
+					if (!MergedPayloadPerGPU[GPUIndex])
+					{
+						MergedPayloadPerGPU[GPUIndex] = new FD3D12Payload(Payload->Queue.Device, ED3D12QueueType::Direct);
+					}
+					FD3D12Payload* MergedPayload = MergedPayloadPerGPU[GPUIndex];
+
+					MergedPayload->CommandListsToExecute.Append(Payload->CommandListsToExecute);
+					MergedPayload->SyncPointsToSignal.Append(Payload->SyncPointsToSignal);
+					MergedPayload->AllocatorsToRelease.Append(Payload->AllocatorsToRelease);
+					MergedPayload->QueryRanges.Append(Payload->QueryRanges);
+
+					// Need to clear out allocator array, so Payload destructor doesn't delete them
+					Payload->AllocatorsToRelease.Empty();
+					delete Payload;
+				}
+				delete Payloads;
+			}
+
+			// Remove any NULL elements from the per-GPU array
+			for (int32 PayloadIndex = 0; PayloadIndex < MergedPayloadPerGPU.Num();)
+			{
+				if (!MergedPayloadPerGPU[PayloadIndex])
+				{
+					// Removing current element, don't increment index
+					MergedPayloadPerGPU.RemoveAt(PayloadIndex);
+				}
+				else
+				{
+					// Otherwise, increment index
+					PayloadIndex++;
+				}
+			}
+
+			SubmitPayloads(MergedPayloadPerGPU);
+
+			return;
+		}
+	}
+
 	for (IRHIPlatformCommandList* Ptr : Commands)
 	{
 		FD3D12FinalizedCommands* Payloads = static_cast<FD3D12FinalizedCommands*>(Ptr);
diff --git a/Engine/Source/Runtime/D3D12RHI/Public/D3D12Viewport.h b/Engine/Source/Runtime/D3D12RHI/Public/D3D12Viewport.h
--- a/Engine/Source/Runtime/D3D12RHI/Public/D3D12Viewport.h
+++ b/Engine/Source/Runtime/D3D12RHI/Public/D3D12Viewport.h
@@ -103,6 +103,12 @@ public:
 	virtual void WaitForFrameEventCompletion() override;
 	virtual void IssueFrameEvent() override;
 
+	/**
+	  * RHIEndDrawingViewport flushes commands internally, so we don't need to flush before calling it.
+	  * Avoids wasteful issuing of another entire command list with just a few commands.
+	  */
+	virtual bool NeedFlushBeforeEndDrawing() { return false; }
+
 #if D3D12_VIEWPORT_EXPOSES_SWAP_CHAIN
 	virtual void* GetNativeSwapChain() const override { return SwapChain1; }
 #endif // #if D3D12_VIEWPORT_EXPOSES_SWAP_CHAIN
diff --git a/Engine/Source/Runtime/RHI/Private/RHICommandList.cpp b/Engine/Source/Runtime/RHI/Private/RHICommandList.cpp
--- a/Engine/Source/Runtime/RHI/Private/RHICommandList.cpp
+++ b/Engine/Source/Runtime/RHI/Private/RHICommandList.cpp
@@ -598,6 +598,8 @@ void FRHICommandListImmediate::QueueAsyncCommandListSubmit(TArrayView<FQueuedCom
 		// The task blocks for each parallel translate completion, in the order they will be submitted in.
 		EnqueueLambda([Tasks](FRHICommandListBase&)
 		{
+			TArray<IRHIPlatformCommandList*> AllCmdLists;
+
 			for (FTask& Task : Tasks)
 			{
 				if (!Task.Event->IsComplete())
@@ -608,13 +610,15 @@ void FRHICommandListImmediate::QueueAsyncCommandListSubmit(TArrayView<FQueuedCom
 					Task.Event->Wait();
 				}
 
-				if (Task.OutCmdLists.Num())
-				{
-					GDynamicRHI->RHISubmitCommandLists(Task.OutCmdLists);
-				}
+				AllCmdLists.Append(Task.OutCmdLists);
 
 				Task.~FTask();
 			}
+
+			if (AllCmdLists.Num())
+			{
+				GDynamicRHI->RHISubmitCommandLists(AllCmdLists);
+			}
 		});
 	}
 	else
@@ -1170,7 +1174,10 @@ void FRHICommandList::EndDrawingViewport(FRHIViewport* Viewport, bool bPresent,
 	// Make sure all prior graphics and async compute work has been submitted.
 	// This is necessary because platform RHIs often submit additional work on the graphics queue during present, and we need to ensure we won't deadlock on async work that wasn't yet submitted by the renderer.
 	// In future, Present() itself should be an enqueued / recorded command, and platform RHIs should never implicitly submit graphics or async compute work.
-	FRHICommandListExecutor::GetImmediateCommandList().ImmediateFlush(EImmediateFlushType::DispatchToRHIThread);
+	if (Viewport->NeedFlushBeforeEndDrawing())
+	{
+		FRHICommandListExecutor::GetImmediateCommandList().ImmediateFlush(EImmediateFlushType::DispatchToRHIThread);
+	}
 
 	check(IsImmediate() && IsInRenderingThread());
 	if (Bypass())
diff --git a/Engine/Source/Runtime/RHI/Public/DynamicRHI.h b/Engine/Source/Runtime/RHI/Public/DynamicRHI.h
--- a/Engine/Source/Runtime/RHI/Public/DynamicRHI.h
+++ b/Engine/Source/Runtime/RHI/Public/DynamicRHI.h
@@ -1147,6 +1147,8 @@ public:
 
 	virtual bool RHIRequiresComputeGenerateMips() const { return false; };
 
+	virtual bool RHIIncludeOptionalFlushes() const { return true; }
+
 #if RHI_RAYTRACING
 
 	virtual FRayTracingAccelerationStructureSize RHICalcRayTracingSceneSize(uint32 MaxInstances, ERayTracingAccelerationStructureFlags Flags)
diff --git a/Engine/Source/Runtime/RHI/Public/RHIResources.h b/Engine/Source/Runtime/RHI/Public/RHIResources.h
--- a/Engine/Source/Runtime/RHI/Public/RHIResources.h
+++ b/Engine/Source/Runtime/RHI/Public/RHIResources.h
@@ -2213,6 +2213,8 @@ public:
 	virtual void WaitForFrameEventCompletion() { }
 
 	virtual void IssueFrameEvent() { }
+
+	virtual bool NeedFlushBeforeEndDrawing() { return true; }
 };
 
 //
diff --git a/Engine/Source/Runtime/RenderCore/Private/RenderGraphBuilder.cpp b/Engine/Source/Runtime/RenderCore/Private/RenderGraphBuilder.cpp
--- a/Engine/Source/Runtime/RenderCore/Private/RenderGraphBuilder.cpp
+++ b/Engine/Source/Runtime/RenderCore/Private/RenderGraphBuilder.cpp
@@ -359,7 +359,11 @@ void FRDGBuilder::BeginFlushResourcesRHI()
 
 	CSV_SCOPED_TIMING_STAT_EXCLUSIVE(STAT_RDG_FlushResourcesRHI);
 	SCOPED_NAMED_EVENT(BeginFlushResourcesRHI, FColor::Emerald);
-	RHICmdList.ImmediateFlush(EImmediateFlushType::DispatchToRHIThread);
+
+	if (GDynamicRHI->RHIIncludeOptionalFlushes())
+	{
+		RHICmdList.ImmediateFlush(EImmediateFlushType::DispatchToRHIThread);
+	}
 }
 
 void FRDGBuilder::EndFlushResourcesRHI()
diff --git a/Engine/Source/Runtime/Renderer/Private/DeferredShadingRenderer.cpp b/Engine/Source/Runtime/Renderer/Private/DeferredShadingRenderer.cpp
--- a/Engine/Source/Runtime/Renderer/Private/DeferredShadingRenderer.cpp
+++ b/Engine/Source/Runtime/Renderer/Private/DeferredShadingRenderer.cpp
@@ -2386,7 +2386,10 @@ void FDeferredShadingSceneRenderer::Render(FRDGBuilder& GraphBuilder)
 	FILCUpdatePrimTaskData ILCTaskData;
 
 	// Find the visible primitives.
-	GraphBuilder.RHICmdList.ImmediateFlush(EImmediateFlushType::DispatchToRHIThread);
+	if (GDynamicRHI->RHIIncludeOptionalFlushes())
+	{
+		GraphBuilder.RHICmdList.ImmediateFlush(EImmediateFlushType::DispatchToRHIThread);
+	}
 
 	FInstanceCullingManager& InstanceCullingManager = *GraphBuilder.AllocObject<FInstanceCullingManager>(Scene->GPUScene.IsEnabled(), GraphBuilder);
 
@@ -3298,6 +3301,9 @@ void FDeferredShadingSceneRenderer::Render(FRDGBuilder& GraphBuilder)
 	}
 	// End shadow and fog after base pass
 
+	// Trigger a command submit here, to avoid GPU bubbles
+	AddDispatchToRHIThreadPass(GraphBuilder);
+	
 	if (bNaniteEnabled)
 	{
 		// Needs doing after shadows such that the checks for shadow atlases etc work.
diff --git a/Engine/Source/Runtime/Renderer/Private/SceneVisibility.cpp b/Engine/Source/Runtime/Renderer/Private/SceneVisibility.cpp
--- a/Engine/Source/Runtime/Renderer/Private/SceneVisibility.cpp
+++ b/Engine/Source/Runtime/Renderer/Private/SceneVisibility.cpp
@@ -5242,7 +5242,10 @@ void FDeferredShadingSceneRenderer::InitViews(FRDGBuilder& GraphBuilder, const F
 		OnStartRender(RHICmdList);
 	}
 
-	RHICmdList.ImmediateFlush(EImmediateFlushType::DispatchToRHIThread);
+	if (GDynamicRHI->RHIIncludeOptionalFlushes())
+	{
+		RHICmdList.ImmediateFlush(EImmediateFlushType::DispatchToRHIThread);
+	}
 }
 
 template<class T>
@@ -5354,7 +5357,10 @@ void FDeferredShadingSceneRenderer::InitViewsAfterPrepass(FRDGBuilder& GraphBuil
 			InitDynamicShadows(RHICmdList, DynamicIndexBufferForInitShadows, DynamicVertexBufferForInitShadows, DynamicReadBufferForInitShadows, InstanceCullingManager);
 		}
 
-		RHICmdList.ImmediateFlush(EImmediateFlushType::DispatchToRHIThread);
+		if (GDynamicRHI->RHIIncludeOptionalFlushes())
+		{
+			RHICmdList.ImmediateFlush(EImmediateFlushType::DispatchToRHIThread);
+		}
 	}
 
 	// If parallel ILC update is disabled, then process it in place.
