Log invalid resource descriptions before they inevitably crash.

diff --git a/Engine/Source/Runtime/D3D12RHI/Private/D3D12Device.cpp b/Engine/Source/Runtime/D3D12RHI/Private/D3D12Device.cpp
--- a/Engine/Source/Runtime/D3D12RHI/Private/D3D12Device.cpp
+++ b/Engine/Source/Runtime/D3D12RHI/Private/D3D12Device.cpp
@@ -555,6 +555,12 @@ D3D12_RESOURCE_ALLOCATION_INFO FD3D12Device::GetResourceAllocationInfo(const FD3
 #endif
 		{
 			Result = GetDevice()->GetResourceAllocationInfo(0, 1, &InDesc);
+		    if (Result.SizeInBytes == UINT64_MAX)
+		    {
+			    // The description provided caused an error per the docs. This will almost certainly crash outside this fn.
+			    UE_LOG(LogD3D12RHI, Error, TEXT("D3D12 GetResourceAllocationInfo failed - likely a resource was requested that has invalid allocation info (e.g. is an invalid texture size)"));
+			    UE_LOG(LogD3D12RHI, Error, TEXT("     W %llu H %d depth %d mips %d pf %d"), InDesc.Width, InDesc.Height, InDesc.DepthOrArraySize, InDesc.MipLevels, InDesc.PixelFormat);
+			}
 		}
 
 		ResourceAllocationInfoMapMutex.WriteLock();
