Remove deprecated RHIRayTraceOcclusion and RHIRayTraceIntersection methods. They are replaced with equivalent new high-level functions in RayTracingBasicShaders.cpp.

diff --git a/Engine/Shaders/Private/RayTracing/RayTracingBuiltInShaders.usf b/Engine/Shaders/Private/RayTracing/RayTracingBuiltInShaders.usf
--- a/Engine/Shaders/Private/RayTracing/RayTracingBuiltInShaders.usf
+++ b/Engine/Shaders/Private/RayTracing/RayTracingBuiltInShaders.usf
@@ -1,26 +1,15 @@
 // Copyright Epic Games, Inc. All Rights Reserved.
 
-#ifdef OVERRIDE_RAYTRACINGBUILTINSHADERS_USH
-#include "/Platform/Private/RayTracingBuiltInShaders.usf"
-#else // OVERRIDE_RAYTRACINGBUILTINSHADERS_USH
 
 #include "../Common.ush"
 #include "RayTracingCommon.ush"
 
-#define RT_CONCATENATE2(a, b) a##b
-#define RT_CONCATENATE(a, b) RT_CONCATENATE2(a, b)
-#define RT_REGISTER(InType, InIndex, InSpace) register(RT_CONCATENATE(InType, InIndex), RT_CONCATENATE(space, InSpace))
+RaytracingAccelerationStructure               TLAS;
+StructuredBuffer<FBasicRayData>               Rays;
+RWStructuredBuffer<uint>                      OcclusionOutput;
+RWStructuredBuffer<FIntersectionPayload>      IntersectionOutput;
 
-RaytracingAccelerationStructure               TLAS               : RT_REGISTER(t, 0, RAY_TRACING_REGISTER_SPACE_GLOBAL);
-StructuredBuffer<FBasicRayData>               Rays               : RT_REGISTER(t, 1, RAY_TRACING_REGISTER_SPACE_GLOBAL);
-RWStructuredBuffer<uint>                      OcclusionOutput    : RT_REGISTER(u, 0, RAY_TRACING_REGISTER_SPACE_GLOBAL);
-RWStructuredBuffer<FIntersectionPayload>      IntersectionOutput : RT_REGISTER(u, 1, RAY_TRACING_REGISTER_SPACE_GLOBAL);
-
-#undef RT_REGISTER
-#undef RT_CONCATENATE
-#undef RT_CONCATENATE2
-
-RAY_TRACING_ENTRY_RAYGEN(OcclusionMainRG)
+RAY_TRACING_ENTRY_RAYGEN(OcclusionMainRGS)
 {
 	const uint RayIndex = DispatchRaysIndex().x;
 	FBasicRayData InputRay = Rays[RayIndex];
@@ -55,7 +44,7 @@ RAY_TRACING_ENTRY_RAYGEN(OcclusionMainRG)
 	OcclusionOutput[RayIndex] = ((FMinimalPayload)Payload).IsHit() ? ~0 : 0;
 }
 
-RAY_TRACING_ENTRY_RAYGEN(IntersectionMainRG)
+RAY_TRACING_ENTRY_RAYGEN(IntersectionMainRGS)
 {
 	const uint RayIndex = DispatchRaysIndex().x;
 	FBasicRayData InputRay = Rays[RayIndex];
@@ -120,5 +109,3 @@ RAY_TRACING_ENTRY_MISS(PackedMaterialClosestHitPayloadMS, FPackedMaterialClosest
 {
 	Payload.SetMiss();
 }
-
-#endif // PLATFORM_RAY_TRACING_BUILTIN_SHADERS_USH
diff --git a/Engine/Shaders/Shared/RayTracingBuiltInResources.h b/Engine/Shaders/Shared/RayTracingBuiltInResources.h
--- a/Engine/Shaders/Shared/RayTracingBuiltInResources.h
+++ b/Engine/Shaders/Shared/RayTracingBuiltInResources.h
@@ -77,9 +77,11 @@ struct FHitGroupSystemRootConstants
 #define RAY_TRACING_SYSTEM_VERTEXBUFFER_REGISTER 1
 #define RAY_TRACING_SYSTEM_ROOTCONSTANT_REGISTER 0
 
-#if INCLUDED_FROM_HLSL_CODE
+#ifndef OVERRIDE_RAY_TRACING_HIT_GROUP_SYSTEM_RESOURCES
+#define OVERRIDE_RAY_TRACING_HIT_GROUP_SYSTEM_RESOURCES 0
+#endif
 
-#ifndef OVERRIDE_RAYTRACINGBUILTINSHADERS_USH
+#if INCLUDED_FROM_HLSL_CODE && !OVERRIDE_RAY_TRACING_HIT_GROUP_SYSTEM_RESOURCES
 	#define RT_CONCATENATE2(a, b) a##b
 	#define RT_CONCATENATE(a, b) RT_CONCATENATE2(a, b)
 	#define RT_REGISTER(InType, InIndex, InSpace) register(RT_CONCATENATE(InType, InIndex), RT_CONCATENATE(space, InSpace))
@@ -90,9 +92,7 @@ struct FHitGroupSystemRootConstants
 	#undef RT_REGISTER
 	#undef RT_CONCATENATE
 	#undef RT_CONCATENATE2
-#endif // !OVERRIDE_RAYTRACINGBUILTINSHADERS_USH
-
-#endif // INCLUDED_FROM_HLSL_CODE
+#endif // INCLUDED_FROM_HLSL_CODE && !OVERRIDE_RAY_TRACING_HIT_GROUP_SYSTEM_RESOURCES
 
 
 #undef INCLUDED_FROM_CPP_CODE
diff --git a/Engine/Source/Runtime/D3D12RHI/Private/D3D12CommandContext.h b/Engine/Source/Runtime/D3D12RHI/Private/D3D12CommandContext.h
--- a/Engine/Source/Runtime/D3D12RHI/Private/D3D12CommandContext.h
+++ b/Engine/Source/Runtime/D3D12RHI/Private/D3D12CommandContext.h
@@ -605,14 +605,6 @@ public:
 	virtual void RHIBuildAccelerationStructures(const TArrayView<const FRayTracingGeometryBuildParams> Params, const FRHIBufferRange& ScratchBufferRange) final override;
 	virtual void RHIBuildAccelerationStructure(const FRayTracingSceneBuildParams& SceneBuildParams) final override;
 	virtual void RHIClearRayTracingBindings(FRHIRayTracingScene* Scene) final override;
-	virtual void RHIRayTraceOcclusion(FRHIRayTracingScene* Scene,
-		FRHIShaderResourceView* Rays,
-		FRHIUnorderedAccessView* Output,
-		uint32 NumRays) final override;
-	virtual void RHIRayTraceIntersection(FRHIRayTracingScene* Scene,
-		FRHIShaderResourceView* Rays,
-		FRHIUnorderedAccessView* Output,
-		uint32 NumRays) final override;
 	virtual void RHIRayTraceDispatch(FRHIRayTracingPipelineState* RayTracingPipelineState, FRHIRayTracingShader* RayGenShader,
 		FRHIRayTracingScene* Scene,
 		const FRayTracingShaderBindings& GlobalResourceBindings,
@@ -989,26 +981,6 @@ public:
 		ContextRedirect(RHIBuildAccelerationStructure(SceneBuildParams));
 	}
 
-	virtual void RHIRayTraceOcclusion(FRHIRayTracingScene* Scene,
-		FRHIShaderResourceView* Rays,
-		FRHIUnorderedAccessView* Output,
-		uint32 NumRays) final override
-	{
-		PRAGMA_DISABLE_DEPRECATION_WARNINGS
-		ContextRedirect(RHIRayTraceOcclusion(Scene, Rays, Output, NumRays));
-		PRAGMA_ENABLE_DEPRECATION_WARNINGS
-	}
-
-	virtual void RHIRayTraceIntersection(FRHIRayTracingScene* Scene,
-		FRHIShaderResourceView* Rays,
-		FRHIUnorderedAccessView* Output,
-		uint32 NumRays) final override
-	{
-		PRAGMA_DISABLE_DEPRECATION_WARNINGS
-		ContextRedirect(RHIRayTraceIntersection(Scene, Rays, Output, NumRays));
-		PRAGMA_ENABLE_DEPRECATION_WARNINGS
-	}
-
 	virtual void RHIRayTraceDispatch(FRHIRayTracingPipelineState* RayTracingPipelineState, FRHIRayTracingShader* RayGenShader,
 		FRHIRayTracingScene* Scene,
 		const FRayTracingShaderBindings& GlobalResourceBindings,
diff --git a/Engine/Source/Runtime/D3D12RHI/Private/D3D12Device.h b/Engine/Source/Runtime/D3D12RHI/Private/D3D12Device.h
--- a/Engine/Source/Runtime/D3D12RHI/Private/D3D12Device.h
+++ b/Engine/Source/Runtime/D3D12RHI/Private/D3D12Device.h
@@ -11,7 +11,6 @@ D3D12Device.h: D3D12 Device Interfaces
 
 class FD3D12Device;
 class FD3D12DynamicRHI;
-class FD3D12BasicRayTracingPipeline;
 class FD3D12Buffer;
 class FD3D12RayTracingDescriptorHeapCache;
 class FD3D12RayTracingPipelineCache;
@@ -284,7 +283,6 @@ public:
 	ID3D12Device5*							  GetDevice5();
 	ID3D12Device7*							  GetDevice7();
 
-	const FD3D12BasicRayTracingPipeline*	  GetBasicRayTracingPipeline           ();
 	FD3D12RayTracingDescriptorHeapCache*	  GetRayTracingDescriptorHeapCache     () { return RayTracingDescriptorHeapCache;      }
 	FD3D12RayTracingPipelineCache*			  GetRayTracingPipelineCache           () { return RayTracingPipelineCache;            }
 	FD3D12Buffer*							  GetRayTracingDispatchRaysDescBuffer  () { return RayTracingDispatchRaysDescBuffer;   }
@@ -401,7 +399,6 @@ private:
 	FD3D12TextureAllocatorPool   TextureAllocator;
 
 #if D3D12_RHI_RAYTRACING
-	FD3D12BasicRayTracingPipeline* BasicRayTracingPipeline = nullptr;
 	FD3D12RayTracingPipelineCache* RayTracingPipelineCache = nullptr;
 	FD3D12RayTracingCompactionRequestHandler* RayTracingCompactionRequestHandler = nullptr;
 	FD3D12Buffer* RayTracingDispatchRaysDescBuffer = nullptr;
diff --git a/Engine/Source/Runtime/D3D12RHI/Private/D3D12RayTracing.cpp b/Engine/Source/Runtime/D3D12RHI/Private/D3D12RayTracing.cpp
--- a/Engine/Source/Runtime/D3D12RHI/Private/D3D12RayTracing.cpp
+++ b/Engine/Source/Runtime/D3D12RHI/Private/D3D12RayTracing.cpp
@@ -2734,35 +2734,6 @@ public:
 			CallableShaders.Identifiers[ShaderIndex] = GetEntryShaderIdentifier(CallableShaderEntries[ShaderIndex]);
 		}
 
-		// Setup default shader binding table, which simply includes all provided RGS and MS plus a single default closest hit shader.
-		// Hit record indexing and local resources access is disabled when using using this SBT.
-
-		FD3D12RayTracingShaderTable::FInitializer SBTInitializer = {};
-		SBTInitializer.NumRayGenShaders = RayGenShaders.Identifiers.Num();
-		SBTInitializer.NumMissShaders = MissShaders.Identifiers.Num();
-		SBTInitializer.NumCallableRecords = 0; // Default SBT does not support callable shaders
-		SBTInitializer.NumMissRecords = 1;
-		SBTInitializer.NumHitRecords = 0; // Default SBT does not support indexable hit shaders
-		SBTInitializer.LocalRootDataSize = 0; // Shaders in default SBT are not allowed to access any local resources
-
-		// Create default shader tables for every node in the LDA group
-		FD3D12Adapter* Adapter = Device->GetParentAdapter();
-
-		for (uint32 GPUIndex : FRHIGPUMask::All())
-		{
-			FD3D12Device* NodeDevice = Adapter->GetDevice(GPUIndex);
-			DefaultShaderTables[GPUIndex].Init(
-				SBTInitializer,
-				NodeDevice,
-				RayGenShaders.Identifiers,
-				HitGroupShaders.Identifiers[0]);
-
-			if (MissShaders.Identifiers.Num())
-			{
-				DefaultShaderTables[GPUIndex].SetMissIdentifier(0, MissShaders.Identifiers[0]);
-			}			
-		}
-
 		PipelineStackSize = PipelineProperties->GetPipelineStackSize();
 
 		TotalCreationTime += FPlatformTime::Cycles64();
@@ -2812,11 +2783,6 @@ public:
 	FD3D12RayTracingShaderLibrary HitGroupShaders;
 	FD3D12RayTracingShaderLibrary CallableShaders;
 
-	// Shader table that can be used to dispatch ray tracing work that doesn't require real SBT bindings.
-	// This is useful for the case where user only provides default RayGen, Miss and HitGroup shaders.
-	// Currently only used by RayTraceOcclusion/Intersection.
-	FD3D12RayTracingShaderTable DefaultShaderTables[MAX_NUM_GPUS];
-
 	ID3D12RootSignature* GlobalRootSignature = nullptr;
 
 	TRefCountPtr<ID3D12StateObject> StateObject;
@@ -2850,66 +2816,6 @@ public:
 #endif // !NO_LOGGING
 };
 
-class FD3D12BasicRayTracingPipeline
-{
-public:
-
-	UE_NONCOPYABLE(FD3D12BasicRayTracingPipeline)
-
-	FD3D12BasicRayTracingPipeline(FD3D12Device* Device)
-	{
-		// Occlusion pipeline
-		{
-			PRAGMA_DISABLE_DEPRECATION_WARNINGS
-
-			FRayTracingPipelineStateInitializer OcclusionInitializer;
-
-			FRHIRayTracingShader* OcclusionRGSTable[] = { GetBuildInRayTracingShader<FOcclusionMainRG>() };
-			OcclusionInitializer.SetRayGenShaderTable(OcclusionRGSTable);
-
-			FRHIRayTracingShader* OcclusionMSTable[] = { GetBuildInRayTracingShader<FDefaultPayloadMS>() };
-			OcclusionInitializer.SetMissShaderTable(OcclusionMSTable);
-
-			OcclusionInitializer.bAllowHitGroupIndexing = false;
-
-			Occlusion = new FD3D12RayTracingPipelineState(Device, OcclusionInitializer);
-
-			PRAGMA_ENABLE_DEPRECATION_WARNINGS
-		}
-
-		// Intersection pipeline
-		{
-			PRAGMA_DISABLE_DEPRECATION_WARNINGS
-
-			FRayTracingPipelineStateInitializer IntersectionInitializer;
-
-			FRHIRayTracingShader* IntersectionRGSTable[] = { GetBuildInRayTracingShader<FIntersectionMainRG>() };
-			IntersectionInitializer.SetRayGenShaderTable(IntersectionRGSTable);
-
-			FRHIRayTracingShader* IntersectionMSTable[] = { GetBuildInRayTracingShader<FDefaultPayloadMS>() };
-			IntersectionInitializer.SetMissShaderTable(IntersectionMSTable);
-
-			FRHIRayTracingShader* IntersectionHitTable[] = { GetBuildInRayTracingShader<FIntersectionMainCHS>() };
-			IntersectionInitializer.SetHitGroupTable(IntersectionHitTable);
-
-			IntersectionInitializer.bAllowHitGroupIndexing = false;
-
-			Intersection = new FD3D12RayTracingPipelineState(Device, IntersectionInitializer);
-
-			PRAGMA_ENABLE_DEPRECATION_WARNINGS
-		}
-	}
-
-	~FD3D12BasicRayTracingPipeline()
-	{
-		delete Intersection;
-		delete Occlusion;
-	}
-
-	FD3D12RayTracingPipelineState* Occlusion;
-	FD3D12RayTracingPipelineState* Intersection;
-};
-
 struct FD3D12RHICommandInitializeRayTracingGeometryString
 {
 	static const TCHAR* TStr() { return TEXT("FD3D12RHICommandInitializeRayTracingGeometryString"); }
@@ -2971,25 +2877,10 @@ void FD3D12Device::InitRayTracing()
 
 	check(RayTracingDescriptorHeapCache == nullptr);
 	RayTracingDescriptorHeapCache = new FD3D12RayTracingDescriptorHeapCache(this);
-
-	check(BasicRayTracingPipeline == nullptr);
-	// the pipeline will be initialized on the first use
-}
-
-const FD3D12BasicRayTracingPipeline* FD3D12Device::GetBasicRayTracingPipeline()
-{
-	if (UNLIKELY(BasicRayTracingPipeline == nullptr))
-	{
-		BasicRayTracingPipeline = new FD3D12BasicRayTracingPipeline(this);
-	}
-	return BasicRayTracingPipeline;
 }
 
 void FD3D12Device::CleanupRayTracing()
 {
-	delete BasicRayTracingPipeline;
-	BasicRayTracingPipeline = nullptr;
-
 	delete RayTracingPipelineCache;
 	RayTracingPipelineCache = nullptr;
 
@@ -5380,82 +5271,6 @@ static void DispatchRays(FD3D12CommandContext& CommandContext,
 	CommandContext.StateCache.GetDescriptorCache()->RestoreAfterExternalHeapsSet();
 }
 
-
-void FD3D12CommandContext::RHIRayTraceOcclusion(FRHIRayTracingScene* InScene,
-	FRHIShaderResourceView* Rays,
-	FRHIUnorderedAccessView* Output,
-	uint32 NumRays)
-{
-	checkf(GetParentDevice()->GetBasicRayTracingPipeline(), TEXT("Ray tracing support is not initialized for this device. Ensure that InitRayTracing() is called before issuing any ray tracing work."));
-
-	FD3D12RayTracingScene* Scene = FD3D12DynamicRHI::ResourceCast(InScene);
-
-	FD3D12RayTracingPipelineState* Pipeline = GetParentDevice()->GetBasicRayTracingPipeline()->Occlusion;
-	FD3D12RayTracingShaderTable& ShaderTable = Pipeline->DefaultShaderTables[GetGPUIndex()];
-
-	if (ShaderTable.bIsDirty)
-	{
-		ShaderTable.CopyToGPU(*this);
-	}
-
-	Scene->UpdateResidency(*this);
-
-	D3D12_DISPATCH_RAYS_DESC DispatchDesc = ShaderTable.GetDispatchRaysDesc(0, 0, 0);
-
-	DispatchDesc.Width = NumRays;
-	DispatchDesc.Height = 1;
-	DispatchDesc.Depth = 1;
-
-	FRayTracingShaderBindings Bindings;
-	Bindings.SRVs[0] = Scene->Layers[0].ShaderResourceView; // TODO: Expose layers
-	Bindings.SRVs[1] = Rays;
-	Bindings.UAVs[0] = Output;
-
-	ShaderTable.UpdateResidency(*this);
-
-	DispatchRays(*this, Bindings, Pipeline, 0, nullptr, DispatchDesc);
-}
-
-void FD3D12CommandContext::RHIRayTraceIntersection(FRHIRayTracingScene* InScene,
-	FRHIShaderResourceView* InRays,
-	FRHIUnorderedAccessView* InOutput,
-	uint32 NumRays)
-{
-	checkf(GetParentDevice()->GetBasicRayTracingPipeline(), TEXT("Ray tracing support is not initialized for this device. Ensure that InitRayTracing() is called before issuing any ray tracing work."));
-
-	FD3D12RayTracingScene* Scene = FD3D12DynamicRHI::ResourceCast(InScene);
-	FD3D12ShaderResourceView* Rays = FD3D12DynamicRHI::ResourceCast(InRays);
-	FD3D12UnorderedAccessView* Output = FD3D12DynamicRHI::ResourceCast(InOutput);
-
-	FD3D12RayTracingPipelineState* Pipeline = GetParentDevice()->GetBasicRayTracingPipeline()->Intersection;
-	FD3D12RayTracingShaderTable& ShaderTable = Pipeline->DefaultShaderTables[GetGPUIndex()];
-
-	if (ShaderTable.bIsDirty)
-	{
-		ShaderTable.CopyToGPU(*this);
-	}
-
-	Scene->UpdateResidency(*this);
-
-	D3D12_DISPATCH_RAYS_DESC DispatchDesc = ShaderTable.GetDispatchRaysDesc(0, 0, 0);
-
-	DispatchDesc.Width = NumRays;
-	DispatchDesc.Height = 1;
-	DispatchDesc.Depth = 1;
-
-	FRayTracingShaderBindings Bindings;
-	Bindings.SRVs[0] = Scene->Layers[0].ShaderResourceView; // TODO: Expose layers
-	Bindings.SRVs[1] = Rays;
-	// #dxr_todo: intersection and occlusion shaders should be split into separate files to avoid resource slot collisions.
-	// Workaround for now is to bind a valid UAV to slots 0 and 1, even though only slot 1 is referenced.
-	Bindings.UAVs[0] = Output;
-	Bindings.UAVs[1] = Output;
-
-	ShaderTable.UpdateResidency(*this);
-
-	DispatchRays(*this, Bindings, Pipeline, 0, nullptr, DispatchDesc);
-}
-
 void FD3D12CommandContext::RHIRayTraceDispatch(FRHIRayTracingPipelineState* InRayTracingPipelineState, FRHIRayTracingShader* RayGenShaderRHI,
 	FRHIRayTracingScene* InScene,
 	const FRayTracingShaderBindings& GlobalResourceBindings,
diff --git a/Engine/Source/Runtime/RHI/Public/RHICommandList.h b/Engine/Source/Runtime/RHI/Public/RHICommandList.h
--- a/Engine/Source/Runtime/RHI/Public/RHICommandList.h
+++ b/Engine/Source/Runtime/RHI/Public/RHICommandList.h
@@ -197,24 +197,6 @@ enum class ERayTracingBindingType : uint8
 	CallableShader,
 	MissShader,
 };
-
-// C++ counter-part of FBasicRayData declared in RayTracingCommon.ush
-struct UE_DEPRECATED(5.1, "Please use an explicit ray generation shader with a custom intersection structure instead.") FBasicRayData
-{
-	float Origin[3];
-	uint32 Mask;
-	float Direction[3];
-	float TFar;
-};
-
-// C++ counter-part of FIntersectionPayload declared in RayTracingCommon.ush
-struct UE_DEPRECATED(5.1, "Please use an explicit ray generation shader with a custom intersection structure instead.") FIntersectionPayload
-{
-	float  HitT;            // Distance from ray origin to the intersection point in the ray direction. Negative on miss.
-	uint32 PrimitiveIndex;  // Index of the primitive within the geometry inside the bottom-level acceleration structure instance. Undefined on miss.
-	uint32 InstanceIndex;   // Index of the current instance in the top-level structure. Undefined on miss.
-	float  Barycentrics[2]; // Primitive barycentric coordinates of the intersection point. Undefined on miss.
-};
 #endif // RHI_RAYTRACING
 
 struct RHI_API FLockTracker
@@ -2255,46 +2237,6 @@ struct FRHICommandBuildAccelerationStructures final : public FRHICommand<FRHICom
 	RHI_API void Execute(FRHICommandListBase& CmdList);
 };
 
-FRHICOMMAND_MACRO(FRHICommandRayTraceOcclusion)
-{
-	FRHIRayTracingScene* Scene;
-	FRHIShaderResourceView* Rays;
-	FRHIUnorderedAccessView* Output;
-	uint32 NumRays;
-
-	FRHICommandRayTraceOcclusion(FRHIRayTracingScene* InScene,
-		FRHIShaderResourceView* InRays,
-		FRHIUnorderedAccessView* InOutput,
-		uint32 InNumRays)
-		: Scene(InScene)
-		, Rays(InRays)
-		, Output(InOutput)
-		, NumRays(InNumRays)
-	{}
-
-	RHI_API void Execute(FRHICommandListBase& CmdList);
-};
-
-FRHICOMMAND_MACRO(FRHICommandRayTraceIntersection)
-{
-	FRHIRayTracingScene* Scene;
-	FRHIShaderResourceView* Rays;
-	FRHIUnorderedAccessView* Output;
-	uint32 NumRays;
-
-	FRHICommandRayTraceIntersection(FRHIRayTracingScene* InScene,
-		FRHIShaderResourceView* InRays,
-		FRHIUnorderedAccessView* InOutput,
-		uint32 InNumRays)
-		: Scene(InScene)
-		, Rays(InRays)
-		, Output(InOutput)
-		, NumRays(InNumRays)
-	{}
-
-	RHI_API void Execute(FRHICommandListBase& CmdList);
-};
-
 FRHICOMMAND_MACRO(FRHICommandRayTraceDispatch)
 {
 	FRayTracingPipelineState* Pipeline;
@@ -3768,51 +3710,6 @@ public:
 		}
 	}
 
-	/**
-	 * Trace rays from an input buffer of FBasicRayData.
-	 * Binary intersection results are written to output buffer as R32_UINTs.
-	 * 0xFFFFFFFF is written if ray intersects any scene triangle, 0 otherwise.
-	 */
-	UE_DEPRECATED(5.1, "Please use an explicit ray generation shader and RayTraceDispatch() instead.")
-	FORCEINLINE_DEBUGGABLE void RayTraceOcclusion(FRHIRayTracingScene* Scene,
-		FRHIShaderResourceView* Rays,
-		FRHIUnorderedAccessView* Output,
-		uint32 NumRays)
-	{
-		PRAGMA_DISABLE_DEPRECATION_WARNINGS
-		if (Bypass())
-		{
-			GetContext().RHIRayTraceOcclusion(Scene, Rays, Output, NumRays);
-		}
-		else
-		{
-			ALLOC_COMMAND(FRHICommandRayTraceOcclusion)(Scene, Rays, Output, NumRays);
-		}
-		PRAGMA_ENABLE_DEPRECATION_WARNINGS
-	}
-
-	/**
-	 * Trace rays from an input buffer of FBasicRayData.
-	 * Primitive intersection results are written to output buffer as FIntersectionPayload.
-	 */
-	UE_DEPRECATED(5.1, "Please use an explicit ray generation shader and RayTraceDispatch() instead.")
-	FORCEINLINE_DEBUGGABLE void RayTraceIntersection(FRHIRayTracingScene* Scene,
-		FRHIShaderResourceView* Rays,
-		FRHIUnorderedAccessView* Output,
-		uint32 NumRays)
-	{
-		PRAGMA_DISABLE_DEPRECATION_WARNINGS
-		if (Bypass())
-		{
-			GetContext().RHIRayTraceIntersection(Scene, Rays, Output, NumRays);
-		}
-		else
-		{
-			ALLOC_COMMAND(FRHICommandRayTraceIntersection)(Scene, Rays, Output, NumRays);
-		}
-		PRAGMA_ENABLE_DEPRECATION_WARNINGS
-	}
-
 	FORCEINLINE_DEBUGGABLE void RayTraceDispatch(FRayTracingPipelineState* Pipeline, FRHIRayTracingShader* RayGenShader, FRHIRayTracingScene* Scene, const FRayTracingShaderBindings& GlobalResourceBindings, uint32 Width, uint32 Height)
 	{
 		if (Bypass())
diff --git a/Engine/Source/Runtime/RHI/Public/RHICommandListCommandExecutes.inl b/Engine/Source/Runtime/RHI/Public/RHICommandListCommandExecutes.inl
--- a/Engine/Source/Runtime/RHI/Public/RHICommandListCommandExecutes.inl
+++ b/Engine/Source/Runtime/RHI/Public/RHICommandListCommandExecutes.inl
@@ -54,8 +54,6 @@ struct FRHICommandTransitionTexturesPipeline;
 struct FRHICommandTransitionTexturesDepth;
 struct FRHICommandTransitionTexturesArray;
 struct FRHICommandClearRayTracingBindings;
-struct FRHICommandRayTraceOcclusion;
-struct FRHICommandRayTraceIntersection;
 struct FRHICommandRayTraceDispatch;
 struct FRHICommandSetRayTracingBindings;
 
@@ -593,22 +591,6 @@ void FRHICommandBuildAccelerationStructures::Execute(FRHICommandListBase& CmdLis
 	INTERNAL_DECORATOR_COMPUTE(RHIBuildAccelerationStructures)(Params, ScratchBufferRange);
 }
 
-void FRHICommandRayTraceOcclusion::Execute(FRHICommandListBase& CmdList)
-{
-	RHISTAT(RayTraceOcclusion);
-PRAGMA_DISABLE_DEPRECATION_WARNINGS
-	INTERNAL_DECORATOR(RHIRayTraceOcclusion)(Scene, Rays, Output, NumRays);
-PRAGMA_ENABLE_DEPRECATION_WARNINGS
-}
-
-void FRHICommandRayTraceIntersection::Execute(FRHICommandListBase& CmdList)
-{
-	RHISTAT(RayTraceIntersection);
-PRAGMA_DISABLE_DEPRECATION_WARNINGS
-	INTERNAL_DECORATOR(RHIRayTraceIntersection)(Scene, Rays, Output, NumRays);
-PRAGMA_ENABLE_DEPRECATION_WARNINGS
-}
-
 void FRHICommandRayTraceDispatch::Execute(FRHICommandListBase& CmdList)
 {
 	RHISTAT(RayTraceDispatch);
diff --git a/Engine/Source/Runtime/RHI/Public/RHIContext.h b/Engine/Source/Runtime/RHI/Public/RHIContext.h
--- a/Engine/Source/Runtime/RHI/Public/RHIContext.h
+++ b/Engine/Source/Runtime/RHI/Public/RHIContext.h
@@ -797,24 +797,6 @@ public:
 		checkNoEntry();
 	}
 
-	UE_DEPRECATED(5.1, "Please use an explicit ray generation shader and RHIRayTraceDispatch() instead.")
-	virtual void RHIRayTraceOcclusion(FRHIRayTracingScene* Scene,
-		FRHIShaderResourceView* Rays,
-		FRHIUnorderedAccessView* Output,
-		uint32 NumRays)
-	{
-		checkNoEntry();
-	}
-
-	UE_DEPRECATED(5.1, "Please use an explicit ray generation shader and RHIRayTraceDispatch() instead.")
-	virtual void RHIRayTraceIntersection(FRHIRayTracingScene* Scene,
-		FRHIShaderResourceView* Rays,
-		FRHIUnorderedAccessView* Output,
-		uint32 NumRays)
-	{
-		checkNoEntry();
-	}
-
 	virtual void RHIRayTraceDispatch(FRHIRayTracingPipelineState* RayTracingPipelineState, FRHIRayTracingShader* RayGenShader,
 		FRHIRayTracingScene* Scene,
 		const FRayTracingShaderBindings& GlobalResourceBindings,
diff --git a/Engine/Source/Runtime/RHI/Public/RHIValidationContext.h b/Engine/Source/Runtime/RHI/Public/RHIValidationContext.h
--- a/Engine/Source/Runtime/RHI/Public/RHIValidationContext.h
+++ b/Engine/Source/Runtime/RHI/Public/RHIValidationContext.h
@@ -1025,26 +1025,6 @@ public:
 		RHIContext->RHIBindAccelerationStructureMemory(Scene, Buffer, BufferOffset);
 	}
 
-	virtual void RHIRayTraceOcclusion(FRHIRayTracingScene* Scene,
-		FRHIShaderResourceView* Rays,
-		FRHIUnorderedAccessView* Output,
-		uint32 NumRays) override final
-	{
-		PRAGMA_DISABLE_DEPRECATION_WARNINGS
-		RHIContext->RHIRayTraceOcclusion(Scene, Rays, Output, NumRays);
-		PRAGMA_ENABLE_DEPRECATION_WARNINGS
-	}
-
-	virtual void RHIRayTraceIntersection(FRHIRayTracingScene* Scene,
-		FRHIShaderResourceView* Rays,
-		FRHIUnorderedAccessView* Output,
-		uint32 NumRays) override final
-	{
-		PRAGMA_DISABLE_DEPRECATION_WARNINGS
-		RHIContext->RHIRayTraceIntersection(Scene, Rays, Output, NumRays);
-		PRAGMA_ENABLE_DEPRECATION_WARNINGS
-	}
-
 	virtual void RHIRayTraceDispatch(FRHIRayTracingPipelineState* RayTracingPipelineState, FRHIRayTracingShader* RayGenShader,
 		FRHIRayTracingScene* Scene,
 		const FRayTracingShaderBindings& GlobalResourceBindings,
diff --git a/Engine/Source/Runtime/RenderCore/Private/BuiltInRayTracingShaders.cpp b/Engine/Source/Runtime/RenderCore/Private/BuiltInRayTracingShaders.cpp
--- a/Engine/Source/Runtime/RenderCore/Private/BuiltInRayTracingShaders.cpp
+++ b/Engine/Source/Runtime/RenderCore/Private/BuiltInRayTracingShaders.cpp
@@ -6,15 +6,9 @@
 
 #if RHI_RAYTRACING
 
-PRAGMA_DISABLE_DEPRECATION_WARNINGS
-IMPLEMENT_GLOBAL_RAYTRACING_SHADER( FOcclusionMainRG,		"/Engine/Private/RayTracing/RayTracingBuiltInShaders.usf",		"OcclusionMainRG",				SF_RayGen, ERayTracingPayloadType::Minimal);
-IMPLEMENT_GLOBAL_RAYTRACING_SHADER( FIntersectionMainRG,	"/Engine/Private/RayTracing/RayTracingBuiltInShaders.usf",		"IntersectionMainRG",			SF_RayGen, ERayTracingPayloadType::Minimal);
-IMPLEMENT_GLOBAL_RAYTRACING_SHADER( FIntersectionMainCHS,	"/Engine/Private/RayTracing/RayTracingBuiltInShaders.usf", "IntersectionMainCHS",	SF_RayHitGroup, ERayTracingPayloadType::Minimal);
-PRAGMA_ENABLE_DEPRECATION_WARNINGS
-
-IMPLEMENT_GLOBAL_RAYTRACING_SHADER( FDefaultMainCHS,		"/Engine/Private/RayTracing/RayTracingBuiltInShaders.usf", "DefaultMainCHS",		SF_RayHitGroup, ERayTracingPayloadType::Minimal);
-IMPLEMENT_GLOBAL_RAYTRACING_SHADER( FDefaultMainCHSOpaqueAHS, "/Engine/Private/RayTracing/RayTracingBuiltInShaders.usf", "closesthit=DefaultMainCHS anyhit=DefaultOpaqueAHS", SF_RayHitGroup, ERayTracingPayloadType::Minimal);
-IMPLEMENT_GLOBAL_RAYTRACING_SHADER( FDefaultPayloadMS,		"/Engine/Private/RayTracing/RayTracingBuiltInShaders.usf", "DefaultPayloadMS",			SF_RayMiss, ERayTracingPayloadType::Minimal);
+IMPLEMENT_GLOBAL_RAYTRACING_SHADER( FDefaultMainCHS, "/Engine/Private/RayTracing/RayTracingBuiltInShaders.usf", "DefaultMainCHS", SF_RayHitGroup, ERayTracingPayloadType::Default);
+IMPLEMENT_GLOBAL_RAYTRACING_SHADER( FDefaultMainCHSOpaqueAHS, "/Engine/Private/RayTracing/RayTracingBuiltInShaders.usf", "closesthit=DefaultMainCHS anyhit=DefaultOpaqueAHS", SF_RayHitGroup, ERayTracingPayloadType::Default);
+IMPLEMENT_GLOBAL_RAYTRACING_SHADER( FDefaultPayloadMS, "/Engine/Private/RayTracing/RayTracingBuiltInShaders.usf", "DefaultPayloadMS", SF_RayMiss, ERayTracingPayloadType::Default);
 IMPLEMENT_GLOBAL_RAYTRACING_SHADER( FPackedMaterialClosestHitPayloadMS, "/Engine/Private/RayTracing/RayTracingBuiltInShaders.usf", "PackedMaterialClosestHitPayloadMS", SF_RayMiss, ERayTracingPayloadType::RayTracingMaterial);
 
 IMPLEMENT_GLOBAL_SHADER(FRayTracingDispatchDescCS, "/Engine/Private/RayTracing/RayTracingDispatchDesc.usf", "RayTracingDispatchDescCS", SF_Compute);
diff --git a/Engine/Source/Runtime/RenderCore/Public/BuiltInRayTracingShaders.h b/Engine/Source/Runtime/RenderCore/Public/BuiltInRayTracingShaders.h
--- a/Engine/Source/Runtime/RenderCore/Public/BuiltInRayTracingShaders.h
+++ b/Engine/Source/Runtime/RenderCore/Public/BuiltInRayTracingShaders.h
@@ -28,41 +28,6 @@ public:
 	{}
 };
 
-class RENDERCORE_API UE_DEPRECATED(5.1, "Please use an explicit ray generation shader instead.") FOcclusionMainRG : public FBuiltInRayTracingShader
-{
-	DECLARE_GLOBAL_SHADER(FOcclusionMainRG);
-	SHADER_USE_ROOT_PARAMETER_STRUCT(FOcclusionMainRG, FBuiltInRayTracingShader);
-
-	BEGIN_SHADER_PARAMETER_STRUCT(FParameters, )
-		SHADER_PARAMETER_RDG_BUFFER_SRV(RaytracingAccelerationStructure, TLAS)
-		SHADER_PARAMETER_RDG_BUFFER_SRV(StructuredBuffer<FBasicRayData>, Rays)
-		SHADER_PARAMETER_RDG_BUFFER_UAV(RWStructuredBuffer<uint>, OcclusionOutput)
-	END_SHADER_PARAMETER_STRUCT()
-};
-
-class RENDERCORE_API UE_DEPRECATED(5.1, "Please use an explicit ray generation shader instead.") FIntersectionMainRG : public FBuiltInRayTracingShader
-{
-	DECLARE_GLOBAL_SHADER(FIntersectionMainRG);
-	SHADER_USE_ROOT_PARAMETER_STRUCT(FIntersectionMainRG, FBuiltInRayTracingShader);
-
-	BEGIN_SHADER_PARAMETER_STRUCT(FParameters, )
-		SHADER_PARAMETER_RDG_BUFFER_SRV(RaytracingAccelerationStructure, TLAS)
-		SHADER_PARAMETER_RDG_BUFFER_SRV(StructuredBuffer<FBasicRayData>, Rays)
-		SHADER_PARAMETER_RDG_BUFFER_UAV(RWStructuredBuffer<FIntersectionPayload>, IntersectionOutput)
-	END_SHADER_PARAMETER_STRUCT()
-};
-
-class UE_DEPRECATED(5.1, "Please use an explicit ray generation and hit shaders instead.") FIntersectionMainCHS : public FBuiltInRayTracingShader
-{
-	DECLARE_EXPORTED_SHADER_TYPE(FIntersectionMainCHS, Global, RENDERCORE_API);
-public:
-
-	FIntersectionMainCHS() = default;
-	FIntersectionMainCHS(const ShaderMetaType::CompiledShaderInitializerType& Initializer)
-		: FBuiltInRayTracingShader(Initializer)
-	{}
-};
-
 class FDefaultMainCHS : public FBuiltInRayTracingShader
 {
 	DECLARE_EXPORTED_SHADER_TYPE(FDefaultMainCHS, Global, RENDERCORE_API);
diff --git a/Engine/Source/Runtime/Renderer/Private/RayTracing/RayTracingBasicShaders.cpp b/Engine/Source/Runtime/Renderer/Private/RayTracing/RayTracingBasicShaders.cpp
new file mode 100644
--- /dev/null
+++ b/Engine/Source/Runtime/Renderer/Private/RayTracing/RayTracingBasicShaders.cpp
@@ -0,0 +1,89 @@
+// Copyright Epic Games, Inc. All Rights Reserved.
+
+#include "RayTracingBasicShaders.h"
+
+#if RHI_RAYTRACING
+
+#include "PipelineStateCache.h"
+#include "ShaderParameterStruct.h"
+
+IMPLEMENT_GLOBAL_RAYTRACING_SHADER( FBasicOcclusionMainRGS, "/Engine/Private/RayTracing/RayTracingBuiltInShaders.usf", "OcclusionMainRGS", SF_RayGen, ERayTracingPayloadType::Default);
+IMPLEMENT_GLOBAL_RAYTRACING_SHADER( FBasicIntersectionMainRGS, "/Engine/Private/RayTracing/RayTracingBuiltInShaders.usf", "IntersectionMainRGS", SF_RayGen, ERayTracingPayloadType::Default);
+IMPLEMENT_GLOBAL_RAYTRACING_SHADER( FBasicIntersectionMainCHS, "/Engine/Private/RayTracing/RayTracingBuiltInShaders.usf", "IntersectionMainCHS", SF_RayHitGroup, ERayTracingPayloadType::Default);
+
+struct FBasicRayTracingPipeline
+{
+	FRayTracingPipelineState* PipelineState = nullptr;
+	TShaderRef<FBasicOcclusionMainRGS> OcclusionRGS;
+	TShaderRef<FBasicIntersectionMainRGS> IntersectionRGS;
+};
+
+/**
+* Returns a ray tracing pipeline with FBasicOcclusionMainRGS, FBasicIntersectionMainRGS, FBasicIntersectionMainCHS and FDefaultPayloadMS.
+* This can be used to perform basic "fixed function" occlusion and intersection ray tracing.
+*/
+FBasicRayTracingPipeline GetBasicRayTracingPipeline(FRHICommandList& RHICmdList, ERHIFeatureLevel::Type FeatureLevel)
+{
+	FGlobalShaderMap* ShaderMap = GetGlobalShaderMap(FeatureLevel);
+
+	FRayTracingPipelineStateInitializer PipelineInitializer;
+	PipelineInitializer.bAllowHitGroupIndexing = false;
+
+	auto OcclusionRGS = ShaderMap->GetShader<FBasicOcclusionMainRGS>();
+	auto IntersectionRGS = ShaderMap->GetShader<FBasicIntersectionMainRGS>();
+
+	FRHIRayTracingShader* RayGenShaderTable[] = { OcclusionRGS.GetRayTracingShader(), IntersectionRGS.GetRayTracingShader() };
+	PipelineInitializer.SetRayGenShaderTable(RayGenShaderTable);
+
+	auto ClosestHitShader = ShaderMap->GetShader<FBasicIntersectionMainCHS>();
+	FRHIRayTracingShader* HitShaderTable[] = { ClosestHitShader.GetRayTracingShader() };
+	PipelineInitializer.SetHitGroupTable(HitShaderTable);
+
+	auto MissShader = ShaderMap->GetShader<FDefaultPayloadMS>();
+	FRHIRayTracingShader* MissShaderTable[] = { MissShader.GetRayTracingShader() };
+	PipelineInitializer.SetMissShaderTable(MissShaderTable);
+
+	FBasicRayTracingPipeline Result;
+
+	Result.PipelineState = PipelineStateCache::GetAndOrCreateRayTracingPipelineState(RHICmdList, PipelineInitializer);
+	Result.OcclusionRGS = OcclusionRGS;
+	Result.IntersectionRGS = IntersectionRGS;
+
+	return Result;
+}
+
+void DispatchBasicOcclusionRays(FRHICommandList& RHICmdList, FRHIRayTracingScene* Scene, FRHIShaderResourceView* SceneView, FRHIShaderResourceView* RayBufferView, FRHIUnorderedAccessView* ResultView, uint32 NumRays)
+{
+	FBasicRayTracingPipeline RayTracingPipeline = GetBasicRayTracingPipeline(RHICmdList, GMaxRHIFeatureLevel);
+
+	RHICmdList.SetRayTracingHitGroup(Scene, 0, 0, 0, RayTracingPipeline.PipelineState, 0, 0, nullptr, 0, nullptr, 0);
+	RHICmdList.SetRayTracingMissShader(Scene, 0, RayTracingPipeline.PipelineState, 0, 0, nullptr, 0);
+
+	FBasicOcclusionMainRGS::FParameters OcclusionParameters;
+	OcclusionParameters.TLAS = SceneView;
+	OcclusionParameters.Rays = RayBufferView;
+	OcclusionParameters.OcclusionOutput = ResultView;
+
+	FRayTracingShaderBindingsWriter GlobalResources;
+	SetShaderParameters(GlobalResources, RayTracingPipeline.OcclusionRGS, OcclusionParameters);
+	RHICmdList.RayTraceDispatch(RayTracingPipeline.PipelineState, RayTracingPipeline.OcclusionRGS.GetRayTracingShader(), Scene, GlobalResources, NumRays, 1);
+}
+
+void DispatchBasicIntersectionRays(FRHICommandList& RHICmdList, FRHIRayTracingScene* Scene, FRHIShaderResourceView* SceneView, FRHIShaderResourceView* RayBufferView, FRHIUnorderedAccessView* ResultView, uint32 NumRays)
+{
+	FBasicRayTracingPipeline RayTracingPipeline = GetBasicRayTracingPipeline(RHICmdList, GMaxRHIFeatureLevel);
+
+	RHICmdList.SetRayTracingHitGroup(Scene, 0, 0, 0, RayTracingPipeline.PipelineState, 0, 0, nullptr, 0, nullptr, 0);
+	RHICmdList.SetRayTracingMissShader(Scene, 0, RayTracingPipeline.PipelineState, 0, 0, nullptr, 0);
+
+	FBasicIntersectionMainRGS::FParameters OcclusionParameters;
+	OcclusionParameters.TLAS = SceneView;
+	OcclusionParameters.Rays = RayBufferView;
+	OcclusionParameters.IntersectionOutput = ResultView;
+
+	FRayTracingShaderBindingsWriter GlobalResources;
+	SetShaderParameters(GlobalResources, RayTracingPipeline.IntersectionRGS, OcclusionParameters);
+	RHICmdList.RayTraceDispatch(RayTracingPipeline.PipelineState, RayTracingPipeline.IntersectionRGS.GetRayTracingShader(), Scene, GlobalResources, NumRays, 1);
+}
+
+#endif // RHI_RAYTRACING
diff --git a/Engine/Source/Runtime/Renderer/Private/RayTracing/RayTracingBasicShaders.h b/Engine/Source/Runtime/Renderer/Private/RayTracing/RayTracingBasicShaders.h
new file mode 100644
--- /dev/null
+++ b/Engine/Source/Runtime/Renderer/Private/RayTracing/RayTracingBasicShaders.h
@@ -0,0 +1,75 @@
+// Copyright Epic Games, Inc. All Rights Reserved.
+
+#pragma once
+
+#include "BuiltInRayTracingShaders.h"
+
+#if RHI_RAYTRACING
+
+// C++ counter-part of FBasicRayData declared in RayTracingCommon.ush
+struct FBasicRayTracingRay
+{
+	float Origin[3];
+	uint32 Mask;
+	float Direction[3];
+	float TFar;
+};
+
+// C++ counter-part of FIntersectionPayload declared in RayTracingCommon.ush
+struct FBasicRayTracingIntersectionResult
+{
+	float  HitT;            // Distance from ray origin to the intersection point in the ray direction. Negative on miss.
+	uint32 PrimitiveIndex;  // Index of the primitive within the geometry inside the bottom-level acceleration structure instance. Undefined on miss.
+	uint32 InstanceIndex;   // Index of the current instance in the top-level structure. Undefined on miss.
+	float  Barycentrics[2]; // Primitive barycentric coordinates of the intersection point. Undefined on miss.
+};
+
+class FBasicOcclusionMainRGS : public FBuiltInRayTracingShader
+{
+	DECLARE_GLOBAL_SHADER(FBasicOcclusionMainRGS);
+	SHADER_USE_ROOT_PARAMETER_STRUCT(FBasicOcclusionMainRGS, FBuiltInRayTracingShader);
+
+	BEGIN_SHADER_PARAMETER_STRUCT(FParameters, )
+		SHADER_PARAMETER_SRV(RaytracingAccelerationStructure, TLAS)
+		SHADER_PARAMETER_SRV(StructuredBuffer<FBasicRayData>, Rays)
+		SHADER_PARAMETER_UAV(RWStructuredBuffer<uint>, OcclusionOutput)
+		END_SHADER_PARAMETER_STRUCT()
+};
+
+class FBasicIntersectionMainRGS : public FBuiltInRayTracingShader
+{
+	DECLARE_GLOBAL_SHADER(FBasicIntersectionMainRGS);
+	SHADER_USE_ROOT_PARAMETER_STRUCT(FBasicIntersectionMainRGS, FBuiltInRayTracingShader);
+
+	BEGIN_SHADER_PARAMETER_STRUCT(FParameters, )
+		SHADER_PARAMETER_SRV(RaytracingAccelerationStructure, TLAS)
+		SHADER_PARAMETER_SRV(StructuredBuffer<FBasicRayData>, Rays)
+		SHADER_PARAMETER_UAV(RWStructuredBuffer<FIntersectionPayload>, IntersectionOutput)
+		END_SHADER_PARAMETER_STRUCT()
+};
+
+class FBasicIntersectionMainCHS : public FBuiltInRayTracingShader
+{
+	DECLARE_GLOBAL_SHADER(FBasicIntersectionMainCHS);
+public:
+
+	FBasicIntersectionMainCHS() = default;
+	FBasicIntersectionMainCHS(const ShaderMetaType::CompiledShaderInitializerType& Initializer)
+		: FBuiltInRayTracingShader(Initializer)
+	{}
+};
+
+/**
+* Trace rays from an input buffer of FBasicRayTracingRay.
+* Binary intersection results are written to output buffer as R32_UINTs.
+* 0xFFFFFFFF is written if ray intersects any scene triangle, 0 otherwise.
+*/
+void DispatchBasicOcclusionRays(FRHICommandList& RHICmdList, FRHIRayTracingScene* Scene, FRHIShaderResourceView* SceneView, FRHIShaderResourceView* RayBufferView, FRHIUnorderedAccessView* ResultView, uint32 NumRays);
+
+/**
+* Trace rays from an input buffer of FBasicRayTracingRay.
+* Primitive intersection results are written to output buffer as FBasicRayTracingIntersectionResult.
+*/
+void DispatchBasicIntersectionRays(FRHICommandList& RHICmdList, FRHIRayTracingScene* Scene, FRHIShaderResourceView* SceneView, FRHIShaderResourceView* RayBufferView, FRHIUnorderedAccessView* ResultView, uint32 NumRays);
+
+#endif // RHI_RAYTRACING
diff --git a/Engine/Source/Runtime/Renderer/Private/Tests/RayTracingTestbed.cpp b/Engine/Source/Runtime/Renderer/Private/Tests/RayTracingTestbed.cpp
--- a/Engine/Source/Runtime/Renderer/Private/Tests/RayTracingTestbed.cpp
+++ b/Engine/Source/Runtime/Renderer/Private/Tests/RayTracingTestbed.cpp
@@ -10,6 +10,7 @@ IMPLEMENT_SIMPLE_AUTOMATION_TEST(FRayTracingTestbed, "System.Renderer.RayTracing
 
 #if RHI_RAYTRACING
 
+#include "RayTracing/RayTracingBasicShaders.h"
 #include "RayTracingInstanceBufferUtil.h"
 #include "RayTracingDefinitions.h"
 #include "GlobalShader.h"
@@ -63,20 +64,18 @@ bool RunRayTracingTestbed_RenderThread(const FString& Parameters)
 	FShaderResourceViewRHIRef RayBufferView;
 
 	{
-		PRAGMA_DISABLE_DEPRECATION_WARNINGS
-		TResourceArray<FBasicRayData> RayData;
+		TResourceArray<FBasicRayTracingRay> RayData;
 		RayData.SetNumUninitialized(NumRays);
-		RayData[0] = FBasicRayData{ { 0.75f, 0.0f, -1.0f}, 0xFFFFFFFF, {0.0f, 0.0f,  1.0f}, 100000.0f }; // expected to hit
-		RayData[1] = FBasicRayData{ { 0.75f, 0.0f, -1.0f}, 0xFFFFFFFF, {0.0f, 0.0f,  1.0f},      0.5f }; // expected to miss (short ray)
-		RayData[2] = FBasicRayData{ { 0.75f, 0.0f,  1.0f}, 0xFFFFFFFF, {0.0f, 0.0f, -1.0f}, 100000.0f }; // expected to hit  (should hit back face)
-		RayData[3] = FBasicRayData{ {-0.75f, 0.0f, -1.0f}, 0xFFFFFFFF, {0.0f, 0.0f,  1.0f}, 100000.0f }; // expected to miss (doesn't intersect)
+		RayData[0] = FBasicRayTracingRay{ { 0.75f, 0.0f, -1.0f}, 0xFFFFFFFF, {0.0f, 0.0f,  1.0f}, 100000.0f }; // expected to hit
+		RayData[1] = FBasicRayTracingRay{ { 0.75f, 0.0f, -1.0f}, 0xFFFFFFFF, {0.0f, 0.0f,  1.0f},      0.5f }; // expected to miss (short ray)
+		RayData[2] = FBasicRayTracingRay{ { 0.75f, 0.0f,  1.0f}, 0xFFFFFFFF, {0.0f, 0.0f, -1.0f}, 100000.0f }; // expected to hit  (should hit back face)
+		RayData[3] = FBasicRayTracingRay{ {-0.75f, 0.0f, -1.0f}, 0xFFFFFFFF, {0.0f, 0.0f,  1.0f}, 100000.0f }; // expected to miss (doesn't intersect)
 
 		FRHIResourceCreateInfo CreateInfo(TEXT("RayBuffer"));
 		CreateInfo.ResourceArray = &RayData;
 
-		RayBuffer = RHICreateStructuredBuffer(sizeof(FBasicRayData), RayData.GetResourceDataSize(), BUF_Static | BUF_ShaderResource, CreateInfo);
+		RayBuffer = RHICreateStructuredBuffer(sizeof(FBasicRayTracingRay), RayData.GetResourceDataSize(), BUF_Static | BUF_ShaderResource, CreateInfo);
 		RayBufferView = RHICreateShaderResourceView(RayBuffer);
-		PRAGMA_ENABLE_DEPRECATION_WARNINGS
 	}
 
 	FBufferRHIRef OcclusionResultBuffer;
@@ -93,9 +92,7 @@ bool RunRayTracingTestbed_RenderThread(const FString& Parameters)
 
 	{
 		FRHIResourceCreateInfo CreateInfo(TEXT("IntersectionResultBuffer"));
-		PRAGMA_DISABLE_DEPRECATION_WARNINGS
-		IntersectionResultBuffer = RHICreateVertexBuffer(sizeof(FIntersectionPayload)*NumRays, BUF_Static | BUF_UnorderedAccess, CreateInfo);
-		PRAGMA_ENABLE_DEPRECATION_WARNINGS
+		IntersectionResultBuffer = RHICreateVertexBuffer(sizeof(FBasicRayTracingIntersectionResult)*NumRays, BUF_Static | BUF_UnorderedAccess, CreateInfo);
 		IntersectionResultBufferView = RHICreateUnorderedAccessView(IntersectionResultBuffer, PF_R32_UINT);
 	}
 
@@ -222,21 +219,22 @@ bool RunRayTracingTestbed_RenderThread(const FString& Parameters)
 	// #yuriy_todo: explicit transitions and state validation for BLAS
 	// RHICmdList.Transition(FRHITransitionInfo(Geometry.GetReference(), ERHIAccess::BVHWrite, ERHIAccess::BVHRead));
 
-	FRayTracingSceneBuildParams Params;
-	Params.Scene = RayTracingScene.Scene;
-	Params.ScratchBuffer = ScratchBuffer;
-	Params.ScratchBufferOffset = 0;
-	Params.InstanceBuffer = InstanceBuffer.Buffer;
-	Params.InstanceBufferOffset = 0;
+	FRayTracingSceneBuildParams BuildParams;
+	BuildParams.Scene = RayTracingScene.Scene;
+	BuildParams.ScratchBuffer = ScratchBuffer;
+	BuildParams.ScratchBufferOffset = 0;
+	BuildParams.InstanceBuffer = InstanceBuffer.Buffer;
+	BuildParams.InstanceBufferOffset = 0;
 
-	RHICmdList.BuildAccelerationStructure(Params);
+	RHICmdList.BuildAccelerationStructure(BuildParams);
 
 	RHICmdList.Transition(FRHITransitionInfo(RayTracingScene.Scene.GetReference(), ERHIAccess::BVHWrite, ERHIAccess::BVHRead));
 
-	PRAGMA_DISABLE_DEPRECATION_WARNINGS
-	RHICmdList.RayTraceOcclusion(RayTracingScene.Scene, RayBufferView, OcclusionResultBufferView, NumRays);
-	RHICmdList.RayTraceIntersection(RayTracingScene.Scene, RayBufferView, IntersectionResultBufferView, NumRays);
-	PRAGMA_ENABLE_DEPRECATION_WARNINGS
+	FShaderResourceViewInitializer RayTracingSceneViewInitializer(SceneBuffer, RayTracingScene.Scene->GetLayerBufferOffset(0), 0);
+	FShaderResourceViewRHIRef RayTracingSceneView = RHICreateShaderResourceView(RayTracingSceneViewInitializer);
+
+	DispatchBasicOcclusionRays(RHICmdList, RayTracingScene.Scene, RayTracingSceneView, RayBufferView, OcclusionResultBufferView, NumRays);
+	DispatchBasicIntersectionRays(RHICmdList, RayTracingScene.Scene, RayTracingSceneView, RayBufferView, IntersectionResultBufferView, NumRays);
 
 	const bool bValidateResults = true;
 	bool bOcclusionTestOK = false;
@@ -267,8 +265,7 @@ bool RunRayTracingTestbed_RenderThread(const FString& Parameters)
 		// Read back and validate intersection trace results
 
 		{
-			PRAGMA_DISABLE_DEPRECATION_WARNINGS
-			auto MappedResults = (const FIntersectionPayload*)RHILockBuffer(IntersectionResultBuffer, 0, sizeof(FIntersectionPayload)*NumRays, RLM_ReadOnly);
+			auto MappedResults = (const FBasicRayTracingIntersectionResult*)RHILockBuffer(IntersectionResultBuffer, 0, sizeof(FBasicRayTracingIntersectionResult)*NumRays, RLM_ReadOnly);
 
 			check(MappedResults);
 
@@ -286,7 +283,6 @@ bool RunRayTracingTestbed_RenderThread(const FString& Parameters)
 			RHIUnlockBuffer(IntersectionResultBuffer);
 
 			bIntersectionTestOK = (MappedResults[0].HitT >= 0) && (MappedResults[1].HitT < 0) && (MappedResults[2].HitT >= 0) && (MappedResults[3].HitT < 0);
-			PRAGMA_ENABLE_DEPRECATION_WARNINGS
 		}
 	}
 
diff --git a/Engine/Source/Runtime/VulkanRHI/Private/VulkanContext.h b/Engine/Source/Runtime/VulkanRHI/Private/VulkanContext.h
--- a/Engine/Source/Runtime/VulkanRHI/Private/VulkanContext.h
+++ b/Engine/Source/Runtime/VulkanRHI/Private/VulkanContext.h
@@ -104,7 +104,6 @@ public:
 	virtual void RHIBindAccelerationStructureMemory(FRHIRayTracingScene* Scene, FRHIBuffer* Buffer, uint32 BufferOffset) final override;
 	virtual void RHIBuildAccelerationStructures(const TArrayView<const FRayTracingGeometryBuildParams> Params, const FRHIBufferRange& ScratchBufferRange) final override;
 	virtual void RHIBuildAccelerationStructure(const FRayTracingSceneBuildParams& SceneBuildParams) final override;
-	virtual void RHIRayTraceOcclusion(FRHIRayTracingScene* Scene, FRHIShaderResourceView* Rays, FRHIUnorderedAccessView* Output, uint32 NumRays) final override;
 #endif
 
 	inline FVulkanCommandBufferManager* GetCommandBufferManager()
diff --git a/Engine/Source/Runtime/VulkanRHI/Private/VulkanDevice.h b/Engine/Source/Runtime/VulkanRHI/Private/VulkanDevice.h
--- a/Engine/Source/Runtime/VulkanRHI/Private/VulkanDevice.h
+++ b/Engine/Source/Runtime/VulkanRHI/Private/VulkanDevice.h
@@ -19,7 +19,6 @@ class FVulkanOcclusionQueryPool;
 class FVulkanRenderPassManager;
 
 #if VULKAN_RHI_RAYTRACING
-class FVulkanBasicRaytracingPipeline;
 class FVulkanRayTracingCompactionRequestHandler;
 #endif
 
@@ -560,7 +559,6 @@ private:
 
 #if VULKAN_RHI_RAYTRACING
 	FRayTracingProperties RayTracingProperties;
-	FVulkanBasicRaytracingPipeline* BasicRayTracingPipeline = nullptr;
 	FVulkanRayTracingCompactionRequestHandler* RayTracingCompactionRequestHandler = nullptr;
 #endif // VULKAN_RHI_RAYTRACING
 
diff --git a/Engine/Source/Runtime/VulkanRHI/Private/VulkanRayTracing.cpp b/Engine/Source/Runtime/VulkanRHI/Private/VulkanRayTracing.cpp
--- a/Engine/Source/Runtime/VulkanRHI/Private/VulkanRayTracing.cpp
+++ b/Engine/Source/Runtime/VulkanRHI/Private/VulkanRayTracing.cpp
@@ -903,12 +903,6 @@ void FVulkanCommandListContext::RHIBuildAccelerationStructure(const FRayTracingS
 		InstanceBuffer, SceneBuildParams.InstanceBufferOffset);
 }
 
-void FVulkanCommandListContext::RHIRayTraceOcclusion(FRHIRayTracingScene* Scene, FRHIShaderResourceView* Rays, FRHIUnorderedAccessView* Output, uint32 NumRays)
-{
-	// todo
-	return;
-}
-
 template<typename ShaderType>
 static FRHIRayTracingShader* GetBuiltInRayTracingShader()
 {
@@ -919,18 +913,10 @@ static FRHIRayTracingShader* GetBuiltInRayTracingShader()
 
 void FVulkanDevice::InitializeRayTracing()
 {
-	check(BasicRayTracingPipeline == nullptr);
-	// the pipeline should be initialized on the first use due to the ability to disable RT in the game settings
-	//BasicRayTracingPipeline = new FVulkanBasicRaytracingPipeline(this);
 }
 
 void FVulkanDevice::CleanUpRayTracing()
 {
-	if (BasicRayTracingPipeline != nullptr)
-	{
-		delete BasicRayTracingPipeline;
-		BasicRayTracingPipeline = nullptr;
-	}
 }
 
 static uint32 GetAlignedSize(uint32 Value, uint32 Alignment)
@@ -1118,42 +1104,6 @@ FVulkanRayTracingPipelineState::~FVulkanRayTracingPipelineState()
 	}
 }
 
-FVulkanBasicRaytracingPipeline::FVulkanBasicRaytracingPipeline(FVulkanDevice* const InDevice)
-{
-	check(Occlusion == nullptr);
-
-	// Occlusion pipeline
-	{
-		PRAGMA_DISABLE_DEPRECATION_WARNINGS
-
-		FRayTracingPipelineStateInitializer OcclusionInitializer;
-
-		FRHIRayTracingShader* OcclusionRGSTable[] = { GetBuiltInRayTracingShader<FOcclusionMainRG>() };
-		OcclusionInitializer.SetRayGenShaderTable(OcclusionRGSTable);
-
-		FRHIRayTracingShader* OcclusionMSTable[] = { GetBuiltInRayTracingShader<FDefaultPayloadMS>() };
-		OcclusionInitializer.SetMissShaderTable(OcclusionMSTable);
-
-		FRHIRayTracingShader* OcclusionCHSTable[] = { GetBuiltInRayTracingShader<FDefaultMainCHS>() };
-		OcclusionInitializer.SetHitGroupTable(OcclusionCHSTable);
-
-		OcclusionInitializer.bAllowHitGroupIndexing = false;
-
-		Occlusion = new FVulkanRayTracingPipelineState(InDevice, OcclusionInitializer);
-
-		PRAGMA_ENABLE_DEPRECATION_WARNINGS
-	}
-}
-
-FVulkanBasicRaytracingPipeline::~FVulkanBasicRaytracingPipeline()
-{
-	if (Occlusion != nullptr)
-	{
-		delete Occlusion;
-		Occlusion = nullptr;
-	}
-}
-
 void FVulkanRayTracingCompactedSizeQueryPool::EndBatch(FVulkanCmdBuffer* InCmdBuffer)
 {
 	check(CmdBuffer == nullptr);
diff --git a/Engine/Source/Runtime/VulkanRHI/Private/VulkanRayTracing.h b/Engine/Source/Runtime/VulkanRHI/Private/VulkanRayTracing.h
--- a/Engine/Source/Runtime/VulkanRHI/Private/VulkanRayTracing.h
+++ b/Engine/Source/Runtime/VulkanRHI/Private/VulkanRayTracing.h
@@ -152,19 +152,6 @@ private:
 	FVkRtAllocation HitShaderBindingTable;
 };
 
-class FVulkanBasicRaytracingPipeline
-{
-public:
-
-	UE_NONCOPYABLE(FVulkanBasicRaytracingPipeline);
-	FVulkanBasicRaytracingPipeline(FVulkanDevice* const InDevice);
-	~FVulkanBasicRaytracingPipeline();
-
-private:
-
-	FVulkanRayTracingPipelineState* Occlusion = nullptr;
-};
-
 class FVulkanRayTracingCompactedSizeQueryPool : public FVulkanQueryPool
 {
 public:
