Fix overzealous ensure and incorrect return value from FD3D12GPUFence::Poll when a fence was polled before it was submitted.

#rnx

diff --git a/Engine/Source/Runtime/D3D12RHI/Private/D3D12DirectCommandListManager.cpp b/Engine/Source/Runtime/D3D12RHI/Private/D3D12DirectCommandListManager.cpp
--- a/Engine/Source/Runtime/D3D12RHI/Private/D3D12DirectCommandListManager.cpp
+++ b/Engine/Source/Runtime/D3D12RHI/Private/D3D12DirectCommandListManager.cpp
@@ -28,7 +28,7 @@ bool FD3D12GPUFence::Poll(FRHIGPUMask GPUMask) const
 {
 	for (uint32 Index : GPUMask)
 	{
-		if (ensureMsgf(SyncPoints[Index], TEXT("Attempt to poll an FRHIGPUFence that was never issued.")) && !SyncPoints[Index]->IsComplete())
+		if (SyncPoints[Index] == nullptr || !SyncPoints[Index]->IsComplete())
 			return false;
 	}
 
