Turn off the interrupt thread, as sometimes the fence events are signaled very late, which introduces excessively long frames.

#rnx

diff --git a/Engine/Source/Runtime/D3D12RHI/Private/D3D12Submission.cpp b/Engine/Source/Runtime/D3D12RHI/Private/D3D12Submission.cpp
--- a/Engine/Source/Runtime/D3D12RHI/Private/D3D12Submission.cpp
+++ b/Engine/Source/Runtime/D3D12RHI/Private/D3D12Submission.cpp
@@ -5,7 +5,7 @@
 #include "HAL/RunnableThread.h"
 
 #ifndef D3D12_PLATFORM_SUPPORTS_BLOCKING_FENCES
-#define D3D12_PLATFORM_SUPPORTS_BLOCKING_FENCES 1
+#define D3D12_PLATFORM_SUPPORTS_BLOCKING_FENCES 0
 #endif
 
 // These defines control which threads are enabled in the GPU submission pipeline.
