MGPU:  fix for fence polling issue when a fence is polled with the GPU mask set to all GPUs, but originally signaled on a subset of GPUs.  This is a common use case for streaming system readback (Virtual Textures, Distance Fields, Nanite), where rendering happens on a single GPU, and streaming scene update happens on all GPUs.  The polling function will return "true" if any SyncPoints array element is non-null, and all SyncPoints have passed, rather than returning "false" on any null element.  Fixes severe Nanite and Virtual Texture artifacts and perf regression.

#rnx

diff --git a/Engine/Source/Runtime/D3D12RHI/Private/D3D12DirectCommandListManager.cpp b/Engine/Source/Runtime/D3D12RHI/Private/D3D12DirectCommandListManager.cpp
--- a/Engine/Source/Runtime/D3D12RHI/Private/D3D12DirectCommandListManager.cpp
+++ b/Engine/Source/Runtime/D3D12RHI/Private/D3D12DirectCommandListManager.cpp
@@ -26,13 +26,21 @@ bool FD3D12GPUFence::Poll() const
 
 bool FD3D12GPUFence::Poll(FRHIGPUMask GPUMask) const
 {
+	bool bHasAnySyncPoint = false;
 	for (uint32 Index : GPUMask)
 	{
-		if (SyncPoints[Index] == nullptr || !SyncPoints[Index]->IsComplete())
-			return false;
+		if (SyncPoints[Index])
+		{
+			if (!SyncPoints[Index]->IsComplete())
+			{
+				return false;
+			}
+			bHasAnySyncPoint = true;
+		}
 	}
 
-	return true;
+	// Return "true" if we had sync points that all successfully completed, or "false" if we have no sync points (fence was never signaled)
+	return bHasAnySyncPoint;
 }
 
 void FD3D12GPUFence::WaitCPU()
