Fix SynthBenchmark
 - GPU timestamp frequency was only being lazy intialized by the "profilegpu" command, which SynthBenchmark does not use.
 - Moved the init into the RHI init function, so its ready immediately after engine boot.

diff --git a/Engine/Source/Runtime/D3D12RHI/Private/D3D12Query.cpp b/Engine/Source/Runtime/D3D12RHI/Private/D3D12Query.cpp
--- a/Engine/Source/Runtime/D3D12RHI/Private/D3D12Query.cpp
+++ b/Engine/Source/Runtime/D3D12RHI/Private/D3D12Query.cpp
@@ -326,23 +326,18 @@ bool FD3D12DynamicRHI::RHIGetRenderQueryResult(FRHIRenderQuery* QueryRHI, uint64
 FD3D12BufferedGPUTiming::FD3D12BufferedGPUTiming(FD3D12Device* InParent)
 	: FD3D12DeviceChild(InParent)
 {
-	// StaticInitialize operates on all devices so only call it once.
-	UE_CALL_ONCE([InParent]()
-	{
-		StaticInitialize(InParent->GetParentAdapter(), PlatformStaticInitialize);
-	});
 }
 
-/**
- * Initializes the static variables, if necessary.
- */
-void FD3D12BufferedGPUTiming::PlatformStaticInitialize(void* UserData)
+void FD3D12BufferedGPUTiming::Initialize(FD3D12Adapter* ParentAdapter)
 {
-	// Are the static variables initialized?
-	check(!GAreGlobalsInitialized);
+	StaticInitialize(ParentAdapter, [](void* UserData)
+	{
+		// Are the static variables initialized?
+		check(!GAreGlobalsInitialized);
 
-	FD3D12Adapter* ParentAdapter = (FD3D12Adapter*)UserData;
-	CalibrateTimers(ParentAdapter);
+		FD3D12Adapter* ParentAdapter = (FD3D12Adapter*)UserData;
+		CalibrateTimers(ParentAdapter);
+	});
 }
 
 void FD3D12BufferedGPUTiming::CalibrateTimers(FD3D12Adapter* ParentAdapter)
diff --git a/Engine/Source/Runtime/D3D12RHI/Private/D3D12Stats.h b/Engine/Source/Runtime/D3D12RHI/Private/D3D12Stats.h
--- a/Engine/Source/Runtime/D3D12RHI/Private/D3D12Stats.h
+++ b/Engine/Source/Runtime/D3D12RHI/Private/D3D12Stats.h
@@ -171,12 +171,9 @@ public:
 
 	static void CalibrateTimers(FD3D12Adapter* ParentAdapter);
 
-private:
-	/**
-	* Initializes the static variables, if necessary.
-	*/
-	static void PlatformStaticInitialize(void* UserData);
+	static void Initialize(FD3D12Adapter* ParentAdapter);
 
+private:
 	struct
 	{
 		uint64 Result = 0;
diff --git a/Engine/Source/Runtime/D3D12RHI/Private/Windows/WindowsD3D12Device.cpp b/Engine/Source/Runtime/D3D12RHI/Private/Windows/WindowsD3D12Device.cpp
--- a/Engine/Source/Runtime/D3D12RHI/Private/Windows/WindowsD3D12Device.cpp
+++ b/Engine/Source/Runtime/D3D12RHI/Private/Windows/WindowsD3D12Device.cpp
@@ -1538,6 +1538,11 @@ void FD3D12DynamicRHI::Init()
 
 	GRHICommandList.GetImmediateCommandList().InitializeImmediateContexts();
 
+	for (TSharedPtr<FD3D12Adapter>& Adapter : ChosenAdapters)
+	{
+		FD3D12BufferedGPUTiming::Initialize(Adapter.Get());
+	}
+
 	FRenderResource::InitPreRHIResources();
 	GIsRHIInitialized = true;
 }
