D3D12:  Fix for fence wait hang when running on console D3D12 with "r.rhicmdbypass 1".  NeedFlushBeforeEndDrawing returns "true" on console D3D12 (the original behavior), and "false" only on Windows D3D12.  Not sure why the extra flush is required there on console, but it's a minor optimization, so it's not too big a deal to disable.

#rnx

diff --git a/Engine/Source/Runtime/D3D12RHI/Private/Windows/WindowsD3D12Viewport.cpp b/Engine/Source/Runtime/D3D12RHI/Private/Windows/WindowsD3D12Viewport.cpp
--- a/Engine/Source/Runtime/D3D12RHI/Private/Windows/WindowsD3D12Viewport.cpp
+++ b/Engine/Source/Runtime/D3D12RHI/Private/Windows/WindowsD3D12Viewport.cpp
@@ -538,6 +538,15 @@ bool FD3D12Viewport::IsPresentAllowed()
 	return !FD3D12DynamicRHI::GetD3DRHI()->RHIIsRenderingSuspended();
 }
 
+/**
+  * RHIEndDrawingViewport flushes commands internally, so we don't need to flush before calling it.
+  * Avoids wasteful issuing of another entire command list with just a few commands.
+  */
+bool FD3D12Viewport::NeedFlushBeforeEndDrawing()
+{
+	return false;
+}
+
 void FD3D12DynamicRHI::RHIGetDisplaysInformation(FDisplayInformationArray& OutDisplayInformation)
 {
 	OutDisplayInformation.Append(DisplayList);
diff --git a/Engine/Source/Runtime/D3D12RHI/Public/D3D12Viewport.h b/Engine/Source/Runtime/D3D12RHI/Public/D3D12Viewport.h
--- a/Engine/Source/Runtime/D3D12RHI/Public/D3D12Viewport.h
+++ b/Engine/Source/Runtime/D3D12RHI/Public/D3D12Viewport.h
@@ -103,11 +103,7 @@ public:
 	virtual void WaitForFrameEventCompletion() override;
 	virtual void IssueFrameEvent() override;
 
-	/**
-	  * RHIEndDrawingViewport flushes commands internally, so we don't need to flush before calling it.
-	  * Avoids wasteful issuing of another entire command list with just a few commands.
-	  */
-	virtual bool NeedFlushBeforeEndDrawing() { return false; }
+	virtual bool NeedFlushBeforeEndDrawing() override;
 
 #if D3D12_VIEWPORT_EXPOSES_SWAP_CHAIN
 	virtual void* GetNativeSwapChain() const override { return SwapChain1; }
