Fix for standalone PIE and quick launched Lyra chugging when focused when "Use Less CPU when in Background" (bThrottleCPUWhenNotForeground) is enabled. Also reenable RHIInterrupt thread.

Since CreateEvent was being passed the same name for both the parent (editor) and child (game) processes when creating the interrupt and submission threads, and that name was in the Global kernel object namespace, both the editor and LyraGame.exe were sharing the same Event. By adding the "Local\\" prefix to the Event name, each process gets its own event, so the editor can no longer throttle the game's interrupt and submission threads.

diff --git a/Engine/Source/Runtime/D3D12RHI/Private/D3D12Submission.cpp b/Engine/Source/Runtime/D3D12RHI/Private/D3D12Submission.cpp
--- a/Engine/Source/Runtime/D3D12RHI/Private/D3D12Submission.cpp
+++ b/Engine/Source/Runtime/D3D12RHI/Private/D3D12Submission.cpp
@@ -5,7 +5,7 @@
 #include "HAL/RunnableThread.h"
 
 #ifndef D3D12_PLATFORM_SUPPORTS_BLOCKING_FENCES
-#define D3D12_PLATFORM_SUPPORTS_BLOCKING_FENCES 0
+#define D3D12_PLATFORM_SUPPORTS_BLOCKING_FENCES 1
 #endif
 
 // These defines control which threads are enabled in the GPU submission pipeline.
@@ -42,7 +42,7 @@ public:
 	FD3D12Thread(TCHAR const* Name, EThreadPriority Priority, FD3D12DynamicRHI* RHI, FQueueFunc QueueFunc)
 		: RHI(RHI)
 		, QueueFunc(QueueFunc)
-		, Event(CreateEvent(nullptr, false, 0, Name))
+		, Event(CreateEvent(nullptr, false, false, *(FString::Printf(TEXT("%s%s"), TEXT("Local\\"), Name))))
 		, Thread(FRunnableThread::Create(this, Name, 0, Priority))
 	{}
 
