Fix for GPU crash when transient heaps are released.

diff --git a/Engine/Source/Runtime/D3D12RHI/Private/D3D12Adapter.cpp b/Engine/Source/Runtime/D3D12RHI/Private/D3D12Adapter.cpp
--- a/Engine/Source/Runtime/D3D12RHI/Private/D3D12Adapter.cpp
+++ b/Engine/Source/Runtime/D3D12RHI/Private/D3D12Adapter.cpp
@@ -1309,6 +1309,8 @@ void FD3D12Adapter::CleanupResources()
 		Viewport->WaitForFrameEventCompletion();
 	}
 
+	TransientMemoryCache = nullptr;
+
 #if D3D12_RHI_RAYTRACING
 	for (uint32 GPUIndex : FRHIGPUMask::All())
 	{
diff --git a/Engine/Source/Runtime/D3D12RHI/Private/D3D12CommandContext.h b/Engine/Source/Runtime/D3D12RHI/Private/D3D12CommandContext.h
--- a/Engine/Source/Runtime/D3D12RHI/Private/D3D12CommandContext.h
+++ b/Engine/Source/Runtime/D3D12RHI/Private/D3D12CommandContext.h
@@ -38,6 +38,7 @@ struct FD3D12DeferredDeleteObject
 	{
 		RHIObject,
 		D3DObject,
+		D3DHeap,
 		BindlessDescriptor,
 		CPUAllocation
 	} Type;
@@ -45,6 +46,7 @@ struct FD3D12DeferredDeleteObject
 	union
 	{
 		FD3D12Resource* RHIObject;
+		FD3D12Heap* D3DHeap;
 		ID3D12Object* D3DObject;
 
 		struct
@@ -61,6 +63,11 @@ struct FD3D12DeferredDeleteObject
 		, RHIObject(RHIObject)
 	{}
 
+	explicit FD3D12DeferredDeleteObject(FD3D12Heap* D3DHeap)
+		: Type(EType::D3DHeap)
+		, D3DHeap(D3DHeap)
+	{}
+
 	explicit FD3D12DeferredDeleteObject(ID3D12Object* D3DObject)
 		: Type(EType::D3DObject)
 		, D3DObject(D3DObject)
diff --git a/Engine/Source/Runtime/D3D12RHI/Private/D3D12RHI.cpp b/Engine/Source/Runtime/D3D12RHI/Private/D3D12RHI.cpp
--- a/Engine/Source/Runtime/D3D12RHI/Private/D3D12RHI.cpp
+++ b/Engine/Source/Runtime/D3D12RHI/Private/D3D12RHI.cpp
@@ -476,6 +476,10 @@ void FD3D12DynamicRHI::RHIPerFrameRHIFlushComplete()
 					check(ObjectToDelete.RHIObject->GetRefCount() == 1);
 					ObjectToDelete.RHIObject->Release();
 					break;
+				case FD3D12DeferredDeleteObject::EType::D3DHeap:
+					// Heaps can have additional references active.
+					ObjectToDelete.D3DHeap->Release();
+					break;
 				case FD3D12DeferredDeleteObject::EType::D3DObject:
 					ObjectToDelete.D3DObject->Release();
 					break;
diff --git a/Engine/Source/Runtime/D3D12RHI/Private/D3D12Resources.cpp b/Engine/Source/Runtime/D3D12RHI/Private/D3D12Resources.cpp
--- a/Engine/Source/Runtime/D3D12RHI/Private/D3D12Resources.cpp
+++ b/Engine/Source/Runtime/D3D12RHI/Private/D3D12Resources.cpp
@@ -186,6 +186,11 @@ FD3D12Heap::~FD3D12Heap()
 	Heap.SafeRelease();
 }
 
+void FD3D12Heap::DeferDelete()
+{
+	FD3D12DynamicRHI::GetD3DRHI()->DeferredDelete(this);
+}
+
 void FD3D12Heap::SetHeap(ID3D12Heap* HeapIn, const TCHAR* const InName, bool bInTrack, bool bForceGetGPUAddress)
 {
 	*Heap.GetInitReference() = HeapIn; 
diff --git a/Engine/Source/Runtime/D3D12RHI/Private/D3D12TransientResourceAllocator.cpp b/Engine/Source/Runtime/D3D12RHI/Private/D3D12TransientResourceAllocator.cpp
--- a/Engine/Source/Runtime/D3D12RHI/Private/D3D12TransientResourceAllocator.cpp
+++ b/Engine/Source/Runtime/D3D12RHI/Private/D3D12TransientResourceAllocator.cpp
@@ -106,6 +106,10 @@ FD3D12TransientHeap::~FD3D12TransientHeap()
 		LLM(FLowLevelMemTracker::Get().OnLowLevelFree(ELLMTracker::Platform, Heap->GetHeap()));
 		MemoryTrace_Free((uint64)Heap->GetHeap(), EMemoryTraceRootHeap::VideoMemory);
 #endif
+
+		// Add a reference to the underlying to heap to be deferred released later.
+		Heap->AddRef();
+		Heap->DeferDelete();
 	}
 }
 
diff --git a/Engine/Source/Runtime/D3D12RHI/Public/D3D12Resources.h b/Engine/Source/Runtime/D3D12RHI/Public/D3D12Resources.h
--- a/Engine/Source/Runtime/D3D12RHI/Public/D3D12Resources.h
+++ b/Engine/Source/Runtime/D3D12RHI/Public/D3D12Resources.h
@@ -78,7 +78,9 @@ public:
 	void SetHeap(ID3D12Heap* HeapIn, const TCHAR* const InName, bool bTrack = true, bool bForceGetGPUAddress = false);
 
 	void BeginTrackingResidency(uint64 Size);
-	
+
+	void DeferDelete();
+
 	inline FName GetName() const { return HeapName; }
 	inline D3D12_HEAP_DESC GetHeapDesc() const { return HeapDesc; }
 	inline FD3D12ResidencyHandle& GetResidencyHandle() { return ResidencyHandle; }
