Remove unnecessary D3D12 ResourceBarrier batching code.

#rn Rendering/D3D12 D3D12 ResourceBarrier batching is unnecessary, removed related code.

diff --git a/Engine/Source/Runtime/D3D12RHI/Private/D3D12RHIPrivate.h b/Engine/Source/Runtime/D3D12RHI/Private/D3D12RHIPrivate.h
--- a/Engine/Source/Runtime/D3D12RHI/Private/D3D12RHIPrivate.h
+++ b/Engine/Source/Runtime/D3D12RHI/Private/D3D12RHIPrivate.h
@@ -663,8 +663,6 @@ public:
 	bool IsQuadBufferStereoEnabled() const;
 	void DisableQuadBufferStereo();
 
-	static int32 GetResourceBarrierBatchSizeLimit();
-
 	FBufferRHIRef CreateBuffer(FRHICommandListBase& RHICmdList, uint32 Size, EBufferUsageFlags Usage, uint32 Stride, ERHIAccess ResourceState, FRHIResourceCreateInfo& CreateInfo);
 	void* LockBuffer(FRHICommandListBase& RHICmdList, FD3D12Buffer* Buffer, uint32 BufferSize, EBufferUsageFlags BufferUsage, uint32 Offset, uint32 Size, EResourceLockMode LockMode);
 	void UnlockBuffer(FRHICommandListBase& RHICmdList, FD3D12Buffer* Buffer, EBufferUsageFlags BufferUsage);
diff --git a/Engine/Source/Runtime/D3D12RHI/Private/D3D12Resources.cpp b/Engine/Source/Runtime/D3D12RHI/Private/D3D12Resources.cpp
--- a/Engine/Source/Runtime/D3D12RHI/Private/D3D12Resources.cpp
+++ b/Engine/Source/Runtime/D3D12RHI/Private/D3D12Resources.cpp
@@ -1007,8 +1007,6 @@ void FD3D12ResourceBarrierBatcher::AddAliasingBarrier(ID3D12Resource* InResource
 
 void FD3D12ResourceBarrierBatcher::FlushIntoCommandList(FD3D12CommandList& CommandList, FD3D12QueryAllocator& TimestampAllocator)
 {
-	int32 const BarrierBatchMax = FD3D12DynamicRHI::GetResourceBarrierBatchSizeLimit();
-
 	auto InsertTimestamp = [&](ED3D12QueryType Type)
 	{
 		CommandList.EndQuery(TimestampAllocator.Allocate(Type, nullptr));
@@ -1020,8 +1018,7 @@ void FD3D12ResourceBarrierBatcher::FlushIntoCommandList(FD3D12CommandList& Comma
 		bool const bIdle = Barriers[BatchEnd].HasIdleFlag();
 
 		while (BatchEnd < Barriers.Num() 
-			&& bIdle == Barriers[BatchEnd].HasIdleFlag() 
-			&& (BatchEnd - BatchStart) < BarrierBatchMax)
+			&& bIdle == Barriers[BatchEnd].HasIdleFlag())
 		{
 			// Clear the idle flag since its not a valid D3D bit.
 			Barriers[BatchEnd++].ClearIdleFlag();
diff --git a/Engine/Source/Runtime/D3D12RHI/Private/Windows/WindowsD3D12Device.cpp b/Engine/Source/Runtime/D3D12RHI/Private/Windows/WindowsD3D12Device.cpp
--- a/Engine/Source/Runtime/D3D12RHI/Private/Windows/WindowsD3D12Device.cpp
+++ b/Engine/Source/Runtime/D3D12RHI/Private/Windows/WindowsD3D12Device.cpp
@@ -1584,12 +1584,6 @@ void FD3D12DynamicRHI::DisableQuadBufferStereo()
 	bIsQuadBufferStereoEnabled = false;
 }
 
-int32 FD3D12DynamicRHI::GetResourceBarrierBatchSizeLimit()
-{
-	return INT32_MAX;
-}
-
-
 void FD3D12Device::CreateSamplerInternal(const D3D12_SAMPLER_DESC& Desc, D3D12_CPU_DESCRIPTOR_HANDLE Descriptor)
 {
 	GetDevice()->CreateSampler(&Desc, Descriptor);
