Fix DLSS crash (GPU device removed)
 - Force re-apply the current descriptor heap when RHIFinishExternalComputeWork is called.

diff --git a/Engine/Source/Runtime/D3D12RHI/Private/D3D12DescriptorCache.cpp b/Engine/Source/Runtime/D3D12RHI/Private/D3D12DescriptorCache.cpp
--- a/Engine/Source/Runtime/D3D12RHI/Private/D3D12DescriptorCache.cpp
+++ b/Engine/Source/Runtime/D3D12RHI/Private/D3D12DescriptorCache.cpp
@@ -72,10 +72,10 @@ void FD3D12DescriptorCache::Init(uint32 InNumLocalViewDescriptors, uint32 InNumS
 	OverrideSamplerHeap = nullptr;
 }
 
-bool FD3D12DescriptorCache::SetDescriptorHeaps()
+bool FD3D12DescriptorCache::SetDescriptorHeaps(bool bForceHeapChanged)
 {
 	// See if the descriptor heaps changed.
-	bool bHeapChanged = false;
+	bool bHeapChanged = bForceHeapChanged;
 
 	ID3D12DescriptorHeap* PendingViewHeap = OverrideViewHeap ? OverrideViewHeap->GetHeap() : CurrentViewHeap->GetHeap();
 	if (LastSetViewHeap != PendingViewHeap)
diff --git a/Engine/Source/Runtime/D3D12RHI/Private/D3D12DescriptorCache.h b/Engine/Source/Runtime/D3D12RHI/Private/D3D12DescriptorCache.h
--- a/Engine/Source/Runtime/D3D12RHI/Private/D3D12DescriptorCache.h
+++ b/Engine/Source/Runtime/D3D12RHI/Private/D3D12DescriptorCache.h
@@ -374,7 +374,7 @@ public:
 
 	// Sets the current descriptor tables on the command list and marks any descriptor tables as dirty if necessary.
 	// Returns true if one of the heaps actually changed, false otherwise.
-	bool SetDescriptorHeaps();
+	bool SetDescriptorHeaps(bool bForceHeapChanged = false);
 
 protected:
 	FD3D12CommandContext& Context;
diff --git a/Engine/Source/Runtime/D3D12RHI/Private/D3D12RHI.cpp b/Engine/Source/Runtime/D3D12RHI/Private/D3D12RHI.cpp
--- a/Engine/Source/Runtime/D3D12RHI/Private/D3D12RHI.cpp
+++ b/Engine/Source/Runtime/D3D12RHI/Private/D3D12RHI.cpp
@@ -616,7 +616,7 @@ void FD3D12DynamicRHI::RHIFinishExternalComputeWork(uint32 InDeviceIndex, ID3D12
 	check(InCommandList == Device->GetDefaultCommandContext().GraphicsCommandList().GetNoRefCount());
 
 	Device->GetDefaultCommandContext().StateCache.ForceSetComputeRootSignature();
-	Device->GetDefaultCommandContext().StateCache.GetDescriptorCache()->SetDescriptorHeaps();
+	Device->GetDefaultCommandContext().StateCache.GetDescriptorCache()->SetDescriptorHeaps(true);
 }
 
 void FD3D12DynamicRHI::RHITransitionResource(FRHICommandList& RHICmdList, FRHITexture* InTexture, D3D12_RESOURCE_STATES InState, uint32 InSubResource)
