PR #9710: D32D12RHI: add #if D3D12_RHI_RAYTRACING wraps to FD3D12CommandContextRedirector (Contributed by [REDACTED_PERSON])

diff --git a/Engine/Source/Runtime/D3D12RHI/Private/D3D12CommandContext.h b/Engine/Source/Runtime/D3D12RHI/Private/D3D12CommandContext.h
--- a/Engine/Source/Runtime/D3D12RHI/Private/D3D12CommandContext.h
+++ b/Engine/Source/Runtime/D3D12RHI/Private/D3D12CommandContext.h
@@ -974,6 +974,7 @@ public:
 		ContextRedirect(RHIEndRenderPass());
 	}
 
+#if D3D12_RHI_RAYTRACING
 	virtual void RHIBuildAccelerationStructures(const TArrayView<const FRayTracingGeometryBuildParams> Params, const FRHIBufferRange& ScratchBufferRange) final override
 	{
 #if WITH_MGPU
@@ -1037,6 +1038,18 @@ public:
 		ContextRedirect(RHISetRayTracingMissShader(Scene, ShaderSlotInScene, Pipeline, ShaderIndexInPipeline, NumUniformBuffers, UniformBuffers, UserData));
 	}
 
+	virtual void RHIClearRayTracingBindings(FRHIRayTracingScene* Scene) final override
+	{
+		ContextRedirect(RHIClearRayTracingBindings(Scene));
+	}
+
+	virtual void RHIBindAccelerationStructureMemory(FRHIRayTracingScene* Scene, FRHIBuffer* Buffer, uint32 BufferOffset) final override
+	{
+		ContextRedirect(RHIBindAccelerationStructureMemory(Scene, Buffer, BufferOffset));
+	}
+
+#endif // D3D12_RHI_RAYTRACING
+
 	virtual void RHISetGPUMask(FRHIGPUMask InGPUMask) final override
 	{
 		GPUMask = InGPUMask;
@@ -1055,16 +1068,6 @@ public:
 		PhysicalGPUMask = InGPUMask;
 	}
 
-	virtual void RHIClearRayTracingBindings(FRHIRayTracingScene* Scene) final override
-	{
-		ContextRedirect(RHIClearRayTracingBindings(Scene));
-	}
-
-	virtual void RHIBindAccelerationStructureMemory(FRHIRayTracingScene* Scene, FRHIBuffer* Buffer, uint32 BufferOffset) final override
-	{
-		ContextRedirect(RHIBindAccelerationStructureMemory(Scene, Buffer, BufferOffset));
-	}
-
 	FORCEINLINE void SetPhysicalContext(FD3D12CommandContext* Context)
 	{
 		check(Context);
