Fix D3D12 crash deleting an unsignaled sync point in FD3D12GPUFence.
 - The sync points allocated in RHIWriteGPUFence_TopOfPipe() are kept alive by the ref count ptr member array "SyncPoints". However, it is possible for a fence to be issued, then immediately cleared by the caller, resulting in the refcount of the sync points hitting zero before the RHI command to signal them on the context has completed.
 - To fix this, the sync points need to be kept alive by the RHI command that signals them, which we're now doing with an EnqueueLambda() rather than using the base RHIWriteGPUFence implementation.

Fix ambiguity in D3D12 RHI handling of MGPU when acquiring the current context from an RHICmdList.
 - Necessary for the above fix, since we need to be able to grab the underlying FD3D12CommandContext (the physical context in MGPU cases) from the given RHICmdList inside an EnqueueLambda.
 - Added a GPUIndex argument to FD3D12CommandContext::Get() to resolve this.
 - Fixed up the places in the RHI that call Get(), or were manually using "RHICmdList.GetComputeContext().GetLowestLevelContext()" so that an explicit GPU context is chosen each time.
 - Removed the virtual override of GetLowestLevelContext() on the context redirector.

Fix double-signal of D3D12 sync point in occlusion queries when using MGPU.
 - There should be one sync point per-GPU used in an occlusion query batch. The code was only allocating one, so was getting signalled N times on multiple GPUs.

diff --git a/Engine/Source/Runtime/D3D12RHI/Private/D3D12Buffer.cpp b/Engine/Source/Runtime/D3D12RHI/Private/D3D12Buffer.cpp
--- a/Engine/Source/Runtime/D3D12RHI/Private/D3D12Buffer.cpp
+++ b/Engine/Source/Runtime/D3D12RHI/Private/D3D12Buffer.cpp
@@ -77,8 +77,11 @@ struct FRHICommandRenameUploadBuffer final : public FRHICommand<FRHICommandRenam
 		// Clear the resource if still bound to make sure the SRVs are rebound again on next operation. This needs to happen
 		// on the RHI timeline when this command runs at the top of the pipe (which can happen when locking buffers in
 		// RLM_WriteOnly_NoOverwrite mode).
-		CmdList.EnqueueLambda([ResourceLocation = &Resource->ResourceLocation](FRHICommandListBase& RHICmdList) {
-			FD3D12CommandContext& Context = (FD3D12CommandContext&)(RHICmdList.GetComputeContext().GetLowestLevelContext());
+		CmdList.EnqueueLambda([ResourceLocation = &Resource->ResourceLocation](FRHICommandListBase& RHICmdList)
+		{
+			const uint32 GPUIndex = 0; // @todo mgpu - seems wrong we're only doing this for the 0th GPU
+
+			FD3D12CommandContext& Context = FD3D12CommandContext::Get(RHICmdList, GPUIndex);
 			Context.ConditionalClearShaderResource(ResourceLocation);
 		});
 
@@ -113,22 +116,9 @@ struct FD3D12RHICommandInitializeBuffer final : public FRHICommand<FD3D12RHIComm
 
 	void Execute(FRHICommandListBase& CmdList)
 	{
-		FD3D12CommandContext& CommandContext = (FD3D12CommandContext&)(CmdList.GetComputeContext().GetLowestLevelContext());
-#if ENABLE_RHI_VALIDATION && WITH_MGPU
-		// Need to do a second level of indirection, to potentially go from FD3D12CommandContextRedirector to FD3D12CommandContext, when
-		// -rhivalidation is enabled.  The issue arises because GetLowestLevelContext on FValidationContext returns the contained
-		// FD3D12CommandContextRedirector, and doesn't attempt to call GetLowestLevelContext on that to get the physical FD3D12CommandContext
-		// required here.  Only when both validation and MGPU are present is it possible for this second level of indirection to be required.
-		// The call to "GetLowestLevelContext" is a nop (returns *this) if already a leaf command context.
-		//
-		// I thought about changing the original FValidationContext::GetLowestLevelContext implementation to attempt a second indirection there,
-		// but I wasn't sure what other side effects there might be for that change.  Here, the change is extremely safe, because the
-		// CommandContext passed to ExecuteOnCommandContext is solely used for a validation assert when WITH_MGPU is true, not actual rendering.
-		// So it can't affect behavior, it's just to avoid a spurious assert.
-		ExecuteOnCommandContext((FD3D12CommandContext&)CommandContext.GetLowestLevelContext());
-#else
+		const uint32 GPUIndex = 0; // @todo mgpu - why isn't this using the RHICmdList GPU mask?
+		FD3D12CommandContext& CommandContext = FD3D12CommandContext::Get(CmdList, GPUIndex);
 		ExecuteOnCommandContext(CommandContext);
-#endif
 	}
 
 	void ExecuteOnCommandContext(FD3D12CommandContext& CommandContext)
diff --git a/Engine/Source/Runtime/D3D12RHI/Private/D3D12CommandContext.h b/Engine/Source/Runtime/D3D12RHI/Private/D3D12CommandContext.h
--- a/Engine/Source/Runtime/D3D12RHI/Private/D3D12CommandContext.h
+++ b/Engine/Source/Runtime/D3D12RHI/Private/D3D12CommandContext.h
@@ -382,7 +382,13 @@ public:
 
 	virtual class FD3D12CommandContextRedirector* AsRedirector() { return nullptr; }
 
+	static FD3D12CommandContextBase& Get(FRHICommandListBase& RHICmdList)
+	{
+		return static_cast<FD3D12CommandContextBase&>(RHICmdList.GetComputeContext().GetLowestLevelContext());
+	}
+
 protected:
+	friend class FD3D12CommandContext;
 	virtual FD3D12CommandContext* GetContext(uint32 InGPUIndex) = 0;
 
 	FRHIGPUMask GPUMask;
@@ -399,9 +405,14 @@ public:
 	FD3D12CommandContext(class FD3D12Device* InParent, ED3D12QueueType QueueType, bool InIsDefaultContext);
 	virtual ~FD3D12CommandContext();
 
-	static FD3D12CommandContext& Get(FRHICommandListBase& RHICmdList)
+	static FD3D12CommandContext& Get(FRHICommandListBase& RHICmdList, uint32 GPUIndex)
 	{
-		return static_cast<FD3D12CommandContext&>(RHICmdList.GetComputeContext().GetLowestLevelContext());
+		FD3D12CommandContextBase& Base = FD3D12CommandContextBase::Get(RHICmdList);
+#if WITH_MGPU
+		return *Base.GetContext(GPUIndex);
+#else
+		return static_cast<FD3D12CommandContext&>(Base);
+#endif
 	}
 
 	virtual void OpenCommandList() override final;
@@ -535,7 +546,6 @@ public:
 	virtual void RHICopyTexture(FRHITexture* SourceTexture, FRHITexture* DestTexture, const FRHICopyTextureInfo& CopyInfo) final override;
 	virtual void RHICopyBufferRegion(FRHIBuffer* DestBuffer, uint64 DstOffset, FRHIBuffer* SourceBuffer, uint64 SrcOffset, uint64 NumBytes) final override;
 	virtual void RHICopyToStagingBuffer(FRHIBuffer* SourceBuffer, FRHIStagingBuffer* DestinationStagingBuffer, uint32 Offset, uint32 NumBytes) final override;
-	virtual void RHIWriteGPUFence(FRHIGPUFence* Fence) final override;
 	virtual void RHIBeginRenderQuery(FRHIRenderQuery* RenderQuery) final override;
 	virtual void RHIEndRenderQuery(FRHIRenderQuery* RenderQuery) final override;
 	virtual void RHICalibrateTimers(FRHITimestampCalibrationQuery* CalibrationQuery) final override;
@@ -757,10 +767,6 @@ public:
 	{
 		ContextRedirect(RHICopyToStagingBuffer(SourceBuffer, DestinationStagingBuffer, Offset, NumBytes));
 	}
-	FORCEINLINE virtual void RHIWriteGPUFence(FRHIGPUFence* Fence) final override
-	{
-		ContextRedirect(RHIWriteGPUFence(Fence));
-	}
 	FORCEINLINE virtual void RHISetShaderTexture(FRHIComputeShader* PixelShader, uint32 TextureIndex, FRHITexture* NewTexture) final override
 	{
 		ContextRedirect(RHISetShaderTexture(PixelShader, TextureIndex, NewTexture));
@@ -1081,13 +1087,6 @@ public:
 		return PhysicalContexts[GPUIndex];
 	}
 
-#if ENABLE_RHI_VALIDATION || WITH_MGPU
-	IRHIComputeContext& GetLowestLevelContext() final override
-	{
-		return *PhysicalContexts[0];
-	}
-#endif
-
 private:
 	TStaticArray<FD3D12CommandContext*, MAX_NUM_GPUS> PhysicalContexts;
 };
diff --git a/Engine/Source/Runtime/D3D12RHI/Private/D3D12DirectCommandListManager.cpp b/Engine/Source/Runtime/D3D12RHI/Private/D3D12DirectCommandListManager.cpp
--- a/Engine/Source/Runtime/D3D12RHI/Private/D3D12DirectCommandListManager.cpp
+++ b/Engine/Source/Runtime/D3D12RHI/Private/D3D12DirectCommandListManager.cpp
@@ -65,20 +65,19 @@ void FD3D12DynamicRHI::RHIWriteGPUFence_TopOfPipe(FRHICommandListBase& RHICmdLis
 		Fence->SyncPoints[GPUIndex] = FD3D12SyncPoint::Create(ED3D12SyncPointType::GPUAndCPU);
 	}
 
-	FDynamicRHI::RHIWriteGPUFence_TopOfPipe(RHICmdList, FenceRHI);
-}
-
-void FD3D12CommandContext::RHIWriteGPUFence(FRHIGPUFence* FenceRHI)
-{
-	FD3D12GPUFence* Fence = FD3D12DynamicRHI::ResourceCast(FenceRHI);
-
-	check(Fence);
-	check(Fence->SyncPoints[GetGPUIndex()]);
+	Fence->NumPendingWriteCommands.Increment();
+	RHICmdList.EnqueueLambda([Fence, SyncPoints = Fence->SyncPoints](FRHICommandListBase& ExecutingCmdList)
+	{
+		for (uint32 GPUIndex : ExecutingCmdList.GetGPUMask())
+		{
+			FD3D12CommandContext& Context = FD3D12CommandContext::Get(ExecutingCmdList, GPUIndex);
+			Context.SignalSyncPoint(SyncPoints[GPUIndex]);
+		}
 
-	SignalSyncPoint(Fence->SyncPoints[GetGPUIndex()]);
+		Fence->NumPendingWriteCommands.Decrement();
+	});
 }
 
-
 FGPUFenceRHIRef FD3D12DynamicRHI::RHICreateGPUFence(const FName& Name)
 {
 	LLM_SCOPE_BYNAME(TEXT("RHIMisc/CreateGPUFence"));
diff --git a/Engine/Source/Runtime/D3D12RHI/Private/D3D12Query.cpp b/Engine/Source/Runtime/D3D12RHI/Private/D3D12Query.cpp
--- a/Engine/Source/Runtime/D3D12RHI/Private/D3D12Query.cpp
+++ b/Engine/Source/Runtime/D3D12RHI/Private/D3D12Query.cpp
@@ -156,27 +156,33 @@ FRenderQueryRHIRef FD3D12DynamicRHI::RHICreateRenderQuery(ERenderQueryType Query
 
 void FD3D12DynamicRHI::RHIBeginOcclusionQueryBatch_TopOfPipe(FRHICommandListBase& RHICmdList, uint32 NumQueriesInBatch)
 {
-	checkf(RHICmdList.QueryBatchData == nullptr, TEXT("An occlusion query batch has already begun on this command list."));
+	// Each occlusion query batch uses a single sync point to signal when the results are ready (one per active GPU).
+	for (uint32 GPUIndex : RHICmdList.GetGPUMask())
+	{
+		checkf(RHICmdList.QueryBatchData[GPUIndex] == nullptr, TEXT("An occlusion query batch has already begun on this command list."));
 
-	// Each occlusion query batch uses a single sync point to signal when the results are ready. Allocate this now.
-	FD3D12SyncPointRef SyncPoint = FD3D12SyncPoint::Create(ED3D12SyncPointType::GPUAndCPU);
+		FD3D12SyncPointRef SyncPoint = FD3D12SyncPoint::Create(ED3D12SyncPointType::GPUAndCPU);
 
-	// Keep a reference on the RHI command list, so we can retrieve it later in BeginQuery/EndQuery/EndBatch.
-	RHICmdList.QueryBatchData = SyncPoint.GetReference();
-	SyncPoint->AddRef();
+		// Keep a reference on the RHI command list, so we can retrieve it later in BeginQuery/EndQuery/EndBatch.
+		RHICmdList.QueryBatchData[GPUIndex] = SyncPoint.GetReference();
+		SyncPoint->AddRef();
+	}
 
 	if (RHIConsoleVariables::GInsertOuterOcclusionQuery)
 	{
 		// Insert an outer query that encloses the whole batch
 		RHICmdList.EnqueueLambda([](FRHICommandListBase& ExecutingCmdList)
 		{
-			FD3D12CommandContext& Context = FD3D12CommandContext::Get(ExecutingCmdList);
+			for (uint32 GPUIndex : ExecutingCmdList.GetGPUMask())
+			{
+				FD3D12CommandContext& Context = FD3D12CommandContext::Get(ExecutingCmdList, GPUIndex);
 
-			if (!Context.OuterOcclusionQuery.IsValid())
-				Context.OuterOcclusionQuery = GDynamicRHI->RHICreateRenderQuery(RQT_Occlusion);
+				if (!Context.OuterOcclusionQuery.IsValid())
+					Context.OuterOcclusionQuery = GDynamicRHI->RHICreateRenderQuery(RQT_Occlusion);
 
-			Context.RHIBeginRenderQuery(Context.OuterOcclusionQuery);
-			Context.bOuterOcclusionQuerySubmitted = true;
+				Context.RHIBeginRenderQuery(Context.OuterOcclusionQuery);
+				Context.bOuterOcclusionQuerySubmitted = true;
+			}
 		});
 	}
 }
@@ -189,8 +195,8 @@ void FD3D12DynamicRHI::RHIBeginRenderQuery_TopOfPipe(FRHICommandListBase& RHICmd
 		FD3D12RenderQuery* Query = ResourceCast(RenderQuery, GPUIndex);
 		checkf(Query->Type == RQT_Occlusion, TEXT("Only occlusion queries support RHIBeginRenderQuery()."));
 
-		checkf(RHICmdList.QueryBatchData, TEXT("Cannot use an occlusion query outside of an occlusion query batch."));
-		Query->SyncPoint = static_cast<FD3D12SyncPoint*>(RHICmdList.QueryBatchData);
+		checkf(RHICmdList.QueryBatchData[GPUIndex], TEXT("Cannot use an occlusion query outside of an occlusion query batch."));
+		Query->SyncPoint = static_cast<FD3D12SyncPoint*>(RHICmdList.QueryBatchData[GPUIndex]);
 	}
 
 	// Enqueue the RHI command to record the BeginQuery() call on the context.
@@ -216,16 +222,16 @@ void FD3D12DynamicRHI::RHIEndRenderQuery_TopOfPipe(FRHICommandListBase& RHICmdLi
 		if (Query->Type == RQT_Occlusion)
 		{
 			// Occlusion query sync points are allocated by BeginOcclusionQueryBatch().
-			checkf(RHICmdList.QueryBatchData, TEXT("Cannot use an occlusion query outside of an occlusion query batch."));
+			checkf(RHICmdList.QueryBatchData[GPUIndex], TEXT("Cannot use an occlusion query outside of an occlusion query batch."));
 		}
 		else
 		{
 			// All other query types use one sync point per query.
 			Query->SyncPoint = FD3D12SyncPoint::Create(ED3D12SyncPointType::GPUAndCPU);
 
-			RHICmdList.EnqueueLambda([SyncPoint = Query->SyncPoint](FRHICommandListBase& ExecutingCmdList) mutable
+			RHICmdList.EnqueueLambda([SyncPoint = Query->SyncPoint, GPUIndex](FRHICommandListBase& ExecutingCmdList) mutable
 			{
-				FD3D12CommandContext& Context = FD3D12CommandContext::Get(ExecutingCmdList);
+				FD3D12CommandContext& Context = FD3D12CommandContext::Get(ExecutingCmdList, GPUIndex);
 				Context.BatchedSyncPoints.ToSignal.Emplace(MoveTemp(SyncPoint));
 			});
 		}
@@ -260,26 +266,29 @@ void FD3D12CommandContext::RHIEndRenderQuery(FRHIRenderQuery* QueryRHI)
 
 void FD3D12DynamicRHI::RHIEndOcclusionQueryBatch_TopOfPipe(FRHICommandListBase& RHICmdList)
 {
-	checkf(RHICmdList.QueryBatchData, TEXT("An occlusion query batch is not open on this command list."));
-	FD3D12SyncPointRef SyncPoint = static_cast<FD3D12SyncPoint*>(RHICmdList.QueryBatchData);
-
-	RHICmdList.EnqueueLambda([SyncPoint](FRHICommandListBase& ExecutingCmdList)
+	for (uint32 GPUIndex : RHICmdList.GetGPUMask())
 	{
-		FD3D12CommandContext& Context = FD3D12CommandContext::Get(ExecutingCmdList);
+		checkf(RHICmdList.QueryBatchData[GPUIndex], TEXT("An occlusion query batch is not open on this command list."));
+		FD3D12SyncPointRef SyncPoint = static_cast<FD3D12SyncPoint*>(RHICmdList.QueryBatchData[GPUIndex]);
 
-		// End the outer query
-		if (Context.bOuterOcclusionQuerySubmitted)
+		// Clear the sync point reference on the RHI command list
+		SyncPoint->Release();
+		RHICmdList.QueryBatchData[GPUIndex] = nullptr;
+
+		RHICmdList.EnqueueLambda([GPUIndex, SyncPoint = MoveTemp(SyncPoint)](FRHICommandListBase& ExecutingCmdList)
 		{
-			Context.RHIEndRenderQuery(Context.OuterOcclusionQuery);
-			Context.bOuterOcclusionQuerySubmitted = false;
-		}
+			FD3D12CommandContext& Context = FD3D12CommandContext::Get(ExecutingCmdList, GPUIndex);
 
-		Context.BatchedSyncPoints.ToSignal.Add(SyncPoint);
-	});
+			// End the outer query
+			if (Context.bOuterOcclusionQuerySubmitted)
+			{
+				Context.RHIEndRenderQuery(Context.OuterOcclusionQuery);
+				Context.bOuterOcclusionQuerySubmitted = false;
+			}
 
-	// Clear the sync point reference on the RHI command list
-	SyncPoint->Release();
-	RHICmdList.QueryBatchData = nullptr;
+			Context.BatchedSyncPoints.ToSignal.Add(SyncPoint);
+		});
+	}
 }
 
 bool FD3D12DynamicRHI::RHIGetRenderQueryResult(FRHIRenderQuery* QueryRHI, uint64& OutResult, bool bWait, uint32 QueryGPUIndex)
diff --git a/Engine/Source/Runtime/D3D12RHI/Private/D3D12RHI.cpp b/Engine/Source/Runtime/D3D12RHI/Private/D3D12RHI.cpp
--- a/Engine/Source/Runtime/D3D12RHI/Private/D3D12RHI.cpp
+++ b/Engine/Source/Runtime/D3D12RHI/Private/D3D12RHI.cpp
@@ -622,20 +622,27 @@ void FD3D12DynamicRHI::RHIFinishExternalComputeWork(uint32 InDeviceIndex, ID3D12
 void FD3D12DynamicRHI::RHITransitionResource(FRHICommandList& RHICmdList, FRHITexture* InTexture, D3D12_RESOURCE_STATES InState, uint32 InSubResource)
 {
 	FD3D12Texture* D3D12Texture = GetD3D12TextureFromRHITexture(InTexture);
-	FD3D12CommandContext::Get(RHICmdList).TransitionResource(D3D12Texture->GetResource(), D3D12_RESOURCE_STATE_TBD, InState, InSubResource);
+	for (uint32 GPUIndex : RHICmdList.GetGPUMask())
+	{
+		FD3D12CommandContext::Get(RHICmdList, GPUIndex).TransitionResource(D3D12Texture->GetResource(), D3D12_RESOURCE_STATE_TBD, InState, InSubResource);
+	}
 }
 
 void FD3D12DynamicRHI::RHISignalManualFence(FRHICommandList& RHICmdList, ID3D12Fence* Fence, uint64 Value)
 {
 	checkf(FRHIGPUMask::All() == FRHIGPUMask::GPU0(), TEXT("RHISignalManualFence cannot be used by multi-GPU code"));
-	FD3D12CommandContext& Context = FD3D12CommandContext::Get(RHICmdList);
+	const uint32 GPUIndex = 0;
+
+	FD3D12CommandContext& Context = FD3D12CommandContext::Get(RHICmdList, GPUIndex);
 	Context.SignalManualFence(Fence, Value);
 }
 
 void FD3D12DynamicRHI::RHIWaitManualFence(FRHICommandList& RHICmdList, ID3D12Fence* Fence, uint64 Value)
 {
 	checkf(FRHIGPUMask::All() == FRHIGPUMask::GPU0(), TEXT("RHIWaitManualFence cannot be used by multi-GPU code"));
-	FD3D12CommandContext& Context = FD3D12CommandContext::Get(RHICmdList);
+	const uint32 GPUIndex = 0;
+
+	FD3D12CommandContext& Context = FD3D12CommandContext::Get(RHICmdList, GPUIndex);
 	Context.WaitManualFence(Fence, Value);
 }
 
diff --git a/Engine/Source/Runtime/RHI/Private/RHICommandList.cpp b/Engine/Source/Runtime/RHI/Private/RHICommandList.cpp
--- a/Engine/Source/Runtime/RHI/Private/RHICommandList.cpp
+++ b/Engine/Source/Runtime/RHI/Private/RHICommandList.cpp
@@ -208,10 +208,14 @@ FRHICommandListBase::FRHICommandListBase(FRHICommandListBase&& Other)
 
 FRHICommandListBase::~FRHICommandListBase()
 {
-	check(QueryBatchData == nullptr);
 	checkf(!HasCommands(), TEXT("FRHICommandListBase has been deleted while it still contained commands. The command list was not submitted."));
 
 #if DO_CHECK
+	for (void* Data : QueryBatchData)
+	{
+		check(Data == nullptr);
+	}
+
 	if (PersistentState.RecordingThread == ERecordingThread::Render)
 	{
 		GRHICommandList.OutstandingCmdListCount.Decrement();
diff --git a/Engine/Source/Runtime/RHI/Public/RHICommandList.h b/Engine/Source/Runtime/RHI/Public/RHICommandList.h
--- a/Engine/Source/Runtime/RHI/Public/RHICommandList.h
+++ b/Engine/Source/Runtime/RHI/Public/RHICommandList.h
@@ -793,7 +793,7 @@ public:
 #endif
 
 public:
-	void* QueryBatchData = nullptr;
+	TStaticArray<void*, MAX_NUM_GPUS> QueryBatchData { InPlace, nullptr };
 
 private:
 	FRHICommandListBase(FPersistentState&& InPersistentState);
diff --git a/Engine/Source/Runtime/RHI/Public/RHIContext.h b/Engine/Source/Runtime/RHI/Public/RHIContext.h
--- a/Engine/Source/Runtime/RHI/Public/RHIContext.h
+++ b/Engine/Source/Runtime/RHI/Public/RHIContext.h
@@ -437,14 +437,10 @@ public:
 		return WrappingContext ? *WrappingContext : *this;
 	}
 
-#elif WITH_MGPU
-	// Needs to be virtual, because it's required for multi GPU resource uploads (FD3D12RHICommandInitializeBuffer, etc)
-	virtual IRHIComputeContext& GetLowestLevelContext() { return *this; }
-	inline IRHIComputeContext& GetHighestLevelContext() { return *this; }
 #else
 
 	// Fast implementations when the RHI validation layer is disabled.
-	inline IRHIComputeContext& GetLowestLevelContext() { return *this; }
+	inline IRHIComputeContext& GetLowestLevelContext () { return *this; }
 	inline IRHIComputeContext& GetHighestLevelContext() { return *this; }
 
 #endif
