Fix D3D12 RHI compile error when building test config

diff --git a/Engine/Source/Runtime/D3D12RHI/Private/D3D12CommandContext.h b/Engine/Source/Runtime/D3D12RHI/Private/D3D12CommandContext.h
--- a/Engine/Source/Runtime/D3D12RHI/Private/D3D12CommandContext.h
+++ b/Engine/Source/Runtime/D3D12RHI/Private/D3D12CommandContext.h
@@ -1087,6 +1087,7 @@ public:
 		return PhysicalContexts[GPUIndex];
 	}
 
+#if WITH_MGPU // @todo mgpu - remove the whole redirector when WITH_MGPU is false
 	virtual void StatsSetCategory(FRHIDrawStats* InStats, uint32 InCategoryID) final override
 	{
 		for (uint32 GPUIndex : PhysicalGPUMask)
@@ -1094,6 +1095,7 @@ public:
 			PhysicalContexts[GPUIndex]->StatsSetCategory(InStats, InCategoryID, GPUIndex);
 		}
 	}
+#endif
 
 private:
 	TStaticArray<FD3D12CommandContext*, MAX_NUM_GPUS> PhysicalContexts;
