Pass nullptr instead of a name to CreateEvent because it isn't clear if adding the "Local" prefix makes it process-local, or just session-local. We know for a fact that passing nullptr will create an anonymous object that no other process can manipulate.

diff --git a/Engine/Source/Runtime/D3D12RHI/Private/D3D12Submission.cpp b/Engine/Source/Runtime/D3D12RHI/Private/D3D12Submission.cpp
--- a/Engine/Source/Runtime/D3D12RHI/Private/D3D12Submission.cpp
+++ b/Engine/Source/Runtime/D3D12RHI/Private/D3D12Submission.cpp
@@ -42,7 +42,7 @@ public:
 	FD3D12Thread(TCHAR const* Name, EThreadPriority Priority, FD3D12DynamicRHI* RHI, FQueueFunc QueueFunc)
 		: RHI(RHI)
 		, QueueFunc(QueueFunc)
-		, Event(CreateEvent(nullptr, false, false, *(FString::Printf(TEXT("%s%s"), TEXT("Local\\"), Name))))
+		, Event(CreateEvent(nullptr, false, false, nullptr))
 		, Thread(FRunnableThread::Create(this, Name, 0, Priority))
 	{}
 
