Add more explicit range checking and error reporting to SetRayTracingHitGroup()

diff --git a/Engine/Source/Runtime/D3D12RHI/Private/D3D12RayTracing.cpp b/Engine/Source/Runtime/D3D12RHI/Private/D3D12RayTracing.cpp
--- a/Engine/Source/Runtime/D3D12RHI/Private/D3D12RayTracing.cpp
+++ b/Engine/Source/Runtime/D3D12RHI/Private/D3D12RayTracing.cpp
@@ -5349,6 +5349,18 @@ static void SetRayTracingHitGroup(
 	const uint32 RecordIndex = Scene->GetHitRecordBaseIndex(InstanceIndex, SegmentIndex) + ShaderSlot;
 	const uint32 GPUIndex = Device->GetGPUIndex();
 
+#if DO_CHECK
+	{
+		const uint32 NumSceneInstances = (uint32)Scene->Initializer.PerInstanceGeometries.Num();
+		checkf(InstanceIndex < NumSceneInstances, TEXT("Instance index %d is out of range for the scene that contains %d instances"), InstanceIndex, NumSceneInstances);
+
+		const FD3D12RayTracingGeometry* Geometry = FD3D12DynamicRHI::ResourceCast(Scene->Initializer.PerInstanceGeometries[InstanceIndex]);
+		const uint32 NumGeometrySegments = Geometry->GetNumSegments();
+		checkf(SegmentIndex < NumGeometrySegments, TEXT("Segment %d is out of range for ray tracing geometry '%s' that contains %d segments"),
+			SegmentIndex, Geometry->DebugName.IsNone() ? TEXT("UNKNOWN") : *Geometry->DebugName.ToString(), NumGeometrySegments);
+	}
+#endif // DO_CHECK
+
 	FHitGroupSystemParameters SystemParameters;
 	if (GPUIndex == 0)
 	{
