D3D12 Support logging improvements
- Log out more DXGI Adapter information when enumerating adapters in D3D12RHI.
- Log out more SM6 relevant information when enumerating adapters in D3D12RHI.
- Log out driver versions when enumerating adapters in D3D12RHI. Had to silence all the log output from FPlatformMisc::GetGPUDriverInfo to make this usable.
- Adding more logging when D3D12 reports it's not supported.
- Adding more logging to WindowsDynamicRHI so it's easier to track what a system did and didn't support.
- Adding -ForceDisableSM6 command line flag to force D3D12RHI to think it doesn't support ERHIFeatureLevel::SM6.
- Adding -ForceZeroAdapters command line flag to simulate when D3D12RHI can't find any valid adapters.

diff --git a/Engine/Source/Runtime/Core/Private/Windows/WindowsPlatformMisc.cpp b/Engine/Source/Runtime/Core/Private/Windows/WindowsPlatformMisc.cpp
--- a/Engine/Source/Runtime/Core/Private/Windows/WindowsPlatformMisc.cpp
+++ b/Engine/Source/Runtime/Core/Private/Windows/WindowsPlatformMisc.cpp
@@ -2486,7 +2486,7 @@ FString FWindowsPlatformMisc::GetPrimaryGPUBrand()
 #undef USE_SP_BACKUP_QUEUE_PARAMS_V1
 #undef USE_SP_INF_SIGNER_INFO_V1
 
-static void GetVideoDriverDetailsFromSetup(const FString& DeviceName, FGPUDriverInfo& Out)
+static void GetVideoDriverDetailsFromSetup(const FString& DeviceName, bool bVerbose, FGPUDriverInfo& Out)
 {
 	
 	HDEVINFO hDevInfo = SetupDiGetClassDevs(&GUID_DEVCLASS_DISPLAY, NULL, NULL, DIGCF_PRESENT);
@@ -2525,7 +2525,7 @@ static void GetVideoDriverDetailsFromSetup(const FString& DeviceName, FGPUDriver
 					}
 					else
 					{
-						UE_LOG(LogWindows, Log, TEXT("Failed to retrieve driver registry key for device %d"), Idx);
+						UE_CLOG(bVerbose, LogWindows, Log, TEXT("Failed to retrieve driver registry key for device %d"), Idx);
 					}
 					
 					break;
@@ -2534,7 +2534,7 @@ static void GetVideoDriverDetailsFromSetup(const FString& DeviceName, FGPUDriver
 			}
 			else
 			{
-				UE_LOG(LogWindows, Log, TEXT("Failed to retrieve driver description for device %d"), Idx);
+				UE_CLOG(bVerbose, LogWindows, Log, TEXT("Failed to retrieve driver description for device %d"), Idx);
 			}
 		}
 
@@ -2549,7 +2549,7 @@ static void GetVideoDriverDetailsFromSetup(const FString& DeviceName, FGPUDriver
 			}
 			else
 			{
-				UE_LOG(LogWindows, Log, TEXT("Failed to find provider name"));
+				UE_CLOG(bVerbose, LogWindows, Log, TEXT("Failed to find provider name"));
 			}
 			// Get the internal driver version
 			if (SetupDiGetDeviceProperty(hDevInfo, &DeviceInfoData, &DEVPKEY_Device_DriverVersion, &DataType,
@@ -2560,7 +2560,7 @@ static void GetVideoDriverDetailsFromSetup(const FString& DeviceName, FGPUDriver
 			}
 			else
 			{
-				UE_LOG(LogWindows, Log, TEXT("Failed to find internal driver version"));
+				UE_CLOG(bVerbose, LogWindows, Log, TEXT("Failed to find internal driver version"));
 			}
 			// Get the driver date
 			FILETIME FileTime;
@@ -2573,19 +2573,19 @@ static void GetVideoDriverDetailsFromSetup(const FString& DeviceName, FGPUDriver
 			}
 			else
 			{
-				UE_LOG(LogWindows, Log, TEXT("Failed to find driver date"));
+				UE_CLOG(bVerbose, LogWindows, Log, TEXT("Failed to find driver date"));
 			}
 		}
 		else
 		{
-			UE_LOG(LogWindows, Log, TEXT("Unable to find requested device '%s' using Setup API."), *DeviceName);
+			UE_CLOG(bVerbose, LogWindows, Log, TEXT("Unable to find requested device '%s' using Setup API."), *DeviceName);
 		}
 		
 		SetupDiDestroyDeviceInfoList(hDevInfo);
 	}
 	else
 	{
-		UE_LOG(LogWindows, Log, TEXT("Failed to initialize Setup API"));
+		UE_CLOG(bVerbose, LogWindows, Log, TEXT("Failed to initialize Setup API"));
 	}
 
 	if (!Out.ProviderName.IsEmpty())
@@ -2740,13 +2740,16 @@ static BOOL CALLBACK MonitorEnumProc(HMONITOR Monitor, HDC MonitorDC, LPRECT Rec
 	return TRUE;
 }
 
-FGPUDriverInfo FWindowsPlatformMisc::GetGPUDriverInfo(const FString& DeviceDescription)
+FGPUDriverInfo FWindowsPlatformMisc::GetGPUDriverInfo(const FString& DeviceDescription, bool bVerbose)
 {
-	// Also report monitor information here, for lack of a better place.
-	UE_LOG(LogWindows, Log, TEXT("Attached monitors:"));
-	int NumMonitors = 0;
-	EnumDisplayMonitors(nullptr, nullptr, MonitorEnumProc, (LPARAM)&NumMonitors);
-	UE_LOG(LogWindows, Log, TEXT("Found %d attached monitors."), NumMonitors);
+	if (bVerbose)
+	{
+		// Also report monitor information here, for lack of a better place.
+		UE_LOG(LogWindows, Log, TEXT("Attached monitors:"));
+		int NumMonitors = 0;
+		EnumDisplayMonitors(nullptr, nullptr, MonitorEnumProc, (LPARAM)&NumMonitors);
+		UE_LOG(LogWindows, Log, TEXT("Found %d attached monitors."), NumMonitors);
+	}
 
 	// to distinguish failed GetGPUDriverInfo() from call to GetGPUDriverInfo()
 	FGPUDriverInfo Ret;
@@ -2764,22 +2767,22 @@ FGPUDriverInfo FWindowsPlatformMisc::GetGPUDriverInfo(const FString& DeviceDescr
 
 	if (Method == 5)
 	{
-		UE_LOG(LogWindows, Log, TEXT("Gathering driver information using Windows Setup API"));
+		UE_CLOG(bVerbose, LogWindows, Log, TEXT("Gathering driver information using Windows Setup API"));
 		FGPUDriverInfo Local;
-		GetVideoDriverDetailsFromSetup(DeviceDescription, Local);
+		GetVideoDriverDetailsFromSetup(DeviceDescription, bVerbose, Local);
 
 		if(Local.IsValid() && Local.DeviceDescription == DeviceDescription)
 		{
 			return Local;
 		}
 
-		UE_LOG(LogWindows, Log, TEXT("Failed to get driver data for device '%s' using Setup API. Switching to fallback method."), *DeviceDescription);
+		UE_CLOG(bVerbose, LogWindows, Log, TEXT("Failed to get driver data for device '%s' using Setup API. Switching to fallback method."), *DeviceDescription);
 		Method = 4; // Switch to method 4 as a fallback if method 5 fails
 	}
 	
 	if(Method == 3 || Method == 4)
 	{
-		UE_LOG(LogWindows, Log, TEXT("EnumDisplayDevices:"));
+		UE_CLOG(bVerbose, LogWindows, Log, TEXT("EnumDisplayDevices:"));
 
 		for(uint32 i = 0; i < 256; ++i)
 		{
@@ -2795,7 +2798,7 @@ FGPUDriverInfo FWindowsPlatformMisc::GetGPUDriverInfo(const FString& DeviceDescr
 				break;
 			}
 
-			UE_LOG(LogWindows, Log, TEXT("   %d. '%s' (P:%d D:%d), name: '%s'"),
+			UE_CLOG(bVerbose, LogWindows, Log, TEXT("   %d. '%s' (P:%d D:%d), name: '%s'"),
 				i,
 				Device.DeviceString,
 				(Device.StateFlags & DISPLAY_DEVICE_PRIMARY_DEVICE) != 0,
@@ -2859,7 +2862,7 @@ FGPUDriverInfo FWindowsPlatformMisc::GetGPUDriverInfo(const FString& DeviceDescr
 
 		if(!DebugString.IsEmpty())
 		{
-			UE_LOG(LogWindows, Log, TEXT("DebugString: %s"), *DebugString);
+			UE_CLOG(bVerbose, LogWindows, Log, TEXT("DebugString: %s"), *DebugString);
 		}
 
 		return Ret;
@@ -2956,7 +2959,7 @@ FGPUDriverInfo FWindowsPlatformMisc::GetGPUDriverInfo(const FString& DeviceDescr
 
 	if(!DebugString.IsEmpty())
 	{
-		UE_LOG(LogWindows, Log, TEXT("DebugString: %s"), *DebugString);
+		UE_CLOG(bVerbose, LogWindows, Log, TEXT("DebugString: %s"), *DebugString);
 	}
 
 	return Ret;
diff --git a/Engine/Source/Runtime/Core/Public/Windows/WindowsPlatformMisc.h b/Engine/Source/Runtime/Core/Public/Windows/WindowsPlatformMisc.h
--- a/Engine/Source/Runtime/Core/Public/Windows/WindowsPlatformMisc.h
+++ b/Engine/Source/Runtime/Core/Public/Windows/WindowsPlatformMisc.h
@@ -180,7 +180,7 @@ struct CORE_API FWindowsPlatformMisc
 	static FString GetCPUVendor();
 	static FString GetCPUBrand();
 	static FString GetPrimaryGPUBrand();
-	static struct FGPUDriverInfo GetGPUDriverInfo(const FString& DeviceDescription);
+	static struct FGPUDriverInfo GetGPUDriverInfo(const FString& DeviceDescription, bool bVerbose = true);
 	static void GetOSVersions( FString& out_OSVersionLabel, FString& out_OSSubVersionLabel );
 	static FString GetOSVersion();
 	static bool GetDiskTotalAndFreeSpace( const FString& InPath, uint64& TotalNumberOfBytes, uint64& NumberOfFreeBytes );
diff --git a/Engine/Source/Runtime/D3D12RHI/Private/D3D12Adapter.cpp b/Engine/Source/Runtime/D3D12RHI/Private/D3D12Adapter.cpp
--- a/Engine/Source/Runtime/D3D12RHI/Private/D3D12Adapter.cpp
+++ b/Engine/Source/Runtime/D3D12RHI/Private/D3D12Adapter.cpp
@@ -211,6 +211,8 @@ FD3D12AdapterDesc::FD3D12AdapterDesc(const DXGI_ADAPTER_DESC& InDesc, int32 InAd
 	, ResourceBindingTier(DeviceInfo.ResourceBindingTier)
 	, ResourceHeapTier(DeviceInfo.ResourceHeapTier)
 	, MaxRHIFeatureLevel(DeviceInfo.MaxRHIFeatureLevel)
+	, bSupportsWaveOps(DeviceInfo.bSupportsWaveOps)
+	, bSupportsAtomic64(DeviceInfo.bSupportsAtomic64)
 {
 }
 
diff --git a/Engine/Source/Runtime/D3D12RHI/Private/D3D12Adapter.h b/Engine/Source/Runtime/D3D12RHI/Private/D3D12Adapter.h
--- a/Engine/Source/Runtime/D3D12RHI/Private/D3D12Adapter.h
+++ b/Engine/Source/Runtime/D3D12RHI/Private/D3D12Adapter.h
@@ -40,6 +40,8 @@ struct FD3D12DeviceBasicInfo
 	D3D12_RESOURCE_BINDING_TIER ResourceBindingTier;
 	D3D12_RESOURCE_HEAP_TIER    ResourceHeapTier;
 	uint32                      NumDeviceNodes;
+	bool                        bSupportsWaveOps;
+	bool                        bSupportsAtomic64;
 
 	ERHIFeatureLevel::Type      MaxRHIFeatureLevel;
 };
@@ -79,6 +81,12 @@ struct FD3D12AdapterDesc
 	/** Whether the GPU is integrated or discrete. */
 	bool bIsIntegrated = false;
 
+	/** Whether SM6.0 wave ops are supported */
+	bool bSupportsWaveOps = false;
+
+	/** Whether SM6.6 atomic64 wave ops are supported */
+	bool bSupportsAtomic64 = false;
+
 #if DXGI_MAX_FACTORY_INTERFACE >= 6
 	DXGI_GPU_PREFERENCE GpuPreference = DXGI_GPU_PREFERENCE_UNSPECIFIED;
 #endif
diff --git a/Engine/Source/Runtime/D3D12RHI/Private/Windows/WindowsD3D12Device.cpp b/Engine/Source/Runtime/D3D12RHI/Private/Windows/WindowsD3D12Device.cpp
--- a/Engine/Source/Runtime/D3D12RHI/Private/Windows/WindowsD3D12Device.cpp
+++ b/Engine/Source/Runtime/D3D12RHI/Private/Windows/WindowsD3D12Device.cpp
@@ -428,32 +428,30 @@ inline bool ShouldCheckBindlessSupport(EShaderPlatform ShaderPlatform)
 		|| RHIGetBindlessSamplersConfiguration(ShaderPlatform) != ERHIBindlessConfiguration::Disabled;
 }
 
-inline ERHIFeatureLevel::Type FindMaxRHIFeatureLevel(IDXGIAdapter* Adapter, ID3D12Device* Device, D3D_FEATURE_LEVEL InMaxFeatureLevel, D3D_SHADER_MODEL InMaxShaderModel, D3D12_RESOURCE_BINDING_TIER ResourceBindingTier)
+inline ERHIFeatureLevel::Type FindMaxRHIFeatureLevel(D3D_FEATURE_LEVEL InMaxFeatureLevel, D3D_SHADER_MODEL InMaxShaderModel, D3D12_RESOURCE_BINDING_TIER ResourceBindingTier, bool bSupportsWaveOps, bool bSupportsAtomic64)
 {
 	ERHIFeatureLevel::Type MaxRHIFeatureLevel = ERHIFeatureLevel::Num;
 
 	if (InMaxFeatureLevel >= D3D_FEATURE_LEVEL_12_0 && InMaxShaderModel >= D3D_SHADER_MODEL_6_6)
 	{
-		D3D12_FEATURE_DATA_D3D12_OPTIONS1 D3D12Caps1{};
-		Device->CheckFeatureSupport(D3D12_FEATURE_D3D12_OPTIONS1, &D3D12Caps1, sizeof(D3D12Caps1));
-
-		D3D12_FEATURE_DATA_D3D12_OPTIONS9 D3D12Caps9{};
-		Device->CheckFeatureSupport(D3D12_FEATURE_D3D12_OPTIONS9, &D3D12Caps9, sizeof(D3D12Caps9));
-
 		bool bHighEnoughBindingTier = true;
 		if (ShouldCheckBindlessSupport(SP_PCD3D_SM6))
 		{
 			bHighEnoughBindingTier = ResourceBindingTier >= D3D12_RESOURCE_BINDING_TIER_3;
 		}
 
-		if (D3D12Caps1.WaveOps && bHighEnoughBindingTier)
-		{
-			if (D3D12Caps9.AtomicInt64OnTypedResourceSupported || CheckDeviceForEmulatedAtomic64Support(Adapter, Device))
+		if (bSupportsWaveOps && bHighEnoughBindingTier && bSupportsAtomic64)
 		{
-			MaxRHIFeatureLevel = ERHIFeatureLevel::SM6;
+			if (FParse::Param(FCommandLine::Get(), TEXT("ForceDisableSM6")))
+			{
+				UE_LOG(LogD3D12RHI, Log, TEXT("ERHIFeatureLevel::SM6 disabled via -ForceDisableSM6"));
+			}
+			else
+			{
+				MaxRHIFeatureLevel = ERHIFeatureLevel::SM6;
+			}
 		}
 	}
-	}
 
 	if (MaxRHIFeatureLevel == ERHIFeatureLevel::Num && InMaxFeatureLevel >= D3D_FEATURE_LEVEL_11_0)
 	{
@@ -472,6 +470,22 @@ inline void GetResourceTiers(ID3D12Device* Device, D3D12_RESOURCE_BINDING_TIER&
 	OutResourceHeapTier = D3D12Caps.ResourceHeapTier;
 }
 
+inline bool GetSupportsWaveOps(ID3D12Device* Device)
+{
+	D3D12_FEATURE_DATA_D3D12_OPTIONS1 D3D12Caps1{};
+	Device->CheckFeatureSupport(D3D12_FEATURE_D3D12_OPTIONS1, &D3D12Caps1, sizeof(D3D12Caps1));
+
+	return D3D12Caps1.WaveOps;
+}
+
+inline bool GetSupportsAtomic64(IDXGIAdapter* Adapter, ID3D12Device* Device)
+{
+	D3D12_FEATURE_DATA_D3D12_OPTIONS9 D3D12Caps9{};
+	Device->CheckFeatureSupport(D3D12_FEATURE_D3D12_OPTIONS9, &D3D12Caps9, sizeof(D3D12Caps9));
+
+	return (CheckDeviceForEmulatedAtomic64Support(Adapter, Device) || D3D12Caps9.AtomicInt64OnTypedResourceSupported);
+}
+
 /**
  * Attempts to create a D3D12 device for the adapter using at minimum MinFeatureLevel.
  * If creation is successful, true is returned and the max supported feature level is set in OutMaxFeatureLevel.
@@ -488,7 +502,11 @@ static bool SafeTestD3D12CreateDevice(IDXGIAdapter* Adapter, D3D_FEATURE_LEVEL M
 			OutInfo.MaxShaderModel = FindHighestShaderModel(Device);
 			GetResourceTiers(Device, OutInfo.ResourceBindingTier, OutInfo.ResourceHeapTier);
 			OutInfo.NumDeviceNodes = Device->GetNodeCount();
-			OutInfo.MaxRHIFeatureLevel = FindMaxRHIFeatureLevel(Adapter, Device, OutInfo.MaxFeatureLevel, OutInfo.MaxShaderModel, OutInfo.ResourceBindingTier);
+
+			OutInfo.bSupportsWaveOps = GetSupportsWaveOps(Device);
+			OutInfo.bSupportsAtomic64 = GetSupportsAtomic64(Adapter, Device);
+
+			OutInfo.MaxRHIFeatureLevel = FindMaxRHIFeatureLevel(OutInfo.MaxFeatureLevel, OutInfo.MaxShaderModel, OutInfo.ResourceBindingTier, OutInfo.bSupportsWaveOps, OutInfo.bSupportsAtomic64);
 
 			Device->Release();
 			return true;
@@ -612,7 +630,7 @@ void FD3D12DynamicRHI::RHIHandleDisplayChange()
 }
 #endif
 
-static bool IsAdapterBlocked(FD3D12Adapter* InAdapter)
+static bool IsAdapterBlocked(const FD3D12Adapter* InAdapter)
 {
 #if !UE_BUILD_SHIPPING
 	if (InAdapter)
@@ -638,17 +656,30 @@ static bool IsAdapterBlocked(FD3D12Adapter* InAdapter)
 	return false;
 }
 
-static bool IsAdapterSupported(FD3D12Adapter* InAdapter, ERHIFeatureLevel::Type InRequestedFeatureLevel)
+inline ERHIFeatureLevel::Type GetAdapterMaxFeatureLevel(const FD3D12Adapter* InAdapter)
 {
 	if (InAdapter)
 	{
 		if (const FD3D12AdapterDesc& Desc = InAdapter->GetDesc(); Desc.IsValid())
 		{
-			return Desc.MaxRHIFeatureLevel != ERHIFeatureLevel::Num && Desc.MaxRHIFeatureLevel >= InRequestedFeatureLevel;
+			return Desc.MaxRHIFeatureLevel;
 		}
 	}
 
-	return false;
+	return ERHIFeatureLevel::Num;
+}
+
+static bool IsAdapterSupported(const FD3D12Adapter* InAdapter, ERHIFeatureLevel::Type InRequestedFeatureLevel)
+{
+	const ERHIFeatureLevel::Type AdapterMaxFeatureLevel = GetAdapterMaxFeatureLevel(InAdapter);
+	return AdapterMaxFeatureLevel != ERHIFeatureLevel::Num && AdapterMaxFeatureLevel >= InRequestedFeatureLevel;
+}
+
+static bool CheckIfAgilitySDKLoaded()
+{
+	const TCHAR* AgilitySDKDllName = TEXT("D3D12Core.dll");
+	HMODULE AgilitySDKDllHandle = ::GetModuleHandle(AgilitySDKDllName);
+	return AgilitySDKDllHandle != NULL;
 }
 
 bool FD3D12DynamicRHIModule::IsSupported(ERHIFeatureLevel::Type RequestedFeatureLevel)
@@ -669,9 +700,47 @@ bool FD3D12DynamicRHIModule::IsSupported(ERHIFeatureLevel::Type RequestedFeature
 		FindAdapter();
 	}
 
-	return ChosenAdapters.Num() > 0
-		&& !IsAdapterBlocked(ChosenAdapters[0].Get())
-		&& IsAdapterSupported(ChosenAdapters[0].Get(), RequestedFeatureLevel);
+	if (ChosenAdapters.Num() == 0)
+	{
+		UE_LOG(LogD3D12RHI, Log, TEXT("No adapters were found."));
+		return false;
+	}
+
+	const FD3D12Adapter* Adapter = ChosenAdapters[0].Get();
+	if (!Adapter || !Adapter->GetDesc().IsValid())
+	{
+		UE_LOG(LogD3D12RHI, Log, TEXT("Adapter was not found"));
+		return false;
+	}
+
+	if (IsAdapterBlocked(Adapter))
+	{
+		UE_LOG(LogD3D12RHI, Log, TEXT("Adapter was blocked by RHI.BlockIHVD3D12"));
+		return false;
+	}
+
+	if (!IsAdapterSupported(Adapter, RequestedFeatureLevel))
+	{
+		auto GetFeatureLevelNameInline = [](ERHIFeatureLevel::Type InFeatureLevel)
+		{
+			FString FeatureLevelName;
+			::GetFeatureLevelName(InFeatureLevel, FeatureLevelName);
+			return FeatureLevelName;
+		};
+
+		const FString SupportedFeatureLevelName = GetFeatureLevelNameInline(GetAdapterMaxFeatureLevel(Adapter));
+		const FString RequestedFeatureLevelName = GetFeatureLevelNameInline(RequestedFeatureLevel);
+
+		UE_LOG(LogD3D12RHI, Log,
+			TEXT("Adapter only supports up to Feature Level '%s', requested Feature Level was '%s'"),
+			*SupportedFeatureLevelName,
+			*RequestedFeatureLevelName
+		);
+
+		return false;
+	}
+
+	return true;
 }
 
 namespace D3D12RHI
@@ -714,6 +783,23 @@ static uint32 CountAdapterOutputs(TRefCountPtr<IDXGIAdapter>& Adapter)
 	return OutputCount;
 }
 
+inline bool IsAdapterIntegrated(IDXGIAdapter* InAdapter)
+{
+	TRefCountPtr<IDXGIAdapter3> Adapter3;
+	InAdapter->QueryInterface(IID_PPV_ARGS(Adapter3.GetInitReference()));
+
+	// Simple heuristic but without profiling it's hard to do better
+	DXGI_QUERY_VIDEO_MEMORY_INFO NonLocalVideoMemoryInfo{};
+	if (Adapter3 && SUCCEEDED(Adapter3->QueryVideoMemoryInfo(0, DXGI_MEMORY_SEGMENT_GROUP_NON_LOCAL, &NonLocalVideoMemoryInfo)))
+	{
+		return NonLocalVideoMemoryInfo.Budget == 0;
+	}
+
+	// TODO: Using GPUDetect for Intel GPUs to check for integrated vs discrete status, pending GPUDetect update
+
+	return true;
+}
+
 void FD3D12DynamicRHIModule::FindAdapter()
 {
 	// Once we've chosen one we don't need to do it again.
@@ -793,55 +879,57 @@ void FD3D12DynamicRHIModule::FindAdapter()
 			if (SafeTestD3D12CreateDevice(TempAdapter, MinRequiredFeatureLevel, DeviceInfo))
 			{
 				check(DeviceInfo.NumDeviceNodes > 0);
+
 				// Log some information about the available D3D12 adapters.
-				DXGI_ADAPTER_DESC AdapterDesc;
+				DXGI_ADAPTER_DESC AdapterDesc{};
 				VERIFYD3D12RESULT(TempAdapter->GetDesc(&AdapterDesc));
-				uint32 OutputCount = CountAdapterOutputs(TempAdapter);
+
+				const uint32 OutputCount = CountAdapterOutputs(TempAdapter);
 
 				UE_LOG(LogD3D12RHI, Log,
-					TEXT("Found D3D12 adapter %u: %s (Max supported Feature Level %s, shader model %d.%d)"),
+					TEXT("Found D3D12 adapter %u: %s (VendorId: %04x, DeviceId: %04x, SubSysId: %04x, Revision: %04x"),
 					AdapterIndex,
 					AdapterDesc.Description,
+					AdapterDesc.VendorId, AdapterDesc.DeviceId, AdapterDesc.SubSysId, AdapterDesc.Revision
+				);
+				UE_LOG(LogD3D12RHI, Log,
+					TEXT("  Max supported Feature Level %s, shader model %d.%d, binding tier %d, wave ops %s, atomic64 %s"),
 					GetFeatureLevelString(DeviceInfo.MaxFeatureLevel),
-					(DeviceInfo.MaxShaderModel >> 4), (DeviceInfo.MaxShaderModel & 0xF)
+					(DeviceInfo.MaxShaderModel >> 4), (DeviceInfo.MaxShaderModel & 0xF),
+					DeviceInfo.ResourceBindingTier,
+					DeviceInfo.bSupportsWaveOps ? TEXT("supported") : TEXT("unsupported"),
+					DeviceInfo.bSupportsAtomic64 ? TEXT("supported") : TEXT("unsupported")
 				);
+
 				UE_LOG(LogD3D12RHI, Log,
-					TEXT("Adapter has %uMB of dedicated video memory, %uMB of dedicated system memory, and %uMB of shared system memory, %d output[s]"),
+					TEXT("  Adapter has %uMB of dedicated video memory, %uMB of dedicated system memory, and %uMB of shared system memory, %d output[s]"),
 					(uint32)(AdapterDesc.DedicatedVideoMemory / (1024 * 1024)),
 					(uint32)(AdapterDesc.DedicatedSystemMemory / (1024 * 1024)),
 					(uint32)(AdapterDesc.SharedSystemMemory / (1024 * 1024)),
 					OutputCount
 				);
 
-				bool bIsAMD = AdapterDesc.VendorId == 0x1002;
-				bool bIsIntel = AdapterDesc.VendorId == 0x8086;
-				bool bIsNVIDIA = AdapterDesc.VendorId == 0x10DE;
-				bool bIsWARP = AdapterDesc.VendorId == 0x1414;
-
-				// Simple heuristic but without profiling it's hard to do better
-				bool bIsNonLocalMemoryPresent = false;
-				TRefCountPtr<IDXGIAdapter3> TempDxgiAdapter3;
-				DXGI_QUERY_VIDEO_MEMORY_INFO NonLocalVideoMemoryInfo;
-				if (SUCCEEDED(TempAdapter->QueryInterface(_uuidof(IDXGIAdapter3), (void**)TempDxgiAdapter3.GetInitReference())) &&
-					TempDxgiAdapter3.IsValid() && SUCCEEDED(TempDxgiAdapter3->QueryVideoMemoryInfo(0, DXGI_MEMORY_SEGMENT_GROUP_NON_LOCAL, &NonLocalVideoMemoryInfo)))
+				const bool bIsWARP = AdapterDesc.VendorId == 0x1414;
+
+				if (!bIsWARP)
 				{
-					bIsNonLocalMemoryPresent = NonLocalVideoMemoryInfo.Budget != 0;
+					const FGPUDriverInfo GPUDriverInfo = FPlatformMisc::GetGPUDriverInfo(AdapterDesc.Description, false);
+					UE_LOG(LogD3D12RHI, Log, TEXT("  Driver Version: %s (internal:%s, unified:%s)"), *GPUDriverInfo.UserDriverVersion, *GPUDriverInfo.InternalDriverVersion, *GPUDriverInfo.GetUnifiedDriverVersion());
+					UE_LOG(LogD3D12RHI, Log, TEXT("     Driver Date: %s"), *GPUDriverInfo.DriverDate);
 				}
-				// TODO: Using GPUDetect for Intel GPUs to check for integrated vs discrete status, pending GPUDetect update
-				const bool bIsIntegrated = !bIsNonLocalMemoryPresent;
-
-				// PerfHUD is for performance profiling
-				const bool bIsPerfHUD = !FCString::Stricmp(AdapterDesc.Description, TEXT("NVIDIA PerfHUD"));
 
 				FD3D12AdapterDesc CurrentAdapter(AdapterDesc, AdapterIndex, DeviceInfo);
 
 				CurrentAdapter.NumDeviceNodes = DeviceInfo.NumDeviceNodes;
 				CurrentAdapter.GpuPreference = GpuPreference;
-				CurrentAdapter.bIsIntegrated = bIsIntegrated;
+				CurrentAdapter.bIsIntegrated = IsAdapterIntegrated(TempAdapter);
 
 				// If requested WARP, then reject all other adapters. If WARP not requested, then reject the WARP device if software rendering support is disallowed
 				const bool bSkipWARP = (bRequestedWARP && !bIsWARP) || (!bRequestedWARP && bIsWARP && !bAllowSoftwareRendering);
 
+				// PerfHUD is for performance profiling
+				const bool bIsPerfHUD = !FCString::Stricmp(AdapterDesc.Description, TEXT("NVIDIA PerfHUD"));
+
 				// we don't allow the PerfHUD adapter
 				const bool bSkipPerfHUDAdapter = bIsPerfHUD && !bAllowPerfHUD;
 
@@ -851,11 +939,11 @@ void FD3D12DynamicRHIModule::FindAdapter()
 				// the user wants a specific adapter, not this one
 				const bool bSkipExplicitAdapter = CVarExplicitAdapterValue >= 0 && AdapterIndex != CVarExplicitAdapterValue;
 
-				const bool bSkipAdapter = bSkipWARP || bSkipPerfHUDAdapter || bSkipHmdGraphicsAdapter || bSkipExplicitAdapter;
+				const bool bSkipAdapter = bSkipWARP || bSkipPerfHUDAdapter || bSkipHmdGraphicsAdapter || bSkipExplicitAdapter || FParse::Param(FCommandLine::Get(), TEXT("ForceZeroAdapters"));
 
 				if (!bSkipAdapter)
 				{
-					if (!bIsWARP && !bIsIntegrated && !FirstDiscreteAdapter.IsValid())
+					if (!bIsWARP && !CurrentAdapter.bIsIntegrated && !FirstDiscreteAdapter.IsValid())
 					{
 						FirstDiscreteAdapter = CurrentAdapter;
 					}
@@ -898,6 +986,8 @@ void FD3D12DynamicRHIModule::FindAdapter()
 		ChosenAdapters.Add(NewAdapter);
 	}
 
+	UE_LOG(LogD3D12RHI, Log, TEXT("DirectX Agility SDK runtime %s."), CheckIfAgilitySDKLoaded() ? TEXT("found") : TEXT("not found"));
+
 	if (ChosenAdapters.Num() > 0 && ChosenAdapters[0]->GetDesc().IsValid())
 	{
 		UE_LOG(LogD3D12RHI, Log, TEXT("Chosen D3D12 Adapter Id = %u"), ChosenAdapters[0]->GetAdapterIndex());
diff --git a/Engine/Source/Runtime/RHI/Private/Windows/WindowsDynamicRHI.cpp b/Engine/Source/Runtime/RHI/Private/Windows/WindowsDynamicRHI.cpp
--- a/Engine/Source/Runtime/RHI/Private/Windows/WindowsDynamicRHI.cpp
+++ b/Engine/Source/Runtime/RHI/Private/Windows/WindowsDynamicRHI.cpp
@@ -679,7 +679,9 @@ static bool HandleUnsupportedFeatureLevel(EWindowsRHI& WindowsRHI, ERHIFeatureLe
 	{
 		if (WindowsRHI == EWindowsRHI::D3D12 && FeatureLevel == ERHIFeatureLevel::SM6)
 		{
-			FMessageDialog::Open(EAppMsgType::Ok, NSLOCTEXT("WindowsDynamicRHI", "RequiredDX12SM6", "DX12 SM6 is not supported on your system. Try running without the -sm6 command line argument."));
+			UE_LOG(LogRHI, Log, TEXT("RHI %s with Feature Level %s is not supported on your system."), ModuleNameFromWindowsRHI(WindowsRHI), *LexToString(FeatureLevel));
+
+			FMessageDialog::Open(EAppMsgType::Ok, NSLOCTEXT("WindowsDynamicRHI", "RequiredDX12SM6", "DirectX 12 with Feature Level SM6 is not supported on your system. Try running without the -sm6 command line argument."));
 			FPlatformMisc::RequestExit(1);
 		}
 
@@ -693,8 +695,6 @@ static bool HandleUnsupportedFeatureLevel(EWindowsRHI& WindowsRHI, ERHIFeatureLe
 
 	if (TOptional<ERHIFeatureLevel::Type> FallbackFeatureLevel = Config.GetNextHighestTargetedFeatureLevel(WindowsRHI, FeatureLevel))
 	{
-		UE_LOG(LogRHI, Log, TEXT("RHI %s with Feature Level %s not supported, falling back to Feature Level %s"), ModuleNameFromWindowsRHI(WindowsRHI), *LexToString(FeatureLevel), *LexToString(FallbackFeatureLevel.GetValue()));
-
 		FeatureLevel = FallbackFeatureLevel.GetValue();
 		return true;
 	}
@@ -708,8 +708,11 @@ static bool HandleUnsupportedRHI(EWindowsRHI& WindowsRHI, ERHIFeatureLevel::Type
 	{
 		if (ForcedRHI == EWindowsRHI::D3D12)
 		{
-			FMessageDialog::Open(EAppMsgType::Ok, NSLOCTEXT("WindowsDynamicRHI", "RequiredDX12", "DX12 is not supported on your system. Try running without the -dx12 or -d3d12 command line argument."));
+			UE_LOG(LogRHI, Log, TEXT("RHI %s with Feature Level %s is not supported on your system."), ModuleNameFromWindowsRHI(WindowsRHI), *LexToString(FeatureLevel));
+
+			FMessageDialog::Open(EAppMsgType::Ok, NSLOCTEXT("WindowsDynamicRHI", "RequiredDX12", "DirectX 12 is not supported on your system. Try running without the -dx12 or -d3d12 command line argument."));
 			FPlatformMisc::RequestExit(1);
+			return false;
 		}
 	}
 
@@ -717,14 +720,17 @@ static bool HandleUnsupportedRHI(EWindowsRHI& WindowsRHI, ERHIFeatureLevel::Type
 	{
 		if (TOptional<ERHIFeatureLevel::Type> D3D11FeatureLevel = Config.GetHighestSupportedFeatureLevel(EWindowsRHI::D3D11))
 		{
-			UE_LOG(LogRHI, Log, TEXT("D3D12 is not supported, falling back to D3D11 with Feature Level %s"), *LexToString(D3D11FeatureLevel.GetValue()));
-
 			WindowsRHI = EWindowsRHI::D3D11;
 			FeatureLevel = D3D11FeatureLevel.GetValue();
 			return true;
 		}
+	}
+
+	UE_LOG(LogRHI, Log, TEXT("RHI %s is not supported on your system."), ModuleNameFromWindowsRHI(WindowsRHI));
 
-		FMessageDialog::Open(EAppMsgType::Ok, NSLOCTEXT("WindowsDynamicRHI", "RequiredDX12", "DX12 is not supported on your system. Try running without the -dx12 or -d3d12 command line argument."));
+	if (WindowsRHI == EWindowsRHI::D3D12)
+	{
+		FMessageDialog::Open(EAppMsgType::Ok, NSLOCTEXT("WindowsDynamicRHI", "RequiredDX12", "DirectX 12 is not supported on your system. Try running without the -dx12 or -d3d12 command line argument."));
 		FPlatformMisc::RequestExit(1);
 	}
 
@@ -794,14 +800,24 @@ static IDynamicRHIModule* LoadDynamicRHIModule(ERHIFeatureLevel::Type& DesiredFe
 		FApp::SetGraphicsRHI(RHIName);
 
 		const TCHAR* ModuleName = ModuleNameFromWindowsRHI(ChosenRHI);
+
+		UE_LOG(LogRHI, Log, TEXT("Loading RHI %s"), ModuleName);
+
 		IDynamicRHIModule* DynamicRHIModule = FModuleManager::LoadModulePtr<IDynamicRHIModule>(ModuleName);
 
+		UE_LOG(LogRHI, Log, TEXT("Checking if RHI %s with Feature Level %s is supported by your system."), ModuleName, *LexToString(DesiredFeatureLevel));
+
 		if (DynamicRHIModule && DynamicRHIModule->IsSupported(DesiredFeatureLevel))
 		{
+			UE_LOG(LogRHI, Log, TEXT("RHI %s with Feature Level %s is supported and will be used."), ModuleName, *LexToString(DesiredFeatureLevel));
+
 			LoadedRHIModuleName = ModuleName;
 			return DynamicRHIModule;
 		}
 
+		const EWindowsRHI PreviousRHI = ChosenRHI;
+		const ERHIFeatureLevel::Type PreviousFeatureLevel = DesiredFeatureLevel;
+
 		bTryWithNewConfig = HandleUnsupportedFeatureLevel(ChosenRHI, DesiredFeatureLevel, ForcedFeatureLevel, Config);
 
 		if (!bTryWithNewConfig)
@@ -809,8 +825,16 @@ static IDynamicRHIModule* LoadDynamicRHIModule(ERHIFeatureLevel::Type& DesiredFe
 			bTryWithNewConfig = HandleUnsupportedRHI(ChosenRHI, DesiredFeatureLevel, ForcedRHI, Config);
 		}
 
+		if (bTryWithNewConfig)
+		{
+			UE_LOG(LogRHI, Log, TEXT("RHI %s with Feature Level %s is not supported on your system, attempting to fall back to RHI %s with Feature Level %s"),
+				ModuleNameFromWindowsRHI(PreviousRHI), *LexToString(PreviousFeatureLevel),
+				ModuleNameFromWindowsRHI(ChosenRHI), *LexToString(DesiredFeatureLevel));
+		}
 	} while (bTryWithNewConfig);
 
+	UE_LOG(LogRHI, Log, TEXT("RHI %s with Feature Level %s is not supported on your system. No RHI was supported, failing initialization."), ModuleNameFromWindowsRHI(ChosenRHI), *LexToString(DesiredFeatureLevel));
+
 	return nullptr;
 }
 
