Enable DRED when r.GPUCrashDebugging is enabled (help track down GPU crashes)

diff --git a/Engine/Source/Runtime/D3D12RHI/Private/D3D12Adapter.cpp b/Engine/Source/Runtime/D3D12RHI/Private/D3D12Adapter.cpp
--- a/Engine/Source/Runtime/D3D12RHI/Private/D3D12Adapter.cpp
+++ b/Engine/Source/Runtime/D3D12RHI/Private/D3D12Adapter.cpp
@@ -1259,7 +1259,7 @@ void FD3D12Adapter::SetupGPUCrashDebuggingModesCommon()
 				EnumAddFlags(GPUCrashDebuggingModes, DebuggingMode);
 			}
 		};
-		ParseCVar(TEXT("r.GPUCrashDebugging"), ED3D12GPUCrashDebuggingModes((int)ED3D12GPUCrashDebuggingModes::BreadCrumbs | (int)ED3D12GPUCrashDebuggingModes::NvAftermath));
+		ParseCVar(TEXT("r.GPUCrashDebugging"), ED3D12GPUCrashDebuggingModes((int)ED3D12GPUCrashDebuggingModes::BreadCrumbs | (int)ED3D12GPUCrashDebuggingModes::NvAftermath | (int)ED3D12GPUCrashDebuggingModes::DRED));
 		ParseCVar(TEXT("r.D3D12.BreadCrumbs"), ED3D12GPUCrashDebuggingModes::BreadCrumbs);
 		ParseCVar(TEXT("r.D3D12.NvAfterMath"), ED3D12GPUCrashDebuggingModes::NvAftermath);
 		ParseCVar(TEXT("r.D3D12.DRED"), ED3D12GPUCrashDebuggingModes::DRED);
