Fixing HoloLens compile errors in D3D12RHI.

diff --git a/Engine/Source/Runtime/D3D12RHI/Private/Windows/WindowsD3D12Device.cpp b/Engine/Source/Runtime/D3D12RHI/Private/Windows/WindowsD3D12Device.cpp
--- a/Engine/Source/Runtime/D3D12RHI/Private/Windows/WindowsD3D12Device.cpp
+++ b/Engine/Source/Runtime/D3D12RHI/Private/Windows/WindowsD3D12Device.cpp
@@ -675,12 +675,14 @@ static bool IsAdapterSupported(const FD3D12Adapter* InAdapter, ERHIFeatureLevel:
 	return AdapterMaxFeatureLevel != ERHIFeatureLevel::Num && AdapterMaxFeatureLevel >= InRequestedFeatureLevel;
 }
 
+#if !PLATFORM_HOLOLENS
 static bool CheckIfAgilitySDKLoaded()
 {
 	const TCHAR* AgilitySDKDllName = TEXT("D3D12Core.dll");
 	HMODULE AgilitySDKDllHandle = ::GetModuleHandle(AgilitySDKDllName);
 	return AgilitySDKDllHandle != NULL;
 }
+#endif // !PLATFORM_HOLOLENS
 
 bool FD3D12DynamicRHIModule::IsSupported(ERHIFeatureLevel::Type RequestedFeatureLevel)
 {
@@ -911,12 +913,14 @@ void FD3D12DynamicRHIModule::FindAdapter()
 
 				const bool bIsWARP = AdapterDesc.VendorId == 0x1414;
 
+#if !PLATFORM_HOLOLENS
 				if (!bIsWARP)
 				{
 					const FGPUDriverInfo GPUDriverInfo = FPlatformMisc::GetGPUDriverInfo(AdapterDesc.Description, false);
 					UE_LOG(LogD3D12RHI, Log, TEXT("  Driver Version: %s (internal:%s, unified:%s)"), *GPUDriverInfo.UserDriverVersion, *GPUDriverInfo.InternalDriverVersion, *GPUDriverInfo.GetUnifiedDriverVersion());
 					UE_LOG(LogD3D12RHI, Log, TEXT("     Driver Date: %s"), *GPUDriverInfo.DriverDate);
 				}
+#endif // !PLATFORM_HOLOLENS
 
 				FD3D12AdapterDesc CurrentAdapter(AdapterDesc, AdapterIndex, DeviceInfo);
 
@@ -986,7 +990,9 @@ void FD3D12DynamicRHIModule::FindAdapter()
 		ChosenAdapters.Add(NewAdapter);
 	}
 
+#if !PLATFORM_HOLOLENS
 	UE_LOG(LogD3D12RHI, Log, TEXT("DirectX Agility SDK runtime %s."), CheckIfAgilitySDKLoaded() ? TEXT("found") : TEXT("not found"));
+#endif
 
 	if (ChosenAdapters.Num() > 0 && ChosenAdapters[0]->GetDesc().IsValid())
 	{
