Make sure correct plane and mip index is used on stencil UAV views

diff --git a/Engine/Source/Runtime/D3D12RHI/Private/D3D12CommandList.cpp b/Engine/Source/Runtime/D3D12RHI/Private/D3D12CommandList.cpp
--- a/Engine/Source/Runtime/D3D12RHI/Private/D3D12CommandList.cpp
+++ b/Engine/Source/Runtime/D3D12RHI/Private/D3D12CommandList.cpp
@@ -492,7 +492,7 @@ void FD3D12ContextCommon::TransitionResource(FD3D12UnorderedAccessView* View, D3
 
 	case D3D12_UAV_DIMENSION_TEXTURE2D:
 		// Only one subresource to transition
-		TransitionResource(Resource, D3D12_RESOURCE_STATE_TBD, After, Desc.Texture2D.MipSlice);
+		TransitionResource(Resource, D3D12_RESOURCE_STATE_TBD, After, View->GetViewSubresourceSubset());
 		break;
 
 	case D3D12_UAV_DIMENSION_TEXTURE2DARRAY:
