[Backout] - CL23009193

- change was incorrect as UAV mip and plane index are relative to the resource description created in RHICreateUnorderedAccessViewStencil not to the original texture

Original CL Desc
-----------------------------------------------------------------
Make sure correct plane and mip index is used on stencil UAV views

diff --git a/Engine/Source/Runtime/D3D12RHI/Private/D3D12CommandList.cpp b/Engine/Source/Runtime/D3D12RHI/Private/D3D12CommandList.cpp
--- a/Engine/Source/Runtime/D3D12RHI/Private/D3D12CommandList.cpp
+++ b/Engine/Source/Runtime/D3D12RHI/Private/D3D12CommandList.cpp
@@ -492,7 +492,7 @@ void FD3D12ContextCommon::TransitionResource(FD3D12UnorderedAccessView* View, D3
 
 	case D3D12_UAV_DIMENSION_TEXTURE2D:
 		// Only one subresource to transition
-		TransitionResource(Resource, D3D12_RESOURCE_STATE_TBD, After, View->GetViewSubresourceSubset());
+		TransitionResource(Resource, D3D12_RESOURCE_STATE_TBD, After, Desc.Texture2D.MipSlice);
 		break;
 
 	case D3D12_UAV_DIMENSION_TEXTURE2DARRAY:
