Enable the D3D12 submission thread
 - When not running with MGPU.
 - Added a cvar to control this at boot time: rhi.UseSubmissionThread

diff --git a/Engine/Source/Runtime/D3D12RHI/Private/D3D12Submission.cpp b/Engine/Source/Runtime/D3D12RHI/Private/D3D12Submission.cpp
--- a/Engine/Source/Runtime/D3D12RHI/Private/D3D12Submission.cpp
+++ b/Engine/Source/Runtime/D3D12RHI/Private/D3D12Submission.cpp
@@ -9,7 +9,7 @@
 #endif
 
 // These defines control which threads are enabled in the GPU submission pipeline.
-#define D3D12_USE_SUBMISSION_THREAD (1 && (WITH_MGPU == 0)) // @todo mgpu - fix crashes when submission thread is enabled
+#define D3D12_USE_SUBMISSION_THREAD (1)
 #define D3D12_USE_INTERRUPT_THREAD  (1 && D3D12_PLATFORM_SUPPORTS_BLOCKING_FENCES)
 
 // When enabled, GPU timestamp queries are adjusted to remove idle time caused by CPU bubbles.
@@ -23,6 +23,16 @@ static FAutoConsoleVariableRef CVarD3D12AsyncPayloadMerge(
 	ECVF_RenderThreadSafe
 );
 
+// @todo mgpu - fix crashes when submission thread is enabled)
+static TAutoConsoleVariable<int32> CVarRHIUseSubmissionThread(
+	TEXT("rhi.UseSubmissionThread"),
+	1,
+	TEXT("Whether to enable the RHI submission thread.\n")
+	TEXT("  0: No\n")
+	TEXT("  1: Yes, but not when running with multi-gpu.\n")
+	TEXT("  2: Yes, always\n"),
+	ECVF_ReadOnly);
+
 DECLARE_CYCLE_STAT(TEXT("Submit"), STAT_D3D12Submit, STATGROUP_D3D12RHI);
 
 DECLARE_CYCLE_STAT(TEXT("GPU Total Time [All Queues]"), STAT_RHI_GPUTotalTime, STATGROUP_D3D12RHI);
@@ -101,7 +111,18 @@ void FD3D12DynamicRHI::InitializeSubmissionPipe()
 #endif
 
 #if D3D12_USE_SUBMISSION_THREAD
-	SubmissionThread = new FD3D12Thread(TEXT("RHISubmissionThread"), TPri_Highest, this, &FD3D12DynamicRHI::ProcessSubmissionQueue);
+
+	bool bUseSubmissionThread = false;
+	switch (CVarRHIUseSubmissionThread.GetValueOnAnyThread())
+	{
+	case 1: bUseSubmissionThread = FRHIGPUMask::All().HasSingleIndex(); break;
+	case 2: bUseSubmissionThread = true; break;
+	}
+
+	if (bUseSubmissionThread)
+	{
+		SubmissionThread = new FD3D12Thread(TEXT("RHISubmissionThread"), TPri_Highest, this, &FD3D12DynamicRHI::ProcessSubmissionQueue);
+	}
 #endif
 
 	FlushTiming(true);
@@ -273,22 +294,21 @@ void FD3D12DynamicRHI::SubmitPayloads(TArrayView<FD3D12Payload*> Payloads)
 		Payload->Queue.PendingSubmission.Enqueue(Payload);
 	}
 
-#if D3D12_USE_SUBMISSION_THREAD
-
-	SubmissionThread->Kick();
-
-#else
-
-	// Since we're processing directly on the calling thread, we need to take a scope lock.
-	// Multiple engine threads might be calling Submit().
+	if (SubmissionThread)
 	{
-		FScopeLock Lock(&SubmissionCS);
-
-		// Process the submission queue until no further progress is being made.
-		while (EnumHasAnyFlags(ProcessSubmissionQueue().Status, EQueueStatus::Processed)) {}
+		SubmissionThread->Kick();
 	}
+	else
+	{
+		// Since we're processing directly on the calling thread, we need to take a scope lock.
+		// Multiple engine threads might be calling Submit().
+		{
+			FScopeLock Lock(&SubmissionCS);
 
-#endif
+			// Process the submission queue until no further progress is being made.
+			while (EnumHasAnyFlags(ProcessSubmissionQueue().Status, EQueueStatus::Processed)) {}
+		}
+	}
 
 	// Use this opportunity to pump the interrupt queue
 	ProcessInterruptQueueUntil(nullptr);
