Allow precache PSO to be removed from LowLevelGraphicsPipelineStateCache

LowLevelGraphicsPipelineStateCache was always keeping FD3D12PipelineState objects alive, so we need to be able to remove some of them when we destroy a FD3D12GraphicsPipelineState
FD3D12PipelineState can be shared by multiple FD3D12GraphicsPipelineState, so we need to refcount them.
If a FD3D12GraphicsPipelineState was the last one to reference a FD3D12PipelineState, and if it was used for precaching purposes, we try to remove it from the map. That removal is not guaranteed though, because another FD3D12GraphicsPipelineState might have requested the same FD3D12PipelineState: we need to acquire the lock for LowLevelGraphicsPipelineStateCache to prevent any other access first, and check the FD3D12PipelineState's refcount again before actually removing it from the map.

diff --git a/Engine/Source/Runtime/D3D12RHI/Private/D3D12PipelineState.h b/Engine/Source/Runtime/D3D12RHI/Private/D3D12PipelineState.h
--- a/Engine/Source/Runtime/D3D12RHI/Private/D3D12PipelineState.h
+++ b/Engine/Source/Runtime/D3D12RHI/Private/D3D12PipelineState.h
@@ -292,7 +292,7 @@ struct FD3D12PipelineStateWorker : public FD3D12AdapterChild, public FNonAbandon
 	TRefCountPtr<ID3D12PipelineState> PSO;
 };
 
-struct FD3D12PipelineState : public FD3D12AdapterChild, public FD3D12MultiNodeGPUObject, public FNoncopyable
+struct FD3D12PipelineState : public FD3D12AdapterChild, public FD3D12MultiNodeGPUObject, public FNoncopyable, public FRefCountBase
 {
 public:
 	explicit FD3D12PipelineState(FD3D12Adapter* Parent);
@@ -337,6 +337,8 @@ public:
 
 	FD3D12PipelineState& operator=(const FD3D12PipelineState& other) = delete;
 
+    static bool UsePSORefCounting();
+
 private:
 	ID3D12PipelineState* InternalGetPipelineState();
 
@@ -496,6 +498,8 @@ protected:
 	FD3D12GraphicsPipelineState* CreateAndAdd(const FGraphicsPipelineStateInitializer& Initializer, const FD3D12RootSignature* RootSignature, const FD3D12LowLevelGraphicsPipelineStateDesc& LowLevelDesc);
 #endif
 public:
+	void RemoveFromLowLevelCache(FD3D12PipelineState* PipelineState, const FGraphicsPipelineStateInitializer& PipelineStateInitializer, const FD3D12RootSignature* RootSignature);
+
 #if D3D12RHI_USE_HIGH_LEVEL_PSO_CACHE
 	FD3D12GraphicsPipelineState* FindInRuntimeCache(const FGraphicsPipelineStateInitializer& Initializer, uint32& OutHash);
 	FD3D12GraphicsPipelineState* FindInLoadedCache(const FGraphicsPipelineStateInitializer& Initializer, uint32 InitializerHash, const FD3D12RootSignature* RootSignature, FD3D12LowLevelGraphicsPipelineStateDesc& OutLowLevelDesc);
diff --git a/Engine/Source/Runtime/D3D12RHI/Private/D3D12Pipelinestate.cpp b/Engine/Source/Runtime/D3D12RHI/Private/D3D12Pipelinestate.cpp
--- a/Engine/Source/Runtime/D3D12RHI/Private/D3D12Pipelinestate.cpp
+++ b/Engine/Source/Runtime/D3D12RHI/Private/D3D12Pipelinestate.cpp
@@ -252,6 +252,7 @@ FD3D12PipelineState::FD3D12PipelineState(FD3D12Adapter* Parent)
 
 FD3D12PipelineState::~FD3D12PipelineState()
 {
+	check(!UsePSORefCounting() || GetRefCount() == 0);
 	if (Worker)
 	{
 		Worker->EnsureCompletion(true);
@@ -262,6 +263,16 @@ FD3D12PipelineState::~FD3D12PipelineState()
 	DEC_DWORD_STAT(STAT_D3D12NumPSOs);
 }
 
+bool FD3D12PipelineState::UsePSORefCounting()
+{
+#if D3D12_USE_DERIVED_PSO
+	return false;
+#else
+	static const auto CVarPSOPrecaching = IConsoleManager::Get().FindConsoleVariable(TEXT("r.PSOPrecaching"));
+	return CVarPSOPrecaching && (CVarPSOPrecaching->GetInt() != 0);
+#endif
+}
+
 ID3D12PipelineState* FD3D12PipelineState::InternalGetPipelineState()
 {
 	// Only one thread should finalize any tasks and set the cached pipeline state.
@@ -347,6 +358,21 @@ FD3D12GraphicsPipelineState::~FD3D12GraphicsPipelineState()
 #if D3D12_USE_DERIVED_PSO
 	delete PipelineState;
 	PipelineState = nullptr;
+#else
+	if (FD3D12PipelineState::UsePSORefCounting() && PipelineState != nullptr)
+	{
+		uint32 RefCount = PipelineState->Release();
+		check(RefCount > 0);
+        // precache PSO are here to avoid hitches at runtime when we want to create one that is actually used. We don't need to keep them
+		// around as this can add up to a lot of system memory
+		if (PipelineStateInitializer.bPSOPrecache && RefCount == 1) 
+		{
+			FD3D12DynamicRHI* D3D12RHI = FD3D12DynamicRHI::GetD3DRHI();
+			FD3D12PipelineStateCache& PSOCache = D3D12RHI->GetAdapter().GetPSOCache();
+			// NB: it's possible that this remove will not do anything because another thread might have requested the same PSO since we entered the if
+			PSOCache.RemoveFromLowLevelCache(PipelineState, PipelineStateInitializer, RootSignature);
+		}
+	}
 #endif
 
 	// release bound RHI resources
@@ -399,7 +425,17 @@ void FD3D12PipelineStateCacheBase::CleanupPipelineStateCaches()
 		for (auto Iter = LowLevelGraphicsPipelineStateCache.CreateConstIterator(); Iter; ++Iter)
 		{
 			const FD3D12PipelineState* const PipelineState = Iter.Value();
-			delete PipelineState;
+			if (FD3D12PipelineState::UsePSORefCounting())
+			{
+				if (PipelineState)
+				{
+					PipelineState->Release();
+				}
+			}
+			else
+			{
+				delete PipelineState;
+			}
 		}
 		LowLevelGraphicsPipelineStateCache.Empty();
 	}
@@ -446,6 +482,10 @@ FD3D12PipelineState* FD3D12PipelineStateCacheBase::FindInLowLevelCache(const FD3
 		if (Found)
 		{
 			INC_DWORD_STAT(STAT_PSOGraphicsLowlevelCacheHit);
+			if (FD3D12PipelineState::UsePSORefCounting())
+			{
+				verify((*Found)->AddRef() >= 2);
+			}
 			return *Found;
 		}
 	}
@@ -466,6 +506,31 @@ FD3D12PipelineState* FD3D12PipelineStateCacheBase::CreateAndAddToLowLevelCache(c
 	return PipelineState;
 }
 
+void FD3D12PipelineStateCacheBase::RemoveFromLowLevelCache(FD3D12PipelineState* PipelineState, const FGraphicsPipelineStateInitializer& PipelineStateInitializer, const FD3D12RootSignature* RootSignature)
+{
+    if (!FD3D12PipelineState::UsePSORefCounting())
+    {
+        checkNoEntry();
+        return;
+    }
+
+	FRWScopeLock Lock(LowLevelGraphicsPipelineStateCacheMutex, FRWScopeLockType::SLT_Write);
+	// By the time we reach this function, it's possible another thread made a request for the same PSO
+	if (PipelineState->Release() == 0)
+	{
+		FD3D12LowLevelGraphicsPipelineStateDesc LowLevelDesc = GetLowLevelGraphicsPipelineStateDesc(PipelineStateInitializer, RootSignature);
+		LowLevelDesc.Desc.NodeMask = FRHIGPUMask::All().GetNative();
+		LowLevelDesc.CombinedHash = FD3D12PipelineStateCacheBase::HashPSODesc(LowLevelDesc);
+		int32 ElementsRemoved = LowLevelGraphicsPipelineStateCache.Remove(LowLevelDesc);
+		ensure(ElementsRemoved == 1);
+	}
+	else
+	{
+		// If another thread requested the pipeline state, we need to restore the refcount we just decremented
+		verify(PipelineState->AddRef() >= 2);
+	}
+}
+
 void FD3D12PipelineStateCacheBase::AddToLowLevelCache(const FD3D12LowLevelGraphicsPipelineStateDesc& Desc, FD3D12PipelineState** OutPipelineState, const FPostCreateGraphicCallback& PostCreateCallback)
 {
 	check(Desc.CombinedHash != 0);
@@ -479,6 +544,10 @@ void FD3D12PipelineStateCacheBase::AddToLowLevelCache(const FD3D12LowLevelGraphi
 		{
 			// This desc already exists.
 			*OutPipelineState = *PipelineState;
+			if (FD3D12PipelineState::UsePSORefCounting())
+			{
+				(*OutPipelineState)->AddRef();
+			}
 			return;
 		}
 
@@ -488,7 +557,20 @@ void FD3D12PipelineStateCacheBase::AddToLowLevelCache(const FD3D12LowLevelGraphi
 		FD3D12PipelineState* NewPipelineState = new FD3D12PipelineState(GetParentAdapter());
 		LowLevelGraphicsPipelineStateCache.Add(Desc, NewPipelineState);
 
+		// This AddRef is for the low level cache
+		if (FD3D12PipelineState::UsePSORefCounting())
+		{
+			NewPipelineState->AddRef();
+		}
+
 		*OutPipelineState = NewPipelineState;
+
+		// This AddRef is for the FD3D12GraphicsPipelineState requesting it
+		if (FD3D12PipelineState::UsePSORefCounting())
+		{
+			(*OutPipelineState)->AddRef();
+		}
+
 	}
 
 	// Create the underlying PSO and then perform any other additional tasks like cleaning up/adding to caches, etc.
diff --git a/Engine/Source/Runtime/RHI/Private/PipelineStateCache.cpp b/Engine/Source/Runtime/RHI/Private/PipelineStateCache.cpp
--- a/Engine/Source/Runtime/RHI/Private/PipelineStateCache.cpp
+++ b/Engine/Source/Runtime/RHI/Private/PipelineStateCache.cpp
@@ -2277,8 +2277,11 @@ FGraphEventRef PipelineStateCache::PrecacheGraphicsPipelineState(const FGraphics
 	// Mark as precache so it will try and use the background thread pool if available
 	bool bPSOPrecache = true;
 
+	FGraphicsPipelineStateInitializer InitializerCopy(Initializer);
+	InitializerCopy.bPSOPrecache = true;
+
 	// Start the precache task
-	return CreateGraphicsPipelineState(Initializer, EPSOPrecacheResult::Active, bDoAsyncCompile, bPSOPrecache, NewGraphicsPipelineState);
+	return CreateGraphicsPipelineState(InitializerCopy, EPSOPrecacheResult::Active, bDoAsyncCompile, bPSOPrecache, NewGraphicsPipelineState);
 }
 
 EPSOPrecacheResult PipelineStateCache::CheckPipelineStateInCache(const FGraphicsPipelineStateInitializer& PipelineStateInitializer)
diff --git a/Engine/Source/Runtime/RHI/Public/RHIResources.h b/Engine/Source/Runtime/RHI/Public/RHIResources.h
--- a/Engine/Source/Runtime/RHI/Public/RHIResources.h
+++ b/Engine/Source/Runtime/RHI/Public/RHIResources.h
@@ -3438,7 +3438,8 @@ public:
 	{
 		struct
 		{
-			uint16					Reserved			: 15;
+			uint16					Reserved			: 14;
+			uint16					bPSOPrecache			: 1;
 			uint16					bFromPSOFileCache	: 1;
 		};
 		uint16						Flags;
