Add the option to prefer timestamps to compute GPUFrameTime

diff --git a/Engine/Source/Runtime/D3D12RHI/Private/D3D12Submission.cpp b/Engine/Source/Runtime/D3D12RHI/Private/D3D12Submission.cpp
--- a/Engine/Source/Runtime/D3D12RHI/Private/D3D12Submission.cpp
+++ b/Engine/Source/Runtime/D3D12RHI/Private/D3D12Submission.cpp
@@ -1096,6 +1096,16 @@ FD3D12PayloadBase::~FD3D12PayloadBase()
 	}
 }
 
+#ifndef D3D12_PREFER_QUERIES_FOR_GPU_TIME
+#define D3D12_PREFER_QUERIES_FOR_GPU_TIME 0
+#endif
+
+static TAutoConsoleVariable<int32> CVarGPUTimeFromTimestamps(
+	TEXT("r.D3D12.GPUTimeFromTimestamps"),
+	D3D12_PREFER_QUERIES_FOR_GPU_TIME,
+	TEXT("Prefer timestamps instead of GetHardwareGPUFrameTime to compute GPU frame time"),
+	ECVF_RenderThreadSafe);
+
 void FD3D12DynamicRHI::ProcessTimestamps(TIndirectArray<FD3D12Timing>& Timing)
 {
 	// The total number of cycles where at least one GPU pipe was busy during the frame.
@@ -1146,7 +1156,7 @@ void FD3D12DynamicRHI::ProcessTimestamps(TIndirectArray<FD3D12Timing>& Timing)
 	}
 
 	check(BusyPipes == 0);
-
+	
 	// @todo mgpu - how to handle multiple devices / queues with potentially different timestamp frequencies?
 	FD3D12Device* Device = GetAdapter().GetDevice(0);
 	double Frequency = Device->GetTimestampFrequency(ED3D12QueueType::Direct);
@@ -1156,9 +1166,9 @@ void FD3D12DynamicRHI::ProcessTimestamps(TIndirectArray<FD3D12Timing>& Timing)
 
 	// Update the global GPU frame time stats
 	SET_CYCLE_COUNTER(STAT_RHI_GPUTotalTime, FPlatformMath::TruncToInt(double(UnionBusyCycles) * Scale64));
-
+	 
 	double HardwareGPUTime = 0.0;
-	if (GetHardwareGPUFrameTime(HardwareGPUTime))
+	if (GetHardwareGPUFrameTime(HardwareGPUTime) && CVarGPUTimeFromTimestamps.GetValueOnAnyThread() == 0)
 	{
 		SET_CYCLE_COUNTER(STAT_RHI_GPUTotalTimeHW, HardwareGPUTime);
 		GGPUFrameTime = HardwareGPUTime;
