Remove check on valid GPU address during free of tracked allocation because GPU address can also be empty when r.D3D12.UseCommittedTexturesNV is enabled - then heap and placed resource allocation workaround is not used

diff --git a/Engine/Source/Runtime/D3D12RHI/Private/D3D12Adapter.cpp b/Engine/Source/Runtime/D3D12RHI/Private/D3D12Adapter.cpp
--- a/Engine/Source/Runtime/D3D12RHI/Private/D3D12Adapter.cpp
+++ b/Engine/Source/Runtime/D3D12RHI/Private/D3D12Adapter.cpp
@@ -1630,7 +1630,7 @@ void FD3D12Adapter::ReleaseTrackedAllocationData(FD3D12ResourceLocation* InAlloc
 	FScopeLock Lock(&TrackedAllocationDataCS);
 
 	D3D12_GPU_VIRTUAL_ADDRESS GPUAddress = InAllocation->GetGPUVirtualAddress();
-	if (GPUAddress != 0 || IsTrackingAllAllocations())
+	if (GPUAddress != 0)
 	{
 		FReleasedAllocationData ReleasedData;
 		ReleasedData.GPUVirtualAddress = GPUAddress;
@@ -1641,8 +1641,6 @@ void FD3D12Adapter::ReleaseTrackedAllocationData(FD3D12ResourceLocation* InAlloc
 		ReleasedData.bDefragFree = bDefragFree;
 		ReleasedData.bBackBuffer = InAllocation->GetResource()->IsBackBuffer();
 		ReleasedData.bTransient = InAllocation->IsTransient();
-		// Only the backbuffer doesn't have a valid gpu virtual address
-		check(ReleasedData.GPUVirtualAddress != 0 || ReleasedData.bBackBuffer);
 		ReleasedAllocationData.Add(ReleasedData);
 	}
 
