Adds CSV stat for GPU memory on D3D12

diff --git a/Engine/Source/Runtime/D3D12RHI/Private/D3D12CommandContext.cpp b/Engine/Source/Runtime/D3D12RHI/Private/D3D12CommandContext.cpp
--- a/Engine/Source/Runtime/D3D12RHI/Private/D3D12CommandContext.cpp
+++ b/Engine/Source/Runtime/D3D12RHI/Private/D3D12CommandContext.cpp
@@ -687,14 +687,30 @@ void FD3D12CommandContextBase::RHIEndFrame()
 	FD3D12DynamicRHI::GetD3DRHI()->ProcessInterruptQueueUntil(nullptr);
 }
 
+#if PLATFORM_WINDOWS && CSV_PROFILER && !UE_BUILD_SHIPPING
+	CSV_DEFINE_CATEGORY(GPUMem, true);
+#endif
+DEFINE_STAT(STAT_D3D12UpdateVideoMemoryStats);
+
 void FD3D12CommandContextBase::UpdateMemoryStats()
 {
-#if PLATFORM_WINDOWS && STATS
+#if PLATFORM_WINDOWS && (STATS || CSV_PROFILER) && !UE_BUILD_SHIPPING
+	SCOPE_CYCLE_COUNTER(STAT_D3D12UpdateVideoMemoryStats);
 	// Refresh captured memory info.
 	ParentAdapter->UpdateMemoryInfo();
 
 	const FD3D12MemoryInfo& MemoryInfo = ParentAdapter->GetMemoryInfo();
 
+#if CSV_PROFILER
+	{
+		CSV_CUSTOM_STAT(GPUMem, Total, double(MemoryInfo.LocalMemoryInfo.Budget), ECsvCustomStatOp::Set);
+		CSV_CUSTOM_STAT(GPUMem, Used, double(MemoryInfo.LocalMemoryInfo.CurrentUsage), ECsvCustomStatOp::Set);
+		CSV_CUSTOM_STAT(GPUMem, Available, double(MemoryInfo.AvailableLocalMemory), ECsvCustomStatOp::Set);
+		CSV_CUSTOM_STAT(GPUMem, Demoted, double(MemoryInfo.DemotedLocalMemory), ECsvCustomStatOp::Set);
+	}
+#endif // CSV_PROFILER
+
+#if STATS
 	SET_MEMORY_STAT(STAT_D3D12UsedVideoMemory, MemoryInfo.LocalMemoryInfo.CurrentUsage);
 	SET_MEMORY_STAT(STAT_D3D12UsedSystemMemory, MemoryInfo.NonLocalMemoryInfo.CurrentUsage);
 	SET_MEMORY_STAT(STAT_D3D12AvailableVideoMemory, MemoryInfo.AvailableLocalMemory);
@@ -718,7 +734,8 @@ void FD3D12CommandContextBase::UpdateMemoryStats()
 		Device->GetDefaultBufferAllocator().UpdateMemoryStats();
 		ParentAdapter->GetUploadHeapAllocator(GPUIndex).UpdateMemoryStats();
 	}
-#endif
+#endif // STATS
+#endif // PLATFORM_WINDOWS && (STATS || CSV_PROFILER)
 }
 
 void FD3D12CommandContext::RHIBeginScene()
diff --git a/Engine/Source/Runtime/D3D12RHI/Private/D3D12Stats.h b/Engine/Source/Runtime/D3D12RHI/Private/D3D12Stats.h
--- a/Engine/Source/Runtime/D3D12RHI/Private/D3D12Stats.h
+++ b/Engine/Source/Runtime/D3D12RHI/Private/D3D12Stats.h
@@ -67,6 +67,7 @@ DECLARE_MEMORY_STAT_EXTERN(TEXT("Available Video Memory"), STAT_D3D12AvailableVi
 DECLARE_MEMORY_STAT_EXTERN(TEXT("Demoted Video Memory"), STAT_D3D12DemotedVideoMemory, STATGROUP_D3D12Memory, );
 DECLARE_MEMORY_STAT_EXTERN(TEXT("Used Video Memory"), STAT_D3D12UsedVideoMemory, STATGROUP_D3D12Memory, );
 DECLARE_MEMORY_STAT_EXTERN(TEXT("Used System Memory"), STAT_D3D12UsedSystemMemory, STATGROUP_D3D12Memory, );
+DECLARE_CYCLE_STAT_EXTERN(TEXT("Video Memory stats update time"), STAT_D3D12UpdateVideoMemoryStats, STATGROUP_D3D12RHI, );
 
 DECLARE_MEMORY_STAT_EXTERN(TEXT("TOTAL"), STAT_D3D12MemoryCurrentTotal, STATGROUP_D3D12Resources, );
 DECLARE_MEMORY_STAT_EXTERN(TEXT("Render Targets"), STAT_D3D12RenderTargets, STATGROUP_D3D12Resources, );
