Unify code that binds shader resources from uniform buffers across platform RHIs
 - All platforms were extending FBaseShaderResourceTable to add a TArray<uint32> TextureMap member, along with copies of the same RHI code to bind these resources.
 - Moved the TextureMap member into the base shader table type, and removed all platform versions. They are all serializer-compatible with the original code, so no shader DDC keys need to be bumped.
 - Also moved the "ResourceTable" array of RHI resources into the base FRHIUniformBuffer class.
 - Added a SetResourcesFromTables function in RHICoreShader.h

Removed GetD3D11TextureFromRHITexture in D3D11RHI, replaced with a ResourceCast.
 - The ResourceCast function for FRHITexture types now always follows the texture reference pointer via the vtable call, like other platforms.

diff --git a/Engine/Plugins/Media/PixelStreaming/Source/PixelStreaming/Private/EncoderFrameFactory.cpp b/Engine/Plugins/Media/PixelStreaming/Source/PixelStreaming/Private/EncoderFrameFactory.cpp
--- a/Engine/Plugins/Media/PixelStreaming/Source/PixelStreaming/Private/EncoderFrameFactory.cpp
+++ b/Engine/Plugins/Media/PixelStreaming/Source/PixelStreaming/Private/EncoderFrameFactory.cpp
@@ -14,6 +14,7 @@
 	#include "ShaderCore.h"
 	#include "D3D11State.h"
 	#include "D3D11Resources.h"
+	#include "ID3D11DynamicRHI.h"
 	#include "ID3D12DynamicRHI.h"
 	#include "Windows/AllowWindowsPlatformTypes.h"
 THIRD_PARTY_INCLUDES_START
@@ -376,10 +377,8 @@ namespace UE::PixelStreaming
 #if PLATFORM_WINDOWS
 	void FEncoderFrameFactory::SetTextureCUDAD3D11(TSharedPtr<AVEncoder::FVideoEncoderInputFrame> InputFrame, const FTextureRHIRef& Texture)
 	{
-		FD3D11Texture* D3D11Texture = GetD3D11TextureFromRHITexture(Texture);
-		unsigned long long TextureMemorySize = D3D11Texture->GetMemorySize();
-
-		ID3D11Texture2D* D3D11NativeTexture = static_cast<ID3D11Texture2D*>(D3D11Texture->GetResource());
+		ID3D11Resource* D3D11NativeTexture = GetID3D11DynamicRHI()->RHIGetResource(Texture);
+		unsigned long long TextureMemorySize = GetID3D11DynamicRHI()->RHIGetResourceMemorySize(Texture);
 
 		TRefCountPtr<IDXGIResource> DXGIResource;
 		HRESULT QueryResult = D3D11NativeTexture->QueryInterface(IID_PPV_ARGS(DXGIResource.GetInitReference()));
diff --git a/Engine/Source/Developer/Windows/ShaderFormatD3D/Private/D3DShaderCompiler.cpp b/Engine/Source/Developer/Windows/ShaderFormatD3D/Private/D3DShaderCompiler.cpp
--- a/Engine/Source/Developer/Windows/ShaderFormatD3D/Private/D3DShaderCompiler.cpp
+++ b/Engine/Source/Developer/Windows/ShaderFormatD3D/Private/D3DShaderCompiler.cpp
@@ -5,7 +5,6 @@
 #include "ShaderCompilerCommon.h"
 #include "ShaderMinifier.h"
 #include "ShaderParameterParser.h"
-#include "D3D11ShaderResources.h"
 #include "D3D12RHI.h"
 #include "Misc/Paths.h"
 #include "Misc/FileHelper.h"
diff --git a/Engine/Source/Developer/Windows/ShaderFormatD3D/Private/D3DShaderCompiler.inl b/Engine/Source/Developer/Windows/ShaderFormatD3D/Private/D3DShaderCompiler.inl
--- a/Engine/Source/Developer/Windows/ShaderFormatD3D/Private/D3DShaderCompiler.inl
+++ b/Engine/Source/Developer/Windows/ShaderFormatD3D/Private/D3DShaderCompiler.inl
@@ -521,7 +521,7 @@ inline void GenerateFinalOutput(TRefCountPtr<TBlob>& CompressedData,
 	TFunction<void(FShaderCode&)> AddOptionalDataCallback)
 {
 	// Build the SRT for this shader.
-	FD3D11ShaderResourceTable SRT;
+	FShaderResourceTable SRT;
 
 	TArray<uint8> UniformBufferNameBytes;
 
diff --git a/Engine/Source/Developer/Windows/ShaderFormatD3D/Private/D3DShaderCompilerDXC.cpp b/Engine/Source/Developer/Windows/ShaderFormatD3D/Private/D3DShaderCompilerDXC.cpp
--- a/Engine/Source/Developer/Windows/ShaderFormatD3D/Private/D3DShaderCompilerDXC.cpp
+++ b/Engine/Source/Developer/Windows/ShaderFormatD3D/Private/D3DShaderCompilerDXC.cpp
@@ -4,7 +4,6 @@
 #include "ShaderPreprocessor.h"
 #include "ShaderCompilerCommon.h"
 #include "ShaderParameterParser.h"
-#include "D3D11ShaderResources.h"
 #include "D3D12RHI.h"
 #include "Misc/Paths.h"
 #include "Misc/FileHelper.h"
diff --git a/Engine/Source/Runtime/Apple/MetalRHI/Private/MetalRenderPass.cpp b/Engine/Source/Runtime/Apple/MetalRHI/Private/MetalRenderPass.cpp
--- a/Engine/Source/Runtime/Apple/MetalRHI/Private/MetalRenderPass.cpp
+++ b/Engine/Source/Runtime/Apple/MetalRHI/Private/MetalRenderPass.cpp
@@ -505,13 +505,9 @@ void FMetalRenderPass::DrawIndexedPrimitive(FMetalBuffer const& IndexBuffer, uin
 				
 				if(ClampedNumInstances < NumInstances)
 				{
-					FString ShaderName = TEXT("Unknown");
-#if !(UE_BUILD_SHIPPING || UE_BUILD_TEST)
-					ShaderName = PipelineState->PixelShader->ShaderName;
-#endif
 					// Setting NumInstances to ClampedNumInstances would fix any visual rendering bugs resulting from this bad call but these draw calls are wrong - don't hide the issue
 					UE_LOG(LogMetal, Error, TEXT("Metal DrawIndexedPrimitive requested to draw %d Instances but vertex stream only has %d instance data available. ShaderName: %s, Deficient Attribute Index: %u"), NumInstances, ClampedNumInstances,
-						   *ShaderName, VertexElem.AttributeIndex);
+						   *PipelineState->PixelShader->GetShaderName(), VertexElem.AttributeIndex);
 				}
 			}
 		}
diff --git a/Engine/Source/Runtime/Apple/MetalRHI/Private/MetalStateCache.cpp b/Engine/Source/Runtime/Apple/MetalRHI/Private/MetalStateCache.cpp
--- a/Engine/Source/Runtime/Apple/MetalRHI/Private/MetalStateCache.cpp
+++ b/Engine/Source/Runtime/Apple/MetalRHI/Private/MetalStateCache.cpp
@@ -1178,11 +1178,6 @@ void FMetalStateCache::BindUniformBuffer(EMetalShaderStages const Freq, uint32 c
 	}
 }
 
-void FMetalStateCache::SetDirtyUniformBuffers(EMetalShaderStages const Freq, uint32 const Dirty)
-{
-	DirtyUniformBuffers[Freq] = Dirty;
-}
-
 void FMetalStateCache::SetVisibilityResultMode(mtlpp::VisibilityResultMode const Mode, NSUInteger const Offset)
 {
 	if (VisibilityMode != Mode || VisibilityOffset != Offset)
@@ -1419,14 +1414,13 @@ void FMetalStateCache::SetShaderSamplerState(EMetalShaderStages const Frequency,
 	}
 }
 
-void FMetalStateCache::SetResource(uint32 ShaderStage, uint32 BindIndex, FRHITexture* RESTRICT TextureRHI, float CurrentTime)
+void FMetalStateCache::SetResource(uint32 ShaderStage, uint32 BindIndex, FRHITexture* RESTRICT TextureRHI)
 {
 	FMetalSurface* Surface = GetMetalSurfaceFromRHITexture(TextureRHI);
 	ns::AutoReleased<FMetalTexture> Texture;
 	mtlpp::ResourceUsage Usage = (mtlpp::ResourceUsage)0;
 	if (Surface != nullptr)
 	{
-		TextureRHI->SetLastRenderTime(CurrentTime);
 		Texture = Surface->Texture;
 		Usage = mtlpp::ResourceUsage(mtlpp::ResourceUsage::Read|mtlpp::ResourceUsage::Sample);
 	}
@@ -1566,7 +1560,7 @@ void FMetalStateCache::SetShaderUnorderedAccessView(EMetalShaderStages ShaderSta
 	}
 }
 
-void FMetalStateCache::SetResource(uint32 ShaderStage, uint32 BindIndex, FMetalShaderResourceView* RESTRICT SRV, float CurrentTime)
+void FMetalStateCache::SetResource(uint32 ShaderStage, uint32 BindIndex, FMetalShaderResourceView* RESTRICT SRV)
 {
 	switch (ShaderStage)
 	{
@@ -1588,7 +1582,7 @@ void FMetalStateCache::SetResource(uint32 ShaderStage, uint32 BindIndex, FMetalS
 	}
 }
 
-void FMetalStateCache::SetResource(uint32 ShaderStage, uint32 BindIndex, FMetalSamplerState* RESTRICT SamplerState, float CurrentTime)
+void FMetalStateCache::SetResource(uint32 ShaderStage, uint32 BindIndex, FMetalSamplerState* RESTRICT SamplerState)
 {
 	check(SamplerState->State);
 	switch (ShaderStage)
@@ -1611,7 +1605,7 @@ void FMetalStateCache::SetResource(uint32 ShaderStage, uint32 BindIndex, FMetalS
 	}
 }
 
-void FMetalStateCache::SetResource(uint32 ShaderStage, uint32 BindIndex, FMetalUnorderedAccessView* RESTRICT UAV, float CurrentTime)
+void FMetalStateCache::SetResource(uint32 ShaderStage, uint32 BindIndex, FMetalUnorderedAccessView* RESTRICT UAV)
 {
 	switch (ShaderStage)
 	{
@@ -1633,42 +1627,13 @@ void FMetalStateCache::SetResource(uint32 ShaderStage, uint32 BindIndex, FMetalU
 	}
 }
 
-
-template <typename MetalResourceType>
-inline int32 FMetalStateCache::SetShaderResourcesFromBuffer(uint32 ShaderStage, FMetalUniformBuffer* RESTRICT Buffer, const uint32* RESTRICT ResourceMap, int32 BufferIndex, float CurrentTime)
-{
-	const TRefCountPtr<FRHIResource>* RESTRICT Resources = Buffer->ResourceTable.GetData();
-	int32 NumSetCalls = 0;
-	uint32 BufferOffset = ResourceMap[BufferIndex];
-	if (BufferOffset > 0)
-	{
-		const uint32* RESTRICT ResourceInfos = &ResourceMap[BufferOffset];
-		uint32 ResourceInfo = *ResourceInfos++;
-		do
-		{
-			checkSlow(FRHIResourceTableEntry::GetUniformBufferIndex(ResourceInfo) == BufferIndex);
-			const uint16 ResourceIndex = FRHIResourceTableEntry::GetResourceIndex(ResourceInfo);
-			const uint8 BindIndex = FRHIResourceTableEntry::GetBindIndex(ResourceInfo);
-			
-			MetalResourceType* ResourcePtr = (MetalResourceType*)Resources[ResourceIndex].GetReference();
-			
-			// todo: could coalesce adjacent bound resources.
-			SetResource(ShaderStage, BindIndex, ResourcePtr, CurrentTime);
-			
-			NumSetCalls++;
-			ResourceInfo = *ResourceInfos++;
-		} while (FRHIResourceTableEntry::GetUniformBufferIndex(ResourceInfo) == BufferIndex);
-	}
-	return NumSetCalls;
-}
-
 template <class ShaderType>
 void FMetalStateCache::SetResourcesFromTables(ShaderType Shader, uint32 ShaderStage)
 {
 	checkSlow(Shader);
 	
 	EMetalShaderStages Frequency;
-	switch(ShaderStage)
+	switch (ShaderStage)
 	{
 		case CrossCompiler::SHADER_STAGE_VERTEX:
 			Frequency = EMetalShaderStages::Vertex;
@@ -1685,30 +1650,35 @@ void FMetalStateCache::SetResourcesFromTables(ShaderType Shader, uint32 ShaderSt
 			break;
 	}
 
-	float CurrentTime = FPlatformTime::Seconds();
-
-	// Mask the dirty bits by those buffers from which the shader has bound resources.
-	uint32 DirtyBits = Shader->Bindings.ShaderResourceTable.ResourceTableBits & GetDirtyUniformBuffers(Frequency);
-	while (DirtyBits)
+	if (!FMetalCommandQueue::SupportsFeature(EMetalFeaturesIABs))
 	{
-		// Scan for the lowest set bit, compute its index, clear it in the set of dirty bits.
-		const uint32 LowestBitMask = (DirtyBits)& (-(int32)DirtyBits);
-		const int32 BufferIndex = FMath::FloorLog2(LowestBitMask); // todo: This has a branch on zero, we know it could never be zero...
-		DirtyBits ^= LowestBitMask;
-		FMetalUniformBuffer* Buffer = (FMetalUniformBuffer*)GetBoundUniformBuffers(Frequency)[BufferIndex];
-		if (Buffer && !FMetalCommandQueue::SupportsFeature(EMetalFeaturesIABs))
+		struct FUniformResourceBinder
 		{
-			check(BufferIndex < Shader->Bindings.ShaderResourceTable.ResourceTableLayoutHashes.Num());
-			check(Buffer->GetLayout().GetHash() == Shader->Bindings.ShaderResourceTable.ResourceTableLayoutHashes[BufferIndex]);
-			
-			// todo: could make this two pass: gather then set
-			SetShaderResourcesFromBuffer<FRHITexture>(ShaderStage, Buffer, Shader->Bindings.ShaderResourceTable.TextureMap.GetData(), BufferIndex, CurrentTime);
-			SetShaderResourcesFromBuffer<FMetalShaderResourceView>(ShaderStage, Buffer, Shader->Bindings.ShaderResourceTable.ShaderResourceViewMap.GetData(), BufferIndex, CurrentTime);
-			SetShaderResourcesFromBuffer<FMetalSamplerState>(ShaderStage, Buffer, Shader->Bindings.ShaderResourceTable.SamplerMap.GetData(), BufferIndex, CurrentTime);
-			SetShaderResourcesFromBuffer<FMetalUnorderedAccessView>(ShaderStage, Buffer, Shader->Bindings.ShaderResourceTable.UnorderedAccessViewMap.GetData(), BufferIndex, CurrentTime);
-		}
+			FMetalStateCache& StateCache;
+			uint32 ShaderStage;
+
+			void SetUAV(FRHIUnorderedAccessView*  UAV, uint8 Index) { StateCache.SetResource(ShaderStage, Index, static_cast<FMetalUnorderedAccessView*>(UAV)); }
+			void SetSRV(FRHIShaderResourceView*   SRV, uint8 Index) { StateCache.SetResource(ShaderStage, Index, static_cast<FMetalShaderResourceView*>(SRV)); }
+
+			void SetTexture(FRHITexture*      Texture, uint8 Index) { StateCache.SetResource(ShaderStage, Index, Texture); }
+			void SetSampler(FRHISamplerState* Sampler, uint8 Index) { StateCache.SetResource(ShaderStage, Index, static_cast<FMetalSamplerState*>(Sampler)); }
+		};
+
+		UE::RHICore::SetResourcesFromTables(
+			  FUniformResourceBinder { *this, ShaderStage }
+			, *Shader
+			, Shader->Bindings.ShaderResourceTable
+			, DirtyUniformBuffers[Frequency]
+			, BoundUniformBuffers[Frequency]
+#if ENABLE_RHI_VALIDATION
+			, nullptr /*Tracker*/ // @todo: the current structure of the Metal RHI prevents easily passing the RHI validation layer tracker here
+#endif
+		);
+	}
+	else
+	{
+		DirtyUniformBuffers[Frequency] = 0;
 	}
-	SetDirtyUniformBuffers(Frequency, 0);
 }
 
 void FMetalStateCache::CommitRenderResources(FMetalCommandEncoder* Raster)
diff --git a/Engine/Source/Runtime/Apple/MetalRHI/Private/MetalStateCache.h b/Engine/Source/Runtime/Apple/MetalRHI/Private/MetalStateCache.h
--- a/Engine/Source/Runtime/Apple/MetalRHI/Private/MetalStateCache.h
+++ b/Engine/Source/Runtime/Apple/MetalRHI/Private/MetalStateCache.h
@@ -57,7 +57,6 @@ public:
 	void SetVertexStream(uint32 const Index, FMetalBuffer* Buffer, FMetalBufferData* Bytes, uint32 const Offset, uint32 const Length);
 	void SetGraphicsPipelineState(FMetalGraphicsPipelineState* State);
 	void BindUniformBuffer(EMetalShaderStages const Freq, uint32 const BufferIndex, FRHIUniformBuffer* BufferRHI);
-	void SetDirtyUniformBuffers(EMetalShaderStages const Freq, uint32 const Dirty);
 	
 	/*
 	 * Monitor if samples pass the depth and stencil tests.
@@ -133,8 +132,6 @@ public:
 	const mtlpp::Viewport& GetViewport(uint32 const Index) const { check(Index < ML_MaxViewports); return Viewport[Index]; }
 	uint32 GetVertexBufferSize(uint32 const Index);
 	uint32 GetRenderTargetArraySize() const { return RenderTargetArraySize; }
-	const FRHIUniformBuffer** GetBoundUniformBuffers(EMetalShaderStages const Freq) { return (const FRHIUniformBuffer**)&BoundUniformBuffers[Freq][0]; }
-	uint32 GetDirtyUniformBuffers(EMetalShaderStages const Freq) const { return DirtyUniformBuffers[Freq]; }
 	FMetalQueryBuffer* GetVisibilityResultsBuffer() const { return VisibilityResults; }
 	bool GetScissorRectEnabled() const { return bScissorRectEnabled; }
 	bool NeedsToSetRenderTarget(const FRHIRenderPassInfo& RenderPassInfo);
@@ -162,16 +159,10 @@ private:
 	void SetDepthStencilState(FMetalDepthStencilState* InDepthStencilState);
 	void SetRasterizerState(FMetalRasterizerState* InRasterizerState);
 
-	FORCEINLINE void SetResource(uint32 ShaderStage, uint32 BindIndex, FRHITexture* RESTRICT TextureRHI, float CurrentTime);
-	
-	FORCEINLINE void SetResource(uint32 ShaderStage, uint32 BindIndex, FMetalShaderResourceView* RESTRICT SRV, float CurrentTime);
-	
-	FORCEINLINE void SetResource(uint32 ShaderStage, uint32 BindIndex, FMetalSamplerState* RESTRICT SamplerState, float CurrentTime);
-	
-	FORCEINLINE void SetResource(uint32 ShaderStage, uint32 BindIndex, FMetalUnorderedAccessView* RESTRICT UAV, float CurrentTime);
-	
-	template <typename MetalResourceType>
-	inline int32 SetShaderResourcesFromBuffer(uint32 ShaderStage, FMetalUniformBuffer* RESTRICT Buffer, const uint32* RESTRICT ResourceMap, int32 BufferIndex, float CurrentTime);
+	FORCEINLINE void SetResource(uint32 ShaderStage, uint32 BindIndex, FRHITexture* RESTRICT TextureRHI);
+	FORCEINLINE void SetResource(uint32 ShaderStage, uint32 BindIndex, FMetalShaderResourceView* RESTRICT SRV);
+	FORCEINLINE void SetResource(uint32 ShaderStage, uint32 BindIndex, FMetalSamplerState* RESTRICT SamplerState);
+	FORCEINLINE void SetResource(uint32 ShaderStage, uint32 BindIndex, FMetalUnorderedAccessView* RESTRICT UAV);
 	
 	template <class ShaderType>
 	void SetResourcesFromTables(ShaderType Shader, uint32 ShaderStage);
diff --git a/Engine/Source/Runtime/Apple/MetalRHI/Private/MetalUniformBuffer.cpp b/Engine/Source/Runtime/Apple/MetalRHI/Private/MetalUniformBuffer.cpp
--- a/Engine/Source/Runtime/Apple/MetalRHI/Private/MetalUniformBuffer.cpp
+++ b/Engine/Source/Runtime/Apple/MetalRHI/Private/MetalUniformBuffer.cpp
@@ -17,7 +17,6 @@ FMetalSuballocatedUniformBuffer::FMetalSuballocatedUniformBuffer(const void *Con
     , Offset(0)
     , Backing(nil)
     , Shadow(FMemory::Malloc(GetSize()))
-    , ResourceTable()
 #if METAL_UNIFORM_BUFFER_VALIDATION
     , Validation(InValidation)
 #endif // METAL_UNIFORM_BUFFER_VALIDATION
diff --git a/Engine/Source/Runtime/Apple/MetalRHI/Public/MetalShaderResources.h b/Engine/Source/Runtime/Apple/MetalRHI/Public/MetalShaderResources.h
--- a/Engine/Source/Runtime/Apple/MetalRHI/Public/MetalShaderResources.h
+++ b/Engine/Source/Runtime/Apple/MetalRHI/Public/MetalShaderResources.h
@@ -99,40 +99,11 @@ enum class EMetalBufferFormat : uint8
 	Max						=54
 };
 
-struct FMetalShaderResourceTable : public FBaseShaderResourceTable
-{
-	/** Mapping of bound Textures to their location in resource tables. */
-	TArray<uint32> TextureMap;
-	friend bool operator==(const FMetalShaderResourceTable &A, const FMetalShaderResourceTable& B)
-	{
-		if (!(((FBaseShaderResourceTable&)A) == ((FBaseShaderResourceTable&)B)))
-		{
-			return false;
-		}
-		if (A.TextureMap.Num() != B.TextureMap.Num())
-		{
-			return false;
-		}
-		if (FMemory::Memcmp(A.TextureMap.GetData(), B.TextureMap.GetData(), A.TextureMap.GetTypeSize()*A.TextureMap.Num()) != 0)
-		{
-			return false;
-		}
-		return true;
-	}
-};
-
-inline FArchive& operator<<(FArchive& Ar, FMetalShaderResourceTable& SRT)
-{
-	Ar << ((FBaseShaderResourceTable&)SRT);
-	Ar << SRT.TextureMap;
-	return Ar;
-}
-
 struct FMetalShaderBindings
 {
 	TArray<TArray<CrossCompiler::FPackedArrayInfo>>	PackedUniformBuffers;
 	TArray<CrossCompiler::FPackedArrayInfo>			PackedGlobalArrays;
-	FMetalShaderResourceTable						ShaderResourceTable;
+	FShaderResourceTable							ShaderResourceTable;
 	TMap<uint8, TArray<uint8>>						ArgumentBufferMasks;
 	CrossCompiler::FShaderBindingInOutMask			InOutMask;
 
diff --git a/Engine/Source/Runtime/Apple/MetalRHI/Public/MetalUniformBuffer.h b/Engine/Source/Runtime/Apple/MetalRHI/Public/MetalUniformBuffer.h
--- a/Engine/Source/Runtime/Apple/MetalRHI/Public/MetalUniformBuffer.h
+++ b/Engine/Source/Runtime/Apple/MetalRHI/Public/MetalUniformBuffer.h
@@ -30,7 +30,6 @@ public:
     // This allows you to upload on a frame but actually use this UB later on
     void* Shadow;
 private:
-    TArray<TRefCountPtr<FRHIResource> > ResourceTable;
 #if METAL_UNIFORM_BUFFER_VALIDATION
     EUniformBufferValidation Validation;
 #endif
diff --git a/Engine/Source/Runtime/D3D12RHI/Private/D3D12CommandContext.h b/Engine/Source/Runtime/D3D12RHI/Private/D3D12CommandContext.h
--- a/Engine/Source/Runtime/D3D12RHI/Private/D3D12CommandContext.h
+++ b/Engine/Source/Runtime/D3D12RHI/Private/D3D12CommandContext.h
@@ -495,7 +495,6 @@ public:
 	void CommitComputeShaderConstants();
 
 	template <class ShaderType> void SetResourcesFromTables(const ShaderType* RESTRICT);
-	template <class ShaderType> uint32 SetUAVPSResourcesFromTables(const ShaderType* RESTRICT Shader);
 
 	void CommitGraphicsResourceTables();
 	void CommitComputeResourceTables();
diff --git a/Engine/Source/Runtime/D3D12RHI/Private/D3D12Commands.cpp b/Engine/Source/Runtime/D3D12RHI/Private/D3D12Commands.cpp
--- a/Engine/Source/Runtime/D3D12RHI/Private/D3D12Commands.cpp
+++ b/Engine/Source/Runtime/D3D12RHI/Private/D3D12Commands.cpp
@@ -1513,325 +1513,53 @@ void FD3D12CommandContext::CommitComputeShaderConstants()
 	bDiscardSharedComputeConstants = false;
 }
 
-template <EShaderFrequency Frequency>
-FORCEINLINE void SetResource(FD3D12CommandContext& CmdContext, uint32 BindIndex, FD3D12ShaderResourceView* RESTRICT SRV)
-{
-	// We set the resource through the RHI to track state for the purposes of unbinding SRVs when a UAV or RTV is bound.
-	CmdContext.StateCache.SetShaderResourceView<Frequency>(SRV, BindIndex);
-}
-
-template <EShaderFrequency Frequency>
-FORCEINLINE void SetResource(FD3D12CommandContext& CmdContext, uint32 BindIndex, FD3D12SamplerState* RESTRICT SamplerState)
-{
-	CmdContext.StateCache.SetSamplerState<Frequency>(SamplerState, BindIndex);
-}
-
-template <EShaderFrequency Frequency>
-FORCEINLINE void SetResource(FD3D12CommandContext& CmdContext, uint32 BindIndex, FD3D12UnorderedAccessView* UAV)
-{
-	uint32 InitialCount = -1;
-
-	// Actually set the UAV
-	CmdContext.StateCache.SetUAVs<SF_Pixel>(BindIndex, 1, &UAV, &InitialCount);
-}
-
-template <EShaderFrequency ShaderFrequency>
-inline int32 SetShaderResourcesFromBuffer_Surface(FD3D12CommandContext& CmdContext, FD3D12UniformBuffer* RESTRICT Buffer, const uint32 * RESTRICT ResourceMap, int32 BufferIndex, const TCHAR* LayoutName)
-{
-#if ENABLE_RHI_VALIDATION
-	constexpr ERHIAccess SRVAccess = (ShaderFrequency == SF_Compute) ? ERHIAccess::SRVCompute : ERHIAccess::SRVGraphics;
-#endif
-
-	const TRefCountPtr<FRHIResource>* RESTRICT Resources = Buffer->ResourceTable.GetData();
-	const float CurrentTime = FApp::GetCurrentTime();
-	int32 NumSetCalls = 0;
-	const uint32 BufferOffset = ResourceMap[BufferIndex];
-	if (BufferOffset > 0)
-	{
-		const uint32* RESTRICT ResourceInfos = &ResourceMap[BufferOffset];
-		uint32 ResourceInfo = *ResourceInfos++;
-		do
-		{
-			checkSlow(FRHIResourceTableEntry::GetUniformBufferIndex(ResourceInfo) == BufferIndex);
-			const uint16 ResourceIndex = FRHIResourceTableEntry::GetResourceIndex(ResourceInfo);
-			const uint8 BindIndex = FRHIResourceTableEntry::GetBindIndex(ResourceInfo);
-
-			FRHITexture* TextureRHI = (FRHITexture*)Resources[ResourceIndex].GetReference();
-			if (!TextureRHI)
-			{
-				UE_LOG(LogD3D12RHI, Fatal, TEXT("Null texture (resource %d bind %d) on UB Layout %s"), ResourceIndex, BindIndex, LayoutName);
-			}
-			TextureRHI->SetLastRenderTime(CurrentTime);
-
-			FD3D12Texture* TextureD3D12 = CmdContext.RetrieveTexture(TextureRHI);
-			FD3D12ShaderResourceView* D3D12Resource = TextureD3D12->GetShaderResourceView();
-			if (D3D12Resource == nullptr)
-			{
-				D3D12Resource = CmdContext.RetrieveTexture(GWhiteTexture->TextureRHI)->GetShaderResourceView();
-			}
-
-#if ENABLE_RHI_VALIDATION
-			if (CmdContext.Tracker)
-			{
-				CmdContext.Tracker->Assert(TextureRHI->GetWholeResourceIdentitySRV(), SRVAccess);
-			}
-#endif
-
-			SetResource<ShaderFrequency>(CmdContext, BindIndex, D3D12Resource);
-			NumSetCalls++;
-			ResourceInfo = *ResourceInfos++;
-		} while (FRHIResourceTableEntry::GetUniformBufferIndex(ResourceInfo) == BufferIndex);
-	}
-
-	INC_DWORD_STAT_BY(STAT_D3D12SetTextureInTableCalls, NumSetCalls);
-	return NumSetCalls;
-}
-
-template <EShaderFrequency ShaderFrequency>
-inline int32 SetShaderResourcesFromBuffer_SRV(FD3D12CommandContext& CmdContext, FD3D12UniformBuffer* RESTRICT Buffer, const uint32 * RESTRICT ResourceMap, int32 BufferIndex, const TCHAR* LayoutName)
-{
-#if ENABLE_RHI_VALIDATION
-	constexpr ERHIAccess SRVAccess = (ShaderFrequency == SF_Compute) ? ERHIAccess::SRVCompute : ERHIAccess::SRVGraphics;
-#endif
-
-	const TRefCountPtr<FRHIResource>* RESTRICT Resources = Buffer->ResourceTable.GetData();
-	int32 NumSetCalls = 0;
-	const uint32 BufferOffset = ResourceMap[BufferIndex];
-	if (BufferOffset > 0)
-	{
-		const uint32* RESTRICT ResourceInfos = &ResourceMap[BufferOffset];
-		uint32 ResourceInfo = *ResourceInfos++;
-		do
-		{
-			checkSlow(FRHIResourceTableEntry::GetUniformBufferIndex(ResourceInfo) == BufferIndex);
-			const uint16 ResourceIndex = FRHIResourceTableEntry::GetResourceIndex(ResourceInfo);
-			const uint8 BindIndex = FRHIResourceTableEntry::GetBindIndex(ResourceInfo);
-
-			FD3D12ShaderResourceView* D3D12Resource = CmdContext.RetrieveObject<FD3D12ShaderResourceView>((FRHIShaderResourceView*)(Resources[ResourceIndex].GetReference()));
-			if (!D3D12Resource)
-			{
-				UE_LOG(LogD3D12RHI, Fatal, TEXT("Null SRV (resource %d bind %d) on UB Layout %s"), ResourceIndex, BindIndex, LayoutName);
-			}
-
-#if ENABLE_RHI_VALIDATION
-			if (CmdContext.Tracker)
-			{
-				CmdContext.Tracker->Assert(D3D12Resource->ViewIdentity, SRVAccess);
-			}
-#endif
-
-			SetResource<ShaderFrequency>(CmdContext, BindIndex, D3D12Resource);
-			NumSetCalls++;
-			ResourceInfo = *ResourceInfos++;
-		} while (FRHIResourceTableEntry::GetUniformBufferIndex(ResourceInfo) == BufferIndex);
-	}
-
-	INC_DWORD_STAT_BY(STAT_D3D12SetTextureInTableCalls, NumSetCalls);
-	return NumSetCalls;
-}
-
-template <EShaderFrequency ShaderFrequency>
-inline int32 SetShaderResourcesFromBuffer_Sampler(FD3D12CommandContext& CmdContext, FD3D12UniformBuffer* RESTRICT Buffer, const uint32 * RESTRICT ResourceMap, int32 BufferIndex)
-{
-	const TRefCountPtr<FRHIResource>* RESTRICT Resources = Buffer->ResourceTable.GetData();
-	int32 NumSetCalls = 0;
-	const uint32 BufferOffset = ResourceMap[BufferIndex];
-	if (BufferOffset > 0)
-	{
-		const uint32* RESTRICT ResourceInfos = &ResourceMap[BufferOffset];
-		uint32 ResourceInfo = *ResourceInfos++;
-		do
-		{
-			checkSlow(FRHIResourceTableEntry::GetUniformBufferIndex(ResourceInfo) == BufferIndex);
-			const uint16 ResourceIndex = FRHIResourceTableEntry::GetResourceIndex(ResourceInfo);
-			const uint8 BindIndex = FRHIResourceTableEntry::GetBindIndex(ResourceInfo);
-
-			// todo: could coalesce adjacent bound resources.
-			FD3D12SamplerState* D3D12Resource = CmdContext.RetrieveObject<FD3D12SamplerState>((FRHISamplerState*)(Resources[ResourceIndex].GetReference()));
-
-			SetResource<ShaderFrequency>(CmdContext, BindIndex, D3D12Resource);
-			NumSetCalls++;
-			ResourceInfo = *ResourceInfos++;
-		} while (FRHIResourceTableEntry::GetUniformBufferIndex(ResourceInfo) == BufferIndex);
-	}
-
-	INC_DWORD_STAT_BY(STAT_D3D12SetTextureInTableCalls, NumSetCalls);
-	return NumSetCalls;
-}
-
 template <class ShaderType>
 void FD3D12CommandContext::SetResourcesFromTables(const ShaderType* RESTRICT Shader)
 {
 	checkSlow(Shader);
 
-	// Mask the dirty bits by those buffers from which the shader has bound resources.
-	uint32 DirtyBits = Shader->ShaderResourceTable.ResourceTableBits & DirtyUniformBuffers[ShaderType::StaticFrequency];
-	while (DirtyBits)
-	{
-		// Scan for the lowest set bit, compute its index, clear it in the set of dirty bits.
-		const uint32 LowestBitMask = (DirtyBits)& (-(int32)DirtyBits);
-		const int32 BufferIndex = FMath::FloorLog2(LowestBitMask); // todo: This has a branch on zero, we know it could never be zero...
-		DirtyBits ^= LowestBitMask;
-		FD3D12UniformBuffer* Buffer = BoundUniformBuffers[ShaderType::StaticFrequency][BufferIndex];
-
-		check(BufferIndex < Shader->ShaderResourceTable.ResourceTableLayoutHashes.Num());
+	static constexpr EShaderFrequency Frequency = static_cast<EShaderFrequency>(ShaderType::StaticFrequency);
 
-		if (!Buffer)
-		{
-			FString ShaderUB;
-#if RHI_INCLUDE_SHADER_DEBUG_DATA
-			if (BufferIndex < Shader->UniformBuffers.Num())
-			{
-				Shader->UniformBuffers[BufferIndex].ToString(ShaderUB);
-			}
-#endif
-			UE_LOG(LogD3D12RHI, Fatal, TEXT("Shader expected a uniform buffer at slot %u but got null instead (Shader='%s' UB='%s').  Rendering code needs to set a valid uniform buffer for this slot."), BufferIndex, Shader->GetShaderName(), *ShaderUB);
-		}
+	struct FUniformResourceBinder
+	{
+		FD3D12CommandContext& Context;
 
-#if RHI_INCLUDE_SHADER_DEBUG_DATA
-		// to track down OR-7159 CRASH: Client crashed at start of match in D3D11Commands.cpp
+		void SetUAV(FRHIUnorderedAccessView* UAV, uint8 Index)
 		{
-			const uint32 LayoutHash = Buffer->GetLayout().GetHash();
+			// Actually set the UAV
+			uint32 InitialCount = -1;
+			FD3D12UnorderedAccessView* D3D12UAV = FD3D12DynamicRHI::ResourceCast(UAV);
 
-			if (LayoutHash != Shader->ShaderResourceTable.ResourceTableLayoutHashes[BufferIndex])
-			{
-				auto& BufferLayout = Buffer->GetLayout();
-				const auto& DebugName = BufferLayout.GetDebugName();
-				const FString& ShaderName = Shader->ShaderName;
-#if UE_BUILD_DEBUG
-				FString ShaderUB;
-				if (BufferIndex < Shader->UniformBuffers.Num())
-				{
-					ShaderUB = FString::Printf(TEXT("expecting UB '%s'"), *Shader->UniformBuffers[BufferIndex].GetPlainNameString());
-				}
-				UE_LOG(LogD3D12RHI, Error, TEXT("SetResourcesFromTables upcoming check(%08x != %08x); Bound Layout='%s' Shader='%s' %s"), BufferLayout.GetHash(), Shader->ShaderResourceTable.ResourceTableLayoutHashes[BufferIndex], *DebugName, *ShaderName, *ShaderUB);
-				FString ResourcesString;
-				for (int32 Index = 0; Index < BufferLayout.Resources.Num(); ++Index)
-				{
-					ResourcesString += FString::Printf(TEXT("%d "), BufferLayout.Resources[Index].MemberType);
-				}
-				UE_LOG(LogD3D12RHI, Error, TEXT("Layout CB Size %d %d Resources: %s"), BufferLayout.ConstantBufferSize, BufferLayout.Resources.Num(), *ResourcesString);
-#else
-				UE_LOG(LogD3D12RHI, Error, TEXT("Bound Layout='%s' Shader='%s', Layout CB Size %d %d"), *DebugName, *ShaderName, BufferLayout.ConstantBufferSize, BufferLayout.Resources.Num());
-#endif
-				// this might mean you are accessing a data you haven't bound e.g. GBuffer
-				checkf(BufferLayout.GetHash() == Shader->ShaderResourceTable.ResourceTableLayoutHashes[BufferIndex],
-					TEXT("Uniform buffer bound to slot %u is not what the shader expected:\n")
-					TEXT("\tBound:    Uniform Buffer[%s] with Hash[%u]\n")
-					TEXT("\tExpected: Uniform Buffer[%s] with Hash[%u]"),
-					BufferIndex, *DebugName, BufferLayout.GetHash(), *Shader->UniformBuffers[BufferIndex].GetPlainNameString(), Shader->ShaderResourceTable.ResourceTableLayoutHashes[BufferIndex]);
-			}
+			Context.StateCache.SetUAVs<SF_Pixel>(Index, 1, &D3D12UAV, &InitialCount);
 		}
-#endif
-
-#if RHI_INCLUDE_SHADER_DEBUG_DATA
-		const TCHAR* LayoutName = *Buffer->GetLayout().GetDebugName();
-#else 
-		const TCHAR* LayoutName = nullptr;
-#endif
-
-		// todo: could make this two pass: gather then set
-		SetShaderResourcesFromBuffer_Surface<(EShaderFrequency)ShaderType::StaticFrequency>(*this, Buffer, Shader->ShaderResourceTable.TextureMap.GetData(), BufferIndex, LayoutName);
-		SetShaderResourcesFromBuffer_SRV<(EShaderFrequency)ShaderType::StaticFrequency>(*this, Buffer, Shader->ShaderResourceTable.ShaderResourceViewMap.GetData(), BufferIndex, LayoutName);
-		SetShaderResourcesFromBuffer_Sampler<(EShaderFrequency)ShaderType::StaticFrequency>(*this, Buffer, Shader->ShaderResourceTable.SamplerMap.GetData(), BufferIndex);
-	}
 
-	DirtyUniformBuffers[ShaderType::StaticFrequency] = 0;
-}
-
-
-template <EShaderFrequency ShaderFrequency>
-inline int32 SetShaderResourcesFromBuffer_UAVPS(FD3D12CommandContext& CmdContext, FD3D12UniformBuffer* RESTRICT Buffer, const uint32 * RESTRICT ResourceMap, int32 BufferIndex, const TCHAR* LayoutName)
-{
-#if ENABLE_RHI_VALIDATION
-	constexpr ERHIAccess UAVAccess = (ShaderFrequency == SF_Compute) ? ERHIAccess::UAVCompute : ERHIAccess::UAVGraphics;
-#endif
-
-	const TRefCountPtr<FRHIResource>* RESTRICT Resources = Buffer->ResourceTable.GetData();
-	int32 NumSetCalls = 0;
-	const uint32 BufferOffset = ResourceMap[BufferIndex];
-	if (BufferOffset > 0)
-	{
-		const uint32* RESTRICT ResourceInfos = &ResourceMap[BufferOffset];
-		uint32 ResourceInfo = *ResourceInfos++;
-		do
+		void SetSRV(FRHIShaderResourceView* SRV, uint8 Index)
 		{
-			checkSlow(FRHIResourceTableEntry::GetUniformBufferIndex(ResourceInfo) == BufferIndex);
-			const uint16 ResourceIndex = FRHIResourceTableEntry::GetResourceIndex(ResourceInfo);
-			const uint8 BindIndex = FRHIResourceTableEntry::GetBindIndex(ResourceInfo);
-
-			FRHIUnorderedAccessView* RHIUAV = (FRHIUnorderedAccessView*)(Resources[ResourceIndex].GetReference());
-			if (!RHIUAV)
-			{
-				UE_LOG(LogD3D12RHI, Fatal, TEXT("Null UAV (resource %d bind %d) on UB Layout %s"), ResourceIndex, BindIndex, LayoutName);
-			}
-
-			FD3D12UnorderedAccessView* D3D12Resource = CmdContext.RetrieveObject<FD3D12UnorderedAccessView>(RHIUAV);
-
-#if ENABLE_RHI_VALIDATION
-			if (CmdContext.Tracker)
-			{
-				CmdContext.Tracker->Assert(D3D12Resource->ViewIdentity, UAVAccess);
-			}
-#endif
-
-			SetResource<ShaderFrequency>(CmdContext, BindIndex, D3D12Resource);
-
-			NumSetCalls++;
-			ResourceInfo = *ResourceInfos++;
-		} while (FRHIResourceTableEntry::GetUniformBufferIndex(ResourceInfo) == BufferIndex);
-	}
-
-	INC_DWORD_STAT_BY(STAT_D3D12SetTextureInTableCalls, NumSetCalls);
-	return NumSetCalls;
-}
-
-
-template <class ShaderType>
-uint32 FD3D12CommandContext::SetUAVPSResourcesFromTables(const ShaderType* RESTRICT Shader)
-{
-	checkSlow(Shader);
-
-	int32 NumChanged = 0;
-	// Mask the dirty bits by those buffers from which the shader has bound resources.
-	uint32 DirtyBits = Shader->ShaderResourceTable.ResourceTableBits & DirtyUniformBuffers[ShaderType::StaticFrequency];
-	while (DirtyBits)
-	{
-		// Scan for the lowest set bit, compute its index, clear it in the set of dirty bits.
-		const uint32 LowestBitMask = (DirtyBits)& (-(int32)DirtyBits);
-		const int32 BufferIndex = FMath::FloorLog2(LowestBitMask); // todo: This has a branch on zero, we know it could never be zero...
-		DirtyBits ^= LowestBitMask;
-		FD3D12UniformBuffer* Buffer = BoundUniformBuffers[ShaderType::StaticFrequency][BufferIndex];
-
-		check(BufferIndex < Shader->ShaderResourceTable.ResourceTableLayoutHashes.Num());
+			Context.StateCache.SetShaderResourceView<Frequency>(FD3D12DynamicRHI::ResourceCast(SRV), Index);
+		}
 
-		if (!Buffer)
+		void SetTexture(FRHITexture* TextureRHI, uint8 Index)
 		{
-			FString ShaderUB;
-#if RHI_INCLUDE_SHADER_DEBUG_DATA
-			if (BufferIndex < Shader->UniformBuffers.Num())
-			{
-				Shader->UniformBuffers[BufferIndex].ToString(ShaderUB);
-			}
-#endif
-			UE_LOG(LogD3D12RHI, Fatal, TEXT("Shader expected a uniform buffer at slot %u but got null instead (Shader='%s' UB='%s').  Rendering code needs to set a valid uniform buffer for this slot."), BufferIndex, Shader->GetShaderName(), *ShaderUB);
+			FD3D12ShaderResourceView* SRV = Context.RetrieveTexture(TextureRHI)->GetShaderResourceView();
+			Context.StateCache.SetShaderResourceView<Frequency>(SRV, Index);
 		}
 
-		check(Buffer->GetLayout().GetHash() == Shader->ShaderResourceTable.ResourceTableLayoutHashes[BufferIndex]);
-
-		if ((EShaderFrequency)ShaderType::StaticFrequency == SF_Pixel)
+		void SetSampler(FRHISamplerState* Sampler, uint8 Index)
 		{
-#if RHI_INCLUDE_SHADER_DEBUG_DATA
-			const TCHAR* LayoutName = *Buffer->GetLayout().GetDebugName();
-#else
-			const TCHAR* LayoutName = nullptr;
-#endif
-
-			NumChanged += SetShaderResourcesFromBuffer_UAVPS<(EShaderFrequency)ShaderType::StaticFrequency>(*this, Buffer, Shader->ShaderResourceTable.UnorderedAccessViewMap.GetData(), BufferIndex, LayoutName);
+			Context.StateCache.SetSamplerState<Frequency>(FD3D12DynamicRHI::ResourceCast(Sampler), Index);
 		}
-	}
-	return NumChanged;
+	};
 
+	UE::RHICore::SetResourcesFromTables(
+		  FUniformResourceBinder { *this }
+		, *Shader
+		, Shader->ShaderResourceTable
+		, DirtyUniformBuffers[Frequency]
+		, BoundUniformBuffers[Frequency]
+#if ENABLE_RHI_VALIDATION
+		, Tracker
+#endif
+	);
 }
 
 void FD3D12CommandContext::CommitGraphicsResourceTables()
@@ -1844,7 +1572,7 @@ void FD3D12CommandContext::CommitGraphicsResourceTables()
 	FD3D12PixelShader* PixelShader = GraphicPSO->GetPixelShader();
 	if (PixelShader)
 	{
-		SetUAVPSResourcesFromTables(PixelShader);
+		SetResourcesFromTables(PixelShader);
 	}
 
 	if (FD3D12VertexShader* Shader = GraphicPSO->GetVertexShader())
@@ -1863,11 +1591,6 @@ void FD3D12CommandContext::CommitGraphicsResourceTables()
 	}
 #endif
 
-	if (PixelShader)
-	{
-		SetResourcesFromTables(PixelShader);
-	}
-
 #if PLATFORM_SUPPORTS_GEOMETRY_SHADERS
 	if (FD3D12GeometryShader* Shader = GraphicPSO->GetGeometryShader())
 	{
diff --git a/Engine/Source/Runtime/D3D12RHI/Private/D3D12RayTracing.cpp b/Engine/Source/Runtime/D3D12RHI/Private/D3D12RayTracing.cpp
--- a/Engine/Source/Runtime/D3D12RHI/Private/D3D12RayTracing.cpp
+++ b/Engine/Source/Runtime/D3D12RHI/Private/D3D12RayTracing.cpp
@@ -4683,6 +4683,7 @@ struct FD3D12RayTracingLocalResourceBinder
 
 template <typename ResourceBinderType>
 static bool SetRayTracingShaderResources(
+	FD3D12CommandContext& CommandContext,
 	const FD3D12RayTracingShader* Shader,
 	uint32 InNumTextures, FRHITexture* const* Textures,
 	uint32 InNumSRVs, FRHIShaderResourceView* const* SRVs,
@@ -4694,42 +4695,92 @@ static bool SetRayTracingShaderResources(
 {
 	const FD3D12RootSignature* RootSignature = Shader->pRootSignature;
 
+	struct FResourceEntry
+	{
+		FD3D12Resource* D3D12Resource = nullptr;
+		FRHIResource* RHIResource = nullptr;
+	};
+
+	struct FBindings
+	{
+		ResourceBinderType& Binder;
+		uint32 GPUIndex;
+
 #if USE_STATIC_ROOT_SIGNATURE
-	D3D12_CPU_DESCRIPTOR_HANDLE LocalCBVs[MAX_CBS];
+		D3D12_CPU_DESCRIPTOR_HANDLE LocalCBVs[MAX_CBS];
 #else // USE_STATIC_ROOT_SIGNATURE
-	D3D12_GPU_VIRTUAL_ADDRESS LocalCBVs[MAX_CBS];
+		D3D12_GPU_VIRTUAL_ADDRESS LocalCBVs[MAX_CBS];
 #endif // USE_STATIC_ROOT_SIGNATURE
 
-	D3D12_CPU_DESCRIPTOR_HANDLE LocalSRVs[MAX_SRVS];
-	D3D12_CPU_DESCRIPTOR_HANDLE LocalUAVs[MAX_UAVS];
-	D3D12_CPU_DESCRIPTOR_HANDLE LocalSamplers[MAX_SAMPLERS];
+		D3D12_CPU_DESCRIPTOR_HANDLE LocalSRVs[MAX_SRVS];
+		D3D12_CPU_DESCRIPTOR_HANDLE LocalUAVs[MAX_UAVS];
+		D3D12_CPU_DESCRIPTOR_HANDLE LocalSamplers[MAX_SAMPLERS];
 
-	struct FResourceEntry
-	{
-		FD3D12Resource* D3D12Resource = nullptr;
-		FRHIResource*   RHIResource   = nullptr;
-	};
-	TArray<FResourceEntry, TInlineAllocator<MAX_CBS + MAX_SRVS + MAX_UAVS>> ReferencedResources;
+		TArray<FResourceEntry, TInlineAllocator<MAX_CBS + MAX_SRVS + MAX_UAVS>> ReferencedResources;
 
-	uint64 BoundSRVMask = 0;
-	uint64 BoundCBVMask = 0;
-	uint64 BoundUAVMask = 0;
-	uint64 BoundSamplerMask = 0;
+		uint64 BoundSRVMask = 0;
+		uint64 BoundCBVMask = 0;
+		uint64 BoundUAVMask = 0;
+		uint64 BoundSamplerMask = 0;
 
-	uint32 GPUIndex = Binder.GetDevice()->GetGPUIndex();
-	for (uint32 SRVIndex = 0; SRVIndex < InNumTextures; ++SRVIndex)
-	{
-		FRHITexture* Resource = Textures[SRVIndex];
-		if (Resource)
+		void SetUAV(FRHIUnorderedAccessView* RHIUAV, uint8 Index)
+		{
+			FD3D12UnorderedAccessView* UAV = FD3D12CommandContext::RetrieveObject<FD3D12UnorderedAccessView>(RHIUAV, GPUIndex);
+			check(UAV != nullptr);
+
+			LocalUAVs[Index] = UAV->GetOfflineCpuHandle();
+			BoundUAVMask |= 1ull << Index;
+
+			ReferencedResources.Add({ UAV->GetResource(), UAV });
+			Binder.AddResourceTransition(UAV);
+		}
+
+		void SetSRV(FRHIShaderResourceView* RHISRV, uint8 Index)
+		{
+			FD3D12ShaderResourceView* SRV = FD3D12CommandContext::RetrieveObject<FD3D12ShaderResourceView>(RHISRV, GPUIndex);
+			check(SRV != nullptr);
+
+			LocalSRVs[Index] = SRV->GetOfflineCpuHandle();
+			BoundSRVMask |= 1ull << Index;
+
+			ReferencedResources.Add({ SRV->GetResource(), SRV });
+			Binder.AddResourceTransition(SRV);
+		}
+
+		void SetTexture(FRHITexture* RHITexture, uint8 Index)
 		{
-			FD3D12Texture* Texture = FD3D12CommandContext::RetrieveTexture(Resource, GPUIndex);
-			FD3D12ShaderResourceView* SRV = Texture->GetShaderResourceView();
-			LocalSRVs[SRVIndex] = SRV->GetOfflineCpuHandle();
-			BoundSRVMask |= 1ull << SRVIndex;
+			FD3D12ShaderResourceView* SRV = FD3D12CommandContext::RetrieveTexture(RHITexture, GPUIndex)->GetShaderResourceView();
+			if (!ensure(SRV))
+			{
+				SRV = FD3D12CommandContext::RetrieveTexture(GBlackTexture->TextureRHI, GPUIndex)->GetShaderResourceView();
+			}
+			check(SRV != nullptr);
 
-			ReferencedResources.Add({ Texture->GetResource(), Resource });
+			LocalSRVs[Index] = SRV->GetOfflineCpuHandle();
+			BoundSRVMask |= 1ull << Index;
+
+			ReferencedResources.Add({ SRV->GetResource(), SRV });
 			Binder.AddResourceTransition(SRV);
 		}
+
+		void SetSampler(FRHISamplerState* RHISampler, uint8 Index)
+		{
+			FD3D12SamplerState* Sampler = FD3D12CommandContext::RetrieveObject<FD3D12SamplerState>(RHISampler, GPUIndex);
+			check(Sampler != nullptr);
+
+			LocalSamplers[Index] = Sampler->OfflineHandle;
+			BoundSamplerMask |= 1ull << Index;
+		}
+
+	} Bindings { Binder, Binder.GetDevice()->GetGPUIndex() };
+
+	for (uint32 TextureIndex = 0; TextureIndex < InNumTextures; ++TextureIndex)
+	{
+		FRHITexture* Resource = Textures[TextureIndex];
+		if (Resource)
+		{
+			Bindings.SetTexture(Resource, TextureIndex);
+		}
 	}
 
 	for (uint32 SRVIndex = 0; SRVIndex < InNumSRVs; ++SRVIndex)
@@ -4737,12 +4788,7 @@ static bool SetRayTracingShaderResources(
 		FRHIShaderResourceView* Resource = SRVs[SRVIndex];
 		if (Resource)
 		{
-			FD3D12ShaderResourceView* SRV = FD3D12CommandContext::RetrieveObject<FD3D12ShaderResourceView>(Resource, GPUIndex);
-			LocalSRVs[SRVIndex] = SRV->GetOfflineCpuHandle();
-			BoundSRVMask |= 1ull << SRVIndex;
-
-			ReferencedResources.Add({ SRV->GetResource(), Resource });
-			Binder.AddResourceTransition(SRV);
+			Bindings.SetSRV(Resource, SRVIndex);
 		}
 	}
 
@@ -4751,15 +4797,15 @@ static bool SetRayTracingShaderResources(
 		FRHIUniformBuffer* Resource = UniformBuffers[CBVIndex];
 		if (Resource)
 		{
-			FD3D12UniformBuffer* CBV = FD3D12CommandContext::RetrieveObject<FD3D12UniformBuffer>(Resource, GPUIndex);
+			FD3D12UniformBuffer* CBV = FD3D12CommandContext::RetrieveObject<FD3D12UniformBuffer>(Resource, Bindings.GPUIndex);
 		#if USE_STATIC_ROOT_SIGNATURE
-			LocalCBVs[CBVIndex] = CBV->View->GetOfflineCpuHandle();
+			Bindings.LocalCBVs[CBVIndex] = CBV->View->GetOfflineCpuHandle();
 		#else // USE_STATIC_ROOT_SIGNATURE
-			LocalCBVs[CBVIndex] = CBV->ResourceLocation.GetGPUVirtualAddress();
+			Bindings.LocalCBVs[CBVIndex] = CBV->ResourceLocation.GetGPUVirtualAddress();
 		#endif // USE_STATIC_ROOT_SIGNATURE
-			BoundCBVMask |= 1ull << CBVIndex;
+			Bindings.BoundCBVMask |= 1ull << CBVIndex;
 
-			ReferencedResources.Add({ CBV->ResourceLocation.GetResource(), Resource });
+			Bindings.ReferencedResources.Add({ CBV->ResourceLocation.GetResource(), Resource });
 		}
 	}
 
@@ -4768,8 +4814,7 @@ static bool SetRayTracingShaderResources(
 		FRHISamplerState* Resource = Samplers[SamplerIndex];
 		if (Resource)
 		{
-			LocalSamplers[SamplerIndex] = FD3D12CommandContext::RetrieveObject<FD3D12SamplerState>(Resource, GPUIndex)->OfflineHandle;
-			BoundSamplerMask |= 1ull << SamplerIndex;
+			Bindings.SetSampler(Resource, SamplerIndex);
 		}
 	}
 
@@ -4778,173 +4823,22 @@ static bool SetRayTracingShaderResources(
 		FRHIUnorderedAccessView* Resource = UAVs[UAVIndex];
 		if (Resource)
 		{
-			FD3D12UnorderedAccessView* UAV = FD3D12CommandContext::RetrieveObject<FD3D12UnorderedAccessView>(Resource, GPUIndex);
-			LocalUAVs[UAVIndex] = UAV->GetOfflineCpuHandle();
-			BoundUAVMask |= 1ull << UAVIndex;
-
-			ReferencedResources.Add({ UAV->GetResource(), Resource });
-			Binder.AddResourceTransition(UAV);
+			Bindings.SetUAV(Resource, UAVIndex);
 		}
 	}
 
-	const FD3D12ShaderResourceTable& ShaderResourceTable = Shader->ShaderResourceTable;
-
-	uint32 DirtyBits = ShaderResourceTable.ResourceTableBits;
-
-	while (DirtyBits)
 	{
-		// Scan for the lowest set bit, compute its index, clear it in the set of dirty bits.
-		const uint32 LowestBitMask = (DirtyBits)& (-(int32)DirtyBits);
-		const int32 BufferIndex = FMath::FloorLog2(LowestBitMask); // todo: This has a branch on zero, we know it could never be zero...
-		DirtyBits ^= LowestBitMask;
-		check(uint32(BufferIndex) < InNumUniformBuffers);
-		FD3D12UniformBuffer* Buffer = FD3D12DynamicRHI::ResourceCast(UniformBuffers[BufferIndex]);
-		check(Buffer);
-		check(BufferIndex < ShaderResourceTable.ResourceTableLayoutHashes.Num());
-		check(Buffer->GetLayout().GetHash() == ShaderResourceTable.ResourceTableLayoutHashes[BufferIndex]);
-
-		// #dxr_todo: could implement all 3 loops using a common template function (and ideally share this with regular dx12 rhi code)
-
-		// Textures
-
-		{
-			const TRefCountPtr<FRHIResource>* RESTRICT Resources = Buffer->ResourceTable.GetData();
-			const TArray<uint32>& ResourceMap = ShaderResourceTable.TextureMap;
-			const uint32 BufferOffset = ResourceMap[BufferIndex];
-			if (BufferOffset > 0)
-			{
-				const uint32* RESTRICT ResourceInfos = &ResourceMap[BufferOffset];
-				uint32 ResourceInfo = *ResourceInfos++;
-				do
-				{
-					checkSlow(FRHIResourceTableEntry::GetUniformBufferIndex(ResourceInfo) == BufferIndex);
-					const uint16 ResourceIndex = FRHIResourceTableEntry::GetResourceIndex(ResourceInfo);
-					const uint8 BindIndex = FRHIResourceTableEntry::GetBindIndex(ResourceInfo);
-
-					FRHITexture* RHITexture = (FRHITexture*)Resources[ResourceIndex].GetReference();
-
-					checkf(RHITexture != nullptr, TEXT("Missing required texture binding for slot %d in uniform buffer %d (UB layout name: '%s')"),
-						BindIndex, BufferIndex, *(Buffer->GetLayout().GetDebugName()));
-
-					FD3D12ShaderResourceView* SRV = FD3D12CommandContext::RetrieveTexture(RHITexture, GPUIndex)->GetShaderResourceView();
-					if (!ensure(SRV))
-					{
-						SRV = FD3D12CommandContext::RetrieveTexture(GBlackTexture->TextureRHI, GPUIndex)->GetShaderResourceView();
-					}
-					check(SRV != nullptr);
-
-					LocalSRVs[BindIndex] = SRV->GetOfflineCpuHandle();
-					BoundSRVMask |= 1ull << BindIndex;
-
-					ReferencedResources.Add({ SRV->GetResource(), SRV });
-					Binder.AddResourceTransition(SRV);
-
-					ResourceInfo = *ResourceInfos++;
-				} while (FRHIResourceTableEntry::GetUniformBufferIndex(ResourceInfo) == BufferIndex);
-			}
-		}
-
-		// SRVs
-
-		{
-			const TRefCountPtr<FRHIResource>* RESTRICT Resources = Buffer->ResourceTable.GetData();
-			const TArray<uint32>& ResourceMap = ShaderResourceTable.ShaderResourceViewMap;
-			const uint32 BufferOffset = ResourceMap[BufferIndex];
-			if (BufferOffset > 0)
-			{
-				const uint32* RESTRICT ResourceInfos = &ResourceMap[BufferOffset];
-				uint32 ResourceInfo = *ResourceInfos++;
-				do
-				{
-					checkSlow(FRHIResourceTableEntry::GetUniformBufferIndex(ResourceInfo) == BufferIndex);
-					const uint16 ResourceIndex = FRHIResourceTableEntry::GetResourceIndex(ResourceInfo);
-					const uint8 BindIndex = FRHIResourceTableEntry::GetBindIndex(ResourceInfo);
-
-					FRHIShaderResourceView* RHISRV = (FRHIShaderResourceView*)Resources[ResourceIndex].GetReference();
-
-					checkf(RHISRV != nullptr, TEXT("Missing required shader resource view binding for slot %d in uniform buffer %d (UB layout name: '%s')"),
-						BindIndex, BufferIndex, *(Buffer->GetLayout().GetDebugName()));
-
-					FD3D12ShaderResourceView* SRV = FD3D12CommandContext::RetrieveObject<FD3D12ShaderResourceView>(RHISRV, GPUIndex);
-					check(SRV != nullptr);
-
-					LocalSRVs[BindIndex] = SRV->GetOfflineCpuHandle();
-					BoundSRVMask |= 1ull << BindIndex;
-
-					ReferencedResources.Add({ SRV->GetResource(), SRV });
-					Binder.AddResourceTransition(SRV);
-
-					ResourceInfo = *ResourceInfos++;
-				} while (FRHIResourceTableEntry::GetUniformBufferIndex(ResourceInfo) == BufferIndex);
-			}
-		}
-
-		// UAVs
-
-		{
-			const TRefCountPtr<FRHIResource>* RESTRICT Resources = Buffer->ResourceTable.GetData();
-			const TArray<uint32>& ResourceMap = ShaderResourceTable.UnorderedAccessViewMap;
-			const uint32 BufferOffset = ResourceMap[BufferIndex];
-			if (BufferOffset > 0)
-			{
-				const uint32* RESTRICT ResourceInfos = &ResourceMap[BufferOffset];
-				uint32 ResourceInfo = *ResourceInfos++;
-				do
-				{
-					checkSlow(FRHIResourceTableEntry::GetUniformBufferIndex(ResourceInfo) == BufferIndex);
-					const uint16 ResourceIndex = FRHIResourceTableEntry::GetResourceIndex(ResourceInfo);
-					const uint8 BindIndex = FRHIResourceTableEntry::GetBindIndex(ResourceInfo);
-
-					FRHIUnorderedAccessView* RHIUAV = (FRHIUnorderedAccessView*)Resources[ResourceIndex].GetReference();
-
-					checkf(RHIUAV != nullptr, TEXT("Missing required unordered access view binding for slot %d in uniform buffer %d (UB layout name: '%s')"),
-						BindIndex, BufferIndex, *(Buffer->GetLayout().GetDebugName()));
-
-					FD3D12UnorderedAccessView* UAV = FD3D12CommandContext::RetrieveObject<FD3D12UnorderedAccessView>(RHIUAV, GPUIndex);
-					check(UAV != nullptr);
-
-					LocalUAVs[BindIndex] = UAV->GetOfflineCpuHandle();
-					BoundUAVMask |= 1ull << BindIndex;
-
-					ReferencedResources.Add({ UAV->GetResource(), UAV });
-					Binder.AddResourceTransition(UAV);
-
-					ResourceInfo = *ResourceInfos++;
-				} while (FRHIResourceTableEntry::GetUniformBufferIndex(ResourceInfo) == BufferIndex);
-			}
-		}
-
-		// Samplers
-
-		{
-			const TRefCountPtr<FRHIResource>* RESTRICT Resources = Buffer->ResourceTable.GetData();
-			const TArray<uint32>& ResourceMap = ShaderResourceTable.SamplerMap;
-			const uint32 BufferOffset = ResourceMap[BufferIndex];
-			if (BufferOffset > 0)
-			{
-				const uint32* RESTRICT ResourceInfos = &ResourceMap[BufferOffset];
-				uint32 ResourceInfo = *ResourceInfos++;
-				do
-				{
-					checkSlow(FRHIResourceTableEntry::GetUniformBufferIndex(ResourceInfo) == BufferIndex);
-					const uint16 ResourceIndex = FRHIResourceTableEntry::GetResourceIndex(ResourceInfo);
-					const uint8 BindIndex = FRHIResourceTableEntry::GetBindIndex(ResourceInfo);
-
-					FRHISamplerState* RHISampler = (FRHISamplerState*)Resources[ResourceIndex].GetReference();
-
-					checkf(RHISampler != nullptr, TEXT("Missing required sampler binding for slot %d in uniform buffer %d (UB layout name: '%s')"),
-						BindIndex, BufferIndex, *(Buffer->GetLayout().GetDebugName()));
-
-					FD3D12SamplerState* Sampler = FD3D12CommandContext::RetrieveObject<FD3D12SamplerState>(RHISampler, GPUIndex);
-					check(Sampler != nullptr);
-
-					LocalSamplers[BindIndex] = Sampler->OfflineHandle;
-					BoundSamplerMask |= 1ull << BindIndex;
-
-					ResourceInfo = *ResourceInfos++;
-				} while (FRHIResourceTableEntry::GetUniformBufferIndex(ResourceInfo) == BufferIndex);
-			}
-		}
+		uint32 DirtyUniformBuffers = ~(0u);
+		UE::RHICore::SetResourcesFromTables(
+			  Bindings
+			, *Shader
+			, Shader->ShaderResourceTable
+			, DirtyUniformBuffers
+			, UniformBuffers
+#if ENABLE_RHI_VALIDATION
+			, CommandContext.Tracker
+#endif
+		);
 	}
 
 	// Bind loose parameters
@@ -4960,12 +4854,12 @@ static bool SetRayTracingShaderResources(
 		D3D12_GPU_VIRTUAL_ADDRESS BufferAddress = Binder.CreateTransientConstantBuffer(InLooseParameterData, InLooseParameterDataSize);
 
 	#if USE_STATIC_ROOT_SIGNATURE
-		LocalCBVs[CBVIndex] = Binder.CreateTransientConstantBufferView(BufferAddress, InLooseParameterDataSize);
+		Bindings.LocalCBVs[CBVIndex] = Binder.CreateTransientConstantBufferView(BufferAddress, InLooseParameterDataSize);
 	#else // USE_STATIC_ROOT_SIGNATURE
-		LocalCBVs[CBVIndex] = BufferAddress;
+		Bindings.LocalCBVs[CBVIndex] = BufferAddress;
 	#endif // USE_STATIC_ROOT_SIGNATURE
 
-		BoundCBVMask |= 1ull << CBVIndex;
+		Bindings.BoundCBVMask |= 1ull << CBVIndex;
 	}
 
 	// Validate that all resources required by the shader are set
@@ -4978,10 +4872,10 @@ static bool SetRayTracingShaderResources(
 		uint64 ExpectedMask = ExpectedCount == 64 ? ~0ull : ((1ull << ExpectedCount) - 1);
 		return (ExpectedMask & BoundMask) == ExpectedMask;
 	};
-	check(IsCompleteBinding(Shader->ResourceCounts.NumSRVs, BoundSRVMask));
-	check(IsCompleteBinding(Shader->ResourceCounts.NumUAVs, BoundUAVMask));
-	check(IsCompleteBinding(Shader->ResourceCounts.NumCBs, BoundCBVMask));
-	check(IsCompleteBinding(Shader->ResourceCounts.NumSamplers, BoundSamplerMask));
+	check(IsCompleteBinding(Shader->ResourceCounts.NumSRVs    , Bindings.BoundSRVMask));
+	check(IsCompleteBinding(Shader->ResourceCounts.NumUAVs    , Bindings.BoundUAVMask));
+	check(IsCompleteBinding(Shader->ResourceCounts.NumCBs     , Bindings.BoundCBVMask));
+	check(IsCompleteBinding(Shader->ResourceCounts.NumSamplers, Bindings.BoundSamplerMask));
 
 	FD3D12RayTracingDescriptorCache& DescriptorCache = Binder.DescriptorCache;
 	const uint32 WorkerIndex = Binder.WorkerIndex;
@@ -4989,7 +4883,7 @@ static bool SetRayTracingShaderResources(
 	const uint32 NumSRVs = Shader->ResourceCounts.NumSRVs;
 	if (NumSRVs)
 	{
-		const int32 DescriptorTableBaseIndex = DescriptorCache.AllocateDescriptorTable(LocalSRVs, NumSRVs, D3D12_DESCRIPTOR_HEAP_TYPE_CBV_SRV_UAV, WorkerIndex);
+		const int32 DescriptorTableBaseIndex = DescriptorCache.AllocateDescriptorTable(Bindings.LocalSRVs, NumSRVs, D3D12_DESCRIPTOR_HEAP_TYPE_CBV_SRV_UAV, WorkerIndex);
 		if (DescriptorTableBaseIndex < 0)
 		{
 			return false;
@@ -5005,7 +4899,7 @@ static bool SetRayTracingShaderResources(
 	const uint32 NumUAVs = Shader->ResourceCounts.NumUAVs;
 	if (NumUAVs)
 	{
-		const int32 DescriptorTableBaseIndex = DescriptorCache.AllocateDescriptorTable(LocalUAVs, NumUAVs, D3D12_DESCRIPTOR_HEAP_TYPE_CBV_SRV_UAV, WorkerIndex);
+		const int32 DescriptorTableBaseIndex = DescriptorCache.AllocateDescriptorTable(Bindings.LocalUAVs, NumUAVs, D3D12_DESCRIPTOR_HEAP_TYPE_CBV_SRV_UAV, WorkerIndex);
 		if (DescriptorTableBaseIndex < 0)
 		{
 			return false;
@@ -5023,7 +4917,7 @@ static bool SetRayTracingShaderResources(
 	{
 	#if USE_STATIC_ROOT_SIGNATURE
 
-		const uint32 DescriptorTableBaseIndex = DescriptorCache.AllocateDescriptorTable(LocalCBVs, NumCBVs, D3D12_DESCRIPTOR_HEAP_TYPE_CBV_SRV_UAV, WorkerIndex);
+		const uint32 DescriptorTableBaseIndex = DescriptorCache.AllocateDescriptorTable(Bindings.LocalCBVs, NumCBVs, D3D12_DESCRIPTOR_HEAP_TYPE_CBV_SRV_UAV, WorkerIndex);
 		const uint32 BindSlot = RootSignature->CBVRDTBindSlot(SF_Compute);
 		check(BindSlot != 0xFF);
 
@@ -5040,7 +4934,7 @@ static bool SetRayTracingShaderResources(
 		for (uint32 i = 0; i < Shader->ResourceCounts.NumCBs; ++i)
 		{
 			const uint64 SlotMask = (1ull << i);
-			D3D12_GPU_VIRTUAL_ADDRESS BufferAddress = (BoundCBVMask & SlotMask) ? LocalCBVs[i] : 0;
+			D3D12_GPU_VIRTUAL_ADDRESS BufferAddress = (Bindings.BoundCBVMask & SlotMask) ? Bindings.LocalCBVs[i] : 0;
 			Binder.SetRootCBV(BindSlot, i, BufferAddress);
 		}
 
@@ -5052,7 +4946,7 @@ static bool SetRayTracingShaderResources(
 	const uint32 NumSamplers = Shader->ResourceCounts.NumSamplers;
 	if (NumSamplers)
 	{
-		const int32 DescriptorTableBaseIndex = DescriptorCache.AllocateDescriptorTable(LocalSamplers, NumSamplers, D3D12_DESCRIPTOR_HEAP_TYPE_SAMPLER, WorkerIndex);
+		const int32 DescriptorTableBaseIndex = DescriptorCache.AllocateDescriptorTable(Bindings.LocalSamplers, NumSamplers, D3D12_DESCRIPTOR_HEAP_TYPE_SAMPLER, WorkerIndex);
 		if (DescriptorTableBaseIndex < 0)
 		{
 			return false;
@@ -5065,7 +4959,7 @@ static bool SetRayTracingShaderResources(
 		Binder.SetRootDescriptorTable(BindSlot, ResourceDescriptorTableBaseGPU);
 	}
 
-	for (const FResourceEntry& Entry : ReferencedResources)
+	for (const FResourceEntry& Entry : Bindings.ReferencedResources)
 	{
 		Binder.AddResourceReference(Entry.D3D12Resource, Entry.RHIResource);
 	}
@@ -5075,6 +4969,7 @@ static bool SetRayTracingShaderResources(
 
 template <typename ResourceBinderType>
 static bool SetRayTracingShaderResources(
+	FD3D12CommandContext& CommandContext,
 	const FD3D12RayTracingShader* Shader,
 	const FRayTracingShaderBindings& ResourceBindings,
 	ResourceBinderType& Binder)
@@ -5092,7 +4987,8 @@ static bool SetRayTracingShaderResources(
 		sizeof(ResourceBindings.UAVs) / sizeof(*ResourceBindings.UAVs) == MAX_UAVS,
 		"Ray Tracing Shader Bindings UAV array size must match D3D12 RHI Limit");
 
-	return SetRayTracingShaderResources(Shader,
+	return SetRayTracingShaderResources(
+		CommandContext, Shader,
 		UE_ARRAY_COUNT(ResourceBindings.Textures), ResourceBindings.Textures,
 		UE_ARRAY_COUNT(ResourceBindings.SRVs), ResourceBindings.SRVs,
 		UE_ARRAY_COUNT(ResourceBindings.UniformBuffers), ResourceBindings.UniformBuffers,
@@ -5201,7 +5097,7 @@ static void DispatchRays(FD3D12CommandContext& CommandContext,
 
 		DescriptorCache->SetDescriptorHeaps(CommandContext);
 		FD3D12RayTracingGlobalResourceBinder ResourceBinder(CommandContext, *DescriptorCache);
-		bResourcesBound = SetRayTracingShaderResources(RayGenShader, GlobalBindings, ResourceBinder);
+		bResourcesBound = SetRayTracingShaderResources(CommandContext, RayGenShader, GlobalBindings, ResourceBinder);
 
 		OptShaderTable->UpdateResidency(CommandContext);
 	}
@@ -5211,7 +5107,7 @@ static void DispatchRays(FD3D12CommandContext& CommandContext,
 		TransientDescriptorCache.Init(MAX_SRVS + MAX_UAVS, MAX_SAMPLERS);
 		TransientDescriptorCache.SetDescriptorHeaps(CommandContext);
 		FD3D12RayTracingGlobalResourceBinder ResourceBinder(CommandContext, TransientDescriptorCache);
-		bResourcesBound = SetRayTracingShaderResources(RayGenShader, GlobalBindings, ResourceBinder);
+		bResourcesBound = SetRayTracingShaderResources(CommandContext, RayGenShader, GlobalBindings, ResourceBinder);
 	}
 
 	if (bResourcesBound)
@@ -5335,6 +5231,7 @@ void FD3D12CommandContext::RHIRayTraceDispatchIndirect(FRHIRayTracingPipelineSta
 }
 
 static void SetRayTracingHitGroup(
+	FD3D12CommandContext& CommandContext,
 	FD3D12Device* Device,
 	FD3D12RayTracingShaderTable* ShaderTable,
 	FD3D12RayTracingScene* Scene,
@@ -5406,7 +5303,7 @@ static void SetRayTracingHitGroup(
 	}
 
 	FD3D12RayTracingLocalResourceBinder ResourceBinder(*Device, *ShaderTable, *(Shader->pRootSignature), RecordIndex, WorkerIndex);
-	const bool bResourcesBound = SetRayTracingShaderResources(Shader,
+	const bool bResourcesBound = SetRayTracingShaderResources(CommandContext, Shader,
 		0, nullptr, // Textures
 		0, nullptr, // SRVs
 		NumUniformBuffers, UniformBuffers,
@@ -5427,6 +5324,7 @@ static void SetRayTracingHitGroup(
 }
 
 static void SetRayTracingCallableShader(
+	FD3D12CommandContext& CommandContext,
 	FD3D12Device* Device,
 	FD3D12RayTracingShaderTable* ShaderTable,
 	FD3D12RayTracingScene* Scene,
@@ -5451,7 +5349,7 @@ static void SetRayTracingCallableShader(
 		const FD3D12RayTracingShader* Shader = Pipeline->CallableShaders.Shaders[ShaderIndexInPipeline];
 
 		FD3D12RayTracingLocalResourceBinder ResourceBinder(*Device, *ShaderTable, *(Shader->pRootSignature), RecordIndex, WorkerIndex);
-		const bool bResourcesBound = SetRayTracingShaderResources(Shader,
+		const bool bResourcesBound = SetRayTracingShaderResources(CommandContext, Shader,
 			0, nullptr, // Textures
 			0, nullptr, // SRVs
 			NumUniformBuffers, UniformBuffers,
@@ -5470,6 +5368,7 @@ static void SetRayTracingCallableShader(
 }
 
 static void SetRayTracingMissShader(
+	FD3D12CommandContext& CommandContext,
 	FD3D12Device* Device,
 	FD3D12RayTracingShaderTable* ShaderTable,
 	FD3D12RayTracingScene* Scene,
@@ -5490,7 +5389,7 @@ static void SetRayTracingMissShader(
 	const FD3D12RayTracingShader* Shader = Pipeline->MissShaders.Shaders[ShaderIndexInPipeline];
 
 	FD3D12RayTracingLocalResourceBinder ResourceBinder(*Device, *ShaderTable, *(Shader->pRootSignature), RecordIndex, WorkerIndex);
-	const bool bResourcesBound = SetRayTracingShaderResources(Shader,
+	const bool bResourcesBound = SetRayTracingShaderResources(CommandContext, Shader,
 		0, nullptr, // Textures
 		0, nullptr, // SRVs
 		NumUniformBuffers, UniformBuffers,
@@ -5540,13 +5439,13 @@ void FD3D12CommandContext::RHISetRayTracingBindings(
 		TaskContexts.Add(FTaskContext{WorkerIndex});
 	}
 
-	auto BindingTask = [Bindings, Device = Device, ShaderTable, Scene, Pipeline, BindingType](const FTaskContext& Context, int32 CurrentIndex)
+	auto BindingTask = [this, Bindings, Device = Device, ShaderTable, Scene, Pipeline, BindingType](const FTaskContext& Context, int32 CurrentIndex)
 	{
 		const FRayTracingLocalShaderBindings& Binding = Bindings[CurrentIndex];
 
 		if (BindingType == ERayTracingBindingType::HitGroup)
 		{
-			SetRayTracingHitGroup(Device, ShaderTable, Scene, Pipeline,
+			SetRayTracingHitGroup(*this, Device, ShaderTable, Scene, Pipeline,
 				Binding.InstanceIndex,
 				Binding.SegmentIndex,
 				Binding.ShaderSlot,
@@ -5560,7 +5459,7 @@ void FD3D12CommandContext::RHISetRayTracingBindings(
 		}
 		else if (BindingType == ERayTracingBindingType::CallableShader)
 		{
-			SetRayTracingCallableShader(Device, ShaderTable, Scene, Pipeline,
+			SetRayTracingCallableShader(*this, Device, ShaderTable, Scene, Pipeline,
 				Binding.ShaderSlot,
 				Binding.ShaderIndexInPipeline,
 				Binding.NumUniformBuffers,
@@ -5572,7 +5471,7 @@ void FD3D12CommandContext::RHISetRayTracingBindings(
 		}
 		else if (BindingType == ERayTracingBindingType::MissShader)
 		{
-			SetRayTracingMissShader(Device, ShaderTable, Scene, Pipeline,
+			SetRayTracingMissShader(*this, Device, ShaderTable, Scene, Pipeline,
 				Binding.ShaderSlot,
 				Binding.ShaderIndexInPipeline,
 				Binding.NumUniformBuffers,
@@ -5613,7 +5512,7 @@ void FD3D12CommandContext::RHISetRayTracingHitGroup(
 		TEXT("Hit shader records are assumed to be at the beginning of local shader table"));
 
 	const uint32 WorkerIndex = 0;
-	SetRayTracingHitGroup(GetParentDevice(), ShaderTable, Scene, Pipeline,
+	SetRayTracingHitGroup(*this, GetParentDevice(), ShaderTable, Scene, Pipeline,
 		InstanceIndex,
 		SegmentIndex,
 		ShaderSlot,
@@ -5637,7 +5536,7 @@ void FD3D12CommandContext::RHISetRayTracingCallableShader(
 	FD3D12RayTracingShaderTable* ShaderTable = Scene->FindOrCreateShaderTable(Pipeline, GetParentDevice());
 	const uint32 WorkerIndex = 0;
 
-	SetRayTracingCallableShader(GetParentDevice(), ShaderTable, Scene, Pipeline, ShaderSlotInScene, ShaderIndexInPipeline, NumUniformBuffers, UniformBuffers, 0, nullptr, UserData, WorkerIndex);
+	SetRayTracingCallableShader(*this, GetParentDevice(), ShaderTable, Scene, Pipeline, ShaderSlotInScene, ShaderIndexInPipeline, NumUniformBuffers, UniformBuffers, 0, nullptr, UserData, WorkerIndex);
 }
 
 void FD3D12CommandContext::RHISetRayTracingMissShader(
@@ -5650,7 +5549,7 @@ void FD3D12CommandContext::RHISetRayTracingMissShader(
 	FD3D12RayTracingPipelineState* Pipeline = FD3D12DynamicRHI::ResourceCast(InPipeline);
 	FD3D12RayTracingShaderTable* ShaderTable = Scene->FindOrCreateShaderTable(Pipeline, GetParentDevice());
 	const uint32 WorkerIndex = 0;
-	SetRayTracingMissShader(GetParentDevice(), ShaderTable, Scene, Pipeline,
+	SetRayTracingMissShader(*this, GetParentDevice(), ShaderTable, Scene, Pipeline,
 		ShaderSlotInScene, ShaderIndexInPipeline,
 		NumUniformBuffers, UniformBuffers,
 		0, nullptr, // Loose parameters
diff --git a/Engine/Source/Runtime/D3D12RHI/Private/D3D12Shader.h b/Engine/Source/Runtime/D3D12RHI/Private/D3D12Shader.h
--- a/Engine/Source/Runtime/D3D12RHI/Private/D3D12Shader.h
+++ b/Engine/Source/Runtime/D3D12RHI/Private/D3D12Shader.h
@@ -44,11 +44,7 @@ struct FD3D12ShaderData
 	/** The shader's bytecode, with custom data in the last byte. */
 	TArray<uint8> Code;
 
-	FD3D12ShaderResourceTable ShaderResourceTable;
-
-#if RHI_INCLUDE_SHADER_DEBUG_DATA
-	TArray<FName> UniformBuffers;
-#endif
+	FShaderResourceTable ShaderResourceTable;
 
 	FShaderCodePackedResourceCounts ResourceCounts{};
 
diff --git a/Engine/Source/Runtime/D3D12RHI/Private/D3D12Shaders.cpp b/Engine/Source/Runtime/D3D12RHI/Private/D3D12Shaders.cpp
--- a/Engine/Source/Runtime/D3D12RHI/Private/D3D12Shaders.cpp
+++ b/Engine/Source/Runtime/D3D12RHI/Private/D3D12Shaders.cpp
@@ -18,7 +18,7 @@ static inline bool ReadShaderOptionalData(FShaderCodeReader& InShaderCode, TShad
 	OutShader.ResourceCounts = *PackedResourceCounts;
 
 #if RHI_INCLUDE_SHADER_DEBUG_DATA
-	OutShader.ShaderName = InShaderCode.FindOptionalData(FShaderCodeName::Key);
+	OutShader.Debug.ShaderName = InShaderCode.FindOptionalData(FShaderCodeName::Key);
 
 	int32 UniformBufferTableSize = 0;
 	const uint8* UniformBufferData = InShaderCode.FindOptionalDataAndSize(FShaderCodeUniformBuffers::Key, UniformBufferTableSize);
@@ -27,10 +27,10 @@ static inline bool ReadShaderOptionalData(FShaderCodeReader& InShaderCode, TShad
 		FBufferReader UBReader((void*)UniformBufferData, UniformBufferTableSize, false);
 		TArray<FString> Names;
 		UBReader << Names;
-		check(OutShader.UniformBuffers.Num() == 0);
+		check(OutShader.Debug.UniformBufferNames.Num() == 0);
 		for (int32 Index = 0; Index < Names.Num(); ++Index)
 		{
-			OutShader.UniformBuffers.Add(FName(*Names[Index]));
+			OutShader.Debug.UniformBufferNames.Add(FName(*Names[Index]));
 		}
 	}
 #endif
diff --git a/Engine/Source/Runtime/D3D12RHI/Private/D3D12UniformBuffer.cpp b/Engine/Source/Runtime/D3D12RHI/Private/D3D12UniformBuffer.cpp
--- a/Engine/Source/Runtime/D3D12RHI/Private/D3D12UniformBuffer.cpp
+++ b/Engine/Source/Runtime/D3D12RHI/Private/D3D12UniformBuffer.cpp
@@ -186,13 +186,13 @@ FUniformBufferRHIRef FD3D12DynamicRHI::RHICreateUniformBuffer(const void* Conten
 
 		for (FD3D12UniformBuffer& CurrentBuffer : *UniformBufferOut)
 		{
-			CurrentBuffer.ResourceTable.SetNumZeroed(NumResources);
+			CurrentBuffer.GetResourceTable().SetNumZeroed(NumResources);
 
 			if (Contents)
 			{
 				for (int32 Index = 0; Index < NumResources; ++Index)
 				{
-					CurrentBuffer.ResourceTable[Index] = GetShaderParameterResourceRHI(Contents, Layout->Resources[Index].MemberOffset, Layout->Resources[Index].MemberType);
+					CurrentBuffer.GetResourceTable()[Index] = GetShaderParameterResourceRHI(Contents, Layout->Resources[Index].MemberOffset, Layout->Resources[Index].MemberType);
 				}
 			}
 		}
@@ -235,8 +235,8 @@ FRHICOMMAND_MACRO(FRHICommandD3D12UpdateUniformBuffer)
 		for (int32 i = 0; i < UpdatedResources.Num(); ++i)
 		{
 			//check(UniformBuffer->ResourceTable[i]);
-			UniformBuffer->ResourceTable[i] = UpdatedResources[i];
-			check(UniformBuffer->ResourceTable[i]);
+			UniformBuffer->GetResourceTable()[i] = UpdatedResources[i];
+			check(UniformBuffer->GetResourceTable()[i]);
 		}
 		FD3D12ResourceLocation::TransferOwnership(UniformBuffer->ResourceLocation, UpdatedLocation);
 #if USE_STATIC_ROOT_SIGNATURE
@@ -278,7 +278,7 @@ void FD3D12DynamicRHI::RHIUpdateUniformBuffer(FRHICommandListBase& RHICmdList, F
 	// Update buffers on all GPUs by looping over FD3D12LinkedAdapterObject chain
 	for (FD3D12UniformBuffer& UniformBuffer : *FirstUniformBuffer)
 	{
-		check(UniformBuffer.ResourceTable.Num() == NumResources);
+		check(UniformBuffer.GetResourceTable().Num() == NumResources);
 
 		FD3D12Device* Device = UniformBuffer.GetParentDevice();
 
diff --git a/Engine/Source/Runtime/D3D12RHI/Public/D3D12Resources.h b/Engine/Source/Runtime/D3D12RHI/Public/D3D12Resources.h
--- a/Engine/Source/Runtime/D3D12RHI/Public/D3D12Resources.h
+++ b/Engine/Source/Runtime/D3D12RHI/Public/D3D12Resources.h
@@ -821,9 +821,6 @@ public:
 	/** The D3D12 constant buffer resource */
 	FD3D12ResourceLocation ResourceLocation;
 
-	/** Resource table containing RHI references. */
-	TArray<TRefCountPtr<FRHIResource> > ResourceTable;
-
 	const EUniformBufferUsage UniformBufferUsage;
 
 	/** Initialization constructor. */
@@ -839,6 +836,10 @@ public:
 	}
 
 	virtual ~FD3D12UniformBuffer();
+
+	// Provides public non-const access to ResourceTable.
+	// @todo refactor uniform buffers to perform updates as a member function, so this isn't necessary.
+	TArray<TRefCountPtr<FRHIResource>>& GetResourceTable() { return ResourceTable; }
 };
 
 class FD3D12Buffer : public FRHIBuffer, public FD3D12BaseShaderResource, public FD3D12LinkedAdapterObject<FD3D12Buffer>
diff --git a/Engine/Source/Runtime/D3D12RHI/Public/D3D12ShaderResources.h b/Engine/Source/Runtime/D3D12RHI/Public/D3D12ShaderResources.h
--- a/Engine/Source/Runtime/D3D12RHI/Public/D3D12ShaderResources.h
+++ b/Engine/Source/Runtime/D3D12RHI/Public/D3D12ShaderResources.h
@@ -11,25 +11,3 @@ const int32 PackedShaderKey = 'XSHA';
 
 // Key indicating whether serialized ray tracing shader contains a DXIL library or a precompiled PSO.
 const int32 RayTracingPrecompiledPSOKey = 'RTPS';
-
-struct FD3D12ShaderResourceTable : public FBaseShaderResourceTable
-{
-	/** Mapping of bound Textures to their location in resource tables. */
-	TArray<uint32> TextureMap;
-
-	friend bool operator==(const FD3D12ShaderResourceTable& A, const FD3D12ShaderResourceTable& B)
-	{
-		const FBaseShaderResourceTable& BaseA = A;
-		const FBaseShaderResourceTable& BaseB = B;
-		return BaseA == BaseB && (FMemory::Memcmp(A.TextureMap.GetData(), B.TextureMap.GetData(), A.TextureMap.GetTypeSize()*A.TextureMap.Num()) == 0);
-	}
-
-	friend inline FArchive& operator<<(FArchive& Ar, FD3D12ShaderResourceTable& SRT)
-	{
-		FBaseShaderResourceTable& BaseSRT = SRT;
-		Ar << BaseSRT;
-		Ar << SRT.TextureMap;
-		return Ar;
-	}
-
-};
diff --git a/Engine/Source/Runtime/OpenGLDrv/Private/OpenGLCommands.cpp b/Engine/Source/Runtime/OpenGLDrv/Private/OpenGLCommands.cpp
--- a/Engine/Source/Runtime/OpenGLDrv/Private/OpenGLCommands.cpp
+++ b/Engine/Source/Runtime/OpenGLDrv/Private/OpenGLCommands.cpp
@@ -1749,7 +1749,7 @@ void FOpenGLDynamicRHI::SetupVertexArrays(FOpenGLContextState& ContextState, uin
 
 	check(IsValidRef(PendingState.BoundShaderState));
 	FOpenGLVertexDeclaration* VertexDeclaration = PendingState.BoundShaderState->VertexDeclaration;
-	const CrossCompiler::FShaderBindingInOutMask& AttributeMask = PendingState.BoundShaderState->GetVertexShader()->Bindings.InOutMask;
+	const CrossCompiler::FShaderBindingInOutMask& AttributeMask = PendingState.BoundShaderState->GetVertexShader()->GetGLResourceObject()->Bindings.InOutMask;
 
 	if (ContextState.VertexDecl != VertexDeclaration || AttributeMask.Bitmask != ContextState.VertexAttrs_EnabledBits)
 	{
@@ -2059,11 +2059,11 @@ void FOpenGLDynamicRHI::CommitNonComputeShaderConstants()
 	FOpenGLLinkedProgram* LinkedProgram = PendingState.BoundShaderState->LinkedProgram;
 	if (GUseEmulatedUniformBuffers)
 	{
-		PendingState.ShaderParameters[CrossCompiler::SHADER_STAGE_VERTEX].CommitPackedUniformBuffers(LinkedProgram, CrossCompiler::SHADER_STAGE_VERTEX, PendingState.BoundUniformBuffers[SF_Vertex], PendingState.BoundShaderState->GetVertexShader()->UniformBuffersCopyInfo);
-		PendingState.ShaderParameters[CrossCompiler::SHADER_STAGE_PIXEL].CommitPackedUniformBuffers(LinkedProgram, CrossCompiler::SHADER_STAGE_PIXEL, PendingState.BoundUniformBuffers[SF_Pixel], PendingState.BoundShaderState->GetPixelShader()->UniformBuffersCopyInfo);
+		PendingState.ShaderParameters[CrossCompiler::SHADER_STAGE_VERTEX].CommitPackedUniformBuffers(LinkedProgram, CrossCompiler::SHADER_STAGE_VERTEX, PendingState.BoundUniformBuffers[SF_Vertex], PendingState.BoundShaderState->GetVertexShader()->GetGLResourceObject()->UniformBuffersCopyInfo);
+		PendingState.ShaderParameters[CrossCompiler::SHADER_STAGE_PIXEL].CommitPackedUniformBuffers(LinkedProgram, CrossCompiler::SHADER_STAGE_PIXEL, PendingState.BoundUniformBuffers[SF_Pixel], PendingState.BoundShaderState->GetPixelShader()->GetGLResourceObject()->UniformBuffersCopyInfo);
 		if (PendingState.BoundShaderState->GetGeometryShader())
 		{
-			PendingState.ShaderParameters[CrossCompiler::SHADER_STAGE_GEOMETRY].CommitPackedUniformBuffers(LinkedProgram, CrossCompiler::SHADER_STAGE_GEOMETRY, PendingState.BoundUniformBuffers[SF_Geometry], PendingState.BoundShaderState->GetGeometryShader()->UniformBuffersCopyInfo);
+			PendingState.ShaderParameters[CrossCompiler::SHADER_STAGE_GEOMETRY].CommitPackedUniformBuffers(LinkedProgram, CrossCompiler::SHADER_STAGE_GEOMETRY, PendingState.BoundUniformBuffers[SF_Geometry], PendingState.BoundShaderState->GetGeometryShader()->GetGLResourceObject()->UniformBuffersCopyInfo);
 		}
 	}
 	
@@ -2123,124 +2123,105 @@ template <> FORCEINLINE uint32 GetNumUAVUnits<SF_Compute>()	{ return FOpenGL::Ge
 template <> FORCEINLINE uint32 GetNumUAVUnits<SF_Pixel>()	{ return FOpenGL::GetMaxPixelUAVUnits(); }
 template <> FORCEINLINE uint32 GetNumUAVUnits<SF_Vertex>()	{ return FOpenGL::GetMaxPixelUAVUnits(); }
 
-template <EShaderFrequency Frequency>
-void SetResource(FOpenGLDynamicRHI* RESTRICT OpenGLRHI, uint32 BindIndex, FRHITexture* RESTRICT TextureRHI)
+template <class ShaderType>
+FORCEINLINE void FOpenGLDynamicRHI::SetResourcesFromTables(ShaderType* RESTRICT Shader)
 {
-	FOpenGLTexture* Texture = GetOpenGLTextureFromRHITexture(TextureRHI);
-#if !UE_BUILD_TEST && !UE_BUILD_SHIPPING
-	ensureMsgf(BindIndex < GetNumTextureUnits<Frequency>(), TEXT("Using more %s texture units (%d) than allowed (%d) on a shader unit!"), GetShaderFrequencyString(Frequency, false), BindIndex, GetNumTextureUnits<Frequency>());
-#endif
-	if (Texture)
-	{
-		TextureRHI->SetLastRenderTime(FApp::GetCurrentTime());
-		OpenGLRHI->InternalSetShaderTexture(Texture, nullptr, GetFirstTextureUnit<Frequency>() + BindIndex, Texture->Target, Texture->GetResource(), Texture->GetNumMips(), -1);
-	}
-	else
-	{
-		OpenGLRHI->InternalSetShaderTexture(Texture, nullptr, GetFirstTextureUnit<Frequency>() + BindIndex, 0, 0, 0, -1);
-	}
-	// clear any previous sampler state
-	OpenGLRHI->InternalSetSamplerStates(GetFirstTextureUnit<Frequency>() + BindIndex, nullptr);
-}
+	checkSlow(Shader);
+	VERIFY_GL_SCOPE();
 
-template <EShaderFrequency Frequency>
-FORCEINLINE void SetResource(FOpenGLDynamicRHI* RESTRICT OpenGLRHI, uint32 BindIndex, FOpenGLSamplerState* RESTRICT SamplerState)
-{
-	OpenGLRHI->InternalSetSamplerStates(GetFirstTextureUnit<Frequency>() + BindIndex, SamplerState);
-}
+	static constexpr EShaderFrequency Frequency = static_cast<EShaderFrequency>(ShaderType::ContainedGLType::StaticFrequency);
 
-template <EShaderFrequency Frequency>
-FORCEINLINE void SetResource(FOpenGLDynamicRHI* RESTRICT OpenGLRHI, uint32 BindIndex, FRHIShaderResourceView* RESTRICT RHISRV)
-{
-#if !UE_BUILD_TEST && !UE_BUILD_SHIPPING
-	ensureMsgf(BindIndex < GetNumTextureUnits<Frequency>(), TEXT("Using more %s texture units (%d) than allowed (%d) on a shader unit!"), GetShaderFrequencyString(Frequency, false), BindIndex, GetNumTextureUnits<Frequency>());	VERIFY_GL_SCOPE();
-#endif
-	auto SRV = FOpenGLDynamicRHI::ResourceCast(RHISRV);
-	if (SRV->Target == GL_SHADER_STORAGE_BUFFER)
+	struct FUniformResourceBinder
 	{
-		OpenGLRHI->InternalSetShaderBufferUAV(GetFirstUAVUnit<Frequency>() + BindIndex, SRV->Resource);
-	}
-	else
-	{
-		OpenGLRHI->InternalSetShaderTexture(NULL, SRV, GetFirstTextureUnit<Frequency>() + BindIndex, SRV->Target, SRV->Resource, 0, SRV->LimitMip);
-		SetResource<Frequency>(OpenGLRHI,BindIndex,OpenGLRHI->GetPointSamplerState());
-	}
-}
+		FOpenGLDynamicRHI& RHI;
 
-template <EShaderFrequency Frequency>
-FORCEINLINE void SetResource(FOpenGLDynamicRHI* RESTRICT OpenGLRHI, uint32 BindIndex, FRHIUnorderedAccessView* RESTRICT RHIUAV)
-{
-#if !UE_BUILD_TEST && !UE_BUILD_SHIPPING
-	ensureMsgf(BindIndex < GetNumUAVUnits<Frequency>(), TEXT("Using more %s image units (%d) than allowed (%d) on a shader unit!"), GetShaderFrequencyString(Frequency, false), BindIndex, GetNumUAVUnits<Frequency>());	VERIFY_GL_SCOPE();
-#endif
-	auto UAV = FOpenGLDynamicRHI::ResourceCast(RHIUAV);
-	GLint UAVIndex = GetFirstUAVUnit<Frequency>() + BindIndex;
-	
-	if (UAV->Resource)
-	{
-		GLenum Access = (Frequency == SF_Compute) ? GL_READ_WRITE : GL_WRITE_ONLY;
-		// TODO: This must be true for 3D textures
-		bool bLayered = false;
-		GLint Layer = 0;
-		OpenGLRHI->InternalSetShaderImageUAV(UAVIndex, UAV->Format, UAV->Resource, bLayered, Layer, Access);
-	}
-	else
-	{
-		OpenGLRHI->InternalSetShaderBufferUAV(UAVIndex, UAV->BufferResource);
-	}
-}
-
-template <class GLResourceType, EShaderFrequency ShaderFrequency>
-FORCEINLINE void SetShaderResourcesFromBuffer(FOpenGLDynamicRHI* RESTRICT OpenGLRHI, FOpenGLUniformBuffer* RESTRICT Buffer, const uint32* RESTRICT ResourceMap, int32 BufferIndex)
-{
-	uint32 BufferOffset = ResourceMap[BufferIndex];
-	if (BufferOffset > 0)
-	{
-		const TRefCountPtr<FRHIResource>* RESTRICT Resources = Buffer->ResourceTable.GetData();
-		const uint32* RESTRICT ResourceInfos = &ResourceMap[BufferOffset];
-		uint32 ResourceInfo = *ResourceInfos++;
-		do 
+		void SetUAV(FRHIUnorderedAccessView* RHIUAV, uint8 Index)
 		{
-			checkSlow(FRHIResourceTableEntry::GetUniformBufferIndex(ResourceInfo) == BufferIndex);
-			const uint16 ResourceIndex = FRHIResourceTableEntry::GetResourceIndex(ResourceInfo);
-			const uint8 BindIndex = FRHIResourceTableEntry::GetBindIndex(ResourceInfo);
+			ensureMsgf(Index < GetNumUAVUnits<Frequency>()
+				, TEXT("Using more %s image units (%d) than allowed (%d) on a shader unit!")
+				, GetShaderFrequencyString(Frequency, false)
+				, Index
+				, GetNumUAVUnits<Frequency>()
+			);
 
-			GLResourceType* ResourcePtr = (GLResourceType*)Resources[ResourceIndex].GetReference();
-			SetResource<ShaderFrequency>(OpenGLRHI, BindIndex, ResourcePtr);
+			auto UAV = FOpenGLDynamicRHI::ResourceCast(RHIUAV);
+			GLint UAVIndex = GetFirstUAVUnit<Frequency>() + Index;
 
-			ResourceInfo = *ResourceInfos++;
-		} while (FRHIResourceTableEntry::GetUniformBufferIndex(ResourceInfo) == BufferIndex);
-	}
-}
+			if (UAV->Resource)
+			{
+				GLenum Access = (Frequency == SF_Compute) ? GL_READ_WRITE : GL_WRITE_ONLY;
+				// TODO: This must be true for 3D textures
+				bool bLayered = false;
+				GLint Layer = 0;
+				RHI.InternalSetShaderImageUAV(UAVIndex, UAV->Format, UAV->Resource, bLayered, Layer, Access);
+			}
+			else
+			{
+				RHI.InternalSetShaderBufferUAV(UAVIndex, UAV->BufferResource);
+			}
+		}
 
-template <class ShaderType>
-FORCEINLINE void FOpenGLDynamicRHI::SetResourcesFromTables(const ShaderType* RESTRICT Shader)
-{
-	checkSlow(Shader);
-	const FOpenGLShaderResourceTable* RESTRICT SRT = &Shader->Bindings.ShaderResourceTable;
+		void SetSRV(FRHIShaderResourceView* RHISRV, uint8 Index)
+		{
+			ensureMsgf(Index < GetNumTextureUnits<Frequency>()
+				, TEXT("Using more %s texture units (%d) than allowed (%d) on a shader unit!")
+				, GetShaderFrequencyString(Frequency, false)
+				, Index
+				, GetNumTextureUnits<Frequency>()
+			);
 
-	// Mask the dirty bits by those buffers from which the shader has bound resources.
-	uint32 DirtyBits = SRT->ResourceTableBits & PendingState.DirtyUniformBuffers[ShaderType::StaticFrequency];
-	int32 BufferIndex = 0;
-	do
-	{
-		uint32 LowBit = (DirtyBits & 1);
-		if (LowBit)
+			auto SRV = FOpenGLDynamicRHI::ResourceCast(RHISRV);
+			if (SRV->Target == GL_SHADER_STORAGE_BUFFER)
+			{
+				RHI.InternalSetShaderBufferUAV(GetFirstUAVUnit<Frequency>() + Index, SRV->Resource);
+			}
+			else
+			{
+				RHI.InternalSetShaderTexture(nullptr, SRV, GetFirstTextureUnit<Frequency>() + Index, SRV->Target, SRV->Resource, 0, SRV->LimitMip);
+				SetSampler(RHI.GetPointSamplerState(), Index);
+			}
+		}
+
+		void SetTexture(FRHITexture* TextureRHI, uint8 Index)
 		{
-			FOpenGLUniformBuffer* Buffer = (FOpenGLUniformBuffer*)PendingState.BoundUniformBuffers[ShaderType::StaticFrequency][BufferIndex].GetReference();
-			check(Buffer);
-			check(BufferIndex < SRT->ResourceTableLayoutHashes.Num());
-			check(Buffer->GetLayout().GetHash() == SRT->ResourceTableLayoutHashes[BufferIndex]);
+			ensureMsgf(Index < GetNumTextureUnits<Frequency>()
+				, TEXT("Using more %s texture units (%d) than allowed (%d) on a shader unit!")
+				, GetShaderFrequencyString(Frequency, false)
+				, Index
+				, GetNumTextureUnits<Frequency>()
+			);
 
-			// todo: could make this two pass: gather then set
-			SetShaderResourcesFromBuffer<FRHITexture, (EShaderFrequency)ShaderType::StaticFrequency>(this, Buffer, SRT->TextureMap.GetData(), BufferIndex);
-			SetShaderResourcesFromBuffer<FRHIShaderResourceView, (EShaderFrequency)ShaderType::StaticFrequency>(this, Buffer, SRT->ShaderResourceViewMap.GetData(), BufferIndex);
-			SetShaderResourcesFromBuffer<FOpenGLSamplerState, (EShaderFrequency)ShaderType::StaticFrequency>(this, Buffer, SRT->SamplerMap.GetData(), BufferIndex);
-			SetShaderResourcesFromBuffer<FRHIUnorderedAccessView, (EShaderFrequency)ShaderType::StaticFrequency>(this, Buffer, SRT->UnorderedAccessViewMap.GetData(), BufferIndex);
+			FOpenGLTexture* Texture = GetOpenGLTextureFromRHITexture(TextureRHI);
+
+			if (Texture)
+			{
+				RHI.InternalSetShaderTexture(Texture, nullptr, GetFirstTextureUnit<Frequency>() + Index, Texture->Target, Texture->GetResource(), Texture->GetNumMips(), -1);
+			}
+			else
+			{
+				RHI.InternalSetShaderTexture(Texture, nullptr, GetFirstTextureUnit<Frequency>() + Index, 0, 0, 0, -1);
+			}
+
+			// clear any previous sampler state
+			RHI.InternalSetSamplerStates(GetFirstTextureUnit<Frequency>() + Index, nullptr);
+		}
+
+		void SetSampler(FRHISamplerState* Sampler, uint8 Index)
+		{
+			RHI.InternalSetSamplerStates(GetFirstTextureUnit<Frequency>() + Index, ResourceCast(Sampler));
 		}
-		BufferIndex++;
-		DirtyBits >>= 1;
-	} while (DirtyBits);
+	};
+
+	UE::RHICore::SetResourcesFromTables(
+		  FUniformResourceBinder { *this }
+		, *Shader
+		, Shader->GetGLResourceObject()->Bindings.ShaderResourceTable
+		, PendingState.DirtyUniformBuffers[Frequency]
+		, PendingState.BoundUniformBuffers[Frequency]
+#if ENABLE_RHI_VALIDATION
+		, Tracker
+#endif
+	);
 }
 
 
@@ -2277,7 +2258,7 @@ void FOpenGLDynamicRHI::CommitGraphicsResourceTablesInner()
 }
 
 
-void FOpenGLDynamicRHI::CommitComputeResourceTables(FOpenGLComputeShader* ComputeShader)
+void FOpenGLDynamicRHI::CommitComputeResourceTables(FOpenGLComputeShaderProxy* ComputeShader)
 {
 	VERIFY_GL_SCOPE();
 
@@ -2811,7 +2792,8 @@ void FOpenGLDynamicRHI::RHIDispatchComputeShader(uint32 ThreadGroupCountX, uint3
 	FRHIComputeShader* ComputeShaderRHI = PendingState.CurrentComputeShader;
 	check(ComputeShaderRHI);
 
-	FOpenGLComputeShader* ComputeShader = ResourceCast(ComputeShaderRHI);
+	FOpenGLComputeShaderProxy* ComputeShaderProxy = ResourceCastProxy(ComputeShaderRHI);
+	FOpenGLComputeShader* ComputeShader = ComputeShaderProxy->GetGLResourceObject();
 
 	if (ComputeShader->LinkedProgram == nullptr)
 	{
@@ -2822,7 +2804,7 @@ void FOpenGLDynamicRHI::RHIDispatchComputeShader(uint32 ThreadGroupCountX, uint3
 	GPUProfilingData.RegisterGPUDispatch(FIntVector(ThreadGroupCountX, ThreadGroupCountY, ThreadGroupCountZ));	
 
 	BindPendingComputeShaderState(ContextState, ComputeShader);
-	CommitComputeResourceTables(ComputeShader);
+	CommitComputeResourceTables(ComputeShaderProxy);
 	SetupTexturesForDraw(ContextState, ComputeShader, FOpenGL::GetMaxComputeTextureImageUnits());
 	SetupUAVsForCompute(ContextState, ComputeShader);
 	CommitComputeShaderConstants(ComputeShader);
diff --git a/Engine/Source/Runtime/OpenGLDrv/Private/OpenGLShaders.cpp b/Engine/Source/Runtime/OpenGLDrv/Private/OpenGLShaders.cpp
--- a/Engine/Source/Runtime/OpenGLDrv/Private/OpenGLShaders.cpp
+++ b/Engine/Source/Runtime/OpenGLDrv/Private/OpenGLShaders.cpp
@@ -1036,7 +1036,7 @@ RHIType* CreateProxyShader(TArrayView<const uint8> Code, const FSHAHash& Hash)
 #if !(UE_BUILD_SHIPPING || UE_BUILD_TEST)
 	if (RHIShader)
 	{
-		RHIShader->ShaderName = ShaderCode.FindOptionalData(FShaderCodeName::Key);
+		RHIShader->Debug.ShaderName = ShaderCode.FindOptionalData(FShaderCodeName::Key);
 	}
 #endif
 
diff --git a/Engine/Source/Runtime/OpenGLDrv/Private/OpenGLUniformBuffer.cpp b/Engine/Source/Runtime/OpenGLDrv/Private/OpenGLUniformBuffer.cpp
--- a/Engine/Source/Runtime/OpenGLDrv/Private/OpenGLUniformBuffer.cpp
+++ b/Engine/Source/Runtime/OpenGLDrv/Private/OpenGLUniformBuffer.cpp
@@ -481,19 +481,19 @@ FOpenGLUniformBuffer::~FOpenGLUniformBuffer()
 }
 
 
-static void SetLayoutTable(FOpenGLUniformBuffer* NewUniformBuffer, const void* Contents, const FRHIUniformBufferLayout* Layout, EUniformBufferValidation Validation)
+void FOpenGLUniformBuffer::SetLayoutTable(const void* Contents, EUniformBufferValidation Validation)
 {
-	if (Layout->Resources.Num())
+	if (GetLayout().Resources.Num())
 	{
-		int32 NumResources = Layout->Resources.Num();
-		NewUniformBuffer->ResourceTable.Empty(NumResources);
-		NewUniformBuffer->ResourceTable.AddZeroed(NumResources);
+		int32 NumResources = GetLayout().Resources.Num();
+		ResourceTable.Empty(NumResources);
+		ResourceTable.AddZeroed(NumResources);
 
 		if (Contents)
 		{
 			for (int32 Index = 0; Index < NumResources; ++Index)
 			{
-				NewUniformBuffer->ResourceTable[Index] = GetShaderParameterResourceRHI(Contents, Layout->Resources[Index].MemberOffset, Layout->Resources[Index].MemberType);
+				ResourceTable[Index] = GetShaderParameterResourceRHI(Contents, GetLayout().Resources[Index].MemberOffset, GetLayout().Resources[Index].MemberType);
 			}
 		}
 	}
@@ -638,7 +638,7 @@ static FUniformBufferRHIRef CreateUniformBuffer(const void* Contents, const FRHI
 	}
 
 	// Initialize the resource table for this uniform buffer.
-	SetLayoutTable(NewUniformBuffer, Contents, Layout, Validation);
+	NewUniformBuffer->SetLayoutTable(Contents, Validation);
 
 	return NewUniformBuffer;
 }	
@@ -755,7 +755,7 @@ FUniformBufferRHIRef FOpenGLDynamicRHI::RHICreateUniformBuffer(const void* Conte
 	}
 
 	// Initialize the resource table for this uniform buffer.
-	SetLayoutTable(NewUniformBuffer, Contents, Layout, Validation);
+	NewUniformBuffer->SetLayoutTable(Contents, Validation);
 
 	return NewUniformBuffer;
 }
@@ -793,7 +793,7 @@ void FOpenGLDynamicRHI::RHIUpdateUniformBuffer(FRHICommandListBase& RHICmdList,
 	const int32 ConstantBufferSize = Layout.ConstantBufferSize;
 	const int32 NumResources = Layout.Resources.Num();
 
-	check(UniformBuffer->ResourceTable.Num() == NumResources);
+	check(UniformBuffer->GetResourceTable().Num() == NumResources);
 
 	uint32 NextUniqueID = UniqueUniformBufferID();
 
@@ -803,7 +803,7 @@ void FOpenGLDynamicRHI::RHIUpdateUniformBuffer(FRHICommandListBase& RHICmdList,
 
 		for (int32 Index = 0; Index < NumResources; ++Index)
 		{
-			UniformBuffer->ResourceTable[Index] = GetShaderParameterResourceRHI(Contents, Layout.Resources[Index].MemberOffset, Layout.Resources[Index].MemberType);
+			UniformBuffer->GetResourceTable()[Index] = GetShaderParameterResourceRHI(Contents, Layout.Resources[Index].MemberOffset, Layout.Resources[Index].MemberType);
 		}
 		
 		UniformBuffer->UniqueID = NextUniqueID;
@@ -837,7 +837,7 @@ void FOpenGLDynamicRHI::RHIUpdateUniformBuffer(FRHICommandListBase& RHICmdList,
 			// Update resource table.
 			for (int32 ResourceIndex = 0; ResourceIndex < NumResources; ++ResourceIndex)
 			{
-				UniformBuffer->ResourceTable[ResourceIndex] = CmdListResources[ResourceIndex];
+				UniformBuffer->GetResourceTable()[ResourceIndex] = CmdListResources[ResourceIndex];
 			}
 			UniformBuffer->UniqueID = NextUniqueID;
 		});
diff --git a/Engine/Source/Runtime/OpenGLDrv/Public/OpenGLDrv.h b/Engine/Source/Runtime/OpenGLDrv/Public/OpenGLDrv.h
--- a/Engine/Source/Runtime/OpenGLDrv/Public/OpenGLDrv.h
+++ b/Engine/Source/Runtime/OpenGLDrv/Public/OpenGLDrv.h
@@ -999,7 +999,7 @@ private:
 	void UpdateScissorRectInOpenGLContext( FOpenGLContextState& ContextState );
 	void UpdateViewportInOpenGLContext( FOpenGLContextState& ContextState );
 	
-	template <class ShaderType> void SetResourcesFromTables(const ShaderType* RESTRICT);
+	template <class ShaderType> void SetResourcesFromTables(ShaderType* RESTRICT);
 	FORCEINLINE void CommitGraphicsResourceTables()
 	{
 		if (PendingState.bAnyDirtyGraphicsUniformBuffers)
@@ -1008,7 +1008,7 @@ private:
 		}
 	}
 	void CommitGraphicsResourceTablesInner();
-	void CommitComputeResourceTables(FOpenGLComputeShader* ComputeShader);
+	void CommitComputeResourceTables(FOpenGLComputeShaderProxy* ComputeShader);
 	void CommitNonComputeShaderConstants();
 	void CommitComputeShaderConstants(FOpenGLComputeShader* ComputeShader);
 	void SetPendingBlendStateForActiveRenderTargets( FOpenGLContextState& ContextState );
diff --git a/Engine/Source/Runtime/OpenGLDrv/Public/OpenGLResources.h b/Engine/Source/Runtime/OpenGLDrv/Public/OpenGLResources.h
--- a/Engine/Source/Runtime/OpenGLDrv/Public/OpenGLResources.h
+++ b/Engine/Source/Runtime/OpenGLDrv/Public/OpenGLResources.h
@@ -949,9 +949,6 @@ public:
 	/** Unique ID for state shadowing purposes. */
 	uint32 UniqueID;
 
-	/** Resource table containing RHI references. */
-	TArray<TRefCountPtr<FRHIResource> > ResourceTable;
-
 	/** Emulated uniform data for ES2. */
 	FOpenGLEUniformBufferDataRef EmulatedBufferData;
 
@@ -974,6 +971,12 @@ public:
 
 	FOpenGLAssertRHIThreadFence AccessFence;
 	FOpenGLAssertRHIThreadFence CopyFence;
+
+	// Provides public non-const access to ResourceTable.
+	// @todo refactor uniform buffers to perform updates as a member function, so this isn't necessary.
+	TArray<TRefCountPtr<FRHIResource>>& GetResourceTable() { return ResourceTable; }
+
+	void SetLayoutTable(const void* Contents, EUniformBufferValidation Validation);
 };
 
 typedef TOpenGLBuffer<FOpenGLBasePixelBuffer, CachedBindPixelUnpackBuffer> FOpenGLPixelBuffer;
@@ -1144,19 +1147,22 @@ public:
 	int32 MaxTextureStageUsed();
 	bool RequiresDriverInstantiation();
 
-	FOpenGLVertexShader* GetVertexShader()
+	FOpenGLVertexShaderProxy* GetVertexShader()
 	{
 		check(IsValidRef(VertexShaderProxy));
-		return VertexShaderProxy->GetGLResourceObject();
+		return VertexShaderProxy;
 	}
 
-	FOpenGLPixelShader* GetPixelShader()
+	FOpenGLPixelShaderProxy* GetPixelShader()
 	{
 		check(IsValidRef(PixelShaderProxy));
-		return PixelShaderProxy->GetGLResourceObject();
+		return PixelShaderProxy;
 	}
 
-	FOpenGLGeometryShader* GetGeometryShader()	{ return GeometryShaderProxy ? GeometryShaderProxy->GetGLResourceObject() : nullptr;}
+	FOpenGLGeometryShaderProxy* GetGeometryShader()
+	{
+		return GeometryShaderProxy;
+	}
 
 	virtual ~FOpenGLBoundShaderState();
 };
diff --git a/Engine/Source/Runtime/OpenGLDrv/Public/OpenGLShaderResources.h b/Engine/Source/Runtime/OpenGLDrv/Public/OpenGLShaderResources.h
--- a/Engine/Source/Runtime/OpenGLDrv/Public/OpenGLShaderResources.h
+++ b/Engine/Source/Runtime/OpenGLDrv/Public/OpenGLShaderResources.h
@@ -33,35 +33,6 @@ enum
 	OGL_UAV_NOT_SUPPORTED_FOR_GRAPHICS_UNIT = -1, // for now, only CS and PS supports UAVs/ images
 };
 
-struct FOpenGLShaderResourceTable : public FBaseShaderResourceTable
-{
-	/** Mapping of bound Textures to their location in resource tables. */
-	TArray<uint32> TextureMap;
-	friend bool operator==(const FOpenGLShaderResourceTable &A, const FOpenGLShaderResourceTable& B)
-	{
-		if (!(((FBaseShaderResourceTable&)A) == ((FBaseShaderResourceTable&)B)))
-		{
-			return false;
-		}
-		if (A.TextureMap.Num() != B.TextureMap.Num())
-		{
-			return false;
-		}
-		if (FMemory::Memcmp(A.TextureMap.GetData(), B.TextureMap.GetData(), A.TextureMap.GetTypeSize()*A.TextureMap.Num()) != 0)
-		{
-			return false;
-		}
-		return true;
-	}
-};
-
-inline FArchive& operator<<(FArchive& Ar, FOpenGLShaderResourceTable& SRT)
-{
-	Ar << ((FBaseShaderResourceTable&)SRT);
-	Ar << SRT.TextureMap;
-	return Ar;
-}
-
 struct FOpenGLShaderVarying
 {
 	TArray<ANSICHAR> Varying;
@@ -100,7 +71,7 @@ struct FOpenGLShaderBindings
 	TArray<CrossCompiler::FPackedArrayInfo>			PackedGlobalArrays;
 	TArray<FOpenGLShaderVarying>					InputVaryings;
 	TArray<FOpenGLShaderVarying>					OutputVaryings;
-	FOpenGLShaderResourceTable						ShaderResourceTable;
+	FShaderResourceTable							ShaderResourceTable;
 	CrossCompiler::FShaderBindingInOutMask			InOutMask;
 
 	uint8	NumSamplers;
diff --git a/Engine/Source/Runtime/RHI/Private/PipelineStateCache.cpp b/Engine/Source/Runtime/RHI/Private/PipelineStateCache.cpp
--- a/Engine/Source/Runtime/RHI/Private/PipelineStateCache.cpp
+++ b/Engine/Source/Runtime/RHI/Private/PipelineStateCache.cpp
@@ -199,23 +199,23 @@ static void HandlePipelineCreationFailure(const FGraphicsPipelineStateInitialize
 #if !(UE_BUILD_SHIPPING || UE_BUILD_TEST)
 	if(Init.BoundShaderState.VertexShaderRHI)
 	{
-		UE_LOG(LogRHI, Error, TEXT("Vertex: %s"), *Init.BoundShaderState.VertexShaderRHI->ShaderName);
+		UE_LOG(LogRHI, Error, TEXT("Vertex: %s"), Init.BoundShaderState.VertexShaderRHI->GetShaderName());
 	}
 	if (Init.BoundShaderState.GetMeshShader())
 	{
-		UE_LOG(LogRHI, Error, TEXT("Mesh: %s"), *Init.BoundShaderState.GetMeshShader()->ShaderName);
+		UE_LOG(LogRHI, Error, TEXT("Mesh: %s"), Init.BoundShaderState.GetMeshShader()->GetShaderName());
 	}
 	if (Init.BoundShaderState.GetAmplificationShader())
 	{
-		UE_LOG(LogRHI, Error, TEXT("Amplification: %s"), *Init.BoundShaderState.GetAmplificationShader()->ShaderName);
+		UE_LOG(LogRHI, Error, TEXT("Amplification: %s"), Init.BoundShaderState.GetAmplificationShader()->GetShaderName());
 	}
 	if(Init.BoundShaderState.GetGeometryShader())
 	{
-		UE_LOG(LogRHI, Error, TEXT("Geometry: %s"), *Init.BoundShaderState.GetGeometryShader()->ShaderName);
+		UE_LOG(LogRHI, Error, TEXT("Geometry: %s"), Init.BoundShaderState.GetGeometryShader()->GetShaderName());
 	}
 	if(Init.BoundShaderState.PixelShaderRHI)
 	{
-		UE_LOG(LogRHI, Error, TEXT("Pixel: %s"), *Init.BoundShaderState.PixelShaderRHI->ShaderName);
+		UE_LOG(LogRHI, Error, TEXT("Pixel: %s"), Init.BoundShaderState.PixelShaderRHI->GetShaderName());
 	}
 	
 	UE_LOG(LogRHI, Error, TEXT("Render Targets: (%u)"), Init.RenderTargetFormats.Num());
diff --git a/Engine/Source/Runtime/RHI/Public/RHIResources.h b/Engine/Source/Runtime/RHI/Public/RHIResources.h
--- a/Engine/Source/Runtime/RHI/Public/RHIResources.h
+++ b/Engine/Source/Runtime/RHI/Public/RHIResources.h
@@ -580,11 +580,33 @@ public:
 	FSHAHash GetHash() const { return Hash; }
 
 #if RHI_INCLUDE_SHADER_DEBUG_DATA
+
 	// for debugging only e.g. MaterialName:ShaderFile.usf or ShaderFile.usf/EntryFunc
-	FString ShaderName;
-	FORCEINLINE const TCHAR* GetShaderName() const { return *ShaderName; }
+	struct
+	{
+		FString ShaderName;
+		TArray<FName> UniformBufferNames;
+	} Debug;
+
+	const TCHAR* GetShaderName() const
+	{
+		return Debug.ShaderName.Len()
+			? *Debug.ShaderName
+			: TEXT("<unknown>");
+	}
+
+	FString GetUniformBufferName(uint32 Index) const
+	{
+		return Debug.UniformBufferNames.IsValidIndex(Index)
+			? Debug.UniformBufferNames[Index].ToString()
+			: TEXT("<unknown>");
+	}
+
 #else
-	FORCEINLINE const TCHAR* GetShaderName() const { return TEXT(""); }
+
+	const TCHAR* GetShaderName() const { return TEXT("<unknown>"); }
+	FString GetUniformBufferName(uint32 Index) const { return TEXT("<unknown>"); }
+
 #endif
 
 	FRHIShader() = delete;
@@ -1064,6 +1086,11 @@ public:
 	mutable int32 NumMeshCommandReferencesForDebugging = 0;
 #endif
 
+	const TArray<TRefCountPtr<FRHIResource>>& GetResourceTable() const { return ResourceTable; }
+
+protected:
+	TArray<TRefCountPtr<FRHIResource>> ResourceTable;
+
 private:
 	/** Layout of the uniform buffer. */
 	TRefCountPtr<const FRHIUniformBufferLayout> Layout;
diff --git a/Engine/Source/Runtime/RHICore/Public/RHICoreShader.h b/Engine/Source/Runtime/RHICore/Public/RHICoreShader.h
--- a/Engine/Source/Runtime/RHICore/Public/RHICoreShader.h
+++ b/Engine/Source/Runtime/RHICore/Public/RHICoreShader.h
@@ -13,7 +13,7 @@ namespace RHICore
 /** Validates that the uniform buffer at the requested static slot. */
 extern RHICORE_API void ValidateStaticUniformBuffer(FRHIUniformBuffer* UniformBuffer, FUniformBufferStaticSlot Slot, uint32 ExpectedHash);
 
-inline void InitStaticUniformBufferSlots(TArray<FUniformBufferStaticSlot>& StaticSlots, const FBaseShaderResourceTable& ShaderResourceTable)
+inline void InitStaticUniformBufferSlots(TArray<FUniformBufferStaticSlot>& StaticSlots, const FShaderResourceTable& ShaderResourceTable)
 {
 	StaticSlots.Reserve(ShaderResourceTable.ResourceTableLayoutHashes.Num());
 
@@ -58,5 +58,175 @@ void ApplyStaticUniformBuffers(
 	}
 }
 
+template <typename TResourceType> struct TResourceTypeStr {};
+template<> struct TResourceTypeStr<FRHISamplerState       > { static constexpr TCHAR String[] = TEXT("Sampler State"); };
+template<> struct TResourceTypeStr<FRHITexture            > { static constexpr TCHAR String[] = TEXT("Texture"); };
+template<> struct TResourceTypeStr<FRHIShaderResourceView > { static constexpr TCHAR String[] = TEXT("Shader Resource View"); };
+template<> struct TResourceTypeStr<FRHIUnorderedAccessView> { static constexpr TCHAR String[] = TEXT("Unordered Access View"); };
+
+template <typename TResourceType, typename TCallback>
+inline void EnumerateUniformBufferResources(FRHIUniformBuffer* RESTRICT Buffer, int32 BufferIndex, const uint32* RESTRICT ResourceMap, TCallback&& Callback)
+{
+	const TRefCountPtr<FRHIResource>* RESTRICT Resources = Buffer->GetResourceTable().GetData();
+
+	uint32 BufferOffset = ResourceMap[BufferIndex];
+	if (BufferOffset > 0)
+	{
+		const uint32* RESTRICT ResourceInfos = &ResourceMap[BufferOffset];
+		uint32 ResourceInfo = *ResourceInfos++;
+		do
+		{
+			checkSlow(FRHIResourceTableEntry::GetUniformBufferIndex(ResourceInfo) == BufferIndex);
+
+			const uint16 ResourceIndex = FRHIResourceTableEntry::GetResourceIndex(ResourceInfo);
+			const uint8  BindIndex     = FRHIResourceTableEntry::GetBindIndex    (ResourceInfo);
+
+			TResourceType* Resource = static_cast<TResourceType*>(Resources[ResourceIndex].GetReference());
+			checkf(Resource
+				, TEXT("Null %s (resource %d bind %d) on UB Layout %s")
+				, TResourceTypeStr<TResourceType>::String
+				, ResourceIndex
+				, BindIndex
+				, *Buffer->GetLayout().GetDebugName()
+			);
+
+			Callback(Resource, BindIndex);
+
+			ResourceInfo = *ResourceInfos++;
+		} while (FRHIResourceTableEntry::GetUniformBufferIndex(ResourceInfo) == BufferIndex);
+	}
+}
+
+template <typename TBinder, typename TUniformBufferArrayType, typename TBitMaskType>
+void SetResourcesFromTables(TBinder&& Binder, FRHIShader const& Shader, FShaderResourceTable const& SRT, TBitMaskType& DirtyUniformBuffers, TUniformBufferArrayType const& BoundUniformBuffers
+#if ENABLE_RHI_VALIDATION
+	, RHIValidation::FTracker* Tracker
+#endif
+)
+{
+	float CurrentTimeForTextureTimes = FApp::GetCurrentTime();
+
+	// Mask the dirty bits by those buffers from which the shader has bound resources.
+	uint32 DirtyBits = SRT.ResourceTableBits & DirtyUniformBuffers;
+	while (DirtyBits)
+	{
+		// Scan for the lowest set bit, compute its index, clear it in the set of dirty bits.
+		const uint32 LowestBitMask = (DirtyBits) & (-(int32)DirtyBits);
+		const int32 BufferIndex = FMath::CountTrailingZeros(LowestBitMask); // todo: This has a branch on zero, we know it could never be zero...
+		DirtyBits ^= LowestBitMask;
+
+		check(BufferIndex < SRT.ResourceTableLayoutHashes.Num());
+
+		FRHIUniformBuffer* Buffer = BoundUniformBuffers[BufferIndex];
+
+#if DO_CHECK
+
+		if (!Buffer)
+		{
+			UE_LOG(LogRHICore, Fatal, TEXT("Shader expected a uniform buffer at slot %u but got null instead (Shader='%s' UB='%s'). Rendering code needs to set a valid uniform buffer for this slot.")
+				, BufferIndex
+				, Shader.GetShaderName()
+				, *Shader.GetUniformBufferName(BufferIndex)
+			);
+		}
+		else if (Buffer->GetLayout().GetHash() != SRT.ResourceTableLayoutHashes[BufferIndex])
+		{
+			FRHIUniformBufferLayout const& BufferLayout = Buffer->GetLayout();
+
+			FString ResourcesString;
+			for (FRHIUniformBufferResource const& Resource : BufferLayout.Resources)
+			{
+				ResourcesString += FString::Printf(TEXT("%s%d")
+					, ResourcesString.Len() ? TEXT(" ") : TEXT("")
+					, Resource.MemberType
+				);
+			}
+
+			// This might mean you are accessing a data you haven't bound e.g. GBuffer
+			UE_LOG(LogRHICore, Fatal,
+				TEXT("Uniform buffer bound to slot %u is not what the shader expected:\n")
+				TEXT("\tBound                : Uniform Buffer[%s] with Hash[0x%08x]\n")
+				TEXT("\tExpected             : Uniform Buffer[%s] with Hash[0x%08x]\n")
+				TEXT("\tShader Name          : %s\n")
+				TEXT("\tLayout CB Size       : %d\n")
+				TEXT("\tLayout Num Resources : %d\n")
+				TEXT("\tResource Types       : %s\n")
+				, BufferIndex
+				, *BufferLayout.GetDebugName(), BufferLayout.GetHash()
+				, *Shader.GetUniformBufferName(BufferIndex), SRT.ResourceTableLayoutHashes[BufferIndex]
+				, Shader.GetShaderName()
+				, BufferLayout.ConstantBufferSize
+				, BufferLayout.Resources.Num()
+				, *ResourcesString
+			);
+		}
+
+#endif // DO_CHECK
+
+		// Samplers
+		EnumerateUniformBufferResources<FRHISamplerState>(Buffer, BufferIndex, SRT.SamplerMap.GetData(),
+			[&](FRHISamplerState* Sampler, uint8 Index)
+			{
+				Binder.SetSampler(Sampler, Index);
+			});
+
+		// Textures
+		EnumerateUniformBufferResources<FRHITexture>(Buffer, BufferIndex, SRT.TextureMap.GetData(),
+			[&](FRHITexture* Texture, uint8 Index)
+			{
+#if ENABLE_RHI_VALIDATION
+				if (Tracker)
+				{
+					ERHIAccess Access = Shader.GetFrequency() == SF_Compute
+						? ERHIAccess::SRVCompute
+						: ERHIAccess::SRVGraphics;
+
+					// Textures bound here only have their "common" plane accessible. Stencil etc is ignored.
+					// (i.e. only access the color plane of a color texture, or depth plane of a depth texture)
+					Tracker->Assert(Texture->GetViewIdentity(0, 0, 0, 0, uint32(RHIValidation::EResourcePlane::Common), 1), Access);
+				}
+#endif
+				Texture->SetLastRenderTime(CurrentTimeForTextureTimes);
+				Binder.SetTexture(Texture, Index);
+			});
+
+		// SRVs
+		EnumerateUniformBufferResources<FRHIShaderResourceView>(Buffer, BufferIndex, SRT.ShaderResourceViewMap.GetData(),
+			[&](FRHIShaderResourceView* SRV, uint8 Index)
+			{
+#if ENABLE_RHI_VALIDATION
+				if (Tracker)
+				{
+					ERHIAccess Access = Shader.GetFrequency() == SF_Compute
+						? ERHIAccess::SRVCompute
+						: ERHIAccess::SRVGraphics;
+
+					Tracker->Assert(SRV->ViewIdentity, Access);
+				}
+#endif
+				Binder.SetSRV(SRV, Index);
+			});
+
+		// UAVs
+		EnumerateUniformBufferResources<FRHIUnorderedAccessView>(Buffer, BufferIndex, SRT.UnorderedAccessViewMap.GetData(),
+			[&](FRHIUnorderedAccessView* UAV, uint8 Index)
+			{
+#if ENABLE_RHI_VALIDATION
+				if (Tracker)
+				{
+					ERHIAccess Access = Shader.GetFrequency() == SF_Compute
+						? ERHIAccess::UAVCompute
+						: ERHIAccess::UAVGraphics;
+
+					Tracker->AssertUAV(UAV, Access, Index);
+				}
+#endif
+				Binder.SetUAV(UAV, Index);
+			});
+	}
+
+	DirtyUniformBuffers = TBitMaskType(0);
+}
+
 } //! RHICore
 } //! UE
\ No newline at end of file
diff --git a/Engine/Source/Runtime/RenderCore/Public/ShaderCore.h b/Engine/Source/Runtime/RenderCore/Public/ShaderCore.h
--- a/Engine/Source/Runtime/RenderCore/Public/ShaderCore.h
+++ b/Engine/Source/Runtime/RenderCore/Public/ShaderCore.h
@@ -378,10 +378,10 @@ private:
 	TMap<FString,FString> Definitions;
 };
 
-struct FBaseShaderResourceTable
+struct FShaderResourceTable
 {
 	/** Bits indicating which resource tables contain resources bound to this shader. */
-	uint32 ResourceTableBits;
+	uint32 ResourceTableBits = 0;
 
 	/** Mapping of bound SRVs to their location in resource tables. */
 	TArray<uint32> ShaderResourceViewMap;
@@ -395,38 +395,41 @@ struct FBaseShaderResourceTable
 	/** Hash of the layouts of resource tables at compile time, used for runtime validation. */
 	TArray<uint32> ResourceTableLayoutHashes;
 
-	FBaseShaderResourceTable() :
-		ResourceTableBits(0)
-	{
-	}
+	/** Mapping of bound Textures to their location in resource tables. */
+	TArray<uint32> TextureMap;
 
-	friend bool operator==(const FBaseShaderResourceTable &A, const FBaseShaderResourceTable& B)
+	friend bool operator==(const FShaderResourceTable&A, const FShaderResourceTable& B)
 	{
 		bool bEqual = true;
 		bEqual &= (A.ResourceTableBits == B.ResourceTableBits);
-		bEqual &= (A.ShaderResourceViewMap.Num() == B.ShaderResourceViewMap.Num());
-		bEqual &= (A.SamplerMap.Num() == B.SamplerMap.Num());
-		bEqual &= (A.UnorderedAccessViewMap.Num() == B.UnorderedAccessViewMap.Num());
+		bEqual &= (A.ShaderResourceViewMap    .Num() == B.ShaderResourceViewMap    .Num());
+		bEqual &= (A.SamplerMap               .Num() == B.SamplerMap               .Num());
+		bEqual &= (A.UnorderedAccessViewMap   .Num() == B.UnorderedAccessViewMap   .Num());
 		bEqual &= (A.ResourceTableLayoutHashes.Num() == B.ResourceTableLayoutHashes.Num());
+		bEqual &= (A.TextureMap               .Num() == B.TextureMap               .Num());
+
 		if (!bEqual)
 		{
 			return false;
 		}
-		bEqual &= (FMemory::Memcmp(A.ShaderResourceViewMap.GetData(), B.ShaderResourceViewMap.GetData(), A.ShaderResourceViewMap.GetTypeSize()*A.ShaderResourceViewMap.Num()) == 0);
-		bEqual &= (FMemory::Memcmp(A.SamplerMap.GetData(), B.SamplerMap.GetData(), A.SamplerMap.GetTypeSize()*A.SamplerMap.Num()) == 0);
-		bEqual &= (FMemory::Memcmp(A.UnorderedAccessViewMap.GetData(), B.UnorderedAccessViewMap.GetData(), A.UnorderedAccessViewMap.GetTypeSize()*A.UnorderedAccessViewMap.Num()) == 0);
-		bEqual &= (FMemory::Memcmp(A.ResourceTableLayoutHashes.GetData(), B.ResourceTableLayoutHashes.GetData(), A.ResourceTableLayoutHashes.GetTypeSize()*A.ResourceTableLayoutHashes.Num()) == 0);
+
+		bEqual &= (FMemory::Memcmp(A.ShaderResourceViewMap    .GetData(), B.ShaderResourceViewMap    .GetData(), A.ShaderResourceViewMap    .GetTypeSize() * A.ShaderResourceViewMap    .Num()) == 0);
+		bEqual &= (FMemory::Memcmp(A.SamplerMap               .GetData(), B.SamplerMap               .GetData(), A.SamplerMap               .GetTypeSize() * A.SamplerMap               .Num()) == 0);
+		bEqual &= (FMemory::Memcmp(A.UnorderedAccessViewMap   .GetData(), B.UnorderedAccessViewMap   .GetData(), A.UnorderedAccessViewMap   .GetTypeSize() * A.UnorderedAccessViewMap   .Num()) == 0);
+		bEqual &= (FMemory::Memcmp(A.ResourceTableLayoutHashes.GetData(), B.ResourceTableLayoutHashes.GetData(), A.ResourceTableLayoutHashes.GetTypeSize() * A.ResourceTableLayoutHashes.Num()) == 0);
+		bEqual &= (FMemory::Memcmp(A.TextureMap               .GetData(), B.TextureMap               .GetData(), A.TextureMap               .GetTypeSize() * A.TextureMap               .Num()) == 0);
 		return bEqual;
 	}
 };
 
-inline FArchive& operator<<(FArchive& Ar, FBaseShaderResourceTable& SRT)
+inline FArchive& operator<<(FArchive& Ar, FShaderResourceTable& SRT)
 {
 	Ar << SRT.ResourceTableBits;
 	Ar << SRT.ShaderResourceViewMap;
 	Ar << SRT.SamplerMap;
 	Ar << SRT.UnorderedAccessViewMap;
 	Ar << SRT.ResourceTableLayoutHashes;
+	Ar << SRT.TextureMap;
 
 	return Ar;
 }
diff --git a/Engine/Source/Runtime/VulkanRHI/Private/VulkanResources.h b/Engine/Source/Runtime/VulkanRHI/Private/VulkanResources.h
--- a/Engine/Source/Runtime/VulkanRHI/Private/VulkanResources.h
+++ b/Engine/Source/Runtime/VulkanRHI/Private/VulkanResources.h
@@ -1045,9 +1045,6 @@ public:
 	FVulkanDevice* Device;
 	VulkanRHI::FVulkanAllocation Allocation;
 	EUniformBufferUsage Usage;
-protected:
-	TArray<TRefCountPtr<FRHIResource>> ResourceTable;
-	
 };
 
 class FVulkanUnorderedAccessView : public FRHIUnorderedAccessView, public VulkanRHI::FVulkanViewBase
diff --git a/Engine/Source/Runtime/VulkanRHI/Private/VulkanShaders.cpp b/Engine/Source/Runtime/VulkanRHI/Private/VulkanShaders.cpp
--- a/Engine/Source/Runtime/VulkanRHI/Private/VulkanShaders.cpp
+++ b/Engine/Source/Runtime/VulkanRHI/Private/VulkanShaders.cpp
@@ -91,7 +91,7 @@ ShaderType* FVulkanShaderFactory::CreateShader(TArrayView<const uint8> Code, FVu
 
 #if !(UE_BUILD_SHIPPING || UE_BUILD_TEST)
 				FShaderCodeReader ShaderCode(Code);
-				RetShader->ShaderName = ShaderCode.FindOptionalData(FShaderCodeName::Key);
+				RetShader->Debug.ShaderName = ShaderCode.FindOptionalData(FShaderCodeName::Key);
 #endif
 			}
 		}
diff --git a/Engine/Source/Runtime/Windows/D3D11RHI/Private/D3D11Commands.cpp b/Engine/Source/Runtime/Windows/D3D11RHI/Private/D3D11Commands.cpp
--- a/Engine/Source/Runtime/Windows/D3D11RHI/Private/D3D11Commands.cpp
+++ b/Engine/Source/Runtime/Windows/D3D11RHI/Private/D3D11Commands.cpp
@@ -89,7 +89,7 @@ void FD3D11DynamicRHI::RHIBeginUpdateMultiFrameResource(FRHITexture* RHITexture)
 {
 	if (!IsRHIDeviceNVIDIA() || GNumAlternateFrameRenderingGroups == 1) return;
 
-	FD3D11Texture* Texture = GetD3D11TextureFromRHITexture(RHITexture);
+	FD3D11Texture* Texture = ResourceCast(RHITexture);
 
 	if (!Texture)
 	{
@@ -113,7 +113,7 @@ void FD3D11DynamicRHI::RHIEndUpdateMultiFrameResource(FRHITexture* RHITexture)
 {
 	if (!IsRHIDeviceNVIDIA() || GNumAlternateFrameRenderingGroups == 1) return;
 
-	FD3D11Texture* Texture = GetD3D11TextureFromRHITexture(RHITexture);
+	FD3D11Texture* Texture = ResourceCast(RHITexture);
 
 	if (!Texture || !Texture->GetIHVResourceHandle())
 	{
@@ -386,7 +386,7 @@ void FD3D11DynamicRHI::RHISetBoundShaderState(FRHIBoundShaderState* BoundShaderS
 
 void FD3D11DynamicRHI::RHISetShaderTexture(FRHIGraphicsShader* ShaderRHI,uint32 TextureIndex, FRHITexture* NewTextureRHI)
 {
-	FD3D11Texture* NewTexture = GetD3D11TextureFromRHITexture(NewTextureRHI);
+	FD3D11Texture* NewTexture = ResourceCast(NewTextureRHI);
 	ID3D11ShaderResourceView* ShaderResourceView = NewTexture ? NewTexture->GetShaderResourceView() : nullptr;
 
 	switch (ShaderRHI->GetFrequency())
@@ -421,7 +421,7 @@ void FD3D11DynamicRHI::RHISetShaderTexture(FRHIComputeShader* ComputeShaderRHI,u
 {
 	//VALIDATE_BOUND_SHADER(ComputeShaderRHI);
 
-	FD3D11Texture* NewTexture = GetD3D11TextureFromRHITexture(NewTextureRHI);
+	FD3D11Texture* NewTexture = ResourceCast(NewTextureRHI);
 	ID3D11ShaderResourceView* ShaderResourceView = NewTexture ? NewTexture->GetShaderResourceView() : nullptr;
 	SetShaderResourceView<SF_Compute>(NewTexture, ShaderResourceView, TextureIndex);
 }
@@ -954,7 +954,7 @@ void FD3D11DynamicRHI::SetRenderTargets(
 	const FRHIRenderTargetView* NewRenderTargetsRHI,
 	const FRHIDepthRenderTargetView* NewDepthStencilTargetRHI)
 {
-	FD3D11Texture* NewDepthStencilTarget = GetD3D11TextureFromRHITexture(NewDepthStencilTargetRHI ? NewDepthStencilTargetRHI->Texture : nullptr);
+	FD3D11Texture* NewDepthStencilTarget = ResourceCast(NewDepthStencilTargetRHI ? NewDepthStencilTargetRHI->Texture : nullptr);
 
 	check(NewNumSimultaneousRenderTargets <= MaxSimultaneousRenderTargets);
 
@@ -990,7 +990,7 @@ void FD3D11DynamicRHI::SetRenderTargets(
 			int32 RTMipIndex = NewRenderTargetsRHI[RenderTargetIndex].MipIndex;
 			int32 RTSliceIndex = NewRenderTargetsRHI[RenderTargetIndex].ArraySliceIndex;
 			
-			FD3D11Texture* NewRenderTarget = GetD3D11TextureFromRHITexture(NewRenderTargetsRHI[RenderTargetIndex].Texture);
+			FD3D11Texture* NewRenderTarget = ResourceCast(NewRenderTargetsRHI[RenderTargetIndex].Texture);
 			RenderTargetView = NewRenderTarget ? NewRenderTarget->GetRenderTargetView(RTMipIndex, RTSliceIndex) : nullptr;
 
 			ensureMsgf(RenderTargetView, TEXT("Texture being set as render target has no RTV"));
@@ -1179,323 +1179,73 @@ void FD3D11DynamicRHI::CommitComputeShaderConstants()
 	}
 }
 
-template <EShaderFrequency Frequency>
-FORCEINLINE void SetResource(FD3D11DynamicRHI* RESTRICT D3D11RHI, FD3D11StateCache* RESTRICT StateCache, uint32 BindIndex, FD3D11BaseShaderResource* RESTRICT ShaderResource, ID3D11ShaderResourceView* RESTRICT SRV)
-{
-	// We set the resource through the RHI to track state for the purposes of unbinding SRVs when a UAV or RTV is bound.
-	// todo: need to support SRV_Static for faster calls when possible
-	D3D11RHI->SetShaderResourceView<Frequency>(ShaderResource, SRV, BindIndex);
-}
-
-template <EShaderFrequency Frequency>
-FORCEINLINE void SetResource(FD3D11DynamicRHI* RESTRICT D3D11RHI, FD3D11StateCache* RESTRICT StateCache, uint32 BindIndex, ID3D11SamplerState* RESTRICT SamplerState)
-{
-	StateCache->SetSamplerState<Frequency>(SamplerState, BindIndex);
-}
-
-template <EShaderFrequency ShaderFrequency>
-inline int32 SetShaderResourcesFromBuffer_Surface(FD3D11DynamicRHI* RESTRICT D3D11RHI, FD3D11StateCache* RESTRICT StateCache, FD3D11UniformBuffer* RESTRICT Buffer, const uint32* RESTRICT ResourceMap, int32 BufferIndex, const TCHAR* LayoutName)
-{
-	const TRefCountPtr<FRHIResource>* RESTRICT Resources = Buffer->ResourceTable.GetData();
-	const int32 NumResourcesInTable = Buffer->ResourceTable.Num();
-	float CurrentTime = FApp::GetCurrentTime();
-	int32 NumSetCalls = 0;
-	uint32 BufferOffset = ResourceMap[BufferIndex];
-	if (BufferOffset > 0)
-	{
-		const uint32* RESTRICT ResourceInfos = &ResourceMap[BufferOffset];
-		uint32 ResourceInfo = *ResourceInfos++;
-		do
-		{
-			checkSlow(FRHIResourceTableEntry::GetUniformBufferIndex(ResourceInfo) == BufferIndex);
-			const uint16 ResourceIndex = FRHIResourceTableEntry::GetResourceIndex(ResourceInfo);
-			const uint8 BindIndex = FRHIResourceTableEntry::GetBindIndex(ResourceInfo);
-
-			check(ResourceIndex < NumResourcesInTable);
-			FRHITexture* TextureRHI = (FRHITexture*)Resources[ResourceIndex].GetReference();
-			if (!TextureRHI)
-			{
-				UE_LOG(LogD3D11RHI, Fatal, TEXT("Null texture (resource %d bind %d) on UB Layout %s"), ResourceIndex, BindIndex, LayoutName);
-			}
-			TextureRHI->SetLastRenderTime(CurrentTime);
-
-			FD3D11Texture* TextureD3D11 = GetD3D11TextureFromRHITexture(TextureRHI);
-			ID3D11ShaderResourceView* D3D11Resource = TextureD3D11->GetShaderResourceView();
-
-#if ENABLE_RHI_VALIDATION
-			if (D3D11RHI->Tracker)
-			{
-				constexpr ERHIAccess Access = ShaderFrequency == EShaderFrequency::SF_Compute
-					? ERHIAccess::SRVCompute
-					: ERHIAccess::SRVGraphics;
-
-				// Textures bound here only have their "common" plane accessible. Stencil etc is ignored.
-				// (i.e. only access the color plane of a color texture, or depth plane of a depth texture)
-				D3D11RHI->Tracker->Assert(TextureD3D11->GetViewIdentity(0, 0, 0, 0, uint32(RHIValidation::EResourcePlane::Common), 1), Access);
-			}
-#endif
-
-			// todo: could coalesce adjacent bound resources.
-			SetResource<ShaderFrequency>(D3D11RHI, StateCache, BindIndex, TextureD3D11, D3D11Resource);
-			NumSetCalls++;
-			ResourceInfo = *ResourceInfos++;
-		} while (FRHIResourceTableEntry::GetUniformBufferIndex(ResourceInfo) == BufferIndex);
-	}
-	return NumSetCalls;
-}
-
-
-template <EShaderFrequency ShaderFrequency>
-inline int32 SetShaderResourcesFromBufferUAVPS(FD3D11DynamicRHI* RESTRICT D3D11RHI, FD3D11StateCache* RESTRICT StateCache, FD3D11UniformBuffer* RESTRICT Buffer, const uint32* RESTRICT ResourceMap, int32 BufferIndex, const TCHAR* LayoutName)
-{
-	const TRefCountPtr<FRHIResource>* RESTRICT Resources = Buffer->ResourceTable.GetData();
-	int32 NumSetCalls = 0;
-	uint32 BufferOffset = ResourceMap[BufferIndex];
-	if (BufferOffset > 0)
-	{
-		const uint32* RESTRICT ResourceInfos = &ResourceMap[BufferOffset];
-		uint32 ResourceInfo = *ResourceInfos++;
-		do
-		{
-			checkSlow(FRHIResourceTableEntry::GetUniformBufferIndex(ResourceInfo) == BufferIndex);
-			const uint16 ResourceIndex = FRHIResourceTableEntry::GetResourceIndex(ResourceInfo);
-			const uint8 BindIndex = FRHIResourceTableEntry::GetBindIndex(ResourceInfo);
-
-			FD3D11UnorderedAccessView* UnorderedAccessViewRHI = (FD3D11UnorderedAccessView*)Resources[ResourceIndex].GetReference();
-			if (!UnorderedAccessViewRHI)
-			{
-				UE_LOG(LogD3D11RHI, Fatal, TEXT("Null UAV (resource %d bind %d) on UB Layout %s"), ResourceIndex, BindIndex, LayoutName);
-			}
-
-#if ENABLE_RHI_VALIDATION
-			if (D3D11RHI->Tracker)
-			{
-				D3D11RHI->Tracker->AssertUAV(UnorderedAccessViewRHI, RHIValidation::EUAVMode::Graphics, BindIndex);
-			}
-#endif
-
-			D3D11RHI->InternalSetUAVPS(BindIndex, UnorderedAccessViewRHI);
-			NumSetCalls++;
-			ResourceInfo = *ResourceInfos++;
-		} while (FRHIResourceTableEntry::GetUniformBufferIndex(ResourceInfo) == BufferIndex);
-	}
-	return NumSetCalls;
-}
-
-
-template <EShaderFrequency ShaderFrequency>
-inline int32 SetShaderResourcesFromBuffer_SRV(FD3D11DynamicRHI* RESTRICT D3D11RHI, FD3D11StateCache* RESTRICT StateCache, FD3D11UniformBuffer* RESTRICT Buffer, const uint32* RESTRICT ResourceMap, int32 BufferIndex, const TCHAR* LayoutName)
-{
-	const TRefCountPtr<FRHIResource>* RESTRICT Resources = Buffer->ResourceTable.GetData();
-	int32 NumSetCalls = 0;
-	uint32 BufferOffset = ResourceMap[BufferIndex];
-	if (BufferOffset > 0)
-	{
-		const uint32* RESTRICT ResourceInfos = &ResourceMap[BufferOffset];
-		uint32 ResourceInfo = *ResourceInfos++;
-		do
-		{
-			checkSlow(FRHIResourceTableEntry::GetUniformBufferIndex(ResourceInfo) == BufferIndex);
-			const uint16 ResourceIndex = FRHIResourceTableEntry::GetResourceIndex(ResourceInfo);
-			const uint8 BindIndex = FRHIResourceTableEntry::GetBindIndex(ResourceInfo);
-
-			FD3D11BaseShaderResource* ShaderResource = nullptr;
-			ID3D11ShaderResourceView* D3D11Resource = nullptr;
-
-			FD3D11ShaderResourceView* ShaderResourceViewRHI = (FD3D11ShaderResourceView*)Resources[ResourceIndex].GetReference();
-			if (!ShaderResourceViewRHI)
-			{
-				UE_LOG(LogD3D11RHI, Fatal, TEXT("Null SRV (resource %d bind %d) on UB Layout %s"), ResourceIndex, BindIndex, LayoutName);
-			}
-
-#if ENABLE_RHI_VALIDATION
-			if (D3D11RHI->Tracker)
-			{
-				constexpr ERHIAccess Access = ShaderFrequency == EShaderFrequency::SF_Compute
-					? ERHIAccess::SRVCompute
-					: ERHIAccess::SRVGraphics;
-
-				D3D11RHI->Tracker->Assert(ShaderResourceViewRHI->ViewIdentity, Access);
-			}
-#endif
-
-			ShaderResource = ShaderResourceViewRHI->Resource.GetReference();
-			D3D11Resource = ShaderResourceViewRHI->View.GetReference();
-
-			// todo: could coalesce adjacent bound resources.
-			SetResource<ShaderFrequency>(D3D11RHI, StateCache, BindIndex, ShaderResource, D3D11Resource);
-			NumSetCalls++;
-			ResourceInfo = *ResourceInfos++;
-		} while (FRHIResourceTableEntry::GetUniformBufferIndex(ResourceInfo) == BufferIndex);
-	}
-	return NumSetCalls;
-}
-
-template <EShaderFrequency ShaderFrequency>
-inline int32 SetShaderResourcesFromBuffer_Sampler(FD3D11DynamicRHI* RESTRICT D3D11RHI, FD3D11StateCache* RESTRICT StateCache, FD3D11UniformBuffer* RESTRICT Buffer, const uint32* RESTRICT ResourceMap, int32 BufferIndex)
-{
-	const TRefCountPtr<FRHIResource>* RESTRICT Resources = Buffer->ResourceTable.GetData();
-	int32 NumSetCalls = 0;
-	uint32 BufferOffset = ResourceMap[BufferIndex];
-	if (BufferOffset > 0)
-	{
-		const uint32* RESTRICT ResourceInfos = &ResourceMap[BufferOffset];
-		uint32 ResourceInfo = *ResourceInfos++;
-		do
-		{
-			checkSlow(FRHIResourceTableEntry::GetUniformBufferIndex(ResourceInfo) == BufferIndex);
-			const uint16 ResourceIndex = FRHIResourceTableEntry::GetResourceIndex(ResourceInfo);
-			const uint8 BindIndex = FRHIResourceTableEntry::GetBindIndex(ResourceInfo);
-
-			ID3D11SamplerState* D3D11Resource = ((FD3D11SamplerState*)Resources[ResourceIndex].GetReference())->Resource.GetReference();
-
-			// todo: could coalesce adjacent bound resources.
-			SetResource<ShaderFrequency>(D3D11RHI, StateCache, BindIndex, D3D11Resource);
-			NumSetCalls++;
-			ResourceInfo = *ResourceInfos++;
-		} while (FRHIResourceTableEntry::GetUniformBufferIndex(ResourceInfo) == BufferIndex);
-	}
-	return NumSetCalls;
-}
-
-
 template <class ShaderType>
 void FD3D11DynamicRHI::SetResourcesFromTables(const ShaderType* RESTRICT Shader)
 {
 	checkSlow(Shader);
+	static constexpr EShaderFrequency Frequency = static_cast<EShaderFrequency>(ShaderType::StaticFrequency);
 
-	// Mask the dirty bits by those buffers from which the shader has bound resources.
-	uint32 DirtyBits = Shader->ShaderResourceTable.ResourceTableBits & DirtyUniformBuffers[ShaderType::StaticFrequency];
-	while (DirtyBits)
+	struct FUniformResourceBinder
 	{
-		// Scan for the lowest set bit, compute its index, clear it in the set of dirty bits.
-		const uint32 LowestBitMask = (DirtyBits) & (-(int32)DirtyBits);
-		const int32 BufferIndex = FMath::FloorLog2(LowestBitMask); // todo: This has a branch on zero, we know it could never be zero...
-		DirtyBits ^= LowestBitMask;
-		FD3D11UniformBuffer* Buffer = (FD3D11UniformBuffer*)BoundUniformBuffers[ShaderType::StaticFrequency][BufferIndex].GetReference();
-		
-		check(BufferIndex < Shader->ShaderResourceTable.ResourceTableLayoutHashes.Num());
+		FD3D11DynamicRHI& RHI;
 
-		if (!Buffer)
+		void SetUAV(FRHIUnorderedAccessView* UAV, uint8 Index)
 		{
-			FString ShaderUB;
-#if RHI_INCLUDE_SHADER_DEBUG_DATA
-			if (BufferIndex < Shader->UniformBuffers.Num())
-			{
-				Shader->UniformBuffers[BufferIndex].ToString(ShaderUB);
-			}
-#endif
-			UE_LOG(LogD3D11RHI, Fatal, TEXT("Shader expected a uniform buffer at slot %u but got null instead (Shader='%s' UB='%s').  Rendering code needs to set a valid uniform buffer for this slot."), BufferIndex, Shader->GetShaderName(), *ShaderUB);
+			RHI.InternalSetUAVPS(Index, ResourceCast(UAV));
 		}
 
-#if RHI_INCLUDE_SHADER_DEBUG_DATA
-		// to track down OR-7159 CRASH: Client crashed at start of match in D3D11Commands.cpp
+		void SetSRV(FRHIShaderResourceView* SRV, uint8 Index)
 		{
-			const uint32 LayoutHash = Buffer->GetLayout().GetHash();
-
-			if (LayoutHash != Shader->ShaderResourceTable.ResourceTableLayoutHashes[BufferIndex])
-			{
-				auto& BufferLayout = Buffer->GetLayout();
-				const auto& DebugName = BufferLayout.GetDebugName();
-				const FString& ShaderName = Shader->ShaderName;
-#if UE_BUILD_DEBUG
-				FString ShaderUB;
-				if (BufferIndex < Shader->UniformBuffers.Num())
-				{
-					ShaderUB = FString::Printf(TEXT("expecting UB '%s'"), *Shader->UniformBuffers[BufferIndex].GetPlainNameString());
-				}
-				UE_LOG(LogD3D11RHI, Error, TEXT("SetResourcesFromTables upcoming check(%08x != %08x); Bound Layout='%s' Shader='%s' %s"), BufferLayout.GetHash(), Shader->ShaderResourceTable.ResourceTableLayoutHashes[BufferIndex], *DebugName, *ShaderName, *ShaderUB);
-				FString ResourcesString;
-				for (int32 Index = 0; Index < BufferLayout.Resources.Num(); ++Index)
-				{
-					ResourcesString += FString::Printf(TEXT("%d "), BufferLayout.Resources[Index].MemberType);
-				}
-				UE_LOG(LogD3D11RHI, Error, TEXT("Layout CB Size %d %d Resources: %s"), BufferLayout.ConstantBufferSize, BufferLayout.Resources.Num(), *ResourcesString);
-#else
-				UE_LOG(LogD3D11RHI, Error, TEXT("Bound Layout='%s' Shader='%s', Layout CB Size %d %d"), *DebugName, *ShaderName, BufferLayout.ConstantBufferSize, BufferLayout.Resources.Num());
-#endif
-				// this might mean you are accessing a data you haven't bound e.g. GBuffer
-				checkf(BufferLayout.GetHash() == Shader->ShaderResourceTable.ResourceTableLayoutHashes[BufferIndex],
-					TEXT("Uniform buffer bound to slot %u is not what the shader expected:\n")
-					TEXT("\tBound:    Uniform Buffer[%s] with Hash[%u]\n")
-					TEXT("\tExpected: Uniform Buffer[%s] with Hash[%u]"),
-					BufferIndex, *DebugName, BufferLayout.GetHash(), *Shader->UniformBuffers[BufferIndex].GetPlainNameString(), Shader->ShaderResourceTable.ResourceTableLayoutHashes[BufferIndex]);
-			}
+			FD3D11ShaderResourceView* SRVRHI = ResourceCast(SRV);
+			RHI.SetShaderResourceView<Frequency>(
+				SRVRHI->Resource,
+				SRVRHI->View,
+				Index
+			);
 		}
-#endif
-
-#if RHI_INCLUDE_SHADER_DEBUG_DATA
-		const TCHAR* LayoutName = *Buffer->GetLayout().GetDebugName();
-#else 
-		const TCHAR* LayoutName = nullptr;
-#endif
-
-		// todo: could make this two pass: gather then set
-		SetShaderResourcesFromBuffer_Surface<(EShaderFrequency)ShaderType::StaticFrequency>(this, &StateCache, Buffer, Shader->ShaderResourceTable.TextureMap.GetData(), BufferIndex, LayoutName);
-		SetShaderResourcesFromBuffer_SRV<(EShaderFrequency)ShaderType::StaticFrequency>(this, &StateCache, Buffer, Shader->ShaderResourceTable.ShaderResourceViewMap.GetData(), BufferIndex, LayoutName);
-		SetShaderResourcesFromBuffer_Sampler<(EShaderFrequency)ShaderType::StaticFrequency>(this, &StateCache, Buffer, Shader->ShaderResourceTable.SamplerMap.GetData(), BufferIndex);
-	}
-	DirtyUniformBuffers[ShaderType::StaticFrequency] = 0;
-}
-
-template <class ShaderType>
-int32 FD3D11DynamicRHI::SetUAVPSResourcesFromTables(const ShaderType* RESTRICT Shader, bool bForceInvalidate)
-{
-	checkSlow(Shader);
-	int32 NumChanged = 0;
-	// Mask the dirty bits by those buffers from which the shader has bound resources.
-	uint16 DirtyMask = bForceInvalidate ? 0xffff : DirtyUniformBuffers[ShaderType::StaticFrequency];
-	uint32 DirtyBits = Shader->ShaderResourceTable.ResourceTableBits & DirtyMask;
-	while (DirtyBits)
-	{
-		// Scan for the lowest set bit, compute its index, clear it in the set of dirty bits.
-		const uint32 LowestBitMask = (DirtyBits) & (-(int32)DirtyBits);
-		const int32 BufferIndex = FMath::FloorLog2(LowestBitMask); // todo: This has a branch on zero, we know it could never be zero...
-		DirtyBits ^= LowestBitMask;
-		FD3D11UniformBuffer* Buffer = (FD3D11UniformBuffer*)BoundUniformBuffers[ShaderType::StaticFrequency][BufferIndex].GetReference();
 
-		check(BufferIndex < Shader->ShaderResourceTable.ResourceTableLayoutHashes.Num());
-
-		if (!Buffer)
+		void SetTexture(FRHITexture* Texture, uint8 Index)
 		{
-			FString ShaderUB;
-#if RHI_INCLUDE_SHADER_DEBUG_DATA
-			if (BufferIndex < Shader->UniformBuffers.Num())
-			{
-				Shader->UniformBuffers[BufferIndex].ToString(ShaderUB);
-			}
-#endif
-			UE_LOG(LogD3D11RHI, Fatal, TEXT("Shader expected a uniform buffer at slot %u but got null instead (Shader='%s' UB='%s').  Rendering code needs to set a valid uniform buffer for this slot."), BufferIndex, Shader->GetShaderName(), *ShaderUB);
+			FD3D11Texture* TextureD3D11 = ResourceCast(Texture);
+			RHI.SetShaderResourceView<Frequency>(
+				TextureD3D11,
+				TextureD3D11->GetShaderResourceView(),
+				Index
+			);
 		}
 
-		if ((EShaderFrequency)ShaderType::StaticFrequency == SF_Pixel)
+		void SetSampler(FRHISamplerState* Sampler, uint8 Index)
 		{
-#if RHI_INCLUDE_SHADER_DEBUG_DATA
-			const TCHAR* LayoutName = *Buffer->GetLayout().GetDebugName();
-#else
-			const TCHAR* LayoutName = nullptr;
-#endif
-
-			NumChanged += SetShaderResourcesFromBufferUAVPS<(EShaderFrequency)ShaderType::StaticFrequency>(this, &StateCache, Buffer, Shader->ShaderResourceTable.UnorderedAccessViewMap.GetData(), BufferIndex, LayoutName);
+			RHI.StateCache.SetSamplerState<Frequency>(ResourceCast(Sampler)->Resource, Index);
 		}
-	}
-	return NumChanged;
-}
+	};
 
-static int32 PeriodicCheck = 0;
+	UE::RHICore::SetResourcesFromTables(
+		  FUniformResourceBinder { *this }
+		, *Shader
+		, Shader->ShaderResourceTable
+		, DirtyUniformBuffers[Frequency]
+		, BoundUniformBuffers[Frequency]
+#if ENABLE_RHI_VALIDATION
+		, Tracker
+#endif
+	);
+}
 
 void FD3D11DynamicRHI::CommitGraphicsResourceTables()
 {
 #if !UE_BUILD_SHIPPING && !UE_BUILD_TEST
 	GDX11CommitGraphicsResourceTables.Increment();
 #endif
+
 	FD3D11BoundShaderState* RESTRICT CurrentBoundShaderState = (FD3D11BoundShaderState*)BoundShaderStateHistory.GetLast();
 	check(CurrentBoundShaderState);
+
 	auto* PixelShader = CurrentBoundShaderState->GetPixelShader();
-	if(PixelShader)
+	if (PixelShader)
 	{
-		//because d3d11 binding uses the same slots for UAVS and RTVS, we have to rebind, when two shaders with different sets of rendertargets are bound
-		//as they can potentially be used by UAVS, which can cause them to unbind RTVs used by subsequent shaders.
+		// Because d3d11 binding uses the same slots for UAVs and RTVs, we have to rebind when two shaders with different sets of rendertargets are bound,
+		// as they can potentially be used by UAVs, which can cause them to unbind RTVs used by subsequent shaders.
 		bool bRTVInvalidate = false;
 		uint32 UAVMask = PixelShader->UAVMask & CurrentRTVOverlapMask;
 		if (GDX11ReduceRTVRebinds && 
@@ -1510,13 +1260,15 @@ void FD3D11DynamicRHI::CommitGraphicsResourceTables()
 			CurrentUAVMask = UAVMask;
 		}
 
-
-		if(bRTVInvalidate)
+		if (bRTVInvalidate)
 		{
 			CommitRenderTargets(true);
+			DirtyUniformBuffers[SF_Pixel] = -1;
 		}
 
-		if(SetUAVPSResourcesFromTables(PixelShader, bRTVInvalidate) || UAVSChanged)
+		SetResourcesFromTables(PixelShader);
+
+		if (UAVSChanged)
 		{
 			CommitUAVs();
 		}
@@ -1526,10 +1278,6 @@ void FD3D11DynamicRHI::CommitGraphicsResourceTables()
 	{
 		SetResourcesFromTables(Shader);
 	}
-	if (PixelShader)
-	{
-		SetResourcesFromTables(PixelShader);
-	}
 	if (auto* Shader = CurrentBoundShaderState->GetGeometryShader())
 	{
 		SetResourcesFromTables(Shader);
diff --git a/Engine/Source/Runtime/Windows/D3D11RHI/Private/D3D11RHI.cpp b/Engine/Source/Runtime/Windows/D3D11RHI/Private/D3D11RHI.cpp
--- a/Engine/Source/Runtime/Windows/D3D11RHI/Private/D3D11RHI.cpp
+++ b/Engine/Source/Runtime/Windows/D3D11RHI/Private/D3D11RHI.cpp
@@ -793,25 +793,25 @@ ID3D11Buffer* FD3D11DynamicRHI::RHIGetResource(FRHIBuffer* InBuffer) const
 
 ID3D11Resource* FD3D11DynamicRHI::RHIGetResource(FRHITexture* InTexture) const
 {
-	FD3D11Texture* D3D11Texture = GetD3D11TextureFromRHITexture(InTexture);
+	FD3D11Texture* D3D11Texture = ResourceCast(InTexture);
 	return D3D11Texture->GetResource();
 }
 
 int64 FD3D11DynamicRHI::RHIGetResourceMemorySize(FRHITexture* InTexture) const
 {
-	FD3D11Texture* D3D11Texture = GetD3D11TextureFromRHITexture(InTexture);
+	FD3D11Texture* D3D11Texture = ResourceCast(InTexture);
 	return D3D11Texture->GetMemorySize();
 }
 
 ID3D11RenderTargetView* FD3D11DynamicRHI::RHIGetRenderTargetView(FRHITexture* InTexture, int32 InMipIndex, int32 InArraySliceIndex) const
 {
-	FD3D11Texture* D3D11Texture = GetD3D11TextureFromRHITexture(InTexture);
+	FD3D11Texture* D3D11Texture = ResourceCast(InTexture);
 	return D3D11Texture->GetRenderTargetView(InMipIndex, InArraySliceIndex);
 }
 
 ID3D11ShaderResourceView* FD3D11DynamicRHI::RHIGetShaderResourceView(FRHITexture* InTexture) const
 {
-	FD3D11Texture* D3D11Texture = GetD3D11TextureFromRHITexture(InTexture);
+	FD3D11Texture* D3D11Texture = ResourceCast(InTexture);
 	return D3D11Texture->GetShaderResourceView();
 }
 
diff --git a/Engine/Source/Runtime/Windows/D3D11RHI/Private/D3D11RHIPrivate.h b/Engine/Source/Runtime/Windows/D3D11RHI/Private/D3D11RHIPrivate.h
--- a/Engine/Source/Runtime/Windows/D3D11RHI/Private/D3D11RHIPrivate.h
+++ b/Engine/Source/Runtime/Windows/D3D11RHI/Private/D3D11RHIPrivate.h
@@ -438,6 +438,19 @@ public:
 		return static_cast<typename TD3D11ResourceTraits<TRHIType>::TConcreteType*>(Resource);
 	}
 
+	static inline FD3D11Texture* ResourceCast(FRHITexture* Texture)
+	{
+		if (!Texture)
+		{
+			return nullptr;
+		}
+
+		FD3D11Texture* Result = static_cast<FD3D11Texture*>(Texture->GetTextureBaseRHI());
+		check(Result);
+
+		return Result;
+	}
+
 	/**
 	 * Reads a D3D query's data into the provided buffer.
 	 * @param Query - The D3D query to read data from.
@@ -963,7 +976,7 @@ protected:
 	virtual void CommitComputeShaderConstants();
 
 	template <class ShaderType> void SetResourcesFromTables(const ShaderType* RESTRICT);
-	template <class ShaderType> int32 SetUAVPSResourcesFromTables(const ShaderType* RESTRICT, bool bForceInvalid);
+
 	void CommitGraphicsResourceTables();
 	void CommitComputeResourceTables(FD3D11ComputeShader* ComputeShader);
 
diff --git a/Engine/Source/Runtime/Windows/D3D11RHI/Private/D3D11RenderTarget.cpp b/Engine/Source/Runtime/Windows/D3D11RHI/Private/D3D11RenderTarget.cpp
--- a/Engine/Source/Runtime/Windows/D3D11RHI/Private/D3D11RenderTarget.cpp
+++ b/Engine/Source/Runtime/Windows/D3D11RHI/Private/D3D11RenderTarget.cpp
@@ -201,14 +201,14 @@ void FD3D11DynamicRHI::RHICopyToResolveTarget(FRHITexture* SourceTextureRHI, FRH
 		return;
 	}
 
-	FD3D11Texture* SourceTexture2D   = GetD3D11TextureFromRHITexture(SourceTextureRHI->GetTexture2D());
-	FD3D11Texture* DestTexture2D     = GetD3D11TextureFromRHITexture(DestTextureRHI->GetTexture2D());
+	FD3D11Texture* SourceTexture2D   = ResourceCast(SourceTextureRHI->GetTexture2D());
+	FD3D11Texture* DestTexture2D     = ResourceCast(DestTextureRHI->GetTexture2D());
 
-	FD3D11Texture* SourceTextureCube = GetD3D11TextureFromRHITexture(SourceTextureRHI->GetTextureCube());
-	FD3D11Texture* DestTextureCube   = GetD3D11TextureFromRHITexture(DestTextureRHI->GetTextureCube());
+	FD3D11Texture* SourceTextureCube = ResourceCast(SourceTextureRHI->GetTextureCube());
+	FD3D11Texture* DestTextureCube   = ResourceCast(DestTextureRHI->GetTextureCube());
 
-	FD3D11Texture* SourceTexture3D   = GetD3D11TextureFromRHITexture(SourceTextureRHI->GetTexture3D());
-	FD3D11Texture* DestTexture3D     = GetD3D11TextureFromRHITexture(DestTextureRHI->GetTexture3D());
+	FD3D11Texture* SourceTexture3D   = ResourceCast(SourceTextureRHI->GetTexture3D());
+	FD3D11Texture* DestTexture3D     = ResourceCast(DestTextureRHI->GetTexture3D());
 		
 	if(SourceTexture2D && DestTexture2D)
 	{
@@ -438,7 +438,7 @@ static uint32 ComputeBytesPerPixel(DXGI_FORMAT Format)
 
 TRefCountPtr<ID3D11Texture2D> FD3D11DynamicRHI::GetStagingTexture(FRHITexture* TextureRHI,FIntRect InRect, FIntRect& StagingRectOUT, FReadSurfaceDataFlags InFlags)
 {
-	FD3D11Texture* Texture = GetD3D11TextureFromRHITexture(TextureRHI);
+	FD3D11Texture* Texture = ResourceCast(TextureRHI);
 
 	D3D11_TEXTURE2D_DESC SourceDesc; 
 	Texture->GetD3D11Texture2D()->GetDesc(&SourceDesc);
@@ -515,7 +515,7 @@ void FD3D11DynamicRHI::ReadSurfaceDataNoMSAARaw(FRHITexture* TextureRHI,FIntRect
 	checkf(InRect.Width() <= TextureRHI->GetSizeXYZ().X >> InFlags.GetMip(), TEXT("Provided rect width (%d), must be smaller or equal to the texture size requested Mip (%d)"), InRect.Width(), TextureRHI->GetSizeXYZ().X >> InFlags.GetMip());
 	checkf(InRect.Height() <= TextureRHI->GetSizeXYZ().Y >> InFlags.GetMip(), TEXT("Provided rect height (%d), must be smaller or equal to the texture size requested Mip (%d)"), InRect.Height(), TextureRHI->GetSizeXYZ().Y >> InFlags.GetMip());
 
-	FD3D11Texture* Texture = GetD3D11TextureFromRHITexture(TextureRHI);
+	FD3D11Texture* Texture = ResourceCast(TextureRHI);
 
 	const uint32 SizeX = InRect.Width();
 	const uint32 SizeY = InRect.Height();
@@ -654,7 +654,7 @@ void FD3D11DynamicRHI::RHIReadSurfaceData(FRHITexture* TextureRHI,FIntRect InRec
 
 	TArray<uint8> OutDataRaw;
 
-	FD3D11Texture* Texture = GetD3D11TextureFromRHITexture(TextureRHI);
+	FD3D11Texture* Texture = ResourceCast(TextureRHI);
 
 	// Check the format of the surface
 	D3D11_TEXTURE2D_DESC TextureDesc;
@@ -686,7 +686,7 @@ void FD3D11DynamicRHI::RHIReadSurfaceData(FRHITexture* TextureRHI,FIntRect InRec
 
 void FD3D11DynamicRHI::ReadSurfaceDataMSAARaw(FRHITexture* TextureRHI,FIntRect InRect,TArray<uint8>& OutData, FReadSurfaceDataFlags InFlags)
 {
-	FD3D11Texture* Texture = GetD3D11TextureFromRHITexture(TextureRHI);
+	FD3D11Texture* Texture = ResourceCast(TextureRHI);
 
 	const uint32 SizeX = InRect.Width();
 	const uint32 SizeY = InRect.Height();
@@ -811,7 +811,7 @@ void FD3D11DynamicRHI::ReadSurfaceDataMSAARaw(FRHITexture* TextureRHI,FIntRect I
 
 void FD3D11DynamicRHI::RHIMapStagingSurface(FRHITexture* TextureRHI, FRHIGPUFence* FenceRHI, void*& OutData, int32& OutWidth, int32& OutHeight, uint32 GPUIndex)
 {
-	ID3D11Texture2D* Texture = GetD3D11TextureFromRHITexture(TextureRHI)->GetD3D11Texture2D();
+	ID3D11Texture2D* Texture = ResourceCast(TextureRHI)->GetD3D11Texture2D();
 	
 	D3D11_TEXTURE2D_DESC TextureDesc;
 	Texture->GetDesc(&TextureDesc);
@@ -829,14 +829,14 @@ void FD3D11DynamicRHI::RHIMapStagingSurface(FRHITexture* TextureRHI, FRHIGPUFenc
 
 void FD3D11DynamicRHI::RHIUnmapStagingSurface(FRHITexture* TextureRHI, uint32 GPUIndex)
 {
-	ID3D11Texture2D* Texture = GetD3D11TextureFromRHITexture(TextureRHI)->GetD3D11Texture2D();
+	ID3D11Texture2D* Texture = ResourceCast(TextureRHI)->GetD3D11Texture2D();
 
 	Direct3DDeviceIMContext->Unmap(Texture,0);
 }
 
 void FD3D11DynamicRHI::RHIReadSurfaceFloatData(FRHITexture* TextureRHI,FIntRect InRect,TArray<FFloat16Color>& OutData,ECubeFace CubeFace,int32 ArrayIndex,int32 MipIndex)
 {
-	FD3D11Texture* Texture = GetD3D11TextureFromRHITexture(TextureRHI);
+	FD3D11Texture* Texture = ResourceCast(TextureRHI);
 
 	uint32 SizeX = InRect.Width();
 	uint32 SizeY = InRect.Height();
@@ -973,7 +973,7 @@ void FD3D11DynamicRHI::RHIReadSurfaceData(FRHITexture* TextureRHI, FIntRect InRe
 {
 	TArray<uint8> OutDataRaw;
 
-	FD3D11Texture* Texture = GetD3D11TextureFromRHITexture(TextureRHI);
+	FD3D11Texture* Texture = ResourceCast(TextureRHI);
 
 	// Check the format of the surface
 	D3D11_TEXTURE2D_DESC TextureDesc;
@@ -1008,7 +1008,7 @@ void FD3D11DynamicRHI::RHIReadSurfaceData(FRHITexture* TextureRHI, FIntRect InRe
 
 void FD3D11DynamicRHI::RHIRead3DSurfaceFloatData(FRHITexture* TextureRHI,FIntRect InRect,FIntPoint ZMinMax,TArray<FFloat16Color>& OutData)
 {
-	FD3D11Texture* Texture = GetD3D11TextureFromRHITexture(TextureRHI);
+	FD3D11Texture* Texture = ResourceCast(TextureRHI);
 
 	uint32 SizeX = InRect.Width();
 	uint32 SizeY = InRect.Height();
@@ -1156,10 +1156,10 @@ void FD3D11DynamicRHI::ResolveTexture(UE::RHICore::FResolveTextureInfo Info)
 {
 	GPUProfilingData.RegisterGPUWork();
 
-	FD3D11Texture* SourceTexture      = GetD3D11TextureFromRHITexture(Info.SourceTexture);
+	FD3D11Texture* SourceTexture      = ResourceCast(Info.SourceTexture);
 	const FRHITextureDesc& SourceDesc = SourceTexture->GetDesc();
 
-	FD3D11Texture* DestTexture        = GetD3D11TextureFromRHITexture(Info.DestTexture);
+	FD3D11Texture* DestTexture        = ResourceCast(Info.DestTexture);
 	const FRHITextureDesc& DestDesc   = DestTexture->GetDesc();
 
 	if (SourceDesc.Format == PF_DepthStencil)
diff --git a/Engine/Source/Runtime/Windows/D3D11RHI/Private/D3D11Shaders.cpp b/Engine/Source/Runtime/Windows/D3D11RHI/Private/D3D11Shaders.cpp
--- a/Engine/Source/Runtime/Windows/D3D11RHI/Private/D3D11Shaders.cpp
+++ b/Engine/Source/Runtime/Windows/D3D11RHI/Private/D3D11Shaders.cpp
@@ -26,7 +26,7 @@ static inline void ReadShaderOptionalData(FShaderCodeReader& InShaderCode, TShad
 	
 	OutShader.bShaderNeedsGlobalConstantBuffer = EnumHasAnyFlags(PackedResourceCounts->UsageFlags, EShaderResourceUsageFlags::GlobalUniformBuffer);
 #if RHI_INCLUDE_SHADER_DEBUG_DATA
-	OutShader.ShaderName = InShaderCode.FindOptionalData(FShaderCodeName::Key);
+	OutShader.Debug.ShaderName = InShaderCode.FindOptionalData(FShaderCodeName::Key);
 
 	int32 UniformBufferTableSize = 0;
 	const uint8* UniformBufferData = InShaderCode.FindOptionalDataAndSize(FShaderCodeUniformBuffers::Key, UniformBufferTableSize);
@@ -35,10 +35,10 @@ static inline void ReadShaderOptionalData(FShaderCodeReader& InShaderCode, TShad
 		FBufferReader UBReader((void*)UniformBufferData, UniformBufferTableSize, false);
 		TArray<FString> Names;
 		UBReader << Names;
-		check(OutShader.UniformBuffers.Num() == 0);
+		check(OutShader.Debug.UniformBufferNames.Num() == 0);
 		for (int32 Index = 0; Index < Names.Num(); ++Index)
 		{
-			OutShader.UniformBuffers.Add(FName(*Names[Index]));
+			OutShader.Debug.UniformBufferNames.Add(FName(*Names[Index]));
 		}
 	}
 #endif
diff --git a/Engine/Source/Runtime/Windows/D3D11RHI/Private/D3D11Texture.cpp b/Engine/Source/Runtime/Windows/D3D11RHI/Private/D3D11Texture.cpp
--- a/Engine/Source/Runtime/Windows/D3D11RHI/Private/D3D11Texture.cpp
+++ b/Engine/Source/Runtime/Windows/D3D11RHI/Private/D3D11Texture.cpp
@@ -1021,7 +1021,7 @@ FTextureRHIRef FD3D11DynamicRHI::RHIAsyncCreateTexture2D(uint32 SizeX, uint32 Si
 
 FShaderResourceViewRHIRef FD3D11DynamicRHI::RHICreateShaderResourceView(FRHITexture* TextureRHI, const FRHITextureSRVCreateInfo& CreateInfo)
 {
-	FD3D11Texture* Texture = GetD3D11TextureFromRHITexture(TextureRHI);
+	FD3D11Texture* Texture = ResourceCast(TextureRHI);
 
 	// Create a Shader Resource View
 	D3D11_SHADER_RESOURCE_VIEW_DESC SRVDesc;
@@ -1102,7 +1102,7 @@ FShaderResourceViewRHIRef FD3D11DynamicRHI::RHICreateShaderResourceView_RenderTh
 /** Generates mip maps for the surface. */
 void FD3D11DynamicRHI::RHIGenerateMips(FRHITexture* TextureRHI)
 {
-	FD3D11Texture* Texture = GetD3D11TextureFromRHITexture(TextureRHI);
+	FD3D11Texture* Texture = ResourceCast(TextureRHI);
 	// Surface must have been created with D3D11_BIND_RENDER_TARGET for GenerateMips to work
 	check(Texture->GetShaderResourceView() && Texture->GetRenderTargetView(0, -1));
 	Direct3DDeviceIMContext->GenerateMips(Texture->GetShaderResourceView());
@@ -1123,7 +1123,7 @@ uint32 FD3D11DynamicRHI::RHIComputeMemorySize(FRHITexture* TextureRHI)
 		return 0;
 	}
 	
-	FD3D11Texture* Texture = GetD3D11TextureFromRHITexture(TextureRHI);
+	FD3D11Texture* Texture = ResourceCast(TextureRHI);
 	return Texture->GetMemorySize();
 }
 
@@ -1628,7 +1628,7 @@ void FD3D11DynamicRHI::RHIUnlockTextureCubeFace_RenderThread(
 
 void FD3D11DynamicRHI::RHIBindDebugLabelName(FRHITexture* TextureRHI, const TCHAR* Name)
 {
-	FD3D11Texture* Texture = GetD3D11TextureFromRHITexture(TextureRHI);
+	FD3D11Texture* Texture = ResourceCast(TextureRHI);
 
 	//todo: require names at texture creation time.
 	FName DebugName(Name);
@@ -1978,8 +1978,8 @@ void FD3D11Texture::AliasResource(FD3D11Texture const& Other)
 
 void FD3D11DynamicRHI::RHIAliasTextureResources(FTextureRHIRef& DstTextureRHI, FTextureRHIRef& SrcTextureRHI)
 {
-	FD3D11Texture* DstTexture = GetD3D11TextureFromRHITexture(DstTextureRHI);
-	FD3D11Texture* SrcTexture = GetD3D11TextureFromRHITexture(SrcTextureRHI);
+	FD3D11Texture* DstTexture = ResourceCast(DstTextureRHI);
+	FD3D11Texture* SrcTexture = ResourceCast(SrcTextureRHI);
 	check(DstTexture && SrcTexture);
 
 	DstTexture->AliasResource(*SrcTexture);
@@ -1987,7 +1987,7 @@ void FD3D11DynamicRHI::RHIAliasTextureResources(FTextureRHIRef& DstTextureRHI, F
 
 FTextureRHIRef FD3D11DynamicRHI::RHICreateAliasedTexture(FTextureRHIRef& SrcTextureRHI)
 {
-	FD3D11Texture* SrcTexture = GetD3D11TextureFromRHITexture(SrcTextureRHI);
+	FD3D11Texture* SrcTexture = ResourceCast(SrcTextureRHI);
 	check(SrcTexture);
 	const FString Name = SrcTextureRHI->GetName().ToString() + TEXT("Alias");
 
@@ -1998,8 +1998,8 @@ void FD3D11DynamicRHI::RHICopyTexture(FRHITexture* SourceTextureRHI, FRHITexture
 {
 	FRHICommandList_RecursiveHazardous RHICmdList(this);	
 
-	FD3D11Texture* SourceTexture = GetD3D11TextureFromRHITexture(SourceTextureRHI);
-	FD3D11Texture* DestTexture = GetD3D11TextureFromRHITexture(DestTextureRHI);
+	FD3D11Texture* SourceTexture = ResourceCast(SourceTextureRHI);
+	FD3D11Texture* DestTexture = ResourceCast(DestTextureRHI);
 
 	check(SourceTexture && DestTexture);
 
diff --git a/Engine/Source/Runtime/Windows/D3D11RHI/Private/D3D11UAV.cpp b/Engine/Source/Runtime/Windows/D3D11RHI/Private/D3D11UAV.cpp
--- a/Engine/Source/Runtime/Windows/D3D11RHI/Private/D3D11UAV.cpp
+++ b/Engine/Source/Runtime/Windows/D3D11RHI/Private/D3D11UAV.cpp
@@ -76,7 +76,7 @@ FUnorderedAccessViewRHIRef FD3D11DynamicRHI::RHICreateUnorderedAccessView_Render
 
 FUnorderedAccessViewRHIRef FD3D11DynamicRHI::RHICreateUnorderedAccessView(FRHITexture* TextureRHI, uint32 MipLevel, uint16 FirstArraySlice, uint16 NumArraySlices)
 {
-	FD3D11Texture* Texture = GetD3D11TextureFromRHITexture(TextureRHI);
+	FD3D11Texture* Texture = ResourceCast(TextureRHI);
 	
 	D3D11_UNORDERED_ACCESS_VIEW_DESC UAVDesc{};
 
diff --git a/Engine/Source/Runtime/Windows/D3D11RHI/Private/D3D11UniformBuffer.cpp b/Engine/Source/Runtime/Windows/D3D11RHI/Private/D3D11UniformBuffer.cpp
--- a/Engine/Source/Runtime/Windows/D3D11RHI/Private/D3D11UniformBuffer.cpp
+++ b/Engine/Source/Runtime/Windows/D3D11RHI/Private/D3D11UniformBuffer.cpp
@@ -260,15 +260,15 @@ FUniformBufferRHIRef FD3D11DynamicRHI::RHICreateUniformBuffer(const void* Conten
 	if (Layout->Resources.Num())
 	{
 		const int32 ResourceCount = Layout->Resources.Num();
-		NewUniformBuffer->ResourceTable.Empty(ResourceCount);
-		NewUniformBuffer->ResourceTable.AddZeroed(ResourceCount);
+		NewUniformBuffer->GetResourceTable().Empty(ResourceCount);
+		NewUniformBuffer->GetResourceTable().AddZeroed(ResourceCount);
 
 		if (Contents)
 		{
 			for (int32 Index = 0; Index < ResourceCount; ++Index)
 			{
 				const auto ResourceParameter = Layout->Resources[Index];
-				NewUniformBuffer->ResourceTable[Index] = GetShaderParameterResourceRHI(Contents, ResourceParameter.MemberOffset, ResourceParameter.MemberType);
+				NewUniformBuffer->GetResourceTable()[Index] = GetShaderParameterResourceRHI(Contents, ResourceParameter.MemberOffset, ResourceParameter.MemberType);
 			}
 		}
 	}
@@ -304,7 +304,7 @@ void FD3D11DynamicRHI::RHIUpdateUniformBuffer(FRHICommandListBase& RHICmdList, F
 	const uint32 ConstantBufferSize = Layout.ConstantBufferSize;
 	const int32 NumResources = Layout.Resources.Num();
 
-	check(UniformBuffer->ResourceTable.Num() == NumResources);
+	check(UniformBuffer->GetResourceTable().Num() == NumResources);
 
 	if (RHICmdList.Bypass())
 	{
@@ -313,7 +313,7 @@ void FD3D11DynamicRHI::RHIUpdateUniformBuffer(FRHICommandListBase& RHICmdList, F
 		for (int32 Index = 0; Index < NumResources; ++Index)
 		{
 			const auto Parameter = Layout.Resources[Index];
-			UniformBuffer->ResourceTable[Index] = GetShaderParameterResourceRHI(Contents, Parameter.MemberOffset, Parameter.MemberType);
+			UniformBuffer->GetResourceTable()[Index] = GetShaderParameterResourceRHI(Contents, Parameter.MemberOffset, Parameter.MemberType);
 		}
 	}
 	else
@@ -351,7 +351,7 @@ void FD3D11DynamicRHI::RHIUpdateUniformBuffer(FRHICommandListBase& RHICmdList, F
 			// Update resource table.
 			for (int32 ResourceIndex = 0; ResourceIndex < NumResources; ++ResourceIndex)
 			{
-				UniformBuffer->ResourceTable[ResourceIndex] = CmdListResources[ResourceIndex];
+				UniformBuffer->GetResourceTable()[ResourceIndex] = CmdListResources[ResourceIndex];
 			}
 		});
 		RHICmdList.RHIThreadFence(true);
diff --git a/Engine/Source/Runtime/Windows/D3D11RHI/Private/D3D11Util.cpp b/Engine/Source/Runtime/Windows/D3D11RHI/Private/D3D11Util.cpp
--- a/Engine/Source/Runtime/Windows/D3D11RHI/Private/D3D11Util.cpp
+++ b/Engine/Source/Runtime/Windows/D3D11RHI/Private/D3D11Util.cpp
@@ -267,20 +267,11 @@ void VerifyD3D11ShaderResult(FRHIShader* Shader, HRESULT D3DResult, const ANSICH
 
 	const FString& ErrorString = GetD3D11ErrorString(D3DResult, Device);
 
-#if !(UE_BUILD_SHIPPING || UE_BUILD_TEST)
-	if (Shader->ShaderName.Len())
-	{
-		UE_LOG(LogD3D11RHI, Error, TEXT("%s failed trying to create shader '%s' with error %s\n at %s:%u"), ANSI_TO_TCHAR(Code), *Shader->ShaderName, *ErrorString, ANSI_TO_TCHAR(Filename), Line);
-		TerminateOnDeviceRemoved(D3DResult, Device);
-		TerminateOnOutOfMemory(D3DResult, false);
+	UE_LOG(LogD3D11RHI, Error, TEXT("%s failed trying to create shader '%s' with error %s\n at %s:%u"), ANSI_TO_TCHAR(Code), Shader->GetShaderName(), *ErrorString, ANSI_TO_TCHAR(Filename), Line);
+	TerminateOnDeviceRemoved(D3DResult, Device);
+	TerminateOnOutOfMemory(D3DResult, false);
 
-		UE_LOG(LogD3D11RHI, Fatal, TEXT("%s failed trying to create shader '%s' with error %s\n at %s:%u"), ANSI_TO_TCHAR(Code), *Shader->ShaderName, *ErrorString, ANSI_TO_TCHAR(Filename), Line);
-	}
-	else
-#endif
-	{
-		VerifyD3D11Result(D3DResult, Code, Filename, Line, Device);
-	}
+	UE_LOG(LogD3D11RHI, Fatal, TEXT("%s failed trying to create shader '%s' with error %s\n at %s:%u"), ANSI_TO_TCHAR(Code), Shader->GetShaderName(), *ErrorString, ANSI_TO_TCHAR(Filename), Line);
 }
 
 void VerifyD3D11CreateTextureResult(HRESULT D3DResult, int32 UEFormat,const ANSICHAR* Code,const ANSICHAR* Filename,uint32 Line,uint32 SizeX,uint32 SizeY,uint32 SizeZ,uint8 D3DFormat,uint32 NumMips,uint32 Flags,
diff --git a/Engine/Source/Runtime/Windows/D3D11RHI/Public/D3D11Resources.h b/Engine/Source/Runtime/Windows/D3D11RHI/Public/D3D11Resources.h
--- a/Engine/Source/Runtime/Windows/D3D11RHI/Public/D3D11Resources.h
+++ b/Engine/Source/Runtime/Windows/D3D11RHI/Public/D3D11Resources.h
@@ -7,7 +7,6 @@
 #pragma once
 
 #include "BoundShaderStateCache.h"
-#include "D3D11ShaderResources.h"
 
 interface ID3D11DeviceContext;
 typedef ID3D11DeviceContext FD3D11DeviceContext;
@@ -40,10 +39,7 @@ public:
 
 struct FD3D11ShaderData
 {
-	FD3D11ShaderResourceTable			ShaderResourceTable;
-#if RHI_INCLUDE_SHADER_DEBUG_DATA
-	TArray<FName>						UniformBuffers;
-#endif
+	FShaderResourceTable				ShaderResourceTable;
 	TArray<FUniformBufferStaticSlot>	StaticSlots;
 	TArray<FShaderCodeVendorExtension>	VendorExtensions;
 	bool								bShaderNeedsGlobalConstantBuffer;
@@ -299,19 +295,6 @@ private:
 	uint8 bAlias : 1;
 };
 
-/** Given a pointer to a RHI texture that was created by the D3D11 RHI, returns a pointer to the FD3D11Texture it encapsulates. */
-FORCEINLINE FD3D11Texture* GetD3D11TextureFromRHITexture(FRHITexture* Texture)
-{
-	if (!Texture)
-	{
-		return nullptr;
-	}
-
-	FD3D11Texture* Result = static_cast<FD3D11Texture*>(Texture->GetTextureBaseRHI());
-	check(Result);
-	return Result;
-}
-
 
 /** D3D11 render query */
 class FD3D11RenderQuery : public FRHIRenderQuery
@@ -369,9 +352,6 @@ public:
 	/** Allocation in the constants ring buffer if applicable. */
 	FRingAllocation RingAllocation;
 
-	/** Resource table containing RHI references. */
-	TArray<TRefCountPtr<FRHIResource> > ResourceTable;
-
 	/** Initialization constructor. */
 	FD3D11UniformBuffer(class FD3D11DynamicRHI* InD3D11RHI, const FRHIUniformBufferLayout* InLayout, ID3D11Buffer* InResource,const FRingAllocation& InRingAllocation, bool bInAllocatedFromPool)
 	: FRHIUniformBuffer(InLayout)
@@ -383,6 +363,10 @@ public:
 
 	virtual ~FD3D11UniformBuffer();
 
+	// Provides public non-const access to ResourceTable.
+	// @todo refactor uniform buffers to perform updates as a member function, so this isn't necessary.
+	TArray<TRefCountPtr<FRHIResource>>& GetResourceTable() { return ResourceTable; }
+
 private:
 	class FD3D11DynamicRHI* D3D11RHI;
 	bool bAllocatedFromPool;
@@ -545,11 +529,6 @@ struct TD3D11ResourceTraits<FRHIBoundShaderState>
 	typedef FD3D11BoundShaderState TConcreteType;
 };
 template<>
-struct TD3D11ResourceTraits<FRHITexture>
-{
-	typedef FD3D11Texture TConcreteType;
-};
-template<>
 struct TD3D11ResourceTraits<FRHIRenderQuery>
 {
 	typedef FD3D11RenderQuery TConcreteType;
diff --git a/Engine/Source/Runtime/Windows/D3D11RHI/Public/D3D11ShaderResources.h b/Engine/Source/Runtime/Windows/D3D11RHI/Public/D3D11ShaderResources.h
deleted file mode 100644
