Force minimum wave size on pre-RDNA AMD GPUs to 64 to work around issues in some drivers misreporting support.

diff --git a/Engine/Source/Runtime/D3D12RHI/Private/Windows/WindowsD3D12Device.cpp b/Engine/Source/Runtime/D3D12RHI/Private/Windows/WindowsD3D12Device.cpp
--- a/Engine/Source/Runtime/D3D12RHI/Private/Windows/WindowsD3D12Device.cpp
+++ b/Engine/Source/Runtime/D3D12RHI/Private/Windows/WindowsD3D12Device.cpp
@@ -1333,6 +1333,26 @@ void FD3D12DynamicRHI::Init()
 	bool bHasVendorSupportForAtomic64 = false;
 
 #if AMD_API_ENABLE
+
+	// Check if the AMD device is pre-RDNA, and ensure it doesn't misreport wave32 support
+	if (IsRHIDeviceAMD() && bAllowVendorDevice && AmdAgsContext)
+	{
+		for (int32 DeviceIndex = 0; DeviceIndex < AmdAgsGpuInfo.numDevices; DeviceIndex++)
+		{
+			const AGSDeviceInfo& DeviceInfo = AmdAgsGpuInfo.devices[DeviceIndex];
+			if (DeviceInfo.deviceId == AdapterDesc.DeviceId && DeviceInfo.vendorId == AdapterDesc.VendorId)
+			{
+				if (DeviceInfo.asicFamily != AGSDeviceInfo::AsicFamily_Unknown && DeviceInfo.asicFamily < AGSDeviceInfo::AsicFamily_RDNA)
+				{
+					GRHIMinimumWaveSize = 64;
+					GRHIMaximumWaveSize = 64;
+				}
+
+				break;
+			}
+		}
+	}
+
 	// Warn if we are trying to use RGP frame markers but are either running on a non-AMD device
 	// or using an older AMD driver without RGP marker support
 	if (IsRHIDeviceAMD())
