Check if AccelerationStructureBuffer is null and include geometry name in message.
- Useful to track down source of bugs.

diff --git a/Engine/Source/Runtime/D3D12RHI/Private/D3D12RayTracing.h b/Engine/Source/Runtime/D3D12RHI/Private/D3D12RayTracing.h
--- a/Engine/Source/Runtime/D3D12RHI/Private/D3D12RayTracing.h
+++ b/Engine/Source/Runtime/D3D12RHI/Private/D3D12RayTracing.h
@@ -29,6 +29,8 @@ public:
 	virtual FRayTracingAccelerationStructureAddress GetAccelerationStructureAddress(uint64 GPUIndex) const final override
 	{
 		checkf(IsInRHIThread() || !IsRunningRHIInSeparateThread(), TEXT("Acceleration structure addresses can only be accessed on RHI timeline due to compaction and defragmentation."));
+		checkf(AccelerationStructureBuffers[GPUIndex], 
+			TEXT("Trying to get address of acceleration structure '%s' without allocated memory."), *DebugName.ToString());
 		return AccelerationStructureBuffers[GPUIndex]->ResourceLocation.GetGPUVirtualAddress();
 	}
 	virtual void SetInitializer(const FRayTracingGeometryInitializer& Initializer) final override;
