Remove use of GetContextSyncPoint in FD3D12OnlineDescriptorManager
 - Replaced with pipelined cleanup of used descriptor blocks.
 - Context sync points cause command lists to split / extra ExecuteCommandList + Signal calls on the submission thread, which results in lower GPU utilization.
 - GetContextSyncPoint was a temporary solution until various D3D12 RHI systems can be refactored to use the submission pipe for cleanup.

diff --git a/Engine/Source/Runtime/D3D12RHI/Private/D3D12CommandContext.h b/Engine/Source/Runtime/D3D12RHI/Private/D3D12CommandContext.h
--- a/Engine/Source/Runtime/D3D12RHI/Private/D3D12CommandContext.h
+++ b/Engine/Source/Runtime/D3D12RHI/Private/D3D12CommandContext.h
@@ -40,7 +40,8 @@ struct FD3D12DeferredDeleteObject
 		D3DObject,
 		D3DHeap,
 		BindlessDescriptor,
-		CPUAllocation
+		CPUAllocation,
+		DescriptorBlock
 	} Type;
 
 	union
@@ -56,6 +57,12 @@ struct FD3D12DeferredDeleteObject
 		} BindlessDescriptor;
 
 		void* CPUAllocation;
+
+		struct
+		{
+			FD3D12OnlineDescriptorBlock* Block;
+			FD3D12OnlineDescriptorManager* Manager;
+		} DescriptorBlock;
 	};
 
 	explicit FD3D12DeferredDeleteObject(FD3D12Resource* RHIObject)
@@ -84,6 +91,11 @@ struct FD3D12DeferredDeleteObject
 	{
 		check(Type == EType::CPUAllocation);
 	}
+
+	explicit FD3D12DeferredDeleteObject(FD3D12OnlineDescriptorBlock* Block, FD3D12OnlineDescriptorManager* Manager)
+		: Type(EType::DescriptorBlock)
+		, DescriptorBlock({ Block, Manager })
+	{}
 };
 
 enum class ED3D12Units
diff --git a/Engine/Source/Runtime/D3D12RHI/Private/D3D12DescriptorCache.cpp b/Engine/Source/Runtime/D3D12RHI/Private/D3D12DescriptorCache.cpp
--- a/Engine/Source/Runtime/D3D12RHI/Private/D3D12DescriptorCache.cpp
+++ b/Engine/Source/Runtime/D3D12RHI/Private/D3D12DescriptorCache.cpp
@@ -953,9 +953,6 @@ bool FD3D12SubAllocatedOnlineHeap::AllocateBlock()
 		check(FirstUsedSlot == 0);
 		CurrentBlock->SizeUsed = NextSlotIndex;
 
-		// Create the sync point on the current command list
-		CurrentBlock->SyncPoint = Context.GetContextSyncPoint();
-
 		OnlineManager.FreeHeapBlock(CurrentBlock);
 		CurrentBlock = nullptr;
 	}
diff --git a/Engine/Source/Runtime/D3D12RHI/Private/D3D12Descriptors.cpp b/Engine/Source/Runtime/D3D12RHI/Private/D3D12Descriptors.cpp
--- a/Engine/Source/Runtime/D3D12RHI/Private/D3D12Descriptors.cpp
+++ b/Engine/Source/Runtime/D3D12RHI/Private/D3D12Descriptors.cpp
@@ -201,7 +201,7 @@ FD3D12OnlineDescriptorManager::FD3D12OnlineDescriptorManager(FD3D12Device* Devic
 
 FD3D12OnlineDescriptorManager::~FD3D12OnlineDescriptorManager() = default;
 
-/** Allocate and initialize the online heap */
+// Allocate and initialize the online heap
 void FD3D12OnlineDescriptorManager::Init(uint32 InTotalSize, uint32 InBlockSize)
 {
 	Heap = GetParentDevice()->GetDescriptorHeapManager().AllocateHeap(
@@ -217,7 +217,6 @@ void FD3D12OnlineDescriptorManager::Init(uint32 InTotalSize, uint32 InBlockSize)
 
 	// Compute amount of free blocks
 	uint32 BlockCount = InTotalSize / InBlockSize;
-	ReleasedBlocks.Reserve(BlockCount);
 
 	// Allocate the free blocks
 	uint32 CurrentBaseSlot = 0;
@@ -230,16 +229,13 @@ void FD3D12OnlineDescriptorManager::Init(uint32 InTotalSize, uint32 InBlockSize)
 	}
 }
 
-/** Allocate a new heap block - will also check if released blocks can be freed again */
+// Allocate a new heap block
 FD3D12OnlineDescriptorBlock* FD3D12OnlineDescriptorManager::AllocateHeapBlock()
 {
 	SCOPED_NAMED_EVENT(FD3D12OnlineViewHeap_AllocateHeapBlock, FColor::Silver);
 
 	FScopeLock Lock(&CriticalSection);
 
-	// Check if certain released blocks are free again
-	UpdateFreeBlocks();
-
 	// Free block
 	FD3D12OnlineDescriptorBlock* Result = nullptr;
 	FreeBlocks.Dequeue(Result);
@@ -255,44 +251,29 @@ FD3D12OnlineDescriptorBlock* FD3D12OnlineDescriptorManager::AllocateHeapBlock()
 	return Result;
 }
 
-/** Free given block - can still be used by the GPU (SyncPoint needs to be setup by the caller and will be used to check if the block can be reused again) */
-void FD3D12OnlineDescriptorManager::FreeHeapBlock(FD3D12OnlineDescriptorBlock * InHeapBlock)
+// Free given block - can still be used by the GPU
+void FD3D12OnlineDescriptorManager::FreeHeapBlock(FD3D12OnlineDescriptorBlock* InHeapBlock)
 {
-	FScopeLock Lock(&CriticalSection);
-
 	// Update stats
 	DEC_DWORD_STAT_BY(STAT_GlobalViewHeapReservedDescriptors, InHeapBlock->Size);
 	INC_DWORD_STAT_BY(STAT_GlobalViewHeapUsedDescriptors, InHeapBlock->SizeUsed);
 	INC_DWORD_STAT_BY(STAT_GlobalViewHeapWastedDescriptors, InHeapBlock->Size - InHeapBlock->SizeUsed);
 
-	ReleasedBlocks.Add(InHeapBlock);
+	FD3D12DynamicRHI::GetD3DRHI()->DeferredDelete(InHeapBlock, this);
 }
 
-/** Find all the blocks which are not used by the GPU anymore */
-void FD3D12OnlineDescriptorManager::UpdateFreeBlocks()
+// Called by the EOP task to recycle blocks
+void FD3D12OnlineDescriptorManager::Recycle(FD3D12OnlineDescriptorBlock* Block)
 {
-	for (int32 BlockIndex = 0; BlockIndex < ReleasedBlocks.Num(); )
-	{
-		// Check if GPU is ready consuming the block data
-		FD3D12OnlineDescriptorBlock* ReleasedBlock = ReleasedBlocks[BlockIndex];
-		if (ReleasedBlock->SyncPoint->IsComplete())
-		{
-			// Update stats
-			DEC_DWORD_STAT_BY(STAT_GlobalViewHeapUsedDescriptors, ReleasedBlock->SizeUsed);
-			DEC_DWORD_STAT_BY(STAT_GlobalViewHeapWastedDescriptors, ReleasedBlock->Size - ReleasedBlock->SizeUsed);
-			INC_DWORD_STAT_BY(STAT_GlobalViewHeapFreeDescriptors, ReleasedBlock->Size);
+	FScopeLock Lock(&CriticalSection);
 
-			ReleasedBlock->SizeUsed = 0;
-			FreeBlocks.Enqueue(ReleasedBlock);
+	// Update stats
+	DEC_DWORD_STAT_BY(STAT_GlobalViewHeapUsedDescriptors, Block->SizeUsed);
+	DEC_DWORD_STAT_BY(STAT_GlobalViewHeapWastedDescriptors, Block->Size - Block->SizeUsed);
+	INC_DWORD_STAT_BY(STAT_GlobalViewHeapFreeDescriptors, Block->Size);
 
-			// don't want to resize, but optional parameter is missing
-			ReleasedBlocks.RemoveAtSwap(BlockIndex);
-		}
-		else
-		{
-			BlockIndex++;
-		}
-	}
+	Block->SizeUsed = 0;
+	FreeBlocks.Enqueue(Block);
 }
 
 ////////////////////////////////////////////////////////////////////////////////////////////////////////////////////////
diff --git a/Engine/Source/Runtime/D3D12RHI/Private/D3D12Descriptors.h b/Engine/Source/Runtime/D3D12RHI/Private/D3D12Descriptors.h
--- a/Engine/Source/Runtime/D3D12RHI/Private/D3D12Descriptors.h
+++ b/Engine/Source/Runtime/D3D12RHI/Private/D3D12Descriptors.h
@@ -124,7 +124,6 @@ struct FD3D12OnlineDescriptorBlock
 	uint32 BaseSlot;
 	uint32 Size;
 	uint32 SizeUsed = 0;
-	FD3D12SyncPointRef SyncPoint;
 };
 
 /** Primary online heap from which sub blocks can be allocated and freed. Used when allocating blocks of descriptors for tables. */
@@ -146,15 +145,14 @@ public:
 
 	D3D12_CPU_DESCRIPTOR_HANDLE GetCPUSlotHandle(FD3D12OnlineDescriptorBlock* InBlock) const { return Heap->GetCPUSlotHandle(InBlock->BaseSlot); }
 	D3D12_GPU_DESCRIPTOR_HANDLE GetGPUSlotHandle(FD3D12OnlineDescriptorBlock* InBlock) const { return Heap->GetGPUSlotHandle(InBlock->BaseSlot); }
+	
+	// Called by the EOP task to recycle blocks
+	void Recycle(FD3D12OnlineDescriptorBlock* Block);
 
 private:
-	// Check all released blocks and check which ones are not used by the GPU anymore
-	void UpdateFreeBlocks();
-
 	FD3D12DescriptorHeapPtr Heap;
 
 	TQueue<FD3D12OnlineDescriptorBlock*> FreeBlocks;
-	TArray<FD3D12OnlineDescriptorBlock*> ReleasedBlocks;
 
 	FCriticalSection CriticalSection;
 };
diff --git a/Engine/Source/Runtime/D3D12RHI/Private/D3D12RHI.cpp b/Engine/Source/Runtime/D3D12RHI/Private/D3D12RHI.cpp
--- a/Engine/Source/Runtime/D3D12RHI/Private/D3D12RHI.cpp
+++ b/Engine/Source/Runtime/D3D12RHI/Private/D3D12RHI.cpp
@@ -489,6 +489,9 @@ void FD3D12DynamicRHI::RHIPerFrameRHIFlushComplete()
 				case FD3D12DeferredDeleteObject::EType::CPUAllocation:
 					FMemory::Free(ObjectToDelete.CPUAllocation);
 					break;
+				case FD3D12DeferredDeleteObject::EType::DescriptorBlock:
+					ObjectToDelete.DescriptorBlock.Manager->Recycle(ObjectToDelete.DescriptorBlock.Block);
+					break;
 				default:
 					checkf(false, TEXT("Unknown ED3D12DeferredDeleteObjectType"));
 					break;
