Remove use of FRHIBoundShaderState resources in D3D12RHI
 - No longer necessary, since they were only used to find/create root signatures, which can be implemented directly in the adapter class.
 - RHICreateBoundShaderState is never called.

Remove D3D12_SUPPORTS_PARALLEL_RHI_EXECUTE
 - This is always true.

diff --git a/Engine/Source/Developer/Windows/ShaderFormatD3D/Private/D3DShaderCompiler.cpp b/Engine/Source/Developer/Windows/ShaderFormatD3D/Private/D3DShaderCompiler.cpp
--- a/Engine/Source/Developer/Windows/ShaderFormatD3D/Private/D3DShaderCompiler.cpp
+++ b/Engine/Source/Developer/Windows/ShaderFormatD3D/Private/D3DShaderCompiler.cpp
@@ -302,7 +302,7 @@ static FString D3D11CreateShaderCompileCommandLine(
 
 
 // Validate that we are not going over to maximum amount of resource bindings support by the default root signature on DX12
-// Currently limited for hard-coded root signature setup (see: FD3D12RootSignatureDesc::GetStaticGraphicsRootSignatureDesc)
+// Currently limited for hard-coded root signature setup (see: FD3D12Adapter::StaticGraphicsRootSignature)
 // In theory this limitation is only required for DX12, but we don't want a shader to compile on DX11 while not working on DX12.
 // (DX11 has an API limit on 128 SRVs, 16 Samplers, 8 UAVs and 14 CBs but if you go over these values then the shader won't compile)
 bool ValidateResourceCounts(uint32 NumSRVs, uint32 NumSamplers, uint32 NumUAVs, uint32 NumCBs, TArray<FString>& OutFilteredErrors)
diff --git a/Engine/Source/Runtime/D3D12RHI/Private/D3D12Adapter.h b/Engine/Source/Runtime/D3D12RHI/Private/D3D12Adapter.h
--- a/Engine/Source/Runtime/D3D12RHI/Private/D3D12Adapter.h
+++ b/Engine/Source/Runtime/D3D12RHI/Private/D3D12Adapter.h
@@ -187,11 +187,6 @@ public:
 	FORCEINLINE IDXGIFactory6* GetDXGIFactory6() const { return DxgiFactory6; }
 #endif
 
-#if PLATFORM_SUPPORTS_BINDLESS_RENDERING
-	FORCEINLINE bool AreBindlessResourcesAllowed() const { return bBindlessResourcesAllowed; }
-	FORCEINLINE bool AreBindlessSamplersAllowed() const { return bBindlessSamplersAllowed; }
-#endif
-
 	FORCEINLINE const bool IsDebugDevice() const { return bDebugDevice; }
 
 	FORCEINLINE const ED3D12GPUCrashDebuggingModes GetGPUCrashDebuggingModes() const { return GPUCrashDebuggingModes; }
@@ -226,29 +221,10 @@ public:
 
 	FORCEINLINE FD3D12PipelineStateCache& GetPSOCache() { return PipelineStateCache; }
 
-#if USE_STATIC_ROOT_SIGNATURE
-	FORCEINLINE const FD3D12RootSignature* GetStaticGraphicsRootSignature()
-	{
-		return &StaticGraphicsRootSignature;
-	}
-	FORCEINLINE const FD3D12RootSignature* GetStaticComputeRootSignature()
-	{
-		return &StaticComputeRootSignature;
-	}
-	FORCEINLINE const FD3D12RootSignature* GetStaticRayTracingGlobalRootSignature()
-	{
-		return &StaticRayTracingGlobalRootSignature;
-	}
-	FORCEINLINE const FD3D12RootSignature* GetStaticRayTracingLocalRootSignature()
-	{
-		return &StaticRayTracingLocalRootSignature;
-	}
-#endif // USE_STATIC_ROOT_SIGNATURE
-
-	FORCEINLINE FD3D12RootSignature* GetRootSignature(const FD3D12QuantizedBoundShaderState& QBSS) 
-	{
-		return RootSignatureManager.GetRootSignature(QBSS);
-	}
+	const FD3D12RootSignature* GetRootSignature(const FBoundShaderStateInput& BoundShaderState);
+	const FD3D12RootSignature* GetRootSignature(const class FD3D12RayTracingShader* Shader);
+	const FD3D12RootSignature* GetRootSignature(const class FD3D12ComputeShader* Shader);
+	const FD3D12RootSignature* GetGlobalRayTracingRootSignature();
 
 	FORCEINLINE FD3D12RootSignatureManager* GetRootSignatureManager()
 	{
diff --git a/Engine/Source/Runtime/D3D12RHI/Private/D3D12RHI.cpp b/Engine/Source/Runtime/D3D12RHI/Private/D3D12RHI.cpp
--- a/Engine/Source/Runtime/D3D12RHI/Private/D3D12RHI.cpp
+++ b/Engine/Source/Runtime/D3D12RHI/Private/D3D12RHI.cpp
@@ -208,7 +208,7 @@ FD3D12DynamicRHI::FD3D12DynamicRHI(const TArray<TSharedPtr<FD3D12Adapter>>& Chos
 
 	GRHISupportsRHIThread = true;
 
-	GRHISupportsParallelRHIExecute = D3D12_SUPPORTS_PARALLEL_RHI_EXECUTE;
+	GRHISupportsParallelRHIExecute = true;
 
 	GRHISupportsRawViewsForAnyBuffer = true;
 
diff --git a/Engine/Source/Runtime/D3D12RHI/Private/D3D12RHIPrivate.h b/Engine/Source/Runtime/D3D12RHI/Private/D3D12RHIPrivate.h
--- a/Engine/Source/Runtime/D3D12RHI/Private/D3D12RHIPrivate.h
+++ b/Engine/Source/Runtime/D3D12RHI/Private/D3D12RHIPrivate.h
@@ -21,7 +21,6 @@
 #include "EngineGlobals.h"
 #include "StaticBoundShaderState.h"
 
-#define D3D12_SUPPORTS_PARALLEL_RHI_EXECUTE 1
 #define D3D12_RHI_RAYTRACING (RHI_RAYTRACING)
 
 // Dependencies.
@@ -437,7 +436,6 @@ public:
 	virtual void* RHILockStagingBuffer(FRHIStagingBuffer* StagingBuffer, FRHIGPUFence* Fence, uint32 Offset, uint32 SizeRHI) final override;
     virtual void RHIUnlockStagingBuffer(FRHIStagingBuffer* StagingBuffer) final override;
 	virtual FBoundShaderStateRHIRef RHICreateBoundShaderState(FRHIVertexDeclaration* VertexDeclaration, FRHIVertexShader* VertexShader, FRHIPixelShader* PixelShader, FRHIGeometryShader* GeometryShader) final override;
-	FBoundShaderStateRHIRef DX12CreateBoundShaderState(const FBoundShaderStateInput& BoundShaderStateInput);
 	virtual FGraphicsPipelineStateRHIRef RHICreateGraphicsPipelineState(const FGraphicsPipelineStateInitializer& Initializer) final override;
 	virtual TRefCountPtr<FRHIComputePipelineState> RHICreateComputePipelineState(FRHIComputeShader* ComputeShader) final override;
 	virtual void RHICreateTransition(FRHITransition* Transition, const FRHITransitionCreateInfo& CreateInfo) final override;
diff --git a/Engine/Source/Runtime/D3D12RHI/Private/D3D12RayTracing.cpp b/Engine/Source/Runtime/D3D12RHI/Private/D3D12RayTracing.cpp
--- a/Engine/Source/Runtime/D3D12RHI/Private/D3D12RayTracing.cpp
+++ b/Engine/Source/Runtime/D3D12RHI/Private/D3D12RayTracing.cpp
@@ -390,8 +390,6 @@ inline void UnregisterD3D12RayTracingGeometry(FD3D12RayTracingGeometry* Geometry
 #define RAY_TRACING_DESCRIPTOR_CACHE_FULL_COMPARE 1
 #endif // RAY_TRACING_DESCRIPTOR_CACHE_FULL_COMPARE
 
-const FD3D12RootSignature* GetGlobalRayTracingRootSignature(FD3D12Adapter& Adapter);
-
 struct FD3D12ShaderIdentifier
 {
 	uint64 Data[4] = {~0ull, ~0ull, ~0ull, ~0ull};
@@ -1259,7 +1257,7 @@ public:
 	ID3D12RootSignature* GetGlobalRootSignature()
 	{
 		FD3D12Adapter* Adapter = GetParentDevice()->GetParentAdapter();
-		const FD3D12RootSignature* RootSignature = GetGlobalRayTracingRootSignature(*Adapter);
+		const FD3D12RootSignature* RootSignature = Adapter->GetGlobalRayTracingRootSignature();
 		return RootSignature->GetRootSignature();
 	}
 
@@ -5075,7 +5073,7 @@ static void DispatchRays(FD3D12CommandContext& CommandContext,
 
 	// Bind diagnostic buffer to allow asserts in ray generation shaders
 	{
-		const FD3D12RootSignature* RootSignature = GetGlobalRayTracingRootSignature(*Adapter);
+		const FD3D12RootSignature* RootSignature = Adapter->GetGlobalRayTracingRootSignature();
 
 		const int8 DiagnosticBufferSlot = RootSignature->GetDiagnosticBufferSlot();
 		FD3D12Queue& Queue = CommandContext.GetParentDevice()->GetQueue(CommandContext.QueueType);
diff --git a/Engine/Source/Runtime/D3D12RHI/Private/D3D12Shader.h b/Engine/Source/Runtime/D3D12RHI/Private/D3D12Shader.h
--- a/Engine/Source/Runtime/D3D12RHI/Private/D3D12Shader.h
+++ b/Engine/Source/Runtime/D3D12RHI/Private/D3D12Shader.h
@@ -133,54 +133,6 @@ public:
 	const FD3D12RootSignature* RootSignature = nullptr;
 };
 
-/**
-* Combined shader state and vertex definition for rendering geometry.
-* Each unique instance consists of a vertex decl, vertex shader, and pixel shader.
-*/
-class FD3D12BoundShaderState : public FRHIBoundShaderState
-{
-public:
-
-#if D3D12_SUPPORTS_PARALLEL_RHI_EXECUTE
-	FCachedBoundShaderStateLink_Threadsafe CacheLink;
-#else
-	FCachedBoundShaderStateLink CacheLink;
-#endif
-
-	const FD3D12RootSignature* pRootSignature;
-
-	/** Initialization constructor. */
-	FD3D12BoundShaderState(
-		FRHIVertexDeclaration* InVertexDeclarationRHI,
-		FRHIVertexShader* InVertexShaderRHI,
-		FRHIPixelShader* InPixelShaderRHI,
-		FRHIGeometryShader* InGeometryShaderRHI,
-		FD3D12Adapter* InAdapter
-	);
-
-#if PLATFORM_SUPPORTS_MESH_SHADERS
-	/** Initialization constructor. */
-	FD3D12BoundShaderState(
-		FRHIMeshShader* InMeshShaderRHI,
-		FRHIAmplificationShader* InAmplificationShaderRHI,
-		FRHIPixelShader* InPixelShaderRHI,
-		FD3D12Adapter* InAdapter
-	);
-#endif
-
-	virtual ~FD3D12BoundShaderState();
-
-	/**
-	* Get the shader for the given frequency.
-	*/
-	FORCEINLINE FD3D12VertexDeclaration*   GetVertexDeclaration()   const { return (FD3D12VertexDeclaration*)   CacheLink.GetVertexDeclaration();   }
-	FORCEINLINE FD3D12VertexShader*        GetVertexShader()        const { return (FD3D12VertexShader*)        CacheLink.GetVertexShader();        }
-	FORCEINLINE FD3D12MeshShader*          GetMeshShader()          const { return (FD3D12MeshShader*)          CacheLink.GetMeshShader();          }
-	FORCEINLINE FD3D12AmplificationShader* GetAmplificationShader() const { return (FD3D12AmplificationShader*) CacheLink.GetAmplificationShader(); }
-	FORCEINLINE FD3D12PixelShader*         GetPixelShader()         const { return (FD3D12PixelShader*)         CacheLink.GetPixelShader();         }
-	FORCEINLINE FD3D12GeometryShader*      GetGeometryShader()      const { return (FD3D12GeometryShader*)      CacheLink.GetGeometryShader();      }
-};
-
 #if D3D12_RHI_RAYTRACING
 
 class FD3D12RayTracingShader : public FRHIRayTracingShader, public FD3D12ShaderData
@@ -234,8 +186,3 @@ struct TD3D12ResourceTraits<FRHIVertexDeclaration>
 {
 	typedef FD3D12VertexDeclaration TConcreteType;
 };
-template<>
-struct TD3D12ResourceTraits<FRHIBoundShaderState>
-{
-	typedef FD3D12BoundShaderState TConcreteType;
-};
\ No newline at end of file
diff --git a/Engine/Source/Runtime/D3D12RHI/Private/D3D12Shaders.cpp b/Engine/Source/Runtime/D3D12RHI/Private/D3D12Shaders.cpp
--- a/Engine/Source/Runtime/D3D12RHI/Private/D3D12Shaders.cpp
+++ b/Engine/Source/Runtime/D3D12RHI/Private/D3D12Shaders.cpp
@@ -151,14 +151,7 @@ FComputeShaderRHIRef FD3D12DynamicRHI::RHICreateComputeShader(TArrayView<const u
 	FD3D12ComputeShader* Shader = CreateStandardShader<FD3D12ComputeShader>(Code);
 	if (Shader)
 	{
-		FD3D12Adapter& Adapter = GetAdapter();
-
-#if USE_STATIC_ROOT_SIGNATURE
-		Shader->RootSignature = Adapter.GetStaticComputeRootSignature();
-#else
-		const FD3D12QuantizedBoundShaderState QBSS = QuantizeBoundComputeShaderState(Adapter, Shader);
-		Shader->RootSignature = Adapter.GetRootSignature(QBSS);
-#endif
+		Shader->RootSignature = GetAdapter().GetRootSignature(Shader);
 	}
 
 	return Shader;
@@ -204,26 +197,7 @@ FRayTracingShaderRHIRef FD3D12DynamicRHI::RHICreateRayTracingShader(TArrayView<c
 
 	UE::RHICore::InitStaticUniformBufferSlots(Shader->StaticSlots, Shader->ShaderResourceTable);
 
-	FD3D12Adapter& Adapter = GetAdapter();
-
-#if USE_STATIC_ROOT_SIGNATURE
-	switch (ShaderFrequency)
-	{
-	case SF_RayGen:
-		Shader->pRootSignature = Adapter.GetStaticRayTracingGlobalRootSignature();
-		break;
-	case SF_RayHitGroup:
-	case SF_RayCallable:
-	case SF_RayMiss:
-		Shader->pRootSignature = Adapter.GetStaticRayTracingLocalRootSignature();
-		break;
-	default:
-		checkNoEntry(); // Unexpected shader target frequency
-	}
-#else // USE_STATIC_ROOT_SIGNATURE
-	const FD3D12QuantizedBoundShaderState QBSS = QuantizeBoundRayTracingShaderState(Adapter, ShaderFrequency, Shader);
-	Shader->pRootSignature = Adapter.GetRootSignature(QBSS);
-#endif // USE_STATIC_ROOT_SIGNATURE
+	Shader->pRootSignature = GetAdapter().GetRootSignature(Shader);
 
 	return Shader;
 }
@@ -236,133 +210,6 @@ void FD3D12CommandContext::RHISetMultipleViewports(uint32 Count, const FViewport
 	StateCache.SetViewports(Count, reinterpret_cast<const D3D12_VIEWPORT*>(Data));
 }
 
-FD3D12BoundShaderState::FD3D12BoundShaderState(
-	FRHIVertexDeclaration* InVertexDeclarationRHI,
-	FRHIVertexShader* InVertexShaderRHI,
-	FRHIPixelShader* InPixelShaderRHI,
-	FRHIGeometryShader* InGeometryShaderRHI,
-	FD3D12Adapter* InAdapter
-	) :
-	CacheLink(InVertexDeclarationRHI, InVertexShaderRHI, InPixelShaderRHI, InGeometryShaderRHI, this)
-{
-	INC_DWORD_STAT(STAT_D3D12NumBoundShaderState);
-
-#if USE_STATIC_ROOT_SIGNATURE
-	pRootSignature = InAdapter->GetStaticGraphicsRootSignature();
-#else
-	const FD3D12QuantizedBoundShaderState QuantizedBoundShaderState = QuantizeBoundGraphicsShaderState(*InAdapter, this);
-	pRootSignature = InAdapter->GetRootSignature(QuantizedBoundShaderState);
-#endif
-
-#if D3D12_SUPPORTS_PARALLEL_RHI_EXECUTE
-	CacheLink.AddToCache();
-#endif
-}
-
-#if PLATFORM_SUPPORTS_MESH_SHADERS
-FD3D12BoundShaderState::FD3D12BoundShaderState(
-	FRHIMeshShader* InMeshShaderRHI,
-	FRHIAmplificationShader* InAmplificationShaderRHI,
-	FRHIPixelShader* InPixelShaderRHI,
-	FD3D12Adapter* InAdapter
-) :
-	CacheLink(InMeshShaderRHI, InAmplificationShaderRHI, InPixelShaderRHI, this)
-{
-	INC_DWORD_STAT(STAT_D3D12NumBoundShaderState);
-
-#if USE_STATIC_ROOT_SIGNATURE
-	pRootSignature = InAdapter->GetStaticGraphicsRootSignature();
-#else
-	const FD3D12QuantizedBoundShaderState QuantizedBoundShaderState = QuantizeBoundGraphicsShaderState(*InAdapter, this);
-	pRootSignature = InAdapter->GetRootSignature(QuantizedBoundShaderState);
-#endif
-
-#if D3D12_SUPPORTS_PARALLEL_RHI_EXECUTE
-	CacheLink.AddToCache();
-#endif
-}
-#endif // PLATFORM_SUPPORTS_MESH_SHADERS
-
-FD3D12BoundShaderState::~FD3D12BoundShaderState()
-{
-	DEC_DWORD_STAT(STAT_D3D12NumBoundShaderState);
-#if D3D12_SUPPORTS_PARALLEL_RHI_EXECUTE
-	CacheLink.RemoveFromCache();
-#endif
-}
-
-FBoundShaderStateRHIRef FD3D12DynamicRHI::DX12CreateBoundShaderState(const FBoundShaderStateInput& BoundShaderStateInput)
-{
-	//SCOPE_CYCLE_COUNTER(STAT_D3D12CreateBoundShaderStateTime);
-
-	checkf(GIsRHIInitialized, (TEXT("Bound shader state RHI resource was created without initializing Direct3D first")));
-
-#if D3D12_SUPPORTS_PARALLEL_RHI_EXECUTE
-	// Check for an existing bound shader state which matches the parameters
-	FBoundShaderStateRHIRef CachedBoundShaderState = GetCachedBoundShaderState_Threadsafe(
-		BoundShaderStateInput.VertexDeclarationRHI,
-		BoundShaderStateInput.VertexShaderRHI,
-		BoundShaderStateInput.PixelShaderRHI,
-		BoundShaderStateInput.GetGeometryShader(),
-		BoundShaderStateInput.GetMeshShader(),
-		BoundShaderStateInput.GetAmplificationShader()
-	);
-	if (CachedBoundShaderState.GetReference())
-	{
-		// If we've already created a bound shader state with these parameters, reuse it.
-		return CachedBoundShaderState;
-	}
-#else
-	check(IsInRenderingThread() || IsInRHIThread());
-	// Check for an existing bound shader state which matches the parameters
-	FCachedBoundShaderStateLink* CachedBoundShaderStateLink = GetCachedBoundShaderState(
-		BoundShaderStateInput.VertexDeclarationRHI,
-		BoundShaderStateInput.VertexShaderRHI,
-		BoundShaderStateInput.PixelShaderRHI,
-		BoundShaderStateInput.GetGeometryShader(),
-		BoundShaderStateInput.GetMeshShader(),
-		BoundShaderStateInput.GetAmplificationShader()
-	);
-	if (CachedBoundShaderStateLink)
-	{
-		// If we've already created a bound shader state with these parameters, reuse it.
-		return CachedBoundShaderStateLink->BoundShaderState;
-	}
-#endif
-	else
-	{
-		SCOPE_CYCLE_COUNTER(STAT_D3D12NewBoundShaderStateTime);
-
-#if PLATFORM_SUPPORTS_MESH_SHADERS
-		if (BoundShaderStateInput.GetMeshShader())
-		{
-			return new FD3D12BoundShaderState(
-				BoundShaderStateInput.GetMeshShader(),
-				BoundShaderStateInput.GetAmplificationShader(),
-				BoundShaderStateInput.PixelShaderRHI,
-				&GetAdapter());
-		}
-		else
-#endif // PLATFORM_SUPPORTS_MESH_SHADERS
-		{
-			return new FD3D12BoundShaderState(
-				BoundShaderStateInput.VertexDeclarationRHI,
-				BoundShaderStateInput.VertexShaderRHI,
-				BoundShaderStateInput.PixelShaderRHI,
-				BoundShaderStateInput.GetGeometryShader(),
-				&GetAdapter());
-		}
-	}
-}
-
-/**
-* Creates a bound shader state instance which encapsulates a decl, vertex shader, and pixel shader
-* @param VertexDeclaration - existing vertex decl
-* @param StreamStrides - optional stream strides
-* @param VertexShader - existing vertex shader
-* @param PixelShader - existing pixel shader
-* @param GeometryShader - existing geometry shader
-*/
 FBoundShaderStateRHIRef FD3D12DynamicRHI::RHICreateBoundShaderState(
 	FRHIVertexDeclaration* VertexDeclarationRHI,
 	FRHIVertexShader* VertexShaderRHI,
@@ -370,10 +217,6 @@ FBoundShaderStateRHIRef FD3D12DynamicRHI::RHICreateBoundShaderState(
 	FRHIGeometryShader* GeometryShaderRHI
 	)
 {
-	FBoundShaderStateInput Inputs(VertexDeclarationRHI, VertexShaderRHI, PixelShaderRHI
-#if PLATFORM_SUPPORTS_GEOMETRY_SHADERS
-		, GeometryShaderRHI
-#endif
-	);
-	return DX12CreateBoundShaderState(Inputs);
+	checkNoEntry();
+	return nullptr;
 }
diff --git a/Engine/Source/Runtime/D3D12RHI/Private/D3D12State.cpp b/Engine/Source/Runtime/D3D12RHI/Private/D3D12State.cpp
--- a/Engine/Source/Runtime/D3D12RHI/Private/D3D12State.cpp
+++ b/Engine/Source/Runtime/D3D12RHI/Private/D3D12State.cpp
@@ -573,12 +573,9 @@ FGraphicsPipelineStateRHIRef FD3D12DynamicRHI::RHICreateGraphicsPipelineState(co
 
 	TRACE_CPUPROFILER_EVENT_SCOPE(FD3D12DynamicRHI::RHICreateGraphicsPipelineState);
 
-	// TODO: Remove the need for BoundShaderState objects. Currently they are needed for the Root Signature.
-	FBoundShaderStateRHIRef const BSS = DX12CreateBoundShaderState(Initializer.BoundShaderState);
+	const FD3D12RootSignature* RootSignature = GetAdapter().GetRootSignature(Initializer.BoundShaderState);
 
 	// Next try to find the PSO based on the hash of its desc.
-	FD3D12BoundShaderState* const BoundShaderState = FD3D12DynamicRHI::ResourceCast(BSS.GetReference());
-	const FD3D12RootSignature* RootSignature = BoundShaderState->pRootSignature;
 
 	FD3D12LowLevelGraphicsPipelineStateDesc LowLevelDesc;
 #if D3D12RHI_USE_HIGH_LEVEL_PSO_CACHE
diff --git a/Engine/Source/Runtime/D3D12RHI/Private/D3D12StateCachePrivate.h b/Engine/Source/Runtime/D3D12RHI/Private/D3D12StateCachePrivate.h
--- a/Engine/Source/Runtime/D3D12RHI/Private/D3D12StateCachePrivate.h
+++ b/Engine/Source/Runtime/D3D12RHI/Private/D3D12StateCachePrivate.h
@@ -411,12 +411,10 @@ protected:
 
 	void InternalSetStreamSource(FD3D12ResourceLocation* VertexBufferLocation, uint32 StreamIndex, uint32 Stride, uint32 Offset);
 
-	template <typename TShader> struct StateCacheShaderTraits;
+	template <typename TShader> struct TStateCacheShaderTraits;
 #define DECLARE_SHADER_TRAITS(Name) \
-	template <> struct StateCacheShaderTraits<FD3D12##Name##Shader> \
+	template <> struct TStateCacheShaderTraits<FD3D12##Name##Shader> \
 	{ \
-		static const EShaderFrequency Frequency = SF_##Name; \
-		static FD3D12##Name##Shader* GetShader(FD3D12BoundShaderState* BSS) { return BSS ? BSS->Get##Name##Shader() : nullptr; } \
 		static FD3D12##Name##Shader* GetShader(FD3D12GraphicsPipelineState* PSO) { return PSO ? (FD3D12##Name##Shader*)PSO->PipelineStateInitializer.BoundShaderState.Get##Name##Shader() : nullptr; } \
 	}
 	DECLARE_SHADER_TRAITS(Vertex);
@@ -619,7 +617,7 @@ public:
 	template <typename TShader>
 	D3D12_STATE_CACHE_INLINE TShader* GetShader()
 	{
-		return StateCacheShaderTraits<TShader>::GetShader(GetGraphicsPipelineState());
+		return TStateCacheShaderTraits<TShader>::GetShader(GetGraphicsPipelineState());
 	}
 
 	void SetNewShaderData(EShaderFrequency InFrequency, const FD3D12ShaderData* InShaderData);
diff --git a/Engine/Source/Runtime/D3D12RHI/Private/D3D12Util.cpp b/Engine/Source/Runtime/D3D12RHI/Private/D3D12Util.cpp
--- a/Engine/Source/Runtime/D3D12RHI/Private/D3D12Util.cpp
+++ b/Engine/Source/Runtime/D3D12RHI/Private/D3D12Util.cpp
@@ -1057,42 +1057,48 @@ static bool IsCompatibleWithBindlessResources(const FD3D12ShaderData* ShaderData
 	return true;
 }
 
-FD3D12QuantizedBoundShaderState QuantizeBoundGraphicsShaderState(FD3D12Adapter& Adapter, const FD3D12BoundShaderState* const BSS)
+const FD3D12RootSignature* FD3D12Adapter::GetRootSignature(const FBoundShaderStateInput& BSS)
 {
+#if USE_STATIC_ROOT_SIGNATURE
+
+	return &StaticGraphicsRootSignature;
+
+#else //! USE_STATIC_ROOT_SIGNATURE
+
 	// BSS quantizer. There is a 1:1 mapping of quantized bound shader state objects to root signatures.
 	// The objective is to allow a single root signature to represent many bound shader state objects.
 	// The bigger the quantization step sizes, the fewer the root signatures.
 	FD3D12QuantizedBoundShaderState QBSS{};
 
-	QBSS.bAllowIAInputLayout = BSS->GetVertexDeclaration() != nullptr;	// Does the root signature need access to vertex buffers?
+	QBSS.bAllowIAInputLayout = BSS.VertexDeclarationRHI != nullptr;	// Does the root signature need access to vertex buffers?
 
-	const D3D12_RESOURCE_BINDING_TIER ResourceBindingTier = Adapter.GetResourceBindingTier();
+	const D3D12_RESOURCE_BINDING_TIER ResourceBindingTier = GetResourceBindingTier();
 
-	QuantizeBoundShaderStateCommon(QBSS, BSS->GetVertexShader(),        ResourceBindingTier, SV_Vertex);
+	QuantizeBoundShaderStateCommon(QBSS, FD3D12DynamicRHI::ResourceCast(BSS.GetVertexShader()),        ResourceBindingTier, SV_Vertex);
 #if PLATFORM_SUPPORTS_MESH_SHADERS
-	QuantizeBoundShaderStateCommon(QBSS, BSS->GetMeshShader(),          ResourceBindingTier, SV_Mesh);
-	QuantizeBoundShaderStateCommon(QBSS, BSS->GetAmplificationShader(), ResourceBindingTier, SV_Amplification);
+	QuantizeBoundShaderStateCommon(QBSS, FD3D12DynamicRHI::ResourceCast(BSS.GetMeshShader()),          ResourceBindingTier, SV_Mesh);
+	QuantizeBoundShaderStateCommon(QBSS, FD3D12DynamicRHI::ResourceCast(BSS.GetAmplificationShader()), ResourceBindingTier, SV_Amplification);
 #endif
-	QuantizeBoundShaderStateCommon(QBSS, BSS->GetPixelShader(),         ResourceBindingTier, SV_Pixel, true /*bAllowUAVs*/);
-	QuantizeBoundShaderStateCommon(QBSS, BSS->GetGeometryShader(),      ResourceBindingTier, SV_Geometry);
+	QuantizeBoundShaderStateCommon(QBSS, FD3D12DynamicRHI::ResourceCast(BSS.GetPixelShader()),         ResourceBindingTier, SV_Pixel, true /*bAllowUAVs*/);
+	QuantizeBoundShaderStateCommon(QBSS, FD3D12DynamicRHI::ResourceCast(BSS.GetGeometryShader()),      ResourceBindingTier, SV_Geometry);
 
 #if DO_CHECK && PLATFORM_SUPPORTS_BINDLESS_RENDERING
 	if (QBSS.bUseDirectlyIndexedResourceHeap || QBSS.bUseDirectlyIndexedSamplerHeap)
 	{
 		struct FGenericShaderPair
 		{
-			const FRHIGraphicsShader* RHI;
 			const FD3D12ShaderData* Data;
+			const FRHIGraphicsShader* RHI;
 		};
 		const FGenericShaderPair ShaderDatas[] =
 		{
-			{ BSS->GetVertexShader(), BSS->GetVertexShader() },
+			{ FD3D12DynamicRHI::ResourceCast(BSS.GetVertexShader()), BSS.GetVertexShader() },
 #if PLATFORM_SUPPORTS_MESH_SHADERS
-			{ BSS->GetMeshShader(), BSS->GetMeshShader() },
-			{ BSS->GetAmplificationShader(), BSS->GetAmplificationShader() },
+			{ FD3D12DynamicRHI::ResourceCast(BSS.GetMeshShader()), BSS.GetMeshShader() },
+			{ FD3D12DynamicRHI::ResourceCast(BSS.GetAmplificationShader()), BSS.GetAmplificationShader() },
 #endif
-			{ BSS->GetPixelShader(), BSS->GetPixelShader() },
-			{ BSS->GetGeometryShader(), BSS->GetGeometryShader() },
+			{ FD3D12DynamicRHI::ResourceCast(BSS.GetPixelShader()), BSS.GetPixelShader() },
+			{ FD3D12DynamicRHI::ResourceCast(BSS.GetGeometryShader()), BSS.GetGeometryShader() },
 		};
 
 		for (const FGenericShaderPair& ShaderPair : ShaderDatas)
@@ -1112,11 +1118,19 @@ FD3D12QuantizedBoundShaderState QuantizeBoundGraphicsShaderState(FD3D12Adapter&
 	}
 #endif
 
-	return QBSS;
+	return RootSignatureManager.GetRootSignature(QBSS);
+
+#endif //! USE_STATIC_ROOT_SIGNATURE
 }
 
-FD3D12QuantizedBoundShaderState QuantizeBoundComputeShaderState(FD3D12Adapter& Adapter, const FD3D12ComputeShader* const ComputeShader)
+const FD3D12RootSignature* FD3D12Adapter::GetRootSignature(const FD3D12ComputeShader* ComputeShader)
 {
+#if USE_STATIC_ROOT_SIGNATURE
+
+	return &StaticComputeRootSignature;
+
+#else //! USE_STATIC_ROOT_SIGNATURE
+
 	check(ComputeShader);
 
 	// BSS quantizer. There is a 1:1 mapping of quantized bound shader state objects to root signatures.
@@ -1124,16 +1138,25 @@ FD3D12QuantizedBoundShaderState QuantizeBoundComputeShaderState(FD3D12Adapter& A
 	// The bigger the quantization step sizes, the fewer the root signatures.
 	FD3D12QuantizedBoundShaderState QBSS{};
 
-	QuantizeBoundShaderStateCommon(QBSS, ComputeShader, Adapter.GetResourceBindingTier(), SV_All, true /*bAllowUAVs*/);
+	QuantizeBoundShaderStateCommon(QBSS, ComputeShader, GetResourceBindingTier(), SV_All, true /*bAllowUAVs*/);
 
 	check(QBSS.bAllowIAInputLayout == false); // No access to vertex buffers needed
-	return QBSS;
+
+	return RootSignatureManager.GetRootSignature(QBSS);
+
+#endif //! USE_STATIC_ROOT_SIGNATURE
 }
 
 #if D3D12_RHI_RAYTRACING
 
-FD3D12QuantizedBoundShaderState GetRayTracingGlobalRootSignatureDesc(const FD3D12Adapter& Adapter)
+const FD3D12RootSignature* FD3D12Adapter::GetGlobalRayTracingRootSignature()
 {
+#if USE_STATIC_ROOT_SIGNATURE
+
+	return StaticRayTracingGlobalRootSignature;
+
+#else //!USE_STATIC_ROOT_SIGNATURE
+
 	FD3D12QuantizedBoundShaderState QBSS{};
 	FShaderRegisterCounts& QBSSRegisterCounts = QBSS.RegisterCounts[SV_All];
 
@@ -1141,8 +1164,8 @@ FD3D12QuantizedBoundShaderState GetRayTracingGlobalRootSignatureDesc(const FD3D1
 	QBSS.bUseDiagnosticBuffer = true;
 
 #if PLATFORM_SUPPORTS_BINDLESS_RENDERING
-	QBSS.bUseDirectlyIndexedResourceHeap = Adapter.AreBindlessResourcesAllowed();
-	QBSS.bUseDirectlyIndexedSamplerHeap = Adapter.AreBindlessSamplersAllowed();
+	QBSS.bUseDirectlyIndexedResourceHeap = bBindlessResourcesAllowed;
+	QBSS.bUseDirectlyIndexedSamplerHeap = bBindlessSamplersAllowed;
 #endif
 
 	QBSSRegisterCounts.SamplerCount = MAX_SAMPLERS;
@@ -1150,34 +1173,41 @@ FD3D12QuantizedBoundShaderState GetRayTracingGlobalRootSignatureDesc(const FD3D1
 	QBSSRegisterCounts.ConstantBufferCount = MAX_CBS;
 	QBSSRegisterCounts.UnorderedAccessCount = MAX_UAVS;
 
-	return QBSS;
+	return RootSignatureManager.GetRootSignature(QBSS);
+
+#endif //! USE_STATIC_ROOT_SIGNATURE
 }
 
-const FD3D12RootSignature* GetGlobalRayTracingRootSignature(FD3D12Adapter& Adapter)
+const FD3D12RootSignature* FD3D12Adapter::GetRootSignature(const FD3D12RayTracingShader* RayTracingShader)
 {
 #if USE_STATIC_ROOT_SIGNATURE
-	return Adapter.GetStaticRayTracingGlobalRootSignature();
-#else // USE_STATIC_ROOT_SIGNATURE
-	FD3D12QuantizedBoundShaderState RSDesc = GetRayTracingGlobalRootSignatureDesc(Adapter);
-	return Adapter.GetRootSignature(RSDesc);
-#endif // USE_STATIC_ROOT_SIGNATURE
-}
 
-FD3D12QuantizedBoundShaderState QuantizeBoundRayTracingShaderState(FD3D12Adapter& Adapter, EShaderFrequency ShaderFrequency, const FD3D12RayTracingShader* const RayTracingShader
-)
-{
+	switch (RayTracingShader->GetFrequency())
+	{
+	default:
+		checkNoEntry(); // Unexpected shader target frequency
+		return nullptr;
+
+	case SF_RayGen:
+		return StaticRayTracingGlobalRootSignature;
+
+	case SF_RayHitGroup:
+	case SF_RayCallable:
+	case SF_RayMiss:
+		return StaticRayTracingLocalRootSignature;
+	}
+
+#else //! USE_STATIC_ROOT_SIGNATURE
+
 	FD3D12QuantizedBoundShaderState QBSS{};
 
 	FShaderRegisterCounts& QBSSRegisterCounts = QBSS.RegisterCounts[SV_All];
 
-	switch (ShaderFrequency)
+	switch (RayTracingShader->GetFrequency())
 	{
 	case SF_RayGen:
-	{
 		// Shared conservative root signature layout is used for all raygen and miss shaders.
-		QBSS = GetRayTracingGlobalRootSignatureDesc(Adapter);
-		break;
-	}
+		return GetGlobalRayTracingRootSignature();
 
 	case SF_RayHitGroup:
 	case SF_RayCallable:
@@ -1207,7 +1237,9 @@ FD3D12QuantizedBoundShaderState QuantizeBoundRayTracingShaderState(FD3D12Adapter
 
 	SetBoundShaderStateFlags(QBSS, RayTracingShader);
 
-	return QBSS;
+	return RootSignatureManager.GetRootSignature(QBSS);
+
+#endif //! USE_STATIC_ROOT_SIGNATURE
 }
 #endif // D3D12_RHI_RAYTRACING
 
diff --git a/Engine/Source/Runtime/D3D12RHI/Private/Windows/WindowsD3D12PipelineState.h b/Engine/Source/Runtime/D3D12RHI/Private/Windows/WindowsD3D12PipelineState.h
--- a/Engine/Source/Runtime/D3D12RHI/Private/Windows/WindowsD3D12PipelineState.h
+++ b/Engine/Source/Runtime/D3D12RHI/Private/Windows/WindowsD3D12PipelineState.h
@@ -5,7 +5,6 @@
 #pragma once
 
 class FD3D12RootSignature; // forward-declare
-class FD3D12BoundShaderState; // forward-declare
 struct FD3D12ComputePipelineStateDesc;
 struct FD3D12LowLevelGraphicsPipelineStateDesc;
 
diff --git a/Engine/Source/Runtime/D3D12RHI/Public/D3D12Util.h b/Engine/Source/Runtime/D3D12RHI/Public/D3D12Util.h
--- a/Engine/Source/Runtime/D3D12RHI/Public/D3D12Util.h
+++ b/Engine/Source/Runtime/D3D12RHI/Public/D3D12Util.h
@@ -115,21 +115,6 @@ struct FD3D12QuantizedBoundShaderState
 	static void InitShaderRegisterCounts(D3D12_RESOURCE_BINDING_TIER ResourceBindingTier, const FShaderCodePackedResourceCounts& Counts, FShaderRegisterCounts& Shader, bool bAllowUAVs = false);
 };
 
-/**
-* Creates a discrete bound shader state object from a collection of graphics pipeline shaders.
-*/
-
-class FD3D12BoundShaderState;
-class FD3D12ComputeShader;
-
-extern FD3D12QuantizedBoundShaderState QuantizeBoundGraphicsShaderState(FD3D12Adapter& Adapter, const FD3D12BoundShaderState* BSS);
-extern FD3D12QuantizedBoundShaderState QuantizeBoundComputeShaderState(FD3D12Adapter& Adapter, const FD3D12ComputeShader* ComputeShader);
-
-#if D3D12_RHI_RAYTRACING
-class FD3D12RayTracingShader;
-extern FD3D12QuantizedBoundShaderState QuantizeBoundRayTracingShaderState(FD3D12Adapter& Adapter, EShaderFrequency ShaderFrequency, const FD3D12RayTracingShader* Shader);
-#endif
-
 /**
 * Convert from ECubeFace to D3DCUBEMAP_FACES type
 * @param Face - ECubeFace type to convert
