Fix CIS errors

#horde 240066

diff --git a/Engine/Source/Runtime/D3D12RHI/Private/D3D12Util.cpp b/Engine/Source/Runtime/D3D12RHI/Private/D3D12Util.cpp
--- a/Engine/Source/Runtime/D3D12RHI/Private/D3D12Util.cpp
+++ b/Engine/Source/Runtime/D3D12RHI/Private/D3D12Util.cpp
@@ -1153,7 +1153,7 @@ const FD3D12RootSignature* FD3D12Adapter::GetGlobalRayTracingRootSignature()
 {
 #if USE_STATIC_ROOT_SIGNATURE
 
-	return StaticRayTracingGlobalRootSignature;
+	return &StaticRayTracingGlobalRootSignature;
 
 #else //!USE_STATIC_ROOT_SIGNATURE
 
@@ -1189,12 +1189,12 @@ const FD3D12RootSignature* FD3D12Adapter::GetRootSignature(const FD3D12RayTracin
 		return nullptr;
 
 	case SF_RayGen:
-		return StaticRayTracingGlobalRootSignature;
+		return &StaticRayTracingGlobalRootSignature;
 
 	case SF_RayHitGroup:
 	case SF_RayCallable:
 	case SF_RayMiss:
-		return StaticRayTracingLocalRootSignature;
+		return &StaticRayTracingLocalRootSignature;
 	}
 
 #else //! USE_STATIC_ROOT_SIGNATURE
