Cleanup shutdown by making sure the thread pool and file caches are cleared before the global object are deleted post shutdown

diff --git a/Engine/Source/Runtime/D3D12RHI/Private/D3D12Pipelinestate.cpp b/Engine/Source/Runtime/D3D12RHI/Private/D3D12Pipelinestate.cpp
--- a/Engine/Source/Runtime/D3D12RHI/Private/D3D12Pipelinestate.cpp
+++ b/Engine/Source/Runtime/D3D12RHI/Private/D3D12Pipelinestate.cpp
@@ -333,6 +333,18 @@ FD3D12GraphicsPipelineState::FD3D12GraphicsPipelineState(
 {
 	// hold on to bound RHI resources
 	PipelineStateInitializer.BoundShaderState.AddRefResources();
+	if (PipelineStateInitializer.BlendState)
+	{
+		PipelineStateInitializer.BlendState->AddRef();
+	}
+	if (PipelineStateInitializer.RasterizerState)
+	{
+		PipelineStateInitializer.RasterizerState->AddRef();
+	}
+	if (PipelineStateInitializer.DepthStencilState)
+	{
+		PipelineStateInitializer.DepthStencilState->AddRef();
+	}
 
 	if (Initializer.BoundShaderState.VertexDeclarationRHI)
 	{
@@ -377,6 +389,18 @@ FD3D12GraphicsPipelineState::~FD3D12GraphicsPipelineState()
 
 	// release bound RHI resources
 	PipelineStateInitializer.BoundShaderState.ReleaseResources();
+	if (PipelineStateInitializer.BlendState)
+	{
+		PipelineStateInitializer.BlendState->Release();
+	}
+	if (PipelineStateInitializer.RasterizerState)
+	{
+		PipelineStateInitializer.RasterizerState->Release();
+	}
+	if (PipelineStateInitializer.DepthStencilState)
+	{
+		PipelineStateInitializer.DepthStencilState->Release();
+	}
 }
 
 FD3D12ComputePipelineState::FD3D12ComputePipelineState(FD3D12ComputeShader* InComputeShader, const FD3D12RootSignature* InRootSignature, FD3D12PipelineState* InPipelineState)
diff --git a/Engine/Source/Runtime/RHI/Private/PipelineFileCache.cpp b/Engine/Source/Runtime/RHI/Private/PipelineFileCache.cpp
--- a/Engine/Source/Runtime/RHI/Private/PipelineFileCache.cpp
+++ b/Engine/Source/Runtime/RHI/Private/PipelineFileCache.cpp
@@ -3197,6 +3197,7 @@ void FPipelineFileCacheManager::Shutdown()
 		NewPSOHashes.Empty();
         NumNewPSOs = 0;
 		
+		FileCacheMap.Empty();
 		FileCacheEnabled = false;
 		
 		SET_MEMORY_STAT(STAT_NewCachedPSOMemory, 0);
diff --git a/Engine/Source/Runtime/RHI/Private/PipelineStateCache.cpp b/Engine/Source/Runtime/RHI/Private/PipelineStateCache.cpp
--- a/Engine/Source/Runtime/RHI/Private/PipelineStateCache.cpp
+++ b/Engine/Source/Runtime/RHI/Private/PipelineStateCache.cpp
@@ -181,7 +181,8 @@ class FPSOPrecacheThreadPool
 public:
 	~FPSOPrecacheThreadPool()
 	{
-		ShutdownThreadPool();
+		// Thread pool needs to be shutdown before the global object is deleted
+		check(PSOPrecompileCompileThreadPool.load() == nullptr);
 	}
 
 	FQueuedThreadPool& Get()
@@ -241,8 +242,11 @@ static FPSOPrecacheThreadPool GPSOPrecacheThreadPool;
 
 void PipelineStateCache::PreCompileComplete()
 {
-	// free up our threads when the precompile completes.
-	GPSOPrecacheThreadPool.ShutdownThreadPool();
+	// free up our threads when the precompile completes and don't have precaching enabled (otherwise the thread are used during gameplay as well)
+	if (!PipelineStateCache::IsPSOPrecachingEnabled())
+	{
+		GPSOPrecacheThreadPool.ShutdownThreadPool();
+	}
 }
 
 extern RHI_API FRHIComputePipelineState* ExecuteSetComputePipelineState(FComputePipelineState* ComputePipelineState);
@@ -2722,6 +2726,8 @@ void PipelineStateCache::Shutdown()
 		Pair.Value->Release();
 	}
 	GVertexDeclarationCache.Empty();
+
+	GPSOPrecacheThreadPool.ShutdownThreadPool();
 }
 
 FRHIVertexDeclaration*	PipelineStateCache::GetOrCreateVertexDeclaration(const FVertexDeclarationElementList& Elements)
