Fix crash when resizing the window.

FD3D12Queue::ExecutePayload does not increment the fence value if there hasn't been any useful work (command lists to submit) since the last invocation. However, this function is not aware of Present calls, since they happen outside of the submission pipeline. This means that the empty payload submitted by BlockUntilIdle can fail to insert a fence wait if a present happened after the previous payload was executed. This reproduces very easily when resizing the window, but could happen during normal rendering as well. The short-term solution is to force the BlockUntilIdle sync payload to always signal the fence. The correct solution is to execute Present calls as part of the submission pipeline, but this is a more involved change which will happen at a later date.

Also fix numeric error codes in GPU crash messages, instead of the corresponding DXGI_ERROR_* name. This happened because the submission code used VERIFYD3D12RESULT() directly on the error code returned by GetDeviceRemovedReason(), which is incorrect. We must pass the device removed code to VerifyD3D12Result when we detect the error, and let it call GetDeviceRemovedReason to print the correct reason code.

#rnx

diff --git a/Engine/Source/Runtime/D3D12RHI/Private/D3D12Adapter.cpp b/Engine/Source/Runtime/D3D12RHI/Private/D3D12Adapter.cpp
--- a/Engine/Source/Runtime/D3D12RHI/Private/D3D12Adapter.cpp
+++ b/Engine/Source/Runtime/D3D12RHI/Private/D3D12Adapter.cpp
@@ -285,8 +285,9 @@ FD3D12Adapter::FD3D12Adapter(FD3D12AdapterDesc& DescIn)
 /** Callback function called when the GPU crashes, when Aftermath is enabled */
 static void D3D12AftermathCrashCallback(const void* InGPUCrashDump, const uint32_t InGPUCrashDumpSize, void* InUserData)
 {
-	// Forward to shared function which is also called when DEVICE_LOST return value is given
-	D3D12RHI::TerminateOnGPUCrash(nullptr, InGPUCrashDump, InGPUCrashDumpSize);
+	// Disabled for now, let the regular rendering code catch the error on the next API call via VERIFYD3D12RESULT & co.
+	// Note that this means we won't be getting aftermath crash dump files anymore, since those are written based on the data passed to this function via InGPUCrashDump.
+	// D3D12RHI::TerminateOnGPUCrash(nullptr, InGPUCrashDump, InGPUCrashDumpSize);
 }
 
 
diff --git a/Engine/Source/Runtime/D3D12RHI/Private/D3D12Device.cpp b/Engine/Source/Runtime/D3D12RHI/Private/D3D12Device.cpp
--- a/Engine/Source/Runtime/D3D12RHI/Private/D3D12Device.cpp
+++ b/Engine/Source/Runtime/D3D12RHI/Private/D3D12Device.cpp
@@ -510,6 +510,7 @@ void FD3D12Device::BlockUntilIdle()
 
 		FD3D12Payload* Payload = new FD3D12Payload(this, (ED3D12QueueType)QueueTypeIndex);
 		Payload->SyncPointsToSignal.Add(SyncPoint);
+		Payload->bAlwaysSignal = true;
 
 		Payloads.Add(Payload);
 		SyncPoints.Add(SyncPoint);
diff --git a/Engine/Source/Runtime/D3D12RHI/Private/D3D12Submission.cpp b/Engine/Source/Runtime/D3D12RHI/Private/D3D12Submission.cpp
--- a/Engine/Source/Runtime/D3D12RHI/Private/D3D12Submission.cpp
+++ b/Engine/Source/Runtime/D3D12RHI/Private/D3D12Submission.cpp
@@ -821,6 +821,8 @@ uint64 FD3D12Queue::ExecutePayload()
 		bRequiresSignal = true;
 	}
 
+	bRequiresSignal |= PayloadToSubmit->bAlwaysSignal;
+
 	// Keep the latest fence value in the submitted payload.
 	// The interrupt thread uses this to determine when work has completed.
 	uint64 FenceValue = SignalFence();
@@ -925,8 +927,14 @@ FD3D12DynamicRHI::FProcessResult FD3D12DynamicRHI::ProcessInterruptQueue()
 
 			if (CompletedFenceValue == UINT64_MAX)
 			{
-				// If the GPU crashes or hangs, the driver will signal all fences to UINT64_MAX. Check the device is still healthy.
-				VERIFYD3D12RESULT(CurrentQueue.Device->GetDevice()->GetDeviceRemovedReason());
+				// If the GPU crashes or hangs, the driver will signal all fences to UINT64_MAX. If we get an error code here, we can't pass it directly to 
+				// VERIFYD3D12RESULT, because that expects DXGI_ERROR_DEVICE_REMOVED, DXGI_ERROR_DEVICE_RESET etc. and wants to obtain the reason code itself
+				// by calling GetDeviceRemovedReason (again).
+				HRESULT DeviceRemovedReason = CurrentQueue.Device->GetDevice()->GetDeviceRemovedReason();
+				if (DeviceRemovedReason != S_OK)
+				{
+					VerifyD3D12Result(DXGI_ERROR_DEVICE_REMOVED, "CurrentQueue.Fence.D3DFence->GetCompletedValue()", __FILE__, __LINE__, CurrentQueue.Device->GetDevice());
+				}
 			}
 
 			if (CompletedFenceValue < Payload->CompletionFenceValue)
diff --git a/Engine/Source/Runtime/D3D12RHI/Private/D3D12Submission.h b/Engine/Source/Runtime/D3D12RHI/Private/D3D12Submission.h
--- a/Engine/Source/Runtime/D3D12RHI/Private/D3D12Submission.h
+++ b/Engine/Source/Runtime/D3D12RHI/Private/D3D12Submission.h
@@ -207,6 +207,9 @@ struct FD3D12PayloadBase
 	FGraphEventRef SubmissionEvent;
 	TOptional<uint64> SubmissionTime;
 
+	// Flags.
+	bool bAlwaysSignal = false;
+
 	// Cleanup
 	TArray<class FD3D12CommandAllocator*> AllocatorsToRelease;
 	TArray<FD3D12QueryLocation> TimestampQueries;
diff --git a/Engine/Source/Runtime/D3D12RHI/Private/D3D12Util.cpp b/Engine/Source/Runtime/D3D12RHI/Private/D3D12Util.cpp
--- a/Engine/Source/Runtime/D3D12RHI/Private/D3D12Util.cpp
+++ b/Engine/Source/Runtime/D3D12RHI/Private/D3D12Util.cpp
@@ -737,8 +737,7 @@ namespace D3D12RHI
 {
 	void TerminateOnGPUCrash(ID3D12Device* InDevice, const void* InGPUCrashDump, const size_t InGPUCrashDumpSize)
 	{		
-		// Lock the cs, and never unlock - don't want another thread processing the same GPU crash
-		// This call will force a request exit
+		// This function can be called outside of VerifyD3D12Result & co, so it uses its own critical section to make sure it's not re-entered.
 		static FCriticalSection cs;
 		cs.Lock();
 
@@ -846,9 +845,15 @@ namespace D3D12RHI
 		}
 	}
 
+	// It's possible for multiple threads to catch GPU crashes or other D3D errors at the same time. Make sure we only log the error once by acquiring
+	// this critical section inside VerifyD3D12Result and VerifyD3D12CreateTextureResult (and never releasing it, because those functions don't return).
+	static FCriticalSection GD3DCallFailedCS;
+
 	void VerifyD3D12Result(HRESULT D3DResult, const ANSICHAR* Code, const ANSICHAR* Filename, uint32 Line, ID3D12Device* Device, FString Message)
 	{
 		check(FAILED(D3DResult));
+		
+		GD3DCallFailedCS.Lock();
 
 		const FString& ErrorString = GetD3D12ErrorString(D3DResult, Device);
 		UE_LOG(LogD3D12RHI, Error, TEXT("%s failed \n at %s:%u \n with error %s\n%s"), ANSI_TO_TCHAR(Code), ANSI_TO_TCHAR(Filename), Line, *ErrorString, *Message);
@@ -875,6 +880,8 @@ namespace D3D12RHI
 	{
 		check(FAILED(D3DResult));
 
+		GD3DCallFailedCS.Lock();
+
 		const FString ErrorString = GetD3D12ErrorString(D3DResult, nullptr);
 		const TCHAR* D3DFormatString = GetD3D12TextureFormatString(TextureDesc.Format);
 
