Fixed a CIS compile issue found on a non-windows platform.

diff --git a/Engine/Source/Runtime/D3D12RHI/Private/D3D12RHIPrivate.h b/Engine/Source/Runtime/D3D12RHI/Private/D3D12RHIPrivate.h
--- a/Engine/Source/Runtime/D3D12RHI/Private/D3D12RHIPrivate.h
+++ b/Engine/Source/Runtime/D3D12RHI/Private/D3D12RHIPrivate.h
@@ -160,6 +160,14 @@ inline const TCHAR* GetD3DCommandQueueTypeName(ED3D12QueueType QueueType)
 #define FD3D12_HEAP_FLAG_CREATE_NOT_ZEROED D3D12_HEAP_FLAG_CREATE_NOT_ZEROED
 #endif
 
+#ifndef WITH_NV_API
+#define WITH_NV_API 0
+#endif
+
+#ifndef WITH_AMD_AGS
+#define WITH_AMD_AGS 0
+#endif
+
 #if DEBUG_RESOURCE_STATES
 	#define LOG_EXECUTE_COMMAND_LISTS 1
 	#define ASSERT_RESOURCE_STATES 0	// Disabled for now.
diff --git a/Engine/Source/Runtime/Windows/D3D11RHI/Private/D3D11RHIPrivate.h b/Engine/Source/Runtime/Windows/D3D11RHI/Private/D3D11RHIPrivate.h
--- a/Engine/Source/Runtime/Windows/D3D11RHI/Private/D3D11RHIPrivate.h
+++ b/Engine/Source/Runtime/Windows/D3D11RHI/Private/D3D11RHIPrivate.h
@@ -102,6 +102,14 @@ THIRD_PARTY_INCLUDES_END
 #define D3D11_NUM_THREAD_LOCAL_CACHES 1
 #endif
 
+#ifndef WITH_NV_API
+#define WITH_NV_API 0
+#endif
+
+#ifndef WITH_AMD_AGS
+#define WITH_AMD_AGS 0
+#endif
+
 /**
  * The D3D RHI stats.
  */
