MGPU resource cast fix -- need to use RetrieveObject, which pulls the GPU index from the context to get the correct GPU version.

#rnx

diff --git a/Engine/Source/Runtime/D3D12RHI/Private/D3D12Commands.cpp b/Engine/Source/Runtime/D3D12RHI/Private/D3D12Commands.cpp
--- a/Engine/Source/Runtime/D3D12RHI/Private/D3D12Commands.cpp
+++ b/Engine/Source/Runtime/D3D12RHI/Private/D3D12Commands.cpp
@@ -1529,14 +1529,14 @@ void FD3D12CommandContext::SetResourcesFromTables(const ShaderType* RESTRICT Sha
 		{
 			// Actually set the UAV
 			uint32 InitialCount = -1;
-			FD3D12UnorderedAccessView* D3D12UAV = FD3D12DynamicRHI::ResourceCast(UAV);
+			FD3D12UnorderedAccessView* D3D12UAV = Context.RetrieveObject<FD3D12UnorderedAccessView>(UAV);
 
 			Context.StateCache.SetUAVs<SF_Pixel>(Index, 1, &D3D12UAV, &InitialCount);
 		}
 
 		void SetSRV(FRHIShaderResourceView* SRV, uint8 Index)
 		{
-			Context.StateCache.SetShaderResourceView<Frequency>(FD3D12DynamicRHI::ResourceCast(SRV), Index);
+			Context.StateCache.SetShaderResourceView<Frequency>(Context.RetrieveObject<FD3D12ShaderResourceView>(SRV), Index);
 		}
 
 		void SetTexture(FRHITexture* TextureRHI, uint8 Index)
@@ -1547,7 +1547,7 @@ void FD3D12CommandContext::SetResourcesFromTables(const ShaderType* RESTRICT Sha
 
 		void SetSampler(FRHISamplerState* Sampler, uint8 Index)
 		{
-			Context.StateCache.SetSamplerState<Frequency>(FD3D12DynamicRHI::ResourceCast(Sampler), Index);
+			Context.StateCache.SetSamplerState<Frequency>(Context.RetrieveObject<FD3D12SamplerState>(Sampler), Index);
 		}
 	};
 
