Add clearPSODriverCache commandline parameter to be able to clear the driver disk cache. Only works on NVidia and AMD right now

diff --git a/Engine/Source/Runtime/D3D12RHI/Private/Windows/WindowsD3D12Device.cpp b/Engine/Source/Runtime/D3D12RHI/Private/Windows/WindowsD3D12Device.cpp
--- a/Engine/Source/Runtime/D3D12RHI/Private/Windows/WindowsD3D12Device.cpp
+++ b/Engine/Source/Runtime/D3D12RHI/Private/Windows/WindowsD3D12Device.cpp
@@ -9,6 +9,7 @@
 #include "Modules/ModuleManager.h"
 #include "Windows/AllowWindowsPlatformTypes.h"
 #include "Windows/WindowsPlatformCrashContext.h"
+#include "HAL/FileManager.h"
 #include <delayimp.h>
 
 #if WITH_AMD_AGS
@@ -1281,6 +1282,58 @@ static void DisableRayTracingSupport()
 	GRHISupportsRayTracingShaders = false;
 }
 
+static void ClearPSODriverCache()
+{
+#if !PLATFORM_HOLOLENS
+	TArray<FString> Files;
+	IFileManager& FileManager = IFileManager::Get();
+    FString PSOPath;
+	if (IsRHIDeviceNVIDIA())
+	{
+		PSOPath = FPlatformMisc::GetEnvironmentVariable(TEXT("LOCALAPPDATA"));
+		if (!PSOPath.IsEmpty())
+		{
+			PSOPath = FString::Printf(L"%s\\NVIDIA\\DXCache", *PSOPath);
+			FileManager.FindFiles(Files, *PSOPath, TEXT(".bin"));
+			FileManager.FindFiles(Files, *PSOPath, TEXT(".toc"));
+		}
+		else
+		{
+			UE_LOG(LogD3D12RHI, Error, TEXT("clearPSODriverCache failed: please ensure that LOCALAPPDATA points to [REDACTED_USER_HOME]\\AppData\\Local"));
+		}
+	}
+	else if (IsRHIDeviceAMD())
+	{
+		PSOPath = FPlatformMisc::GetEnvironmentVariable(TEXT("LOCALAPPDATA"));
+		if (!PSOPath.IsEmpty())
+		{
+			PSOPath = FString::Printf(L"%s\\AMD\\DxCache", *PSOPath);
+			FileManager.FindFiles(Files, *PSOPath, TEXT(".bin"));
+		}
+		else
+		{
+			UE_LOG(LogD3D12RHI, Error, TEXT("clearPSODriverCache failed: please ensure that LOCALAPPDATA points to [REDACTED_USER_HOME]\\AppData\\Local"));
+		}
+	}
+	else if (IsRHIDeviceIntel())
+	{
+		UE_LOG(LogD3D12RHI, Warning, TEXT("clearPSODriverCache not implemented yet for Intel"));
+	}
+
+	if (!PSOPath.IsEmpty() && Files.IsEmpty())
+	{
+		UE_LOG(LogD3D12RHI, Warning, TEXT("clearPSODriverCache hasn't found any file under \"%s\""), *PSOPath);
+	}
+
+	for (FString& File : Files)
+	{
+		FString FilePath = FString::Printf(L"%s\\%s", *PSOPath, *File);
+		FileManager.Delete(*FilePath, /*RequireExists*/ false, /*EvenReadOnly*/ true, /*Quiet*/ true);
+	}
+
+#endif
+}
+
 void FD3D12DynamicRHI::Init()
 {
 #if D3D12RHI_SUPPORTS_WIN_PIX
@@ -1501,6 +1554,11 @@ void FD3D12DynamicRHI::Init()
 	}
 #endif // INTEL_EXTENSIONS
 
+	if (FParse::Param(FCommandLine::Get(), TEXT("clearPSODriverCache")))
+    {
+	    ClearPSODriverCache();
+    }
+
 	GRHIPersistentThreadGroupCount = 1440; // TODO: Revisit based on vendor/adapter/perf query
 
 	GTexturePoolSize = 0;
