Fixing issues when using D3D12 SM6 on Intel cards.
Fix provided by Intel.

#rnx

diff --git a/Engine/Source/Runtime/D3D12RHI/Private/Windows/WindowsD3D12Device.cpp b/Engine/Source/Runtime/D3D12RHI/Private/Windows/WindowsD3D12Device.cpp
--- a/Engine/Source/Runtime/D3D12RHI/Private/Windows/WindowsD3D12Device.cpp
+++ b/Engine/Source/Runtime/D3D12RHI/Private/Windows/WindowsD3D12Device.cpp
@@ -1541,9 +1541,6 @@ void FD3D12DynamicRHI::Init()
 		if (IntelExtensionContext)
 		{
 			bHasVendorSupportForAtomic64 = (INTCExtensionInfo.RequestedExtensionVersion.HWFeatureLevel > 0);
-
-			GRHISupportsAtomicUInt64 = EnableIntelAtomic64Support(IntelExtensionContext, INTCExtensionInfo);
-			GRHISupportsDX12AtomicUInt64 = GRHISupportsAtomicUInt64;
 		}
 
 		if (GRHISupportsMeshShadersTier0 || GRHISupportsMeshShadersTier1)
