- Fix file and extensions for AMD PSO cache clearing
- Move PSO cache clear function eariler otherwise the files are locked already by the AMD driver

#rnx

diff --git a/Engine/Source/Runtime/D3D12RHI/Private/Windows/WindowsD3D12Device.cpp b/Engine/Source/Runtime/D3D12RHI/Private/Windows/WindowsD3D12Device.cpp
--- a/Engine/Source/Runtime/D3D12RHI/Private/Windows/WindowsD3D12Device.cpp
+++ b/Engine/Source/Runtime/D3D12RHI/Private/Windows/WindowsD3D12Device.cpp
@@ -1285,52 +1285,44 @@ static void DisableRayTracingSupport()
 static void ClearPSODriverCache()
 {
 #if !PLATFORM_HOLOLENS
-	TArray<FString> Files;
-	IFileManager& FileManager = IFileManager::Get();
-    FString PSOPath;
-	if (IsRHIDeviceNVIDIA())
+	FString LocalAppDataFolder = FPlatformMisc::GetEnvironmentVariable(TEXT("LOCALAPPDATA"));
+	if (!LocalAppDataFolder.IsEmpty())
 	{
-		PSOPath = FPlatformMisc::GetEnvironmentVariable(TEXT("LOCALAPPDATA"));
-		if (!PSOPath.IsEmpty())
+		auto ClearFolder = [](const FString& PSOPath, const TCHAR* Extension)
 		{
-			PSOPath = FString::Printf(L"%s\\NVIDIA\\DXCache", *PSOPath);
-			FileManager.FindFiles(Files, *PSOPath, TEXT(".bin"));
-			FileManager.FindFiles(Files, *PSOPath, TEXT(".toc"));
-		}
-		else
+			TArray<FString> Files;
+			IFileManager& FileManager = IFileManager::Get();
+			FileManager.FindFiles(Files, *PSOPath, Extension);
+			for (FString& File : Files)
+			{
+				FString FilePath = FString::Printf(L"%s\\%s", *PSOPath, *File);
+				FileManager.Delete(*FilePath, /*RequireExists*/ false, /*EvenReadOnly*/ true, /*Quiet*/ true);
+			}
+		};
+
+		if (IsRHIDeviceNVIDIA())
 		{
-			UE_LOG(LogD3D12RHI, Error, TEXT("clearPSODriverCache failed: please ensure that LOCALAPPDATA points to [REDACTED_USER_HOME]\\AppData\\Local"));
+			FString PSOPath = FString::Printf(L"%s\\NVIDIA\\DXCache", *LocalAppDataFolder);
+			ClearFolder(PSOPath, TEXT(".bin"));
+			ClearFolder(PSOPath, TEXT(".toc"));
 		}
-	}
-	else if (IsRHIDeviceAMD())
-	{
-		PSOPath = FPlatformMisc::GetEnvironmentVariable(TEXT("LOCALAPPDATA"));
-		if (!PSOPath.IsEmpty())
+		else if (IsRHIDeviceAMD())
 		{
-			PSOPath = FString::Printf(L"%s\\AMD\\DxCache", *PSOPath);
-			FileManager.FindFiles(Files, *PSOPath, TEXT(".bin"));
+			FString PSOPath = FString::Printf(L"%s\\AMD\\DxCache", *LocalAppDataFolder);
+			ClearFolder(PSOPath, TEXT(".bin"));
+
+			PSOPath = FString::Printf(L"%s\\AMD\\DxcCache", *LocalAppDataFolder);
+			ClearFolder(PSOPath, TEXT(".parc"));
 		}
-		else
+		else if (IsRHIDeviceIntel())
 		{
-			UE_LOG(LogD3D12RHI, Error, TEXT("clearPSODriverCache failed: please ensure that LOCALAPPDATA points to [REDACTED_USER_HOME]\\AppData\\Local"));
+			UE_LOG(LogD3D12RHI, Warning, TEXT("clearPSODriverCache not implemented yet for Intel"));
 		}
 	}
-	else if (IsRHIDeviceIntel())
-	{
-		UE_LOG(LogD3D12RHI, Warning, TEXT("clearPSODriverCache not implemented yet for Intel"));
-	}
-
-	if (!PSOPath.IsEmpty() && Files.IsEmpty())
-	{
-		UE_LOG(LogD3D12RHI, Warning, TEXT("clearPSODriverCache hasn't found any file under \"%s\""), *PSOPath);
-	}
-
-	for (FString& File : Files)
+	else
 	{
-		FString FilePath = FString::Printf(L"%s\\%s", *PSOPath, *File);
-		FileManager.Delete(*FilePath, /*RequireExists*/ false, /*EvenReadOnly*/ true, /*Quiet*/ true);
+		UE_LOG(LogD3D12RHI, Error, TEXT("clearPSODriverCache failed: please ensure that LOCALAPPDATA points to [REDACTED_USER_HOME]\\AppData\\Local"));
 	}
-
 #endif
 }
 
@@ -1352,6 +1344,12 @@ void FD3D12DynamicRHI::Init()
 
 	check(!GIsRHIInitialized);
 
+	if (FParse::Param(FCommandLine::Get(), TEXT("clearPSODriverCache")))
+	{
+		ClearPSODriverCache();
+	}
+
+
 	const DXGI_ADAPTER_DESC& AdapterDesc = GetAdapter().GetD3DAdapterDesc();
 
 	UE_LOG(LogD3D12RHI, Log, TEXT("    GPU DeviceId: 0x%x (for the marketing name, search the web for \"GPU Device Id\")"), AdapterDesc.DeviceId);
@@ -1551,11 +1549,6 @@ void FD3D12DynamicRHI::Init()
 	}
 #endif // INTEL_EXTENSIONS
 
-	if (FParse::Param(FCommandLine::Get(), TEXT("clearPSODriverCache")))
-    {
-	    ClearPSODriverCache();
-    }
-
 	GRHIPersistentThreadGroupCount = 1440; // TODO: Revisit based on vendor/adapter/perf query
 
 	GTexturePoolSize = 0;
