Change FD3D12Device::ObtainContext to use CreateCommandContext so derived classes can create the correct context.

diff --git a/Engine/Source/Runtime/D3D12RHI/Private/D3D12Device.cpp b/Engine/Source/Runtime/D3D12RHI/Private/D3D12Device.cpp
--- a/Engine/Source/Runtime/D3D12RHI/Private/D3D12Device.cpp
+++ b/Engine/Source/Runtime/D3D12RHI/Private/D3D12Device.cpp
@@ -584,9 +584,9 @@ FD3D12ContextCommon* FD3D12Device::ObtainContext(ED3D12QueueType QueueType)
 		switch (QueueType)
 		{
 		default: checkNoEntry(); // fallthrough
-		case ED3D12QueueType::Direct: Context = new FD3D12CommandContext(this, QueueType, false); break;
-		case ED3D12QueueType::Async : Context = new FD3D12CommandContext(this, QueueType, false); break;
-		case ED3D12QueueType::Copy  : Context = new FD3D12ContextCopy   (this);                   break;
+		case ED3D12QueueType::Direct: Context = FD3D12DynamicRHI::GetD3DRHI()->CreateCommandContext(this, QueueType, false); break;
+		case ED3D12QueueType::Async : Context = FD3D12DynamicRHI::GetD3DRHI()->CreateCommandContext(this, QueueType, false); break;
+		case ED3D12QueueType::Copy  : Context = new FD3D12ContextCopy(this); break;
 		}
 	}
 
