Limit D3D12 ray tracing setup in FD3D12Adapter::InitializeDevices() to PC and Hololens

- other D3D12 platforms already do their own ray tracing setup and this code results in erroneous log messages about the state of ray tracing support

diff --git a/Engine/Source/Runtime/D3D12RHI/Private/D3D12Adapter.cpp b/Engine/Source/Runtime/D3D12RHI/Private/D3D12Adapter.cpp
--- a/Engine/Source/Runtime/D3D12RHI/Private/D3D12Adapter.cpp
+++ b/Engine/Source/Runtime/D3D12RHI/Private/D3D12Adapter.cpp
@@ -974,6 +974,7 @@ PRAGMA_ENABLE_DEPRECATION_WARNINGS
 			}
 
 #if D3D12_RHI_RAYTRACING
+#if PLATFORM_WINDOWS || PLATFORM_HOLOLENS
 			D3D12_FEATURE_DATA_D3D12_OPTIONS5 D3D12Caps5 = {};
 			if (SUCCEEDED(RootDevice->CheckFeatureSupport(D3D12_FEATURE_D3D12_OPTIONS5, &D3D12Caps5, sizeof(D3D12Caps5))))
 			{
@@ -1023,6 +1024,7 @@ PRAGMA_ENABLE_DEPRECATION_WARNINGS
 					UE_LOG(LogD3D12RHI, Warning, TEXT("Ray Tracing is disabled because the RenderDoc plugin is currently not compatible with D3D12 ray tracing."));
 				}
 			}
+#endif // PLATFORM_WINDOWS || PLATFORM_HOLOLENS
 
 			GRHIRayTracingAccelerationStructureAlignment = uint32(D3D12_RAYTRACING_ACCELERATION_STRUCTURE_BYTE_ALIGNMENT);
 			GRHIRayTracingScratchBufferAlignment = uint32(D3D12_RAYTRACING_ACCELERATION_STRUCTURE_BYTE_ALIGNMENT);
