RayTracing shouldn't override heaps when bindless is enabled.

diff --git a/Engine/Source/Runtime/D3D12RHI/Private/D3D12DescriptorCache.cpp b/Engine/Source/Runtime/D3D12RHI/Private/D3D12DescriptorCache.cpp
--- a/Engine/Source/Runtime/D3D12RHI/Private/D3D12DescriptorCache.cpp
+++ b/Engine/Source/Runtime/D3D12RHI/Private/D3D12DescriptorCache.cpp
@@ -670,13 +670,19 @@ bool FD3D12DescriptorCache::SetHeapOverrides(FD3D12DescriptorHeap* InOverrideVie
 
 void FD3D12DescriptorCache::OverrideLastSetHeaps(ID3D12DescriptorHeap* ViewHeap, ID3D12DescriptorHeap* SamplerHeap)
 {
-	LastSetViewHeap = ViewHeap;
-	LastSetSamplerHeap = SamplerHeap;
+	ID3D12DescriptorHeap* ViewHeapToSet = ViewHeap ? ViewHeap : LastSetViewHeap;
+	ID3D12DescriptorHeap* SamplerHeapToSet = SamplerHeap ? SamplerHeap : LastSetSamplerHeap;
 
-	ID3D12DescriptorHeap* ppHeaps[] = { ViewHeap, SamplerHeap };
-	Context.GraphicsCommandList()->SetDescriptorHeaps(UE_ARRAY_COUNT(ppHeaps), ppHeaps);
+	if (ViewHeapToSet != LastSetViewHeap || SamplerHeapToSet != LastSetSamplerHeap)
+	{
+		LastSetViewHeap = ViewHeapToSet;
+		LastSetSamplerHeap = SamplerHeapToSet;
+
+		ID3D12DescriptorHeap* ppHeaps[] = { ViewHeapToSet, SamplerHeapToSet };
+		Context.GraphicsCommandList()->SetDescriptorHeaps(UE_ARRAY_COUNT(ppHeaps), ppHeaps);
 
-	bHeapsOverridden = true;
+		bHeapsOverridden = true;
+	}
 }
 
 void FD3D12DescriptorCache::RestoreAfterExternalHeapsSet()
diff --git a/Engine/Source/Runtime/D3D12RHI/Private/D3D12Descriptors.cpp b/Engine/Source/Runtime/D3D12RHI/Private/D3D12Descriptors.cpp
--- a/Engine/Source/Runtime/D3D12RHI/Private/D3D12Descriptors.cpp
+++ b/Engine/Source/Runtime/D3D12RHI/Private/D3D12Descriptors.cpp
@@ -179,6 +179,19 @@ FD3D12DescriptorHeap* FD3D12BindlessDescriptorManager::GetHeapForType(ERHIDescri
 	return nullptr;
 }
 
+bool FD3D12BindlessDescriptorManager::HasHeapForType(ERHIDescriptorHeapType InType) const
+{
+	for (const FD3D12DescriptorManager& Manager : Managers)
+	{
+		if (Manager.HandlesAllocation(InType))
+		{
+			return true;
+		}
+	}
+
+	return false;
+}
+
 D3D12_GPU_DESCRIPTOR_HANDLE FD3D12BindlessDescriptorManager::GetGpuHandle(FRHIDescriptorHandle InHandle) const
 {
 	for (const FD3D12DescriptorManager& Manager : Managers)
diff --git a/Engine/Source/Runtime/D3D12RHI/Private/D3D12Descriptors.h b/Engine/Source/Runtime/D3D12RHI/Private/D3D12Descriptors.h
--- a/Engine/Source/Runtime/D3D12RHI/Private/D3D12Descriptors.h
+++ b/Engine/Source/Runtime/D3D12RHI/Private/D3D12Descriptors.h
@@ -104,6 +104,7 @@ public:
 	void UpdateDeferred(FRHIDescriptorHandle InHandle, D3D12_CPU_DESCRIPTOR_HANDLE InSourceCpuHandle);
 
 	FD3D12DescriptorHeap* GetHeapForType(ERHIDescriptorHeapType InType);
+	bool HasHeapForType(ERHIDescriptorHeapType InType) const;
 
 	D3D12_GPU_DESCRIPTOR_HANDLE GetGpuHandle(FRHIDescriptorHandle InHandle) const;
 
diff --git a/Engine/Source/Runtime/D3D12RHI/Private/D3D12RayTracing.cpp b/Engine/Source/Runtime/D3D12RHI/Private/D3D12RayTracing.cpp
--- a/Engine/Source/Runtime/D3D12RHI/Private/D3D12RayTracing.cpp
+++ b/Engine/Source/Runtime/D3D12RHI/Private/D3D12RayTracing.cpp
@@ -1597,6 +1597,11 @@ public:
 
 	void Init(uint32 NumViewDescriptors, uint32 NumSamplerDescriptors)
 	{
+		FD3D12BindlessDescriptorManager& BindlessManager = GetParentDevice()->GetBindlessDescriptorManager();
+
+		bBindlessViews = BindlessManager.HasHeapForType(ERHIDescriptorHeapType::Standard);
+		bBindlessSamplers = BindlessManager.HasHeapForType(ERHIDescriptorHeapType::Sampler);
+
 		ViewHeap.Init(NumViewDescriptors, D3D12_DESCRIPTOR_HEAP_TYPE_CBV_SRV_UAV);
 		SamplerHeap.Init(NumSamplerDescriptors, D3D12_DESCRIPTOR_HEAP_TYPE_SAMPLER);
 	}
@@ -1615,7 +1620,10 @@ public:
 		check(ViewHeap.GetParentDevice() == CommandContext.GetParentDevice());
 		check(SamplerHeap.GetParentDevice() == CommandContext.GetParentDevice());
 
-		CommandContext.StateCache.GetDescriptorCache()->OverrideLastSetHeaps(ViewHeap.D3D12Heap, SamplerHeap.D3D12Heap);
+		ID3D12DescriptorHeap* ViewHeapToSet = bBindlessViews ? nullptr : ViewHeap.D3D12Heap;
+		ID3D12DescriptorHeap* SamplerHeapToSet = bBindlessSamplers ? nullptr : SamplerHeap.D3D12Heap;
+
+		CommandContext.StateCache.GetDescriptorCache()->OverrideLastSetHeaps(ViewHeapToSet, SamplerHeapToSet);
 	}
 
 	// Returns descriptor heap base index for this descriptor table allocation or -1 if allocation failed.
@@ -1703,6 +1711,9 @@ public:
 	FD3D12RayTracingDescriptorHeap ViewHeap;
 	FD3D12RayTracingDescriptorHeap SamplerHeap;
 
+	bool bBindlessViews = false;
+	bool bBindlessSamplers = false;
+
 	template<typename KeyType>
 	struct TIdentityHash
 	{
