Fix scratch buffer size assert message format string (%lld was incorrectly used to format 32 bit integers)

diff --git a/Engine/Source/Runtime/D3D12RHI/Private/D3D12RayTracing.cpp b/Engine/Source/Runtime/D3D12RHI/Private/D3D12RayTracing.cpp
--- a/Engine/Source/Runtime/D3D12RHI/Private/D3D12RayTracing.cpp
+++ b/Engine/Source/Runtime/D3D12RHI/Private/D3D12RayTracing.cpp
@@ -4040,7 +4040,7 @@ void FD3D12RayTracingScene::BuildAccelerationStructure(FD3D12CommandContext& Com
 			PrebuildInfo.ScratchDataSizeInBytes, Layer.SizeInfo.BuildScratchSize);
 
 		checkf(PrebuildInfo.ScratchDataSizeInBytes + ScratchBufferOffset <= ScratchBuffer->GetSize(),
-			TEXT("TLAS scratch buffer size is %lld bytes with offset %lld (%lld bytes available), but the build requires %lld bytes. ")
+			TEXT("TLAS scratch buffer size is %d bytes with offset %d (%d bytes available), but the build requires %lld bytes. ")
 			TEXT("BuildInputs.NumDescs = %d, Instances.Num = %d, PerInstanceGeometries.Num = %d."),
 			ScratchBuffer->GetSize(), ScratchBufferOffset, ScratchBuffer->GetSize() - ScratchBufferOffset,
 			PrebuildInfo.ScratchDataSizeInBytes,
