Changing strongly typed D3D12 code to be based on shader frequency indices. This reduces the amount of duplicated code in D3D12 RHISetXXX commands.

diff --git a/Engine/Source/Runtime/D3D12RHI/Private/D3D12CommandContext.cpp b/Engine/Source/Runtime/D3D12RHI/Private/D3D12CommandContext.cpp
--- a/Engine/Source/Runtime/D3D12RHI/Private/D3D12CommandContext.cpp
+++ b/Engine/Source/Runtime/D3D12RHI/Private/D3D12CommandContext.cpp
@@ -103,12 +103,14 @@ FD3D12CommandContext::FD3D12CommandContext(FD3D12Device* InParent, ED3D12QueueTy
 	, ConstantsAllocator(InParent, InParent->GetGPUMask())
 	, StateCache(*this, InParent->GetGPUMask())
 	, ValidResourceStates(GetValidResourceStates(QueueType))
-	, VSConstantBuffer(InParent, ConstantsAllocator)
-	, MSConstantBuffer(InParent, ConstantsAllocator)
-	, ASConstantBuffer(InParent, ConstantsAllocator)
-	, PSConstantBuffer(InParent, ConstantsAllocator)
-	, GSConstantBuffer(InParent, ConstantsAllocator)
-	, CSConstantBuffer(InParent, ConstantsAllocator)
+	, StageConstantBuffers{
+		FD3D12ConstantBuffer(InParent, ConstantsAllocator),
+		FD3D12ConstantBuffer(InParent, ConstantsAllocator),
+		FD3D12ConstantBuffer(InParent, ConstantsAllocator),
+		FD3D12ConstantBuffer(InParent, ConstantsAllocator),
+		FD3D12ConstantBuffer(InParent, ConstantsAllocator),
+		FD3D12ConstantBuffer(InParent, ConstantsAllocator),
+	}
 {
 	StaticUniformBuffers.AddZeroed(FUniformBufferStaticSlotRegistry::Get().GetSlotCount());
 	ClearState();
@@ -630,12 +632,11 @@ void FD3D12CommandContext::ClearState(EClearStateMode Mode)
 void FD3D12CommandContext::ConditionalClearShaderResource(FD3D12ResourceLocation* Resource)
 {
 	check(Resource);
-	StateCache.ClearShaderResourceViews<SF_Vertex>(Resource);
-	StateCache.ClearShaderResourceViews<SF_Mesh>(Resource);
-	StateCache.ClearShaderResourceViews<SF_Amplification>(Resource);
-	StateCache.ClearShaderResourceViews<SF_Pixel>(Resource);
-	StateCache.ClearShaderResourceViews<SF_Geometry>(Resource);
-	StateCache.ClearShaderResourceViews<SF_Compute>(Resource);
+
+	for (int32 Index = 0; Index < SF_NumStandardFrequencies; Index++)
+	{
+		StateCache.ClearShaderResourceViews(static_cast<EShaderFrequency>(Index), Resource);
+	}
 }
 
 void FD3D12CommandContext::ClearAllShaderResources()
diff --git a/Engine/Source/Runtime/D3D12RHI/Private/D3D12CommandContext.h b/Engine/Source/Runtime/D3D12RHI/Private/D3D12CommandContext.h
--- a/Engine/Source/Runtime/D3D12RHI/Private/D3D12CommandContext.h
+++ b/Engine/Source/Runtime/D3D12RHI/Private/D3D12CommandContext.h
@@ -477,6 +477,9 @@ public:
 	D3D12_RESOURCE_STATES SkipFastClearEliminateState = D3D12_RESOURCE_STATES(0);
 	D3D12_RESOURCE_STATES ValidResourceStates;
 
+	/** Constant buffers for Set*ShaderParameter calls. */
+	FD3D12ConstantBuffer StageConstantBuffers[SF_NumStandardFrequencies];
+
 #if PLATFORM_SUPPORTS_VIRTUAL_TEXTURES
 	bool bNeedFlushTextureCache = false;
 	void InvalidateTextureCache() { bNeedFlushTextureCache = true; }
@@ -492,14 +495,6 @@ public:
 	virtual void FlushTextureCache() {};
 #endif
 
-	/** Constant buffers for Set*ShaderParameter calls. */
-	FD3D12ConstantBuffer VSConstantBuffer;
-	FD3D12ConstantBuffer MSConstantBuffer;
-	FD3D12ConstantBuffer ASConstantBuffer;
-	FD3D12ConstantBuffer PSConstantBuffer;
-	FD3D12ConstantBuffer GSConstantBuffer;
-	FD3D12ConstantBuffer CSConstantBuffer;
-
 	/** needs to be called before each draw call */
 	void CommitNonComputeShaderConstants();
 
diff --git a/Engine/Source/Runtime/D3D12RHI/Private/D3D12Commands.cpp b/Engine/Source/Runtime/D3D12RHI/Private/D3D12Commands.cpp
--- a/Engine/Source/Runtime/D3D12RHI/Private/D3D12Commands.cpp
+++ b/Engine/Source/Runtime/D3D12RHI/Private/D3D12Commands.cpp
@@ -43,31 +43,46 @@ static FAutoConsoleVariableRef CVarD3D12AllowDiscardResources(
 
 using namespace D3D12RHI;
 
-template<typename TRHIShaderType>
-inline void ValidateBoundShader(FD3D12StateCache & InStateCache, TRHIShaderType* InShaderRHI)
+inline void ValidateBoundShader(FD3D12StateCache& InStateCache, FRHIShader* InShaderRHI)
 {
 #if DO_CHECK
-	auto* CachedShader = InStateCache.GetShader<typename TD3D12ResourceTraits<TRHIShaderType>::TConcreteType>();
-	auto* ShaderType = FD3D12DynamicRHI::ResourceCast(InShaderRHI);
-	ensureMsgf(CachedShader == ShaderType, TEXT("Parameters are being set for a %sShader which is not currently bound"),
-		GetShaderFrequencyString<TRHIShaderType>(false));
+	const EShaderFrequency ShaderFrequency = InShaderRHI->GetFrequency();
+	FRHIShader* CachedShader = InStateCache.GetShader(ShaderFrequency);
+	ensureMsgf(CachedShader == InShaderRHI, TEXT("Parameters are being set for a %sShader which is not currently bound"), GetShaderFrequencyString(ShaderFrequency, false));
 #endif
 }
 
-template<typename TRHIShaderType>
-inline void ValidateBoundUniformBuffer(FD3D12UniformBuffer * InUniformBuffer, TRHIShaderType* InShaderRHI, uint32 InBufferIndex)
+inline FD3D12ShaderData* GetShaderData(FRHIShader* InShaderRHI)
+{
+	switch (InShaderRHI->GetFrequency())
+	{
+	case SF_Vertex:        return FD3D12DynamicRHI::ResourceCast(static_cast<FRHIVertexShader*>(InShaderRHI));
+#if PLATFORM_SUPPORTS_MESH_SHADERS
+	case SF_Mesh:          return FD3D12DynamicRHI::ResourceCast(static_cast<FRHIMeshShader*>(InShaderRHI));
+	case SF_Amplification: return FD3D12DynamicRHI::ResourceCast(static_cast<FRHIAmplificationShader*>(InShaderRHI));
+#endif
+	case SF_Pixel:         return FD3D12DynamicRHI::ResourceCast(static_cast<FRHIPixelShader*>(InShaderRHI));
+	case SF_Geometry:      return FD3D12DynamicRHI::ResourceCast(static_cast<FRHIGeometryShader*>(InShaderRHI));
+	case SF_Compute:       return FD3D12DynamicRHI::ResourceCast(static_cast<FRHIComputeShader*>(InShaderRHI));
+	}
+	return nullptr;
+}
+
+inline void ValidateBoundUniformBuffer(FD3D12UniformBuffer * InUniformBuffer, FRHIShader* InShaderRHI, uint32 InBufferIndex)
 {
 #if DO_CHECK
-	auto* D3D12Shader = FD3D12DynamicRHI::ResourceCast(InShaderRHI);
-	if (InBufferIndex < (uint32)D3D12Shader->ShaderResourceTable.ResourceTableLayoutHashes.Num())
+	if (FD3D12ShaderData* ShaderData = GetShaderData(InShaderRHI))
 	{
-		uint32 UniformBufferHash = InUniformBuffer->GetLayout().GetHash(); 
-		uint32 ShaderTableHash = D3D12Shader->ShaderResourceTable.ResourceTableLayoutHashes[InBufferIndex];
-		ensureMsgf(ShaderTableHash == 0 || UniformBufferHash == ShaderTableHash,
-			TEXT("Invalid uniform buffer %s bound on %sShader at index %d."),
-			*(InUniformBuffer->GetLayout().GetDebugName()),
-			GetShaderFrequencyString<TRHIShaderType>(false),
-			InBufferIndex);
+		if (InBufferIndex < (uint32)ShaderData->ShaderResourceTable.ResourceTableLayoutHashes.Num())
+		{
+			uint32 UniformBufferHash = InUniformBuffer->GetLayout().GetHash();
+			uint32 ShaderTableHash = ShaderData->ShaderResourceTable.ResourceTableLayoutHashes[InBufferIndex];
+			ensureMsgf(ShaderTableHash == 0 || UniformBufferHash == ShaderTableHash,
+				TEXT("Invalid uniform buffer %s bound on %sShader at index %d."),
+				*(InUniformBuffer->GetLayout().GetDebugName()),
+				GetShaderFrequencyString(InShaderRHI->GetFrequency(), false),
+				InBufferIndex);
+		}
 	}
 #endif
 }
@@ -800,11 +815,10 @@ void FD3D12CommandContext::RHISetGraphicsPipelineState(FRHIGraphicsPipelineState
 
 	// TODO: [PSO API] Every thing inside this scope is only necessary to keep the PSO shadow in sync while we convert the high level to only use PSOs
 	// Ensure the command buffers are reset to reduce the amount of data that needs to be versioned.
-	VSConstantBuffer.Reset();
-	MSConstantBuffer.Reset();
-	ASConstantBuffer.Reset();
-	PSConstantBuffer.Reset();
-	GSConstantBuffer.Reset();
+	for (int32 Index = 0; Index < SF_NumGraphicsFrequencies; Index++)
+	{
+		StageConstantBuffers[Index].Reset();
+	}
 	
 	// @TODO : really should only discard the constants if the shader state has actually changed.
 	bDiscardSharedGraphicsConstants = true;
@@ -840,7 +854,7 @@ void FD3D12CommandContext::RHISetComputePipelineState(FRHIComputePipelineState*
 
 	FD3D12ComputePipelineState* ComputePipelineState = FD3D12DynamicRHI::ResourceCast(ComputeState);
 
-	CSConstantBuffer.Reset();
+	StageConstantBuffers[SF_Compute].Reset();
 	bDiscardSharedComputeConstants = true;
 
 	StateCache.SetComputePipelineState(ComputePipelineState);
@@ -850,56 +864,30 @@ void FD3D12CommandContext::RHISetComputePipelineState(FRHIComputePipelineState*
 
 void FD3D12CommandContext::RHISetShaderTexture(FRHIGraphicsShader* ShaderRHI, uint32 TextureIndex, FRHITexture* NewTextureRHI)
 {
-	FD3D12Texture* const NewTexture = RetrieveTexture(NewTextureRHI);
-	switch (ShaderRHI->GetFrequency())
+	const EShaderFrequency ShaderFrequency = ShaderRHI->GetFrequency();
+	if (IsValidGraphicsFrequency(ShaderFrequency))
 	{
-	case SF_Vertex:
-	{
-		FRHIVertexShader* VertexShaderRHI = static_cast<FRHIVertexShader*>(ShaderRHI);
-		ValidateBoundShader(StateCache, VertexShaderRHI);
-		StateCache.SetShaderResourceView<SF_Vertex>(NewTexture ? NewTexture->GetShaderResourceView() : nullptr, TextureIndex);
-	}
-	break;
-#if PLATFORM_SUPPORTS_MESH_SHADERS
-	case SF_Mesh:
-	{
-		FRHIMeshShader* MeshShaderRHI = static_cast<FRHIMeshShader*>(ShaderRHI);
-		ValidateBoundShader(StateCache, MeshShaderRHI);
-		StateCache.SetShaderResourceView<SF_Mesh>(NewTexture ? NewTexture->GetShaderResourceView() : nullptr, TextureIndex);
-	}
-	break;
-	case SF_Amplification:
-	{
-		FRHIAmplificationShader* AmplificationShaderRHI = static_cast<FRHIAmplificationShader*>(ShaderRHI);
-		ValidateBoundShader(StateCache, AmplificationShaderRHI);
-		StateCache.SetShaderResourceView<SF_Amplification>(NewTexture ? NewTexture->GetShaderResourceView() : nullptr, TextureIndex);
-	}
-	break;
-#endif
-	case SF_Geometry:
-	{
-		FRHIGeometryShader* GeometryShaderRHI = static_cast<FRHIGeometryShader*>(ShaderRHI);
-		ValidateBoundShader(StateCache, GeometryShaderRHI);
-		StateCache.SetShaderResourceView<SF_Geometry>(NewTexture ? NewTexture->GetShaderResourceView() : nullptr, TextureIndex);
+		ValidateBoundShader(StateCache, ShaderRHI);
+
+		FD3D12Texture* NewTexture = RetrieveTexture(NewTextureRHI);
+		FD3D12ShaderResourceView* ViewToSet = NewTexture ? NewTexture->GetShaderResourceView() : nullptr;
+
+		StateCache.SetShaderResourceView(ShaderFrequency, ViewToSet, TextureIndex);
 	}
-	break;
-	case SF_Pixel:
+	else
 	{
-		FRHIPixelShader* PixelShaderRHI = static_cast<FRHIPixelShader*>(ShaderRHI);
-		ValidateBoundShader(StateCache, PixelShaderRHI);
-		StateCache.SetShaderResourceView<SF_Pixel>(NewTexture ? NewTexture->GetShaderResourceView() : nullptr, TextureIndex);
-	}
-	break;
-	default:
-		checkf(0, TEXT("Undefined FRHIShader Type %d!"), (int32)ShaderRHI->GetFrequency());
+		checkf(0, TEXT("Unsupported FRHIGraphicsShader Type '%s'!"), GetShaderFrequencyString(ShaderFrequency, false));
 	}
 }
 
 void FD3D12CommandContext::RHISetShaderTexture(FRHIComputeShader* ComputeShaderRHI, uint32 TextureIndex, FRHITexture* NewTextureRHI)
 {
-	//ValidateBoundShader(StateCache, ComputeShaderRHI);
+	ValidateBoundShader(StateCache, ComputeShaderRHI);
+
 	FD3D12Texture* const NewTexture = RetrieveTexture(NewTextureRHI);
-	StateCache.SetShaderResourceView<SF_Compute>(NewTexture ? NewTexture->GetShaderResourceView() : nullptr, TextureIndex);
+	FD3D12ShaderResourceView* ViewToSet = NewTexture ? NewTexture->GetShaderResourceView() : nullptr;
+
+	StateCache.SetShaderResourceView(SF_Compute, ViewToSet, TextureIndex);
 }
 
 void FD3D12CommandContext::RHISetUAVParameter(FRHIPixelShader* PixelShaderRHI, uint32 UAVIndex, FRHIUnorderedAccessView* UAVRHI)
@@ -914,13 +902,13 @@ void FD3D12CommandContext::RHISetUAVParameter(FRHIPixelShader* PixelShaderRHI, u
 	uint32 InitialCount = -1;
 
 	// Actually set the UAV
-	StateCache.SetUAVs<SF_Pixel>(UAVIndex, 1, &UAV, &InitialCount);
+	StateCache.SetUAVs(SF_Pixel, UAVIndex, 1, &UAV, &InitialCount);
 }
 
 
 void FD3D12CommandContext::RHISetUAVParameter(FRHIComputeShader* ComputeShaderRHI, uint32 UAVIndex, FRHIUnorderedAccessView* UAVRHI)
 {
-	//ValidateBoundShader(StateCache, ComputeShaderRHI);
+	ValidateBoundShader(StateCache, ComputeShaderRHI);
 
 	FD3D12UnorderedAccessView* UAV = RetrieveObject<FD3D12UnorderedAccessView>(UAVRHI);
 
@@ -932,12 +920,12 @@ void FD3D12CommandContext::RHISetUAVParameter(FRHIComputeShader* ComputeShaderRH
 	uint32 InitialCount = -1;
 
 	// Actually set the UAV
-	StateCache.SetUAVs<SF_Compute>(UAVIndex, 1, &UAV, &InitialCount);
+	StateCache.SetUAVs(SF_Compute, UAVIndex, 1, &UAV, &InitialCount);
 }
 
 void FD3D12CommandContext::RHISetUAVParameter(FRHIComputeShader* ComputeShaderRHI, uint32 UAVIndex, FRHIUnorderedAccessView* UAVRHI, uint32 InitialCount)
 {
-	//ValidateBoundShader(StateCache, ComputeShaderRHI);
+	ValidateBoundShader(StateCache, ComputeShaderRHI);
 
 	FD3D12UnorderedAccessView* UAV = RetrieveObject<FD3D12UnorderedAccessView>(UAVRHI);
 
@@ -946,193 +934,93 @@ void FD3D12CommandContext::RHISetUAVParameter(FRHIComputeShader* ComputeShaderRH
 		ConditionalClearShaderResource(UAV->GetResourceLocation());
 	}
 
-	StateCache.SetUAVs<SF_Compute>(UAVIndex, 1, &UAV, &InitialCount);
+	StateCache.SetUAVs(SF_Compute, UAVIndex, 1, &UAV, &InitialCount);
 }
 
 void FD3D12CommandContext::RHISetShaderResourceViewParameter(FRHIGraphicsShader* ShaderRHI, uint32 TextureIndex, FRHIShaderResourceView* SRVRHI)
 {
-	FD3D12ShaderResourceView* const SRV = RetrieveObject<FD3D12ShaderResourceView>(SRVRHI);
-	switch (ShaderRHI->GetFrequency())
-	{
-	case SF_Vertex:
-	{
-		FRHIVertexShader* VertexShaderRHI = static_cast<FRHIVertexShader*>(ShaderRHI);
-		ValidateBoundShader(StateCache, VertexShaderRHI);
-		StateCache.SetShaderResourceView<SF_Vertex>(SRV, TextureIndex);
-	}
-	break;
-#if PLATFORM_SUPPORTS_MESH_SHADERS
-	case SF_Mesh:
-	{
-		FRHIMeshShader* MeshShaderRHI = static_cast<FRHIMeshShader*>(ShaderRHI);
-		ValidateBoundShader(StateCache, MeshShaderRHI);
-		StateCache.SetShaderResourceView<SF_Mesh>(SRV, TextureIndex);
-	}
-	break;
-	case SF_Amplification:
-	{
-		FRHIAmplificationShader* AmplificationShaderRHI = static_cast<FRHIAmplificationShader*>(ShaderRHI);
-		ValidateBoundShader(StateCache, AmplificationShaderRHI);
-		StateCache.SetShaderResourceView<SF_Amplification>(SRV, TextureIndex);
-	}
-	break;
-#endif
-	case SF_Geometry:
+	const EShaderFrequency ShaderFrequency = ShaderRHI->GetFrequency();
+
+	if (IsValidGraphicsFrequency(ShaderFrequency))
 	{
-		FRHIGeometryShader* GeometryShaderRHI = static_cast<FRHIGeometryShader*>(ShaderRHI);
-		ValidateBoundShader(StateCache, GeometryShaderRHI);
-		StateCache.SetShaderResourceView<SF_Geometry>(SRV, TextureIndex);
+		ValidateBoundShader(StateCache, ShaderRHI);
+
+		FD3D12ShaderResourceView* SRV = RetrieveObject<FD3D12ShaderResourceView>(SRVRHI);
+
+		StateCache.SetShaderResourceView(ShaderFrequency, SRV, TextureIndex);
 	}
-	break;
-	case SF_Pixel:
+	else
 	{
-		FRHIPixelShader* PixelShaderRHI = static_cast<FRHIPixelShader*>(ShaderRHI);
-		ValidateBoundShader(StateCache, PixelShaderRHI);
-		StateCache.SetShaderResourceView<SF_Pixel>(SRV, TextureIndex);
-	}
-	break;
-	default:
-		checkf(0, TEXT("Undefined FRHIShader Type %d!"), (int32)ShaderRHI->GetFrequency());
+		checkf(0, TEXT("Unsupported FRHIGraphicsShader Type '%s'!"), GetShaderFrequencyString(ShaderFrequency, false));
 	}
-
 }
 
 void FD3D12CommandContext::RHISetShaderResourceViewParameter(FRHIComputeShader* ComputeShaderRHI, uint32 TextureIndex, FRHIShaderResourceView* SRVRHI)
 {
-	//ValidateBoundShader(StateCache, ComputeShaderRHI);
+	ValidateBoundShader(StateCache, ComputeShaderRHI);
 	FD3D12ShaderResourceView* const SRV = RetrieveObject<FD3D12ShaderResourceView>(SRVRHI);
-	StateCache.SetShaderResourceView<SF_Compute>(SRV, TextureIndex);
+	StateCache.SetShaderResourceView(SF_Compute, SRV, TextureIndex);
 }
 
 void FD3D12CommandContext::RHISetShaderSampler(FRHIGraphicsShader* ShaderRHI, uint32 SamplerIndex, FRHISamplerState* NewStateRHI)
 {
-	FD3D12SamplerState* NewState = RetrieveObject<FD3D12SamplerState>(NewStateRHI);
-	switch (ShaderRHI->GetFrequency())
-	{
-	case SF_Vertex:
-	{
-		FRHIVertexShader* VertexShaderRHI = static_cast<FRHIVertexShader*>(ShaderRHI);
-		ValidateBoundShader(StateCache, VertexShaderRHI);
-		StateCache.SetSamplerState<SF_Vertex>(NewState, SamplerIndex);
-	}
-	break;
-#if PLATFORM_SUPPORTS_MESH_SHADERS
-	case SF_Mesh:
-	{
-		FRHIMeshShader* MeshShaderRHI = static_cast<FRHIMeshShader*>(ShaderRHI);
-		ValidateBoundShader(StateCache, MeshShaderRHI);
-		StateCache.SetSamplerState<SF_Mesh>(NewState, SamplerIndex);
-	}
-	break;
-	case SF_Amplification:
-	{
-		FRHIAmplificationShader* AmplificationShaderRHI = static_cast<FRHIAmplificationShader*>(ShaderRHI);
-		ValidateBoundShader(StateCache, AmplificationShaderRHI);
-		StateCache.SetSamplerState<SF_Amplification>(NewState, SamplerIndex);
-	}
-	break;
-#endif
-	case SF_Geometry:
+	const EShaderFrequency ShaderFrequency = ShaderRHI->GetFrequency();
+
+	if (IsValidGraphicsFrequency(ShaderFrequency))
 	{
-		FRHIGeometryShader* GeometryShaderRHI = static_cast<FRHIGeometryShader*>(ShaderRHI);
-		ValidateBoundShader(StateCache, GeometryShaderRHI);
-		StateCache.SetSamplerState<SF_Geometry>(NewState, SamplerIndex);
+		ValidateBoundShader(StateCache, ShaderRHI);
+
+		FD3D12SamplerState* NewState = RetrieveObject<FD3D12SamplerState>(NewStateRHI);
+
+		StateCache.SetSamplerState(ShaderFrequency, NewState, SamplerIndex);
 	}
-	break;
-	case SF_Pixel:
+	else
 	{
-		FRHIPixelShader* PixelShaderRHI = static_cast<FRHIPixelShader*>(ShaderRHI);
-		ValidateBoundShader(StateCache, PixelShaderRHI);
-		StateCache.SetSamplerState<SF_Pixel>(NewState, SamplerIndex);
-	}
-	break;
-	default:
-		checkf(0, TEXT("Undefined FRHIShader Type %d!"), (int32)ShaderRHI->GetFrequency());
+		checkf(0, TEXT("Unsupported FRHIGraphicsShader Type '%s'!"), GetShaderFrequencyString(ShaderFrequency, false));
 	}
 }
 
 void FD3D12CommandContext::RHISetShaderSampler(FRHIComputeShader* ComputeShaderRHI, uint32 SamplerIndex, FRHISamplerState* NewStateRHI)
 {
-	//ValidateBoundShader(StateCache, ComputeShaderRHI);
+	ValidateBoundShader(StateCache, ComputeShaderRHI);
 	FD3D12SamplerState* NewState = RetrieveObject<FD3D12SamplerState>(NewStateRHI);
-	StateCache.SetSamplerState<SF_Compute>(NewState, SamplerIndex);
+	StateCache.SetSamplerState(SF_Compute, NewState, SamplerIndex);
 }
 
 void FD3D12CommandContext::RHISetShaderUniformBuffer(FRHIGraphicsShader* ShaderRHI, uint32 BufferIndex, FRHIUniformBuffer* BufferRHI)
 {
-	//SCOPE_CYCLE_COUNTER(STAT_D3D12SetShaderUniformBuffer);
-	FD3D12UniformBuffer* Buffer = RetrieveObject<FD3D12UniformBuffer>(BufferRHI);
-	EShaderFrequency Stage = SF_NumFrequencies;
-	switch (ShaderRHI->GetFrequency())
-	{
-	case SF_Vertex:
-	{
-		FRHIVertexShader* VertexShaderRHI = static_cast<FRHIVertexShader*>(ShaderRHI);
-		ValidateBoundShader(StateCache, VertexShaderRHI);
-		ValidateBoundUniformBuffer(Buffer, VertexShaderRHI, BufferIndex);
-		StateCache.SetConstantsFromUniformBuffer<SF_Vertex>(BufferIndex, Buffer);
-		Stage = SF_Vertex;
-	}
-	break;
-#if PLATFORM_SUPPORTS_MESH_SHADERS
-	case SF_Mesh:
-	{
-		FRHIMeshShader* MeshShaderRHI = static_cast<FRHIMeshShader*>(ShaderRHI);
-		ValidateBoundShader(StateCache, MeshShaderRHI);
-		ValidateBoundUniformBuffer(Buffer, MeshShaderRHI, BufferIndex);
-		StateCache.SetConstantsFromUniformBuffer<SF_Mesh>(BufferIndex, Buffer);
-		Stage = SF_Mesh;
-	}
-	break;
-	case SF_Amplification:
-	{
-		FRHIAmplificationShader* AmplificationShaderRHI = static_cast<FRHIAmplificationShader*>(ShaderRHI);
-		ValidateBoundShader(StateCache, AmplificationShaderRHI);
-		ValidateBoundUniformBuffer(Buffer, AmplificationShaderRHI, BufferIndex);
-		StateCache.SetConstantsFromUniformBuffer<SF_Amplification>(BufferIndex, Buffer);
-		Stage = SF_Amplification;
-	}
-	break;
-#endif
-	case SF_Geometry:
-	{
-		FRHIGeometryShader* GeometryShaderRHI = static_cast<FRHIGeometryShader*>(ShaderRHI);
-		ValidateBoundShader(StateCache, GeometryShaderRHI);
-		ValidateBoundUniformBuffer(Buffer, GeometryShaderRHI, BufferIndex);
-		StateCache.SetConstantsFromUniformBuffer<SF_Geometry>(BufferIndex, Buffer);
-		Stage = SF_Geometry;
-	}
-	break;
-	case SF_Pixel:
+	const EShaderFrequency ShaderFrequency = ShaderRHI->GetFrequency();
+
+	if (IsValidGraphicsFrequency(ShaderFrequency))
 	{
-		FRHIPixelShader* PixelShaderRHI = static_cast<FRHIPixelShader*>(ShaderRHI);
-		ValidateBoundShader(StateCache, PixelShaderRHI);
-		ValidateBoundUniformBuffer(Buffer, PixelShaderRHI, BufferIndex);
-		StateCache.SetConstantsFromUniformBuffer<SF_Pixel>(BufferIndex, Buffer);
-		Stage = SF_Pixel;
-	}
-	break;
-	default:
-		checkf(0, TEXT("Undefined FRHIShader Type %d!"), (int32)ShaderRHI->GetFrequency());
-		return;
-	}
+		ValidateBoundShader(StateCache, ShaderRHI);
 
-	if (!GRHINeedsExtraDeletionLatency)
+		FD3D12UniformBuffer* Buffer = RetrieveObject<FD3D12UniformBuffer>(BufferRHI);
+		ValidateBoundUniformBuffer(Buffer, ShaderRHI, BufferIndex);
+
+		StateCache.SetConstantsFromUniformBuffer(ShaderFrequency, BufferIndex, Buffer);
+
+		if (!GRHINeedsExtraDeletionLatency)
+		{
+			BoundUniformBufferRefs[ShaderFrequency][BufferIndex] = BufferRHI;
+		}
+
+		BoundUniformBuffers[ShaderFrequency][BufferIndex] = Buffer;
+		DirtyUniformBuffers[ShaderFrequency] |= (1 << BufferIndex);
+	}
+	else
 	{
-		BoundUniformBufferRefs[Stage][BufferIndex] = BufferRHI;
+		checkf(0, TEXT("Unsupported FRHIGraphicsShader Type '%s'!"), GetShaderFrequencyString(ShaderFrequency, false));
 	}
-
-	BoundUniformBuffers[Stage][BufferIndex] = Buffer;
-	DirtyUniformBuffers[Stage] |= (1 << BufferIndex);
 }
 
 void FD3D12CommandContext::RHISetShaderUniformBuffer(FRHIComputeShader* ComputeShader, uint32 BufferIndex, FRHIUniformBuffer* BufferRHI)
 {
 	//SCOPE_CYCLE_COUNTER(STAT_D3D12SetShaderUniformBuffer);
-	//ValidateBoundShader(StateCache, ComputeShader);
+	ValidateBoundShader(StateCache, ComputeShader);
 	FD3D12UniformBuffer* Buffer = RetrieveObject<FD3D12UniformBuffer>(BufferRHI);
 
-	StateCache.SetConstantsFromUniformBuffer<SF_Compute>(BufferIndex, Buffer);
+	StateCache.SetConstantsFromUniformBuffer(SF_Compute, BufferIndex, Buffer);
 
 	if (!GRHINeedsExtraDeletionLatency)
 	{
@@ -1171,55 +1059,24 @@ void FD3D12CommandContext::RHISetShaderParameter(FRHIGraphicsShader* ShaderRHI,
 {
 	checkSlow(BufferIndex == 0);
 
-	switch (ShaderRHI->GetFrequency())
-	{
-	case SF_Vertex:
-	{
-		FRHIVertexShader* VertexShaderRHI = static_cast<FRHIVertexShader*>(ShaderRHI);
-		ValidateBoundShader(StateCache, VertexShaderRHI);
-		VSConstantBuffer.UpdateConstant((const uint8*)NewValue, Offset, NumBytes);
-	}
-	break;
-#if PLATFORM_SUPPORTS_MESH_SHADERS
-	case SF_Mesh:
-	{
-		FRHIMeshShader* MeshShaderRHI = static_cast<FRHIMeshShader*>(ShaderRHI);
-		ValidateBoundShader(StateCache, MeshShaderRHI);
-		MSConstantBuffer.UpdateConstant((const uint8*)NewValue, Offset, NumBytes);
-	}
-	break;
-	case SF_Amplification:
+	const EShaderFrequency ShaderFrequency = ShaderRHI->GetFrequency();
+	if (IsValidGraphicsFrequency(ShaderFrequency))
 	{
-		FRHIAmplificationShader* AmplificationShaderRHI = static_cast<FRHIAmplificationShader*>(ShaderRHI);
-		ValidateBoundShader(StateCache, AmplificationShaderRHI);
-		ASConstantBuffer.UpdateConstant((const uint8*)NewValue, Offset, NumBytes);
-	}
-	break;
-#endif
-	case SF_Geometry:
-	{
-		FRHIGeometryShader* GeometryShaderRHI = static_cast<FRHIGeometryShader*>(ShaderRHI);
-		ValidateBoundShader(StateCache, GeometryShaderRHI);
-		GSConstantBuffer.UpdateConstant((const uint8*)NewValue, Offset, NumBytes);
+		ValidateBoundShader(StateCache, ShaderRHI);
+
+		StageConstantBuffers[ShaderFrequency].UpdateConstant((const uint8*)NewValue, Offset, NumBytes);
 	}
-	break;
-	case SF_Pixel:
+	else
 	{
-		FRHIPixelShader* PixelShaderRHI = static_cast<FRHIPixelShader*>(ShaderRHI);
-		ValidateBoundShader(StateCache, PixelShaderRHI);
-		PSConstantBuffer.UpdateConstant((const uint8*)NewValue, Offset, NumBytes);
-	}
-	break;
-	default:
-		checkf(0, TEXT("Undefined FRHIShader Type %d!"), (int32)ShaderRHI->GetFrequency());
+		checkf(0, TEXT("Unsupported FRHIGraphicsShader Type '%s'!"), GetShaderFrequencyString(ShaderFrequency, false));
 	}
 }
 
 void FD3D12CommandContext::RHISetShaderParameter(FRHIComputeShader* ComputeShaderRHI, uint32 BufferIndex, uint32 Offset, uint32 NumBytes, const void* NewValue)
 {
-	//ValidateBoundShader(StateCache, ComputeShaderRHI);
+	ValidateBoundShader(StateCache, ComputeShaderRHI);
 	checkSlow(BufferIndex == 0);
-	CSConstantBuffer.UpdateConstant((const uint8*)NewValue, Offset, NumBytes);
+	StageConstantBuffers[SF_Compute].UpdateConstant((const uint8*)NewValue, Offset, NumBytes);
 }
 
 void FD3D12CommandContext::RHISetStencilRef(uint32 StencilRef)
@@ -1332,7 +1189,7 @@ void FD3D12CommandContext::SetRenderTargets(
 	}
 
 	StateCache.SetRenderTargets(NewNumSimultaneousRenderTargets, NewRenderTargetViews, DepthStencilView);
-	StateCache.ClearUAVs<SF_Pixel>();
+	StateCache.ClearUAVs(SF_Pixel);
 
 	// Set the viewport to the full size of render target 0.
 	if (NewRenderTargetViews[0])
@@ -1470,31 +1327,14 @@ void FD3D12CommandContext::CommitNonComputeShaderConstants()
 
 	// Only set the constant buffer if this shader needs the global constant buffer bound
 	// Otherwise we will overwrite a different constant buffer
-	if (GraphicPSO->bShaderNeedsGlobalConstantBuffer[SF_Vertex])
-	{
-		StateCache.SetConstantBuffer<SF_Vertex>(VSConstantBuffer, bDiscardSharedGraphicsConstants);
-	}
-
-#if PLATFORM_SUPPORTS_MESH_SHADERS
-	if (GraphicPSO->bShaderNeedsGlobalConstantBuffer[SF_Mesh])
-	{
-		StateCache.SetConstantBuffer<SF_Mesh>(MSConstantBuffer, bDiscardSharedGraphicsConstants);
-	}
-
-	if (GraphicPSO->bShaderNeedsGlobalConstantBuffer[SF_Amplification])
-	{
-		StateCache.SetConstantBuffer<SF_Amplification>(ASConstantBuffer, bDiscardSharedGraphicsConstants);
-	}
-#endif
 
-	if (GraphicPSO->bShaderNeedsGlobalConstantBuffer[SF_Geometry])
+	for (int32 Index = 0; Index < SF_NumGraphicsFrequencies; Index++)
 	{
-		StateCache.SetConstantBuffer<SF_Geometry>(GSConstantBuffer, bDiscardSharedGraphicsConstants);
-	}
-
-	if (GraphicPSO->bShaderNeedsGlobalConstantBuffer[SF_Pixel])
-	{
-		StateCache.SetConstantBuffer<SF_Pixel>(PSConstantBuffer, bDiscardSharedGraphicsConstants);
+		const EShaderFrequency ShaderFrequency = static_cast<EShaderFrequency>(Index);
+		if (IsValidGraphicsFrequency(ShaderFrequency) && GraphicPSO->bShaderNeedsGlobalConstantBuffer[Index])
+		{
+			StateCache.SetConstantBuffer(ShaderFrequency, StageConstantBuffers[Index], bDiscardSharedGraphicsConstants);
+		}
 	}
 
 	bDiscardSharedGraphicsConstants = false;
@@ -1508,7 +1348,7 @@ void FD3D12CommandContext::CommitComputeShaderConstants()
 
 	if (ComputePSO->bShaderNeedsGlobalConstantBuffer)
 	{
-		StateCache.SetConstantBuffer<SF_Compute>(CSConstantBuffer, bDiscardSharedComputeConstants);
+		StateCache.SetConstantBuffer(SF_Compute, StageConstantBuffers[SF_Compute], bDiscardSharedComputeConstants);
 	}
 
 	bDiscardSharedComputeConstants = false;
@@ -1531,23 +1371,23 @@ void FD3D12CommandContext::SetResourcesFromTables(const ShaderType* RESTRICT Sha
 			uint32 InitialCount = -1;
 			FD3D12UnorderedAccessView* D3D12UAV = Context.RetrieveObject<FD3D12UnorderedAccessView>(UAV);
 
-			Context.StateCache.SetUAVs<SF_Pixel>(Index, 1, &D3D12UAV, &InitialCount);
+			Context.StateCache.SetUAVs(SF_Pixel, Index, 1, &D3D12UAV, &InitialCount);
 		}
 
 		void SetSRV(FRHIShaderResourceView* SRV, uint8 Index)
 		{
-			Context.StateCache.SetShaderResourceView<Frequency>(Context.RetrieveObject<FD3D12ShaderResourceView>(SRV), Index);
+			Context.StateCache.SetShaderResourceView(Frequency, Context.RetrieveObject<FD3D12ShaderResourceView>(SRV), Index);
 		}
 
 		void SetTexture(FRHITexture* TextureRHI, uint8 Index)
 		{
 			FD3D12ShaderResourceView* SRV = Context.RetrieveTexture(TextureRHI)->GetShaderResourceView();
-			Context.StateCache.SetShaderResourceView<Frequency>(SRV, Index);
+			Context.StateCache.SetShaderResourceView(Frequency, SRV, Index);
 		}
 
 		void SetSampler(FRHISamplerState* Sampler, uint8 Index)
 		{
-			Context.StateCache.SetSamplerState<Frequency>(Context.RetrieveObject<FD3D12SamplerState>(Sampler), Index);
+			Context.StateCache.SetSamplerState(Frequency, Context.RetrieveObject<FD3D12SamplerState>(Sampler), Index);
 		}
 	};
 
diff --git a/Engine/Source/Runtime/D3D12RHI/Private/D3D12RenderTarget.cpp b/Engine/Source/Runtime/D3D12RHI/Private/D3D12RenderTarget.cpp
--- a/Engine/Source/Runtime/D3D12RHI/Private/D3D12RenderTarget.cpp
+++ b/Engine/Source/Runtime/D3D12RHI/Private/D3D12RenderTarget.cpp
@@ -268,7 +268,7 @@ void FD3D12CommandContext::ResolveTextureUsingShader(
 			RHICmdList.RunOnContext([TextureIndex, SourceTexture](FD3D12CommandContext& Context)
 			{
 				// Set the source texture.
-				Context.StateCache.SetShaderResourceView<SF_Pixel>(SourceTexture->GetShaderResourceView(), TextureIndex);
+				Context.StateCache.SetShaderResourceView(SF_Pixel, SourceTexture->GetShaderResourceView(), TextureIndex);
 			});
 
 			RHICmdList.DrawPrimitive(0, 2, 1);
diff --git a/Engine/Source/Runtime/D3D12RHI/Private/D3D12StateCache.cpp b/Engine/Source/Runtime/D3D12RHI/Private/D3D12StateCache.cpp
--- a/Engine/Source/Runtime/D3D12RHI/Private/D3D12StateCache.cpp
+++ b/Engine/Source/Runtime/D3D12RHI/Private/D3D12StateCache.cpp
@@ -3,13 +3,7 @@
 // Implementation of Device Context State Caching to improve draw
 //	thread performance by removing redundant device context calls.
 
-//-----------------------------------------------------------------------------
-//	Include Files
-//-----------------------------------------------------------------------------
 #include "D3D12RHIPrivate.h"
-#if !PLATFORM_HOLOLENS && PLATFORM_CPU_X86_FAMILY
-#include <emmintrin.h>
-#endif
 
 // This value defines how many descriptors will be in the device global descriptor heap. This heap contains all shader visible view descriptors.
 // Other shader visible descriptor heaps (e.g. OnlineViewHeap) are allocated from this pool. Non-visible heaps (e.g. LocalViewHeap) are allocated as standalone.
@@ -90,18 +84,6 @@ inline bool operator!=(D3D12_CPU_DESCRIPTOR_HANDLE lhs, D3D12_CPU_DESCRIPTOR_HAN
 }
 
 // Define template functions that are only declared in the header.
-template void FD3D12StateCache::SetShaderResourceView<SF_Vertex>(FD3D12ShaderResourceView* SRV, uint32 ResourceIndex);
-template void FD3D12StateCache::SetShaderResourceView<SF_Mesh>(FD3D12ShaderResourceView* SRV, uint32 ResourceIndex);
-template void FD3D12StateCache::SetShaderResourceView<SF_Amplification>(FD3D12ShaderResourceView* SRV, uint32 ResourceIndex);
-template void FD3D12StateCache::SetShaderResourceView<SF_Geometry>(FD3D12ShaderResourceView* SRV, uint32 ResourceIndex);
-template void FD3D12StateCache::SetShaderResourceView<SF_Pixel>(FD3D12ShaderResourceView* SRV, uint32 ResourceIndex);
-template void FD3D12StateCache::SetShaderResourceView<SF_Compute>(FD3D12ShaderResourceView* SRV, uint32 ResourceIndex);
-
-template void FD3D12StateCache::SetUAVs<SF_Pixel>(uint32 UAVStartSlot, uint32 NumSimultaneousUAVs, FD3D12UnorderedAccessView** UAVArray, uint32* UAVInitialCountArray);
-template void FD3D12StateCache::SetUAVs<SF_Compute>(uint32 UAVStartSlot, uint32 NumSimultaneousUAVs, FD3D12UnorderedAccessView** UAVArray, uint32* UAVInitialCountArray);
-
-
-template void FD3D12StateCache::ClearUAVs<SF_Pixel>();
 
 template void FD3D12StateCache::ApplyState<ED3D12PipelineType::Graphics>();
 template void FD3D12StateCache::ApplyState<ED3D12PipelineType::Compute>();
@@ -172,6 +154,25 @@ void FD3D12StateCache::ClearSRVs()
 	bSRVSCleared = true;
 }
 
+void FD3D12StateCache::ClearShaderResourceViews(EShaderFrequency ShaderFrequency, FD3D12ResourceLocation*& ResourceLocation)
+{
+	//SCOPE_CYCLE_COUNTER(STAT_D3D12ClearShaderResourceViewsTime);
+
+	if (PipelineState.Common.SRVCache.MaxBoundIndex[ShaderFrequency] < 0)
+	{
+		return;
+	}
+
+	auto& CurrentShaderResourceViews = PipelineState.Common.SRVCache.Views[ShaderFrequency];
+	for (int32 i = 0; i <= PipelineState.Common.SRVCache.MaxBoundIndex[ShaderFrequency]; ++i)
+	{
+		if (CurrentShaderResourceViews[i] && CurrentShaderResourceViews[i]->GetResourceLocation() == ResourceLocation)
+		{
+			SetShaderResourceView(ShaderFrequency, nullptr, i);
+		}
+	}
+}
+
 void FD3D12StateCache::FlushComputeShaderCache(bool bForce)
 {
 	if (bForce)
@@ -1165,8 +1166,7 @@ bool FD3D12StateCache::AssertResourceStates(ED3D12PipelineType PipelineType)
 #endif
 }
 
-template <EShaderFrequency ShaderStage>
-void FD3D12StateCache::ClearUAVs()
+void FD3D12StateCache::ClearUAVs(EShaderFrequency ShaderStage)
 {
 	FD3D12UnorderedAccessViewCache& Cache = PipelineState.Common.UAVCache;
 	const bool bIsCompute = ShaderStage == SF_Compute;
@@ -1181,8 +1181,7 @@ void FD3D12StateCache::ClearUAVs()
 	}
 }
 
-template <EShaderFrequency ShaderStage>
-void FD3D12StateCache::SetUAVs(uint32 UAVStartSlot, uint32 NumSimultaneousUAVs, FD3D12UnorderedAccessView** UAVArray, uint32* UAVInitialCountArray)
+void FD3D12StateCache::SetUAVs(EShaderFrequency ShaderStage, uint32 UAVStartSlot, uint32 NumSimultaneousUAVs, FD3D12UnorderedAccessView** UAVArray, uint32* UAVInitialCountArray)
 {
 	SCOPE_CYCLE_COUNTER(STAT_D3D12SetUnorderedAccessViewTime);
 	check(NumSimultaneousUAVs > 0);
@@ -1425,8 +1424,7 @@ void FD3D12StateCache::InternalSetStreamSource(FD3D12ResourceLocation* VertexBuf
 	}
 }
 
-template <EShaderFrequency ShaderFrequency>
-void FD3D12StateCache::SetShaderResourceView(FD3D12ShaderResourceView* SRV, uint32 ResourceIndex)
+void FD3D12StateCache::SetShaderResourceView(EShaderFrequency ShaderFrequency, FD3D12ShaderResourceView* SRV, uint32 ResourceIndex)
 {
 	//SCOPE_CYCLE_COUNTER(STAT_D3D12SetShaderResourceViewTime);
 
diff --git a/Engine/Source/Runtime/D3D12RHI/Private/D3D12StateCachePrivate.h b/Engine/Source/Runtime/D3D12RHI/Private/D3D12StateCachePrivate.h
--- a/Engine/Source/Runtime/D3D12RHI/Private/D3D12StateCachePrivate.h
+++ b/Engine/Source/Runtime/D3D12RHI/Private/D3D12StateCachePrivate.h
@@ -408,26 +408,8 @@ protected:
 	FD3D12DescriptorCache DescriptorCache;
 
 	void InternalSetIndexBuffer(FD3D12Resource* Resource);
-
 	void InternalSetStreamSource(FD3D12ResourceLocation* VertexBufferLocation, uint32 StreamIndex, uint32 Stride, uint32 Offset);
 
-	template <typename TShader> struct TStateCacheShaderTraits;
-#define DECLARE_SHADER_TRAITS(Name) \
-	template <> struct TStateCacheShaderTraits<FD3D12##Name##Shader> \
-	{ \
-		static FD3D12##Name##Shader* GetShader(FD3D12GraphicsPipelineState* PSO) { return PSO ? (FD3D12##Name##Shader*)PSO->PipelineStateInitializer.BoundShaderState.Get##Name##Shader() : nullptr; } \
-	}
-	DECLARE_SHADER_TRAITS(Vertex);
-#if PLATFORM_SUPPORTS_MESH_SHADERS
-	DECLARE_SHADER_TRAITS(Mesh);
-	DECLARE_SHADER_TRAITS(Amplification);
-#endif
-	DECLARE_SHADER_TRAITS(Pixel);
-#if PLATFORM_SUPPORTS_GEOMETRY_SHADERS
-	DECLARE_SHADER_TRAITS(Geometry);
-#endif
-#undef DECLARE_SHADER_TRAITS
-
 	bool InternalSetRootSignature(ED3D12PipelineType InPipelineType, const FD3D12RootSignature* InRootSignature);
 	void InternalSetPipelineState(FD3D12PipelineState* InPipelineState);
 
@@ -472,28 +454,9 @@ public:
 
 	void ClearSRVs();
 
-	template <EShaderFrequency ShaderFrequency>
-	void ClearShaderResourceViews(FD3D12ResourceLocation*& ResourceLocation)
-	{
-		//SCOPE_CYCLE_COUNTER(STAT_D3D12ClearShaderResourceViewsTime);
-
-		if (PipelineState.Common.SRVCache.MaxBoundIndex[ShaderFrequency] < 0)
-		{
-			return;
-		}
-
-		auto& CurrentShaderResourceViews = PipelineState.Common.SRVCache.Views[ShaderFrequency];
-		for (int32 i = 0; i <= PipelineState.Common.SRVCache.MaxBoundIndex[ShaderFrequency]; ++i)
-		{
-			if (CurrentShaderResourceViews[i] && CurrentShaderResourceViews[i]->GetResourceLocation() == ResourceLocation)
-			{
-				SetShaderResourceView<ShaderFrequency>(nullptr, i);
-			}
-		}
-	}
+	void ClearShaderResourceViews(EShaderFrequency ShaderFrequency, FD3D12ResourceLocation*& ResourceLocation);
 
-	template <EShaderFrequency ShaderFrequency>
-	void SetShaderResourceView(FD3D12ShaderResourceView* SRV, uint32 ResourceIndex);
+	void SetShaderResourceView(EShaderFrequency ShaderFrequency, FD3D12ShaderResourceView* SRV, uint32 ResourceIndex);
 	
 	void SetScissorRects(uint32 Count, const D3D12_RECT* const ScissorRects);
 	void SetScissorRect(const D3D12_RECT& ScissorRect);
@@ -537,8 +500,7 @@ public:
 		*Count = PipelineState.Graphics.CurrentNumberOfViewports;
 	}
 
-	template <EShaderFrequency ShaderFrequency>
-	D3D12_STATE_CACHE_INLINE void SetSamplerState(FD3D12SamplerState* SamplerState, uint32 SamplerIndex)
+	D3D12_STATE_CACHE_INLINE void SetSamplerState(EShaderFrequency ShaderFrequency, FD3D12SamplerState* SamplerState, uint32 SamplerIndex)
 	{
 		check(SamplerIndex < MAX_SAMPLERS);
 		auto& Samplers = PipelineState.Common.SamplerCache.States[ShaderFrequency];
@@ -549,8 +511,7 @@ public:
 		}
 	}
 
-	template <EShaderFrequency ShaderFrequency>
-	void D3D12_STATE_CACHE_INLINE SetConstantsFromUniformBuffer(uint32 SlotIndex, FD3D12UniformBuffer* UniformBuffer)
+	D3D12_STATE_CACHE_INLINE void SetConstantsFromUniformBuffer(EShaderFrequency ShaderFrequency, uint32 SlotIndex, FD3D12UniformBuffer* UniformBuffer)
 	{
 		check(SlotIndex < MAX_CBS);
 		FD3D12ConstantBufferCache& CBVCache = PipelineState.Common.CBVCache;
@@ -588,8 +549,7 @@ public:
 		}
 	}
 
-	template <EShaderFrequency ShaderFrequency>
-	void D3D12_STATE_CACHE_INLINE SetConstantBuffer(FD3D12ConstantBuffer& Buffer, bool bDiscardSharedConstants)
+	D3D12_STATE_CACHE_INLINE void SetConstantBuffer(EShaderFrequency ShaderFrequency, FD3D12ConstantBuffer& Buffer, bool bDiscardSharedConstants)
 	{
 		FD3D12ResourceLocation Location(GetParentDevice());
 
@@ -614,10 +574,18 @@ public:
 	void SetBlendFactor(const float BlendFactor[4]);
 	void SetStencilRef(uint32 StencilRef);
 
-	template <typename TShader>
-	D3D12_STATE_CACHE_INLINE TShader* GetShader()
+	FRHIShader* GetShader(EShaderFrequency InFrequency)
 	{
-		return TStateCacheShaderTraits<TShader>::GetShader(GetGraphicsPipelineState());
+		switch (InFrequency)
+		{
+		case SF_Vertex:        return GetGraphicsPipelineState()->GetVertexShader();
+		case SF_Mesh:          return GetGraphicsPipelineState()->GetMeshShader();
+		case SF_Amplification: return GetGraphicsPipelineState()->GetAmplificationShader();
+		case SF_Pixel:         return GetGraphicsPipelineState()->GetPixelShader();
+		case SF_Geometry:      return GetGraphicsPipelineState()->GetGeometryShader();
+		case SF_Compute:       return GetComputePipelineState()->GetComputeShader();
+		default:               return nullptr;
+		}
 	}
 
 	void SetNewShaderData(EShaderFrequency InFrequency, const FD3D12ShaderData* InShaderData);
@@ -713,11 +681,8 @@ public:
 		}
 	}
 
-	template <EShaderFrequency ShaderStage>
-	void SetUAVs(uint32 UAVStartSlot, uint32 NumSimultaneousUAVs, FD3D12UnorderedAccessView** UAVArray, uint32* UAVInitialCountArray);
-	template <EShaderFrequency ShaderStage>
-	void ClearUAVs();
-
+	void SetUAVs(EShaderFrequency ShaderStage, uint32 UAVStartSlot, uint32 NumSimultaneousUAVs, FD3D12UnorderedAccessView** UAVArray, uint32* UAVInitialCountArray);
+	void ClearUAVs(EShaderFrequency ShaderStage);
 
 	void SetDepthBounds(float MinDepth, float MaxDepth)
 	{
diff --git a/Engine/Source/Runtime/RHI/Public/RHIDefinitions.h b/Engine/Source/Runtime/RHI/Public/RHIDefinitions.h
--- a/Engine/Source/Runtime/RHI/Public/RHIDefinitions.h
+++ b/Engine/Source/Runtime/RHI/Public/RHIDefinitions.h
@@ -107,6 +107,21 @@ enum EShaderFrequency : uint8
 };
 static_assert(SF_NumFrequencies <= (1 << SF_NumBits), "SF_NumFrequencies will not fit on SF_NumBits");
 
+inline bool IsValidGraphicsFrequency(EShaderFrequency InShaderFrequency)
+{
+	switch (InShaderFrequency)
+	{
+	case SF_Vertex:        return true;
+#if PLATFORM_SUPPORTS_MESH_SHADERS
+	case SF_Mesh:          return true;
+	case SF_Amplification: return true;
+#endif
+	case SF_Pixel:         return true;
+	case SF_Geometry:      return true;
+	}
+	return false;
+}
+
 inline bool IsComputeShaderFrequency(EShaderFrequency ShaderFrequency)
 {
 	switch (ShaderFrequency)
