fix for RHIAsyncCreateTexture2D with split allocation

It is related to the way UpdateSubresources / GetCopyableFootprints / GetRequiredIntermediateSize handle the FirstSubresource parameter.
We were expecting these functions would return offsets / size that would ignore the [0,FirstSubresource[ but it's not the case.

I'm changing the code to perform the size / offsets computation ourselves and call the UpdateSubresources function that take D3D12_PLACED_SUBRESOURCE_FOOTPRINT as parameters.
As a bonus, we bypass the HeapAlloc/HeapFree calls for these temporary structures, since we know the upper bound on the number of mips.

diff --git a/Engine/Source/Runtime/D3D12RHI/Private/D3D12Texture.cpp b/Engine/Source/Runtime/D3D12RHI/Private/D3D12Texture.cpp
--- a/Engine/Source/Runtime/D3D12RHI/Private/D3D12Texture.cpp
+++ b/Engine/Source/Runtime/D3D12RHI/Private/D3D12Texture.cpp
@@ -1063,7 +1063,7 @@ FTextureRHIRef FD3D12DynamicRHI::RHIAsyncCreateTexture2D(uint32 SizeX, uint32 Si
 
 		FD3D12FastAllocator& FastAllocator = TextureOut->GetParentDevice()->GetDefaultFastAllocator();
 		uint64 Size = GetRequiredIntermediateSize(TextureOut->GetResource()->GetResource(), 0, NumMips);
-		uint64 SizeLowMips;
+		uint64 SizeLowMips = 0;
 
 		FD3D12ResourceLocation TempResourceLocation(FastAllocator.GetParentDevice());
 		FD3D12ResourceLocation TempResourceLocationLowMips(FastAllocator.GetParentDevice());
@@ -1074,13 +1074,45 @@ FTextureRHIRef FD3D12DynamicRHI::RHIAsyncCreateTexture2D(uint32 SizeX, uint32 Si
 		// we split the top mip into a separate allocation, allowing it to fit within 4MB.
 		const bool bSplitAllocation = (Size > 4 * 1024 * 1024) && (NumMips > 1);
 
+		// Data used for split allocation - Workaround for GetCopyableFootprints returning unexpected values, see UE-173385
+		ID3D12Device* D3D12Device = TextureOut->GetParentDevice()->GetDevice();
+		D3D12_PLACED_SUBRESOURCE_FOOTPRINT Layouts[MAX_TEXTURE_MIP_COUNT] = { };
+		UINT NumRows[MAX_TEXTURE_MIP_COUNT] = { };
+		UINT64 RowSizesInBytes[MAX_TEXTURE_MIP_COUNT] = { };
+		UINT64 TotalBytes = 0; // unused, got some issues with it
+		uint64 SizeMip0 = 0;
 		if (bSplitAllocation)
 		{
-			Size = GetRequiredIntermediateSize(TextureOut->GetResource()->GetResource(), 0, 1);
-			SizeLowMips = GetRequiredIntermediateSize(TextureOut->GetResource()->GetResource(), 1, NumMips - 1);
-
-			FastAllocator.Allocate(Size, D3D12_TEXTURE_DATA_PLACEMENT_ALIGNMENT, &TempResourceLocation);
+			// Setup for the copies: we get the fullmip chain here to get the offsets first
+			const uint64 FirstSubresource = 0;
+			D3D12Device->GetCopyableFootprints(&TextureDesc, FirstSubresource, NumMips, 0, Layouts, NumRows, RowSizesInBytes, &TotalBytes);
+			
+			// Mip 0
+			SizeMip0 = RowSizesInBytes[0] * NumRows[0];
+			FastAllocator.Allocate(SizeMip0, D3D12_TEXTURE_DATA_PLACEMENT_ALIGNMENT, &TempResourceLocation);
+			Layouts[0].Offset = TempResourceLocation.GetOffsetFromBaseOfResource();
+
+			// Remaining mip chain
+			SizeLowMips = 0;
+			for (uint64 MipIndex = 1; MipIndex < NumMips; ++MipIndex)
+			{
+				SizeLowMips += RowSizesInBytes[MipIndex] * NumRows[MipIndex];
+			}
 			FastAllocator.Allocate(SizeLowMips, D3D12_TEXTURE_DATA_PLACEMENT_ALIGNMENT, &TempResourceLocationLowMips);
+
+			uint64 BaseOffset = Layouts[1].Offset;
+			for (uint64 MipIndex = 1; MipIndex < NumMips; ++MipIndex)
+			{
+				check(Layouts[MipIndex].Offset >= BaseOffset);
+				// The original offsets for the remaining mipchain were originally computed with mip0, so we need to remove that offset
+				Layouts[MipIndex].Offset -= BaseOffset;
+				// The intermediate resource we get might be already used, so we need to account for the offset within this resource
+				Layouts[MipIndex].Offset += TempResourceLocationLowMips.GetOffsetFromBaseOfResource();
+				// We are going to map TempResourceLocationLowMips.GetResource(), and memcpy RowSizesInBytes[MipIndex] * NumRows[MipIndex] bytes in there offsetted
+				// Make sure that the buffer is large enough before proceeding
+				check(Layouts[MipIndex].Offset + RowSizesInBytes[MipIndex] * NumRows[MipIndex] <= TempResourceLocationLowMips.GetResource()->GetDesc().Width);
+			}
+			
 			TempResourceLocationLowMips.GetResource()->AddRef();
 		}
 		else
@@ -1109,21 +1141,34 @@ FTextureRHIRef FD3D12DynamicRHI::RHIAsyncCreateTexture2D(uint32 SizeX, uint32 Si
 
 				if (bSplitAllocation)
 				{
-					UpdateSubresources(
+					UINT64 SizeCopiedMip0 = UpdateSubresources(
 						CopyScope.Context.CopyCommandList().Get(),
 						Resource->GetResource(),
 						TempResourceLocation.GetResource()->GetResource(),
-						TempResourceLocation.GetOffsetFromBaseOfResource(),
-						0, 1,
+						0, // FirstSubresource
+						1, // NumSubresources
+						SizeMip0, // RequiredSize
+						Layouts,
+						NumRows,
+						RowSizesInBytes,
 						SubResourceData);
 
-					UpdateSubresources(
+					ensure(SizeCopiedMip0 == SizeMip0);
+
+					UINT64 SizeCopiedLowMips = UpdateSubresources(
 						CopyScope.Context.CopyCommandList().Get(),
 						Resource->GetResource(),
 						TempResourceLocationLowMips.GetResource()->GetResource(),
-						TempResourceLocationLowMips.GetOffsetFromBaseOfResource(),
-						1, NumMips - 1,
-						SubResourceData + 1);
+						1, // FirstSubresource
+						NumMips-1, // NumSubresources
+						SizeLowMips, // RequiredSize
+						Layouts+1,
+						NumRows+1,
+						RowSizesInBytes+1,
+						SubResourceData+1);
+
+					ensure(SizeCopiedLowMips == SizeLowMips);
+
 				}
 				else
 				{
