Re-disabling ValidateBoundShader checks on compute shaders since I didn't test it thoroughly enough.

diff --git a/Engine/Source/Runtime/D3D12RHI/Private/D3D12Commands.cpp b/Engine/Source/Runtime/D3D12RHI/Private/D3D12Commands.cpp
--- a/Engine/Source/Runtime/D3D12RHI/Private/D3D12Commands.cpp
+++ b/Engine/Source/Runtime/D3D12RHI/Private/D3D12Commands.cpp
@@ -882,7 +882,7 @@ void FD3D12CommandContext::RHISetShaderTexture(FRHIGraphicsShader* ShaderRHI, ui
 
 void FD3D12CommandContext::RHISetShaderTexture(FRHIComputeShader* ComputeShaderRHI, uint32 TextureIndex, FRHITexture* NewTextureRHI)
 {
-	ValidateBoundShader(StateCache, ComputeShaderRHI);
+	//ValidateBoundShader(StateCache, ComputeShaderRHI);
 
 	FD3D12Texture* const NewTexture = RetrieveTexture(NewTextureRHI);
 	FD3D12ShaderResourceView* ViewToSet = NewTexture ? NewTexture->GetShaderResourceView() : nullptr;
@@ -908,7 +908,7 @@ void FD3D12CommandContext::RHISetUAVParameter(FRHIPixelShader* PixelShaderRHI, u
 
 void FD3D12CommandContext::RHISetUAVParameter(FRHIComputeShader* ComputeShaderRHI, uint32 UAVIndex, FRHIUnorderedAccessView* UAVRHI)
 {
-	ValidateBoundShader(StateCache, ComputeShaderRHI);
+	//ValidateBoundShader(StateCache, ComputeShaderRHI);
 
 	FD3D12UnorderedAccessView* UAV = RetrieveObject<FD3D12UnorderedAccessView>(UAVRHI);
 
@@ -925,7 +925,7 @@ void FD3D12CommandContext::RHISetUAVParameter(FRHIComputeShader* ComputeShaderRH
 
 void FD3D12CommandContext::RHISetUAVParameter(FRHIComputeShader* ComputeShaderRHI, uint32 UAVIndex, FRHIUnorderedAccessView* UAVRHI, uint32 InitialCount)
 {
-	ValidateBoundShader(StateCache, ComputeShaderRHI);
+	//ValidateBoundShader(StateCache, ComputeShaderRHI);
 
 	FD3D12UnorderedAccessView* UAV = RetrieveObject<FD3D12UnorderedAccessView>(UAVRHI);
 
@@ -957,7 +957,7 @@ void FD3D12CommandContext::RHISetShaderResourceViewParameter(FRHIGraphicsShader*
 
 void FD3D12CommandContext::RHISetShaderResourceViewParameter(FRHIComputeShader* ComputeShaderRHI, uint32 TextureIndex, FRHIShaderResourceView* SRVRHI)
 {
-	ValidateBoundShader(StateCache, ComputeShaderRHI);
+	//ValidateBoundShader(StateCache, ComputeShaderRHI);
 	FD3D12ShaderResourceView* const SRV = RetrieveObject<FD3D12ShaderResourceView>(SRVRHI);
 	StateCache.SetShaderResourceView(SF_Compute, SRV, TextureIndex);
 }
@@ -982,7 +982,7 @@ void FD3D12CommandContext::RHISetShaderSampler(FRHIGraphicsShader* ShaderRHI, ui
 
 void FD3D12CommandContext::RHISetShaderSampler(FRHIComputeShader* ComputeShaderRHI, uint32 SamplerIndex, FRHISamplerState* NewStateRHI)
 {
-	ValidateBoundShader(StateCache, ComputeShaderRHI);
+	//ValidateBoundShader(StateCache, ComputeShaderRHI);
 	FD3D12SamplerState* NewState = RetrieveObject<FD3D12SamplerState>(NewStateRHI);
 	StateCache.SetSamplerState(SF_Compute, NewState, SamplerIndex);
 }
@@ -1017,7 +1017,7 @@ void FD3D12CommandContext::RHISetShaderUniformBuffer(FRHIGraphicsShader* ShaderR
 void FD3D12CommandContext::RHISetShaderUniformBuffer(FRHIComputeShader* ComputeShader, uint32 BufferIndex, FRHIUniformBuffer* BufferRHI)
 {
 	//SCOPE_CYCLE_COUNTER(STAT_D3D12SetShaderUniformBuffer);
-	ValidateBoundShader(StateCache, ComputeShader);
+	//ValidateBoundShader(StateCache, ComputeShader);
 	FD3D12UniformBuffer* Buffer = RetrieveObject<FD3D12UniformBuffer>(BufferRHI);
 
 	StateCache.SetConstantsFromUniformBuffer(SF_Compute, BufferIndex, Buffer);
@@ -1074,7 +1074,7 @@ void FD3D12CommandContext::RHISetShaderParameter(FRHIGraphicsShader* ShaderRHI,
 
 void FD3D12CommandContext::RHISetShaderParameter(FRHIComputeShader* ComputeShaderRHI, uint32 BufferIndex, uint32 Offset, uint32 NumBytes, const void* NewValue)
 {
-	ValidateBoundShader(StateCache, ComputeShaderRHI);
+	//ValidateBoundShader(StateCache, ComputeShaderRHI);
 	checkSlow(BufferIndex == 0);
 	StageConstantBuffers[SF_Compute].UpdateConstant((const uint8*)NewValue, Offset, NumBytes);
 }
