Explicitly set bindless descriptor heaps before dispatching ray tracing work

Previously the code assumed that bindless heap is already set by default, which is not guaranteed to be the case.

diff --git a/Engine/Source/Runtime/D3D12RHI/Private/D3D12RayTracing.cpp b/Engine/Source/Runtime/D3D12RHI/Private/D3D12RayTracing.cpp
--- a/Engine/Source/Runtime/D3D12RHI/Private/D3D12RayTracing.cpp
+++ b/Engine/Source/Runtime/D3D12RHI/Private/D3D12RayTracing.cpp
@@ -1602,26 +1602,43 @@ public:
 		bBindlessViews = BindlessManager.HasHeapForType(ERHIDescriptorHeapType::Standard);
 		bBindlessSamplers = BindlessManager.HasHeapForType(ERHIDescriptorHeapType::Sampler);
 
-		ViewHeap.Init(NumViewDescriptors, D3D12_DESCRIPTOR_HEAP_TYPE_CBV_SRV_UAV);
-		SamplerHeap.Init(NumSamplerDescriptors, D3D12_DESCRIPTOR_HEAP_TYPE_SAMPLER);
+		if (!bBindlessViews)
+		{
+			ViewHeap.Init(NumViewDescriptors, D3D12_DESCRIPTOR_HEAP_TYPE_CBV_SRV_UAV);
+		}
+		
+		if (!bBindlessSamplers)
+		{
+			SamplerHeap.Init(NumSamplerDescriptors, D3D12_DESCRIPTOR_HEAP_TYPE_SAMPLER);
+		}
 	}
 
 	void UpdateSyncPoint()
 	{
 		check(IsInRHIThread() || !IsRunningRHIInSeparateThread());
-		ViewHeap.UpdateSyncPoint();
-		SamplerHeap.UpdateSyncPoint();
+
+		if (!bBindlessViews)
+		{
+			ViewHeap.UpdateSyncPoint();
+		}
+
+		if (!bBindlessSamplers)
+		{
+			SamplerHeap.UpdateSyncPoint();
+		}
 	}
 
 	void SetDescriptorHeaps(FD3D12CommandContext& CommandContext)
 	{
 		UpdateSyncPoint();
 
-		check(ViewHeap.GetParentDevice() == CommandContext.GetParentDevice());
-		check(SamplerHeap.GetParentDevice() == CommandContext.GetParentDevice());
+		check(bBindlessViews || ViewHeap.GetParentDevice() == CommandContext.GetParentDevice());
+		check(bBindlessSamplers || SamplerHeap.GetParentDevice() == CommandContext.GetParentDevice());
+
+		FD3D12BindlessDescriptorManager& BindlessManager = GetParentDevice()->GetBindlessDescriptorManager();
 
-		ID3D12DescriptorHeap* ViewHeapToSet = bBindlessViews ? nullptr : ViewHeap.D3D12Heap;
-		ID3D12DescriptorHeap* SamplerHeapToSet = bBindlessSamplers ? nullptr : SamplerHeap.D3D12Heap;
+		ID3D12DescriptorHeap* ViewHeapToSet = bBindlessViews ? BindlessManager.GetHeapForType(ERHIDescriptorHeapType::Standard)->GetHeap() : ViewHeap.D3D12Heap;
+		ID3D12DescriptorHeap* SamplerHeapToSet = bBindlessSamplers ? BindlessManager.GetHeapForType(ERHIDescriptorHeapType::Sampler)->GetHeap() : SamplerHeap.D3D12Heap;
 
 		CommandContext.StateCache.GetDescriptorCache()->OverrideLastSetHeaps(ViewHeapToSet, SamplerHeapToSet);
 	}
