D3D12 Bindless resource state and residency tracking
- Making custom implementation of RHISetShaderParameters for D3D12.
- Bindless resources are now queued for transitions and residency tracking before the next draw/dispatch.

diff --git a/Engine/Source/Runtime/D3D12RHI/Private/D3D12CommandContext.h b/Engine/Source/Runtime/D3D12RHI/Private/D3D12CommandContext.h
--- a/Engine/Source/Runtime/D3D12RHI/Private/D3D12CommandContext.h
+++ b/Engine/Source/Runtime/D3D12RHI/Private/D3D12CommandContext.h
@@ -660,13 +660,13 @@ public:
 #endif // D3D12_RHI_RAYTRACING
 
 	template<typename ObjectType, typename RHIType>
-	static FORCEINLINE_DEBUGGABLE ObjectType* RetrieveObject(RHIType RHIObject, uint32 GPUIndex)
+	static FORCEINLINE_DEBUGGABLE ObjectType* RetrieveObject(RHIType* RHIObject, uint32 GPUIndex)
 	{
-		return FD3D12DynamicRHI::ResourceCast(RHIObject, GPUIndex);
+		return FD3D12DynamicRHI::ResourceCast<RHIType, ObjectType>(RHIObject, GPUIndex);
 	}
 
 	template<typename ObjectType, typename RHIType>
-	FORCEINLINE_DEBUGGABLE ObjectType* RetrieveObject(RHIType RHIObject)
+	FORCEINLINE_DEBUGGABLE ObjectType* RetrieveObject(RHIType* RHIObject)
 	{
 		return RetrieveObject<ObjectType, RHIType>(RHIObject, GetGPUIndex());
 	}
diff --git a/Engine/Source/Runtime/D3D12RHI/Private/D3D12Commands.cpp b/Engine/Source/Runtime/D3D12RHI/Private/D3D12Commands.cpp
--- a/Engine/Source/Runtime/D3D12RHI/Private/D3D12Commands.cpp
+++ b/Engine/Source/Runtime/D3D12RHI/Private/D3D12Commands.cpp
@@ -87,6 +87,27 @@ inline void ValidateBoundUniformBuffer(FD3D12UniformBuffer * InUniformBuffer, FR
 #endif
 }
 
+inline FD3D12ShaderResourceView* GetShaderResourceViewFromTexture(FRHIResource* InResource, uint32 GpuIndex)
+{
+	FD3D12Texture* Texture = FD3D12CommandContext::RetrieveTexture(static_cast<FRHITexture*>(InResource), GpuIndex);
+	return Texture ? Texture->GetShaderResourceView() : nullptr;
+}
+
+static void BindUniformBuffer(FD3D12CommandContext& Context, FRHIShader* Shader, EShaderFrequency ShaderFrequency, uint32 BufferIndex, FD3D12UniformBuffer* InBuffer)
+{
+	ValidateBoundUniformBuffer(InBuffer, Shader, BufferIndex);
+
+	Context.StateCache.SetConstantsFromUniformBuffer(ShaderFrequency, BufferIndex, InBuffer);
+
+	if (!GRHINeedsExtraDeletionLatency)
+	{
+		Context.BoundUniformBufferRefs[ShaderFrequency][BufferIndex] = InBuffer;
+	}
+
+	Context.BoundUniformBuffers[ShaderFrequency][BufferIndex] = InBuffer;
+	Context.DirtyUniformBuffers[ShaderFrequency] |= (1 << BufferIndex);
+}
+
 #if !defined(D3D12_PLATFORM_SUPPORTS_RESOLVE_SHADERS)
 	#define D3D12_PLATFORM_SUPPORTS_RESOLVE_SHADERS 1
 #endif
@@ -995,18 +1016,7 @@ void FD3D12CommandContext::RHISetShaderUniformBuffer(FRHIGraphicsShader* ShaderR
 	{
 		ValidateBoundShader(StateCache, ShaderRHI);
 
-		FD3D12UniformBuffer* Buffer = RetrieveObject<FD3D12UniformBuffer>(BufferRHI);
-		ValidateBoundUniformBuffer(Buffer, ShaderRHI, BufferIndex);
-
-		StateCache.SetConstantsFromUniformBuffer(ShaderFrequency, BufferIndex, Buffer);
-
-		if (!GRHINeedsExtraDeletionLatency)
-		{
-			BoundUniformBufferRefs[ShaderFrequency][BufferIndex] = BufferRHI;
-		}
-
-		BoundUniformBuffers[ShaderFrequency][BufferIndex] = Buffer;
-		DirtyUniformBuffers[ShaderFrequency] |= (1 << BufferIndex);
+		BindUniformBuffer(*this, ShaderRHI, ShaderFrequency, BufferIndex, RetrieveObject<FD3D12UniformBuffer>(BufferRHI));
 	}
 	else
 	{
@@ -1018,36 +1028,153 @@ void FD3D12CommandContext::RHISetShaderUniformBuffer(FRHIComputeShader* ComputeS
 {
 	//SCOPE_CYCLE_COUNTER(STAT_D3D12SetShaderUniformBuffer);
 	//ValidateBoundShader(StateCache, ComputeShader);
-	FD3D12UniformBuffer* Buffer = RetrieveObject<FD3D12UniformBuffer>(BufferRHI);
 
-	StateCache.SetConstantsFromUniformBuffer(SF_Compute, BufferIndex, Buffer);
+	BindUniformBuffer(*this, ComputeShader, SF_Compute, BufferIndex, RetrieveObject<FD3D12UniformBuffer>(BufferRHI));
+}
 
-	if (!GRHINeedsExtraDeletionLatency)
+static void SetShaderParametersOnContext(
+	FD3D12CommandContext& Context
+	, FRHIShader* Shader
+	, EShaderFrequency ShaderFrequency
+	, TArrayView<const uint8> InParametersData
+	, TArrayView<const FRHIShaderParameter> InParameters
+	, TArrayView<const FRHIShaderParameterResource> InResourceParameters
+	, TArrayView<const FRHIShaderParameterResource> InBindlessParameters)
+{
+	FD3D12ConstantBuffer& ConstantBuffer = Context.StageConstantBuffers[ShaderFrequency];
+	FD3D12StateCache& StateCache = Context.StateCache;
+	const uint32 GpuIndex = Context.GetGPUIndex();
+
+	for (const FRHIShaderParameter& Parameter : InParameters)
 	{
-		BoundUniformBufferRefs[SF_Compute][BufferIndex] = BufferRHI;
+		checkSlow(Parameter.BufferIndex == 0);
+		ConstantBuffer.UpdateConstant(&InParametersData[Parameter.ByteOffset], Parameter.BaseIndex, Parameter.ByteSize);
 	}
 
-	BoundUniformBuffers[SF_Compute][BufferIndex] = Buffer;
-	DirtyUniformBuffers[SF_Compute] |= (1 << BufferIndex);
+#if PLATFORM_SUPPORTS_BINDLESS_RENDERING
+	for (const FRHIShaderParameterResource& Parameter : InBindlessParameters)
+	{
+		if (FRHIResource* Resource = Parameter.Resource)
+		{
+			FRHIDescriptorHandle Handle;
+
+			switch (Parameter.Type)
+			{
+			case FRHIShaderParameterResource::EType::Texture:
+			{
+				Handle = static_cast<FRHITexture*>(Resource)->GetDefaultBindlessHandle();
+				StateCache.QueueBindlessSRV(ShaderFrequency, GetShaderResourceViewFromTexture(Resource, GpuIndex));
+			}
+			break;
+			case FRHIShaderParameterResource::EType::ResourceView:
+			{
+				Handle = static_cast<FRHIShaderResourceView*>(Resource)->GetBindlessHandle();
+				StateCache.QueueBindlessSRV(ShaderFrequency, FD3D12CommandContext::RetrieveObject<FD3D12ShaderResourceView>(Parameter.Resource, GpuIndex));
+			}
+			break;
+			case FRHIShaderParameterResource::EType::UnorderedAccessView:
+			{
+				Handle = static_cast<FRHIUnorderedAccessView*>(Resource)->GetBindlessHandle();
+				FD3D12UnorderedAccessView* UAV = FD3D12CommandContext::RetrieveObject<FD3D12UnorderedAccessView>(Parameter.Resource, GpuIndex);
+				if (UAV)
+				{
+					Context.ConditionalClearShaderResource(UAV->GetResourceLocation());
+				}
+				StateCache.QueueBindlessUAV(ShaderFrequency, UAV);
+			}
+			break;
+			case FRHIShaderParameterResource::EType::Sampler:
+			{
+				Handle = static_cast<FRHISamplerState*>(Resource)->GetBindlessHandle();
+			}
+			break;
+			}
+
+			if (Handle.IsValid())
+			{
+				const uint32 BindlessIndex = Handle.GetIndex();
+				ConstantBuffer.UpdateConstant(reinterpret_cast<const uint8*>(&BindlessIndex), Parameter.Index, 4);
+			}
+		}
+	}
+#endif
+
+	for (const FRHIShaderParameterResource& Parameter : InResourceParameters)
+	{
+		if (Parameter.Type == FRHIShaderParameterResource::EType::UnorderedAccessView)
+		{
+			if (ShaderFrequency == SF_Pixel)
+			{
+				FD3D12UnorderedAccessView* UAV = FD3D12CommandContext::RetrieveObject<FD3D12UnorderedAccessView>(Parameter.Resource, GpuIndex);
+				if (UAV)
+				{
+					Context.ConditionalClearShaderResource(UAV->GetResourceLocation());
+				}
+				StateCache.SetUAV(SF_Pixel, Parameter.Index, UAV);
+			}
+			else
+			{
+				checkf(false, TEXT("TShaderRHI Can't have compute shader to be set. UAVs are not supported on vertex, tessellation and geometry shaders."));
+			}
+		}
+	}
+
+	for (const FRHIShaderParameterResource& Parameter : InResourceParameters)
+	{
+		switch (Parameter.Type)
+		{
+		case FRHIShaderParameterResource::EType::Texture:
+			StateCache.SetShaderResourceView(ShaderFrequency, GetShaderResourceViewFromTexture(Parameter.Resource, GpuIndex), Parameter.Index);
+			break;
+		case FRHIShaderParameterResource::EType::ResourceView:
+			StateCache.SetShaderResourceView(ShaderFrequency, FD3D12CommandContext::RetrieveObject<FD3D12ShaderResourceView>(Parameter.Resource, GpuIndex), Parameter.Index);
+			break;
+		case FRHIShaderParameterResource::EType::UnorderedAccessView:
+			break;
+		case FRHIShaderParameterResource::EType::Sampler:
+			StateCache.SetSamplerState(ShaderFrequency, FD3D12CommandContext::RetrieveObject<FD3D12SamplerState>(Parameter.Resource, GpuIndex), Parameter.Index);
+			break;
+		case FRHIShaderParameterResource::EType::UniformBuffer:
+			BindUniformBuffer(Context, Shader, ShaderFrequency, Parameter.Index, FD3D12CommandContext::RetrieveObject<FD3D12UniformBuffer>(Parameter.Resource, GpuIndex));
+			break;
+		default:
+			checkf(false, TEXT("Unhandled resource type?"));
+			break;
+		}
+	}
 }
 
 void FD3D12CommandContext::RHISetShaderParameters(FRHIGraphicsShader* Shader, TArrayView<const uint8> InParametersData, TArrayView<const FRHIShaderParameter> InParameters, TArrayView<const FRHIShaderParameterResource> InResourceParameters, TArrayView<const FRHIShaderParameterResource> InBindlessParameters)
 {
-	UE::RHICore::RHISetShaderParametersShared(
-		*this
-		, Shader
-		, InParametersData
-		, InParameters
-		, InResourceParameters
-		, InBindlessParameters
-	);
+	const EShaderFrequency ShaderFrequency = Shader->GetFrequency();
+	if (IsValidGraphicsFrequency(ShaderFrequency))
+	{
+		ValidateBoundShader(StateCache, Shader);
+
+		SetShaderParametersOnContext(
+			*this
+			, Shader
+			, ShaderFrequency
+			, InParametersData
+			, InParameters
+			, InResourceParameters
+			, InBindlessParameters
+		);
+	}
+	else
+	{
+		checkf(0, TEXT("Unsupported FRHIGraphicsShader Type '%s'!"), GetShaderFrequencyString(ShaderFrequency, false));
+	}
 }
 
 void FD3D12CommandContext::RHISetShaderParameters(FRHIComputeShader* Shader, TArrayView<const uint8> InParametersData, TArrayView<const FRHIShaderParameter> InParameters, TArrayView<const FRHIShaderParameterResource> InResourceParameters, TArrayView<const FRHIShaderParameterResource> InBindlessParameters)
 {
-	UE::RHICore::RHISetShaderParametersShared(
+	//ValidateBoundShader(StateCache, Shader);
+
+	SetShaderParametersOnContext(
 		*this
 		, Shader
+		, SF_Compute
 		, InParametersData
 		, InParameters
 		, InResourceParameters
diff --git a/Engine/Source/Runtime/D3D12RHI/Private/D3D12DescriptorCache.cpp b/Engine/Source/Runtime/D3D12RHI/Private/D3D12DescriptorCache.cpp
--- a/Engine/Source/Runtime/D3D12RHI/Private/D3D12DescriptorCache.cpp
+++ b/Engine/Source/Runtime/D3D12RHI/Private/D3D12DescriptorCache.cpp
@@ -401,6 +401,24 @@ void FD3D12DescriptorCache::SetSamplers(const FD3D12RootSignature* RootSignature
 #endif
 }
 
+static D3D12_RESOURCE_STATES GetBindingResourceState(FD3D12CommandContext& Context, EShaderFrequency ShaderStage, FD3D12ShaderResourceView* SRV)
+{
+	D3D12_RESOURCE_STATES State = (ShaderStage == SF_Compute)
+		? D3D12_RESOURCE_STATE_NON_PIXEL_SHADER_RESOURCE
+		: D3D12_RESOURCE_STATE_NON_PIXEL_SHADER_RESOURCE | D3D12_RESOURCE_STATE_PIXEL_SHADER_RESOURCE;
+
+	if (SRV->IsDepthStencilResource())
+	{
+		State |= D3D12_RESOURCE_STATE_DEPTH_READ;
+	}
+	else if (SRV->GetSkipFastClearFinalize())
+	{
+		State |= Context.SkipFastClearEliminateState;
+	}
+
+	return State;
+}
+
 template <EShaderFrequency ShaderStage>
 void FD3D12DescriptorCache::SetSRVs(const FD3D12RootSignature* RootSignature, FD3D12ShaderResourceViewCache& Cache, const SRVSlotMask& SlotsNeededMask, uint32 SlotsNeeded, uint32& HeapSlot)
 {
@@ -424,24 +442,13 @@ void FD3D12DescriptorCache::SetSRVs(const FD3D12RootSignature* RootSignature, FD
 
 	for (uint32 SlotIndex = 0; SlotIndex < SlotsNeeded; SlotIndex++)
 	{
-		if (SRVs[SlotIndex] != nullptr)
+		if (FD3D12ShaderResourceView* SRV = SRVs[SlotIndex])
 		{
-			SrcDescriptors[SlotIndex] = SRVs[SlotIndex]->GetOfflineCpuHandle();
+			SrcDescriptors[SlotIndex] = SRV->GetOfflineCpuHandle();
 
-			D3D12_RESOURCE_STATES State = (ShaderStage == SF_Compute) 
-				? D3D12_RESOURCE_STATE_NON_PIXEL_SHADER_RESOURCE
-				: D3D12_RESOURCE_STATE_NON_PIXEL_SHADER_RESOURCE | D3D12_RESOURCE_STATE_PIXEL_SHADER_RESOURCE;
-
-			if (SRVs[SlotIndex]->IsDepthStencilResource())
-			{
-				State |= D3D12_RESOURCE_STATE_DEPTH_READ;
-			}
-			else if (SRVs[SlotIndex]->GetSkipFastClearFinalize())
-			{
-				State |= Context.SkipFastClearEliminateState;
-			}
-			Context.TransitionResource(SRVs[SlotIndex], State & ValidResourceStates);
+			const D3D12_RESOURCE_STATES State = GetBindingResourceState(Context, ShaderStage, SRV);
 
+			Context.TransitionResource(SRV, State & ValidResourceStates);
 			Context.UpdateResidency(Cache.ResidencyHandles[ShaderStage][SlotIndex]);
 		}
 		else
@@ -482,6 +489,30 @@ void FD3D12DescriptorCache::SetSRVs(const FD3D12RootSignature* RootSignature, FD
 #endif
 }
 
+void FD3D12DescriptorCache::PrepareBindlessViews(EShaderFrequency ShaderStage, TConstArrayView<FD3D12ShaderResourceView*> SRVs, TConstArrayView<FD3D12UnorderedAccessView*> UAVs)
+{
+	const D3D12_RESOURCE_STATES ValidResourceStates = Context.ValidResourceStates;
+
+	for (FD3D12ShaderResourceView* SRV : SRVs)
+	{
+		if (ensure(SRV))
+		{
+			const D3D12_RESOURCE_STATES State = GetBindingResourceState(Context, ShaderStage, SRV);
+			Context.TransitionResource(SRV, State & ValidResourceStates);
+			Context.UpdateResidency(&SRV->GetResidencyHandle());
+		}
+	}
+
+	for (FD3D12UnorderedAccessView* UAV : UAVs)
+	{
+		if (ensure(UAV))
+		{
+			Context.TransitionResource(UAV, D3D12_RESOURCE_STATE_UNORDERED_ACCESS);
+			Context.UpdateResidency(&UAV->GetResidencyHandle());
+		}
+	}
+}
+
 void FD3D12DescriptorCache::SetConstantBufferViews(EShaderFrequency ShaderStage, const FD3D12RootSignature* RootSignature, FD3D12ConstantBufferCache& Cache, CBVSlotMask SlotsNeededMask, uint32 SlotsNeeded, uint32& HeapSlot)
 {
 #if USE_STATIC_ROOT_SIGNATURE
diff --git a/Engine/Source/Runtime/D3D12RHI/Private/D3D12DescriptorCache.h b/Engine/Source/Runtime/D3D12RHI/Private/D3D12DescriptorCache.h
--- a/Engine/Source/Runtime/D3D12RHI/Private/D3D12DescriptorCache.h
+++ b/Engine/Source/Runtime/D3D12RHI/Private/D3D12DescriptorCache.h
@@ -352,6 +352,8 @@ public:
 
 	void SetRootConstantBuffers(EShaderFrequency ShaderStage, const FD3D12RootSignature* RootSignature, FD3D12ConstantBufferCache& Cache, CBVSlotMask SlotsNeededMask);
 
+	void PrepareBindlessViews(EShaderFrequency ShaderStage, TConstArrayView<FD3D12ShaderResourceView*> SRVs, TConstArrayView<FD3D12UnorderedAccessView*> UAVs);
+
 	void SetStreamOutTargets(FD3D12Resource **Buffers, uint32 Count, const uint32* Offsets);
 
 	bool HeapRolledOver(ERHIDescriptorHeapType InHeapType);
diff --git a/Engine/Source/Runtime/D3D12RHI/Private/D3D12RHIPrivate.h b/Engine/Source/Runtime/D3D12RHI/Private/D3D12RHIPrivate.h
--- a/Engine/Source/Runtime/D3D12RHI/Private/D3D12RHIPrivate.h
+++ b/Engine/Source/Runtime/D3D12RHI/Private/D3D12RHIPrivate.h
@@ -391,18 +391,17 @@ public:
 	virtual void Shutdown() override;
 	virtual const TCHAR* GetName() override { return TEXT("D3D12"); }
 
-	template<typename TRHIType>
-	static FORCEINLINE typename TD3D12ResourceTraits<TRHIType>::TConcreteType* ResourceCast(TRHIType* Resource)
+	template<typename TRHIType, typename TReturnType = TD3D12ResourceTraits<TRHIType>::TConcreteType>
+	static FORCEINLINE TReturnType* ResourceCast(TRHIType* Resource)
 	{
-		return static_cast<typename TD3D12ResourceTraits<TRHIType>::TConcreteType*>(Resource);
+		return static_cast<TReturnType*>(Resource);
 	}
 
-	template<typename TRHIType>
-	static FORCEINLINE_DEBUGGABLE typename TD3D12ResourceTraits<TRHIType>::TConcreteType* ResourceCast(TRHIType* Resource, uint32 GPUIndex)
+	template<typename TRHIType, typename TReturnType = TD3D12ResourceTraits<TRHIType>::TConcreteType>
+	static FORCEINLINE_DEBUGGABLE TReturnType* ResourceCast(TRHIType* Resource, uint32 GPUIndex)
 	{
-		using ReturnType = typename TD3D12ResourceTraits<TRHIType>::TConcreteType;
-		ReturnType* Object = ResourceCast(Resource);
-		return Object ? static_cast<ReturnType*>(Object->GetLinkedObject(GPUIndex)) : nullptr;
+		TReturnType* Object = ResourceCast<TRHIType, TReturnType>(Resource);
+		return Object ? static_cast<TReturnType*>(Object->GetLinkedObject(GPUIndex)) : nullptr;
 	}
 
 	virtual FD3D12CommandContext* CreateCommandContext(FD3D12Device* InParent, ED3D12QueueType InQueueType, bool InIsDefaultContext);
diff --git a/Engine/Source/Runtime/D3D12RHI/Private/D3D12StateCache.cpp b/Engine/Source/Runtime/D3D12RHI/Private/D3D12StateCache.cpp
--- a/Engine/Source/Runtime/D3D12RHI/Private/D3D12StateCache.cpp
+++ b/Engine/Source/Runtime/D3D12RHI/Private/D3D12StateCache.cpp
@@ -558,6 +558,10 @@ void FD3D12StateCache::ApplyState()
 	{
 		ApplyResources(PSOCommonData->RootSignature, StartStage, EndStage);
 	}
+	else if (bBindlessResources)
+	{
+		ApplyBindlessResources(PSOCommonData->RootSignature, StartStage, EndStage);
+	}
 
 	ApplyConstants(PSOCommonData->RootSignature, StartStage, EndStage);
 
@@ -745,6 +749,20 @@ void FD3D12StateCache::ApplyResources(const FD3D12RootSignature* const pRootSign
 #endif // USE_STATIC_ROOT_SIGNATURE
 }
 
+void FD3D12StateCache::ApplyBindlessResources(const FD3D12RootSignature* const pRootSignature, uint32 StartStage, uint32 EndStage)
+{
+	for (uint32 Index = StartStage; Index < EndStage; Index++)
+	{
+		DescriptorCache.PrepareBindlessViews(
+			static_cast<EShaderFrequency>(Index)
+			, PipelineState.Common.QueuedBindlessSRVs[Index]
+			, PipelineState.Common.QueuedBindlessUAVs[Index]);
+
+		PipelineState.Common.QueuedBindlessSRVs[Index].Reset();
+		PipelineState.Common.QueuedBindlessUAVs[Index].Reset();
+	}
+}
+
 void FD3D12StateCache::ApplyConstants(const FD3D12RootSignature* const pRootSignature, uint32 StartStage, uint32 EndStage)
 {
 #if !USE_STATIC_ROOT_SIGNATURE
diff --git a/Engine/Source/Runtime/D3D12RHI/Private/D3D12StateCachePrivate.h b/Engine/Source/Runtime/D3D12RHI/Private/D3D12StateCachePrivate.h
--- a/Engine/Source/Runtime/D3D12RHI/Private/D3D12StateCachePrivate.h
+++ b/Engine/Source/Runtime/D3D12RHI/Private/D3D12StateCachePrivate.h
@@ -402,6 +402,9 @@ protected:
 			uint32 CurrentShaderSRVCounts    [SF_NumStandardFrequencies] = {};
 			uint32 CurrentShaderCBCounts     [SF_NumStandardFrequencies] = {};
 			uint32 CurrentShaderUAVCounts    [SF_NumStandardFrequencies] = {};
+
+			TArray<FD3D12ShaderResourceView*> QueuedBindlessSRVs[SF_NumStandardFrequencies];
+			TArray<FD3D12UnorderedAccessView*> QueuedBindlessUAVs[SF_NumStandardFrequencies];
 		} Common = {};
 	} PipelineState = {};
 
@@ -659,6 +662,7 @@ public:
 	void ApplyState();
 	void ApplySamplers(const FD3D12RootSignature* const pRootSignature, uint32 StartStage, uint32 EndStage);
 	void ApplyResources(const FD3D12RootSignature* const pRootSignature, uint32 StartStage, uint32 EndStage);
+	void ApplyBindlessResources(const FD3D12RootSignature* const pRootSignature, uint32 StartStage, uint32 EndStage);
 	void ApplyConstants(const FD3D12RootSignature* const pRootSignature, uint32 StartStage, uint32 EndStage);
 	void DirtyStateForNewCommandList();
 	void DirtyState();
@@ -682,6 +686,11 @@ public:
 	}
 
 	void SetUAVs(EShaderFrequency ShaderStage, uint32 UAVStartSlot, uint32 NumSimultaneousUAVs, FD3D12UnorderedAccessView** UAVArray, uint32* UAVInitialCountArray);
+	void SetUAV(EShaderFrequency ShaderStage, uint32 SlotIndex, FD3D12UnorderedAccessView* UAV)
+	{
+		uint32 InitialCount = -1;
+		SetUAVs(ShaderStage, SlotIndex, 1, &UAV, &InitialCount);
+	}
 	void ClearUAVs(EShaderFrequency ShaderStage);
 
 	void SetDepthBounds(float MinDepth, float MaxDepth)
@@ -738,4 +747,13 @@ public:
 	void ClearState();
 
 	void ForceSetComputeRootSignature() { PipelineState.Compute.bNeedSetRootSignature = true; }
+
+	void QueueBindlessSRV(EShaderFrequency ShaderFrequency, FD3D12ShaderResourceView* SRV)
+	{
+		PipelineState.Common.QueuedBindlessSRVs[ShaderFrequency].Emplace(SRV);
+	}
+	void QueueBindlessUAV(EShaderFrequency ShaderFrequency, FD3D12UnorderedAccessView* UAV)
+	{
+		PipelineState.Common.QueuedBindlessUAVs[ShaderFrequency].Emplace(UAV);
+	}
 };
diff --git a/Engine/Source/Runtime/D3D12RHI/Private/D3D12Viewport.cpp b/Engine/Source/Runtime/D3D12RHI/Private/D3D12Viewport.cpp
--- a/Engine/Source/Runtime/D3D12RHI/Private/D3D12Viewport.cpp
+++ b/Engine/Source/Runtime/D3D12RHI/Private/D3D12Viewport.cpp
@@ -667,8 +667,8 @@ bool FD3D12Viewport::Present(bool bLockToVsync)
 		FD3D12CommandContext& DefaultContext = Adapter->GetDevice(GPUIndex)->GetDefaultCommandContext();
 
 		// Those are not necessarily the swap chain back buffer in case of multi-gpu
-		FD3D12Texture* DeviceBackBuffer = DefaultContext.RetrieveObject<FD3D12Texture, FRHITexture2D*>(GetBackBuffer_RHIThread());
-		FD3D12Texture* DeviceSDRBackBuffer = DefaultContext.RetrieveObject<FD3D12Texture, FRHITexture2D*>(GetSDRBackBuffer_RHIThread());
+		FD3D12Texture* DeviceBackBuffer = DefaultContext.RetrieveObject<FD3D12Texture, FRHITexture2D>(GetBackBuffer_RHIThread());
+		FD3D12Texture* DeviceSDRBackBuffer = DefaultContext.RetrieveObject<FD3D12Texture, FRHITexture2D>(GetSDRBackBuffer_RHIThread());
 
 		DefaultContext.TransitionResource(
 			DeviceBackBuffer->GetShaderResourceView()->GetResource(),
