Fixing bad type check that slid in after testing of previous change.

diff --git a/Engine/Source/Runtime/D3D12RHI/Private/D3D12Commands.cpp b/Engine/Source/Runtime/D3D12RHI/Private/D3D12Commands.cpp
--- a/Engine/Source/Runtime/D3D12RHI/Private/D3D12Commands.cpp
+++ b/Engine/Source/Runtime/D3D12RHI/Private/D3D12Commands.cpp
@@ -1103,14 +1103,14 @@ static void SetShaderParametersOnContext(
 	{
 		if (Parameter.Type == FRHIShaderParameterResource::EType::UnorderedAccessView)
 		{
-			if (ShaderFrequency == SF_Pixel)
+			if (ShaderFrequency == SF_Pixel || ShaderFrequency == SF_Compute)
 			{
 				FD3D12UnorderedAccessView* UAV = FD3D12CommandContext::RetrieveObject<FD3D12UnorderedAccessView>(Parameter.Resource, GpuIndex);
 				if (UAV)
 				{
 					Context.ConditionalClearShaderResource(UAV->GetResourceLocation());
 				}
-				StateCache.SetUAV(SF_Pixel, Parameter.Index, UAV);
+				StateCache.SetUAV(ShaderFrequency, Parameter.Index, UAV);
 			}
 			else
 			{
