Fixing template error in D3D12RHI.

diff --git a/Engine/Source/Runtime/D3D12RHI/Private/D3D12RHIPrivate.h b/Engine/Source/Runtime/D3D12RHI/Private/D3D12RHIPrivate.h
--- a/Engine/Source/Runtime/D3D12RHI/Private/D3D12RHIPrivate.h
+++ b/Engine/Source/Runtime/D3D12RHI/Private/D3D12RHIPrivate.h
@@ -391,13 +391,13 @@ public:
 	virtual void Shutdown() override;
 	virtual const TCHAR* GetName() override { return TEXT("D3D12"); }
 
-	template<typename TRHIType, typename TReturnType = TD3D12ResourceTraits<TRHIType>::TConcreteType>
+	template<typename TRHIType, typename TReturnType = typename TD3D12ResourceTraits<TRHIType>::TConcreteType>
 	static FORCEINLINE TReturnType* ResourceCast(TRHIType* Resource)
 	{
 		return static_cast<TReturnType*>(Resource);
 	}
 
-	template<typename TRHIType, typename TReturnType = TD3D12ResourceTraits<TRHIType>::TConcreteType>
+	template<typename TRHIType, typename TReturnType = typename TD3D12ResourceTraits<TRHIType>::TConcreteType>
 	static FORCEINLINE_DEBUGGABLE TReturnType* ResourceCast(TRHIType* Resource, uint32 GPUIndex)
 	{
 		TReturnType* Object = ResourceCast<TRHIType, TReturnType>(Resource);
