Changing more strongly typed D3D12 code to be based on shader frequency indices. There was already support for indexing like this and the templates aren't actually necessary.

diff --git a/Engine/Source/Runtime/D3D12RHI/Private/D3D12Commands.cpp b/Engine/Source/Runtime/D3D12RHI/Private/D3D12Commands.cpp
--- a/Engine/Source/Runtime/D3D12RHI/Private/D3D12Commands.cpp
+++ b/Engine/Source/Runtime/D3D12RHI/Private/D3D12Commands.cpp
@@ -158,7 +158,7 @@ void FD3D12CommandContext::RHIDispatchComputeShader(uint32 ThreadGroupCountX, ui
 
 	CommitComputeShaderConstants();
 	CommitComputeResourceTables();
-	StateCache.ApplyState<ED3D12PipelineType::Compute>();
+	StateCache.ApplyState(ED3D12PipelineType::Compute);
 
 	GraphicsCommandList()->Dispatch(ThreadGroupCountX, ThreadGroupCountY, ThreadGroupCountZ);
 	
@@ -181,7 +181,7 @@ void FD3D12CommandContext::RHIDispatchIndirectComputeShader(FRHIBuffer* Argument
 
 	FD3D12ResourceLocation& Location = ArgumentBuffer->ResourceLocation;
 
-	StateCache.ApplyState<ED3D12PipelineType::Compute>();
+	StateCache.ApplyState(ED3D12PipelineType::Compute);
 
 	// Indirect args buffer can be a previously pending UAV, which becomes PS\Non-PS read. ApplyState will flush pending transitions, so enqueue the indirect arg transition and flush here.
 	D3D12_RESOURCE_STATES IndirectState = D3D12_RESOURCE_STATE_INDIRECT_ARGUMENT | D3D12_RESOURCE_STATE_NON_PIXEL_SHADER_RESOURCE;
@@ -1591,7 +1591,7 @@ void FD3D12CommandContext::RHIDrawPrimitive(uint32 BaseVertexIndex, uint32 NumPr
 		GetParentDevice()->RegisterGPUWork(NumPrimitives * NumInstances, VertexCount * NumInstances);
 	}
 
-	StateCache.ApplyState<ED3D12PipelineType::Graphics>();
+	StateCache.ApplyState(ED3D12PipelineType::Graphics);
 	GraphicsCommandList()->DrawInstanced(VertexCount, NumInstances, BaseVertexIndex, 0);
 
 	ConditionalSplitCommandList();
@@ -1614,7 +1614,7 @@ void FD3D12CommandContext::RHIDrawPrimitiveIndirect(FRHIBuffer* ArgumentBufferRH
 
 	FD3D12ResourceLocation& Location = ArgumentBuffer->ResourceLocation;
 
-	StateCache.ApplyState<ED3D12PipelineType::Graphics>();
+	StateCache.ApplyState(ED3D12PipelineType::Graphics);
 
 	// Indirect args buffer can be a previously pending UAV, which becomes PS\Non-PS read. ApplyState will flush pending transitions, so enqueue the indirect
 	// arg transition and flush here.
@@ -1661,7 +1661,7 @@ void FD3D12CommandContext::RHIDrawIndexedIndirect(FRHIBuffer* IndexBufferRHI, FR
 
 	FD3D12ResourceLocation& Location = ArgumentsBuffer->ResourceLocation;
 
-	StateCache.ApplyState<ED3D12PipelineType::Graphics>();
+	StateCache.ApplyState(ED3D12PipelineType::Graphics);
 
 	// Indirect args buffer can be a previously pending UAV, which becomes PS\Non-PS read. ApplyState will flush pending transitions, so enqueue the indirect
 	// arg transition and flush here.
@@ -1712,7 +1712,7 @@ void FD3D12CommandContext::RHIDrawIndexedPrimitive(FRHIBuffer* IndexBufferRHI, i
 	// determine 16bit vs 32bit indices
 	const DXGI_FORMAT Format = (IndexBuffer->GetStride() == sizeof(uint16) ? DXGI_FORMAT_R16_UINT : DXGI_FORMAT_R32_UINT);
 	StateCache.SetIndexBuffer(IndexBuffer->ResourceLocation, Format, 0);
-	StateCache.ApplyState<ED3D12PipelineType::Graphics>();
+	StateCache.ApplyState(ED3D12PipelineType::Graphics);
 
 	GraphicsCommandList()->DrawIndexedInstanced(IndexCount, NumInstances, StartIndex, BaseVertexIndex, FirstInstance);
 
@@ -1742,7 +1742,7 @@ void FD3D12CommandContext::RHIDrawIndexedPrimitiveIndirect(FRHIBuffer* IndexBuff
 
 	FD3D12ResourceLocation& Location = ArgumentBuffer->ResourceLocation;
 
-	StateCache.ApplyState<ED3D12PipelineType::Graphics>();
+	StateCache.ApplyState(ED3D12PipelineType::Graphics);
 
 	// Indirect args buffer can be a previously pending UAV, which becomes PS\Non-PS read. ApplyState will flush pending transitions, so enqueue the indirect
 	// arg transition and flush here.
@@ -1776,7 +1776,7 @@ void FD3D12CommandContext::RHIDispatchMeshShader(uint32 ThreadGroupCountX, uint3
 	CommitGraphicsResourceTables();
 	CommitNonComputeShaderConstants();
 
-	StateCache.ApplyState<ED3D12PipelineType::Graphics>();
+	StateCache.ApplyState(ED3D12PipelineType::Graphics);
 
 	GraphicsCommandList6()->DispatchMesh(ThreadGroupCountX, ThreadGroupCountY, ThreadGroupCountZ);
 
@@ -1800,7 +1800,7 @@ void FD3D12CommandContext::RHIDispatchIndirectMeshShader(FRHIBuffer* ArgumentBuf
 
 	FD3D12ResourceLocation& Location = ArgumentBuffer->ResourceLocation;
 
-	StateCache.ApplyState<ED3D12PipelineType::Graphics>();
+	StateCache.ApplyState(ED3D12PipelineType::Graphics);
 
 	// Indirect args buffer can be a previously pending UAV, which becomes PS\Non-PS read. ApplyState will flush pending transitions, so enqueue the indirect
 	// arg transition and flush here.
diff --git a/Engine/Source/Runtime/D3D12RHI/Private/D3D12DescriptorCache.cpp b/Engine/Source/Runtime/D3D12RHI/Private/D3D12DescriptorCache.cpp
--- a/Engine/Source/Runtime/D3D12RHI/Private/D3D12DescriptorCache.cpp
+++ b/Engine/Source/Runtime/D3D12RHI/Private/D3D12DescriptorCache.cpp
@@ -5,24 +5,6 @@
 //-----------------------------------------------------------------------------
 #include "D3D12RHIPrivate.h"
 
-// Define template functions that are only declared in the header.
-template void FD3D12DescriptorCache::SetSRVs<SF_Vertex>(const FD3D12RootSignature* RootSignature, FD3D12ShaderResourceViewCache& Cache, const SRVSlotMask& SlotsNeededMask, uint32 SlotsNeeded, uint32& HeapSlot);
-template void FD3D12DescriptorCache::SetSRVs<SF_Mesh>(const FD3D12RootSignature* RootSignature, FD3D12ShaderResourceViewCache& Cache, const SRVSlotMask& SlotsNeededMask, uint32 SlotsNeeded, uint32& HeapSlot);
-template void FD3D12DescriptorCache::SetSRVs<SF_Amplification>(const FD3D12RootSignature* RootSignature, FD3D12ShaderResourceViewCache& Cache, const SRVSlotMask& SlotsNeededMask, uint32 SlotsNeeded, uint32& HeapSlot);
-template void FD3D12DescriptorCache::SetSRVs<SF_Geometry>(const FD3D12RootSignature* RootSignature, FD3D12ShaderResourceViewCache& Cache, const SRVSlotMask& SlotsNeededMask, uint32 SlotsNeeded, uint32& HeapSlot);
-template void FD3D12DescriptorCache::SetSRVs<SF_Pixel>(const FD3D12RootSignature* RootSignature, FD3D12ShaderResourceViewCache& Cache, const SRVSlotMask& SlotsNeededMask, uint32 SlotsNeeded, uint32& HeapSlot);
-template void FD3D12DescriptorCache::SetSRVs<SF_Compute>(const FD3D12RootSignature* RootSignature, FD3D12ShaderResourceViewCache& Cache, const SRVSlotMask& SlotsNeededMask, uint32 SlotsNeeded, uint32& HeapSlot);
-
-template void FD3D12DescriptorCache::SetUAVs<SF_Pixel>(const FD3D12RootSignature* RootSignature, FD3D12UnorderedAccessViewCache& Cache, const UAVSlotMask& SlotsNeededMask, uint32 SlotsNeeded, uint32& HeapSlot);
-template void FD3D12DescriptorCache::SetUAVs<SF_Compute>(const FD3D12RootSignature* RootSignature, FD3D12UnorderedAccessViewCache& Cache, const UAVSlotMask& SlotsNeededMask, uint32 SlotsNeeded, uint32& HeapSlot);
-
-template void FD3D12DescriptorCache::SetSamplers<SF_Vertex>(const FD3D12RootSignature* RootSignature, FD3D12SamplerStateCache& Cache, const SamplerSlotMask& SlotsNeededMask, uint32 SlotsNeeded, uint32& HeapSlot);
-template void FD3D12DescriptorCache::SetSamplers<SF_Mesh>(const FD3D12RootSignature* RootSignature, FD3D12SamplerStateCache& Cache, const SamplerSlotMask& SlotsNeededMask, uint32 SlotsNeeded, uint32& HeapSlot);
-template void FD3D12DescriptorCache::SetSamplers<SF_Amplification>(const FD3D12RootSignature* RootSignature, FD3D12SamplerStateCache& Cache, const SamplerSlotMask& SlotsNeededMask, uint32 SlotsNeeded, uint32& HeapSlot);
-template void FD3D12DescriptorCache::SetSamplers<SF_Geometry>(const FD3D12RootSignature* RootSignature, FD3D12SamplerStateCache& Cache, const SamplerSlotMask& SlotsNeededMask, uint32 SlotsNeeded, uint32& HeapSlot);
-template void FD3D12DescriptorCache::SetSamplers<SF_Pixel>(const FD3D12RootSignature* RootSignature, FD3D12SamplerStateCache& Cache, const SamplerSlotMask& SlotsNeededMask, uint32 SlotsNeeded, uint32& HeapSlot);
-template void FD3D12DescriptorCache::SetSamplers<SF_Compute>(const FD3D12RootSignature* RootSignature, FD3D12SamplerStateCache& Cache, const SamplerSlotMask& SlotsNeededMask, uint32 SlotsNeeded, uint32& HeapSlot);
-
 bool FD3D12DescriptorCache::HeapRolledOver(ERHIDescriptorHeapType InHeapType)
 {
 	// A heap rolled over, so set the descriptor heaps again and return if the heaps actually changed.
@@ -179,11 +161,8 @@ void FD3D12DescriptorCache::SetVertexBuffers(FD3D12VertexBufferCache& Cache)
 	}
 }
 
-template <EShaderFrequency ShaderStage>
-void FD3D12DescriptorCache::SetUAVs(const FD3D12RootSignature* RootSignature, FD3D12UnorderedAccessViewCache& Cache, const UAVSlotMask& SlotsNeededMask, uint32 SlotsNeeded, uint32& HeapSlot)
+void FD3D12DescriptorCache::SetUAVs(EShaderFrequency ShaderStage, const FD3D12RootSignature* RootSignature, FD3D12UnorderedAccessViewCache& Cache, const UAVSlotMask& SlotsNeededMask, uint32 SlotsNeeded, uint32& HeapSlot)
 {
-	static_assert(ShaderStage < SF_NumStandardFrequencies, "Unexpected shader frequency.");
-
 	UAVSlotMask& CurrentDirtySlotMask = Cache.DirtySlotMask[ShaderStage];
 	check(CurrentDirtySlotMask != 0);	// All dirty slots for the current shader stage.
 	check(SlotsNeededMask != 0);		// All dirty slots for the current shader stage AND used by the current shader stage.
@@ -295,10 +274,8 @@ void FD3D12DescriptorCache::SetStreamOutTargets(FD3D12Resource** Buffers, uint32
 	check(Count == 0);
 }
 
-template <EShaderFrequency ShaderStage>
-void FD3D12DescriptorCache::SetSamplers(const FD3D12RootSignature* RootSignature, FD3D12SamplerStateCache& Cache, const SamplerSlotMask& SlotsNeededMask, uint32 SlotsNeeded, uint32& HeapSlot)
+void FD3D12DescriptorCache::SetSamplers(EShaderFrequency ShaderStage, const FD3D12RootSignature* RootSignature, FD3D12SamplerStateCache& Cache, const SamplerSlotMask& SlotsNeededMask, uint32 SlotsNeeded, uint32& HeapSlot)
 {
-	static_assert(ShaderStage < SF_NumStandardFrequencies, "Unexpected shader frequency.");
 	check(!UsingGlobalSamplerHeap());
 
 	SamplerSlotMask& CurrentDirtySlotMask = Cache.DirtySlotMask[ShaderStage];
@@ -419,11 +396,8 @@ static D3D12_RESOURCE_STATES GetBindingResourceState(FD3D12CommandContext& Conte
 	return State;
 }
 
-template <EShaderFrequency ShaderStage>
-void FD3D12DescriptorCache::SetSRVs(const FD3D12RootSignature* RootSignature, FD3D12ShaderResourceViewCache& Cache, const SRVSlotMask& SlotsNeededMask, uint32 SlotsNeeded, uint32& HeapSlot)
+void FD3D12DescriptorCache::SetSRVs(EShaderFrequency ShaderStage, const FD3D12RootSignature* RootSignature, FD3D12ShaderResourceViewCache& Cache, const SRVSlotMask& SlotsNeededMask, uint32 SlotsNeeded, uint32& HeapSlot)
 {
-	static_assert(ShaderStage < SF_NumStandardFrequencies, "Unexpected shader frequency.");
-
 	SRVSlotMask& CurrentDirtySlotMask = Cache.DirtySlotMask[ShaderStage];
 	check(CurrentDirtySlotMask != 0);	// All dirty slots for the current shader stage.
 	check(SlotsNeededMask != 0);		// All dirty slots for the current shader stage AND used by the current shader stage.
diff --git a/Engine/Source/Runtime/D3D12RHI/Private/D3D12DescriptorCache.h b/Engine/Source/Runtime/D3D12RHI/Private/D3D12DescriptorCache.h
--- a/Engine/Source/Runtime/D3D12RHI/Private/D3D12DescriptorCache.h
+++ b/Engine/Source/Runtime/D3D12RHI/Private/D3D12DescriptorCache.h
@@ -339,17 +339,11 @@ public:
 	void SetVertexBuffers(FD3D12VertexBufferCache& Cache);
 	void SetRenderTargets(FD3D12RenderTargetView** RenderTargetViewArray, uint32 Count, FD3D12DepthStencilView* DepthStencilTarget);
 
-	template <EShaderFrequency ShaderStage>
-	void SetUAVs(const FD3D12RootSignature* RootSignature, FD3D12UnorderedAccessViewCache& Cache, const UAVSlotMask& SlotsNeededMask, uint32 Count, uint32 &HeapSlot);
-
-	template <EShaderFrequency ShaderStage>
-	void SetSamplers(const FD3D12RootSignature* RootSignature, FD3D12SamplerStateCache& Cache, const SamplerSlotMask& SlotsNeededMask, uint32 Count, uint32& HeapSlot);
-
-	template <EShaderFrequency ShaderStage>
-	void SetSRVs(const FD3D12RootSignature* RootSignature, FD3D12ShaderResourceViewCache& Cache, const SRVSlotMask& SlotsNeededMask, uint32 Count, uint32& HeapSlot);
+	void SetUAVs(EShaderFrequency ShaderStage, const FD3D12RootSignature* RootSignature, FD3D12UnorderedAccessViewCache& Cache, const UAVSlotMask& SlotsNeededMask, uint32 Count, uint32 &HeapSlot);
+	void SetSamplers(EShaderFrequency ShaderStage, const FD3D12RootSignature* RootSignature, FD3D12SamplerStateCache& Cache, const SamplerSlotMask& SlotsNeededMask, uint32 Count, uint32& HeapSlot);
+	void SetSRVs(EShaderFrequency ShaderStage, const FD3D12RootSignature* RootSignature, FD3D12ShaderResourceViewCache& Cache, const SRVSlotMask& SlotsNeededMask, uint32 Count, uint32& HeapSlot);
 
 	void SetConstantBufferViews(EShaderFrequency ShaderStage, const FD3D12RootSignature* RootSignature, FD3D12ConstantBufferCache& Cache, CBVSlotMask SlotsNeededMask, uint32 Count, uint32& HeapSlot);
-
 	void SetRootConstantBuffers(EShaderFrequency ShaderStage, const FD3D12RootSignature* RootSignature, FD3D12ConstantBufferCache& Cache, CBVSlotMask SlotsNeededMask);
 
 	void PrepareBindlessViews(EShaderFrequency ShaderStage, TConstArrayView<FD3D12ShaderResourceView*> SRVs, TConstArrayView<FD3D12UnorderedAccessView*> UAVs);
diff --git a/Engine/Source/Runtime/D3D12RHI/Private/D3D12StateCache.cpp b/Engine/Source/Runtime/D3D12RHI/Private/D3D12StateCache.cpp
--- a/Engine/Source/Runtime/D3D12RHI/Private/D3D12StateCache.cpp
+++ b/Engine/Source/Runtime/D3D12RHI/Private/D3D12StateCache.cpp
@@ -83,11 +83,6 @@ inline bool operator!=(D3D12_CPU_DESCRIPTOR_HANDLE lhs, D3D12_CPU_DESCRIPTOR_HAN
 	return lhs.ptr != rhs.ptr;
 }
 
-// Define template functions that are only declared in the header.
-
-template void FD3D12StateCache::ApplyState<ED3D12PipelineType::Graphics>();
-template void FD3D12StateCache::ApplyState<ED3D12PipelineType::Compute>();
-
 #if D3D12_STATE_CACHE_RUNTIME_TOGGLE
 
 // Default the state caching system to on.
@@ -394,8 +389,7 @@ void FD3D12StateCache::InternalSetPipelineState(FD3D12PipelineState* InPipelineS
 	}
 }
 
-template <ED3D12PipelineType PipelineType>
-void FD3D12StateCache::ApplyState()
+void FD3D12StateCache::ApplyState(ED3D12PipelineType PipelineType)
 {
 	//SCOPE_CYCLE_COUNTER(STAT_D3D12ApplyStateTime);
 	const bool bForceState = false;
@@ -682,14 +676,7 @@ void FD3D12StateCache::ApplyResources(const FD3D12RootSignature* const pRootSign
 	if (CurrentShaderDirtyUAVSlots)
 	{
 		SCOPE_CYCLE_COUNTER(STAT_D3D12ApplyStateSetUAVTime);
-		if (UAVStage == SF_Pixel)
-		{
-			DescriptorCache.SetUAVs<SF_Pixel>(pRootSignature, PipelineState.Common.UAVCache, CurrentShaderDirtyUAVSlots, NumUAVs, ViewHeapSlot);
-		}
-		else
-		{
-			DescriptorCache.SetUAVs<SF_Compute>(pRootSignature, PipelineState.Common.UAVCache, CurrentShaderDirtyUAVSlots, NumUAVs, ViewHeapSlot);
-		}
+		DescriptorCache.SetUAVs(UAVStage, pRootSignature, PipelineState.Common.UAVCache, CurrentShaderDirtyUAVSlots, NumUAVs, ViewHeapSlot);
 	}
 
 	// Shader resource views
@@ -697,26 +684,13 @@ void FD3D12StateCache::ApplyResources(const FD3D12RootSignature* const pRootSign
 		//SCOPE_CYCLE_COUNTER(STAT_D3D12ApplyStateSetSRVTime);
 		FD3D12ShaderResourceViewCache& SRVCache = PipelineState.Common.SRVCache;
 
-#define CONDITIONAL_SET_SRVS(Shader) \
-		if (CurrentShaderDirtySRVSlots[Shader]) \
-		{ \
-			DescriptorCache.SetSRVs<Shader>(pRootSignature, SRVCache, CurrentShaderDirtySRVSlots[Shader], NumSRVs[Shader], ViewHeapSlot); \
-		}
-
-		if (StartStage == SF_Compute)
+		for (uint32 Index = StartStage; Index < EndStage; Index++)
 		{
-			// Note that ray tracing pipeline shares state with compute
-			CONDITIONAL_SET_SRVS(SF_Compute);
-		}
-		else
-		{
-			CONDITIONAL_SET_SRVS(SF_Vertex);
-			CONDITIONAL_SET_SRVS(SF_Mesh);
-			CONDITIONAL_SET_SRVS(SF_Amplification);
-			CONDITIONAL_SET_SRVS(SF_Geometry);
-			CONDITIONAL_SET_SRVS(SF_Pixel);
+			if (CurrentShaderDirtySRVSlots[Index])
+			{
+				DescriptorCache.SetSRVs(static_cast<EShaderFrequency>(Index), pRootSignature, SRVCache, CurrentShaderDirtySRVSlots[Index], NumSRVs[Index], ViewHeapSlot);
+			}
 		}
-#undef CONDITIONAL_SET_SRVS
 	}
 
 #if USE_STATIC_ROOT_SIGNATURE
@@ -725,26 +699,13 @@ void FD3D12StateCache::ApplyResources(const FD3D12RootSignature* const pRootSign
 		//SCOPE_CYCLE_COUNTER(STAT_D3D12ApplyStateSetConstantBufferTime);
 		FD3D12ConstantBufferCache& CBVCache = PipelineState.Common.CBVCache;
 
-	#define CONDITIONAL_SET_CBVS(Shader) \
-		if (CurrentShaderDirtyCBVSlots[Shader]) \
-		{ \
-			DescriptorCache.SetConstantBufferViews(Shader, pRootSignature, CBVCache, CurrentShaderDirtyCBVSlots[Shader], NumCBVs[Shader], ViewHeapSlot); \
-		}
-
-		if (StartStage == SF_Compute)
+		for (uint32 Index = StartStage; Index < EndStage; Index++)
 		{
-			// Note that ray tracing pipeline shares state with compute
-			CONDITIONAL_SET_CBVS(SF_Compute);
-		}
-		else
-		{
-			CONDITIONAL_SET_CBVS(SF_Vertex);
-			CONDITIONAL_SET_CBVS(SF_Mesh);
-			CONDITIONAL_SET_CBVS(SF_Amplification);
-			CONDITIONAL_SET_CBVS(SF_Geometry);
-			CONDITIONAL_SET_CBVS(SF_Pixel);
+			if (CurrentShaderDirtyCBVSlots[Index])
+			{
+				DescriptorCache.SetConstantBufferViews(static_cast<EShaderFrequency>(Index), pRootSignature, CBVCache, CurrentShaderDirtyCBVSlots[Index], NumCBVs[Index], ViewHeapSlot);
+			}
 		}
-#undef CONDITIONAL_SET_CBVS
 	}
 #endif // USE_STATIC_ROOT_SIGNATURE
 }
@@ -787,26 +748,13 @@ void FD3D12StateCache::ApplyConstants(const FD3D12RootSignature* const pRootSign
 		//SCOPE_CYCLE_COUNTER(STAT_D3D12ApplyStateSetConstantBufferTime);
 		FD3D12ConstantBufferCache& CBVCache = PipelineState.Common.CBVCache;
 
-#define CONDITIONAL_SET_CBVS(Shader) \
-		if (CurrentShaderDirtyCBVSlots[##Shader]) \
-		{ \
-			DescriptorCache.SetRootConstantBuffers(##Shader, pRootSignature, CBVCache, CurrentShaderDirtyCBVSlots[##Shader]); \
-		}
-
-		if (StartStage == SF_Compute)
+		for (uint32 Index = StartStage; Index < EndStage; Index++)
 		{
-			// Note that ray tracing pipeline shares state with compute
-			CONDITIONAL_SET_CBVS(SF_Compute);
-		}
-		else
-		{
-			CONDITIONAL_SET_CBVS(SF_Vertex);
-			CONDITIONAL_SET_CBVS(SF_Mesh);
-			CONDITIONAL_SET_CBVS(SF_Amplification);
-			CONDITIONAL_SET_CBVS(SF_Geometry);
-			CONDITIONAL_SET_CBVS(SF_Pixel);
+			if (CurrentShaderDirtyCBVSlots[Index])
+			{
+				DescriptorCache.SetRootConstantBuffers(static_cast<EShaderFrequency>(Index), pRootSignature, CBVCache, CurrentShaderDirtyCBVSlots[Index]);
+			}
 		}
-#undef CONDITIONAL_SET_CBVS
 	}
 #endif // !USE_STATIC_ROOT_SIGNATURE
 }
@@ -946,25 +894,13 @@ void FD3D12StateCache::ApplySamplers(const FD3D12RootSignature* const pRootSigna
 	}
 	uint32 SamplerHeapSlot = SamplerHeap->ReserveSlots(NumSamplers[SF_NumStandardFrequencies]);
 
-#define CONDITIONAL_SET_SAMPLERS(Shader) \
-	if (CurrentShaderDirtySamplerSlots[Shader]) \
-	{ \
-		DescriptorCache.SetSamplers<Shader>(pRootSignature, Cache, CurrentShaderDirtySamplerSlots[Shader], NumSamplers[Shader], SamplerHeapSlot); \
-	}
-
-	if (StartStage == SF_Compute)
-	{
-		CONDITIONAL_SET_SAMPLERS(SF_Compute);
-	}
-	else
+	for (uint32 Index = StartStage; Index < EndStage; Index++)
 	{
-		CONDITIONAL_SET_SAMPLERS(SF_Vertex);
-		CONDITIONAL_SET_SAMPLERS(SF_Mesh);
-		CONDITIONAL_SET_SAMPLERS(SF_Amplification);
-		CONDITIONAL_SET_SAMPLERS(SF_Geometry);
-		CONDITIONAL_SET_SAMPLERS(SF_Pixel);
+		if (CurrentShaderDirtySamplerSlots[Index])
+		{
+			DescriptorCache.SetSamplers(static_cast<EShaderFrequency>(Index), pRootSignature, Cache, CurrentShaderDirtySamplerSlots[Index], NumSamplers[Index], SamplerHeapSlot);
+		}
 	}
-#undef CONDITIONAL_SET_SAMPLERS
 
 	SamplerHeap->SetNextSlot(SamplerHeapSlot);
 }
diff --git a/Engine/Source/Runtime/D3D12RHI/Private/D3D12StateCachePrivate.h b/Engine/Source/Runtime/D3D12RHI/Private/D3D12StateCachePrivate.h
--- a/Engine/Source/Runtime/D3D12RHI/Private/D3D12StateCachePrivate.h
+++ b/Engine/Source/Runtime/D3D12RHI/Private/D3D12StateCachePrivate.h
@@ -658,8 +658,7 @@ public:
 	ED3D12PipelineType LastComputePipelineType = ED3D12PipelineType::Compute;
 #endif // D3D12_RHI_RAYTRACING
 
-	template <ED3D12PipelineType PipelineType> 
-	void ApplyState();
+	void ApplyState(ED3D12PipelineType PipelineType);
 	void ApplySamplers(const FD3D12RootSignature* const pRootSignature, uint32 StartStage, uint32 EndStage);
 	void ApplyResources(const FD3D12RootSignature* const pRootSignature, uint32 StartStage, uint32 EndStage);
 	void ApplyBindlessResources(const FD3D12RootSignature* const pRootSignature, uint32 StartStage, uint32 EndStage);
