D3D12 binding cleanup
- Renaming and reusing FUniformResourceBinder to bind parameters in RHISetShaderParameters
- Fixing an issue where SetResourcesFromTables wouldn't set Compute UAVs
- Replacing FD3D12StateCache::SetUAVs with a single UAV setter
- Fixing static analysis issue in FD3D12StateCache::ApplyState
- Adding shortcut calls to ConditionalClearShaderResource

diff --git a/Engine/Source/Runtime/D3D12RHI/Private/D3D12CommandContext.cpp b/Engine/Source/Runtime/D3D12RHI/Private/D3D12CommandContext.cpp
--- a/Engine/Source/Runtime/D3D12RHI/Private/D3D12CommandContext.cpp
+++ b/Engine/Source/Runtime/D3D12RHI/Private/D3D12CommandContext.cpp
@@ -639,6 +639,22 @@ void FD3D12CommandContext::ConditionalClearShaderResource(FD3D12ResourceLocation
 	}
 }
 
+void FD3D12CommandContext::ClearShaderResources(FD3D12UnorderedAccessView* UAV)
+{
+	if (UAV)
+	{
+		ConditionalClearShaderResource(UAV->GetResourceLocation());
+	}
+}
+
+void FD3D12CommandContext::ClearShaderResources(FD3D12BaseShaderResource* Resource)
+{
+	if (Resource)
+	{
+		ConditionalClearShaderResource(&Resource->ResourceLocation);
+	}
+}
+
 void FD3D12CommandContext::ClearAllShaderResources()
 {
 	StateCache.ClearSRVs();
diff --git a/Engine/Source/Runtime/D3D12RHI/Private/D3D12CommandContext.h b/Engine/Source/Runtime/D3D12RHI/Private/D3D12CommandContext.h
--- a/Engine/Source/Runtime/D3D12RHI/Private/D3D12CommandContext.h
+++ b/Engine/Source/Runtime/D3D12RHI/Private/D3D12CommandContext.h
@@ -439,6 +439,8 @@ public:
 
 	virtual void ClearState(EClearStateMode ClearStateMode = EClearStateMode::All) override final;
 	void ConditionalClearShaderResource(FD3D12ResourceLocation* Resource);
+	void ClearShaderResources(FD3D12UnorderedAccessView* UAV);
+	void ClearShaderResources(FD3D12BaseShaderResource* Resource);
 	void ClearAllShaderResources();
 
 	FD3D12FastConstantAllocator ConstantsAllocator;
diff --git a/Engine/Source/Runtime/D3D12RHI/Private/D3D12Commands.cpp b/Engine/Source/Runtime/D3D12RHI/Private/D3D12Commands.cpp
--- a/Engine/Source/Runtime/D3D12RHI/Private/D3D12Commands.cpp
+++ b/Engine/Source/Runtime/D3D12RHI/Private/D3D12Commands.cpp
@@ -537,7 +537,7 @@ void FD3D12CommandContext::HandleTransientAliasing(const FD3D12TransitionData* T
 			TransitionResource(Resource, D3D12_RESOURCE_STATE_TBD, FinalState, D3D12_RESOURCE_BARRIER_ALL_SUBRESOURCES);
 			
 			// Remove from caches
-			ConditionalClearShaderResource(&BaseShaderResource->ResourceLocation);
+			ClearShaderResources(BaseShaderResource);
 		}
 	}
 }
@@ -914,34 +914,17 @@ void FD3D12CommandContext::RHISetShaderTexture(FRHIComputeShader* ComputeShaderR
 void FD3D12CommandContext::RHISetUAVParameter(FRHIPixelShader* PixelShaderRHI, uint32 UAVIndex, FRHIUnorderedAccessView* UAVRHI)
 {
 	FD3D12UnorderedAccessView* UAV = RetrieveObject<FD3D12UnorderedAccessView>(UAVRHI);
-
-	if (UAV)
-	{
-		ConditionalClearShaderResource(UAV->GetResourceLocation());
-	}
-
-	uint32 InitialCount = -1;
-
-	// Actually set the UAV
-	StateCache.SetUAVs(SF_Pixel, UAVIndex, 1, &UAV, &InitialCount);
+	ClearShaderResources(UAV);
+	StateCache.SetUAV(SF_Pixel, UAVIndex, UAV);
 }
 
-
 void FD3D12CommandContext::RHISetUAVParameter(FRHIComputeShader* ComputeShaderRHI, uint32 UAVIndex, FRHIUnorderedAccessView* UAVRHI)
 {
 	//ValidateBoundShader(StateCache, ComputeShaderRHI);
 
 	FD3D12UnorderedAccessView* UAV = RetrieveObject<FD3D12UnorderedAccessView>(UAVRHI);
-
-	if (UAV)
-	{
-		ConditionalClearShaderResource(UAV->GetResourceLocation());
-	}
-
-	uint32 InitialCount = -1;
-
-	// Actually set the UAV
-	StateCache.SetUAVs(SF_Compute, UAVIndex, 1, &UAV, &InitialCount);
+	ClearShaderResources(UAV);
+	StateCache.SetUAV(SF_Compute, UAVIndex, UAV);
 }
 
 void FD3D12CommandContext::RHISetUAVParameter(FRHIComputeShader* ComputeShaderRHI, uint32 UAVIndex, FRHIUnorderedAccessView* UAVRHI, uint32 InitialCount)
@@ -949,13 +932,8 @@ void FD3D12CommandContext::RHISetUAVParameter(FRHIComputeShader* ComputeShaderRH
 	//ValidateBoundShader(StateCache, ComputeShaderRHI);
 
 	FD3D12UnorderedAccessView* UAV = RetrieveObject<FD3D12UnorderedAccessView>(UAVRHI);
-
-	if (UAV)
-	{
-		ConditionalClearShaderResource(UAV->GetResourceLocation());
-	}
-
-	StateCache.SetUAVs(SF_Compute, UAVIndex, 1, &UAV, &InitialCount);
+	ClearShaderResources(UAV);
+	StateCache.SetUAV(SF_Compute, UAVIndex, UAV, InitialCount);
 }
 
 void FD3D12CommandContext::RHISetShaderResourceViewParameter(FRHIGraphicsShader* ShaderRHI, uint32 TextureIndex, FRHIShaderResourceView* SRVRHI)
@@ -965,10 +943,7 @@ void FD3D12CommandContext::RHISetShaderResourceViewParameter(FRHIGraphicsShader*
 	if (IsValidGraphicsFrequency(ShaderFrequency))
 	{
 		ValidateBoundShader(StateCache, ShaderRHI);
-
-		FD3D12ShaderResourceView* SRV = RetrieveObject<FD3D12ShaderResourceView>(SRVRHI);
-
-		StateCache.SetShaderResourceView(ShaderFrequency, SRV, TextureIndex);
+		StateCache.SetShaderResourceView(ShaderFrequency, RetrieveObject<FD3D12ShaderResourceView>(SRVRHI), TextureIndex);
 	}
 	else
 	{
@@ -979,8 +954,7 @@ void FD3D12CommandContext::RHISetShaderResourceViewParameter(FRHIGraphicsShader*
 void FD3D12CommandContext::RHISetShaderResourceViewParameter(FRHIComputeShader* ComputeShaderRHI, uint32 TextureIndex, FRHIShaderResourceView* SRVRHI)
 {
 	//ValidateBoundShader(StateCache, ComputeShaderRHI);
-	FD3D12ShaderResourceView* const SRV = RetrieveObject<FD3D12ShaderResourceView>(SRVRHI);
-	StateCache.SetShaderResourceView(SF_Compute, SRV, TextureIndex);
+	StateCache.SetShaderResourceView(SF_Compute, RetrieveObject<FD3D12ShaderResourceView>(SRVRHI), TextureIndex);
 }
 
 void FD3D12CommandContext::RHISetShaderSampler(FRHIGraphicsShader* ShaderRHI, uint32 SamplerIndex, FRHISamplerState* NewStateRHI)
@@ -1032,6 +1006,45 @@ void FD3D12CommandContext::RHISetShaderUniformBuffer(FRHIComputeShader* ComputeS
 	BindUniformBuffer(*this, ComputeShader, SF_Compute, BufferIndex, RetrieveObject<FD3D12UniformBuffer>(BufferRHI));
 }
 
+struct FD3D12ResourceBinder
+{
+	FD3D12CommandContext& Context;
+	const uint32 GpuIndex;
+	const EShaderFrequency Frequency;
+
+	FD3D12ResourceBinder(FD3D12CommandContext& InContext, EShaderFrequency InFrequency)
+		: Context(InContext)
+		, GpuIndex(InContext.GetGPUIndex())
+		, Frequency(InFrequency)
+	{
+	}
+
+	void SetUAV(FRHIUnorderedAccessView* UAV, uint8 Index, bool bClearResources = false)
+	{
+		FD3D12UnorderedAccessView* D3D12UAV = FD3D12CommandContext::RetrieveObject<FD3D12UnorderedAccessView>(UAV, GpuIndex);
+		if (bClearResources)
+		{
+			Context.ClearShaderResources(D3D12UAV);
+		}
+		Context.StateCache.SetUAV(Frequency, Index, D3D12UAV);
+	}
+
+	void SetSRV(FRHIShaderResourceView* SRV, uint8 Index)
+	{
+		Context.StateCache.SetShaderResourceView(Frequency, FD3D12CommandContext::RetrieveObject<FD3D12ShaderResourceView>(SRV, GpuIndex), Index);
+	}
+
+	void SetTexture(FRHITexture* TextureRHI, uint8 Index)
+	{
+		Context.StateCache.SetShaderResourceView(Frequency, GetShaderResourceViewFromTexture(TextureRHI, GpuIndex), Index);
+	}
+
+	void SetSampler(FRHISamplerState* Sampler, uint8 Index)
+	{
+		Context.StateCache.SetSamplerState(Frequency, FD3D12CommandContext::RetrieveObject<FD3D12SamplerState>(Sampler, GpuIndex), Index);
+	}
+};
+
 static void SetShaderParametersOnContext(
 	FD3D12CommandContext& Context
 	, FRHIShader* Shader
@@ -1076,10 +1089,7 @@ static void SetShaderParametersOnContext(
 			{
 				Handle = static_cast<FRHIUnorderedAccessView*>(Resource)->GetBindlessHandle();
 				FD3D12UnorderedAccessView* UAV = FD3D12CommandContext::RetrieveObject<FD3D12UnorderedAccessView>(Parameter.Resource, GpuIndex);
-				if (UAV)
-				{
-					Context.ConditionalClearShaderResource(UAV->GetResourceLocation());
-				}
+				Context.ClearShaderResources(UAV);
 				StateCache.QueueBindlessUAV(ShaderFrequency, UAV);
 			}
 			break;
@@ -1099,18 +1109,15 @@ static void SetShaderParametersOnContext(
 	}
 #endif
 
+	FD3D12ResourceBinder Binder(Context, ShaderFrequency);
+
 	for (const FRHIShaderParameterResource& Parameter : InResourceParameters)
 	{
 		if (Parameter.Type == FRHIShaderParameterResource::EType::UnorderedAccessView)
 		{
 			if (ShaderFrequency == SF_Pixel || ShaderFrequency == SF_Compute)
 			{
-				FD3D12UnorderedAccessView* UAV = FD3D12CommandContext::RetrieveObject<FD3D12UnorderedAccessView>(Parameter.Resource, GpuIndex);
-				if (UAV)
-				{
-					Context.ConditionalClearShaderResource(UAV->GetResourceLocation());
-				}
-				StateCache.SetUAV(ShaderFrequency, Parameter.Index, UAV);
+				Binder.SetUAV(static_cast<FRHIUnorderedAccessView*>(Parameter.Resource), Parameter.Index, true);
 			}
 			else
 			{
@@ -1124,15 +1131,15 @@ static void SetShaderParametersOnContext(
 		switch (Parameter.Type)
 		{
 		case FRHIShaderParameterResource::EType::Texture:
-			StateCache.SetShaderResourceView(ShaderFrequency, GetShaderResourceViewFromTexture(Parameter.Resource, GpuIndex), Parameter.Index);
+			Binder.SetTexture(static_cast<FRHITexture*>(Parameter.Resource), Parameter.Index);
 			break;
 		case FRHIShaderParameterResource::EType::ResourceView:
-			StateCache.SetShaderResourceView(ShaderFrequency, FD3D12CommandContext::RetrieveObject<FD3D12ShaderResourceView>(Parameter.Resource, GpuIndex), Parameter.Index);
+			Binder.SetSRV(static_cast<FRHIShaderResourceView*>(Parameter.Resource), Parameter.Index);
 			break;
 		case FRHIShaderParameterResource::EType::UnorderedAccessView:
 			break;
 		case FRHIShaderParameterResource::EType::Sampler:
-			StateCache.SetSamplerState(ShaderFrequency, FD3D12CommandContext::RetrieveObject<FD3D12SamplerState>(Parameter.Resource, GpuIndex), Parameter.Index);
+			Binder.SetSampler(static_cast<FRHISamplerState*>(Parameter.Resource), Parameter.Index);
 			break;
 		case FRHIShaderParameterResource::EType::UniformBuffer:
 			BindUniformBuffer(Context, Shader, ShaderFrequency, Parameter.Index, FD3D12CommandContext::RetrieveObject<FD3D12UniformBuffer>(Parameter.Resource, GpuIndex));
@@ -1291,7 +1298,7 @@ void FD3D12CommandContext::SetRenderTargets(
 		DepthStencilView = NewDepthStencilTarget->GetDepthStencilView(NewDepthStencilTargetRHI->GetDepthStencilAccess());
 
 		// Unbind any shader views of the depth stencil target that are bound.
-		ConditionalClearShaderResource(&NewDepthStencilTarget->ResourceLocation);
+		ClearShaderResources(NewDepthStencilTarget);
 	}
 
 	// Gather the render target views for the new render targets.
@@ -1309,7 +1316,7 @@ void FD3D12CommandContext::SetRenderTargets(
 			ensureMsgf(RenderTargetView, TEXT("Texture being set as render target has no RTV"));
 
 			// Unbind any shader views of the render target that are bound.
-			ConditionalClearShaderResource(&NewRenderTarget->ResourceLocation);
+			ClearShaderResources(NewRenderTarget);
 		}
 
 		NewRenderTargetViews[RenderTargetIndex] = RenderTargetView;
@@ -1482,44 +1489,15 @@ void FD3D12CommandContext::CommitComputeShaderConstants()
 }
 
 template <class ShaderType>
-void FD3D12CommandContext::SetResourcesFromTables(const ShaderType* RESTRICT Shader)
+void FD3D12CommandContext::SetResourcesFromTables(const ShaderType* Shader)
 {
 	checkSlow(Shader);
 
 	static constexpr EShaderFrequency Frequency = static_cast<EShaderFrequency>(ShaderType::StaticFrequency);
 
-	struct FUniformResourceBinder
-	{
-		FD3D12CommandContext& Context;
-
-		void SetUAV(FRHIUnorderedAccessView* UAV, uint8 Index)
-		{
-			// Actually set the UAV
-			uint32 InitialCount = -1;
-			FD3D12UnorderedAccessView* D3D12UAV = Context.RetrieveObject<FD3D12UnorderedAccessView>(UAV);
-
-			Context.StateCache.SetUAVs(SF_Pixel, Index, 1, &D3D12UAV, &InitialCount);
-		}
-
-		void SetSRV(FRHIShaderResourceView* SRV, uint8 Index)
-		{
-			Context.StateCache.SetShaderResourceView(Frequency, Context.RetrieveObject<FD3D12ShaderResourceView>(SRV), Index);
-		}
-
-		void SetTexture(FRHITexture* TextureRHI, uint8 Index)
-		{
-			FD3D12ShaderResourceView* SRV = Context.RetrieveTexture(TextureRHI)->GetShaderResourceView();
-			Context.StateCache.SetShaderResourceView(Frequency, SRV, Index);
-		}
-
-		void SetSampler(FRHISamplerState* Sampler, uint8 Index)
-		{
-			Context.StateCache.SetSamplerState(Frequency, Context.RetrieveObject<FD3D12SamplerState>(Sampler), Index);
-		}
-	};
-
+	FD3D12ResourceBinder Binder(*this, Frequency);
 	UE::RHICore::SetResourcesFromTables(
-		  FUniformResourceBinder { *this }
+		  Binder
 		, *Shader
 		, Shader->ShaderResourceTable
 		, DirtyUniformBuffers[Frequency]
@@ -1537,10 +1515,9 @@ void FD3D12CommandContext::CommitGraphicsResourceTables()
 	const FD3D12GraphicsPipelineState* const RESTRICT GraphicPSO = StateCache.GetGraphicsPipelineState();
 	check(GraphicPSO);
 
-	FD3D12PixelShader* PixelShader = GraphicPSO->GetPixelShader();
-	if (PixelShader)
+	if (FD3D12PixelShader* Shader = GraphicPSO->GetPixelShader())
 	{
-		SetResourcesFromTables(PixelShader);
+		SetResourcesFromTables(Shader);
 	}
 
 	if (FD3D12VertexShader* Shader = GraphicPSO->GetVertexShader())
diff --git a/Engine/Source/Runtime/D3D12RHI/Private/D3D12DescriptorCache.cpp b/Engine/Source/Runtime/D3D12RHI/Private/D3D12DescriptorCache.cpp
--- a/Engine/Source/Runtime/D3D12RHI/Private/D3D12DescriptorCache.cpp
+++ b/Engine/Source/Runtime/D3D12RHI/Private/D3D12DescriptorCache.cpp
@@ -695,6 +695,7 @@ void FD3D12DescriptorCache::RestoreAfterExternalHeapsSet()
 	if (bHeapsOverridden)
 	{
 		SetDescriptorHeaps();
+		bHeapsOverridden = false;
 	}
 }
 
diff --git a/Engine/Source/Runtime/D3D12RHI/Private/D3D12RenderTarget.cpp b/Engine/Source/Runtime/D3D12RHI/Private/D3D12RenderTarget.cpp
--- a/Engine/Source/Runtime/D3D12RHI/Private/D3D12RenderTarget.cpp
+++ b/Engine/Source/Runtime/D3D12RHI/Private/D3D12RenderTarget.cpp
@@ -161,10 +161,7 @@ void FD3D12CommandContext::ResolveTextureUsingShader(
 	GraphicsPSOInit.RasterizerState = TStaticRasterizerState<FM_Solid, CM_None>::GetRHI();
 
 	// Make sure the destination is not bound as a shader resource.
-	if (DestTexture)
-	{
-		ConditionalClearShaderResource(&DestTexture->ResourceLocation);
-	}
+	ClearShaderResources(DestTexture);
 
 	// Determine if the entire destination surface is being resolved to.
 	// If the entire surface is being resolved to, then it means we can clear it and signal the driver that it can discard
@@ -275,7 +272,7 @@ void FD3D12CommandContext::ResolveTextureUsingShader(
 		})(RHICmdList);
 	}
 
-	ConditionalClearShaderResource(&SourceTexture->ResourceLocation);
+	ClearShaderResources(SourceTexture);
 
 	// Reset saved viewport
 	{
diff --git a/Engine/Source/Runtime/D3D12RHI/Private/D3D12StateCache.cpp b/Engine/Source/Runtime/D3D12RHI/Private/D3D12StateCache.cpp
--- a/Engine/Source/Runtime/D3D12RHI/Private/D3D12StateCache.cpp
+++ b/Engine/Source/Runtime/D3D12RHI/Private/D3D12StateCache.cpp
@@ -414,6 +414,11 @@ void FD3D12StateCache::ApplyState(ED3D12PipelineType PipelineType)
 	{
 		PSOCommonData = GetGraphicsPipelineState();
 	}
+	else
+	{
+		checkf(false, TEXT("Unexpected pipeline type: %d"), (uint32)PipelineType);
+		return;
+	}
 
 	const bool bRootSignatureChanged = InternalSetRootSignature(PipelineType, PSOCommonData->RootSignature);
 
@@ -1135,63 +1140,57 @@ void FD3D12StateCache::ClearUAVs(EShaderFrequency ShaderStage)
 	}
 }
 
-void FD3D12StateCache::SetUAVs(EShaderFrequency ShaderStage, uint32 UAVStartSlot, uint32 NumSimultaneousUAVs, FD3D12UnorderedAccessView** UAVArray, uint32* UAVInitialCountArray)
+void FD3D12StateCache::SetUAV(EShaderFrequency ShaderStage, uint32 SlotIndex, FD3D12UnorderedAccessView* UAV, uint32 InitialCount)
 {
 	SCOPE_CYCLE_COUNTER(STAT_D3D12SetUnorderedAccessViewTime);
-	check(NumSimultaneousUAVs > 0);
 
 	FD3D12UnorderedAccessViewCache& Cache = PipelineState.Common.UAVCache;
 
 	// When setting UAV's for Graphics, it wipes out all existing bound resources.
 	const bool bIsCompute = ShaderStage == SF_Compute;
-	Cache.StartSlot[ShaderStage] = bIsCompute ? FMath::Min(UAVStartSlot, Cache.StartSlot[ShaderStage]) : 0;
+	Cache.StartSlot[ShaderStage] = bIsCompute ? FMath::Min(SlotIndex, Cache.StartSlot[ShaderStage]) : 0;
 
-	for (uint32 i = 0; i < NumSimultaneousUAVs; ++i)
-	{
-		FD3D12UnorderedAccessView* UAV = UAVArray[i];
+	Cache.Views[ShaderStage][SlotIndex] = UAV;
+	FD3D12UnorderedAccessViewCache::DirtySlot(Cache.DirtySlotMask[ShaderStage], SlotIndex);
 
-		Cache.Views[ShaderStage][UAVStartSlot + i] = UAV;
-		FD3D12UnorderedAccessViewCache::DirtySlot(Cache.DirtySlotMask[ShaderStage], UAVStartSlot + i);
+	if (UAV)
+	{
+		Cache.ResidencyHandles[ShaderStage][SlotIndex] = &UAV->GetResidencyHandle();
 
-		if (UAV)
+		FD3D12Resource* CounterResource = UAV->GetCounterResource();
+		if (CounterResource && (!UAV->IsCounterResourceInitialized() || InitialCount != -1))
 		{
-			Cache.ResidencyHandles[ShaderStage][UAVStartSlot + i] = &UAV->GetResidencyHandle();
-
-			FD3D12Resource* CounterResource = UAV->GetCounterResource();
-			if (CounterResource && (!UAV->IsCounterResourceInitialized() || UAVInitialCountArray[i] != -1))
-			{
-				FD3D12Device* Device = CounterResource->GetParentDevice();
-				FD3D12ResourceLocation UploadBufferLocation(Device);
+			FD3D12Device* Device = CounterResource->GetParentDevice();
+			FD3D12ResourceLocation UploadBufferLocation(Device);
 
-				uint32* CounterUploadHeapData = static_cast<uint32*>(CmdContext.ConstantsAllocator.Allocate(sizeof(uint32), UploadBufferLocation, nullptr));
+			uint32* CounterUploadHeapData = static_cast<uint32*>(CmdContext.ConstantsAllocator.Allocate(sizeof(uint32), UploadBufferLocation, nullptr));
 
-				// Initialize the counter to 0 if it's not been previously initialized and the UAVInitialCount is -1, if not use the value that was passed.
-				*CounterUploadHeapData = (!UAV->IsCounterResourceInitialized() && UAVInitialCountArray[i] == -1) ? 0 : UAVInitialCountArray[i];
+			// Initialize the counter to 0 if it's not been previously initialized and the UAVInitialCount is -1, if not use the value that was passed.
+			*CounterUploadHeapData = (!UAV->IsCounterResourceInitialized() && InitialCount == -1) ? 0 : InitialCount;
 
-				// Transition to copy dest
-				CmdContext.TransitionResource(CounterResource, D3D12_RESOURCE_STATE_UNORDERED_ACCESS, D3D12_RESOURCE_STATE_COPY_DEST, 0);
-				CmdContext.FlushResourceBarriers();
+			// Transition to copy dest
+			CmdContext.TransitionResource(CounterResource, D3D12_RESOURCE_STATE_UNORDERED_ACCESS, D3D12_RESOURCE_STATE_COPY_DEST, 0);
+			CmdContext.FlushResourceBarriers();
 
-				CmdContext.GraphicsCommandList()->CopyBufferRegion(
-					CounterResource->GetResource(),
-					0,
-					UploadBufferLocation.GetResource()->GetResource(),
-					UploadBufferLocation.GetOffsetFromBaseOfResource(),
-					4);
+			CmdContext.GraphicsCommandList()->CopyBufferRegion(
+				CounterResource->GetResource(),
+				0,
+				UploadBufferLocation.GetResource()->GetResource(),
+				UploadBufferLocation.GetOffsetFromBaseOfResource(),
+				4);
 
-				// Restore UAV state
-				CmdContext.TransitionResource(CounterResource, D3D12_RESOURCE_STATE_COPY_DEST, D3D12_RESOURCE_STATE_UNORDERED_ACCESS, 0);
+			// Restore UAV state
+			CmdContext.TransitionResource(CounterResource, D3D12_RESOURCE_STATE_COPY_DEST, D3D12_RESOURCE_STATE_UNORDERED_ACCESS, 0);
 
-				CmdContext.UpdateResidency(CounterResource);
+			CmdContext.UpdateResidency(CounterResource);
 
-				UAV->MarkCounterResourceInitialized();
-			}
-		}
-		else
-		{
-			Cache.ResidencyHandles[ShaderStage][UAVStartSlot + i] = nullptr;
+			UAV->MarkCounterResourceInitialized();
 		}
 	}
+	else
+	{
+		Cache.ResidencyHandles[ShaderStage][SlotIndex] = nullptr;
+	}
 }
 
 void FD3D12StateCache::SetBlendFactor(const float BlendFactor[4])
diff --git a/Engine/Source/Runtime/D3D12RHI/Private/D3D12StateCachePrivate.h b/Engine/Source/Runtime/D3D12RHI/Private/D3D12StateCachePrivate.h
--- a/Engine/Source/Runtime/D3D12RHI/Private/D3D12StateCachePrivate.h
+++ b/Engine/Source/Runtime/D3D12RHI/Private/D3D12StateCachePrivate.h
@@ -684,12 +684,8 @@ public:
 		}
 	}
 
-	void SetUAVs(EShaderFrequency ShaderStage, uint32 UAVStartSlot, uint32 NumSimultaneousUAVs, FD3D12UnorderedAccessView** UAVArray, uint32* UAVInitialCountArray);
-	void SetUAV(EShaderFrequency ShaderStage, uint32 SlotIndex, FD3D12UnorderedAccessView* UAV)
-	{
-		uint32 InitialCount = -1;
-		SetUAVs(ShaderStage, SlotIndex, 1, &UAV, &InitialCount);
-	}
+	void SetUAV(EShaderFrequency ShaderStage, uint32 SlotIndex, FD3D12UnorderedAccessView* UAV, uint32 InitialCount = -1);
+
 	void ClearUAVs(EShaderFrequency ShaderStage);
 
 	void SetDepthBounds(float MinDepth, float MaxDepth)
