Temporary speculative workaround for a GPU Crash in indirect dispatch rays.

diff --git a/Engine/Source/Runtime/D3D12RHI/Private/D3D12RayTracing.cpp b/Engine/Source/Runtime/D3D12RHI/Private/D3D12RayTracing.cpp
--- a/Engine/Source/Runtime/D3D12RHI/Private/D3D12RayTracing.cpp
+++ b/Engine/Source/Runtime/D3D12RHI/Private/D3D12RayTracing.cpp
@@ -5031,7 +5031,7 @@ static void DispatchRays(FD3D12CommandContext& CommandContext,
 
 		// Copy SBT data part of the dispatch desc to upload memory
 		FD3D12ResourceLocation UploadResourceLocation(Device);
-		void* Data = Device->GetDefaultFastAllocator().Allocate(DispatchRayDescSize, 256, &UploadResourceLocation);
+		void* Data = Adapter->GetUploadHeapAllocator(Device->GetGPUIndex()).AllocUploadResource(DispatchRayDescSize, 256, UploadResourceLocation);
 		FMemory::Memcpy(Data, &DispatchDesc, SBTPartSize);
 
 		// Copy SBT data part to resource
