Moving UniformBuffer updating with Bindless descriptors into a shared location.

diff --git a/Engine/Source/Runtime/D3D12RHI/Private/D3D12UniformBuffer.cpp b/Engine/Source/Runtime/D3D12RHI/Private/D3D12UniformBuffer.cpp
--- a/Engine/Source/Runtime/D3D12RHI/Private/D3D12UniformBuffer.cpp
+++ b/Engine/Source/Runtime/D3D12RHI/Private/D3D12UniformBuffer.cpp
@@ -7,111 +7,7 @@
 #include "D3D12RHIPrivate.h"
 #include "UniformBuffer.h"
 #include "ShaderParameterStruct.h"
-
-/// Utility class for reading shader parameters out of the data blob passed in.
-struct FUniformDataReader
-{
-	FUniformDataReader() = delete;
-	FUniformDataReader(const void* InData) : Data(reinterpret_cast<const uint8*>(InData)) { }
-
-	template<typename TResourceOut>
-	const TResourceOut& Read(const FRHIUniformBufferResource& InResource) const
-	{
-		return *reinterpret_cast<const TResourceOut*>(Data + InResource.MemberOffset);
-	}
-
-	const uint8* Data;
-};
-
-inline FRHIDescriptorHandle GetBindlessResourceHandle(FUniformDataReader Reader, const FRHIUniformBufferResource& Resource)
-{
-	switch (Resource.MemberType)
-	{
-	case UBMT_TEXTURE:
-	{
-		FRHITexture* Texture = Reader.Read<FRHITexture*>(Resource);
-		return Texture ? Texture->GetDefaultBindlessHandle() : FRHIDescriptorHandle();
-	}
-	break;
-	case UBMT_SRV:
-	{
-		FRHIShaderResourceView* ShaderResourceView = Reader.Read<FRHIShaderResourceView*>(Resource);
-		return ShaderResourceView ? ShaderResourceView->GetBindlessHandle() : FRHIDescriptorHandle();
-	}
-	break;
-	case UBMT_SAMPLER:
-	{
-		FRHISamplerState* SamplerState = Reader.Read<FRHISamplerState*>(Resource);
-		return SamplerState ? SamplerState->GetBindlessHandle() : FRHIDescriptorHandle();
-	}
-	break;
-	case UBMT_RDG_TEXTURE:
-	{
-		FRDGTexture* RDGTexture = Reader.Read<FRDGTexture*>(Resource);
-		FRHITexture* Texture = RDGTexture ? RDGTexture->GetRHI() : nullptr;
-		return Texture ? Texture->GetDefaultBindlessHandle() : FRHIDescriptorHandle();
-	}
-	break;
-	case UBMT_RDG_TEXTURE_SRV:
-	case UBMT_RDG_BUFFER_SRV:
-	{
-		FRDGShaderResourceView* RDGShaderResourceView = Reader.Read<FRDGShaderResourceView*>(Resource);
-		if (RDGShaderResourceView)
-			RDGShaderResourceView->MarkResourceAsUsed();
-
-		FRHIShaderResourceView* ShaderResourceView = RDGShaderResourceView ? RDGShaderResourceView->GetRHI() : nullptr;
-		return ShaderResourceView ? ShaderResourceView->GetBindlessHandle() : FRHIDescriptorHandle();
-	}
-	break;
-	case UBMT_UAV:
-	{
-		FRHIUnorderedAccessView* UnorderedAccessView = Reader.Read<FRHIUnorderedAccessView*>(Resource);
-		return UnorderedAccessView ? UnorderedAccessView->GetBindlessHandle() : FRHIDescriptorHandle();
-	}
-	case UBMT_RDG_TEXTURE_UAV:
-	case UBMT_RDG_BUFFER_UAV:
-	{
-		FRDGUnorderedAccessView* RDGUnorderedAccessView = Reader.Read<FRDGUnorderedAccessView*>(Resource);
-		FRHIUnorderedAccessView* UnorderedAccessView = RDGUnorderedAccessView ? RDGUnorderedAccessView->GetRHI() : nullptr;
-		return UnorderedAccessView ? UnorderedAccessView->GetBindlessHandle() : FRHIDescriptorHandle();
-	}
-	break;
-	case UBMT_NESTED_STRUCT:
-	case UBMT_INCLUDED_STRUCT:
-	case UBMT_REFERENCED_STRUCT:
-	{
-		// Do nothing?
-	}
-	break;
-
-	default:
-		//checkf(false, TEXT("Unhandled resource type?"));
-		break;
-	}
-	return FRHIDescriptorHandle();
-}
-
-static void UpdateUniformBufferConstants(void* DestinationData, const void* SourceData, const FRHIUniformBufferLayout* Layout)
-{
-	check(DestinationData != nullptr);
-	check(SourceData != nullptr);
-
-	// First copy wholesale
-	FMemory::Memcpy(DestinationData, SourceData, Layout->ConstantBufferSize);
-
-	FUniformDataReader Reader(SourceData);
-
-	// Then copy indices over
-	for (const FRHIUniformBufferResource& Resource : Layout->Resources)
-	{
-		const FRHIDescriptorHandle Handle = GetBindlessResourceHandle(Reader, Resource);
-		if (Handle.IsValid())
-		{
-			const uint32 BindlessIndex = Handle.GetIndex();
-			FMemory::Memcpy(reinterpret_cast<uint8*>(DestinationData) + Resource.MemberOffset, &BindlessIndex, sizeof(BindlessIndex));
-		}
-	}
-}
+#include "RHIUniformBufferDataShared.h"
 
 FUniformBufferRHIRef FD3D12DynamicRHI::RHICreateUniformBuffer(const void* Contents, const FRHIUniformBufferLayout* Layout, EUniformBufferUsage Usage, EUniformBufferValidation Validation)
 {
@@ -165,7 +61,7 @@ FUniformBufferRHIRef FD3D12DynamicRHI::RHICreateUniformBuffer(const void* Conten
 				// Copy the data to the upload heap
 				check(MappedData != nullptr);
 
-				UpdateUniformBufferConstants(MappedData, Contents, Layout);
+				UE::RHICore::UpdateUniformBufferConstants(MappedData, Contents, *Layout);
 
 #if USE_STATIC_ROOT_SIGNATURE
 				NewUniformBuffer->View->Create(NewUniformBuffer->ResourceLocation.GetGPUVirtualAddress(), NumBytesActualData);
@@ -301,7 +197,7 @@ void FD3D12DynamicRHI::RHIUpdateUniformBuffer(FRHICommandListBase& RHICmdList, F
 			}
 
 			check(MappedData != nullptr);
-			UpdateUniformBufferConstants(MappedData, Contents, &Layout);
+			UE::RHICore::UpdateUniformBufferConstants(MappedData, Contents, Layout);
 		}
 
 		if (bBypass)
diff --git a/Engine/Source/Runtime/RHICore/Internal/RHIUniformBufferDataShared.h b/Engine/Source/Runtime/RHICore/Internal/RHIUniformBufferDataShared.h
new file mode 100644
--- /dev/null
+++ b/Engine/Source/Runtime/RHICore/Internal/RHIUniformBufferDataShared.h
@@ -0,0 +1,115 @@
+// Copyright Epic Games, Inc. All Rights Reserved.
+
+#pragma once
+
+#include "RHI.h"
+
+namespace UE::RHICore
+{
+	struct FUniformDataReader
+	{
+		FUniformDataReader() = delete;
+		FUniformDataReader(const void* InData) : Data(reinterpret_cast<const uint8*>(InData)) { }
+
+		template<typename TResourceOut>
+		const TResourceOut& Read(const FRHIUniformBufferResource& InResource) const
+		{
+			return *reinterpret_cast<const TResourceOut*>(Data + InResource.MemberOffset);
+		}
+
+		const uint8* Data;
+	};
+
+	inline FRHIDescriptorHandle GetBindlessResourceHandle(FUniformDataReader Reader, const FRHIUniformBufferResource& Resource)
+	{
+		switch (Resource.MemberType)
+		{
+		case UBMT_TEXTURE:
+		{
+			FRHITexture* Texture = Reader.Read<FRHITexture*>(Resource);
+			return Texture ? Texture->GetDefaultBindlessHandle() : FRHIDescriptorHandle();
+		}
+		break;
+		case UBMT_SRV:
+		{
+			FRHIShaderResourceView* ShaderResourceView = Reader.Read<FRHIShaderResourceView*>(Resource);
+			return ShaderResourceView ? ShaderResourceView->GetBindlessHandle() : FRHIDescriptorHandle();
+		}
+		break;
+		case UBMT_SAMPLER:
+		{
+			FRHISamplerState* SamplerState = Reader.Read<FRHISamplerState*>(Resource);
+			return SamplerState ? SamplerState->GetBindlessHandle() : FRHIDescriptorHandle();
+		}
+		break;
+		case UBMT_RDG_TEXTURE:
+		{
+			FRDGTexture* RDGTexture = Reader.Read<FRDGTexture*>(Resource);
+			FRHITexture* Texture = RDGTexture ? RDGTexture->GetRHI() : nullptr;
+			return Texture ? Texture->GetDefaultBindlessHandle() : FRHIDescriptorHandle();
+		}
+		break;
+		case UBMT_RDG_TEXTURE_SRV:
+		case UBMT_RDG_BUFFER_SRV:
+		{
+			FRDGShaderResourceView* RDGShaderResourceView = Reader.Read<FRDGShaderResourceView*>(Resource);
+			if (RDGShaderResourceView)
+				RDGShaderResourceView->MarkResourceAsUsed();
+
+			FRHIShaderResourceView* ShaderResourceView = RDGShaderResourceView ? RDGShaderResourceView->GetRHI() : nullptr;
+			return ShaderResourceView ? ShaderResourceView->GetBindlessHandle() : FRHIDescriptorHandle();
+		}
+		break;
+		case UBMT_UAV:
+		{
+			FRHIUnorderedAccessView* UnorderedAccessView = Reader.Read<FRHIUnorderedAccessView*>(Resource);
+			return UnorderedAccessView ? UnorderedAccessView->GetBindlessHandle() : FRHIDescriptorHandle();
+		}
+		case UBMT_RDG_TEXTURE_UAV:
+		case UBMT_RDG_BUFFER_UAV:
+		{
+			FRDGUnorderedAccessView* RDGUnorderedAccessView = Reader.Read<FRDGUnorderedAccessView*>(Resource);
+			FRHIUnorderedAccessView* UnorderedAccessView = RDGUnorderedAccessView ? RDGUnorderedAccessView->GetRHI() : nullptr;
+			return UnorderedAccessView ? UnorderedAccessView->GetBindlessHandle() : FRHIDescriptorHandle();
+		}
+		break;
+		case UBMT_NESTED_STRUCT:
+		case UBMT_INCLUDED_STRUCT:
+		case UBMT_REFERENCED_STRUCT:
+		{
+			// Do nothing?
+		}
+		break;
+
+		default:
+			//checkf(false, TEXT("Unhandled resource type?"));
+			break;
+		}
+		return FRHIDescriptorHandle();
+	}
+
+	inline void UpdateUniformBufferConstants(void* DestinationData, const void* SourceData, const FRHIUniformBufferLayout& Layout, bool bAllowBindless = true)
+	{
+		check(DestinationData != nullptr);
+		check(SourceData != nullptr);
+
+		// First copy wholesale
+		FMemory::Memcpy(DestinationData, SourceData, Layout.ConstantBufferSize);
+
+		if (bAllowBindless)
+		{
+			FUniformDataReader Reader(SourceData);
+
+			// Then copy indices over
+			for (const FRHIUniformBufferResource& Resource : Layout.Resources)
+			{
+				const FRHIDescriptorHandle Handle = GetBindlessResourceHandle(Reader, Resource);
+				if (Handle.IsValid())
+				{
+					const uint32 BindlessIndex = Handle.GetIndex();
+					FMemory::Memcpy(reinterpret_cast<uint8*>(DestinationData) + Resource.MemberOffset, &BindlessIndex, sizeof(BindlessIndex));
+				}
+			}
+		}
+	}
+}
diff --git a/Engine/Source/Runtime/VulkanRHI/Private/VulkanUniformBuffer.cpp b/Engine/Source/Runtime/VulkanRHI/Private/VulkanUniformBuffer.cpp
--- a/Engine/Source/Runtime/VulkanRHI/Private/VulkanUniformBuffer.cpp
+++ b/Engine/Source/Runtime/VulkanRHI/Private/VulkanUniformBuffer.cpp
@@ -8,6 +8,7 @@
 #include "VulkanContext.h"
 #include "VulkanLLM.h"
 #include "ShaderParameterStruct.h"
+#include "RHIUniformBufferDataShared.h"
 
 static int32 GVulkanAllowUniformUpload = 1;
 static FAutoConsoleVariableRef CVarVulkanAllowUniformUpload(
@@ -29,62 +30,11 @@ enum
 -----------------------------------------------------------------------------*/
 
 
-
-// Utility class for reading shader parameters out of the data blob passed in.
-struct FUniformDataReader
-{
-	FUniformDataReader() = delete;
-	FUniformDataReader(const void* InData) : Data(reinterpret_cast<const uint8*>(InData)) { }
-
-	template<typename TResourceOut>
-	const TResourceOut& Read(const FRHIUniformBufferResource& InResource) const
-	{
-		return *reinterpret_cast<const TResourceOut*>(Data + InResource.MemberOffset);
-	}
-
-	const uint8* Data;
-};
-
 static void UpdateUniformBufferConstants(FVulkanDevice* Device, void* DestinationData, const void* SourceData, const FRHIUniformBufferLayout* Layout)
 {
-	check(DestinationData != nullptr);
-	check(SourceData != nullptr);
-
-	// First copy wholesale
-	FMemory::Memcpy(DestinationData, SourceData, Layout->ConstantBufferSize);
-
-	// Then copy indices over
-	if (Device->SupportsBindless())
-	{
-		FUniformDataReader Reader(SourceData);
-
-		for (const FRHIUniformBufferResource& Resource : Layout->Resources)
-		{
-			FRHIDescriptorHandle Handle;
-
-			if (Resource.MemberType == UBMT_SAMPLER)
-			{
-				FRHISamplerState* SamplerState = Reader.Read<FRHISamplerState*>(Resource);
-				if (SamplerState)
-				{
-					Handle = SamplerState->GetBindlessHandle();
-				}
-			}
-			else
-			{
-				// todo-jn: bindless resources : handle other types
-			}
-
-			if (Handle.IsValid())
-			{
-				const uint32 BindlessIndex = Handle.GetIndex();
-				FMemory::Memcpy(reinterpret_cast<uint8*>(DestinationData) + Resource.MemberOffset, &BindlessIndex, sizeof(BindlessIndex));
-			}
-		}
-	}
+	UE::RHICore::UpdateUniformBufferConstants(DestinationData, SourceData, *Layout, Device->SupportsBindless());
 }
 
-
 static inline EBufferUsageFlags UniformBufferToBufferUsage(EUniformBufferUsage Usage)
 {
 	switch (Usage)
