Disable D3D12 submission thread when running with RenderDoc.
This is because the programmatic catures in the render doc plugin rely on being able to call EndCapture() at a point just after the work is submitted. But we don't have a way to synchronize that when using a submission thread.

diff --git a/Engine/Plugins/Developer/RenderDocPlugin/Source/RenderDocPlugin/Private/RenderDocPluginModule.h b/Engine/Plugins/Developer/RenderDocPlugin/Source/RenderDocPlugin/Private/RenderDocPluginModule.h
--- a/Engine/Plugins/Developer/RenderDocPlugin/Source/RenderDocPlugin/Private/RenderDocPluginModule.h
+++ b/Engine/Plugins/Developer/RenderDocPlugin/Source/RenderDocPlugin/Private/RenderDocPluginModule.h
@@ -16,6 +16,7 @@ class FRenderDocPluginModule : public IRenderDocPlugin
 
 public:
 	// Begin IRenderCaptureProvider interface.
+	virtual bool CanSupportSubmissionThread() const { return false; }
 	virtual void CaptureFrame(FViewport* InViewport, uint32 InFlags, FString const& InDestFileName) override;
 	virtual void BeginCapture(FRHICommandListImmediate* InRHICommandList, uint32 InFlags, FString const& InDestFileName) override;
 	virtual void EndCapture(FRHICommandListImmediate* InRHICommandList) override;
diff --git a/Engine/Source/Runtime/D3D12RHI/Private/D3D12Submission.cpp b/Engine/Source/Runtime/D3D12RHI/Private/D3D12Submission.cpp
--- a/Engine/Source/Runtime/D3D12RHI/Private/D3D12Submission.cpp
+++ b/Engine/Source/Runtime/D3D12RHI/Private/D3D12Submission.cpp
@@ -3,6 +3,7 @@
 #include "D3D12RHIPrivate.h"
 #include "HAL/Runnable.h"
 #include "HAL/RunnableThread.h"
+#include "IRenderCaptureProvider.h"
 
 #ifndef D3D12_PLATFORM_SUPPORTS_BLOCKING_FENCES
 #define D3D12_PLATFORM_SUPPORTS_BLOCKING_FENCES 1
@@ -119,6 +120,9 @@ void FD3D12DynamicRHI::InitializeSubmissionPipe()
 	case 2: bUseSubmissionThread = true; break;
 	}
 
+	// Currently RenderDoc can't make programmatic captures when we use a submission thread.
+	bUseSubmissionThread &= !IRenderCaptureProvider::IsAvailable() || IRenderCaptureProvider::Get().CanSupportSubmissionThread();
+
 	if (bUseSubmissionThread)
 	{
 		SubmissionThread = new FD3D12Thread(TEXT("RHISubmissionThread"), TPri_Highest, this, &FD3D12DynamicRHI::ProcessSubmissionQueue);
diff --git a/Engine/Source/Runtime/RenderCore/Public/IRenderCaptureProvider.h b/Engine/Source/Runtime/RenderCore/Public/IRenderCaptureProvider.h
--- a/Engine/Source/Runtime/RenderCore/Public/IRenderCaptureProvider.h
+++ b/Engine/Source/Runtime/RenderCore/Public/IRenderCaptureProvider.h
@@ -46,6 +46,9 @@ public:
 		return IModularFeatures::Get().GetModularFeature<IRenderCaptureProvider>(GetModularFeatureName());
 	}
 
+	/** Derived classes return false if they cannot support an RHI submission thread. */
+	virtual bool CanSupportSubmissionThread() const { return true; }
+
 	/**
 	 * Flags to pass to capture API.
 	 */
