Fix calculation of temporary texture upload resource in D3D12 RHI

RowSizesInBytes was used to compute sizes of the first mip and the tail and to allocate upload resource memory. This can result in under-estimation, as the row size in bytes (i.e. valid data range) may be smaller than the row pitch (i.e. number of bytes between rows).

The pitch is what's used inside MemcpySubresource (called from D3DX12 UpdateSubresources) to compute the destination address. Since the smaller memory footprint was previously computed, smaller buffer range is sub-allocated from FD3D12FastAllocator and adjacent sub-allocations could be corrupted by MemcpySubresource.

The solution is to just use the total size of the resource (including all mips) returned by GetCopyableFootprints() and the offset of mip1 to compute the size of mip0 and the total size of the tail.

diff --git a/Engine/Source/Runtime/D3D12RHI/Private/D3D12Texture.cpp b/Engine/Source/Runtime/D3D12RHI/Private/D3D12Texture.cpp
--- a/Engine/Source/Runtime/D3D12RHI/Private/D3D12Texture.cpp
+++ b/Engine/Source/Runtime/D3D12RHI/Private/D3D12Texture.cpp
@@ -1079,7 +1079,7 @@ FTextureRHIRef FD3D12DynamicRHI::RHIAsyncCreateTexture2D(uint32 SizeX, uint32 Si
 		D3D12_PLACED_SUBRESOURCE_FOOTPRINT Layouts[MAX_TEXTURE_MIP_COUNT] = { };
 		UINT NumRows[MAX_TEXTURE_MIP_COUNT] = { };
 		UINT64 RowSizesInBytes[MAX_TEXTURE_MIP_COUNT] = { };
-		UINT64 TotalBytes = 0; // unused, got some issues with it
+		UINT64 TotalBytes = 0;
 		uint64 SizeMip0 = 0;
 		if (bSplitAllocation)
 		{
@@ -1088,16 +1088,12 @@ FTextureRHIRef FD3D12DynamicRHI::RHIAsyncCreateTexture2D(uint32 SizeX, uint32 Si
 			D3D12Device->GetCopyableFootprints(&TextureDesc, FirstSubresource, NumMips, 0, Layouts, NumRows, RowSizesInBytes, &TotalBytes);
 			
 			// Mip 0
-			SizeMip0 = RowSizesInBytes[0] * NumRows[0];
+			SizeMip0 = Layouts[1].Offset;
 			FastAllocator.Allocate(SizeMip0, D3D12_TEXTURE_DATA_PLACEMENT_ALIGNMENT, &TempResourceLocation);
 			Layouts[0].Offset = TempResourceLocation.GetOffsetFromBaseOfResource();
 
 			// Remaining mip chain
-			SizeLowMips = 0;
-			for (uint64 MipIndex = 1; MipIndex < NumMips; ++MipIndex)
-			{
-				SizeLowMips += RowSizesInBytes[MipIndex] * NumRows[MipIndex];
-			}
+			SizeLowMips = TotalBytes - SizeMip0;
 			FastAllocator.Allocate(SizeLowMips, D3D12_TEXTURE_DATA_PLACEMENT_ALIGNMENT, &TempResourceLocationLowMips);
 
 			uint64 BaseOffset = Layouts[1].Offset;
@@ -1108,9 +1104,9 @@ FTextureRHIRef FD3D12DynamicRHI::RHIAsyncCreateTexture2D(uint32 SizeX, uint32 Si
 				Layouts[MipIndex].Offset -= BaseOffset;
 				// The intermediate resource we get might be already used, so we need to account for the offset within this resource
 				Layouts[MipIndex].Offset += TempResourceLocationLowMips.GetOffsetFromBaseOfResource();
-				// We are going to map TempResourceLocationLowMips.GetResource(), and memcpy RowSizesInBytes[MipIndex] * NumRows[MipIndex] bytes in there offsetted
+				// We are going to map TempResourceLocationLowMips.GetResource(), and memcpy Layouts[MipIndex].Footprint.RowPitch * NumRows[MipIndex] bytes in there offsetted
 				// Make sure that the buffer is large enough before proceeding
-				check(Layouts[MipIndex].Offset + RowSizesInBytes[MipIndex] * NumRows[MipIndex] <= TempResourceLocationLowMips.GetResource()->GetDesc().Width);
+				check(Layouts[MipIndex].Offset + Layouts[MipIndex].Footprint.RowPitch * NumRows[MipIndex] <= TempResourceLocationLowMips.GetResource()->GetDesc().Width);
 			}
 			
 			TempResourceLocationLowMips.GetResource()->AddRef();
