Fix streaming textures occasionally missing inline mips when r.VirtualTextureReducedMemory is enabled. Courtesy to Serge for the fix.

#tests ab tested using a replay

diff --git a/Engine/Source/Runtime/D3D12RHI/Private/D3D12CommandContext.h b/Engine/Source/Runtime/D3D12RHI/Private/D3D12CommandContext.h
--- a/Engine/Source/Runtime/D3D12RHI/Private/D3D12CommandContext.h
+++ b/Engine/Source/Runtime/D3D12RHI/Private/D3D12CommandContext.h
@@ -41,7 +41,8 @@ struct FD3D12DeferredDeleteObject
 		D3DHeap,
 		BindlessDescriptor,
 		CPUAllocation,
-		DescriptorBlock
+		DescriptorBlock,
+		VirtualAllocation
 	} Type;
 
 	union
@@ -63,6 +64,14 @@ struct FD3D12DeferredDeleteObject
 			FD3D12OnlineDescriptorBlock* Block;
 			FD3D12OnlineDescriptorManager* Manager;
 		} DescriptorBlock;
+
+		struct
+		{
+			FPlatformMemory::FPlatformVirtualMemoryBlock VirtualBlock;
+			ETextureCreateFlags Flags;
+			uint64 CommittedTextureSize;
+			void* RawMemory;
+		} VirtualAllocDescriptor;
 	};
 
 	explicit FD3D12DeferredDeleteObject(FD3D12Resource* RHIObject)
@@ -96,6 +105,11 @@ struct FD3D12DeferredDeleteObject
 		: Type(EType::DescriptorBlock)
 		, DescriptorBlock({ Block, Manager })
 	{}
+
+	explicit FD3D12DeferredDeleteObject(FPlatformMemory::FPlatformVirtualMemoryBlock& VirtualBlock, ETextureCreateFlags Flags, uint64 CommittedTextureSize, void* RawMemory)
+		: Type(EType::VirtualAllocation)
+		, VirtualAllocDescriptor({ VirtualBlock, Flags, CommittedTextureSize, RawMemory })
+	{}
 };
 
 enum class ED3D12Units
diff --git a/Engine/Source/Runtime/D3D12RHI/Private/D3D12RHI.cpp b/Engine/Source/Runtime/D3D12RHI/Private/D3D12RHI.cpp
--- a/Engine/Source/Runtime/D3D12RHI/Private/D3D12RHI.cpp
+++ b/Engine/Source/Runtime/D3D12RHI/Private/D3D12RHI.cpp
@@ -489,6 +489,15 @@ void FD3D12DynamicRHI::RHIPerFrameRHIFlushComplete()
 				case FD3D12DeferredDeleteObject::EType::DescriptorBlock:
 					ObjectToDelete.DescriptorBlock.Manager->Recycle(ObjectToDelete.DescriptorBlock.Block);
 					break;
+#if PLATFORM_SUPPORTS_VIRTUAL_TEXTURES	
+				case FD3D12DeferredDeleteObject::EType::VirtualAllocation:
+					FD3D12DynamicRHI::GetD3DRHI()->DestroyVirtualTexture(
+						ObjectToDelete.VirtualAllocDescriptor.Flags,
+						ObjectToDelete.VirtualAllocDescriptor.RawMemory,
+						const_cast<FPlatformMemory::FPlatformVirtualMemoryBlock&>(ObjectToDelete.VirtualAllocDescriptor.VirtualBlock),
+						ObjectToDelete.VirtualAllocDescriptor.CommittedTextureSize);
+					break;
+#endif
 				default:
 					checkf(false, TEXT("Unknown ED3D12DeferredDeleteObjectType"));
 					break;
