MGPU:  nDisplay view takes significantly longer on second GPU, due to occlusion culling failing to work on the second GPU.  Query API was always fetching the query from GPU 0 -- it needs to iterate over each GPU to search for a valid query.  Bug was introduced in 5.1 in CL 22239063 (Dev Parallel Rendering integration).  The previous version of the code that shipped with 5.0 included logic similar to this fix, so this basically restores that.

#rnx

diff --git a/Engine/Source/Runtime/D3D12RHI/Private/D3D12Query.cpp b/Engine/Source/Runtime/D3D12RHI/Private/D3D12Query.cpp
--- a/Engine/Source/Runtime/D3D12RHI/Private/D3D12Query.cpp
+++ b/Engine/Source/Runtime/D3D12RHI/Private/D3D12Query.cpp
@@ -293,13 +293,38 @@ void FD3D12DynamicRHI::RHIEndOcclusionQueryBatch_TopOfPipe(FRHICommandListBase&
 
 bool FD3D12DynamicRHI::RHIGetRenderQueryResult(FRHIRenderQuery* QueryRHI, uint64& OutResult, bool bWait, uint32 QueryGPUIndex)
 {
-	if (QueryGPUIndex == INDEX_NONE)
+	FD3D12RenderQuery* Query;
+
+	// This will be the common case, as most users aren't running MGPU, so check this first (also becomes a constant compare if WITH_MGPU disabled)
+	if (GNumExplicitGPUsForRendering <= 1)
+	{
+		Query = FD3D12DynamicRHI::ResourceCast(QueryRHI, 0);
+	}
+	else if (QueryGPUIndex != INDEX_NONE)
+	{
+		Query = FD3D12DynamicRHI::ResourceCast(QueryRHI, QueryGPUIndex);
+	}
+	else
 	{
-		// @todo mgpu - not sure how to handle "any" GPU here. Just pick the 0th.
-		QueryGPUIndex = 0;
+		// Pick the first query that has a valid sync point.  If none have a valid sync point, the function will return failure, so it
+		// doesn't matter which we pick.  We need to set Query outisde the loop to avoid an uninitialized variable compile error, so
+		// we check the first item before starting the loop.
+		FD3D12RenderQuery::FLinkedObjectIterator CurrentQuery((FD3D12RenderQuery*)QueryRHI);
+		Query = CurrentQuery.Get();
+
+		if (!Query->SyncPoint.IsValid())
+		{
+			for (++CurrentQuery; CurrentQuery; ++CurrentQuery)
+			{
+				Query = CurrentQuery.Get();
+				if (Query->SyncPoint.IsValid())
+				{
+					break;
+				}
+			}
+		}
 	}
 
-	FD3D12RenderQuery* Query = FD3D12DynamicRHI::ResourceCast(QueryRHI, QueryGPUIndex);
 	if (!ensureMsgf(Query->SyncPoint, TEXT("Attempt to get result data for an FRHIRenderQuery that was never used in a command list.")))
 	{
 		OutResult = 0;
