Added video memory tracking for D3D12 buffer to Unreal Insights

diff --git a/Engine/Source/Runtime/D3D12RHI/Private/D3D12Buffer.cpp b/Engine/Source/Runtime/D3D12RHI/Private/D3D12Buffer.cpp
--- a/Engine/Source/Runtime/D3D12RHI/Private/D3D12Buffer.cpp
+++ b/Engine/Source/Runtime/D3D12RHI/Private/D3D12Buffer.cpp
@@ -14,11 +14,10 @@ FD3D12Buffer::~FD3D12Buffer()
 		DefaultContext.StateCache.ClearVertexBuffer(&ResourceLocation);
 	}
 
-	int64 BufferSize = ResourceLocation.GetSize();
 	bool bTransient = ResourceLocation.IsTransient();
 	if (!bTransient)
 	{
-		UpdateBufferStats((EBufferUsageFlags)GetUsage(), -BufferSize);
+		UpdateBufferStats(this, false);
 	}
 }
 
@@ -404,7 +403,7 @@ FD3D12Buffer* FD3D12Adapter::CreateRHIBuffer(
 	// Don't track transient buffer stats here
 	if (!BufferOut->ResourceLocation.IsTransient())
 	{
-		UpdateBufferStats((EBufferUsageFlags)InUsage, BufferOut->ResourceLocation.GetSize());
+		UpdateBufferStats(BufferOut, true);
 	}
 
 	return BufferOut;
@@ -461,7 +460,7 @@ void FD3D12Buffer::ReleaseUnderlyingResource()
 	bool bTransient = ResourceLocation.IsTransient();
 	if (!bTransient)
 	{
-		UpdateBufferStats((EBufferUsageFlags)GetUsage(), -BufferSize);
+		UpdateBufferStats(this, false);
 	}
 
 	check(IsHeadLink());
diff --git a/Engine/Source/Runtime/D3D12RHI/Private/D3D12Stats.cpp b/Engine/Source/Runtime/D3D12RHI/Private/D3D12Stats.cpp
--- a/Engine/Source/Runtime/D3D12RHI/Private/D3D12Stats.cpp
+++ b/Engine/Source/Runtime/D3D12RHI/Private/D3D12Stats.cpp
@@ -7,6 +7,7 @@ D3D12Stats.cpp:RHI Stats and timing implementation.
 #include "D3D12RHIPrivate.h"
 #include "Engine/Engine.h"
 #include "Engine/GameViewportClient.h"
+#include "ProfilingDebugging/MemoryTrace.h"
 
 void D3D12RHI::FD3DGPUProfiler::BeginFrame()
 {
@@ -296,11 +297,24 @@ static FName GetD3D12BufferStat(EBufferUsageFlags InUsageFlags)
 	}
 }
 
-void UpdateBufferStats(EBufferUsageFlags InUsageFlags, int64 RequestedSize)
+void UpdateBufferStats(FD3D12Buffer* Buffer, bool bAllocating)
 {
-	INC_MEMORY_STAT_BY_FName(GetRHIBufferStats(InUsageFlags), RequestedSize);
-	INC_MEMORY_STAT_BY_FName(GetD3D12BufferStat(InUsageFlags), RequestedSize);
+	int64 BufferSize = Buffer->ResourceLocation.GetSize();
+	int64 RequestedSize = bAllocating ? BufferSize : -BufferSize;
+	EBufferUsageFlags UsageFlags = Buffer->GetUsage();
+
+	INC_MEMORY_STAT_BY_FName(GetRHIBufferStats(UsageFlags), RequestedSize);
+	INC_MEMORY_STAT_BY_FName(GetD3D12BufferStat(UsageFlags), RequestedSize);
 	INC_MEMORY_STAT_BY(STAT_D3D12MemoryCurrentTotal, RequestedSize);
+
+	if (bAllocating)
+	{
+		MemoryTrace_Alloc((uint64)Buffer->GetResource(), BufferSize, 0, EMemoryTraceRootHeap::VideoMemory);
+	}
+	else
+	{
+		MemoryTrace_Free((uint64)Buffer->GetResource(), EMemoryTraceRootHeap::VideoMemory);
+	}
 }
 
 #if NV_AFTERMATH
diff --git a/Engine/Source/Runtime/D3D12RHI/Public/D3D12Resources.h b/Engine/Source/Runtime/D3D12RHI/Public/D3D12Resources.h
--- a/Engine/Source/Runtime/D3D12RHI/Public/D3D12Resources.h
+++ b/Engine/Source/Runtime/D3D12RHI/Public/D3D12Resources.h
@@ -808,7 +808,7 @@ public:
 	}
 };
 
-extern void UpdateBufferStats(EBufferUsageFlags UsageFlags, int64 RequestedSize);
+extern void UpdateBufferStats(class FD3D12Buffer* Buffer, bool bAllocating);
 
 /** Uniform buffer resource class. */
 class FD3D12UniformBuffer : public FRHIUniformBuffer, public FD3D12DeviceChild, public FD3D12LinkedAdapterObject<FD3D12UniformBuffer>
