Only explicitly require resource binding tier 3 when non-RayTracing shaders are configured for Bindless. RayTracing shaders already have a runtime check for Tier3 so we don't need to raise the requirement for projects that aren't configured for RayTracing.

diff --git a/Engine/Source/Runtime/D3D12RHI/Private/Windows/WindowsD3D12Device.cpp b/Engine/Source/Runtime/D3D12RHI/Private/Windows/WindowsD3D12Device.cpp
--- a/Engine/Source/Runtime/D3D12RHI/Private/Windows/WindowsD3D12Device.cpp
+++ b/Engine/Source/Runtime/D3D12RHI/Private/Windows/WindowsD3D12Device.cpp
@@ -419,8 +419,9 @@ static bool CheckDeviceForEmulatedAtomic64Support(IDXGIAdapter* Adapter, ID3D12D
 
 inline bool ShouldCheckBindlessSupport(EShaderPlatform ShaderPlatform)
 {
-	return RHIGetBindlessResourcesConfiguration(ShaderPlatform) != ERHIBindlessConfiguration::Disabled
-		|| RHIGetBindlessSamplersConfiguration(ShaderPlatform) != ERHIBindlessConfiguration::Disabled;
+	// Note: only checking against AllShaders allows the RayTracingOnly configuration to not raise the requirements for projects that aren't using RayTracing.
+	return RHIGetBindlessResourcesConfiguration(ShaderPlatform) == ERHIBindlessConfiguration::AllShaders
+		|| RHIGetBindlessSamplersConfiguration(ShaderPlatform) == ERHIBindlessConfiguration::AllShaders;
 }
 
 inline ERHIFeatureLevel::Type FindMaxRHIFeatureLevel(D3D_FEATURE_LEVEL InMaxFeatureLevel, D3D_SHADER_MODEL InMaxShaderModel, D3D12_RESOURCE_BINDING_TIER ResourceBindingTier, bool bSupportsWaveOps, bool bSupportsAtomic64)
