Add CVar to keep lowlevel PSO in memory

diff --git a/Engine/Source/Runtime/D3D12RHI/Private/D3D12Pipelinestate.cpp b/Engine/Source/Runtime/D3D12RHI/Private/D3D12Pipelinestate.cpp
--- a/Engine/Source/Runtime/D3D12RHI/Private/D3D12Pipelinestate.cpp
+++ b/Engine/Source/Runtime/D3D12RHI/Private/D3D12Pipelinestate.cpp
@@ -15,6 +15,14 @@ static TAutoConsoleVariable<float> CVarPSOStallWarningThresholdInMs(
 	TEXT("Value is in milliseconds. (100 is the default)\n"),
 	ECVF_ReadOnly);
 
+int32 GPSOPrecacheKeepLowLevel = 0;
+static FAutoConsoleVariableRef CVarPSOPrecacheKeepLowLevel(
+	TEXT("D3D12.PSOPrecache.KeepLowLevel"),
+	GPSOPrecacheKeepLowLevel,
+	TEXT("Keep in memory the d3d12 PSO blob for precached PSOs. Consumes more memory but reduces stalls.\n"),
+	ECVF_ReadOnly
+);
+
 /// @cond DOXYGEN_WARNINGS
 
 static FD3D12LowLevelGraphicsPipelineStateDesc GetLowLevelGraphicsPipelineStateDesc(const FGraphicsPipelineStateInitializer& Initializer, const FD3D12RootSignature* RootSignature)
@@ -377,7 +385,7 @@ FD3D12GraphicsPipelineState::~FD3D12GraphicsPipelineState()
 		check(RefCount > 0);
         // precache PSO are here to avoid hitches at runtime when we want to create one that is actually used. We don't need to keep them
 		// around as this can add up to a lot of system memory
-		if (PipelineStateInitializer.bPSOPrecache && RefCount == 1) 
+		if (PipelineStateInitializer.bPSOPrecache && RefCount == 1 && GPSOPrecacheKeepLowLevel == 0)
 		{
 			FD3D12DynamicRHI* D3D12RHI = FD3D12DynamicRHI::GetD3DRHI();
 			FD3D12PipelineStateCache& PSOCache = D3D12RHI->GetAdapter().GetPSOCache();
