Fix mip tail buffer size validation condition in FD3D12DynamicRHI::RHIAsyncCreateTexture2D

Mip copy footprint size was computed incorrectly in the check(), assuming that RowPitch worth of data will be accessed for every row. This is wrong for the last row, where only RowSize bytes will be accessed.

diff --git a/Engine/Source/Runtime/D3D12RHI/Private/D3D12Texture.cpp b/Engine/Source/Runtime/D3D12RHI/Private/D3D12Texture.cpp
--- a/Engine/Source/Runtime/D3D12RHI/Private/D3D12Texture.cpp
+++ b/Engine/Source/Runtime/D3D12RHI/Private/D3D12Texture.cpp
@@ -1097,17 +1097,39 @@ FTextureRHIRef FD3D12DynamicRHI::RHIAsyncCreateTexture2D(uint32 SizeX, uint32 Si
 			SizeLowMips = TotalBytes - SizeMip0;
 			FastAllocator.Allocate(SizeLowMips, D3D12_TEXTURE_DATA_PLACEMENT_ALIGNMENT, &TempResourceLocationLowMips);
 
-			uint64 BaseOffset = Layouts[1].Offset;
+			const uint64 LowMipsTotalBufferSize = TempResourceLocationLowMips.GetResource()->GetDesc().Width;
+			
+			const uint64 BaseOffset = Layouts[1].Offset;
+
 			for (uint64 MipIndex = 1; MipIndex < NumMips; ++MipIndex)
 			{
 				check(Layouts[MipIndex].Offset >= BaseOffset);
+
+				const uint64 RelativeMipCopyOffset = Layouts[MipIndex].Offset - BaseOffset; // Offset relative to mip1
+
 				// The original offsets for the remaining mipchain were originally computed with mip0, so we need to remove that offset
 				Layouts[MipIndex].Offset -= BaseOffset;
 				// The intermediate resource we get might be already used, so we need to account for the offset within this resource
 				Layouts[MipIndex].Offset += TempResourceLocationLowMips.GetOffsetFromBaseOfResource();
-				// We are going to map TempResourceLocationLowMips.GetResource(), and memcpy Layouts[MipIndex].Footprint.RowPitch * NumRows[MipIndex] bytes in there offsetted
-				// Make sure that the buffer is large enough before proceeding
-				check(Layouts[MipIndex].Offset + Layouts[MipIndex].Footprint.RowPitch * NumRows[MipIndex] <= TempResourceLocationLowMips.GetResource()->GetDesc().Width);
+
+				// UpdateSubresources copies mip levels taking into account RowPitch (number of bytes between rows) and RowSize (number of valid texture data bytes).
+				// For each row, the destination address is computed as RowIndex*RowPitch and the copy size is always RowSize.
+				// If RowSize is smaller than RowPitch, the remaining bytes in the copy destination buffer are not touched.
+				// See MemcpySubresource() in d3dx12_resource_helpers.h
+				check(NumRows[MipIndex] != 0);
+				const uint64 MipCopySize = Layouts[MipIndex].Footprint.RowPitch * (NumRows[MipIndex] - 1) + RowSizesInBytes[MipIndex];
+
+				// Make sure that the buffer is large enough before proceeding.
+
+				const uint64 RelativeMipCopyEndOffset = RelativeMipCopyOffset + MipCopySize;
+				checkf(RelativeMipCopyEndOffset <= SizeLowMips,
+					TEXT("Mip tail upload buffer allocation is too small for mip %llu. RelativeMipCopyOffset=%llu, MipCopySize=%llu, RelativeMipCopyEndOffset=%llu, SizeLowMips=%llu."),
+					MipIndex, RelativeMipCopyOffset, MipCopySize, RelativeMipCopyEndOffset, SizeLowMips);
+
+				const uint64 AbsoluteMipCopyEndOffset = Layouts[MipIndex].Offset + MipCopySize;
+				checkf(AbsoluteMipCopyEndOffset <= LowMipsTotalBufferSize,
+					TEXT("Mip tail upload buffer total size is too small for mip %llu. Layouts[MipIndex].Offset=%llu, MipCopySize=%llu, AbsoluteMipCopyEndOffset=%llu, LowMipsTotalBufferSize=%llu."),
+					MipIndex, Layouts[MipIndex].Offset, MipCopySize, AbsoluteMipCopyEndOffset, LowMipsTotalBufferSize);
 			}
 			
 			TempResourceLocationLowMips.GetResource()->AddRef();
