Insight memory trace DirectX12 memory allocator tracking

diff --git a/Engine/Source/Runtime/D3D12RHI/Private/D3D12Adapter.cpp b/Engine/Source/Runtime/D3D12RHI/Private/D3D12Adapter.cpp
--- a/Engine/Source/Runtime/D3D12RHI/Private/D3D12Adapter.cpp
+++ b/Engine/Source/Runtime/D3D12RHI/Private/D3D12Adapter.cpp
@@ -1081,7 +1081,8 @@ PRAGMA_ENABLE_DEPRECATION_WARNINGS
 #if TRACK_RESOURCE_ALLOCATIONS
 		// Set flag if we want to track all allocations - comes with some overhead and only possible when Tier 2 is available
 		// (because we will create placed buffers for texture allocation to retrieve the GPU virtual addresses)
-		bTrackAllAllocation = (GD3D12TrackAllAlocations || GPUCrashDebuggingModes == ED3D12GPUCrashDebuggingModes::All) && (GetResourceHeapTier() == D3D12_RESOURCE_HEAP_TIER_2);
+		const bool bTraceMemAlloc = UE_TRACE_CHANNELEXPR_IS_ENABLED(MemAllocChannel);
+		bTrackAllAllocation = (GD3D12TrackAllAlocations || GPUCrashDebuggingModes == ED3D12GPUCrashDebuggingModes::All || bTraceMemAlloc) && (GetResourceHeapTier() == D3D12_RESOURCE_HEAP_TIER_2);
 #endif 
 
 		// Context redirectors allow RHI commands to be executed on multiple GPUs at the
diff --git a/Engine/Source/Runtime/D3D12RHI/Private/D3D12Allocation.cpp b/Engine/Source/Runtime/D3D12RHI/Private/D3D12Allocation.cpp
--- a/Engine/Source/Runtime/D3D12RHI/Private/D3D12Allocation.cpp
+++ b/Engine/Source/Runtime/D3D12RHI/Private/D3D12Allocation.cpp
@@ -9,6 +9,7 @@
 #include "D3D12Allocation.h"
 #include "Misc/BufferedOutputDevice.h"
 #include "HAL/PlatformStackWalk.h"
+#include "ProfilingDebugging/MemoryTrace.h"
 
 // Fix for random GPU crashes on draw indirects on multiple IHVs. Force all indirect arg buffers as committed resources (see UE-115982)
 static int32 GD3D12AllowPoolAllocateIndirectArgBuffers = 0;
@@ -232,7 +233,8 @@ FD3D12BuddyAllocator::FD3D12BuddyAllocator(FD3D12Device* ParentDevice,
 	EResourceAllocationStrategy InAllocationStrategy,
 	uint32 MaxSizeForPooling,
 	uint32 InMaxBlockSize,
-	uint32 InMinBlockSize)
+	uint32 InMinBlockSize,
+	HeapId InTraceParentHeapId)
 	: FD3D12ResourceAllocator(ParentDevice, VisibleNodes, InInitConfig, Name, MaxSizeForPooling)
 	, MaxBlockSize(InMaxBlockSize)
 	, MinBlockSize(InMinBlockSize)
@@ -242,6 +244,9 @@ FD3D12BuddyAllocator::FD3D12BuddyAllocator(FD3D12Device* ParentDevice,
 	, TotalSizeUsed(0)
 	, HeapFullMessageDisplayed(false)
 {
+#if UE_MEMORY_TRACE_ENABLED
+	TraceHeapId = MemoryTrace_HeapSpec(InTraceParentHeapId, AllocationStrategy == EResourceAllocationStrategy::kPlacedResource ? TEXT("BuddyAllocator (PlacedResource)") : TEXT("BuddyAllocator (ManualSubAllocation)"));
+#endif
 	// maxBlockSize should be evenly dividable by MinBlockSize and  
 	// maxBlockSize / MinBlockSize should be a power of two  
 	check((MaxBlockSize / MinBlockSize) * MinBlockSize == MaxBlockSize); // Evenly dividable  
@@ -282,7 +287,7 @@ void FD3D12BuddyAllocator::Initialize()
 			VERIFYD3D12RESULT(Adapter->GetD3DDevice()->CreateHeap(&Desc, IID_PPV_ARGS(&Heap)));
 		}
 
-		BackingHeap = new FD3D12Heap(GetParentDevice(), GetVisibilityMask());
+		BackingHeap = new FD3D12Heap(GetParentDevice(), GetVisibilityMask(), TraceHeapId);
 		BackingHeap->SetHeap(Heap, TEXT("BuddyAllocator Heap"));
 
 		// Only track resources that cannot be accessed on the CPU.
@@ -297,6 +302,9 @@ void FD3D12BuddyAllocator::Initialize()
 			LLM_SCOPED_PAUSE_TRACKING_FOR_TRACKER(ELLMTracker::Default, ELLMAllocType::System);
 			const D3D12_HEAP_PROPERTIES HeapProps = CD3DX12_HEAP_PROPERTIES(InitConfig.HeapType, GetGPUMask().GetNative(), GetVisibilityMask().GetNative());
 			VERIFYD3D12RESULT(Adapter->CreateBuffer(HeapProps, GetGPUMask(), InitConfig.InitialResourceState, ED3D12ResourceStateMode::SingleState, InitConfig.InitialResourceState, MaxBlockSize, BackingResource.GetInitReference(), TEXT("Resource Allocator Underlying Buffer"), InitConfig.ResourceFlags));
+#if UE_MEMORY_TRACE_ENABLED
+			MemoryTrace_MarkAllocAsHeap(BackingResource->GetGPUVirtualAddress(), TraceHeapId);
+#endif
 		}
 
 		if (IsCPUAccessible(InitConfig.HeapType))
@@ -452,6 +460,9 @@ void FD3D12BuddyAllocator::Allocate(uint32 SizeInBytes, uint32 Alignment, FD3D12
 	// which means that the memory would be counted twice. Because of this the tracking had to be disabled here.
 	// This does mean that non-buffer memory that goes through this allocator won't be tracked, so this does need a better solution.
 	// see UpdateBufferStats for a more detailed explanation.
+#if UE_MEMORY_TRACE_ENABLED
+	MemoryTrace_Alloc(ResourceLocation.GetGPUVirtualAddress(), SizeInBytes, Alignment, EMemoryTraceRootHeap::VideoMemory);
+#endif
 #endif
 }
 
@@ -510,6 +521,9 @@ void FD3D12BuddyAllocator::Deallocate(FD3D12ResourceLocation& ResourceLocation)
 	// This does mean that non-buffer memory that goes through this allocator won't be tracked, so this does need a better solution.
 	// see UpdateBufferStats for a more detailed explanation.
 	LLM(FLowLevelMemTracker::Get().OnLowLevelFree(ELLMTracker::Default, ResourceLocation.GetAddressForLLMTracking()));
+#if UE_MEMORY_TRACE_ENABLED
+	MemoryTrace_Free(ResourceLocation.GetGPUVirtualAddress(), EMemoryTraceRootHeap::VideoMemory);
+#endif
 #endif
 }
 
@@ -567,6 +581,15 @@ void FD3D12BuddyAllocator::ReleaseAllResources()
 {
 	LLM_SCOPED_PAUSE_TRACKING_FOR_TRACKER(ELLMTracker::Default, ELLMAllocType::System);
 
+#if UE_MEMORY_TRACE_ENABLED
+	D3D12_GPU_VIRTUAL_ADDRESS GPUAddress = 
+		AllocationStrategy == EResourceAllocationStrategy::kPlacedResource ? (BackingHeap ? BackingHeap->GetGPUVirtualAddress() : 0) : (BackingResource ? BackingResource->GetGPUVirtualAddress() : 0);
+	if (GPUAddress > 0)
+	{
+		MemoryTrace_UnmarkAllocAsHeap(GPUAddress, TraceHeapId);
+	}
+#endif
+
 	for (RetiredBlock& Block : DeferredDeletionQueue)
 	{
 		DeallocateInternal(Block);
@@ -677,12 +700,17 @@ FD3D12MultiBuddyAllocator::FD3D12MultiBuddyAllocator(FD3D12Device* ParentDevice,
 	EResourceAllocationStrategy InAllocationStrategy,
 	uint32 InMaxAllocationSize,
 	uint32 InDefaultPoolSize,
-	uint32 InMinBlockSize)
+	uint32 InMinBlockSize,
+	HeapId InTraceParentHeapId)
 	: FD3D12ResourceAllocator(ParentDevice, VisibleNodes, InInitConfig, Name, InMaxAllocationSize)
 	, AllocationStrategy(InAllocationStrategy)
 	, MinBlockSize(InMinBlockSize)
 	, DefaultPoolSize(InDefaultPoolSize)
-{}
+{
+#if UE_MEMORY_TRACE_ENABLED
+	TraceHeapId = MemoryTrace_HeapSpec(InTraceParentHeapId, *Name);
+#endif
+}
 
 FD3D12MultiBuddyAllocator::~FD3D12MultiBuddyAllocator()
 {
@@ -723,7 +751,8 @@ FD3D12BuddyAllocator* FD3D12MultiBuddyAllocator::CreateNewAllocator(uint32 InMin
 		AllocationStrategy,
 		AllocationSize,
 		AllocationSize,
-		MinBlockSize);
+		MinBlockSize,
+		TraceHeapId);
 }
 
 void FD3D12MultiBuddyAllocator::Initialize()
@@ -1018,13 +1047,14 @@ void FD3D12BucketAllocator::Reset()
 FD3D12UploadHeapAllocator::FD3D12UploadHeapAllocator(FD3D12Adapter* InParent, FD3D12Device* InParentDevice, const FString& InName)
 	: FD3D12AdapterChild(InParent)
 	, FD3D12DeviceChild(InParentDevice)
-	, FD3D12MultiNodeGPUObject(InParentDevice->GetGPUMask(), FRHIGPUMask::All()), // Upload memory, thus they can be trivially visibile to all GPUs
-	SmallBlockAllocator(InParentDevice, GetVisibilityMask(), FD3D12ResourceInitConfig::CreateUpload(), InName, EResourceAllocationStrategy::kManualSubAllocation,
-		GD3D12UploadHeapSmallBlockMaxAllocationSize, GD3D12UploadHeapSmallBlockPoolSize, D3D12_TEXTURE_DATA_PLACEMENT_ALIGNMENT),
-	BigBlockAllocator(InParentDevice, GetVisibilityMask(), FD3D12ResourceInitConfig::CreateUpload(), InName, EResourceAllocationStrategy::kManualSubAllocation,
-		GD3D12UploadHeapBigBlockPoolSize, D3D12_TEXTURE_DATA_PLACEMENT_ALIGNMENT, GD3D12UploadHeapBigBlockMaxAllocationSize, FRHIMemoryPool::EFreeListOrder::SortByOffset, false /*defrag*/),
-	FastConstantPageAllocator(InParentDevice, GetVisibilityMask(), FD3D12ResourceInitConfig::CreateUpload(), InName, EResourceAllocationStrategy::kManualSubAllocation,
-		GD3D12FastConstantAllocatorPageSize * 64, GD3D12UploadHeapSmallBlockPoolSize, GD3D12FastConstantAllocatorPageSize)
+	, FD3D12MultiNodeGPUObject(InParentDevice->GetGPUMask(), FRHIGPUMask::All()) // Upload memory, thus they can be trivially visibile to all GPUs
+	, TraceHeapId(MemoryTrace_HeapSpec(EMemoryTraceRootHeap::VideoMemory, *FString(InName + TEXT(" (UploadHeapAllocator)")))),
+	SmallBlockAllocator(InParentDevice, GetVisibilityMask(), FD3D12ResourceInitConfig::CreateUpload(), TEXT("Small Block Multi Buddy Allocator"), EResourceAllocationStrategy::kManualSubAllocation,
+		GD3D12UploadHeapSmallBlockMaxAllocationSize, GD3D12UploadHeapSmallBlockPoolSize, D3D12_TEXTURE_DATA_PLACEMENT_ALIGNMENT, TraceHeapId),
+	BigBlockAllocator(InParentDevice, GetVisibilityMask(), FD3D12ResourceInitConfig::CreateUpload(), TEXT("Big Block Pool Allocator"), EResourceAllocationStrategy::kManualSubAllocation,
+		GD3D12UploadHeapBigBlockPoolSize, D3D12_TEXTURE_DATA_PLACEMENT_ALIGNMENT, GD3D12UploadHeapBigBlockMaxAllocationSize, FRHIMemoryPool::EFreeListOrder::SortByOffset, false /*defrag*/, TraceHeapId),
+	FastConstantPageAllocator(InParentDevice, GetVisibilityMask(), FD3D12ResourceInitConfig::CreateUpload(), TEXT("Fast Constant Page Multi Buddy Allocator"), EResourceAllocationStrategy::kManualSubAllocation,
+		GD3D12FastConstantAllocatorPageSize * 64, GD3D12UploadHeapSmallBlockPoolSize, GD3D12FastConstantAllocatorPageSize, TraceHeapId)
 {
 }
 
@@ -1336,6 +1366,7 @@ void FD3D12DefaultBufferPool::UpdateMemoryStats(uint32& IOMemoryAllocated, uint3
 FD3D12DefaultBufferAllocator::FD3D12DefaultBufferAllocator(FD3D12Device* InParent, FRHIGPUMask VisibleNodes)
 	: FD3D12DeviceChild(InParent)
 	, FD3D12MultiNodeGPUObject(InParent->GetGPUMask(), VisibleNodes)
+	, TraceHeapId(MemoryTrace_HeapSpec(EMemoryTraceRootHeap::VideoMemory, TEXT("Default Buffer Allocator")))
 {
 	FMemory::Memset(DefaultBufferPools, 0);
 }
@@ -1371,7 +1402,7 @@ FD3D12BufferPool* FD3D12DefaultBufferAllocator::CreateBufferPool(D3D12_HEAP_TYPE
 	}
 #endif // D3D12_RHI_RAYTRACING
 
-	FD3D12BufferPool* NewPool = new FD3D12PoolAllocator(Device, GetVisibilityMask(), InitConfig, Name, AllocationStrategy, PoolSize, PoolAlignment, MaxAllocationSize, FreeListOrder, bDefragEnabled);
+	FD3D12BufferPool* NewPool = new FD3D12PoolAllocator(Device, GetVisibilityMask(), InitConfig, Name, AllocationStrategy, PoolSize, PoolAlignment, MaxAllocationSize, FreeListOrder, bDefragEnabled, TraceHeapId);
 
 #else // USE_BUFFER_POOL_ALLOCATOR
 
@@ -1604,7 +1635,8 @@ void FD3D12DefaultBufferAllocator::UpdateMemoryStats()
 #if USE_TEXTURE_POOL_ALLOCATOR
 FD3D12TextureAllocatorPool::FD3D12TextureAllocatorPool(FD3D12Device* Device, FRHIGPUMask VisibilityNode) :
 	FD3D12DeviceChild(Device),
-	FD3D12MultiNodeGPUObject(Device->GetGPUMask(), VisibilityNode)
+	FD3D12MultiNodeGPUObject(Device->GetGPUMask(), VisibilityNode),
+	TraceHeapId(MemoryTrace_HeapSpec(EMemoryTraceRootHeap::VideoMemory, TEXT("Texture Allocator Pool")))
 {	
 	FD3D12ResourceInitConfig SharedInitConfig;
 
@@ -1626,7 +1658,7 @@ FD3D12TextureAllocatorPool::FD3D12TextureAllocatorPool(FD3D12Device* Device, FRH
 		uint64 PoolSize = 4 * 1024 * 1024;
 		uint64 MaxAllocationSize = PoolSize;
 		bool bDefragEnabled = false; // Disable defrag on 4K pool because it shouldn't really fragment - all allocations are 4K or multiple of 4K and pretty small
-		FD3D12PoolAllocator* ReadOnly4KPool = new FD3D12PoolAllocator(Device, GetVisibilityMask(), InitConfig, Name, AllocationStrategy, PoolSize, D3D12_SMALL_RESOURCE_PLACEMENT_ALIGNMENT, MaxAllocationSize, FreeListOrder, bDefragEnabled);
+		FD3D12PoolAllocator* ReadOnly4KPool = new FD3D12PoolAllocator(Device, GetVisibilityMask(), InitConfig, Name, AllocationStrategy, PoolSize, D3D12_SMALL_RESOURCE_PLACEMENT_ALIGNMENT, MaxAllocationSize, FreeListOrder, bDefragEnabled, TraceHeapId);
 		ReadOnly4KPool->Initialize();
 
 		PoolAllocators[(int)EPoolType::ReadOnly4K] = ReadOnly4KPool;
@@ -1640,7 +1672,7 @@ FD3D12TextureAllocatorPool::FD3D12TextureAllocatorPool(FD3D12Device* Device, FRH
 		uint64 PoolSize = GD3D12PoolAllocatorReadOnlyTextureVRAMPoolSize;
 		uint64 MaxAllocationSize = GD3D12PoolAllocatorReadOnlyTextureVRAMMaxAllocationSize;
 		bool bDefragEnabled = true;
-		FD3D12PoolAllocator* ReadOnlyPool = new FD3D12PoolAllocator(Device, GetVisibilityMask(), InitConfig, Name, AllocationStrategy, PoolSize, D3D12_DEFAULT_RESOURCE_PLACEMENT_ALIGNMENT, MaxAllocationSize, FreeListOrder, bDefragEnabled);
+		FD3D12PoolAllocator* ReadOnlyPool = new FD3D12PoolAllocator(Device, GetVisibilityMask(), InitConfig, Name, AllocationStrategy, PoolSize, D3D12_DEFAULT_RESOURCE_PLACEMENT_ALIGNMENT, MaxAllocationSize, FreeListOrder, bDefragEnabled, TraceHeapId);
 		ReadOnlyPool->Initialize();
 
 		PoolAllocators[(int)EPoolType::ReadOnly] = ReadOnlyPool;
@@ -1654,7 +1686,7 @@ FD3D12TextureAllocatorPool::FD3D12TextureAllocatorPool(FD3D12Device* Device, FRH
 		uint64 PoolSize = GD3D12PoolAllocatorRTUAVTextureVRAMPoolSize;
 		uint64 MaxAllocationSize = GD3D12PoolAllocatorRTUAVTextureVRAMMaxAllocationSize;
 		bool bDefragEnabled = true;
-		FD3D12PoolAllocator* RTPool = new FD3D12PoolAllocator(Device, GetVisibilityMask(), InitConfig, Name, AllocationStrategy, PoolSize, D3D12_DEFAULT_RESOURCE_PLACEMENT_ALIGNMENT, MaxAllocationSize, FreeListOrder, bDefragEnabled);
+		FD3D12PoolAllocator* RTPool = new FD3D12PoolAllocator(Device, GetVisibilityMask(), InitConfig, Name, AllocationStrategy, PoolSize, D3D12_DEFAULT_RESOURCE_PLACEMENT_ALIGNMENT, MaxAllocationSize, FreeListOrder, bDefragEnabled, TraceHeapId);
 		RTPool->Initialize();
 
 		PoolAllocators[(int)EPoolType::RenderTarget] = RTPool;
@@ -1668,7 +1700,7 @@ FD3D12TextureAllocatorPool::FD3D12TextureAllocatorPool(FD3D12Device* Device, FRH
 		uint64 PoolSize = GD3D12PoolAllocatorRTUAVTextureVRAMPoolSize;
 		uint64 MaxAllocationSize = GD3D12PoolAllocatorRTUAVTextureVRAMMaxAllocationSize;
 		bool bDefragEnabled = true;
-		FD3D12PoolAllocator* UAVPool = new FD3D12PoolAllocator(Device, GetVisibilityMask(), InitConfig, Name, AllocationStrategy, PoolSize, D3D12_DEFAULT_RESOURCE_PLACEMENT_ALIGNMENT, MaxAllocationSize, FreeListOrder, bDefragEnabled);
+		FD3D12PoolAllocator* UAVPool = new FD3D12PoolAllocator(Device, GetVisibilityMask(), InitConfig, Name, AllocationStrategy, PoolSize, D3D12_DEFAULT_RESOURCE_PLACEMENT_ALIGNMENT, MaxAllocationSize, FreeListOrder, bDefragEnabled, TraceHeapId);
 		UAVPool->Initialize();
 
 		PoolAllocators[(int)EPoolType::UAV] = UAVPool;
@@ -1947,7 +1979,7 @@ HRESULT FD3D12TextureAllocatorPool::AllocateTexture(FD3D12ResourceDesc Desc, con
 
 		ID3D12Heap* Heap = nullptr;
 		VERIFYD3D12RESULT(Adapter->GetD3DDevice()->CreateHeap(&HeapDesc, IID_PPV_ARGS(&Heap)));
-		TRefCountPtr<FD3D12Heap> BackingHeap = new FD3D12Heap(GetParentDevice(), GetVisibilityMask());
+		TRefCountPtr<FD3D12Heap> BackingHeap = new FD3D12Heap(GetParentDevice(), GetVisibilityMask(), TraceHeapId);
 		bool bTrack = false;
 		BackingHeap->SetHeap(Heap, Name, bTrack);
 
diff --git a/Engine/Source/Runtime/D3D12RHI/Private/D3D12Allocation.h b/Engine/Source/Runtime/D3D12RHI/Private/D3D12Allocation.h
--- a/Engine/Source/Runtime/D3D12RHI/Private/D3D12Allocation.h
+++ b/Engine/Source/Runtime/D3D12RHI/Private/D3D12Allocation.h
@@ -143,7 +143,8 @@ public:
 		EResourceAllocationStrategy InAllocationStrategy,
 		uint32 MaxSizeForPooling,
 		uint32 InMaxBlockSize,
-		uint32 InMinBlockSize);
+		uint32 InMinBlockSize,
+		HeapId InTraceParentHeapId);
 
 	bool TryAllocate(uint32 SizeInBytes, uint32 Alignment, FD3D12ResourceLocation& ResourceLocation);
 
@@ -204,6 +205,7 @@ private:
 	uint64 LastUsedFrameFence;
 	uint32 MaxOrder;
 	uint32 TotalSizeUsed;
+	HeapId TraceHeapId;
 
 	bool HeapFullMessageDisplayed;
 
@@ -253,7 +255,8 @@ public:
 		EResourceAllocationStrategy InAllocationStrategy,
 		uint32 InMaxAllocationSize,
 		uint32 InDefaultPoolSize,
-		uint32 InMinBlockSize);
+		uint32 InMinBlockSize,
+		HeapId InTraceParentHeapId);
 	~FD3D12MultiBuddyAllocator();
 
 	const EResourceAllocationStrategy GetAllocationStrategy() const { return AllocationStrategy; }
@@ -284,6 +287,9 @@ protected:
 	const uint32 DefaultPoolSize;
 
 	TArray<FD3D12BuddyAllocator*> Allocators;
+
+private:
+	HeapId TraceHeapId;
 };
 
 //-----------------------------------------------------------------------------
@@ -376,7 +382,7 @@ public:
 	void UpdateMemoryStats();
 
 private:
-
+	HeapId						TraceHeapId;
 	// Buddy allocator used for all 'small' allocation - fast but aligns to power of 2
 	FD3D12MultiBuddyAllocator	SmallBlockAllocator;
 	// Pool allocator for all bigger allocations - less fast but less alignment waste
@@ -438,6 +444,7 @@ private:
 
 	TArray<FD3D12BufferPool*> DefaultBufferPools;
 	FCriticalSection CS;
+	HeapId TraceHeapId;
 };
 
 //-----------------------------------------------------------------------------
@@ -925,6 +932,7 @@ private:
 	};
 
 	FD3D12PoolAllocator* PoolAllocators[(int)EPoolType::Count];
+	HeapId TraceHeapId;
 };
 #elif D3D12RHI_SEGREGATED_TEXTURE_ALLOC
 class FD3D12TextureAllocatorPool : public FD3D12DeviceChild, public FD3D12MultiNodeGPUObject
diff --git a/Engine/Source/Runtime/D3D12RHI/Private/D3D12PoolAllocator.cpp b/Engine/Source/Runtime/D3D12RHI/Private/D3D12PoolAllocator.cpp
--- a/Engine/Source/Runtime/D3D12RHI/Private/D3D12PoolAllocator.cpp
+++ b/Engine/Source/Runtime/D3D12RHI/Private/D3D12PoolAllocator.cpp
@@ -30,6 +30,7 @@ FD3D12MemoryPool::FD3D12MemoryPool(
 		, uint32 InPoolAlignment
 		, ERHIPoolResourceTypes InSupportedResourceTypes
 		, EFreeListOrder InFreeListOrder
+		, HeapId InTraceParentHeapId
 	)
 	: FRHIMemoryPool(InPoolIndex, InPoolSize, InPoolAlignment, InSupportedResourceTypes, InFreeListOrder)
 	, FD3D12DeviceChild(ParentDevice)
@@ -39,6 +40,9 @@ FD3D12MemoryPool::FD3D12MemoryPool(
 	, AllocationStrategy(InAllocationStrategy)
 	, LastUsedFrameFence(0)
 {
+#if UE_MEMORY_TRACE_ENABLED
+	TraceHeapId = MemoryTrace_HeapSpec(InTraceParentHeapId, AllocationStrategy == EResourceAllocationStrategy::kPlacedResource ? TEXT("D3D12MemoryPool (PlacedResource)") : TEXT("D3D12MemoryPool (ManualSubAllocation)"));
+#endif
 }
 
 
@@ -94,7 +98,7 @@ void FD3D12MemoryPool::Init()
 			VERIFYD3D12RESULT(Adapter->GetD3DDevice()->CreateHeap(&Desc, IID_PPV_ARGS(&Heap)));
 		}
 
-		BackingHeap = new FD3D12Heap(GetParentDevice(), GetVisibilityMask());
+		BackingHeap = new FD3D12Heap(GetParentDevice(), GetVisibilityMask(), TraceHeapId);
 		BackingHeap->SetHeap(Heap, TEXT("PoolAllocator Heap"));
 
 		// Only track resources that cannot be accessed on the CPU.
@@ -109,6 +113,9 @@ void FD3D12MemoryPool::Init()
 			LLM_SCOPED_PAUSE_TRACKING_FOR_TRACKER(ELLMTracker::Default, ELLMAllocType::System);
 			const D3D12_HEAP_PROPERTIES HeapProps = CD3DX12_HEAP_PROPERTIES(InitConfig.HeapType, GetGPUMask().GetNative(), GetVisibilityMask().GetNative());
 			VERIFYD3D12RESULT(Adapter->CreateBuffer(HeapProps, GetGPUMask(), InitConfig.InitialResourceState, ED3D12ResourceStateMode::SingleState, InitConfig.InitialResourceState, PoolSize, BackingResource.GetInitReference(), TEXT("Resource Allocator Underlying Buffer"), InitConfig.ResourceFlags));
+#if UE_MEMORY_TRACE_ENABLED
+			MemoryTrace_MarkAllocAsHeap(BackingResource->GetGPUVirtualAddress(), TraceHeapId);
+#endif
 		}
 
 		if (IsCPUAccessible(InitConfig.HeapType))
@@ -145,6 +152,9 @@ void FD3D12MemoryPool::Destroy()
 
 	if (BackingResource)
 	{
+#if UE_MEMORY_TRACE_ENABLED
+		MemoryTrace_UnmarkAllocAsHeap(BackingResource->GetGPUVirtualAddress(), TraceHeapId);
+#endif
 		ensure(BackingResource->GetRefCount() == 1 || GNumExplicitGPUsForRendering > 1);
 		BackingResource = nullptr;
 	}
@@ -218,14 +228,18 @@ EResourceAllocationStrategy FD3D12PoolAllocator::GetResourceAllocationStrategy(D
 
 
 FD3D12PoolAllocator::FD3D12PoolAllocator(FD3D12Device* ParentDevice, FRHIGPUMask VisibleNodes, const FD3D12ResourceInitConfig& InInitConfig, const FString& InName,
-	EResourceAllocationStrategy InAllocationStrategy, uint64 InDefaultPoolSize, uint32 InPoolAlignment, uint32 InMaxAllocationSize, FRHIMemoryPool::EFreeListOrder InFreeListOrder, bool bInDefragEnabled) :
+	EResourceAllocationStrategy InAllocationStrategy, uint64 InDefaultPoolSize, uint32 InPoolAlignment, uint32 InMaxAllocationSize, FRHIMemoryPool::EFreeListOrder InFreeListOrder, bool bInDefragEnabled, HeapId InTraceParentHeapId) :
 	FRHIPoolAllocator(InDefaultPoolSize, InPoolAlignment, InMaxAllocationSize, InFreeListOrder, bInDefragEnabled), 
 	FD3D12DeviceChild(ParentDevice), 
 	FD3D12MultiNodeGPUObject(ParentDevice->GetGPUMask(), VisibleNodes), 
 	InitConfig(InInitConfig), 
 	Name(InName), 
 	AllocationStrategy(InAllocationStrategy)
-{}
+{
+#if UE_MEMORY_TRACE_ENABLED
+	TraceHeapId = MemoryTrace_HeapSpec(InTraceParentHeapId, *Name);
+#endif
+}
 
 
 FD3D12PoolAllocator::~FD3D12PoolAllocator()
@@ -409,7 +423,7 @@ void FD3D12PoolAllocator::AllocateResource(uint32 GPUIndex, D3D12_HEAP_TYPE InHe
 
 			ID3D12Heap* Heap = nullptr;
 			VERIFYD3D12RESULT(Adapter->GetD3DDevice()->CreateHeap(&HeapDesc, IID_PPV_ARGS(&Heap)));
-			TRefCountPtr<FD3D12Heap> BackingHeap = new FD3D12Heap(GetParentDevice(), GetVisibilityMask());
+			TRefCountPtr<FD3D12Heap> BackingHeap = new FD3D12Heap(GetParentDevice(), GetVisibilityMask(), TraceHeapId);
 			bool bTrack = false;
 			BackingHeap->SetHeap(Heap, InName, bTrack);		
 			
@@ -544,7 +558,7 @@ FRHIMemoryPool* FD3D12PoolAllocator::CreateNewPool(int16 InPoolIndex, uint32 InM
 	}
 
 	FD3D12MemoryPool* NewPool = new FD3D12MemoryPool(GetParentDevice(),	GetVisibilityMask(), InitConfig,
-		Name, AllocationStrategy, InPoolIndex, PoolSize, PoolAlignment, InAllocationResourceType, FreeListOrder);
+		Name, AllocationStrategy, InPoolIndex, PoolSize, PoolAlignment, InAllocationResourceType, FreeListOrder, TraceHeapId);
 	NewPool->Init();
 	return NewPool;
 }
diff --git a/Engine/Source/Runtime/D3D12RHI/Private/D3D12PoolAllocator.h b/Engine/Source/Runtime/D3D12RHI/Private/D3D12PoolAllocator.h
--- a/Engine/Source/Runtime/D3D12RHI/Private/D3D12PoolAllocator.h
+++ b/Engine/Source/Runtime/D3D12RHI/Private/D3D12PoolAllocator.h
@@ -87,7 +87,7 @@ public:
 
 	// Constructor
 	FD3D12MemoryPool(FD3D12Device* ParentDevice, FRHIGPUMask VisibleNodes, const FD3D12ResourceInitConfig& InInitConfig, const FString& Name,
-		EResourceAllocationStrategy InAllocationStrategy, int16 InPoolIndex, uint64 InPoolSize, uint32 InPoolAlignment, ERHIPoolResourceTypes InSupportedResourceTypes, EFreeListOrder InFreeListOrder);
+		EResourceAllocationStrategy InAllocationStrategy, int16 InPoolIndex, uint64 InPoolSize, uint32 InPoolAlignment, ERHIPoolResourceTypes InSupportedResourceTypes, EFreeListOrder InFreeListOrder, HeapId InTraceParentHeapId);
 	virtual ~FD3D12MemoryPool();
 
 	// Setup/Shutdown
@@ -111,6 +111,9 @@ protected:
 	TRefCountPtr<FD3D12Heap> BackingHeap;
 
 	uint64 LastUsedFrameFence;
+
+private:
+	HeapId TraceHeapId;
 };
 
 
@@ -127,7 +130,7 @@ public:
 
 	// Constructor
 	FD3D12PoolAllocator(FD3D12Device* ParentDevice, FRHIGPUMask VisibleNodes, const FD3D12ResourceInitConfig& InInitConfig, const FString& InName,
-		EResourceAllocationStrategy InAllocationStrategy, uint64 InDefaultPoolSize, uint32 InPoolAlignment, uint32 InMaxAllocationSize, FRHIMemoryPool::EFreeListOrder InFreeListOrder, bool bInDefragEnabled);
+		EResourceAllocationStrategy InAllocationStrategy, uint64 InDefaultPoolSize, uint32 InPoolAlignment, uint32 InMaxAllocationSize, FRHIMemoryPool::EFreeListOrder InFreeListOrder, bool bInDefragEnabled, HeapId InTraceParentHeapId);
 	~FD3D12PoolAllocator();
 	
 	// Function names currently have to match with FD3D12DefaultBufferPool until we can make full replacement of the allocator
@@ -206,5 +209,8 @@ protected:
 
 	// Allocation data pool used to reduce allocation overhead
 	TArray<FRHIPoolAllocationData*> AllocationDataPool;
+
+private:
+	HeapId TraceHeapId;
 };
 
diff --git a/Engine/Source/Runtime/D3D12RHI/Private/D3D12Resources.cpp b/Engine/Source/Runtime/D3D12RHI/Private/D3D12Resources.cpp
--- a/Engine/Source/Runtime/D3D12RHI/Private/D3D12Resources.cpp
+++ b/Engine/Source/Runtime/D3D12RHI/Private/D3D12Resources.cpp
@@ -7,6 +7,7 @@ D3D12Resources.cpp: D3D RHI utility implementation.
 #include "D3D12RHIPrivate.h"
 #include "EngineModule.h"
 #include "HAL/LowLevelMemTracker.h"
+#include "ProfilingDebugging/MemoryTrace.h"
 
 #if INTEL_EXTENSIONS
 	#define INTC_IGDEXT_D3D12 1
@@ -157,15 +158,23 @@ void FD3D12Resource::DeferDelete()
 //	FD3D12 Heap
 /////////////////////////////////////////////////////////////////////
 
-FD3D12Heap::FD3D12Heap(FD3D12Device* Parent, FRHIGPUMask VisibleNodes) :
+FD3D12Heap::FD3D12Heap(FD3D12Device* Parent, FRHIGPUMask VisibleNodes, HeapId InTraceParentHeapId) :
 	FD3D12DeviceChild(Parent),
 	FD3D12MultiNodeGPUObject(Parent->GetGPUMask(), VisibleNodes),
-	ResidencyHandle()
+	ResidencyHandle(),
+	TraceParentHeapId(InTraceParentHeapId)
 {
 }
 
 FD3D12Heap::~FD3D12Heap()
 {
+#if UE_MEMORY_TRACE_ENABLED
+	if (GPUVirtualAddress != 0)
+	{
+		MemoryTrace_UnmarkAllocAsHeap(GPUVirtualAddress, TraceHeapId);
+	}
+#endif	
+
 #if TRACK_RESOURCE_ALLOCATIONS
 	FD3D12Adapter* Adapter = GetParentDevice()->GetParentAdapter();
 	if (GPUVirtualAddress != 0 && bTrack)
@@ -212,7 +221,11 @@ void FD3D12Heap::SetHeap(ID3D12Heap* HeapIn, const TCHAR* const InName, bool bIn
 		const D3D12_RESOURCE_DESC BufDesc = CD3DX12_RESOURCE_DESC::Buffer(HeapSize, D3D12_RESOURCE_FLAG_NONE);
 		VERIFYD3D12RESULT(Adapter->GetD3DDevice()->CreatePlacedResource(Heap, 0, &BufDesc, D3D12_RESOURCE_STATE_COMMON, nullptr, IID_PPV_ARGS(TempResource.GetInitReference())));
 		GPUVirtualAddress = TempResource->GetGPUVirtualAddress();
-				
+#if UE_MEMORY_TRACE_ENABLED
+		TraceHeapId = MemoryTrace_HeapSpec(TraceParentHeapId, *(FString(InName) + TEXT(" D3D12Heap")));
+		MemoryTrace_Alloc(GPUVirtualAddress, HeapSize, HeapDesc.Alignment, EMemoryTraceRootHeap::VideoMemory);
+		MemoryTrace_MarkAllocAsHeap(GPUVirtualAddress, TraceHeapId);
+#endif				
 #if TRACK_RESOURCE_ALLOCATIONS
 		if (bTrack)
 		{
@@ -301,6 +314,16 @@ HRESULT FD3D12Adapter::CreateCommittedResource(const FD3D12ResourceDesc& InDesc,
 		{
 			(*ppOutResource)->StartTrackingForResidency();
 		}
+
+#if UE_MEMORY_TRACE_ENABLED
+		// Calling GetResourceAllocationInfo is not cheap so check memory allocation tracking is enabled
+		if (UE_TRACE_CHANNELEXPR_IS_ENABLED(MemAllocChannel))
+		{
+			const D3D12_RESOURCE_DESC ResourceDesc = (*ppOutResource)->GetResource()->GetDesc();
+			const D3D12_RESOURCE_ALLOCATION_INFO Info = RootDevice->GetResourceAllocationInfo(0, 1, &ResourceDesc);
+			MemoryTrace_Alloc((*ppOutResource)->GetGPUVirtualAddress(), Info.SizeInBytes, Info.Alignment, EMemoryTraceRootHeap::VideoMemory);
+		}
+#endif
 	}
 	else	
 	{
@@ -375,7 +398,15 @@ HRESULT FD3D12Adapter::CreatePlacedResource(const FD3D12ResourceDesc& InDesc, FD
 			}
 		}
 #endif		
-
+#if UE_MEMORY_TRACE_ENABLED
+		// Calling GetResourceAllocationInfo is not cheap so check memory allocation tracking is enabled
+		if (UE_TRACE_CHANNELEXPR_IS_ENABLED(MemAllocChannel))
+		{
+			const D3D12_RESOURCE_DESC ResourceDesc = (*ppOutResource)->GetResource()->GetDesc();
+			const D3D12_RESOURCE_ALLOCATION_INFO Info = RootDevice->GetResourceAllocationInfo(0, 1, &ResourceDesc);
+			MemoryTrace_Alloc((*ppOutResource)->GetGPUVirtualAddress(), Info.SizeInBytes, Info.Alignment, EMemoryTraceRootHeap::VideoMemory);
+		}
+#endif
 		// Set a default name (can override later).
 		SetName(*ppOutResource, Name);
 
@@ -876,6 +907,10 @@ bool FD3D12ResourceLocation::OnAllocationMoved(FRHIPoolAllocationData* InNewData
 	FD3D12Resource* CurrentResource = GetResource();
 	FD3D12PoolAllocator* NewAllocator = GetPoolAllocator();
 
+#if UE_MEMORY_TRACE_ENABLED
+	MemoryTrace_ReallocFree(GPUVirtualAddress, EMemoryTraceRootHeap::VideoMemory);
+#endif
+
 	// If sub allocated and not placed only update the internal data
 	if (NewAllocator->GetAllocationStrategy() == EResourceAllocationStrategy::kManualSubAllocation)
 	{
@@ -883,6 +918,10 @@ bool FD3D12ResourceLocation::OnAllocationMoved(FRHIPoolAllocationData* InNewData
 
 		OffsetFromBaseOfResource = AllocationData.GetOffset();
 		UnderlyingResource = NewAllocator->GetBackingResource(*this);
+
+#if UE_MEMORY_TRACE_ENABLED
+		MemoryTrace_ReallocAlloc(UnderlyingResource->GetGPUVirtualAddress() + OffsetFromBaseOfResource, AllocationData.GetSize(), AllocationData.GetAlignment(), EMemoryTraceRootHeap::VideoMemory);
+#endif
 	}
 	else
 	{
diff --git a/Engine/Source/Runtime/D3D12RHI/Public/D3D12Resources.h b/Engine/Source/Runtime/D3D12RHI/Public/D3D12Resources.h
--- a/Engine/Source/Runtime/D3D12RHI/Public/D3D12Resources.h
+++ b/Engine/Source/Runtime/D3D12RHI/Public/D3D12Resources.h
@@ -71,7 +71,7 @@ public:
 class FD3D12Heap : public FThreadSafeRefCountedObject, public FD3D12DeviceChild, public FD3D12MultiNodeGPUObject
 {
 public:
-	FD3D12Heap(FD3D12Device* Parent, FRHIGPUMask VisibleNodes);
+	FD3D12Heap(FD3D12Device* Parent, FRHIGPUMask VisibleNodes, HeapId InTraceParentHeapId = EMemoryTraceRootHeap::VideoMemory);
 	~FD3D12Heap();
 
 	inline ID3D12Heap* GetHeap() const { return Heap.GetReference(); }
@@ -94,6 +94,8 @@ private:
 	D3D12_HEAP_DESC HeapDesc;
 	D3D12_GPU_VIRTUAL_ADDRESS GPUVirtualAddress = 0;
 	FD3D12ResidencyHandle ResidencyHandle;
+	HeapId TraceHeapId;
+	HeapId TraceParentHeapId;
 };
 
 struct FD3D12ResourceDesc : public D3D12_RESOURCE_DESC
