Fix build errors, fallout from CL24072149

diff --git a/Engine/Source/Runtime/D3D12RHI/Private/D3D12Allocation.cpp b/Engine/Source/Runtime/D3D12RHI/Private/D3D12Allocation.cpp
--- a/Engine/Source/Runtime/D3D12RHI/Private/D3D12Allocation.cpp
+++ b/Engine/Source/Runtime/D3D12RHI/Private/D3D12Allocation.cpp
@@ -1419,7 +1419,8 @@ FD3D12BufferPool* FD3D12DefaultBufferAllocator::CreateBufferPool(D3D12_HEAP_TYPE
 		AllocationStrategy,
 		InHeapType == D3D12_HEAP_TYPE_READBACK ? READBACK_BUFFER_POOL_MAX_ALLOC_SIZE : DEFAULT_BUFFER_POOL_MAX_ALLOC_SIZE,
 		InHeapType == D3D12_HEAP_TYPE_READBACK ? READBACK_BUFFER_POOL_DEFAULT_POOL_SIZE : DEFAULT_BUFFER_POOL_DEFAULT_POOL_SIZE,
-		MinBlockSize
+		MinBlockSize,
+		TraceHeapId
 		);
 
 	FD3D12DefaultBufferPool* NewPool = new FD3D12DefaultBufferPool(Device, Allocator);
@@ -1876,7 +1877,8 @@ FD3D12TextureAllocator::FD3D12TextureAllocator(FD3D12Device* Device,
 	FRHIGPUMask VisibleNodes,
 	const FString& Name,
 	uint32 HeapSize,
-	D3D12_HEAP_FLAGS Flags) :
+	D3D12_HEAP_FLAGS Flags,
+	HeapId InTraceParentHeapId) :
 	FD3D12MultiBuddyAllocator(Device,
 		VisibleNodes,
 		FD3D12ResourceInitConfig
@@ -1890,7 +1892,8 @@ FD3D12TextureAllocator::FD3D12TextureAllocator(FD3D12Device* Device,
 		EResourceAllocationStrategy::kPlacedResource,
 		D3D12_DEFAULT_RESOURCE_PLACEMENT_ALIGNMENT,
 		HeapSize,
-		D3D12_SMALL_RESOURCE_PLACEMENT_ALIGNMENT)
+		D3D12_SMALL_RESOURCE_PLACEMENT_ALIGNMENT,
+		InTraceParentHeapId)
 {
 }
 
@@ -1940,7 +1943,8 @@ HRESULT FD3D12TextureAllocator::AllocateTexture(FD3D12ResourceDesc Desc, const D
 FD3D12TextureAllocatorPool::FD3D12TextureAllocatorPool(FD3D12Device* Device, FRHIGPUMask VisibilityNode) :
 	FD3D12DeviceChild(Device),
 	FD3D12MultiNodeGPUObject(Device->GetGPUMask(), VisibilityNode),
-	ReadOnlyTexturePool(Device, VisibilityNode, FString(L"Small Read-Only Texture allocator"), TEXTURE_POOL_SIZE, D3D12_HEAP_FLAG_ALLOW_ONLY_NON_RT_DS_TEXTURES)
+	TraceHeapId(MemoryTrace_HeapSpec(EMemoryTraceRootHeap::VideoMemory, TEXT("Texture Allocator Pool"))),
+	ReadOnlyTexturePool(Device, VisibilityNode, FString(L"Small Read-Only Texture allocator"), TEXTURE_POOL_SIZE, D3D12_HEAP_FLAG_ALLOW_ONLY_NON_RT_DS_TEXTURES, TraceHeapId)
 {};
 
 HRESULT FD3D12TextureAllocatorPool::AllocateTexture(FD3D12ResourceDesc Desc, const D3D12_CLEAR_VALUE* ClearValue, EPixelFormat UEFormat, FD3D12ResourceLocation& TextureLocation, const D3D12_RESOURCE_STATES InitialState, const TCHAR* Name)
diff --git a/Engine/Source/Runtime/D3D12RHI/Private/D3D12Allocation.h b/Engine/Source/Runtime/D3D12RHI/Private/D3D12Allocation.h
--- a/Engine/Source/Runtime/D3D12RHI/Private/D3D12Allocation.h
+++ b/Engine/Source/Runtime/D3D12RHI/Private/D3D12Allocation.h
@@ -965,7 +965,7 @@ private:
 class FD3D12TextureAllocator : public FD3D12MultiBuddyAllocator
 {
 public:
-	FD3D12TextureAllocator(FD3D12Device* Device, FRHIGPUMask VisibleNodes, const FString& Name, uint32 HeapSize, D3D12_HEAP_FLAGS Flags);
+	FD3D12TextureAllocator(FD3D12Device* Device, FRHIGPUMask VisibleNodes, const FString& Name, uint32 HeapSize, D3D12_HEAP_FLAGS Flags, HeapId InTraceParentHeapId);
 
 	~FD3D12TextureAllocator();
 
@@ -986,6 +986,7 @@ public:
 	bool GetMemoryStats(uint64& OutTotalAllocated, uint64& OutTotalUnused) const { OutTotalAllocated = 0; OutTotalUnused = 0; return false; }
 
 private:
+	HeapId TraceHeapId;
 	FD3D12TextureAllocator ReadOnlyTexturePool;
 };
 #endif
