Use IDXGISwapChain3::GetCurrentBackBufferIndex when it's available.

#rnx

diff --git a/Engine/Source/Runtime/D3D12RHI/Private/D3D12Viewport.cpp b/Engine/Source/Runtime/D3D12RHI/Private/D3D12Viewport.cpp
--- a/Engine/Source/Runtime/D3D12RHI/Private/D3D12Viewport.cpp
+++ b/Engine/Source/Runtime/D3D12RHI/Private/D3D12Viewport.cpp
@@ -716,8 +716,17 @@ bool FD3D12Viewport::Present(bool bLockToVsync)
 	if (bNativelyPresented || (CustomPresent && CustomPresent->NeedsAdvanceBackbuffer()))
 	{
 		// Increment back buffer
-		CurrentBackBufferIndex_RHIThread++;
-		CurrentBackBufferIndex_RHIThread = CurrentBackBufferIndex_RHIThread % NumBackBuffers;
+#if DXGI_MAX_SWAPCHAIN_INTERFACE >= 3
+		if (bNativelyPresented && SwapChain3)
+		{
+			CurrentBackBufferIndex_RHIThread = SwapChain3->GetCurrentBackBufferIndex();
+		}
+		else
+#endif
+		{
+			CurrentBackBufferIndex_RHIThread++;
+			CurrentBackBufferIndex_RHIThread = CurrentBackBufferIndex_RHIThread % NumBackBuffers;
+		}
 		BackBuffer_RHIThread = BackBuffers[CurrentBackBufferIndex_RHIThread].GetReference();
 		SDRBackBuffer_RHIThread = SDRBackBuffers[CurrentBackBufferIndex_RHIThread].GetReference();
 
