Change tracked DX12 buffer address to be GPUVirtualAddress, and fix dynamic buffer resources throwing up "Invalid FREE event " errors in Insight memory trace. RHILockBuffer function reallocates the resource causing the tracked GPUVirtualAddress during allocation and freeing to mismatch.

diff --git a/Engine/Source/Runtime/D3D12RHI/Private/D3D12Buffer.cpp b/Engine/Source/Runtime/D3D12RHI/Private/D3D12Buffer.cpp
--- a/Engine/Source/Runtime/D3D12RHI/Private/D3D12Buffer.cpp
+++ b/Engine/Source/Runtime/D3D12RHI/Private/D3D12Buffer.cpp
@@ -594,12 +594,21 @@ void* FD3D12DynamicRHI::LockBuffer(FRHICommandListBase& RHICmdList, FD3D12Buffer
 			{
 				FRHICommandRenameUploadBuffer* Command = ALLOC_COMMAND_CL(RHICmdList, FRHICommandRenameUploadBuffer)(Buffer, Device);
 				Data = Adapter.GetUploadHeapAllocator(Device->GetGPUIndex()).AllocUploadResource(BufferSize, Buffer->BufferAlignment, Command->NewLocation);
+#if UE_MEMORY_TRACE_ENABLED
+				MemoryTrace_ReallocFree(Buffer->ResourceLocation.GetGPUVirtualAddress(), EMemoryTraceRootHeap::VideoMemory);
+				MemoryTrace_ReallocAlloc(Command->NewLocation.GetGPUVirtualAddress(), BufferSize, Buffer->BufferAlignment, EMemoryTraceRootHeap::VideoMemory);
+#endif
+
 				RHICmdList.RHIThreadFence(true);
 			}
 			else
 			{
 				FRHICommandRenameUploadBuffer Command(Buffer, Device);
 				Data = Adapter.GetUploadHeapAllocator(Device->GetGPUIndex()).AllocUploadResource(BufferSize, Buffer->BufferAlignment, Command.NewLocation);
+#if UE_MEMORY_TRACE_ENABLED
+				MemoryTrace_ReallocFree(Buffer->ResourceLocation.GetGPUVirtualAddress(), EMemoryTraceRootHeap::VideoMemory);
+				MemoryTrace_ReallocAlloc(Command.NewLocation.GetGPUVirtualAddress(), BufferSize, Buffer->BufferAlignment, EMemoryTraceRootHeap::VideoMemory);
+#endif
 				Command.Execute(RHICmdList);
 			}
 		}
diff --git a/Engine/Source/Runtime/D3D12RHI/Private/D3D12Stats.cpp b/Engine/Source/Runtime/D3D12RHI/Private/D3D12Stats.cpp
--- a/Engine/Source/Runtime/D3D12RHI/Private/D3D12Stats.cpp
+++ b/Engine/Source/Runtime/D3D12RHI/Private/D3D12Stats.cpp
@@ -299,6 +299,7 @@ static FName GetD3D12BufferStat(EBufferUsageFlags InUsageFlags)
 
 void UpdateBufferStats(FD3D12Buffer* Buffer, bool bAllocating)
 {
+	D3D12_GPU_VIRTUAL_ADDRESS GPUAddress = Buffer->ResourceLocation.GetGPUVirtualAddress();
 	int64 BufferSize = Buffer->ResourceLocation.GetSize();
 	int64 RequestedSize = bAllocating ? BufferSize : -BufferSize;
 	EBufferUsageFlags UsageFlags = Buffer->GetUsage();
@@ -307,14 +308,16 @@ void UpdateBufferStats(FD3D12Buffer* Buffer, bool bAllocating)
 	INC_MEMORY_STAT_BY_FName(GetD3D12BufferStat(UsageFlags), RequestedSize);
 	INC_MEMORY_STAT_BY(STAT_D3D12MemoryCurrentTotal, RequestedSize);
 
+#if UE_MEMORY_TRACE_ENABLED
 	if (bAllocating)
 	{
-		MemoryTrace_Alloc((uint64)Buffer->GetResource(), BufferSize, 0, EMemoryTraceRootHeap::VideoMemory);
+		MemoryTrace_Alloc(GPUAddress, BufferSize, Buffer->BufferAlignment, EMemoryTraceRootHeap::VideoMemory);
 	}
 	else
 	{
-		MemoryTrace_Free((uint64)Buffer->GetResource(), EMemoryTraceRootHeap::VideoMemory);
+		MemoryTrace_Free(GPUAddress, EMemoryTraceRootHeap::VideoMemory);
 	}
+#endif
 }
 
 #if NV_AFTERMATH
