Fix DX12 default allocator defragmentation not freeing old memory, throwing up "Invalid FREE event " errors in Insight memory trace. For ManualSubAllocation strategy use MemoryTrace_ReallocFree/MemoryTrace_ReallocAlloc pair. For PlacedResource strategy use MemoryTrace_Free in to match the MemoryTrace_Alloc call from inside CreatePlacedResource function.

diff --git a/Engine/Source/Runtime/D3D12RHI/Private/D3D12Resources.cpp b/Engine/Source/Runtime/D3D12RHI/Private/D3D12Resources.cpp
--- a/Engine/Source/Runtime/D3D12RHI/Private/D3D12Resources.cpp
+++ b/Engine/Source/Runtime/D3D12RHI/Private/D3D12Resources.cpp
@@ -916,11 +916,15 @@ bool FD3D12ResourceLocation::OnAllocationMoved(FRHIPoolAllocationData* InNewData
 	{
 		check(!CurrentResource->IsPlacedResource());
 
+#if UE_MEMORY_TRACE_ENABLED
+		MemoryTrace_ReallocFree(GPUVirtualAddress, EMemoryTraceRootHeap::VideoMemory);
+#endif
 		OffsetFromBaseOfResource = AllocationData.GetOffset();
 		UnderlyingResource = NewAllocator->GetBackingResource(*this);
+		GPUVirtualAddress = UnderlyingResource->GetGPUVirtualAddress() + OffsetFromBaseOfResource;
 
 #if UE_MEMORY_TRACE_ENABLED
-		MemoryTrace_ReallocAlloc(UnderlyingResource->GetGPUVirtualAddress() + OffsetFromBaseOfResource, AllocationData.GetSize(), AllocationData.GetAlignment(), EMemoryTraceRootHeap::VideoMemory);
+		MemoryTrace_ReallocAlloc(GPUVirtualAddress, AllocationData.GetSize(), AllocationData.GetAlignment(), EMemoryTraceRootHeap::VideoMemory);
 #endif
 	}
 	else
@@ -928,6 +932,10 @@ bool FD3D12ResourceLocation::OnAllocationMoved(FRHIPoolAllocationData* InNewData
 		check(CurrentResource->IsPlacedResource());
 		check(OffsetFromBaseOfResource == 0);
 
+#if UE_MEMORY_TRACE_ENABLED
+		// CreatePlacedResource function below calls MemoryTrace_Alloc to track new memory, so call MemoryTrace_Free to match (instead of calling MemoryTrace_ReallocFree/MemoryTrace_ReallocAlloc).
+		MemoryTrace_Free(GPUVirtualAddress, EMemoryTraceRootHeap::VideoMemory);
+#endif
 		// recreate the placed resource (ownership of current resource is already handled during the internal move)
 		FD3D12HeapAndOffset HeapAndOffset = NewAllocator->GetBackingHeapAndAllocationOffsetInBytes(*this);
 
@@ -955,9 +963,9 @@ bool FD3D12ResourceLocation::OnAllocationMoved(FRHIPoolAllocationData* InNewData
 		VERIFYD3D12RESULT(CurrentResource->GetParentDevice()->GetParentAdapter()->CreatePlacedResource(CurrentResource->GetDesc(), HeapAndOffset.Heap, HeapAndOffset.Offset, CreateState, ResourceStateMode, D3D12_RESOURCE_STATE_TBD, ClearValue, &NewResource, *Name.ToString()));
 
 		UnderlyingResource = NewResource;
+		GPUVirtualAddress = UnderlyingResource->GetGPUVirtualAddress() + OffsetFromBaseOfResource;
 	}
-
-	GPUVirtualAddress = UnderlyingResource->GetGPUVirtualAddress() + OffsetFromBaseOfResource;
+	
 	ResidencyHandle = &UnderlyingResource->GetResidencyHandle();
 
 	// Refresh aliases
