Remove MemoryTrace_ReallocFree that should have been removed in CL24209975

diff --git a/Engine/Source/Runtime/D3D12RHI/Private/D3D12Resources.cpp b/Engine/Source/Runtime/D3D12RHI/Private/D3D12Resources.cpp
--- a/Engine/Source/Runtime/D3D12RHI/Private/D3D12Resources.cpp
+++ b/Engine/Source/Runtime/D3D12RHI/Private/D3D12Resources.cpp
@@ -907,10 +907,6 @@ bool FD3D12ResourceLocation::OnAllocationMoved(FRHIPoolAllocationData* InNewData
 	FD3D12Resource* CurrentResource = GetResource();
 	FD3D12PoolAllocator* NewAllocator = GetPoolAllocator();
 
-#if UE_MEMORY_TRACE_ENABLED
-	MemoryTrace_ReallocFree(GPUVirtualAddress, EMemoryTraceRootHeap::VideoMemory);
-#endif
-
 	// If sub allocated and not placed only update the internal data
 	if (NewAllocator->GetAllocationStrategy() == EResourceAllocationStrategy::kManualSubAllocation)
 	{
