Implement limited support for reserved resource creation in D3D12 RHI

* Add GRHISupportsReservedResources to query support at run-time
* Add TexCreate_ReservedResource to request a texture to be created as a reserved resource
* Add TexCreate_ImmediateCommit to allocate and commit physical memory to a reserved resource on creation
** This is intended as a temporary measure, until an explicit API for memory allocation and tile mapping can be implemented
** CVar d3d12.ReservedResourceHeapSizeMB can be used to control the size of heaps used for ImmediateCommit mode

The primary purpose of this CL is to allow splitting very large resource allocations into smaller physical chunks. This is useful to improve performance under high memory usage / over-commitment conditions on Windows, where some allocations may be paged out to system memory by the kernel.
Individual allocations are paged in/out as a whole and must be contiguous in physical memory space. Under high memory usage conditions, Windows video memory manager (VidMm) may not be able to page large resources back into local memory due to fragmentation.
This can lead to a significant performance degradation, especially when a paged-out resource is accessed by the GPU for writing. Splitting large allocations into smaller physical chunks allows VidMm to perform paging at finer granularity, reducing the chance of a large performance-critical resource getting entirely paged out.

diff --git a/Engine/Source/Runtime/D3D12RHI/Private/D3D12Adapter.h b/Engine/Source/Runtime/D3D12RHI/Private/D3D12Adapter.h
--- a/Engine/Source/Runtime/D3D12RHI/Private/D3D12Adapter.h
+++ b/Engine/Source/Runtime/D3D12RHI/Private/D3D12Adapter.h
@@ -307,6 +307,16 @@ public:
 		const TCHAR* Name,
 		bool bVerifyHResult = true);
 
+	HRESULT CreateReservedResource(const FD3D12ResourceDesc& Desc,
+		FRHIGPUMask CreationNode,
+		D3D12_RESOURCE_STATES InInitialState,
+		ED3D12ResourceStateMode InResourceStateMode,
+		D3D12_RESOURCE_STATES InDefaultState,
+		const D3D12_CLEAR_VALUE* ClearValue,
+		FD3D12Resource** ppOutResource,
+		const TCHAR* Name,
+		bool bVerifyHResult = true);
+
 	HRESULT CreateBuffer(D3D12_HEAP_TYPE HeapType,
 		FRHIGPUMask CreationNode,
 		FRHIGPUMask VisibleNodes,
diff --git a/Engine/Source/Runtime/D3D12RHI/Private/D3D12Resources.cpp b/Engine/Source/Runtime/D3D12RHI/Private/D3D12Resources.cpp
--- a/Engine/Source/Runtime/D3D12RHI/Private/D3D12Resources.cpp
+++ b/Engine/Source/Runtime/D3D12RHI/Private/D3D12Resources.cpp
@@ -17,6 +17,13 @@ D3D12Resources.cpp: D3D RHI utility implementation.
 	THIRD_PARTY_INCLUDES_END
 #endif
 
+static TAutoConsoleVariable<int32> CVarD3D12ReservedResourceHeapSizeMB(
+	TEXT("d3d12.ReservedResourceHeapSizeMB"),
+	16,
+	TEXT("Size of the backing heaps for reserved resources in megabytes (default 16MB)."),
+	ECVF_ReadOnly
+);
+
 /////////////////////////////////////////////////////////////////////
 //	ID3D12ResourceAllocator
 /////////////////////////////////////////////////////////////////////
@@ -99,6 +106,12 @@ FD3D12Resource::FD3D12Resource(FD3D12Device* ParentDevice,
 		GFSDK_Aftermath_DX12_RegisterResource(InResource, &AftermathHandle);
 	}
 #endif
+
+	if (Desc.bReservedResource)
+	{
+		checkf(Heap == nullptr, TEXT("Reserved resources are not expected to have a heap"));
+		ReservedResourceData = MakeUnique<FD3D12ReservedResourceData>();
+	}
 }
 
 FD3D12Resource::~FD3D12Resource()
@@ -124,6 +137,119 @@ FD3D12Resource::~FD3D12Resource()
 	}
 }
 
+void FD3D12Resource::CommitReservedResource()
+{
+	check(Desc.bReservedResource);
+	check(ReservedResourceData.IsValid());
+	checkf(ReservedResourceData->BackingHeaps.IsEmpty(), TEXT("Reserved resource is already committed"));
+	checkf(Desc.Dimension == D3D12_RESOURCE_DIMENSION_TEXTURE2D, TEXT("CommitReservedResource is currently only implemented for 2D textures"));
+	checkf(Desc.MipLevels == 1, TEXT("CommitReservedResource is currently only implemented for textures without mips"));
+
+	uint32 NumTiles = 0;
+	D3D12_PACKED_MIP_INFO PackedMipDesc = {};
+	D3D12_TILE_SHAPE TileShape = {};
+	const uint32 FirstSubresource = 0;
+
+	// We assume that all subresources in a 2D texture array are identical, so only query the tiling config for the first
+	uint32 NumSubresourceTilings = 1;
+	D3D12_SUBRESOURCE_TILING SubresourceTiling = {}; 
+
+	ID3D12Device* D3DDevice = GetParentDevice()->GetDevice();
+	FD3D12Adapter* Adapter = GetParentDevice()->GetParentAdapter();
+
+	D3DDevice->GetResourceTiling(GetResource(), &NumTiles, &PackedMipDesc, &TileShape, &NumSubresourceTilings, FirstSubresource, &SubresourceTiling);
+
+	const uint64 TileSizeInBytes = 65536; // reserved resource tiles are always 64KB
+	const uint64 TotalSize = NumTiles * TileSizeInBytes;
+	const uint64 MaxHeapSize = uint64(CVarD3D12ReservedResourceHeapSizeMB.GetValueOnAnyThread()) * 1024 * 1024;
+	const uint64 NumHeaps = FMath::DivideAndRoundUp(TotalSize, MaxHeapSize);
+
+	TArray<TRefCountPtr<ID3D12Heap>>& Heaps = ReservedResourceData->BackingHeaps;
+	Heaps.Reserve(NumHeaps);
+
+	const uint32 NumTilesPerHeap = uint32(MaxHeapSize / TileSizeInBytes);
+
+	// NOTE: Accessing the queue from this thread is OK, as D3D12 runtime acquires a lock around all command queue APIs.
+	// https://microsoft.github.io/DirectX-Specs/d3d/CPUEfficiency.html#threading
+	FD3D12Queue& Queue = GetParentDevice()->GetQueue(ED3D12QueueType::Direct);
+	ID3D12CommandQueue* D3DCommandQueue = Queue.D3DCommandQueue.GetReference();
+
+	const D3D12_TILE_MAPPING_FLAGS MappingFlags = D3D12_TILE_MAPPING_FLAG_NONE;
+
+	// Set high residency priority based on the same heuristics as D3D12 committed resources,
+	// i.e. normal priority unless it's a UAV/RT/DS texture.
+	const bool bHighPriorityResource = EnumHasAnyFlags(Desc.Flags, D3D12_RESOURCE_FLAG_ALLOW_RENDER_TARGET | D3D12_RESOURCE_FLAG_ALLOW_DEPTH_STENCIL | D3D12_RESOURCE_FLAG_ALLOW_UNORDERED_ACCESS);
+	const uint32 GPUIndex = GetParentDevice()->GetGPUIndex();
+
+	D3D12_HEAP_PROPERTIES BackingHeapProps = {};
+	BackingHeapProps.Type = D3D12_HEAP_TYPE_DEFAULT;
+	BackingHeapProps.CPUPageProperty = D3D12_CPU_PAGE_PROPERTY_UNKNOWN;
+	BackingHeapProps.MemoryPoolPreference = D3D12_MEMORY_POOL_UNKNOWN;
+	BackingHeapProps.CreationNodeMask = GetGPUMask().GetNative();
+	BackingHeapProps.VisibleNodeMask = GetVisibilityMask().GetNative();
+
+	const uint32 NumTilesX = SubresourceTiling.WidthInTiles;
+	const uint32 NumTilesY = SubresourceTiling.HeightInTiles;
+
+	check(SubresourceTiling.DepthInTiles == 1);
+
+	const uint32 NumTilesPerSubresource = NumTilesX * NumTilesY;
+	const uint32 NumTotalTiles = NumTilesPerSubresource * SubresourceCount;
+
+	uint32 NumMappedTiles = 0;
+
+	while (NumMappedTiles < NumTotalTiles)
+	{
+		const uint32 NumRemainingTiles = NumTotalTiles - NumMappedTiles;
+
+		D3D12_TILE_REGION_SIZE RegionSize = {};
+		RegionSize.UseBox = false;
+		RegionSize.NumTiles = FMath::Min(NumTilesPerHeap, NumRemainingTiles);
+
+		const uint32 ThisHeapSize = RegionSize.NumTiles * TileSizeInBytes;
+		D3D12_HEAP_DESC NewHeapDesc = {};
+		NewHeapDesc.Alignment = D3D12_DEFAULT_RESOURCE_PLACEMENT_ALIGNMENT;
+		NewHeapDesc.Flags = D3D12_HEAP_FLAG_NONE;
+		NewHeapDesc.SizeInBytes = ThisHeapSize;
+		NewHeapDesc.Properties = BackingHeapProps;
+		TRefCountPtr<ID3D12Heap> NewHeap;
+		VERIFYD3D12RESULT(D3DDevice->CreateHeap(&NewHeapDesc, IID_PPV_ARGS(NewHeap.GetInitReference())));
+
+		if (bHighPriorityResource)
+		{
+			Adapter->SetResidencyPriority(NewHeap, D3D12_RESIDENCY_PRIORITY_HIGH, GPUIndex);
+		}
+
+		D3D12_TILED_RESOURCE_COORDINATE ResourceCoordinate = {}; // Coordinates are in tiles, not pixels
+		const uint32 TileIndexInSubresource = NumMappedTiles % NumTilesPerSubresource;
+		ResourceCoordinate.Subresource = NumMappedTiles / NumTilesPerSubresource;
+		ResourceCoordinate.X = TileIndexInSubresource % NumTilesX;
+		ResourceCoordinate.Y = (TileIndexInSubresource / NumTilesX) % NumTilesY;
+		ResourceCoordinate.Z = 0; // Only simple 2D / Array2D textures are implemented
+
+		const D3D12_TILE_RANGE_FLAGS RangeFlags = D3D12_TILE_RANGE_FLAG_NONE;
+		const uint32 HeapRangeStartOffsetInTiles = 0;
+		const uint32 RangeTileCount = RegionSize.NumTiles;
+
+		D3DCommandQueue->UpdateTileMappings(GetResource(), 1,
+			&ResourceCoordinate, &RegionSize, NewHeap.GetReference(),
+			1,
+			&RangeFlags,
+			&HeapRangeStartOffsetInTiles,
+			&RangeTileCount,
+			MappingFlags);
+
+	#if NAME_OBJECTS
+		const int32 HeapIndex = Heaps.Num();
+		FString HeapName = FString::Printf(TEXT("%s.Heap[%d]"), DebugName.IsValid() ? *DebugName.ToString() : TEXT("UNKNOWN"), HeapIndex);
+		::SetName(NewHeap, *HeapName);
+	#endif // NAME_OBJECTS
+
+		NumMappedTiles += RegionSize.NumTiles;
+		Heaps.Add(MoveTemp(NewHeap));
+	}
+}
+
 ID3D12Pageable* FD3D12Resource::GetPageable()
 {
 	if (IsPlacedResource())
@@ -339,6 +465,58 @@ HRESULT FD3D12Adapter::CreateCommittedResource(const FD3D12ResourceDesc& InDesc,
 	return hr;
 }
 
+HRESULT FD3D12Adapter::CreateReservedResource(const FD3D12ResourceDesc& InDesc, FRHIGPUMask CreationNode, D3D12_RESOURCE_STATES InInitialState,
+	ED3D12ResourceStateMode InResourceStateMode, D3D12_RESOURCE_STATES InDefaultState, const D3D12_CLEAR_VALUE* ClearValue, FD3D12Resource** ppOutResource, const TCHAR* Name, bool bVerifyHResult)
+{
+	if (!ppOutResource)
+	{
+		return E_POINTER;
+	}
+
+	TRACE_CPUPROFILER_EVENT_SCOPE(CreateReservedResource);
+
+	LLM_PLATFORM_SCOPE(ELLMTag::GraphicsPlatform);
+
+	TRefCountPtr<ID3D12Resource> pResource;
+
+	FD3D12ResourceDesc LocalDesc = InDesc;
+	LocalDesc.Layout = D3D12_TEXTURE_LAYOUT_64KB_UNDEFINED_SWIZZLE;
+	LocalDesc.bReservedResource = true;
+
+#if D3D12_RHI_RAYTRACING
+	if (InDefaultState == D3D12_RESOURCE_STATE_RAYTRACING_ACCELERATION_STRUCTURE)
+	{
+		LocalDesc.Flags |= D3D12_RESOURCE_FLAG_ALLOW_UNORDERED_ACCESS;
+	}
+#endif // D3D12_RHI_RAYTRACING
+
+	HRESULT hr = RootDevice->CreateReservedResource(&LocalDesc, InInitialState, ClearValue, IID_PPV_ARGS(pResource.GetInitReference()));
+
+	if (SUCCEEDED(hr))
+	{
+		// Set the output pointer
+		*ppOutResource = new FD3D12Resource(GetDevice(CreationNode.ToIndex()), CreationNode, pResource, InInitialState, InResourceStateMode, InDefaultState, LocalDesc, nullptr /*Heap*/, D3D12_HEAP_TYPE_DEFAULT);
+		(*ppOutResource)->AddRef();
+
+		// Set a default name (can override later).
+		SetName(*ppOutResource, Name);
+		
+		// NOTE: reserved resource residency is not tracked/managed by the engine, so we don't need to call StartTrackingForResidency().
+	}
+	else
+	{
+		UE_LOG(LogD3D12RHI, Display, TEXT("D3D12 CreateReservedResource failed with params:\n\tResource Dimension: %d\n\tResource Width: %d\n\tResource Height: %d\n\tFormat: %d\n\tResource Flags: %d"),
+			LocalDesc.Dimension, LocalDesc.Width, LocalDesc.Height, LocalDesc.PixelFormat, LocalDesc.Flags);
+
+		if (bVerifyHResult)
+		{
+			VERIFYD3D12RESULT_EX(hr, RootDevice);
+		}
+	}
+
+	return hr;
+}
+
 HRESULT FD3D12Adapter::CreatePlacedResource(const FD3D12ResourceDesc& InDesc, FD3D12Heap* BackingHeap, uint64 HeapOffset, D3D12_RESOURCE_STATES InInitialState, ED3D12ResourceStateMode InResourceStateMode,
 	D3D12_RESOURCE_STATES InDefaultState, const D3D12_CLEAR_VALUE* ClearValue, FD3D12Resource** ppOutResource, const TCHAR* Name, bool bVerifyHResult)
 {
diff --git a/Engine/Source/Runtime/D3D12RHI/Private/D3D12Texture.cpp b/Engine/Source/Runtime/D3D12RHI/Private/D3D12Texture.cpp
--- a/Engine/Source/Runtime/D3D12RHI/Private/D3D12Texture.cpp
+++ b/Engine/Source/Runtime/D3D12RHI/Private/D3D12Texture.cpp
@@ -754,6 +754,15 @@ bool FD3D12DynamicRHI::RHIGetTextureMemoryVisualizeData(FColor* /*TextureData*/,
 	return false;
 }
 
+static bool ShouldCreateReservedTexture(const FD3D12ResourceDesc& TextureDesc, ETextureCreateFlags Flags)
+{
+	// Only 2D textures without mips are implemented/supported, to avoid the complexity associated with packed mips
+	return GRHISupportsReservedResources
+		&& EnumHasAllFlags(Flags, TexCreate_ReservedResource)
+		&& TextureDesc.MipLevels == 1
+		&& TextureDesc.Dimension == D3D12_RESOURCE_DIMENSION_TEXTURE2D;
+}
+
 /**
  * Creates a 2D texture optionally guarded by a structured exception handler.
  */
@@ -792,7 +801,27 @@ void SafeCreateTexture2D(FD3D12Device* pDevice,
 
 		case D3D12_HEAP_TYPE_DEFAULT:
 		{
-			VERIFYD3D12CREATETEXTURERESULT(pDevice->GetTextureAllocator().AllocateTexture(TextureDesc, ClearValue, Format, *OutTexture2D, InitialState, Name), TextureDesc, pDevice->GetDevice());
+			if (ShouldCreateReservedTexture(TextureDesc, Flags))
+			{
+				FD3D12Resource* Resource = nullptr;
+				VERIFYD3D12CREATETEXTURERESULT(
+					Adapter->CreateReservedResource(TextureDesc, pDevice->GetGPUMask(), InitialState, ED3D12ResourceStateMode::MultiState, InitialState, ClearValue, &Resource, Name, false),
+					TextureDesc, pDevice->GetDevice());
+
+				D3D12_RESOURCE_ALLOCATION_INFO AllocInfo = pDevice->GetResourceAllocationInfo(TextureDesc);
+
+				OutTexture2D->AsStandAlone(Resource, AllocInfo.SizeInBytes);
+
+				if (EnumHasAllFlags(Flags, TexCreate_ImmediateCommit))
+				{
+					Resource->CommitReservedResource();
+				}
+			}
+			else
+			{
+				VERIFYD3D12CREATETEXTURERESULT(pDevice->GetTextureAllocator().AllocateTexture(TextureDesc, ClearValue, Format, *OutTexture2D, InitialState, Name), TextureDesc, pDevice->GetDevice());
+			}
+
 			OutTexture2D->SetOwner(Owner);
 			break;
 		}
diff --git a/Engine/Source/Runtime/D3D12RHI/Private/Windows/WindowsD3D12Device.cpp b/Engine/Source/Runtime/D3D12RHI/Private/Windows/WindowsD3D12Device.cpp
--- a/Engine/Source/Runtime/D3D12RHI/Private/Windows/WindowsD3D12Device.cpp
+++ b/Engine/Source/Runtime/D3D12RHI/Private/Windows/WindowsD3D12Device.cpp
@@ -1628,6 +1628,16 @@ void FD3D12DynamicRHI::Init()
 		UE_LOG(LogD3D12RHI, Log, TEXT("RHI does not have support for 64 bit atomics"));
 	}
 
+	// Detect reserved resource support
+	{
+		D3D12_FEATURE_DATA_D3D12_OPTIONS Options = {};
+		if (SUCCEEDED(GetAdapter().GetD3DDevice()->CheckFeatureSupport(D3D12_FEATURE_D3D12_OPTIONS, &Options, sizeof(Options))))
+		{
+			// Tier 2 is guaranteed for all adapters with feature level 12_0.
+			GRHISupportsReservedResources = Options.TiledResourcesTier >= D3D12_TILED_RESOURCES_TIER_2;
+		}
+	}
+
 	D3D12_FEATURE_DATA_D3D12_OPTIONS6 options = {};
 	HRESULT Options6HR = GetAdapter().GetD3DDevice()->CheckFeatureSupport(D3D12_FEATURE_D3D12_OPTIONS6, &options, sizeof(options));
 
diff --git a/Engine/Source/Runtime/D3D12RHI/Public/D3D12Resources.h b/Engine/Source/Runtime/D3D12RHI/Public/D3D12Resources.h
--- a/Engine/Source/Runtime/D3D12RHI/Public/D3D12Resources.h
+++ b/Engine/Source/Runtime/D3D12RHI/Public/D3D12Resources.h
@@ -14,6 +14,7 @@
 #include "D3D12State.h"
 #include "D3D12DirectCommandListManager.h"
 #include "RHIPoolAllocator.h"
+#include "Templates/UniquePtr.h"
 
 constexpr D3D12_RESOURCE_STATES BackBufferBarrierWriteTransitionTargets = D3D12_RESOURCE_STATES(
 	uint32(D3D12_RESOURCE_STATE_RENDER_TARGET) |
@@ -121,6 +122,8 @@ struct FD3D12ResourceDesc : public D3D12_RESOURCE_DESC
 	bool bRequires64BitAtomicSupport{ false };
 #endif
 
+	bool bReservedResource { false };
+
 	// Used primarily to help treat this resource description as writable.
 	inline bool NeedsUAVAliasWorkarounds() const { return UAVAliasPixelFormat != PF_Unknown; }
 };
@@ -169,6 +172,12 @@ private:
 	static int64 NoStateTrackingResourceCount;
 #endif
 
+	struct FD3D12ReservedResourceData
+	{
+		TArray<TRefCountPtr<ID3D12Heap>> BackingHeaps;
+	};
+	TUniquePtr<FD3D12ReservedResourceData> ReservedResourceData;
+
 public:
 	FD3D12Resource() = delete;
 	explicit FD3D12Resource(FD3D12Device* ParentDevice,
@@ -362,6 +371,8 @@ public:
 		const uint32 bReadBackResource : 1;
 	};
 
+	void CommitReservedResource();
+
 private:
 	void InitalizeResourceState(D3D12_RESOURCE_STATES InInitialState, ED3D12ResourceStateMode InResourceStateMode, D3D12_RESOURCE_STATES InDefaultState)
 	{
diff --git a/Engine/Source/Runtime/RHI/Private/RHI.cpp b/Engine/Source/Runtime/RHI/Private/RHI.cpp
--- a/Engine/Source/Runtime/RHI/Private/RHI.cpp
+++ b/Engine/Source/Runtime/RHI/Private/RHI.cpp
@@ -1398,6 +1398,8 @@ bool GRHISupportsSeparateDepthStencilCopyAccess = true;
 ERHIBindlessSupport GRHIBindlessSupport = ERHIBindlessSupport::Unsupported;
 bool GRHISupportsBindless = false;
 
+bool GRHISupportsReservedResources = false;
+
 RHI_API int32 volatile GCurrentTextureMemorySize = 0;
 RHI_API int32 volatile GCurrentRendertargetMemorySize = 0;
 RHI_API int64 GTexturePoolSize = 0 * 1024 * 1024;
diff --git a/Engine/Source/Runtime/RHI/Public/RHI.h b/Engine/Source/Runtime/RHI/Public/RHI.h
--- a/Engine/Source/Runtime/RHI/Public/RHI.h
+++ b/Engine/Source/Runtime/RHI/Public/RHI.h
@@ -669,6 +669,9 @@ extern RHI_API ERHIBindlessSupport GRHIBindlessSupport;
 UE_DEPRECATED(5.2, "You must use GRHIBindlessSupport instead.")
 extern RHI_API bool GRHISupportsBindless;
 
+/** True if the RHI supports reserved (AKA tiled, virtual or sparse) resources and operations related to them*/
+extern RHI_API bool GRHISupportsReservedResources;
+
 // Calculate the index of the sample in GRHIDefaultMSAASampleOffsets
 extern RHI_API int32 CalculateMSAASampleArrayIndex(int32 NumSamples, int32 SampleIndex);
 
diff --git a/Engine/Source/Runtime/RHI/Public/RHIDefinitions.h b/Engine/Source/Runtime/RHI/Public/RHIDefinitions.h
--- a/Engine/Source/Runtime/RHI/Public/RHIDefinitions.h
+++ b/Engine/Source/Runtime/RHI/Public/RHIDefinitions.h
@@ -962,6 +962,10 @@ enum class ETextureCreateFlags : uint64
 	External                		  = 1ull << 34,
 	/** Don't automatically transfer across GPUs in multi-GPU scenarios.  For example, if you are transferring it yourself manually. */
 	MultiGPUGraphIgnore				  = 1ull << 35,
+	/** EXPERIMENTAL: Allow the texture to be created as a reserved (AKA tiled/sparse/virtual) resource internally, without physical memory backing. */
+	ReservedResource                  = 1ull << 37,
+	/** EXPERIMENTAL: Used with ReservedResource flag to immediately allocate and commit memory on creation. May use N small physical memory allocations instead of a single large one. */
+	ImmediateCommit                   = 1ull << 38,
 };
 ENUM_CLASS_FLAGS(ETextureCreateFlags);
 
@@ -1002,6 +1006,8 @@ ENUM_CLASS_FLAGS(ETextureCreateFlags);
 #define TexCreate_AtomicCompatible               ETextureCreateFlags::AtomicCompatible
 #define TexCreate_External               		 ETextureCreateFlags::External
 #define TexCreate_MultiGPUGraphIgnore            ETextureCreateFlags::MultiGPUGraphIgnore
+#define TexCreate_ReservedResource               ETextureCreateFlags::ReservedResource
+#define TexCreate_ImmediateCommit                ETextureCreateFlags::ImmediateCommit
 
 enum EAsyncComputePriority
 {
