Fix for missing transition to depth read.
The issue can currently happen when CullingContext.Configuration.bTwoPassOcclusion is 0 and UseComputeDepthExport is true: in this case, SceneDepthZ will transition to UAV first, but won't be used as an SRV with BuildPreviousOccluderHZB. Instead we try to bind it as a read only depth buffer during BeginOcclusionTests  and we were keeping it in the UAV state.

diff --git a/Engine/Source/Runtime/D3D12RHI/Private/D3D12CommandList.cpp b/Engine/Source/Runtime/D3D12RHI/Private/D3D12CommandList.cpp
--- a/Engine/Source/Runtime/D3D12RHI/Private/D3D12CommandList.cpp
+++ b/Engine/Source/Runtime/D3D12RHI/Private/D3D12CommandList.cpp
@@ -381,11 +381,19 @@ void FD3D12ContextCommon::TransitionResource(FD3D12DepthStencilView* View)
 	{
 		TransitionResource(Resource, D3D12_RESOURCE_STATE_TBD, D3D12_RESOURCE_STATE_DEPTH_WRITE, View->GetDepthOnlyViewSubresourceSubset());
 	}
+	else if (bHasDepth)
+	{
+		TransitionResource(Resource, D3D12_RESOURCE_STATE_TBD, D3D12_RESOURCE_STATE_DEPTH_READ, View->GetDepthOnlyViewSubresourceSubset());
+	}
 
 	if (bStencilIsWritable)
 	{
 		TransitionResource(Resource, D3D12_RESOURCE_STATE_TBD, D3D12_RESOURCE_STATE_DEPTH_WRITE, View->GetStencilOnlyViewSubresourceSubset());
 	}
+	else if (bHasStencil)
+	{
+		TransitionResource(Resource, D3D12_RESOURCE_STATE_TBD, D3D12_RESOURCE_STATE_DEPTH_READ, View->GetStencilOnlyViewSubresourceSubset());
+	}
 }
 
 void FD3D12ContextCommon::TransitionResource(FD3D12DepthStencilView* View, D3D12_RESOURCE_STATES After)
@@ -613,6 +621,7 @@ static inline bool IsTransitionNeeded(D3D12_RESOURCE_STATES Before, D3D12_RESOUR
 	// Depth write is actually a suitable for read operations as a "normal" depth buffer.
 	if (Before == D3D12_RESOURCE_STATE_DEPTH_WRITE && After == D3D12_RESOURCE_STATE_DEPTH_READ)
 	{
+		After = D3D12_RESOURCE_STATE_DEPTH_WRITE;
 		return false;
 	}
 
