Fixed a few Insight video memory trace errors:
- Fixed fast allocator page pool not freed
- Fixed dynamic resource Lock/Unlock function tracking alloc/free in missteps due to the operation gets deferred to RHI thread
- Fixed freeing heap and its sub resources in incorrect order
- Fixed memory allocation get tracked twice when buffer or texture is created as committed standalone resource
- Fixed memory allocation get tracked twice when buffer or texture is created as sub-allocated placed resource from pool allocator
- Fixed texture tracked with incorrect pointers
- Skip transient resource tracking

diff --git a/Engine/Source/Runtime/D3D12RHI/Private/D3D12Allocation.cpp b/Engine/Source/Runtime/D3D12RHI/Private/D3D12Allocation.cpp
--- a/Engine/Source/Runtime/D3D12RHI/Private/D3D12Allocation.cpp
+++ b/Engine/Source/Runtime/D3D12RHI/Private/D3D12Allocation.cpp
@@ -2161,6 +2161,14 @@ FD3D12FastAllocatorPage* FD3D12FastAllocatorPagePool::RequestFastAllocatorPage()
 	return Page;
 }
 
+FD3D12FastAllocatorPage::~FD3D12FastAllocatorPage()
+{
+#if UE_MEMORY_TRACE_ENABLED
+	// Matches MemoryTrace_Alloc issued from CreateBuffer call inside RequestFastAllocatorPage() function
+	MemoryTrace_Free(FastAllocBuffer->GetGPUVirtualAddress(), EMemoryTraceRootHeap::VideoMemory);
+#endif
+}
+
 void FD3D12FastAllocatorPage::UpdateFence()
 {
 	// Fence value must be updated every time the page is used to service an allocation.
diff --git a/Engine/Source/Runtime/D3D12RHI/Private/D3D12Allocation.h b/Engine/Source/Runtime/D3D12RHI/Private/D3D12Allocation.h
--- a/Engine/Source/Runtime/D3D12RHI/Private/D3D12Allocation.h
+++ b/Engine/Source/Runtime/D3D12RHI/Private/D3D12Allocation.h
@@ -465,6 +465,8 @@ struct FD3D12FastAllocatorPage
 		, FastAllocData(nullptr)
 		, FrameFence(0) {};
 
+	~FD3D12FastAllocatorPage();
+
 	void Reset()
 	{
 		NextFastAllocOffset = 0;
diff --git a/Engine/Source/Runtime/D3D12RHI/Private/D3D12Buffer.cpp b/Engine/Source/Runtime/D3D12RHI/Private/D3D12Buffer.cpp
--- a/Engine/Source/Runtime/D3D12RHI/Private/D3D12Buffer.cpp
+++ b/Engine/Source/Runtime/D3D12RHI/Private/D3D12Buffer.cpp
@@ -84,6 +84,11 @@ struct FRHICommandRenameUploadBuffer final : public FRHICommand<FRHICommandRenam
 			Context.ConditionalClearShaderResource(ResourceLocation);
 		});
 
+#if UE_MEMORY_TRACE_ENABLED
+		// This memory trace happens before RenameLDAChain so the old & new GPU addresses are correct
+		MemoryTrace_ReallocFree(Resource->ResourceLocation.GetGPUVirtualAddress(), EMemoryTraceRootHeap::VideoMemory);
+		MemoryTrace_ReallocAlloc(NewLocation.GetGPUVirtualAddress(), Resource->ResourceLocation.GetSize(), Resource->BufferAlignment, EMemoryTraceRootHeap::VideoMemory);
+#endif
 		Resource->RenameLDAChain(NewLocation);
 
 		if (PreviousPipeline.IsSet())
@@ -594,21 +599,12 @@ void* FD3D12DynamicRHI::LockBuffer(FRHICommandListBase& RHICmdList, FD3D12Buffer
 			{
 				FRHICommandRenameUploadBuffer* Command = ALLOC_COMMAND_CL(RHICmdList, FRHICommandRenameUploadBuffer)(Buffer, Device);
 				Data = Adapter.GetUploadHeapAllocator(Device->GetGPUIndex()).AllocUploadResource(BufferSize, Buffer->BufferAlignment, Command->NewLocation);
-#if UE_MEMORY_TRACE_ENABLED
-				MemoryTrace_ReallocFree(Buffer->ResourceLocation.GetGPUVirtualAddress(), EMemoryTraceRootHeap::VideoMemory);
-				MemoryTrace_ReallocAlloc(Command->NewLocation.GetGPUVirtualAddress(), BufferSize, Buffer->BufferAlignment, EMemoryTraceRootHeap::VideoMemory);
-#endif
-
 				RHICmdList.RHIThreadFence(true);
 			}
 			else
 			{
 				FRHICommandRenameUploadBuffer Command(Buffer, Device);
 				Data = Adapter.GetUploadHeapAllocator(Device->GetGPUIndex()).AllocUploadResource(BufferSize, Buffer->BufferAlignment, Command.NewLocation);
-#if UE_MEMORY_TRACE_ENABLED
-				MemoryTrace_ReallocFree(Buffer->ResourceLocation.GetGPUVirtualAddress(), EMemoryTraceRootHeap::VideoMemory);
-				MemoryTrace_ReallocAlloc(Command.NewLocation.GetGPUVirtualAddress(), BufferSize, Buffer->BufferAlignment, EMemoryTraceRootHeap::VideoMemory);
-#endif
 				Command.Execute(RHICmdList);
 			}
 		}
diff --git a/Engine/Source/Runtime/D3D12RHI/Private/D3D12Resources.cpp b/Engine/Source/Runtime/D3D12RHI/Private/D3D12Resources.cpp
--- a/Engine/Source/Runtime/D3D12RHI/Private/D3D12Resources.cpp
+++ b/Engine/Source/Runtime/D3D12RHI/Private/D3D12Resources.cpp
@@ -298,6 +298,7 @@ FD3D12Heap::~FD3D12Heap()
 	if (GPUVirtualAddress != 0)
 	{
 		MemoryTrace_UnmarkAllocAsHeap(GPUVirtualAddress, TraceHeapId);
+		MemoryTrace_Free(GPUVirtualAddress, EMemoryTraceRootHeap::VideoMemory);
 	}
 #endif	
 
@@ -349,8 +350,14 @@ void FD3D12Heap::SetHeap(ID3D12Heap* HeapIn, const TCHAR* const InName, bool bIn
 		GPUVirtualAddress = TempResource->GetGPUVirtualAddress();
 #if UE_MEMORY_TRACE_ENABLED
 		TraceHeapId = MemoryTrace_HeapSpec(TraceParentHeapId, *(FString(InName) + TEXT(" D3D12Heap")));
-		MemoryTrace_Alloc(GPUVirtualAddress, HeapSize, HeapDesc.Alignment, EMemoryTraceRootHeap::VideoMemory);
-		MemoryTrace_MarkAllocAsHeap(GPUVirtualAddress, TraceHeapId);
+		// Calling GetResourceAllocationInfo is not trivial, only do it if memory trace is enabled
+		if (UE_TRACE_CHANNELEXPR_IS_ENABLED(MemAllocChannel))
+		{
+			const D3D12_RESOURCE_DESC ResourceDesc = TempResource->GetDesc();
+			const D3D12_RESOURCE_ALLOCATION_INFO Info = Adapter->GetD3DDevice()->GetResourceAllocationInfo(0, 1, &ResourceDesc);
+			MemoryTrace_Alloc(GPUVirtualAddress, Info.SizeInBytes, Info.Alignment, EMemoryTraceRootHeap::VideoMemory);
+			MemoryTrace_MarkAllocAsHeap(GPUVirtualAddress, TraceHeapId);
+		}
 #endif				
 #if TRACK_RESOURCE_ALLOCATIONS
 		if (bTrack)
@@ -578,7 +585,8 @@ HRESULT FD3D12Adapter::CreatePlacedResource(const FD3D12ResourceDesc& InDesc, FD
 #endif		
 #if UE_MEMORY_TRACE_ENABLED
 		// Calling GetResourceAllocationInfo is not cheap so check memory allocation tracking is enabled
-		if (UE_TRACE_CHANNELEXPR_IS_ENABLED(MemAllocChannel))
+		// Don't track resources allocated on transient heaps
+		if (UE_TRACE_CHANNELEXPR_IS_ENABLED(MemAllocChannel) && !BackingHeap->GetIsTransient())
 		{
 			const D3D12_RESOURCE_DESC ResourceDesc = (*ppOutResource)->GetResource()->GetDesc();
 			const D3D12_RESOURCE_ALLOCATION_INFO Info = RootDevice->GetResourceAllocationInfo(0, 1, &ResourceDesc);
diff --git a/Engine/Source/Runtime/D3D12RHI/Private/D3D12Stats.cpp b/Engine/Source/Runtime/D3D12RHI/Private/D3D12Stats.cpp
--- a/Engine/Source/Runtime/D3D12RHI/Private/D3D12Stats.cpp
+++ b/Engine/Source/Runtime/D3D12RHI/Private/D3D12Stats.cpp
@@ -311,7 +311,16 @@ void UpdateBufferStats(FD3D12Buffer* Buffer, bool bAllocating)
 #if UE_MEMORY_TRACE_ENABLED
 	if (bAllocating)
 	{
-		MemoryTrace_Alloc(GPUAddress, BufferSize, Buffer->BufferAlignment, EMemoryTraceRootHeap::VideoMemory);
+		bool bStandalone = Buffer->ResourceLocation.GetType() == FD3D12ResourceLocation::ResourceLocationType::eStandAlone;
+		bool bPoolPlacedResource = (!bStandalone && Buffer->ResourceLocation.GetAllocatorType() == FD3D12ResourceLocation::AT_Pool) ?
+									Buffer->ResourceLocation.GetPoolAllocator()->GetAllocationStrategy() == EResourceAllocationStrategy::kPlacedResource : false;
+		// Skip if it's created as a
+		// 1) standalone resource, because MemoryTrace_Alloc has been called in FD3D12Adapter::CreateCommittedResource
+		// 2) placed resource from a pool allocator, because MemoryTrace_Alloc has been called in FD3D12Adapter::CreatePlacedResource
+		if (!bStandalone && !bPoolPlacedResource)
+		{
+			MemoryTrace_Alloc(GPUAddress, BufferSize, Buffer->BufferAlignment, EMemoryTraceRootHeap::VideoMemory);
+		}
 	}
 	else
 	{
diff --git a/Engine/Source/Runtime/D3D12RHI/Private/D3D12Texture.cpp b/Engine/Source/Runtime/D3D12RHI/Private/D3D12Texture.cpp
--- a/Engine/Source/Runtime/D3D12RHI/Private/D3D12Texture.cpp
+++ b/Engine/Source/Runtime/D3D12RHI/Private/D3D12Texture.cpp
@@ -385,6 +385,8 @@ void FD3D12TextureStats::UpdateD3D12TextureStats(FD3D12Texture& Texture, const D
 	INC_MEMORY_STAT_BY_FName(GetRHIStatEnum(Desc.Flags, bCubeMap, b3D).GetName(), TextureSize);
 	INC_MEMORY_STAT_BY(STAT_D3D12MemoryCurrentTotal, TextureSize);
 
+	D3D12_GPU_VIRTUAL_ADDRESS GPUAddress = Texture.ResourceLocation.GetGPUVirtualAddress();
+
 	if (TextureSize > 0)
 	{
 #if PLATFORM_WINDOWS
@@ -394,7 +396,19 @@ void FD3D12TextureStats::UpdateD3D12TextureStats(FD3D12Texture& Texture, const D
 		LLM(FLowLevelMemTracker::Get().OnLowLevelAlloc(ELLMTracker::Default, Texture.GetResource(), TextureSize, ELLMTag::Textures));
 		{
 			LLM(UE_MEMSCOPE_DEFAULT(ELLMTag::Textures));
-			MemoryTrace_Alloc((uint64)Texture.GetResource(), TextureSize, 1024, EMemoryTraceRootHeap::VideoMemory);
+
+#if UE_MEMORY_TRACE_ENABLED
+			bool bStandalone = (Texture.ResourceLocation.GetType() == FD3D12ResourceLocation::ResourceLocationType::eStandAlone);
+			bool bPoolPlacedResource = (!bStandalone && Texture.ResourceLocation.GetAllocatorType() == FD3D12ResourceLocation::AT_Pool) ?
+										Texture.ResourceLocation.GetPoolAllocator()->GetAllocationStrategy() == EResourceAllocationStrategy::kPlacedResource : false;
+			// Skip if it's created as a
+			// 1) standalone resource, because MemoryTrace_Alloc has been called in FD3D12Adapter::CreateCommittedResource
+			// 2) placed resource from a pool allocator, because MemoryTrace_Alloc has been called in FD3D12Adapter::CreatePlacedResource
+			if (!bStandalone && !bPoolPlacedResource)
+			{
+				MemoryTrace_Alloc(GPUAddress, TextureSize, Desc.Alignment, EMemoryTraceRootHeap::VideoMemory);
+			}
+#endif
 		}
 #endif
 		INC_DWORD_STAT(STAT_D3D12TexturesAllocated);
@@ -404,7 +418,10 @@ void FD3D12TextureStats::UpdateD3D12TextureStats(FD3D12Texture& Texture, const D
 #if PLATFORM_WINDOWS
 		LLM(FLowLevelMemTracker::Get().OnLowLevelFree(ELLMTracker::Platform, Texture.GetResource()));
 		LLM(FLowLevelMemTracker::Get().OnLowLevelFree(ELLMTracker::Default, Texture.GetResource()));
-		MemoryTrace_Free((uint64)Texture.GetResource(), EMemoryTraceRootHeap::VideoMemory);
+
+#if UE_MEMORY_TRACE_ENABLED
+		MemoryTrace_Free(GPUAddress, EMemoryTraceRootHeap::VideoMemory);
+#endif
 #endif
 		INC_DWORD_STAT(STAT_D3D12TexturesReleased);
 	}
diff --git a/Engine/Source/Runtime/D3D12RHI/Private/D3D12TransientResourceAllocator.cpp b/Engine/Source/Runtime/D3D12RHI/Private/D3D12TransientResourceAllocator.cpp
--- a/Engine/Source/Runtime/D3D12RHI/Private/D3D12TransientResourceAllocator.cpp
+++ b/Engine/Source/Runtime/D3D12RHI/Private/D3D12TransientResourceAllocator.cpp
@@ -87,6 +87,7 @@ FD3D12TransientHeap::FD3D12TransientHeap(const FInitializer& Initializer, FD3D12
 
 	Heap = new FD3D12Heap(Device, VisibleNodeMask);
 	Heap->SetHeap(D3DHeap, TEXT("TransientResourceAllocator Backing Heap"), true, true);
+	Heap->SetIsTransient(true);
 	Heap->BeginTrackingResidency(Desc.SizeInBytes);
 
 	SetGpuVirtualAddress(Heap->GetGPUVirtualAddress());
diff --git a/Engine/Source/Runtime/D3D12RHI/Public/D3D12Resources.h b/Engine/Source/Runtime/D3D12RHI/Public/D3D12Resources.h
--- a/Engine/Source/Runtime/D3D12RHI/Public/D3D12Resources.h
+++ b/Engine/Source/Runtime/D3D12RHI/Public/D3D12Resources.h
@@ -86,6 +86,8 @@ public:
 	inline D3D12_HEAP_DESC GetHeapDesc() const { return HeapDesc; }
 	inline FD3D12ResidencyHandle& GetResidencyHandle() { return ResidencyHandle; }
 	inline D3D12_GPU_VIRTUAL_ADDRESS GetGPUVirtualAddress() const { return GPUVirtualAddress; }
+	inline void SetIsTransient(bool bInIsTransient) { bIsTransient = bInIsTransient; }
+	inline bool GetIsTransient() const { return bIsTransient; }
 
 private:
 
@@ -97,6 +99,7 @@ private:
 	FD3D12ResidencyHandle ResidencyHandle;
 	HeapId TraceHeapId;
 	HeapId TraceParentHeapId;
+	bool bIsTransient = false; // Whether this is a transient heap
 };
 
 struct FD3D12ResourceDesc : public D3D12_RESOURCE_DESC
