Move the logic to determine whether a buffer/texture is a standalone or pooled placed resource into a utility function in FD3D12ResourceLocation

diff --git a/Engine/Source/Runtime/D3D12RHI/Private/D3D12Resources.cpp b/Engine/Source/Runtime/D3D12RHI/Private/D3D12Resources.cpp
--- a/Engine/Source/Runtime/D3D12RHI/Private/D3D12Resources.cpp
+++ b/Engine/Source/Runtime/D3D12RHI/Private/D3D12Resources.cpp
@@ -1176,6 +1176,13 @@ void FD3D12ResourceLocation::UnlockPoolData()
 	}
 }
 
+bool FD3D12ResourceLocation::IsStandaloneOrPooledPlacedResource() const
+{
+	bool bStandalone = Type == ResourceLocationType::eStandAlone;
+	bool bPoolPlacedResource = (!bStandalone && AllocatorType == AT_Pool) ? PoolAllocator->GetAllocationStrategy() == EResourceAllocationStrategy::kPlacedResource : false;
+	return bStandalone || bPoolPlacedResource;
+}
+
 /////////////////////////////////////////////////////////////////////
 //	FD3D12 Resource Barrier Batcher
 /////////////////////////////////////////////////////////////////////
diff --git a/Engine/Source/Runtime/D3D12RHI/Private/D3D12Stats.cpp b/Engine/Source/Runtime/D3D12RHI/Private/D3D12Stats.cpp
--- a/Engine/Source/Runtime/D3D12RHI/Private/D3D12Stats.cpp
+++ b/Engine/Source/Runtime/D3D12RHI/Private/D3D12Stats.cpp
@@ -311,13 +311,10 @@ void UpdateBufferStats(FD3D12Buffer* Buffer, bool bAllocating)
 #if UE_MEMORY_TRACE_ENABLED
 	if (bAllocating)
 	{
-		bool bStandalone = Buffer->ResourceLocation.GetType() == FD3D12ResourceLocation::ResourceLocationType::eStandAlone;
-		bool bPoolPlacedResource = (!bStandalone && Buffer->ResourceLocation.GetAllocatorType() == FD3D12ResourceLocation::AT_Pool) ?
-									Buffer->ResourceLocation.GetPoolAllocator()->GetAllocationStrategy() == EResourceAllocationStrategy::kPlacedResource : false;
 		// Skip if it's created as a
 		// 1) standalone resource, because MemoryTrace_Alloc has been called in FD3D12Adapter::CreateCommittedResource
 		// 2) placed resource from a pool allocator, because MemoryTrace_Alloc has been called in FD3D12Adapter::CreatePlacedResource
-		if (!bStandalone && !bPoolPlacedResource)
+		if (!Buffer->ResourceLocation.IsStandaloneOrPooledPlacedResource())
 		{
 			MemoryTrace_Alloc(GPUAddress, BufferSize, Buffer->BufferAlignment, EMemoryTraceRootHeap::VideoMemory);
 		}
diff --git a/Engine/Source/Runtime/D3D12RHI/Private/D3D12Texture.cpp b/Engine/Source/Runtime/D3D12RHI/Private/D3D12Texture.cpp
--- a/Engine/Source/Runtime/D3D12RHI/Private/D3D12Texture.cpp
+++ b/Engine/Source/Runtime/D3D12RHI/Private/D3D12Texture.cpp
@@ -398,13 +398,10 @@ void FD3D12TextureStats::UpdateD3D12TextureStats(FD3D12Texture& Texture, const D
 			LLM(UE_MEMSCOPE_DEFAULT(ELLMTag::Textures));
 
 #if UE_MEMORY_TRACE_ENABLED
-			bool bStandalone = (Texture.ResourceLocation.GetType() == FD3D12ResourceLocation::ResourceLocationType::eStandAlone);
-			bool bPoolPlacedResource = (!bStandalone && Texture.ResourceLocation.GetAllocatorType() == FD3D12ResourceLocation::AT_Pool) ?
-										Texture.ResourceLocation.GetPoolAllocator()->GetAllocationStrategy() == EResourceAllocationStrategy::kPlacedResource : false;
 			// Skip if it's created as a
 			// 1) standalone resource, because MemoryTrace_Alloc has been called in FD3D12Adapter::CreateCommittedResource
 			// 2) placed resource from a pool allocator, because MemoryTrace_Alloc has been called in FD3D12Adapter::CreatePlacedResource
-			if (!bStandalone && !bPoolPlacedResource)
+			if (!Texture.ResourceLocation.IsStandaloneOrPooledPlacedResource())
 			{
 				MemoryTrace_Alloc(GPUAddress, TextureSize, Desc.Alignment, EMemoryTraceRootHeap::VideoMemory);
 			}
diff --git a/Engine/Source/Runtime/D3D12RHI/Public/D3D12Resources.h b/Engine/Source/Runtime/D3D12RHI/Public/D3D12Resources.h
--- a/Engine/Source/Runtime/D3D12RHI/Public/D3D12Resources.h
+++ b/Engine/Source/Runtime/D3D12RHI/Public/D3D12Resources.h
@@ -592,6 +592,7 @@ public:
 	const inline bool IsValid() const { return Type != ResourceLocationType::eUndefined; }
 
 	void AsStandAlone(FD3D12Resource* Resource, uint64 InSize = 0, bool bInIsTransient = false, const D3D12_HEAP_PROPERTIES* CustomHeapProperties = nullptr);
+	bool IsStandaloneOrPooledPlacedResource() const;
 
 	inline void AsHeapAliased(FD3D12Resource* Resource)
 	{
