Only free texture memory trace if its GPUVirtualAddress is valid, this skips back buffers which don't have valid GPUVirtualAddress.

diff --git a/Engine/Source/Runtime/D3D12RHI/Private/D3D12Texture.cpp b/Engine/Source/Runtime/D3D12RHI/Private/D3D12Texture.cpp
--- a/Engine/Source/Runtime/D3D12RHI/Private/D3D12Texture.cpp
+++ b/Engine/Source/Runtime/D3D12RHI/Private/D3D12Texture.cpp
@@ -417,7 +417,11 @@ void FD3D12TextureStats::UpdateD3D12TextureStats(FD3D12Texture& Texture, const D
 		LLM(FLowLevelMemTracker::Get().OnLowLevelFree(ELLMTracker::Default, Texture.GetResource()));
 
 #if UE_MEMORY_TRACE_ENABLED
-		MemoryTrace_Free(GPUAddress, EMemoryTraceRootHeap::VideoMemory);
+		// Skip back buffers that aren't traced on alloc and don't have valid GPUVirtualAddress
+		if (GPUAddress != 0)
+		{
+			MemoryTrace_Free(GPUAddress, EMemoryTraceRootHeap::VideoMemory);
+		}
 #endif
 #endif
 		INC_DWORD_STAT(STAT_D3D12TexturesReleased);
