Set NewHeapDesc.Flags to D3D12_HEAP_FLAG_DENY_BUFFERS | D3D12_HEAP_FLAG_DENY_RT_DS_TEXTURES

diff --git a/Engine/Source/Runtime/D3D12RHI/Private/D3D12Resources.cpp b/Engine/Source/Runtime/D3D12RHI/Private/D3D12Resources.cpp
--- a/Engine/Source/Runtime/D3D12RHI/Private/D3D12Resources.cpp
+++ b/Engine/Source/Runtime/D3D12RHI/Private/D3D12Resources.cpp
@@ -209,7 +209,7 @@ void FD3D12Resource::CommitReservedResource()
 		const uint32 ThisHeapSize = RegionSize.NumTiles * TileSizeInBytes;
 		D3D12_HEAP_DESC NewHeapDesc = {};
 		NewHeapDesc.Alignment = D3D12_DEFAULT_RESOURCE_PLACEMENT_ALIGNMENT;
-		NewHeapDesc.Flags = D3D12_HEAP_FLAG_NONE;
+		NewHeapDesc.Flags = D3D12_HEAP_FLAG_DENY_BUFFERS | D3D12_HEAP_FLAG_DENY_RT_DS_TEXTURES;
 		NewHeapDesc.SizeInBytes = ThisHeapSize;
 		NewHeapDesc.Properties = BackingHeapProps;
 		TRefCountPtr<ID3D12Heap> NewHeap;
