Fix a couple of Insight invalid memory allocation errors
- Free heap memory after UnmarkAllocAsHeap on destroying D3D12MemoryPool
- Free temporary allocated memory after texture readback finishes

diff --git a/Engine/Source/Runtime/D3D12RHI/Private/D3D12PoolAllocator.cpp b/Engine/Source/Runtime/D3D12RHI/Private/D3D12PoolAllocator.cpp
--- a/Engine/Source/Runtime/D3D12RHI/Private/D3D12PoolAllocator.cpp
+++ b/Engine/Source/Runtime/D3D12RHI/Private/D3D12PoolAllocator.cpp
@@ -154,6 +154,7 @@ void FD3D12MemoryPool::Destroy()
 	{
 #if UE_MEMORY_TRACE_ENABLED
 		MemoryTrace_UnmarkAllocAsHeap(BackingResource->GetGPUVirtualAddress(), TraceHeapId);
+		MemoryTrace_Free(BackingResource->GetGPUVirtualAddress(), EMemoryTraceRootHeap::VideoMemory);
 #endif
 		ensure(BackingResource->GetRefCount() == 1 || GNumExplicitGPUsForRendering > 1);
 		BackingResource = nullptr;
diff --git a/Engine/Source/Runtime/D3D12RHI/Private/D3D12RenderTarget.cpp b/Engine/Source/Runtime/D3D12RHI/Private/D3D12RenderTarget.cpp
--- a/Engine/Source/Runtime/D3D12RHI/Private/D3D12RenderTarget.cpp
+++ b/Engine/Source/Runtime/D3D12RHI/Private/D3D12RenderTarget.cpp
@@ -1576,6 +1576,11 @@ void FD3D12DynamicRHI::RHIReadSurfaceFloatData(FRHITexture* TextureRHI, FIntRect
 	}
 
 	TempTexture2D->GetResource()->Unmap(0, nullptr);
+
+#if UE_MEMORY_TRACE_ENABLED
+	// Free the temporary texture after read back finishes. This matches the MemoryTrace_Alloc call in Adapter->CreateBuffer() above.
+	MemoryTrace_Free(TempTexture2D->GetGPUVirtualAddress(), EMemoryTraceRootHeap::VideoMemory);
+#endif
 }
 
 void FD3D12DynamicRHI::RHIRead3DSurfaceFloatData(FRHITexture* TextureRHI, FIntRect InRect, FIntPoint ZMinMax, TArray<FFloat16Color>& OutData)
