Fix use of Intel Extension without extension initialized.

Change provided by [REDACTED_PERSON] at Intel

#rnx

diff --git a/Engine/Source/Runtime/D3D12RHI/Private/D3D12Device.cpp b/Engine/Source/Runtime/D3D12RHI/Private/D3D12Device.cpp
--- a/Engine/Source/Runtime/D3D12RHI/Private/D3D12Device.cpp
+++ b/Engine/Source/Runtime/D3D12RHI/Private/D3D12Device.cpp
@@ -542,7 +542,7 @@ D3D12_RESOURCE_ALLOCATION_INFO FD3D12Device::GetResourceAllocationInfo(const FD3
 	{
 		D3D12_RESOURCE_ALLOCATION_INFO Result;
 #if INTEL_EXTENSIONS
-		if (InDesc.bRequires64BitAtomicSupport && IsRHIDeviceIntel() && GRHISupportsAtomicUInt64)
+		if (InDesc.bRequires64BitAtomicSupport && IsRHIDeviceIntel() && GDX12INTCAtomicUInt64Emulation)
 		{
 			FD3D12ResourceDesc LocalDesc = InDesc;
 
diff --git a/Engine/Source/Runtime/D3D12RHI/Private/D3D12RHIPrivate.h b/Engine/Source/Runtime/D3D12RHI/Private/D3D12RHIPrivate.h
--- a/Engine/Source/Runtime/D3D12RHI/Private/D3D12RHIPrivate.h
+++ b/Engine/Source/Runtime/D3D12RHI/Private/D3D12RHIPrivate.h
@@ -102,6 +102,10 @@ inline const TCHAR* GetD3DCommandQueueTypeName(ED3D12QueueType QueueType)
 
 #endif // NV_AFTERMATH
 
+#if INTEL_EXTENSIONS
+	extern bool GDX12INTCAtomicUInt64Emulation;
+#endif //INTEL_EXTENSIONS
+
 #include "D3D12Residency.h"
 
 // D3D RHI public headers.
diff --git a/Engine/Source/Runtime/D3D12RHI/Private/D3D12Resources.cpp b/Engine/Source/Runtime/D3D12RHI/Private/D3D12Resources.cpp
--- a/Engine/Source/Runtime/D3D12RHI/Private/D3D12Resources.cpp
+++ b/Engine/Source/Runtime/D3D12RHI/Private/D3D12Resources.cpp
@@ -420,7 +420,7 @@ HRESULT FD3D12Adapter::CreateCommittedResource(const FD3D12ResourceDesc& InDesc,
 
 	HRESULT hr = S_OK;
 #if INTEL_EXTENSIONS
-	if (InDesc.bRequires64BitAtomicSupport && IsRHIDeviceIntel() && GRHISupportsAtomicUInt64)
+	if (InDesc.bRequires64BitAtomicSupport && IsRHIDeviceIntel() && GDX12INTCAtomicUInt64Emulation)
 	{
 		INTC_D3D12_RESOURCE_DESC_0001 IntelLocalDesc{};
 		IntelLocalDesc.pD3D12Desc = &LocalDesc;
@@ -537,7 +537,7 @@ HRESULT FD3D12Adapter::CreatePlacedResource(const FD3D12ResourceDesc& InDesc, FD
 	TRefCountPtr<ID3D12Resource> pResource;
 	HRESULT hr = S_OK;
 #if INTEL_EXTENSIONS
-	if (InDesc.bRequires64BitAtomicSupport && IsRHIDeviceIntel() && GRHISupportsAtomicUInt64)
+	if (InDesc.bRequires64BitAtomicSupport && IsRHIDeviceIntel() && GDX12INTCAtomicUInt64Emulation)
 	{
 		FD3D12ResourceDesc LocalDesc = InDesc;
 		INTC_D3D12_RESOURCE_DESC_0001 IntelLocalDesc{};
diff --git a/Engine/Source/Runtime/D3D12RHI/Private/Windows/WindowsD3D12Device.cpp b/Engine/Source/Runtime/D3D12RHI/Private/Windows/WindowsD3D12Device.cpp
--- a/Engine/Source/Runtime/D3D12RHI/Private/Windows/WindowsD3D12Device.cpp
+++ b/Engine/Source/Runtime/D3D12RHI/Private/Windows/WindowsD3D12Device.cpp
@@ -27,6 +27,8 @@
 	THIRD_PARTY_INCLUDES_START
 	#include "igdext.h"
 	THIRD_PARTY_INCLUDES_END
+
+	bool GDX12INTCAtomicUInt64Emulation = false;
 #endif
 
 #include "Windows/HideWindowsPlatformTypes.h"
@@ -365,8 +367,6 @@ static INTCExtensionContext* CreateIntelExtensionsContext(ID3D12Device* Device,
 
 static bool EnableIntelAtomic64Support(INTCExtensionContext* IntelExtensionContext, INTCExtensionInfo& INTCExtensionInfo)
 {
-	bool bEmulatedAtomic64Support = false;
-
 	if (IntelExtensionContext)
 	{
 #if D3D12_CORE_ENABLED
@@ -379,7 +379,7 @@ static bool EnableIntelAtomic64Support(INTCExtensionContext* IntelExtensionConte
 			const HRESULT hr = INTC_D3D12_SetFeatureSupport(IntelExtensionContext, &INTCFeature);
 			if (SUCCEEDED(hr))
 			{
-				bEmulatedAtomic64Support = true;
+				GDX12INTCAtomicUInt64Emulation = true;
 				UE_LOG(LogD3D12RHI, Log, TEXT("Intel Extensions 64-bit Typed Atomics emulation enabled."));
 			}
 			else
@@ -390,7 +390,7 @@ static bool EnableIntelAtomic64Support(INTCExtensionContext* IntelExtensionConte
 #endif
 	}
 
-	return bEmulatedAtomic64Support;
+	return GDX12INTCAtomicUInt64Emulation;
 }
 #endif
 
