Added more debug info when passing an invalid size to RHILockBuffer.

#rnx

diff --git a/Engine/Source/Runtime/D3D12RHI/Private/D3D12Buffer.cpp b/Engine/Source/Runtime/D3D12RHI/Private/D3D12Buffer.cpp
--- a/Engine/Source/Runtime/D3D12RHI/Private/D3D12Buffer.cpp
+++ b/Engine/Source/Runtime/D3D12RHI/Private/D3D12Buffer.cpp
@@ -571,7 +571,7 @@ void* FD3D12DynamicRHI::LockBuffer(FRHICommandListBase& RHICmdList, FD3D12Buffer
 {
 	SCOPE_CYCLE_COUNTER(STAT_D3D12LockBufferTime);
 
-	check(Size <= BufferSize);
+	checkf(Size <= BufferSize, TEXT("Requested lock size %u is larger than the total size %u for buffer '%s'."), Size, BufferSize, *Buffer->GetName().ToString());
 
 	FD3D12LockedResource& LockedData = Buffer->LockedData;
 	check(LockedData.bLocked == false);
