Fix a deadlock that can happen when two streaming updates are blocked on task graph events in a critical section. The threads where the two updates run on can pull in and execute FAsyncCancelIORequestsTask while waiting. If they grab each others FAsyncCancelIORequestsTask, it will deadlock since the tasks will also try to enter the critical section but won't succeed.

diff --git a/Engine/Source/Runtime/Apple/MetalRHI/Private/MetalDynamicRHI.h b/Engine/Source/Runtime/Apple/MetalRHI/Private/MetalDynamicRHI.h
--- a/Engine/Source/Runtime/Apple/MetalRHI/Private/MetalDynamicRHI.h
+++ b/Engine/Source/Runtime/Apple/MetalRHI/Private/MetalDynamicRHI.h
@@ -76,7 +76,7 @@ public:
 	virtual void RHIGetTextureMemoryStats(FTextureMemoryStats& OutStats) final override;
 	virtual bool RHIGetTextureMemoryVisualizeData(FColor* TextureData,int32 SizeX,int32 SizeY,int32 Pitch,int32 PixelSize) final override;
 	virtual FTextureRHIRef RHICreateTexture(const FRHITextureCreateDesc& CreateDesc) final override;
-	virtual FTexture2DRHIRef RHIAsyncCreateTexture2D(uint32 SizeX, uint32 SizeY, uint8 Format, uint32 NumMips, ETextureCreateFlags Flags, ERHIAccess InResourceState, void** InitialMipData, uint32 NumInitialMips) final override;
+	virtual FTexture2DRHIRef RHIAsyncCreateTexture2D(uint32 SizeX, uint32 SizeY, uint8 Format, uint32 NumMips, ETextureCreateFlags Flags, ERHIAccess InResourceState, void** InitialMipData, uint32 NumInitialMips, FGraphEventRef& OutCompletionEvent) final override;
 	virtual void RHIGenerateMips(FRHITexture* Texture) final override;
 	virtual uint32 RHIComputeMemorySize(FRHITexture* TextureRHI) final override;
 	virtual FTexture2DRHIRef RHIAsyncReallocateTexture2D(FRHITexture2D* Texture2D, int32 NewMipCount, int32 NewSizeX, int32 NewSizeY, FThreadSafeCounter* RequestStatus) final override;
diff --git a/Engine/Source/Runtime/Apple/MetalRHI/Private/MetalTexture.cpp b/Engine/Source/Runtime/Apple/MetalRHI/Private/MetalTexture.cpp
--- a/Engine/Source/Runtime/Apple/MetalRHI/Private/MetalTexture.cpp
+++ b/Engine/Source/Runtime/Apple/MetalRHI/Private/MetalTexture.cpp
@@ -1518,7 +1518,7 @@ FTextureRHIRef FMetalDynamicRHI::RHICreateTexture(const FRHITextureCreateDesc& C
 	}
 }
 
-FTexture2DRHIRef FMetalDynamicRHI::RHIAsyncCreateTexture2D(uint32 SizeX, uint32 SizeY, uint8 Format, uint32 NumMips, ETextureCreateFlags Flags, ERHIAccess InResourceState, void** InitialMipData, uint32 NumInitialMips)
+FTexture2DRHIRef FMetalDynamicRHI::RHIAsyncCreateTexture2D(uint32 SizeX, uint32 SizeY, uint8 Format, uint32 NumMips, ETextureCreateFlags Flags, ERHIAccess InResourceState, void** InitialMipData, uint32 NumInitialMips, FGraphEventRef& OutCompletionEvent)
 {
 	UE_LOG(LogMetal, Fatal, TEXT("RHIAsyncCreateTexture2D is not supported"));
 	return FTexture2DRHIRef();
diff --git a/Engine/Source/Runtime/D3D12RHI/Private/D3D12RHIPrivate.h b/Engine/Source/Runtime/D3D12RHI/Private/D3D12RHIPrivate.h
--- a/Engine/Source/Runtime/D3D12RHI/Private/D3D12RHIPrivate.h
+++ b/Engine/Source/Runtime/D3D12RHI/Private/D3D12RHIPrivate.h
@@ -461,7 +461,7 @@ public:
 	virtual uint64 RHIGetMinimumAlignmentForBufferBackedSRV(EPixelFormat Format) final override;
 	virtual void RHIGetTextureMemoryStats(FTextureMemoryStats& OutStats) final override;
 	virtual bool RHIGetTextureMemoryVisualizeData(FColor* TextureData, int32 SizeX, int32 SizeY, int32 Pitch, int32 PixelSize) final override;
-	virtual FTextureRHIRef RHIAsyncCreateTexture2D(uint32 SizeX, uint32 SizeY, uint8 Format, uint32 NumMips, ETextureCreateFlags Flags, ERHIAccess InResourceState, void** InitialMipData, uint32 NumInitialMips) final override;
+	virtual FTextureRHIRef RHIAsyncCreateTexture2D(uint32 SizeX, uint32 SizeY, uint8 Format, uint32 NumMips, ETextureCreateFlags Flags, ERHIAccess InResourceState, void** InitialMipData, uint32 NumInitialMips, FGraphEventRef& OutCompletionEvent) final override;
 	virtual FTextureRHIRef RHICreateTexture(const FRHITextureCreateDesc& CreateDesc) override;
 	virtual FShaderResourceViewRHIRef RHICreateShaderResourceView(FRHITexture* Texture, const FRHITextureSRVCreateInfo& CreateInfo) final override;
 	virtual uint32 RHIComputeMemorySize(FRHITexture* TextureRHI) final override;
diff --git a/Engine/Source/Runtime/D3D12RHI/Private/D3D12Texture.cpp b/Engine/Source/Runtime/D3D12RHI/Private/D3D12Texture.cpp
--- a/Engine/Source/Runtime/D3D12RHI/Private/D3D12Texture.cpp
+++ b/Engine/Source/Runtime/D3D12RHI/Private/D3D12Texture.cpp
@@ -1008,7 +1008,61 @@ FTextureRHIRef FD3D12DynamicRHI::RHICreateTexture_RenderThread(class FRHICommand
 }
 
 
-FTextureRHIRef FD3D12DynamicRHI::RHIAsyncCreateTexture2D(uint32 SizeX, uint32 SizeY, uint8 Format, uint32 NumMips, ETextureCreateFlags Flags, ERHIAccess InResourceState, void** InitialMipData, uint32 NumInitialMips)
+class FWaitInitialMipDataUploadTask
+{
+public:
+	TRefCountPtr<FD3D12Texture> Texture;
+	FD3D12ResourceLocation TempResourceLocation;
+	FD3D12ResourceLocation TempResourceLocationLowMips;
+
+	FWaitInitialMipDataUploadTask(
+		FD3D12Texture* InTexture,
+		FD3D12ResourceLocation& InTempResourceLocation,
+		FD3D12ResourceLocation& InTempResourceLocationLowMips)
+		: Texture(InTexture)
+		, TempResourceLocation(InTempResourceLocation.GetParentDevice())
+		, TempResourceLocationLowMips(InTempResourceLocationLowMips.GetParentDevice())
+	{
+		FD3D12ResourceLocation::TransferOwnership(TempResourceLocation, InTempResourceLocation);
+		FD3D12ResourceLocation::TransferOwnership(TempResourceLocationLowMips, InTempResourceLocationLowMips);
+	}
+
+	void DoTask(ENamedThreads::Type CurrentThread, const FGraphEventRef& MyCompletionGraphEvent)
+	{
+		for (FD3D12Texture& CurrentTexture : *Texture)
+		{
+			// Initial data upload is done
+			CurrentTexture.ResourceLocation.UnlockPoolData();
+		}
+
+		// These are clear to be recycled now because GPU is done with it at this point because this task use the copy command list sync points
+		// as prerequisites. No defer delete required but can be reused immediately
+		TempResourceLocation.GetResource()->DoNotDeferDelete();
+		TempResourceLocation.GetResource()->Release();
+		if (TempResourceLocationLowMips.IsValid())
+		{
+			TempResourceLocationLowMips.GetResource()->DoNotDeferDelete();
+			TempResourceLocationLowMips.GetResource()->Release();
+		}
+	}
+
+	static ESubsequentsMode::Type GetSubsequentsMode()
+	{
+		return ESubsequentsMode::TrackSubsequents;
+	}
+
+	ENamedThreads::Type GetDesiredThread()
+	{
+		return ENamedThreads::AnyNormalThreadNormalTask;
+	}
+
+	FORCEINLINE TStatId GetStatId() const
+	{
+		RETURN_QUICK_DECLARE_CYCLE_STAT(FGatherRequestsTask, STATGROUP_D3D12RHI);
+	}
+};
+
+FTextureRHIRef FD3D12DynamicRHI::RHIAsyncCreateTexture2D(uint32 SizeX, uint32 SizeY, uint8 Format, uint32 NumMips, ETextureCreateFlags Flags, ERHIAccess InResourceState, void** InitialMipData, uint32 NumInitialMips, FGraphEventRef& OutCompletionEvent)
 {	
 	check(GRHISupportsAsyncTextureCreation);
 	
@@ -1102,6 +1156,9 @@ FTextureRHIRef FD3D12DynamicRHI::RHIAsyncCreateTexture2D(uint32 SizeX, uint32 Si
 		return NewTexture;
 	});
 
+	FGraphEventArray CopyCompleteEvents;
+	OutCompletionEvent = nullptr;
+
 	if (TextureOut)
 	{
 		// SubResourceData is only used in async texture creation (RHIAsyncCreateTexture2D). We need to manually transition the resource to
@@ -1200,6 +1257,8 @@ FTextureRHIRef FD3D12DynamicRHI::RHIAsyncCreateTexture2D(uint32 SizeX, uint32 Si
 				FD3D12CopyScope CopyScope(Device, ED3D12SyncPointType::GPUAndCPU);
 				SyncPoint = CopyScope.GetSyncPoint();
 
+				CopyCompleteEvents.Add(SyncPoint->GetGraphEvent());
+
 				// NB: Do not increment NumCopies because that will count as work on the direct
 				// queue, not the copy queue, possibly causing it to flush prematurely. We are
 				// explicitly submitting the copy command list so there's no need to increment any
@@ -1249,26 +1308,16 @@ FTextureRHIRef FD3D12DynamicRHI::RHIAsyncCreateTexture2D(uint32 SizeX, uint32 Si
 
 				CopyScope.Context.UpdateResidency(Resource);
 			}
-
-			// The FD3D12CopyScope above immediately submits work to the copy queue when the scope ends.
-			// Block this CPU thread until the copy queue work has completed on the GPU.
-			SyncPoint->Wait();
-
-			// Blocking update is done
-			CurrentTexture.ResourceLocation.UnlockPoolData();
 		}
 
 		FD3D12TextureStats::D3D12TextureAllocated(*TextureOut);
 
-		// These are clear to be recycled now because GPU is done with it at this point. We wait on GPU in ExecuteCommandList() above.
-		// No defer delete required but can be reused immediately
-		TempResourceLocation.GetResource()->DoNotDeferDelete();
-		TempResourceLocation.GetResource()->Release();
-		if (bSplitAllocation)
-		{
-			TempResourceLocationLowMips.GetResource()->DoNotDeferDelete();
-			TempResourceLocationLowMips.GetResource()->Release();
-		}
+		check(CopyCompleteEvents.Num() > 0);
+
+		OutCompletionEvent = TGraphTask<FWaitInitialMipDataUploadTask>::CreateTask(&CopyCompleteEvents).ConstructAndDispatchWhenReady(
+			TextureOut,
+			TempResourceLocation,
+			TempResourceLocationLowMips);
 	}
 
 	if (TempBufferSize != ZeroBufferSize)
diff --git a/Engine/Source/Runtime/Engine/Private/Streaming/Texture2DMipAllocator_AsyncCreate.cpp b/Engine/Source/Runtime/Engine/Private/Streaming/Texture2DMipAllocator_AsyncCreate.cpp
--- a/Engine/Source/Runtime/Engine/Private/Streaming/Texture2DMipAllocator_AsyncCreate.cpp
+++ b/Engine/Source/Runtime/Engine/Private/Streaming/Texture2DMipAllocator_AsyncCreate.cpp
@@ -82,7 +82,29 @@ bool FTexture2DMipAllocator_AsyncCreate::FinalizeMips(const FTextureUpdateContex
 	if (!IntermediateTextureRHI)
 	{
 		// Create the intermediate texture.
-		IntermediateTextureRHI = RHIAsyncCreateTexture2D(FinalSizeX, FinalSizeY, FinalFormat, ResourceState.NumRequestedLODs, Context.Resource->GetCreationFlags(), FinalMipData.GetData(), FinalMipData.Num());
+		FGraphEventRef CompletionEvent;
+		IntermediateTextureRHI = RHIAsyncCreateTexture2D(
+			FinalSizeX,
+			FinalSizeY,
+			FinalFormat,
+			ResourceState.NumRequestedLODs,
+			Context.Resource->GetCreationFlags(),
+			FinalMipData.GetData(),
+			FinalMipData.Num(),
+			CompletionEvent);
+
+		if (CompletionEvent)
+		{
+			SyncOptions.Counter->Increment();
+			FFunctionGraphTask::CreateAndDispatchWhenReady(
+				[SyncCounter = SyncOptions.Counter]()
+				{
+					SyncCounter->Decrement();
+				},
+				TStatId{},
+				CompletionEvent);
+		}
+
 		// Free the temporary mip data, since copy is now in the RHIAsyncCreateTexture2D command.
 		ReleaseAllocatedMipData();
 
diff --git a/Engine/Source/Runtime/Engine/Private/Streaming/Texture2DStreamIn.cpp b/Engine/Source/Runtime/Engine/Private/Streaming/Texture2DStreamIn.cpp
--- a/Engine/Source/Runtime/Engine/Private/Streaming/Texture2DStreamIn.cpp
+++ b/Engine/Source/Runtime/Engine/Private/Streaming/Texture2DStreamIn.cpp
@@ -147,6 +147,7 @@ void FTexture2DStreamIn::DoAsyncCreateWithNewMips(const FContext& Context)
 		FTexture2DResource::WarnRequiresTightPackedMip(RequestedMipMap.SizeX,RequestedMipMap.SizeY,Context.Resource->GetPixelFormat(),MipData[PendingFirstLODIdx].Pitch);
 
 		ensure(IntermediateTextureRHI == nullptr);
+		FGraphEventRef CompletionEvent;
 		IntermediateTextureRHI = RHIAsyncCreateTexture2D(
 			RequestedMipMap.SizeX,
 			RequestedMipMap.SizeY,
@@ -154,6 +155,19 @@ void FTexture2DStreamIn::DoAsyncCreateWithNewMips(const FContext& Context)
 			ResourceState.NumRequestedLODs,
 			Context.Resource->GetCreationFlags(),
 			InitialMipDataForAsyncCreate+PendingFirstLODIdx,
-			ResourceState.NumRequestedLODs - ResourceState.NumResidentLODs);
+			ResourceState.NumRequestedLODs - ResourceState.NumResidentLODs,
+			CompletionEvent);
+
+		if (CompletionEvent)
+		{
+			TaskSynchronization.Increment();
+			FFunctionGraphTask::CreateAndDispatchWhenReady(
+				[this]()
+				{
+					TaskSynchronization.Decrement();
+				},
+				TStatId{},
+				CompletionEvent);
+		}
 	}
 }
diff --git a/Engine/Source/Runtime/Engine/Private/Streaming/Texture2DStreamOut_AsyncCreate.cpp b/Engine/Source/Runtime/Engine/Private/Streaming/Texture2DStreamOut_AsyncCreate.cpp
--- a/Engine/Source/Runtime/Engine/Private/Streaming/Texture2DStreamOut_AsyncCreate.cpp
+++ b/Engine/Source/Runtime/Engine/Private/Streaming/Texture2DStreamOut_AsyncCreate.cpp
@@ -34,13 +34,28 @@ void FTexture2DStreamOut_AsyncCreate::AsyncCreate(const FContext& Context)
 		const FTexture2DMipMap& RequestedMipMap = *Context.MipsView[PendingFirstLODIdx];
 
 		ensure(IntermediateTextureRHI == nullptr);
+		FGraphEventRef CompletionEvent;
 		IntermediateTextureRHI = RHIAsyncCreateTexture2D(
 			RequestedMipMap.SizeX,
 			RequestedMipMap.SizeY,
 			Context.Resource->GetPixelFormat(),
 			ResourceState.NumRequestedLODs,
 			Context.Resource->GetCreationFlags(),
-			nullptr, 0);
+			nullptr,
+			0,
+			CompletionEvent);
+
+		if (CompletionEvent)
+		{
+			TaskSynchronization.Increment();
+			FFunctionGraphTask::CreateAndDispatchWhenReady(
+				[this]()
+				{
+					TaskSynchronization.Decrement();
+				},
+				TStatId{},
+				CompletionEvent);
+		}
 	}
 
 	PushTask(Context, TT_Render, SRA_UPDATE_CALLBACK(Finalize), TT_Render, SRA_UPDATE_CALLBACK(Cancel));
diff --git a/Engine/Source/Runtime/NullDrv/Public/NullRHI.h b/Engine/Source/Runtime/NullDrv/Public/NullRHI.h
--- a/Engine/Source/Runtime/NullDrv/Public/NullRHI.h
+++ b/Engine/Source/Runtime/NullDrv/Public/NullRHI.h
@@ -216,7 +216,7 @@ public:
 		return new FNullTexture(CreateDesc);
 	}
 
-	virtual FTextureRHIRef RHIAsyncCreateTexture2D(uint32 SizeX, uint32 SizeY, uint8 Format, uint32 NumMips, ETextureCreateFlags Flags, ERHIAccess InResourceState, void** InitialMipData, uint32 NumInitialMips) final override
+	virtual FTextureRHIRef RHIAsyncCreateTexture2D(uint32 SizeX, uint32 SizeY, uint8 Format, uint32 NumMips, ETextureCreateFlags Flags, ERHIAccess InResourceState, void** InitialMipData, uint32 NumInitialMips, FGraphEventRef& OutCompletionEvent) final override
 	{ 
 		const FRHITextureCreateDesc Desc =
 			FRHITextureCreateDesc::Create2D(TEXT("FNullDynamicRHI::RHIAsyncCreateTexture2D"), SizeX, SizeY, (EPixelFormat)Format)
@@ -224,7 +224,7 @@ public:
 			.SetFlags(Flags)
 			.SetNumMips(NumMips)
 			.SetInitialState(InResourceState);
-
+		OutCompletionEvent = nullptr;
 		return this->RHICreateTexture(Desc);
 	}
 
diff --git a/Engine/Source/Runtime/OpenGLDrv/Private/OpenGLTexture.cpp b/Engine/Source/Runtime/OpenGLDrv/Private/OpenGLTexture.cpp
--- a/Engine/Source/Runtime/OpenGLDrv/Private/OpenGLTexture.cpp
+++ b/Engine/Source/Runtime/OpenGLDrv/Private/OpenGLTexture.cpp
@@ -1464,7 +1464,7 @@ FTextureRHIRef FOpenGLDynamicRHI::RHICreateTexture(const FRHITextureCreateDesc&
 	return new FOpenGLTexture(CreateDesc);
 }
 
-FTextureRHIRef FOpenGLDynamicRHI::RHIAsyncCreateTexture2D(uint32 SizeX, uint32 SizeY, uint8 Format, uint32 NumMips, ETextureCreateFlags Flags, ERHIAccess InResourceState, void** InitialMipData, uint32 NumInitialMips)
+FTextureRHIRef FOpenGLDynamicRHI::RHIAsyncCreateTexture2D(uint32 SizeX, uint32 SizeY, uint8 Format, uint32 NumMips, ETextureCreateFlags Flags, ERHIAccess InResourceState, void** InitialMipData, uint32 NumInitialMips, FGraphEventRef& OutCompletionEvent)
 {
 	check(0);
 	return FTexture2DRHIRef();
diff --git a/Engine/Source/Runtime/OpenGLDrv/Public/OpenGLDrv.h b/Engine/Source/Runtime/OpenGLDrv/Public/OpenGLDrv.h
--- a/Engine/Source/Runtime/OpenGLDrv/Public/OpenGLDrv.h
+++ b/Engine/Source/Runtime/OpenGLDrv/Public/OpenGLDrv.h
@@ -360,7 +360,7 @@ public:
 	virtual void RHIGetTextureMemoryStats(FTextureMemoryStats& OutStats) final override;
 	virtual bool RHIGetTextureMemoryVisualizeData(FColor* TextureData,int32 SizeX,int32 SizeY,int32 Pitch,int32 PixelSize) final override;
 	virtual FTextureRHIRef RHICreateTexture(const FRHITextureCreateDesc& CreateDesc) final override;
-	virtual FTextureRHIRef RHIAsyncCreateTexture2D(uint32 SizeX, uint32 SizeY, uint8 Format, uint32 NumMips, ETextureCreateFlags Flags, ERHIAccess InResourceState, void** InitialMipData, uint32 NumInitialMips) final override;
+	virtual FTextureRHIRef RHIAsyncCreateTexture2D(uint32 SizeX, uint32 SizeY, uint8 Format, uint32 NumMips, ETextureCreateFlags Flags, ERHIAccess InResourceState, void** InitialMipData, uint32 NumInitialMips, FGraphEventRef& OutCompletionEvent) final override;
 	virtual FShaderResourceViewRHIRef RHICreateShaderResourceView(FRHITexture* Texture, const FRHITextureSRVCreateInfo& CreateInfo) final override;
 	virtual void RHIGenerateMips(FRHITexture* Texture) final override;
 	virtual uint32 RHIComputeMemorySize(FRHITexture* TextureRHI) final override;
diff --git a/Engine/Source/Runtime/RHI/Private/DynamicRHI.cpp b/Engine/Source/Runtime/RHI/Private/DynamicRHI.cpp
--- a/Engine/Source/Runtime/RHI/Private/DynamicRHI.cpp
+++ b/Engine/Source/Runtime/RHI/Private/DynamicRHI.cpp
@@ -569,6 +569,17 @@ uint64 FDynamicRHI::RHIGetMinimumAlignmentForBufferBackedSRV(EPixelFormat Format
 	return 1;
 }
 
+FTextureRHIRef FDynamicRHI::RHIAsyncCreateTexture2D(uint32 SizeX, uint32 SizeY, uint8 Format, uint32 NumMips, ETextureCreateFlags Flags, ERHIAccess InResourceState, void** InitialMipData, uint32 NumInitialMips)
+{
+	FGraphEventRef CompletionEvent;
+	FTextureRHIRef Result = this->RHIAsyncCreateTexture2D(SizeX, SizeY, Format, NumMips, Flags, InResourceState, InitialMipData, NumInitialMips, CompletionEvent);
+	if (CompletionEvent)
+	{
+		CompletionEvent->Wait();
+	}
+	return Result;
+}
+
 uint64 FDynamicRHI::RHIComputePrecachePSOHash(const FGraphicsPipelineStateInitializer& Initializer)
 {
 	// When compute precache PSO hash we assume a valid state precache PSO hash is already provided
diff --git a/Engine/Source/Runtime/RHI/Public/DynamicRHI.h b/Engine/Source/Runtime/RHI/Public/DynamicRHI.h
--- a/Engine/Source/Runtime/RHI/Public/DynamicRHI.h
+++ b/Engine/Source/Runtime/RHI/Public/DynamicRHI.h
@@ -14,6 +14,7 @@ DynamicRHI.h: Dynamically bound Render Hardware Interface definitions.
 #include "Serialization/MemoryLayout.h"
 #include "Containers/ArrayView.h"
 #include "Misc/EnumClassFlags.h"
+#include "Async/TaskGraphInterfaces.h"
 
 class FBlendStateInitializerRHI;
 class FGraphicsPipelineStateInitializer;
@@ -511,10 +512,14 @@ public:
 	* @param Flags - ETextureCreateFlags creation flags
 	* @param InitialMipData - pointers to mip data with which to create the texture
 	* @param NumInitialMips - how many mips are provided in InitialMipData
+	* @param OutCompletionEvent - An event signaled on operation completion. Can return null. Operation can still be pending after function returns (e.g. initial data upload in-flight)
 	* @returns a reference to a 2D texture resource
 	*/
 	// FlushType: Thread safe
-	virtual FTextureRHIRef RHIAsyncCreateTexture2D(uint32 SizeX, uint32 SizeY, uint8 Format, uint32 NumMips, ETextureCreateFlags Flags, ERHIAccess InResourceState, void** InitialMipData, uint32 NumInitialMips) = 0;
+	virtual FTextureRHIRef RHIAsyncCreateTexture2D(uint32 SizeX, uint32 SizeY, uint8 Format, uint32 NumMips, ETextureCreateFlags Flags, ERHIAccess InResourceState, void** InitialMipData, uint32 NumInitialMips, FGraphEventRef& OutCompletionEvent) = 0;
+
+	UE_DEPRECATED(5.2, "RHIAsyncCreateTexture2D now requires a completion callback. Using the old version in a critical section can lead to deadlock. Please switch to the new function signature.")
+	FTextureRHIRef RHIAsyncCreateTexture2D(uint32 SizeX, uint32 SizeY, uint8 Format, uint32 NumMips, ETextureCreateFlags Flags, ERHIAccess InResourceState, void** InitialMipData, uint32 NumInitialMips);
 
 	/**
 	* Creates a shader resource view for a texture
diff --git a/Engine/Source/Runtime/RHI/Public/RHICommandList.h b/Engine/Source/Runtime/RHI/Public/RHICommandList.h
--- a/Engine/Source/Runtime/RHI/Public/RHICommandList.h
+++ b/Engine/Source/Runtime/RHI/Public/RHICommandList.h
@@ -5184,14 +5184,32 @@ FORCEINLINE FTextureRHIRef RHICreateTexture(const FRHITextureCreateDesc& CreateD
 }
 
 //UE_DEPRECATED(5.1, "FRHITexture2D is deprecated, please use RHICreateTexture(const FRHITextureCreateDesc&).")
+FORCEINLINE FTexture2DRHIRef RHIAsyncCreateTexture2D(uint32 SizeX, uint32 SizeY, uint8 Format, uint32 NumMips, ETextureCreateFlags Flags, ERHIAccess InResourceState, void** InitialMipData, uint32 NumInitialMips, FGraphEventRef& OutCompletionEvent)
+{
+	LLM_SCOPE(EnumHasAnyFlags(Flags, TexCreate_RenderTargetable | TexCreate_DepthStencilTargetable) ? ELLMTag::RenderTargets : ELLMTag::Textures);
+	const ERHIAccess ResourceState = InResourceState == ERHIAccess::Unknown ? RHIGetDefaultResourceState((ETextureCreateFlags)Flags, InitialMipData != nullptr) : InResourceState;
+	return GDynamicRHI->RHIAsyncCreateTexture2D(SizeX, SizeY, Format, NumMips, Flags, ResourceState, InitialMipData, NumInitialMips, OutCompletionEvent);
+}
+
+//UE_DEPRECATED(5.1, "FRHITexture2D is deprecated, please use RHICreateTexture(const FRHITextureCreateDesc&).")
+FORCEINLINE FTexture2DRHIRef RHIAsyncCreateTexture2D(uint32 SizeX, uint32 SizeY, uint8 Format, uint32 NumMips, ETextureCreateFlags Flags, void** InitialMipData, uint32 NumInitialMips, FGraphEventRef& OutCompletionEvent)
+{
+PRAGMA_DISABLE_DEPRECATION_WARNINGS
+	return RHIAsyncCreateTexture2D(SizeX, SizeY, Format, NumMips, Flags, ERHIAccess::Unknown, InitialMipData, NumInitialMips, OutCompletionEvent);
+PRAGMA_ENABLE_DEPRECATION_WARNINGS
+}
+
+UE_DEPRECATED(5.2, "RHIAsyncCreateTexture2D now requires a completion callback. Using the old version in a critical section can lead to deadlock. Please switch to the new function signature.")
 FORCEINLINE FTexture2DRHIRef RHIAsyncCreateTexture2D(uint32 SizeX, uint32 SizeY, uint8 Format, uint32 NumMips, ETextureCreateFlags Flags, ERHIAccess InResourceState, void** InitialMipData, uint32 NumInitialMips)
 {
 	LLM_SCOPE(EnumHasAnyFlags(Flags, TexCreate_RenderTargetable | TexCreate_DepthStencilTargetable) ? ELLMTag::RenderTargets : ELLMTag::Textures);
 	const ERHIAccess ResourceState = InResourceState == ERHIAccess::Unknown ? RHIGetDefaultResourceState((ETextureCreateFlags)Flags, InitialMipData != nullptr) : InResourceState;
+PRAGMA_DISABLE_DEPRECATION_WARNINGS
 	return GDynamicRHI->RHIAsyncCreateTexture2D(SizeX, SizeY, Format, NumMips, Flags, ResourceState, InitialMipData, NumInitialMips);
+PRAGMA_ENABLE_DEPRECATION_WARNINGS
 }
 
-//UE_DEPRECATED(5.1, "FRHITexture2D is deprecated, please use RHICreateTexture(const FRHITextureCreateDesc&).")
+UE_DEPRECATED(5.2, "RHIAsyncCreateTexture2D now requires a completion callback. Using the old version in a critical section can lead to deadlock. Please switch to the new function signature.")
 FORCEINLINE FTexture2DRHIRef RHIAsyncCreateTexture2D(uint32 SizeX, uint32 SizeY, uint8 Format, uint32 NumMips, ETextureCreateFlags Flags, void** InitialMipData, uint32 NumInitialMips)
 {
 PRAGMA_DISABLE_DEPRECATION_WARNINGS
diff --git a/Engine/Source/Runtime/RHI/Public/RHIValidation.h b/Engine/Source/Runtime/RHI/Public/RHIValidation.h
--- a/Engine/Source/Runtime/RHI/Public/RHIValidation.h
+++ b/Engine/Source/Runtime/RHI/Public/RHIValidation.h
@@ -602,14 +602,15 @@ public:
 	* @param Flags - ETextureCreateFlags creation flags
 	* @param InitialMipData - pointers to mip data with which to create the texture
 	* @param NumInitialMips - how many mips are provided in InitialMipData
+	* @param OutCompletionEvent - An event signaled on operation completion. Can return null. Operation can still be pending after function returns (e.g. initial data upload in-flight)
 	* @returns a reference to a 2D texture resource
 	*/
 	// FlushType: Thread safe
-	virtual FTextureRHIRef RHIAsyncCreateTexture2D(uint32 SizeX, uint32 SizeY, uint8 Format, uint32 NumMips, ETextureCreateFlags Flags, ERHIAccess InResourceState, void** InitialMipData, uint32 NumInitialMips) override final
+	virtual FTextureRHIRef RHIAsyncCreateTexture2D(uint32 SizeX, uint32 SizeY, uint8 Format, uint32 NumMips, ETextureCreateFlags Flags, ERHIAccess InResourceState, void** InitialMipData, uint32 NumInitialMips, FGraphEventRef& OutCompletionEvent) override final
 	{
 		check(GRHISupportsAsyncTextureCreation);
 		ensure(FMath::Max(SizeX, SizeY) >= (1u << (FMath::Max(1u, NumMips) - 1)));
-		FTextureRHIRef Texture = RHI->RHIAsyncCreateTexture2D(SizeX, SizeY, Format, NumMips, Flags, InResourceState, InitialMipData, NumInitialMips);
+		FTextureRHIRef Texture = RHI->RHIAsyncCreateTexture2D(SizeX, SizeY, Format, NumMips, Flags, InResourceState, InitialMipData, NumInitialMips, OutCompletionEvent);
 		ensure(Texture->IsBarrierTrackingInitialized());
 		return Texture;
 	}
diff --git a/Engine/Source/Runtime/VulkanRHI/Private/VulkanTexture.cpp b/Engine/Source/Runtime/VulkanRHI/Private/VulkanTexture.cpp
--- a/Engine/Source/Runtime/VulkanRHI/Private/VulkanTexture.cpp
+++ b/Engine/Source/Runtime/VulkanRHI/Private/VulkanTexture.cpp
@@ -947,7 +947,7 @@ FTextureRHIRef FVulkanDynamicRHI::RHICreateTexture(const FRHITextureCreateDesc&
 	return new FVulkanTexture(*Device, CreateDesc, nullptr);
 }
 
-FTextureRHIRef FVulkanDynamicRHI::RHIAsyncCreateTexture2D(uint32 SizeX,uint32 SizeY,uint8 Format,uint32 NumMips,ETextureCreateFlags Flags, ERHIAccess InResourceState,void** InitialMipData,uint32 NumInitialMips)
+FTextureRHIRef FVulkanDynamicRHI::RHIAsyncCreateTexture2D(uint32 SizeX,uint32 SizeY,uint8 Format,uint32 NumMips,ETextureCreateFlags Flags, ERHIAccess InResourceState,void** InitialMipData,uint32 NumInitialMips, FGraphEventRef& OutCompletionEvent)
 {
 	UE_LOG(LogVulkan, Fatal, TEXT("RHIAsyncCreateTexture2D is not supported"));
 	VULKAN_SIGNAL_UNIMPLEMENTED();
diff --git a/Engine/Source/Runtime/VulkanRHI/Public/VulkanDynamicRHI.h b/Engine/Source/Runtime/VulkanRHI/Public/VulkanDynamicRHI.h
--- a/Engine/Source/Runtime/VulkanRHI/Public/VulkanDynamicRHI.h
+++ b/Engine/Source/Runtime/VulkanRHI/Public/VulkanDynamicRHI.h
@@ -127,7 +127,7 @@ public:
 	virtual void RHIGetTextureMemoryStats(FTextureMemoryStats& OutStats) final override;
 	virtual bool RHIGetTextureMemoryVisualizeData(FColor* TextureData, int32 SizeX, int32 SizeY, int32 Pitch, int32 PixelSize) final override;
 	virtual FTextureRHIRef RHICreateTexture(const FRHITextureCreateDesc& CreateDesc) final override;
-	virtual FTextureRHIRef RHIAsyncCreateTexture2D(uint32 SizeX, uint32 SizeY, uint8 Format, uint32 NumMips, ETextureCreateFlags Flags, ERHIAccess InResourceState, void** InitialMipData, uint32 NumInitialMips) final override;
+	virtual FTextureRHIRef RHIAsyncCreateTexture2D(uint32 SizeX, uint32 SizeY, uint8 Format, uint32 NumMips, ETextureCreateFlags Flags, ERHIAccess InResourceState, void** InitialMipData, uint32 NumInitialMips, FGraphEventRef& OutCompletionEvent) final override;
 	virtual FShaderResourceViewRHIRef RHICreateShaderResourceView(FRHITexture* Texture, const FRHITextureSRVCreateInfo& CreateInfo) final override;
 	virtual uint32 RHIComputeMemorySize(FRHITexture* TextureRHI) final override;
 	virtual FTexture2DRHIRef RHIAsyncReallocateTexture2D(FRHITexture2D* Texture2D, int32 NewMipCount, int32 NewSizeX, int32 NewSizeY, FThreadSafeCounter* RequestStatus) final override;
diff --git a/Engine/Source/Runtime/Windows/D3D11RHI/Private/D3D11RHIPrivate.h b/Engine/Source/Runtime/Windows/D3D11RHI/Private/D3D11RHIPrivate.h
--- a/Engine/Source/Runtime/Windows/D3D11RHI/Private/D3D11RHIPrivate.h
+++ b/Engine/Source/Runtime/Windows/D3D11RHI/Private/D3D11RHIPrivate.h
@@ -506,7 +506,7 @@ public:
 	virtual void RHIGetTextureMemoryStats(FTextureMemoryStats& OutStats) final override;
 	virtual bool RHIGetTextureMemoryVisualizeData(FColor* TextureData,int32 SizeX,int32 SizeY,int32 Pitch,int32 PixelSize) final override;
 	virtual FTextureRHIRef RHICreateTexture(const FRHITextureCreateDesc& CreateDesc) final override;
-	virtual FTextureRHIRef RHIAsyncCreateTexture2D(uint32 SizeX, uint32 SizeY, uint8 Format, uint32 NumMips, ETextureCreateFlags Flags, ERHIAccess InResourceState, void** InitialMipData, uint32 NumInitialMips) final override;
+	virtual FTextureRHIRef RHIAsyncCreateTexture2D(uint32 SizeX, uint32 SizeY, uint8 Format, uint32 NumMips, ETextureCreateFlags Flags, ERHIAccess InResourceState, void** InitialMipData, uint32 NumInitialMips, FGraphEventRef& OutCompletionEvent) final override;
 	virtual FShaderResourceViewRHIRef RHICreateShaderResourceView(FRHITexture* TextureRHI, const FRHITextureSRVCreateInfo& CreateInfo) final override;
 	virtual void RHIGenerateMips(FRHITexture* Texture) final override;
 	virtual uint32 RHIComputeMemorySize(FRHITexture* TextureRHI) final override;
diff --git a/Engine/Source/Runtime/Windows/D3D11RHI/Private/D3D11Texture.cpp b/Engine/Source/Runtime/Windows/D3D11RHI/Private/D3D11Texture.cpp
--- a/Engine/Source/Runtime/Windows/D3D11RHI/Private/D3D11Texture.cpp
+++ b/Engine/Source/Runtime/Windows/D3D11RHI/Private/D3D11Texture.cpp
@@ -941,7 +941,7 @@ FTextureRHIRef FD3D11DynamicRHI::RHICreateTexture_RenderThread(class FRHICommand
 	return RHICreateTexture(CreateDesc);
 }
 
-FTextureRHIRef FD3D11DynamicRHI::RHIAsyncCreateTexture2D(uint32 SizeX, uint32 SizeY, uint8 Format, uint32 NumMips, ETextureCreateFlags Flags, ERHIAccess InResourceState, void** InitialMipData, uint32 NumInitialMips)
+FTextureRHIRef FD3D11DynamicRHI::RHIAsyncCreateTexture2D(uint32 SizeX, uint32 SizeY, uint8 Format, uint32 NumMips, ETextureCreateFlags Flags, ERHIAccess InResourceState, void** InitialMipData, uint32 NumInitialMips, FGraphEventRef& OutCompletionEvent)
 {
 	TArray<D3D11_SUBRESOURCE_DATA, TInlineAllocator<12>> SubresourceData;
 	SubresourceData.SetNumUninitialized(NumMips);
@@ -992,6 +992,8 @@ FTextureRHIRef FD3D11DynamicRHI::RHIAsyncCreateTexture2D(uint32 SizeX, uint32 Si
 		FMemory::Free(TempBuffer);
 	}
 
+	OutCompletionEvent = nullptr;
+
 	return Texture;
 }
 
