Fix GPU crashes when d3d12.PoolAllocator.RTUAVTextureMaxAllocationSize is non-zero.

UAVs of compressed textures (which need the aliased resource workaround) were not handled correctly when the resource was allocated from a pool. CreateUAVAliasResource was expecting Location.GetOffsetFromBaseOfResource to return the offset to the start of the heap, but that's not the case for placed resources. Also, defrag isn't handling these resources at all, so we need to disable defrag on the RTV and UAV pools until proper support is added.

Fixed an unrelated crash which happened when MSAA targets were pooled, because the pools have lower alignment requirements than these resources.

#rnx

diff --git a/Engine/Source/Runtime/D3D12RHI/Private/D3D12Allocation.cpp b/Engine/Source/Runtime/D3D12RHI/Private/D3D12Allocation.cpp
--- a/Engine/Source/Runtime/D3D12RHI/Private/D3D12Allocation.cpp
+++ b/Engine/Source/Runtime/D3D12RHI/Private/D3D12Allocation.cpp
@@ -1686,7 +1686,8 @@ FD3D12TextureAllocatorPool::FD3D12TextureAllocatorPool(FD3D12Device* Device, FRH
 		const FString Name(L"D3D12 RT Texture Pool Allocator");
 		uint64 PoolSize = GD3D12PoolAllocatorRTUAVTextureVRAMPoolSize;
 		uint64 MaxAllocationSize = GD3D12PoolAllocatorRTUAVTextureVRAMMaxAllocationSize;
-		bool bDefragEnabled = true;
+		// FD3D12ResourceLocation::OnAllocationMoved doesn't correctly retrieve the clear value when recreating moved resources, so we need to disable defrag for this pool for the time being.
+		bool bDefragEnabled = false;
 		FD3D12PoolAllocator* RTPool = new FD3D12PoolAllocator(Device, GetVisibilityMask(), InitConfig, Name, AllocationStrategy, PoolSize, D3D12_DEFAULT_RESOURCE_PLACEMENT_ALIGNMENT, MaxAllocationSize, FreeListOrder, bDefragEnabled, TraceHeapId);
 		RTPool->Initialize();
 
@@ -1700,7 +1701,8 @@ FD3D12TextureAllocatorPool::FD3D12TextureAllocatorPool(FD3D12Device* Device, FRH
 		const FString Name(L"D3D12 UAV Texture Pool Allocator");
 		uint64 PoolSize = GD3D12PoolAllocatorRTUAVTextureVRAMPoolSize;
 		uint64 MaxAllocationSize = GD3D12PoolAllocatorRTUAVTextureVRAMMaxAllocationSize;
-		bool bDefragEnabled = true;
+		// Defrag doesn't correctly handle resources which need the BCn/UINT UAV aliasing workaround, so we'll turn off defrag for this heap for now.
+		bool bDefragEnabled = false;
 		FD3D12PoolAllocator* UAVPool = new FD3D12PoolAllocator(Device, GetVisibilityMask(), InitConfig, Name, AllocationStrategy, PoolSize, D3D12_DEFAULT_RESOURCE_PLACEMENT_ALIGNMENT, MaxAllocationSize, FreeListOrder, bDefragEnabled, TraceHeapId);
 		UAVPool->Initialize();
 
diff --git a/Engine/Source/Runtime/D3D12RHI/Private/D3D12PoolAllocator.cpp b/Engine/Source/Runtime/D3D12RHI/Private/D3D12PoolAllocator.cpp
--- a/Engine/Source/Runtime/D3D12RHI/Private/D3D12PoolAllocator.cpp
+++ b/Engine/Source/Runtime/D3D12RHI/Private/D3D12PoolAllocator.cpp
@@ -315,22 +315,32 @@ void FD3D12PoolAllocator::AllocateResource(uint32 GPUIndex, D3D12_HEAP_TYPE InHe
 	check(GPUIndex == Device->GetGPUIndex());
 
 	const bool bSharedResource = EnumHasAnyFlags(InDesc.Flags, D3D12_RESOURCE_FLAG_ALLOW_SIMULTANEOUS_ACCESS);
-	const bool PoolResource = InSize <= MaxAllocationSize && !bSharedResource;
+	
+	bool bPoolResource = InSize <= MaxAllocationSize && !bSharedResource;
+	bool bPlacedResource = bPoolResource && (AllocationStrategy == EResourceAllocationStrategy::kPlacedResource);
+	if (bPlacedResource)
+	{
+		// Make sure that we can satisfy the alignment requirements. If not, we have to fall back to a committed resource. This can happen for MSAA targets,
+		// because the pools are aligned to D3D12_DEFAULT_RESOURCE_PLACEMENT_ALIGNMENT, but MSAA needs more.
+		if (InAllocationAlignment > PoolAlignment || InDesc.Alignment > PoolAlignment)
+		{
+			bPoolResource = false;
+			bPlacedResource = false;
+		}
+	}
 
 	// Disable pooling for VRAM allocated textures and force use the committed resource path
 	bool bForceCommittedResourcePath = false;
 #if PLATFORM_WINDOWS
 	const bool bUseCommittedTexturesNV = GUseCommittedTexturesNV && IsRHIDeviceNVIDIA();
-	if (PoolResource && bUseCommittedTexturesNV && InHeapType == D3D12_HEAP_TYPE_DEFAULT && InDesc.Dimension != D3D12_RESOURCE_DIMENSION_BUFFER)
+	if (bPoolResource && bUseCommittedTexturesNV && InHeapType == D3D12_HEAP_TYPE_DEFAULT && InDesc.Dimension != D3D12_RESOURCE_DIMENSION_BUFFER)
 	{
 		bForceCommittedResourcePath = true;
 	}
 #endif // PLATFORM_WINDOWS
 
-	if (PoolResource && !bForceCommittedResourcePath)
+	if (bPoolResource && !bForceCommittedResourcePath)
 	{
-		const bool bPlacedResource = (AllocationStrategy == EResourceAllocationStrategy::kPlacedResource);
-
 		uint32 AllocationAlignment = InAllocationAlignment;
 
 		// Ensure we're allocating from the correct pool
@@ -338,8 +348,6 @@ void FD3D12PoolAllocator::AllocateResource(uint32 GPUIndex, D3D12_HEAP_TYPE InHe
 		{
 			// If it's a placed resource then base offset will always be 0 from the actual d3d resource so ignore the allocation alignment - no extra offset required
 			// for creating the views!
-			check(InAllocationAlignment <= PoolAlignment);
-			check(InDesc.Alignment <= PoolAlignment);
 			AllocationAlignment = PoolAlignment;
 		}
 		else
diff --git a/Engine/Source/Runtime/D3D12RHI/Private/D3D12Resources.cpp b/Engine/Source/Runtime/D3D12RHI/Private/D3D12Resources.cpp
--- a/Engine/Source/Runtime/D3D12RHI/Private/D3D12Resources.cpp
+++ b/Engine/Source/Runtime/D3D12RHI/Private/D3D12Resources.cpp
@@ -663,7 +663,15 @@ void FD3D12Adapter::CreateUAVAliasResource(D3D12_CLEAR_VALUE* ClearValuePtr, con
 
 	if (ensure(ResourceHeap != nullptr) && ensure(SourceFormat != PF_Unknown) && SourceFormat != AliasTextureFormat)
 	{
-		const uint64 SourceOffset = Location.GetOffsetFromBaseOfResource();
+		uint64 SourceOffset = Location.GetOffsetFromBaseOfResource();
+
+		if (Location.GetAllocatorType() == FD3D12ResourceLocation::AT_Pool && Location.GetType() == FD3D12ResourceLocation::ResourceLocationType::eSubAllocation)
+		{
+			FD3D12HeapAndOffset HeapAndOffset = Location.GetPoolAllocator()->GetBackingHeapAndAllocationOffsetInBytes(Location);
+			check(SourceOffset == 0);
+			check(ResourceHeap == HeapAndOffset.Heap);
+			SourceOffset = HeapAndOffset.Offset;
+		}
 
 		FD3D12ResourceDesc AliasTextureDesc = SourceDesc;
 		AliasTextureDesc.Format = (DXGI_FORMAT)GPixelFormats[AliasTextureFormat].PlatformFormat;
@@ -1143,7 +1151,6 @@ bool FD3D12ResourceLocation::OnAllocationMoved(FRHIPoolAllocationData* InNewData
 
 		FD3D12Resource* NewResource = nullptr;
 		VERIFYD3D12RESULT(CurrentResource->GetParentDevice()->GetParentAdapter()->CreatePlacedResource(CurrentResource->GetDesc(), HeapAndOffset.Heap, HeapAndOffset.Offset, CreateState, ResourceStateMode, D3D12_RESOURCE_STATE_TBD, ClearValue, &NewResource, *Name.ToString()));
-
 		UnderlyingResource = NewResource;
 		GPUVirtualAddress = UnderlyingResource->GetGPUVirtualAddress() + OffsetFromBaseOfResource;
 	}
@@ -1161,6 +1168,10 @@ bool FD3D12ResourceLocation::OnAllocationMoved(FRHIPoolAllocationData* InNewData
 		OtherResourceLocation->ResidencyHandle = ResidencyHandle;
 	}
 
+	// TODO: recreate the aliased UAV resource if we have one. For the time being we just check that we don't get here with such a resource,
+	// because defrag is disabled for the UAV pool in the FD3D12TextureAllocatorPool constructor.
+	check(!CurrentResource->GetDesc().NeedsUAVAliasWorkarounds());
+
 	// Notify all the dependent resources about the change
 	Owner->ResourceRenamed(this);
 
