Speculative fix for an access violation crash. The suspicion is that the RHI texture is deleted before FD3D12DynamicRHI::RHIAsyncCreateTexture2D returns.

diff --git a/Engine/Source/Runtime/D3D12RHI/Private/D3D12Texture.cpp b/Engine/Source/Runtime/D3D12RHI/Private/D3D12Texture.cpp
--- a/Engine/Source/Runtime/D3D12RHI/Private/D3D12Texture.cpp
+++ b/Engine/Source/Runtime/D3D12RHI/Private/D3D12Texture.cpp
@@ -1155,6 +1155,7 @@ FTextureRHIRef FD3D12DynamicRHI::RHIAsyncCreateTexture2D(uint32 SizeX, uint32 Si
 
 		return NewTexture;
 	});
+	FTextureRHIRef TextureOutRHI = TextureOut;
 
 	FGraphEventArray CopyCompleteEvents;
 	OutCompletionEvent = nullptr;
@@ -1325,7 +1326,7 @@ FTextureRHIRef FD3D12DynamicRHI::RHIAsyncCreateTexture2D(uint32 SizeX, uint32 Si
 		FMemory::Free(TempBuffer);
 	}
 
-	return TextureOut;
+	return MoveTemp(TextureOutRHI);
 }
 
 
