Add Null guard to prevent DrawIndexedPrimitive crash

#rnx

diff --git a/Engine/Source/Runtime/D3D12RHI/Private/D3D12Commands.cpp b/Engine/Source/Runtime/D3D12RHI/Private/D3D12Commands.cpp
--- a/Engine/Source/Runtime/D3D12RHI/Private/D3D12Commands.cpp
+++ b/Engine/Source/Runtime/D3D12RHI/Private/D3D12Commands.cpp
@@ -33,6 +33,17 @@ static FAutoConsoleVariableRef CVarD3D12AllowDiscardResources(
 	TEXT("Whether to call DiscardResources after transient aliasing acquire. This is not needed on some platforms if newly acquired resources are cleared before use."),
 	ECVF_RenderThreadSafe);
 
+#if UE_BUILD_SHIPPING
+int32 GD3D12SkipDrawIndexedWithNullIB = 1;
+#else
+int32 GD3D12SkipDrawIndexedWithNullIB = 2;
+#endif
+static FAutoConsoleVariableRef CVarSkipDrawIndexedWithNullIB(
+	TEXT("D3D12.SkipDrawIndexedWithNullIB"),
+	GD3D12SkipDrawIndexedWithNullIB,
+	TEXT("1 == Skip RHIDrawIndexedPrimitive when IB is nullptr, 2 == Crash with infos"),
+	ECVF_RenderThreadSafe);
+
 using namespace D3D12RHI;
 
 inline void ValidateBoundShader(FD3D12StateCache& InStateCache, FRHIShader* InShaderRHI)
@@ -1627,6 +1638,20 @@ void FD3D12CommandContext::RHIDrawIndexedPrimitive(FRHIBuffer* IndexBufferRHI, i
 
 	NumInstances = FMath::Max<uint32>(1, NumInstances);
 
+	FD3D12Buffer* IndexBuffer = RetrieveObject<FD3D12Buffer>(IndexBufferRHI);
+
+	if (GD3D12SkipDrawIndexedWithNullIB != 0 && IndexBuffer->ResourceLocation.GetResource() == nullptr)
+	{
+		if (GD3D12SkipDrawIndexedWithNullIB == 2)
+		{
+			FD3D12ResourceLocation* StackCpyResourceLocation = (FD3D12ResourceLocation*)alloca(sizeof(FD3D12ResourceLocation));
+			FMemory::Memcpy(StackCpyResourceLocation, &IndexBuffer->ResourceLocation, sizeof(FD3D12ResourceLocation));
+			PLATFORM_BREAK();
+		}
+
+		return;
+	}
+
 	if (bTrackingEvents)
 	{
 		GetParentDevice()->RegisterGPUWork(NumPrimitives * NumInstances, NumVertices * NumInstances);
@@ -1637,8 +1662,6 @@ void FD3D12CommandContext::RHIDrawIndexedPrimitive(FRHIBuffer* IndexBufferRHI, i
 
 	uint32 IndexCount = StateCache.GetVertexCount(NumPrimitives);
 
-	FD3D12Buffer* IndexBuffer = RetrieveObject<FD3D12Buffer>(IndexBufferRHI);
-
 	// Verify that we are not trying to read outside the index buffer range
 	// test is an optimized version of: StartIndex + IndexCount <= IndexBuffer->GetSize() / IndexBuffer->GetStride() 
 	checkf((StartIndex + IndexCount) * IndexBuffer->GetStride() <= IndexBuffer->GetSize(),
