Remove D3D12_PLATFORM_SUPPORTS_RESOLVE_SHADERS which is unreferenced

- remove preprocessor def for old platform specific workaround which is no longer implemented

#rnx

diff --git a/Engine/Source/Runtime/D3D12RHI/Private/D3D12Commands.cpp b/Engine/Source/Runtime/D3D12RHI/Private/D3D12Commands.cpp
--- a/Engine/Source/Runtime/D3D12RHI/Private/D3D12Commands.cpp
+++ b/Engine/Source/Runtime/D3D12RHI/Private/D3D12Commands.cpp
@@ -106,10 +106,6 @@ static void BindUniformBuffer(FD3D12CommandContext& Context, FRHIShader* Shader,
 	Context.DirtyUniformBuffers[ShaderFrequency] |= (1 << BufferIndex);
 }
 
-#if !defined(D3D12_PLATFORM_SUPPORTS_RESOLVE_SHADERS)
-	#define D3D12_PLATFORM_SUPPORTS_RESOLVE_SHADERS 1
-#endif
-
 // Vertex state.
 void FD3D12CommandContext::RHISetStreamSource(uint32 StreamIndex, FRHIBuffer* VertexBufferRHI, uint32 Offset)
 {
