RHI: Move to common RHI function for preferred vendors.

#rnx

diff --git a/Engine/Source/Runtime/D3D12RHI/Private/Windows/WindowsD3D12Device.cpp b/Engine/Source/Runtime/D3D12RHI/Private/Windows/WindowsD3D12Device.cpp
--- a/Engine/Source/Runtime/D3D12RHI/Private/Windows/WindowsD3D12Device.cpp
+++ b/Engine/Source/Runtime/D3D12RHI/Private/Windows/WindowsD3D12Device.cpp
@@ -37,6 +37,7 @@
 #include "IHeadMountedDisplayModule.h"
 #include "GenericPlatform/GenericPlatformDriver.h"			// FGPUDriverInfo
 #include "RHIValidation.h"
+#include "RHIUtilities.h"
 
 #include "ShaderCompiler.h"
 
@@ -142,26 +143,6 @@ static FAutoConsoleVariableRef CVarAutoAttachPIX(
 );
 #endif // D3D12RHI_SUPPORTS_WIN_PIX
 
-static inline int D3D12RHI_PreferAdapterVendor()
-{
-	if (FParse::Param(FCommandLine::Get(), TEXT("preferAMD")))
-	{
-		return 0x1002;
-	}
-
-	if (FParse::Param(FCommandLine::Get(), TEXT("preferIntel")))
-	{
-		return 0x8086;
-	}
-
-	if (FParse::Param(FCommandLine::Get(), TEXT("preferNvidia")))
-	{
-		return 0x10DE;
-	}
-
-	return -1;
-}
-
 using namespace D3D12RHI;
 
 static bool bIsQuadBufferStereoEnabled = false;
@@ -402,7 +383,7 @@ static bool CheckDeviceForEmulatedAtomic64Support(IDXGIAdapter* Adapter, ID3D12D
 	DXGI_ADAPTER_DESC AdapterDesc{};
 	Adapter->GetDesc(&AdapterDesc);
 
-	if (AdapterDesc.VendorId == 0x8086 && !FParse::Param(FCommandLine::Get(), TEXT("novendordevice")))
+	if ((RHIConvertToGpuVendorId(AdapterDesc.VendorId) == EGpuVendorId::Intel) && !FParse::Param(FCommandLine::Get(), TEXT("novendordevice")))
 	{
 		INTCExtensionInfo INTCExtensionInfo{};
 		if (INTCExtensionContext* IntelExtensionContext = CreateIntelExtensionsContext(Device, INTCExtensionInfo))
@@ -866,7 +847,7 @@ void FD3D12DynamicRHIModule::FindAdapter()
 	default: GpuPreference = DXGI_GPU_PREFERENCE_UNSPECIFIED; break;
 	}
 
-	int PreferredVendor = D3D12RHI_PreferAdapterVendor();
+	const EGpuVendorId PreferredVendor = RHIGetPreferredAdapterVendor();
 
 	// Enumerate the DXGIFactory's adapters.
 	for (uint32 AdapterIndex = 0; FD3D12AdapterDesc::EnumAdapters(AdapterIndex, GpuPreference, DXGIFactory4, DXGIFactory6, TempAdapter.GetInitReference()) != DXGI_ERROR_NOT_FOUND; ++AdapterIndex)
@@ -908,7 +889,7 @@ void FD3D12DynamicRHIModule::FindAdapter()
 					OutputCount
 				);
 
-				const bool bIsWARP = AdapterDesc.VendorId == 0x1414;
+				const bool bIsWARP = (RHIConvertToGpuVendorId(AdapterDesc.VendorId) == EGpuVendorId::Microsoft);
 
 #if !PLATFORM_HOLOLENS
 				if (!bIsWARP)
@@ -944,7 +925,7 @@ void FD3D12DynamicRHIModule::FindAdapter()
 
 				if (!bSkipAdapter)
 				{
-					if (PreferredVendor == AdapterDesc.VendorId && FirstDiscreteAdapter.IsValid())
+					if ((PreferredVendor != EGpuVendorId::Unknown) && (PreferredVendor == RHIConvertToGpuVendorId(AdapterDesc.VendorId)) && FirstDiscreteAdapter.IsValid())
 					{
 						FirstDiscreteAdapter = CurrentAdapter;
 					}
@@ -965,7 +946,7 @@ void FD3D12DynamicRHIModule::FindAdapter()
 					{
 						FirstAdapter = CurrentAdapter;
 					}
-					else if (PreferredVendor == AdapterDesc.VendorId && FirstAdapter.IsValid())
+					else if ((PreferredVendor != EGpuVendorId::Unknown) && (PreferredVendor == RHIConvertToGpuVendorId(AdapterDesc.VendorId)) && FirstAdapter.IsValid())
 					{
 						FirstAdapter = CurrentAdapter;
 					}
diff --git a/Engine/Source/Runtime/RHI/Private/RHIStrings.cpp b/Engine/Source/Runtime/RHI/Private/RHIStrings.cpp
--- a/Engine/Source/Runtime/RHI/Private/RHIStrings.cpp
+++ b/Engine/Source/Runtime/RHI/Private/RHIStrings.cpp
@@ -72,6 +72,7 @@ const TCHAR* RHIVendorIdToString(EGpuVendorId VendorId)
 	case EGpuVendorId::Codeplay:	return TEXT("Codeplay");
 	case EGpuVendorId::Mesa:		return TEXT("Mesa");
 	case EGpuVendorId::SamsungAMD:	return TEXT("SamsungAMD");
+	case EGpuVendorId::Microsoft:	return TEXT("Microsoft");
 	case EGpuVendorId::NotQueried:	return TEXT("Not Queried");
 	default:                        return TEXT("Unknown");
 	}
diff --git a/Engine/Source/Runtime/RHI/Private/RHIUtilities.cpp b/Engine/Source/Runtime/RHI/Private/RHIUtilities.cpp
--- a/Engine/Source/Runtime/RHI/Private/RHIUtilities.cpp
+++ b/Engine/Source/Runtime/RHI/Private/RHIUtilities.cpp
@@ -733,3 +733,28 @@ void DecodeRenderTargetMode(ESimpleRenderTargetMode Mode, ERenderTargetLoadActio
 		StencilStoreAction = ERenderTargetStoreAction::ENoAction;
 	}
 }
+
+EGpuVendorId RHIGetPreferredAdapterVendor()
+{
+	if (FParse::Param(FCommandLine::Get(), TEXT("preferAMD")))
+	{
+		return EGpuVendorId::Amd;
+	}
+
+	if (FParse::Param(FCommandLine::Get(), TEXT("preferIntel")))
+	{
+		return EGpuVendorId::Intel;
+	}
+
+	if (FParse::Param(FCommandLine::Get(), TEXT("preferNvidia")))
+	{
+		return EGpuVendorId::Nvidia;
+	}
+
+	if (FParse::Param(FCommandLine::Get(), TEXT("preferMS")) || FParse::Param(FCommandLine::Get(), TEXT("preferMicrosoft")))
+	{
+		return EGpuVendorId::Microsoft;
+	}
+
+	return EGpuVendorId::Unknown;
+}
diff --git a/Engine/Source/Runtime/RHI/Public/RHIDefinitions.h b/Engine/Source/Runtime/RHI/Public/RHIDefinitions.h
--- a/Engine/Source/Runtime/RHI/Public/RHIDefinitions.h
+++ b/Engine/Source/Runtime/RHI/Public/RHIDefinitions.h
@@ -820,6 +820,7 @@ enum class EGpuVendorId
 	Vivante		= 0x7a05,
 	VeriSilicon	= 0x1EB1,
 	SamsungAMD  = 0x144D,
+	Microsoft   = 0x1414,
 
 	Kazan		= 0x10003,	// VkVendorId
 	Codeplay	= 0x10004,	// VkVendorId
@@ -1286,6 +1287,7 @@ inline EGpuVendorId RHIConvertToGpuVendorId(uint32 VendorId)
 	case EGpuVendorId::Qualcomm:
 	case EGpuVendorId::Intel:
 	case EGpuVendorId::SamsungAMD:
+	case EGpuVendorId::Microsoft:
 		return (EGpuVendorId)VendorId;
 
 	default:
diff --git a/Engine/Source/Runtime/RHI/Public/RHIUtilities.h b/Engine/Source/Runtime/RHI/Public/RHIUtilities.h
--- a/Engine/Source/Runtime/RHI/Public/RHIUtilities.h
+++ b/Engine/Source/Runtime/RHI/Public/RHIUtilities.h
@@ -574,6 +574,9 @@ extern RHI_API float RHIGetFrameTime();
 
 extern RHI_API void RHICalculateFrameTime();
 
+/** Returns the VendorID of the preferred vendor or -1 if none were specified. */
+extern RHI_API EGpuVendorId RHIGetPreferredAdapterVendor();
+
 #if UE_ENABLE_INCLUDE_ORDER_DEPRECATED_IN_5_3
 #include "RHILockTracker.h"
 #include "PixelFormat.h"
diff --git a/Engine/Source/Runtime/VulkanRHI/Private/VulkanRHI.cpp b/Engine/Source/Runtime/VulkanRHI/Private/VulkanRHI.cpp
--- a/Engine/Source/Runtime/VulkanRHI/Private/VulkanRHI.cpp
+++ b/Engine/Source/Runtime/VulkanRHI/Private/VulkanRHI.cpp
@@ -21,6 +21,7 @@
 #include "Misc/EngineVersion.h"
 #include "GlobalShader.h"
 #include "RHIValidation.h"
+#include "RHIUtilities.h"
 #include "IHeadMountedDisplayModule.h"
 #include "VulkanRenderpass.h"
 #include "VulkanTransientResourceAllocator.h"
@@ -508,27 +509,6 @@ void FVulkanDynamicRHI::CreateInstance()
 #endif
 }
 
-//#todo-rco: Common RHI should handle this...
-static inline int32 PreferAdapterVendor()
-{
-	if (FParse::Param(FCommandLine::Get(), TEXT("preferAMD")))
-	{
-		return 0x1002;
-	}
-
-	if (FParse::Param(FCommandLine::Get(), TEXT("preferIntel")))
-	{
-		return 0x8086;
-	}
-
-	if (FParse::Param(FCommandLine::Get(), TEXT("preferNvidia")))
-	{
-		return 0x10DE;
-	}
-
-	return -1;
-}
-
 void FVulkanDynamicRHI::SelectDevice()
 {
 	uint32 GpuCount = 0;
@@ -639,13 +619,13 @@ void FVulkanDynamicRHI::SelectDevice()
 			}
 			else if (DiscreteDevices.Num() > 0 && CVarExplicitAdapterValue == -1)
 			{
-				int32 PreferredVendor = PreferAdapterVendor();
-				if (DiscreteDevices.Num() > 1 && PreferredVendor != -1)
+				const EGpuVendorId PreferredVendor = RHIGetPreferredAdapterVendor();
+				if (DiscreteDevices.Num() > 1 && (PreferredVendor != EGpuVendorId::Unknown))
 				{
 					// Check for preferred
 					for (int32 Index = 0; Index < DiscreteDevices.Num(); ++Index)
 					{
-						if (DiscreteDevices[Index].Device->GpuProps.vendorID == PreferredVendor)
+						if (RHIConvertToGpuVendorId(DiscreteDevices[Index].Device->GpuProps.vendorID) == PreferredVendor)
 						{
 							DeviceIndex = DiscreteDevices[Index].DeviceIndex;
 							Device = DiscreteDevices[Index].Device;
diff --git a/Engine/Source/Runtime/Windows/D3D11RHI/Private/Windows/WindowsD3D11Device.cpp b/Engine/Source/Runtime/Windows/D3D11RHI/Private/Windows/WindowsD3D11Device.cpp
--- a/Engine/Source/Runtime/Windows/D3D11RHI/Private/Windows/WindowsD3D11Device.cpp
+++ b/Engine/Source/Runtime/Windows/D3D11RHI/Private/Windows/WindowsD3D11Device.cpp
@@ -24,6 +24,7 @@
 #include "GenericPlatform/GenericPlatformDriver.h"			// FGPUDriverInfo
 #include "GenericPlatform/GenericPlatformCrashContext.h"
 #include "RHIValidation.h"
+#include "RHIUtilities.h"
 #include "HAL/ExceptionHandling.h"
 #include "HDRHelper.h"
 #include "GlobalShader.h"
@@ -43,31 +44,6 @@ FD3D11DynamicRHI*	GD3D11RHI = nullptr;
 extern bool D3D11RHI_ShouldCreateWithD3DDebug();
 extern bool D3D11RHI_ShouldAllowAsyncResourceCreation();
 
-static int D3D11RHI_PreferAdapterVendor()
-{
-	if (FParse::Param(FCommandLine::Get(), TEXT("preferAMD")))
-	{
-		return 0x1002;
-	}
-
-	if (FParse::Param(FCommandLine::Get(), TEXT("preferIntel")))
-	{
-		return 0x8086;
-	}
-
-	if (FParse::Param(FCommandLine::Get(), TEXT("preferNvidia")))
-	{
-		return 0x10DE;
-	}
-
-	if (FParse::Param(FCommandLine::Get(), TEXT("preferMS")))
-	{
-		return 0x1414;
-	}
-
-	return -1;
-}
-
 static bool D3D11RHI_AllowSoftwareFallback()
 {
 	if (FParse::Param(FCommandLine::Get(), TEXT("AllowSoftwareRendering")))
@@ -1020,7 +996,7 @@ void FD3D11DynamicRHIModule::FindAdapter()
 
 	UE_LOG(LogD3D11RHI, Log, TEXT("D3D11 adapters:"));
 
-	int PreferredVendor = D3D11RHI_PreferAdapterVendor();
+	const EGpuVendorId PreferredVendor = RHIGetPreferredAdapterVendor();
 	bool bAllowSoftwareFallback = D3D11RHI_AllowSoftwareFallback();
 
 
@@ -1132,7 +1108,7 @@ void FD3D11DynamicRHIModule::FindAdapter()
 					{
 						FirstWithoutIntegratedAdapter = CurrentAdapter;
 					}
-					else if (PreferredVendor == AdapterDesc.VendorId && FirstWithoutIntegratedAdapter.IsValid())
+					else if ((PreferredVendor != EGpuVendorId::Unknown) && (PreferredVendor == RHIConvertToGpuVendorId(AdapterDesc.VendorId)) && FirstWithoutIntegratedAdapter.IsValid())
 					{
 						FirstWithoutIntegratedAdapter = CurrentAdapter;
 					}
@@ -1141,7 +1117,7 @@ void FD3D11DynamicRHIModule::FindAdapter()
 					{
 						FirstAdapter = CurrentAdapter;
 					}
-					else if (PreferredVendor == AdapterDesc.VendorId && FirstAdapter.IsValid())
+					else if ((PreferredVendor != EGpuVendorId::Unknown) && (PreferredVendor == RHIConvertToGpuVendorId(AdapterDesc.VendorId)) && FirstAdapter.IsValid())
 					{
 						FirstAdapter = CurrentAdapter;
 					}
