Skip D3D12 drawcall in RHIDrawIndexedPrimitive and RHIDrawIndexedPrimitiveIndirect if IB is empty or underlying resource is null

#rnx

diff --git a/Engine/Source/Runtime/D3D12RHI/Private/D3D12Commands.cpp b/Engine/Source/Runtime/D3D12RHI/Private/D3D12Commands.cpp
--- a/Engine/Source/Runtime/D3D12RHI/Private/D3D12Commands.cpp
+++ b/Engine/Source/Runtime/D3D12RHI/Private/D3D12Commands.cpp
@@ -33,17 +33,6 @@ static FAutoConsoleVariableRef CVarD3D12AllowDiscardResources(
 	TEXT("Whether to call DiscardResources after transient aliasing acquire. This is not needed on some platforms if newly acquired resources are cleared before use."),
 	ECVF_RenderThreadSafe);
 
-#if UE_BUILD_SHIPPING
-int32 GD3D12SkipDrawIndexedWithNullIB = 1;
-#else
-int32 GD3D12SkipDrawIndexedWithNullIB = 2;
-#endif
-static FAutoConsoleVariableRef CVarSkipDrawIndexedWithNullIB(
-	TEXT("D3D12.SkipDrawIndexedWithNullIB"),
-	GD3D12SkipDrawIndexedWithNullIB,
-	TEXT("1 == Skip RHIDrawIndexedPrimitive when IB is nullptr, 2 == Crash with infos"),
-	ECVF_RenderThreadSafe);
-
 using namespace D3D12RHI;
 
 inline void ValidateBoundShader(FD3D12StateCache& InStateCache, FRHIShader* InShaderRHI)
@@ -1628,26 +1617,22 @@ void FD3D12CommandContext::RHIDrawIndexedIndirect(FRHIBuffer* IndexBufferRHI, FR
 
 void FD3D12CommandContext::RHIDrawIndexedPrimitive(FRHIBuffer* IndexBufferRHI, int32 BaseVertexIndex, uint32 FirstInstance, uint32 NumVertices, uint32 StartIndex, uint32 NumPrimitives, uint32 NumInstances)
 {
+	FD3D12Buffer* IndexBuffer = RetrieveObject<FD3D12Buffer>(IndexBufferRHI);
+
 	// called should make sure the input is valid, this avoid hidden bugs
 	ensure(NumPrimitives > 0);
-	RHI_DRAW_CALL_STATS(StateCache.GetGraphicsPipelinePrimitiveType(), FMath::Max(NumInstances, 1U) * NumPrimitives);
+	ensure(IndexBufferRHI->GetSize() > 0);
+	ensure(IndexBuffer->ResourceLocation.GetResource() != nullptr);
 
-	NumInstances = FMath::Max<uint32>(1, NumInstances);
-
-	FD3D12Buffer* IndexBuffer = RetrieveObject<FD3D12Buffer>(IndexBufferRHI);
-
-	if (GD3D12SkipDrawIndexedWithNullIB != 0 && IndexBuffer->ResourceLocation.GetResource() == nullptr)
+	if (IndexBufferRHI->GetSize() == 0 || IndexBuffer->ResourceLocation.GetResource() == nullptr)
 	{
-		if (GD3D12SkipDrawIndexedWithNullIB == 2)
-		{
-			FD3D12ResourceLocation* StackCpyResourceLocation = (FD3D12ResourceLocation*)alloca(sizeof(FD3D12ResourceLocation));
-			FMemory::Memcpy(StackCpyResourceLocation, &IndexBuffer->ResourceLocation, sizeof(FD3D12ResourceLocation));
-			PLATFORM_BREAK();
-		}
-
 		return;
 	}
 
+	RHI_DRAW_CALL_STATS(StateCache.GetGraphicsPipelinePrimitiveType(), FMath::Max(NumInstances, 1U) * NumPrimitives);
+
+	NumInstances = FMath::Max<uint32>(1, NumInstances);
+
 	if (bTrackingEvents)
 	{
 		GetParentDevice()->RegisterGPUWork(NumPrimitives * NumInstances, NumVertices * NumInstances);
@@ -1677,8 +1662,18 @@ void FD3D12CommandContext::RHIDrawIndexedPrimitive(FRHIBuffer* IndexBufferRHI, i
 
 void FD3D12CommandContext::RHIDrawIndexedPrimitiveIndirect(FRHIBuffer* IndexBufferRHI, FRHIBuffer* ArgumentBufferRHI, uint32 ArgumentOffset)
 {
-	const uint32 IndexBufferStride = FD3D12DynamicRHI::ResourceCast(IndexBufferRHI)->GetStride();
 	FD3D12Buffer* IndexBuffer = RetrieveObject<FD3D12Buffer>(IndexBufferRHI);
+
+	// called should make sure the input is valid, this avoid hidden bugs
+	ensure(IndexBufferRHI->GetSize() > 0);
+	ensure(IndexBuffer->ResourceLocation.GetResource() != nullptr);
+
+	if (IndexBufferRHI->GetSize() == 0 || IndexBuffer->ResourceLocation.GetResource() == nullptr)
+	{
+		return;
+	}
+
+	const uint32 IndexBufferStride = FD3D12DynamicRHI::ResourceCast(IndexBufferRHI)->GetStride();
 	FD3D12Buffer* ArgumentBuffer = RetrieveObject<FD3D12Buffer>(ArgumentBufferRHI);
 
 	RHI_DRAW_CALL_INC();
