Add package field to asset metadata tracing.

Add a package field to the asset metadat scope, in addition to the existing asset name and object class fields.

diff --git a/Engine/Source/Developer/TraceInsights/Private/Insights/MemoryProfiler/ViewModels/MemAllocTable.cpp b/Engine/Source/Developer/TraceInsights/Private/Insights/MemoryProfiler/ViewModels/MemAllocTable.cpp
--- a/Engine/Source/Developer/TraceInsights/Private/Insights/MemoryProfiler/ViewModels/MemAllocTable.cpp
+++ b/Engine/Source/Developer/TraceInsights/Private/Insights/MemoryProfiler/ViewModels/MemAllocTable.cpp
@@ -38,6 +38,7 @@ const FName FMemAllocTableColumns::CountColumnId(TEXT("Count"));
 const FName FMemAllocTableColumns::SizeColumnId(TEXT("Size"));
 const FName FMemAllocTableColumns::TagColumnId(TEXT("Tag"));
 const FName FMemAllocTableColumns::AssetColumnId(TEXT("Asset"));
+const FName FMemAllocTableColumns::PackageColumnId(TEXT("Package"));
 const FName FMemAllocTableColumns::ClassNameColumnId(TEXT("ClassName"));
 const FName FMemAllocTableColumns::FunctionColumnId(TEXT("Function"));
 const FName FMemAllocTableColumns::SourceFileColumnId(TEXT("SourceFile"));
@@ -737,6 +738,64 @@ void FMemAllocTable::AddDefaultColumns()
 		AddColumn(ColumnRef);
 	}
 	//////////////////////////////////////////////////
+	// Package Column
+	{
+		TSharedRef<FTableColumn> ColumnRef = MakeShared<FTableColumn>(FMemAllocTableColumns::PackageColumnId);
+		FTableColumn& Column = *ColumnRef;
+
+		Column.SetIndex(ColumnIndex++);
+
+		Column.SetShortName(LOCTEXT("PackageColumnName", "Package"));
+		Column.SetTitleName(LOCTEXT("PackageColumnTitle", "Package"));
+		Column.SetDescription(LOCTEXT("PackageColumnDesc", "Asset package associated with allocation"));
+
+		Column.SetFlags(ETableColumnFlags::CanBeHidden | ETableColumnFlags::CanBeFiltered);
+
+		Column.SetHorizontalAlignment(HAlign_Left);
+		Column.SetInitialWidth(320.0f);
+
+		Column.SetDataType(ETableCellDataType::CString);
+
+		class FPackageMetadataValueGetter : public FTableCellValueGetter
+		{
+		public:
+			virtual const TOptional<FTableCellValue> GetValue(const FTableColumn& Column, const FBaseTreeNode& Node) const
+			{
+				if (Node.IsGroup())
+				{
+					const FTableTreeNode& NodePtr = static_cast<const FTableTreeNode&>(Node);
+					if (NodePtr.HasAggregatedValue(Column.GetId()))
+					{
+						return NodePtr.GetAggregatedValue(Column.GetId());
+					}
+				}
+				else //if (Node->Is<FMemAllocNode>())
+				{
+					const FMemAllocNode& MemAllocNode = static_cast<const FMemAllocNode&>(Node);
+					const FMemoryAlloc* Alloc = MemAllocNode.GetMemAlloc();
+					if (Alloc)
+					{
+						return FTableCellValue(Alloc->GetPackage());
+					}
+				}
+
+				return TOptional<FTableCellValue>();
+			}
+		};
+		TSharedRef<ITableCellValueGetter> Getter = MakeShared<FPackageMetadataValueGetter>();
+		Column.SetValueGetter(Getter);
+
+		TSharedRef<ITableCellValueFormatter> Formatter = MakeShared<FCStringValueFormatterAsText>();
+		Column.SetValueFormatter(Formatter);
+
+		TSharedRef<ITableCellValueSorter> Sorter = MakeShared<FSorterByCStringValue>(ColumnRef);
+		Column.SetValueSorter(Sorter);
+
+		Column.SetAggregation(ETableColumnAggregation::SameValue);
+
+		AddColumn(ColumnRef);
+	}
+	//////////////////////////////////////////////////
 	// Asset Column
 	{
 		TSharedRef<FTableColumn> ColumnRef = MakeShared<FTableColumn>(FMemAllocTableColumns::AssetColumnId);
diff --git a/Engine/Source/Developer/TraceInsights/Private/Insights/MemoryProfiler/ViewModels/MemAllocTable.h b/Engine/Source/Developer/TraceInsights/Private/Insights/MemoryProfiler/ViewModels/MemAllocTable.h
--- a/Engine/Source/Developer/TraceInsights/Private/Insights/MemoryProfiler/ViewModels/MemAllocTable.h
+++ b/Engine/Source/Developer/TraceInsights/Private/Insights/MemoryProfiler/ViewModels/MemAllocTable.h
@@ -30,6 +30,7 @@ struct FMemAllocTableColumns
 	static const FName SizeColumnId;
 	static const FName TagColumnId;
 	static const FName AssetColumnId;
+	static const FName PackageColumnId;
 	static const FName ClassNameColumnId;
 	static const FName FunctionColumnId;
 	static const FName SourceFileColumnId;
diff --git a/Engine/Source/Developer/TraceInsights/Private/Insights/MemoryProfiler/ViewModels/MemoryAlloc.cpp b/Engine/Source/Developer/TraceInsights/Private/Insights/MemoryProfiler/ViewModels/MemoryAlloc.cpp
--- a/Engine/Source/Developer/TraceInsights/Private/Insights/MemoryProfiler/ViewModels/MemoryAlloc.cpp
+++ b/Engine/Source/Developer/TraceInsights/Private/Insights/MemoryProfiler/ViewModels/MemoryAlloc.cpp
@@ -19,6 +19,7 @@ FMemoryAlloc::FMemoryAlloc()
 	, TagId(0)
 	, Tag(nullptr)
 	, Asset(nullptr)
+	, Package(nullptr)
 	, Callstack(nullptr)
 	, FreeCallstack(nullptr)
 	, RootHeap(0)
diff --git a/Engine/Source/Developer/TraceInsights/Private/Insights/MemoryProfiler/ViewModels/MemoryAlloc.h b/Engine/Source/Developer/TraceInsights/Private/Insights/MemoryProfiler/ViewModels/MemoryAlloc.h
--- a/Engine/Source/Developer/TraceInsights/Private/Insights/MemoryProfiler/ViewModels/MemoryAlloc.h
+++ b/Engine/Source/Developer/TraceInsights/Private/Insights/MemoryProfiler/ViewModels/MemoryAlloc.h
@@ -32,6 +32,7 @@ public:
 	TraceServices::TagIdType GetTagId() const { return TagId; }
 	const TCHAR* GetTag() const { return Tag; }
 	const TCHAR* GetAsset() const { return Asset; }
+	const TCHAR* GetPackage() const { return Package; }
 	const TCHAR* GetClassName() const { return ClassName; }
 	uint32 GetCallstackId() const { return CallstackId; }
 	uint32 GetFreeCallstackId() const { return FreeCallstackId; }
@@ -51,6 +52,7 @@ private:
 	TraceServices::TagIdType TagId;
 	const TCHAR* Tag;
 	const TCHAR* Asset;
+	const TCHAR* Package;
 	const TCHAR* ClassName;
 	uint32 CallstackId;
 	uint32 FreeCallstackId;
diff --git a/Engine/Source/Developer/TraceInsights/Private/Insights/MemoryProfiler/Widgets/SMemAllocTableTreeView.cpp b/Engine/Source/Developer/TraceInsights/Private/Insights/MemoryProfiler/Widgets/SMemAllocTableTreeView.cpp
--- a/Engine/Source/Developer/TraceInsights/Private/Insights/MemoryProfiler/Widgets/SMemAllocTableTreeView.cpp
+++ b/Engine/Source/Developer/TraceInsights/Private/Insights/MemoryProfiler/Widgets/SMemAllocTableTreeView.cpp
@@ -405,6 +405,8 @@ void SMemAllocTableTreeView::UpdateQuery(TraceServices::IAllocationsProvider::EQ
 
 					Alloc.Asset = nullptr;
 					Alloc.ClassName = nullptr;
+					Alloc.Package = nullptr;
+					
 					const uint32 MetadataId = Allocation->GetMetadataId();
 					if (MetadataId != TraceServices::IMetadataProvider::InvalidMetadataId && MetadataProvider && DefinitionProvider)
 					{
@@ -428,6 +430,12 @@ void SMemAllocTableTreeView::UpdateQuery(TraceServices::IAllocationsProvider::EQ
 									{
 										Alloc.ClassName = ClassName->Display;
 									}
+									const auto PackageNameRef = Reader.GetValueAs<UE::Trace::FEventRef32>((uint8*)Data, 2);
+									const auto PackageName = DefinitionProvider->Get<TraceServices::FStringDefinition>(*PackageNameRef);
+									if (PackageName)
+									{
+										Alloc.Package = PackageName->Display;
+									}
 									return false;
 								}
 								return true;
diff --git a/Engine/Source/Runtime/Core/Public/ProfilingDebugging/AssetMetadataTrace.h b/Engine/Source/Runtime/Core/Public/ProfilingDebugging/AssetMetadataTrace.h
--- a/Engine/Source/Runtime/Core/Public/ProfilingDebugging/AssetMetadataTrace.h
+++ b/Engine/Source/Runtime/Core/Public/ProfilingDebugging/AssetMetadataTrace.h
@@ -27,22 +27,28 @@ CORE_API UE_TRACE_CHANNEL_EXTERN(AssetMetadataChannel);
 CORE_API UE_TRACE_METADATA_EVENT_BEGIN_EXTERN(Asset)
 	UE_TRACE_METADATA_EVENT_REFERENCE_FIELD(Strings, FName, Name)
 	UE_TRACE_METADATA_EVENT_REFERENCE_FIELD(Strings, FName, Class)
+	UE_TRACE_METADATA_EVENT_REFERENCE_FIELD(Strings, FName, Package)
 UE_TRACE_METADATA_EVENT_END()
 
 /**
- * Utility macro to use asset scope
+ * Utility macro to create asset scope from and object and object class
  * @note When using this macro in an module outside of Core, make sure a dependency to the "TraceLog" module
  *		 is added.
  */
 #define UE_TRACE_METADATA_SCOPE_ASSET(Object, ObjClass) \
-	UE_TRACE_METADATA_SCOPE_ASSET_FNAME(Object->GetFName(), ObjClass->GetFName())
+	UE_TRACE_METADATA_SCOPE_ASSET_FNAME(Object->GetFName(), ObjClass->GetFName(), Object->GetOutermost()->GetFName())
 
-#define UE_TRACE_METADATA_SCOPE_ASSET_FNAME(ObjectName, ObjClassName) \
+/**
+ * Utility macro to create an asset scope by specifying object name, class name and package name explicitly.
+ */
+#define UE_TRACE_METADATA_SCOPE_ASSET_FNAME(ObjectName, ObjClassName, PackageName) \
 	auto MetaNameRef = UE_TRACE_CHANNELEXPR_IS_ENABLED(MetadataChannel) ? FStringTrace::GetNameRef(ObjectName) : UE::Trace::FEventRef32(0,0);\
 	auto ClassNameRef = UE_TRACE_CHANNELEXPR_IS_ENABLED(MetadataChannel) ? FStringTrace::GetNameRef(ObjClassName) : UE::Trace::FEventRef32(0,0);\
+	auto PackageNameRef = UE_TRACE_CHANNELEXPR_IS_ENABLED(MetadataChannel) ? FStringTrace::GetNameRef(PackageName) : UE::Trace::FEventRef32(0,0);\
 	UE_TRACE_METADATA_SCOPE(Asset, AssetMetadataChannel)\
 		<< Asset.Name(MetaNameRef)\
-		<< Asset.Class(ClassNameRef);
+		<< Asset.Class(ClassNameRef)\
+		<< Asset.Package(PackageNameRef);
 
 #else
 
diff --git a/Engine/Source/Runtime/D3D12RHI/Private/D3D12Buffer.cpp b/Engine/Source/Runtime/D3D12RHI/Private/D3D12Buffer.cpp
--- a/Engine/Source/Runtime/D3D12RHI/Private/D3D12Buffer.cpp
+++ b/Engine/Source/Runtime/D3D12RHI/Private/D3D12Buffer.cpp
@@ -72,7 +72,7 @@ struct FRHICommandRenameUploadBuffer final : public FRHICommand<FRHICommandRenam
 
 	void Execute(FRHICommandListBase& CmdList)
 	{
-		UE_TRACE_METADATA_SCOPE_ASSET_FNAME(ChooseMetadataName(Resource->GetOwnerName(), Resource->GetName()), FName(TEXT("FRHICommandRenameUploadBuffer::Execute")));
+		UE_TRACE_METADATA_SCOPE_ASSET_FNAME(Resource->GetName(), FName(TEXT("FRHICommandRenameUploadBuffer::Execute")), Resource->GetOwnerName());
 
 		// Make sure we have an active pipeline when running at the top of the pipe, so the lambda below can obtain a context.
 		TOptional<ERHIPipeline> PreviousPipeline;
@@ -200,9 +200,9 @@ struct FD3D12RHICommandInitializeBuffer final : public FRHICommand<FD3D12RHIComm
 };
 
 
-void FD3D12Buffer::UploadResourceData(FRHICommandListBase& RHICmdList, FResourceArrayInterface* InResourceArray, D3D12_RESOURCE_STATES InDestinationState, const FName& AssetName, const FName& ClassName)
+void FD3D12Buffer::UploadResourceData(FRHICommandListBase& RHICmdList, FResourceArrayInterface* InResourceArray, D3D12_RESOURCE_STATES InDestinationState, const FName& AssetName, const FName& ClassName, const FName& PackageName)
 {
-	UE_TRACE_METADATA_SCOPE_ASSET_FNAME(AssetName, ClassName);
+	UE_TRACE_METADATA_SCOPE_ASSET_FNAME(AssetName, ClassName, PackageName);
 
 	check(InResourceArray);
 	check(ResourceLocation.IsValid());
@@ -354,7 +354,7 @@ FD3D12Buffer* FD3D12Adapter::CreateRHIBuffer(
 	TRACE_CPUPROFILER_EVENT_SCOPE(D3D12RHI::CreateRHIBuffer);
 	SCOPE_CYCLE_COUNTER(STAT_D3D12CreateBufferTime);
 
-	UE_TRACE_METADATA_SCOPE_ASSET_FNAME(ChooseMetadataName(AssetName, InDebugName), ChooseMetadataName(ClassName, TEXT("FRHIBuffer")));
+	UE_TRACE_METADATA_SCOPE_ASSET_FNAME(FName(InDebugName), ChooseMetadataName(ClassName, TEXT("FRHIBuffer")), AssetName);
 
 	check(InDesc.Width == Size);
 
@@ -539,8 +539,9 @@ FBufferRHIRef FD3D12DynamicRHI::CreateBuffer(FRHICommandListBase& RHICmdList, ui
 
 FD3D12Buffer* FD3D12DynamicRHI::CreateD3D12Buffer(class FRHICommandListBase* RHICmdList, uint32 Size, EBufferUsageFlags Usage, uint32 Stride, ERHIAccess InResourceState, FRHIResourceCreateInfo& CreateInfo, ID3D12ResourceAllocator* ResourceAllocator)
 {
-	FName TraceAssetName = ChooseMetadataName(CreateInfo.GetTraceAssetName(), CreateInfo.DebugName);
+	FName TraceAssetName = CreateInfo.DebugName;
 	FName TraceClassName = ChooseMetadataName(CreateInfo.GetTraceClassName(), TEXT("FRHIBuffer"));
+	FName TracePackageName = CreateInfo.GetTraceAssetName();
 
 	D3D12_RESOURCE_DESC Desc;
 	uint32 Alignment;
@@ -572,7 +573,7 @@ FD3D12Buffer* FD3D12DynamicRHI::CreateD3D12Buffer(class FRHICommandListBase* RHI
 	if (bHasInitialData)
 	{
 		check(RHICmdList);
-		Buffer->UploadResourceData(*RHICmdList, CreateInfo.ResourceArray, DesiredState, TraceAssetName, TraceClassName);
+		Buffer->UploadResourceData(*RHICmdList, CreateInfo.ResourceArray, DesiredState, TraceAssetName, TraceClassName, TracePackageName);
 	}
 
 	return Buffer;
@@ -587,7 +588,7 @@ FRHIBuffer* FD3D12DynamicRHI::CreateBuffer(const FRHIBufferCreateInfo& CreateInf
 void* FD3D12DynamicRHI::LockBuffer(FRHICommandListBase& RHICmdList, FD3D12Buffer* Buffer, uint32 BufferSize, EBufferUsageFlags BufferUsage, uint32 Offset, uint32 Size, EResourceLockMode LockMode)
 {
 	SCOPE_CYCLE_COUNTER(STAT_D3D12LockBufferTime);
-	UE_TRACE_METADATA_SCOPE_ASSET_FNAME(ChooseMetadataName(Buffer->GetOwnerName(), Buffer->GetName()), Buffer->GetName());
+	UE_TRACE_METADATA_SCOPE_ASSET_FNAME(Buffer->GetName(), Buffer->GetName(), Buffer->GetOwnerName());
 
 	checkf(Size <= BufferSize, TEXT("Requested lock size %u is larger than the total size %u for buffer '%s'."), Size, BufferSize, *Buffer->GetName().ToString());
 
diff --git a/Engine/Source/Runtime/D3D12RHI/Private/D3D12Commands.cpp b/Engine/Source/Runtime/D3D12RHI/Private/D3D12Commands.cpp
--- a/Engine/Source/Runtime/D3D12RHI/Private/D3D12Commands.cpp
+++ b/Engine/Source/Runtime/D3D12RHI/Private/D3D12Commands.cpp
@@ -676,8 +676,7 @@ void FD3D12CommandContext::RHICopyToStagingBuffer(FRHIBuffer* SourceBufferRHI, F
 {
 	SCOPE_CYCLE_COUNTER(STAT_D3D12CopyToStagingBufferTime);
 
-	FName AssetName = SourceBufferRHI->GetOwnerName() == NAME_None ? SourceBufferRHI->GetName() : SourceBufferRHI->GetOwnerName();
-	UE_TRACE_METADATA_SCOPE_ASSET_FNAME(AssetName, FName(TEXT("FRHIStagingBuffer")));
+	UE_TRACE_METADATA_SCOPE_ASSET_FNAME(SourceBufferRHI->GetName(), FName(TEXT("FRHIStagingBuffer")), SourceBufferRHI->GetOwnerName());
 
 	FD3D12StagingBuffer* StagingBuffer = FD3D12DynamicRHI::ResourceCast(StagingBufferRHI);
 	check(StagingBuffer);
diff --git a/Engine/Source/Runtime/D3D12RHI/Private/D3D12Query.cpp b/Engine/Source/Runtime/D3D12RHI/Private/D3D12Query.cpp
--- a/Engine/Source/Runtime/D3D12RHI/Private/D3D12Query.cpp
+++ b/Engine/Source/Runtime/D3D12RHI/Private/D3D12Query.cpp
@@ -75,7 +75,7 @@ FD3D12QueryHeap::FD3D12QueryHeap(FD3D12Device* Device, D3D12_QUERY_TYPE QueryTyp
 		break;
 	}
 
-	UE_TRACE_METADATA_SCOPE_ASSET_FNAME(FName(ResultBufferName), FName(TEXT("FD3D12QueryHeap")));
+	UE_TRACE_METADATA_SCOPE_ASSET_FNAME(FName(ResultBufferName), FName(TEXT("FD3D12QueryHeap")), NAME_None);
 
 	// Create the query heap
 #if D3D12RHI_PLATFORM_USES_TIMESTAMP_QUERIES == 0
diff --git a/Engine/Source/Runtime/D3D12RHI/Private/D3D12Resources.cpp b/Engine/Source/Runtime/D3D12RHI/Private/D3D12Resources.cpp
--- a/Engine/Source/Runtime/D3D12RHI/Private/D3D12Resources.cpp
+++ b/Engine/Source/Runtime/D3D12RHI/Private/D3D12Resources.cpp
@@ -1102,7 +1102,7 @@ bool FD3D12ResourceLocation::OnAllocationMoved(FRHIPoolAllocationData* InNewData
 	FD3D12Resource* CurrentResource = GetResource();
 	FD3D12PoolAllocator* NewAllocator = GetPoolAllocator();
 
-	UE_TRACE_METADATA_SCOPE_ASSET_FNAME(CurrentResource->GetName(), FName(TEXT("FD3D12ResourceLocation::OnAllocationMoved")));
+	UE_TRACE_METADATA_SCOPE_ASSET_FNAME(CurrentResource->GetName(), FName(TEXT("FD3D12ResourceLocation::OnAllocationMoved")), NAME_None);
 
 	// If sub allocated and not placed only update the internal data
 	if (NewAllocator->GetAllocationStrategy() == EResourceAllocationStrategy::kManualSubAllocation)
diff --git a/Engine/Source/Runtime/D3D12RHI/Private/D3D12Texture.cpp b/Engine/Source/Runtime/D3D12RHI/Private/D3D12Texture.cpp
--- a/Engine/Source/Runtime/D3D12RHI/Private/D3D12Texture.cpp
+++ b/Engine/Source/Runtime/D3D12RHI/Private/D3D12Texture.cpp
@@ -880,7 +880,7 @@ FD3D12Texture* FD3D12DynamicRHI::CreateD3D12Texture(const FRHITextureCreateDesc&
 #if PLATFORM_WINDOWS || PLATFORM_HOLOLENS
 	
 	TRACE_CPUPROFILER_EVENT_SCOPE(D3D12RHI::CreateD3D12Texture);
-	UE_TRACE_METADATA_SCOPE_ASSET_FNAME(InCreateDesc.GetTraceAssetName(), InCreateDesc.GetTraceClassName());
+	UE_TRACE_METADATA_SCOPE_ASSET_FNAME(FName(InCreateDesc.DebugName), InCreateDesc.GetTraceClassName(), InCreateDesc.GetTraceAssetName());
 
 	// Make local copy of create desc because certain flags might be removed before actual texture creation
 	FRHITextureCreateDesc CreateDesc = InCreateDesc;
@@ -1068,7 +1068,7 @@ FTextureRHIRef FD3D12DynamicRHI::RHIAsyncCreateTexture2D(uint32 SizeX, uint32 Si
 {	
 	check(GRHISupportsAsyncTextureCreation);
 	
-	UE_TRACE_METADATA_SCOPE_ASSET_FNAME(FName(TEXT("RHIAsyncCreateTexture2D")), FName(TEXT("FRHITexture")));
+	UE_TRACE_METADATA_SCOPE_ASSET_FNAME(FName(TEXT("RHIAsyncCreateTexture2D")), FName(TEXT("FRHITexture")), NAME_None);
 
 	const ETextureCreateFlags InvalidFlags = TexCreate_RenderTargetable | TexCreate_ResolveTargetable | TexCreate_DepthStencilTargetable | TexCreate_GenerateMipCapable | TexCreate_UAV | TexCreate_Presentable | TexCreate_CPUReadback;
 	check(!EnumHasAnyFlags(Flags, InvalidFlags));
@@ -1791,8 +1791,7 @@ void* FD3D12Texture::Lock(class FRHICommandListImmediate* RHICmdList, uint32 Mip
 {
 	SCOPE_CYCLE_COUNTER(STAT_D3D12LockTextureTime);
 
-	FName AssetName = GetOwnerName() == NAME_None ? GetName() : GetOwnerName();
-	UE_TRACE_METADATA_SCOPE_ASSET_FNAME(AssetName, FName(TEXT("FRHITexture Lock")));
+	UE_TRACE_METADATA_SCOPE_ASSET_FNAME(GetName(), FName(TEXT("FRHITexture Lock")), GetOwnerName());
 
 	FD3D12Device* Device = GetParentDevice();
 	FD3D12Adapter* Adapter = Device->GetParentAdapter();
diff --git a/Engine/Source/Runtime/D3D12RHI/Private/D3D12TransientResourceAllocator.cpp b/Engine/Source/Runtime/D3D12RHI/Private/D3D12TransientResourceAllocator.cpp
--- a/Engine/Source/Runtime/D3D12RHI/Private/D3D12TransientResourceAllocator.cpp
+++ b/Engine/Source/Runtime/D3D12RHI/Private/D3D12TransientResourceAllocator.cpp
@@ -29,7 +29,7 @@ D3D12_RESOURCE_STATES GetInitialResourceState(const D3D12_RESOURCE_DESC& InDesc)
 FD3D12TransientHeap::FD3D12TransientHeap(const FInitializer& Initializer, FD3D12Adapter* Adapter, FD3D12Device* Device, FRHIGPUMask VisibleNodeMask)
 	: FRHITransientHeap(Initializer)
 {
-	UE_TRACE_METADATA_SCOPE_ASSET_FNAME(FName(TEXT("FD3D12TransientHeap")), FName(TEXT("FD3D12TransientHeap")));
+	UE_TRACE_METADATA_SCOPE_ASSET_FNAME(FName(TEXT("FD3D12TransientHeap")), FName(TEXT("FD3D12TransientHeap")), NAME_None);
 
 	D3D12_HEAP_FLAGS HeapFlags = D3D12_HEAP_FLAG_ALLOW_ALL_BUFFERS_AND_TEXTURES;
 
diff --git a/Engine/Source/Runtime/D3D12RHI/Private/Windows/WindowsD3D12Device.cpp b/Engine/Source/Runtime/D3D12RHI/Private/Windows/WindowsD3D12Device.cpp
--- a/Engine/Source/Runtime/D3D12RHI/Private/Windows/WindowsD3D12Device.cpp
+++ b/Engine/Source/Runtime/D3D12RHI/Private/Windows/WindowsD3D12Device.cpp
@@ -1969,7 +1969,7 @@ void FWindowsD3D12Adapter::CreateCommandSignatures()
 
 TUniquePtr<FD3D12DiagnosticBuffer> FD3D12Device::CreateDiagnosticBuffer(const D3D12_RESOURCE_DESC& Desc, const TCHAR* Name)
 {
-	UE_TRACE_METADATA_SCOPE_ASSET_FNAME(FName(TEXT("FD3D12DiagnosticBuffer")), FName(TEXT("FD3D12Device::CreateDiagnosticBuffer")));
+	UE_TRACE_METADATA_SCOPE_ASSET_FNAME(FName(TEXT("FD3D12DiagnosticBuffer")), FName(TEXT("FD3D12Device::CreateDiagnosticBuffer")), NAME_None);
 
 	TRefCountPtr<ID3D12Device3> D3D12Device3;
 	HRESULT hr = GetDevice()->QueryInterface(IID_PPV_ARGS(D3D12Device3.GetInitReference()));
diff --git a/Engine/Source/Runtime/D3D12RHI/Public/D3D12Resources.h b/Engine/Source/Runtime/D3D12RHI/Public/D3D12Resources.h
--- a/Engine/Source/Runtime/D3D12RHI/Public/D3D12Resources.h
+++ b/Engine/Source/Runtime/D3D12RHI/Public/D3D12Resources.h
@@ -879,7 +879,7 @@ public:
 
 	virtual uint32 GetParentGPUIndex() const override;
 
-	void UploadResourceData(FRHICommandListBase& InRHICmdList, FResourceArrayInterface* InResourceArray, D3D12_RESOURCE_STATES InDestinationState, const FName& AssetName = NAME_None, const FName& ClassName = NAME_None);
+	void UploadResourceData(FRHICommandListBase& InRHICmdList, FResourceArrayInterface* InResourceArray, D3D12_RESOURCE_STATES InDestinationState, const FName& AssetName = NAME_None, const FName& ClassName = NAME_None, const FName& PackageName = NAME_None);
 	FD3D12SyncPointRef UploadResourceDataViaCopyQueue(FResourceArrayInterface* InResourceArray);
 
 	// FRHIResource overrides
