Fix texture Insight memory tracing throws alloc/free 0x0 on RootHeap=1 error. It happens when D3D12 resource heap tier is level 2, in which case IsTrackingAllAllocations() is false and textures don't have valid GPUVirtualAddress.

diff --git a/Engine/Source/Runtime/D3D12RHI/Private/D3D12Adapter.h b/Engine/Source/Runtime/D3D12RHI/Private/D3D12Adapter.h
--- a/Engine/Source/Runtime/D3D12RHI/Private/D3D12Adapter.h
+++ b/Engine/Source/Runtime/D3D12RHI/Private/D3D12Adapter.h
@@ -599,4 +599,8 @@ protected:
 	FD3D12RootSignature StaticRayTracingGlobalRootSignature;
 	FD3D12RootSignature StaticRayTracingLocalRootSignature;
 #endif
+
+private:
+	// Insight memory trace helper
+	void TraceMemoryAllocation(FD3D12Resource* Resource);
 };
diff --git a/Engine/Source/Runtime/D3D12RHI/Private/D3D12Resources.cpp b/Engine/Source/Runtime/D3D12RHI/Private/D3D12Resources.cpp
--- a/Engine/Source/Runtime/D3D12RHI/Private/D3D12Resources.cpp
+++ b/Engine/Source/Runtime/D3D12RHI/Private/D3D12Resources.cpp
@@ -449,15 +449,7 @@ HRESULT FD3D12Adapter::CreateCommittedResource(const FD3D12ResourceDesc& InDesc,
 			(*ppOutResource)->StartTrackingForResidency();
 		}
 
-#if UE_MEMORY_TRACE_ENABLED
-		// Calling GetResourceAllocationInfo is not cheap so check memory allocation tracking is enabled
-		if (UE_TRACE_CHANNELEXPR_IS_ENABLED(MemAllocChannel))
-		{
-			const D3D12_RESOURCE_DESC ResourceDesc = (*ppOutResource)->GetResource()->GetDesc();
-			const D3D12_RESOURCE_ALLOCATION_INFO Info = RootDevice->GetResourceAllocationInfo(0, 1, &ResourceDesc);
-			MemoryTrace_Alloc((*ppOutResource)->GetGPUVirtualAddress(), Info.SizeInBytes, Info.Alignment, EMemoryTraceRootHeap::VideoMemory);
-		}
-#endif
+		TraceMemoryAllocation(*ppOutResource);
 	}
 	else	
 	{
@@ -584,16 +576,12 @@ HRESULT FD3D12Adapter::CreatePlacedResource(const FD3D12ResourceDesc& InDesc, FD
 			}
 		}
 #endif		
-#if UE_MEMORY_TRACE_ENABLED
-		// Calling GetResourceAllocationInfo is not cheap so check memory allocation tracking is enabled
 		// Don't track resources allocated on transient heaps
-		if (UE_TRACE_CHANNELEXPR_IS_ENABLED(MemAllocChannel) && !BackingHeap->GetIsTransient())
+		if (!BackingHeap->GetIsTransient())
 		{
-			const D3D12_RESOURCE_DESC ResourceDesc = (*ppOutResource)->GetResource()->GetDesc();
-			const D3D12_RESOURCE_ALLOCATION_INFO Info = RootDevice->GetResourceAllocationInfo(0, 1, &ResourceDesc);
-			MemoryTrace_Alloc((*ppOutResource)->GetGPUVirtualAddress(), Info.SizeInBytes, Info.Alignment, EMemoryTraceRootHeap::VideoMemory);
+			TraceMemoryAllocation(*ppOutResource);
 		}
-#endif
+
 		// Set a default name (can override later).
 		SetName(*ppOutResource, Name);
 
@@ -613,6 +601,24 @@ HRESULT FD3D12Adapter::CreatePlacedResource(const FD3D12ResourceDesc& InDesc, FD
 	return hr;
 }
 
+void FD3D12Adapter::TraceMemoryAllocation(FD3D12Resource* Resource)
+{
+#if UE_MEMORY_TRACE_ENABLED
+	// Calling GetResourceAllocationInfo is not cheap so check memory allocation tracking is enabled
+	if (UE_TRACE_CHANNELEXPR_IS_ENABLED(MemAllocChannel))
+	{
+		const D3D12_RESOURCE_DESC ResourceDesc = Resource->GetResource()->GetDesc();
+		const D3D12_RESOURCE_ALLOCATION_INFO Info = RootDevice->GetResourceAllocationInfo(0, 1, &ResourceDesc);
+		D3D12_GPU_VIRTUAL_ADDRESS GPUAddress = Resource->GetGPUVirtualAddress();
+		// Textures don't have valid GPUVirtualAddress when IsTrackingAllAllocations() is false, so don't do memory trace in this case.
+		if (IsTrackingAllAllocations() || GPUAddress != 0)
+		{
+			MemoryTrace_Alloc(GPUAddress, Info.SizeInBytes, Info.Alignment, EMemoryTraceRootHeap::VideoMemory);
+		}
+	}
+#endif
+}
+
 HRESULT FD3D12Adapter::CreateBuffer(D3D12_HEAP_TYPE HeapType, FRHIGPUMask CreationNode, FRHIGPUMask VisibleNodes, uint64 HeapSize, FD3D12Resource** ppOutResource, const TCHAR* Name, D3D12_RESOURCE_FLAGS Flags)
 {
 	const D3D12_HEAP_PROPERTIES HeapProps = CD3DX12_HEAP_PROPERTIES(HeapType, CreationNode.GetNative(), VisibleNodes.GetNative());
@@ -1104,20 +1110,25 @@ bool FD3D12ResourceLocation::OnAllocationMoved(FRHIPoolAllocationData* InNewData
 
 	UE_TRACE_METADATA_SCOPE_ASSET_FNAME(CurrentResource->GetName(), FName(TEXT("FD3D12ResourceLocation::OnAllocationMoved")), NAME_None);
 
+	// Textures don't have valid GPUVirtualAddress when IsTrackingAllAllocations() is false, so don't do memory trace in this case.
+	const bool bTrackingAllAllocations = GetParentDevice()->GetParentAdapter()->IsTrackingAllAllocations();
+	const bool bMemoryTrace = bTrackingAllAllocations || GPUVirtualAddress != 0;
+
 	// If sub allocated and not placed only update the internal data
 	if (NewAllocator->GetAllocationStrategy() == EResourceAllocationStrategy::kManualSubAllocation)
 	{
 		check(!CurrentResource->IsPlacedResource());
-
-#if UE_MEMORY_TRACE_ENABLED
-		MemoryTrace_ReallocFree(GPUVirtualAddress, EMemoryTraceRootHeap::VideoMemory);
-#endif
+		D3D12_GPU_VIRTUAL_ADDRESS OldGPUAddress = GPUVirtualAddress;
 		OffsetFromBaseOfResource = AllocationData.GetOffset();
 		UnderlyingResource = NewAllocator->GetBackingResource(*this);
 		GPUVirtualAddress = UnderlyingResource->GetGPUVirtualAddress() + OffsetFromBaseOfResource;
 
 #if UE_MEMORY_TRACE_ENABLED
-		MemoryTrace_ReallocAlloc(GPUVirtualAddress, AllocationData.GetSize(), AllocationData.GetAlignment(), EMemoryTraceRootHeap::VideoMemory);
+		if (bMemoryTrace)
+		{
+			MemoryTrace_ReallocFree(OldGPUAddress, EMemoryTraceRootHeap::VideoMemory);
+			MemoryTrace_ReallocAlloc(GPUVirtualAddress, AllocationData.GetSize(), AllocationData.GetAlignment(), EMemoryTraceRootHeap::VideoMemory);
+		}
 #endif
 	}
 	else
@@ -1126,8 +1137,11 @@ bool FD3D12ResourceLocation::OnAllocationMoved(FRHIPoolAllocationData* InNewData
 		check(OffsetFromBaseOfResource == 0);
 
 #if UE_MEMORY_TRACE_ENABLED
-		// CreatePlacedResource function below calls MemoryTrace_Alloc to track new memory, so call MemoryTrace_Free to match (instead of calling MemoryTrace_ReallocFree/MemoryTrace_ReallocAlloc).
-		MemoryTrace_Free(GPUVirtualAddress, EMemoryTraceRootHeap::VideoMemory);
+		if (bMemoryTrace)
+		{
+			// CreatePlacedResource function below calls MemoryTrace_Alloc to track new memory, so call MemoryTrace_Free to match (instead of calling MemoryTrace_ReallocFree/MemoryTrace_ReallocAlloc).
+			MemoryTrace_Free(GPUVirtualAddress, EMemoryTraceRootHeap::VideoMemory);
+		}
 #endif
 		// recreate the placed resource (ownership of current resource is already handled during the internal move)
 		FD3D12HeapAndOffset HeapAndOffset = NewAllocator->GetBackingHeapAndAllocationOffsetInBytes(*this);
diff --git a/Engine/Source/Runtime/D3D12RHI/Private/D3D12Texture.cpp b/Engine/Source/Runtime/D3D12RHI/Private/D3D12Texture.cpp
--- a/Engine/Source/Runtime/D3D12RHI/Private/D3D12Texture.cpp
+++ b/Engine/Source/Runtime/D3D12RHI/Private/D3D12Texture.cpp
@@ -388,6 +388,12 @@ void FD3D12TextureStats::UpdateD3D12TextureStats(FD3D12Texture& Texture, const D
 
 	D3D12_GPU_VIRTUAL_ADDRESS GPUAddress = Texture.ResourceLocation.GetGPUVirtualAddress();
 
+#if UE_MEMORY_TRACE_ENABLED
+	// Textures don't have valid GPUVirtualAddress when IsTrackingAllAllocations() is false, so don't do memory trace in this case.
+	const bool bTrackingAllAllocations = Texture.GetParentDevice()->GetParentAdapter()->IsTrackingAllAllocations();
+	const bool bMemoryTrace = bTrackingAllAllocations || GPUAddress != 0;
+#endif
+
 	if (TextureSize > 0)
 	{
 #if PLATFORM_WINDOWS
@@ -402,7 +408,7 @@ void FD3D12TextureStats::UpdateD3D12TextureStats(FD3D12Texture& Texture, const D
 			// Skip if it's created as a
 			// 1) standalone resource, because MemoryTrace_Alloc has been called in FD3D12Adapter::CreateCommittedResource
 			// 2) placed resource from a pool allocator, because MemoryTrace_Alloc has been called in FD3D12Adapter::CreatePlacedResource
-			if (!Texture.ResourceLocation.IsStandaloneOrPooledPlacedResource())
+			if (bMemoryTrace && !Texture.ResourceLocation.IsStandaloneOrPooledPlacedResource())
 			{
 				MemoryTrace_Alloc(GPUAddress, TextureSize, Desc.Alignment, EMemoryTraceRootHeap::VideoMemory);
 			}
