Added CVar to control whether to bind D3D12 debug labels.

diff --git a/Engine/Source/Runtime/D3D12RHI/Private/D3D12Buffer.cpp b/Engine/Source/Runtime/D3D12RHI/Private/D3D12Buffer.cpp
--- a/Engine/Source/Runtime/D3D12RHI/Private/D3D12Buffer.cpp
+++ b/Engine/Source/Runtime/D3D12RHI/Private/D3D12Buffer.cpp
@@ -12,6 +12,8 @@ FName ChooseMetadataName(const FName& FirstChoice, const FName& SecondChoice)
 	return FirstChoice == NAME_None ? SecondChoice : FirstChoice;
 }
 
+extern int32 GD3D12BindResourceLabels;
+
 FD3D12Buffer::~FD3D12Buffer()
 {
 	if (EnumHasAnyFlags(GetUsage(), EBufferUsageFlags::VertexBuffer) && GetParentDevice())
@@ -860,7 +862,7 @@ void FD3D12DynamicRHI::RHICopyBuffer(FRHIBuffer* SourceBufferRHI, FRHIBuffer* De
 
 void FD3D12DynamicRHI::RHIBindDebugLabelName(FRHIBuffer* BufferRHI, const TCHAR* Name)
 {
-	if (BufferRHI == nullptr)
+	if (BufferRHI == nullptr || !GD3D12BindResourceLabels)
 	{
 		return;
 	}
diff --git a/Engine/Source/Runtime/D3D12RHI/Private/D3D12RHI.cpp b/Engine/Source/Runtime/D3D12RHI/Private/D3D12RHI.cpp
--- a/Engine/Source/Runtime/D3D12RHI/Private/D3D12RHI.cpp
+++ b/Engine/Source/Runtime/D3D12RHI/Private/D3D12RHI.cpp
@@ -21,6 +21,12 @@
 DEFINE_LOG_CATEGORY(LogD3D12RHI);
 DEFINE_LOG_CATEGORY(LogD3D12GapRecorder);
 
+int32 GD3D12BindResourceLabels = 1;
+static FAutoConsoleVariableRef CVarD3D12BindResourceLabels(
+	TEXT("d3d12.BindResourceLabels"),
+	GD3D12BindResourceLabels,
+	TEXT("Whether to enable binding of debug names to D3D12 resources."));
+
 static TAutoConsoleVariable<int32> CVarD3D12UseD24(
 	TEXT("r.D3D12.Depth24Bit"),
 	0,
diff --git a/Engine/Source/Runtime/D3D12RHI/Private/D3D12Texture.cpp b/Engine/Source/Runtime/D3D12RHI/Private/D3D12Texture.cpp
--- a/Engine/Source/Runtime/D3D12RHI/Private/D3D12Texture.cpp
+++ b/Engine/Source/Runtime/D3D12RHI/Private/D3D12Texture.cpp
@@ -48,6 +48,7 @@ static TAutoConsoleVariable<bool> CVarTexturePoolOnlyAccountStreamableTexture(
 	ECVF_ReadOnly
 );
 
+extern int32 GD3D12BindResourceLabels;
 
 ///////////////////////////////////////////////////////////////////////////////////////////
 // Texture Commands
@@ -2721,6 +2722,11 @@ void FD3D12DynamicRHI::RHIUnlockTextureCubeFace(FRHITextureCube* TextureCubeRHI,
 
 void FD3D12DynamicRHI::RHIBindDebugLabelName(FRHITexture* TextureRHI, const TCHAR* Name)
 {
+	if (!TextureRHI || !GD3D12BindResourceLabels)
+	{
+		return;
+	}
+
 #if NAME_OBJECTS
 	FD3D12Texture::FLinkedObjectIterator BaseTexture(GetD3D12TextureFromRHITexture(TextureRHI));
 
