Exclude back buffers and query heaps from D3D12 resource residency tracking system

* Resource eviction is tracked based on command list completion fences, however back buffers may be referenced outside of command lists. Present() calls are queue-level commands, which are not tracked correctly
* Query heap objects may be missing the required usage tracking calls. Disabling explicit residency/eviction as a temporary workaround

diff --git a/Engine/Source/Runtime/D3D12RHI/Private/D3D12Query.cpp b/Engine/Source/Runtime/D3D12RHI/Private/D3D12Query.cpp
--- a/Engine/Source/Runtime/D3D12RHI/Private/D3D12Query.cpp
+++ b/Engine/Source/Runtime/D3D12RHI/Private/D3D12Query.cpp
@@ -90,7 +90,7 @@ FD3D12QueryHeap::FD3D12QueryHeap(FD3D12Device* Device, D3D12_QUERY_TYPE QueryTyp
 		VERIFYD3D12RESULT(Device->GetDevice()->CreateQueryHeap(&QueryHeapDesc, IID_PPV_ARGS(D3DQueryHeap.GetInitReference())));
 		SetName(D3DQueryHeap, QueryHeapName);
 
-#if ENABLE_RESIDENCY_MANAGEMENT
+#if ENABLE_RESIDENCY_MANAGEMENT && 0 // Temporary workaround for missing resource usage tracking for query heap
 		D3DX12Residency::Initialize(ResidencyHandle, D3DQueryHeap, ResultSize * NumQueries, this);
 		D3DX12Residency::BeginTrackingObject(Device->GetResidencyManager(), ResidencyHandle);
 #endif
@@ -132,7 +132,10 @@ FD3D12QueryHeap::~FD3D12QueryHeap()
 	}
 
 #if ENABLE_RESIDENCY_MANAGEMENT
-	D3DX12Residency::EndTrackingObject(Device->GetResidencyManager(), ResidencyHandle);
+	if (D3DX12Residency::IsInitialized(ResidencyHandle))
+	{
+		D3DX12Residency::EndTrackingObject(Device->GetResidencyManager(), ResidencyHandle);
+	}
 #endif
 
 	DEC_DWORD_STAT(STAT_D3D12NumQueryHeaps);
diff --git a/Engine/Source/Runtime/D3D12RHI/Private/D3D12Resources.cpp b/Engine/Source/Runtime/D3D12RHI/Private/D3D12Resources.cpp
--- a/Engine/Source/Runtime/D3D12RHI/Private/D3D12Resources.cpp
+++ b/Engine/Source/Runtime/D3D12RHI/Private/D3D12Resources.cpp
@@ -266,6 +266,14 @@ ID3D12Pageable* FD3D12Resource::GetPageable()
 void FD3D12Resource::StartTrackingForResidency()
 {
 #if ENABLE_RESIDENCY_MANAGEMENT
+	if (bBackBuffer)
+	{
+		// Back buffers may be referenced outside of command lists (during presents), however D3DX12Residency.h library 
+		// uses fences tied to command lists to detect when it's safe to evict a resource, which is wrong for back buffers.
+		// Simply disable residency tracking for back buffers as a workaround (keep them always resident).
+		return;
+	}
+
 	check(IsGPUOnly(HeapType));	// This is checked at a higher level before calling this function.
 	check(D3DX12Residency::IsInitialized(ResidencyHandle) == false);
 	const D3D12_RESOURCE_DESC ResourceDesc = Resource->GetDesc();
diff --git a/Engine/Source/Runtime/D3D12RHI/Private/D3D12Viewport.cpp b/Engine/Source/Runtime/D3D12RHI/Private/D3D12Viewport.cpp
--- a/Engine/Source/Runtime/D3D12RHI/Private/D3D12Viewport.cpp
+++ b/Engine/Source/Runtime/D3D12RHI/Private/D3D12Viewport.cpp
@@ -206,6 +206,7 @@ FD3D12Texture* GetSwapChainSurface(FD3D12Device* Parent, EPixelFormat PixelForma
 		if (Device->GetGPUIndex() == Parent->GetGPUIndex())
 		{
 			FD3D12Resource* NewResourceWrapper = new FD3D12Resource(Device, FRHIGPUMask::All(), BackBufferResource, InitialState, BackBufferDesc);
+			NewResourceWrapper->SetIsBackBuffer(true);
 			NewResourceWrapper->AddRef();
 			NewResourceWrapper->StartTrackingForResidency();
 			NewTexture->ResourceLocation.AsStandAlone(NewResourceWrapper);
