Fix Insight memory allocation error in DX12 buddy allocator

diff --git a/Engine/Source/Runtime/D3D12RHI/Private/D3D12Allocation.cpp b/Engine/Source/Runtime/D3D12RHI/Private/D3D12Allocation.cpp
--- a/Engine/Source/Runtime/D3D12RHI/Private/D3D12Allocation.cpp
+++ b/Engine/Source/Runtime/D3D12RHI/Private/D3D12Allocation.cpp
@@ -582,12 +582,17 @@ void FD3D12BuddyAllocator::ReleaseAllResources()
 	LLM_SCOPED_PAUSE_TRACKING_FOR_TRACKER(ELLMTracker::Default, ELLMAllocType::System);
 
 #if UE_MEMORY_TRACE_ENABLED
-	D3D12_GPU_VIRTUAL_ADDRESS GPUAddress = 
-		AllocationStrategy == EResourceAllocationStrategy::kPlacedResource ? (BackingHeap ? BackingHeap->GetGPUVirtualAddress() : 0) : (BackingResource ? BackingResource->GetGPUVirtualAddress() : 0);
-	if (GPUAddress > 0)
+	if (AllocationStrategy != EResourceAllocationStrategy::kPlacedResource)
 	{
-		MemoryTrace_UnmarkAllocAsHeap(GPUAddress, TraceHeapId);
+		// Free memory & heap to match alloc operations
+		D3D12_GPU_VIRTUAL_ADDRESS GPUAddress = BackingResource ? BackingResource->GetGPUVirtualAddress() : 0;
+		if (GPUAddress > 0)
+		{
+			MemoryTrace_UnmarkAllocAsHeap(GPUAddress, TraceHeapId);
+			MemoryTrace_Free(GPUAddress, EMemoryTraceRootHeap::VideoMemory);
+		}
 	}
+
 #endif
 
 	for (RetiredBlock& Block : DeferredDeletionQueue)
