RHI: Avoid re-allocating the OutData array in RHIRead*SurfaceData when possible

This behavior was inconsistent across RHIs. In some cases we called SetNumUninitialized() which is the behavior we want (a no-op when the array already has the right size). On other RHIs, the array was being emptied first and then grown which forced a re-allocation.

diff --git a/Engine/Source/Runtime/Apple/MetalRHI/Private/MetalRenderTarget.cpp b/Engine/Source/Runtime/Apple/MetalRHI/Private/MetalRenderTarget.cpp
--- a/Engine/Source/Runtime/Apple/MetalRHI/Private/MetalRenderTarget.cpp
+++ b/Engine/Source/Runtime/Apple/MetalRHI/Private/MetalRenderTarget.cpp
@@ -174,8 +174,7 @@ void FMetalDynamicRHI::RHIReadSurfaceData(FRHITexture* TextureRHI, FIntRect InRe
 	TArray<FColor> OutDataUnConverted;
 	RHIReadSurfaceData(TextureRHI, InRect, OutDataUnConverted, InFlags);
 
-	OutData.Empty();
-	OutData.AddUninitialized(OutDataUnConverted.Num());
+	OutData.SetNumUninitialized(OutDataUnConverted.Num());
 
 	for (uint32 i = 0; i < OutDataUnConverted.Num(); ++i)
 	{
@@ -236,21 +235,20 @@ static void ConvertSurfaceDataToFColor(EPixelFormat Format, uint32 Width, uint32
 void FMetalDynamicRHI::RHIReadSurfaceData(FRHITexture* TextureRHI, FIntRect Rect, TArray<FColor>& OutData, FReadSurfaceDataFlags InFlags)
 {
 	@autoreleasepool {
+
+	// allocate output space
+	const uint32 SizeX = Rect.Width();
+	const uint32 SizeY = Rect.Height();
+	OutData.SetNumUninitialized(SizeX * SizeY);
+
 	if (!ensure(TextureRHI))
 	{
-		OutData.Empty();
-		OutData.AddZeroed(Rect.Width() * Rect.Height());
+		FMemory::Memzero(OutData.GetData(), sizeof(FColor) * OutData.Num());
 		return;
 	}
 
 	FMetalSurface* Surface = GetMetalSurfaceFromRHITexture(TextureRHI);
 
-	// allocate output space
-	const uint32 SizeX = Rect.Width();
-	const uint32 SizeY = Rect.Height();
-	OutData.Empty();
-	OutData.AddUninitialized(SizeX * SizeY);
-	
 	FColor* OutDataPtr = OutData.GetData();
 	mtlpp::Region Region(Rect.Min.X, Rect.Min.Y, SizeX, SizeY);
     
@@ -420,8 +418,7 @@ void FMetalDynamicRHI::RHIReadSurfaceFloatData(FRHITexture* TextureRHI, FIntRect
 	// allocate output space
 	const uint32 SizeX = Rect.Width();
 	const uint32 SizeY = Rect.Height();
-	OutData.Empty();
-	OutData.AddUninitialized(SizeX * SizeY);
+	OutData.SetNumUninitialized(SizeX * SizeY);
 	
 	mtlpp::Region Region = mtlpp::Region(Rect.Min.X, Rect.Min.Y, SizeX, SizeY);
 	
@@ -486,8 +483,7 @@ void FMetalDynamicRHI::RHIRead3DSurfaceFloatData(FRHITexture* TextureRHI,FIntRec
 	const uint32 SizeX = InRect.Width();
 	const uint32 SizeY = InRect.Height();
 	const uint32 SizeZ = ZMinMax.Y - ZMinMax.X;
-	OutData.Empty();
-	OutData.AddUninitialized(SizeX * SizeY * SizeZ);
+	OutData.SetNumUninitialized(SizeX * SizeY * SizeZ);
 	
 	mtlpp::Region Region = mtlpp::Region(InRect.Min.X, InRect.Min.Y, ZMinMax.X, SizeX, SizeY, SizeZ);
 	
diff --git a/Engine/Source/Runtime/D3D12RHI/Private/D3D12RenderTarget.cpp b/Engine/Source/Runtime/D3D12RHI/Private/D3D12RenderTarget.cpp
--- a/Engine/Source/Runtime/D3D12RHI/Private/D3D12RenderTarget.cpp
+++ b/Engine/Source/Runtime/D3D12RHI/Private/D3D12RenderTarget.cpp
@@ -761,8 +761,7 @@ void FD3D12DynamicRHI::ReadSurfaceDataNoMSAARaw(FRHITexture* TextureRHI, FIntRec
 	uint32 BytesPerPixel = GPixelFormats[TextureRHI->GetFormat()].BlockBytes;
 
 	// Allocate the output buffer.
-	OutData.Empty();
-	OutData.AddUninitialized(SizeX * SizeY * BytesPerPixel);
+	OutData.SetNumUninitialized(SizeX * SizeY * BytesPerPixel);
 
 	uint32 BytesPerLine = BytesPerPixel * InRect.Width();
 	const uint32 XBytesAligned = Align((uint32)readBackHeapDesc.Footprint.Width * BytesPerPixel, FD3D12_TEXTURE_DATA_PITCH_ALIGNMENT);
@@ -1181,8 +1180,7 @@ void FD3D12DynamicRHI::RHIReadSurfaceData(FRHITexture* TextureRHI, FIntRect InRe
 	const uint32 SizeY = InRect.Height();
 
 	// Allocate the output buffer.
-	OutData.Empty();
-	OutData.AddUninitialized(SizeX * SizeY);
+	OutData.SetNumUninitialized(SizeX * SizeY);
 
 	uint32 BytesPerPixel = ComputeBytesPerPixel(TextureDesc.Format);
 	uint32 SrcPitch = SizeX * BytesPerPixel;
@@ -1196,8 +1194,8 @@ void FD3D12DynamicRHI::RHIReadSurfaceData(FRHITexture* InRHITexture, FIntRect In
 {
 	if (!ensure(InRHITexture))
 	{
-		OutData.Empty();
-		OutData.AddZeroed(InRect.Width() * InRect.Height());
+		OutData.SetNumUninitialized(InRect.Width() * InRect.Height());
+		FMemory::Memzero(OutData.GetData(), sizeof(FColor) * OutData.Num());
 		return;
 	}
 
@@ -1239,8 +1237,7 @@ void FD3D12DynamicRHI::RHIReadSurfaceData(FRHITexture* InRHITexture, FIntRect In
 	const uint32 SizeY = InRect.Height();
 
 	// Allocate the output buffer.
-	OutData.Empty();
-	OutData.AddUninitialized(SizeX * SizeY);
+	OutData.SetNumUninitialized(SizeX * SizeY);
 
 	FPixelFormatInfo FormatInfo = GPixelFormats[DestTexture2D->GetFormat()];
 	uint32 BytesPerPixel = FormatInfo.BlockBytes;
@@ -1478,8 +1475,8 @@ void FD3D12DynamicRHI::RHIReadSurfaceFloatData(FRHITexture* TextureRHI, FIntRect
 
 	FD3D12Texture* Texture = GetD3D12TextureFromRHITexture(TextureRHI, GPUIndex);
 
-	uint32 SizeX = InRect.Width();
-	uint32 SizeY = InRect.Height();
+	const uint32 SizeX = InRect.Width();
+	const uint32 SizeY = InRect.Height();
 
 	// Check the format of the surface
 	D3D12_RESOURCE_DESC const& TextureDesc = Texture->GetResource()->GetDesc();
@@ -1487,7 +1484,7 @@ void FD3D12DynamicRHI::RHIReadSurfaceFloatData(FRHITexture* TextureRHI, FIntRect
 	check(TextureDesc.Format == GPixelFormats[PF_FloatRGBA].PlatformFormat);
 
 	// Allocate the output buffer.
-	OutData.Empty(SizeX * SizeY);
+	OutData.SetNumUninitialized(SizeX * SizeY);
 
 	// Read back the surface data from defined rect
 	D3D12_BOX	Rect;
@@ -1558,13 +1555,6 @@ void FD3D12DynamicRHI::RHIReadSurfaceFloatData(FRHITexture* TextureRHI, FIntRect
 	D3D12_RANGE Range = { 0, MipBytesAligned };
 	VERIFYD3D12RESULT(TempTexture2D->GetResource()->Map(0, &Range, &pData));
 
-	// Presize the array
-	int32 TotalCount = SizeX * SizeY;
-	if (TotalCount >= OutData.Num())
-	{
-		OutData.AddZeroed(TotalCount);
-	}
-
 	for (int32 Y = InRect.Min.Y; Y < InRect.Max.Y; Y++)
 	{
 		FFloat16Color* SrcPtr = (FFloat16Color*)((uint8*)pData + (Y - InRect.Min.Y) * XBytesAligned);
@@ -1599,9 +1589,9 @@ void FD3D12DynamicRHI::RHIRead3DSurfaceFloatData(FRHITexture* TextureRHI, FIntRe
 
 	FD3D12Texture* Texture = GetD3D12TextureFromRHITexture(TextureRHI, GPUIndex);
 
-	uint32 SizeX = InRect.Width();
-	uint32 SizeY = InRect.Height();
-	uint32 SizeZ = ZMinMax.Y - ZMinMax.X;
+	const uint32 SizeX = InRect.Width();
+	const uint32 SizeY = InRect.Height();
+	const uint32 SizeZ = ZMinMax.Y - ZMinMax.X;
 
 	// Check the format of the surface
 	D3D12_RESOURCE_DESC const& TextureDesc = Texture->GetResource()->GetDesc();
@@ -1611,7 +1601,7 @@ void FD3D12DynamicRHI::RHIRead3DSurfaceFloatData(FRHITexture* TextureRHI, FIntRe
 	check(bIsRGBAFmt || bIsR16FFmt || bIsR32FFmt);
 
 	// Allocate the output buffer.
-	OutData.Empty(SizeX * SizeY * SizeZ * sizeof(FFloat16Color));
+	OutData.SetNumUninitialized(SizeX * SizeY * SizeZ);
 
 	// Read back the surface data from defined rect
 	D3D12_BOX	Rect;
@@ -1669,13 +1659,6 @@ void FD3D12DynamicRHI::RHIRead3DSurfaceFloatData(FRHITexture* TextureRHI, FIntRe
 	void* pData;
 	VERIFYD3D12RESULT(TempTexture3D->GetResource()->Map(0, nullptr, &pData));
 
-	// Presize the array
-	int32 TotalCount = SizeX * SizeY * SizeZ;
-	if (TotalCount >= OutData.Num())
-	{
-		OutData.AddZeroed(TotalCount);
-	}
-
 	// Read the data out of the buffer
 	if (bIsRGBAFmt)
 	{
diff --git a/Engine/Source/Runtime/OpenGLDrv/Private/OpenGLRenderTarget.cpp b/Engine/Source/Runtime/OpenGLDrv/Private/OpenGLRenderTarget.cpp
--- a/Engine/Source/Runtime/OpenGLDrv/Private/OpenGLRenderTarget.cpp
+++ b/Engine/Source/Runtime/OpenGLDrv/Private/OpenGLRenderTarget.cpp
@@ -888,10 +888,12 @@ void FOpenGLDynamicRHI::ReadSurfaceDataRaw(FOpenGLContextState& ContextState, FR
 
 void FOpenGLDynamicRHI::RHIReadSurfaceData(FRHITexture* TextureRHI,FIntRect Rect,TArray<FColor>& OutData, FReadSurfaceDataFlags InFlags)
 {
+	const uint32 Size = Rect.Width() * Rect.Height();
+	OutData.SetNumUninitialized(Size);
+
 	if (!ensure(TextureRHI))
 	{
-		OutData.Empty();
-		OutData.AddZeroed(Rect.Width() * Rect.Height());
+		FMemory::Memzero(OutData.GetData(), Size * sizeof(FColor));
 		return;
 	}
 
@@ -901,15 +903,10 @@ void FOpenGLDynamicRHI::RHIReadSurfaceData(FRHITexture* TextureRHI,FIntRect Rect
 	TArray<uint8> Temp;
 
 	FOpenGLContextState& ContextState = GetContextStateForCurrentContext();
-	OutData.Empty();
 	if (&ContextState != &InvalidContextState)
 	{
 		ReadSurfaceDataRaw(ContextState, TextureRHI, Rect, Temp, InFlags);
 
-		uint32 Size = Rect.Width() * Rect.Height();
-
-		OutData.AddUninitialized(Size);
-
 		FMemory::Memcpy(OutData.GetData(), Temp.GetData(), Size * sizeof(FColor));
 	}
 	RHITHREAD_GLCOMMAND_EPILOGUE();
@@ -943,8 +940,7 @@ void FOpenGLDynamicRHI::RHIReadSurfaceData(FRHITexture* TextureRHI, FIntRect Rec
 	uint32 SizeY = Rect.Height();
 
 	// Initialize output
-	OutData.Empty(SizeX * SizeY);
-	OutData.AddUninitialized(SizeX * SizeY);
+	OutData.SetNumUninitialized(SizeX * SizeY);
 
 	// Bind the framebuffer
 	// @TODO: Do we need to worry about multisampling?
@@ -1036,8 +1032,7 @@ void FOpenGLDynamicRHI::RHIReadSurfaceFloatData(FRHITexture* TextureRHI,FIntRect
 	uint32 SizeX = Rect.Width();
 	uint32 SizeY = Rect.Height();
 
-	OutData.Empty(SizeX * SizeY);
-	OutData.AddUninitialized(SizeX * SizeY);
+	OutData.SetNumUninitialized(SizeX * SizeY);
 
 	glBindFramebuffer(UGL_READ_FRAMEBUFFER, SourceFramebuffer);
 	FOpenGL::ReadBuffer(SourceFramebuffer == 0 ? GL_BACK : GL_COLOR_ATTACHMENT0);
@@ -1092,8 +1087,7 @@ void FOpenGLDynamicRHI::RHIRead3DSurfaceFloatData(FRHITexture* TextureRHI,FIntRe
 	uint32 SizeZ = ZMinMax.Y - ZMinMax.X;
 
 	// Allocate the output buffer.
-	OutData.Empty(SizeX * SizeY * SizeZ * sizeof(FFloat16Color));
-	OutData.AddZeroed(SizeX * SizeY * SizeZ);
+	OutData.SetNumUninitialized(SizeX * SizeY * SizeZ);
 
 	// Set up the source as a temporary FBO
 	uint32 MipmapLevel = 0;
diff --git a/Engine/Source/Runtime/VulkanRHI/Private/VulkanRenderTarget.cpp b/Engine/Source/Runtime/VulkanRHI/Private/VulkanRenderTarget.cpp
--- a/Engine/Source/Runtime/VulkanRHI/Private/VulkanRenderTarget.cpp
+++ b/Engine/Source/Runtime/VulkanRHI/Private/VulkanRenderTarget.cpp
@@ -385,11 +385,11 @@ void FVulkanDynamicRHI::RHIReadSurfaceData(FRHITexture* TextureRHI, FIntRect Rec
 	const uint32 DestWidth = Rect.Max.X - Rect.Min.X;
 	const uint32 DestHeight = Rect.Max.Y - Rect.Min.Y;
 	const uint32 NumRequestedPixels = DestWidth * DestHeight;
+	OutData.SetNumUninitialized(NumRequestedPixels);
 	if (GIgnoreCPUReads)
 	{
 		// Debug: Fill with CPU
-		OutData.Empty(0);
-		OutData.AddZeroed(NumRequestedPixels);
+		FMemory::Memzero(OutData.GetData(), NumRequestedPixels * sizeof(FColor));
 		return;
 	}
 
@@ -405,8 +405,7 @@ void FVulkanDynamicRHI::RHIReadSurfaceData(FRHITexture* TextureRHI, FIntRect Rec
 
 	default:
 		// Just return black for texture types we don't support.
-		OutData.Empty(0);
-		OutData.AddZeroed(NumRequestedPixels);
+		FMemory::Memzero(OutData.GetData(), NumRequestedPixels * sizeof(FColor));
 		return;
 	}
 
@@ -487,7 +486,6 @@ void FVulkanDynamicRHI::RHIReadSurfaceData(FRHITexture* TextureRHI, FIntRect Rec
 		SrcPitch = DestWidth * PixelByteSize;
 	}
 
-	OutData.SetNum(NumRequestedPixels);
 	FColor* Dest = OutData.GetData();
 	ConvertRawDataToFColor(Surface.StorageFormat, DestWidth, DestHeight, In, SrcPitch, Dest, InFlags);
 
@@ -503,9 +501,10 @@ void FVulkanDynamicRHI::RHIReadSurfaceData(FRHITexture* TextureRHI, FIntRect Rec
 {
 	TArray<FColor> FromColorData;
 	RHIReadSurfaceData(TextureRHI, Rect, FromColorData, InFlags);
-	for (FColor& From : FromColorData)
+	OutData.SetNumUninitialized(FromColorData.Num());
+	for (int Index = 0, Num = FromColorData.Num(); Index < Num; Index++)
 	{
-		OutData.Emplace(FLinearColor(From));
+		OutData[Index] = FLinearColor(FromColorData[Index]);
 	}
 }
 
@@ -609,7 +608,7 @@ void FVulkanDynamicRHI::RHIReadSurfaceFloatData(FRHITexture* TextureRHI, FIntRec
 
 		uint32 OutWidth = InRect.Max.X - InRect.Min.X;
 		uint32 OutHeight= InRect.Max.Y - InRect.Min.Y;
-		OutputData.SetNum(OutWidth * OutHeight);
+		OutputData.SetNumUninitialized(OutWidth * OutHeight);
 		uint32 OutIndex = 0;
 		FFloat16Color* Dest = OutputData.GetData();
 		for (int32 Row = InRect.Min.Y; Row < InRect.Max.Y; ++Row)
@@ -687,12 +686,12 @@ void FVulkanDynamicRHI::RHIRead3DSurfaceFloatData(FRHITexture* TextureRHI, FIntR
 	const uint32 Size = NumPixels * sizeof(FFloat16Color);
 
 	// Allocate the output buffer.
-	OutData.Reserve(Size);
+	OutData.SetNumUninitialized(Size);
 
 	if (GIgnoreCPUReads == 1)
 	{
 		// Debug: Fill with CPU
-		OutData.AddZeroed(Size);
+		FMemory::Memzero(OutData.GetData(), Size * sizeof(FFloat16Color));
 		return;
 	}
 
@@ -753,7 +752,6 @@ void FVulkanDynamicRHI::RHIRead3DSurfaceFloatData(FRHITexture* TextureRHI, FIntR
 
 	StagingBuffer->InvalidateMappedMemory();
 
-	OutData.SetNum(NumPixels);
 	FFloat16Color* Dest = OutData.GetData();
 	for (int32 Layer = ZMinMax.X; Layer < ZMinMax.Y; ++Layer)
 	{
diff --git a/Engine/Source/Runtime/Windows/D3D11RHI/Private/D3D11RenderTarget.cpp b/Engine/Source/Runtime/Windows/D3D11RHI/Private/D3D11RenderTarget.cpp
--- a/Engine/Source/Runtime/Windows/D3D11RHI/Private/D3D11RenderTarget.cpp
+++ b/Engine/Source/Runtime/Windows/D3D11RHI/Private/D3D11RenderTarget.cpp
@@ -647,8 +647,8 @@ void FD3D11DynamicRHI::RHIReadSurfaceData(FRHITexture* TextureRHI,FIntRect InRec
 {
 	if (!ensure(TextureRHI))
 	{
-		OutData.Empty();
-		OutData.AddZeroed(InRect.Width() * InRect.Height());
+		OutData.SetNumUninitialized(InRect.Width() * InRect.Height());
+		FMemory::Memzero(OutData.GetData(), OutData.Num() * sizeof(FColor));
 		return;
 	}
 
@@ -675,8 +675,7 @@ void FD3D11DynamicRHI::RHIReadSurfaceData(FRHITexture* TextureRHI,FIntRect InRec
 	const uint32 SizeY = InRect.Height();
 
 	// Allocate the output buffer.
-	OutData.Empty();
-	OutData.AddUninitialized(SizeX * SizeY);
+	OutData.SetNumUninitialized(SizeX * SizeY);
 
 	uint32 BytesPerPixel = ComputeBytesPerPixel(TextureDesc.Format);
 	uint32 SrcPitch = SizeX * BytesPerPixel;
@@ -848,7 +847,7 @@ void FD3D11DynamicRHI::RHIReadSurfaceFloatData(FRHITexture* TextureRHI,FIntRect
 	check(TextureDesc.Format == GPixelFormats[PF_FloatRGBA].PlatformFormat);
 
 	// Allocate the output buffer.
-	OutData.Empty(SizeX * SizeY);
+	OutData.SetNumUninitialized(SizeX * SizeY);
 
 	// Read back the surface data from defined rect
 	D3D11_BOX	Rect;
@@ -894,13 +893,6 @@ void FD3D11DynamicRHI::RHIReadSurfaceFloatData(FRHITexture* TextureRHI,FIntRect
 	D3D11_MAPPED_SUBRESOURCE LockedRect;
 	VERIFYD3D11RESULT_EX(Direct3DDeviceIMContext->Map(TempTexture2D,0,D3D11_MAP_READ,0,&LockedRect), Direct3DDevice);
 
-	// Presize the array
-	int32 TotalCount = SizeX * SizeY;
-	if (TotalCount >= OutData.Num())
-	{
-		OutData.AddZeroed(TotalCount);
-	}
-
 	for(int32 Y = InRect.Min.Y; Y < InRect.Max.Y; Y++)
 	{
 		FFloat16Color* SrcPtr = (FFloat16Color*)((uint8*)LockedRect.pData + (Y - InRect.Min.Y) * LockedRect.RowPitch);
@@ -994,8 +986,7 @@ void FD3D11DynamicRHI::RHIReadSurfaceData(FRHITexture* TextureRHI, FIntRect InRe
 	const uint32 SizeY = InRect.Height();
 
 	// Allocate the output buffer.
-	OutData.Empty();
-	OutData.AddUninitialized(SizeX * SizeY);
+	OutData.SetNumUninitialized(SizeX * SizeY);
 
 	uint32 BytesPerPixel = ComputeBytesPerPixel(TextureDesc.Format);
 	uint32 SrcPitch = SizeX * BytesPerPixel;
@@ -1024,7 +1015,7 @@ void FD3D11DynamicRHI::RHIRead3DSurfaceFloatData(FRHITexture* TextureRHI,FIntRec
 	check(bIsRGBAFmt || bIsR16FFmt || bIsR32FFmt);
 
 	// Allocate the output buffer.
-	OutData.Empty(SizeX * SizeY * SizeZ * sizeof(FFloat16Color));
+	OutData.SetNumUninitialized(SizeX * SizeY * SizeZ);
 
 	// Read back the surface data from defined rect
 	D3D11_BOX	Rect;
@@ -1058,13 +1049,6 @@ void FD3D11DynamicRHI::RHIRead3DSurfaceFloatData(FRHITexture* TextureRHI,FIntRec
 	D3D11_MAPPED_SUBRESOURCE LockedRect;
 	VERIFYD3D11RESULT_EX(Direct3DDeviceIMContext->Map(TempTexture3D,0,D3D11_MAP_READ,0,&LockedRect), Direct3DDevice);
 
-	// Presize the array
-	int32 TotalCount = SizeX * SizeY * SizeZ;
-	if (TotalCount >= OutData.Num())
-	{
-		OutData.AddZeroed(TotalCount);
-	}
-
 	// Read the data out of the buffer
 	if (bIsRGBAFmt)
 	{
