Fix crash when running D3D12 with -nothreading
 - Skip creating the interrupt thread and fallback to polling mode.

diff --git a/Engine/Source/Runtime/D3D12RHI/Private/D3D12Submission.cpp b/Engine/Source/Runtime/D3D12RHI/Private/D3D12Submission.cpp
--- a/Engine/Source/Runtime/D3D12RHI/Private/D3D12Submission.cpp
+++ b/Engine/Source/Runtime/D3D12RHI/Private/D3D12Submission.cpp
@@ -107,27 +107,29 @@ private:
 
 void FD3D12DynamicRHI::InitializeSubmissionPipe()
 {
+	if (FPlatformProcess::SupportsMultithreading())
+	{
 #if D3D12_USE_INTERRUPT_THREAD
-	InterruptThread = new FD3D12Thread(TEXT("RHIInterruptThread"), TPri_Highest, this, &FD3D12DynamicRHI::ProcessInterruptQueue);
+		InterruptThread = new FD3D12Thread(TEXT("RHIInterruptThread"), TPri_Highest, this, &FD3D12DynamicRHI::ProcessInterruptQueue);
 #endif
 
 #if D3D12_USE_SUBMISSION_THREAD
+		bool bUseSubmissionThread = false;
+		switch (CVarRHIUseSubmissionThread.GetValueOnAnyThread())
+		{
+		case 1: bUseSubmissionThread = FRHIGPUMask::All().HasSingleIndex(); break;
+		case 2: bUseSubmissionThread = true; break;
+		}
 
-	bool bUseSubmissionThread = false;
-	switch (CVarRHIUseSubmissionThread.GetValueOnAnyThread())
-	{
-	case 1: bUseSubmissionThread = FRHIGPUMask::All().HasSingleIndex(); break;
-	case 2: bUseSubmissionThread = true; break;
-	}
-
-	// Currently RenderDoc can't make programmatic captures when we use a submission thread.
-	bUseSubmissionThread &= !IRenderCaptureProvider::IsAvailable() || IRenderCaptureProvider::Get().CanSupportSubmissionThread();
+		// Currently RenderDoc can't make programmatic captures when we use a submission thread.
+		bUseSubmissionThread &= !IRenderCaptureProvider::IsAvailable() || IRenderCaptureProvider::Get().CanSupportSubmissionThread();
 
-	if (bUseSubmissionThread)
-	{
-		SubmissionThread = new FD3D12Thread(TEXT("RHISubmissionThread"), TPri_Highest, this, &FD3D12DynamicRHI::ProcessSubmissionQueue);
-	}
+		if (bUseSubmissionThread)
+		{
+			SubmissionThread = new FD3D12Thread(TEXT("RHISubmissionThread"), TPri_Highest, this, &FD3D12DynamicRHI::ProcessSubmissionQueue);
+		}
 #endif
+	}
 
 	FlushTiming(true);
 }
@@ -585,12 +587,10 @@ FD3D12DynamicRHI::FProcessResult FD3D12DynamicRHI::ProcessSubmissionQueue()
 		Payload->Queue.PendingInterrupt.Enqueue(Payload);
 	}
 
-#if D3D12_USE_INTERRUPT_THREAD
-	if (EnumHasAnyFlags(Result.Status, EQueueStatus::Processed))
+	if (InterruptThread && EnumHasAnyFlags(Result.Status, EQueueStatus::Processed))
 	{
 		InterruptThread->Kick();
 	}
-#endif
 
 	return Result;
 }
@@ -875,45 +875,44 @@ void FD3D12SyncPoint::Wait() const
 
 void FD3D12DynamicRHI::ProcessInterruptQueueUntil(FGraphEvent* GraphEvent)
 {
-#if D3D12_USE_INTERRUPT_THREAD
-
-	if (GraphEvent && !GraphEvent->IsComplete())
+	if (InterruptThread)
 	{
-		GraphEvent->Wait();
-	}
-
-#else
-
-	// Use the current thread to process the interrupt queue until the sync point we're waiting for is signaled.
-	// If GraphEvent is nullptr, process the queue until no further progress is made (assuming we can acquire the lock), then return.
-	if (!GraphEvent || !GraphEvent->IsComplete())
-	{
-		// If we're waiting for a sync point, accumulate the idle time
-		FThreadIdleStats::FScopeIdle IdleScope(/* bIgnore = */GraphEvent == nullptr);
-
-	Retry:
-		if (InterruptCS.TryLock())
+		if (GraphEvent && !GraphEvent->IsComplete())
 		{
-			FProcessResult Result;
-			do { Result = ProcessInterruptQueue(); }
-			// If we have a sync point, keep processing until the sync point is signaled.
-			// Otherwise, process until no more progress is being made.
-			while (GraphEvent
-				? !GraphEvent->IsComplete()
-				: EnumHasAllFlags(Result.Status, EQueueStatus::Processed)
-			);
-
-			InterruptCS.Unlock();
+			GraphEvent->Wait();
 		}
-		else if (GraphEvent && !GraphEvent->IsComplete())
+	}
+	else
+	{
+		// Use the current thread to process the interrupt queue until the sync point we're waiting for is signaled.
+		// If GraphEvent is nullptr, process the queue until no further progress is made (assuming we can acquire the lock), then return.
+		if (!GraphEvent || !GraphEvent->IsComplete())
 		{
-			// Failed to get the lock. Another thread is processing the interrupt queue. Try again...
-			FPlatformProcess::SleepNoStats(0);
-			goto Retry;
+			// If we're waiting for a sync point, accumulate the idle time
+			FThreadIdleStats::FScopeIdle IdleScope(/* bIgnore = */GraphEvent == nullptr);
+
+		Retry:
+			if (InterruptCS.TryLock())
+			{
+				FProcessResult Result;
+				do { Result = ProcessInterruptQueue(); }
+				// If we have a sync point, keep processing until the sync point is signaled.
+				// Otherwise, process until no more progress is being made.
+				while (GraphEvent
+					? !GraphEvent->IsComplete()
+					: EnumHasAllFlags(Result.Status, EQueueStatus::Processed)
+				);
+
+				InterruptCS.Unlock();
+			}
+			else if (GraphEvent && !GraphEvent->IsComplete())
+			{
+				// Failed to get the lock. Another thread is processing the interrupt queue. Try again...
+				FPlatformProcess::SleepNoStats(0);
+				goto Retry;
+			}
 		}
 	}
-
-#endif
 }
 
 FD3D12DynamicRHI::FProcessResult FD3D12DynamicRHI::ProcessInterruptQueue()
@@ -943,16 +942,14 @@ FD3D12DynamicRHI::FProcessResult FD3D12DynamicRHI::ProcessInterruptQueue()
 
 			if (CompletedFenceValue < Payload->CompletionFenceValue)
 			{
-#if D3D12_USE_INTERRUPT_THREAD
 				// Command list batch has not yet completed on this queue.
 				// Ask the driver to wake this thread again when the required value is reached.
-				if (!CurrentQueue.Fence.bInterruptAwaited)
+				if (InterruptThread && !CurrentQueue.Fence.bInterruptAwaited)
 				{
 					SCOPED_NAMED_EVENT_TEXT("SetEventOnCompletion", FColor::Red);
 					VERIFYD3D12RESULT(CurrentQueue.Fence.D3DFence->SetEventOnCompletion(Payload->CompletionFenceValue, InterruptThread->Event));
 					CurrentQueue.Fence.bInterruptAwaited = true;
 				}
-#endif // D3D12_USE_INTERRUPT_THREAD
 
 				// Skip processing this queue and move on to the next.
 				Result.Status |= EQueueStatus::Pending;
