D3D12:  Fix for D3D validation errors due to creating constant buffers smaller than the required alignment size

- issue
  - ray tracing "loose parameter" constant buffer size is getting aligned up for the allocation but the actual buffer creation is using the pre-aligned size

- fix
  - align up the buffer size prior to buffer allocation and creation

diff --git a/Engine/Source/Runtime/D3D12RHI/Private/D3D12RayTracing.cpp b/Engine/Source/Runtime/D3D12RHI/Private/D3D12RayTracing.cpp
--- a/Engine/Source/Runtime/D3D12RHI/Private/D3D12RayTracing.cpp
+++ b/Engine/Source/Runtime/D3D12RHI/Private/D3D12RayTracing.cpp
@@ -4847,10 +4847,11 @@ static bool SetRayTracingShaderResources(
 	if (InLooseParameterData && Shader->UsesGlobalUniformBuffer())
 	{
 		const uint32 CBVIndex = 0; // Global uniform buffer is always assumed to be in slot 0
-		D3D12_GPU_VIRTUAL_ADDRESS BufferAddress = Binder.CreateTransientConstantBuffer(InLooseParameterData, InLooseParameterDataSize);
+		const uint32 AlignedSize = Align(InLooseParameterDataSize, D3D12_CONSTANT_BUFFER_DATA_PLACEMENT_ALIGNMENT);
+		D3D12_GPU_VIRTUAL_ADDRESS BufferAddress = Binder.CreateTransientConstantBuffer(InLooseParameterData, AlignedSize);
 
 	#if USE_STATIC_ROOT_SIGNATURE
-		Bindings.LocalCBVs[CBVIndex] = Binder.CreateTransientConstantBufferView(BufferAddress, InLooseParameterDataSize);
+		Bindings.LocalCBVs[CBVIndex] = Binder.CreateTransientConstantBufferView(BufferAddress, AlignedSize);
 	#else // USE_STATIC_ROOT_SIGNATURE
 		Bindings.LocalCBVs[CBVIndex] = BufferAddress;
 	#endif // USE_STATIC_ROOT_SIGNATURE
