Suppress -Wenum-constexpr-conversion in d3d12rhi (clang 16)

diff --git a/Engine/Source/Runtime/D3D12RHI/Public/D3D12Resources.h b/Engine/Source/Runtime/D3D12RHI/Public/D3D12Resources.h
--- a/Engine/Source/Runtime/D3D12RHI/Public/D3D12Resources.h
+++ b/Engine/Source/Runtime/D3D12RHI/Public/D3D12Resources.h
@@ -927,8 +927,10 @@ public:
 
 class FD3D12ResourceBarrierBatcher
 {
+	PRAGMA_DISABLE_ENUM_CONSTEXPR_CONVERSION_WARNINGS
 	// Use the top bit of the flags enum to mark transitions as "idle" time (used to remove the swapchain wait time for back buffers).
 	static constexpr D3D12_RESOURCE_BARRIER_FLAGS BarrierFlag_CountAsIdleTime = D3D12_RESOURCE_BARRIER_FLAGS(1ull << ((sizeof(D3D12_RESOURCE_BARRIER_FLAGS) * 8) - 1));
+	PRAGMA_ENABLE_ENUM_CONSTEXPR_CONVERSION_WARNINGS
 
 	struct FD3D12ResourceBarrier : public D3D12_RESOURCE_BARRIER
 	{
