Fix format string errors when logging resource creation errors.

The resource width is a 64-bit int, but we were using %u for it.

#rnx

diff --git a/Engine/Source/Runtime/D3D12RHI/Private/D3D12Resources.cpp b/Engine/Source/Runtime/D3D12RHI/Private/D3D12Resources.cpp
--- a/Engine/Source/Runtime/D3D12RHI/Private/D3D12Resources.cpp
+++ b/Engine/Source/Runtime/D3D12RHI/Private/D3D12Resources.cpp
@@ -461,7 +461,7 @@ HRESULT FD3D12Adapter::CreateCommittedResource(const FD3D12ResourceDesc& InDesc,
 	}
 	else	
 	{
-		UE_LOG(LogD3D12RHI, Display, TEXT("D3D12 CreateCommittedResource failed with params:\n\tHeap Type: %d\n\tHeap Flags: %d\n\tResource Dimension: %d\n\tResource Width: %d\n\tResource Height: %d\n\tFormat: %d\n\tResource Flags: %d"),
+		UE_LOG(LogD3D12RHI, Display, TEXT("D3D12 CreateCommittedResource failed with params:\n\tHeap Type: %d\n\tHeap Flags: %d\n\tResource Dimension: %d\n\tResource Width: %llu\n\tResource Height: %u\n\tFormat: %d\n\tResource Flags: %d"),
 			HeapProps.Type, HeapFlags, LocalDesc.Dimension, LocalDesc.Width, LocalDesc.Height, LocalDesc.PixelFormat, LocalDesc.Flags);
 
 		if (bVerifyHResult)
@@ -513,7 +513,7 @@ HRESULT FD3D12Adapter::CreateReservedResource(const FD3D12ResourceDesc& InDesc,
 	}
 	else
 	{
-		UE_LOG(LogD3D12RHI, Display, TEXT("D3D12 CreateReservedResource failed with params:\n\tResource Dimension: %d\n\tResource Width: %d\n\tResource Height: %d\n\tFormat: %d\n\tResource Flags: %d"),
+		UE_LOG(LogD3D12RHI, Display, TEXT("D3D12 CreateReservedResource failed with params:\n\tResource Dimension: %d\n\tResource Width: %llu\n\tResource Height: %u\n\tFormat: %d\n\tResource Flags: %d"),
 			LocalDesc.Dimension, LocalDesc.Width, LocalDesc.Height, LocalDesc.PixelFormat, LocalDesc.Flags);
 
 		if (bVerifyHResult)
@@ -597,7 +597,7 @@ HRESULT FD3D12Adapter::CreatePlacedResource(const FD3D12ResourceDesc& InDesc, FD
 	}
 	else
 	{
-		UE_LOG(LogD3D12RHI, Display, TEXT("D3D12 CreatePlacedResource failed with params:\n\tHeap Type: %d\n\tHeap Flags: %d\n\tResource Dimension: %d\n\tResource Width: %d\n\tResource Height: %d\n\tHeightFormat: %d\n\tResource Flags: %d"),
+		UE_LOG(LogD3D12RHI, Display, TEXT("D3D12 CreatePlacedResource failed with params:\n\tHeap Type: %d\n\tHeap Flags: %d\n\tResource Dimension: %d\n\tResource Width: %llu\n\tResource Height: %u\n\tFormat: %d\n\tResource Flags: %d"),
 			BackingHeap->GetHeapDesc().Properties.Type, BackingHeap->GetHeapDesc().Flags, InDesc.Dimension, InDesc.Width, InDesc.Height, InDesc.PixelFormat, InDesc.Flags);
 
 		if (bVerifyHResult)
