Add lightweight DRED (ID3D12DeviceRemovedExtendedDataSettings2). Turn lightweight DRED on by default (when supported) and disable Nvidia Aftermath so they don't interfere with each other. Aftermath remains on by default (old behavior) if lightweight DRED is unsupported or heavyweight DRED is requested.

diff --git a/Engine/Source/Runtime/D3D12RHI/Private/D3D12Adapter.cpp b/Engine/Source/Runtime/D3D12RHI/Private/D3D12Adapter.cpp
--- a/Engine/Source/Runtime/D3D12RHI/Private/D3D12Adapter.cpp
+++ b/Engine/Source/Runtime/D3D12RHI/Private/D3D12Adapter.cpp
@@ -383,41 +383,96 @@ void FD3D12Adapter::CreateRootDevice(bool bWithDebug)
 
 	// Setup DRED if requested
 	bool bDRED = false;
+	bool bDREDMarkersOnly = false;
 	bool bDREDContext = false;
-	if (EnumHasAnyFlags(GPUCrashDebuggingModes, ED3D12GPUCrashDebuggingModes::DRED))
 	{
-		TRefCountPtr<ID3D12DeviceRemovedExtendedDataSettings> DredSettings;
-		HRESULT hr = D3D12GetDebugInterface(IID_PPV_ARGS(DredSettings.GetInitReference()));
+		bool bEnableDRED = EnumHasAnyFlags(GPUCrashDebuggingModes, ED3D12GPUCrashDebuggingModes::DRED);
+		
+		HMODULE d3d12DllHandle = (HMODULE)FPlatformProcess::GetDllHandle(TEXT("d3d12.dll"));
+		typedef HRESULT(WINAPI* FD3D12GetInterface)(REFCLSID, REFIID, void**);
 
-		// Can fail if not on correct Windows Version - needs 1903 or newer
-		if (SUCCEEDED(hr))
+		if (d3d12DllHandle)
 		{
-			// Turn on AutoBreadcrumbs and Page Fault reporting
-			DredSettings->SetAutoBreadcrumbsEnablement(D3D12_DRED_ENABLEMENT_FORCED_ON);
-			DredSettings->SetPageFaultEnablement(D3D12_DRED_ENABLEMENT_FORCED_ON);
+			FD3D12GetInterface D3D12GetInterfaceFnPtr = (FD3D12GetInterface)(void*)(GetProcAddress(d3d12DllHandle, "D3D12GetInterface"));
 
-			bDRED = true;
-			UE_LOG(LogD3D12RHI, Log, TEXT("[DRED] Dred enabled"));
-		}
-		else
-		{
-			UE_LOG(LogD3D12RHI, Log, TEXT("[DRED] DRED requested but interface was not found, hresult: %x. DRED only works on Windows 10 1903+."), hr);
-		}
+			if (D3D12GetInterfaceFnPtr != nullptr)
+			{
+				if (bEnableDRED)
+				{
+					TRefCountPtr<ID3D12DeviceRemovedExtendedDataSettings> DredSettings;
+					HRESULT hr = D3D12GetInterfaceFnPtr(CLSID_D3D12DeviceRemovedExtendedData, IID_PPV_ARGS(DredSettings.GetInitReference()));
+
+					// Can fail if not on correct Windows Version - needs 1903 or newer
+					if (SUCCEEDED(hr))
+					{
+						// Turn on AutoBreadcrumbs and Page Fault reporting
+						DredSettings->SetAutoBreadcrumbsEnablement(D3D12_DRED_ENABLEMENT_FORCED_ON);
+						DredSettings->SetPageFaultEnablement(D3D12_DRED_ENABLEMENT_FORCED_ON);
+
+						bDRED = true;
+						UE_LOG(LogD3D12RHI, Log, TEXT("[DRED] Dred enabled"));
+					}
+					else
+					{
+						UE_LOG(LogD3D12RHI, Log, TEXT("[DRED] DRED requested but interface was not found, hresult: %x. DRED only works on Windows 10 1903+."), hr);
+					}
 
 #ifdef __ID3D12DeviceRemovedExtendedDataSettings1_INTERFACE_DEFINED__
-		TRefCountPtr<ID3D12DeviceRemovedExtendedDataSettings1> DredSettings1;
-		hr = D3D12GetDebugInterface(IID_PPV_ARGS(DredSettings1.GetInitReference()));
-		if (SUCCEEDED(hr))
-		{
-			DredSettings1->SetBreadcrumbContextEnablement(D3D12_DRED_ENABLEMENT_FORCED_ON);
-			bDREDContext = true;
-			UE_LOG(LogD3D12RHI, Log, TEXT("[DRED] Dred breadcrumb context enabled"));
-		}
+					TRefCountPtr<ID3D12DeviceRemovedExtendedDataSettings1> DredSettings1;
+					hr = D3D12GetInterfaceFnPtr(CLSID_D3D12DeviceRemovedExtendedData, IID_PPV_ARGS(DredSettings1.GetInitReference()));
+
+					if (SUCCEEDED(hr))
+					{
+						DredSettings1->SetBreadcrumbContextEnablement(D3D12_DRED_ENABLEMENT_FORCED_ON);
+						bDREDContext = true;
+						UE_LOG(LogD3D12RHI, Log, TEXT("[DRED] Dred breadcrumb context enabled"));
+					}
+#endif // __ID3D12DeviceRemovedExtendedDataSettings1_INTERFACE_DEFINED__
+				}
+
+#ifdef __ID3D12DeviceRemovedExtendedDataSettings2_INTERFACE_DEFINED__
+				{
+					SetEmitDrawEvents(true);
+
+					TRefCountPtr<ID3D12DeviceRemovedExtendedDataSettings2> DredSettings2;
+					HRESULT hr = D3D12GetInterfaceFnPtr(CLSID_D3D12DeviceRemovedExtendedData, IID_PPV_ARGS(DredSettings2.GetInitReference()));
+
+					if (SUCCEEDED(hr))
+					{
+						// Always enable DRED if lightweight option is available (version 1.2)
+						bEnableDRED = true;
+
+						// Turn on AutoBreadcrumbs and Page Fault reporting
+						DredSettings2->SetAutoBreadcrumbsEnablement(D3D12_DRED_ENABLEMENT_FORCED_ON);
+						DredSettings2->SetPageFaultEnablement(D3D12_DRED_ENABLEMENT_FORCED_ON);
+
+						if (EnumHasAnyFlags(GPUCrashDebuggingModes, ED3D12GPUCrashDebuggingModes::DRED))
+						{
+							bDREDMarkersOnly = false;
+							DredSettings2->UseMarkersOnlyAutoBreadcrumbs(false);
+						}
+						else
+						{
+							bDREDMarkersOnly = true;
+							DredSettings2->UseMarkersOnlyAutoBreadcrumbs(true);
+							UE_LOG(LogD3D12RHI, Log, TEXT("[DRED] Using lightweight DRED."));
+#if NV_AFTERMATH
+							// Disable Nvidia Aftermath if lightweight DRED is enabled by default so they don't interfere with each other.
+							GDX12NVAfterMathEnabled = 0;
+							EnumRemoveFlags(GPUCrashDebuggingModes, ED3D12GPUCrashDebuggingModes::NvAftermath);
+							UE_LOG(LogD3D12RHI, Log, TEXT("[DRED] Disabled Nvidia Aftermath to avoid conflicts."));
 #endif
+						}
+					}
+				}
+			}
+#endif // __ID3D12DeviceRemovedExtendedDataSettings2_INTERFACE_DEFINED__
+		}
 	}
 
 	FGenericCrashContext::SetEngineData(TEXT("RHI.DRED"), bDRED ? TEXT("true") : TEXT("false"));
-	FGenericCrashContext::SetEngineData(TEXT("RHI.DREDContext"), bDREDContext ? TEXT("true") : TEXT("false"));
+	FGenericCrashContext::SetEngineData(TEXT("RHI.DREDMarkersOnly"), bDREDMarkersOnly && bDRED ? TEXT("true") : TEXT("false"));
+	FGenericCrashContext::SetEngineData(TEXT("RHI.DREDContext"), bDREDContext && bDRED ? TEXT("true") : TEXT("false"));
 
 #endif // PLATFORM_WINDOWS || (PLATFORM_HOLOLENS && !UE_BUILD_SHIPPING && WITH_PIX_EVENT_RUNTIME)
 
@@ -717,12 +772,9 @@ void FD3D12Adapter::CreateRootDevice(bool bWithDebug)
 
 			// Break on D3D debug errors.
 			pd3dInfoQueue->SetBreakOnSeverity(D3D12_MESSAGE_SEVERITY_ERROR, true);
-
-			// Enable this to break on a specific id in order to quickly get a callstack
-			//pd3dInfoQueue->SetBreakOnID(D3D12_MESSAGE_ID_DEVICE_DRAW_CONSTANT_BUFFER_TOO_SMALL, true);
-
-			// Break on D3D warnings if warning log or warning breakpoint is enabled
-			if (bLogWarnings)
+			
+			// By default trace warnings to log (don't break) - exception handler takes care this
+			if (!FParse::Param(FCommandLine::Get(), TEXT("nod3dbreakonwarning")))
 			{
 				pd3dInfoQueue->SetBreakOnSeverity(D3D12_MESSAGE_SEVERITY_WARNING, true);
 			}
diff --git a/Engine/Source/Runtime/D3D12RHI/Public/Windows/D3D12ThirdParty.h b/Engine/Source/Runtime/D3D12RHI/Public/Windows/D3D12ThirdParty.h
--- a/Engine/Source/Runtime/D3D12RHI/Public/Windows/D3D12ThirdParty.h
+++ b/Engine/Source/Runtime/D3D12RHI/Public/Windows/D3D12ThirdParty.h
@@ -6,6 +6,8 @@
 #pragma warning(push)
 #pragma warning(disable : 4005)	// macro redefinition
 
+#include "HAL/Platform.h"
+
 #if PLATFORM_64BITS
 	#pragma pack(push,16)
 #else
@@ -13,6 +15,7 @@
 #endif
 
 #include "Microsoft/AllowMicrosoftPlatformTypes.h"
+#define INITGUID
 THIRD_PARTY_INCLUDES_START
 	#include <d3d12.h>
 	#include <d3dx12.h>
