Fix D3D12 ray tracing GPU crash when decals, light functions and landscape are used together.
The core issue was miss/callable record indices were calculated incorrectly when the SBT record size was not a multiple of D3D12_RAYTRACING_SHADER_TABLE_BYTE_ALIGNMENT.
This resulted in bindings written to incorrect addresses and a GPU crash.
Modified SBT code to use offsets instead of indices.

diff --git a/Engine/Source/Runtime/D3D12RHI/Private/D3D12RayTracing.cpp b/Engine/Source/Runtime/D3D12RHI/Private/D3D12RayTracing.cpp
--- a/Engine/Source/Runtime/D3D12RHI/Private/D3D12RayTracing.cpp
+++ b/Engine/Source/Runtime/D3D12RHI/Private/D3D12RayTracing.cpp
@@ -1765,14 +1765,13 @@ private:
 		bIsDirty = true;
 	}
 
-	void WriteLocalShaderRecord(uint32 RecordIndex, uint32 OffsetWithinRecord, const void* InData, uint32 InDataSize)
+	void WriteLocalShaderRecord(uint32 ShaderTableOffset, uint32 RecordIndex, uint32 OffsetWithinRecord, const void* InData, uint32 InDataSize)
 	{
 		checkfSlow(OffsetWithinRecord % 4 == 0, TEXT("SBT record parameters must be written on DWORD-aligned boundary"));
 		checkfSlow(InDataSize % 4 == 0, TEXT("SBT record parameters must be DWORD-aligned"));
 		checkfSlow(OffsetWithinRecord + InDataSize <= LocalRecordSizeUnaligned, TEXT("SBT record write request is out of bounds"));
-		checkfSlow(RecordIndex < NumLocalRecords, TEXT("SBT local record write request is out of bounds"));
 
-		const uint32 WriteOffset = LocalShaderTableOffset + LocalRecordStride * RecordIndex + OffsetWithinRecord;
+		const uint32 WriteOffset = ShaderTableOffset + LocalRecordStride * RecordIndex + OffsetWithinRecord;
 
 		WriteData(WriteOffset, InData, InDataSize);
 	}
@@ -1870,12 +1869,6 @@ public:
 		TotalDataSize += Initializer.NumMissRecords * LocalRecordStride;
 		TotalDataSize = RoundUpToNextMultiple(TotalDataSize, D3D12_RAYTRACING_SHADER_TABLE_BYTE_ALIGNMENT);
 
-		// Hit groups and callable shaders are stored in a consecutive memory block and are accessed using common local record indices.
-		LocalShaderTableOffset = HitGroupShaderTableOffset;
-		CallableShaderRecordIndexOffset = (CallableShaderTableOffset - LocalShaderTableOffset) / LocalRecordStride;
-		MissShaderRecordIndexOffset = (MissShaderTableOffset - LocalShaderTableOffset) / LocalRecordStride;
-		NumLocalRecords = (TotalDataSize - LocalShaderTableOffset) / LocalRecordStride;
-
 		Data.SetNumZeroed(TotalDataSize);
 #if DO_CHECK
 		bWasDefaultMissShaderSet = false;
@@ -1890,19 +1883,33 @@ public:
 	}
 
 	template <typename T>
-	void SetLocalShaderParameters(uint32 RecordIndex, uint32 InOffsetWithinRootSignature, const T& Parameters)
+	void SetLocalShaderParameters(uint32 ShaderTableOffset, uint32 RecordIndex, uint32 InOffsetWithinRootSignature, const T& Parameters)
 	{
-		WriteLocalShaderRecord(RecordIndex, ShaderIdentifierSize + InOffsetWithinRootSignature, &Parameters, sizeof(Parameters));
+		WriteLocalShaderRecord(ShaderTableOffset, RecordIndex, ShaderIdentifierSize + InOffsetWithinRootSignature, &Parameters, sizeof(Parameters));
 	}
 
-	void SetLocalShaderParameters(uint32 RecordIndex, uint32 InOffsetWithinRootSignature, const void* InData, uint32 InDataSize)
+	void SetLocalShaderParameters(uint32 ShaderTableOffset, uint32 RecordIndex, uint32 InOffsetWithinRootSignature, const void* InData, uint32 InDataSize)
 	{
-		WriteLocalShaderRecord(RecordIndex, ShaderIdentifierSize + InOffsetWithinRootSignature, InData, InDataSize);
+		WriteLocalShaderRecord(ShaderTableOffset, RecordIndex, ShaderIdentifierSize + InOffsetWithinRootSignature, InData, InDataSize);
 	}
 
-	void CopyLocalShaderParameters(uint32 InDestRecordIndex, uint32 InSourceRecordIndex, uint32 InOffsetWithinRootSignature)
+	template <typename T>
+	void SetMissShaderParameters(uint32 RecordIndex, uint32 InOffsetWithinRootSignature, const T& Parameters)
+	{
+		const uint32 ShaderTableOffset = MissShaderTableOffset;
+		WriteLocalShaderRecord(ShaderTableOffset, RecordIndex, ShaderIdentifierSize + InOffsetWithinRootSignature, &Parameters, sizeof(Parameters));
+	}
+
+	template <typename T>
+	void SetCallableShaderParameters(uint32 RecordIndex, uint32 InOffsetWithinRootSignature, const T& Parameters)
+	{
+		const uint32 ShaderTableOffset = CallableShaderTableOffset;
+		WriteLocalShaderRecord(ShaderTableOffset, RecordIndex, ShaderIdentifierSize + InOffsetWithinRootSignature, &Parameters, sizeof(Parameters));
+	}
+	
+	void CopyLocalShaderParameters(uint32 InShaderTableOffset, uint32 InDestRecordIndex, uint32 InSourceRecordIndex, uint32 InOffsetWithinRootSignature)
 	{
-		const uint32 BaseOffset = LocalShaderTableOffset + ShaderIdentifierSize + InOffsetWithinRootSignature;
+		const uint32 BaseOffset = InShaderTableOffset + ShaderIdentifierSize + InOffsetWithinRootSignature;
 		const uint32 DestOffset   = BaseOffset + LocalRecordStride * InDestRecordIndex;
 		const uint32 SourceOffset = BaseOffset + LocalRecordStride * InSourceRecordIndex;
 		const uint32 CopySize = LocalRecordStride - ShaderIdentifierSize - InOffsetWithinRootSignature;
@@ -1916,11 +1923,10 @@ public:
 		bIsDirty = true;
 	}
 
-	void SetLocalShaderIdentifier(uint32 RecordIndex, const void* ShaderIdentifierData, uint32 InShaderIdentifierSize)
+	void CopyHitGroupParameters(uint32 InDestRecordIndex, uint32 InSourceRecordIndex, uint32 InOffsetWithinRootSignature)
 	{
-		checkSlow(InShaderIdentifierSize == ShaderIdentifierSize);
-
-		WriteLocalShaderRecord(RecordIndex, 0, ShaderIdentifierData, InShaderIdentifierSize);
+		const uint32 ShaderTableOffset = HitGroupShaderTableOffset;
+		CopyLocalShaderParameters(ShaderTableOffset, InDestRecordIndex, InSourceRecordIndex, InOffsetWithinRootSignature);
 	}
 
 	void SetRayGenIdentifier(uint32 RecordIndex, const FD3D12ShaderIdentifier& ShaderIdentifier)
@@ -1941,6 +1947,12 @@ public:
 		WriteData(WriteOffset, ShaderIdentifier.Data, ShaderIdentifierSize);
 	}
 
+	void SetCallableIdentifier(uint32 RecordIndex, const FD3D12ShaderIdentifier& ShaderIdentifier)
+	{
+		const uint32 WriteOffset = CallableShaderTableOffset + RecordIndex * LocalRecordStride;
+		WriteData(WriteOffset, ShaderIdentifier.Data, ShaderIdentifierSize);
+	}
+
 	void SetDefaultHitGroupIdentifier(const FD3D12ShaderIdentifier& ShaderIdentifier)
 	{
 		const uint32 WriteOffset = DefaultHitGroupShaderTableOffset;
@@ -1950,15 +1962,17 @@ public:
 	void SetHitGroupSystemParameters(uint32 RecordIndex, const FHitGroupSystemParameters& SystemParameters)
 	{
 		const uint32 OffsetWithinRootSignature = 0; // System parameters are always first in the RS.
-		SetLocalShaderParameters(RecordIndex, OffsetWithinRootSignature, SystemParameters);
+		const uint32 ShaderTableOffset = HitGroupShaderTableOffset;
+		SetLocalShaderParameters(ShaderTableOffset, RecordIndex, OffsetWithinRootSignature, SystemParameters);
 	}
 
-	void SetLocalShaderIdentifier(uint32 RecordIndex, const FD3D12ShaderIdentifier& ShaderIdentifier)
+	void SetHitGroupIdentifier(uint32 RecordIndex, const FD3D12ShaderIdentifier& ShaderIdentifier)
 	{
 		checkfSlow(ShaderIdentifier.IsValid(), TEXT("Shader identifier must be initialized FD3D12RayTracingPipelineState::GetShaderIdentifier() before use."));
 		checkSlow(sizeof(ShaderIdentifier.Data) >= ShaderIdentifierSize);
 
-		SetLocalShaderIdentifier(RecordIndex, ShaderIdentifier.Data, ShaderIdentifierSize);
+		const uint32 WriteOffset = HitGroupShaderTableOffset + RecordIndex * LocalRecordStride;
+		WriteData(WriteOffset, ShaderIdentifier.Data, ShaderIdentifierSize);
 	}
 
 	void SetRayGenIdentifiers(const TArrayView<const FD3D12ShaderIdentifier>& Identifiers)
@@ -1975,15 +1989,19 @@ public:
 		// Set all slots to the same default
 		for (uint32 Index = 0; Index < NumMissRecords; ++Index)
 		{
-			SetLocalShaderIdentifier(MissShaderRecordIndexOffset + Index, ShaderIdentifier);
+			SetMissIdentifier(Index, ShaderIdentifier);
 		}
+
+#if DO_CHECK
+		bWasDefaultMissShaderSet = false;
+#endif
 	}
 
 	void SetDefaultCallableShaderIdentifier(const FD3D12ShaderIdentifier& ShaderIdentifier)
 	{
 		for (uint32 Index = 0; Index < NumCallableRecords; ++Index)
 		{
-			SetLocalShaderIdentifier(CallableShaderRecordIndexOffset + Index, ShaderIdentifier);
+			SetCallableIdentifier(Index, ShaderIdentifier);
 		}
 	}
 
@@ -2027,7 +2045,7 @@ public:
 		return Buffer->ResourceLocation.GetGPUVirtualAddress();
 	}
 
-	D3D12_DISPATCH_RAYS_DESC GetDispatchRaysDesc(uint32 RayGenShaderIndex, uint32 MissShaderBaseIndex, bool bAllowHitGroupIndexing) const
+	D3D12_DISPATCH_RAYS_DESC GetDispatchRaysDesc(uint32 RayGenShaderIndex, bool bAllowHitGroupIndexing) const
 	{
 		D3D12_GPU_VIRTUAL_ADDRESS ShaderTableAddress = GetShaderTableAddress();
 
@@ -2036,7 +2054,7 @@ public:
 		Desc.RayGenerationShaderRecord.StartAddress = ShaderTableAddress + RayGenShaderTableOffset + RayGenShaderIndex * RayGenRecordStride;
 		Desc.RayGenerationShaderRecord.SizeInBytes = RayGenRecordStride;
 
-		Desc.MissShaderTable.StartAddress = ShaderTableAddress + MissShaderTableOffset + MissShaderBaseIndex * LocalRecordStride;
+		Desc.MissShaderTable.StartAddress = ShaderTableAddress + MissShaderTableOffset;
 		Desc.MissShaderTable.StrideInBytes = LocalRecordStride;
 		Desc.MissShaderTable.SizeInBytes = LocalRecordStride * NumMissRecords;
 
@@ -2076,9 +2094,6 @@ public:
 	uint32 DefaultHitGroupShaderTableOffset = 0;
 	uint32 HitGroupShaderTableOffset = 0;
 	uint32 CallableShaderTableOffset = 0;
-	uint32 LocalShaderTableOffset = 0;
-	uint32 CallableShaderRecordIndexOffset = 0;
-	uint32 MissShaderRecordIndexOffset = 0;
 
 	uint64 LastCommandListID = 0;
 
@@ -4577,7 +4592,7 @@ struct FD3D12RayTracingGlobalResourceBinder
 
 struct FD3D12RayTracingLocalResourceBinder
 {
-	FD3D12RayTracingLocalResourceBinder(FD3D12Device& InDevice, FD3D12RayTracingShaderTable& InShaderTable, const FD3D12RootSignature& InRootSignature, uint32 InRecordIndex, uint32 InWorkerIndex)
+	FD3D12RayTracingLocalResourceBinder(FD3D12Device& InDevice, FD3D12RayTracingShaderTable& InShaderTable, const FD3D12RootSignature& InRootSignature, uint32 InRecordIndex, uint32 InWorkerIndex, ERayTracingBindingType InBindingType)
 		: Device(InDevice)
 		, ShaderTable(InShaderTable)
 		, DescriptorCache(*InShaderTable.DescriptorCache)
@@ -4589,6 +4604,21 @@ struct FD3D12RayTracingLocalResourceBinder
 		check(WorkerIndex < InShaderTable.MaxBindingWorkers);
 		check(WorkerIndex < DescriptorCache.MaxBindingWorkers);
 		check(RecordIndex != ~0u);
+
+		switch (InBindingType)
+		{
+		case ERayTracingBindingType::CallableShader:
+			ShaderTableOffset = InShaderTable.CallableShaderTableOffset;
+			break;
+		case ERayTracingBindingType::HitGroup:
+			ShaderTableOffset = InShaderTable.HitGroupShaderTableOffset;
+			break;
+		case ERayTracingBindingType::MissShader:
+			ShaderTableOffset = InShaderTable.MissShaderTableOffset;
+			break;
+		default:
+			checkNoEntry();
+		}
 	}
 
 	void SetRootDescriptor(uint32 BaseSlotIndex, uint32 DescriptorIndex, D3D12_GPU_VIRTUAL_ADDRESS Address)
@@ -4596,7 +4626,7 @@ struct FD3D12RayTracingLocalResourceBinder
 		const uint32 BindOffsetBase = RootSignature.GetBindSlotOffsetInBytes(BaseSlotIndex);
 		const uint32 DescriptorSize = uint32(sizeof(D3D12_GPU_VIRTUAL_ADDRESS));
 		const uint32 CurrentOffset = BindOffsetBase + DescriptorIndex * DescriptorSize;
-		ShaderTable.SetLocalShaderParameters(RecordIndex, CurrentOffset, Address);
+		ShaderTable.SetLocalShaderParameters(ShaderTableOffset, RecordIndex, CurrentOffset, Address);
 	}
 
 	void SetRootCBV(uint32 BaseSlotIndex, uint32 DescriptorIndex, D3D12_GPU_VIRTUAL_ADDRESS Address)
@@ -4612,7 +4642,7 @@ struct FD3D12RayTracingLocalResourceBinder
 	void SetRootDescriptorTable(uint32 SlotIndex, D3D12_GPU_DESCRIPTOR_HANDLE DescriptorTable)
 	{
 		const uint32 BindOffset = RootSignature.GetBindSlotOffsetInBytes(SlotIndex);
-		ShaderTable.SetLocalShaderParameters(RecordIndex, BindOffset, DescriptorTable);
+		ShaderTable.SetLocalShaderParameters(ShaderTableOffset, RecordIndex, BindOffset, DescriptorTable);
 	}
 
 #if USE_STATIC_ROOT_SIGNATURE
@@ -4673,6 +4703,7 @@ struct FD3D12RayTracingLocalResourceBinder
 	FD3D12RayTracingShaderTable& ShaderTable;
 	FD3D12RayTracingDescriptorCache& DescriptorCache;
 	const FD3D12RootSignature& RootSignature;
+	uint32 ShaderTableOffset = 0;
 	uint32 RecordIndex = ~0u;
 	uint32 WorkerIndex = 0;
 };
@@ -5190,7 +5221,7 @@ void FD3D12CommandContext::RHIRayTraceDispatch(FRHIRayTracingPipelineState* InRa
 		TEXT("All RayGen shaders must be declared when creating RTPSO."),
 			*(RayGenShader->EntryPoint));
 
-	D3D12_DISPATCH_RAYS_DESC DispatchDesc = ShaderTable->GetDispatchRaysDesc(RayGenShaderIndex, 0, Pipeline->bAllowHitGroupIndexing);
+	D3D12_DISPATCH_RAYS_DESC DispatchDesc = ShaderTable->GetDispatchRaysDesc(RayGenShaderIndex, Pipeline->bAllowHitGroupIndexing);
 
 	DispatchDesc.Width = Width;
 	DispatchDesc.Height = Height;
@@ -5223,7 +5254,7 @@ void FD3D12CommandContext::RHIRayTraceDispatchIndirect(FRHIRayTracingPipelineSta
 	const int32 RayGenShaderIndex = Pipeline->RayGenShaders.Find(RayGenShader->GetHash());
 	checkf(RayGenShaderIndex != INDEX_NONE, TEXT("RayGen shader is not present in the given ray tracing pipeline. All RayGen shaders must be declared when creating RTPSO."));
 
-	D3D12_DISPATCH_RAYS_DESC DispatchDesc = ShaderTable->GetDispatchRaysDesc(RayGenShaderIndex, 0, Pipeline->bAllowHitGroupIndexing);
+	D3D12_DISPATCH_RAYS_DESC DispatchDesc = ShaderTable->GetDispatchRaysDesc(RayGenShaderIndex, Pipeline->bAllowHitGroupIndexing);
 	DispatchRays(*this, GlobalResourceBindings, Pipeline, RayGenShaderIndex, ShaderTable, DispatchDesc, RetrieveObject<FD3D12Buffer>(ArgumentBuffer), ArgumentOffset);
 }
 
@@ -5293,13 +5324,13 @@ static void SetRayTracingHitGroup(
 		{
 			// Simply copy local shader parameters from existing SBT record and set the shader identifier, skipping resource binding work.
 			const uint32 OffsetFromRootSignatureStart = sizeof(FHitGroupSystemParameters);
-			ShaderTable->SetLocalShaderIdentifier(RecordIndex, Pipeline->HitGroupShaders.Identifiers[HitGroupIndex]);
-			ShaderTable->CopyLocalShaderParameters(RecordIndex, *ExistingRecordIndex, OffsetFromRootSignatureStart);
+			ShaderTable->SetHitGroupIdentifier(RecordIndex, Pipeline->HitGroupShaders.Identifiers[HitGroupIndex]);
+			ShaderTable->CopyHitGroupParameters(RecordIndex, *ExistingRecordIndex, OffsetFromRootSignatureStart);
 			return;
 		}
 	}
 
-	FD3D12RayTracingLocalResourceBinder ResourceBinder(*Device, *ShaderTable, *(Shader->pRootSignature), RecordIndex, WorkerIndex);
+	FD3D12RayTracingLocalResourceBinder ResourceBinder(*Device, *ShaderTable, *(Shader->pRootSignature), RecordIndex, WorkerIndex, ERayTracingBindingType::HitGroup);
 	const bool bResourcesBound = SetRayTracingShaderResources(CommandContext, Shader,
 		0, nullptr, // Textures
 		0, nullptr, // SRVs
@@ -5314,7 +5345,7 @@ static void SetRayTracingHitGroup(
 		ShaderTable->ShaderRecordCache[WorkerIndex].Add(CacheKey, RecordIndex);
 	}
 
-	ShaderTable->SetLocalShaderIdentifier(RecordIndex, 
+	ShaderTable->SetHitGroupIdentifier(RecordIndex,
 		bResourcesBound
 		? Pipeline->HitGroupShaders.Identifiers[HitGroupIndex]
 		: FD3D12ShaderIdentifier::Null);
@@ -5334,10 +5365,10 @@ static void SetRayTracingCallableShader(
 {
 	checkf(ShaderSlotInScene < Scene->Initializer.NumCallableShaderSlots, TEXT("Shader slot is invalid. Make sure that NumCallableShaderSlots is correct on FRayTracingSceneInitializer."));
 
-	const uint32 RecordIndex = ShaderTable->CallableShaderRecordIndexOffset + ShaderSlotInScene;
+	const uint32 RecordIndex = ShaderSlotInScene;
 
 	const uint32 UserDataOffset = offsetof(FHitGroupSystemParameters, RootConstants) + offsetof(FHitGroupSystemRootConstants, UserData);
-	ShaderTable->SetLocalShaderParameters(RecordIndex, UserDataOffset, UserData);
+	ShaderTable->SetCallableShaderParameters(RecordIndex, UserDataOffset, UserData);
 
 	const FD3D12ShaderIdentifier* ShaderIdentifier = &FD3D12ShaderIdentifier::Null;
 
@@ -5345,7 +5376,7 @@ static void SetRayTracingCallableShader(
 	{
 		const FD3D12RayTracingShader* Shader = Pipeline->CallableShaders.Shaders[ShaderIndexInPipeline];
 
-		FD3D12RayTracingLocalResourceBinder ResourceBinder(*Device, *ShaderTable, *(Shader->pRootSignature), RecordIndex, WorkerIndex);
+		FD3D12RayTracingLocalResourceBinder ResourceBinder(*Device, *ShaderTable, *(Shader->pRootSignature), RecordIndex, WorkerIndex, ERayTracingBindingType::CallableShader);
 		const bool bResourcesBound = SetRayTracingShaderResources(CommandContext, Shader,
 			0, nullptr, // Textures
 			0, nullptr, // SRVs
@@ -5361,7 +5392,7 @@ static void SetRayTracingCallableShader(
 		}
 	}
 
-	ShaderTable->SetLocalShaderIdentifier(RecordIndex, *ShaderIdentifier);
+	ShaderTable->SetCallableIdentifier(RecordIndex, *ShaderIdentifier);
 }
 
 static void SetRayTracingMissShader(
@@ -5378,14 +5409,14 @@ static void SetRayTracingMissShader(
 {
 	checkf(ShaderSlotInScene < Scene->Initializer.NumMissShaderSlots, TEXT("Shader slot is invalid. Make sure that NumMissShaderSlots is correct on FRayTracingSceneInitializer."));
 
-	const uint32 RecordIndex = ShaderTable->MissShaderRecordIndexOffset + ShaderSlotInScene;
+	const uint32 RecordIndex = ShaderSlotInScene;
 
 	const uint32 UserDataOffset = offsetof(FHitGroupSystemParameters, RootConstants) + offsetof(FHitGroupSystemRootConstants, UserData);
-	ShaderTable->SetLocalShaderParameters(RecordIndex, UserDataOffset, UserData);
+	ShaderTable->SetMissShaderParameters(RecordIndex, UserDataOffset, UserData);
 
 	const FD3D12RayTracingShader* Shader = Pipeline->MissShaders.Shaders[ShaderIndexInPipeline];
 
-	FD3D12RayTracingLocalResourceBinder ResourceBinder(*Device, *ShaderTable, *(Shader->pRootSignature), RecordIndex, WorkerIndex);
+	FD3D12RayTracingLocalResourceBinder ResourceBinder(*Device, *ShaderTable, *(Shader->pRootSignature), RecordIndex, WorkerIndex, ERayTracingBindingType::MissShader);
 	const bool bResourcesBound = SetRayTracingShaderResources(CommandContext, Shader,
 		0, nullptr, // Textures
 		0, nullptr, // SRVs
@@ -5415,8 +5446,6 @@ void FD3D12CommandContext::RHISetRayTracingBindings(
 	checkf(Scene->bBuilt, TEXT("Ray tracing scene must be built before any shaders can be bound to it. Make sure that RHIBuildAccelerationStructure() command has been executed."));
 
 	FD3D12RayTracingShaderTable* ShaderTable = Scene->FindOrCreateShaderTable(Pipeline, Device);
-	checkf(ShaderTable->LocalShaderTableOffset == ShaderTable->HitGroupShaderTableOffset,
-		TEXT("Hit shader records are assumed to be at the beginning of local shader table"));
 
 	FGraphEventArray TaskList;
 
@@ -5505,8 +5534,6 @@ void FD3D12CommandContext::RHISetRayTracingHitGroup(
 	checkf(ShaderSlot < Scene->Initializer.ShaderSlotsPerGeometrySegment, TEXT("Shader slot is invalid. Make sure that ShaderSlotsPerGeometrySegment is correct on FRayTracingSceneInitializer."));
 
 	FD3D12RayTracingShaderTable* ShaderTable = Scene->FindOrCreateShaderTable(Pipeline, GetParentDevice());
-	checkf(ShaderTable->LocalShaderTableOffset == ShaderTable->HitGroupShaderTableOffset,
-		TEXT("Hit shader records are assumed to be at the beginning of local shader table"));
 
 	const uint32 WorkerIndex = 0;
 	SetRayTracingHitGroup(*this, GetParentDevice(), ShaderTable, Scene, Pipeline,
