Improved DX 12 GPU breadcrumbs that track all active, finished and not started scopes on the GPU. The current implementation only logs the last finished or last started scope, which in many cases may not be relevant to the crash. The new implementation logs all scopes currently scheduled or active on the GPU.

diff --git a/Engine/Source/Runtime/D3D12RHI/Private/D3D12CommandContext.cpp b/Engine/Source/Runtime/D3D12RHI/Private/D3D12CommandContext.cpp
--- a/Engine/Source/Runtime/D3D12RHI/Private/D3D12CommandContext.cpp
+++ b/Engine/Source/Runtime/D3D12RHI/Private/D3D12CommandContext.cpp
@@ -121,78 +121,6 @@ FD3D12CommandContext::~FD3D12CommandContext()
 	ClearState();
 }
 
-
-/** Write out the event stack to the bread crumb resource if available */
-void FD3D12CommandContext::WriteGPUEventStackToBreadCrumbData(bool bBeginEvent)
-{
-	// Write directly to command list if breadcrumb resource is available
-	TUniquePtr<FD3D12DiagnosticBuffer>& DiagnosticBuffer = Device->GetQueue(QueueType).DiagnosticBuffer;
-	if (!DiagnosticBuffer)
-		return;
-
-	if (!GraphicsCommandList2())
-		return;
-
-	// Find the max parameter count from the resource
-	const int32 MaxParameterCount = DiagnosticBuffer->BreadCrumbsSize / sizeof(uint32);
-
-	// allocate the parameters on the stack if smaller than 4K
-	int32 ParameterCount = GPUEventStack.Num() < (MaxParameterCount - 2) ? GPUEventStack.Num() + 2 : MaxParameterCount;
-	size_t MemSize = ParameterCount * (sizeof(D3D12_WRITEBUFFERIMMEDIATE_PARAMETER) + sizeof(D3D12_WRITEBUFFERIMMEDIATE_PARAMETER));
-	const bool bAllocateOnStack = (MemSize < 4096);
-	void* Mem = bAllocateOnStack ? FMemory_Alloca(MemSize) : FMemory::Malloc(MemSize);
-
-	if (Mem)
-	{
-		D3D12_WRITEBUFFERIMMEDIATE_PARAMETER* Parameters = (D3D12_WRITEBUFFERIMMEDIATE_PARAMETER*)Mem;
-		D3D12_WRITEBUFFERIMMEDIATE_MODE* Modes = (D3D12_WRITEBUFFERIMMEDIATE_MODE*)(Parameters + ParameterCount);
-		for (int i = 0; i < ParameterCount; ++i)
-		{
-			Parameters[i].Dest = DiagnosticBuffer->Resource->GetGPUVirtualAddress() + 4 * i;
-
-		#if 1 // This is more accurate, but may have a higher theoretical GPU overhead
-			if (bBeginEvent)
-			{
-				// The write operation is guaranteed to occur after all preceding commands in the command stream have started, including previous.
-				// We use this mode because we want to know which ops have started on the GPU.
-				Modes[i] = D3D12_WRITEBUFFERIMMEDIATE_MODE_MARKER_IN;
-			}
-			else
-			{
-				// The write operation is deferred until all previous commands in the command stream have completed through the GPU pipeline, including previous WriteBufferImmediate operations.
-				// We want this mode when ending breadcrumb scopes because the GPU might start the next scope before we hit a problem in the current one in some cases (when there are no barriers).
-				Modes[i] = D3D12_WRITEBUFFERIMMEDIATE_MODE_MARKER_OUT;
-			}
-		#else // This is less accurate (we don't know for sure if the event has finished), but has lower theoretical GPU overhead
-			Modes[i] = D3D12_WRITEBUFFERIMMEDIATE_MODE_MARKER_IN;
-		#endif
-
-			// Write event stack count first
-			if (i == 0)
-			{
-				Parameters[i].Value = GPUEventStack.Num();
-			}
-			// Then if it's the begin or end event
-			else if (i == 1)
-			{
-				Parameters[i].Value = bBeginEvent ? 1 : 0;
-			}
-			// Otherwise the actual stack value
-			else
-			{
-				Parameters[i].Value = GPUEventStack[i - 2];
-			}
-		}
-		GraphicsCommandList2()->WriteBufferImmediate(ParameterCount, Parameters, Modes);
-	}
-
-	if (!bAllocateOnStack)
-	{
-		FMemory::Free(Mem);
-	}
-}
-
-
 void FD3D12CommandContext::RHIPushEvent(const TCHAR* Name, FColor Color)
 {
 	D3D12RHI::FD3DGPUProfiler& GPUProfiler = GetParentDevice()->GetGPUProfiler();
@@ -211,7 +139,7 @@ void FD3D12CommandContext::RHIPushEvent(const TCHAR* Name, FColor Color)
 		uint32 CRC = GPUProfiler.GetOrAddEventStringHash(Name);
 
 		GPUEventStack.Push(CRC);
-		WriteGPUEventStackToBreadCrumbData(true);
+		WriteGPUEventStackToBreadCrumbData(Name, CRC);
 
 #if NV_AFTERMATH
 		// Only track aftermath for default context?
@@ -247,7 +175,7 @@ void FD3D12CommandContext::RHIPopEvent()
 
 	if (GPUProfiler.bTrackingGPUCrashData)
 	{
-		WriteGPUEventStackToBreadCrumbData(false);
+		PopGPUEventStackFromBreadCrumbData();
 
 		// need to look for unbalanced push/pop
 		if (GPUEventStack.Num() > 0)
@@ -538,6 +466,131 @@ FD3D12SyncPoint* FD3D12CopyScope::GetSyncPoint() const
 	return SyncPoint;
 }
 
+void FD3D12ContextCommon::InitPayloadBreadcrumbs()
+{
+	FD3D12Payload* Payload = GetPayload(EPhase::Execute);
+
+	if (Payload->BreadcrumbStacks.IsEmpty() || !BreadcrumbStack.IsValid())
+	{
+		TUniquePtr<FD3D12DiagnosticBuffer>& DiagnosticBuffer = Device->GetQueue(QueueType).DiagnosticBuffer;
+		if (!BreadcrumbStack.IsValid())
+		{
+			BreadcrumbStack = MakeShared<FBreadcrumbStack>();
+			BreadcrumbStack->Queue = &Device->GetQueue(QueueType);
+			BreadcrumbStack->Initialize(DiagnosticBuffer);
+		}
+
+		Payload->BreadcrumbStacks.Add(BreadcrumbStack);
+	}
+}
+
+void FD3D12ContextCommon::WriteGPUEventStackToBreadCrumbData(const TCHAR* Name, int32 CRC)
+{
+	InitPayloadBreadcrumbs();
+
+	FBreadcrumbStack::FScope NewScope;
+	NewScope.NameCRC = CRC;
+	NewScope.MarkerIndex = (BreadcrumbStack->NextIdx++);
+	NewScope.Sibling = 0;
+	NewScope.Child = 0;
+
+	const uint32 ThisScopeIndex = BreadcrumbStack->Scopes.Num();
+
+	if (!BreadcrumbStack->ScopeStack.IsEmpty())
+	{
+		auto& TopScope = BreadcrumbStack->Scopes[BreadcrumbStack->ScopeStack.Last()];
+		if (BreadcrumbStack->bTopIsOpen)
+		{
+			TopScope.Child = ThisScopeIndex;
+		}
+		else
+		{
+			TopScope.Sibling = ThisScopeIndex;
+			BreadcrumbStack->ScopeStack.Pop();
+		}
+	}
+	BreadcrumbStack->Scopes.Add(NewScope);
+	BreadcrumbStack->ScopeStack.Add(ThisScopeIndex);
+
+	BreadcrumbStack->bTopIsOpen = true;
+
+	if (NewScope.MarkerIndex < BreadcrumbStack->MaxMarkers)
+	{
+		WriteGPUEventToBreadCrumbData(BreadcrumbStack.Get(), NewScope.MarkerIndex, true);
+	}
+}
+
+void FD3D12ContextCommon::PopGPUEventStackFromBreadCrumbData()
+{
+	InitPayloadBreadcrumbs();
+
+	if (BreadcrumbStack->ScopeStack.IsEmpty())
+	{
+		UE_LOG(LogD3D12RHI, Log, TEXT("Cannot end block when stack is empty"));
+	}
+	else
+	{
+		// If top of scope stack isn't open, then our last child is there, and we need to pop that off.
+		{
+			if (!BreadcrumbStack->bTopIsOpen)
+			{
+				if (BreadcrumbStack->ScopeStack.Num() <= 1)
+				{
+					UE_LOG(LogD3D12RHI, Log, TEXT("Cannot end block when stack is empty"));
+				}
+				else
+				{
+					BreadcrumbStack->ScopeStack.Pop();
+				}
+			}
+		}
+
+		{
+			const FBreadcrumbStack::FScope& ThisScope = BreadcrumbStack->Scopes[BreadcrumbStack->ScopeStack.Last()];
+			checkf(ThisScope.Sibling == 0, TEXT("Shouldn't have a sibling already"));
+
+			if (ThisScope.MarkerIndex < BreadcrumbStack->MaxMarkers)
+			{
+				WriteGPUEventToBreadCrumbData(BreadcrumbStack.Get(), ThisScope.MarkerIndex, false);
+			}
+		}
+
+		if (BreadcrumbStack->ScopeStack.Num() == 1 && BreadcrumbStack->Scopes.Num() > 100)
+		{
+			BreadcrumbStack->ScopeStack.Reset();
+			BreadcrumbStack.Reset();
+		}
+		else
+		{
+			// Don't remove ourselves from the stack, we stay there for any siblings.
+			BreadcrumbStack->bTopIsOpen = false;
+		}
+	}
+}
+
+void FD3D12ContextCommon::WriteGPUEventToBreadCrumbData(FBreadcrumbStack* Breadcrumbs, uint32 MarkerIndex, bool bBeginEvent)
+{
+	if (!GraphicsCommandList2())
+		return;
+
+	// Find the max parameter count from the resource
+	const int32 MaxParameterCount = Breadcrumbs->MaxMarkers;
+
+	if (static_cast<int32>(MarkerIndex) > MaxParameterCount)
+	{
+		UE_LOG(LogD3D12RHI, Log, TEXT("Breadcrumbs parameter overflow: %u"), MarkerIndex);
+		return;
+	}
+
+	D3D12_WRITEBUFFERIMMEDIATE_PARAMETER Parameter;
+	Parameter.Dest = Breadcrumbs->WriteAddress + MarkerIndex * sizeof(uint32);
+	Parameter.Value = bBeginEvent ? 1 : 2;
+	D3D12_WRITEBUFFERIMMEDIATE_MODE Mode;
+	Mode = bBeginEvent ? D3D12_WRITEBUFFERIMMEDIATE_MODE_MARKER_IN : D3D12_WRITEBUFFERIMMEDIATE_MODE_MARKER_OUT;
+
+	GraphicsCommandList2()->WriteBufferImmediate(1, &Parameter, &Mode);
+}
+
 void FD3D12ContextCommon::FlushCommands(ED3D12FlushFlags FlushFlags)
 {
 	// We should only be flushing the default context
diff --git a/Engine/Source/Runtime/D3D12RHI/Private/D3D12CommandContext.h b/Engine/Source/Runtime/D3D12RHI/Private/D3D12CommandContext.h
--- a/Engine/Source/Runtime/D3D12RHI/Private/D3D12CommandContext.h
+++ b/Engine/Source/Runtime/D3D12RHI/Private/D3D12CommandContext.h
@@ -228,6 +228,9 @@ private:
 	// A sync point signaled when all payloads in this context have completed.
 	FD3D12SyncPointRef ContextSyncPoint;
 
+	// Stack containing GPU breadcrumbs for crash debugging
+	TSharedPtr<FBreadcrumbStack> BreadcrumbStack;
+
 	// Returns the current command list (or creates a new one if the command list was not open).
 	FD3D12CommandList& GetCommandList()
 	{
@@ -246,6 +249,11 @@ public:
 	}
 
 protected:
+	void WriteGPUEventStackToBreadCrumbData(const TCHAR* Name, int32 CRC);
+	void WriteGPUEventToBreadCrumbData(FBreadcrumbStack* Breadcrumbs, uint32 MarkerIndex, bool bBeginEvent);
+	void InitPayloadBreadcrumbs();
+	void PopGPUEventStackFromBreadCrumbData();
+
 	enum class EPhase
 	{
 		Wait,
@@ -699,8 +707,6 @@ protected:
 	{  
 		return InGPUIndex == GetGPUIndex() ? this : nullptr; 
 	}
-	
-	void WriteGPUEventStackToBreadCrumbData(bool bBeginEvent);
 
 private:
 
diff --git a/Engine/Source/Runtime/D3D12RHI/Private/D3D12Device.cpp b/Engine/Source/Runtime/D3D12RHI/Private/D3D12Device.cpp
--- a/Engine/Source/Runtime/D3D12RHI/Private/D3D12Device.cpp
+++ b/Engine/Source/Runtime/D3D12RHI/Private/D3D12Device.cpp
@@ -67,15 +67,13 @@ void FD3D12Queue::SetupAfterDeviceCreation()
 		HRESULT hr = Device->GetDevice()->QueryInterface(IID_PPV_ARGS(D3D12Device3.GetInitReference()));
 		if (SUCCEEDED(hr))
 		{
-			// find out how many entries we can much push in a single event (limit to MAX_GPU_BREADCRUMB_DEPTH)
-			int32 GPUCrashDataDepth = Device->GetGPUProfiler().GPUCrashDataDepth;
-			int32 MaxEventCount = GPUCrashDataDepth > 0 ? FMath::Min(GPUCrashDataDepth, MAX_GPU_BREADCRUMB_DEPTH) : MAX_GPU_BREADCRUMB_DEPTH;
-
 			const uint32 ShaderDiagnosticBufferSize = sizeof(FD3D12DiagnosticBufferData);
 
 			// Allocate persistent CPU readable memory which will still be valid after a device lost and wrap this data in a placed resource
 			// so the GPU command list can write to it
-			const uint32 EventBufferSize = MaxEventCount * sizeof(uint32);
+			int32 MaxBreadcrumbsContexts = MAX_GPU_BREADCRUMB_CONTEXTS;
+			int32 MaxBreadcrumbsSize = MAX_GPU_BREADCRUMB_SIZE;
+			const uint32 EventBufferSize = MaxBreadcrumbsSize * MaxBreadcrumbsContexts * sizeof(uint32);
 			const uint32 TotalBufferSize = EventBufferSize + ShaderDiagnosticBufferSize;
 
 			// Create the platform-specific diagnostic buffer
@@ -90,6 +88,13 @@ void FD3D12Queue::SetupAfterDeviceCreation()
 				DiagnosticBuffer->BreadCrumbsOffset = 0;
 				DiagnosticBuffer->BreadCrumbsSize = EventBufferSize;
 
+				DiagnosticBuffer->BreadCrumbsContextSize = EventBufferSize / MaxBreadcrumbsContexts;
+
+				for (uint16 Idx = 0; Idx < MaxBreadcrumbsContexts; ++Idx)
+				{
+					DiagnosticBuffer->FreeContextIds.Add(MaxBreadcrumbsContexts - Idx);
+				}
+
 				DiagnosticBuffer->DiagnosticsOffset = DiagnosticBuffer->BreadCrumbsOffset + DiagnosticBuffer->BreadCrumbsSize;
 				DiagnosticBuffer->DiagnosticsSize = ShaderDiagnosticBufferSize;
 			}
diff --git a/Engine/Source/Runtime/D3D12RHI/Private/D3D12Device.h b/Engine/Source/Runtime/D3D12RHI/Private/D3D12Device.h
--- a/Engine/Source/Runtime/D3D12RHI/Private/D3D12Device.h
+++ b/Engine/Source/Runtime/D3D12RHI/Private/D3D12Device.h
@@ -49,6 +49,10 @@ struct FD3D12DiagnosticBuffer
 		, GpuAddress(GpuAddress)
 	{}
 
+	TArray<uint16> FreeContextIds;
+	FCriticalSection CriticalSection;
+	uint32 BreadCrumbsContextSize = 0;
+
 	uint32 BreadCrumbsOffset = 0;
 	uint32 BreadCrumbsSize = 0;
 
diff --git a/Engine/Source/Runtime/D3D12RHI/Private/D3D12RHIPrivate.h b/Engine/Source/Runtime/D3D12RHI/Private/D3D12RHIPrivate.h
--- a/Engine/Source/Runtime/D3D12RHI/Private/D3D12RHIPrivate.h
+++ b/Engine/Source/Runtime/D3D12RHI/Private/D3D12RHIPrivate.h
@@ -155,6 +155,8 @@ inline const TCHAR* GetD3DCommandQueueTypeName(ED3D12QueueType QueueType)
 #define TEXTURE_POOL_SIZE (8 * 1024 * 1024)
 
 #define MAX_GPU_BREADCRUMB_DEPTH 1024
+#define MAX_GPU_BREADCRUMB_CONTEXTS 128
+#define MAX_GPU_BREADCRUMB_SIZE 4096
 
 #ifndef FD3D12_HEAP_FLAG_CREATE_NOT_ZEROED
 #define FD3D12_HEAP_FLAG_CREATE_NOT_ZEROED D3D12_HEAP_FLAG_CREATE_NOT_ZEROED
diff --git a/Engine/Source/Runtime/D3D12RHI/Private/D3D12Submission.cpp b/Engine/Source/Runtime/D3D12RHI/Private/D3D12Submission.cpp
--- a/Engine/Source/Runtime/D3D12RHI/Private/D3D12Submission.cpp
+++ b/Engine/Source/Runtime/D3D12RHI/Private/D3D12Submission.cpp
@@ -255,6 +255,8 @@ void FD3D12DynamicRHI::SubmitCommands(TConstArrayView<FD3D12FinalizedCommands*>
 					MergedPayload->SyncPointsToSignal.Append(Payload->SyncPointsToSignal);
 					MergedPayload->AllocatorsToRelease.Append(Payload->AllocatorsToRelease);
 					MergedPayload->QueryRanges.Append(Payload->QueryRanges);
+					MergedPayload->BreadcrumbStacks.Append(Payload->BreadcrumbStacks);
+					Payload->BreadcrumbStacks.Empty();
 
 					// Need to clear out allocator array, so Payload destructor doesn't delete them
 					Payload->AllocatorsToRelease.Empty();
@@ -1105,6 +1107,52 @@ FD3D12PayloadBase::~FD3D12PayloadBase()
 	}
 }
 
+void FBreadcrumbStack::Initialize(TUniquePtr<FD3D12DiagnosticBuffer>& DiagnosticBuffer)
+{
+	{
+		FScopeLock Lock(&DiagnosticBuffer->CriticalSection);
+
+		if (DiagnosticBuffer->FreeContextIds.Num() == 0)
+		{
+			ContextId = 0;
+			return;
+		}
+
+		ContextId = DiagnosticBuffer->FreeContextIds.Pop();
+	}
+
+	check(ContextId > 0);
+	{
+		const uint32 ContextSize = DiagnosticBuffer->BreadCrumbsContextSize;
+		const uint32 Offset = DiagnosticBuffer->BreadCrumbsContextSize * (ContextId - 1);
+
+		MaxMarkers = ContextSize / sizeof(uint32);
+		WriteAddress = DiagnosticBuffer->GpuAddress + Offset;
+		CPUAddress = (uint8*)(DiagnosticBuffer->CpuAddress) + Offset;
+
+		FMemory::Memzero(CPUAddress, DiagnosticBuffer->BreadCrumbsContextSize);
+	}
+}
+
+FBreadcrumbStack::FBreadcrumbStack()
+{
+	Scopes.Reserve(2048);
+	ScopeStack.Reserve(128);
+}
+
+FBreadcrumbStack::~FBreadcrumbStack()
+{
+	TUniquePtr<FD3D12DiagnosticBuffer>& DiagnosticBuffer = Queue->DiagnosticBuffer;
+	if (ContextId > 0)
+	{
+		{
+			FScopeLock Lock(&DiagnosticBuffer->CriticalSection);
+			DiagnosticBuffer->FreeContextIds.Push(ContextId);
+		}
+		ContextId = 0;
+	}
+}
+
 #ifndef D3D12_PREFER_QUERIES_FOR_GPU_TIME
 #define D3D12_PREFER_QUERIES_FOR_GPU_TIME 0
 #endif
diff --git a/Engine/Source/Runtime/D3D12RHI/Private/D3D12Submission.h b/Engine/Source/Runtime/D3D12RHI/Private/D3D12Submission.h
--- a/Engine/Source/Runtime/D3D12RHI/Private/D3D12Submission.h
+++ b/Engine/Source/Runtime/D3D12RHI/Private/D3D12Submission.h
@@ -152,6 +152,33 @@ struct FD3D12QueryLocation
 	operator bool() const { return Heap != nullptr; }
 };
 
+struct FBreadcrumbStack
+{
+	struct FScope
+	{
+		uint32 NameCRC;
+		uint32 MarkerIndex;
+		uint32 Child;
+		uint32 Sibling;
+	};
+
+	FD3D12Queue* Queue;
+	uint32 NextIdx{ 0 };
+	int32 ContextId;
+	uint32 MaxMarkers{ 0 };
+	D3D12_GPU_VIRTUAL_ADDRESS WriteAddress;
+	void* CPUAddress;
+
+	TArray<FScope> Scopes;
+	TArray<uint32> ScopeStack;
+	bool bTopIsOpen{ false };
+
+	FBreadcrumbStack();
+	~FBreadcrumbStack();
+
+	void Initialize(TUniquePtr<struct FD3D12DiagnosticBuffer>& DiagnosticBuffer);
+};
+
 struct FD3D12QueryRange
 {
 	TRefCountPtr<FD3D12QueryHeap> Heap;
@@ -216,6 +243,9 @@ struct FD3D12PayloadBase
 	TArray<FD3D12QueryLocation> OcclusionQueries;
 	TArray<FD3D12QueryRange> QueryRanges;
 
+	// GPU crash breadcrumbs stack
+	TArray<TSharedPtr<FBreadcrumbStack>> BreadcrumbStacks;
+
 	virtual ~FD3D12PayloadBase();
 
 protected:
diff --git a/Engine/Source/Runtime/D3D12RHI/Private/D3D12Util.cpp b/Engine/Source/Runtime/D3D12RHI/Private/D3D12Util.cpp
--- a/Engine/Source/Runtime/D3D12RHI/Private/D3D12Util.cpp
+++ b/Engine/Source/Runtime/D3D12RHI/Private/D3D12Util.cpp
@@ -193,38 +193,85 @@ static FString GetD3D12TextureFlagString(uint32 TextureFlags)
 	return TextureFormatText;
 }
 
-/** Log the GPU progress of the given queue to the Error log if breadcrumb data is available */
-static bool LogBreadcrumbData(D3D12RHI::FD3DGPUProfiler& GPUProfiler, FD3D12Queue& Queue)
+void DumpBreadcrumbScopeRecursive(
+	D3D12RHI::FD3DGPUProfiler& GPUProfiler,
+	const TSharedPtr<FBreadcrumbStack>& Stack,
+	FString& OutGpuProgress,
+	const FBreadcrumbStack::FScope& Scope,
+	uint32 Indent)
 {
-	uint32* BreadCrumbData = (uint32*)Queue.GetBreadCrumbBufferData();
-	if (BreadCrumbData == nullptr)
+	const volatile uint32* Markers = reinterpret_cast<uint32*>(Stack->CPUAddress);
+	const bool bHaveMarkerIndex = (Scope.MarkerIndex < Stack->MaxMarkers);
+	const uint32 Marker = bHaveMarkerIndex ? Markers[Scope.MarkerIndex] : 0;
+
+	const bool bDidOpen = Marker > 0;
+	const bool bDidClose = Marker > 1;
+
+	const TCHAR* Prefix = TEXT("  ");
+	if (bDidOpen && bDidClose)
 	{
-		return false;
+		Prefix = TEXT("| ");
+	}
+	else if (bDidOpen && !bDidClose)
+	{
+		Prefix = TEXT("> ");
 	}
 
-	uint32 EventCount = BreadCrumbData[0];
-	bool bBeginEvent = BreadCrumbData[1] > 0;
-	check(EventCount >= 0 && EventCount < (MAX_GPU_BREADCRUMB_DEPTH - 2));
+	const TCHAR* Suffix = TEXT("");
+	if (!bHaveMarkerIndex)
+	{
+		Suffix = TEXT(" [overflow]");
+	}
+
+	const FString* EventName = GPUProfiler.FindEventString(Scope.NameCRC);
 
-	FString GpuProgress = FString::Printf(TEXT("[GPUBreadCrumb]\t%s Queue %d - %s"), GetD3DCommandQueueTypeName(Queue.QueueType),
-		Queue.Device->GetGPUIndex(), EventCount == 0 ? TEXT("No Data") : (bBeginEvent ? TEXT("Begin: ") : TEXT("End: ")));
-	for (uint32 EventIndex = 0; EventIndex < EventCount; ++EventIndex)
+	for (uint32 Idx = 0; Idx < Indent * 2; ++Idx)
 	{
-		if (EventIndex > 0)
+		OutGpuProgress.AppendChar(' ');
+	}
+
+	const TCHAR* MarkerState[] = { TEXT("Not started"), TEXT("Active"), TEXT("Finished"), TEXT("Invalid")};
+	OutGpuProgress.Append(FString::Printf(TEXT("Breadcrumbs: %s%s [%s]%s\n"), Prefix, **EventName, MarkerState[FMath::Min(3U, Marker)], Suffix));
+
+	if (bDidOpen && !bDidClose)
+	{
+		for (uint32 Child = Scope.Child; Child != 0; Child = Stack->Scopes[Child].Sibling)
 		{
-			GpuProgress.Append(TEXT(" - "));
+			DumpBreadcrumbScopeRecursive(GPUProfiler, Stack, OutGpuProgress, Stack->Scopes[Child], Indent + 1);
 		}
+	}
+};
 
-		// get the crc and try and translate back into a string
-		uint32 EventCrc = BreadCrumbData[EventIndex + 2];
-		const FString* EventName = GPUProfiler.FindEventString(EventCrc);
-		if (EventName)
+/** Log the GPU progress of the given queue to the Error log if breadcrumb data is available */
+static bool LogBreadcrumbData(D3D12RHI::FD3DGPUProfiler& GPUProfiler, FD3D12Queue& Queue)
+{
+	FString GpuProgress = FString::Printf(TEXT("[GPUBreadCrumb]\t%s Queue %d\n"), GetD3DCommandQueueTypeName(Queue.QueueType),
+		Queue.Device->GetGPUIndex());
+
+	TArray<TSharedPtr<FBreadcrumbStack>, TInlineAllocator<8>> UniqueStacks;
+
+	while (const FD3D12Payload* Payload = Queue.PendingInterrupt.Peek())
+	{
+		for (const TSharedPtr<FBreadcrumbStack>& Stack : Payload->BreadcrumbStacks)
 		{
-			GpuProgress.Append(*EventName);
+			UniqueStacks.AddUnique(Stack);
 		}
-		else
+		Queue.PendingInterrupt.Pop();
+	}
+
+	for (const TSharedPtr<FBreadcrumbStack>& Stack : UniqueStacks)
+	{
+		if (!Stack->Scopes.IsEmpty())
 		{
-			GpuProgress.Append(TEXT("Unknown Event"));
+			if (Stack->ContextId > 0)
+			{
+				uint32 Scope = 0;
+				do
+				{
+					DumpBreadcrumbScopeRecursive(GPUProfiler, Stack, GpuProgress, Stack->Scopes[Scope], 0);
+					Scope = Stack->Scopes[Scope].Sibling;
+				} while (Scope != 0);
+			}
 		}
 	}
 
