Fix crash when async compute is enabled.

The diagnostic buffer was implicitly destroyed by the queue destructor before the context pool was drained, which can lead to a situation where there are still some FBreadcrumbStack objects owned by a context, and their destructor tries to use the diagnostic buffer after it's been freed. Made the queue destructor explicitly free the diagnostic buffer and set it to null before the breadcrumb stacks are deleted.

#rnx

diff --git a/Engine/Source/Runtime/D3D12RHI/Private/D3D12Device.cpp b/Engine/Source/Runtime/D3D12RHI/Private/D3D12Device.cpp
--- a/Engine/Source/Runtime/D3D12RHI/Private/D3D12Device.cpp
+++ b/Engine/Source/Runtime/D3D12RHI/Private/D3D12Device.cpp
@@ -104,6 +104,10 @@ void FD3D12Queue::SetupAfterDeviceCreation()
 
 FD3D12Queue::~FD3D12Queue()
 {
+	// The diagnostic buffer would be implicitly destroyed before the context pool, which can lead to a situation where there
+	// are still some FBreadcrumbStack objects owned by a context, and their destructor tries to use the diagnostic buffer
+	// after it's been freed. Explicitly freeing the buffer now and setting it to null works around this problem.
+	DiagnosticBuffer = nullptr;
 	check(PendingSubmission.IsEmpty());
 	check(PendingInterrupt.IsEmpty());
 }
diff --git a/Engine/Source/Runtime/D3D12RHI/Private/D3D12Submission.cpp b/Engine/Source/Runtime/D3D12RHI/Private/D3D12Submission.cpp
--- a/Engine/Source/Runtime/D3D12RHI/Private/D3D12Submission.cpp
+++ b/Engine/Source/Runtime/D3D12RHI/Private/D3D12Submission.cpp
@@ -1143,7 +1143,7 @@ FBreadcrumbStack::FBreadcrumbStack()
 FBreadcrumbStack::~FBreadcrumbStack()
 {
 	TUniquePtr<FD3D12DiagnosticBuffer>& DiagnosticBuffer = Queue->DiagnosticBuffer;
-	if (ContextId > 0)
+	if (ContextId > 0 && DiagnosticBuffer.IsValid())
 	{
 		{
 			FScopeLock Lock(&DiagnosticBuffer->CriticalSection);
diff --git a/Engine/Source/Runtime/D3D12RHI/Private/D3D12Submission.h b/Engine/Source/Runtime/D3D12RHI/Private/D3D12Submission.h
--- a/Engine/Source/Runtime/D3D12RHI/Private/D3D12Submission.h
+++ b/Engine/Source/Runtime/D3D12RHI/Private/D3D12Submission.h
@@ -162,7 +162,7 @@ struct FBreadcrumbStack
 		uint32 Sibling;
 	};
 
-	FD3D12Queue* Queue;
+	FD3D12Queue* Queue = nullptr;
 	uint32 NextIdx{ 0 };
 	int32 ContextId;
 	uint32 MaxMarkers{ 0 };
