Bindless Resources
- Adding correct handling for bindless uniform buffer parameters.
- Reworking RHITextureReference creation to allow the RHI to create them differently, especially with bindless views.
- Adding RHITextureReference bindless view updates to D3D12RHI.

diff --git a/Engine/Shaders/Public/Platform/D3D/D3DCommon.ush b/Engine/Shaders/Public/Platform/D3D/D3DCommon.ush
--- a/Engine/Shaders/Public/Platform/D3D/D3DCommon.ush
+++ b/Engine/Shaders/Public/Platform/D3D/D3DCommon.ush
@@ -5,7 +5,7 @@ D3DCommon.usf: Common defines for the d3d shader compiler
 =============================================================================*/
 
 // Update this GUID to invalidate D3D shaders and force a shader recompilation for D3D RHIs.
-// GUID: 55E9AB45-C918-4014-9500-D4912DBF9980
+// GUID: A44B0FDA-31ED-4E3B-9E04-A80DFDD5B015
 
 #if SM6_PROFILE
 	#define COMPILER_SUPPORTS_WAVE_SIZE 1
diff --git a/Engine/Source/Developer/ShaderCompilerCommon/Private/ShaderCompilerCommon.cpp b/Engine/Source/Developer/ShaderCompilerCommon/Private/ShaderCompilerCommon.cpp
--- a/Engine/Source/Developer/ShaderCompilerCommon/Private/ShaderCompilerCommon.cpp
+++ b/Engine/Source/Developer/ShaderCompilerCommon/Private/ShaderCompilerCommon.cpp
@@ -90,16 +90,20 @@ bool BuildResourceTableMapping(
 	int32 MaxBoundResourceTable = -1;
 
 	// Go through ALL the members of ALL the UB resources
-	for( auto MapIt = ResourceTableMap.CreateConstIterator(); MapIt; ++MapIt )
+	for (const TPair<FString, FResourceTableEntry>& ResourceEntry : ResourceTableMap)
 	{
-		const FString& Name	= MapIt->Key;
-		const FResourceTableEntry& Entry = MapIt->Value;
-
-		uint16 BufferIndex, BaseIndex, Size;
+		const FString& Name = ResourceEntry.Key;
+		const FResourceTableEntry& Entry = ResourceEntry.Value;
 
 		// If the shaders uses this member (eg View_PerlinNoise3DTexture)...
-		if (ParameterMap.FindParameterAllocation( *Name, BufferIndex, BaseIndex, Size ) )
+		if (TOptional<FParameterAllocation> Allocation = ParameterMap.FindParameterAllocation(Name))
 		{
+			const EShaderParameterType ParameterType = Allocation->Type;
+			const bool bBindlessParameter = (ParameterType == EShaderParameterType::BindlessResourceIndex || ParameterType == EShaderParameterType::BindlessSamplerIndex);
+
+			// Force bindless "indices" to zero since they're not needed in SetResourcesFromTables
+			const uint16 BaseIndex = bBindlessParameter ? 0 : Allocation->BaseIndex;
+
 			ParameterMap.RemoveParameterAllocation(*Name);
 
 			uint16 UniformBufferIndex = INDEX_NONE;
@@ -150,7 +154,7 @@ bool BuildResourceTableMapping(
 
 	// Emit hashes for all uniform buffers in the parameter map. We need to include the ones without resources as well
 	// (i.e. just constants), since the global uniform buffer bindings rely on valid hashes.
-	for (const auto& KeyValue : ParameterMap.GetParameterMap())
+	for (const TPair<FString, FParameterAllocation>& KeyValue : ParameterMap.GetParameterMap())
 	{
 		const FString& UniformBufferName = KeyValue.Key;
 		const FParameterAllocation& UniformBufferParameter = KeyValue.Value;
@@ -738,17 +742,14 @@ void HandleReflectedGlobalConstantBufferMember(
 		ParameterType);
 }
 
-void HandleReflectedRootConstantBufferMember(
-	const FShaderCompilerInput& Input,
-	const FShaderParameterParser& ShaderParameterParser,
+void HandleReflectedUniformBufferConstantBufferMember(
+	int32 UniformBufferSlot,
 	const FString& InMemberName,
 	int32 ReflectionOffset,
 	int32 ReflectionSize,
 	FShaderCompilerOutput& Output
 )
 {
-	ShaderParameterParser.ValidateShaderParameterType(Input, InMemberName, ReflectionOffset, ReflectionSize, Output);
-
 	FString MemberName = InMemberName;
 	const EShaderParameterType ParameterType = UE::ShaderCompilerCommon::ParseAndRemoveBindlessParameterPrefix(MemberName);
 
@@ -756,13 +757,34 @@ void HandleReflectedRootConstantBufferMember(
 	{
 		Output.ParameterMap.AddParameterAllocation(
 			*MemberName,
-			FShaderParametersMetadata::kRootCBufferBindingIndex,
+			UniformBufferSlot,
 			ReflectionOffset,
 			1,
-			ParameterType);
+			ParameterType
+		);
 	}
 }
 
+void HandleReflectedRootConstantBufferMember(
+	const FShaderCompilerInput& Input,
+	const FShaderParameterParser& ShaderParameterParser,
+	const FString& MemberName,
+	int32 ReflectionOffset,
+	int32 ReflectionSize,
+	FShaderCompilerOutput& Output
+)
+{
+	ShaderParameterParser.ValidateShaderParameterType(Input, MemberName, ReflectionOffset, ReflectionSize, Output);
+
+	HandleReflectedUniformBufferConstantBufferMember(
+		FShaderParametersMetadata::kRootCBufferBindingIndex,
+		MemberName,
+		ReflectionOffset,
+		ReflectionSize,
+		Output
+	);
+}
+
 void HandleReflectedRootConstantBuffer(
 	int32 ConstantBufferSize,
 	FShaderCompilerOutput& CompilerOutput
diff --git a/Engine/Source/Developer/ShaderCompilerCommon/Public/ShaderCompilerCommon.h b/Engine/Source/Developer/ShaderCompilerCommon/Public/ShaderCompilerCommon.h
--- a/Engine/Source/Developer/ShaderCompilerCommon/Public/ShaderCompilerCommon.h
+++ b/Engine/Source/Developer/ShaderCompilerCommon/Public/ShaderCompilerCommon.h
@@ -133,6 +133,14 @@ extern SHADERCOMPILERCOMMON_API void HandleReflectedGlobalConstantBufferMember(
 	FShaderCompilerOutput& CompilerOutput
 );
 
+extern SHADERCOMPILERCOMMON_API void HandleReflectedUniformBufferConstantBufferMember(
+	int32 UniformBufferSlot,
+	const FString& MemberName,
+	int32 ReflectionOffset,
+	int32 ReflectionSize,
+	FShaderCompilerOutput& CompilerOutput
+);
+
 extern SHADERCOMPILERCOMMON_API void HandleReflectedRootConstantBufferMember(
 	const FShaderCompilerInput& Input,
 	const FShaderParameterParser& ShaderParameterParser,
diff --git a/Engine/Source/Developer/Windows/ShaderFormatD3D/Private/D3DShaderCompiler.inl b/Engine/Source/Developer/Windows/ShaderFormatD3D/Private/D3DShaderCompiler.inl
--- a/Engine/Source/Developer/Windows/ShaderFormatD3D/Private/D3DShaderCompiler.inl
+++ b/Engine/Source/Developer/Windows/ShaderFormatD3D/Private/D3DShaderCompiler.inl
@@ -305,6 +305,8 @@ template <typename ID3D1xShaderReflection, typename D3D1x_SHADER_DESC, typename
 		bool& bGlobalUniformBufferUsed, bool& bDiagnosticBufferUsed, uint32& NumSamplers, uint32& NumSRVs, uint32& NumCBs, uint32& NumUAVs,
 		FShaderCompilerOutput& Output, TArray<FString>& UniformBufferNames, TBitArray<>& UsedUniformBufferSlots, TArray<FShaderCodeVendorExtension>& VendorExtensions)
 {
+	const bool bBindlessEnabled = (Input.Environment.CompilerFlags.Contains(CFLAG_BindlessResources) || Input.Environment.CompilerFlags.Contains(CFLAG_BindlessSamplers));
+
 	// Add parameters for shader resources (constant buffers, textures, samplers, etc. */
 	for (uint32 ResourceIndex = 0; ResourceIndex < ShaderDesc.BoundResources; ResourceIndex++)
 	{
@@ -425,6 +427,28 @@ template <typename ID3D1xShaderReflection, typename D3D1x_SHADER_DESC, typename
 				HandleReflectedUniformBuffer(UniformBufferName, CBIndex, Output);
 				
 				UsedUniformBufferSlots[CBIndex] = true;
+
+				if (bBindlessEnabled)
+				{
+					for (uint32 ConstantIndex = 0; ConstantIndex < CBDesc.Variables; ConstantIndex++)
+					{
+						ID3D1xShaderReflectionVariable* Variable = ConstantBuffer->GetVariableByIndex(ConstantIndex);
+
+						D3D1x_SHADER_VARIABLE_DESC VariableDesc;
+						Variable->GetDesc(&VariableDesc);
+
+						if (VariableDesc.uFlags & D3D_SVF_USED)
+						{
+							HandleReflectedUniformBufferConstantBufferMember(
+								CBIndex,
+								FString(VariableDesc.Name),
+								VariableDesc.StartOffset,
+								VariableDesc.Size,
+								Output
+							);
+						}
+					}
+				}
 			}
 
 			if (UniformBufferNames.Num() <= (int32)CBIndex)
diff --git a/Engine/Source/Runtime/D3D12RHI/Private/D3D12Commands.cpp b/Engine/Source/Runtime/D3D12RHI/Private/D3D12Commands.cpp
--- a/Engine/Source/Runtime/D3D12RHI/Private/D3D12Commands.cpp
+++ b/Engine/Source/Runtime/D3D12RHI/Private/D3D12Commands.cpp
@@ -80,12 +80,6 @@ inline void ValidateBoundUniformBuffer(FD3D12UniformBuffer * InUniformBuffer, FR
 #endif
 }
 
-inline FD3D12ShaderResourceView* GetShaderResourceViewFromTexture(FRHIResource* InResource, uint32 GpuIndex)
-{
-	FD3D12Texture* Texture = FD3D12CommandContext::RetrieveTexture(static_cast<FRHITexture*>(InResource), GpuIndex);
-	return Texture ? Texture->GetShaderResourceView() : nullptr;
-}
-
 static void BindUniformBuffer(FD3D12CommandContext& Context, FRHIShader* Shader, EShaderFrequency ShaderFrequency, uint32 BufferIndex, FD3D12UniformBuffer* InBuffer)
 {
 	ValidateBoundUniformBuffer(InBuffer, Shader, BufferIndex);
@@ -967,39 +961,105 @@ void FD3D12CommandContext::RHISetShaderUniformBuffer(FRHIComputeShader* ComputeS
 struct FD3D12ResourceBinder
 {
 	FD3D12CommandContext& Context;
+	FD3D12ConstantBuffer& ConstantBuffer;
 	const uint32 GpuIndex;
 	const EShaderFrequency Frequency;
+#if PLATFORM_SUPPORTS_BINDLESS_RENDERING
+	const bool bBindlessResources;
+	const bool bBindlessSamplers;
+#endif
 
-	FD3D12ResourceBinder(FD3D12CommandContext& InContext, EShaderFrequency InFrequency)
+	FD3D12ResourceBinder(FD3D12CommandContext& InContext, EShaderFrequency InFrequency, const FD3D12ShaderData* ShaderData)
 		: Context(InContext)
+		, ConstantBuffer(InContext.StageConstantBuffers[InFrequency])
 		, GpuIndex(InContext.GetGPUIndex())
 		, Frequency(InFrequency)
+#if PLATFORM_SUPPORTS_BINDLESS_RENDERING
+		, bBindlessResources(EnumHasAnyFlags(ShaderData->ResourceCounts.UsageFlags, EShaderResourceUsageFlags::BindlessResources))
+		, bBindlessSamplers(EnumHasAnyFlags(ShaderData->ResourceCounts.UsageFlags, EShaderResourceUsageFlags::BindlessSamplers))
+#endif
+	{
+	}
+
+#if PLATFORM_SUPPORTS_BINDLESS_RENDERING
+	void SetBindlessHandle(const FRHIDescriptorHandle& Handle, uint32 Offset)
 	{
+		if (Handle.IsValid())
+		{
+			const uint32 BindlessIndex = Handle.GetIndex();
+			ConstantBuffer.UpdateConstant(reinterpret_cast<const uint8*>(&BindlessIndex), Offset, 4);
+		}
 	}
+#endif
 
-	void SetUAV(FRHIUnorderedAccessView* UAV, uint8 Index, bool bClearResources = false)
+	void SetUAV(FRHIUnorderedAccessView* InUnorderedAccessView, uint32 Index, bool bClearResources = false)
 	{
-		FD3D12UnorderedAccessView* D3D12UAV = FD3D12CommandContext::RetrieveObject<FD3D12UnorderedAccessView>(UAV, GpuIndex);
+		FD3D12UnorderedAccessView* D3D12UnorderedAccessView = FD3D12CommandContext::RetrieveObject<FD3D12UnorderedAccessView>(InUnorderedAccessView, GpuIndex);
 		if (bClearResources)
 		{
-			Context.ClearShaderResources(D3D12UAV);
+			Context.ClearShaderResources(D3D12UnorderedAccessView);
+		}
+
+#if PLATFORM_SUPPORTS_BINDLESS_RENDERING
+		if (bBindlessResources)
+		{
+			Context.StateCache.QueueBindlessUAV(Frequency, D3D12UnorderedAccessView);
+		}
+		else
+#endif
+		{
+			Context.StateCache.SetUAV(Frequency, Index, D3D12UnorderedAccessView);
 		}
-		Context.StateCache.SetUAV(Frequency, Index, D3D12UAV);
 	}
 
-	void SetSRV(FRHIShaderResourceView* SRV, uint8 Index)
+	void SetSRV(FRHIShaderResourceView* InShaderResourceView, uint32 Index)
 	{
-		Context.StateCache.SetShaderResourceView(Frequency, FD3D12CommandContext::RetrieveObject<FD3D12ShaderResourceView>(SRV, GpuIndex), Index);
+		FD3D12ShaderResourceView* D3D12ShaderResourceView = FD3D12CommandContext::RetrieveObject<FD3D12ShaderResourceView>(InShaderResourceView, GpuIndex);
+
+#if PLATFORM_SUPPORTS_BINDLESS_RENDERING
+		if (bBindlessResources)
+		{
+			Context.StateCache.QueueBindlessSRV(Frequency, D3D12ShaderResourceView);
+		}
+		else
+#endif
+		{
+			Context.StateCache.SetShaderResourceView(Frequency, D3D12ShaderResourceView, Index);
+		}
 	}
 
-	void SetTexture(FRHITexture* TextureRHI, uint8 Index)
+	void SetTexture(FRHITexture* InTexture, uint32 Index)
 	{
-		Context.StateCache.SetShaderResourceView(Frequency, GetShaderResourceViewFromTexture(TextureRHI, GpuIndex), Index);
+		FD3D12Texture* D3D12Texture = FD3D12CommandContext::RetrieveTexture(InTexture, GpuIndex);
+		check(D3D12Texture);
+		FD3D12ShaderResourceView* D3D12ShaderResourceView = D3D12Texture->GetShaderResourceView();
+
+#if PLATFORM_SUPPORTS_BINDLESS_RENDERING
+		if (bBindlessResources)
+		{
+			Context.StateCache.QueueBindlessSRV(Frequency, D3D12ShaderResourceView);
+		}
+		else
+#endif
+		{
+			Context.StateCache.SetShaderResourceView(Frequency, D3D12ShaderResourceView, Index);
+		}
 	}
 
-	void SetSampler(FRHISamplerState* Sampler, uint8 Index)
+	void SetSampler(FRHISamplerState* Sampler, uint32 Index)
 	{
-		Context.StateCache.SetSamplerState(Frequency, FD3D12CommandContext::RetrieveObject<FD3D12SamplerState>(Sampler, GpuIndex), Index);
+		FD3D12SamplerState* D3D12SamplerState = FD3D12CommandContext::RetrieveObject<FD3D12SamplerState>(Sampler, GpuIndex);
+
+#if PLATFORM_SUPPORTS_BINDLESS_RENDERING
+		if (bBindlessSamplers)
+		{
+			// Nothing to do, only needs constants set
+		}
+		else
+#endif
+		{
+			Context.StateCache.SetSamplerState(Frequency, D3D12SamplerState, Index);
+		}
 	}
 };
 
@@ -1022,6 +1082,8 @@ static void SetShaderParametersOnContext(
 		ConstantBuffer.UpdateConstant(&InParametersData[Parameter.ByteOffset], Parameter.BaseIndex, Parameter.ByteSize);
 	}
 
+	FD3D12ResourceBinder Binder(Context, ShaderFrequency, GetShaderData(Shader));
+
 #if PLATFORM_SUPPORTS_BINDLESS_RENDERING
 	for (const FRHIShaderParameterResource& Parameter : InBindlessParameters)
 	{
@@ -1032,43 +1094,28 @@ static void SetShaderParametersOnContext(
 			switch (Parameter.Type)
 			{
 			case FRHIShaderParameterResource::EType::Texture:
-			{
 				Handle = static_cast<FRHITexture*>(Resource)->GetDefaultBindlessHandle();
-				StateCache.QueueBindlessSRV(ShaderFrequency, GetShaderResourceViewFromTexture(Resource, GpuIndex));
-			}
-			break;
+				Binder.SetTexture(static_cast<FRHITexture*>(Resource), Parameter.Index);
+				break;
 			case FRHIShaderParameterResource::EType::ResourceView:
-			{
 				Handle = static_cast<FRHIShaderResourceView*>(Resource)->GetBindlessHandle();
-				StateCache.QueueBindlessSRV(ShaderFrequency, FD3D12CommandContext::RetrieveObject<FD3D12ShaderResourceView>(Parameter.Resource, GpuIndex));
-			}
-			break;
+				Binder.SetSRV(static_cast<FRHIShaderResourceView*>(Resource), Parameter.Index);
+				break;
 			case FRHIShaderParameterResource::EType::UnorderedAccessView:
-			{
 				Handle = static_cast<FRHIUnorderedAccessView*>(Resource)->GetBindlessHandle();
-				FD3D12UnorderedAccessView* UAV = FD3D12CommandContext::RetrieveObject<FD3D12UnorderedAccessView>(Parameter.Resource, GpuIndex);
-				Context.ClearShaderResources(UAV);
-				StateCache.QueueBindlessUAV(ShaderFrequency, UAV);
-			}
-			break;
+				Binder.SetUAV(static_cast<FRHIUnorderedAccessView*>(Resource), Parameter.Index, true);
+				break;
 			case FRHIShaderParameterResource::EType::Sampler:
-			{
 				Handle = static_cast<FRHISamplerState*>(Resource)->GetBindlessHandle();
-			}
-			break;
+				Binder.SetSampler(static_cast<FRHISamplerState*>(Resource), Parameter.Index);
+				break;
 			}
 
-			if (Handle.IsValid())
-			{
-				const uint32 BindlessIndex = Handle.GetIndex();
-				ConstantBuffer.UpdateConstant(reinterpret_cast<const uint8*>(&BindlessIndex), Parameter.Index, 4);
-			}
+			Binder.SetBindlessHandle(Handle, Parameter.Index);
 		}
 	}
 #endif
 
-	FD3D12ResourceBinder Binder(Context, ShaderFrequency);
-
 	for (const FRHIShaderParameterResource& Parameter : InResourceParameters)
 	{
 		if (Parameter.Type == FRHIShaderParameterResource::EType::UnorderedAccessView)
@@ -1453,7 +1500,7 @@ void FD3D12CommandContext::SetResourcesFromTables(const ShaderType* Shader)
 
 	static constexpr EShaderFrequency Frequency = static_cast<EShaderFrequency>(ShaderType::StaticFrequency);
 
-	FD3D12ResourceBinder Binder(*this, Frequency);
+	FD3D12ResourceBinder Binder(*this, Frequency, Shader);
 	UE::RHICore::SetResourcesFromTables(
 		  Binder
 		, *Shader
diff --git a/Engine/Source/Runtime/D3D12RHI/Private/D3D12RHIPrivate.h b/Engine/Source/Runtime/D3D12RHI/Private/D3D12RHIPrivate.h
--- a/Engine/Source/Runtime/D3D12RHI/Private/D3D12RHIPrivate.h
+++ b/Engine/Source/Runtime/D3D12RHI/Private/D3D12RHIPrivate.h
@@ -459,6 +459,7 @@ public:
 	virtual FShaderResourceViewRHIRef RHICreateShaderResourceView(const FShaderResourceViewInitializer& Initializer) final override;
 	virtual void RHIUpdateShaderResourceView(FRHIShaderResourceView* SRV, FRHIBuffer* Buffer, uint32 Stride, uint8 Format) final override;
 	virtual void RHIUpdateShaderResourceView(FRHIShaderResourceView* SRV, FRHIBuffer* Buffer) final override;
+	virtual void RHIUpdateTextureReference(FRHITextureReference* TextureRef, FRHITexture* NewTexture) final override;
 	virtual FRHICalcTextureSizeResult RHICalcTexturePlatformSize(const FRHITextureDesc& Desc, uint32 FirstMipIndex) override;
 	virtual uint64 RHIGetMinimumAlignmentForBufferBackedSRV(EPixelFormat Format) final override;
 	virtual void RHIGetTextureMemoryStats(FTextureMemoryStats& OutStats) final override;
diff --git a/Engine/Source/Runtime/D3D12RHI/Private/D3D12Texture.cpp b/Engine/Source/Runtime/D3D12RHI/Private/D3D12Texture.cpp
--- a/Engine/Source/Runtime/D3D12RHI/Private/D3D12Texture.cpp
+++ b/Engine/Source/Runtime/D3D12RHI/Private/D3D12Texture.cpp
@@ -1016,6 +1016,22 @@ FTextureRHIRef FD3D12DynamicRHI::RHICreateTexture_RenderThread(class FRHICommand
 	return CreateD3D12Texture(CreateDesc, RHIImmediateCmdList, ResourceAllocator);
 }
 
+void FD3D12DynamicRHI::RHIUpdateTextureReference(FRHITextureReference* TextureRef, FRHITexture* InNewTexture)
+{
+	FRHITexture* NewTexture = InNewTexture ? InNewTexture : FRHITextureReference::GetDefaultTexture();
+
+	if (FRHIShaderResourceView* TextureRefSRV = TextureRef ? TextureRef->GetBindlessView() : nullptr)
+	{
+		FD3D12ShaderResourceView* TextureRefSRVD3D12 = ResourceCast(TextureRefSRV);
+		if (TextureRefSRVD3D12->GetBindlessHandle().IsValid())
+		{
+			FD3D12BindlessDescriptorManager& BindlessDescriptorManager = TextureRefSRVD3D12->GetParentDevice()->GetBindlessDescriptorManager();
+			BindlessDescriptorManager.UpdateDeferred(TextureRefSRVD3D12->GetBindlessHandle(), ResourceCast(NewTexture)->GetShaderResourceView()->GetOfflineCpuHandle());
+		}
+	}
+
+	FDynamicRHI::RHIUpdateTextureReference(TextureRef, NewTexture);
+}
 
 class FWaitInitialMipDataUploadTask
 {
diff --git a/Engine/Source/Runtime/RHI/Private/DynamicRHI.cpp b/Engine/Source/Runtime/RHI/Private/DynamicRHI.cpp
--- a/Engine/Source/Runtime/RHI/Private/DynamicRHI.cpp
+++ b/Engine/Source/Runtime/RHI/Private/DynamicRHI.cpp
@@ -21,6 +21,7 @@
 #include "RHICommandList.h"
 #include "RHIImmutableSamplerState.h"
 #include "RHIStrings.h"
+#include "RHITextureReference.h"
 #include "Serialization/MemoryImage.h"
 
 IMPLEMENT_TYPE_LAYOUT(FRayTracingGeometryInitializer);
@@ -537,9 +538,10 @@ void FDynamicRHI::RHIUpdateShaderResourceView(FRHIShaderResourceView* SRV, FRHIB
 	UE_LOG(LogRHI, Fatal, TEXT("RHIUpdateShaderResourceView isn't implemented for the current RHI"));
 }
 
-void FDynamicRHI::RHIUpdateTextureReference(FRHITextureReference* TextureRef, FRHITexture* NewTexture)
+void FDynamicRHI::RHIUpdateTextureReference(FRHITextureReference* TextureRef, FRHITexture* InReferencedTexture)
 {
-	TextureRef->SetReferencedTexture(NewTexture);
+	FRHITexture* ReferencedTexture = InReferencedTexture ? InReferencedTexture : FRHITextureReference::GetDefaultTexture();
+	TextureRef->SetReferencedTexture(ReferencedTexture);
 }
 
 uint64 FDynamicRHI::RHIGetMinimumAlignmentForBufferBackedSRV(EPixelFormat Format)
@@ -558,6 +560,23 @@ FTextureRHIRef FDynamicRHI::RHIAsyncCreateTexture2D(uint32 SizeX, uint32 SizeY,
 	return Result;
 }
 
+FTextureReferenceRHIRef FDynamicRHI::RHICreateTextureReference(FRHITexture* InReferencedTexture)
+{
+	FRHITexture* ReferencedTexture = InReferencedTexture ? InReferencedTexture : FRHITextureReference::GetDefaultTexture();
+
+	FShaderResourceViewRHIRef ShaderResourceView;
+
+#if PLATFORM_SUPPORTS_BINDLESS_RENDERING
+	// If the referenced texture is configured for bindless, make sure we also create an SRV to use for bindless.
+	if (ReferencedTexture && ReferencedTexture->GetDefaultBindlessHandle().IsValid())
+	{
+		ShaderResourceView = ::RHICreateShaderResourceView(ReferencedTexture, 0u);
+	}
+#endif
+
+	return new FRHITextureReference(ReferencedTexture, ShaderResourceView);
+}
+
 uint64 FDynamicRHI::RHIComputePrecachePSOHash(const FGraphicsPipelineStateInitializer& Initializer)
 {
 	// When compute precache PSO hash we assume a valid state precache PSO hash is already provided
diff --git a/Engine/Source/Runtime/RHI/Private/RHITextureReference.cpp b/Engine/Source/Runtime/RHI/Private/RHITextureReference.cpp
--- a/Engine/Source/Runtime/RHI/Private/RHITextureReference.cpp
+++ b/Engine/Source/Runtime/RHI/Private/RHITextureReference.cpp
@@ -7,10 +7,18 @@
 
 FRHITextureReference::FRHITextureReference()
 	: FRHITexture(RRT_TextureReference)
+	, ReferencedTexture(DefaultTexture.GetReference())
+	, BindlessView(nullptr)
+{
+	check(DefaultTexture);
+}
+
+FRHITextureReference::FRHITextureReference(FRHITexture* InReferencedTexture, FRHIShaderResourceView* InBindlessView)
+	: FRHITexture(RRT_TextureReference)
+	, ReferencedTexture(InReferencedTexture ? InReferencedTexture : DefaultTexture.GetReference())
+	, BindlessView(InBindlessView)
 {
 	check(DefaultTexture);
-	ReferencedTexture = DefaultTexture;
-	UpdateBindlessShaderResourceView();
 }
 
 FRHITextureReference::~FRHITextureReference() = default;
@@ -25,9 +33,12 @@ FRHIDescriptorHandle FRHITextureReference::GetDefaultBindlessHandle() const
 	check(ReferencedTexture);
 
 	// If an SRV has been created, return its handle
-	return BindlessShaderResourceViewRHI.IsValid() ?
-		BindlessShaderResourceViewRHI->GetBindlessHandle() : 
-		ReferencedTexture->GetDefaultBindlessHandle();
+	if (BindlessView.IsValid())
+	{
+		return BindlessView->GetBindlessHandle();
+	}
+
+	return ReferencedTexture->GetDefaultBindlessHandle();
 }
 
 void* FRHITextureReference::GetNativeResource() const
@@ -67,11 +78,3 @@ const FRHITextureDesc& FRHITextureReference::GetDesc() const
 	check(ReferencedTexture);
 	return ReferencedTexture->GetDesc();
 }
-
-void FRHITextureReference::UpdateBindlessShaderResourceView()
-{
-	if (RHIGetRuntimeBindlessResourcesConfiguration(GMaxRHIShaderPlatform) != ERHIBindlessConfiguration::Disabled)
-	{
-		BindlessShaderResourceViewRHI = RHICreateShaderResourceView(ReferencedTexture.GetReference(), 0u);
-	}
-}
diff --git a/Engine/Source/Runtime/RHI/Public/DynamicRHI.h b/Engine/Source/Runtime/RHI/Public/DynamicRHI.h
--- a/Engine/Source/Runtime/RHI/Public/DynamicRHI.h
+++ b/Engine/Source/Runtime/RHI/Public/DynamicRHI.h
@@ -521,6 +521,9 @@ public:
 	UE_DEPRECATED(5.2, "RHIAsyncCreateTexture2D now requires a completion callback. Using the old version in a critical section can lead to deadlock. Please switch to the new function signature.")
 	FTextureRHIRef RHIAsyncCreateTexture2D(uint32 SizeX, uint32 SizeY, uint8 Format, uint32 NumMips, ETextureCreateFlags Flags, ERHIAccess InResourceState, void** InitialMipData, uint32 NumInitialMips);
 
+	/** Create a texture reference. InReferencedTexture can be null. */
+	RHI_API virtual FTextureReferenceRHIRef RHICreateTextureReference(FRHITexture* InReferencedTexture);
+
 	/**
 	* Creates a shader resource view for a texture
 	*/
diff --git a/Engine/Source/Runtime/RHI/Public/RHICommandList.h b/Engine/Source/Runtime/RHI/Public/RHICommandList.h
--- a/Engine/Source/Runtime/RHI/Public/RHICommandList.h
+++ b/Engine/Source/Runtime/RHI/Public/RHICommandList.h
@@ -4827,9 +4827,9 @@ FORCEINLINE void RHIUpdateRHIResources(FRHIResourceUpdateInfo* UpdateInfos, int3
 	return FRHICommandListExecutor::GetImmediateCommandList().UpdateRHIResources(UpdateInfos, Num, bNeedReleaseRefs);
 }
 
-FORCEINLINE FTextureReferenceRHIRef RHICreateTextureReference()
+FORCEINLINE FTextureReferenceRHIRef RHICreateTextureReference(FRHITexture* InReferencedTexture = nullptr)
 {
-	return new FRHITextureReference();
+	return GDynamicRHI->RHICreateTextureReference(InReferencedTexture);
 }
 
 FORCEINLINE void RHIUpdateTextureReference(FRHITextureReference* TextureRef, FRHITexture* NewTexture)
diff --git a/Engine/Source/Runtime/RHI/Public/RHITextureReference.h b/Engine/Source/Runtime/RHI/Public/RHITextureReference.h
--- a/Engine/Source/Runtime/RHI/Public/RHITextureReference.h
+++ b/Engine/Source/Runtime/RHI/Public/RHITextureReference.h
@@ -7,7 +7,9 @@
 class FRHITextureReference final : public FRHITexture
 {
 public:
-	RHI_API explicit FRHITextureReference();
+	UE_DEPRECATED(5.3, "FRHITextureReference now requires a texture reference on construction")
+	RHI_API FRHITextureReference();
+	RHI_API FRHITextureReference(FRHITexture* InReferencedTexture, FRHIShaderResourceView* InBindlessView = nullptr);
 	RHI_API ~FRHITextureReference();
 
 	RHI_API virtual class FRHITextureReference* GetTextureReference() override;
@@ -25,29 +27,25 @@ public:
 	RHI_API virtual RHIValidation::FResource* GetTrackerResource() final override;
 #endif
 
-	inline FRHITexture * GetReferencedTexture() const
-	{
-		return ReferencedTexture.GetReference();
-	}
+	inline FRHITexture* GetReferencedTexture() const { return ReferencedTexture.GetReference(); }
+	inline FRHIShaderResourceView* GetBindlessView() const { return BindlessView.GetReference(); }
 
-	RHI_API void UpdateBindlessShaderResourceView();
+	static inline FRHITexture* GetDefaultTexture() { return DefaultTexture; }
 
 private:
 	friend class FRHICommandListImmediate;
 	friend class FDynamicRHI;
 
-	// Called only from FRHICommandListImmediate::UpdateTextureReference()
+	// Called only from FDynamicRHI::RHIUpdateTextureReference
 	void SetReferencedTexture(FRHITexture* InTexture)
 	{
-		ReferencedTexture = InTexture
-			? InTexture
-			: DefaultTexture.GetReference();
+		ReferencedTexture = InTexture;
 	}
 
 	TRefCountPtr<FRHITexture> ReferencedTexture;
 
 	// SRV to create a reusable slot for bindless handles for this texture reference
-	FShaderResourceViewRHIRef BindlessShaderResourceViewRHI;
+	FShaderResourceViewRHIRef BindlessView;
 
 	// This pointer is set by the InitRHI() function on the FBlackTextureWithSRV global resource,
 	// to allow FRHITextureReference to use the global black texture when the reference is nullptr.
