When enumerating adapters, make sure vendor preference is considered when choosing based on max video memory.

#rnx

diff --git a/Engine/Source/Runtime/D3D12RHI/Private/Windows/WindowsD3D12Device.cpp b/Engine/Source/Runtime/D3D12RHI/Private/Windows/WindowsD3D12Device.cpp
--- a/Engine/Source/Runtime/D3D12RHI/Private/Windows/WindowsD3D12Device.cpp
+++ b/Engine/Source/Runtime/D3D12RHI/Private/Windows/WindowsD3D12Device.cpp
@@ -951,7 +951,8 @@ void FD3D12DynamicRHIModule::FindAdapter()
 							FirstDiscreteAdapter = CurrentAdapter;
 						}
 
-						if (CurrentAdapter.Desc.DedicatedVideoMemory > BestMemoryAdapter.Desc.DedicatedVideoMemory)
+						if ((PreferredVendor == -1 || PreferredVendor == AdapterDesc.VendorId) &&
+						    CurrentAdapter.Desc.DedicatedVideoMemory > BestMemoryAdapter.Desc.DedicatedVideoMemory)
 						{
 							BestMemoryAdapter = CurrentAdapter;
 						}
