Changing new PreferredVendor checks to use EGpuVendorId.

diff --git a/Engine/Source/Runtime/D3D12RHI/Private/Windows/WindowsD3D12Device.cpp b/Engine/Source/Runtime/D3D12RHI/Private/Windows/WindowsD3D12Device.cpp
--- a/Engine/Source/Runtime/D3D12RHI/Private/Windows/WindowsD3D12Device.cpp
+++ b/Engine/Source/Runtime/D3D12RHI/Private/Windows/WindowsD3D12Device.cpp
@@ -951,7 +951,7 @@ void FD3D12DynamicRHIModule::FindAdapter()
 							FirstDiscreteAdapter = CurrentAdapter;
 						}
 
-						if ((PreferredVendor == -1 || PreferredVendor == AdapterDesc.VendorId) &&
+						if ((PreferredVendor == EGpuVendorId::Unknown || PreferredVendor == RHIConvertToGpuVendorId(AdapterDesc.VendorId)) &&
 						    CurrentAdapter.Desc.DedicatedVideoMemory > BestMemoryAdapter.Desc.DedicatedVideoMemory)
 						{
 							BestMemoryAdapter = CurrentAdapter;
