Improve logging when the Intel driver extensions fail to initialize.

#rnx

diff --git a/Engine/Source/Runtime/D3D12RHI/Private/Windows/WindowsD3D12Device.cpp b/Engine/Source/Runtime/D3D12RHI/Private/Windows/WindowsD3D12Device.cpp
--- a/Engine/Source/Runtime/D3D12RHI/Private/Windows/WindowsD3D12Device.cpp
+++ b/Engine/Source/Runtime/D3D12RHI/Private/Windows/WindowsD3D12Device.cpp
@@ -321,17 +321,21 @@ INTCExtensionContext* CreateIntelExtensionsContext(ID3D12Device* Device, INTCExt
 
 	if (SUCCEEDED(hr))
 	{
-		UE_LOG(LogD3D12RHI, Log, TEXT("Intel Extensions Framework enabled"));
+		UE_LOG(LogD3D12RHI, Log, TEXT("Intel Extensions Framework enabled."));
 	}
 	else
 	{
-		if (hr == E_OUTOFMEMORY)
+		if (hr == E_NOINTERFACE)
 		{
-			UE_LOG(LogD3D12RHI, Log, TEXT("Intel Extensions Framework not supported by driver"));
+			UE_LOG(LogD3D12RHI, Log, TEXT("Intel Extensions Framework not supported by driver. Please check if a driver update is available."));
 		}
 		else if (hr == E_INVALIDARG)
 		{
-			UE_LOG(LogD3D12RHI, Log, TEXT("Intel Extensions Framework passed invalid creation arguments"));
+			UE_LOG(LogD3D12RHI, Log, TEXT("Intel Extensions Framework passed invalid creation arguments."));
+		}
+		else
+		{
+			UE_LOG(LogD3D12RHI, Log, TEXT("Intel Extensions Framework failed to initialize. Error code: 0x%08x. Please check if a driver update is available."), hr);
 		}
 
 		if (IntelExtensionContext)
@@ -356,25 +360,32 @@ bool EnableIntelAtomic64Support(INTCExtensionContext* IntelExtensionContext, INT
 		INTC_D3D12_FEATURE_DATA_D3D12_OPTIONS1 INTCFeatureSupportData;;
 		const HRESULT hrCheck = INTC_D3D12_CheckFeatureSupport(IntelExtensionContext, INTC_D3D12_FEATURE_D3D12_OPTIONS1, &INTCFeatureSupportData, sizeof(INTCFeatureSupportData));
 
-		if (SUCCEEDED(hrCheck) && INTCFeatureSupportData.EmulatedTyped64bitAtomics)
+		if (SUCCEEDED(hrCheck))
 		{
-			INTC_D3D12_FEATURE INTCFeature;
-			INTCFeature.EmulatedTyped64bitAtomics = true;
-
-			const HRESULT hrSet = INTC_D3D12_SetFeatureSupport(IntelExtensionContext, &INTCFeature);
-			if (SUCCEEDED(hrSet))
+			if (INTCFeatureSupportData.EmulatedTyped64bitAtomics)
 			{
-				GDX12INTCAtomicUInt64Emulation = true;
-				UE_LOG(LogD3D12RHI, Log, TEXT("Intel Extensions 64-bit Typed Atomics emulation enabled."));
+				INTC_D3D12_FEATURE INTCFeature;
+				INTCFeature.EmulatedTyped64bitAtomics = true;
+
+				const HRESULT hrSet = INTC_D3D12_SetFeatureSupport(IntelExtensionContext, &INTCFeature);
+				if (SUCCEEDED(hrSet))
+				{
+					GDX12INTCAtomicUInt64Emulation = true;
+					UE_LOG(LogD3D12RHI, Log, TEXT("Intel Extensions 64-bit Typed Atomics emulation enabled."));
+				}
+				else
+				{
+					UE_LOG(LogD3D12RHI, Log, TEXT("Failed to enable Intel Extensions 64-bit Typed Atomics emulation, error code: 0x%08x."), hrSet);
+				}
 			}
 			else
 			{
-				UE_LOG(LogD3D12RHI, Log, TEXT("Failed to enable Intel Extensions 64-bit Typed Atomics emulation."));
+				UE_LOG(LogD3D12RHI, Log, TEXT("Intel Extensions 64-bit Typed Atomics emulation not needed."));
 			}
 		}
 		else
 		{
-			UE_LOG(LogD3D12RHI, Log, TEXT("Failed to check for Intel Extensions 64-bit Typed Atomics emulation."));
+			UE_LOG(LogD3D12RHI, Log, TEXT("Failed to check for Intel Extensions 64-bit Typed Atomics emulation, error code: 0x%08x."), hrCheck);
 		}
 	}
 
