Add CVar for lightweight DRED

diff --git a/Engine/Source/Runtime/D3D12RHI/Private/D3D12Adapter.cpp b/Engine/Source/Runtime/D3D12RHI/Private/D3D12Adapter.cpp
--- a/Engine/Source/Runtime/D3D12RHI/Private/D3D12Adapter.cpp
+++ b/Engine/Source/Runtime/D3D12RHI/Private/D3D12Adapter.cpp
@@ -56,10 +56,12 @@ static FAutoConsoleVariableRef CTrackedReleasedAllocationFrameRetention(
 static int32 GD3D12EnableGPUBreadCrumbs = 0;
 static int32 GD3D12EnableNvAftermath = 0;
 static int32 GD3D12EnableDRED = 0;
+static int32 GD3D12EnableLightweightDRED = 1;
 #else
 static int32 GD3D12EnableGPUBreadCrumbs = 1;
 static int32 GD3D12EnableNvAftermath = 1;
 static int32 GD3D12EnableDRED = 0;
+static int32 GD3D12EnableLightweightDRED = 1;
 #endif // UE_BUILD_SHIPPING || UE_BUILD_TEST
 
 static FAutoConsoleVariableRef CVarD3D12EnableGPUBreadCrumbs(
@@ -82,6 +84,13 @@ static FAutoConsoleVariableRef CVarD3D12EnableDRED(
 	TEXT("Has GPU overhead but gives the most information on the current GPU state when it crashes or hangs.\n"),
 	ECVF_RenderThreadSafe | ECVF_ReadOnly);
 
+static FAutoConsoleVariableRef CVarD3D12EnableLightweightDRED(
+	TEXT("r.D3D12.LightweightDRED"),
+	GD3D12EnableLightweightDRED,
+	TEXT("Enable Lightweight DRED GPU Crash debugging mode to track the current GPU state and logs information what operations the GPU executed last.")
+	TEXT("Gives the basic information on the current GPU state when it crashes or hangs on all PC hardware.\n"),
+	ECVF_RenderThreadSafe | ECVF_ReadOnly);
+
 bool GD3D12TrackAllAlocations = false;
 static TAutoConsoleVariable<int32> CVarD3D12TrackAllAllocations(
 	TEXT("D3D12.TrackAllAllocations"),
@@ -391,11 +400,9 @@ void FD3D12Adapter::CreateRootDevice(bool bWithDebug)
 
 	// Setup DRED if requested
 	bool bDRED = false;
-	bool bDREDMarkersOnly = false;
+	bool bDREDMarkersOnly = false; // LightweightDRED
 	bool bDREDContext = false;
-	{
-		bool bEnableDRED = EnumHasAnyFlags(GPUCrashDebuggingModes, ED3D12GPUCrashDebuggingModes::DRED);
-		
+	{		
 		HMODULE d3d12DllHandle = (HMODULE)FPlatformProcess::GetDllHandle(TEXT("d3d12.dll"));
 		typedef HRESULT(WINAPI* FD3D12GetInterface)(REFCLSID, REFIID, void**);
 
@@ -405,7 +412,7 @@ void FD3D12Adapter::CreateRootDevice(bool bWithDebug)
 
 			if (D3D12GetInterfaceFnPtr != nullptr)
 			{
-				if (bEnableDRED)
+				if (EnumHasAnyFlags(GPUCrashDebuggingModes, ED3D12GPUCrashDebuggingModes::DRED))
 				{
 					TRefCountPtr<ID3D12DeviceRemovedExtendedDataSettings> DredSettings;
 					HRESULT hr = D3D12GetInterfaceFnPtr(CLSID_D3D12DeviceRemovedExtendedData, IID_PPV_ARGS(DredSettings.GetInitReference()));
@@ -418,7 +425,9 @@ void FD3D12Adapter::CreateRootDevice(bool bWithDebug)
 						DredSettings->SetPageFaultEnablement(D3D12_DRED_ENABLEMENT_FORCED_ON);
 
 						bDRED = true;
-						UE_LOG(LogD3D12RHI, Log, TEXT("[DRED] Dred enabled"));
+						bDREDMarkersOnly = false;
+						SetEmitDrawEvents(true);
+						UE_LOG(LogD3D12RHI, Log, TEXT("[DRED] DRED enabled"));
 					}
 					else
 					{
@@ -438,48 +447,37 @@ void FD3D12Adapter::CreateRootDevice(bool bWithDebug)
 #endif // __ID3D12DeviceRemovedExtendedDataSettings1_INTERFACE_DEFINED__
 				}
 
-#ifdef __ID3D12DeviceRemovedExtendedDataSettings2_INTERFACE_DEFINED__
+				else if(GD3D12EnableLightweightDRED)
 				{
-					SetEmitDrawEvents(true);
-
+#ifdef __ID3D12DeviceRemovedExtendedDataSettings2_INTERFACE_DEFINED__
 					TRefCountPtr<ID3D12DeviceRemovedExtendedDataSettings2> DredSettings2;
 					HRESULT hr = D3D12GetInterfaceFnPtr(CLSID_D3D12DeviceRemovedExtendedData, IID_PPV_ARGS(DredSettings2.GetInitReference()));
 
 					if (SUCCEEDED(hr))
 					{
-						// Always enable DRED if lightweight option is available (version 1.2)
-						bEnableDRED = true;
-
 						// Turn on AutoBreadcrumbs and Page Fault reporting
 						DredSettings2->SetAutoBreadcrumbsEnablement(D3D12_DRED_ENABLEMENT_FORCED_ON);
 						DredSettings2->SetPageFaultEnablement(D3D12_DRED_ENABLEMENT_FORCED_ON);
 
-						if (EnumHasAnyFlags(GPUCrashDebuggingModes, ED3D12GPUCrashDebuggingModes::DRED))
-						{
-							bDREDMarkersOnly = false;
-							DredSettings2->UseMarkersOnlyAutoBreadcrumbs(false);
-						}
-						else
-						{
-							bDREDMarkersOnly = true;
-							DredSettings2->UseMarkersOnlyAutoBreadcrumbs(true);
-							UE_LOG(LogD3D12RHI, Log, TEXT("[DRED] Using lightweight DRED."));
-#if NV_AFTERMATH
-							// Disable Nvidia Aftermath if lightweight DRED is enabled by default so they don't interfere with each other.
-							GDX12NVAfterMathEnabled = 0;
-							EnumRemoveFlags(GPUCrashDebuggingModes, ED3D12GPUCrashDebuggingModes::NvAftermath);
-							UE_LOG(LogD3D12RHI, Log, TEXT("[DRED] Disabled Nvidia Aftermath to avoid conflicts."));
-#endif
-						}
+						bDREDMarkersOnly = true;
+						DredSettings2->UseMarkersOnlyAutoBreadcrumbs(true);
+						SetEmitDrawEvents(true);
+						UE_LOG(LogD3D12RHI, Log, TEXT("[DRED] Using lightweight DRED."));
+					}
+					else
+					{
+						UE_LOG(LogD3D12RHI, Log, TEXT("[DRED] Lightweight DRED requested but interface was not found, hresult: %x. DRED only works on Windows 10 1903+."), hr);
 					}
+#else
+					UE_LOG(LogD3D12RHI, Log, TEXT("[DRED] Lightweight DRED unsupported."));
+#endif // __ID3D12DeviceRemovedExtendedDataSettings2_INTERFACE_DEFINED__
 				}
 			}
-#endif // __ID3D12DeviceRemovedExtendedDataSettings2_INTERFACE_DEFINED__
 		}
 	}
 
 	FGenericCrashContext::SetEngineData(TEXT("RHI.DRED"), bDRED ? TEXT("true") : TEXT("false"));
-	FGenericCrashContext::SetEngineData(TEXT("RHI.DREDMarkersOnly"), bDREDMarkersOnly && bDRED ? TEXT("true") : TEXT("false"));
+	FGenericCrashContext::SetEngineData(TEXT("RHI.DREDMarkersOnly"), bDREDMarkersOnly ? TEXT("true") : TEXT("false"));
 	FGenericCrashContext::SetEngineData(TEXT("RHI.DREDContext"), bDREDContext && bDRED ? TEXT("true") : TEXT("false"));
 
 #endif // PLATFORM_WINDOWS || (PLATFORM_HOLOLENS && !UE_BUILD_SHIPPING && WITH_PIX_EVENT_RUNTIME)
@@ -803,9 +801,12 @@ void FD3D12Adapter::CreateRootDevice(bool bWithDebug)
 
 			// Break on D3D debug errors.
 			pd3dInfoQueue->SetBreakOnSeverity(D3D12_MESSAGE_SEVERITY_ERROR, true);
-			
-			// By default trace warnings to log (don't break) - exception handler takes care this
-			if (!FParse::Param(FCommandLine::Get(), TEXT("nod3dbreakonwarning")))
+
+			// Enable this to break on a specific id in order to quickly get a callstack
+			//pd3dInfoQueue->SetBreakOnID(D3D12_MESSAGE_ID_DEVICE_DRAW_CONSTANT_BUFFER_TOO_SMALL, true);
+
+			// Break on D3D warnings if warning log or warning breakpoint is enabled
+			if (bLogWarnings)
 			{
 				pd3dInfoQueue->SetBreakOnSeverity(D3D12_MESSAGE_SEVERITY_WARNING, true);
 			}
