Restoring handling of a null texture in D3D12 binding code.

diff --git a/Engine/Source/Runtime/D3D12RHI/Private/D3D12Commands.cpp b/Engine/Source/Runtime/D3D12RHI/Private/D3D12Commands.cpp
--- a/Engine/Source/Runtime/D3D12RHI/Private/D3D12Commands.cpp
+++ b/Engine/Source/Runtime/D3D12RHI/Private/D3D12Commands.cpp
@@ -1031,8 +1031,7 @@ struct FD3D12ResourceBinder
 	void SetTexture(FRHITexture* InTexture, uint32 Index)
 	{
 		FD3D12Texture* D3D12Texture = FD3D12CommandContext::RetrieveTexture(InTexture, GpuIndex);
-		check(D3D12Texture);
-		FD3D12ShaderResourceView* D3D12ShaderResourceView = D3D12Texture->GetShaderResourceView();
+		FD3D12ShaderResourceView* D3D12ShaderResourceView = D3D12Texture ? D3D12Texture->GetShaderResourceView() : nullptr;
 
 #if PLATFORM_SUPPORTS_BINDLESS_RENDERING
 		if (bBindlessResources)
