Removing CopyToResolveTarget.

diff --git a/Engine/Plugins/Media/MediaIOFramework/Source/MediaIOCore/Public/MediaCapture.h b/Engine/Plugins/Media/MediaIOFramework/Source/MediaIOCore/Public/MediaCapture.h
--- a/Engine/Plugins/Media/MediaIOFramework/Source/MediaIOCore/Public/MediaCapture.h
+++ b/Engine/Plugins/Media/MediaIOFramework/Source/MediaIOCore/Public/MediaCapture.h
@@ -478,12 +478,6 @@ protected:
 	 */
 	virtual void OnRHIResourceCaptured_AnyThread(const FCaptureBaseData& InBaseData, TSharedPtr<FMediaCaptureUserData, ESPMode::ThreadSafe> InUserData, FTextureRHIRef InTexture) { }
 	virtual void OnRHIResourceCaptured_AnyThread(const FCaptureBaseData& InBaseData, TSharedPtr<FMediaCaptureUserData, ESPMode::ThreadSafe> InUserData, FBufferRHIRef InBuffer) { }
-	
-
-	UE_DEPRECATED(5.1, "OnCustomCapture_RenderingThread has been deprecated while moving to a RDG based pipeline. Please use the RDG version instead.")
-	virtual void OnCustomCapture_RenderingThread(FRHICommandListImmediate& RHICmdList, const FCaptureBaseData& InBaseData, TSharedPtr<FMediaCaptureUserData, ESPMode::ThreadSafe> InUserData, FTextureRHIRef InSourceTexture, FTextureRHIRef TargetableTexture, FResolveParams& ResolveParams, FVector2D CropU, FVector2D CropV)
-	{
-	}
 
 	/**
 	 * Callback when the source texture is ready to be copied on the GPU. You need to copy yourself. You may use the callback to add passes to the graph builder.
diff --git a/Engine/Source/Runtime/Apple/MetalRHI/Private/MetalRenderTarget.cpp b/Engine/Source/Runtime/Apple/MetalRHI/Private/MetalRenderTarget.cpp
--- a/Engine/Source/Runtime/Apple/MetalRHI/Private/MetalRenderTarget.cpp
+++ b/Engine/Source/Runtime/Apple/MetalRHI/Private/MetalRenderTarget.cpp
@@ -33,116 +33,6 @@ static FAutoConsoleVariableRef CVarMetalUseTexGetBytes(
 								ECVF_RenderThreadSafe
 								);
 
-void FMetalRHICommandContext::RHICopyToResolveTarget(FRHITexture* SourceTextureRHI, FRHITexture* DestTextureRHI, const FResolveParams& ResolveParams)
-{
-	@autoreleasepool {
-	if (!SourceTextureRHI || !DestTextureRHI)
-	{
-		// nothing to do if one of the textures is null!
-		return;
-	}
-	if (SourceTextureRHI != DestTextureRHI)
-	{
-		FMetalSurface* Source = GetMetalSurfaceFromRHITexture(SourceTextureRHI);
-		FMetalSurface* Destination = GetMetalSurfaceFromRHITexture(DestTextureRHI);
-		
-		const FRHITextureDesc& SourceDesc = Source->GetDesc();
-		const FRHITextureDesc& DestinationDesc = Destination->GetDesc();
-
-		// Only valid to have nil Metal textures when they are TexCreate_Presentable
-		if (!Source->Texture)
-		{
-			// Source RHI texture is valid with no Presentable Metal texture - there is nothing to copy from
-			check(EnumHasAnyFlags(SourceDesc.Flags, TexCreate_Presentable));
-			return;
-		}
-		if (!Destination->Texture)
-		{
-			// Destination RHI texture is valid with no Presentable Metal texture - force fetch it now so we can complete the copy
-			check(EnumHasAnyFlags(DestinationDesc.Flags, TexCreate_Presentable));
-			Destination->GetDrawableTexture();
-			if(!Destination->Texture)
-			{
-				UE_LOG(LogRHI, Error, TEXT("Drawable for destination texture resolve target unavailable"));
-				return;
-			}
-		}
-
-		mtlpp::Origin Origin(0, 0, 0);
-		mtlpp::Size Size(0, 0, 1);
-		if (ResolveParams.Rect.IsValid())
-		{
-			// Partial copy
-			Origin.x = ResolveParams.Rect.X1;
-			Origin.y = ResolveParams.Rect.Y1;
-			Size.width = ResolveParams.Rect.X2 - ResolveParams.Rect.X1;
-			Size.height = ResolveParams.Rect.Y2 - ResolveParams.Rect.Y1;
-		}
-		else
-		{
-			// Whole of source copy
-			Origin.x = 0;
-			Origin.y = 0;
-			
-			Size.width = FMath::Max<uint32>(1, SourceDesc.Extent.X >> ResolveParams.MipIndex);
-			Size.height = FMath::Max<uint32>(1, SourceDesc.Extent.Y >> ResolveParams.MipIndex);
-			// clamp to a destination size
-			Size.width = FMath::Min<uint32>(Size.width, DestinationDesc.Extent.X >> ResolveParams.MipIndex);
-			Size.height = FMath::Min<uint32>(Size.height, DestinationDesc.Extent.Y >> ResolveParams.MipIndex);
-		}
-		
-		const bool bSrcCubemap  = SourceDesc.IsTextureCube();
-		const bool bDestCubemap = DestinationDesc.IsTextureCube();
-		
-		uint32 DestIndex = ResolveParams.DestArrayIndex * (bDestCubemap ? 6 : 1) + (bDestCubemap ? uint32(ResolveParams.CubeFace) : 0);
-		uint32 SrcIndex  = ResolveParams.SourceArrayIndex * (bSrcCubemap ? 6 : 1) + (bSrcCubemap ? uint32(ResolveParams.CubeFace) : 0);
-		
-		if(Profiler)
-		{
-			Profiler->RegisterGPUWork();
-		}
-
-		const bool bMSAASource = Source->MSAATexture;
-        const bool bMSAADest = Destination->MSAATexture;
-        const bool bDepthStencil = SourceDesc.Format == PF_DepthStencil;
-		if (bMSAASource && !bMSAADest)
-		{
-			// Resolve required - Device must support this - Using Shader for resolve not supported amd NumSamples should be 1
-			const bool bSupportsMSAADepthResolve = GetMetalDeviceContext().SupportsFeature(EMetalFeaturesMSAADepthResolve);
-			const bool bSupportsMSAAStoreAndResolve = GetMetalDeviceContext().SupportsFeature(EMetalFeaturesMSAAStoreAndResolve);
-			check( (!bDepthStencil && bSupportsMSAAStoreAndResolve) || (bDepthStencil && bSupportsMSAADepthResolve) );
-			
-			Context->CopyFromTextureToTexture(Source->MSAAResolveTexture, SrcIndex, ResolveParams.MipIndex, Origin, Size, Destination->Texture, DestIndex, ResolveParams.MipIndex, Origin);
-		}
-		else if(Source->Texture.GetPixelFormat() == Destination->Texture.GetPixelFormat())
-		{
-			// Blit Copy for matching formats
-			Context->CopyFromTextureToTexture(Source->Texture, SrcIndex, ResolveParams.MipIndex, Origin, Size, Destination->Texture, DestIndex, ResolveParams.MipIndex, Origin);
-		}
-		else
-		{
-			const FPixelFormatInfo& SourceFormatInfo = GPixelFormats[SourceDesc.Format];
-			const FPixelFormatInfo& DestFormatInfo = GPixelFormats[DestinationDesc.Format];
-			bool bUsingPixelFormatView = (Source->Texture.GetUsage() & mtlpp::TextureUsage::PixelFormatView) != 0;
-			
-			// Attempt to Resolve with a texture view - source Texture doesn't have to be created with MTLTextureUsagePixelFormatView for these cases e.g:
-			// If we are resolving to/from sRGB linear color space within the same format OR using same bit length color format
-			if	(	SourceFormatInfo.BlockBytes == DestFormatInfo.BlockBytes &&
-					(bUsingPixelFormatView || SourceFormatInfo.NumComponents == DestFormatInfo.NumComponents)
-				)
-			{
-				FMetalTexture SourceTextureView = Source->Texture.NewTextureView(Destination->Texture.GetPixelFormat(), Source->Texture.GetTextureType(), ns::Range(ResolveParams.MipIndex, 1), ns::Range(SrcIndex, 1));
-				if(SourceTextureView)
-				{
-					Context->CopyFromTextureToTexture(SourceTextureView, 0, 0, Origin, Size, Destination->Texture, DestIndex, ResolveParams.MipIndex, Origin);
-					SafeReleaseMetalTexture(SourceTextureView);
-				}
-			}
-		}
-	}
-	}
-}
-
 /** Helper for accessing R10G10B10A2 colors. */
 struct FMetalR10G10B10A2
 {
diff --git a/Engine/Source/Runtime/Apple/MetalRHI/Public/MetalRHIContext.h b/Engine/Source/Runtime/Apple/MetalRHI/Public/MetalRHIContext.h
--- a/Engine/Source/Runtime/Apple/MetalRHI/Public/MetalRHIContext.h
+++ b/Engine/Source/Runtime/Apple/MetalRHI/Public/MetalRHIContext.h
@@ -42,23 +42,12 @@ public:
 	virtual void RHIClearUAVUint(FRHIUnorderedAccessView* UnorderedAccessViewRHI, const FUintVector4& Values) final override;
 	
 	virtual void RHICopyTexture(FRHITexture* SourceTextureRHI, FRHITexture* DestTextureRHI, const FRHICopyTextureInfo& CopyInfo) final override;
-	
 	virtual void RHICopyBufferRegion(FRHIBuffer* DstBufferRHI, uint64 DstOffset, FRHIBuffer* SrcBufferRHI, uint64 SrcOffset, uint64 NumBytes) final override;
 
-	/**
-	 * Resolves from one texture to another.
-	 * @param SourceTexture - texture to resolve from, 0 is silenty ignored
-	 * @param DestTexture - texture to resolve to, 0 is silenty ignored
-	 * @param ResolveParams - optional resolve params
-	 */
-	virtual void RHICopyToResolveTarget(FRHITexture* SourceTexture, FRHITexture* DestTexture, const FResolveParams& ResolveParams) final override;
-	
 	virtual void RHIBeginRenderQuery(FRHIRenderQuery* RenderQuery) final override;
-	
 	virtual void RHIEndRenderQuery(FRHIRenderQuery* RenderQuery) final override;
 	
 	void RHIBeginOcclusionQueryBatch(uint32 NumQueriesInBatch);
-	
 	void RHIEndOcclusionQueryBatch();
 	
 	virtual void RHISubmitCommandsHint() override;
diff --git a/Engine/Source/Runtime/D3D12RHI/Private/D3D12CommandContext.h b/Engine/Source/Runtime/D3D12RHI/Private/D3D12CommandContext.h
--- a/Engine/Source/Runtime/D3D12RHI/Private/D3D12CommandContext.h
+++ b/Engine/Source/Runtime/D3D12RHI/Private/D3D12CommandContext.h
@@ -582,7 +582,6 @@ public:
 	virtual void RHISetMultipleViewports(uint32 Count, const FViewportBounds* Data) final override;
 	virtual void RHIClearUAVFloat(FRHIUnorderedAccessView* UnorderedAccessViewRHI, const FVector4f& Values) final override;
 	virtual void RHIClearUAVUint(FRHIUnorderedAccessView* UnorderedAccessViewRHI, const FUintVector4& Values) final override;
-	virtual void RHICopyToResolveTarget(FRHITexture* SourceTexture, FRHITexture* DestTexture, const FResolveParams& ResolveParams) override;
 	virtual void RHICopyTexture(FRHITexture* SourceTexture, FRHITexture* DestTexture, const FRHICopyTextureInfo& CopyInfo) final override;
 	virtual void RHICopyBufferRegion(FRHIBuffer* DestBuffer, uint64 DstOffset, FRHIBuffer* SourceBuffer, uint64 SrcOffset, uint64 NumBytes) final override;
 	virtual void RHICopyToStagingBuffer(FRHIBuffer* SourceBuffer, FRHIStagingBuffer* DestinationStagingBuffer, uint32 Offset, uint32 NumBytes) final override;
@@ -850,10 +849,6 @@ public:
 	{
 		ContextRedirect(RHIClearUAVUint(UnorderedAccessViewRHI, Values));
 	}
-	FORCEINLINE virtual void RHICopyToResolveTarget(FRHITexture* SourceTexture, FRHITexture* DestTexture, const FResolveParams& ResolveParams) final override
-	{
-		ContextRedirect(RHICopyToResolveTarget(SourceTexture, DestTexture, ResolveParams));
-	}
 	FORCEINLINE virtual void RHICopyTexture(FRHITexture* SourceTextureRHI, FRHITexture* DestTextureRHI, const FRHICopyTextureInfo& CopyInfo) final override
 	{
 		ContextRedirect(RHICopyTexture(SourceTextureRHI, DestTextureRHI, CopyInfo));
diff --git a/Engine/Source/Runtime/D3D12RHI/Private/D3D12RenderTarget.cpp b/Engine/Source/Runtime/D3D12RHI/Private/D3D12RenderTarget.cpp
--- a/Engine/Source/Runtime/D3D12RHI/Private/D3D12RenderTarget.cpp
+++ b/Engine/Source/Runtime/D3D12RHI/Private/D3D12RenderTarget.cpp
@@ -297,283 +297,6 @@ static DXGI_FORMAT GetPlaneFormat(DXGI_FORMAT InFormat, uint32 InPlaneSlice)
 	return InFormat;
 }
 
-/**
-* Copies the contents of the given surface to its resolve target texture.
-* @param SourceSurface - surface with a resolve texture to copy to
-* @param ResolveParams - optional resolve params
-*/
-void FD3D12CommandContext::RHICopyToResolveTarget(FRHITexture* SourceTextureRHI, FRHITexture* DestTextureRHI, const FResolveParams& ResolveParams)
-{
-	if (!SourceTextureRHI || !DestTextureRHI)
-	{
-		// no need to do anything (silently ignored)
-		return;
-	}
-
-	uint32 GPUIndex = GetGPUIndex();
-	FD3D12Texture* SourceTexture = GetD3D12TextureFromRHITexture(SourceTextureRHI, GPUIndex);
-	FD3D12Texture* DestTexture = GetD3D12TextureFromRHITexture(DestTextureRHI, GPUIndex);
-
-	const FRHITextureDesc& SourceDesc = SourceTexture->GetDesc();
-	const FRHITextureDesc& DestDesc = DestTexture->GetDesc();
-
-	if (SourceDesc.IsTexture2D() && DestDesc.IsTexture2D())
-	{
-		const D3D_FEATURE_LEVEL FeatureLevel = GetParentDevice()->GetParentAdapter()->GetFeatureLevel();
-		if (SourceTexture != DestTexture)
-		{
-			if (IsDefaultContext())
-			{
-				GetParentDevice()->RegisterGPUWork();
-			}
-
-			if (SourceTexture->GetDepthStencilView(FExclusiveDepthStencil::DepthWrite_StencilWrite)
-				&& SourceTextureRHI->IsMultisampled()
-				&& !DestTextureRHI->IsMultisampled())
-			{
-				D3D12_RESOURCE_DESC const& ResolveTargetDesc = DestTexture->GetResource()->GetDesc();
-
-				ResolveTextureUsingShader<FResolveDepthPS>(
-					SourceTexture,
-					DestTexture,
-					DestTexture->GetRenderTargetView(0, -1),
-					DestTexture->GetDepthStencilView(FExclusiveDepthStencil::DepthWrite_StencilWrite),
-					ResolveTargetDesc,
-					GetDefaultRect(ResolveParams.Rect, DestTexture->GetSizeX(), DestTexture->GetSizeY()),
-					GetDefaultRect(ResolveParams.Rect, DestTexture->GetSizeX(), DestTexture->GetSizeY()),
-					FDummyResolveParameter()
-					);
-			}
-			else
-			{
-				DXGI_FORMAT SrcFmt = (DXGI_FORMAT)GPixelFormats[SourceTextureRHI->GetFormat()].PlatformFormat;
-				DXGI_FORMAT DstFmt = (DXGI_FORMAT)GPixelFormats[DestTexture->GetFormat()].PlatformFormat;
-
-				DXGI_FORMAT Fmt = ConvertTypelessToUnorm((DXGI_FORMAT)GPixelFormats[DestTexture->GetFormat()].PlatformFormat);
-
-				// Determine whether a MSAA resolve is needed, or just a copy.
-				if (SourceTextureRHI->IsMultisampled() && !DestTexture->IsMultisampled())
-				{
-					FScopedResourceBarrier ConditionalScopeResourceBarrierDst(*this, DestTexture->GetResource()  , D3D12_RESOURCE_STATE_RESOLVE_DEST  , ResolveParams.DestArrayIndex  );
-					FScopedResourceBarrier ConditionalScopeResourceBarrierSrc(*this, SourceTexture->GetResource(), D3D12_RESOURCE_STATE_RESOLVE_SOURCE, ResolveParams.SourceArrayIndex);
-
-					FlushResourceBarriers();
-					GraphicsCommandList()->ResolveSubresource(
-						DestTexture->GetResource()->GetResource(),
-						ResolveParams.DestArrayIndex,
-						SourceTexture->GetResource()->GetResource(),
-						ResolveParams.SourceArrayIndex,
-						Fmt
-						);
-
-					UpdateResidency(SourceTexture->GetResource());
-					UpdateResidency(DestTexture->GetResource());
-				}
-				else
-				{
-					D3D12_RESOURCE_DESC const& srcDesc = SourceTexture->GetResource()->GetDesc();
-					D3D12_RESOURCE_DESC const& ResolveTargetDesc = DestTexture->GetResource()->GetDesc();
-					bool bCopySrcSubRect = ResolveParams.Rect.IsValid() && (ResolveParams.Rect.X1 != 0 || ResolveParams.Rect.Y1 != 0 || ResolveParams.Rect.X2 != srcDesc.Width || ResolveParams.Rect.Y2 != srcDesc.Height);
-					bool bCopyDstSubRect = ResolveParams.DestRect.IsValid() && (ResolveParams.DestRect.X1 != 0 || ResolveParams.DestRect.Y1 != 0 || ResolveParams.DestRect.X2 != ResolveTargetDesc.Width || ResolveParams.DestRect.Y2 != ResolveTargetDesc.Height);
-
-					if ((bCopySrcSubRect || bCopyDstSubRect)
-						&& !SourceTextureRHI->IsMultisampled()
-						&& !DestTexture->GetDepthStencilView(FExclusiveDepthStencil::DepthWrite_StencilWrite))
-					{
-						// currently no support for readback buffers
-						check(ResolveTargetDesc.Dimension != D3D12_RESOURCE_DIMENSION_BUFFER);
-
-						const FResolveRect& SrcRect = ResolveParams.Rect.IsValid() ? ResolveParams.Rect : FResolveRect(0, 0, srcDesc.Width, srcDesc.Height);
-						D3D12_BOX SrcBox;
-
-						SrcBox.left = SrcRect.X1;
-						SrcBox.top = SrcRect.Y1;
-						SrcBox.front = 0;
-						SrcBox.right = SrcRect.X2;
-						SrcBox.bottom = SrcRect.Y2;
-						SrcBox.back = 1;
-
-						const FResolveRect& DestRect = ResolveParams.DestRect.IsValid() ? ResolveParams.DestRect : SrcRect;
-
-						FScopedResourceBarrier ConditionalScopeResourceBarrierDst(*this, DestTexture->GetResource()  , D3D12_RESOURCE_STATE_COPY_DEST  , ResolveParams.DestArrayIndex  );
-						FScopedResourceBarrier ConditionalScopeResourceBarrierSrc(*this, SourceTexture->GetResource(), D3D12_RESOURCE_STATE_COPY_SOURCE, ResolveParams.SourceArrayIndex);
-
-						CD3DX12_TEXTURE_COPY_LOCATION DestCopyLocation(DestTexture->GetResource()->GetResource(), ResolveParams.DestArrayIndex);
-						CD3DX12_TEXTURE_COPY_LOCATION SourceCopyLocation(SourceTexture->GetResource()->GetResource(), ResolveParams.SourceArrayIndex);
-
-						FlushResourceBarriers();
-						GraphicsCommandList()->CopyTextureRegion(
-							&DestCopyLocation,
-							DestRect.X1, DestRect.Y1, 0,
-							&SourceCopyLocation,
-							&SrcBox);
-
-						UpdateResidency(SourceTexture->GetResource());
-						UpdateResidency(DestTexture->GetResource());
-					}
-					else
-					{
-						FScopedResourceBarrier ConditionalScopeResourceBarrierSrc(*this, SourceTexture->GetResource(), D3D12_RESOURCE_STATE_COPY_SOURCE, ResolveParams.SourceArrayIndex);
-
-						// Resolve to a buffer.
-						if (ResolveTargetDesc.Dimension == D3D12_RESOURCE_DIMENSION_BUFFER)
-						{
-							check(IsDefaultContext());
-
-							const uint32 PlaneSlice = 0;
-							const DXGI_FORMAT DestFormat = GetPlaneFormat(srcDesc.Format, PlaneSlice);
-
-							const uint32 BlockBytes = ComputeBytesPerPixel(DestFormat);
-							const uint32 XBytes = (uint32)srcDesc.Width * BlockBytes;
-							const uint32 XBytesAligned = Align(XBytes, FD3D12_TEXTURE_DATA_PITCH_ALIGNMENT);
-
-							D3D12_PLACED_SUBRESOURCE_FOOTPRINT DestFootprint{};
-							DestFootprint.Footprint.Depth = 1;
-							DestFootprint.Footprint.Height = srcDesc.Height;
-							DestFootprint.Footprint.Width = srcDesc.Width;
-							DestFootprint.Footprint.Format = DestFormat;
-							DestFootprint.Footprint.RowPitch = XBytesAligned;
-
-							const UINT SourceSubresource = D3D12CalcSubresource(ResolveParams.MipIndex, ResolveParams.SourceArrayIndex, PlaneSlice, srcDesc.MipLevels, srcDesc.DepthOrArraySize);
-
-							CD3DX12_TEXTURE_COPY_LOCATION DestCopyLocation(DestTexture->GetResource()->GetResource(), DestFootprint);
-							CD3DX12_TEXTURE_COPY_LOCATION SourceCopyLocation(SourceTexture->GetResource()->GetResource(), SourceSubresource);
-
-							FlushResourceBarriers();
-							GraphicsCommandList()->CopyTextureRegion(
-								&DestCopyLocation,
-								0, 0, 0,
-								&SourceCopyLocation,
-								nullptr);
-
-							UpdateResidency(SourceTexture->GetResource());
-							UpdateResidency(DestTexture->GetResource());
-						}
-						// Resolve to a texture.
-						else
-						{
-							// Transition to the copy dest state.
-							FScopedResourceBarrier ConditionalScopeResourceBarrierDst(*this, DestTexture->GetResource(), D3D12_RESOURCE_STATE_COPY_DEST, 0);
-
-							CD3DX12_TEXTURE_COPY_LOCATION DestCopyLocation(DestTexture->GetResource()->GetResource(), ResolveParams.DestArrayIndex);
-							CD3DX12_TEXTURE_COPY_LOCATION SourceCopyLocation(SourceTexture->GetResource()->GetResource(), ResolveParams.SourceArrayIndex);
-
-							FlushResourceBarriers();
-							GraphicsCommandList()->CopyTextureRegion(
-								&DestCopyLocation,
-								0, 0, 0,
-								&SourceCopyLocation,
-								nullptr);
-
-							UpdateResidency(SourceTexture->GetResource());
-							UpdateResidency(DestTexture->GetResource());
-						}
-					}
-				}
-			}
-		}
-	}
-	else if (SourceDesc.IsTextureCube() && DestDesc.IsTextureCube())
-	{
-		if (SourceTexture != DestTexture)
-		{
-			if (IsDefaultContext())
-			{
-				GetParentDevice()->RegisterGPUWork();
-			}
-
-			// Determine the cubemap face being resolved.
-			const uint32 D3DFace = GetD3D12CubeFace(ResolveParams.CubeFace);
-			const uint32 SourceSubresource = CalcSubresource(ResolveParams.MipIndex, ResolveParams.SourceArrayIndex * 6 + D3DFace, SourceTexture->GetNumMips());
-			const uint32 DestSubresource = CalcSubresource(ResolveParams.MipIndex, ResolveParams.DestArrayIndex * 6 + D3DFace, DestTexture->GetNumMips());
-
-			// Determine whether a MSAA resolve is needed, or just a copy.
-			if (SourceTextureRHI->IsMultisampled() && !DestTexture->IsMultisampled())
-			{
-				FScopedResourceBarrier ConditionalScopeResourceBarrierDst(*this, DestTexture->GetResource()  , D3D12_RESOURCE_STATE_RESOLVE_DEST  , DestSubresource  );
-				FScopedResourceBarrier ConditionalScopeResourceBarrierSrc(*this, SourceTexture->GetResource(), D3D12_RESOURCE_STATE_RESOLVE_SOURCE, SourceSubresource);
-
-				FlushResourceBarriers();
-				GraphicsCommandList()->ResolveSubresource(
-					DestTexture->GetResource()->GetResource(),
-					DestSubresource,
-					SourceTexture->GetResource()->GetResource(),
-					SourceSubresource,
-					(DXGI_FORMAT)GPixelFormats[DestTexture->GetFormat()].PlatformFormat
-					);
-
-				UpdateResidency(SourceTexture->GetResource());
-				UpdateResidency(DestTexture->GetResource());
-			}
-			else
-			{
-				CD3DX12_TEXTURE_COPY_LOCATION DestCopyLocation(DestTexture->GetResource()->GetResource(), DestSubresource);
-				CD3DX12_TEXTURE_COPY_LOCATION SourceCopyLocation(SourceTexture->GetResource()->GetResource(), SourceSubresource);
-
-				FScopedResourceBarrier ConditionalScopeResourceBarrierDst(*this, DestTexture->GetResource()  , D3D12_RESOURCE_STATE_COPY_DEST  , DestCopyLocation.SubresourceIndex  );
-				FScopedResourceBarrier ConditionalScopeResourceBarrierSrc(*this, SourceTexture->GetResource(), D3D12_RESOURCE_STATE_COPY_SOURCE, SourceCopyLocation.SubresourceIndex);
-
-				FlushResourceBarriers();
-				GraphicsCommandList()->CopyTextureRegion(
-					&DestCopyLocation,
-					0, 0, 0,
-					&SourceCopyLocation,
-					nullptr);
-
-				UpdateResidency(SourceTexture->GetResource());
-				UpdateResidency(DestTexture->GetResource());
-			}
-		}
-	}
-	else if (SourceDesc.IsTexture2D() && DestDesc.IsTextureCube())
-	{
-		// If source is 2D and Dest is a cube then copy the 2D texture to the specified cube face.
-		// Determine the cubemap face being resolved.
-		const uint32 D3DFace = GetD3D12CubeFace(ResolveParams.CubeFace);
-		const uint32 Subresource = CalcSubresource(0, D3DFace, 1);
-
-		CD3DX12_TEXTURE_COPY_LOCATION DestCopyLocation(DestTexture->GetResource()->GetResource(), Subresource);
-		CD3DX12_TEXTURE_COPY_LOCATION SourceCopyLocation(SourceTexture->GetResource()->GetResource(), 0);
-
-		FScopedResourceBarrier ConditionalScopeResourceBarrierDst(*this, DestTexture->GetResource()  , D3D12_RESOURCE_STATE_COPY_DEST  , DestCopyLocation.SubresourceIndex  );
-		FScopedResourceBarrier ConditionalScopeResourceBarrierSrc(*this, SourceTexture->GetResource(), D3D12_RESOURCE_STATE_COPY_SOURCE, SourceCopyLocation.SubresourceIndex);
-
-		FlushResourceBarriers();
-		GraphicsCommandList()->CopyTextureRegion(
-			&DestCopyLocation,
-			0, 0, 0,
-			&SourceCopyLocation,
-			nullptr);
-
-		UpdateResidency(SourceTexture->GetResource());
-		UpdateResidency(DestTexture->GetResource());
-	}
-
-	else if (SourceDesc.IsTexture3D() && DestDesc.IsTexture3D())
-	{
-		// bit of a hack.  no one resolves slice by slice and 0 is the default value.  assume for the moment they are resolving the whole texture.
-		check(ResolveParams.SourceArrayIndex == 0);
-		check(SourceTexture == DestTexture);
-	}
-
-	// Force transtion to readable since RHI user expects both source and dest to be in readable state after calling RHICopyToResolveTarget
-	{
-		TRHICommandList_RecursiveHazardous<FD3D12CommandContext> RHICmdList(this);
-		if (SourceTextureRHI && ResolveParams.SourceAccessFinal != ERHIAccess::Unknown)
-		{
-			RHICmdList.Transition(FRHITransitionInfo(SourceTextureRHI, ERHIAccess::Unknown, ResolveParams.SourceAccessFinal));
-		}
-		if (DestTextureRHI && SourceTextureRHI != DestTextureRHI && ResolveParams.DestAccessFinal != ERHIAccess::Unknown)
-		{
-			RHICmdList.Transition(FRHITransitionInfo(DestTextureRHI, ERHIAccess::Unknown, ResolveParams.DestAccessFinal));
-		}
-	}
-
-	ConditionalSplitCommandList();
-
-	DEBUG_EXECUTE_COMMAND_LIST(this);
-}
-
 void FD3D12CommandContext::ResolveTexture(UE::RHICore::FResolveTextureInfo Info)
 {
 	uint32 GPUIndex = GetGPUIndex();
diff --git a/Engine/Source/Runtime/NullDrv/Public/NullRHI.h b/Engine/Source/Runtime/NullDrv/Public/NullRHI.h
--- a/Engine/Source/Runtime/NullDrv/Public/NullRHI.h
+++ b/Engine/Source/Runtime/NullDrv/Public/NullRHI.h
@@ -266,10 +266,6 @@ public:
 	}
 	virtual void RHIUnlockTextureCubeFace(FRHITextureCube* Texture, uint32 FaceIndex, uint32 ArrayIndex, uint32 MipIndex, bool bLockWithinMiptail) final override
 	{
-	}
-	virtual void RHICopyToResolveTarget(FRHITexture* SourceTexture, FRHITexture* DestTexture, const FResolveParams& ResolveParams) final override
-	{
-
 	}
 
 	virtual void RHICopyTexture(FRHITexture* SourceTexture, FRHITexture* DestTexture, const FRHICopyTextureInfo& CopyInfo) final override
diff --git a/Engine/Source/Runtime/OpenGLDrv/Private/OpenGLRenderTarget.cpp b/Engine/Source/Runtime/OpenGLDrv/Private/OpenGLRenderTarget.cpp
--- a/Engine/Source/Runtime/OpenGLDrv/Private/OpenGLRenderTarget.cpp
+++ b/Engine/Source/Runtime/OpenGLDrv/Private/OpenGLRenderTarget.cpp
@@ -442,153 +442,6 @@ void FOpenGLDynamicRHI::PurgeFramebufferFromCaches( GLuint Framebuffer )
 	}
 }
 
-void FOpenGLDynamicRHI::RHICopyToResolveTarget(FRHITexture* SourceTextureRHI, FRHITexture* DestTextureRHI, const FResolveParams& ResolveParams)
-{
-	if (!SourceTextureRHI || !DestTextureRHI)
-	{
-		// no need to do anything (silently ignored)
-		return;
-	}
-
-	FOpenGLTexture* SourceTexture = ResourceCast(SourceTextureRHI);
-	FOpenGLTexture* DestTexture = ResourceCast(DestTextureRHI);
-
-	if (SourceTexture != DestTexture)
-	{
-		const FRHITextureDesc& SrcDesc = SourceTexture->GetDesc();
-		const FRHITextureDesc& DstDesc = DestTexture->GetDesc();
-
-		VERIFY_GL_SCOPE();
-
-		check(GMaxRHIFeatureLevel >= ERHIFeatureLevel::SM5 || ResolveParams.SourceArrayIndex == 0);
-		check(GMaxRHIFeatureLevel >= ERHIFeatureLevel::SM5 || ResolveParams.DestArrayIndex == 0);
-
-		const bool bSrcCubemap  = SourceTextureRHI->GetTextureCube() != NULL;
-		const bool bDestCubemap = DestTextureRHI->GetTextureCube() != NULL;
-
-		uint32 DestIndex = ResolveParams.DestArrayIndex * (bDestCubemap ? 6 : 1) + (bDestCubemap ? uint32(ResolveParams.CubeFace) : 0);
-		uint32 SrcIndex  = ResolveParams.SourceArrayIndex * (bSrcCubemap ? 6 : 1) + (bSrcCubemap ? uint32(ResolveParams.CubeFace) : 0);
-
-		FIntRect SrcRect(ResolveParams.Rect.X1, ResolveParams.Rect.Y1, ResolveParams.Rect.X2, ResolveParams.Rect.Y2);
-		if (SrcRect.IsEmpty())
-		{
-			// Empty rect mans that the entire source is to be copied. Note that we can't use ResolveParams.Rect.IsValid(), because it
-			// returns false if the rectangle is "inside out" (e.g. X1 > X2), and we want to perform flipping when that's the case.
-			SrcRect.Min.X = 0;
-			SrcRect.Min.Y = 0;
-			SrcRect.Max.X = SrcDesc.Extent.X;
-			SrcRect.Max.Y = SrcDesc.Extent.Y;
-			SrcRect.Max.X = FMath::Max<int32>(1, SrcRect.Max.X >> ResolveParams.MipIndex);
-			SrcRect.Max.Y = FMath::Max<int32>(1, SrcRect.Max.Y >> ResolveParams.MipIndex);
-		}
-
-		FIntRect DestRect(ResolveParams.DestRect.X1, ResolveParams.DestRect.Y1, ResolveParams.DestRect.X2, ResolveParams.DestRect.Y2);
-		if(DestRect.IsEmpty())
-		{
-			DestRect.Min.X = 0;
-			DestRect.Min.Y = 0;
-			DestRect.Max.X = DstDesc.Extent.X;
-			DestRect.Max.Y = DstDesc.Extent.Y;
-			DestRect.Max.X = FMath::Max<int32>(1, DestRect.Max.X >> ResolveParams.MipIndex);
-			DestRect.Max.Y = FMath::Max<int32>(1, DestRect.Max.Y >> ResolveParams.MipIndex);
-		}
-
-		GPUProfilingData.RegisterGPUWork();
-		const uint32 MipmapLevel = ResolveParams.MipIndex;
-
-		const bool bTrueBlit = !SourceTextureRHI->IsMultisampled()
-			&& !DestTextureRHI->IsMultisampled()
-			&& SourceTextureRHI->GetFormat() == DestTextureRHI->GetFormat()
-			&& SrcRect.Size() == DestRect.Size()
-			&& SrcRect.Width() > 0
-			&& SrcRect.Height() > 0
-			&& SourceTexture->Target == DestTexture->Target // glCopyImageSubData() doesn't like copying from a texture to a renderbuffer on Android
-			;
-		
-		if ( !bTrueBlit)
-		{
-			// Color buffers can be GL_NONE for attachment purposes if they aren't used as render targets
-			const bool bIsColorBuffer = SourceTexture->Attachment != GL_DEPTH_STENCIL_ATTACHMENT
-																	&& SourceTexture->Attachment != GL_DEPTH_ATTACHMENT;
-			check( bIsColorBuffer || (SrcIndex == 0 && DestIndex == 0));
-			check( bIsColorBuffer || MipmapLevel == 0);
-			GLuint SrcFramebuffer = bIsColorBuffer ?
-				GetOpenGLFramebuffer(1, &SourceTexture, &SrcIndex, &MipmapLevel, NULL) :
-				GetOpenGLFramebuffer(0, NULL, NULL, NULL, SourceTexture);
-			GLuint DestFramebuffer = bIsColorBuffer ?
-				GetOpenGLFramebuffer(1, &DestTexture, &DestIndex, &MipmapLevel, NULL) :
-				GetOpenGLFramebuffer(0, NULL, NULL, NULL, DestTexture);
-
-			glBindFramebuffer(UGL_DRAW_FRAMEBUFFER, DestFramebuffer);
-			FOpenGL::DrawBuffer(bIsColorBuffer ? GL_COLOR_ATTACHMENT0 : GL_NONE);
-			glBindFramebuffer(UGL_READ_FRAMEBUFFER, SrcFramebuffer);
-			FOpenGL::ReadBuffer(bIsColorBuffer ? GL_COLOR_ATTACHMENT0 : GL_NONE);
-
-			// ToDo - Scissor and possibly color mask can impact blits
-			//  These should be disabled
-
-			GLbitfield Mask = GL_NONE;
-			if (bIsColorBuffer)
-			{
-				Mask = GL_COLOR_BUFFER_BIT;
-			}
-			else if (SourceTexture->Attachment == GL_DEPTH_ATTACHMENT)
-			{
-				Mask = GL_DEPTH_BUFFER_BIT;
-			}
-			else
-			{
-				Mask = GL_DEPTH_BUFFER_BIT | GL_STENCIL_BUFFER_BIT;
-			}
-
-			FOpenGL::BlitFramebuffer(
-				SrcRect.Min.X, SrcRect.Min.Y, SrcRect.Max.X, SrcRect.Max.Y,
-				DestRect.Min.X, DestRect.Min.Y, DestRect.Max.X, DestRect.Max.Y,
-				Mask,
-				GL_NEAREST
-				);
-		}
-		else
-		{
-			// CopyImageSubData seems like a better analog to what UE wants in most cases
-			// It has no interactions with any other state, and there is no filtering/conversion
-			// It does not support MSAA resolves though
-			FOpenGL::CopyImageSubData(	SourceTexture->GetResource(),
-										SourceTexture->Target,
-										MipmapLevel,
-										SrcRect.Min.X,
-										SrcRect.Min.Y,
-										SrcIndex,
-										DestTexture->GetResource(),
-										DestTexture->Target,
-										MipmapLevel,
-										DestRect.Min.X,
-										DestRect.Min.Y,
-										DestIndex,
-										SrcRect.Width(),
-										SrcRect.Height(),
-										1);
-		}
-
-	
-		// For CPU readback resolve targets we should issue the resolve to the internal PBO immediately.
-		// This makes any subsequent locking of that texture much cheaper as it won't have to stall on a pixel pack op.
-		bool bLockableTarget = DestTextureRHI->GetTexture2D() && EnumHasAnyFlags(DestTextureRHI->GetFlags(), TexCreate_CPUReadback) && !EnumHasAnyFlags(DestTextureRHI->GetFlags(), TexCreate_RenderTargetable|TexCreate_DepthStencilTargetable) && !DestTextureRHI->IsMultisampled();
-		if(bLockableTarget && !ResolveParams.Rect.IsValid())
-		{
-			DestTexture->Resolve(MipmapLevel, DestIndex);
-		}
-
-		GetContextStateForCurrentContext().Framebuffer = (GLuint)-1;
-	}
-	else
-	{
-		// no need to do anything (silently ignored)
-	}
-}
-
-
-
 void FOpenGLDynamicRHI::ReadSurfaceDataRaw(FOpenGLContextState& ContextState, FRHITexture* TextureRHI,FIntRect Rect,TArray<uint8>& OutData, FReadSurfaceDataFlags InFlags)
 {
 	VERIFY_GL_SCOPE();
diff --git a/Engine/Source/Runtime/OpenGLDrv/Public/OpenGLDrv.h b/Engine/Source/Runtime/OpenGLDrv/Public/OpenGLDrv.h
--- a/Engine/Source/Runtime/OpenGLDrv/Public/OpenGLDrv.h
+++ b/Engine/Source/Runtime/OpenGLDrv/Public/OpenGLDrv.h
@@ -428,7 +428,6 @@ public:
 	virtual void RHISetMultipleViewports(uint32 Count, const FViewportBounds* Data) final override;
 	virtual void RHIClearUAVFloat(FRHIUnorderedAccessView* UnorderedAccessViewRHI, const FVector4f& Values) final override;
 	virtual void RHIClearUAVUint(FRHIUnorderedAccessView* UnorderedAccessViewRHI, const FUintVector4& Values) final override;
-	virtual void RHICopyToResolveTarget(FRHITexture* SourceTexture, FRHITexture* DestTexture, const FResolveParams& ResolveParams) final override;
 	virtual void RHIBeginRenderQuery(FRHIRenderQuery* RenderQuery) final override;
 	virtual void RHIEndRenderQuery(FRHIRenderQuery* RenderQuery) final override;
 	virtual void RHISubmitCommandsHint() final override;
diff --git a/Engine/Source/Runtime/RHI/Public/RHICommandList.h b/Engine/Source/Runtime/RHI/Public/RHICommandList.h
--- a/Engine/Source/Runtime/RHI/Public/RHICommandList.h
+++ b/Engine/Source/Runtime/RHI/Public/RHICommandList.h
@@ -1527,23 +1527,6 @@ FRHICOMMAND_MACRO(FRHICommandClearUAVUint)
 	RHI_API void Execute(FRHICommandListBase& CmdList);
 };
 
-FRHICOMMAND_MACRO(FRHICommandCopyToResolveTarget)
-{
-	FResolveParams ResolveParams;
-	FRHITexture* SourceTexture;
-	FRHITexture* DestTexture;
-
-	FORCEINLINE_DEBUGGABLE FRHICommandCopyToResolveTarget(FRHITexture* InSourceTexture, FRHITexture* InDestTexture, const FResolveParams& InResolveParams)
-		: ResolveParams(InResolveParams)
-		, SourceTexture(InSourceTexture)
-		, DestTexture(InDestTexture)
-	{
-		ensure(SourceTexture);
-		ensure(DestTexture);
-	}
-	RHI_API void Execute(FRHICommandListBase& CmdList);
-};
-
 FRHICOMMAND_MACRO(FRHICommandCopyTexture)
 {
 	FRHICopyTextureInfo CopyInfo;
@@ -3222,18 +3205,6 @@ PRAGMA_ENABLE_DEPRECATION_WARNINGS
 #endif
 	}
 
-	UE_DEPRECATED(5.1, "CopyToResolveTarget is deprecated. Use render passes for MSAA resolves, or CopyTexture for copies instead.")
-	FORCEINLINE_DEBUGGABLE void CopyToResolveTarget(FRHITexture* SourceTextureRHI, FRHITexture* DestTextureRHI, const FResolveParams& ResolveParams)
-	{
-		//check(IsOutsideRenderPass());
-		if (Bypass())
-		{
-			GetContext().RHICopyToResolveTarget(SourceTextureRHI, DestTextureRHI, ResolveParams);
-			return;
-		}
-		ALLOC_COMMAND(FRHICommandCopyToResolveTarget)(SourceTextureRHI, DestTextureRHI, ResolveParams);
-	}
-
 	FORCEINLINE_DEBUGGABLE void CopyTexture(FRHITexture* SourceTextureRHI, FRHITexture* DestTextureRHI, const FRHICopyTextureInfo& CopyInfo)
 	{
 		check(SourceTextureRHI && DestTextureRHI);
diff --git a/Engine/Source/Runtime/RHI/Public/RHICommandListCommandExecutes.inl b/Engine/Source/Runtime/RHI/Public/RHICommandListCommandExecutes.inl
--- a/Engine/Source/Runtime/RHI/Public/RHICommandListCommandExecutes.inl
+++ b/Engine/Source/Runtime/RHI/Public/RHICommandListCommandExecutes.inl
@@ -27,7 +27,6 @@ struct FRHICommandBeginScene;
 struct FRHICommandBuildLocalBoundShaderState;
 struct FRHICommandBuildLocalGraphicsPipelineState;
 struct FRHICommandBuildLocalUniformBuffer;
-struct FRHICommandCopyToResolveTarget;
 struct FRHICommandDrawIndexedIndirect;
 struct FRHICommandDrawIndexedPrimitive;
 struct FRHICommandDrawIndexedPrimitiveIndirect;
@@ -290,12 +289,6 @@ void FRHICommandClearUAVUint::Execute(FRHICommandListBase& CmdList)
 	INTERNAL_DECORATOR_COMPUTE(RHIClearUAVUint)(UnorderedAccessViewRHI, Values);
 }
 
-void FRHICommandCopyToResolveTarget::Execute(FRHICommandListBase& CmdList)
-{
-	RHISTAT(CopyToResolveTarget);
-	INTERNAL_DECORATOR(RHICopyToResolveTarget)(SourceTexture, DestTexture, ResolveParams);
-}
-
 void FRHICommandCopyTexture::Execute(FRHICommandListBase& CmdList)
 {
 	RHISTAT(CopyTexture);
diff --git a/Engine/Source/Runtime/RHI/Public/RHIContext.h b/Engine/Source/Runtime/RHI/Public/RHIContext.h
--- a/Engine/Source/Runtime/RHI/Public/RHIContext.h
+++ b/Engine/Source/Runtime/RHI/Public/RHIContext.h
@@ -23,7 +23,6 @@
 class FRHIDepthRenderTargetView;
 class FRHIRenderTargetView;
 class FRHISetRenderTargetsInfo;
-struct FResolveParams;
 struct FViewportBounds;
 struct FRayTracingGeometryInstance;
 struct FRayTracingShaderBindings;
@@ -550,15 +549,6 @@ public:
 	// @param Data must not be 0
 	virtual void RHISetMultipleViewports(uint32 Count, const FViewportBounds* Data) = 0;
 
-	/**
-	* Resolves from one texture to another.
-	* @param SourceTexture - texture to resolve from, 0 is silently ignored
-	* @param DestTexture - texture to resolve to, 0 is silently ignored
-	* @param ResolveParams - optional resolve params
-	* @param Fence - optional fence, will be set once copy is completed by GPU
-	*/
-	virtual void RHICopyToResolveTarget(FRHITexture* SourceTexture, FRHITexture* DestTexture, const FResolveParams& ResolveParams) = 0;
-
 	/**
 	* Rebuilds the depth target HTILE meta data (on supported platforms).
 	* @param DepthTexture - the depth surface to resummarize.
diff --git a/Engine/Source/Runtime/RHI/Public/RHIResources.h b/Engine/Source/Runtime/RHI/Public/RHIResources.h
--- a/Engine/Source/Runtime/RHI/Public/RHIResources.h
+++ b/Engine/Source/Runtime/RHI/Public/RHIResources.h
@@ -4337,51 +4337,6 @@ struct FResolveRect
 	}
 };
 
-struct FResolveParams
-{
-	/** used to specify face when resolving to a cube map texture */
-	ECubeFace CubeFace;
-	/** resolve RECT bounded by [X1,Y1]..[X2,Y2]. Or -1 for fullscreen */
-	FResolveRect Rect;
-	FResolveRect DestRect;
-	/** The mip index to resolve in both source and dest. */
-	int32 MipIndex;
-	/** Array index to resolve in the source. */
-	int32 SourceArrayIndex;
-	/** Array index to resolve in the dest. */
-	int32 DestArrayIndex;
-	/** States to transition to at the end of the resolve operation. */
-	ERHIAccess SourceAccessFinal = ERHIAccess::SRVMask;
-	ERHIAccess DestAccessFinal = ERHIAccess::SRVMask;
-
-	/** constructor */
-	FResolveParams(
-		const FResolveRect& InRect = FResolveRect(),
-		ECubeFace InCubeFace = CubeFace_PosX,
-		int32 InMipIndex = 0,
-		int32 InSourceArrayIndex = 0,
-		int32 InDestArrayIndex = 0,
-		const FResolveRect& InDestRect = FResolveRect())
-		: CubeFace(InCubeFace)
-		, Rect(InRect)
-		, DestRect(InDestRect)
-		, MipIndex(InMipIndex)
-		, SourceArrayIndex(InSourceArrayIndex)
-		, DestArrayIndex(InDestArrayIndex)
-	{}
-
-	FORCEINLINE FResolveParams(const FResolveParams& Other)
-		: CubeFace(Other.CubeFace)
-		, Rect(Other.Rect)
-		, DestRect(Other.DestRect)
-		, MipIndex(Other.MipIndex)
-		, SourceArrayIndex(Other.SourceArrayIndex)
-		, DestArrayIndex(Other.DestArrayIndex)
-		, SourceAccessFinal(Other.SourceAccessFinal)
-		, DestAccessFinal(Other.DestAccessFinal)
-	{}
-};
-
 struct FRHIRenderPassInfo
 {
 	struct FColorEntry
diff --git a/Engine/Source/Runtime/RHI/Public/RHIValidationContext.h b/Engine/Source/Runtime/RHI/Public/RHIValidationContext.h
--- a/Engine/Source/Runtime/RHI/Public/RHIValidationContext.h
+++ b/Engine/Source/Runtime/RHI/Public/RHIValidationContext.h
@@ -455,17 +455,6 @@ public:
 		RHIContext->RHIEndUAVOverlap(UAVs);
 	}
 
-	/**
-	* Resolves from one texture to another.
-	* @param SourceTexture - texture to resolve from, 0 is silenty ignored
-	* @param DestTexture - texture to resolve to, 0 is silenty ignored
-	* @param ResolveParams - optional resolve params
-	*/
-	virtual void RHICopyToResolveTarget(FRHITexture* SourceTexture, FRHITexture* DestTexture, const FResolveParams& ResolveParams) override final
-	{
-		RHIContext->RHICopyToResolveTarget(SourceTexture, DestTexture, ResolveParams);
-	}
-
 	virtual void RHIResummarizeHTile(FRHITexture2D* DepthTexture) override final
 	{
 		Tracker->Assert(DepthTexture->GetWholeResourceIdentity(), ERHIAccess::DSVWrite);
diff --git a/Engine/Source/Runtime/RenderCore/Private/RenderGraphUtils.cpp b/Engine/Source/Runtime/RenderCore/Private/RenderGraphUtils.cpp
--- a/Engine/Source/Runtime/RenderCore/Private/RenderGraphUtils.cpp
+++ b/Engine/Source/Runtime/RenderCore/Private/RenderGraphUtils.cpp
@@ -349,54 +349,6 @@ RENDERCORE_API void AddDrawTexturePass(
 	}
 }
 
-BEGIN_SHADER_PARAMETER_STRUCT(FCopyToResolveTargetParameters, )
-	RDG_TEXTURE_ACCESS_DYNAMIC(Input)
-	RDG_TEXTURE_ACCESS_DYNAMIC(Output)
-END_SHADER_PARAMETER_STRUCT()
-
-void AddCopyToResolveTargetPass(
-	FRDGBuilder& GraphBuilder,
-	FRDGTextureRef InputTexture,
-	FRDGTextureRef OutputTexture,
-	const FResolveParams& ResolveParams)
-{
-	check(InputTexture && OutputTexture);
-
-	if (InputTexture == OutputTexture)
-	{
-		return;
-	}
-
-	ERHIAccess AccessSource = ERHIAccess::ResolveSrc;
-	ERHIAccess AccessDest = ERHIAccess::ResolveDst;
-
-	// This might also just be a copy.
-	if (InputTexture->Desc.NumSamples == OutputTexture->Desc.NumSamples)
-	{
-		AccessSource = ERHIAccess::CopySrc;
-		AccessDest = ERHIAccess::CopyDest;
-	}
-
-	FCopyToResolveTargetParameters* Parameters = GraphBuilder.AllocParameters<FCopyToResolveTargetParameters>();
-	Parameters->Input = FRDGTextureAccess(InputTexture, AccessSource);
-	Parameters->Output = FRDGTextureAccess(OutputTexture, AccessDest);
-
-	FResolveParams LocalResolveParams = ResolveParams;
-	LocalResolveParams.SourceAccessFinal = AccessSource;
-	LocalResolveParams.DestAccessFinal = AccessDest;
-
-	GraphBuilder.AddPass(
-		RDG_EVENT_NAME("CopyToResolveTarget(%s -> %s)", InputTexture->Name, OutputTexture->Name),
-		Parameters,
-		ERDGPassFlags::Copy | ERDGPassFlags::Raster | ERDGPassFlags::SkipRenderPass,
-		[InputTexture, OutputTexture, LocalResolveParams](FRHICommandList& RHICmdList)
-	{
-	PRAGMA_DISABLE_DEPRECATION_WARNINGS
-		RHICmdList.CopyToResolveTarget(InputTexture->GetRHI(), OutputTexture->GetRHI(), LocalResolveParams);
-	PRAGMA_ENABLE_DEPRECATION_WARNINGS
-	});
-}
-
 BEGIN_SHADER_PARAMETER_STRUCT(FCopyBufferParameters, )
 	RDG_BUFFER_ACCESS(SrcBuffer, ERHIAccess::CopySrc)
 	RDG_BUFFER_ACCESS(DstBuffer, ERHIAccess::CopyDest)
diff --git a/Engine/Source/Runtime/RenderCore/Public/RenderGraphUtils.h b/Engine/Source/Runtime/RenderCore/Public/RenderGraphUtils.h
--- a/Engine/Source/Runtime/RenderCore/Public/RenderGraphUtils.h
+++ b/Engine/Source/Runtime/RenderCore/Public/RenderGraphUtils.h
@@ -734,16 +734,6 @@ RENDERCORE_API void AddDrawTexturePass(
 	FRDGTextureRef OutputTexture,
 	const FRDGDrawTextureInfo& DrawInfo);
 
-/** Adds a render graph pass to resolve from one texture to another. Uses RHICopyToResolveTarget under the hood.
- *  The formats of the two textures don't need to match.
- */
-UE_DEPRECATED(5.1, "AddCopyToResolveTargetPass is deprecated. Use AddCopyTexturePass to perform texture copies, or FRenderTargetBinding to perform resolves.")
-RENDERCORE_API void AddCopyToResolveTargetPass(
-	FRDGBuilder& GraphBuilder,
-	FRDGTextureRef InputTexture,
-	FRDGTextureRef OutputTexture,
-	const FResolveParams& ResolveParams);
-
 /** Adds a render graph pass to clear a texture or buffer UAV with a single typed value. */
 RENDERCORE_API void AddClearUAVPass(FRDGBuilder& GraphBuilder, FRDGBufferUAVRef BufferUAV, uint32 Value, ERDGPassFlags ComputePassFlags = ERDGPassFlags::Compute);
 
diff --git a/Engine/Source/Runtime/VulkanRHI/Private/VulkanContext.h b/Engine/Source/Runtime/VulkanRHI/Private/VulkanContext.h
--- a/Engine/Source/Runtime/VulkanRHI/Private/VulkanContext.h
+++ b/Engine/Source/Runtime/VulkanRHI/Private/VulkanContext.h
@@ -77,7 +77,6 @@ public:
 	virtual void RHISetMultipleViewports(uint32 Count, const FViewportBounds* Data) final override;
 	virtual void RHIClearUAVFloat(FRHIUnorderedAccessView* UnorderedAccessViewRHI, const FVector4f& Values) final override;
 	virtual void RHIClearUAVUint(FRHIUnorderedAccessView* UnorderedAccessViewRHI, const FUintVector4& Values) final override;
-	virtual void RHICopyToResolveTarget(FRHITexture* SourceTexture, FRHITexture* DestTexture, const FResolveParams& ResolveParams) final override;
 	virtual void RHICopyTexture(FRHITexture* SourceTexture, FRHITexture* DestTexture, const FRHICopyTextureInfo& CopyInfo) final override;
 	virtual void RHICopyBufferRegion(FRHIBuffer* DstBuffer, uint64 DstOffset, FRHIBuffer* SrcBuffer, uint64 SrcOffset, uint64 NumBytes) final override;
 	virtual void RHIBeginTransitions(TArrayView<const FRHITransition*> Transitions) override final;
diff --git a/Engine/Source/Runtime/VulkanRHI/Private/VulkanRenderTarget.cpp b/Engine/Source/Runtime/VulkanRHI/Private/VulkanRenderTarget.cpp
--- a/Engine/Source/Runtime/VulkanRHI/Private/VulkanRenderTarget.cpp
+++ b/Engine/Source/Runtime/VulkanRHI/Private/VulkanRenderTarget.cpp
@@ -25,17 +25,6 @@ static FAutoConsoleVariableRef CVarGVulkanFlushOnMapStaging(
 	ECVF_Default
 );
 
-
-static int32 GSubmitOnCopyToResolve = 0;
-static FAutoConsoleVariableRef CVarVulkanSubmitOnCopyToResolve(
-	TEXT("r.Vulkan.SubmitOnCopyToResolve"),
-	GSubmitOnCopyToResolve,
-	TEXT("Submits the Queue to the GPU on every RHICopyToResolveTarget call.\n")
-	TEXT(" 0: Do not submit (default)\n")
-	TEXT(" 1: Submit"),
-	ECVF_Default
-	);
-
 static int32 GIgnoreCPUReads = 0;
 static FAutoConsoleVariableRef CVarVulkanIgnoreCPUReads(
 	TEXT("r.Vulkan.IgnoreCPUReads"),
@@ -85,251 +74,6 @@ FVulkanRenderPass* FVulkanCommandListContext::PrepareRenderPassForPSOCreation(co
 	return RenderPass;
 }
 
-template<typename RegionType>
-static void SetupCopyOrResolveRegion(RegionType& Region, const FVulkanTexture& SrcSurface, const FVulkanTexture& DstSurface, const VkImageSubresourceRange& SrcRange, const VkImageSubresourceRange& DstRange, const FResolveParams& ResolveParams)
-{
-	FMemory::Memzero(Region);
-
-	const FRHITextureDesc& SrcDesc = SrcSurface.GetDesc();
-	const FRHITextureDesc& DstDesc = DstSurface.GetDesc();
-
-	bool bCopySrcSubRect = ResolveParams.Rect.IsValid()     && (ResolveParams.Rect.X1 != 0     || ResolveParams.Rect.Y1 != 0     || ResolveParams.Rect.X2 != SrcDesc.Extent.X || ResolveParams.Rect.Y2 != SrcDesc.Extent.Y);
-	bool bCopyDstSubRect = ResolveParams.DestRect.IsValid() && (ResolveParams.DestRect.X1 != 0 || ResolveParams.DestRect.Y1 != 0 || ResolveParams.DestRect.X2 != DstDesc.Extent.X || ResolveParams.DestRect.Y2 != DstDesc.Extent.Y);
-	if(bCopySrcSubRect || bCopyDstSubRect)
-	{
-		const FResolveRect& SrcSubRect = ResolveParams.Rect.IsValid()     ? ResolveParams.Rect     : FResolveRect(0, 0, SrcDesc.Extent.X, SrcDesc.Extent.Y);
-		const FResolveRect& DstSubRect = ResolveParams.DestRect.IsValid() ? ResolveParams.DestRect : FResolveRect(0, 0, DstDesc.Extent.X, DstDesc.Extent.Y);
-
-		const uint32 SrcSubRectWidth  = SrcSubRect.X2 - SrcSubRect.X1;
-		const uint32 SrcSubRectHeight = SrcSubRect.Y2 - SrcSubRect.Y1;
-
-		const uint32 DstSubRectWidth  = DstSubRect.X2 - DstSubRect.X1;
-		const uint32 DstSubRectHeight = DstSubRect.Y2 - DstSubRect.Y1;
-
-		ensure(SrcSubRectWidth == DstSubRectWidth && SrcSubRectHeight == DstSubRectHeight);
-
-		Region.srcOffset.x = SrcSubRect.X1;
-		Region.srcOffset.y = SrcSubRect.Y1;
-		Region.dstOffset.x = DstSubRect.X1;
-		Region.dstOffset.y = DstSubRect.Y1;
-
-		Region.extent.width = FMath::Max(1u, SrcSubRectWidth >> ResolveParams.MipIndex);
-		Region.extent.height = FMath::Max(1u, SrcSubRectHeight >> ResolveParams.MipIndex);
-	}
-	else
-	{
-		ensure(SrcDesc.Extent.X == DstDesc.Extent.X && SrcDesc.Extent.Y == DstDesc.Extent.Y);
-
-		Region.extent.width = FMath::Max(1, SrcDesc.Extent.X >> ResolveParams.MipIndex);
-		Region.extent.height = FMath::Max(1, SrcDesc.Extent.Y >> ResolveParams.MipIndex);
-	}
-
-	Region.extent.depth = 1;
-	Region.srcSubresource.aspectMask = SrcSurface.GetFullAspectMask();
-	Region.srcSubresource.baseArrayLayer = SrcRange.baseArrayLayer;
-	Region.srcSubresource.layerCount = 1;
-	Region.srcSubresource.mipLevel = ResolveParams.MipIndex;
-	Region.dstSubresource.aspectMask = DstSurface.GetFullAspectMask();
-	Region.dstSubresource.baseArrayLayer = DstRange.baseArrayLayer;
-	Region.dstSubresource.layerCount = 1;
-	Region.dstSubresource.mipLevel = ResolveParams.MipIndex;
-}
-
-void FVulkanCommandListContext::RHICopyToResolveTarget(FRHITexture* SourceTextureRHI, FRHITexture* DestTextureRHI, const FResolveParams& InResolveParams)
-{
-	if (!SourceTextureRHI || !DestTextureRHI)
-	{
-		// no need to do anything (silently ignored)
-		return;
-	}
-
-	FVulkanCmdBuffer* CmdBuffer = CommandBufferManager->GetActiveCmdBuffer();
-	check(CmdBuffer->IsOutsideRenderPass());
-
-	const FVulkanTexture& SrcSurface = *ResourceCast(SourceTextureRHI);
-	const FVulkanTexture& DstSurface = *ResourceCast(DestTextureRHI);
-
-	const FRHITextureDesc& SrcDesc = SrcSurface.GetDesc();
-	const FRHITextureDesc& DstDesc = DstSurface.GetDesc();
-
-	uint32 SrcNumLayers, DstNumLayers, SrcCubeFace = 0, DstCubeFace = 0;
-
-	switch (SourceTextureRHI->GetDesc().Dimension)
-	{
-	case ETextureDimension::Texture2D:
-		SrcNumLayers = 1;
-		if (DestTextureRHI->GetDesc().Dimension == ETextureDimension::Texture2D)
-		{
-			DstNumLayers = 1;
-		}
-		else
-		{
-			// Allow copying a 2D texture to a face of a cube texture.
-			check(DestTextureRHI->GetDesc().Dimension == ETextureDimension::TextureCube);
-			DstNumLayers = 6;
-			DstCubeFace = (uint32)InResolveParams.CubeFace;
-		}
-		break;
-
-	case ETextureDimension::Texture2DArray:
-		check(DestTextureRHI->GetDesc().Dimension == ETextureDimension::Texture2DArray);
-		SrcNumLayers = SrcDesc.ArraySize;
-		DstNumLayers = DstDesc.ArraySize;
-		break;
-
-	case ETextureDimension::Texture3D:
-		check(DestTextureRHI->GetDesc().Dimension == ETextureDimension::Texture3D);
-		SrcNumLayers = DstNumLayers = 1;
-		break;
-
-	case ETextureDimension::TextureCube:
-		SrcNumLayers = 6;
-		SrcCubeFace = DstCubeFace = (uint32)InResolveParams.CubeFace;
-		if (DestTextureRHI->GetDesc().Dimension == ETextureDimension::TextureCube)
-		{
-			DstNumLayers = 6;
-		}
-		else
-		{
-			// Allow copying a cube texture to a slice of a cube texture array.
-			check(DestTextureRHI->GetDesc().Dimension == ETextureDimension::TextureCubeArray);
-			DstNumLayers = 6 * DstDesc.ArraySize;
-		}
-		break;
-
-	case ETextureDimension::TextureCubeArray:
-		check(DestTextureRHI->GetDesc().Dimension == ETextureDimension::TextureCubeArray);
-		SrcNumLayers = 6 * SrcDesc.ArraySize;
-		DstNumLayers = 6 * DstDesc.ArraySize;
-		SrcCubeFace = DstCubeFace = (uint32)InResolveParams.CubeFace;
-		break;
-
-	default:
-		checkNoEntry();
-		return;
-	}
-
-	VkImageSubresourceRange SrcRange;
-	SrcRange.aspectMask = SrcSurface.GetFullAspectMask();
-	SrcRange.baseMipLevel = InResolveParams.MipIndex;
-	SrcRange.levelCount = 1;
-	SrcRange.baseArrayLayer = InResolveParams.SourceArrayIndex * SrcNumLayers + SrcCubeFace;
-	SrcRange.layerCount = 1;
-
-	VkImageSubresourceRange DstRange;
-	DstRange.aspectMask = DstSurface.GetFullAspectMask();
-	DstRange.baseMipLevel = InResolveParams.MipIndex;
-	DstRange.levelCount = 1;
-	DstRange.baseArrayLayer = InResolveParams.DestArrayIndex * DstNumLayers + DstCubeFace;
-	DstRange.layerCount = 1;
-
-	check(!EnumHasAnyFlags(SrcSurface.GetDesc().Flags, TexCreate_CPUReadback));
-	VkImageLayout SrcLayout = FVulkanLayoutManager::SetExpectedLayout(CmdBuffer, SrcSurface, ERHIAccess::CopySrc);
-
-	if (EnumHasAnyFlags(DstSurface.GetDesc().Flags, TexCreate_CPUReadback))
-	{
-		//Readback textures are represented as a buffer, so we can support miplevels on hardware that does not expose it.
-
-		const FVulkanCpuReadbackBuffer* CpuReadbackBuffer = DstSurface.GetCpuReadbackBuffer();
-		check(DstRange.baseArrayLayer == 0);
-		VkBufferImageCopy CopyRegion;
-		FMemory::Memzero(CopyRegion);
-		uint32 MipLevel = InResolveParams.MipIndex;
-		uint32 SizeX = FMath::Max(SrcDesc.Extent.X >> MipLevel, 1);
-		uint32 SizeY = FMath::Max(SrcDesc.Extent.Y >> MipLevel, 1);
-		CopyRegion.bufferOffset = CpuReadbackBuffer->MipOffsets[MipLevel];
-		CopyRegion.bufferRowLength = SizeX;
-		CopyRegion.bufferImageHeight = SizeY;
-		CopyRegion.imageSubresource.aspectMask = SrcSurface.GetFullAspectMask();
-		CopyRegion.imageSubresource.mipLevel = MipLevel;
-		CopyRegion.imageSubresource.baseArrayLayer = 0;
-		CopyRegion.imageSubresource.layerCount = 1;
-		CopyRegion.imageExtent.width = SizeX;
-		CopyRegion.imageExtent.height = SizeY;
-		CopyRegion.imageExtent.depth = 1;
-		VulkanRHI::vkCmdCopyImageToBuffer(CmdBuffer->GetHandle(), SrcSurface.Image, VK_IMAGE_LAYOUT_TRANSFER_SRC_OPTIMAL, CpuReadbackBuffer->Buffer, 1, &CopyRegion);
-
-		{
-			FVulkanPipelineBarrier BarrierAfter;
-			BarrierAfter.AddMemoryBarrier(VK_ACCESS_TRANSFER_WRITE_BIT, VK_ACCESS_HOST_READ_BIT, VK_PIPELINE_STAGE_ALL_COMMANDS_BIT, VK_PIPELINE_STAGE_HOST_BIT);
-
-			if (InResolveParams.SourceAccessFinal != ERHIAccess::CopySrc && InResolveParams.SourceAccessFinal != ERHIAccess::Unknown)
-			{
-				BarrierAfter.AddImageAccessTransition(SrcSurface, ERHIAccess::CopySrc, InResolveParams.SourceAccessFinal, FVulkanPipelineBarrier::MakeSubresourceRange(SrcSurface.GetFullAspectMask()), SrcLayout);
-			}
-			CmdBuffer->GetLayoutManager().SetFullLayout(SrcSurface, SrcLayout);
-
-			BarrierAfter.Execute(CmdBuffer);
-		}
-	}
-	else
-	{
-		VkImageLayout DstLayout = FVulkanLayoutManager::SetExpectedLayout(CmdBuffer, DstSurface, ERHIAccess::CopyDest);
-
-		if (SrcSurface.Image != DstSurface.Image)
-		{
-			const bool bIsResolve = SrcSurface.GetNumSamples() > DstSurface.GetNumSamples();
-			checkf(!bIsResolve || !DstSurface.IsDepthOrStencilAspect(), TEXT("Vulkan does not support multisample depth resolve."));
-
-			FVulkanPipelineBarrier BarrierBefore;
-		
-			// We'll transition the entire resources to the correct copy states, so we don't need to worry about sub-resource states.
-			if (DstLayout != VK_IMAGE_LAYOUT_TRANSFER_DST_OPTIMAL)
-			{
-				BarrierBefore.AddImageLayoutTransition(DstSurface.Image, DstLayout, VK_IMAGE_LAYOUT_TRANSFER_DST_OPTIMAL, FVulkanPipelineBarrier::MakeSubresourceRange(DstSurface.GetFullAspectMask()));
-				DstLayout = VK_IMAGE_LAYOUT_TRANSFER_DST_OPTIMAL;
-			}
-
-			BarrierBefore.Execute(CmdBuffer);
-
-			if (!bIsResolve)
-			{
-				VkImageCopy Region;
-				SetupCopyOrResolveRegion(Region, SrcSurface, DstSurface, SrcRange, DstRange, InResolveParams);
-				VulkanRHI::vkCmdCopyImage(CmdBuffer->GetHandle(),
-					SrcSurface.Image, VK_IMAGE_LAYOUT_TRANSFER_SRC_OPTIMAL,
-					DstSurface.Image, VK_IMAGE_LAYOUT_TRANSFER_DST_OPTIMAL,
-					1, &Region);
-			}
-			else
-			{
-				check(DstSurface.GetNumSamples() == 1);
-				VkImageResolve Region;
-				SetupCopyOrResolveRegion(Region, SrcSurface, DstSurface, SrcRange, DstRange, InResolveParams);
-				VulkanRHI::vkCmdResolveImage(CmdBuffer->GetHandle(),
-					SrcSurface.Image, VK_IMAGE_LAYOUT_TRANSFER_SRC_OPTIMAL,
-					DstSurface.Image, VK_IMAGE_LAYOUT_TRANSFER_DST_OPTIMAL,
-					1, &Region);
-			}
-		}
-
-		if ((InResolveParams.SourceAccessFinal != ERHIAccess::Unknown) || (InResolveParams.DestAccessFinal != ERHIAccess::Unknown))
-		{
-			FVulkanPipelineBarrier BarrierAfter;
-
-			if ((SrcSurface.Image != DstSurface.Image) && (InResolveParams.SourceAccessFinal != ERHIAccess::CopySrc) && (InResolveParams.SourceAccessFinal != ERHIAccess::Unknown))
-			{
-				BarrierAfter.AddImageAccessTransition(SrcSurface, ERHIAccess::CopySrc, InResolveParams.SourceAccessFinal, FVulkanPipelineBarrier::MakeSubresourceRange(SrcSurface.GetFullAspectMask()), SrcLayout);
-				CmdBuffer->GetLayoutManager().SetFullLayout(SrcSurface, SrcLayout);
-			}
-
-			if ((InResolveParams.DestAccessFinal != ERHIAccess::CopyDest) && (InResolveParams.DestAccessFinal != ERHIAccess::Unknown))
-			{
-				BarrierAfter.AddImageAccessTransition(DstSurface, ERHIAccess::CopyDest, InResolveParams.DestAccessFinal, FVulkanPipelineBarrier::MakeSubresourceRange(DstSurface.GetFullAspectMask()), DstLayout);
-				CmdBuffer->GetLayoutManager().SetFullLayout(DstSurface, DstLayout);
-			}
-
-			BarrierAfter.Execute(CmdBuffer);
-		}
-	}
-
-	if (GSubmitOnCopyToResolve)
-	{
-		InternalSubmitActiveCmdBuffer();
-	}
-}
-
-
 static void ConvertRawDataToFColor(VkFormat VulkanFormat, uint32 DestWidth, uint32 DestHeight, uint8* In, uint32 SrcPitch, FColor* Dest, const FReadSurfaceDataFlags& InFlags)
 {
 	const bool bLinearToGamma = InFlags.GetLinearToGamma();
diff --git a/Engine/Source/Runtime/Windows/D3D11RHI/Private/D3D11RHIPrivate.h b/Engine/Source/Runtime/Windows/D3D11RHI/Private/D3D11RHIPrivate.h
--- a/Engine/Source/Runtime/Windows/D3D11RHI/Private/D3D11RHIPrivate.h
+++ b/Engine/Source/Runtime/Windows/D3D11RHI/Private/D3D11RHIPrivate.h
@@ -562,7 +562,6 @@ public:
 	virtual void RHISetMultipleViewports(uint32 Count, const FViewportBounds* Data) final override;
 	virtual void RHIClearUAVFloat(FRHIUnorderedAccessView* UnorderedAccessViewRHI, const FVector4f& Values) final override;
 	virtual void RHIClearUAVUint(FRHIUnorderedAccessView* UnorderedAccessViewRHI, const FUintVector4& Values) final override;
-	virtual void RHICopyToResolveTarget(FRHITexture* SourceTexture, FRHITexture* DestTexture, const FResolveParams& ResolveParams) final override;
 	virtual void RHICopyTexture(FRHITexture* SourceTexture, FRHITexture* DestTexture, const FRHICopyTextureInfo& CopyInfo) final override;
 	virtual void RHICopyBufferRegion(FRHIBuffer* DestBuffer, uint64 DstOffset, FRHIBuffer* SourceBuffer, uint64 SrcOffset, uint64 NumBytes) final override;
 	virtual void RHICreateTransition(FRHITransition* Transition, const FRHITransitionCreateInfo& CreateInfo) final override;
diff --git a/Engine/Source/Runtime/Windows/D3D11RHI/Private/D3D11RenderTarget.cpp b/Engine/Source/Runtime/Windows/D3D11RHI/Private/D3D11RenderTarget.cpp
--- a/Engine/Source/Runtime/Windows/D3D11RHI/Private/D3D11RenderTarget.cpp
+++ b/Engine/Source/Runtime/Windows/D3D11RHI/Private/D3D11RenderTarget.cpp
@@ -187,161 +187,6 @@ void FD3D11DynamicRHI::ResolveTextureUsingShader(
 	This->CurrentDepthTexture = OriginalDepthTexture;
 }
 
-/**
-* Copies the contents of the given surface to its resolve target texture.
-* @param SourceSurface - surface with a resolve texture to copy to
-* @param bKeepOriginalSurface - true if the original surface will still be used after this function so must remain valid
-* @param ResolveParams - optional resolve params
-*/
-void FD3D11DynamicRHI::RHICopyToResolveTarget(FRHITexture* SourceTextureRHI, FRHITexture* DestTextureRHI, const FResolveParams& ResolveParams)
-{
-	if (!SourceTextureRHI || !DestTextureRHI)
-	{
-		// no need to do anything (silently ignored)
-		return;
-	}
-
-	FD3D11Texture* SourceTexture2D   = ResourceCast(SourceTextureRHI->GetTexture2D());
-	FD3D11Texture* DestTexture2D     = ResourceCast(DestTextureRHI->GetTexture2D());
-
-	FD3D11Texture* SourceTextureCube = ResourceCast(SourceTextureRHI->GetTextureCube());
-	FD3D11Texture* DestTextureCube   = ResourceCast(DestTextureRHI->GetTextureCube());
-
-	FD3D11Texture* SourceTexture3D   = ResourceCast(SourceTextureRHI->GetTexture3D());
-	FD3D11Texture* DestTexture3D     = ResourceCast(DestTextureRHI->GetTexture3D());
-		
-	if(SourceTexture2D && DestTexture2D)
-	{
-		check(!SourceTextureCube && !DestTextureCube);
-		if(SourceTexture2D != DestTexture2D)
-		{
-			GPUProfilingData.RegisterGPUWork();
-		
-			if(DestTexture2D->GetDepthStencilView(FExclusiveDepthStencil::DepthWrite_StencilWrite)
-				&& SourceTextureRHI->IsMultisampled()
-				&& !DestTextureRHI->IsMultisampled())
-			{
-				D3D11_TEXTURE2D_DESC ResolveTargetDesc;
-				DestTexture2D->GetD3D11Texture2D()->GetDesc(&ResolveTargetDesc);
-
-				ResolveTextureUsingShader<FResolveDepthPS>(
-					this,
-					SourceTexture2D,
-					DestTexture2D,
-					DestTexture2D->GetRenderTargetView(0, -1),
-					DestTexture2D->GetDepthStencilView(FExclusiveDepthStencil::DepthWrite_StencilWrite),
-					ResolveTargetDesc,
-					GetDefaultRect(ResolveParams.Rect,DestTexture2D->GetSizeX(),DestTexture2D->GetSizeY()),
-					GetDefaultRect(ResolveParams.Rect,DestTexture2D->GetSizeX(),DestTexture2D->GetSizeY()),
-					FDummyResolveParameter()
-					);
-			}
-			else
-			{
-				DXGI_FORMAT SrcFmt = (DXGI_FORMAT)GPixelFormats[SourceTextureRHI->GetFormat()].PlatformFormat;
-				DXGI_FORMAT DstFmt = (DXGI_FORMAT)GPixelFormats[DestTexture2D->GetFormat()].PlatformFormat;
-				
-				DXGI_FORMAT Fmt = ConvertTypelessToUnorm((DXGI_FORMAT)GPixelFormats[DestTexture2D->GetFormat()].PlatformFormat);
-
-				// Determine whether a MSAA resolve is needed, or just a copy.
-				if(SourceTextureRHI->IsMultisampled() && !DestTexture2D->IsMultisampled())
-				{
-					Direct3DDeviceIMContext->ResolveSubresource(
-						DestTexture2D->GetResource(),
-						ResolveParams.DestArrayIndex,
-						SourceTexture2D->GetResource(),
-						ResolveParams.SourceArrayIndex,
-						Fmt
-						);
-				}
-				else
-				{
-					if(ResolveParams.Rect.IsValid() 
-						&& !SourceTextureRHI->IsMultisampled()
-						&& !DestTexture2D->GetDepthStencilView(FExclusiveDepthStencil::DepthWrite_StencilWrite))
-					{
-						D3D11_BOX SrcBox;
-
-						SrcBox.left = ResolveParams.Rect.X1;
-						SrcBox.top = ResolveParams.Rect.Y1;
-						SrcBox.front = 0;
-						SrcBox.right = ResolveParams.Rect.X2;
-						SrcBox.bottom = ResolveParams.Rect.Y2;
-						SrcBox.back = 1;
-
-						const FResolveRect& DestRect = ResolveParams.DestRect.IsValid() ? ResolveParams.DestRect : ResolveParams.Rect;
-						Direct3DDeviceIMContext->CopySubresourceRegion(DestTexture2D->GetResource(), ResolveParams.DestArrayIndex, DestRect.X1, DestRect.Y1, 0, SourceTexture2D->GetResource(), ResolveParams.SourceArrayIndex, &SrcBox);
-					}
-					else
-					{
-						Direct3DDeviceIMContext->CopyResource(DestTexture2D->GetResource(), SourceTexture2D->GetResource());
-					}
-				}
-			}
-		}
-	}
-	else if(SourceTextureCube && DestTextureCube)
-	{
-		check(!SourceTexture2D && !DestTexture2D);			
-
-		if(SourceTextureCube != DestTextureCube)
-		{
-			GPUProfilingData.RegisterGPUWork();
-
-			// Determine the cubemap face being resolved.
-			const uint32 D3DFace = GetD3D11CubeFace(ResolveParams.CubeFace);
-			const uint32 SourceSubresource = D3D11CalcSubresource(ResolveParams.MipIndex, ResolveParams.SourceArrayIndex * 6 + D3DFace, SourceTextureCube->GetNumMips());
-			const uint32 DestSubresource = D3D11CalcSubresource(ResolveParams.MipIndex, ResolveParams.DestArrayIndex * 6 + D3DFace, DestTextureCube->GetNumMips());
-
-			// Determine whether a MSAA resolve is needed, or just a copy.
-			if(SourceTextureRHI->IsMultisampled() && !DestTextureCube->IsMultisampled())
-			{
-				Direct3DDeviceIMContext->ResolveSubresource(
-					DestTextureCube->GetResource(),
-					DestSubresource,
-					SourceTextureCube->GetResource(),
-					SourceSubresource,
-					(DXGI_FORMAT)GPixelFormats[DestTextureCube->GetFormat()].PlatformFormat
-					);
-			}
-			else
-			{
-				if (ResolveParams.Rect.IsValid())
-				{
-					D3D11_BOX SrcBox;
-
-					SrcBox.left = ResolveParams.Rect.X1;
-					SrcBox.top = ResolveParams.Rect.Y1;
-					SrcBox.front = 0;
-					SrcBox.right = ResolveParams.Rect.X2;
-					SrcBox.bottom = ResolveParams.Rect.Y2;
-					SrcBox.back = 1;
-
-					Direct3DDeviceIMContext->CopySubresourceRegion(DestTextureCube->GetResource(), DestSubresource, 0, 0, 0, SourceTextureCube->GetResource(), SourceSubresource, &SrcBox);
-				}
-				else
-				{
-					Direct3DDeviceIMContext->CopySubresourceRegion(DestTextureCube->GetResource(), DestSubresource, 0, 0, 0, SourceTextureCube->GetResource(), SourceSubresource, NULL);
-				}
-			}
-		}
-	}
-	else if(SourceTexture2D && DestTextureCube)
-	{
-		// If source is 2D and Dest is a cube then copy the 2D texture to the specified cube face.
-		// Determine the cubemap face being resolved.
-		const uint32 D3DFace = GetD3D11CubeFace(ResolveParams.CubeFace);
-		const uint32 Subresource = D3D11CalcSubresource(0, D3DFace, 1);
-		Direct3DDeviceIMContext->CopySubresourceRegion(DestTextureCube->GetResource(), Subresource, 0, 0, 0, SourceTexture2D->GetResource(), 0, NULL);
-	}
-	else if (SourceTexture3D && DestTexture3D)
-	{
-		// bit of a hack.  no one resolves slice by slice and 0 is the default value.  assume for the moment they are resolving the whole texture.
-		check(ResolveParams.SourceArrayIndex == 0);
-		check(SourceTexture3D == DestTexture3D);
-	}
-}
-
 // Only supports the formats that are supported by ConvertRAWSurfaceDataToFColor()
 static uint32 ComputeBytesPerPixel(DXGI_FORMAT Format)
 {
