Multi-GPU Upload Heap check fix:
Structured upload buffer (Upload Heap) released via RHIBufferRef. When it gets down to releasing d3d12 resource in ::RHIPerFrameRHIFlushComplete()  - EnqueueEndOfPipeTaskit sometimes fails a check check(ObjectToDelete.RHIObject->GetRefCount() == 1);
This CL forces the buffer to lose the references kept by multiple gpu nodes before being released.

diff --git a/Engine/Source/Runtime/D3D12RHI/Private/D3D12Resources.cpp b/Engine/Source/Runtime/D3D12RHI/Private/D3D12Resources.cpp
--- a/Engine/Source/Runtime/D3D12RHI/Private/D3D12Resources.cpp
+++ b/Engine/Source/Runtime/D3D12RHI/Private/D3D12Resources.cpp
@@ -859,16 +859,26 @@ void FD3D12ResourceLocation::ReleaseResource()
 		bool bIncrement = false;
 		UpdateStandAloneStats(bIncrement);
 
-		// Multi-GPU support : because of references, several GPU nodes can refrence the same stand-alone resource.
-		check(UnderlyingResource->GetRefCount() == 1 || GNumExplicitGPUsForRendering > 1);
-		
-		if (UnderlyingResource->ShouldDeferDelete())
+		// Multi-GPU support : When the resource enters this point for the first time the number of references should be the same as number of GPUs.
+		// Shouldn't queue deferred deletion until all references are released as this could cause issues at the end of the pipe.
+		// Instead reduce number of references until nothing else holds the resource.
+		if (GNumExplicitGPUsForRendering > 1 && UnderlyingResource->GetRefCount() > 1)
 		{
-			UnderlyingResource->DeferDelete();
+			check(UnderlyingResource->GetRefCount() <= GNumExplicitGPUsForRendering)
+			UnderlyingResource->Release();
 		}
 		else
 		{
-			UnderlyingResource->Release();
+			check(UnderlyingResource->GetRefCount() == 1);
+
+			if (UnderlyingResource->ShouldDeferDelete())
+			{
+				UnderlyingResource->DeferDelete();
+			}
+			else
+			{
+				UnderlyingResource->Release();
+			}
 		}
 		break;
 	}
