Moving a number of D3D12RHI headers from Public to Private.
Cleaning up a comment in D3D12ShaderResources.h
Touching build script so that UBT is guaranteed to pick up changes.

diff --git a/Engine/Source/Runtime/D3D12RHI/D3D12RHI.Build.cs b/Engine/Source/Runtime/D3D12RHI/D3D12RHI.Build.cs
--- a/Engine/Source/Runtime/D3D12RHI/D3D12RHI.Build.cs
+++ b/Engine/Source/Runtime/D3D12RHI/D3D12RHI.Build.cs
@@ -24,10 +24,6 @@ public class D3D12RHI : ModuleRules
 
 		PublicIncludePathModuleNames.Add("HeadMountedDisplay");
 
-		///////////////////////////////////////////////////////////////
-        // Platform specific defines
-        ///////////////////////////////////////////////////////////////
-
 		if (!Target.Platform.IsInGroup(UnrealPlatformGroup.Windows))
         {
             PrecompileForTargets = PrecompileTargetsType.None;
diff --git a/Engine/Source/Runtime/D3D12RHI/Private/D3D12ConstantBuffer.h b/Engine/Source/Runtime/D3D12RHI/Private/D3D12ConstantBuffer.h
new file mode 100644
--- /dev/null
+++ b/Engine/Source/Runtime/D3D12RHI/Private/D3D12ConstantBuffer.h
@@ -0,0 +1,75 @@
+// Copyright Epic Games, Inc. All Rights Reserved.
+
+#pragma once
+
+// Chunk size for global constant buffer shadow data. Should always be a multiple of 256
+#define GLOBAL_CONSTANT_BUFFER_CHUNK_SIZE 512
+
+using namespace D3D12RHI;
+
+class FD3D12FastConstantAllocator;
+
+/**
+ * A D3D constant buffer
+ */
+class FD3D12ConstantBuffer : public FD3D12DeviceChild
+{
+public:
+	FD3D12ConstantBuffer(FD3D12Device* InParent, FD3D12FastConstantAllocator& InAllocator);
+	virtual ~FD3D12ConstantBuffer();
+
+	/**
+	* Updates a variable in the constant buffer.
+	* @param Data - The data to copy into the constant buffer
+	* @param Offset - The offset in the constant buffer to place the data at
+	* @param InSize - The size of the data being copied
+	*/
+	FORCEINLINE_DEBUGGABLE void UpdateConstant(const uint8* Data, uint16 Offset, uint16 InSize)
+	{
+		// Make sure we have enough memory in our shadow data
+		const uint32 RequiredSize = (uint32)Offset + (uint32)InSize;
+		if ((uint32)ShadowData.Num() < RequiredSize)
+		{
+			// Make sure we grow in chunks to prevent too many reallocs
+			const uint32 SizeToGrowTo = Align(RequiredSize, GLOBAL_CONSTANT_BUFFER_CHUNK_SIZE);
+			ShadowData.SetNumZeroed(SizeToGrowTo);
+		}
+
+		FMemory::Memcpy(ShadowData.GetData() + Offset, Data, InSize);
+		CurrentUpdateSize = FMath::Max(RequiredSize, CurrentUpdateSize);
+
+		bIsDirty = true;
+	}
+
+	FORCEINLINE void Reset() { CurrentUpdateSize = 0; }
+
+	bool Version(FD3D12ResourceLocation& BufferOut, bool bDiscardSharedConstants);
+
+#if USE_STATIC_ROOT_SIGNATURE
+	inline D3D12_CPU_DESCRIPTOR_HANDLE GetOfflineCpuHandle() const { return View->GetOfflineCpuHandle(); }
+#endif
+
+protected:
+#if USE_STATIC_ROOT_SIGNATURE
+	FD3D12ConstantBufferView* View = nullptr;
+#endif
+
+	TArray<uint8> ShadowData;
+	
+	/** Size of all constants that has been updated since the last call to Commit. */
+	uint32	CurrentUpdateSize;
+
+	/**
+	 * Size of all constants that has been updated since the last Discard.
+	 * Includes "shared" constants that don't necessarily gets updated between every Commit.
+	 */
+	uint32	TotalUpdateSize;
+	
+	// Indicates that a constant has been updated but this one hasn't been flushed.
+	bool bIsDirty;
+
+	FD3D12FastConstantAllocator& Allocator;
+};
+
+DECLARE_CYCLE_STAT_EXTERN(TEXT("Global Constant buffer update time"), STAT_D3D12GlobalConstantBufferUpdateTime, STATGROUP_D3D12RHI, );
+
diff --git a/Engine/Source/Runtime/D3D12RHI/Private/D3D12Resources.h b/Engine/Source/Runtime/D3D12RHI/Private/D3D12Resources.h
new file mode 100644
--- /dev/null
+++ b/Engine/Source/Runtime/D3D12RHI/Private/D3D12Resources.h
@@ -0,0 +1,1057 @@
+// Copyright Epic Games, Inc. All Rights Reserved.
+
+/*=============================================================================
+	D3D12Resources.h: D3D resource RHI definitions.
+	=============================================================================*/
+
+#pragma once
+
+#include "D3D12RHIPrivate.h"
+#include "BoundShaderStateCache.h"
+#include "D3D12ShaderResources.h"
+#include "D3D12Residency.h"
+#include "D3D12Util.h"
+#include "D3D12State.h"
+#include "D3D12DirectCommandListManager.h"
+#include "RHIPoolAllocator.h"
+#include "Templates/UniquePtr.h"
+
+constexpr D3D12_RESOURCE_STATES BackBufferBarrierWriteTransitionTargets = D3D12_RESOURCE_STATES(
+	uint32(D3D12_RESOURCE_STATE_RENDER_TARGET) |
+	uint32(D3D12_RESOURCE_STATE_UNORDERED_ACCESS) |
+	uint32(D3D12_RESOURCE_STATE_STREAM_OUT) |
+	uint32(D3D12_RESOURCE_STATE_COPY_DEST) |
+	uint32(D3D12_RESOURCE_STATE_RESOLVE_DEST));
+
+// Forward Decls
+class FD3D12Resource;
+class FD3D12StateCache;
+class FD3D12CommandListManager;
+class FD3D12CommandContext;
+class FD3D12SegListAllocator;
+class FD3D12PoolAllocator;
+struct FD3D12GraphicsPipelineState;
+struct FD3D12ResourceDesc;
+
+#if D3D12_RHI_RAYTRACING
+class FD3D12RayTracingGeometry;
+class FD3D12RayTracingScene;
+class FD3D12RayTracingPipelineState;
+class FD3D12RayTracingShader;
+#endif // D3D12_RHI_RAYTRACING
+
+#ifndef D3D12_WITH_CUSTOM_TEXTURE_LAYOUT
+#define D3D12_WITH_CUSTOM_TEXTURE_LAYOUT 0
+#endif
+
+#if D3D12_WITH_CUSTOM_TEXTURE_LAYOUT
+extern void ApplyCustomTextureLayout(FD3D12ResourceDesc& TextureLayout, FD3D12Adapter& Adapter);
+#endif
+
+enum class ED3D12ResourceStateMode
+{
+	Default,					//< Decide if tracking is required based on flags
+	SingleState,				//< Force disable state tracking of resource - resource will always be in the initial resource state
+	MultiState,					//< Force enable state tracking of resource
+};
+
+class FD3D12PendingResourceBarrier
+{
+public:
+	FD3D12Resource*       Resource;
+	D3D12_RESOURCE_STATES State;
+	uint32                SubResource;
+
+	FD3D12PendingResourceBarrier(FD3D12Resource* Resource, D3D12_RESOURCE_STATES State, uint32 SubResource)
+		: Resource(Resource)
+		, State(State)
+		, SubResource(SubResource)
+	{}
+};
+
+class FD3D12Heap : public FThreadSafeRefCountedObject, public FD3D12DeviceChild, public FD3D12MultiNodeGPUObject
+{
+public:
+	FD3D12Heap(FD3D12Device* Parent, FRHIGPUMask VisibleNodes, HeapId InTraceParentHeapId = EMemoryTraceRootHeap::VideoMemory);
+	~FD3D12Heap();
+
+	inline ID3D12Heap* GetHeap() const { return Heap.GetReference(); }
+	void SetHeap(ID3D12Heap* HeapIn, const TCHAR* const InName, bool bTrack = true, bool bForceGetGPUAddress = false);
+
+	void BeginTrackingResidency(uint64 Size);
+
+	void DeferDelete();
+
+	inline FName GetName() const { return HeapName; }
+	inline D3D12_HEAP_DESC GetHeapDesc() const { return HeapDesc; }
+	inline FD3D12ResidencyHandle& GetResidencyHandle() { return ResidencyHandle; }
+	inline D3D12_GPU_VIRTUAL_ADDRESS GetGPUVirtualAddress() const { return GPUVirtualAddress; }
+	inline void SetIsTransient(bool bInIsTransient) { bIsTransient = bInIsTransient; }
+	inline bool GetIsTransient() const { return bIsTransient; }
+
+private:
+
+	TRefCountPtr<ID3D12Heap> Heap;
+	FName HeapName;
+	bool bTrack = true;
+	D3D12_HEAP_DESC HeapDesc;
+	D3D12_GPU_VIRTUAL_ADDRESS GPUVirtualAddress = 0;
+	FD3D12ResidencyHandle ResidencyHandle;
+	HeapId TraceHeapId;
+	HeapId TraceParentHeapId;
+	bool bIsTransient = false; // Whether this is a transient heap
+};
+
+struct FD3D12ResourceDesc : public D3D12_RESOURCE_DESC
+{
+	FD3D12ResourceDesc() = default;
+	FD3D12ResourceDesc(const CD3DX12_RESOURCE_DESC& Other)
+		: D3D12_RESOURCE_DESC(Other)
+	{
+	}
+	
+	// TODO: use this type everywhere and disallow implicit conversion
+	/*explicit*/ FD3D12ResourceDesc(const D3D12_RESOURCE_DESC& Other)
+		: D3D12_RESOURCE_DESC(Other)
+	{
+	}
+
+	EPixelFormat PixelFormat{ PF_Unknown };
+
+	// PixelFormat for the Resource that aliases our current resource.
+	EPixelFormat UAVAliasPixelFormat{ PF_Unknown };
+
+#if D3D12RHI_NEEDS_VENDOR_EXTENSIONS
+	bool bRequires64BitAtomicSupport{ false };
+#endif
+
+	bool bReservedResource { false };
+
+	// Used primarily to help treat this resource description as writable.
+	inline bool NeedsUAVAliasWorkarounds() const { return UAVAliasPixelFormat != PF_Unknown; }
+};
+
+class FD3D12Resource : public FThreadSafeRefCountedObject, public FD3D12DeviceChild, public FD3D12MultiNodeGPUObject
+{
+private:
+	TRefCountPtr<ID3D12Resource> Resource;
+	// Since certain formats cannot be aliased in D3D12, we have to create a separate ID3D12Resource that aliases the
+	// resource's memory and use this separate resource to create the UAV.
+	// TODO: UE-116727: request better DX12 API support and clean this up.
+	TRefCountPtr<ID3D12Resource> UAVAccessResource;
+	TRefCountPtr<FD3D12Heap> Heap;
+
+	FD3D12ResidencyHandle ResidencyHandle;
+
+	D3D12_GPU_VIRTUAL_ADDRESS GPUVirtualAddress{};
+	void* ResourceBaseAddress{};
+
+#if NV_AFTERMATH
+	GFSDK_Aftermath_ResourceHandle AftermathHandle{};
+#endif
+
+	const FD3D12ResourceDesc Desc;
+	CResourceState ResourceState;
+	D3D12_RESOURCE_STATES DefaultResourceState{ D3D12_RESOURCE_STATE_TBD };
+	D3D12_RESOURCE_STATES ReadableState{ D3D12_RESOURCE_STATE_CORRUPT };
+	D3D12_RESOURCE_STATES WritableState{ D3D12_RESOURCE_STATE_CORRUPT };
+#ifdef PLATFORM_SUPPORTS_RESOURCE_COMPRESSION
+	D3D12_RESOURCE_STATES CompressedState{ D3D12_RESOURCE_STATE_CORRUPT };
+#endif
+
+	D3D12_HEAP_TYPE HeapType;
+	FName DebugName;
+
+	int32 NumMapCalls = 0;
+	uint16 SubresourceCount{};
+	uint8 PlaneCount;
+	bool bRequiresResourceStateTracking : 1;
+	bool bDepthStencil : 1;
+	bool bDeferDelete : 1;
+	bool bBackBuffer : 1;
+
+#if UE_BUILD_DEBUG
+	static int64 TotalResourceCount;
+	static int64 NoStateTrackingResourceCount;
+#endif
+
+	struct FD3D12ReservedResourceData
+	{
+		TArray<TRefCountPtr<ID3D12Heap>> BackingHeaps;
+	};
+	TUniquePtr<FD3D12ReservedResourceData> ReservedResourceData;
+
+public:
+	FD3D12Resource() = delete;
+	explicit FD3D12Resource(FD3D12Device* ParentDevice,
+		FRHIGPUMask VisibleNodes,
+		ID3D12Resource* InResource,
+		D3D12_RESOURCE_STATES InInitialResourceState,
+		FD3D12ResourceDesc const& InDesc,
+		FD3D12Heap* InHeap = nullptr,
+		D3D12_HEAP_TYPE InHeapType = D3D12_HEAP_TYPE_DEFAULT);
+
+	explicit FD3D12Resource(FD3D12Device* ParentDevice,
+		FRHIGPUMask VisibleNodes,
+		ID3D12Resource* InResource,
+		D3D12_RESOURCE_STATES InInitialResourceState,
+		ED3D12ResourceStateMode InResourceStateMode,
+		D3D12_RESOURCE_STATES InDefaultResourceState,
+		FD3D12ResourceDesc const& InDesc,
+		FD3D12Heap* InHeap,
+		D3D12_HEAP_TYPE InHeapType);
+
+	virtual ~FD3D12Resource();
+
+	operator ID3D12Resource&() { return *Resource; }
+
+	ID3D12Resource* GetResource() const { return Resource.GetReference(); }
+	ID3D12Resource* GetUAVAccessResource() const { return UAVAccessResource.GetReference(); }
+	void SetUAVAccessResource(ID3D12Resource* InUAVAccessResource) { UAVAccessResource = InUAVAccessResource; }
+
+	inline void* Map(const D3D12_RANGE* ReadRange = nullptr)
+	{
+		if (NumMapCalls == 0)
+		{
+			check(Resource);
+			check(ResourceBaseAddress == nullptr);
+			VERIFYD3D12RESULT(Resource->Map(0, ReadRange, &ResourceBaseAddress));
+		}
+		else
+		{
+			check(ResourceBaseAddress);
+		}
+		++NumMapCalls;
+
+		return ResourceBaseAddress;
+	}
+
+	inline void Unmap()
+	{
+		check(Resource);
+		check(ResourceBaseAddress);
+		check(NumMapCalls > 0);
+
+		--NumMapCalls;
+		if (NumMapCalls == 0)
+		{
+			Resource->Unmap(0, nullptr);
+			ResourceBaseAddress = nullptr;
+		}
+	}
+
+	ID3D12Pageable* GetPageable();
+	const FD3D12ResourceDesc& GetDesc() const { return Desc; }
+	D3D12_HEAP_TYPE GetHeapType() const { return HeapType; }
+	D3D12_GPU_VIRTUAL_ADDRESS GetGPUVirtualAddress() const { return GPUVirtualAddress; }
+	inline void SetGPUVirtualAddress(D3D12_GPU_VIRTUAL_ADDRESS Value) { GPUVirtualAddress = Value; }
+	void* GetResourceBaseAddress() const { check(ResourceBaseAddress); return ResourceBaseAddress; }
+	uint16 GetMipLevels() const { return Desc.MipLevels; }
+	uint16 GetArraySize() const { return (Desc.Dimension == D3D12_RESOURCE_DIMENSION_TEXTURE3D) ? 1 : Desc.DepthOrArraySize; }
+	uint8 GetPlaneCount() const { return PlaneCount; }
+	uint16 GetSubresourceCount() const { return SubresourceCount; }
+	CResourceState& GetResourceState_OnResource()
+	{
+		check(bRequiresResourceStateTracking);
+		// This state is used as the resource's "global" state between command lists. It's only needed for resources that
+		// require state tracking.
+		return ResourceState;
+	}
+	D3D12_RESOURCE_STATES GetDefaultResourceState() const { check(!bRequiresResourceStateTracking); return DefaultResourceState; }
+	D3D12_RESOURCE_STATES GetWritableState() const { return WritableState; }
+	D3D12_RESOURCE_STATES GetReadableState() const { return ReadableState; }
+#ifdef PLATFORM_SUPPORTS_RESOURCE_COMPRESSION
+	D3D12_RESOURCE_STATES GetCompressedState() const { return CompressedState; }
+	void SetCompressedState(D3D12_RESOURCE_STATES State) { CompressedState = State; }
+#endif
+	bool RequiresResourceStateTracking() const { return bRequiresResourceStateTracking; }
+
+	inline bool IsBackBuffer() const { return bBackBuffer; }
+	inline void SetIsBackBuffer(bool bBackBufferIn) { bBackBuffer = bBackBufferIn; }
+
+	void SetName(const TCHAR* Name)
+	{
+		// Check name before setting it.  Saves FName lookup and driver call.  Names are frequently the same for pooled buffers
+		// that end up getting reused for the same purpose every frame (2/3 of calls to this function on a given frame).
+		if (DebugName != Name)
+		{
+			DebugName = FName(Name);
+			::SetName(Resource, Name);
+		}
+	}
+
+	FName GetName() const
+	{
+		return DebugName;
+	}
+	
+	void DoNotDeferDelete()
+	{
+		bDeferDelete = false;
+	}
+
+	inline bool ShouldDeferDelete() const { return bDeferDelete; }
+	void DeferDelete();
+
+	inline bool IsPlacedResource() const { return Heap.GetReference() != nullptr; }
+	inline FD3D12Heap* GetHeap() const { return Heap; };
+	inline bool IsDepthStencilResource() const { return bDepthStencil; }
+
+	void StartTrackingForResidency();
+
+	FD3D12ResidencyHandle& GetResidencyHandle()
+	{
+		return IsPlacedResource() ? Heap->GetResidencyHandle() : ResidencyHandle;
+	}
+
+	struct FD3D12ResourceTypeHelper
+	{
+		FD3D12ResourceTypeHelper(const FD3D12ResourceDesc& Desc, D3D12_HEAP_TYPE HeapType) :
+			bSRV(!EnumHasAnyFlags(Desc.Flags, D3D12_RESOURCE_FLAG_DENY_SHADER_RESOURCE)),
+			bDSV(EnumHasAnyFlags(Desc.Flags, D3D12_RESOURCE_FLAG_ALLOW_DEPTH_STENCIL)),
+			bRTV(EnumHasAnyFlags(Desc.Flags, D3D12_RESOURCE_FLAG_ALLOW_RENDER_TARGET)),
+			bUAV(EnumHasAnyFlags(Desc.Flags, D3D12_RESOURCE_FLAG_ALLOW_UNORDERED_ACCESS) || Desc.NeedsUAVAliasWorkarounds()),
+			bWritable(bDSV || bRTV || bUAV),
+			bSRVOnly(bSRV && !bWritable),
+			bBuffer(Desc.Dimension == D3D12_RESOURCE_DIMENSION_BUFFER),
+			bReadBackResource(HeapType == D3D12_HEAP_TYPE_READBACK)
+		{}
+
+		const D3D12_RESOURCE_STATES GetOptimalInitialState(ERHIAccess InResourceState, bool bAccurateWriteableStates) const
+		{
+			// Ignore the requested resource state for non tracked resource because RHI will assume it's always in default resource 
+			// state then when a transition is required (will transition via scoped push/pop to requested state)
+			if (!bSRVOnly && InResourceState != ERHIAccess::Unknown && InResourceState != ERHIAccess::Discard)
+			{
+				bool bAsyncCompute = false;
+				return GetD3D12ResourceState(InResourceState, bAsyncCompute);
+			}
+			else
+			{
+				if (bSRVOnly)
+				{
+					return D3D12_RESOURCE_STATE_NON_PIXEL_SHADER_RESOURCE | D3D12_RESOURCE_STATE_PIXEL_SHADER_RESOURCE;
+				}
+				else if (bBuffer && !bUAV)
+				{
+					return (bReadBackResource) ? D3D12_RESOURCE_STATE_COPY_DEST : D3D12_RESOURCE_STATE_GENERIC_READ;
+				}
+				else if (bWritable)
+				{
+					if (bAccurateWriteableStates)
+					{
+						if (bDSV)
+						{
+							return D3D12_RESOURCE_STATE_DEPTH_WRITE;
+						}
+						else if (bRTV)
+						{
+							return D3D12_RESOURCE_STATE_RENDER_TARGET;
+						}
+						else if (bUAV)
+						{
+							return D3D12_RESOURCE_STATE_UNORDERED_ACCESS;
+						}
+					}
+					else
+					{
+						// This things require tracking anyway
+						return D3D12_RESOURCE_STATE_COMMON;
+					}
+				}
+
+				return D3D12_RESOURCE_STATE_COMMON;
+			}
+		}
+
+		const uint32 bSRV : 1;
+		const uint32 bDSV : 1;
+		const uint32 bRTV : 1;
+		const uint32 bUAV : 1;
+		const uint32 bWritable : 1;
+		const uint32 bSRVOnly : 1;
+		const uint32 bBuffer : 1;
+		const uint32 bReadBackResource : 1;
+	};
+
+	void CommitReservedResource();
+
+private:
+	void InitalizeResourceState(D3D12_RESOURCE_STATES InInitialState, ED3D12ResourceStateMode InResourceStateMode, D3D12_RESOURCE_STATES InDefaultState)
+	{
+		SubresourceCount = GetMipLevels() * GetArraySize() * GetPlaneCount();
+
+		if (InResourceStateMode == ED3D12ResourceStateMode::SingleState)
+		{
+			// make sure a valid default state is set
+			check(IsValidD3D12ResourceState(InDefaultState));
+
+#if UE_BUILD_DEBUG
+			FPlatformAtomics::InterlockedIncrement(&NoStateTrackingResourceCount);
+#endif
+			DefaultResourceState = InDefaultState;
+			WritableState = D3D12_RESOURCE_STATE_CORRUPT;
+			ReadableState = D3D12_RESOURCE_STATE_CORRUPT;
+			bRequiresResourceStateTracking = false;
+		}
+		else
+		{
+			DetermineResourceStates(InDefaultState, InResourceStateMode);
+		}
+
+		if (bRequiresResourceStateTracking)
+		{
+#if D3D12_RHI_RAYTRACING
+			// No state tracking for acceleration structures because they can't have another state
+			check(InDefaultState != D3D12_RESOURCE_STATE_RAYTRACING_ACCELERATION_STRUCTURE && InInitialState != D3D12_RESOURCE_STATE_RAYTRACING_ACCELERATION_STRUCTURE);
+#endif // D3D12_RHI_RAYTRACING
+
+			// Only a few resources (~1%) actually need resource state tracking
+			ResourceState.Initialize(SubresourceCount);
+			ResourceState.SetResourceState(InInitialState);
+		}
+	}
+
+	void DetermineResourceStates(D3D12_RESOURCE_STATES InDefaultState, ED3D12ResourceStateMode InResourceStateMode)
+	{
+		const FD3D12ResourceTypeHelper Type(Desc, HeapType);
+
+		bDepthStencil = Type.bDSV;
+
+#ifdef PLATFORM_SUPPORTS_RESOURCE_COMPRESSION
+		SetCompressedState(D3D12_RESOURCE_STATE_COMMON);
+#endif
+
+		if (Type.bWritable || InResourceStateMode == ED3D12ResourceStateMode::MultiState)
+		{
+			// Determine the resource's write/read states.
+			if (Type.bRTV)
+			{
+				// Note: The resource could also be used as a UAV however we don't store that writable state. UAV's are handled in a separate RHITransitionResources() specially for UAVs so we know the writeable state in that case should be UAV.
+				check(!Type.bDSV && !Type.bBuffer);
+				WritableState = D3D12_RESOURCE_STATE_RENDER_TARGET;
+				ReadableState = Type.bSRV ? D3D12_RESOURCE_STATE_NON_PIXEL_SHADER_RESOURCE | D3D12_RESOURCE_STATE_PIXEL_SHADER_RESOURCE : D3D12_RESOURCE_STATE_CORRUPT;
+			}
+			else if (Type.bDSV)
+			{
+				check(!Type.bRTV && (!Type.bUAV || GRHISupportsDepthUAV) && !Type.bBuffer);
+				WritableState = D3D12_RESOURCE_STATE_DEPTH_WRITE;
+				ReadableState = Type.bSRV ? D3D12_RESOURCE_STATE_DEPTH_READ | D3D12_RESOURCE_STATE_NON_PIXEL_SHADER_RESOURCE | D3D12_RESOURCE_STATE_PIXEL_SHADER_RESOURCE : D3D12_RESOURCE_STATE_DEPTH_READ;
+			}
+			else
+			{
+				WritableState = Type.bUAV ? D3D12_RESOURCE_STATE_UNORDERED_ACCESS : D3D12_RESOURCE_STATE_CORRUPT;
+				ReadableState = Type.bSRV ? D3D12_RESOURCE_STATE_NON_PIXEL_SHADER_RESOURCE | D3D12_RESOURCE_STATE_PIXEL_SHADER_RESOURCE : D3D12_RESOURCE_STATE_CORRUPT;
+			}
+		}
+		else
+		{
+			bRequiresResourceStateTracking = false;
+						
+#if UE_BUILD_DEBUG
+			FPlatformAtomics::InterlockedIncrement(&NoStateTrackingResourceCount);
+#endif
+			if (InDefaultState != D3D12_RESOURCE_STATE_TBD)
+			{
+				DefaultResourceState = InDefaultState;
+			}
+			else if (Type.bBuffer)
+			{
+				DefaultResourceState = (HeapType == D3D12_HEAP_TYPE_READBACK) ? D3D12_RESOURCE_STATE_COPY_DEST : D3D12_RESOURCE_STATE_GENERIC_READ;
+			}
+			else
+			{
+				check(Type.bSRVOnly);
+				DefaultResourceState = D3D12_RESOURCE_STATE_NON_PIXEL_SHADER_RESOURCE | D3D12_RESOURCE_STATE_PIXEL_SHADER_RESOURCE;
+			}
+		}
+	}
+};
+
+typedef class FD3D12BuddyAllocator FD3D12BaseAllocatorType;
+
+struct FD3D12BuddyAllocatorPrivateData
+{
+	uint32 Offset;
+	uint32 Order;
+
+	void Init()
+	{
+		Offset = 0;
+		Order = 0;
+	}
+};
+
+struct FD3D12BlockAllocatorPrivateData
+{
+	uint64 FrameFence;
+	uint32 BucketIndex;
+	uint32 Offset;
+	FD3D12Resource* ResourceHeap;
+
+	void Init()
+	{
+		FrameFence = 0;
+		BucketIndex = 0;
+		Offset = 0;
+		ResourceHeap = nullptr;
+	}
+};
+
+struct FD3D12SegListAllocatorPrivateData
+{
+	uint32 Offset;
+
+	void Init()
+	{
+		Offset = 0;
+	}
+};
+
+struct FD3D12PoolAllocatorPrivateData
+{
+	FRHIPoolAllocationData PoolData;
+
+	void Init()
+	{
+		PoolData.Reset();
+	}
+};
+
+class FD3D12ResourceAllocator;
+class FD3D12BaseShaderResource;
+
+// A very light-weight and cache friendly way of accessing a GPU resource
+class FD3D12ResourceLocation : public FRHIPoolResource, public FD3D12DeviceChild, public FNoncopyable
+{
+public:
+	enum class ResourceLocationType : uint8
+	{
+		eUndefined,
+		eStandAlone,
+		eSubAllocation,
+		eFastAllocation,
+		eMultiFrameFastAllocation,
+		eAliased, // XR HMDs are the only use cases
+		eNodeReference,
+		eHeapAliased, 
+	};
+
+	enum EAllocatorType : uint8
+	{
+		AT_Default, // FD3D12BaseAllocatorType
+		AT_SegList, // FD3D12SegListAllocator
+		AT_Pool,	// FD3D12PoolAllocator
+		AT_Unknown = 0xff
+	};
+
+	FD3D12ResourceLocation(FD3D12Device* Parent);
+	~FD3D12ResourceLocation();
+
+	void Clear();
+
+	// Transfers the contents of 1 resource location to another, destroying the original but preserving the underlying resource 
+	static void TransferOwnership(FD3D12ResourceLocation& Destination, FD3D12ResourceLocation& Source);
+
+	// Setters
+	void SetOwner   (FD3D12BaseShaderResource* InOwner) { Owner = InOwner; }
+	void SetType    (ResourceLocationType Value)        { Type = Value;}
+	void SetResource(FD3D12Resource* Value);
+
+	void SetAllocator       (FD3D12BaseAllocatorType* Value) { Allocator        = Value; AllocatorType = AT_Default; }
+	void SetSegListAllocator(FD3D12SegListAllocator*  Value) { SegListAllocator = Value; AllocatorType = AT_SegList; }
+	void SetPoolAllocator   (FD3D12PoolAllocator*     Value) { PoolAllocator    = Value; AllocatorType = AT_Pool;    }
+
+	void ClearAllocator() { Allocator = nullptr; AllocatorType = AT_Unknown; }
+
+	void SetMappedBaseAddress       (void* Value                    ) { MappedBaseAddress = Value; }
+	void SetGPUVirtualAddress       (D3D12_GPU_VIRTUAL_ADDRESS Value) { GPUVirtualAddress = Value; }
+	void SetOffsetFromBaseOfResource(uint64 Value                   ) { OffsetFromBaseOfResource = Value; }
+	void SetSize                    (uint64 Value                   ) { Size = Value; }
+
+	// Getters
+	ResourceLocationType               GetType                       () const { return Type;                                                 }
+	EAllocatorType                     GetAllocatorType              () const { return AllocatorType;                                        }
+	FD3D12BaseAllocatorType*           GetAllocator                  ()       { check(AT_Default == AllocatorType); return Allocator;        }
+	FD3D12SegListAllocator*            GetSegListAllocator           ()       { check(AT_SegList == AllocatorType); return SegListAllocator; }
+	FD3D12PoolAllocator*               GetPoolAllocator              ()       { check(AT_Pool == AllocatorType);    return PoolAllocator;    }
+	FD3D12Resource*                    GetResource                   () const { return UnderlyingResource;                                   }
+	void*                              GetMappedBaseAddress          () const { return MappedBaseAddress;                                    }
+	D3D12_GPU_VIRTUAL_ADDRESS          GetGPUVirtualAddress          () const { return GPUVirtualAddress;                                    }
+	uint64                             GetOffsetFromBaseOfResource   () const { return OffsetFromBaseOfResource;                             }
+	uint64                             GetSize                       () const { return Size;                                                 }
+	FD3D12ResidencyHandle&             GetResidencyHandle            ()       { check(ResidencyHandle); return *ResidencyHandle;             }
+	FD3D12BuddyAllocatorPrivateData&   GetBuddyAllocatorPrivateData  ()       { return AllocatorData.BuddyAllocatorPrivateData;              }
+	FD3D12BlockAllocatorPrivateData&   GetBlockAllocatorPrivateData  ()       { return AllocatorData.BlockAllocatorPrivateData;              }
+	FD3D12SegListAllocatorPrivateData& GetSegListAllocatorPrivateData()       { return AllocatorData.SegListAllocatorPrivateData;            }
+	FD3D12PoolAllocatorPrivateData&    GetPoolAllocatorPrivateData   ()       { return AllocatorData.PoolAllocatorPrivateData;               }
+
+	// Pool allocation specific functions
+	bool OnAllocationMoved(FRHIPoolAllocationData* InNewData);
+	void UnlockPoolData();
+
+	bool IsValid() const { return Type != ResourceLocationType::eUndefined; }
+
+	void AsStandAlone(FD3D12Resource* Resource, uint64 InSize = 0, bool bInIsTransient = false, const D3D12_HEAP_PROPERTIES* CustomHeapProperties = nullptr);
+	bool IsStandaloneOrPooledPlacedResource() const;
+
+	void AsHeapAliased(FD3D12Resource* Resource)
+	{
+		check(Resource->GetHeapType() != D3D12_HEAP_TYPE_READBACK);
+
+		SetType(FD3D12ResourceLocation::ResourceLocationType::eHeapAliased);
+		SetResource(Resource);
+		SetSize(0);
+
+		if (IsCPUWritable(Resource->GetHeapType()))
+		{
+			D3D12_RANGE range = { 0, 0 };
+			SetMappedBaseAddress(Resource->Map(&range));
+		}
+		SetGPUVirtualAddress(Resource->GetGPUVirtualAddress());
+	}
+
+
+	void AsFastAllocation(FD3D12Resource* Resource, uint32 BufferSize, D3D12_GPU_VIRTUAL_ADDRESS GPUBase, void* CPUBase, uint64 ResourceOffsetBase, uint64 Offset, bool bMultiFrame = false)
+	{
+		if (bMultiFrame)
+		{
+			Resource->AddRef();
+			SetType(ResourceLocationType::eMultiFrameFastAllocation);
+		}
+		else
+		{
+			SetType(ResourceLocationType::eFastAllocation);
+		}
+		SetResource(Resource);
+		SetSize(BufferSize);
+		SetOffsetFromBaseOfResource(ResourceOffsetBase + Offset);
+
+		if (CPUBase != nullptr)
+		{
+			SetMappedBaseAddress((uint8*)CPUBase + Offset);
+		}
+		SetGPUVirtualAddress(GPUBase + Offset);
+	}
+
+	// XR plugins alias textures so this allows 2+ resource locations to reference the same underlying
+	// resource. We should avoid this as much as possible as it requires expensive reference counting and
+	// it complicates the resource ownership model.
+	static void Alias(FD3D12ResourceLocation& Destination, FD3D12ResourceLocation& Source);
+	static void ReferenceNode(FD3D12Device* NodeDevice, FD3D12ResourceLocation& Destination, FD3D12ResourceLocation& Source);
+	bool IsAliased() const
+	{
+		return (Type == ResourceLocationType::eAliased) 
+			|| (Type == ResourceLocationType::eHeapAliased);
+	}
+
+	void SetTransient(bool bInTransient)
+	{
+		bTransient = bInTransient;
+	}
+	bool IsTransient() const
+	{
+		return bTransient;
+	}
+
+	/**  Get an address used by LLM to track the GPU allocation that this location represents. */
+	void* GetAddressForLLMTracking() const
+	{
+		return (uint8*)this + 1;
+	}
+
+private:
+
+	template<bool bReleaseResource>
+	void InternalClear();
+
+	void ReleaseResource();
+	void UpdateStandAloneStats(bool bIncrement);
+
+	FD3D12BaseShaderResource* Owner{};
+	FD3D12Resource* UnderlyingResource{};
+	FD3D12ResidencyHandle* ResidencyHandle{};
+
+	// Which allocator this belongs to
+	union
+	{
+		FD3D12BaseAllocatorType* Allocator;
+		FD3D12SegListAllocator* SegListAllocator;
+		FD3D12PoolAllocator* PoolAllocator;
+	};
+
+	// Union to save memory
+	union PrivateAllocatorData
+	{
+		FD3D12BuddyAllocatorPrivateData BuddyAllocatorPrivateData;
+		FD3D12BlockAllocatorPrivateData BlockAllocatorPrivateData;
+		FD3D12SegListAllocatorPrivateData SegListAllocatorPrivateData;
+		FD3D12PoolAllocatorPrivateData PoolAllocatorPrivateData;
+	} AllocatorData;
+
+	// Note: These values refer to the start of this location including any padding *NOT* the start of the underlying resource
+	void* MappedBaseAddress{};
+	D3D12_GPU_VIRTUAL_ADDRESS GPUVirtualAddress{};
+	uint64 OffsetFromBaseOfResource{};
+
+	// The size the application asked for
+	uint64 Size{};
+
+	ResourceLocationType Type{ ResourceLocationType::eUndefined };
+	EAllocatorType AllocatorType{ AT_Unknown };
+	bool bTransient{ false };
+};
+
+// Generic interface for every type D3D12 specific allocator
+struct ID3D12ResourceAllocator
+{
+	// Helper function for textures to compute the correct size and alignment
+	void AllocateTexture(uint32 GPUIndex, D3D12_HEAP_TYPE InHeapType, const FD3D12ResourceDesc& InDesc, EPixelFormat InUEFormat, ED3D12ResourceStateMode InResourceStateMode,
+		D3D12_RESOURCE_STATES InCreateState, const D3D12_CLEAR_VALUE* InClearValue, const TCHAR* InName, FD3D12ResourceLocation& ResourceLocation);
+
+	// Actual pure virtual resource allocation function
+	virtual void AllocateResource(uint32 GPUIndex, D3D12_HEAP_TYPE InHeapType, const FD3D12ResourceDesc& InDesc, uint64 InSize, uint32 InAllocationAlignment, ED3D12ResourceStateMode InResourceStateMode,
+		D3D12_RESOURCE_STATES InCreateState, const D3D12_CLEAR_VALUE* InClearValue, const TCHAR* InName, FD3D12ResourceLocation& ResourceLocation) = 0;
+};
+
+struct FD3D12LockedResource : public FD3D12DeviceChild
+{
+	FD3D12LockedResource(FD3D12Device* Device)
+		: FD3D12DeviceChild(Device)
+		, ResourceLocation(Device)
+		, LockedOffset(0)
+		, LockedPitch(0)
+		, bLocked(false)
+		, bLockedForReadOnly(false)
+		, bHasNeverBeenLocked(true)
+	{}
+
+	inline void Reset()
+	{
+		ResourceLocation.Clear();
+		bLocked = false;
+		bLockedForReadOnly = false;
+		LockedOffset = 0;
+		LockedPitch = 0;
+	}
+
+	FD3D12ResourceLocation ResourceLocation;
+	uint32 LockedOffset;
+	uint32 LockedPitch;
+	uint32 bLocked : 1;
+	uint32 bLockedForReadOnly : 1;
+	uint32 bHasNeverBeenLocked : 1;
+};
+
+/** Resource which might needs to be notified about changes on dependent resources (Views, RTGeometryObject, Cached binding tables) */
+class FD3D12ShaderResourceRenameListener
+{
+protected:
+
+	friend class FD3D12BaseShaderResource;
+	virtual void ResourceRenamed(FD3D12BaseShaderResource* InRenamedResource, FD3D12ResourceLocation* InNewResourceLocation) = 0;
+};
+
+
+/** The base class of resources that may be bound as shader resources (texture or buffer). */
+class FD3D12BaseShaderResource : public FD3D12DeviceChild, public IRefCountedObject
+{
+private:
+	mutable FCriticalSection RenameListenersCS;
+	TArray<FD3D12ShaderResourceRenameListener*> RenameListeners;
+
+public:
+	FD3D12Resource* GetResource() const { return ResourceLocation.GetResource(); }
+
+	void AddRenameListener(FD3D12ShaderResourceRenameListener* InRenameListener)
+	{
+		FScopeLock Lock(&RenameListenersCS);
+		check(!RenameListeners.Contains(InRenameListener));
+		RenameListeners.Add(InRenameListener);
+	}
+
+	void RemoveRenameListener(FD3D12ShaderResourceRenameListener* InRenameListener)
+	{
+		FScopeLock Lock(&RenameListenersCS);
+		uint32 Removed = RenameListeners.Remove(InRenameListener);
+
+		checkf(Removed == 1, TEXT("Should have exactly one registered listener during remove (same listener shouldn't registered twice and we shouldn't call this if not registered"));
+	}
+
+	bool HasLinkedViews() const
+	{
+		FScopeLock Lock(&RenameListenersCS);
+		return RenameListeners.Num() != 0;
+	}
+
+	void ResourceRenamed()
+	{
+		FScopeLock Lock(&RenameListenersCS);
+		for (FD3D12ShaderResourceRenameListener* RenameListener : RenameListeners)
+		{
+			RenameListener->ResourceRenamed(this, &ResourceLocation);
+		}
+	}
+
+	FD3D12ResourceLocation ResourceLocation;
+
+public:
+	FD3D12BaseShaderResource(FD3D12Device* InParent)
+		: FD3D12DeviceChild(InParent)
+		, ResourceLocation(InParent)
+	{
+		ResourceLocation.SetOwner(this);
+	}
+
+	~FD3D12BaseShaderResource()
+	{
+		check(!HasLinkedViews());
+	}
+};
+
+extern void UpdateBufferStats(class FD3D12Buffer* Buffer, bool bAllocating);
+
+/** Uniform buffer resource class. */
+class FD3D12UniformBuffer : public FRHIUniformBuffer, public FD3D12DeviceChild, public FD3D12LinkedAdapterObject<FD3D12UniformBuffer>
+{
+public:
+#if USE_STATIC_ROOT_SIGNATURE
+	class FD3D12ConstantBufferView* View;
+#endif
+
+	/** The D3D12 constant buffer resource */
+	FD3D12ResourceLocation ResourceLocation;
+
+	const EUniformBufferUsage UniformBufferUsage;
+
+	/** Initialization constructor. */
+	FD3D12UniformBuffer(class FD3D12Device* InParent, const FRHIUniformBufferLayout* InLayout, EUniformBufferUsage InUniformBufferUsage)
+		: FRHIUniformBuffer(InLayout)
+		, FD3D12DeviceChild(InParent)
+#if USE_STATIC_ROOT_SIGNATURE
+		, View(nullptr)
+#endif
+		, ResourceLocation(InParent)
+		, UniformBufferUsage(InUniformBufferUsage)
+	{
+	}
+
+	virtual ~FD3D12UniformBuffer();
+
+	// Provides public non-const access to ResourceTable.
+	// @todo refactor uniform buffers to perform updates as a member function, so this isn't necessary.
+	TArray<TRefCountPtr<FRHIResource>>& GetResourceTable() { return ResourceTable; }
+};
+
+class FD3D12Buffer : public FRHIBuffer, public FD3D12BaseShaderResource, public FD3D12LinkedAdapterObject<FD3D12Buffer>
+{
+public:
+	FD3D12Buffer(FD3D12Device* InParent, FRHIBufferDesc const& Desc)
+		: FRHIBuffer(Desc)
+		, FD3D12BaseShaderResource(InParent)
+		, LockedData(InParent)
+	{}
+
+	virtual ~FD3D12Buffer();
+
+	virtual uint32 GetParentGPUIndex() const override;
+
+	void UploadResourceData(FRHICommandListBase& InRHICmdList, FResourceArrayInterface* InResourceArray, D3D12_RESOURCE_STATES InDestinationState, const TCHAR* AssetName = nullptr, const FName& ClassName = NAME_None, const FName& PackageName = NAME_None);
+	FD3D12SyncPointRef UploadResourceDataViaCopyQueue(FResourceArrayInterface* InResourceArray);
+
+	// FRHIResource overrides
+#if RHI_ENABLE_RESOURCE_INFO
+	bool GetResourceInfo(FRHIResourceInfo& OutResourceInfo) const override
+	{
+		OutResourceInfo = FRHIResourceInfo{};
+		OutResourceInfo.Name = GetName();
+		OutResourceInfo.Type = GetType();
+		OutResourceInfo.VRamAllocation.AllocationSize = ResourceLocation.GetSize();
+		OutResourceInfo.IsTransient = ResourceLocation.IsTransient();
+#if ENABLE_RESIDENCY_MANAGEMENT
+		OutResourceInfo.bResident = GetResource() && GetResource()->GetResidencyHandle().ResidencyStatus == D3DX12Residency::ManagedObject::RESIDENCY_STATUS::RESIDENT;
+#endif
+		
+		return true;
+	}
+#endif
+
+	void Rename(FD3D12ResourceLocation& NewLocation);
+	void RenameLDAChain(FD3D12ResourceLocation& NewLocation);
+
+	void TakeOwnership(FD3D12Buffer& Other);
+	void ReleaseOwnership();
+
+	// IRefCountedObject interface.
+	virtual uint32 AddRef() const
+	{
+		return FRHIResource::AddRef();
+	}
+	virtual uint32 Release() const
+	{
+		return FRHIResource::Release();
+	}
+	virtual uint32 GetRefCount() const
+	{
+		return FRHIResource::GetRefCount();
+	}
+
+	static void GetResourceDescAndAlignment(uint64 InSize, uint32 InStride, EBufferUsageFlags InUsage, D3D12_RESOURCE_DESC& ResourceDesc, uint32& Alignment);
+
+	FD3D12LockedResource LockedData;
+	uint32 BufferAlignment = 0;
+};
+
+class FD3D12ResourceBarrierBatcher
+{
+	PRAGMA_DISABLE_ENUM_CONSTEXPR_CONVERSION_WARNINGS
+	// Use the top bit of the flags enum to mark transitions as "idle" time (used to remove the swapchain wait time for back buffers).
+	static constexpr D3D12_RESOURCE_BARRIER_FLAGS BarrierFlag_CountAsIdleTime = D3D12_RESOURCE_BARRIER_FLAGS(1ull << ((sizeof(D3D12_RESOURCE_BARRIER_FLAGS) * 8) - 1));
+	PRAGMA_ENABLE_ENUM_CONSTEXPR_CONVERSION_WARNINGS
+
+	struct FD3D12ResourceBarrier : public D3D12_RESOURCE_BARRIER
+	{
+		FD3D12ResourceBarrier() = default;
+		FD3D12ResourceBarrier(D3D12_RESOURCE_BARRIER&& Barrier) : D3D12_RESOURCE_BARRIER(MoveTemp(Barrier)) {}
+
+		bool HasIdleFlag() const { return !!(Flags & BarrierFlag_CountAsIdleTime); }
+		void ClearIdleFlag() { Flags &= ~BarrierFlag_CountAsIdleTime; }
+	};
+	static_assert(sizeof(FD3D12ResourceBarrier) == sizeof(D3D12_RESOURCE_BARRIER), "FD3D12ResourceBarrier is a wrapper to add helper functions. Do not add members.");
+
+public:
+	// Add a UAV barrier to the batch. Ignoring the actual resource for now.
+	void AddUAV();
+
+	// Add a transition resource barrier to the batch. Returns the number of barriers added, which may be negative if an existing barrier was cancelled.
+	int32 AddTransition(FD3D12Resource* pResource, D3D12_RESOURCE_STATES Before, D3D12_RESOURCE_STATES After, uint32 Subresource);
+
+	void AddAliasingBarrier(ID3D12Resource* InResourceBefore, ID3D12Resource* InResourceAfter);
+
+	void FlushIntoCommandList(FD3D12CommandList& CommandList, class FD3D12QueryAllocator& TimestampAllocator);
+
+	int32 Num() const { return Barriers.Num(); }
+
+private:
+	TArray<FD3D12ResourceBarrier> Barriers;
+};
+
+class FD3D12StagingBuffer final : public FRHIStagingBuffer
+{
+	friend class FD3D12CommandContext;
+	friend class FD3D12DynamicRHI;
+
+public:
+	FD3D12StagingBuffer(FD3D12Device* InDevice)
+		: FRHIStagingBuffer()
+		, ResourceLocation(InDevice)
+		, ShadowBufferSize(0)
+	{}
+	~FD3D12StagingBuffer() override;
+
+	void SafeRelease()
+	{
+		ResourceLocation.Clear();
+	}
+
+	void* Lock(uint32 Offset, uint32 NumBytes) override;
+	void Unlock() override;
+	uint64 GetGPUSizeBytes() const override { return ShadowBufferSize; }
+
+private:
+	FD3D12ResourceLocation ResourceLocation;
+	uint32 ShadowBufferSize;
+};
+
+class FD3D12GPUFence final : public FRHIGPUFence
+{
+public:
+	FD3D12GPUFence(FName InName);
+
+	virtual void Clear() override;
+	virtual bool Poll() const override;
+	virtual bool Poll(FRHIGPUMask GPUMask) const override;
+
+	void WaitCPU();
+
+	TArray<FD3D12SyncPointRef, TInlineAllocator<MAX_NUM_GPUS>> SyncPoints;
+};
+
+template<class T>
+struct TD3D12ResourceTraits
+{
+};
+template<>
+struct TD3D12ResourceTraits<FRHIUniformBuffer>
+{
+	typedef FD3D12UniformBuffer TConcreteType;
+};
+template<>
+struct TD3D12ResourceTraits<FRHIBuffer>
+{
+	typedef FD3D12Buffer TConcreteType;
+};
+template<>
+struct TD3D12ResourceTraits<FRHISamplerState>
+{
+	typedef FD3D12SamplerState TConcreteType;
+};
+template<>
+struct TD3D12ResourceTraits<FRHIRasterizerState>
+{
+	typedef FD3D12RasterizerState TConcreteType;
+};
+template<>
+struct TD3D12ResourceTraits<FRHIDepthStencilState>
+{
+	typedef FD3D12DepthStencilState TConcreteType;
+};
+template<>
+struct TD3D12ResourceTraits<FRHIBlendState>
+{
+	typedef FD3D12BlendState TConcreteType;
+};
+template<>
+struct TD3D12ResourceTraits<FRHIGraphicsPipelineState>
+{
+	typedef FD3D12GraphicsPipelineState TConcreteType;
+};
+template<>
+struct TD3D12ResourceTraits<FRHIComputePipelineState>
+{
+	typedef FD3D12ComputePipelineState TConcreteType;
+};
+template<>
+struct TD3D12ResourceTraits<FRHIGPUFence>
+{
+	typedef FD3D12GPUFence TConcreteType;
+};
+template<>
+struct TD3D12ResourceTraits<FRHIStagingBuffer>
+{
+	typedef FD3D12StagingBuffer TConcreteType;
+};
+
+
+#if D3D12_RHI_RAYTRACING
+template<>
+struct TD3D12ResourceTraits<FRHIRayTracingScene>
+{
+	typedef FD3D12RayTracingScene TConcreteType;
+};
+template<>
+struct TD3D12ResourceTraits<FRHIRayTracingGeometry>
+{
+	typedef FD3D12RayTracingGeometry TConcreteType;
+};
+template<>
+struct TD3D12ResourceTraits<FRHIRayTracingPipelineState>
+{
+	typedef FD3D12RayTracingPipelineState TConcreteType;
+};
+template<>
+struct TD3D12ResourceTraits<FRHIRayTracingShader>
+{
+	typedef FD3D12RayTracingShader TConcreteType;
+};
+#endif // D3D12_RHI_RAYTRACING
diff --git a/Engine/Source/Runtime/D3D12RHI/Private/D3D12State.h b/Engine/Source/Runtime/D3D12RHI/Private/D3D12State.h
new file mode 100644
--- /dev/null
+++ b/Engine/Source/Runtime/D3D12RHI/Private/D3D12State.h
@@ -0,0 +1,53 @@
+// Copyright Epic Games, Inc. All Rights Reserved.
+
+/*=============================================================================
+	D3D12State.h: D3D state definitions.
+	=============================================================================*/
+
+#pragma once
+
+#include "D3D12Descriptors.h"
+
+class FD3D12SamplerState : public FRHISamplerState, public FD3D12DeviceChild, public FD3D12LinkedAdapterObject<FD3D12SamplerState>
+{
+public:
+	FD3D12OfflineDescriptor OfflineDescriptor;
+	FRHIDescriptorHandle BindlessHandle;
+	const uint16 ID;
+
+	FD3D12SamplerState() = delete;
+	FD3D12SamplerState(FD3D12Device* InParent, const D3D12_SAMPLER_DESC& Desc, uint16 SamplerID);
+	~FD3D12SamplerState();
+
+	virtual FRHIDescriptorHandle GetBindlessHandle() const final { return BindlessHandle; }
+};
+
+class FD3D12RasterizerState : public FRHIRasterizerState
+{
+public:
+	D3D12_RASTERIZER_DESC Desc;
+
+	virtual bool GetInitializer(struct FRasterizerStateInitializerRHI& Init) final override;
+};
+
+class FD3D12DepthStencilState : public FRHIDepthStencilState
+{
+public:
+
+	D3D12_DEPTH_STENCIL_DESC1 Desc;
+
+	/* Describes the read/write state of the separate depth and stencil components of the DSV. */
+	FExclusiveDepthStencil AccessType;
+
+	virtual bool GetInitializer(struct FDepthStencilStateInitializerRHI& Init) final override;
+};
+
+class FD3D12BlendState : public FRHIBlendState
+{
+public:
+
+	D3D12_BLEND_DESC Desc;
+
+	virtual bool GetInitializer(class FBlendStateInitializerRHI& Init) final override;
+};
+
diff --git a/Engine/Source/Runtime/D3D12RHI/Private/D3D12Util.h b/Engine/Source/Runtime/D3D12RHI/Private/D3D12Util.h
new file mode 100644
--- /dev/null
+++ b/Engine/Source/Runtime/D3D12RHI/Private/D3D12Util.h
@@ -0,0 +1,1022 @@
+// Copyright Epic Games, Inc. All Rights Reserved.
+
+/*=============================================================================
+	D3D12Util.h: D3D RHI utility definitions.
+	=============================================================================*/
+
+#pragma once
+
+#include "CoreMinimal.h"
+#include "Containers/Queue.h"
+#include "RenderUtils.h"
+
+namespace D3D12RHI
+{
+	/**
+	 * Dump & Log all the information we have on a GPU crash (NvAfterMath, DRED, Breadcrumbs, ...) and force quit
+	 */
+	extern void TerminateOnGPUCrash(ID3D12Device* InDevice, const void* InGPUCrashDump, const size_t InGPUCrashDumpSize);
+
+	/**
+	 * Checks that the given result isn't a failure.  If it is, the application exits with an appropriate error message.
+	 * @param	Result - The result code to check
+	 * @param	Code - The code which yielded the result.
+	 * @param	Filename - The filename of the source file containing Code.
+	 * @param	Line - The line number of Code within Filename.
+	 */
+	extern void VerifyD3D12Result(HRESULT Result, const ANSICHAR* Code, const ANSICHAR* Filename, uint32 Line, ID3D12Device* Device, FString Message = FString());
+
+	/**
+	* Checks that the given result isn't a failure.  If it is, the application exits with an appropriate error message.
+	* @param	Result - The result code to check
+	* @param	Code - The code which yielded the result.
+	* @param	Filename - The filename of the source file containing Code.
+	* @param	Line - The line number of Code within Filename.
+	*/
+	extern void VerifyD3D12CreateTextureResult(HRESULT D3DResult, const ANSICHAR* Code, const ANSICHAR* Filename, uint32 Line, const D3D12_RESOURCE_DESC& TextureDesc, ID3D12Device* Device);
+
+	/**
+	 * A macro for using VERIFYD3D12RESULT that automatically passes in the code and filename/line.
+	 */
+#define VERIFYD3D12RESULT_LAMBDA(x, Device, Lambda)	{HRESULT hres = x; if (FAILED(hres)) { VerifyD3D12Result(hres, #x, __FILE__, __LINE__, Device, Lambda()); }}
+#define VERIFYD3D12RESULT_EX(x, Device)	{HRESULT hres = x; if (FAILED(hres)) { VerifyD3D12Result(hres, #x, __FILE__, __LINE__, Device); }}
+#define VERIFYD3D12RESULT(x)			{HRESULT hres = x; if (FAILED(hres)) { VerifyD3D12Result(hres, #x, __FILE__, __LINE__, nullptr); }}
+#define VERIFYD3D12CREATETEXTURERESULT(x, Desc, Device) {HRESULT hres = x; if (FAILED(hres)) { VerifyD3D12CreateTextureResult(hres, #x, __FILE__, __LINE__, Desc, Device); }}
+
+	/**
+	 * Checks that a COM object has the expected number of references.
+	 */
+	extern void VerifyComRefCount(IUnknown* Object, int32 ExpectedRefs, const TCHAR* Code, const TCHAR* Filename, int32 Line);
+#define checkComRefCount(Obj,ExpectedRefs) VerifyComRefCount(Obj,ExpectedRefs,TEXT(#Obj),TEXT(__FILE__),__LINE__)
+
+	/** Returns a string for the provided DXGI format. */
+	const TCHAR* GetD3D12TextureFormatString(DXGI_FORMAT TextureFormat);
+
+	/** Checks if given GPU virtual address corresponds to any known resource allocations and logs results */
+	void LogPageFaultData(class FD3D12Adapter* InAdapter, FD3D12Device* InDevice, D3D12_GPU_VIRTUAL_ADDRESS InPageFaultAddress);
+	
+} // namespace D3D12RHI
+
+using namespace D3D12RHI;
+
+class FD3D12Resource;
+
+void SetName(ID3D12Object* const Object, const TCHAR* const Name);
+void SetName(FD3D12Resource* const Resource, const TCHAR* const Name);
+
+enum EShaderVisibility
+{
+	SV_Vertex,
+	SV_Pixel,
+	SV_Geometry,
+#if PLATFORM_SUPPORTS_MESH_SHADERS
+	SV_Mesh,
+	SV_Amplification,
+#endif
+	SV_All,
+	SV_ShaderVisibilityCount
+};
+
+enum ERTRootSignatureType
+{
+	RS_Raster,
+	RS_RayTracingGlobal,
+	RS_RayTracingLocal,
+};
+
+struct FShaderRegisterCounts
+{
+	uint8 SamplerCount;
+	uint8 ConstantBufferCount;
+	uint8 ShaderResourceCount;
+	uint8 UnorderedAccessCount;
+};
+
+struct FD3D12QuantizedBoundShaderState
+{
+	FShaderRegisterCounts RegisterCounts[SV_ShaderVisibilityCount];
+	ERTRootSignatureType RootSignatureType = RS_Raster;
+	uint8 bAllowIAInputLayout : 1;
+	uint8 bNeedsAgsIntrinsicsSpace : 1;
+	uint8 bUseDiagnosticBuffer : 1;
+	uint8 bUseDirectlyIndexedResourceHeap : 1;
+	uint8 bUseDirectlyIndexedSamplerHeap : 1;
+	uint8 Padding : 3;
+
+	inline bool operator==(const FD3D12QuantizedBoundShaderState& RHS) const
+	{
+		return 0 == FMemory::Memcmp(this, &RHS, sizeof(RHS));
+	}
+
+	friend uint32 GetTypeHash(const FD3D12QuantizedBoundShaderState& Key);
+
+	static void InitShaderRegisterCounts(D3D12_RESOURCE_BINDING_TIER ResourceBindingTier, const FShaderCodePackedResourceCounts& Counts, FShaderRegisterCounts& Shader, bool bAllowUAVs = false);
+};
+
+/**
+* Convert from ECubeFace to D3DCUBEMAP_FACES type
+* @param Face - ECubeFace type to convert
+* @return D3D cube face enum value
+*/
+FORCEINLINE uint32 GetD3D12CubeFace(ECubeFace Face)
+{
+	switch (Face)
+	{
+	case CubeFace_PosX:
+	default:
+		return 0;//D3DCUBEMAP_FACE_POSITIVE_X;
+	case CubeFace_NegX:
+		return 1;//D3DCUBEMAP_FACE_NEGATIVE_X;
+	case CubeFace_PosY:
+		return 2;//D3DCUBEMAP_FACE_POSITIVE_Y;
+	case CubeFace_NegY:
+		return 3;//D3DCUBEMAP_FACE_NEGATIVE_Y;
+	case CubeFace_PosZ:
+		return 4;//D3DCUBEMAP_FACE_POSITIVE_Z;
+	case CubeFace_NegZ:
+		return 5;//D3DCUBEMAP_FACE_NEGATIVE_Z;
+	};
+}
+
+/**
+* Calculate a subresource index for a texture
+*/
+FORCEINLINE uint32 CalcSubresource(uint32 MipSlice, uint32 ArraySlice, uint32 MipLevels)
+{
+	return MipSlice + ArraySlice * MipLevels;
+}
+
+/**
+ * Keeps track of Locks for D3D12 objects
+ */
+class FD3D12LockedKey
+{
+public:
+	void* SourceObject;
+	uint32 Subresource;
+
+public:
+	FD3D12LockedKey() : SourceObject(NULL)
+		, Subresource(0)
+	{}
+	FD3D12LockedKey(FD3D12Resource* source, uint32 subres = 0) : SourceObject((void*)source)
+		, Subresource(subres)
+	{}
+	FD3D12LockedKey(class FD3D12ResourceLocation* source, uint32 subres = 0) : SourceObject((void*)source)
+		, Subresource(subres)
+	{}
+
+	template<class ClassType>
+	FD3D12LockedKey(ClassType* source, uint32 subres = 0) : SourceObject((void*)source)
+		, Subresource(subres)
+	{}
+	bool operator==(const FD3D12LockedKey& Other) const
+	{
+		return SourceObject == Other.SourceObject && Subresource == Other.Subresource;
+	}
+	bool operator!=(const FD3D12LockedKey& Other) const
+	{
+		return SourceObject != Other.SourceObject || Subresource != Other.Subresource;
+	}
+	FD3D12LockedKey& operator=(const FD3D12LockedKey& Other)
+	{
+		SourceObject = Other.SourceObject;
+		Subresource = Other.Subresource;
+		return *this;
+	}
+	uint32 GetHash() const
+	{
+		return PointerHash(SourceObject);
+	}
+
+	/** Hashing function. */
+	friend uint32 GetTypeHash(const FD3D12LockedKey& K)
+	{
+		return K.GetHash();
+	}
+};
+
+class FD3D12RenderTargetView;
+class FD3D12DepthStencilView;
+
+/**
+ * Class for retrieving render targets currently bound to the device context.
+ */
+class FD3D12BoundRenderTargets
+{
+public:
+	/** Initialization constructor: requires the state cache. */
+	explicit FD3D12BoundRenderTargets(FD3D12RenderTargetView** RTArray, uint32 NumActiveRTs, FD3D12DepthStencilView* DSView);
+
+	/** Destructor. */
+	~FD3D12BoundRenderTargets();
+
+	/** Accessors. */
+	FORCEINLINE int32 GetNumActiveTargets() const { return NumActiveTargets; }
+	FORCEINLINE FD3D12RenderTargetView* GetRenderTargetView(int32 TargetIndex) { return RenderTargetViews[TargetIndex]; }
+	FORCEINLINE FD3D12DepthStencilView* GetDepthStencilView() { return DepthStencilView; }
+
+private:
+	/** Active render target views. */
+	FD3D12RenderTargetView* RenderTargetViews[MaxSimultaneousRenderTargets];
+
+	/** Active depth stencil view. */
+	FD3D12DepthStencilView* DepthStencilView;
+
+	/** The number of active render targets. */
+	int32 NumActiveTargets;
+};
+
+void LogExecuteCommandLists(uint32 NumCommandLists, ID3D12CommandList *const *ppCommandLists);
+FString ConvertToResourceStateString(uint32 ResourceState);
+void LogResourceBarriers(TConstArrayView<D3D12_RESOURCE_BARRIER> Barriers, ID3D12CommandList *const pCommandList);
+
+// Custom resource states
+// To Be Determined (TBD) means we need to fill out a resource barrier before the command list is executed.
+#define D3D12_RESOURCE_STATE_TBD D3D12_RESOURCE_STATES(-1 ^ (1 << 31))
+#define D3D12_RESOURCE_STATE_CORRUPT D3D12_RESOURCE_STATES(-2 ^ (1 << 31))
+
+static bool IsValidD3D12ResourceState(D3D12_RESOURCE_STATES InState)
+{
+	return (InState != D3D12_RESOURCE_STATE_TBD && InState != D3D12_RESOURCE_STATE_CORRUPT);
+}
+
+static bool IsDirectQueueExclusiveD3D12State(D3D12_RESOURCE_STATES InState)
+{
+	return EnumHasAnyFlags(InState, D3D12_RESOURCE_STATE_RENDER_TARGET | D3D12_RESOURCE_STATE_DEPTH_WRITE | D3D12_RESOURCE_STATE_DEPTH_READ | D3D12_RESOURCE_STATE_PIXEL_SHADER_RESOURCE);
+}
+
+D3D12_RESOURCE_STATES GetD3D12ResourceState(ERHIAccess InRHIAccess, bool InIsAsyncCompute);
+
+//==================================================================================================================================
+// CResourceState
+// Tracking of per-resource or per-subresource state
+//==================================================================================================================================
+class CResourceState
+{
+public:
+	void Initialize(uint32 SubresourceCount);
+
+	bool AreAllSubresourcesSame() const;
+	bool CheckResourceState(D3D12_RESOURCE_STATES State) const;
+	bool CheckResourceStateInitalized() const;
+	D3D12_RESOURCE_STATES GetSubresourceState(uint32 SubresourceIndex) const;
+	bool CheckAllSubresourceSame();
+	void SetResourceState(D3D12_RESOURCE_STATES State);
+	void SetSubresourceState(uint32 SubresourceIndex, D3D12_RESOURCE_STATES State);
+
+	D3D12_RESOURCE_STATES GetUAVHiddenResourceState() const
+	{
+		return static_cast<D3D12_RESOURCE_STATES>(UAVHiddenResourceState);
+	}
+
+	void SetUAVHiddenResourceState(D3D12_RESOURCE_STATES InUAVHiddenResourceState)
+	{
+		// The hidden state can never include UAV
+		check(InUAVHiddenResourceState == D3D12_RESOURCE_STATE_TBD || !EnumHasAnyFlags(InUAVHiddenResourceState, D3D12_RESOURCE_STATE_UNORDERED_ACCESS));
+		check((InUAVHiddenResourceState & (1 << 31)) == 0);
+
+		UAVHiddenResourceState = (uint32)InUAVHiddenResourceState;
+	}
+
+	void SetHasInternalTransition()
+	{
+		bHasInternalTransition = 1;
+	}
+	bool HasInternalTransition() const
+	{
+		return bHasInternalTransition != 0;
+	}
+
+private:
+	// Only used if m_AllSubresourcesSame is 1.
+	// Bits defining the state of the full resource, bits are from D3D12_RESOURCE_STATES
+	uint32 m_ResourceState : 31;
+
+	// Set to 1 if m_ResourceState is valid.  In this case, all subresources have the same state
+	// Set to 0 if m_SubresourceState is valid.  In this case, each subresources may have a different state (or may be unknown)
+	uint32 m_AllSubresourcesSame : 1;
+
+	// Special resource state to track previous state before resource transitioned to UAV when the resource
+	// has a UAV aliasing resource so correct previous state can be found (only single state allowed)
+	uint32 UAVHiddenResourceState : 31;
+
+	// Was the resource used for another transition than the pending transition
+	uint32 bHasInternalTransition : 1;
+
+	// Only used if m_AllSubresourcesSame is 0.
+	// The state of each subresources.  Bits are from D3D12_RESOURCE_STATES.
+	TArray<D3D12_RESOURCE_STATES, TInlineAllocator<4>> m_SubresourceState;
+};
+
+//==================================================================================================================================
+// FD3D12ShaderBytecode
+// Encapsulates D3D12 shader bytecode and creates a hash for the shader bytecode
+//==================================================================================================================================
+struct ShaderBytecodeHash
+{
+	uint64 Hash[2];
+
+	bool operator ==(const ShaderBytecodeHash &b) const
+	{
+		return (Hash[0] == b.Hash[0] && Hash[1] == b.Hash[1]);
+	}
+
+	bool operator !=(const ShaderBytecodeHash &b) const
+	{
+		return (Hash[0] != b.Hash[0] || Hash[1] != b.Hash[1]);
+	}
+};
+
+/**
+ * The base class of threadsafe reference counted objects.
+ */
+template <class Type>
+struct FThreadsafeQueue
+{
+private:
+	mutable FCriticalSection	SynchronizationObject; // made this mutable so this class can have const functions and still be thread safe
+	TQueue<Type>				Items;
+	uint32						Size = 0;
+public:
+
+	inline const uint32 GetSize() const { return Size; }
+
+	void Enqueue(const Type& Item)
+	{
+		FScopeLock ScopeLock(&SynchronizationObject);
+		Items.Enqueue(Item);
+		Size++;
+	}
+
+	bool Dequeue(Type& Result)
+	{
+		FScopeLock ScopeLock(&SynchronizationObject);
+		Size--;
+		return Items.Dequeue(Result);
+	}
+
+	template <typename CompareFunc>
+	bool Dequeue(Type& Result, const CompareFunc& Func)
+	{
+		FScopeLock ScopeLock(&SynchronizationObject);
+
+		if (Items.Peek(Result))
+		{
+			if (Func(Result))
+			{
+				Items.Dequeue(Result);
+				Size--;
+
+				return true;
+			}
+		}
+
+		return false;
+	}
+
+	template <typename ResultType, typename CompareFunc>
+	bool BatchDequeue(TArray<ResultType>& Result, const CompareFunc& Func, uint32 MaxItems)
+	{
+		FScopeLock ScopeLock(&SynchronizationObject);
+
+		uint32 i = 0;
+		Type Item;
+		while (Items.Peek(Item) && i <= MaxItems)
+		{
+			if (Func(Item))
+			{
+				Items.Dequeue(Item);
+				Size--;
+				Result.Emplace(Item);
+
+				i++;
+			}
+			else
+			{
+				break;
+			}
+		}
+
+		return i > 0;
+	}
+
+	bool Peek(Type& Result)
+	{
+		FScopeLock ScopeLock(&SynchronizationObject);
+		return Items.Peek(Result);
+	}
+
+	bool IsEmpty()
+	{
+		FScopeLock ScopeLock(&SynchronizationObject);
+		return Items.IsEmpty();
+	}
+
+	void Empty()
+	{
+		FScopeLock ScopeLock(&SynchronizationObject);
+
+		Type Result;
+		while (Items.Dequeue(Result)) {}
+	}
+};
+
+inline bool IsCPUWritable(D3D12_HEAP_TYPE HeapType, const D3D12_HEAP_PROPERTIES *pCustomHeapProperties = nullptr)
+{
+	check(HeapType == D3D12_HEAP_TYPE_CUSTOM ? pCustomHeapProperties != nullptr : true);
+	return HeapType == D3D12_HEAP_TYPE_UPLOAD ||
+		(HeapType == D3D12_HEAP_TYPE_CUSTOM &&
+			(pCustomHeapProperties->CPUPageProperty == D3D12_CPU_PAGE_PROPERTY_WRITE_COMBINE || pCustomHeapProperties->CPUPageProperty == D3D12_CPU_PAGE_PROPERTY_WRITE_BACK));
+}
+
+inline bool IsGPUOnly(D3D12_HEAP_TYPE HeapType, const D3D12_HEAP_PROPERTIES *pCustomHeapProperties = nullptr)
+{
+	check(HeapType == D3D12_HEAP_TYPE_CUSTOM ? pCustomHeapProperties != nullptr : true);
+	return HeapType == D3D12_HEAP_TYPE_DEFAULT ||
+		(HeapType == D3D12_HEAP_TYPE_CUSTOM &&
+		(pCustomHeapProperties->CPUPageProperty == D3D12_CPU_PAGE_PROPERTY_NOT_AVAILABLE));
+}
+
+inline bool IsCPUAccessible(D3D12_HEAP_TYPE HeapType, const D3D12_HEAP_PROPERTIES* pCustomHeapProperties = nullptr)
+{
+	return !IsGPUOnly(HeapType, pCustomHeapProperties);
+}
+
+inline D3D12_RESOURCE_STATES DetermineInitialResourceState(D3D12_HEAP_TYPE HeapType, const D3D12_HEAP_PROPERTIES *pCustomHeapProperties = nullptr)
+{
+	if (HeapType == D3D12_HEAP_TYPE_DEFAULT || IsCPUWritable(HeapType, pCustomHeapProperties))
+	{
+		return D3D12_RESOURCE_STATE_GENERIC_READ;
+	}
+	else
+	{
+		check(HeapType == D3D12_HEAP_TYPE_READBACK);
+		return D3D12_RESOURCE_STATE_COPY_DEST;
+	}
+}
+
+static bool IsBlockCompressFormat(DXGI_FORMAT Format)
+{
+	// Returns true if BC1, BC2, BC3, BC4, BC5, BC6, BC7
+	return (Format >= DXGI_FORMAT_BC1_TYPELESS && Format <= DXGI_FORMAT_BC5_SNORM) ||
+		(Format >= DXGI_FORMAT_BC6H_TYPELESS && Format <= DXGI_FORMAT_BC7_UNORM_SRGB);
+}
+
+static inline uint64 GetTilesNeeded(uint32 Width, uint32 Height, uint32 Depth, const D3D12_TILE_SHAPE& Shape)
+{
+	return uint64((Width + Shape.WidthInTexels - 1) / Shape.WidthInTexels) *
+		((Height + Shape.HeightInTexels - 1) / Shape.HeightInTexels) *
+		((Depth + Shape.DepthInTexels - 1) / Shape.DepthInTexels);
+}
+
+static uint32 GetWidthAlignment(DXGI_FORMAT Format)
+{
+	switch (Format)
+	{
+	case DXGI_FORMAT_R8G8_B8G8_UNORM: return 2;
+	case DXGI_FORMAT_G8R8_G8B8_UNORM: return 2;
+	case DXGI_FORMAT_NV12: return 2;
+	case DXGI_FORMAT_P010: return 2;
+	case DXGI_FORMAT_P016: return 2;
+	case DXGI_FORMAT_420_OPAQUE: return 2;
+	case DXGI_FORMAT_YUY2: return 2;
+	case DXGI_FORMAT_Y210: return 2;
+	case DXGI_FORMAT_Y216: return 2;
+	case DXGI_FORMAT_BC1_TYPELESS: return 4;
+	case DXGI_FORMAT_BC1_UNORM: return 4;
+	case DXGI_FORMAT_BC1_UNORM_SRGB: return 4;
+	case DXGI_FORMAT_BC2_TYPELESS: return 4;
+	case DXGI_FORMAT_BC2_UNORM: return 4;
+	case DXGI_FORMAT_BC2_UNORM_SRGB: return 4;
+	case DXGI_FORMAT_BC3_TYPELESS: return 4;
+	case DXGI_FORMAT_BC3_UNORM: return 4;
+	case DXGI_FORMAT_BC3_UNORM_SRGB: return 4;
+	case DXGI_FORMAT_BC4_TYPELESS: return 4;
+	case DXGI_FORMAT_BC4_UNORM: return 4;
+	case DXGI_FORMAT_BC4_SNORM: return 4;
+	case DXGI_FORMAT_BC5_TYPELESS: return 4;
+	case DXGI_FORMAT_BC5_UNORM: return 4;
+	case DXGI_FORMAT_BC5_SNORM: return 4;
+	case DXGI_FORMAT_BC6H_TYPELESS: return 4;
+	case DXGI_FORMAT_BC6H_UF16: return 4;
+	case DXGI_FORMAT_BC6H_SF16: return 4;
+	case DXGI_FORMAT_BC7_TYPELESS: return 4;
+	case DXGI_FORMAT_BC7_UNORM: return 4;
+	case DXGI_FORMAT_BC7_UNORM_SRGB: return 4;
+	case DXGI_FORMAT_NV11: return 4;
+	case DXGI_FORMAT_R1_UNORM: return 8;
+	default: return 1;
+	}
+}
+
+static uint32 GetHeightAlignment(DXGI_FORMAT Format)
+{
+	switch (Format)
+	{
+	case DXGI_FORMAT_NV12: return 2;
+	case DXGI_FORMAT_P010: return 2;
+	case DXGI_FORMAT_P016: return 2;
+	case DXGI_FORMAT_420_OPAQUE: return 2;
+	case DXGI_FORMAT_BC1_TYPELESS: return 4;
+	case DXGI_FORMAT_BC1_UNORM: return 4;
+	case DXGI_FORMAT_BC1_UNORM_SRGB: return 4;
+	case DXGI_FORMAT_BC2_TYPELESS: return 4;
+	case DXGI_FORMAT_BC2_UNORM: return 4;
+	case DXGI_FORMAT_BC2_UNORM_SRGB: return 4;
+	case DXGI_FORMAT_BC3_TYPELESS: return 4;
+	case DXGI_FORMAT_BC3_UNORM: return 4;
+	case DXGI_FORMAT_BC3_UNORM_SRGB: return 4;
+	case DXGI_FORMAT_BC4_TYPELESS: return 4;
+	case DXGI_FORMAT_BC4_UNORM: return 4;
+	case DXGI_FORMAT_BC4_SNORM: return 4;
+	case DXGI_FORMAT_BC5_TYPELESS: return 4;
+	case DXGI_FORMAT_BC5_UNORM: return 4;
+	case DXGI_FORMAT_BC5_SNORM: return 4;
+	case DXGI_FORMAT_BC6H_TYPELESS: return 4;
+	case DXGI_FORMAT_BC6H_UF16: return 4;
+	case DXGI_FORMAT_BC6H_SF16: return 4;
+	case DXGI_FORMAT_BC7_TYPELESS: return 4;
+	case DXGI_FORMAT_BC7_UNORM: return 4;
+	case DXGI_FORMAT_BC7_UNORM_SRGB: return 4;
+	default: return 1;
+	}
+}
+
+static void Get4KTileShape(D3D12_TILE_SHAPE* pTileShape, DXGI_FORMAT DXGIFormat, EPixelFormat UEFormat, D3D12_RESOURCE_DIMENSION Dimension, uint32 SampleCount)
+{
+	//Bits per unit
+	uint32 BPU = GPixelFormats[UEFormat].BlockBytes * 8;
+
+	switch (Dimension)
+	{
+	case D3D12_RESOURCE_DIMENSION_BUFFER:
+	case D3D12_RESOURCE_DIMENSION_TEXTURE1D:
+	{
+		check(!IsBlockCompressFormat(DXGIFormat));
+		pTileShape->WidthInTexels = (BPU == 0) ? 4096 : 4096 * 8 / BPU;
+		pTileShape->HeightInTexels = 1;
+		pTileShape->DepthInTexels = 1;
+	}
+	break;
+	case D3D12_RESOURCE_DIMENSION_TEXTURE2D:
+	{
+		pTileShape->DepthInTexels = 1;
+		if (IsBlockCompressFormat(DXGIFormat))
+		{
+			// Currently only supported block sizes are 64 and 128.
+			// These equations calculate the size in texels for a tile. It relies on the fact that 16*16*16 blocks fit in a tile if the block size is 128 bits.
+			check(BPU == 64 || BPU == 128);
+			pTileShape->WidthInTexels = 16 * GetWidthAlignment(DXGIFormat);
+			pTileShape->HeightInTexels = 16 * GetHeightAlignment(DXGIFormat);
+			if (BPU == 64)
+			{
+				// If bits per block are 64 we double width so it takes up the full tile size.
+				// This is only true for BC1 and BC4
+				check((DXGIFormat >= DXGI_FORMAT_BC1_TYPELESS && DXGIFormat <= DXGI_FORMAT_BC1_UNORM_SRGB) ||
+					(DXGIFormat >= DXGI_FORMAT_BC4_TYPELESS && DXGIFormat <= DXGI_FORMAT_BC4_SNORM));
+				pTileShape->WidthInTexels *= 2;
+			}
+		}
+		else
+		{
+			if (BPU <= 8)
+			{
+				pTileShape->WidthInTexels = 64;
+				pTileShape->HeightInTexels = 64;
+			}
+			else if (BPU <= 16)
+			{
+				pTileShape->WidthInTexels = 64;
+				pTileShape->HeightInTexels = 32;
+			}
+			else if (BPU <= 32)
+			{
+				pTileShape->WidthInTexels = 32;
+				pTileShape->HeightInTexels = 32;
+			}
+			else if (BPU <= 64)
+			{
+				pTileShape->WidthInTexels = 32;
+				pTileShape->HeightInTexels = 16;
+			}
+			else if (BPU <= 128)
+			{
+				pTileShape->WidthInTexels = 16;
+				pTileShape->HeightInTexels = 16;
+			}
+			else
+			{
+				check(false);
+			}
+
+			if (SampleCount <= 1)
+			{ /* Do nothing */
+			}
+			else if (SampleCount <= 2)
+			{
+				pTileShape->WidthInTexels /= 2;
+				pTileShape->HeightInTexels /= 1;
+			}
+			else if (SampleCount <= 4)
+			{
+				pTileShape->WidthInTexels /= 2;
+				pTileShape->HeightInTexels /= 2;
+			}
+			else if (SampleCount <= 8)
+			{
+				pTileShape->WidthInTexels /= 4;
+				pTileShape->HeightInTexels /= 2;
+			}
+			else if (SampleCount <= 16)
+			{
+				pTileShape->WidthInTexels /= 4;
+				pTileShape->HeightInTexels /= 4;
+			}
+			else
+			{
+				check(false);
+			}
+
+			check(GetWidthAlignment(DXGIFormat) == 1);
+			check(GetHeightAlignment(DXGIFormat) == 1);
+		}
+
+		break;
+	}
+	case D3D12_RESOURCE_DIMENSION_TEXTURE3D:
+	{
+		if (IsBlockCompressFormat(DXGIFormat))
+		{
+			// Currently only supported block sizes are 64 and 128.
+			// These equations calculate the size in texels for a tile. It relies on the fact that 16*16*16 blocks fit in a tile if the block size is 128 bits.
+			check(BPU == 64 || BPU == 128);
+			pTileShape->WidthInTexels = 8 * GetWidthAlignment(DXGIFormat);
+			pTileShape->HeightInTexels = 8 * GetHeightAlignment(DXGIFormat);
+			pTileShape->DepthInTexels = 4;
+			if (BPU == 64)
+			{
+				// If bits per block are 64 we double width so it takes up the full tile size.
+				// This is only true for BC1 and BC4
+				check((DXGIFormat >= DXGI_FORMAT_BC1_TYPELESS && DXGIFormat <= DXGI_FORMAT_BC1_UNORM_SRGB) ||
+					(DXGIFormat >= DXGI_FORMAT_BC4_TYPELESS && DXGIFormat <= DXGI_FORMAT_BC4_SNORM));
+				pTileShape->DepthInTexels *= 2;
+			}
+		}
+		else
+		{
+			if (BPU <= 8)
+			{
+				pTileShape->WidthInTexels = 16;
+				pTileShape->HeightInTexels = 16;
+				pTileShape->DepthInTexels = 16;
+			}
+			else if (BPU <= 16)
+			{
+				pTileShape->WidthInTexels = 16;
+				pTileShape->HeightInTexels = 16;
+				pTileShape->DepthInTexels = 8;
+			}
+			else if (BPU <= 32)
+			{
+				pTileShape->WidthInTexels = 16;
+				pTileShape->HeightInTexels = 8;
+				pTileShape->DepthInTexels = 8;
+			}
+			else if (BPU <= 64)
+			{
+				pTileShape->WidthInTexels = 8;
+				pTileShape->HeightInTexels = 8;
+				pTileShape->DepthInTexels = 8;
+			}
+			else if (BPU <= 128)
+			{
+				pTileShape->WidthInTexels = 8;
+				pTileShape->HeightInTexels = 8;
+				pTileShape->DepthInTexels = 4;
+			}
+			else
+			{
+				check(false);
+			}
+
+			check(GetWidthAlignment(DXGIFormat) == 1);
+			check(GetHeightAlignment(DXGIFormat) == 1);
+		}
+	}
+	break;
+	}
+}
+
+#define ASSERT_RESOURCE_STATES 0	// Disabled for now.
+
+#if ASSERT_RESOURCE_STATES
+class FD3D12View;
+class FD3D12ViewSubset;
+
+template <class TView>
+bool AssertResourceState(ID3D12CommandList* pCommandList, FD3D12View* pView, const D3D12_RESOURCE_STATES& State);
+
+bool AssertResourceState(ID3D12CommandList* pCommandList, FD3D12Resource* pResource, const D3D12_RESOURCE_STATES& State, uint32 Subresource);
+bool AssertResourceState(ID3D12CommandList* pCommandList, FD3D12Resource* pResource, const D3D12_RESOURCE_STATES& State, const FD3D12ViewSubset& ViewSubset);
+#endif
+
+FORCEINLINE_DEBUGGABLE D3D12_PRIMITIVE_TOPOLOGY_TYPE TranslatePrimitiveTopologyType(EPrimitiveTopologyType TopologyType)
+{
+	switch (TopologyType)
+	{
+	case EPrimitiveTopologyType::Triangle:	return D3D12_PRIMITIVE_TOPOLOGY_TYPE_TRIANGLE;
+	case EPrimitiveTopologyType::Patch:		return D3D12_PRIMITIVE_TOPOLOGY_TYPE_PATCH;
+	case EPrimitiveTopologyType::Line:		return D3D12_PRIMITIVE_TOPOLOGY_TYPE_LINE;
+	case EPrimitiveTopologyType::Point:		return D3D12_PRIMITIVE_TOPOLOGY_TYPE_POINT;
+	default:
+		ensure(0);
+		return D3D12_PRIMITIVE_TOPOLOGY_TYPE_UNDEFINED;
+	}
+}
+
+FORCEINLINE_DEBUGGABLE D3D_PRIMITIVE_TOPOLOGY TranslatePrimitiveType(EPrimitiveType PrimitiveType)
+{
+	switch (PrimitiveType)
+	{
+	case PT_TriangleList:				return D3D_PRIMITIVE_TOPOLOGY_TRIANGLELIST;
+	case PT_TriangleStrip:				return D3D11_PRIMITIVE_TOPOLOGY_TRIANGLESTRIP;
+	case PT_LineList:					return D3D_PRIMITIVE_TOPOLOGY_LINELIST;
+	case PT_PointList:					return D3D_PRIMITIVE_TOPOLOGY_POINTLIST;
+	#if defined(D3D12RHI_PRIMITIVE_TOPOLOGY_RECTLIST)
+		case PT_RectList:				return D3D12RHI_PRIMITIVE_TOPOLOGY_RECTLIST;
+	#endif
+
+	default:
+		ensure(0);
+		return D3D_PRIMITIVE_TOPOLOGY_UNDEFINED;
+	}
+}
+
+#pragma warning(push)
+#pragma warning(disable: 4063)
+FORCEINLINE_DEBUGGABLE D3D12_PRIMITIVE_TOPOLOGY_TYPE D3D12PrimitiveTypeToTopologyType(D3D_PRIMITIVE_TOPOLOGY PrimitiveType)
+{
+	switch (PrimitiveType)
+	{
+	case D3D_PRIMITIVE_TOPOLOGY_POINTLIST:
+		return D3D12_PRIMITIVE_TOPOLOGY_TYPE_POINT;
+
+	case D3D_PRIMITIVE_TOPOLOGY_LINELIST:
+	case D3D_PRIMITIVE_TOPOLOGY_LINESTRIP:
+	case D3D_PRIMITIVE_TOPOLOGY_LINELIST_ADJ:
+	case D3D_PRIMITIVE_TOPOLOGY_LINESTRIP_ADJ:
+		return D3D12_PRIMITIVE_TOPOLOGY_TYPE_LINE;
+
+	case D3D_PRIMITIVE_TOPOLOGY_TRIANGLELIST:
+	case D3D_PRIMITIVE_TOPOLOGY_TRIANGLESTRIP:
+	case D3D_PRIMITIVE_TOPOLOGY_TRIANGLELIST_ADJ:
+	case D3D_PRIMITIVE_TOPOLOGY_TRIANGLESTRIP_ADJ:
+		return D3D12_PRIMITIVE_TOPOLOGY_TYPE_TRIANGLE;
+
+	#if defined(D3D12RHI_PRIMITIVE_TOPOLOGY_RECTLIST)
+		case D3D12RHI_PRIMITIVE_TOPOLOGY_RECTLIST:
+			return D3D12RHI_PRIMITIVE_TOPOLOGY_TYPE_RECT;
+	#endif
+
+	case D3D_PRIMITIVE_TOPOLOGY_UNDEFINED:
+		return D3D12_PRIMITIVE_TOPOLOGY_TYPE_UNDEFINED;
+
+	default:
+		return D3D12_PRIMITIVE_TOPOLOGY_TYPE_PATCH;
+	}
+}
+#pragma warning(pop)
+
+// @return 0xffffffff if not not supported
+FORCEINLINE_DEBUGGABLE uint32 GetMaxMSAAQuality(uint32 SampleCount)
+{
+	if (SampleCount <= DX_MAX_MSAA_COUNT)
+	{
+		// 0 has better quality (a more even distribution)
+		// higher quality levels might be useful for non box filtered AA or when using weighted samples 
+		return 0;
+	}
+
+	// not supported
+	return 0xffffffff;
+}
+
+/** Find the appropriate depth-stencil typeless DXGI format for the given format. */
+inline DXGI_FORMAT FindDepthStencilParentDXGIFormat(DXGI_FORMAT InFormat)
+{
+	switch (InFormat)
+	{
+	case DXGI_FORMAT_D24_UNORM_S8_UINT      : return DXGI_FORMAT_R24G8_TYPELESS;
+	case DXGI_FORMAT_X24_TYPELESS_G8_UINT   : return DXGI_FORMAT_R24G8_TYPELESS;
+
+	case DXGI_FORMAT_D32_FLOAT_S8X24_UINT   : return DXGI_FORMAT_R32G8X24_TYPELESS;
+	case DXGI_FORMAT_X32_TYPELESS_G8X24_UINT: return DXGI_FORMAT_R32G8X24_TYPELESS;
+
+	case DXGI_FORMAT_D32_FLOAT              : return DXGI_FORMAT_R32_TYPELESS;
+	case DXGI_FORMAT_D16_UNORM              : return DXGI_FORMAT_R16_TYPELESS;
+	};
+	return InFormat;
+}
+
+static uint8 GetPlaneSliceFromViewFormat(DXGI_FORMAT ResourceFormat, DXGI_FORMAT ViewFormat)
+{
+	// Currently, the only planar resources used are depth-stencil formats
+	switch (FindDepthStencilParentDXGIFormat(ResourceFormat))
+	{
+	case DXGI_FORMAT_R24G8_TYPELESS:
+		switch (ViewFormat)
+		{
+		case DXGI_FORMAT_R24_UNORM_X8_TYPELESS: return 0;
+		case DXGI_FORMAT_X24_TYPELESS_G8_UINT : return 1;
+		}
+		break;
+	
+	case DXGI_FORMAT_R32G8X24_TYPELESS:
+		switch (ViewFormat)
+		{
+		case DXGI_FORMAT_R32_FLOAT_X8X24_TYPELESS: return 0;
+		case DXGI_FORMAT_X32_TYPELESS_G8X24_UINT : return 1;
+		}
+		break;
+
+	case DXGI_FORMAT_NV12:
+		switch (ViewFormat)
+		{
+		case DXGI_FORMAT_R8_UNORM  : return 0;
+		case DXGI_FORMAT_R8G8_UNORM: return 1;
+		}
+		break;
+	}
+
+	return 0;
+}
+
+static uint8 GetPlaneCount(DXGI_FORMAT Format)
+{
+	// Currently, the only planar resources used are depth-stencil formats
+	// Note there is a D3D12 helper for this, D3D12GetFormatPlaneCount
+	switch (FindDepthStencilParentDXGIFormat(Format))
+	{
+	case DXGI_FORMAT_R24G8_TYPELESS:
+	case DXGI_FORMAT_R32G8X24_TYPELESS:
+		return 2;
+	default:
+		return 1;
+	}
+}
+
+static uint32 D3D12GetFormatSizeInBits(DXGI_FORMAT Format)
+{
+	switch (Format)
+	{
+	default: checkNoEntry(); [[fallthrough]];
+	case DXGI_FORMAT_UNKNOWN:
+		return 0;
+
+	case DXGI_FORMAT_R32G32B32A32_TYPELESS:
+	case DXGI_FORMAT_R32G32B32A32_FLOAT:
+	case DXGI_FORMAT_R32G32B32A32_UINT:
+	case DXGI_FORMAT_R32G32B32A32_SINT:
+	case DXGI_FORMAT_BC2_TYPELESS:
+	case DXGI_FORMAT_BC2_UNORM:
+	case DXGI_FORMAT_BC2_UNORM_SRGB:
+	case DXGI_FORMAT_BC3_TYPELESS:
+	case DXGI_FORMAT_BC3_UNORM:
+	case DXGI_FORMAT_BC3_UNORM_SRGB:
+	case DXGI_FORMAT_BC5_TYPELESS:
+	case DXGI_FORMAT_BC5_UNORM:
+	case DXGI_FORMAT_BC5_SNORM:
+	case DXGI_FORMAT_BC6H_TYPELESS:
+	case DXGI_FORMAT_BC6H_UF16:
+	case DXGI_FORMAT_BC6H_SF16:
+	case DXGI_FORMAT_BC7_TYPELESS:
+	case DXGI_FORMAT_BC7_UNORM:
+	case DXGI_FORMAT_BC7_UNORM_SRGB:
+		return 128;
+
+	case DXGI_FORMAT_R32G32B32_TYPELESS:
+	case DXGI_FORMAT_R32G32B32_FLOAT:
+	case DXGI_FORMAT_R32G32B32_UINT:
+	case DXGI_FORMAT_R32G32B32_SINT:
+		return 96;
+
+	case DXGI_FORMAT_R16G16B16A16_TYPELESS:
+	case DXGI_FORMAT_R16G16B16A16_FLOAT:
+	case DXGI_FORMAT_R16G16B16A16_UNORM:
+	case DXGI_FORMAT_R16G16B16A16_UINT:
+	case DXGI_FORMAT_R16G16B16A16_SNORM:
+	case DXGI_FORMAT_R16G16B16A16_SINT:
+	case DXGI_FORMAT_R32G32_TYPELESS:
+	case DXGI_FORMAT_R32G32_FLOAT:
+	case DXGI_FORMAT_R32G32_UINT:
+	case DXGI_FORMAT_R32G32_SINT:
+	case DXGI_FORMAT_R32G8X24_TYPELESS:
+	case DXGI_FORMAT_D32_FLOAT_S8X24_UINT:
+	case DXGI_FORMAT_R32_FLOAT_X8X24_TYPELESS:
+	case DXGI_FORMAT_X32_TYPELESS_G8X24_UINT:
+	case DXGI_FORMAT_BC1_TYPELESS:
+	case DXGI_FORMAT_BC1_UNORM:
+	case DXGI_FORMAT_BC1_UNORM_SRGB:
+	case DXGI_FORMAT_BC4_TYPELESS:
+	case DXGI_FORMAT_BC4_UNORM:
+	case DXGI_FORMAT_BC4_SNORM:
+	case DXGI_FORMAT_Y416:
+	case DXGI_FORMAT_Y210:
+	case DXGI_FORMAT_Y216:
+		return 64;
+
+	case DXGI_FORMAT_R10G10B10A2_TYPELESS:
+	case DXGI_FORMAT_R10G10B10A2_UNORM:
+	case DXGI_FORMAT_R10G10B10A2_UINT:
+	case DXGI_FORMAT_R11G11B10_FLOAT:
+	case DXGI_FORMAT_R8G8B8A8_TYPELESS:
+	case DXGI_FORMAT_R8G8B8A8_UNORM:
+	case DXGI_FORMAT_R8G8B8A8_UNORM_SRGB:
+	case DXGI_FORMAT_R8G8B8A8_UINT:
+	case DXGI_FORMAT_R8G8B8A8_SNORM:
+	case DXGI_FORMAT_R8G8B8A8_SINT:
+	case DXGI_FORMAT_R16G16_TYPELESS:
+	case DXGI_FORMAT_R16G16_FLOAT:
+	case DXGI_FORMAT_R16G16_UNORM:
+	case DXGI_FORMAT_R16G16_UINT:
+	case DXGI_FORMAT_R16G16_SNORM:
+	case DXGI_FORMAT_R16G16_SINT:
+	case DXGI_FORMAT_R32_TYPELESS:
+	case DXGI_FORMAT_D32_FLOAT:
+	case DXGI_FORMAT_R32_FLOAT:
+	case DXGI_FORMAT_R32_UINT:
+	case DXGI_FORMAT_R32_SINT:
+	case DXGI_FORMAT_R24G8_TYPELESS:
+	case DXGI_FORMAT_D24_UNORM_S8_UINT:
+	case DXGI_FORMAT_R24_UNORM_X8_TYPELESS:
+	case DXGI_FORMAT_X24_TYPELESS_G8_UINT:
+	case DXGI_FORMAT_R9G9B9E5_SHAREDEXP:
+	case DXGI_FORMAT_R8G8_B8G8_UNORM:
+	case DXGI_FORMAT_G8R8_G8B8_UNORM:
+	case DXGI_FORMAT_B8G8R8A8_UNORM:
+	case DXGI_FORMAT_B8G8R8X8_UNORM:
+	case DXGI_FORMAT_R10G10B10_XR_BIAS_A2_UNORM:
+	case DXGI_FORMAT_B8G8R8A8_TYPELESS:
+	case DXGI_FORMAT_B8G8R8A8_UNORM_SRGB:
+	case DXGI_FORMAT_B8G8R8X8_TYPELESS:
+	case DXGI_FORMAT_B8G8R8X8_UNORM_SRGB:
+	case DXGI_FORMAT_AYUV:
+	case DXGI_FORMAT_Y410:
+	case DXGI_FORMAT_P010:
+	case DXGI_FORMAT_P016:
+	case DXGI_FORMAT_YUY2:
+		return 32;
+
+	case DXGI_FORMAT_R8G8_TYPELESS:
+	case DXGI_FORMAT_R8G8_UNORM:
+	case DXGI_FORMAT_R8G8_UINT:
+	case DXGI_FORMAT_R8G8_SNORM:
+	case DXGI_FORMAT_R8G8_SINT:
+	case DXGI_FORMAT_R16_TYPELESS:
+	case DXGI_FORMAT_R16_FLOAT:
+	case DXGI_FORMAT_D16_UNORM:
+	case DXGI_FORMAT_R16_UNORM:
+	case DXGI_FORMAT_R16_UINT:
+	case DXGI_FORMAT_R16_SNORM:
+	case DXGI_FORMAT_R16_SINT:
+	case DXGI_FORMAT_B5G6R5_UNORM:
+	case DXGI_FORMAT_B5G5R5A1_UNORM:
+	case DXGI_FORMAT_B4G4R4A4_UNORM:
+	case DXGI_FORMAT_NV12:
+	case DXGI_FORMAT_NV11:
+		return 16;
+
+	case DXGI_FORMAT_R8_TYPELESS:
+	case DXGI_FORMAT_R8_UNORM:
+	case DXGI_FORMAT_R8_UINT:
+	case DXGI_FORMAT_R8_SNORM:
+	case DXGI_FORMAT_R8_SINT:
+	case DXGI_FORMAT_A8_UNORM:
+		return 8;
+
+	case DXGI_FORMAT_R1_UNORM:
+		return 1;
+	}
+}
+
+static uint32 D3D12GetFormatSizeInBytes(DXGI_FORMAT Format)
+{
+	return D3D12GetFormatSizeInBits(Format) / 8;
+}
+
+struct FD3D12ScopeLock
+{
+public:
+	FD3D12ScopeLock(FCriticalSection* CritSec) : CS(CritSec) { CS->Lock(); }
+	~FD3D12ScopeLock() { CS->Unlock(); }
+private:
+	FCriticalSection* CS;
+};
+
+struct FD3D12ScopeNoLock
+{
+public:
+	FD3D12ScopeNoLock(FCriticalSection* CritSec) { /* Do Nothing! */ }
+	~FD3D12ScopeNoLock() { /* Do Nothing! */ }
+};
\ No newline at end of file
diff --git a/Engine/Source/Runtime/D3D12RHI/Private/D3D12Viewport.h b/Engine/Source/Runtime/D3D12RHI/Private/D3D12Viewport.h
new file mode 100644
--- /dev/null
+++ b/Engine/Source/Runtime/D3D12RHI/Private/D3D12Viewport.h
@@ -0,0 +1,279 @@
+// Copyright Epic Games, Inc. All Rights Reserved.
+
+/*=============================================================================
+	D3D12Viewport.h: D3D viewport RHI definitions.
+=============================================================================*/
+
+#pragma once
+
+#include "CoreMinimal.h"
+#include "HAL/Runnable.h"
+
+#define ALTERNATE_TIMESTAMP_METRIC 0
+
+#ifndef D3D12_VIEWPORT_EXPOSES_SWAP_CHAIN
+#define D3D12_VIEWPORT_EXPOSES_SWAP_CHAIN 1
+#endif // D3D12_VIEWPORT_EXPOSES_SWAP_CHAIN
+
+#if WITH_MGPU
+class FD3D12FramePacing : public FRunnable, public FD3D12AdapterChild
+{
+public:
+	explicit FD3D12FramePacing(FD3D12Adapter* Parent);
+	~FD3D12FramePacing();
+	bool Init() override;
+	void Stop() override;
+	void Exit() override;
+	uint32 Run() override;
+
+	void PrePresentQueued(ID3D12CommandQueue* Queue);
+
+private:
+	static const uint32 MaxFrames = MAX_NUM_GPUS + 1;
+
+	TRefCountPtr<ID3D12Fence> Fence;
+	uint64 NextIndex = 0;
+	uint64 CurIndex = 0;
+	uint32 SleepTimes[MaxFrames];
+	HANDLE Semaphore;
+	bool bKeepRunning;
+
+	float AvgFrameTimeMs;
+	uint64 LastFrameTimeMs;
+
+	// ======== Some knobs for tweaking the algorithm ========
+	// How long to average the GPU time over, in seconds.
+	// - Higher = Smoother when framerate is steady, less smooth when frametime drops.
+	// - Lower = Quicker to smooth out after frametime drops, less smooth from incremental changes.
+	const float FramePacingAvgTimePeriod = 0.25f;
+	// What percentage of average GPU time to wait for on the pacing thread.
+	// - Higher = More consistent pacing, potential to starve the GPU in order to maintain pacing.
+	// - Lower = More allowable deviation between frame times, depending on GPU workload.
+#if ALTERNATE_TIMESTAMP_METRIC
+	const float FramePacingPercentage = 1.15f;
+#else
+	const float FramePacingPercentage = 1.05f;
+#endif
+
+	FRunnableThread* Thread;
+};
+#endif //WITH_MGPU
+
+
+class FD3D12Viewport : public FRHIViewport, public FD3D12AdapterChild
+{
+public:
+
+	// Lock viewport windows association and back buffer destruction because of possible crash inside DXGI factory during a call to MakeWindowAssociation
+	// Backbuffer release will wait on the call to MakeWindowAssociation while this will fail internally with 'The requested operation is not implemented.' in KernelBase.dll
+	// Reported & known problem in DXGI and will be fixed with future release but DXGI is not part of the Agility SDK so a code side fix is needed for now.
+	static FCriticalSection DXGIBackBufferLock;
+
+	FD3D12Viewport(class FD3D12Adapter* InParent, HWND InWindowHandle, uint32 InSizeX, uint32 InSizeY, bool bInIsFullscreen, EPixelFormat InPixelFormat);
+
+	void Init();
+
+	~FD3D12Viewport();
+
+	void Resize(uint32 InSizeX, uint32 InSizeY, bool bInIsFullscreen, EPixelFormat PreferredPixelFormat);
+
+	/**
+	 * If the swap chain has been invalidated by DXGI, resets the swap chain to the expected state; otherwise, does nothing.
+	 * Called once/frame by the game thread on all viewports.
+	 * @param bIgnoreFocus - Whether the reset should happen regardless of whether the window is focused.
+	 */
+	void ConditionalResetSwapChain(bool bIgnoreFocus);
+
+	/** Presents the swap chain.
+	 * Returns true if Present was done by Engine.
+	 */
+	bool Present(bool bLockToVsync);
+
+	// Accessors.
+	FIntPoint GetSizeXY() const { return FIntPoint(SizeX, SizeY); }
+
+	FD3D12Texture* GetDummyBackBuffer_RenderThread(bool bInIsSDR) const;
+	FD3D12Texture* GetBackBuffer_RenderThread() const;
+
+	FD3D12Texture* GetBackBuffer_RHIThread() const { return BackBuffer_RHIThread; }
+	FD3D12Texture* GetSDRBackBuffer_RHIThread() const { return (PixelFormat == SDRPixelFormat) ? GetBackBuffer_RHIThread() : SDRBackBuffer_RHIThread; }
+
+	FD3D12UnorderedAccessView_RHI* GetBackBufferUAV_RenderThread() const;
+
+	virtual void WaitForFrameEventCompletion() override;
+	virtual void IssueFrameEvent() override;
+
+	virtual bool NeedFlushBeforeEndDrawing() override;
+
+#if D3D12_VIEWPORT_EXPOSES_SWAP_CHAIN
+	virtual void* GetNativeSwapChain() const override { return SwapChain1; }
+#endif // #if D3D12_VIEWPORT_EXPOSES_SWAP_CHAIN
+
+	virtual void* GetNativeBackBufferTexture() const override { return GetBackBuffer_RHIThread()->GetResource(); }
+	virtual void* GetNativeBackBufferRT() const override { return GetBackBuffer_RHIThread()->GetRenderTargetView(0, 0); }
+
+	virtual void SetCustomPresent(FRHICustomPresent* InCustomPresent) override
+	{
+		CustomPresent = InCustomPresent;
+	}
+	virtual FRHICustomPresent* GetCustomPresent() const { return CustomPresent; }
+
+	virtual void* GetNativeWindow(void** AddParam = nullptr) const override { return (void*)WindowHandle; }
+
+	uint32 GetNumBackBuffers() const { return NumBackBuffers; }
+
+	inline const bool IsFullscreen() const { return bIsFullscreen; }
+
+	/** Query the swap chain's current connected output for HDR support. */
+	bool CurrentOutputSupportsHDR() const;
+
+	/** Advance and get the next present GPU index */
+	void AdvanceExpectedBackBufferIndex_RenderThread();
+#if WITH_MGPU
+	uint32 GetNextPresentGPUIndex() const
+	{
+		FScopeLock Lock(&ExpectedBackBufferIndexLock);
+		return BackBufferGPUIndices.IsValidIndex(ExpectedBackBufferIndex_RenderThread) ? BackBufferGPUIndices[ExpectedBackBufferIndex_RenderThread] : 0;
+	}
+#endif // WITH_MGPU
+
+	void OnResumeRendering();
+	void OnSuspendRendering();
+
+	static DXGI_FORMAT GetRenderTargetFormat(EPixelFormat PixelFormat)
+	{
+		DXGI_FORMAT	DXFormat = (DXGI_FORMAT)GPixelFormats[PixelFormat].PlatformFormat;
+		switch (DXFormat)
+		{
+		case DXGI_FORMAT_B8G8R8A8_TYPELESS:		return DXGI_FORMAT_B8G8R8A8_UNORM;
+		case DXGI_FORMAT_BC1_TYPELESS:			return DXGI_FORMAT_BC1_UNORM;
+		case DXGI_FORMAT_BC2_TYPELESS:			return DXGI_FORMAT_BC2_UNORM;
+		case DXGI_FORMAT_BC3_TYPELESS:			return DXGI_FORMAT_BC3_UNORM;
+		case DXGI_FORMAT_R16_TYPELESS:			return DXGI_FORMAT_R16_UNORM;
+		case DXGI_FORMAT_R8G8B8A8_TYPELESS:		return DXGI_FORMAT_R8G8B8A8_UNORM;
+		default: 								return DXFormat;
+		}
+	}
+
+private:
+	bool IsPresentAllowed();
+
+	/**
+	 * Create the dummy back buffer textures
+	 */
+	FD3D12Texture* CreateDummyBackBufferTextures(FD3D12Adapter* InAdapter, EPixelFormat InPixelFormat, uint32 InSizeX, uint32 InSizeY, bool bInIsSDR);
+
+	/**
+	 * Presents the swap chain checking the return result.
+	 * Returns true if Present was done by Engine.
+	 */
+	bool PresentChecked(int32 SyncInterval);
+
+	/**
+	 * Presents the backbuffer to the viewport window.
+	 * Returns the HRESULT for the call.
+	 */
+	HRESULT PresentInternal(int32 SyncInterval);
+
+	void ResizeInternal();
+	void FinalDestroyInternal();
+	void ClearPresentQueue();
+
+	HWND WindowHandle;
+	uint32 SizeX;
+	uint32 SizeY;
+	bool bIsFullscreen;
+	bool bFullscreenLost;
+	EPixelFormat PixelFormat;
+	bool bIsValid;
+	bool bAllowTearing;
+
+#if D3D12_VIEWPORT_EXPOSES_SWAP_CHAIN
+	TRefCountPtr<IDXGISwapChain1> SwapChain1;
+#if DXGI_MAX_SWAPCHAIN_INTERFACE >= 2
+	TRefCountPtr<IDXGISwapChain2> SwapChain2;
+#endif
+#if DXGI_MAX_SWAPCHAIN_INTERFACE >= 3
+	TRefCountPtr<IDXGISwapChain3> SwapChain3;
+#endif
+#if DXGI_MAX_SWAPCHAIN_INTERFACE >= 4
+	TRefCountPtr<IDXGISwapChain4> SwapChain4;
+#endif
+
+#if PLATFORM_WINDOWS || PLATFORM_HOLOLENS
+	DXGI_COLOR_SPACE_TYPE ColorSpace;
+#endif
+#endif // D3D12_VIEWPORT_EXPOSES_SWAP_CHAIN
+
+	TArray<TRefCountPtr<FD3D12Texture>> BackBuffers;
+	TArray<TRefCountPtr<FD3D12UnorderedAccessView_RHI>> BackBuffersUAV;
+	uint32 NumBackBuffers;
+
+	TRefCountPtr<FD3D12Texture> DummyBackBuffer_RenderThread; // Dummy back buffer texture which always references the current back buffer on the RHI thread
+	uint32 CurrentBackBufferIndex_RHIThread;
+	FD3D12Texture* BackBuffer_RHIThread;
+	FD3D12Texture* BackBuffer_RenderThread;
+	FD3D12UnorderedAccessView_RHI* BackBufferUAV_RenderThread;
+
+#if WITH_MGPU
+	int32 BackbufferMultiGPUBinding; // where INDEX_NONE cycles through the GPU, otherwise the GPU index.
+	mutable FCriticalSection ExpectedBackBufferIndexLock; // Can very rarely be modified on the RHI thread as well if present is skipped
+#endif
+	uint32 ExpectedBackBufferIndex_RenderThread; // Expected back buffer GPU index - used and updated on RenderThread!
+#if WITH_MGPU
+	TArray<uint32> BackBufferGPUIndices;
+	FD3D12SyncPointRef LastFrameSyncPoint;
+#endif // WITH_MGPU
+
+	/** 
+	 * When HDR is enabled, SDR backbuffers may be required on some architectures for game DVR or broadcasting
+	 */
+	TArray<TRefCountPtr<FD3D12Texture>> SDRBackBuffers;
+	TRefCountPtr<FD3D12Texture> SDRDummyBackBuffer_RenderThread;
+	FD3D12Texture* SDRBackBuffer_RHIThread;
+	EPixelFormat SDRPixelFormat;
+	EDisplayColorGamut DisplayColorGamut;
+	EDisplayOutputFormat DisplayOutputFormat;
+
+	/** A fence value used to track the GPU's progress. */
+	TArray<FD3D12SyncPointRef> FrameSyncPoints;
+
+	// Determine how deep the swapchain should be
+	void CalculateSwapChainDepth(int32 DefaultSwapChainDepth);
+
+	FCustomPresentRHIRef CustomPresent;
+
+	DXGI_MODE_DESC SetupDXGI_MODE_DESC() const;
+
+#if WITH_MGPU
+	FD3D12FramePacing* FramePacerRunnable;
+#endif //WITH_MGPU
+
+	struct DisplayChromacities
+	{
+		float RedX, RedY;
+		float GreenX, GreenY;
+		float BlueX, BlueY;
+		float WpX, WpY;
+	};
+
+	/** See if HDR can be enabled or not based on RHI support and current engine settings. */
+	bool CheckHDRSupport();
+
+	/** Enable HDR meta data transmission and set the necessary color space. */
+	void EnableHDR();
+
+	/** Disable HDR meta data transmission and set the necessary color space. */
+	void ShutdownHDR();
+
+#if PLATFORM_WINDOWS || PLATFORM_HOLOLENS
+	/** Ensure the correct color space is set on the swap chain */
+	void EnsureColorSpace(EDisplayColorGamut DisplayGamut, EDisplayOutputFormat OutputDevice);
+#endif
+};
+
+template<>
+struct TD3D12ResourceTraits<FRHIViewport>
+{
+	typedef FD3D12Viewport TConcreteType;
+};
diff --git a/Engine/Source/Runtime/D3D12RHI/Public/D3D12ConstantBuffer.h b/Engine/Source/Runtime/D3D12RHI/Public/D3D12ConstantBuffer.h
deleted file mode 100644
diff --git a/Engine/Source/Runtime/D3D12RHI/Public/D3D12Resources.h b/Engine/Source/Runtime/D3D12RHI/Public/D3D12Resources.h
deleted file mode 100644
diff --git a/Engine/Source/Runtime/D3D12RHI/Public/D3D12RootSignatureDefinitions.h b/Engine/Source/Runtime/D3D12RHI/Public/D3D12RootSignatureDefinitions.h
--- a/Engine/Source/Runtime/D3D12RHI/Public/D3D12RootSignatureDefinitions.h
+++ b/Engine/Source/Runtime/D3D12RHI/Public/D3D12RootSignatureDefinitions.h
@@ -2,6 +2,7 @@
 
 #pragma once
 
+#include "D3D12RHI.h"
 #include "RayTracingBuiltInResources.h"
 
 enum class ED3D12RootSignatureFlags
diff --git a/Engine/Source/Runtime/D3D12RHI/Public/D3D12ShaderResources.h b/Engine/Source/Runtime/D3D12RHI/Public/D3D12ShaderResources.h
--- a/Engine/Source/Runtime/D3D12RHI/Public/D3D12ShaderResources.h
+++ b/Engine/Source/Runtime/D3D12RHI/Public/D3D12ShaderResources.h
@@ -1,7 +1,7 @@
 // Copyright Epic Games, Inc. All Rights Reserved.
 
 /*=============================================================================
-	D3D12Resources.h: D3D resource RHI definitions.
+	D3D12ShaderResources.h: Public D3D12 RHI shader definitions.
 =============================================================================*/
 
 #pragma once
diff --git a/Engine/Source/Runtime/D3D12RHI/Public/D3D12State.h b/Engine/Source/Runtime/D3D12RHI/Public/D3D12State.h
deleted file mode 100644
diff --git a/Engine/Source/Runtime/D3D12RHI/Public/D3D12Util.h b/Engine/Source/Runtime/D3D12RHI/Public/D3D12Util.h
deleted file mode 100644
diff --git a/Engine/Source/Runtime/D3D12RHI/Public/D3D12Viewport.h b/Engine/Source/Runtime/D3D12RHI/Public/D3D12Viewport.h
deleted file mode 100644
