Fix startup crash in D3D12 when breadcrumbs are disabled.
 - Introduced in CL 24909702. DiagnosticsBuffer can be nullptr if breadcrumbs aren't enabled on startup.

diff --git a/Engine/Source/Runtime/D3D12RHI/Private/D3D12CommandContext.cpp b/Engine/Source/Runtime/D3D12RHI/Private/D3D12CommandContext.cpp
--- a/Engine/Source/Runtime/D3D12RHI/Private/D3D12CommandContext.cpp
+++ b/Engine/Source/Runtime/D3D12RHI/Private/D3D12CommandContext.cpp
@@ -466,13 +466,16 @@ FD3D12SyncPoint* FD3D12CopyScope::GetSyncPoint() const
 	return SyncPoint;
 }
 
-void FD3D12ContextCommon::InitPayloadBreadcrumbs()
+bool FD3D12ContextCommon::InitPayloadBreadcrumbs()
 {
-	FD3D12Payload* Payload = GetPayload(EPhase::Execute);
+	TUniquePtr<FD3D12DiagnosticBuffer>& DiagnosticBuffer = Device->GetQueue(QueueType).DiagnosticBuffer;
+
+	if (!DiagnosticBuffer)
+		return false;
 
+	FD3D12Payload* Payload = GetPayload(EPhase::Execute);
 	if (Payload->BreadcrumbStacks.IsEmpty() || !BreadcrumbStack.IsValid())
 	{
-		TUniquePtr<FD3D12DiagnosticBuffer>& DiagnosticBuffer = Device->GetQueue(QueueType).DiagnosticBuffer;
 		if (!BreadcrumbStack.IsValid())
 		{
 			BreadcrumbStack = MakeShared<FBreadcrumbStack>();
@@ -482,11 +485,14 @@ void FD3D12ContextCommon::InitPayloadBreadcrumbs()
 
 		Payload->BreadcrumbStacks.Add(BreadcrumbStack);
 	}
+
+	return true;
 }
 
 void FD3D12ContextCommon::WriteGPUEventStackToBreadCrumbData(const TCHAR* Name, int32 CRC)
 {
-	InitPayloadBreadcrumbs();
+	if (!InitPayloadBreadcrumbs())
+		return;
 
 	FBreadcrumbStack::FScope NewScope;
 	NewScope.NameCRC = CRC;
@@ -522,7 +528,8 @@ void FD3D12ContextCommon::WriteGPUEventStackToBreadCrumbData(const TCHAR* Name,
 
 void FD3D12ContextCommon::PopGPUEventStackFromBreadCrumbData()
 {
-	InitPayloadBreadcrumbs();
+	if (!InitPayloadBreadcrumbs())
+		return;
 
 	if (BreadcrumbStack->ScopeStack.IsEmpty())
 	{
diff --git a/Engine/Source/Runtime/D3D12RHI/Private/D3D12CommandContext.h b/Engine/Source/Runtime/D3D12RHI/Private/D3D12CommandContext.h
--- a/Engine/Source/Runtime/D3D12RHI/Private/D3D12CommandContext.h
+++ b/Engine/Source/Runtime/D3D12RHI/Private/D3D12CommandContext.h
@@ -251,7 +251,7 @@ public:
 protected:
 	void WriteGPUEventStackToBreadCrumbData(const TCHAR* Name, int32 CRC);
 	void WriteGPUEventToBreadCrumbData(FBreadcrumbStack* Breadcrumbs, uint32 MarkerIndex, bool bBeginEvent);
-	void InitPayloadBreadcrumbs();
+	[[nodiscard]] bool InitPayloadBreadcrumbs();
 	void PopGPUEventStackFromBreadCrumbData();
 
 	enum class EPhase
