Fix false sharing in D3D12 Ray Tracing shader binding table management code

diff --git a/Engine/Source/Runtime/D3D12RHI/Private/D3D12RayTracing.cpp b/Engine/Source/Runtime/D3D12RHI/Private/D3D12RayTracing.cpp
--- a/Engine/Source/Runtime/D3D12RHI/Private/D3D12RayTracing.cpp
+++ b/Engine/Source/Runtime/D3D12RHI/Private/D3D12RayTracing.cpp
@@ -1649,7 +1649,9 @@ public:
 		checkSlow(Type == D3D12_DESCRIPTOR_HEAP_TYPE_CBV_SRV_UAV || Type == D3D12_DESCRIPTOR_HEAP_TYPE_SAMPLER);
 
 		FD3D12RayTracingDescriptorHeap& Heap = (Type == D3D12_DESCRIPTOR_HEAP_TYPE_CBV_SRV_UAV) ? ViewHeap : SamplerHeap;
-		TDescriptorHashMap& Map = (Type == D3D12_DESCRIPTOR_HEAP_TYPE_CBV_SRV_UAV) ? ViewDescriptorTableCache[WorkerIndex]: SamplerDescriptorTableCache[WorkerIndex];
+		TDescriptorHashMap& Map = (Type == D3D12_DESCRIPTOR_HEAP_TYPE_CBV_SRV_UAV) 
+			? WorkerData[WorkerIndex].ViewDescriptorTableCache
+			: WorkerData[WorkerIndex].SamplerDescriptorTableCache;
 
 		const uint64 Key = CityHash64((const char*)Descriptors, sizeof(Descriptors[0]) * NumDescriptors);
 
@@ -1746,8 +1748,14 @@ public:
 
 	static constexpr uint32 MaxBindingWorkers = FD3D12RayTracingScene::MaxBindingWorkers;
 	using TDescriptorHashMap = Experimental::TSherwoodMap<uint64, int32, TIdentityHash<uint64>>;
-	TDescriptorHashMap ViewDescriptorTableCache[MaxBindingWorkers];
-	TDescriptorHashMap SamplerDescriptorTableCache[MaxBindingWorkers];
+
+	struct alignas(PLATFORM_CACHE_LINE_SIZE) FWorkerThreadData
+	{
+		TDescriptorHashMap ViewDescriptorTableCache;
+		TDescriptorHashMap SamplerDescriptorTableCache;
+	};
+
+	FWorkerThreadData WorkerData[MaxBindingWorkers];
 };
 
 class FD3D12RayTracingShaderTable
@@ -1801,9 +1809,9 @@ public:
 	{
 		delete DescriptorCache;
 	#if USE_STATIC_ROOT_SIGNATURE
-		for (const TArray<FD3D12ConstantBufferView*>& Container : TransientCBVs)
+		for (FWorkerThreadData& ThisWorkerData : WorkerData)
 		{
-			for (FD3D12ConstantBufferView* CBV : Container)
+			for (FD3D12ConstantBufferView* CBV : ThisWorkerData.TransientCBVs)
 			{
 				delete CBV;
 			}
@@ -2154,18 +2162,14 @@ public:
 			return uint32(Key.Hash);
 		}
 	};
-	TMap<FShaderRecordCacheKey, uint32> ShaderRecordCache[MaxBindingWorkers];
-
-	// A set of all resources referenced by this shader table for the purpose of updating residency before ray tracing work dispatch.
-	Experimental::TSherwoodSet<void*> ReferencedD3D12ResourceSet[MaxBindingWorkers];
-	TArray<TRefCountPtr<FD3D12Resource>> ReferencedD3D12Resources[MaxBindingWorkers];
+	
 	void AddResourceReference(FD3D12Resource* D3D12Resource, uint32 WorkerIndex)
 	{
 		bool bIsAlreadyInSet = false;
-		ReferencedD3D12ResourceSet[WorkerIndex].Add(D3D12Resource, &bIsAlreadyInSet);
+		WorkerData[WorkerIndex].ReferencedD3D12ResourceSet.Add(D3D12Resource, &bIsAlreadyInSet);
 		if (!bIsAlreadyInSet)
 		{
-			ReferencedD3D12Resources[WorkerIndex].Add(D3D12Resource);
+			WorkerData[WorkerIndex].ReferencedD3D12Resources.Add(D3D12Resource);
 		}
 	}
 
@@ -2181,17 +2185,17 @@ public:
 
 		for (uint32 WorkerIndex = 1; WorkerIndex < MaxBindingWorkers; ++WorkerIndex)
 		{
-			for (FD3D12Resource* Resource : ReferencedD3D12Resources[WorkerIndex])
+			for (FD3D12Resource* Resource : WorkerData[WorkerIndex].ReferencedD3D12Resources)
 			{
 				AddResourceReference(Resource, 0);
 			}
 			
-			ReferencedD3D12Resources[WorkerIndex].Empty();
+			WorkerData[WorkerIndex].ReferencedD3D12Resources.Empty();
 		}
 
 		// Use the main (merged) set data to actually update resource residency
 
-		for (FD3D12Resource* Resource : ReferencedD3D12Resources[0])
+		for (FD3D12Resource* Resource : WorkerData[0].ReferencedD3D12Resources)
 		{
 			CommandContext.UpdateResidency(Resource);
 		}
@@ -2201,20 +2205,14 @@ public:
 		LastCommandListID = CommandContext.GetCommandListID();
 	}
 
-	// Some resources referenced in SBT may be dynamic (written on GPU timeline) and may require transition barriers.
-	// We save such resources while we fill the SBT and issue transitions before the SBT is used.
-
-	TSet<FD3D12ShaderResourceView*> TransitionSRVs[MaxBindingWorkers];
-	TSet<FD3D12UnorderedAccessView*> TransitionUAVs[MaxBindingWorkers];
-
 	void AddResourceTransition(FD3D12ShaderResourceView* SRV, uint32 WorkerIndex)
 	{
-		TransitionSRVs[WorkerIndex].Add(SRV);
+		WorkerData[WorkerIndex].TransitionSRVs.Add(SRV);
 	}
 
 	void AddResourceTransition(FD3D12UnorderedAccessView* UAV, uint32 WorkerIndex)
 	{
-		TransitionUAVs[WorkerIndex].Add(UAV);
+		WorkerData[WorkerIndex].TransitionUAVs.Add(UAV);
 	}
 
 	void TransitionResources(FD3D12CommandContext& CommandContext)
@@ -2223,36 +2221,52 @@ public:
 
 		for (uint32 WorkerIndex = 1; WorkerIndex < MaxBindingWorkers; ++WorkerIndex)
 		{
-			for (FD3D12ShaderResourceView* SRV : TransitionSRVs[WorkerIndex])
+			for (FD3D12ShaderResourceView* SRV : WorkerData[WorkerIndex].TransitionSRVs)
 			{
 				AddResourceTransition(SRV, 0);
 			}
 
-			for (FD3D12UnorderedAccessView* UAV : TransitionUAVs[WorkerIndex])
+			for (FD3D12UnorderedAccessView* UAV : WorkerData[WorkerIndex].TransitionUAVs)
 			{
 				AddResourceTransition(UAV, 0);
 			}
 
-			TransitionSRVs[WorkerIndex].Empty();
-			TransitionUAVs[WorkerIndex].Empty();
+			WorkerData[WorkerIndex].TransitionSRVs.Empty();
+			WorkerData[WorkerIndex].TransitionUAVs.Empty();
 		}
 
 		// Use the main (merged) set data to perform resource transitions
 
-		for (FD3D12ShaderResourceView* SRV : TransitionSRVs[0])
+		for (FD3D12ShaderResourceView* SRV : WorkerData[0].TransitionSRVs)
 		{
 			CommandContext.TransitionResource(SRV, D3D12_RESOURCE_STATE_NON_PIXEL_SHADER_RESOURCE);
 		}
 
-		for (FD3D12UnorderedAccessView* UAV : TransitionUAVs[0])
+		for (FD3D12UnorderedAccessView* UAV : WorkerData[0].TransitionUAVs)
 		{
 			CommandContext.TransitionResource(UAV, D3D12_RESOURCE_STATE_UNORDERED_ACCESS);
 		}
 	}
 
+	struct alignas(PLATFORM_CACHE_LINE_SIZE) FWorkerThreadData
+	{
+		TMap<FShaderRecordCacheKey, uint32> ShaderRecordCache;
+
+		// A set of all resources referenced by this shader table for the purpose of updating residency before ray tracing work dispatch.
+		Experimental::TSherwoodSet<void*> ReferencedD3D12ResourceSet;
+		TArray<TRefCountPtr<FD3D12Resource>> ReferencedD3D12Resources;
+
+		// Some resources referenced in SBT may be dynamic (written on GPU timeline) and may require transition barriers.
+		// We save such resources while we fill the SBT and issue transitions before the SBT is used.
+		TSet<FD3D12ShaderResourceView*> TransitionSRVs;
+		TSet<FD3D12UnorderedAccessView*> TransitionUAVs;
+
 #if USE_STATIC_ROOT_SIGNATURE
-	TArray<FD3D12ConstantBufferView*> TransientCBVs[MaxBindingWorkers];
+		TArray<FD3D12ConstantBufferView*> TransientCBVs;
 #endif // USE_STATIC_ROOT_SIGNATURE
+	};
+
+	FWorkerThreadData WorkerData[MaxBindingWorkers];
 };
 
 
@@ -4641,7 +4655,7 @@ struct FD3D12RayTracingLocalResourceBinder
 	{
 		FD3D12ConstantBufferView* BufferView = new FD3D12ConstantBufferView(GetDevice());
 		BufferView->CreateView(&ResourceLocation, 0, DataSize);
-		ShaderTable.TransientCBVs[WorkerIndex].Add(BufferView);
+		ShaderTable.WorkerData[WorkerIndex].TransientCBVs.Add(BufferView);
 		return BufferView->GetOfflineCpuHandle();
 	}
 #endif // USE_STATIC_ROOT_SIGNATURE
@@ -5302,7 +5316,7 @@ static void SetRayTracingHitGroup(
 	{
 		CacheKey = FD3D12RayTracingShaderTable::FShaderRecordCacheKey(NumUniformBuffers, UniformBuffers, HitGroupIndex);
 
-		uint32* ExistingRecordIndex = ShaderTable->ShaderRecordCache[WorkerIndex].Find(CacheKey);
+		uint32* ExistingRecordIndex = ShaderTable->WorkerData[WorkerIndex].ShaderRecordCache.Find(CacheKey);
 		if (ExistingRecordIndex)
 		{
 			// Simply copy local shader parameters from existing SBT record and set the shader identifier, skipping resource binding work.
@@ -5325,7 +5339,7 @@ static void SetRayTracingHitGroup(
 
 	if (bCanUseRecordCache && bResourcesBound)
 	{
-		ShaderTable->ShaderRecordCache[WorkerIndex].Add(CacheKey, RecordIndex);
+		ShaderTable->WorkerData[WorkerIndex].ShaderRecordCache.Add(CacheKey, RecordIndex);
 	}
 
 	ShaderTable->SetHitGroupIdentifier(RecordIndex,
