Several micro-optimizations for D3D12 shader binding table management

* Use XXHash instead of CityHash
* Use TSherwoodSet/Map instead of TSet/Map when possible

diff --git a/Engine/Source/Runtime/D3D12RHI/Private/D3D12RayTracing.cpp b/Engine/Source/Runtime/D3D12RHI/Private/D3D12RayTracing.cpp
--- a/Engine/Source/Runtime/D3D12RHI/Private/D3D12RayTracing.cpp
+++ b/Engine/Source/Runtime/D3D12RHI/Private/D3D12RayTracing.cpp
@@ -10,7 +10,7 @@
 #include "Experimental/Containers/SherwoodHashTable.h"
 #include "BuiltInRayTracingShaders.h"
 #include "RayTracingValidationShaders.h"
-#include "Hash/CityHash.h"
+#include "Hash/xxhash.h"
 #include "HAL/CriticalSection.h"
 #include "HAL/IConsoleManager.h"
 #include "HAL/FileManagerGeneric.h"
@@ -386,7 +386,7 @@ inline void UnregisterD3D12RayTracingGeometry(FD3D12RayTracingGeometry* Geometry
 }
 #endif // UE_BUILD_SHIPPING
 
-// Whether to compare the full descriptor table on cache lookup or only use CityHash64 digest.
+// Whether to compare the full descriptor table on cache lookup or only use FXxHash64 digest.
 #ifndef RAY_TRACING_DESCRIPTOR_CACHE_FULL_COMPARE
 #define RAY_TRACING_DESCRIPTOR_CACHE_FULL_COMPARE 1
 #endif // RAY_TRACING_DESCRIPTOR_CACHE_FULL_COMPARE
@@ -1653,7 +1653,7 @@ public:
 			? WorkerData[WorkerIndex].ViewDescriptorTableCache
 			: WorkerData[WorkerIndex].SamplerDescriptorTableCache;
 
-		const uint64 Key = CityHash64((const char*)Descriptors, sizeof(Descriptors[0]) * NumDescriptors);
+		const uint64 Key = FXxHash64::HashBuffer(Descriptors, sizeof(Descriptors[0]) * NumDescriptors).Hash;
 
 		const int32 InvalidIndex = -1;
 		int32& DescriptorTableBaseIndex = Map.FindOrAdd(Key, InvalidIndex);
@@ -2126,9 +2126,9 @@ public:
 	{
 		static constexpr uint32 MaxUniformBuffers = 6;
 		FRHIUniformBuffer* const* UniformBuffers[MaxUniformBuffers];
-		uint64 Hash;
-		uint32 NumUniformBuffers;
-		uint32 ShaderIndex;
+		uint64 Hash = 0;
+		uint32 NumUniformBuffers = 0;
+		uint32 ShaderIndex = 0;
 
 		FShaderRecordCacheKey() = default;
 		FShaderRecordCacheKey(uint32 InNumUniformBuffers, FRHIUniformBuffer* const* InUniformBuffers, uint32 InShaderIndex)
@@ -2140,7 +2140,7 @@ public:
 
 			const uint64 DataSizeInBytes = sizeof(FRHIUniformBuffer*) * NumUniformBuffers;
 			FMemory::Memcpy(UniformBuffers, InUniformBuffers, DataSizeInBytes);
-			Hash = CityHash64(reinterpret_cast<const char*>(UniformBuffers), DataSizeInBytes);
+			Hash = FXxHash64::HashBuffer(UniformBuffers, DataSizeInBytes).Hash;
 		}
 
 		bool operator == (const FShaderRecordCacheKey& Other) const
@@ -2207,12 +2207,23 @@ public:
 
 	void AddResourceTransition(FD3D12ShaderResourceView* SRV, uint32 WorkerIndex)
 	{
-		WorkerData[WorkerIndex].TransitionSRVs.Add(SRV);
+		bool bAlreadyInSet = false;
+		WorkerData[WorkerIndex].TransitionViewSet.Add(SRV, &bAlreadyInSet);
+		if (!bAlreadyInSet)
+		{
+			WorkerData[WorkerIndex].TransitionSRVs.Add(SRV);
+		}
 	}
 
 	void AddResourceTransition(FD3D12UnorderedAccessView* UAV, uint32 WorkerIndex)
 	{
-		WorkerData[WorkerIndex].TransitionUAVs.Add(UAV);
+		bool bAlreadyInSet = false;
+		WorkerData[WorkerIndex].TransitionViewSet.Add(UAV, &bAlreadyInSet);
+
+		if (!bAlreadyInSet)
+		{
+			WorkerData[WorkerIndex].TransitionUAVs.Add(UAV);
+		}
 	}
 
 	void TransitionResources(FD3D12CommandContext& CommandContext)
@@ -2233,6 +2244,7 @@ public:
 
 			WorkerData[WorkerIndex].TransitionSRVs.Empty();
 			WorkerData[WorkerIndex].TransitionUAVs.Empty();
+			WorkerData[WorkerIndex].TransitionViewSet.Empty();
 		}
 
 		// Use the main (merged) set data to perform resource transitions
@@ -2250,7 +2262,7 @@ public:
 
 	struct alignas(PLATFORM_CACHE_LINE_SIZE) FWorkerThreadData
 	{
-		TMap<FShaderRecordCacheKey, uint32> ShaderRecordCache;
+		Experimental::TSherwoodMap<FShaderRecordCacheKey, uint32> ShaderRecordCache;
 
 		// A set of all resources referenced by this shader table for the purpose of updating residency before ray tracing work dispatch.
 		Experimental::TSherwoodSet<void*> ReferencedD3D12ResourceSet;
@@ -2258,8 +2270,9 @@ public:
 
 		// Some resources referenced in SBT may be dynamic (written on GPU timeline) and may require transition barriers.
 		// We save such resources while we fill the SBT and issue transitions before the SBT is used.
-		TSet<FD3D12ShaderResourceView*> TransitionSRVs;
-		TSet<FD3D12UnorderedAccessView*> TransitionUAVs;
+		Experimental::TSherwoodSet<FD3D12View*> TransitionViewSet;
+		TArray<FD3D12ShaderResourceView*> TransitionSRVs;
+		TArray<FD3D12UnorderedAccessView*> TransitionUAVs;
 
 #if USE_STATIC_ROOT_SIGNATURE
 		TArray<FD3D12ConstantBufferView*> TransientCBVs;
@@ -5339,7 +5352,7 @@ static void SetRayTracingHitGroup(
 
 	if (bCanUseRecordCache && bResourcesBound)
 	{
-		ShaderTable->WorkerData[WorkerIndex].ShaderRecordCache.Add(CacheKey, RecordIndex);
+		ShaderTable->WorkerData[WorkerIndex].ShaderRecordCache.FindOrAdd(CacheKey, RecordIndex);
 	}
 
 	ShaderTable->SetHitGroupIdentifier(RecordIndex,
