Add version to FD3D12OfflineDescriptor to detect updates that should invalidate cached descriptor tables

* Fixes the invalid descriptor table reuse in FD3D12RayTracingDescriptorCache when underlying descriptor data has been changed, but the descriptor pointer itself is the same

diff --git a/Engine/Source/Runtime/D3D12RHI/Private/D3D12Descriptors.h b/Engine/Source/Runtime/D3D12RHI/Private/D3D12Descriptors.h
--- a/Engine/Source/Runtime/D3D12RHI/Private/D3D12Descriptors.h
+++ b/Engine/Source/Runtime/D3D12RHI/Private/D3D12Descriptors.h
@@ -4,7 +4,17 @@
 
 #include "RHIDescriptorAllocator.h"
 #include "MultiGPU.h"
+#include "D3D12RHI.h"
 #include "D3D12RHICommon.h"
+#include "Templates/RefCounting.h"
+#include "Containers/Queue.h"
+#include "Containers/List.h"
+
+#if PLATFORM_WINDOWS
+#include "Windows/D3D12ThirdParty.h"
+#else
+#include "D3D12ThirdParty.h"
+#endif
 
 inline D3D12_DESCRIPTOR_HEAP_TYPE Translate(ERHIDescriptorHeapType InHeapType)
 {
@@ -194,9 +204,15 @@ public:
 	operator bool () const { return ptr != 0; }
 	FD3D12OfflineDescriptor() { ptr = 0; }
 
+	// Descriptor version can be used to invalidate any caches that are based on D3D12_CPU_DESCRIPTOR_HANDLE,
+	// for example to detect when a view has been updated following a resource rename.
+	void IncrementVersion() { ++Version; }
+	uint32 GetVersion() const { return Version; }
+
 private:
 	friend class FD3D12OfflineDescriptorManager;
 	uint32 HeapIndex = 0;
+	uint32 Version = 0;
 };
 
 /** Manages and allows allocations of CPU descriptors only. Creates small heaps on demand to satisfy allocations. */
diff --git a/Engine/Source/Runtime/D3D12RHI/Private/D3D12RayTracing.cpp b/Engine/Source/Runtime/D3D12RHI/Private/D3D12RayTracing.cpp
--- a/Engine/Source/Runtime/D3D12RHI/Private/D3D12RayTracing.cpp
+++ b/Engine/Source/Runtime/D3D12RHI/Private/D3D12RayTracing.cpp
@@ -1644,7 +1644,7 @@ public:
 	}
 
 	// Returns descriptor heap base index for this descriptor table allocation or -1 if allocation failed.
-	int32 AllocateDescriptorTable(const D3D12_CPU_DESCRIPTOR_HANDLE* Descriptors, uint32 NumDescriptors, D3D12_DESCRIPTOR_HEAP_TYPE Type, uint32 WorkerIndex)
+	int32 AllocateDescriptorTable(const uint32* DescriptorVersions, const D3D12_CPU_DESCRIPTOR_HANDLE* Descriptors, uint32 NumDescriptors, D3D12_DESCRIPTOR_HEAP_TYPE Type, uint32 WorkerIndex)
 	{
 		checkSlow(Type == D3D12_DESCRIPTOR_HEAP_TYPE_CBV_SRV_UAV || Type == D3D12_DESCRIPTOR_HEAP_TYPE_SAMPLER);
 
@@ -1653,12 +1653,13 @@ public:
 			? WorkerData[WorkerIndex].ViewDescriptorTableCache
 			: WorkerData[WorkerIndex].SamplerDescriptorTableCache;
 
-		const uint64 Key = FXxHash64::HashBuffer(Descriptors, sizeof(Descriptors[0]) * NumDescriptors).Hash;
+		const uint64 VersionHash = FXxHash64::HashBuffer(DescriptorVersions, sizeof(DescriptorVersions[0]) * NumDescriptors).Hash;
+		const uint64 DescriptorHash = FXxHash64::HashBuffer(Descriptors, sizeof(Descriptors[0]) * NumDescriptors).Hash;
+		const uint64 Key = VersionHash ^ DescriptorHash;
 
-		const int32 InvalidIndex = -1;
-		int32& DescriptorTableBaseIndex = Map.FindOrAdd(Key, InvalidIndex);
+		int32& DescriptorTableBaseIndex = Map.FindOrAdd(Key, INDEX_NONE);
 
-		if (DescriptorTableBaseIndex != InvalidIndex)
+		if (DescriptorTableBaseIndex != INDEX_NONE)
 		{
 		#if RAY_TRACING_DESCRIPTOR_CACHE_FULL_COMPARE
 			if (ensureMsgf(Heap.CompareDescriptors(DescriptorTableBaseIndex, Descriptors, NumDescriptors), 
@@ -1676,7 +1677,7 @@ public:
 			// Per-thread descriptor table deduplication hash tables introduce a lot of redundancy in the heap
 			// which is reduced by looking for global matches on hash table lookup miss.
 
-			int32 FoundIndex = InvalidIndex;
+			int32 FoundIndex = INDEX_NONE;
 			int32 SearchEndPos = Heap.NumWrittenSamplerDescriptors;
 			for (int32 SearchIndex = 0; SearchIndex + int32(NumDescriptors) < SearchEndPos; ++SearchIndex)
 			{
@@ -1687,7 +1688,7 @@ public:
 				}
 			}
 
-			if (FoundIndex != InvalidIndex)
+			if (FoundIndex != INDEX_NONE)
 			{
 				DescriptorTableBaseIndex = FoundIndex;
 				return DescriptorTableBaseIndex;
@@ -1696,9 +1697,9 @@ public:
 
 		DescriptorTableBaseIndex = Heap.Allocate(NumDescriptors);
 
-		if (DescriptorTableBaseIndex == InvalidIndex)
+		if (DescriptorTableBaseIndex == INDEX_NONE)
 		{
-			return InvalidIndex;
+			return INDEX_NONE;
 		}
 
 		Heap.CopyDescriptors(DescriptorTableBaseIndex, Descriptors, NumDescriptors);
@@ -4751,6 +4752,11 @@ static bool SetRayTracingShaderResources(
 		D3D12_CPU_DESCRIPTOR_HANDLE LocalUAVs[MAX_UAVS];
 		D3D12_CPU_DESCRIPTOR_HANDLE LocalSamplers[MAX_SAMPLERS];
 
+		uint32 CBVVersions[MAX_CBS];
+		uint32 SRVVersions[MAX_SRVS];
+		uint32 UAVVersions[MAX_SRVS];
+		uint32 SamplerVersions[MAX_SRVS];
+
 		TArray<FD3D12Resource*, TInlineAllocator<MAX_CBS + MAX_SRVS + MAX_UAVS>> ReferencedResources;
 
 		uint64 BoundSRVMask = 0;
@@ -4763,7 +4769,10 @@ static bool SetRayTracingShaderResources(
 			FD3D12UnorderedAccessView* UAV = FD3D12CommandContext::RetrieveObject<FD3D12UnorderedAccessView_RHI>(RHIUAV, GPUIndex);
 			check(UAV != nullptr);
 
-			LocalUAVs[Index] = UAV->GetOfflineCpuHandle();
+			FD3D12OfflineDescriptor Descriptor = UAV->GetOfflineCpuHandle();
+			LocalUAVs[Index] = Descriptor;
+			UAVVersions[Index] = Descriptor.GetVersion();
+
 			BoundUAVMask |= 1ull << Index;
 
 			ReferencedResources.Add(UAV->GetResource());
@@ -4775,7 +4784,10 @@ static bool SetRayTracingShaderResources(
 			FD3D12ShaderResourceView* SRV = FD3D12CommandContext::RetrieveObject<FD3D12ShaderResourceView_RHI>(RHISRV, GPUIndex);
 			check(SRV != nullptr);
 
-			LocalSRVs[Index] = SRV->GetOfflineCpuHandle();
+			FD3D12OfflineDescriptor Descriptor = SRV->GetOfflineCpuHandle();
+			LocalSRVs[Index] = Descriptor;
+			SRVVersions[Index] = Descriptor.GetVersion();
+
 			BoundSRVMask |= 1ull << Index;
 
 			ReferencedResources.Add(SRV->GetResource());
@@ -4791,7 +4803,10 @@ static bool SetRayTracingShaderResources(
 			}
 			check(SRV != nullptr);
 
-			LocalSRVs[Index] = SRV->GetOfflineCpuHandle();
+			FD3D12OfflineDescriptor Descriptor = SRV->GetOfflineCpuHandle();
+			LocalSRVs[Index] = Descriptor;
+			SRVVersions[Index] = Descriptor.GetVersion();
+
 			BoundSRVMask |= 1ull << Index;
 
 			ReferencedResources.Add(SRV->GetResource());
@@ -4803,7 +4818,10 @@ static bool SetRayTracingShaderResources(
 			FD3D12SamplerState* Sampler = FD3D12CommandContext::RetrieveObject<FD3D12SamplerState>(RHISampler, GPUIndex);
 			check(Sampler != nullptr);
 
-			LocalSamplers[Index] = Sampler->OfflineDescriptor;
+			FD3D12OfflineDescriptor Descriptor = Sampler->OfflineDescriptor;
+			LocalSamplers[Index] = Descriptor;
+			SamplerVersions[Index] = Descriptor.GetVersion();
+
 			BoundSamplerMask |= 1ull << Index;
 		}
 
@@ -4834,9 +4852,12 @@ static bool SetRayTracingShaderResources(
 		{
 			FD3D12UniformBuffer* CBV = FD3D12CommandContext::RetrieveObject<FD3D12UniformBuffer>(Resource, Bindings.GPUIndex);
 		#if USE_STATIC_ROOT_SIGNATURE
-			Bindings.LocalCBVs[CBVIndex] = CBV->View->GetOfflineCpuHandle();
+			FD3D12OfflineDescriptor Descriptor = CBV->View->GetOfflineCpuHandle();
+			Bindings.LocalCBVs[CBVIndex] = Descriptor;
+			Bindings.CBVVersions[CBVIndex] = Descriptor.GetVersion();
 		#else // USE_STATIC_ROOT_SIGNATURE
 			Bindings.LocalCBVs[CBVIndex] = CBV->ResourceLocation.GetGPUVirtualAddress();
+			Bindings.CBVVersions[CBVIndex] = 0; // not available with GPU address path
 		#endif // USE_STATIC_ROOT_SIGNATURE
 			Bindings.BoundCBVMask |= 1ull << CBVIndex;
 
@@ -4921,7 +4942,7 @@ static bool SetRayTracingShaderResources(
 	const uint32 NumSRVs = Shader->ResourceCounts.NumSRVs;
 	if (NumSRVs)
 	{
-		const int32 DescriptorTableBaseIndex = DescriptorCache.AllocateDescriptorTable(Bindings.LocalSRVs, NumSRVs, D3D12_DESCRIPTOR_HEAP_TYPE_CBV_SRV_UAV, WorkerIndex);
+		const int32 DescriptorTableBaseIndex = DescriptorCache.AllocateDescriptorTable(Bindings.SRVVersions, Bindings.LocalSRVs, NumSRVs, D3D12_DESCRIPTOR_HEAP_TYPE_CBV_SRV_UAV, WorkerIndex);
 		if (DescriptorTableBaseIndex < 0)
 		{
 			return false;
@@ -4937,7 +4958,7 @@ static bool SetRayTracingShaderResources(
 	const uint32 NumUAVs = Shader->ResourceCounts.NumUAVs;
 	if (NumUAVs)
 	{
-		const int32 DescriptorTableBaseIndex = DescriptorCache.AllocateDescriptorTable(Bindings.LocalUAVs, NumUAVs, D3D12_DESCRIPTOR_HEAP_TYPE_CBV_SRV_UAV, WorkerIndex);
+		const int32 DescriptorTableBaseIndex = DescriptorCache.AllocateDescriptorTable(Bindings.UAVVersions, Bindings.LocalUAVs, NumUAVs, D3D12_DESCRIPTOR_HEAP_TYPE_CBV_SRV_UAV, WorkerIndex);
 		if (DescriptorTableBaseIndex < 0)
 		{
 			return false;
@@ -4955,7 +4976,7 @@ static bool SetRayTracingShaderResources(
 	{
 	#if USE_STATIC_ROOT_SIGNATURE
 
-		const uint32 DescriptorTableBaseIndex = DescriptorCache.AllocateDescriptorTable(Bindings.LocalCBVs, NumCBVs, D3D12_DESCRIPTOR_HEAP_TYPE_CBV_SRV_UAV, WorkerIndex);
+		const uint32 DescriptorTableBaseIndex = DescriptorCache.AllocateDescriptorTable(Bindings.CBVVersions, Bindings.LocalCBVs, NumCBVs, D3D12_DESCRIPTOR_HEAP_TYPE_CBV_SRV_UAV, WorkerIndex);
 		const uint32 BindSlot = RootSignature->CBVRDTBindSlot(SF_Compute);
 		check(BindSlot != 0xFF);
 
@@ -4984,7 +5005,7 @@ static bool SetRayTracingShaderResources(
 	const uint32 NumSamplers = Shader->ResourceCounts.NumSamplers;
 	if (NumSamplers)
 	{
-		const int32 DescriptorTableBaseIndex = DescriptorCache.AllocateDescriptorTable(Bindings.LocalSamplers, NumSamplers, D3D12_DESCRIPTOR_HEAP_TYPE_SAMPLER, WorkerIndex);
+		const int32 DescriptorTableBaseIndex = DescriptorCache.AllocateDescriptorTable(Bindings.SamplerVersions, Bindings.LocalSamplers, NumSamplers, D3D12_DESCRIPTOR_HEAP_TYPE_SAMPLER, WorkerIndex);
 		if (DescriptorTableBaseIndex < 0)
 		{
 			return false;
diff --git a/Engine/Source/Runtime/D3D12RHI/Private/D3D12SRV.cpp b/Engine/Source/Runtime/D3D12RHI/Private/D3D12SRV.cpp
--- a/Engine/Source/Runtime/D3D12RHI/Private/D3D12SRV.cpp
+++ b/Engine/Source/Runtime/D3D12RHI/Private/D3D12SRV.cpp
@@ -103,6 +103,8 @@ void FD3D12ShaderResourceView::UpdateDescriptor()
 		&D3DViewDesc,
 		OfflineCpuHandle
 	);
+
+	OfflineCpuHandle.IncrementVersion();
 }
 
 void FD3D12ShaderResourceView_RHI::CreateView()
diff --git a/Engine/Source/Runtime/D3D12RHI/Private/D3D12State.h b/Engine/Source/Runtime/D3D12RHI/Private/D3D12State.h
--- a/Engine/Source/Runtime/D3D12RHI/Private/D3D12State.h
+++ b/Engine/Source/Runtime/D3D12RHI/Private/D3D12State.h
@@ -7,6 +7,7 @@
 #pragma once
 
 #include "D3D12Descriptors.h"
+#include "RHIResources.h"
 
 class FD3D12SamplerState : public FRHISamplerState, public FD3D12DeviceChild, public FD3D12LinkedAdapterObject<FD3D12SamplerState>
 {
diff --git a/Engine/Source/Runtime/D3D12RHI/Private/D3D12UAV.cpp b/Engine/Source/Runtime/D3D12RHI/Private/D3D12UAV.cpp
--- a/Engine/Source/Runtime/D3D12RHI/Private/D3D12UAV.cpp
+++ b/Engine/Source/Runtime/D3D12RHI/Private/D3D12UAV.cpp
@@ -85,6 +85,8 @@ void FD3D12UnorderedAccessView::UpdateDescriptor()
 		&D3DViewDesc,
 		OfflineCpuHandle
 	);
+
+	OfflineCpuHandle.IncrementVersion();
 }
 
 void FD3D12UnorderedAccessView_RHI::CreateView()
diff --git a/Engine/Source/Runtime/D3D12RHI/Private/D3D12View.cpp b/Engine/Source/Runtime/D3D12RHI/Private/D3D12View.cpp
--- a/Engine/Source/Runtime/D3D12RHI/Private/D3D12View.cpp
+++ b/Engine/Source/Runtime/D3D12RHI/Private/D3D12View.cpp
@@ -290,6 +290,8 @@ void FD3D12View::CreateView(FResourceInfo const& InResource, FNullDescPtr NullDe
 			Device->GetDefaultViews().*NullDescriptor,
 			Translate(HeapType)
 		);
+
+		OfflineCpuHandle.IncrementVersion();
 	}
 
 	UpdateBindlessSlot(Reason);
@@ -340,6 +342,8 @@ void FD3D12ConstantBufferView::ResourceRenamed(FD3D12BaseShaderResource* InRenam
 void FD3D12ConstantBufferView::UpdateDescriptor()
 {
 	GetParentDevice()->GetDevice()->CreateConstantBufferView(&D3DViewDesc, OfflineCpuHandle);
+
+	OfflineCpuHandle.IncrementVersion();
 }
 
 
@@ -361,6 +365,8 @@ void FD3D12RenderTargetView::UpdateDescriptor()
 		&D3DViewDesc,
 		OfflineCpuHandle
 	);
+
+	OfflineCpuHandle.IncrementVersion();
 }
 
 
@@ -382,4 +388,6 @@ void FD3D12DepthStencilView::UpdateDescriptor()
 		&D3DViewDesc,
 		OfflineCpuHandle
 	);
+
+	OfflineCpuHandle.IncrementVersion();
 }
