Adding new D3D12 "UAV Aliasing" code path to use the new APIs that allow us to create UINT views of BCn textures.

diff --git a/Engine/Source/Runtime/D3D12RHI/Private/D3D12Adapter.cpp b/Engine/Source/Runtime/D3D12RHI/Private/D3D12Adapter.cpp
--- a/Engine/Source/Runtime/D3D12RHI/Private/D3D12Adapter.cpp
+++ b/Engine/Source/Runtime/D3D12RHI/Private/D3D12Adapter.cpp
@@ -934,6 +934,12 @@ void FD3D12Adapter::InitializeDevices()
 				UE_LOG(LogD3D12RHI, Log, TEXT("ID3D12Device11 is supported."));
 			}
 #endif
+#if D3D12_MAX_DEVICE_INTERFACE >= 12
+			if (SUCCEEDED(RootDevice->QueryInterface(IID_PPV_ARGS(RootDevice12.GetInitReference()))))
+			{
+				UE_LOG(LogD3D12RHI, Log, TEXT("ID3D12Device12 is supported."));
+			}
+#endif
 
 			const bool bRenderDocPresent = D3D12RHI_IsRenderDocPresent(RootDevice);
 
diff --git a/Engine/Source/Runtime/D3D12RHI/Private/D3D12Adapter.h b/Engine/Source/Runtime/D3D12RHI/Private/D3D12Adapter.h
--- a/Engine/Source/Runtime/D3D12RHI/Private/D3D12Adapter.h
+++ b/Engine/Source/Runtime/D3D12RHI/Private/D3D12Adapter.h
@@ -186,6 +186,9 @@ public:
 #if D3D12_MAX_DEVICE_INTERFACE >= 11
 	FORCEINLINE ID3D12Device11* GetD3DDevice11() { return RootDevice11; }
 #endif
+#if D3D12_MAX_DEVICE_INTERFACE >= 12
+	FORCEINLINE ID3D12Device12* GetD3DDevice12() { return RootDevice12; }
+#endif
 
 	FORCEINLINE IDXGIFactory2* GetDXGIFactory2() const { return DxgiFactory2; }
 #if DXGI_MAX_FACTORY_INTERFACE >= 3
@@ -494,6 +497,9 @@ protected:
 #if D3D12_MAX_DEVICE_INTERFACE >= 11
 	TRefCountPtr<ID3D12Device11> RootDevice11;
 #endif
+#if D3D12_MAX_DEVICE_INTERFACE >= 12
+	TRefCountPtr<ID3D12Device12> RootDevice12;
+#endif
 
 	TRefCountPtr<IDXGIFactory2> DxgiFactory2;
 #if DXGI_MAX_FACTORY_INTERFACE >= 3
diff --git a/Engine/Source/Runtime/D3D12RHI/Private/D3D12Allocation.cpp b/Engine/Source/Runtime/D3D12RHI/Private/D3D12Allocation.cpp
--- a/Engine/Source/Runtime/D3D12RHI/Private/D3D12Allocation.cpp
+++ b/Engine/Source/Runtime/D3D12RHI/Private/D3D12Allocation.cpp
@@ -1764,7 +1764,8 @@ HRESULT FD3D12TextureAllocatorPool::AllocateTexture(
 	// The top mip level must be less than 64 KB to use 4 KB alignment
 	bool b4KAligment = FD3D12Texture::CanBe4KAligned(Desc, (EPixelFormat)UEFormat);
 	Desc.Alignment = b4KAligment ?	D3D12_SMALL_RESOURCE_PLACEMENT_ALIGNMENT : (Desc.SampleDesc.Count > 1 ? D3D12_DEFAULT_MSAA_RESOURCE_PLACEMENT_ALIGNMENT : D3D12_DEFAULT_RESOURCE_PLACEMENT_ALIGNMENT);
-	const D3D12_RESOURCE_ALLOCATION_INFO Info = GetParentDevice()->GetDevice()->GetResourceAllocationInfo(0, 1, &Desc);
+
+	const D3D12_RESOURCE_ALLOCATION_INFO Info = GetParentDevice()->GetResourceAllocationInfoUncached(Desc);
 
 	const bool bIsRenderTarget = EnumHasAnyFlags(Desc.Flags, D3D12_RESOURCE_FLAG_ALLOW_RENDER_TARGET | D3D12_RESOURCE_FLAG_ALLOW_DEPTH_STENCIL);
 	const bool bIsReadOnly = !bIsRenderTarget && !EnumHasAnyFlags(Desc.Flags, D3D12_RESOURCE_FLAG_ALLOW_UNORDERED_ACCESS) && !Desc.NeedsUAVAliasWorkarounds();
@@ -1875,7 +1876,7 @@ HRESULT FD3D12TextureAllocatorPool::AllocateTexture(
 		Desc.Alignment = TextureCanBe4KAligned(Desc, UEFormat) ?
 			D3D12_SMALL_RESOURCE_PLACEMENT_ALIGNMENT :
 			D3D12_DEFAULT_RESOURCE_PLACEMENT_ALIGNMENT;
-		const D3D12_RESOURCE_ALLOCATION_INFO Info = Device->GetDevice()->GetResourceAllocationInfo(0, 1, &Desc);
+		const D3D12_RESOURCE_ALLOCATION_INFO Info = Device->GetResourceAllocationInfoUncached(Desc);
 
 		TRefCountPtr<FD3D12SegHeap> BackingHeap;
 		const uint32 Offset = ReadOnlyTexturePool.Allocate(Info.SizeInBytes, Info.Alignment, BackingHeap);
@@ -1946,7 +1947,7 @@ HRESULT FD3D12TextureAllocator::AllocateTexture(FD3D12ResourceDesc Desc, const D
 
 	TextureLocation.Clear();
 
-	D3D12_RESOURCE_ALLOCATION_INFO Info = Device->GetDevice()->GetResourceAllocationInfo(0, 1, &Desc);
+	D3D12_RESOURCE_ALLOCATION_INFO Info = Device->GetResourceAllocationInfoUncached(Desc);
 
 	if (Info.SizeInBytes < D3D12_DEFAULT_RESOURCE_PLACEMENT_ALIGNMENT)
 	{
@@ -2001,7 +2002,7 @@ HRESULT FD3D12TextureAllocatorPool::AllocateTexture(FD3D12ResourceDesc Desc, con
 	FD3D12Resource* Resource = nullptr;
 
 	const D3D12_HEAP_PROPERTIES HeapProps = CD3DX12_HEAP_PROPERTIES(D3D12_HEAP_TYPE_DEFAULT, GetGPUMask().GetNative(), GetVisibilityMask().GetNative());
-	const D3D12_RESOURCE_ALLOCATION_INFO Info = GetParentDevice()->GetDevice()->GetResourceAllocationInfo(0, 1, &Desc);
+	const D3D12_RESOURCE_ALLOCATION_INFO Info = GetParentDevice()->GetResourceAllocationInfoUncached(Desc);
 
 	// UAV Aliasing needs a Heap to create the aliased resource in.
 	if (Desc.NeedsUAVAliasWorkarounds())
diff --git a/Engine/Source/Runtime/D3D12RHI/Private/D3D12Device.cpp b/Engine/Source/Runtime/D3D12RHI/Private/D3D12Device.cpp
--- a/Engine/Source/Runtime/D3D12RHI/Private/D3D12Device.cpp
+++ b/Engine/Source/Runtime/D3D12RHI/Private/D3D12Device.cpp
@@ -547,9 +547,55 @@ void FD3D12Device::BlockUntilIdle()
 	}
 }
 
+D3D12_RESOURCE_ALLOCATION_INFO FD3D12Device::GetResourceAllocationInfoUncached(const FD3D12ResourceDesc& InDesc)
+{
+	D3D12_RESOURCE_ALLOCATION_INFO Result;
+
+#if INTEL_EXTENSIONS
+	if (InDesc.bRequires64BitAtomicSupport && IsRHIDeviceIntel() && GDX12INTCAtomicUInt64Emulation)
+	{
+		FD3D12ResourceDesc LocalDesc = InDesc;
+
+		INTC_D3D12_RESOURCE_DESC_0001 IntelLocalDesc{};
+		IntelLocalDesc.pD3D12Desc = &LocalDesc;
+		IntelLocalDesc.EmulatedTyped64bitAtomics = true;
+
+		Result = INTC_D3D12_GetResourceAllocationInfo(FD3D12DynamicRHI::GetD3DRHI()->GetIntelExtensionContext(), 0, 1, &IntelLocalDesc);
+	}
+	else
+#endif
+#if D3D12RHI_SUPPORTS_UNCOMPRESSED_UAV
+	if (InDesc.SupportsUncompressedUAV())
+	{
+		// Convert the desc to the version required by GetResourceAllocationInfo3
+		const CD3DX12_RESOURCE_DESC1 LocalDesc(InDesc);
+
+		const TArray<DXGI_FORMAT, TInlineAllocator<4>> CastableFormats = InDesc.GetCastableFormats();
+
+		const UINT32 NumCastableFormats = CastableFormats.Num();
+		D3D12_RESOURCE_ALLOCATION_INFO1* NoExtraAllocationInfo = nullptr;
+
+		const DXGI_FORMAT* const Formats[] = { CastableFormats.GetData() };
+
+		Result = GetParentAdapter()->GetD3DDevice12()->GetResourceAllocationInfo3(0, 1, &LocalDesc, &NumCastableFormats, Formats, NoExtraAllocationInfo);
+	}
+	else
+#endif
+	{
+		Result = GetDevice()->GetResourceAllocationInfo(0, 1, &InDesc);
+		if (Result.SizeInBytes == UINT64_MAX)
+		{
+			// The description provided caused an error per the docs. This will almost certainly crash outside this fn.
+			UE_LOG(LogD3D12RHI, Error, TEXT("D3D12 GetResourceAllocationInfo failed - likely a resource was requested that has invalid allocation info (e.g. is an invalid texture size)"));
+			UE_LOG(LogD3D12RHI, Error, TEXT("     W %llu H %d depth %d mips %d pf %d"), InDesc.Width, InDesc.Height, InDesc.DepthOrArraySize, InDesc.MipLevels, InDesc.PixelFormat);
+		}
+	}
+	return Result;
+}
+
 D3D12_RESOURCE_ALLOCATION_INFO FD3D12Device::GetResourceAllocationInfo(const FD3D12ResourceDesc& InDesc)
 {
-	uint64 Hash = CityHash64((const char*)&InDesc, sizeof(FD3D12ResourceDesc));
+	const uint64 Hash = CityHash64((const char*)&InDesc, sizeof(FD3D12ResourceDesc));
 
 	// By default there'll be more threads trying to read this than to write it.
 	ResourceAllocationInfoMapMutex.ReadLock();
@@ -562,29 +608,7 @@ D3D12_RESOURCE_ALLOCATION_INFO FD3D12Device::GetResourceAllocationInfo(const FD3
 	}
 	else
 	{
-		D3D12_RESOURCE_ALLOCATION_INFO Result;
-#if INTEL_EXTENSIONS
-		if (InDesc.bRequires64BitAtomicSupport && IsRHIDeviceIntel() && GDX12INTCAtomicUInt64Emulation)
-		{
-			FD3D12ResourceDesc LocalDesc = InDesc;
-
-			INTC_D3D12_RESOURCE_DESC_0001 IntelLocalDesc{};
-			IntelLocalDesc.pD3D12Desc = &LocalDesc;
-			IntelLocalDesc.EmulatedTyped64bitAtomics = true;
-
-			Result = INTC_D3D12_GetResourceAllocationInfo(FD3D12DynamicRHI::GetD3DRHI()->GetIntelExtensionContext(), 0, 1, &IntelLocalDesc);
-		}
-		else
-#endif
-		{
-			Result = GetDevice()->GetResourceAllocationInfo(0, 1, &InDesc);
-		    if (Result.SizeInBytes == UINT64_MAX)
-		    {
-			    // The description provided caused an error per the docs. This will almost certainly crash outside this fn.
-			    UE_LOG(LogD3D12RHI, Error, TEXT("D3D12 GetResourceAllocationInfo failed - likely a resource was requested that has invalid allocation info (e.g. is an invalid texture size)"));
-			    UE_LOG(LogD3D12RHI, Error, TEXT("     W %llu H %d depth %d mips %d pf %d"), InDesc.Width, InDesc.Height, InDesc.DepthOrArraySize, InDesc.MipLevels, InDesc.PixelFormat);
-			}
-		}
+		const D3D12_RESOURCE_ALLOCATION_INFO Result = GetResourceAllocationInfoUncached(InDesc);
 
 		ResourceAllocationInfoMapMutex.WriteLock();
 		// Try search again with write lock because could have been added already
diff --git a/Engine/Source/Runtime/D3D12RHI/Private/D3D12Device.h b/Engine/Source/Runtime/D3D12RHI/Private/D3D12Device.h
--- a/Engine/Source/Runtime/D3D12RHI/Private/D3D12Device.h
+++ b/Engine/Source/Runtime/D3D12RHI/Private/D3D12Device.h
@@ -263,6 +263,7 @@ public:
 
 	// Misc
 	void BlockUntilIdle();
+	D3D12_RESOURCE_ALLOCATION_INFO GetResourceAllocationInfoUncached(const FD3D12ResourceDesc& InDesc);
 	D3D12_RESOURCE_ALLOCATION_INFO GetResourceAllocationInfo(const FD3D12ResourceDesc& InDesc);
 	TUniquePtr<FD3D12DiagnosticBuffer> CreateDiagnosticBuffer(const D3D12_RESOURCE_DESC& Desc, const TCHAR* Name);
 
diff --git a/Engine/Source/Runtime/D3D12RHI/Private/D3D12RHICommon.h b/Engine/Source/Runtime/D3D12RHI/Private/D3D12RHICommon.h
--- a/Engine/Source/Runtime/D3D12RHI/Private/D3D12RHICommon.h
+++ b/Engine/Source/Runtime/D3D12RHI/Private/D3D12RHICommon.h
@@ -10,6 +10,12 @@ D3D12RHICommon.h: Common D3D12 RHI definitions for Windows.
 #include "MultiGPU.h"
 #include "Containers/StaticArray.h"
 
+#if D3D12_MAX_DEVICE_INTERFACE >= 12
+	#define D3D12RHI_SUPPORTS_UNCOMPRESSED_UAV 1
+#else
+	#define D3D12RHI_SUPPORTS_UNCOMPRESSED_UAV 0
+#endif
+
 DECLARE_STATS_GROUP(TEXT("D3D12RHI"), STATGROUP_D3D12RHI, STATCAT_Advanced);
 DECLARE_STATS_GROUP(TEXT("D3D12RHI: Memory"), STATGROUP_D3D12Memory, STATCAT_Advanced);
 DECLARE_STATS_GROUP(TEXT("D3D12RHI: Memory Details"), STATGROUP_D3D12MemoryDetails, STATCAT_Advanced);
diff --git a/Engine/Source/Runtime/D3D12RHI/Private/D3D12Resources.cpp b/Engine/Source/Runtime/D3D12RHI/Private/D3D12Resources.cpp
--- a/Engine/Source/Runtime/D3D12RHI/Private/D3D12Resources.cpp
+++ b/Engine/Source/Runtime/D3D12RHI/Private/D3D12Resources.cpp
@@ -43,6 +43,24 @@ void ID3D12ResourceAllocator::AllocateTexture(uint32 GPUIndex, D3D12_HEAP_TYPE I
 	AllocateResource(GPUIndex, InHeapType, Desc, Info.SizeInBytes, Info.Alignment, InResourceStateMode, InCreateState, InClearValue, InName, ResourceLocation);
 }
 
+#if D3D12RHI_SUPPORTS_UNCOMPRESSED_UAV
+TArray<DXGI_FORMAT, TInlineAllocator<4>> FD3D12ResourceDesc::GetCastableFormats() const
+{
+	TArray<DXGI_FORMAT, TInlineAllocator<4>> Result;
+
+	// We have to add the 'implied' castable formats for SRVs. Since we don't have any sRGB flags here, just add both formats.
+	Result.Emplace(FindShaderResourceDXGIFormat(Format, true));
+	Result.Emplace(FindShaderResourceDXGIFormat(Format, false));
+
+	if (UAVPixelFormat != PF_Unknown)
+	{
+		// Add the uncompressed UAV format we want
+		Result.Emplace((DXGI_FORMAT)GPixelFormats[UAVPixelFormat].PlatformFormat);
+	}
+
+	return Result;
+}
+#endif // D3D12RHI_SUPPORTS_UNCOMPRESSED_UAV
 
 /////////////////////////////////////////////////////////////////////
 //	FD3D12 Resource
@@ -276,8 +294,8 @@ void FD3D12Resource::StartTrackingForResidency()
 
 	check(IsGPUOnly(HeapType));	// This is checked at a higher level before calling this function.
 	check(D3DX12Residency::IsInitialized(ResidencyHandle) == false);
-	const D3D12_RESOURCE_DESC ResourceDesc = Resource->GetDesc();
-	const D3D12_RESOURCE_ALLOCATION_INFO Info = GetParentDevice()->GetDevice()->GetResourceAllocationInfo(0, 1, &ResourceDesc);
+
+	const D3D12_RESOURCE_ALLOCATION_INFO Info = GetParentDevice()->GetResourceAllocationInfoUncached(Desc);
 
 	D3DX12Residency::Initialize(ResidencyHandle, Resource.GetReference(), Info.SizeInBytes, this);
 	D3DX12Residency::BeginTrackingObject(GetParentDevice()->GetResidencyManager(), ResidencyHandle);
@@ -438,6 +456,38 @@ HRESULT FD3D12Adapter::CreateCommittedResource(const FD3D12ResourceDesc& InDesc,
 		hr = INTC_D3D12_CreateCommittedResource(FD3D12DynamicRHI::GetD3DRHI()->GetIntelExtensionContext(), &HeapProps, HeapFlags, &IntelLocalDesc, InInitialState, ClearValue, IID_PPV_ARGS(pResource.GetInitReference()));
 	}
 	else
+#endif
+#if D3D12RHI_SUPPORTS_UNCOMPRESSED_UAV
+	if (InDesc.SupportsUncompressedUAV())
+	{
+		// Convert the desc to the version required by CreateCommittedResource3
+		const CD3DX12_RESOURCE_DESC1 LocalDesc1(LocalDesc);
+
+		// Common layout is the required starting state for any "legacy" transitions
+		const D3D12_BARRIER_LAYOUT InitialLayout = D3D12_BARRIER_LAYOUT_COMMON;
+		
+		ID3D12ProtectedResourceSession* ProtectedSession = nullptr;
+
+		const TArray<DXGI_FORMAT, TInlineAllocator<4>> CastableFormats = InDesc.GetCastableFormats();
+
+		hr = RootDevice12->CreateCommittedResource3(
+			&HeapProps,
+			HeapFlags,
+			&LocalDesc1,
+			InitialLayout,
+			ClearValue,
+			ProtectedSession,
+			CastableFormats.Num(),
+			CastableFormats.GetData(),
+			IID_PPV_ARGS(pResource.GetInitReference()));
+
+		if (SUCCEEDED(hr))
+		{
+			// Change the initial state to match the initial layout on construction
+			InInitialState = D3D12_RESOURCE_STATE_COMMON;
+		}
+	}
+	else
 #endif
 	{
 		hr = RootDevice->CreateCommittedResource(&HeapProps, HeapFlags, &LocalDesc, InInitialState, ClearValue, IID_PPV_ARGS(pResource.GetInitReference()));
@@ -498,7 +548,37 @@ HRESULT FD3D12Adapter::CreateReservedResource(const FD3D12ResourceDesc& InDesc,
 	}
 #endif // D3D12_RHI_RAYTRACING
 
-	HRESULT hr = RootDevice->CreateReservedResource(&LocalDesc, InInitialState, ClearValue, IID_PPV_ARGS(pResource.GetInitReference()));
+	HRESULT hr = S_OK;
+	
+#if D3D12RHI_SUPPORTS_UNCOMPRESSED_UAV
+	if (InDesc.SupportsUncompressedUAV())
+	{
+		// Common layout is the require starting state for any "legacy" transitions
+		const D3D12_BARRIER_LAYOUT InitialLayout = D3D12_BARRIER_LAYOUT_COMMON;
+
+		ID3D12ProtectedResourceSession* ProtectedSession = nullptr;
+		const TArray<DXGI_FORMAT, TInlineAllocator<4>> CastableFormats = InDesc.GetCastableFormats();
+
+		hr = RootDevice12->CreateReservedResource2(
+			&LocalDesc,
+			InitialLayout,
+			ClearValue,
+			ProtectedSession,
+			CastableFormats.Num(),
+			CastableFormats.GetData(),
+			IID_PPV_ARGS(pResource.GetInitReference()));
+
+		if (SUCCEEDED(hr))
+		{
+			// Change the initial state to match the initial layout on construction
+			InInitialState = D3D12_RESOURCE_STATE_COMMON;
+		}
+	}
+	else
+#endif
+	{
+		hr = RootDevice->CreateReservedResource(&LocalDesc, InInitialState, ClearValue, IID_PPV_ARGS(pResource.GetInitReference()));
+	}
 
 	if (SUCCEEDED(hr))
 	{
@@ -548,6 +628,35 @@ HRESULT FD3D12Adapter::CreatePlacedResource(const FD3D12ResourceDesc& InDesc, FD
 		hr = INTC_D3D12_CreatePlacedResource(FD3D12DynamicRHI::GetD3DRHI()->GetIntelExtensionContext(), Heap, HeapOffset, &IntelLocalDesc, InInitialState, ClearValue, IID_PPV_ARGS(pResource.GetInitReference()));
 	}
 	else
+#endif
+#if D3D12RHI_SUPPORTS_UNCOMPRESSED_UAV
+	if (InDesc.SupportsUncompressedUAV())
+	{
+		// Convert the desc to the version required by CreatePlacedResource2
+		const CD3DX12_RESOURCE_DESC1 LocalDesc(InDesc);
+
+		// Common layout is the required starting state for any "legacy" transitions
+		const D3D12_BARRIER_LAYOUT InitialLayout = D3D12_BARRIER_LAYOUT_COMMON;
+		
+		const TArray<DXGI_FORMAT, TInlineAllocator<4>> CastableFormats = InDesc.GetCastableFormats();
+
+		hr = RootDevice10->CreatePlacedResource2(
+			Heap,
+			HeapOffset,
+			&LocalDesc,
+			InitialLayout,
+			ClearValue,
+			CastableFormats.Num(),
+			CastableFormats.GetData(),
+			IID_PPV_ARGS(pResource.GetInitReference()));
+
+		if (SUCCEEDED(hr))
+		{
+			// Change the initial state to match the initial layout on construction
+			InInitialState = D3D12_RESOURCE_STATE_COMMON;
+		}
+	}
+	else
 #endif
 	{
 		hr = RootDevice->CreatePlacedResource(Heap, HeapOffset, &InDesc, InInitialState, ClearValue, IID_PPV_ARGS(pResource.GetInitReference()));
@@ -674,7 +783,7 @@ void FD3D12Adapter::CreateUAVAliasResource(D3D12_CLEAR_VALUE* ClearValuePtr, con
 	const FD3D12Heap* const ResourceHeap = SourceResource->GetHeap();
 
 	const EPixelFormat SourceFormat = SourceDesc.PixelFormat;
-	const EPixelFormat AliasTextureFormat = SourceDesc.UAVAliasPixelFormat;
+	const EPixelFormat AliasTextureFormat = SourceDesc.UAVPixelFormat;
 
 	if (ensure(ResourceHeap != nullptr) && ensure(SourceFormat != PF_Unknown) && SourceFormat != AliasTextureFormat)
 	{
@@ -697,7 +806,7 @@ void FD3D12Adapter::CreateUAVAliasResource(D3D12_CLEAR_VALUE* ClearValuePtr, con
 
 		EnumAddFlags(AliasTextureDesc.Flags, D3D12_RESOURCE_FLAG_ALLOW_UNORDERED_ACCESS);
 
-		AliasTextureDesc.UAVAliasPixelFormat = PF_Unknown;
+		AliasTextureDesc.UAVPixelFormat = PF_Unknown;
 
 		TRefCountPtr<ID3D12Resource> pAliasResource;
 		HRESULT AliasHR = GetD3DDevice()->CreatePlacedResource(
@@ -954,7 +1063,7 @@ void FD3D12ResourceLocation::UpdateStandAloneStats(bool bIncrement)
 {
 	if (UnderlyingResource->GetHeapType() == D3D12_HEAP_TYPE_DEFAULT)
 	{
-		D3D12_RESOURCE_DESC Desc = UnderlyingResource->GetDesc();
+		FD3D12ResourceDesc Desc = UnderlyingResource->GetDesc();
 		bool bIsBuffer = (Desc.Dimension == D3D12_RESOURCE_DIMENSION_BUFFER);
 		bool bIsRenderTarget = (Desc.Flags & D3D12_RESOURCE_FLAG_ALLOW_RENDER_TARGET || Desc.Flags & D3D12_RESOURCE_FLAG_ALLOW_DEPTH_STENCIL);
 		bool bIsUAV = (Desc.Flags & D3D12_RESOURCE_FLAG_ALLOW_UNORDERED_ACCESS) > 0;
@@ -966,7 +1075,7 @@ void FD3D12ResourceLocation::UpdateStandAloneStats(bool bIncrement)
 		}
 
 		// Get the desired size and allocated size for stand alone resources - allocated are very slow anyway
-		D3D12_RESOURCE_ALLOCATION_INFO Info = UnderlyingResource->GetParentDevice()->GetDevice()->GetResourceAllocationInfo(0, 1, &Desc);
+		D3D12_RESOURCE_ALLOCATION_INFO Info = UnderlyingResource->GetParentDevice()->GetResourceAllocationInfoUncached(Desc);
 
 		int64 SizeInBytes = bIncrement ? Info.SizeInBytes : -(int64)Info.SizeInBytes;
 		int32 Count = bIncrement ? 1 : -1;
diff --git a/Engine/Source/Runtime/D3D12RHI/Private/D3D12Resources.h b/Engine/Source/Runtime/D3D12RHI/Private/D3D12Resources.h
--- a/Engine/Source/Runtime/D3D12RHI/Private/D3D12Resources.h
+++ b/Engine/Source/Runtime/D3D12RHI/Private/D3D12Resources.h
@@ -119,7 +119,7 @@ struct FD3D12ResourceDesc : public D3D12_RESOURCE_DESC
 	EPixelFormat PixelFormat{ PF_Unknown };
 
 	// PixelFormat for the Resource that aliases our current resource.
-	EPixelFormat UAVAliasPixelFormat{ PF_Unknown };
+	EPixelFormat UAVPixelFormat{ PF_Unknown };
 
 #if D3D12RHI_NEEDS_VENDOR_EXTENSIONS
 	bool bRequires64BitAtomicSupport{ false };
@@ -127,8 +127,16 @@ struct FD3D12ResourceDesc : public D3D12_RESOURCE_DESC
 
 	bool bReservedResource { false };
 
-	// Used primarily to help treat this resource description as writable.
-	inline bool NeedsUAVAliasWorkarounds() const { return UAVAliasPixelFormat != PF_Unknown; }
+	// If we support the new format list casting, use the newer APIs; otherwise, fall back to our UAV Aliasing approach.
+#if D3D12RHI_SUPPORTS_UNCOMPRESSED_UAV
+	inline bool NeedsUAVAliasWorkarounds() const { return false; }
+	inline bool SupportsUncompressedUAV() const { return UAVPixelFormat != PF_Unknown; }
+
+	TArray<DXGI_FORMAT, TInlineAllocator<4>> GetCastableFormats() const;
+#else
+	inline bool NeedsUAVAliasWorkarounds() const { return UAVPixelFormat != PF_Unknown; }
+	inline bool SupportsUncompressedUAV() const { return false; }
+#endif
 };
 
 class FD3D12Resource : public FThreadSafeRefCountedObject, public FD3D12DeviceChild, public FD3D12MultiNodeGPUObject
diff --git a/Engine/Source/Runtime/D3D12RHI/Private/D3D12Texture.cpp b/Engine/Source/Runtime/D3D12RHI/Private/D3D12Texture.cpp
--- a/Engine/Source/Runtime/D3D12RHI/Private/D3D12Texture.cpp
+++ b/Engine/Source/Runtime/D3D12RHI/Private/D3D12Texture.cpp
@@ -640,13 +640,12 @@ FD3D12ResourceDesc FD3D12DynamicRHI::GetResourceDesc(const FRHITextureDesc& Text
 		ResourceDesc.Flags |= D3D12_RESOURCE_FLAG_DENY_SHADER_RESOURCE;
 	}
 
-	bool bBCTextureNeedsUAVAlias = EnumHasAnyFlags(TextureDesc.Flags, TexCreate_UAV) && IsBlockCompressedFormat(TextureDesc.Format);
-	if (bBCTextureNeedsUAVAlias)
+	if (EnumHasAnyFlags(TextureDesc.Flags, ETextureCreateFlags::UAV) && IsBlockCompressedFormat(TextureDesc.Format))
 	{
-		ResourceDesc.UAVAliasPixelFormat = GetBlockCompressedFormatUAVAliasFormat(TextureDesc.Format);
+		ResourceDesc.UAVPixelFormat = GetBlockCompressedFormatUAVAliasFormat(TextureDesc.Format);
 	}
 
-	if (EnumHasAnyFlags(TextureDesc.Flags, TexCreate_UAV) && !bBCTextureNeedsUAVAlias)
+	if (EnumHasAnyFlags(TextureDesc.Flags, TexCreate_UAV) && !ResourceDesc.NeedsUAVAliasWorkarounds())
 	{
 		ResourceDesc.Flags |= D3D12_RESOURCE_FLAG_ALLOW_UNORDERED_ACCESS;
 	}
@@ -655,6 +654,8 @@ FD3D12ResourceDesc FD3D12DynamicRHI::GetResourceDesc(const FRHITextureDesc& Text
 
 #if D3D12RHI_NEEDS_VENDOR_EXTENSIONS
 	ResourceDesc.bRequires64BitAtomicSupport = EnumHasAnyFlags(TextureDesc.Flags, ETextureCreateFlags::Atomic64Compatible);
+
+	checkf(!(ResourceDesc.bRequires64BitAtomicSupport&& ResourceDesc.SupportsUncompressedUAV()), TEXT("Intel resource creation extensions don't support the new resource casting parameters."));
 #endif
 
 	// Check if the 4K aligment is possible
@@ -972,10 +973,8 @@ FD3D12Texture* FD3D12DynamicRHI::CreateD3D12Texture(const FRHITextureCreateDesc&
 
 		check(Location.IsValid());
 
-		bool bBCTextureNeedsUAVAlias = EnumHasAnyFlags(CreateDesc.Flags, TexCreate_UAV) && IsBlockCompressedFormat(CreateDesc.Format);
-		if (bBCTextureNeedsUAVAlias)
+		if (ResourceDesc.NeedsUAVAliasWorkarounds())
 		{
-			check(ResourceDesc.UAVAliasPixelFormat != PF_Unknown);
 			Adapter->CreateUAVAliasResource(ClearValuePtr, CreateDesc.DebugName, Location);
 		}
 
diff --git a/Engine/Source/Runtime/D3D12RHI/Private/Windows/WindowsD3D12Device.cpp b/Engine/Source/Runtime/D3D12RHI/Private/Windows/WindowsD3D12Device.cpp
--- a/Engine/Source/Runtime/D3D12RHI/Private/Windows/WindowsD3D12Device.cpp
+++ b/Engine/Source/Runtime/D3D12RHI/Private/Windows/WindowsD3D12Device.cpp
@@ -1725,7 +1725,17 @@ void FD3D12DynamicRHI::Init()
 		GRHIVariableRateShadingImageTileMaxHeight = 1;
 	}
 
-	GRHISupportsUAVFormatAliasing = (GetAdapter().GetResourceHeapTier() > D3D12_RESOURCE_HEAP_TIER_1 && IsRHIDeviceNVIDIA());
+#if D3D12RHI_SUPPORTS_UNCOMPRESSED_UAV
+	D3D12_FEATURE_DATA_D3D12_OPTIONS12 Options12{};
+	if (SUCCEEDED(GetAdapter().GetD3DDevice()->CheckFeatureSupport(D3D12_FEATURE_D3D12_OPTIONS12, &Options12, sizeof(Options12))))
+	{
+		GRHIGlobals.SupportsUAVFormatAliasing = (Options12.RelaxedFormatCastingSupported != 0)
+			// We require ID3D12Device12 for GetResourceAllocationInfo3
+			&& GetAdapter().GetD3DDevice12() != nullptr;
+	}
+#else
+	GRHIGlobals.SupportsUAVFormatAliasing = (GetAdapter().GetResourceHeapTier() > D3D12_RESOURCE_HEAP_TIER_1 && IsRHIDeviceNVIDIA());
+#endif
 
 	InitializeSubmissionPipe();
 
