Moving USE_STATIC_ROOT_SIGNATURE into D3D12RHI Private.

diff --git a/Engine/Platforms/Hololens/Source/Runtime/D3D12RHI/Private/D3D12RHIBasePrivate.h b/Engine/Platforms/Hololens/Source/Runtime/D3D12RHI/Private/D3D12RHIBasePrivate.h
--- a/Engine/Platforms/Hololens/Source/Runtime/D3D12RHI/Private/D3D12RHIBasePrivate.h
+++ b/Engine/Platforms/Hololens/Source/Runtime/D3D12RHI/Private/D3D12RHIBasePrivate.h
@@ -33,6 +33,8 @@ typedef ID3D12GraphicsCommandList ID3D12CopyCommandList;
 // but make code path shared when it becomes available
 #define FD3D12_HEAP_FLAG_CREATE_NOT_ZEROED D3D12_HEAP_FLAG_NONE
 
+#define USE_STATIC_ROOT_SIGNATURE			0
+
 #include "D3D12Util.h"
 
 // Windows desktop version is used.
diff --git a/Engine/Source/Runtime/D3D12RHI/Private/Windows/D3D12RHIBasePrivate.h b/Engine/Source/Runtime/D3D12RHI/Private/Windows/D3D12RHIBasePrivate.h
--- a/Engine/Source/Runtime/D3D12RHI/Private/Windows/D3D12RHIBasePrivate.h
+++ b/Engine/Source/Runtime/D3D12RHI/Private/Windows/D3D12RHIBasePrivate.h
@@ -29,6 +29,8 @@ typedef ID3D12GraphicsCommandList ID3D12CopyCommandList;
 #define D3D12RHI_NEEDS_VENDOR_EXTENSIONS     1
 #define D3D12RHI_NEEDS_SHADER_FEATURE_CHECKS 1
 
+#define USE_STATIC_ROOT_SIGNATURE			0
+
 #include "D3D12Util.h"
 #include "WindowsD3D12DiskCache.h"
 #include "WindowsD3D12PipelineState.h"
diff --git a/Engine/Source/Runtime/D3D12RHI/Public/D3D12RHI.h b/Engine/Source/Runtime/D3D12RHI/Public/D3D12RHI.h
--- a/Engine/Source/Runtime/D3D12RHI/Public/D3D12RHI.h
+++ b/Engine/Source/Runtime/D3D12RHI/Public/D3D12RHI.h
@@ -29,11 +29,6 @@
 // Note: Using root descriptors significantly increases the size of root signatures (each root descriptor is 2 DWORDs).
 #define MAX_ROOT_CBVS	MAX_CBS
 
-// So outside callers can override this
-#ifndef USE_STATIC_ROOT_SIGNATURE
-	#define USE_STATIC_ROOT_SIGNATURE 0
-#endif
-
 // How many residency packets can be in flight before the rendering thread
 // blocks for them to drain. Should be ~ NumBufferedFrames * AvgNumSubmissionsPerFrame i.e.
 // enough to ensure that the GPU is rarely blocked by residency work
