Fix a CPU crash in FD3D12RayTracingLocalResourceBinder::CreateTransientConstantBuffer because memcpy was using incorrect source data size:
* Use low-level Allocate to create a constant buffer view instead that is already properly aligned.

diff --git a/Engine/Source/Runtime/D3D12RHI/Private/D3D12RayTracing.cpp b/Engine/Source/Runtime/D3D12RHI/Private/D3D12RayTracing.cpp
--- a/Engine/Source/Runtime/D3D12RHI/Private/D3D12RayTracing.cpp
+++ b/Engine/Source/Runtime/D3D12RHI/Private/D3D12RayTracing.cpp
@@ -4571,17 +4571,10 @@ struct FD3D12RayTracingGlobalResourceBinder
 		CommandContext.GraphicsCommandList()->SetComputeRootDescriptorTable(SlotIndex, DescriptorTable);
 	}
 
-#if USE_STATIC_ROOT_SIGNATURE
-	D3D12_CPU_DESCRIPTOR_HANDLE CreateTransientConstantBufferView(FD3D12ResourceLocation& ResourceLocation, uint32 DataSize)
-	{
-		checkf(0, TEXT("Loose parameters and transient constant buffers are not implemented for global ray tracing shaders (raygen, miss, callable)"));
-		return D3D12_CPU_DESCRIPTOR_HANDLE {};
-	}
-#endif // USE_STATIC_ROOT_SIGNATURE
-
-	void CreateTransientConstantBuffer(FD3D12ResourceLocation& ResourceLocation, const void* Data, uint32 DataSize)
+	FD3D12ConstantBufferView* CreateTransientConstantBuffer(FD3D12ResourceLocation& ResourceLocation, const void* Data, uint32 DataSize)
 	{
 		checkf(0, TEXT("Loose parameters and transient constant buffers are not implemented for global ray tracing shaders (raygen, miss, callable)"));
+		return nullptr;
 	}
 
 	void AddResourceReference(FD3D12Resource* D3D12Resource)
@@ -4664,28 +4657,27 @@ struct FD3D12RayTracingLocalResourceBinder
 		ShaderTable.SetLocalShaderParameters(ShaderTableOffset, RecordIndex, BindOffset, DescriptorTable);
 	}
 
-#if USE_STATIC_ROOT_SIGNATURE
-	D3D12_CPU_DESCRIPTOR_HANDLE CreateTransientConstantBufferView(FD3D12ResourceLocation& ResourceLocation, uint32 DataSize)
-	{
-		FD3D12ConstantBufferView* BufferView = new FD3D12ConstantBufferView(GetDevice());
-		BufferView->CreateView(&ResourceLocation, 0, DataSize);
-		ShaderTable.WorkerData[WorkerIndex].TransientCBVs.Add(BufferView);
-		return BufferView->GetOfflineCpuHandle();
-	}
-#endif // USE_STATIC_ROOT_SIGNATURE
-
-	void CreateTransientConstantBuffer(FD3D12ResourceLocation& ResourceLocation, const void* Data, uint32 DataSize)
+	FD3D12ConstantBufferView* CreateTransientConstantBuffer(FD3D12ResourceLocation& ResourceLocation, const void* Data, uint32 DataSize)
 	{
 		// If we see a significant number of transient allocations coming through this path, we should consider
 		// caching constant buffer blocks inside ShaderTable and linearly sub-allocate from them.
 		// If the amount of data is relatively small, it may also be possible to use root constants and avoid extra allocations entirely.
 
+	#if USE_STATIC_ROOT_SIGNATURE
+		FD3D12ConstantBufferView* ConstantBufferView = new FD3D12ConstantBufferView(GetDevice());
+		ShaderTable.WorkerData[WorkerIndex].TransientCBVs.Add(ConstantBufferView);
+	#else // USE_STATIC_ROOT_SIGNATURE
+		FD3D12ConstantBufferView* ConstantBufferView = nullptr;
+	#endif // USE_STATIC_ROOT_SIGNATURE
+
 		FD3D12FastConstantAllocator& Allocator = Device.GetParentAdapter()->GetTransientUniformBufferAllocator();
-		void* MappedData = Allocator.Allocate(DataSize, ResourceLocation, nullptr);
+		void* MappedData = Allocator.Allocate(DataSize, ResourceLocation, ConstantBufferView);
 
 		FMemory::Memcpy(MappedData, Data, DataSize);
 
 		ShaderTable.AddResourceReference(ResourceLocation.GetResource(), WorkerIndex);
+
+		return ConstantBufferView;
 	}
 
 	void AddResourceReference(FD3D12Resource* D3D12Resource)
@@ -4909,11 +4901,10 @@ static bool SetRayTracingShaderResources(
 		const uint32 CBVIndex = 0; // Global uniform buffer is always assumed to be in slot 0
 
 		FD3D12ResourceLocation ResourceLocation(Binder.GetDevice());
-		const uint32 AlignedSize = Align(InLooseParameterDataSize, D3D12_CONSTANT_BUFFER_DATA_PLACEMENT_ALIGNMENT);
-		Binder.CreateTransientConstantBuffer(ResourceLocation, InLooseParameterData, AlignedSize);
+		FD3D12ConstantBufferView* ConstantBufferView = Binder.CreateTransientConstantBuffer(ResourceLocation, InLooseParameterData, InLooseParameterDataSize);
 
 	#if USE_STATIC_ROOT_SIGNATURE
-		Bindings.LocalCBVs[CBVIndex] = Binder.CreateTransientConstantBufferView(ResourceLocation, AlignedSize);
+		Bindings.LocalCBVs[CBVIndex] = ConstantBufferView->GetOfflineCpuHandle();
 	#else // USE_STATIC_ROOT_SIGNATURE
 		Bindings.LocalCBVs[CBVIndex] = ResourceLocation.GetGPUVirtualAddress();
 	#endif // USE_STATIC_ROOT_SIGNATURE
