When running with RenderDoc, only consider RelaxedFormatCasting supported when RenderDoc eventually supports the necessary interfaces.

diff --git a/Engine/Source/Runtime/D3D12RHI/Private/Windows/WindowsD3D12Device.cpp b/Engine/Source/Runtime/D3D12RHI/Private/Windows/WindowsD3D12Device.cpp
--- a/Engine/Source/Runtime/D3D12RHI/Private/Windows/WindowsD3D12Device.cpp
+++ b/Engine/Source/Runtime/D3D12RHI/Private/Windows/WindowsD3D12Device.cpp
@@ -1731,7 +1731,12 @@ void FD3D12DynamicRHI::Init()
 	{
 		GRHIGlobals.SupportsUAVFormatAliasing = (Options12.RelaxedFormatCastingSupported != 0)
 			// We require ID3D12Device12 for GetResourceAllocationInfo3
-			&& GetAdapter().GetD3DDevice12() != nullptr;
+			&& GetAdapter().GetD3DDevice12() != nullptr
+#if PLATFORM_WINDOWS
+			// Make sure RenderDoc supports the new interfaces
+			&& (D3D12RHI_IsRenderDocPresent(GetAdapter().GetD3DDevice()) == D3D12RHI_IsRenderDocPresent(GetAdapter().GetD3DDevice12()))
+#endif
+			;
 	}
 #else
 	GRHIGlobals.SupportsUAVFormatAliasing = (GetAdapter().GetResourceHeapTier() > D3D12_RESOURCE_HEAP_TIER_1 && IsRHIDeviceNVIDIA());
