Moving the bulk of D3D12RHI defines to their own file.

diff --git a/Engine/Source/Runtime/D3D12RHI/Private/D3D12Allocation.h b/Engine/Source/Runtime/D3D12RHI/Private/D3D12Allocation.h
--- a/Engine/Source/Runtime/D3D12RHI/Private/D3D12Allocation.h
+++ b/Engine/Source/Runtime/D3D12RHI/Private/D3D12Allocation.h
@@ -5,9 +5,12 @@ D3D12Allocation.h: A Collection of allocators
 =============================================================================*/
 
 #pragma once
+
 #include "D3D12Resources.h"
 #include "D3D12PoolAllocator.h"
 
+#define SUB_ALLOCATED_DEFAULT_ALLOCATIONS	1
+
 // Enable pool allocator by default only on Windows until all initial issues have been resolved.
 // Xbox is currently using the buddy allocator with a lot smaller pool sizes so the alignment waste is a lot less
 // It also has less of a problem with committed resource allocations
diff --git a/Engine/Source/Runtime/D3D12RHI/Private/D3D12RHICommon.h b/Engine/Source/Runtime/D3D12RHI/Private/D3D12RHICommon.h
--- a/Engine/Source/Runtime/D3D12RHI/Private/D3D12RHICommon.h
+++ b/Engine/Source/Runtime/D3D12RHI/Private/D3D12RHICommon.h
@@ -1,20 +1,49 @@
 // Copyright Epic Games, Inc. All Rights Reserved.
 
 /*=============================================================================
-D3D12RHICommon.h: Common D3D12 RHI definitions for Windows.
-=============================================================================*/
+  D3D12RHICommon.h: Common D3D12 RHI definitions
+  =============================================================================*/
 
 #pragma once
 
-#include "Stats/Stats2.h"
-#include "MultiGPU.h"
+#include "HAL/Platform.h"
+
+#if PLATFORM_WINDOWS
+	#include "Windows/D3D12ThirdParty.h"
+#else
+	#include "D3D12ThirdParty.h"
+#endif
+
+#include "D3D12RHI.h"
+#include "D3D12RHIDefinitions.h"
+
 #include "Containers/StaticArray.h"
+#include "MultiGPU.h"
+#include "Stats/Stats2.h"
+
+typedef uint16 CBVSlotMask;
+
+static_assert(MAX_ROOT_CBVS <= MAX_CBS, "MAX_ROOT_CBVS must be <= MAX_CBS.");
+static_assert((8 * sizeof(CBVSlotMask)) >= MAX_CBS, "CBVSlotMask isn't large enough to cover all CBs. Please increase the size.");
+static_assert((8 * sizeof(CBVSlotMask)) >= MAX_ROOT_CBVS, "CBVSlotMask isn't large enough to cover all CBs. Please increase the size.");
+static const CBVSlotMask GRootCBVSlotMask = (1 << MAX_ROOT_CBVS) - 1; // Mask for all slots that are used by root descriptors.
+static const CBVSlotMask GDescriptorTableCBVSlotMask = static_cast<CBVSlotMask>(-1) & ~(GRootCBVSlotMask); // Mask for all slots that are used by a root descriptor table.
 
-#if D3D12_MAX_DEVICE_INTERFACE >= 12
-	#define D3D12RHI_SUPPORTS_UNCOMPRESSED_UAV 1
+#if MAX_SRVS > 32
+	typedef uint64 SRVSlotMask;
 #else
-	#define D3D12RHI_SUPPORTS_UNCOMPRESSED_UAV 0
+	typedef uint32 SRVSlotMask;
 #endif
+static_assert((8 * sizeof(SRVSlotMask)) >= MAX_SRVS, "SRVSlotMask isn't large enough to cover all SRVs. Please increase the size.");
+
+typedef uint16 SamplerSlotMask;
+static_assert((8 * sizeof(SamplerSlotMask)) >= MAX_SAMPLERS, "SamplerSlotMask isn't large enough to cover all Samplers. Please increase the size.");
+
+typedef uint16 UAVSlotMask;
+static_assert((8 * sizeof(UAVSlotMask)) >= MAX_UAVS, "UAVSlotMask isn't large enough to cover all UAVs. Please increase the size.");
+
+DECLARE_LOG_CATEGORY_EXTERN(LogD3D12RHI, Log, All);
+DECLARE_LOG_CATEGORY_EXTERN(LogD3D12GapRecorder, Log, All);
 
 DECLARE_STATS_GROUP(TEXT("D3D12RHI"), STATGROUP_D3D12RHI, STATCAT_Advanced);
 DECLARE_STATS_GROUP(TEXT("D3D12RHI: Memory"), STATGROUP_D3D12Memory, STATCAT_Advanced);
diff --git a/Engine/Source/Runtime/D3D12RHI/Private/D3D12RHIDefinitions.h b/Engine/Source/Runtime/D3D12RHI/Private/D3D12RHIDefinitions.h
new file mode 100644
--- /dev/null
+++ b/Engine/Source/Runtime/D3D12RHI/Private/D3D12RHIDefinitions.h
@@ -0,0 +1,100 @@
+// Copyright Epic Games, Inc. All Rights Reserved.
+
+#pragma once
+
+#include "RHIDefinitions.h"
+
+// DX12 doesn't support higher MSAA count
+#define DX_MAX_MSAA_COUNT	8
+
+// How many residency packets can be in flight before the rendering thread
+// blocks for them to drain. Should be ~ NumBufferedFrames * AvgNumSubmissionsPerFrame i.e.
+// enough to ensure that the GPU is rarely blocked by residency work
+#define RESIDENCY_PIPELINE_DEPTH	6
+
+#define DEBUG_RESOURCE_STATES	0
+
+#define D3D12_RHI_RAYTRACING (RHI_RAYTRACING)
+
+#if D3D12_MAX_DEVICE_INTERFACE >= 12
+	#define D3D12RHI_SUPPORTS_UNCOMPRESSED_UAV 1
+#else
+	#define D3D12RHI_SUPPORTS_UNCOMPRESSED_UAV 0
+#endif
+
+// This value controls how many root constant buffers can be used per shader stage in a root signature.
+// Note: Using root descriptors significantly increases the size of root signatures (each root descriptor is 2 DWORDs).
+#define MAX_ROOT_CBVS	MAX_CBS
+
+#define EXECUTE_DEBUG_COMMAND_LISTS 0
+#define NAME_OBJECTS !(UE_BUILD_SHIPPING || UE_BUILD_TEST)	// Name objects in all builds except shipping
+#define LOG_PSO_CREATES (0 && STATS)	// Logs Create Pipeline State timings (also requires STATS)
+#define TRACK_RESOURCE_ALLOCATIONS (PLATFORM_WINDOWS && !UE_BUILD_SHIPPING && !UE_BUILD_TEST)
+
+#define READBACK_BUFFER_POOL_MAX_ALLOC_SIZE (64 * 1024)
+#define READBACK_BUFFER_POOL_DEFAULT_POOL_SIZE (4 * 1024 * 1024)
+
+#define TEXTURE_POOL_SIZE (8 * 1024 * 1024)
+
+#define MAX_GPU_BREADCRUMB_DEPTH 1024
+#define MAX_GPU_BREADCRUMB_CONTEXTS 128
+#define MAX_GPU_BREADCRUMB_SIZE 4096
+
+#if DEBUG_RESOURCE_STATES
+	#define LOG_EXECUTE_COMMAND_LISTS 1
+	#define LOG_PRESENT 1
+#else
+	#define LOG_EXECUTE_COMMAND_LISTS 0
+	#define LOG_PRESENT 0
+#endif
+
+#define DEBUG_FRAME_TIMING 0
+
+#if DEBUG_FRAME_TIMING
+	#define LOG_VIEWPORT_EVENTS 1
+	#define LOG_PRESENT 1
+	#define LOG_EXECUTE_COMMAND_LISTS 1
+#else
+	#define LOG_VIEWPORT_EVENTS 0
+#endif
+
+#if EXECUTE_DEBUG_COMMAND_LISTS
+	#define DEBUG_EXECUTE_COMMAND_LIST(scope) if (scope##->ActiveQueries == 0) { scope##->FlushCommands(ED3D12FlushFlags::WaitForCompletion); }
+	#define DEBUG_EXECUTE_COMMAND_CONTEXT(context) if (context.ActiveQueries == 0) { context##.FlushCommands(ED3D12FlushFlags::WaitForCompletion); }
+	#define DEBUG_RHI_EXECUTE_COMMAND_LIST(scope) if (scope##->GetRHIDevice(0)->GetDefaultCommandContext().ActiveQueries == 0) { scope##->GetRHIDevice(0)->GetDefaultCommandContext().FlushCommands(ED3D12FlushFlags::WaitForCompletion); }
+#else
+	#define DEBUG_EXECUTE_COMMAND_LIST(scope) 
+	#define DEBUG_EXECUTE_COMMAND_CONTEXT(context) 
+	#define DEBUG_RHI_EXECUTE_COMMAND_LIST(scope) 
+#endif
+
+#ifndef D3D12_VIEWPORT_EXPOSES_SWAP_CHAIN
+	#define D3D12_VIEWPORT_EXPOSES_SWAP_CHAIN 1
+#endif
+
+#if D3D12_VIEWPORT_EXPOSES_SWAP_CHAIN
+	#ifndef DXGI_PRESENT_ALLOW_TEARING
+		#define DXGI_PRESENT_ALLOW_TEARING          0x00000200UL
+		#define DXGI_SWAP_CHAIN_FLAG_ALLOW_TEARING  2048
+	#endif
+#endif
+
+#if D3D12_SUPPORTS_INFO_QUEUE
+	#define EMBED_DXGI_ERROR_LIST(PerEntry, Terminator)	\
+		PerEntry(DXGI_ERROR_UNSUPPORTED) Terminator \
+		PerEntry(DXGI_ERROR_NOT_CURRENT) Terminator \
+		PerEntry(DXGI_ERROR_MORE_DATA) Terminator \
+		PerEntry(DXGI_ERROR_MODE_CHANGE_IN_PROGRESS) Terminator \
+		PerEntry(DXGI_ERROR_ALREADY_EXISTS) Terminator \
+		PerEntry(DXGI_ERROR_SESSION_DISCONNECTED) Terminator \
+		PerEntry(DXGI_ERROR_ACCESS_DENIED) Terminator \
+		PerEntry(DXGI_ERROR_NON_COMPOSITED_UI) Terminator \
+		PerEntry(DXGI_ERROR_CACHE_FULL) Terminator \
+		PerEntry(DXGI_ERROR_NOT_CURRENTLY_AVAILABLE) Terminator \
+		PerEntry(DXGI_ERROR_CACHE_CORRUPT) Terminator \
+		PerEntry(DXGI_ERROR_WAIT_TIMEOUT) Terminator \
+		PerEntry(DXGI_ERROR_FRAME_STATISTICS_DISJOINT) Terminator \
+		PerEntry(DXGI_ERROR_DYNAMIC_CODE_POLICY_VIOLATION) Terminator \
+		PerEntry(DXGI_ERROR_REMOTE_OUTOFMEMORY) Terminator \
+		PerEntry(DXGI_ERROR_ACCESS_LOST) Terminator
+#endif
diff --git a/Engine/Source/Runtime/D3D12RHI/Private/D3D12RHIPrivate.h b/Engine/Source/Runtime/D3D12RHI/Private/D3D12RHIPrivate.h
--- a/Engine/Source/Runtime/D3D12RHI/Private/D3D12RHIPrivate.h
+++ b/Engine/Source/Runtime/D3D12RHI/Private/D3D12RHIPrivate.h
@@ -6,6 +6,9 @@
 
 #pragma once
 
+#include "D3D12RHICommon.h"
+#include "ID3D12DynamicRHI.h"
+
 #include "CoreMinimal.h"
 #include "Misc/CommandLine.h"
 #include "Misc/App.h"
@@ -21,22 +24,12 @@
 #include "EngineGlobals.h"
 #include "StaticBoundShaderState.h"
 
-#define D3D12_RHI_RAYTRACING (RHI_RAYTRACING)
-
 // Dependencies.
-#include "CoreMinimal.h"
-#include "ID3D12DynamicRHI.h"
 #include "DXGIUtilities.h"
 #include "GPUProfiler.h"
 #include "ShaderCore.h"
 #include "HDRHelper.h"
 
-DECLARE_LOG_CATEGORY_EXTERN(LogD3D12RHI, Log, All);
-DECLARE_LOG_CATEGORY_EXTERN(LogD3D12GapRecorder, Log, All);
-
-#include "D3D12RHI.h"
-#include "D3D12RHICommon.h"
-
 // Defines a unique command queue type within a FD3D12Device (owner by the command list managers).
 enum class ED3D12QueueType
 {
@@ -50,9 +43,25 @@ enum class ED3D12QueueType
 #include "D3D12Submission.h"
 
 #if PLATFORM_WINDOWS
-#include "Windows/D3D12RHIBasePrivate.h"
+	#include "Windows/D3D12RHIBasePrivate.h"
 #else
-#include "D3D12RHIBasePrivate.h"
+	#include "D3D12RHIBasePrivate.h"
+#endif
+
+// Note: the following defines depend on D3D12RHIBasePrivate.h
+
+//@TODO: Improve allocator efficiency so we can increase these thresholds and improve performance
+// We measured 149MB of wastage in 340MB of allocations with DEFAULT_BUFFER_POOL_MAX_ALLOC_SIZE set to 512KB
+#if !defined(DEFAULT_BUFFER_POOL_MAX_ALLOC_SIZE)
+	#if D3D12_RHI_RAYTRACING
+		// #dxr_todo: Reevaluate these values. Currently optimized to reduce number of CreateCommitedResource() calls, at the expense of memory use.
+		#define DEFAULT_BUFFER_POOL_MAX_ALLOC_SIZE    (64 * 1024 * 1024)
+		#define DEFAULT_BUFFER_POOL_DEFAULT_POOL_SIZE (16 * 1024 * 1024)
+	#else
+		// On PC, buffers are 64KB aligned, so anything smaller should be sub-allocated
+		#define DEFAULT_BUFFER_POOL_MAX_ALLOC_SIZE    (64 * 1024)
+		#define DEFAULT_BUFFER_POOL_DEFAULT_POOL_SIZE (8 * 1024 * 1024)
+	#endif //D3D12_RHI_RAYTRACING
 #endif
 
 static constexpr uint32 GD3D12MaxNumQueues = MAX_NUM_GPUS * (uint32)ED3D12QueueType::Count;
@@ -138,65 +147,6 @@ inline const TCHAR* GetD3DCommandQueueTypeName(ED3D12QueueType QueueType)
 #include "D3D12Device.h"
 #include "D3D12Adapter.h"
 
-#define EXECUTE_DEBUG_COMMAND_LISTS 0
-#define NAME_OBJECTS !(UE_BUILD_SHIPPING || UE_BUILD_TEST)	// Name objects in all builds except shipping
-#define LOG_PSO_CREATES (0 && STATS)	// Logs Create Pipeline State timings (also requires STATS)
-#define TRACK_RESOURCE_ALLOCATIONS (PLATFORM_WINDOWS && !UE_BUILD_SHIPPING && !UE_BUILD_TEST)
-
-//@TODO: Improve allocator efficiency so we can increase these thresholds and improve performance
-// We measured 149MB of wastage in 340MB of allocations with DEFAULT_BUFFER_POOL_MAX_ALLOC_SIZE set to 512KB
-#if !defined(DEFAULT_BUFFER_POOL_MAX_ALLOC_SIZE)
-	#if D3D12_RHI_RAYTRACING
-		// #dxr_todo: Reevaluate these values. Currently optimized to reduce number of CreateCommitedResource() calls, at the expense of memory use.
-		#define DEFAULT_BUFFER_POOL_MAX_ALLOC_SIZE    (64 * 1024 * 1024)
-		#define DEFAULT_BUFFER_POOL_DEFAULT_POOL_SIZE (16 * 1024 * 1024)
-	#else
-		// On PC, buffers are 64KB aligned, so anything smaller should be sub-allocated
-		#define DEFAULT_BUFFER_POOL_MAX_ALLOC_SIZE    (64 * 1024)
-		#define DEFAULT_BUFFER_POOL_DEFAULT_POOL_SIZE (8 * 1024 * 1024)
-	#endif //D3D12_RHI_RAYTRACING
-#endif //DEFAULT_BUFFER_POOL_MAX_ALLOC_SIZE
-
-#define READBACK_BUFFER_POOL_MAX_ALLOC_SIZE (64 * 1024)
-#define READBACK_BUFFER_POOL_DEFAULT_POOL_SIZE (4 * 1024 * 1024)
-
-#define TEXTURE_POOL_SIZE (8 * 1024 * 1024)
-
-#define MAX_GPU_BREADCRUMB_DEPTH 1024
-#define MAX_GPU_BREADCRUMB_CONTEXTS 128
-#define MAX_GPU_BREADCRUMB_SIZE 4096
-
-#ifndef FD3D12_HEAP_FLAG_CREATE_NOT_ZEROED
-#define FD3D12_HEAP_FLAG_CREATE_NOT_ZEROED D3D12_HEAP_FLAG_CREATE_NOT_ZEROED
-#endif
-
-#if DEBUG_RESOURCE_STATES
-	#define LOG_EXECUTE_COMMAND_LISTS 1
-	#define LOG_PRESENT 1
-#else
-	#define LOG_EXECUTE_COMMAND_LISTS 0
-	#define LOG_PRESENT 0
-#endif
-
-#define DEBUG_FRAME_TIMING 0
-#if DEBUG_FRAME_TIMING
-	#define LOG_VIEWPORT_EVENTS 1
-	#define LOG_PRESENT 1
-	#define LOG_EXECUTE_COMMAND_LISTS 1
-#else
-	#define LOG_VIEWPORT_EVENTS 0
-#endif
-
-#if EXECUTE_DEBUG_COMMAND_LISTS
-	#define DEBUG_EXECUTE_COMMAND_LIST(scope) if (scope##->ActiveQueries == 0) { scope##->FlushCommands(ED3D12FlushFlags::WaitForCompletion); }
-	#define DEBUG_EXECUTE_COMMAND_CONTEXT(context) if (context.ActiveQueries == 0) { context##.FlushCommands(ED3D12FlushFlags::WaitForCompletion); }
-	#define DEBUG_RHI_EXECUTE_COMMAND_LIST(scope) if (scope##->GetRHIDevice(0)->GetDefaultCommandContext().ActiveQueries == 0) { scope##->GetRHIDevice(0)->GetDefaultCommandContext().FlushCommands(ED3D12FlushFlags::WaitForCompletion); }
-#else
-	#define DEBUG_EXECUTE_COMMAND_LIST(scope) 
-	#define DEBUG_EXECUTE_COMMAND_CONTEXT(context) 
-	#define DEBUG_RHI_EXECUTE_COMMAND_LIST(scope) 
-#endif
-
 template< typename t_A, typename t_B >
 inline t_A RoundUpToNextMultiple(const t_A& a, const t_B& b)
 {
@@ -812,35 +762,3 @@ public:
 
 // Returns the given format as a string. Unsupported formats are treated as DXGI_FORMAT_UNKNOWN.
 const TCHAR* LexToString(DXGI_FORMAT Format);
-
-#if (PLATFORM_WINDOWS || PLATFORM_HOLOLENS)
-
-#ifndef DXGI_PRESENT_ALLOW_TEARING
-#define DXGI_PRESENT_ALLOW_TEARING          0x00000200UL
-#define DXGI_SWAP_CHAIN_FLAG_ALLOW_TEARING  2048
-
-#endif
-
-
-
-#define EMBED_DXGI_ERROR_LIST(PerEntry, Terminator)	\
-	PerEntry(DXGI_ERROR_UNSUPPORTED) Terminator \
-	PerEntry(DXGI_ERROR_NOT_CURRENT) Terminator \
-	PerEntry(DXGI_ERROR_MORE_DATA) Terminator \
-	PerEntry(DXGI_ERROR_MODE_CHANGE_IN_PROGRESS) Terminator \
-	PerEntry(DXGI_ERROR_ALREADY_EXISTS) Terminator \
-	PerEntry(DXGI_ERROR_SESSION_DISCONNECTED) Terminator \
-	PerEntry(DXGI_ERROR_ACCESS_DENIED) Terminator \
-	PerEntry(DXGI_ERROR_NON_COMPOSITED_UI) Terminator \
-	PerEntry(DXGI_ERROR_CACHE_FULL) Terminator \
-	PerEntry(DXGI_ERROR_NOT_CURRENTLY_AVAILABLE) Terminator \
-	PerEntry(DXGI_ERROR_CACHE_CORRUPT) Terminator \
-	PerEntry(DXGI_ERROR_WAIT_TIMEOUT) Terminator \
-	PerEntry(DXGI_ERROR_FRAME_STATISTICS_DISJOINT) Terminator \
-	PerEntry(DXGI_ERROR_DYNAMIC_CODE_POLICY_VIOLATION) Terminator \
-	PerEntry(DXGI_ERROR_REMOTE_OUTOFMEMORY) Terminator \
-	PerEntry(DXGI_ERROR_ACCESS_LOST) Terminator
-
-
-
-#endif //(PLATFORM_WINDOWS || PLATFORM_HOLOLENS)
diff --git a/Engine/Source/Runtime/D3D12RHI/Private/D3D12Viewport.h b/Engine/Source/Runtime/D3D12RHI/Private/D3D12Viewport.h
--- a/Engine/Source/Runtime/D3D12RHI/Private/D3D12Viewport.h
+++ b/Engine/Source/Runtime/D3D12RHI/Private/D3D12Viewport.h
@@ -11,10 +11,6 @@
 
 #define ALTERNATE_TIMESTAMP_METRIC 0
 
-#ifndef D3D12_VIEWPORT_EXPOSES_SWAP_CHAIN
-#define D3D12_VIEWPORT_EXPOSES_SWAP_CHAIN 1
-#endif // D3D12_VIEWPORT_EXPOSES_SWAP_CHAIN
-
 #if WITH_MGPU
 class FD3D12FramePacing : public FRunnable, public FD3D12AdapterChild
 {
diff --git a/Engine/Source/Runtime/D3D12RHI/Private/Windows/D3D12RHIBasePrivate.h b/Engine/Source/Runtime/D3D12RHI/Private/Windows/D3D12RHIBasePrivate.h
--- a/Engine/Source/Runtime/D3D12RHI/Private/Windows/D3D12RHIBasePrivate.h
+++ b/Engine/Source/Runtime/D3D12RHI/Private/Windows/D3D12RHIBasePrivate.h
@@ -29,6 +29,8 @@ typedef ID3D12GraphicsCommandList ID3D12CopyCommandList;
 #define D3D12RHI_RESOURCE_FLAG_ALLOW_INDIRECT_BUFFER D3D12_RESOURCE_FLAG_NONE
 #define D3D12RHI_HEAP_FLAG_ALLOW_INDIRECT_BUFFERS		D3D12_HEAP_FLAG_NONE
 
+#define FD3D12_HEAP_FLAG_CREATE_NOT_ZEROED				D3D12_HEAP_FLAG_CREATE_NOT_ZEROED
+
 #define D3D12RHI_NEEDS_VENDOR_EXTENSIONS     1
 #define D3D12RHI_NEEDS_SHADER_FEATURE_CHECKS 1
 
diff --git a/Engine/Source/Runtime/D3D12RHI/Public/D3D12RHI.h b/Engine/Source/Runtime/D3D12RHI/Public/D3D12RHI.h
--- a/Engine/Source/Runtime/D3D12RHI/Public/D3D12RHI.h
+++ b/Engine/Source/Runtime/D3D12RHI/Public/D3D12RHI.h
@@ -8,16 +8,6 @@
 
 #include "HAL/Platform.h"
 
-///////////////////////////////////////////////////////////////
-// Platform agnostic defines
-///////////////////////////////////////////////////////////////
-#define SUB_ALLOCATED_DEFAULT_ALLOCATIONS	1
-
-#define DEBUG_RESOURCE_STATES	0
-
-// DX12 doesn't support higher MSAA count
-#define DX_MAX_MSAA_COUNT	8
-
 // This is a value that should be tweaked to fit the app, lower numbers will have better performance
 // Titles using many terrain layers may want to set MAX_SRVS to 64 to avoid shader compilation errors. This will have a small performance hit of around 0.1%
 #define MAX_SRVS		64
@@ -25,15 +15,6 @@
 #define MAX_UAVS		16
 #define MAX_CBS			16
 
-// This value controls how many root constant buffers can be used per shader stage in a root signature.
-// Note: Using root descriptors significantly increases the size of root signatures (each root descriptor is 2 DWORDs).
-#define MAX_ROOT_CBVS	MAX_CBS
-
-// How many residency packets can be in flight before the rendering thread
-// blocks for them to drain. Should be ~ NumBufferedFrames * AvgNumSubmissionsPerFrame i.e.
-// enough to ensure that the GPU is rarely blocked by residency work
-#define RESIDENCY_PIPELINE_DEPTH	6
-
 #if !defined(WITH_PIX_EVENT_RUNTIME)
 	#define WITH_PIX_EVENT_RUNTIME 0
 #endif
@@ -48,23 +29,3 @@
 #else
 	#include "D3D12RHIPlatformPublic.h"
 #endif
-
-typedef uint16 CBVSlotMask;
-static_assert(MAX_ROOT_CBVS <= MAX_CBS, "MAX_ROOT_CBVS must be <= MAX_CBS.");
-static_assert((8 * sizeof(CBVSlotMask)) >= MAX_CBS, "CBVSlotMask isn't large enough to cover all CBs. Please increase the size.");
-static_assert((8 * sizeof(CBVSlotMask)) >= MAX_ROOT_CBVS, "CBVSlotMask isn't large enough to cover all CBs. Please increase the size.");
-static const CBVSlotMask GRootCBVSlotMask = (1 << MAX_ROOT_CBVS) - 1; // Mask for all slots that are used by root descriptors.
-static const CBVSlotMask GDescriptorTableCBVSlotMask = static_cast<CBVSlotMask>(-1) & ~(GRootCBVSlotMask); // Mask for all slots that are used by a root descriptor table.
-
-#if MAX_SRVS > 32
-typedef uint64 SRVSlotMask;
-#else
-typedef uint32 SRVSlotMask;
-#endif
-static_assert((8 * sizeof(SRVSlotMask)) >= MAX_SRVS, "SRVSlotMask isn't large enough to cover all SRVs. Please increase the size.");
-
-typedef uint16 SamplerSlotMask;
-static_assert((8 * sizeof(SamplerSlotMask)) >= MAX_SAMPLERS, "SamplerSlotMask isn't large enough to cover all Samplers. Please increase the size.");
-
-typedef uint16 UAVSlotMask;
-static_assert((8 * sizeof(UAVSlotMask)) >= MAX_UAVS, "UAVSlotMask isn't large enough to cover all UAVs. Please increase the size.");
