D3D12RHI header dependency cleanup
- Restrict D3D12RHI to Microsoft platforms. UBT was building the module on every platform when -allmodules was set.
- Moving a few build script items into individual platform scripts instead of being based on the platform group.
- Adding Platform specific versions of D3D12RHISubmission.h and moving FD3D12Payload into those files.
- Moving ED3D12QueueType to its own file, D3D12Queue.h. This should eventually also contain FD3D12Queue.
- Adding missing includes to D3D12RHISubmission.h so it can compile on its own.

diff --git a/Engine/Platforms/Hololens/Source/Runtime/D3D12RHI/D3D12RHI_HoloLens.build.cs b/Engine/Platforms/Hololens/Source/Runtime/D3D12RHI/D3D12RHI_HoloLens.build.cs
--- a/Engine/Platforms/Hololens/Source/Runtime/D3D12RHI/D3D12RHI_HoloLens.build.cs
+++ b/Engine/Platforms/Hololens/Source/Runtime/D3D12RHI/D3D12RHI_HoloLens.build.cs
@@ -8,6 +8,9 @@ namespace UnrealBuildTool.Rules
 
 		public D3D12RHI_HoloLens(ReadOnlyTargetRules Target) : base(Target)
 		{
+			PrecompileForTargets = PrecompileTargetsType.None;
+
+			PrivateDefinitions.Add("D3D12RHI_USE_D3DDISASSEMBLE=0");
 		}
 	}
 }
diff --git a/Engine/Platforms/Hololens/Source/Runtime/D3D12RHI/Private/D3D12RHIBasePrivate.h b/Engine/Platforms/Hololens/Source/Runtime/D3D12RHI/Private/D3D12RHIBasePrivate.h
--- a/Engine/Platforms/Hololens/Source/Runtime/D3D12RHI/Private/D3D12RHIBasePrivate.h
+++ b/Engine/Platforms/Hololens/Source/Runtime/D3D12RHI/Private/D3D12RHIBasePrivate.h
@@ -8,14 +8,6 @@
 
 #include "D3D12ThirdParty.h"
 
-// Standard Hololens implementation, used to mark the type as 'final'.
-struct FD3D12Payload final : public FD3D12PayloadBase
-{
-	FD3D12Payload(FD3D12Device* const Device, ED3D12QueueType const QueueType)
-		: FD3D12PayloadBase(Device, QueueType)
-	{}
-};
-
 // Hololens uses the graphics command list interface for the copy queue.
 typedef ID3D12GraphicsCommandList ID3D12CopyCommandList;
 
diff --git a/Engine/Platforms/Hololens/Source/Runtime/D3D12RHI/Private/HoloLensD3D12Submission.h b/Engine/Platforms/Hololens/Source/Runtime/D3D12RHI/Private/HoloLensD3D12Submission.h
new file mode 100644
--- /dev/null
+++ b/Engine/Platforms/Hololens/Source/Runtime/D3D12RHI/Private/HoloLensD3D12Submission.h
@@ -0,0 +1,14 @@
+// Copyright Epic Games, Inc. All Rights Reserved.
+
+#pragma once
+
+#include "D3D12Submission.h"
+
+// Standard Hololens implementation, used to mark the type as 'final'.
+struct FD3D12Payload final : public FD3D12PayloadBase
+{
+	FD3D12Payload(FD3D12Device* Device, ED3D12QueueType QueueType)
+		: FD3D12PayloadBase(Device, QueueType)
+	{
+	}
+};
diff --git a/Engine/Source/Runtime/D3D12RHI/D3D12RHI.Build.cs b/Engine/Source/Runtime/D3D12RHI/D3D12RHI.Build.cs
--- a/Engine/Source/Runtime/D3D12RHI/D3D12RHI.Build.cs
+++ b/Engine/Source/Runtime/D3D12RHI/D3D12RHI.Build.cs
@@ -2,6 +2,7 @@
 
 using UnrealBuildTool;
 
+[SupportedPlatformGroups("Microsoft")]
 public class D3D12RHI : ModuleRules
 {
 	protected virtual bool bUsesWindowsD3D12 { get => Target.Platform.IsInGroup(UnrealPlatformGroup.Windows); }
@@ -24,11 +25,6 @@ public class D3D12RHI : ModuleRules
 
 		PublicIncludePathModuleNames.Add("HeadMountedDisplay");
 
-		if (!Target.Platform.IsInGroup(UnrealPlatformGroup.Windows))
-        {
-            PrecompileForTargets = PrecompileTargetsType.None;
-        }
-
 		AddEngineThirdPartyPrivateStaticDependencies(Target, "IntelExtensionsFramework");
 		AddEngineThirdPartyPrivateStaticDependencies(Target, "AMD_AGS");
 		AddEngineThirdPartyPrivateStaticDependencies(Target, "NVAPI");
@@ -48,10 +44,5 @@ public class D3D12RHI : ModuleRules
 				PublicDependencyModuleNames.Add("WinPixEventRuntime");
 			}
 		}
-
-		if (!Target.Platform.IsInGroup(UnrealPlatformGroup.Windows))
-		{
-			PrivateDefinitions.Add("D3D12RHI_USE_D3DDISASSEMBLE=0");
-		}
 	}
 }
diff --git a/Engine/Source/Runtime/D3D12RHI/Private/D3D12Queue.h b/Engine/Source/Runtime/D3D12RHI/Private/D3D12Queue.h
new file mode 100644
--- /dev/null
+++ b/Engine/Source/Runtime/D3D12RHI/Private/D3D12Queue.h
@@ -0,0 +1,41 @@
+// Copyright Epic Games, Inc. All Rights Reserved.
+
+#pragma once
+
+#include "D3D12RHICommon.h"
+#include "RHIPipeline.h"
+
+// Defines a unique command queue type within a FD3D12Device (owner by the command list managers).
+enum class ED3D12QueueType
+{
+	Direct = 0,
+	Copy,
+	Async,
+
+	Count,
+};
+
+static constexpr uint32 GD3D12MaxNumQueues = MAX_NUM_GPUS * (uint32)ED3D12QueueType::Count;
+
+inline ED3D12QueueType GetD3DCommandQueueType(ERHIPipeline Pipeline)
+{
+	switch (Pipeline)
+	{
+	default: checkNoEntry(); // fallthrough
+	case ERHIPipeline::Graphics:     return ED3D12QueueType::Direct;
+	case ERHIPipeline::AsyncCompute: return ED3D12QueueType::Async;
+	}
+}
+
+inline const TCHAR* GetD3DCommandQueueTypeName(ED3D12QueueType QueueType)
+{
+	switch (QueueType)
+	{
+	default: checkNoEntry(); // fallthrough
+	case ED3D12QueueType::Direct: return TEXT("3D");
+	case ED3D12QueueType::Async:  return TEXT("Compute");
+	case ED3D12QueueType::Copy:   return TEXT("Copy");
+	}
+}
+
+// TODO: move FD3D12Queue here
diff --git a/Engine/Source/Runtime/D3D12RHI/Private/D3D12RHIPrivate.h b/Engine/Source/Runtime/D3D12RHI/Private/D3D12RHIPrivate.h
--- a/Engine/Source/Runtime/D3D12RHI/Private/D3D12RHIPrivate.h
+++ b/Engine/Source/Runtime/D3D12RHI/Private/D3D12RHIPrivate.h
@@ -30,17 +30,7 @@
 #include "ShaderCore.h"
 #include "HDRHelper.h"
 
-// Defines a unique command queue type within a FD3D12Device (owner by the command list managers).
-enum class ED3D12QueueType
-{
-	Direct = 0,
-	Copy,
-	Async,
-
-	Count,
-};
-
-#include "D3D12Submission.h"
+#include COMPILED_PLATFORM_HEADER(D3D12Submission.h)
 
 #if PLATFORM_WINDOWS
 	#include "Windows/D3D12RHIBasePrivate.h"
@@ -64,37 +54,15 @@ enum class ED3D12QueueType
 	#endif //D3D12_RHI_RAYTRACING
 #endif
 
-static constexpr uint32 GD3D12MaxNumQueues = MAX_NUM_GPUS * (uint32)ED3D12QueueType::Count;
-
-inline ED3D12QueueType GetD3DCommandQueueType(ERHIPipeline Pipeline)
-{
-	switch (Pipeline)
-	{
-	default: checkNoEntry(); // fallthrough
-	case ERHIPipeline::Graphics    : return ED3D12QueueType::Direct;
-	case ERHIPipeline::AsyncCompute: return ED3D12QueueType::Async;
-	}
-}
-
+// TODO: move D3D12RHI_PLATFORM_COPY_COMMAND_LIST_TYPE out of D3D12RHIBasePrivate
 inline D3D12_COMMAND_LIST_TYPE GetD3DCommandListType(ED3D12QueueType QueueType)
 {
 	switch (QueueType)
 	{
 	default: checkNoEntry(); // fallthrough
 	case ED3D12QueueType::Direct: return D3D12_COMMAND_LIST_TYPE_DIRECT;
-	case ED3D12QueueType::Copy  : return D3D12RHI_PLATFORM_COPY_COMMAND_LIST_TYPE;
-	case ED3D12QueueType::Async : return D3D12_COMMAND_LIST_TYPE_COMPUTE;
-	}
-}
-
-inline const TCHAR* GetD3DCommandQueueTypeName(ED3D12QueueType QueueType)
-{
-	switch (QueueType)
-	{
-	default: checkNoEntry(); // fallthrough
-	case ED3D12QueueType::Direct: return TEXT("3D");
-	case ED3D12QueueType::Async : return TEXT("Compute");
-	case ED3D12QueueType::Copy  : return TEXT("Copy");
+	case ED3D12QueueType::Copy:   return D3D12RHI_PLATFORM_COPY_COMMAND_LIST_TYPE;
+	case ED3D12QueueType::Async:  return D3D12_COMMAND_LIST_TYPE_COMPUTE;
 	}
 }
 
diff --git a/Engine/Source/Runtime/D3D12RHI/Private/D3D12Submission.h b/Engine/Source/Runtime/D3D12RHI/Private/D3D12Submission.h
--- a/Engine/Source/Runtime/D3D12RHI/Private/D3D12Submission.h
+++ b/Engine/Source/Runtime/D3D12RHI/Private/D3D12Submission.h
@@ -2,6 +2,11 @@
 
 #pragma once
 
+#include "Async/TaskGraphInterfaces.h"
+#include "D3D12RHICommon.h"
+#include "D3D12Queue.h"
+#include "Templates/RefCounting.h"
+
 enum class ED3D12SyncPointType
 {
 	// Sync points of this type do not include an FGraphEvent, so cannot
diff --git a/Engine/Source/Runtime/D3D12RHI/Private/Windows/D3D12RHIBasePrivate.h b/Engine/Source/Runtime/D3D12RHI/Private/Windows/D3D12RHIBasePrivate.h
--- a/Engine/Source/Runtime/D3D12RHI/Private/Windows/D3D12RHIBasePrivate.h
+++ b/Engine/Source/Runtime/D3D12RHI/Private/Windows/D3D12RHIBasePrivate.h
@@ -8,14 +8,6 @@
 
 #include "Windows/D3D12ThirdParty.h"
 
-// Standard Windows implementation, used to mark the type as 'final'.
-struct FD3D12Payload final : public FD3D12PayloadBase
-{
-	FD3D12Payload(FD3D12Device* const Device, ED3D12QueueType const QueueType)
-		: FD3D12PayloadBase(Device, QueueType)
-	{}
-};
-
 // Windows uses the graphics command list interface for the copy queue.
 typedef ID3D12GraphicsCommandList ID3D12CopyCommandList;
 
diff --git a/Engine/Source/Runtime/D3D12RHI/Private/Windows/WindowsD3D12Submission.h b/Engine/Source/Runtime/D3D12RHI/Private/Windows/WindowsD3D12Submission.h
new file mode 100644
--- /dev/null
+++ b/Engine/Source/Runtime/D3D12RHI/Private/Windows/WindowsD3D12Submission.h
@@ -0,0 +1,14 @@
+// Copyright Epic Games, Inc. All Rights Reserved.
+
+#pragma once
+
+#include "D3D12Submission.h"
+
+// Standard Windows implementation, used to mark the type as 'final'.
+struct FD3D12Payload final : public FD3D12PayloadBase
+{
+	FD3D12Payload(FD3D12Device* Device, ED3D12QueueType QueueType)
+		: FD3D12PayloadBase(Device, QueueType)
+	{
+	}
+};
