Implement RHILock/UnlockTextureCubeFace_RenderThread in d3d12 RHI to fix ImmediateFlush calls in the default dynamic RHI implementation

diff --git a/Engine/Source/Runtime/D3D12RHI/Private/D3D12RHIPrivate.h b/Engine/Source/Runtime/D3D12RHI/Private/D3D12RHIPrivate.h
--- a/Engine/Source/Runtime/D3D12RHI/Private/D3D12RHIPrivate.h
+++ b/Engine/Source/Runtime/D3D12RHI/Private/D3D12RHIPrivate.h
@@ -374,6 +374,8 @@ public:
 	virtual void RHIUnlockTexture2D(FRHITexture2D* Texture, uint32 MipIndex, bool bLockWithinMiptail) final override;
 	virtual void* RHILockTexture2DArray(FRHITexture2DArray* Texture, uint32 TextureIndex, uint32 MipIndex, EResourceLockMode LockMode, uint32& DestStride, bool bLockWithinMiptail) final override;
 	virtual void RHIUnlockTexture2DArray(FRHITexture2DArray* Texture, uint32 TextureIndex, uint32 MipIndex, bool bLockWithinMiptail) final override;
+	virtual void* RHILockTextureCubeFace_RenderThread(class FRHICommandListImmediate& RHICmdList, FRHITextureCube* TextureRHI, uint32 FaceIndex, uint32 ArrayIndex, uint32 MipIndex, EResourceLockMode LockMode, uint32& DestStride, bool bLockWithinMiptail) final override;
+	virtual void RHIUnlockTextureCubeFace_RenderThread(class FRHICommandListImmediate& RHICmdList, FRHITextureCube* TextureRHI, uint32 FaceIndex, uint32 ArrayIndex, uint32 MipIndex, bool bLockWithinMiptail) final override;
 	virtual void RHIUpdateTexture2D(FRHICommandListBase& RHICmdList, FRHITexture2D* Texture, uint32 MipIndex, const struct FUpdateTextureRegion2D& UpdateRegion, uint32 SourcePitch, const uint8* SourceData) final override;
 	virtual void RHIUpdateTexture3D(FRHICommandListBase& RHICmdList, FRHITexture3D* Texture, uint32 MipIndex, const struct FUpdateTextureRegion3D& UpdateRegion, uint32 SourceRowPitch, uint32 SourceDepthPitch, const uint8* SourceData) final override;
 	virtual FUpdateTexture3DData RHIBeginUpdateTexture3D(FRHICommandListBase& RHICmdList, FRHITexture3D* Texture, uint32 MipIndex, const struct FUpdateTextureRegion3D& UpdateRegion) final override;
diff --git a/Engine/Source/Runtime/D3D12RHI/Private/D3D12Texture.cpp b/Engine/Source/Runtime/D3D12RHI/Private/D3D12Texture.cpp
--- a/Engine/Source/Runtime/D3D12RHI/Private/D3D12Texture.cpp
+++ b/Engine/Source/Runtime/D3D12RHI/Private/D3D12Texture.cpp
@@ -2311,6 +2311,22 @@ void FD3D12DynamicRHI::RHIUnlockTexture2DArray(FRHITexture2DArray* TextureRHI, u
 	Texture->Unlock(nullptr, MipIndex, TextureIndex);
 }
 
+void* FD3D12DynamicRHI::RHILockTextureCubeFace_RenderThread(class FRHICommandListImmediate& RHICmdList, FRHITextureCube* TextureRHI, uint32 FaceIndex, uint32 ArrayIndex, uint32 MipIndex, EResourceLockMode LockMode, uint32& DestStride, bool bLockWithinMiptail)
+{
+	check(TextureRHI);
+	FD3D12Texture* TextureCube = FD3D12DynamicRHI::ResourceCast(TextureRHI);
+	uint32 D3DFace = GetD3D12CubeFace((ECubeFace)FaceIndex);
+	return TextureCube->Lock(&RHICmdList, MipIndex, D3DFace + ArrayIndex * 6, LockMode, DestStride);
+}
+
+void FD3D12DynamicRHI::RHIUnlockTextureCubeFace_RenderThread(class FRHICommandListImmediate& RHICmdList, FRHITextureCube* TextureRHI, uint32 FaceIndex, uint32 ArrayIndex, uint32 MipIndex, bool bLockWithinMiptail)
+{
+	check(TextureRHI);
+	FD3D12Texture* TextureCube = FD3D12DynamicRHI::ResourceCast(TextureRHI);
+	uint32 D3DFace = GetD3D12CubeFace((ECubeFace)FaceIndex);
+	TextureCube->Unlock(&RHICmdList, MipIndex, D3DFace + ArrayIndex * 6);
+}
+
 void FD3D12DynamicRHI::RHIUpdateTexture2D(FRHICommandListBase& RHICmdList, FRHITexture2D* TextureRHI, uint32 MipIndex, const FUpdateTextureRegion2D& UpdateRegion, uint32 SourcePitch, const uint8* SourceData)
 {
 	check(TextureRHI);
