D3D12 Platform Header Cleanup
- Removing D3D12RHIBasePrivate.h in favor of platform versions of D3D12RHIDefinitions.h, D3D12DiskCache.h, and D3D12PipelineState.h
- These will eventually be included via their non-platform header after other code is moved around.
- Putting D3D12ThirdParty.h platform include in its own contained header.

diff --git a/Engine/Platforms/Hololens/Source/Runtime/D3D12RHI/Private/D3D12RHIBasePrivate.h b/Engine/Platforms/Hololens/Source/Runtime/D3D12RHI/Private/D3D12RHIBasePrivate.h
deleted file mode 100644
diff --git a/Engine/Platforms/Hololens/Source/Runtime/D3D12RHI/Private/HoloLensD3D12DiskCache.h b/Engine/Platforms/Hololens/Source/Runtime/D3D12RHI/Private/HoloLensD3D12DiskCache.h
new file mode 100644
--- /dev/null
+++ b/Engine/Platforms/Hololens/Source/Runtime/D3D12RHI/Private/HoloLensD3D12DiskCache.h
@@ -0,0 +1,5 @@
+// Copyright Epic Games, Inc. All Rights Reserved.
+
+#pragma once
+
+#include "Windows/WindowsD3D12DiskCache.h"
diff --git a/Engine/Platforms/Hololens/Source/Runtime/D3D12RHI/Private/HoloLensD3D12PipelineState.h b/Engine/Platforms/Hololens/Source/Runtime/D3D12RHI/Private/HoloLensD3D12PipelineState.h
new file mode 100644
--- /dev/null
+++ b/Engine/Platforms/Hololens/Source/Runtime/D3D12RHI/Private/HoloLensD3D12PipelineState.h
@@ -0,0 +1,5 @@
+// Copyright Epic Games, Inc. All Rights Reserved.
+
+#pragma once
+
+#include "Windows/WindowsD3D12PipelineState.h"
diff --git a/Engine/Platforms/Hololens/Source/Runtime/D3D12RHI/Private/HoloLensD3D12RHIDefinitions.h b/Engine/Platforms/Hololens/Source/Runtime/D3D12RHI/Private/HoloLensD3D12RHIDefinitions.h
new file mode 100644
--- /dev/null
+++ b/Engine/Platforms/Hololens/Source/Runtime/D3D12RHI/Private/HoloLensD3D12RHIDefinitions.h
@@ -0,0 +1,26 @@
+// Copyright Epic Games, Inc. All Rights Reserved.
+
+#pragma once
+
+#define D3D12RHI_PLATFORM_COPY_COMMAND_LIST_TYPE D3D12_COMMAND_LIST_TYPE_COPY
+
+#define FD3D12_DESCRIPTOR_HEAP_TYPE_SAMPLER D3D12_DESCRIPTOR_HEAP_TYPE_SAMPLER
+
+// Hololens uses standard D3D12 query heaps for timestamps
+#define D3D12RHI_PLATFORM_USES_TIMESTAMP_QUERIES 1
+
+#define D3D12RHI_RESOURCE_FLAG_ALLOW_INDIRECT_BUFFER D3D12_RESOURCE_FLAG_NONE
+#define D3D12RHI_HEAP_FLAG_ALLOW_INDIRECT_BUFFERS		D3D12_HEAP_FLAG_NONE
+
+#define D3D12RHI_NEEDS_VENDOR_EXTENSIONS     0
+#define D3D12RHI_NEEDS_SHADER_FEATURE_CHECKS 0
+
+// Heap create not zeroed flag is not available on Hololens so use internal define to disable the feature
+// but make code path shared when it becomes available
+#define FD3D12_HEAP_FLAG_CREATE_NOT_ZEROED D3D12_HEAP_FLAG_NONE
+
+#define USE_STATIC_ROOT_SIGNATURE			0
+#define D3D12_USE_DUMMY_BACKBUFFER			1
+
+// Hololens uses the graphics command list interface for the copy queue.
+typedef ID3D12GraphicsCommandList ID3D12CopyCommandList;
diff --git a/Engine/Plugins/VirtualProduction/TextureShare/Source/TextureShareCore/Private/Core/ResourcesCache/D3D12/TextureShareCoreD3D12ResourcesCacheItem.cpp b/Engine/Plugins/VirtualProduction/TextureShare/Source/TextureShareCore/Private/Core/ResourcesCache/D3D12/TextureShareCoreD3D12ResourcesCacheItem.cpp
--- a/Engine/Plugins/VirtualProduction/TextureShare/Source/TextureShareCore/Private/Core/ResourcesCache/D3D12/TextureShareCoreD3D12ResourcesCacheItem.cpp
+++ b/Engine/Plugins/VirtualProduction/TextureShare/Source/TextureShareCore/Private/Core/ResourcesCache/D3D12/TextureShareCoreD3D12ResourcesCacheItem.cpp
@@ -4,7 +4,7 @@
 #include "Module/TextureShareCoreLog.h"
 #include "Templates/RefCounting.h"
 
-#include COMPILED_PLATFORM_HEADER(D3D12ThirdParty.h)
+#include "D3D12ThirdParty.h"
 
 namespace UE::TextureShareCore::D3D12ResourcesCacheItem
 {
diff --git a/Engine/Source/Runtime/D3D12RHI/Private/D3D12DiskCache.h b/Engine/Source/Runtime/D3D12RHI/Private/D3D12DiskCache.h
new file mode 100644
--- /dev/null
+++ b/Engine/Source/Runtime/D3D12RHI/Private/D3D12DiskCache.h
@@ -0,0 +1,7 @@
+// Copyright Epic Games, Inc. All Rights Reserved.
+
+#pragma once
+
+#include "HAL/PreprocessorHelpers.h"
+
+#include COMPILED_PLATFORM_HEADER(D3D12DiskCache.h)
diff --git a/Engine/Source/Runtime/D3D12RHI/Private/D3D12PipelineState.h b/Engine/Source/Runtime/D3D12RHI/Private/D3D12PipelineState.h
--- a/Engine/Source/Runtime/D3D12RHI/Private/D3D12PipelineState.h
+++ b/Engine/Source/Runtime/D3D12RHI/Private/D3D12PipelineState.h
@@ -5,10 +5,9 @@
 #pragma once
 
 #include "Async/AsyncWork.h"
+#include "D3D12Util.h"
 
-#if PLATFORM_HOLOLENS
-	#include "d3d12.h"
-#endif
+#include "D3D12DiskCache.h"
 
 class FD3D12VertexShader;
 class FD3D12MeshShader;
diff --git a/Engine/Source/Runtime/D3D12RHI/Private/D3D12Queue.h b/Engine/Source/Runtime/D3D12RHI/Private/D3D12Queue.h
--- a/Engine/Source/Runtime/D3D12RHI/Private/D3D12Queue.h
+++ b/Engine/Source/Runtime/D3D12RHI/Private/D3D12Queue.h
@@ -3,6 +3,7 @@
 #pragma once
 
 #include "D3D12RHICommon.h"
+#include "D3D12RHIDefinitions.h"
 #include "RHIPipeline.h"
 
 // Defines a unique command queue type within a FD3D12Device (owner by the command list managers).
@@ -38,4 +39,15 @@ inline const TCHAR* GetD3DCommandQueueTypeName(ED3D12QueueType QueueType)
 	}
 }
 
+inline D3D12_COMMAND_LIST_TYPE GetD3DCommandListType(ED3D12QueueType QueueType)
+{
+	switch (QueueType)
+	{
+	default: checkNoEntry(); // fallthrough
+	case ED3D12QueueType::Direct: return D3D12_COMMAND_LIST_TYPE_DIRECT;
+	case ED3D12QueueType::Copy:   return D3D12RHI_PLATFORM_COPY_COMMAND_LIST_TYPE;
+	case ED3D12QueueType::Async:  return D3D12_COMMAND_LIST_TYPE_COMPUTE;
+	}
+}
+
 // TODO: move FD3D12Queue here
diff --git a/Engine/Source/Runtime/D3D12RHI/Private/D3D12RHICommon.h b/Engine/Source/Runtime/D3D12RHI/Private/D3D12RHICommon.h
--- a/Engine/Source/Runtime/D3D12RHI/Private/D3D12RHICommon.h
+++ b/Engine/Source/Runtime/D3D12RHI/Private/D3D12RHICommon.h
@@ -8,8 +8,7 @@
 
 #include "HAL/Platform.h"
 
-#include COMPILED_PLATFORM_HEADER(D3D12ThirdParty.h)
-
+#include "D3D12ThirdParty.h"
 #include "D3D12RHI.h"
 #include "D3D12RHIDefinitions.h"
 
diff --git a/Engine/Source/Runtime/D3D12RHI/Private/D3D12RHIDefinitions.h b/Engine/Source/Runtime/D3D12RHI/Private/D3D12RHIDefinitions.h
--- a/Engine/Source/Runtime/D3D12RHI/Private/D3D12RHIDefinitions.h
+++ b/Engine/Source/Runtime/D3D12RHI/Private/D3D12RHIDefinitions.h
@@ -4,6 +4,8 @@
 
 #include "RHIDefinitions.h"
 
+#include COMPILED_PLATFORM_HEADER(D3D12RHIDefinitions.h)
+
 // DX12 doesn't support higher MSAA count
 #define DX_MAX_MSAA_COUNT	8
 
@@ -98,3 +100,19 @@
 		PerEntry(DXGI_ERROR_REMOTE_OUTOFMEMORY) Terminator \
 		PerEntry(DXGI_ERROR_ACCESS_LOST) Terminator
 #endif
+
+// Note: the following defines depend on D3D12RHIBasePrivate.h
+
+//@TODO: Improve allocator efficiency so we can increase these thresholds and improve performance
+// We measured 149MB of wastage in 340MB of allocations with DEFAULT_BUFFER_POOL_MAX_ALLOC_SIZE set to 512KB
+#if !defined(DEFAULT_BUFFER_POOL_MAX_ALLOC_SIZE)
+	#if D3D12_RHI_RAYTRACING
+		// #dxr_todo: Reevaluate these values. Currently optimized to reduce number of CreateCommitedResource() calls, at the expense of memory use.
+		#define DEFAULT_BUFFER_POOL_MAX_ALLOC_SIZE    (64 * 1024 * 1024)
+		#define DEFAULT_BUFFER_POOL_DEFAULT_POOL_SIZE (16 * 1024 * 1024)
+	#else
+		// On PC, buffers are 64KB aligned, so anything smaller should be sub-allocated
+		#define DEFAULT_BUFFER_POOL_MAX_ALLOC_SIZE    (64 * 1024)
+		#define DEFAULT_BUFFER_POOL_DEFAULT_POOL_SIZE (8 * 1024 * 1024)
+	#endif //D3D12_RHI_RAYTRACING
+#endif
diff --git a/Engine/Source/Runtime/D3D12RHI/Private/D3D12RHIPrivate.h b/Engine/Source/Runtime/D3D12RHI/Private/D3D12RHIPrivate.h
--- a/Engine/Source/Runtime/D3D12RHI/Private/D3D12RHIPrivate.h
+++ b/Engine/Source/Runtime/D3D12RHI/Private/D3D12RHIPrivate.h
@@ -30,41 +30,15 @@
 #include "ShaderCore.h"
 #include "HDRHelper.h"
 
+// TODO reorder includes so we just include D3D12Submission.h here
 #include COMPILED_PLATFORM_HEADER(D3D12Submission.h)
 
-#if PLATFORM_WINDOWS
-	#include "Windows/D3D12RHIBasePrivate.h"
-#else
-	#include "D3D12RHIBasePrivate.h"
-#endif
+#include "D3D12RHIDefinitions.h"
 
-// Note: the following defines depend on D3D12RHIBasePrivate.h
-
-//@TODO: Improve allocator efficiency so we can increase these thresholds and improve performance
-// We measured 149MB of wastage in 340MB of allocations with DEFAULT_BUFFER_POOL_MAX_ALLOC_SIZE set to 512KB
-#if !defined(DEFAULT_BUFFER_POOL_MAX_ALLOC_SIZE)
-	#if D3D12_RHI_RAYTRACING
-		// #dxr_todo: Reevaluate these values. Currently optimized to reduce number of CreateCommitedResource() calls, at the expense of memory use.
-		#define DEFAULT_BUFFER_POOL_MAX_ALLOC_SIZE    (64 * 1024 * 1024)
-		#define DEFAULT_BUFFER_POOL_DEFAULT_POOL_SIZE (16 * 1024 * 1024)
-	#else
-		// On PC, buffers are 64KB aligned, so anything smaller should be sub-allocated
-		#define DEFAULT_BUFFER_POOL_MAX_ALLOC_SIZE    (64 * 1024)
-		#define DEFAULT_BUFFER_POOL_DEFAULT_POOL_SIZE (8 * 1024 * 1024)
-	#endif //D3D12_RHI_RAYTRACING
-#endif
+// TODO reorder includes so we just include D3D12PipelineState.h here
+#include COMPILED_PLATFORM_HEADER(D3D12PipelineState.h)
 
-// TODO: move D3D12RHI_PLATFORM_COPY_COMMAND_LIST_TYPE out of D3D12RHIBasePrivate
-inline D3D12_COMMAND_LIST_TYPE GetD3DCommandListType(ED3D12QueueType QueueType)
-{
-	switch (QueueType)
-	{
-	default: checkNoEntry(); // fallthrough
-	case ED3D12QueueType::Direct: return D3D12_COMMAND_LIST_TYPE_DIRECT;
-	case ED3D12QueueType::Copy:   return D3D12RHI_PLATFORM_COPY_COMMAND_LIST_TYPE;
-	case ED3D12QueueType::Async:  return D3D12_COMMAND_LIST_TYPE_COMPUTE;
-	}
-}
+#include "D3D12DiskCache.h"
 
 #if NV_AFTERMATH
 
diff --git a/Engine/Source/Runtime/D3D12RHI/Private/D3D12Util.h b/Engine/Source/Runtime/D3D12RHI/Private/D3D12Util.h
--- a/Engine/Source/Runtime/D3D12RHI/Private/D3D12Util.h
+++ b/Engine/Source/Runtime/D3D12RHI/Private/D3D12Util.h
@@ -8,7 +8,10 @@
 
 #include "CoreMinimal.h"
 #include "Containers/Queue.h"
+#include "D3D12RHICommon.h"
+#include "DXGIUtilities.h"
 #include "RenderUtils.h"
+#include "ShaderCore.h"
 
 namespace D3D12RHI
 {
diff --git a/Engine/Source/Runtime/D3D12RHI/Private/Windows/D3D12RHIBasePrivate.h b/Engine/Source/Runtime/D3D12RHI/Private/Windows/D3D12RHIBasePrivate.h
deleted file mode 100644
diff --git a/Engine/Source/Runtime/D3D12RHI/Private/Windows/WindowsD3D12PipelineState.h b/Engine/Source/Runtime/D3D12RHI/Private/Windows/WindowsD3D12PipelineState.h
--- a/Engine/Source/Runtime/D3D12RHI/Private/Windows/WindowsD3D12PipelineState.h
+++ b/Engine/Source/Runtime/D3D12RHI/Private/Windows/WindowsD3D12PipelineState.h
@@ -4,6 +4,8 @@
 
 #pragma once
 
+#include "D3D12RHICommon.h"
+
 class FD3D12RootSignature; // forward-declare
 struct FD3D12ComputePipelineStateDesc;
 struct FD3D12LowLevelGraphicsPipelineStateDesc;
diff --git a/Engine/Source/Runtime/D3D12RHI/Private/Windows/WindowsD3D12RHIDefinitions.h b/Engine/Source/Runtime/D3D12RHI/Private/Windows/WindowsD3D12RHIDefinitions.h
new file mode 100644
--- /dev/null
+++ b/Engine/Source/Runtime/D3D12RHI/Private/Windows/WindowsD3D12RHIDefinitions.h
@@ -0,0 +1,24 @@
+// Copyright Epic Games, Inc. All Rights Reserved.
+
+#pragma once
+
+// Windows uses the graphics command list interface for the copy queue.
+typedef ID3D12GraphicsCommandList ID3D12CopyCommandList;
+
+#define D3D12RHI_PLATFORM_COPY_COMMAND_LIST_TYPE D3D12_COMMAND_LIST_TYPE_COPY
+
+#define FD3D12_DESCRIPTOR_HEAP_TYPE_SAMPLER D3D12_DESCRIPTOR_HEAP_TYPE_SAMPLER
+
+// Windows uses standard D3D12 query heaps for timestamps
+#define D3D12RHI_PLATFORM_USES_TIMESTAMP_QUERIES 1
+
+#define D3D12RHI_RESOURCE_FLAG_ALLOW_INDIRECT_BUFFER D3D12_RESOURCE_FLAG_NONE
+#define D3D12RHI_HEAP_FLAG_ALLOW_INDIRECT_BUFFERS		D3D12_HEAP_FLAG_NONE
+
+#define FD3D12_HEAP_FLAG_CREATE_NOT_ZEROED				D3D12_HEAP_FLAG_CREATE_NOT_ZEROED
+
+#define D3D12RHI_NEEDS_VENDOR_EXTENSIONS     1
+#define D3D12RHI_NEEDS_SHADER_FEATURE_CHECKS 1
+
+#define USE_STATIC_ROOT_SIGNATURE			0
+#define D3D12_USE_DUMMY_BACKBUFFER			1
diff --git a/Engine/Source/Runtime/D3D12RHI/Public/D3D12ThirdParty.h b/Engine/Source/Runtime/D3D12RHI/Public/D3D12ThirdParty.h
new file mode 100644
--- /dev/null
+++ b/Engine/Source/Runtime/D3D12RHI/Public/D3D12ThirdParty.h
@@ -0,0 +1,7 @@
+// Copyright Epic Games, Inc. All Rights Reserved.
+
+#pragma once
+
+#include "HAL/PreprocessorHelpers.h"
+
+#include COMPILED_PLATFORM_HEADER(D3D12ThirdParty.h)
diff --git a/Engine/Source/Runtime/D3D12RHI/Public/ID3D12DynamicRHI.h b/Engine/Source/Runtime/D3D12RHI/Public/ID3D12DynamicRHI.h
--- a/Engine/Source/Runtime/D3D12RHI/Public/ID3D12DynamicRHI.h
+++ b/Engine/Source/Runtime/D3D12RHI/Public/ID3D12DynamicRHI.h
@@ -2,11 +2,10 @@
 
 #pragma once
 
+#include "D3D12ThirdParty.h"
 #include "RHI.h"
 #include "Templates/Function.h"
 
-#include COMPILED_PLATFORM_HEADER(D3D12ThirdParty.h)
-
 struct FD3D12MinimalAdapterDesc
 {
 	DXGI_ADAPTER_DESC Desc{};
