Enable the forced committed resource path for texture creation when driver is too old

diff --git a/Engine/Source/Runtime/D3D12RHI/Private/D3D12PoolAllocator.cpp b/Engine/Source/Runtime/D3D12RHI/Private/D3D12PoolAllocator.cpp
--- a/Engine/Source/Runtime/D3D12RHI/Private/D3D12PoolAllocator.cpp
+++ b/Engine/Source/Runtime/D3D12RHI/Private/D3D12PoolAllocator.cpp
@@ -7,14 +7,6 @@
 #define NEEDS_D3D12_INDIRECT_ARGUMENT_HEAP_WORKAROUND 0
 #endif
 
-int32 GUseCommittedTexturesNV = 0;
-static FAutoConsoleVariableRef CVarUseCommittedTexturesNV(
-	TEXT("r.D3D12.UseCommittedTexturesNV"),
-	GUseCommittedTexturesNV,
-	TEXT("Force committed resource creation for textures to fix possible driver bug"),
-	ECVF_Default
-);
-
 //-----------------------------------------------------------------------------
 //	FD3D12MemoryPool
 //-----------------------------------------------------------------------------
@@ -332,8 +324,7 @@ void FD3D12PoolAllocator::AllocateResource(uint32 GPUIndex, D3D12_HEAP_TYPE InHe
 	// Disable pooling for VRAM allocated textures and force use the committed resource path
 	bool bForceCommittedResourcePath = false;
 #if PLATFORM_WINDOWS
-	const bool bUseCommittedTexturesNV = GUseCommittedTexturesNV && IsRHIDeviceNVIDIA();
-	if (bPoolResource && bUseCommittedTexturesNV && InHeapType == D3D12_HEAP_TYPE_DEFAULT && InDesc.Dimension != D3D12_RESOURCE_DIMENSION_BUFFER)
+	if (bPoolResource && GD3D12WorkaroundFlags.bForceCommittedResourceTextureAllocation && InHeapType == D3D12_HEAP_TYPE_DEFAULT && InDesc.Dimension != D3D12_RESOURCE_DIMENSION_BUFFER)
 	{
 		bForceCommittedResourcePath = true;
 	}
diff --git a/Engine/Source/Runtime/D3D12RHI/Private/D3D12RHIPrivate.h b/Engine/Source/Runtime/D3D12RHI/Private/D3D12RHIPrivate.h
--- a/Engine/Source/Runtime/D3D12RHI/Private/D3D12RHIPrivate.h
+++ b/Engine/Source/Runtime/D3D12RHI/Private/D3D12RHIPrivate.h
@@ -182,6 +182,12 @@ struct FD3D12WorkaroundFlags
 	* This is less efficient and can also trigger another known issue with D3D12 Agility version <= 4.
 	*/
 	bool bAllowGetShaderIdentifierOnCollectionSubObject = true;
+
+	/**
+	* Certain drivers can cause texture corruption when sub allocating textures from a shared heap. Committed resource
+	* allocations are used then
+	*/
+	bool bForceCommittedResourceTextureAllocation = false;
 };
 
 extern FD3D12WorkaroundFlags GD3D12WorkaroundFlags;
diff --git a/Engine/Source/Runtime/D3D12RHI/Private/Windows/WindowsD3D12Device.cpp b/Engine/Source/Runtime/D3D12RHI/Private/Windows/WindowsD3D12Device.cpp
--- a/Engine/Source/Runtime/D3D12RHI/Private/Windows/WindowsD3D12Device.cpp
+++ b/Engine/Source/Runtime/D3D12RHI/Private/Windows/WindowsD3D12Device.cpp
@@ -1511,6 +1511,13 @@ void FD3D12DynamicRHI::Init()
 			GD3D12WorkaroundFlags.bAllowGetShaderIdentifierOnCollectionSubObject = false;
 			UE_LOG(LogD3D12RHI, Warning, TEXT("GD3D12WorkaroundFlags.bAllowGetShaderIdentifierOnCollectionSubObject is disabled due to a known issue with current driver version."));
 		}
+
+		if (DriverVersion < 53141u)
+		{
+			GD3D12WorkaroundFlags.bForceCommittedResourceTextureAllocation = true;
+			UE_LOG(LogD3D12RHI, Log, TEXT("GD3D12WorkaroundFlags.bForceCommittedResourceTextureAllocation is set due to a known issue with current driver version."));
+		}
+
 	} // if NVIDIA
 #endif // NVAPI
 
