Move GetRaytracingAccelerationStructurePrebuildInfo into the device to be able to modify it per platform.

diff --git a/Engine/Source/Runtime/D3D12RHI/Private/D3D12Device.h b/Engine/Source/Runtime/D3D12RHI/Private/D3D12Device.h
--- a/Engine/Source/Runtime/D3D12RHI/Private/D3D12Device.h
+++ b/Engine/Source/Runtime/D3D12RHI/Private/D3D12Device.h
@@ -282,6 +282,10 @@ public:
 
 	TRefCountPtr<ID3D12StateObject>			  DeserializeRayTracingStateObject(D3D12_SHADER_BYTECODE Bytecode, ID3D12RootSignature* RootSignature);
 
+	void GetRaytracingAccelerationStructurePrebuildInfo(
+		const D3D12_BUILD_RAYTRACING_ACCELERATION_STRUCTURE_INPUTS* pDesc,
+		D3D12_RAYTRACING_ACCELERATION_STRUCTURE_PREBUILD_INFO* pInfo);
+
 	// Queries ray tracing pipeline state object metrics such as VGPR usage (if available/supported). Returns true if query succeeded.
 	bool GetRayTracingPipelineInfo(ID3D12StateObject* Pipeline, FD3D12RayTracingPipelineInfo* OutInfo);
 #endif // D3D12_RHI_RAYTRACING
diff --git a/Engine/Source/Runtime/D3D12RHI/Private/D3D12RayTracing.cpp b/Engine/Source/Runtime/D3D12RHI/Private/D3D12RayTracing.cpp
--- a/Engine/Source/Runtime/D3D12RHI/Private/D3D12RayTracing.cpp
+++ b/Engine/Source/Runtime/D3D12RHI/Private/D3D12RayTracing.cpp
@@ -2968,9 +2968,8 @@ FRayTracingAccelerationStructureSize FD3D12DynamicRHI::RHICalcRayTracingSceneSiz
 	FRayTracingAccelerationStructureSize SizeInfo = {};
 	for (uint32 GPUIndex = 0; GPUIndex < GNumExplicitGPUsForRendering; ++GPUIndex)
 	{
-		ID3D12Device5* RayTracingDevice = Adapter.GetDevice(GPUIndex)->GetDevice5();
 		D3D12_RAYTRACING_ACCELERATION_STRUCTURE_PREBUILD_INFO PrebuildInfo = {};
-		RayTracingDevice->GetRaytracingAccelerationStructurePrebuildInfo(&BuildInputs, &PrebuildInfo);
+		Adapter.GetDevice(GPUIndex)->GetRaytracingAccelerationStructurePrebuildInfo(&BuildInputs, &PrebuildInfo);
 
 		SizeInfo.ResultSize = FMath::Max(SizeInfo.ResultSize, PrebuildInfo.ResultDataMaxSizeInBytes);
 		SizeInfo.BuildScratchSize  = FMath::Max(SizeInfo.BuildScratchSize, PrebuildInfo.ScratchDataSizeInBytes);
@@ -3176,9 +3175,8 @@ FRayTracingAccelerationStructureSize FD3D12DynamicRHI::RHICalcRayTracingGeometry
 		// Get maximum buffer sizes for all GPUs in the system
 		for (uint32 GPUIndex = 0; GPUIndex < GNumExplicitGPUsForRendering; ++GPUIndex)
 		{
-			ID3D12Device5* RayTracingDevice = Adapter.GetDevice(GPUIndex)->GetDevice5();
 			D3D12_RAYTRACING_ACCELERATION_STRUCTURE_PREBUILD_INFO PrebuildInfo = {};
-			RayTracingDevice->GetRaytracingAccelerationStructurePrebuildInfo(&PrebuildDescInputs, &PrebuildInfo);
+			Adapter.GetDevice(GPUIndex)->GetRaytracingAccelerationStructurePrebuildInfo(&PrebuildDescInputs, &PrebuildInfo);
 
 			SizeInfo.ResultSize = FMath::Max(SizeInfo.ResultSize, PrebuildInfo.ResultDataMaxSizeInBytes);
 			SizeInfo.BuildScratchSize = FMath::Max(SizeInfo.BuildScratchSize, PrebuildInfo.ScratchDataSizeInBytes);
@@ -3881,8 +3879,7 @@ void FD3D12RayTracingGeometry::CreateAccelerationStructureBuildDesc(FD3D12Comman
 
 	D3D12_RAYTRACING_ACCELERATION_STRUCTURE_PREBUILD_INFO PrebuildInfo = {};
 
-	ID3D12Device5* RayTracingDevice = CommandContext.GetParentDevice()->GetDevice5();
-	RayTracingDevice->GetRaytracingAccelerationStructurePrebuildInfo(&PrebuildDescInputs, &PrebuildInfo);
+	CommandContext.GetParentDevice()->GetRaytracingAccelerationStructurePrebuildInfo(&PrebuildDescInputs, &PrebuildInfo);
 
 	// Must make sure that values computed in the constructor are valid.
 	check(PrebuildInfo.ResultDataMaxSizeInBytes <= SizeInfo.ResultSize);
@@ -4060,7 +4057,6 @@ void FD3D12RayTracingScene::BuildAccelerationStructure(FD3D12CommandContext& Com
 
 	const uint32 GPUIndex = CommandContext.GetGPUIndex();
 	FD3D12Adapter* Adapter = CommandContext.GetParentAdapter();
-	ID3D12Device5* RayTracingDevice = CommandContext.GetParentDevice()->GetDevice5();
 
 	TRefCountPtr<FD3D12Buffer> AutoScratchBuffer;
 	if (ScratchBuffer == nullptr)
@@ -4081,7 +4077,7 @@ void FD3D12RayTracingScene::BuildAccelerationStructure(FD3D12CommandContext& Com
 
 		D3D12_RAYTRACING_ACCELERATION_STRUCTURE_PREBUILD_INFO PrebuildInfo = {};
 
-		RayTracingDevice->GetRaytracingAccelerationStructurePrebuildInfo(&Layer.BuildInputs, &PrebuildInfo);
+		CommandContext.GetParentDevice()->GetRaytracingAccelerationStructurePrebuildInfo(&Layer.BuildInputs, &PrebuildInfo);
 
 		checkf(PrebuildInfo.ResultDataMaxSizeInBytes <= Layer.SizeInfo.ResultSize,
 			TEXT("TLAS build result buffer now requires %lld bytes, but only %lld was calculated in the constructor."),
diff --git a/Engine/Source/Runtime/D3D12RHI/Private/Windows/WindowsD3D12Device.cpp b/Engine/Source/Runtime/D3D12RHI/Private/Windows/WindowsD3D12Device.cpp
--- a/Engine/Source/Runtime/D3D12RHI/Private/Windows/WindowsD3D12Device.cpp
+++ b/Engine/Source/Runtime/D3D12RHI/Private/Windows/WindowsD3D12Device.cpp
@@ -1775,6 +1775,14 @@ void FD3D12Device::CreateSamplerInternal(const D3D12_SAMPLER_DESC& Desc, D3D12_C
 }
 
 #if D3D12_RHI_RAYTRACING
+void FD3D12Device::GetRaytracingAccelerationStructurePrebuildInfo(
+	const D3D12_BUILD_RAYTRACING_ACCELERATION_STRUCTURE_INPUTS* pDesc,
+	D3D12_RAYTRACING_ACCELERATION_STRUCTURE_PREBUILD_INFO* pInfo)
+{
+	ID3D12Device5* RayTracingDevice = GetDevice5();
+	RayTracingDevice->GetRaytracingAccelerationStructurePrebuildInfo(pDesc, pInfo);
+}
+
 TRefCountPtr<ID3D12StateObject> FD3D12Device::DeserializeRayTracingStateObject(D3D12_SHADER_BYTECODE Bytecode, ID3D12RootSignature* RootSignature)
 {
 	checkNoEntry();
