MGPU:  Enable D3D12 submission thread by default.  It had been disabled with MGPU due to crashes, but the unknown cause of the issue seems to have been fixed by more recent parallel rendering work.  All the Virtual Production test maps have been successfully tested with the submission thread re-enabled, with no crashes.

#rnx

diff --git a/Engine/Source/Runtime/D3D12RHI/Private/D3D12Submission.cpp b/Engine/Source/Runtime/D3D12RHI/Private/D3D12Submission.cpp
--- a/Engine/Source/Runtime/D3D12RHI/Private/D3D12Submission.cpp
+++ b/Engine/Source/Runtime/D3D12RHI/Private/D3D12Submission.cpp
@@ -27,7 +27,7 @@ static FAutoConsoleVariableRef CVarD3D12AsyncPayloadMerge(
 // @todo mgpu - fix crashes when submission thread is enabled)
 static TAutoConsoleVariable<int32> CVarRHIUseSubmissionThread(
 	TEXT("rhi.UseSubmissionThread"),
-	1,
+	2,
 	TEXT("Whether to enable the RHI submission thread.\n")
 	TEXT("  0: No\n")
 	TEXT("  1: Yes, but not when running with multi-gpu.\n")
