Moving D3D12 Extension includes into their own headers.

diff --git a/Engine/Source/Runtime/D3D12RHI/Private/D3D12Adapter.cpp b/Engine/Source/Runtime/D3D12RHI/Private/D3D12Adapter.cpp
--- a/Engine/Source/Runtime/D3D12RHI/Private/D3D12Adapter.cpp
+++ b/Engine/Source/Runtime/D3D12RHI/Private/D3D12Adapter.cpp
@@ -5,30 +5,21 @@ D3D12Adapter.cpp:D3D12 Adapter implementation.
 =============================================================================*/
 
 #include "D3D12RHIPrivate.h"
+#include "D3D12AmdExtensions.h"
+#include "D3D12IntelExtensions.h"
 #include "Misc/CommandLine.h"
 #include "Misc/EngineVersion.h"
 #include "Misc/OutputDeviceRedirector.h"
+#include "DataDrivenShaderPlatformInfo.h"
 #include "Windows/AllowWindowsPlatformTypes.h"
 #include "GenericPlatform/GenericPlatformCrashContext.h"
-#include "DataDrivenShaderPlatformInfo.h"
 #if PLATFORM_WINDOWS
 #include "Windows/WindowsPlatformMisc.h"
 #include "Windows/WindowsPlatformStackWalk.h"
 #endif
 #include "Modules/ModuleManager.h"
-
-#if WITH_AMD_AGS
-	#include "amd_ags.h"
-#endif
 #include "Windows/HideWindowsPlatformTypes.h"
 
-#if INTEL_EXTENSIONS
-	#define INTC_IGDEXT_D3D12 1
-
-	THIRD_PARTY_INCLUDES_START
-	#include "igdext.h"
-	THIRD_PARTY_INCLUDES_END
-#endif
 
 #if ENABLE_RESIDENCY_MANAGEMENT
 bool GEnableResidencyManagement = true;
diff --git a/Engine/Source/Runtime/D3D12RHI/Private/D3D12AmdExtensions.h b/Engine/Source/Runtime/D3D12RHI/Private/D3D12AmdExtensions.h
new file mode 100644
--- /dev/null
+++ b/Engine/Source/Runtime/D3D12RHI/Private/D3D12AmdExtensions.h
@@ -0,0 +1,11 @@
+// Copyright Epic Games, Inc. All Rights Reserved.
+
+#pragma once
+
+#if WITH_AMD_AGS
+	#include "Microsoft/AllowMicrosoftPlatformTypes.h"
+	THIRD_PARTY_INCLUDES_START
+		#include <amd_ags.h>
+	THIRD_PARTY_INCLUDES_END
+	#include "Microsoft/HideMicrosoftPlatformTypes.h"
+#endif // WITH_AMD_AGS
diff --git a/Engine/Source/Runtime/D3D12RHI/Private/D3D12CommandContext.cpp b/Engine/Source/Runtime/D3D12RHI/Private/D3D12CommandContext.cpp
--- a/Engine/Source/Runtime/D3D12RHI/Private/D3D12CommandContext.cpp
+++ b/Engine/Source/Runtime/D3D12RHI/Private/D3D12CommandContext.cpp
@@ -6,12 +6,7 @@ D3D12CommandContext.cpp: RHI  Command Context implementation.
 
 #include "D3D12RHIPrivate.h"
 
-#if WITH_AMD_AGS
-#include "Windows/AllowWindowsPlatformTypes.h"
-#include "amd_ags.h"
-#include "Windows/HideWindowsPlatformTypes.h"
-#endif
-
+#include "D3D12AmdExtensions.h"
 #include "D3D12RayTracing.h"
 
 int32 GD3D12MaxCommandsPerCommandList = 10000;
diff --git a/Engine/Source/Runtime/D3D12RHI/Private/D3D12CommandList.h b/Engine/Source/Runtime/D3D12RHI/Private/D3D12CommandList.h
--- a/Engine/Source/Runtime/D3D12RHI/Private/D3D12CommandList.h
+++ b/Engine/Source/Runtime/D3D12RHI/Private/D3D12CommandList.h
@@ -6,6 +6,7 @@ D3D12CommandList.h: Implementation of D3D12 Command List functions
 #pragma once
 
 #include "CoreMinimal.h"
+#include "D3D12NvidiaExtensions.h"
 
 #if !defined(D3D12_PLATFORM_SUPPORTS_ASSERTRESOURCESTATES)
 	#define D3D12_PLATFORM_SUPPORTS_ASSERTRESOURCESTATES 1
diff --git a/Engine/Source/Runtime/D3D12RHI/Private/D3D12Device.cpp b/Engine/Source/Runtime/D3D12RHI/Private/D3D12Device.cpp
--- a/Engine/Source/Runtime/D3D12RHI/Private/D3D12Device.cpp
+++ b/Engine/Source/Runtime/D3D12RHI/Private/D3D12Device.cpp
@@ -1,16 +1,9 @@
 // Copyright Epic Games, Inc. All Rights Reserved.
 
 #include "D3D12RHIPrivate.h"
+#include "D3D12IntelExtensions.h"
 #include "D3D12RayTracing.h"
 
-#if INTEL_EXTENSIONS
-	#define INTC_IGDEXT_D3D12 1
-
-	THIRD_PARTY_INCLUDES_START
-	#include "igdext.h"
-	THIRD_PARTY_INCLUDES_END
-#endif
-
 static TAutoConsoleVariable<int32> CVarD3D12GPUTimeout(
 	TEXT("r.D3D12.GPUTimeout"),
 	1,
diff --git a/Engine/Source/Runtime/D3D12RHI/Private/D3D12IntelExtensions.h b/Engine/Source/Runtime/D3D12RHI/Private/D3D12IntelExtensions.h
new file mode 100644
--- /dev/null
+++ b/Engine/Source/Runtime/D3D12RHI/Private/D3D12IntelExtensions.h
@@ -0,0 +1,22 @@
+// Copyright Epic Games, Inc. All Rights Reserved.
+
+#pragma once
+
+#if INTEL_EXTENSIONS
+	#define INTC_IGDEXT_D3D12 1
+
+	#include "Microsoft/AllowMicrosoftPlatformTypes.h"
+	THIRD_PARTY_INCLUDES_START
+		#include <igdext.h>
+	THIRD_PARTY_INCLUDES_END
+	#include "Microsoft/HideMicrosoftPlatformTypes.h"
+
+	extern bool GDX12INTCAtomicUInt64Emulation;
+
+	struct INTCExtensionContext;
+	struct INTCExtensionInfo;
+
+	void DestroyIntelExtensionsContext(INTCExtensionContext* IntelExtensionContext);
+	INTCExtensionContext* CreateIntelExtensionsContext(ID3D12Device* Device, INTCExtensionInfo& INTCExtensionInfo);
+	bool EnableIntelAtomic64Support(INTCExtensionContext* IntelExtensionContext, INTCExtensionInfo& INTCExtensionInfo);
+#endif //INTEL_EXTENSIONS
diff --git a/Engine/Source/Runtime/D3D12RHI/Private/D3D12NvidiaExtensions.h b/Engine/Source/Runtime/D3D12RHI/Private/D3D12NvidiaExtensions.h
new file mode 100644
--- /dev/null
+++ b/Engine/Source/Runtime/D3D12RHI/Private/D3D12NvidiaExtensions.h
@@ -0,0 +1,28 @@
+// Copyright Epic Games, Inc. All Rights Reserved.
+
+#pragma once
+
+#if WITH_NVAPI
+	#include "Microsoft/AllowMicrosoftPlatformTypes.h"
+	THIRD_PARTY_INCLUDES_START
+		#include "nvapi.h"
+		#include "nvShaderExtnEnums.h"
+	THIRD_PARTY_INCLUDES_END
+	#include "Microsoft/HideMicrosoftPlatformTypes.h"
+#endif
+
+#if NV_AFTERMATH
+	#include "Microsoft/AllowMicrosoftPlatformTypes.h"
+	THIRD_PARTY_INCLUDES_START
+		#define GFSDK_Aftermath_WITH_DX12 1
+		#include "GFSDK_Aftermath.h"
+		#include "GFSDK_Aftermath_GpuCrashdump.h"
+		#undef GFSDK_Aftermath_WITH_DX12
+	THIRD_PARTY_INCLUDES_END
+	#include "Microsoft/HideMicrosoftPlatformTypes.h"
+
+	extern bool GDX12NVAfterMathModuleLoaded;
+	extern int32 GDX12NVAfterMathEnabled;
+	extern int32 GDX12NVAfterMathTrackResources;
+	extern int32 GDX12NVAfterMathMarkers;
+#endif // NV_AFTERMATH
diff --git a/Engine/Source/Runtime/D3D12RHI/Private/D3D12RHI.cpp b/Engine/Source/Runtime/D3D12RHI/Private/D3D12RHI.cpp
--- a/Engine/Source/Runtime/D3D12RHI/Private/D3D12RHI.cpp
+++ b/Engine/Source/Runtime/D3D12RHI/Private/D3D12RHI.cpp
@@ -8,11 +8,8 @@
 #include "RHIStaticStates.h"
 #include "OneColorShader.h"
 
-#if WITH_AMD_AGS
-#include "Windows/AllowWindowsPlatformTypes.h"
-#include "amd_ags.h"
-#include "Windows/HideWindowsPlatformTypes.h"
-#endif
+#include "D3D12AmdExtensions.h"
+#include "D3D12IntelExtensions.h"
 
 #if !defined(D3D12_PLATFORM_NEEDS_DISPLAY_MODE_ENUMERATION)
 	#define D3D12_PLATFORM_NEEDS_DISPLAY_MODE_ENUMERATION 1
diff --git a/Engine/Source/Runtime/D3D12RHI/Private/D3D12RHIPrivate.h b/Engine/Source/Runtime/D3D12RHI/Private/D3D12RHIPrivate.h
--- a/Engine/Source/Runtime/D3D12RHI/Private/D3D12RHIPrivate.h
+++ b/Engine/Source/Runtime/D3D12RHI/Private/D3D12RHIPrivate.h
@@ -39,32 +39,8 @@
 #include COMPILED_PLATFORM_HEADER(D3D12PipelineState.h)
 
 #include "D3D12DiskCache.h"
-
-#if NV_AFTERMATH
-
-	#define GFSDK_Aftermath_WITH_DX12 1
-		#include "GFSDK_Aftermath.h"
-		#include "GFSDK_Aftermath_GpuCrashdump.h"
-	#undef GFSDK_Aftermath_WITH_DX12
-
-	extern bool GDX12NVAfterMathModuleLoaded;
-	extern int32 GDX12NVAfterMathEnabled;
-	extern int32 GDX12NVAfterMathTrackResources;
-	extern int32 GDX12NVAfterMathMarkers;
-
-#endif // NV_AFTERMATH
-
-#if INTEL_EXTENSIONS
-	extern bool GDX12INTCAtomicUInt64Emulation;
-
-	struct INTCExtensionContext;
-	struct INTCExtensionInfo;
-
-	void DestroyIntelExtensionsContext(INTCExtensionContext* IntelExtensionContext);
-	INTCExtensionContext* CreateIntelExtensionsContext(ID3D12Device* Device, INTCExtensionInfo& INTCExtensionInfo);
-	bool EnableIntelAtomic64Support(INTCExtensionContext* IntelExtensionContext, INTCExtensionInfo& INTCExtensionInfo);
-#endif //INTEL_EXTENSIONS
-
+#include "D3D12NvidiaExtensions.h"
+#include "D3D12IntelExtensions.h"
 #include "D3D12Residency.h"
 
 // D3D RHI public headers.
diff --git a/Engine/Source/Runtime/D3D12RHI/Private/D3D12Resources.cpp b/Engine/Source/Runtime/D3D12RHI/Private/D3D12Resources.cpp
--- a/Engine/Source/Runtime/D3D12RHI/Private/D3D12Resources.cpp
+++ b/Engine/Source/Runtime/D3D12RHI/Private/D3D12Resources.cpp
@@ -5,19 +5,12 @@ D3D12Resources.cpp: D3D RHI utility implementation.
 =============================================================================*/
 
 #include "D3D12RHIPrivate.h"
+#include "D3D12IntelExtensions.h"
 #include "EngineModule.h"
 #include "HAL/LowLevelMemTracker.h"
 #include "ProfilingDebugging/MemoryTrace.h"
 #include "ProfilingDebugging/AssetMetadataTrace.h"
 
-#if INTEL_EXTENSIONS
-	#define INTC_IGDEXT_D3D12 1
-
-	THIRD_PARTY_INCLUDES_START
-	#include "igdext.h"
-	THIRD_PARTY_INCLUDES_END
-#endif
-
 static TAutoConsoleVariable<int32> CVarD3D12ReservedResourceHeapSizeMB(
 	TEXT("d3d12.ReservedResourceHeapSizeMB"),
 	16,
diff --git a/Engine/Source/Runtime/D3D12RHI/Private/D3D12Resources.h b/Engine/Source/Runtime/D3D12RHI/Private/D3D12Resources.h
--- a/Engine/Source/Runtime/D3D12RHI/Private/D3D12Resources.h
+++ b/Engine/Source/Runtime/D3D12RHI/Private/D3D12Resources.h
@@ -8,6 +8,7 @@
 
 #include "D3D12RHIPrivate.h"
 #include "BoundShaderStateCache.h"
+#include "D3D12NvidiaExtensions.h"
 #include "D3D12ShaderResources.h"
 #include "D3D12Residency.h"
 #include "D3D12Util.h"
diff --git a/Engine/Source/Runtime/D3D12RHI/Private/D3D12RootSignature.cpp b/Engine/Source/Runtime/D3D12RHI/Private/D3D12RootSignature.cpp
--- a/Engine/Source/Runtime/D3D12RHI/Private/D3D12RootSignature.cpp
+++ b/Engine/Source/Runtime/D3D12RHI/Private/D3D12RootSignature.cpp
@@ -5,6 +5,7 @@
 =============================================================================*/
 
 #include "D3D12RHIPrivate.h"
+#include "D3D12AmdExtensions.h"
 #include "D3D12RootSignatureDefinitions.h"
 #include "RayTracingBuiltInResources.h"
 
@@ -12,12 +13,6 @@
 #define FD3D12_ROOT_SIGNATURE_FLAG_GLOBAL_ROOT_SIGNATURE D3D12_ROOT_SIGNATURE_FLAG_NONE
 #endif
 
-#if WITH_AMD_AGS
-#include "Microsoft/AllowMicrosoftPlatformTypes.h"
-#include "amd_ags.h"
-#include "Microsoft/HideMicrosoftPlatformTypes.h"
-#endif
-
 // Allows to automatically bind UEDiagnosticBuffer UAV, available to all shaders.
 // If a shader is compiled with diagnostics enabled, it will fail to load/create 
 // unless D3D12_ALLOW_SHADER_DIAGNOSTIC_BUFFER is enabled.
diff --git a/Engine/Source/Runtime/D3D12RHI/Private/D3D12Stats.h b/Engine/Source/Runtime/D3D12RHI/Private/D3D12Stats.h
--- a/Engine/Source/Runtime/D3D12RHI/Private/D3D12Stats.h
+++ b/Engine/Source/Runtime/D3D12RHI/Private/D3D12Stats.h
@@ -1,10 +1,9 @@
 // Copyright Epic Games, Inc. All Rights Reserved.
 
-/*=============================================================================
-D3D12Stats.h: D3D12 Statistics and Timing Interfaces
-=============================================================================*/
 #pragma once
 
+#include "D3D12NvidiaExtensions.h"
+
 /**
 * The D3D RHI stats.
 */
diff --git a/Engine/Source/Runtime/D3D12RHI/Private/Windows/WindowsD3D12Device.cpp b/Engine/Source/Runtime/D3D12RHI/Private/Windows/WindowsD3D12Device.cpp
--- a/Engine/Source/Runtime/D3D12RHI/Private/Windows/WindowsD3D12Device.cpp
+++ b/Engine/Source/Runtime/D3D12RHI/Private/Windows/WindowsD3D12Device.cpp
@@ -5,6 +5,9 @@
 =============================================================================*/
 
 #include "D3D12RHIPrivate.h"
+#include "D3D12AmdExtensions.h"
+#include "D3D12IntelExtensions.h"
+#include "D3D12NvidiaExtensions.h"
 #include "WindowsD3D12Adapter.h"
 #include "Modules/ModuleManager.h"
 #include "Windows/AllowWindowsPlatformTypes.h"
@@ -12,30 +15,8 @@
 #include "HAL/FileManager.h"
 #include "ProfilingDebugging/AssetMetadataTrace.h"
 #include <delayimp.h>
-
-#if WITH_AMD_AGS
-	#include "amd_ags.h"
-#endif
-
-#if WITH_NVAPI
-	#include "nvapi.h"
-	#include "nvShaderExtnEnums.h"
-#endif
-
-#if INTEL_EXTENSIONS
-	bool GDX12INTCAtomicUInt64Emulation = false;
-#endif
-
 #include "Windows/HideWindowsPlatformTypes.h"
 
-#if INTEL_EXTENSIONS
-	#define INTC_IGDEXT_D3D12 1
-
-	THIRD_PARTY_INCLUDES_START
-	#include "igdext.h"
-	THIRD_PARTY_INCLUDES_END
-#endif
-
 #include "HardwareInfo.h"
 #include "IHeadMountedDisplayModule.h"
 #include "GenericPlatform/GenericPlatformDriver.h"			// FGPUDriverInfo
@@ -53,6 +34,10 @@ extern bool D3D12RHI_AllowSoftwareFallback();
 extern bool D3D12RHI_ShouldAllowAsyncResourceCreation();
 extern bool D3D12RHI_ShouldForceCompatibility();
 
+#if INTEL_EXTENSIONS
+bool GDX12INTCAtomicUInt64Emulation = false;
+#endif
+
 FD3D12DynamicRHI* GD3D12RHI = nullptr;
 
 #if NV_AFTERMATH
diff --git a/Engine/Source/Runtime/D3D12RHI/Private/Windows/WindowsD3D12PipelineState.cpp b/Engine/Source/Runtime/D3D12RHI/Private/Windows/WindowsD3D12PipelineState.cpp
--- a/Engine/Source/Runtime/D3D12RHI/Private/Windows/WindowsD3D12PipelineState.cpp
+++ b/Engine/Source/Runtime/D3D12RHI/Private/Windows/WindowsD3D12PipelineState.cpp
@@ -6,13 +6,9 @@
 //	Include Files
 //-----------------------------------------------------------------------------
 #include "D3D12RHIPrivate.h"
+#include "D3D12NvidiaExtensions.h"
 #include "Misc/ScopeRWLock.h"
 #include "Stats/StatsMisc.h"
-
-#if WITH_NVAPI
-	#include "nvapi.h"
-#endif
-
 #include "d3dcompiler.h"
 
 // UE-65533
