Getting a few D3D12RHI headers to compile independently by moving some code to D3D12RHICommon.h

diff --git a/Engine/Source/Runtime/D3D12RHI/Private/D3D12CommandContext.h b/Engine/Source/Runtime/D3D12RHI/Private/D3D12CommandContext.h
--- a/Engine/Source/Runtime/D3D12RHI/Private/D3D12CommandContext.h
+++ b/Engine/Source/Runtime/D3D12RHI/Private/D3D12CommandContext.h
@@ -7,10 +7,11 @@ D3D12CommandContext.h: D3D12 Command Context Interfaces
 #pragma once
 
 #include "D3D12RHIPrivate.h"
+#include "D3D12Queue.h"
+
 #include "Windows/AllowWindowsPlatformTypes.h"
-THIRD_PARTY_INCLUDES_START
-#include <delayimp.h>
 
+THIRD_PARTY_INCLUDES_START
 #if USE_PIX
 	#include "pix3.h"
 #endif
diff --git a/Engine/Source/Runtime/D3D12RHI/Private/D3D12CommandList.h b/Engine/Source/Runtime/D3D12RHI/Private/D3D12CommandList.h
--- a/Engine/Source/Runtime/D3D12RHI/Private/D3D12CommandList.h
+++ b/Engine/Source/Runtime/D3D12RHI/Private/D3D12CommandList.h
@@ -5,17 +5,33 @@ D3D12CommandList.h: Implementation of D3D12 Command List functions
 =============================================================================*/
 #pragma once
 
-#include "CoreMinimal.h"
+#include "D3D12RHICommon.h"
+#include "RHICommandList.h"
 #include "D3D12NvidiaExtensions.h"
+#include "D3D12Resources.h"
+#include "D3D12Submission.h"
+#include "D3D12Util.h"
 
 #if !defined(D3D12_PLATFORM_SUPPORTS_ASSERTRESOURCESTATES)
 	#define D3D12_PLATFORM_SUPPORTS_ASSERTRESOURCESTATES 1
 #endif
 
-class FD3D12Device;
 class FD3D12ContextCommon;
-class FD3D12Queue;
+class FD3D12Device;
+class FD3D12DynamicRHI;
 class FD3D12QueryAllocator;
+class FD3D12Queue;
+class FD3D12ResourceBarrierBatcher;
+
+struct FD3D12ResidencyHandle;
+
+enum class ED3D12QueueType;
+
+namespace D3DX12Residency
+{
+	class ResidencySet;
+}
+typedef D3DX12Residency::ResidencySet FD3D12ResidencySet;
 
 //
 // Wraps a D3D command list allocator object.
diff --git a/Engine/Source/Runtime/D3D12RHI/Private/D3D12Device.h b/Engine/Source/Runtime/D3D12RHI/Private/D3D12Device.h
--- a/Engine/Source/Runtime/D3D12RHI/Private/D3D12Device.h
+++ b/Engine/Source/Runtime/D3D12RHI/Private/D3D12Device.h
@@ -86,19 +86,6 @@ public:
 	{}
 };
 
-// Lock free pointer list that auto-destructs items remaining in the list.
-template <typename TObjectType>
-struct TD3D12ObjectPool : public TLockFreePointerListUnordered<TObjectType, PLATFORM_CACHE_LINE_SIZE>
-{
-	~TD3D12ObjectPool()
-	{
-		while (TObjectType* Object = TLockFreePointerListUnordered<TObjectType, PLATFORM_CACHE_LINE_SIZE>::Pop())
-		{
-			delete Object;
-		}
-	}
-};
-
 // Encapsulates a single D3D command queue, and maintains the 
 // state required by the submission thread for managing the queue.
 class FD3D12Queue final
diff --git a/Engine/Source/Runtime/D3D12RHI/Private/D3D12PipelineState.h b/Engine/Source/Runtime/D3D12RHI/Private/D3D12PipelineState.h
--- a/Engine/Source/Runtime/D3D12RHI/Private/D3D12PipelineState.h
+++ b/Engine/Source/Runtime/D3D12RHI/Private/D3D12PipelineState.h
@@ -5,9 +5,8 @@
 #pragma once
 
 #include "Async/AsyncWork.h"
-#include "D3D12Util.h"
-
 #include "D3D12DiskCache.h"
+#include "D3D12Shader.h"
 
 class FD3D12VertexShader;
 class FD3D12MeshShader;
diff --git a/Engine/Source/Runtime/D3D12RHI/Private/D3D12RHICommon.h b/Engine/Source/Runtime/D3D12RHI/Private/D3D12RHICommon.h
--- a/Engine/Source/Runtime/D3D12RHI/Private/D3D12RHICommon.h
+++ b/Engine/Source/Runtime/D3D12RHI/Private/D3D12RHICommon.h
@@ -16,6 +16,12 @@
 #include "MultiGPU.h"
 #include "Stats/Stats2.h"
 
+class FD3D12Adapter;
+class FD3D12Device;
+
+template <typename ObjectType0, typename ObjectType1>
+class TD3D12DualLinkedObjectIterator;
+
 typedef uint16 CBVSlotMask;
 
 static_assert(MAX_ROOT_CBVS <= MAX_CBS, "MAX_ROOT_CBVS must be <= MAX_CBS.");
@@ -48,9 +54,6 @@ DECLARE_STATS_GROUP(TEXT("D3D12RHI: Buffer Details"), STATGROUP_D3D12BufferDetai
 DECLARE_STATS_GROUP(TEXT("D3D12RHI: Pipeline State (PSO)"), STATGROUP_D3D12PipelineState, STATCAT_Advanced);
 DECLARE_STATS_GROUP(TEXT("D3D12RHI: Descriptor Heap (GPU Visible)"), STATGROUP_D3D12DescriptorHeap, STATCAT_Advanced);
 
-class FD3D12Adapter;
-class FD3D12Device;
-
 class FD3D12AdapterChild
 {
 protected:
@@ -143,10 +146,6 @@ public:
 	}
 };
 
-
-template <typename ObjectType0, typename ObjectType1>
-class TD3D12DualLinkedObjectIterator;
-
 template <typename ObjectType>
 class FD3D12LinkedAdapterObject
 {
@@ -351,11 +350,22 @@ private:
 	ObjectType1* Object1;
 };
 
-namespace D3D12RHI
+// Template helper class for converting RHI resource types to D3D12RHI resource types
+template<class T>
+struct TD3D12ResourceTraits
+{
+};
+
+// Lock free pointer list that auto-destructs items remaining in the list.
+template <typename TObjectType>
+struct TD3D12ObjectPool : public TLockFreePointerListUnordered<TObjectType, PLATFORM_CACHE_LINE_SIZE>
 {
-	template <typename ObjectType0, typename ObjectType1>
-	TD3D12DualLinkedObjectIterator<ObjectType0, ObjectType1> MakeDualLinkedObjectIterator(ObjectType0* InObject0, ObjectType1* InObject1)
+	~TD3D12ObjectPool()
 	{
-		return TD3D12DualLinkedObjectIterator<ObjectType0, ObjectType1>(InObject0, InObject1);
+		while (TObjectType* Object = TLockFreePointerListUnordered<TObjectType, PLATFORM_CACHE_LINE_SIZE>::Pop())
+		{
+			delete Object;
+		}
 	}
-}
+};
+
diff --git a/Engine/Source/Runtime/D3D12RHI/Private/D3D12RHIPrivate.h b/Engine/Source/Runtime/D3D12RHI/Private/D3D12RHIPrivate.h
--- a/Engine/Source/Runtime/D3D12RHI/Private/D3D12RHIPrivate.h
+++ b/Engine/Source/Runtime/D3D12RHI/Private/D3D12RHIPrivate.h
@@ -30,9 +30,7 @@
 #include "ShaderCore.h"
 #include "HDRHelper.h"
 
-// TODO reorder includes so we just include D3D12Submission.h here
-#include COMPILED_PLATFORM_HEADER(D3D12Submission.h)
-
+#include "D3D12Submission.h"
 #include "D3D12RHIDefinitions.h"
 
 // TODO reorder includes so we just include D3D12PipelineState.h here
diff --git a/Engine/Source/Runtime/D3D12RHI/Private/D3D12Resources.h b/Engine/Source/Runtime/D3D12RHI/Private/D3D12Resources.h
--- a/Engine/Source/Runtime/D3D12RHI/Private/D3D12Resources.h
+++ b/Engine/Source/Runtime/D3D12RHI/Private/D3D12Resources.h
@@ -6,7 +6,6 @@
 
 #pragma once
 
-#include "D3D12RHIPrivate.h"
 #include "BoundShaderStateCache.h"
 #include "D3D12NvidiaExtensions.h"
 #include "D3D12ShaderResources.h"
@@ -986,10 +985,6 @@ public:
 	TArray<FD3D12SyncPointRef, TInlineAllocator<MAX_NUM_GPUS>> SyncPoints;
 };
 
-template<class T>
-struct TD3D12ResourceTraits
-{
-};
 template<>
 struct TD3D12ResourceTraits<FRHIUniformBuffer>
 {
diff --git a/Engine/Source/Runtime/D3D12RHI/Private/D3D12Shader.h b/Engine/Source/Runtime/D3D12RHI/Private/D3D12Shader.h
--- a/Engine/Source/Runtime/D3D12RHI/Private/D3D12Shader.h
+++ b/Engine/Source/Runtime/D3D12RHI/Private/D3D12Shader.h
@@ -6,6 +6,13 @@
 
 #pragma once
 
+#include "D3D12RHICommon.h"
+#include "RHIResources.h"
+#include "ShaderCore.h"
+#include "Templates/UnrealTypeTraits.h"
+
+class FD3D12RootSignature;
+
 template <>
 struct TTypeTraits<D3D12_INPUT_ELEMENT_DESC> : public TTypeTraitsBase < D3D12_INPUT_ELEMENT_DESC >
 {
@@ -39,6 +46,25 @@ public:
 	virtual uint32 GetPrecachePSOHash() const final override { return HashNoStrides;  }
 };
 
+//==================================================================================================================================
+// FD3D12ShaderBytecode
+// Encapsulates D3D12 shader bytecode and creates a hash for the shader bytecode
+//==================================================================================================================================
+struct ShaderBytecodeHash
+{
+	uint64 Hash[2];
+
+	bool operator ==(const ShaderBytecodeHash& b) const
+	{
+		return (Hash[0] == b.Hash[0] && Hash[1] == b.Hash[1]);
+	}
+
+	bool operator !=(const ShaderBytecodeHash& b) const
+	{
+		return (Hash[0] != b.Hash[0] || Hash[1] != b.Hash[1]);
+	}
+};
+
 struct FD3D12ShaderData
 {
 	/** The shader's bytecode, with custom data in the last byte. */
diff --git a/Engine/Source/Runtime/D3D12RHI/Private/D3D12Submission.h b/Engine/Source/Runtime/D3D12RHI/Private/D3D12Submission.h
--- a/Engine/Source/Runtime/D3D12RHI/Private/D3D12Submission.h
+++ b/Engine/Source/Runtime/D3D12RHI/Private/D3D12Submission.h
@@ -4,9 +4,20 @@
 
 #include "Async/TaskGraphInterfaces.h"
 #include "D3D12RHICommon.h"
-#include "D3D12Queue.h"
 #include "Templates/RefCounting.h"
 
+enum class ED3D12QueueType;
+
+class FD3D12CommandAllocator;
+class FD3D12CommandList;
+class FD3D12DynamicRHI;
+class FD3D12QueryHeap;
+class FD3D12Queue;
+class FD3D12Timing;
+
+class FD3D12SyncPoint;
+using FD3D12SyncPointRef = TRefCountPtr<FD3D12SyncPoint>;
+
 enum class ED3D12SyncPointType
 {
 	// Sync points of this type do not include an FGraphEvent, so cannot
@@ -48,11 +59,10 @@ struct FD3D12ResolvedFence
 // Sync points are one-shot, meaning they represent a single timeline point, and are released after use, via ref-counting.
 // Use FD3D12SyncPoint::Create() to make a new sync point and hold a reference to it via a FD3D12SyncPointRef object.
 //
-typedef TRefCountPtr<class FD3D12SyncPoint> FD3D12SyncPointRef;
 class FD3D12SyncPoint final : public FThreadSafeRefCountedObject
 {
-	friend class FD3D12DynamicRHI;
-	friend class FD3D12Queue;
+	friend FD3D12DynamicRHI;
+	friend FD3D12Queue;
 
 	static TLockFreePointerListUnordered<void, PLATFORM_CACHE_LINE_SIZE> MemoryPool;
 
@@ -133,7 +143,7 @@ enum class ED3D12QueryType
 struct FD3D12QueryLocation
 {
 	// The heap in which the result is contained.
-	class FD3D12QueryHeap* Heap = nullptr;
+	FD3D12QueryHeap* Heap = nullptr;
 
 	// The index of the query within the heap.
 	uint32 Index = 0;
@@ -147,7 +157,7 @@ struct FD3D12QueryLocation
 	inline uint64 GetResult() const;
 
 	FD3D12QueryLocation() = default;
-	FD3D12QueryLocation(class FD3D12QueryHeap* Heap, uint32 Index, ED3D12QueryType Type, uint64* Target)
+	FD3D12QueryLocation(FD3D12QueryHeap* Heap, uint32 Index, ED3D12QueryType Type, uint64* Target)
 		: Heap	(Heap  )
 		, Index	(Index )
 		, Type	(Type  )
@@ -229,11 +239,11 @@ struct FD3D12PayloadBase
 	TArray<FManualFence> FencesToWait;
 
 	// Execute
-	TArray<class FD3D12CommandList*> CommandListsToExecute;
+	TArray<FD3D12CommandList*> CommandListsToExecute;
 
 	// Signal
 	TArray<FManualFence> FencesToSignal;
-	TOptional<class FD3D12Timing*> Timing;
+	TOptional<FD3D12Timing*> Timing;
 	TArray<FD3D12SyncPointRef> SyncPointsToSignal;
 	uint64 CompletionFenceValue = 0;
 	FGraphEventRef SubmissionEvent;
@@ -243,7 +253,7 @@ struct FD3D12PayloadBase
 	bool bAlwaysSignal = false;
 
 	// Cleanup
-	TArray<class FD3D12CommandAllocator*> AllocatorsToRelease;
+	TArray<FD3D12CommandAllocator*> AllocatorsToRelease;
 	TArray<FD3D12QueryLocation> TimestampQueries;
 	TArray<FD3D12QueryLocation> OcclusionQueries;
 	TArray<FD3D12QueryRange> QueryRanges;
@@ -254,5 +264,7 @@ struct FD3D12PayloadBase
 	virtual ~FD3D12PayloadBase();
 
 protected:
-	FD3D12PayloadBase(FD3D12Device* const Device, ED3D12QueueType const QueueType);
+	FD3D12PayloadBase(FD3D12Device* Device, ED3D12QueueType QueueType);
 };
+
+#include COMPILED_PLATFORM_HEADER(D3D12Submission.h)
diff --git a/Engine/Source/Runtime/D3D12RHI/Private/D3D12Util.h b/Engine/Source/Runtime/D3D12RHI/Private/D3D12Util.h
--- a/Engine/Source/Runtime/D3D12RHI/Private/D3D12Util.h
+++ b/Engine/Source/Runtime/D3D12RHI/Private/D3D12Util.h
@@ -303,25 +303,6 @@ private:
 	TArray<D3D12_RESOURCE_STATES, TInlineAllocator<4>> m_SubresourceState;
 };
 
-//==================================================================================================================================
-// FD3D12ShaderBytecode
-// Encapsulates D3D12 shader bytecode and creates a hash for the shader bytecode
-//==================================================================================================================================
-struct ShaderBytecodeHash
-{
-	uint64 Hash[2];
-
-	bool operator ==(const ShaderBytecodeHash &b) const
-	{
-		return (Hash[0] == b.Hash[0] && Hash[1] == b.Hash[1]);
-	}
-
-	bool operator !=(const ShaderBytecodeHash &b) const
-	{
-		return (Hash[0] != b.Hash[0] || Hash[1] != b.Hash[1]);
-	}
-};
-
 /**
  * The base class of threadsafe reference counted objects.
  */
diff --git a/Engine/Source/Runtime/D3D12RHI/Private/D3D12View.h b/Engine/Source/Runtime/D3D12RHI/Private/D3D12View.h
--- a/Engine/Source/Runtime/D3D12RHI/Private/D3D12View.h
+++ b/Engine/Source/Runtime/D3D12RHI/Private/D3D12View.h
@@ -2,7 +2,11 @@
 
 #pragma once
 
+#include "Containers/Array.h"
+#include "D3D12RHICommon.h"
 #include "D3D12Descriptors.h"
+#include "D3D12Resources.h"
+#include "RHIResources.h"
 
 class FD3D12Buffer;
 class FD3D12Texture;
