Fix breadcrumbs not showing up in lightweight DRED.

diff --git a/Engine/Source/Runtime/D3D12RHI/Private/D3D12Adapter.cpp b/Engine/Source/Runtime/D3D12RHI/Private/D3D12Adapter.cpp
--- a/Engine/Source/Runtime/D3D12RHI/Private/D3D12Adapter.cpp
+++ b/Engine/Source/Runtime/D3D12RHI/Private/D3D12Adapter.cpp
@@ -471,6 +471,10 @@ void FD3D12Adapter::CreateRootDevice(bool bWithDebug)
 
 					if (SUCCEEDED(hr))
 					{
+						DredSettings2->SetBreadcrumbContextEnablement(D3D12_DRED_ENABLEMENT_FORCED_ON);
+						bDREDContext = true;
+						UE_LOG(LogD3D12RHI, Log, TEXT("[DRED] Dred breadcrumb context enabled"));
+
 						// Turn on AutoBreadcrumbs and Page Fault reporting
 						DredSettings2->SetAutoBreadcrumbsEnablement(D3D12_DRED_ENABLEMENT_FORCED_ON);
 						DredSettings2->SetPageFaultEnablement(D3D12_DRED_ENABLEMENT_FORCED_ON);
