Add a cvar d3d12.ExtraDepthTransitions which defaults to off. This makes the additional depth transitions added in 24227380 optional. These additional transitions fix a validation error but delay async compute work, causing async DF shadow and AO dispatches to run serially.

diff --git a/Engine/Source/Runtime/D3D12RHI/Private/D3D12CommandList.cpp b/Engine/Source/Runtime/D3D12RHI/Private/D3D12CommandList.cpp
--- a/Engine/Source/Runtime/D3D12RHI/Private/D3D12CommandList.cpp
+++ b/Engine/Source/Runtime/D3D12RHI/Private/D3D12CommandList.cpp
@@ -10,6 +10,13 @@ static FAutoConsoleVariableRef CVarD3D12BatchResourceBarriers(
 	GD3D12BatchResourceBarriers,
 	TEXT("Whether to allow batching resource barriers"));
 
+static int32 GD3D12ExtraDepthTransitions = 0;
+static FAutoConsoleVariableRef CVarD3D12ExtraDepthTransitions(
+	TEXT("d3d12.ExtraDepthTransitions"),
+	GD3D12ExtraDepthTransitions,
+	TEXT("Adds extra transitions for the depth buffer to fix validation issues. However, this currently breaks async compute"));
+
+
 int64 FD3D12CommandList::FState::NextCommandListID = 0;
 
 void FD3D12CommandList::UpdateResidency(TConstArrayView<FD3D12ResidencyHandle*> Handles)
@@ -390,7 +397,7 @@ void FD3D12ContextCommon::TransitionResource(FD3D12DepthStencilView* View)
 	{
 		TransitionResource(Resource, D3D12_RESOURCE_STATE_TBD, D3D12_RESOURCE_STATE_DEPTH_WRITE, View->GetDepthOnlySubset());
 	}
-	else if (bHasDepth)
+	else if (bHasDepth && GD3D12ExtraDepthTransitions)
 	{
 		TransitionResource(Resource, D3D12_RESOURCE_STATE_TBD, D3D12_RESOURCE_STATE_DEPTH_READ, View->GetDepthOnlySubset());
 	}
@@ -399,7 +406,7 @@ void FD3D12ContextCommon::TransitionResource(FD3D12DepthStencilView* View)
 	{
 		TransitionResource(Resource, D3D12_RESOURCE_STATE_TBD, D3D12_RESOURCE_STATE_DEPTH_WRITE, View->GetStencilOnlySubset());
 	}
-	else if (bHasStencil)
+	else if (bHasStencil && GD3D12ExtraDepthTransitions)
 	{
 		TransitionResource(Resource, D3D12_RESOURCE_STATE_TBD, D3D12_RESOURCE_STATE_DEPTH_READ, View->GetStencilOnlySubset());
 	}
@@ -627,7 +634,10 @@ static inline bool IsTransitionNeeded(D3D12_RESOURCE_STATES Before, D3D12_RESOUR
 	// Depth write is actually a suitable for read operations as a "normal" depth buffer.
 	if (Before == D3D12_RESOURCE_STATE_DEPTH_WRITE && After == D3D12_RESOURCE_STATE_DEPTH_READ)
 	{
-		After = D3D12_RESOURCE_STATE_DEPTH_WRITE;
+		if (GD3D12ExtraDepthTransitions)
+		{
+			After = D3D12_RESOURCE_STATE_DEPTH_WRITE;
+		}
 		return false;
 	}
 
