Add fields to the crash report indicating whether DRED has collected any useful data.

diff --git a/Engine/Source/Runtime/D3D12RHI/Private/D3D12Util.cpp b/Engine/Source/Runtime/D3D12RHI/Private/D3D12Util.cpp
--- a/Engine/Source/Runtime/D3D12RHI/Private/D3D12Util.cpp
+++ b/Engine/Source/Runtime/D3D12RHI/Private/D3D12Util.cpp
@@ -442,11 +442,13 @@ static bool LogDREDData(ID3D12Device* Device, bool bTrackingAllAllocations, D3D1
 	};
 	static_assert(UE_ARRAY_COUNT(AllocTypesNames) == D3D12_DRED_ALLOCATION_TYPE_VIDEO_EXTENSION_COMMAND - D3D12_DRED_ALLOCATION_TYPE_COMMAND_QUEUE + 1, "AllocTypes array length mismatch");
 
+	bool bHasValidBreadcrumbData = false;
 	FDred_T Dred(Device);
 	if (Dred.Data.IsValid())
 	{
 		if (Dred.BreadcrumbHead)
 		{
+			bHasValidBreadcrumbData = true;
 			UE_LOG(LogD3D12RHI, Error, TEXT("DRED: Last tracked GPU operations:"));
 
 			FString ContextStr;
@@ -506,9 +508,13 @@ static bool LogDREDData(ID3D12Device* Device, bool bTrackingAllAllocations, D3D1
 			UE_LOG(LogD3D12RHI, Error, TEXT("DRED: No breadcrumb head found."));
 		}
 
+		FGenericCrashContext::SetEngineData(TEXT("RHI.DRED.HasBreadcrumbData"), bHasValidBreadcrumbData ? TEXT("true") : TEXT("false"));
+
+		bool bHasValidPageFaultData = false;
 		D3D12_DRED_PAGE_FAULT_OUTPUT DredPageFaultOutput;
 		if (SUCCEEDED(Dred.Data->GetPageFaultAllocationOutput(&DredPageFaultOutput)) && DredPageFaultOutput.PageFaultVA != 0)
 		{
+			bHasValidPageFaultData = true;
 			OutPageFaultGPUAddress = DredPageFaultOutput.PageFaultVA;
 			UE_LOG(LogD3D12RHI, Error, TEXT("DRED: PageFault at VA GPUAddress \"0x%llX\""), (long long)DredPageFaultOutput.PageFaultVA);
 			
@@ -553,6 +559,8 @@ static bool LogDREDData(ID3D12Device* Device, bool bTrackingAllAllocations, D3D1
 			UE_LOG(LogD3D12RHI, Error, TEXT("DRED: No PageFault data."));
 		}
 
+		FGenericCrashContext::SetEngineData(TEXT("RHI.DRED.HasPageFaultData"), bHasValidPageFaultData ? TEXT("true") : TEXT("false"));
+
 		return true;
 	}
 	else
