Cleaning up more platform code.

diff --git a/Engine/Source/Developer/Windows/ShaderFormatD3D/Private/ShaderFormatD3D.cpp b/Engine/Source/Developer/Windows/ShaderFormatD3D/Private/ShaderFormatD3D.cpp
--- a/Engine/Source/Developer/Windows/ShaderFormatD3D/Private/ShaderFormatD3D.cpp
+++ b/Engine/Source/Developer/Windows/ShaderFormatD3D/Private/ShaderFormatD3D.cpp
@@ -11,7 +11,6 @@
 static FName NAME_PCD3D_SM6(TEXT("PCD3D_SM6"));
 static FName NAME_PCD3D_SM5(TEXT("PCD3D_SM5"));
 static FName NAME_PCD3D_ES3_1(TEXT("PCD3D_ES31"));
-static FName NAME_D3D_ES3_1_HOLOLENS(TEXT("D3D_ES3_1_HOLOLENS"));
 
 class FShaderFormatD3D : public UE::ShaderCompilerCommon::FBaseShaderFormat 
 {
@@ -23,7 +22,6 @@ class FShaderFormatD3D : public UE::ShaderCompilerCommon::FBaseShaderFormat
 		UE_SHADER_PCD3D_SM6_VER = 7,
 		UE_SHADER_PCD3D_SM5_VER = 12,
 		UE_SHADER_PCD3D_ES3_1_VER = 8,
-		UE_SHADER_D3D_ES3_1_HOLOLENS_VER = UE_SHADER_PCD3D_ES3_1_VER,
 	};
 
 	uint32 DxcVersionHash = 0;
@@ -84,11 +82,6 @@ public:
 			// Shader DXC signature is intentionally not included, as ES3_1 target always uses legacy compiler.
 			return GetVersionHash(UE_SHADER_PCD3D_ES3_1_VER);
 		}
-		else if (Format == NAME_D3D_ES3_1_HOLOLENS)
-		{
-			// Shader DXC signature is intentionally not included, as ES3_1 target always uses legacy compiler.
-			return GetVersionHash(UE_SHADER_D3D_ES3_1_HOLOLENS_VER);
-		}
 		checkf(0, TEXT("Unknown Format %s"), *Format.ToString());
 		return 0;
 	}
@@ -98,7 +91,6 @@ public:
 		OutFormats.Add(NAME_PCD3D_SM6);
 		OutFormats.Add(NAME_PCD3D_SM5);
 		OutFormats.Add(NAME_PCD3D_ES3_1);
-		OutFormats.Add(NAME_D3D_ES3_1_HOLOLENS);
 	}
 
 	ELanguage FormatToLanguage(FName Format) const
@@ -111,7 +103,7 @@ public:
 		{
 			return ELanguage::SM5;
 		}
-		else if (Format == NAME_PCD3D_ES3_1 || Format == NAME_D3D_ES3_1_HOLOLENS)
+		else if (Format == NAME_PCD3D_ES3_1)
 		{
 			return ELanguage::ES3_1;
 		}
@@ -189,13 +181,6 @@ public:
 			Input.Environment.SetDefine(TEXT("__SHADER_TARGET_MAJOR"), 5);
 			Input.Environment.SetDefine(TEXT("__SHADER_TARGET_MINOR"), 0);
 		}
-		else if (Input.ShaderFormat == NAME_D3D_ES3_1_HOLOLENS)
-		{
-			Input.Environment.SetDefine(TEXT("ES3_1_PROFILE"), 1);
-			Input.Environment.SetDefine(TEXT("COMPILER_DXC"), 0);
-			Input.Environment.SetDefine(TEXT("__SHADER_TARGET_MAJOR"), 5);
-			Input.Environment.SetDefine(TEXT("__SHADER_TARGET_MINOR"), 0);
-		}
 		else
 		{
 			checkf(0, TEXT("Unknown format %s"), *Input.ShaderFormat.ToString());
diff --git a/Engine/Source/Editor/DetailCustomizations/Private/MacTargetSettingsDetails.cpp b/Engine/Source/Editor/DetailCustomizations/Private/MacTargetSettingsDetails.cpp
--- a/Engine/Source/Editor/DetailCustomizations/Private/MacTargetSettingsDetails.cpp
+++ b/Engine/Source/Editor/DetailCustomizations/Private/MacTargetSettingsDetails.cpp
@@ -154,9 +154,6 @@ static FText GetFriendlyNameFromRHINameMac(FName InRHIName)
 	case SP_VULKAN_ES3_1_ANDROID:
 		FriendlyRHIName = LOCTEXT("VulkanMobile", "Vulkan (Mobile)");
 		break;
-	case SP_D3D_ES3_1_HOLOLENS:
-		FriendlyRHIName = LOCTEXT("D3DMobileHL", "Direct3D (Mobile, Hololens)");
-		break;
 	default:
 		FriendlyRHIName = FText::FromString(InRHIName.ToString());
 		break;
diff --git a/Engine/Source/Editor/DetailCustomizations/Private/WindowsTargetSettingsDetails.cpp b/Engine/Source/Editor/DetailCustomizations/Private/WindowsTargetSettingsDetails.cpp
--- a/Engine/Source/Editor/DetailCustomizations/Private/WindowsTargetSettingsDetails.cpp
+++ b/Engine/Source/Editor/DetailCustomizations/Private/WindowsTargetSettingsDetails.cpp
@@ -67,9 +67,6 @@ static FText GetFriendlyNameFromWindowsShaderPlatform(FName InShaderPlatformName
 	case SP_PCD3D_ES3_1:
 		FriendlyRHIName = LOCTEXT("DirectXMobile", "DirectX Mobile Emulation (Mobile)");
 		break;
-	case SP_D3D_ES3_1_HOLOLENS:
-		FriendlyRHIName = LOCTEXT("DirectXMobileHL", "DirectX Hololens (Mobile)");
-		break;
 	case SP_VULKAN_SM5:
 		FriendlyRHIName = LOCTEXT("VulkanSM5", "Vulkan (SM5)");
 		break;
diff --git a/Engine/Source/Editor/MaterialEditor/Private/MaterialStats.cpp b/Engine/Source/Editor/MaterialEditor/Private/MaterialStats.cpp
--- a/Engine/Source/Editor/MaterialEditor/Private/MaterialStats.cpp
+++ b/Engine/Source/Editor/MaterialEditor/Private/MaterialStats.cpp
@@ -472,9 +472,6 @@ void FMaterialStats::BuildShaderPlatformDB()
 	AddShaderPlatform(EPlatformCategoryType::Desktop, SP_PCD3D_SM5, TEXT("DirectX SM5"), true, true, TEXT("Desktop, DirectX, Shader Model 5"));
 	AddShaderPlatform(EPlatformCategoryType::Desktop, SP_PCD3D_SM6, TEXT("DirectX SM6"), true, true, TEXT("Desktop, DirectX, Shader Model 6"));
 	AddShaderPlatform(EPlatformCategoryType::Desktop, SP_PCD3D_ES3_1, TEXT("DirectX  Mobile"), true, true, TEXT("Desktop, DirectX, Mobile"));
-
-	// Desktop is the closest fit for Hololens (for now)
-	AddShaderPlatform(EPlatformCategoryType::Desktop, SP_D3D_ES3_1_HOLOLENS, TEXT("DirectX Mobile Hololens"), true, true, TEXT("Hololens, DirectX, Mobile"));
 #endif
 
 	// Vulkan
diff --git a/Engine/Source/Runtime/D3D12RHI/Private/D3D12Adapter.cpp b/Engine/Source/Runtime/D3D12RHI/Private/D3D12Adapter.cpp
--- a/Engine/Source/Runtime/D3D12RHI/Private/D3D12Adapter.cpp
+++ b/Engine/Source/Runtime/D3D12RHI/Private/D3D12Adapter.cpp
@@ -41,7 +41,7 @@ static FAutoConsoleVariableRef CTrackedReleasedAllocationFrameRetention(
 );
 #endif
 
-#if PLATFORM_WINDOWS || PLATFORM_HOLOLENS
+#if PLATFORM_WINDOWS
 
 #if UE_BUILD_SHIPPING || UE_BUILD_TEST
 static int32 GD3D12EnableGPUBreadCrumbs = 0;
@@ -122,7 +122,7 @@ static FAutoConsoleVariableRef CVarD3D12EvictAllResidentResourcesInBackground(
 	TEXT("Force D3D12 resource residency manager to evict all tracked unused resources when the application is not focused\n"),
 	ECVF_Default);
 
-#endif // PLATFORM_WINDOWS || PLATFORM_HOLOLENS
+#endif // PLATFORM_WINDOWS
 
 #if D3D12_SUPPORTS_INFO_QUEUE
 static bool CheckD3DStoredMessages()
@@ -375,7 +375,7 @@ void FD3D12Adapter::CreateRootDevice(bool bWithDebug)
 	UE_LOG(LogD3D12RHI, Log, TEXT("InitD3DDevice: -D3DDebug = %s -D3D12GPUValidation = %s"), bWithDebug ? TEXT("on") : TEXT("off"), bWithGPUValidation ? TEXT("on") : TEXT("off"));
 #endif
 
-#if PLATFORM_WINDOWS || (PLATFORM_HOLOLENS && !UE_BUILD_SHIPPING && WITH_PIX_EVENT_RUNTIME)
+#if PLATFORM_WINDOWS
 	
     SetupGPUCrashDebuggingModesCommon();
 
@@ -500,7 +500,7 @@ void FD3D12Adapter::CreateRootDevice(bool bWithDebug)
 	FGenericCrashContext::SetEngineData(TEXT("RHI.DREDMarkersOnly"), bDREDMarkersOnly ? TEXT("true") : TEXT("false"));
 	FGenericCrashContext::SetEngineData(TEXT("RHI.DREDContext"), bDREDContext && bDRED ? TEXT("true") : TEXT("false"));
 
-#endif // PLATFORM_WINDOWS || (PLATFORM_HOLOLENS && !UE_BUILD_SHIPPING && WITH_PIX_EVENT_RUNTIME)
+#endif // PLATFORM_WINDOWS
 
 #if USE_PIX
 	UE_LOG(LogD3D12RHI, Log, TEXT("Emitting draw events for PIX profiling."));
@@ -704,7 +704,7 @@ void FD3D12Adapter::CreateRootDevice(bool bWithDebug)
 			}
 		}
 	}
-#endif //  (PLATFORM_WINDOWS || PLATFORM_HOLOLENS)
+#endif // D3D12_SUPPORTS_DXGI_DEBUG
 
 #if UE_BUILD_DEBUG	&& D3D12_SUPPORTS_INFO_QUEUE
 	//break on debug
@@ -1080,7 +1080,7 @@ PRAGMA_ENABLE_DEPRECATION_WARNINGS
 			}
 
 #if D3D12_RHI_RAYTRACING
-#if PLATFORM_WINDOWS || PLATFORM_HOLOLENS
+#if PLATFORM_WINDOWS
 			D3D12_FEATURE_DATA_D3D12_OPTIONS5 D3D12Caps5 = {};
 			if (SUCCEEDED(RootDevice->CheckFeatureSupport(D3D12_FEATURE_D3D12_OPTIONS5, &D3D12Caps5, sizeof(D3D12Caps5))))
 			{
@@ -1130,7 +1130,7 @@ PRAGMA_ENABLE_DEPRECATION_WARNINGS
 					UE_LOG(LogD3D12RHI, Warning, TEXT("Ray Tracing is disabled because the RenderDoc plugin is currently not compatible with D3D12 ray tracing."));
 				}
 			}
-#endif // PLATFORM_WINDOWS || PLATFORM_HOLOLENS
+#endif // PLATFORM_WINDOWS
 
 			GRHIRayTracingAccelerationStructureAlignment = uint32(D3D12_RAYTRACING_ACCELERATION_STRUCTURE_BYTE_ALIGNMENT);
 			GRHIRayTracingScratchBufferAlignment = uint32(D3D12_RAYTRACING_ACCELERATION_STRUCTURE_BYTE_ALIGNMENT);
@@ -1193,7 +1193,7 @@ PRAGMA_ENABLE_DEPRECATION_WARNINGS
 #endif // PLATFORM_WINDOWS
 		}
 
-#if PLATFORM_WINDOWS || PLATFORM_HOLOLENS
+#if PLATFORM_WINDOWS
 		D3D12_FEATURE_DATA_D3D12_OPTIONS2 D3D12Caps2 = {};
 		if (FAILED(RootDevice->CheckFeatureSupport(D3D12_FEATURE_D3D12_OPTIONS2, &D3D12Caps2, sizeof(D3D12Caps2))))
 		{
@@ -1529,37 +1529,32 @@ FD3D12Adapter::~FD3D12Adapter()
 
 void FD3D12Adapter::CreateDXGIFactory(TRefCountPtr<IDXGIFactory2>& DxgiFactory2, bool bWithDebug, HMODULE DxgiDllHandle)
 {
-#if PLATFORM_WINDOWS || PLATFORM_HOLOLENS
+#if PLATFORM_WINDOWS
 	typedef HRESULT(WINAPI FCreateDXGIFactory2)(UINT, REFIID, void**);
 
-#if PLATFORM_WINDOWS
 	FCreateDXGIFactory2* CreateDXGIFactory2FnPtr = (FCreateDXGIFactory2*)(void*)::GetProcAddress(DxgiDllHandle, "CreateDXGIFactory2");
-#else
-	FCreateDXGIFactory2* CreateDXGIFactory2FnPtr = &CreateDXGIFactory2;
-#endif
 
 	check(CreateDXGIFactory2FnPtr);
 
 	const uint32 Flags = bWithDebug ? DXGI_CREATE_FACTORY_DEBUG : 0;
 	VERIFYD3D12RESULT(CreateDXGIFactory2FnPtr(Flags, IID_PPV_ARGS(DxgiFactory2.GetInitReference())));
-#endif // #if PLATFORM_WINDOWS || PLATFORM_HOLOLENS
+#endif // PLATFORM_WINDOWS
 }
 
 void FD3D12Adapter::CreateDXGIFactory(bool bWithDebug)
 {
-#if PLATFORM_WINDOWS || PLATFORM_HOLOLENS
-	HMODULE UsedDxgiDllHandle = (HMODULE)0;
 #if PLATFORM_WINDOWS
+	HMODULE UsedDxgiDllHandle = (HMODULE)0;
+
 	// Dynamically load this otherwise Win7 fails to boot as it's missing on that DLL
 	DxgiDllHandle = (HMODULE)FPlatformProcess::GetDllHandle(TEXT("dxgi.dll"));
 	check(DxgiDllHandle);
 	UsedDxgiDllHandle = DxgiDllHandle;
-#endif
 
 	CreateDXGIFactory(DxgiFactory2, bWithDebug, UsedDxgiDllHandle);
 
 	InitDXGIFactoryVariants(DxgiFactory2);
-#endif // #if PLATFORM_WINDOWS || PLATFORM_HOLOLENS
+#endif // PLATFORM_WINDOWS
 }
 
 void FD3D12Adapter::InitDXGIFactoryVariants(IDXGIFactory2* InDxgiFactory2)
@@ -1648,7 +1643,7 @@ void FD3D12Adapter::ReleaseTransientUniformBufferAllocator(FTransientUniformBuff
 
 void FD3D12Adapter::UpdateMemoryInfo()
 {
-#if PLATFORM_WINDOWS || PLATFORM_HOLOLENS
+#if PLATFORM_WINDOWS
 	const uint64 UpdateFrame = FrameFence != nullptr ? FrameFence->GetNextFenceToSignal() : 0;
 
 	// Avoid spurious query calls if we have already captured this frame.
diff --git a/Engine/Source/Runtime/D3D12RHI/Private/D3D12Device.cpp b/Engine/Source/Runtime/D3D12RHI/Private/D3D12Device.cpp
--- a/Engine/Source/Runtime/D3D12RHI/Private/D3D12Device.cpp
+++ b/Engine/Source/Runtime/D3D12RHI/Private/D3D12Device.cpp
@@ -185,7 +185,7 @@ ID3D12Device7* FD3D12Device::GetDevice7()
 }
 #endif // D3D12_RHI_RAYTRACING
 
-#if (PLATFORM_WINDOWS || PLATFORM_HOLOLENS)
+#if D3D12_SUPPORTS_DXGI_DEBUG
 typedef HRESULT(WINAPI *FDXGIGetDebugInterface1)(UINT, REFIID, void **);
 #endif
 
@@ -277,7 +277,7 @@ void FD3D12Device::SetupAfterDeviceCreation()
 	GRHISupportsArrayIndexFromAnyShader = true;
 	GRHISupportsStencilRefFromPixelShader = false; // TODO: Sort out DXC shader database SM6.0 usage. DX12 supports this feature, but need to improve DXC support.
 
-#if (PLATFORM_WINDOWS || PLATFORM_HOLOLENS)
+#if PLATFORM_WINDOWS
 	// Check if we're running under GPU capture
 	bool bUnderGPUCapture = false;
 
@@ -306,19 +306,13 @@ void FD3D12Device::SetupAfterDeviceCreation()
 		// Running on AMD with RGP profiling enabled, so enable capturing mode
 		bUnderGPUCapture = true;
 	}
-#if USE_PIX
 
-	// Only check windows version on PLATFORM_WINDOWS - Hololens can assume windows > 10.0 so the condition would always be true.
-#if PLATFORM_WINDOWS
+#if USE_PIX
 	// PIX (note that DXGIGetDebugInterface1 requires Windows 8.1 and up)
 	if (FPlatformMisc::VerifyWindowsVersion(6, 3))
-#endif
 	{
 		FDXGIGetDebugInterface1 DXGIGetDebugInterface1FnPtr = nullptr;
 
-#if PLATFORM_HOLOLENS
-		DXGIGetDebugInterface1FnPtr = DXGIGetDebugInterface1;
-#else
 		// CreateDXGIFactory2 is only available on Win8.1+, find it if it exists
 		HMODULE DxgiDLL = LoadLibraryA("dxgi.dll");
 		if (DxgiDLL)
@@ -329,7 +323,6 @@ void FD3D12Device::SetupAfterDeviceCreation()
 #pragma warning(pop)
 			FreeLibrary(DxgiDLL);
 		}
-#endif
 		
 		if (DXGIGetDebugInterface1FnPtr)
 		{
@@ -351,7 +344,7 @@ void FD3D12Device::SetupAfterDeviceCreation()
 	{
 		GDynamicRHI->EnableIdealGPUCaptureOptions(true);
 	}
-#endif // (PLATFORM_WINDOWS || PLATFORM_HOLOLENS)
+#endif // PLATFORM_WINDOWS
 
 
 	const int32 MaximumResourceHeapSize = GetParentAdapter()->GetMaxDescriptorsForHeapType(ERHIDescriptorHeapType::Standard);
diff --git a/Engine/Source/Runtime/D3D12RHI/Private/D3D12PipelineState.h b/Engine/Source/Runtime/D3D12RHI/Private/D3D12PipelineState.h
--- a/Engine/Source/Runtime/D3D12RHI/Private/D3D12PipelineState.h
+++ b/Engine/Source/Runtime/D3D12RHI/Private/D3D12PipelineState.h
@@ -249,7 +249,7 @@ template <> struct equality_pipeline_state_desc<FD3D12ComputePipelineStateDesc>
 	bool operator()(const FD3D12ComputePipelineStateDesc& lhs, const FD3D12ComputePipelineStateDesc& rhs)
 	{
 		PSO_IF_NOT_EQUAL_RETURN_FALSE(Desc.CS.BytecodeLength)
-#if PLATFORM_WINDOWS || PLATFORM_HOLOLENS
+#if PLATFORM_WINDOWS
 		PSO_IF_NOT_EQUAL_RETURN_FALSE(Desc.Flags)
 #endif
 		PSO_IF_NOT_EQUAL_RETURN_FALSE(Desc.pRootSignature)
diff --git a/Engine/Source/Runtime/D3D12RHI/Private/D3D12RHI.cpp b/Engine/Source/Runtime/D3D12RHI/Private/D3D12RHI.cpp
--- a/Engine/Source/Runtime/D3D12RHI/Private/D3D12RHI.cpp
+++ b/Engine/Source/Runtime/D3D12RHI/Private/D3D12RHI.cpp
@@ -65,7 +65,7 @@ FD3D12DynamicRHI::FD3D12DynamicRHI(const TArray<TSharedPtr<FD3D12Adapter>>& Chos
 	FeatureLevel = GetAdapter().GetFeatureLevel();
 	check(FeatureLevel >= D3D_FEATURE_LEVEL_11_0);
 
-#if PLATFORM_WINDOWS || PLATFORM_HOLOLENS
+#if PLATFORM_WINDOWS
 	// Allocate a buffer of zeroes. This is used when we need to pass D3D memory
 	// that we don't care about and will overwrite with valid data in the future.
 	ZeroBufferSize = FMath::Max(CVarD3D12ZeroBufferSizeInMB.GetValueOnAnyThread(), 0) * (1 << 20);
diff --git a/Engine/Source/Runtime/D3D12RHI/Private/D3D12Texture.cpp b/Engine/Source/Runtime/D3D12RHI/Private/D3D12Texture.cpp
--- a/Engine/Source/Runtime/D3D12RHI/Private/D3D12Texture.cpp
+++ b/Engine/Source/Runtime/D3D12RHI/Private/D3D12Texture.cpp
@@ -692,7 +692,7 @@ void FD3D12DynamicRHI::RHIGetTextureMemoryStats(FTextureMemoryStats& OutStats)
 	OutStats.TexturePoolSize = GTexturePoolSize;
 	OutStats.PendingMemoryAdjustment = 0;
 
-#if PLATFORM_WINDOWS || PLATFORM_HOLOLENS
+#if PLATFORM_WINDOWS
 	if (GAdjustTexturePoolSizeBasedOnBudget)
 	{
 		GetAdapter().UpdateMemoryInfo();
@@ -761,7 +761,7 @@ void FD3D12DynamicRHI::RHIGetTextureMemoryStats(FTextureMemoryStats& OutStats)
 
 	// Cache the last requested texture pool size.
 	RequestedTexturePoolSize = OutStats.TexturePoolSize;
-#endif // PLATFORM_WINDOWS || PLATFORM_HOLOLENS
+#endif // PLATFORM_WINDOWS
 }
 
 /**
@@ -885,8 +885,7 @@ FD3D12Texture* FD3D12DynamicRHI::CreateNewD3D12Texture(const FRHITextureCreateDe
 
 FD3D12Texture* FD3D12DynamicRHI::CreateD3D12Texture(const FRHITextureCreateDesc& InCreateDesc, class FRHICommandListImmediate* RHICmdList, ID3D12ResourceAllocator* ResourceAllocator)
 {
-#if PLATFORM_WINDOWS || PLATFORM_HOLOLENS
-	
+#if PLATFORM_WINDOWS
 	TRACE_CPUPROFILER_EVENT_SCOPE(D3D12RHI::CreateD3D12Texture);
 	UE_TRACE_METADATA_SCOPE_ASSET_FNAME(FName(InCreateDesc.DebugName), InCreateDesc.GetTraceClassName(), InCreateDesc.OwnerName);
 
@@ -996,7 +995,7 @@ FD3D12Texture* FD3D12DynamicRHI::CreateD3D12Texture(const FRHITextureCreateDesc&
 #else
 	checkf(false, TEXT("XBOX_CODE_MERGE : Removed. The Xbox platform version should be used."));
 	return nullptr;
-#endif // PLATFORM_WINDOWS || PLATFORM_HOLOLENS
+#endif // PLATFORM_WINDOWS
 }
 
 
diff --git a/Engine/Source/Runtime/D3D12RHI/Private/D3D12Util.cpp b/Engine/Source/Runtime/D3D12RHI/Private/D3D12Util.cpp
--- a/Engine/Source/Runtime/D3D12RHI/Private/D3D12Util.cpp
+++ b/Engine/Source/Runtime/D3D12RHI/Private/D3D12Util.cpp
@@ -16,7 +16,7 @@ D3D12Util.h: D3D RHI utility implementation.
 #define LOCTEXT_NAMESPACE "Developer.MessageLog"
 
 // GPU crashes are nonfatal on windows/nonshipping so as not to interfere with GPU crash dump processing
-#if PLATFORM_WINDOWS || PLATFORM_HOLOLENS || !UE_BUILD_SHIPPING
+#if PLATFORM_WINDOWS || !UE_BUILD_SHIPPING
   #define D3D12RHI_GPU_CRASH_LOG_VERBOSITY Error
 #else
   #define D3D12RHI_GPU_CRASH_LOG_VERBOSITY Fatal
diff --git a/Engine/Source/Runtime/D3D12RHI/Private/D3D12Viewport.cpp b/Engine/Source/Runtime/D3D12RHI/Private/D3D12Viewport.cpp
--- a/Engine/Source/Runtime/D3D12RHI/Private/D3D12Viewport.cpp
+++ b/Engine/Source/Runtime/D3D12RHI/Private/D3D12Viewport.cpp
@@ -320,14 +320,14 @@ FD3D12Viewport::~FD3D12Viewport()
 {
 	check(IsInRHIThread() || IsInRenderingThread());
 
-#if !PLATFORM_HOLOLENS && D3D12_VIEWPORT_EXPOSES_SWAP_CHAIN
+#if D3D12_VIEWPORT_EXPOSES_SWAP_CHAIN
 	// If the swap chain was in fullscreen mode, switch back to windowed before releasing the swap chain.
 	// DXGI throws an error otherwise.
 	if (SwapChain1)
 	{
 		SwapChain1->SetFullscreenState(0, nullptr);
 	}
-#endif // !PLATFORM_HOLOLENS && D3D12_VIEWPORT_EXPOSES_SWAP_CHAIN
+#endif // D3D12_VIEWPORT_EXPOSES_SWAP_CHAIN
 
 	GetParentAdapter()->GetViewports().Remove(this);
 
diff --git a/Engine/Source/Runtime/D3D12RHI/Private/D3D12Viewport.h b/Engine/Source/Runtime/D3D12RHI/Private/D3D12Viewport.h
--- a/Engine/Source/Runtime/D3D12RHI/Private/D3D12Viewport.h
+++ b/Engine/Source/Runtime/D3D12RHI/Private/D3D12Viewport.h
@@ -196,7 +196,7 @@ private:
 	TRefCountPtr<IDXGISwapChain4> SwapChain4;
 #endif
 
-#if PLATFORM_WINDOWS || PLATFORM_HOLOLENS
+#if PLATFORM_WINDOWS
 	DXGI_COLOR_SPACE_TYPE ColorSpace;
 #endif
 #endif // D3D12_VIEWPORT_EXPOSES_SWAP_CHAIN
@@ -262,7 +262,7 @@ private:
 	/** Disable HDR meta data transmission and set the necessary color space. */
 	void ShutdownHDR();
 
-#if PLATFORM_WINDOWS || PLATFORM_HOLOLENS
+#if PLATFORM_WINDOWS
 	/** Ensure the correct color space is set on the swap chain */
 	void EnsureColorSpace(EDisplayColorGamut DisplayGamut, EDisplayOutputFormat OutputDevice);
 #endif
diff --git a/Engine/Source/Runtime/D3D12RHI/Private/Windows/WindowsD3D12Device.cpp b/Engine/Source/Runtime/D3D12RHI/Private/Windows/WindowsD3D12Device.cpp
--- a/Engine/Source/Runtime/D3D12RHI/Private/Windows/WindowsD3D12Device.cpp
+++ b/Engine/Source/Runtime/D3D12RHI/Private/Windows/WindowsD3D12Device.cpp
@@ -159,11 +159,7 @@ static void SafeCreateDXGIFactory(IDXGIFactory4** DXGIFactory)
 	{
 		bIsQuadBufferStereoEnabled = FParse::Param(FCommandLine::Get(), TEXT("quad_buffer_stereo"));
 
-#if PLATFORM_HOLOLENS
-		CreateDXGIFactory1(__uuidof(IDXGIFactory4), (void**)DXGIFactory);
-#else
 		CreateDXGIFactory(__uuidof(IDXGIFactory4), (void**)DXGIFactory);
-#endif
 	}
 	__except (IsDelayLoadException(GetExceptionInformation()))
 	{
@@ -673,7 +669,6 @@ static bool CheckIfAgilitySDKLoaded()
 
 bool FD3D12DynamicRHIModule::IsSupported(ERHIFeatureLevel::Type RequestedFeatureLevel)
 {
-#if !PLATFORM_HOLOLENS
 	// Windows version 15063 is Windows 1703 aka "Windows Creator Update"
 	// This is the first version that supports ID3D12Device2 which is our minimum runtime device version.
 	if (!FPlatformMisc::VerifyWindowsVersion(10, 0, 15063))
@@ -681,7 +676,6 @@ bool FD3D12DynamicRHIModule::IsSupported(ERHIFeatureLevel::Type RequestedFeature
 		UE_LOG(LogD3D12RHI, Warning, TEXT("Missing full support for Direct3D 12. Update to Windows 1703 or newer for D3D12 support."));
 		return false;
 	}
-#endif
 
 	// If not computed yet
 	if (ChosenAdapters.Num() == 0)
@@ -901,14 +895,12 @@ void FD3D12DynamicRHIModule::FindAdapter()
 
 				const bool bIsWARP = (RHIConvertToGpuVendorId(AdapterDesc.VendorId) == EGpuVendorId::Microsoft);
 
-#if !PLATFORM_HOLOLENS
 				if (!bIsWARP)
 				{
 					const FGPUDriverInfo GPUDriverInfo = FPlatformMisc::GetGPUDriverInfo(AdapterDesc.Description, false);
 					UE_LOG(LogD3D12RHI, Log, TEXT("  Driver Version: %s (internal:%s, unified:%s)"), *GPUDriverInfo.UserDriverVersion, *GPUDriverInfo.InternalDriverVersion, *GPUDriverInfo.GetUnifiedDriverVersion());
 					UE_LOG(LogD3D12RHI, Log, TEXT("     Driver Date: %s"), *GPUDriverInfo.DriverDate);
 				}
-#endif // !PLATFORM_HOLOLENS
 
 				FD3D12AdapterDesc CurrentAdapter(AdapterDesc, AdapterIndex, DeviceInfo);
 
@@ -1027,21 +1019,12 @@ static bool DoesAnyAdapterSupportSM6(const TArray<TSharedPtr<FD3D12Adapter>>& Ad
 
 FDynamicRHI* FD3D12DynamicRHIModule::CreateRHI(ERHIFeatureLevel::Type RequestedFeatureLevel)
 {
-#if PLATFORM_HOLOLENS
-	check(RequestedFeatureLevel == ERHIFeatureLevel::ES3_1);
-
-	GMaxRHIFeatureLevel = ERHIFeatureLevel::ES3_1;
-	GMaxRHIShaderPlatform = SP_D3D_ES3_1_HOLOLENS;
-
-	GShaderPlatformForFeatureLevel[ERHIFeatureLevel::ES3_1] = SP_D3D_ES3_1_HOLOLENS;
-#else
 	GShaderPlatformForFeatureLevel[ERHIFeatureLevel::ES3_1] = SP_PCD3D_ES3_1;
 	GShaderPlatformForFeatureLevel[ERHIFeatureLevel::SM5] = SP_PCD3D_SM5;
 	if (DoesAnyAdapterSupportSM6(ChosenAdapters))
 	{
 		GShaderPlatformForFeatureLevel[ERHIFeatureLevel::SM6] = SP_PCD3D_SM6;
 	}
-#endif
 
 	ERHIFeatureLevel::Type PreviewFeatureLevel;
 	if (!GIsEditor && RHIGetPreviewFeatureLevel(PreviewFeatureLevel))
@@ -1277,7 +1260,6 @@ static void DisableRayTracingSupport()
 
 static void ClearPSODriverCache()
 {
-#if !PLATFORM_HOLOLENS
 	FString LocalAppDataFolder = FPlatformMisc::GetEnvironmentVariable(TEXT("LOCALAPPDATA"));
 	if (!LocalAppDataFolder.IsEmpty())
 	{
@@ -1316,7 +1298,6 @@ static void ClearPSODriverCache()
 	{
 		UE_LOG(LogD3D12RHI, Error, TEXT("clearPSODriverCache failed: please ensure that LOCALAPPDATA points to [REDACTED_USER_HOME]\\AppData\\Local"));
 	}
-#endif
 }
 
 void FD3D12DynamicRHI::Init()
@@ -1437,9 +1418,7 @@ void FD3D12DynamicRHI::Init()
 	}
 #endif
 
-#if !PLATFORM_HOLOLENS
 	// Disable ray tracing for Windows build versions
-
 	if (GRHISupportsRayTracing
 		&& GMinimumWindowsBuildVersionForRayTracing > 0
 		&& !FPlatformMisc::VerifyWindowsVersion(10, 0, GMinimumWindowsBuildVersionForRayTracing))
@@ -1447,7 +1426,6 @@ void FD3D12DynamicRHI::Init()
 		DisableRayTracingSupport();
 		UE_LOG(LogD3D12RHI, Warning, TEXT("Ray tracing is disabled because it requires Windows 10 version %u"), (uint32)GMinimumWindowsBuildVersionForRayTracing);
 	}
-#endif
 
 #if WITH_NVAPI
 	if (IsRHIDeviceNVIDIA() && bAllowVendorDevice)
diff --git a/Engine/Source/Runtime/D3D12RHI/Private/Windows/WindowsD3D12DiskCache.cpp b/Engine/Source/Runtime/D3D12RHI/Private/Windows/WindowsD3D12DiskCache.cpp
--- a/Engine/Source/Runtime/D3D12RHI/Private/Windows/WindowsD3D12DiskCache.cpp
+++ b/Engine/Source/Runtime/D3D12RHI/Private/Windows/WindowsD3D12DiskCache.cpp
@@ -93,11 +93,7 @@ void FDiskCacheInterface::GrowMapping(SIZE_T size, bool firstrun)
 
 	uint32 flag = (mCacheExists) ? OPEN_EXISTING : CREATE_NEW;
 	// open the shader cache file
-#if PLATFORM_HOLOLENS
-	mFile = CreateFile2(mFileName.GetCharArray().GetData(), GENERIC_READ | GENERIC_WRITE, FILE_SHARE_READ, flag, NULL);
-#else
 	mFile = CreateFile(mFileName.GetCharArray().GetData(), GENERIC_READ | GENERIC_WRITE, FILE_SHARE_READ, NULL, flag, FILE_ATTRIBUTE_NORMAL, NULL);
-#endif
 
 	if (mFile == INVALID_HANDLE_VALUE)
 	{
@@ -108,18 +104,8 @@ void FDiskCacheInterface::GrowMapping(SIZE_T size, bool firstrun)
 
 	mCacheExists = true;
 
-#if PLATFORM_HOLOLENS
-	LARGE_INTEGER largeFileSize;
-	if (!GetFileSizeEx(mFile, &largeFileSize))
-	{
-		//error state!
-		mInErrorState = true;
-		return;
-	}
-	uint32 fileSize = largeFileSize.LowPart;
-#else
 	uint32 fileSize = GetFileSize(mFile, NULL);
-#endif
+
 	if (fileSize == 0)
 	{
 		byte data[64];
diff --git a/Engine/Source/Runtime/D3D12RHI/Private/Windows/WindowsD3D12Viewport.cpp b/Engine/Source/Runtime/D3D12RHI/Private/Windows/WindowsD3D12Viewport.cpp
--- a/Engine/Source/Runtime/D3D12RHI/Private/Windows/WindowsD3D12Viewport.cpp
+++ b/Engine/Source/Runtime/D3D12RHI/Private/Windows/WindowsD3D12Viewport.cpp
@@ -114,36 +114,6 @@ void FD3D12Viewport::Init()
 	ID3D12CommandQueue* CommandQueue = Adapter->GetDevice(0)->GetQueue(ED3D12QueueType::Direct).D3DCommandQueue;
 
 	// Create the swapchain.
-#if PLATFORM_HOLOLENS
-	{
-		// MSAA Sample count
-		DXGI_SWAP_CHAIN_DESC1 SwapChainDesc{};
-		SwapChainDesc.SampleDesc.Count = 1;
-		SwapChainDesc.SampleDesc.Quality = 0;
-		SwapChainDesc.BufferUsage = DXGI_USAGE_RENDER_TARGET_OUTPUT | DXGI_USAGE_SHADER_INPUT;
-		SwapChainDesc.Width = BufferDesc.Width;
-		SwapChainDesc.Height = BufferDesc.Height;
-		SwapChainDesc.Format = BufferDesc.Format;
-		// 1:single buffering, 2:double buffering, 3:triple buffering
-		SwapChainDesc.BufferCount = NumBackBuffers;
-		SwapChainDesc.Scaling = DXGI_SCALING_NONE;
-		SwapChainDesc.AlphaMode = DXGI_ALPHA_MODE_IGNORE;
-		// DXGI_SWAP_EFFECT_DISCARD / DXGI_SWAP_EFFECT_SEQUENTIAL
-		SwapChainDesc.SwapEffect = DXGI_SWAP_EFFECT_FLIP_DISCARD;
-		SwapChainDesc.Flags = SwapChainFlags;
-
-		IDXGIFactory4* Factory4 = Adapter->GetDXGIFactory4();
-		check(Factory4);
-
-		VERIFYD3D12RESULT(Factory4->CreateSwapChainForCoreWindow(
-			CommandQueue,
-			reinterpret_cast<IUnknown*>(Windows::UI::Core::CoreWindow::GetForCurrentThread()),
-			&SwapChainDesc,
-			NULL,
-			SwapChain1.GetInitReference()
-		));
-	}
-#else
 
 	extern bool bNeedSwapChain;
 
@@ -225,7 +195,6 @@ void FD3D12Viewport::Init()
 			VERIFYD3D12RESULT(SwapChain->QueryInterface(IID_PPV_ARGS(SwapChain1.GetInitReference())));
 		}
 	}
-#endif
 
 	if (SwapChain1)
 	{
@@ -249,11 +218,9 @@ void FD3D12Viewport::Init()
 	// Resize to setup mGPU correctly.
 	Resize(BufferDesc.Width, BufferDesc.Height, bIsFullscreen, PixelFormat);
 
-#if !PLATFORM_HOLOLENS
 	// Tell the window to redraw when they can.
 	// @todo: For Slate viewports, it doesn't make sense to post WM_PAINT messages (we swallow those.)
 	::PostMessageW(WindowHandle, WM_PAINT, 0, 0);
-#endif
 }
 
 void FD3D12Viewport::FinalDestroyInternal()
@@ -276,13 +243,11 @@ void FD3D12Viewport::ConditionalResetSwapChain(bool bIgnoreFocus)
 		}
 		else
 		{
-#if !PLATFORM_HOLOLENS
 			// Check if the viewport's window is focused before resetting the swap chain's fullscreen state.
 			HWND FocusWindow = ::GetFocus();
 			const bool bIsFocused = FocusWindow == WindowHandle;
 			const bool bIsIconic = !!::IsIconic(WindowHandle);
 			if (bIgnoreFocus || (bIsFocused && !bIsIconic))
-#endif
 			{
 				FlushRenderingCommands();
 
diff --git a/Engine/Source/Runtime/D3D12RHI/Public/D3D12RHI.h b/Engine/Source/Runtime/D3D12RHI/Public/D3D12RHI.h
--- a/Engine/Source/Runtime/D3D12RHI/Public/D3D12RHI.h
+++ b/Engine/Source/Runtime/D3D12RHI/Public/D3D12RHI.h
@@ -19,7 +19,7 @@
 	#define WITH_PIX_EVENT_RUNTIME 0
 #endif
 
-#if PLATFORM_WINDOWS || PLATFORM_HOLOLENS
+#if PLATFORM_WINDOWS
 	#define ENABLE_RESIDENCY_MANAGEMENT				1
 	#define PIPELINE_STATE_FILE_LOCATION			FPaths::ProjectSavedDir()
 	#define USE_PIX									WITH_PIX_EVENT_RUNTIME
diff --git a/Engine/Source/Runtime/RHI/Private/RHIStrings.cpp b/Engine/Source/Runtime/RHI/Private/RHIStrings.cpp
--- a/Engine/Source/Runtime/RHI/Private/RHIStrings.cpp
+++ b/Engine/Source/Runtime/RHI/Private/RHIStrings.cpp
@@ -99,7 +99,6 @@ FString LexToString(EShaderPlatform Platform, bool bError)
 	case SP_VULKAN_SM5: return TEXT("VULKAN_SM5");
 	case SP_VULKAN_SM6: return TEXT("VULKAN_SM6");
 	case SP_VULKAN_SM5_ANDROID: return TEXT("VULKAN_SM5_ANDROID");
-	case SP_D3D_ES3_1_HOLOLENS: return TEXT("D3D_ES3_1_HOLOLENS");
 	case SP_CUSTOM_PLATFORM_FIRST: return TEXT("SP_CUSTOM_PLATFORM_FIRST");
 	case SP_CUSTOM_PLATFORM_LAST: return TEXT("SP_CUSTOM_PLATFORM_LAST");
 
@@ -204,7 +203,6 @@ FName ShaderPlatformToPlatformName(EShaderPlatform Platform)
 	case SP_VULKAN_PCES3_1:
 	case SP_VULKAN_SM5:
 	case SP_VULKAN_SM6:
-	case SP_D3D_ES3_1_HOLOLENS:
 		return NAME_PLATFORM_WINDOWS;
 	case SP_VULKAN_ES3_1_ANDROID:
 	case SP_VULKAN_SM5_ANDROID:
diff --git a/Engine/Source/Runtime/RHI/Public/RHIShaderFormatDefinitions.inl b/Engine/Source/Runtime/RHI/Public/RHIShaderFormatDefinitions.inl
--- a/Engine/Source/Runtime/RHI/Public/RHIShaderFormatDefinitions.inl
+++ b/Engine/Source/Runtime/RHI/Public/RHIShaderFormatDefinitions.inl
@@ -30,8 +30,6 @@ static FName NAME_VULKAN_SM5(TEXT("SF_VULKAN_SM5"));
 static FName NAME_VULKAN_SM6(TEXT("SF_VULKAN_SM6"));
 static FName NAME_VULKAN_SM5_ANDROID(TEXT("SF_VULKAN_SM5_ANDROID"));
 
-static FName NAME_D3D_ES3_1_HOLOLENS(TEXT("D3D_ES3_1_HOLOLENS"));
-
 static EShaderPlatform ShaderFormatNameToShaderPlatform(FName ShaderFormat)
 {
 	if (ShaderFormat == NAME_PCD3D_SM6)					return SP_PCD3D_SM6;
@@ -55,8 +53,6 @@ static EShaderPlatform ShaderFormatNameToShaderPlatform(FName ShaderFormat)
 	if (ShaderFormat == NAME_VULKAN_SM6)				return SP_VULKAN_SM6;
 	if (ShaderFormat == NAME_VULKAN_SM5_ANDROID)		return SP_VULKAN_SM5_ANDROID;
 
-	if (ShaderFormat == NAME_D3D_ES3_1_HOLOLENS)		return SP_D3D_ES3_1_HOLOLENS;
-
 	for (int32 StaticPlatform = SP_StaticPlatform_First; StaticPlatform <= SP_StaticPlatform_Last; ++StaticPlatform)
 	{
 		if (ShaderFormat == FStaticShaderPlatformNames::Get().GetShaderFormat(EShaderPlatform(StaticPlatform)))
diff --git a/Engine/Source/Runtime/RHI/Public/RHIShaderPlatform.h b/Engine/Source/Runtime/RHI/Public/RHIShaderPlatform.h
--- a/Engine/Source/Runtime/RHI/Public/RHIShaderPlatform.h
+++ b/Engine/Source/Runtime/RHI/Public/RHIShaderPlatform.h
@@ -42,9 +42,8 @@ enum EShaderPlatform : uint16
 	//  Add new platforms below this line, starting from (SP_StaticPlatform_Last + 1)
 	//---------------------------------------------------------------------------------
 	SP_VULKAN_SM5_ANDROID = SP_StaticPlatform_Last + 1,
-	SP_PCD3D_SM6,
-	SP_D3D_ES3_1_HOLOLENS,
-	SP_VULKAN_SM6,
+	SP_PCD3D_SM6          = SP_StaticPlatform_Last + 2,
+	SP_VULKAN_SM6         = SP_StaticPlatform_Last + 4,
 
 	SP_CUSTOM_PLATFORM_FIRST,
 	SP_CUSTOM_PLATFORM_LAST = (SP_CUSTOM_PLATFORM_FIRST + 100),
diff --git a/Engine/Source/Runtime/Renderer/Private/PostProcess/PostProcessMaterial.cpp b/Engine/Source/Runtime/Renderer/Private/PostProcess/PostProcessMaterial.cpp
--- a/Engine/Source/Runtime/Renderer/Private/PostProcess/PostProcessMaterial.cpp
+++ b/Engine/Source/Runtime/Renderer/Private/PostProcess/PostProcessMaterial.cpp
@@ -485,7 +485,7 @@ FScreenPassTexture AddPostProcessMaterialPass(
 	// We can re-use the scene color texture as the render target if we're not simultaneously reading from it.
 	// This is only necessary to do if we're going to be priming content from the render target since it avoids
 	// the copy. Otherwise, we just allocate a new render target.
-	const bool bValidShaderPlatform = (GMaxRHIShaderPlatform != SP_PCD3D_ES3_1) && (GMaxRHIShaderPlatform != SP_D3D_ES3_1_HOLOLENS); // This might actually work with Hololens GPUs, but we haven't enabled it before
+	const bool bValidShaderPlatform = (GMaxRHIShaderPlatform != SP_PCD3D_ES3_1);
 	if (!Output.IsValid() && !MaterialShaderMap->UsesSceneTexture(PPI_PostProcessInput0) && bPrimeOutputColor && !bForceIntermediateTarget && Inputs.bAllowSceneColorInputAsOutput && bValidShaderPlatform)
 	{
 		Output = FScreenPassRenderTarget(SceneColor, ERenderTargetLoadAction::ELoad);
diff --git a/Engine/Source/Runtime/Windows/D3D11RHI/Private/D3D11Commands.cpp b/Engine/Source/Runtime/Windows/D3D11RHI/Private/D3D11Commands.cpp
--- a/Engine/Source/Runtime/Windows/D3D11RHI/Private/D3D11Commands.cpp
+++ b/Engine/Source/Runtime/Windows/D3D11RHI/Private/D3D11Commands.cpp
@@ -1572,7 +1572,6 @@ void FD3D11DynamicRHI::EnableUAVOverlap()
 
 	bUAVOverlapEnabled = true;
 
-#if !PLATFORM_HOLOLENS
 	if (IsRHIDeviceNVIDIA())
 	{
 #if WITH_NVAPI
@@ -1594,7 +1593,6 @@ void FD3D11DynamicRHI::EnableUAVOverlap()
 		}
 #endif
 	}
-#endif
 }
 
 void FD3D11DynamicRHI::DisableUAVOverlap()
@@ -1606,7 +1604,6 @@ void FD3D11DynamicRHI::DisableUAVOverlap()
 		return;
 	}
 
-#if !PLATFORM_HOLOLENS
 	if (IsRHIDeviceNVIDIA())
 	{
 #if WITH_NVAPI
@@ -1628,7 +1625,6 @@ void FD3D11DynamicRHI::DisableUAVOverlap()
 		}
 #endif
 	}
-#endif
 
 	bUAVOverlapEnabled = false;
 }
diff --git a/Engine/Source/Runtime/Windows/D3D11RHI/Private/D3D11Device.cpp b/Engine/Source/Runtime/Windows/D3D11RHI/Private/D3D11Device.cpp
--- a/Engine/Source/Runtime/Windows/D3D11RHI/Private/D3D11Device.cpp
+++ b/Engine/Source/Runtime/Windows/D3D11RHI/Private/D3D11Device.cpp
@@ -261,20 +261,6 @@ void FD3D11DynamicRHI::RHIPushEvent(const TCHAR* Name, FColor Color)
 	GPUProfilingData.PushEvent(Name, Color);
 }
 
-#if PLATFORM_HOLOLENS
-void FD3D11DynamicRHI::RHISuspendRendering()
-{
-	TRefCountPtr<IDXGIDevice3> dxgiDevice;
-	if (Direct3DDevice->QueryInterface(IID_PPV_ARGS(dxgiDevice.GetInitReference())))
-	{
-		dxgiDevice->Trim();
-	}
-}
-void FD3D11DynamicRHI::RHIResumeRendering()
-{
-}
-#endif
-
 void FD3D11DynamicRHI::RHIPopEvent()
 { 
 	GPUProfilingData.PopEvent(); 
diff --git a/Engine/Source/Runtime/Windows/D3D11RHI/Private/D3D11RHIPrivate.h b/Engine/Source/Runtime/Windows/D3D11RHI/Private/D3D11RHIPrivate.h
--- a/Engine/Source/Runtime/Windows/D3D11RHI/Private/D3D11RHIPrivate.h
+++ b/Engine/Source/Runtime/Windows/D3D11RHI/Private/D3D11RHIPrivate.h
@@ -432,11 +432,6 @@ public:
 	virtual void Shutdown() override;
 	virtual const TCHAR* GetName() override { return TEXT("D3D11"); }
 
-#if PLATFORM_HOLOLENS
-	virtual void RHISuspendRendering() override;
-	virtual void RHIResumeRendering() override;
-#endif
-
 	// HDR display output
 	virtual void EnableHDR();
 	virtual void ShutdownHDR();
diff --git a/Engine/Source/Runtime/Windows/D3D11RHI/Private/D3D11Shaders.cpp b/Engine/Source/Runtime/Windows/D3D11RHI/Private/D3D11Shaders.cpp
--- a/Engine/Source/Runtime/Windows/D3D11RHI/Private/D3D11Shaders.cpp
+++ b/Engine/Source/Runtime/Windows/D3D11RHI/Private/D3D11Shaders.cpp
@@ -66,7 +66,7 @@ static bool ApplyVendorExtensions(ID3D11Device* Direct3DDevice, EShaderFrequency
 	}
 
 	bool IsValidHardwareExtension = true;
-#if !PLATFORM_HOLOLENS
+
 	for (const FShaderCodeVendorExtension& Extension : ShaderData->VendorExtensions)
 	{
 		if (Extension.VendorId == EGpuVendorId::Nvidia)
@@ -105,7 +105,7 @@ static bool ApplyVendorExtensions(ID3D11Device* Direct3DDevice, EShaderFrequency
 			// TODO: [REDACTED_PROFILE_URL]/intel-graphics-compiler/blob/master/inc/IntelExtensions.hlsl
 		}
 	}
-#endif
+
 	return IsValidHardwareExtension;
 }
 
diff --git a/Engine/Source/Runtime/Windows/D3D11RHI/Private/D3D11Viewport.cpp b/Engine/Source/Runtime/Windows/D3D11RHI/Private/D3D11Viewport.cpp
--- a/Engine/Source/Runtime/Windows/D3D11RHI/Private/D3D11Viewport.cpp
+++ b/Engine/Source/Runtime/Windows/D3D11RHI/Private/D3D11Viewport.cpp
@@ -209,12 +209,10 @@ FD3D11Viewport::~FD3D11Viewport()
 
 	// If the swap chain was in fullscreen mode, switch back to windowed before releasing the swap chain.
 	// DXGI throws an error otherwise.
-#if !PLATFORM_HOLOLENS
 	if (SwapChain)
 	{
 		VERIFYD3D11RESULT_EX(SwapChain->SetFullscreenState(false, NULL), D3DRHI->GetDevice());
 	}
-#endif
 
 	D3DRHI->Viewports.Remove(this);
 }
@@ -593,7 +591,6 @@ bool FD3D11Viewport::Present(bool bLockToVsync)
 {
 	bool bNativelyPresented = true;
 #if	D3D11_WITH_DWMAPI
-#if !PLATFORM_HOLOLENS
 	// We can't call Present if !bIsValid, as it waits a window message to be processed, but the main thread may not be pumping the message handler.
 	if(ValidState != 0 && SwapChain.IsValid())
 	{
@@ -607,7 +604,6 @@ bool FD3D11Viewport::Present(bool bLockToVsync)
 			ValidState = VIEWPORT_INVALID;
 		}
 	}
-#endif
 	if (MaximumFrameLatency != RHIConsoleVariables::MaximumFrameLatency)
 	{
 		MaximumFrameLatency = RHIConsoleVariables::MaximumFrameLatency;	
diff --git a/Engine/Source/Runtime/Windows/D3D11RHI/Private/Windows/WindowsD3D11Device.cpp b/Engine/Source/Runtime/Windows/D3D11RHI/Private/Windows/WindowsD3D11Device.cpp
--- a/Engine/Source/Runtime/Windows/D3D11RHI/Private/Windows/WindowsD3D11Device.cpp
+++ b/Engine/Source/Runtime/Windows/D3D11RHI/Private/Windows/WindowsD3D11Device.cpp
@@ -277,12 +277,8 @@ static void SafeCreateDXGIFactory(IDXGIFactory1** DXGIFactory1, bool bWithDebug)
 	__try
 	{
 		bool bQuadBufferStereoRequested = FParse::Param(FCommandLine::Get(), TEXT("quad_buffer_stereo"));
-
-#if PLATFORM_HOLOLENS
-		bool bIsWin8OrNewer = true;
-#else
 		bool bIsWin8OrNewer = FPlatformMisc::VerifyWindowsVersion(8, 0);
-#endif
+
 		if (bIsWin8OrNewer && (bQuadBufferStereoRequested || bWithDebug))
 		{
 			// CreateDXGIFactory2 is only available on Win8.1+, find it if it exists
@@ -427,11 +423,7 @@ static bool SafeTestD3D11CreateDevice(IDXGIAdapter* Adapter,D3D_FEATURE_LEVEL Mi
 		// Log any reason for failure to create test device. Extra debug help.
 		VERIFYD3D11RESULT_NOEXIT(Result);
 
-#if PLATFORM_HOLOLENS
-		bool bIsWin10 = true;
-#else
 		bool bIsWin10 = FPlatformMisc::VerifyWindowsVersion(10, 0);
-#endif
 
 		// Fatal error on 0x887A002D
 		if (DXGI_ERROR_SDK_COMPONENT_MISSING == Result && bIsWin10)
@@ -1192,19 +1184,10 @@ void FD3D11DynamicRHIModule::FindAdapter()
 
 FDynamicRHI* FD3D11DynamicRHIModule::CreateRHI(ERHIFeatureLevel::Type RequestedFeatureLevel)
 {
-#if PLATFORM_HOLOLENS
-	GMaxRHIFeatureLevel = ERHIFeatureLevel::ES3_1;
-	GMaxRHIShaderPlatform = SP_D3D_ES3_1_HOLOLENS;
-#endif
-
 	IDXGIFactory1* DXGIFactory1;
 	VERIFYD3D11RESULT(ChosenAdapter.DXGIAdapter->GetParent(__uuidof(DXGIFactory1), reinterpret_cast<void**>(&DXGIFactory1)));
 
-#if PLATFORM_HOLOLENS
-	GShaderPlatformForFeatureLevel[ERHIFeatureLevel::ES3_1] = SP_D3D_ES3_1_HOLOLENS;
-#else
 	GShaderPlatformForFeatureLevel[ERHIFeatureLevel::ES3_1] = SP_PCD3D_ES3_1;
-#endif
 	GShaderPlatformForFeatureLevel[ERHIFeatureLevel::SM5] = SP_PCD3D_SM5;
 
 	ERHIFeatureLevel::Type PreviewFeatureLevel;
diff --git a/Engine/Source/Runtime/Windows/D3D11RHI/Public/D3D11Viewport.h b/Engine/Source/Runtime/Windows/D3D11RHI/Public/D3D11Viewport.h
--- a/Engine/Source/Runtime/Windows/D3D11RHI/Public/D3D11Viewport.h
+++ b/Engine/Source/Runtime/Windows/D3D11RHI/Public/D3D11Viewport.h
@@ -78,11 +78,7 @@ public:
 		FrameSyncEvent.IssueEvent();
 	}
 
-#if PLATFORM_HOLOLENS
-	IDXGISwapChain1* GetSwapChain() const { return SwapChain; } 
-#else
 	IDXGISwapChain* GetSwapChain() const { return SwapChain; }
-#endif
 
 	virtual void* GetNativeSwapChain() const override { return GetSwapChain(); }
 	virtual void* GetNativeBackBufferTexture() const override { return GetBackBuffer()->GetD3D11Texture2D(); }
@@ -150,11 +146,7 @@ protected:
 
 	static uint32 GSwapChainFlags;
 
-#if PLATFORM_HOLOLENS
-	TRefCountPtr<IDXGISwapChain1> SwapChain;
-#else
 	TRefCountPtr<IDXGISwapChain> SwapChain;
-#endif
 	TRefCountPtr<FD3D11Texture> BackBuffer;
 
 	// Support for selecting non-default output for display in fullscreen exclusive
