D3D12 Bindless Rework
- Moving FD3D12BindlessDescriptorManager to the new file D3D12BindlessDescriptors.h.
- Adding FD3D12BindlessDescriptorHeapManager to manage the bindless heaps.
- Changing other descriptor allocating methods to not do anything when bindless is enabled.
- Root Signatures should track Table CBVs differently from Root CBVs since only the former will be bindless.
- Wrapping more D3D12 bindless code with PLATFORM_SUPPORTS_BINDLESS_RENDERING
- Adding delayed deletion of FD3D12DescriptorHeap

diff --git a/Engine/Source/Runtime/D3D12RHI/Private/D3D12BindlessDescriptors.cpp b/Engine/Source/Runtime/D3D12RHI/Private/D3D12BindlessDescriptors.cpp
new file mode 100644
--- /dev/null
+++ b/Engine/Source/Runtime/D3D12RHI/Private/D3D12BindlessDescriptors.cpp
@@ -0,0 +1,320 @@
+// Copyright Epic Games, Inc. All Rights Reserved.
+
+#include "D3D12RHIPrivate.h"
+#include "D3D12BindlessDescriptors.h"
+#include "D3D12Descriptors.h"
+
+#if PLATFORM_SUPPORTS_BINDLESS_RENDERING
+
+int32 GBindlessResourceDescriptorHeapSize = 1000 * 1000;
+static FAutoConsoleVariableRef CVarBindlessResourceDescriptorHeapSize(
+	TEXT("D3D12.Bindless.ResourceDescriptorHeapSize"),
+	GBindlessResourceDescriptorHeapSize,
+	TEXT("Bindless resource descriptor heap size"),
+	ECVF_ReadOnly
+);
+
+int32 GBindlessSamplerDescriptorHeapSize = 2048;
+static FAutoConsoleVariableRef CVarBindlessSamplerDescriptorHeapSize(
+	TEXT("D3D12.Bindless.SamplerDescriptorHeapSize"),
+	GBindlessSamplerDescriptorHeapSize,
+	TEXT("Bindless sampler descriptor heap size"),
+	ECVF_ReadOnly
+);
+
+////////////////////////////////////////////////////////////////////////////////////////////////////////////////////////
+// FD3D12BindlessDescriptorHeapManager
+
+FD3D12BindlessDescriptorHeapManager::FD3D12BindlessDescriptorHeapManager(FD3D12Device* InDevice, ERHIDescriptorHeapType InType, ERHIBindlessConfiguration InConfiguration, uint32 InNumDescriptorsPerHeap, TConstArrayView<TStatId> InStats)
+	: FD3D12DeviceChild(InDevice)
+	, Type(InType)
+	, Configuration(InConfiguration)
+	, NumDescriptorsPerHeap(InNumDescriptorsPerHeap)
+{
+#if STATS
+	Stats = InStats;
+#endif
+
+	SetupInitialState(InNumDescriptorsPerHeap);
+}
+
+FRHIDescriptorHandle FD3D12BindlessDescriptorHeapManager::Allocate()
+{
+	FScopeLock Lock(&CriticalSection);
+
+	int32 AllocatedIndex = INDEX_NONE;
+
+	while ((AllocatedIndex = Allocations.FindAndSetFirstZeroBit()) == INDEX_NONE)
+	{
+		const int32 NewNumDescriptorsPerHeap = NumDescriptorsPerHeap * 1.5;
+		ResizeHeaps(NewNumDescriptorsPerHeap);
+	}
+
+	NumAllocations++;
+	RecordAlloc();
+
+	checkSlow(NumAllocations == Allocations.CountSetBits());
+
+	return FRHIDescriptorHandle(Type, AllocatedIndex);
+}
+
+void FD3D12BindlessDescriptorHeapManager::Free(FRHIDescriptorHandle InHandle)
+{
+	if (InHandle.IsValid())
+	{
+		FScopeLock Lock(&CriticalSection);
+
+		const int32 Index = InHandle.GetIndex();
+		check(Allocations[Index] == true);
+		Allocations[Index] = false;
+
+		NumAllocations--;
+		RecordFree();
+
+		checkSlow(NumAllocations == Allocations.CountSetBits());
+	}
+}
+
+void FD3D12BindlessDescriptorHeapManager::UpdateImmediately(FRHIDescriptorHandle InHandle, D3D12_CPU_DESCRIPTOR_HANDLE InSourceCpuHandle)
+{
+	UE::D3D12Descriptors::CopyDescriptor(GetParentDevice(), CpuHeap, InHandle, InSourceCpuHandle);
+	UE::D3D12Descriptors::CopyDescriptor(GetParentDevice(), GpuHeap, InHandle, InSourceCpuHandle);
+}
+
+void FD3D12BindlessDescriptorHeapManager::UpdateDeferred(FRHIDescriptorHandle InHandle, D3D12_CPU_DESCRIPTOR_HANDLE InSourceCpuHandle)
+{
+	// TODO: implement deferred updates
+	UpdateImmediately(InHandle, InSourceCpuHandle);
+}
+
+void FD3D12BindlessDescriptorHeapManager::SetupInitialState(uint32 InNumDescriptorsPerHeap)
+{
+	CpuHeap = CreateCpuHeapInternal(InNumDescriptorsPerHeap);
+	GpuHeap = CreateGpuHeapInternal(InNumDescriptorsPerHeap);
+
+	Allocations = AllocationListType(false, InNumDescriptorsPerHeap);
+	NumAllocations = 0;
+
+	NumDescriptorsPerHeap = InNumDescriptorsPerHeap;
+}
+
+void FD3D12BindlessDescriptorHeapManager::ResizeHeaps(uint32 InNewNumDescriptorsPerHeap)
+{
+	if (InNewNumDescriptorsPerHeap > NumDescriptorsPerHeap)
+	{
+		// Save off the old heaps and descriptor count
+		FD3D12DescriptorHeapPtr OldCpuHeap = CpuHeap;
+		FD3D12DescriptorHeapPtr OldGpuHeap = GpuHeap;
+		const uint32 OldNumDescriptorsPerHeap = NumDescriptorsPerHeap;
+
+		// Create the new heaps
+		CpuHeap = CreateCpuHeapInternal(InNewNumDescriptorsPerHeap);
+		GpuHeap = CreateGpuHeapInternal(InNewNumDescriptorsPerHeap);
+
+		// Resize the allocation list
+		Allocations.SetNum(InNewNumDescriptorsPerHeap, false);
+
+		// Validate the count is still correct
+		checkSlow(NumAllocations == Allocations.CountSetBits());
+
+		// Now we can switch to the new size
+		NumDescriptorsPerHeap = InNewNumDescriptorsPerHeap;
+
+		// Copy the old descriptor range from the old heap to the new heaps
+		UE::D3D12Descriptors::CopyDescriptors(GetParentDevice(), CpuHeap, OldCpuHeap, OldNumDescriptorsPerHeap);
+		UE::D3D12Descriptors::CopyDescriptors(GetParentDevice(), GpuHeap, OldCpuHeap, OldNumDescriptorsPerHeap);
+
+		// Finally, queue the old heaps for deletion
+		FD3D12DescriptorHeapManager& HeapManager = GetParentDevice()->GetDescriptorHeapManager();
+		HeapManager.DeferredFreeHeap(OldCpuHeap);
+		HeapManager.DeferredFreeHeap(OldGpuHeap);
+	}
+}
+
+FD3D12DescriptorHeap* FD3D12BindlessDescriptorHeapManager::CreateCpuHeapInternal(uint32 InNewNumDescriptorsPerHeap)
+{
+	const TCHAR* const HeapName = (GetType() == ERHIDescriptorHeapType::Standard) ? TEXT("BindlessResourcesCPU") : TEXT("BindlessSamplersCPU");
+
+	return GetParentDevice()->GetDescriptorHeapManager().AllocateIndependentHeap(
+		HeapName,
+		GetType(),
+		InNewNumDescriptorsPerHeap,
+		D3D12_DESCRIPTOR_HEAP_FLAG_NONE
+	);
+}
+
+FD3D12DescriptorHeap* FD3D12BindlessDescriptorHeapManager::CreateGpuHeapInternal(uint32 InNewNumDescriptorsPerHeap)
+{
+	const TCHAR* const HeapName = (GetType() == ERHIDescriptorHeapType::Standard) ? TEXT("BindlessResources") : TEXT("BindlessSamplers");
+
+	return GetParentDevice()->GetDescriptorHeapManager().AllocateIndependentHeap(
+		HeapName,
+		GetType(),
+		InNewNumDescriptorsPerHeap,
+		D3D12_DESCRIPTOR_HEAP_FLAG_SHADER_VISIBLE
+	);
+}
+
+void FD3D12BindlessDescriptorHeapManager::UpdateGpuHeap(FD3D12DescriptorHeap* InGpuHeap)
+{
+	UE::D3D12Descriptors::CopyDescriptors(GetParentDevice(), InGpuHeap, CpuHeap, NumDescriptorsPerHeap);
+}
+
+////////////////////////////////////////////////////////////////////////////////////////////////////////////////////////
+// FD3D12BindlessDescriptorManager
+
+FD3D12BindlessDescriptorManager::FD3D12BindlessDescriptorManager(FD3D12Device* InDevice)
+	: FD3D12DeviceChild(InDevice)
+{
+}
+
+FD3D12BindlessDescriptorManager::~FD3D12BindlessDescriptorManager() = default;
+
+void FD3D12BindlessDescriptorManager::Init()
+{
+	ResourcesConfiguration = RHIGetRuntimeBindlessResourcesConfiguration(GMaxRHIShaderPlatform);
+	SamplersConfiguration  = RHIGetRuntimeBindlessSamplersConfiguration(GMaxRHIShaderPlatform);
+
+	const uint32 NumResourceDescriptors = ResourcesConfiguration != ERHIBindlessConfiguration::Disabled ? GBindlessResourceDescriptorHeapSize : 0;
+	const uint32 NumSamplerDescriptors = SamplersConfiguration != ERHIBindlessConfiguration::Disabled ? GBindlessSamplerDescriptorHeapSize  : 0;
+
+	if (NumResourceDescriptors > 0)
+	{
+		const TStatId Stats[] =
+		{
+			GET_STATID(STAT_ResourceDescriptorsAllocated),
+			GET_STATID(STAT_BindlessResourceDescriptorsAllocated),
+		};
+
+		Managers.Emplace(GetParentDevice(), ERHIDescriptorHeapType::Standard, ResourcesConfiguration, NumResourceDescriptors, Stats);
+	}
+
+	if (NumSamplerDescriptors > 0)
+	{
+		const TStatId Stats[] =
+		{
+			GET_STATID(STAT_SamplerDescriptorsAllocated),
+			GET_STATID(STAT_BindlessSamplerDescriptorsAllocated),
+		};
+
+		Managers.Emplace(GetParentDevice(), ERHIDescriptorHeapType::Sampler, SamplersConfiguration, NumSamplerDescriptors, Stats);
+	}
+}
+
+FRHIDescriptorHandle FD3D12BindlessDescriptorManager::Allocate(ERHIDescriptorHeapType InType)
+{
+	for (FD3D12BindlessDescriptorHeapManager& Manager : Managers)
+	{
+		if (Manager.HandlesAllocation(InType))
+		{
+			return Manager.Allocate();
+		}
+	}
+	return FRHIDescriptorHandle();
+}
+
+void FD3D12BindlessDescriptorManager::ImmediateFree(FRHIDescriptorHandle InHandle)
+{
+	for (FD3D12BindlessDescriptorHeapManager& Manager : Managers)
+	{
+		if (Manager.HandlesAllocation(InHandle.GetType()))
+		{
+			Manager.Free(InHandle);
+			return;
+		}
+	}
+
+	// Bad configuration?
+	checkNoEntry();
+}
+
+void FD3D12BindlessDescriptorManager::DeferredFreeFromDestructor(FRHIDescriptorHandle InHandle)
+{
+	if (InHandle.IsValid())
+	{
+		FD3D12DynamicRHI::GetD3DRHI()->DeferredDelete(InHandle, GetParentDevice());
+	}
+}
+
+void FD3D12BindlessDescriptorManager::UpdateImmediately(FRHIDescriptorHandle InHandle, D3D12_CPU_DESCRIPTOR_HANDLE InSourceCpuHandle)
+{
+	for (FD3D12BindlessDescriptorHeapManager& Manager : Managers)
+	{
+		if (Manager.HandlesAllocation(InHandle.GetType()))
+		{
+			Manager.UpdateImmediately(InHandle, InSourceCpuHandle);
+			return;
+		}
+	}
+
+	// Bad configuration?
+	checkNoEntry();
+}
+
+void FD3D12BindlessDescriptorManager::UpdateDeferred(FRHIDescriptorHandle InHandle, D3D12_CPU_DESCRIPTOR_HANDLE InSourceCpuHandle)
+{
+	for (FD3D12BindlessDescriptorHeapManager& Manager : Managers)
+	{
+		if (Manager.HandlesAllocation(InHandle.GetType()))
+		{
+			Manager.UpdateDeferred(InHandle, InSourceCpuHandle);
+			return;
+		}
+	}
+
+	// Bad configuration?
+	checkNoEntry();
+}
+
+FD3D12DescriptorHeap* FD3D12BindlessDescriptorManager::GetHeap(ERHIDescriptorHeapType InType)
+{
+	for (FD3D12BindlessDescriptorHeapManager& Manager : Managers)
+	{
+		if (Manager.HandlesAllocation(InType))
+		{
+			return Manager.GetHeap();
+		}
+	}
+
+	return nullptr;
+}
+
+FD3D12DescriptorHeap* FD3D12BindlessDescriptorManager::GetHeap(ERHIDescriptorHeapType InType, ERHIBindlessConfiguration InConfiguration)
+{
+	for (FD3D12BindlessDescriptorHeapManager& Manager : Managers)
+	{
+		if (Manager.HandlesAllocation(InType, InConfiguration))
+		{
+			return Manager.GetHeap();
+		}
+	}
+	return nullptr;
+}
+
+
+bool FD3D12BindlessDescriptorManager::HasHeap(ERHIDescriptorHeapType InType, ERHIBindlessConfiguration InConfiguration) const
+{
+	for (const FD3D12BindlessDescriptorHeapManager& Manager : Managers)
+	{
+		if (Manager.HandlesAllocation(InType, InConfiguration))
+		{
+			return true;
+		}
+	}
+
+	return false;
+}
+
+D3D12_GPU_DESCRIPTOR_HANDLE FD3D12BindlessDescriptorManager::GetGpuHandle(FRHIDescriptorHandle InHandle) const
+{
+	for (const FD3D12BindlessDescriptorHeapManager& Manager : Managers)
+	{
+		if (Manager.HandlesAllocation(InHandle.GetType()))
+		{
+			return Manager.GetHeap()->GetGPUSlotHandle(InHandle.GetIndex());
+		}
+	}
+	return D3D12_GPU_DESCRIPTOR_HANDLE{};
+}
+
+#endif // PLATFORM_SUPPORTS_BINDLESS_RENDERING
diff --git a/Engine/Source/Runtime/D3D12RHI/Private/D3D12BindlessDescriptors.h b/Engine/Source/Runtime/D3D12RHI/Private/D3D12BindlessDescriptors.h
new file mode 100644
--- /dev/null
+++ b/Engine/Source/Runtime/D3D12RHI/Private/D3D12BindlessDescriptors.h
@@ -0,0 +1,122 @@
+// Copyright Epic Games, Inc. All Rights Reserved.
+
+#pragma once
+
+#include "D3D12RHICommon.h"
+#include "RHIDefinitions.h"
+#include "Templates/RefCounting.h"
+
+#if PLATFORM_SUPPORTS_BINDLESS_RENDERING
+
+struct FD3D12DescriptorHeap;
+using FD3D12DescriptorHeapPtr = TRefCountPtr<FD3D12DescriptorHeap>;
+
+class FD3D12BindlessDescriptorHeapManager : public FD3D12DeviceChild
+{
+public:
+	FD3D12BindlessDescriptorHeapManager() = delete;
+	FD3D12BindlessDescriptorHeapManager(FD3D12Device* InDevice, ERHIDescriptorHeapType InType, ERHIBindlessConfiguration InConfiguration, uint32 InNumDescriptorsPerHeap, TConstArrayView<TStatId> InStats);
+
+	FRHIDescriptorHandle Allocate();
+	void                 Free(FRHIDescriptorHandle InHandle);
+
+	void UpdateImmediately(FRHIDescriptorHandle InHandle, D3D12_CPU_DESCRIPTOR_HANDLE InSourceCpuHandle);
+	void UpdateDeferred   (FRHIDescriptorHandle InHandle, D3D12_CPU_DESCRIPTOR_HANDLE InSourceCpuHandle);
+
+	      FD3D12DescriptorHeap* GetHeap()       { return GpuHeap.GetReference(); }
+	const FD3D12DescriptorHeap* GetHeap() const { return GpuHeap.GetReference(); }
+	ERHIDescriptorHeapType      GetType() const { return Type; }
+	
+	bool HandlesAllocation(ERHIDescriptorHeapType InType) const { return GetType() == InType; }
+	bool HandlesConfiguration(ERHIBindlessConfiguration InConfiguration) const { return Configuration != ERHIBindlessConfiguration::Disabled && Configuration <= InConfiguration; }
+
+	bool HandlesAllocation(ERHIDescriptorHeapType InType, ERHIBindlessConfiguration InConfiguration) const
+	{
+		return HandlesAllocation(InType) && HandlesConfiguration(InConfiguration);
+	}
+
+private:
+	void SetupInitialState(uint32 InNumDescriptorsPerHeap);
+	void ResizeHeaps(uint32 InNewNumDescriptorsPerHeap);
+
+	FD3D12DescriptorHeap* CreateCpuHeapInternal(uint32 InNumDescriptorsPerHeap);
+	FD3D12DescriptorHeap* CreateGpuHeapInternal(uint32 InNumDescriptorsPerHeap);
+
+	void UpdateGpuHeap(FD3D12DescriptorHeap* GpuHeap);
+
+	void RecordAlloc()
+	{
+#if STATS
+		for (TStatId Stat : Stats)
+		{
+			INC_DWORD_STAT_FName(Stat.GetName());
+		}
+#endif
+	}
+
+	void RecordFree()
+	{
+#if STATS
+		for (TStatId Stat : Stats)
+		{
+			DEC_DWORD_STAT_FName(Stat.GetName());
+		}
+#endif
+	}
+
+private:
+	using AllocationListType = TBitArray<>;
+
+	FCriticalSection        CriticalSection;
+	
+	FD3D12DescriptorHeapPtr CpuHeap;
+	FD3D12DescriptorHeapPtr GpuHeap;
+
+	AllocationListType      Allocations;
+	int32                   NumAllocations = 0;
+
+	const ERHIDescriptorHeapType    Type;
+	const ERHIBindlessConfiguration Configuration;
+	uint32                          NumDescriptorsPerHeap;
+
+#if STATS
+	TArray<TStatId> Stats;
+#endif
+};
+
+/** Manager for resource descriptors used in bindless rendering. */
+class FD3D12BindlessDescriptorManager : public FD3D12DeviceChild
+{
+public:
+	FD3D12BindlessDescriptorManager(FD3D12Device* InDevice);
+	~FD3D12BindlessDescriptorManager();
+
+	void Init();
+
+	ERHIBindlessConfiguration GetResourcesConfiguration() const { return ResourcesConfiguration; }
+	ERHIBindlessConfiguration GetSamplersConfiguration()  const { return SamplersConfiguration; }
+
+	bool AreResourcesFullyBindless() const { return GetResourcesConfiguration() == ERHIBindlessConfiguration::AllShaders; }
+	bool AreSamplersFullyBindless () const { return GetSamplersConfiguration()  == ERHIBindlessConfiguration::AllShaders; }
+
+	FRHIDescriptorHandle Allocate(ERHIDescriptorHeapType InType);
+	void                 ImmediateFree(FRHIDescriptorHandle InHandle);
+	void                 DeferredFreeFromDestructor(FRHIDescriptorHandle InHandle);
+
+	void UpdateImmediately(FRHIDescriptorHandle InHandle, D3D12_CPU_DESCRIPTOR_HANDLE InSourceCpuHandle);
+	void UpdateDeferred(FRHIDescriptorHandle InHandle, D3D12_CPU_DESCRIPTOR_HANDLE InSourceCpuHandle);
+
+	FD3D12DescriptorHeap* GetHeap(ERHIDescriptorHeapType InType);
+	FD3D12DescriptorHeap* GetHeap(ERHIDescriptorHeapType InType, ERHIBindlessConfiguration InConfiguration);
+	bool                  HasHeap(ERHIDescriptorHeapType InType, ERHIBindlessConfiguration InConfiguration) const;
+
+	D3D12_GPU_DESCRIPTOR_HANDLE GetGpuHandle(FRHIDescriptorHandle InHandle) const;
+
+private:
+	TArray<FD3D12BindlessDescriptorHeapManager> Managers;
+
+	ERHIBindlessConfiguration ResourcesConfiguration{};
+	ERHIBindlessConfiguration SamplersConfiguration{};
+};
+
+#endif // PLATFORM_SUPPORTS_BINDLESS_RENDERING
diff --git a/Engine/Source/Runtime/D3D12RHI/Private/D3D12CommandContext.h b/Engine/Source/Runtime/D3D12RHI/Private/D3D12CommandContext.h
--- a/Engine/Source/Runtime/D3D12RHI/Private/D3D12CommandContext.h
+++ b/Engine/Source/Runtime/D3D12RHI/Private/D3D12CommandContext.h
@@ -10,7 +10,6 @@ D3D12CommandContext.h: D3D12 Command Context Interfaces
 #include "D3D12Queue.h"
 
 #include "Windows/AllowWindowsPlatformTypes.h"
-
 THIRD_PARTY_INCLUDES_START
 #if USE_PIX
 	#include "pix3.h"
@@ -21,6 +20,7 @@ THIRD_PARTY_INCLUDES_END
 #include "RHICoreShader.h"
 #include "RHICore.h"
 
+struct FD3D12DescriptorHeap;
 struct FRayTracingShaderBindings;
 class FD3D12Device;
 
@@ -30,8 +30,11 @@ struct FD3D12DeferredDeleteObject
 	{
 		RHIObject,
 		D3DObject,
-		D3DHeap,
+		Heap,
+		DescriptorHeap,
+#if PLATFORM_SUPPORTS_BINDLESS_RENDERING
 		BindlessDescriptor,
+#endif
 		CPUAllocation,
 		DescriptorBlock,
 		VirtualAllocation
@@ -40,7 +43,8 @@ struct FD3D12DeferredDeleteObject
 	union
 	{
 		FD3D12Resource* RHIObject;
-		FD3D12Heap* D3DHeap;
+		FD3D12Heap* Heap;
+		FD3D12DescriptorHeap* DescriptorHeap;
 		ID3D12Object* D3DObject;
 
 		struct
@@ -71,9 +75,14 @@ struct FD3D12DeferredDeleteObject
 		, RHIObject(RHIObject)
 	{}
 
-	explicit FD3D12DeferredDeleteObject(FD3D12Heap* D3DHeap)
-		: Type(EType::D3DHeap)
-		, D3DHeap(D3DHeap)
+	explicit FD3D12DeferredDeleteObject(FD3D12Heap* InHeap)
+		: Type(EType::Heap)
+		, Heap(InHeap)
+	{}
+
+	explicit FD3D12DeferredDeleteObject(FD3D12DescriptorHeap* InDescriptorHeap)
+		: Type(EType::DescriptorHeap)
+		, DescriptorHeap(InDescriptorHeap)
 	{}
 
 	explicit FD3D12DeferredDeleteObject(ID3D12Object* D3DObject)
@@ -81,10 +90,12 @@ struct FD3D12DeferredDeleteObject
 		, D3DObject(D3DObject)
 	{}
 
+#if PLATFORM_SUPPORTS_BINDLESS_RENDERING
 	explicit FD3D12DeferredDeleteObject(FRHIDescriptorHandle Handle, FD3D12Device* Device)
 		: Type(EType::BindlessDescriptor)
 		, BindlessDescriptor({ Handle, Device })
 	{}
+#endif
 
 	explicit FD3D12DeferredDeleteObject(void* Ptr, EType Type)
 		: Type(Type)
diff --git a/Engine/Source/Runtime/D3D12RHI/Private/D3D12DescriptorCache.cpp b/Engine/Source/Runtime/D3D12RHI/Private/D3D12DescriptorCache.cpp
--- a/Engine/Source/Runtime/D3D12RHI/Private/D3D12DescriptorCache.cpp
+++ b/Engine/Source/Runtime/D3D12RHI/Private/D3D12DescriptorCache.cpp
@@ -40,18 +40,22 @@ FD3D12DescriptorCache::~FD3D12DescriptorCache()
 
 void FD3D12DescriptorCache::Init(uint32 InNumLocalViewDescriptors, uint32 InNumSamplerDescriptors)
 {
+#if PLATFORM_SUPPORTS_BINDLESS_RENDERING
+	FD3D12BindlessDescriptorManager& BindlessDescriptorManager = GetParentDevice()->GetBindlessDescriptorManager();
+
+	bBindlessResources = BindlessDescriptorManager.AreResourcesFullyBindless();
+	bBindlessSamplers = BindlessDescriptorManager.AreSamplersFullyBindless();
+#endif
+
 	// Always Init a local sampler heap as the high level cache will always miss initialy
 	// so we need something to fall back on (The view heap never rolls over so we init that one
 	// lazily as a backup to save memory)
-	LocalSamplerHeap.Init(InNumSamplerDescriptors, ERHIDescriptorHeapType::Sampler);
+	LocalSamplerHeap.Init(IsUsingBindlessSamplers() ? 0 : InNumSamplerDescriptors, ERHIDescriptorHeapType::Sampler);
 
-	NumLocalViewDescriptors = InNumLocalViewDescriptors;
-
-	CurrentViewHeap = &SubAllocatedViewHeap;
-	CurrentSamplerHeap = &LocalSamplerHeap;
+	NumLocalViewDescriptors = IsUsingBindlessResources() ? 0 : InNumLocalViewDescriptors;
 
-	OverrideViewHeap = nullptr;
-	OverrideSamplerHeap = nullptr;
+	CurrentViewHeap = IsUsingBindlessResources() ? nullptr : &SubAllocatedViewHeap;
+	CurrentSamplerHeap = IsUsingBindlessSamplers() ? nullptr : &LocalSamplerHeap;
 }
 
 bool FD3D12DescriptorCache::SetDescriptorHeaps(bool bForceHeapChanged)
@@ -59,13 +63,29 @@ bool FD3D12DescriptorCache::SetDescriptorHeaps(bool bForceHeapChanged)
 	// See if the descriptor heaps changed.
 	bool bHeapChanged = bForceHeapChanged;
 
-	ID3D12DescriptorHeap* PendingViewHeap = OverrideViewHeap ? OverrideViewHeap->GetHeap() : CurrentViewHeap->GetHeap();
+#if PLATFORM_SUPPORTS_BINDLESS_RENDERING && DO_ENSURE
+	FD3D12BindlessDescriptorManager& BindlessDescriptorManager = GetParentDevice()->GetBindlessDescriptorManager();
+	if (IsUsingBindlessResources())
+	{
+		check(BindlessResourcesHeap == BindlessDescriptorManager.GetHeap(ERHIDescriptorHeapType::Standard));
+	}
+	if (IsUsingBindlessSamplers())
+	{
+		check(BindlessSamplersHeap == BindlessDescriptorManager.GetHeap(ERHIDescriptorHeapType::Sampler));
+	}
+#endif
+
+	ID3D12DescriptorHeap* PendingViewHeap =
+#if PLATFORM_SUPPORTS_BINDLESS_RENDERING
+		IsUsingBindlessResources() ? BindlessResourcesHeap->GetHeap() :
+#endif
+		CurrentViewHeap->GetHeap();
 	if (LastSetViewHeap != PendingViewHeap)
 	{
 		// The view heap changed, so dirty the descriptor tables.
 		bHeapChanged = true;
 
-		if (!OverrideViewHeap)
+		if (!IsUsingBindlessResources())
 		{
 			Context.StateCache.DirtyViewDescriptorTables();
 		}
@@ -73,13 +93,17 @@ bool FD3D12DescriptorCache::SetDescriptorHeaps(bool bForceHeapChanged)
 		INC_DWORD_STAT_BY(STAT_ViewHeapChanged, LastSetViewHeap == nullptr ? 0 : 1);	// Don't count the initial set on a command list.
 	}
 
-	ID3D12DescriptorHeap* PendingSamplerHeap = OverrideSamplerHeap ? OverrideSamplerHeap->GetHeap() : CurrentSamplerHeap->GetHeap();
+	ID3D12DescriptorHeap* PendingSamplerHeap =
+#if PLATFORM_SUPPORTS_BINDLESS_RENDERING
+		IsUsingBindlessSamplers() ? BindlessSamplersHeap->GetHeap() :
+#endif
+		CurrentSamplerHeap->GetHeap();
 	if (LastSetSamplerHeap != PendingSamplerHeap)
 	{
 		// The sampler heap changed, so dirty the descriptor tables.
 		bHeapChanged = true;
 
-		if (!OverrideSamplerHeap)
+		if (!IsUsingBindlessSamplers())
 		{
 			Context.StateCache.DirtySamplerDescriptorTables();
 
@@ -111,23 +135,56 @@ void FD3D12DescriptorCache::OpenCommandList()
 	// Clear the previous heap pointers (since it's a new command list) and then set the current descriptor heaps.
 	LastSetViewHeap = nullptr;
 	LastSetSamplerHeap = nullptr;
-	OverrideViewHeap = nullptr;
-	OverrideSamplerHeap = nullptr;
 
-	// The global sampler heap doesn't care about the current command list
-	LocalSamplerHeap.OpenCommandList();
-	CurrentViewHeap->OpenCommandList();
 
-	SwitchToGlobalSamplerHeap();
+#if PLATFORM_SUPPORTS_BINDLESS_RENDERING
+	if (IsUsingBindlessSamplers())
+	{
+		BindlessSamplersHeap = GetParentDevice()->GetBindlessDescriptorManager().GetHeap(ERHIDescriptorHeapType::Sampler);
+	}
+	else
+#endif
+	{
+		// The global sampler heap doesn't care about the current command list
+		LocalSamplerHeap.OpenCommandList();
+	}
+
+#if PLATFORM_SUPPORTS_BINDLESS_RENDERING
+	if (IsUsingBindlessResources())
+	{
+		BindlessResourcesHeap = GetParentDevice()->GetBindlessDescriptorManager().GetHeap(ERHIDescriptorHeapType::Standard);
+	}
+	else
+#endif
+	{
+		CurrentViewHeap->OpenCommandList();
+	}
+
+	if (!IsUsingBindlessSamplers())
+	{
+		SwitchToGlobalSamplerHeap();
+	}
+
+	// Make sure the heaps are set
+	SetDescriptorHeaps();
+
+	check(IsUsingBindlessSamplers() || IsHeapSet(GetParentDevice()->GetGlobalSamplerHeap().GetHeap()));
 }
 
 void FD3D12DescriptorCache::CloseCommandList()
 {
-	CurrentViewHeap->CloseCommandList();
-	LocalSamplerHeap.CloseCommandList();
+	if (!IsUsingBindlessResources())
+	{
+		CurrentViewHeap->CloseCommandList();
+	}
+
+	if (!IsUsingBindlessSamplers())
+	{
+		LocalSamplerHeap.CloseCommandList();
 
-	GetParentDevice()->GetGlobalSamplerHeap().ConsolidateUniqueSamplerTables(UniqueTables);
-	UniqueTables.Reset();
+		GetParentDevice()->GetGlobalSamplerHeap().ConsolidateUniqueSamplerTables(UniqueTables);
+		UniqueTables.Reset();
+	}
 }
 
 void FD3D12DescriptorCache::SetVertexBuffers(FD3D12VertexBufferCache& Cache)
@@ -608,6 +665,8 @@ void FD3D12DescriptorCache::SetRootConstantBuffers(EShaderFrequency ShaderStage,
 
 bool FD3D12DescriptorCache::SwitchToContextLocalViewHeap()
 {
+	check(!IsUsingBindlessResources());
+
 	if (LocalViewHeap == nullptr)
 	{
 		UE_LOG(LogD3D12RHI, Log, TEXT("This should only happen in the Editor where it doesn't matter as much. If it happens in game you should increase the device global heap size!"));
@@ -638,6 +697,8 @@ bool FD3D12DescriptorCache::SwitchToContextLocalViewHeap()
 
 bool FD3D12DescriptorCache::SwitchToContextLocalSamplerHeap()
 {
+	check(!IsUsingBindlessSamplers());
+
 	CurrentSamplerHeap = &LocalSamplerHeap;
 
 	bool bDescriptorHeapsChanged = SetDescriptorHeaps();
@@ -646,31 +707,13 @@ bool FD3D12DescriptorCache::SwitchToContextLocalSamplerHeap()
 	return bDescriptorHeapsChanged;
 }
 
-bool FD3D12DescriptorCache::SwitchToGlobalSamplerHeap()
+void FD3D12DescriptorCache::SwitchToGlobalSamplerHeap()
 {
+	check(!IsUsingBindlessSamplers());
+
 	FD3D12GlobalOnlineSamplerHeap& GlobalSamplerHeap = GetParentDevice()->GetGlobalSamplerHeap();
 	LocalSamplerSet = GlobalSamplerHeap.GetUniqueDescriptorTables();
 	CurrentSamplerHeap = &GlobalSamplerHeap;
-
-	bool bDescriptorHeapsChanged = SetDescriptorHeaps();
-	check(IsHeapSet(GetParentDevice()->GetGlobalSamplerHeap().GetHeap()));
-
-	return bDescriptorHeapsChanged;
-}
-
-bool FD3D12DescriptorCache::SetHeapOverrides(FD3D12DescriptorHeap* InOverrideViewHeap, FD3D12DescriptorHeap* InOverrideSamplerHeap)
-{
-	bool bDescriptorHeapsChanged = false;
-
-	if (OverrideViewHeap != InOverrideViewHeap || OverrideSamplerHeap != InOverrideSamplerHeap)
-	{
-		OverrideViewHeap = InOverrideViewHeap;
-		OverrideSamplerHeap = InOverrideSamplerHeap;
-
-		bDescriptorHeapsChanged = SetDescriptorHeaps();
-	}
-
-	return bDescriptorHeapsChanged;
 }
 
 void FD3D12DescriptorCache::OverrideLastSetHeaps(ID3D12DescriptorHeap* ViewHeap, ID3D12DescriptorHeap* SamplerHeap)
@@ -1010,27 +1053,35 @@ Initialize a thread local online heap
 **/
 void FD3D12LocalOnlineHeap::Init(uint32 InNumDescriptors, ERHIDescriptorHeapType InHeapType)
 {
-	const TCHAR* DebugName = InHeapType == ERHIDescriptorHeapType::Standard ? L"Thread Local - Online View Heap" : L"Thread Local - Online Sampler Heap";
-	Heap = GetParentDevice()->GetDescriptorHeapManager().AllocateHeap(
-		DebugName,
-		InHeapType,
-		InNumDescriptors,
-		D3D12_DESCRIPTOR_HEAP_FLAG_SHADER_VISIBLE
-	);
+	if (InNumDescriptors > 0)
+	{
+		const TCHAR* DebugName = InHeapType == ERHIDescriptorHeapType::Standard ? L"Thread Local - Online View Heap" : L"Thread Local - Online Sampler Heap";
+		Heap = GetParentDevice()->GetDescriptorHeapManager().AllocateHeap(
+			DebugName,
+			InHeapType,
+			InNumDescriptors,
+			D3D12_DESCRIPTOR_HEAP_FLAG_SHADER_VISIBLE
+		);
 
-	Entry.Heap = Heap;
+		Entry.Heap = Heap;
 
-	if (InHeapType == ERHIDescriptorHeapType::Standard)
-	{
-		INC_DWORD_STAT(STAT_NumViewOnlineDescriptorHeaps);
-		INC_MEMORY_STAT_BY(STAT_ViewOnlineDescriptorHeapMemory, Heap->GetMemorySize());
+		if (InHeapType == ERHIDescriptorHeapType::Standard)
+		{
+			INC_DWORD_STAT(STAT_NumViewOnlineDescriptorHeaps);
+			INC_MEMORY_STAT_BY(STAT_ViewOnlineDescriptorHeapMemory, Heap->GetMemorySize());
+		}
+		else
+		{
+			INC_DWORD_STAT(STAT_NumSamplerOnlineDescriptorHeaps);
+			INC_MEMORY_STAT_BY(STAT_SamplerOnlineDescriptorHeapMemory, Heap->GetMemorySize());
+		}
+		INC_MEMORY_STAT_BY(STAT_D3D12MemoryCurrentTotal, Heap->GetMemorySize());
 	}
 	else
 	{
-		INC_DWORD_STAT(STAT_NumSamplerOnlineDescriptorHeaps);
-		INC_MEMORY_STAT_BY(STAT_SamplerOnlineDescriptorHeapMemory, Heap->GetMemorySize());
+		Heap = nullptr;
+		Entry.Heap = nullptr;
 	}
-	INC_MEMORY_STAT_BY(STAT_D3D12MemoryCurrentTotal, Heap->GetMemorySize());
 }
 
 
diff --git a/Engine/Source/Runtime/D3D12RHI/Private/D3D12DescriptorCache.h b/Engine/Source/Runtime/D3D12RHI/Private/D3D12DescriptorCache.h
--- a/Engine/Source/Runtime/D3D12RHI/Private/D3D12DescriptorCache.h
+++ b/Engine/Source/Runtime/D3D12RHI/Private/D3D12DescriptorCache.h
@@ -315,12 +315,17 @@ public:
 
 	~FD3D12DescriptorCache();
 
-	inline bool IsViewHeapOverridden() const { return OverrideViewHeap != nullptr; }
-	inline bool IsSamplerHeapOverridden() const { return OverrideSamplerHeap != nullptr; }
-
 	inline FD3D12OnlineHeap* GetCurrentViewHeap() const { return CurrentViewHeap; }
 	inline FD3D12OnlineHeap* GetCurrentSamplerHeap() const { return CurrentSamplerHeap; }
 
+#if PLATFORM_SUPPORTS_BINDLESS_RENDERING
+	bool IsUsingBindlessResources() const { return bBindlessResources; }
+	bool IsUsingBindlessSamplers()  const { return bBindlessSamplers;  }
+#else
+	constexpr bool IsUsingBindlessResources() const { return false; }
+	constexpr bool IsUsingBindlessSamplers()  const { return false; }
+#endif
+
 	// Checks if the specified descriptor heap has been set on the current command list.
 	bool IsHeapSet(ID3D12DescriptorHeap* const pHeap) const
 	{
@@ -356,9 +361,7 @@ public:
 
 	bool SwitchToContextLocalViewHeap();
 	bool SwitchToContextLocalSamplerHeap();
-	bool SwitchToGlobalSamplerHeap();
-
-	bool SetHeapOverrides(FD3D12DescriptorHeap* InOverrideViewHeap, FD3D12DescriptorHeap* InOverrideSamplerHeap);
+	void SwitchToGlobalSamplerHeap();
 
 	void OverrideLastSetHeaps(ID3D12DescriptorHeap* ViewHeap, ID3D12DescriptorHeap* SamplerHeap);
 	void RestoreAfterExternalHeapsSet();
@@ -382,9 +385,6 @@ private:
 	FD3D12OnlineHeap* CurrentViewHeap = nullptr;
 	FD3D12OnlineHeap* CurrentSamplerHeap = nullptr;
 
-	FD3D12DescriptorHeap* OverrideViewHeap = nullptr;
-	FD3D12DescriptorHeap* OverrideSamplerHeap = nullptr;
-
 	FD3D12LocalOnlineHeap* LocalViewHeap = nullptr;
 	FD3D12LocalOnlineHeap LocalSamplerHeap;
 	FD3D12SubAllocatedOnlineHeap SubAllocatedViewHeap;
@@ -397,4 +397,12 @@ private:
 	bool bHeapsOverridden = false;
 
 	uint32 NumLocalViewDescriptors = 0;
+
+#if PLATFORM_SUPPORTS_BINDLESS_RENDERING
+	bool bBindlessResources = false;
+	bool bBindlessSamplers = false;
+
+	FD3D12DescriptorHeap* BindlessResourcesHeap = nullptr;
+	FD3D12DescriptorHeap* BindlessSamplersHeap = nullptr;
+#endif
 };
\ No newline at end of file
diff --git a/Engine/Source/Runtime/D3D12RHI/Private/D3D12Descriptors.cpp b/Engine/Source/Runtime/D3D12RHI/Private/D3D12Descriptors.cpp
--- a/Engine/Source/Runtime/D3D12RHI/Private/D3D12Descriptors.cpp
+++ b/Engine/Source/Runtime/D3D12RHI/Private/D3D12Descriptors.cpp
@@ -3,6 +3,33 @@
 #include "D3D12RHIPrivate.h"
 #include "D3D12Descriptors.h"
 
+////////////////////////////////////////////////////////////////////////////////////////////////////////////////////////
+
+void UE::D3D12Descriptors::CopyDescriptor(FD3D12Device* Device, FD3D12DescriptorHeap* TargetHeap, FRHIDescriptorHandle Handle, D3D12_CPU_DESCRIPTOR_HANDLE SourceCpuHandle)
+{
+	if (Handle.IsValid())
+	{
+		const D3D12_CPU_DESCRIPTOR_HANDLE DestCpuHandle = TargetHeap->GetCPUSlotHandle(Handle.GetIndex());
+		const D3D12_DESCRIPTOR_HEAP_TYPE D3DHeapType = Translate(Handle.GetType());
+
+		Device->GetDevice()->CopyDescriptorsSimple(1, DestCpuHandle, SourceCpuHandle, D3DHeapType);
+	}
+}
+
+void UE::D3D12Descriptors::CopyDescriptors(FD3D12Device* Device, FD3D12DescriptorHeap* TargetHeap, FD3D12DescriptorHeap* SourceHeap, uint32 NumDescriptors)
+{
+	const D3D12_CPU_DESCRIPTOR_HANDLE TargetStart = TargetHeap->GetCPUSlotHandle(0);
+	const D3D12_CPU_DESCRIPTOR_HANDLE SourceStart = SourceHeap->GetCPUSlotHandle(0);
+	const D3D12_DESCRIPTOR_HEAP_TYPE D3DHeapType = Translate(TargetHeap->GetType());
+
+	Device->GetDevice()->CopyDescriptorsSimple(
+		NumDescriptors,
+		TargetStart,
+		SourceStart,
+		D3DHeapType
+	);
+}
+
 ////////////////////////////////////////////////////////////////////////////////////////////////////////////////////////
 // FD3D12DescriptorHeap
 
@@ -79,158 +106,7 @@ FD3D12DescriptorManager::~FD3D12DescriptorManager() = default;
 
 void FD3D12DescriptorManager::UpdateImmediately(FRHIDescriptorHandle InHandle, D3D12_CPU_DESCRIPTOR_HANDLE InSourceCpuHandle)
 {
-	if (InHandle.IsValid())
-	{
-		const D3D12_CPU_DESCRIPTOR_HANDLE DestCpuHandle = Heap->GetCPUSlotHandle(InHandle.GetIndex());
-		const D3D12_DESCRIPTOR_HEAP_TYPE D3DHeapType = Translate(GetType());
-
-		GetParentDevice()->GetDevice()->CopyDescriptorsSimple(1, DestCpuHandle, InSourceCpuHandle, D3DHeapType);
-	}
-}
-
-////////////////////////////////////////////////////////////////////////////////////////////////////////////////////////
-// FD3D12BindlessDescriptorManager
-
-FD3D12BindlessDescriptorManager::FD3D12BindlessDescriptorManager(FD3D12Device* Device)
-	: FD3D12DeviceChild(Device)
-{
-}
-
-FD3D12BindlessDescriptorManager::~FD3D12BindlessDescriptorManager() = default;
-
-void FD3D12BindlessDescriptorManager::Init(uint32 InNumResourceDescriptors, uint32 InNumSamplerDescriptors)
-{
-	FD3D12DescriptorHeapManager& HeapManager = GetParentDevice()->GetDescriptorHeapManager();
-
-	if (InNumResourceDescriptors > 0)
-	{
-		FD3D12DescriptorHeap* Heap = HeapManager.AllocateHeap(
-			TEXT("BindlessResources"),
-			ERHIDescriptorHeapType::Standard,
-			InNumResourceDescriptors,
-			D3D12_DESCRIPTOR_HEAP_FLAG_SHADER_VISIBLE
-		);
-
-		const TStatId Stats[] =
-		{
-			GET_STATID(STAT_ResourceDescriptorsAllocated),
-			GET_STATID(STAT_BindlessResourceDescriptorsAllocated),
-		};
-
-		Managers.Emplace(GetParentDevice(), Heap, Stats);
-	}
-
-	if (InNumSamplerDescriptors > 0)
-	{
-		FD3D12DescriptorHeap* Heap = HeapManager.AllocateHeap(
-			TEXT("BindlessSamplers"),
-			ERHIDescriptorHeapType::Sampler,
-			InNumSamplerDescriptors,
-			D3D12_DESCRIPTOR_HEAP_FLAG_SHADER_VISIBLE
-		);
-
-		const TStatId Stats[] =
-		{
-			GET_STATID(STAT_SamplerDescriptorsAllocated),
-			GET_STATID(STAT_BindlessSamplerDescriptorsAllocated),
-		};
-
-		Managers.Emplace(GetParentDevice(), Heap, Stats);
-	}
-}
-
-FRHIDescriptorHandle FD3D12BindlessDescriptorManager::Allocate(ERHIDescriptorHeapType InType)
-{
-	for (FD3D12DescriptorManager& Manager : Managers)
-	{
-		if (Manager.HandlesAllocation(InType))
-		{
-			return Manager.Allocate();
-		}
-	}
-	return FRHIDescriptorHandle();
-}
-
-void FD3D12BindlessDescriptorManager::ImmediateFree(FRHIDescriptorHandle InHandle)
-{
-	for (FD3D12DescriptorManager& Manager : Managers)
-	{
-		if (Manager.HandlesAllocation(InHandle.GetType()))
-		{
-			Manager.Free(InHandle);
-			return;
-		}
-	}
-
-	// Bad configuration?
-	checkNoEntry();
-}
-
-void FD3D12BindlessDescriptorManager::DeferredFreeFromDestructor(FRHIDescriptorHandle InHandle)
-{
-	if (InHandle.IsValid())
-	{
-		FD3D12DynamicRHI::GetD3DRHI()->DeferredDelete(InHandle, GetParentDevice());
-	}
-}
-
-void FD3D12BindlessDescriptorManager::UpdateImmediately(FRHIDescriptorHandle InHandle, D3D12_CPU_DESCRIPTOR_HANDLE InSourceCpuHandle)
-{
-	for (FD3D12DescriptorManager& Manager : Managers)
-	{
-		if (Manager.HandlesAllocation(InHandle.GetType()))
-		{
-			Manager.UpdateImmediately(InHandle, InSourceCpuHandle);
-			return;
-		}
-	}
-
-	// Bad configuration?
-	checkNoEntry();
-}
-
-void FD3D12BindlessDescriptorManager::UpdateDeferred(FRHIDescriptorHandle InHandle, D3D12_CPU_DESCRIPTOR_HANDLE InSourceCpuHandle)
-{
-	// TODO: implement deferred updates
-	UpdateImmediately(InHandle, InSourceCpuHandle);
-}
-
-FD3D12DescriptorHeap* FD3D12BindlessDescriptorManager::GetHeapForType(ERHIDescriptorHeapType InType)
-{
-	for (FD3D12DescriptorManager& Manager : Managers)
-	{
-		if (Manager.HandlesAllocation(InType))
-		{
-			return Manager.GetHeap();
-		}
-	}
-
-	return nullptr;
-}
-
-bool FD3D12BindlessDescriptorManager::HasHeapForType(ERHIDescriptorHeapType InType) const
-{
-	for (const FD3D12DescriptorManager& Manager : Managers)
-	{
-		if (Manager.HandlesAllocation(InType))
-		{
-			return true;
-		}
-	}
-
-	return false;
-}
-
-D3D12_GPU_DESCRIPTOR_HANDLE FD3D12BindlessDescriptorManager::GetGpuHandle(FRHIDescriptorHandle InHandle) const
-{
-	for (const FD3D12DescriptorManager& Manager : Managers)
-	{
-		if (Manager.HandlesAllocation(InHandle.GetType()))
-		{
-			return Manager.GetHeap()->GetGPUSlotHandle(InHandle.GetIndex());
-		}
-	}
-	return D3D12_GPU_DESCRIPTOR_HANDLE{};
+	UE::D3D12Descriptors::CopyDescriptor(GetParentDevice(), GetHeap(), InHandle, InSourceCpuHandle);
 }
 
 ////////////////////////////////////////////////////////////////////////////////////////////////////////////////////////
@@ -514,6 +390,11 @@ void FD3D12DescriptorHeapManager::Destroy()
 {
 }
 
+FD3D12DescriptorHeap* FD3D12DescriptorHeapManager::AllocateIndependentHeap(const TCHAR* InDebugName, ERHIDescriptorHeapType InHeapType, uint32 InNumDescriptors, D3D12_DESCRIPTOR_HEAP_FLAGS InHeapFlags)
+{
+	return CreateDescriptorHeap(GetParentDevice(), InDebugName, InHeapType, InNumDescriptors, InHeapFlags, false);
+}
+
 FD3D12DescriptorHeap* FD3D12DescriptorHeapManager::AllocateHeap(const TCHAR* InDebugName, ERHIDescriptorHeapType InHeapType, uint32 InNumDescriptors, D3D12_DESCRIPTOR_HEAP_FLAGS InHeapFlags)
 {
 	for (FD3D12DescriptorManager& GlobalHeap : GlobalHeaps)
@@ -531,18 +412,46 @@ FD3D12DescriptorHeap* FD3D12DescriptorHeapManager::AllocateHeap(const TCHAR* InD
 		}
 	}
 
-	return CreateDescriptorHeap(GetParentDevice(), InDebugName, InHeapType, InNumDescriptors, InHeapFlags, false);
+	return AllocateIndependentHeap(InDebugName, InHeapType, InNumDescriptors, InHeapFlags);
 }
 
-void FD3D12DescriptorHeapManager::FreeHeap(FD3D12DescriptorHeap* InHeap)
+void FD3D12DescriptorHeapManager::DeferredFreeHeap(FD3D12DescriptorHeap* InHeap)
 {
-	for (FD3D12DescriptorManager& GlobalHeap : GlobalHeaps)
+	if (InHeap->IsGlobal())
 	{
-		if (GlobalHeap.IsHeapAChild(InHeap))
+		for (FD3D12DescriptorManager& GlobalHeap : GlobalHeaps)
 		{
-			GlobalHeap.Free(InHeap->GetOffset(), InHeap->GetNumDescriptors());
-			return;
+			if (GlobalHeap.IsHeapAChild(InHeap))
+			{
+				InHeap->AddRef();
+				FD3D12DynamicRHI::GetD3DRHI()->DeferredDelete(InHeap);
+				return;
+			}
+		}
+	}
+	else
+	{
+		InHeap->AddRef();
+		FD3D12DynamicRHI::GetD3DRHI()->DeferredDelete(InHeap);
+	}
+}
+
+void FD3D12DescriptorHeapManager::ImmediateFreeHeap(FD3D12DescriptorHeap* InHeap)
+{
+	if (InHeap->IsGlobal())
+	{
+		for (FD3D12DescriptorManager& GlobalHeap : GlobalHeaps)
+		{
+			if (GlobalHeap.IsHeapAChild(InHeap))
+			{
+				GlobalHeap.Free(InHeap->GetOffset(), InHeap->GetNumDescriptors());
+				InHeap->Release();
+				return;
+			}
 		}
 	}
-	check(!InHeap->IsGlobal());
+	else
+	{
+		InHeap->Release();
+	}
 }
diff --git a/Engine/Source/Runtime/D3D12RHI/Private/D3D12Descriptors.h b/Engine/Source/Runtime/D3D12RHI/Private/D3D12Descriptors.h
--- a/Engine/Source/Runtime/D3D12RHI/Private/D3D12Descriptors.h
+++ b/Engine/Source/Runtime/D3D12RHI/Private/D3D12Descriptors.h
@@ -10,6 +10,8 @@
 #include "Containers/Queue.h"
 #include "Containers/List.h"
 
+struct FD3D12DescriptorHeap;
+
 inline D3D12_DESCRIPTOR_HEAP_TYPE Translate(ERHIDescriptorHeapType InHeapType)
 {
 	switch (InHeapType)
@@ -22,6 +24,13 @@ inline D3D12_DESCRIPTOR_HEAP_TYPE Translate(ERHIDescriptorHeapType InHeapType)
 	}
 }
 
+namespace UE::D3D12Descriptors
+{
+	FD3D12DescriptorHeap* CreateDescriptorHeap(FD3D12Device* Device, const TCHAR* DebugName, ERHIDescriptorHeapType HeapType, uint32 NumDescriptors, D3D12_DESCRIPTOR_HEAP_FLAGS Flags, bool bGlobal = false);
+	void CopyDescriptor(FD3D12Device* Device, FD3D12DescriptorHeap* TargetHeap, FRHIDescriptorHandle Handle, D3D12_CPU_DESCRIPTOR_HANDLE SourceCpuHandle);
+	void CopyDescriptors(FD3D12Device* Device, FD3D12DescriptorHeap* TargetHeap, FD3D12DescriptorHeap* SourceHeap, uint32 NumDescriptors);
+}
+
 struct FD3D12DescriptorHeap : public FD3D12DeviceChild, public FThreadSafeRefCountedObject
 {
 public:
@@ -103,31 +112,6 @@ private:
 	FD3D12DescriptorHeapPtr Heap;
 };
 
-/** Manager for resource descriptors used in bindless rendering. */
-class FD3D12BindlessDescriptorManager : public FD3D12DeviceChild
-{
-public:
-	FD3D12BindlessDescriptorManager(FD3D12Device* Device);
-	~FD3D12BindlessDescriptorManager();
-
-	void Init(uint32 InNumResourceDescriptors, uint32 InNumSamplerDescriptors);
-
-	FRHIDescriptorHandle Allocate(ERHIDescriptorHeapType InType);
-	void ImmediateFree(FRHIDescriptorHandle InHandle);
-	void DeferredFreeFromDestructor(FRHIDescriptorHandle InHandle);
-
-	void UpdateImmediately(FRHIDescriptorHandle InHandle, D3D12_CPU_DESCRIPTOR_HANDLE InSourceCpuHandle);
-	void UpdateDeferred(FRHIDescriptorHandle InHandle, D3D12_CPU_DESCRIPTOR_HANDLE InSourceCpuHandle);
-
-	FD3D12DescriptorHeap* GetHeapForType(ERHIDescriptorHeapType InType);
-	bool HasHeapForType(ERHIDescriptorHeapType InType) const;
-
-	D3D12_GPU_DESCRIPTOR_HANDLE GetGpuHandle(FRHIDescriptorHandle InHandle) const;
-
-private:
-	TArray<FD3D12DescriptorManager> Managers;
-};
-
 /** Heap sub block of an online heap */
 struct FD3D12OnlineDescriptorBlock
 {
@@ -246,8 +230,10 @@ public:
 	void Init(uint32 InNumGlobalResourceDescriptors, uint32 InNumGlobalSamplerDescriptors);
 	void Destroy();
 
-	FD3D12DescriptorHeap* AllocateHeap(const TCHAR* DebugName, ERHIDescriptorHeapType InHeapType, uint32 InNumDescriptors, D3D12_DESCRIPTOR_HEAP_FLAGS InFlags);
-	void FreeHeap(FD3D12DescriptorHeap* InHeap);
+	FD3D12DescriptorHeap* AllocateIndependentHeap(const TCHAR* InDebugName, ERHIDescriptorHeapType InHeapType, uint32 InNumDescriptors, D3D12_DESCRIPTOR_HEAP_FLAGS InHeapFlags);
+	FD3D12DescriptorHeap* AllocateHeap(const TCHAR* InDebugName, ERHIDescriptorHeapType InHeapType, uint32 InNumDescriptors, D3D12_DESCRIPTOR_HEAP_FLAGS InHeapFlags);
+	void DeferredFreeHeap(FD3D12DescriptorHeap* InHeap);
+	void ImmediateFreeHeap(FD3D12DescriptorHeap* InHeap);
 
 	//inline FD3D12DescriptorHeap* GetGlobalHeap() const { return GlobalHeap; }
 
diff --git a/Engine/Source/Runtime/D3D12RHI/Private/D3D12Device.cpp b/Engine/Source/Runtime/D3D12RHI/Private/D3D12Device.cpp
--- a/Engine/Source/Runtime/D3D12RHI/Private/D3D12Device.cpp
+++ b/Engine/Source/Runtime/D3D12RHI/Private/D3D12Device.cpp
@@ -111,7 +111,9 @@ FD3D12Device::FD3D12Device(FRHIGPUMask InGPUMask, FD3D12Adapter* InAdapter)
 	, GPUProfilingData         (this)
 	, ResidencyManager         (*this)
 	, DescriptorHeapManager    (this)
+#if PLATFORM_SUPPORTS_BINDLESS_RENDERING
 	, BindlessDescriptorManager(this)
+#endif
 	, GlobalSamplerHeap        (this)
 	, OnlineDescriptorManager  (this)
 	, DefaultBufferAllocator   (this, FRHIGPUMask::All()) //Note: Cross node buffers are possible 
@@ -358,23 +360,28 @@ void FD3D12Device::SetupAfterDeviceCreation()
 	check(GGlobalSamplerHeapSize <= MaximumSamplerHeapSize);
 
 	check(GOnlineDescriptorHeapSize <= GGlobalResourceDescriptorHeapSize);
-	check(GBindlessResourceDescriptorHeapSize <= GGlobalResourceDescriptorHeapSize);
-	check(GBindlessSamplerDescriptorHeapSize <= GGlobalSamplerDescriptorHeapSize);
 
-	DescriptorHeapManager.Init(GGlobalResourceDescriptorHeapSize, GGlobalSamplerDescriptorHeapSize);
+	bool bFullyBindlessResources = false;
+	bool bFullyBindlessSamplers = false;
 
 #if PLATFORM_SUPPORTS_BINDLESS_RENDERING
-	const bool bBindlessResources = RHIGetRuntimeBindlessResourcesConfiguration(GMaxRHIShaderPlatform) != ERHIBindlessConfiguration::Disabled;
-	const bool bBindlessSamplers = RHIGetRuntimeBindlessSamplersConfiguration(GMaxRHIShaderPlatform) != ERHIBindlessConfiguration::Disabled;
-	if (bBindlessResources || bBindlessSamplers)
-	{
-		BindlessDescriptorManager.Init(bBindlessResources ? GBindlessResourceDescriptorHeapSize : 0, bBindlessSamplers ? GBindlessSamplerDescriptorHeapSize : 0);
-	}
+	BindlessDescriptorManager.Init();
+
+	bFullyBindlessResources = BindlessDescriptorManager.AreResourcesFullyBindless();
+	bFullyBindlessSamplers = BindlessDescriptorManager.AreSamplersFullyBindless();
 #endif
 
-	GlobalSamplerHeap.Init(GGlobalSamplerHeapSize);
+	DescriptorHeapManager.Init(bFullyBindlessResources ? 0 : GGlobalResourceDescriptorHeapSize, bFullyBindlessSamplers ? 0 : GGlobalSamplerDescriptorHeapSize);
 
-	OnlineDescriptorManager.Init(GOnlineDescriptorHeapSize, GOnlineDescriptorHeapBlockSize);
+	if (!bFullyBindlessSamplers)
+	{
+		GlobalSamplerHeap.Init(GGlobalSamplerHeapSize);
+	}
+
+	if (!bFullyBindlessResources)
+	{
+		OnlineDescriptorManager.Init(GOnlineDescriptorHeapSize, GOnlineDescriptorHeapBlockSize);
+	}
 
 	// Make sure we create the default views before the first command context
 	CreateDefaultViews();
diff --git a/Engine/Source/Runtime/D3D12RHI/Private/D3D12Device.h b/Engine/Source/Runtime/D3D12RHI/Private/D3D12Device.h
--- a/Engine/Source/Runtime/D3D12RHI/Private/D3D12Device.h
+++ b/Engine/Source/Runtime/D3D12RHI/Private/D3D12Device.h
@@ -7,6 +7,7 @@ D3D12Device.h: D3D12 Device Interfaces
 #pragma once
 
 #include "CoreMinimal.h"
+#include "D3D12BindlessDescriptors.h"
 #include "D3D12Descriptors.h"
 
 class FD3D12Device;
@@ -284,7 +285,9 @@ public:
 
 	// Descriptor Managers
 	inline FD3D12DescriptorHeapManager&     GetDescriptorHeapManager    () { return DescriptorHeapManager;     }
+#if PLATFORM_SUPPORTS_BINDLESS_RENDERING
 	inline FD3D12BindlessDescriptorManager& GetBindlessDescriptorManager() { return BindlessDescriptorManager; }
+#endif
 	inline FD3D12OnlineDescriptorManager&   GetOnlineDescriptorManager  () { return OnlineDescriptorManager;   }
 	inline FD3D12OfflineDescriptorManager&  GetOfflineDescriptorManager (ERHIDescriptorHeapType InType)
 	{
@@ -349,7 +352,9 @@ private:
 	} ResidencyManager;
 
 	FD3D12DescriptorHeapManager     DescriptorHeapManager;
+#if PLATFORM_SUPPORTS_BINDLESS_RENDERING
 	FD3D12BindlessDescriptorManager BindlessDescriptorManager;
+#endif
 	TArray<FD3D12OfflineDescriptorManager, TInlineAllocator<(uint32)ERHIDescriptorHeapType::Count>> OfflineDescriptorManagers;
 
 	FD3D12GlobalOnlineSamplerHeap GlobalSamplerHeap;
diff --git a/Engine/Source/Runtime/D3D12RHI/Private/D3D12RHI.cpp b/Engine/Source/Runtime/D3D12RHI/Private/D3D12RHI.cpp
--- a/Engine/Source/Runtime/D3D12RHI/Private/D3D12RHI.cpp
+++ b/Engine/Source/Runtime/D3D12RHI/Private/D3D12RHI.cpp
@@ -520,16 +520,21 @@ void FD3D12DynamicRHI::ProcessDeferredDeletionQueue()
 					check(ObjectToDelete.RHIObject->GetRefCount() == 1);
 					ObjectToDelete.RHIObject->Release();
 					break;
-				case FD3D12DeferredDeleteObject::EType::D3DHeap:
+				case FD3D12DeferredDeleteObject::EType::Heap:
 					// Heaps can have additional references active.
-					ObjectToDelete.D3DHeap->Release();
+					ObjectToDelete.Heap->Release();
+					break;
+				case FD3D12DeferredDeleteObject::EType::DescriptorHeap:
+					ObjectToDelete.DescriptorHeap->GetParentDevice()->GetDescriptorHeapManager().ImmediateFreeHeap(ObjectToDelete.DescriptorHeap);
 					break;
 				case FD3D12DeferredDeleteObject::EType::D3DObject:
 					ObjectToDelete.D3DObject->Release();
 					break;
+#if PLATFORM_SUPPORTS_BINDLESS_RENDERING
 				case FD3D12DeferredDeleteObject::EType::BindlessDescriptor:
 					ObjectToDelete.BindlessDescriptor.Device->GetBindlessDescriptorManager().ImmediateFree(ObjectToDelete.BindlessDescriptor.Handle);
 					break;
+#endif
 				case FD3D12DeferredDeleteObject::EType::CPUAllocation:
 					FMemory::Free(ObjectToDelete.CPUAllocation);
 					break;
diff --git a/Engine/Source/Runtime/D3D12RHI/Private/D3D12RayTracing.cpp b/Engine/Source/Runtime/D3D12RHI/Private/D3D12RayTracing.cpp
--- a/Engine/Source/Runtime/D3D12RHI/Private/D3D12RayTracing.cpp
+++ b/Engine/Source/Runtime/D3D12RHI/Private/D3D12RayTracing.cpp
@@ -1599,8 +1599,8 @@ public:
 	{
 		FD3D12BindlessDescriptorManager& BindlessManager = GetParentDevice()->GetBindlessDescriptorManager();
 
-		bBindlessViews = BindlessManager.HasHeapForType(ERHIDescriptorHeapType::Standard);
-		bBindlessSamplers = BindlessManager.HasHeapForType(ERHIDescriptorHeapType::Sampler);
+		bBindlessViews = BindlessManager.HasHeap(ERHIDescriptorHeapType::Standard, ERHIBindlessConfiguration::RayTracingShaders);
+		bBindlessSamplers = BindlessManager.HasHeap(ERHIDescriptorHeapType::Sampler, ERHIBindlessConfiguration::RayTracingShaders);
 
 		if (!bBindlessViews)
 		{
@@ -1637,8 +1637,8 @@ public:
 
 		FD3D12BindlessDescriptorManager& BindlessManager = GetParentDevice()->GetBindlessDescriptorManager();
 
-		ID3D12DescriptorHeap* ViewHeapToSet = bBindlessViews ? BindlessManager.GetHeapForType(ERHIDescriptorHeapType::Standard)->GetHeap() : ViewHeap.D3D12Heap;
-		ID3D12DescriptorHeap* SamplerHeapToSet = bBindlessSamplers ? BindlessManager.GetHeapForType(ERHIDescriptorHeapType::Sampler)->GetHeap() : SamplerHeap.D3D12Heap;
+		ID3D12DescriptorHeap* ViewHeapToSet = bBindlessViews ? BindlessManager.GetHeap(ERHIDescriptorHeapType::Standard)->GetHeap() : ViewHeap.D3D12Heap;
+		ID3D12DescriptorHeap* SamplerHeapToSet = bBindlessSamplers ? BindlessManager.GetHeap(ERHIDescriptorHeapType::Sampler)->GetHeap() : SamplerHeap.D3D12Heap;
 
 		CommandContext.StateCache.GetDescriptorCache()->OverrideLastSetHeaps(ViewHeapToSet, SamplerHeapToSet);
 	}
diff --git a/Engine/Source/Runtime/D3D12RHI/Private/D3D12RootSignature.cpp b/Engine/Source/Runtime/D3D12RHI/Private/D3D12RootSignature.cpp
--- a/Engine/Source/Runtime/D3D12RHI/Private/D3D12RootSignature.cpp
+++ b/Engine/Source/Runtime/D3D12RHI/Private/D3D12RootSignature.cpp
@@ -485,6 +485,7 @@ void FD3D12RootSignature::InternalAnalyzeSignature(const RootSignatureDescType&
 		bHasUAVs = false;
 		bHasSRVs = false;
 		bHasCBVs = false;
+		bHasRootCBs = false;
 		bHasSamplers = false;
 
 		FMemory::Memset(BindSlotOffsetsInDWORDs, 0, sizeof(BindSlotOffsetsInDWORDs));
diff --git a/Engine/Source/Runtime/D3D12RHI/Private/D3D12RootSignature.h b/Engine/Source/Runtime/D3D12RHI/Private/D3D12RootSignature.h
--- a/Engine/Source/Runtime/D3D12RHI/Private/D3D12RootSignature.h
+++ b/Engine/Source/Runtime/D3D12RHI/Private/D3D12RootSignature.h
@@ -178,7 +178,8 @@ public:
 	inline bool HasUAVs() const { return bHasUAVs; }
 	inline bool HasSRVs() const { return bHasSRVs; }
 	inline bool HasCBVs() const { return bHasCBVs; }
-	inline bool HasResources() const { return bHasUAVs || bHasSRVs || bHasCBVs; }
+	inline bool HasRootCBs() const { return bHasRootCBs; }
+	inline bool HasTableResources() const { return bHasUAVs || bHasSRVs || bHasCBVs; }
 	inline bool HasSamplers() const { return bHasSamplers; }
 	inline bool UsesDynamicResources() const { return bUsesDynamicResources; }
 	inline bool UsesDynamicSamplers() const { return bUsesDynamicSamplers; }
@@ -305,7 +306,7 @@ private:
 		check(*pBindSlot == 0xFF);
 		*pBindSlot = RootParameterIndex;
 
-		bHasCBVs = true;
+		bHasRootCBs = true;
 	}
 
 	inline void SetUAVRDTBindSlot(EShaderFrequency SF, uint8 RootParameterIndex)
@@ -446,11 +447,11 @@ private:
 	uint8 bHasUAVs : 1;
 	uint8 bHasSRVs : 1;
 	uint8 bHasCBVs : 1;
+	uint8 bHasRootCBs : 1;
 	uint8 bHasSamplers : 1;
 	uint8 bHasVendorExtensionSpace : 1;
 	uint8 bUsesDynamicResources : 1;
 	uint8 bUsesDynamicSamplers : 1;
-	uint8 bUnused : 1;
 };
 
 class FD3D12RootSignatureManager : public FD3D12AdapterChild
diff --git a/Engine/Source/Runtime/D3D12RHI/Private/D3D12State.cpp b/Engine/Source/Runtime/D3D12RHI/Private/D3D12State.cpp
--- a/Engine/Source/Runtime/D3D12RHI/Private/D3D12State.cpp
+++ b/Engine/Source/Runtime/D3D12RHI/Private/D3D12State.cpp
@@ -640,6 +640,7 @@ FD3D12SamplerState::FD3D12SamplerState(FD3D12Device* InParent, const D3D12_SAMPL
 
 	GetParentDevice()->CreateSamplerInternal(Desc, OfflineDescriptor);
 
+#if PLATFORM_SUPPORTS_BINDLESS_RENDERING
 	FD3D12BindlessDescriptorManager& BindlessDescriptorManager = GetParentDevice()->GetBindlessDescriptorManager();
 	BindlessHandle = BindlessDescriptorManager.Allocate(ERHIDescriptorHeapType::Sampler);
 
@@ -647,6 +648,7 @@ FD3D12SamplerState::FD3D12SamplerState(FD3D12Device* InParent, const D3D12_SAMPL
 	{
 		GetParentDevice()->GetBindlessDescriptorManager().UpdateImmediately(BindlessHandle, OfflineDescriptor);
 	}
+#endif
 }
 
 FD3D12SamplerState::~FD3D12SamplerState()
@@ -656,9 +658,11 @@ FD3D12SamplerState::~FD3D12SamplerState()
 		FD3D12OfflineDescriptorManager& OfflineAllocator = GetParentDevice()->GetOfflineDescriptorManager(ERHIDescriptorHeapType::Sampler);
 		OfflineAllocator.FreeHeapSlot(OfflineDescriptor);
 
+#if PLATFORM_SUPPORTS_BINDLESS_RENDERING
 		if (BindlessHandle.IsValid())
 		{
 			GetParentDevice()->GetBindlessDescriptorManager().DeferredFreeFromDestructor(BindlessHandle);
 		}
+#endif
 	}
 }
diff --git a/Engine/Source/Runtime/D3D12RHI/Private/D3D12StateCache.cpp b/Engine/Source/Runtime/D3D12RHI/Private/D3D12StateCache.cpp
--- a/Engine/Source/Runtime/D3D12RHI/Private/D3D12StateCache.cpp
+++ b/Engine/Source/Runtime/D3D12RHI/Private/D3D12StateCache.cpp
@@ -23,22 +23,6 @@ static FAutoConsoleVariableRef CVarGlobalSamplerDescriptorHeapSize(
 	ECVF_ReadOnly
 );
 
-int32 GBindlessResourceDescriptorHeapSize = 1000 * 1000;
-static FAutoConsoleVariableRef CVarBindlessResourceDescriptorHeapSize(
-	TEXT("D3D12.Bindless.ResourceDescriptorHeapSize"),
-	GBindlessResourceDescriptorHeapSize,
-	TEXT("Bindless resource descriptor heap size"),
-	ECVF_ReadOnly
-);
-
-int32 GBindlessSamplerDescriptorHeapSize = 2048;
-static FAutoConsoleVariableRef CVarBindlessSamplerDescriptorHeapSize(
-	TEXT("D3D12.Bindless.SamplerDescriptorHeapSize"),
-	GBindlessSamplerDescriptorHeapSize,
-	TEXT("Bindless sampler descriptor heap size"),
-	ECVF_ReadOnly
-);
-
 // This value defines how many descriptors will be in the device local view heap which
 // This should be tweaked for each title as heaps require VRAM. The default value of 512k takes up ~16MB
 int32 GLocalViewHeapSize = 500 * 1000;
@@ -426,27 +410,26 @@ void FD3D12StateCache::ApplyState(ED3D12PipelineType PipelineType)
 	bool bBindlessResources = PSOCommonData->RootSignature->UsesDynamicResources();
 	bool bBindlessSamplers = PSOCommonData->RootSignature->UsesDynamicSamplers();
 
-	const bool bApplyResources = !bBindlessResources && PSOCommonData->RootSignature->HasResources();
-	const bool bApplySamplers = !bBindlessSamplers && PSOCommonData->RootSignature->HasSamplers();
+	const bool bHasTableResources = PSOCommonData->RootSignature->HasTableResources();
+	const bool bHasSamplers = PSOCommonData->RootSignature->HasSamplers();
 
-	const bool bBindlessResourcesAlreadyEnabled = DescriptorCache.IsViewHeapOverridden();
-	const bool bBindlessSamplersAlreadyEnabled = DescriptorCache.IsSamplerHeapOverridden();
+	const bool bApplyResources = !bBindlessResources && bHasTableResources;
+	const bool bApplySamplers = !bBindlessSamplers && bHasSamplers;
 
-	if (bRootSignatureChanged || bBindlessResources != bBindlessResourcesAlreadyEnabled || bBindlessSamplers != bBindlessSamplersAlreadyEnabled)
+#if PLATFORM_SUPPORTS_BINDLESS_RENDERING
 	{
 		FD3D12BindlessDescriptorManager& BindlessManager = GetParentDevice()->GetBindlessDescriptorManager();
 
-		FD3D12DescriptorHeap* ResourceHeap = BindlessManager.GetHeapForType(ERHIDescriptorHeapType::Standard);
-		FD3D12DescriptorHeap* SamplerHeap = BindlessManager.GetHeapForType(ERHIDescriptorHeapType::Sampler);
+		FD3D12DescriptorHeap* ResourceHeap = BindlessManager.GetHeap(ERHIDescriptorHeapType::Standard, ERHIBindlessConfiguration::AllShaders);
+		FD3D12DescriptorHeap* SamplerHeap = BindlessManager.GetHeap(ERHIDescriptorHeapType::Sampler, ERHIBindlessConfiguration::AllShaders);
 
 		checkf(!bBindlessResources || ResourceHeap != nullptr, TEXT("Using dynamic samplers without the bindless sampler heap configured. Please check your configuration."));
 		checkf(!bBindlessSamplers || SamplerHeap != nullptr, TEXT("Using dynamic samplers without the bindless sampler heap configured. Please check your configuration."));
-		
-		DescriptorCache.SetHeapOverrides(
-			bBindlessResources ? BindlessManager.GetHeapForType(ERHIDescriptorHeapType::Standard) : nullptr,
-			bBindlessSamplers ? BindlessManager.GetHeapForType(ERHIDescriptorHeapType::Sampler) : nullptr
-		);
+
+		check(!(ResourceHeap != nullptr && bHasTableResources));
+		check(!(SamplerHeap  != nullptr && bHasSamplers      ));
 	}
+#endif
 
 	if (bRootSignatureChanged)
 	{
diff --git a/Engine/Source/Runtime/D3D12RHI/Private/D3D12StateCachePrivate.h b/Engine/Source/Runtime/D3D12RHI/Private/D3D12StateCachePrivate.h
--- a/Engine/Source/Runtime/D3D12RHI/Private/D3D12StateCachePrivate.h
+++ b/Engine/Source/Runtime/D3D12RHI/Private/D3D12StateCachePrivate.h
@@ -32,9 +32,6 @@ static const bool GD3D12SkipStateCaching = false;
 extern int32 GGlobalResourceDescriptorHeapSize;
 extern int32 GGlobalSamplerDescriptorHeapSize;
 
-extern int32 GBindlessResourceDescriptorHeapSize;
-extern int32 GBindlessSamplerDescriptorHeapSize;
-
 extern int32 GGlobalSamplerHeapSize;
 extern int32 GOnlineDescriptorHeapSize;
 extern int32 GOnlineDescriptorHeapBlockSize;
diff --git a/Engine/Source/Runtime/D3D12RHI/Private/D3D12Texture.cpp b/Engine/Source/Runtime/D3D12RHI/Private/D3D12Texture.cpp
--- a/Engine/Source/Runtime/D3D12RHI/Private/D3D12Texture.cpp
+++ b/Engine/Source/Runtime/D3D12RHI/Private/D3D12Texture.cpp
@@ -1018,15 +1018,17 @@ void FD3D12DynamicRHI::RHIUpdateTextureReference(FRHITextureReference* TextureRe
 {
 	FRHITexture* NewTexture = InNewTexture ? InNewTexture : FRHITextureReference::GetDefaultTexture();
 
+#if PLATFORM_SUPPORTS_BINDLESS_RENDERING
 	if (FRHIShaderResourceView* TextureRefSRV = TextureRef ? TextureRef->GetBindlessView() : nullptr)
 	{
 		FD3D12ShaderResourceView* TextureRefSRVD3D12 = ResourceCast(TextureRefSRV);
-		if (TextureRefSRVD3D12->GetBindlessHandle().IsValid())
+		if (TextureRefSRVD3D12->IsBindless())
 		{
 			FD3D12BindlessDescriptorManager& BindlessDescriptorManager = TextureRefSRVD3D12->GetParentDevice()->GetBindlessDescriptorManager();
 			BindlessDescriptorManager.UpdateDeferred(TextureRefSRVD3D12->GetBindlessHandle(), ResourceCast(NewTexture)->GetShaderResourceView()->GetOfflineCpuHandle());
 		}
 	}
+#endif // PLATFORM_SUPPORTS_BINDLESS_RENDERING
 
 	FDynamicRHI::RHIUpdateTextureReference(TextureRef, NewTexture);
 }
diff --git a/Engine/Source/Runtime/D3D12RHI/Private/D3D12UAV.cpp b/Engine/Source/Runtime/D3D12RHI/Private/D3D12UAV.cpp
--- a/Engine/Source/Runtime/D3D12RHI/Private/D3D12UAV.cpp
+++ b/Engine/Source/Runtime/D3D12RHI/Private/D3D12UAV.cpp
@@ -316,8 +316,8 @@ void FD3D12CommandContext::ClearUAV(TRHICommandList_RecursiveHazardous<FD3D12Com
 				FD3D12UnorderedAccessView UAV(ParentDevice);
 				UAV.CreateView(UnorderedAccessView->GetResourceLocation(), R32UAVDesc, FD3D12UnorderedAccessView::EFlags::None);
 
-				D3D12_GPU_DESCRIPTOR_HANDLE GPUHandle;
-				if (!Context.StateCache.GetDescriptorCache()->IsViewHeapOverridden())
+				D3D12_GPU_DESCRIPTOR_HANDLE GPUHandle{};
+				if (!Context.StateCache.GetDescriptorCache()->IsUsingBindlessResources())
 				{
 					// Check if the view heap is full and needs to rollover.
 					if (!Context.StateCache.GetDescriptorCache()->GetCurrentViewHeap()->CanReserveSlots(1))
@@ -331,10 +331,12 @@ void FD3D12CommandContext::ClearUAV(TRHICommandList_RecursiveHazardous<FD3D12Com
 
 					Device->CopyDescriptorsSimple(1, DestSlot, UAV.GetOfflineCpuHandle(), D3D12_DESCRIPTOR_HEAP_TYPE_CBV_SRV_UAV);
 				}
+#if PLATFORM_SUPPORTS_BINDLESS_RENDERING
 				else
 				{
 					GPUHandle = ParentDevice->GetBindlessDescriptorManager().GetGpuHandle(UAV.GetBindlessHandle());
 				}
+#endif
 
 				Context.TransitionResource(UnorderedAccessView, D3D12_RESOURCE_STATE_UNORDERED_ACCESS);
 
diff --git a/Engine/Source/Runtime/D3D12RHI/Private/D3D12View.cpp b/Engine/Source/Runtime/D3D12RHI/Private/D3D12View.cpp
--- a/Engine/Source/Runtime/D3D12RHI/Private/D3D12View.cpp
+++ b/Engine/Source/Runtime/D3D12RHI/Private/D3D12View.cpp
@@ -217,10 +217,12 @@ FD3D12View::FD3D12View(FD3D12Device* InDevice, ERHIDescriptorHeapType InHeapType
 	, OfflineCpuHandle(InDevice->GetOfflineDescriptorManager(InHeapType).AllocateHeapSlot())
 	, HeapType(InHeapType)
 {
+#if PLATFORM_SUPPORTS_BINDLESS_RENDERING
 	if (HeapType == ERHIDescriptorHeapType::Standard)
 	{
 		BindlessHandle = InDevice->GetBindlessDescriptorManager().Allocate(ERHIDescriptorHeapType::Standard);
 	}
+#endif
 }
 
 FD3D12View::~FD3D12View()
@@ -232,15 +234,18 @@ FD3D12View::~FD3D12View()
 	// Free the descriptor heap slot and bindless handle
 	GetParentDevice()->GetOfflineDescriptorManager(HeapType).FreeHeapSlot(OfflineCpuHandle);
 
+#if PLATFORM_SUPPORTS_BINDLESS_RENDERING
 	if (BindlessHandle.IsValid())
 	{
 		GetParentDevice()->GetBindlessDescriptorManager().DeferredFreeFromDestructor(BindlessHandle);
 		BindlessHandle = {};
 	}
+#endif
 }
 
 void FD3D12View::UpdateBindlessSlot(EReason Reason)
 {
+#if PLATFORM_SUPPORTS_BINDLESS_RENDERING
 	if (BindlessHandle.IsValid())
 	{
 		FD3D12BindlessDescriptorManager& BindlessManager = GetParentDevice()->GetBindlessDescriptorManager();
@@ -253,6 +258,7 @@ void FD3D12View::UpdateBindlessSlot(EReason Reason)
 			BindlessManager.UpdateDeferred(BindlessHandle, OfflineCpuHandle);
 		}
 	}
+#endif
 }
 
 void FD3D12View::CreateView(FResourceInfo const& InResource, FNullDescPtr NullDescriptor)
diff --git a/Engine/Source/Runtime/D3D12RHI/Private/D3D12View.h b/Engine/Source/Runtime/D3D12RHI/Private/D3D12View.h
--- a/Engine/Source/Runtime/D3D12RHI/Private/D3D12View.h
+++ b/Engine/Source/Runtime/D3D12RHI/Private/D3D12View.h
@@ -210,9 +210,14 @@ public:
 	FD3D12ResidencyHandle&      GetResidencyHandle () const { check(IsInitialized()); return *ResourceInfo.ResidencyHandle; }
 	FD3D12ViewSubset const&     GetViewSubset      () const { check(IsInitialized()); return ViewSubset;                    }
 	FD3D12OfflineDescriptor     GetOfflineCpuHandle() const { check(IsInitialized()); return OfflineCpuHandle;              }
-	FRHIDescriptorHandle        GetBindlessHandle  () const { /* no check */          return BindlessHandle;                }
 
-	bool IsBindless() const { return BindlessHandle.IsValid(); }
+#if PLATFORM_SUPPORTS_BINDLESS_RENDERING
+	FRHIDescriptorHandle        GetBindlessHandle() const { return BindlessHandle;           }
+	bool                        IsBindless       () const { return BindlessHandle.IsValid(); }
+#else
+	FRHIDescriptorHandle        GetBindlessHandle() const { return FRHIDescriptorHandle();   }
+	constexpr bool              IsBindless       () const { return false;                    }
+#endif
 
 protected:
 	FD3D12View(FD3D12Device* InDevice, ERHIDescriptorHeapType InHeapType);
@@ -231,7 +236,9 @@ protected:
 
 	FD3D12OfflineDescriptor OfflineCpuHandle;
 
+#if PLATFORM_SUPPORTS_BINDLESS_RENDERING
 	FRHIDescriptorHandle BindlessHandle;
+#endif
 	ERHIDescriptorHeapType const HeapType;
 };
 
