Fix D3D12 UpdateBuffer incorrect use of default command context
 - Causing stomps etc when the UpdateBuffer command was enqueued on a parallel RHICmdList.

diff --git a/Engine/Source/Runtime/D3D12RHI/Private/D3D12Buffer.cpp b/Engine/Source/Runtime/D3D12RHI/Private/D3D12Buffer.cpp
--- a/Engine/Source/Runtime/D3D12RHI/Private/D3D12Buffer.cpp
+++ b/Engine/Source/Runtime/D3D12RHI/Private/D3D12Buffer.cpp
@@ -52,7 +52,7 @@ struct FRHICommandUpdateBuffer final : public FRHICommand<FRHICommandUpdateBuffe
 
 	void Execute(FRHICommandListBase& CmdList)
 	{
-		FD3D12DynamicRHI::GetD3DRHI()->UpdateBuffer(Destination, DestinationOffset, &Source, 0, NumBytes);
+		FD3D12CommandContextBase::Get(CmdList).UpdateBuffer(Destination, DestinationOffset, &Source, 0, NumBytes);
 	}
 };
 
@@ -744,7 +744,7 @@ void FD3D12DynamicRHI::UnlockBuffer(FRHICommandListBase& RHICmdList, FD3D12Buffe
 				}
 				else
 				{
-					UpdateBuffer(&CurrentBuffer->ResourceLocation, LockedData.LockedOffset, &LockedData.ResourceLocation, 0, LockedData.LockedPitch);
+					FD3D12CommandContextBase::Get(RHICmdList).UpdateBuffer(&CurrentBuffer->ResourceLocation, LockedData.LockedOffset, &LockedData.ResourceLocation, 0, LockedData.LockedPitch);
 				}
 			}
 		}
diff --git a/Engine/Source/Runtime/D3D12RHI/Private/D3D12CommandContext.h b/Engine/Source/Runtime/D3D12RHI/Private/D3D12CommandContext.h
--- a/Engine/Source/Runtime/D3D12RHI/Private/D3D12CommandContext.h
+++ b/Engine/Source/Runtime/D3D12RHI/Private/D3D12CommandContext.h
@@ -439,6 +439,8 @@ public:
 		return static_cast<FD3D12CommandContextBase&>(RHICmdList.GetComputeContext().GetLowestLevelContext());
 	}
 
+	virtual void UpdateBuffer(FD3D12ResourceLocation* Dest, uint32 DestOffset, FD3D12ResourceLocation* Source, uint32 SourceOffset, uint32 NumBytes) = 0;
+
 protected:
 	friend class FD3D12CommandContext;
 	virtual FD3D12CommandContext* GetContext(uint32 InGPUIndex) = 0;
@@ -725,6 +727,8 @@ public:
 		ensure(InGPUMask == GPUMask);
 	}
 
+	virtual void UpdateBuffer(FD3D12ResourceLocation* Dest, uint32 DestOffset, FD3D12ResourceLocation* Source, uint32 SourceOffset, uint32 NumBytes) final override;
+
 protected:
 
 	FD3D12CommandContext* GetContext(uint32 InGPUIndex) final override 
@@ -1006,6 +1010,11 @@ public:
 		ContextRedirect(RHIEndRenderPass());
 	}
 
+	virtual void UpdateBuffer(FD3D12ResourceLocation* Dest, uint32 DestOffset, FD3D12ResourceLocation* Source, uint32 SourceOffset, uint32 NumBytes) final override
+	{
+		ContextRedirect(UpdateBuffer(Dest, DestOffset, Source, SourceOffset, NumBytes));
+	}
+
 #if D3D12_RHI_RAYTRACING
 	virtual void RHIBuildAccelerationStructures(const TArrayView<const FRayTracingGeometryBuildParams> Params, const FRHIBufferRange& ScratchBufferRange) final override
 	{
diff --git a/Engine/Source/Runtime/D3D12RHI/Private/D3D12Commands.cpp b/Engine/Source/Runtime/D3D12RHI/Private/D3D12Commands.cpp
--- a/Engine/Source/Runtime/D3D12RHI/Private/D3D12Commands.cpp
+++ b/Engine/Source/Runtime/D3D12RHI/Private/D3D12Commands.cpp
@@ -2043,3 +2043,27 @@ void FD3D12CommandContext::RHISubmitCommandsHint()
 {
 	// Nothing to do
 }
+
+void FD3D12CommandContext::UpdateBuffer(FD3D12ResourceLocation* Dest, uint32 DestOffset, FD3D12ResourceLocation* Source, uint32 SourceOffset, uint32 NumBytes)
+{
+	FD3D12Resource* SourceResource = Source->GetResource();
+	uint32 SourceFullOffset = Source->GetOffsetFromBaseOfResource() + SourceOffset;
+
+	FD3D12Resource* DestResource = Dest->GetResource();
+	uint32 DestFullOffset = Dest->GetOffsetFromBaseOfResource() + DestOffset;
+
+	// Clear the resource if still bound to make sure the SRVs are rebound again on next operation (and get correct resource transitions enqueued)
+	ConditionalClearShaderResource(Dest);
+
+	FScopedResourceBarrier ScopeResourceBarrierDest(*this, DestResource, D3D12_RESOURCE_STATE_COPY_DEST, 0);
+	// Don't need to transition upload heaps
+
+	FlushResourceBarriers();
+	GraphicsCommandList()->CopyBufferRegion(DestResource->GetResource(), DestFullOffset, SourceResource->GetResource(), SourceFullOffset, NumBytes);
+	UpdateResidency(DestResource);
+	UpdateResidency(SourceResource);
+
+	ConditionalSplitCommandList();
+
+	DEBUG_RHI_EXECUTE_COMMAND_LIST(this);
+}
diff --git a/Engine/Source/Runtime/D3D12RHI/Private/D3D12RHI.cpp b/Engine/Source/Runtime/D3D12RHI/Private/D3D12RHI.cpp
--- a/Engine/Source/Runtime/D3D12RHI/Private/D3D12RHI.cpp
+++ b/Engine/Source/Runtime/D3D12RHI/Private/D3D12RHI.cpp
@@ -400,34 +400,6 @@ IRHIComputeContext* FD3D12DynamicRHI::RHIGetDefaultAsyncComputeContext()
 	return nullptr;
 }
 
-void FD3D12DynamicRHI::UpdateBuffer(FD3D12ResourceLocation* Dest, uint32 DestOffset, FD3D12ResourceLocation* Source, uint32 SourceOffset, uint32 NumBytes)
-{
-	FD3D12Resource* SourceResource = Source->GetResource();
-	uint32 SourceFullOffset = Source->GetOffsetFromBaseOfResource() + SourceOffset;
-
-	FD3D12Resource* DestResource = Dest->GetResource();
-	uint32 DestFullOffset = Dest->GetOffsetFromBaseOfResource() + DestOffset;
-
-	FD3D12Device* Device = DestResource->GetParentDevice();
-
-	FD3D12CommandContext& DefaultContext = Device->GetDefaultCommandContext();
-	
-	// Clear the resource if still bound to make sure the SRVs are rebound again on next operation (and get correct resource transitions enqueued)
-	DefaultContext.ConditionalClearShaderResource(Dest);
-
-	FScopedResourceBarrier ScopeResourceBarrierDest(DefaultContext, DestResource, D3D12_RESOURCE_STATE_COPY_DEST, 0);
-	// Don't need to transition upload heaps
-
-	DefaultContext.FlushResourceBarriers();
-	DefaultContext.GraphicsCommandList()->CopyBufferRegion(DestResource->GetResource(), DestFullOffset, SourceResource->GetResource(), SourceFullOffset, NumBytes);
-	DefaultContext.UpdateResidency(DestResource);
-	DefaultContext.UpdateResidency(SourceResource);
-	
-	DefaultContext.ConditionalSplitCommandList();
-
-	DEBUG_RHI_EXECUTE_COMMAND_LIST(this);
-}
-
 void FD3D12DynamicRHI::RHIFlushResources()
 {
 	// Nothing to do (yet!)
diff --git a/Engine/Source/Runtime/D3D12RHI/Private/D3D12RHIPrivate.h b/Engine/Source/Runtime/D3D12RHI/Private/D3D12RHIPrivate.h
--- a/Engine/Source/Runtime/D3D12RHI/Private/D3D12RHIPrivate.h
+++ b/Engine/Source/Runtime/D3D12RHI/Private/D3D12RHIPrivate.h
@@ -476,8 +476,6 @@ public:
 	FUpdateTexture3DData BeginUpdateTexture3D_Internal(FRHITexture3D* Texture, uint32 MipIndex, const struct FUpdateTextureRegion3D& UpdateRegion);
 	void EndUpdateTexture3D_Internal(FRHICommandListBase& RHICmdList, FUpdateTexture3DData& UpdateData);
 
-	void UpdateBuffer(FD3D12ResourceLocation* Dest, uint32 DestOffset, FD3D12ResourceLocation* Source, uint32 SourceOffset, uint32 NumBytes);
-
 public:
 
 #if	PLATFORM_SUPPORTS_VIRTUAL_TEXTURES
