Adding interfaces for VerifyD3D11/12Result. These should go through the RHI interfaces to ensure the user is running the same RHI they're calling.

diff --git a/Engine/Source/Runtime/D3D12RHI/Private/D3D12RHI.cpp b/Engine/Source/Runtime/D3D12RHI/Private/D3D12RHI.cpp
--- a/Engine/Source/Runtime/D3D12RHI/Private/D3D12RHI.cpp
+++ b/Engine/Source/Runtime/D3D12RHI/Private/D3D12RHI.cpp
@@ -691,6 +691,11 @@ void FD3D12DynamicRHI::RHIWaitManualFence(FRHICommandList& RHICmdList, ID3D12Fen
 	Context.WaitManualFence(Fence, Value);
 }
 
+void FD3D12DynamicRHI::RHIVerifyResult(ID3D12Device* Device, HRESULT Result, const ANSICHAR* Code, const ANSICHAR* Filename, uint32 Line, FString Message) const
+{
+	D3D12RHI::VerifyD3D12Result(Result, Code, Filename, Line, Device, Message);
+}
+
 void* FD3D12DynamicRHI::RHIGetNativeDevice()
 {
 	return (void*)GetAdapter().GetD3DDevice();
diff --git a/Engine/Source/Runtime/D3D12RHI/Private/D3D12RHIPrivate.h b/Engine/Source/Runtime/D3D12RHI/Private/D3D12RHIPrivate.h
--- a/Engine/Source/Runtime/D3D12RHI/Private/D3D12RHIPrivate.h
+++ b/Engine/Source/Runtime/D3D12RHI/Private/D3D12RHIPrivate.h
@@ -410,6 +410,7 @@ public:
 	virtual void RHITransitionResource(FRHICommandList& RHICmdList, FRHITexture* InTexture, D3D12_RESOURCE_STATES InState, uint32 InSubResource) final override;
 	virtual void RHISignalManualFence(FRHICommandList& RHICmdList, ID3D12Fence* Fence, uint64 Value) final override;
 	virtual void RHIWaitManualFence(FRHICommandList& RHICmdList, ID3D12Fence* Fence, uint64 Value) final override;
+	virtual void RHIVerifyResult(ID3D12Device* Device, HRESULT Result, const ANSICHAR* Code, const ANSICHAR* Filename, uint32 Line, FString Message = FString()) const final override;
 
 	//
 	// The Following functions are the _RenderThread version of the above functions. They allow the RHI to control the thread synchronization for greater efficiency.
diff --git a/Engine/Source/Runtime/D3D12RHI/Public/ID3D12DynamicRHI.h b/Engine/Source/Runtime/D3D12RHI/Public/ID3D12DynamicRHI.h
--- a/Engine/Source/Runtime/D3D12RHI/Public/ID3D12DynamicRHI.h
+++ b/Engine/Source/Runtime/D3D12RHI/Public/ID3D12DynamicRHI.h
@@ -43,8 +43,10 @@ struct ID3D12DynamicRHI : public FDynamicRHI
 	virtual void                  RHIFinishExternalComputeWork(uint32 InDeviceIndex, ID3D12GraphicsCommandList* InCommandList) = 0;
 	virtual void                  RHITransitionResource(FRHICommandList& RHICmdList, FRHITexture* InTexture, D3D12_RESOURCE_STATES InState, uint32 InSubResource) = 0;
 
-	virtual void						RHISignalManualFence(FRHICommandList& RHICmdList, ID3D12Fence* Fence, uint64 Value) = 0;
-	virtual void						RHIWaitManualFence(FRHICommandList& RHICmdList, ID3D12Fence* Fence, uint64 Value) = 0;
+	virtual void                  RHISignalManualFence(FRHICommandList& RHICmdList, ID3D12Fence* Fence, uint64 Value) = 0;
+	virtual void                  RHIWaitManualFence(FRHICommandList& RHICmdList, ID3D12Fence* Fence, uint64 Value) = 0;
+
+	virtual void                  RHIVerifyResult(ID3D12Device* Device, HRESULT Result, const ANSICHAR* Code, const ANSICHAR* Filename, uint32 Line, FString Message = FString()) const = 0;
 };
 
 inline bool IsRHID3D12()
diff --git a/Engine/Source/Runtime/Windows/D3D11RHI/Private/D3D11RHI.cpp b/Engine/Source/Runtime/Windows/D3D11RHI/Private/D3D11RHI.cpp
--- a/Engine/Source/Runtime/Windows/D3D11RHI/Private/D3D11RHI.cpp
+++ b/Engine/Source/Runtime/Windows/D3D11RHI/Private/D3D11RHI.cpp
@@ -805,3 +805,8 @@ void FD3D11DynamicRHI::RHIRegisterWork(uint32 NumPrimitives)
 {
 	RegisterGPUWork(NumPrimitives);
 }
+
+void FD3D11DynamicRHI::RHIVerifyResult(ID3D11Device* Device, HRESULT Result, const ANSICHAR* Code, const ANSICHAR* Filename, uint32 Line) const
+{
+	VerifyD3D11Result(Result, Code, Filename, Line, Device);
+}
diff --git a/Engine/Source/Runtime/Windows/D3D11RHI/Private/D3D11RHIPrivate.h b/Engine/Source/Runtime/Windows/D3D11RHI/Private/D3D11RHIPrivate.h
--- a/Engine/Source/Runtime/Windows/D3D11RHI/Private/D3D11RHIPrivate.h
+++ b/Engine/Source/Runtime/Windows/D3D11RHI/Private/D3D11RHIPrivate.h
@@ -667,6 +667,8 @@ public:
 	virtual ID3D11RenderTargetView* RHIGetRenderTargetView(FRHITexture* InTexture, int32 InMipIndex = 0, int32 InArraySliceIndex = -1) const final override;
 	virtual ID3D11ShaderResourceView* RHIGetShaderResourceView(FRHITexture* InTexture) const final override;
 	virtual void                  RHIRegisterWork(uint32 NumPrimitives) final override;
+	virtual void                  RHIVerifyResult(ID3D11Device* Device, HRESULT Result, const ANSICHAR* Code, const ANSICHAR* Filename, uint32 Line) const final override;
+
 	virtual void				  RHIGetDisplaysInformation(FDisplayInformationArray& OutDisplayInformation) final override;
 
 	// Accessors.
diff --git a/Engine/Source/Runtime/Windows/D3D11RHI/Public/ID3D11DynamicRHI.h b/Engine/Source/Runtime/Windows/D3D11RHI/Public/ID3D11DynamicRHI.h
--- a/Engine/Source/Runtime/Windows/D3D11RHI/Public/ID3D11DynamicRHI.h
+++ b/Engine/Source/Runtime/Windows/D3D11RHI/Public/ID3D11DynamicRHI.h
@@ -34,6 +34,8 @@ struct ID3D11DynamicRHI : public FDynamicRHIPSOFallback
 	virtual ID3D11ShaderResourceView* RHIGetShaderResourceView(FRHITexture* InTexture) const = 0;
 
 	virtual void                  RHIRegisterWork(uint32 NumPrimitives) = 0;
+
+	virtual void                  RHIVerifyResult(ID3D11Device* Device, HRESULT Result, const ANSICHAR* Code, const ANSICHAR* Filename, uint32 Line) const = 0;
 };
 
 inline bool IsRHID3D11()
