Adding an RHICommandList argument to FDynamicRHI::RHIUpdateTextureReference. Adding this first so each implementation can use the arg independently.

diff --git a/Engine/Source/Runtime/D3D12RHI/Private/D3D12RHIPrivate.h b/Engine/Source/Runtime/D3D12RHI/Private/D3D12RHIPrivate.h
--- a/Engine/Source/Runtime/D3D12RHI/Private/D3D12RHIPrivate.h
+++ b/Engine/Source/Runtime/D3D12RHI/Private/D3D12RHIPrivate.h
@@ -304,7 +304,7 @@ public:
 	virtual void* RHILockBufferMGPU(FRHICommandListBase& RHICmdList, FRHIBuffer* Buffer, uint32 GPUIndex, uint32 Offset, uint32 Size, EResourceLockMode LockMode) final override;
 	virtual void RHIUnlockBuffer(FRHICommandListBase& RHICmdList, FRHIBuffer* Buffer) final override;
 	virtual void RHIUnlockBufferMGPU(FRHICommandListBase& RHICmdList, FRHIBuffer* Buffer, uint32 GPUIndex) final override;
-	virtual void RHIUpdateTextureReference(FRHITextureReference* TextureRef, FRHITexture* NewTexture) final override;
+	virtual void RHIUpdateTextureReference(FRHICommandListBase& RHICmdList, FRHITextureReference* TextureRef, FRHITexture* NewTexture) final override;
 	virtual FRHICalcTextureSizeResult RHICalcTexturePlatformSize(const FRHITextureDesc& Desc, uint32 FirstMipIndex) override;
 	virtual void RHIGetTextureMemoryStats(FTextureMemoryStats& OutStats) final override;
 	virtual bool RHIGetTextureMemoryVisualizeData(FColor* TextureData, int32 SizeX, int32 SizeY, int32 Pitch, int32 PixelSize) final override;
diff --git a/Engine/Source/Runtime/D3D12RHI/Private/D3D12Texture.cpp b/Engine/Source/Runtime/D3D12RHI/Private/D3D12Texture.cpp
--- a/Engine/Source/Runtime/D3D12RHI/Private/D3D12Texture.cpp
+++ b/Engine/Source/Runtime/D3D12RHI/Private/D3D12Texture.cpp
@@ -1014,7 +1014,7 @@ FTextureRHIRef FD3D12DynamicRHI::RHICreateTexture_RenderThread(class FRHICommand
 	return CreateD3D12Texture(CreateDesc, RHIImmediateCmdList, ResourceAllocator);
 }
 
-void FD3D12DynamicRHI::RHIUpdateTextureReference(FRHITextureReference* TextureRef, FRHITexture* InNewTexture)
+void FD3D12DynamicRHI::RHIUpdateTextureReference(FRHICommandListBase& RHICmdList, FRHITextureReference* TextureRef, FRHITexture* InNewTexture)
 {
 	FRHITexture* NewTexture = InNewTexture ? InNewTexture : FRHITextureReference::GetDefaultTexture();
 
@@ -1030,7 +1030,7 @@ void FD3D12DynamicRHI::RHIUpdateTextureReference(FRHITextureReference* TextureRe
 	}
 #endif // PLATFORM_SUPPORTS_BINDLESS_RENDERING
 
-	FDynamicRHI::RHIUpdateTextureReference(TextureRef, NewTexture);
+	FDynamicRHI::RHIUpdateTextureReference(RHICmdList, TextureRef, NewTexture);
 }
 
 class FWaitInitialMipDataUploadTask
diff --git a/Engine/Source/Runtime/RHI/Private/DynamicRHI.cpp b/Engine/Source/Runtime/RHI/Private/DynamicRHI.cpp
--- a/Engine/Source/Runtime/RHI/Private/DynamicRHI.cpp
+++ b/Engine/Source/Runtime/RHI/Private/DynamicRHI.cpp
@@ -561,7 +561,7 @@ void FDynamicRHI::RHITransferBufferUnderlyingResource(FRHIBuffer* DestBuffer, FR
 	UE_LOG(LogRHI, Fatal, TEXT("RHITransferBufferUnderlyingResource isn't implemented for the current RHI"));
 }
 
-void FDynamicRHI::RHIUpdateTextureReference(FRHITextureReference* TextureRef, FRHITexture* InReferencedTexture)
+void FDynamicRHI::RHIUpdateTextureReference(FRHICommandListBase& RHICmdList, FRHITextureReference* TextureRef, FRHITexture* InReferencedTexture)
 {
 	FRHITexture* ReferencedTexture = InReferencedTexture ? InReferencedTexture : FRHITextureReference::GetDefaultTexture();
 	TextureRef->SetReferencedTexture(ReferencedTexture);
diff --git a/Engine/Source/Runtime/RHI/Private/RHICommandList.cpp b/Engine/Source/Runtime/RHI/Private/RHICommandList.cpp
--- a/Engine/Source/Runtime/RHI/Private/RHICommandList.cpp
+++ b/Engine/Source/Runtime/RHI/Private/RHICommandList.cpp
@@ -1795,9 +1795,9 @@ void FRHICommandListImmediate::UpdateTextureReference(FRHITextureReference* Text
 		return;
 	}
 
-	EnqueueLambda([TextureRef, NewTexture](auto&)
+	EnqueueLambda([TextureRef, NewTexture](FRHICommandListBase& RHICmdList)
 	{
-		GDynamicRHI->RHIUpdateTextureReference(TextureRef, NewTexture);
+		GDynamicRHI->RHIUpdateTextureReference(RHICmdList, TextureRef, NewTexture);
 	});
 
 	RHIThreadFence(true);
diff --git a/Engine/Source/Runtime/RHI/Public/DynamicRHI.h b/Engine/Source/Runtime/RHI/Public/DynamicRHI.h
--- a/Engine/Source/Runtime/RHI/Public/DynamicRHI.h
+++ b/Engine/Source/Runtime/RHI/Public/DynamicRHI.h
@@ -353,7 +353,7 @@ public:
 	RHI_API virtual void RHIUnlockBuffer(FRHICommandListBase& RHICmdList, FRHIBuffer* Buffer);
 	RHI_API virtual void RHIUnlockBufferMGPU(FRHICommandListBase& RHICmdList, FRHIBuffer* Buffer, uint32 GPUIndex);
 
-	RHI_API virtual void RHIUpdateTextureReference(FRHITextureReference* TextureRef, FRHITexture* NewTexture);
+	RHI_API virtual void RHIUpdateTextureReference(FRHICommandListBase& RHICmdList, FRHITextureReference* TextureRef, FRHITexture* NewTexture);
 
 	struct FRHICalcTextureSizeResult
 	{
diff --git a/Engine/Source/Runtime/VulkanRHI/Private/VulkanTexture.cpp b/Engine/Source/Runtime/VulkanRHI/Private/VulkanTexture.cpp
--- a/Engine/Source/Runtime/VulkanRHI/Private/VulkanTexture.cpp
+++ b/Engine/Source/Runtime/VulkanRHI/Private/VulkanTexture.cpp
@@ -910,7 +910,7 @@ uint32 FVulkanDynamicRHI::RHIComputeMemorySize(FRHITexture* TextureRHI)
 	return ResourceCast(TextureRHI)->GetMemorySize();
 }
 
-void FVulkanDynamicRHI::RHIUpdateTextureReference(FRHITextureReference* TextureRef, FRHITexture* InNewTexture)
+void FVulkanDynamicRHI::RHIUpdateTextureReference(FRHICommandListBase& RHICmdList, FRHITextureReference* TextureRef, FRHITexture* InNewTexture)
 {
 	if (Device->SupportsBindless())
 	{
@@ -945,7 +945,7 @@ void FVulkanDynamicRHI::RHIUpdateTextureReference(FRHITextureReference* TextureR
 
 	}
 
-	FDynamicRHI::RHIUpdateTextureReference(TextureRef, InNewTexture);
+	FDynamicRHI::RHIUpdateTextureReference(RHICmdList, TextureRef, InNewTexture);
 }
 
 /*-----------------------------------------------------------------------------
diff --git a/Engine/Source/Runtime/VulkanRHI/Public/VulkanDynamicRHI.h b/Engine/Source/Runtime/VulkanRHI/Public/VulkanDynamicRHI.h
--- a/Engine/Source/Runtime/VulkanRHI/Public/VulkanDynamicRHI.h
+++ b/Engine/Source/Runtime/VulkanRHI/Public/VulkanDynamicRHI.h
@@ -120,7 +120,7 @@ public:
 	virtual void* LockBuffer_BottomOfPipe(FRHICommandListBase& RHICmdList, FRHIBuffer* Buffer, uint32 Offset, uint32 Size, EResourceLockMode LockMode) final override;
 	virtual void UnlockBuffer_BottomOfPipe(FRHICommandListBase& RHICmdList, FRHIBuffer* Buffer) final override;
 	virtual void RHIUnlockBuffer(FRHICommandListBase& RHICmdList, FRHIBuffer* Buffer) final override;
-	virtual void RHIUpdateTextureReference(FRHITextureReference* TextureRef, FRHITexture* NewTexture) final override;
+	virtual void RHIUpdateTextureReference(FRHICommandListBase& RHICmdList, FRHITextureReference* TextureRef, FRHITexture* NewTexture) final override;
 	virtual FRHICalcTextureSizeResult RHICalcTexturePlatformSize(FRHITextureDesc const& Desc, uint32 FirstMipIndex) final override;
 	virtual void RHIGetTextureMemoryStats(FTextureMemoryStats& OutStats) final override;
 	virtual bool RHIGetTextureMemoryVisualizeData(FColor* TextureData, int32 SizeX, int32 SizeY, int32 Pitch, int32 PixelSize) final override;
