Use megabytes instead of bytes for storing D3D12 GPU memory statistics in CSV profiles.

The CSV profiler doesn't have support for 64-bit stats so any value larger than 2GiB cannot be represented.

#pf skip
#rnx
#tests tested locally that CSV profiles use MBs as expected.

diff --git a/Engine/Source/Runtime/D3D12RHI/Private/D3D12CommandContext.cpp b/Engine/Source/Runtime/D3D12RHI/Private/D3D12CommandContext.cpp
--- a/Engine/Source/Runtime/D3D12RHI/Private/D3D12CommandContext.cpp
+++ b/Engine/Source/Runtime/D3D12RHI/Private/D3D12CommandContext.cpp
@@ -767,10 +767,10 @@ void FD3D12CommandContextBase::UpdateMemoryStats()
 
 #if CSV_PROFILER
 	{
-		CSV_CUSTOM_STAT(GPUMem, Total, double(MemoryInfo.LocalMemoryInfo.Budget), ECsvCustomStatOp::Set);
-		CSV_CUSTOM_STAT(GPUMem, Used, double(MemoryInfo.LocalMemoryInfo.CurrentUsage), ECsvCustomStatOp::Set);
-		CSV_CUSTOM_STAT(GPUMem, Available, double(MemoryInfo.AvailableLocalMemory), ECsvCustomStatOp::Set);
-		CSV_CUSTOM_STAT(GPUMem, Demoted, double(MemoryInfo.DemotedLocalMemory), ECsvCustomStatOp::Set);
+		CSV_CUSTOM_STAT(GPUMem, TotalMB, float(MemoryInfo.LocalMemoryInfo.Budget / 1024.0 / 1024.0), ECsvCustomStatOp::Set);
+		CSV_CUSTOM_STAT(GPUMem, UsedMB, float(MemoryInfo.LocalMemoryInfo.CurrentUsage / 1024.0 / 1024.0), ECsvCustomStatOp::Set);
+		CSV_CUSTOM_STAT(GPUMem, AvailableMB, float(MemoryInfo.AvailableLocalMemory / 1024.0 / 1024.0), ECsvCustomStatOp::Set);
+		CSV_CUSTOM_STAT(GPUMem, DemotedMB, float(MemoryInfo.DemotedLocalMemory / 1024.0 / 1024.0), ECsvCustomStatOp::Set);
 	}
 #endif // CSV_PROFILER
 
