Add basic support for MultiDrawIndirect to D3D12 RHI

diff --git a/Engine/Source/Runtime/D3D12RHI/Private/D3D12CommandContext.h b/Engine/Source/Runtime/D3D12RHI/Private/D3D12CommandContext.h
--- a/Engine/Source/Runtime/D3D12RHI/Private/D3D12CommandContext.h
+++ b/Engine/Source/Runtime/D3D12RHI/Private/D3D12CommandContext.h
@@ -628,6 +628,7 @@ public:
 	virtual void RHIDrawIndexedIndirect(FRHIBuffer* IndexBufferRHI, FRHIBuffer* ArgumentsBufferRHI, int32 DrawArgumentsIndex, uint32 NumInstances) final override;
 	virtual void RHIDrawIndexedPrimitive(FRHIBuffer* IndexBuffer, int32 BaseVertexIndex, uint32 FirstInstance, uint32 NumVertices, uint32 StartIndex, uint32 NumPrimitives, uint32 NumInstances) final override;
 	virtual void RHIDrawIndexedPrimitiveIndirect(FRHIBuffer* IndexBuffer, FRHIBuffer* ArgumentBuffer, uint32 ArgumentOffset) final override;
+	virtual void RHIMultiDrawIndexedPrimitiveIndirect(FRHIBuffer* IndexBuffer, FRHIBuffer* ArgumentBuffer, uint32 ArgumentOffset, FRHIBuffer* CountBuffer, uint32 CountBufferOffset, uint32 MaxDrawArguments) final override;
 #if PLATFORM_SUPPORTS_MESH_SHADERS
 	virtual void RHIDispatchMeshShader(uint32 ThreadGroupCountX, uint32 ThreadGroupCountY, uint32 ThreadGroupCountZ) final override;
 	virtual void RHIDispatchIndirectMeshShader(FRHIBuffer* ArgumentBuffer, uint32 ArgumentOffset) final override;
@@ -987,6 +988,10 @@ public:
 	{
 		ContextRedirect(RHIDrawIndexedPrimitiveIndirect(IndexBuffer, ArgumentBuffer, ArgumentOffset));
 	}
+	FORCEINLINE virtual void RHIMultiDrawIndexedPrimitiveIndirect(FRHIBuffer* IndexBuffer, FRHIBuffer* ArgumentBuffer, uint32 ArgumentOffset, FRHIBuffer* CountBuffer, uint32 CountBufferOffset, uint32 MaxDrawArguments) final override
+	{
+		ContextRedirect(RHIMultiDrawIndexedPrimitiveIndirect(IndexBuffer, ArgumentBuffer, ArgumentOffset, CountBuffer, CountBufferOffset, MaxDrawArguments));
+	}
 #if PLATFORM_SUPPORTS_MESH_SHADERS
 	FORCEINLINE virtual void RHIDispatchMeshShader(uint32 ThreadGroupCountX, uint32 ThreadGroupCountY, uint32 ThreadGroupCountZ) final override
 	{
diff --git a/Engine/Source/Runtime/D3D12RHI/Private/D3D12Commands.cpp b/Engine/Source/Runtime/D3D12RHI/Private/D3D12Commands.cpp
--- a/Engine/Source/Runtime/D3D12RHI/Private/D3D12Commands.cpp
+++ b/Engine/Source/Runtime/D3D12RHI/Private/D3D12Commands.cpp
@@ -1765,7 +1765,7 @@ void FD3D12CommandContext::RHIDrawIndexedPrimitive(FRHIBuffer* IndexBufferRHI, i
 	DEBUG_EXECUTE_COMMAND_LIST(this);
 }
 
-void FD3D12CommandContext::RHIDrawIndexedPrimitiveIndirect(FRHIBuffer* IndexBufferRHI, FRHIBuffer* ArgumentBufferRHI, uint32 ArgumentOffset)
+void FD3D12CommandContext::RHIMultiDrawIndexedPrimitiveIndirect(FRHIBuffer* IndexBufferRHI, FRHIBuffer* ArgumentBufferRHI, uint32 ArgumentOffset, FRHIBuffer* CountBufferRHI, uint32 CountBufferOffset, uint32 MaxDrawArguments)
 {
 	FD3D12Buffer* IndexBuffer = RetrieveObject<FD3D12Buffer>(IndexBufferRHI);
 
@@ -1781,6 +1781,18 @@ void FD3D12CommandContext::RHIDrawIndexedPrimitiveIndirect(FRHIBuffer* IndexBuff
 	const uint32 IndexBufferStride = FD3D12DynamicRHI::ResourceCast(IndexBufferRHI)->GetStride();
 	FD3D12Buffer* ArgumentBuffer = RetrieveObject<FD3D12Buffer>(ArgumentBufferRHI);
 
+	ID3D12Resource* CountBufferResource = nullptr;
+	uint64 CountBufferOffsetFromResourceBase = 0;
+	if (CountBufferRHI)
+	{
+		FD3D12Buffer* CountBuffer = RetrieveObject<FD3D12Buffer>(CountBufferRHI);
+		FD3D12ResourceLocation& CounterLocation = CountBuffer->ResourceLocation;
+		CountBufferResource = CounterLocation.GetResource()->GetResource();
+		TransitionResource(CounterLocation.GetResource(), D3D12_RESOURCE_STATE_TBD, D3D12_RESOURCE_STATE_INDIRECT_ARGUMENT, D3D12_RESOURCE_BARRIER_ALL_SUBRESOURCES);
+		CountBufferOffsetFromResourceBase = CounterLocation.GetOffsetFromBaseOfResource() + CountBufferOffset;
+		UpdateResidency(CounterLocation.GetResource());
+	}
+
 	RHI_DRAW_CALL_INC();
 	if (bTrackingEvents)
 	{
@@ -1805,11 +1817,11 @@ void FD3D12CommandContext::RHIDrawIndexedPrimitiveIndirect(FRHIBuffer* IndexBuff
 
 	GraphicsCommandList()->ExecuteIndirect(
 		GetParentDevice()->GetParentAdapter()->GetDrawIndexedIndirectCommandSignature(),
-		1,
+		MaxDrawArguments,
 		Location.GetResource()->GetResource(),
 		Location.GetOffsetFromBaseOfResource() + ArgumentOffset,
-		NULL,
-		0
+		CountBufferResource,
+		CountBufferOffsetFromResourceBase
 	);
 
 	UpdateResidency(Location.GetResource());
@@ -1819,6 +1831,12 @@ void FD3D12CommandContext::RHIDrawIndexedPrimitiveIndirect(FRHIBuffer* IndexBuff
 	DEBUG_EXECUTE_COMMAND_LIST(this);
 }
 
+void FD3D12CommandContext::RHIDrawIndexedPrimitiveIndirect(FRHIBuffer* IndexBufferRHI, FRHIBuffer* ArgumentBufferRHI, uint32 ArgumentOffset)
+{
+	// DrawIndexedPrimitiveIndirect is a special case of a more general MDI in D3D12
+	RHIMultiDrawIndexedPrimitiveIndirect(IndexBufferRHI, ArgumentBufferRHI, ArgumentOffset, nullptr /*CounterBuffer*/, 0 /*CounterBufferOffset*/, 1 /*MaxDrawArguments*/);
+}
+
 #if PLATFORM_SUPPORTS_MESH_SHADERS
 void FD3D12CommandContext::RHIDispatchMeshShader(uint32 ThreadGroupCountX, uint32 ThreadGroupCountY, uint32 ThreadGroupCountZ)
 {
diff --git a/Engine/Source/Runtime/D3D12RHI/Private/D3D12RHI.cpp b/Engine/Source/Runtime/D3D12RHI/Private/D3D12RHI.cpp
--- a/Engine/Source/Runtime/D3D12RHI/Private/D3D12RHI.cpp
+++ b/Engine/Source/Runtime/D3D12RHI/Private/D3D12RHI.cpp
@@ -76,6 +76,8 @@ FD3D12DynamicRHI::FD3D12DynamicRHI(const TArray<TSharedPtr<FD3D12Adapter>>& Chos
 	ZeroBuffer = nullptr;
 #endif // PLATFORM_WINDOWS
 
+	GRHIGlobals.SupportsMultiDrawIndirect = true;
+
 	GRHISupportsMultithreading = true;
 	GRHISupportsMultithreadedResources = true;
 	GRHISupportsAsyncGetRenderQueryResult = true;
diff --git a/Engine/Source/Runtime/NullDrv/Public/NullRHI.h b/Engine/Source/Runtime/NullDrv/Public/NullRHI.h
--- a/Engine/Source/Runtime/NullDrv/Public/NullRHI.h
+++ b/Engine/Source/Runtime/NullDrv/Public/NullRHI.h
@@ -552,11 +552,17 @@ public:
 	{
 
 	}
+
 	virtual void RHIDrawIndexedPrimitiveIndirect(FRHIBuffer* IndexBuffer, FRHIBuffer* ArgumentBuffer, uint32 ArgumentOffset) final override
 	{
 
 	}
 
+	virtual void RHIMultiDrawIndexedPrimitiveIndirect(FRHIBuffer* IndexBuffer, FRHIBuffer* ArgumentBuffer, uint32 ArgumentOffset, FRHIBuffer* CountBuffer, uint32 CountBufferOffset, uint32 MaxDrawArguments) final override
+	{
+
+	}
+
 	virtual void RHIBlockUntilGPUIdle() final override
 	{
 	}
diff --git a/Engine/Source/Runtime/RHI/Public/RHICommandList.h b/Engine/Source/Runtime/RHI/Public/RHICommandList.h
--- a/Engine/Source/Runtime/RHI/Public/RHICommandList.h
+++ b/Engine/Source/Runtime/RHI/Public/RHICommandList.h
@@ -1696,6 +1696,26 @@ FRHICOMMAND_MACRO(FRHICommandDrawIndexedPrimitiveIndirect)
 	RHI_API void Execute(FRHICommandListBase& CmdList);
 };
 
+FRHICOMMAND_MACRO(FRHICommandMultiDrawIndexedPrimitiveIndirect)
+{
+	FRHIBuffer* IndexBuffer;
+	FRHIBuffer* ArgumentBuffer;
+	uint32 ArgumentOffset;
+	FRHIBuffer* CountBuffer;
+	uint32 CountBufferOffset;
+	uint32 MaxDrawArguments;
+	FORCEINLINE_DEBUGGABLE FRHICommandMultiDrawIndexedPrimitiveIndirect(FRHIBuffer* InIndexBuffer, FRHIBuffer* InArgumentBuffer, uint32 InArgumentOffset, FRHIBuffer* InCountBuffer, uint32 InCountBufferOffset, uint32 InMaxDrawArguments)
+		: IndexBuffer(InIndexBuffer)
+		, ArgumentBuffer(InArgumentBuffer)
+		, ArgumentOffset(InArgumentOffset)
+		, CountBuffer(InCountBuffer)
+		, CountBufferOffset(InCountBufferOffset)
+		, MaxDrawArguments(InMaxDrawArguments)
+	{
+	}
+	RHI_API void Execute(FRHICommandListBase & CmdList);
+};
+
 FRHICOMMAND_MACRO(FRHICommandDispatchMeshShader)
 {
 	uint32 ThreadGroupCountX;
@@ -3490,6 +3510,18 @@ PRAGMA_ENABLE_DEPRECATION_WARNINGS
 		ALLOC_COMMAND(FRHICommandDrawIndexedPrimitiveIndirect)(IndexBuffer, ArgumentsBuffer, ArgumentOffset);
 	}
 
+	FORCEINLINE_DEBUGGABLE void MultiDrawIndexedPrimitiveIndirect(FRHIBuffer* IndexBuffer, FRHIBuffer* ArgumentsBuffer, uint32 ArgumentOffset, FRHIBuffer* CountBuffer, uint32 CountBufferOffset, uint32 MaxDrawArguments)
+	{
+		//check(IsOutsideRenderPass());
+		PreDraw();
+		if (Bypass())
+		{
+			GetContext().RHIMultiDrawIndexedPrimitiveIndirect(IndexBuffer, ArgumentsBuffer, ArgumentOffset, CountBuffer, CountBufferOffset, MaxDrawArguments);
+			return;
+		}
+		ALLOC_COMMAND(FRHICommandMultiDrawIndexedPrimitiveIndirect)(IndexBuffer, ArgumentsBuffer, ArgumentOffset, CountBuffer, CountBufferOffset, MaxDrawArguments);
+	}
+
 	FORCEINLINE_DEBUGGABLE void DispatchMeshShader(uint32 ThreadGroupCountX, uint32 ThreadGroupCountY, uint32 ThreadGroupCountZ)
 	{
 		if (Bypass())
diff --git a/Engine/Source/Runtime/RHI/Public/RHICommandListCommandExecutes.inl b/Engine/Source/Runtime/RHI/Public/RHICommandListCommandExecutes.inl
--- a/Engine/Source/Runtime/RHI/Public/RHICommandListCommandExecutes.inl
+++ b/Engine/Source/Runtime/RHI/Public/RHICommandListCommandExecutes.inl
@@ -32,6 +32,7 @@ struct FRHICommandDrawIndexedPrimitive;
 struct FRHICommandDrawIndexedPrimitiveIndirect;
 struct FRHICommandDrawPrimitive;
 struct FRHICommandDrawPrimitiveIndirect;
+struct FRHICommandMultiDrawPrimitiveIndirect;
 struct FRHICommandSetDepthBounds;
 struct FRHICommandEndDrawingViewport;
 struct FRHICommandEndFrame;
@@ -283,6 +284,12 @@ void FRHICommandDrawIndexedPrimitiveIndirect::Execute(FRHICommandListBase& CmdLi
 	INTERNAL_DECORATOR(RHIDrawIndexedPrimitiveIndirect)(IndexBuffer, ArgumentsBuffer, ArgumentOffset);
 }
 
+void FRHICommandMultiDrawIndexedPrimitiveIndirect::Execute(FRHICommandListBase& CmdList)
+{
+	RHISTAT(MultiDrawIndexedPrimitiveIndirect);
+	INTERNAL_DECORATOR(RHIMultiDrawIndexedPrimitiveIndirect)(IndexBuffer, ArgumentBuffer, ArgumentOffset, CountBuffer, CountBufferOffset, MaxDrawArguments);
+}
+
 void FRHICommandDispatchMeshShader::Execute(FRHICommandListBase& CmdList)
 {
 	RHISTAT(DispatchMeshShader);
diff --git a/Engine/Source/Runtime/RHI/Public/RHIContext.h b/Engine/Source/Runtime/RHI/Public/RHIContext.h
--- a/Engine/Source/Runtime/RHI/Public/RHIContext.h
+++ b/Engine/Source/Runtime/RHI/Public/RHIContext.h
@@ -751,6 +751,21 @@ public:
 
 	virtual void RHIDrawIndexedPrimitiveIndirect(FRHIBuffer* IndexBuffer, FRHIBuffer* ArgumentBuffer, uint32 ArgumentOffset) = 0;
 
+	/**
+	* Similar to RHIDrawIndexedPrimitiveIndirect, but allows many draw arguments to be provided at once.
+	* GRHIGlobals.SupportsDrawIndirect must be checked to detect support on the current machine.
+	* @ param IndexBuffer			Buffer containing primitive indices
+	* @ param ArgumentsBuffer		Buffer containing FRHIDrawIndexedIndirectParameters structures
+	* @ param ArgumentOffset		Offset in bytes of the first element in ArgumentsBuffer that will be used for drawing
+	* @ param CountBuffer			Buffer containing uint32 count of valid draw arguments that should be consumed (may be nullptr, indicating that only MaxDrawArguments value should be used)
+	* @ param CountBuffeOffset		Offset in bytes for the CountBuffer element that will be used to source the draw argument count
+	* @ param MaxDrawArguments		How many draw arguments should be processed at most, i.e. NumDrawArguments = min(MaxDrawArguments, ValueFromCountBuffer)
+	*/
+	virtual void RHIMultiDrawIndexedPrimitiveIndirect(FRHIBuffer* IndexBuffer, FRHIBuffer* ArgumentBuffer, uint32 ArgumentOffset, FRHIBuffer* CountBuffer, uint32 CountBuffeOffset, uint32 MaxDrawArguments)
+	{
+		checkNoEntry();
+	}
+
 	virtual void RHIDispatchMeshShader(uint32 ThreadGroupCountX, uint32 ThreadGroupCountY, uint32 ThreadGroupCountZ)
 	{
 		/* empty default implementation */
diff --git a/Engine/Source/Runtime/RHI/Public/RHIGlobals.h b/Engine/Source/Runtime/RHI/Public/RHIGlobals.h
--- a/Engine/Source/Runtime/RHI/Public/RHIGlobals.h
+++ b/Engine/Source/Runtime/RHI/Public/RHIGlobals.h
@@ -107,6 +107,9 @@ struct FRHIGlobals
 	/** true if the RHI supports Draw Indirect */
 	bool SupportsDrawIndirect = true;
 
+	/** true if the RHI supports Multi Draw Indirect, a variant of Draw Indirect with variable number of sub-commands generated by GPU */
+	bool SupportsMultiDrawIndirect = false;
+
 	/** Whether the RHI can send commands to the device context from multiple threads. Used in the GPU readback to avoid stalling the RHI threads. */
 	bool SupportsMultithreading = false;
 
@@ -590,6 +593,7 @@ extern RHI_API FRHIGlobals GRHIGlobals;
 #define GSupportsQuadBufferStereo                              GRHIGlobals.SupportsQuadBufferStereo
 #define GSupportsRenderDepthTargetableShaderResources          GRHIGlobals.SupportsRenderDepthTargetableShaderResources
 #define GRHISupportsDrawIndirect                               GRHIGlobals.SupportsDrawIndirect
+#define GRHISupportsMultiDrawIndirect                          GRHIGlobals.SupportsMultiDrawIndirect
 #define GRHISupportsMultithreading                             GRHIGlobals.SupportsMultithreading
 #define GRHISupportsAsyncGetRenderQueryResult                  GRHIGlobals.SupportsAsyncGetRenderQueryResult
 #define GRHIAdapterName                                        GRHIGlobals.GpuInfo.AdapterName
diff --git a/Engine/Source/Runtime/RHI/Public/RHIValidationContext.h b/Engine/Source/Runtime/RHI/Public/RHIValidationContext.h
--- a/Engine/Source/Runtime/RHI/Public/RHIValidationContext.h
+++ b/Engine/Source/Runtime/RHI/Public/RHIValidationContext.h
@@ -930,6 +930,21 @@ public:
 		Tracker->Draw();
 	}
 
+	virtual void RHIMultiDrawIndexedPrimitiveIndirect(FRHIBuffer* IndexBuffer, FRHIBuffer* ArgumentBuffer, uint32 ArgumentOffset, FRHIBuffer* CountBuffer, uint32 CountBufferOffset, uint32 MaxDrawArguments) override final
+	{
+		checkf(State.bGfxPSOSet, TEXT("A Graphics PSO has to be set to draw!"));
+		FValidationRHI::ValidateIndirectArgsBuffer(ArgumentBuffer, ArgumentOffset, sizeof(FRHIDrawIndexedIndirectParameters), 0);
+		SetupDrawing();
+		Tracker->Assert(ArgumentBuffer->GetWholeResourceIdentity(), ERHIAccess::IndirectArgs);
+		if (CountBuffer)
+		{
+			Tracker->Assert(CountBuffer->GetWholeResourceIdentity(), ERHIAccess::IndirectArgs);
+		}
+		Tracker->Assert(IndexBuffer->GetWholeResourceIdentity(), ERHIAccess::VertexOrIndexBuffer);
+		RHIContext->RHIMultiDrawIndexedPrimitiveIndirect(IndexBuffer, ArgumentBuffer, ArgumentOffset, CountBuffer, CountBufferOffset, MaxDrawArguments);
+		Tracker->Draw();
+	}
+
 	virtual void RHIDispatchMeshShader(uint32 ThreadGroupCountX, uint32 ThreadGroupCountY, uint32 ThreadGroupCountZ) final override
 	{
 		checkf(State.bGfxPSOSet, TEXT("A Graphics PSO has to be set to draw!"));
