Only call TerminateOnGPUCrash() when the GPU has actually crashed.

We used to call this function on any D3D error, which produced confusing error reports when the problem was caused by incorrect API usage (e.g. E_INVALIDARG) rather than actual GPU crashes.

#rnx

diff --git a/Engine/Source/Runtime/D3D12RHI/Private/D3D12Util.cpp b/Engine/Source/Runtime/D3D12RHI/Private/D3D12Util.cpp
--- a/Engine/Source/Runtime/D3D12RHI/Private/D3D12Util.cpp
+++ b/Engine/Source/Runtime/D3D12RHI/Private/D3D12Util.cpp
@@ -884,7 +884,7 @@ namespace D3D12RHI
 		{
 			TerminateOnOutOfMemory(Device, D3DResult, false);
 		}
-		else
+		else if (D3DResult == DXGI_ERROR_DEVICE_REMOVED || D3DResult == DXGI_ERROR_DEVICE_HUNG || D3DResult == DXGI_ERROR_DEVICE_RESET)
 		{
 			TerminateOnGPUCrash(Device, nullptr, 0);
 		}
@@ -922,7 +922,7 @@ namespace D3D12RHI
 			*GetD3D12TextureFlagString(TextureDesc.Flags));
 
 		// Terminate with device removed but we don't have any GPU crash dump information
-		if (D3DResult == DXGI_ERROR_DEVICE_REMOVED || D3DResult == DXGI_ERROR_DEVICE_HUNG)
+		if (D3DResult == DXGI_ERROR_DEVICE_REMOVED || D3DResult == DXGI_ERROR_DEVICE_HUNG || D3DResult == DXGI_ERROR_DEVICE_RESET)
 		{
 			TerminateOnGPUCrash(Device, nullptr, 0);
 		}
