Don't mark UAV slot dirty in FD3D12StateCache::SetUAV if it is set to the current value

diff --git a/Engine/Source/Runtime/D3D12RHI/Private/D3D12StateCache.cpp b/Engine/Source/Runtime/D3D12RHI/Private/D3D12StateCache.cpp
--- a/Engine/Source/Runtime/D3D12RHI/Private/D3D12StateCache.cpp
+++ b/Engine/Source/Runtime/D3D12RHI/Private/D3D12StateCache.cpp
@@ -1127,6 +1127,10 @@ void FD3D12StateCache::SetUAV(EShaderFrequency ShaderStage, uint32 SlotIndex, FD
 	SCOPE_CYCLE_COUNTER(STAT_D3D12SetUnorderedAccessViewTime);
 
 	FD3D12UnorderedAccessViewCache& Cache = PipelineState.Common.UAVCache;
+	if (Cache.Views[ShaderStage][SlotIndex] == UAV)
+	{
+		return;
+	}
 
 	// When setting UAV's for Graphics, it wipes out all existing bound resources.
 	const bool bIsCompute = ShaderStage == SF_Compute;
