Implement whole-frame D3D12 pipeline statistics as a new stat group: "stat D3D12RHIPipeline"

diff --git a/Engine/Source/Runtime/D3D12RHI/Private/D3D12CommandContext.cpp b/Engine/Source/Runtime/D3D12RHI/Private/D3D12CommandContext.cpp
--- a/Engine/Source/Runtime/D3D12RHI/Private/D3D12CommandContext.cpp
+++ b/Engine/Source/Runtime/D3D12RHI/Private/D3D12CommandContext.cpp
@@ -201,6 +201,7 @@ FD3D12ContextCommon::FD3D12ContextCommon(FD3D12Device* Device, ED3D12QueueType Q
 	, bIsDefaultContext(bIsDefaultContext)
 	, TimestampQueries(Device, QueueType, D3D12_QUERY_TYPE_TIMESTAMP)
 	, OcclusionQueries(Device, QueueType, D3D12_QUERY_TYPE_OCCLUSION)
+	, PipelineStatsQueries(Device, QueueType, D3D12_QUERY_TYPE_PIPELINE_STATISTICS)
 {
 }
 
@@ -244,7 +245,7 @@ void FD3D12ContextCommon::WaitManualFence(ID3D12Fence* Fence, uint64 Value)
 	GetPayload(EPhase::Wait)->FencesToWait.Emplace(Fence, Value);
 }
 
-FD3D12QueryLocation FD3D12ContextCommon::AllocateQuery(ED3D12QueryType Type, uint64* Target)
+FD3D12QueryLocation FD3D12ContextCommon::AllocateQuery(ED3D12QueryType Type, void* Target)
 {
 	switch (Type)
 	{
@@ -258,6 +259,9 @@ FD3D12QueryLocation FD3D12ContextCommon::AllocateQuery(ED3D12QueryType Type, uin
 
 	case ED3D12QueryType::Occlusion:
 		return OcclusionQueries.Allocate(Type, Target);
+
+	case ED3D12QueryType::PipelineStats:
+		return PipelineStatsQueries.Allocate(Type, Target);
 	}
 }
 
@@ -292,7 +296,7 @@ void FD3D12ContextCommon::OpenCommandList()
 	}
 
 	// Get a new command list
-	CommandList = Device->ObtainCommandList(CommandAllocator, &TimestampQueries);
+	CommandList = Device->ObtainCommandList(CommandAllocator, &TimestampQueries, &PipelineStatsQueries);
 	GetPayload(EPhase::Execute)->CommandListsToExecute.Add(CommandList);
 
 	check(ActiveQueries == 0);
@@ -324,6 +328,7 @@ void FD3D12ContextCommon::CloseCommandList()
 
 	TimestampQueries.CloseAndReset(Payload->QueryRanges);
 	OcclusionQueries.CloseAndReset(Payload->QueryRanges);
+	PipelineStatsQueries.CloseAndReset(Payload->QueryRanges);
 }
 
 void FD3D12CommandContext::CloseCommandList()
@@ -369,6 +374,7 @@ void FD3D12ContextCommon::Finalize(TArray<FD3D12Payload*>& OutPayloads)
 
 	check(!TimestampQueries.HasQueries());
 	check(!OcclusionQueries.HasQueries());
+	check(!PipelineStatsQueries.HasQueries());
 
 	ContextSyncPoint = nullptr;
 
@@ -376,7 +382,7 @@ void FD3D12ContextCommon::Finalize(TArray<FD3D12Payload*>& OutPayloads)
 	OutPayloads.Append(MoveTemp(Payloads));
 }
 
-FD3D12QueryLocation FD3D12QueryAllocator::Allocate(ED3D12QueryType Type, uint64* Target)
+FD3D12QueryLocation FD3D12QueryAllocator::Allocate(ED3D12QueryType Type, void* Target)
 {
 	check(Type != ED3D12QueryType::None);
 
diff --git a/Engine/Source/Runtime/D3D12RHI/Private/D3D12CommandContext.h b/Engine/Source/Runtime/D3D12RHI/Private/D3D12CommandContext.h
--- a/Engine/Source/Runtime/D3D12RHI/Private/D3D12CommandContext.h
+++ b/Engine/Source/Runtime/D3D12RHI/Private/D3D12CommandContext.h
@@ -177,7 +177,7 @@ public:
 	FD3D12QueryLocation InsertTimestamp(ED3D12Units Units, uint64* Target);
 
 	// Allocates a query of the specified type, returning its location.
-	FD3D12QueryLocation AllocateQuery(ED3D12QueryType Type, uint64* Target);
+	FD3D12QueryLocation AllocateQuery(ED3D12QueryType Type, void* Target);
 
 	// Complete recording of the current command list set, and appends the resulting
 	// payloads to the given array. Resets the context so new commands can be recorded.
@@ -220,7 +220,7 @@ public:
 
 private:
 	// Allocators to manage query heaps
-	FD3D12QueryAllocator TimestampQueries, OcclusionQueries;
+	FD3D12QueryAllocator TimestampQueries, OcclusionQueries, PipelineStatsQueries;
 
 	// Batches resource barriers together until it's explicitly flushed
 	FD3D12ResourceBarrierBatcher ResourceBarrierBatcher;
diff --git a/Engine/Source/Runtime/D3D12RHI/Private/D3D12CommandList.cpp b/Engine/Source/Runtime/D3D12RHI/Private/D3D12CommandList.cpp
--- a/Engine/Source/Runtime/D3D12RHI/Private/D3D12CommandList.cpp
+++ b/Engine/Source/Runtime/D3D12RHI/Private/D3D12CommandList.cpp
@@ -147,11 +147,11 @@ void FD3D12CommandAllocator::Reset()
 	VERIFYD3D12RESULT(CommandAllocator->Reset());
 }
 
-FD3D12CommandList::FD3D12CommandList(FD3D12CommandAllocator* CommandAllocator, FD3D12QueryAllocator* TimestampAllocator)
+FD3D12CommandList::FD3D12CommandList(FD3D12CommandAllocator* CommandAllocator, FD3D12QueryAllocator* TimestampAllocator, FD3D12QueryAllocator* PipelineStatsAllocator)
 	: Device(CommandAllocator->Device)
 	, QueueType(CommandAllocator->QueueType)
 	, ResidencySet(D3DX12Residency::CreateResidencySet(Device->GetResidencyManager()))
-	, State(CommandAllocator, TimestampAllocator)
+	, State(CommandAllocator, TimestampAllocator, PipelineStatsAllocator)
 {
 	switch (QueueType)
 	{
@@ -236,7 +236,7 @@ FD3D12CommandList::FD3D12CommandList(FD3D12CommandAllocator* CommandAllocator, F
 #endif
 
 	D3DX12Residency::Open(ResidencySet);
-	WriteBeginTimestamp();
+	BeginLocalQueries();
 }
 
 FD3D12CommandList::~FD3D12CommandList()
@@ -256,7 +256,7 @@ FD3D12CommandList::~FD3D12CommandList()
 	DEC_DWORD_STAT(STAT_D3D12NumCommandLists);
 }
 
-void FD3D12CommandList::Reset(FD3D12CommandAllocator* NewCommandAllocator, FD3D12QueryAllocator* TimestampAllocator)
+void FD3D12CommandList::Reset(FD3D12CommandAllocator* NewCommandAllocator, FD3D12QueryAllocator* TimestampAllocator, FD3D12QueryAllocator* PipelineStatsAllocator)
 {
 	check(IsClosed());
 	check(NewCommandAllocator->Device == Device && NewCommandAllocator->QueueType == QueueType);
@@ -270,14 +270,14 @@ void FD3D12CommandList::Reset(FD3D12CommandAllocator* NewCommandAllocator, FD3D1
 	}
 	D3DX12Residency::Open(ResidencySet);
 
-	State = FState(NewCommandAllocator, TimestampAllocator);
-	WriteBeginTimestamp();
+	State = FState(NewCommandAllocator, TimestampAllocator, PipelineStatsAllocator);
+	BeginLocalQueries();
 }
 
 void FD3D12CommandList::Close()
 {
 	check(IsOpen());
-	WriteEndTimestamp();
+	EndLocalQueries();
 
 	if (Interfaces.CopyCommandList)
 	{
@@ -291,28 +291,46 @@ void FD3D12CommandList::Close()
 	State.IsClosed = true;
 }
 
-void FD3D12CommandList::WriteBeginTimestamp()
+void FD3D12CommandList::BeginLocalQueries()
 {
-	if (!State.bBeginTimestampWritten && State.BeginTimestamp)
+	if (!State.bLocalQueriesBegun)
 	{
-		EndQuery(State.BeginTimestamp);
-		State.bBeginTimestampWritten = true;
+		if (State.BeginTimestamp)
+		{
+			EndQuery(State.BeginTimestamp);
+		}
+
+		if (State.PipelineStats)
+		{
+			BeginQuery(State.PipelineStats);
+		}
+
+		State.bLocalQueriesBegun = true;
 	}
 }
 
-void FD3D12CommandList::WriteEndTimestamp()
+void FD3D12CommandList::EndLocalQueries()
 {
-	if (!State.bEndTimestampWritten && State.EndTimestamp)
+	if (!State.bLocalQueriesEnded)
 	{
-		EndQuery(State.EndTimestamp);
-		State.bEndTimestampWritten = true;
+		if (State.PipelineStats)
+		{
+			EndQuery(State.PipelineStats);
+		}
+
+		if (State.EndTimestamp)
+		{
+			EndQuery(State.EndTimestamp);
+		}
+
+		State.bLocalQueriesEnded = true;
 	}
 }
 
 void FD3D12CommandList::BeginQuery(FD3D12QueryLocation const& Location)
 {
 	check(Location);
-	check(Location.Heap->QueryType == D3D12_QUERY_TYPE_OCCLUSION);
+	check(Location.Heap->QueryType == D3D12_QUERY_TYPE_OCCLUSION || Location.Heap->QueryType == D3D12_QUERY_TYPE_PIPELINE_STATISTICS);
 
 	GraphicsCommandList()->BeginQuery(
 		Location.Heap->GetD3DQueryHeap(),
@@ -330,6 +348,15 @@ void FD3D12CommandList::EndQuery(FD3D12QueryLocation const& Location)
 		checkNoEntry();
 		break;
 
+	case D3D12_QUERY_TYPE_PIPELINE_STATISTICS:
+		GraphicsCommandList()->EndQuery(
+			Location.Heap->GetD3DQueryHeap(),
+			Location.Heap->QueryType,
+			Location.Index
+		);
+		State.PipelineStatsQueries.Add(Location);
+		break;
+
 	case D3D12_QUERY_TYPE_OCCLUSION:
 		GraphicsCommandList()->EndQuery(
 			Location.Heap->GetD3DQueryHeap(),
@@ -363,7 +390,7 @@ void FD3D12CommandList::WriteTimestamp(FD3D12QueryLocation const& Location)
 }
 #endif // D3D12RHI_PLATFORM_USES_TIMESTAMP_QUERIES
 
-FD3D12CommandList::FState::FState(FD3D12CommandAllocator* CommandAllocator, FD3D12QueryAllocator* TimestampAllocator)
+FD3D12CommandList::FState::FState(FD3D12CommandAllocator* CommandAllocator, FD3D12QueryAllocator* TimestampAllocator, FD3D12QueryAllocator* PipelineStatsAllocator)
 	: CommandAllocator(CommandAllocator)
 	, CommandListID   (FPlatformAtomics::InterlockedIncrement(&NextCommandListID))
 {
@@ -374,6 +401,11 @@ FD3D12CommandList::FState::FState(FD3D12CommandAllocator* CommandAllocator, FD3D
 		BeginTimestamp = TimestampAllocator->Allocate(ED3D12QueryType::CommandListBegin, nullptr);
 		EndTimestamp   = TimestampAllocator->Allocate(ED3D12QueryType::CommandListEnd  , nullptr);
 	}
+
+	if (PipelineStatsAllocator)
+	{
+		PipelineStats = PipelineStatsAllocator->Allocate(ED3D12QueryType::PipelineStats, nullptr);
+	}
 }
 
 void FD3D12ContextCommon::TransitionResource(FD3D12DepthStencilView* View)
diff --git a/Engine/Source/Runtime/D3D12RHI/Private/D3D12CommandList.h b/Engine/Source/Runtime/D3D12RHI/Private/D3D12CommandList.h
--- a/Engine/Source/Runtime/D3D12RHI/Private/D3D12CommandList.h
+++ b/Engine/Source/Runtime/D3D12RHI/Private/D3D12CommandList.h
@@ -76,10 +76,10 @@ private:
 	FD3D12CommandList(FD3D12CommandList const&) = delete;
 	FD3D12CommandList(FD3D12CommandList&&) = delete;
 
-	FD3D12CommandList(FD3D12CommandAllocator* CommandAllocator, FD3D12QueryAllocator* TimestampAllocator);
+	FD3D12CommandList(FD3D12CommandAllocator* CommandAllocator, FD3D12QueryAllocator* TimestampAllocator, FD3D12QueryAllocator* PipelineStatsAllocator);
 
-	void WriteBeginTimestamp();
-	void WriteEndTimestamp();
+	void BeginLocalQueries();
+	void EndLocalQueries();
 
 public:
 	~FD3D12CommandList();
@@ -87,7 +87,7 @@ public:
 	void BeginQuery(FD3D12QueryLocation const& Location);
 	void EndQuery  (FD3D12QueryLocation const& Location);
 
-	void Reset(FD3D12CommandAllocator* CommandAllocator, FD3D12QueryAllocator* TimestampAllocator);
+	void Reset(FD3D12CommandAllocator* CommandAllocator, FD3D12QueryAllocator* TimestampAllocator, FD3D12QueryAllocator* PipelineStatsAllocator);
 	void Close();
 
 	bool IsOpen  () const { return !State.IsClosed; }
@@ -241,7 +241,7 @@ private:
 	{
 		static int64 NextCommandListID;
 
-		FState(FD3D12CommandAllocator* CommandAllocator, FD3D12QueryAllocator* TimestampAllocator);
+		FState(FD3D12CommandAllocator* CommandAllocator, FD3D12QueryAllocator* TimestampAllocator, FD3D12QueryAllocator* PipelineStatsAllocator);
 
 		// The allocator currently assigned to this command list.
 		FD3D12CommandAllocator* CommandAllocator;
@@ -265,14 +265,17 @@ private:
 		FD3D12QueryLocation BeginTimestamp;
 		FD3D12QueryLocation EndTimestamp;
 
+		FD3D12QueryLocation PipelineStats;
+
 		TArray<FD3D12QueryLocation> TimestampQueries;
 		TArray<FD3D12QueryLocation> OcclusionQueries;
+		TArray<FD3D12QueryLocation> PipelineStatsQueries;
 
 		uint32 NumCommands = 0;
 
 		bool IsClosed = false;
-		bool bBeginTimestampWritten = false;
-		bool bEndTimestampWritten = false;
+		bool bLocalQueriesBegun = false;
+		bool bLocalQueriesEnded = false;
 
 	} State;
 };
diff --git a/Engine/Source/Runtime/D3D12RHI/Private/D3D12Device.cpp b/Engine/Source/Runtime/D3D12RHI/Private/D3D12Device.cpp
--- a/Engine/Source/Runtime/D3D12RHI/Private/D3D12Device.cpp
+++ b/Engine/Source/Runtime/D3D12RHI/Private/D3D12Device.cpp
@@ -20,6 +20,7 @@ static uint32 GetQueryHeapPoolIndex(D3D12_QUERY_HEAP_TYPE HeapType)
 	case D3D12_QUERY_HEAP_TYPE_OCCLUSION:            return 0;
 	case D3D12_QUERY_HEAP_TYPE_TIMESTAMP:            return 1;
 	case D3D12_QUERY_HEAP_TYPE_COPY_QUEUE_TIMESTAMP: return 2;
+	case D3D12_QUERY_HEAP_TYPE_PIPELINE_STATISTICS:  return 3;
 	}
 }
 
@@ -659,18 +660,18 @@ void FD3D12Device::ReleaseCommandAllocator(FD3D12CommandAllocator* Allocator)
 	Queues[(uint32)Allocator->QueueType].ObjectPool.Allocators.Push(Allocator);
 }
 
-FD3D12CommandList* FD3D12Device::ObtainCommandList(FD3D12CommandAllocator* CommandAllocator, FD3D12QueryAllocator* TimestampAllocator)
+FD3D12CommandList* FD3D12Device::ObtainCommandList(FD3D12CommandAllocator* CommandAllocator, FD3D12QueryAllocator* TimestampAllocator, FD3D12QueryAllocator* PipelineStatsAllocator)
 {
 	check(CommandAllocator->Device == this);
 
 	FD3D12CommandList* List = Queues[(uint32)CommandAllocator->QueueType].ObjectPool.Lists.Pop();
 	if (!List)
 	{
-		List = new FD3D12CommandList(CommandAllocator, TimestampAllocator);
+		List = new FD3D12CommandList(CommandAllocator, TimestampAllocator, PipelineStatsAllocator);
 	}
 	else
 	{
-		List->Reset(CommandAllocator, TimestampAllocator);
+		List->Reset(CommandAllocator, TimestampAllocator, PipelineStatsAllocator);
 	}
 	
 	check(List);
@@ -710,13 +711,25 @@ TRefCountPtr<FD3D12QueryHeap> FD3D12Device::ObtainQueryHeap(ED3D12QueueType Queu
 				HeapType = D3D12_QUERY_HEAP_TYPE_TIMESTAMP;
 			}
 			break;
+
+		case D3D12_QUERY_TYPE_PIPELINE_STATISTICS:
+#if D3D12RHI_ENABLE_PIPELINE_STATISTICS
+			if (QueueType != ED3D12QueueType::Direct)
+			{
+				// Only graphics/direct queues support pipeline statistics
+				return nullptr;
+			}
+			HeapType = D3D12_QUERY_HEAP_TYPE_PIPELINE_STATISTICS;
+			break;
+#else
+			return nullptr;
+#endif
 	}
 
 	FD3D12QueryHeap* QueryHeap = QueryHeapPool[GetQueryHeapPoolIndex(HeapType)].Pop();
 	if (!QueryHeap)
 	{
-		// Size the query heap to fill a 64KB page
-		QueryHeap = new FD3D12QueryHeap(this, QueryType, HeapType, 65536 / FD3D12QueryHeap::ResultSize);
+		QueryHeap = new FD3D12QueryHeap(this, QueryType, HeapType);
 	}
 
 	check(QueryHeap->QueryType == QueryType);
diff --git a/Engine/Source/Runtime/D3D12RHI/Private/D3D12Device.h b/Engine/Source/Runtime/D3D12RHI/Private/D3D12Device.h
--- a/Engine/Source/Runtime/D3D12RHI/Private/D3D12Device.h
+++ b/Engine/Source/Runtime/D3D12RHI/Private/D3D12Device.h
@@ -74,6 +74,8 @@ public:
 
 	uint64 BusyCycles = 0;
 
+	D3D12_QUERY_DATA_PIPELINE_STATISTICS PipelineStats {};
+
 	uint64 GetCurrentTimestamp()  const { return Timestamps[TimestampIndex]; }
 	uint64 GetPreviousTimestamp() const { return Timestamps[TimestampIndex - 1]; }
 
@@ -116,6 +118,7 @@ public:
 	TArray<FD3D12QueryRange   > PendingQueryRanges;
 	TArray<FD3D12QueryLocation> PendingTimestampQueries;
 	TArray<FD3D12QueryLocation> PendingOcclusionQueries;
+	TArray<FD3D12QueryLocation> PendingPipelineStatsQueries;
 
 	// Executes the current payload, returning the latest fence value signaled for this queue.
 	uint64 ExecutePayload();
@@ -326,7 +329,7 @@ public:
 	void                          ReleaseQueryHeap(FD3D12QueryHeap* QueryHeap);
 	
 	// Command Lists
-	FD3D12CommandList* ObtainCommandList (FD3D12CommandAllocator* CommandAllocator, FD3D12QueryAllocator* TimestampAllocator);
+	FD3D12CommandList* ObtainCommandList (FD3D12CommandAllocator* CommandAllocator, FD3D12QueryAllocator* TimestampAllocator, FD3D12QueryAllocator* PipelineStatsAllocator);
 	void               ReleaseCommandList(FD3D12CommandList* CommandList);
 
 	// Queues
@@ -362,7 +365,7 @@ private:
 
 	FD3D12DefaultViews DefaultViews;
 
-	TD3D12ObjectPool<FD3D12QueryHeap> QueryHeapPool[3];
+	TStaticArray<TD3D12ObjectPool<FD3D12QueryHeap>, 4> QueryHeapPool { InPlace };
 	
 	FD3D12CommandContext* ImmediateCommandContext = nullptr;
 
diff --git a/Engine/Source/Runtime/D3D12RHI/Private/D3D12Query.cpp b/Engine/Source/Runtime/D3D12RHI/Private/D3D12Query.cpp
--- a/Engine/Source/Runtime/D3D12RHI/Private/D3D12Query.cpp
+++ b/Engine/Source/Runtime/D3D12RHI/Private/D3D12Query.cpp
@@ -45,12 +45,12 @@ FD3D12RenderQuery::~FD3D12RenderQuery()
 	Result = nullptr;
 }
 
-FD3D12QueryHeap::FD3D12QueryHeap(FD3D12Device* Device, D3D12_QUERY_TYPE QueryType, D3D12_QUERY_HEAP_TYPE HeapType, uint32 NumQueries)
+FD3D12QueryHeap::FD3D12QueryHeap(FD3D12Device* Device, D3D12_QUERY_TYPE QueryType, D3D12_QUERY_HEAP_TYPE HeapType)
 	: FD3D12SingleNodeGPUObject(Device->GetGPUMask())
 	, Device(Device)
 	, QueryType(QueryType)
 	, HeapType(HeapType)
-	, NumQueries(NumQueries)
+	, NumQueries(MaxHeapSize / GetResultSize())
 {
 	INC_DWORD_STAT(STAT_D3D12NumQueryHeaps);
 
@@ -73,6 +73,11 @@ FD3D12QueryHeap::FD3D12QueryHeap(FD3D12Device* Device, D3D12_QUERY_TYPE QueryTyp
 		QueryHeapName    = TEXT("Timestamp Query Heap (Copy)");
 		ResultBufferName = TEXT("Timestamp Query Heap Result Buffer (Copy)");
 		break;
+
+	case D3D12_QUERY_HEAP_TYPE_PIPELINE_STATISTICS:
+		QueryHeapName    = TEXT("Pipeline Statistics Query Heap (Copy)");
+		ResultBufferName = TEXT("Pipeline Statistics Query Heap Result Buffer (Copy)");
+		break;
 	}
 
 	const static FLazyName D3D12QueryHeapName(TEXT("FD3D12QueryHeap"));
@@ -92,7 +97,7 @@ FD3D12QueryHeap::FD3D12QueryHeap(FD3D12Device* Device, D3D12_QUERY_TYPE QueryTyp
 		SetName(D3DQueryHeap, QueryHeapName);
 
 #if ENABLE_RESIDENCY_MANAGEMENT && 0 // Temporary workaround for missing resource usage tracking for query heap
-		D3DX12Residency::Initialize(ResidencyHandle, D3DQueryHeap, ResultSize * NumQueries, this);
+		D3DX12Residency::Initialize(ResidencyHandle, D3DQueryHeap, GetResultSize() * NumQueries, this);
 		D3DX12Residency::BeginTrackingObject(Device->GetResidencyManager(), ResidencyHandle);
 #endif
 	}
@@ -104,7 +109,7 @@ FD3D12QueryHeap::FD3D12QueryHeap(FD3D12Device* Device, D3D12_QUERY_TYPE QueryTyp
 			GetGPUMask().GetNative(),
 			GetVisibilityMask().GetNative());
 			
-		D3D12_RESOURCE_DESC BufferDesc = CD3DX12_RESOURCE_DESC::Buffer(ResultSize * NumQueries);
+		D3D12_RESOURCE_DESC BufferDesc = CD3DX12_RESOURCE_DESC::Buffer(GetResultSize() * NumQueries);
 
 		// Create the readback heap
 		VERIFYD3D12RESULT(Device->GetParentAdapter()->CreateCommittedResource(
@@ -121,7 +126,7 @@ FD3D12QueryHeap::FD3D12QueryHeap(FD3D12Device* Device, D3D12_QUERY_TYPE QueryTyp
 
 	// Map the readback buffer. Resources in a readback heap are allowed
 	// to be persistently mapped, so we only need to do this once.
-	ResultPtr = (uint64*)ResultBuffer->Map();
+	ResultPtr = static_cast<uint8 const*>(ResultBuffer->Map());
 }
 
 FD3D12QueryHeap::~FD3D12QueryHeap()
diff --git a/Engine/Source/Runtime/D3D12RHI/Private/D3D12Query.h b/Engine/Source/Runtime/D3D12RHI/Private/D3D12Query.h
--- a/Engine/Source/Runtime/D3D12RHI/Private/D3D12Query.h
+++ b/Engine/Source/Runtime/D3D12RHI/Private/D3D12Query.h
@@ -35,6 +35,9 @@ struct TD3D12ResourceTraits<FRHIRenderQuery>
 // Ref-counting is used to recycle the heaps on parent device when all refering command lists have completed on the GPU.
 class FD3D12QueryHeap final : public FD3D12SingleNodeGPUObject
 {
+	// All query heaps are allocated to fill a single 64KB page
+	static constexpr uint32 MaxHeapSize = 65536;
+
 private:
 	friend FD3D12Device;
 	friend FD3D12QueryLocation;
@@ -42,7 +45,7 @@ private:
 	FD3D12QueryHeap(FD3D12QueryHeap const&) = delete;
 	FD3D12QueryHeap(FD3D12QueryHeap&&) = delete;
 
-	FD3D12QueryHeap(FD3D12Device* Device, D3D12_QUERY_TYPE QueryType, D3D12_QUERY_HEAP_TYPE HeapType, uint32 NumQueries);
+	FD3D12QueryHeap(FD3D12Device* Device, D3D12_QUERY_TYPE QueryType, D3D12_QUERY_HEAP_TYPE HeapType);
 
 public:
 	~FD3D12QueryHeap();
@@ -53,7 +56,19 @@ public:
 	uint32 const NumQueries;
 
 	// The byte size of a result for a single query
-	static constexpr uint32 ResultSize = sizeof(uint64);
+	uint32 GetResultSize() const
+	{
+		switch (QueryType)
+		{
+		default: checkNoEntry(); [[fallthrough]];
+		case D3D12_QUERY_TYPE_TIMESTAMP:
+		case D3D12_QUERY_TYPE_OCCLUSION:
+			return sizeof(uint64);
+
+		case D3D12_QUERY_TYPE_PIPELINE_STATISTICS:
+			return sizeof(D3D12_QUERY_DATA_PIPELINE_STATISTICS);
+		}
+	}
 
 	ID3D12QueryHeap* GetD3DQueryHeap() const { return D3DQueryHeap; }
 	FD3D12Resource* GetResultBuffer() const { return ResultBuffer; }
@@ -66,18 +81,31 @@ public:
 
 private:
 	TRefCountPtr<FD3D12Resource> ResultBuffer;
-	uint64* ResultPtr = nullptr;
+	uint8 const* ResultPtr = nullptr;
 	TRefCountPtr<ID3D12QueryHeap> D3DQueryHeap;
 	FD3D12ResidencyHandle ResidencyHandle;
 
 	FThreadSafeCounter NumRefs;
 };
 
-inline uint64 FD3D12QueryLocation::GetResult() const
+inline void FD3D12QueryLocation::CopyResultTo(void* Dst) const
 {
+	check(Dst);
 	check(Index < Heap->NumQueries);
 	check(Heap->ResultPtr);
-	return Heap->ResultPtr[Index];
+
+	void const* Src = Heap->ResultPtr + Index * Heap->GetResultSize();
+	FMemory::Memcpy(Dst, Src, Heap->GetResultSize());
+}
+
+template <typename TValueType>
+inline TValueType FD3D12QueryLocation::GetResult() const
+{
+	check(sizeof(TValueType) >= Heap->GetResultSize());
+
+	TValueType Value;
+	CopyResultTo(&Value);
+	return Value;
 }
 
 inline bool FD3D12QueryRange::IsFull() const
@@ -100,7 +128,7 @@ public:
 
 	// Allocate a query on a query heap, returning its location.
 	// The "target" is where the interrupt thread will write the result when completed by the GPU.
-	FD3D12QueryLocation Allocate(ED3D12QueryType Type, uint64* Target);
+	FD3D12QueryLocation Allocate(ED3D12QueryType Type, void* Target);
 
 	// Resets the allocator and returns the used query ranges
 	void CloseAndReset(TArray<FD3D12QueryRange>& OutRanges);
diff --git a/Engine/Source/Runtime/D3D12RHI/Private/D3D12Submission.cpp b/Engine/Source/Runtime/D3D12RHI/Private/D3D12Submission.cpp
--- a/Engine/Source/Runtime/D3D12RHI/Private/D3D12Submission.cpp
+++ b/Engine/Source/Runtime/D3D12RHI/Private/D3D12Submission.cpp
@@ -42,6 +42,20 @@ DECLARE_CYCLE_STAT(TEXT("GPU Total Time [Graphics]"), STAT_RHI_GPUTotalTimeGraph
 DECLARE_CYCLE_STAT(TEXT("GPU Total Time [Async Compute]"), STAT_RHI_GPUTotalTimeAsyncCompute, STATGROUP_D3D12RHI);
 DECLARE_CYCLE_STAT(TEXT("GPU Total Time [Copy]"), STAT_RHI_GPUTotalTimeCopy, STATGROUP_D3D12RHI);
 
+DECLARE_STATS_GROUP(TEXT("D3D12RHIPipeline"), STATGROUP_D3D12RHIPipeline, STATCAT_Advanced);
+
+DECLARE_DWORD_ACCUMULATOR_STAT(TEXT("GPU IA Vertices"   ), STAT_D3D12RHI_IAVertices   , STATGROUP_D3D12RHIPipeline);
+DECLARE_DWORD_ACCUMULATOR_STAT(TEXT("GPU IA Primitives" ), STAT_D3D12RHI_IAPrimitives , STATGROUP_D3D12RHIPipeline);
+DECLARE_DWORD_ACCUMULATOR_STAT(TEXT("GPU VS Invocations"), STAT_D3D12RHI_VSInvocations, STATGROUP_D3D12RHIPipeline);
+DECLARE_DWORD_ACCUMULATOR_STAT(TEXT("GPU GS Invocations"), STAT_D3D12RHI_GSInvocations, STATGROUP_D3D12RHIPipeline);
+DECLARE_DWORD_ACCUMULATOR_STAT(TEXT("GPU GS Primitives" ), STAT_D3D12RHI_GSPrimitives , STATGROUP_D3D12RHIPipeline);
+DECLARE_DWORD_ACCUMULATOR_STAT(TEXT("GPU C Invocations" ), STAT_D3D12RHI_CInvocations , STATGROUP_D3D12RHIPipeline);
+DECLARE_DWORD_ACCUMULATOR_STAT(TEXT("GPU C Primitives"  ), STAT_D3D12RHI_CPrimitives  , STATGROUP_D3D12RHIPipeline);
+DECLARE_DWORD_ACCUMULATOR_STAT(TEXT("GPU PS Invocations"), STAT_D3D12RHI_PSInvocations, STATGROUP_D3D12RHIPipeline);
+DECLARE_DWORD_ACCUMULATOR_STAT(TEXT("GPU HS Invocations"), STAT_D3D12RHI_HSInvocations, STATGROUP_D3D12RHIPipeline);
+DECLARE_DWORD_ACCUMULATOR_STAT(TEXT("GPU DS Invocations"), STAT_D3D12RHI_DSInvocations, STATGROUP_D3D12RHIPipeline);
+DECLARE_DWORD_ACCUMULATOR_STAT(TEXT("GPU CS Invocations"), STAT_D3D12RHI_CSInvocations, STATGROUP_D3D12RHIPipeline);
+
 // The maximum time, in seconds, that a submitted GPU command list is allowed to take before the RHI reports a GPU hang.
 const double SubmissionTimeOutInSeconds = 5.0;
 
@@ -231,7 +245,8 @@ void FD3D12DynamicRHI::SubmitCommands(TConstArrayView<FD3D12FinalizedCommands*>
 					Payload->SubmissionEvent ||
 					Payload->SubmissionTime.IsSet() ||
 					Payload->TimestampQueries.Num() ||
-					Payload->OcclusionQueries.Num())
+					Payload->OcclusionQueries.Num() ||
+					Payload->PipelineStatsQueries.Num())
 				{
 					bCanMergePayloads = false;
 					break;
@@ -409,8 +424,9 @@ FD3D12DynamicRHI::FProcessResult FD3D12DynamicRHI::ProcessSubmissionQueue()
 				{
 					FD3D12Queue& TargetQueue = CommandList->Device->GetQueue(CommandList->QueueType);
 
-					// Occlusion Queries
+					// Occlusion + Pipeline Stats Queries
 					TargetQueue.PendingOcclusionQueries.Append(MoveTemp(CommandList->State.OcclusionQueries));
+					TargetQueue.PendingPipelineStatsQueries.Append(MoveTemp(CommandList->State.PipelineStatsQueries));
 
 					// Timestamp Queries
 					// Keep only the first Begin() in the batch
@@ -483,6 +499,7 @@ FD3D12DynamicRHI::FProcessResult FD3D12DynamicRHI::ProcessSubmissionQueue()
 				check(Payload->QueryRanges.Num() == 0);
 				check(Payload->TimestampQueries.Num() == 0);
 				check(Payload->OcclusionQueries.Num() == 0);
+				check(Payload->PipelineStatsQueries.Num() == 0);
 
 				if (Queue->PendingQueryRanges.Num())
 				{
@@ -510,7 +527,7 @@ FD3D12DynamicRHI::FProcessResult FD3D12DynamicRHI::ProcessSubmissionQueue()
 								if (!Queue->BarrierAllocator)
 									Queue->BarrierAllocator = Queue->Device->ObtainCommandAllocator(Queue->QueueType); 
 
-								return ResolveCommandList = Queue->Device->ObtainCommandList(Queue->BarrierAllocator, nullptr);
+								return ResolveCommandList = Queue->Device->ObtainCommandList(Queue->BarrierAllocator, nullptr, nullptr);
 							};
 
 							// We've got queries to resolve. Allocate a command list.
@@ -542,15 +559,16 @@ FD3D12DynamicRHI::FProcessResult FD3D12DynamicRHI::ProcessSubmissionQueue()
 										Range.Start,
 										Range.End - Range.Start,
 										Range.Heap->GetResultBuffer()->GetResource(),
-										Range.Start * FD3D12QueryHeap::ResultSize
+										Range.Start * Range.Heap->GetResultSize()
 									);
 								}
 							}
 						}
 
-						Payload->QueryRanges      = MoveTemp(Queue->PendingQueryRanges     );
-						Payload->TimestampQueries = MoveTemp(Queue->PendingTimestampQueries);
-						Payload->OcclusionQueries = MoveTemp(Queue->PendingOcclusionQueries);
+						Payload->QueryRanges          = MoveTemp(Queue->PendingQueryRanges         );
+						Payload->TimestampQueries     = MoveTemp(Queue->PendingTimestampQueries    );
+						Payload->OcclusionQueries     = MoveTemp(Queue->PendingOcclusionQueries    );
+						Payload->PipelineStatsQueries = MoveTemp(Queue->PendingPipelineStatsQueries);
 
 						if (ResolveCommandList)
 						{
@@ -735,7 +753,7 @@ FD3D12CommandList* FD3D12DynamicRHI::GenerateBarrierCommandListAndUpdateState(FD
 		Queue.BarrierAllocator = Queue.Device->ObtainCommandAllocator(Queue.QueueType);
 
 	// Get a new command list
-	FD3D12CommandList* BarrierCommandList = Queue.Device->ObtainCommandList(Queue.BarrierAllocator, &Queue.BarrierTimestamps);
+	FD3D12CommandList* BarrierCommandList = Queue.Device->ObtainCommandList(Queue.BarrierAllocator, &Queue.BarrierTimestamps, nullptr);
 
 #if ENABLE_RESIDENCY_MANAGEMENT
 	BarrierCommandList->UpdateResidency(ResidencyHandles);
@@ -922,6 +940,22 @@ void FD3D12DynamicRHI::ProcessInterruptQueueUntil(FGraphEvent* GraphEvent)
 	}
 }
 
+D3D12_QUERY_DATA_PIPELINE_STATISTICS& operator += (D3D12_QUERY_DATA_PIPELINE_STATISTICS& LHS, D3D12_QUERY_DATA_PIPELINE_STATISTICS const& RHS)
+{
+	LHS.IAVertices	  += RHS.IAVertices;
+	LHS.IAPrimitives  += RHS.IAPrimitives;
+	LHS.VSInvocations += RHS.VSInvocations;
+	LHS.GSInvocations += RHS.GSInvocations;
+	LHS.GSPrimitives  += RHS.GSPrimitives;
+	LHS.CInvocations  += RHS.CInvocations;
+	LHS.CPrimitives	  += RHS.CPrimitives;
+	LHS.PSInvocations += RHS.PSInvocations;
+	LHS.HSInvocations += RHS.HSInvocations;
+	LHS.DSInvocations += RHS.DSInvocations;
+	LHS.CSInvocations += RHS.CSInvocations;
+	return LHS;
+}
+
 FD3D12DynamicRHI::FProcessResult FD3D12DynamicRHI::ProcessInterruptQueue()
 {
 	SCOPED_NAMED_EVENT_TEXT("InterruptQueue_Process", FColor::Yellow);
@@ -1004,7 +1038,20 @@ FD3D12DynamicRHI::FProcessResult FD3D12DynamicRHI::ProcessInterruptQueue()
 				for (FD3D12QueryLocation& Query : Payload->OcclusionQueries)
 				{
 					check(Query.Target);
-					*Query.Target = Query.GetResult();
+					Query.CopyResultTo(Query.Target);
+				}
+
+				for (FD3D12QueryLocation& Query : Payload->PipelineStatsQueries)
+				{
+					if (Query.Target)
+					{
+						Query.CopyResultTo(Query.Target);
+					}
+					else
+					{
+						// Pipeline stats queries without targets are the ones that surround whole command lists.
+						CurrentQueue.Timing->PipelineStats += Query.GetResult<D3D12_QUERY_DATA_PIPELINE_STATISTICS>();
+					}
 				}
 
 				if (Payload->TimestampQueries.Num())
@@ -1019,7 +1066,7 @@ FD3D12DynamicRHI::FProcessResult FD3D12DynamicRHI::ProcessInterruptQueue()
 					{
 						if (Query.Target)
 						{
-							*Query.Target = Query.GetResult();
+							Query.CopyResultTo(Query.Target);
 						}
 
 						switch (Query.Type)
@@ -1029,7 +1076,7 @@ FD3D12DynamicRHI::FProcessResult FD3D12DynamicRHI::ProcessInterruptQueue()
 						case ED3D12QueryType::IdleBegin:
 						case ED3D12QueryType::IdleEnd:
 							check(CurrentQueue.Timing);
-							CurrentQueue.Timing->Timestamps.Add(Query.GetResult());
+							CurrentQueue.Timing->Timestamps.Add(Query.GetResult<uint64>());
 							break;
 						}
 
@@ -1043,8 +1090,8 @@ FD3D12DynamicRHI::FProcessResult FD3D12DynamicRHI::ProcessInterruptQueue()
 						case ED3D12QueryType::CommandListEnd:
 							check(ListBegin != nullptr && IdleBegin == nullptr);
 							// Accumulate the number of ticks that have elapsed between the end of the previous command list, and the start of this one.
-							CurrentQueue.CumulativeIdleTicks += (CurrentQueue.LastEndTime != 0) ? ListBegin->GetResult() - CurrentQueue.LastEndTime : 0; //-V522
-							CurrentQueue.LastEndTime = Query.GetResult();
+							CurrentQueue.CumulativeIdleTicks += (CurrentQueue.LastEndTime != 0) ? ListBegin->GetResult<uint64>() - CurrentQueue.LastEndTime : 0; //-V522
+							CurrentQueue.LastEndTime = Query.GetResult<uint64>();
 							ListBegin = nullptr;
 							break;
 
@@ -1056,7 +1103,7 @@ FD3D12DynamicRHI::FProcessResult FD3D12DynamicRHI::ProcessInterruptQueue()
 						case ED3D12QueryType::IdleEnd:
 							check(ListBegin != nullptr && IdleBegin != nullptr);
 							// Accumulate the time this pipe spent in an idle scope. This includes vsync and waiting on other pipes.
-							CurrentQueue.CumulativeIdleTicks += Query.GetResult() - IdleBegin->GetResult(); //-V522
+							CurrentQueue.CumulativeIdleTicks += Query.GetResult<uint64>() - IdleBegin->GetResult<uint64>(); //-V522
 							IdleBegin = nullptr;
 							break;
 
@@ -1070,7 +1117,7 @@ FD3D12DynamicRHI::FProcessResult FD3D12DynamicRHI::ProcessInterruptQueue()
 							if (Query.Type == ED3D12QueryType::AdjustedMicroseconds)
 							{
 								// Convert to microseconds
-								*Query.Target = FPlatformMath::TruncToInt(double(*Query.Target) * MicrosecondsScale);
+								*static_cast<uint64*>(Query.Target) = FPlatformMath::TruncToInt(double(*static_cast<uint64*>(Query.Target)) * MicrosecondsScale);
 							}
 							break;
 						}
@@ -1219,6 +1266,24 @@ void FD3D12DynamicRHI::ProcessTimestamps(TIndirectArray<FD3D12Timing>& Timing)
 
 	check(BusyPipes == 0);
 	
+	D3D12_QUERY_DATA_PIPELINE_STATISTICS PipelineStats{};
+	for (FD3D12Timing& Current : Timing)
+	{
+		PipelineStats += Current.PipelineStats;
+	}
+
+	SET_DWORD_STAT(STAT_D3D12RHI_IAVertices   , PipelineStats.IAVertices   );
+	SET_DWORD_STAT(STAT_D3D12RHI_IAPrimitives , PipelineStats.IAPrimitives );
+	SET_DWORD_STAT(STAT_D3D12RHI_VSInvocations, PipelineStats.VSInvocations);
+	SET_DWORD_STAT(STAT_D3D12RHI_GSInvocations, PipelineStats.GSInvocations);
+	SET_DWORD_STAT(STAT_D3D12RHI_GSPrimitives , PipelineStats.GSPrimitives );
+	SET_DWORD_STAT(STAT_D3D12RHI_CInvocations , PipelineStats.CInvocations );
+	SET_DWORD_STAT(STAT_D3D12RHI_CPrimitives  , PipelineStats.CPrimitives  );
+	SET_DWORD_STAT(STAT_D3D12RHI_PSInvocations, PipelineStats.PSInvocations);
+	SET_DWORD_STAT(STAT_D3D12RHI_HSInvocations, PipelineStats.HSInvocations);
+	SET_DWORD_STAT(STAT_D3D12RHI_DSInvocations, PipelineStats.DSInvocations);
+	SET_DWORD_STAT(STAT_D3D12RHI_CSInvocations, PipelineStats.CSInvocations);
+
 	// @todo mgpu - how to handle multiple devices / queues with potentially different timestamp frequencies?
 	FD3D12Device* Device = GetAdapter().GetDevice(0);
 	double Frequency = Device->GetTimestampFrequency(ED3D12QueueType::Direct);
diff --git a/Engine/Source/Runtime/D3D12RHI/Private/D3D12Submission.h b/Engine/Source/Runtime/D3D12RHI/Private/D3D12Submission.h
--- a/Engine/Source/Runtime/D3D12RHI/Private/D3D12Submission.h
+++ b/Engine/Source/Runtime/D3D12RHI/Private/D3D12Submission.h
@@ -132,6 +132,7 @@ enum class ED3D12QueryType
 	None,
 	CommandListBegin,
 	CommandListEnd,
+	PipelineStats,
 	IdleBegin,
 	IdleEnd,
 	AdjustedMicroseconds,
@@ -151,13 +152,16 @@ struct FD3D12QueryLocation
 	ED3D12QueryType Type = ED3D12QueryType::None;
 
 	// The location into which the result is written by the interrupt thread.
-	uint64* Target = nullptr;
+	void* Target = nullptr;
 
 	// Reads the query result from the heap
-	inline uint64 GetResult() const;
+	inline void CopyResultTo(void* Dst) const;
+
+	template <typename TValueType>
+	inline TValueType GetResult() const;
 
 	FD3D12QueryLocation() = default;
-	FD3D12QueryLocation(FD3D12QueryHeap* Heap, uint32 Index, ED3D12QueryType Type, uint64* Target)
+	FD3D12QueryLocation(FD3D12QueryHeap* Heap, uint32 Index, ED3D12QueryType Type, void* Target)
 		: Heap	(Heap  )
 		, Index	(Index )
 		, Type	(Type  )
@@ -256,6 +260,7 @@ struct FD3D12PayloadBase
 	TArray<FD3D12CommandAllocator*> AllocatorsToRelease;
 	TArray<FD3D12QueryLocation> TimestampQueries;
 	TArray<FD3D12QueryLocation> OcclusionQueries;
+	TArray<FD3D12QueryLocation> PipelineStatsQueries;
 	TArray<FD3D12QueryRange> QueryRanges;
 
 	// GPU crash breadcrumbs stack
diff --git a/Engine/Source/Runtime/D3D12RHI/Private/Windows/WindowsD3D12RHIDefinitions.h b/Engine/Source/Runtime/D3D12RHI/Private/Windows/WindowsD3D12RHIDefinitions.h
--- a/Engine/Source/Runtime/D3D12RHI/Private/Windows/WindowsD3D12RHIDefinitions.h
+++ b/Engine/Source/Runtime/D3D12RHI/Private/Windows/WindowsD3D12RHIDefinitions.h
@@ -22,3 +22,6 @@ typedef ID3D12GraphicsCommandList ID3D12CopyCommandList;
 
 #define USE_STATIC_ROOT_SIGNATURE			0
 #define D3D12_USE_DUMMY_BACKBUFFER			1
+
+// Only enable pipeline statistics if we've got the stats system enabled to display them
+#define D3D12RHI_ENABLE_PIPELINE_STATISTICS (1 && STATS)
