Making sure we wait on Aftermath to finish processing GPU crash dumps before we force quit.

diff --git a/Engine/Source/Runtime/D3D12RHI/Private/D3D12Adapter.cpp b/Engine/Source/Runtime/D3D12RHI/Private/D3D12Adapter.cpp
--- a/Engine/Source/Runtime/D3D12RHI/Private/D3D12Adapter.cpp
+++ b/Engine/Source/Runtime/D3D12RHI/Private/D3D12Adapter.cpp
@@ -7,6 +7,7 @@ D3D12Adapter.cpp:D3D12 Adapter implementation.
 #include "D3D12RHIPrivate.h"
 #include "D3D12AmdExtensions.h"
 #include "D3D12IntelExtensions.h"
+#include "HAL/FileManager.h"
 #include "Misc/CommandLine.h"
 #include "Misc/EngineVersion.h"
 #include "Misc/OutputDeviceRedirector.h"
@@ -16,6 +17,7 @@ D3D12Adapter.cpp:D3D12 Adapter implementation.
 #if PLATFORM_WINDOWS
 #include "Windows/WindowsPlatformMisc.h"
 #include "Windows/WindowsPlatformStackWalk.h"
+#include "Windows/WindowsPlatformCrashContext.h"
 #endif
 #include "Modules/ModuleManager.h"
 #include "Windows/HideWindowsPlatformTypes.h"
@@ -328,14 +330,60 @@ FD3D12Adapter::FD3D12Adapter(FD3D12AdapterDesc& DescIn)
 	}
 }
 
+#if NV_AFTERMATH
 /** Callback function called when the GPU crashes, when Aftermath is enabled */
 static void D3D12AftermathCrashCallback(const void* InGPUCrashDump, const uint32_t InGPUCrashDumpSize, void* InUserData)
 {
-	// Disabled for now, let the regular rendering code catch the error on the next API call via VERIFYD3D12RESULT & co.
-	// Note that this means we won't be getting aftermath crash dump files anymore, since those are written based on the data passed to this function via InGPUCrashDump.
-	// D3D12RHI::TerminateOnGPUCrash(nullptr, InGPUCrashDump, InGPUCrashDumpSize);
+	// If we have crash dump data then dump to disc
+	if (InGPUCrashDump != nullptr)
+	{
+		// Write out crash dump to project log dir - exception handling code will take care of copying it to the correct location
+		const FString GpuMiniDumpPath = FPaths::Combine(FPaths::ProjectLogDir(), FWindowsPlatformCrashContext::UEGPUAftermathMinidumpName);
+
+		UE_LOG(LogD3D12RHI, Error, TEXT("Aftermath: Writing Aftermath dump to: %s"), *GpuMiniDumpPath);
+
+		if (FArchive* Writer = IFileManager::Get().CreateFileWriter(*GpuMiniDumpPath))
+		{
+			Writer->Serialize((void*)InGPUCrashDump, InGPUCrashDumpSize);
+			Writer->Close();
+		}
+	}
 }
 
+void EnableNVAftermathCrashDumps(ED3D12GPUCrashDebuggingModes GPUCrashDebuggingModes)
+{
+	// GPUcrash dump handler must be attached prior to device creation
+	if (GDX12NVAfterMathModuleLoaded && EnumHasAnyFlags(GPUCrashDebuggingModes, ED3D12GPUCrashDebuggingModes::NvAftermath))
+	{
+		const HANDLE CurrentThread = ::GetCurrentThread();
+
+		const GFSDK_Aftermath_Result Result = GFSDK_Aftermath_EnableGpuCrashDumps(
+			GFSDK_Aftermath_Version_API,
+			GFSDK_Aftermath_GpuCrashDumpWatchedApiFlags_DX,
+			GFSDK_Aftermath_GpuCrashDumpFeatureFlags_Default,
+			&D3D12AftermathCrashCallback,
+			nullptr, //Shader debug callback
+			nullptr, // description callback
+			nullptr, // resolve marker callback
+			CurrentThread
+		); // user data
+
+		if (Result == GFSDK_Aftermath_Result_Success)
+		{
+			UE_LOG(LogD3D12RHI, Log, TEXT("[Aftermath] Aftermath crash dumping enabled"));
+
+			// enable core Aftermath to set the init flags
+			GDX12NVAfterMathEnabled = 1;
+		}
+		else
+		{
+			UE_LOG(LogD3D12RHI, Log, TEXT("[Aftermath] Aftermath crash dumping failed to initialize (%x)"), Result);
+
+			GDX12NVAfterMathEnabled = 0;
+		}
+	}
+}
+#endif
 
 void FD3D12Adapter::CreateRootDevice(bool bWithDebug)
 {
@@ -382,35 +430,7 @@ void FD3D12Adapter::CreateRootDevice(bool bWithDebug)
 #if NV_AFTERMATH
 	if (IsRHIDeviceNVIDIA() && GDX12NVAfterMathModuleLoaded)
 	{
-		// GPUcrash dump handler must be attached prior to device creation
-		if (EnumHasAnyFlags(GPUCrashDebuggingModes, ED3D12GPUCrashDebuggingModes::NvAftermath))
-		{
-			HANDLE CurrentThread = ::GetCurrentThread();
-
-			GFSDK_Aftermath_Result Result = GFSDK_Aftermath_EnableGpuCrashDumps(
-				GFSDK_Aftermath_Version_API,
-				GFSDK_Aftermath_GpuCrashDumpWatchedApiFlags_DX,
-				GFSDK_Aftermath_GpuCrashDumpFeatureFlags_Default,
-				&D3D12AftermathCrashCallback,
-				nullptr, //Shader debug callback
-				nullptr, // description callback
-				nullptr, // resolve marker callback
-				CurrentThread); // user data
-
-			if (Result == GFSDK_Aftermath_Result_Success)
-			{
-				UE_LOG(LogD3D12RHI, Log, TEXT("[Aftermath] Aftermath crash dumping enabled"));
-
-				// enable core Aftermath to set the init flags
-				GDX12NVAfterMathEnabled = 1;
-			}
-			else
-			{
-				UE_LOG(LogD3D12RHI, Log, TEXT("[Aftermath] Aftermath crash dumping failed to initialize (%x)"), Result);
-
-				GDX12NVAfterMathEnabled = 0;
-			}
-		}
+		EnableNVAftermathCrashDumps(GPUCrashDebuggingModes);
 	}
 #endif
 
diff --git a/Engine/Source/Runtime/D3D12RHI/Private/D3D12NvidiaExtensions.h b/Engine/Source/Runtime/D3D12RHI/Private/D3D12NvidiaExtensions.h
--- a/Engine/Source/Runtime/D3D12RHI/Private/D3D12NvidiaExtensions.h
+++ b/Engine/Source/Runtime/D3D12RHI/Private/D3D12NvidiaExtensions.h
@@ -24,5 +24,6 @@
 	extern bool GDX12NVAfterMathModuleLoaded;
 	extern int32 GDX12NVAfterMathEnabled;
 	extern int32 GDX12NVAfterMathTrackResources;
+	extern float GDX12NVAfterMathDumpWaitTime;
 	extern int32 GDX12NVAfterMathMarkers;
 #endif // NV_AFTERMATH
diff --git a/Engine/Source/Runtime/D3D12RHI/Private/D3D12Util.cpp b/Engine/Source/Runtime/D3D12RHI/Private/D3D12Util.cpp
--- a/Engine/Source/Runtime/D3D12RHI/Private/D3D12Util.cpp
+++ b/Engine/Source/Runtime/D3D12RHI/Private/D3D12Util.cpp
@@ -9,8 +9,10 @@ D3D12Util.h: D3D RHI utility implementation.
 #include "RendererInterface.h"
 #include "CoreGlobals.h"
 #include "Misc/OutputDeviceRedirector.h"
-#include "Windows/WindowsPlatformCrashContext.h"
 #include "HAL/ExceptionHandling.h"
+#if PLATFORM_WINDOWS
+#include "HAL/PlatformCrashContext.h"
+#endif
 
 #define D3DERR(x) case x: ErrorCodeText = TEXT(#x); break;
 #define LOCTEXT_NAMESPACE "Developer.MessageLog"
@@ -508,7 +510,7 @@ static bool LogDREDData(ID3D12Device* Device, bool bTrackingAllAllocations, D3D1
 			UE_LOG(LogD3D12RHI, Error, TEXT("DRED: No breadcrumb head found."));
 		}
 
-		FGenericCrashContext::SetEngineData(TEXT("RHI.DRED.HasBreadcrumbData"), bHasValidBreadcrumbData ? TEXT("true") : TEXT("false"));
+		FPlatformCrashContext::SetEngineData(TEXT("RHI.DRED.HasBreadcrumbData"), bHasValidBreadcrumbData ? TEXT("true") : TEXT("false"));
 
 		bool bHasValidPageFaultData = false;
 		D3D12_DRED_PAGE_FAULT_OUTPUT DredPageFaultOutput;
@@ -559,7 +561,7 @@ static bool LogDREDData(ID3D12Device* Device, bool bTrackingAllAllocations, D3D1
 			UE_LOG(LogD3D12RHI, Error, TEXT("DRED: No PageFault data."));
 		}
 
-		FGenericCrashContext::SetEngineData(TEXT("RHI.DRED.HasPageFaultData"), bHasValidPageFaultData ? TEXT("true") : TEXT("false"));
+		FPlatformCrashContext::SetEngineData(TEXT("RHI.DRED.HasPageFaultData"), bHasValidPageFaultData ? TEXT("true") : TEXT("false"));
 
 		return true;
 	}
@@ -757,7 +759,7 @@ static void TerminateOnOutOfMemory(ID3D12Device* InDevice, HRESULT D3DResult, bo
 
 namespace D3D12RHI
 {
-	void TerminateOnGPUCrash(ID3D12Device* InDevice, const void* InGPUCrashDump, const size_t InGPUCrashDumpSize)
+	void TerminateOnGPUCrash(ID3D12Device* InDevice)
 	{		
 		// This function can be called outside of VerifyD3D12Result & co, so it uses its own critical section to make sure it's not re-entered.
 		static FCriticalSection cs;
@@ -799,6 +801,26 @@ namespace D3D12RHI
 				}
 			});
 #endif  // PLATFORM_WINDOWS
+
+#if NV_AFTERMATH
+		GFSDK_Aftermath_CrashDump_Status AftermathStatus{};
+		if (GDX12NVAfterMathEnabled)
+		{
+			GFSDK_Aftermath_GetCrashDumpStatus(&AftermathStatus);
+			if (AftermathStatus != GFSDK_Aftermath_CrashDump_Status_Unknown && AftermathStatus != GFSDK_Aftermath_CrashDump_Status_NotStarted)
+			{
+				const float StartTime = FPlatformTime::Seconds();
+				const float EndTime = StartTime + GDX12NVAfterMathDumpWaitTime;
+				while (AftermathStatus != GFSDK_Aftermath_CrashDump_Status_CollectingDataFailed
+					&& AftermathStatus != GFSDK_Aftermath_CrashDump_Status_Finished
+					&& FPlatformTime::Seconds() < EndTime)
+				{
+					FPlatformProcess::Sleep(0.01f);
+					GFSDK_Aftermath_GetCrashDumpStatus(&AftermathStatus);
+				}
+			}
+		}
+#endif
 		
 		// Build the error message
 		FTextBuilder ErrorMessage;
@@ -817,10 +839,12 @@ namespace D3D12RHI
 		}
 
 		// And info on gpu crash dump as well
-		if (InGPUCrashDump)
+#if NV_AFTERMATH
+		if (AftermathStatus == GFSDK_Aftermath_CrashDump_Status_Finished)
 		{
-			ErrorMessage.AppendLine(LOCTEXT("GPU CrashDump", "\nA GPU mini dump will be saved in the Crashes folder."));
+			ErrorMessage.AppendLine(LOCTEXT("GPU CrashDump", "\nA GPU mini dump was be saved in the Logs folder."));
 		}
+#endif
 		
 		// Make sure the log is flushed!
 		GLog->Panic();
@@ -837,34 +861,18 @@ namespace D3D12RHI
 			UE_LOG(LogD3D12RHI, D3D12RHI_GPU_CRASH_LOG_VERBOSITY, TEXT("%s"), *ErrorMessage.ToText().ToString());
 		}
 
-#if PLATFORM_WINDOWS
-		// If we have crash dump data then dump to disc
-		if (InGPUCrashDump != nullptr)
-		{
-			// Write out crash dump to project log dir - exception handling code will take care of copying it to the correct location
-			const FString GPUMiniDumpPath = FPaths::Combine(FPaths::ProjectLogDir(), FWindowsPlatformCrashContext::UEGPUAftermathMinidumpName);
-
-			// Just use raw windows file routines for the GPU minidump (TODO: refactor to our own functions?)
-			HANDLE FileHandle = CreateFileW(*GPUMiniDumpPath, GENERIC_WRITE, 0, NULL, CREATE_ALWAYS, FILE_ATTRIBUTE_NORMAL, NULL);
-			if (FileHandle != INVALID_HANDLE_VALUE)
-			{
-				WriteFile(FileHandle, InGPUCrashDump, InGPUCrashDumpSize, nullptr, nullptr);
-			}
-			CloseHandle(FileHandle);
-
-			// Report the GPU crash which will raise the exception (only interesting if we have a GPU dump)
-			ReportGPUCrash(TEXT("Aftermath GPU Crash dump Triggered"), nullptr);
-
-			// Force shutdown, we can't do anything useful anymore.
-			FPlatformMisc::RequestExit(true, TEXT("D3D12Util.TerminateOnGPUCrash"));
-		}
-#endif // PLATFORM_WINDOWS
-
 		// hard break here when the debugger is attached
 		if (IsDebuggerPresent())
 		{
 			UE_DEBUG_BREAK();
 		}
+
+#if PLATFORM_WINDOWS
+		ReportGPUCrash(TEXT("GPU Crash dump Triggered"), nullptr);
+#endif
+
+		// Force shutdown, we can't do anything useful anymore.
+		FPlatformMisc::RequestExit(true, TEXT("D3D12Util.TerminateOnGPUCrash"));
 	}
 
 	// It's possible for multiple threads to catch GPU crashes or other D3D errors at the same time. Make sure we only log the error once by acquiring
@@ -886,7 +894,7 @@ namespace D3D12RHI
 		}
 		else if (D3DResult == DXGI_ERROR_DEVICE_REMOVED || D3DResult == DXGI_ERROR_DEVICE_HUNG || D3DResult == DXGI_ERROR_DEVICE_RESET)
 		{
-			TerminateOnGPUCrash(Device, nullptr, 0);
+			TerminateOnGPUCrash(Device);
 		}
 
 		// Make sure the log is flushed!
@@ -924,7 +932,7 @@ namespace D3D12RHI
 		// Terminate with device removed but we don't have any GPU crash dump information
 		if (D3DResult == DXGI_ERROR_DEVICE_REMOVED || D3DResult == DXGI_ERROR_DEVICE_HUNG || D3DResult == DXGI_ERROR_DEVICE_RESET)
 		{
-			TerminateOnGPUCrash(Device, nullptr, 0);
+			TerminateOnGPUCrash(Device);
 		}
 		else if (D3DResult == E_OUTOFMEMORY)
 		{
diff --git a/Engine/Source/Runtime/D3D12RHI/Private/D3D12Util.h b/Engine/Source/Runtime/D3D12RHI/Private/D3D12Util.h
--- a/Engine/Source/Runtime/D3D12RHI/Private/D3D12Util.h
+++ b/Engine/Source/Runtime/D3D12RHI/Private/D3D12Util.h
@@ -18,7 +18,7 @@ namespace D3D12RHI
 	/**
 	 * Dump & Log all the information we have on a GPU crash (NvAfterMath, DRED, Breadcrumbs, ...) and force quit
 	 */
-	extern void TerminateOnGPUCrash(ID3D12Device* InDevice, const void* InGPUCrashDump, const size_t InGPUCrashDumpSize);
+	CA_NO_RETURN void TerminateOnGPUCrash(ID3D12Device* InDevice);
 
 	/**
 	 * Checks that the given result isn't a failure.  If it is, the application exits with an appropriate error message.
diff --git a/Engine/Source/Runtime/D3D12RHI/Private/Windows/WindowsD3D12Device.cpp b/Engine/Source/Runtime/D3D12RHI/Private/Windows/WindowsD3D12Device.cpp
--- a/Engine/Source/Runtime/D3D12RHI/Private/Windows/WindowsD3D12Device.cpp
+++ b/Engine/Source/Runtime/D3D12RHI/Private/Windows/WindowsD3D12Device.cpp
@@ -46,19 +46,27 @@ FD3D12DynamicRHI* GD3D12RHI = nullptr;
 
 	// Disabled by default since introduces stalls between render and driver threads
 	int32 GDX12NVAfterMathEnabled = 0;
-	static FAutoConsoleVariableRef CVarDX12NVAfterMathBufferSize(
-	TEXT("r.DX12NVAfterMathEnabled"),
-	GDX12NVAfterMathEnabled,
-	TEXT("Use NV Aftermath for GPU crash analysis in D3D12"),
-	ECVF_ReadOnly
+	static FAutoConsoleVariableRef CVarDX12NVAfterMathEnabled(
+		TEXT("r.DX12NVAfterMathEnabled"),
+		GDX12NVAfterMathEnabled,
+		TEXT("Use NV Aftermath for GPU crash analysis in D3D12"),
+		ECVF_ReadOnly
 	);
 
 	int32 GDX12NVAfterMathTrackResources = 0;
 	static FAutoConsoleVariableRef CVarDX12NVAfterMathTrackResources(
-	TEXT("r.DX12NVAfterMathTrackResources"),
-	GDX12NVAfterMathTrackResources,
-	TEXT("Enable NV Aftermath resource tracing in D3D12"),
-	ECVF_ReadOnly
+		TEXT("r.DX12NVAfterMathTrackResources"),
+		GDX12NVAfterMathTrackResources,
+		TEXT("Enable NV Aftermath resource tracing in D3D12"),
+		ECVF_ReadOnly
+	);
+
+	float GDX12NVAfterMathDumpWaitTime = 10.0f;
+	static FAutoConsoleVariableRef CVarDX12NVAfterMathDumpWaitTime(
+		TEXT("r.DX12NVAfterMathDumpWaitTime"),
+		GDX12NVAfterMathDumpWaitTime,
+		TEXT("Amount of time to wait for NV Aftermath to finish processing GPU crash dumps."),
+		ECVF_Default
 	);
 
 	int32 GDX12NVAfterMathMarkers = 0;
diff --git a/Engine/Source/Runtime/VulkanRHI/Private/VulkanUtil.cpp b/Engine/Source/Runtime/VulkanRHI/Private/VulkanUtil.cpp
--- a/Engine/Source/Runtime/VulkanRHI/Private/VulkanUtil.cpp
+++ b/Engine/Source/Runtime/VulkanRHI/Private/VulkanUtil.cpp
@@ -29,6 +29,16 @@ static const uint32	EventDeepCRC = FCrc::StrCrc32<TCHAR>(*EventDeepString);
 static const uint32 BUFFERED_TIMING_QUERIES = 1;
 static const uint32 TIMING_QUERY_RETRIES = 1;
 
+#if NV_AFTERMATH
+	float GVulkanNVAfterMathDumpWaitTime = 10.0f;
+	static FAutoConsoleVariableRef CVarVulkanNVAfterMathDumpWaitTime(
+		TEXT("r.VulkanNVAfterMathDumpWaitTime"),
+		GVulkanNVAfterMathDumpWaitTime,
+		TEXT("Amount of time to wait for NV Aftermath to finish processing GPU crash dumps."),
+		ECVF_Default
+	);
+#endif
+
 /**
  * Initializes the static variables, if necessary.
  */
@@ -1093,6 +1103,27 @@ namespace VulkanRHI
 #endif
 
 			CheckDeviceFault(Device);
+
+			// Make sure we wait on the Aftermath crash dump before we crash.
+#if NV_AFTERMATH
+			if (GGPUCrashDebuggingEnabled && GVulkanNVAftermathModuleLoaded)
+			{
+				GFSDK_Aftermath_CrashDump_Status AftermathStatus{};
+				GFSDK_Aftermath_GetCrashDumpStatus(&AftermathStatus);
+				if (AftermathStatus != GFSDK_Aftermath_CrashDump_Status_Unknown && AftermathStatus != GFSDK_Aftermath_CrashDump_Status_NotStarted)
+				{
+					const float StartTime = FPlatformTime::Seconds();
+					const float EndTime = StartTime + GVulkanNVAfterMathDumpWaitTime;
+					while (AftermathStatus != GFSDK_Aftermath_CrashDump_Status_CollectingDataFailed
+						&& AftermathStatus != GFSDK_Aftermath_CrashDump_Status_Finished
+						&& FPlatformTime::Seconds() < EndTime)
+					{
+						FPlatformProcess::Sleep(0.01f);
+						GFSDK_Aftermath_GetCrashDumpStatus(&AftermathStatus);
+					}
+				}
+			}
+#endif
 		}
 
 #if UE_BUILD_DEBUG || UE_BUILD_DEVELOPMENT
diff --git a/Engine/Source/Runtime/Windows/D3D11RHI/Private/D3D11RHIPrivate.h b/Engine/Source/Runtime/Windows/D3D11RHI/Private/D3D11RHIPrivate.h
--- a/Engine/Source/Runtime/Windows/D3D11RHI/Private/D3D11RHIPrivate.h
+++ b/Engine/Source/Runtime/Windows/D3D11RHI/Private/D3D11RHIPrivate.h
@@ -47,6 +47,7 @@ DECLARE_LOG_CATEGORY_EXTERN(LogD3D11RHI, Log, All);
 #undef GFSDK_Aftermath_WITH_DX11
 extern bool GDX11NVAfterMathEnabled;
 extern bool GDX11NVAfterMathMarkers;
+extern float GDX11NVAfterMathDumpWaitTime;
 #endif
 
 #if INTEL_METRICSDISCOVERY
diff --git a/Engine/Source/Runtime/Windows/D3D11RHI/Private/D3D11Util.cpp b/Engine/Source/Runtime/Windows/D3D11RHI/Private/D3D11Util.cpp
--- a/Engine/Source/Runtime/Windows/D3D11RHI/Private/D3D11Util.cpp
+++ b/Engine/Source/Runtime/Windows/D3D11RHI/Private/D3D11Util.cpp
@@ -8,6 +8,7 @@
 #include "EngineModule.h"
 #include "RendererInterface.h"
 #include "ProfilingDebugging/ScopedDebugInfo.h"
+#include "HAL/ExceptionHandling.h"
 
 #define D3DERR(x) case x: ErrorCodeText = TEXT(#x); break;
 #define LOCTEXT_NAMESPACE "Developer.MessageLog"
@@ -98,14 +99,29 @@ static void TerminateOnDeviceRemoved(HRESULT D3DResult, ID3D11Device* Direct3DDe
 	if (D3DResult == DXGI_ERROR_DEVICE_REMOVED)
 	{
 #if NV_AFTERMATH
-		uint32 Result = 0xffffffff;
+		GFSDK_Aftermath_Result Result{};
 		uint32 bDeviceActive = 0;
 		if (GDX11NVAfterMathEnabled)
 		{
+			// Wait until the Aftermath crash dump has been handled.
+			GFSDK_Aftermath_CrashDump_Status AftermathStatus{};
+			GFSDK_Aftermath_GetCrashDumpStatus(&AftermathStatus);
+			if (AftermathStatus != GFSDK_Aftermath_CrashDump_Status_Unknown && AftermathStatus != GFSDK_Aftermath_CrashDump_Status_NotStarted)
+			{
+				const float StartTime = FPlatformTime::Seconds();
+				const float EndTime = StartTime + GDX11NVAfterMathDumpWaitTime;
+				while (AftermathStatus != GFSDK_Aftermath_CrashDump_Status_CollectingDataFailed
+					&& AftermathStatus != GFSDK_Aftermath_CrashDump_Status_Finished
+					&& FPlatformTime::Seconds() < EndTime)
+				{
+					FPlatformProcess::Sleep(0.01f);
+					GFSDK_Aftermath_GetCrashDumpStatus(&AftermathStatus);
+				}
+			}
+
 			GFSDK_Aftermath_Device_Status Status;
-			auto Res = GFSDK_Aftermath_GetDeviceStatus(&Status);
-			Result = uint32(Res);
-			if (Res == GFSDK_Aftermath_Result_Success)
+			Result = GFSDK_Aftermath_GetDeviceStatus(&Status);
+			if (Result == GFSDK_Aftermath_Result_Success)
 			{
 				bDeviceActive = Status == GFSDK_Aftermath_Device_Status_Active ? 1 : 0;
 			}
@@ -115,6 +131,9 @@ static void TerminateOnDeviceRemoved(HRESULT D3DResult, ID3D11Device* Direct3DDe
 		UE_LOG(LogD3D11RHI, Log, TEXT("[Aftermath] NV_AFTERMATH is not set"));
 #endif
 
+		// Report the GPU crash which will raise the exception
+		ReportGPUCrash(TEXT("GPU Crash dump Triggered"), nullptr);
+
 		GIsGPUCrashed = true;		
 		if (Direct3DDevice)
 		{
diff --git a/Engine/Source/Runtime/Windows/D3D11RHI/Private/Windows/WindowsD3D11Device.cpp b/Engine/Source/Runtime/Windows/D3D11RHI/Private/Windows/WindowsD3D11Device.cpp
--- a/Engine/Source/Runtime/Windows/D3D11RHI/Private/Windows/WindowsD3D11Device.cpp
+++ b/Engine/Source/Runtime/Windows/D3D11RHI/Private/Windows/WindowsD3D11Device.cpp
@@ -5,6 +5,7 @@
 =============================================================================*/
 #include "Misc/EngineVersion.h"
 #include "D3D11RHIPrivate.h"
+#include "HAL/FileManager.h"
 #include "Misc/CommandLine.h"
 #include "Misc/EngineVersion.h"
 #include "Windows/AllowWindowsPlatformTypes.h"
@@ -25,7 +26,6 @@
 #include "GenericPlatform/GenericPlatformCrashContext.h"
 #include "RHIValidation.h"
 #include "RHIUtilities.h"
-#include "HAL/ExceptionHandling.h"
 #include "HDRHelper.h"
 #include "GlobalShader.h"
 
@@ -33,6 +33,14 @@
 bool GDX11NVAfterMathEnabled = false;
 bool GNVAftermathModuleLoaded = false;
 bool GDX11NVAfterMathMarkers = false;
+
+float GDX11NVAfterMathDumpWaitTime = 10.0f;
+static FAutoConsoleVariableRef CVarDX12NVAfterMathDumpWaitTime(
+	TEXT("r.DX11NVAfterMathDumpWaitTime"),
+	GDX11NVAfterMathDumpWaitTime,
+	TEXT("Amount of time to wait for NV Aftermath to finish processing GPU crash dumps."),
+	ECVF_Default
+);
 #endif
 
 #if INTEL_METRICSDISCOVERY
@@ -1404,19 +1412,20 @@ static void D3D11AftermathCrashCallback(const void* InGPUCrashDump, const uint32
 		GDynamicRHI->CheckGpuHeartbeat();
 	}
 
-	// Write out crash dump to project log dir - exception handling code will take care of copying it to the correct location
-	const FString GPUMiniDumpPath = FPaths::Combine(FPaths::ProjectLogDir(), FWindowsPlatformCrashContext::UEGPUAftermathMinidumpName);
-
-	// Just use raw windows file routines for the GPU minidump (TODO: refactor to our own functions?)
-	HANDLE FileHandle = CreateFileW(*GPUMiniDumpPath, GENERIC_WRITE, 0, NULL, CREATE_ALWAYS, FILE_ATTRIBUTE_NORMAL, NULL);
-	if (FileHandle != INVALID_HANDLE_VALUE)
+	// If we have crash dump data then dump to disc
+	if (InGPUCrashDump)
 	{
-		WriteFile(FileHandle, InGPUCrashDump, InGPUCrashDumpSize, nullptr, nullptr);
-	}
-	CloseHandle(FileHandle);
+		// Write out crash dump to project log dir - exception handling code will take care of copying it to the correct location
+		const FString GpuMiniDumpPath = FPaths::Combine(FPaths::ProjectLogDir(), FWindowsPlatformCrashContext::UEGPUAftermathMinidumpName);
+
+		UE_LOG(LogD3D11RHI, Error, TEXT("Aftermath: Writing Aftermath dump to: %s"), *GpuMiniDumpPath);
 
-	// Report the GPU crash which will raise the exception
-	ReportGPUCrash(TEXT("Aftermath GPU Crash dump Triggered"), nullptr);
+		if (FArchive* Writer = IFileManager::Get().CreateFileWriter(*GpuMiniDumpPath))
+		{
+			Writer->Serialize((void*)InGPUCrashDump, InGPUCrashDumpSize);
+			Writer->Close();
+		}
+	}
 }
 
 void EnableNVAftermathCrashDumps()
@@ -1426,7 +1435,6 @@ void EnableNVAftermathCrashDumps()
 		static IConsoleVariable* GPUCrashDump = IConsoleManager::Get().FindConsoleVariable(TEXT("r.GPUCrashDump"));
 		if (FParse::Param(FCommandLine::Get(), TEXT("gpucrashdump")) || (GPUCrashDump && GPUCrashDump->GetInt()))
 		{
-
 			GFSDK_Aftermath_Result Result = GFSDK_Aftermath_EnableGpuCrashDumps(
 				GFSDK_Aftermath_Version_API,
 				GFSDK_Aftermath_GpuCrashDumpWatchedApiFlags_DX,
