Add implementation of `ClearPSODriverCache` for Intel.

Intel drivers store the shader cache in `AppData\LocalLow\Intel\ShaderCache`.

#pf skip

diff --git a/Engine/Source/Runtime/D3D12RHI/Private/Windows/WindowsD3D12Device.cpp b/Engine/Source/Runtime/D3D12RHI/Private/Windows/WindowsD3D12Device.cpp
--- a/Engine/Source/Runtime/D3D12RHI/Private/Windows/WindowsD3D12Device.cpp
+++ b/Engine/Source/Runtime/D3D12RHI/Private/Windows/WindowsD3D12Device.cpp
@@ -1289,21 +1289,23 @@ static void ClearPSODriverCache()
 
 		if (IsRHIDeviceNVIDIA())
 		{
-			FString PSOPath = FString::Printf(L"%s\\NVIDIA\\DXCache", *LocalAppDataFolder);
+			FString PSOPath = FPaths::Combine(*LocalAppDataFolder, TEXT("NVIDIA"), TEXT("DXCache"));
 			ClearFolder(PSOPath, TEXT(".bin"));
 			ClearFolder(PSOPath, TEXT(".toc"));
 		}
 		else if (IsRHIDeviceAMD())
 		{
-			FString PSOPath = FString::Printf(L"%s\\AMD\\DxCache", *LocalAppDataFolder);
+			FString PSOPath = FPaths::Combine(*LocalAppDataFolder, TEXT("AMD"), TEXT("DxCache"));
 			ClearFolder(PSOPath, TEXT(".bin"));
 
-			PSOPath = FString::Printf(L"%s\\AMD\\DxcCache", *LocalAppDataFolder);
+			PSOPath = FPaths::Combine(*LocalAppDataFolder, TEXT("AMD"), TEXT("DxcCache"));
 			ClearFolder(PSOPath, TEXT(".parc"));
 		}
 		else if (IsRHIDeviceIntel())
 		{
-			UE_LOG(LogD3D12RHI, Warning, TEXT("clearPSODriverCache not implemented yet for Intel"));
+			// Intel stores the cache in LocalLow.
+			FString PSOPath = FPaths::Combine(*LocalAppDataFolder, TEXT(".."), TEXT("LocalLow"), TEXT("Intel"), TEXT("ShaderCache"));
+			ClearFolder(PSOPath, nullptr);
 		}
 	}
 	else
