Deprecating StreamOut symbols since they're unused.

diff --git a/Engine/Source/Runtime/D3D12RHI/Private/D3D12DescriptorCache.cpp b/Engine/Source/Runtime/D3D12RHI/Private/D3D12DescriptorCache.cpp
--- a/Engine/Source/Runtime/D3D12RHI/Private/D3D12DescriptorCache.cpp
+++ b/Engine/Source/Runtime/D3D12RHI/Private/D3D12DescriptorCache.cpp
@@ -326,11 +326,6 @@ void FD3D12DescriptorCache::SetRenderTargets(FD3D12RenderTargetView** RenderTarg
 	}
 }
 
-void FD3D12DescriptorCache::SetStreamOutTargets(FD3D12Resource** Buffers, uint32 Count, const uint32* Offsets)
-{
-	check(Count == 0);
-}
-
 void FD3D12DescriptorCache::SetSamplers(EShaderFrequency ShaderStage, const FD3D12RootSignature* RootSignature, FD3D12SamplerStateCache& Cache, const SamplerSlotMask& SlotsNeededMask, uint32 SlotsNeeded, uint32& HeapSlot)
 {
 	check(!UsingGlobalSamplerHeap());
diff --git a/Engine/Source/Runtime/D3D12RHI/Private/D3D12DescriptorCache.h b/Engine/Source/Runtime/D3D12RHI/Private/D3D12DescriptorCache.h
--- a/Engine/Source/Runtime/D3D12RHI/Private/D3D12DescriptorCache.h
+++ b/Engine/Source/Runtime/D3D12RHI/Private/D3D12DescriptorCache.h
@@ -353,8 +353,6 @@ public:
 
 	void PrepareBindlessViews(EShaderFrequency ShaderStage, TConstArrayView<FD3D12ShaderResourceView*> SRVs, TConstArrayView<FD3D12UnorderedAccessView*> UAVs);
 
-	void SetStreamOutTargets(FD3D12Resource **Buffers, uint32 Count, const uint32* Offsets);
-
 	bool HeapRolledOver(ERHIDescriptorHeapType InHeapType);
 	void HeapLoopedAround(ERHIDescriptorHeapType InHeapType);
 	void Init(uint32 InNumLocalViewDescriptors, uint32 InNumSamplerDescriptors);
diff --git a/Engine/Source/Runtime/D3D12RHI/Private/D3D12StateCache.cpp b/Engine/Source/Runtime/D3D12RHI/Private/D3D12StateCache.cpp
--- a/Engine/Source/Runtime/D3D12RHI/Private/D3D12StateCache.cpp
+++ b/Engine/Source/Runtime/D3D12RHI/Private/D3D12StateCache.cpp
@@ -174,7 +174,6 @@ void FD3D12StateCache::DirtyStateForNewCommandList()
 	// IndexBuffers are set in DrawIndexed*() calls, so there's no way to depend on previously set IndexBuffers without making a new DrawIndexed*() call.
 	PipelineState.Graphics.IBCache.Clear();
 
-	if (PipelineState.Graphics.CurrentNumberOfStreamOutTargets) { bNeedSetSOs = true; }
 	if (PipelineState.Graphics.CurrentNumberOfRenderTargets || PipelineState.Graphics.CurrentDepthStencilTarget) { bNeedSetRTs = true; }
 	if (PipelineState.Graphics.CurrentNumberOfViewports) { bNeedSetViewports = true; }
 	if (PipelineState.Graphics.CurrentNumberOfScissorRects) { bNeedSetScissorRects = true; }
@@ -213,7 +212,6 @@ void FD3D12StateCache::DirtyState()
 	PipelineState.Compute.bNeedSetRootSignature = true;
 	PipelineState.Graphics.bNeedSetRootSignature = true;
 	bNeedSetVB = true;
-	bNeedSetSOs = true;
 	bNeedSetRTs = true;
 	bNeedSetViewports = true;
 	bNeedSetScissorRects = true;
@@ -465,11 +463,6 @@ void FD3D12StateCache::ApplyState(ED3D12PipelineType PipelineType)
 			//SCOPE_CYCLE_COUNTER(STAT_D3D12ApplyStateSetVertexBufferTime);
 			DescriptorCache.SetVertexBuffers(PipelineState.Graphics.VBCache);
 		}
-		if (bNeedSetSOs)
-		{
-			bNeedSetSOs = false;
-			DescriptorCache.SetStreamOutTargets(PipelineState.Graphics.CurrentStreamOutTargets, PipelineState.Graphics.CurrentNumberOfStreamOutTargets, PipelineState.Graphics.CurrentSOOffsets);
-		}
 		if (bNeedSetViewports)
 		{
 			bNeedSetViewports = false;
diff --git a/Engine/Source/Runtime/D3D12RHI/Private/D3D12StateCachePrivate.h b/Engine/Source/Runtime/D3D12RHI/Private/D3D12StateCachePrivate.h
--- a/Engine/Source/Runtime/D3D12RHI/Private/D3D12StateCachePrivate.h
+++ b/Engine/Source/Runtime/D3D12RHI/Private/D3D12StateCachePrivate.h
@@ -286,7 +286,6 @@ protected:
 
 	bool bNeedSetVB = true;
 	bool bNeedSetRTs = true;
-	bool bNeedSetSOs = true;
 	bool bNeedSetViewports = true;
 	bool bNeedSetScissorRects = true;
 	bool bNeedSetPrimitiveTopology = true;
@@ -329,11 +328,6 @@ protected:
 			// Vertex Buffer State
 			FD3D12VertexBufferCache VBCache = {};
 
-			// SO
-			uint32			CurrentNumberOfStreamOutTargets = 0;
-			FD3D12Resource* CurrentStreamOutTargets[D3D12_SO_STREAM_COUNT] = {};
-			uint32			CurrentSOOffsets       [D3D12_SO_STREAM_COUNT] = {};
-
 			// Index Buffer State
 			FD3D12IndexBufferCache IBCache = {};
 
diff --git a/Engine/Source/Runtime/RHI/Private/RHIStrings.cpp b/Engine/Source/Runtime/RHI/Private/RHIStrings.cpp
--- a/Engine/Source/Runtime/RHI/Private/RHIStrings.cpp
+++ b/Engine/Source/Runtime/RHI/Private/RHIStrings.cpp
@@ -319,7 +319,9 @@ FString GetBufferUsageFlagsName(EBufferUsageFlags BufferUsage)
 		case BUF_UnorderedAccess:		return TEXT("BUF_UnorderedAccess");
 		case BUF_ByteAddressBuffer:		return TEXT("BUF_ByteAddressBuffer");
 		case BUF_SourceCopy:			return TEXT("BUF_SourceCopy");
+PRAGMA_DISABLE_DEPRECATION_WARNINGS
 		case BUF_StreamOutput:			return TEXT("BUF_StreamOutput");
+PRAGMA_ENABLE_DEPRECATION_WARNINGS
 		case BUF_DrawIndirect:			return TEXT("BUF_DrawIndirect");
 		case BUF_ShaderResource:		return TEXT("BUF_ShaderResource");
 		case BUF_KeepCPUAccessible:		return TEXT("BUF_KeepCPUAccessible");
@@ -633,7 +635,9 @@ const TCHAR* GetBufferUsageFlagString(EBufferUsageFlags BufferUsage)
 		return TEXT("ByteAddressBuffer");
 	case EBufferUsageFlags::SourceCopy:
 		return TEXT("SourceCopy");
+PRAGMA_DISABLE_DEPRECATION_WARNINGS
 	case EBufferUsageFlags::StreamOutput:
+PRAGMA_ENABLE_DEPRECATION_WARNINGS
 		return TEXT("StreamOutput");
 	case EBufferUsageFlags::DrawIndirect:
 		return TEXT("DrawIndirect");
diff --git a/Engine/Source/Runtime/RHI/Public/RHI.h b/Engine/Source/Runtime/RHI/Public/RHI.h
--- a/Engine/Source/Runtime/RHI/Public/RHI.h
+++ b/Engine/Source/Runtime/RHI/Public/RHI.h
@@ -212,7 +212,7 @@ typedef TArray<FVertexElement,TFixedAllocator<MaxVertexElementCount> > FVertexDe
 
 /** RHI representation of a single stream out element. */
 //#todo-RemoveStreamOut
-struct FStreamOutElement
+struct UE_DEPRECATED(5.3, "StreamOut is not supported") FStreamOutElement
 {
 	/** Index of the output stream from the geometry shader. */
 	uint32 Stream;
@@ -244,7 +244,10 @@ struct FStreamOutElement
 };
 
 //#todo-RemoveStreamOut
+UE_DEPRECATED(5.3, "StreamOut is not supported")
+PRAGMA_DISABLE_DEPRECATION_WARNINGS
 typedef TArray<FStreamOutElement,TFixedAllocator<MaxVertexElementCount> > FStreamOutElementList;
+PRAGMA_ENABLE_DEPRECATION_WARNINGS
 
 struct FSamplerStateInitializerRHI
 {
diff --git a/Engine/Source/Runtime/RHI/Public/RHIDefinitions.h b/Engine/Source/Runtime/RHI/Public/RHIDefinitions.h
--- a/Engine/Source/Runtime/RHI/Public/RHIDefinitions.h
+++ b/Engine/Source/Runtime/RHI/Public/RHIDefinitions.h
@@ -728,7 +728,7 @@ enum class EBufferUsageFlags : uint32
 	SourceCopy              = 1 << 5,
 
 	/** Create a buffer that can be bound as a stream output target. */
-	StreamOutput            = 1 << 6,
+	StreamOutput            UE_DEPRECATED(5.3, "StreamOut is not supported") = 1 << 6,
 
 	/** Create a buffer which contains the arguments used by DispatchIndirect or DrawIndirect. */
 	DrawIndirect            = 1 << 7,
