- Fixed wrong descriptor passed to d3d12Resource constructor. Shared buffers were causing asserts when getting resource allocation info.

diff --git a/Engine/Source/Runtime/D3D12RHI/Private/D3D12Resources.cpp b/Engine/Source/Runtime/D3D12RHI/Private/D3D12Resources.cpp
--- a/Engine/Source/Runtime/D3D12RHI/Private/D3D12Resources.cpp
+++ b/Engine/Source/Runtime/D3D12RHI/Private/D3D12Resources.cpp
@@ -488,7 +488,7 @@ HRESULT FD3D12Adapter::CreateCommittedResource(const FD3D12ResourceDesc& InDesc,
 	if (SUCCEEDED(hr))
 	{
 		// Set the output pointer
-		*ppOutResource = new FD3D12Resource(GetDevice(CreationNode.ToIndex()), CreationNode, pResource, InInitialState, InResourceStateMode, InDefaultState, InDesc, nullptr, HeapProps.Type);
+		*ppOutResource = new FD3D12Resource(GetDevice(CreationNode.ToIndex()), CreationNode, pResource, InInitialState, InResourceStateMode, InDefaultState, LocalDesc, nullptr, HeapProps.Type);
 		(*ppOutResource)->AddRef();
 
 		// Set a default name (can override later).
