Replace Stats2.h by Stats.h to fix nopch compilation

diff --git a/Engine/Source/Runtime/D3D12RHI/Private/D3D12RHICommon.h b/Engine/Source/Runtime/D3D12RHI/Private/D3D12RHICommon.h
--- a/Engine/Source/Runtime/D3D12RHI/Private/D3D12RHICommon.h
+++ b/Engine/Source/Runtime/D3D12RHI/Private/D3D12RHICommon.h
@@ -14,7 +14,7 @@
 
 #include "Containers/StaticArray.h"
 #include "MultiGPU.h"
-#include "Stats/Stats2.h"
+#include "Stats/Stats.h"
 
 class FD3D12Adapter;
 class FD3D12Device;
