Working on making Texture and Buffer memory accounting consistent
- Renaming GRHIGlobals.CurrentTextureMemorySize to GRHIGlobals.StreamingTextureMemorySizeInKB
- Renaming GRHIGlobals.CurrentRendertargetMemorySize to GRHIGlobals.NonStreamingTextureMemorySizeInKB
- Adding GRHIGlobals.BufferMemorySize and GRHIGlobals.UniformBufferMemorySize
- Deduplicating RHI stat accounting code and changing it to use RHI Descs.
- UE::RHICore::UpdateGlobalTextureStats is for StreamingTextureMemorySizeInKB and NonStreamingTextureMemorySizeInKB. This still allows platforms to control where the memory is counted.
- Adding more Buffer RHI stat accounting code to RHIs that were missing it.
- UE::RHICore::UpdateGlobalBufferStats is for BufferMemorySize and UniformBufferMemorySize.
- FillBaselineTextureMemoryStats initializes and fills the common members of FTextureMemoryStats.
- Adding ETextureCreateFlags::ForceIntoNonStreamingMemoryTracking which will force the texture to be counted in NonStreamingTextureMemorySizeInKB.
- Changing all RenderTargetPool textures to be flagged with ETextureCreateFlags::ForceIntoNonStreamingMemoryTracking so they aren't counted against streaming budgets.
- Adding UAV Texture stats instead of lumping them with regular textures
- Adding Buffer stats for ByteAddressBuffer, DrawIndirect and Misc

diff --git a/Engine/Source/Runtime/Apple/MetalRHI/Private/MetalTexture.cpp b/Engine/Source/Runtime/Apple/MetalRHI/Private/MetalTexture.cpp
--- a/Engine/Source/Runtime/Apple/MetalRHI/Private/MetalTexture.cpp
+++ b/Engine/Source/Runtime/Apple/MetalRHI/Private/MetalTexture.cpp
@@ -12,6 +12,7 @@
 #include "Misc/ScopeRWLock.h"
 #include "MetalLLM.h"
 #include "RHILockTracker.h"
+#include "RHICoreStats.h"
 
 volatile int64 FMetalSurface::ActiveUploads = 0;
 
@@ -710,67 +711,24 @@ FMetalSurface::FMetalSurface(FMetalTextureCreateDesc const& CreateDesc)
 	}
 	
 	// track memory usage
-	
-	if (CreateDesc.bIsRenderTarget)
-	{
-		GCurrentRendertargetMemorySize += Align(TotalTextureSize, 1024) / 1024;
-	}
-	else
-	{
-		GCurrentTextureMemorySize += Align(TotalTextureSize, 1024) / 1024;
-	}
-	
-#if STATS
-	if (CreateDesc.IsTextureCube())
-	{
-		if (CreateDesc.bIsRenderTarget)
-		{
-			INC_MEMORY_STAT_BY(STAT_RenderTargetMemoryCube, TotalTextureSize);
-		}
-		else
-		{
-			INC_MEMORY_STAT_BY(STAT_TextureMemoryCube, TotalTextureSize);
-		}
-	}
-	else if (CreateDesc.IsTexture3D())
-	{
-		if (CreateDesc.bIsRenderTarget)
-		{
-			INC_MEMORY_STAT_BY(STAT_RenderTargetMemory3D, TotalTextureSize);
-		}
-		else
-		{
-			INC_MEMORY_STAT_BY(STAT_TextureMemory3D, TotalTextureSize);
-		}
-	}
-	else
+	const bool bOnlyStreamableTexturesInTexturePool = false;
+	UE::RHICore::UpdateGlobalTextureStats(GetDesc(), TotalTextureSize, bOnlyStreamableTexturesInTexturePool, true);
+
+	if (Texture && EnumHasAnyFlags(CreateDesc.Flags, TexCreate_ShaderResource | TexCreate_UAV) &&
+		!(Texture.GetUsage() & mtlpp::TextureUsage::PixelFormatView))
 	{
-		if (CreateDesc.bIsRenderTarget)
-		{
-			INC_MEMORY_STAT_BY(STAT_RenderTargetMemory2D, TotalTextureSize);
-		}
-		else
-		{
-			INC_MEMORY_STAT_BY(STAT_TextureMemory2D, TotalTextureSize);
-		}
+		// If the texture was created without PixelFormatView delete the resources
+		// unless we definitely use this feature or we are throwing ~4% performance vs. Windows on the floor.
+		check(0);
 	}
-#endif
-    
-    if (Texture && EnumHasAnyFlags(CreateDesc.Flags, TexCreate_ShaderResource | TexCreate_UAV) &&
-            !(Texture.GetUsage() & mtlpp::TextureUsage::PixelFormatView))
-    {
-        // If the texture was created without PixelFormatView delete the resources
-        // unless we definitely use this feature or we are throwing ~4% performance vs. Windows on the floor.
-        check(0);
-    }
 }
 
 @interface FMetalDeferredStats : FApplePlatformObject
 {
 @public
-	uint64 TextureSize;
 	ETextureDimension Dimension;
-	bool bIsRenderTarget;
+	ETextureCreateFlags Flags;
+	uint64 TextureSize;
 }
 @end
 
@@ -778,57 +736,14 @@ FMetalSurface::FMetalSurface(FMetalTextureCreateDesc const& CreateDesc)
 APPLE_PLATFORM_OBJECT_ALLOC_OVERRIDES(FMetalDeferredStats)
 -(void)dealloc
 {
-#if STATS
-	if (Dimension == ETextureDimension::TextureCube || Dimension == ETextureDimension::TextureCubeArray)
-	{
-		if (bIsRenderTarget)
-		{
-			DEC_MEMORY_STAT_BY(STAT_RenderTargetMemoryCube, TextureSize);
-		}
-		else
-		{
-			DEC_MEMORY_STAT_BY(STAT_TextureMemoryCube, TextureSize);
-		}
-	}
-	else if (Dimension == ETextureDimension::Texture3D)
-	{
-		if (bIsRenderTarget)
-		{
-			DEC_MEMORY_STAT_BY(STAT_RenderTargetMemory3D, TextureSize);
-		}
-		else
-		{
-			DEC_MEMORY_STAT_BY(STAT_TextureMemory3D, TextureSize);
-		}
-	}
-	else
-	{
-		if (bIsRenderTarget)
-		{
-			DEC_MEMORY_STAT_BY(STAT_RenderTargetMemory2D, TextureSize);
-		}
-		else
-		{
-			DEC_MEMORY_STAT_BY(STAT_TextureMemory2D, TextureSize);
-		}
-	}
-#endif
-	if (bIsRenderTarget)
-	{
-		GCurrentRendertargetMemorySize -= Align(TextureSize, 1024) / 1024;
-	}
-	else
-	{
-		GCurrentTextureMemorySize -= Align(TextureSize, 1024) / 1024;
-	}
+	const bool bOnlyStreamableTexturesInTexturePool = false;
+	UE::RHICore::UpdateGlobalTextureStats(Flags, Dimension, TextureSize, bOnlyStreamableTexturesInTexturePool, false);
 	[super dealloc];
 }
 @end
 
 FMetalSurface::~FMetalSurface()
 {
-	bool const bIsRenderTarget = IsRenderTarget(GetDesc().Flags);
-	
 	if (MSAATexture.GetPtr())
 	{
 		if (Texture.GetPtr() != MSAATexture.GetPtr())
@@ -837,7 +752,6 @@ FMetalSurface::~FMetalSurface()
 		}
 	}
 	
-	
 	//do the same as above.  only do a [release] if it's the same as texture.
 	if (MSAAResolveTexture.GetPtr())
 	{
@@ -859,8 +773,8 @@ FMetalSurface::~FMetalSurface()
 	// track memory usage
 	FMetalDeferredStats* Block = [FMetalDeferredStats new];
 	Block->Dimension = GetDesc().Dimension;
+	Block->Flags = GetDesc().Flags;
 	Block->TextureSize = TotalTextureSize;
-	Block->bIsRenderTarget = bIsRenderTarget;
 	SafeReleaseMetalObject(Block);
 	
 	if(ImageSurfaceRef)
@@ -1393,25 +1307,17 @@ ns::AutoReleased<FMetalTexture> FMetalSurface::GetCurrentTexture()
 
 void FMetalDynamicRHI::RHIGetTextureMemoryStats(FTextureMemoryStats& OutStats)
 {
-	if(MemoryStats.TotalGraphicsMemory > 0)
+	UE::RHICore::FillBaselineTextureMemoryStats(OutStats);
+
+	if (MemoryStats.TotalGraphicsMemory > 0)
 	{
 		OutStats.DedicatedVideoMemory = MemoryStats.DedicatedVideoMemory;
 		OutStats.DedicatedSystemMemory = MemoryStats.DedicatedSystemMemory;
 		OutStats.SharedSystemMemory = MemoryStats.SharedSystemMemory;
 		OutStats.TotalGraphicsMemory = MemoryStats.TotalGraphicsMemory;
 	}
-	else
-	{
-		OutStats.DedicatedVideoMemory = 0;
-		OutStats.DedicatedSystemMemory = 0;
-		OutStats.SharedSystemMemory = 0;
-		OutStats.TotalGraphicsMemory = 0;
-	}
-	
-	OutStats.AllocatedMemorySize = int64(GCurrentTextureMemorySize) * 1024;
-	OutStats.LargestContiguousAllocation = OutStats.AllocatedMemorySize;
-	OutStats.TexturePoolSize = GTexturePoolSize;
-	OutStats.PendingMemoryAdjustment = 0;
+
+	OutStats.LargestContiguousAllocation = OutStats.StreamingMemorySize;
 }
 
 bool FMetalDynamicRHI::RHIGetTextureMemoryVisualizeData( FColor* /*TextureData*/, int32 /*SizeX*/, int32 /*SizeY*/, int32 /*Pitch*/, int32 /*PixelSize*/ )
diff --git a/Engine/Source/Runtime/D3D12RHI/Private/D3D12Buffer.cpp b/Engine/Source/Runtime/D3D12RHI/Private/D3D12Buffer.cpp
--- a/Engine/Source/Runtime/D3D12RHI/Private/D3D12Buffer.cpp
+++ b/Engine/Source/Runtime/D3D12RHI/Private/D3D12Buffer.cpp
@@ -6,6 +6,7 @@ D3D12Buffer.cpp: D3D Common code for buffers.
 
 #include "D3D12RHIPrivate.h"
 #include "ProfilingDebugging/AssetMetadataTrace.h"
+#include "RHICoreStats.h"
 
 FName GetRHIBufferClassName(const FName& ClassName)
 {
@@ -26,7 +27,7 @@ FD3D12Buffer::~FD3D12Buffer()
 	bool bTransient = ResourceLocation.IsTransient();
 	if (!bTransient)
 	{
-		UpdateBufferStats(this, false);
+		D3D12BufferStats::UpdateBufferStats(*this, false);
 	}
 }
 
@@ -424,7 +425,7 @@ FD3D12Buffer* FD3D12Adapter::CreateRHIBuffer(
 	// Don't track transient buffer stats here
 	if (!BufferOut->ResourceLocation.IsTransient())
 	{
-		UpdateBufferStats(BufferOut, true);
+		D3D12BufferStats::UpdateBufferStats(*BufferOut, true);
 	}
 
 	return BufferOut;
@@ -481,7 +482,7 @@ void FD3D12Buffer::ReleaseOwnership()
 
 	if (!ResourceLocation.IsTransient())
 	{
-		UpdateBufferStats(this, false);
+		D3D12BufferStats::UpdateBufferStats(*this, false);
 	}
 
 	ResourceLocation.Clear();
diff --git a/Engine/Source/Runtime/D3D12RHI/Private/D3D12Resources.h b/Engine/Source/Runtime/D3D12RHI/Private/D3D12Resources.h
--- a/Engine/Source/Runtime/D3D12RHI/Private/D3D12Resources.h
+++ b/Engine/Source/Runtime/D3D12RHI/Private/D3D12Resources.h
@@ -817,8 +817,6 @@ public:
 	}
 };
 
-extern void UpdateBufferStats(class FD3D12Buffer* Buffer, bool bAllocating);
-
 /** Uniform buffer resource class. */
 class FD3D12UniformBuffer : public FRHIUniformBuffer, public FD3D12DeviceChild, public FD3D12LinkedAdapterObject<FD3D12UniformBuffer>
 {
@@ -910,6 +908,11 @@ public:
 	uint32 BufferAlignment = 0;
 };
 
+namespace D3D12BufferStats
+{
+	void UpdateBufferStats(FD3D12Buffer& Buffer, bool bAllocating);
+}
+
 class FD3D12ResourceBarrierBatcher
 {
 	PRAGMA_DISABLE_ENUM_CONSTEXPR_CONVERSION_WARNINGS
diff --git a/Engine/Source/Runtime/D3D12RHI/Private/D3D12Stats.cpp b/Engine/Source/Runtime/D3D12RHI/Private/D3D12Stats.cpp
--- a/Engine/Source/Runtime/D3D12RHI/Private/D3D12Stats.cpp
+++ b/Engine/Source/Runtime/D3D12RHI/Private/D3D12Stats.cpp
@@ -8,6 +8,7 @@ D3D12Stats.cpp:RHI Stats and timing implementation.
 #include "Engine/Engine.h"
 #include "Engine/GameViewportClient.h"
 #include "ProfilingDebugging/MemoryTrace.h"
+#include "RHICoreStats.h"
 
 void D3D12RHI::FD3DGPUProfiler::BeginFrame()
 {
@@ -261,62 +262,45 @@ float FD3D12EventNode::GetTiming()
 	return Result;
 }
 
-inline FName GetRHIBufferStats(EBufferUsageFlags Usage)
+static TStatId GetD3D12BufferStat(const FRHIBufferDesc& BufferDesc)
 {
-	if (EnumHasAnyFlags(Usage, BUF_VertexBuffer))
+	if (EnumHasAnyFlags(BufferDesc.Usage, EBufferUsageFlags::UnorderedAccess))
 	{
-		return GET_STATFNAME(STAT_VertexBufferMemory);
+		return GET_STATID(STAT_D3D12UAVBuffers);
 	}
-	else if (EnumHasAnyFlags(Usage, BUF_IndexBuffer))
+	
+	if (EnumHasAnyFlags(BufferDesc.Usage, EBufferUsageFlags::AccelerationStructure))
 	{
-		return GET_STATFNAME(STAT_IndexBufferMemory);
+		return GET_STATID(STAT_D3D12RTBuffers);
 	}
-	else if (EnumHasAnyFlags(Usage, BUF_AccelerationStructure))
-	{
-		return GET_STATFNAME(STAT_RTAccelerationStructureMemory);
-	}
-	else
-	{
-		return GET_STATFNAME(STAT_StructuredBufferMemory);
-	}
-}
 
-static FName GetD3D12BufferStat(EBufferUsageFlags InUsageFlags)
-{
-	if (EnumHasAnyFlags(InUsageFlags, BUF_UnorderedAccess))
-	{
-		return GET_STATFNAME(STAT_D3D12UAVBuffers);
-	}
-	else if (EnumHasAnyFlags(InUsageFlags, BUF_AccelerationStructure))
-	{
-		return GET_STATFNAME(STAT_D3D12RTBuffers);
-	}
-	else
-	{
-		return GET_STATFNAME(STAT_D3D12Buffer);
-	}
+	return GET_STATID(STAT_D3D12Buffer);
 }
 
-void UpdateBufferStats(FD3D12Buffer* Buffer, bool bAllocating)
+void D3D12BufferStats::UpdateBufferStats(FD3D12Buffer& Buffer, bool bAllocating)
 {
-	D3D12_GPU_VIRTUAL_ADDRESS GPUAddress = Buffer->ResourceLocation.GetGPUVirtualAddress();
-	int64 BufferSize = Buffer->ResourceLocation.GetSize();
-	int64 RequestedSize = bAllocating ? BufferSize : -BufferSize;
-	EBufferUsageFlags UsageFlags = Buffer->GetUsage();
+	const FRHIBufferDesc& BufferDesc = Buffer.GetDesc();
+	const FD3D12ResourceLocation& Location = Buffer.ResourceLocation;
+
+	const int64 BufferSize = Location.GetSize();
+	const int64 RequestedSize = bAllocating ? BufferSize : -BufferSize;
 
-	INC_MEMORY_STAT_BY_FName(GetRHIBufferStats(UsageFlags), RequestedSize);
-	INC_MEMORY_STAT_BY_FName(GetD3D12BufferStat(UsageFlags), RequestedSize);
+	UE::RHICore::UpdateGlobalBufferStats(BufferDesc, RequestedSize);
+
+	INC_MEMORY_STAT_BY_FName(GetD3D12BufferStat(BufferDesc).GetName(), RequestedSize);
 	INC_MEMORY_STAT_BY(STAT_D3D12MemoryCurrentTotal, RequestedSize);
 
 #if UE_MEMORY_TRACE_ENABLED
+	const D3D12_GPU_VIRTUAL_ADDRESS GPUAddress = Location.GetGPUVirtualAddress();
+
 	if (bAllocating)
 	{
 		// Skip if it's created as a
 		// 1) standalone resource, because MemoryTrace_Alloc has been called in FD3D12Adapter::CreateCommittedResource
 		// 2) placed resource from a pool allocator, because MemoryTrace_Alloc has been called in FD3D12Adapter::CreatePlacedResource
-		if (!Buffer->ResourceLocation.IsStandaloneOrPooledPlacedResource())
+		if (!Location.IsStandaloneOrPooledPlacedResource())
 		{
-			MemoryTrace_Alloc(GPUAddress, BufferSize, Buffer->BufferAlignment, EMemoryTraceRootHeap::VideoMemory);
+			MemoryTrace_Alloc(GPUAddress, BufferSize, Buffer.BufferAlignment, EMemoryTraceRootHeap::VideoMemory);
 		}
 	}
 	else
diff --git a/Engine/Source/Runtime/D3D12RHI/Private/D3D12Texture.cpp b/Engine/Source/Runtime/D3D12RHI/Private/D3D12Texture.cpp
--- a/Engine/Source/Runtime/D3D12RHI/Private/D3D12Texture.cpp
+++ b/Engine/Source/Runtime/D3D12RHI/Private/D3D12Texture.cpp
@@ -9,6 +9,7 @@
 #include "ProfilingDebugging/MemoryTrace.h"
 #include "RHIUtilities.h"
 #include "ProfilingDebugging/AssetMetadataTrace.h"
+#include "RHICoreStats.h"
 
 int64 FD3D12GlobalStats::GDedicatedVideoMemory = 0;
 int64 FD3D12GlobalStats::GDedicatedSystemMemory = 0;
@@ -38,9 +39,10 @@ static TAutoConsoleVariable<int32> CVarUseUpdateTexture3DComputeShader(
 	TEXT(" 1: on"),
 	ECVF_RenderThreadSafe );
 
-static TAutoConsoleVariable<bool> CVarTexturePoolOnlyAccountStreamableTexture(
+static bool GTexturePoolOnlyAccountStreamableTexture = false;
+static FAutoConsoleVariableRef CVarTexturePoolOnlyAccountStreamableTexture(
 	TEXT("D3D12.TexturePoolOnlyAccountStreamableTexture"),
-	0,
+	GTexturePoolOnlyAccountStreamableTexture,
 	TEXT("Texture streaming pool size only account streamable texture .\n")
 	TEXT(" - 0: All texture types are counted in the pool (legacy, default).\n")
 	TEXT(" - 1: Only streamable textures are counted in the pool.\n")
@@ -282,84 +284,34 @@ struct FRHICommandD3D12AsyncReallocateTexture2D final : public FRHICommand<FRHIC
 // Texture Stats
 ///////////////////////////////////////////////////////////////////////////////////////////
 
-
-static bool ShouldCountAsTextureMemory(D3D12_RESOURCE_FLAGS MiscFlags)
-{
-	// Shouldn't be used for DEPTH, RENDER TARGET, or UNORDERED ACCESS
-	return !EnumHasAnyFlags(MiscFlags, D3D12_RESOURCE_FLAG_ALLOW_DEPTH_STENCIL | D3D12_RESOURCE_FLAG_ALLOW_RENDER_TARGET | D3D12_RESOURCE_FLAG_ALLOW_UNORDERED_ACCESS);
-}
-
-
-static TStatId GetRHIStatEnum(D3D12_RESOURCE_FLAGS MiscFlags, bool bCubeMap, bool b3D)
-{
 #if STATS
-	if (ShouldCountAsTextureMemory(MiscFlags))
-	{
-		// normal texture
-		if (bCubeMap)
-		{
-			return GET_STATID(STAT_TextureMemoryCube);
-		}
-		else if (b3D)
-		{
-			return GET_STATID(STAT_TextureMemory3D);
-		}
-		else
-		{
-			return GET_STATID(STAT_TextureMemory2D);
-		}
-	}
-	else
-	{
-		// render target
-		if (bCubeMap)
-		{
-			return GET_STATID(STAT_RenderTargetMemoryCube);
-		}
-		else if (b3D)
-		{
-			return GET_STATID(STAT_RenderTargetMemory3D);
-		}
-		else
-		{
-			return GET_STATID(STAT_RenderTargetMemory2D);
-		}
-	}
-#endif
-	return TStatId();
-}
-
-
-static TStatId GetD3D12StatEnum(D3D12_RESOURCE_FLAGS MiscFlags)
+static TStatId GetD3D12StatEnum(const FD3D12ResourceDesc& ResourceDesc)
 {
-#if STATS
-	if (EnumHasAnyFlags(MiscFlags, D3D12_RESOURCE_FLAG_ALLOW_RENDER_TARGET | D3D12_RESOURCE_FLAG_ALLOW_DEPTH_STENCIL))
+	if (EnumHasAnyFlags(ResourceDesc.Flags, D3D12_RESOURCE_FLAG_ALLOW_RENDER_TARGET | D3D12_RESOURCE_FLAG_ALLOW_DEPTH_STENCIL))
 	{
 		return GET_STATID(STAT_D3D12RenderTargets);
 	}
-	else if (EnumHasAnyFlags(MiscFlags, D3D12_RESOURCE_FLAG_ALLOW_UNORDERED_ACCESS))
+
+	if (EnumHasAnyFlags(ResourceDesc.Flags, D3D12_RESOURCE_FLAG_ALLOW_UNORDERED_ACCESS))
 	{
 		return GET_STATID(STAT_D3D12UAVTextures);
 	}
-	else
-	{
-		return GET_STATID(STAT_D3D12Textures);
-	}
-#endif
-	return TStatId();
+
+	return GET_STATID(STAT_D3D12Textures);
 }
+#endif // STATS
 
 
-void FD3D12TextureStats::UpdateD3D12TextureStats(FD3D12Texture& Texture, const D3D12_RESOURCE_DESC& Desc, int64 TextureSize, bool b3D, bool bCubeMap, bool bStreamable, bool bNewTexture)
+void FD3D12TextureStats::UpdateD3D12TextureStats(FD3D12Texture& Texture, const FD3D12ResourceDesc& ResourceDesc, const FRHITextureDesc& TextureDesc, uint64 TextureSize, bool bNewTexture, bool bAllocating)
 {
 #if TEXTURE_PROFILER_ENABLED
 	if (!bNewTexture
 		&& !Texture.ResourceLocation.IsTransient() 
-		&& !EnumHasAnyFlags(Texture.GetFlags(), TexCreate_Virtual)
+		&& !EnumHasAnyFlags(TextureDesc.Flags, ETextureCreateFlags::Virtual)
 		&& !Texture.ResourceLocation.IsAliased())
 	{
-		const uint64 SafeSize = (uint64)(TextureSize >= 0 ? TextureSize : 0);
-		FTextureProfiler::Get()->UpdateTextureAllocation(&Texture, SafeSize, Desc.Alignment, 0);
+		const uint64 SafeSize = bAllocating ? TextureSize : 0;
+		FTextureProfiler::Get()->UpdateTextureAllocation(&Texture, SafeSize, ResourceDesc.Alignment, 0);
 	}
 #endif
 
@@ -368,24 +320,12 @@ void FD3D12TextureStats::UpdateD3D12TextureStats(FD3D12Texture& Texture, const D
 		return;
 	}
 
-	const int64 AlignedSize = (TextureSize > 0) ? Align(TextureSize, 1024) / 1024 : -(Align(-TextureSize, 1024) / 1024);
-	if (ShouldCountAsTextureMemory(Desc.Flags))
-	{
-		const bool bOnlyStreamableTextureAccounted = CVarTexturePoolOnlyAccountStreamableTexture.GetValueOnAnyThread();
+	UE::RHICore::UpdateGlobalTextureStats(TextureDesc, TextureSize, GTexturePoolOnlyAccountStreamableTexture, bAllocating);
 
-		if (!bOnlyStreamableTextureAccounted || bStreamable)
-		{
-			FPlatformAtomics::InterlockedAdd(&GCurrentTextureMemorySize, AlignedSize);
-		}
-	}
-	else
-	{
-		FPlatformAtomics::InterlockedAdd(&GCurrentRendertargetMemorySize, AlignedSize);
-	}
+	const int64 TextureSizeDeltaInBytes = bAllocating ? static_cast<int64>(TextureSize) : -static_cast<int64>(TextureSize);
 
-	INC_MEMORY_STAT_BY_FName(GetD3D12StatEnum(Desc.Flags).GetName(), TextureSize);
-	INC_MEMORY_STAT_BY_FName(GetRHIStatEnum(Desc.Flags, bCubeMap, b3D).GetName(), TextureSize);
-	INC_MEMORY_STAT_BY(STAT_D3D12MemoryCurrentTotal, TextureSize);
+	INC_MEMORY_STAT_BY_FName(GetD3D12StatEnum(ResourceDesc).GetName(), TextureSizeDeltaInBytes);
+	INC_MEMORY_STAT_BY(STAT_D3D12MemoryCurrentTotal, TextureSizeDeltaInBytes);
 
 	D3D12_GPU_VIRTUAL_ADDRESS GPUAddress = Texture.ResourceLocation.GetGPUVirtualAddress();
 
@@ -395,7 +335,7 @@ void FD3D12TextureStats::UpdateD3D12TextureStats(FD3D12Texture& Texture, const D
 	const bool bMemoryTrace = bTrackingAllAllocations || GPUAddress != 0;
 #endif
 
-	if (TextureSize > 0)
+	if (bAllocating)
 	{
 #if PLATFORM_WINDOWS
 		// On Windows there is no way to hook into the low level d3d allocations and frees.
@@ -411,7 +351,7 @@ void FD3D12TextureStats::UpdateD3D12TextureStats(FD3D12Texture& Texture, const D
 			// 2) placed resource from a pool allocator, because MemoryTrace_Alloc has been called in FD3D12Adapter::CreatePlacedResource
 			if (bMemoryTrace && !Texture.ResourceLocation.IsStandaloneOrPooledPlacedResource())
 			{
-				MemoryTrace_Alloc(GPUAddress, TextureSize, Desc.Alignment, EMemoryTraceRootHeap::VideoMemory);
+				MemoryTrace_Alloc(GPUAddress, TextureSize, ResourceDesc.Alignment, EMemoryTraceRootHeap::VideoMemory);
 			}
 #endif
 		}
@@ -437,33 +377,26 @@ void FD3D12TextureStats::UpdateD3D12TextureStats(FD3D12Texture& Texture, const D
 }
 
 
-void FD3D12TextureStats::D3D12TextureAllocated(FD3D12Texture& Texture, const D3D12_RESOURCE_DESC *Desc)
+void FD3D12TextureStats::D3D12TextureAllocated(FD3D12Texture& Texture)
 {
-	FD3D12Resource* D3D12Texture = Texture.GetResource();
-
-	if (D3D12Texture)
+	if (FD3D12Resource* D3D12Resource = Texture.GetResource())
 	{
-		if (!Desc)
-		{
-			Desc = &D3D12Texture->GetDesc();
-		}
+		const FD3D12ResourceDesc& ResourceDesc = D3D12Resource->GetDesc();
+		const FRHITextureDesc& TextureDesc = Texture.GetDesc();
 
 		// Don't update state for readback, virtual, or transient textures	
-		const ETextureCreateFlags CreateFlags = Texture.GetDesc().Flags;
-		if (!EnumHasAnyFlags(CreateFlags, TexCreate_Virtual | TexCreate_CPUReadback) && !Texture.ResourceLocation.IsTransient())
+		if (!EnumHasAnyFlags(TextureDesc.Flags, ETextureCreateFlags::Virtual | ETextureCreateFlags::CPUReadback) && !Texture.ResourceLocation.IsTransient())
 		{
 			TRACE_CPUPROFILER_EVENT_SCOPE(D3D12RHI::UpdateTextureStats);
 
-			const int64 TextureSize = Texture.ResourceLocation.GetSize();			
+			const uint64 TextureSize = Texture.ResourceLocation.GetSize();
 			const bool bNewTexture = true;
-			const bool bIsStreamable = EnumHasAnyFlags(CreateFlags, ETextureCreateFlags::Streamable);
-			UpdateD3D12TextureStats(Texture, *Desc, TextureSize, Texture.GetDesc().IsTexture3D(), Texture.GetDesc().IsTextureCube(), bIsStreamable, bNewTexture);
+			UpdateD3D12TextureStats(Texture, ResourceDesc, TextureDesc, TextureSize, bNewTexture, true);
 
 #if TEXTURE_PROFILER_ENABLED
 			if (!Texture.ResourceLocation.IsAliased())
 			{
-				const uint32 Alignment = Desc->Alignment;
-				FTextureProfiler::Get()->AddTextureAllocation(&Texture, TextureSize, Alignment, 0);
+				FTextureProfiler::Get()->AddTextureAllocation(&Texture, TextureSize, ResourceDesc.Alignment, 0);
 			}
 #endif
 		}
@@ -474,24 +407,22 @@ void FD3D12TextureStats::D3D12TextureAllocated(FD3D12Texture& Texture, const D3D
 
 void FD3D12TextureStats::D3D12TextureDeleted(FD3D12Texture& Texture)
 {
-	FD3D12Resource* D3D12Texture = Texture.GetResource();
-
-	if (D3D12Texture)
+	if (FD3D12Resource* D3D12Resource = Texture.GetResource())
 	{
+		const FD3D12ResourceDesc& ResourceDesc = D3D12Resource->GetDesc();
+		const FRHITextureDesc& TextureDesc = Texture.GetDesc();
+
 		// Don't update state for readback or transient textures, but virtual textures need to have their size deducted from calls to RHIVirtualTextureSetFirstMipInMemory.
-		const ETextureCreateFlags CreateFlags = Texture.GetDesc().Flags;
-		if (!EnumHasAnyFlags(CreateFlags, TexCreate_CPUReadback) && !Texture.ResourceLocation.IsTransient())
+		if (!EnumHasAnyFlags(TextureDesc.Flags, ETextureCreateFlags::CPUReadback) && !Texture.ResourceLocation.IsTransient())
 		{
-			const D3D12_RESOURCE_DESC& Desc = D3D12Texture->GetDesc();
-			const int64 TextureSize = Texture.ResourceLocation.GetSize();
-			ensure(TextureSize > 0 || EnumHasAnyFlags(CreateFlags, TexCreate_Virtual) || Texture.ResourceLocation.IsAliased());
+			const uint64 TextureSize = Texture.ResourceLocation.GetSize();
+			ensure(TextureSize > 0 || EnumHasAnyFlags(TextureDesc.Flags, ETextureCreateFlags::Virtual) || Texture.ResourceLocation.IsAliased());
 
 			const bool bNewTexture = false;
-			const bool bIsStreamable = EnumHasAnyFlags(CreateFlags, ETextureCreateFlags::Streamable);
-			UpdateD3D12TextureStats(Texture, Desc, -TextureSize, Texture.GetDesc().IsTexture3D(), Texture.GetDesc().IsTextureCube(), bIsStreamable, bNewTexture);
+			UpdateD3D12TextureStats(Texture, ResourceDesc, TextureDesc, TextureSize, bNewTexture, false);
 
 #if TEXTURE_PROFILER_ENABLED
-			if (!EnumHasAnyFlags(CreateFlags, TexCreate_Virtual) && !Texture.ResourceLocation.IsAliased())
+			if (!EnumHasAnyFlags(TextureDesc.Flags, ETextureCreateFlags::Virtual) && !Texture.ResourceLocation.IsAliased())
 			{
 				FTextureProfiler::Get()->RemoveTextureAllocation(&Texture);
 			}
@@ -682,15 +613,14 @@ FDynamicRHI::FRHICalcTextureSizeResult FD3D12DynamicRHI::RHICalcTexturePlatformS
  */
 void FD3D12DynamicRHI::RHIGetTextureMemoryStats(FTextureMemoryStats& OutStats)
 {
+	UE::RHICore::FillBaselineTextureMemoryStats(OutStats);
+
 	OutStats.DedicatedVideoMemory = FD3D12GlobalStats::GDedicatedVideoMemory;
 	OutStats.DedicatedSystemMemory = FD3D12GlobalStats::GDedicatedSystemMemory;
 	OutStats.SharedSystemMemory = FD3D12GlobalStats::GSharedSystemMemory;
 	OutStats.TotalGraphicsMemory = FD3D12GlobalStats::GTotalGraphicsMemory ? FD3D12GlobalStats::GTotalGraphicsMemory : -1;
 
-	OutStats.AllocatedMemorySize = int64(GCurrentTextureMemorySize) * 1024;
-	OutStats.LargestContiguousAllocation = OutStats.AllocatedMemorySize;
-	OutStats.TexturePoolSize = GTexturePoolSize;
-	OutStats.PendingMemoryAdjustment = 0;
+	OutStats.LargestContiguousAllocation = OutStats.StreamingMemorySize;
 
 #if PLATFORM_WINDOWS
 	if (GAdjustTexturePoolSizeBasedOnBudget)
@@ -716,12 +646,12 @@ void FD3D12DynamicRHI::RHIGetTextureMemoryStats(FTextureMemoryStats& OutStats)
 			// Attempt to lower the texture pool size to meet the budget.
 			const bool bOverActualBudget = LocalVideoMemoryInfo.CurrentUsage > LocalVideoMemoryInfo.Budget;
 			UE_CLOG(bOverActualBudget, LogD3D12RHI, Warning,
-				TEXT("Video memory usage is overbudget by %llu MB (using %lld MB/%lld MB budget). Usage breakdown: %lld MB (Textures), %lld MB (Render targets). Last requested texture pool size is %lld MB. This can cause stuttering due to paging."),
+				TEXT("Video memory usage is overbudget by %llu MB (using %lld MB/%lld MB budget). Usage breakdown: %lld MB (Streaming Textures), %lld MB (Non Streaming Textures). Last requested texture pool size is %lld MB. This can cause stuttering due to paging."),
 				(LocalVideoMemoryInfo.CurrentUsage - LocalVideoMemoryInfo.Budget) / 1024ll / 1024ll,
 				LocalVideoMemoryInfo.CurrentUsage / 1024ll / 1024ll,
 				LocalVideoMemoryInfo.Budget / 1024ll / 1024ll,
-				GCurrentTextureMemorySize / 1024ll,
-				GCurrentRendertargetMemorySize / 1024ll,
+				GRHIGlobals.StreamingTextureMemorySizeInKB / 1024ll,
+				GRHIGlobals.NonStreamingTextureMemorySizeInKB / 1024ll,
 				PreviousTexturePoolSize / 1024ll / 1024ll);
 
 			const int64 DesiredTexturePoolSize = PreviousTexturePoolSize + AvailableSpace - BudgetPadding;
diff --git a/Engine/Source/Runtime/D3D12RHI/Private/D3D12Texture.h b/Engine/Source/Runtime/D3D12RHI/Private/D3D12Texture.h
--- a/Engine/Source/Runtime/D3D12RHI/Private/D3D12Texture.h
+++ b/Engine/Source/Runtime/D3D12RHI/Private/D3D12Texture.h
@@ -256,16 +256,14 @@ FORCEINLINE FD3D12Texture* GetD3D12TextureFromRHITexture(FRHITexture* Texture, u
 	}
 }
 
-class FD3D12TextureStats
+namespace FD3D12TextureStats
 {
-public:
-
 	// Note: This function can be called from many different threads
 	// @param TextureSize >0 to allocate, <0 to deallocate
 	// @param b3D true:3D, false:2D or cube map
 	// @param bStreamable true:Streamable, false:not streamable
-	static void UpdateD3D12TextureStats(FD3D12Texture& Texture, const D3D12_RESOURCE_DESC& Desc, int64 TextureSize, bool b3D, bool bCubeMap, bool bStreamable, bool bNewTexture);
+	void UpdateD3D12TextureStats(FD3D12Texture& Texture, const FD3D12ResourceDesc& ResourceDesc, const FRHITextureDesc& TextureDesc, uint64 TextureSize, bool bNewTexture, bool bAllocating);
 
-	static void D3D12TextureAllocated(FD3D12Texture& Texture, const D3D12_RESOURCE_DESC *Desc = nullptr);
-	static void D3D12TextureDeleted(FD3D12Texture& Texture);
+	void D3D12TextureAllocated(FD3D12Texture& Texture);
+	void D3D12TextureDeleted(FD3D12Texture& Texture);
 };
diff --git a/Engine/Source/Runtime/Engine/Private/Streaming/AsyncTextureStreaming.cpp b/Engine/Source/Runtime/Engine/Private/Streaming/AsyncTextureStreaming.cpp
--- a/Engine/Source/Runtime/Engine/Private/Streaming/AsyncTextureStreaming.cpp
+++ b/Engine/Source/Runtime/Engine/Private/Streaming/AsyncTextureStreaming.cpp
@@ -580,7 +580,7 @@ void FRenderAssetStreamingMipCalcTask::UpdateBudgetedMips_Async()
 	if (GPoolSizeVRAMPercentage > 0 && TotalGraphicsMemory > 0)
 	{
 		const int64 UsableVRAM = FMath::Max<int64>(TotalGraphicsMemory * GPoolSizeVRAMPercentage / 100, TotalGraphicsMemory - Settings.VRAMPercentageClamp * 1024ll * 1024ll);
-		const int64 UsedVRAM = (int64)GCurrentRendertargetMemorySize * 1024ll + NonStreamingRenderAssetMemory; // Add any other...
+		const int64 UsedVRAM = (int64)GRHIGlobals.NonStreamingTextureMemorySizeInKB * 1024ll + NonStreamingRenderAssetMemory; // Add any other...
 		const int64 AvailableVRAMForStreaming = FMath::Min<int64>(UsableVRAM - UsedVRAM - MemoryMargin, PoolSize);
 		if (Settings.bLimitPoolSizeToVRAM || AvailableVRAMForStreaming > AvailableMemoryForStreaming)
 		{
@@ -589,7 +589,7 @@ void FRenderAssetStreamingMipCalcTask::UpdateBudgetedMips_Async()
 	}
 
 	// Update EffectiveStreamingPoolSize, trying to stabilize it independently of temp memory, allocator overhead and non-streaming resources normal variation.
-	// It's hard to know how much temp memory and allocator overhead is actually in AllocatedMemorySize as it is platform specific.
+	// It's hard to know how much temp memory and allocator overhead is actually in StreamingMemorySize as it is platform specific.
 	// We handle it by not using all memory available. If temp memory and memory margin values are effectively bigger than the actual used values, the pool will stabilize.
 	if (AvailableMemoryForStreaming < MemoryBudget)
 	{
diff --git a/Engine/Source/Runtime/Engine/Private/Streaming/StreamingManagerTexture.cpp b/Engine/Source/Runtime/Engine/Private/Streaming/StreamingManagerTexture.cpp
--- a/Engine/Source/Runtime/Engine/Private/Streaming/StreamingManagerTexture.cpp
+++ b/Engine/Source/Runtime/Engine/Private/Streaming/StreamingManagerTexture.cpp
@@ -703,12 +703,12 @@ void FRenderAssetStreamingManager::PrepareAsyncTask(bool bProcessEverything)
 	if (Stats.IsUsingLimitedPoolSize() && !bProcessEverything && !Settings.bFullyLoadUsedTextures)
 	{
 		const int64 TempMemoryBudget = static_cast<int64>(Settings.MaxTempMemoryAllowed) * 1024 * 1024;
-		AsyncTask.Reset(Stats.TotalGraphicsMemory, Stats.AllocatedMemorySize, Stats.TexturePoolSize, TempMemoryBudget, MemoryMargin);
+		AsyncTask.Reset(Stats.TotalGraphicsMemory, Stats.StreamingMemorySize, Stats.TexturePoolSize, TempMemoryBudget, MemoryMargin);
 	}
 	else
 	{
 		// Temp must be smaller since membudget only updates if it has a least temp memory available.
-		AsyncTask.Reset(0, Stats.AllocatedMemorySize, MAX_int64, MAX_int64 / 2, 0);
+		AsyncTask.Reset(0, Stats.StreamingMemorySize, MAX_int64, MAX_int64 / 2, 0);
 	}
 	AsyncTask.StreamingData.Init(CurrentViewInfos, LastWorldUpdateTime, LevelRenderAssetManagers, DynamicComponentManager);
 
diff --git a/Engine/Source/Runtime/Engine/Private/UnrealEngine.cpp b/Engine/Source/Runtime/Engine/Private/UnrealEngine.cpp
--- a/Engine/Source/Runtime/Engine/Private/UnrealEngine.cpp
+++ b/Engine/Source/Runtime/Engine/Private/UnrealEngine.cpp
@@ -1510,10 +1510,10 @@ static void EngineMemoryWarningHandler(const FGenericMemoryWarningContext& Gener
 {
 	FPlatformMemoryStats Stats = FPlatformMemory::GetStats();
 
-	FPlatformMisc::LowLevelOutputDebugStringf(TEXT("EngineMemoryWarningHandler: Mem Used %.2f MB, Texture Memory %.2f MB, Render Target memory %.2f MB, OS Free %.2f MB\n"), 
+	FPlatformMisc::LowLevelOutputDebugStringf(TEXT("EngineMemoryWarningHandler: Mem Used %.2f MB, Streaming Texture Memory %.2f MB, Non Streaming Texture Memory %.2f MB, OS Free %.2f MB\n"), 
 		Stats.UsedPhysical / 1048576.0f, 
-		GCurrentTextureMemorySize / 1024.f,
-		GCurrentRendertargetMemorySize / 1024.f,
+		static_cast<double>(GRHIGlobals.StreamingTextureMemorySizeInKB) / 1024.0,
+		static_cast<double>(GRHIGlobals.NonStreamingTextureMemorySizeInKB) / 1024.0,
 		Stats.AvailablePhysical / 1048576.0f);
 
 #if !UE_BUILD_SHIPPING && !UE_BUILD_TEST
diff --git a/Engine/Source/Runtime/OpenGLDrv/Private/OpenGLDevice.cpp b/Engine/Source/Runtime/OpenGLDrv/Private/OpenGLDevice.cpp
--- a/Engine/Source/Runtime/OpenGLDrv/Private/OpenGLDevice.cpp
+++ b/Engine/Source/Runtime/OpenGLDrv/Private/OpenGLDevice.cpp
@@ -1458,7 +1458,7 @@ void FOpenGLDynamicRHI::Cleanup()
 	{
 		FOpenGL::DeleteBuffers(1, &PendingState.ZeroFilledDummyUniformBuffer);
 		PendingState.ZeroFilledDummyUniformBuffer = 0;
-		DecrementBufferMemory(GL_UNIFORM_BUFFER, ZERO_FILLED_DUMMY_UNIFORM_BUFFER_SIZE);
+		OpenGLBufferStats::UpdateUniformBufferStats(ZERO_FILLED_DUMMY_UNIFORM_BUFFER_SIZE, false);
 	}
 
 	// Release pending shader
diff --git a/Engine/Source/Runtime/OpenGLDrv/Private/OpenGLShaders.cpp b/Engine/Source/Runtime/OpenGLDrv/Private/OpenGLShaders.cpp
--- a/Engine/Source/Runtime/OpenGLDrv/Private/OpenGLShaders.cpp
+++ b/Engine/Source/Runtime/OpenGLDrv/Private/OpenGLShaders.cpp
@@ -1119,7 +1119,8 @@ void FOpenGLDynamicRHI::BindUniformBufferBase(FOpenGLContextState& ContextState,
 				CachedBindUniformBuffer(ContextState,PendingState.ZeroFilledDummyUniformBuffer);
 				glBufferData(GL_UNIFORM_BUFFER, ZERO_FILLED_DUMMY_UNIFORM_BUFFER_SIZE, ZeroBuffer, GL_STATIC_DRAW);
 				FMemory::Free(ZeroBuffer);
-				IncrementBufferMemory(GL_UNIFORM_BUFFER, ZERO_FILLED_DUMMY_UNIFORM_BUFFER_SIZE);
+
+				OpenGLBufferStats::UpdateUniformBufferStats(ZERO_FILLED_DUMMY_UNIFORM_BUFFER_SIZE, true);
 			}
 
 			Buffer = PendingState.ZeroFilledDummyUniformBuffer;
diff --git a/Engine/Source/Runtime/OpenGLDrv/Private/OpenGLTexture.cpp b/Engine/Source/Runtime/OpenGLDrv/Private/OpenGLTexture.cpp
--- a/Engine/Source/Runtime/OpenGLDrv/Private/OpenGLTexture.cpp
+++ b/Engine/Source/Runtime/OpenGLDrv/Private/OpenGLTexture.cpp
@@ -13,6 +13,7 @@
 #include "OpenGLDrvPrivate.h"
 #include "HAL/LowLevelMemTracker.h"
 #include "Engine/Texture.h"
+#include "RHICoreStats.h"
 
 #if PLATFORM_ANDROID
 #include "ThirdParty/Android/detex/AndroidETC.h"
@@ -97,47 +98,21 @@ void FOpenGLTexture::UpdateTextureStats(FOpenGLTexture* Texture, bool bAllocatin
 {
 	const FRHITextureDesc& Desc = Texture->GetDesc();
 
-	const bool bRenderTarget = EnumHasAnyFlags(Desc.Flags, TexCreate_RenderTargetable | TexCreate_ResolveTargetable | TexCreate_DepthStencilTargetable);
-	const int64 TextureSize = bAllocating
-		?  int64(Texture->MemorySize)
-		: -int64(Texture->MemorySize);
+	const uint64 TextureSize = Texture->MemorySize;
 
-	const int64 TextureSizeInKBs = bAllocating
-		?  FMath::DivideAndRoundUp(int64(Texture->MemorySize), 1024ll)
-		: -FMath::DivideAndRoundUp(int64(Texture->MemorySize), 1024ll);
+	const bool bOnlyStreamableTexturesInTexturePool = false;
+	UE::RHICore::UpdateGlobalTextureStats(Desc, TextureSize, bOnlyStreamableTexturesInTexturePool, bAllocating);
 
-	if (bRenderTarget)
-	{
-		switch (Desc.Dimension)
-		{
-		default: checkNoEntry();
-		case ETextureDimension::Texture2D:		  INC_MEMORY_STAT_BY(STAT_RenderTargetMemory2D  , TextureSize); break;
-		case ETextureDimension::Texture2DArray:	  INC_MEMORY_STAT_BY(STAT_RenderTargetMemory2D  , TextureSize); break;
-		case ETextureDimension::Texture3D:		  INC_MEMORY_STAT_BY(STAT_RenderTargetMemory3D  , TextureSize); break;
-		case ETextureDimension::TextureCube:	  INC_MEMORY_STAT_BY(STAT_RenderTargetMemoryCube, TextureSize); break;
-		case ETextureDimension::TextureCubeArray: INC_MEMORY_STAT_BY(STAT_RenderTargetMemoryCube, TextureSize); break;
-		}
+	const int64 TextureSizeDelta = bAllocating ? int64(TextureSize) : -int64(TextureSize);
 
-		GCurrentRendertargetMemorySize = GCurrentRendertargetMemorySize + TextureSizeInKBs;
-		LLM_SCOPED_PAUSE_TRACKING_WITH_ENUM_AND_AMOUNT(ELLMTag::GraphicsPlatform, TextureSize, ELLMTracker::Platform, ELLMAllocType::None);
-		LLM_SCOPED_PAUSE_TRACKING_WITH_ENUM_AND_AMOUNT(ELLMTag::RenderTargets   , TextureSize, ELLMTracker::Default , ELLMAllocType::None);
-	}
-	else
-	{
-		switch (Desc.Dimension)
-		{
-		default: checkNoEntry();
-		case ETextureDimension::Texture2D:		  INC_MEMORY_STAT_BY(STAT_TextureMemory2D  , TextureSize); break;
-		case ETextureDimension::Texture2DArray:	  INC_MEMORY_STAT_BY(STAT_TextureMemory2D  , TextureSize); break;
-		case ETextureDimension::Texture3D:		  INC_MEMORY_STAT_BY(STAT_TextureMemory3D  , TextureSize); break;
-		case ETextureDimension::TextureCube:	  INC_MEMORY_STAT_BY(STAT_TextureMemoryCube, TextureSize); break;
-		case ETextureDimension::TextureCubeArray: INC_MEMORY_STAT_BY(STAT_TextureMemoryCube, TextureSize); break;
-		}
+#if ENABLE_LOW_LEVEL_MEM_TRACKER
+	const ELLMTag TextureTag = EnumHasAnyFlags(Desc.Flags, ETextureCreateFlags::RenderTargetable | ETextureCreateFlags::ResolveTargetable | ETextureCreateFlags::DepthStencilTargetable)
+		? ELLMTag::RenderTargets
+		: ELLMTag::Textures;
 
-		GCurrentTextureMemorySize = GCurrentTextureMemorySize + TextureSizeInKBs;
-		LLM_SCOPED_PAUSE_TRACKING_WITH_ENUM_AND_AMOUNT(ELLMTag::GraphicsPlatform, TextureSize, ELLMTracker::Platform, ELLMAllocType::None);
-		LLM_SCOPED_PAUSE_TRACKING_WITH_ENUM_AND_AMOUNT(ELLMTag::Textures        , TextureSize, ELLMTracker::Default , ELLMAllocType::None);
-	}
+	LLM_SCOPED_PAUSE_TRACKING_WITH_ENUM_AND_AMOUNT(ELLMTag::GraphicsPlatform, TextureSizeDelta, ELLMTracker::Platform, ELLMAllocType::None);
+	LLM_SCOPED_PAUSE_TRACKING_WITH_ENUM_AND_AMOUNT(TextureTag               , TextureSizeDelta, ELLMTracker::Default , ELLMAllocType::None);
+#endif // ENABLE_LOW_LEVEL_MEM_TRACKER
 }
 
 FDynamicRHI::FRHICalcTextureSizeResult FOpenGLDynamicRHI::RHICalcTexturePlatformSize(FRHITextureDesc const& Desc, uint32 FirstMipIndex)
@@ -155,15 +130,12 @@ FDynamicRHI::FRHICalcTextureSizeResult FOpenGLDynamicRHI::RHICalcTexturePlatform
  */
 void FOpenGLDynamicRHI::RHIGetTextureMemoryStats(FTextureMemoryStats& OutStats)
 {
+	UE::RHICore::FillBaselineTextureMemoryStats(OutStats);
+
 	OutStats.DedicatedVideoMemory = GOpenGLDedicatedVideoMemory;
-    OutStats.DedicatedSystemMemory = 0;
-    OutStats.SharedSystemMemory = 0;
 	OutStats.TotalGraphicsMemory = GOpenGLTotalGraphicsMemory ? GOpenGLTotalGraphicsMemory : -1;
 
-	OutStats.AllocatedMemorySize = int64(GCurrentTextureMemorySize) * 1024;
-	OutStats.LargestContiguousAllocation = OutStats.AllocatedMemorySize;
-	OutStats.TexturePoolSize = GTexturePoolSize;
-	OutStats.PendingMemoryAdjustment = 0;
+	OutStats.LargestContiguousAllocation = OutStats.StreamingMemorySize;
 }
 
 
diff --git a/Engine/Source/Runtime/OpenGLDrv/Private/OpenGLUniformBuffer.cpp b/Engine/Source/Runtime/OpenGLDrv/Private/OpenGLUniformBuffer.cpp
--- a/Engine/Source/Runtime/OpenGLDrv/Private/OpenGLUniformBuffer.cpp
+++ b/Engine/Source/Runtime/OpenGLDrv/Private/OpenGLUniformBuffer.cpp
@@ -233,7 +233,7 @@ static void ReleaseUniformBuffer(bool bEmulatedBufferData, GLuint Resource, uint
 
 		RunOnGLRenderContextThread(MoveTemp(DeleteGLBuffer));
 	}
-	DecrementBufferMemory(GL_UNIFORM_BUFFER, AllocatedSize);
+	OpenGLBufferStats::UpdateUniformBufferStats(AllocatedSize, false);
 }
 
 // Does per-frame global updating for the uniform buffer pool.
@@ -606,7 +606,7 @@ static FUniformBufferRHIRef CreateUniformBuffer(const void* Contents, const FRHI
 		}
 	}
 
-	IncrementBufferMemory(GL_UNIFORM_BUFFER, AllocatedSize);
+	OpenGLBufferStats::UpdateUniformBufferStats(AllocatedSize, true);
 
 	check(!bUseEmulatedUBs || (IsValidRef(EmulatedUniformDataRef) && (EmulatedUniformDataRef->Data.Num() * EmulatedUniformDataRef->Data.GetTypeSize() == AllocatedSize)));
 
diff --git a/Engine/Source/Runtime/OpenGLDrv/Private/OpenGLUtil.cpp b/Engine/Source/Runtime/OpenGLDrv/Private/OpenGLUtil.cpp
--- a/Engine/Source/Runtime/OpenGLDrv/Private/OpenGLUtil.cpp
+++ b/Engine/Source/Runtime/OpenGLDrv/Private/OpenGLUtil.cpp
@@ -8,6 +8,7 @@
 #include "Stats/Stats.h"
 #include "OpenGLDrv.h"
 #include "OpenGLDrvPrivate.h"
+#include "RHICoreStats.h"
 
 void VerifyOpenGLResult(GLenum ErrorCode, const TCHAR* Msg1, const TCHAR* Msg2, const TCHAR* Filename, uint32 Line)
 {
@@ -101,54 +102,15 @@ DEFINE_STAT(STAT_OpenGLUniformBindTime);
 DEFINE_STAT(STAT_OpenGLVBOSetupTime);
 #endif
 
-void IncrementBufferMemory(GLenum Type, uint32 NumBytes)
+
+void OpenGLBufferStats::UpdateUniformBufferStats(int64 BufferSize, bool bAllocating)
 {
-	if (Type == GL_SHADER_STORAGE_BUFFER)
-	{
-		INC_MEMORY_STAT_BY(STAT_StructuredBufferMemory,NumBytes);
-	}
-	else if (Type == GL_UNIFORM_BUFFER)
-	{
-		INC_MEMORY_STAT_BY(STAT_UniformBufferMemory,NumBytes);
-	}
-	else if (Type == GL_ELEMENT_ARRAY_BUFFER)
-	{
-		INC_MEMORY_STAT_BY(STAT_IndexBufferMemory,NumBytes);
-	}
-	else if (Type == GL_PIXEL_UNPACK_BUFFER)
-	{
-		INC_MEMORY_STAT_BY(STAT_PixelBufferMemory,NumBytes);
-	}
-	else
-	{
-		check(Type == GL_ARRAY_BUFFER);
-		INC_MEMORY_STAT_BY(STAT_VertexBufferMemory,NumBytes);
-	}
+	UE::RHICore::UpdateGlobalUniformBufferStats(BufferSize, bAllocating);
 }
 
-void DecrementBufferMemory(GLenum Type, uint32 NumBytes)
+void OpenGLBufferStats::UpdateBufferStats(const FRHIBufferDesc& BufferDesc, bool bAllocating)
 {
-	if (Type == GL_SHADER_STORAGE_BUFFER)
-	{
-		DEC_MEMORY_STAT_BY(STAT_StructuredBufferMemory,NumBytes);
-	}
-	else if (Type == GL_UNIFORM_BUFFER)
-	{
-		DEC_MEMORY_STAT_BY(STAT_UniformBufferMemory,NumBytes);
-	}
-	else if (Type == GL_ELEMENT_ARRAY_BUFFER)
-	{
-		DEC_MEMORY_STAT_BY(STAT_IndexBufferMemory,NumBytes);
-	}
-	else if (Type == GL_PIXEL_UNPACK_BUFFER)
-	{
-		DEC_MEMORY_STAT_BY(STAT_PixelBufferMemory,NumBytes);
-	}
-	else if (Type != 0) // CreateInfo.bWithoutNativeResource
-	{
-		check(Type == GL_ARRAY_BUFFER);
-		DEC_MEMORY_STAT_BY(STAT_VertexBufferMemory,NumBytes);
-	}
+	UE::RHICore::UpdateGlobalBufferStats(BufferDesc, BufferDesc.Size, bAllocating);
 }
 
 // Run passed function on whichever thread owns the render context.
diff --git a/Engine/Source/Runtime/OpenGLDrv/Public/OpenGLResources.h b/Engine/Source/Runtime/OpenGLDrv/Public/OpenGLResources.h
--- a/Engine/Source/Runtime/OpenGLDrv/Public/OpenGLResources.h
+++ b/Engine/Source/Runtime/OpenGLDrv/Public/OpenGLResources.h
@@ -58,8 +58,11 @@ namespace OpenGLConsoleVariables
 #define RESTRICT_SUBDATA_SIZE 0
 #endif
 
-void IncrementBufferMemory(GLenum Type, uint32 NumBytes);
-void DecrementBufferMemory(GLenum Type, uint32 NumBytes);
+namespace OpenGLBufferStats
+{
+	void UpdateUniformBufferStats(int64 BufferSize, bool bAllocating);
+	void UpdateBufferStats(const FRHIBufferDesc& BufferDesc, bool bAllocating);
+}
 
 // Extra stats for finer-grained timing
 // They shouldn't always be on, as they may impact overall performance
@@ -173,7 +176,7 @@ public:
 				LoadData(0, BaseType::GetSize(), InData);
 			}
 #endif
-			IncrementBufferMemory(Type, BaseType::GetSize());
+			OpenGLBufferStats::UpdateBufferStats(BaseType::GetDesc(), true);
 		};
 
 		if (!BufferDesc.IsNull())
@@ -509,7 +512,7 @@ public:
 			}
 
 			LockBuffer = nullptr;
-			DecrementBufferMemory(Type, BaseType::GetSize());
+			OpenGLBufferStats::UpdateBufferStats(BaseType::GetDesc(), false);
 
 			ReleaseCachedBuffer();
 		}
@@ -546,6 +549,7 @@ public:
 		return true;
 	}
 
+	const FRHIBufferDesc& GetDesc() const { return Desc; }
 	uint32 GetSize() const { return Desc.Size; }
 	EBufferUsageFlags GetUsage() const { return Desc.Usage; }
 
@@ -597,12 +601,12 @@ struct FOpenGLEUniformBufferData : public FRefCountedObject
 		uint32 SizeInUint32s = (SizeInBytes + 3) / 4;
 		Data.Empty(SizeInUint32s);
 		Data.AddUninitialized(SizeInUint32s);
-		IncrementBufferMemory(GL_UNIFORM_BUFFER,Data.GetAllocatedSize());
+		OpenGLBufferStats::UpdateUniformBufferStats(Data.GetAllocatedSize(), true);
 	}
 
 	~FOpenGLEUniformBufferData()
 	{
-		DecrementBufferMemory(GL_UNIFORM_BUFFER,Data.GetAllocatedSize());
+		OpenGLBufferStats::UpdateUniformBufferStats(Data.GetAllocatedSize(), false);
 	}
 
 	TArray<uint32> Data;
diff --git a/Engine/Source/Runtime/RHI/Private/RHIStats.cpp b/Engine/Source/Runtime/RHI/Private/RHIStats.cpp
--- a/Engine/Source/Runtime/RHI/Private/RHIStats.cpp
+++ b/Engine/Source/Runtime/RHI/Private/RHIStats.cpp
@@ -19,6 +19,7 @@ int32 GNumPrimitivesDrawnRHI[MAX_NUM_GPUS] = {};
 DEFINE_STAT(STAT_RenderTargetMemory2D);
 DEFINE_STAT(STAT_RenderTargetMemory3D);
 DEFINE_STAT(STAT_RenderTargetMemoryCube);
+DEFINE_STAT(STAT_UAVTextureMemory);
 DEFINE_STAT(STAT_TextureMemory2D);
 DEFINE_STAT(STAT_TextureMemory3D);
 DEFINE_STAT(STAT_TextureMemoryCube);
@@ -27,7 +28,9 @@ DEFINE_STAT(STAT_IndexBufferMemory);
 DEFINE_STAT(STAT_VertexBufferMemory);
 DEFINE_STAT(STAT_RTAccelerationStructureMemory);
 DEFINE_STAT(STAT_StructuredBufferMemory);
-DEFINE_STAT(STAT_PixelBufferMemory);
+DEFINE_STAT(STAT_ByteAddressBufferMemory);
+DEFINE_STAT(STAT_DrawIndirectBufferMemory);
+DEFINE_STAT(STAT_MiscBufferMemory);
 
 DEFINE_STAT(STAT_SamplerDescriptorsAllocated);
 DEFINE_STAT(STAT_ResourceDescriptorsAllocated);
diff --git a/Engine/Source/Runtime/RHI/Public/RHIDefinitions.h b/Engine/Source/Runtime/RHI/Public/RHIDefinitions.h
--- a/Engine/Source/Runtime/RHI/Public/RHIDefinitions.h
+++ b/Engine/Source/Runtime/RHI/Public/RHIDefinitions.h
@@ -971,6 +971,9 @@ enum class ETextureCreateFlags : uint64
 	ReservedResource                  = 1ull << 37,
 	/** EXPERIMENTAL: Used with ReservedResource flag to immediately allocate and commit memory on creation. May use N small physical memory allocations instead of a single large one. */
 	ImmediateCommit                   = 1ull << 38,
+
+	/** Don't lump this texture with streaming memory when tracking total texture allocation sizes */
+	ForceIntoNonStreamingMemoryTracking = 1ull << 39,
 };
 ENUM_CLASS_FLAGS(ETextureCreateFlags);
 
diff --git a/Engine/Source/Runtime/RHI/Public/RHIGlobals.h b/Engine/Source/Runtime/RHI/Public/RHIGlobals.h
--- a/Engine/Source/Runtime/RHI/Public/RHIGlobals.h
+++ b/Engine/Source/Runtime/RHI/Public/RHIGlobals.h
@@ -361,11 +361,19 @@ struct FRHIGlobals
 	bool SupportsTextureStreaming = false;
 
 	/** Amount of memory allocated by textures. In kilobytes. */
+	UE_DEPRECATED(5.3, "CurrentTextureMemorySize was too vague, use StreamingTextureMemorySizeInKB in its place")
 	volatile int32 CurrentTextureMemorySize = 0;
 
 	/** Amount of memory allocated by rendertargets. In kilobytes. */
+	UE_DEPRECATED(5.3, "CurrentRendertargetMemorySize was too vague, use NonStreamingTextureMemorySizeInKB in its place")
 	volatile int32 CurrentRendertargetMemorySize = 0;
 
+	/** Amount of memory allocated by streaming textures. In kilobytes. */
+	volatile uint64 StreamingTextureMemorySizeInKB = 0;
+
+	/** Amount of memory allocated by non streaming textures. In kilobytes. */
+	volatile uint64 NonStreamingTextureMemorySizeInKB = 0;
+
 	/** Current texture streaming pool size, in bytes. 0 means unlimited. */
 	int64 TexturePoolSize = 0 * 1024 * 1024;
 
@@ -375,6 +383,12 @@ struct FRHIGlobals
 	/** Amount of local video memory demoted to system memory. In bytes. */
 	uint64 DemotedLocalMemorySize = 0;
 
+	/** Amount of memory allocated by buffers */
+	volatile uint64 BufferMemorySize = 0;
+
+	/** Amount of memory allocated by uniform buffers */
+	volatile uint64 UniformBufferMemorySize = 0;
+
 	/** Whether or not the RHI can handle a non-zero BaseVertexIndex - extra SetStreamSource calls will be needed if this is false */
 	bool SupportsBaseVertexIndex = true;
 
diff --git a/Engine/Source/Runtime/RHI/Public/RHIStats.h b/Engine/Source/Runtime/RHI/Public/RHIStats.h
--- a/Engine/Source/Runtime/RHI/Public/RHIStats.h
+++ b/Engine/Source/Runtime/RHI/Public/RHIStats.h
@@ -12,35 +12,36 @@ struct FTextureMemoryStats
 	// Hardware state (never change after device creation):
 
 	// -1 if unknown, in bytes
-	int64 DedicatedVideoMemory;
+	int64 DedicatedVideoMemory = -1;
+
 	// -1 if unknown, in bytes
-	int64 DedicatedSystemMemory;
+	int64 DedicatedSystemMemory = -1;
+
 	// -1 if unknown, in bytes
-	int64 SharedSystemMemory;
+	int64 SharedSystemMemory = -1;
+
 	// Total amount of "graphics memory" that we think we can use for all our graphics resources, in bytes. -1 if unknown.
-	int64 TotalGraphicsMemory;
+	int64 TotalGraphicsMemory = -1;
 
 	// Size of allocated memory, in bytes
-	int64 AllocatedMemorySize;
+	UE_DEPRECATED(5.3, "AllocatedMemorySize was too vague, use StreamingMemorySize in its place")
+	int64 AllocatedMemorySize = 0;
+
+	// Size of memory allocated to streaming textures, in bytes
+	uint64 StreamingMemorySize = 0;
+
+	// Size of memory allocated to non-streaming textures, in bytes
+	uint64 NonStreamingMemorySize = 0;
+
 	// Size of the largest memory fragment, in bytes
-	int64 LargestContiguousAllocation;
+	int64 LargestContiguousAllocation = 0;
+	
 	// 0 if streaming pool size limitation is disabled, in bytes
-	int64 TexturePoolSize;
+	int64 TexturePoolSize = 0;
+
 	// Upcoming adjustments to allocated memory, in bytes (async reallocations)
-	int32 PendingMemoryAdjustment;
-
-	// defaults
-	FTextureMemoryStats()
-		: DedicatedVideoMemory(-1)
-		, DedicatedSystemMemory(-1)
-		, SharedSystemMemory(-1)
-		, TotalGraphicsMemory(-1)
-		, AllocatedMemorySize(0)
-		, LargestContiguousAllocation(0)
-		, TexturePoolSize(0)
-		, PendingMemoryAdjustment(0)
-	{
-	}
+	UE_DEPRECATED(5.3, "PendingMemoryAdjustment is unused")
+	int32 PendingMemoryAdjustment = 0;
 
 	bool AreHardwareStatsValid() const
 	{
@@ -55,7 +56,7 @@ struct FTextureMemoryStats
 
 	int64 ComputeAvailableMemorySize() const
 	{
-		return FMath::Max(TexturePoolSize - AllocatedMemorySize, (int64)0);
+		return FMath::Max<int64>(TexturePoolSize - StreamingMemorySize, 0);
 	}
 };
 
@@ -202,18 +203,24 @@ private:
 };
 
 // RHI memory stats.
-DECLARE_MEMORY_STAT_POOL_EXTERN(TEXT("Render target memory 2D"), STAT_RenderTargetMemory2D, STATGROUP_RHI, FPlatformMemory::MCR_GPU, RHI_API);
-DECLARE_MEMORY_STAT_POOL_EXTERN(TEXT("Render target memory 3D"), STAT_RenderTargetMemory3D, STATGROUP_RHI, FPlatformMemory::MCR_GPU, RHI_API);
-DECLARE_MEMORY_STAT_POOL_EXTERN(TEXT("Render target memory Cube"), STAT_RenderTargetMemoryCube, STATGROUP_RHI, FPlatformMemory::MCR_GPU, RHI_API);
-DECLARE_MEMORY_STAT_POOL_EXTERN(TEXT("Texture memory 2D"), STAT_TextureMemory2D, STATGROUP_RHI, FPlatformMemory::MCR_GPU, RHI_API);
-DECLARE_MEMORY_STAT_POOL_EXTERN(TEXT("Texture memory 3D"), STAT_TextureMemory3D, STATGROUP_RHI, FPlatformMemory::MCR_GPU, RHI_API);
-DECLARE_MEMORY_STAT_POOL_EXTERN(TEXT("Texture memory Cube"), STAT_TextureMemoryCube, STATGROUP_RHI, FPlatformMemory::MCR_GPU, RHI_API);
-DECLARE_MEMORY_STAT_POOL_EXTERN(TEXT("Uniform buffer memory"), STAT_UniformBufferMemory, STATGROUP_RHI, FPlatformMemory::MCR_GPU, RHI_API);
-DECLARE_MEMORY_STAT_POOL_EXTERN(TEXT("Index buffer memory"), STAT_IndexBufferMemory, STATGROUP_RHI, FPlatformMemory::MCR_GPU, RHI_API);
-DECLARE_MEMORY_STAT_POOL_EXTERN(TEXT("Vertex buffer memory"), STAT_VertexBufferMemory, STATGROUP_RHI, FPlatformMemory::MCR_GPU, RHI_API);
-DECLARE_MEMORY_STAT_POOL_EXTERN(TEXT("Ray Tracing Acceleration Structure memory"), STAT_RTAccelerationStructureMemory, STATGROUP_RHI, FPlatformMemory::MCR_GPU, RHI_API);
-DECLARE_MEMORY_STAT_POOL_EXTERN(TEXT("Structured buffer memory"), STAT_StructuredBufferMemory, STATGROUP_RHI, FPlatformMemory::MCR_GPU, RHI_API);
-DECLARE_MEMORY_STAT_POOL_EXTERN(TEXT("Pixel buffer memory"), STAT_PixelBufferMemory, STATGROUP_RHI, FPlatformMemory::MCR_GPU, RHI_API);
+DECLARE_MEMORY_STAT_POOL_EXTERN(TEXT("Render Target 2D Memory"), STAT_RenderTargetMemory2D, STATGROUP_RHI, FPlatformMemory::MCR_GPU, RHI_API);
+DECLARE_MEMORY_STAT_POOL_EXTERN(TEXT("Render Target 3D Memory"), STAT_RenderTargetMemory3D, STATGROUP_RHI, FPlatformMemory::MCR_GPU, RHI_API);
+DECLARE_MEMORY_STAT_POOL_EXTERN(TEXT("Render Target Cube Memory"), STAT_RenderTargetMemoryCube, STATGROUP_RHI, FPlatformMemory::MCR_GPU, RHI_API);
+
+DECLARE_MEMORY_STAT_POOL_EXTERN(TEXT("UAV Texture Memory"), STAT_UAVTextureMemory, STATGROUP_RHI, FPlatformMemory::MCR_GPU, RHI_API);
+
+DECLARE_MEMORY_STAT_POOL_EXTERN(TEXT("Texture 2D Memory"), STAT_TextureMemory2D, STATGROUP_RHI, FPlatformMemory::MCR_GPU, RHI_API);
+DECLARE_MEMORY_STAT_POOL_EXTERN(TEXT("Texture 3D Memory"), STAT_TextureMemory3D, STATGROUP_RHI, FPlatformMemory::MCR_GPU, RHI_API);
+DECLARE_MEMORY_STAT_POOL_EXTERN(TEXT("Texture Cube Memory"), STAT_TextureMemoryCube, STATGROUP_RHI, FPlatformMemory::MCR_GPU, RHI_API);
+
+DECLARE_MEMORY_STAT_POOL_EXTERN(TEXT("Uniform Buffer Memory"), STAT_UniformBufferMemory, STATGROUP_RHI, FPlatformMemory::MCR_GPU, RHI_API);
+DECLARE_MEMORY_STAT_POOL_EXTERN(TEXT("Index Buffer Memory"), STAT_IndexBufferMemory, STATGROUP_RHI, FPlatformMemory::MCR_GPU, RHI_API);
+DECLARE_MEMORY_STAT_POOL_EXTERN(TEXT("Vertex Buffer Memory"), STAT_VertexBufferMemory, STATGROUP_RHI, FPlatformMemory::MCR_GPU, RHI_API);
+DECLARE_MEMORY_STAT_POOL_EXTERN(TEXT("RayTracing Acceleration Structure Memory"), STAT_RTAccelerationStructureMemory, STATGROUP_RHI, FPlatformMemory::MCR_GPU, RHI_API);
+DECLARE_MEMORY_STAT_POOL_EXTERN(TEXT("Structured Buffer Memory"), STAT_StructuredBufferMemory, STATGROUP_RHI, FPlatformMemory::MCR_GPU, RHI_API);
+DECLARE_MEMORY_STAT_POOL_EXTERN(TEXT("Byte Address Buffer Memory"), STAT_ByteAddressBufferMemory, STATGROUP_RHI, FPlatformMemory::MCR_GPU, RHI_API);
+DECLARE_MEMORY_STAT_POOL_EXTERN(TEXT("Draw Indirect Buffer Memory"), STAT_DrawIndirectBufferMemory, STATGROUP_RHI, FPlatformMemory::MCR_GPU, RHI_API);
+DECLARE_MEMORY_STAT_POOL_EXTERN(TEXT("Misc Buffer Memory"), STAT_MiscBufferMemory, STATGROUP_RHI, FPlatformMemory::MCR_GPU, RHI_API);
 
 DECLARE_DWORD_ACCUMULATOR_STAT_EXTERN(TEXT("Sampler Descriptors Allocated"), STAT_SamplerDescriptorsAllocated, STATGROUP_RHI, RHI_API);
 DECLARE_DWORD_ACCUMULATOR_STAT_EXTERN(TEXT("Resource Descriptors Allocated"), STAT_ResourceDescriptorsAllocated, STATGROUP_RHI, RHI_API);
diff --git a/Engine/Source/Runtime/RHICore/Internal/RHICoreStats.h b/Engine/Source/Runtime/RHICore/Internal/RHICoreStats.h
new file mode 100644
--- /dev/null
+++ b/Engine/Source/Runtime/RHICore/Internal/RHICoreStats.h
@@ -0,0 +1,152 @@
+// Copyright Epic Games, Inc. All Rights Reserved.
+
+#pragma once
+
+#include "HAL/PlatformAtomics.h"
+#include "RHIGlobals.h"
+#include "RHIResources.h"
+#include "RHIStats.h"
+
+namespace UE::RHICore
+{
+	inline void UpdateGlobalTextureStats(ETextureCreateFlags TextureFlags, ETextureDimension Dimension, uint64 TextureSizeInBytes, bool bOnlyStreamableTexturesInTexturePool, bool bAllocating)
+	{
+		constexpr ETextureCreateFlags AllRenderTargetFlags = ETextureCreateFlags::RenderTargetable | ETextureCreateFlags::ResolveTargetable | ETextureCreateFlags::DepthStencilTargetable;
+#if STATS
+		const int64 TextureSizeDeltaInBytes = bAllocating ? static_cast<int64>(TextureSizeInBytes) : -static_cast<int64>(TextureSizeInBytes);
+		if (EnumHasAnyFlags(TextureFlags, AllRenderTargetFlags))
+		{
+			switch (Dimension)
+			{
+			default:
+				checkNoEntry();
+				[[fallthrough]];
+			case ETextureDimension::Texture2D:
+				[[fallthrough]];
+			case ETextureDimension::Texture2DArray:
+				INC_MEMORY_STAT_BY(STAT_RenderTargetMemory2D, TextureSizeDeltaInBytes);
+				break;
+			case ETextureDimension::Texture3D:
+				INC_MEMORY_STAT_BY(STAT_RenderTargetMemory3D, TextureSizeDeltaInBytes);
+				break;
+			case ETextureDimension::TextureCube:
+				[[fallthrough]];
+			case ETextureDimension::TextureCubeArray:
+				INC_MEMORY_STAT_BY(STAT_RenderTargetMemoryCube, TextureSizeDeltaInBytes);
+				break;
+			};
+		}
+		else if (EnumHasAnyFlags(TextureFlags, ETextureCreateFlags::UAV))
+		{
+			INC_MEMORY_STAT_BY(STAT_UAVTextureMemory, TextureSizeDeltaInBytes);
+		}
+		else
+		{
+			switch (Dimension)
+			{
+			default:
+				checkNoEntry();
+				[[fallthrough]];
+			case ETextureDimension::Texture2D:
+				[[fallthrough]];
+			case ETextureDimension::Texture2DArray:
+				INC_MEMORY_STAT_BY(STAT_TextureMemory2D, TextureSizeDeltaInBytes);
+				break;
+			case ETextureDimension::Texture3D:
+				INC_MEMORY_STAT_BY(STAT_TextureMemory3D, TextureSizeDeltaInBytes);
+				break;
+			case ETextureDimension::TextureCube:
+				[[fallthrough]];
+			case ETextureDimension::TextureCubeArray:
+				INC_MEMORY_STAT_BY(STAT_TextureMemoryCube, TextureSizeDeltaInBytes);
+				break;
+			};
+		}
+#endif
+		const int64 AlignedTextureSizeInKB = static_cast<int64>(Align(TextureSizeInBytes, 1024) / 1024);
+		const int64 TextureSizeDeltaInKB = bAllocating ? AlignedTextureSizeInKB : -AlignedTextureSizeInKB;
+
+		const bool bAlwaysExcludedFromStreamingSize = EnumHasAnyFlags(TextureFlags,
+			AllRenderTargetFlags
+			| ETextureCreateFlags::UAV
+			| ETextureCreateFlags::ForceIntoNonStreamingMemoryTracking
+		);
+
+		const bool bStreamable = EnumHasAnyFlags(TextureFlags, ETextureCreateFlags::Streamable);
+
+		if (!bAlwaysExcludedFromStreamingSize && (!bOnlyStreamableTexturesInTexturePool || bStreamable))
+		{
+			FPlatformAtomics::InterlockedAdd((volatile int64*)&GRHIGlobals.StreamingTextureMemorySizeInKB, TextureSizeDeltaInKB);
+		}
+		else
+		{
+			FPlatformAtomics::InterlockedAdd((volatile int64*)&GRHIGlobals.NonStreamingTextureMemorySizeInKB, TextureSizeDeltaInKB);
+		}
+	}
+
+	inline void UpdateGlobalTextureStats(const FRHITextureDesc& TextureDesc, uint64 TextureSizeInBytes, bool bOnlyStreamableTexturesInTexturePool, bool bAllocating)
+	{
+		return UpdateGlobalTextureStats(TextureDesc.Flags, TextureDesc.Dimension, TextureSizeInBytes, bOnlyStreamableTexturesInTexturePool, bAllocating);
+	}
+
+	inline void FillBaselineTextureMemoryStats(FTextureMemoryStats& OutStats)
+	{
+		OutStats.StreamingMemorySize    = GRHIGlobals.StreamingTextureMemorySizeInKB * 1024;
+		OutStats.NonStreamingMemorySize = GRHIGlobals.NonStreamingTextureMemorySizeInKB * 1024;
+		OutStats.TexturePoolSize        = GRHIGlobals.TexturePoolSize;
+
+		PRAGMA_DISABLE_DEPRECATION_WARNINGS
+		OutStats.AllocatedMemorySize = OutStats.StreamingMemorySize;
+		PRAGMA_ENABLE_DEPRECATION_WARNINGS
+	}
+
+	inline void UpdateGlobalBufferStats(const FRHIBufferDesc& BufferDesc, int64 BufferSizeDelta)
+	{
+#if STATS
+		if (EnumHasAnyFlags(BufferDesc.Usage, EBufferUsageFlags::VertexBuffer))
+		{
+			INC_MEMORY_STAT_BY(STAT_VertexBufferMemory, BufferSizeDelta);
+		}
+		else if (EnumHasAnyFlags(BufferDesc.Usage, EBufferUsageFlags::IndexBuffer))
+		{
+			INC_MEMORY_STAT_BY(STAT_IndexBufferMemory, BufferSizeDelta);
+		}
+		else if (EnumHasAnyFlags(BufferDesc.Usage, EBufferUsageFlags::AccelerationStructure))
+		{
+			INC_MEMORY_STAT_BY(STAT_RTAccelerationStructureMemory, BufferSizeDelta);
+		}
+		else if (EnumHasAnyFlags(BufferDesc.Usage, EBufferUsageFlags::ByteAddressBuffer))
+		{
+			INC_MEMORY_STAT_BY(STAT_ByteAddressBufferMemory, BufferSizeDelta);
+		}
+		else if (EnumHasAnyFlags(BufferDesc.Usage, EBufferUsageFlags::DrawIndirect))
+		{
+			INC_MEMORY_STAT_BY(STAT_DrawIndirectBufferMemory, BufferSizeDelta);
+		}
+		else if (EnumHasAnyFlags(BufferDesc.Usage, EBufferUsageFlags::StructuredBuffer))
+		{
+			INC_MEMORY_STAT_BY(STAT_StructuredBufferMemory, BufferSizeDelta);
+		}
+		else
+		{
+			INC_MEMORY_STAT_BY(STAT_MiscBufferMemory, BufferSizeDelta);
+		}
+#endif
+
+		FPlatformAtomics::InterlockedAdd((volatile int64*)&GRHIGlobals.BufferMemorySize, BufferSizeDelta);
+	}
+
+	inline void UpdateGlobalBufferStats(const FRHIBufferDesc& BufferDesc, uint64 BufferSize, bool bAllocating)
+	{
+		const int64 BufferSizeDelta = bAllocating ? static_cast<int64>(BufferSize) : -static_cast<int64>(BufferSize);
+		UpdateGlobalBufferStats(BufferDesc, BufferSizeDelta);
+	}
+
+	inline void UpdateGlobalUniformBufferStats(int64 BufferSize, bool bAllocating)
+	{
+		const int64 BufferSizeDelta = bAllocating ? static_cast<int64>(BufferSize) : -static_cast<int64>(BufferSize);
+
+		INC_MEMORY_STAT_BY(STAT_UniformBufferMemory, BufferSizeDelta);
+		FPlatformAtomics::InterlockedAdd((volatile int64*)&GRHIGlobals.UniformBufferMemorySize, BufferSizeDelta);
+	}
+}
diff --git a/Engine/Source/Runtime/RenderCore/Private/RenderTargetPool.cpp b/Engine/Source/Runtime/RenderCore/Private/RenderTargetPool.cpp
--- a/Engine/Source/Runtime/RenderCore/Private/RenderTargetPool.cpp
+++ b/Engine/Source/Runtime/RenderCore/Private/RenderTargetPool.cpp
@@ -76,6 +76,9 @@ TRefCountPtr<IPooledRenderTarget> FRenderTargetPool::FindFreeElement(FRHITexture
 	// We always want SRV access
 	Desc.Flags |= TexCreate_ShaderResource;
 
+	// Render target pool always forces textures into non streaming memory.
+	Desc.Flags |= ETextureCreateFlags::ForceIntoNonStreamingMemoryTracking;
+
 	const uint32 DescHash = GetTypeHash(Desc);
 
 	for (uint32 Index = 0, Num = (uint32)PooledRenderTargets.Num(); Index < Num; ++Index)
diff --git a/Engine/Source/Runtime/RenderCore/Public/RenderGraphResources.inl b/Engine/Source/Runtime/RenderCore/Public/RenderGraphResources.inl
--- a/Engine/Source/Runtime/RenderCore/Public/RenderGraphResources.inl
+++ b/Engine/Source/Runtime/RenderCore/Public/RenderGraphResources.inl
@@ -126,7 +126,7 @@ inline FRDGTextureDesc Translate(const FPooledRenderTargetDesc& InDesc)
 		OutDesc.Dimension = ETextureDimension::Texture2DArray;
 	}
 
-	OutDesc.Flags = InDesc.Flags;
+	OutDesc.Flags = InDesc.Flags | ETextureCreateFlags::ForceIntoNonStreamingMemoryTracking;
 	OutDesc.FastVRAMPercentage = InDesc.FastVRAMPercentage;
 	check(OutDesc.IsValid());
 
diff --git a/Engine/Source/Runtime/VulkanRHI/Private/VulkanBuffer.cpp b/Engine/Source/Runtime/VulkanRHI/Private/VulkanBuffer.cpp
--- a/Engine/Source/Runtime/VulkanRHI/Private/VulkanBuffer.cpp
+++ b/Engine/Source/Runtime/VulkanRHI/Private/VulkanBuffer.cpp
@@ -11,6 +11,7 @@
 #include "VulkanLLM.h"
 #include "VulkanRayTracing.h"
 #include "VulkanTransientResourceAllocator.h"
+#include "RHICoreStats.h"
 
 static TMap<FVulkanResourceMultiBuffer*, VulkanRHI::FPendingBufferLock> GPendingLockIBs;
 static FCriticalSection GPendingLockIBsMutex;
@@ -36,59 +37,9 @@ static FORCEINLINE VulkanRHI::FPendingBufferLock GetPendingBufferLock(FVulkanRes
 	return PendingLock;
 }
 
-static FORCEINLINE void UpdateVulkanBufferStats(uint64_t Size, VkBufferUsageFlags Usage, bool Allocating)
+static void UpdateVulkanBufferStats(const FRHIBufferDesc& BufferDesc, int64 BufferSize, bool bAllocating)
 {
-	const bool bUniformBuffer = !!(Usage & VK_BUFFER_USAGE_UNIFORM_BUFFER_BIT);
-	const bool bIndexBuffer = !!(Usage & VK_BUFFER_USAGE_INDEX_BUFFER_BIT);
-	const bool bVertexBuffer = !!(Usage & VK_BUFFER_USAGE_VERTEX_BUFFER_BIT);
-	const bool bAccelerationStructure = !!(Usage & VK_BUFFER_USAGE_ACCELERATION_STRUCTURE_STORAGE_BIT_KHR);
-
-	if (Allocating)
-	{
-		if (bUniformBuffer)
-		{
-			INC_MEMORY_STAT_BY(STAT_UniformBufferMemory, Size);
-		}
-		else if (bIndexBuffer)
-		{
-			INC_MEMORY_STAT_BY(STAT_IndexBufferMemory, Size);
-		}
-		else if (bVertexBuffer)
-		{
-			INC_MEMORY_STAT_BY(STAT_VertexBufferMemory, Size);
-		}
-		else if (bAccelerationStructure)
-		{
-			INC_MEMORY_STAT_BY(STAT_RTAccelerationStructureMemory, Size);
-		}
-		else
-		{
-			INC_MEMORY_STAT_BY(STAT_StructuredBufferMemory, Size);
-		}
-	}
-	else
-	{
-		if (bUniformBuffer)
-		{
-			DEC_MEMORY_STAT_BY(STAT_UniformBufferMemory, Size);
-		}
-		else if (bIndexBuffer)
-		{
-			DEC_MEMORY_STAT_BY(STAT_IndexBufferMemory, Size);
-		}
-		else if (bVertexBuffer)
-		{
-			DEC_MEMORY_STAT_BY(STAT_VertexBufferMemory, Size);
-		}
-		else if (bAccelerationStructure)
-		{
-			INC_MEMORY_STAT_BY(STAT_RTAccelerationStructureMemory, Size);
-		}
-		else
-		{
-			DEC_MEMORY_STAT_BY(STAT_StructuredBufferMemory, Size);
-		}
-	}
+	UE::RHICore::UpdateGlobalBufferStats(BufferDesc, BufferSize, bAllocating);
 }
 
 static VkDeviceAddress GetBufferDeviceAddress(FVulkanDevice* Device, VkBuffer Buffer)
@@ -297,7 +248,7 @@ void FVulkanResourceMultiBuffer::AdvanceBufferIndex()
 		NewBufferAlloc.AllocStatus = FBufferAlloc::EAllocStatus::InUse;
 		NewBufferAlloc.DeviceAddress = GetBufferDeviceAddress(Device, NewBufferAlloc.Alloc.GetBufferHandle()) + NewBufferAlloc.Alloc.Offset;
 
-		UpdateVulkanBufferStats(BufferSize, BufferUsageFlags, true);
+		UpdateVulkanBufferStats(GetDesc(), BufferSize, true);
 	}
 }
 
@@ -651,7 +602,7 @@ void FVulkanResourceMultiBuffer::ReleaseOwnership()
 
 	if (TotalSize > 0)
 	{
-		UpdateVulkanBufferStats(TotalSize, BufferUsageFlags, false);
+		UpdateVulkanBufferStats(GetDesc(), TotalSize, false);
 	}
 }
 
diff --git a/Engine/Source/Runtime/VulkanRHI/Private/VulkanTexture.cpp b/Engine/Source/Runtime/VulkanRHI/Private/VulkanTexture.cpp
--- a/Engine/Source/Runtime/VulkanRHI/Private/VulkanTexture.cpp
+++ b/Engine/Source/Runtime/VulkanRHI/Private/VulkanTexture.cpp
@@ -12,6 +12,7 @@
 #include "VulkanLLM.h"
 #include "VulkanBarriers.h"
 #include "VulkanTransientResourceAllocator.h"
+#include "RHICoreStats.h"
 
 int32 GVulkanSubmitOnTextureUnlock = 1;
 static FAutoConsoleVariableRef CVarVulkanSubmitOnTextureUnlock(
@@ -121,35 +122,20 @@ static TStatId GetVulkanStatEnum(bool bIsCube, bool bIs3D, bool bIsRT)
 	return TStatId();
 }
 
-static void UpdateVulkanTextureStats(int64 TextureSize, bool bIsCube, bool bIs3D, bool bIsRT)
+static void UpdateVulkanTextureStats(const FRHITextureDesc& TextureDesc, uint64 TextureSize, bool bAllocating)
 {
-	const int64 AlignedSize = (TextureSize > 0) ? Align(TextureSize, 1024) / 1024 : -(Align(-TextureSize, 1024) / 1024);
-	if (bIsRT == false)
-	{
-		FPlatformAtomics::InterlockedAdd(&GCurrentTextureMemorySize, AlignedSize);
-	}
-	else
-	{
-		FPlatformAtomics::InterlockedAdd(&GCurrentRendertargetMemorySize, AlignedSize);
-	}
-
-	INC_MEMORY_STAT_BY_FName(GetVulkanStatEnum(bIsCube, bIs3D, bIsRT).GetName(), TextureSize);
+	const bool bOnlyStreamableTexturesInTexturePool = false;
+	UE::RHICore::UpdateGlobalTextureStats(TextureDesc, TextureSize, bOnlyStreamableTexturesInTexturePool, bAllocating);
 }
 
-static void VulkanTextureAllocated(uint64 Size, VkImageViewType ImageType, bool bIsRT)
+static void VulkanTextureAllocated(const FRHITextureDesc& TextureDesc, uint64 Size)
 {
-	bool bIsCube = ImageType == VK_IMAGE_VIEW_TYPE_CUBE || ImageType == VK_IMAGE_VIEW_TYPE_CUBE_ARRAY;
-	bool bIs3D = ImageType == VK_IMAGE_VIEW_TYPE_3D ;
-
-	UpdateVulkanTextureStats(Size, bIsCube, bIs3D, bIsRT);
+	UpdateVulkanTextureStats(TextureDesc, Size, true);
 }
 
-static void VulkanTextureDestroyed(uint64 Size, VkImageViewType ImageType, bool bIsRT)
+static void VulkanTextureDestroyed(const FRHITextureDesc& TextureDesc, uint64 Size)
 {
-	bool bIsCube = ImageType == VK_IMAGE_VIEW_TYPE_CUBE || ImageType == VK_IMAGE_VIEW_TYPE_CUBE_ARRAY;
-	bool bIs3D = ImageType == VK_IMAGE_VIEW_TYPE_3D;
-
-	UpdateVulkanTextureStats(-(int64)Size, bIsCube, bIs3D, bIsRT);
+	UpdateVulkanTextureStats(TextureDesc, Size, false);
 }
 
 inline void FVulkanTexture::InternalLockWrite(FVulkanCommandListContext& Context, FVulkanTexture* Surface, const VkBufferImageCopy& Region, VulkanRHI::FStagingBuffer* StagingBuffer)
@@ -698,7 +684,7 @@ void FVulkanTexture::DestroySurface()
 			// If we don't own the allocation, it's transient memory not included in stats
 			if (Allocation.HasAllocation())
 			{
-				VulkanTextureDestroyed(Allocation.Size, GetViewType(), bRenderTarget);
+				VulkanTextureDestroyed(GetDesc(), Allocation.Size);
 			}
 
 			if (Image != VK_NULL_HANDLE)
@@ -878,6 +864,8 @@ void FVulkanTexture::SetInitialImageState(FVulkanCommandListContext& Context, Vk
 
 void FVulkanDynamicRHI::RHIGetTextureMemoryStats(FTextureMemoryStats& OutStats)
 {
+	UE::RHICore::FillBaselineTextureMemoryStats(OutStats);
+
 	check(Device);
 	const uint64 TotalGPUMemory = Device->GetDeviceMemoryManager().GetTotalMemory(true);
 	const uint64 TotalCPUMemory = Device->GetDeviceMemoryManager().GetTotalMemory(false);
@@ -887,10 +875,7 @@ void FVulkanDynamicRHI::RHIGetTextureMemoryStats(FTextureMemoryStats& OutStats)
 	OutStats.SharedSystemMemory = -1;
 	OutStats.TotalGraphicsMemory = TotalGPUMemory ? TotalGPUMemory : -1;
 
-	OutStats.AllocatedMemorySize = int64(GCurrentTextureMemorySize) * 1024;
-	OutStats.LargestContiguousAllocation = OutStats.AllocatedMemorySize;
-	OutStats.TexturePoolSize = GTexturePoolSize;
-	OutStats.PendingMemoryAdjustment = 0;
+	OutStats.LargestContiguousAllocation = OutStats.StreamingMemorySize;
 }
 
 bool FVulkanDynamicRHI::RHIGetTextureMemoryVisualizeData( FColor* /*TextureData*/, int32 /*SizeX*/, int32 /*SizeY*/, int32 /*Pitch*/, int32 /*PixelSize*/ )
@@ -1565,7 +1550,7 @@ FVulkanTexture::FVulkanTexture(FVulkanDevice& InDevice, const FRHITextureCreateD
 			}
 
 			// update rhi stats
-			VulkanTextureAllocated(Allocation.Size, GetViewType(), bRenderTarget);
+			VulkanTextureAllocated(GetDesc(), Allocation.Size);
 		}
 		Allocation.BindImage(Device, Image);
 
diff --git a/Engine/Source/Runtime/Windows/D3D11RHI/Private/D3D11Buffer.cpp b/Engine/Source/Runtime/Windows/D3D11RHI/Private/D3D11Buffer.cpp
--- a/Engine/Source/Runtime/Windows/D3D11RHI/Private/D3D11Buffer.cpp
+++ b/Engine/Source/Runtime/Windows/D3D11RHI/Private/D3D11Buffer.cpp
@@ -13,6 +13,25 @@ TAutoConsoleVariable<int32> GCVarUseSharedKeyedMutex(
 	TEXT("with the D3D11_RESOURCE_MISC_SHARED_KEYEDMUTEX flag instead of D3D11_RESOURCE_MISC_SHARED (default).\n"),
 	ECVF_Default);
 
+FD3D11Buffer::~FD3D11Buffer()
+{
+	D3D11BufferStats::UpdateBufferStats(*this, false);
+}
+
+void FD3D11Buffer::TakeOwnership(FD3D11Buffer& Other)
+{
+	FRHIBuffer::TakeOwnership(Other);
+	Resource = MoveTemp(Other.Resource);
+}
+
+void FD3D11Buffer::ReleaseOwnership()
+{
+	FRHIBuffer::ReleaseOwnership();
+
+	D3D11BufferStats::UpdateBufferStats(*this, false);
+	Resource = nullptr;
+}
+
 FBufferRHIRef FD3D11DynamicRHI::RHICreateBuffer(FRHICommandListBase& RHICmdList, FRHIBufferDesc const& BufferDesc, ERHIAccess ResourceState, FRHIResourceCreateInfo& CreateInfo)
 {
 	if (BufferDesc.IsNull())
@@ -24,7 +43,8 @@ FBufferRHIRef FD3D11DynamicRHI::RHICreateBuffer(FRHICommandListBase& RHICmdList,
 	checkf(BufferDesc.Size > 0, TEXT("Attempt to create buffer '%s' with size 0."), CreateInfo.DebugName ? CreateInfo.DebugName : TEXT("(null)"));
 
 	// Describe the buffer.
-	D3D11_BUFFER_DESC Desc = { BufferDesc.Size };
+	D3D11_BUFFER_DESC Desc{};
+	Desc.ByteWidth = BufferDesc.Size;
 
 	if (EnumHasAnyFlags(BufferDesc.Usage, BUF_AnyDynamic))
 	{
@@ -107,8 +127,6 @@ FBufferRHIRef FD3D11DynamicRHI::RHICreateBuffer(FRHICommandListBase& RHICmdList,
 		BufferResource->SetPrivateData(WKPDID_D3DDebugObjectName, FCString::Strlen(CreateInfo.DebugName) + 1, TCHAR_TO_ANSI(CreateInfo.DebugName));
 	}
 
-	UpdateBufferStats(BufferResource, true);
-
 	if (CreateInfo.DebugName)
 	{
 		BufferResource->SetPrivateData(WKPDID_D3DDebugObjectName, FCString::Strlen(CreateInfo.DebugName) + 1, TCHAR_TO_ANSI(CreateInfo.DebugName));
@@ -125,6 +143,9 @@ FBufferRHIRef FD3D11DynamicRHI::RHICreateBuffer(FRHICommandListBase& RHICmdList,
 	{
 		NewBuffer->SetName(CreateInfo.DebugName);
 	}
+
+	D3D11BufferStats::UpdateBufferStats(*NewBuffer, true);
+
 	return NewBuffer;
 }
 
diff --git a/Engine/Source/Runtime/Windows/D3D11RHI/Private/D3D11RHI.cpp b/Engine/Source/Runtime/Windows/D3D11RHI/Private/D3D11RHI.cpp
--- a/Engine/Source/Runtime/Windows/D3D11RHI/Private/D3D11RHI.cpp
+++ b/Engine/Source/Runtime/Windows/D3D11RHI/Private/D3D11RHI.cpp
@@ -9,6 +9,7 @@
 #include "StaticBoundShaderState.h"
 #include "Engine/GameViewportClient.h"
 #include "ProfilingDebugging/MemoryTrace.h"
+#include "RHICoreStats.h"
 
 #if WITH_DX_PERF
 	// For perf events
@@ -683,64 +684,40 @@ void FD3D11EventNodeFrame::LogDisjointQuery()
 	}
 }
 
-void UpdateBufferStats(TRefCountPtr<ID3D11Buffer> Buffer, bool bAllocating)
+static void D3D11UpdateBufferStatsCommon(ID3D11Buffer* Buffer, int64 BufferSize, bool bAllocating)
 {
-	D3D11_BUFFER_DESC Desc;
-	Buffer->GetDesc(&Desc);
+	// this is a work-around on Windows. Due to the fact that there is no way
+	// to hook the actual d3d allocations we can't track the memory in the normal way.
+	// Instead we simply tell LLM the size of these resources.
 
-	const bool bUniformBuffer = !!(Desc.BindFlags & D3D11_BIND_CONSTANT_BUFFER);
-	const bool bIndexBuffer = !!(Desc.BindFlags & D3D11_BIND_INDEX_BUFFER);
-	const bool bVertexBuffer = !!(Desc.BindFlags & D3D11_BIND_VERTEX_BUFFER);
+	LLM_SCOPED_PAUSE_TRACKING_WITH_ENUM_AND_AMOUNT(ELLMTag::GraphicsPlatform, bAllocating ? BufferSize : -BufferSize, ELLMTracker::Platform, ELLMAllocType::None);
 
+#if UE_MEMORY_TRACE_ENABLED
 	if (bAllocating)
 	{
-		if (bUniformBuffer)
-		{
-			INC_MEMORY_STAT_BY(STAT_UniformBufferMemory,Desc.ByteWidth);
-		}
-		else if (bIndexBuffer)
-		{
-			INC_MEMORY_STAT_BY(STAT_IndexBufferMemory,Desc.ByteWidth);
-		}
-		else if (bVertexBuffer)
-		{
-			INC_MEMORY_STAT_BY(STAT_VertexBufferMemory,Desc.ByteWidth);
-		}
-		else
-		{
-			INC_MEMORY_STAT_BY(STAT_StructuredBufferMemory,Desc.ByteWidth);
-		}
-
-		// this is a work-around on Windows. Due to the fact that there is no way
-		// to hook the actual d3d allocations we can't track the memory in the normal way.
-		// Instead we simply tell LLM the size of these resources.
-		LLM_SCOPED_PAUSE_TRACKING_WITH_ENUM_AND_AMOUNT(ELLMTag::GraphicsPlatform, Desc.ByteWidth, ELLMTracker::Platform, ELLMAllocType::None);
-		MemoryTrace_Alloc((uint64)Buffer.GetReference(), Desc.ByteWidth, 0, EMemoryTraceRootHeap::VideoMemory);
+		MemoryTrace_Alloc((uint64)Buffer, BufferSize, 0, EMemoryTraceRootHeap::VideoMemory);
 	}
 	else
-	{ //-V523
-		if (bUniformBuffer)
-		{
-			DEC_MEMORY_STAT_BY(STAT_UniformBufferMemory,Desc.ByteWidth);
-		}
-		else if (bIndexBuffer)
-		{
-			DEC_MEMORY_STAT_BY(STAT_IndexBufferMemory,Desc.ByteWidth);
-		}
-		else if (bVertexBuffer)
-		{
-			DEC_MEMORY_STAT_BY(STAT_VertexBufferMemory,Desc.ByteWidth);
-		}
-		else
-		{
-			DEC_MEMORY_STAT_BY(STAT_StructuredBufferMemory,Desc.ByteWidth);
-		}
+	{
+		MemoryTrace_Free((uint64)Buffer, EMemoryTraceRootHeap::VideoMemory);
+	}
+#endif
+}
+
+void D3D11BufferStats::UpdateUniformBufferStats(ID3D11Buffer* Buffer, int64 BufferSize, bool bAllocating)
+{
+	UE::RHICore::UpdateGlobalUniformBufferStats(BufferSize, bAllocating);
+	D3D11UpdateBufferStatsCommon(Buffer, BufferSize, bAllocating);
+}
+
+void D3D11BufferStats::UpdateBufferStats(FD3D11Buffer& Buffer, bool bAllocating)
+{
+	if (ID3D11Buffer* Resource = Buffer.Resource)
+	{
+		const FRHIBufferDesc& BufferDesc = Buffer.GetDesc();
 
-		// this is a work-around on Windows. Due to the fact that there is no way
-		// to hook the actual d3d allocations we can't track the memory in the normal way.
-		// Instead we simply tell LLM the size of these resources.
-		LLM_SCOPED_PAUSE_TRACKING_WITH_ENUM_AND_AMOUNT(ELLMTag::GraphicsPlatform, -(int64)Desc.ByteWidth, ELLMTracker::Platform, ELLMAllocType::None);
-		MemoryTrace_Free((uint64)Buffer.GetReference(), EMemoryTraceRootHeap::VideoMemory);
+		UE::RHICore::UpdateGlobalBufferStats(BufferDesc, BufferDesc.Size, bAllocating);
+		D3D11UpdateBufferStatsCommon(Resource, BufferDesc.Size, bAllocating);
 	}
 }
 
diff --git a/Engine/Source/Runtime/Windows/D3D11RHI/Private/D3D11Texture.cpp b/Engine/Source/Runtime/Windows/D3D11RHI/Private/D3D11Texture.cpp
--- a/Engine/Source/Runtime/Windows/D3D11RHI/Private/D3D11Texture.cpp
+++ b/Engine/Source/Runtime/Windows/D3D11RHI/Private/D3D11Texture.cpp
@@ -18,6 +18,7 @@
 
 #include "HAL/LowLevelMemTracker.h"
 #include "ProfilingDebugging/MemoryTrace.h"
+#include "RHICoreStats.h"
 
 int64 FD3D11GlobalStats::GDedicatedVideoMemory = 0;
 int64 FD3D11GlobalStats::GDedicatedSystemMemory = 0;
@@ -29,92 +30,31 @@ int64 FD3D11GlobalStats::GTotalGraphicsMemory = 0;
 	Texture allocator support.
 -----------------------------------------------------------------------------*/
 
-static bool ShouldCountAsTextureMemory(uint32 BindFlags)
-{
-	return (BindFlags & (D3D11_BIND_DEPTH_STENCIL | D3D11_BIND_RENDER_TARGET | D3D11_BIND_UNORDERED_ACCESS)) == 0;
-}
-
-// @param b3D true:3D, false:2D or cube map
-static TStatId GetD3D11StatEnum(uint32 BindFlags, bool bCubeMap, bool b3D)
-{
-#if STATS
-	if(ShouldCountAsTextureMemory(BindFlags))
-	{
-		// normal texture
-		if(bCubeMap)
-		{
-			return GET_STATID(STAT_TextureMemoryCube);
-		}
-		else if(b3D)
-		{
-			return GET_STATID(STAT_TextureMemory3D);
-		}
-		else
-		{
-			return GET_STATID(STAT_TextureMemory2D);
-		}
-	}
-	else
-	{
-		// render target
-		if(bCubeMap)
-		{
-			return GET_STATID(STAT_RenderTargetMemoryCube);
-		}
-		else if(b3D)
-		{
-			return GET_STATID(STAT_RenderTargetMemory3D);
-		}
-		else
-		{
-			return GET_STATID(STAT_RenderTargetMemory2D);
-		}
-	}
-#endif
-	return TStatId();
-}
-
 // Note: This function can be called from many different threads
 // @param TextureSize >0 to allocate, <0 to deallocate
 // @param b3D true:3D, false:2D or cube map
 void UpdateD3D11TextureStats(FD3D11Texture& Texture, bool bAllocating)
 {
-	uint32 BindFlags, MiscFlags;
-	bool b3D = Texture.IsTexture3D();
-	if (b3D)
+	const FRHITextureDesc& TextureDesc = Texture.GetDesc();
+
+	uint32 BindFlags;
+	if (TextureDesc.IsTexture3D())
 	{
 		D3D11_TEXTURE3D_DESC Desc;
 		Texture.GetD3D11Texture3D()->GetDesc(&Desc);
 		BindFlags = Desc.BindFlags;
-		MiscFlags = Desc.MiscFlags;
 	}
 	else
 	{
 		D3D11_TEXTURE2D_DESC Desc;
 		Texture.GetD3D11Texture2D()->GetDesc(&Desc);
 		BindFlags = Desc.BindFlags;
-		MiscFlags = Desc.MiscFlags;
 	}
 
-	int64 TextureSize = Texture.GetMemorySize();
-	if (!bAllocating)
-	{
-		TextureSize = -TextureSize;
-	}
+	const uint64 TextureSize = Texture.GetMemorySize();
 
-	int64 AlignedSize = (TextureSize > 0) ? Align(TextureSize, 1024) / 1024 : -(Align(-TextureSize, 1024) / 1024);
-	if(ShouldCountAsTextureMemory(BindFlags))
-	{
-		FPlatformAtomics::InterlockedAdd(&GCurrentTextureMemorySize, AlignedSize);
-	}
-	else
-	{
-		FPlatformAtomics::InterlockedAdd(&GCurrentRendertargetMemorySize, AlignedSize);
-	}
-
-	bool bCubeMap = (MiscFlags & D3D11_RESOURCE_MISC_TEXTURECUBE) != 0;
-
-	INC_MEMORY_STAT_BY_FName(GetD3D11StatEnum(BindFlags, bCubeMap, b3D).GetName(), TextureSize);
+	const bool bOnlyStreamableTexturesInTexturePool = false;
+	UE::RHICore::UpdateGlobalTextureStats(TextureDesc, TextureSize, bOnlyStreamableTexturesInTexturePool, bAllocating);
 
 	if (bAllocating)
 	{
@@ -159,15 +99,14 @@ FDynamicRHI::FRHICalcTextureSizeResult FD3D11DynamicRHI::RHICalcTexturePlatformS
  */
 void FD3D11DynamicRHI::RHIGetTextureMemoryStats(FTextureMemoryStats& OutStats)
 {
+	UE::RHICore::FillBaselineTextureMemoryStats(OutStats);
+
 	OutStats.DedicatedVideoMemory = FD3D11GlobalStats::GDedicatedVideoMemory;
     OutStats.DedicatedSystemMemory = FD3D11GlobalStats::GDedicatedSystemMemory;
     OutStats.SharedSystemMemory = FD3D11GlobalStats::GSharedSystemMemory;
 	OutStats.TotalGraphicsMemory = FD3D11GlobalStats::GTotalGraphicsMemory ? FD3D11GlobalStats::GTotalGraphicsMemory : -1;
 
-	OutStats.AllocatedMemorySize = int64(GCurrentTextureMemorySize) * 1024;
-	OutStats.LargestContiguousAllocation = OutStats.AllocatedMemorySize;
-	OutStats.TexturePoolSize = GTexturePoolSize;
-	OutStats.PendingMemoryAdjustment = 0;
+	OutStats.LargestContiguousAllocation = OutStats.StreamingMemorySize;
 }
 
 /**
diff --git a/Engine/Source/Runtime/Windows/D3D11RHI/Private/D3D11UniformBuffer.cpp b/Engine/Source/Runtime/Windows/D3D11RHI/Private/D3D11UniformBuffer.cpp
--- a/Engine/Source/Runtime/Windows/D3D11RHI/Private/D3D11UniformBuffer.cpp
+++ b/Engine/Source/Runtime/Windows/D3D11RHI/Private/D3D11UniformBuffer.cpp
@@ -77,7 +77,7 @@ void UniformBufferBeginFrame()
 				DEC_DWORD_STAT(STAT_D3D11NumFreeUniformBuffers);
 				DEC_MEMORY_STAT_BY(STAT_D3D11FreeUniformBufferMemory, PoolEntry.CreatedSize);
 				NumCleaned++;
-				UpdateBufferStats(PoolEntry.Buffer, false);
+				D3D11BufferStats::UpdateUniformBufferStats(PoolEntry.Buffer, PoolEntry.CreatedSize, false);
 				PoolEntry.Buffer.SafeRelease();
 				UniformBufferPool[BucketIndex].RemoveAtSwap(EntryIndex);
 			}
@@ -149,7 +149,7 @@ static TRefCountPtr<ID3D11Buffer> CreateAndUpdatePooledUniformBuffer(
 
 		VERIFYD3D11RESULT_EX(Device->CreateBuffer(&Desc, NULL, UniformBufferResource.GetInitReference()), Device);
 
-		UpdateBufferStats(UniformBufferResource, true);
+		D3D11BufferStats::UpdateUniformBufferStats(UniformBufferResource, Desc.ByteWidth, true);
 	}
 
 	check(IsValidRef(UniformBufferResource));
diff --git a/Engine/Source/Runtime/Windows/D3D11RHI/Private/Windows/WindowsD3D11ConstantBuffer.cpp b/Engine/Source/Runtime/Windows/D3D11RHI/Private/Windows/WindowsD3D11ConstantBuffer.cpp
--- a/Engine/Source/Runtime/Windows/D3D11RHI/Private/Windows/WindowsD3D11ConstantBuffer.cpp
+++ b/Engine/Source/Runtime/Windows/D3D11RHI/Private/Windows/WindowsD3D11ConstantBuffer.cpp
@@ -31,7 +31,7 @@ void FWinD3D11ConstantBuffer::InitRHI(FRHICommandListBase& RHICmdList)
 
 		CA_SUPPRESS(6385);	// Doesn't like COM
 		VERIFYD3D11RESULT_EX(D3DRHI->GetDevice()->CreateBuffer(&BufferDesc, nullptr, SubBuffer.Buffer.GetInitReference()), D3DRHI->GetDevice());
-		UpdateBufferStats(SubBuffer.Buffer, true);
+		D3D11BufferStats::UpdateUniformBufferStats(SubBuffer.Buffer, SubBuffer.Size, true);
 
 		BufferDesc.ByteWidth = Align(BufferDesc.ByteWidth / 2, 16);
 	}
@@ -43,9 +43,9 @@ void FWinD3D11ConstantBuffer::ReleaseRHI()
 {
 	FD3D11ConstantBuffer::ReleaseRHI();
 
-	for (FSubBuffer& SubBuffer : SubBuffers)
+	for (const FSubBuffer& SubBuffer : SubBuffers)
 	{
-		UpdateBufferStats(SubBuffer.Buffer, false);
+		D3D11BufferStats::UpdateUniformBufferStats(SubBuffer.Buffer, SubBuffer.Size, false);
 	}
 
 	SubBuffers.Empty();
diff --git a/Engine/Source/Runtime/Windows/D3D11RHI/Public/D3D11Resources.h b/Engine/Source/Runtime/Windows/D3D11RHI/Public/D3D11Resources.h
--- a/Engine/Source/Runtime/Windows/D3D11RHI/Public/D3D11Resources.h
+++ b/Engine/Source/Runtime/Windows/D3D11RHI/Public/D3D11Resources.h
@@ -334,9 +334,6 @@ public:
 
 };
 
-/** Updates tracked stats for a buffer. */
-extern void UpdateBufferStats(TRefCountPtr<ID3D11Buffer> Buffer, bool bAllocating);
-
 /** Forward declare the constants ring buffer. */
 class FD3D11ConstantsRingBuffer;
 
@@ -407,27 +404,10 @@ public:
 	}
 #endif
 
-	virtual ~FD3D11Buffer()
-	{
-		if (Resource)
-		{
-			UpdateBufferStats(Resource, false);
-		}
-	}
+	virtual ~FD3D11Buffer();
 
-	void TakeOwnership(FD3D11Buffer& Other)
-	{
-		FRHIBuffer::TakeOwnership(Other);
-		Resource = MoveTemp(Other.Resource);
-	}
-
-	void ReleaseOwnership()
-	{
-		FRHIBuffer::ReleaseOwnership();
-
-		UpdateBufferStats(Resource, false);
-		Resource = nullptr;
-	}
+	void TakeOwnership(FD3D11Buffer& Other);
+	void ReleaseOwnership();
 
 	// IRefCountedObject interface.
 	virtual uint32 AddRef() const
@@ -464,6 +444,12 @@ private:
 	uint32 ShadowBufferSize;
 };
 
+namespace D3D11BufferStats
+{
+	void UpdateUniformBufferStats(ID3D11Buffer* Buffer, int64 BufferSize, bool bAllocating);
+	void UpdateBufferStats(FD3D11Buffer& Buffer, bool bAllocating);
+}
+
 class FD3D11View : public TIntrusiveLinkedList<FD3D11View>
 {
 public:
