Fixing D3D12 resource allocation info tracing to go through the GetResourceAllocationInfo path that handles our texture corner cases.

diff --git a/Engine/Source/Runtime/D3D12RHI/Private/D3D12Resources.cpp b/Engine/Source/Runtime/D3D12RHI/Private/D3D12Resources.cpp
--- a/Engine/Source/Runtime/D3D12RHI/Private/D3D12Resources.cpp
+++ b/Engine/Source/Runtime/D3D12RHI/Private/D3D12Resources.cpp
@@ -32,7 +32,7 @@ void ID3D12ResourceAllocator::AllocateTexture(uint32 GPUIndex, D3D12_HEAP_TYPE I
 	Desc.Alignment = b4KAligment ? D3D12_SMALL_RESOURCE_PLACEMENT_ALIGNMENT : D3D12_DEFAULT_RESOURCE_PLACEMENT_ALIGNMENT;
 
 	// Get the size and alignment for the allocation
-	D3D12_RESOURCE_ALLOCATION_INFO Info = FD3D12DynamicRHI::GetD3DRHI()->GetAdapter().GetDevice(0)->GetResourceAllocationInfo(Desc);
+	D3D12_RESOURCE_ALLOCATION_INFO Info = FD3D12DynamicRHI::GetD3DRHI()->GetAdapter().GetDevice(GPUIndex)->GetResourceAllocationInfo(Desc);
 	AllocateResource(GPUIndex, InHeapType, Desc, Info.SizeInBytes, Info.Alignment, InResourceStateMode, InCreateState, InClearValue, InName, ResourceLocation);
 }
 
@@ -717,8 +717,7 @@ void FD3D12Adapter::TraceMemoryAllocation(FD3D12Resource* Resource)
 	// Calling GetResourceAllocationInfo is not cheap so check memory allocation tracking is enabled
 	if (UE_TRACE_CHANNELEXPR_IS_ENABLED(MemAllocChannel))
 	{
-		const D3D12_RESOURCE_DESC ResourceDesc = Resource->GetResource()->GetDesc();
-		const D3D12_RESOURCE_ALLOCATION_INFO Info = RootDevice->GetResourceAllocationInfo(0, 1, &ResourceDesc);
+		const D3D12_RESOURCE_ALLOCATION_INFO Info = Resource->GetParentDevice()->GetResourceAllocationInfo(Resource->GetDesc());
 		D3D12_GPU_VIRTUAL_ADDRESS GPUAddress = Resource->GetGPUVirtualAddress();
 		// Textures don't have valid GPUVirtualAddress when IsTrackingAllAllocations() is false, so don't do memory trace in this case.
 		if (IsTrackingAllAllocations() || GPUAddress != 0)
