Fix invalid format string when printing information on failed graphics PSOs.

#rnx

diff --git a/Engine/Source/Runtime/D3D12RHI/Private/Windows/WindowsD3D12PipelineState.cpp b/Engine/Source/Runtime/D3D12RHI/Private/Windows/WindowsD3D12PipelineState.cpp
--- a/Engine/Source/Runtime/D3D12RHI/Private/Windows/WindowsD3D12PipelineState.cpp
+++ b/Engine/Source/Runtime/D3D12RHI/Private/Windows/WindowsD3D12PipelineState.cpp
@@ -670,7 +670,7 @@ static void DumpGraphicsPSO(const FD3D12_GRAPHICS_PIPELINE_STATE_DESC& Desc, con
 		const D3D12_DEPTH_STENCIL_DESC1& DSState = Desc.DepthStencilState;
 		String.Appendf(TEXT("DepthEnable = %u\n"),               DSState.DepthEnable);
 		String.Appendf(TEXT("DepthWriteMask = %u\n"),            DSState.DepthWriteMask);
-		String.Appendf(TEXT("DepthFunc = u\n"),                  DSState.DepthFunc);
+		String.Appendf(TEXT("DepthFunc = %u\n"),                 DSState.DepthFunc);
 		String.Appendf(TEXT("StencilEnable = %u\n"),             DSState.StencilEnable);
 		String.Appendf(TEXT("StencilReadMask = 0x%X\n"),         DSState.StencilReadMask);
 		String.Appendf(TEXT("StencilWriteMask = 0x%X\n"),        DSState.StencilWriteMask);
