PR #10531: Fix CleanSlots precision overflow situations.

diff --git a/Engine/Source/Runtime/D3D12RHI/Private/D3D12StateCachePrivate.h b/Engine/Source/Runtime/D3D12RHI/Private/D3D12StateCachePrivate.h
--- a/Engine/Source/Runtime/D3D12RHI/Private/D3D12StateCachePrivate.h
+++ b/Engine/Source/Runtime/D3D12RHI/Private/D3D12StateCachePrivate.h
@@ -103,7 +103,7 @@ struct FD3D12ResourceCache
 
 	static inline void CleanSlots(ResourceSlotMask& SlotMask, uint32 NumSlots)
 	{
-		SlotMask &= ~(((ResourceSlotMask)1 << NumSlots) - 1);
+		SlotMask &= (NumSlots >= std::numeric_limits<ResourceSlotMask>::digits) ? 0 : ~(((ResourceSlotMask)1 << NumSlots) - 1);
 	}
 
 	static inline void DirtySlot(ResourceSlotMask& SlotMask, uint32 SlotIndex)
