Break apart FD3D12DescriptorCache SetSRVs, SetUAVs, and SetSamplers so you can allocate and copy GPU descriptors and retrieve the handles without actually setting descriptor tables (needed for an upcoming GPU feature).

diff --git a/Engine/Source/Runtime/D3D12RHI/Private/D3D12DescriptorCache.cpp b/Engine/Source/Runtime/D3D12RHI/Private/D3D12DescriptorCache.cpp
--- a/Engine/Source/Runtime/D3D12RHI/Private/D3D12DescriptorCache.cpp
+++ b/Engine/Source/Runtime/D3D12RHI/Private/D3D12DescriptorCache.cpp
@@ -218,7 +218,7 @@ void FD3D12DescriptorCache::SetVertexBuffers(FD3D12VertexBufferCache& Cache)
 	}
 }
 
-void FD3D12DescriptorCache::SetUAVs(EShaderFrequency ShaderStage, const FD3D12RootSignature* RootSignature, FD3D12UnorderedAccessViewCache& Cache, const UAVSlotMask& SlotsNeededMask, uint32 SlotsNeeded, uint32& HeapSlot)
+D3D12_GPU_DESCRIPTOR_HANDLE FD3D12DescriptorCache::BuildUAVTable(EShaderFrequency ShaderStage, const FD3D12RootSignature* RootSignature, FD3D12UnorderedAccessViewCache& Cache, const UAVSlotMask& SlotsNeededMask, uint32 SlotsNeeded, uint32& HeapSlot)
 {
 	UAVSlotMask& CurrentDirtySlotMask = Cache.DirtySlotMask[ShaderStage];
 	check(CurrentDirtySlotMask != 0);	// All dirty slots for the current shader stage.
@@ -262,16 +262,21 @@ void FD3D12DescriptorCache::SetUAVs(EShaderFrequency ShaderStage, const FD3D12Ro
 		SlotsNeeded, SrcDescriptors, nullptr /* sizes */,
 		D3D12_DESCRIPTOR_HEAP_TYPE_CBV_SRV_UAV);
 
+	return BindDescriptor;
+}
+
+void FD3D12DescriptorCache::SetUAVTable(EShaderFrequency ShaderStage, const FD3D12RootSignature* RootSignature, FD3D12UnorderedAccessViewCache& Cache, uint32 SlotsNeeded, const D3D12_GPU_DESCRIPTOR_HANDLE& BindDescriptor)
+{
+	check(ShaderStage == SF_Compute || ShaderStage == SF_Pixel);
+	const uint32 RootParameterIndex = RootSignature->UAVRDTBindSlot(ShaderStage);
+
 	if (ShaderStage == SF_Pixel)
 	{
-		const uint32 RDTIndex = RootSignature->UAVRDTBindSlot(ShaderStage);
-		Context.GraphicsCommandList()->SetGraphicsRootDescriptorTable(RDTIndex, BindDescriptor);
+		Context.GraphicsCommandList()->SetGraphicsRootDescriptorTable(RootParameterIndex, BindDescriptor);
 	}
 	else
 	{
-		check(ShaderStage == SF_Compute);
-		const uint32 RDTIndex = RootSignature->UAVRDTBindSlot(ShaderStage);
-		Context.GraphicsCommandList()->SetComputeRootDescriptorTable(RDTIndex, BindDescriptor);
+		Context.GraphicsCommandList()->SetComputeRootDescriptorTable(RootParameterIndex, BindDescriptor);
 	}
 
 	// We changed the descriptor table, so all resources bound to slots outside of the table's range are now dirty.
@@ -326,7 +331,7 @@ void FD3D12DescriptorCache::SetRenderTargets(FD3D12RenderTargetView** RenderTarg
 	}
 }
 
-void FD3D12DescriptorCache::SetSamplers(EShaderFrequency ShaderStage, const FD3D12RootSignature* RootSignature, FD3D12SamplerStateCache& Cache, const SamplerSlotMask& SlotsNeededMask, uint32 SlotsNeeded, uint32& HeapSlot)
+D3D12_GPU_DESCRIPTOR_HANDLE FD3D12DescriptorCache::BuildSamplerTable(EShaderFrequency ShaderStage, const FD3D12RootSignature* RootSignature, FD3D12SamplerStateCache& Cache, const SamplerSlotMask& SlotsNeededMask, uint32 SlotsNeeded, uint32& HeapSlot)
 {
 	check(!UsingGlobalSamplerHeap());
 
@@ -408,15 +413,20 @@ void FD3D12DescriptorCache::SetSamplers(EShaderFrequency ShaderStage, const FD3D
 		}
 	}
 
+	return BindDescriptor;
+}
+
+void FD3D12DescriptorCache::SetSamplerTable(EShaderFrequency ShaderStage, const FD3D12RootSignature* RootSignature, FD3D12SamplerStateCache& Cache, uint32 SlotsNeeded, const D3D12_GPU_DESCRIPTOR_HANDLE& BindDescriptor)
+{
+	const uint32 RootParameterIndex = RootSignature->SamplerRDTBindSlot(ShaderStage);
+
 	if (ShaderStage == SF_Compute)
 	{
-		const uint32 RDTIndex = RootSignature->SamplerRDTBindSlot(ShaderStage);
-		Context.GraphicsCommandList()->SetComputeRootDescriptorTable(RDTIndex, BindDescriptor);
+		Context.GraphicsCommandList()->SetComputeRootDescriptorTable(RootParameterIndex, BindDescriptor);
 	}
 	else
 	{
-		const uint32 RDTIndex = RootSignature->SamplerRDTBindSlot(ShaderStage);
-		Context.GraphicsCommandList()->SetGraphicsRootDescriptorTable(RDTIndex, BindDescriptor);
+		Context.GraphicsCommandList()->SetGraphicsRootDescriptorTable(RootParameterIndex, BindDescriptor);
 	}
 
 	// We changed the descriptor table, so all resources bound to slots outside of the table's range are now dirty.
@@ -448,7 +458,7 @@ static D3D12_RESOURCE_STATES GetBindingResourceState(FD3D12CommandContext& Conte
 	return State;
 }
 
-void FD3D12DescriptorCache::SetSRVs(EShaderFrequency ShaderStage, const FD3D12RootSignature* RootSignature, FD3D12ShaderResourceViewCache& Cache, const SRVSlotMask& SlotsNeededMask, uint32 SlotsNeeded, uint32& HeapSlot)
+D3D12_GPU_DESCRIPTOR_HANDLE FD3D12DescriptorCache::BuildSRVTable(EShaderFrequency ShaderStage, const FD3D12RootSignature* RootSignature, FD3D12ShaderResourceViewCache& Cache, const SRVSlotMask& SlotsNeededMask, uint32 SlotsNeeded, uint32& HeapSlot)
 {
 	SRVSlotMask& CurrentDirtySlotMask = Cache.DirtySlotMask[ShaderStage];
 	check(CurrentDirtySlotMask != 0);	// All dirty slots for the current shader stage.
@@ -488,20 +498,24 @@ void FD3D12DescriptorCache::SetSRVs(EShaderFrequency ShaderStage, const FD3D12Ro
 	ID3D12Device* Device = GetParentDevice()->GetDevice();
 	Device->CopyDescriptors(1, &DestDescriptor, &SlotsNeeded, SlotsNeeded, SrcDescriptors, nullptr, D3D12_DESCRIPTOR_HEAP_TYPE_CBV_SRV_UAV);
 
-
 	check((CurrentDirtySlotMask & SlotsNeededMask) == 0);	// Check all slots that needed to be set, were set.
 
 	const D3D12_GPU_DESCRIPTOR_HANDLE BindDescriptor = CurrentViewHeap->GetGPUSlotHandle(FirstSlotIndex);
 
+	return BindDescriptor;
+}
+
+void FD3D12DescriptorCache::SetSRVTable(EShaderFrequency ShaderStage, const FD3D12RootSignature* RootSignature, FD3D12ShaderResourceViewCache& Cache, uint32 SlotsNeeded, const D3D12_GPU_DESCRIPTOR_HANDLE& BindDescriptor)
+{
+	const uint32 RootParameterIndex = RootSignature->SRVRDTBindSlot(ShaderStage);
+
 	if (ShaderStage == SF_Compute)
 	{
-		const uint32 RDTIndex = RootSignature->SRVRDTBindSlot(ShaderStage);
-		Context.GraphicsCommandList()->SetComputeRootDescriptorTable(RDTIndex, BindDescriptor);
+		Context.GraphicsCommandList()->SetComputeRootDescriptorTable(RootParameterIndex, BindDescriptor);
 	}
 	else
 	{
-		const uint32 RDTIndex = RootSignature->SRVRDTBindSlot(ShaderStage);
-		Context.GraphicsCommandList()->SetGraphicsRootDescriptorTable(RDTIndex, BindDescriptor);
+		Context.GraphicsCommandList()->SetGraphicsRootDescriptorTable(RootParameterIndex, BindDescriptor);
 	}
 
 	// We changed the descriptor table, so all resources bound to slots outside of the table's range are now dirty.
diff --git a/Engine/Source/Runtime/D3D12RHI/Private/D3D12DescriptorCache.h b/Engine/Source/Runtime/D3D12RHI/Private/D3D12DescriptorCache.h
--- a/Engine/Source/Runtime/D3D12RHI/Private/D3D12DescriptorCache.h
+++ b/Engine/Source/Runtime/D3D12RHI/Private/D3D12DescriptorCache.h
@@ -344,9 +344,13 @@ public:
 	void SetVertexBuffers(FD3D12VertexBufferCache& Cache);
 	void SetRenderTargets(FD3D12RenderTargetView** RenderTargetViewArray, uint32 Count, FD3D12DepthStencilView* DepthStencilTarget);
 
-	void SetUAVs(EShaderFrequency ShaderStage, const FD3D12RootSignature* RootSignature, FD3D12UnorderedAccessViewCache& Cache, const UAVSlotMask& SlotsNeededMask, uint32 Count, uint32 &HeapSlot);
-	void SetSamplers(EShaderFrequency ShaderStage, const FD3D12RootSignature* RootSignature, FD3D12SamplerStateCache& Cache, const SamplerSlotMask& SlotsNeededMask, uint32 Count, uint32& HeapSlot);
-	void SetSRVs(EShaderFrequency ShaderStage, const FD3D12RootSignature* RootSignature, FD3D12ShaderResourceViewCache& Cache, const SRVSlotMask& SlotsNeededMask, uint32 Count, uint32& HeapSlot);
+	D3D12_GPU_DESCRIPTOR_HANDLE BuildUAVTable(EShaderFrequency ShaderStage, const FD3D12RootSignature* RootSignature, FD3D12UnorderedAccessViewCache& Cache, const UAVSlotMask& SlotsNeededMask, uint32 Count, uint32 &HeapSlot);
+	D3D12_GPU_DESCRIPTOR_HANDLE BuildSamplerTable(EShaderFrequency ShaderStage, const FD3D12RootSignature* RootSignature, FD3D12SamplerStateCache& Cache, const SamplerSlotMask& SlotsNeededMask, uint32 Count, uint32& HeapSlot);
+	D3D12_GPU_DESCRIPTOR_HANDLE BuildSRVTable(EShaderFrequency ShaderStage, const FD3D12RootSignature* RootSignature, FD3D12ShaderResourceViewCache& Cache, const SRVSlotMask& SlotsNeededMask, uint32 Count, uint32& HeapSlot);
+
+	void SetUAVTable(EShaderFrequency ShaderStage, const FD3D12RootSignature* RootSignature, FD3D12UnorderedAccessViewCache& Cache, uint32 SlotsNeeded, const D3D12_GPU_DESCRIPTOR_HANDLE& BindDescriptor);
+	void SetSamplerTable(EShaderFrequency ShaderStage, const FD3D12RootSignature* RootSignature, FD3D12SamplerStateCache& Cache, uint32 SlotsNeeded, const D3D12_GPU_DESCRIPTOR_HANDLE& BindDescriptor);
+	void SetSRVTable(EShaderFrequency ShaderStage, const FD3D12RootSignature* RootSignature, FD3D12ShaderResourceViewCache& Cache, uint32 SlotsNeeded, const D3D12_GPU_DESCRIPTOR_HANDLE& BindDescriptor);
 
 	void SetConstantBufferViews(EShaderFrequency ShaderStage, const FD3D12RootSignature* RootSignature, FD3D12ConstantBufferCache& Cache, CBVSlotMask SlotsNeededMask, uint32 Count, uint32& HeapSlot);
 	void SetRootConstantBuffers(EShaderFrequency ShaderStage, const FD3D12RootSignature* RootSignature, FD3D12ConstantBufferCache& Cache, CBVSlotMask SlotsNeededMask);
diff --git a/Engine/Source/Runtime/D3D12RHI/Private/D3D12StateCache.cpp b/Engine/Source/Runtime/D3D12RHI/Private/D3D12StateCache.cpp
--- a/Engine/Source/Runtime/D3D12RHI/Private/D3D12StateCache.cpp
+++ b/Engine/Source/Runtime/D3D12RHI/Private/D3D12StateCache.cpp
@@ -655,7 +655,8 @@ void FD3D12StateCache::ApplyResources(const FD3D12RootSignature* const pRootSign
 	if (CurrentShaderDirtyUAVSlots)
 	{
 		SCOPE_CYCLE_COUNTER(STAT_D3D12ApplyStateSetUAVTime);
-		DescriptorCache.SetUAVs(UAVStage, pRootSignature, PipelineState.Common.UAVCache, CurrentShaderDirtyUAVSlots, NumUAVs, ViewHeapSlot);
+		const D3D12_GPU_DESCRIPTOR_HANDLE BindDescriptor = DescriptorCache.BuildUAVTable(UAVStage, pRootSignature, PipelineState.Common.UAVCache, CurrentShaderDirtyUAVSlots, NumUAVs, ViewHeapSlot);
+		DescriptorCache.SetUAVTable(UAVStage, pRootSignature, PipelineState.Common.UAVCache, NumUAVs, BindDescriptor);
 	}
 
 	// Shader resource views
@@ -667,7 +668,8 @@ void FD3D12StateCache::ApplyResources(const FD3D12RootSignature* const pRootSign
 		{
 			if (CurrentShaderDirtySRVSlots[Index])
 			{
-				DescriptorCache.SetSRVs(static_cast<EShaderFrequency>(Index), pRootSignature, SRVCache, CurrentShaderDirtySRVSlots[Index], NumSRVs[Index], ViewHeapSlot);
+				const D3D12_GPU_DESCRIPTOR_HANDLE BindDescriptor = DescriptorCache.BuildSRVTable(static_cast<EShaderFrequency>(Index), pRootSignature, SRVCache, CurrentShaderDirtySRVSlots[Index], NumSRVs[Index], ViewHeapSlot);
+				DescriptorCache.SetSRVTable(static_cast<EShaderFrequency>(Index), pRootSignature, SRVCache, NumSRVs[Index], BindDescriptor);
 			}
 		}
 	}
@@ -877,7 +879,8 @@ void FD3D12StateCache::ApplySamplers(const FD3D12RootSignature* const pRootSigna
 	{
 		if (CurrentShaderDirtySamplerSlots[Index])
 		{
-			DescriptorCache.SetSamplers(static_cast<EShaderFrequency>(Index), pRootSignature, Cache, CurrentShaderDirtySamplerSlots[Index], NumSamplers[Index], SamplerHeapSlot);
+			D3D12_GPU_DESCRIPTOR_HANDLE BindDescriptor = DescriptorCache.BuildSamplerTable(static_cast<EShaderFrequency>(Index), pRootSignature, Cache, CurrentShaderDirtySamplerSlots[Index], NumSamplers[Index], SamplerHeapSlot);
+			DescriptorCache.SetSamplerTable(static_cast<EShaderFrequency>(Index), pRootSignature, Cache, NumSamplers[Index], BindDescriptor);
 		}
 	}
 
