Log the error code when D3D12GetDebugInterface fails.

#rnx

diff --git a/Engine/Source/Runtime/D3D12RHI/Private/D3D12Adapter.cpp b/Engine/Source/Runtime/D3D12RHI/Private/D3D12Adapter.cpp
--- a/Engine/Source/Runtime/D3D12RHI/Private/D3D12Adapter.cpp
+++ b/Engine/Source/Runtime/D3D12RHI/Private/D3D12Adapter.cpp
@@ -398,7 +398,8 @@ void FD3D12Adapter::CreateRootDevice(bool bWithDebug)
 	if (bWithDebug)
 	{
 		TRefCountPtr<ID3D12Debug> DebugController;
-		if (SUCCEEDED(D3D12GetDebugInterface(IID_PPV_ARGS(DebugController.GetInitReference()))))
+		HRESULT hr = D3D12GetDebugInterface(IID_PPV_ARGS(DebugController.GetInitReference()));
+		if (SUCCEEDED(hr))
 		{
 			DebugController->EnableDebugLayer();
 			bDebugDevice = true;
@@ -415,7 +416,7 @@ void FD3D12Adapter::CreateRootDevice(bool bWithDebug)
 		}
 		else
 		{
-			UE_LOG(LogD3D12RHI, Fatal, TEXT("The debug interface requires the D3D12 SDK Layers. Please install the Graphics Tools for Windows. See: https://docs.microsoft.com/en-us/windows/uwp/gaming/use-the-directx-runtime-and-visual-studio-graphics-diagnostic-features"));
+			UE_LOG(LogD3D12RHI, Fatal, TEXT("Failed to create D3D debug interface, error %x. The debug interface requires the D3D12 SDK Layers. Please install the Graphics Tools for Windows. See: https://docs.microsoft.com/en-us/windows/uwp/gaming/use-the-directx-runtime-and-visual-studio-graphics-diagnostic-features"), hr);
 		}
 	}
 
