Use unique ptr for residency handle to reduce memory size of FD3D12Resource because residency tracking is only needed it the resource is committed allocated (not a placed resource)

GitHub 10618 pull request

diff --git a/Engine/Source/Runtime/D3D12RHI/Private/D3D12Resources.cpp b/Engine/Source/Runtime/D3D12RHI/Private/D3D12Resources.cpp
--- a/Engine/Source/Runtime/D3D12RHI/Private/D3D12Resources.cpp
+++ b/Engine/Source/Runtime/D3D12RHI/Private/D3D12Resources.cpp
@@ -119,6 +119,11 @@ FD3D12Resource::FD3D12Resource(FD3D12Device* ParentDevice,
 	}
 #endif
 
+	if (!IsPlacedResource())
+	{
+		ResidencyHandle = MakeUnique<FD3D12ResidencyHandle>();
+	}
+
 	if (Desc.bReservedResource)
 	{
 		checkf(Heap == nullptr, TEXT("Reserved resources are not expected to have a heap"));
@@ -128,9 +133,9 @@ FD3D12Resource::FD3D12Resource(FD3D12Device* ParentDevice,
 
 FD3D12Resource::~FD3D12Resource()
 {
-	if (D3DX12Residency::IsInitialized(ResidencyHandle))
+	if (!IsPlacedResource() && D3DX12Residency::IsInitialized(*ResidencyHandle))
 	{
-		D3DX12Residency::EndTrackingObject(GetParentDevice()->GetResidencyManager(), ResidencyHandle);
+		D3DX12Residency::EndTrackingObject(GetParentDevice()->GetResidencyManager(), *ResidencyHandle);
 	}
 
 #if NV_AFTERMATH
@@ -285,13 +290,15 @@ void FD3D12Resource::StartTrackingForResidency()
 		return;
 	}
 
-	check(IsGPUOnly(HeapType));	// This is checked at a higher level before calling this function.
-	check(D3DX12Residency::IsInitialized(ResidencyHandle) == false);
-
-	const D3D12_RESOURCE_ALLOCATION_INFO Info = GetParentDevice()->GetResourceAllocationInfoUncached(Desc);
+	check(IsGPUOnly(HeapType));	// This is checked at a higher level before calling this function.		
+	if (!IsPlacedResource())
+	{
+		check(D3DX12Residency::IsInitialized(*ResidencyHandle) == false);
 
-	D3DX12Residency::Initialize(ResidencyHandle, Resource.GetReference(), Info.SizeInBytes, this);
-	D3DX12Residency::BeginTrackingObject(GetParentDevice()->GetResidencyManager(), ResidencyHandle);
+		const D3D12_RESOURCE_ALLOCATION_INFO Info = GetParentDevice()->GetResourceAllocationInfoUncached(Desc);
+		D3DX12Residency::Initialize(*ResidencyHandle, Resource.GetReference(), Info.SizeInBytes, this);
+		D3DX12Residency::BeginTrackingObject(GetParentDevice()->GetResidencyManager(), *ResidencyHandle);
+	}
 #endif
 }
 
diff --git a/Engine/Source/Runtime/D3D12RHI/Private/D3D12Resources.h b/Engine/Source/Runtime/D3D12RHI/Private/D3D12Resources.h
--- a/Engine/Source/Runtime/D3D12RHI/Private/D3D12Resources.h
+++ b/Engine/Source/Runtime/D3D12RHI/Private/D3D12Resources.h
@@ -149,7 +149,7 @@ private:
 	TRefCountPtr<ID3D12Resource> UAVAccessResource;
 	TRefCountPtr<FD3D12Heap> Heap;
 
-	FD3D12ResidencyHandle ResidencyHandle;
+	TUniquePtr<FD3D12ResidencyHandle> ResidencyHandle;
 
 	D3D12_GPU_VIRTUAL_ADDRESS GPUVirtualAddress{};
 	void* ResourceBaseAddress{};
@@ -309,7 +309,7 @@ public:
 
 	FD3D12ResidencyHandle& GetResidencyHandle()
 	{
-		return IsPlacedResource() ? Heap->GetResidencyHandle() : ResidencyHandle;
+		return IsPlacedResource() ? Heap->GetResidencyHandle() : *ResidencyHandle;
 	}
 
 	struct FD3D12ResourceTypeHelper
