Refactored D3D12 ray tracing descriptor heap/cache functionality into a more generic utility usable outside of ray tracing when explicitly managed descriptor tables are necessary (includes renaming FD3D12RayTracingDescriptorHeapCache -> FD3D12ExplicitDescriptorHeapCache, FD3D12RayTracingDescriptorHeap -> FD3D12ExplicitDescriptorHeap, FD3D12RayTracingDescriptorCache -> FD3D12ExplicitDescriptorCache)

diff --git a/Engine/Source/Runtime/D3D12RHI/Private/D3D12Adapter.cpp b/Engine/Source/Runtime/D3D12RHI/Private/D3D12Adapter.cpp
--- a/Engine/Source/Runtime/D3D12RHI/Private/D3D12Adapter.cpp
+++ b/Engine/Source/Runtime/D3D12RHI/Private/D3D12Adapter.cpp
@@ -1340,6 +1340,13 @@ void FD3D12Adapter::InitializeDevices()
 		CreateCommandSignatures();
 	}
 }
+void FD3D12Adapter::InitializeExplicitDescriptorHeap()
+{
+	for (uint32 GPUIndex : FRHIGPUMask::All())
+	{
+		Devices[GPUIndex]->InitExplicitDescriptorHeap();
+	}
+}
 
 void FD3D12Adapter::InitializeRayTracing()
 {
diff --git a/Engine/Source/Runtime/D3D12RHI/Private/D3D12Adapter.h b/Engine/Source/Runtime/D3D12RHI/Private/D3D12Adapter.h
--- a/Engine/Source/Runtime/D3D12RHI/Private/D3D12Adapter.h
+++ b/Engine/Source/Runtime/D3D12RHI/Private/D3D12Adapter.h
@@ -145,6 +145,7 @@ public:
 	void CleanupResources();
 
 	void InitializeDevices();
+	void InitializeExplicitDescriptorHeap();
 	void InitializeRayTracing();
 
 	// Getters
diff --git a/Engine/Source/Runtime/D3D12RHI/Private/D3D12Device.cpp b/Engine/Source/Runtime/D3D12RHI/Private/D3D12Device.cpp
--- a/Engine/Source/Runtime/D3D12RHI/Private/D3D12Device.cpp
+++ b/Engine/Source/Runtime/D3D12RHI/Private/D3D12Device.cpp
@@ -3,6 +3,7 @@
 #include "D3D12RHIPrivate.h"
 #include "D3D12IntelExtensions.h"
 #include "D3D12RayTracing.h"
+#include "D3D12ExplicitDescriptorCache.h"
 
 static TAutoConsoleVariable<int32> CVarD3D12GPUTimeout(
 	TEXT("r.D3D12.GPUTimeout"),
@@ -139,10 +140,10 @@ FD3D12Device::~FD3D12Device()
 #if D3D12_RHI_RAYTRACING
 	delete RayTracingCompactionRequestHandler;
 	RayTracingCompactionRequestHandler = nullptr;
-
-	DestroyRayTracingDescriptorCache(); // #dxr_todo UE-72158: unify RT descriptor cache with main FD3D12DescriptorCache
 #endif
 
+	DestroyExplicitDescriptorCache(); // #dxr_todo UE-72158: unify RT descriptor cache with main FD3D12DescriptorCache
+
 	// Cleanup the allocator near the end, as some resources may be returned to the allocator or references are shared by multiple GPUs
 	DefaultBufferAllocator.FreeDefaultBufferPools();
 
@@ -770,3 +771,18 @@ FGPUTimingCalibrationTimestamp FD3D12Device::GetCalibrationTimestamp(ED3D12Queue
 
 	return Result;
 }
+
+void FD3D12Device::InitExplicitDescriptorHeap()
+{
+	check(ExplicitDescriptorHeapCache == nullptr);
+	ExplicitDescriptorHeapCache = new FD3D12ExplicitDescriptorHeapCache(this);
+
+	// Note: ExplicitDescriptorHeapCache is destroyed in ~FD3D12Device, after all deferred deletion is processed
+}
+
+void FD3D12Device::DestroyExplicitDescriptorCache()
+{
+	delete ExplicitDescriptorHeapCache;
+	ExplicitDescriptorHeapCache = nullptr;
+}
+
diff --git a/Engine/Source/Runtime/D3D12RHI/Private/D3D12Device.h b/Engine/Source/Runtime/D3D12RHI/Private/D3D12Device.h
--- a/Engine/Source/Runtime/D3D12RHI/Private/D3D12Device.h
+++ b/Engine/Source/Runtime/D3D12RHI/Private/D3D12Device.h
@@ -13,7 +13,7 @@ D3D12Device.h: D3D12 Device Interfaces
 class FD3D12Device;
 class FD3D12DynamicRHI;
 class FD3D12Buffer;
-class FD3D12RayTracingDescriptorHeapCache;
+class FD3D12ExplicitDescriptorHeapCache;
 class FD3D12RayTracingPipelineCache;
 class FD3D12RayTracingCompactionRequestHandler;
 struct FD3D12RayTracingPipelineInfo;
@@ -258,6 +258,9 @@ public:
 	D3D12_RESOURCE_ALLOCATION_INFO GetResourceAllocationInfo(const FD3D12ResourceDesc& InDesc);
 	TUniquePtr<FD3D12DiagnosticBuffer> CreateDiagnosticBuffer(const D3D12_RESOURCE_DESC& Desc, const TCHAR* Name);
 
+	void									  InitExplicitDescriptorHeap();
+	FD3D12ExplicitDescriptorHeapCache*		  GetExplicitDescriptorHeapCache() { return ExplicitDescriptorHeapCache; }
+
 	// Ray Tracing
 #if D3D12_RHI_RAYTRACING
 	void									  InitRayTracing();
@@ -266,7 +269,6 @@ public:
 	ID3D12Device5*							  GetDevice5();
 	ID3D12Device7*							  GetDevice7();
 
-	FD3D12RayTracingDescriptorHeapCache*	  GetRayTracingDescriptorHeapCache     () { return RayTracingDescriptorHeapCache;      }
 	FD3D12RayTracingPipelineCache*			  GetRayTracingPipelineCache           () { return RayTracingPipelineCache;            }
 	FD3D12Buffer*							  GetRayTracingDispatchRaysDescBuffer  () { return RayTracingDispatchRaysDescBuffer;   }
 	FD3D12RayTracingCompactionRequestHandler* GetRayTracingCompactionRequestHandler() { return RayTracingCompactionRequestHandler; }
@@ -394,7 +396,8 @@ private:
 	FD3D12RayTracingCompactionRequestHandler* RayTracingCompactionRequestHandler = nullptr;
 	FD3D12Buffer* RayTracingDispatchRaysDescBuffer = nullptr;
 // #dxr_todo UE-72158: unify RT descriptor cache with main FD3D12DescriptorCache
-	FD3D12RayTracingDescriptorHeapCache* RayTracingDescriptorHeapCache = nullptr;
-	void DestroyRayTracingDescriptorCache();
 #endif
+
+	FD3D12ExplicitDescriptorHeapCache* ExplicitDescriptorHeapCache = nullptr;
+	void DestroyExplicitDescriptorCache();
 };
diff --git a/Engine/Source/Runtime/D3D12RHI/Private/D3D12ExplicitDescriptorCache.cpp b/Engine/Source/Runtime/D3D12RHI/Private/D3D12ExplicitDescriptorCache.cpp
new file mode 100644
--- /dev/null
+++ b/Engine/Source/Runtime/D3D12RHI/Private/D3D12ExplicitDescriptorCache.cpp
@@ -0,0 +1,418 @@
+// Copyright Epic Games, Inc. All Rights Reserved.
+
+#include "D3D12ExplicitDescriptorCache.h"
+#include "Hash/xxhash.h"
+
+// Whether to compare the full descriptor table on cache lookup or only use FXxHash64 digest.
+#ifndef EXPLICIT_DESCRIPTOR_CACHE_FULL_COMPARE
+#define EXPLICIT_DESCRIPTOR_CACHE_FULL_COMPARE 1
+#endif
+
+static int32 GD3D12ExplicitDeduplicateSamplers = 1;
+static FAutoConsoleVariableRef CVarD3D12ExplicitDeduplicateSamplers(
+	TEXT("r.D3D12.ExplicitDescriptorHeap.DeduplicateSamplers"),
+	GD3D12ExplicitDeduplicateSamplers,
+	TEXT("Use an exhaustive search to deduplicate sampler descriptors when generating shader binding tables. Reduces sampler heap usage at the cost of some CPU time. (default = 1)")
+);
+
+int32 GD3D12ExplicitViewDescriptorHeapSize = 250'000;
+int32 GD3D12ExplicitViewDescriptorHeapOverflowReported = 0;
+static FAutoConsoleVariableRef CVarD3D12ExplicitViewDescriptorHeapSize(
+	TEXT("r.D3D12.ExplicitDescriptorHeap.ViewDescriptorHeapSize"),
+	GD3D12ExplicitViewDescriptorHeapSize,
+	TEXT("Maximum number of descriptors per explicit view descriptor heap. (default = 250k, ~8MB per heap)\n")
+	TEXT("Typical measured descriptor heap usage in large scenes is ~50k. An error is reported when this limit is reached and shader bindings for subsequent objects are skipped.\n"),
+	ECVF_ReadOnly
+);
+
+FD3D12ExplicitDescriptorHeapCache::~FD3D12ExplicitDescriptorHeapCache()
+{
+	check(AllocatedEntries == 0);
+
+	FScopeLock Lock(&CriticalSection);
+	for (const Entry& It : Entries)
+	{
+		if (It.Type == D3D12_DESCRIPTOR_HEAP_TYPE_CBV_SRV_UAV)
+		{
+			DEC_DWORD_STAT(STAT_ExplicitViewDescriptorHeaps);
+			DEC_DWORD_STAT_BY(STAT_ExplicitViewDescriptors, It.NumDescriptors);
+		}
+		else
+		{
+			DEC_DWORD_STAT(STAT_ExplicitSamplerDescriptorHeaps);
+			DEC_DWORD_STAT_BY(STAT_ExplicitSamplerDescriptors, It.NumDescriptors);
+		}
+
+		It.Heap->Release();
+	}
+	Entries.Empty();
+}
+
+void FD3D12ExplicitDescriptorHeapCache::ReleaseHeap(FD3D12ExplicitDescriptorHeapCache::Entry& Entry)
+{
+	FScopeLock Lock(&CriticalSection);
+
+	Entries.Add(Entry);
+
+	check(AllocatedEntries != 0);
+	--AllocatedEntries;
+}
+
+FD3D12ExplicitDescriptorHeapCache::Entry FD3D12ExplicitDescriptorHeapCache::AllocateHeap(D3D12_DESCRIPTOR_HEAP_TYPE Type, uint32 NumDescriptors)
+{
+	FScopeLock Lock(&CriticalSection);
+
+	++AllocatedEntries;
+
+	Entry Result = {};
+
+	FD3D12ManualFence& Fence = GetParentDevice()->GetParentAdapter()->GetFrameFence();
+	const uint64 CompletedFenceValue = Fence.GetCompletedFenceValue(/* bUpdateCachedFenceValue */ false);
+
+	for (int32 EntryIndex = 0; EntryIndex < Entries.Num(); ++EntryIndex)
+	{
+		const Entry& It = Entries[EntryIndex];
+		if (It.Type == Type && It.NumDescriptors >= NumDescriptors && It.FenceValue <= CompletedFenceValue)
+		{
+			Result = It;
+
+			Entries[EntryIndex] = Entries.Last();
+			Entries.Pop();
+
+			return Result;
+		}
+	}
+
+	// Compatible heap was not found in cache, so create a new one.
+
+	ReleaseStaleEntries(100, CompletedFenceValue); // Release heaps that were not used for 100 frames before allocating new.
+
+	D3D12_DESCRIPTOR_HEAP_DESC Desc = {};
+
+	Desc.Flags = D3D12_DESCRIPTOR_HEAP_FLAG_SHADER_VISIBLE;
+	Desc.Type = Type;
+	Desc.NumDescriptors = NumDescriptors;
+	Desc.NodeMask = GetParentDevice()->GetGPUMask().GetNative();
+
+	ID3D12DescriptorHeap* D3D12Heap = nullptr;
+
+	const TCHAR* HeapName = Desc.Type == D3D12_DESCRIPTOR_HEAP_TYPE_CBV_SRV_UAV ? TEXT("RT View Heap") : TEXT("RT Sampler Heap");
+	UE_LOG(LogD3D12RHI, Log, TEXT("Creating %s with %d entries"), HeapName, NumDescriptors);
+
+	VERIFYD3D12RESULT(GetParentDevice()->GetDevice()->CreateDescriptorHeap(&Desc, IID_PPV_ARGS(&D3D12Heap)));
+	SetName(D3D12Heap, Desc.Type == D3D12_DESCRIPTOR_HEAP_TYPE_CBV_SRV_UAV ? L"RT View Heap" : L"RT Sampler Heap");
+
+	Result.NumDescriptors = NumDescriptors;
+	Result.Type = Type;
+	Result.Heap = D3D12Heap;
+
+	if (Desc.Type == D3D12_DESCRIPTOR_HEAP_TYPE_CBV_SRV_UAV)
+	{
+		INC_DWORD_STAT(STAT_ExplicitViewDescriptorHeaps);
+		INC_DWORD_STAT_BY(STAT_ExplicitViewDescriptors, NumDescriptors);
+	}
+	else
+	{
+		INC_DWORD_STAT(STAT_ExplicitSamplerDescriptorHeaps);
+		INC_DWORD_STAT_BY(STAT_ExplicitSamplerDescriptors, NumDescriptors);
+	}
+
+	return Result;
+}
+
+void FD3D12ExplicitDescriptorHeapCache::ReleaseStaleEntries(uint32 MaxAge, uint64 CompletedFenceValue)
+{
+	int32 EntryIndex = 0;
+	while (EntryIndex < Entries.Num())
+	{
+		Entry& It = Entries[EntryIndex];
+		if ((It.FenceValue + MaxAge) <= CompletedFenceValue)
+		{
+			if (It.Type == D3D12_DESCRIPTOR_HEAP_TYPE_CBV_SRV_UAV)
+			{
+				DEC_DWORD_STAT(STAT_ExplicitViewDescriptorHeaps);
+				DEC_DWORD_STAT_BY(STAT_ExplicitViewDescriptors, It.NumDescriptors);
+			}
+			else
+			{
+				DEC_DWORD_STAT(STAT_ExplicitSamplerDescriptorHeaps);
+				DEC_DWORD_STAT_BY(STAT_ExplicitSamplerDescriptors, It.NumDescriptors);
+			}
+
+			It.Heap->Release();
+
+			Entries[EntryIndex] = Entries.Last();
+			Entries.Pop(false);
+		}
+		else
+		{
+			EntryIndex++;
+		}
+	}
+}
+
+void FD3D12ExplicitDescriptorHeapCache::Flush()
+{
+	FScopeLock Lock(&CriticalSection);
+
+	for (const Entry& It : Entries)
+	{
+		FD3D12DynamicRHI::GetD3DRHI()->DeferredDelete(It.Heap);
+	}
+	Entries.Empty();
+}
+
+///
+
+FD3D12ExplicitDescriptorHeap::~FD3D12ExplicitDescriptorHeap()
+{
+	if (D3D12Heap)
+	{
+		GetParentDevice()->GetExplicitDescriptorHeapCache()->ReleaseHeap(HeapCacheEntry);
+	}
+}
+
+void FD3D12ExplicitDescriptorHeap::Init(uint32 InMaxNumDescriptors, D3D12_DESCRIPTOR_HEAP_TYPE InType)
+{
+	check(D3D12Heap == nullptr);
+
+	Type = InType;
+	HeapCacheEntry = GetParentDevice()->GetExplicitDescriptorHeapCache()->AllocateHeap(Type, InMaxNumDescriptors);
+
+	MaxNumDescriptors = HeapCacheEntry.NumDescriptors;
+	D3D12Heap = HeapCacheEntry.Heap;
+
+	CPUBase = D3D12Heap->GetCPUDescriptorHandleForHeapStart();
+	GPUBase = D3D12Heap->GetGPUDescriptorHandleForHeapStart();
+
+	checkf(CPUBase.ptr, TEXT("Ray tracing descriptor heap of type %d returned from descriptor heap cache is invalid."), Type);
+
+	DescriptorSize = GetParentDevice()->GetDevice()->GetDescriptorHandleIncrementSize(Type);
+
+	bExhaustiveSamplerDeduplication = GD3D12ExplicitDeduplicateSamplers == 1;
+
+	if (bExhaustiveSamplerDeduplication && InType == D3D12_DESCRIPTOR_HEAP_TYPE_SAMPLER)
+	{
+		// When exhaustive descriptor deduplication is active, all shadow descriptor table entries must be initialized.
+		// Deduplication works by looping over all elements, but they may be written out of order by worker threads.
+		// Initializing descriptors to 0 avoids accidentally matching wrong descriptor entries.
+		Descriptors.SetNumZeroed(MaxNumDescriptors);
+	}
+#if EXPLICIT_DESCRIPTOR_CACHE_FULL_COMPARE
+	else
+	{
+		// If exhaustive sampler deduplication is not used, then descriptor table copy is only used for validation.
+		// Entries are guaranteed to be written before they're read, so we don't have to spend the time on initialization.
+		Descriptors.SetNumUninitialized(MaxNumDescriptors);
+	}
+#endif
+}
+
+// Returns descriptor heap base index or -1 if allocation is not possible.
+// Thread-safe (uses atomic linear allocation).
+int32 FD3D12ExplicitDescriptorHeap::Allocate(uint32 InNumDescriptors)
+{
+	int32 Result = FPlatformAtomics::InterlockedAdd(&NumAllocatedDescriptors, InNumDescriptors);
+
+	if (Result + InNumDescriptors > MaxNumDescriptors)
+	{
+		if (Type == D3D12_DESCRIPTOR_HEAP_TYPE_SAMPLER)
+		{
+			UE_LOG(LogD3D12RHI, Fatal,
+			       TEXT("Explicit sampler descriptor heap overflow. ")
+			       TEXT("It is not possible to recover from this error, as maximum D3D12 sampler heap size is 2048."));
+		}
+		else if (GD3D12ExplicitViewDescriptorHeapSize == MaxNumDescriptors
+			&& FPlatformAtomics::InterlockedOr(&GD3D12ExplicitViewDescriptorHeapOverflowReported, 1) == 0)
+		{
+			// NOTE: GD3D12RayTracingViewDescriptorHeapOverflowReported is set atomically because multiple 
+			// allocations may be happening simultaneously, but we only want to report the error once.
+
+			UE_LOG(LogD3D12RHI, Error,
+			       TEXT("Ray tracing view descriptor heap overflow. Current frame will not be rendered correctly. ")
+			       TEXT("Increase r.D3D12.RayTracingViewDescriptorHeapSize to at least %d to fix this issue."),
+			       MaxNumDescriptors * 2);
+		}
+
+		Result = -1;
+	}
+
+	return Result;
+}
+
+void FD3D12ExplicitDescriptorHeap::CopyDescriptors(int32 BaseIndex, const D3D12_CPU_DESCRIPTOR_HANDLE* InDescriptors, uint32 InNumDescriptors)
+{
+	D3D12_CPU_DESCRIPTOR_HANDLE DestDescriptor = GetDescriptorCPU(BaseIndex);
+	GetParentDevice()->GetDevice()->CopyDescriptors(1, &DestDescriptor, &InNumDescriptors, InNumDescriptors, InDescriptors, nullptr, Type);
+	for (uint32 i = 0; i < InNumDescriptors; ++i)
+	{
+		Descriptors[BaseIndex + i].ptr = InDescriptors[i].ptr;
+	}
+}
+
+bool FD3D12ExplicitDescriptorHeap::CompareDescriptors(int32 BaseIndex, const D3D12_CPU_DESCRIPTOR_HANDLE* InDescriptors, uint32 InNumDescriptors)
+{
+	for (uint32 i = 0; i < InNumDescriptors; ++i)
+	{
+		if (Descriptors[BaseIndex + i].ptr != InDescriptors[i].ptr)
+		{
+			return false;
+		}
+	}
+	return true;
+}
+
+D3D12_CPU_DESCRIPTOR_HANDLE FD3D12ExplicitDescriptorHeap::GetDescriptorCPU(uint32 Index) const
+{
+	checkSlow(Index < MaxNumDescriptors);
+	D3D12_CPU_DESCRIPTOR_HANDLE Result = { CPUBase.ptr + Index * DescriptorSize };
+	return Result;
+}
+
+D3D12_GPU_DESCRIPTOR_HANDLE FD3D12ExplicitDescriptorHeap::GetDescriptorGPU(uint32 Index) const
+{
+	checkSlow(Index < MaxNumDescriptors);
+	D3D12_GPU_DESCRIPTOR_HANDLE Result = { GPUBase.ptr + Index * DescriptorSize };
+	return Result;
+}
+
+void FD3D12ExplicitDescriptorHeap::UpdateSyncPoint()
+{
+	const FD3D12ManualFence& Fence = GetParentDevice()->GetParentAdapter()->GetFrameFence();
+	HeapCacheEntry.FenceValue = FMath::Max(HeapCacheEntry.FenceValue, Fence.GetNextFenceToSignal());
+}
+
+///
+
+void FD3D12ExplicitDescriptorCache::Init(uint32 NumViewDescriptors, uint32 NumSamplerDescriptors)
+{
+	FD3D12BindlessDescriptorManager& BindlessManager = GetParentDevice()->GetBindlessDescriptorManager();
+
+	bBindlessViews = BindlessManager.HasHeap(ERHIDescriptorHeapType::Standard, ERHIBindlessConfiguration::RayTracingShaders);
+	bBindlessSamplers = BindlessManager.HasHeap(ERHIDescriptorHeapType::Sampler, ERHIBindlessConfiguration::RayTracingShaders);
+
+	if (!bBindlessViews)
+	{
+		ViewHeap.Init(NumViewDescriptors, D3D12_DESCRIPTOR_HEAP_TYPE_CBV_SRV_UAV);
+	}
+		
+	if (!bBindlessSamplers)
+	{
+		SamplerHeap.Init(NumSamplerDescriptors, D3D12_DESCRIPTOR_HEAP_TYPE_SAMPLER);
+	}
+}
+
+void FD3D12ExplicitDescriptorCache::UpdateSyncPoint()
+{
+	check(IsInRHIThread() || !IsRunningRHIInSeparateThread());
+
+	if (!bBindlessViews)
+	{
+		ViewHeap.UpdateSyncPoint();
+	}
+
+	if (!bBindlessSamplers)
+	{
+		SamplerHeap.UpdateSyncPoint();
+	}
+}
+
+void FD3D12ExplicitDescriptorCache::SetDescriptorHeaps(FD3D12CommandContext& CommandContext)
+{
+	UpdateSyncPoint();
+
+	check(bBindlessViews || ViewHeap.GetParentDevice() == CommandContext.GetParentDevice());
+	check(bBindlessSamplers || SamplerHeap.GetParentDevice() == CommandContext.GetParentDevice());
+
+	FD3D12BindlessDescriptorManager& BindlessManager = GetParentDevice()->GetBindlessDescriptorManager();
+
+	ID3D12DescriptorHeap* ViewHeapToSet = bBindlessViews ? BindlessManager.GetHeap(ERHIDescriptorHeapType::Standard)->GetHeap() : ViewHeap.D3D12Heap;
+	ID3D12DescriptorHeap* SamplerHeapToSet = bBindlessSamplers ? BindlessManager.GetHeap(ERHIDescriptorHeapType::Sampler)->GetHeap() : SamplerHeap.D3D12Heap;
+
+	CommandContext.StateCache.GetDescriptorCache()->OverrideLastSetHeaps(ViewHeapToSet, SamplerHeapToSet);
+}
+
+// Returns descriptor heap base index for this descriptor table allocation or -1 if allocation failed.
+int32 FD3D12ExplicitDescriptorCache::AllocateDescriptorTable(const uint32* DescriptorVersions, const D3D12_CPU_DESCRIPTOR_HANDLE* Descriptors, uint32 NumDescriptors, D3D12_DESCRIPTOR_HEAP_TYPE Type, uint32 WorkerIndex)
+{
+	checkSlow(Type == D3D12_DESCRIPTOR_HEAP_TYPE_CBV_SRV_UAV || Type == D3D12_DESCRIPTOR_HEAP_TYPE_SAMPLER);
+
+	FD3D12ExplicitDescriptorHeap& Heap = (Type == D3D12_DESCRIPTOR_HEAP_TYPE_CBV_SRV_UAV) ? ViewHeap : SamplerHeap;
+	TDescriptorHashMap& Map = (Type == D3D12_DESCRIPTOR_HEAP_TYPE_CBV_SRV_UAV) 
+		? WorkerData[WorkerIndex].ViewDescriptorTableCache
+		: WorkerData[WorkerIndex].SamplerDescriptorTableCache;
+
+	const uint64 VersionHash = FXxHash64::HashBuffer(DescriptorVersions, sizeof(DescriptorVersions[0]) * NumDescriptors).Hash;
+	const uint64 DescriptorHash = FXxHash64::HashBuffer(Descriptors, sizeof(Descriptors[0]) * NumDescriptors).Hash;
+	const uint64 Key = VersionHash ^ DescriptorHash;
+
+	int32& DescriptorTableBaseIndex = Map.FindOrAdd(Key, INDEX_NONE);
+
+	if (DescriptorTableBaseIndex != INDEX_NONE)
+	{
+	#if EXPLICIT_DESCRIPTOR_CACHE_FULL_COMPARE
+		if (ensureMsgf(Heap.CompareDescriptors(DescriptorTableBaseIndex, Descriptors, NumDescriptors), 
+		               TEXT("Ray tracing descriptor cache hash collision detected!")))
+	#endif
+		{
+			return DescriptorTableBaseIndex;
+		}
+	}
+
+	if (Heap.bExhaustiveSamplerDeduplication && Type == D3D12_DESCRIPTOR_HEAP_TYPE_SAMPLER && WorkerData.Num() > 1)
+	{
+		// Exhaustive search for a sampler table:
+		// We have to do this because sampler heap space is precious (hard limit of 2048 total entries).
+		// Per-thread descriptor table deduplication hash tables introduce a lot of redundancy in the heap
+		// which is reduced by looking for global matches on hash table lookup miss.
+
+		int32 FoundIndex = INDEX_NONE;
+		int32 SearchEndPos = Heap.NumWrittenSamplerDescriptors;
+		for (int32 SearchIndex = 0; SearchIndex + int32(NumDescriptors) < SearchEndPos; ++SearchIndex)
+		{
+			if (Heap.CompareDescriptors(SearchIndex, Descriptors, NumDescriptors))
+			{
+				FoundIndex = SearchIndex;
+				break;
+			}
+		}
+
+		if (FoundIndex != INDEX_NONE)
+		{
+			DescriptorTableBaseIndex = FoundIndex;
+			return DescriptorTableBaseIndex;
+		}
+	}
+
+	DescriptorTableBaseIndex = Heap.Allocate(NumDescriptors);
+
+	if (DescriptorTableBaseIndex == INDEX_NONE)
+	{
+		return INDEX_NONE;
+	}
+
+	Heap.CopyDescriptors(DescriptorTableBaseIndex, Descriptors, NumDescriptors);
+
+	if (Heap.bExhaustiveSamplerDeduplication && Type == D3D12_DESCRIPTOR_HEAP_TYPE_SAMPLER)
+	{
+		FPlatformAtomics::InterlockedAdd(&Heap.NumWrittenSamplerDescriptors, NumDescriptors);
+	}
+
+	if (Type == D3D12_DESCRIPTOR_HEAP_TYPE_CBV_SRV_UAV)
+	{
+		INC_DWORD_STAT_BY(STAT_ExplicitUsedViewDescriptors, NumDescriptors);
+	}
+	else
+	{
+		INC_DWORD_STAT_BY(STAT_ExplicitUsedSamplerDescriptors, NumDescriptors);
+
+		static uint64 GMaxNumUsedSamplerDescriptors = 0;
+		if (Heap.NumAllocatedDescriptors > GMaxNumUsedSamplerDescriptors)
+		{
+			GMaxNumUsedSamplerDescriptors = Heap.NumAllocatedDescriptors;
+		}
+		SET_DWORD_STAT(STAT_ExplicitMaxUsedSamplerDescriptors, GMaxNumUsedSamplerDescriptors);
+	}
+
+	return DescriptorTableBaseIndex;
+}
\ No newline at end of file
diff --git a/Engine/Source/Runtime/D3D12RHI/Private/D3D12ExplicitDescriptorCache.h b/Engine/Source/Runtime/D3D12RHI/Private/D3D12ExplicitDescriptorCache.h
new file mode 100644
--- /dev/null
+++ b/Engine/Source/Runtime/D3D12RHI/Private/D3D12ExplicitDescriptorCache.h
@@ -0,0 +1,156 @@
+// Copyright Epic Games, Inc. All Rights Reserved.
+
+#pragma once
+
+#include "D3D12RHIPrivate.h"
+#include "Experimental/Containers/SherwoodHashTable.h"
+
+class FD3D12DynamicRHI;
+struct FD3D12DefaultViews;
+class FD3D12DescriptorCache;
+struct FD3D12VertexBufferCache;
+struct FD3D12IndexBufferCache;
+struct FD3D12ConstantBufferCache;
+struct FD3D12ShaderResourceViewCache;
+struct FD3D12UnorderedAccessViewCache;
+struct FD3D12SamplerStateCache;
+
+// #dxr_todo UE-72158: FD3D12Device::GlobalViewHeap/GlobalSamplerHeap should be used instead of ad-hoc heaps here.
+// Unfortunately, this requires a major refactor of how global heaps work.
+// FD3D12CommandContext-s should not get static chunks of the global heap, but instead should dynamically allocate
+// chunks on as-needed basis and release them when possible.
+// This would allow calling code to sub-allocate heap blocks from the same global heap.
+class FD3D12ExplicitDescriptorHeapCache : FD3D12DeviceChild
+{
+public:
+
+	UE_NONCOPYABLE(FD3D12ExplicitDescriptorHeapCache)
+
+	struct Entry
+	{
+		ID3D12DescriptorHeap* Heap = nullptr;
+		uint64 FenceValue = 0;
+		uint32 NumDescriptors = 0;
+		D3D12_DESCRIPTOR_HEAP_TYPE Type = D3D12_DESCRIPTOR_HEAP_TYPE_NUM_TYPES;
+	};
+
+	FD3D12ExplicitDescriptorHeapCache(FD3D12Device* Device)
+	: FD3D12DeviceChild(Device)
+	{
+	}
+
+	~FD3D12ExplicitDescriptorHeapCache();
+
+	void ReleaseHeap(Entry& Entry);
+
+	Entry AllocateHeap(D3D12_DESCRIPTOR_HEAP_TYPE Type, uint32 NumDescriptors);
+
+	void ReleaseStaleEntries(uint32 MaxAge, uint64 CompletedFenceValue);
+
+	void Flush();
+
+	FCriticalSection CriticalSection;
+	TArray<Entry> Entries;
+	uint32 AllocatedEntries = 0;
+};
+
+struct FD3D12ExplicitDescriptorHeap : public FD3D12DeviceChild
+{
+	UE_NONCOPYABLE(FD3D12ExplicitDescriptorHeap)
+
+	FD3D12ExplicitDescriptorHeap(FD3D12Device* Device)
+	: FD3D12DeviceChild(Device)
+	{
+	}
+
+	~FD3D12ExplicitDescriptorHeap();
+
+	void Init(uint32 InMaxNumDescriptors, D3D12_DESCRIPTOR_HEAP_TYPE InType);
+
+	// Returns descriptor heap base index or -1 if allocation is not possible.
+	// Thread-safe (uses atomic linear allocation).
+	int32 Allocate(uint32 InNumDescriptors);
+
+	void CopyDescriptors(int32 BaseIndex, const D3D12_CPU_DESCRIPTOR_HANDLE* InDescriptors, uint32 InNumDescriptors);
+
+	bool CompareDescriptors(int32 BaseIndex, const D3D12_CPU_DESCRIPTOR_HANDLE* InDescriptors, uint32 InNumDescriptors);
+
+	D3D12_CPU_DESCRIPTOR_HANDLE GetDescriptorCPU(uint32 Index) const;
+
+	D3D12_GPU_DESCRIPTOR_HANDLE GetDescriptorGPU(uint32 Index) const;
+
+	void UpdateSyncPoint();
+
+	D3D12_DESCRIPTOR_HEAP_TYPE Type = D3D12_DESCRIPTOR_HEAP_TYPE_NUM_TYPES;
+	ID3D12DescriptorHeap* D3D12Heap = nullptr;
+	uint32 MaxNumDescriptors = 0;
+
+	int32 NumAllocatedDescriptors = 0;
+
+	// Marks the valid range of the heap when exhaustive sampler deduplication is enabled. Not used otherwise.
+	int32 NumWrittenSamplerDescriptors = 0;
+
+	uint32 DescriptorSize = 0;
+	D3D12_CPU_DESCRIPTOR_HANDLE CPUBase = {};
+	D3D12_GPU_DESCRIPTOR_HANDLE GPUBase = {};
+
+	FD3D12ExplicitDescriptorHeapCache::Entry HeapCacheEntry;
+
+	TArray<D3D12_CPU_DESCRIPTOR_HANDLE> Descriptors;
+
+	bool bExhaustiveSamplerDeduplication = false;
+};
+
+class FD3D12ExplicitDescriptorCache : public FD3D12DeviceChild
+{
+public:
+
+	UE_NONCOPYABLE(FD3D12ExplicitDescriptorCache)
+
+	FD3D12ExplicitDescriptorCache(FD3D12Device* Device, uint32 MaxWorkerCount)
+	: FD3D12DeviceChild(Device)
+	, ViewHeap(Device)
+	, SamplerHeap(Device)
+	{
+		check(MaxWorkerCount > 0u);
+		WorkerData.SetNum(MaxWorkerCount);
+	}
+
+	void Init(uint32 NumViewDescriptors, uint32 NumSamplerDescriptors);
+
+	void UpdateSyncPoint();
+
+	void SetDescriptorHeaps(FD3D12CommandContext& CommandContext);
+
+	// Returns descriptor heap base index for this descriptor table allocation or -1 if allocation failed.
+	int32 AllocateDescriptorTable(const uint32* DescriptorVersions, const D3D12_CPU_DESCRIPTOR_HANDLE* Descriptors, uint32 NumDescriptors, D3D12_DESCRIPTOR_HEAP_TYPE Type, uint32 WorkerIndex);
+
+	FD3D12ExplicitDescriptorHeap ViewHeap;
+	FD3D12ExplicitDescriptorHeap SamplerHeap;
+
+	bool bBindlessViews = false;
+	bool bBindlessSamplers = false;
+
+	template<typename KeyType>
+	struct TIdentityHash
+	{
+		static FORCEINLINE bool Matches(KeyType A, KeyType B)
+		{
+			return A == B;
+		}
+		static FORCEINLINE uint32 GetKeyHash(KeyType Key)
+		{
+			return (uint32)Key;
+		}
+	};
+
+	using TDescriptorHashMap = Experimental::TSherwoodMap<uint64, int32, TIdentityHash<uint64>>;
+
+	struct alignas(PLATFORM_CACHE_LINE_SIZE) FWorkerThreadData
+	{
+		TDescriptorHashMap ViewDescriptorTableCache;
+		TDescriptorHashMap SamplerDescriptorTableCache;
+	};
+
+	TArray<FWorkerThreadData> WorkerData;
+};
\ No newline at end of file
diff --git a/Engine/Source/Runtime/D3D12RHI/Private/D3D12RHI.cpp b/Engine/Source/Runtime/D3D12RHI/Private/D3D12RHI.cpp
--- a/Engine/Source/Runtime/D3D12RHI/Private/D3D12RHI.cpp
+++ b/Engine/Source/Runtime/D3D12RHI/Private/D3D12RHI.cpp
@@ -259,9 +259,11 @@ FD3D12DynamicRHI::FD3D12DynamicRHI(const TArray<TSharedPtr<FD3D12Adapter>>& Chos
 
 void FD3D12DynamicRHI::PostInit()
 {
-	if (GRHISupportsRayTracing)
+	for (TSharedPtr<FD3D12Adapter>& Adapter : ChosenAdapters)
 	{
-		for (TSharedPtr<FD3D12Adapter>& Adapter : ChosenAdapters)
+		Adapter->InitializeExplicitDescriptorHeap();
+
+		if (GRHISupportsRayTracing)
 		{
 			Adapter->InitializeRayTracing();
 		}
diff --git a/Engine/Source/Runtime/D3D12RHI/Private/D3D12RayTracing.cpp b/Engine/Source/Runtime/D3D12RHI/Private/D3D12RayTracing.cpp
--- a/Engine/Source/Runtime/D3D12RHI/Private/D3D12RayTracing.cpp
+++ b/Engine/Source/Runtime/D3D12RHI/Private/D3D12RayTracing.cpp
@@ -20,6 +20,10 @@
 #include "String/LexFromString.h"
 #include "GlobalRenderResources.h"
 #include "D3D12RayTracingDebug.h"
+#include "D3D12ExplicitDescriptorCache.h"
+
+extern int32 GD3D12ExplicitViewDescriptorHeapSize;
+extern int32 GD3D12ExplicitViewDescriptorHeapOverflowReported;
 
 static int32 GRayTracingDebugForceBuildMode = 0;
 static FAutoConsoleVariableRef CVarRayTracingDebugForceFastTrace(
@@ -40,16 +44,6 @@ static FAutoConsoleVariableRef CVarRayTracingShaderRecordCache(
 	TEXT("This mode assumes that contents of uniform buffers does not change during ray tracing resource binding.")
 );
 
-static int32 GD3D12RayTracingViewDescriptorHeapSize = 250'000;
-static int32 GD3D12RayTracingViewDescriptorHeapOverflowReported = 0;
-static FAutoConsoleVariableRef CVarD3D12RayTracingViewDescriptorHeapSize(
-	TEXT("r.D3D12.RayTracing.ViewDescriptorHeapSize"),
-	GD3D12RayTracingViewDescriptorHeapSize,
-	TEXT("Maximum number of descriptors per ray tracing view descriptor heap. (default = 250k, ~8MB per heap)\n")
-	TEXT("Typical measured descriptor heap usage in large scenes is ~50k. An error is reported when this limit is reached and shader bindings for subsequent objects are skipped.\n"),
-	ECVF_ReadOnly
-);
-
 static int32 GD3D12RayTracingAllowCompaction = 1;
 static FAutoConsoleVariableRef CVarD3D12RayTracingAllowCompaction(
 	TEXT("r.D3D12.RayTracing.AllowCompaction"),
@@ -83,13 +77,6 @@ static FAutoConsoleVariableRef CVarRayTracingAllowSpecializedStateObjects(
 	TEXT("This is intended for performance testingand has no effect if r.D3D12.RayTracing.SpecializeStateObjects is 0. (default = 1)\n")
 );
 
-static int32 GD3D12RayTracingDeduplicateSamplers = 1;
-static FAutoConsoleVariableRef CVarD3D12RayTracingDeduplicateSamplers(
-	TEXT("r.D3D12.RayTracing.DeduplicateSamplers"),
-	GD3D12RayTracingDeduplicateSamplers,
-	TEXT("Use an exhaustive search to deduplicate sampler descriptors when generating shader binding tables. Reduces sampler heap usage at the cost of some CPU time. (default = 1)")
-);
-
 static int32 GD3D12RayTracingGPUValidation = 0;
 static FAutoConsoleVariableRef CVarD3D12RayTracingGPUValidation(
 	TEXT("r.D3D12.RayTracing.GPUValidation"),
@@ -108,16 +95,6 @@ DECLARE_DWORD_ACCUMULATOR_STAT(TEXT("Allocated bottom level acceleration structu
 DECLARE_DWORD_ACCUMULATOR_STAT(TEXT("Allocated top level acceleration structures"), STAT_D3D12RayTracingAllocatedTLAS, STATGROUP_D3D12RayTracing);
 DECLARE_DWORD_ACCUMULATOR_STAT(TEXT("Triangles in all BL acceleration structures"), STAT_D3D12RayTracingTrianglesBLAS, STATGROUP_D3D12RayTracing);
 
-DECLARE_DWORD_ACCUMULATOR_STAT(TEXT("Allocated sampler descriptor heaps"), STAT_D3D12RayTracingSamplerDescriptorHeaps, STATGROUP_D3D12RayTracing);
-DECLARE_DWORD_ACCUMULATOR_STAT(TEXT("Allocated sampler descriptors"), STAT_D3D12RayTracingSamplerDescriptors, STATGROUP_D3D12RayTracing);
-
-DECLARE_DWORD_ACCUMULATOR_STAT(TEXT("Allocated view descriptor heaps"), STAT_D3D12RayTracingViewDescriptorHeaps, STATGROUP_D3D12RayTracing);
-DECLARE_DWORD_ACCUMULATOR_STAT(TEXT("Allocated view descriptors"), STAT_D3D12RayTracingViewDescriptors, STATGROUP_D3D12RayTracing);
-
-DECLARE_DWORD_COUNTER_STAT(TEXT("Max used sampler descriptors in a single heap"), STAT_D3D12RayTracingMaxUsedSamplerDescriptors, STATGROUP_D3D12RayTracing);
-DECLARE_DWORD_COUNTER_STAT(TEXT("Used sampler descriptors (per frame)"), STAT_D3D12RayTracingUsedSamplerDescriptors, STATGROUP_D3D12RayTracing);
-DECLARE_DWORD_COUNTER_STAT(TEXT("Used view descriptors (per frame)"), STAT_D3D12RayTracingUsedViewDescriptors, STATGROUP_D3D12RayTracing);
-
 DECLARE_DWORD_COUNTER_STAT(TEXT("Built BL AS (per frame)"), STAT_D3D12RayTracingBuiltBLAS, STATGROUP_D3D12RayTracing);
 DECLARE_DWORD_COUNTER_STAT(TEXT("Updated BL AS (per frame)"), STAT_D3D12RayTracingUpdatedBLAS, STATGROUP_D3D12RayTracing);
 DECLARE_DWORD_COUNTER_STAT(TEXT("Built TL AS (per frame)"), STAT_D3D12RayTracingBuiltTLAS, STATGROUP_D3D12RayTracing);
@@ -387,11 +364,6 @@ inline void UnregisterD3D12RayTracingGeometry(FD3D12RayTracingGeometry* Geometry
 }
 #endif // UE_BUILD_SHIPPING
 
-// Whether to compare the full descriptor table on cache lookup or only use FXxHash64 digest.
-#ifndef RAY_TRACING_DESCRIPTOR_CACHE_FULL_COMPARE
-#define RAY_TRACING_DESCRIPTOR_CACHE_FULL_COMPARE 1
-#endif // RAY_TRACING_DESCRIPTOR_CACHE_FULL_COMPARE
-
 struct FD3D12ShaderIdentifier
 {
 	uint64 Data[4] = {~0ull, ~0ull, ~0ull, ~0ull};
@@ -1269,497 +1241,6 @@ private:
 	FD3D12RootSignature DefaultLocalRootSignature; // Default empty root signature used for default hit shaders.
 };
 
-// #dxr_todo UE-72158: FD3D12Device::GlobalViewHeap/GlobalSamplerHeap should be used instead of ad-hoc heaps here.
-// Unfortunately, this requires a major refactor of how global heaps work.
-// FD3D12CommandContext-s should not get static chunks of the global heap, but instead should dynamically allocate
-// chunks on as-needed basis and release them when possible.
-// This would allow ray tracing code to sub-allocate heap blocks from the same global heap.
-class FD3D12RayTracingDescriptorHeapCache : FD3D12DeviceChild
-{
-public:
-
-	UE_NONCOPYABLE(FD3D12RayTracingDescriptorHeapCache)
-
-	struct Entry
-	{
-		ID3D12DescriptorHeap* Heap = nullptr;
-		uint64 FenceValue = 0;
-		uint32 NumDescriptors = 0;
-		D3D12_DESCRIPTOR_HEAP_TYPE Type = D3D12_DESCRIPTOR_HEAP_TYPE_NUM_TYPES;
-	};
-
-	FD3D12RayTracingDescriptorHeapCache(FD3D12Device* Device)
-		: FD3D12DeviceChild(Device)
-	{
-
-	}
-
-	~FD3D12RayTracingDescriptorHeapCache()
-	{
-		check(AllocatedEntries == 0);
-
-		FScopeLock Lock(&CriticalSection);
-		for (const Entry& It : Entries)
-		{
-			if (It.Type == D3D12_DESCRIPTOR_HEAP_TYPE_CBV_SRV_UAV)
-			{
-				DEC_DWORD_STAT(STAT_D3D12RayTracingViewDescriptorHeaps);
-				DEC_DWORD_STAT_BY(STAT_D3D12RayTracingViewDescriptors, It.NumDescriptors);
-			}
-			else
-			{
-				DEC_DWORD_STAT(STAT_D3D12RayTracingSamplerDescriptorHeaps);
-				DEC_DWORD_STAT_BY(STAT_D3D12RayTracingSamplerDescriptors, It.NumDescriptors);
-			}
-
-			It.Heap->Release();
-		}
-		Entries.Empty();
-	}
-
-	void ReleaseHeap(Entry& Entry)
-	{
-		FScopeLock Lock(&CriticalSection);
-
-		Entries.Add(Entry);
-
-		check(AllocatedEntries != 0);
-		--AllocatedEntries;
-	}
-
-	Entry AllocateHeap(D3D12_DESCRIPTOR_HEAP_TYPE Type, uint32 NumDescriptors)
-	{
-		FScopeLock Lock(&CriticalSection);
-
-		++AllocatedEntries;
-
-		Entry Result = {};
-
-		FD3D12ManualFence& Fence = GetParentDevice()->GetParentAdapter()->GetFrameFence();
-		const uint64 CompletedFenceValue = Fence.GetCompletedFenceValue(/* bUpdateCachedFenceValue */ false);
-
-		for (int32 EntryIndex = 0; EntryIndex < Entries.Num(); ++EntryIndex)
-		{
-			const Entry& It = Entries[EntryIndex];
-			if (It.Type == Type && It.NumDescriptors >= NumDescriptors && It.FenceValue <= CompletedFenceValue)
-			{
-				Result = It;
-
-				Entries[EntryIndex] = Entries.Last();
-				Entries.Pop();
-
-				return Result;
-			}
-		}
-
-		// Compatible heap was not found in cache, so create a new one.
-
-		ReleaseStaleEntries(100, CompletedFenceValue); // Release heaps that were not used for 100 frames before allocating new.
-
-		D3D12_DESCRIPTOR_HEAP_DESC Desc = {};
-
-		Desc.Flags = D3D12_DESCRIPTOR_HEAP_FLAG_SHADER_VISIBLE;
-		Desc.Type = Type;
-		Desc.NumDescriptors = NumDescriptors;
-		Desc.NodeMask = GetParentDevice()->GetGPUMask().GetNative();
-
-		ID3D12DescriptorHeap* D3D12Heap = nullptr;
-
-		const TCHAR* HeapName = Desc.Type == D3D12_DESCRIPTOR_HEAP_TYPE_CBV_SRV_UAV ? TEXT("RT View Heap") : TEXT("RT Sampler Heap");
-		UE_LOG(LogD3D12RHI, Log, TEXT("Creating %s with %d entries"), HeapName, NumDescriptors);
-
-		VERIFYD3D12RESULT(GetParentDevice()->GetDevice()->CreateDescriptorHeap(&Desc, IID_PPV_ARGS(&D3D12Heap)));
-		SetName(D3D12Heap, Desc.Type == D3D12_DESCRIPTOR_HEAP_TYPE_CBV_SRV_UAV ? L"RT View Heap" : L"RT Sampler Heap");
-
-		Result.NumDescriptors = NumDescriptors;
-		Result.Type = Type;
-		Result.Heap = D3D12Heap;
-
-		if (Desc.Type == D3D12_DESCRIPTOR_HEAP_TYPE_CBV_SRV_UAV)
-		{
-			INC_DWORD_STAT(STAT_D3D12RayTracingViewDescriptorHeaps);
-			INC_DWORD_STAT_BY(STAT_D3D12RayTracingViewDescriptors, NumDescriptors);
-		}
-		else
-		{
-			INC_DWORD_STAT(STAT_D3D12RayTracingSamplerDescriptorHeaps);
-			INC_DWORD_STAT_BY(STAT_D3D12RayTracingSamplerDescriptors, NumDescriptors);
-		}
-
-		return Result;
-	}
-
-	void ReleaseStaleEntries(uint32 MaxAge, uint64 CompletedFenceValue)
-	{
-		int32 EntryIndex = 0;
-		while (EntryIndex < Entries.Num())
-		{
-			Entry& It = Entries[EntryIndex];
-			if ((It.FenceValue + MaxAge) <= CompletedFenceValue)
-			{
-				if (It.Type == D3D12_DESCRIPTOR_HEAP_TYPE_CBV_SRV_UAV)
-				{
-					DEC_DWORD_STAT(STAT_D3D12RayTracingViewDescriptorHeaps);
-					DEC_DWORD_STAT_BY(STAT_D3D12RayTracingViewDescriptors, It.NumDescriptors);
-				}
-				else
-				{
-					DEC_DWORD_STAT(STAT_D3D12RayTracingSamplerDescriptorHeaps);
-					DEC_DWORD_STAT_BY(STAT_D3D12RayTracingSamplerDescriptors, It.NumDescriptors);
-				}
-
-				It.Heap->Release();
-
-				Entries[EntryIndex] = Entries.Last();
-				Entries.Pop(false);
-			}
-			else
-			{
-				EntryIndex++;
-			}
-		}
-	}
-
-	void Flush()
-	{
-		FScopeLock Lock(&CriticalSection);
-
-		for (const Entry& It : Entries)
-		{
-			FD3D12DynamicRHI::GetD3DRHI()->DeferredDelete(It.Heap);
-		}
-		Entries.Empty();
-	}
-
-	FCriticalSection CriticalSection;
-	TArray<Entry> Entries;
-	uint32 AllocatedEntries = 0;
-};
-
-struct FD3D12RayTracingDescriptorHeap : public FD3D12DeviceChild
-{
-	UE_NONCOPYABLE(FD3D12RayTracingDescriptorHeap)
-
-	FD3D12RayTracingDescriptorHeap(FD3D12Device* Device)
-		: FD3D12DeviceChild(Device)
-	{
-
-	}
-
-	~FD3D12RayTracingDescriptorHeap()
-	{
-		if (D3D12Heap)
-		{
-			GetParentDevice()->GetRayTracingDescriptorHeapCache()->ReleaseHeap(HeapCacheEntry);
-		}
-	}
-
-	void Init(uint32 InMaxNumDescriptors, D3D12_DESCRIPTOR_HEAP_TYPE InType)
-	{
-		check(D3D12Heap == nullptr);
-
-		Type = InType;
-		HeapCacheEntry = GetParentDevice()->GetRayTracingDescriptorHeapCache()->AllocateHeap(Type, InMaxNumDescriptors);
-
-		MaxNumDescriptors = HeapCacheEntry.NumDescriptors;
-		D3D12Heap = HeapCacheEntry.Heap;
-
-		CPUBase = D3D12Heap->GetCPUDescriptorHandleForHeapStart();
-		GPUBase = D3D12Heap->GetGPUDescriptorHandleForHeapStart();
-
-		checkf(CPUBase.ptr, TEXT("Ray tracing descriptor heap of type %d returned from descriptor heap cache is invalid."), Type);
-
-		DescriptorSize = GetParentDevice()->GetDevice()->GetDescriptorHandleIncrementSize(Type);
-
-		bExhaustiveSamplerDeduplication = GD3D12RayTracingDeduplicateSamplers == 1;
-
-		if (bExhaustiveSamplerDeduplication && InType == D3D12_DESCRIPTOR_HEAP_TYPE_SAMPLER)
-		{
-			// When exhaustive descriptor deduplication is active, all shadow descriptor table entries must be initialized.
-			// Deduplication works by looping over all elements, but they may be written out of order by worker threads.
-			// Initializing descriptors to 0 avoids accidentally matching wrong descriptor entries.
-			Descriptors.SetNumZeroed(MaxNumDescriptors);
-		}
-#if RAY_TRACING_DESCRIPTOR_CACHE_FULL_COMPARE
-		else
-		{
-			// If exhaustive sampler deduplication is not used, then descriptor table copy is only used for validation.
-			// Entries are guaranteed to be written before they're read, so we don't have to spend the time on initialization.
-			Descriptors.SetNumUninitialized(MaxNumDescriptors);
-		}
-#endif // RAY_TRACING_DESCRIPTOR_CACHE_FULL_COMPARE
-	}
-
-	// Returns descriptor heap base index or -1 if allocation is not possible.
-	// Thread-safe (uses atomic linear allocation).
-	int32 Allocate(uint32 InNumDescriptors)
-	{
-		int32 Result = FPlatformAtomics::InterlockedAdd(&NumAllocatedDescriptors, InNumDescriptors);
-
-		if (Result + InNumDescriptors > MaxNumDescriptors)
-		{
-			if (Type == D3D12_DESCRIPTOR_HEAP_TYPE_SAMPLER)
-			{
-				UE_LOG(LogD3D12RHI, Fatal,
-					TEXT("Ray tracing sampler descriptor heap overflow. ")
-					TEXT("It is not possible to recover from this error, as maximum D3D12 sampler heap size is 2048."));
-			}
-			else if (GD3D12RayTracingViewDescriptorHeapSize == MaxNumDescriptors
-				&& FPlatformAtomics::InterlockedOr(&GD3D12RayTracingViewDescriptorHeapOverflowReported, 1) == 0)
-			{
-				// NOTE: GD3D12RayTracingViewDescriptorHeapOverflowReported is set atomically because multiple 
-				// allocations may be happening simultaneously, but we only want to report the error once.
-
-				UE_LOG(LogD3D12RHI, Error,
-					TEXT("Ray tracing view descriptor heap overflow. Current frame will not be rendered correctly. ")
-					TEXT("Increase r.D3D12.RayTracingViewDescriptorHeapSize to at least %d to fix this issue."),
-					MaxNumDescriptors * 2);
-			}
-
-			Result = -1;
-		}
-
-		return Result;
-	}
-
-	void CopyDescriptors(int32 BaseIndex, const D3D12_CPU_DESCRIPTOR_HANDLE* InDescriptors, uint32 InNumDescriptors)
-	{
-		D3D12_CPU_DESCRIPTOR_HANDLE DestDescriptor = GetDescriptorCPU(BaseIndex);
-		GetParentDevice()->GetDevice()->CopyDescriptors(1, &DestDescriptor, &InNumDescriptors, InNumDescriptors, InDescriptors, nullptr, Type);
-		for (uint32 i = 0; i < InNumDescriptors; ++i)
-		{
-			Descriptors[BaseIndex + i].ptr = InDescriptors[i].ptr;
-		}
-	}
-
-	bool CompareDescriptors(int32 BaseIndex, const D3D12_CPU_DESCRIPTOR_HANDLE* InDescriptors, uint32 InNumDescriptors)
-	{
-		for (uint32 i = 0; i < InNumDescriptors; ++i)
-		{
-			if (Descriptors[BaseIndex + i].ptr != InDescriptors[i].ptr)
-			{
-				return false;
-			}
-		}
-		return true;
-	}
-
-	D3D12_CPU_DESCRIPTOR_HANDLE GetDescriptorCPU(uint32 Index) const
-	{
-		checkSlow(Index < MaxNumDescriptors);
-		D3D12_CPU_DESCRIPTOR_HANDLE Result = { CPUBase.ptr + Index * DescriptorSize };
-		return Result;
-	}
-
-	D3D12_GPU_DESCRIPTOR_HANDLE GetDescriptorGPU(uint32 Index) const
-	{
-		checkSlow(Index < MaxNumDescriptors);
-		D3D12_GPU_DESCRIPTOR_HANDLE Result = { GPUBase.ptr + Index * DescriptorSize };
-		return Result;
-	}
-
-	void UpdateSyncPoint()
-	{
-		const FD3D12ManualFence& Fence = GetParentDevice()->GetParentAdapter()->GetFrameFence();
-		HeapCacheEntry.FenceValue = FMath::Max(HeapCacheEntry.FenceValue, Fence.GetNextFenceToSignal());
-	}
-
-	D3D12_DESCRIPTOR_HEAP_TYPE Type = D3D12_DESCRIPTOR_HEAP_TYPE_NUM_TYPES;
-	ID3D12DescriptorHeap* D3D12Heap = nullptr;
-	uint32 MaxNumDescriptors = 0;
-
-	int32 NumAllocatedDescriptors = 0;
-
-	// Marks the valid range of the heap when exhaustive sampler deduplication is enabled. Not used otherwise.
-	int32 NumWrittenSamplerDescriptors = 0;
-
-	uint32 DescriptorSize = 0;
-	D3D12_CPU_DESCRIPTOR_HANDLE CPUBase = {};
-	D3D12_GPU_DESCRIPTOR_HANDLE GPUBase = {};
-
-	FD3D12RayTracingDescriptorHeapCache::Entry HeapCacheEntry;
-
-	TArray<D3D12_CPU_DESCRIPTOR_HANDLE> Descriptors;
-
-	bool bExhaustiveSamplerDeduplication = false;
-};
-
-class FD3D12RayTracingDescriptorCache : public FD3D12DeviceChild
-{
-public:
-
-	UE_NONCOPYABLE(FD3D12RayTracingDescriptorCache)
-
-	FD3D12RayTracingDescriptorCache(FD3D12Device* Device)
-		: FD3D12DeviceChild(Device)
-		, ViewHeap(Device)
-		, SamplerHeap(Device)
-	{}
-
-	void Init(uint32 NumViewDescriptors, uint32 NumSamplerDescriptors)
-	{
-		FD3D12BindlessDescriptorManager& BindlessManager = GetParentDevice()->GetBindlessDescriptorManager();
-
-		bBindlessViews = BindlessManager.HasHeap(ERHIDescriptorHeapType::Standard, ERHIBindlessConfiguration::RayTracingShaders);
-		bBindlessSamplers = BindlessManager.HasHeap(ERHIDescriptorHeapType::Sampler, ERHIBindlessConfiguration::RayTracingShaders);
-
-		if (!bBindlessViews)
-		{
-			ViewHeap.Init(NumViewDescriptors, D3D12_DESCRIPTOR_HEAP_TYPE_CBV_SRV_UAV);
-		}
-		
-		if (!bBindlessSamplers)
-		{
-			SamplerHeap.Init(NumSamplerDescriptors, D3D12_DESCRIPTOR_HEAP_TYPE_SAMPLER);
-		}
-	}
-
-	void UpdateSyncPoint()
-	{
-		check(IsInRHIThread() || !IsRunningRHIInSeparateThread());
-
-		if (!bBindlessViews)
-		{
-			ViewHeap.UpdateSyncPoint();
-		}
-
-		if (!bBindlessSamplers)
-		{
-			SamplerHeap.UpdateSyncPoint();
-		}
-	}
-
-	void SetDescriptorHeaps(FD3D12CommandContext& CommandContext)
-	{
-		UpdateSyncPoint();
-
-		check(bBindlessViews || ViewHeap.GetParentDevice() == CommandContext.GetParentDevice());
-		check(bBindlessSamplers || SamplerHeap.GetParentDevice() == CommandContext.GetParentDevice());
-
-		FD3D12BindlessDescriptorManager& BindlessManager = GetParentDevice()->GetBindlessDescriptorManager();
-
-		ID3D12DescriptorHeap* ViewHeapToSet = bBindlessViews ? BindlessManager.GetHeap(ERHIDescriptorHeapType::Standard)->GetHeap() : ViewHeap.D3D12Heap;
-		ID3D12DescriptorHeap* SamplerHeapToSet = bBindlessSamplers ? BindlessManager.GetHeap(ERHIDescriptorHeapType::Sampler)->GetHeap() : SamplerHeap.D3D12Heap;
-
-		CommandContext.StateCache.GetDescriptorCache()->OverrideLastSetHeaps(ViewHeapToSet, SamplerHeapToSet);
-	}
-
-	// Returns descriptor heap base index for this descriptor table allocation or -1 if allocation failed.
-	int32 AllocateDescriptorTable(const uint32* DescriptorVersions, const D3D12_CPU_DESCRIPTOR_HANDLE* Descriptors, uint32 NumDescriptors, D3D12_DESCRIPTOR_HEAP_TYPE Type, uint32 WorkerIndex)
-	{
-		checkSlow(Type == D3D12_DESCRIPTOR_HEAP_TYPE_CBV_SRV_UAV || Type == D3D12_DESCRIPTOR_HEAP_TYPE_SAMPLER);
-
-		FD3D12RayTracingDescriptorHeap& Heap = (Type == D3D12_DESCRIPTOR_HEAP_TYPE_CBV_SRV_UAV) ? ViewHeap : SamplerHeap;
-		TDescriptorHashMap& Map = (Type == D3D12_DESCRIPTOR_HEAP_TYPE_CBV_SRV_UAV) 
-			? WorkerData[WorkerIndex].ViewDescriptorTableCache
-			: WorkerData[WorkerIndex].SamplerDescriptorTableCache;
-
-		const uint64 VersionHash = FXxHash64::HashBuffer(DescriptorVersions, sizeof(DescriptorVersions[0]) * NumDescriptors).Hash;
-		const uint64 DescriptorHash = FXxHash64::HashBuffer(Descriptors, sizeof(Descriptors[0]) * NumDescriptors).Hash;
-		const uint64 Key = VersionHash ^ DescriptorHash;
-
-		int32& DescriptorTableBaseIndex = Map.FindOrAdd(Key, INDEX_NONE);
-
-		if (DescriptorTableBaseIndex != INDEX_NONE)
-		{
-		#if RAY_TRACING_DESCRIPTOR_CACHE_FULL_COMPARE
-			if (ensureMsgf(Heap.CompareDescriptors(DescriptorTableBaseIndex, Descriptors, NumDescriptors), 
-				TEXT("Ray tracing descriptor cache hash collision detected!")))
-		#endif // RAY_TRACING_DESCRIPTOR_CACHE_FULL_COMPARE
-			{
-				return DescriptorTableBaseIndex;
-			}
-		}
-
-		if (Heap.bExhaustiveSamplerDeduplication && Type == D3D12_DESCRIPTOR_HEAP_TYPE_SAMPLER)
-		{
-			// Exhaustive search for a sampler table:
-			// We have to do this because sampler heap space is precious (hard limit of 2048 total entries).
-			// Per-thread descriptor table deduplication hash tables introduce a lot of redundancy in the heap
-			// which is reduced by looking for global matches on hash table lookup miss.
-
-			int32 FoundIndex = INDEX_NONE;
-			int32 SearchEndPos = Heap.NumWrittenSamplerDescriptors;
-			for (int32 SearchIndex = 0; SearchIndex + int32(NumDescriptors) < SearchEndPos; ++SearchIndex)
-			{
-				if (Heap.CompareDescriptors(SearchIndex, Descriptors, NumDescriptors))
-				{
-					FoundIndex = SearchIndex;
-					break;
-				}
-			}
-
-			if (FoundIndex != INDEX_NONE)
-			{
-				DescriptorTableBaseIndex = FoundIndex;
-				return DescriptorTableBaseIndex;
-			}
-		}
-
-		DescriptorTableBaseIndex = Heap.Allocate(NumDescriptors);
-
-		if (DescriptorTableBaseIndex == INDEX_NONE)
-		{
-			return INDEX_NONE;
-		}
-
-		Heap.CopyDescriptors(DescriptorTableBaseIndex, Descriptors, NumDescriptors);
-
-		if (Heap.bExhaustiveSamplerDeduplication && Type == D3D12_DESCRIPTOR_HEAP_TYPE_SAMPLER)
-		{
-			FPlatformAtomics::InterlockedAdd(&Heap.NumWrittenSamplerDescriptors, NumDescriptors);
-		}
-
-		if (Type == D3D12_DESCRIPTOR_HEAP_TYPE_CBV_SRV_UAV)
-		{
-			INC_DWORD_STAT_BY(STAT_D3D12RayTracingUsedViewDescriptors, NumDescriptors);
-		}
-		else
-		{
-			INC_DWORD_STAT_BY(STAT_D3D12RayTracingUsedSamplerDescriptors, NumDescriptors);
-
-			static uint64 GMaxNumUsedSamplerDescriptors = 0;
-			if (Heap.NumAllocatedDescriptors > GMaxNumUsedSamplerDescriptors)
-			{
-				GMaxNumUsedSamplerDescriptors = Heap.NumAllocatedDescriptors;
-			}
-			SET_DWORD_STAT(STAT_D3D12RayTracingMaxUsedSamplerDescriptors, GMaxNumUsedSamplerDescriptors);
-		}
-
-		return DescriptorTableBaseIndex;
-	}
-
-	FD3D12RayTracingDescriptorHeap ViewHeap;
-	FD3D12RayTracingDescriptorHeap SamplerHeap;
-
-	bool bBindlessViews = false;
-	bool bBindlessSamplers = false;
-
-	template<typename KeyType>
-	struct TIdentityHash
-	{
-		static FORCEINLINE bool Matches(KeyType A, KeyType B)
-		{
-			return A == B;
-		}
-		static FORCEINLINE uint32 GetKeyHash(KeyType Key)
-		{
-			return (uint32)Key;
-		}
-	};
-
-	static constexpr uint32 MaxBindingWorkers = FD3D12RayTracingScene::MaxBindingWorkers;
-	using TDescriptorHashMap = Experimental::TSherwoodMap<uint64, int32, TIdentityHash<uint64>>;
-
-	struct alignas(PLATFORM_CACHE_LINE_SIZE) FWorkerThreadData
-	{
-		TDescriptorHashMap ViewDescriptorTableCache;
-		TDescriptorHashMap SamplerDescriptorTableCache;
-	};
-
-	FWorkerThreadData WorkerData[MaxBindingWorkers];
-};
-
 class FD3D12RayTracingShaderTable
 {
 private:
@@ -1837,18 +1318,18 @@ public:
 		{
 			// #dxr_todo UE-72158: Remove this when RT descriptors are sub-allocated from the global view descriptor heap.
 
-			if (GD3D12RayTracingViewDescriptorHeapOverflowReported)
+			if (GD3D12ExplicitViewDescriptorHeapOverflowReported)
 			{
-				GD3D12RayTracingViewDescriptorHeapSize = GD3D12RayTracingViewDescriptorHeapSize * 2;
-				GD3D12RayTracingViewDescriptorHeapOverflowReported = 0;
+				GD3D12ExplicitViewDescriptorHeapSize = GD3D12ExplicitViewDescriptorHeapSize * 2;
+				GD3D12ExplicitViewDescriptorHeapOverflowReported = 0;
 			}
 
 			// D3D12 is guaranteed to support 1M (D3D12_MAX_SHADER_VISIBLE_DESCRIPTOR_HEAP_SIZE_TIER_1) descriptors in a CBV/SRV/UAV heap, so clamp the size to this.
 			// https://docs.microsoft.com/en-us/windows/desktop/direct3d12/hardware-support
-			const uint32 NumViewDescriptors = FMath::Min(D3D12_MAX_SHADER_VISIBLE_DESCRIPTOR_HEAP_SIZE_TIER_1, GD3D12RayTracingViewDescriptorHeapSize);
+			const uint32 NumViewDescriptors = FMath::Min(D3D12_MAX_SHADER_VISIBLE_DESCRIPTOR_HEAP_SIZE_TIER_1, GD3D12ExplicitViewDescriptorHeapSize);
 			const uint32 NumSamplerDescriptors = D3D12_MAX_SHADER_VISIBLE_SAMPLER_HEAP_SIZE;
 
-			DescriptorCache = new FD3D12RayTracingDescriptorCache(Device);
+			DescriptorCache = new FD3D12ExplicitDescriptorCache(Device, FD3D12RayTracingScene::MaxBindingWorkers);
 			DescriptorCache->Init(NumViewDescriptors, NumSamplerDescriptors);
 		}
 
@@ -2122,7 +1603,7 @@ public:
 #endif
 
 	// SBTs have their own descriptor heaps
-	FD3D12RayTracingDescriptorCache* DescriptorCache = nullptr;
+	FD3D12ExplicitDescriptorCache* DescriptorCache = nullptr;
 
 	struct FShaderRecordCacheKey
 	{
@@ -2284,13 +1765,6 @@ public:
 	FWorkerThreadData WorkerData[MaxBindingWorkers];
 };
 
-
-void FD3D12Device::DestroyRayTracingDescriptorCache()
-{
-	delete RayTracingDescriptorHeapCache;
-	RayTracingDescriptorHeapCache = nullptr;
-}
-
 struct FD3D12RayTracingShaderLibrary
 {
 	void Reserve(uint32 NumShaders)
@@ -2912,9 +2386,6 @@ void FD3D12Device::InitRayTracing()
 	LLM_SCOPE_BYNAME(TEXT("FD3D12RT"));
 	check(RayTracingPipelineCache == nullptr);
 	RayTracingPipelineCache = new FD3D12RayTracingPipelineCache(this);
-
-	check(RayTracingDescriptorHeapCache == nullptr);
-	RayTracingDescriptorHeapCache = new FD3D12RayTracingDescriptorHeapCache(this);
 }
 
 void FD3D12Device::CleanupRayTracing()
@@ -2924,8 +2395,6 @@ void FD3D12Device::CleanupRayTracing()
 
 	delete RayTracingDispatchRaysDescBuffer;
 	RayTracingDispatchRaysDescBuffer = nullptr;
-
-	// Note: RayTracingDescriptorHeapCache is destroyed in ~FD3D12Device, after all deferred deletion is processed
 }
 
 static D3D12_RAYTRACING_INSTANCE_FLAGS TranslateRayTracingInstanceFlags(ERayTracingInstanceFlags InFlags)
@@ -4550,7 +4019,7 @@ void FD3D12CommandContext::RHIClearRayTracingBindings(FRHIRayTracingScene* InSce
 
 struct FD3D12RayTracingGlobalResourceBinder
 {
-	FD3D12RayTracingGlobalResourceBinder(FD3D12CommandContext& InCommandContext, FD3D12RayTracingDescriptorCache& InDescriptorCache)
+	FD3D12RayTracingGlobalResourceBinder(FD3D12CommandContext& InCommandContext, FD3D12ExplicitDescriptorCache& InDescriptorCache)
 		: CommandContext(InCommandContext)
 		, DescriptorCache(InDescriptorCache)
 	{
@@ -4599,7 +4068,7 @@ struct FD3D12RayTracingGlobalResourceBinder
 	}
 
 	FD3D12CommandContext& CommandContext;
-	FD3D12RayTracingDescriptorCache& DescriptorCache;
+	FD3D12ExplicitDescriptorCache& DescriptorCache;
 	static constexpr uint32 WorkerIndex = 0;
 };
 
@@ -4615,7 +4084,7 @@ struct FD3D12RayTracingLocalResourceBinder
 	{
 		check(InShaderTable.DescriptorCache != nullptr);
 		check(WorkerIndex < InShaderTable.MaxBindingWorkers);
-		check(WorkerIndex < DescriptorCache.MaxBindingWorkers);
+		check(WorkerIndex < uint32(DescriptorCache.WorkerData.Num()));
 		check(RecordIndex != ~0u);
 
 		switch (InBindingType)
@@ -4709,7 +4178,7 @@ struct FD3D12RayTracingLocalResourceBinder
 
 	FD3D12Device& Device;
 	FD3D12RayTracingShaderTable& ShaderTable;
-	FD3D12RayTracingDescriptorCache& DescriptorCache;
+	FD3D12ExplicitDescriptorCache& DescriptorCache;
 	const FD3D12RootSignature& RootSignature;
 	uint32 ShaderTableOffset = 0;
 	uint32 RecordIndex = ~0u;
@@ -4928,7 +4397,7 @@ static bool SetRayTracingShaderResources(
 	check(IsCompleteBinding(Shader->ResourceCounts.NumCBs     , Bindings.BoundCBVMask));
 	check(IsCompleteBinding(Shader->ResourceCounts.NumSamplers, Bindings.BoundSamplerMask));
 
-	FD3D12RayTracingDescriptorCache& DescriptorCache = Binder.DescriptorCache;
+	FD3D12ExplicitDescriptorCache& DescriptorCache = Binder.DescriptorCache;
 	const uint32 WorkerIndex = Binder.WorkerIndex;
 
 	const uint32 NumSRVs = Shader->ResourceCounts.NumSRVs;
@@ -5143,7 +4612,7 @@ static void DispatchRays(FD3D12CommandContext& CommandContext,
 	bool bResourcesBound = false;
 	if (OptShaderTable && OptShaderTable->DescriptorCache)
 	{
-		FD3D12RayTracingDescriptorCache* DescriptorCache = OptShaderTable->DescriptorCache;
+		FD3D12ExplicitDescriptorCache* DescriptorCache = OptShaderTable->DescriptorCache;
 		check(DescriptorCache != nullptr);
 
 		DescriptorCache->SetDescriptorHeaps(CommandContext);
@@ -5154,7 +4623,7 @@ static void DispatchRays(FD3D12CommandContext& CommandContext,
 	}
 	else
 	{
-		FD3D12RayTracingDescriptorCache TransientDescriptorCache(CommandContext.GetParentDevice());
+		FD3D12ExplicitDescriptorCache TransientDescriptorCache(CommandContext.GetParentDevice(), FD3D12RayTracingScene::MaxBindingWorkers);
 		TransientDescriptorCache.Init(MAX_SRVS + MAX_UAVS, MAX_SAMPLERS);
 		TransientDescriptorCache.SetDescriptorHeaps(CommandContext);
 		FD3D12RayTracingGlobalResourceBinder ResourceBinder(CommandContext, TransientDescriptorCache);
@@ -5603,4 +5072,5 @@ void FD3D12CommandContext::RHISetRayTracingMissShader(
 		UserData,
 		WorkerIndex);
 }
+
 #endif // D3D12_RHI_RAYTRACING
diff --git a/Engine/Source/Runtime/D3D12RHI/Private/D3D12Stats.h b/Engine/Source/Runtime/D3D12RHI/Private/D3D12Stats.h
--- a/Engine/Source/Runtime/D3D12RHI/Private/D3D12Stats.h
+++ b/Engine/Source/Runtime/D3D12RHI/Private/D3D12Stats.h
@@ -136,6 +136,15 @@ DECLARE_DWORD_ACCUMULATOR_STAT_EXTERN(TEXT("View Global: Used Descriptors"), STA
 DECLARE_DWORD_ACCUMULATOR_STAT_EXTERN(TEXT("View Global: Wasted Descriptors"), STAT_GlobalViewHeapWastedDescriptors, STATGROUP_D3D12DescriptorHeap, );
 DECLARE_DWORD_COUNTER_STAT_EXTERN(TEXT("View Global: Block Allocations"), STAT_GlobalViewHeapBlockAllocations, STATGROUP_D3D12DescriptorHeap, );
 
+DECLARE_DWORD_ACCUMULATOR_STAT_EXTERN(TEXT("Explicit sampler descriptor heaps"), STAT_ExplicitSamplerDescriptorHeaps, STATGROUP_D3D12DescriptorHeap, );
+DECLARE_DWORD_ACCUMULATOR_STAT_EXTERN(TEXT("Explicit sampler descriptors"), STAT_ExplicitSamplerDescriptors, STATGROUP_D3D12DescriptorHeap, );
+
+DECLARE_DWORD_ACCUMULATOR_STAT_EXTERN(TEXT("Explicit view descriptor heaps"), STAT_ExplicitViewDescriptorHeaps, STATGROUP_D3D12DescriptorHeap, );
+DECLARE_DWORD_ACCUMULATOR_STAT_EXTERN(TEXT("Explicit view descriptors"), STAT_ExplicitViewDescriptors, STATGROUP_D3D12DescriptorHeap, );
+
+DECLARE_DWORD_COUNTER_STAT_EXTERN(TEXT("Max used explicit sampler descriptors in a single heap"), STAT_ExplicitMaxUsedSamplerDescriptors, STATGROUP_D3D12DescriptorHeap, );
+DECLARE_DWORD_COUNTER_STAT_EXTERN(TEXT("Used explicit sampler descriptors (per frame)"), STAT_ExplicitUsedSamplerDescriptors, STATGROUP_D3D12DescriptorHeap, );
+DECLARE_DWORD_COUNTER_STAT_EXTERN(TEXT("Used explicit view descriptors (per frame)"), STAT_ExplicitUsedViewDescriptors, STATGROUP_D3D12DescriptorHeap, );
 
 struct FD3D12GlobalStats
 {
diff --git a/Engine/Source/Runtime/D3D12RHI/Private/D3D12Util.cpp b/Engine/Source/Runtime/D3D12RHI/Private/D3D12Util.cpp
--- a/Engine/Source/Runtime/D3D12RHI/Private/D3D12Util.cpp
+++ b/Engine/Source/Runtime/D3D12RHI/Private/D3D12Util.cpp
@@ -1845,4 +1845,14 @@ DEFINE_STAT(STAT_GlobalViewHeapBlockAllocations);
 DEFINE_STAT(STAT_ViewOnlineDescriptorHeapMemory);
 DEFINE_STAT(STAT_SamplerOnlineDescriptorHeapMemory);
 
+DEFINE_STAT(STAT_ExplicitSamplerDescriptorHeaps);
+DEFINE_STAT(STAT_ExplicitSamplerDescriptors);
+
+DEFINE_STAT(STAT_ExplicitViewDescriptorHeaps);
+DEFINE_STAT(STAT_ExplicitViewDescriptors);
+
+DEFINE_STAT(STAT_ExplicitMaxUsedSamplerDescriptors);
+DEFINE_STAT(STAT_ExplicitUsedSamplerDescriptors);
+DEFINE_STAT(STAT_ExplicitUsedViewDescriptors);
+
 #undef LOCTEXT_NAMESPACE
