Compile out bindless support from D3D12 explicit descriptor cache if PLATFORM_SUPPORTS_BINDLESS_RENDERING=0

diff --git a/Engine/Source/Runtime/D3D12RHI/Private/D3D12ExplicitDescriptorCache.cpp b/Engine/Source/Runtime/D3D12RHI/Private/D3D12ExplicitDescriptorCache.cpp
--- a/Engine/Source/Runtime/D3D12RHI/Private/D3D12ExplicitDescriptorCache.cpp
+++ b/Engine/Source/Runtime/D3D12RHI/Private/D3D12ExplicitDescriptorCache.cpp
@@ -96,11 +96,11 @@ FD3D12ExplicitDescriptorHeapCache::Entry FD3D12ExplicitDescriptorHeapCache::Allo
 
 	ID3D12DescriptorHeap* D3D12Heap = nullptr;
 
-	const TCHAR* HeapName = Desc.Type == D3D12_DESCRIPTOR_HEAP_TYPE_CBV_SRV_UAV ? TEXT("RT View Heap") : TEXT("RT Sampler Heap");
+	const TCHAR* HeapName = Desc.Type == D3D12_DESCRIPTOR_HEAP_TYPE_CBV_SRV_UAV ? TEXT("Explicit View Heap") : TEXT("Explicit Sampler Heap");
 	UE_LOG(LogD3D12RHI, Log, TEXT("Creating %s with %d entries"), HeapName, NumDescriptors);
 
 	VERIFYD3D12RESULT(GetParentDevice()->GetDevice()->CreateDescriptorHeap(&Desc, IID_PPV_ARGS(&D3D12Heap)));
-	SetName(D3D12Heap, Desc.Type == D3D12_DESCRIPTOR_HEAP_TYPE_CBV_SRV_UAV ? L"RT View Heap" : L"RT Sampler Heap");
+	SetName(D3D12Heap, Desc.Type == D3D12_DESCRIPTOR_HEAP_TYPE_CBV_SRV_UAV ? L"Explicit View Heap" : L"Explicit Sampler Heap");
 
 	Result.NumDescriptors = NumDescriptors;
 	Result.Type = Type;
@@ -286,10 +286,15 @@ void FD3D12ExplicitDescriptorHeap::UpdateSyncPoint()
 
 void FD3D12ExplicitDescriptorCache::Init(uint32 NumViewDescriptors, uint32 NumSamplerDescriptors)
 {
+#if PLATFORM_SUPPORTS_BINDLESS_RENDERING
 	FD3D12BindlessDescriptorManager& BindlessManager = GetParentDevice()->GetBindlessDescriptorManager();
 
 	bBindlessViews = BindlessManager.HasHeap(ERHIDescriptorHeapType::Standard, ERHIBindlessConfiguration::RayTracingShaders);
 	bBindlessSamplers = BindlessManager.HasHeap(ERHIDescriptorHeapType::Sampler, ERHIBindlessConfiguration::RayTracingShaders);
+#else
+	const bool bBindlessViews = false;
+	const bool bBindlessSamplers = false;
+#endif
 
 	if (!bBindlessViews)
 	{
@@ -306,6 +311,11 @@ void FD3D12ExplicitDescriptorCache::UpdateSyncPoint()
 {
 	check(IsInRHIThread() || !IsRunningRHIInSeparateThread());
 
+#if !PLATFORM_SUPPORTS_BINDLESS_RENDERING
+	const bool bBindlessViews = false;
+	const bool bBindlessSamplers = false;
+#endif
+
 	if (!bBindlessViews)
 	{
 		ViewHeap.UpdateSyncPoint();
@@ -321,6 +331,7 @@ void FD3D12ExplicitDescriptorCache::SetDescriptorHeaps(FD3D12CommandContext& Com
 {
 	UpdateSyncPoint();
 
+#if PLATFORM_SUPPORTS_BINDLESS_RENDERING
 	check(bBindlessViews || ViewHeap.GetParentDevice() == CommandContext.GetParentDevice());
 	check(bBindlessSamplers || SamplerHeap.GetParentDevice() == CommandContext.GetParentDevice());
 
@@ -328,6 +339,13 @@ void FD3D12ExplicitDescriptorCache::SetDescriptorHeaps(FD3D12CommandContext& Com
 
 	ID3D12DescriptorHeap* ViewHeapToSet = bBindlessViews ? BindlessManager.GetHeap(ERHIDescriptorHeapType::Standard)->GetHeap() : ViewHeap.D3D12Heap;
 	ID3D12DescriptorHeap* SamplerHeapToSet = bBindlessSamplers ? BindlessManager.GetHeap(ERHIDescriptorHeapType::Sampler)->GetHeap() : SamplerHeap.D3D12Heap;
+#else
+	check(ViewHeap.GetParentDevice() == CommandContext.GetParentDevice());
+	check(SamplerHeap.GetParentDevice() == CommandContext.GetParentDevice());
+
+	ID3D12DescriptorHeap* ViewHeapToSet = ViewHeap.D3D12Heap;
+	ID3D12DescriptorHeap* SamplerHeapToSet = SamplerHeap.D3D12Heap;
+#endif
 
 	CommandContext.StateCache.GetDescriptorCache()->OverrideLastSetHeaps(ViewHeapToSet, SamplerHeapToSet);
 }
diff --git a/Engine/Source/Runtime/D3D12RHI/Private/D3D12ExplicitDescriptorCache.h b/Engine/Source/Runtime/D3D12RHI/Private/D3D12ExplicitDescriptorCache.h
--- a/Engine/Source/Runtime/D3D12RHI/Private/D3D12ExplicitDescriptorCache.h
+++ b/Engine/Source/Runtime/D3D12RHI/Private/D3D12ExplicitDescriptorCache.h
@@ -128,8 +128,10 @@ public:
 	FD3D12ExplicitDescriptorHeap ViewHeap;
 	FD3D12ExplicitDescriptorHeap SamplerHeap;
 
+#if PLATFORM_SUPPORTS_BINDLESS_RENDERING
 	bool bBindlessViews = false;
 	bool bBindlessSamplers = false;
+#endif
 
 	template<typename KeyType>
 	struct TIdentityHash
