Added static uniform buffer accessor to FD3D12CommandContext

diff --git a/Engine/Source/Runtime/D3D12RHI/Private/D3D12CommandContext.h b/Engine/Source/Runtime/D3D12RHI/Private/D3D12CommandContext.h
--- a/Engine/Source/Runtime/D3D12RHI/Private/D3D12CommandContext.h
+++ b/Engine/Source/Runtime/D3D12RHI/Private/D3D12CommandContext.h
@@ -739,6 +739,11 @@ public:
 
 	virtual void UpdateBuffer(FD3D12ResourceLocation* Dest, uint32 DestOffset, FD3D12ResourceLocation* Source, uint32 SourceOffset, uint32 NumBytes) final override;
 
+	inline const TArray<FRHIUniformBuffer*> GetStaticUniformBuffers() const
+	{
+		return StaticUniformBuffers;
+	}
+
 protected:
 
 	FD3D12CommandContext* GetContext(uint32 InGPUIndex) final override 
