Added a bindless config arg to the D3D12 explicit descriptor cache init (so bindless can be configured regardless of ray tracing)

diff --git a/Engine/Source/Runtime/D3D12RHI/Private/D3D12ExplicitDescriptorCache.cpp b/Engine/Source/Runtime/D3D12RHI/Private/D3D12ExplicitDescriptorCache.cpp
--- a/Engine/Source/Runtime/D3D12RHI/Private/D3D12ExplicitDescriptorCache.cpp
+++ b/Engine/Source/Runtime/D3D12RHI/Private/D3D12ExplicitDescriptorCache.cpp
@@ -284,13 +284,13 @@ void FD3D12ExplicitDescriptorHeap::UpdateSyncPoint()
 
 ///
 
-void FD3D12ExplicitDescriptorCache::Init(uint32 NumViewDescriptors, uint32 NumSamplerDescriptors)
+void FD3D12ExplicitDescriptorCache::Init(uint32 NumViewDescriptors, uint32 NumSamplerDescriptors, ERHIBindlessConfiguration BindlessConfig)
 {
 #if PLATFORM_SUPPORTS_BINDLESS_RENDERING
 	FD3D12BindlessDescriptorManager& BindlessManager = GetParentDevice()->GetBindlessDescriptorManager();
 
-	bBindlessViews = BindlessManager.HasHeap(ERHIDescriptorHeapType::Standard, ERHIBindlessConfiguration::RayTracingShaders);
-	bBindlessSamplers = BindlessManager.HasHeap(ERHIDescriptorHeapType::Sampler, ERHIBindlessConfiguration::RayTracingShaders);
+	bBindlessViews = BindlessManager.HasHeap(ERHIDescriptorHeapType::Standard, BindlessConfig);
+	bBindlessSamplers = BindlessManager.HasHeap(ERHIDescriptorHeapType::Sampler, BindlessConfig);
 #else
 	const bool bBindlessViews = false;
 	const bool bBindlessSamplers = false;
diff --git a/Engine/Source/Runtime/D3D12RHI/Private/D3D12ExplicitDescriptorCache.h b/Engine/Source/Runtime/D3D12RHI/Private/D3D12ExplicitDescriptorCache.h
--- a/Engine/Source/Runtime/D3D12RHI/Private/D3D12ExplicitDescriptorCache.h
+++ b/Engine/Source/Runtime/D3D12RHI/Private/D3D12ExplicitDescriptorCache.h
@@ -116,7 +116,7 @@ public:
 		WorkerData.SetNum(MaxWorkerCount);
 	}
 
-	void Init(uint32 NumViewDescriptors, uint32 NumSamplerDescriptors);
+	void Init(uint32 NumViewDescriptors, uint32 NumSamplerDescriptors, ERHIBindlessConfiguration BindlessConfig);
 
 	void UpdateSyncPoint();
 
diff --git a/Engine/Source/Runtime/D3D12RHI/Private/D3D12RayTracing.cpp b/Engine/Source/Runtime/D3D12RHI/Private/D3D12RayTracing.cpp
--- a/Engine/Source/Runtime/D3D12RHI/Private/D3D12RayTracing.cpp
+++ b/Engine/Source/Runtime/D3D12RHI/Private/D3D12RayTracing.cpp
@@ -1330,7 +1330,7 @@ public:
 			const uint32 NumSamplerDescriptors = D3D12_MAX_SHADER_VISIBLE_SAMPLER_HEAP_SIZE;
 
 			DescriptorCache = new FD3D12ExplicitDescriptorCache(Device, FD3D12RayTracingScene::MaxBindingWorkers);
-			DescriptorCache->Init(NumViewDescriptors, NumSamplerDescriptors);
+			DescriptorCache->Init(NumViewDescriptors, NumSamplerDescriptors, ERHIBindlessConfiguration::RayTracingShaders);
 		}
 
 		NumRayGenShaders = Initializer.NumRayGenShaders;
@@ -4624,7 +4624,7 @@ static void DispatchRays(FD3D12CommandContext& CommandContext,
 	else
 	{
 		FD3D12ExplicitDescriptorCache TransientDescriptorCache(CommandContext.GetParentDevice(), FD3D12RayTracingScene::MaxBindingWorkers);
-		TransientDescriptorCache.Init(MAX_SRVS + MAX_UAVS, MAX_SAMPLERS);
+		TransientDescriptorCache.Init(MAX_SRVS + MAX_UAVS, MAX_SAMPLERS, ERHIBindlessConfiguration::RayTracingShaders);
 		TransientDescriptorCache.SetDescriptorHeaps(CommandContext);
 		FD3D12RayTracingGlobalResourceBinder ResourceBinder(CommandContext, TransientDescriptorCache);
 		bResourcesBound = SetRayTracingShaderResources(CommandContext, RayGenShader, GlobalBindings, ResourceBinder);
