[On behalf of Daren Cheng]
Fix Notification Ensure

#rnx

diff --git a/Engine/Source/Runtime/D3D12RHI/Private/D3D12StateCache.cpp b/Engine/Source/Runtime/D3D12RHI/Private/D3D12StateCache.cpp
--- a/Engine/Source/Runtime/D3D12RHI/Private/D3D12StateCache.cpp
+++ b/Engine/Source/Runtime/D3D12RHI/Private/D3D12StateCache.cpp
@@ -268,11 +268,23 @@ void FD3D12StateCache::SetViewports(uint32 Count, const D3D12_VIEWPORT* const Vi
 
 static void ValidateScissorRect(const D3D12_VIEWPORT& Viewport, const D3D12_RECT& ScissorRect)
 {
-	ensure(ScissorRect.left   >= (LONG)Viewport.TopLeftX);
-	ensure(ScissorRect.top    >= (LONG)Viewport.TopLeftY);
-	ensure(ScissorRect.right  <= (LONG)Viewport.TopLeftX + (LONG)Viewport.Width);
-	ensure(ScissorRect.bottom <= (LONG)Viewport.TopLeftY + (LONG)Viewport.Height);
-	ensure(ScissorRect.left <= ScissorRect.right && ScissorRect.top <= ScissorRect.bottom);
+	bool bScissorRectValid = true;
+	bScissorRectValid = bScissorRectValid && ScissorRect.left >= (LONG)Viewport.TopLeftX;
+	bScissorRectValid = bScissorRectValid && ScissorRect.top >= (LONG)Viewport.TopLeftY;
+	bScissorRectValid = bScissorRectValid && ScissorRect.right <= (LONG)Viewport.TopLeftX + (LONG)Viewport.Width;
+	bScissorRectValid = bScissorRectValid && ScissorRect.bottom <= (LONG)Viewport.TopLeftY + (LONG)Viewport.Height;
+	bScissorRectValid = bScissorRectValid && ScissorRect.left <= ScissorRect.right && ScissorRect.top <= ScissorRect.bottom;
+
+	ensureMsgf(bScissorRectValid,
+		TEXT("Scissor invalid with current Viewport. Scissor: [left:%li, top:%li, right:%li, bottom:%li]. Viewport: [left:%li, top:%li, right:%li, bottom:%li]")
+			, ScissorRect.left
+			, ScissorRect.top
+			, ScissorRect.right
+			, ScissorRect.bottom
+			, (LONG)Viewport.TopLeftX
+			, (LONG)Viewport.TopLeftY
+			, (LONG)Viewport.TopLeftX + (LONG)Viewport.Width
+			, (LONG)Viewport.TopLeftY + (LONG)Viewport.Height);
 }
 
 void FD3D12StateCache::SetScissorRect(const D3D12_RECT& ScissorRect)
diff --git a/Engine/Source/Runtime/Interchange/Engine/Private/InterchangeManager.cpp b/Engine/Source/Runtime/Interchange/Engine/Private/InterchangeManager.cpp
--- a/Engine/Source/Runtime/Interchange/Engine/Private/InterchangeManager.cpp
+++ b/Engine/Source/Runtime/Interchange/Engine/Private/InterchangeManager.cpp
@@ -1125,8 +1125,15 @@ void UInterchangeManager::StartQueuedTasks(bool bCancelAllTasks /*= false*/)
 
 	bGCEndDelegateCancellAllTask = false;
 
-	auto UpdateNotification = [this]()
+	uint64 LastNotificationFrame = 0;
+	auto UpdateNotification = [this, &LastNotificationFrame]()
 	{
+		if (LastNotificationFrame == GFrameCounter)
+		{
+			return;
+		}
+		LastNotificationFrame = GFrameCounter;
+
 		if (Notification.IsValid())
 		{
 			int32 ImportTaskNumber = ImportTasks.Num() + QueueTaskCount;
