Deprecate IDXGISwapchainProvider::GetName and replace it with GetProviderName, because due to an unfortunate overshight, the previous method returned a TCHAR* instead of a const TCHAR*.

diff --git a/Engine/Source/Runtime/D3D12RHI/Private/Windows/WindowsD3D12Viewport.cpp b/Engine/Source/Runtime/D3D12RHI/Private/Windows/WindowsD3D12Viewport.cpp
--- a/Engine/Source/Runtime/D3D12RHI/Private/Windows/WindowsD3D12Viewport.cpp
+++ b/Engine/Source/Runtime/D3D12RHI/Private/Windows/WindowsD3D12Viewport.cpp
@@ -83,7 +83,7 @@ void FD3D12Viewport::Init()
 		static bool bCustomSwapchainLogged = false;
 		if (!bCustomSwapchainLogged)
 		{
-			UE_LOG(LogD3D12RHI, Log, TEXT("Found a custom swapchain provider: '%s'."), DXGISwapchainProvider->GetName());
+			UE_LOG(LogD3D12RHI, Log, TEXT("Found a custom swapchain provider: '%s'."), DXGISwapchainProvider->GetProviderName());
 			bCustomSwapchainLogged = true;
 		}
 	}
@@ -187,7 +187,7 @@ void FD3D12Viewport::Init()
 				UE_LOG(LogD3D12RHI, Warning, TEXT("\tWindows handle: 0x%x (IsWindow: %s)"), WindowHandle, IsWindow(WindowHandle) ? TEXT("true") : TEXT("false"));
 				UE_LOG(LogD3D12RHI, Warning, TEXT("\tFullscreen: %s"), bIsFullscreen ? TEXT("true") : TEXT("false"));
 				UE_LOG(LogD3D12RHI, Warning, TEXT("\tSwapchain flags: %d"), SwapChainFlags);
-				UE_LOG(LogD3D12RHI, Warning, TEXT("\tCustom swapchain provider: %s"), DXGISwapchainProvider ? DXGISwapchainProvider->GetName() : TEXT("none"));
+				UE_LOG(LogD3D12RHI, Warning, TEXT("\tCustom swapchain provider: %s"), DXGISwapchainProvider ? DXGISwapchainProvider->GetProviderName() : TEXT("none"));
 
 				VERIFYD3D12RESULT(hr);
 			}
diff --git a/Engine/Source/Runtime/RHI/Public/Windows/IDXGISwapchainProvider.h b/Engine/Source/Runtime/RHI/Public/Windows/IDXGISwapchainProvider.h
--- a/Engine/Source/Runtime/RHI/Public/Windows/IDXGISwapchainProvider.h
+++ b/Engine/Source/Runtime/RHI/Public/Windows/IDXGISwapchainProvider.h
@@ -16,7 +16,11 @@ public:
 	}
 
 	virtual bool SupportsRHI(ERHIInterfaceType RHIType) const = 0;
-	virtual TCHAR* GetName() const = 0;
+
+	UE_DEPRECATED(5.3, "IDXGISwapchainProvider::GetName is deprecated. Please use GetProviderName instead!")
+	virtual TCHAR* GetName() const { return nullptr; }
+
+	virtual const TCHAR* GetProviderName() const = 0;
 
 	virtual HRESULT CreateSwapChainForHwnd(IDXGIFactory2* pFactory, IUnknown* pDevice, HWND hWnd, const DXGI_SWAP_CHAIN_DESC1* pDesc, const DXGI_SWAP_CHAIN_FULLSCREEN_DESC* pFulScreenDesc, IDXGIOutput* pRestrictToOutput, IDXGISwapChain1** ppSwapChain) = 0;
 	virtual HRESULT CreateSwapChain(IDXGIFactory* pFactory, IUnknown* pDevice, DXGI_SWAP_CHAIN_DESC* pDesc, IDXGISwapChain** ppSwapChain) = 0;
diff --git a/Engine/Source/Runtime/Windows/D3D11RHI/Private/Windows/WindowsD3D11Viewport.cpp b/Engine/Source/Runtime/Windows/D3D11RHI/Private/Windows/WindowsD3D11Viewport.cpp
--- a/Engine/Source/Runtime/Windows/D3D11RHI/Private/Windows/WindowsD3D11Viewport.cpp
+++ b/Engine/Source/Runtime/Windows/D3D11RHI/Private/Windows/WindowsD3D11Viewport.cpp
@@ -77,7 +77,7 @@ FD3D11Viewport::FD3D11Viewport(FD3D11DynamicRHI* InD3DRHI,HWND InWindowHandle,ui
 		static bool bCustomSwapchainLogged = false;
 		if (!bCustomSwapchainLogged)
 		{
-			UE_LOG(LogD3D11RHI, Log, TEXT("Found a custom swapchain provider: '%s'."), DXGISwapchainProvider->GetName());
+			UE_LOG(LogD3D11RHI, Log, TEXT("Found a custom swapchain provider: '%s'."), DXGISwapchainProvider->GetProviderName());
 			bCustomSwapchainLogged = true;
 		}
 	}
