PC D3D12 shader bundle dispatch emulation functional

diff --git a/Engine/Source/Runtime/D3D12RHI/Private/D3D12Resources.h b/Engine/Source/Runtime/D3D12RHI/Private/D3D12Resources.h
--- a/Engine/Source/Runtime/D3D12RHI/Private/D3D12Resources.h
+++ b/Engine/Source/Runtime/D3D12RHI/Private/D3D12Resources.h
@@ -992,9 +992,8 @@ public:
 
 	virtual void CalcDispatchBufferSizes(uint32& OutDataBufferSize, uint32& OutExecutionBufferSize) const override
 	{
-		checkNoEntry();
-		OutDataBufferSize = 0u;
-		OutExecutionBufferSize = 0u;
+		OutDataBufferSize = 4u;
+		OutExecutionBufferSize = 4u;
 	}
 };
 
diff --git a/Engine/Source/Runtime/D3D12RHI/Private/Windows/WindowsD3D12Device.cpp b/Engine/Source/Runtime/D3D12RHI/Private/Windows/WindowsD3D12Device.cpp
--- a/Engine/Source/Runtime/D3D12RHI/Private/Windows/WindowsD3D12Device.cpp
+++ b/Engine/Source/Runtime/D3D12RHI/Private/Windows/WindowsD3D12Device.cpp
@@ -1690,7 +1690,7 @@ void FD3D12DynamicRHI::Init()
 	GRHISupportsTextureStreaming = true;
 	GRHISupportsFirstInstance = true;
 
-	GRHISupportsDispatchShaderBundle = false; // TODO: Implement
+	GRHISupportsDispatchShaderBundle = true;
 
 	// Indicate that the RHI needs to use the engine's deferred deletion queue.
 	GRHINeedsExtraDeletionLatency = true;
