Align staging buffer mip size to the block size. The Agility SDK doesn't require this, but older, OS-provided D3D runtimes on machines that don't support the Agility SDK do, so using non-aligned mips would crash with E_INVALIDARG.

diff --git a/Engine/Source/Runtime/D3D12RHI/Private/D3D12Texture.cpp b/Engine/Source/Runtime/D3D12RHI/Private/D3D12Texture.cpp
--- a/Engine/Source/Runtime/D3D12RHI/Private/D3D12Texture.cpp
+++ b/Engine/Source/Runtime/D3D12RHI/Private/D3D12Texture.cpp
@@ -2048,8 +2048,8 @@ void FD3D12Texture::UnlockInternal(class FRHICommandListImmediate* RHICmdList, F
 			const D3D12_RESOURCE_DESC& ResourceDesc = Resource->GetDesc();
 			D3D12_SUBRESOURCE_FOOTPRINT BufferPitchDesc;
 			BufferPitchDesc.Depth = 1;
-			BufferPitchDesc.Height = MipSizeY;
-			BufferPitchDesc.Width = MipSizeX;
+			BufferPitchDesc.Height = Align(MipSizeY, BlockSizeY);
+			BufferPitchDesc.Width = Align(MipSizeX, BlockSizeX);
 			BufferPitchDesc.Format = ResourceDesc.Format;
 			BufferPitchDesc.RowPitch = LockedResource->LockedPitch;
 			check(BufferPitchDesc.RowPitch % FD3D12_TEXTURE_DATA_PITCH_ALIGNMENT == 0);
