When checking for BCn<->UINT aliasing support, we also need to check EnhancedBarriersSupported.

diff --git a/Engine/Source/Runtime/D3D12RHI/Private/Windows/WindowsD3D12Device.cpp b/Engine/Source/Runtime/D3D12RHI/Private/Windows/WindowsD3D12Device.cpp
--- a/Engine/Source/Runtime/D3D12RHI/Private/Windows/WindowsD3D12Device.cpp
+++ b/Engine/Source/Runtime/D3D12RHI/Private/Windows/WindowsD3D12Device.cpp
@@ -1732,7 +1732,10 @@ void FD3D12DynamicRHI::Init()
 	D3D12_FEATURE_DATA_D3D12_OPTIONS12 Options12{};
 	if (SUCCEEDED(GetAdapter().GetD3DDevice()->CheckFeatureSupport(D3D12_FEATURE_D3D12_OPTIONS12, &Options12, sizeof(Options12))))
 	{
-		GRHIGlobals.SupportsUAVFormatAliasing = (Options12.RelaxedFormatCastingSupported != 0)
+		GRHIGlobals.SupportsUAVFormatAliasing =
+			(Options12.RelaxedFormatCastingSupported != 0)
+			// Our BCn casting requires enhanced barrier support
+			&& (Options12.EnhancedBarriersSupported != 0)
 			// We require ID3D12Device12 for GetResourceAllocationInfo3
 			&& GetAdapter().GetD3DDevice12() != nullptr
 #if PLATFORM_WINDOWS
