Moved RHICreateTexture to support non-immediate command lists.

diff --git a/Engine/Plugins/Experimental/AVCodecs/AVCodecsCore/Source/AVCodecsCoreRHI/Private/Video/Resources/VideoResourceRHI.cpp b/Engine/Plugins/Experimental/AVCodecs/AVCodecsCore/Source/AVCodecsCoreRHI/Private/Video/Resources/VideoResourceRHI.cpp
--- a/Engine/Plugins/Experimental/AVCodecs/AVCodecsCore/Source/AVCodecsCoreRHI/Private/Video/Resources/VideoResourceRHI.cpp
+++ b/Engine/Plugins/Experimental/AVCodecs/AVCodecsCore/Source/AVCodecsCoreRHI/Private/Video/Resources/VideoResourceRHI.cpp
@@ -162,7 +162,7 @@ TSharedPtr<FVideoResourceRHI> FVideoResourceRHI::Create(TSharedPtr<FAVDevice> co
 		}
 
 //TODO-TE THIS IS THE REAL DEAL?
-		return MakeShareable(new FVideoResourceRHI(Device.ToSharedRef(), { GDynamicRHI->RHICreateTexture(TextureDesc), nullptr, 0 }, Descriptor));
+		return MakeShareable(new FVideoResourceRHI(Device.ToSharedRef(), { RHICreateTexture(TextureDesc), nullptr, 0 }, Descriptor));
 	}
 
 	return nullptr;
diff --git a/Engine/Plugins/Media/PixelCapture/Source/PixelCapture/Private/PixelCaptureCapturerI420ToRHI.cpp b/Engine/Plugins/Media/PixelCapture/Source/PixelCapture/Private/PixelCaptureCapturerI420ToRHI.cpp
--- a/Engine/Plugins/Media/PixelCapture/Source/PixelCapture/Private/PixelCaptureCapturerI420ToRHI.cpp
+++ b/Engine/Plugins/Media/PixelCapture/Source/PixelCapture/Private/PixelCaptureCapturerI420ToRHI.cpp
@@ -56,7 +56,7 @@ IPixelCaptureOutputFrame *FPixelCaptureCapturerI420ToRHI::CreateOutputBuffer(int
 		TextureDesc.AddFlags(ETextureCreateFlags::Shared);
 	}
 
-	return new FPixelCaptureOutputFrameRHI(GDynamicRHI->RHICreateTexture(TextureDesc));
+	return new FPixelCaptureOutputFrameRHI(RHICreateTexture(TextureDesc));
 }
 
 void FPixelCaptureCapturerI420ToRHI::BeginProcess(const IPixelCaptureInputFrame &InputFrame, IPixelCaptureOutputFrame *OutputBuffer)
diff --git a/Engine/Plugins/Media/PixelCapture/Source/PixelCapture/Private/PixelCaptureCapturerNV12ToRHI.cpp b/Engine/Plugins/Media/PixelCapture/Source/PixelCapture/Private/PixelCaptureCapturerNV12ToRHI.cpp
--- a/Engine/Plugins/Media/PixelCapture/Source/PixelCapture/Private/PixelCaptureCapturerNV12ToRHI.cpp
+++ b/Engine/Plugins/Media/PixelCapture/Source/PixelCapture/Private/PixelCaptureCapturerNV12ToRHI.cpp
@@ -55,7 +55,7 @@ IPixelCaptureOutputFrame* FPixelCaptureCapturerNV12ToRHI::CreateOutputBuffer(int
 		TextureDesc.AddFlags(ETextureCreateFlags::Shared);
 	}
 
-	return new FPixelCaptureOutputFrameRHI(GDynamicRHI->RHICreateTexture(TextureDesc), { InputWidth, InputHeight, EPixelFormat::PF_NV12 });
+	return new FPixelCaptureOutputFrameRHI(RHICreateTexture(TextureDesc), { InputWidth, InputHeight, EPixelFormat::PF_NV12 });
 }
 
 void FPixelCaptureCapturerNV12ToRHI::BeginProcess(const IPixelCaptureInputFrame& InputFrame, IPixelCaptureOutputFrame* OutputBuffer)
diff --git a/Engine/Plugins/Media/PixelCapture/Source/PixelCapture/Private/PixelCaptureCapturerRHI.cpp b/Engine/Plugins/Media/PixelCapture/Source/PixelCapture/Private/PixelCaptureCapturerRHI.cpp
--- a/Engine/Plugins/Media/PixelCapture/Source/PixelCapture/Private/PixelCaptureCapturerRHI.cpp
+++ b/Engine/Plugins/Media/PixelCapture/Source/PixelCapture/Private/PixelCaptureCapturerRHI.cpp
@@ -41,7 +41,7 @@ IPixelCaptureOutputFrame* FPixelCaptureCapturerRHI::CreateOutputBuffer(int32 Inp
 		TextureDesc.AddFlags(ETextureCreateFlags::Shared);
 	}
 
-	return new FPixelCaptureOutputFrameRHI(GDynamicRHI->RHICreateTexture(TextureDesc));
+	return new FPixelCaptureOutputFrameRHI(RHICreateTexture(TextureDesc));
 }
 
 void FPixelCaptureCapturerRHI::BeginProcess(const IPixelCaptureInputFrame& InputFrame, IPixelCaptureOutputFrame* OutputBuffer)
diff --git a/Engine/Plugins/Media/PixelCapture/Source/PixelCapture/Private/PixelCaptureCapturerRHIRDG.cpp b/Engine/Plugins/Media/PixelCapture/Source/PixelCapture/Private/PixelCaptureCapturerRHIRDG.cpp
--- a/Engine/Plugins/Media/PixelCapture/Source/PixelCapture/Private/PixelCaptureCapturerRHIRDG.cpp
+++ b/Engine/Plugins/Media/PixelCapture/Source/PixelCapture/Private/PixelCaptureCapturerRHIRDG.cpp
@@ -39,7 +39,7 @@ IPixelCaptureOutputFrame* FPixelCaptureCapturerRHIRDG::CreateOutputBuffer(int32
 		TextureDesc.AddFlags(ETextureCreateFlags::Shared);
 	}
 
-	return new FPixelCaptureOutputFrameRHI(GDynamicRHI->RHICreateTexture(TextureDesc));
+	return new FPixelCaptureOutputFrameRHI(RHICreateTexture(TextureDesc));
 }
 
 void FPixelCaptureCapturerRHIRDG::BeginProcess(const IPixelCaptureInputFrame& InputFrame, IPixelCaptureOutputFrame* OutputBuffer)
diff --git a/Engine/Plugins/Media/PixelCapture/Source/PixelCapture/Private/PixelCaptureCapturerRHIToI420CPU.cpp b/Engine/Plugins/Media/PixelCapture/Source/PixelCapture/Private/PixelCaptureCapturerRHIToI420CPU.cpp
--- a/Engine/Plugins/Media/PixelCapture/Source/PixelCapture/Private/PixelCaptureCapturerRHIToI420CPU.cpp
+++ b/Engine/Plugins/Media/PixelCapture/Source/PixelCapture/Private/PixelCaptureCapturerRHIToI420CPU.cpp
@@ -44,7 +44,7 @@ void FPixelCaptureCapturerRHIToI420CPU::Initialize(int32 InputWidth, int32 Input
 		TextureDesc.AddFlags(ETextureCreateFlags::Shared);
 	}
 
-	StagingTexture = GDynamicRHI->RHICreateTexture(TextureDesc);
+	StagingTexture = RHICreateTexture(TextureDesc);
 
 	FRHITextureCreateDesc ReadbackDesc =
 		FRHITextureCreateDesc::Create2D(TEXT("FPixelCaptureCapturerRHIToI420CPU ReadbackTexture"), Width, Height, EPixelFormat::PF_B8G8R8A8)
@@ -53,7 +53,7 @@ void FPixelCaptureCapturerRHIToI420CPU::Initialize(int32 InputWidth, int32 Input
 			.SetInitialState(ERHIAccess::CPURead)
 			.DetermineInititialState();
 
-	ReadbackTexture = GDynamicRHI->RHICreateTexture(ReadbackDesc);
+	ReadbackTexture = RHICreateTexture(ReadbackDesc);
 
 	int32 BufferWidth = 0, BufferHeight = 0;
 	GDynamicRHI->RHIMapStagingSurface(ReadbackTexture, nullptr, ResultsBuffer, BufferWidth, BufferHeight);
diff --git a/Engine/Plugins/Media/PixelStreaming/Source/PixelStreaming/Private/UtilsRender.h b/Engine/Plugins/Media/PixelStreaming/Source/PixelStreaming/Private/UtilsRender.h
--- a/Engine/Plugins/Media/PixelStreaming/Source/PixelStreaming/Private/UtilsRender.h
+++ b/Engine/Plugins/Media/PixelStreaming/Source/PixelStreaming/Private/UtilsRender.h
@@ -35,7 +35,7 @@ namespace UE::PixelStreaming
 			TextureDesc.AddFlags(ETextureCreateFlags::Shared);
 		}
 
-		return GDynamicRHI->RHICreateTexture(TextureDesc);
+		return RHICreateTexture(TextureDesc);
 	}
 
 	/*
diff --git a/Engine/Source/Runtime/Apple/MetalRHI/Private/MetalDynamicRHI.h b/Engine/Source/Runtime/Apple/MetalRHI/Private/MetalDynamicRHI.h
--- a/Engine/Source/Runtime/Apple/MetalRHI/Private/MetalDynamicRHI.h
+++ b/Engine/Source/Runtime/Apple/MetalRHI/Private/MetalDynamicRHI.h
@@ -71,7 +71,7 @@ public:
 	virtual uint64 RHIGetMinimumAlignmentForBufferBackedSRV(EPixelFormat Format) final override;
 	virtual void RHIGetTextureMemoryStats(FTextureMemoryStats& OutStats) final override;
 	virtual bool RHIGetTextureMemoryVisualizeData(FColor* TextureData,int32 SizeX,int32 SizeY,int32 Pitch,int32 PixelSize) final override;
-	virtual FTextureRHIRef RHICreateTexture(const FRHITextureCreateDesc& CreateDesc) final override;
+	virtual FTextureRHIRef RHICreateTexture(FRHICommandListBase& RHICmdList, const FRHITextureCreateDesc& CreateDesc) final override;
 	virtual FTexture2DRHIRef RHIAsyncCreateTexture2D(uint32 SizeX, uint32 SizeY, uint8 Format, uint32 NumMips, ETextureCreateFlags Flags, ERHIAccess InResourceState, void** InitialMipData, uint32 NumInitialMips, FGraphEventRef& OutCompletionEvent) final override;
 	virtual void RHIGenerateMips(FRHITexture* Texture) final override;
 	virtual uint32 RHIComputeMemorySize(FRHITexture* TextureRHI) final override;
@@ -129,8 +129,6 @@ public:
 	virtual FUpdateTexture3DData RHIBeginUpdateTexture3D(FRHICommandListBase& RHICmdList, FRHITexture3D* Texture, uint32 MipIndex, const struct FUpdateTextureRegion3D& UpdateRegion) final override;
 	virtual void RHIEndUpdateTexture3D(FRHICommandListBase& RHICmdList, FUpdateTexture3DData& UpdateData) final override;
 
-	virtual FTextureRHIRef RHICreateTexture_RenderThread(class FRHICommandListImmediate& RHICmdList, const FRHITextureCreateDesc& CreateDesc) final override;
-
 	virtual void RHICreateTransition(FRHITransition* Transition, const FRHITransitionCreateInfo& CreateInfo) final override;
 	virtual void RHIReleaseTransition(FRHITransition* Transition) final override;
 
diff --git a/Engine/Source/Runtime/Apple/MetalRHI/Private/MetalTexture.cpp b/Engine/Source/Runtime/Apple/MetalRHI/Private/MetalTexture.cpp
--- a/Engine/Source/Runtime/Apple/MetalRHI/Private/MetalTexture.cpp
+++ b/Engine/Source/Runtime/Apple/MetalRHI/Private/MetalTexture.cpp
@@ -473,7 +473,7 @@ FMetalTextureCreateDesc::FMetalTextureCreateDesc(FRHITextureCreateDesc const& In
 	}
 }
 
-FMetalSurface::FMetalSurface(FMetalTextureCreateDesc const& CreateDesc)
+FMetalSurface::FMetalSurface(FRHICommandListBase* RHICmdList, FMetalTextureCreateDesc const& CreateDesc)
 	: FRHITexture       (CreateDesc)
 	, FormatKey         (CreateDesc.FormatKey)
 	, Texture           (nil)
@@ -617,16 +617,16 @@ FMetalSurface::FMetalSurface(FMetalTextureCreateDesc const& CreateDesc)
 		// Regular texture has some bulk data to handle
 		UE_LOG(LogMetal, Display, TEXT("Got a bulk data texture, with %d mips"), CreateDesc.NumMips);
 		checkf(CreateDesc.NumMips == 1, TEXT("Only handling bulk data with 1 mip and 1 array length"));
+		check(RHICmdList);
 
-		check(IsInRenderingThread());
-		FRHICommandListImmediate& RHICmdList = FRHICommandListExecutor::GetImmediateCommandList();
+		FRHICommandListImmediate& RHICmdListImmediate = RHICmdList->GetAsImmediate();
 
 		// lock, copy, unlock
 		uint32 Stride;
-		void* LockedData = FMetalDynamicRHI::Get().LockTexture2D_RenderThread(RHICmdList, this, 0, RLM_WriteOnly, Stride, false);
+		void* LockedData = FMetalDynamicRHI::Get().LockTexture2D_RenderThread(RHICmdListImmediate, this, 0, RLM_WriteOnly, Stride, false);
 		check(LockedData);
 		FMemory::Memcpy(LockedData, BulkData->GetResourceBulkData(), BulkData->GetResourceBulkDataSize());
-		FMetalDynamicRHI::Get().UnlockTexture2D_RenderThread(RHICmdList, this, 0, false);
+		FMetalDynamicRHI::Get().UnlockTexture2D_RenderThread(RHICmdListImmediate, this, 0, false);
 
 		// bulk data can be unloaded now
 		BulkData->Discard();
@@ -1350,17 +1350,10 @@ uint32 FMetalDynamicRHI::RHIComputeMemorySize(FRHITexture* TextureRHI)
  2D texture support.
  -----------------------------------------------------------------------------*/
 
-FTextureRHIRef FMetalDynamicRHI::RHICreateTexture_RenderThread(class FRHICommandListImmediate& RHICmdList, const FRHITextureCreateDesc& CreateDesc)
-{
-	@autoreleasepool{
-		return this->RHICreateTexture(CreateDesc);
-	}
-}
-
-FTextureRHIRef FMetalDynamicRHI::RHICreateTexture(const FRHITextureCreateDesc& CreateDesc)
+FTextureRHIRef FMetalDynamicRHI::RHICreateTexture(FRHICommandListBase& RHICmdList, const FRHITextureCreateDesc& CreateDesc)
 {
 	@autoreleasepool{
-		return new FMetalSurface(CreateDesc);
+		return new FMetalSurface(&RHICmdList, CreateDesc);
 	}
 }
 
@@ -1417,7 +1410,7 @@ FTexture2DRHIRef FMetalDynamicRHI::RHIAsyncReallocateTexture2D(FRHITexture2D* Ol
 			TEXT("RHIAsyncReallocateTexture2D")
 		);
 		
-		FMetalSurface* NewTexture = new FMetalSurface(CreateDesc);
+		FMetalSurface* NewTexture = new FMetalSurface(&RHICmdList, CreateDesc);
 
 		// Copy shared mips
 		RHICmdList.EnqueueLambda([this, OldTexture, NewSizeX, NewSizeY, NewTexture, RequestStatus](FRHICommandListImmediate& RHICmdList)
diff --git a/Engine/Source/Runtime/Apple/MetalRHI/Private/MetalViewport.cpp b/Engine/Source/Runtime/Apple/MetalRHI/Private/MetalViewport.cpp
--- a/Engine/Source/Runtime/Apple/MetalRHI/Private/MetalViewport.cpp
+++ b/Engine/Source/Runtime/Apple/MetalRHI/Private/MetalViewport.cpp
@@ -260,12 +260,12 @@ void FMetalViewport::Resize(uint32 InSizeX, uint32 InSizeY, bool bInIsFullscreen
 		
 		CreateDesc.SetInitialState(RHIGetDefaultResourceState(CreateDesc.Flags, false));
 
-		NewBackBuffer = new FMetalSurface(CreateDesc);
+		NewBackBuffer = new FMetalSurface(nullptr, CreateDesc);
 		NewBackBuffer->Viewport = this;
 
         if (GMetalSupportsIntermediateBackBuffer && GMetalSeparatePresentThread)
         {
-            DoubleBuffer = new FMetalSurface(CreateDesc);
+            DoubleBuffer = new FMetalSurface(nullptr, CreateDesc);
             DoubleBuffer->Viewport = this;
         }
 
diff --git a/Engine/Source/Runtime/Apple/MetalRHI/Public/MetalResources.h b/Engine/Source/Runtime/Apple/MetalRHI/Public/MetalResources.h
--- a/Engine/Source/Runtime/Apple/MetalRHI/Public/MetalResources.h
+++ b/Engine/Source/Runtime/Apple/MetalRHI/Public/MetalResources.h
@@ -177,7 +177,7 @@ public:
 	/** 
 	 * Constructor that will create Texture and Color/DepthBuffers as needed
 	 */
-	FMetalSurface(FMetalTextureCreateDesc const& CreateDesc);
+	FMetalSurface(FRHICommandListBase* RHICmdList, FMetalTextureCreateDesc const& CreateDesc);
 	
 	/**
 	 * Destructor
diff --git a/Engine/Source/Runtime/D3D12RHI/Private/D3D12RHIPrivate.h b/Engine/Source/Runtime/D3D12RHI/Private/D3D12RHIPrivate.h
--- a/Engine/Source/Runtime/D3D12RHI/Private/D3D12RHIPrivate.h
+++ b/Engine/Source/Runtime/D3D12RHI/Private/D3D12RHIPrivate.h
@@ -309,7 +309,7 @@ public:
 	virtual void RHIGetTextureMemoryStats(FTextureMemoryStats& OutStats) final override;
 	virtual bool RHIGetTextureMemoryVisualizeData(FColor* TextureData, int32 SizeX, int32 SizeY, int32 Pitch, int32 PixelSize) final override;
 	virtual FTextureRHIRef RHIAsyncCreateTexture2D(uint32 SizeX, uint32 SizeY, uint8 Format, uint32 NumMips, ETextureCreateFlags Flags, ERHIAccess InResourceState, void** InitialMipData, uint32 NumInitialMips, FGraphEventRef& OutCompletionEvent) final override;
-	virtual FTextureRHIRef RHICreateTexture(const FRHITextureCreateDesc& CreateDesc) override;
+	virtual FTextureRHIRef RHICreateTexture(FRHICommandListBase& RHICmdList, const FRHITextureCreateDesc& CreateDesc) override;
 	virtual uint32 RHIComputeMemorySize(FRHITexture* TextureRHI) final override;
 	virtual FTexture2DRHIRef RHIAsyncReallocateTexture2D(FRHITexture2D* Texture2D, int32 NewMipCount, int32 NewSizeX, int32 NewSizeY, FThreadSafeCounter* RequestStatus) override;
 	virtual ETextureReallocationStatus RHIFinalizeAsyncReallocateTexture2D(FRHITexture2D* Texture2D, bool bBlockUntilCompleted) final override;
@@ -433,8 +433,6 @@ public:
 		return RHICancelAsyncReallocateTexture2D(Texture2D, bBlockUntilCompleted);
 	}
 
-	virtual FTextureRHIRef RHICreateTexture_RenderThread(class FRHICommandListImmediate& RHICmdList, const FRHITextureCreateDesc& CreateDesc) override;
-
 	void RHICalibrateTimers() override;
 
 #if D3D12_RHI_RAYTRACING
@@ -540,7 +538,7 @@ public:
 
 	virtual FD3D12ResourceDesc GetResourceDesc(const FRHITextureDesc& CreateInfo) const;
 
-	virtual FD3D12Texture* CreateD3D12Texture(const FRHITextureCreateDesc& CreateDesc, class FRHICommandListImmediate* RHICmdList, ID3D12ResourceAllocator* ResourceAllocator = nullptr);
+	virtual FD3D12Texture* CreateD3D12Texture(const FRHITextureCreateDesc& CreateDesc, FRHICommandListBase* RHICmdList, ID3D12ResourceAllocator* ResourceAllocator = nullptr);
 	FD3D12Buffer* CreateD3D12Buffer(class FRHICommandListBase* RHICmdList, FRHIBufferDesc const& BufferDesc, ERHIAccess ResourceState, FRHIResourceCreateInfo& CreateInfo, ID3D12ResourceAllocator* ResourceAllocator = nullptr);
 	virtual FD3D12Texture* CreateNewD3D12Texture(const FRHITextureCreateDesc& CreateDesc, class FD3D12Device* Device);
 
diff --git a/Engine/Source/Runtime/D3D12RHI/Private/D3D12Texture.cpp b/Engine/Source/Runtime/D3D12RHI/Private/D3D12Texture.cpp
--- a/Engine/Source/Runtime/D3D12RHI/Private/D3D12Texture.cpp
+++ b/Engine/Source/Runtime/D3D12RHI/Private/D3D12Texture.cpp
@@ -56,7 +56,7 @@ extern int32 GD3D12BindResourceLabels;
 // Texture Commands
 ///////////////////////////////////////////////////////////////////////////////////////////
 
-static bool ShouldDeferCmdListOperation(FRHICommandList* RHICmdList)
+static bool ShouldDeferCmdListOperation(FRHICommandListBase* RHICmdList)
 {
 	if (RHICmdList == nullptr)
 	{
@@ -132,12 +132,7 @@ struct FD3D12RHICommandInitializeTexture final : public FRHICommand<FD3D12RHICom
 		FD3D12ResourceLocation::TransferOwnership(SrcResourceLoc, InSrcResourceLoc);
 	}
 
-	void Execute(FRHICommandListBase& /* unused */)
-	{
-		ExecuteNoCmdList();
-	}
-
-	void ExecuteNoCmdList()
+	void Execute(FRHICommandListBase& RHICmdList)
 	{
 		size_t MemSize = NumSubresources * (sizeof(D3D12_PLACED_SUBRESOURCE_FOOTPRINT) + sizeof(UINT) + sizeof(UINT64));
 		const bool bAllocateOnStack = (MemSize < 4096);
@@ -161,8 +156,7 @@ struct FD3D12RHICommandInitializeTexture final : public FRHICommand<FD3D12RHICom
 		{
 			FD3D12Device* Device = CurrentTexture.GetParentDevice();
 			FD3D12Resource* Resource = CurrentTexture.GetResource();
-
-			FD3D12CommandContext& Context = Device->GetDefaultCommandContext();
+			FD3D12CommandContext& Context = FD3D12CommandContext::Get(RHICmdList, Device->GetGPUIndex());
 
 			// resource should be in copy dest already, because it's created like that, so no transition required here
 
@@ -813,7 +807,7 @@ FD3D12Texture* FD3D12DynamicRHI::CreateNewD3D12Texture(const FRHITextureCreateDe
 	return new FD3D12Texture(CreateDesc, Device);
 }
 
-FD3D12Texture* FD3D12DynamicRHI::CreateD3D12Texture(const FRHITextureCreateDesc& InCreateDesc, class FRHICommandListImmediate* RHICmdList, ID3D12ResourceAllocator* ResourceAllocator)
+FD3D12Texture* FD3D12DynamicRHI::CreateD3D12Texture(const FRHITextureCreateDesc& InCreateDesc, class FRHICommandListBase* RHICmdList, ID3D12ResourceAllocator* ResourceAllocator)
 {
 #if PLATFORM_WINDOWS
 	TRACE_CPUPROFILER_EVENT_SCOPE(D3D12RHI::CreateD3D12Texture);
@@ -938,7 +932,8 @@ FD3D12Texture* FD3D12DynamicRHI::CreateD3D12Texture(const FRHITextureCreateDesc&
 	// Initialize if data is given
 	if (CreateDesc.BulkData != nullptr)
 	{
-		D3D12TextureOut->InitializeTextureData(RHICmdList, CreateDesc, InitialState);
+		check(RHICmdList);
+		D3D12TextureOut->InitializeTextureData(*RHICmdList, CreateDesc, InitialState);
 		CreateDesc.BulkData->Discard();
 	}
 	return D3D12TextureOut;
@@ -949,19 +944,10 @@ FD3D12Texture* FD3D12DynamicRHI::CreateD3D12Texture(const FRHITextureCreateDesc&
 }
 
 
-FTextureRHIRef FD3D12DynamicRHI::RHICreateTexture(const FRHITextureCreateDesc& CreateDesc)
-{
-	FRHICommandListImmediate* RHIImmediateCmdList = nullptr;
-	ID3D12ResourceAllocator* ResourceAllocator = nullptr;
-	return CreateD3D12Texture(CreateDesc, RHIImmediateCmdList, ResourceAllocator);
-}
-
-
-FTextureRHIRef FD3D12DynamicRHI::RHICreateTexture_RenderThread(class FRHICommandListImmediate& RHICmdList, const FRHITextureCreateDesc& CreateDesc)
+FTextureRHIRef FD3D12DynamicRHI::RHICreateTexture(FRHICommandListBase& RHICmdList, const FRHITextureCreateDesc& CreateDesc)
 {
-	FRHICommandListImmediate* RHIImmediateCmdList = &RHICmdList;
 	ID3D12ResourceAllocator* ResourceAllocator = nullptr;
-	return CreateD3D12Texture(CreateDesc, RHIImmediateCmdList, ResourceAllocator);
+	return CreateD3D12Texture(CreateDesc, &RHICmdList, ResourceAllocator);
 }
 
 void FD3D12DynamicRHI::RHIUpdateTextureReference(FRHICommandListBase& RHICmdList, FRHITextureReference* TextureRef, FRHITexture* InNewTexture)
@@ -1932,7 +1918,7 @@ void FD3D12Texture::CopyTextureRegion(uint32 DestX, uint32 DestY, uint32 DestZ,
 	DefaultContext.UpdateResidency(GetResource());
 }
 
-void FD3D12Texture::InitializeTextureData(FRHICommandListImmediate* RHICmdList, const FRHITextureCreateDesc& CreateDesc, D3D12_RESOURCE_STATES DestinationState)
+void FD3D12Texture::InitializeTextureData(FRHICommandListBase& RHICmdList, const FRHITextureCreateDesc& CreateDesc, D3D12_RESOURCE_STATES DestinationState)
 {
 	// each mip of each array slice counts as a subresource
 	uint16 ArraySize = CreateDesc.IsTextureArray() ? CreateDesc.ArraySize : 1;
@@ -1991,14 +1977,14 @@ void FD3D12Texture::InitializeTextureData(FRHICommandListImmediate* RHICmdList,
 
 	check(SrcData == (uint8*)CreateDesc.BulkData->GetResourceBulkData() + CreateDesc.BulkData->GetResourceBulkDataSize());
 
-	if (ShouldDeferCmdListOperation(RHICmdList))
+	if (RHICmdList.IsTopOfPipe())
 	{
-		ALLOC_COMMAND_CL(*RHICmdList, FD3D12RHICommandInitializeTexture)(this, SrcResourceLoc, NumSubresources, DestinationState);
+		ALLOC_COMMAND_CL(RHICmdList, FD3D12RHICommandInitializeTexture)(this, SrcResourceLoc, NumSubresources, DestinationState);
 	}
 	else
 	{
 		FD3D12RHICommandInitializeTexture Command(this, SrcResourceLoc, NumSubresources, DestinationState);
-		Command.ExecuteNoCmdList();
+		Command.Execute(RHICmdList);
 	}
 
 	if (!bAllocateOnStack)
diff --git a/Engine/Source/Runtime/D3D12RHI/Private/D3D12Texture.h b/Engine/Source/Runtime/D3D12RHI/Private/D3D12Texture.h
--- a/Engine/Source/Runtime/D3D12RHI/Private/D3D12Texture.h
+++ b/Engine/Source/Runtime/D3D12RHI/Private/D3D12Texture.h
@@ -74,7 +74,7 @@ public:
 #endif // #if PLATFORM_REQUIRES_TYPELESS_RESOURCE_DISCARD_WORKAROUND
 		
 	// Setup functionality
-	void InitializeTextureData(class FRHICommandListImmediate* RHICmdList, const FRHITextureCreateDesc& CreateDesc, D3D12_RESOURCE_STATES DestinationState);
+	void InitializeTextureData(FRHICommandListBase& RHICmdList, const FRHITextureCreateDesc& CreateDesc, D3D12_RESOURCE_STATES DestinationState);
 	void CreateViews();
 	void SetCreatedRTVsPerSlice(bool Value, int32 InRTVArraySize)
 	{
diff --git a/Engine/Source/Runtime/NullDrv/Public/NullRHI.h b/Engine/Source/Runtime/NullDrv/Public/NullRHI.h
--- a/Engine/Source/Runtime/NullDrv/Public/NullRHI.h
+++ b/Engine/Source/Runtime/NullDrv/Public/NullRHI.h
@@ -177,7 +177,7 @@ public:
 		{}
 	};
 
-	virtual FTextureRHIRef RHICreateTexture(const FRHITextureCreateDesc& CreateDesc) final override
+	virtual FTextureRHIRef RHICreateTexture(FRHICommandListBase&, const FRHITextureCreateDesc& CreateDesc) final override
 	{
 		return new FNullTexture(CreateDesc);
 	}
@@ -191,7 +191,7 @@ public:
 			.SetNumMips(NumMips)
 			.SetInitialState(InResourceState);
 		OutCompletionEvent = nullptr;
-		return this->RHICreateTexture(Desc);
+		return new FNullTexture(Desc);
 	}
 
 	virtual FShaderResourceViewRHIRef RHICreateShaderResourceView(class FRHICommandListBase& RHICmdList, FRHIViewableResource* Resource, FRHIViewDesc const& ViewDesc)
@@ -221,7 +221,7 @@ public:
 			.SetNumMips(NewMipCount)
 			.SetNumSamples(Texture2D->GetNumSamples());
 
-		return this->RHICreateTexture(Desc);
+		return new FNullTexture(Desc);
 	}
 	virtual ETextureReallocationStatus RHIFinalizeAsyncReallocateTexture2D(FRHITexture2D* Texture2D, bool bBlockUntilCompleted) final override
 	{ 
@@ -345,7 +345,7 @@ public:
 			FRHITextureCreateDesc::Create2D(TEXT("FNullDynamicRHI::RHIGetViewportBackBuffer"), 1, 1, PF_B8G8R8A8)
 			.SetFlags(ETextureCreateFlags::RenderTargetable);
 
-		return this->RHICreateTexture(Desc);
+		return new FNullTexture(Desc);
 	}
 
 	virtual void RHIBeginFrame() final override
diff --git a/Engine/Source/Runtime/OpenGLDrv/Private/Android/AndroidOpenGL.cpp b/Engine/Source/Runtime/OpenGLDrv/Private/Android/AndroidOpenGL.cpp
--- a/Engine/Source/Runtime/OpenGLDrv/Private/Android/AndroidOpenGL.cpp
+++ b/Engine/Source/Runtime/OpenGLDrv/Private/Android/AndroidOpenGL.cpp
@@ -476,7 +476,7 @@ FOpenGLTexture* PlatformCreateBuiltinBackBuffer(FOpenGLDynamicRHI* OpenGLRHI, ui
 			.SetFlags(ETextureCreateFlags::RenderTargetable | ETextureCreateFlags::Presentable | ETextureCreateFlags::ResolveTargetable)
 			.DetermineInititialState();
 
-		return new FOpenGLTexture(CreateDesc);
+		return new FOpenGLTexture(FRHICommandListImmediate::Get(), CreateDesc);
 	}
 	else
 	{
diff --git a/Engine/Source/Runtime/OpenGLDrv/Private/OpenGLTexture.cpp b/Engine/Source/Runtime/OpenGLDrv/Private/OpenGLTexture.cpp
--- a/Engine/Source/Runtime/OpenGLDrv/Private/OpenGLTexture.cpp
+++ b/Engine/Source/Runtime/OpenGLDrv/Private/OpenGLTexture.cpp
@@ -328,7 +328,7 @@ FOpenGLTexture::FOpenGLTexture(FOpenGLTextureCreateDesc const& CreateDesc, GLuin
 {}
 
 // Standard constructor.
-FOpenGLTexture::FOpenGLTexture(FOpenGLTextureCreateDesc const& CreateDesc)
+FOpenGLTexture::FOpenGLTexture(FRHICommandListBase& RHICmdList, FOpenGLTextureCreateDesc const& CreateDesc)
 	: FRHITexture        (CreateDesc)
 	, Target             (CreateDesc.Target)
 	, Attachment         (CreateDesc.Attachment)
@@ -342,9 +342,6 @@ FOpenGLTexture::FOpenGLTexture(FOpenGLTextureCreateDesc const& CreateDesc)
 	, bAlias             (false)
 	, bMultisampleRenderbuffer(CreateDesc.bMultisampleRenderbuffer)
 {
-	check(IsInRenderingThread());
-	FRHICommandListImmediate& RHICmdList = FRHICommandListExecutor::GetImmediateCommandList();
-
 	SCOPE_CYCLE_COUNTER(STAT_OpenGLCreateTextureTime);
 
 	if (bCanCreateAsEvicted)
@@ -358,7 +355,7 @@ FOpenGLTexture::FOpenGLTexture(FOpenGLTextureCreateDesc const& CreateDesc)
 
 	if (CreateDesc.BulkData)
 	{
-		if (!ShouldRunGLRenderContextOpOnThisThread(RHICmdList))
+		if (RHICmdList.IsTopOfPipe())
 		{
 			// If bulk data is provided, and texture initialization is done by the RHI thread, it needs to be copied out of the FResourceBulkDataInterface.
 			// It is not safe to pass this pointer to the RHI thread, as the interface may have been stack allocated in the renderer.
@@ -378,7 +375,7 @@ FOpenGLTexture::FOpenGLTexture(FOpenGLTextureCreateDesc const& CreateDesc)
 		}
 	}
 
-	RunOnGLRenderContextThread([this, BulkDataPtr, BulkDataSize, bFreeBulkData]()
+	RHICmdList.EnqueueLambda([this, BulkDataPtr, BulkDataSize, bFreeBulkData](FRHICommandListBase&)
 	{
 		FOpenGLDynamicRHI::Get().InitializeGLTexture(this, BulkDataPtr, BulkDataSize);
 		if (bFreeBulkData)
@@ -1431,9 +1428,9 @@ void FOpenGLTexture::CloneViaCopyImage(FOpenGLTexture* Src, uint32 InNumMips, in
 	2D texture support.
 -----------------------------------------------------------------------------*/
 
-FTextureRHIRef FOpenGLDynamicRHI::RHICreateTexture(const FRHITextureCreateDesc& CreateDesc)
+FTextureRHIRef FOpenGLDynamicRHI::RHICreateTexture(FRHICommandListBase& RHICmdList, const FRHITextureCreateDesc& CreateDesc)
 {
-	return new FOpenGLTexture(CreateDesc);
+	return new FOpenGLTexture(RHICmdList, CreateDesc);
 }
 
 FTextureRHIRef FOpenGLDynamicRHI::RHIAsyncCreateTexture2D(uint32 SizeX, uint32 SizeY, uint8 Format, uint32 NumMips, ETextureCreateFlags Flags, ERHIAccess InResourceState, void** InitialMipData, uint32 NumInitialMips, FGraphEventRef& OutCompletionEvent)
@@ -1509,7 +1506,7 @@ FTexture2DRHIRef FOpenGLDynamicRHI::RHIAsyncReallocateTexture2D(FRHITexture2D* T
 		TEXT("RHIAsyncReallocateTexture2D")
 	);
 
-	FOpenGLTexture* NewTexture = new FOpenGLTexture(CreateDesc);
+	FOpenGLTexture* NewTexture = new FOpenGLTexture(RHICmdList, CreateDesc);
 
 	RHICmdList.EnqueueLambda([OldTexture, NewTexture, SourceMipCount, NewMipCount, RequestStatus](FRHICommandListImmediate& RHICmdList)
 	{
diff --git a/Engine/Source/Runtime/OpenGLDrv/Private/OpenGLViewport.cpp b/Engine/Source/Runtime/OpenGLDrv/Private/OpenGLViewport.cpp
--- a/Engine/Source/Runtime/OpenGLDrv/Private/OpenGLViewport.cpp
+++ b/Engine/Source/Runtime/OpenGLDrv/Private/OpenGLViewport.cpp
@@ -305,7 +305,7 @@ void FOpenGLViewport::Resize(uint32 InSizeX,uint32 InSizeY,bool bInIsFullscreen)
 				.SetFlags(ETextureCreateFlags::RenderTargetable | ETextureCreateFlags::ResolveTargetable)
 				.DetermineInititialState();
 
-			BackBuffer = new FOpenGLTexture(Desc);
+			BackBuffer = new FOpenGLTexture(RHICmdList, Desc);
 		}
 
 		RHICmdList.EnqueueLambda([this, InSizeX, InSizeY, bInIsFullscreen, bWasFullscreen](FRHICommandListImmediate&)
diff --git a/Engine/Source/Runtime/OpenGLDrv/Private/Windows/OpenGLWindows.cpp b/Engine/Source/Runtime/OpenGLDrv/Private/Windows/OpenGLWindows.cpp
--- a/Engine/Source/Runtime/OpenGLDrv/Private/Windows/OpenGLWindows.cpp
+++ b/Engine/Source/Runtime/OpenGLDrv/Private/Windows/OpenGLWindows.cpp
@@ -1010,7 +1010,7 @@ FOpenGLTexture* PlatformCreateBuiltinBackBuffer(FOpenGLDynamicRHI* OpenGLRHI, ui
 			.SetFlags(ETextureCreateFlags::RenderTargetable | ETextureCreateFlags::Presentable | ETextureCreateFlags::ResolveTargetable)
 			.DetermineInititialState();
 
-		return new FOpenGLTexture(Desc);
+		return new FOpenGLTexture(FRHICommandListImmediate::Get(), Desc);
 	}
 
 	return nullptr;
diff --git a/Engine/Source/Runtime/OpenGLDrv/Public/OpenGLDrv.h b/Engine/Source/Runtime/OpenGLDrv/Public/OpenGLDrv.h
--- a/Engine/Source/Runtime/OpenGLDrv/Public/OpenGLDrv.h
+++ b/Engine/Source/Runtime/OpenGLDrv/Public/OpenGLDrv.h
@@ -367,7 +367,7 @@ public:
 	virtual FRHICalcTextureSizeResult RHICalcTexturePlatformSize(FRHITextureDesc const& Desc, uint32 FirstMipIndex) final override;
 	virtual void RHIGetTextureMemoryStats(FTextureMemoryStats& OutStats) final override;
 	virtual bool RHIGetTextureMemoryVisualizeData(FColor* TextureData,int32 SizeX,int32 SizeY,int32 Pitch,int32 PixelSize) final override;
-	virtual FTextureRHIRef RHICreateTexture(const FRHITextureCreateDesc& CreateDesc) final override;
+	virtual FTextureRHIRef RHICreateTexture(FRHICommandListBase& RHICmdList, const FRHITextureCreateDesc& CreateDesc) final override;
 	virtual FTextureRHIRef RHIAsyncCreateTexture2D(uint32 SizeX, uint32 SizeY, uint8 Format, uint32 NumMips, ETextureCreateFlags Flags, ERHIAccess InResourceState, void** InitialMipData, uint32 NumInitialMips, FGraphEventRef& OutCompletionEvent) final override;
 	virtual void RHIGenerateMips(FRHITexture* Texture) final override;
 	virtual uint32 RHIComputeMemorySize(FRHITexture* TextureRHI) final override;
@@ -744,12 +744,6 @@ public:
 	virtual void* LockTexture2DArray_RenderThread(class FRHICommandListImmediate& RHICmdList, FRHITexture2DArray* Texture, uint32 ArrayIndex, uint32 MipIndex, EResourceLockMode LockMode, uint32& DestStride, bool bLockWithinMiptail) final override;
 	virtual void UnlockTexture2DArray_RenderThread(class FRHICommandListImmediate& RHICmdList, FRHITexture2DArray* Texture, uint32 ArrayIndex, uint32 MipIndex, bool bLockWithinMiptail) final override;
 
-	virtual FTextureRHIRef RHICreateTexture_RenderThread(class FRHICommandListImmediate& RHICmdList, const FRHITextureCreateDesc& CreateDesc) final override
-	{
-		check(IsInRenderingThread());
-		return this->RHICreateTexture(CreateDesc);
-	}
-
 	virtual FTexture2DRHIRef AsyncReallocateTexture2D_RenderThread(class FRHICommandListImmediate& RHICmdList, FRHITexture2D* Texture2D, int32 NewMipCount, int32 NewSizeX, int32 NewSizeY, FThreadSafeCounter* RequestStatus) final override;
 
 	virtual ETextureReallocationStatus FinalizeAsyncReallocateTexture2D_RenderThread(class FRHICommandListImmediate& RHICmdList, FRHITexture2D* Texture2D, bool bBlockUntilCompleted) final override
diff --git a/Engine/Source/Runtime/OpenGLDrv/Public/OpenGLResources.h b/Engine/Source/Runtime/OpenGLDrv/Public/OpenGLResources.h
--- a/Engine/Source/Runtime/OpenGLDrv/Public/OpenGLResources.h
+++ b/Engine/Source/Runtime/OpenGLDrv/Public/OpenGLResources.h
@@ -961,7 +961,7 @@ class OPENGLDRV_API FOpenGLTexture : public FRHITexture, public FOpenGLViewableR
 
 public:
 	// Standard constructor.
-	explicit FOpenGLTexture(FOpenGLTextureCreateDesc const& CreateDesc);
+	explicit FOpenGLTexture(FRHICommandListBase& RHICmdList, FOpenGLTextureCreateDesc const& CreateDesc);
 
 	// Constructor for external resources (RHICreateTexture2DFromResource etc).
 	explicit FOpenGLTexture(FOpenGLTextureCreateDesc const& CreateDesc, GLuint Resource);
diff --git a/Engine/Source/Runtime/RHI/Private/RHICommandList.cpp b/Engine/Source/Runtime/RHI/Private/RHICommandList.cpp
--- a/Engine/Source/Runtime/RHI/Private/RHICommandList.cpp
+++ b/Engine/Source/Runtime/RHI/Private/RHICommandList.cpp
@@ -1450,13 +1450,6 @@ void FRHIComputeCommandList::BuildAccelerationStructures(const TArrayView<const
 }
 #endif
 
-FBufferRHIRef FDynamicRHI::CreateBuffer_RenderThread(class FRHICommandListBase& RHICmdList, uint32 Size, EBufferUsageFlags Usage, uint32 Stride, ERHIAccess ResourceState, FRHIResourceCreateInfo& CreateInfo)
-{
-	CSV_SCOPED_TIMING_STAT(RHITStalls, CreateBuffer_RenderThread);
-	FScopedRHIThreadStaller StallRHIThread(RHICmdList.GetAsImmediate());
-	return GDynamicRHI->RHICreateBuffer(RHICmdList, FRHIBufferDesc(Size, Stride, Usage), ResourceState, CreateInfo);
-}
-
 static FLockTracker GLockTracker;
 
 void* FDynamicRHI::RHILockBuffer(class FRHICommandListBase& RHICmdList, FRHIBuffer* Buffer, uint32 Offset, uint32 SizeRHI, EResourceLockMode LockMode)
@@ -1728,13 +1721,6 @@ FRHIShaderLibraryRef FDynamicRHI::RHICreateShaderLibrary_RenderThread(class FRHI
 	return GDynamicRHI->RHICreateShaderLibrary(Platform, FilePath, Name);
 }
 
-FTextureRHIRef FDynamicRHI::RHICreateTexture_RenderThread(class FRHICommandListImmediate& RHICmdList, const FRHITextureCreateDesc& CreateDesc)
-{
-	CSV_SCOPED_TIMING_STAT(RHITStalls, RHICreateTexture_RenderThread);
-	FScopedRHIThreadStaller StallRHIThread(RHICmdList);
-	return GDynamicRHI->RHICreateTexture(CreateDesc);
-}
-
 void* FDynamicRHI::RHILockTextureCubeFace_RenderThread(class FRHICommandListImmediate& RHICmdList, FRHITextureCube* Texture, uint32 FaceIndex, uint32 ArrayIndex, uint32 MipIndex, EResourceLockMode LockMode, uint32& DestStride, bool bLockWithinMiptail)
 {
 	QUICK_SCOPE_CYCLE_COUNTER(STAT_RHIMETHOD_LockTextureCubeFace_Flush);
diff --git a/Engine/Source/Runtime/RHI/Private/RHIGPUReadback.cpp b/Engine/Source/Runtime/RHI/Private/RHIGPUReadback.cpp
--- a/Engine/Source/Runtime/RHI/Private/RHIGPUReadback.cpp
+++ b/Engine/Source/Runtime/RHI/Private/RHIGPUReadback.cpp
@@ -177,7 +177,7 @@ void FRHIGPUTextureReadback::EnqueueCopy(FRHICommandList& RHICmdList, FRHITextur
 						break;
 				}
 
-				DestinationStagingTextures[GPUIndex] = RHICreateTexture(Desc);
+				DestinationStagingTextures[GPUIndex] = RHICmdList.CreateTexture(Desc);
 			}
 			else
 			{
diff --git a/Engine/Source/Runtime/RHI/Public/DynamicRHI.h b/Engine/Source/Runtime/RHI/Public/DynamicRHI.h
--- a/Engine/Source/Runtime/RHI/Public/DynamicRHI.h
+++ b/Engine/Source/Runtime/RHI/Public/DynamicRHI.h
@@ -396,7 +396,7 @@ public:
 	/**
 	* Creates an RHI texture resource.
 	*/
-	virtual FTextureRHIRef RHICreateTexture(const FRHITextureCreateDesc& CreateDesc) = 0;
+	virtual FTextureRHIRef RHICreateTexture(FRHICommandListBase& RHICmdList, const FRHITextureCreateDesc& CreateDesc) = 0;
 
 	/**
 	* Thread-safe function that can be used to create a texture outside of the
@@ -882,8 +882,6 @@ public:
 	//
 	virtual void RHISubmitCommandLists(TArrayView<IRHIPlatformCommandList*> CommandLists, bool bFlushResources) = 0;
 
-	UE_DEPRECATED(5.1, "CreateBuffer_RenderThread is deprecated. Use RHICreateBuffer instead.")
-	RHI_API virtual FBufferRHIRef CreateBuffer_RenderThread(class FRHICommandListBase& RHICmdList, uint32 Size, EBufferUsageFlags Usage, uint32 Stride, ERHIAccess ResourceState, FRHIResourceCreateInfo& CreateInfo);
 	RHI_API virtual FTexture2DRHIRef AsyncReallocateTexture2D_RenderThread(class FRHICommandListImmediate& RHICmdList, FRHITexture2D* Texture2D, int32 NewMipCount, int32 NewSizeX, int32 NewSizeY, FThreadSafeCounter* RequestStatus);
 	RHI_API virtual ETextureReallocationStatus FinalizeAsyncReallocateTexture2D_RenderThread(class FRHICommandListImmediate& RHICmdList, FRHITexture2D* Texture2D, bool bBlockUntilCompleted);
 	RHI_API virtual ETextureReallocationStatus CancelAsyncReallocateTexture2D_RenderThread(class FRHICommandListImmediate& RHICmdList, FRHITexture2D* Texture2D, bool bBlockUntilCompleted);
@@ -901,8 +899,6 @@ public:
 
 	RHI_API virtual FRHIShaderLibraryRef RHICreateShaderLibrary_RenderThread(class FRHICommandListImmediate& RHICmdList, EShaderPlatform Platform, FString FilePath, FString Name);
 
-	RHI_API virtual FTextureRHIRef RHICreateTexture_RenderThread(class FRHICommandListImmediate& RHICmdList, const FRHITextureCreateDesc& CreateDesc);
-	
 	RHI_API virtual void* RHILockTextureCubeFace_RenderThread(class FRHICommandListImmediate& RHICmdList, FRHITextureCube* Texture, uint32 FaceIndex, uint32 ArrayIndex, uint32 MipIndex, EResourceLockMode LockMode, uint32& DestStride, bool bLockWithinMiptail);
 	RHI_API virtual void RHIUnlockTextureCubeFace_RenderThread(class FRHICommandListImmediate& RHICmdList, FRHITextureCube* Texture, uint32 FaceIndex, uint32 ArrayIndex, uint32 MipIndex, bool bLockWithinMiptail);
 
diff --git a/Engine/Source/Runtime/RHI/Public/RHICommandList.h b/Engine/Source/Runtime/RHI/Public/RHICommandList.h
--- a/Engine/Source/Runtime/RHI/Public/RHICommandList.h
+++ b/Engine/Source/Runtime/RHI/Public/RHICommandList.h
@@ -762,6 +762,22 @@ public:
 		GDynamicRHI->RHIUpdateTexture2D(*this, Texture, MipIndex, UpdateRegion, SourcePitch, SourceData);
 	}
 
+	FORCEINLINE FTextureRHIRef CreateTexture(const FRHITextureCreateDesc& CreateDesc)
+	{
+		LLM_SCOPE(EnumHasAnyFlags(CreateDesc.Flags, TexCreate_RenderTargetable | TexCreate_DepthStencilTargetable) ? ELLMTag::RenderTargets : ELLMTag::Textures);
+
+		if (CreateDesc.InitialState == ERHIAccess::Unknown)
+		{
+			// Need to copy the incoming descriptor since we need to override the initial state.
+			FRHITextureCreateDesc NewCreateDesc(CreateDesc);
+			NewCreateDesc.SetInitialState(RHIGetDefaultResourceState(CreateDesc.Flags, CreateDesc.BulkData != nullptr));
+
+			return GDynamicRHI->RHICreateTexture(*this, NewCreateDesc);
+		}
+
+		return GDynamicRHI->RHICreateTexture(*this, CreateDesc);
+	}
+
 	FORCEINLINE void UpdateFromBufferTexture2D(FRHITexture2D* Texture, uint32 MipIndex, const struct FUpdateTextureRegion2D& UpdateRegion, uint32 SourcePitch, FRHIBuffer* Buffer, uint32 BufferOffset)
 	{
 		checkf(UpdateRegion.DestX + UpdateRegion.Width <= Texture->GetSizeX(), TEXT("UpdateFromBufferTexture2D out of bounds on X. Texture: %s, %i, %i, %i"), *Texture->GetName().ToString(), UpdateRegion.DestX, UpdateRegion.Width, Texture->GetSizeX());
@@ -5077,21 +5093,7 @@ FORCEINLINE void RHIUpdateTextureReference(FRHITextureReference* TextureRef, FRH
 
 FORCEINLINE FTextureRHIRef RHICreateTexture(const FRHITextureCreateDesc& CreateDesc)
 {
-	//check(IsInRenderingThread()); // @todo: texture type unification Some passes call this function on parallel rendering threads (e.g. FRHIGPUTextureReadback::EnqueueCopyInternal)
-	FRHICommandListImmediate& RHICmdList = FRHICommandListExecutor::GetImmediateCommandList();
-
-	LLM_SCOPE(EnumHasAnyFlags(CreateDesc.Flags, TexCreate_RenderTargetable | TexCreate_DepthStencilTargetable) ? ELLMTag::RenderTargets : ELLMTag::Textures);
-
-	if (CreateDesc.InitialState == ERHIAccess::Unknown)
-	{
-		// Need to copy the incoming descriptor since we need to override the initial state.
-		FRHITextureCreateDesc NewCreateDesc(CreateDesc);
-		NewCreateDesc.SetInitialState(RHIGetDefaultResourceState(CreateDesc.Flags, CreateDesc.BulkData != nullptr));
-
-		return GDynamicRHI->RHICreateTexture_RenderThread(RHICmdList, NewCreateDesc);
-	}
-
-	return GDynamicRHI->RHICreateTexture_RenderThread(RHICmdList, CreateDesc);
+	return FRHICommandListImmediate::Get().CreateTexture(CreateDesc);
 }
 
 //UE_DEPRECATED(5.1, "FRHITexture2D is deprecated, please use RHICreateTexture(const FRHITextureCreateDesc&).")
diff --git a/Engine/Source/Runtime/RHI/Public/RHIValidation.h b/Engine/Source/Runtime/RHI/Public/RHIValidation.h
--- a/Engine/Source/Runtime/RHI/Public/RHIValidation.h
+++ b/Engine/Source/Runtime/RHI/Public/RHIValidation.h
@@ -493,10 +493,10 @@ public:
 	/**
 	* Creates an RHI texture resource.
 	*/
-	virtual FTextureRHIRef RHICreateTexture(const FRHITextureCreateDesc& CreateDesc) override final
+	virtual FTextureRHIRef RHICreateTexture(FRHICommandListBase& RHICmdList, const FRHITextureCreateDesc& CreateDesc) override final
 	{
 		CreateDesc.CheckValidity();
-		FTextureRHIRef Texture = RHI->RHICreateTexture(CreateDesc);
+		FTextureRHIRef Texture = RHI->RHICreateTexture(RHICmdList, CreateDesc);
 		ensure(Texture->IsBarrierTrackingInitialized());
 		return Texture;
 	}
@@ -1188,14 +1188,6 @@ public:
 		return RHI->RHICreateShaderLibrary_RenderThread(RHICmdList, Platform, FilePath, Name);
 	}
 
-	virtual FTextureRHIRef RHICreateTexture_RenderThread(class FRHICommandListImmediate& RHICmdList, const FRHITextureCreateDesc& CreateDesc) override final
-	{
-		CreateDesc.CheckValidity();
-		FTextureRHIRef Texture = RHI->RHICreateTexture_RenderThread(RHICmdList, CreateDesc);
-		ensure(Texture->IsBarrierTrackingInitialized());
-		return Texture;
-	}
-
 	virtual void* RHILockTextureCubeFace_RenderThread(class FRHICommandListImmediate& RHICmdList, FRHITextureCube* Texture, uint32 FaceIndex, uint32 ArrayIndex, uint32 MipIndex, EResourceLockMode LockMode, uint32& DestStride, bool bLockWithinMiptail) override final
 	{
 		return RHI->RHILockTextureCubeFace_RenderThread(RHICmdList, Texture, FaceIndex, ArrayIndex, MipIndex, LockMode, DestStride, bLockWithinMiptail);
diff --git a/Engine/Source/Runtime/VulkanRHI/Private/VulkanResources.h b/Engine/Source/Runtime/VulkanRHI/Private/VulkanResources.h
--- a/Engine/Source/Runtime/VulkanRHI/Private/VulkanResources.h
+++ b/Engine/Source/Runtime/VulkanRHI/Private/VulkanResources.h
@@ -552,7 +552,11 @@ class FVulkanTexture : public FRHITexture, public FVulkanEvictable, public FVulk
 {
 public:
 	// Regular constructor.
-	FVulkanTexture(FVulkanDevice& InDevice, const FRHITextureCreateDesc& InCreateDesc, const FRHITransientHeapAllocation* InTransientHeapAllocation);
+	FVulkanTexture(FRHICommandListBase* RHICmdList, FVulkanDevice& InDevice, const FRHITextureCreateDesc& InCreateDesc, const FRHITransientHeapAllocation* InTransientHeapAllocation);
+
+	FVulkanTexture(FVulkanDevice& InDevice, const FRHITextureCreateDesc& InCreateDesc, const FRHITransientHeapAllocation* InTransientHeapAllocation)
+		: FVulkanTexture(nullptr, InDevice, InCreateDesc, InTransientHeapAllocation)
+	{}
 
 	// Construct from external resource.
 	// FIXME: HUGE HACK: the bUnused argument is there to disambiguate this overload from the one above when passing nullptr, since nullptr is a valid VkImage. Get rid of this code smell when unifying FVulkanSurface and FVulkanTexture.
diff --git a/Engine/Source/Runtime/VulkanRHI/Private/VulkanTexture.cpp b/Engine/Source/Runtime/VulkanRHI/Private/VulkanTexture.cpp
--- a/Engine/Source/Runtime/VulkanRHI/Private/VulkanTexture.cpp
+++ b/Engine/Source/Runtime/VulkanRHI/Private/VulkanTexture.cpp
@@ -949,10 +949,10 @@ void FVulkanDynamicRHI::RHIUpdateTextureReference(FRHICommandListBase& RHICmdLis
 	2D texture support.
 -----------------------------------------------------------------------------*/
 
-FTextureRHIRef FVulkanDynamicRHI::RHICreateTexture(const FRHITextureCreateDesc& CreateDesc)
+FTextureRHIRef FVulkanDynamicRHI::RHICreateTexture(FRHICommandListBase& RHICmdList, const FRHITextureCreateDesc& CreateDesc)
 {
 	LLM_SCOPE_VULKAN(GetMemoryTagForTextureFlags(CreateDesc.Flags));
-	return new FVulkanTexture(*Device, CreateDesc, nullptr);
+	return new FVulkanTexture(&RHICmdList, *Device, CreateDesc, nullptr);
 }
 
 FTextureRHIRef FVulkanDynamicRHI::RHIAsyncCreateTexture2D(uint32 SizeX,uint32 SizeY,uint8 Format,uint32 NumMips,ETextureCreateFlags Flags, ERHIAccess InResourceState,void** InitialMipData,uint32 NumInitialMips, FGraphEventRef& OutCompletionEvent)
@@ -1080,7 +1080,7 @@ FTexture2DRHIRef FVulkanDynamicRHI::AsyncReallocateTexture2D_RenderThread(FRHICo
 		.SetNumSamples(OldDesc.NumSamples)
 		.DetermineInititialState();
 
-	FVulkanTexture* NewTexture = new FVulkanTexture(*Device, Desc, nullptr);
+	FVulkanTexture* NewTexture = new FVulkanTexture(&RHICmdList, *Device, Desc, nullptr);
 	ALLOC_COMMAND_CL(RHICmdList, FRHICommandVulkanAsyncReallocateTexture2D)(OldTexture, NewTexture, NewMipCount, NewSizeX, NewSizeY, RequestStatus);
 
 	return NewTexture;
@@ -1424,7 +1424,7 @@ void FVulkanDynamicRHI::InternalUpdateTexture3D(FRHICommandListBase& RHICmdList,
 	}
 }
 
-FVulkanTexture::FVulkanTexture(FVulkanDevice& InDevice, const FRHITextureCreateDesc& InCreateDesc, const FRHITransientHeapAllocation* InTransientHeapAllocation)
+FVulkanTexture::FVulkanTexture(FRHICommandListBase* RHICmdList, FVulkanDevice& InDevice, const FRHITextureCreateDesc& InCreateDesc, const FRHITransientHeapAllocation* InTransientHeapAllocation)
 	: FRHITexture(InCreateDesc)
 	, PartialView(nullptr)
 	, Device(&InDevice)
@@ -1567,15 +1567,22 @@ FVulkanTexture::FVulkanTexture(FVulkanDevice& InDevice, const FRHITextureCreateD
 
 		if (InitialLayout != VK_IMAGE_LAYOUT_UNDEFINED || bDoInitialClear)
 		{
-			FRHICommandList& RHICmdList = FRHICommandListExecutor::GetImmediateCommandList();
-			if (!IsInRenderingThread() || (RHICmdList.Bypass() || !IsRunningRHIInSeparateThread()))
+			if (RHICmdList && RHICmdList->IsTopOfPipe())
 			{
-				SetInitialImageState(Device->GetImmediateContext(), InitialLayout, bDoInitialClear, InCreateDesc.ClearValue, bIsTransientResource);
+				ALLOC_COMMAND_CL(*RHICmdList, FRHICommandSetInitialImageState)(this, InitialLayout, false, bDoInitialClear, InCreateDesc.ClearValue, bIsTransientResource);
 			}
 			else
 			{
-				check(IsInRenderingThread());
-				ALLOC_COMMAND_CL(RHICmdList, FRHICommandSetInitialImageState)(this, InitialLayout, false, bDoInitialClear, InCreateDesc.ClearValue, bIsTransientResource);
+				RHICmdList = &FRHICommandListExecutor::GetImmediateCommandList();
+				if (!IsInRenderingThread() || (RHICmdList->Bypass() || !IsRunningRHIInSeparateThread()))
+				{
+					SetInitialImageState(Device->GetImmediateContext(), InitialLayout, bDoInitialClear, InCreateDesc.ClearValue, bIsTransientResource);
+				}
+				else
+				{
+					check(IsInRenderingThread());
+					ALLOC_COMMAND_CL(*RHICmdList, FRHICommandSetInitialImageState)(this, InitialLayout, false, bDoInitialClear, InCreateDesc.ClearValue, bIsTransientResource);
+				}
 			}
 		}
 
@@ -1661,15 +1668,15 @@ FVulkanTexture::FVulkanTexture(FVulkanDevice& InDevice, const FRHITextureCreateD
 	Region.imageExtent.height = Region.bufferImageHeight;
 	Region.imageExtent.depth = InCreateDesc.Depth;
 
-	FRHICommandList& RHICmdList = FRHICommandListExecutor::GetImmediateCommandList();
-	if (RHICmdList.Bypass() || !IsRunningRHIInSeparateThread())
+	checkf(RHICmdList, TEXT("FVulkanTexture requires a command list for creating bulk data."));
+
+	if (RHICmdList->IsTopOfPipe())
 	{
-		FVulkanTexture::InternalLockWrite(InDevice.GetImmediateContext(), this, Region, StagingBuffer);
+		ALLOC_COMMAND_CL(*RHICmdList, FRHICommandLockWriteTexture)(this, Region, StagingBuffer);
 	}
 	else
 	{
-		check(IsInRenderingThread());
-		ALLOC_COMMAND_CL(RHICmdList, FRHICommandLockWriteTexture)(this, Region, StagingBuffer);
+		FVulkanTexture::InternalLockWrite(InDevice.GetImmediateContext(), this, Region, StagingBuffer);
 	}
 }
 
diff --git a/Engine/Source/Runtime/VulkanRHI/Public/VulkanDynamicRHI.h b/Engine/Source/Runtime/VulkanRHI/Public/VulkanDynamicRHI.h
--- a/Engine/Source/Runtime/VulkanRHI/Public/VulkanDynamicRHI.h
+++ b/Engine/Source/Runtime/VulkanRHI/Public/VulkanDynamicRHI.h
@@ -121,7 +121,7 @@ public:
 	virtual FRHICalcTextureSizeResult RHICalcTexturePlatformSize(FRHITextureDesc const& Desc, uint32 FirstMipIndex) final override;
 	virtual void RHIGetTextureMemoryStats(FTextureMemoryStats& OutStats) final override;
 	virtual bool RHIGetTextureMemoryVisualizeData(FColor* TextureData, int32 SizeX, int32 SizeY, int32 Pitch, int32 PixelSize) final override;
-	virtual FTextureRHIRef RHICreateTexture(const FRHITextureCreateDesc& CreateDesc) final override;
+	virtual FTextureRHIRef RHICreateTexture(FRHICommandListBase& RHICmdList, const FRHITextureCreateDesc& CreateDesc) final override;
 	virtual FTextureRHIRef RHIAsyncCreateTexture2D(uint32 SizeX, uint32 SizeY, uint8 Format, uint32 NumMips, ETextureCreateFlags Flags, ERHIAccess InResourceState, void** InitialMipData, uint32 NumInitialMips, FGraphEventRef& OutCompletionEvent) final override;
 	virtual uint32 RHIComputeMemorySize(FRHITexture* TextureRHI) final override;
 	virtual FTexture2DRHIRef RHIAsyncReallocateTexture2D(FRHITexture2D* Texture2D, int32 NewMipCount, int32 NewSizeX, int32 NewSizeY, FThreadSafeCounter* RequestStatus) final override;
@@ -210,12 +210,6 @@ public:
 	virtual FUpdateTexture3DData RHIBeginUpdateTexture3D(FRHICommandListBase& RHICmdList, FRHITexture3D* Texture, uint32 MipIndex, const struct FUpdateTextureRegion3D& UpdateRegion) override final;
 	virtual void RHIEndUpdateTexture3D(FRHICommandListBase& RHICmdList, FUpdateTexture3DData& UpdateData) override final;
 
-	virtual FTextureRHIRef RHICreateTexture_RenderThread(class FRHICommandListImmediate& RHICmdList, const FRHITextureCreateDesc& CreateDesc) final override
-	{
-		// this-> is required to avoid calling the global version of this function
-		return this->RHICreateTexture(CreateDesc);
-	}
-
 	// SRV / UAV creation functions
 	virtual FShaderResourceViewRHIRef  RHICreateShaderResourceView (class FRHICommandListBase& RHICmdList, FRHIViewableResource* Resource, FRHIViewDesc const& ViewDesc) final override;
 	virtual FUnorderedAccessViewRHIRef RHICreateUnorderedAccessView(class FRHICommandListBase& RHICmdList, FRHIViewableResource* Resource, FRHIViewDesc const& ViewDesc) final override;
diff --git a/Engine/Source/Runtime/Windows/D3D11RHI/Private/D3D11RHIPrivate.h b/Engine/Source/Runtime/Windows/D3D11RHI/Private/D3D11RHIPrivate.h
--- a/Engine/Source/Runtime/Windows/D3D11RHI/Private/D3D11RHIPrivate.h
+++ b/Engine/Source/Runtime/Windows/D3D11RHI/Private/D3D11RHIPrivate.h
@@ -498,7 +498,7 @@ public:
 	virtual FRHICalcTextureSizeResult RHICalcTexturePlatformSize(FRHITextureDesc const& Desc, uint32 FirstMipIndex) final override;
 	virtual void RHIGetTextureMemoryStats(FTextureMemoryStats& OutStats) final override;
 	virtual bool RHIGetTextureMemoryVisualizeData(FColor* TextureData,int32 SizeX,int32 SizeY,int32 Pitch,int32 PixelSize) final override;
-	virtual FTextureRHIRef RHICreateTexture(const FRHITextureCreateDesc& CreateDesc) final override;
+	virtual FTextureRHIRef RHICreateTexture(FRHICommandListBase& RHICmdList, const FRHITextureCreateDesc& CreateDesc) final override;
 	virtual FTextureRHIRef RHIAsyncCreateTexture2D(uint32 SizeX, uint32 SizeY, uint8 Format, uint32 NumMips, ETextureCreateFlags Flags, ERHIAccess InResourceState, void** InitialMipData, uint32 NumInitialMips, FGraphEventRef& OutCompletionEvent) final override;
 	virtual void RHIGenerateMips(FRHITexture* Texture) final override;
 	virtual uint32 RHIComputeMemorySize(FRHITexture* TextureRHI) final override;
@@ -647,7 +647,6 @@ public:
 	virtual void RHIEndUpdateTexture3D(FRHICommandListBase& RHICmdList, FUpdateTexture3DData& UpdateData) final override;
 	virtual void* RHILockTextureCubeFace_RenderThread(class FRHICommandListImmediate& RHICmdList, FRHITextureCube* Texture, uint32 FaceIndex, uint32 ArrayIndex, uint32 MipIndex, EResourceLockMode LockMode, uint32& DestStride, bool bLockWithinMiptail) final override;
 	virtual void RHIUnlockTextureCubeFace_RenderThread(class FRHICommandListImmediate& RHICmdList, FRHITextureCube* Texture, uint32 FaceIndex, uint32 ArrayIndex, uint32 MipIndex, bool bLockWithinMiptail) final override;
-	virtual FTextureRHIRef RHICreateTexture_RenderThread(class FRHICommandListImmediate& RHICmdList, const FRHITextureCreateDesc& CreateDesc) final override;
 
 	virtual void RHIBeginRenderPass(const FRHIRenderPassInfo& InInfo, const TCHAR* InName) final override;
 	virtual void RHIEndRenderPass() final override;
diff --git a/Engine/Source/Runtime/Windows/D3D11RHI/Private/D3D11Texture.cpp b/Engine/Source/Runtime/Windows/D3D11RHI/Private/D3D11Texture.cpp
--- a/Engine/Source/Runtime/Windows/D3D11RHI/Private/D3D11Texture.cpp
+++ b/Engine/Source/Runtime/Windows/D3D11RHI/Private/D3D11Texture.cpp
@@ -806,18 +806,13 @@ FD3D11Texture* FD3D11DynamicRHI::CreateD3D11Texture3D(FRHITextureCreateDesc cons
 	2D texture support.
 -----------------------------------------------------------------------------*/
 
-FTextureRHIRef FD3D11DynamicRHI::RHICreateTexture(const FRHITextureCreateDesc& CreateDesc)
+FTextureRHIRef FD3D11DynamicRHI::RHICreateTexture(FRHICommandListBase&, const FRHITextureCreateDesc& CreateDesc)
 {
 	return CreateDesc.IsTexture3D()
 		? CreateD3D11Texture3D(CreateDesc)
 		: CreateD3D11Texture2D(CreateDesc);
 }
 
-FTextureRHIRef FD3D11DynamicRHI::RHICreateTexture_RenderThread(class FRHICommandListImmediate& RHICmdList, const FRHITextureCreateDesc& CreateDesc)
-{
-	return RHICreateTexture(CreateDesc);
-}
-
 FTextureRHIRef FD3D11DynamicRHI::RHIAsyncCreateTexture2D(uint32 SizeX, uint32 SizeY, uint8 Format, uint32 NumMips, ETextureCreateFlags Flags, ERHIAccess InResourceState, void** InitialMipData, uint32 NumInitialMips, FGraphEventRef& OutCompletionEvent)
 {
 	TArray<D3D11_SUBRESOURCE_DATA, TInlineAllocator<12>> SubresourceData;
