RHITransferBufferUnderlyingResource and RHITransferRayTracingGeometryUnderlyingResource need to be given a command list for future rename support.

diff --git a/Engine/Source/Runtime/Apple/MetalRHI/Private/MetalDynamicRHI.h b/Engine/Source/Runtime/Apple/MetalRHI/Private/MetalDynamicRHI.h
--- a/Engine/Source/Runtime/Apple/MetalRHI/Private/MetalDynamicRHI.h
+++ b/Engine/Source/Runtime/Apple/MetalRHI/Private/MetalDynamicRHI.h
@@ -58,7 +58,7 @@ public:
 	virtual TRefCountPtr<FRHIComputePipelineState> RHICreateComputePipelineState(FRHIComputeShader* ComputeShader) final override;
 	virtual FUniformBufferRHIRef RHICreateUniformBuffer(const void* Contents, const FRHIUniformBufferLayout* Layout, EUniformBufferUsage Usage, EUniformBufferValidation Validation) final override;
 	virtual FBufferRHIRef RHICreateBuffer(FRHICommandListBase& RHICmdList, FRHIBufferDesc const& Desc, ERHIAccess ResourceState, FRHIResourceCreateInfo& CreateInfo) final override;
-	virtual void RHITransferBufferUnderlyingResource(FRHIBuffer* DestBuffer, FRHIBuffer* SrcBuffer) final override;
+	virtual void RHITransferBufferUnderlyingResource(FRHICommandListBase& RHICmdList, FRHIBuffer* DestBuffer, FRHIBuffer* SrcBuffer) final override;
 	virtual void RHICopyBuffer(FRHIBuffer* SourceBuffer, FRHIBuffer* DestBuffer) final override;
 	virtual void* LockBuffer_BottomOfPipe(FRHICommandListBase& RHICmdList, FRHIBuffer* Buffer, uint32 Offset, uint32 SizeRHI, EResourceLockMode LockMode) final override;
 	virtual void UnlockBuffer_BottomOfPipe(FRHICommandListBase& RHICmdList, FRHIBuffer* Buffer) final override;
@@ -152,7 +152,7 @@ public:
 	virtual FRayTracingSceneRHIRef RHICreateRayTracingScene(FRayTracingSceneInitializer2 Initializer) final override;
 	virtual FRayTracingShaderRHIRef RHICreateRayTracingShader(TArrayView<const uint8> Code, const FSHAHash& Hash, EShaderFrequency ShaderFrequency) final override;
 	virtual FRayTracingPipelineStateRHIRef RHICreateRayTracingPipelineState(const FRayTracingPipelineStateInitializer& Initializer) final override;
-	virtual void RHITransferRayTracingGeometryUnderlyingResource(FRHIRayTracingGeometry* DestGeometry, FRHIRayTracingGeometry* SrcGeometry) final override;
+	virtual void RHITransferRayTracingGeometryUnderlyingResource(FRHICommandListBase& RHICmdList, FRHIRayTracingGeometry* DestGeometry, FRHIRayTracingGeometry* SrcGeometry) final override;
 #endif // METAL_RHI_RAYTRACING
 
 private:
diff --git a/Engine/Source/Runtime/Apple/MetalRHI/Private/MetalRayTracing.cpp b/Engine/Source/Runtime/Apple/MetalRHI/Private/MetalRayTracing.cpp
--- a/Engine/Source/Runtime/Apple/MetalRHI/Private/MetalRayTracing.cpp
+++ b/Engine/Source/Runtime/Apple/MetalRHI/Private/MetalRayTracing.cpp
@@ -664,7 +664,7 @@ void FMetalRHICommandContext::RHIBuildAccelerationStructure(const FRayTracingSce
 		InstanceBuffer, SceneBuildParams.InstanceBufferOffset);
 }
 
-void FMetalDynamicRHI::RHITransferRayTracingGeometryUnderlyingResource(FRHIRayTracingGeometry* DestGeometry, FRHIRayTracingGeometry* SrcGeometry)
+void FMetalDynamicRHI::RHITransferRayTracingGeometryUnderlyingResource(FRHICommandListBase& RHICmdList, FRHIRayTracingGeometry* DestGeometry, FRHIRayTracingGeometry* SrcGeometry)
 {
 	check(DestGeometry);
 	FMetalRayTracingGeometry* Dest = ResourceCast(DestGeometry);
diff --git a/Engine/Source/Runtime/Apple/MetalRHI/Private/MetalVertexBuffer.cpp b/Engine/Source/Runtime/Apple/MetalRHI/Private/MetalVertexBuffer.cpp
--- a/Engine/Source/Runtime/Apple/MetalRHI/Private/MetalVertexBuffer.cpp
+++ b/Engine/Source/Runtime/Apple/MetalRHI/Private/MetalVertexBuffer.cpp
@@ -671,7 +671,7 @@ void FMetalDynamicRHI::RHICopyBuffer(FRHIBuffer* SourceBufferRHI, FRHIBuffer* De
 	}
 }
 
-void FMetalDynamicRHI::RHITransferBufferUnderlyingResource(FRHIBuffer* DestBuffer, FRHIBuffer* SrcBuffer)
+void FMetalDynamicRHI::RHITransferBufferUnderlyingResource(FRHICommandListBase& RHICmdList, FRHIBuffer* DestBuffer, FRHIBuffer* SrcBuffer)
 {
 	@autoreleasepool {
         FMetalRHIBuffer* Dst = ResourceCast(DestBuffer);
diff --git a/Engine/Source/Runtime/D3D12RHI/Private/D3D12Buffer.cpp b/Engine/Source/Runtime/D3D12RHI/Private/D3D12Buffer.cpp
--- a/Engine/Source/Runtime/D3D12RHI/Private/D3D12Buffer.cpp
+++ b/Engine/Source/Runtime/D3D12RHI/Private/D3D12Buffer.cpp
@@ -794,7 +794,7 @@ void FD3D12DynamicRHI::RHIUnlockBufferMGPU(FRHICommandListBase& RHICmdList, FRHI
 	UnlockBuffer(RHICmdList, Buffer, Buffer->GetUsage());
 }
 
-void FD3D12DynamicRHI::RHITransferBufferUnderlyingResource(FRHIBuffer* DestBuffer, FRHIBuffer* SrcBuffer)
+void FD3D12DynamicRHI::RHITransferBufferUnderlyingResource(FRHICommandListBase& RHICmdList, FRHIBuffer* DestBuffer, FRHIBuffer* SrcBuffer)
 {
 	FD3D12Buffer* Dst = ResourceCast(DestBuffer);
 	FD3D12Buffer* Src = ResourceCast(SrcBuffer);
diff --git a/Engine/Source/Runtime/D3D12RHI/Private/D3D12RHIPrivate.h b/Engine/Source/Runtime/D3D12RHI/Private/D3D12RHIPrivate.h
--- a/Engine/Source/Runtime/D3D12RHI/Private/D3D12RHIPrivate.h
+++ b/Engine/Source/Runtime/D3D12RHI/Private/D3D12RHIPrivate.h
@@ -298,7 +298,7 @@ public:
 	virtual FUniformBufferRHIRef RHICreateUniformBuffer(const void* Contents, const FRHIUniformBufferLayout* Layout, EUniformBufferUsage Usage, EUniformBufferValidation Validation) final override;
 	virtual void RHIUpdateUniformBuffer(FRHICommandListBase& RHICmdList, FRHIUniformBuffer* UniformBufferRHI, const void* Contents) final override;
 	virtual FBufferRHIRef RHICreateBuffer(FRHICommandListBase& RHICmdList, FRHIBufferDesc const& Desc, ERHIAccess ResourceState, FRHIResourceCreateInfo& CreateInfo) final override;
-	virtual void RHITransferBufferUnderlyingResource(FRHIBuffer* DestBuffer, FRHIBuffer* SrcBuffer) final override;
+	virtual void RHITransferBufferUnderlyingResource(FRHICommandListBase& RHICmdList, FRHIBuffer* DestBuffer, FRHIBuffer* SrcBuffer) final override;
 	virtual void RHICopyBuffer(FRHIBuffer* SourceBuffer, FRHIBuffer* DestBuffer) final override;
 	virtual void* RHILockBuffer(FRHICommandListBase& RHICmdList, FRHIBuffer* Buffer, uint32 Offset, uint32 Size, EResourceLockMode LockMode) final override;
 	virtual void* RHILockBufferMGPU(FRHICommandListBase& RHICmdList, FRHIBuffer* Buffer, uint32 GPUIndex, uint32 Offset, uint32 Size, EResourceLockMode LockMode) final override;
@@ -444,7 +444,7 @@ public:
 	virtual FRayTracingSceneRHIRef RHICreateRayTracingScene(FRayTracingSceneInitializer2 Initializer) final override;
 	virtual FRayTracingShaderRHIRef RHICreateRayTracingShader(TArrayView<const uint8> Code, const FSHAHash& Hash, EShaderFrequency ShaderFrequency) final override;
 	virtual FRayTracingPipelineStateRHIRef RHICreateRayTracingPipelineState(const FRayTracingPipelineStateInitializer& Initializer) final override;
-	virtual void RHITransferRayTracingGeometryUnderlyingResource(FRHIRayTracingGeometry* DestGeometry, FRHIRayTracingGeometry* SrcGeometry) final override;
+	virtual void RHITransferRayTracingGeometryUnderlyingResource(FRHICommandListBase& RHICmdList, FRHIRayTracingGeometry* DestGeometry, FRHIRayTracingGeometry* SrcGeometry) final override;
 #endif //D3D12_RHI_RAYTRACING
 
 	virtual FShaderBundleRHIRef RHICreateShaderBundle(uint32 NumRecords, bool bInEmulated) override;
diff --git a/Engine/Source/Runtime/D3D12RHI/Private/D3D12RayTracing.cpp b/Engine/Source/Runtime/D3D12RHI/Private/D3D12RayTracing.cpp
--- a/Engine/Source/Runtime/D3D12RHI/Private/D3D12RayTracing.cpp
+++ b/Engine/Source/Runtime/D3D12RHI/Private/D3D12RayTracing.cpp
@@ -2675,7 +2675,7 @@ FRayTracingGeometryRHIRef FD3D12DynamicRHI::RHICreateRayTracingGeometry(FRHIComm
 	return new FD3D12RayTracingGeometry(RHICmdList, &Adapter, Initializer);
 }
 
-void FD3D12DynamicRHI::RHITransferRayTracingGeometryUnderlyingResource(FRHIRayTracingGeometry* DestGeometry, FRHIRayTracingGeometry* SrcGeometry)
+void FD3D12DynamicRHI::RHITransferRayTracingGeometryUnderlyingResource(FRHICommandListBase& RHICmdList, FRHIRayTracingGeometry* DestGeometry, FRHIRayTracingGeometry* SrcGeometry)
 {
 	check(DestGeometry);
 	FD3D12RayTracingGeometry* Dest = ResourceCast(DestGeometry);
diff --git a/Engine/Source/Runtime/NullDrv/Public/NullRHI.h b/Engine/Source/Runtime/NullDrv/Public/NullRHI.h
--- a/Engine/Source/Runtime/NullDrv/Public/NullRHI.h
+++ b/Engine/Source/Runtime/NullDrv/Public/NullRHI.h
@@ -139,7 +139,7 @@ public:
 
 	}
 
-	virtual void RHITransferBufferUnderlyingResource(FRHIBuffer* DestBuffer, FRHIBuffer* SrcBuffer) final override
+	virtual void RHITransferBufferUnderlyingResource(FRHICommandListBase& RHICmdList, FRHIBuffer* DestBuffer, FRHIBuffer* SrcBuffer) final override
 	{
 
 	}
diff --git a/Engine/Source/Runtime/OpenGLDrv/Private/OpenGLBuffer.cpp b/Engine/Source/Runtime/OpenGLDrv/Private/OpenGLBuffer.cpp
--- a/Engine/Source/Runtime/OpenGLDrv/Private/OpenGLBuffer.cpp
+++ b/Engine/Source/Runtime/OpenGLDrv/Private/OpenGLBuffer.cpp
@@ -264,7 +264,7 @@ void FOpenGLDynamicRHI::RHICopyBufferRegion(FRHIBuffer* DestBufferRHI, uint64 Ds
 	glBindBuffer(GL_COPY_WRITE_BUFFER, 0);
 }
 
-void FOpenGLDynamicRHI::RHITransferBufferUnderlyingResource(FRHIBuffer* DestBuffer, FRHIBuffer* SrcBuffer)
+void FOpenGLDynamicRHI::RHITransferBufferUnderlyingResource(FRHICommandListBase& RHICmdList, FRHIBuffer* DestBuffer, FRHIBuffer* SrcBuffer)
 {
 	VERIFY_GL_SCOPE();
 	FOpenGLBuffer* Dst = ResourceCast(DestBuffer);
diff --git a/Engine/Source/Runtime/OpenGLDrv/Public/OpenGLDrv.h b/Engine/Source/Runtime/OpenGLDrv/Public/OpenGLDrv.h
--- a/Engine/Source/Runtime/OpenGLDrv/Public/OpenGLDrv.h
+++ b/Engine/Source/Runtime/OpenGLDrv/Public/OpenGLDrv.h
@@ -361,7 +361,7 @@ public:
 	virtual void RHIUpdateUniformBuffer(FRHICommandListBase& RHICmdList, FRHIUniformBuffer* UniformBufferRHI, const void* Contents) final override;
 	virtual FBufferRHIRef RHICreateBuffer(FRHICommandListBase& RHICmdList, FRHIBufferDesc const& Desc, ERHIAccess ResourceState, FRHIResourceCreateInfo& CreateInfo) final override;
 	virtual void RHICopyBuffer(FRHIBuffer* SourceBuffer, FRHIBuffer* DestBuffer) final override;
-	virtual void RHITransferBufferUnderlyingResource(FRHIBuffer* DestBuffer, FRHIBuffer* SrcBuffer) final override;
+	virtual void RHITransferBufferUnderlyingResource(FRHICommandListBase& RHICmdList, FRHIBuffer* DestBuffer, FRHIBuffer* SrcBuffer) final override;
 	virtual void* LockBuffer_BottomOfPipe(FRHICommandListBase& RHICmdList, FRHIBuffer* Buffer, uint32 Offset, uint32 Size, EResourceLockMode LockMode) final override;
 	virtual void UnlockBuffer_BottomOfPipe(FRHICommandListBase& RHICmdList, FRHIBuffer* Buffer) final override;
 	virtual FRHICalcTextureSizeResult RHICalcTexturePlatformSize(FRHITextureDesc const& Desc, uint32 FirstMipIndex) final override;
diff --git a/Engine/Source/Runtime/RHI/Private/DynamicRHI.cpp b/Engine/Source/Runtime/RHI/Private/DynamicRHI.cpp
--- a/Engine/Source/Runtime/RHI/Private/DynamicRHI.cpp
+++ b/Engine/Source/Runtime/RHI/Private/DynamicRHI.cpp
@@ -553,11 +553,6 @@ void FDynamicRHI::EnableIdealGPUCaptureOptions(bool bEnabled)
 	}	
 }
 
-void FDynamicRHI::RHITransferBufferUnderlyingResource(FRHIBuffer* DestBuffer, FRHIBuffer* SrcBuffer)
-{
-	UE_LOG(LogRHI, Fatal, TEXT("RHITransferBufferUnderlyingResource isn't implemented for the current RHI"));
-}
-
 void FDynamicRHI::RHIUpdateTextureReference(FRHICommandListBase& RHICmdList, FRHITextureReference* TextureRef, FRHITexture* InReferencedTexture)
 {
 	FRHITexture* ReferencedTexture = InReferencedTexture ? InReferencedTexture : FRHITextureReference::GetDefaultTexture();
diff --git a/Engine/Source/Runtime/RHI/Public/DynamicRHI.h b/Engine/Source/Runtime/RHI/Public/DynamicRHI.h
--- a/Engine/Source/Runtime/RHI/Public/DynamicRHI.h
+++ b/Engine/Source/Runtime/RHI/Public/DynamicRHI.h
@@ -329,7 +329,7 @@ public:
 	 * @param DestBuffer - the buffer to update
 	 * @param SrcBuffer - don't use after call. If null, will release any resource owned by DestBuffer
 	 */
-	RHI_API virtual void RHITransferBufferUnderlyingResource(FRHIBuffer* DestBuffer, FRHIBuffer* SrcBuffer);
+	virtual void RHITransferBufferUnderlyingResource(FRHICommandListBase& CmdList, FRHIBuffer* DestBuffer, FRHIBuffer* SrcBuffer) = 0;
 
 	/**
 	* @param ResourceArray - An optional pointer to a resource array containing the resource's data.
@@ -977,7 +977,7 @@ public:
 		return nullptr;
 	}
 
-	virtual void RHITransferRayTracingGeometryUnderlyingResource(FRHIRayTracingGeometry* DestGeometry, FRHIRayTracingGeometry* SrcGeometry)
+	virtual void RHITransferRayTracingGeometryUnderlyingResource(FRHICommandListBase& RHICmdList, FRHIRayTracingGeometry* DestGeometry, FRHIRayTracingGeometry* SrcGeometry)
 	{
 		checkNoEntry();
 	}
diff --git a/Engine/Source/Runtime/RHI/Public/RHICommandListCommandExecutes.inl b/Engine/Source/Runtime/RHI/Public/RHICommandListCommandExecutes.inl
--- a/Engine/Source/Runtime/RHI/Public/RHICommandListCommandExecutes.inl
+++ b/Engine/Source/Runtime/RHI/Public/RHICommandListCommandExecutes.inl
@@ -592,12 +592,14 @@ void FRHICommandUpdateRHIResources::Execute(FRHICommandListBase& CmdList)
 		{
 		case FRHIResourceUpdateInfo::UT_Buffer:
 			GDynamicRHI->RHITransferBufferUnderlyingResource(
+				CmdList,
 				Info.Buffer.DestBuffer,
 				Info.Buffer.SrcBuffer);
 			break;
 #if RHI_RAYTRACING
 		case FRHIResourceUpdateInfo::UT_RayTracingGeometry:
 			GDynamicRHI->RHITransferRayTracingGeometryUnderlyingResource(
+				CmdList,
 				Info.RayTracingGeometry.DestGeometry,
 				Info.RayTracingGeometry.SrcGeometry);
 			break;
diff --git a/Engine/Source/Runtime/RHI/Public/RHIValidation.h b/Engine/Source/Runtime/RHI/Public/RHIValidation.h
--- a/Engine/Source/Runtime/RHI/Public/RHIValidation.h
+++ b/Engine/Source/Runtime/RHI/Public/RHIValidation.h
@@ -525,9 +525,9 @@ public:
 		return Texture;
 	}
 
-	void RHITransferBufferUnderlyingResource(FRHIBuffer* DestBuffer, FRHIBuffer* SrcBuffer) override final
+	void RHITransferBufferUnderlyingResource(FRHICommandListBase& RHICmdList, FRHIBuffer* DestBuffer, FRHIBuffer* SrcBuffer) override final
 	{
-		RHI->RHITransferBufferUnderlyingResource(DestBuffer, SrcBuffer);
+		RHI->RHITransferBufferUnderlyingResource(RHICmdList, DestBuffer, SrcBuffer);
 	}
 
 	/**
@@ -1276,9 +1276,9 @@ public:
 		return RHI->RHICalcRayTracingGeometrySize(RHICmdList, Initializer);
 	}
 
-	void RHITransferRayTracingGeometryUnderlyingResource(FRHIRayTracingGeometry* DestGeometry, FRHIRayTracingGeometry* SrcGeometry) override final
+	void RHITransferRayTracingGeometryUnderlyingResource(FRHICommandListBase& RHICmdList, FRHIRayTracingGeometry* DestGeometry, FRHIRayTracingGeometry* SrcGeometry) override final
 	{
-		RHI->RHITransferRayTracingGeometryUnderlyingResource(DestGeometry, SrcGeometry);
+		RHI->RHITransferRayTracingGeometryUnderlyingResource(RHICmdList, DestGeometry, SrcGeometry);
 	}
 #endif // RHI_RAYTRACING
 
diff --git a/Engine/Source/Runtime/VulkanRHI/Private/VulkanBuffer.cpp b/Engine/Source/Runtime/VulkanRHI/Private/VulkanBuffer.cpp
--- a/Engine/Source/Runtime/VulkanRHI/Private/VulkanBuffer.cpp
+++ b/Engine/Source/Runtime/VulkanRHI/Private/VulkanBuffer.cpp
@@ -632,7 +632,7 @@ void FVulkanDynamicRHI::RHICopyBuffer(FRHIBuffer* SourceBufferRHI, FRHIBuffer* D
 	VULKAN_SIGNAL_UNIMPLEMENTED();
 }
 
-void FVulkanDynamicRHI::RHITransferBufferUnderlyingResource(FRHIBuffer* DestBuffer, FRHIBuffer* SrcBuffer)
+void FVulkanDynamicRHI::RHITransferBufferUnderlyingResource(FRHICommandListBase& RHICmdList, FRHIBuffer* DestBuffer, FRHIBuffer* SrcBuffer)
 {
 	FVulkanResourceMultiBuffer* Dst = ResourceCast(DestBuffer);
 	FVulkanResourceMultiBuffer* Src = ResourceCast(SrcBuffer);
diff --git a/Engine/Source/Runtime/VulkanRHI/Private/VulkanRayTracing.cpp b/Engine/Source/Runtime/VulkanRHI/Private/VulkanRayTracing.cpp
--- a/Engine/Source/Runtime/VulkanRHI/Private/VulkanRayTracing.cpp
+++ b/Engine/Source/Runtime/VulkanRHI/Private/VulkanRayTracing.cpp
@@ -1053,7 +1053,7 @@ FVulkanRayTracingShaderTable* FVulkanRayTracingScene::FindOrCreateShaderTable(co
 	return CreatedShaderTable;
 }
 
-void FVulkanDynamicRHI::RHITransferRayTracingGeometryUnderlyingResource(FRHIRayTracingGeometry* DestGeometry, FRHIRayTracingGeometry* SrcGeometry)
+void FVulkanDynamicRHI::RHITransferRayTracingGeometryUnderlyingResource(FRHICommandListBase& RHICmdList, FRHIRayTracingGeometry* DestGeometry, FRHIRayTracingGeometry* SrcGeometry)
 {
 	check(DestGeometry);
 	FVulkanRayTracingGeometry* Dest = ResourceCast(DestGeometry);
diff --git a/Engine/Source/Runtime/VulkanRHI/Public/VulkanDynamicRHI.h b/Engine/Source/Runtime/VulkanRHI/Public/VulkanDynamicRHI.h
--- a/Engine/Source/Runtime/VulkanRHI/Public/VulkanDynamicRHI.h
+++ b/Engine/Source/Runtime/VulkanRHI/Public/VulkanDynamicRHI.h
@@ -113,7 +113,7 @@ public:
 	virtual void RHIUpdateUniformBuffer(FRHICommandListBase& RHICmdList, FRHIUniformBuffer* UniformBufferRHI, const void* Contents) final override;
 	virtual FBufferRHIRef RHICreateBuffer(FRHICommandListBase& RHICmdList, FRHIBufferDesc const& Desc, ERHIAccess ResourceState, FRHIResourceCreateInfo& CreateInfo) final override;
 	virtual void RHICopyBuffer(FRHIBuffer* SourceBuffer, FRHIBuffer* DestBuffer) final override;
-	virtual void RHITransferBufferUnderlyingResource(FRHIBuffer* DestBuffer, FRHIBuffer* SrcBuffer) final override;
+	virtual void RHITransferBufferUnderlyingResource(FRHICommandListBase& RHICmdList, FRHIBuffer* DestBuffer, FRHIBuffer* SrcBuffer) final override;
 	virtual void* LockBuffer_BottomOfPipe(FRHICommandListBase& RHICmdList, FRHIBuffer* Buffer, uint32 Offset, uint32 Size, EResourceLockMode LockMode) final override;
 	virtual void UnlockBuffer_BottomOfPipe(FRHICommandListBase& RHICmdList, FRHIBuffer* Buffer) final override;
 	virtual void RHIUnlockBuffer(FRHICommandListBase& RHICmdList, FRHIBuffer* Buffer) final override;
@@ -223,7 +223,7 @@ public:
 	virtual FRayTracingSceneRHIRef RHICreateRayTracingScene(FRayTracingSceneInitializer2 Initializer) final override;
 	virtual FRayTracingShaderRHIRef RHICreateRayTracingShader(TArrayView<const uint8> Code, const FSHAHash& Hash, EShaderFrequency ShaderFrequency) final override;
 	virtual FRayTracingPipelineStateRHIRef RHICreateRayTracingPipelineState(const FRayTracingPipelineStateInitializer& Initializer) final override;
-	virtual void RHITransferRayTracingGeometryUnderlyingResource(FRHIRayTracingGeometry* DestGeometry, FRHIRayTracingGeometry* SrcGeometry) final override;
+	virtual void RHITransferRayTracingGeometryUnderlyingResource(FRHICommandListBase& RHICmdList, FRHIRayTracingGeometry* DestGeometry, FRHIRayTracingGeometry* SrcGeometry) final override;
 #endif
 
 	// Historical number of times we've presented any and all viewports
diff --git a/Engine/Source/Runtime/Windows/D3D11RHI/Private/D3D11Buffer.cpp b/Engine/Source/Runtime/Windows/D3D11RHI/Private/D3D11Buffer.cpp
--- a/Engine/Source/Runtime/Windows/D3D11RHI/Private/D3D11Buffer.cpp
+++ b/Engine/Source/Runtime/Windows/D3D11RHI/Private/D3D11Buffer.cpp
@@ -278,7 +278,7 @@ void FD3D11DynamicRHI::RHICopyBuffer(FRHIBuffer* SourceBufferRHI, FRHIBuffer* De
 	GPUProfilingData.RegisterGPUWork(1);
 }
 
-void FD3D11DynamicRHI::RHITransferBufferUnderlyingResource(FRHIBuffer* DestBuffer, FRHIBuffer* SrcBuffer)
+void FD3D11DynamicRHI::RHITransferBufferUnderlyingResource(FRHICommandListBase& RHICmdList, FRHIBuffer* DestBuffer, FRHIBuffer* SrcBuffer)
 {
 	FD3D11Buffer* Dst = ResourceCast(DestBuffer);
 	FD3D11Buffer* Src = ResourceCast(SrcBuffer);
diff --git a/Engine/Source/Runtime/Windows/D3D11RHI/Private/D3D11RHIPrivate.h b/Engine/Source/Runtime/Windows/D3D11RHI/Private/D3D11RHIPrivate.h
--- a/Engine/Source/Runtime/Windows/D3D11RHI/Private/D3D11RHIPrivate.h
+++ b/Engine/Source/Runtime/Windows/D3D11RHI/Private/D3D11RHIPrivate.h
@@ -492,7 +492,7 @@ public:
 	virtual void RHIUpdateUniformBuffer(FRHICommandListBase& RHICmdList, FRHIUniformBuffer* UniformBufferRHI, const void* Contents) final override;
 	virtual FBufferRHIRef RHICreateBuffer(FRHICommandListBase& RHICmdList, FRHIBufferDesc const& Desc, ERHIAccess ResourceState, FRHIResourceCreateInfo& CreateInfo) final override;
 	virtual void RHICopyBuffer(FRHIBuffer* SourceBuffer, FRHIBuffer* DestBuffer) final override;
-	virtual void RHITransferBufferUnderlyingResource(FRHIBuffer* DestBuffer, FRHIBuffer* SrcBuffer) final override;
+	virtual void RHITransferBufferUnderlyingResource(FRHICommandListBase& RHICmdList, FRHIBuffer* DestBuffer, FRHIBuffer* SrcBuffer) final override;
 	virtual void* LockBuffer_BottomOfPipe(FRHICommandListBase& RHICmdList, FRHIBuffer* Buffer, uint32 Offset, uint32 Size, EResourceLockMode LockMode) final override;
 	virtual void UnlockBuffer_BottomOfPipe(FRHICommandListBase& RHICmdList, FRHIBuffer* Buffer) final override;
 	virtual FRHICalcTextureSizeResult RHICalcTexturePlatformSize(FRHITextureDesc const& Desc, uint32 FirstMipIndex) final override;
