D3D12RHI: Fixed CreateDiagnosticBuffer to use MEM_COMMIT | MEM_RESERVE (instead of just MEM_COMMIT) with VirtualAlloc in order for memory tracing to correctly track these buffers.

diff --git a/Engine/Source/Runtime/D3D12RHI/Private/Windows/WindowsD3D12Device.cpp b/Engine/Source/Runtime/D3D12RHI/Private/Windows/WindowsD3D12Device.cpp
--- a/Engine/Source/Runtime/D3D12RHI/Private/Windows/WindowsD3D12Device.cpp
+++ b/Engine/Source/Runtime/D3D12RHI/Private/Windows/WindowsD3D12Device.cpp
@@ -2033,7 +2033,7 @@ TUniquePtr<FD3D12DiagnosticBuffer> FD3D12Device::CreateDiagnosticBuffer(const D3
 	HRESULT hr = GetDevice()->QueryInterface(IID_PPV_ARGS(D3D12Device3.GetInitReference()));
 	if (SUCCEEDED(hr))
 	{
-		void* BreadCrumbResourceAddress = VirtualAlloc(nullptr, Desc.Width, MEM_COMMIT, PAGE_READWRITE);
+		void* BreadCrumbResourceAddress = VirtualAlloc(nullptr, Desc.Width, MEM_COMMIT | MEM_RESERVE, PAGE_READWRITE);
 		if (BreadCrumbResourceAddress)
 		{
 			ID3D12Heap* D3D12Heap = nullptr;
