Added a d3d12 scope cycle counter for RHI dispatch shader bundle

diff --git a/Engine/Source/Runtime/D3D12RHI/Private/D3D12Commands.cpp b/Engine/Source/Runtime/D3D12RHI/Private/D3D12Commands.cpp
--- a/Engine/Source/Runtime/D3D12RHI/Private/D3D12Commands.cpp
+++ b/Engine/Source/Runtime/D3D12RHI/Private/D3D12Commands.cpp
@@ -1621,6 +1621,9 @@ void FD3D12CommandContext::RHIDispatchShaderBundle(
 	bool bEmulated
 )
 {
+	TRACE_CPUPROFILER_EVENT_SCOPE(RHIDispatchShaderBundle);
+	SCOPE_CYCLE_COUNTER(STAT_D3D12DispatchShaderBundle);
+
 	check(ShaderBundle != nullptr && Dispatches.Num() > 0);
 	TRHICommandList_RecursiveHazardous<FD3D12CommandContext> RHICmdList(this);
 	UE::RHICore::DispatchShaderBundleEmulation(RHICmdList, ShaderBundle, RecordArgBufferSRV->GetBuffer(), Dispatches);
diff --git a/Engine/Source/Runtime/D3D12RHI/Private/D3D12Stats.h b/Engine/Source/Runtime/D3D12RHI/Private/D3D12Stats.h
--- a/Engine/Source/Runtime/D3D12RHI/Private/D3D12Stats.h
+++ b/Engine/Source/Runtime/D3D12RHI/Private/D3D12Stats.h
@@ -41,6 +41,8 @@ DECLARE_CYCLE_STAT_WITH_FLAGS_EXTERN(TEXT("Update uniform buffer"), STAT_D3D12Up
 DECLARE_CYCLE_STAT_EXTERN(TEXT("Commit resource tables"), STAT_D3D12CommitResourceTables, STATGROUP_D3D12RHI, );
 DECLARE_DWORD_COUNTER_STAT_EXTERN(TEXT("Num textures in tables"), STAT_D3D12SetTextureInTableCalls, STATGROUP_D3D12RHI, );
 
+DECLARE_CYCLE_STAT_EXTERN(TEXT("Dispatch shader bundle time"), STAT_D3D12DispatchShaderBundle, STATGROUP_D3D12RHI, );
+
 DECLARE_CYCLE_STAT_WITH_FLAGS_EXTERN(TEXT("Clear SRVs time"), STAT_D3D12ClearShaderResourceViewsTime, STATGROUP_D3D12RHI, EStatFlags::Verbose, );
 DECLARE_CYCLE_STAT_WITH_FLAGS_EXTERN(TEXT("Set SRV time"), STAT_D3D12SetShaderResourceViewTime, STATGROUP_D3D12RHI, EStatFlags::Verbose, );
 DECLARE_CYCLE_STAT_WITH_FLAGS_EXTERN(TEXT("Set UAV time"), STAT_D3D12SetUnorderedAccessViewTime, STATGROUP_D3D12RHI, EStatFlags::Verbose, );
diff --git a/Engine/Source/Runtime/D3D12RHI/Private/D3D12Util.cpp b/Engine/Source/Runtime/D3D12RHI/Private/D3D12Util.cpp
--- a/Engine/Source/Runtime/D3D12RHI/Private/D3D12Util.cpp
+++ b/Engine/Source/Runtime/D3D12RHI/Private/D3D12Util.cpp
@@ -1767,6 +1767,8 @@ DEFINE_STAT(STAT_D3D12UpdateUniformBufferTime);
 DEFINE_STAT(STAT_D3D12CommitResourceTables);
 DEFINE_STAT(STAT_D3D12SetTextureInTableCalls);
 
+DEFINE_STAT(STAT_D3D12DispatchShaderBundle);
+
 DEFINE_STAT(STAT_D3D12ClearShaderResourceViewsTime);
 DEFINE_STAT(STAT_D3D12SetShaderResourceViewTime);
 DEFINE_STAT(STAT_D3D12SetUnorderedAccessViewTime);
