Add missed files in last submit

diff --git a/Engine/Config/Windows/DataDrivenPlatformInfo.ini b/Engine/Config/Windows/DataDrivenPlatformInfo.ini
--- a/Engine/Config/Windows/DataDrivenPlatformInfo.ini
+++ b/Engine/Config/Windows/DataDrivenPlatformInfo.ini
@@ -118,6 +118,7 @@ bSupportsUInt64ImageAtomics=true
 bSupportsAnisotropicMaterials=true
 bSupportsVariableRateShading=true
 bSupportsGen5TemporalAA=true
+bSupportsShaderRootConstants=true
 bSupportsMeshShadersTier0=true
 bSupportsMeshShadersTier1=true
 bSupportsMeshShadersWithClipDistance=true
diff --git a/Engine/Source/Developer/Windows/ShaderFormatD3D/Private/D3DShaderCompiler.inl b/Engine/Source/Developer/Windows/ShaderFormatD3D/Private/D3DShaderCompiler.inl
--- a/Engine/Source/Developer/Windows/ShaderFormatD3D/Private/D3DShaderCompiler.inl
+++ b/Engine/Source/Developer/Windows/ShaderFormatD3D/Private/D3DShaderCompiler.inl
@@ -32,7 +32,8 @@ template <typename ID3D1xShaderReflection, typename D3D1x_SHADER_DESC, typename
 			ID3D1xShaderReflectionConstantBuffer* ConstantBuffer = Reflector->GetConstantBufferByName(BindDesc.Name);
 			D3D1x_SHADER_BUFFER_DESC CBDesc;
 			ConstantBuffer->GetDesc(&CBDesc);
-			bool bGlobalCB = (FCStringAnsi::Strcmp(CBDesc.Name, "$Globals") == 0);
+			const bool bGlobalCB = (FCStringAnsi::Strcmp(CBDesc.Name, "$Globals") == 0);
+			const bool bRootConstantsCB = (FCStringAnsi::Strcmp(CBDesc.Name, "UERootConstants") == 0);
 			const bool bIsRootCB = FCString::Strcmp(ANSI_TO_TCHAR(CBDesc.Name), FShaderParametersMetadata::kRootUniformBufferBindingName) == 0;
 
 			if (bGlobalCB)
@@ -81,6 +82,11 @@ template <typename ID3D1xShaderReflection, typename D3D1x_SHADER_DESC, typename
 					}
 				}
 			}
+			else if (bRootConstantsCB)
+			{
+				// For the UERootConstants root constant CB, we want to fully skip adding it to the parameter map, or
+				// updating the used slots or num CBs (all those assume space0).
+			}
 			else if (bIsRootCB && Input.ShouldUseStableConstantBuffer())
 			{
 				if (CBIndex == FShaderParametersMetadata::kRootCBufferBindingIndex)
diff --git a/Engine/Source/Runtime/D3D12RHI/Private/D3D12StateCache.cpp b/Engine/Source/Runtime/D3D12RHI/Private/D3D12StateCache.cpp
--- a/Engine/Source/Runtime/D3D12RHI/Private/D3D12StateCache.cpp
+++ b/Engine/Source/Runtime/D3D12RHI/Private/D3D12StateCache.cpp
@@ -165,6 +165,7 @@ void FD3D12StateCache::DirtyStateForNewCommandList()
 
 	// Always need to set PSOs and root signatures
 	PipelineState.Common.bNeedSetPSO = true;
+	PipelineState.Common.bNeedSetRootConstants = true;
 	PipelineState.Compute.bNeedSetRootSignature = true;
 	PipelineState.Graphics.bNeedSetRootSignature = true;
 	bNeedSetPrimitiveTopology = true;
@@ -209,6 +210,7 @@ void FD3D12StateCache::DirtyState()
 {
 	// Mark bits dirty so the next call to ApplyState will set all this state again
 	PipelineState.Common.bNeedSetPSO = true;
+	PipelineState.Common.bNeedSetRootConstants = true;
 	PipelineState.Compute.bNeedSetRootSignature = true;
 	PipelineState.Graphics.bNeedSetRootSignature = true;
 	bNeedSetVB = true;
@@ -337,6 +339,7 @@ bool FD3D12StateCache::InternalSetRootSignature(ED3D12PipelineType InPipelineTyp
 			PipelineState.Common.UAVCache.DirtyCompute();
 			PipelineState.Common.SamplerCache.DirtyCompute();
 			PipelineState.Common.CBVCache.DirtyCompute();
+			PipelineState.Common.bNeedSetRootConstants = true;
 
 			bWasRootSignatureChanged = true;
 		}
@@ -354,6 +357,7 @@ bool FD3D12StateCache::InternalSetRootSignature(ED3D12PipelineType InPipelineTyp
 			PipelineState.Common.UAVCache.DirtyGraphics();
 			PipelineState.Common.SamplerCache.DirtyGraphics();
 			PipelineState.Common.CBVCache.DirtyGraphics();
+			PipelineState.Common.bNeedSetRootConstants = true;
 
 			bWasRootSignatureChanged = true;
 		}
@@ -550,6 +554,26 @@ void FD3D12StateCache::ApplyState(ED3D12PipelineType PipelineType)
 
 	ApplyConstants(PSOCommonData->RootSignature, StartStage, EndStage);
 
+	int8 RootConstantsSlot = PSOCommonData->RootSignature->GetRootConstantsSlot();
+	if (PipelineState.Common.bNeedSetRootConstants && RootConstantsSlot >= 0)
+	{
+		PipelineState.Common.bNeedSetRootConstants = false;
+
+		if (PipelineType == ED3D12PipelineType::Compute)
+		{
+			uint32 UERootConstants[4];
+			UERootConstants[0] = PipelineState.Common.ShaderRootConstants.X;
+			UERootConstants[1] = PipelineState.Common.ShaderRootConstants.Y;
+			UERootConstants[2] = PipelineState.Common.ShaderRootConstants.Z;
+			UERootConstants[3] = PipelineState.Common.ShaderRootConstants.W;
+			CmdContext.GraphicsCommandList()->SetComputeRoot32BitConstants(RootConstantsSlot, 4, &UERootConstants[0], 0);
+		}
+		else
+		{
+			checkNoEntry();
+		}
+	}
+
 	// Flush any needed resource barriers
 	CmdContext.FlushResourceBarriers();
 
@@ -1115,6 +1139,15 @@ bool FD3D12StateCache::AssertResourceStates(ED3D12PipelineType PipelineType)
 }
 #endif
 
+void FD3D12StateCache::SetRootConstants(const FUint32Vector4& Constants)
+{
+	if (Constants != PipelineState.Common.ShaderRootConstants)
+	{
+		PipelineState.Common.ShaderRootConstants = Constants;
+		PipelineState.Common.bNeedSetRootConstants = true;
+	}
+}
+
 void FD3D12StateCache::ClearUAVs(EShaderFrequency ShaderStage)
 {
 	FD3D12UnorderedAccessViewCache& Cache = PipelineState.Common.UAVCache;
diff --git a/Engine/Source/Runtime/D3D12RHI/Private/D3D12StateCachePrivate.h b/Engine/Source/Runtime/D3D12RHI/Private/D3D12StateCachePrivate.h
--- a/Engine/Source/Runtime/D3D12RHI/Private/D3D12StateCachePrivate.h
+++ b/Engine/Source/Runtime/D3D12RHI/Private/D3D12StateCachePrivate.h
@@ -389,6 +389,10 @@ protected:
 			ID3D12PipelineState* CurrentPipelineStateObject = nullptr;
 			bool bNeedSetPSO;
 
+			// Root Constants
+			bool bNeedSetRootConstants;
+			FUint32Vector4 ShaderRootConstants = FUint32Vector4(0, 0, 0, 0);
+
 			uint32 CurrentShaderSamplerCounts[SF_NumStandardFrequencies] = {};
 			uint32 CurrentShaderSRVCounts    [SF_NumStandardFrequencies] = {};
 			uint32 CurrentShaderCBCounts     [SF_NumStandardFrequencies] = {};
@@ -640,6 +644,7 @@ public:
 		if (LastComputePipelineType != PipelineType)
 		{
 			PipelineState.Common.bNeedSetPSO = true;
+			PipelineState.Common.bNeedSetRootConstants = true;
 			PipelineState.Compute.bNeedSetRootSignature = true;
 
 			LastComputePipelineType = PipelineType;
@@ -675,6 +680,8 @@ public:
 		}
 	}
 
+	void SetRootConstants(const FUint32Vector4& Constants);
+
 	void SetUAV(EShaderFrequency ShaderStage, uint32 SlotIndex, FD3D12UnorderedAccessView* UAV, uint32 InitialCount = -1);
 
 	void ClearUAVs(EShaderFrequency ShaderStage);
