Disallow async compute in mGPU mode because it causes CPU hangs (likely due to command lists being submitted in the wrong order).

#rnx

diff --git a/Engine/Source/Runtime/D3D12RHI/Private/Windows/WindowsD3D12Device.cpp b/Engine/Source/Runtime/D3D12RHI/Private/Windows/WindowsD3D12Device.cpp
--- a/Engine/Source/Runtime/D3D12RHI/Private/Windows/WindowsD3D12Device.cpp
+++ b/Engine/Source/Runtime/D3D12RHI/Private/Windows/WindowsD3D12Device.cpp
@@ -1664,6 +1664,14 @@ void FD3D12DynamicRHI::Init()
 
 	GSupportsEfficientAsyncCompute = GAllowAsyncCompute && (FParse::Param(FCommandLine::Get(), TEXT("ForceAsyncCompute")) || (GRHISupportsParallelRHIExecute && (IsRHIDeviceAMD() || bnVidiaAsyncComputeSupported)));
 
+#if WITH_MGPU
+	// Disallow async compute in mGPU mode due to submission order bugs (UE-193929).
+	if (GNumExplicitGPUsForRendering > 1)
+	{
+		GSupportsEfficientAsyncCompute = false;
+	}
+#endif
+
 	GSupportsDepthBoundsTest = SupportsDepthBoundsTest(this);
 
 	{
