Add optional extra memory for diagnostic buffer.

diff --git a/Engine/Source/Runtime/D3D12RHI/Private/D3D12Device.cpp b/Engine/Source/Runtime/D3D12RHI/Private/D3D12Device.cpp
--- a/Engine/Source/Runtime/D3D12RHI/Private/D3D12Device.cpp
+++ b/Engine/Source/Runtime/D3D12RHI/Private/D3D12Device.cpp
@@ -13,6 +13,13 @@ static TAutoConsoleVariable<int32> CVarD3D12GPUTimeout(
 	ECVF_ReadOnly
 );
 
+static TAutoConsoleVariable<int32> CVarD3D12ExtraDiagnosticBufferMemory(
+	TEXT("r.D3D12.DiagnosticBufferExtraMemory"),
+	0,
+	TEXT("Extra allocated memory for diagnostic buffer"),
+	ECVF_ReadOnly
+);
+
 static uint32 GetQueryHeapPoolIndex(D3D12_QUERY_HEAP_TYPE HeapType)
 {
 	switch (HeapType)
@@ -62,7 +69,7 @@ void FD3D12Queue::SetupAfterDeviceCreation()
 		HRESULT hr = Device->GetDevice()->QueryInterface(IID_PPV_ARGS(D3D12Device3.GetInitReference()));
 		if (SUCCEEDED(hr))
 		{
-			const uint32 ShaderDiagnosticBufferSize = sizeof(FD3D12DiagnosticBufferData);
+			const uint32 ShaderDiagnosticBufferSize = sizeof(FD3D12DiagnosticBufferData) + FMath::Max(0, CVarD3D12ExtraDiagnosticBufferMemory.GetValueOnAnyThread());
 
 			// Allocate persistent CPU readable memory which will still be valid after a device lost and wrap this data in a placed resource
 			// so the GPU command list can write to it
