Fix acceleration structures not being released when mesh streaming request is cancelled.
- when request is cancelled (see FStaticMeshStreamIn::DoFinishUpdate()), FD3D12RayTracingGeometry marked as ERayTracingGeometryInitializerType::StreamingSource is released while still owning AccelerationStructureBuffers. However FD3D12RayTracingGeometry::ReleaseUnderlyingResource() wouldn't actually release those buffers.

diff --git a/Engine/Source/Runtime/D3D12RHI/Private/D3D12RayTracing.cpp b/Engine/Source/Runtime/D3D12RHI/Private/D3D12RayTracing.cpp
--- a/Engine/Source/Runtime/D3D12RHI/Private/D3D12RayTracing.cpp
+++ b/Engine/Source/Runtime/D3D12RHI/Private/D3D12RayTracing.cpp
@@ -2981,27 +2981,27 @@ void FD3D12RayTracingGeometry::ReleaseUnderlyingResource()
 		UnregisterAsRenameListener(GPUIndex);
 	}
 
-	if (Initializer.Type != ERayTracingGeometryInitializerType::StreamingSource)
+	for (TRefCountPtr<FD3D12Buffer>& Buffer : AccelerationStructureBuffers)
 	{
-		for (TRefCountPtr<FD3D12Buffer>& Buffer : AccelerationStructureBuffers)
+		if (Buffer)
 		{
-			if (Buffer)
-			{
-				DEC_MEMORY_STAT_BY(STAT_D3D12RayTracingUsedVideoMemory, Buffer->GetSize());
-				DEC_MEMORY_STAT_BY(STAT_D3D12RayTracingBLASMemory, Buffer->GetSize());
+			DEC_MEMORY_STAT_BY(STAT_D3D12RayTracingUsedVideoMemory, Buffer->GetSize());
+			DEC_MEMORY_STAT_BY(STAT_D3D12RayTracingBLASMemory, Buffer->GetSize());
 
-				ERayTracingAccelerationStructureFlags BuildFlags = GetRayTracingAccelerationStructureBuildFlags(Initializer);
-				if (EnumHasAllFlags(BuildFlags, ERayTracingAccelerationStructureFlags::AllowUpdate))
-				{
-					DEC_MEMORY_STAT_BY(STAT_D3D12RayTracingDynamicBLASMemory, Buffer->GetSize());
-				}
-				else
-				{
-					DEC_MEMORY_STAT_BY(STAT_D3D12RayTracingStaticBLASMemory, Buffer->GetSize());
-				}
+			ERayTracingAccelerationStructureFlags BuildFlags = GetRayTracingAccelerationStructureBuildFlags(Initializer);
+			if (EnumHasAllFlags(BuildFlags, ERayTracingAccelerationStructureFlags::AllowUpdate))
+			{
+				DEC_MEMORY_STAT_BY(STAT_D3D12RayTracingDynamicBLASMemory, Buffer->GetSize());
+			}
+			else
+			{
+				DEC_MEMORY_STAT_BY(STAT_D3D12RayTracingStaticBLASMemory, Buffer->GetSize());
 			}
 		}
+	}
 
+	if (Initializer.Type != ERayTracingGeometryInitializerType::StreamingSource)
+	{
 		DEC_DWORD_STAT_BY(STAT_D3D12RayTracingTrianglesBLAS, Initializer.TotalPrimitiveCount);
 		DEC_DWORD_STAT(STAT_D3D12RayTracingAllocatedBLAS);
 	}
