D3D12 cleanup
- Removing the old RHISetShaderXXX methods from D3D12

diff --git a/Engine/Source/Runtime/D3D12RHI/Private/D3D12CommandContext.h b/Engine/Source/Runtime/D3D12RHI/Private/D3D12CommandContext.h
--- a/Engine/Source/Runtime/D3D12RHI/Private/D3D12CommandContext.h
+++ b/Engine/Source/Runtime/D3D12RHI/Private/D3D12CommandContext.h
@@ -579,14 +579,6 @@ public:
 	virtual void RHIDispatchComputeShader(uint32 ThreadGroupCountX, uint32 ThreadGroupCountY, uint32 ThreadGroupCountZ) final override;
 	virtual void RHIDispatchIndirectComputeShader(FRHIBuffer* ArgumentBuffer, uint32 ArgumentOffset) final override;
 	virtual void RHISetStaticUniformBuffers(const FUniformBufferStaticBindings& InUniformBuffers) final override;
-	void RHISetShaderTexture(FRHIComputeShader* PixelShader, uint32 TextureIndex, FRHITexture* NewTexture);
-	void RHISetShaderSampler(FRHIComputeShader* ComputeShader, uint32 SamplerIndex, FRHISamplerState* NewState);
-	void RHISetUAVParameter(FRHIPixelShader* PixelShaderRHI, uint32 UAVIndex, FRHIUnorderedAccessView* UAVRHI);
-	void RHISetUAVParameter(FRHIComputeShader* ComputeShader, uint32 UAVIndex, FRHIUnorderedAccessView* UAV);
-	void RHISetUAVParameter(FRHIComputeShader* ComputeShader, uint32 UAVIndex, FRHIUnorderedAccessView* UAV, uint32 InitialCount);
-	void RHISetShaderResourceViewParameter(FRHIComputeShader* ComputeShader, uint32 SamplerIndex, FRHIShaderResourceView* SRV);
-	void RHISetShaderUniformBuffer(FRHIComputeShader* ComputeShader, uint32 BufferIndex, FRHIUniformBuffer* Buffer);
-	void RHISetShaderParameter(FRHIComputeShader* ComputeShader, uint32 BufferIndex, uint32 BaseIndex, uint32 NumBytes, const void* NewValue);
 	virtual void RHISetShaderParameters(FRHIComputeShader* Shader, TConstArrayView<uint8> InParametersData, TConstArrayView<FRHIShaderParameter> InParameters, TConstArrayView<FRHIShaderParameterResource> InResourceParameters, TConstArrayView<FRHIShaderParameterResource> InBindlessParameters) final override;
 	virtual void RHISetShaderUnbinds(FRHIComputeShader* Shader, TConstArrayView<FRHIShaderParameterUnbind> InUnbinds) final override;
 	virtual void RHISetShaderUnbinds(FRHIGraphicsShader* Shader, TConstArrayView<FRHIShaderParameterUnbind> InUnbinds) final override;
@@ -611,11 +603,6 @@ public:
 	virtual void RHISetStereoViewport(float LeftMinX, float RightMinX, float LeftMinY, float RightMinY, float MinZ, float LeftMaxX, float RightMaxX, float LeftMaxY, float RightMaxY, float MaxZ) override;
 	virtual void RHISetScissorRect(bool bEnable, uint32 MinX, uint32 MinY, uint32 MaxX, uint32 MaxY) final override;
 	virtual void RHISetGraphicsPipelineState(FRHIGraphicsPipelineState* GraphicsPipelineState, uint32 StencilRef, bool bApplyAdditionalState) final override;
-	void RHISetShaderTexture(FRHIGraphicsShader* Shader, uint32 TextureIndex, FRHITexture* NewTexture);
-	void RHISetShaderSampler(FRHIGraphicsShader* Shader, uint32 SamplerIndex, FRHISamplerState* NewState);
-	void RHISetShaderResourceViewParameter(FRHIGraphicsShader* Shader, uint32 SamplerIndex, FRHIShaderResourceView* SRV);
-	void RHISetShaderUniformBuffer(FRHIGraphicsShader* Shader, uint32 BufferIndex, FRHIUniformBuffer* Buffer);
-	void RHISetShaderParameter(FRHIGraphicsShader* Shader, uint32 BufferIndex, uint32 BaseIndex, uint32 NumBytes, const void* NewValue);
 	virtual void RHISetShaderParameters(FRHIGraphicsShader* Shader, TConstArrayView<uint8> InParametersData, TConstArrayView<FRHIShaderParameter> InParameters, TConstArrayView<FRHIShaderParameterResource> InResourceParameters, TConstArrayView<FRHIShaderParameterResource> InBindlessParameters) final override;
 	virtual void RHISetStencilRef(uint32 StencilRef) final override;
 	virtual void RHISetBlendFactor(const FLinearColor& BlendFactor) final override;
@@ -759,15 +746,6 @@ private:
 
 	static void ClearUAV(TRHICommandList_RecursiveHazardous<FD3D12CommandContext>& RHICmdList, FD3D12UnorderedAccessView_RHI* UAV, const void* ClearValues, bool bFloat);
 
-	template <typename TRHIShader>
-	void ApplyStaticUniformBuffers(TRHIShader* Shader)
-	{
-		if (Shader)
-		{
-			UE::RHICore::ApplyStaticUniformBuffers(this, Shader, Shader->StaticSlots, Shader->ShaderResourceTable.ResourceTableLayoutHashes, StaticUniformBuffers);
-		}
-	}
-
 	void HandleDiscardResources          (TArrayView<const FRHITransition*> Transitions, bool bIsBeginTransition);
 	void HandleResourceTransitions       (const struct FD3D12TransitionData* TransitionData, bool& bUAVBarrier);
 	void HandleTransientAliasing         (const struct FD3D12TransitionData* TransitionData);
diff --git a/Engine/Source/Runtime/D3D12RHI/Private/D3D12Commands.cpp b/Engine/Source/Runtime/D3D12RHI/Private/D3D12Commands.cpp
--- a/Engine/Source/Runtime/D3D12RHI/Private/D3D12Commands.cpp
+++ b/Engine/Source/Runtime/D3D12RHI/Private/D3D12Commands.cpp
@@ -782,6 +782,32 @@ void FD3D12CommandContext::RHISetScissorRect(bool bEnable, uint32 MinX, uint32 M
 	}
 }
 
+static void ApplyStaticUniformBuffersOnContext(FD3D12CommandContext& Context, FRHIShader* Shader, FD3D12ShaderData* ShaderData)
+{
+	if (Shader)
+	{
+		const uint32 GpuIndex = Context.GetGPUIndex();
+		const EShaderFrequency ShaderFrequency = Shader->GetFrequency();
+
+		UE::RHICore::ApplyStaticUniformBuffers(
+			Shader,
+			ShaderData->StaticSlots,
+			ShaderData->ShaderResourceTable.ResourceTableLayoutHashes,
+			Context.GetStaticUniformBuffers(),
+			[&Context, Shader, ShaderFrequency, GpuIndex](int32 BufferIndex, FRHIUniformBuffer* Buffer)
+			{
+				BindUniformBuffer(Context, Shader, ShaderFrequency, BufferIndex, FD3D12CommandContext::RetrieveObject<FD3D12UniformBuffer>(Buffer, GpuIndex));
+			}
+		);
+	}
+}
+
+template<typename TShader>
+static void ApplyStaticUniformBuffersOnContext(FD3D12CommandContext& Context, TShader* Shader)
+{
+	ApplyStaticUniformBuffersOnContext(Context, Shader, static_cast<FD3D12ShaderData*>(Shader));
+}
+
 void FD3D12CommandContext::RHISetGraphicsPipelineState(FRHIGraphicsPipelineState* GraphicsState, uint32 StencilRef, bool bApplyAdditionalState)
 {
 	FD3D12GraphicsPipelineState* GraphicsPipelineState = FD3D12DynamicRHI::ResourceCast(GraphicsState);
@@ -811,11 +837,11 @@ void FD3D12CommandContext::RHISetGraphicsPipelineState(FRHIGraphicsPipelineState
 
 	if (bApplyAdditionalState)
 	{
-		ApplyStaticUniformBuffers(GraphicsPipelineState->GetVertexShader());
-		ApplyStaticUniformBuffers(GraphicsPipelineState->GetMeshShader());
-		ApplyStaticUniformBuffers(GraphicsPipelineState->GetAmplificationShader());
-		ApplyStaticUniformBuffers(GraphicsPipelineState->GetGeometryShader());
-		ApplyStaticUniformBuffers(GraphicsPipelineState->GetPixelShader());
+		ApplyStaticUniformBuffersOnContext(*this, GraphicsPipelineState->GetVertexShader());
+		ApplyStaticUniformBuffersOnContext(*this, GraphicsPipelineState->GetMeshShader());
+		ApplyStaticUniformBuffersOnContext(*this, GraphicsPipelineState->GetAmplificationShader());
+		ApplyStaticUniformBuffersOnContext(*this, GraphicsPipelineState->GetGeometryShader());
+		ApplyStaticUniformBuffersOnContext(*this, GraphicsPipelineState->GetPixelShader());
 	}
 }
 
@@ -832,35 +858,7 @@ void FD3D12CommandContext::RHISetComputePipelineState(FRHIComputePipelineState*
 
 	StateCache.SetComputePipelineState(ComputePipelineState);
 
-	ApplyStaticUniformBuffers(ComputePipelineState->ComputeShader.GetReference());
-}
-
-void FD3D12CommandContext::RHISetShaderTexture(FRHIGraphicsShader* ShaderRHI, uint32 TextureIndex, FRHITexture* NewTextureRHI)
-{
-	const EShaderFrequency ShaderFrequency = ShaderRHI->GetFrequency();
-	if (IsValidGraphicsFrequency(ShaderFrequency))
-	{
-		ValidateBoundShader(StateCache, ShaderRHI);
-
-		FD3D12Texture* NewTexture = RetrieveTexture(NewTextureRHI);
-		FD3D12ShaderResourceView* ViewToSet = NewTexture ? NewTexture->GetShaderResourceView() : nullptr;
-
-		StateCache.SetShaderResourceView(ShaderFrequency, ViewToSet, TextureIndex);
-	}
-	else
-	{
-		checkf(0, TEXT("Unsupported FRHIGraphicsShader Type '%s'!"), GetShaderFrequencyString(ShaderFrequency, false));
-	}
-}
-
-void FD3D12CommandContext::RHISetShaderTexture(FRHIComputeShader* ComputeShaderRHI, uint32 TextureIndex, FRHITexture* NewTextureRHI)
-{
-	//ValidateBoundShader(StateCache, ComputeShaderRHI);
-
-	FD3D12Texture* const NewTexture = RetrieveTexture(NewTextureRHI);
-	FD3D12ShaderResourceView* ViewToSet = NewTexture ? NewTexture->GetShaderResourceView() : nullptr;
-
-	StateCache.SetShaderResourceView(SF_Compute, ViewToSet, TextureIndex);
+	ApplyStaticUniformBuffersOnContext(*this, ComputePipelineState->ComputeShader.GetReference());
 }
 
 void FD3D12CommandContext::SetUAVParameter(EShaderFrequency Frequency, uint32 UAVIndex, FD3D12UnorderedAccessView* UAV)
@@ -875,92 +873,11 @@ void FD3D12CommandContext::SetUAVParameter(EShaderFrequency Frequency, uint32 UA
 	StateCache.SetUAV(SF_Pixel, UAVIndex, UAV, InitialCount);
 }
 
-void FD3D12CommandContext::RHISetUAVParameter(FRHIPixelShader* PixelShaderRHI, uint32 UAVIndex, FRHIUnorderedAccessView* UAVRHI)
-{
-	SetUAVParameter(SF_Pixel, UAVIndex, RetrieveObject<FD3D12UnorderedAccessView_RHI>(UAVRHI));
-}
-
-void FD3D12CommandContext::RHISetUAVParameter(FRHIComputeShader* ComputeShaderRHI, uint32 UAVIndex, FRHIUnorderedAccessView* UAVRHI)
-{
-	//ValidateBoundShader(StateCache, ComputeShaderRHI);
-	SetUAVParameter(SF_Compute, UAVIndex, RetrieveObject<FD3D12UnorderedAccessView_RHI>(UAVRHI));
-}
-
-void FD3D12CommandContext::RHISetUAVParameter(FRHIComputeShader* ComputeShaderRHI, uint32 UAVIndex, FRHIUnorderedAccessView* UAVRHI, uint32 InitialCount)
-{
-	//ValidateBoundShader(StateCache, ComputeShaderRHI);
-	SetUAVParameter(SF_Compute, UAVIndex, RetrieveObject<FD3D12UnorderedAccessView_RHI>(UAVRHI), InitialCount);
-}
-
 void FD3D12CommandContext::SetSRVParameter(EShaderFrequency Frequency, uint32 SRVIndex, FD3D12ShaderResourceView* SRV)
 {
 	StateCache.SetShaderResourceView(Frequency, SRV, SRVIndex);
 }
 
-void FD3D12CommandContext::RHISetShaderResourceViewParameter(FRHIGraphicsShader* ShaderRHI, uint32 TextureIndex, FRHIShaderResourceView* SRVRHI)
-{
-	const EShaderFrequency ShaderFrequency = ShaderRHI->GetFrequency();
-	checkf(IsValidGraphicsFrequency(ShaderFrequency), TEXT("Unsupported FRHIGraphicsShader Type '%s'!"), GetShaderFrequencyString(ShaderFrequency, false));
-
-	ValidateBoundShader(StateCache, ShaderRHI);
-	SetSRVParameter(ShaderFrequency, TextureIndex, RetrieveObject<FD3D12ShaderResourceView_RHI>(SRVRHI));
-}
-
-void FD3D12CommandContext::RHISetShaderResourceViewParameter(FRHIComputeShader* ComputeShaderRHI, uint32 TextureIndex, FRHIShaderResourceView* SRVRHI)
-{
-	//ValidateBoundShader(StateCache, ComputeShaderRHI);
-	SetSRVParameter(SF_Compute, TextureIndex, RetrieveObject<FD3D12ShaderResourceView_RHI>(SRVRHI));
-}
-
-void FD3D12CommandContext::RHISetShaderSampler(FRHIGraphicsShader* ShaderRHI, uint32 SamplerIndex, FRHISamplerState* NewStateRHI)
-{
-	const EShaderFrequency ShaderFrequency = ShaderRHI->GetFrequency();
-
-	if (IsValidGraphicsFrequency(ShaderFrequency))
-	{
-		ValidateBoundShader(StateCache, ShaderRHI);
-
-		FD3D12SamplerState* NewState = RetrieveObject<FD3D12SamplerState>(NewStateRHI);
-
-		StateCache.SetSamplerState(ShaderFrequency, NewState, SamplerIndex);
-	}
-	else
-	{
-		checkf(0, TEXT("Unsupported FRHIGraphicsShader Type '%s'!"), GetShaderFrequencyString(ShaderFrequency, false));
-	}
-}
-
-void FD3D12CommandContext::RHISetShaderSampler(FRHIComputeShader* ComputeShaderRHI, uint32 SamplerIndex, FRHISamplerState* NewStateRHI)
-{
-	//ValidateBoundShader(StateCache, ComputeShaderRHI);
-	FD3D12SamplerState* NewState = RetrieveObject<FD3D12SamplerState>(NewStateRHI);
-	StateCache.SetSamplerState(SF_Compute, NewState, SamplerIndex);
-}
-
-void FD3D12CommandContext::RHISetShaderUniformBuffer(FRHIGraphicsShader* ShaderRHI, uint32 BufferIndex, FRHIUniformBuffer* BufferRHI)
-{
-	const EShaderFrequency ShaderFrequency = ShaderRHI->GetFrequency();
-
-	if (IsValidGraphicsFrequency(ShaderFrequency))
-	{
-		ValidateBoundShader(StateCache, ShaderRHI);
-
-		BindUniformBuffer(*this, ShaderRHI, ShaderFrequency, BufferIndex, RetrieveObject<FD3D12UniformBuffer>(BufferRHI));
-	}
-	else
-	{
-		checkf(0, TEXT("Unsupported FRHIGraphicsShader Type '%s'!"), GetShaderFrequencyString(ShaderFrequency, false));
-	}
-}
-
-void FD3D12CommandContext::RHISetShaderUniformBuffer(FRHIComputeShader* ComputeShader, uint32 BufferIndex, FRHIUniformBuffer* BufferRHI)
-{
-	//SCOPE_CYCLE_COUNTER(STAT_D3D12SetShaderUniformBuffer);
-	//ValidateBoundShader(StateCache, ComputeShader);
-
-	BindUniformBuffer(*this, ComputeShader, SF_Compute, BufferIndex, RetrieveObject<FD3D12UniformBuffer>(BufferRHI));
-}
-
 struct FD3D12ResourceBinder
 {
 	FD3D12CommandContext& Context;
@@ -1248,30 +1165,6 @@ void FD3D12CommandContext::RHISetShaderUnbinds(FRHIComputeShader* Shader, TConst
 	SetShaderUnbindsOnContext(*this, Shader, SF_Compute, InUnbinds);
 }
 
-void FD3D12CommandContext::RHISetShaderParameter(FRHIGraphicsShader* ShaderRHI, uint32 BufferIndex, uint32 Offset, uint32 NumBytes, const void* NewValue)
-{
-	checkSlow(BufferIndex == 0);
-
-	const EShaderFrequency ShaderFrequency = ShaderRHI->GetFrequency();
-	if (IsValidGraphicsFrequency(ShaderFrequency))
-	{
-		ValidateBoundShader(StateCache, ShaderRHI);
-
-		StageConstantBuffers[ShaderFrequency].UpdateConstant((const uint8*)NewValue, Offset, NumBytes);
-	}
-	else
-	{
-		checkf(0, TEXT("Unsupported FRHIGraphicsShader Type '%s'!"), GetShaderFrequencyString(ShaderFrequency, false));
-	}
-}
-
-void FD3D12CommandContext::RHISetShaderParameter(FRHIComputeShader* ComputeShaderRHI, uint32 BufferIndex, uint32 Offset, uint32 NumBytes, const void* NewValue)
-{
-	//ValidateBoundShader(StateCache, ComputeShaderRHI);
-	checkSlow(BufferIndex == 0);
-	StageConstantBuffers[SF_Compute].UpdateConstant((const uint8*)NewValue, Offset, NumBytes);
-}
-
 void FD3D12CommandContext::RHISetStencilRef(uint32 StencilRef)
 {
 	StateCache.SetStencilRef(StencilRef);
