Implemented XSX support for CFLAG_CheckForDerivativeOps

diff --git a/Engine/Source/Runtime/D3D12RHI/Private/D3D12Shader.h b/Engine/Source/Runtime/D3D12RHI/Private/D3D12Shader.h
--- a/Engine/Source/Runtime/D3D12RHI/Private/D3D12Shader.h
+++ b/Engine/Source/Runtime/D3D12RHI/Private/D3D12Shader.h
@@ -117,7 +117,6 @@ struct FD3D12ShaderData
 	FORCEINLINE bool UsesBindlessResources() const { return EnumHasAnyFlags(ResourceCounts.UsageFlags, EShaderResourceUsageFlags::BindlessResources); }
 	FORCEINLINE bool UsesBindlessSamplers() const { return EnumHasAnyFlags(ResourceCounts.UsageFlags, EShaderResourceUsageFlags::BindlessSamplers); }
 	FORCEINLINE bool UsesRootConstants() const { return EnumHasAnyFlags(ResourceCounts.UsageFlags, EShaderResourceUsageFlags::RootConstants); }
-	FORCEINLINE bool NoDerivativeOps() const { return EnumHasAnyFlags(ResourceCounts.UsageFlags, EShaderResourceUsageFlags::NoDerivativeOps); }
 
 	bool InitCommon(TArrayView<const uint8> InCode);
 };
diff --git a/Engine/Source/Runtime/D3D12RHI/Private/D3D12Shaders.cpp b/Engine/Source/Runtime/D3D12RHI/Private/D3D12Shaders.cpp
--- a/Engine/Source/Runtime/D3D12RHI/Private/D3D12Shaders.cpp
+++ b/Engine/Source/Runtime/D3D12RHI/Private/D3D12Shaders.cpp
@@ -160,6 +160,7 @@ FComputeShaderRHIRef FD3D12DynamicRHI::RHICreateComputeShader(TArrayView<const u
 	if (Shader)
 	{
 		Shader->RootSignature = GetAdapter().GetRootSignature(Shader);
+		Shader->SetNoDerivativeOps(EnumHasAnyFlags(Shader->ResourceCounts.UsageFlags, EShaderResourceUsageFlags::NoDerivativeOps));
 	}
 
 	return Shader;
