Adding D3D12RHI_USE_CONSTANT_BUFFER_VIEWS to untie CBVs from USE_STATIC_ROOT_SIGNATURE. These will eventually be used for bindless rendering.

diff --git a/Engine/Source/Runtime/D3D12RHI/Private/D3D12ConstantBuffer.cpp b/Engine/Source/Runtime/D3D12RHI/Private/D3D12ConstantBuffer.cpp
--- a/Engine/Source/Runtime/D3D12RHI/Private/D3D12ConstantBuffer.cpp
+++ b/Engine/Source/Runtime/D3D12RHI/Private/D3D12ConstantBuffer.cpp
@@ -12,14 +12,14 @@ FD3D12ConstantBuffer::FD3D12ConstantBuffer(FD3D12Device* InParent, FD3D12FastCon
 	bIsDirty(false),
 	Allocator(InAllocator)
 {
-#if USE_STATIC_ROOT_SIGNATURE
+#if D3D12RHI_USE_CONSTANT_BUFFER_VIEWS
 	View = new FD3D12ConstantBufferView(InParent);
 #endif
 }
 
 FD3D12ConstantBuffer::~FD3D12ConstantBuffer()
 {
-#if USE_STATIC_ROOT_SIGNATURE
+#if D3D12RHI_USE_CONSTANT_BUFFER_VIEWS
 	delete View;
 #endif
 }
@@ -46,7 +46,7 @@ bool FD3D12ConstantBuffer::Version(FD3D12ResourceLocation& BufferOut, bool bDisc
 	}
 
 	FD3D12ConstantBufferView* ViewToUse = nullptr;
-#if USE_STATIC_ROOT_SIGNATURE
+#if D3D12RHI_USE_CONSTANT_BUFFER_VIEWS
 	ViewToUse = View;
 #endif
 
diff --git a/Engine/Source/Runtime/D3D12RHI/Private/D3D12ConstantBuffer.h b/Engine/Source/Runtime/D3D12RHI/Private/D3D12ConstantBuffer.h
--- a/Engine/Source/Runtime/D3D12RHI/Private/D3D12ConstantBuffer.h
+++ b/Engine/Source/Runtime/D3D12RHI/Private/D3D12ConstantBuffer.h
@@ -45,12 +45,12 @@ public:
 
 	bool Version(FD3D12ResourceLocation& BufferOut, bool bDiscardSharedConstants);
 
-#if USE_STATIC_ROOT_SIGNATURE
+#if D3D12RHI_USE_CONSTANT_BUFFER_VIEWS
 	inline D3D12_CPU_DESCRIPTOR_HANDLE GetOfflineCpuHandle() const { return View->GetOfflineCpuHandle(); }
 #endif
 
 protected:
-#if USE_STATIC_ROOT_SIGNATURE
+#if D3D12RHI_USE_CONSTANT_BUFFER_VIEWS
 	FD3D12ConstantBufferView* View = nullptr;
 #endif
 
diff --git a/Engine/Source/Runtime/D3D12RHI/Private/D3D12DescriptorCache.cpp b/Engine/Source/Runtime/D3D12RHI/Private/D3D12DescriptorCache.cpp
--- a/Engine/Source/Runtime/D3D12RHI/Private/D3D12DescriptorCache.cpp
+++ b/Engine/Source/Runtime/D3D12RHI/Private/D3D12DescriptorCache.cpp
@@ -555,7 +555,7 @@ void FD3D12DescriptorCache::PrepareBindlessViews(EShaderFrequency ShaderStage, T
 
 void FD3D12DescriptorCache::SetConstantBufferViews(EShaderFrequency ShaderStage, const FD3D12RootSignature* RootSignature, FD3D12ConstantBufferCache& Cache, CBVSlotMask SlotsNeededMask, uint32 SlotsNeeded, uint32& HeapSlot)
 {
-#if USE_STATIC_ROOT_SIGNATURE
+#if D3D12RHI_USE_CONSTANT_BUFFER_VIEWS
 	CBVSlotMask& CurrentDirtySlotMask = Cache.DirtySlotMask[ShaderStage];
 	check(CurrentDirtySlotMask != 0);	// All dirty slots for the current shader stage.
 	check(SlotsNeededMask != 0);		// All dirty slots for the current shader stage AND used by the current shader stage.
@@ -625,7 +625,7 @@ void FD3D12DescriptorCache::SetConstantBufferViews(EShaderFrequency ShaderStage,
 	FMsg::Logf(__FILE__, __LINE__, TEXT("DescriptorCache"), ELogVerbosity::Log, TEXT("SetShaderResourceViewTable [STAGE %d] to slots %d - %d"), (int32)ShaderStage, FirstSlotIndex, FirstSlotIndex + SlotsNeeded - 1);
 #endif
 
-#endif // USE_STATIC_ROOT_SIGNATURE
+#endif // D3D12RHI_USE_CONSTANT_BUFFER_VIEWS
 }
 
 void FD3D12DescriptorCache::SetRootConstantBuffers(EShaderFrequency ShaderStage, const FD3D12RootSignature* RootSignature, FD3D12ConstantBufferCache& Cache, CBVSlotMask SlotsNeededMask)
diff --git a/Engine/Source/Runtime/D3D12RHI/Private/D3D12Descriptors.cpp b/Engine/Source/Runtime/D3D12RHI/Private/D3D12Descriptors.cpp
--- a/Engine/Source/Runtime/D3D12RHI/Private/D3D12Descriptors.cpp
+++ b/Engine/Source/Runtime/D3D12RHI/Private/D3D12Descriptors.cpp
@@ -202,7 +202,7 @@ static uint32 GetOfflineDescriptorHeapDefaultSize(ERHIDescriptorHeapType InHeapT
 	switch (InHeapType)
 	{
 	default: checkNoEntry();
-#if USE_STATIC_ROOT_SIGNATURE
+#if D3D12RHI_USE_CONSTANT_BUFFER_VIEWS
 	case ERHIDescriptorHeapType::Standard:     return 4096;
 #else
 	case ERHIDescriptorHeapType::Standard:     return 2048;
diff --git a/Engine/Source/Runtime/D3D12RHI/Private/D3D12RayTracing.cpp b/Engine/Source/Runtime/D3D12RHI/Private/D3D12RayTracing.cpp
--- a/Engine/Source/Runtime/D3D12RHI/Private/D3D12RayTracing.cpp
+++ b/Engine/Source/Runtime/D3D12RHI/Private/D3D12RayTracing.cpp
@@ -1298,7 +1298,7 @@ public:
 	~FD3D12RayTracingShaderTable()
 	{
 		delete DescriptorCache;
-	#if USE_STATIC_ROOT_SIGNATURE
+	#if D3D12RHI_USE_CONSTANT_BUFFER_VIEWS
 		for (FWorkerThreadData& ThisWorkerData : WorkerData)
 		{
 			for (FD3D12ConstantBufferView* CBV : ThisWorkerData.TransientCBVs)
@@ -1306,7 +1306,7 @@ public:
 				delete CBV;
 			}
 		}
-	#endif // USE_STATIC_ROOT_SIGNATURE
+	#endif // D3D12RHI_USE_CONSTANT_BUFFER_VIEWS
 	}
 
 	void Init(const FInitializer& Initializer, FD3D12Device* Device, const TArrayView<const FD3D12ShaderIdentifier>& RaygenIdentifiers, const FD3D12ShaderIdentifier& DefaultHitGroupIdentifier)
@@ -1764,9 +1764,9 @@ public:
 		TArray<FD3D12ShaderResourceView*> TransitionSRVs;
 		TArray<FD3D12UnorderedAccessView*> TransitionUAVs;
 
-#if USE_STATIC_ROOT_SIGNATURE
+#if D3D12RHI_USE_CONSTANT_BUFFER_VIEWS
 		TArray<FD3D12ConstantBufferView*> TransientCBVs;
-#endif // USE_STATIC_ROOT_SIGNATURE
+#endif // D3D12RHI_USE_CONSTANT_BUFFER_VIEWS
 	};
 
 	FWorkerThreadData WorkerData[MaxBindingWorkers];
@@ -4141,12 +4141,12 @@ struct FD3D12RayTracingLocalResourceBinder
 		// caching constant buffer blocks inside ShaderTable and linearly sub-allocate from them.
 		// If the amount of data is relatively small, it may also be possible to use root constants and avoid extra allocations entirely.
 
-	#if USE_STATIC_ROOT_SIGNATURE
+	#if D3D12RHI_USE_CONSTANT_BUFFER_VIEWS
 		FD3D12ConstantBufferView* ConstantBufferView = new FD3D12ConstantBufferView(GetDevice());
 		ShaderTable.WorkerData[WorkerIndex].TransientCBVs.Add(ConstantBufferView);
-	#else // USE_STATIC_ROOT_SIGNATURE
+	#else // D3D12RHI_USE_CONSTANT_BUFFER_VIEWS
 		FD3D12ConstantBufferView* ConstantBufferView = nullptr;
-	#endif // USE_STATIC_ROOT_SIGNATURE
+	#endif // D3D12RHI_USE_CONSTANT_BUFFER_VIEWS
 
 		FD3D12FastConstantAllocator& Allocator = Device.GetParentAdapter()->GetTransientUniformBufferAllocator();
 		void* MappedData = Allocator.Allocate(DataSize, ResourceLocation, ConstantBufferView);
@@ -4212,11 +4212,11 @@ static bool SetRayTracingShaderResources(
 		ResourceBinderType& Binder;
 		uint32 GPUIndex;
 
-#if USE_STATIC_ROOT_SIGNATURE
+#if D3D12RHI_USE_CONSTANT_BUFFER_VIEWS
 		D3D12_CPU_DESCRIPTOR_HANDLE LocalCBVs[MAX_CBS];
-#else // USE_STATIC_ROOT_SIGNATURE
+#else // D3D12RHI_USE_CONSTANT_BUFFER_VIEWS
 		D3D12_GPU_VIRTUAL_ADDRESS LocalCBVs[MAX_CBS];
-#endif // USE_STATIC_ROOT_SIGNATURE
+#endif // D3D12RHI_USE_CONSTANT_BUFFER_VIEWS
 
 		D3D12_CPU_DESCRIPTOR_HANDLE LocalSRVs[MAX_SRVS];
 		D3D12_CPU_DESCRIPTOR_HANDLE LocalUAVs[MAX_UAVS];
@@ -4321,14 +4321,14 @@ static bool SetRayTracingShaderResources(
 		if (Resource)
 		{
 			FD3D12UniformBuffer* CBV = FD3D12CommandContext::RetrieveObject<FD3D12UniformBuffer>(Resource, Bindings.GPUIndex);
-		#if USE_STATIC_ROOT_SIGNATURE
+		#if D3D12RHI_USE_CONSTANT_BUFFER_VIEWS
 			FD3D12OfflineDescriptor Descriptor = CBV->View->GetOfflineCpuHandle();
 			Bindings.LocalCBVs[CBVIndex] = Descriptor;
 			Bindings.CBVVersions[CBVIndex] = Descriptor.GetVersion();
-		#else // USE_STATIC_ROOT_SIGNATURE
+		#else // D3D12RHI_USE_CONSTANT_BUFFER_VIEWS
 			Bindings.LocalCBVs[CBVIndex] = CBV->ResourceLocation.GetGPUVirtualAddress();
 			Bindings.CBVVersions[CBVIndex] = 0; // not available with GPU address path
-		#endif // USE_STATIC_ROOT_SIGNATURE
+		#endif // D3D12RHI_USE_CONSTANT_BUFFER_VIEWS
 			Bindings.BoundCBVMask |= 1ull << CBVIndex;
 
 			Bindings.ReferencedResources.Add(CBV->ResourceLocation.GetResource());
@@ -4381,11 +4381,11 @@ static bool SetRayTracingShaderResources(
 		FD3D12ResourceLocation ResourceLocation(Binder.GetDevice());
 		FD3D12ConstantBufferView* ConstantBufferView = Binder.CreateTransientConstantBuffer(ResourceLocation, InLooseParameterData, InLooseParameterDataSize);
 
-	#if USE_STATIC_ROOT_SIGNATURE
+	#if D3D12RHI_USE_CONSTANT_BUFFER_VIEWS
 		Bindings.LocalCBVs[CBVIndex] = ConstantBufferView->GetOfflineCpuHandle();
-	#else // USE_STATIC_ROOT_SIGNATURE
+	#else // D3D12RHI_USE_CONSTANT_BUFFER_VIEWS
 		Bindings.LocalCBVs[CBVIndex] = ResourceLocation.GetGPUVirtualAddress();
-	#endif // USE_STATIC_ROOT_SIGNATURE
+	#endif // D3D12RHI_USE_CONSTANT_BUFFER_VIEWS
 
 		Bindings.BoundCBVMask |= 1ull << CBVIndex;
 	}
@@ -4443,7 +4443,7 @@ static bool SetRayTracingShaderResources(
 	const uint32 NumCBVs = Shader->ResourceCounts.NumCBs;
 	if (Shader->ResourceCounts.NumCBs)
 	{
-	#if USE_STATIC_ROOT_SIGNATURE
+	#if D3D12RHI_USE_CONSTANT_BUFFER_VIEWS
 
 		const uint32 DescriptorTableBaseIndex = DescriptorCache.AllocateDescriptorTable(Bindings.CBVVersions, Bindings.LocalCBVs, NumCBVs, D3D12_DESCRIPTOR_HEAP_TYPE_CBV_SRV_UAV, WorkerIndex);
 		const uint32 BindSlot = RootSignature->CBVRDTBindSlot(SF_Compute);
@@ -4452,7 +4452,7 @@ static bool SetRayTracingShaderResources(
 		const D3D12_GPU_DESCRIPTOR_HANDLE ResourceDescriptorTableBaseGPU = DescriptorCache.ViewHeap.GetDescriptorGPU(DescriptorTableBaseIndex);
 		Binder.SetRootDescriptorTable(BindSlot, ResourceDescriptorTableBaseGPU);
 
-	#else // USE_STATIC_ROOT_SIGNATURE
+	#else // D3D12RHI_USE_CONSTANT_BUFFER_VIEWS
 
 		checkf(RootSignature->CBVRDTBindSlot(SF_Compute) == 0xFF, TEXT("Root CBV descriptor tables are not implemented for ray tracing shaders."));
 
@@ -4466,7 +4466,7 @@ static bool SetRayTracingShaderResources(
 			Binder.SetRootCBV(BindSlot, i, BufferAddress);
 		}
 
-	#endif // USE_STATIC_ROOT_SIGNATURE
+	#endif // D3D12RHI_USE_CONSTANT_BUFFER_VIEWS
 	}
 
 	// Bind samplers
diff --git a/Engine/Source/Runtime/D3D12RHI/Private/D3D12Resources.h b/Engine/Source/Runtime/D3D12RHI/Private/D3D12Resources.h
--- a/Engine/Source/Runtime/D3D12RHI/Private/D3D12Resources.h
+++ b/Engine/Source/Runtime/D3D12RHI/Private/D3D12Resources.h
@@ -821,8 +821,8 @@ public:
 class FD3D12UniformBuffer : public FRHIUniformBuffer, public FD3D12DeviceChild, public FD3D12LinkedAdapterObject<FD3D12UniformBuffer>
 {
 public:
-#if USE_STATIC_ROOT_SIGNATURE
-	class FD3D12ConstantBufferView* View;
+#if D3D12RHI_USE_CONSTANT_BUFFER_VIEWS
+	class FD3D12ConstantBufferView* View = nullptr;
 #endif
 
 	/** The D3D12 constant buffer resource */
@@ -834,9 +834,6 @@ public:
 	FD3D12UniformBuffer(class FD3D12Device* InParent, const FRHIUniformBufferLayout* InLayout, EUniformBufferUsage InUniformBufferUsage)
 		: FRHIUniformBuffer(InLayout)
 		, FD3D12DeviceChild(InParent)
-#if USE_STATIC_ROOT_SIGNATURE
-		, View(nullptr)
-#endif
 		, ResourceLocation(InParent)
 		, UniformBufferUsage(InUniformBufferUsage)
 	{
diff --git a/Engine/Source/Runtime/D3D12RHI/Private/D3D12StateCache.cpp b/Engine/Source/Runtime/D3D12RHI/Private/D3D12StateCache.cpp
--- a/Engine/Source/Runtime/D3D12RHI/Private/D3D12StateCache.cpp
+++ b/Engine/Source/Runtime/D3D12RHI/Private/D3D12StateCache.cpp
@@ -593,13 +593,13 @@ void FD3D12StateCache::ApplyResources(const FD3D12RootSignature* const pRootSign
 	// Determine what resource bind slots are dirty for the current shaders and how many descriptor table slots we need.
 	// We only set dirty resources that can be used for the upcoming Draw/Dispatch.
 	SRVSlotMask CurrentShaderDirtySRVSlots[SF_NumStandardFrequencies] = {};
-#if USE_STATIC_ROOT_SIGNATURE
+#if D3D12RHI_USE_CONSTANT_BUFFER_VIEWS
 	CBVSlotMask CurrentShaderDirtyCBVSlots[SF_NumStandardFrequencies] = {};
 #endif
 	UAVSlotMask CurrentShaderDirtyUAVSlots = 0;
 	uint32 NumUAVs = 0;
 	uint32 NumSRVs[SF_NumStandardFrequencies] = {};
-#if USE_STATIC_ROOT_SIGNATURE
+#if D3D12RHI_USE_CONSTANT_BUFFER_VIEWS
 	uint32 NumCBVs[SF_NumStandardFrequencies] ={};
 #endif
 	uint32 NumViews = 0;
@@ -652,7 +652,7 @@ void FD3D12StateCache::ApplyResources(const FD3D12RootSignature* const pRootSign
 				NumViews += NumSRVs[Stage];
 			}
 
-#if USE_STATIC_ROOT_SIGNATURE
+#if D3D12RHI_USE_CONSTANT_BUFFER_VIEWS
 			const CBVSlotMask CurrentShaderCBVRegisterMask = BitMask<CBVSlotMask>(PipelineState.Common.CurrentShaderCBCounts[Stage]);
 			CurrentShaderDirtyCBVSlots[Stage] = CurrentShaderCBVRegisterMask & PipelineState.Common.CBVCache.DirtySlotMask[Stage];
 			if (CurrentShaderDirtyCBVSlots[Stage])
@@ -715,7 +715,7 @@ void FD3D12StateCache::ApplyResources(const FD3D12RootSignature* const pRootSign
 		}
 	}
 
-#if USE_STATIC_ROOT_SIGNATURE
+#if D3D12RHI_USE_CONSTANT_BUFFER_VIEWS
 	// Constant buffers
 	{
 		//SCOPE_CYCLE_COUNTER(STAT_D3D12ApplyStateSetConstantBufferTime);
@@ -729,7 +729,7 @@ void FD3D12StateCache::ApplyResources(const FD3D12RootSignature* const pRootSign
 			}
 		}
 	}
-#endif // USE_STATIC_ROOT_SIGNATURE
+#endif // D3D12RHI_USE_CONSTANT_BUFFER_VIEWS
 }
 
 void FD3D12StateCache::ApplyBindlessResources(const FD3D12RootSignature* const pRootSignature, uint32 StartStage, uint32 EndStage)
@@ -748,7 +748,7 @@ void FD3D12StateCache::ApplyBindlessResources(const FD3D12RootSignature* const p
 
 void FD3D12StateCache::ApplyConstants(const FD3D12RootSignature* const pRootSignature, uint32 StartStage, uint32 EndStage)
 {
-#if !USE_STATIC_ROOT_SIGNATURE
+#if !D3D12RHI_USE_CONSTANT_BUFFER_VIEWS
 	// Determine what resource bind slots are dirty for the current shaders and how many descriptor table slots we need.
 	// We only set dirty resources that can be used for the upcoming Draw/Dispatch.
 	CBVSlotMask CurrentShaderDirtyCBVSlots[SF_NumStandardFrequencies] = {};
@@ -778,7 +778,7 @@ void FD3D12StateCache::ApplyConstants(const FD3D12RootSignature* const pRootSign
 			}
 		}
 	}
-#endif // !USE_STATIC_ROOT_SIGNATURE
+#endif // !D3D12RHI_USE_CONSTANT_BUFFER_VIEWS
 }
 
 
diff --git a/Engine/Source/Runtime/D3D12RHI/Private/D3D12StateCachePrivate.h b/Engine/Source/Runtime/D3D12RHI/Private/D3D12StateCachePrivate.h
--- a/Engine/Source/Runtime/D3D12RHI/Private/D3D12StateCachePrivate.h
+++ b/Engine/Source/Runtime/D3D12RHI/Private/D3D12StateCachePrivate.h
@@ -162,12 +162,12 @@ struct FD3D12ConstantBufferCache : public FD3D12ResourceCache<CBVSlotMask>
 
 		FMemory::Memzero(CurrentGPUVirtualAddress, sizeof(CurrentGPUVirtualAddress));
 		FMemory::Memzero(ResidencyHandles, sizeof(ResidencyHandles));
-#if USE_STATIC_ROOT_SIGNATURE
+#if D3D12RHI_USE_CONSTANT_BUFFER_VIEWS
 		FMemory::Memzero(CBHandles, sizeof(CBHandles));
 #endif
 	}
 
-#if USE_STATIC_ROOT_SIGNATURE
+#if D3D12RHI_USE_CONSTANT_BUFFER_VIEWS
 	D3D12_CPU_DESCRIPTOR_HANDLE CBHandles[SF_NumStandardFrequencies][MAX_CBS];
 #endif
 	D3D12_GPU_VIRTUAL_ADDRESS CurrentGPUVirtualAddress[SF_NumStandardFrequencies][MAX_CBS];
@@ -526,7 +526,7 @@ public:
 				FD3D12ConstantBufferCache::DirtySlot(CBVCache.DirtySlotMask[ShaderFrequency], SlotIndex);
 			}
 
-#if USE_STATIC_ROOT_SIGNATURE
+#if D3D12RHI_USE_CONSTANT_BUFFER_VIEWS
 			CBVCache.CBHandles[ShaderFrequency][SlotIndex] = UniformBuffer->View->GetOfflineCpuHandle();
 #endif
 		}
@@ -535,13 +535,13 @@ public:
 			CurrentGPUVirtualAddress = 0;
 			CBVCache.ResidencyHandles[ShaderFrequency][SlotIndex] = nullptr;
 			FD3D12ConstantBufferCache::DirtySlot(CBVCache.DirtySlotMask[ShaderFrequency], SlotIndex);
-#if USE_STATIC_ROOT_SIGNATURE
+#if D3D12RHI_USE_CONSTANT_BUFFER_VIEWS
 			CBVCache.CBHandles[ShaderFrequency][SlotIndex].ptr = 0;
 #endif
 		}
 		else
 		{
-#if USE_STATIC_ROOT_SIGNATURE
+#if D3D12RHI_USE_CONSTANT_BUFFER_VIEWS
 			CBVCache.CBHandles[ShaderFrequency][SlotIndex].ptr = 0;
 #endif
 		}
@@ -563,7 +563,7 @@ public:
 			CBVCache.ResidencyHandles[ShaderFrequency][SlotIndex] = &Location.GetResource()->GetResidencyHandle();
 			FD3D12ConstantBufferCache::DirtySlot(CBVCache.DirtySlotMask[ShaderFrequency], SlotIndex);
 
-#if USE_STATIC_ROOT_SIGNATURE
+#if D3D12RHI_USE_CONSTANT_BUFFER_VIEWS
 			CBVCache.CBHandles[ShaderFrequency][SlotIndex] = Buffer.GetOfflineCpuHandle();
 #endif
 		}
diff --git a/Engine/Source/Runtime/D3D12RHI/Private/D3D12UniformBuffer.cpp b/Engine/Source/Runtime/D3D12RHI/Private/D3D12UniformBuffer.cpp
--- a/Engine/Source/Runtime/D3D12RHI/Private/D3D12UniformBuffer.cpp
+++ b/Engine/Source/Runtime/D3D12RHI/Private/D3D12UniformBuffer.cpp
@@ -35,7 +35,7 @@ FUniformBufferRHIRef FD3D12DynamicRHI::RHICreateUniformBuffer(const void* Conten
 			check(Align(Contents, 16) == Contents);
 			check(NumBytesActualData <= D3D12_REQ_CONSTANT_BUFFER_ELEMENT_COUNT * 16);
 
-#if USE_STATIC_ROOT_SIGNATURE
+#if D3D12RHI_USE_CONSTANT_BUFFER_VIEWS
 			// Create an offline CBV descriptor
 			NewUniformBuffer->View = new FD3D12ConstantBufferView(Device);
 #endif
@@ -63,7 +63,7 @@ FUniformBufferRHIRef FD3D12DynamicRHI::RHICreateUniformBuffer(const void* Conten
 
 				UE::RHICore::UpdateUniformBufferConstants(MappedData, Contents, *Layout);
 
-#if USE_STATIC_ROOT_SIGNATURE
+#if D3D12RHI_USE_CONSTANT_BUFFER_VIEWS
 				NewUniformBuffer->View->CreateView(&NewUniformBuffer->ResourceLocation, 0, NumBytesActualData);
 #endif
 			}
@@ -135,7 +135,7 @@ FRHICOMMAND_MACRO(FRHICommandD3D12UpdateUniformBuffer)
 			check(UniformBuffer->GetResourceTable()[i]);
 		}
 		FD3D12ResourceLocation::TransferOwnership(UniformBuffer->ResourceLocation, UpdatedLocation);
-#if USE_STATIC_ROOT_SIGNATURE
+#if D3D12RHI_USE_CONSTANT_BUFFER_VIEWS
 		const uint32 NumBytes = Align(UniformBuffer->GetLayout().ConstantBufferSize, D3D12_CONSTANT_BUFFER_DATA_PLACEMENT_ALIGNMENT);
 		UniformBuffer->View->CreateView(&UniformBuffer->ResourceLocation, 0, NumBytes);
 #endif
@@ -223,7 +223,7 @@ FD3D12UniformBuffer::~FD3D12UniformBuffer()
 	int64 BufferSize = ResourceLocation.GetSize();
 	DEC_MEMORY_STAT_BY(STAT_UniformBufferMemory, BufferSize);
 
-#if USE_STATIC_ROOT_SIGNATURE
+#if D3D12RHI_USE_CONSTANT_BUFFER_VIEWS
 	delete View;
 #endif
 }
diff --git a/Engine/Source/Runtime/D3D12RHI/Private/Windows/WindowsD3D12RHIDefinitions.h b/Engine/Source/Runtime/D3D12RHI/Private/Windows/WindowsD3D12RHIDefinitions.h
--- a/Engine/Source/Runtime/D3D12RHI/Private/Windows/WindowsD3D12RHIDefinitions.h
+++ b/Engine/Source/Runtime/D3D12RHI/Private/Windows/WindowsD3D12RHIDefinitions.h
@@ -21,6 +21,7 @@ typedef ID3D12GraphicsCommandList ID3D12CopyCommandList;
 #define D3D12RHI_NEEDS_SHADER_FEATURE_CHECKS 1
 
 #define USE_STATIC_ROOT_SIGNATURE			0
+#define D3D12RHI_USE_CONSTANT_BUFFER_VIEWS	0
 #define D3D12_USE_DUMMY_BACKBUFFER			1
 
 // Only enable pipeline statistics if we've got the stats system enabled to display them
