Add support for previewing platforms on PC D3D that support more than 16 samplers (max: 32). To acheive this the following changes were required:
- Added support for up to 32 samplers to the D3D12 RHI.
- Added 'MaxSamplers' to ShaderPlatform in DataDrivenPlatformInfo, this is defaulted to 16 for all shader platforms and can be modified the the DDPI ini files. This value will set the shader compiler define 'PLATFORM_MAX_SAMPLERS'.
- Added a 'RequiredSamplersSwitch' material editor node to branch based on a given shader platform's maximum sampler count.
- To support more than 16 samplers on feature level SM5 platforms Dxc and sm6.0 are used to compile the shader. On platforms that don't support Dxc and sm6.0 the preivew menu item will be disabled.

diff --git a/Engine/Config/Windows/DataDrivenPlatformInfo.ini b/Engine/Config/Windows/DataDrivenPlatformInfo.ini
--- a/Engine/Config/Windows/DataDrivenPlatformInfo.ini
+++ b/Engine/Config/Windows/DataDrivenPlatformInfo.ini
@@ -137,6 +137,7 @@ EnablesHLSL2021ByDefault=1
 SupportsMultiViewport=RuntimeGuaranteed
 bSupportsDualSourceBlending=true
 bSupportsNNEShaders=true
+MaxSamplers=32
 
 [PreviewPlatform PCD3D_SM6]
 PlatformName=PCD3D_SM6
diff --git a/Engine/Source/Developer/Windows/ShaderFormatD3D/Private/D3DShaderCompiler.cpp b/Engine/Source/Developer/Windows/ShaderFormatD3D/Private/D3DShaderCompiler.cpp
--- a/Engine/Source/Developer/Windows/ShaderFormatD3D/Private/D3DShaderCompiler.cpp
+++ b/Engine/Source/Developer/Windows/ShaderFormatD3D/Private/D3DShaderCompiler.cpp
@@ -1110,7 +1110,8 @@ inline bool ShouldUseDXC(const FShaderCompilerInput& Input, ELanguage Language)
 		|| Input.IsRayTracingShader()
 		|| Input.Environment.CompilerFlags.Contains(CFLAG_WaveOperations)
 		|| Input.Environment.CompilerFlags.Contains(CFLAG_ForceDXC)
-		|| Input.Environment.CompilerFlags.Contains(CFLAG_InlineRayTracing);
+		|| Input.Environment.CompilerFlags.Contains(CFLAG_InlineRayTracing)
+		|| Input.Environment.GetIntegerValue(TEXT("PLATFORM_MAX_SAMPLERS")) > 16;
 }
 
 bool PreprocessD3DShader(
diff --git a/Engine/Source/Editor/LevelEditor/Private/LevelEditorActions.cpp b/Engine/Source/Editor/LevelEditor/Private/LevelEditorActions.cpp
--- a/Engine/Source/Editor/LevelEditor/Private/LevelEditorActions.cpp
+++ b/Engine/Source/Editor/LevelEditor/Private/LevelEditorActions.cpp
@@ -718,6 +718,21 @@ bool FLevelEditorActionCallbacks::CanExecutePreviewPlatform(FPreviewPlatformInfo
 		}
 	}
 
+	// When the preview platform's DDSPI MaxSamplers is > 16 and the current RHI device has support
+	// for > 16 samplers we rely on the shader compiler being able to choose an appropriate profile for the
+	// preview feature level that supports > 16 samplers. On D3D12 SM5 the D3D shader compiler will use Dxc and
+	// sm6.0. Vulkan SM5 also appears to handle > 16 samplers fine.
+	// TODO: Look into support for Metal (MacOS).
+	if (FDataDrivenShaderPlatformInfo::IsValid(NewPreviewPlatform.ShaderPlatform) &&
+		FDataDrivenShaderPlatformInfo::GetIsPreviewPlatform(NewPreviewPlatform.ShaderPlatform))
+	{
+		const EShaderPlatform ParentShaderPlatform = FDataDrivenShaderPlatformInfo::GetPreviewShaderPlatformParent(NewPreviewPlatform.ShaderPlatform);
+		if (GMaxTextureSamplers < (int32)FDataDrivenShaderPlatformInfo::GetMaxSamplers(ParentShaderPlatform))
+		{
+			return false;
+		}
+	}
+
 	return true;
 }
 
diff --git a/Engine/Source/Editor/UnrealEd/Private/SEditorViewport.cpp b/Engine/Source/Editor/UnrealEd/Private/SEditorViewport.cpp
--- a/Engine/Source/Editor/UnrealEd/Private/SEditorViewport.cpp
+++ b/Engine/Source/Editor/UnrealEd/Private/SEditorViewport.cpp
@@ -874,7 +874,9 @@ EVisibility SEditorViewport::GetCurrentFeatureLevelPreviewTextVisibility() const
 {
 	if (Client->GetWorld())
 	{
-		return (Client->GetWorld()->GetFeatureLevel() != GMaxRHIFeatureLevel) ? EVisibility::SelfHitTestInvisible : EVisibility::Collapsed;
+		UMaterialShaderQualitySettings* MaterialShaderQualitySettings = UMaterialShaderQualitySettings::Get();
+		const FName& PreviewPlatform = MaterialShaderQualitySettings->GetPreviewPlatform();
+		return PreviewPlatform != NAME_None ? EVisibility::SelfHitTestInvisible : EVisibility::Collapsed;
 	}
 	else
 	{
diff --git a/Engine/Source/Runtime/D3D12RHI/Private/D3D12DescriptorCache.h b/Engine/Source/Runtime/D3D12RHI/Private/D3D12DescriptorCache.h
--- a/Engine/Source/Runtime/D3D12RHI/Private/D3D12DescriptorCache.h
+++ b/Engine/Source/Runtime/D3D12RHI/Private/D3D12DescriptorCache.h
@@ -83,7 +83,7 @@ uint32 GetTypeHash(const D3D12_SAMPLER_DESC& Desc);
 struct FD3D12SamplerArrayDesc
 {
 	uint32 Count;
-	uint16 SamplerID[16];
+	uint16 SamplerID[MAX_SAMPLERS];
 	inline bool operator==(const FD3D12SamplerArrayDesc& rhs) const
 	{
 		check(Count <= UE_ARRAY_COUNT(SamplerID));
diff --git a/Engine/Source/Runtime/D3D12RHI/Private/D3D12RHI.cpp b/Engine/Source/Runtime/D3D12RHI/Private/D3D12RHI.cpp
--- a/Engine/Source/Runtime/D3D12RHI/Private/D3D12RHI.cpp
+++ b/Engine/Source/Runtime/D3D12RHI/Private/D3D12RHI.cpp
@@ -7,6 +7,7 @@
 #include "D3D12RHIPrivate.h"
 #include "RHIStaticStates.h"
 #include "OneColorShader.h"
+#include "DataDrivenShaderPlatformInfo.h"
 
 #include "D3D12AmdExtensions.h"
 #include "D3D12IntelExtensions.h"
@@ -206,6 +207,8 @@ FD3D12DynamicRHI::FD3D12DynamicRHI(const TArray<TSharedPtr<FD3D12Adapter>>& Chos
 	GMaxShadowDepthBufferSizeY = GMaxTextureDimensions;
 	GRHISupportsArrayIndexFromAnyShader = true;
 
+	GMaxTextureSamplers = FMath::Min<int32>(MAX_SAMPLERS, (GetAdapter().GetResourceBindingTier() >= D3D12_RESOURCE_BINDING_TIER_2 ? D3D12_MAX_SHADER_VISIBLE_SAMPLER_HEAP_SIZE : D3D12_COMMONSHADER_SAMPLER_REGISTER_COUNT));
+
 	GRHIMaxDispatchThreadGroupsPerDimension.X = D3D12_CS_DISPATCH_MAX_THREAD_GROUPS_PER_DIMENSION;
 	GRHIMaxDispatchThreadGroupsPerDimension.Y = D3D12_CS_DISPATCH_MAX_THREAD_GROUPS_PER_DIMENSION;
 	GRHIMaxDispatchThreadGroupsPerDimension.Z = D3D12_CS_DISPATCH_MAX_THREAD_GROUPS_PER_DIMENSION;
diff --git a/Engine/Source/Runtime/D3D12RHI/Private/D3D12RHICommon.h b/Engine/Source/Runtime/D3D12RHI/Private/D3D12RHICommon.h
--- a/Engine/Source/Runtime/D3D12RHI/Private/D3D12RHICommon.h
+++ b/Engine/Source/Runtime/D3D12RHI/Private/D3D12RHICommon.h
@@ -28,7 +28,7 @@ static_assert(MAX_ROOT_CBVS <= MAX_CBS, "MAX_ROOT_CBVS must be <= MAX_CBS.");
 static_assert((8 * sizeof(CBVSlotMask)) >= MAX_CBS, "CBVSlotMask isn't large enough to cover all CBs. Please increase the size.");
 static_assert((8 * sizeof(CBVSlotMask)) >= MAX_ROOT_CBVS, "CBVSlotMask isn't large enough to cover all CBs. Please increase the size.");
 static const CBVSlotMask GRootCBVSlotMask = (1 << MAX_ROOT_CBVS) - 1; // Mask for all slots that are used by root descriptors.
-static const CBVSlotMask GDescriptorTableCBVSlotMask = static_cast<CBVSlotMask>(-1) & ~(GRootCBVSlotMask); // Mask for all slots that are used by a root descriptor table.
+static const CBVSlotMask GDescriptorTableCBVSlotMask = static_cast<CBVSlotMask>((1<<MAX_CBS) - 1) & ~(GRootCBVSlotMask); // Mask for all slots that are used by a root descriptor table.
 
 #if MAX_SRVS > 32
 	typedef uint64 SRVSlotMask;
@@ -37,7 +37,7 @@ static const CBVSlotMask GDescriptorTableCBVSlotMask = static_cast<CBVSlotMask>(
 #endif
 static_assert((8 * sizeof(SRVSlotMask)) >= MAX_SRVS, "SRVSlotMask isn't large enough to cover all SRVs. Please increase the size.");
 
-typedef uint16 SamplerSlotMask;
+typedef uint32 SamplerSlotMask;
 static_assert((8 * sizeof(SamplerSlotMask)) >= MAX_SAMPLERS, "SamplerSlotMask isn't large enough to cover all Samplers. Please increase the size.");
 
 typedef uint16 UAVSlotMask;
diff --git a/Engine/Source/Runtime/D3D12RHI/Private/D3D12Util.cpp b/Engine/Source/Runtime/D3D12RHI/Private/D3D12Util.cpp
--- a/Engine/Source/Runtime/D3D12RHI/Private/D3D12Util.cpp
+++ b/Engine/Source/Runtime/D3D12RHI/Private/D3D12Util.cpp
@@ -1002,15 +1002,22 @@ namespace D3D12RHI
 	}
 }
 
-void FD3D12QuantizedBoundShaderState::InitShaderRegisterCounts(D3D12_RESOURCE_BINDING_TIER ResourceBindingTier, const FShaderCodePackedResourceCounts& Counts, FShaderRegisterCounts& Shader, bool bAllowUAVs)
+void FD3D12QuantizedBoundShaderState::InitShaderRegisterCounts(const D3D12_RESOURCE_BINDING_TIER ResourceBindingTier, const FShaderCodePackedResourceCounts& Counts, FShaderRegisterCounts& Shader, bool bAllowUAVs)
 {
-	static const uint32 MaxSamplerCount = MAX_SAMPLERS;
-	static const uint32 MaxConstantBufferCount = MAX_CBS;
-	static const uint32 MaxShaderResourceCount = MAX_SRVS;
-	static const uint32 MaxUnorderedAccessCount = MAX_UAVS;
-
-	// Round up and clamp values to their max
-	// Note: Rounding and setting counts based on binding tier allows us to create fewer root signatures.
+	uint32 MaxSRVs = MAX_SRVS;
+	uint32 MaxSamplers = MAX_SAMPLERS;
+	uint32 MaxUAVs = MAX_UAVS;
+	uint32 MaxCBs = MAX_CBS;
+
+	// On tier 1 & 2 HW the actual descriptor table size used during the draw/dispatch must match that of the
+	// root signature so we round the size up to the closest power of 2 to accomplish 2 goals: 1) keep the size of
+	// the table closer to the required size to limit descriptor heap usage due to required empty descriptors,
+	// 2) encourage root signature reuse by having other shader root signature table sizes fall within the size rounding.
+	// Sampler and Shader resouce view table sizes must match signature on Tier 1 hardware and Constant buffer and
+	// Unorded access views table sizes must match signature on tier 2 hardware. On hardware > tier 2 the actual descriptor
+	// table size used during the draw/dispatch doesn't need to match the root signature size so we encourage reuse by using
+	// the max size. More info here: https://learn.microsoft.com/en-us/windows/win32/direct3d12/hardware-support,
+	// https://en.wikipedia.org/wiki/Feature_levels_in_Direct3D
 
 	// To reduce the size of the root signature, we only allow UAVs for certain shaders. 
 	// This code makes the assumption that the engine only uses UAVs at the PS or CS shader stages.
@@ -1018,24 +1025,24 @@ void FD3D12QuantizedBoundShaderState::InitShaderRegisterCounts(D3D12_RESOURCE_BI
 
 	if (ResourceBindingTier <= D3D12_RESOURCE_BINDING_TIER_1)
 	{
-		Shader.SamplerCount = (Counts.NumSamplers > 0) ? FMath::Min(MaxSamplerCount, FMath::RoundUpToPowerOfTwo(Counts.NumSamplers)) : Counts.NumSamplers;
-		Shader.ShaderResourceCount = (Counts.NumSRVs > 0) ? FMath::Min(MaxShaderResourceCount, FMath::RoundUpToPowerOfTwo(Counts.NumSRVs)) : Counts.NumSRVs;
+		Shader.SamplerCount = (Counts.NumSamplers > 0) ? FMath::Min(MaxSamplers, FMath::RoundUpToPowerOfTwo(Counts.NumSamplers)) : Counts.NumSamplers;
+		Shader.ShaderResourceCount = (Counts.NumSRVs > 0) ? FMath::Min(MaxSRVs, FMath::RoundUpToPowerOfTwo(Counts.NumSRVs)) : Counts.NumSRVs;
 	}
 	else
 	{
-		Shader.SamplerCount = Counts.NumSamplers > 0 ? MaxSamplerCount : 0;
-		Shader.ShaderResourceCount = Counts.NumSRVs > 0 ? MaxShaderResourceCount : 0;
+		Shader.SamplerCount = Counts.NumSamplers > 0 ? MaxSamplers : 0;
+		Shader.ShaderResourceCount = Counts.NumSRVs > 0 ? MaxSRVs : 0;
 	}
 
 	if (ResourceBindingTier <= D3D12_RESOURCE_BINDING_TIER_2)
 	{
-		Shader.ConstantBufferCount = (Counts.NumCBs > MAX_ROOT_CBVS) ? FMath::Min(MaxConstantBufferCount, FMath::RoundUpToPowerOfTwo(Counts.NumCBs)) : Counts.NumCBs;
-		Shader.UnorderedAccessCount = (Counts.NumUAVs > 0 && bAllowUAVs) ? FMath::Min(MaxUnorderedAccessCount, FMath::RoundUpToPowerOfTwo(Counts.NumUAVs)) : 0;
+		Shader.ConstantBufferCount = (Counts.NumCBs > MAX_ROOT_CBVS) ? FMath::Min(MaxCBs, FMath::RoundUpToPowerOfTwo(Counts.NumCBs)) : Counts.NumCBs;
+		Shader.UnorderedAccessCount = (Counts.NumUAVs > 0 && bAllowUAVs) ? FMath::Min(MaxUAVs, FMath::RoundUpToPowerOfTwo(Counts.NumUAVs)) : 0;
 	}
 	else
 	{
-		Shader.ConstantBufferCount = (Counts.NumCBs > MAX_ROOT_CBVS) ? MaxConstantBufferCount : Counts.NumCBs;
-		Shader.UnorderedAccessCount = Counts.NumUAVs > 0 && bAllowUAVs ? MaxUnorderedAccessCount : 0;
+		Shader.ConstantBufferCount = (Counts.NumCBs > MAX_ROOT_CBVS) ? MaxCBs : Counts.NumCBs;
+		Shader.UnorderedAccessCount = (Counts.NumUAVs > 0 && bAllowUAVs) ? MaxUAVs : 0;
 	}
 }
 
diff --git a/Engine/Source/Runtime/D3D12RHI/Private/Windows/WindowsD3D12Device.cpp b/Engine/Source/Runtime/D3D12RHI/Private/Windows/WindowsD3D12Device.cpp
--- a/Engine/Source/Runtime/D3D12RHI/Private/Windows/WindowsD3D12Device.cpp
+++ b/Engine/Source/Runtime/D3D12RHI/Private/Windows/WindowsD3D12Device.cpp
@@ -1031,8 +1031,11 @@ static bool DoesAnyAdapterSupportSM6(const TArray<TSharedPtr<FD3D12Adapter>>& Ad
 
 FDynamicRHI* FD3D12DynamicRHIModule::CreateRHI(ERHIFeatureLevel::Type RequestedFeatureLevel)
 {
+	GShaderPlatformForFeatureLevel[ERHIFeatureLevel::ES2_REMOVED] = SP_NumPlatforms;
 	GShaderPlatformForFeatureLevel[ERHIFeatureLevel::ES3_1] = SP_PCD3D_ES3_1;
+	GShaderPlatformForFeatureLevel[ERHIFeatureLevel::SM4_REMOVED] = SP_NumPlatforms;
 	GShaderPlatformForFeatureLevel[ERHIFeatureLevel::SM5] = SP_PCD3D_SM5;
+	GShaderPlatformForFeatureLevel[ERHIFeatureLevel::SM6] = SP_NumPlatforms;
 	if (DoesAnyAdapterSupportSM6(ChosenAdapters))
 	{
 		GShaderPlatformForFeatureLevel[ERHIFeatureLevel::SM6] = SP_PCD3D_SM6;
@@ -1081,6 +1084,19 @@ FDynamicRHI* FD3D12DynamicRHIModule::CreateRHI(ERHIFeatureLevel::Type RequestedF
 		return new FValidationRHI(GD3D12RHI);
 	}
 #endif
+
+	for (int32 Index = 0; Index < ERHIFeatureLevel::Num; ++Index)
+	{
+		if (GShaderPlatformForFeatureLevel[Index] != SP_NumPlatforms)
+		{
+			check(GMaxTextureSamplers >= (int32)FDataDrivenShaderPlatformInfo::GetMaxSamplers(GShaderPlatformForFeatureLevel[Index]));
+			if (GMaxTextureSamplers < (int32)FDataDrivenShaderPlatformInfo::GetMaxSamplers(GShaderPlatformForFeatureLevel[Index]))
+			{
+				UE_LOG(LogD3D12RHI, Error, TEXT("Shader platform requires at least: %d samplers, device supports: %d."), FDataDrivenShaderPlatformInfo::GetMaxSamplers(GShaderPlatformForFeatureLevel[Index]), GMaxTextureSamplers);
+			}
+		}
+	}
+
 	return GD3D12RHI;
 }
 
diff --git a/Engine/Source/Runtime/D3D12RHI/Public/D3D12RHI.h b/Engine/Source/Runtime/D3D12RHI/Public/D3D12RHI.h
--- a/Engine/Source/Runtime/D3D12RHI/Public/D3D12RHI.h
+++ b/Engine/Source/Runtime/D3D12RHI/Public/D3D12RHI.h
@@ -10,8 +10,14 @@
 
 // This is a value that should be tweaked to fit the app, lower numbers will have better performance
 // Titles using many terrain layers may want to set MAX_SRVS to 64 to avoid shader compilation errors. This will have a small performance hit of around 0.1%
+// Note: We do not need to worry about these counts exceeding resource binding tiers 1 & 2 due to 2 reasons:
+// 1) On PC the actual counts used to build the root signature are derived from the shader and the shader compiler
+// profile required for a desired feature level will enforce the proper limits.
+// 2) All platforms currently using the global root signature support resource binding tier 3, these counts are well within the limits
+// of tier 3.
+// More info here: https://learn.microsoft.com/en-us/windows/win32/direct3d12/hardware-support, https://en.wikipedia.org/wiki/Feature_levels_in_Direct3D
 #define MAX_SRVS		64
-#define MAX_SAMPLERS	16
+#define MAX_SAMPLERS	32
 #define MAX_UAVS		16
 #define MAX_CBS			16
 
diff --git a/Engine/Source/Runtime/Engine/Classes/Materials/MaterialExpressionRequiredSamplersSwitch.h b/Engine/Source/Runtime/Engine/Classes/Materials/MaterialExpressionRequiredSamplersSwitch.h
new file mode 100644
--- /dev/null
+++ b/Engine/Source/Runtime/Engine/Classes/Materials/MaterialExpressionRequiredSamplersSwitch.h
@@ -0,0 +1,38 @@
+// Copyright Epic Games, Inc. All Rights Reserved.
+
+
+#pragma once
+
+#include "CoreMinimal.h"
+#include "UObject/ObjectMacros.h"
+#include "MaterialExpressionIO.h"
+#include "MaterialValueType.h"
+#include "Materials/MaterialExpression.h"
+#include "RHIFeatureLevel.h"
+#include "MaterialExpressionRequiredSamplersSwitch.generated.h"
+
+UCLASS(MinimalAPI)
+class UMaterialExpressionRequiredSamplersSwitch : public UMaterialExpression
+{
+	GENERATED_UCLASS_BODY()
+		UPROPERTY()
+		FExpressionInput InputTrue;
+		
+		UPROPERTY()
+		FExpressionInput InputFalse;
+
+		UPROPERTY(EditAnywhere, Category=MaterialExpressionRequiredSamplersSwitch, meta=(UIMin = "16", UIMax = "255"))
+		uint32 RequiredSamplers;
+
+	//~ Begin UMaterialExpression Interface
+#if WITH_EDITOR
+	virtual int32 Compile(class FMaterialCompiler* Compiler, int32 OutputIndex) override;
+	virtual void GetCaption(TArray<FString>& OutCaptions) const override;
+	virtual bool IsInputConnectionRequired(int32 InputIndex) const override;
+	virtual FName GetInputName(int32 InputIndex) const override;
+	virtual bool IsResultMaterialAttributes(int32 OutputIndex) override;
+	virtual uint32 GetInputType(int32 InputIndex) override { return MCT_Unknown; }
+	virtual uint32 GetOutputType(int32 OutputIndex) override { return MCT_Unknown; }
+#endif // WITH_EDITOR
+	//~ End UMaterialExpression Interface
+};
diff --git a/Engine/Source/Runtime/Engine/Private/Materials/MaterialExpressions.cpp b/Engine/Source/Runtime/Engine/Private/Materials/MaterialExpressions.cpp
--- a/Engine/Source/Runtime/Engine/Private/Materials/MaterialExpressions.cpp
+++ b/Engine/Source/Runtime/Engine/Private/Materials/MaterialExpressions.cpp
@@ -101,6 +101,7 @@
 #include "Materials/MaterialExpressionEyeAdaptationInverse.h"
 #include "Materials/MaterialExpressionFeatureLevelSwitch.h"
 #include "Materials/MaterialExpressionDataDrivenShaderPlatformInfoSwitch.h"
+#include "Materials/MaterialExpressionRequiredSamplersSwitch.h"
 #include "Materials/MaterialExpressionFloor.h"	
 #include "Materials/MaterialExpressionFmod.h"
 #include "Materials/MaterialExpressionFontSample.h"
@@ -9940,6 +9941,73 @@ void UMaterialExpressionDataDrivenShaderPlatformInfoSwitch::PostEditChangeProper
 }
 #endif // WITH_EDITOR
 
+//
+//	UMaterialExpressionRequiredSamplersSwitch
+//
+
+UMaterialExpressionRequiredSamplersSwitch::UMaterialExpressionRequiredSamplersSwitch(const FObjectInitializer& ObjectInitializer)
+	: Super(ObjectInitializer),
+	RequiredSamplers(16)
+{
+}
+
+#if WITH_EDITOR
+
+int32 UMaterialExpressionRequiredSamplersSwitch::Compile(class FMaterialCompiler* Compiler, int32 OutputIndex)
+{
+	const EShaderPlatform ShaderPlatform = Compiler->GetShaderPlatform();
+	check(FDataDrivenShaderPlatformInfo::IsValid(ShaderPlatform));
+	const bool bCheck = RequiredSamplers <= FDataDrivenShaderPlatformInfo::GetMaxSamplers(ShaderPlatform);
+	if (bCheck)
+	{
+		return InputTrue.Compile(Compiler);
+	}
+	else
+	{
+		return InputFalse.Compile(Compiler);
+	}
+}
+
+void UMaterialExpressionRequiredSamplersSwitch::GetCaption(TArray<FString>& OutCaptions) const
+{
+	OutCaptions.Add(FString(TEXT("Required Samplers Switch")));
+}
+
+bool UMaterialExpressionRequiredSamplersSwitch::IsInputConnectionRequired(int32 InputIndex) const
+{
+	return true;
+}
+
+FName UMaterialExpressionRequiredSamplersSwitch::GetInputName(int32 InputIndex) const
+{
+	if (InputIndex == 0)
+	{
+		return TEXT("Within platform limit");
+	}
+	else if (InputIndex == 1)
+	{
+		return TEXT("Over platform limit");
+	}
+	return NAME_None;
+}
+
+bool UMaterialExpressionRequiredSamplersSwitch::IsResultMaterialAttributes(int32 OutputIndex)
+{
+	check(OutputIndex == 0);
+	for (FExpressionInput* ExpressionInput : GetInputsView())
+	{
+		// If there is a loop anywhere in this expression's inputs then we can't risk checking them
+		if (ExpressionInput->Expression && ExpressionInput->Expression->IsResultMaterialAttributes(ExpressionInput->OutputIndex))
+		{
+			return true;
+		}
+	}
+
+	return false;
+}
+
+#endif // WITH_EDITOR
+
 //
 //	UMaterialExpressionShadingPathSwitch
 //
diff --git a/Engine/Source/Runtime/Engine/Private/ShaderCompiler/ShaderCompiler.cpp b/Engine/Source/Runtime/Engine/Private/ShaderCompiler/ShaderCompiler.cpp
--- a/Engine/Source/Runtime/Engine/Private/ShaderCompiler/ShaderCompiler.cpp
+++ b/Engine/Source/Runtime/Engine/Private/ShaderCompiler/ShaderCompiler.cpp
@@ -7860,6 +7860,11 @@ void GlobalBeginCompileShader(
 		SET_SHADER_DEFINE(Input.Environment, PLATFORM_SUPPORTS_VERTEX_SHADER_SRVS, bSupportsVertexShaderSRVs ? 1u : 0u);
 	}
 
+	{
+		const uint32 MaxSamplers = FDataDrivenShaderPlatformInfo::GetMaxSamplers((EShaderPlatform)Target.Platform);
+		SET_SHADER_DEFINE(Input.Environment, PLATFORM_MAX_SAMPLERS, MaxSamplers);
+	}
+
 	ITargetPlatform* TargetPlatform = GetTargetPlatformManager()->FindTargetPlatformWithSupport(TEXT("ShaderFormat"), ShaderFormatName);
 	bool bForwardShading = false;
 	{
diff --git a/Engine/Source/Runtime/RHI/Private/DataDrivenShaderPlatformInfo.cpp b/Engine/Source/Runtime/RHI/Private/DataDrivenShaderPlatformInfo.cpp
--- a/Engine/Source/Runtime/RHI/Private/DataDrivenShaderPlatformInfo.cpp
+++ b/Engine/Source/Runtime/RHI/Private/DataDrivenShaderPlatformInfo.cpp
@@ -147,6 +147,7 @@ void FGenericDataDrivenShaderPlatformInfo::SetDefaultValues()
 	bSupportsVolumeTextureAtomics = true;
 	bSupportsClipDistance = true;
 	bSupportsShaderPipelines = true;
+	MaxSamplers = 16;
 }
 
 void FGenericDataDrivenShaderPlatformInfo::ParseDataDrivenShaderInfo(const FConfigSection& Section, uint32 Index)
@@ -290,6 +291,7 @@ void FGenericDataDrivenShaderPlatformInfo::ParseDataDrivenShaderInfo(const FConf
 	GET_SECTION_BOOL_HELPER(bSupportsShaderPipelines);
 	GET_SECTION_BOOL_HELPER(bSupportsUniformBufferObjects);
 	GET_SECTION_BOOL_HELPER(bRequiresBindfulUtilityShaders);
+	GET_SECTION_INT_HELPER(MaxSamplers);
 #undef GET_SECTION_BOOL_HELPER
 #undef GET_SECTION_INT_HELPER
 #undef GET_SECTION_SUPPORT_HELPER
@@ -365,10 +367,9 @@ void FGenericDataDrivenShaderPlatformInfo::Initialize()
 						if (Item.ShaderPlatformToPreview == PreviewPlatformName)
 						{
 							const EShaderPlatform PreviewShaderPlatform = EShaderPlatform(CustomShaderPlatform++);
-							ParseDataDrivenShaderInfo(Section.Value, PreviewShaderPlatform);
-
 							FGenericDataDrivenShaderPlatformInfo& PreviewInfo = Infos[PreviewShaderPlatform];
 							PreviewInfo.Name = Item.PreviewShaderPlatformName;
+							ParseDataDrivenShaderInfo(Section.Value, PreviewShaderPlatform);
 							PreviewInfo.bIsPreviewPlatform = true;
 							PreviewInfo.bContainsValidPlatformInfo = true;
 
@@ -406,6 +407,7 @@ void FGenericDataDrivenShaderPlatformInfo::UpdatePreviewPlatforms()
 				Infos[ShaderPlatform].bIsHlslcc = Infos[EditorSPForPreviewMaxFeatureLevel].bIsHlslcc;
 				Infos[ShaderPlatform].bSupportsDxc = Infos[EditorSPForPreviewMaxFeatureLevel].bSupportsDxc;
 				Infos[ShaderPlatform].bSupportsGPUScene = Infos[EditorSPForPreviewMaxFeatureLevel].bSupportsGPUScene;
+				Infos[ShaderPlatform].MaxMeshShaderThreadGroupSize = Infos[EditorSPForPreviewMaxFeatureLevel].MaxMeshShaderThreadGroupSize;
 				Infos[ShaderPlatform].bIsPC = true;
 				Infos[ShaderPlatform].bSupportsDebugViewShaders = true;
 				Infos[ShaderPlatform].bIsConsole = false;
@@ -443,6 +445,10 @@ void FGenericDataDrivenShaderPlatformInfo::UpdatePreviewPlatforms()
 				Infos[ShaderPlatform].bSupports4ComponentUAVReadWrite = false;
 				Infos[ShaderPlatform].bSupportsUniformBufferObjects = Infos[EditorSPForPreviewMaxFeatureLevel].bSupportsUniformBufferObjects;
 				Infos[ShaderPlatform].bContainsValidPlatformInfo = true;
+
+				// Seeing as we are merging the two shader platforms merge the hash key as well, this way
+				// any changes in the editor feature level shader platform will dirty the preview key.
+				Infos[ShaderPlatform].ShaderPropertiesHash = HashCombine(Infos[ShaderPlatform].ShaderPropertiesHash, Infos[EditorSPForPreviewMaxFeatureLevel].ShaderPropertiesHash);
 			}
 		}
 	}
diff --git a/Engine/Source/Runtime/RHI/Public/DataDrivenShaderPlatformInfo.h b/Engine/Source/Runtime/RHI/Public/DataDrivenShaderPlatformInfo.h
--- a/Engine/Source/Runtime/RHI/Public/DataDrivenShaderPlatformInfo.h
+++ b/Engine/Source/Runtime/RHI/Public/DataDrivenShaderPlatformInfo.h
@@ -120,6 +120,7 @@ class FGenericDataDrivenShaderPlatformInfo
 	uint32 bRequiresBindfulUtilityShaders : 1;
 	// NOTE: When adding fields, you must also add to ParseDataDrivenShaderInfo!
 	uint32 bContainsValidPlatformInfo : 1;
+	uint32 MaxSamplers : 8;
 
 	FGenericDataDrivenShaderPlatformInfo()
 	{
@@ -759,6 +760,12 @@ public:
 		return Infos[Platform].bRequiresBindfulUtilityShaders;
 	}
 	
+	static FORCEINLINE_DEBUGGABLE const uint32 GetMaxSamplers(const FStaticShaderPlatform Platform)
+	{
+		check(IsValid(Platform));
+		return Infos[Platform].MaxSamplers;
+	}
+
 	static FORCEINLINE_DEBUGGABLE const bool IsValid(const FStaticShaderPlatform Platform)
 	{
 		return Infos[Platform].bContainsValidPlatformInfo;
diff --git a/Engine/Source/Runtime/RHI/Public/RHICommandList.h b/Engine/Source/Runtime/RHI/Public/RHICommandList.h
--- a/Engine/Source/Runtime/RHI/Public/RHICommandList.h
+++ b/Engine/Source/Runtime/RHI/Public/RHICommandList.h
@@ -244,7 +244,7 @@ struct FRayTracingShaderBindings
 	FRHITexture* Textures[64] = {};
 	FRHIShaderResourceView* SRVs[64] = {};
 	FRHIUniformBuffer* UniformBuffers[16] = {};
-	FRHISamplerState* Samplers[16] = {};
+	FRHISamplerState* Samplers[32] = {};
 	FRHIUnorderedAccessView* UAVs[16] = {};
 };
 
diff --git a/Engine/Source/Runtime/VulkanRHI/Private/VulkanRHI.cpp b/Engine/Source/Runtime/VulkanRHI/Private/VulkanRHI.cpp
--- a/Engine/Source/Runtime/VulkanRHI/Private/VulkanRHI.cpp
+++ b/Engine/Source/Runtime/VulkanRHI/Private/VulkanRHI.cpp
@@ -395,6 +395,18 @@ FDynamicRHI* FVulkanDynamicRHIModule::CreateRHI(ERHIFeatureLevel::Type InRequest
 	}
 #endif
 
+	for (int32 Index = 0; Index < ERHIFeatureLevel::Num; ++Index)
+	{
+		if (ShaderPlatformForFeatureLevel[Index] != SP_NumPlatforms)
+		{
+			check(GMaxTextureSamplers >= (int32)FDataDrivenShaderPlatformInfo::GetMaxSamplers(ShaderPlatformForFeatureLevel[Index]));
+			if (GMaxTextureSamplers < (int32)FDataDrivenShaderPlatformInfo::GetMaxSamplers(ShaderPlatformForFeatureLevel[Index]))
+			{
+				UE_LOG(LogVulkanRHI, Error, TEXT("Shader platform requires at least: %d samplers, device supports: %d."), FDataDrivenShaderPlatformInfo::GetMaxSamplers(ShaderPlatformForFeatureLevel[Index]), GMaxTextureSamplers);
+			}
+		}
+	}
+
 	return FinalRHI;
 }
 
@@ -955,7 +967,8 @@ void FVulkanDynamicRHI::InitInstance()
 		GSupportsSeparateRenderTargetBlendState = true;
 		GRHISupportsSeparateDepthStencilCopyAccess = Device->SupportsParallelRendering();
 		GRHIBindlessSupport = Device->SupportsBindless() ? RHIGetBindlessSupport(GMaxRHIShaderPlatform) : ERHIBindlessSupport::Unsupported;
-
+		GMaxTextureSamplers = Props.limits.maxPerStageDescriptorSamplers;
+		
 		GRHIMaxDispatchThreadGroupsPerDimension.X = FMath::Min<uint32>(Limits.maxComputeWorkGroupCount[0], 0x7fffffff);
 		GRHIMaxDispatchThreadGroupsPerDimension.Y = FMath::Min<uint32>(Limits.maxComputeWorkGroupCount[1], 0x7fffffff);
 		GRHIMaxDispatchThreadGroupsPerDimension.Z = FMath::Min<uint32>(Limits.maxComputeWorkGroupCount[2], 0x7fffffff);
