Use FGenericCrashContext for GPU breadcrumbs.

diff --git a/Engine/Source/Runtime/D3D12RHI/Private/D3D12Util.cpp b/Engine/Source/Runtime/D3D12RHI/Private/D3D12Util.cpp
--- a/Engine/Source/Runtime/D3D12RHI/Private/D3D12Util.cpp
+++ b/Engine/Source/Runtime/D3D12RHI/Private/D3D12Util.cpp
@@ -264,7 +264,7 @@ static bool LogBreadcrumbData(D3D12RHI::FD3DGPUProfiler& GPUProfiler, FD3D12Queu
 
 	if (!Nodes.IsEmpty())
 	{
-		FPlatformCrashContext::SetGPUBreadcrumbs(GPUQueueName, Nodes);
+		FGenericCrashContext::SetGPUBreadcrumbs(GPUQueueName, Nodes);
 	}
 
 	const FD3D12DiagnosticBufferData* DiagnosticData = Queue.GetDiagnosticBufferData();
