Fix Memory Stats being decremented twice in D3D12RayTracing when using mesh streaming.

diff --git a/Engine/Source/Runtime/D3D12RHI/Private/D3D12RayTracing.cpp b/Engine/Source/Runtime/D3D12RHI/Private/D3D12RayTracing.cpp
--- a/Engine/Source/Runtime/D3D12RHI/Private/D3D12RayTracing.cpp
+++ b/Engine/Source/Runtime/D3D12RHI/Private/D3D12RayTracing.cpp
@@ -2958,8 +2958,8 @@ void FD3D12RayTracingGeometry::Swap(FD3D12RayTracingGeometry& Other)
 {
 	for (uint32 i = 0; i < MAX_NUM_GPUS; i++)
 	{
-		AccelerationStructureBuffers[i] = Other.AccelerationStructureBuffers[i];
-		bIsAccelerationStructureDirty[i] = Other.bIsAccelerationStructureDirty[i];
+		::Swap(AccelerationStructureBuffers[i], Other.AccelerationStructureBuffers[i]);
+		::Swap(bIsAccelerationStructureDirty[i], Other.bIsAccelerationStructureDirty[i]);
 	}
 	::Swap(AccelerationStructureCompactedSize, Other.AccelerationStructureCompactedSize);
 
