Fix D3D12 TerminateOnOutOfMemory to not pop up a message box when running -unattended, causing build machine jobs to spin for 30m until timeout instead of crashing in just a couple minutes.

diff --git a/Engine/Source/Runtime/D3D12RHI/Private/D3D12Util.cpp b/Engine/Source/Runtime/D3D12RHI/Private/D3D12Util.cpp
--- a/Engine/Source/Runtime/D3D12RHI/Private/D3D12Util.cpp
+++ b/Engine/Source/Runtime/D3D12RHI/Private/D3D12Util.cpp
@@ -750,15 +750,18 @@ static void TerminateOnOutOfMemory(ID3D12Device* InDevice, HRESULT D3DResult, bo
 		FCoreDelegates::GetGPUOutOfMemoryDelegate().Broadcast(MemoryInfo.Budget, MemoryInfo.CurrentUsage);
 	}
 
-
-	if (bCreatingTextures)
+	if (!FApp::IsUnattended())
 	{
-		FPlatformMisc::MessageBoxExt(EAppMsgType::Ok, *LOCTEXT("OutOfVideoMemoryTextures", "Out of video memory trying to allocate a texture! Make sure your video card has the minimum required memory, try lowering the resolution and/or closing other applications that are running. Exiting...").ToString(), TEXT("Error"));
-	}
-	else
-	{
-		FPlatformMisc::MessageBoxExt(EAppMsgType::Ok, *NSLOCTEXT("D3D12RHI", "OutOfMemory", "Out of video memory trying to allocate a rendering resource. Make sure your video card has the minimum required memory, try lowering the resolution and/or closing other applications that are running. Exiting...").ToString(), TEXT("Error"));
+		if (bCreatingTextures)
+		{
+			FPlatformMisc::MessageBoxExt(EAppMsgType::Ok, *LOCTEXT("OutOfVideoMemoryTextures", "Out of video memory trying to allocate a texture! Make sure your video card has the minimum required memory, try lowering the resolution and/or closing other applications that are running. Exiting...").ToString(), TEXT("Error"));
+		}
+		else
+		{
+			FPlatformMisc::MessageBoxExt(EAppMsgType::Ok, *NSLOCTEXT("D3D12RHI", "OutOfMemory", "Out of video memory trying to allocate a rendering resource. Make sure your video card has the minimum required memory, try lowering the resolution and/or closing other applications that are running. Exiting...").ToString(), TEXT("Error"));
+		}
 	}
+
 #if STATS
 	GetRendererModule().DebugLogOnCrash();
 #endif
