GetStaticUniformBuffers shouldn't return the array by value.

diff --git a/Engine/Source/Runtime/D3D12RHI/Private/D3D12CommandContext.h b/Engine/Source/Runtime/D3D12RHI/Private/D3D12CommandContext.h
--- a/Engine/Source/Runtime/D3D12RHI/Private/D3D12CommandContext.h
+++ b/Engine/Source/Runtime/D3D12RHI/Private/D3D12CommandContext.h
@@ -730,7 +730,7 @@ public:
 
 	virtual void UpdateBuffer(FD3D12ResourceLocation* Dest, uint32 DestOffset, FD3D12ResourceLocation* Source, uint32 SourceOffset, uint32 NumBytes) final override;
 
-	inline const TArray<FRHIUniformBuffer*> GetStaticUniformBuffers() const
+	inline const TArray<FRHIUniformBuffer*>& GetStaticUniformBuffers() const
 	{
 		return StaticUniformBuffers;
 	}
