Fix for hitting ensure in EnsureColorSpace when running with -renderoffscreen by instead checking to make sure there is a SwapChain4 before trying to use it by calling CurrentOutputSupportsHDR(). If you don't have a valid SwapChain4 at that point, it should crash anyway.

diff --git a/Engine/Source/Runtime/D3D12RHI/Private/Windows/WindowsD3D12Viewport.cpp b/Engine/Source/Runtime/D3D12RHI/Private/Windows/WindowsD3D12Viewport.cpp
--- a/Engine/Source/Runtime/D3D12RHI/Private/Windows/WindowsD3D12Viewport.cpp
+++ b/Engine/Source/Runtime/D3D12RHI/Private/Windows/WindowsD3D12Viewport.cpp
@@ -470,8 +470,6 @@ static const FString GetDXGIColorSpaceString(DXGI_COLOR_SPACE_TYPE ColorSpace)
 
 void FD3D12Viewport::EnsureColorSpace(EDisplayColorGamut DisplayGamut, EDisplayOutputFormat OutputDevice)
 {
-	ensure(SwapChain4.GetReference());
-
 	DXGI_COLOR_SPACE_TYPE NewColorSpace = DXGI_COLOR_SPACE_RGB_FULL_G22_NONE_P709;	// sRGB;
 	const bool bPrimaries2020 = (DisplayGamut == EDisplayColorGamut::Rec2020_D65);
 
@@ -501,6 +499,7 @@ void FD3D12Viewport::EnsureColorSpace(EDisplayColorGamut DisplayGamut, EDisplayO
 
 	if (ColorSpace != NewColorSpace)
 	{
+		check(CurrentOutputSupportsHDR());
 		uint32 ColorSpaceSupport = 0;
 		HRESULT hr = SwapChain4->CheckColorSpaceSupport(NewColorSpace, &ColorSpaceSupport);
 		FString NewColorSpaceName = GetDXGIColorSpaceString(NewColorSpace);
