Numerous optimizations to shader bundles, including avoiding descriptor table deduplication for SRVs and CBVs, which will never cache hit (unlike UAVs and Samplers).

diff --git a/Engine/Source/Runtime/D3D12RHI/Private/D3D12ExplicitDescriptorCache.cpp b/Engine/Source/Runtime/D3D12RHI/Private/D3D12ExplicitDescriptorCache.cpp
--- a/Engine/Source/Runtime/D3D12RHI/Private/D3D12ExplicitDescriptorCache.cpp
+++ b/Engine/Source/Runtime/D3D12RHI/Private/D3D12ExplicitDescriptorCache.cpp
@@ -5,7 +5,7 @@
 
 // Whether to compare the full descriptor table on cache lookup or only use FXxHash64 digest.
 #ifndef EXPLICIT_DESCRIPTOR_CACHE_FULL_COMPARE
-#define EXPLICIT_DESCRIPTOR_CACHE_FULL_COMPARE 1
+#define EXPLICIT_DESCRIPTOR_CACHE_FULL_COMPARE 0
 #endif
 
 static int32 GD3D12ExplicitDeduplicateSamplers = 1;
@@ -244,10 +244,12 @@ void FD3D12ExplicitDescriptorHeap::CopyDescriptors(int32 BaseIndex, const D3D12_
 {
 	D3D12_CPU_DESCRIPTOR_HANDLE DestDescriptor = GetDescriptorCPU(BaseIndex);
 	GetParentDevice()->GetDevice()->CopyDescriptors(1, &DestDescriptor, &InNumDescriptors, InNumDescriptors, InDescriptors, nullptr, Type);
+#if EXPLICIT_DESCRIPTOR_CACHE_FULL_COMPARE
 	for (uint32 i = 0; i < InNumDescriptors; ++i)
 	{
 		Descriptors[BaseIndex + i].ptr = InDescriptors[i].ptr;
 	}
+#endif
 }
 
 bool FD3D12ExplicitDescriptorHeap::CompareDescriptors(int32 BaseIndex, const D3D12_CPU_DESCRIPTOR_HANDLE* InDescriptors, uint32 InNumDescriptors)
@@ -350,8 +352,8 @@ void FD3D12ExplicitDescriptorCache::SetDescriptorHeaps(FD3D12CommandContext& Com
 	CommandContext.StateCache.GetDescriptorCache()->OverrideLastSetHeaps(ViewHeapToSet, SamplerHeapToSet);
 }
 
-// Returns descriptor heap base index for this descriptor table allocation or -1 if allocation failed.
-int32 FD3D12ExplicitDescriptorCache::AllocateDescriptorTable(const uint32* DescriptorVersions, const D3D12_CPU_DESCRIPTOR_HANDLE* Descriptors, uint32 NumDescriptors, D3D12_DESCRIPTOR_HEAP_TYPE Type, uint32 WorkerIndex)
+// Returns descriptor heap base index for this descriptor table allocation (checking for duplicates and reusing existing tables) or -1 if allocation failed.
+int32 FD3D12ExplicitDescriptorCache::AllocateDeduplicated(const uint32* DescriptorVersions, const D3D12_CPU_DESCRIPTOR_HANDLE* Descriptors, uint32 NumDescriptors, D3D12_DESCRIPTOR_HEAP_TYPE Type, uint32 WorkerIndex)
 {
 	checkSlow(Type == D3D12_DESCRIPTOR_HEAP_TYPE_CBV_SRV_UAV || Type == D3D12_DESCRIPTOR_HEAP_TYPE_SAMPLER);
 
@@ -369,8 +371,8 @@ int32 FD3D12ExplicitDescriptorCache::AllocateDescriptorTable(const uint32* Descr
 	if (DescriptorTableBaseIndex != INDEX_NONE)
 	{
 	#if EXPLICIT_DESCRIPTOR_CACHE_FULL_COMPARE
-		if (ensureMsgf(Heap.CompareDescriptors(DescriptorTableBaseIndex, Descriptors, NumDescriptors), 
-		               TEXT("Explicit descriptor cache hash collision detected!")))
+		if (ensureMsgf(Heap.CompareDescriptors(DescriptorTableBaseIndex, Descriptors, NumDescriptors),
+			TEXT("Explicit descriptor cache hash collision detected!")))
 	#endif
 		{
 			return DescriptorTableBaseIndex;
@@ -402,7 +404,18 @@ int32 FD3D12ExplicitDescriptorCache::AllocateDescriptorTable(const uint32* Descr
 		}
 	}
 
-	DescriptorTableBaseIndex = Heap.Allocate(NumDescriptors);
+	DescriptorTableBaseIndex = Allocate(Descriptors, NumDescriptors, Type, WorkerIndex);
+	return DescriptorTableBaseIndex;
+}
+
+// Returns descriptor heap base index for this descriptor table allocation or -1 if allocation failed.
+int32 FD3D12ExplicitDescriptorCache::Allocate(const D3D12_CPU_DESCRIPTOR_HANDLE* Descriptors, uint32 NumDescriptors, D3D12_DESCRIPTOR_HEAP_TYPE Type, uint32 WorkerIndex)
+{
+	checkSlow(Type == D3D12_DESCRIPTOR_HEAP_TYPE_CBV_SRV_UAV || Type == D3D12_DESCRIPTOR_HEAP_TYPE_SAMPLER);
+
+	FD3D12ExplicitDescriptorHeap& Heap = (Type == D3D12_DESCRIPTOR_HEAP_TYPE_CBV_SRV_UAV) ? ViewHeap : SamplerHeap;
+
+	int32 DescriptorTableBaseIndex = Heap.Allocate(NumDescriptors);
 
 	if (DescriptorTableBaseIndex == INDEX_NONE)
 	{
@@ -433,4 +446,4 @@ int32 FD3D12ExplicitDescriptorCache::AllocateDescriptorTable(const uint32* Descr
 	}
 
 	return DescriptorTableBaseIndex;
-}
\ No newline at end of file
+}
diff --git a/Engine/Source/Runtime/D3D12RHI/Private/D3D12ExplicitDescriptorCache.h b/Engine/Source/Runtime/D3D12RHI/Private/D3D12ExplicitDescriptorCache.h
--- a/Engine/Source/Runtime/D3D12RHI/Private/D3D12ExplicitDescriptorCache.h
+++ b/Engine/Source/Runtime/D3D12RHI/Private/D3D12ExplicitDescriptorCache.h
@@ -123,7 +123,10 @@ public:
 	void SetDescriptorHeaps(FD3D12CommandContext& CommandContext);
 
 	// Returns descriptor heap base index for this descriptor table allocation or -1 if allocation failed.
-	int32 AllocateDescriptorTable(const uint32* DescriptorVersions, const D3D12_CPU_DESCRIPTOR_HANDLE* Descriptors, uint32 NumDescriptors, D3D12_DESCRIPTOR_HEAP_TYPE Type, uint32 WorkerIndex);
+	int32 Allocate(const D3D12_CPU_DESCRIPTOR_HANDLE* Descriptors, uint32 NumDescriptors, D3D12_DESCRIPTOR_HEAP_TYPE Type, uint32 WorkerIndex);
+
+	// Returns descriptor heap base index for this descriptor table allocation (checking for duplicates and reusing existing tables) or -1 if allocation failed.
+	int32 AllocateDeduplicated(const uint32* DescriptorVersions, const D3D12_CPU_DESCRIPTOR_HANDLE* Descriptors, uint32 NumDescriptors, D3D12_DESCRIPTOR_HEAP_TYPE Type, uint32 WorkerIndex);
 
 	FD3D12ExplicitDescriptorHeap ViewHeap;
 	FD3D12ExplicitDescriptorHeap SamplerHeap;
diff --git a/Engine/Source/Runtime/D3D12RHI/Private/D3D12RayTracing.cpp b/Engine/Source/Runtime/D3D12RHI/Private/D3D12RayTracing.cpp
--- a/Engine/Source/Runtime/D3D12RHI/Private/D3D12RayTracing.cpp
+++ b/Engine/Source/Runtime/D3D12RHI/Private/D3D12RayTracing.cpp
@@ -4411,7 +4411,7 @@ static bool SetRayTracingShaderResources(
 	const uint32 NumSRVs = Shader->ResourceCounts.NumSRVs;
 	if (NumSRVs)
 	{
-		const int32 DescriptorTableBaseIndex = DescriptorCache.AllocateDescriptorTable(Bindings.SRVVersions, Bindings.LocalSRVs, NumSRVs, D3D12_DESCRIPTOR_HEAP_TYPE_CBV_SRV_UAV, WorkerIndex);
+		const int32 DescriptorTableBaseIndex = DescriptorCache.AllocateDeduplicated(Bindings.SRVVersions, Bindings.LocalSRVs, NumSRVs, D3D12_DESCRIPTOR_HEAP_TYPE_CBV_SRV_UAV, WorkerIndex);
 		if (DescriptorTableBaseIndex < 0)
 		{
 			return false;
@@ -4427,7 +4427,7 @@ static bool SetRayTracingShaderResources(
 	const uint32 NumUAVs = Shader->ResourceCounts.NumUAVs;
 	if (NumUAVs)
 	{
-		const int32 DescriptorTableBaseIndex = DescriptorCache.AllocateDescriptorTable(Bindings.UAVVersions, Bindings.LocalUAVs, NumUAVs, D3D12_DESCRIPTOR_HEAP_TYPE_CBV_SRV_UAV, WorkerIndex);
+		const int32 DescriptorTableBaseIndex = DescriptorCache.AllocateDeduplicated(Bindings.UAVVersions, Bindings.LocalUAVs, NumUAVs, D3D12_DESCRIPTOR_HEAP_TYPE_CBV_SRV_UAV, WorkerIndex);
 		if (DescriptorTableBaseIndex < 0)
 		{
 			return false;
@@ -4445,7 +4445,7 @@ static bool SetRayTracingShaderResources(
 	{
 	#if D3D12RHI_USE_CONSTANT_BUFFER_VIEWS
 
-		const uint32 DescriptorTableBaseIndex = DescriptorCache.AllocateDescriptorTable(Bindings.CBVVersions, Bindings.LocalCBVs, NumCBVs, D3D12_DESCRIPTOR_HEAP_TYPE_CBV_SRV_UAV, WorkerIndex);
+		const uint32 DescriptorTableBaseIndex = DescriptorCache.AllocateDeduplicated(Bindings.CBVVersions, Bindings.LocalCBVs, NumCBVs, D3D12_DESCRIPTOR_HEAP_TYPE_CBV_SRV_UAV, WorkerIndex);
 		const uint32 BindSlot = RootSignature->CBVRDTBindSlot(SF_Compute);
 		check(BindSlot != 0xFF);
 
@@ -4474,7 +4474,7 @@ static bool SetRayTracingShaderResources(
 	const uint32 NumSamplers = Shader->ResourceCounts.NumSamplers;
 	if (NumSamplers)
 	{
-		const int32 DescriptorTableBaseIndex = DescriptorCache.AllocateDescriptorTable(Bindings.SamplerVersions, Bindings.LocalSamplers, NumSamplers, D3D12_DESCRIPTOR_HEAP_TYPE_SAMPLER, WorkerIndex);
+		const int32 DescriptorTableBaseIndex = DescriptorCache.AllocateDeduplicated(Bindings.SamplerVersions, Bindings.LocalSamplers, NumSamplers, D3D12_DESCRIPTOR_HEAP_TYPE_SAMPLER, WorkerIndex);
 		if (DescriptorTableBaseIndex < 0)
 		{
 			return false;
