Make sure that number of constants in ray tracing local root signature matches in different places.

diff --git a/Engine/Source/Runtime/D3D12RHI/Public/D3D12RootSignatureDefinitions.h b/Engine/Source/Runtime/D3D12RHI/Public/D3D12RootSignatureDefinitions.h
--- a/Engine/Source/Runtime/D3D12RHI/Public/D3D12RootSignatureDefinitions.h
+++ b/Engine/Source/Runtime/D3D12RHI/Public/D3D12RootSignatureDefinitions.h
@@ -401,7 +401,7 @@ namespace D3D12ShaderUtils
 			Creator.AddShaderResourceViewParameter(RAY_TRACING_SYSTEM_INDEXBUFFER_REGISTER, UE_HLSL_SPACE_RAY_TRACING_SYSTEM);
 			Creator.AddShaderResourceViewParameter(RAY_TRACING_SYSTEM_VERTEXBUFFER_REGISTER, UE_HLSL_SPACE_RAY_TRACING_SYSTEM);
 
-			uint32 NumConstants = 4; // sizeof(FHitGroupSystemRootConstants) / sizeof(uint32);
+			uint32 NumConstants = sizeof(FHitGroupSystemRootConstants) / sizeof(uint32);
 			Creator.AddConstantsParameter(NumConstants, RAY_TRACING_SYSTEM_ROOTCONSTANT_REGISTER, UE_HLSL_SPACE_RAY_TRACING_SYSTEM);
 		}
 
