Kills waveops on -d3d12 -sm5

#tests win64

diff --git a/Engine/Config/Windows/DataDrivenPlatformInfo.ini b/Engine/Config/Windows/DataDrivenPlatformInfo.ini
--- a/Engine/Config/Windows/DataDrivenPlatformInfo.ini
+++ b/Engine/Config/Windows/DataDrivenPlatformInfo.ini
@@ -46,9 +46,6 @@ bSupportsRayTracingProceduralPrimitive=true
 bSupportsRayTracingIndirectInstanceData=true
 bSupportsPathTracing=true
 bSupportsHighEndRayTracingReflections=true
-bSupportsWaveOperations=RuntimeDependent
-MinimumWaveSize=4
-MaximumWaveSize=128
 bSupportsGPUScene=true
 bSupportsNanite=false
 bSupportsSceneDataCompressedTransforms=true
diff --git a/Engine/Source/Runtime/D3D12RHI/Private/D3D12Adapter.cpp b/Engine/Source/Runtime/D3D12RHI/Private/D3D12Adapter.cpp
--- a/Engine/Source/Runtime/D3D12RHI/Private/D3D12Adapter.cpp
+++ b/Engine/Source/Runtime/D3D12RHI/Private/D3D12Adapter.cpp
@@ -1081,11 +1081,14 @@ void FD3D12Adapter::InitializeDevices()
 
 			// Detect availability of shader model 6.0 wave operations
 			{
-				D3D12_FEATURE_DATA_D3D12_OPTIONS1 Features{};
-				RootDevice->CheckFeatureSupport(D3D12_FEATURE_D3D12_OPTIONS1, &Features, sizeof(Features));
-				GRHISupportsWaveOperations = Features.WaveOps;
-				GRHIMinimumWaveSize = Features.WaveLaneCountMin;
-				GRHIMaximumWaveSize = Features.WaveLaneCountMax;
+				if (GMaxRHIFeatureLevel >= ERHIFeatureLevel::SM6)
+				{
+					D3D12_FEATURE_DATA_D3D12_OPTIONS1 Features{};
+					RootDevice->CheckFeatureSupport(D3D12_FEATURE_D3D12_OPTIONS1, &Features, sizeof(Features));
+					GRHISupportsWaveOperations = Features.WaveOps;
+					GRHIMinimumWaveSize = Features.WaveLaneCountMin;
+					GRHIMaximumWaveSize = Features.WaveLaneCountMax;
+				}
 
 				if (GRHISupportsWaveOperations)
 				{
