Export DRED breadcrumbs as part of the crash payload.

DRED breadcrumbs are processed into breadcrumb nodes first and then added to the crash payload, unless RHI breadcrumbs are enabled (in that case, RHI breadcrumbs are preferred).

A breadcrumb tree is reconstructed from the linearized DRED history. BeginEvent/EndEvent operations are used to mark the start/end of a scope (breadcrumb node with potential children). An event is marked as "Finished" if the EndEvent operation has been reported as completed (e.g. it's at or before the last completed operation as reported by DRED).

Note that especially when lightweight (markers only) DRED is active, the command history may begin and end at arbitrary points. This change should handle all possible cases, in particular the case where the history starts in the middle of a command queue where EndEvent operations do not have a matching BeginEvent operation before. In that case, we consider those events "Unknown" (since the contextual name is lost). For example, the following linear history where the final EndEvent does not have an associated BeginEvent:

```
1) BeginEvent [BeginOcclusionTests]
2) BeginEvent [ShadowFrustumQueries]
3) EndEvent
4) EndEvent
5) EndEvent
```

becomes

```
> Unknown Event
--- > BeginOcclusionTests
--- --- > ShadowFrustumQueries
```

#tests Tested that the produced breadcrumb tree from DRED matches the tree produced by RHI breadcrumbs (to the extent possible)

diff --git a/Engine/Source/Runtime/Core/Private/GenericPlatform/GenericPlatformCrashContext.cpp b/Engine/Source/Runtime/Core/Private/GenericPlatform/GenericPlatformCrashContext.cpp
--- a/Engine/Source/Runtime/Core/Private/GenericPlatform/GenericPlatformCrashContext.cpp
+++ b/Engine/Source/Runtime/Core/Private/GenericPlatform/GenericPlatformCrashContext.cpp
@@ -237,6 +237,7 @@ namespace NCached
 	static TMap<FString, FString> EngineData;
 	static TMap<FString, FString> GameData;
 	static TMap<FString, FGPUBreadcrumbCrashData> GPUBreadcrumbsByQueue;
+	static FString GPUBreadcrumbsSource;
 
 	template <size_t CharCount, typename CharType>
 	void Set(CharType(&Dest)[CharCount], const CharType* pSrc)
@@ -1098,6 +1099,8 @@ void FGenericCrashContext::AddGPUBreadcrumbs() const
 	// if changes are made in the format exported by the engine.
 	AddSection(CommonBuffer, TEXT("FormatVersion"), FGPUBreadcrumbCrashData::GetBreadcrumbStringFormatVersion());
 
+	AddSection(CommonBuffer, TEXT("Source"), NCached::GPUBreadcrumbsSource);
+
 	for (auto& [Queue, Breadcrumbs] : NCached::GPUBreadcrumbsByQueue)
 	{
 		BeginSection(CommonBuffer, TEXT("Queue"));
@@ -1315,6 +1318,22 @@ void FGenericCrashContext::SetGPUBreadcrumbs(const FString& GPUQueueName, const
 	NCached::GPUBreadcrumbsByQueue.Emplace(GPUQueueName, FGPUBreadcrumbCrashData(Breadcrumbs));
 }
 
+void FGenericCrashContext::SetGPUBreadcrumbsSource(const FString& GPUBreadcrumbsSource)
+{
+	NCached::GPUBreadcrumbsSource = GPUBreadcrumbsSource;
+}
+
+const FString& FGenericCrashContext::GetGPUBreadcrumbsSource()
+{
+	return NCached::GPUBreadcrumbsSource;
+}
+
+void FGenericCrashContext::ResetGPUBreadcrumbsData()
+{
+	NCached::GPUBreadcrumbsByQueue.Empty();
+	NCached::GPUBreadcrumbsSource.Empty();
+}
+
 const TMap<FString, FString>& FGenericCrashContext::GetEngineData()
 {
 	return NCached::EngineData;
diff --git a/Engine/Source/Runtime/Core/Public/GenericPlatform/GenericPlatformCrashContext.h b/Engine/Source/Runtime/Core/Public/GenericPlatform/GenericPlatformCrashContext.h
--- a/Engine/Source/Runtime/Core/Public/GenericPlatform/GenericPlatformCrashContext.h
+++ b/Engine/Source/Runtime/Core/Public/GenericPlatform/GenericPlatformCrashContext.h
@@ -497,9 +497,18 @@ public:
 	/** Updates (or adds if not already present) arbitrary engine data to the crash context (will remove the key if passed an empty string) */
 	CORE_API static void SetEngineData(const FString& Key, const FString& Value);
 
-	/** Updates (or adds if not already present) GPU breadcrumb data for a given GPU queue */
+	/** Updates (or adds if not already present) GPU breadcrumb data for a given GPU queue. */
 	CORE_API static void SetGPUBreadcrumbs(const FString& GPUQueueName, const TArray<FBreadcrumbNode>& Breadcrumbs);
 
+	/** Sets a named source for the GPU breadcrumbs, mainly used to identify which system produced them. */
+	CORE_API static void SetGPUBreadcrumbsSource(const FString& GPUBreadcrumbsSource);
+
+	/** Gets the named source for the GPU breadcrumbs. */
+	CORE_API static const FString& GetGPUBreadcrumbsSource();
+
+	/** Clears all the GPU breadcrumb data. */
+	CORE_API static void ResetGPUBreadcrumbsData();
+
 	/** Accessor for engine data change callback delegate */
 	static FEngineDataSetDelegate& OnEngineDataSetDelegate() { return OnEngineDataSet; }
 
diff --git a/Engine/Source/Runtime/D3D12RHI/Private/D3D12Util.cpp b/Engine/Source/Runtime/D3D12RHI/Private/D3D12Util.cpp
--- a/Engine/Source/Runtime/D3D12RHI/Private/D3D12Util.cpp
+++ b/Engine/Source/Runtime/D3D12RHI/Private/D3D12Util.cpp
@@ -163,6 +163,23 @@ static FString GetD3D12TextureFlagString(uint32 TextureFlags)
 	return TextureFormatText;
 }
 
+/* Export GPU breadcrumbs data as part of the crash payload. RHI breadcrumbs are preferred over other types and will overwrite if necessary. */
+static void ExportBreadcrumbDataAsCrashPayload(const FString& BreadcrumbSource, const FString& GPUQueueName, const TArray<FBreadcrumbNode>& Breadcrumbs)
+{
+	const FString& CurrentSource = FGenericCrashContext::GetGPUBreadcrumbsSource();
+
+	if (CurrentSource.IsEmpty() || BreadcrumbSource == CurrentSource || BreadcrumbSource == TEXT("RHI"))
+	{
+		if (CurrentSource != BreadcrumbSource)
+		{
+			FGenericCrashContext::ResetGPUBreadcrumbsData();
+		}
+
+		FGenericCrashContext::SetGPUBreadcrumbsSource(BreadcrumbSource);
+		FGenericCrashContext::SetGPUBreadcrumbs(GPUQueueName, Breadcrumbs);
+	}
+}
+
 FBreadcrumbNode CollectBreadcrumbNode(
 	D3D12RHI::FD3DGPUProfiler& GPUProfiler,
 	const TSharedPtr<FBreadcrumbStack>& Stack,
@@ -264,7 +281,7 @@ static bool LogBreadcrumbData(D3D12RHI::FD3DGPUProfiler& GPUProfiler, FD3D12Queu
 
 	if (!Nodes.IsEmpty())
 	{
-		FGenericCrashContext::SetGPUBreadcrumbs(GPUQueueName, Nodes);
+		ExportBreadcrumbDataAsCrashPayload(TEXT("RHI"), GPUQueueName, MoveTemp(Nodes));
 	}
 
 	const FD3D12DiagnosticBufferData* DiagnosticData = Queue.GetDiagnosticBufferData();
@@ -388,58 +405,176 @@ struct FDred_1_2
 	const D3D12_AUTO_BREADCRUMB_NODE1* BreadcrumbHead = nullptr;
 };
 
+// Should match all values from D3D12_AUTO_BREADCRUMB_OP
+static const TCHAR* BreadcrumbOpNames[] =
+{
+	TEXT("SetMarker"),
+	TEXT("BeginEvent"),
+	TEXT("EndEvent"),
+	TEXT("DrawInstanced"),
+	TEXT("DrawIndexedInstanced"),
+	TEXT("ExecuteIndirect"),
+	TEXT("Dispatch"),
+	TEXT("CopyBufferRegion"),
+	TEXT("CopyTextureRegion"),
+	TEXT("CopyResource"),
+	TEXT("CopyTiles"),
+	TEXT("ResolveSubresource"),
+	TEXT("ClearRenderTargetView"),
+	TEXT("ClearUnorderedAccessView"),
+	TEXT("ClearDepthStencilView"),
+	TEXT("ResourceBarrier"),
+	TEXT("ExecuteBundle"),
+	TEXT("Present"),
+	TEXT("ResolveQueryData"),
+	TEXT("BeginSubmission"),
+	TEXT("EndSubmission"),
+	TEXT("DecodeFrame"),
+	TEXT("ProcessFrames"),
+	TEXT("AtomicCopyBufferUint"),
+	TEXT("AtomicCopyBufferUint64"),
+	TEXT("ResolveSubresourceRegion"),
+	TEXT("WriteBufferImmediate"),
+	TEXT("DecodeFrame1"),
+	TEXT("SetProtectedResourceSession"),
+	TEXT("DecodeFrame2"),
+	TEXT("ProcessFrames1"),
+	TEXT("BuildRaytracingAccelerationStructure"),
+	TEXT("EmitRaytracingAccelerationStructurePostBuildInfo"),
+	TEXT("CopyRaytracingAccelerationStructure"),
+	TEXT("DispatchRays"),
+	TEXT("InitializeMetaCommand"),
+	TEXT("ExecuteMetaCommand"),
+	TEXT("EstimateMotion"),
+	TEXT("ResolveMotionVectorHeap"),
+	TEXT("SetPipelineState1"),
+	TEXT("InitializeExtensionCommand"),
+	TEXT("ExecuteExtensionCommand"),
+};
+static_assert(UE_ARRAY_COUNT(BreadcrumbOpNames) == D3D12_AUTO_BREADCRUMB_OP_EXECUTEEXTENSIONCOMMAND + 1, "OpNames array length mismatch");
+
+/** 
+ * Calculate the number of active scopes in the case of a DRED history where the number of 
+ * EndEvent operations does not match the number of BeginEvent operations.
+ * Practically, this would be the number of "missing" BeginEvent operations that, if added at
+ * the beginning of the history, would balance out all EndEvent operations found later on.
+ */
+template <typename FDredNode_T>
+static uint32 CalculateDREDUnknownActiveScopes(const FDredNode_T* DredNode)
+{
+	check(DredNode);
+
+	int32 NumOpenEvents = 0;
+	int32 MaxUnknownActiveScopes = 0;
+	for (uint32 Op = 0; Op < DredNode->BreadcrumbCount; ++Op)
+	{
+		D3D12_AUTO_BREADCRUMB_OP BreadcrumbOp = DredNode->pCommandHistory[Op];
+		if (BreadcrumbOp == D3D12_AUTO_BREADCRUMB_OP_BEGINEVENT)
+		{
+			NumOpenEvents++;
+		}
+		else if (BreadcrumbOp == D3D12_AUTO_BREADCRUMB_OP_ENDEVENT)
+		{
+			NumOpenEvents--;
+		}
+
+		MaxUnknownActiveScopes = FMath::Min(NumOpenEvents, MaxUnknownActiveScopes);
+	}
+
+	return FMath::Abs(MaxUnknownActiveScopes);
+}
+
+template <typename FDredNode_T>
+static TArray<FBreadcrumbNode> CollectDREDBreadcrumbNodes(const FDredNode_T* DredNode)
+{
+	check(DredNode && DredNode->pLastBreadcrumbValue);
+	uint32 LastCompletedOp = *DredNode->pLastBreadcrumbValue;
+	if (LastCompletedOp == DredNode->BreadcrumbCount || LastCompletedOp == 0)
+	{
+		return {};
+	}
+
+	TMap<uint32, const wchar_t*> ContextStrings;
+	for (const D3D12_DRED_BREADCRUMB_CONTEXT& Context : GetBreadcrumbContexts(DredNode))
+	{
+		ContextStrings.Add(Context.BreadcrumbIndex, Context.pContextString);
+	}
+
+	// Create a root node that will hold all events as children. The root itself will be discarded.
+	FBreadcrumbNode Root;
+	Root.Name = TEXT("");
+	Root.State = EBreadcrumbState::Invalid;
+
+	TArray<FBreadcrumbNode*> ParentChain = { &Root };
+
+	// If we have open scopes, create them now as "Unknown events".
+	uint32 NumOpenScopes = CalculateDREDUnknownActiveScopes(DredNode);
+	for (uint32 i = 0; i < NumOpenScopes; ++i)
+	{
+		FBreadcrumbNode& UnknownNode = ParentChain.Last()->Children.Emplace_GetRef();
+		UnknownNode.Name = TEXT("Unknown event");
+		UnknownNode.State = EBreadcrumbState::Active;
+		ParentChain.Push(&UnknownNode);
+	}
+
+	for (uint32 Op = 0; Op < DredNode->BreadcrumbCount; ++Op)
+	{
+		D3D12_AUTO_BREADCRUMB_OP BreadcrumbOp = DredNode->pCommandHistory[Op];
+		bool bCompleted = Op < LastCompletedOp;
+		auto OpContextStr = ContextStrings.Find(Op);
+
+		if (BreadcrumbOp == D3D12_AUTO_BREADCRUMB_OP_BEGINEVENT)
+		{
+			// This is a begin event, potentially with children events.
+			FBreadcrumbNode& BreadcrumbNode = ParentChain.Last()->Children.Emplace_GetRef();
+			BreadcrumbNode.Name = OpContextStr ? *OpContextStr : TEXT("Unknown event");
+			BreadcrumbNode.State = bCompleted ? EBreadcrumbState::Active : EBreadcrumbState::NotStarted;
+
+			ParentChain.Push(&BreadcrumbNode);
+		}
+		else if (BreadcrumbOp == D3D12_AUTO_BREADCRUMB_OP_ENDEVENT)
+		{
+			FBreadcrumbNode* Parent = ParentChain.Pop();
+			if (Parent->State == EBreadcrumbState::Invalid)
+			{
+				// If we reach this point, the DRED breadcrumbs are malformed, and some
+				// basic invariants around matching BeginEvent/EndEvent do not hold.
+				// Return gracefully and do not attempt to process further.
+				return {};
+			}
+
+			// This is the end event for the parent node. Mark the whole event as finished
+			// if this end event was completed.
+			if (bCompleted && Parent->State == EBreadcrumbState::Active)
+			{
+				Parent->State = EBreadcrumbState::Finished;
+			}
+		}
+		else
+		{
+			// This is a miscellaneous event between a BeginEvent and an EndEvent.
+			const TCHAR* OpName = (BreadcrumbOp < UE_ARRAY_COUNT(BreadcrumbOpNames)) ? BreadcrumbOpNames[BreadcrumbOp] : TEXT("Unknown Op");
+
+			FBreadcrumbNode& BreadcrumbNode = ParentChain.Last()->Children.Emplace_GetRef();
+			if (OpContextStr)
+			{
+				BreadcrumbNode.Name = FString::Printf(TEXT("%s [%s]"), OpName, *OpContextStr);
+			}
+			else
+			{
+				BreadcrumbNode.Name = OpName;
+			}
+			BreadcrumbNode.State = bCompleted ? EBreadcrumbState::Finished : EBreadcrumbState::NotStarted;
+		}
+	}
+
+	return Root.Children;
+}
+
 /** Log the DRED data to Error log if available */
 template <typename FDred_T>
 static bool LogDREDData(ID3D12Device* Device, bool bTrackingAllAllocations, D3D12_GPU_VIRTUAL_ADDRESS& OutPageFaultGPUAddress)
 {
-	// Should match all values from D3D12_AUTO_BREADCRUMB_OP
-	static const TCHAR* OpNames[] =
-	{
-		TEXT("SetMarker"),
-		TEXT("BeginEvent"),
-		TEXT("EndEvent"),
-		TEXT("DrawInstanced"),
-		TEXT("DrawIndexedInstanced"),
-		TEXT("ExecuteIndirect"),
-		TEXT("Dispatch"),
-		TEXT("CopyBufferRegion"),
-		TEXT("CopyTextureRegion"),
-		TEXT("CopyResource"),
-		TEXT("CopyTiles"),
-		TEXT("ResolveSubresource"),
-		TEXT("ClearRenderTargetView"),
-		TEXT("ClearUnorderedAccessView"),
-		TEXT("ClearDepthStencilView"),
-		TEXT("ResourceBarrier"),
-		TEXT("ExecuteBundle"),
-		TEXT("Present"),
-		TEXT("ResolveQueryData"),
-		TEXT("BeginSubmission"),
-		TEXT("EndSubmission"),
-		TEXT("DecodeFrame"),
-		TEXT("ProcessFrames"),
-		TEXT("AtomicCopyBufferUint"),
-		TEXT("AtomicCopyBufferUint64"),
-		TEXT("ResolveSubresourceRegion"),
-		TEXT("WriteBufferImmediate"),
-		TEXT("DecodeFrame1"),
-		TEXT("SetProtectedResourceSession"),
-		TEXT("DecodeFrame2"),
-		TEXT("ProcessFrames1"),
-		TEXT("BuildRaytracingAccelerationStructure"),
-		TEXT("EmitRaytracingAccelerationStructurePostBuildInfo"),
-		TEXT("CopyRaytracingAccelerationStructure"),
-		TEXT("DispatchRays"),
-		TEXT("InitializeMetaCommand"),
-		TEXT("ExecuteMetaCommand"),
-		TEXT("EstimateMotion"),
-		TEXT("ResolveMotionVectorHeap"),
-		TEXT("SetPipelineState1"),
-		TEXT("InitializeExtensionCommand"),
-		TEXT("ExecuteExtensionCommand"),
-	};
-	static_assert(UE_ARRAY_COUNT(OpNames) == D3D12_AUTO_BREADCRUMB_OP_EXECUTEEXTENSIONCOMMAND + 1, "OpNames array length mismatch");
-
 	// Should match all valid values from D3D12_DRED_ALLOCATION_TYPE
 	static const TCHAR* AllocTypesNames[] =
 	{
@@ -481,7 +616,6 @@ static bool LogDREDData(ID3D12Device* Device, bool bTrackingAllAllocations, D3D1
 	{
 		if (Dred.BreadcrumbHead)
 		{
-			bHasValidBreadcrumbData = true;
 			UE_LOG(LogD3D12RHI, Error, TEXT("DRED: Last tracked GPU operations:"));
 
 			FString ContextStr;
@@ -495,6 +629,7 @@ static bool LogDREDData(ID3D12Device* Device, bool bTrackingAllAllocations, D3D1
 
 				if (LastCompletedOp != Node->BreadcrumbCount && LastCompletedOp != 0)
 				{
+					bHasValidBreadcrumbData = true;
 					UE_LOG(LogD3D12RHI, Error, TEXT("DRED: Commandlist \"%s\" on CommandQueue \"%s\", %d completed of %d"), Node->pCommandListDebugNameW, Node->pCommandQueueDebugNameW, LastCompletedOp, Node->BreadcrumbCount);
 					TracedCommandLists++;
 
@@ -523,9 +658,16 @@ static bool LogDREDData(ID3D12Device* Device, bool bTrackingAllAllocations, D3D1
 							ContextStr.Reset();
 						}
 
-						const TCHAR* OpName = (BreadcrumbOp < UE_ARRAY_COUNT(OpNames)) ? OpNames[BreadcrumbOp] : TEXT("Unknown Op");
+						const TCHAR* OpName = (BreadcrumbOp < UE_ARRAY_COUNT(BreadcrumbOpNames)) ? BreadcrumbOpNames[BreadcrumbOp] : TEXT("Unknown Op");
 						UE_LOG(LogD3D12RHI, Error, TEXT("\tOp: %d, %s%s%s"), Op, OpName, *ContextStr, (Op + 1 == LastCompletedOp) ? TEXT(" - LAST COMPLETED") : TEXT(""));
 					}
+
+					// Collect and export breadcrumb data separately as part of the crash payload.
+					TArray<FBreadcrumbNode> Breadcrumbs = CollectDREDBreadcrumbNodes(Node);
+					if (!Breadcrumbs.IsEmpty())
+					{
+						ExportBreadcrumbDataAsCrashPayload(TEXT("DRED"), Node->pCommandQueueDebugNameW, MoveTemp(Breadcrumbs));
+					}
 				}
 
 				Node = Node->pNext;
