Disable residency tracking for transient heaps.

We seem to have a bug where a transient heap is evicted by the residency manager, but is not restored when a resource from that heap is needed in a draw or dispatch, leading to GPU page faults. We don't really need to evict these heaps, since they're used throughout the frame, and they're deleted after a few frames when they're no longer needed. Disabling tracking for them shouldn't affect memory usage, but might mitigate this crash, or at least provide better GPU crash information when it happens.

#rnx

diff --git a/Engine/Source/Runtime/D3D12RHI/Private/D3D12TransientResourceAllocator.cpp b/Engine/Source/Runtime/D3D12RHI/Private/D3D12TransientResourceAllocator.cpp
--- a/Engine/Source/Runtime/D3D12RHI/Private/D3D12TransientResourceAllocator.cpp
+++ b/Engine/Source/Runtime/D3D12RHI/Private/D3D12TransientResourceAllocator.cpp
@@ -92,7 +92,17 @@ FD3D12TransientHeap::FD3D12TransientHeap(const FInitializer& Initializer, FD3D12
 	Heap = new FD3D12Heap(Device, VisibleNodeMask);
 	Heap->SetHeap(D3DHeap, TEXT("TransientResourceAllocator Backing Heap"), true, true);
 	Heap->SetIsTransient(true);
-	Heap->BeginTrackingResidency(Desc.SizeInBytes);
+
+	// UE-174791: we seem to have a bug related to residency where transient heaps are evicted, but are not restored correctly before a resource
+	// is needed, leading to GPU page faults like this one:
+	// 
+	// PageFault: Found 1 active heaps containing page fault address
+	//  	GPU Address : "0x1008800000" - Size : 128.00 MB - Name : TransientResourceAllocator Backing Heap
+	// 
+	// We don't really need to evict these heaps anyway, since they are used throughout the frame, and are garbage-collected after a few frames
+	// when they're no longer needed. Disabling residency tracking will not fix the underlying bug, but should make it less likely to occur,
+	// and might make the GPU crash data more useful when it does happen.
+	//Heap->BeginTrackingResidency(Desc.SizeInBytes);
 
 	SetGpuVirtualAddress(Heap->GetGPUVirtualAddress());
 
