Adding RayTracing RootSignature string generation for use in shader versioning.

diff --git a/Engine/Source/Runtime/D3D12RHI/Public/D3D12RootSignatureDefinitions.h b/Engine/Source/Runtime/D3D12RHI/Public/D3D12RootSignatureDefinitions.h
--- a/Engine/Source/Runtime/D3D12RHI/Public/D3D12RootSignatureDefinitions.h
+++ b/Engine/Source/Runtime/D3D12RHI/Public/D3D12RootSignatureDefinitions.h
@@ -172,10 +172,12 @@ namespace D3D12ShaderUtils
 	struct FRootSignatureCreator
 	{
 		ED3D12RootSignatureFlags Flags = ED3D12RootSignatureFlags::None;
+		uint32 RegisterSpace = 0;
 
-		virtual ~FRootSignatureCreator() { }
+		virtual ~FRootSignatureCreator() = default;
 
 		virtual void AddRootFlag(D3D12_ROOT_SIGNATURE_FLAGS Flag) = 0;
+		virtual void AddShaderResourceViewParameter(uint32 Register, uint32 Space) = 0;
 		virtual void AddTable(ERootSignatureVisibility Visibility, ERootSignatureRangeType Type, int32 NumDescriptors, D3D12_DESCRIPTOR_RANGE_FLAGS FlagsOverride = D3D12_DESCRIPTOR_RANGE_FLAG_NONE) = 0;
 		virtual void AddConstantsParameter(uint32 Num32BitValues, uint32 Register, uint32 Space) = 0;
 
@@ -199,6 +201,11 @@ namespace D3D12ShaderUtils
 			}
 		}
 
+		void SetRegisterSpace(uint32 InSpace)
+		{
+			RegisterSpace = InSpace;
+		}
+
 		inline bool ShouldSkipType(ERootSignatureRangeType Type)
 		{
 			if (Type == ERootSignatureRangeType::SRV || Type == ERootSignatureRangeType::UAV)
@@ -218,7 +225,6 @@ namespace D3D12ShaderUtils
 	struct FBinaryRootSignatureCreator : public FRootSignatureCreator
 	{
 		D3D12_ROOT_SIGNATURE_FLAGS RootFlags = D3D12_ROOT_SIGNATURE_FLAG_NONE;
-		uint32 RegisterSpace = 0;
 		TArray<CD3DX12_DESCRIPTOR_RANGE1> DescriptorRanges;
 		TArray<CD3DX12_ROOT_PARAMETER1> Parameters;
 		CD3DX12_VERSIONED_ROOT_SIGNATURE_DESC RootDesc;
@@ -232,11 +238,6 @@ namespace D3D12ShaderUtils
 		};
 		TArray<FPendingTable> PendingTables;
 
-		void SetRegisterSpace(uint32 InSpace)
-		{
-			RegisterSpace = InSpace;
-		}
-
 		void AddRootFlag(D3D12_ROOT_SIGNATURE_FLAGS RootFlag) override
 		{
 			RootFlags |= RootFlag;
@@ -305,6 +306,16 @@ namespace D3D12ShaderUtils
 			RootFlags += GetFlagName(InFlag);
 		}
 
+		void AddShaderResourceViewParameter(uint32 Register, uint32 Space) override
+		{
+			FString Line = FString::Printf(TEXT("SRV(t%d, space=%d))"), Register, Space);
+			if (Table.Len() > 0)
+			{
+				Table += ",";
+			}
+			Table += Line;
+		}
+
 		void AddConstantsParameter(uint32 Num32BitValues, uint32 Register, uint32 Space) override
 		{
 			Constants.Appendf(TEXT("RootConstants(num32BitConstants=%d, b%d, space=%d),"), Num32BitValues, Register, Space);
@@ -314,9 +325,8 @@ namespace D3D12ShaderUtils
 		{
 			if (!ShouldSkipType(Type))
 			{
-				FString Line = FString::Printf(TEXT("DescriptorTable(visibility=%s, %s0, numDescriptors=%d%s"/*, flags = DESCRIPTORS_VOLATILE*/"))"),
-					GetVisibilityFlag(Visibility), GetTypePrefix(Type), NumDescriptors,
-					TEXT("")
+				FString Line = FString::Printf(TEXT("DescriptorTable(visibility=%s, %s0, space=%d, numDescriptors=%d))"),
+					GetVisibilityFlag(Visibility), GetTypePrefix(Type), RegisterSpace, NumDescriptors
 				);
 				if (Table.Len() > 0)
 				{
@@ -390,7 +400,7 @@ namespace D3D12ShaderUtils
 	}
 
 #if !defined(D3D12RHI_TOOLS_RAYTRACING_SHADERS_UNSUPPORTED)
-	inline void CreateRayTracingSignature(FBinaryRootSignatureCreator& Creator, bool bLocalRootSignature, D3D12_ROOT_SIGNATURE_FLAGS BaseRootFlags, ED3D12RootSignatureFlags InFlags)
+	inline void CreateRayTracingSignature(FRootSignatureCreator& Creator, bool bLocalRootSignature, D3D12_ROOT_SIGNATURE_FLAGS BaseRootFlags, ED3D12RootSignatureFlags InFlags)
 	{
 		Creator.SetFlags(InFlags);
 		Creator.AddRootFlag(BaseRootFlags);
@@ -409,6 +419,13 @@ namespace D3D12ShaderUtils
 
 		Creator.AddTable(ERootSignatureVisibility::All, ERootSignatureRangeType::UAV, MAX_UAVS, D3D12_DESCRIPTOR_RANGE_FLAG_DATA_STATIC_WHILE_SET_AT_EXECUTE | D3D12_DESCRIPTOR_RANGE_FLAG_DESCRIPTORS_VOLATILE);
 	}
+
+	inline FString GenerateRayTracingSignatureString(bool bLocalRootSignature, D3D12_ROOT_SIGNATURE_FLAGS BaseRootFlags, ED3D12RootSignatureFlags InFlags)
+	{
+		FTextRootSignatureCreator Creator;
+		CreateRayTracingSignature(Creator, bLocalRootSignature, BaseRootFlags, InFlags);
+		return Creator.GenerateString();
+	}
 #endif //!defined(D3D12RHI_TOOLS_RAYTRACING_SHADERS_UNSUPPORTED)
 
 	inline FString GenerateRootSignatureString(EShaderFrequency InFrequency, ED3D12RootSignatureFlags InFlags)
