D3D12 Reserved Resource improvements

- Add basic reserved texture creation unit test
- Add basic support for packed tile layout (still only single mip and no array support, for simplicity)
- Deprecated GRHISupportsReservedResources (use GRHIGlobals.ReservedResources.Supported directly)

diff --git a/Engine/Plugins/Tests/RHITests/Source/RHITests/Private/RHIReservedResourceTests.cpp b/Engine/Plugins/Tests/RHITests/Source/RHITests/Private/RHIReservedResourceTests.cpp
new file mode 100644
--- /dev/null
+++ b/Engine/Plugins/Tests/RHITests/Source/RHITests/Private/RHIReservedResourceTests.cpp
@@ -0,0 +1,86 @@
+// Copyright Epic Games, Inc. All Rights Reserved.
+
+#include "RHIReservedResourceTests.h"
+#include "CommonRenderResources.h"
+#include "RenderCaptureInterface.h"
+#include "RHIStaticStates.h"
+
+bool FRHIReservedResourceTests::Test_ReservedResource_CreateTexture(FRHICommandListImmediate& RHICmdList)
+{
+	if (!GRHIGlobals.ReservedResources.Supported)
+	{
+		return true;
+	}
+
+	const ETextureCreateFlags CreateFlags =
+		  TexCreate_ShaderResource
+		| TexCreate_ReservedResource
+		| TexCreate_ImmediateCommit;
+
+	const EPixelFormat TestFormats[] =
+	{
+		PF_DXT1,
+		PF_R8,
+		PF_R32_UINT,
+	};
+
+	struct FTestTextureConfig
+	{
+		FIntPoint Extent = FIntPoint(1,1);
+		int32 NumMips = 1;
+		int32 ArraySize = 1;
+	};
+
+	const int32 TileSize = GRHIGlobals.ReservedResources.TextureArrayMinimumMipDimension;
+
+	FTestTextureConfig TestConfigs[] =
+	{
+		// Texture arrays
+		FTestTextureConfig{.Extent = FIntPoint(TileSize * 8, TileSize),  .NumMips = 1, .ArraySize = 2},
+		FTestTextureConfig{.Extent = FIntPoint(TileSize * 2, TileSize),  .NumMips = 1, .ArraySize = 3},
+
+		// Regular textures
+		FTestTextureConfig{.Extent = FIntPoint(4096, 128), .NumMips = 1, .ArraySize = 1},
+		FTestTextureConfig{.Extent = FIntPoint(128, 128),  .NumMips = 1, .ArraySize = 1},
+		FTestTextureConfig{.Extent = FIntPoint(512, 128),  .NumMips = 1, .ArraySize = 1},
+		FTestTextureConfig{.Extent = FIntPoint(512, 128),  .NumMips = 1, .ArraySize = 1},
+		FTestTextureConfig{.Extent = FIntPoint(4, 4),      .NumMips = 1, .ArraySize = 1},
+
+#if 0 // TODO: as of 2023-10-18, reserved textures with mips are not supported
+		FTestTextureConfig{.Extent = FIntPoint(128, 128),  .NumMips = 2, .ArraySize = 2},
+		FTestTextureConfig{.Extent = FIntPoint(512, 128),  .NumMips = 2, .ArraySize = 2},
+		FTestTextureConfig{.Extent = FIntPoint(512, 128),  .NumMips = 5, .ArraySize = 2},
+		FTestTextureConfig{.Extent = FIntPoint(4, 4),      .NumMips = 2, .ArraySize = 2},
+#endif
+	};
+
+	// Try to create resources of various formats and dimensions
+
+	for (FTestTextureConfig Config : TestConfigs)
+	{
+		for (EPixelFormat Format : TestFormats)
+		{
+			FRHITextureCreateDesc Desc;
+
+			if (Config.ArraySize > 1)
+			{
+				Desc = FRHITextureCreateDesc::Create2DArray(TEXT("TestReservedTexture2DArray"));
+				Desc.SetArraySize(Config.ArraySize);
+			}
+			else
+			{
+				Desc = FRHITextureCreateDesc::Create2D(TEXT("TestReservedTexture2D"));
+			}
+
+			Desc.SetFlags(CreateFlags)
+				.SetExtent(Config.Extent)
+				.SetNumMips(Config.NumMips)
+				.SetFormat(Format)
+				.SetInitialState(ERHIAccess::SRVCompute);
+
+			FTextureRHIRef Texture = RHICreateTexture(Desc);
+		}
+	}
+
+	return true;
+}
diff --git a/Engine/Plugins/Tests/RHITests/Source/RHITests/Private/RHITestsModule.cpp b/Engine/Plugins/Tests/RHITests/Source/RHITests/Private/RHITestsModule.cpp
--- a/Engine/Plugins/Tests/RHITests/Source/RHITests/Private/RHITestsModule.cpp
+++ b/Engine/Plugins/Tests/RHITests/Source/RHITests/Private/RHITestsModule.cpp
@@ -9,6 +9,7 @@
 #include "RHITextureTests.h"
 #include "RHIDrawTests.h"
 #include "RHIReadbackTests.h"
+#include "RHIReservedResourceTests.h"
 
 #define LOCTEXT_NAMESPACE "FRHITestsModule"
 
@@ -16,6 +17,13 @@ static bool RunTests_RenderThread(FRHICommandListImmediate& RHICmdList)
 {
 	bool bResult = true;
 
+	// ------------------------------------------------
+	// Reserved resources
+	// ------------------------------------------------
+	{
+		RUN_TEST(FRHIReservedResourceTests::Test_ReservedResource_CreateTexture(RHICmdList));
+	}
+
 	// ------------------------------------------------
 	// Drawing
 	// ------------------------------------------------
diff --git a/Engine/Plugins/Tests/RHITests/Source/RHITests/Public/RHIReservedResourceTests.h b/Engine/Plugins/Tests/RHITests/Source/RHITests/Public/RHIReservedResourceTests.h
new file mode 100644
--- /dev/null
+++ b/Engine/Plugins/Tests/RHITests/Source/RHITests/Public/RHIReservedResourceTests.h
@@ -0,0 +1,13 @@
+// Copyright Epic Games, Inc. All Rights Reserved.
+
+#pragma once
+
+#include "RHITestsCommon.h"
+
+class FRHIReservedResourceTests
+{
+public:
+
+	static bool Test_ReservedResource_CreateTexture(FRHICommandListImmediate& RHICmdList);
+};
+
diff --git a/Engine/Source/Runtime/D3D12RHI/Private/D3D12Allocation.cpp b/Engine/Source/Runtime/D3D12RHI/Private/D3D12Allocation.cpp
--- a/Engine/Source/Runtime/D3D12RHI/Private/D3D12Allocation.cpp
+++ b/Engine/Source/Runtime/D3D12RHI/Private/D3D12Allocation.cpp
@@ -1873,7 +1873,7 @@ HRESULT FD3D12TextureAllocatorPool::AllocateTexture(
 		Desc.SampleDesc.Count == 1)
 	{
 		// The top mip level must be less than 64 KB to use 4 KB alignment
-		Desc.Alignment = TextureCanBe4KAligned(Desc, UEFormat) ?
+		Desc.Alignment = FD3D12Texture::CanBe4KAligned(Desc, UEFormat) ?
 			D3D12_SMALL_RESOURCE_PLACEMENT_ALIGNMENT :
 			D3D12_DEFAULT_RESOURCE_PLACEMENT_ALIGNMENT;
 		const D3D12_RESOURCE_ALLOCATION_INFO Info = Device->GetResourceAllocationInfoUncached(Desc);
diff --git a/Engine/Source/Runtime/D3D12RHI/Private/D3D12Resources.cpp b/Engine/Source/Runtime/D3D12RHI/Private/D3D12Resources.cpp
--- a/Engine/Source/Runtime/D3D12RHI/Private/D3D12Resources.cpp
+++ b/Engine/Source/Runtime/D3D12RHI/Private/D3D12Resources.cpp
@@ -162,7 +162,7 @@ void FD3D12Resource::CommitReservedResource()
 	checkf(Desc.Dimension == D3D12_RESOURCE_DIMENSION_TEXTURE2D, TEXT("CommitReservedResource is currently only implemented for 2D textures"));
 	checkf(Desc.MipLevels == 1, TEXT("CommitReservedResource is currently only implemented for textures without mips"));
 
-	uint32 NumTiles = 0;
+	uint32 D3DResourceNumTiles = 0;
 	D3D12_PACKED_MIP_INFO PackedMipDesc = {};
 	D3D12_TILE_SHAPE TileShape = {};
 	const uint32 FirstSubresource = 0;
@@ -174,17 +174,17 @@ void FD3D12Resource::CommitReservedResource()
 	ID3D12Device* D3DDevice = GetParentDevice()->GetDevice();
 	FD3D12Adapter* Adapter = GetParentDevice()->GetParentAdapter();
 
-	D3DDevice->GetResourceTiling(GetResource(), &NumTiles, &PackedMipDesc, &TileShape, &NumSubresourceTilings, FirstSubresource, &SubresourceTiling);
+	D3DDevice->GetResourceTiling(GetResource(), &D3DResourceNumTiles, &PackedMipDesc, &TileShape, &NumSubresourceTilings, FirstSubresource, &SubresourceTiling);
 
 	const uint64 TileSizeInBytes = 65536; // reserved resource tiles are always 64KB
-	const uint64 TotalSize = NumTiles * TileSizeInBytes;
+	const uint64 TotalSize = D3DResourceNumTiles * TileSizeInBytes;
 	const uint64 MaxHeapSize = uint64(CVarD3D12ReservedResourceHeapSizeMB.GetValueOnAnyThread()) * 1024 * 1024;
 	const uint64 NumHeaps = FMath::DivideAndRoundUp(TotalSize, MaxHeapSize);
 
 	TArray<TRefCountPtr<ID3D12Heap>>& Heaps = ReservedResourceData->BackingHeaps;
 	Heaps.Reserve(NumHeaps);
 
-	const uint32 NumTilesPerHeap = uint32(MaxHeapSize / TileSizeInBytes);
+	const uint32 MaxTilesPerHeap = uint32(MaxHeapSize / TileSizeInBytes);
 
 	// NOTE: Accessing the queue from this thread is OK, as D3D12 runtime acquires a lock around all command queue APIs.
 	// https://microsoft.github.io/DirectX-Specs/d3d/CPUEfficiency.html#threading
@@ -205,13 +205,24 @@ void FD3D12Resource::CommitReservedResource()
 	BackingHeapProps.CreationNodeMask = GetGPUMask().GetNative();
 	BackingHeapProps.VisibleNodeMask = GetVisibilityMask().GetNative();
 
-	const uint32 NumTilesX = SubresourceTiling.WidthInTiles;
-	const uint32 NumTilesY = SubresourceTiling.HeightInTiles;
+	uint32 NumStandardTilesPerSubresource = 0;
+	uint32 NumTotalTiles = 0;
 
-	check(SubresourceTiling.DepthInTiles == 1);
+	if (PackedMipDesc.NumStandardMips != 0)
+	{
+		checkf(SubresourceTiling.DepthInTiles == 1, TEXT("3D reserved textures are not supported/implemented"));
+		NumStandardTilesPerSubresource = SubresourceTiling.WidthInTiles * SubresourceTiling.HeightInTiles;
+		NumTotalTiles = NumStandardTilesPerSubresource * SubresourceCount;
+	}
+	else // packed mip case
+	{
+		checkf(PackedMipDesc.NumPackedMips != 0, TEXT("It is expected that reserved resources have at least one standard or one packed tile mip level"));
+		NumTotalTiles = PackedMipDesc.NumTilesForPackedMips * SubresourceCount;
+	}
 
-	const uint32 NumTilesPerSubresource = NumTilesX * NumTilesY;
-	const uint32 NumTotalTiles = NumTilesPerSubresource * SubresourceCount;
+	checkf(D3DResourceNumTiles == NumTotalTiles,
+		TEXT("D3D resource size in tiles: %d, computed size in tiles: %d"),
+		D3DResourceNumTiles, NumTotalTiles);
 
 	uint32 NumMappedTiles = 0;
 
@@ -221,7 +232,7 @@ void FD3D12Resource::CommitReservedResource()
 
 		D3D12_TILE_REGION_SIZE RegionSize = {};
 		RegionSize.UseBox = false;
-		RegionSize.NumTiles = FMath::Min(NumTilesPerHeap, NumRemainingTiles);
+		RegionSize.NumTiles = FMath::Min(MaxTilesPerHeap, NumRemainingTiles);
 
 		const uint32 ThisHeapSize = RegionSize.NumTiles * TileSizeInBytes;
 		D3D12_HEAP_DESC NewHeapDesc = {};
@@ -238,11 +249,34 @@ void FD3D12Resource::CommitReservedResource()
 		}
 
 		D3D12_TILED_RESOURCE_COORDINATE ResourceCoordinate = {}; // Coordinates are in tiles, not pixels
-		const uint32 TileIndexInSubresource = NumMappedTiles % NumTilesPerSubresource;
-		ResourceCoordinate.Subresource = NumMappedTiles / NumTilesPerSubresource;
-		ResourceCoordinate.X = TileIndexInSubresource % NumTilesX;
-		ResourceCoordinate.Y = (TileIndexInSubresource / NumTilesX) % NumTilesY;
-		ResourceCoordinate.Z = 0; // Only simple 2D / Array2D textures are implemented
+		if (NumStandardTilesPerSubresource)
+		{
+			const uint32 TileIndexInSubresource = NumMappedTiles % NumStandardTilesPerSubresource;
+			ResourceCoordinate.Subresource = NumMappedTiles / NumStandardTilesPerSubresource;
+			ResourceCoordinate.X = TileIndexInSubresource % SubresourceTiling.WidthInTiles;
+			ResourceCoordinate.Y = (TileIndexInSubresource / SubresourceTiling.WidthInTiles) % SubresourceTiling.HeightInTiles;
+			ResourceCoordinate.Z = 0; // Only simple 2D / Array2D textures are implemented
+		}
+		else
+		{
+			// Packed mip level case:
+			// - Only simple 2D textures are expected (single subresource, no arrays)
+			// - Entire packed mip level must be covered in one map operation, so mapping origin is always 0
+
+			checkf(SubresourceCount == 1,
+				TEXT("Reserved textures with packed mips and multiple subresources are not supported. Current subresource count: %d"),
+				SubresourceCount);
+
+			checkf(RegionSize.NumTiles <= MaxTilesPerHeap,
+				TEXT("Reserved texture packed mip level requires tiles: %d, maximum supported tiles: %d. ")
+				TEXT("Increase d3d12.ReservedResourceHeapSizeMB or avoid packed mips by using a larger texture dimensions."),
+				RegionSize.NumTiles, MaxTilesPerHeap);
+
+			ResourceCoordinate.Subresource = 0; // Packed mips are currently supported for single subresource textures
+			ResourceCoordinate.X = 0;
+			ResourceCoordinate.Y = 0;
+			ResourceCoordinate.Z = 0;
+		}
 
 		const D3D12_TILE_RANGE_FLAGS RangeFlags = D3D12_TILE_RANGE_FLAG_NONE;
 		const uint32 HeapRangeStartOffsetInTiles = 0;
@@ -265,6 +299,10 @@ void FD3D12Resource::CommitReservedResource()
 		NumMappedTiles += RegionSize.NumTiles;
 		Heaps.Add(MoveTemp(NewHeap));
 	}
+
+	checkf(NumMappedTiles == D3DResourceNumTiles,
+		TEXT("Reserved resource was not fully processed while committing physical memory. Expected to process tiles: %d, actually processed: %d"),
+		D3DResourceNumTiles, NumMappedTiles);
 }
 
 ID3D12Pageable* FD3D12Resource::GetPageable()
@@ -538,8 +576,17 @@ HRESULT FD3D12Adapter::CreateReservedResource(const FD3D12ResourceDesc& InDesc,
 	TRefCountPtr<ID3D12Resource> pResource;
 
 	FD3D12ResourceDesc LocalDesc = InDesc;
-	LocalDesc.Layout = D3D12_TEXTURE_LAYOUT_64KB_UNDEFINED_SWIZZLE;
-	LocalDesc.bReservedResource = true;
+
+	checkf(LocalDesc.bReservedResource,
+		TEXT("FD3D12ResourceDesc is expected to be initialized as a reserved resource. See FD3D12DynamicRHI::GetResourceDesc()."));
+
+	checkf(LocalDesc.Layout == D3D12_TEXTURE_LAYOUT_64KB_UNDEFINED_SWIZZLE,
+		TEXT("Reserved textures are expected to have layout %d (64KB_UNDEFINED_SWIZZLE), but have %d. See FD3D12DynamicRHI::GetResourceDesc()."),
+		uint32(D3D12_TEXTURE_LAYOUT_64KB_UNDEFINED_SWIZZLE), uint32(LocalDesc.Layout));
+
+	checkf(LocalDesc.Alignment == 0 || LocalDesc.Alignment == 65536,
+		TEXT("Reserved resources must use either 64KB alignment or 0 (unspecified/default), but have %d. See FD3D12DynamicRHI::GetResourceDesc()."),
+			uint32(LocalDesc.Alignment));
 
 #if D3D12_RHI_RAYTRACING
 	if (InDefaultState == D3D12_RESOURCE_STATE_RAYTRACING_ACCELERATION_STRUCTURE)
diff --git a/Engine/Source/Runtime/D3D12RHI/Private/D3D12Texture.cpp b/Engine/Source/Runtime/D3D12RHI/Private/D3D12Texture.cpp
--- a/Engine/Source/Runtime/D3D12RHI/Private/D3D12Texture.cpp
+++ b/Engine/Source/Runtime/D3D12RHI/Private/D3D12Texture.cpp
@@ -428,6 +428,11 @@ void FD3D12TextureStats::D3D12TextureDeleted(FD3D12Texture& Texture)
 
 bool FD3D12Texture::CanBe4KAligned(const FD3D12ResourceDesc& Desc, EPixelFormat UEFormat)
 {
+	if (Desc.bReservedResource)
+	{
+		return false;
+	}
+
 	// Exclude video related formats
 	if (UEFormat == PF_NV12 ||
 		UEFormat == PF_P010)
@@ -575,6 +580,16 @@ FD3D12ResourceDesc FD3D12DynamicRHI::GetResourceDesc(const FRHITextureDesc& Text
 		ResourceDesc.Flags |= D3D12_RESOURCE_FLAG_ALLOW_UNORDERED_ACCESS;
 	}
 
+	// Only 2D textures without mips are implemented/supported, to avoid the complexity associated with packed mips
+	if (GRHIGlobals.ReservedResources.Supported
+		&& EnumHasAllFlags(TextureDesc.Flags, TexCreate_ReservedResource)
+		&& TextureDesc.NumMips == 1
+		&& (TextureDesc.Dimension == ETextureDimension::Texture2D || TextureDesc.Dimension == ETextureDimension::Texture2DArray))
+	{
+		ResourceDesc.bReservedResource = true;
+		ResourceDesc.Layout = D3D12_TEXTURE_LAYOUT_64KB_UNDEFINED_SWIZZLE;
+	}
+
 	ResourceDesc.PixelFormat = TextureDesc.Format;
 
 #if D3D12RHI_NEEDS_VENDOR_EXTENSIONS
@@ -705,15 +720,6 @@ bool FD3D12DynamicRHI::RHIGetTextureMemoryVisualizeData(FColor* /*TextureData*/,
 	return false;
 }
 
-static bool ShouldCreateReservedTexture(const FD3D12ResourceDesc& TextureDesc, ETextureCreateFlags Flags)
-{
-	// Only 2D textures without mips are implemented/supported, to avoid the complexity associated with packed mips
-	return GRHISupportsReservedResources
-		&& EnumHasAllFlags(Flags, TexCreate_ReservedResource)
-		&& TextureDesc.MipLevels == 1
-		&& TextureDesc.Dimension == D3D12_RESOURCE_DIMENSION_TEXTURE2D;
-}
-
 /**
  * Creates a 2D texture optionally guarded by a structured exception handler.
  */
@@ -752,7 +758,7 @@ void SafeCreateTexture2D(FD3D12Device* pDevice,
 
 		case D3D12_HEAP_TYPE_DEFAULT:
 		{
-			if (ShouldCreateReservedTexture(TextureDesc, Flags))
+			if (TextureDesc.bReservedResource)
 			{
 				FD3D12Resource* Resource = nullptr;
 				VERIFYD3D12CREATETEXTURERESULT(
diff --git a/Engine/Source/Runtime/D3D12RHI/Private/Windows/WindowsD3D12Device.cpp b/Engine/Source/Runtime/D3D12RHI/Private/Windows/WindowsD3D12Device.cpp
--- a/Engine/Source/Runtime/D3D12RHI/Private/Windows/WindowsD3D12Device.cpp
+++ b/Engine/Source/Runtime/D3D12RHI/Private/Windows/WindowsD3D12Device.cpp
@@ -1666,7 +1666,7 @@ void FD3D12DynamicRHI::Init()
 		if (SUCCEEDED(GetAdapter().GetD3DDevice()->CheckFeatureSupport(D3D12_FEATURE_D3D12_OPTIONS, &Options, sizeof(Options))))
 		{
 			// Tier 2 is guaranteed for all adapters with feature level 12_0.
-			GRHISupportsReservedResources = Options.TiledResourcesTier >= D3D12_TILED_RESOURCES_TIER_2;
+			GRHIGlobals.ReservedResources.Supported = Options.TiledResourcesTier >= D3D12_TILED_RESOURCES_TIER_2;
 		}
 	}
 
diff --git a/Engine/Source/Runtime/RHI/Private/RHI.cpp b/Engine/Source/Runtime/RHI/Private/RHI.cpp
--- a/Engine/Source/Runtime/RHI/Private/RHI.cpp
+++ b/Engine/Source/Runtime/RHI/Private/RHI.cpp
@@ -1732,6 +1732,35 @@ bool FRHITextureDesc::Validate(const FRHITextureCreateInfo& Desc, const TCHAR* N
 		ValidateResourceDesc(Desc.NumMips <= GMaxTextureMipCount, TEXT("Texture %s's NumMips=%d is too large."), Name, Desc.NumMips);
 	}
 
+	// Validate reserved resource restrictions
+	if (EnumHasAnyFlags(Desc.Flags, TexCreate_ReservedResource))
+	{
+		ValidateResourceDesc(GRHIGlobals.ReservedResources.Supported,
+			TEXT("Reserved Texture %s's can't be created because current RHI does not support reserved resources."),
+			Name);
+
+		ValidateResourceDesc(
+			Desc.Dimension == ETextureDimension::Texture2D ||
+			Desc.Dimension == ETextureDimension::Texture2DArray,
+			TEXT("Reserved Texture %s's Desc.Dimension=%s is invalid. Expected Texture2D or Texture2DArray."),
+			Name, GetTextureDimensionString(Desc.Dimension));
+
+		ValidateResourceDesc(Desc.NumMips == 1,
+			TEXT("Reserved Texture %s's NumMips=%d is invalid. Expected only 1 mip level."),
+			Name, Desc.NumMips);
+
+		if (Desc.Dimension == ETextureDimension::Texture2DArray)
+		{
+			ValidateResourceDesc(Desc.Extent.X >= GRHIGlobals.ReservedResources.TextureArrayMinimumMipDimension,
+				TEXT("Reserved Texture array %s's Desc.Extent.X=%d is invalid. It is required to be be no less than %d."),
+				Name, Desc.Extent.X, GRHIGlobals.ReservedResources.TextureArrayMinimumMipDimension);
+
+			ValidateResourceDesc(Desc.Extent.Y >= GRHIGlobals.ReservedResources.TextureArrayMinimumMipDimension,
+				TEXT("Reserved Texture array %s's Desc.Extent.Y=%d is invalid. It is required to be be no less than %d."),
+				Name, Desc.Extent.Y, GRHIGlobals.ReservedResources.TextureArrayMinimumMipDimension);
+		}
+	}
+
 	return true;
 }
 
diff --git a/Engine/Source/Runtime/RHI/Public/RHIGlobals.h b/Engine/Source/Runtime/RHI/Public/RHIGlobals.h
--- a/Engine/Source/Runtime/RHI/Public/RHIGlobals.h
+++ b/Engine/Source/Runtime/RHI/Public/RHIGlobals.h
@@ -595,8 +595,22 @@ struct FRHIGlobals
 	/** Whether dynamic (bindless) resources are supported */
 	ERHIBindlessSupport BindlessSupport = ERHIBindlessSupport::Unsupported;
 
-	/** True if the RHI supports reserved (AKA tiled, virtual or sparse) resources and operations related to them*/
-	bool SupportsReservedResources = false;
+	struct FReservedResources
+	{
+		/**
+		* True if the RHI supports reserved (AKA tiled, virtual or sparse) resources and operations related to them.
+		*/
+		bool Supported = true;
+
+		/**
+		* Smallest mip dimension of reserved texture arrays must be greater or equal to this value.
+		* Tiled/reserved resources with both more than one array slice and any mipmap that
+		* has a dimension less than a tile in extent are not supported by some hardware.
+		* This is a conservative value chosen by the engine, independent of the texture format for simplicity.
+		*/
+		int32 TextureArrayMinimumMipDimension = 256;
+
+	} ReservedResources;
 
 	/** Table for finding out which shader platform corresponds to a given feature level for this RHI. */
 	EShaderPlatform ShaderPlatformForFeatureLevel[ERHIFeatureLevel::Num];
@@ -767,7 +781,7 @@ extern RHI_API FRHIGlobals GRHIGlobals;
 #define GRHIDefaultMSAASampleOffsets                           GRHIGlobals.DefaultMSAASampleOffsets
 #define GRHISupportsAsyncPipelinePrecompile                    GRHIGlobals.SupportsAsyncPipelinePrecompile
 #define GRHIBindlessSupport                                    GRHIGlobals.BindlessSupport
-#define GRHISupportsReservedResources                          GRHIGlobals.SupportsReservedResources
+#define GRHISupportsReservedResources                          GRHIGlobals.ReservedResources.Supported UE_DEPRECATED_MACRO(5.4, "GRHISupportsReservedResources has been deprecated - please use GRHIGlobals.ReservedResources.Supported instead.")
 #define GShaderPlatformForFeatureLevel                         GRHIGlobals.ShaderPlatformForFeatureLevel
 #define GRHIIsDebugLayerEnabled                                GRHIGlobals.IsDebugLayerEnabled
 
diff --git a/Engine/Source/Runtime/Renderer/Private/VirtualShadowMaps/VirtualShadowMapCacheManager.cpp b/Engine/Source/Runtime/Renderer/Private/VirtualShadowMaps/VirtualShadowMapCacheManager.cpp
--- a/Engine/Source/Runtime/Renderer/Private/VirtualShadowMaps/VirtualShadowMapCacheManager.cpp
+++ b/Engine/Source/Runtime/Renderer/Private/VirtualShadowMaps/VirtualShadowMapCacheManager.cpp
@@ -679,7 +679,7 @@ void FVirtualShadowMapArrayCacheManager::SetPhysicalPoolSize(FRDGBuilder& GraphB
 {
 	// Using ReservedResource|ImmediateCommit flags hint to the RHI that the resource can be allocated using N small physical memory allocations,
 	// instead of a single large contighous allocation. This helps Windows video memory manager page allocations in and out of local memory more efficiently.
-	ETextureCreateFlags RequestedCreateFlags = (CVarVSMReservedResource.GetValueOnRenderThread() && GRHISupportsReservedResources)
+	ETextureCreateFlags RequestedCreateFlags = (CVarVSMReservedResource.GetValueOnRenderThread() && GRHIGlobals.ReservedResources.Supported)
 		? (TexCreate_ReservedResource | TexCreate_ImmediateCommit)
 		: TexCreate_None;
 
