Implement basic support for D3D12 Reserved Buffers

- Add BUF_ReservedResource and BUF_ImmediateCommit, similar to corresponding TexCreate flags
- Add GRHIGlobals.ReservedResources.TileSizeInBytes to indicate minimum unit of physical memory that can be mapped to a reserved resource
- Add few basic RHI unit tests to exercise new code paths

diff --git a/Engine/Plugins/Tests/RHITests/Source/RHITests/Private/RHIReservedResourceTests.cpp b/Engine/Plugins/Tests/RHITests/Source/RHITests/Private/RHIReservedResourceTests.cpp
--- a/Engine/Plugins/Tests/RHITests/Source/RHITests/Private/RHIReservedResourceTests.cpp
+++ b/Engine/Plugins/Tests/RHITests/Source/RHITests/Private/RHIReservedResourceTests.cpp
@@ -84,3 +84,36 @@ bool FRHIReservedResourceTests::Test_ReservedResource_CreateTexture(FRHICommandL
 
 	return true;
 }
+
+bool FRHIReservedResourceTests::Test_ReservedResource_CreateBuffer(FRHICommandListImmediate& RHICmdList)
+{
+	if (!GRHIGlobals.ReservedResources.Supported)
+	{
+		return true;
+	}
+
+	const EBufferUsageFlags CommonFlags = BUF_ReservedResource | BUF_ImmediateCommit;
+
+	{
+		FRHIResourceCreateInfo CreateInfo(TEXT("TestSmallReservedVertexBuffer"));
+		FBufferRHIRef Buffer = RHICmdList.CreateBuffer(32768, CommonFlags | BUF_VertexBuffer, 4, ERHIAccess::CopyDest, CreateInfo);
+	}
+
+	{
+		FRHIResourceCreateInfo CreateInfo(TEXT("TestReservedVertexBuffer"));
+		FBufferRHIRef Buffer = RHICmdList.CreateBuffer(32 * 1024 * 1024, CommonFlags | BUF_VertexBuffer, 4, ERHIAccess::CopyDest, CreateInfo);
+	}
+
+	{
+		FRHIResourceCreateInfo CreateInfo(TEXT("TestReservedAccelerationStructureBuffer"));
+		FBufferRHIRef Buffer = RHICmdList.CreateBuffer(32 * 1024 * 1024, CommonFlags | BUF_AccelerationStructure, 4, ERHIAccess::BVHWrite, CreateInfo);
+	}
+
+	{
+		FRHIResourceCreateInfo CreateInfo(TEXT("TestReservedRayTracingScratchBuffer"));
+		FBufferRHIRef Buffer = RHICmdList.CreateBuffer(32 * 1024 * 1024, CommonFlags | BUF_RayTracingScratch, 4, ERHIAccess::UAVCompute, CreateInfo);
+	}
+
+	return true;
+}
+
diff --git a/Engine/Plugins/Tests/RHITests/Source/RHITests/Private/RHITests.spec.cpp b/Engine/Plugins/Tests/RHITests/Source/RHITests/Private/RHITests.spec.cpp
--- a/Engine/Plugins/Tests/RHITests/Source/RHITests/Private/RHITests.spec.cpp
+++ b/Engine/Plugins/Tests/RHITests/Source/RHITests/Private/RHITests.spec.cpp
@@ -8,6 +8,7 @@
 #include "RHITextureTests.h"
 #include "RHIDrawTests.h"
 #include "RHIReadbackTests.h"
+#include "RHIReservedResourceTests.h"
 
 BEGIN_DEFINE_SPEC(FAutomationRHITest, "Rendering.RHI", EAutomationTestFlags::EngineFilter | EAutomationTestFlags::ApplicationContextMask | EAutomationTestFlags::NonNullRHI)
 END_DEFINE_SPEC(FAutomationRHITest)
@@ -117,4 +118,19 @@ void FAutomationRHITest::Define()
 			TestEqual("RHI MultiDrawIndirect", bResult, 1);
 		});
 	});
+
+	Describe("Test RHI Reserved Resource", [this]()
+	{
+		It("Create Reserved Texture", [this]()
+		{
+			bool bResult = RunOnRenderThreadSynchronous(FRHIReservedResourceTests::Test_ReservedResource_CreateTexture);
+			TestEqual("Create Reserved Texture failed", bResult, 1);
+		});
+
+		It("Create Reserved Buffer", [this]()
+		{
+			bool bResult = RunOnRenderThreadSynchronous(FRHIReservedResourceTests::Test_ReservedResource_CreateBuffer);
+			TestEqual("Create Reserved Buffer failed", bResult, 1);
+		});
+	});
 }
diff --git a/Engine/Plugins/Tests/RHITests/Source/RHITests/Private/RHITestsModule.cpp b/Engine/Plugins/Tests/RHITests/Source/RHITests/Private/RHITestsModule.cpp
--- a/Engine/Plugins/Tests/RHITests/Source/RHITests/Private/RHITestsModule.cpp
+++ b/Engine/Plugins/Tests/RHITests/Source/RHITests/Private/RHITestsModule.cpp
@@ -21,6 +21,7 @@ static bool RunTests_RenderThread(FRHICommandListImmediate& RHICmdList)
 	// Reserved resources
 	// ------------------------------------------------
 	{
+		RUN_TEST(FRHIReservedResourceTests::Test_ReservedResource_CreateBuffer(RHICmdList));
 		RUN_TEST(FRHIReservedResourceTests::Test_ReservedResource_CreateTexture(RHICmdList));
 	}
 
diff --git a/Engine/Plugins/Tests/RHITests/Source/RHITests/Public/RHIReservedResourceTests.h b/Engine/Plugins/Tests/RHITests/Source/RHITests/Public/RHIReservedResourceTests.h
--- a/Engine/Plugins/Tests/RHITests/Source/RHITests/Public/RHIReservedResourceTests.h
+++ b/Engine/Plugins/Tests/RHITests/Source/RHITests/Public/RHIReservedResourceTests.h
@@ -9,5 +9,6 @@ class FRHIReservedResourceTests
 public:
 
 	static bool Test_ReservedResource_CreateTexture(FRHICommandListImmediate& RHICmdList);
+	static bool Test_ReservedResource_CreateBuffer(FRHICommandListImmediate& RHICmdList);
 };
 
diff --git a/Engine/Source/Runtime/D3D12RHI/Private/D3D12Allocation.cpp b/Engine/Source/Runtime/D3D12RHI/Private/D3D12Allocation.cpp
--- a/Engine/Source/Runtime/D3D12RHI/Private/D3D12Allocation.cpp
+++ b/Engine/Source/Runtime/D3D12RHI/Private/D3D12Allocation.cpp
@@ -1505,6 +1505,8 @@ D3D12_RESOURCE_STATES FD3D12DefaultBufferAllocator::GetDefaultInitialResourceSta
 // Grab a buffer from the available buffers or create a new buffer if none are available
 void FD3D12DefaultBufferAllocator::AllocDefaultResource(D3D12_HEAP_TYPE InHeapType, const D3D12_RESOURCE_DESC& InResourceDesc, EBufferUsageFlags InBufferUsage, ED3D12ResourceStateMode InResourceStateMode, D3D12_RESOURCE_STATES InCreateState, FD3D12ResourceLocation& ResourceLocation, uint32 Alignment, const TCHAR* Name)
 {
+	FD3D12Adapter* Adapter = GetParentDevice()->GetParentAdapter();
+
 	// Force indirect args to stand alone allocations instead of pooled
 	if (!GD3D12AllowPoolAllocateIndirectArgBuffers && EnumHasAnyFlags(InBufferUsage, BUF_DrawIndirect))
 	{
@@ -1514,9 +1516,32 @@ void FD3D12DefaultBufferAllocator::AllocDefaultResource(D3D12_HEAP_TYPE InHeapTy
 		const D3D12_HEAP_PROPERTIES HeapProps = CD3DX12_HEAP_PROPERTIES(InHeapType, GetGPUMask().GetNative(), GetVisibilityMask().GetNative());
 		D3D12_RESOURCE_DESC Desc = InResourceDesc;
 		Desc.Alignment = 0;
-		VERIFYD3D12RESULT(GetParentDevice()->GetParentAdapter()->CreateCommittedResource(Desc, GetGPUMask(), HeapProps, InCreateState, InResourceStateMode, InCreateState, nullptr, &NewResource, Name, false));
+		VERIFYD3D12RESULT(Adapter->CreateCommittedResource(Desc, GetGPUMask(), HeapProps, InCreateState, InResourceStateMode, InCreateState, nullptr, &NewResource, Name, false));
+
+		ResourceLocation.AsStandAlone(NewResource, InResourceDesc.Width);
+
+		return;
+	}
+
+	if (EnumHasAnyFlags(InBufferUsage, BUF_ReservedResource))
+	{
+		ResourceLocation.Clear();
+
+		FD3D12Resource* NewResource = nullptr;
+		checkf(Alignment % GRHIGlobals.ReservedResources.TileSizeInBytes == 0,
+			TEXT("Reserved buffer alignment is expected to be a multiple of the reserved resource tile size"));
+		FD3D12ResourceDesc Desc = InResourceDesc;
+		Desc.Alignment = Alignment;
+		Desc.bReservedResource = true;
+		VERIFYD3D12RESULT(Adapter->CreateReservedResource(Desc, GetGPUMask(), InCreateState, InResourceStateMode, InCreateState, nullptr, &NewResource, Name, false));
 
 		ResourceLocation.AsStandAlone(NewResource, InResourceDesc.Width);
+
+		if (EnumHasAnyFlags(InBufferUsage, BUF_ImmediateCommit))
+		{
+			NewResource->CommitReservedResource();
+		}
+
 		return;
 	}
 
diff --git a/Engine/Source/Runtime/D3D12RHI/Private/D3D12Buffer.cpp b/Engine/Source/Runtime/D3D12RHI/Private/D3D12Buffer.cpp
--- a/Engine/Source/Runtime/D3D12RHI/Private/D3D12Buffer.cpp
+++ b/Engine/Source/Runtime/D3D12RHI/Private/D3D12Buffer.cpp
@@ -497,7 +497,7 @@ void FD3D12Buffer::GetResourceDescAndAlignment(uint64 InSize, uint32 InStride, E
 		ResourceDesc.Flags |= D3D12_RESOURCE_FLAG_ALLOW_UNORDERED_ACCESS;
 	}
 
-	if (!EnumHasAnyFlags(InUsage, BUF_ShaderResource))
+	if (!EnumHasAnyFlags(InUsage, BUF_ShaderResource | BUF_AccelerationStructure))
 	{
 		ResourceDesc.Flags |= D3D12_RESOURCE_FLAG_DENY_SHADER_RESOURCE;
 	}
@@ -512,8 +512,19 @@ void FD3D12Buffer::GetResourceDescAndAlignment(uint64 InSize, uint32 InStride, E
 		ResourceDesc.Flags |= D3D12_RESOURCE_FLAG_ALLOW_SIMULTANEOUS_ACCESS;
 	}
 
-	// Structured buffers, non-ByteAddress buffers, need to be aligned to their stride to ensure that they can be addressed correctly with element based offsets.
-	Alignment = (InStride > 0) && (EnumHasAnyFlags(InUsage, BUF_StructuredBuffer) || !EnumHasAnyFlags(InUsage, BUF_ByteAddressBuffer | BUF_DrawIndirect)) ? InStride : 4;
+	if (EnumHasAnyFlags(InUsage, BUF_ReservedResource))
+	{
+		checkf(InStride <= GRHIGlobals.ReservedResources.TileSizeInBytes,
+			TEXT("Reserved buffer stride %d must not be greater than the reserved resource tile size %d"), 
+			InStride, GRHIGlobals.ReservedResources.TileSizeInBytes);
+
+		Alignment = GRHIGlobals.ReservedResources.TileSizeInBytes;
+	}
+	else
+	{
+		// Structured buffers, non-ByteAddress buffers, need to be aligned to their stride to ensure that they can be addressed correctly with element based offsets.
+		Alignment = (InStride > 0) && (EnumHasAnyFlags(InUsage, BUF_StructuredBuffer) || !EnumHasAnyFlags(InUsage, BUF_ByteAddressBuffer | BUF_DrawIndirect)) ? InStride : 4;
+	}
 }
 
 FBufferRHIRef FD3D12DynamicRHI::RHICreateBuffer(FRHICommandListBase& RHICmdList, FRHIBufferDesc const& Desc, ERHIAccess ResourceState, FRHIResourceCreateInfo& CreateInfo)
@@ -547,6 +558,13 @@ FD3D12Buffer* FD3D12DynamicRHI::CreateD3D12Buffer(class FRHICommandListBase* RHI
 		: ED3D12ResourceStateMode::Default;
 
 	const bool bIsDynamic = EnumHasAnyFlags(BufferDesc.Usage, BUF_AnyDynamic);
+
+	if (EnumHasAnyFlags(BufferDesc.Usage, BUF_ReservedResource))
+	{
+		checkf(!bIsDynamic, TEXT("Reserved resources may not be dynamic"));
+		checkf(!ResourceAllocator, TEXT("Reserved resources may not use a custom resource allocator"));
+	}
+
 	D3D12_HEAP_TYPE HeapType = bIsDynamic ? D3D12_HEAP_TYPE_UPLOAD : D3D12_HEAP_TYPE_DEFAULT;
 	const FD3D12Resource::FD3D12ResourceTypeHelper Type(Desc, HeapType);
 
diff --git a/Engine/Source/Runtime/D3D12RHI/Private/D3D12Resources.cpp b/Engine/Source/Runtime/D3D12RHI/Private/D3D12Resources.cpp
--- a/Engine/Source/Runtime/D3D12RHI/Private/D3D12Resources.cpp
+++ b/Engine/Source/Runtime/D3D12RHI/Private/D3D12Resources.cpp
@@ -156,11 +156,19 @@ FD3D12Resource::~FD3D12Resource()
 
 void FD3D12Resource::CommitReservedResource()
 {
+	TRACE_CPUPROFILER_EVENT_SCOPE(CommitReservedResource);
+
 	check(Desc.bReservedResource);
 	check(ReservedResourceData.IsValid());
 	checkf(ReservedResourceData->BackingHeaps.IsEmpty(), TEXT("Reserved resource is already committed"));
-	checkf(Desc.Dimension == D3D12_RESOURCE_DIMENSION_TEXTURE2D, TEXT("CommitReservedResource is currently only implemented for 2D textures"));
-	checkf(Desc.MipLevels == 1, TEXT("CommitReservedResource is currently only implemented for textures without mips"));
+	checkf(Desc.Dimension == D3D12_RESOURCE_DIMENSION_BUFFER
+		|| Desc.Dimension == D3D12_RESOURCE_DIMENSION_TEXTURE2D, 
+		TEXT("CommitReservedResource is currently only implemented for 2D textures and buffers"));
+
+	if (Desc.Dimension != D3D12_RESOURCE_DIMENSION_BUFFER)
+	{
+		checkf(Desc.MipLevels == 1, TEXT("CommitReservedResource is currently only implemented for textures without mips"));
+	}
 
 	uint32 D3DResourceNumTiles = 0;
 	D3D12_PACKED_MIP_INFO PackedMipDesc = {};
@@ -176,7 +184,9 @@ void FD3D12Resource::CommitReservedResource()
 
 	D3DDevice->GetResourceTiling(GetResource(), &D3DResourceNumTiles, &PackedMipDesc, &TileShape, &NumSubresourceTilings, FirstSubresource, &SubresourceTiling);
 
-	const uint64 TileSizeInBytes = 65536; // reserved resource tiles are always 64KB
+	static constexpr uint64 TileSizeInBytes = GRHIGlobals.ReservedResources.TileSizeInBytes;
+	static_assert(TileSizeInBytes == 65536, "Reserved resource tiles are expected to always be 64KB");
+
 	const uint64 TotalSize = D3DResourceNumTiles * TileSizeInBytes;
 	const uint64 MaxHeapSize = uint64(CVarD3D12ReservedResourceHeapSizeMB.GetValueOnAnyThread()) * 1024 * 1024;
 	const uint64 NumHeaps = FMath::DivideAndRoundUp(TotalSize, MaxHeapSize);
@@ -208,7 +218,13 @@ void FD3D12Resource::CommitReservedResource()
 	uint32 NumStandardTilesPerSubresource = 0;
 	uint32 NumTotalTiles = 0;
 
-	if (PackedMipDesc.NumStandardMips != 0)
+	if (Desc.Dimension == D3D12_RESOURCE_DIMENSION_BUFFER)
+	{
+		NumTotalTiles = NumStandardTilesPerSubresource = D3DResourceNumTiles;
+		checkf(D3DResourceNumTiles == SubresourceTiling.WidthInTiles,
+			TEXT("Reserved buffers are expected to have trivial tiling configuration: single 1D subresource that contains all tiles."));
+	}
+	else if (PackedMipDesc.NumStandardMips != 0)
 	{
 		checkf(SubresourceTiling.DepthInTiles == 1, TEXT("3D reserved textures are not supported/implemented"));
 		NumStandardTilesPerSubresource = SubresourceTiling.WidthInTiles * SubresourceTiling.HeightInTiles;
@@ -226,6 +242,12 @@ void FD3D12Resource::CommitReservedResource()
 
 	uint32 NumMappedTiles = 0;
 
+	const D3D12_HEAP_FLAGS HeapFlags = Desc.Dimension == D3D12_RESOURCE_DIMENSION_BUFFER
+		? D3D12_HEAP_FLAG_ALLOW_ONLY_BUFFERS
+		: D3D12_HEAP_FLAG_ALLOW_ONLY_NON_RT_DS_TEXTURES;
+
+	static_assert((D3D12_HEAP_FLAG_DENY_BUFFERS | D3D12_HEAP_FLAG_DENY_RT_DS_TEXTURES) == D3D12_HEAP_FLAG_ALLOW_ONLY_NON_RT_DS_TEXTURES);
+
 	while (NumMappedTiles < NumTotalTiles)
 	{
 		const uint32 NumRemainingTiles = NumTotalTiles - NumMappedTiles;
@@ -237,7 +259,7 @@ void FD3D12Resource::CommitReservedResource()
 		const uint32 ThisHeapSize = RegionSize.NumTiles * TileSizeInBytes;
 		D3D12_HEAP_DESC NewHeapDesc = {};
 		NewHeapDesc.Alignment = D3D12_DEFAULT_RESOURCE_PLACEMENT_ALIGNMENT;
-		NewHeapDesc.Flags = D3D12_HEAP_FLAG_DENY_BUFFERS | D3D12_HEAP_FLAG_DENY_RT_DS_TEXTURES;
+		NewHeapDesc.Flags = HeapFlags;
 		NewHeapDesc.SizeInBytes = ThisHeapSize;
 		NewHeapDesc.Properties = BackingHeapProps;
 		TRefCountPtr<ID3D12Heap> NewHeap;
@@ -580,9 +602,14 @@ HRESULT FD3D12Adapter::CreateReservedResource(const FD3D12ResourceDesc& InDesc,
 	checkf(LocalDesc.bReservedResource,
 		TEXT("FD3D12ResourceDesc is expected to be initialized as a reserved resource. See FD3D12DynamicRHI::GetResourceDesc()."));
 
-	checkf(LocalDesc.Layout == D3D12_TEXTURE_LAYOUT_64KB_UNDEFINED_SWIZZLE,
-		TEXT("Reserved textures are expected to have layout %d (64KB_UNDEFINED_SWIZZLE), but have %d. See FD3D12DynamicRHI::GetResourceDesc()."),
-		uint32(D3D12_TEXTURE_LAYOUT_64KB_UNDEFINED_SWIZZLE), uint32(LocalDesc.Layout));
+	if (LocalDesc.Dimension == D3D12_RESOURCE_DIMENSION_TEXTURE1D
+		|| LocalDesc.Dimension == D3D12_RESOURCE_DIMENSION_TEXTURE2D
+		|| LocalDesc.Dimension == D3D12_RESOURCE_DIMENSION_TEXTURE3D)
+	{
+		checkf(LocalDesc.Layout == D3D12_TEXTURE_LAYOUT_64KB_UNDEFINED_SWIZZLE,
+			TEXT("Reserved textures are expected to have layout %d (64KB_UNDEFINED_SWIZZLE), but have %d. See FD3D12DynamicRHI::GetResourceDesc()."),
+			uint32(D3D12_TEXTURE_LAYOUT_64KB_UNDEFINED_SWIZZLE), uint32(LocalDesc.Layout));
+	}
 
 	checkf(LocalDesc.Alignment == 0 || LocalDesc.Alignment == 65536,
 		TEXT("Reserved resources must use either 64KB alignment or 0 (unspecified/default), but have %d. See FD3D12DynamicRHI::GetResourceDesc()."),
diff --git a/Engine/Source/Runtime/RHI/Public/RHIDefinitions.h b/Engine/Source/Runtime/RHI/Public/RHIDefinitions.h
--- a/Engine/Source/Runtime/RHI/Public/RHIDefinitions.h
+++ b/Engine/Source/Runtime/RHI/Public/RHIDefinitions.h
@@ -783,7 +783,7 @@ enum class EBufferUsageFlags : uint32
 	/**
 	 * Buffer contains opaque ray tracing acceleration structure data.
 	 * Resources with this flag can't be bound directly to any shader stage and only can be used with ray tracing APIs.
-	 * This flag is mutually exclusive with all other buffer flags except BUF_Static.
+	 * This flag is mutually exclusive with all other buffer flags except Static, ReservedResource and ImmediateCommit.
 	*/
 	AccelerationStructure   = 1 << 13,
 
@@ -812,6 +812,15 @@ enum class EBufferUsageFlags : uint32
 	/** Buffer can be used as uniform buffer on platforms that do support uniform buffer objects. */
 	UniformBuffer = 1 << 21,
 
+	/**
+	* EXPERIMENTAL: Allow the buffer to be created as a reserved (AKA tiled/sparse/virtual) resource internally, without physical memory backing.
+	* May not be used with Dynamic and other buffer flags that prevent the resource from being allocated in local GPU memory.
+	*/
+	ReservedResource = 1 << 22,
+
+	/** EXPERIMENTAL: Used with ReservedResource flag to immediately allocate and commit memory on creation. May use N small physical memory allocations instead of a single large one. */
+	ImmediateCommit = 1 << 23,
+
 	// Helper bit-masks
 	AnyDynamic = (Dynamic | Volatile),
 };
@@ -841,6 +850,8 @@ ENUM_CLASS_FLAGS(EBufferUsageFlags);
 #define BUF_MultiGPUGraphIgnore    EBufferUsageFlags::MultiGPUGraphIgnore
 #define BUF_NullResource           EBufferUsageFlags::NullResource
 #define BUF_UniformBuffer          EBufferUsageFlags::UniformBuffer
+#define BUF_ReservedResource       EBufferUsageFlags::ReservedResource
+#define BUF_ImmediateCommit        EBufferUsageFlags::ImmediateCommit
 
 enum class EGpuVendorId : uint32
 {
@@ -1004,7 +1015,10 @@ enum class ETextureCreateFlags : uint64
 	External                		  = 1ull << 34,
 	/** Don't automatically transfer across GPUs in multi-GPU scenarios.  For example, if you are transferring it yourself manually. */
 	MultiGPUGraphIgnore				  = 1ull << 35,
-	/** EXPERIMENTAL: Allow the texture to be created as a reserved (AKA tiled/sparse/virtual) resource internally, without physical memory backing. */
+	/**
+	* EXPERIMENTAL: Allow the texture to be created as a reserved (AKA tiled/sparse/virtual) resource internally, without physical memory backing. 
+	* May not be used with Dynamic and other buffer flags that prevent the resource from being allocated in local GPU memory.
+	*/
 	ReservedResource                  = 1ull << 37,
 	/** EXPERIMENTAL: Used with ReservedResource flag to immediately allocate and commit memory on creation. May use N small physical memory allocations instead of a single large one. */
 	ImmediateCommit                   = 1ull << 38,
diff --git a/Engine/Source/Runtime/RHI/Public/RHIGlobals.h b/Engine/Source/Runtime/RHI/Public/RHIGlobals.h
--- a/Engine/Source/Runtime/RHI/Public/RHIGlobals.h
+++ b/Engine/Source/Runtime/RHI/Public/RHIGlobals.h
@@ -610,6 +610,13 @@ struct FRHIGlobals
 		*/
 		int32 TextureArrayMinimumMipDimension = 256;
 
+		/**
+		* Size that corresponds to a minimum unit of physical memory that may be mapped to a region of a reserved resource.
+		* High-level code should aim to allocate reserved resources such that their size is a multiple of this tile size.
+		* Guaranteed to be the same value on all platforms, regardless of the native virtual memory page size.
+		*/
+		static constexpr int32 TileSizeInBytes = 65536;
+
 	} ReservedResources;
 
 	/** Table for finding out which shader platform corresponds to a given feature level for this RHI. */
diff --git a/Engine/Source/Runtime/RHI/Public/RHIValidation.h b/Engine/Source/Runtime/RHI/Public/RHIValidation.h
--- a/Engine/Source/Runtime/RHI/Public/RHIValidation.h
+++ b/Engine/Source/Runtime/RHI/Public/RHIValidation.h
@@ -414,6 +414,13 @@ public:
 	*/
 	virtual FBufferRHIRef RHICreateBuffer(FRHICommandListBase& RHICmdList, FRHIBufferDesc const& Desc, ERHIAccess ResourceState, FRHIResourceCreateInfo& CreateInfo) override final
 	{
+		if (EnumHasAnyFlags(Desc.Usage, BUF_ReservedResource))
+		{
+			RHI_VALIDATION_CHECK(GRHIGlobals.ReservedResources.Supported, TEXT("Reserved buffers are not supported by current RHI"));
+			RHI_VALIDATION_CHECK(!EnumHasAnyFlags(Desc.Usage, BUF_AnyDynamic), TEXT("Reserved buffers must not be dynamic"));
+			RHI_VALIDATION_CHECK(Desc.Stride <= GRHIGlobals.ReservedResources.TileSizeInBytes, TEXT("Reserved buffer stride must not be greater than reserved resource tile size"));
+		}
+
 		if (CreateInfo.ResourceArray && RHICmdList.IsInsideRenderPass())
 		{
 			FString Msg = FString::Printf(TEXT("Creating buffers with initial data during a render pass is not supported, buffer name: \"%s\""), CreateInfo.DebugName);
