Use GetCopyableFootprints for sizing the staging buffer in FD3D12Texture::Lock, rather than manually doing the math.

The math we were doing was wrong in some corner cases (such as when trying to read a PF_DepthStencil texture, because for legacy reasons GPixelFormats says that the block size for that format is 5 bytes), and the engine crashed because D3D returned E_INVALIDARG when we tried to execute the copy operation. While reading back a depth-stencil surface like that is almost never a good idea, we shouldn't crash when somebody does it. Anyway, using GetCopyableFootprints gets rid of a bunch of error-prone copy-paste and works in all cases.

#rnx

diff --git a/Engine/Source/Runtime/D3D12RHI/Private/D3D12Buffer.cpp b/Engine/Source/Runtime/D3D12RHI/Private/D3D12Buffer.cpp
--- a/Engine/Source/Runtime/D3D12RHI/Private/D3D12Buffer.cpp
+++ b/Engine/Source/Runtime/D3D12RHI/Private/D3D12Buffer.cpp
@@ -708,8 +708,8 @@ void* FD3D12DynamicRHI::LockBuffer(FRHICommandListBase& RHICmdList, FD3D12Buffer
 		}
 	}
 
-	LockedData.LockedOffset = Offset;
-	LockedData.LockedPitch = Size;
+	LockedData.LockOffset = Offset;
+	LockedData.LockSize = Size;
 	LockedData.bLocked = true;
 	LockedData.bHasNeverBeenLocked = false;
 
@@ -752,7 +752,7 @@ void FD3D12DynamicRHI::UnlockBuffer(FRHICommandListBase& RHICmdList, FD3D12Buffe
 					if (CurrentBuffer.Get() == LastBuffer)
 					{
 						// Command associated with last buffer (will be only buffer if single GPU) receives ownership of locked data
-						ALLOC_COMMAND_CL(RHICmdList, FRHICommandUpdateBuffer)(&CurrentBuffer->ResourceLocation, LockedData.ResourceLocation, LockedData.LockedOffset, LockedData.LockedPitch);
+						ALLOC_COMMAND_CL(RHICmdList, FRHICommandUpdateBuffer)(&CurrentBuffer->ResourceLocation, LockedData.ResourceLocation, LockedData.LockOffset, LockedData.LockSize);
 					}
 					else
 					{
@@ -760,12 +760,12 @@ void FD3D12DynamicRHI::UnlockBuffer(FRHICommandListBase& RHICmdList, FD3D12Buffe
 						// last command handling clean up the locked data after it has been propagated to all GPUs.
 						FD3D12ResourceLocation NodeResourceLocation(LockedData.ResourceLocation.GetParentDevice());
 						FD3D12ResourceLocation::ReferenceNode(NodeResourceLocation.GetParentDevice(), NodeResourceLocation, LockedData.ResourceLocation);
-						ALLOC_COMMAND_CL(RHICmdList, FRHICommandUpdateBuffer)(&CurrentBuffer->ResourceLocation, NodeResourceLocation, LockedData.LockedOffset, LockedData.LockedPitch);
+						ALLOC_COMMAND_CL(RHICmdList, FRHICommandUpdateBuffer)(&CurrentBuffer->ResourceLocation, NodeResourceLocation, LockedData.LockOffset, LockedData.LockSize);
 					}
 				}
 				else
 				{
-					FD3D12CommandContextBase::Get(RHICmdList).UpdateBuffer(&CurrentBuffer->ResourceLocation, LockedData.LockedOffset, &LockedData.ResourceLocation, 0, LockedData.LockedPitch);
+					FD3D12CommandContextBase::Get(RHICmdList).UpdateBuffer(&CurrentBuffer->ResourceLocation, LockedData.LockOffset, &LockedData.ResourceLocation, 0, LockedData.LockSize);
 				}
 			}
 		}
diff --git a/Engine/Source/Runtime/D3D12RHI/Private/D3D12Resources.h b/Engine/Source/Runtime/D3D12RHI/Private/D3D12Resources.h
--- a/Engine/Source/Runtime/D3D12RHI/Private/D3D12Resources.h
+++ b/Engine/Source/Runtime/D3D12RHI/Private/D3D12Resources.h
@@ -727,11 +727,6 @@ struct FD3D12LockedResource : public FD3D12DeviceChild
 	FD3D12LockedResource(FD3D12Device* Device)
 		: FD3D12DeviceChild(Device)
 		, ResourceLocation(Device)
-		, LockedOffset(0)
-		, LockedPitch(0)
-		, bLocked(false)
-		, bLockedForReadOnly(false)
-		, bHasNeverBeenLocked(true)
 	{}
 
 	inline void Reset()
@@ -739,16 +734,18 @@ struct FD3D12LockedResource : public FD3D12DeviceChild
 		ResourceLocation.Clear();
 		bLocked = false;
 		bLockedForReadOnly = false;
-		LockedOffset = 0;
-		LockedPitch = 0;
+		LockOffset = 0;
+		LockSize = 0;
+		FMemory::Memzero(Footprint);
 	}
 
 	FD3D12ResourceLocation ResourceLocation;
-	uint32 LockedOffset;
-	uint32 LockedPitch;
-	uint32 bLocked : 1;
-	uint32 bLockedForReadOnly : 1;
-	uint32 bHasNeverBeenLocked : 1;
+	D3D12_SUBRESOURCE_FOOTPRINT Footprint = {};
+	uint32 LockOffset = 0;
+	uint32 LockSize = 0;
+	uint32 bLocked : 1 = false;
+	uint32 bLockedForReadOnly : 1 = false;
+	uint32 bHasNeverBeenLocked : 1 = true;
 };
 
 /** Resource which might needs to be notified about changes on dependent resources (Views, RTGeometryObject, Cached binding tables) */
diff --git a/Engine/Source/Runtime/D3D12RHI/Private/D3D12Texture.cpp b/Engine/Source/Runtime/D3D12RHI/Private/D3D12Texture.cpp
--- a/Engine/Source/Runtime/D3D12RHI/Private/D3D12Texture.cpp
+++ b/Engine/Source/Runtime/D3D12RHI/Private/D3D12Texture.cpp
@@ -1748,21 +1748,17 @@ void* FD3D12Texture::Lock(class FRHICommandListImmediate* RHICmdList, uint32 Mip
 	check(LockedMap.Find(Subresource) == nullptr);
 	FD3D12LockedResource* LockedResource = new FD3D12LockedResource(Device);
 
-	// Calculate the dimensions of the mip-map.
-	const uint32 BlockSizeX = GPixelFormats[this->GetFormat()].BlockSizeX;
-	const uint32 BlockSizeY = GPixelFormats[this->GetFormat()].BlockSizeY;
-	const uint32 BlockBytes = GPixelFormats[this->GetFormat()].BlockBytes;
-	const uint32 MipSizeX = FMath::Max(this->GetSizeX() >> MipIndex, BlockSizeX);
-	const uint32 MipSizeY = FMath::Max(this->GetSizeY() >> MipIndex, BlockSizeY);
-	const uint32 NumBlocksX = (MipSizeX + BlockSizeX - 1) / BlockSizeX;
-	const uint32 NumBlocksY = (MipSizeY + BlockSizeY - 1) / BlockSizeY;
-
-	const uint32 XBytesAligned = Align(NumBlocksX * BlockBytes, FD3D12_TEXTURE_DATA_PITCH_ALIGNMENT);
-	const uint32 MipBytesAligned = XBytesAligned * NumBlocksY;
+	const D3D12_RESOURCE_DESC& ResourceDesc = GetResource()->GetDesc();
 
+	UINT64 TotalBytes;
+	D3D12_PLACED_SUBRESOURCE_FOOTPRINT PlacedFootprint;
+	Device->GetDevice()->GetCopyableFootprints(&ResourceDesc, Subresource, 1, 0, &PlacedFootprint, nullptr, nullptr, &TotalBytes);
+
+	LockedResource->Footprint = PlacedFootprint.Footprint;
+	DestStride = LockedResource->Footprint.RowPitch;
 	if (OutLockedByteCount)
 	{
-		*OutLockedByteCount = MipBytesAligned;
+		*OutLockedByteCount = TotalBytes;
 	}
 
 	FD3D12CommandContext& Context = Device->GetDefaultCommandContext();
@@ -1783,21 +1779,14 @@ void* FD3D12Texture::Lock(class FRHICommandListImmediate* RHICmdList, uint32 Mip
 	{
 		// If we're writing to the texture, allocate a system memory buffer to receive the new contents.
 		// Use an upload heap to copy data to a default resource.
-		//const uint32 bufferSize = (uint32)GetRequiredIntermediateSize(this->GetResource()->GetResource(), Subresource, 1);
-		const uint32 bufferSize = Align(MipBytesAligned, D3D12_TEXTURE_DATA_PLACEMENT_ALIGNMENT);
 
-		void* pData = Device->GetDefaultFastAllocator().Allocate(bufferSize, D3D12_TEXTURE_DATA_PLACEMENT_ALIGNMENT, &LockedResource->ResourceLocation);
+		void* pData = Device->GetDefaultFastAllocator().Allocate(TotalBytes, D3D12_TEXTURE_DATA_PLACEMENT_ALIGNMENT, &LockedResource->ResourceLocation);
 		if (nullptr == pData)
 		{
 			check(false);
 			return nullptr;
 		}
 
-		DestStride = XBytesAligned;
-		LockedResource->LockedPitch = XBytesAligned;
-
-		check(LockedResource->LockedPitch % FD3D12_TEXTURE_DATA_PITCH_ALIGNMENT == 0);
-
 		Data = LockedResource->ResourceLocation.GetMappedBaseAddress();
 	}
 	else
@@ -1810,28 +1799,14 @@ void* FD3D12Texture::Lock(class FRHICommandListImmediate* RHICmdList, uint32 Mip
 		// If we're reading from the texture, we create a staging resource, copy the texture contents to it, and map it.
 
 		// Create the staging texture.
-		const D3D12_RESOURCE_DESC& StagingTextureDesc = GetResource()->GetDesc();
 		FD3D12Resource* StagingTexture = nullptr;
 
 		const FRHIGPUMask Node = Device->GetGPUMask();
-		VERIFYD3D12RESULT(Adapter->CreateBuffer(D3D12_HEAP_TYPE_READBACK, Node, Node, MipBytesAligned, &StagingTexture, nullptr));
-
-		LockedResource->ResourceLocation.AsStandAlone(StagingTexture, MipBytesAligned);
+		VERIFYD3D12RESULT(Adapter->CreateBuffer(D3D12_HEAP_TYPE_READBACK, Node, Node, TotalBytes, &StagingTexture, nullptr));
 
-		// Copy the mip-map data from the real resource into the staging resource
-		D3D12_SUBRESOURCE_FOOTPRINT destSubresource;
-		destSubresource.Depth = 1;
-		destSubresource.Height = MipSizeY;
-		destSubresource.Width = MipSizeX;
-		destSubresource.Format = StagingTextureDesc.Format;
-		destSubresource.RowPitch = XBytesAligned;
-		check(destSubresource.RowPitch % FD3D12_TEXTURE_DATA_PITCH_ALIGNMENT == 0);
+		LockedResource->ResourceLocation.AsStandAlone(StagingTexture, TotalBytes);
 
-		D3D12_PLACED_SUBRESOURCE_FOOTPRINT placedTexture2D = { 0 };
-		placedTexture2D.Offset = 0;
-		placedTexture2D.Footprint = destSubresource;
-
-		CD3DX12_TEXTURE_COPY_LOCATION DestCopyLocation(StagingTexture->GetResource(), placedTexture2D);
+		CD3DX12_TEXTURE_COPY_LOCATION DestCopyLocation(StagingTexture->GetResource(), PlacedFootprint);
 		CD3DX12_TEXTURE_COPY_LOCATION SourceCopyLocation(GetResource()->GetResource(), Subresource);
 
 		const auto& pfnCopyTextureRegion = [&]()
@@ -1865,11 +1840,6 @@ void* FD3D12Texture::Lock(class FRHICommandListImmediate* RHICmdList, uint32 Mip
 		// We need to execute the command list so we can read the data from the map below
 		Context.FlushCommands(ED3D12FlushFlags::WaitForCompletion);
 
-		LockedResource->LockedPitch = XBytesAligned;
-		DestStride = XBytesAligned;
-		check(LockedResource->LockedPitch % FD3D12_TEXTURE_DATA_PITCH_ALIGNMENT == 0);
-		check(DestStride == XBytesAligned);
-
 		Data = LockedResource->ResourceLocation.GetMappedBaseAddress();
 	}
 
@@ -2011,14 +1981,6 @@ void FD3D12Texture::UnlockInternal(class FRHICommandListImmediate* RHICmdList, F
 	// Calculate the subresource index corresponding to the specified mip-map.
 	const uint32 Subresource = CalcSubresource(MipIndex, ArrayIndex, this->GetNumMips());
 
-	// Calculate the dimensions of the mip-map.
-	const uint32 BlockSizeX = GPixelFormats[this->GetFormat()].BlockSizeX;
-	const uint32 BlockSizeY = GPixelFormats[this->GetFormat()].BlockSizeY;
-	const uint32 BlockBytes = GPixelFormats[this->GetFormat()].BlockBytes;
-	// MipSize not aligned to BlockSize is correct here...
-	const uint32 MipSizeX = FMath::Max(this->GetSizeX() >> MipIndex, BlockSizeX);
-	const uint32 MipSizeY = FMath::Max(this->GetSizeY() >> MipIndex, BlockSizeY);
-
 	auto* FirstObject = static_cast<FD3D12Texture*>(GetFirstLinkedObject());
 	TMap<uint32, FD3D12LockedResource*>& Map = FirstObject->LockedMap;
 	FD3D12LockedResource* LockedResource = Map[Subresource];
@@ -2037,20 +1999,11 @@ void FD3D12Texture::UnlockInternal(class FRHICommandListImmediate* RHICmdList, F
 			FD3D12ResourceLocation& UploadLocation = LockedResource->ResourceLocation;
 
 			// Copy the mip-map data from the real resource into the staging resource
-			const D3D12_RESOURCE_DESC& ResourceDesc = Resource->GetDesc();
-			D3D12_SUBRESOURCE_FOOTPRINT BufferPitchDesc;
-			BufferPitchDesc.Depth = 1;
-			BufferPitchDesc.Height = Align(MipSizeY, BlockSizeY);
-			BufferPitchDesc.Width = Align(MipSizeX, BlockSizeX);
-			BufferPitchDesc.Format = ResourceDesc.Format;
-			BufferPitchDesc.RowPitch = LockedResource->LockedPitch;
-			check(BufferPitchDesc.RowPitch % FD3D12_TEXTURE_DATA_PITCH_ALIGNMENT == 0);
-
-			D3D12_PLACED_SUBRESOURCE_FOOTPRINT PlacedTexture2D = { 0 };
-			PlacedTexture2D.Offset = UploadLocation.GetOffsetFromBaseOfResource();
-			PlacedTexture2D.Footprint = BufferPitchDesc;
-
-			CD3DX12_TEXTURE_COPY_LOCATION SourceCopyLocation(UploadLocation.GetResource()->GetResource(), PlacedTexture2D);
+
+			D3D12_PLACED_SUBRESOURCE_FOOTPRINT PlacedFootprint;
+			PlacedFootprint.Offset = UploadLocation.GetOffsetFromBaseOfResource();
+			PlacedFootprint.Footprint = LockedResource->Footprint;
+			CD3DX12_TEXTURE_COPY_LOCATION SourceCopyLocation(UploadLocation.GetResource()->GetResource(), PlacedFootprint);
 
 			// If we are on the render thread, queue up the copy on the RHIThread so it happens at the correct time.
 			if (ShouldDeferCmdListOperation(RHICmdList))
