D3D12 Bindless Resources
- Allowing constant buffer views to be allocated from the bindless heap.
- Changing RayTracing to use root constant buffer views when bindless resources are enabled.

diff --git a/Engine/Source/Runtime/D3D12RHI/Private/D3D12Adapter.cpp b/Engine/Source/Runtime/D3D12RHI/Private/D3D12Adapter.cpp
--- a/Engine/Source/Runtime/D3D12RHI/Private/D3D12Adapter.cpp
+++ b/Engine/Source/Runtime/D3D12RHI/Private/D3D12Adapter.cpp
@@ -1346,7 +1346,7 @@ void FD3D12Adapter::InitializeDevices()
 		}
 
 		StaticRayTracingGlobalRootSignature.InitStaticRayTracingGlobalRootSignatureDesc(RayTracingFlags);
-		StaticRayTracingLocalRootSignature.InitStaticRayTracingLocalRootSignatureDesc();
+		StaticRayTracingLocalRootSignature.InitStaticRayTracingLocalRootSignatureDesc(RayTracingFlags);
 #endif
 #endif // USE_STATIC_ROOT_SIGNATURE
 
diff --git a/Engine/Source/Runtime/D3D12RHI/Private/D3D12BindlessDescriptors.cpp b/Engine/Source/Runtime/D3D12RHI/Private/D3D12BindlessDescriptors.cpp
--- a/Engine/Source/Runtime/D3D12RHI/Private/D3D12BindlessDescriptors.cpp
+++ b/Engine/Source/Runtime/D3D12RHI/Private/D3D12BindlessDescriptors.cpp
@@ -38,6 +38,24 @@ FD3D12BindlessDescriptorHeapManager::FD3D12BindlessDescriptorHeapManager(FD3D12D
 	SetupInitialState(InNumDescriptorsPerHeap);
 }
 
+FD3D12DescriptorHeap* FD3D12BindlessDescriptorHeapManager::AllocateHeap(int32 InSize)
+{
+	check(NumAllocations == 0);
+
+	const int32 NewNumDescriptorsPerHeap = NumDescriptorsPerHeap + InSize;
+	ResizeHeaps(NewNumDescriptorsPerHeap);
+
+	check(Allocations.Num() == NewNumDescriptorsPerHeap);
+
+	const int32 Offset = 0;
+	const int32 Count = InSize;
+
+	Allocations.SetRange(Offset, Count, true);
+	NumAllocations += InSize;
+
+	return new FD3D12DescriptorHeap(GpuHeap.GetReference(), Offset, Count);
+}
+
 FRHIDescriptorHandle FD3D12BindlessDescriptorHeapManager::Allocate()
 {
 	FScopeLock Lock(&CriticalSection);
@@ -53,7 +71,8 @@ FRHIDescriptorHandle FD3D12BindlessDescriptorHeapManager::Allocate()
 	NumAllocations++;
 	RecordAlloc();
 
-	checkSlow(NumAllocations == Allocations.CountSetBits());
+	// TODO - find a faster way to verify Allocations
+	//checkSlow(NumAllocations == Allocations.CountSetBits());
 
 	return FRHIDescriptorHandle(Type, AllocatedIndex);
 }
@@ -71,7 +90,8 @@ void FD3D12BindlessDescriptorHeapManager::Free(FRHIDescriptorHandle InHandle)
 		NumAllocations--;
 		RecordFree();
 
-		checkSlow(NumAllocations == Allocations.CountSetBits());
+		// TODO - find a faster way to verify Allocations
+		//checkSlow(NumAllocations == Allocations.CountSetBits());
 	}
 }
 
@@ -266,6 +286,21 @@ void FD3D12BindlessDescriptorManager::UpdateDeferred(FRHIDescriptorHandle InHand
 	checkNoEntry();
 }
 
+FD3D12DescriptorHeap* FD3D12BindlessDescriptorManager::AllocateHeap(ERHIDescriptorHeapType InType, int32 InSize)
+{
+	for (FD3D12BindlessDescriptorHeapManager& Manager : Managers)
+	{
+		if (Manager.HandlesAllocation(InType))
+		{
+			return Manager.AllocateHeap(InSize);
+		}
+	}
+
+	// Bad configuration?
+	checkNoEntry();
+	return nullptr;
+}
+
 FD3D12DescriptorHeap* FD3D12BindlessDescriptorManager::GetHeap(ERHIDescriptorHeapType InType)
 {
 	for (FD3D12BindlessDescriptorHeapManager& Manager : Managers)
diff --git a/Engine/Source/Runtime/D3D12RHI/Private/D3D12BindlessDescriptors.h b/Engine/Source/Runtime/D3D12RHI/Private/D3D12BindlessDescriptors.h
--- a/Engine/Source/Runtime/D3D12RHI/Private/D3D12BindlessDescriptors.h
+++ b/Engine/Source/Runtime/D3D12RHI/Private/D3D12BindlessDescriptors.h
@@ -17,6 +17,8 @@ public:
 	FD3D12BindlessDescriptorHeapManager() = delete;
 	FD3D12BindlessDescriptorHeapManager(FD3D12Device* InDevice, ERHIDescriptorHeapType InType, ERHIBindlessConfiguration InConfiguration, uint32 InNumDescriptorsPerHeap, TConstArrayView<TStatId> InStats);
 
+	FD3D12DescriptorHeap* AllocateHeap(int32 InSize);
+
 	FRHIDescriptorHandle Allocate();
 	void                 Free(FRHIDescriptorHandle InHandle);
 
@@ -106,6 +108,8 @@ public:
 	void UpdateImmediately(FRHIDescriptorHandle InHandle, D3D12_CPU_DESCRIPTOR_HANDLE InSourceCpuHandle);
 	void UpdateDeferred(FRHIDescriptorHandle InHandle, D3D12_CPU_DESCRIPTOR_HANDLE InSourceCpuHandle);
 
+	FD3D12DescriptorHeap* AllocateHeap(ERHIDescriptorHeapType InType, int32 InSize);
+
 	FD3D12DescriptorHeap* GetHeap(ERHIDescriptorHeapType InType);
 	FD3D12DescriptorHeap* GetHeap(ERHIDescriptorHeapType InType, ERHIBindlessConfiguration InConfiguration);
 	bool                  HasHeap(ERHIDescriptorHeapType InType, ERHIBindlessConfiguration InConfiguration) const;
diff --git a/Engine/Source/Runtime/D3D12RHI/Private/D3D12DescriptorCache.cpp b/Engine/Source/Runtime/D3D12RHI/Private/D3D12DescriptorCache.cpp
--- a/Engine/Source/Runtime/D3D12RHI/Private/D3D12DescriptorCache.cpp
+++ b/Engine/Source/Runtime/D3D12RHI/Private/D3D12DescriptorCache.cpp
@@ -52,9 +52,9 @@ void FD3D12DescriptorCache::Init(uint32 InNumLocalViewDescriptors, uint32 InNumS
 	// lazily as a backup to save memory)
 	LocalSamplerHeap.Init(IsUsingBindlessSamplers() ? 0 : InNumSamplerDescriptors, ERHIDescriptorHeapType::Sampler);
 
-	NumLocalViewDescriptors = IsUsingBindlessResources() ? 0 : InNumLocalViewDescriptors;
+	NumLocalViewDescriptors = InNumLocalViewDescriptors;
 
-	CurrentViewHeap = IsUsingBindlessResources() ? nullptr : &SubAllocatedViewHeap;
+	CurrentViewHeap = &SubAllocatedViewHeap;
 	CurrentSamplerHeap = IsUsingBindlessSamplers() ? nullptr : &LocalSamplerHeap;
 }
 
@@ -154,11 +154,9 @@ void FD3D12DescriptorCache::OpenCommandList()
 	{
 		BindlessResourcesHeap = GetParentDevice()->GetBindlessDescriptorManager().GetHeap(ERHIDescriptorHeapType::Standard);
 	}
-	else
 #endif
-	{
-		CurrentViewHeap->OpenCommandList();
-	}
+
+	CurrentViewHeap->OpenCommandList();
 
 	if (!IsUsingBindlessSamplers())
 	{
@@ -173,10 +171,7 @@ void FD3D12DescriptorCache::OpenCommandList()
 
 void FD3D12DescriptorCache::CloseCommandList()
 {
-	if (!IsUsingBindlessResources())
-	{
-		CurrentViewHeap->CloseCommandList();
-	}
+	CurrentViewHeap->CloseCommandList();
 
 	if (!IsUsingBindlessSamplers())
 	{
diff --git a/Engine/Source/Runtime/D3D12RHI/Private/D3D12Descriptors.cpp b/Engine/Source/Runtime/D3D12RHI/Private/D3D12Descriptors.cpp
--- a/Engine/Source/Runtime/D3D12RHI/Private/D3D12Descriptors.cpp
+++ b/Engine/Source/Runtime/D3D12RHI/Private/D3D12Descriptors.cpp
@@ -120,13 +120,22 @@ FD3D12OnlineDescriptorManager::FD3D12OnlineDescriptorManager(FD3D12Device* Devic
 FD3D12OnlineDescriptorManager::~FD3D12OnlineDescriptorManager() = default;
 
 // Allocate and initialize the online heap
-void FD3D12OnlineDescriptorManager::Init(uint32 InTotalSize, uint32 InBlockSize)
+void FD3D12OnlineDescriptorManager::Init(uint32 InTotalSize, uint32 InBlockSize, bool bBindlessResources)
 {
-	Heap = GetParentDevice()->GetDescriptorHeapManager().AllocateHeap(
-		TEXT("Device Global - Online View Heap"),
-		ERHIDescriptorHeapType::Standard,
-		InTotalSize,
-		D3D12_DESCRIPTOR_HEAP_FLAG_SHADER_VISIBLE);
+#if PLATFORM_SUPPORTS_BINDLESS_RENDERING
+	if (bBindlessResources)
+	{
+		Heap = GetParentDevice()->GetBindlessDescriptorManager().AllocateHeap(ERHIDescriptorHeapType::Standard, InTotalSize);
+	}
+	else
+#endif
+	{
+		Heap = GetParentDevice()->GetDescriptorHeapManager().AllocateHeap(
+			TEXT("Device Global - Online View Heap"),
+			ERHIDescriptorHeapType::Standard,
+			InTotalSize,
+			D3D12_DESCRIPTOR_HEAP_FLAG_SHADER_VISIBLE);
+	}
 
 	// Update the stats
 	INC_DWORD_STAT(STAT_NumViewOnlineDescriptorHeaps);
diff --git a/Engine/Source/Runtime/D3D12RHI/Private/D3D12Descriptors.h b/Engine/Source/Runtime/D3D12RHI/Private/D3D12Descriptors.h
--- a/Engine/Source/Runtime/D3D12RHI/Private/D3D12Descriptors.h
+++ b/Engine/Source/Runtime/D3D12RHI/Private/D3D12Descriptors.h
@@ -135,7 +135,7 @@ public:
 	~FD3D12OnlineDescriptorManager();
 
 	// Setup the actual heap
-	void Init(uint32 InTotalSize, uint32 InBlockSize);
+	void Init(uint32 InTotalSize, uint32 InBlockSize, bool bBindlessResources);
 
 	// Allocate an available sub heap block from the global heap
 	FD3D12OnlineDescriptorBlock* AllocateHeapBlock();
diff --git a/Engine/Source/Runtime/D3D12RHI/Private/D3D12Device.cpp b/Engine/Source/Runtime/D3D12RHI/Private/D3D12Device.cpp
--- a/Engine/Source/Runtime/D3D12RHI/Private/D3D12Device.cpp
+++ b/Engine/Source/Runtime/D3D12RHI/Private/D3D12Device.cpp
@@ -387,10 +387,7 @@ void FD3D12Device::SetupAfterDeviceCreation()
 		GlobalSamplerHeap.Init(GGlobalSamplerHeapSize);
 	}
 
-	if (!bFullyBindlessResources)
-	{
-		OnlineDescriptorManager.Init(GOnlineDescriptorHeapSize, GOnlineDescriptorHeapBlockSize);
-	}
+	OnlineDescriptorManager.Init(GOnlineDescriptorHeapSize, GOnlineDescriptorHeapBlockSize, bFullyBindlessResources);
 
 	// Make sure we create the default views before the first command context
 	CreateDefaultViews();
diff --git a/Engine/Source/Runtime/D3D12RHI/Private/D3D12RayTracing.cpp b/Engine/Source/Runtime/D3D12RHI/Private/D3D12RayTracing.cpp
--- a/Engine/Source/Runtime/D3D12RHI/Private/D3D12RayTracing.cpp
+++ b/Engine/Source/Runtime/D3D12RHI/Private/D3D12RayTracing.cpp
@@ -4214,15 +4214,16 @@ static bool SetRayTracingShaderResources(
 
 #if D3D12RHI_USE_CONSTANT_BUFFER_VIEWS
 		D3D12_CPU_DESCRIPTOR_HANDLE LocalCBVs[MAX_CBS];
-#else // D3D12RHI_USE_CONSTANT_BUFFER_VIEWS
-		D3D12_GPU_VIRTUAL_ADDRESS LocalCBVs[MAX_CBS];
-#endif // D3D12RHI_USE_CONSTANT_BUFFER_VIEWS
+#endif
+		D3D12_GPU_VIRTUAL_ADDRESS RemoteCBVs[MAX_CBS];
 
 		D3D12_CPU_DESCRIPTOR_HANDLE LocalSRVs[MAX_SRVS];
 		D3D12_CPU_DESCRIPTOR_HANDLE LocalUAVs[MAX_UAVS];
 		D3D12_CPU_DESCRIPTOR_HANDLE LocalSamplers[MAX_SAMPLERS];
 
+#if D3D12RHI_USE_CONSTANT_BUFFER_VIEWS
 		uint32 CBVVersions[MAX_CBS];
+#endif
 		uint32 SRVVersions[MAX_SRVS];
 		uint32 UAVVersions[MAX_SRVS];
 		uint32 SamplerVersions[MAX_SRVS];
@@ -4325,10 +4326,8 @@ static bool SetRayTracingShaderResources(
 			FD3D12OfflineDescriptor Descriptor = CBV->View->GetOfflineCpuHandle();
 			Bindings.LocalCBVs[CBVIndex] = Descriptor;
 			Bindings.CBVVersions[CBVIndex] = Descriptor.GetVersion();
-		#else // D3D12RHI_USE_CONSTANT_BUFFER_VIEWS
-			Bindings.LocalCBVs[CBVIndex] = CBV->ResourceLocation.GetGPUVirtualAddress();
-			Bindings.CBVVersions[CBVIndex] = 0; // not available with GPU address path
 		#endif // D3D12RHI_USE_CONSTANT_BUFFER_VIEWS
+			Bindings.RemoteCBVs[CBVIndex] = CBV->ResourceLocation.GetGPUVirtualAddress();
 			Bindings.BoundCBVMask |= 1ull << CBVIndex;
 
 			Bindings.ReferencedResources.Add(CBV->ResourceLocation.GetResource());
@@ -4383,9 +4382,8 @@ static bool SetRayTracingShaderResources(
 
 	#if D3D12RHI_USE_CONSTANT_BUFFER_VIEWS
 		Bindings.LocalCBVs[CBVIndex] = ConstantBufferView->GetOfflineCpuHandle();
-	#else // D3D12RHI_USE_CONSTANT_BUFFER_VIEWS
-		Bindings.LocalCBVs[CBVIndex] = ResourceLocation.GetGPUVirtualAddress();
 	#endif // D3D12RHI_USE_CONSTANT_BUFFER_VIEWS
+		Bindings.RemoteCBVs[CBVIndex] = ResourceLocation.GetGPUVirtualAddress();
 
 		Bindings.BoundCBVMask |= 1ull << CBVIndex;
 	}
@@ -4444,29 +4442,30 @@ static bool SetRayTracingShaderResources(
 	if (Shader->ResourceCounts.NumCBs)
 	{
 	#if D3D12RHI_USE_CONSTANT_BUFFER_VIEWS
+		if (!EnumHasAllFlags(Shader->ResourceCounts.UsageFlags, EShaderResourceUsageFlags::BindlessResources))
+		{
+			const uint32 DescriptorTableBaseIndex = DescriptorCache.AllocateDeduplicated(Bindings.CBVVersions, Bindings.LocalCBVs, NumCBVs, D3D12_DESCRIPTOR_HEAP_TYPE_CBV_SRV_UAV, WorkerIndex);
+			const uint32 BindSlot = RootSignature->CBVRDTBindSlot(SF_Compute);
+			check(BindSlot != 0xFF);
 
-		const uint32 DescriptorTableBaseIndex = DescriptorCache.AllocateDeduplicated(Bindings.CBVVersions, Bindings.LocalCBVs, NumCBVs, D3D12_DESCRIPTOR_HEAP_TYPE_CBV_SRV_UAV, WorkerIndex);
-		const uint32 BindSlot = RootSignature->CBVRDTBindSlot(SF_Compute);
-		check(BindSlot != 0xFF);
-
-		const D3D12_GPU_DESCRIPTOR_HANDLE ResourceDescriptorTableBaseGPU = DescriptorCache.ViewHeap.GetDescriptorGPU(DescriptorTableBaseIndex);
-		Binder.SetRootDescriptorTable(BindSlot, ResourceDescriptorTableBaseGPU);
-
-	#else // D3D12RHI_USE_CONSTANT_BUFFER_VIEWS
-
-		checkf(RootSignature->CBVRDTBindSlot(SF_Compute) == 0xFF, TEXT("Root CBV descriptor tables are not implemented for ray tracing shaders."));
+			const D3D12_GPU_DESCRIPTOR_HANDLE ResourceDescriptorTableBaseGPU = DescriptorCache.ViewHeap.GetDescriptorGPU(DescriptorTableBaseIndex);
+			Binder.SetRootDescriptorTable(BindSlot, ResourceDescriptorTableBaseGPU);
+		}
+		else
+	#endif // D3D12RHI_USE_CONSTANT_BUFFER_VIEWS
+		{
+			checkf(RootSignature->CBVRDTBindSlot(SF_Compute) == 0xFF, TEXT("Root CBV descriptor tables are not implemented for ray tracing shaders."));
 
-		const uint32 BindSlot = RootSignature->CBVRDBaseBindSlot(SF_Compute);
-		check(BindSlot != 0xFF);
+			const uint32 BindSlot = RootSignature->CBVRDBaseBindSlot(SF_Compute);
+			check(BindSlot != 0xFF);
 
-		for (uint32 i = 0; i < Shader->ResourceCounts.NumCBs; ++i)
-		{
-			const uint64 SlotMask = (1ull << i);
-			D3D12_GPU_VIRTUAL_ADDRESS BufferAddress = (Bindings.BoundCBVMask & SlotMask) ? Bindings.LocalCBVs[i] : 0;
-			Binder.SetRootCBV(BindSlot, i, BufferAddress);
+			for (uint32 i = 0; i < Shader->ResourceCounts.NumCBs; ++i)
+			{
+				const uint64 SlotMask = (1ull << i);
+				D3D12_GPU_VIRTUAL_ADDRESS BufferAddress = (Bindings.BoundCBVMask & SlotMask) ? Bindings.RemoteCBVs[i] : 0;
+				Binder.SetRootCBV(BindSlot, i, BufferAddress);
+			}
 		}
-
-	#endif // D3D12RHI_USE_CONSTANT_BUFFER_VIEWS
 	}
 
 	// Bind samplers
diff --git a/Engine/Source/Runtime/D3D12RHI/Private/D3D12RootSignature.cpp b/Engine/Source/Runtime/D3D12RHI/Private/D3D12RootSignature.cpp
--- a/Engine/Source/Runtime/D3D12RHI/Private/D3D12RootSignature.cpp
+++ b/Engine/Source/Runtime/D3D12RHI/Private/D3D12RootSignature.cpp
@@ -444,10 +444,10 @@ void FD3D12RootSignature::InitStaticRayTracingGlobalRootSignatureDesc(ED3D12Root
 	Init(Creator.Finalize(), UE_HLSL_SPACE_RAY_TRACING_GLOBAL);
 }
 
-void FD3D12RootSignature::InitStaticRayTracingLocalRootSignatureDesc()
+void FD3D12RootSignature::InitStaticRayTracingLocalRootSignatureDesc(ED3D12RootSignatureFlags InFlags)
 {
 	D3D12ShaderUtils::FBinaryRootSignatureCreator Creator;
-	D3D12ShaderUtils::CreateRayTracingSignature(Creator, true, D3D12_ROOT_SIGNATURE_FLAG_LOCAL_ROOT_SIGNATURE, ED3D12RootSignatureFlags::None);
+	D3D12ShaderUtils::CreateRayTracingSignature(Creator, true, D3D12_ROOT_SIGNATURE_FLAG_LOCAL_ROOT_SIGNATURE, InFlags);
 	Init(Creator.Finalize(), UE_HLSL_SPACE_RAY_TRACING_LOCAL);
 }
 #endif // D3D12_RHI_RAYTRACING
diff --git a/Engine/Source/Runtime/D3D12RHI/Private/D3D12RootSignature.h b/Engine/Source/Runtime/D3D12RHI/Private/D3D12RootSignature.h
--- a/Engine/Source/Runtime/D3D12RHI/Private/D3D12RootSignature.h
+++ b/Engine/Source/Runtime/D3D12RHI/Private/D3D12RootSignature.h
@@ -93,7 +93,7 @@ public:
 	void InitStaticComputeRootSignatureDesc(ED3D12RootSignatureFlags InFlags);
 #if D3D12_RHI_RAYTRACING
 	void InitStaticRayTracingGlobalRootSignatureDesc(ED3D12RootSignatureFlags InFlags);
-	void InitStaticRayTracingLocalRootSignatureDesc();
+	void InitStaticRayTracingLocalRootSignatureDesc(ED3D12RootSignatureFlags InFlags);
 #endif
 
 	ID3D12RootSignature* GetRootSignature() const { return RootSignature.GetReference(); }
@@ -181,7 +181,8 @@ public:
 	inline bool HasSRVs() const { return bHasSRVs; }
 	inline bool HasCBVs() const { return bHasCBVs; }
 	inline bool HasRootCBs() const { return bHasRootCBs; }
-	inline bool HasTableResources() const { return bHasUAVs || bHasSRVs || bHasCBVs; }
+	inline bool HasTableResources() const { return bHasUAVs || bHasSRVs; }
+	inline bool HasTableConstants() const { return bHasCBVs; }
 	inline bool HasSamplers() const { return bHasSamplers; }
 	inline bool UsesDynamicResources() const { return bUsesDynamicResources; }
 	inline bool UsesDynamicSamplers() const { return bUsesDynamicSamplers; }
diff --git a/Engine/Source/Runtime/D3D12RHI/Private/D3D12StateCache.cpp b/Engine/Source/Runtime/D3D12RHI/Private/D3D12StateCache.cpp
--- a/Engine/Source/Runtime/D3D12RHI/Private/D3D12StateCache.cpp
+++ b/Engine/Source/Runtime/D3D12RHI/Private/D3D12StateCache.cpp
@@ -590,6 +590,10 @@ void FD3D12StateCache::ApplyState(ED3D12PipelineType PipelineType)
 
 void FD3D12StateCache::ApplyResources(const FD3D12RootSignature* const pRootSignature, uint32 StartStage, uint32 EndStage)
 {
+	const bool bUAVs = pRootSignature->HasUAVs();
+	const bool bSRVs = pRootSignature->HasSRVs();
+	const bool bCBVs = pRootSignature->HasCBVs();
+
 	// Determine what resource bind slots are dirty for the current shaders and how many descriptor table slots we need.
 	// We only set dirty resources that can be used for the upcoming Draw/Dispatch.
 	SRVSlotMask CurrentShaderDirtySRVSlots[SF_NumStandardFrequencies] = {};
@@ -608,22 +612,25 @@ void FD3D12StateCache::ApplyResources(const FD3D12RootSignature* const pRootSign
 
 	for (uint32 iTries = 0; iTries < 2; ++iTries)
 	{
-		const UAVSlotMask CurrentShaderUAVRegisterMask = BitMask<UAVSlotMask>(PipelineState.Common.CurrentShaderUAVCounts[UAVStage]);
-		CurrentShaderDirtyUAVSlots = CurrentShaderUAVRegisterMask & PipelineState.Common.UAVCache.DirtySlotMask[UAVStage];
-		if (CurrentShaderDirtyUAVSlots)
+		if (bUAVs)
 		{
-			if (ResourceBindingTier <= D3D12_RESOURCE_BINDING_TIER_2)
+			const UAVSlotMask CurrentShaderUAVRegisterMask = BitMask<UAVSlotMask>(PipelineState.Common.CurrentShaderUAVCounts[UAVStage]);
+			CurrentShaderDirtyUAVSlots = CurrentShaderUAVRegisterMask & PipelineState.Common.UAVCache.DirtySlotMask[UAVStage];
+			if (CurrentShaderDirtyUAVSlots)
 			{
-				// Tier 1 and 2 HW requires the full number of UAV descriptors defined in the root signature's descriptor table.
-				NumUAVs = pRootSignature->MaxUAVCount(UAVStage);
-			}
-			else
-			{
-				NumUAVs = PipelineState.Common.CurrentShaderUAVCounts[UAVStage];
-			}
+				if (ResourceBindingTier <= D3D12_RESOURCE_BINDING_TIER_2)
+				{
+					// Tier 1 and 2 HW requires the full number of UAV descriptors defined in the root signature's descriptor table.
+					NumUAVs = pRootSignature->MaxUAVCount(UAVStage);
+				}
+				else
+				{
+					NumUAVs = PipelineState.Common.CurrentShaderUAVCounts[UAVStage];
+				}
 
-			check(NumUAVs > 0 && NumUAVs <= MAX_UAVS);
-			NumViews += NumUAVs;
+				check(NumUAVs > 0 && NumUAVs <= MAX_UAVS);
+				NumViews += NumUAVs;
+			}
 		}
 
 		for (uint32 Stage = StartStage; Stage < EndStage; ++Stage)
@@ -633,42 +640,48 @@ void FD3D12StateCache::ApplyResources(const FD3D12RootSignature* const pRootSign
 				continue;
 			}
 
-			// Note this code assumes the starting register is index 0.
-			const SRVSlotMask CurrentShaderSRVRegisterMask = BitMask<SRVSlotMask>(PipelineState.Common.CurrentShaderSRVCounts[Stage]);
-			CurrentShaderDirtySRVSlots[Stage] = CurrentShaderSRVRegisterMask & PipelineState.Common.SRVCache.DirtySlotMask[Stage];
-			if (CurrentShaderDirtySRVSlots[Stage])
+			if (bSRVs)
 			{
-				if (ResourceBindingTier == D3D12_RESOURCE_BINDING_TIER_1)
+				// Note this code assumes the starting register is index 0.
+				const SRVSlotMask CurrentShaderSRVRegisterMask = BitMask<SRVSlotMask>(PipelineState.Common.CurrentShaderSRVCounts[Stage]);
+				CurrentShaderDirtySRVSlots[Stage] = CurrentShaderSRVRegisterMask & PipelineState.Common.SRVCache.DirtySlotMask[Stage];
+				if (CurrentShaderDirtySRVSlots[Stage])
 				{
-					// Tier 1 HW requires the full number of SRV descriptors defined in the root signature's descriptor table.
-					NumSRVs[Stage] = pRootSignature->MaxSRVCount(Stage);
-				}
-				else
-				{
-					NumSRVs[Stage] = PipelineState.Common.CurrentShaderSRVCounts[Stage];
-				}
+					if (ResourceBindingTier == D3D12_RESOURCE_BINDING_TIER_1)
+					{
+						// Tier 1 HW requires the full number of SRV descriptors defined in the root signature's descriptor table.
+						NumSRVs[Stage] = pRootSignature->MaxSRVCount(Stage);
+					}
+					else
+					{
+						NumSRVs[Stage] = PipelineState.Common.CurrentShaderSRVCounts[Stage];
+					}
 
-				check(NumSRVs[Stage] > 0 && NumSRVs[Stage] <= MAX_SRVS);
-				NumViews += NumSRVs[Stage];
+					check(NumSRVs[Stage] > 0 && NumSRVs[Stage] <= MAX_SRVS);
+					NumViews += NumSRVs[Stage];
+				}
 			}
 
 #if D3D12RHI_USE_CONSTANT_BUFFER_VIEWS
-			const CBVSlotMask CurrentShaderCBVRegisterMask = BitMask<CBVSlotMask>(PipelineState.Common.CurrentShaderCBCounts[Stage]);
-			CurrentShaderDirtyCBVSlots[Stage] = CurrentShaderCBVRegisterMask & PipelineState.Common.CBVCache.DirtySlotMask[Stage];
-			if (CurrentShaderDirtyCBVSlots[Stage])
+			if (bCBVs)
 			{
-				if (ResourceBindingTier == D3D12_RESOURCE_BINDING_TIER_1)
+				const CBVSlotMask CurrentShaderCBVRegisterMask = BitMask<CBVSlotMask>(PipelineState.Common.CurrentShaderCBCounts[Stage]);
+				CurrentShaderDirtyCBVSlots[Stage] = CurrentShaderCBVRegisterMask & PipelineState.Common.CBVCache.DirtySlotMask[Stage];
+				if (CurrentShaderDirtyCBVSlots[Stage])
 				{
-					// Tier 1 HW requires the full number of SRV descriptors defined in the root signature's descriptor table.
-					NumCBVs[Stage] = pRootSignature->MaxCBVCount(Stage);
-				}
-				else
-				{
-					NumCBVs[Stage] = PipelineState.Common.CurrentShaderCBCounts[Stage];
-				}
+					if (ResourceBindingTier == D3D12_RESOURCE_BINDING_TIER_1)
+					{
+						// Tier 1 HW requires the full number of SRV descriptors defined in the root signature's descriptor table.
+						NumCBVs[Stage] = pRootSignature->MaxCBVCount(Stage);
+					}
+					else
+					{
+						NumCBVs[Stage] = PipelineState.Common.CurrentShaderCBCounts[Stage];
+					}
 
-				check(NumCBVs[Stage] > 0 && NumCBVs[Stage] <= MAX_SRVS);
-				NumViews += NumCBVs[Stage];
+					check(NumCBVs[Stage] > 0 && NumCBVs[Stage] <= MAX_CBS);
+					NumViews += NumCBVs[Stage];
+				}
 			}
 #endif
 			// Note: CBVs don't currently use descriptor tables but we still need to know what resource point slots are dirty.
@@ -701,6 +714,7 @@ void FD3D12StateCache::ApplyResources(const FD3D12RootSignature* const pRootSign
 	}
 
 	// Shader resource views
+	if (bSRVs)
 	{
 		//SCOPE_CYCLE_COUNTER(STAT_D3D12ApplyStateSetSRVTime);
 		FD3D12ShaderResourceViewCache& SRVCache = PipelineState.Common.SRVCache;
@@ -717,6 +731,7 @@ void FD3D12StateCache::ApplyResources(const FD3D12RootSignature* const pRootSign
 
 #if D3D12RHI_USE_CONSTANT_BUFFER_VIEWS
 	// Constant buffers
+	if (bCBVs)
 	{
 		//SCOPE_CYCLE_COUNTER(STAT_D3D12ApplyStateSetConstantBufferTime);
 		FD3D12ConstantBufferCache& CBVCache = PipelineState.Common.CBVCache;
@@ -744,6 +759,63 @@ void FD3D12StateCache::ApplyBindlessResources(const FD3D12RootSignature* const p
 		PipelineState.Common.QueuedBindlessSRVs[Index].Reset();
 		PipelineState.Common.QueuedBindlessUAVs[Index].Reset();
 	}
+
+#if D3D12RHI_USE_CONSTANT_BUFFER_VIEWS
+	if (pRootSignature->HasCBVs())
+	{
+		FD3D12ConstantBufferCache& CBVCache = PipelineState.Common.CBVCache;
+
+		CBVSlotMask CurrentShaderDirtyCBVSlots[SF_NumStandardFrequencies] = {};
+		uint32 NumCBVs[SF_NumStandardFrequencies] = {};
+
+		uint32 NumViews = 0;
+
+		for (uint32 iTries = 0; iTries < 2; ++iTries)
+		{
+			for (uint32 Stage = StartStage; Stage < EndStage; ++Stage)
+			{
+				if (ShouldSkipStage(Stage))
+				{
+					continue;
+				}
+
+				const uint32 ConstantBufferCount = PipelineState.Common.CurrentShaderCBCounts[Stage];
+
+				const CBVSlotMask CurrentShaderCBVRegisterMask = BitMask<CBVSlotMask>(ConstantBufferCount);
+				CurrentShaderDirtyCBVSlots[Stage] = CurrentShaderCBVRegisterMask & CBVCache.DirtySlotMask[Stage];
+				if (CurrentShaderDirtyCBVSlots[Stage])
+				{
+					check(ConstantBufferCount > 0 && ConstantBufferCount <= MAX_CBS);
+
+					NumCBVs[Stage] = ConstantBufferCount;
+					NumViews += ConstantBufferCount;
+				}
+				// Note: CBVs don't currently use descriptor tables but we still need to know what resource point slots are dirty.
+			}
+
+			// See if the descriptor slots will fit
+			if (!DescriptorCache.GetCurrentViewHeap()->CanReserveSlots(NumViews))
+			{
+				if (DescriptorCache.GetCurrentViewHeap()->RollOver())
+				{
+					// If descriptor heaps changed, then all our tables are dirty again and we need to recalculate the number of slots we need.
+					NumViews = 0;
+					continue;
+				}
+			}
+		}
+
+		uint32 ViewHeapSlot = DescriptorCache.GetCurrentViewHeap()->ReserveSlots(NumViews);
+
+		for (uint32 Index = StartStage; Index < EndStage; Index++)
+		{
+			if (CurrentShaderDirtyCBVSlots[Index])
+			{
+				DescriptorCache.SetConstantBufferViews(static_cast<EShaderFrequency>(Index), pRootSignature, CBVCache, CurrentShaderDirtyCBVSlots[Index], NumCBVs[Index], ViewHeapSlot);
+			}
+		}
+	}
+#endif // D3D12RHI_USE_CONSTANT_BUFFER_VIEWS
 }
 
 void FD3D12StateCache::ApplyConstants(const FD3D12RootSignature* const pRootSignature, uint32 StartStage, uint32 EndStage)
diff --git a/Engine/Source/Runtime/D3D12RHI/Public/D3D12RootSignatureDefinitions.h b/Engine/Source/Runtime/D3D12RHI/Public/D3D12RootSignatureDefinitions.h
--- a/Engine/Source/Runtime/D3D12RHI/Public/D3D12RootSignatureDefinitions.h
+++ b/Engine/Source/Runtime/D3D12RHI/Public/D3D12RootSignatureDefinitions.h
@@ -3,7 +3,16 @@
 #pragma once
 
 #include "D3D12RHI.h"
+#include "Containers/Array.h"
+#include "Containers/UnrealString.h"
+#include "Misc/EnumClassFlags.h"
+#include "Misc/StringBuilder.h"
 #include "RayTracingBuiltInResources.h"
+#include "RHIDefinitions.h"
+
+#if !defined(D3DINTERFACE)
+	#include "D3D12ThirdParty.h"
+#endif
 
 enum class ED3D12RootSignatureFlags
 {
@@ -21,10 +30,13 @@ namespace D3D12ShaderUtils
 	namespace StaticRootSignatureConstants
 	{
 		// Assume descriptors are volatile because we don't initialize all the descriptors in a table, just the ones used by the current shaders.
-		const D3D12_DESCRIPTOR_RANGE_FLAGS SRVDescriptorRangeFlags = D3D12_DESCRIPTOR_RANGE_FLAG_DATA_STATIC_WHILE_SET_AT_EXECUTE | D3D12_DESCRIPTOR_RANGE_FLAG_DESCRIPTORS_VOLATILE;
-		const D3D12_DESCRIPTOR_RANGE_FLAGS CBVDescriptorRangeFlags = D3D12_DESCRIPTOR_RANGE_FLAG_DATA_STATIC_WHILE_SET_AT_EXECUTE | D3D12_DESCRIPTOR_RANGE_FLAG_DESCRIPTORS_VOLATILE;
-		const D3D12_DESCRIPTOR_RANGE_FLAGS UAVDescriptorRangeFlags = D3D12_DESCRIPTOR_RANGE_FLAG_DATA_VOLATILE | D3D12_DESCRIPTOR_RANGE_FLAG_DESCRIPTORS_VOLATILE;
-		const D3D12_DESCRIPTOR_RANGE_FLAGS SamplerDescriptorRangeFlags = D3D12_DESCRIPTOR_RANGE_FLAG_DESCRIPTORS_VOLATILE;
+		constexpr D3D12_DESCRIPTOR_RANGE_FLAGS SRVDescriptorRangeFlags = D3D12_DESCRIPTOR_RANGE_FLAG_DATA_STATIC_WHILE_SET_AT_EXECUTE | D3D12_DESCRIPTOR_RANGE_FLAG_DESCRIPTORS_VOLATILE;
+		constexpr D3D12_DESCRIPTOR_RANGE_FLAGS CBVDescriptorRangeFlags = D3D12_DESCRIPTOR_RANGE_FLAG_DATA_STATIC_WHILE_SET_AT_EXECUTE | D3D12_DESCRIPTOR_RANGE_FLAG_DESCRIPTORS_VOLATILE;
+		constexpr D3D12_DESCRIPTOR_RANGE_FLAGS UAVDescriptorRangeFlags = D3D12_DESCRIPTOR_RANGE_FLAG_DATA_VOLATILE | D3D12_DESCRIPTOR_RANGE_FLAG_DESCRIPTORS_VOLATILE;
+		constexpr D3D12_DESCRIPTOR_RANGE_FLAGS SamplerDescriptorRangeFlags = D3D12_DESCRIPTOR_RANGE_FLAG_DESCRIPTORS_VOLATILE;
+
+		// We always set the data in an upload heap before calling Set*RootConstantBufferView.
+		constexpr D3D12_ROOT_DESCRIPTOR_FLAGS CBVRootDescriptorFlags = D3D12_ROOT_DESCRIPTOR_FLAG_DATA_STATIC;
 	}
 
 	enum class ERootSignatureRangeType
@@ -177,7 +189,8 @@ namespace D3D12ShaderUtils
 		virtual ~FRootSignatureCreator() = default;
 
 		virtual void AddRootFlag(D3D12_ROOT_SIGNATURE_FLAGS Flag) = 0;
-		virtual void AddShaderResourceViewParameter(uint32 Register, uint32 Space) = 0;
+		virtual void AddConstantBufferViewParameter(uint32 Register, uint32 Space, ERootSignatureVisibility Visibility = ERootSignatureVisibility::All) = 0;
+		virtual void AddShaderResourceViewParameter(uint32 Register, uint32 Space, ERootSignatureVisibility Visibility = ERootSignatureVisibility::All) = 0;
 		virtual void AddTable(ERootSignatureVisibility Visibility, ERootSignatureRangeType Type, int32 NumDescriptors, D3D12_DESCRIPTOR_RANGE_FLAGS FlagsOverride = D3D12_DESCRIPTOR_RANGE_FLAG_NONE) = 0;
 		virtual void AddConstantsParameter(uint32 Num32BitValues, uint32 Register, uint32 Space) = 0;
 
@@ -201,12 +214,22 @@ namespace D3D12ShaderUtils
 			}
 		}
 
+		bool HasFlags(ED3D12RootSignatureFlags InFlags) const
+		{
+			return EnumHasAllFlags(Flags, InFlags);
+		}
+
 		void SetRegisterSpace(uint32 InSpace)
 		{
 			RegisterSpace = InSpace;
 		}
 
-		inline bool ShouldSkipType(ERootSignatureRangeType Type)
+		uint32 GetRegisterSpace() const
+		{
+			return RegisterSpace;
+		}
+
+		inline bool ShouldSkipType(ERootSignatureRangeType Type) const
 		{
 			if (Type == ERootSignatureRangeType::SRV || Type == ERootSignatureRangeType::UAV)
 			{
@@ -222,7 +245,7 @@ namespace D3D12ShaderUtils
 		}
 	};
 
-	struct FBinaryRootSignatureCreator : public FRootSignatureCreator
+	struct FBinaryRootSignatureCreator final : public FRootSignatureCreator
 	{
 		D3D12_ROOT_SIGNATURE_FLAGS RootFlags = D3D12_ROOT_SIGNATURE_FLAG_NONE;
 		TArray<CD3DX12_DESCRIPTOR_RANGE1> DescriptorRanges;
@@ -243,10 +266,16 @@ namespace D3D12ShaderUtils
 			RootFlags |= RootFlag;
 		}
 
-		void AddShaderResourceViewParameter(uint32 Register, uint32 Space)
+		void AddConstantBufferViewParameter(uint32 Register, uint32 Space, ERootSignatureVisibility Visibility = ERootSignatureVisibility::All) override
 		{
 			CD3DX12_ROOT_PARAMETER1& Parameter = Parameters.AddZeroed_GetRef();
-			Parameter.InitAsShaderResourceView(Register, Space);
+			Parameter.InitAsConstantBufferView(Register, Space, StaticRootSignatureConstants::CBVRootDescriptorFlags, GetD3D12ShaderVisibility(Visibility));
+		}
+
+		void AddShaderResourceViewParameter(uint32 Register, uint32 Space, ERootSignatureVisibility Visibility = ERootSignatureVisibility::All) override
+		{
+			CD3DX12_ROOT_PARAMETER1& Parameter = Parameters.AddZeroed_GetRef();
+			Parameter.InitAsShaderResourceView(Register, Space, D3D12_ROOT_DESCRIPTOR_FLAG_NONE, GetD3D12ShaderVisibility(Visibility));
 		}
 
 		void AddConstantsParameter(uint32 Num32BitValues, uint32 Register, uint32 Space) override
@@ -295,7 +324,7 @@ namespace D3D12ShaderUtils
 	};
 
 	/* Root signature generator for DXC */
-	struct FTextRootSignatureCreator : public FRootSignatureCreator
+	struct FTextRootSignatureCreator final : public FRootSignatureCreator
 	{
 		void AddRootFlag(D3D12_ROOT_SIGNATURE_FLAGS InFlag) override
 		{
@@ -306,9 +335,19 @@ namespace D3D12ShaderUtils
 			RootFlags += GetFlagName(InFlag);
 		}
 
-		void AddShaderResourceViewParameter(uint32 Register, uint32 Space) override
+		void AddConstantBufferViewParameter(uint32 Register, uint32 Space, ERootSignatureVisibility Visibility = ERootSignatureVisibility::All) override
 		{
-			FString Line = FString::Printf(TEXT("SRV(t%d, space=%d))"), Register, Space);
+			FString Line = FString::Printf(TEXT("CBV(t%d, space=%d, visibility=%s))"), Register, Space, GetVisibilityFlag(Visibility));
+			if (Table.Len() > 0)
+			{
+				Table += ",";
+			}
+			Table += Line;
+		}
+
+		void AddShaderResourceViewParameter(uint32 Register, uint32 Space, ERootSignatureVisibility Visibility = ERootSignatureVisibility::All) override
+		{
+			FString Line = FString::Printf(TEXT("SRV(t%d, space=%d, visibility=%s))"), Register, Space, GetVisibilityFlag(Visibility));
 			if (Table.Len() > 0)
 			{
 				Table += ",";
@@ -415,7 +454,17 @@ namespace D3D12ShaderUtils
 			Creator.AddConstantsParameter(NumConstants, RAY_TRACING_SYSTEM_ROOTCONSTANT_REGISTER, UE_HLSL_SPACE_RAY_TRACING_SYSTEM);
 		}
 
-		AddAllStandardTablesForVisibility(Creator, ERootSignatureVisibility::All);
+		if (Creator.HasFlags(ED3D12RootSignatureFlags::BindlessResources))
+		{
+			for (uint32 Index = 0; Index < MAX_CBS; Index++)
+			{
+				Creator.AddConstantBufferViewParameter(Index, Creator.GetRegisterSpace());
+			}
+		}
+		else
+		{
+			AddAllStandardTablesForVisibility(Creator, ERootSignatureVisibility::All);
+		}
 
 		Creator.AddTable(ERootSignatureVisibility::All, ERootSignatureRangeType::UAV, MAX_UAVS, D3D12_DESCRIPTOR_RANGE_FLAG_DATA_STATIC_WHILE_SET_AT_EXECUTE | D3D12_DESCRIPTOR_RANGE_FLAG_DESCRIPTORS_VOLATILE);
 	}
