Fixing EditorCooked compiles by checking for a common D3D12 define before including d3d12.h.

diff --git a/Engine/Source/Runtime/D3D12RHI/Public/D3D12RootSignatureDefinitions.h b/Engine/Source/Runtime/D3D12RHI/Public/D3D12RootSignatureDefinitions.h
--- a/Engine/Source/Runtime/D3D12RHI/Public/D3D12RootSignatureDefinitions.h
+++ b/Engine/Source/Runtime/D3D12RHI/Public/D3D12RootSignatureDefinitions.h
@@ -10,7 +10,7 @@
 #include "RayTracingBuiltInResources.h"
 #include "RHIDefinitions.h"
 
-#if !defined(D3DINTERFACE)
+#if !defined(D3D12_MAJOR_VERSION)
 	#include "D3D12ThirdParty.h"
 #endif
 
