Don't version upload heap memory when WriteOnly_NoOverwrite is used for dynamic data.

diff --git a/Engine/Source/Runtime/D3D12RHI/Private/D3D12Buffer.cpp b/Engine/Source/Runtime/D3D12RHI/Private/D3D12Buffer.cpp
--- a/Engine/Source/Runtime/D3D12RHI/Private/D3D12Buffer.cpp
+++ b/Engine/Source/Runtime/D3D12RHI/Private/D3D12Buffer.cpp
@@ -616,7 +616,7 @@ void* FD3D12DynamicRHI::LockBuffer(FRHICommandListBase& RHICmdList, FD3D12Buffer
 	{
 		check(LockMode == RLM_WriteOnly || LockMode == RLM_WriteOnly_NoOverwrite);
 
-		if (LockedData.bHasNeverBeenLocked)
+		if (LockedData.bHasNeverBeenLocked || LockMode == RLM_WriteOnly_NoOverwrite)
 		{
 			// Buffers on upload heap are mapped right after creation
 			Data = Buffer->ResourceLocation.GetMappedBaseAddress();
@@ -627,7 +627,7 @@ void* FD3D12DynamicRHI::LockBuffer(FRHICommandListBase& RHICmdList, FD3D12Buffer
 			FD3D12Device* Device = Buffer->GetParentDevice();
 
 			// If on the RenderThread, queue up a command on the RHIThread to rename this buffer at the correct time
-			if (RHICmdList.IsTopOfPipe() && LockMode == RLM_WriteOnly)
+			if (RHICmdList.IsTopOfPipe())
 			{
 				FRHICommandRenameUploadBuffer* Command = ALLOC_COMMAND_CL(RHICmdList, FRHICommandRenameUploadBuffer)(Buffer, Device);
 				Data = Adapter.GetUploadHeapAllocator(Device->GetGPUIndex()).AllocUploadResource(BufferSize, Buffer->BufferAlignment, Command->NewLocation);
