Fixed bug introduced in CL 29014082 which made the staging buffer allocated by FD3D12Texture::Lock have the size of the entire resource, instead of being the size of only the locked subresource.

This significantly increased memory usage in some cases, leading to out of memory situations.

#rnx

diff --git a/Engine/Source/Runtime/D3D12RHI/Private/D3D12Texture.cpp b/Engine/Source/Runtime/D3D12RHI/Private/D3D12Texture.cpp
--- a/Engine/Source/Runtime/D3D12RHI/Private/D3D12Texture.cpp
+++ b/Engine/Source/Runtime/D3D12RHI/Private/D3D12Texture.cpp
@@ -1750,15 +1750,21 @@ void* FD3D12Texture::Lock(class FRHICommandListImmediate* RHICmdList, uint32 Mip
 
 	const D3D12_RESOURCE_DESC& ResourceDesc = GetResource()->GetDesc();
 
-	UINT64 TotalBytes;
 	D3D12_PLACED_SUBRESOURCE_FOOTPRINT PlacedFootprint;
-	Device->GetDevice()->GetCopyableFootprints(&ResourceDesc, Subresource, 1, 0, &PlacedFootprint, nullptr, nullptr, &TotalBytes);
+	Device->GetDevice()->GetCopyableFootprints(&ResourceDesc, Subresource, 1, 0, &PlacedFootprint, nullptr, nullptr, nullptr);
 
+	// GetCopyableFootprints returns the offset from the start of the resource to the specified subresource, but our staging buffer represents
+	// only the selected subresource, so we need to reset the offset to 0.
+	PlacedFootprint.Offset = 0;
+
+	// Store the footprint information so we don't have to recompute it in Unlock.
 	LockedResource->Footprint = PlacedFootprint.Footprint;
+
 	DestStride = LockedResource->Footprint.RowPitch;
+	const uint64 SubresourceSize = LockedResource->Footprint.RowPitch * LockedResource->Footprint.Height * LockedResource->Footprint.Depth;
 	if (OutLockedByteCount)
 	{
-		*OutLockedByteCount = TotalBytes;
+		*OutLockedByteCount = SubresourceSize;
 	}
 
 	FD3D12CommandContext& Context = Device->GetDefaultCommandContext();
@@ -1780,7 +1786,8 @@ void* FD3D12Texture::Lock(class FRHICommandListImmediate* RHICmdList, uint32 Mip
 		// If we're writing to the texture, allocate a system memory buffer to receive the new contents.
 		// Use an upload heap to copy data to a default resource.
 
-		void* pData = Device->GetDefaultFastAllocator().Allocate(TotalBytes, D3D12_TEXTURE_DATA_PLACEMENT_ALIGNMENT, &LockedResource->ResourceLocation);
+		const uint64 BufferSize = Align(SubresourceSize, D3D12_TEXTURE_DATA_PLACEMENT_ALIGNMENT);
+		void* pData = Device->GetDefaultFastAllocator().Allocate(BufferSize, D3D12_TEXTURE_DATA_PLACEMENT_ALIGNMENT, &LockedResource->ResourceLocation);
 		if (nullptr == pData)
 		{
 			check(false);
@@ -1802,9 +1809,9 @@ void* FD3D12Texture::Lock(class FRHICommandListImmediate* RHICmdList, uint32 Mip
 		FD3D12Resource* StagingTexture = nullptr;
 
 		const FRHIGPUMask Node = Device->GetGPUMask();
-		VERIFYD3D12RESULT(Adapter->CreateBuffer(D3D12_HEAP_TYPE_READBACK, Node, Node, TotalBytes, &StagingTexture, nullptr));
+		VERIFYD3D12RESULT(Adapter->CreateBuffer(D3D12_HEAP_TYPE_READBACK, Node, Node, SubresourceSize, &StagingTexture, nullptr));
 
-		LockedResource->ResourceLocation.AsStandAlone(StagingTexture, TotalBytes);
+		LockedResource->ResourceLocation.AsStandAlone(StagingTexture, SubresourceSize);
 
 		CD3DX12_TEXTURE_COPY_LOCATION DestCopyLocation(StagingTexture->GetResource(), PlacedFootprint);
 		CD3DX12_TEXTURE_COPY_LOCATION SourceCopyLocation(GetResource()->GetResource(), Subresource);
