Fix crash when using mGPU.

RHISetShaderRootConstants was not implemented in the context redirector, so it went to the default implementation which just does a checkNoEntry().

#rnx

diff --git a/Engine/Source/Runtime/D3D12RHI/Private/D3D12CommandContext.h b/Engine/Source/Runtime/D3D12RHI/Private/D3D12CommandContext.h
--- a/Engine/Source/Runtime/D3D12RHI/Private/D3D12CommandContext.h
+++ b/Engine/Source/Runtime/D3D12RHI/Private/D3D12CommandContext.h
@@ -910,6 +910,10 @@ public:
 	{
 		ContextRedirect(RHISetBlendFactor(BlendFactor));
 	}
+	FORCEINLINE void RHISetShaderRootConstants(const FUint32Vector4& Constants) final override
+	{
+		ContextRedirect(RHISetShaderRootConstants(Constants));
+	}
 	FORCEINLINE virtual void RHIDrawPrimitive(uint32 BaseVertexIndex, uint32 NumPrimitives, uint32 NumInstances) final override
 	{
 		ContextRedirect(RHIDrawPrimitive(BaseVertexIndex, NumPrimitives, NumInstances));
