Addressing some issues with Bindless Samplers
- Adding logging when we get close to the 2048 SamplerState limit.
- Don't try and resize the Sampler heap past the D3D12 limits.

diff --git a/Engine/Source/Runtime/D3D12RHI/Private/D3D12BindlessDescriptors.cpp b/Engine/Source/Runtime/D3D12RHI/Private/D3D12BindlessDescriptors.cpp
--- a/Engine/Source/Runtime/D3D12RHI/Private/D3D12BindlessDescriptors.cpp
+++ b/Engine/Source/Runtime/D3D12RHI/Private/D3D12BindlessDescriptors.cpp
@@ -65,6 +65,12 @@ FRHIDescriptorHandle FD3D12BindlessDescriptorHeapManager::Allocate()
 	while ((AllocatedIndex = Allocations.FindAndSetFirstZeroBit()) == INDEX_NONE)
 	{
 		const int32 NewNumDescriptorsPerHeap = NumDescriptorsPerHeap * 1.5;
+
+		if (Type == ERHIDescriptorHeapType::Sampler && NewNumDescriptorsPerHeap > D3D12_MAX_SHADER_VISIBLE_SAMPLER_HEAP_SIZE)
+		{
+			return FRHIDescriptorHandle();
+		}
+
 		ResizeHeaps(NewNumDescriptorsPerHeap);
 	}
 
diff --git a/Engine/Source/Runtime/D3D12RHI/Private/D3D12State.cpp b/Engine/Source/Runtime/D3D12RHI/Private/D3D12State.cpp
--- a/Engine/Source/Runtime/D3D12RHI/Private/D3D12State.cpp
+++ b/Engine/Source/Runtime/D3D12RHI/Private/D3D12State.cpp
@@ -224,6 +224,30 @@ FSamplerStateRHIRef FD3D12DynamicRHI::RHICreateSamplerState(const FSamplerStateI
 	});
 }
 
+int32 GD3D12SamplerWarningThreshold = 10;
+static FAutoConsoleVariableRef CVarD3D12SamplerWarningThreshold(
+	TEXT("D3D12.SamplerWarningThreshold"),
+	GD3D12SamplerWarningThreshold,
+	TEXT("Threshold to start warning about creating too many sampler states")
+);
+
+static void LogSamplerStateWarning(const FSamplerStateInitializerRHI& Initializer)
+{
+	UE_LOG(LogD3D12RHI, Warning,
+		TEXT("Approaching SamplerState limit: FSamplerStateInitializerRHI(Filter: %d, AddressU: %d, AddressV: %d, AddressW: %d, MipBias: %f, MinMipLevel: %f, MaxMipLevel: %f, MaxAnisotropy: %d, BorderColor: %d, SamplerComparisonFunction: %d)"),
+		Initializer.Filter.GetIntValue(),
+		Initializer.AddressU.GetIntValue(),
+		Initializer.AddressV.GetIntValue(),
+		Initializer.AddressW.GetIntValue(),
+		Initializer.MipBias,
+		Initializer.MinMipLevel,
+		Initializer.MaxMipLevel,
+		Initializer.MaxAnisotropy,
+		Initializer.BorderColor,
+		Initializer.SamplerComparisonFunction.GetIntValue()
+	);
+}
+
 FD3D12SamplerState* FD3D12Device::CreateSampler(const FSamplerStateInitializerRHI& Initializer)
 {
 	D3D12_SAMPLER_DESC SamplerDesc;
@@ -286,6 +310,11 @@ FD3D12SamplerState* FD3D12Device::CreateSampler(const FSamplerStateInitializerRH
 		// 16-bit IDs are used for faster hashing
 		check(SamplerID < 0xffff);
 
+		if (static_cast<int32>(SamplerID + 1) > (D3D12_MAX_SHADER_VISIBLE_SAMPLER_HEAP_SIZE - GD3D12SamplerWarningThreshold))
+		{
+			LogSamplerStateWarning(Initializer);
+		}
+
 		FD3D12SamplerState* NewSampler = new FD3D12SamplerState(this, SamplerDesc, static_cast<uint16>(SamplerID));
 
 		SamplerMap.Add(SamplerDesc, NewSampler);
