Add a material setting to disable VRS (both per-draw and screen-space) per-material.
Updates D3D12 state cache SetShadingRate and SetShadingRateImage functions to match the underlying command context functionality - both VRS combiners are now handled by SetShadingRate. If a material is opted-out of VRS, it will force passthrough a 1x1 rate.

diff --git a/Engine/Source/Runtime/D3D12RHI/Private/D3D12Commands.cpp b/Engine/Source/Runtime/D3D12RHI/Private/D3D12Commands.cpp
--- a/Engine/Source/Runtime/D3D12RHI/Private/D3D12Commands.cpp
+++ b/Engine/Source/Runtime/D3D12RHI/Private/D3D12Commands.cpp
@@ -829,7 +829,14 @@ void FD3D12CommandContext::RHISetGraphicsPipelineState(FRHIGraphicsPipelineState
 
 	if (GRHISupportsPipelineVariableRateShading && GRHIVariableRateShadingEnabled)
 	{
-		StateCache.SetShadingRate(GraphicsPipelineState->PipelineStateInitializer.ShadingRate, VRSRB_Passthrough);
+		if (GraphicsPipelineState->PipelineStateInitializer.bAllowVariableRateShading)
+		{
+			StateCache.SetShadingRate(GraphicsPipelineState->PipelineStateInitializer.ShadingRate, VRSRB_Passthrough, VRSRB_Max);
+		}
+		else
+		{
+			StateCache.SetShadingRate(EVRSShadingRate::VRSSR_1x1, VRSRB_Passthrough, VRSRB_Passthrough);
+		}
 	}
 
 	StateCache.SetGraphicsPipelineState(GraphicsPipelineState);
@@ -1378,17 +1385,16 @@ void FD3D12CommandContext::SetRenderTargetsAndClear(const FRHISetRenderTargetsIn
 					D3D12_RESOURCE_STATE_SHADING_RATE_SOURCE,
 					D3D12_RESOURCE_BARRIER_ALL_SUBRESOURCES);
 
-				StateCache.SetShadingRateImage(Resource, RenderTargetsInfo.ShadingRateTextureCombiner);
+				StateCache.SetShadingRateImage(Resource);
 			}
 			else
 			{
-				StateCache.SetShadingRateImage(nullptr, EVRSRateCombiner::VRSRB_Passthrough);
+				StateCache.SetShadingRateImage(nullptr);
 			}
 		}
 		else
 		{
-			// Ensure this is set appropriate if image-based VRS not supported or not enabled.
-			StateCache.SetShadingRateImage(nullptr, EVRSRateCombiner::VRSRB_Passthrough);
+			StateCache.SetShadingRateImage(nullptr);
 		}
 	}
 #endif
@@ -1981,7 +1987,8 @@ void FD3D12CommandContext::SetDepthBounds(float MinDepth, float MaxDepth)
 void FD3D12CommandContext::RHISetShadingRate(EVRSShadingRate ShadingRate, EVRSRateCombiner Combiner)
 {
 #if PLATFORM_SUPPORTS_VARIABLE_RATE_SHADING
-	StateCache.SetShadingRate(ShadingRate, Combiner);
+	// Note - this will override per-material VRS opt-out, but FRHICommandSetShadingRate isn't called from anywhere
+	StateCache.SetShadingRate(ShadingRate, Combiner, VRSRB_Max);
 #endif
 }
 
diff --git a/Engine/Source/Runtime/D3D12RHI/Private/D3D12StateCachePrivate.h b/Engine/Source/Runtime/D3D12RHI/Private/D3D12StateCachePrivate.h
--- a/Engine/Source/Runtime/D3D12RHI/Private/D3D12StateCachePrivate.h
+++ b/Engine/Source/Runtime/D3D12RHI/Private/D3D12StateCachePrivate.h
@@ -697,32 +697,26 @@ public:
 		}
 	}
 
-	void SetShadingRate(EVRSShadingRate ShadingRate, EVRSRateCombiner Combiner)
+	void SetShadingRate(EVRSShadingRate ShadingRate, EVRSRateCombiner PerPrimitiveCombiner, EVRSRateCombiner ScreenSpaceCombiner)
 	{
-		if (PipelineState.Graphics.DrawShadingRate != ShadingRate || PipelineState.Graphics.Combiners[ED3D12VRSCombinerStages::PerPrimitive] != Combiner)
+		if (PipelineState.Graphics.DrawShadingRate != ShadingRate
+			|| PipelineState.Graphics.Combiners[ED3D12VRSCombinerStages::PerPrimitive] != PerPrimitiveCombiner
+			|| PipelineState.Graphics.Combiners[ED3D12VRSCombinerStages::ScreenSpace] != ScreenSpaceCombiner)
 		{
 			PipelineState.Graphics.DrawShadingRate = ShadingRate;
-			PipelineState.Graphics.Combiners[ED3D12VRSCombinerStages::PerPrimitive] = Combiner;
+			PipelineState.Graphics.Combiners[ED3D12VRSCombinerStages::PerPrimitive] = PerPrimitiveCombiner;
+			PipelineState.Graphics.Combiners[ED3D12VRSCombinerStages::ScreenSpace] = ScreenSpaceCombiner;
 			bNeedSetShadingRate = GRHISupportsPipelineVariableRateShading && GRHIVariableRateShadingEnabled;
 		}
-	}
+	}	
 
-	void SetShadingRateImage(FD3D12Resource* ShadingRateImage, EVRSRateCombiner Combiner)
+	void SetShadingRateImage(FD3D12Resource* ShadingRateImage)
 	{
 		if (PipelineState.Graphics.ShadingRateImage != ShadingRateImage)
 		{
 			PipelineState.Graphics.ShadingRateImage = ShadingRateImage;
 			bNeedSetShadingRateImage = GRHISupportsAttachmentVariableRateShading && GRHIAttachmentVariableRateShadingEnabled;
 		}
-
-		// If we aren't provided a ShadingRateImage, we should request the passthrough combiner
-		ensure((PipelineState.Graphics.ShadingRateImage != nullptr) || (Combiner == EVRSRateCombiner::VRSRB_Passthrough));
-
-		if (PipelineState.Graphics.Combiners[ED3D12VRSCombinerStages::ScreenSpace] != Combiner)
-		{
-			PipelineState.Graphics.Combiners[ED3D12VRSCombinerStages::ScreenSpace] = Combiner;
-			bNeedSetShadingRate = GRHISupportsAttachmentVariableRateShading && GRHIAttachmentVariableRateShadingEnabled;
-		}
 	}
 
 	void SetComputeBudget(EAsyncComputeBudget ComputeBudget)
diff --git a/Engine/Source/Runtime/Engine/Classes/Materials/Material.h b/Engine/Source/Runtime/Engine/Classes/Materials/Material.h
--- a/Engine/Source/Runtime/Engine/Classes/Materials/Material.h
+++ b/Engine/Source/Runtime/Engine/Classes/Materials/Material.h
@@ -892,6 +892,10 @@ public:
 	UPROPERTY(EditAnywhere, Category = Material, AdvancedDisplay)
 	TEnumAsByte<EMaterialShadingRate> ShadingRate;
 
+	/** Allows the use of variable rate shading when evaluating this material. This will only apply to the base/translucency pass. */
+	UPROPERTY(EditAnywhere, Category = Material, AdvancedDisplay)
+	uint8 bAllowVariableRateShading : 1;
+
 #if WITH_EDITORONLY_DATA
 	UPROPERTY()
 	int32 EditorX;
diff --git a/Engine/Source/Runtime/Engine/Private/Materials/Material.cpp b/Engine/Source/Runtime/Engine/Private/Materials/Material.cpp
--- a/Engine/Source/Runtime/Engine/Private/Materials/Material.cpp
+++ b/Engine/Source/Runtime/Engine/Private/Materials/Material.cpp
@@ -1018,6 +1018,8 @@ UMaterial::UMaterial(const FObjectInitializer& ObjectInitializer)
 	RefractionMethod = RM_None;
 	RefractionCoverageMode = RCM_CoverageAccountedFor;
 
+	bAllowVariableRateShading = true;
+
 #if WITH_EDITORONLY_DATA
 	MaterialGraph = NULL;
 #endif //WITH_EDITORONLY_DATA
diff --git a/Engine/Source/Runtime/Engine/Private/Materials/MaterialShared.cpp b/Engine/Source/Runtime/Engine/Private/Materials/MaterialShared.cpp
--- a/Engine/Source/Runtime/Engine/Private/Materials/MaterialShared.cpp
+++ b/Engine/Source/Runtime/Engine/Private/Materials/MaterialShared.cpp
@@ -1730,6 +1730,11 @@ bool FMaterialResource::HasPixelDepthOffsetConnected() const { return Material->
 bool FMaterialResource::HasMaterialAttributesConnected() const { return (Material->bUseMaterialAttributes && Material->GetEditorOnlyData()->MaterialAttributes.IsConnected()); }
 #endif
 EMaterialShadingRate FMaterialResource::GetShadingRate() const { return Material->ShadingRate; }
+bool FMaterialResource::IsVariableRateShadingAllowed() const {
+	// Automatically disable VRS on anything using pixel discard as coarse shading causes the whole block to get 
+	// discarded resulting in noticeable artifacts
+	return Material->bAllowVariableRateShading && !IsMasked();
+}
 FString FMaterialResource::GetBaseMaterialPathName() const { return Material->GetPathName(); }
 FString FMaterialResource::GetDebugName() const
 {
@@ -2773,7 +2778,8 @@ void FMaterial::SetupMaterialEnvironment(
 	}
 
 	if (FDataDrivenShaderPlatformInfo::GetSupportsVariableRateShading(Platform) && 
-		(GRHIAttachmentVariableRateShadingEnabled || GetShadingRate() != MSR_1x1))
+		(GRHIAttachmentVariableRateShadingEnabled || GetShadingRate() != MSR_1x1) &&
+		IsVariableRateShadingAllowed())
 	{
 		OutEnvironment.SetCompileArgument(TEXT("USING_VARIABLE_RATE_SHADING"), true);
 	}
diff --git a/Engine/Source/Runtime/Engine/Public/MaterialShared.h b/Engine/Source/Runtime/Engine/Public/MaterialShared.h
--- a/Engine/Source/Runtime/Engine/Public/MaterialShared.h
+++ b/Engine/Source/Runtime/Engine/Public/MaterialShared.h
@@ -2006,6 +2006,7 @@ public:
 	virtual bool HasRenderTracePhysicalMaterialOutputs() const { return false; }
 	virtual EMaterialShadingRate GetShadingRate() const { return MSR_1x1; }
 	virtual int32 GetNeuralProfileId() const { return INDEX_NONE; }
+	virtual bool IsVariableRateShadingAllowed() const { return true; }
 	/**
 	 * Should shaders compiled for this material be saved to disk?
 	 */
@@ -2642,6 +2643,7 @@ public:
 	ENGINE_API virtual bool IsDeferredDecal() const override;
 	ENGINE_API virtual bool IsVolumetricPrimitive() const override;
 	ENGINE_API virtual bool IsWireframe() const override;
+	ENGINE_API virtual bool IsVariableRateShadingAllowed() const override;
 	ENGINE_API virtual EMaterialShadingRate  GetShadingRate() const override;
 	ENGINE_API virtual bool IsUIMaterial() const override;
 	ENGINE_API virtual bool IsPostProcessMaterial() const override;
diff --git a/Engine/Source/Runtime/RHI/Public/RHIResources.h b/Engine/Source/Runtime/RHI/Public/RHIResources.h
--- a/Engine/Source/Runtime/RHI/Public/RHIResources.h
+++ b/Engine/Source/Runtime/RHI/Public/RHIResources.h
@@ -3729,6 +3729,7 @@ public:
 		, bDepthBounds(false)
 		, MultiViewCount(0)
 		, bHasFragmentDensityAttachment(false)
+		, bAllowVariableRateShading(false)
 		, ShadingRate(EVRSShadingRate::VRSSR_1x1)
 		, Flags(0)
 		, StatePrecachePSOHash(0)
@@ -3764,6 +3765,7 @@ public:
 		bool						bInDepthBounds,
 		uint8						InMultiViewCount,
 		bool						bInHasFragmentDensityAttachment,
+		bool						bInAllowVariableRateShading,
 		EVRSShadingRate				InShadingRate)
 		: BoundShaderState(InBoundShaderState)
 		, BlendState(InBlendState)
@@ -3788,6 +3790,7 @@ public:
 		, bDepthBounds(bInDepthBounds)
 		, MultiViewCount(InMultiViewCount)
 		, bHasFragmentDensityAttachment(bInHasFragmentDensityAttachment)
+		, bAllowVariableRateShading(bInAllowVariableRateShading)
 		, ShadingRate(InShadingRate)
 		, Flags(InFlags)
 		, StatePrecachePSOHash(0)
@@ -3810,6 +3813,7 @@ public:
 			bDepthBounds != rhs.bDepthBounds ||
 			MultiViewCount != rhs.MultiViewCount ||
 			ShadingRate != rhs.ShadingRate ||
+			bAllowVariableRateShading != rhs.bAllowVariableRateShading ||
 			bHasFragmentDensityAttachment != rhs.bHasFragmentDensityAttachment ||
 			RenderTargetsEnabled != rhs.RenderTargetsEnabled ||
 			RenderTargetFormats != rhs.RenderTargetFormats || 
@@ -3904,6 +3908,7 @@ public:
 	bool							bDepthBounds;
 	uint8							MultiViewCount;
 	bool							bHasFragmentDensityAttachment;
+	bool							bAllowVariableRateShading;
 	EVRSShadingRate					ShadingRate;
 	
 	// Note: these flags do NOT affect compilation of this PSO.
diff --git a/Engine/Source/Runtime/Renderer/Public/MeshPassProcessor.h b/Engine/Source/Runtime/Renderer/Public/MeshPassProcessor.h
--- a/Engine/Source/Runtime/Renderer/Public/MeshPassProcessor.h
+++ b/Engine/Source/Runtime/Renderer/Public/MeshPassProcessor.h
@@ -319,6 +319,7 @@ public:
 		, DepthStencilState(InMinimalState.DepthStencilState)
 		, ImmutableSamplerState(InMinimalState.ImmutableSamplerState)
 		, bDepthBounds(InMinimalState.bDepthBounds)
+		, bAllowVariableRateShading(InMinimalState.bAllowVariableRateShading)
 		, DrawShadingRate(InMinimalState.DrawShadingRate)
 		, PrimitiveType(InMinimalState.PrimitiveType)
 		, StatePrecachePSOHash(InMinimalState.StatePrecachePSOHash)
@@ -357,6 +358,7 @@ public:
 			, bDepthBounds
 			, MultiViewCount
 			, bHasFragmentDensityAttachment
+			, bAllowVariableRateShading
 			, DrawShadingRate
 		);
 
@@ -387,6 +389,7 @@ public:
 			bDepthBounds != rhs.bDepthBounds ||
 			MultiViewCount != rhs.MultiViewCount ||
 			bHasFragmentDensityAttachment != rhs.bHasFragmentDensityAttachment ||
+			bAllowVariableRateShading != rhs.bAllowVariableRateShading ||
 			DrawShadingRate != rhs.DrawShadingRate ||
 			PrimitiveType != rhs.PrimitiveType)
 		{
@@ -449,6 +452,7 @@ public:
 			COMPARE_FIELD(bDepthBounds)
 			COMPARE_FIELD(MultiViewCount)
 			COMPARE_FIELD(bHasFragmentDensityAttachment)
+			COMPARE_FIELD(bAllowVariableRateShading)
 			COMPARE_FIELD(DrawShadingRate)
 			COMPARE_FIELD(PrimitiveType)
 		COMPARE_FIELD_END;
@@ -475,6 +479,7 @@ public:
 			COMPARE_FIELD(bDepthBounds)
 			COMPARE_FIELD(MultiViewCount)
 			COMPARE_FIELD(bHasFragmentDensityAttachment)
+			COMPARE_FIELD(bAllowVariableRateShading)
 			COMPARE_FIELD(DrawShadingRate)
 			COMPARE_FIELD(PrimitiveType)
 			COMPARE_FIELD_END;
@@ -500,6 +505,7 @@ public:
 	bool							bDepthBounds = false;
 	uint8							MultiViewCount = 0;
 	bool							bHasFragmentDensityAttachment = false;
+	bool							bAllowVariableRateShading = false;
 	EVRSShadingRate					DrawShadingRate  = EVRSShadingRate::VRSSR_1x1;
 
 	EPrimitiveType					PrimitiveType;
diff --git a/Engine/Source/Runtime/Renderer/Public/MeshPassProcessor.inl b/Engine/Source/Runtime/Renderer/Public/MeshPassProcessor.inl
--- a/Engine/Source/Runtime/Renderer/Public/MeshPassProcessor.inl
+++ b/Engine/Source/Runtime/Renderer/Public/MeshPassProcessor.inl
@@ -102,6 +102,7 @@ void FMeshPassProcessor::BuildMeshDrawCommands(
 	PipelineState.BlendState = DrawRenderState.GetBlendState();
 	PipelineState.DepthStencilState = DrawRenderState.GetDepthStencilState();
 	PipelineState.DrawShadingRate = GetShadingRateFromMaterial(MaterialResource.GetShadingRate());
+	PipelineState.bAllowVariableRateShading = MaterialResource.IsVariableRateShadingAllowed();
 
 	// PSO Precache hash only needed when PSO precaching is enabled
 	if (PipelineStateCache::IsPSOPrecachingEnabled())
