Initial implementation of explicit D3D12 reserved buffer commit / resize

* First pass implementation only supports growing
* Commit is implemented as a resource transition to allow asynchronous memory mapping operations in the future

diff --git a/Engine/Plugins/Tests/RHITests/Source/RHITests/Private/RHIReservedResourceTests.cpp b/Engine/Plugins/Tests/RHITests/Source/RHITests/Private/RHIReservedResourceTests.cpp
--- a/Engine/Plugins/Tests/RHITests/Source/RHITests/Private/RHIReservedResourceTests.cpp
+++ b/Engine/Plugins/Tests/RHITests/Source/RHITests/Private/RHIReservedResourceTests.cpp
@@ -1,6 +1,7 @@
 // Copyright Epic Games, Inc. All Rights Reserved.
 
 #include "RHIReservedResourceTests.h"
+#include "RHIBufferTests.h" // for VerifyBufferContents
 #include "CommonRenderResources.h"
 #include "RenderCaptureInterface.h"
 #include "RHIStaticStates.h"
@@ -117,3 +118,85 @@ bool FRHIReservedResourceTests::Test_ReservedResource_CreateBuffer(FRHICommandLi
 	return true;
 }
 
+static void CommitBuffer(FRHICommandListImmediate& RHICmdList, FRHIBuffer* Buffer, uint64 CommitSize, ERHIAccess StateBefore, ERHIAccess StateAfter)
+{
+	FRHITransitionInfo TransitionInfo(Buffer, StateBefore, StateAfter, FRHICommitResourceInfo(CommitSize));
+
+	TArrayView<const FRHITransitionInfo> TransitionInfos = MakeArrayView(&TransitionInfo, 1);
+
+	FRHITransitionCreateInfo CreateInfo(
+		ERHIPipeline::Graphics, ERHIPipeline::Graphics,
+		ERHITransitionCreateFlags::None, TransitionInfos);
+
+	const FRHITransition* Transition = RHICreateTransition(CreateInfo);
+
+	RHICmdList.BeginTransition(Transition);
+	RHICmdList.EndTransition(Transition);
+}
+
+bool FRHIReservedResourceTests::Test_ReservedResource_CommitBuffer(FRHICommandListImmediate& RHICmdList)
+{
+	if (!GRHIGlobals.ReservedResources.Supported)
+	{
+		return true;
+	}
+
+	const int32 TileSizeInBytes = GRHIGlobals.ReservedResources.TileSizeInBytes;
+	const int32 BufferSizeInBytes = TileSizeInBytes * 128;
+
+	FRHIResourceCreateInfo CreateInfo(TEXT("TestReservedBufferExplicitCommit"));
+
+	FBufferRHIRef Buffer = RHICmdList.CreateBuffer(BufferSizeInBytes,
+		BUF_ReservedResource | BUF_UnorderedAccess | BUF_ShaderResource | BUF_SourceCopy,
+		4, ERHIAccess::UAVCompute, CreateInfo);
+
+	FUnorderedAccessViewRHIRef BufferUAV = RHICmdList.CreateUnorderedAccessView(Buffer,
+		FRHIViewDesc::CreateBufferUAV()
+		.SetType(FRHIViewDesc::EBufferType::Typed)
+		.SetFormat(PF_R32_UINT));
+
+	// Commit half of the resource, leaving the tail unmapped. 
+	// The RHI follows D3D12 Tier 2 Reserved Resource semantics:
+	// - Unmapped page writes are discarded
+	// - Unmapped page reads return 0
+
+	const int32 CommitSizeInBytes = BufferSizeInBytes / 2;
+	CommitBuffer(RHICmdList, Buffer, CommitSizeInBytes, ERHIAccess::UAVCompute, ERHIAccess::UAVCompute);
+
+	RHICmdList.ClearUAVUint(BufferUAV, FUintVector4(~0u));
+
+	RHICmdList.Transition(FRHITransitionInfo(Buffer, ERHIAccess::UAVCompute, ERHIAccess::CopySrc));
+	
+	FRHIBuffer* Buffers[] = { Buffer.GetReference() };
+	bool bSucceeded = FRHIBufferTests::VerifyBufferContents(TEXT("Test_ReservedResource_CommitBuffer"), RHICmdList, Buffers, 
+		[BufferSizeInBytes, CommitSizeInBytes](int32 BufferIndex, void* Ptr, uint32 NumBytes)
+		{
+			uint64 ExpectedCommittedValue = ~0ull;
+
+			uint32 CommittedSizeInElements = CommitSizeInBytes / sizeof(ExpectedCommittedValue);
+			uint32 TotalSizeInElements = BufferSizeInBytes / sizeof(ExpectedCommittedValue);
+			const uint64* BufferData = reinterpret_cast<const uint64*>(Ptr);
+
+			for (uint32 i = 0; i < CommittedSizeInElements; ++i)
+			{
+				if (BufferData[i] != ExpectedCommittedValue)
+				{
+					return false;
+				}
+			}
+
+			uint64 ExpectedTailValue = 0;
+			for (uint32 i = CommittedSizeInElements; i < TotalSizeInElements; ++i)
+			{
+				if (BufferData[i] != ExpectedTailValue)
+				{
+					return false;
+				}
+			}
+
+			return true;
+		});
+
+	return bSucceeded;
+}
+
diff --git a/Engine/Plugins/Tests/RHITests/Source/RHITests/Private/RHITests.spec.cpp b/Engine/Plugins/Tests/RHITests/Source/RHITests/Private/RHITests.spec.cpp
--- a/Engine/Plugins/Tests/RHITests/Source/RHITests/Private/RHITests.spec.cpp
+++ b/Engine/Plugins/Tests/RHITests/Source/RHITests/Private/RHITests.spec.cpp
@@ -132,5 +132,11 @@ void FAutomationRHITest::Define()
 			bool bResult = RunOnRenderThreadSynchronous(FRHIReservedResourceTests::Test_ReservedResource_CreateBuffer);
 			TestEqual("Create Reserved Buffer failed", bResult, 1);
 		});
+
+		It("Commit Reserved Buffer", [this]()
+		{
+			bool bResult = RunOnRenderThreadSynchronous(FRHIReservedResourceTests::Test_ReservedResource_CommitBuffer);
+			TestEqual("Commit Reserved Buffer failed", bResult, 1);
+		});
 	});
 }
diff --git a/Engine/Plugins/Tests/RHITests/Source/RHITests/Private/RHITestsModule.cpp b/Engine/Plugins/Tests/RHITests/Source/RHITests/Private/RHITestsModule.cpp
--- a/Engine/Plugins/Tests/RHITests/Source/RHITests/Private/RHITestsModule.cpp
+++ b/Engine/Plugins/Tests/RHITests/Source/RHITests/Private/RHITestsModule.cpp
@@ -21,6 +21,7 @@ static bool RunTests_RenderThread(FRHICommandListImmediate& RHICmdList)
 	// Reserved resources
 	// ------------------------------------------------
 	{
+		RUN_TEST(FRHIReservedResourceTests::Test_ReservedResource_CommitBuffer(RHICmdList));
 		RUN_TEST(FRHIReservedResourceTests::Test_ReservedResource_CreateBuffer(RHICmdList));
 		RUN_TEST(FRHIReservedResourceTests::Test_ReservedResource_CreateTexture(RHICmdList));
 	}
diff --git a/Engine/Plugins/Tests/RHITests/Source/RHITests/Public/RHIReservedResourceTests.h b/Engine/Plugins/Tests/RHITests/Source/RHITests/Public/RHIReservedResourceTests.h
--- a/Engine/Plugins/Tests/RHITests/Source/RHITests/Public/RHIReservedResourceTests.h
+++ b/Engine/Plugins/Tests/RHITests/Source/RHITests/Public/RHIReservedResourceTests.h
@@ -10,5 +10,6 @@ public:
 
 	static bool Test_ReservedResource_CreateTexture(FRHICommandListImmediate& RHICmdList);
 	static bool Test_ReservedResource_CreateBuffer(FRHICommandListImmediate& RHICmdList);
+	static bool Test_ReservedResource_CommitBuffer(FRHICommandListImmediate& RHICmdList);
 };
 
diff --git a/Engine/Source/Runtime/D3D12RHI/Private/D3D12Allocation.cpp b/Engine/Source/Runtime/D3D12RHI/Private/D3D12Allocation.cpp
--- a/Engine/Source/Runtime/D3D12RHI/Private/D3D12Allocation.cpp
+++ b/Engine/Source/Runtime/D3D12RHI/Private/D3D12Allocation.cpp
@@ -1539,7 +1539,10 @@ void FD3D12DefaultBufferAllocator::AllocDefaultResource(D3D12_HEAP_TYPE InHeapTy
 
 		if (EnumHasAnyFlags(InBufferUsage, BUF_ImmediateCommit))
 		{
-			NewResource->CommitReservedResource();
+			// NOTE: Accessing the queue from this thread is OK, as D3D12 runtime acquires a lock around all command queue APIs.
+			// https://microsoft.github.io/DirectX-Specs/d3d/CPUEfficiency.html#threading
+			FD3D12Queue& Queue = GetParentDevice()->GetQueue(ED3D12QueueType::Direct);
+			NewResource->CommitReservedResource(Queue.D3DCommandQueue, UINT64_MAX /*commit entire resource*/);
 		}
 
 		return;
diff --git a/Engine/Source/Runtime/D3D12RHI/Private/D3D12Buffer.cpp b/Engine/Source/Runtime/D3D12RHI/Private/D3D12Buffer.cpp
--- a/Engine/Source/Runtime/D3D12RHI/Private/D3D12Buffer.cpp
+++ b/Engine/Source/Runtime/D3D12RHI/Private/D3D12Buffer.cpp
@@ -557,10 +557,13 @@ FD3D12Buffer* FD3D12DynamicRHI::CreateD3D12Buffer(class FRHICommandListBase* RHI
 		? ED3D12ResourceStateMode::SingleState 
 		: ED3D12ResourceStateMode::Default;
 
+	const bool bHasInitialData = CreateInfo.ResourceArray != nullptr;
+
 	const bool bIsDynamic = EnumHasAnyFlags(BufferDesc.Usage, BUF_AnyDynamic);
 
 	if (EnumHasAnyFlags(BufferDesc.Usage, BUF_ReservedResource))
 	{
+		checkf(!bHasInitialData || EnumHasAllFlags(BufferDesc.Usage, BUF_ImmediateCommit), TEXT("Reserved resources with initial data must be committed at creation"));
 		checkf(!bIsDynamic, TEXT("Reserved resources may not be dynamic"));
 		checkf(!ResourceAllocator, TEXT("Reserved resources may not use a custom resource allocator"));
 	}
@@ -577,7 +580,6 @@ FD3D12Buffer* FD3D12DynamicRHI::CreateD3D12Buffer(class FRHICommandListBase* RHI
 
 	// Setup the state at which the resource needs to be created - copy dest only supported for placed resources
 	D3D12_RESOURCE_STATES CreateState = (CreateInfo.ResourceArray && bSupportResourceStateTracking) ? D3D12_RESOURCE_STATE_COPY_DEST : DesiredState;
-	bool bHasInitialData = CreateInfo.ResourceArray != nullptr;
 
 	FD3D12Buffer* Buffer = GetAdapter().CreateRHIBuffer(Desc, Alignment, BufferDesc, StateMode, CreateState, bHasInitialData, CreateInfo.GPUMask, ResourceAllocator, CreateInfo.DebugName, CreateInfo.OwnerName, TraceClassName);
 	check(Buffer->ResourceLocation.IsValid());
diff --git a/Engine/Source/Runtime/D3D12RHI/Private/D3D12CommandContext.cpp b/Engine/Source/Runtime/D3D12RHI/Private/D3D12CommandContext.cpp
--- a/Engine/Source/Runtime/D3D12RHI/Private/D3D12CommandContext.cpp
+++ b/Engine/Source/Runtime/D3D12RHI/Private/D3D12CommandContext.cpp
@@ -284,6 +284,20 @@ FD3D12QueryLocation FD3D12ContextCommon::InsertTimestamp(ED3D12Units Units, uint
 	return Location;
 }
 
+void FD3D12ContextCommon::SetReservedBufferCommitSize(FD3D12Buffer* Buffer, uint64 CommitSizeInBytes)
+{
+	if (IsOpen())
+	{
+		CloseCommandList();
+	}
+
+	FD3D12CommitReservedResourceDesc CommitDesc;
+	CommitDesc.Resource = Buffer->GetResource();
+	CommitDesc.CommitSizeInBytes = CommitSizeInBytes;
+
+	GetPayload(EPhase::UpdateReservedResources)->ReservedResourcesToCommit.Add(CommitDesc);
+}
+
 void FD3D12ContextCommon::OpenCommandList()
 {
 	LLM_SCOPE_BYNAME(TEXT("RHIMisc/OpenCommandList"));
diff --git a/Engine/Source/Runtime/D3D12RHI/Private/D3D12CommandContext.h b/Engine/Source/Runtime/D3D12RHI/Private/D3D12CommandContext.h
--- a/Engine/Source/Runtime/D3D12RHI/Private/D3D12CommandContext.h
+++ b/Engine/Source/Runtime/D3D12RHI/Private/D3D12CommandContext.h
@@ -179,6 +179,11 @@ public:
 	// Allocates a query of the specified type, returning its location.
 	FD3D12QueryLocation AllocateQuery(ED3D12QueryType Type, void* Target);
 
+	// Resizes physical memory allocation for a buffer. Allocates new backing heaps as necessary.
+	// Causes the command list to be split, as reserved resource update operations are performed on the D3D12 queue.
+	// The actual work is deferred via FD3D12Payload.
+	void SetReservedBufferCommitSize(FD3D12Buffer* Buffer, uint64 CommitSizeInBytes);
+
 	// Complete recording of the current command list set, and appends the resulting
 	// payloads to the given array. Resets the context so new commands can be recorded.
 	void Finalize(TArray<FD3D12Payload*>& OutPayloads);
@@ -269,6 +274,7 @@ protected:
 	enum class EPhase
 	{
 		Wait,
+		UpdateReservedResources,
 		Execute,
 		Signal
 	} CurrentPhase = EPhase::Wait;
@@ -750,6 +756,7 @@ private:
 	void HandleResourceTransitions       (const struct FD3D12TransitionData* TransitionData, bool& bUAVBarrier);
 	void HandleTransientAliasing         (const struct FD3D12TransitionData* TransitionData);
 	void HandleResourceDiscardTransitions(const struct FD3D12TransitionData* TransitionData, TArray<struct FD3D12DiscardResource>& ResourcesToDiscard);
+	void HandleReservedResourceCommits   (const struct FD3D12TransitionData* TransitionData);
 
 	TArray<FRHIUniformBuffer*> StaticUniformBuffers;
 };
diff --git a/Engine/Source/Runtime/D3D12RHI/Private/D3D12Commands.cpp b/Engine/Source/Runtime/D3D12RHI/Private/D3D12Commands.cpp
--- a/Engine/Source/Runtime/D3D12RHI/Private/D3D12Commands.cpp
+++ b/Engine/Source/Runtime/D3D12RHI/Private/D3D12Commands.cpp
@@ -287,6 +287,25 @@ struct FD3D12DiscardResource
 #endif // #if PLATFORM_REQUIRES_TYPELESS_RESOURCE_DISCARD_WORKAROUND
 };
 
+void FD3D12CommandContext::HandleReservedResourceCommits(const struct FD3D12TransitionData* TransitionData)
+{
+	for (const FRHITransitionInfo& Info : TransitionData->TransitionInfos)
+	{
+		if (const FRHICommitResourceInfo* CommitInfo = Info.CommitInfo.GetPtrOrNull())
+		{
+			if (Info.Type == FRHITransitionInfo::EType::Buffer)
+			{
+				FD3D12Buffer* Buffer = RetrieveObject<FD3D12Buffer>(Info.Buffer);
+				SetReservedBufferCommitSize(Buffer, CommitInfo->SizeInBytes);
+			}
+			else
+			{
+				checkNoEntry();
+			}
+		}
+	}
+}
+
 void FD3D12CommandContext::HandleResourceDiscardTransitions(
 	const FD3D12TransitionData* TransitionData,
 	TArray<FD3D12DiscardResource>& ResourcesToDiscard)
@@ -625,6 +644,13 @@ void FD3D12CommandContext::RHIEndTransitions(TArrayView<const FRHITransition*> T
 	const bool bShowTransitionEvents = CVarShowTransitions->GetInt() != 0;
 	SCOPED_RHI_CONDITIONAL_DRAW_EVENTF(*this, RHIEndTransitions, bShowTransitionEvents, TEXT("RHIEndTransitions"));
 
+	// Update reserved resource memory mapping
+	for (const FRHITransition* Transition : Transitions)
+	{
+		const FD3D12TransitionData* Data = Transition->GetPrivateData<FD3D12TransitionData>();
+		HandleReservedResourceCommits(Data);
+	}
+
 	for (const FRHITransition* Transition : Transitions)
 	{
 		const FD3D12TransitionData* Data = Transition->GetPrivateData<FD3D12TransitionData>();
diff --git a/Engine/Source/Runtime/D3D12RHI/Private/D3D12Resources.cpp b/Engine/Source/Runtime/D3D12RHI/Private/D3D12Resources.cpp
--- a/Engine/Source/Runtime/D3D12RHI/Private/D3D12Resources.cpp
+++ b/Engine/Source/Runtime/D3D12RHI/Private/D3D12Resources.cpp
@@ -154,13 +154,15 @@ FD3D12Resource::~FD3D12Resource()
 	}
 }
 
-void FD3D12Resource::CommitReservedResource()
+void FD3D12Resource::CommitReservedResource(ID3D12CommandQueue* D3DCommandQueue, uint64 RequiredCommitSizeInBytes)
 {
 	TRACE_CPUPROFILER_EVENT_SCOPE(CommitReservedResource);
 
+	static constexpr uint64 TileSizeInBytes = GRHIGlobals.ReservedResources.TileSizeInBytes;
+	static_assert(TileSizeInBytes == 65536, "Reserved resource tiles are expected to always be 64KB");
+
 	check(Desc.bReservedResource);
 	check(ReservedResourceData.IsValid());
-	checkf(ReservedResourceData->BackingHeaps.IsEmpty(), TEXT("Reserved resource is already committed"));
 	checkf(Desc.Dimension == D3D12_RESOURCE_DIMENSION_BUFFER
 		|| Desc.Dimension == D3D12_RESOURCE_DIMENSION_TEXTURE2D, 
 		TEXT("CommitReservedResource is currently only implemented for 2D textures and buffers"));
@@ -184,10 +186,11 @@ void FD3D12Resource::CommitReservedResource()
 
 	D3DDevice->GetResourceTiling(GetResource(), &D3DResourceNumTiles, &PackedMipDesc, &TileShape, &NumSubresourceTilings, FirstSubresource, &SubresourceTiling);
 
-	static constexpr uint64 TileSizeInBytes = GRHIGlobals.ReservedResources.TileSizeInBytes;
-	static_assert(TileSizeInBytes == 65536, "Reserved resource tiles are expected to always be 64KB");
-
 	const uint64 TotalSize = D3DResourceNumTiles * TileSizeInBytes;
+
+	RequiredCommitSizeInBytes = FMath::Min<uint64>(RequiredCommitSizeInBytes, TotalSize);
+	RequiredCommitSizeInBytes = AlignArbitrary(RequiredCommitSizeInBytes, TileSizeInBytes);
+
 	const uint64 MaxHeapSize = uint64(CVarD3D12ReservedResourceHeapSizeMB.GetValueOnAnyThread()) * 1024 * 1024;
 	const uint64 NumHeaps = FMath::DivideAndRoundUp(TotalSize, MaxHeapSize);
 
@@ -196,11 +199,6 @@ void FD3D12Resource::CommitReservedResource()
 
 	const uint32 MaxTilesPerHeap = uint32(MaxHeapSize / TileSizeInBytes);
 
-	// NOTE: Accessing the queue from this thread is OK, as D3D12 runtime acquires a lock around all command queue APIs.
-	// https://microsoft.github.io/DirectX-Specs/d3d/CPUEfficiency.html#threading
-	FD3D12Queue& Queue = GetParentDevice()->GetQueue(ED3D12QueueType::Direct);
-	ID3D12CommandQueue* D3DCommandQueue = Queue.D3DCommandQueue.GetReference();
-
 	const D3D12_TILE_MAPPING_FLAGS MappingFlags = D3D12_TILE_MAPPING_FLAG_NONE;
 
 	// Set high residency priority based on the same heuristics as D3D12 committed resources,
@@ -240,22 +238,35 @@ void FD3D12Resource::CommitReservedResource()
 		TEXT("D3D resource size in tiles: %d, computed size in tiles: %d"),
 		D3DResourceNumTiles, NumTotalTiles);
 
-	uint32 NumMappedTiles = 0;
+	checkf(RequiredCommitSizeInBytes % TileSizeInBytes == 0,
+		TEXT("Reserved resources memory is expected to be committed at tile granularity"));
+
+	checkf(ReservedResourceData->CommittedSizeInBytes % TileSizeInBytes == 0,
+		TEXT("Reserved resources memory is expected to be committed at tile granularity"));
+
+	checkf(RequiredCommitSizeInBytes >= ReservedResourceData->CommittedSizeInBytes,
+		TEXT("Shrinking committed resources is not yet implemented")); // #yuriy_todo
 
-	const D3D12_HEAP_FLAGS HeapFlags = Desc.Dimension == D3D12_RESOURCE_DIMENSION_BUFFER
-		? D3D12_HEAP_FLAG_ALLOW_ONLY_BUFFERS
-		: D3D12_HEAP_FLAG_ALLOW_ONLY_NON_RT_DS_TEXTURES;
+	const uint32 NumRequiredCommitTiles = RequiredCommitSizeInBytes / TileSizeInBytes;
+	checkf(NumRequiredCommitTiles <= NumTotalTiles,
+		TEXT("Shrinking committed resources is not yet implemented")); // #yuriy_todo
 
-	static_assert((D3D12_HEAP_FLAG_DENY_BUFFERS | D3D12_HEAP_FLAG_DENY_RT_DS_TEXTURES) == D3D12_HEAP_FLAG_ALLOW_ONLY_NON_RT_DS_TEXTURES);
+	uint32 NumMappedTiles = ReservedResourceData->CommittedSizeInBytes / TileSizeInBytes;
 
-	while (NumMappedTiles < NumTotalTiles)
+	while (NumMappedTiles < NumRequiredCommitTiles)
 	{
-		const uint32 NumRemainingTiles = NumTotalTiles - NumMappedTiles;
+		const uint32 NumRemainingTiles = NumRequiredCommitTiles - NumMappedTiles;
 
 		D3D12_TILE_REGION_SIZE RegionSize = {};
 		RegionSize.UseBox = false;
 		RegionSize.NumTiles = FMath::Min(MaxTilesPerHeap, NumRemainingTiles);
 
+		const D3D12_HEAP_FLAGS HeapFlags = Desc.Dimension == D3D12_RESOURCE_DIMENSION_BUFFER
+			? D3D12_HEAP_FLAG_ALLOW_ONLY_BUFFERS
+			: D3D12_HEAP_FLAG_ALLOW_ONLY_NON_RT_DS_TEXTURES;
+
+		static_assert((D3D12_HEAP_FLAG_DENY_BUFFERS | D3D12_HEAP_FLAG_DENY_RT_DS_TEXTURES) == D3D12_HEAP_FLAG_ALLOW_ONLY_NON_RT_DS_TEXTURES);
+
 		const uint32 ThisHeapSize = RegionSize.NumTiles * TileSizeInBytes;
 		D3D12_HEAP_DESC NewHeapDesc = {};
 		NewHeapDesc.Alignment = D3D12_DEFAULT_RESOURCE_PLACEMENT_ALIGNMENT;
@@ -322,7 +333,9 @@ void FD3D12Resource::CommitReservedResource()
 		Heaps.Add(MoveTemp(NewHeap));
 	}
 
-	checkf(NumMappedTiles == D3DResourceNumTiles,
+	ReservedResourceData->CommittedSizeInBytes = NumMappedTiles * TileSizeInBytes;
+
+	checkf(NumMappedTiles == NumRequiredCommitTiles,
 		TEXT("Reserved resource was not fully processed while committing physical memory. Expected to process tiles: %d, actually processed: %d"),
 		D3DResourceNumTiles, NumMappedTiles);
 }
diff --git a/Engine/Source/Runtime/D3D12RHI/Private/D3D12Resources.h b/Engine/Source/Runtime/D3D12RHI/Private/D3D12Resources.h
--- a/Engine/Source/Runtime/D3D12RHI/Private/D3D12Resources.h
+++ b/Engine/Source/Runtime/D3D12RHI/Private/D3D12Resources.h
@@ -186,6 +186,7 @@ private:
 	struct FD3D12ReservedResourceData
 	{
 		TArray<TRefCountPtr<ID3D12Heap>> BackingHeaps;
+		uint64 CommittedSizeInBytes = 0;
 	};
 	TUniquePtr<FD3D12ReservedResourceData> ReservedResourceData;
 
@@ -382,7 +383,9 @@ public:
 		const uint32 bReadBackResource : 1;
 	};
 
-	void CommitReservedResource();
+	// Note: RequiredCommitSizeInBytes is clamped to the maximum size of the resource.
+	// Use UINT64_MAX to commit the entire resource.
+	void CommitReservedResource(ID3D12CommandQueue* D3DCommandQueue, uint64 RequiredCommitSizeInBytes);
 
 private:
 	void InitalizeResourceState(D3D12_RESOURCE_STATES InInitialState, ED3D12ResourceStateMode InResourceStateMode, D3D12_RESOURCE_STATES InDefaultState)
diff --git a/Engine/Source/Runtime/D3D12RHI/Private/D3D12Submission.cpp b/Engine/Source/Runtime/D3D12RHI/Private/D3D12Submission.cpp
--- a/Engine/Source/Runtime/D3D12RHI/Private/D3D12Submission.cpp
+++ b/Engine/Source/Runtime/D3D12RHI/Private/D3D12Submission.cpp
@@ -271,6 +271,7 @@ void FD3D12DynamicRHI::SubmitCommands(TConstArrayView<FD3D12FinalizedCommands*>
 					}
 					FD3D12Payload* MergedPayload = MergedPayloadPerGPU[GPUIndex];
 
+					MergedPayload->ReservedResourcesToCommit.Append(Payload->ReservedResourcesToCommit);
 					MergedPayload->CommandListsToExecute.Append(Payload->CommandListsToExecute);
 					MergedPayload->SyncPointsToSignal.Append(Payload->SyncPointsToSignal);
 					MergedPayload->AllocatorsToRelease.Append(Payload->AllocatorsToRelease);
@@ -779,6 +780,11 @@ uint64 FD3D12Queue::ExecutePayload()
 
 	PayloadToSubmit->PreExecute();
 
+	for (const FD3D12CommitReservedResourceDesc& CommitDesc : PayloadToSubmit->ReservedResourcesToCommit)
+	{
+		CommitDesc.Resource->CommitReservedResource(D3DCommandQueue, CommitDesc.CommitSizeInBytes);
+	}
+
 	if (const int32 NumCommandLists = PayloadToSubmit->CommandListsToExecute.Num())
 	{
 		// Build SOA layout needed to call ExecuteCommandLists().
diff --git a/Engine/Source/Runtime/D3D12RHI/Private/D3D12Submission.h b/Engine/Source/Runtime/D3D12RHI/Private/D3D12Submission.h
--- a/Engine/Source/Runtime/D3D12RHI/Private/D3D12Submission.h
+++ b/Engine/Source/Runtime/D3D12RHI/Private/D3D12Submission.h
@@ -14,6 +14,8 @@ class FD3D12DynamicRHI;
 class FD3D12QueryHeap;
 class FD3D12Queue;
 class FD3D12Timing;
+class FD3D12Buffer;
+class FD3D12Resource;
 
 class FD3D12SyncPoint;
 using FD3D12SyncPointRef = TRefCountPtr<FD3D12SyncPoint>;
@@ -206,6 +208,12 @@ struct FD3D12QueryRange
 	inline bool IsFull() const;
 };
 
+struct FD3D12CommitReservedResourceDesc
+{
+	FD3D12Resource* Resource = nullptr;
+	uint64 CommitSizeInBytes = 0;
+};
+
 // A single unit of work (specific to a single GPU node and queue type) to be processed by the submission thread.
 struct FD3D12PayloadBase
 {
@@ -242,6 +250,9 @@ struct FD3D12PayloadBase
 	// Wait
 	TArray<FManualFence> FencesToWait;
 
+	// UpdateReservedResources
+	TArray<FD3D12CommitReservedResourceDesc> ReservedResourcesToCommit;
+
 	// Execute
 	TArray<FD3D12CommandList*> CommandListsToExecute;
 
diff --git a/Engine/Source/Runtime/D3D12RHI/Private/D3D12Texture.cpp b/Engine/Source/Runtime/D3D12RHI/Private/D3D12Texture.cpp
--- a/Engine/Source/Runtime/D3D12RHI/Private/D3D12Texture.cpp
+++ b/Engine/Source/Runtime/D3D12RHI/Private/D3D12Texture.cpp
@@ -771,7 +771,9 @@ void SafeCreateTexture2D(FD3D12Device* pDevice,
 
 				if (EnumHasAllFlags(Flags, TexCreate_ImmediateCommit))
 				{
-					Resource->CommitReservedResource();
+					// NOTE: Accessing the queue from this thread is OK, as D3D12 runtime acquires a lock around all command queue APIs.
+					// https://microsoft.github.io/DirectX-Specs/d3d/CPUEfficiency.html#threading
+					Resource->CommitReservedResource(pDevice->GetQueue(ED3D12QueueType::Direct).D3DCommandQueue, UINT64_MAX /*commit entire resource*/);
 				}
 			}
 			else
diff --git a/Engine/Source/Runtime/RHI/Private/RHIValidation.cpp b/Engine/Source/Runtime/RHI/Private/RHIValidation.cpp
--- a/Engine/Source/Runtime/RHI/Private/RHIValidation.cpp
+++ b/Engine/Source/Runtime/RHI/Private/RHIValidation.cpp
@@ -622,6 +622,25 @@ void FValidationRHI::RHICreateTransition(FRHITransition* Transition, const FRHIT
 		checkf(Info.AccessAfter != ERHIAccess::Unknown, TEXT("FRHITransitionInfo::AccessAfter cannot be Unknown when creating a resource transition."));
 		checkf(Info.Type != FRHITransitionInfo::EType::Unknown, TEXT("FRHITransitionInfo::Type cannot be Unknown when creating a resource transition."));
 
+		if (const FRHICommitResourceInfo* CommitInfo = Info.CommitInfo.GetPtrOrNull())
+		{
+			RHI_VALIDATION_CHECK((SrcPipelines == ERHIPipeline::Graphics && DstPipelines == ERHIPipeline::Graphics),
+				TEXT("Reserved resource commit operations are only supported on the graphics pipeline and must not cross pipeline boundary."));
+
+			if (Info.Type == FRHITransitionInfo::EType::Buffer)
+			{
+				const FRHIBuffer* Buffer = Info.Buffer;
+				const EBufferUsageFlags BufferUsage = Buffer->GetUsage();
+				const uint32 BufferSize = Buffer->GetSize();
+				RHI_VALIDATION_CHECK(EnumHasAllFlags(BufferUsage, BUF_ReservedResource), TEXT("Commit transitions can only be used with reserved resources."));
+				RHI_VALIDATION_CHECK(CommitInfo->SizeInBytes <= BufferSize, TEXT("Buffer commit size request must not be larger than the size of the buffer itself, as virtual memory allocation cannot be resized."));
+			}
+			else
+			{
+				RHI_VALIDATION_CHECK(false, TEXT("Reserved resource commit is only supported for buffers"));
+			}
+		}
+
 		FResourceIdentity Identity;
 
 		switch (Info.Type)
diff --git a/Engine/Source/Runtime/RHI/Public/RHITransition.h b/Engine/Source/Runtime/RHI/Public/RHITransition.h
--- a/Engine/Source/Runtime/RHI/Public/RHITransition.h
+++ b/Engine/Source/Runtime/RHI/Public/RHITransition.h
@@ -77,6 +77,21 @@ struct FRHISubresourceRange
 	}
 };
 
+/**
+* Represents a change in physical memory allocation for a resource that was created with TexCreate/BUF_ReservedResource flag.
+* Physical memory is allocated in tiles/pages and mapped to the tail of the currently committed region of the resource.
+* This API may be used to grow or shrink reserved resources without moving the bulk of the data or re-creating SRVs/UAVs.
+* The contents of the newly committed region of the resource is undefined and must be overwritten by the application before use.
+* Reserved resources must be created with maximum expected size, which will not cost any memory until committed.
+* Commit size must be smaller or equal to the maximum resource size specified at creation.
+* Check GRHIGlobals.ReservedResources.Supported before using this API or TexCreate/BUF_ReservedResource flag.
+*/
+struct FRHICommitResourceInfo
+{
+	uint64 SizeInBytes = 0;
+	FRHICommitResourceInfo(uint64 InSizeInBytes) : SizeInBytes(InSizeInBytes) {}
+};
+
 struct FRHITransitionInfo : public FRHISubresourceRange
 {
 	union
@@ -101,6 +116,7 @@ struct FRHITransitionInfo : public FRHISubresourceRange
 	ERHIAccess AccessBefore = ERHIAccess::Unknown;
 	ERHIAccess AccessAfter = ERHIAccess::Unknown;
 	EResourceTransitionFlags Flags = EResourceTransitionFlags::None;
+	TOptional<FRHICommitResourceInfo> CommitInfo;
 
 	FRHITransitionInfo() = default;
 
@@ -136,6 +152,14 @@ struct FRHITransitionInfo : public FRHISubresourceRange
 		, Flags(InFlags)
 	{}
 
+	FRHITransitionInfo(class FRHIBuffer* InRHIBuffer, ERHIAccess InPreviousState, ERHIAccess InNewState, FRHICommitResourceInfo InCommitInfo)
+		: Buffer(InRHIBuffer)
+		, Type(EType::Buffer)
+		, AccessBefore(InPreviousState)
+		, AccessAfter(InNewState)
+		, CommitInfo(InCommitInfo)
+	{}
+
 	FRHITransitionInfo(class FRHIRayTracingAccelerationStructure* InBVH, ERHIAccess InPreviousState, ERHIAccess InNewState, EResourceTransitionFlags InFlags = EResourceTransitionFlags::None)
 		: BVH(InBVH)
 		, Type(EType::BVH)
