D3D12 Bindless Descriptor Updates
- Descriptor updates are queued on the D3D12 context
- Split descriptor management type into a deferred, per-pipeline manager for Resources and a simple one for Samplers.
- Initial descriptor update flushing implementations are still WIP, CPU side heap copying/switching is up next.
- Added ability for platforms to provide custom descriptor update functions.
- Moved common code from Dispatch/Draw commands into common functions: SetupDispatch, SetupIndirectArgument, SetupDraw, SetupDispatchDraw and PostGpuEvent.

diff --git a/Engine/Source/Runtime/D3D12RHI/Private/D3D12Adapter.cpp b/Engine/Source/Runtime/D3D12RHI/Private/D3D12Adapter.cpp
--- a/Engine/Source/Runtime/D3D12RHI/Private/D3D12Adapter.cpp
+++ b/Engine/Source/Runtime/D3D12RHI/Private/D3D12Adapter.cpp
@@ -43,6 +43,14 @@ static FAutoConsoleVariableRef CTrackedReleasedAllocationFrameRetention(
 );
 #endif
 
+int32 GAllowAsyncCompute = 1;
+static FAutoConsoleVariableRef CVarAllowAsyncCompute(
+	TEXT("r.D3D12.AllowAsyncCompute"),
+	GAllowAsyncCompute,
+	TEXT("Allow usage of async compute"),
+	ECVF_ReadOnly | ECVF_RenderThreadSafe
+);
+
 #if PLATFORM_WINDOWS
 
 #if UE_BUILD_SHIPPING || UE_BUILD_TEST
@@ -1223,6 +1231,36 @@ void FD3D12Adapter::InitializeDevices()
 #endif // PLATFORM_WINDOWS
 		}
 
+		if (FParse::Param(FCommandLine::Get(), TEXT("DisableAsyncCompute")))
+		{
+			GSupportsEfficientAsyncCompute = false;
+		}
+		else if (FParse::Param(FCommandLine::Get(), TEXT("ForceAsyncCompute")))
+		{
+			GSupportsEfficientAsyncCompute = true;
+		}
+		else if (!GSupportsEfficientAsyncCompute && GAllowAsyncCompute && GRHISupportsParallelRHIExecute)
+		{
+			if (IsRHIDeviceAMD())
+			{
+				GSupportsEfficientAsyncCompute = true;
+			}
+#if PLATFORM_WINDOWS
+			else
+			{
+				D3D12_FEATURE_DATA_D3D12_OPTIONS6 D3D12Caps6{};
+				HRESULT Options6HR = RootDevice->CheckFeatureSupport(D3D12_FEATURE_D3D12_OPTIONS6, &D3D12Caps6, sizeof(D3D12Caps6));
+
+				// Allow async compute by default on nVidia cards which support PerPrimitiveShadingRateSupportedWithViewportIndexing 
+				// this should be a good metric according to nVidia itself (this is set for Ampere and newer cards)
+				if (IsRHIDeviceNVIDIA() && Options6HR == S_OK && D3D12Caps6.PerPrimitiveShadingRateSupportedWithViewportIndexing)
+				{
+					GSupportsEfficientAsyncCompute = true;
+				}
+			}
+#endif
+		}
+
 #if PLATFORM_WINDOWS
 		D3D12_FEATURE_DATA_D3D12_OPTIONS2 D3D12Caps2 = {};
 		if (FAILED(RootDevice->CheckFeatureSupport(D3D12_FEATURE_D3D12_OPTIONS2, &D3D12Caps2, sizeof(D3D12Caps2))))
diff --git a/Engine/Source/Runtime/D3D12RHI/Private/D3D12Allocation.cpp b/Engine/Source/Runtime/D3D12RHI/Private/D3D12Allocation.cpp
--- a/Engine/Source/Runtime/D3D12RHI/Private/D3D12Allocation.cpp
+++ b/Engine/Source/Runtime/D3D12RHI/Private/D3D12Allocation.cpp
@@ -1595,7 +1595,7 @@ void FD3D12DefaultBufferAllocator::FreeDefaultBufferPools()
 }
 
 
-void FD3D12DefaultBufferAllocator::BeginFrame()
+void FD3D12DefaultBufferAllocator::BeginFrame(FRHICommandListBase& RHICmdList)
 {
 #if USE_BUFFER_POOL_ALLOCATOR
 	FScopeLock Lock(&CS);
@@ -1610,7 +1610,7 @@ void FD3D12DefaultBufferAllocator::BeginFrame()
 		{
 			if (DefaultBufferPool)
 			{
-				DefaultBufferPool->Defrag(MaxCopySize, CopySize);
+				DefaultBufferPool->Defrag(RHICmdList, MaxCopySize, CopySize);
 
 				// break when we reach the max copy size
 				if (CopySize >= MaxCopySize)
@@ -1806,7 +1806,7 @@ HRESULT FD3D12TextureAllocatorPool::AllocateTexture(
 }
 
 
-void FD3D12TextureAllocatorPool::BeginFrame()
+void FD3D12TextureAllocatorPool::BeginFrame(FRHICommandListBase& RHICmdList)
 {
 	if (GD3D12VRAMTexturePoolDefrag > 0 && GD3D12VRAMTexturePoolDefragMaxCopySizePerFrame > 0)
 	{
@@ -1816,7 +1816,7 @@ void FD3D12TextureAllocatorPool::BeginFrame()
 		uint32 CopySize = 0;
 		for (uint32 PoolIndex = 0; PoolIndex < (uint32)EPoolType::Count; ++PoolIndex)
 		{
-			PoolAllocators[PoolIndex]->Defrag(MaxCopySize, CopySize);
+			PoolAllocators[PoolIndex]->Defrag(RHICmdList, MaxCopySize, CopySize);
 		}
 	}
 
diff --git a/Engine/Source/Runtime/D3D12RHI/Private/D3D12Allocation.h b/Engine/Source/Runtime/D3D12RHI/Private/D3D12Allocation.h
--- a/Engine/Source/Runtime/D3D12RHI/Private/D3D12Allocation.h
+++ b/Engine/Source/Runtime/D3D12RHI/Private/D3D12Allocation.h
@@ -434,7 +434,7 @@ public:
 	// Grab a buffer from the available buffers or create a new buffer if none are available
 	void AllocDefaultResource(D3D12_HEAP_TYPE InHeapType, const D3D12_RESOURCE_DESC& pDesc, EBufferUsageFlags InBufferUsage, ED3D12ResourceStateMode InResourceStateMode, D3D12_RESOURCE_STATES InCreateState, FD3D12ResourceLocation& ResourceLocation, uint32 Alignment, const TCHAR* Name);
 	void FreeDefaultBufferPools();
-	void BeginFrame();
+	void BeginFrame(FRHICommandListBase& RHICmdList);
 	void CleanupFreeBlocks(uint64 InFrameLag);	
 	void UpdateMemoryStats();
 
@@ -920,7 +920,7 @@ public:
 
 	HRESULT AllocateTexture(FD3D12ResourceDesc Desc, const D3D12_CLEAR_VALUE* ClearValue, EPixelFormat UEFormat, FD3D12ResourceLocation& TextureLocation, const D3D12_RESOURCE_STATES InitialState, const TCHAR* Name);
 	
-	void BeginFrame();
+	void BeginFrame(FRHICommandListBase& RHICmdList);
 	void CleanUpAllocations();
 	void Destroy();
 	bool GetMemoryStats(uint64& OutTotalAllocated, uint64& OutTotalUnused) const;
@@ -947,7 +947,7 @@ public:
 
 	HRESULT AllocateTexture(FD3D12ResourceDesc Desc, const D3D12_CLEAR_VALUE* ClearValue, EPixelFormat UEFormat, FD3D12ResourceLocation& TextureLocation, const D3D12_RESOURCE_STATES InitialState, const TCHAR* Name);
 
-	void BeginFrame() {}
+	void BeginFrame(FRHICommandListBase& RHICmdList) {}
 	void CleanUpAllocations()
 	{
 		ReadOnlyTexturePool.CleanUpAllocations();
@@ -984,7 +984,7 @@ public:
 
 	HRESULT AllocateTexture(FD3D12ResourceDesc Desc, const D3D12_CLEAR_VALUE* ClearValue, EPixelFormat UEFormat, FD3D12ResourceLocation& TextureLocation, const D3D12_RESOURCE_STATES InitialState, const TCHAR* Name);
 	
-	void BeginFrame() {}
+	void BeginFrame(FRHICommandListBase& RHICmdList) {}
 	void CleanUpAllocations() { ReadOnlyTexturePool.CleanUpAllocations(0); }
 
 	void Destroy() { ReadOnlyTexturePool.Destroy(); }
diff --git a/Engine/Source/Runtime/D3D12RHI/Private/D3D12BindlessDescriptors.cpp b/Engine/Source/Runtime/D3D12RHI/Private/D3D12BindlessDescriptors.cpp
--- a/Engine/Source/Runtime/D3D12RHI/Private/D3D12BindlessDescriptors.cpp
+++ b/Engine/Source/Runtime/D3D12RHI/Private/D3D12BindlessDescriptors.cpp
@@ -22,170 +22,125 @@ static FAutoConsoleVariableRef CVarBindlessSamplerDescriptorHeapSize(
 	ECVF_ReadOnly
 );
 
-////////////////////////////////////////////////////////////////////////////////////////////////////////////////////////
-// FD3D12BindlessDescriptorHeapManager
-
-FD3D12BindlessDescriptorHeapManager::FD3D12BindlessDescriptorHeapManager(FD3D12Device* InDevice, ERHIDescriptorHeapType InType, ERHIBindlessConfiguration InConfiguration, uint32 InNumDescriptorsPerHeap, TConstArrayView<TStatId> InStats)
-	: FD3D12DeviceChild(InDevice)
-	, Type(InType)
-	, Configuration(InConfiguration)
-	, NumDescriptorsPerHeap(InNumDescriptorsPerHeap)
+FD3D12DescriptorHeap* UE::D3D12BindlessDescriptors::CreateCpuHeap(FD3D12Device* InDevice, ERHIDescriptorHeapType InType, uint32 InNewNumDescriptorsPerHeap)
 {
-#if STATS
-	Stats = InStats;
-#endif
+	const TCHAR* const HeapName = (InType == ERHIDescriptorHeapType::Standard) ? TEXT("BindlessResourcesCPU") : TEXT("BindlessSamplersCPU");
 
-	SetupInitialState(InNumDescriptorsPerHeap);
+	return InDevice->GetDescriptorHeapManager().AllocateIndependentHeap(
+		HeapName,
+		InType,
+		InNewNumDescriptorsPerHeap,
+		D3D12_DESCRIPTOR_HEAP_FLAG_NONE
+	);
 }
 
-FD3D12DescriptorHeap* FD3D12BindlessDescriptorHeapManager::AllocateHeap(int32 InSize)
+FD3D12DescriptorHeap* UE::D3D12BindlessDescriptors::CreateGpuHeap(FD3D12Device* InDevice, ERHIDescriptorHeapType InType, uint32 InNewNumDescriptorsPerHeap)
 {
-	check(NumAllocations == 0);
-
-	const int32 NewNumDescriptorsPerHeap = NumDescriptorsPerHeap + InSize;
-	ResizeHeaps(NewNumDescriptorsPerHeap);
+	const TCHAR* const HeapName = (InType == ERHIDescriptorHeapType::Standard) ? TEXT("BindlessResources") : TEXT("BindlessSamplers");
 
-	check(Allocations.Num() == NewNumDescriptorsPerHeap);
-
-	const int32 Offset = 0;
-	const int32 Count = InSize;
+	return InDevice->GetDescriptorHeapManager().AllocateIndependentHeap(
+		HeapName,
+		InType,
+		InNewNumDescriptorsPerHeap,
+		D3D12_DESCRIPTOR_HEAP_FLAG_SHADER_VISIBLE
+	);
+}
 
-	Allocations.SetRange(Offset, Count, true);
-	NumAllocations += InSize;
+////////////////////////////////////////////////////////////////////////////////////////////////////////////////////////
+// FD3D12BindlessSamplerManager
 
-	return new FD3D12DescriptorHeap(GpuHeap.GetReference(), Offset, Count);
+FD3D12BindlessSamplerManager::FD3D12BindlessSamplerManager(FD3D12Device* InDevice, uint32 InNumDescriptors, TConstArrayView<TStatId> InStats)
+	: FD3D12DeviceChild(InDevice)
+	, Allocator(ERHIDescriptorHeapType::Sampler, InNumDescriptors, InStats)
+{
+	GpuHeap = UE::D3D12BindlessDescriptors::CreateGpuHeap(InDevice, ERHIDescriptorHeapType::Sampler, InNumDescriptors);
 }
 
-FRHIDescriptorHandle FD3D12BindlessDescriptorHeapManager::Allocate()
+FRHIDescriptorHandle FD3D12BindlessSamplerManager::Allocate()
 {
-	FScopeLock Lock(&CriticalSection);
-
-	int32 AllocatedIndex = INDEX_NONE;
+	FRHIDescriptorHandle Result = Allocator.Allocate();
+	check(Result.IsValid());
+	return Result;
+}
 
-	while ((AllocatedIndex = Allocations.FindAndSetFirstZeroBit()) == INDEX_NONE)
+void FD3D12BindlessSamplerManager::Free(FRHIDescriptorHandle InHandle)
+{
+	if (InHandle.IsValid())
 	{
-		const int32 NewNumDescriptorsPerHeap = NumDescriptorsPerHeap * 1.5;
-
-		if (Type == ERHIDescriptorHeapType::Sampler && NewNumDescriptorsPerHeap > D3D12_MAX_SHADER_VISIBLE_SAMPLER_HEAP_SIZE)
-		{
-			return FRHIDescriptorHandle();
-		}
-
-		ResizeHeaps(NewNumDescriptorsPerHeap);
+		Allocator.Free(InHandle);
 	}
-
-	NumAllocations++;
-	RecordAlloc();
-
-	// TODO - find a faster way to verify Allocations
-	//checkSlow(NumAllocations == Allocations.CountSetBits());
-
-	return FRHIDescriptorHandle(Type, AllocatedIndex);
 }
 
-void FD3D12BindlessDescriptorHeapManager::Free(FRHIDescriptorHandle InHandle)
+void FD3D12BindlessSamplerManager::UpdateDescriptorImmediately(FRHIDescriptorHandle DstHandle, D3D12_CPU_DESCRIPTOR_HANDLE SrcDescriptor)
 {
-	if (InHandle.IsValid())
+	if (DstHandle.IsValid())
 	{
-		FScopeLock Lock(&CriticalSection);
+		UE::D3D12Descriptors::CopyDescriptor(GetParentDevice(), GpuHeap, DstHandle, SrcDescriptor);
+	}
+}
 
-		const int32 Index = InHandle.GetIndex();
-		check(Allocations[Index] == true);
-		Allocations[Index] = false;
+////////////////////////////////////////////////////////////////////////////////////////////////////////////////////////
+// FD3D12BindlessResourceManager
 
-		NumAllocations--;
-		RecordFree();
+#if !D3D12RHI_CUSTOM_BINDLESS_RESOURCE_MANAGER
 
-		// TODO - find a faster way to verify Allocations
-		//checkSlow(NumAllocations == Allocations.CountSetBits());
+FD3D12BindlessResourceManager::FD3D12BindlessResourceManager(FD3D12Device* InDevice, uint32 InNumDescriptors, TConstArrayView<TStatId> InStats)
+	: FD3D12DeviceChild(InDevice)
+	, Allocator(ERHIDescriptorHeapType::Standard, InNumDescriptors, InStats)
+{
+	ConfiguredPipelines = ERHIPipeline::Graphics;
+	if (GSupportsEfficientAsyncCompute)
+	{
+		EnumAddFlags(ConfiguredPipelines, ERHIPipeline::AsyncCompute);
 	}
-}
 
-void FD3D12BindlessDescriptorHeapManager::UpdateImmediately(FRHIDescriptorHandle InHandle, D3D12_CPU_DESCRIPTOR_HANDLE InSourceCpuHandle)
-{
-	UE::D3D12Descriptors::CopyDescriptor(GetParentDevice(), CpuHeap, InHandle, InSourceCpuHandle);
-	UE::D3D12Descriptors::CopyDescriptor(GetParentDevice(), GpuHeap, InHandle, InSourceCpuHandle);
-}
+	EnumeratePipelines([&](ERHIPipeline PipelineIndex)
+	{
+		FPipeline& Pipeline = Pipelines[PipelineIndex];
 
-void FD3D12BindlessDescriptorHeapManager::UpdateDeferred(FRHIDescriptorHandle InHandle, D3D12_CPU_DESCRIPTOR_HANDLE InSourceCpuHandle)
-{
-	// TODO: implement deferred updates
-	UpdateImmediately(InHandle, InSourceCpuHandle);
+		Pipeline.CpuHeap = UE::D3D12BindlessDescriptors::CreateCpuHeap(InDevice, ERHIDescriptorHeapType::Standard, InNumDescriptors);
+		Pipeline.GpuHeap = UE::D3D12BindlessDescriptors::CreateGpuHeap(InDevice, ERHIDescriptorHeapType::Standard, InNumDescriptors);
+	});
 }
 
-void FD3D12BindlessDescriptorHeapManager::SetupInitialState(uint32 InNumDescriptorsPerHeap)
+FRHIDescriptorHandle FD3D12BindlessResourceManager::Allocate()
 {
-	CpuHeap = CreateCpuHeapInternal(InNumDescriptorsPerHeap);
-	GpuHeap = CreateGpuHeapInternal(InNumDescriptorsPerHeap);
-
-	Allocations = AllocationListType(false, InNumDescriptorsPerHeap);
-	NumAllocations = 0;
-
-	NumDescriptorsPerHeap = InNumDescriptorsPerHeap;
+	FRHIDescriptorHandle Result = Allocator.Allocate();
+	check(Result.IsValid());
+	return Result;
 }
 
-void FD3D12BindlessDescriptorHeapManager::ResizeHeaps(uint32 InNewNumDescriptorsPerHeap)
+void FD3D12BindlessResourceManager::Free(FRHIDescriptorHandle InHandle)
 {
-	if (InNewNumDescriptorsPerHeap > NumDescriptorsPerHeap)
+	if (InHandle.IsValid())
 	{
-		// Save off the old heaps and descriptor count
-		FD3D12DescriptorHeapPtr OldCpuHeap = CpuHeap;
-		FD3D12DescriptorHeapPtr OldGpuHeap = GpuHeap;
-		const uint32 OldNumDescriptorsPerHeap = NumDescriptorsPerHeap;
-
-		// Create the new heaps
-		CpuHeap = CreateCpuHeapInternal(InNewNumDescriptorsPerHeap);
-		GpuHeap = CreateGpuHeapInternal(InNewNumDescriptorsPerHeap);
-
-		// Resize the allocation list
-		Allocations.SetNum(InNewNumDescriptorsPerHeap, false);
-
-		// Validate the count is still correct
-		checkSlow(NumAllocations == Allocations.CountSetBits());
-
-		// Now we can switch to the new size
-		NumDescriptorsPerHeap = InNewNumDescriptorsPerHeap;
-
-		// Copy the old descriptor range from the old heap to the new heaps
-		UE::D3D12Descriptors::CopyDescriptors(GetParentDevice(), CpuHeap, OldCpuHeap, OldNumDescriptorsPerHeap);
-		UE::D3D12Descriptors::CopyDescriptors(GetParentDevice(), GpuHeap, OldCpuHeap, OldNumDescriptorsPerHeap);
-
-		// Finally, queue the old heaps for deletion
-		FD3D12DescriptorHeapManager& HeapManager = GetParentDevice()->GetDescriptorHeapManager();
-		HeapManager.DeferredFreeHeap(OldCpuHeap);
-		HeapManager.DeferredFreeHeap(OldGpuHeap);
+		Allocator.Free(InHandle);
 	}
 }
 
-FD3D12DescriptorHeap* FD3D12BindlessDescriptorHeapManager::CreateCpuHeapInternal(uint32 InNewNumDescriptorsPerHeap)
+void FD3D12BindlessResourceManager::UpdateDescriptorImmediately(FRHIDescriptorHandle DstHandle, D3D12_CPU_DESCRIPTOR_HANDLE SrcDescriptor)
 {
-	const TCHAR* const HeapName = (GetType() == ERHIDescriptorHeapType::Standard) ? TEXT("BindlessResourcesCPU") : TEXT("BindlessSamplersCPU");
-
-	return GetParentDevice()->GetDescriptorHeapManager().AllocateIndependentHeap(
-		HeapName,
-		GetType(),
-		InNewNumDescriptorsPerHeap,
-		D3D12_DESCRIPTOR_HEAP_FLAG_NONE
-	);
+	if (DstHandle.IsValid())
+	{
+		EnumeratePipelines([&](ERHIPipeline PipelineIndex)
+		{
+			UE::D3D12Descriptors::CopyDescriptor(GetParentDevice(), Pipelines[PipelineIndex].GpuHeap, DstHandle, SrcDescriptor);
+		});
+	}
 }
 
-FD3D12DescriptorHeap* FD3D12BindlessDescriptorHeapManager::CreateGpuHeapInternal(uint32 InNewNumDescriptorsPerHeap)
+void FD3D12BindlessResourceManager::UpdateDescriptor(FRHICommandListBase& RHICmdList, FRHIDescriptorHandle DstHandle, D3D12_CPU_DESCRIPTOR_HANDLE SrcDescriptor)
 {
-	const TCHAR* const HeapName = (GetType() == ERHIDescriptorHeapType::Standard) ? TEXT("BindlessResources") : TEXT("BindlessSamplers");
-
-	return GetParentDevice()->GetDescriptorHeapManager().AllocateIndependentHeap(
-		HeapName,
-		GetType(),
-		InNewNumDescriptorsPerHeap,
-		D3D12_DESCRIPTOR_HEAP_FLAG_SHADER_VISIBLE
-	);
+	UpdateDescriptorImmediately(DstHandle, SrcDescriptor);
 }
 
-void FD3D12BindlessDescriptorHeapManager::UpdateGpuHeap(FD3D12DescriptorHeap* InGpuHeap)
+void FD3D12BindlessResourceManager::FlushPendingDescriptorUpdates(FD3D12CommandContext& Context, ERHIPipeline PipelineIndex, const FD3D12PendingResourceDescriptorUpdates& PendingDescriptorUpdates)
 {
-	UE::D3D12Descriptors::CopyDescriptors(GetParentDevice(), InGpuHeap, CpuHeap, NumDescriptorsPerHeap);
+	checkNoEntry(); // TODO
 }
 
+#endif // D3D12RHI_CUSTOM_BINDLESS_RESOURCE_MANAGER
+
 ////////////////////////////////////////////////////////////////////////////////////////////////////////////////////////
 // FD3D12BindlessDescriptorManager
 
@@ -201,10 +156,7 @@ void FD3D12BindlessDescriptorManager::Init()
 	ResourcesConfiguration = RHIGetRuntimeBindlessResourcesConfiguration(GMaxRHIShaderPlatform);
 	SamplersConfiguration  = RHIGetRuntimeBindlessSamplersConfiguration(GMaxRHIShaderPlatform);
 
-	const uint32 NumResourceDescriptors = ResourcesConfiguration != ERHIBindlessConfiguration::Disabled ? GBindlessResourceDescriptorHeapSize : 0;
-	const uint32 NumSamplerDescriptors = SamplersConfiguration != ERHIBindlessConfiguration::Disabled ? GBindlessSamplerDescriptorHeapSize  : 0;
-
-	if (NumResourceDescriptors > 0)
+	if (ResourcesConfiguration != ERHIBindlessConfiguration::Disabled)
 	{
 		const TStatId Stats[] =
 		{
@@ -212,10 +164,15 @@ void FD3D12BindlessDescriptorManager::Init()
 			GET_STATID(STAT_BindlessResourceDescriptorsAllocated),
 		};
 
-		Managers.Emplace(GetParentDevice(), ERHIDescriptorHeapType::Standard, ResourcesConfiguration, NumResourceDescriptors, Stats);
+		uint32 NumResourceDescriptors = GBindlessResourceDescriptorHeapSize;
+#if D3D12RHI_USE_CONSTANT_BUFFER_VIEWS
+		NumResourceDescriptors += GBindlessOnlineDescriptorHeapBlockSize;
+#endif
+
+		ResourceManager = MakeUnique<FD3D12BindlessResourceManager>(GetParentDevice(), NumResourceDescriptors, Stats);
 	}
 
-	if (NumSamplerDescriptors > 0)
+	if (SamplersConfiguration != ERHIBindlessConfiguration::Disabled)
 	{
 		const TStatId Stats[] =
 		{
@@ -223,31 +180,47 @@ void FD3D12BindlessDescriptorManager::Init()
 			GET_STATID(STAT_BindlessSamplerDescriptorsAllocated),
 		};
 
-		Managers.Emplace(GetParentDevice(), ERHIDescriptorHeapType::Sampler, SamplersConfiguration, NumSamplerDescriptors, Stats);
+		uint32 NumSamplerDescriptors = GBindlessSamplerDescriptorHeapSize;
+		if (NumSamplerDescriptors > D3D12_MAX_SHADER_VISIBLE_SAMPLER_HEAP_SIZE)
+		{
+			UE_LOG(LogD3D12RHI, Error, TEXT("D3D12.Bindless.SamplerDescriptorHeapSize was set to %d, which is higher than the D3D12 maximum of %d. Adjusting the value to prevent a crash."),
+				NumSamplerDescriptors,
+				D3D12_MAX_SHADER_VISIBLE_SAMPLER_HEAP_SIZE
+			);
+			NumSamplerDescriptors = D3D12_MAX_SHADER_VISIBLE_SAMPLER_HEAP_SIZE;
+		}
+
+		SamplerManager = MakeUnique<FD3D12BindlessSamplerManager>(GetParentDevice(), NumSamplerDescriptors, Stats);
 	}
 }
 
 FRHIDescriptorHandle FD3D12BindlessDescriptorManager::Allocate(ERHIDescriptorHeapType InType)
 {
-	for (FD3D12BindlessDescriptorHeapManager& Manager : Managers)
+	if (InType == ERHIDescriptorHeapType::Standard && ResourceManager)
 	{
-		if (Manager.HandlesAllocation(InType))
-		{
-			return Manager.Allocate();
-		}
+		return ResourceManager->Allocate();
 	}
+
+	if (InType == ERHIDescriptorHeapType::Sampler && SamplerManager)
+	{
+		return SamplerManager->Allocate();
+	}
+
 	return FRHIDescriptorHandle();
 }
 
 void FD3D12BindlessDescriptorManager::ImmediateFree(FRHIDescriptorHandle InHandle)
 {
-	for (FD3D12BindlessDescriptorHeapManager& Manager : Managers)
+	if (InHandle.GetType() == ERHIDescriptorHeapType::Standard && ResourceManager)
 	{
-		if (Manager.HandlesAllocation(InHandle.GetType()))
-		{
-			Manager.Free(InHandle);
-			return;
-		}
+		ResourceManager->Free(InHandle);
+		return;
+	}
+
+	if (InHandle.GetType() == ERHIDescriptorHeapType::Sampler && SamplerManager)
+	{
+		SamplerManager->Free(InHandle);
+		return;
 	}
 
 	// Bad configuration?
@@ -262,98 +235,93 @@ void FD3D12BindlessDescriptorManager::DeferredFreeFromDestructor(FRHIDescriptorH
 	}
 }
 
-void FD3D12BindlessDescriptorManager::UpdateImmediately(FRHIDescriptorHandle InHandle, D3D12_CPU_DESCRIPTOR_HANDLE InSourceCpuHandle)
+void FD3D12BindlessDescriptorManager::UpdateDescriptorImmediately(FRHIDescriptorHandle DstHandle, D3D12_CPU_DESCRIPTOR_HANDLE SrcDescriptor)
 {
-	for (FD3D12BindlessDescriptorHeapManager& Manager : Managers)
+	if (DstHandle.GetType() == ERHIDescriptorHeapType::Standard && ResourceManager)
 	{
-		if (Manager.HandlesAllocation(InHandle.GetType()))
-		{
-			Manager.UpdateImmediately(InHandle, InSourceCpuHandle);
-			return;
-		}
+		ResourceManager->UpdateDescriptorImmediately(DstHandle, SrcDescriptor);
+		return;
+	}
+
+	if (DstHandle.GetType() == ERHIDescriptorHeapType::Sampler && SamplerManager)
+	{
+		SamplerManager->UpdateDescriptorImmediately(DstHandle, SrcDescriptor);
+		return;
 	}
 
 	// Bad configuration?
 	checkNoEntry();
 }
 
-void FD3D12BindlessDescriptorManager::UpdateDeferred(FRHIDescriptorHandle InHandle, D3D12_CPU_DESCRIPTOR_HANDLE InSourceCpuHandle)
+void FD3D12BindlessDescriptorManager::UpdateResourceDescriptor(FRHICommandListBase& RHICmdList, FRHIDescriptorHandle DstHandle, D3D12_CPU_DESCRIPTOR_HANDLE SrcDescriptor)
 {
-	for (FD3D12BindlessDescriptorHeapManager& Manager : Managers)
+	if (ResourceManager)
 	{
-		if (Manager.HandlesAllocation(InHandle.GetType()))
-		{
-			Manager.UpdateDeferred(InHandle, InSourceCpuHandle);
-			return;
-		}
+		ResourceManager->UpdateDescriptor(RHICmdList, DstHandle, SrcDescriptor);
+		return;
 	}
 
 	// Bad configuration?
 	checkNoEntry();
 }
 
-FD3D12DescriptorHeap* FD3D12BindlessDescriptorManager::AllocateHeap(ERHIDescriptorHeapType InType, int32 InSize)
+void FD3D12BindlessDescriptorManager::FlushPendingDescriptorUpdates(FD3D12CommandContext& Context, ERHIPipeline PipelineIndex, const FD3D12PendingResourceDescriptorUpdates& PendingDescriptorUpdates)
 {
-	for (FD3D12BindlessDescriptorHeapManager& Manager : Managers)
+	if (ResourceManager)
 	{
-		if (Manager.HandlesAllocation(InType))
-		{
-			return Manager.AllocateHeap(InSize);
-		}
+		ResourceManager->FlushPendingDescriptorUpdates(Context, PipelineIndex, PendingDescriptorUpdates);
+	}
+}
+
+#if D3D12RHI_USE_CONSTANT_BUFFER_VIEWS
+TRHIPipelineArray<FD3D12DescriptorHeapPtr> FD3D12BindlessDescriptorManager::AllocateResourceHeapsForAllPipelines(int32 InSize)
+{
+	if (ResourceManager)
+	{
+		return ResourceManager->AllocateResourceHeapsForAllPipelines(InSize);
 	}
 
 	// Bad configuration?
 	checkNoEntry();
-	return nullptr;
+	return TRHIPipelineArray<FD3D12DescriptorHeapPtr>();
 }
+#endif // D3D12RHI_USE_CONSTANT_BUFFER_VIEWS
 
-FD3D12DescriptorHeap* FD3D12BindlessDescriptorManager::GetHeap(ERHIDescriptorHeapType InType)
+FD3D12DescriptorHeap* FD3D12BindlessDescriptorManager::GetResourceHeap(ERHIPipeline Pipeline)
 {
-	for (FD3D12BindlessDescriptorHeapManager& Manager : Managers)
-	{
-		if (Manager.HandlesAllocation(InType))
-		{
-			return Manager.GetHeap();
-		}
-	}
+	return ResourceManager->GetHeap(Pipeline);
+}
 
-	return nullptr;
+FD3D12DescriptorHeap* FD3D12BindlessDescriptorManager::GetSamplerHeap()
+{
+	return SamplerManager->GetHeap();
 }
 
-FD3D12DescriptorHeap* FD3D12BindlessDescriptorManager::GetHeap(ERHIDescriptorHeapType InType, ERHIBindlessConfiguration InConfiguration)
+FD3D12DescriptorHeap* FD3D12BindlessDescriptorManager::GetResourceHeap(ERHIPipeline Pipeline, ERHIBindlessConfiguration InConfiguration)
 {
-	for (FD3D12BindlessDescriptorHeapManager& Manager : Managers)
+	if (AreResourcesBindless(InConfiguration))
 	{
-		if (Manager.HandlesAllocation(InType, InConfiguration))
-		{
-			return Manager.GetHeap();
-		}
+		return GetResourceHeap(Pipeline);
 	}
+
 	return nullptr;
 }
 
-
-bool FD3D12BindlessDescriptorManager::HasHeap(ERHIDescriptorHeapType InType, ERHIBindlessConfiguration InConfiguration) const
+FD3D12DescriptorHeap* FD3D12BindlessDescriptorManager::GetSamplerHeap(ERHIBindlessConfiguration InConfiguration)
 {
-	for (const FD3D12BindlessDescriptorHeapManager& Manager : Managers)
+	if (AreSamplersBindless(InConfiguration))
 	{
-		if (Manager.HandlesAllocation(InType, InConfiguration))
-		{
-			return true;
-		}
+		return GetSamplerHeap();
 	}
 
-	return false;
+	return nullptr;
 }
 
-D3D12_GPU_DESCRIPTOR_HANDLE FD3D12BindlessDescriptorManager::GetGpuHandle(FRHIDescriptorHandle InHandle) const
+D3D12_GPU_DESCRIPTOR_HANDLE FD3D12BindlessDescriptorManager::GetResourceGpuHandle(ERHIPipeline Pipeline, FRHIDescriptorHandle InHandle) const
 {
-	for (const FD3D12BindlessDescriptorHeapManager& Manager : Managers)
+	if (ResourceManager && ResourceManager->GetHeap(Pipeline))
 	{
-		if (Manager.HandlesAllocation(InHandle.GetType()))
-		{
-			return Manager.GetHeap()->GetGPUSlotHandle(InHandle.GetIndex());
-		}
+		return ResourceManager->GetHeap(Pipeline)->GetGPUSlotHandle(InHandle.GetIndex());
 	}
 	return D3D12_GPU_DESCRIPTOR_HANDLE{};
 }
diff --git a/Engine/Source/Runtime/D3D12RHI/Private/D3D12BindlessDescriptors.h b/Engine/Source/Runtime/D3D12RHI/Private/D3D12BindlessDescriptors.h
--- a/Engine/Source/Runtime/D3D12RHI/Private/D3D12BindlessDescriptors.h
+++ b/Engine/Source/Runtime/D3D12RHI/Private/D3D12BindlessDescriptors.h
@@ -4,89 +4,122 @@
 
 #include "D3D12RHICommon.h"
 #include "RHIDefinitions.h"
+#include "RHIDescriptorAllocator.h"
 #include "Templates/RefCounting.h"
 
-#if PLATFORM_SUPPORTS_BINDLESS_RENDERING
+class FD3D12CommandContext;
 
 struct FD3D12DescriptorHeap;
 using FD3D12DescriptorHeapPtr = TRefCountPtr<FD3D12DescriptorHeap>;
 
-class FD3D12BindlessDescriptorHeapManager : public FD3D12DeviceChild
+#if PLATFORM_SUPPORTS_BINDLESS_RENDERING
+
+#include COMPILED_PLATFORM_HEADER(D3D12BindlessDescriptors.h)
+
+namespace UE::D3D12BindlessDescriptors
 {
-public:
-	FD3D12BindlessDescriptorHeapManager() = delete;
-	FD3D12BindlessDescriptorHeapManager(FD3D12Device* InDevice, ERHIDescriptorHeapType InType, ERHIBindlessConfiguration InConfiguration, uint32 InNumDescriptorsPerHeap, TConstArrayView<TStatId> InStats);
+	FD3D12DescriptorHeap* CreateCpuHeap(FD3D12Device* InDevice, ERHIDescriptorHeapType InType, uint32 InNewNumDescriptorsPerHeap);
+	FD3D12DescriptorHeap* CreateGpuHeap(FD3D12Device* InDevice, ERHIDescriptorHeapType InType, uint32 InNewNumDescriptorsPerHeap);
+}
 
-	FD3D12DescriptorHeap* AllocateHeap(int32 InSize);
+/** Manager specifically for bindless sampler descriptors. */
+class FD3D12BindlessSamplerManager : public FD3D12DeviceChild
+{
+public:
+	FD3D12BindlessSamplerManager() = delete;
+	FD3D12BindlessSamplerManager(FD3D12Device* InDevice, uint32 InNumDescriptors, TConstArrayView<TStatId> InStats);
 
 	FRHIDescriptorHandle Allocate();
 	void                 Free(FRHIDescriptorHandle InHandle);
 
-	void UpdateImmediately(FRHIDescriptorHandle InHandle, D3D12_CPU_DESCRIPTOR_HANDLE InSourceCpuHandle);
-	void UpdateDeferred   (FRHIDescriptorHandle InHandle, D3D12_CPU_DESCRIPTOR_HANDLE InSourceCpuHandle);
+	void UpdateDescriptorImmediately(FRHIDescriptorHandle DstHandle, D3D12_CPU_DESCRIPTOR_HANDLE SrcDescriptor);
 
-	      FD3D12DescriptorHeap* GetHeap()       { return GpuHeap.GetReference(); }
-	const FD3D12DescriptorHeap* GetHeap() const { return GpuHeap.GetReference(); }
-	ERHIDescriptorHeapType      GetType() const { return Type; }
-	
-	bool HandlesAllocation(ERHIDescriptorHeapType InType) const { return GetType() == InType; }
-	bool HandlesConfiguration(ERHIBindlessConfiguration InConfiguration) const { return Configuration != ERHIBindlessConfiguration::Disabled && Configuration <= InConfiguration; }
-
-	bool HandlesAllocation(ERHIDescriptorHeapType InType, ERHIBindlessConfiguration InConfiguration) const
-	{
-		return HandlesAllocation(InType) && HandlesConfiguration(InConfiguration);
-	}
+	FD3D12DescriptorHeap* GetHeap() { return GpuHeap.GetReference(); }
 
 private:
-	void SetupInitialState(uint32 InNumDescriptorsPerHeap);
-	void ResizeHeaps(uint32 InNewNumDescriptorsPerHeap);
+	FD3D12DescriptorHeapPtr GpuHeap;
+	FRHIHeapDescriptorAllocator  Allocator;
+};
+
+#if !D3D12RHI_CUSTOM_BINDLESS_RESOURCE_MANAGER
 
-	FD3D12DescriptorHeap* CreateCpuHeapInternal(uint32 InNumDescriptorsPerHeap);
-	FD3D12DescriptorHeap* CreateGpuHeapInternal(uint32 InNumDescriptorsPerHeap);
+struct FD3D12PendingResourceDescriptorUpdates
+{
+	TArray<FRHIDescriptorHandle> Handles;
+	TArray<D3D12_CPU_DESCRIPTOR_HANDLE> Descriptors;
 
-	void UpdateGpuHeap(FD3D12DescriptorHeap* GpuHeap);
+	int32 Num() const
+	{
+		return Handles.Num();
+	}
 
-	void RecordAlloc()
+	void Add(FRHIDescriptorHandle InHandle, const D3D12_CPU_DESCRIPTOR_HANDLE& InDescriptor)
 	{
-#if STATS
-		for (TStatId Stat : Stats)
+		if (ensure(InHandle.IsValid()))
 		{
-			INC_DWORD_STAT_FName(Stat.GetName());
+			Handles.Emplace(InHandle);
+			Descriptors.Emplace(InDescriptor);
 		}
-#endif
 	}
 
-	void RecordFree()
+	void Empty()
+	{
+		Handles.Empty();
+		Descriptors.Empty();
+	}
+};
+
+/** Manager specifically for bindless resource descriptors. Has to handle renames on command lists. */
+class FD3D12BindlessResourceManager : public FD3D12DeviceChild
+{
+public:
+	FD3D12BindlessResourceManager() = delete;
+	FD3D12BindlessResourceManager(FD3D12Device* InDevice, uint32 InNumDescriptors, TConstArrayView<TStatId> InStats);
+
+	FRHIDescriptorHandle Allocate();
+	void                 Free(FRHIDescriptorHandle InHandle);
+
+	void UpdateDescriptorImmediately(FRHIDescriptorHandle DstHandle, D3D12_CPU_DESCRIPTOR_HANDLE SrcDescriptor);
+	void UpdateDescriptor(FRHICommandListBase& RHICmdList, FRHIDescriptorHandle DstHandle, D3D12_CPU_DESCRIPTOR_HANDLE SrcDescriptor);
+
+	void FlushPendingDescriptorUpdates(FD3D12CommandContext& Context, ERHIPipeline PipelineIndex, const FD3D12PendingResourceDescriptorUpdates& PendingDescriptorUpdates);
+
+	bool IsEnabledForPipeline(ERHIPipeline Pipeline) const
+	{
+		return EnumHasAnyFlags(ConfiguredPipelines, Pipeline);
+	}
+
+	FD3D12DescriptorHeap* GetHeap(ERHIPipeline Pipeline)
 	{
-#if STATS
-		for (TStatId Stat : Stats)
+		if (IsEnabledForPipeline(Pipeline))
 		{
-			DEC_DWORD_STAT_FName(Stat.GetName());
+			return Pipelines[Pipeline].GpuHeap.GetReference();
 		}
-#endif
+		return nullptr;
 	}
 
 private:
-	using AllocationListType = TBitArray<>;
+	template <typename TFunctionType>
+	void EnumeratePipelines(TFunctionType&& Function)
+	{
+		EnumerateRHIPipelines(ConfiguredPipelines, Forward<TFunctionType>(Function));
+	}
 
-	FCriticalSection        CriticalSection;
-	
-	FD3D12DescriptorHeapPtr CpuHeap;
-	FD3D12DescriptorHeapPtr GpuHeap;
+	struct FPipeline
+	{
+		FD3D12DescriptorHeapPtr CpuHeap;
+		FD3D12DescriptorHeapPtr GpuHeap;
+	};
 
-	AllocationListType      Allocations;
-	int32                   NumAllocations = 0;
+	ERHIPipeline                 ConfiguredPipelines;
+	TRHIPipelineArray<FPipeline> Pipelines;
 
-	const ERHIDescriptorHeapType    Type;
-	const ERHIBindlessConfiguration Configuration;
-	uint32                          NumDescriptorsPerHeap;
+	FRHIHeapDescriptorAllocator  Allocator;
+};
 
-#if STATS
-	TArray<TStatId> Stats;
 #endif
-};
 
-/** Manager for resource descriptors used in bindless rendering. */
+/** Manager for descriptors used in bindless rendering. */
 class FD3D12BindlessDescriptorManager : public FD3D12DeviceChild
 {
 public:
@@ -98,6 +131,12 @@ public:
 	ERHIBindlessConfiguration GetResourcesConfiguration() const { return ResourcesConfiguration; }
 	ERHIBindlessConfiguration GetSamplersConfiguration()  const { return SamplersConfiguration; }
 
+	bool AreResourcesBindless() const { return GetResourcesConfiguration() != ERHIBindlessConfiguration::Disabled; }
+	bool AreSamplersBindless()  const { return GetSamplersConfiguration()  != ERHIBindlessConfiguration::Disabled; }
+
+	bool AreResourcesBindless(ERHIBindlessConfiguration InConfiguration) const { return GetResourcesConfiguration() != ERHIBindlessConfiguration::Disabled && GetResourcesConfiguration() <= InConfiguration; }
+	bool AreSamplersBindless(ERHIBindlessConfiguration InConfiguration)  const { return GetSamplersConfiguration()  != ERHIBindlessConfiguration::Disabled && GetSamplersConfiguration() <= InConfiguration; }
+
 	bool AreResourcesFullyBindless() const { return GetResourcesConfiguration() == ERHIBindlessConfiguration::AllShaders; }
 	bool AreSamplersFullyBindless () const { return GetSamplersConfiguration()  == ERHIBindlessConfiguration::AllShaders; }
 
@@ -105,19 +144,26 @@ public:
 	void                 ImmediateFree(FRHIDescriptorHandle InHandle);
 	void                 DeferredFreeFromDestructor(FRHIDescriptorHandle InHandle);
 
-	void UpdateImmediately(FRHIDescriptorHandle InHandle, D3D12_CPU_DESCRIPTOR_HANDLE InSourceCpuHandle);
-	void UpdateDeferred(FRHIDescriptorHandle InHandle, D3D12_CPU_DESCRIPTOR_HANDLE InSourceCpuHandle);
+	void UpdateDescriptorImmediately(FRHIDescriptorHandle DstHandle, D3D12_CPU_DESCRIPTOR_HANDLE SrcDescriptor);
+	void UpdateResourceDescriptor(FRHICommandListBase& RHICmdList, FRHIDescriptorHandle DstHandle, D3D12_CPU_DESCRIPTOR_HANDLE SrcDescriptor);
+
+	void FlushPendingDescriptorUpdates(FD3D12CommandContext& Context, ERHIPipeline PipelineIndex, const FD3D12PendingResourceDescriptorUpdates& PendingDescriptorUpdates);
+
+#if D3D12RHI_USE_CONSTANT_BUFFER_VIEWS
+	TRHIPipelineArray<FD3D12DescriptorHeapPtr> AllocateResourceHeapsForAllPipelines(int32 InSize);
+#endif
 
-	FD3D12DescriptorHeap* AllocateHeap(ERHIDescriptorHeapType InType, int32 InSize);
+	FD3D12DescriptorHeap* GetResourceHeap(ERHIPipeline Pipeline);
+	FD3D12DescriptorHeap* GetSamplerHeap();
 
-	FD3D12DescriptorHeap* GetHeap(ERHIDescriptorHeapType InType);
-	FD3D12DescriptorHeap* GetHeap(ERHIDescriptorHeapType InType, ERHIBindlessConfiguration InConfiguration);
-	bool                  HasHeap(ERHIDescriptorHeapType InType, ERHIBindlessConfiguration InConfiguration) const;
+	FD3D12DescriptorHeap* GetResourceHeap(ERHIPipeline Pipeline, ERHIBindlessConfiguration InConfiguration);
+	FD3D12DescriptorHeap* GetSamplerHeap(ERHIBindlessConfiguration InConfiguration);
 
-	D3D12_GPU_DESCRIPTOR_HANDLE GetGpuHandle(FRHIDescriptorHandle InHandle) const;
+	D3D12_GPU_DESCRIPTOR_HANDLE GetResourceGpuHandle(ERHIPipeline Pipeline, FRHIDescriptorHandle InHandle) const;
 
 private:
-	TArray<FD3D12BindlessDescriptorHeapManager> Managers;
+	TUniquePtr<FD3D12BindlessResourceManager> ResourceManager;
+	TUniquePtr<FD3D12BindlessSamplerManager>  SamplerManager;
 
 	ERHIBindlessConfiguration ResourcesConfiguration{};
 	ERHIBindlessConfiguration SamplersConfiguration{};
diff --git a/Engine/Source/Runtime/D3D12RHI/Private/D3D12Buffer.cpp b/Engine/Source/Runtime/D3D12RHI/Private/D3D12Buffer.cpp
--- a/Engine/Source/Runtime/D3D12RHI/Private/D3D12Buffer.cpp
+++ b/Engine/Source/Runtime/D3D12RHI/Private/D3D12Buffer.cpp
@@ -102,7 +102,7 @@ struct FRHICommandRenameUploadBuffer final : public FRHICommand<FRHICommandRenam
 		MemoryTrace_ReallocFree(Resource->ResourceLocation.GetGPUVirtualAddress(), EMemoryTraceRootHeap::VideoMemory);
 		MemoryTrace_ReallocAlloc(NewLocation.GetGPUVirtualAddress(), Resource->ResourceLocation.GetSize(), Resource->BufferAlignment, EMemoryTraceRootHeap::VideoMemory);
 #endif
-		Resource->RenameLDAChain(NewLocation);
+		Resource->RenameLDAChain(CmdList, NewLocation);
 
 		if (PreviousPipeline.IsSet())
 		{
@@ -431,17 +431,17 @@ FD3D12Buffer* FD3D12Adapter::CreateRHIBuffer(
 	return BufferOut;
 }
 
-void FD3D12Buffer::Rename(FD3D12ResourceLocation& NewLocation)
+void FD3D12Buffer::Rename(FRHICommandListBase& RHICmdList, FD3D12ResourceLocation& NewLocation)
 {
 	FD3D12ResourceLocation::TransferOwnership(ResourceLocation, NewLocation);
-	ResourceRenamed();
+	ResourceRenamed(RHICmdList);
 }
 
-void FD3D12Buffer::RenameLDAChain(FD3D12ResourceLocation& NewLocation)
+void FD3D12Buffer::RenameLDAChain(FRHICommandListBase& RHICmdList, FD3D12ResourceLocation& NewLocation)
 {
 	// Dynamic buffers use cross-node resources (with the exception of BUF_MultiGPUAllocate)
 	//ensure(GetUsage() & BUF_AnyDynamic);
-	Rename(NewLocation);
+	Rename(RHICmdList, NewLocation);
 
 	if (GNumExplicitGPUsForRendering > 1)
 	{
@@ -455,7 +455,7 @@ void FD3D12Buffer::RenameLDAChain(FD3D12ResourceLocation& NewLocation)
 			for (auto NextBuffer = ++FLinkedObjectIterator(this); NextBuffer; ++NextBuffer)
 			{
 				FD3D12ResourceLocation::ReferenceNode(NextBuffer->GetParentDevice(), NextBuffer->ResourceLocation, ResourceLocation);
-				NextBuffer->ResourceRenamed();
+				NextBuffer->ResourceRenamed(RHICmdList);
 			}
 		}
 	}
@@ -831,7 +831,7 @@ void FD3D12DynamicRHI::RHITransferBufferUnderlyingResource(FRHICommandListBase&
 		Dst->ReleaseOwnership();
 	}
 
-	Dst->ResourceRenamed();
+	Dst->ResourceRenamed(RHICmdList);
 }
 
 void FD3D12DynamicRHI::RHICopyBuffer(FRHIBuffer* SourceBufferRHI, FRHIBuffer* DestBufferRHI)
diff --git a/Engine/Source/Runtime/D3D12RHI/Private/D3D12CommandContext.cpp b/Engine/Source/Runtime/D3D12RHI/Private/D3D12CommandContext.cpp
--- a/Engine/Source/Runtime/D3D12RHI/Private/D3D12CommandContext.cpp
+++ b/Engine/Source/Runtime/D3D12RHI/Private/D3D12CommandContext.cpp
@@ -396,6 +396,13 @@ void FD3D12ContextCommon::Finalize(TArray<FD3D12Payload*>& OutPayloads)
 	OutPayloads.Append(MoveTemp(Payloads));
 }
 
+void FD3D12CommandContext::Finalize(TArray<FD3D12Payload*>& OutPayloads)
+{
+	FlushPendingDescriptorUpdates();
+
+	FD3D12ContextCommon::Finalize(OutPayloads);
+}
+
 FD3D12QueryLocation FD3D12QueryAllocator::Allocate(ED3D12QueryType Type, void* Target)
 {
 	check(Type != ED3D12QueryType::None);
@@ -666,13 +673,21 @@ void FD3D12ContextCommon::ConditionalSplitCommandList()
 	}
 }
 
-void FD3D12CommandContext::RHIBeginFrame()
+void FD3D12DynamicRHI::RHIBeginFrame(FRHICommandListImmediate& RHICmdList)
 {
-	Device->GetGPUProfiler().BeginFrame();
+	RHICmdList.EnqueueLambda([](FRHICommandListBase& ExecutingCmdList)
+	{
+		FD3D12CommandContext& Context = static_cast<FD3D12CommandContext&>(ExecutingCmdList.GetContext());
+		FD3D12Device* Device = Context.Device;
 
-	Device->GetDefaultBufferAllocator().BeginFrame();
-	Device->GetTextureAllocator().BeginFrame();
+		Device->GetGPUProfiler().BeginFrame();
+		Device->GetDefaultBufferAllocator().BeginFrame(ExecutingCmdList);
+		Device->GetTextureAllocator().BeginFrame(ExecutingCmdList);
+	});
+}
 
+void FD3D12CommandContext::RHIBeginFrame()
+{
 	bTrackingEvents = IsDefaultContext() && Device->GetGPUProfiler().bTrackingEvents;
 
 #if D3D12_RHI_RAYTRACING
diff --git a/Engine/Source/Runtime/D3D12RHI/Private/D3D12CommandContext.h b/Engine/Source/Runtime/D3D12RHI/Private/D3D12CommandContext.h
--- a/Engine/Source/Runtime/D3D12RHI/Private/D3D12CommandContext.h
+++ b/Engine/Source/Runtime/D3D12RHI/Private/D3D12CommandContext.h
@@ -8,6 +8,7 @@ D3D12CommandContext.h: D3D12 Command Context Interfaces
 
 #include "D3D12RHIPrivate.h"
 #include "D3D12Queue.h"
+#include "D3D12BindlessDescriptors.h"
 
 #include "Windows/AllowWindowsPlatformTypes.h"
 THIRD_PARTY_INCLUDES_START
@@ -186,7 +187,7 @@ public:
 
 	// Complete recording of the current command list set, and appends the resulting
 	// payloads to the given array. Resets the context so new commands can be recorded.
-	void Finalize(TArray<FD3D12Payload*>& OutPayloads);
+	virtual void Finalize(TArray<FD3D12Payload*>& OutPayloads);
 
 	// The owner device of this context
 	FD3D12Device* const Device;
@@ -741,6 +742,11 @@ public:
 		return StaticUniformBuffers;
 	}
 
+	void AddPendingDescriptorUpdate(FRHIDescriptorHandle InHandle, D3D12_CPU_DESCRIPTOR_HANDLE InDescriptor);
+	void FlushPendingDescriptorUpdates();
+
+	virtual void Finalize(TArray<FD3D12Payload*>& OutPayloads) override;
+
 protected:
 
 	FD3D12CommandContext* GetContext(uint32 InGPUIndex) final override 
@@ -749,6 +755,11 @@ protected:
 	}
 
 private:
+	void SetupDispatch(uint32 ThreadGroupCountX, uint32 ThreadGroupCountY, uint32 ThreadGroupCountZ);
+	void SetupDraw(FRHIBuffer* IndexBufferRHI, uint32 NumPrimitives = 0, uint32 NumVertices = 0);
+	void SetupDispatchDraw(uint32 ThreadGroupCountX, uint32 ThreadGroupCountY, uint32 ThreadGroupCountZ);
+	FD3D12ResourceLocation& SetupIndirectArgument(FRHIBuffer* ArgumentBufferRHI, D3D12_RESOURCE_STATES ExtraStates = static_cast<D3D12_RESOURCE_STATES>(0));
+	void PostGpuEvent();
 
 	static void ClearUAV(TRHICommandList_RecursiveHazardous<FD3D12CommandContext>& RHICmdList, FD3D12UnorderedAccessView_RHI* UAV, const void* ClearValues, bool bFloat);
 
@@ -759,6 +770,10 @@ private:
 	void HandleReservedResourceCommits   (const struct FD3D12TransitionData* TransitionData);
 
 	TArray<FRHIUniformBuffer*> StaticUniformBuffers;
+
+#if PLATFORM_SUPPORTS_BINDLESS_RENDERING
+	FD3D12PendingResourceDescriptorUpdates PendingDescriptorUpdates;
+#endif
 };
 
 // Version of command context to handle multi-GPU.  Because IRHICommandContext is pure virtual we can return the normal
diff --git a/Engine/Source/Runtime/D3D12RHI/Private/D3D12Commands.cpp b/Engine/Source/Runtime/D3D12RHI/Private/D3D12Commands.cpp
--- a/Engine/Source/Runtime/D3D12RHI/Private/D3D12Commands.cpp
+++ b/Engine/Source/Runtime/D3D12RHI/Private/D3D12Commands.cpp
@@ -90,6 +90,24 @@ static void BindUniformBuffer(FD3D12CommandContext& Context, FRHIShader* Shader,
 	Context.DirtyUniformBuffers[ShaderFrequency] |= (1 << BufferIndex);
 }
 
+void FD3D12CommandContext::AddPendingDescriptorUpdate(FRHIDescriptorHandle InHandle, D3D12_CPU_DESCRIPTOR_HANDLE InDescriptor)
+{
+#if PLATFORM_SUPPORTS_BINDLESS_RENDERING
+	PendingDescriptorUpdates.Add(InHandle, InDescriptor);
+#endif
+}
+
+void FD3D12CommandContext::FlushPendingDescriptorUpdates()
+{
+#if PLATFORM_SUPPORTS_BINDLESS_RENDERING
+	if (PendingDescriptorUpdates.Num())
+	{
+		GetParentDevice()->GetBindlessDescriptorManager().FlushPendingDescriptorUpdates(*this, GetPipeline(), PendingDescriptorUpdates);
+		PendingDescriptorUpdates.Empty();
+	}
+#endif
+}
+
 // Vertex state.
 void FD3D12CommandContext::RHISetStreamSource(uint32 StreamIndex, FRHIBuffer* VertexBufferRHI, uint32 Offset)
 {
@@ -98,44 +116,56 @@ void FD3D12CommandContext::RHISetStreamSource(uint32 StreamIndex, FRHIBuffer* Ve
 	StateCache.SetStreamSource(VertexBuffer ? &VertexBuffer->ResourceLocation : nullptr, StreamIndex, Offset);
 }
 
-void FD3D12CommandContext::RHIDispatchComputeShader(uint32 ThreadGroupCountX, uint32 ThreadGroupCountY, uint32 ThreadGroupCountZ)
+void FD3D12CommandContext::SetupDispatch(uint32 ThreadGroupCountX, uint32 ThreadGroupCountY, uint32 ThreadGroupCountZ)
 {
 	if (IsDefaultContext())
 	{
-		GetParentDevice()->RegisterGPUDispatch(FIntVector(ThreadGroupCountX, ThreadGroupCountY, ThreadGroupCountZ));	
+		GetParentDevice()->RegisterGPUDispatch(FIntVector(ThreadGroupCountX, ThreadGroupCountY, ThreadGroupCountZ));
 	}
 
 	CommitComputeShaderConstants();
 	CommitComputeResourceTables();
-	StateCache.ApplyState(ED3D12PipelineType::Compute);
 
-	GraphicsCommandList()->Dispatch(ThreadGroupCountX, ThreadGroupCountY, ThreadGroupCountZ);
-	
-	ConditionalSplitCommandList();
+	StateCache.ApplyState(GetPipeline(), ED3D12PipelineType::Compute);
 
-	DEBUG_EXECUTE_COMMAND_LIST(this);
+	FlushPendingDescriptorUpdates();
 }
 
-void FD3D12CommandContext::RHIDispatchIndirectComputeShader(FRHIBuffer* ArgumentBufferRHI, uint32 ArgumentOffset)
+FD3D12ResourceLocation& FD3D12CommandContext::SetupIndirectArgument(FRHIBuffer* ArgumentBufferRHI, D3D12_RESOURCE_STATES ExtraStates)
 {
-	FD3D12Buffer* ArgumentBuffer = RetrieveObject<FD3D12Buffer>(ArgumentBufferRHI);
+	FD3D12ResourceLocation& ArgumentBufferLocation = RetrieveObject<FD3D12Buffer>(ArgumentBufferRHI)->ResourceLocation;
 
-	if (IsDefaultContext())
-	{
-		GetParentDevice()->RegisterGPUDispatch(FIntVector(1, 1, 1));	
-	}
+	// Indirect args buffer can be a previously pending UAV, which becomes PS\Non-PS read. ApplyState will flush pending transitions, so enqueue the indirect arg transition and flush here.
+	TransitionResource(ArgumentBufferLocation.GetResource(), D3D12_RESOURCE_STATE_TBD, D3D12_RESOURCE_STATE_INDIRECT_ARGUMENT | ExtraStates, D3D12_RESOURCE_BARRIER_ALL_SUBRESOURCES);
 
-	CommitComputeShaderConstants();
-	CommitComputeResourceTables();
+	// Must flush so the desired state is actually set.
+	FlushResourceBarriers();
 
-	FD3D12ResourceLocation& Location = ArgumentBuffer->ResourceLocation;
+	UpdateResidency(ArgumentBufferLocation.GetResource());
 
-	StateCache.ApplyState(ED3D12PipelineType::Compute);
+	return ArgumentBufferLocation;
+}
 
-	// Indirect args buffer can be a previously pending UAV, which becomes PS\Non-PS read. ApplyState will flush pending transitions, so enqueue the indirect arg transition and flush here.
-	D3D12_RESOURCE_STATES IndirectState = D3D12_RESOURCE_STATE_INDIRECT_ARGUMENT | D3D12_RESOURCE_STATE_NON_PIXEL_SHADER_RESOURCE;
-	TransitionResource(Location.GetResource(), D3D12_RESOURCE_STATE_TBD, IndirectState, D3D12_RESOURCE_BARRIER_ALL_SUBRESOURCES);
-	FlushResourceBarriers();	// Must flush so the desired state is actually set.
+void FD3D12CommandContext::PostGpuEvent()
+{
+	ConditionalSplitCommandList();
+	DEBUG_EXECUTE_COMMAND_LIST(this);
+}
+
+void FD3D12CommandContext::RHIDispatchComputeShader(uint32 ThreadGroupCountX, uint32 ThreadGroupCountY, uint32 ThreadGroupCountZ)
+{
+	SetupDispatch(ThreadGroupCountX, ThreadGroupCountY, ThreadGroupCountZ);
+
+	GraphicsCommandList()->Dispatch(ThreadGroupCountX, ThreadGroupCountY, ThreadGroupCountZ);
+	
+	PostGpuEvent();
+}
+
+void FD3D12CommandContext::RHIDispatchIndirectComputeShader(FRHIBuffer* ArgumentBufferRHI, uint32 ArgumentOffset)
+{
+	SetupDispatch(1, 1, 1);
+	
+	FD3D12ResourceLocation& ArgumentBufferLocation = SetupIndirectArgument(ArgumentBufferRHI, D3D12_RESOURCE_STATE_NON_PIXEL_SHADER_RESOURCE);
 
 	FD3D12Adapter* Adapter = GetParentDevice()->GetParentAdapter();
 	ID3D12CommandSignature* CommandSignature = IsAsyncComputeContext()
@@ -145,16 +175,13 @@ void FD3D12CommandContext::RHIDispatchIndirectComputeShader(FRHIBuffer* Argument
 	GraphicsCommandList()->ExecuteIndirect(
 		CommandSignature,
 		1,
-		Location.GetResource()->GetResource(),
-		Location.GetOffsetFromBaseOfResource() + ArgumentOffset,
+		ArgumentBufferLocation.GetResource()->GetResource(),
+		ArgumentBufferLocation.GetOffsetFromBaseOfResource() + ArgumentOffset,
 		NULL,
 		0
-		);
-	UpdateResidency(Location.GetResource());
+	);
 	
-	ConditionalSplitCommandList();
-
-	DEBUG_EXECUTE_COMMAND_LIST(this);
+	PostGpuEvent();
 }
 
 template <typename FunctionType>
@@ -1537,8 +1564,7 @@ void FD3D12CommandContext::CommitComputeResourceTables()
 	SetResourcesFromTables(ComputePSO->GetComputeShader());
 }
 
-void FD3D12CommandContext::RHISetShaderRootConstants(
-	const FUint32Vector4& Constants)
+void FD3D12CommandContext::RHISetShaderRootConstants(const FUint32Vector4& Constants)
 {
 	StateCache.SetRootConstants(Constants);
 }
@@ -1560,112 +1586,98 @@ void FD3D12CommandContext::RHIDispatchShaderBundle(
 	UE::RHICore::DispatchShaderBundleEmulation(RHICmdList, ShaderBundle, RecordArgBufferSRV->GetBuffer(), Dispatches);
 }
 
-void FD3D12CommandContext::RHIDrawPrimitive(uint32 BaseVertexIndex, uint32 NumPrimitives, uint32 NumInstances)
+void FD3D12CommandContext::SetupDraw(FRHIBuffer* IndexBufferRHI, uint32 NumPrimitives /* = 0 */, uint32 NumVertices /* = 0 */)
 {
-	RHI_DRAW_CALL_STATS(StateCache.GetGraphicsPipelinePrimitiveType(), FMath::Max(NumInstances, 1U) * NumPrimitives);
+	if (bTrackingEvents)
+	{
+		GetParentDevice()->RegisterGPUWork(NumPrimitives, NumVertices);
+	}
 
 	CommitGraphicsResourceTables();
 	CommitNonComputeShaderConstants();
 
-	uint32 VertexCount = StateCache.GetVertexCount(NumPrimitives);
-	NumInstances = FMath::Max<uint32>(1, NumInstances);
-
-	if (bTrackingEvents)
+	if (IndexBufferRHI)
 	{
-		GetParentDevice()->RegisterGPUWork(NumPrimitives * NumInstances, VertexCount * NumInstances);
-	}
+		FD3D12Buffer* IndexBuffer = RetrieveObject<FD3D12Buffer>(IndexBufferRHI);
 
-	StateCache.ApplyState(ED3D12PipelineType::Graphics);
-	GraphicsCommandList()->DrawInstanced(VertexCount, NumInstances, BaseVertexIndex, 0);
+		// determine 16bit vs 32bit indices
+		const DXGI_FORMAT Format = (IndexBuffer->GetStride() == sizeof(uint16) ? DXGI_FORMAT_R16_UINT : DXGI_FORMAT_R32_UINT);
 
-	ConditionalSplitCommandList();
+		StateCache.SetIndexBuffer(IndexBuffer->ResourceLocation, Format, 0);
+	}
 
-	DEBUG_EXECUTE_COMMAND_LIST(this);
+	StateCache.ApplyState(GetPipeline(), ED3D12PipelineType::Graphics);
+
+	FlushPendingDescriptorUpdates();
 }
 
-void FD3D12CommandContext::RHIDrawPrimitiveIndirect(FRHIBuffer* ArgumentBufferRHI, uint32 ArgumentOffset)
+void FD3D12CommandContext::SetupDispatchDraw(uint32 ThreadGroupCountX, uint32 ThreadGroupCountY, uint32 ThreadGroupCountZ)
 {
-	FD3D12Buffer* ArgumentBuffer = RetrieveObject<FD3D12Buffer>(ArgumentBufferRHI);
-
-	RHI_DRAW_CALL_INC();
 	if (bTrackingEvents)
 	{
-		GetParentDevice()->RegisterGPUWork(0);
+		GetParentDevice()->RegisterGPUDispatch(FIntVector(ThreadGroupCountX, ThreadGroupCountY, ThreadGroupCountZ));
 	}
 
 	CommitGraphicsResourceTables();
 	CommitNonComputeShaderConstants();
 
-	FD3D12ResourceLocation& Location = ArgumentBuffer->ResourceLocation;
+	StateCache.ApplyState(GetPipeline(), ED3D12PipelineType::Graphics);
 
-	StateCache.ApplyState(ED3D12PipelineType::Graphics);
+	FlushPendingDescriptorUpdates();
+}
 
-	// Indirect args buffer can be a previously pending UAV, which becomes PS\Non-PS read. ApplyState will flush pending transitions, so enqueue the indirect
-	// arg transition and flush here.
-	TransitionResource(Location.GetResource(), D3D12_RESOURCE_STATE_TBD, D3D12_RESOURCE_STATE_INDIRECT_ARGUMENT, D3D12_RESOURCE_BARRIER_ALL_SUBRESOURCES);
-	FlushResourceBarriers();	// Must flush so the desired state is actually set.
+void FD3D12CommandContext::RHIDrawPrimitive(uint32 BaseVertexIndex, uint32 NumPrimitives, uint32 NumInstances)
+{
+	RHI_DRAW_CALL_STATS(StateCache.GetGraphicsPipelinePrimitiveType(), FMath::Max(NumInstances, 1U) * NumPrimitives);
 
-	GraphicsCommandList()->ExecuteIndirect(
-		GetParentDevice()->GetParentAdapter()->GetDrawIndirectCommandSignature(),
-		1,
-		Location.GetResource()->GetResource(),
-		Location.GetOffsetFromBaseOfResource() + ArgumentOffset,
-		NULL,
-		0
-	);
+	uint32 VertexCount = StateCache.GetVertexCount(NumPrimitives);
+	NumInstances = FMath::Max<uint32>(1, NumInstances);
 
-	UpdateResidency(Location.GetResource());
+	SetupDraw(nullptr, NumPrimitives * NumInstances, VertexCount * NumInstances);
 
-	ConditionalSplitCommandList();
+	GraphicsCommandList()->DrawInstanced(VertexCount, NumInstances, BaseVertexIndex, 0);
 
-	DEBUG_EXECUTE_COMMAND_LIST(this);
+	PostGpuEvent();
 }
 
-void FD3D12CommandContext::RHIDrawIndexedIndirect(FRHIBuffer* IndexBufferRHI, FRHIBuffer* ArgumentsBufferRHI, int32 DrawArgumentsIndex, uint32 /*NumInstances*/)
+void FD3D12CommandContext::RHIDrawPrimitiveIndirect(FRHIBuffer* ArgumentBufferRHI, uint32 ArgumentOffset)
 {
-	const uint32 IndexBufferStride = FD3D12DynamicRHI::ResourceCast(IndexBufferRHI)->GetStride();
-	const uint32 ArgumentsBufferStride = FD3D12DynamicRHI::ResourceCast(ArgumentsBufferRHI)->GetStride();
-
-	FD3D12Buffer* IndexBuffer = RetrieveObject<FD3D12Buffer>(IndexBufferRHI);
-	FD3D12Buffer* ArgumentsBuffer = RetrieveObject<FD3D12Buffer>(ArgumentsBufferRHI);
-
 	RHI_DRAW_CALL_INC();
-	if (bTrackingEvents)
-	{
-		GetParentDevice()->RegisterGPUWork(1);
-	}
 
-	CommitGraphicsResourceTables();
-	CommitNonComputeShaderConstants();
+	SetupDraw(nullptr, 0, 0);
 
-	// determine 16bit vs 32bit indices
-	const DXGI_FORMAT Format = (IndexBufferStride == sizeof(uint16) ? DXGI_FORMAT_R16_UINT : DXGI_FORMAT_R32_UINT);
+	FD3D12ResourceLocation& ArgumentBufferLocation = SetupIndirectArgument(ArgumentBufferRHI);
+
+	GraphicsCommandList()->ExecuteIndirect(
+		GetParentDevice()->GetParentAdapter()->GetDrawIndirectCommandSignature(),
+		1,
+		ArgumentBufferLocation.GetResource()->GetResource(),
+		ArgumentBufferLocation.GetOffsetFromBaseOfResource() + ArgumentOffset,
+		NULL,
+		0
+	);
 
-	StateCache.SetIndexBuffer(IndexBuffer->ResourceLocation, Format, 0);
+	PostGpuEvent();
+}
 
-	FD3D12ResourceLocation& Location = ArgumentsBuffer->ResourceLocation;
+void FD3D12CommandContext::RHIDrawIndexedIndirect(FRHIBuffer* IndexBufferRHI, FRHIBuffer* ArgumentBufferRHI, int32 DrawArgumentsIndex, uint32 /*NumInstances*/)
+{
+	RHI_DRAW_CALL_INC();
 
-	StateCache.ApplyState(ED3D12PipelineType::Graphics);
+	SetupDraw(IndexBufferRHI, 1);
 
-	// Indirect args buffer can be a previously pending UAV, which becomes PS\Non-PS read. ApplyState will flush pending transitions, so enqueue the indirect
-	// arg transition and flush here.
-	TransitionResource(Location.GetResource(), D3D12_RESOURCE_STATE_TBD, D3D12_RESOURCE_STATE_INDIRECT_ARGUMENT, D3D12_RESOURCE_BARRIER_ALL_SUBRESOURCES);
-	FlushResourceBarriers();	// Must flush so the desired state is actually set.
+	FD3D12ResourceLocation& ArgumentBufferLocation = SetupIndirectArgument(ArgumentBufferRHI);
 
 	GraphicsCommandList()->ExecuteIndirect(
 		GetParentDevice()->GetParentAdapter()->GetDrawIndexedIndirectCommandSignature(),
 		1,
-		Location.GetResource()->GetResource(),
-		Location.GetOffsetFromBaseOfResource() + DrawArgumentsIndex * ArgumentsBufferStride,
+		ArgumentBufferLocation.GetResource()->GetResource(),
+		ArgumentBufferLocation.GetOffsetFromBaseOfResource() + DrawArgumentsIndex * ArgumentBufferRHI->GetStride(),
 		NULL,
 		0
 	);
 
-	UpdateResidency(Location.GetResource());
-
-	ConditionalSplitCommandList();
-
-	DEBUG_EXECUTE_COMMAND_LIST(this);
+	PostGpuEvent();
 }
 
 void FD3D12CommandContext::RHIDrawIndexedPrimitive(FRHIBuffer* IndexBufferRHI, int32 BaseVertexIndex, uint32 FirstInstance, uint32 NumVertices, uint32 StartIndex, uint32 NumPrimitives, uint32 NumInstances)
@@ -1686,14 +1698,6 @@ void FD3D12CommandContext::RHIDrawIndexedPrimitive(FRHIBuffer* IndexBufferRHI, i
 
 	NumInstances = FMath::Max<uint32>(1, NumInstances);
 
-	if (bTrackingEvents)
-	{
-		GetParentDevice()->RegisterGPUWork(NumPrimitives * NumInstances, NumVertices * NumInstances);
-	}
-
-	CommitGraphicsResourceTables();
-	CommitNonComputeShaderConstants();
-
 	uint32 IndexCount = StateCache.GetVertexCount(NumPrimitives);
 
 	// Verify that we are not trying to read outside the index buffer range
@@ -1701,16 +1705,11 @@ void FD3D12CommandContext::RHIDrawIndexedPrimitive(FRHIBuffer* IndexBufferRHI, i
 	checkf((StartIndex + IndexCount) * IndexBuffer->GetStride() <= IndexBuffer->GetSize(),
 		TEXT("Start %u, Count %u, Type %u, Buffer Size %u, Buffer stride %u"), StartIndex, IndexCount, StateCache.GetGraphicsPipelinePrimitiveType(), IndexBuffer->GetSize(), IndexBuffer->GetStride());
 
-	// determine 16bit vs 32bit indices
-	const DXGI_FORMAT Format = (IndexBuffer->GetStride() == sizeof(uint16) ? DXGI_FORMAT_R16_UINT : DXGI_FORMAT_R32_UINT);
-	StateCache.SetIndexBuffer(IndexBuffer->ResourceLocation, Format, 0);
-	StateCache.ApplyState(ED3D12PipelineType::Graphics);
+	SetupDraw(IndexBufferRHI, NumPrimitives * NumInstances, NumVertices * NumInstances);
 
 	GraphicsCommandList()->DrawIndexedInstanced(IndexCount, NumInstances, StartIndex, BaseVertexIndex, FirstInstance);
 
-	ConditionalSplitCommandList();
-
-	DEBUG_EXECUTE_COMMAND_LIST(this);
+	PostGpuEvent();
 }
 
 void FD3D12CommandContext::RHIMultiDrawIndexedPrimitiveIndirect(FRHIBuffer* IndexBufferRHI, FRHIBuffer* ArgumentBufferRHI, uint32 ArgumentOffset, FRHIBuffer* CountBufferRHI, uint32 CountBufferOffset, uint32 MaxDrawArguments)
@@ -1718,17 +1717,11 @@ void FD3D12CommandContext::RHIMultiDrawIndexedPrimitiveIndirect(FRHIBuffer* Inde
 	FD3D12Buffer* IndexBuffer = RetrieveObject<FD3D12Buffer>(IndexBufferRHI);
 
 	// called should make sure the input is valid, this avoid hidden bugs
-	ensure(IndexBufferRHI->GetSize() > 0);
-	ensure(IndexBuffer->ResourceLocation.GetResource() != nullptr);
-
-	if (IndexBufferRHI->GetSize() == 0 || IndexBuffer->ResourceLocation.GetResource() == nullptr)
+	if (!ensure(IndexBufferRHI->GetSize() > 0) || !ensure(IndexBuffer->ResourceLocation.GetResource() != nullptr))
 	{
 		return;
 	}
 
-	const uint32 IndexBufferStride = FD3D12DynamicRHI::ResourceCast(IndexBufferRHI)->GetStride();
-	FD3D12Buffer* ArgumentBuffer = RetrieveObject<FD3D12Buffer>(ArgumentBufferRHI);
-
 	ID3D12Resource* CountBufferResource = nullptr;
 	uint64 CountBufferOffsetFromResourceBase = 0;
 	if (CountBufferRHI)
@@ -1742,41 +1735,21 @@ void FD3D12CommandContext::RHIMultiDrawIndexedPrimitiveIndirect(FRHIBuffer* Inde
 	}
 
 	RHI_DRAW_CALL_INC();
-	if (bTrackingEvents)
-	{
-		GetParentDevice()->RegisterGPUWork(0);
-	}
-
-	CommitGraphicsResourceTables();
-	CommitNonComputeShaderConstants();
-
-	// Set the index buffer.
-	const DXGI_FORMAT Format = (IndexBufferStride == sizeof(uint16) ? DXGI_FORMAT_R16_UINT : DXGI_FORMAT_R32_UINT);
-	StateCache.SetIndexBuffer(IndexBuffer->ResourceLocation, Format, 0);
 
-	FD3D12ResourceLocation& Location = ArgumentBuffer->ResourceLocation;
+	SetupDraw(IndexBufferRHI, 0, 0);
 
-	StateCache.ApplyState(ED3D12PipelineType::Graphics);
-
-	// Indirect args buffer can be a previously pending UAV, which becomes PS\Non-PS read. ApplyState will flush pending transitions, so enqueue the indirect
-	// arg transition and flush here.
-	TransitionResource(Location.GetResource(), D3D12_RESOURCE_STATE_TBD, D3D12_RESOURCE_STATE_INDIRECT_ARGUMENT, D3D12_RESOURCE_BARRIER_ALL_SUBRESOURCES);
-	FlushResourceBarriers();	// Must flush so the desired state is actually set.
+	FD3D12ResourceLocation& ArgumentBufferLocation = SetupIndirectArgument(ArgumentBufferRHI);
 
 	GraphicsCommandList()->ExecuteIndirect(
 		GetParentDevice()->GetParentAdapter()->GetDrawIndexedIndirectCommandSignature(),
 		MaxDrawArguments,
-		Location.GetResource()->GetResource(),
-		Location.GetOffsetFromBaseOfResource() + ArgumentOffset,
+		ArgumentBufferLocation.GetResource()->GetResource(),
+		ArgumentBufferLocation.GetOffsetFromBaseOfResource() + ArgumentOffset,
 		CountBufferResource,
 		CountBufferOffsetFromResourceBase
 	);
 
-	UpdateResidency(Location.GetResource());
-
-	ConditionalSplitCommandList();
-
-	DEBUG_EXECUTE_COMMAND_LIST(this);
+	PostGpuEvent();
 }
 
 void FD3D12CommandContext::RHIDrawIndexedPrimitiveIndirect(FRHIBuffer* IndexBufferRHI, FRHIBuffer* ArgumentBufferRHI, uint32 ArgumentOffset)
@@ -1788,59 +1761,31 @@ void FD3D12CommandContext::RHIDrawIndexedPrimitiveIndirect(FRHIBuffer* IndexBuff
 #if PLATFORM_SUPPORTS_MESH_SHADERS
 void FD3D12CommandContext::RHIDispatchMeshShader(uint32 ThreadGroupCountX, uint32 ThreadGroupCountY, uint32 ThreadGroupCountZ)
 {
-	if (bTrackingEvents)
-	{
-		GetParentDevice()->RegisterGPUDispatch(FIntVector(ThreadGroupCountX, ThreadGroupCountY, ThreadGroupCountZ));
-	}
-
-	CommitGraphicsResourceTables();
-	CommitNonComputeShaderConstants();
-
-	StateCache.ApplyState(ED3D12PipelineType::Graphics);
+	SetupDispatchDraw(ThreadGroupCountX, ThreadGroupCountY, ThreadGroupCountZ);
 
 	GraphicsCommandList6()->DispatchMesh(ThreadGroupCountX, ThreadGroupCountY, ThreadGroupCountZ);
 
-	ConditionalSplitCommandList();
-
-	DEBUG_EXECUTE_COMMAND_LIST(this);
+	PostGpuEvent();
 }
 
 void FD3D12CommandContext::RHIDispatchIndirectMeshShader(FRHIBuffer* ArgumentBufferRHI, uint32 ArgumentOffset)
 {
-	FD3D12Buffer* ArgumentBuffer = RetrieveObject<FD3D12Buffer>(ArgumentBufferRHI);
-
 	RHI_DRAW_CALL_INC();
-	if (bTrackingEvents)
-	{
-		GetParentDevice()->RegisterGPUWork(0);
-	}
-
-	CommitGraphicsResourceTables();
-	CommitNonComputeShaderConstants();
-
-	FD3D12ResourceLocation& Location = ArgumentBuffer->ResourceLocation;
 
-	StateCache.ApplyState(ED3D12PipelineType::Graphics);
+	SetupDispatchDraw(1, 1, 1);
 
-	// Indirect args buffer can be a previously pending UAV, which becomes PS\Non-PS read. ApplyState will flush pending transitions, so enqueue the indirect
-	// arg transition and flush here.
-	TransitionResource(Location.GetResource(), D3D12_RESOURCE_STATE_TBD, D3D12_RESOURCE_STATE_INDIRECT_ARGUMENT, D3D12_RESOURCE_BARRIER_ALL_SUBRESOURCES);
-	FlushResourceBarriers();	// Must flush so the desired state is actually set.
+	FD3D12ResourceLocation& ArgumentBufferLocation = SetupIndirectArgument(ArgumentBufferRHI);
 
 	GraphicsCommandList()->ExecuteIndirect(
 		GetParentDevice()->GetParentAdapter()->GetDispatchIndirectMeshCommandSignature(),
 		1,
-		Location.GetResource()->GetResource(),
-		Location.GetOffsetFromBaseOfResource() + ArgumentOffset,
+		ArgumentBufferLocation.GetResource()->GetResource(),
+		ArgumentBufferLocation.GetOffsetFromBaseOfResource() + ArgumentOffset,
 		NULL,
 		0
 	);
 
-	UpdateResidency(Location.GetResource());
-
-	ConditionalSplitCommandList();
-
-	DEBUG_EXECUTE_COMMAND_LIST(this);
+	PostGpuEvent();
 }
 #endif // PLATFORM_SUPPORTS_MESH_SHADERS
 
diff --git a/Engine/Source/Runtime/D3D12RHI/Private/D3D12DescriptorCache.cpp b/Engine/Source/Runtime/D3D12RHI/Private/D3D12DescriptorCache.cpp
--- a/Engine/Source/Runtime/D3D12RHI/Private/D3D12DescriptorCache.cpp
+++ b/Engine/Source/Runtime/D3D12RHI/Private/D3D12DescriptorCache.cpp
@@ -45,6 +45,10 @@ void FD3D12DescriptorCache::Init(uint32 InNumLocalViewDescriptors, uint32 InNumS
 
 	bBindlessResources = BindlessDescriptorManager.AreResourcesFullyBindless();
 	bBindlessSamplers = BindlessDescriptorManager.AreSamplersFullyBindless();
+
+#if !D3D12RHI_USE_CONSTANT_BUFFER_VIEWS
+	bUsingViewHeap = !bBindlessResources;
+#endif
 #endif
 
 	// Always Init a local sampler heap as the high level cache will always miss initialy
@@ -52,14 +56,16 @@ void FD3D12DescriptorCache::Init(uint32 InNumLocalViewDescriptors, uint32 InNumS
 	// lazily as a backup to save memory)
 	LocalSamplerHeap.Init(IsUsingBindlessSamplers() ? 0 : InNumSamplerDescriptors, ERHIDescriptorHeapType::Sampler);
 
-	NumLocalViewDescriptors = InNumLocalViewDescriptors;
+	NumLocalViewDescriptors = bUsingViewHeap ? InNumLocalViewDescriptors : 0;
 
-	CurrentViewHeap = &SubAllocatedViewHeap;
+	CurrentViewHeap = bUsingViewHeap  ? &SubAllocatedViewHeap : nullptr;
 	CurrentSamplerHeap = IsUsingBindlessSamplers() ? nullptr : &LocalSamplerHeap;
 }
 
 bool FD3D12DescriptorCache::SetDescriptorHeaps(bool bForceHeapChanged)
 {
+	const ERHIPipeline Pipeline = Context.GetPipeline();
+
 	// See if the descriptor heaps changed.
 	bool bHeapChanged = bForceHeapChanged;
 
@@ -67,11 +73,11 @@ bool FD3D12DescriptorCache::SetDescriptorHeaps(bool bForceHeapChanged)
 	FD3D12BindlessDescriptorManager& BindlessDescriptorManager = GetParentDevice()->GetBindlessDescriptorManager();
 	if (IsUsingBindlessResources())
 	{
-		check(BindlessResourcesHeap == BindlessDescriptorManager.GetHeap(ERHIDescriptorHeapType::Standard));
+		check(BindlessResourcesHeap == BindlessDescriptorManager.GetResourceHeap(Pipeline));
 	}
 	if (IsUsingBindlessSamplers())
 	{
-		check(BindlessSamplersHeap == BindlessDescriptorManager.GetHeap(ERHIDescriptorHeapType::Sampler));
+		check(BindlessSamplersHeap == BindlessDescriptorManager.GetSamplerHeap());
 	}
 #endif
 
@@ -140,7 +146,7 @@ void FD3D12DescriptorCache::OpenCommandList()
 #if PLATFORM_SUPPORTS_BINDLESS_RENDERING
 	if (IsUsingBindlessSamplers())
 	{
-		BindlessSamplersHeap = GetParentDevice()->GetBindlessDescriptorManager().GetHeap(ERHIDescriptorHeapType::Sampler);
+		BindlessSamplersHeap = GetParentDevice()->GetBindlessDescriptorManager().GetSamplerHeap();
 	}
 	else
 #endif
@@ -152,11 +158,13 @@ void FD3D12DescriptorCache::OpenCommandList()
 #if PLATFORM_SUPPORTS_BINDLESS_RENDERING
 	if (IsUsingBindlessResources())
 	{
-		BindlessResourcesHeap = GetParentDevice()->GetBindlessDescriptorManager().GetHeap(ERHIDescriptorHeapType::Standard);
+		BindlessResourcesHeap = GetParentDevice()->GetBindlessDescriptorManager().GetResourceHeap(Context.GetPipeline());
 	}
 #endif
-
-	CurrentViewHeap->OpenCommandList();
+	if (CurrentViewHeap)
+	{
+		CurrentViewHeap->OpenCommandList();
+	}
 
 	if (!IsUsingBindlessSamplers())
 	{
@@ -171,7 +179,10 @@ void FD3D12DescriptorCache::OpenCommandList()
 
 void FD3D12DescriptorCache::CloseCommandList()
 {
-	CurrentViewHeap->CloseCommandList();
+	if (CurrentViewHeap)
+	{
+		CurrentViewHeap->CloseCommandList();
+	}
 
 	if (!IsUsingBindlessSamplers())
 	{
@@ -973,7 +984,7 @@ void FD3D12GlobalOnlineSamplerHeap::ConsolidateUniqueSamplerTables(TArrayView<FD
 ////////////////////////////////////////////////////////////////////////////////////////////////////////////////////////
 // FD3D12SubAllocatedOnlineHeap
 
-FD3D12SubAllocatedOnlineHeap::FD3D12SubAllocatedOnlineHeap(FD3D12DescriptorCache& DescriptorCache, FD3D12ContextCommon& Context)
+FD3D12SubAllocatedOnlineHeap::FD3D12SubAllocatedOnlineHeap(FD3D12DescriptorCache& DescriptorCache, FD3D12CommandContext& Context)
 	: FD3D12OnlineHeap(Context.Device, false)
 	, DescriptorCache(DescriptorCache)
 	, Context(Context)
@@ -1027,7 +1038,7 @@ bool FD3D12SubAllocatedOnlineHeap::AllocateBlock()
 	// Extract global heap data
 	if (CurrentBlock)
 	{
-		Heap = new FD3D12DescriptorHeap(OnlineManager.GetDescriptorHeap(), CurrentBlock->BaseSlot, CurrentBlock->Size);
+		Heap = new FD3D12DescriptorHeap(OnlineManager.GetDescriptorHeap(Context.GetPipeline()), CurrentBlock->BaseSlot, CurrentBlock->Size);
 	}
 	else
 	{
diff --git a/Engine/Source/Runtime/D3D12RHI/Private/D3D12DescriptorCache.h b/Engine/Source/Runtime/D3D12RHI/Private/D3D12DescriptorCache.h
--- a/Engine/Source/Runtime/D3D12RHI/Private/D3D12DescriptorCache.h
+++ b/Engine/Source/Runtime/D3D12RHI/Private/D3D12DescriptorCache.h
@@ -223,7 +223,7 @@ private:
 class FD3D12SubAllocatedOnlineHeap : public FD3D12OnlineHeap
 {
 public:
-	FD3D12SubAllocatedOnlineHeap(FD3D12DescriptorCache& DescriptorCache, FD3D12ContextCommon& Context);
+	FD3D12SubAllocatedOnlineHeap(FD3D12DescriptorCache& DescriptorCache, FD3D12CommandContext& Context);
 
 	// Override FD3D12OnlineHeap functions
 	virtual bool RollOver() final override;
@@ -240,7 +240,7 @@ private:
 	FD3D12OnlineDescriptorBlock* CurrentBlock = nullptr;
 
 	FD3D12DescriptorCache& DescriptorCache;
-	FD3D12ContextCommon& Context;
+	FD3D12CommandContext& Context;
 };
 
 
@@ -397,6 +397,7 @@ private:
 
 	TSharedPtr<FD3D12SamplerSet> LocalSamplerSet;
 	bool bHeapsOverridden = false;
+	bool bUsingViewHeap = true;
 
 	uint32 NumLocalViewDescriptors = 0;
 
diff --git a/Engine/Source/Runtime/D3D12RHI/Private/D3D12Descriptors.cpp b/Engine/Source/Runtime/D3D12RHI/Private/D3D12Descriptors.cpp
--- a/Engine/Source/Runtime/D3D12RHI/Private/D3D12Descriptors.cpp
+++ b/Engine/Source/Runtime/D3D12RHI/Private/D3D12Descriptors.cpp
@@ -104,7 +104,7 @@ FD3D12DescriptorManager::FD3D12DescriptorManager(FD3D12Device* Device, FD3D12Des
 
 FD3D12DescriptorManager::~FD3D12DescriptorManager() = default;
 
-void FD3D12DescriptorManager::UpdateImmediately(FRHIDescriptorHandle InHandle, D3D12_CPU_DESCRIPTOR_HANDLE InSourceCpuHandle)
+void FD3D12DescriptorManager::UpdateDescriptorImmediately(FRHIDescriptorHandle InHandle, D3D12_CPU_DESCRIPTOR_HANDLE InSourceCpuHandle)
 {
 	UE::D3D12Descriptors::CopyDescriptor(GetParentDevice(), GetHeap(), InHandle, InSourceCpuHandle);
 }
@@ -125,21 +125,40 @@ void FD3D12OnlineDescriptorManager::Init(uint32 InTotalSize, uint32 InBlockSize,
 #if PLATFORM_SUPPORTS_BINDLESS_RENDERING
 	if (bBindlessResources)
 	{
-		Heap = GetParentDevice()->GetBindlessDescriptorManager().AllocateHeap(ERHIDescriptorHeapType::Standard, InTotalSize);
+#if D3D12RHI_USE_CONSTANT_BUFFER_VIEWS
+		Heaps = GetParentDevice()->GetBindlessDescriptorManager().AllocateResourceHeapsForAllPipelines(InTotalSize);
+		for (ERHIPipeline Pipeline : GetRHIPipelines())
+		{
+			if (Heaps[Pipeline])
+			{
+				INC_DWORD_STAT(STAT_NumViewOnlineDescriptorHeaps);
+				INC_MEMORY_STAT_BY(STAT_ViewOnlineDescriptorHeapMemory, Heaps[Pipeline]->GetMemorySize());
+			}
+		}
+#else
+		// We won't need CBVs, so don't allocate the heap for them
+		return;
+#endif
 	}
 	else
 #endif
 	{
-		Heap = GetParentDevice()->GetDescriptorHeapManager().AllocateHeap(
+		FD3D12DescriptorHeapPtr Heap = GetParentDevice()->GetDescriptorHeapManager().AllocateHeap(
 			TEXT("Device Global - Online View Heap"),
 			ERHIDescriptorHeapType::Standard,
 			InTotalSize,
 			D3D12_DESCRIPTOR_HEAP_FLAG_SHADER_VISIBLE);
+
+		INC_DWORD_STAT(STAT_NumViewOnlineDescriptorHeaps);
+		INC_MEMORY_STAT_BY(STAT_ViewOnlineDescriptorHeapMemory, Heap->GetMemorySize());
+
+		for (ERHIPipeline Pipeline : GetRHIPipelines())
+		{
+			Heaps[Pipeline] = Heap;
+		}
 	}
 
 	// Update the stats
-	INC_DWORD_STAT(STAT_NumViewOnlineDescriptorHeaps);
-	INC_MEMORY_STAT_BY(STAT_ViewOnlineDescriptorHeapMemory, Heap->GetMemorySize());
 	INC_DWORD_STAT_BY(STAT_GlobalViewHeapFreeDescriptors, InTotalSize);
 
 	// Compute amount of free blocks
diff --git a/Engine/Source/Runtime/D3D12RHI/Private/D3D12Descriptors.h b/Engine/Source/Runtime/D3D12RHI/Private/D3D12Descriptors.h
--- a/Engine/Source/Runtime/D3D12RHI/Private/D3D12Descriptors.h
+++ b/Engine/Source/Runtime/D3D12RHI/Private/D3D12Descriptors.h
@@ -93,7 +93,7 @@ public:
 	FD3D12DescriptorManager(FD3D12Device* Device, FD3D12DescriptorHeap* InHeap, TConstArrayView<TStatId> InStats);
 	~FD3D12DescriptorManager();
 
-	void UpdateImmediately(FRHIDescriptorHandle InHandle, D3D12_CPU_DESCRIPTOR_HANDLE InSourceCpuHandle);
+	void UpdateDescriptorImmediately(FRHIDescriptorHandle InHandle, D3D12_CPU_DESCRIPTOR_HANDLE InSourceCpuHandle);
 
 	inline       FD3D12DescriptorHeap* GetHeap()       { return Heap.GetReference(); }
 	inline const FD3D12DescriptorHeap* GetHeap() const { return Heap.GetReference(); }
@@ -141,17 +141,17 @@ public:
 	FD3D12OnlineDescriptorBlock* AllocateHeapBlock();
 	void FreeHeapBlock(FD3D12OnlineDescriptorBlock* InHeapBlock);
 
-	ID3D12DescriptorHeap* GetHeap() { return Heap->GetHeap(); }
-	FD3D12DescriptorHeap* GetDescriptorHeap() { return Heap.GetReference(); }
+	ID3D12DescriptorHeap* GetHeap(ERHIPipeline Pipeline) { return Heaps[Pipeline]->GetHeap(); }
+	FD3D12DescriptorHeap* GetDescriptorHeap(ERHIPipeline Pipeline) { return Heaps[Pipeline].GetReference(); }
 
-	D3D12_CPU_DESCRIPTOR_HANDLE GetCPUSlotHandle(FD3D12OnlineDescriptorBlock* InBlock) const { return Heap->GetCPUSlotHandle(InBlock->BaseSlot); }
-	D3D12_GPU_DESCRIPTOR_HANDLE GetGPUSlotHandle(FD3D12OnlineDescriptorBlock* InBlock) const { return Heap->GetGPUSlotHandle(InBlock->BaseSlot); }
+	D3D12_CPU_DESCRIPTOR_HANDLE GetCPUSlotHandle(ERHIPipeline Pipeline, FD3D12OnlineDescriptorBlock* InBlock) const { return Heaps[Pipeline]->GetCPUSlotHandle(InBlock->BaseSlot); }
+	D3D12_GPU_DESCRIPTOR_HANDLE GetGPUSlotHandle(ERHIPipeline Pipeline, FD3D12OnlineDescriptorBlock* InBlock) const { return Heaps[Pipeline]->GetGPUSlotHandle(InBlock->BaseSlot); }
 	
 	// Called by the EOP task to recycle blocks
 	void Recycle(FD3D12OnlineDescriptorBlock* Block);
 
 private:
-	FD3D12DescriptorHeapPtr Heap;
+	TRHIPipelineArray<FD3D12DescriptorHeapPtr> Heaps;
 
 	TQueue<FD3D12OnlineDescriptorBlock*> FreeBlocks;
 
diff --git a/Engine/Source/Runtime/D3D12RHI/Private/D3D12Device.cpp b/Engine/Source/Runtime/D3D12RHI/Private/D3D12Device.cpp
--- a/Engine/Source/Runtime/D3D12RHI/Private/D3D12Device.cpp
+++ b/Engine/Source/Runtime/D3D12RHI/Private/D3D12Device.cpp
@@ -387,7 +387,12 @@ void FD3D12Device::SetupAfterDeviceCreation()
 		GlobalSamplerHeap.Init(GGlobalSamplerHeapSize);
 	}
 
-	OnlineDescriptorManager.Init(GOnlineDescriptorHeapSize, GOnlineDescriptorHeapBlockSize, bFullyBindlessResources);
+	{
+		const uint32 HeapSize = bFullyBindlessResources ? GBindlessOnlineDescriptorHeapSize : GOnlineDescriptorHeapSize;
+		const uint32 BlockSize = bFullyBindlessResources ? GBindlessOnlineDescriptorHeapBlockSize : GOnlineDescriptorHeapBlockSize;
+
+		OnlineDescriptorManager.Init(HeapSize, BlockSize, bFullyBindlessResources);
+	}
 
 	// Make sure we create the default views before the first command context
 	CreateDefaultViews();
diff --git a/Engine/Source/Runtime/D3D12RHI/Private/D3D12ExplicitDescriptorCache.cpp b/Engine/Source/Runtime/D3D12RHI/Private/D3D12ExplicitDescriptorCache.cpp
--- a/Engine/Source/Runtime/D3D12RHI/Private/D3D12ExplicitDescriptorCache.cpp
+++ b/Engine/Source/Runtime/D3D12RHI/Private/D3D12ExplicitDescriptorCache.cpp
@@ -291,8 +291,8 @@ void FD3D12ExplicitDescriptorCache::Init(uint32 NumViewDescriptors, uint32 NumSa
 #if PLATFORM_SUPPORTS_BINDLESS_RENDERING
 	FD3D12BindlessDescriptorManager& BindlessManager = GetParentDevice()->GetBindlessDescriptorManager();
 
-	bBindlessViews = BindlessManager.HasHeap(ERHIDescriptorHeapType::Standard, BindlessConfig);
-	bBindlessSamplers = BindlessManager.HasHeap(ERHIDescriptorHeapType::Sampler, BindlessConfig);
+	bBindlessViews = BindlessManager.AreResourcesBindless(BindlessConfig);
+	bBindlessSamplers = BindlessManager.AreSamplersBindless(BindlessConfig);
 #else
 	const bool bBindlessViews = false;
 	const bool bBindlessSamplers = false;
@@ -339,8 +339,8 @@ void FD3D12ExplicitDescriptorCache::SetDescriptorHeaps(FD3D12CommandContext& Com
 
 	FD3D12BindlessDescriptorManager& BindlessManager = GetParentDevice()->GetBindlessDescriptorManager();
 
-	ID3D12DescriptorHeap* ViewHeapToSet = bBindlessViews ? BindlessManager.GetHeap(ERHIDescriptorHeapType::Standard)->GetHeap() : ViewHeap.D3D12Heap;
-	ID3D12DescriptorHeap* SamplerHeapToSet = bBindlessSamplers ? BindlessManager.GetHeap(ERHIDescriptorHeapType::Sampler)->GetHeap() : SamplerHeap.D3D12Heap;
+	ID3D12DescriptorHeap* ViewHeapToSet = bBindlessViews ? BindlessManager.GetResourceHeap(CommandContext.GetPipeline())->GetHeap() : ViewHeap.D3D12Heap;
+	ID3D12DescriptorHeap* SamplerHeapToSet = bBindlessSamplers ? BindlessManager.GetSamplerHeap()->GetHeap() : SamplerHeap.D3D12Heap;
 #else
 	check(ViewHeap.GetParentDevice() == CommandContext.GetParentDevice());
 	check(SamplerHeap.GetParentDevice() == CommandContext.GetParentDevice());
diff --git a/Engine/Source/Runtime/D3D12RHI/Private/D3D12PoolAllocator.cpp b/Engine/Source/Runtime/D3D12RHI/Private/D3D12PoolAllocator.cpp
--- a/Engine/Source/Runtime/D3D12RHI/Private/D3D12PoolAllocator.cpp
+++ b/Engine/Source/Runtime/D3D12RHI/Private/D3D12PoolAllocator.cpp
@@ -564,7 +564,7 @@ FRHIMemoryPool* FD3D12PoolAllocator::CreateNewPool(int16 InPoolIndex, uint32 InM
 }
 
 
-bool FD3D12PoolAllocator::HandleDefragRequest(FRHIPoolAllocationData* InSourceBlock, FRHIPoolAllocationData& InTmpTargetBlock)
+bool FD3D12PoolAllocator::HandleDefragRequest(FRHICommandListBase& RHICmdList, FRHIPoolAllocationData* InSourceBlock, FRHIPoolAllocationData& InTmpTargetBlock)
 {
 	// Cache source copy data
 	FD3D12ResourceLocation* Owner = (FD3D12ResourceLocation*)InSourceBlock->GetOwner();
@@ -582,7 +582,7 @@ bool FD3D12PoolAllocator::HandleDefragRequest(FRHIPoolAllocationData* InSourceBl
 	Owner->SetPoolAllocator(this);
 
 	// Notify owner of moved allocation data (recreated resources and SRVs if needed)
-	Owner->OnAllocationMoved(InSourceBlock);
+	Owner->OnAllocationMoved(RHICmdList, InSourceBlock);
 
 	// Add request to unlock the source block on the next fence value (copy operation should have been done by then)
 	FD3D12Adapter* Adapter = GetParentDevice()->GetParentAdapter();
diff --git a/Engine/Source/Runtime/D3D12RHI/Private/D3D12PoolAllocator.h b/Engine/Source/Runtime/D3D12RHI/Private/D3D12PoolAllocator.h
--- a/Engine/Source/Runtime/D3D12RHI/Private/D3D12PoolAllocator.h
+++ b/Engine/Source/Runtime/D3D12RHI/Private/D3D12PoolAllocator.h
@@ -162,7 +162,7 @@ protected:
 
 	// Implementation of FRHIPoolAllocator pure virtuals
 	virtual FRHIMemoryPool* CreateNewPool(int16 InPoolIndex, uint32 InMinimumAllocationSize, ERHIPoolResourceTypes InAllocationResourceType) override;
-	virtual bool HandleDefragRequest(FRHIPoolAllocationData* InSourceBlock, FRHIPoolAllocationData& InTmpTargetBlock) override;
+	virtual bool HandleDefragRequest(FRHICommandListBase& RHICmdList, FRHIPoolAllocationData* InSourceBlock, FRHIPoolAllocationData& InTmpTargetBlock) override;
 	
 	// Placed resource allocation helper function which can be overriden
 	virtual FD3D12Resource* CreatePlacedResource(const FRHIPoolAllocationData& InAllocationData, const FD3D12ResourceDesc& InDesc, D3D12_RESOURCE_STATES InCreateState, ED3D12ResourceStateMode InResourceStateMode, const D3D12_CLEAR_VALUE* InClearValue, const TCHAR* InName);
diff --git a/Engine/Source/Runtime/D3D12RHI/Private/D3D12RHIPrivate.h b/Engine/Source/Runtime/D3D12RHI/Private/D3D12RHIPrivate.h
--- a/Engine/Source/Runtime/D3D12RHI/Private/D3D12RHIPrivate.h
+++ b/Engine/Source/Runtime/D3D12RHI/Private/D3D12RHIPrivate.h
@@ -275,6 +275,7 @@ public:
 		return false;
 	}
 
+	virtual void RHIBeginFrame(FRHICommandListImmediate& RHICmdList) final override;
 	virtual FSamplerStateRHIRef RHICreateSamplerState(const FSamplerStateInitializerRHI& Initializer) final override;
 	virtual FRasterizerStateRHIRef RHICreateRasterizerState(const FRasterizerStateInitializerRHI& Initializer) final override;
 	virtual FDepthStencilStateRHIRef RHICreateDepthStencilState(const FDepthStencilStateInitializerRHI& Initializer) final override;
diff --git a/Engine/Source/Runtime/D3D12RHI/Private/D3D12RayTracing.cpp b/Engine/Source/Runtime/D3D12RHI/Private/D3D12RayTracing.cpp
--- a/Engine/Source/Runtime/D3D12RHI/Private/D3D12RayTracing.cpp
+++ b/Engine/Source/Runtime/D3D12RHI/Private/D3D12RayTracing.cpp
@@ -3088,7 +3088,7 @@ void FD3D12RayTracingGeometry::UnregisterAsRenameListener(uint32 InGPUIndex)
 	bRegisteredAsRenameListener[InGPUIndex] = false;
 }
 
-void FD3D12RayTracingGeometry::ResourceRenamed(FD3D12BaseShaderResource* InRenamedResource, FD3D12ResourceLocation* InNewResourceLocation)
+void FD3D12RayTracingGeometry::ResourceRenamed(FRHICommandListBase& RHICmdList, FD3D12BaseShaderResource* InRenamedResource, FD3D12ResourceLocation* InNewResourceLocation)
 {
 	// Empty resource location is used on destruction of the base shader resource but this
 	// shouldn't happen for RT Geometries because it keeps smart pointers to it's resources.
diff --git a/Engine/Source/Runtime/D3D12RHI/Private/D3D12RayTracing.h b/Engine/Source/Runtime/D3D12RHI/Private/D3D12RayTracing.h
--- a/Engine/Source/Runtime/D3D12RHI/Private/D3D12RayTracing.h
+++ b/Engine/Source/Runtime/D3D12RHI/Private/D3D12RayTracing.h
@@ -43,7 +43,7 @@ public:
 											D3D12_BUILD_RAYTRACING_ACCELERATION_STRUCTURE_DESC& OutDesc, TArrayView<D3D12_RAYTRACING_GEOMETRY_DESC>& OutGeometryDescs) const;
 	
 	// Implement FD3D12ShaderResourceRenameListener interface
-	virtual void ResourceRenamed(FD3D12BaseShaderResource* InRenamedResource, FD3D12ResourceLocation* InNewResourceLocation) override;
+	virtual void ResourceRenamed(FRHICommandListBase& RHICmdList, FD3D12BaseShaderResource* InRenamedResource, FD3D12ResourceLocation* InNewResourceLocation) override;
 
 	void RegisterAsRenameListener(uint32 InGPUIndex);
 	void UnregisterAsRenameListener(uint32 InGPUIndex);
diff --git a/Engine/Source/Runtime/D3D12RHI/Private/D3D12Resources.cpp b/Engine/Source/Runtime/D3D12RHI/Private/D3D12Resources.cpp
--- a/Engine/Source/Runtime/D3D12RHI/Private/D3D12Resources.cpp
+++ b/Engine/Source/Runtime/D3D12RHI/Private/D3D12Resources.cpp
@@ -1243,7 +1243,7 @@ void FD3D12ResourceLocation::AsStandAlone(FD3D12Resource* Resource, uint64 InSiz
 }
 
 
-bool FD3D12ResourceLocation::OnAllocationMoved(FRHIPoolAllocationData* InNewData)
+bool FD3D12ResourceLocation::OnAllocationMoved(FRHICommandListBase& RHICmdList, FRHIPoolAllocationData* InNewData)
 {
 	// Assume linked list allocated for now - only defragging allocator
 	FRHIPoolAllocationData& AllocationData = GetPoolAllocatorPrivateData().PoolData;
@@ -1339,7 +1339,7 @@ bool FD3D12ResourceLocation::OnAllocationMoved(FRHIPoolAllocationData* InNewData
 	check(!CurrentResource->GetDesc().NeedsUAVAliasWorkarounds());
 
 	// Notify all the dependent resources about the change
-	Owner->ResourceRenamed();
+	Owner->ResourceRenamed(RHICmdList);
 
 	return true;
 }
diff --git a/Engine/Source/Runtime/D3D12RHI/Private/D3D12Resources.h b/Engine/Source/Runtime/D3D12RHI/Private/D3D12Resources.h
--- a/Engine/Source/Runtime/D3D12RHI/Private/D3D12Resources.h
+++ b/Engine/Source/Runtime/D3D12RHI/Private/D3D12Resources.h
@@ -598,7 +598,7 @@ public:
 	FD3D12PoolAllocatorPrivateData&    GetPoolAllocatorPrivateData   ()       { return AllocatorData.PoolAllocatorPrivateData;               }
 
 	// Pool allocation specific functions
-	bool OnAllocationMoved(FRHIPoolAllocationData* InNewData);
+	bool OnAllocationMoved(FRHICommandListBase& RHICmdList, FRHIPoolAllocationData* InNewData);
 	void UnlockPoolData();
 
 	bool IsValid() const { return Type != ResourceLocationType::eUndefined; }
@@ -757,7 +757,7 @@ class FD3D12ShaderResourceRenameListener
 protected:
 
 	friend class FD3D12BaseShaderResource;
-	virtual void ResourceRenamed(FD3D12BaseShaderResource* InRenamedResource, FD3D12ResourceLocation* InNewResourceLocation) = 0;
+	virtual void ResourceRenamed(FRHICommandListBase& RHICmdList, FD3D12BaseShaderResource* InRenamedResource, FD3D12ResourceLocation* InNewResourceLocation) = 0;
 };
 
 
@@ -792,12 +792,12 @@ public:
 		return RenameListeners.Num() != 0;
 	}
 
-	void ResourceRenamed()
+	void ResourceRenamed(FRHICommandListBase& RHICmdList)
 	{
 		FScopeLock Lock(&RenameListenersCS);
 		for (FD3D12ShaderResourceRenameListener* RenameListener : RenameListeners)
 		{
-			RenameListener->ResourceRenamed(this, &ResourceLocation);
+			RenameListener->ResourceRenamed(RHICmdList, this, &ResourceLocation);
 		}
 	}
 
@@ -879,8 +879,8 @@ public:
 	}
 #endif
 
-	void Rename(FD3D12ResourceLocation& NewLocation);
-	void RenameLDAChain(FD3D12ResourceLocation& NewLocation);
+	void Rename(FRHICommandListBase& RHICmdList, FD3D12ResourceLocation& NewLocation);
+	void RenameLDAChain(FRHICommandListBase& RHICmdList, FD3D12ResourceLocation& NewLocation);
 
 	void TakeOwnership(FD3D12Buffer& Other);
 	void ReleaseOwnership();
diff --git a/Engine/Source/Runtime/D3D12RHI/Private/D3D12SRV.cpp b/Engine/Source/Runtime/D3D12RHI/Private/D3D12SRV.cpp
--- a/Engine/Source/Runtime/D3D12RHI/Private/D3D12SRV.cpp
+++ b/Engine/Source/Runtime/D3D12RHI/Private/D3D12SRV.cpp
@@ -12,7 +12,7 @@ FD3D12ShaderResourceView::FD3D12ShaderResourceView(FD3D12Device* InDevice)
 	: TD3D12View(InDevice, ERHIDescriptorHeapType::Standard)
 {}
 
-void FD3D12ShaderResourceView::CreateView(FResourceInfo const& InResource, D3D12_SHADER_RESOURCE_VIEW_DESC const& InD3DViewDesc, EFlags InFlags)
+void FD3D12ShaderResourceView::UpdateResourceInfo(const FResourceInfo& InResource, const D3D12_SHADER_RESOURCE_VIEW_DESC& InD3DViewDesc, EFlags InFlags)
 {
 	OffsetInBytes = 0;
 	StrideInBytes = 0;
@@ -42,11 +42,21 @@ void FD3D12ShaderResourceView::CreateView(FResourceInfo const& InResource, D3D12
 		StrideInBytes = 1;
 	}
 #endif
+}
 
+void FD3D12ShaderResourceView::CreateView(FResourceInfo const& InResource, D3D12_SHADER_RESOURCE_VIEW_DESC const& InD3DViewDesc, EFlags InFlags)
+{
+	UpdateResourceInfo(InResource, InD3DViewDesc, InFlags);
 	TD3D12View::CreateView(InResource, InD3DViewDesc);
 }
 
-void FD3D12ShaderResourceView::ResourceRenamed(FD3D12BaseShaderResource* InRenamedResource, FD3D12ResourceLocation* InNewResourceLocation)
+void FD3D12ShaderResourceView::UpdateView(FRHICommandListBase& RHICmdList, const FResourceInfo& InResource, const D3D12_SHADER_RESOURCE_VIEW_DESC& InD3DViewDesc, EFlags InFlags)
+{
+	UpdateResourceInfo(InResource, InD3DViewDesc, InFlags);
+	TD3D12View::UpdateView(RHICmdList, InResource, InD3DViewDesc);
+}
+
+void FD3D12ShaderResourceView::ResourceRenamed(FRHICommandListBase& RHICmdList, FD3D12BaseShaderResource* InRenamedResource, FD3D12ResourceLocation* InNewResourceLocation)
 {
 	check(IsInitialized());
 
@@ -66,10 +76,10 @@ void FD3D12ShaderResourceView::ResourceRenamed(FD3D12BaseShaderResource* InRenam
 	}
 #endif
 
-	TD3D12View::ResourceRenamed(InRenamedResource, InNewResourceLocation);
+	TD3D12View::ResourceRenamed(RHICmdList, InRenamedResource, InNewResourceLocation);
 }
 
-void FD3D12ShaderResourceView::UpdateMinLODClamp(float MinLODClamp)
+void FD3D12ShaderResourceView::UpdateMinLODClamp(FRHICommandListBase& RHICmdList, float MinLODClamp)
 {
 	check(IsInitialized());
 
@@ -84,7 +94,7 @@ void FD3D12ShaderResourceView::UpdateMinLODClamp(float MinLODClamp)
 	}
 
 	UpdateDescriptor();
-	UpdateBindlessSlot(EReason::UpdateOrRename);
+	UpdateBindlessSlot(RHICmdList);
 }
 
 void FD3D12ShaderResourceView::UpdateDescriptor()
@@ -107,139 +117,175 @@ void FD3D12ShaderResourceView::UpdateDescriptor()
 	OfflineCpuHandle.IncrementVersion();
 }
 
-void FD3D12ShaderResourceView_RHI::CreateView()
+static FD3D12ShaderResourceView::EFlags TranslateDesc(D3D12_SHADER_RESOURCE_VIEW_DESC& SRVDesc, FD3D12Buffer* Buffer, const FRHIViewDesc::FBufferSRV::FViewInfo& Info)
 {
-	if (IsBuffer())
+	if (!Info.bNullView)
 	{
-		FD3D12Buffer* Buffer = FD3D12DynamicRHI::ResourceCast(GetBuffer());
-
-		D3D12_SHADER_RESOURCE_VIEW_DESC SRVDesc{};
+		SRVDesc.Shader4ComponentMapping = D3D12_DEFAULT_SHADER_4_COMPONENT_MAPPING;
 
-		auto const Info = ViewDesc.Buffer.SRV.GetViewInfo(Buffer);
-		if (!Info.bNullView)
+		if (Info.BufferType == FRHIViewDesc::EBufferType::AccelerationStructure)
 		{
-			SRVDesc.Shader4ComponentMapping = D3D12_DEFAULT_SHADER_4_COMPONENT_MAPPING;
-
-			if (Info.BufferType == FRHIViewDesc::EBufferType::AccelerationStructure)
-			{
 #if D3D12_RHI_RAYTRACING
-				SRVDesc.ViewDimension = D3D12_SRV_DIMENSION_RAYTRACING_ACCELERATION_STRUCTURE;
-				SRVDesc.Format        = DXGI_FORMAT_UNKNOWN;
+			SRVDesc.ViewDimension = D3D12_SRV_DIMENSION_RAYTRACING_ACCELERATION_STRUCTURE;
+			SRVDesc.Format = DXGI_FORMAT_UNKNOWN;
 
-				SRVDesc.RaytracingAccelerationStructure.Location = Info.OffsetInBytes + Buffer->ResourceLocation.GetGPUVirtualAddress();
+			SRVDesc.RaytracingAccelerationStructure.Location = Info.OffsetInBytes + Buffer->ResourceLocation.GetGPUVirtualAddress();
 #else
-				UE_LOG(LogD3D12RHI, Fatal, TEXT("Raytracing not implemented."));
+			UE_LOG(LogD3D12RHI, Fatal, TEXT("Raytracing not implemented."));
 #endif
-			}
-			else
+		}
+		else
+		{
+			SRVDesc.ViewDimension = D3D12_SRV_DIMENSION_BUFFER;
+			SRVDesc.Format = UE::DXGIUtilities::FindShaderResourceFormat(DXGI_FORMAT(GPixelFormats[Info.Format].PlatformFormat), false);
+			SRVDesc.Buffer.FirstElement = (Info.OffsetInBytes + Buffer->ResourceLocation.GetOffsetFromBaseOfResource()) / Info.StrideInBytes;
+			SRVDesc.Buffer.NumElements = Info.NumElements;
+
+			switch (Info.BufferType)
 			{
-				SRVDesc.ViewDimension       = D3D12_SRV_DIMENSION_BUFFER;
-				SRVDesc.Format              = UE::DXGIUtilities::FindShaderResourceFormat(DXGI_FORMAT(GPixelFormats[Info.Format].PlatformFormat), false);
-				SRVDesc.Buffer.FirstElement = (Info.OffsetInBytes + Buffer->ResourceLocation.GetOffsetFromBaseOfResource()) / Info.StrideInBytes;
-				SRVDesc.Buffer.NumElements  = Info.NumElements;
-
-				switch (Info.BufferType)
-				{
-				case FRHIViewDesc::EBufferType::Raw:
-					SRVDesc.Format       = DXGI_FORMAT_R32_TYPELESS;
-					SRVDesc.Buffer.Flags = D3D12_BUFFER_SRV_FLAG_RAW;
-					break;
-
-				case FRHIViewDesc::EBufferType::Structured:
-					SRVDesc.Buffer.StructureByteStride = Info.StrideInBytes;
-					break;
-
-				case FRHIViewDesc::EBufferType::Typed:
-					// Nothing more to specify
-					break;
-				}
+			case FRHIViewDesc::EBufferType::Raw:
+				SRVDesc.Format = DXGI_FORMAT_R32_TYPELESS;
+				SRVDesc.Buffer.Flags = D3D12_BUFFER_SRV_FLAG_RAW;
+				break;
+
+			case FRHIViewDesc::EBufferType::Structured:
+				SRVDesc.Buffer.StructureByteStride = Info.StrideInBytes;
+				break;
+
+			case FRHIViewDesc::EBufferType::Typed:
+				// Nothing more to specify
+				break;
 			}
 		}
+	}
 
-		FD3D12ShaderResourceView::CreateView(Buffer, SRVDesc, EFlags::None);
+	return FD3D12ShaderResourceView::EFlags::None;
+}
+
+static FD3D12ShaderResourceView::EFlags TranslateDesc(D3D12_SHADER_RESOURCE_VIEW_DESC& SRVDesc, FD3D12Texture* Texture, const FRHIViewDesc::FTextureSRV::FViewInfo& Info)
+{
+	FRHITextureDesc const& TextureDesc = Texture->GetDesc();
+
+	DXGI_FORMAT const ViewFormat = UE::DXGIUtilities::FindShaderResourceFormat(DXGI_FORMAT(GPixelFormats[Info.Format].PlatformFormat), Info.bSRGB);
+	DXGI_FORMAT const BaseFormat = UE::DXGIUtilities::GetPlatformTextureResourceFormat(DXGI_FORMAT(GPixelFormats[TextureDesc.Format].PlatformFormat), TextureDesc.Flags);
+
+	SRVDesc.Format = ViewFormat;
+	SRVDesc.Shader4ComponentMapping = D3D12_DEFAULT_SHADER_4_COMPONENT_MAPPING;
+
+	uint32 const PlaneSlice = UE::DXGIUtilities::GetPlaneSliceFromViewFormat(BaseFormat, ViewFormat);
+	FRHIRange8 const PlaneRange(PlaneSlice, 1);
+
+	// No need to use Info.Dimension, since D3D supports mixing Texture2D view types.
+	// Create a view which matches the underlying resource dimension.
+	switch (TextureDesc.Dimension)
+	{
+	case ETextureDimension::Texture2D:
+		if (TextureDesc.NumSamples > 1)
+		{
+			SRVDesc.ViewDimension = D3D12_SRV_DIMENSION_TEXTURE2DMS;
+		}
+		else
+		{
+			SRVDesc.ViewDimension = D3D12_SRV_DIMENSION_TEXTURE2D;
+			SRVDesc.Texture2D.MostDetailedMip = Info.MipRange.First;
+			SRVDesc.Texture2D.MipLevels = Info.MipRange.Num;
+			SRVDesc.Texture2D.PlaneSlice = PlaneSlice;
+		}
+		break;
+
+	case ETextureDimension::Texture2DArray:
+		if (TextureDesc.NumSamples > 1)
+		{
+			SRVDesc.ViewDimension = D3D12_SRV_DIMENSION_TEXTURE2DMSARRAY;
+			SRVDesc.Texture2DMSArray.FirstArraySlice = Info.ArrayRange.First;
+			SRVDesc.Texture2DMSArray.ArraySize = Info.ArrayRange.Num;
+		}
+		else
+		{
+			SRVDesc.ViewDimension = D3D12_SRV_DIMENSION_TEXTURE2DARRAY;
+			SRVDesc.Texture2DArray.FirstArraySlice = Info.ArrayRange.First;
+			SRVDesc.Texture2DArray.ArraySize = Info.ArrayRange.Num;
+			SRVDesc.Texture2DArray.MostDetailedMip = Info.MipRange.First;
+			SRVDesc.Texture2DArray.MipLevels = Info.MipRange.Num;
+			SRVDesc.Texture2DArray.PlaneSlice = PlaneSlice;
+		}
+		break;
+
+	case ETextureDimension::Texture3D:
+		SRVDesc.ViewDimension = D3D12_SRV_DIMENSION_TEXTURE3D;
+		SRVDesc.Texture3D.MostDetailedMip = Info.MipRange.First;
+		SRVDesc.Texture3D.MipLevels = Info.MipRange.Num;
+		break;
+
+	case ETextureDimension::TextureCube:
+		SRVDesc.ViewDimension = D3D12_SRV_DIMENSION_TEXTURECUBE;
+		SRVDesc.TextureCube.MostDetailedMip = Info.MipRange.First;
+		SRVDesc.TextureCube.MipLevels = Info.MipRange.Num;
+		break;
+
+	case ETextureDimension::TextureCubeArray:
+		SRVDesc.ViewDimension = D3D12_SRV_DIMENSION_TEXTURECUBEARRAY;
+		SRVDesc.TextureCubeArray.MostDetailedMip = Info.MipRange.First;
+		SRVDesc.TextureCubeArray.MipLevels = Info.MipRange.Num;
+		SRVDesc.TextureCubeArray.First2DArrayFace = Info.ArrayRange.First * 6;
+		SRVDesc.TextureCubeArray.NumCubes = Info.ArrayRange.Num;
+		break;
+
+	default:
+		checkNoEntry();
+		break;
+	}
+
+	FD3D12ShaderResourceView::EFlags Flags = FD3D12ShaderResourceView::EFlags::None;
+	if (Texture->SkipsFastClearFinalize())
+	{
+		Flags |= FD3D12ShaderResourceView::EFlags::SkipFastClearFinalize;
+	}
+
+	return Flags;
+}
+
+void FD3D12ShaderResourceView_RHI::CreateView()
+{
+	if (IsBuffer())
+	{
+		FD3D12Buffer* Buffer = FD3D12DynamicRHI::ResourceCast(GetBuffer());
+
+		D3D12_SHADER_RESOURCE_VIEW_DESC SRVDesc{};
+		const EFlags CreateFlags = TranslateDesc(SRVDesc, Buffer, ViewDesc.Buffer.SRV.GetViewInfo(Buffer));
+
+		FD3D12ShaderResourceView::CreateView(Buffer, SRVDesc, CreateFlags);
 	}
 	else
 	{
 		FD3D12Texture* Texture = FD3D12DynamicRHI::ResourceCast(GetTexture());
-		FRHITextureDesc const& TextureDesc = Texture->GetDesc();
 
-		auto const Info = ViewDesc.Texture.SRV.GetViewInfo(Texture);
+		D3D12_SHADER_RESOURCE_VIEW_DESC SRVDesc{};
+		const EFlags CreateFlags = TranslateDesc(SRVDesc, Texture, ViewDesc.Texture.SRV.GetViewInfo(Texture));
 
-		DXGI_FORMAT const ViewFormat = UE::DXGIUtilities::FindShaderResourceFormat(DXGI_FORMAT(GPixelFormats[Info.Format       ].PlatformFormat), Info.bSRGB       );
-		DXGI_FORMAT const BaseFormat = UE::DXGIUtilities::GetPlatformTextureResourceFormat(DXGI_FORMAT(GPixelFormats[TextureDesc.Format].PlatformFormat), TextureDesc.Flags);
+		FD3D12ShaderResourceView::CreateView(Texture, SRVDesc, CreateFlags);
+	}
+}
 
-		D3D12_SHADER_RESOURCE_VIEW_DESC SRVDesc{};
-		SRVDesc.Format = ViewFormat;
-		SRVDesc.Shader4ComponentMapping = D3D12_DEFAULT_SHADER_4_COMPONENT_MAPPING;
+void FD3D12ShaderResourceView_RHI::UpdateView(FRHICommandListBase& RHICmdList)
+{
+	if (IsBuffer())
+	{
+		FD3D12Buffer* Buffer = FD3D12DynamicRHI::ResourceCast(GetBuffer());
 
-		uint32 const PlaneSlice = UE::DXGIUtilities::GetPlaneSliceFromViewFormat(BaseFormat, ViewFormat);
-		FRHIRange8 const PlaneRange(PlaneSlice, 1);
+		D3D12_SHADER_RESOURCE_VIEW_DESC SRVDesc{};
+		const EFlags CreateFlags = TranslateDesc(SRVDesc, Buffer, ViewDesc.Buffer.SRV.GetViewInfo(Buffer));
 
-		// No need to use Info.Dimension, since D3D supports mixing Texture2D view types.
-		// Create a view which matches the underlying resource dimension.
-		switch (TextureDesc.Dimension)
-		{
-		case ETextureDimension::Texture2D:
-			if (TextureDesc.NumSamples > 1)
-			{
-				SRVDesc.ViewDimension = D3D12_SRV_DIMENSION_TEXTURE2DMS;
-			}
-			else
-			{
-				SRVDesc.ViewDimension = D3D12_SRV_DIMENSION_TEXTURE2D;
-				SRVDesc.Texture2D.MostDetailedMip = Info.MipRange.First;
-				SRVDesc.Texture2D.MipLevels       = Info.MipRange.Num;
-				SRVDesc.Texture2D.PlaneSlice      = PlaneSlice;
-			}
-			break;
+		FD3D12ShaderResourceView::UpdateView(RHICmdList, Buffer, SRVDesc, CreateFlags);
+	}
+	else
+	{
+		FD3D12Texture* Texture = FD3D12DynamicRHI::ResourceCast(GetTexture());
 
-		case ETextureDimension::Texture2DArray:
-			if (TextureDesc.NumSamples > 1)
-			{
-				SRVDesc.ViewDimension = D3D12_SRV_DIMENSION_TEXTURE2DMSARRAY;
-				SRVDesc.Texture2DMSArray.FirstArraySlice = Info.ArrayRange.First;
-				SRVDesc.Texture2DMSArray.ArraySize       = Info.ArrayRange.Num;
-			}
-			else
-			{
-				SRVDesc.ViewDimension = D3D12_SRV_DIMENSION_TEXTURE2DARRAY;
-				SRVDesc.Texture2DArray.FirstArraySlice = Info.ArrayRange.First;
-				SRVDesc.Texture2DArray.ArraySize       = Info.ArrayRange.Num;
-				SRVDesc.Texture2DArray.MostDetailedMip = Info.MipRange.First;
-				SRVDesc.Texture2DArray.MipLevels       = Info.MipRange.Num;
-				SRVDesc.Texture2DArray.PlaneSlice      = PlaneSlice;
-			}
-			break;
-
-		case ETextureDimension::Texture3D:
-			SRVDesc.ViewDimension = D3D12_SRV_DIMENSION_TEXTURE3D;
-			SRVDesc.Texture3D.MostDetailedMip = Info.MipRange.First;
-			SRVDesc.Texture3D.MipLevels       = Info.MipRange.Num;
-			break;
-
-		case ETextureDimension::TextureCube:
-			SRVDesc.ViewDimension = D3D12_SRV_DIMENSION_TEXTURECUBE;
-			SRVDesc.TextureCube.MostDetailedMip = Info.MipRange.First;
-			SRVDesc.TextureCube.MipLevels       = Info.MipRange.Num;
-			break;
-
-		case ETextureDimension::TextureCubeArray:
-			SRVDesc.ViewDimension = D3D12_SRV_DIMENSION_TEXTURECUBEARRAY;
-			SRVDesc.TextureCubeArray.MostDetailedMip  = Info.MipRange.First;
-			SRVDesc.TextureCubeArray.MipLevels        = Info.MipRange.Num;
-			SRVDesc.TextureCubeArray.First2DArrayFace = Info.ArrayRange.First * 6;
-			SRVDesc.TextureCubeArray.NumCubes         = Info.ArrayRange.Num;
-			break;
-
-		default:
-			checkNoEntry();
-			break;
-		}
+		D3D12_SHADER_RESOURCE_VIEW_DESC SRVDesc{};
+		const EFlags CreateFlags = TranslateDesc(SRVDesc, Texture, ViewDesc.Texture.SRV.GetViewInfo(Texture));
 
-		FD3D12ShaderResourceView::CreateView(Texture, SRVDesc, Texture->SkipsFastClearFinalize()
-			? EFlags::SkipFastClearFinalize
-			: EFlags::None
-		);
+		FD3D12ShaderResourceView::UpdateView(RHICmdList, Texture, SRVDesc, CreateFlags);
 	}
 }
 
diff --git a/Engine/Source/Runtime/D3D12RHI/Private/D3D12State.cpp b/Engine/Source/Runtime/D3D12RHI/Private/D3D12State.cpp
--- a/Engine/Source/Runtime/D3D12RHI/Private/D3D12State.cpp
+++ b/Engine/Source/Runtime/D3D12RHI/Private/D3D12State.cpp
@@ -674,7 +674,7 @@ FD3D12SamplerState::FD3D12SamplerState(FD3D12Device* InParent, const D3D12_SAMPL
 
 	if (BindlessHandle.IsValid())
 	{
-		GetParentDevice()->GetBindlessDescriptorManager().UpdateImmediately(BindlessHandle, OfflineDescriptor);
+		GetParentDevice()->GetBindlessDescriptorManager().UpdateDescriptorImmediately(BindlessHandle, OfflineDescriptor);
 	}
 #endif
 }
diff --git a/Engine/Source/Runtime/D3D12RHI/Private/D3D12StateCache.cpp b/Engine/Source/Runtime/D3D12RHI/Private/D3D12StateCache.cpp
--- a/Engine/Source/Runtime/D3D12RHI/Private/D3D12StateCache.cpp
+++ b/Engine/Source/Runtime/D3D12RHI/Private/D3D12StateCache.cpp
@@ -60,6 +60,22 @@ static FAutoConsoleVariableRef CVarOnlineDescriptorHeapBlockSize(
 	ECVF_ReadOnly
 );
 
+int32 GBindlessOnlineDescriptorHeapSize = 500 * 1000;
+static FAutoConsoleVariableRef CVarBindlessOnlineDescriptorHeapSize(
+	TEXT("D3D12.BindlessOnlineDescriptorHeapSize"),
+	GBindlessOnlineDescriptorHeapSize,
+	TEXT("Online descriptor heap size"),
+	ECVF_ReadOnly
+);
+
+int32 GBindlessOnlineDescriptorHeapBlockSize = 2000;
+static FAutoConsoleVariableRef CVarBindlessOnlineDescriptorHeapBlockSize(
+	TEXT("D3D12.BindlessOnlineDescriptorHeapBlockSize"),
+	GBindlessOnlineDescriptorHeapBlockSize,
+	TEXT("Block size for sub allocations on the global view descriptor heap."),
+	ECVF_ReadOnly
+);
+
 inline bool operator!=(D3D12_CPU_DESCRIPTOR_HANDLE lhs, D3D12_CPU_DESCRIPTOR_HANDLE rhs)
 {
 	return lhs.ptr != rhs.ptr;
@@ -99,14 +115,17 @@ FD3D12StateCache::FD3D12StateCache(FD3D12CommandContext& Context, FRHIGPUMask No
 	// Cache the resource binding tier
 	ResourceBindingTier = Adapter->GetResourceBindingTier();
 
-	const int32 MaximumResourceHeapSize = Adapter->GetMaxDescriptorsForHeapType(ERHIDescriptorHeapType::Standard);
-	const int32 MaximumSamplerHeapSize = Adapter->GetMaxDescriptorsForHeapType(ERHIDescriptorHeapType::Sampler);
+	const uint32 NumSamplerDescriptors = NUM_SAMPLER_DESCRIPTORS;
 
-	check(GLocalViewHeapSize <= MaximumResourceHeapSize || MaximumResourceHeapSize < 0);
-	check(GOnlineDescriptorHeapSize <= MaximumResourceHeapSize || MaximumResourceHeapSize < 0);
+	checkCode(
+		const int32 MaximumResourceHeapSize = Adapter->GetMaxDescriptorsForHeapType(ERHIDescriptorHeapType::Standard);
+		const int32 MaximumSamplerHeapSize = Adapter->GetMaxDescriptorsForHeapType(ERHIDescriptorHeapType::Sampler);
 
-	const uint32 NumSamplerDescriptors = NUM_SAMPLER_DESCRIPTORS;
-	check(NumSamplerDescriptors <= MaximumSamplerHeapSize);
+		check(GLocalViewHeapSize <= MaximumResourceHeapSize || MaximumResourceHeapSize < 0);
+		check(GOnlineDescriptorHeapSize <= MaximumResourceHeapSize || MaximumResourceHeapSize < 0);
+
+		check(NumSamplerDescriptors <= MaximumSamplerHeapSize);
+	);
 
 	DescriptorCache.Init(GLocalViewHeapSize, NumSamplerDescriptors);
 
@@ -385,7 +404,7 @@ void FD3D12StateCache::InternalSetPipelineState(FD3D12PipelineState* InPipelineS
 	}
 }
 
-void FD3D12StateCache::ApplyState(ED3D12PipelineType PipelineType)
+void FD3D12StateCache::ApplyState(ERHIPipeline HardwarePipe, ED3D12PipelineType PipelineType)
 {
 	//SCOPE_CYCLE_COUNTER(STAT_D3D12ApplyStateTime);
 	const bool bForceState = false;
@@ -434,8 +453,8 @@ void FD3D12StateCache::ApplyState(ED3D12PipelineType PipelineType)
 	{
 		FD3D12BindlessDescriptorManager& BindlessManager = GetParentDevice()->GetBindlessDescriptorManager();
 
-		FD3D12DescriptorHeap* ResourceHeap = BindlessManager.GetHeap(ERHIDescriptorHeapType::Standard, ERHIBindlessConfiguration::AllShaders);
-		FD3D12DescriptorHeap* SamplerHeap = BindlessManager.GetHeap(ERHIDescriptorHeapType::Sampler, ERHIBindlessConfiguration::AllShaders);
+		FD3D12DescriptorHeap* ResourceHeap = BindlessManager.GetResourceHeap(HardwarePipe, ERHIBindlessConfiguration::AllShaders);
+		FD3D12DescriptorHeap* SamplerHeap = BindlessManager.GetSamplerHeap(ERHIBindlessConfiguration::AllShaders);
 
 		checkf(!bBindlessResources || ResourceHeap != nullptr, TEXT("Using dynamic samplers without the bindless sampler heap configured. Please check your configuration."));
 		checkf(!bBindlessSamplers || SamplerHeap != nullptr, TEXT("Using dynamic samplers without the bindless sampler heap configured. Please check your configuration."));
diff --git a/Engine/Source/Runtime/D3D12RHI/Private/D3D12StateCachePrivate.h b/Engine/Source/Runtime/D3D12RHI/Private/D3D12StateCachePrivate.h
--- a/Engine/Source/Runtime/D3D12RHI/Private/D3D12StateCachePrivate.h
+++ b/Engine/Source/Runtime/D3D12RHI/Private/D3D12StateCachePrivate.h
@@ -35,6 +35,8 @@ extern int32 GGlobalSamplerDescriptorHeapSize;
 extern int32 GGlobalSamplerHeapSize;
 extern int32 GOnlineDescriptorHeapSize;
 extern int32 GOnlineDescriptorHeapBlockSize;
+extern int32 GBindlessOnlineDescriptorHeapSize;
+extern int32 GBindlessOnlineDescriptorHeapBlockSize;
 
 enum class ED3D12PipelineType : uint8
 {
@@ -654,7 +656,7 @@ public:
 	ED3D12PipelineType LastComputePipelineType = ED3D12PipelineType::Compute;
 #endif // D3D12_RHI_RAYTRACING
 
-	void ApplyState(ED3D12PipelineType PipelineType);
+	void ApplyState(ERHIPipeline HardwarePipe, ED3D12PipelineType PipelineType);
 	void ApplySamplers(const FD3D12RootSignature* const pRootSignature, uint32 StartStage, uint32 EndStage);
 	void ApplyResources(const FD3D12RootSignature* const pRootSignature, uint32 StartStage, uint32 EndStage);
 	void ApplyBindlessResources(const FD3D12RootSignature* const pRootSignature, uint32 StartStage, uint32 EndStage);
diff --git a/Engine/Source/Runtime/D3D12RHI/Private/D3D12Submission.h b/Engine/Source/Runtime/D3D12RHI/Private/D3D12Submission.h
--- a/Engine/Source/Runtime/D3D12RHI/Private/D3D12Submission.h
+++ b/Engine/Source/Runtime/D3D12RHI/Private/D3D12Submission.h
@@ -5,6 +5,7 @@
 #include "Async/TaskGraphInterfaces.h"
 #include "D3D12RHICommon.h"
 #include "Templates/RefCounting.h"
+#include "D3D12BindlessDescriptors.h"
 
 enum class ED3D12QueueType;
 
diff --git a/Engine/Source/Runtime/D3D12RHI/Private/D3D12Texture.cpp b/Engine/Source/Runtime/D3D12RHI/Private/D3D12Texture.cpp
--- a/Engine/Source/Runtime/D3D12RHI/Private/D3D12Texture.cpp
+++ b/Engine/Source/Runtime/D3D12RHI/Private/D3D12Texture.cpp
@@ -963,13 +963,15 @@ void FD3D12DynamicRHI::RHIUpdateTextureReference(FRHICommandListBase& RHICmdList
 	FRHITexture* NewTexture = InNewTexture ? InNewTexture : FRHITextureReference::GetDefaultTexture();
 
 #if PLATFORM_SUPPORTS_BINDLESS_RENDERING
-	if (FRHIShaderResourceView* TextureRefSRV = TextureRef ? TextureRef->GetBindlessView() : nullptr)
+	if (FRHIShaderResourceView* DestinationSRVRHI = TextureRef ? TextureRef->GetBindlessView() : nullptr)
 	{
-		FD3D12ShaderResourceView* TextureRefSRVD3D12 = ResourceCast(TextureRefSRV);
-		if (TextureRefSRVD3D12->IsBindless())
+		FD3D12ShaderResourceView* DestinationSRV = ResourceCast(DestinationSRVRHI);
+		if (DestinationSRV->IsBindless())
 		{
-			FD3D12BindlessDescriptorManager& BindlessDescriptorManager = TextureRefSRVD3D12->GetParentDevice()->GetBindlessDescriptorManager();
-			BindlessDescriptorManager.UpdateDeferred(TextureRefSRVD3D12->GetBindlessHandle(), ResourceCast(NewTexture)->GetShaderResourceView()->GetOfflineCpuHandle());
+			FD3D12ShaderResourceView* SourceSRV = ResourceCast(NewTexture)->GetShaderResourceView();
+
+			FD3D12BindlessDescriptorManager& BindlessDescriptorManager = DestinationSRV->GetParentDevice()->GetBindlessDescriptorManager();
+			BindlessDescriptorManager.UpdateResourceDescriptor(RHICmdList, DestinationSRV->GetBindlessHandle(), SourceSRV->GetOfflineCpuHandle());
 		}
 	}
 #endif // PLATFORM_SUPPORTS_BINDLESS_RENDERING
diff --git a/Engine/Source/Runtime/D3D12RHI/Private/D3D12UAV.cpp b/Engine/Source/Runtime/D3D12RHI/Private/D3D12UAV.cpp
--- a/Engine/Source/Runtime/D3D12RHI/Private/D3D12UAV.cpp
+++ b/Engine/Source/Runtime/D3D12RHI/Private/D3D12UAV.cpp
@@ -13,7 +13,7 @@ FD3D12UnorderedAccessView::FD3D12UnorderedAccessView(FD3D12Device* InDevice)
 	: TD3D12View(InDevice, ERHIDescriptorHeapType::Standard)
 {}
 
-void FD3D12UnorderedAccessView::CreateView(FResourceInfo const& InResource, D3D12_UNORDERED_ACCESS_VIEW_DESC const& InD3DViewDesc, EFlags InFlags)
+void FD3D12UnorderedAccessView::UpdateResourceInfo(const FResourceInfo& InResource, const D3D12_UNORDERED_ACCESS_VIEW_DESC& InD3DViewDesc, EFlags InFlags)
 {
 	OffsetInBytes = 0;
 	StrideInBytes = 0;
@@ -56,17 +56,29 @@ void FD3D12UnorderedAccessView::CreateView(FResourceInfo const& InResource, D3D1
 			, D3D12_RESOURCE_FLAG_ALLOW_UNORDERED_ACCESS
 		);
 	}
+}
 
+void FD3D12UnorderedAccessView::CreateView(FResourceInfo const& InResource, D3D12_UNORDERED_ACCESS_VIEW_DESC const& InD3DViewDesc, EFlags InFlags)
+{
+	UpdateResourceInfo(InResource, InD3DViewDesc, InFlags);
 	TD3D12View::CreateView(InResource, InD3DViewDesc);
 }
 
-void FD3D12UnorderedAccessView::ResourceRenamed(FD3D12BaseShaderResource* InRenamedResource, FD3D12ResourceLocation* InNewResourceLocation)
+void FD3D12UnorderedAccessView::UpdateView(FRHICommandListBase& RHICmdList, const FResourceInfo& InResource, const D3D12_UNORDERED_ACCESS_VIEW_DESC& InD3DViewDesc, EFlags InFlags)
+{
+	UpdateResourceInfo(InResource, InD3DViewDesc, InFlags);
+	TD3D12View::UpdateView(RHICmdList, InResource, InD3DViewDesc);
+}
+
+void FD3D12UnorderedAccessView::ResourceRenamed(FRHICommandListBase& RHICmdList, FD3D12BaseShaderResource* InRenamedResource, FD3D12ResourceLocation* InNewResourceLocation)
 {
 	// Buffer SRV descriptors contain offsets / GPU virtual addresses which need to be updated to match the new resource location.
 	if (D3DViewDesc.ViewDimension == D3D12_UAV_DIMENSION_BUFFER)
+	{
 		D3DViewDesc.Buffer.FirstElement = (OffsetInBytes + InNewResourceLocation->GetOffsetFromBaseOfResource()) / StrideInBytes;
+	}
 
-	TD3D12View::ResourceRenamed(InRenamedResource, InNewResourceLocation);
+	TD3D12View::ResourceRenamed(RHICmdList, InRenamedResource, InNewResourceLocation);
 }
 
 void FD3D12UnorderedAccessView::UpdateDescriptor()
@@ -89,107 +101,141 @@ void FD3D12UnorderedAccessView::UpdateDescriptor()
 	OfflineCpuHandle.IncrementVersion();
 }
 
-void FD3D12UnorderedAccessView_RHI::CreateView()
+static FD3D12UnorderedAccessView::EFlags TranslateDesc(D3D12_UNORDERED_ACCESS_VIEW_DESC& UAVDesc, FD3D12Buffer* Buffer, const FRHIViewDesc::FBufferUAV::FViewInfo& Info)
 {
-	if (IsBuffer())
+	FD3D12UnorderedAccessView::EFlags Flags = FD3D12UnorderedAccessView::EFlags::None;
+	if (!Info.bNullView)
 	{
-		FD3D12Buffer* Buffer = FD3D12DynamicRHI::ResourceCast(GetBuffer());
+		if (Info.bAppendBuffer || Info.bAtomicCounter)
+		{
+			Flags = FD3D12UnorderedAccessView::EFlags::NeedsCounter;
+		}
 
-		auto const Info = ViewDesc.Buffer.UAV.GetViewInfo(Buffer);
+		UAVDesc.ViewDimension = D3D12_UAV_DIMENSION_BUFFER;
+		UAVDesc.Format = UE::DXGIUtilities::FindUnorderedAccessFormat(DXGI_FORMAT(GPixelFormats[Info.Format].PlatformFormat));
+		UAVDesc.Buffer.FirstElement = (Info.OffsetInBytes + Buffer->ResourceLocation.GetOffsetFromBaseOfResource()) / Info.StrideInBytes;
+		UAVDesc.Buffer.NumElements = Info.NumElements;
 
-		D3D12_UNORDERED_ACCESS_VIEW_DESC UAVDesc{};
-		EFlags Flags = EFlags::None;
-		if (!Info.bNullView)
+		switch (Info.BufferType)
 		{
-			if (Info.bAppendBuffer || Info.bAtomicCounter)
-				Flags = EFlags::NeedsCounter;
+		case FRHIViewDesc::EBufferType::Raw:
+			UAVDesc.Format = DXGI_FORMAT_R32_TYPELESS;
+			UAVDesc.Buffer.Flags = D3D12_BUFFER_UAV_FLAG_RAW;
+			break;
 
-			UAVDesc.ViewDimension       = D3D12_UAV_DIMENSION_BUFFER;
-			UAVDesc.Format              = UE::DXGIUtilities::FindUnorderedAccessFormat(DXGI_FORMAT(GPixelFormats[Info.Format].PlatformFormat));
-			UAVDesc.Buffer.FirstElement = (Info.OffsetInBytes + Buffer->ResourceLocation.GetOffsetFromBaseOfResource()) / Info.StrideInBytes;
-			UAVDesc.Buffer.NumElements  = Info.NumElements;
+		case FRHIViewDesc::EBufferType::Structured:
+			UAVDesc.Buffer.StructureByteStride = Info.StrideInBytes;
+			break;
 
-			switch (Info.BufferType)
-			{
-			case FRHIViewDesc::EBufferType::Raw:
-				UAVDesc.Format       = DXGI_FORMAT_R32_TYPELESS;
-				UAVDesc.Buffer.Flags = D3D12_BUFFER_UAV_FLAG_RAW;
-				break;
-
-			case FRHIViewDesc::EBufferType::Structured:
-				UAVDesc.Buffer.StructureByteStride = Info.StrideInBytes;
-				break;
-
-			case FRHIViewDesc::EBufferType::Typed:
-				// Nothing more to specify
-				break;
-
-			default:
-				checkNoEntry(); // unsupported / unimplemented
-				break;
-			}
-		}
+		case FRHIViewDesc::EBufferType::Typed:
+			// Nothing more to specify
+			break;
 
-		FD3D12UnorderedAccessView::CreateView(Buffer, UAVDesc, Flags);
+		default:
+			checkNoEntry(); // unsupported / unimplemented
+			break;
+		}
 	}
-	else
+
+	return Flags;
+}
+
+static FD3D12UnorderedAccessView::EFlags TranslateDesc(D3D12_UNORDERED_ACCESS_VIEW_DESC& UAVDesc, FD3D12Texture* Texture, const FRHIViewDesc::FTextureUAV::FViewInfo& Info)
+{
+	const FRHITextureDesc& TextureDesc = Texture->GetDesc();
+	check(TextureDesc.NumSamples == 1);
+
+	DXGI_FORMAT const ViewFormat = UE::DXGIUtilities::FindUnorderedAccessFormat(DXGI_FORMAT(GPixelFormats[Info.Format].PlatformFormat));
+	DXGI_FORMAT const BaseFormat = UE::DXGIUtilities::GetPlatformTextureResourceFormat(DXGI_FORMAT(GPixelFormats[TextureDesc.Format].PlatformFormat), TextureDesc.Flags);
+
+	UAVDesc.Format = ViewFormat;
+
+	uint32 const PlaneSlice = UE::DXGIUtilities::GetPlaneSliceFromViewFormat(BaseFormat, ViewFormat);
+	FRHIRange8 const PlaneRange(PlaneSlice, 1);
+
+	// No need to use Info.Dimension, since D3D supports mixing Texture2D view types.
+	// Create a view which matches the underlying resource dimension.
+	switch (TextureDesc.Dimension)
 	{
-		FD3D12Texture* Texture = FD3D12DynamicRHI::ResourceCast(GetTexture());
+	case ETextureDimension::Texture2D:
+		UAVDesc.ViewDimension = D3D12_UAV_DIMENSION_TEXTURE2D;
+		UAVDesc.Texture2D.MipSlice = Info.MipLevel;
+		UAVDesc.Texture2D.PlaneSlice = PlaneSlice;
+		break;
+
+	case ETextureDimension::TextureCube:
+	case ETextureDimension::TextureCubeArray:
+		UAVDesc.ViewDimension = D3D12_UAV_DIMENSION_TEXTURE2DARRAY;
+		UAVDesc.Texture2DArray.FirstArraySlice = Info.ArrayRange.First * 6;
+		UAVDesc.Texture2DArray.ArraySize = Info.ArrayRange.Num * 6;
+		UAVDesc.Texture2DArray.MipSlice = Info.MipLevel;
+		UAVDesc.Texture2DArray.PlaneSlice = PlaneSlice;
+		break;
+
+	case ETextureDimension::Texture2DArray:
+		UAVDesc.ViewDimension = D3D12_UAV_DIMENSION_TEXTURE2DARRAY;
+		UAVDesc.Texture2DArray.FirstArraySlice = Info.ArrayRange.First;
+		UAVDesc.Texture2DArray.ArraySize = Info.ArrayRange.Num;
+		UAVDesc.Texture2DArray.MipSlice = Info.MipLevel;
+		UAVDesc.Texture2DArray.PlaneSlice = PlaneSlice;
+		break;
 
-		FRHITextureDesc const& TextureDesc = Texture->GetDesc();
-		check(TextureDesc.NumSamples == 1);
+	case ETextureDimension::Texture3D:
+		UAVDesc.ViewDimension = D3D12_UAV_DIMENSION_TEXTURE3D;
+		UAVDesc.Texture3D.FirstWSlice = 0;
+		UAVDesc.Texture3D.WSize = FMath::Max(TextureDesc.Depth >> Info.MipLevel, 1);
+		UAVDesc.Texture3D.MipSlice = Info.MipLevel;
+		break;
 
-		auto const Info = ViewDesc.Texture.UAV.GetViewInfo(Texture);
+	default:
+		checkNoEntry();
+		break;
+	}
 
-		DXGI_FORMAT const ViewFormat = UE::DXGIUtilities::FindUnorderedAccessFormat(DXGI_FORMAT(GPixelFormats[Info.Format       ].PlatformFormat));
-		DXGI_FORMAT const BaseFormat = UE::DXGIUtilities::GetPlatformTextureResourceFormat(DXGI_FORMAT(GPixelFormats[TextureDesc.Format].PlatformFormat), TextureDesc.Flags);
+	return FD3D12UnorderedAccessView::EFlags::None;
+}
+
+void FD3D12UnorderedAccessView_RHI::CreateView()
+{
+	if (IsBuffer())
+	{
+		FD3D12Buffer* Buffer = FD3D12DynamicRHI::ResourceCast(GetBuffer());
 
 		D3D12_UNORDERED_ACCESS_VIEW_DESC UAVDesc{};
-		UAVDesc.Format = ViewFormat;
+		const EFlags CreateFlags = TranslateDesc(UAVDesc, Buffer, ViewDesc.Buffer.UAV.GetViewInfo(Buffer));
 
-		uint32 const PlaneSlice = UE::DXGIUtilities::GetPlaneSliceFromViewFormat(BaseFormat, ViewFormat);
-		FRHIRange8 const PlaneRange(PlaneSlice, 1);
+		FD3D12UnorderedAccessView::CreateView(Buffer, UAVDesc, CreateFlags);
+	}
+	else
+	{
+		FD3D12Texture* Texture = FD3D12DynamicRHI::ResourceCast(GetTexture());
 
-		// No need to use Info.Dimension, since D3D supports mixing Texture2D view types.
-		// Create a view which matches the underlying resource dimension.
-		switch (TextureDesc.Dimension)
-		{
-		case ETextureDimension::Texture2D:
-			UAVDesc.ViewDimension = D3D12_UAV_DIMENSION_TEXTURE2D;
-			UAVDesc.Texture2D.MipSlice   = Info.MipLevel;
-			UAVDesc.Texture2D.PlaneSlice = PlaneSlice;
-			break;
+		D3D12_UNORDERED_ACCESS_VIEW_DESC UAVDesc{};
+		const EFlags CreateFlags = TranslateDesc(UAVDesc, Texture, ViewDesc.Texture.UAV.GetViewInfo(Texture));
 
-		case ETextureDimension::TextureCube:
-		case ETextureDimension::TextureCubeArray:
-			UAVDesc.ViewDimension = D3D12_UAV_DIMENSION_TEXTURE2DARRAY;
-			UAVDesc.Texture2DArray.FirstArraySlice = Info.ArrayRange.First * 6;
-			UAVDesc.Texture2DArray.ArraySize       = Info.ArrayRange.Num * 6;
-			UAVDesc.Texture2DArray.MipSlice        = Info.MipLevel;
-			UAVDesc.Texture2DArray.PlaneSlice      = PlaneSlice;
-			break;
+		FD3D12UnorderedAccessView::CreateView(Texture, UAVDesc, CreateFlags);
+	}
+}
 
-		case ETextureDimension::Texture2DArray:
-			UAVDesc.ViewDimension = D3D12_UAV_DIMENSION_TEXTURE2DARRAY;
-			UAVDesc.Texture2DArray.FirstArraySlice = Info.ArrayRange.First;
-			UAVDesc.Texture2DArray.ArraySize       = Info.ArrayRange.Num;
-			UAVDesc.Texture2DArray.MipSlice        = Info.MipLevel;
-			UAVDesc.Texture2DArray.PlaneSlice      = PlaneSlice;
-			break;
+void FD3D12UnorderedAccessView_RHI::UpdateView(FRHICommandListBase& RHICmdList)
+{
+	if (IsBuffer())
+	{
+		FD3D12Buffer* Buffer = FD3D12DynamicRHI::ResourceCast(GetBuffer());
 
-		case ETextureDimension::Texture3D:
-			UAVDesc.ViewDimension = D3D12_UAV_DIMENSION_TEXTURE3D;
-			UAVDesc.Texture3D.FirstWSlice = 0;
-			UAVDesc.Texture3D.WSize       = FMath::Max(TextureDesc.Depth >> Info.MipLevel, 1);
-			UAVDesc.Texture3D.MipSlice    = Info.MipLevel;
-			break;
+		D3D12_UNORDERED_ACCESS_VIEW_DESC UAVDesc{};
+		const EFlags CreateFlags = TranslateDesc(UAVDesc, Buffer, ViewDesc.Buffer.UAV.GetViewInfo(Buffer));
 
-		default:
-			checkNoEntry();
-			break;
-		}
+		FD3D12UnorderedAccessView::UpdateView(RHICmdList, Buffer, UAVDesc, CreateFlags);
+	}
+	else
+	{
+		FD3D12Texture* Texture = FD3D12DynamicRHI::ResourceCast(GetTexture());
+
+		D3D12_UNORDERED_ACCESS_VIEW_DESC UAVDesc{};
+		const EFlags CreateFlags = TranslateDesc(UAVDesc, Texture, ViewDesc.Texture.UAV.GetViewInfo(Texture));
 
-		FD3D12UnorderedAccessView::CreateView(Texture, UAVDesc, EFlags::None);
+		FD3D12UnorderedAccessView::UpdateView(RHICmdList, Texture, UAVDesc, CreateFlags);
 	}
 }
 
@@ -206,7 +252,7 @@ FD3D12UnorderedAccessView_RHI::FD3D12UnorderedAccessView_RHI(FD3D12Device* InDev
 //
 // -----------------------------------------------------------------------------------------------------
 
-FUnorderedAccessViewRHIRef FD3D12DynamicRHI::RHICreateUnorderedAccessView(class FRHICommandListBase& RHICmdList, FRHIViewableResource* Resource, FRHIViewDesc const& ViewDesc)
+FUnorderedAccessViewRHIRef FD3D12DynamicRHI::RHICreateUnorderedAccessView(FRHICommandListBase& RHICmdList, FRHIViewableResource* Resource, FRHIViewDesc const& ViewDesc)
 {
 	FRHIGPUMask RelevantGPUs = ViewDesc.IsBuffer()
 		? FD3D12DynamicRHI::ResourceCast(static_cast<FRHIBuffer* >(Resource))->GetLinkedObjectsGPUMask()
@@ -334,7 +380,7 @@ void FD3D12CommandContext::ClearUAV(TRHICommandList_RecursiveHazardous<FD3D12Com
 #if PLATFORM_SUPPORTS_BINDLESS_RENDERING
 				else
 				{
-					GPUHandle = ParentDevice->GetBindlessDescriptorManager().GetGpuHandle(UAV.GetBindlessHandle());
+					GPUHandle = ParentDevice->GetBindlessDescriptorManager().GetResourceGpuHandle(Context.GetPipeline(), UAV.GetBindlessHandle());
 				}
 #endif
 
diff --git a/Engine/Source/Runtime/D3D12RHI/Private/D3D12View.cpp b/Engine/Source/Runtime/D3D12RHI/Private/D3D12View.cpp
--- a/Engine/Source/Runtime/D3D12RHI/Private/D3D12View.cpp
+++ b/Engine/Source/Runtime/D3D12RHI/Private/D3D12View.cpp
@@ -243,30 +243,31 @@ FD3D12View::~FD3D12View()
 #endif
 }
 
-void FD3D12View::UpdateBindlessSlot(EReason Reason)
+void FD3D12View::InitializeBindlessSlot()
 {
 #if PLATFORM_SUPPORTS_BINDLESS_RENDERING
 	if (BindlessHandle.IsValid())
 	{
 		FD3D12BindlessDescriptorManager& BindlessManager = GetParentDevice()->GetBindlessDescriptorManager();
-		if (Reason == EReason::InitialCreate)
-		{
-			BindlessManager.UpdateImmediately(BindlessHandle, OfflineCpuHandle);
-		}
-		else
-		{
-			BindlessManager.UpdateDeferred(BindlessHandle, OfflineCpuHandle);
-		}
+		BindlessManager.UpdateDescriptorImmediately(BindlessHandle, OfflineCpuHandle);
 	}
 #endif
 }
 
-void FD3D12View::CreateView(FResourceInfo const& InResource, FNullDescPtr NullDescriptor)
+void FD3D12View::UpdateBindlessSlot(FRHICommandListBase& RHICmdList)
 {
-	EReason Reason = IsInitialized()
-		? EReason::UpdateOrRename
-		: EReason::InitialCreate;
+#if PLATFORM_SUPPORTS_BINDLESS_RENDERING
+	if (BindlessHandle.IsValid())
+	{
+		check(BindlessHandle.GetType() == ERHIDescriptorHeapType::Standard);
+		FD3D12BindlessDescriptorManager& BindlessManager = GetParentDevice()->GetBindlessDescriptorManager();
+		BindlessManager.UpdateResourceDescriptor(RHICmdList, BindlessHandle, OfflineCpuHandle);
+	}
+#endif
+}
 
+void FD3D12View::UpdateResourceInfo(FResourceInfo const& InResource, FNullDescPtr NullDescriptor)
+{
 	if (ResourceInfo.BaseResource != InResource.BaseResource)
 	{
 		if (ResourceInfo.BaseResource)
@@ -299,21 +300,30 @@ void FD3D12View::CreateView(FResourceInfo const& InResource, FNullDescPtr NullDe
 
 		OfflineCpuHandle.IncrementVersion();
 	}
+}
 
-	UpdateBindlessSlot(Reason);
+void FD3D12View::CreateView(FResourceInfo const& InResource, FNullDescPtr NullDescriptor)
+{
+	UpdateResourceInfo(InResource, NullDescriptor);
+	InitializeBindlessSlot();
+}
+
+void FD3D12View::UpdateView(FRHICommandListBase& RHICmdList, const FResourceInfo& InResource, FNullDescPtr NullDescriptor)
+{
+	UpdateResourceInfo(InResource, NullDescriptor);
+	UpdateBindlessSlot(RHICmdList);
 }
 
-void FD3D12View::ResourceRenamed(FD3D12BaseShaderResource* InRenamedResource, FD3D12ResourceLocation* InNewResourceLocation)
+void FD3D12View::ResourceRenamed(FRHICommandListBase& RHICmdList, FD3D12BaseShaderResource* InRenamedResource, FD3D12ResourceLocation* InNewResourceLocation)
 {
 	// Can only be called if the base shader resource is not null.
-	check(ResourceInfo.BaseResource == InRenamedResource &&
-		ResourceInfo.ResourceLocation == InNewResourceLocation);
+	check(ResourceInfo.BaseResource == InRenamedResource && ResourceInfo.ResourceLocation == InNewResourceLocation);
 
 	// Update the cached resource pointers
 	ResourceInfo = InRenamedResource;
 
 	UpdateDescriptor();
-	UpdateBindlessSlot(EReason::UpdateOrRename);
+	UpdateBindlessSlot(RHICmdList);
 }
 
 
@@ -339,10 +349,10 @@ void FD3D12ConstantBufferView::CreateView(FResourceInfo const& InResource, uint3
 	TD3D12View::CreateView(InResource, CBVDesc);
 }
 
-void FD3D12ConstantBufferView::ResourceRenamed(FD3D12BaseShaderResource* InRenamedResource, FD3D12ResourceLocation* InNewResourceLocation)
+void FD3D12ConstantBufferView::ResourceRenamed(FRHICommandListBase& RHICmdList, FD3D12BaseShaderResource* InRenamedResource, FD3D12ResourceLocation* InNewResourceLocation)
 {
 	D3DViewDesc.BufferLocation = InNewResourceLocation->GetGPUVirtualAddress() + Offset;
-	TD3D12View::ResourceRenamed(InRenamedResource, InNewResourceLocation);
+	TD3D12View::ResourceRenamed(RHICmdList, InRenamedResource, InNewResourceLocation);
 }
 
 void FD3D12ConstantBufferView::UpdateDescriptor()
diff --git a/Engine/Source/Runtime/D3D12RHI/Private/D3D12View.h b/Engine/Source/Runtime/D3D12RHI/Private/D3D12View.h
--- a/Engine/Source/Runtime/D3D12RHI/Private/D3D12View.h
+++ b/Engine/Source/Runtime/D3D12RHI/Private/D3D12View.h
@@ -223,13 +223,17 @@ protected:
 	FD3D12View(FD3D12Device* InDevice, ERHIDescriptorHeapType InHeapType);
 	virtual ~FD3D12View();
 
-	virtual void ResourceRenamed(FD3D12BaseShaderResource* InRenamedResource, FD3D12ResourceLocation* InNewResourceLocation) override;
+	virtual void ResourceRenamed(FRHICommandListBase& RHICmdList, FD3D12BaseShaderResource* InRenamedResource, FD3D12ResourceLocation* InNewResourceLocation) override;
 	virtual void UpdateDescriptor() = 0;
 
+	void UpdateResourceInfo(FResourceInfo const& InResource, FNullDescPtr NullDescriptor);
 	void CreateView(FResourceInfo const& InResource, FNullDescPtr NullDescriptor);
+	void UpdateView(FRHICommandListBase& RHICmdList, const FResourceInfo& InResource, FNullDescPtr NullDescriptor);
 
 	bool IsInitialized() const { return ResourceInfo.ResourceLocation != nullptr; }
-	void UpdateBindlessSlot(EReason Reason);
+
+	void InitializeBindlessSlot();
+	void UpdateBindlessSlot(FRHICommandListBase& RHICmdList);
 
 	FResourceInfo ResourceInfo;
 	FD3D12ViewSubset ViewSubset;
@@ -259,6 +263,13 @@ protected:
 		FD3D12View::CreateView(InResource, TParent::Null);
 	}
 
+	void UpdateView(FRHICommandListBase& RHICmdList, FResourceInfo const& InResource, TDesc const& InD3DViewDesc)
+	{
+		D3DViewDesc = InD3DViewDesc;
+		ViewSubset.Range = InD3DViewDesc;
+		FD3D12View::UpdateView(RHICmdList, InResource, TParent::Null);
+	}
+
 public:
 	TDesc const& GetD3DDesc() const { return D3DViewDesc; }
 };
@@ -272,7 +283,7 @@ public:
 	void CreateView(FResourceInfo const& InResource, uint32 InOffset, uint32 InAlignedSize);
 
 private:
-	virtual void ResourceRenamed(FD3D12BaseShaderResource* InRenamedResource, FD3D12ResourceLocation* InNewResourceLocation) override;
+	virtual void ResourceRenamed(FRHICommandListBase& RHICmdList, FD3D12BaseShaderResource* InRenamedResource, FD3D12ResourceLocation* InNewResourceLocation) override;
 	virtual void UpdateDescriptor() override;
 
 	uint32 Offset;
@@ -292,12 +303,14 @@ public:
 
 	FD3D12ShaderResourceView(FD3D12Device* InDevice);
 	void CreateView(FResourceInfo const& InResource, D3D12_SHADER_RESOURCE_VIEW_DESC const& InD3DViewDesc, EFlags InFlags);
+	void UpdateView(FRHICommandListBase& RHICmdList, const FResourceInfo& InResource, const D3D12_SHADER_RESOURCE_VIEW_DESC& InD3DViewDesc, EFlags InFlags);
 
 	bool GetSkipFastClearFinalize() const { return EnumHasAnyFlags(Flags, EFlags::SkipFastClearFinalize); }
-	void UpdateMinLODClamp(float MinLODClamp);
+	void UpdateMinLODClamp(FRHICommandListBase& RHICmdList, float MinLODClamp);
 
 protected:
-	virtual void ResourceRenamed(FD3D12BaseShaderResource* InRenamedResource, FD3D12ResourceLocation* InNewResourceLocation) override;
+	void UpdateResourceInfo(const FResourceInfo& InResource, const D3D12_SHADER_RESOURCE_VIEW_DESC& InD3DViewDesc, EFlags InFlags);
+	virtual void ResourceRenamed(FRHICommandListBase& RHICmdList, FD3D12BaseShaderResource* InRenamedResource, FD3D12ResourceLocation* InNewResourceLocation) override;
 	virtual void UpdateDescriptor() override;
 
 	// Required for resource renaming
@@ -323,6 +336,7 @@ public:
 
 	FD3D12UnorderedAccessView(FD3D12Device* InDevice);
 	void CreateView(FResourceInfo const& InResource, D3D12_UNORDERED_ACCESS_VIEW_DESC const& InD3DViewDesc, EFlags InFlags);
+	void UpdateView(FRHICommandListBase& RHICmdList, const FResourceInfo& InResource, const D3D12_UNORDERED_ACCESS_VIEW_DESC& InD3DViewDesc, EFlags InFlags);
 
 	FD3D12Resource* GetCounterResource() const
 	{
@@ -330,7 +344,8 @@ public:
 	}
 
 protected:
-	virtual void ResourceRenamed(FD3D12BaseShaderResource* InRenamedResource, FD3D12ResourceLocation* InNewResourceLocation) override;
+	void UpdateResourceInfo(const FResourceInfo& InResource, const D3D12_UNORDERED_ACCESS_VIEW_DESC& InD3DViewDesc, EFlags InFlags);
+	virtual void ResourceRenamed(FRHICommandListBase& RHICmdList, FD3D12BaseShaderResource* InRenamedResource, FD3D12ResourceLocation* InNewResourceLocation) override;
 	virtual void UpdateDescriptor() override;
 
 	TRefCountPtr<FD3D12Resource> CounterResource;
@@ -410,12 +425,13 @@ public:
 	FD3D12ShaderResourceView_RHI(FD3D12Device* InDevice, FRHIViewableResource* InResource, FRHIViewDesc const& InViewDesc);
 
 	virtual void CreateView();
+	virtual void UpdateView(FRHICommandListBase& RHICmdList);
 
-	virtual void ResourceRenamed(FD3D12BaseShaderResource*, FD3D12ResourceLocation*) override
+	virtual void ResourceRenamed(FRHICommandListBase& RHICmdList, FD3D12BaseShaderResource*, FD3D12ResourceLocation*) override
 	{
 		// Recreate the view from the FRHIViewDesc rather than simply updating the D3D12 descriptor handle from the existing D3D view desc.
 		// This is because the streaming system may have replaced the underlying resource with one that has a different layout.
-		CreateView();
+		UpdateView(RHICmdList);
 	}
 
 	virtual FRHIDescriptorHandle GetBindlessHandle() const override { return FD3D12ShaderResourceView::GetBindlessHandle(); }
@@ -430,12 +446,13 @@ public:
 	FD3D12UnorderedAccessView_RHI(FD3D12Device* InDevice, FRHIViewableResource* InResource, FRHIViewDesc const& InViewDesc);
 
 	virtual void CreateView();
+	virtual void UpdateView(FRHICommandListBase& RHICmdList);
 
-	virtual void ResourceRenamed(FD3D12BaseShaderResource*, FD3D12ResourceLocation*) override
+	virtual void ResourceRenamed(FRHICommandListBase& RHICmdList, FD3D12BaseShaderResource*, FD3D12ResourceLocation*) override
 	{
 		// Recreate the view from the FRHIViewDesc rather than simply updating the D3D12 descriptor handle from the existing D3D view desc.
 		// This is because the streaming system may have replaced the underlying resource with one that has a different layout.
-		CreateView();
+		UpdateView(RHICmdList);
 	}
 
 	virtual FRHIDescriptorHandle GetBindlessHandle() const override { return FD3D12UnorderedAccessView::GetBindlessHandle(); }
diff --git a/Engine/Source/Runtime/D3D12RHI/Private/D3D12Viewport.cpp b/Engine/Source/Runtime/D3D12RHI/Private/D3D12Viewport.cpp
--- a/Engine/Source/Runtime/D3D12RHI/Private/D3D12Viewport.cpp
+++ b/Engine/Source/Runtime/D3D12RHI/Private/D3D12Viewport.cpp
@@ -767,7 +767,7 @@ void FD3D12Viewport::IssueFrameEvent()
 	TArray<FD3D12Payload*> Payloads;
 	for (FD3D12Device* Device : ParentAdapter->GetDevices())
 	{
-		FD3D12ContextCommon& Context = Device->GetDefaultCommandContext();
+		FD3D12CommandContext& Context = Device->GetDefaultCommandContext();
 
 		FD3D12SyncPointRef SyncPoint = FD3D12SyncPoint::Create(ED3D12SyncPointType::GPUAndCPU);
 
diff --git a/Engine/Source/Runtime/D3D12RHI/Private/Windows/WindowsD3D12BindlessDescriptors.h b/Engine/Source/Runtime/D3D12RHI/Private/Windows/WindowsD3D12BindlessDescriptors.h
new file mode 100644
--- /dev/null
+++ b/Engine/Source/Runtime/D3D12RHI/Private/Windows/WindowsD3D12BindlessDescriptors.h
@@ -0,0 +1,5 @@
+// Copyright Epic Games, Inc. All Rights Reserved.
+
+#pragma once
+
+#define D3D12RHI_CUSTOM_BINDLESS_RESOURCE_MANAGER 0
diff --git a/Engine/Source/Runtime/D3D12RHI/Private/Windows/WindowsD3D12Device.cpp b/Engine/Source/Runtime/D3D12RHI/Private/Windows/WindowsD3D12Device.cpp
--- a/Engine/Source/Runtime/D3D12RHI/Private/Windows/WindowsD3D12Device.cpp
+++ b/Engine/Source/Runtime/D3D12RHI/Private/Windows/WindowsD3D12Device.cpp
@@ -112,14 +112,6 @@ static FAutoConsoleVariableRef CVarMinDriverVersionForRayTracingAMD(
 	ECVF_ReadOnly | ECVF_RenderThreadSafe
 );
 
-int32 GAllowAsyncCompute = 1;
-static FAutoConsoleVariableRef CVarAllowAsyncCompute(
-	TEXT("r.D3D12.AllowAsyncCompute"),
-	GAllowAsyncCompute,
-	TEXT("Allow usage of async compute"),
-	ECVF_ReadOnly | ECVF_RenderThreadSafe
-);
-
 #if !UE_BUILD_SHIPPING
 static TAutoConsoleVariable<int32> CVarExperimentalShaderModels(
 	TEXT("r.D3D12.ExperimentalShaderModels"),
@@ -1680,16 +1672,6 @@ void FD3D12DynamicRHI::Init()
 	D3D12_FEATURE_DATA_D3D12_OPTIONS6 options = {};
 	HRESULT Options6HR = GetAdapter().GetD3DDevice()->CheckFeatureSupport(D3D12_FEATURE_D3D12_OPTIONS6, &options, sizeof(options));
 
-	// Allow async compute by default on nVidia cards which support PerPrimitiveShadingRateSupportedWithViewportIndexing 
-	// this should be a good metric according to nVidia itself (this is set for Ampere and newer cards)
-	bool bnVidiaAsyncComputeSupported = false;
-	if (IsRHIDeviceNVIDIA() && Options6HR == S_OK && options.PerPrimitiveShadingRateSupportedWithViewportIndexing)
-	{
-		bnVidiaAsyncComputeSupported = true;
-	}
-
-	GSupportsEfficientAsyncCompute = GAllowAsyncCompute && (FParse::Param(FCommandLine::Get(), TEXT("ForceAsyncCompute")) || (GRHISupportsParallelRHIExecute && (IsRHIDeviceAMD() || bnVidiaAsyncComputeSupported)));
-
 #if WITH_MGPU
 	// Disallow async compute in mGPU mode due to submission order bugs (UE-193929).
 	if (GNumExplicitGPUsForRendering > 1)
diff --git a/Engine/Source/Runtime/NullDrv/Public/NullRHI.h b/Engine/Source/Runtime/NullDrv/Public/NullRHI.h
--- a/Engine/Source/Runtime/NullDrv/Public/NullRHI.h
+++ b/Engine/Source/Runtime/NullDrv/Public/NullRHI.h
@@ -348,6 +348,7 @@ public:
 		return new FNullTexture(Desc);
 	}
 
+	using FDynamicRHI::RHIBeginFrame;
 	virtual void RHIBeginFrame() final override
 	{
 
diff --git a/Engine/Source/Runtime/OpenGLDrv/Public/OpenGLDrv.h b/Engine/Source/Runtime/OpenGLDrv/Public/OpenGLDrv.h
--- a/Engine/Source/Runtime/OpenGLDrv/Public/OpenGLDrv.h
+++ b/Engine/Source/Runtime/OpenGLDrv/Public/OpenGLDrv.h
@@ -433,6 +433,7 @@ public:
 	virtual void RHISubmitCommandsHint() final override;
 	virtual void RHIBeginDrawingViewport(FRHIViewport* Viewport, FRHITexture* RenderTargetRHI) final override;
 	virtual void RHIEndDrawingViewport(FRHIViewport* Viewport, bool bPresent, bool bLockToVsync) final override;
+	using FDynamicRHI::RHIBeginFrame;
 	virtual void RHIBeginFrame() final override;
 	virtual void RHIEndFrame() final override;
 	virtual void RHIBeginScene() final override;
diff --git a/Engine/Source/Runtime/RHI/Private/DynamicRHI.cpp b/Engine/Source/Runtime/RHI/Private/DynamicRHI.cpp
--- a/Engine/Source/Runtime/RHI/Private/DynamicRHI.cpp
+++ b/Engine/Source/Runtime/RHI/Private/DynamicRHI.cpp
@@ -477,6 +477,10 @@ void RHIExit()
 	FRHICommandListImmediate::CleanupGraphEvents();
 }
 
+void FDynamicRHI::RHIBeginFrame(FRHICommandListImmediate& RHICmdList)
+{
+}
+
 // Default fallback; will not work for non-8-bit surfaces and it's extremely slow.
 void FDynamicRHI::RHIReadSurfaceData(FRHITexture* Texture, FIntRect Rect, TArray<FLinearColor>& OutData, FReadSurfaceDataFlags InFlags)
 {
diff --git a/Engine/Source/Runtime/RHI/Private/RHICommandList.cpp b/Engine/Source/Runtime/RHI/Private/RHICommandList.cpp
--- a/Engine/Source/Runtime/RHI/Private/RHICommandList.cpp
+++ b/Engine/Source/Runtime/RHI/Private/RHICommandList.cpp
@@ -1329,6 +1329,9 @@ void FRHICommandListImmediate::EndDrawingViewport(FRHIViewport* Viewport, bool b
 void FRHICommandListImmediate::BeginFrame()
 {
 	check(IsImmediate() && IsInRenderingThread());
+
+	GDynamicRHI->RHIBeginFrame(*this);
+
 	if (Bypass())
 	{
 		ProcessStats();
diff --git a/Engine/Source/Runtime/RHI/Public/DynamicRHI.h b/Engine/Source/Runtime/RHI/Public/DynamicRHI.h
--- a/Engine/Source/Runtime/RHI/Public/DynamicRHI.h
+++ b/Engine/Source/Runtime/RHI/Public/DynamicRHI.h
@@ -135,6 +135,8 @@ public:
 
 	/////// RHI Methods
 
+	RHI_API virtual void RHIBeginFrame(FRHICommandListImmediate& RHICmdList);
+
 	// FlushType: Thread safe
 	virtual FSamplerStateRHIRef RHICreateSamplerState(const FSamplerStateInitializerRHI& Initializer) = 0;
 
diff --git a/Engine/Source/Runtime/RHICore/Private/RHIPoolAllocator.cpp b/Engine/Source/Runtime/RHICore/Private/RHIPoolAllocator.cpp
--- a/Engine/Source/Runtime/RHICore/Private/RHIPoolAllocator.cpp
+++ b/Engine/Source/Runtime/RHICore/Private/RHIPoolAllocator.cpp
@@ -408,7 +408,7 @@ void FRHIMemoryPool::Deallocate(FRHIPoolAllocationData& AllocationData)
 }
 
 
-void FRHIMemoryPool::TryClear(FRHIPoolAllocator* InAllocator, uint32 InMaxCopySize, uint32& CopySize, const TArray<FRHIMemoryPool*>& InTargetPools)
+void FRHIMemoryPool::TryClear(FRHICommandListBase& RHICmdList, FRHIPoolAllocator* InAllocator, uint32 InMaxCopySize, uint32& CopySize, const TArray<FRHIMemoryPool*>& InTargetPools)
 {
 	TRACE_CPUPROFILER_EVENT_SCOPE(FRHIMemoryPool::TryClear);
 
@@ -431,7 +431,7 @@ void FRHIMemoryPool::TryClear(FRHIPoolAllocator* InAllocator, uint32 InMaxCopySi
 				if (TargetPool->TryAllocate(SizeToAllocate, AllocationAlignment, SupportedResourceTypes, TempTargetAllocation))
 				{
 					// RHI specific handling of the actual defrag request
-					InAllocator->HandleDefragRequest(BlockToMove, TempTargetAllocation);
+					InAllocator->HandleDefragRequest(RHICmdList, BlockToMove, TempTargetAllocation);
 
 					// Increment the working copy size
 					CopySize += BlockToMove->GetSize();
@@ -727,7 +727,7 @@ void FRHIPoolAllocator::DeallocateInternal(FRHIPoolAllocationData& AllocationDat
 }
 
 
-void FRHIPoolAllocator::Defrag(uint32 InMaxCopySize, uint32& CurrentCopySize)
+void FRHIPoolAllocator::Defrag(FRHICommandListBase& RHICmdList, uint32 InMaxCopySize, uint32& CurrentCopySize)
 {
 	// Don't do anything when defrag is disabled for this allocator
 	if (!bDefragEnabled)
@@ -783,7 +783,7 @@ void FRHIPoolAllocator::Defrag(uint32 InMaxCopySize, uint32& CurrentCopySize)
 		for (int32 PoolIndex = 0; PoolIndex < SortedTargetPools.Num() - 1; ++PoolIndex)
 		{
 			FRHIMemoryPool* PoolToClear = SortedTargetPools[PoolIndex];
-			PoolToClear->TryClear(this, InMaxCopySize, CurrentCopySize, TargetPools);
+			PoolToClear->TryClear(RHICmdList, this, InMaxCopySize, CurrentCopySize, TargetPools);
 
 			// Remove last allocator since we will try and clear that one next		
 			TargetPools.Pop();
@@ -830,7 +830,7 @@ void FRHIPoolAllocator::Defrag(uint32 InMaxCopySize, uint32& CurrentCopySize)
 				continue;
 			}
 
-			PoolToClear->TryClear(this, InMaxCopySize, CurrentCopySize, TargetPools);
+			PoolToClear->TryClear(RHICmdList, this, InMaxCopySize, CurrentCopySize, TargetPools);
 			DefraggedPoolCount++;
 
 			if (CurrentCopySize >= InMaxCopySize || (DefraggedPoolCount >= GRHIPoolAllocatorDefragMaxPoolsToClear))
diff --git a/Engine/Source/Runtime/RHICore/Public/RHIPoolAllocator.h b/Engine/Source/Runtime/RHICore/Public/RHIPoolAllocator.h
--- a/Engine/Source/Runtime/RHICore/Public/RHIPoolAllocator.h
+++ b/Engine/Source/Runtime/RHICore/Public/RHIPoolAllocator.h
@@ -149,7 +149,7 @@ public:
 	RHICORE_API void Deallocate(FRHIPoolAllocationData& AllocationData);
 
 	// Bookkeeping and clearing
-	RHICORE_API void TryClear(FRHIPoolAllocator* InAllocator, uint32 InMaxCopySize, uint32& CopySize, const TArray<FRHIMemoryPool*>& InTargetPools);
+	RHICORE_API void TryClear(FRHICommandListBase& RHICmdList, FRHIPoolAllocator* InAllocator, uint32 InMaxCopySize, uint32& CopySize, const TArray<FRHIMemoryPool*>& InTargetPools);
 
 	// Getters
 	int16 GetPoolIndex() const { return PoolIndex; }
@@ -224,7 +224,7 @@ public:
 	RHICORE_API void Destroy();
 
 	// Defrag & cleanup operation
-	RHICORE_API void Defrag(uint32 InMaxCopySize, uint32& CurrentCopySize);
+	RHICORE_API void Defrag(FRHICommandListBase& RHICmdList, uint32 InMaxCopySize, uint32& CurrentCopySize);
 
 	// Stats
 	RHICORE_API void UpdateMemoryStats(uint32& IOMemoryAllocated, uint32& IOMemoryUsed, uint32& IOMemoryFree, uint32& IOMemoryEndFree, uint32& IOAlignmentWaste, uint32& IOAllocatedPageCount, uint32& IOFullPageCount);
@@ -240,7 +240,7 @@ protected:
 
 	// Handle a rhi specific defrag op
 	friend class FRHIMemoryPool;
-	virtual bool HandleDefragRequest(FRHIPoolAllocationData* InSourceBlock, FRHIPoolAllocationData& InTmpTargetBlock) = 0;
+	virtual bool HandleDefragRequest(FRHICommandListBase& RHICmdList, FRHIPoolAllocationData* InSourceBlock, FRHIPoolAllocationData& InTmpTargetBlock) = 0;
 
 	// Const creation members - used to create new pools
 	const uint64 DefaultPoolSize;
diff --git a/Engine/Source/Runtime/Windows/D3D11RHI/Private/D3D11RHIPrivate.h b/Engine/Source/Runtime/Windows/D3D11RHI/Private/D3D11RHIPrivate.h
--- a/Engine/Source/Runtime/Windows/D3D11RHI/Private/D3D11RHIPrivate.h
+++ b/Engine/Source/Runtime/Windows/D3D11RHI/Private/D3D11RHIPrivate.h
@@ -575,6 +575,7 @@ public:
 	virtual void RHISubmitCommandsHint() final override;
 	virtual void RHIBeginDrawingViewport(FRHIViewport* Viewport, FRHITexture* RenderTargetRHI) final override;
 	virtual void RHIEndDrawingViewport(FRHIViewport* Viewport, bool bPresent, bool bLockToVsync) final override;
+	using FDynamicRHI::RHIBeginFrame;
 	virtual void RHIBeginFrame() override;
 	virtual void RHIEndFrame() override;
 	virtual void RHIBeginScene() final override;
