Fix crash caused by race condition in RHISetRayTracingBindings when RHI validation is enabled.

#rnx

diff --git a/Engine/Source/Runtime/D3D12RHI/Private/D3D12RayTracing.cpp b/Engine/Source/Runtime/D3D12RHI/Private/D3D12RayTracing.cpp
--- a/Engine/Source/Runtime/D3D12RHI/Private/D3D12RayTracing.cpp
+++ b/Engine/Source/Runtime/D3D12RHI/Private/D3D12RayTracing.cpp
@@ -4075,6 +4075,13 @@ struct FD3D12RayTracingGlobalResourceBinder
 		return CommandContext.GetParentDevice();
 	}
 
+#if ENABLE_RHI_VALIDATION
+	RHIValidation::FTracker* GetValidationTracker()
+	{
+		return CommandContext.Tracker;
+	}
+#endif
+
 	FD3D12CommandContext& CommandContext;
 	FD3D12ExplicitDescriptorCache& DescriptorCache;
 	static constexpr uint32 WorkerIndex = 0;
@@ -4184,6 +4191,15 @@ struct FD3D12RayTracingLocalResourceBinder
 		return &Device;
 	}
 
+#if ENABLE_RHI_VALIDATION
+	RHIValidation::FTracker* GetValidationTracker()
+	{
+		// We can't validate resource states in RHISetRayTracingBindings because there's no command context at that point, and because the states will
+		// change before the raytracing command is dispatched anyway.
+		return nullptr;
+	}
+#endif
+
 	FD3D12Device& Device;
 	FD3D12RayTracingShaderTable& ShaderTable;
 	FD3D12ExplicitDescriptorCache& DescriptorCache;
@@ -4195,7 +4211,6 @@ struct FD3D12RayTracingLocalResourceBinder
 
 template <typename ResourceBinderType>
 static bool SetRayTracingShaderResources(
-	FD3D12CommandContext& CommandContext,
 	const FD3D12RayTracingShader* Shader,
 	uint32 InNumTextures, FRHITexture* const* Textures,
 	uint32 InNumSRVs, FRHIShaderResourceView* const* SRVs,
@@ -4361,7 +4376,7 @@ static bool SetRayTracingShaderResources(
 			, DirtyUniformBuffers
 			, UniformBuffers
 #if ENABLE_RHI_VALIDATION
-			, CommandContext.Tracker
+			, Binder.GetValidationTracker()
 #endif
 		);
 	}
@@ -4496,7 +4511,6 @@ static bool SetRayTracingShaderResources(
 
 template <typename ResourceBinderType>
 static bool SetRayTracingShaderResources(
-	FD3D12CommandContext& CommandContext,
 	const FD3D12RayTracingShader* Shader,
 	const FRayTracingShaderBindings& ResourceBindings,
 	ResourceBinderType& Binder)
@@ -4515,7 +4529,7 @@ static bool SetRayTracingShaderResources(
 		"Ray Tracing Shader Bindings UAV array size must match D3D12 RHI Limit");
 
 	return SetRayTracingShaderResources(
-		CommandContext, Shader,
+		Shader,
 		UE_ARRAY_COUNT(ResourceBindings.Textures), ResourceBindings.Textures,
 		UE_ARRAY_COUNT(ResourceBindings.SRVs), ResourceBindings.SRVs,
 		UE_ARRAY_COUNT(ResourceBindings.UniformBuffers), ResourceBindings.UniformBuffers,
@@ -4624,7 +4638,7 @@ static void DispatchRays(FD3D12CommandContext& CommandContext,
 
 		DescriptorCache->SetDescriptorHeaps(CommandContext);
 		FD3D12RayTracingGlobalResourceBinder ResourceBinder(CommandContext, *DescriptorCache);
-		bResourcesBound = SetRayTracingShaderResources(CommandContext, RayGenShader, GlobalBindings, ResourceBinder);
+		bResourcesBound = SetRayTracingShaderResources(RayGenShader, GlobalBindings, ResourceBinder);
 
 		OptShaderTable->UpdateResidency(CommandContext);
 	}
@@ -4634,7 +4648,7 @@ static void DispatchRays(FD3D12CommandContext& CommandContext,
 		TransientDescriptorCache.Init(MAX_SRVS + MAX_UAVS, MAX_SAMPLERS, ERHIBindlessConfiguration::RayTracingShaders);
 		TransientDescriptorCache.SetDescriptorHeaps(CommandContext);
 		FD3D12RayTracingGlobalResourceBinder ResourceBinder(CommandContext, TransientDescriptorCache);
-		bResourcesBound = SetRayTracingShaderResources(CommandContext, RayGenShader, GlobalBindings, ResourceBinder);
+		bResourcesBound = SetRayTracingShaderResources(RayGenShader, GlobalBindings, ResourceBinder);
 	}
 
 	if (bResourcesBound)
@@ -4758,7 +4772,6 @@ void FD3D12CommandContext::RHIRayTraceDispatchIndirect(FRHIRayTracingPipelineSta
 }
 
 static void SetRayTracingHitGroup(
-	FD3D12CommandContext& CommandContext,
 	FD3D12Device* Device,
 	FD3D12RayTracingShaderTable* ShaderTable,
 	FD3D12RayTracingScene* Scene,
@@ -4830,7 +4843,7 @@ static void SetRayTracingHitGroup(
 	}
 
 	FD3D12RayTracingLocalResourceBinder ResourceBinder(*Device, *ShaderTable, *(Shader->pRootSignature), RecordIndex, WorkerIndex, ERayTracingBindingType::HitGroup);
-	const bool bResourcesBound = SetRayTracingShaderResources(CommandContext, Shader,
+	const bool bResourcesBound = SetRayTracingShaderResources(Shader,
 		0, nullptr, // Textures
 		0, nullptr, // SRVs
 		NumUniformBuffers, UniformBuffers,
@@ -4851,7 +4864,6 @@ static void SetRayTracingHitGroup(
 }
 
 static void SetRayTracingCallableShader(
-	FD3D12CommandContext& CommandContext,
 	FD3D12Device* Device,
 	FD3D12RayTracingShaderTable* ShaderTable,
 	FD3D12RayTracingScene* Scene,
@@ -4876,7 +4888,7 @@ static void SetRayTracingCallableShader(
 		const FD3D12RayTracingShader* Shader = Pipeline->CallableShaders.Shaders[ShaderIndexInPipeline];
 
 		FD3D12RayTracingLocalResourceBinder ResourceBinder(*Device, *ShaderTable, *(Shader->pRootSignature), RecordIndex, WorkerIndex, ERayTracingBindingType::CallableShader);
-		const bool bResourcesBound = SetRayTracingShaderResources(CommandContext, Shader,
+		const bool bResourcesBound = SetRayTracingShaderResources(Shader,
 			0, nullptr, // Textures
 			0, nullptr, // SRVs
 			NumUniformBuffers, UniformBuffers,
@@ -4895,7 +4907,6 @@ static void SetRayTracingCallableShader(
 }
 
 static void SetRayTracingMissShader(
-	FD3D12CommandContext& CommandContext,
 	FD3D12Device* Device,
 	FD3D12RayTracingShaderTable* ShaderTable,
 	FD3D12RayTracingScene* Scene,
@@ -4916,7 +4927,7 @@ static void SetRayTracingMissShader(
 	const FD3D12RayTracingShader* Shader = Pipeline->MissShaders.Shaders[ShaderIndexInPipeline];
 
 	FD3D12RayTracingLocalResourceBinder ResourceBinder(*Device, *ShaderTable, *(Shader->pRootSignature), RecordIndex, WorkerIndex, ERayTracingBindingType::MissShader);
-	const bool bResourcesBound = SetRayTracingShaderResources(CommandContext, Shader,
+	const bool bResourcesBound = SetRayTracingShaderResources(Shader,
 		0, nullptr, // Textures
 		0, nullptr, // SRVs
 		NumUniformBuffers, UniformBuffers,
@@ -4964,13 +4975,13 @@ void FD3D12CommandContext::RHISetRayTracingBindings(
 		TaskContexts.Add(FTaskContext{WorkerIndex});
 	}
 
-	auto BindingTask = [this, Bindings, Device = Device, ShaderTable, Scene, Pipeline, BindingType](const FTaskContext& Context, int32 CurrentIndex)
+	auto BindingTask = [Bindings, Device = Device, ShaderTable, Scene, Pipeline, BindingType](const FTaskContext& Context, int32 CurrentIndex)
 	{
 		const FRayTracingLocalShaderBindings& Binding = Bindings[CurrentIndex];
 
 		if (BindingType == ERayTracingBindingType::HitGroup)
 		{
-			SetRayTracingHitGroup(*this, Device, ShaderTable, Scene, Pipeline,
+			SetRayTracingHitGroup(Device, ShaderTable, Scene, Pipeline,
 				Binding.InstanceIndex,
 				Binding.SegmentIndex,
 				Binding.ShaderSlot,
@@ -4984,7 +4995,7 @@ void FD3D12CommandContext::RHISetRayTracingBindings(
 		}
 		else if (BindingType == ERayTracingBindingType::CallableShader)
 		{
-			SetRayTracingCallableShader(*this, Device, ShaderTable, Scene, Pipeline,
+			SetRayTracingCallableShader(Device, ShaderTable, Scene, Pipeline,
 				Binding.ShaderSlot,
 				Binding.ShaderIndexInPipeline,
 				Binding.NumUniformBuffers,
@@ -4996,7 +5007,7 @@ void FD3D12CommandContext::RHISetRayTracingBindings(
 		}
 		else if (BindingType == ERayTracingBindingType::MissShader)
 		{
-			SetRayTracingMissShader(*this, Device, ShaderTable, Scene, Pipeline,
+			SetRayTracingMissShader(Device, ShaderTable, Scene, Pipeline,
 				Binding.ShaderSlot,
 				Binding.ShaderIndexInPipeline,
 				Binding.NumUniformBuffers,
@@ -5035,7 +5046,7 @@ void FD3D12CommandContext::RHISetRayTracingHitGroup(
 	FD3D12RayTracingShaderTable* ShaderTable = Scene->FindOrCreateShaderTable(Pipeline, GetParentDevice());
 
 	const uint32 WorkerIndex = 0;
-	SetRayTracingHitGroup(*this, GetParentDevice(), ShaderTable, Scene, Pipeline,
+	SetRayTracingHitGroup(GetParentDevice(), ShaderTable, Scene, Pipeline,
 		InstanceIndex,
 		SegmentIndex,
 		ShaderSlot,
@@ -5059,7 +5070,7 @@ void FD3D12CommandContext::RHISetRayTracingCallableShader(
 	FD3D12RayTracingShaderTable* ShaderTable = Scene->FindOrCreateShaderTable(Pipeline, GetParentDevice());
 	const uint32 WorkerIndex = 0;
 
-	SetRayTracingCallableShader(*this, GetParentDevice(), ShaderTable, Scene, Pipeline, ShaderSlotInScene, ShaderIndexInPipeline, NumUniformBuffers, UniformBuffers, 0, nullptr, UserData, WorkerIndex);
+	SetRayTracingCallableShader(GetParentDevice(), ShaderTable, Scene, Pipeline, ShaderSlotInScene, ShaderIndexInPipeline, NumUniformBuffers, UniformBuffers, 0, nullptr, UserData, WorkerIndex);
 }
 
 void FD3D12CommandContext::RHISetRayTracingMissShader(
@@ -5072,7 +5083,7 @@ void FD3D12CommandContext::RHISetRayTracingMissShader(
 	FD3D12RayTracingPipelineState* Pipeline = FD3D12DynamicRHI::ResourceCast(InPipeline);
 	FD3D12RayTracingShaderTable* ShaderTable = Scene->FindOrCreateShaderTable(Pipeline, GetParentDevice());
 	const uint32 WorkerIndex = 0;
-	SetRayTracingMissShader(*this, GetParentDevice(), ShaderTable, Scene, Pipeline,
+	SetRayTracingMissShader(GetParentDevice(), ShaderTable, Scene, Pipeline,
 		ShaderSlotInScene, ShaderIndexInPipeline,
 		NumUniformBuffers, UniformBuffers,
 		0, nullptr, // Loose parameters
