Disable lightweight DRED on Intel cards.

Intel currently suffers a significant performance hit (~10%).

#localization none

diff --git a/Engine/Source/Runtime/D3D12RHI/Private/D3D12Adapter.cpp b/Engine/Source/Runtime/D3D12RHI/Private/D3D12Adapter.cpp
--- a/Engine/Source/Runtime/D3D12RHI/Private/D3D12Adapter.cpp
+++ b/Engine/Source/Runtime/D3D12RHI/Private/D3D12Adapter.cpp
@@ -494,7 +494,7 @@ void FD3D12Adapter::CreateRootDevice(bool bWithDebug)
 #endif // __ID3D12DeviceRemovedExtendedDataSettings1_INTERFACE_DEFINED__
 				}
 
-				else if(GD3D12EnableLightweightDRED)
+				else if(GD3D12EnableLightweightDRED && !IsRHIDeviceIntel()) // Intel suffers a significant performance hit.
 				{
 #ifdef __ID3D12DeviceRemovedExtendedDataSettings2_INTERFACE_DEFINED__
 					TRefCountPtr<ID3D12DeviceRemovedExtendedDataSettings2> DredSettings2;
