PR #10480: Make D3D12 submission timeout configurable

diff --git a/Engine/Source/Runtime/D3D12RHI/Private/D3D12Submission.cpp b/Engine/Source/Runtime/D3D12RHI/Private/D3D12Submission.cpp
--- a/Engine/Source/Runtime/D3D12RHI/Private/D3D12Submission.cpp
+++ b/Engine/Source/Runtime/D3D12RHI/Private/D3D12Submission.cpp
@@ -56,8 +56,12 @@ DECLARE_DWORD_ACCUMULATOR_STAT(TEXT("GPU HS Invocations"), STAT_D3D12RHI_HSInvoc
 DECLARE_DWORD_ACCUMULATOR_STAT(TEXT("GPU DS Invocations"), STAT_D3D12RHI_DSInvocations, STATGROUP_D3D12RHIPipeline);
 DECLARE_DWORD_ACCUMULATOR_STAT(TEXT("GPU CS Invocations"), STAT_D3D12RHI_CSInvocations, STATGROUP_D3D12RHIPipeline);
 
-// The maximum time, in seconds, that a submitted GPU command list is allowed to take before the RHI reports a GPU hang.
-const double SubmissionTimeOutInSeconds = 5.0;
+static float GD3D12SubmissionTimeout = 5.0;
+static FAutoConsoleVariableRef CVarD3D12SubmissionTimeout(
+	TEXT("r.D3D12.SubmissionTimeout"),
+	GD3D12SubmissionTimeout,
+	TEXT("The maximum time, in seconds, that a submitted GPU command list is allowed to take before the RHI reports a GPU hang"),
+	ECVF_RenderThreadSafe);
 
 class FD3D12Thread final : private FRunnable
 {
@@ -1004,7 +1008,7 @@ FD3D12DynamicRHI::FProcessResult FD3D12DynamicRHI::ProcessInterruptQueue()
 				// Detect a hung GPU
 				if (Payload->SubmissionTime.IsSet())
 				{
-					static const uint64 TimeoutCycles = FMath::TruncToInt64(SubmissionTimeOutInSeconds / FPlatformTime::GetSecondsPerCycle64());
+					static const uint64 TimeoutCycles = FMath::TruncToInt64(GD3D12SubmissionTimeout / FPlatformTime::GetSecondsPerCycle64());
 					static const double CyclesPerSecond = 1.0 / FPlatformTime::GetSecondsPerCycle64();
 
 					uint64 ElapsedCycles = FPlatformTime::Cycles64() - Payload->SubmissionTime.GetValue();
