Always log the fatal out of memory message on D3D12 when failing to create a resource.

diff --git a/Engine/Source/Runtime/D3D12RHI/Private/D3D12Util.cpp b/Engine/Source/Runtime/D3D12RHI/Private/D3D12Util.cpp
--- a/Engine/Source/Runtime/D3D12RHI/Private/D3D12Util.cpp
+++ b/Engine/Source/Runtime/D3D12RHI/Private/D3D12Util.cpp
@@ -909,7 +909,8 @@ static void TerminateOnOutOfMemory(ID3D12Device* InDevice, HRESULT D3DResult, bo
 #endif
 
 	static IConsoleVariable* GPUCrashOOM = IConsoleManager::Get().FindConsoleVariable(TEXT("r.GPUCrashOnOutOfMemory"));
-	if (GPUCrashOOM && GPUCrashOOM->GetInt())
+	const bool bGPUCrashOOM = GPUCrashOOM && GPUCrashOOM->GetInt();
+	if (bGPUCrashOOM)
 	{
 		// If no device provided then try and log the DRED status of each device
 		D3D12RHI->ForEachDevice(InDevice, [&](FD3D12Device* IterationDevice)
@@ -920,11 +921,12 @@ static void TerminateOnOutOfMemory(ID3D12Device* InDevice, HRESULT D3DResult, bo
 					LogMemoryInfo(Adapter);
 				}
 			});
-		UE_LOG(LogD3D12RHI, Fatal, TEXT("Out of video memory trying to allocate a rendering resource"));
 	}
-	else
+	
+	UE_LOG(LogD3D12RHI, Fatal, TEXT("Out of video memory trying to allocate a rendering resource"));
+	if (!bGPUCrashOOM)
 	{
-		// Exit silently without reporting a crash because an OOM is not necessarily our fault
+		// Exit silently without reporting a crash because an OOM is not necessarily our fault		
 		FPlatformMisc::RequestExit(true, TEXT("D3D12Util.TerminateOnOutOfMemory"));
 	}
 
