Silence warning V547: Expression is always true.
#rnx

diff --git a/Engine/Source/Runtime/D3D12RHI/Private/D3D12StateCache.cpp b/Engine/Source/Runtime/D3D12RHI/Private/D3D12StateCache.cpp
--- a/Engine/Source/Runtime/D3D12RHI/Private/D3D12StateCache.cpp
+++ b/Engine/Source/Runtime/D3D12RHI/Private/D3D12StateCache.cpp
@@ -588,7 +588,7 @@ void FD3D12StateCache::ApplyState(ERHIPipeline HardwarePipe, ED3D12PipelineType
 		{
 			CmdContext.GraphicsCommandList()->SetComputeRoot32BitConstants(RootConstantsSlot, 4, &UERootConstants[0], 0);
 		}
-		else if (PipelineType == ED3D12PipelineType::Graphics)
+		else if (PipelineType == ED3D12PipelineType::Graphics) //-V547
 		{
 			CmdContext.GraphicsCommandList()->SetGraphicsRoot32BitConstants(RootConstantsSlot, 4, &UERootConstants[0], 0);
 		}
