Fix some color conversion rounding in AgcReadback and OpenGLRenderTarget

diff --git a/Engine/Source/Runtime/Core/Public/Math/Color.h b/Engine/Source/Runtime/Core/Public/Math/Color.h
--- a/Engine/Source/Runtime/Core/Public/Math/Color.h
+++ b/Engine/Source/Runtime/Core/Public/Math/Color.h
@@ -582,28 +582,26 @@ public:
 
 	static uint8 QuantizeUNormFloatTo8( float UnitFloat )
 	{
-		UnitFloat = FMath::Clamp(UnitFloat,0.f,1.f);
-		return (uint8)( 0.5f + UnitFloat * 255.f );
+		return (uint8)( 0.5f + FLinearColor::Clamp01NansTo0(UnitFloat) * 255.f );
 	}
 	
 	static uint16 QuantizeUNormFloatTo16( float UnitFloat )
 	{
-		UnitFloat = FMath::Clamp(UnitFloat,0.f,1.f);
-		return (uint16)( 0.5f + UnitFloat * 65535.f );
+		return (uint16)( 0.5f + FLinearColor::Clamp01NansTo0(UnitFloat) * 65535.f );
 	}
 
 	static float DequantizeUNorm8ToFloat( int Value8 )
 	{
 		check( Value8 >= 0 && Value8 <= 255 );
 
-		return (float)Value8 / 255.f;
+		return (float)Value8 * (1.f/255.f);
 	}
 	
 	static float DequantizeUNorm16ToFloat( int Value16 )
 	{
 		check( Value16 >= 0 && Value16 <= 65535 );
 
-		return (float)Value16 / 65535.f;
+		return (float)Value16 * (1.f/65535.f);
 	}
 
 	static uint8 Requantize10to8( int Value10 )
@@ -667,7 +665,8 @@ public:
 	 */
 	FORCEINLINE FLinearColor ReinterpretAsLinear() const
 	{
-		return FLinearColor(R / 255.f, G / 255.f, B / 255.f, A / 255.f);
+		constexpr float inv255 = 1.f / 255.f;
+		return FLinearColor(R * inv255, G * inv255, B * inv255, A * inv255);
 	}
 
 	/**
diff --git a/Engine/Source/Runtime/D3D12RHI/Private/D3D12RenderTarget.cpp b/Engine/Source/Runtime/D3D12RHI/Private/D3D12RenderTarget.cpp
--- a/Engine/Source/Runtime/D3D12RHI/Private/D3D12RenderTarget.cpp
+++ b/Engine/Source/Runtime/D3D12RHI/Private/D3D12RenderTarget.cpp
@@ -603,7 +603,7 @@ static void ConvertRAWSurfaceDataToFLinearColor(EPixelFormat Format, uint32 Widt
 			for (uint32 X = 0; X < Width; X++)
 			{
 				uint16 Value16 = *SrcPtr;
-				float Value = Value16 / (float)(0xffff);
+				float Value = Value16 * (1.f/0xffff);
 
 				*DestPtr = FLinearColor(Value, Value, Value);
 				++SrcPtr;
@@ -620,7 +620,7 @@ static void ConvertRAWSurfaceDataToFLinearColor(EPixelFormat Format, uint32 Widt
 			FLinearColor* DestPtr = Out + Y * Width;
 			for (uint32 X = 0; X < Width; X++)
 			{
-				FColor sRGBColor = FColor(SrcPtr->B, SrcPtr->G, SrcPtr->R, SrcPtr->A);
+				FColor sRGBColor = FColor(SrcPtr->B, SrcPtr->G, SrcPtr->R, SrcPtr->A); // swap RB
 				*DestPtr = FLinearColor(sRGBColor);
 				++SrcPtr;
 				++DestPtr;
@@ -635,8 +635,7 @@ static void ConvertRAWSurfaceDataToFLinearColor(EPixelFormat Format, uint32 Widt
 			FLinearColor* DestPtr = Out + Y * Width;
 			for (uint32 X = 0; X < Width; X++)
 			{
-				FColor sRGBColor = FColor(SrcPtr->R, SrcPtr->G, SrcPtr->B, SrcPtr->A);
-				*DestPtr = FLinearColor(sRGBColor);
+				*DestPtr = FLinearColor(*SrcPtr);
 				++SrcPtr;
 				++DestPtr;
 			}
@@ -668,14 +667,12 @@ static void ConvertRAWSurfaceDataToFLinearColor(EPixelFormat Format, uint32 Widt
 		{
 			for (uint32 Y = 0; Y < Height; Y++)
 			{
-				FFloat16* SrcPtr = (FFloat16*)(In + Y * SrcPitch);
+				const FFloat16Color* SrcPtr = (FFloat16Color*)(In + Y * SrcPitch);
 				FLinearColor* DestPtr = Out + Y * Width;
 
 				for (uint32 X = 0; X < Width; X++)
 				{
-					*DestPtr = FLinearColor((float)SrcPtr[0], (float)SrcPtr[1], (float)SrcPtr[2], (float)SrcPtr[3]);
-					++DestPtr;
-					SrcPtr += 4;
+					DestPtr[X] = SrcPtr[X].GetFloats();
 				}
 			}
 		}
diff --git a/Engine/Source/Runtime/OpenGLDrv/Private/OpenGLRenderTarget.cpp b/Engine/Source/Runtime/OpenGLDrv/Private/OpenGLRenderTarget.cpp
--- a/Engine/Source/Runtime/OpenGLDrv/Private/OpenGLRenderTarget.cpp
+++ b/Engine/Source/Runtime/OpenGLDrv/Private/OpenGLRenderTarget.cpp
@@ -594,7 +594,7 @@ void FOpenGLDynamicRHI::ReadSurfaceDataRaw(FOpenGLContextState& ContextState, FR
 		uint8* TargetPtr = TargetBuffer;
 		for( int32 DepthValueIndex = 0; DepthValueIndex < DepthValueCount; ++DepthValueIndex )
 		{
-			uint8 Value = (uint8)( *DataPtr++ * 255.0f );
+			uint8 Value = FColor::QuantizeUNormFloatTo8( *DataPtr++ );
 			*TargetPtr++ = Value;
 			*TargetPtr++ = Value;
 			*TargetPtr++ = Value;
diff --git a/Engine/Source/Runtime/RHI/Public/RHICommandList.h b/Engine/Source/Runtime/RHI/Public/RHICommandList.h
--- a/Engine/Source/Runtime/RHI/Public/RHICommandList.h
+++ b/Engine/Source/Runtime/RHI/Public/RHICommandList.h
@@ -4462,7 +4462,9 @@ public:
 		LLM_SCOPE(ELLMTag::Textures);
 		GDynamicRHI->RHIUnlockTextureCubeFace_RenderThread(*this, Texture, FaceIndex, ArrayIndex, MipIndex, bLockWithinMiptail);
 	}
-
+	
+	// ReadSurfaceFloatData reads texture data into FColor
+	//	pixels in other formats are converted to FColor
 	FORCEINLINE void ReadSurfaceData(FRHITexture* Texture,FIntRect Rect,TArray<FColor>& OutData,FReadSurfaceDataFlags InFlags)
 	{
 		QUICK_SCOPE_CYCLE_COUNTER(STAT_RHIMETHOD_ReadSurfaceData_Flush);
@@ -4470,7 +4472,10 @@ public:
 		ImmediateFlush(EImmediateFlushType::FlushRHIThread);  
 		GDynamicRHI->RHIReadSurfaceData(Texture,Rect,OutData,InFlags);
 	}
-
+	
+	// ReadSurfaceFloatData reads texture data into FLinearColor
+	//	pixels in other formats are converted to FLinearColor
+	// reading from float surfaces remaps the values into an interpolation of their {min,max} ; use RCM_MinMax to prevent that
 	FORCEINLINE void ReadSurfaceData(FRHITexture* Texture, FIntRect Rect, TArray<FLinearColor>& OutData, FReadSurfaceDataFlags InFlags)
 	{
 		QUICK_SCOPE_CYCLE_COUNTER(STAT_RHIMETHOD_ReadSurfaceData_Flush);
@@ -4497,6 +4502,9 @@ public:
 		GDynamicRHI->RHIUnmapStagingSurface_RenderThread(*this, Texture, GPUIndex);
 	}
 	
+	// ReadSurfaceFloatData reads texture data into FFloat16Color
+	//	it only works if Texture is exactly PF_FloatRGBA (RGBA16F) !
+	//	no conversion is done
 	FORCEINLINE void ReadSurfaceFloatData(FRHITexture* Texture,FIntRect Rect,TArray<FFloat16Color>& OutData,ECubeFace CubeFace,int32 ArrayIndex,int32 MipIndex)
 	{
 		LLM_SCOPE(ELLMTag::Textures);
