Removing unnecessary includes of ShaderCompiler.h. Also removing a few other unnecessary includes.

diff --git a/Engine/Source/Runtime/D3D12RHI/Private/D3D12Commands.cpp b/Engine/Source/Runtime/D3D12RHI/Private/D3D12Commands.cpp
--- a/Engine/Source/Runtime/D3D12RHI/Private/D3D12Commands.cpp
+++ b/Engine/Source/Runtime/D3D12RHI/Private/D3D12Commands.cpp
@@ -5,20 +5,8 @@ D3D12Commands.cpp: D3D RHI commands implementation.
 =============================================================================*/
 
 #include "D3D12RHIPrivate.h"
-#include "StaticBoundShaderState.h"
-#include "GlobalShader.h"
-#include "OneColorShader.h"
-#include "RHICommandList.h"
-#include "RHIStaticStates.h"
-#include "ShaderParameterUtils.h"
-#include "ShaderCompiler.h"
-#include "ScreenRendering.h"
-#include "ResolveShader.h"
-#include "SceneUtils.h"
-#include "RenderUtils.h"
-#include "GlobalRenderResources.h"
-#include "RHIShaderParametersShared.h"
 #include "ProfilingDebugging/AssetMetadataTrace.h"
+#include "ProfilingDebugging/RealtimeGPUProfiler.h"
 
 static int32 GD3D12TransientAllocatorFullAliasingBarrier = 0;
 static FAutoConsoleVariableRef CVarD3D12TransientAllocatorFullAliasingBarrier(
diff --git a/Engine/Source/Runtime/SlateRHIRenderer/Private/SlateRHIRenderer.cpp b/Engine/Source/Runtime/SlateRHIRenderer/Private/SlateRHIRenderer.cpp
--- a/Engine/Source/Runtime/SlateRHIRenderer/Private/SlateRHIRenderer.cpp
+++ b/Engine/Source/Runtime/SlateRHIRenderer/Private/SlateRHIRenderer.cpp
@@ -36,7 +36,6 @@
 #include "SceneUtils.h"
 #include "TextureResource.h"
 #include "VolumeRendering.h"
-#include "ShaderCompiler.h"
 #include "PipelineStateCache.h"
 #include "EngineModule.h"
 #include "Interfaces/ISlate3DRenderer.h"
@@ -54,6 +53,10 @@
 #include "RenderCore.h"
 #include "DataDrivenShaderPlatformInfo.h"
 
+#if WITH_EDITORONLY_DATA
+#include "ShaderCompiler.h"
+#endif
+
 DECLARE_CYCLE_STAT(TEXT("Slate RT: Rendering"), STAT_SlateRenderingRTTime, STATGROUP_Slate);
 
 DECLARE_CYCLE_STAT(TEXT("Slate RT: Draw Batches"), STAT_SlateRTDrawBatches, STATGROUP_Slate);
diff --git a/Engine/Source/Runtime/Windows/D3D11RHI/Private/D3D11RHIPrivate.h b/Engine/Source/Runtime/Windows/D3D11RHI/Private/D3D11RHIPrivate.h
--- a/Engine/Source/Runtime/Windows/D3D11RHI/Private/D3D11RHIPrivate.h
+++ b/Engine/Source/Runtime/Windows/D3D11RHI/Private/D3D11RHIPrivate.h
@@ -34,8 +34,6 @@ DECLARE_LOG_CATEGORY_EXTERN(LogD3D11RHI, Log, All);
 #include "RHIValidationCommon.h"
 #include "RHICoreShader.h"
 
-#include "ShaderCompiler.h"
-
 #ifndef WITH_DX_PERF
 #define WITH_DX_PERF	1
 #endif
