un-duplicate code in D3D11 and 12 rendertarget pixel conversion routines

diff --git a/Engine/Source/Runtime/D3D12RHI/Private/D3D12RenderTarget.cpp b/Engine/Source/Runtime/D3D12RHI/Private/D3D12RenderTarget.cpp
--- a/Engine/Source/Runtime/D3D12RHI/Private/D3D12RenderTarget.cpp
+++ b/Engine/Source/Runtime/D3D12RHI/Private/D3D12RenderTarget.cpp
@@ -15,7 +15,7 @@
 #include "RHISurfaceDataConversion.h"
 #include "CommonRenderResources.h"
 
-static uint32 ComputeBytesPerPixel(DXGI_FORMAT Format)
+static uint32 D3D12RT_ComputeBytesPerPixel(DXGI_FORMAT Format)
 {
 	uint32 BytesPerPixel = 0;
 
@@ -100,16 +100,23 @@ static uint32 ComputeBytesPerPixel(DXGI_FORMAT Format)
 	case DXGI_FORMAT_R32G32B32A32_SINT:
 		BytesPerPixel = 16;
 		break;
+	default:
+		// format not supported yet
+		check(0);
+		break;
 	}
 
-	// format not supported yet
-	check(BytesPerPixel);
+	// @@ DXGIUtilities supercedes; remove this function ??
+	check(BytesPerPixel == UE::DXGIUtilities::GetFormatSizeInBytes(Format) );
 
 	return BytesPerPixel;
 }
 
 static inline DXGI_FORMAT ConvertTypelessToUnorm(DXGI_FORMAT Format)
 {
+	// prefer DXGIUtilities::FindSharedResourceFormat ?
+	//	or something? lots of these mappers in DXGIUtilities already
+
 	// required to prevent 
 	// D3D11: ERROR: ID3D11DeviceContext::ResolveSubresource: The Format (0x1b, R8G8B8A8_TYPELESS) is never able to resolve multisampled resources. [ RESOURCE_MANIPULATION ERROR #294: DEVICE_RESOLVESUBRESOURCE_FORMAT_INVALID ]
 	// D3D11: **BREAK** enabled for the previous D3D11 message, which was: [ RESOURCE_MANIPULATION ERROR #294: DEVICE_RESOLVESUBRESOURCE_FORMAT_INVALID ]
@@ -282,6 +289,8 @@ void FD3D12CommandContext::ResolveTextureUsingShader(
 
 static DXGI_FORMAT GetPlaneFormat(DXGI_FORMAT InFormat, uint32 InPlaneSlice)
 {
+	// in D3D12 mixed formats are in split planes not interleaved (change from D3D11)
+	//	?? but these look wrong ??
 	if (InFormat == DXGI_FORMAT_R32G8X24_TYPELESS || InFormat == DXGI_FORMAT_R24G8_TYPELESS)
 	{
 		if (InPlaneSlice == 0)
@@ -405,7 +414,7 @@ TRefCountPtr<FD3D12Resource> FD3D12DynamicRHI::GetStagingTexture(FRHITexture* Te
 	const uint32 PlaneSlice = 0;
 	const DXGI_FORMAT DestFormat = GetPlaneFormat(SourceDesc.Format, PlaneSlice);
 
-	const uint32 BlockBytes = ComputeBytesPerPixel(DestFormat);
+	const uint32 BlockBytes = D3D12RT_ComputeBytesPerPixel(DestFormat);
 	const uint32 XBytesAligned = Align((uint32)SourceDesc.Width * BlockBytes, FD3D12_TEXTURE_DATA_PITCH_ALIGNMENT);
 	const uint32 MipBytesAligned = XBytesAligned * SourceDesc.Height;
 	VERIFYD3D12RESULT(Adapter->CreateBuffer(D3D12_HEAP_TYPE_READBACK, Node, Node, MipBytesAligned, TempTexture2D.GetInitReference(), nullptr));
@@ -505,444 +514,6 @@ void FD3D12DynamicRHI::ReadSurfaceDataNoMSAARaw(FRHITexture* TextureRHI, FIntRec
 	TempTexture2D->GetResource()->Unmap(0, nullptr);
 }
 
-/** Helper for accessing R10G10B10A2 colors. */
-struct FD3DR10G10B10A2
-{
-	uint32 R : 10;
-	uint32 G : 10;
-	uint32 B : 10;
-	uint32 A : 2;
-};
-
-/** Helper for accessing R16G16 colors. */
-struct FD3DRG16
-{
-	uint16 R;
-	uint16 G;
-};
-
-/** Helper for accessing R16G16B16A16 colors. */
-struct FD3DRGBA16
-{
-	uint16 R;
-	uint16 G;
-	uint16 B;
-	uint16 A;
-};
-
-/** Convert D3D format type to general pixel format type*/
-static void ConvertDXGIToFColor(DXGI_FORMAT Format, uint32 Width, uint32 Height, uint8 *In, uint32 SrcPitch, FColor* Out, FReadSurfaceDataFlags InFlags)
-{
-	bool bLinearToGamma = InFlags.GetLinearToGamma();
-	switch (Format)
-	{
-	case DXGI_FORMAT_R16_TYPELESS:
-		ConvertRawR16DataToFColor(Width, Height, In, SrcPitch, Out);
-		break;
-	case DXGI_FORMAT_R8G8B8A8_TYPELESS:
-	case DXGI_FORMAT_R8G8B8A8_UNORM:
-	case DXGI_FORMAT_R8G8B8A8_UNORM_SRGB:
-		ConvertRawR8G8B8A8DataToFColor(Width, Height, In, SrcPitch, Out);
-		break;
-	case DXGI_FORMAT_B8G8R8A8_TYPELESS:
-	case DXGI_FORMAT_B8G8R8A8_UNORM:
-	case DXGI_FORMAT_B8G8R8A8_UNORM_SRGB:
-		ConvertRawB8G8R8A8DataToFColor(Width, Height, In, SrcPitch, Out);
-		break;
-	case DXGI_FORMAT_R10G10B10A2_UNORM:
-		ConvertRawR10G10B10A2DataToFColor(Width, Height, In, SrcPitch, Out);
-		break;
-	case DXGI_FORMAT_R16G16B16A16_FLOAT:
-		ConvertRawR16G16B16A16FDataToFColor(Width, Height, In, SrcPitch, Out, bLinearToGamma);
-		break;
-	case DXGI_FORMAT_R11G11B10_FLOAT:
-		ConvertRawR11G11B10DataToFColor(Width, Height, In, SrcPitch, Out, bLinearToGamma);
-		break;
-	case DXGI_FORMAT_R9G9B9E5_SHAREDEXP:
-		ConvertRawR9G9B9E5DataToFColor(Width, Height, In, SrcPitch, Out, bLinearToGamma);
-		break;
-	case DXGI_FORMAT_R32G32B32A32_FLOAT:
-		ConvertRawR32G32B32A32DataToFColor(Width, Height, In, SrcPitch, Out, bLinearToGamma);
-		break;
-	case DXGI_FORMAT_R24G8_TYPELESS:
-		ConvertRawR24G8DataToFColor(Width, Height, In, SrcPitch, Out, InFlags);
-		break;
-	case DXGI_FORMAT_R32G8X24_TYPELESS:
-		ConvertRawR32DataToFColor(Width, Height, In, SrcPitch, Out, InFlags);
-		break;
-	case DXGI_FORMAT_R16G16B16A16_UNORM:
-		ConvertRawR16G16B16A16DataToFColor(Width, Height, In, SrcPitch, Out);
-		break;
-	case DXGI_FORMAT_R16G16_UNORM:
-		ConvertRawR16G16DataToFColor(Width, Height, In, SrcPitch, Out);
-		break;
-	case DXGI_FORMAT_R8_UNORM:
-		ConvertRawR8DataToFColor(Width, Height, In, SrcPitch, Out);
-		break;
-	case DXGI_FORMAT_R8G8_UNORM:
-		ConvertRawR8G8DataToFColor(Width, Height, In, SrcPitch, Out);
-		break;
-	default:
-		checkf(0, TEXT("Unknown surface format!"));
-		break;
-	}
-}
-
-
-
-static void ConvertRAWSurfaceDataToFLinearColor(EPixelFormat Format, uint32 Width, uint32 Height, uint8 *In, uint32 SrcPitch, FLinearColor* Out, FReadSurfaceDataFlags InFlags)
-{
-	// @todo Oodle : need to factor this out of here to somewhere generic and share it with all the RHI's
-	// also get the UNorm/MinMax stuff out of the inner loops here, do that as a post-process
-
-	if (Format == PF_R16_UINT )
-	{
-		for (uint32 Y = 0; Y < Height; Y++)
-		{
-			uint16* SrcPtr = (uint16*)(In + Y * SrcPitch);
-			FLinearColor* DestPtr = Out + Y * Width;
-
-			for (uint32 X = 0; X < Width; X++)
-			{
-				uint16 Value16 = *SrcPtr;
-				float Value = Value16 * (1.f/0xffff);
-
-				*DestPtr = FLinearColor(Value, Value, Value); // G (gray) not R
-				++SrcPtr;
-				++DestPtr;
-			}
-		}
-	}
-	else if (Format == PF_R16F || Format == PF_R16F_FILTER)
-	{
-		for (uint32 Y = 0; Y < Height; Y++)
-		{
-			FFloat16 * SrcPtr = (FFloat16*)(In + Y * SrcPitch);
-			FLinearColor* DestPtr = Out + Y * Width;
-
-			for (uint32 X = 0; X < Width; X++)
-			{
-				float Value = SrcPtr->GetFloat();
-				*DestPtr = FLinearColor(Value, Value, Value);
-				++SrcPtr;
-				++DestPtr;
-			}
-		}
-	}
-	else if (Format == PF_R8G8B8A8)
-	{
-		// Read the data out of the buffer, converting it from ABGR to ARGB.
-		for (uint32 Y = 0; Y < Height; Y++)
-		{
-			FColor* SrcPtr = (FColor*)(In + Y * SrcPitch);
-			FLinearColor* DestPtr = Out + Y * Width;
-			for (uint32 X = 0; X < Width; X++)
-			{
-				FColor sRGBColor = FColor(SrcPtr->B, SrcPtr->G, SrcPtr->R, SrcPtr->A); // swap RB
-				*DestPtr = FLinearColor(sRGBColor);
-				++SrcPtr;
-				++DestPtr;
-			}
-		}
-	}
-	else if (Format == PF_B8G8R8A8)
-	{
-		for (uint32 Y = 0; Y < Height; Y++)
-		{
-			FColor* SrcPtr = (FColor*)(In + Y * SrcPitch);
-			FLinearColor* DestPtr = Out + Y * Width;
-			for (uint32 X = 0; X < Width; X++)
-			{
-				*DestPtr = FLinearColor(*SrcPtr);
-				++SrcPtr;
-				++DestPtr;
-			}
-		}
-	}
-	else if (Format == PF_A2B10G10R10)
-	{
-		// Read the data out of the buffer, converting it from R10G10B10A2 to FLinearColor.
-		for (uint32 Y = 0; Y < Height; Y++)
-		{
-			FD3DR10G10B10A2* SrcPtr = (FD3DR10G10B10A2*)(In + Y * SrcPitch);
-			FLinearColor* DestPtr = Out + Y * Width;
-			for (uint32 X = 0; X < Width; X++)
-			{
-				*DestPtr = FLinearColor(
-					(float)SrcPtr->R / 1023.0f,
-					(float)SrcPtr->G / 1023.0f,
-					(float)SrcPtr->B / 1023.0f,
-					(float)SrcPtr->A / 3.0f
-				);
-				++SrcPtr;
-				++DestPtr;
-			}
-		}
-	}
-	else if (Format == PF_FloatRGBA)
-	{
-		if (InFlags.GetCompressionMode() == RCM_MinMax)
-		{
-			for (uint32 Y = 0; Y < Height; Y++)
-			{
-				const FFloat16Color* SrcPtr = (FFloat16Color*)(In + Y * SrcPitch);
-				FLinearColor* DestPtr = Out + Y * Width;
-
-				for (uint32 X = 0; X < Width; X++)
-				{
-					DestPtr[X] = SrcPtr[X].GetFloats();
-				}
-			}
-		}
-		else
-		{
-			FPlane	MinValue(0.0f, 0.0f, 0.0f, 0.0f);
-			FPlane	MaxValue(1.0f, 1.0f, 1.0f, 1.0f);
-
-			check(sizeof(FFloat16) == sizeof(uint16));
-
-			for (uint32 Y = 0; Y < Height; Y++)
-			{
-				FFloat16* SrcPtr = (FFloat16*)(In + Y * SrcPitch);
-
-				for (uint32 X = 0; X < Width; X++)
-				{
-					MinValue.X = FMath::Min<float>(SrcPtr[0], MinValue.X);
-					MinValue.Y = FMath::Min<float>(SrcPtr[1], MinValue.Y);
-					MinValue.Z = FMath::Min<float>(SrcPtr[2], MinValue.Z);
-					MinValue.W = FMath::Min<float>(SrcPtr[3], MinValue.W);
-					MaxValue.X = FMath::Max<float>(SrcPtr[0], MaxValue.X);
-					MaxValue.Y = FMath::Max<float>(SrcPtr[1], MaxValue.Y);
-					MaxValue.Z = FMath::Max<float>(SrcPtr[2], MaxValue.Z);
-					MaxValue.W = FMath::Max<float>(SrcPtr[3], MaxValue.W);
-					SrcPtr += 4;
-				}
-			}
-
-			for (uint32 Y = 0; Y < Height; Y++)
-			{
-				FFloat16* SrcPtr = (FFloat16*)(In + Y * SrcPitch);
-				FLinearColor* DestPtr = Out + Y * Width;
-
-				for (uint32 X = 0; X < Width; X++)
-				{
-					*DestPtr = FLinearColor(
-						(SrcPtr[0] - MinValue.X) / (MaxValue.X - MinValue.X),
-						(SrcPtr[1] - MinValue.Y) / (MaxValue.Y - MinValue.Y),
-						(SrcPtr[2] - MinValue.Z) / (MaxValue.Z - MinValue.Z),
-						(SrcPtr[3] - MinValue.W) / (MaxValue.W - MinValue.W)
-					);
-					++DestPtr;
-					SrcPtr += 4;
-				}
-			}
-		}
-	}
-	else if (Format == PF_FloatRGB || Format == PF_FloatR11G11B10)
-	{
-		check(sizeof(FFloat3Packed) == sizeof(uint32));
-
-		for (uint32 Y = 0; Y < Height; Y++)
-		{
-			FFloat3Packed* SrcPtr = (FFloat3Packed*)(In + Y * SrcPitch);
-			FLinearColor* DestPtr = Out + Y * Width;
-
-			for (uint32 X = 0; X < Width; X++)
-			{
-				*DestPtr = (*SrcPtr).ToLinearColor();
-				++DestPtr;
-				++SrcPtr;
-			}
-		}
-	}
-	else if (Format == PF_A32B32G32R32F)
-	{
-		if (InFlags.GetCompressionMode() == RCM_MinMax)
-		{
-			// Copy data directly, respecting existing min-max values
-			FLinearColor* SrcPtr = (FLinearColor*)In;
-			FLinearColor* DestPtr = (FLinearColor*)Out;
-			const int32 ImageSize = sizeof(FLinearColor) * Height * Width;
-
-			FMemory::Memcpy(DestPtr, SrcPtr, ImageSize);
-		}
-		else
-		{
-			// Normalize data
-			FPlane MinValue(0.0f, 0.0f, 0.0f, 0.0f);
-			FPlane MaxValue(1.0f, 1.0f, 1.0f, 1.0f);
-
-			for (uint32 Y = 0; Y < Height; Y++)
-			{
-				float* SrcPtr = (float*)(In + Y * SrcPitch);
-
-				for (uint32 X = 0; X < Width; X++)
-				{
-					MinValue.X = FMath::Min<float>(SrcPtr[0], MinValue.X);
-					MinValue.Y = FMath::Min<float>(SrcPtr[1], MinValue.Y);
-					MinValue.Z = FMath::Min<float>(SrcPtr[2], MinValue.Z);
-					MinValue.W = FMath::Min<float>(SrcPtr[3], MinValue.W);
-					MaxValue.X = FMath::Max<float>(SrcPtr[0], MaxValue.X);
-					MaxValue.Y = FMath::Max<float>(SrcPtr[1], MaxValue.Y);
-					MaxValue.Z = FMath::Max<float>(SrcPtr[2], MaxValue.Z);
-					MaxValue.W = FMath::Max<float>(SrcPtr[3], MaxValue.W);
-					SrcPtr += 4;
-				}
-			}
-
-			float* SrcPtr = (float*)In;
-
-			for (uint32 Y = 0; Y < Height; Y++)
-			{
-				FLinearColor* DestPtr = Out + Y * Width;
-
-				for (uint32 X = 0; X < Width; X++)
-				{
-					*DestPtr = FLinearColor(
-						(SrcPtr[0] - MinValue.X) / (MaxValue.X - MinValue.X),
-						(SrcPtr[1] - MinValue.Y) / (MaxValue.Y - MinValue.Y),
-						(SrcPtr[2] - MinValue.Z) / (MaxValue.Z - MinValue.Z),
-						(SrcPtr[3] - MinValue.W) / (MaxValue.W - MinValue.W)
-					);
-					++DestPtr;
-					SrcPtr += 4;
-				}
-			}
-		}
-	}
-	else if (Format == PF_DepthStencil || Format == PF_D24)
-	{
-		// Depth stencil
-		for (uint32 Y = 0; Y < Height; Y++)
-		{
-			uint32* SrcPtr = (uint32 *)(In + Y * SrcPitch);
-			FLinearColor* DestPtr = Out + Y * Width;
-
-			for (uint32 X = 0; X < Width; X++)
-			{
-				float DeviceStencil = 0.0f;
-				DeviceStencil = (float)((*SrcPtr & 0xFF000000) >> 24) / 255.0f;
-				float DeviceZ = (*SrcPtr & 0xffffff) / (float)(1 << 24);
-				float LinearValue = FMath::Min(InFlags.ComputeNormalizedDepth(DeviceZ), 1.0f);
-				*DestPtr = FLinearColor(LinearValue, DeviceStencil, 0.0f, 0.0f);
-				++DestPtr;
-				++SrcPtr;
-			}
-		}
-	}
-	else if (Format == PF_DepthStencil)
-	{
-		// Depth stencil
-		for (uint32 Y = 0; Y < Height; Y++)
-		{
-			uint8* SrcStart = (uint8 *)(In + Y * SrcPitch);
-			FLinearColor* DestPtr = Out + Y * Width;
-
-			for (uint32 X = 0; X < Width; X++)
-			{
-				float DeviceZ = *((float *)(SrcStart));
-				float LinearValue = FMath::Min(InFlags.ComputeNormalizedDepth(DeviceZ), 1.0f);
-				float DeviceStencil = (float)(*(SrcStart + 4)) / 255.0f;
-				*DestPtr = FLinearColor(LinearValue, DeviceStencil, 0.0f, 0.0f);
-				SrcStart += 8; //64 bit format with the last 24 bit ignore
-			}
-		}
-	}
-	else if (Format == PF_A16B16G16R16)
-	{
-		// Read the data out of the buffer, converting it to FLinearColor.
-		for (uint32 Y = 0; Y < Height; Y++)
-		{
-			FD3DRGBA16* SrcPtr = (FD3DRGBA16*)(In + Y * SrcPitch);
-			FLinearColor* DestPtr = Out + Y * Width;
-			for (uint32 X = 0; X < Width; X++)
-			{
-				*DestPtr = FLinearColor(
-					(float)SrcPtr->R / 65535.0f,
-					(float)SrcPtr->G / 65535.0f,
-					(float)SrcPtr->B / 65535.0f,
-					(float)SrcPtr->A / 65535.0f
-				);
-				++SrcPtr;
-				++DestPtr;
-			}
-		}
-	}
-	else if (Format == PF_G16R16)
-	{
-		// Read the data out of the buffer, converting it to FLinearColor.
-		for (uint32 Y = 0; Y < Height; Y++)
-		{
-			FD3DRG16* SrcPtr = (FD3DRG16*)(In + Y * SrcPitch);
-			FLinearColor* DestPtr = Out + Y * Width;
-			for (uint32 X = 0; X < Width; X++)
-			{
-				*DestPtr = FLinearColor(
-					(float)SrcPtr->R / 65535.0f,
-					(float)SrcPtr->G / 65535.0f,
-					0);
-				++SrcPtr;
-				++DestPtr;
-			}
-		}
-	}
-	else if (Format == PF_G16R16F)
-	{
-		// Read the data out of the buffer, converting it to FLinearColor.
-		for (uint32 Y = 0; Y < Height; Y++)
-		{
-			FFloat16 * SrcPtr = (FFloat16*)(In + Y * SrcPitch);
-			FLinearColor* DestPtr = Out + Y * Width;
-			for (uint32 X = 0; X < Width; X++)
-			{
-				*DestPtr = FLinearColor( SrcPtr[0].GetFloat(), SrcPtr[1].GetFloat(), 0.f,1.f);
-				SrcPtr += 2;
-				++DestPtr;
-			}
-		}
-	}
-	else if (Format == PF_G32R32F)
-	{
-		// not doing MinMax/Unorm remap here
-	
-		// Read the data out of the buffer, converting it to FLinearColor.
-		for (uint32 Y = 0; Y < Height; Y++)
-		{
-			float * SrcPtr = (float *)(In + Y * SrcPitch);
-			FLinearColor* DestPtr = Out + Y * Width;
-			for (uint32 X = 0; X < Width; X++)
-			{
-				*DestPtr = FLinearColor( SrcPtr[0], SrcPtr[1], 0.f, 1.f );
-				SrcPtr += 2;
-				++DestPtr;
-			}
-		}
-	}
-	else if (Format == PF_R32_FLOAT)
-	{
-		// not doing MinMax/Unorm remap here
-	
-		// Read the data out of the buffer, converting it to FLinearColor.
-		for (uint32 Y = 0; Y < Height; Y++)
-		{
-			float * SrcPtr = (float *)(In + Y * SrcPitch);
-			FLinearColor* DestPtr = Out + Y * Width;
-			for (uint32 X = 0; X < Width; X++)
-			{
-				*DestPtr = FLinearColor( SrcPtr[0], 0.f, 0.f, 1.f );
-				++SrcPtr;
-				++DestPtr;
-			}
-		}
-	}
-	else
-	{
-		// not supported yet
-		check(0);
-	}
-}
-
-
 void FD3D12DynamicRHI::RHIReadSurfaceData(FRHITexture* TextureRHI, FIntRect InRect, TArray<FLinearColor>& OutData, FReadSurfaceDataFlags InFlags)
 {
 	TArray<uint8> OutDataRaw;
@@ -969,16 +540,31 @@ void FD3D12DynamicRHI::RHIReadSurfaceData(FRHITexture* TextureRHI, FIntRect InRe
 	// Allocate the output buffer.
 	OutData.SetNumUninitialized(SizeX * SizeY);
 
-	uint32 BytesPerPixel = ComputeBytesPerPixel(TextureDesc.Format);
+	uint32 BytesPerPixel = D3D12RT_ComputeBytesPerPixel(TextureDesc.Format);
 	uint32 SrcPitch = SizeX * BytesPerPixel;
+
+	// switching on the EPixelFormat is risky if the mapping is not what you expect
+	//	verify against TextureDesc.Format
+
 	EPixelFormat Format = TextureRHI->GetFormat();
-	if (Format != PF_Unknown)
+	check( GPixelFormats[Format].PlatformFormat == TextureDesc.Format );
+	check( GPixelFormats[Format].BlockBytes == D3D12RT_ComputeBytesPerPixel(TextureDesc.Format) );
+
+	if ( ! ConvertRAWSurfaceDataToFLinearColor(Format, SizeX, SizeY, OutDataRaw.GetData(), SrcPitch, OutData.GetData(), InFlags) )
 	{
-		ConvertRAWSurfaceDataToFLinearColor(Format, SizeX, SizeY, OutDataRaw.GetData(), SrcPitch, OutData.GetData(), InFlags);
+		checkf(0, TEXT("Unsupported surface format!"));
+		OutData.Empty();
 	}
 }
+
 void FD3D12DynamicRHI::RHIReadSurfaceData(FRHITexture* InRHITexture, FIntRect InRect, TArray<FColor>& OutData, FReadSurfaceDataFlags InFlags)
 {
+	if ( InRect.Width() == 0 || InRect.Height() == 0 )
+	{
+		OutData.Empty();
+		return;
+	}
+
 	if (!ensure(InRHITexture))
 	{
 		OutData.SetNumUninitialized(InRect.Width() * InRect.Height());
@@ -1007,7 +593,7 @@ void FD3D12DynamicRHI::RHIReadSurfaceData(FRHITexture* InRHITexture, FIntRect In
 	FD3D12Texture* DestTexture2D = CommandContext.RetrieveTexture(InRHITexture);
 
 	// Check the format of the surface
-	D3D12_RESOURCE_DESC const& TextureDesc = DestTexture2D->GetResource()->GetDesc();
+	FD3D12ResourceDesc const& TextureDesc = DestTexture2D->GetResource()->GetDesc();
 
 	check(TextureDesc.SampleDesc.Count >= 1);
 
@@ -1025,12 +611,40 @@ void FD3D12DynamicRHI::RHIReadSurfaceData(FRHITexture* InRHITexture, FIntRect In
 
 	// Allocate the output buffer.
 	OutData.SetNumUninitialized(SizeX * SizeY);
+	
+	DXGI_FORMAT DXGIFormat = TextureDesc.Format;
+	EPixelFormat PixelFormat = DestTexture2D->GetFormat();
+
+	check( PixelFormat != PF_Unknown );
+	check( PixelFormat == TextureDesc.PixelFormat || TextureDesc.PixelFormat == PF_Unknown );
+
+	const FPixelFormatInfo & FormatInfo = GPixelFormats[PixelFormat];
 
-	FPixelFormatInfo FormatInfo = GPixelFormats[DestTexture2D->GetFormat()];
 	uint32 BytesPerPixel = FormatInfo.BlockBytes;
 	uint32 SrcPitch = SizeX * BytesPerPixel;
+	
+	// switching on the EPixelFormat is risky if the mapping is not what you expect
+	//	verify against TextureDesc.Format
+
+	if ( DXGIFormat == DXGI_FORMAT_UNKNOWN )
+	{
+		// when called on actual textures, format is valid
+		// but this is also called on untyped buffers
+		DXGIFormat = (DXGI_FORMAT) FormatInfo.PlatformFormat;
+	}
+	else
+	{
+		check( FormatInfo.PlatformFormat == DXGIFormat );
+		check( FormatInfo.BlockBytes == D3D12RT_ComputeBytesPerPixel(DXGIFormat) );
+	}
 
-	ConvertDXGIToFColor((DXGI_FORMAT)FormatInfo.PlatformFormat, SizeX, SizeY, OutDataRaw.GetData(), SrcPitch, OutData.GetData(), InFlags);
+	// ConvertDXGIToFColor switches on the hardware format, not the EPixelFormat :
+
+	if ( ! ConvertDXGIToFColor(DXGIFormat, SizeX, SizeY, OutDataRaw.GetData(), SrcPitch, OutData.GetData(), InFlags) )
+	{
+		checkf(0, TEXT("Unsupported surface format!"));
+		OutData.Empty();
+	}
 }
 
 void FD3D12DynamicRHI::ReadSurfaceDataMSAARaw(FRHITexture* TextureRHI, FIntRect InRect, TArray<uint8>& OutData, FReadSurfaceDataFlags InFlags)
@@ -1050,7 +664,7 @@ void FD3D12DynamicRHI::ReadSurfaceDataMSAARaw(FRHITexture* TextureRHI, FIntRect
 	// Check the format of the surface
 	D3D12_RESOURCE_DESC const& TextureDesc = Texture->GetResource()->GetDesc();
 
-	uint32 BytesPerPixel = ComputeBytesPerPixel(TextureDesc.Format);
+	uint32 BytesPerPixel = D3D12RT_ComputeBytesPerPixel(TextureDesc.Format);
 
 	const uint32 NumSamples = TextureDesc.SampleDesc.Count;
 
@@ -1100,7 +714,7 @@ void FD3D12DynamicRHI::ReadSurfaceDataMSAARaw(FRHITexture* TextureRHI, FIntRect
 	const uint32 PlaneSlice = 0;
 	const DXGI_FORMAT DestFormat = GetPlaneFormat(TextureDesc.Format, PlaneSlice);
 
-	const uint32 BlockBytes = ComputeBytesPerPixel(DestFormat);
+	const uint32 BlockBytes = D3D12RT_ComputeBytesPerPixel(DestFormat);
 	const uint32 XBytesAligned = Align(SizeX * BlockBytes, FD3D12_TEXTURE_DATA_PITCH_ALIGNMENT);
 	const uint32 MipBytesAligned = XBytesAligned * SizeY;
 	VERIFYD3D12RESULT(Adapter->CreateBuffer(D3D12_HEAP_TYPE_READBACK, NodeMask, NodeMask, MipBytesAligned, StagingTexture2D.GetInitReference(), nullptr));
@@ -1132,8 +746,7 @@ void FD3D12DynamicRHI::ReadSurfaceDataMSAARaw(FRHITexture* TextureRHI, FIntRect
 	CD3DX12_TEXTURE_COPY_LOCATION SourceCopyLocation(NonMSAATexture2D->GetResource(), Subresource);
 
 	// Allocate the output buffer.
-	OutData.Empty();
-	OutData.AddUninitialized(SizeX * SizeY * NumSamples * BytesPerPixel);
+	OutData.SetNumUninitialized(SizeX * SizeY * NumSamples * BytesPerPixel);
 
 	// Can be optimized by doing all subsamples into a large enough rendertarget in one pass (multiple draw calls)
 	for (uint32 SampleIndex = 0; SampleIndex < NumSamples; ++SampleIndex)
@@ -1200,7 +813,7 @@ void FD3D12DynamicRHI::RHIMapStagingSurface(FRHITexture* TextureRHI, FRHIGPUFenc
 
 	DXGI_FORMAT Format = (DXGI_FORMAT)GPixelFormats[DestTexture->GetFormat()].PlatformFormat;
 
-	uint32 BytesPerPixel = ComputeBytesPerPixel(Format);
+	uint32 BytesPerPixel = D3D12RT_ComputeBytesPerPixel(Format);
 
 	if (FenceRHI && !FenceRHI->Poll())
 	{
@@ -1265,7 +878,11 @@ void FD3D12DynamicRHI::RHIReadSurfaceFloatData(FRHITexture* TextureRHI, FIntRect
 	// Check the format of the surface
 	D3D12_RESOURCE_DESC const& TextureDesc = Texture->GetResource()->GetDesc();
 
-	check(TextureDesc.Format == GPixelFormats[PF_FloatRGBA].PlatformFormat);
+	if ( ! ensure(TextureDesc.Format == GPixelFormats[PF_FloatRGBA].PlatformFormat) )
+	{
+		OutData.Empty();
+		return;
+	}
 
 	// Allocate the output buffer.
 	OutData.SetNumUninitialized(SizeX * SizeY);
@@ -1285,7 +902,7 @@ void FD3D12DynamicRHI::RHIReadSurfaceFloatData(FRHITexture* TextureRHI, FIntRect
 	const uint32 PlaneSlice = 0;
 	const DXGI_FORMAT DestFormat = GetPlaneFormat(TextureDesc.Format, PlaneSlice);
 
-	const uint32 BlockBytes = ComputeBytesPerPixel(DestFormat);
+	const uint32 BlockBytes = D3D12RT_ComputeBytesPerPixel(DestFormat);
 	const uint32 XBytesAligned = Align((uint32)TextureDesc.Width * BlockBytes, FD3D12_TEXTURE_DATA_PITCH_ALIGNMENT);
 	const uint32 MipBytesAligned = XBytesAligned * TextureDesc.Height;
 	VERIFYD3D12RESULT(Adapter->CreateBuffer(D3D12_HEAP_TYPE_READBACK, Node, Node, MipBytesAligned, TempTexture2D.GetInitReference(), nullptr));
@@ -1382,7 +999,11 @@ void FD3D12DynamicRHI::RHIRead3DSurfaceFloatData(FRHITexture* TextureRHI, FIntRe
 	bool bIsRGBAFmt = TextureDesc.Format == GPixelFormats[PF_FloatRGBA].PlatformFormat;
 	bool bIsR16FFmt = TextureDesc.Format == GPixelFormats[PF_R16F].PlatformFormat;
 	bool bIsR32FFmt = TextureDesc.Format == GPixelFormats[PF_R32_FLOAT].PlatformFormat;
-	check(bIsRGBAFmt || bIsR16FFmt || bIsR32FFmt);
+	if ( ! ensure(bIsRGBAFmt || bIsR16FFmt || bIsR32FFmt) )
+	{
+		OutData.Empty();
+		return;
+	}
 
 	// Allocate the output buffer.
 	OutData.SetNumUninitialized(SizeX * SizeY * SizeZ);
@@ -1406,7 +1027,7 @@ void FD3D12DynamicRHI::RHIRead3DSurfaceFloatData(FRHITexture* TextureRHI, FIntRe
 
 	// Copy the data to a staging resource.
 	uint32 Subresource = 0;
-	uint32 BytesPerPixel = ComputeBytesPerPixel(TextureDesc.Format);
+	uint32 BytesPerPixel = D3D12RT_ComputeBytesPerPixel(TextureDesc.Format);
 	D3D12_SUBRESOURCE_FOOTPRINT DestSubresource;
 	DestSubresource.Depth = TextureDesc.DepthOrArraySize;
 	DestSubresource.Height = TextureDesc.Height;
@@ -1495,6 +1116,11 @@ void FD3D12DynamicRHI::RHIRead3DSurfaceFloatData(FRHITexture* TextureRHI, FIntRe
 			}
 		}
 	}
+	else
+	{
+		// unsupported format; checked for this earlier
+		check(0);
+	}
 
 	TempTexture3D->GetResource()->Unmap(0, nullptr);
 }
diff --git a/Engine/Source/Runtime/Engine/Private/ImageUtils.cpp b/Engine/Source/Runtime/Engine/Private/ImageUtils.cpp
--- a/Engine/Source/Runtime/Engine/Private/ImageUtils.cpp
+++ b/Engine/Source/Runtime/Engine/Private/ImageUtils.cpp
@@ -339,6 +339,10 @@ bool FImageUtils::GetRenderTargetImage(UTextureRenderTarget2D* TexRT, FImage & I
 
 	ERawImageFormat::Type ReadFormat = UTextureRenderTarget::GetReadPixelsFormat(RTFormat,false);
 		
+	// RCM_MinMax means don't renormalize, just read the pixels as they are
+	//	default RCM_UNorm does funny scalings
+	FReadSurfaceDataFlags InFlags(RCM_MinMax, CubeFace_MAX);
+
 	if ( ReadFormat == ERawImageFormat::RGBA16F )
 	{
 		// ReadFloat16Pixels does no conversions
@@ -347,7 +351,7 @@ bool FImageUtils::GetRenderTargetImage(UTextureRenderTarget2D* TexRT, FImage & I
 		Image.Init(TexRT->SizeX,TexRT->SizeY,ERawImageFormat::RGBA16F,EGammaSpace::Linear);
 		
 		TArray<FFloat16Color> Colors;
-		if ( ! RenderTarget->ReadFloat16Pixels(Colors) )
+		if ( ! RenderTarget->ReadFloat16Pixels(Colors,InFlags) )
 		{
 			return false;
 		}
@@ -363,8 +367,9 @@ bool FImageUtils::GetRenderTargetImage(UTextureRenderTarget2D* TexRT, FImage & I
 		
 		Image.Init(TexRT->SizeX,TexRT->SizeY,ERawImageFormat::BGRA8,GammaSpace);
 		
-		FReadSurfaceDataFlags InFlags(RCM_MinMax, CubeFace_MAX);
-
+		// "LinearToGamma" is basically moot; that would only be used if we were reading float pixels to FColor
+		//	but in that case the ReadFormat should have been float, so we won't be here
+		//	gamma conversion will be handled by FImage after the pixel read, not inside RHI
 		InFlags.SetLinearToGamma( GammaSpace == EGammaSpace::sRGB );
 
 		TArray<FColor> Colors;
@@ -380,10 +385,6 @@ bool FImageUtils::GetRenderTargetImage(UTextureRenderTarget2D* TexRT, FImage & I
 	{
 		Image.Init(TexRT->SizeX,TexRT->SizeY,ERawImageFormat::RGBA32F,EGammaSpace::Linear);
 		
-		FReadSurfaceDataFlags InFlags(RCM_MinMax, CubeFace_MAX);
-
-		InFlags.SetLinearToGamma( false );
-
 		TArray<FLinearColor> Colors;
 		if ( ! RenderTarget->ReadLinearColorPixels(Colors,InFlags) )
 		{
diff --git a/Engine/Source/Runtime/RHI/Public/RHISurfaceDataConversion.h b/Engine/Source/Runtime/RHI/Public/RHISurfaceDataConversion.h
--- a/Engine/Source/Runtime/RHI/Public/RHISurfaceDataConversion.h
+++ b/Engine/Source/Runtime/RHI/Public/RHISurfaceDataConversion.h
@@ -13,6 +13,9 @@
 #include "RHI.h"
 #include "RHITypes.h"
 
+namespace
+{
+
 /** Helper for accessing R10G10B10A2 colors. */
 struct FRHIR10G10B10A2
 {
@@ -38,6 +41,7 @@ struct FRHIRG16
 	uint16 G;
 };
 
+
 inline void ConvertRawR16DataToFColor(uint32 Width, uint32 Height, uint8 *In, uint32 SrcPitch, FColor* Out)
 {
 	// e.g. shadow maps
@@ -319,7 +323,7 @@ inline void ConvertRawR24G8DataToFColor(uint32 Width, uint32 Height, uint8 *In,
 	}
 }
 
-inline void ConvertRawR32DataToFColor(uint32 Width, uint32 Height, uint8 *In, uint32 SrcPitch, FColor* Out, FReadSurfaceDataFlags InFlags)
+inline void ConvertRawDepthStencil64DataToFColor(uint32 Width, uint32 Height, uint8 *In, uint32 SrcPitch, FColor* Out, FReadSurfaceDataFlags InFlags)
 {
 	bool bLinearToGamma = InFlags.GetLinearToGamma();
 	for (uint32 Y = 0; Y < Height; Y++)
@@ -454,7 +458,7 @@ inline void ConvertRawD32S8DataToFColor(uint32 Width, uint32 Height, uint8 *In,
 
 // Linear functions
 
-inline void ConvertRawR16DataToFLinearColor(uint32 Width, uint32 Height, uint8 *In, uint32 SrcPitch, FLinearColor* Out)
+inline void ConvertRawR16UDataToFLinearColor(uint32 Width, uint32 Height, uint8 *In, uint32 SrcPitch, FLinearColor* Out)
 {
 	// e.g. shadow maps
 	for (uint32 Y = 0; Y < Height; Y++)
@@ -465,7 +469,26 @@ inline void ConvertRawR16DataToFLinearColor(uint32 Width, uint32 Height, uint8 *
 		for (uint32 X = 0; X < Width; X++)
 		{
 			uint16 Value16 = *SrcPtr;
-			float Value = Value16 / (float)(0xffff);
+			float Value = Value16 * (1.f/0xffff);
+
+			*DestPtr = FLinearColor(Value, Value, Value);
+			++SrcPtr;
+			++DestPtr;
+		}
+	}
+}
+
+inline void ConvertRawR16FDataToFLinearColor(uint32 Width, uint32 Height, uint8 *In, uint32 SrcPitch, FLinearColor* Out)
+{
+	// e.g. shadow maps
+	for (uint32 Y = 0; Y < Height; Y++)
+	{
+		FFloat16 * SrcPtr = (FFloat16 *)(In + Y * SrcPitch);
+		FLinearColor* DestPtr = Out + Y * Width;
+
+		for (uint32 X = 0; X < Width; X++)
+		{
+			float Value = SrcPtr->GetFloat();
 
 			*DestPtr = FLinearColor(Value, Value, Value);
 			++SrcPtr;
@@ -483,7 +506,8 @@ inline void ConvertRawR8G8B8A8DataToFLinearColor(uint32 Width, uint32 Height, ui
 		FLinearColor* DestPtr = Out + Y * Width;
 		for (uint32 X = 0; X < Width; X++)
 		{
-			FColor sRGBColor = FColor(SrcPtr->B, SrcPtr->G, SrcPtr->R, SrcPtr->A);
+			FColor sRGBColor = FColor(SrcPtr->B, SrcPtr->G, SrcPtr->R, SrcPtr->A); // swap RB
+			// FLinearColor(FColor) does SRGB -> Linear
 			*DestPtr = FLinearColor(sRGBColor);
 			++SrcPtr;
 			++DestPtr;
@@ -499,8 +523,8 @@ inline void ConvertRawB8G8R8A8DataToFLinearColor(uint32 Width, uint32 Height, ui
 		FLinearColor* DestPtr = Out + Y * Width;
 		for (uint32 X = 0; X < Width; X++)
 		{
-			FColor sRGBColor = FColor(SrcPtr->R, SrcPtr->G, SrcPtr->B, SrcPtr->A);
-			*DestPtr = FLinearColor(sRGBColor);
+			// FLinearColor(FColor) does SRGB -> Linear
+			*DestPtr = FLinearColor(*SrcPtr);
 			++SrcPtr;
 			++DestPtr;
 		}
@@ -591,7 +615,7 @@ inline void ConvertRawR16G16B16A16FDataToFLinearColor(uint32 Width, uint32 Heigh
 }
 
 
-inline void ConvertRawRR11G11B10DataToFLinearColor(uint32 Width, uint32 Height, uint8 *In, uint32 SrcPitch, FLinearColor* Out)
+inline void ConvertRawR11G11B10FDataToFLinearColor(uint32 Width, uint32 Height, uint8 *In, uint32 SrcPitch, FLinearColor* Out)
 {
 	check(sizeof(FFloat3Packed) == sizeof(uint32));
 
@@ -686,7 +710,7 @@ inline void ConvertRawR24G8DataToFLinearColor(uint32 Width, uint32 Height, uint8
 	}
 }
 
-inline void ConvertRawR32DataToFLinearColor(uint32 Width, uint32 Height, uint8 *In, uint32 SrcPitch, FLinearColor* Out, FReadSurfaceDataFlags InFlags)
+inline void ConvertRawDepthStencil64DataToFLinearColor(uint32 Width, uint32 Height, uint8 *In, uint32 SrcPitch, FLinearColor* Out, FReadSurfaceDataFlags InFlags)
 {
 	// Depth stencil
 	for (uint32 Y = 0; Y < Height; Y++)
@@ -746,3 +770,218 @@ inline void ConvertRawR16G16DataToFLinearColor(uint32 Width, uint32 Height, uint
 	}
 }
 
+}; // namespace
+
+static bool ConvertRAWSurfaceDataToFLinearColor(EPixelFormat Format, uint32 Width, uint32 Height, uint8 *In, uint32 SrcPitch, FLinearColor* Out, FReadSurfaceDataFlags InFlags)
+{
+	// InFlags.GetLinearToGamma() is ignored by the FLinearColor reader
+
+	// Flags RCM_MinMax means pass the values out unchanged
+	//	default flags RCM_UNorm rescales them to [0,1] if they were outside that range
+
+	if (Format == PF_G16 || Format == PF_R16_UINT )
+	{
+		ConvertRawR16UDataToFLinearColor(Width, Height, In, SrcPitch, Out);
+		return true;
+	}
+	else if (Format == PF_R16F || Format == PF_R16F_FILTER)
+	{
+		ConvertRawR16FDataToFLinearColor(Width, Height, In, SrcPitch, Out);
+		return true;
+	}
+	else if (Format == PF_R8G8B8A8)
+	{
+		ConvertRawR8G8B8A8DataToFLinearColor(Width, Height, In, SrcPitch, Out);
+		return true;
+	}
+	else if (Format == PF_B8G8R8A8)
+	{
+		ConvertRawB8G8R8A8DataToFLinearColor(Width, Height, In, SrcPitch, Out);
+		return true;
+	}
+	else if (Format == PF_A2B10G10R10)
+	{
+		ConvertRawA2B10G10R10DataToFLinearColor(Width, Height, In, SrcPitch, Out);
+		return true;
+	}
+	else if (Format == PF_FloatRGBA)
+	{
+		ConvertRawR16G16B16A16FDataToFLinearColor(Width, Height, In, SrcPitch, Out, InFlags);
+		return true;
+	}
+	else if (Format == PF_FloatRGB || Format == PF_FloatR11G11B10)
+	{
+		// assume FloatRGB == PF_FloatR11G11B10 always here
+		check( GPixelFormats[Format].BlockBytes == 4 );
+		ConvertRawR11G11B10FDataToFLinearColor(Width, Height, In, SrcPitch, Out);
+		return true;
+	}
+	else if (Format == PF_A32B32G32R32F)
+	{
+		ConvertRawR32G32B32A32DataToFLinearColor(Width, Height, In, SrcPitch, Out, InFlags);
+		return true;
+	}
+	else if ( Format == PF_D24 ||
+		( (Format == PF_X24_G8 || Format == PF_DepthStencil ) && GPixelFormats[Format].BlockBytes == 4 )
+		)
+	{
+		//	see CVarD3D11UseD24/CVarD3D12UseD24
+		ConvertRawR24G8DataToFLinearColor(Width, Height, In, SrcPitch, Out, InFlags);
+		return true;
+	}
+	else if ( (Format == PF_X24_G8 || Format == PF_DepthStencil ) && GPixelFormats[Format].BlockBytes > 4 )
+	{
+		// @@ D3D 11/12 different?
+
+		/**
+		
+		D3D11 and 12 formats are almost the same, one difference :
+
+		DepthStencil formats are interleaved	DepthStencil formats are planar.
+		For example a format of 24 bits of depth, 8 bits of stencil would be
+		stored in the format 24/8/24/8... etc in Direct3D 11, but as 24/24/24...
+		followed by 8/8/8... in Direct3D 12. Note that each plane is its own
+		subresource in D3D12
+
+		**/
+
+		ConvertRawDepthStencil64DataToFLinearColor(Width, Height, In, SrcPitch, Out, InFlags);
+		return true;
+	}
+	else if (Format == PF_A16B16G16R16)
+	{
+		ConvertRawR16G16B16A16DataToFLinearColor(Width, Height, In, SrcPitch, Out);
+		return true;
+	}
+	else if (Format == PF_G16R16)
+	{
+		ConvertRawR16G16DataToFLinearColor(Width, Height, In, SrcPitch, Out);
+		return true;
+	}
+	else if (Format == PF_G16R16F)
+	{
+		// Read the data out of the buffer, converting it to FLinearColor.
+		for (uint32 Y = 0; Y < Height; Y++)
+		{
+			FFloat16 * SrcPtr = (FFloat16*)(In + Y * SrcPitch);
+			FLinearColor* DestPtr = Out + Y * Width;
+			for (uint32 X = 0; X < Width; X++)
+			{
+				*DestPtr = FLinearColor( SrcPtr[0].GetFloat(), SrcPtr[1].GetFloat(), 0.f,1.f);
+				SrcPtr += 2;
+				++DestPtr;
+			}
+		}
+		return true;
+	}
+	else if (Format == PF_G32R32F)
+	{
+		// not doing MinMax/Unorm remap here
+	
+		// Read the data out of the buffer, converting it to FLinearColor.
+		for (uint32 Y = 0; Y < Height; Y++)
+		{
+			float * SrcPtr = (float *)(In + Y * SrcPitch);
+			FLinearColor* DestPtr = Out + Y * Width;
+			for (uint32 X = 0; X < Width; X++)
+			{
+				*DestPtr = FLinearColor( SrcPtr[0], SrcPtr[1], 0.f, 1.f );
+				SrcPtr += 2;
+				++DestPtr;
+			}
+		}
+		return true;
+	}
+	else if (Format == PF_R32_FLOAT)
+	{
+		// not doing MinMax/Unorm remap here
+	
+		// Read the data out of the buffer, converting it to FLinearColor.
+		for (uint32 Y = 0; Y < Height; Y++)
+		{
+			float * SrcPtr = (float *)(In + Y * SrcPitch);
+			FLinearColor* DestPtr = Out + Y * Width;
+			for (uint32 X = 0; X < Width; X++)
+			{
+				*DestPtr = FLinearColor( SrcPtr[0], 0.f, 0.f, 1.f );
+				++SrcPtr;
+				++DestPtr;
+			}
+		}
+		return true;
+	}
+	else
+	{
+		// not supported yet
+		check(0);
+		return false;
+	}
+}
+
+#if defined(DXGI_FORMAT_DEFINED) && DXGI_FORMAT_DEFINED
+
+// switching on platform format
+// could switch on EPixelFormat instead and make this more generic
+// have to be a little careful with that as the mapping is not one to one
+
+/** Convert D3D format type to general pixel format type*/
+static bool ConvertDXGIToFColor(DXGI_FORMAT Format, uint32 Width, uint32 Height, uint8 *In, uint32 SrcPitch, FColor* Out, FReadSurfaceDataFlags InFlags)
+{
+	// todo : use a helper to map all the typeless to unorm before switching here ; see DXGIUtilities
+
+	bool bLinearToGamma = InFlags.GetLinearToGamma();
+	switch (Format)
+	{
+	case DXGI_FORMAT_R16_TYPELESS:
+	case DXGI_FORMAT_R16_UNORM:
+		ConvertRawR16DataToFColor(Width, Height, In, SrcPitch, Out);
+		return true;
+	case DXGI_FORMAT_R8G8B8A8_TYPELESS:
+	case DXGI_FORMAT_R8G8B8A8_UNORM:
+	case DXGI_FORMAT_R8G8B8A8_UNORM_SRGB:
+		ConvertRawR8G8B8A8DataToFColor(Width, Height, In, SrcPitch, Out);
+		return true;
+	case DXGI_FORMAT_B8G8R8A8_TYPELESS:
+	case DXGI_FORMAT_B8G8R8A8_UNORM:
+	case DXGI_FORMAT_B8G8R8A8_UNORM_SRGB:
+		ConvertRawB8G8R8A8DataToFColor(Width, Height, In, SrcPitch, Out);
+		return true;
+	case DXGI_FORMAT_R10G10B10A2_UNORM:
+		ConvertRawR10G10B10A2DataToFColor(Width, Height, In, SrcPitch, Out);
+		return true;
+	case DXGI_FORMAT_R16G16B16A16_FLOAT:
+		ConvertRawR16G16B16A16FDataToFColor(Width, Height, In, SrcPitch, Out, bLinearToGamma);
+		return true;
+	case DXGI_FORMAT_R11G11B10_FLOAT:
+		ConvertRawR11G11B10DataToFColor(Width, Height, In, SrcPitch, Out, bLinearToGamma);
+		return true;
+	case DXGI_FORMAT_R9G9B9E5_SHAREDEXP:
+		ConvertRawR9G9B9E5DataToFColor(Width, Height, In, SrcPitch, Out, bLinearToGamma);
+		return true;
+	case DXGI_FORMAT_R32G32B32A32_FLOAT:
+		ConvertRawR32G32B32A32DataToFColor(Width, Height, In, SrcPitch, Out, bLinearToGamma);
+		return true;
+	case DXGI_FORMAT_R24G8_TYPELESS:
+		ConvertRawR24G8DataToFColor(Width, Height, In, SrcPitch, Out, InFlags);
+		return true;
+	case DXGI_FORMAT_R32G8X24_TYPELESS:
+		ConvertRawDepthStencil64DataToFColor(Width, Height, In, SrcPitch, Out, InFlags);
+		return true;
+	case DXGI_FORMAT_R16G16B16A16_UNORM:
+		ConvertRawR16G16B16A16DataToFColor(Width, Height, In, SrcPitch, Out);
+		return true;
+	case DXGI_FORMAT_R16G16_UNORM:
+		ConvertRawR16G16DataToFColor(Width, Height, In, SrcPitch, Out);
+		return true;
+	case DXGI_FORMAT_R8_UNORM:
+		ConvertRawR8DataToFColor(Width, Height, In, SrcPitch, Out);
+		return true;
+	case DXGI_FORMAT_R8G8_UNORM:
+		ConvertRawR8G8DataToFColor(Width, Height, In, SrcPitch, Out);
+		return true;
+	default:
+		return false;
+	}
+}
+
+#endif // DXGI_FORMAT_DEFINED
diff --git a/Engine/Source/Runtime/Windows/D3D11RHI/Private/D3D11RenderTarget.cpp b/Engine/Source/Runtime/Windows/D3D11RHI/Private/D3D11RenderTarget.cpp
--- a/Engine/Source/Runtime/Windows/D3D11RHI/Private/D3D11RenderTarget.cpp
+++ b/Engine/Source/Runtime/Windows/D3D11RHI/Private/D3D11RenderTarget.cpp
@@ -16,6 +16,9 @@
 
 static inline DXGI_FORMAT ConvertTypelessToUnorm(DXGI_FORMAT Format)
 {
+	// prefer DXGIUtilities::FindSharedResourceFormat ?
+	//	or something? lots of these mappers in DXGIUtilities already
+
 	// required to prevent 
 	// D3D11: ERROR: ID3D11DeviceContext::ResolveSubresource: The Format (0x1b, R8G8B8A8_TYPELESS) is never able to resolve multisampled resources. [ RESOURCE_MANIPULATION ERROR #294: DEVICE_RESOLVESUBRESOURCE_FORMAT_INVALID ]
 	// D3D11: **BREAK** enabled for the previous D3D11 message, which was: [ RESOURCE_MANIPULATION ERROR #294: DEVICE_RESOLVESUBRESOURCE_FORMAT_INVALID ]
@@ -188,7 +191,7 @@ void FD3D11DynamicRHI::ResolveTextureUsingShader(
 }
 
 // Only supports the formats that are supported by ConvertRAWSurfaceDataToFColor()
-static uint32 ComputeBytesPerPixel(DXGI_FORMAT Format)
+static uint32 D3D11RT_ComputeBytesPerPixel(DXGI_FORMAT Format)
 {
 	uint32 BytesPerPixel = 0;
 
@@ -273,10 +276,15 @@ static uint32 ComputeBytesPerPixel(DXGI_FORMAT Format)
 		case DXGI_FORMAT_R32G32B32A32_SINT:
 			BytesPerPixel = 16;
 			break;
+
+		default:
+			// format not supported yet
+			check(false);
+			break;
 	}
 
-	// format not supported yet
-	check(BytesPerPixel);
+	// this function is superceded by DXGIUtilities, delete me ??
+	check(BytesPerPixel == UE::DXGIUtilities::GetFormatSizeInBytes(Format) );
 
 	return BytesPerPixel;
 }
@@ -370,11 +378,10 @@ void FD3D11DynamicRHI::ReadSurfaceDataNoMSAARaw(FRHITexture* TextureRHI,FIntRect
 	D3D11_TEXTURE2D_DESC TextureDesc;
 	Texture->GetD3D11Texture2D()->GetDesc(&TextureDesc);
 	
-	uint32 BytesPerPixel = ComputeBytesPerPixel(TextureDesc.Format);
+	uint32 BytesPerPixel = D3D11RT_ComputeBytesPerPixel(TextureDesc.Format);
 	
 	// Allocate the output buffer.
-	OutData.Empty();
-	OutData.AddUninitialized(SizeX * SizeY * BytesPerPixel);
+	OutData.SetNumUninitialized(SizeX * SizeY * BytesPerPixel);
 
 	bool bIsUsingTempStagingTexture = TextureDesc.Usage != D3D11_USAGE_STAGING;
 	FIntRect StagingRect;
@@ -398,100 +405,6 @@ void FD3D11DynamicRHI::ReadSurfaceDataNoMSAARaw(FRHITexture* TextureRHI,FIntRect
 	Direct3DDeviceIMContext->Unmap(TempTexture2D, MappedSubresource);
 }
 
-
-/** Helper for accessing R10G10B10A2 colors. */
-struct FD3DR10G10B10A2
-{
-	uint32 R : 10;
-	uint32 G : 10;
-	uint32 B : 10;
-	uint32 A : 2;
-};
-
-struct FD3DR32G8
-{
-	uint32 R : 32;
-	uint32 G : 8;
-};
-
-struct FD3DR24G8
-{
-	uint32 R : 24;
-	uint32 G : 8;
-};
-
-
-/** Helper for accessing R16G16 colors. */
-struct FD3DRG16
-{
-	uint16 R;
-	uint16 G;
-};
-
-/** Helper for accessing R16G16B16A16 colors. */
-struct FD3DRGBA16
-{
-	uint16 R;
-	uint16 G;
-	uint16 B;
-	uint16 A;
-};
-
-/** Convert D3D format type to general pixel format type*/
-static void ConvertDXGIToFColor(DXGI_FORMAT Format, uint32 Width, uint32 Height, uint8 *In, uint32 SrcPitch, FColor* Out, FReadSurfaceDataFlags InFlags)
-{
-	bool bLinearToGamma = InFlags.GetLinearToGamma();
-	switch (Format)
-	{
-		case DXGI_FORMAT_R16_TYPELESS:
-			ConvertRawR16DataToFColor(Width, Height, In, SrcPitch, Out);
-			break;
-		case DXGI_FORMAT_R8G8B8A8_TYPELESS:
-		case DXGI_FORMAT_R8G8B8A8_UNORM:
-		case DXGI_FORMAT_R8G8B8A8_UNORM_SRGB:
-			ConvertRawR8G8B8A8DataToFColor(Width, Height, In, SrcPitch, Out);
-			break;
-		case DXGI_FORMAT_B8G8R8A8_TYPELESS:
-		case DXGI_FORMAT_B8G8R8A8_UNORM:
-		case DXGI_FORMAT_B8G8R8A8_UNORM_SRGB:
-			ConvertRawB8G8R8A8DataToFColor(Width, Height, In, SrcPitch, Out);
-			break;
-		case DXGI_FORMAT_R10G10B10A2_UNORM:
-			ConvertRawR10G10B10A2DataToFColor(Width, Height, In, SrcPitch, Out);
-			break;
-		case DXGI_FORMAT_R16G16B16A16_FLOAT:
-			ConvertRawR16G16B16A16FDataToFColor(Width, Height, In, SrcPitch, Out, bLinearToGamma);
-			break;
-		case DXGI_FORMAT_R11G11B10_FLOAT:
-			ConvertRawR11G11B10DataToFColor(Width, Height, In, SrcPitch, Out, bLinearToGamma);
-			break;
-		case DXGI_FORMAT_R32G32B32A32_FLOAT:
-			ConvertRawR32G32B32A32DataToFColor(Width, Height, In, SrcPitch, Out, bLinearToGamma);
-			break;
-		case DXGI_FORMAT_R24G8_TYPELESS:
-			ConvertRawR24G8DataToFColor(Width, Height, In, SrcPitch, Out, InFlags);
-			break;
-		case DXGI_FORMAT_R32G8X24_TYPELESS:
-			ConvertRawR32DataToFColor(Width, Height, In, SrcPitch, Out, InFlags);
-			break;
-		case DXGI_FORMAT_R16G16B16A16_UNORM:
-			ConvertRawR16G16B16A16DataToFColor(Width, Height, In, SrcPitch, Out);
-			break;
-		case DXGI_FORMAT_R16G16_UNORM:
-			ConvertRawR16G16DataToFColor(Width, Height, In, SrcPitch, Out);
-			break;
-		case DXGI_FORMAT_R8_UNORM:
-			ConvertRawR8DataToFColor(Width, Height, In, SrcPitch, Out);
-			break;
-		case DXGI_FORMAT_R8G8_UNORM:
-			ConvertRawR8G8DataToFColor(Width, Height, In, SrcPitch, Out);
-			break;
-		default:
-			checkf(0, TEXT("Unknown surface format!"));
-			break;
-	}
-}
-
 void FD3D11DynamicRHI::RHIReadSurfaceData(FRHITexture* TextureRHI,FIntRect InRect,TArray<FColor>& OutData, FReadSurfaceDataFlags InFlags)
 {
 	if (!ensure(TextureRHI))
@@ -526,10 +439,23 @@ void FD3D11DynamicRHI::RHIReadSurfaceData(FRHITexture* TextureRHI,FIntRect InRec
 	// Allocate the output buffer.
 	OutData.SetNumUninitialized(SizeX * SizeY);
 
-	uint32 BytesPerPixel = ComputeBytesPerPixel(TextureDesc.Format);
+	uint32 BytesPerPixel = D3D11RT_ComputeBytesPerPixel(TextureDesc.Format);
 	uint32 SrcPitch = SizeX * BytesPerPixel;
+	
+	// switching on the EPixelFormat is risky if the mapping is not what you expect
+	//	verify against TextureDesc.Format
+
+	EPixelFormat Format = TextureRHI->GetFormat();
+	check( GPixelFormats[Format].PlatformFormat == TextureDesc.Format );
+	check( GPixelFormats[Format].BlockBytes == D3D11RT_ComputeBytesPerPixel(TextureDesc.Format) );
 
-	ConvertDXGIToFColor(TextureDesc.Format, SizeX, SizeY, OutDataRaw.GetData(), SrcPitch, OutData.GetData(), InFlags);
+	// ConvertDXGIToFColor switches on the hardware format, not the EPixelFormat :
+
+	if ( ! ConvertDXGIToFColor(TextureDesc.Format, SizeX, SizeY, OutDataRaw.GetData(), SrcPitch, OutData.GetData(), InFlags) )
+	{
+		checkf(0, TEXT("Unsupported surface format!"));
+		OutData.Empty();
+	}
 }
 
 void FD3D11DynamicRHI::ReadSurfaceDataMSAARaw(FRHITexture* TextureRHI,FIntRect InRect,TArray<uint8>& OutData, FReadSurfaceDataFlags InFlags)
@@ -543,7 +469,7 @@ void FD3D11DynamicRHI::ReadSurfaceDataMSAARaw(FRHITexture* TextureRHI,FIntRect I
 	D3D11_TEXTURE2D_DESC TextureDesc;
 	Texture->GetD3D11Texture2D()->GetDesc(&TextureDesc);
 
-	uint32 BytesPerPixel = ComputeBytesPerPixel(TextureDesc.Format);
+	uint32 BytesPerPixel = D3D11RT_ComputeBytesPerPixel(TextureDesc.Format);
 
 	const uint32 NumSamples = TextureDesc.SampleDesc.Count;
 
@@ -615,8 +541,7 @@ void FD3D11DynamicRHI::ReadSurfaceDataMSAARaw(FRHITexture* TextureRHI,FIntRect I
 	}
 	
 	// Allocate the output buffer.
-	OutData.Empty();
-	OutData.AddUninitialized(SizeX * SizeY * NumSamples * BytesPerPixel);
+	OutData.SetNumUninitialized(SizeX * SizeY * NumSamples * BytesPerPixel);
 
 	// Can be optimized by doing all subsamples into a large enough rendertarget in one pass (multiple draw calls)
 	for(uint32 SampleIndex = 0;SampleIndex < NumSamples;++SampleIndex)
@@ -667,7 +592,7 @@ void FD3D11DynamicRHI::RHIMapStagingSurface(FRHITexture* TextureRHI, FRHIGPUFenc
 	ID3D11Resource* Resource = ResourceCast(TextureRHI)->GetResource();
 
 	DXGI_FORMAT Format = (DXGI_FORMAT)GPixelFormats[TextureRHI->GetDesc().Format].PlatformFormat;
-	uint32 BytesPerPixel = ComputeBytesPerPixel(Format);
+	uint32 BytesPerPixel = D3D11RT_ComputeBytesPerPixel(Format);
 
 	D3D11_MAPPED_SUBRESOURCE LockedRect;
 	VERIFYD3D11RESULT_EX(Direct3DDeviceIMContext->Map(Resource, 0, D3D11_MAP_READ, 0, &LockedRect), Direct3DDevice);
@@ -696,7 +621,13 @@ void FD3D11DynamicRHI::RHIReadSurfaceFloatData(FRHITexture* TextureRHI,FIntRect
 	D3D11_TEXTURE2D_DESC TextureDesc;
 	Texture->GetD3D11Texture2D()->GetDesc(&TextureDesc);
 
-	check(TextureDesc.Format == GPixelFormats[PF_FloatRGBA].PlatformFormat);
+	// only supports exactly RGBA16F textures
+	if ( ! ensure(TextureDesc.Format == GPixelFormats[PF_FloatRGBA].PlatformFormat) )
+	{
+		checkf(0, TEXT("Unsupported surface format!"));
+		OutData.Empty();
+		return;
+	}
 
 	// Allocate the output buffer.
 	OutData.SetNumUninitialized(SizeX * SizeY);
@@ -758,64 +689,6 @@ void FD3D11DynamicRHI::RHIReadSurfaceFloatData(FRHITexture* TextureRHI,FIntRect
 	Direct3DDeviceIMContext->Unmap(TempTexture2D,0);
 }
 
-static void ConvertRAWSurfaceDataToFLinearColor(EPixelFormat Format, uint32 Width, uint32 Height, uint8 *In, uint32 SrcPitch, FLinearColor* Out, FReadSurfaceDataFlags InFlags)
-{
-	// @@ missing formats I added to D3D12RenderTarget
-	//	factor out and share code!
-
-	bool bLinearToGamma = InFlags.GetLinearToGamma();
-	if (Format == PF_R16F || Format == PF_R16F_FILTER)
-	{
-		ConvertRawR16DataToFLinearColor(Width, Height, In, SrcPitch, Out);
-	}
-	else if (Format == PF_R8G8B8A8)
-	{
-		ConvertRawR8G8B8A8DataToFLinearColor(Width, Height, In, SrcPitch, Out);
-	}
-	else if (Format == PF_B8G8R8A8)
-	{
-		ConvertRawB8G8R8A8DataToFLinearColor(Width, Height, In, SrcPitch, Out);
-	}
-	else if (Format == PF_A2B10G10R10)
-	{
-		ConvertRawA2B10G10R10DataToFLinearColor(Width, Height, In, SrcPitch, Out);
-	}
-	else if (Format == PF_FloatRGBA)
-	{
-		ConvertRawR16G16B16A16FDataToFLinearColor(Width, Height, In, SrcPitch, Out, InFlags);
-	}
-	else if (Format == PF_FloatRGB || Format == PF_FloatR11G11B10)
-	{
-		ConvertRawRR11G11B10DataToFLinearColor(Width, Height, In, SrcPitch, Out);
-	}
-	else if (Format == PF_A32B32G32R32F)
-	{
-		ConvertRawR32G32B32A32DataToFLinearColor(Width, Height, In, SrcPitch, Out, InFlags);
-	}
-	else if (Format == PF_D24)
-	{
-		ConvertRawR24G8DataToFLinearColor(Width, Height, In, SrcPitch, Out, InFlags);
-	}
-	// Changing Depth Buffers to 32 bit on Dingo as D24S8 is actually implemented as a 32 bit buffer in the hardware
-	else if (Format == PF_DepthStencil)
-	{
-		ConvertRawR32DataToFLinearColor(Width, Height, In, SrcPitch, Out, InFlags);
-	}
-	else if (Format == PF_A16B16G16R16)
-	{
-		ConvertRawR16G16B16A16DataToFLinearColor(Width, Height, In, SrcPitch, Out);
-	}
-	else if (Format == PF_G16R16)
-	{
-		ConvertRawR16G16DataToFLinearColor(Width, Height, In, SrcPitch, Out);
-	}
-	else
-	{
-		// not supported yet
-		check(0);
-	}
-}
-
 void FD3D11DynamicRHI::RHIReadSurfaceData(FRHITexture* TextureRHI, FIntRect InRect, TArray<FLinearColor>& OutData, FReadSurfaceDataFlags InFlags)
 {
 	TArray<uint8> OutDataRaw;
@@ -843,12 +716,21 @@ void FD3D11DynamicRHI::RHIReadSurfaceData(FRHITexture* TextureRHI, FIntRect InRe
 	// Allocate the output buffer.
 	OutData.SetNumUninitialized(SizeX * SizeY);
 
-	uint32 BytesPerPixel = ComputeBytesPerPixel(TextureDesc.Format);
+	uint32 BytesPerPixel = D3D11RT_ComputeBytesPerPixel(TextureDesc.Format);
 	uint32 SrcPitch = SizeX * BytesPerPixel;
 	EPixelFormat Format = TextureRHI->GetFormat();
-	if (Format != PF_Unknown)
+
+	check( GPixelFormats[Format].PlatformFormat == TextureDesc.Format );
+
+	// switching on the EPixelFormat is risky if the mapping is not what you expect
+	//	verify against TextureDesc.Format
+
+	check( GPixelFormats[Format].BlockBytes == D3D11RT_ComputeBytesPerPixel(TextureDesc.Format) );
+
+	if ( ! ConvertRAWSurfaceDataToFLinearColor(Format, SizeX, SizeY, OutDataRaw.GetData(), SrcPitch, OutData.GetData(), InFlags) )
 	{
-		ConvertRAWSurfaceDataToFLinearColor(Format, SizeX, SizeY, OutDataRaw.GetData(), SrcPitch, OutData.GetData(), InFlags);
+		checkf(0, TEXT("Unsupported surface format!"));
+		OutData.Empty();
 	}
 }
 
@@ -867,7 +749,11 @@ void FD3D11DynamicRHI::RHIRead3DSurfaceFloatData(FRHITexture* TextureRHI,FIntRec
 	bool bIsRGBAFmt = TextureDesc.Format == GPixelFormats[PF_FloatRGBA].PlatformFormat;
 	bool bIsR16FFmt = TextureDesc.Format == GPixelFormats[PF_R16F].PlatformFormat;	
 	bool bIsR32FFmt = TextureDesc.Format == GPixelFormats[PF_R32_FLOAT].PlatformFormat;
-	check(bIsRGBAFmt || bIsR16FFmt || bIsR32FFmt);
+	if ( ! ensure(bIsRGBAFmt || bIsR16FFmt || bIsR32FFmt) )
+	{
+		OutData.Empty();
+		return;
+	}
 
 	// Allocate the output buffer.
 	OutData.SetNumUninitialized(SizeX * SizeY * SizeZ);
@@ -956,6 +842,11 @@ void FD3D11DynamicRHI::RHIRead3DSurfaceFloatData(FRHITexture* TextureRHI,FIntRec
 			}
 		}
 	}
+	else
+	{
+		// unsupported format; checked for this earlier
+		check(0);
+	}
 
 	Direct3DDeviceIMContext->Unmap(TempTexture3D,0);
 }
