fix -dxdebug not working anymore on some platforms

diff --git a/Engine/Source/Runtime/D3D12RHI/Private/D3D12Adapter.cpp b/Engine/Source/Runtime/D3D12RHI/Private/D3D12Adapter.cpp
--- a/Engine/Source/Runtime/D3D12RHI/Private/D3D12Adapter.cpp
+++ b/Engine/Source/Runtime/D3D12RHI/Private/D3D12Adapter.cpp
@@ -65,6 +65,8 @@ static int32 GD3D12EnableDRED = 0;
 static int32 GD3D12EnableLightweightDRED = 1;
 #endif // UE_BUILD_SHIPPING || UE_BUILD_TEST
 
+#endif // PLATFORM_WINDOWS
+
 TAutoConsoleVariable<int32> GD3D12DebugCvar (
 	TEXT("r.D3D12.EnableD3DDebug"),
 	0,
@@ -90,6 +92,7 @@ bool D3D12_ShouldBreakOnD3DDebugWarnings()
 	return GD3D12DebugCvar.GetValueOnAnyThread() > 3;
 }
 
+#if PLATFORM_WINDOWS
 static FAutoConsoleVariableRef CVarD3D12EnableGPUBreadCrumbs(
 	TEXT("r.D3D12.BreadCrumbs"),
 	GD3D12EnableGPUBreadCrumbs,
diff --git a/Engine/Source/Runtime/D3D12RHI/Private/D3D12RHI.cpp b/Engine/Source/Runtime/D3D12RHI/Private/D3D12RHI.cpp
--- a/Engine/Source/Runtime/D3D12RHI/Private/D3D12RHI.cpp
+++ b/Engine/Source/Runtime/D3D12RHI/Private/D3D12RHI.cpp
@@ -854,6 +854,31 @@ void FD3D12DynamicRHI::HandleGpuTimeout(FD3D12Payload* Payload, double SecondsSi
 	);
 }
 
+void FD3D12DynamicRHI::SetupD3D12Debug()
+{
+	// Use a debug device if specified on the command line.
+	if (FParse::Param(FCommandLine::Get(), TEXT("d3ddebug")) ||
+		FParse::Param(FCommandLine::Get(), TEXT("d3debug")) ||
+		FParse::Param(FCommandLine::Get(), TEXT("dxdebug")))
+	{
+		GD3D12DebugCvar->Set(1, ECVF_SetByCommandline);
+	}
+	if (FParse::Param(FCommandLine::Get(), TEXT("d3dlogwarnings")))
+	{
+		GD3D12DebugCvar->Set(2, ECVF_SetByCommandline);
+	}
+	if (FParse::Param(FCommandLine::Get(), TEXT("d3dbreakonwarning")))
+	{
+		GD3D12DebugCvar->Set(3, ECVF_SetByCommandline);
+	}
+	if (FParse::Param(FCommandLine::Get(), TEXT("d3dcontinueonerrors")))
+	{
+		GD3D12DebugCvar->Set(4, ECVF_SetByCommandline);
+	}
+	GRHIGlobals.IsDebugLayerEnabled = (GD3D12DebugCvar.GetValueOnAnyThread() > 0);
+
+}
+
 const TCHAR* LexToString(DXGI_FORMAT Format)
 {
 	switch (Format)
diff --git a/Engine/Source/Runtime/D3D12RHI/Private/D3D12RHIPrivate.h b/Engine/Source/Runtime/D3D12RHI/Private/D3D12RHIPrivate.h
--- a/Engine/Source/Runtime/D3D12RHI/Private/D3D12RHIPrivate.h
+++ b/Engine/Source/Runtime/D3D12RHI/Private/D3D12RHIPrivate.h
@@ -579,6 +579,8 @@ protected:
 	// This should only be called by Dynamic RHI member functions
 	FD3D12Device* GetRHIDevice(uint32 GPUIndex) const;
 
+	void SetupD3D12Debug();
+
 	HANDLE FlipEvent;
 
 	const bool bAllowVendorDevice;
diff --git a/Engine/Source/Runtime/D3D12RHI/Private/Windows/WindowsD3D12Device.cpp b/Engine/Source/Runtime/D3D12RHI/Private/Windows/WindowsD3D12Device.cpp
--- a/Engine/Source/Runtime/D3D12RHI/Private/Windows/WindowsD3D12Device.cpp
+++ b/Engine/Source/Runtime/D3D12RHI/Private/Windows/WindowsD3D12Device.cpp
@@ -1348,26 +1348,7 @@ void FD3D12DynamicRHI::Init()
 	}
 #endif
 
-	// Use a debug device if specified on the command line.
-	if (FParse::Param(FCommandLine::Get(), TEXT("d3ddebug")) ||
-		FParse::Param(FCommandLine::Get(), TEXT("d3debug")) ||
-		FParse::Param(FCommandLine::Get(), TEXT("dxdebug")))
-	{
-		GD3D12DebugCvar->Set(1, ECVF_SetByCommandline);
-	}
-	if (FParse::Param(FCommandLine::Get(), TEXT("d3dlogwarnings")))
-	{
-		GD3D12DebugCvar->Set(2, ECVF_SetByCommandline);
-	}
-	if (FParse::Param(FCommandLine::Get(), TEXT("d3dbreakonwarning")))
-	{
-		GD3D12DebugCvar->Set(3, ECVF_SetByCommandline);
-	}
-	if (FParse::Param(FCommandLine::Get(), TEXT("d3dcontinueonerrors")))
-	{
-		GD3D12DebugCvar->Set(4, ECVF_SetByCommandline);
-	}
-	GRHIGlobals.IsDebugLayerEnabled = (GD3D12DebugCvar.GetValueOnAnyThread() > 0);
+	SetupD3D12Debug();
 
 	check(!GIsRHIInitialized);
 
