Refactor D3D12 resource residency tracking to allow multiple residency handles per resource. Add reserved resource backing heaps to the residency tracking system.

diff --git a/Engine/Source/Runtime/D3D12RHI/Private/D3D12Adapter.cpp b/Engine/Source/Runtime/D3D12RHI/Private/D3D12Adapter.cpp
--- a/Engine/Source/Runtime/D3D12RHI/Private/D3D12Adapter.cpp
+++ b/Engine/Source/Runtime/D3D12RHI/Private/D3D12Adapter.cpp
@@ -1959,22 +1959,16 @@ void FD3D12Adapter::DumpTrackedAllocationData(FOutputDevice& OutputDevice, bool
 		if (ResourceDesc.Dimension == D3D12_RESOURCE_DIMENSION_BUFFER)
 		{
 			BufferAllocations.Add(AllocationData);
-			TotalAllocatedBufferSize += AllocationData.AllocationSize;			
-#if ENABLE_RESIDENCY_MANAGEMENT
-			TotalResidentBufferSize += (AllocationData.ResourceAllocation->GetResidencyHandle().ResidencyStatus == D3DX12Residency::ManagedObject::RESIDENCY_STATUS::RESIDENT) ? AllocationData.AllocationSize : 0;
-#else
-			TotalResidentBufferSize += AllocationData.AllocationSize;
-#endif 
+			TotalAllocatedBufferSize += AllocationData.AllocationSize;
+			// TODO: accurately account for partially-resident resources
+			TotalResidentBufferSize += (AllocationData.ResourceAllocation->GetResource()->IsResident()) ? AllocationData.AllocationSize : 0;
 		}
 		else
 		{
 			TextureAllocations.Add(AllocationData);
 			TotalAllocatedTextureSize += AllocationData.AllocationSize;
-#if ENABLE_RESIDENCY_MANAGEMENT
-			TotalResidentTextureSize += (AllocationData.ResourceAllocation->GetResidencyHandle().ResidencyStatus == D3DX12Residency::ManagedObject::RESIDENCY_STATUS::RESIDENT) ? AllocationData.AllocationSize : 0;
-#else
-			TotalResidentTextureSize += AllocationData.AllocationSize;
-#endif 
+			// TODO: accurately account for partially-resident resources
+			TotalResidentTextureSize += (AllocationData.ResourceAllocation->GetResource()->IsResident()) ? AllocationData.AllocationSize : 0;
 		}
 	}
 
@@ -1989,7 +1983,7 @@ void FD3D12Adapter::DumpTrackedAllocationData(FOutputDevice& OutputDevice, bool
 
 		bool bResident = true;
 #if ENABLE_RESIDENCY_MANAGEMENT
-		bResident = AllocationData.ResourceAllocation->GetResidencyHandle().ResidencyStatus == D3DX12Residency::ManagedObject::RESIDENCY_STATUS::RESIDENT;
+		bResident = AllocationData.ResourceAllocation->GetResource()->IsResident();
 #endif 
 		if (!bResident && bResidentOnly)
 		{
@@ -2036,7 +2030,7 @@ void FD3D12Adapter::DumpTrackedAllocationData(FOutputDevice& OutputDevice, bool
 
 		bool bResident = true;
 #if ENABLE_RESIDENCY_MANAGEMENT
-		bResident = AllocationData.ResourceAllocation->GetResidencyHandle().ResidencyStatus == D3DX12Residency::ManagedObject::RESIDENCY_STATUS::RESIDENT;
+		bResident = AllocationData.ResourceAllocation->GetResource()->IsResident();
 #endif 
 		if (!bResident && bResidentOnly)
 		{
diff --git a/Engine/Source/Runtime/D3D12RHI/Private/D3D12CommandContext.h b/Engine/Source/Runtime/D3D12RHI/Private/D3D12CommandContext.h
--- a/Engine/Source/Runtime/D3D12RHI/Private/D3D12CommandContext.h
+++ b/Engine/Source/Runtime/D3D12RHI/Private/D3D12CommandContext.h
@@ -346,7 +346,7 @@ public:
 	void UpdateResidency(TConstArrayView<FD3D12ResidencyHandle*> Handles) { GetCommandList().UpdateResidency(Handles); }
 	void UpdateResidency(FD3D12ResidencyHandle& Handle                  ) { GetCommandList().UpdateResidency({ &Handle }); }
 	void UpdateResidency(FD3D12ResidencyHandle* Handle                  ) { check(Handle  ); GetCommandList().UpdateResidency({ Handle }); }
-	void UpdateResidency(FD3D12Resource* Resource                       ) { check(Resource); GetCommandList().UpdateResidency({ &Resource->GetResidencyHandle() }); }
+	void UpdateResidency(FD3D12Resource* Resource                       ) { check(Resource); GetCommandList().UpdateResidency(Resource->GetResidencyHandles()); }
 #else
 	void UpdateResidency(TConstArrayView<FD3D12ResidencyHandle*> Handles) { }
 	void UpdateResidency(FD3D12ResidencyHandle& Handle                  ) { }
diff --git a/Engine/Source/Runtime/D3D12RHI/Private/D3D12CommandList.cpp b/Engine/Source/Runtime/D3D12RHI/Private/D3D12CommandList.cpp
--- a/Engine/Source/Runtime/D3D12RHI/Private/D3D12CommandList.cpp
+++ b/Engine/Source/Runtime/D3D12RHI/Private/D3D12CommandList.cpp
@@ -24,8 +24,7 @@ void FD3D12CommandList::UpdateResidency(TConstArrayView<FD3D12ResidencyHandle*>
 #if ENABLE_RESIDENCY_MANAGEMENT
 	for (FD3D12ResidencyHandle* Handle : Handles)
 	{
-		check(Handle);
-		if (D3DX12Residency::IsInitialized(*Handle))
+		if (D3DX12Residency::IsInitialized(Handle))
 		{
 			check(Device->GetGPUMask() == Handle->GPUObject->GetGPUMask());
 			D3DX12Residency::Insert(*ResidencySet, *Handle);
diff --git a/Engine/Source/Runtime/D3D12RHI/Private/D3D12DescriptorCache.cpp b/Engine/Source/Runtime/D3D12RHI/Private/D3D12DescriptorCache.cpp
--- a/Engine/Source/Runtime/D3D12RHI/Private/D3D12DescriptorCache.cpp
+++ b/Engine/Source/Runtime/D3D12RHI/Private/D3D12DescriptorCache.cpp
@@ -207,7 +207,7 @@ void FD3D12DescriptorCache::SetVertexBuffers(FD3D12VertexBufferCache& Cache)
 	{
 		if (Cache.CurrentVertexBufferResources[i])
 		{
-			Context.UpdateResidency(Cache.ResidencyHandles[i]);
+			Context.UpdateResidency(Cache.Resources[i]);
 
 			FD3D12Resource* Resource = Cache.CurrentVertexBufferResources[i]->GetResource();
 			if (Resource && Resource->RequiresResourceStateTracking())
@@ -255,7 +255,7 @@ D3D12_GPU_DESCRIPTOR_HANDLE FD3D12DescriptorCache::BuildUAVTable(EShaderFrequenc
 			SrcDescriptors[SlotIndex] = UAVs[SlotIndex]->GetOfflineCpuHandle();
 
 			Context.TransitionResource(UAVs[SlotIndex], D3D12_RESOURCE_STATE_UNORDERED_ACCESS);
-			Context.UpdateResidency(Cache.ResidencyHandles[ShaderStage][SlotIndex]);
+			Context.UpdateResidency(Cache.Resources[ShaderStage][SlotIndex]);
 		}
 	}
 	FD3D12UnorderedAccessViewCache::CleanSlots(CurrentDirtySlotMask, SlotsNeeded);
@@ -491,7 +491,7 @@ D3D12_GPU_DESCRIPTOR_HANDLE FD3D12DescriptorCache::BuildSRVTable(EShaderFrequenc
 			const D3D12_RESOURCE_STATES State = GetBindingResourceState(Context, ShaderStage, SRV);
 
 			Context.TransitionResource(SRV, State & ValidResourceStates);
-			Context.UpdateResidency(Cache.ResidencyHandles[ShaderStage][SlotIndex]);
+			Context.UpdateResidency(Cache.Resources[ShaderStage][SlotIndex]);
 		}
 		else
 		{
@@ -545,7 +545,7 @@ void FD3D12DescriptorCache::PrepareBindlessViews(EShaderFrequency ShaderStage, T
 		{
 			const D3D12_RESOURCE_STATES State = GetBindingResourceState(Context, ShaderStage, SRV);
 			Context.TransitionResource(SRV, State & ValidResourceStates);
-			Context.UpdateResidency(&SRV->GetResidencyHandle());
+			Context.UpdateResidency(SRV->GetResource());
 		}
 	}
 
@@ -554,7 +554,7 @@ void FD3D12DescriptorCache::PrepareBindlessViews(EShaderFrequency ShaderStage, T
 		if (ensure(UAV))
 		{
 			Context.TransitionResource(UAV, D3D12_RESOURCE_STATE_UNORDERED_ACCESS);
-			Context.UpdateResidency(&UAV->GetResidencyHandle());
+			Context.UpdateResidency(UAV->GetResource());
 		}
 	}
 }
@@ -591,7 +591,7 @@ void FD3D12DescriptorCache::SetConstantBufferViews(EShaderFrequency ShaderStage,
 			Device->CopyDescriptorsSimple(1, DestDescriptor, CBVHandles[SlotIndex], D3D12_DESCRIPTOR_HEAP_TYPE_CBV_SRV_UAV);
 
 			// Update residency.
-			Context.UpdateResidency(Cache.ResidencyHandles[ShaderStage][SlotIndex]);
+			Context.UpdateResidency(Cache.Resources[ShaderStage][SlotIndex]);
 		}
 		else
 		{
@@ -667,7 +667,7 @@ void FD3D12DescriptorCache::SetRootConstantBuffers(EShaderFrequency ShaderStage,
 			}
 
 			// Update residency.
-			Context.UpdateResidency(Cache.ResidencyHandles[ShaderStage][SlotIndex]);
+			Context.UpdateResidency(Cache.Resources[ShaderStage][SlotIndex]);
 
 			// Clear the dirty bit.
 			FD3D12ConstantBufferCache::CleanSlot(CurrentDirtySlotMask, SlotIndex);
diff --git a/Engine/Source/Runtime/D3D12RHI/Private/D3D12RayTracing.cpp b/Engine/Source/Runtime/D3D12RHI/Private/D3D12RayTracing.cpp
--- a/Engine/Source/Runtime/D3D12RHI/Private/D3D12RayTracing.cpp
+++ b/Engine/Source/Runtime/D3D12RHI/Private/D3D12RayTracing.cpp
@@ -3596,14 +3596,16 @@ void FD3D12RayTracingScene::BuildAccelerationStructure(FD3D12CommandContext& Com
 			auto AddResidencyHandleForResource = [&UniqueResidencyHandles, &GeometryResidencyHandlesForThisGPU] (FD3D12Resource* Resource)
 			{
 			#if ENABLE_RESIDENCY_MANAGEMENT
-				FD3D12ResidencyHandle& ResidencyHandle = Resource->GetResidencyHandle();
-				if (D3DX12Residency::IsInitialized(ResidencyHandle))
+				for (FD3D12ResidencyHandle* ResidencyHandle : Resource->GetResidencyHandles())
 				{
-					bool bIsAlreadyInSet = false;
-					UniqueResidencyHandles.Add(&ResidencyHandle, &bIsAlreadyInSet);
-					if (!bIsAlreadyInSet)
+					if (D3DX12Residency::IsInitialized(*ResidencyHandle))
 					{
-						GeometryResidencyHandlesForThisGPU.Add(&ResidencyHandle);
+						bool bIsAlreadyInSet = false;
+						UniqueResidencyHandles.Add(ResidencyHandle, &bIsAlreadyInSet);
+						if (!bIsAlreadyInSet)
+						{
+							GeometryResidencyHandlesForThisGPU.Add(ResidencyHandle);
+						}
 					}
 				}
 			#endif // ENABLE_RESIDENCY_MANAGEMENT
diff --git a/Engine/Source/Runtime/D3D12RHI/Private/D3D12Resources.cpp b/Engine/Source/Runtime/D3D12RHI/Private/D3D12Resources.cpp
--- a/Engine/Source/Runtime/D3D12RHI/Private/D3D12Resources.cpp
+++ b/Engine/Source/Runtime/D3D12RHI/Private/D3D12Resources.cpp
@@ -119,11 +119,6 @@ FD3D12Resource::FD3D12Resource(FD3D12Device* ParentDevice,
 	}
 #endif
 
-	if (!IsPlacedResource())
-	{
-		ResidencyHandle = MakeUnique<FD3D12ResidencyHandle>();
-	}
-
 	if (Desc.bReservedResource)
 	{
 		checkf(Heap == nullptr, TEXT("Reserved resources are not expected to have a heap"));
@@ -133,10 +128,13 @@ FD3D12Resource::FD3D12Resource(FD3D12Device* ParentDevice,
 
 FD3D12Resource::~FD3D12Resource()
 {
-	if (!IsPlacedResource() && D3DX12Residency::IsInitialized(*ResidencyHandle))
+#if ENABLE_RESIDENCY_MANAGEMENT
+	if (D3DX12Residency::IsInitialized(ResidencyHandle))
 	{
 		D3DX12Residency::EndTrackingObject(GetParentDevice()->GetResidencyManager(), *ResidencyHandle);
 	}
+	delete ResidencyHandle;
+#endif // ENABLE_RESIDENCY_MANAGEMENT
 
 #if NV_AFTERMATH
 	if (GDX12NVAfterMathTrackResources)
@@ -194,8 +192,7 @@ void FD3D12Resource::CommitReservedResource(ID3D12CommandQueue* D3DCommandQueue,
 	const uint64 MaxHeapSize = uint64(CVarD3D12ReservedResourceHeapSizeMB.GetValueOnAnyThread()) * 1024 * 1024;
 	const uint64 NumHeaps = FMath::DivideAndRoundUp(TotalSize, MaxHeapSize);
 
-	TArray<TRefCountPtr<ID3D12Heap>>& Heaps = ReservedResourceData->BackingHeaps;
-	Heaps.Reserve(NumHeaps);
+	ReservedResourceData->BackingHeaps.Reserve(NumHeaps);
 
 	const uint32 MaxTilesPerHeap = uint32(MaxHeapSize / TileSizeInBytes);
 
@@ -273,12 +270,12 @@ void FD3D12Resource::CommitReservedResource(ID3D12CommandQueue* D3DCommandQueue,
 		NewHeapDesc.Flags = HeapFlags;
 		NewHeapDesc.SizeInBytes = ThisHeapSize;
 		NewHeapDesc.Properties = BackingHeapProps;
-		TRefCountPtr<ID3D12Heap> NewHeap;
-		VERIFYD3D12RESULT(D3DDevice->CreateHeap(&NewHeapDesc, IID_PPV_ARGS(NewHeap.GetInitReference())));
+		ID3D12Heap* D3DHeap = nullptr;
+		VERIFYD3D12RESULT(D3DDevice->CreateHeap(&NewHeapDesc, IID_PPV_ARGS(&D3DHeap)));
 
 		if (bHighPriorityResource)
 		{
-			Adapter->SetResidencyPriority(NewHeap, D3D12_RESIDENCY_PRIORITY_HIGH, GPUIndex);
+			Adapter->SetResidencyPriority(D3DHeap, D3D12_RESIDENCY_PRIORITY_HIGH, GPUIndex);
 		}
 
 		D3D12_TILED_RESOURCE_COORDINATE ResourceCoordinate = {}; // Coordinates are in tiles, not pixels
@@ -316,7 +313,7 @@ void FD3D12Resource::CommitReservedResource(ID3D12CommandQueue* D3DCommandQueue,
 		const uint32 RangeTileCount = RegionSize.NumTiles;
 
 		D3DCommandQueue->UpdateTileMappings(GetResource(), 1,
-			&ResourceCoordinate, &RegionSize, NewHeap.GetReference(),
+			&ResourceCoordinate, &RegionSize, D3DHeap,
 			1,
 			&RangeFlags,
 			&HeapRangeStartOffsetInTiles,
@@ -324,13 +321,21 @@ void FD3D12Resource::CommitReservedResource(ID3D12CommandQueue* D3DCommandQueue,
 			MappingFlags);
 
 	#if NAME_OBJECTS
-		const int32 HeapIndex = Heaps.Num();
+		const int32 HeapIndex = ReservedResourceData->BackingHeaps.Num();
 		FString HeapName = FString::Printf(TEXT("%s.Heap[%d]"), DebugName.IsValid() ? *DebugName.ToString() : TEXT("UNKNOWN"), HeapIndex);
-		::SetName(NewHeap, *HeapName);
+		const TCHAR* HeapNameChars = *HeapName;
+	#else
+		const TCHAR* HeapNameChars = TEXT("ReservedResourceBackingHeap");
 	#endif // NAME_OBJECTS
 
 		NumMappedTiles += RegionSize.NumTiles;
-		Heaps.Add(MoveTemp(NewHeap));
+
+		TRefCountPtr<FD3D12Heap> NewHeap = new FD3D12Heap(GetParentDevice(), GetVisibilityMask());
+		NewHeap->SetHeap(D3DHeap, HeapNameChars, true /*bTrack*/, false /*bForceGetGPUAddress*/);
+		NewHeap->BeginTrackingResidency(ThisHeapSize);
+
+		ReservedResourceData->ResidencyHandles.Append(NewHeap->GetResidencyHandles());
+		ReservedResourceData->BackingHeaps.Add(MoveTemp(NewHeap));
 	}
 
 	ReservedResourceData->CommittedSizeInBytes = NumMappedTiles * TileSizeInBytes;
@@ -363,16 +368,17 @@ void FD3D12Resource::StartTrackingForResidency()
 		return;
 	}
 
-	check(IsGPUOnly(HeapType));	// This is checked at a higher level before calling this function.		
-	if (!IsPlacedResource())
+	check(IsGPUOnly(HeapType));	// This is checked at a higher level before calling this function.
+	if (!IsPlacedResource() && !IsReservedResource())
 	{
-		check(D3DX12Residency::IsInitialized(*ResidencyHandle) == false);
+		checkf(!ResidencyHandle, TEXT("Residency tracking is already initialzied for this resource"));
+		ResidencyHandle = new FD3D12ResidencyHandle;
 
 		const D3D12_RESOURCE_ALLOCATION_INFO Info = GetParentDevice()->GetResourceAllocationInfoUncached(Desc);
 		D3DX12Residency::Initialize(*ResidencyHandle, Resource.GetReference(), Info.SizeInBytes, this);
 		D3DX12Residency::BeginTrackingObject(GetParentDevice()->GetResidencyManager(), *ResidencyHandle);
 	}
-#endif
+#endif // ENABLE_RESIDENCY_MANAGEMENT
 }
 
 void FD3D12Resource::DeferDelete()
@@ -387,7 +393,6 @@ void FD3D12Resource::DeferDelete()
 FD3D12Heap::FD3D12Heap(FD3D12Device* Parent, FRHIGPUMask VisibleNodes, HeapId InTraceParentHeapId) :
 	FD3D12DeviceChild(Parent),
 	FD3D12MultiNodeGPUObject(Parent->GetGPUMask(), VisibleNodes),
-	ResidencyHandle(),
 	TraceParentHeapId(InTraceParentHeapId)
 {
 }
@@ -413,9 +418,9 @@ FD3D12Heap::~FD3D12Heap()
 #if ENABLE_RESIDENCY_MANAGEMENT
 	if (D3DX12Residency::IsInitialized(ResidencyHandle))
 	{
-		D3DX12Residency::EndTrackingObject(GetParentDevice()->GetResidencyManager(), ResidencyHandle);
-		ResidencyHandle = {};
+		D3DX12Residency::EndTrackingObject(GetParentDevice()->GetResidencyManager(), *ResidencyHandle);
 	}
+	delete ResidencyHandle;
 #endif // ENABLE_RESIDENCY_MANAGEMENT
 
 	// Release actual d3d object
@@ -471,8 +476,10 @@ void FD3D12Heap::SetHeap(ID3D12Heap* HeapIn, const TCHAR* const InName, bool bIn
 void FD3D12Heap::BeginTrackingResidency(uint64 Size)
 {
 #if ENABLE_RESIDENCY_MANAGEMENT
-	D3DX12Residency::Initialize(ResidencyHandle, Heap.GetReference(), Size, this);
-	D3DX12Residency::BeginTrackingObject(GetParentDevice()->GetResidencyManager(), ResidencyHandle);
+	checkf(!ResidencyHandle, TEXT("Residency tracking is already initialzied for this resource"));
+	ResidencyHandle = new FD3D12ResidencyHandle;
+	D3DX12Residency::Initialize(*ResidencyHandle, Heap.GetReference(), Size, this);
+	D3DX12Residency::BeginTrackingObject(GetParentDevice()->GetResidencyManager(), *ResidencyHandle);
 #endif
 }
 
@@ -954,7 +961,6 @@ void FD3D12ResourceLocation::InternalClear()
 	UnderlyingResource = nullptr;
 	MappedBaseAddress = nullptr;
 	GPUVirtualAddress = 0;
-	ResidencyHandle = nullptr;
 	Size = 0;
 	OffsetFromBaseOfResource = 0;
 	FMemory::Memzero(AllocatorData);
@@ -1215,12 +1221,10 @@ void FD3D12ResourceLocation::UpdateStandAloneStats(bool bIncrement)
 void FD3D12ResourceLocation::SetResource(FD3D12Resource* Value)
 {
 	check(UnderlyingResource == nullptr);
-	check(ResidencyHandle == nullptr);
 
 	GPUVirtualAddress = Value->GetGPUVirtualAddress();
 
 	UnderlyingResource = Value;
-	ResidencyHandle = &UnderlyingResource->GetResidencyHandle();
 }
 
 
@@ -1320,8 +1324,6 @@ bool FD3D12ResourceLocation::OnAllocationMoved(FRHICommandListBase& RHICmdList,
 		UnderlyingResource = NewResource;
 		GPUVirtualAddress = UnderlyingResource->GetGPUVirtualAddress() + OffsetFromBaseOfResource;
 	}
-	
-	ResidencyHandle = &UnderlyingResource->GetResidencyHandle();
 
 	// Refresh aliases
 	for (FRHIPoolAllocationData* OtherAlias = AllocationData.GetFirstAlias(); OtherAlias; OtherAlias = OtherAlias->GetNext())
@@ -1331,7 +1333,6 @@ bool FD3D12ResourceLocation::OnAllocationMoved(FRHICommandListBase& RHICmdList,
 		OtherResourceLocation->OffsetFromBaseOfResource = OffsetFromBaseOfResource;
 		OtherResourceLocation->UnderlyingResource = UnderlyingResource;
 		OtherResourceLocation->GPUVirtualAddress = GPUVirtualAddress;
-		OtherResourceLocation->ResidencyHandle = ResidencyHandle;
 	}
 
 	// TODO: recreate the aliased UAV resource if we have one. For the time being we just check that we don't get here with such a resource,
diff --git a/Engine/Source/Runtime/D3D12RHI/Private/D3D12Resources.h b/Engine/Source/Runtime/D3D12RHI/Private/D3D12Resources.h
--- a/Engine/Source/Runtime/D3D12RHI/Private/D3D12Resources.h
+++ b/Engine/Source/Runtime/D3D12RHI/Private/D3D12Resources.h
@@ -84,7 +84,7 @@ public:
 
 	inline FName GetName() const { return HeapName; }
 	inline D3D12_HEAP_DESC GetHeapDesc() const { return HeapDesc; }
-	inline FD3D12ResidencyHandle& GetResidencyHandle() { return ResidencyHandle; }
+	inline TConstArrayView<FD3D12ResidencyHandle*> GetResidencyHandles() { return MakeArrayView(&ResidencyHandle, 1); }
 	inline D3D12_GPU_VIRTUAL_ADDRESS GetGPUVirtualAddress() const { return GPUVirtualAddress; }
 	inline void SetIsTransient(bool bInIsTransient) { bIsTransient = bInIsTransient; }
 	inline bool GetIsTransient() const { return bIsTransient; }
@@ -96,7 +96,7 @@ private:
 	bool bTrack = true;
 	D3D12_HEAP_DESC HeapDesc;
 	D3D12_GPU_VIRTUAL_ADDRESS GPUVirtualAddress = 0;
-	FD3D12ResidencyHandle ResidencyHandle;
+	FD3D12ResidencyHandle* ResidencyHandle = nullptr; // Residency handle owned by this object
 	HeapId TraceHeapId;
 	HeapId TraceParentHeapId;
 	bool bIsTransient = false; // Whether this is a transient heap
@@ -149,7 +149,7 @@ private:
 	TRefCountPtr<ID3D12Resource> UAVAccessResource;
 	TRefCountPtr<FD3D12Heap> Heap;
 
-	TUniquePtr<FD3D12ResidencyHandle> ResidencyHandle;
+	FD3D12ResidencyHandle* ResidencyHandle = nullptr; // Residency handle owned by this object
 
 	D3D12_GPU_VIRTUAL_ADDRESS GPUVirtualAddress{};
 	void* ResourceBaseAddress{};
@@ -185,7 +185,8 @@ private:
 
 	struct FD3D12ReservedResourceData
 	{
-		TArray<TRefCountPtr<ID3D12Heap>> BackingHeaps;
+		TArray<TRefCountPtr<FD3D12Heap>> BackingHeaps;
+		TArray<FD3D12ResidencyHandle*> ResidencyHandles; // Flattened array of residency handles owned by backing heaps, used to support batched GetResidencyHandles()
 		uint64 CommittedSizeInBytes = 0;
 	};
 	TUniquePtr<FD3D12ReservedResourceData> ReservedResourceData;
@@ -302,15 +303,49 @@ public:
 	inline bool ShouldDeferDelete() const { return bDeferDelete; }
 	void DeferDelete();
 
+	inline bool IsReservedResource() const { return ReservedResourceData.IsValid(); }
 	inline bool IsPlacedResource() const { return Heap.GetReference() != nullptr; }
 	inline FD3D12Heap* GetHeap() const { return Heap; };
 	inline bool IsDepthStencilResource() const { return bDepthStencil; }
 
 	void StartTrackingForResidency();
 
-	FD3D12ResidencyHandle& GetResidencyHandle()
+	bool IsResident() const
 	{
-		return IsPlacedResource() ? Heap->GetResidencyHandle() : *ResidencyHandle;
+#if ENABLE_RESIDENCY_MANAGEMENT
+		// Treat resource as resident if at least one backing heap is resident.
+		// Technically we should return a partial residency status.
+		for (FD3D12ResidencyHandle* Handle : GetResidencyHandles())
+		{
+			if (Handle->ResidencyStatus == FD3D12ResidencyHandle::RESIDENCY_STATUS::RESIDENT)
+			{
+				return true;
+			}
+		}
+		return false;
+#else // ENABLE_RESIDENCY_MANAGEMENT
+		return true;
+#endif // ENABLE_RESIDENCY_MANAGEMENT
+	}
+
+	TConstArrayView<FD3D12ResidencyHandle*> GetResidencyHandles() const
+	{
+#if ENABLE_RESIDENCY_MANAGEMENT
+		if (IsPlacedResource())
+		{
+			return Heap->GetResidencyHandles();
+		}
+		else if (IsReservedResource())
+		{
+			return ReservedResourceData->ResidencyHandles;
+		}
+		else
+		{
+			return MakeArrayView(&ResidencyHandle, 1);
+		}
+#else // ENABLE_RESIDENCY_MANAGEMENT
+		return {};
+#endif // ENABLE_RESIDENCY_MANAGEMENT
 	}
 
 	struct FD3D12ResourceTypeHelper
@@ -591,7 +626,6 @@ public:
 	D3D12_GPU_VIRTUAL_ADDRESS          GetGPUVirtualAddress          () const { return GPUVirtualAddress;                                    }
 	uint64                             GetOffsetFromBaseOfResource   () const { return OffsetFromBaseOfResource;                             }
 	uint64                             GetSize                       () const { return Size;                                                 }
-	FD3D12ResidencyHandle&             GetResidencyHandle            ()       { check(ResidencyHandle); return *ResidencyHandle;             }
 	FD3D12BuddyAllocatorPrivateData&   GetBuddyAllocatorPrivateData  ()       { return AllocatorData.BuddyAllocatorPrivateData;              }
 	FD3D12BlockAllocatorPrivateData&   GetBlockAllocatorPrivateData  ()       { return AllocatorData.BlockAllocatorPrivateData;              }
 	FD3D12SegListAllocatorPrivateData& GetSegListAllocatorPrivateData()       { return AllocatorData.SegListAllocatorPrivateData;            }
@@ -681,7 +715,6 @@ private:
 
 	FD3D12BaseShaderResource* Owner{};
 	FD3D12Resource* UnderlyingResource{};
-	FD3D12ResidencyHandle* ResidencyHandle{};
 
 	// Which allocator this belongs to
 	union
@@ -871,10 +904,8 @@ public:
 		OutResourceInfo.Type = GetType();
 		OutResourceInfo.VRamAllocation.AllocationSize = ResourceLocation.GetSize();
 		OutResourceInfo.IsTransient = ResourceLocation.IsTransient();
-#if ENABLE_RESIDENCY_MANAGEMENT
-		OutResourceInfo.bResident = GetResource() && GetResource()->GetResidencyHandle().ResidencyStatus == D3DX12Residency::ManagedObject::RESIDENCY_STATUS::RESIDENT;
-#endif
-		
+		OutResourceInfo.bResident = GetResource() && GetResource()->IsResident();
+
 		return true;
 	}
 #endif
diff --git a/Engine/Source/Runtime/D3D12RHI/Private/D3D12StateCache.cpp b/Engine/Source/Runtime/D3D12RHI/Private/D3D12StateCache.cpp
--- a/Engine/Source/Runtime/D3D12RHI/Private/D3D12StateCache.cpp
+++ b/Engine/Source/Runtime/D3D12RHI/Private/D3D12StateCache.cpp
@@ -1278,7 +1278,7 @@ void FD3D12StateCache::SetUAV(EShaderFrequency ShaderStage, uint32 SlotIndex, FD
 
 	if (UAV)
 	{
-		Cache.ResidencyHandles[ShaderStage][SlotIndex] = &UAV->GetResidencyHandle();
+		Cache.Resources[ShaderStage][SlotIndex] = UAV->GetResource();
 
 		FD3D12Resource* CounterResource = UAV->GetCounterResource();
 		if (CounterResource)
@@ -1316,7 +1316,7 @@ void FD3D12StateCache::SetUAV(EShaderFrequency ShaderStage, uint32 SlotIndex, FD
 	}
 	else
 	{
-		Cache.ResidencyHandles[ShaderStage][SlotIndex] = nullptr;
+		Cache.Resources[ShaderStage][SlotIndex] = nullptr;
 	}
 }
 
@@ -1480,7 +1480,7 @@ void FD3D12StateCache::InternalSetStreamSource(FD3D12ResourceLocation* VertexBuf
 
 		if (VertexBufferLocation != nullptr)
 		{
-			PipelineState.Graphics.VBCache.ResidencyHandles[StreamIndex] = &VertexBufferLocation->GetResource()->GetResidencyHandle();
+			PipelineState.Graphics.VBCache.Resources[StreamIndex] = VertexBufferLocation->GetResource();
 			FMemory::Memcpy(CurrentView, NewView);
 			PipelineState.Graphics.VBCache.BoundVBMask |= ((VBSlotMask)1 << StreamIndex);
 		}
@@ -1488,7 +1488,7 @@ void FD3D12StateCache::InternalSetStreamSource(FD3D12ResourceLocation* VertexBuf
 		{
 			FMemory::Memzero(&CurrentView, sizeof(CurrentView));
 			PipelineState.Graphics.VBCache.CurrentVertexBufferResources[StreamIndex] = nullptr;
-			PipelineState.Graphics.VBCache.ResidencyHandles[StreamIndex] = nullptr;
+			PipelineState.Graphics.VBCache.Resources[StreamIndex] = nullptr;
 
 			PipelineState.Graphics.VBCache.BoundVBMask &= ~((VBSlotMask)1 << StreamIndex);
 		}
@@ -1520,12 +1520,12 @@ void FD3D12StateCache::SetShaderResourceView(EShaderFrequency ShaderFrequency, F
 			bSRVSCleared = false;
 
 			Cache.BoundMask[ShaderFrequency] |= ((SRVSlotMask)1 << ResourceIndex);
-			Cache.ResidencyHandles[ShaderFrequency][ResourceIndex] = &SRV->GetResidencyHandle();
+			Cache.Resources[ShaderFrequency][ResourceIndex] = SRV->GetResource();
 		}
 		else
 		{
 			Cache.BoundMask[ShaderFrequency] &= ~((SRVSlotMask)1 << ResourceIndex);
-			Cache.ResidencyHandles[ShaderFrequency][ResourceIndex] = nullptr;
+			Cache.Resources[ShaderFrequency][ResourceIndex] = nullptr;
 		}
 
 		// Find the highest set SRV
diff --git a/Engine/Source/Runtime/D3D12RHI/Private/D3D12StateCachePrivate.h b/Engine/Source/Runtime/D3D12RHI/Private/D3D12StateCachePrivate.h
--- a/Engine/Source/Runtime/D3D12RHI/Private/D3D12StateCachePrivate.h
+++ b/Engine/Source/Runtime/D3D12RHI/Private/D3D12StateCachePrivate.h
@@ -68,14 +68,14 @@ struct FD3D12VertexBufferCache
 	{
 		FMemory::Memzero(CurrentVertexBufferViews, sizeof(CurrentVertexBufferViews));
 		FMemory::Memzero(CurrentVertexBufferResources, sizeof(CurrentVertexBufferResources));
-		FMemory::Memzero(ResidencyHandles, sizeof(ResidencyHandles));
+		FMemory::Memzero(Resources, sizeof(Resources));
 		MaxBoundVertexBufferIndex = INDEX_NONE;
 		BoundVBMask = 0;
 	}
 
 	D3D12_VERTEX_BUFFER_VIEW CurrentVertexBufferViews[MAX_VBS];
 	FD3D12ResourceLocation* CurrentVertexBufferResources[MAX_VBS];
-	FD3D12ResidencyHandle* ResidencyHandles[MAX_VBS];
+	FD3D12Resource* Resources[MAX_VBS];
 	int32 MaxBoundVertexBufferIndex;
 	VBSlotMask BoundVBMask;
 };
@@ -163,7 +163,7 @@ struct FD3D12ConstantBufferCache : public FD3D12ResourceCache<CBVSlotMask>
 		DirtyAll();
 
 		FMemory::Memzero(CurrentGPUVirtualAddress, sizeof(CurrentGPUVirtualAddress));
-		FMemory::Memzero(ResidencyHandles, sizeof(ResidencyHandles));
+		FMemory::Memzero(Resources, sizeof(Resources));
 #if D3D12RHI_USE_CONSTANT_BUFFER_VIEWS
 		FMemory::Memzero(CBHandles, sizeof(CBHandles));
 #endif
@@ -173,7 +173,7 @@ struct FD3D12ConstantBufferCache : public FD3D12ResourceCache<CBVSlotMask>
 	D3D12_CPU_DESCRIPTOR_HANDLE CBHandles[SF_NumStandardFrequencies][MAX_CBS];
 #endif
 	D3D12_GPU_VIRTUAL_ADDRESS CurrentGPUVirtualAddress[SF_NumStandardFrequencies][MAX_CBS];
-	FD3D12ResidencyHandle* ResidencyHandles[SF_NumStandardFrequencies][MAX_CBS];
+	FD3D12Resource* Resources[SF_NumStandardFrequencies][MAX_CBS];
 };
 
 struct FD3D12ShaderResourceViewCache : public FD3D12ResourceCache<SRVSlotMask>
@@ -187,7 +187,7 @@ struct FD3D12ShaderResourceViewCache : public FD3D12ResourceCache<SRVSlotMask>
 	{
 		DirtyAll();
 
-		FMemory::Memzero(ResidencyHandles);
+		FMemory::Memzero(Resources);
 		FMemory::Memzero(BoundMask);
 		
 		for (int32& Index : MaxBoundIndex)
@@ -205,7 +205,7 @@ struct FD3D12ShaderResourceViewCache : public FD3D12ResourceCache<SRVSlotMask>
 	}
 
 	FD3D12ShaderResourceView* Views[SF_NumStandardFrequencies][MAX_SRVS];
-	FD3D12ResidencyHandle* ResidencyHandles[SF_NumStandardFrequencies][MAX_SRVS];
+	FD3D12Resource* Resources[SF_NumStandardFrequencies][MAX_SRVS];
 
 	SRVSlotMask BoundMask[SF_NumStandardFrequencies];
 	int32 MaxBoundIndex[SF_NumStandardFrequencies];
@@ -223,7 +223,7 @@ struct FD3D12UnorderedAccessViewCache : public FD3D12ResourceCache<UAVSlotMask>
 		DirtyAll();
 
 		FMemory::Memzero(Views);
-		FMemory::Memzero(ResidencyHandles);
+		FMemory::Memzero(Resources);
 
 		for (uint32& Index : StartSlot)
 		{
@@ -232,7 +232,7 @@ struct FD3D12UnorderedAccessViewCache : public FD3D12ResourceCache<UAVSlotMask>
 	}
 
 	FD3D12UnorderedAccessView* Views[SF_NumStandardFrequencies][MAX_UAVS];
-	FD3D12ResidencyHandle* ResidencyHandles[SF_NumStandardFrequencies][MAX_UAVS];
+	FD3D12Resource* Resources[SF_NumStandardFrequencies][MAX_UAVS];
 	uint32 StartSlot[SF_NumStandardFrequencies];
 };
 
@@ -524,7 +524,7 @@ public:
 			if (ResourceLocation.GetGPUVirtualAddress() != CurrentGPUVirtualAddress)
 			{
 				CurrentGPUVirtualAddress = ResourceLocation.GetGPUVirtualAddress();
-				CBVCache.ResidencyHandles[ShaderFrequency][SlotIndex] = &ResourceLocation.GetResource()->GetResidencyHandle();
+				CBVCache.Resources[ShaderFrequency][SlotIndex] = ResourceLocation.GetResource();
 				FD3D12ConstantBufferCache::DirtySlot(CBVCache.DirtySlotMask[ShaderFrequency], SlotIndex);
 			}
 
@@ -535,7 +535,7 @@ public:
 		else if (CurrentGPUVirtualAddress != 0)
 		{
 			CurrentGPUVirtualAddress = 0;
-			CBVCache.ResidencyHandles[ShaderFrequency][SlotIndex] = nullptr;
+			CBVCache.Resources[ShaderFrequency][SlotIndex] = {};
 			FD3D12ConstantBufferCache::DirtySlot(CBVCache.DirtySlotMask[ShaderFrequency], SlotIndex);
 #if D3D12RHI_USE_CONSTANT_BUFFER_VIEWS
 			CBVCache.CBHandles[ShaderFrequency][SlotIndex].ptr = 0;
@@ -562,7 +562,7 @@ public:
 			D3D12_GPU_VIRTUAL_ADDRESS& CurrentGPUVirtualAddress = CBVCache.CurrentGPUVirtualAddress[ShaderFrequency][SlotIndex];
 			check(Location.GetGPUVirtualAddress() != CurrentGPUVirtualAddress);
 			CurrentGPUVirtualAddress = Location.GetGPUVirtualAddress();
-			CBVCache.ResidencyHandles[ShaderFrequency][SlotIndex] = &Location.GetResource()->GetResidencyHandle();
+			CBVCache.Resources[ShaderFrequency][SlotIndex] = Location.GetResource();
 			FD3D12ConstantBufferCache::DirtySlot(CBVCache.DirtySlotMask[ShaderFrequency], SlotIndex);
 
 #if D3D12RHI_USE_CONSTANT_BUFFER_VIEWS
diff --git a/Engine/Source/Runtime/D3D12RHI/Private/D3D12Submission.cpp b/Engine/Source/Runtime/D3D12RHI/Private/D3D12Submission.cpp
--- a/Engine/Source/Runtime/D3D12RHI/Private/D3D12Submission.cpp
+++ b/Engine/Source/Runtime/D3D12RHI/Private/D3D12Submission.cpp
@@ -541,19 +541,10 @@ FD3D12DynamicRHI::FProcessResult FD3D12DynamicRHI::ProcessSubmissionQueue()
 								check(Range.End > Range.Start);
 
 	#if ENABLE_RESIDENCY_MANAGEMENT
-								FD3D12ResidencyHandle* ResidencyHandles[] =
-								{
-									&Range.Heap->GetHeapResidencyHandle(),
-									&Range.Heap->GetResultBuffer()->GetResidencyHandle()
-								};
-
-								for (FD3D12ResidencyHandle* Handle : ResidencyHandles)
-								{
-									if (Handle->IsInitialized())
-									{
-										GetResolveCommandList()->UpdateResidency({ Handle });
-									}
-								}
+								TArray<FD3D12ResidencyHandle*, TInlineAllocator<2>> ResidencyHandles;
+								ResidencyHandles.Add(&Range.Heap->GetHeapResidencyHandle());
+								ResidencyHandles.Append(Range.Heap->GetResultBuffer()->GetResidencyHandles());
+								GetResolveCommandList()->UpdateResidency(ResidencyHandles);
 	#endif // ENABLE_RESIDENCY_MANAGEMENT
 
 								if (Range.Heap->GetD3DQueryHeap())
@@ -705,7 +696,10 @@ FD3D12CommandList* FD3D12DynamicRHI::GenerateBarrierCommandListAndUpdateState(FD
 		}
 
 #if ENABLE_RESIDENCY_MANAGEMENT
-		ResidencyHandles.Add(&PRB.Resource->GetResidencyHandle());
+		for (FD3D12ResidencyHandle* Handle : PRB.Resource->GetResidencyHandles())
+		{
+			ResidencyHandles.Add(Handle);
+		}
 #endif // ENABLE_RESIDENCY_MANAGEMENT
 	}
 
diff --git a/Engine/Source/Runtime/D3D12RHI/Private/D3D12Texture.cpp b/Engine/Source/Runtime/D3D12RHI/Private/D3D12Texture.cpp
--- a/Engine/Source/Runtime/D3D12RHI/Private/D3D12Texture.cpp
+++ b/Engine/Source/Runtime/D3D12RHI/Private/D3D12Texture.cpp
@@ -1426,7 +1426,7 @@ bool FD3D12Texture::GetResourceInfo(FRHIResourceInfo& OutResourceInfo) const
 	OutResourceInfo.VRamAllocation.AllocationSize = ResourceLocation.GetSize();
 	OutResourceInfo.IsTransient = ResourceLocation.IsTransient();
 #if ENABLE_RESIDENCY_MANAGEMENT
-	OutResourceInfo.bResident = GetResource() && GetResource()->GetResidencyHandle().ResidencyStatus == D3DX12Residency::ManagedObject::RESIDENCY_STATUS::RESIDENT;
+	OutResourceInfo.bResident = GetResource() && GetResource()->IsResident();
 #endif
 
 	return true;
diff --git a/Engine/Source/Runtime/D3D12RHI/Private/D3D12UAV.cpp b/Engine/Source/Runtime/D3D12RHI/Private/D3D12UAV.cpp
--- a/Engine/Source/Runtime/D3D12RHI/Private/D3D12UAV.cpp
+++ b/Engine/Source/Runtime/D3D12RHI/Private/D3D12UAV.cpp
@@ -388,7 +388,7 @@ void FD3D12CommandContext::ClearUAV(TRHICommandList_RecursiveHazardous<FD3D12Com
 
 				Context.FlushResourceBarriers();
 				Context.GraphicsCommandList()->ClearUnorderedAccessViewUint(GPUHandle, UAV.GetOfflineCpuHandle(), UAV.GetResource()->GetResource(), *reinterpret_cast<const UINT(*)[4]>(ClearValues), 0, nullptr);
-				Context.UpdateResidency(UnorderedAccessView->GetResidencyHandle());
+				Context.UpdateResidency(UnorderedAccessView->GetResidencyHandles());
 				Context.ConditionalSplitCommandList();
 
 				if (Context.IsDefaultContext())
diff --git a/Engine/Source/Runtime/D3D12RHI/Private/D3D12View.h b/Engine/Source/Runtime/D3D12RHI/Private/D3D12View.h
--- a/Engine/Source/Runtime/D3D12RHI/Private/D3D12View.h
+++ b/Engine/Source/Runtime/D3D12RHI/Private/D3D12View.h
@@ -10,6 +10,7 @@
 
 class FD3D12Buffer;
 class FD3D12Texture;
+struct FD3D12ResidencyHandle;
 
 struct FD3D12DefaultViews
 {
@@ -181,10 +182,9 @@ public:
 
 	struct FResourceInfo
 	{
-		FD3D12BaseShaderResource* BaseResource     = nullptr;
-		FD3D12ResourceLocation*   ResourceLocation = nullptr;
-		FD3D12Resource*           Resource         = nullptr;
-		FD3D12ResidencyHandle*    ResidencyHandle  = nullptr;
+		FD3D12BaseShaderResource*               BaseResource     = nullptr;
+		FD3D12ResourceLocation*                 ResourceLocation = nullptr;
+		FD3D12Resource*                         Resource         = nullptr;
 
 		FResourceInfo() = default;
 
@@ -193,7 +193,6 @@ public:
 			: BaseResource    (InBaseResource)
 			, ResourceLocation(InBaseResource ? &InBaseResource->ResourceLocation : nullptr)
 			, Resource        (InBaseResource ? InBaseResource->GetResource()     : nullptr)
-			, ResidencyHandle (Resource       ? &Resource->GetResidencyHandle()   : nullptr)
 		{}
 
 		// Constructor for manual views (does not automatically register for resource renames)
@@ -201,15 +200,14 @@ public:
 			: BaseResource    (nullptr)
 			, ResourceLocation(InResourceLocation)
 			, Resource        (InResourceLocation ? InResourceLocation->GetResource() : nullptr)
-			, ResidencyHandle (Resource           ? &Resource->GetResidencyHandle()   : nullptr)
 		{}
 	};
 
-	FD3D12Resource*             GetResource        () const { check(IsInitialized()); return ResourceInfo.Resource;         }
-	FD3D12ResourceLocation*     GetResourceLocation() const { check(IsInitialized()); return ResourceInfo.ResourceLocation; }
-	FD3D12ResidencyHandle&      GetResidencyHandle () const { check(IsInitialized()); return *ResourceInfo.ResidencyHandle; }
-	FD3D12ViewSubset const&     GetViewSubset      () const { check(IsInitialized()); return ViewSubset;                    }
-	FD3D12OfflineDescriptor     GetOfflineCpuHandle() const { check(IsInitialized()); return OfflineCpuHandle;              }
+	FD3D12Resource*                         GetResource        () const { check(IsInitialized()); return ResourceInfo.Resource;         }
+	FD3D12ResourceLocation*                 GetResourceLocation() const { check(IsInitialized()); return ResourceInfo.ResourceLocation; }
+	TConstArrayView<FD3D12ResidencyHandle*> GetResidencyHandles() const { check(IsInitialized()); return ResourceInfo.Resource ? ResourceInfo.Resource->GetResidencyHandles() : TConstArrayView<FD3D12ResidencyHandle*>(); }
+	FD3D12ViewSubset const&                 GetViewSubset      () const { check(IsInitialized()); return ViewSubset;                    }
+	FD3D12OfflineDescriptor                 GetOfflineCpuHandle() const { check(IsInitialized()); return OfflineCpuHandle;              }
 
 #if PLATFORM_SUPPORTS_BINDLESS_RENDERING
 	FRHIDescriptorHandle        GetBindlessHandle() const { return BindlessHandle;           }
