Adding support for executing outside code "on" the graphics/copy queue as platform extension to RHI D3D12 Windows
#rnx
#tests on local machine in testbed project

diff --git a/Engine/Source/Runtime/D3D12RHI/Private/D3D12RHI.cpp b/Engine/Source/Runtime/D3D12RHI/Private/D3D12RHI.cpp
--- a/Engine/Source/Runtime/D3D12RHI/Private/D3D12RHI.cpp
+++ b/Engine/Source/Runtime/D3D12RHI/Private/D3D12RHI.cpp
@@ -879,6 +879,27 @@ void FD3D12DynamicRHI::SetupD3D12Debug()
 
 }
 
+void FD3D12DynamicRHI::RHIRunOnQueue(ED3D12RHIRunOnQueueType QueueType, TFunction<void(ID3D12CommandQueue*)>&& CodeToRun, bool bWaitForSubmission)
+{
+	FGraphEventRef SubmissionEvent;
+
+	FD3D12Payload* Payload = new FD3D12Payload(GetRHIDevice(0), (QueueType == ED3D12RHIRunOnQueueType::Graphics) ?  ED3D12QueueType::Direct : ED3D12QueueType::Copy);
+	Payload->PreExecuteCallback = MoveTemp(CodeToRun);
+
+	if (bWaitForSubmission)
+	{
+		SubmissionEvent = FGraphEvent::CreateGraphEvent();
+		Payload->SubmissionEvent = SubmissionEvent;
+	}
+
+	SubmitPayloads(MakeArrayView(&Payload, 1));
+
+	if (SubmissionEvent && !SubmissionEvent->IsComplete())
+	{
+		SubmissionEvent->Wait();
+	}
+}
+
 const TCHAR* LexToString(DXGI_FORMAT Format)
 {
 	switch (Format)
diff --git a/Engine/Source/Runtime/D3D12RHI/Private/D3D12RHIPrivate.h b/Engine/Source/Runtime/D3D12RHI/Private/D3D12RHIPrivate.h
--- a/Engine/Source/Runtime/D3D12RHI/Private/D3D12RHIPrivate.h
+++ b/Engine/Source/Runtime/D3D12RHI/Private/D3D12RHIPrivate.h
@@ -379,6 +379,8 @@ public:
 	virtual IRHIPlatformCommandList* RHIFinalizeContext(IRHIComputeContext* Context) final override;
 	virtual void RHISubmitCommandLists(TArrayView<IRHIPlatformCommandList*> CommandLists, bool bFlushResources) final override;
 
+	virtual void RHIRunOnQueue(ED3D12RHIRunOnQueueType QueueType, TFunction<void(ID3D12CommandQueue*)>&& CodeToRun, bool bWaitForSubmission) final override;
+
 	virtual IRHITransientResourceAllocator* RHICreateTransientResourceAllocator() override;
 
 	virtual void RHIWriteGPUFence_TopOfPipe(FRHICommandListBase& RHICmdList, FRHIGPUFence* FenceRHI) final override;
diff --git a/Engine/Source/Runtime/D3D12RHI/Private/D3D12Submission.cpp b/Engine/Source/Runtime/D3D12RHI/Private/D3D12Submission.cpp
--- a/Engine/Source/Runtime/D3D12RHI/Private/D3D12Submission.cpp
+++ b/Engine/Source/Runtime/D3D12RHI/Private/D3D12Submission.cpp
@@ -1163,6 +1163,14 @@ FD3D12PayloadBase::~FD3D12PayloadBase()
 	}
 }
 
+void FD3D12PayloadBase::PreExecute()
+{
+	if (PreExecuteCallback)
+	{
+		PreExecuteCallback(Queue.D3DCommandQueue);
+	}
+}
+
 void FBreadcrumbStack::Initialize(TUniquePtr<FD3D12DiagnosticBuffer>& DiagnosticBuffer)
 {
 	{
diff --git a/Engine/Source/Runtime/D3D12RHI/Private/D3D12Submission.h b/Engine/Source/Runtime/D3D12RHI/Private/D3D12Submission.h
--- a/Engine/Source/Runtime/D3D12RHI/Private/D3D12Submission.h
+++ b/Engine/Source/Runtime/D3D12RHI/Private/D3D12Submission.h
@@ -246,7 +246,7 @@ struct FD3D12PayloadBase
 
 	} SyncPointsToWait;
 
-	virtual void PreExecute() {}
+	virtual void PreExecute();
 
 	// Wait
 	TArray<FManualFence> FencesToWait;
@@ -280,6 +280,9 @@ struct FD3D12PayloadBase
 
 	virtual ~FD3D12PayloadBase();
 
+	// Used by RHIRunOnQueue
+	TFunction<void(ID3D12CommandQueue*)> PreExecuteCallback;
+
 protected:
 	FD3D12PayloadBase(FD3D12Device* Device, ED3D12QueueType QueueType);
 };
diff --git a/Engine/Source/Runtime/D3D12RHI/Public/ID3D12DynamicRHI.h b/Engine/Source/Runtime/D3D12RHI/Public/ID3D12DynamicRHI.h
--- a/Engine/Source/Runtime/D3D12RHI/Public/ID3D12DynamicRHI.h
+++ b/Engine/Source/Runtime/D3D12RHI/Public/ID3D12DynamicRHI.h
@@ -12,6 +12,12 @@ struct FD3D12MinimalAdapterDesc
 	uint32 NumDeviceNodes{};
 };
 
+enum class ED3D12RHIRunOnQueueType
+{
+	Graphics = 0,
+	Copy,
+};
+
 struct ID3D12DynamicRHI : public FDynamicRHI
 {
 	virtual ERHIInterfaceType     GetInterfaceType() const override final { return ERHIInterfaceType::D3D12; }
@@ -47,6 +53,8 @@ struct ID3D12DynamicRHI : public FDynamicRHI
 	virtual void                  RHIWaitManualFence(FRHICommandList& RHICmdList, ID3D12Fence* Fence, uint64 Value) = 0;
 
 	virtual void                  RHIVerifyResult(ID3D12Device* Device, HRESULT Result, const ANSICHAR* Code, const ANSICHAR* Filename, uint32 Line, FString Message = FString()) const = 0;
+
+	virtual void				  RHIRunOnQueue(ED3D12RHIRunOnQueueType QueueType, TFunction<void(ID3D12CommandQueue*)>&& CodeToRun, bool bWaitForSubmission) = 0;
 };
 
 inline bool IsRHID3D12()
