D3D12 Bindless Resources with CPU managed descriptor heap updates
- D3D12 PC Bindless needs descriptor heaps managed on the CPU; we cannot update them on the GPU timeline.
- Each context now has a FD3D12ContextBindlessState that contains the per-context GPU descriptor heap as well as descriptor rollbacks to apply to the heap before submission.
- We have to roll descriptors back to their values before the incoming view updates were applied so that we can leverage the CPU heap copy at all times. This isn't deferring the updates; it's storing the values before the updates and making sure the heap is used with those values.
- When a context encounters a draw/dispatch and there were any descriptor updates, the previously used heap is updated with the correct set of descriptors before a new heap is created for the subsequent draws/dispatches.

Additional changes
- Allowing RHIs to override FRHITextureReference for custom bindless implementations.
- Adding FRHIDescriptorAllocator::GetAllocatedRange to allow managers to find the smallest range of descriptors that need to be copied to new heaps.
- DescriptorCache now holds reference counted pointers to the bindless heaps to avoid potential use after free scenarios.
- Adding ED3D12DescriptorHeapFlags to mirror D3D12_DESCRIPTOR_HEAP_FLAGS while adding new flags.
- Adding ability to pool descriptor heaps to avoid high OS overhead when constantly allocating new ones.
- Pooling descriptor heaps required more descriptor heap managers to implement CleanupResources.

diff --git a/Engine/Config/Windows/WindowsEngine.ini b/Engine/Config/Windows/WindowsEngine.ini
--- a/Engine/Config/Windows/WindowsEngine.ini
+++ b/Engine/Config/Windows/WindowsEngine.ini
@@ -20,6 +20,8 @@ QualityLevelMapping="high"
 ; Enable single RHI thread stall for D3D11
 r.Occlusion.SingleRHIThreadStall=1
 r.Shadow.DetectVertexShaderLayerAtRuntime=1
+D3D12.Bindless.ResourceDescriptorHeapSize=100000
+D3D12.Bindless.SamplerDescriptorHeapSize=2048
 
 [PlatformCrypto]
 PlatformRequiresDataCrypto=True
diff --git a/Engine/Source/Runtime/D3D12RHI/Private/D3D12Adapter.cpp b/Engine/Source/Runtime/D3D12RHI/Private/D3D12Adapter.cpp
--- a/Engine/Source/Runtime/D3D12RHI/Private/D3D12Adapter.cpp
+++ b/Engine/Source/Runtime/D3D12RHI/Private/D3D12Adapter.cpp
@@ -1471,12 +1471,10 @@ void FD3D12Adapter::CleanupResources()
 
 	TransientMemoryCache = nullptr;
 
-#if D3D12_RHI_RAYTRACING
 	for (uint32 GPUIndex : FRHIGPUMask::All())
 	{
-		Devices[GPUIndex]->CleanupRayTracing();
+		Devices[GPUIndex]->CleanupResources();
 	}
-#endif // D3D12_RHI_RAYTRACING
 }
 
 FD3D12Adapter::~FD3D12Adapter()
diff --git a/Engine/Source/Runtime/D3D12RHI/Private/D3D12BindlessDescriptors.cpp b/Engine/Source/Runtime/D3D12RHI/Private/D3D12BindlessDescriptors.cpp
--- a/Engine/Source/Runtime/D3D12RHI/Private/D3D12BindlessDescriptors.cpp
+++ b/Engine/Source/Runtime/D3D12RHI/Private/D3D12BindlessDescriptors.cpp
@@ -30,22 +30,29 @@ FD3D12DescriptorHeap* UE::D3D12BindlessDescriptors::CreateCpuHeap(FD3D12Device*
 		HeapName,
 		InType,
 		InNewNumDescriptorsPerHeap,
-		D3D12_DESCRIPTOR_HEAP_FLAG_NONE
+		ED3D12DescriptorHeapFlags::None
 	);
 }
 
 FD3D12DescriptorHeap* UE::D3D12BindlessDescriptors::CreateGpuHeap(FD3D12Device* InDevice, ERHIDescriptorHeapType InType, uint32 InNewNumDescriptorsPerHeap)
 {
+	SCOPED_NAMED_EVENT_F(TEXT("CreateNewBindlessHeap (%d)"), FColor::Turquoise, InNewNumDescriptorsPerHeap);
+
 	const TCHAR* const HeapName = (InType == ERHIDescriptorHeapType::Standard) ? TEXT("BindlessResources") : TEXT("BindlessSamplers");
 
 	return InDevice->GetDescriptorHeapManager().AllocateIndependentHeap(
 		HeapName,
 		InType,
 		InNewNumDescriptorsPerHeap,
-		D3D12_DESCRIPTOR_HEAP_FLAG_SHADER_VISIBLE
+		ED3D12DescriptorHeapFlags::GpuVisible | ED3D12DescriptorHeapFlags::Poolable
 	);
 }
 
+void UE::D3D12BindlessDescriptors::DeferredFreeHeap(FD3D12Device* InDevice, FD3D12DescriptorHeap* InHeap)
+{
+	InDevice->GetDescriptorHeapManager().DeferredFreeHeap(InHeap);
+}
+
 ////////////////////////////////////////////////////////////////////////////////////////////////////////////////////////
 // FD3D12BindlessSamplerManager
 
@@ -56,9 +63,18 @@ FD3D12BindlessSamplerManager::FD3D12BindlessSamplerManager(FD3D12Device* InDevic
 	GpuHeap = UE::D3D12BindlessDescriptors::CreateGpuHeap(InDevice, ERHIDescriptorHeapType::Sampler, InNumDescriptors);
 }
 
-FRHIDescriptorHandle FD3D12BindlessSamplerManager::Allocate()
+void FD3D12BindlessSamplerManager::CleanupResources()
+{
+	GpuHeap = nullptr;
+}
+
+FRHIDescriptorHandle FD3D12BindlessSamplerManager::AllocateAndInitialize(FD3D12SamplerState* SamplerState)
 {
 	FRHIDescriptorHandle Result = Allocator.Allocate();
+	if (ensure(Result.IsValid()))
+	{
+		UE::D3D12Descriptors::CopyDescriptor(GetParentDevice(), GpuHeap, Result, SamplerState->OfflineDescriptor);
+	}
 	check(Result.IsValid());
 	return Result;
 }
@@ -71,12 +87,14 @@ void FD3D12BindlessSamplerManager::Free(FRHIDescriptorHandle InHandle)
 	}
 }
 
-void FD3D12BindlessSamplerManager::UpdateDescriptorImmediately(FRHIDescriptorHandle DstHandle, D3D12_CPU_DESCRIPTOR_HANDLE SrcDescriptor)
+void FD3D12BindlessSamplerManager::OpenCommandList(FD3D12CommandContext& Context)
 {
-	if (DstHandle.IsValid())
-	{
-		UE::D3D12Descriptors::CopyDescriptor(GetParentDevice(), GpuHeap, DstHandle, SrcDescriptor);
-	}
+	Context.StateCache.GetDescriptorCache()->SetBindlessSamplersHeapDirectly(GetHeap());
+}
+
+void FD3D12BindlessSamplerManager::CloseCommandList(FD3D12CommandContext& Context)
+{
+	Context.StateCache.GetDescriptorCache()->SetBindlessSamplersHeapDirectly(nullptr);
 }
 
 ////////////////////////////////////////////////////////////////////////////////////////////////////////////////////////
@@ -86,21 +104,14 @@ void FD3D12BindlessSamplerManager::UpdateDescriptorImmediately(FRHIDescriptorHan
 
 FD3D12BindlessResourceManager::FD3D12BindlessResourceManager(FD3D12Device* InDevice, uint32 InNumDescriptors, TConstArrayView<TStatId> InStats)
 	: FD3D12DeviceChild(InDevice)
+	, CpuHeap(UE::D3D12BindlessDescriptors::CreateCpuHeap(InDevice, ERHIDescriptorHeapType::Standard, InNumDescriptors))
 	, Allocator(ERHIDescriptorHeapType::Standard, InNumDescriptors, InStats)
 {
-	ConfiguredPipelines = ERHIPipeline::Graphics;
-	if (GSupportsEfficientAsyncCompute)
-	{
-		EnumAddFlags(ConfiguredPipelines, ERHIPipeline::AsyncCompute);
-	}
-
-	EnumeratePipelines([&](ERHIPipeline PipelineIndex)
-	{
-		FPipeline& Pipeline = Pipelines[PipelineIndex];
+}
 
-		Pipeline.CpuHeap = UE::D3D12BindlessDescriptors::CreateCpuHeap(InDevice, ERHIDescriptorHeapType::Standard, InNumDescriptors);
-		Pipeline.GpuHeap = UE::D3D12BindlessDescriptors::CreateGpuHeap(InDevice, ERHIDescriptorHeapType::Standard, InNumDescriptors);
-	});
+void FD3D12BindlessResourceManager::CleanupResources()
+{
+	CpuHeap.SafeRelease();
 }
 
 FRHIDescriptorHandle FD3D12BindlessResourceManager::Allocate()
@@ -118,25 +129,158 @@ void FD3D12BindlessResourceManager::Free(FRHIDescriptorHandle InHandle)
 	}
 }
 
-void FD3D12BindlessResourceManager::UpdateDescriptorImmediately(FRHIDescriptorHandle DstHandle, D3D12_CPU_DESCRIPTOR_HANDLE SrcDescriptor)
+void FD3D12BindlessResourceManager::UpdateDescriptorImmediately(FRHIDescriptorHandle DstHandle, FD3D12View* View)
+{
+	if (DstHandle.IsValid())
+	{
+		UE::D3D12Descriptors::CopyDescriptor(GetParentDevice(), CpuHeap, DstHandle, View->GetOfflineCpuHandle());
+	}
+}
+
+void FD3D12BindlessResourceManager::UpdateDescriptor(FRHICommandListBase& RHICmdList, FRHIDescriptorHandle DstHandle, FD3D12View* View)
 {
 	if (DstHandle.IsValid())
 	{
-		EnumeratePipelines([&](ERHIPipeline PipelineIndex)
+		for (ERHIPipeline PipelineIndex : GetRHIPipelines())
+		{
+			FRHICommandListScopedPipeline Scope(RHICmdList, PipelineIndex);
+			RHICmdList.EnqueueLambda([this, View, DstHandle](FRHICommandListBase& ExecutingCmdList)
+			{
+				FD3D12CommandContext& Context =
+					ExecutingCmdList.IsGraphics()
+					? static_cast<FD3D12CommandContext&>(ExecutingCmdList.GetContext())
+					: static_cast<FD3D12CommandContext&>(ExecutingCmdList.GetComputeContext());
+
+				if (Context.IsOpen())
+				{
+					FD3D12OfflineDescriptor CopyOfPreviousDescriptorValue = UE::D3D12Descriptors::CreateOfflineCopy(GetParentDevice(), CpuHeap, DstHandle);
+					Context.GetBindlessState().PendingDescriptorRollbacks.Add(GetParentDevice(), DstHandle, CopyOfPreviousDescriptorValue);
+				}
+			});
+		}
+
+		RHICmdList.EnqueueLambda([this, DstHandle, View](FRHICommandListBase& ExecutingCmdList)
 		{
-			UE::D3D12Descriptors::CopyDescriptor(GetParentDevice(), Pipelines[PipelineIndex].GpuHeap, DstHandle, SrcDescriptor);
+			UpdateDescriptorImmediately(DstHandle, View);
 		});
 	}
 }
 
-void FD3D12BindlessResourceManager::UpdateDescriptor(FRHICommandListBase& RHICmdList, FRHIDescriptorHandle DstHandle, D3D12_CPU_DESCRIPTOR_HANDLE SrcDescriptor)
+void FD3D12BindlessResourceManager::FlushPendingDescriptorUpdates(FD3D12CommandContext& Context)
+{
+	FD3D12ContextBindlessState& State = Context.GetBindlessState();
+
+	if (State.PendingDescriptorRollbacks.Num() > 0)
+	{
+		// If we have rollbacks to apply, we have to move to a new heap.
+
+		// First finalize the previous heap if it was set.
+		FinalizeHeapOnState(State);
+
+		// Then create a new heap to use
+		CreateHeapOnState(State);
+
+		if (ensure(Context.IsOpen()))
+		{
+			// Finally tell the Context that we're using this heap,
+			// this call also makes sure the heap is set on the d3d command list.
+			Context.StateCache.GetDescriptorCache()->SwitchToNewBindlessResourceHeap(State.CurrentGpuHeap);
+		}
+	}
+}
+
+void FD3D12BindlessResourceManager::OpenCommandList(FD3D12CommandContext& Context)
+{
+	FD3D12ContextBindlessState& State = Context.GetBindlessState();
+
+	// Always create a new descriptor heap to use.
+	// TODO: defer this to the first FlushPendingDescriptorUpdates
+	CreateHeapOnState(State);
+
+	// Assign the heap to the descriptor cache
+	Context.StateCache.GetDescriptorCache()->SetBindlessResourcesHeapDirectly(State.CurrentGpuHeap);
+}
+
+void FD3D12BindlessResourceManager::CloseCommandList(FD3D12CommandContext& Context)
+{
+	FD3D12ContextBindlessState& State = Context.GetBindlessState();
+
+	// First finalize the current heap if any was set
+	FinalizeHeapOnState(State);
+
+	// Then clear the reference from the state cache
+	Context.StateCache.GetDescriptorCache()->SetBindlessResourcesHeapDirectly(nullptr);
+}
+
+void FD3D12BindlessResourceManager::FinalizeContext(FD3D12CommandContext& Context)
+{
+	if (Context.IsOpen())
+	{
+		Context.CloseCommandList();
+	}
+
+	FD3D12ContextBindlessState& State = Context.GetBindlessState();
+
+	if (State.UsedHeaps.Num() > 0)
+	{
+		for (const FD3D12DescriptorHeapPtr& UsedHeap : State.UsedHeaps)
+		{
+			checkSlow(UsedHeap);
+
+			// Now queue it up for deletion
+			UE::D3D12BindlessDescriptors::DeferredFreeHeap(GetParentDevice(), UsedHeap);
+		}
+
+		State.UsedHeaps.Empty();
+	}
+
+	check(!Context.GetBindlessState().HasAnyPending());
+}
+
+void FD3D12BindlessResourceManager::CopyCpuHeap(FD3D12DescriptorHeap* DestinationHeap)
+{
+	// Copy the smallest possible set of descriptors from the CPU heap to the new GPU heap.
+	FRHIDescriptorAllocatorRange AllocatedRange(0, 0);
+	if (Allocator.GetAllocatedRange(AllocatedRange))
+	{
+		const uint32 NumDescriptorsToCopy = AllocatedRange.Last - AllocatedRange.First + 1;
+		UE::D3D12Descriptors::CopyDescriptors(GetParentDevice(), DestinationHeap, CpuHeap, AllocatedRange.First, NumDescriptorsToCopy);
+	}
+}
+
+void FD3D12BindlessResourceManager::CreateHeapOnState(FD3D12ContextBindlessState& State)
 {
-	UpdateDescriptorImmediately(DstHandle, SrcDescriptor);
+	checkf(State.CurrentGpuHeap == nullptr, TEXT("FinalizeHeapOnState was not called before CreateHeapOnState"));
+
+	State.CurrentGpuHeap = UE::D3D12BindlessDescriptors::CreateGpuHeap(GetParentDevice(), Allocator.GetType(), Allocator.GetCapacity());
 }
 
-void FD3D12BindlessResourceManager::FlushPendingDescriptorUpdates(FD3D12CommandContext& Context, ERHIPipeline PipelineIndex, const FD3D12PendingResourceDescriptorUpdates& PendingDescriptorUpdates)
+void FD3D12BindlessResourceManager::FinalizeHeapOnState(FD3D12ContextBindlessState& State)
 {
-	checkNoEntry(); // TODO
+	if (State.CurrentGpuHeap)
+	{
+		// Since we're about to stop using this heap, make sure it's in the correct state for the GPU:
+		//  1. Update wholesale from the CPU copy
+		//  2. Take the set of previous descriptor values and apply them to this heap since they were updated in the CPU copy
+		// This effectively "rolls back" the dynamic updates to their correct state while keeping all other immediate updates
+
+		CopyCpuHeap(State.CurrentGpuHeap);
+
+		FD3D12PendingDescriptorUpdates& Updates = State.PendingDescriptorRollbacks;
+
+		if (!Updates.IsEmpty())
+		{
+			UE::D3D12Descriptors::CopyDescriptors(GetParentDevice(), State.CurrentGpuHeap, Updates.Handles, Updates.OfflineDescriptors);
+
+			Updates.Empty(GetParentDevice());
+		}
+
+		// Move this heap to the used list for cleanup at context finalize.
+		State.UsedHeaps.Emplace(MoveTemp(State.CurrentGpuHeap));
+	}
+
+	// Always clear the pending rollbacks. If we didn't previously bind a descriptor heap, we don't need to roll any descriptors back
+	State.PendingDescriptorRollbacks.Empty(GetParentDevice());
 }
 
 #endif // D3D12RHI_CUSTOM_BINDLESS_RESOURCE_MANAGER
@@ -194,16 +338,34 @@ void FD3D12BindlessDescriptorManager::Init()
 	}
 }
 
-FRHIDescriptorHandle FD3D12BindlessDescriptorManager::Allocate(ERHIDescriptorHeapType InType)
+void FD3D12BindlessDescriptorManager::CleanupResources()
+{
+	if (ResourceManager)
+	{
+		ResourceManager->CleanupResources();
+	}
+
+	if (SamplerManager)
+	{
+		SamplerManager->CleanupResources();
+	}
+}
+
+FRHIDescriptorHandle FD3D12BindlessDescriptorManager::AllocateResourceHandle()
 {
-	if (InType == ERHIDescriptorHeapType::Standard && ResourceManager)
+	if (ResourceManager)
 	{
 		return ResourceManager->Allocate();
 	}
 
-	if (InType == ERHIDescriptorHeapType::Sampler && SamplerManager)
+	return FRHIDescriptorHandle();
+}
+
+FRHIDescriptorHandle FD3D12BindlessDescriptorManager::AllocateAndInitialize(FD3D12SamplerState* SamplerState)
+{
+	if (SamplerManager)
 	{
-		return SamplerManager->Allocate();
+		return SamplerManager->AllocateAndInitialize(SamplerState);
 	}
 
 	return FRHIDescriptorHandle();
@@ -235,17 +397,23 @@ void FD3D12BindlessDescriptorManager::DeferredFreeFromDestructor(FRHIDescriptorH
 	}
 }
 
-void FD3D12BindlessDescriptorManager::UpdateDescriptorImmediately(FRHIDescriptorHandle DstHandle, D3D12_CPU_DESCRIPTOR_HANDLE SrcDescriptor)
+void FD3D12BindlessDescriptorManager::UpdateDescriptorImmediately(FRHIDescriptorHandle DstHandle, FD3D12View* View)
 {
 	if (DstHandle.GetType() == ERHIDescriptorHeapType::Standard && ResourceManager)
 	{
-		ResourceManager->UpdateDescriptorImmediately(DstHandle, SrcDescriptor);
+		ResourceManager->UpdateDescriptorImmediately(DstHandle, View);
 		return;
 	}
 
-	if (DstHandle.GetType() == ERHIDescriptorHeapType::Sampler && SamplerManager)
+	// Bad configuration?
+	checkNoEntry();
+}
+
+void FD3D12BindlessDescriptorManager::UpdateDescriptor(FRHICommandListBase& RHICmdList, FRHIDescriptorHandle DstHandle, FD3D12View* View)
+{
+	if (ResourceManager)
 	{
-		SamplerManager->UpdateDescriptorImmediately(DstHandle, SrcDescriptor);
+		ResourceManager->UpdateDescriptor(RHICmdList, DstHandle, View);
 		return;
 	}
 
@@ -253,23 +421,45 @@ void FD3D12BindlessDescriptorManager::UpdateDescriptorImmediately(FRHIDescriptor
 	checkNoEntry();
 }
 
-void FD3D12BindlessDescriptorManager::UpdateResourceDescriptor(FRHICommandListBase& RHICmdList, FRHIDescriptorHandle DstHandle, D3D12_CPU_DESCRIPTOR_HANDLE SrcDescriptor)
+void FD3D12BindlessDescriptorManager::FinalizeContext(FD3D12CommandContext& Context)
 {
 	if (ResourceManager)
 	{
-		ResourceManager->UpdateDescriptor(RHICmdList, DstHandle, SrcDescriptor);
-		return;
+		ResourceManager->FinalizeContext(Context);
+	}
+}
+
+void FD3D12BindlessDescriptorManager::OpenCommandList(FD3D12CommandContext& Context)
+{
+	if (ResourceManager)
+	{
+		ResourceManager->OpenCommandList(Context);
 	}
 
-	// Bad configuration?
-	checkNoEntry();
+	if (SamplerManager)
+	{
+		SamplerManager->OpenCommandList(Context);
+	}
 }
 
-void FD3D12BindlessDescriptorManager::FlushPendingDescriptorUpdates(FD3D12CommandContext& Context, ERHIPipeline PipelineIndex, const FD3D12PendingResourceDescriptorUpdates& PendingDescriptorUpdates)
+void FD3D12BindlessDescriptorManager::CloseCommandList(FD3D12CommandContext& Context)
 {
 	if (ResourceManager)
 	{
-		ResourceManager->FlushPendingDescriptorUpdates(Context, PipelineIndex, PendingDescriptorUpdates);
+		ResourceManager->CloseCommandList(Context);
+	}
+
+	if (SamplerManager)
+	{
+		SamplerManager->CloseCommandList(Context);
+	}
+}
+
+void FD3D12BindlessDescriptorManager::FlushPendingDescriptorUpdates(FD3D12CommandContext& Context)
+{
+	if (ResourceManager)
+	{
+		ResourceManager->FlushPendingDescriptorUpdates(Context);
 	}
 }
 
@@ -317,13 +507,4 @@ FD3D12DescriptorHeap* FD3D12BindlessDescriptorManager::GetSamplerHeap(ERHIBindle
 	return nullptr;
 }
 
-D3D12_GPU_DESCRIPTOR_HANDLE FD3D12BindlessDescriptorManager::GetResourceGpuHandle(ERHIPipeline Pipeline, FRHIDescriptorHandle InHandle) const
-{
-	if (ResourceManager && ResourceManager->GetHeap(Pipeline))
-	{
-		return ResourceManager->GetHeap(Pipeline)->GetGPUSlotHandle(InHandle.GetIndex());
-	}
-	return D3D12_GPU_DESCRIPTOR_HANDLE{};
-}
-
 #endif // PLATFORM_SUPPORTS_BINDLESS_RENDERING
diff --git a/Engine/Source/Runtime/D3D12RHI/Private/D3D12BindlessDescriptors.h b/Engine/Source/Runtime/D3D12RHI/Private/D3D12BindlessDescriptors.h
--- a/Engine/Source/Runtime/D3D12RHI/Private/D3D12BindlessDescriptors.h
+++ b/Engine/Source/Runtime/D3D12RHI/Private/D3D12BindlessDescriptors.h
@@ -3,14 +3,19 @@
 #pragma once
 
 #include "D3D12RHICommon.h"
+#include "D3D12Descriptors.h"
 #include "RHIDefinitions.h"
 #include "RHIDescriptorAllocator.h"
+#include "RHIPipeline.h"
 #include "Templates/RefCounting.h"
 
 class FD3D12CommandContext;
+class FRHICommandListBase;
+class FD3D12SamplerState;
+class FD3D12ShaderResourceView;
+class FD3D12UnorderedAccessView;
 
-struct FD3D12DescriptorHeap;
-using FD3D12DescriptorHeapPtr = TRefCountPtr<FD3D12DescriptorHeap>;
+struct FD3D12Payload;
 
 #if PLATFORM_SUPPORTS_BINDLESS_RENDERING
 
@@ -20,6 +25,7 @@ namespace UE::D3D12BindlessDescriptors
 {
 	FD3D12DescriptorHeap* CreateCpuHeap(FD3D12Device* InDevice, ERHIDescriptorHeapType InType, uint32 InNewNumDescriptorsPerHeap);
 	FD3D12DescriptorHeap* CreateGpuHeap(FD3D12Device* InDevice, ERHIDescriptorHeapType InType, uint32 InNewNumDescriptorsPerHeap);
+	void DeferredFreeHeap(FD3D12Device* InDevice, FD3D12DescriptorHeap* InHeap);
 }
 
 /** Manager specifically for bindless sampler descriptors. */
@@ -29,43 +35,77 @@ public:
 	FD3D12BindlessSamplerManager() = delete;
 	FD3D12BindlessSamplerManager(FD3D12Device* InDevice, uint32 InNumDescriptors, TConstArrayView<TStatId> InStats);
 
-	FRHIDescriptorHandle Allocate();
+	void CleanupResources();
+
+	FRHIDescriptorHandle AllocateAndInitialize(FD3D12SamplerState* SamplerState);
 	void                 Free(FRHIDescriptorHandle InHandle);
 
-	void UpdateDescriptorImmediately(FRHIDescriptorHandle DstHandle, D3D12_CPU_DESCRIPTOR_HANDLE SrcDescriptor);
+	void OpenCommandList(FD3D12CommandContext& Context);
+	void CloseCommandList(FD3D12CommandContext& Context);
 
 	FD3D12DescriptorHeap* GetHeap() { return GpuHeap.GetReference(); }
 
 private:
-	FD3D12DescriptorHeapPtr GpuHeap;
+	FD3D12DescriptorHeapPtr      GpuHeap;
 	FRHIHeapDescriptorAllocator  Allocator;
 };
 
 #if !D3D12RHI_CUSTOM_BINDLESS_RESOURCE_MANAGER
 
-struct FD3D12PendingResourceDescriptorUpdates
+struct FD3D12PendingDescriptorUpdates
 {
-	TArray<FRHIDescriptorHandle> Handles;
-	TArray<D3D12_CPU_DESCRIPTOR_HANDLE> Descriptors;
+	// List of handles that need updating
+	TArray<FRHIDescriptorHandle>    Handles;
+	// Copies of descriptors to update each handle with. These need to be freed in a specific way.
+	TArray<FD3D12OfflineDescriptor> OfflineDescriptors;
 
-	int32 Num() const
+	~FD3D12PendingDescriptorUpdates()
 	{
-		return Handles.Num();
+		checkSlow(Handles.Num() == 0);
 	}
 
-	void Add(FRHIDescriptorHandle InHandle, const D3D12_CPU_DESCRIPTOR_HANDLE& InDescriptor)
+	void Add(FD3D12Device* Device, FRHIDescriptorHandle DestinationHandle, const FD3D12OfflineDescriptor& OfflineDescriptor)
 	{
-		if (ensure(InHandle.IsValid()))
+		if (ensure(DestinationHandle.IsValid()))
 		{
-			Handles.Emplace(InHandle);
-			Descriptors.Emplace(InDescriptor);
+			Handles.Emplace(DestinationHandle);
+			OfflineDescriptors.Emplace(OfflineDescriptor);
 		}
 	}
 
-	void Empty()
+	void Empty(FD3D12Device* Device)
 	{
+		for (FD3D12OfflineDescriptor& OfflineDescriptor : OfflineDescriptors)
+		{
+			UE::D3D12Descriptors::FreeOfflineCopy(Device, OfflineDescriptor, ERHIDescriptorHeapType::Standard);
+		}
+
 		Handles.Empty();
-		Descriptors.Empty();
+		OfflineDescriptors.Empty();
+	}
+
+	int32 Num()     const { return Handles.Num(); }
+	bool  IsEmpty() const { return Num() == 0; }
+};
+
+// Helper container for all context related bindless state.
+struct FD3D12ContextBindlessState
+{
+	FD3D12PendingDescriptorUpdates   PendingDescriptorRollbacks;
+	FD3D12DescriptorHeapPtr          CurrentGpuHeap;
+
+	// All heaps used on the context. Used for lifetime management.
+	TArray<FD3D12DescriptorHeapPtr> UsedHeaps;
+
+	FD3D12ContextBindlessState() = default;
+	~FD3D12ContextBindlessState()
+	{
+		check(PendingDescriptorRollbacks.IsEmpty());
+	}
+
+	bool HasAnyPending() const
+	{
+		return UsedHeaps.Num() > 0 || PendingDescriptorRollbacks.Num() > 0;
 	}
 };
 
@@ -76,44 +116,32 @@ public:
 	FD3D12BindlessResourceManager() = delete;
 	FD3D12BindlessResourceManager(FD3D12Device* InDevice, uint32 InNumDescriptors, TConstArrayView<TStatId> InStats);
 
+	void CleanupResources();
+
 	FRHIDescriptorHandle Allocate();
 	void                 Free(FRHIDescriptorHandle InHandle);
 
-	void UpdateDescriptorImmediately(FRHIDescriptorHandle DstHandle, D3D12_CPU_DESCRIPTOR_HANDLE SrcDescriptor);
-	void UpdateDescriptor(FRHICommandListBase& RHICmdList, FRHIDescriptorHandle DstHandle, D3D12_CPU_DESCRIPTOR_HANDLE SrcDescriptor);
+	void UpdateDescriptorImmediately(FRHIDescriptorHandle DstHandle, FD3D12View* View);
+	void UpdateDescriptor(FRHICommandListBase& RHICmdList, FRHIDescriptorHandle DstHandle, FD3D12View* View);
 
-	void FlushPendingDescriptorUpdates(FD3D12CommandContext& Context, ERHIPipeline PipelineIndex, const FD3D12PendingResourceDescriptorUpdates& PendingDescriptorUpdates);
+	void FlushPendingDescriptorUpdates(FD3D12CommandContext& Context);
 
-	bool IsEnabledForPipeline(ERHIPipeline Pipeline) const
-	{
-		return EnumHasAnyFlags(ConfiguredPipelines, Pipeline);
-	}
+	void OpenCommandList(FD3D12CommandContext& Context);
+	void CloseCommandList(FD3D12CommandContext& Context);
+	void FinalizeContext(FD3D12CommandContext& Context);
 
 	FD3D12DescriptorHeap* GetHeap(ERHIPipeline Pipeline)
 	{
-		if (IsEnabledForPipeline(Pipeline))
-		{
-			return Pipelines[Pipeline].GpuHeap.GetReference();
-		}
+		checkNoEntry();
 		return nullptr;
 	}
 
 private:
-	template <typename TFunctionType>
-	void EnumeratePipelines(TFunctionType&& Function)
-	{
-		EnumerateRHIPipelines(ConfiguredPipelines, Forward<TFunctionType>(Function));
-	}
-
-	struct FPipeline
-	{
-		FD3D12DescriptorHeapPtr CpuHeap;
-		FD3D12DescriptorHeapPtr GpuHeap;
-	};
-
-	ERHIPipeline                 ConfiguredPipelines;
-	TRHIPipelineArray<FPipeline> Pipelines;
+	void CopyCpuHeap(FD3D12DescriptorHeap* DestinationHeap);
+	void CreateHeapOnState(FD3D12ContextBindlessState& State);
+	void FinalizeHeapOnState(FD3D12ContextBindlessState& State);
 
+	FD3D12DescriptorHeapPtr      CpuHeap;
 	FRHIHeapDescriptorAllocator  Allocator;
 };
 
@@ -127,6 +155,7 @@ public:
 	~FD3D12BindlessDescriptorManager();
 
 	void Init();
+	void CleanupResources();
 
 	ERHIBindlessConfiguration GetResourcesConfiguration() const { return ResourcesConfiguration; }
 	ERHIBindlessConfiguration GetSamplersConfiguration()  const { return SamplersConfiguration; }
@@ -140,14 +169,20 @@ public:
 	bool AreResourcesFullyBindless() const { return GetResourcesConfiguration() == ERHIBindlessConfiguration::AllShaders; }
 	bool AreSamplersFullyBindless () const { return GetSamplersConfiguration()  == ERHIBindlessConfiguration::AllShaders; }
 
-	FRHIDescriptorHandle Allocate(ERHIDescriptorHeapType InType);
+	FRHIDescriptorHandle AllocateResourceHandle();
+	FRHIDescriptorHandle AllocateAndInitialize(FD3D12SamplerState* SamplerState);
 	void                 ImmediateFree(FRHIDescriptorHandle InHandle);
 	void                 DeferredFreeFromDestructor(FRHIDescriptorHandle InHandle);
 
-	void UpdateDescriptorImmediately(FRHIDescriptorHandle DstHandle, D3D12_CPU_DESCRIPTOR_HANDLE SrcDescriptor);
-	void UpdateResourceDescriptor(FRHICommandListBase& RHICmdList, FRHIDescriptorHandle DstHandle, D3D12_CPU_DESCRIPTOR_HANDLE SrcDescriptor);
+	void UpdateDescriptorImmediately(FRHIDescriptorHandle DstHandle, FD3D12View* View);
+	void UpdateDescriptor(FRHICommandListBase& RHICmdList, FRHIDescriptorHandle DstHandle, FD3D12View* SourceView);
 
-	void FlushPendingDescriptorUpdates(FD3D12CommandContext& Context, ERHIPipeline PipelineIndex, const FD3D12PendingResourceDescriptorUpdates& PendingDescriptorUpdates);
+	void FinalizeContext(FD3D12CommandContext& Context);
+
+	void OpenCommandList(FD3D12CommandContext& Context);
+	void CloseCommandList(FD3D12CommandContext& Context);
+
+	void FlushPendingDescriptorUpdates(FD3D12CommandContext& Context);
 
 #if D3D12RHI_USE_CONSTANT_BUFFER_VIEWS
 	TRHIPipelineArray<FD3D12DescriptorHeapPtr> AllocateResourceHeapsForAllPipelines(int32 InSize);
@@ -159,8 +194,6 @@ public:
 	FD3D12DescriptorHeap* GetResourceHeap(ERHIPipeline Pipeline, ERHIBindlessConfiguration InConfiguration);
 	FD3D12DescriptorHeap* GetSamplerHeap(ERHIBindlessConfiguration InConfiguration);
 
-	D3D12_GPU_DESCRIPTOR_HANDLE GetResourceGpuHandle(ERHIPipeline Pipeline, FRHIDescriptorHandle InHandle) const;
-
 private:
 	TUniquePtr<FD3D12BindlessResourceManager> ResourceManager;
 	TUniquePtr<FD3D12BindlessSamplerManager>  SamplerManager;
diff --git a/Engine/Source/Runtime/D3D12RHI/Private/D3D12CommandContext.cpp b/Engine/Source/Runtime/D3D12RHI/Private/D3D12CommandContext.cpp
--- a/Engine/Source/Runtime/D3D12RHI/Private/D3D12CommandContext.cpp
+++ b/Engine/Source/Runtime/D3D12RHI/Private/D3D12CommandContext.cpp
@@ -398,7 +398,9 @@ void FD3D12ContextCommon::Finalize(TArray<FD3D12Payload*>& OutPayloads)
 
 void FD3D12CommandContext::Finalize(TArray<FD3D12Payload*>& OutPayloads)
 {
-	FlushPendingDescriptorUpdates();
+#if PLATFORM_SUPPORTS_BINDLESS_RENDERING
+	GetParentDevice()->GetBindlessDescriptorManager().FinalizeContext(*this);
+#endif
 
 	FD3D12ContextCommon::Finalize(OutPayloads);
 }
diff --git a/Engine/Source/Runtime/D3D12RHI/Private/D3D12CommandContext.h b/Engine/Source/Runtime/D3D12RHI/Private/D3D12CommandContext.h
--- a/Engine/Source/Runtime/D3D12RHI/Private/D3D12CommandContext.h
+++ b/Engine/Source/Runtime/D3D12RHI/Private/D3D12CommandContext.h
@@ -742,11 +742,14 @@ public:
 		return StaticUniformBuffers;
 	}
 
-	void AddPendingDescriptorUpdate(FRHIDescriptorHandle InHandle, D3D12_CPU_DESCRIPTOR_HANDLE InDescriptor);
 	void FlushPendingDescriptorUpdates();
 
 	virtual void Finalize(TArray<FD3D12Payload*>& OutPayloads) override;
 
+#if PLATFORM_SUPPORTS_BINDLESS_RENDERING
+	FD3D12ContextBindlessState& GetBindlessState() { return BindlessState; }
+#endif
+
 protected:
 
 	FD3D12CommandContext* GetContext(uint32 InGPUIndex) final override 
@@ -772,7 +775,7 @@ private:
 	TArray<FRHIUniformBuffer*> StaticUniformBuffers;
 
 #if PLATFORM_SUPPORTS_BINDLESS_RENDERING
-	FD3D12PendingResourceDescriptorUpdates PendingDescriptorUpdates;
+	FD3D12ContextBindlessState BindlessState;
 #endif
 };
 
diff --git a/Engine/Source/Runtime/D3D12RHI/Private/D3D12Commands.cpp b/Engine/Source/Runtime/D3D12RHI/Private/D3D12Commands.cpp
--- a/Engine/Source/Runtime/D3D12RHI/Private/D3D12Commands.cpp
+++ b/Engine/Source/Runtime/D3D12RHI/Private/D3D12Commands.cpp
@@ -78,21 +78,10 @@ static void BindUniformBuffer(FD3D12CommandContext& Context, FRHIShader* Shader,
 	Context.DirtyUniformBuffers[ShaderFrequency] |= (1 << BufferIndex);
 }
 
-void FD3D12CommandContext::AddPendingDescriptorUpdate(FRHIDescriptorHandle InHandle, D3D12_CPU_DESCRIPTOR_HANDLE InDescriptor)
-{
-#if PLATFORM_SUPPORTS_BINDLESS_RENDERING
-	PendingDescriptorUpdates.Add(InHandle, InDescriptor);
-#endif
-}
-
 void FD3D12CommandContext::FlushPendingDescriptorUpdates()
 {
 #if PLATFORM_SUPPORTS_BINDLESS_RENDERING
-	if (PendingDescriptorUpdates.Num())
-	{
-		GetParentDevice()->GetBindlessDescriptorManager().FlushPendingDescriptorUpdates(*this, GetPipeline(), PendingDescriptorUpdates);
-		PendingDescriptorUpdates.Empty();
-	}
+	GetParentDevice()->GetBindlessDescriptorManager().FlushPendingDescriptorUpdates(*this);
 #endif
 }
 
@@ -1078,6 +1067,7 @@ static void SetShaderParametersOnContext(
 				break;
 			}
 
+			checkf(Handle.IsValid(), TEXT("D3D12 resource did not provide a valid descriptor handle. Please validate that all D3D12 types can provide this or that the resource is still valid."));
 			Binder.SetBindlessHandle(Handle, Parameter.Index);
 		}
 	}
diff --git a/Engine/Source/Runtime/D3D12RHI/Private/D3D12DescriptorCache.cpp b/Engine/Source/Runtime/D3D12RHI/Private/D3D12DescriptorCache.cpp
--- a/Engine/Source/Runtime/D3D12RHI/Private/D3D12DescriptorCache.cpp
+++ b/Engine/Source/Runtime/D3D12RHI/Private/D3D12DescriptorCache.cpp
@@ -69,15 +69,14 @@ bool FD3D12DescriptorCache::SetDescriptorHeaps(bool bForceHeapChanged)
 	// See if the descriptor heaps changed.
 	bool bHeapChanged = bForceHeapChanged;
 
-#if PLATFORM_SUPPORTS_BINDLESS_RENDERING && DO_ENSURE
-	FD3D12BindlessDescriptorManager& BindlessDescriptorManager = GetParentDevice()->GetBindlessDescriptorManager();
+#if PLATFORM_SUPPORTS_BINDLESS_RENDERING && DO_CHECK
 	if (IsUsingBindlessResources())
 	{
-		check(BindlessResourcesHeap == BindlessDescriptorManager.GetResourceHeap(Pipeline));
+		checkf(BindlessResourcesHeap, TEXT("Bindless resource heap was not set in OpenCommandList!"));
 	}
 	if (IsUsingBindlessSamplers())
 	{
-		check(BindlessSamplersHeap == BindlessDescriptorManager.GetSamplerHeap());
+		checkf(BindlessSamplersHeap, TEXT("Bindless sampler heap was not set in OpenCommandList!"));
 	}
 #endif
 
@@ -142,25 +141,19 @@ void FD3D12DescriptorCache::OpenCommandList()
 	LastSetViewHeap = nullptr;
 	LastSetSamplerHeap = nullptr;
 
-
 #if PLATFORM_SUPPORTS_BINDLESS_RENDERING
-	if (IsUsingBindlessSamplers())
+	if (IsUsingBindlessSamplers() || IsUsingBindlessResources())
 	{
-		BindlessSamplersHeap = GetParentDevice()->GetBindlessDescriptorManager().GetSamplerHeap();
+		GetParentDevice()->GetBindlessDescriptorManager().OpenCommandList(Context);
 	}
-	else
+
+	if (!IsUsingBindlessSamplers())
 #endif
 	{
 		// The global sampler heap doesn't care about the current command list
 		LocalSamplerHeap.OpenCommandList();
 	}
 
-#if PLATFORM_SUPPORTS_BINDLESS_RENDERING
-	if (IsUsingBindlessResources())
-	{
-		BindlessResourcesHeap = GetParentDevice()->GetBindlessDescriptorManager().GetResourceHeap(Context.GetPipeline());
-	}
-#endif
 	if (CurrentViewHeap)
 	{
 		CurrentViewHeap->OpenCommandList();
@@ -184,7 +177,14 @@ void FD3D12DescriptorCache::CloseCommandList()
 		CurrentViewHeap->CloseCommandList();
 	}
 
+#if PLATFORM_SUPPORTS_BINDLESS_RENDERING
+	if (IsUsingBindlessSamplers() || IsUsingBindlessResources())
+	{
+		GetParentDevice()->GetBindlessDescriptorManager().CloseCommandList(Context);
+	}
+
 	if (!IsUsingBindlessSamplers())
+#endif
 	{
 		LocalSamplerHeap.CloseCommandList();
 
@@ -731,6 +731,25 @@ void FD3D12DescriptorCache::SwitchToGlobalSamplerHeap()
 	CurrentSamplerHeap = &GlobalSamplerHeap;
 }
 
+#if PLATFORM_SUPPORTS_BINDLESS_RENDERING
+void FD3D12DescriptorCache::SwitchToNewBindlessResourceHeap(FD3D12DescriptorHeap* InHeap)
+{
+	if (ensure(IsUsingBindlessResources()))
+	{
+		BindlessResourcesHeap = InHeap;
+
+		// TODO: should we be forced open before here?
+		if (IsUsingBindlessSamplers())
+		{
+			check(BindlessSamplersHeap != nullptr);
+		}
+
+		// If we didn't change heaps, then the caller sent us the wrong heap.
+		ensure(SetDescriptorHeaps());
+	}
+}
+#endif
+
 void FD3D12DescriptorCache::OverrideLastSetHeaps(ID3D12DescriptorHeap* ViewHeap, ID3D12DescriptorHeap* SamplerHeap)
 {
 	ID3D12DescriptorHeap* ViewHeapToSet = ViewHeap ? ViewHeap : LastSetViewHeap;
@@ -912,7 +931,7 @@ void FD3D12GlobalOnlineSamplerHeap::Init(uint32 TotalSize)
 		TEXT("Device Global - Online Sampler Heap"),
 		ERHIDescriptorHeapType::Sampler,
 		TotalSize,
-		D3D12_DESCRIPTOR_HEAP_FLAG_SHADER_VISIBLE
+		ED3D12DescriptorHeapFlags::GpuVisible
 	);
 
 	INC_DWORD_STAT(STAT_NumSamplerOnlineDescriptorHeaps);
@@ -1075,7 +1094,7 @@ void FD3D12LocalOnlineHeap::Init(uint32 InNumDescriptors, ERHIDescriptorHeapType
 			DebugName,
 			InHeapType,
 			InNumDescriptors,
-			D3D12_DESCRIPTOR_HEAP_FLAG_SHADER_VISIBLE
+			ED3D12DescriptorHeapFlags::GpuVisible
 		);
 
 		Entry.Heap = Heap;
@@ -1129,7 +1148,7 @@ bool FD3D12LocalOnlineHeap::RollOver()
 			DebugName,
 			HeapType,
 			NumDescriptors,
-			D3D12_DESCRIPTOR_HEAP_FLAG_SHADER_VISIBLE
+			ED3D12DescriptorHeapFlags::GpuVisible
 		);
 
 		if (HeapType == ERHIDescriptorHeapType::Standard)
diff --git a/Engine/Source/Runtime/D3D12RHI/Private/D3D12DescriptorCache.h b/Engine/Source/Runtime/D3D12RHI/Private/D3D12DescriptorCache.h
--- a/Engine/Source/Runtime/D3D12RHI/Private/D3D12DescriptorCache.h
+++ b/Engine/Source/Runtime/D3D12RHI/Private/D3D12DescriptorCache.h
@@ -364,6 +364,9 @@ public:
 	bool SwitchToContextLocalViewHeap();
 	bool SwitchToContextLocalSamplerHeap();
 	void SwitchToGlobalSamplerHeap();
+#if PLATFORM_SUPPORTS_BINDLESS_RENDERING
+	void SwitchToNewBindlessResourceHeap(FD3D12DescriptorHeap* InHeap);
+#endif
 
 	void OverrideLastSetHeaps(ID3D12DescriptorHeap* ViewHeap, ID3D12DescriptorHeap* SamplerHeap);
 	void RestoreAfterExternalHeapsSet();
@@ -375,6 +378,28 @@ public:
 	// Returns true if one of the heaps actually changed, false otherwise.
 	bool SetDescriptorHeaps(bool bForceHeapChanged = false);
 
+#if PLATFORM_SUPPORTS_BINDLESS_RENDERING
+	void SetBindlessResourcesHeapDirectly(FD3D12DescriptorHeap* InHeap)
+	{
+		BindlessResourcesHeap = InHeap;
+	}
+
+	void SetBindlessSamplersHeapDirectly(FD3D12DescriptorHeap* InHeap)
+	{
+		BindlessSamplersHeap = InHeap;
+	}
+
+	FD3D12DescriptorHeap* GetBindlessResourcesHeap() const
+	{
+		return BindlessResourcesHeap;
+	}
+
+	FD3D12DescriptorHeap* GetBindlessSamplersHeap() const
+	{
+		return BindlessSamplersHeap;
+	}
+#endif
+
 protected:
 	FD3D12CommandContext& Context;
 	const FD3D12DefaultViews& DefaultViews;
@@ -405,7 +430,7 @@ private:
 	bool bBindlessResources = false;
 	bool bBindlessSamplers = false;
 
-	FD3D12DescriptorHeap* BindlessResourcesHeap = nullptr;
-	FD3D12DescriptorHeap* BindlessSamplersHeap = nullptr;
+	FD3D12DescriptorHeapPtr BindlessResourcesHeap = nullptr;
+	FD3D12DescriptorHeapPtr BindlessSamplersHeap = nullptr;
 #endif
 };
\ No newline at end of file
diff --git a/Engine/Source/Runtime/D3D12RHI/Private/D3D12Descriptors.cpp b/Engine/Source/Runtime/D3D12RHI/Private/D3D12Descriptors.cpp
--- a/Engine/Source/Runtime/D3D12RHI/Private/D3D12Descriptors.cpp
+++ b/Engine/Source/Runtime/D3D12RHI/Private/D3D12Descriptors.cpp
@@ -5,21 +5,23 @@
 
 ////////////////////////////////////////////////////////////////////////////////////////////////////////////////////////
 
-void UE::D3D12Descriptors::CopyDescriptor(FD3D12Device* Device, FD3D12DescriptorHeap* TargetHeap, FRHIDescriptorHandle Handle, D3D12_CPU_DESCRIPTOR_HANDLE SourceCpuHandle)
+void UE::D3D12Descriptors::CopyDescriptor(FD3D12Device* Device, FD3D12DescriptorHeap* TargetHeap, FRHIDescriptorHandle DstHandle, D3D12_CPU_DESCRIPTOR_HANDLE SrcCpuHandle)
 {
-	if (Handle.IsValid())
+	if (DstHandle.IsValid())
 	{
-		const D3D12_CPU_DESCRIPTOR_HANDLE DestCpuHandle = TargetHeap->GetCPUSlotHandle(Handle.GetIndex());
-		const D3D12_DESCRIPTOR_HEAP_TYPE D3DHeapType = Translate(Handle.GetType());
+		const D3D12_CPU_DESCRIPTOR_HANDLE DestCpuHandle = TargetHeap->GetCPUSlotHandle(DstHandle.GetIndex());
+		const D3D12_DESCRIPTOR_HEAP_TYPE D3DHeapType = Translate(DstHandle.GetType());
 
-		Device->GetDevice()->CopyDescriptorsSimple(1, DestCpuHandle, SourceCpuHandle, D3DHeapType);
+		Device->GetDevice()->CopyDescriptorsSimple(1, DestCpuHandle, SrcCpuHandle, D3DHeapType);
 	}
 }
 
-void UE::D3D12Descriptors::CopyDescriptors(FD3D12Device* Device, FD3D12DescriptorHeap* TargetHeap, FD3D12DescriptorHeap* SourceHeap, uint32 NumDescriptors)
+void UE::D3D12Descriptors::CopyDescriptors(FD3D12Device* Device, FD3D12DescriptorHeap* TargetHeap, FD3D12DescriptorHeap* SourceHeap, uint32 FirstDescriptor, uint32 NumDescriptors)
 {
-	const D3D12_CPU_DESCRIPTOR_HANDLE TargetStart = TargetHeap->GetCPUSlotHandle(0);
-	const D3D12_CPU_DESCRIPTOR_HANDLE SourceStart = SourceHeap->GetCPUSlotHandle(0);
+	SCOPED_NAMED_EVENT_F(TEXT("CopyDescriptors HeapToHeap (%d)"), FColor::Turquoise, NumDescriptors);
+
+	const D3D12_CPU_DESCRIPTOR_HANDLE TargetStart = TargetHeap->GetCPUSlotHandle(FirstDescriptor);
+	const D3D12_CPU_DESCRIPTOR_HANDLE SourceStart = SourceHeap->GetCPUSlotHandle(FirstDescriptor);
 	const D3D12_DESCRIPTOR_HEAP_TYPE D3DHeapType = Translate(TargetHeap->GetType());
 
 	Device->GetDevice()->CopyDescriptorsSimple(
@@ -30,14 +32,67 @@ void UE::D3D12Descriptors::CopyDescriptors(FD3D12Device* Device, FD3D12Descripto
 	);
 }
 
+void UE::D3D12Descriptors::CopyDescriptors(FD3D12Device* Device, FD3D12DescriptorHeap* TargetHeap, TConstArrayView<FRHIDescriptorHandle> DstHandles, TConstArrayView<FD3D12OfflineDescriptor> SrcOfflineDescriptors)
+{
+	SCOPED_NAMED_EVENT_F(TEXT("CopyDescriptors Scattered (%d)"), FColor::Turquoise, DstHandles.Num());
+
+	check(DstHandles.Num() == SrcOfflineDescriptors.Num());
+
+	const int32 NumDescriptors = DstHandles.Num();
+	const D3D12_DESCRIPTOR_HEAP_TYPE D3DHeapType = Translate(TargetHeap->GetType());
+
+	TArray<D3D12_CPU_DESCRIPTOR_HANDLE> DstDescriptors;
+	TArray<D3D12_CPU_DESCRIPTOR_HANDLE> SrcDescriptors;
+	TArray<UINT> CopyCounts;
+
+	DstDescriptors.SetNumUninitialized(NumDescriptors);
+	SrcDescriptors.SetNumUninitialized(NumDescriptors);
+	CopyCounts.SetNumUninitialized(NumDescriptors);
+
+	for (int32 Index = 0; Index < NumDescriptors; Index++)
+	{
+		DstDescriptors[Index] = TargetHeap->GetCPUSlotHandle(DstHandles[Index].GetIndex());
+		SrcDescriptors[Index] = SrcOfflineDescriptors[Index];
+		CopyCounts[Index] = 1;
+	}
+
+	Device->GetDevice()->CopyDescriptors(
+		NumDescriptors, DstDescriptors.GetData(), CopyCounts.GetData(),
+		NumDescriptors, SrcDescriptors.GetData(), CopyCounts.GetData(),
+		D3DHeapType
+	);
+}
+
+FD3D12OfflineDescriptor UE::D3D12Descriptors::CreateOfflineCopy(FD3D12Device* Device, D3D12_CPU_DESCRIPTOR_HANDLE InDescriptor, ERHIDescriptorHeapType InType)
+{
+	FD3D12OfflineDescriptor NewDescriptor = Device->GetOfflineDescriptorManager(InType).AllocateHeapSlot();
+
+	Device->GetDevice()->CopyDescriptorsSimple(1, NewDescriptor, InDescriptor, Translate(InType));
+
+	NewDescriptor.IncrementVersion();
+
+	return NewDescriptor;
+}
+
+FD3D12OfflineDescriptor UE::D3D12Descriptors::CreateOfflineCopy(FD3D12Device* Device, FD3D12DescriptorHeap* InHeap, FRHIDescriptorHandle InHandle)
+{
+	D3D12_CPU_DESCRIPTOR_HANDLE SourceDescriptor = InHeap->GetCPUSlotHandle(InHandle.GetIndex());
+	return UE::D3D12Descriptors::CreateOfflineCopy(Device, SourceDescriptor, InHeap->GetType());
+}
+
+void UE::D3D12Descriptors::FreeOfflineCopy(FD3D12Device* Device, FD3D12OfflineDescriptor& InDescriptor, ERHIDescriptorHeapType InType)
+{
+	Device->GetOfflineDescriptorManager(InType).FreeHeapSlot(InDescriptor);
+}
+
 ////////////////////////////////////////////////////////////////////////////////////////////////////////////////////////
 // FD3D12DescriptorHeap
 
-FD3D12DescriptorHeap::FD3D12DescriptorHeap(FD3D12Device* InDevice, ID3D12DescriptorHeap* InHeap, uint32 InNumDescriptors, ERHIDescriptorHeapType InType, D3D12_DESCRIPTOR_HEAP_FLAGS InFlags, bool bInIsGlobal)
+FD3D12DescriptorHeap::FD3D12DescriptorHeap(FD3D12Device* InDevice, TRefCountPtr<ID3D12DescriptorHeap>&& InHeap, uint32 InNumDescriptors, ERHIDescriptorHeapType InType, ED3D12DescriptorHeapFlags InFlags, bool bInIsGlobal)
 	: FD3D12DeviceChild(InDevice)
-	, Heap(InHeap)
-	, CpuBase(InHeap->GetCPUDescriptorHandleForHeapStart())
-	, GpuBase((InFlags& D3D12_DESCRIPTOR_HEAP_FLAG_SHADER_VISIBLE) ? InHeap->GetGPUDescriptorHandleForHeapStart() : D3D12_GPU_DESCRIPTOR_HANDLE{})
+	, Heap(Forward<TRefCountPtr<ID3D12DescriptorHeap>>(InHeap))
+	, CpuBase(Heap->GetCPUDescriptorHandleForHeapStart())
+	, GpuBase(EnumHasAnyFlags(InFlags, ED3D12DescriptorHeapFlags::GpuVisible) ? Heap->GetGPUDescriptorHandleForHeapStart() : D3D12_GPU_DESCRIPTOR_HANDLE{})
 	, Offset(0)
 	, NumDescriptors(InNumDescriptors)
 	, DescriptorSize(InDevice->GetDevice()->GetDescriptorHandleIncrementSize(Translate(InType)))
@@ -76,6 +131,11 @@ FD3D12DescriptorHeap::FD3D12DescriptorHeap(FD3D12DescriptorHeap* SubAllocateSour
 
 FD3D12DescriptorHeap::~FD3D12DescriptorHeap()
 {
+	if (CanBePooled())
+	{
+		GetParentDevice()->GetDescriptorHeapManager().AddHeapToPool(MoveTemp(Heap), GetType(), GetNumDescriptors(), GetFlags());
+	}
+
 #if STATS
 	if (!bIsSuballocation)
 	{
@@ -147,7 +207,7 @@ void FD3D12OnlineDescriptorManager::Init(uint32 InTotalSize, uint32 InBlockSize,
 			TEXT("Device Global - Online View Heap"),
 			ERHIDescriptorHeapType::Standard,
 			InTotalSize,
-			D3D12_DESCRIPTOR_HEAP_FLAG_SHADER_VISIBLE);
+			ED3D12DescriptorHeapFlags::GpuVisible);
 
 		INC_DWORD_STAT(STAT_NumViewOnlineDescriptorHeaps);
 		INC_MEMORY_STAT_BY(STAT_ViewOnlineDescriptorHeapMemory, Heap->GetMemorySize());
@@ -175,6 +235,14 @@ void FD3D12OnlineDescriptorManager::Init(uint32 InTotalSize, uint32 InBlockSize,
 	}
 }
 
+void FD3D12OnlineDescriptorManager::CleanupResources()
+{
+	for (ERHIPipeline Pipeline : GetRHIPipelines())
+	{
+		Heaps[Pipeline] = nullptr;
+	}
+}
+
 // Allocate a new heap block
 FD3D12OnlineDescriptorBlock* FD3D12OnlineDescriptorManager::AllocateHeapBlock()
 {
@@ -258,7 +326,7 @@ void FD3D12OfflineDescriptorManager::AllocateHeap()
 		TEXT("FD3D12OfflineDescriptorManager"),
 		HeapType,
 		NumDescriptorsPerHeap,
-		D3D12_DESCRIPTOR_HEAP_FLAG_NONE
+		ED3D12DescriptorHeapFlags::None
 	);
 
 	D3D12_CPU_DESCRIPTOR_HANDLE HeapBase = Heap->GetCPUSlotHandle(0);
@@ -352,6 +420,14 @@ void FD3D12OfflineDescriptorManager::FreeHeapSlot(FD3D12OfflineDescriptor& Descr
 	Descriptor = {};
 }
 
+void FD3D12OfflineDescriptorManager::CleanupResources()
+{
+	for (FD3D12OfflineHeapEntry& HeapEntry : Heaps)
+	{
+		HeapEntry.Heap = nullptr;
+	}
+}
+
 ////////////////////////////////////////////////////////////////////////////////////////////////////////////////////////
 // FD3D12DescriptorHeapManager
 
@@ -365,12 +441,14 @@ FD3D12DescriptorHeapManager::~FD3D12DescriptorHeapManager()
 	Destroy();
 }
 
-static FD3D12DescriptorHeap* CreateDescriptorHeap(FD3D12Device* Device, const TCHAR* DebugName, ERHIDescriptorHeapType HeapType, uint32 NumDescriptors, D3D12_DESCRIPTOR_HEAP_FLAGS Flags, bool bIsGlobal)
+static FD3D12DescriptorHeap* CreateDescriptorHeap(FD3D12Device* Device, const TCHAR* DebugName, ERHIDescriptorHeapType HeapType, uint32 NumDescriptors, ED3D12DescriptorHeapFlags Flags, bool bIsGlobal)
 {
+	SCOPED_NAMED_EVENT_F(TEXT("D3D12Device::CreateDescriptorHeap(%s, %d)"), FColor::Turquoise, DebugName, NumDescriptors);
+
 	D3D12_DESCRIPTOR_HEAP_DESC Desc{};
 	Desc.Type = Translate(HeapType);
 	Desc.NumDescriptors = NumDescriptors;
-	Desc.Flags = Flags;
+	Desc.Flags = Translate(Flags);
 	Desc.NodeMask = Device->GetGPUMask().GetNative();
 
 	TRefCountPtr<ID3D12DescriptorHeap> Heap;
@@ -378,7 +456,7 @@ static FD3D12DescriptorHeap* CreateDescriptorHeap(FD3D12Device* Device, const TC
 
 	SetName(Heap, DebugName);
 
-	return new FD3D12DescriptorHeap(Device, Heap, NumDescriptors, HeapType, Flags, bIsGlobal);
+	return new FD3D12DescriptorHeap(Device, MoveTemp(Heap), NumDescriptors, HeapType, Flags, bIsGlobal);
 }
 
 void FD3D12DescriptorHeapManager::Init(uint32 InNumGlobalResourceDescriptors, uint32 InNumGlobalSamplerDescriptors)
@@ -390,7 +468,7 @@ void FD3D12DescriptorHeapManager::Init(uint32 InNumGlobalResourceDescriptors, ui
 			TEXT("GlobalResourceHeap"),
 			ERHIDescriptorHeapType::Standard,
 			InNumGlobalResourceDescriptors,
-			D3D12_DESCRIPTOR_HEAP_FLAG_SHADER_VISIBLE,
+			ED3D12DescriptorHeapFlags::GpuVisible,
 			true /* bIsGlobal */
 		);
 
@@ -405,7 +483,7 @@ void FD3D12DescriptorHeapManager::Init(uint32 InNumGlobalResourceDescriptors, ui
 			TEXT("GlobalSamplerHeap"),
 			ERHIDescriptorHeapType::Sampler,
 			InNumGlobalSamplerDescriptors,
-			D3D12_DESCRIPTOR_HEAP_FLAG_SHADER_VISIBLE,
+			ED3D12DescriptorHeapFlags::GpuVisible,
 			true /* bIsGlobal */
 		);
 
@@ -418,12 +496,42 @@ void FD3D12DescriptorHeapManager::Destroy()
 {
 }
 
-FD3D12DescriptorHeap* FD3D12DescriptorHeapManager::AllocateIndependentHeap(const TCHAR* InDebugName, ERHIDescriptorHeapType InHeapType, uint32 InNumDescriptors, D3D12_DESCRIPTOR_HEAP_FLAGS InHeapFlags)
+TRefCountPtr<ID3D12DescriptorHeap> FD3D12DescriptorHeapManager::AcquirePooledHeap(ERHIDescriptorHeapType InType, uint32 InNumDescriptors, ED3D12DescriptorHeapFlags InFlags)
 {
+	FScopeLock Lock(&PooledHeapsCS);
+	const int32 Index = PooledHeaps.IndexOfByPredicate([InType, InNumDescriptors, InFlags](const FPooledHeap& PooledHeap) -> bool
+		{
+			return PooledHeap.Type == InType
+				&& PooledHeap.NumDescriptors == InNumDescriptors
+				&& PooledHeap.Flags == InFlags;
+		});
+
+	if (Index != INDEX_NONE)
+	{
+		TRefCountPtr<ID3D12DescriptorHeap> Result;
+		PooledHeaps[Index].Heap.Swap(Result);
+		PooledHeaps.RemoveAt(Index);
+		return Result;
+	}
+
+	return nullptr;
+}
+
+FD3D12DescriptorHeap* FD3D12DescriptorHeapManager::AllocateIndependentHeap(const TCHAR* InDebugName, ERHIDescriptorHeapType InHeapType, uint32 InNumDescriptors, ED3D12DescriptorHeapFlags InHeapFlags)
+{
+	if (EnumHasAnyFlags(InHeapFlags, ED3D12DescriptorHeapFlags::Poolable))
+	{
+		TRefCountPtr<ID3D12DescriptorHeap> PooledHeap = AcquirePooledHeap(InHeapType, InNumDescriptors, InHeapFlags);
+		if (PooledHeap.GetReference())
+		{
+			return new FD3D12DescriptorHeap(GetParentDevice(), MoveTemp(PooledHeap), InNumDescriptors, InHeapType, InHeapFlags, false);
+		}
+	}
+
 	return CreateDescriptorHeap(GetParentDevice(), InDebugName, InHeapType, InNumDescriptors, InHeapFlags, false);
 }
 
-FD3D12DescriptorHeap* FD3D12DescriptorHeapManager::AllocateHeap(const TCHAR* InDebugName, ERHIDescriptorHeapType InHeapType, uint32 InNumDescriptors, D3D12_DESCRIPTOR_HEAP_FLAGS InHeapFlags)
+FD3D12DescriptorHeap* FD3D12DescriptorHeapManager::AllocateHeap(const TCHAR* InDebugName, ERHIDescriptorHeapType InHeapType, uint32 InNumDescriptors, ED3D12DescriptorHeapFlags InHeapFlags)
 {
 	for (FD3D12DescriptorManager& GlobalHeap : GlobalHeaps)
 	{
@@ -483,3 +591,9 @@ void FD3D12DescriptorHeapManager::ImmediateFreeHeap(FD3D12DescriptorHeap* InHeap
 		InHeap->Release();
 	}
 }
+
+void FD3D12DescriptorHeapManager::AddHeapToPool(TRefCountPtr<ID3D12DescriptorHeap>&& InHeap, ERHIDescriptorHeapType InType, uint32 InNumDescriptors, ED3D12DescriptorHeapFlags InFlags)
+{
+	FScopeLock Lock(&PooledHeapsCS);
+	PooledHeaps.Emplace(Forward<TRefCountPtr<ID3D12DescriptorHeap>>(InHeap), InType, InNumDescriptors, InFlags);
+}
diff --git a/Engine/Source/Runtime/D3D12RHI/Private/D3D12Descriptors.h b/Engine/Source/Runtime/D3D12RHI/Private/D3D12Descriptors.h
--- a/Engine/Source/Runtime/D3D12RHI/Private/D3D12Descriptors.h
+++ b/Engine/Source/Runtime/D3D12RHI/Private/D3D12Descriptors.h
@@ -3,14 +3,18 @@
 #pragma once
 
 #include "RHIDescriptorAllocator.h"
-#include "MultiGPU.h"
+#include "Containers/List.h"
+#include "Containers/Queue.h"
 #include "D3D12RHI.h"
 #include "D3D12RHICommon.h"
+#include "MultiGPU.h"
+#include "RHIPipeline.h"
 #include "Templates/RefCounting.h"
-#include "Containers/Queue.h"
-#include "Containers/List.h"
+
+class FD3D12View;
 
 struct FD3D12DescriptorHeap;
+struct FD3D12OfflineDescriptor;
 
 inline D3D12_DESCRIPTOR_HEAP_TYPE Translate(ERHIDescriptorHeapType InHeapType)
 {
@@ -24,11 +28,34 @@ inline D3D12_DESCRIPTOR_HEAP_TYPE Translate(ERHIDescriptorHeapType InHeapType)
 	}
 }
 
+enum class ED3D12DescriptorHeapFlags : uint8
+{
+	None       = 0,
+	GpuVisible = 1 << 0,
+	Poolable   = 1 << 1,
+};
+ENUM_CLASS_FLAGS(ED3D12DescriptorHeapFlags)
+
+inline D3D12_DESCRIPTOR_HEAP_FLAGS Translate(ED3D12DescriptorHeapFlags InHeapFlags)
+{
+	if (EnumHasAnyFlags(InHeapFlags, ED3D12DescriptorHeapFlags::GpuVisible))
+	{
+		return D3D12_DESCRIPTOR_HEAP_FLAG_SHADER_VISIBLE;
+	}
+
+	return D3D12_DESCRIPTOR_HEAP_FLAG_NONE;
+}
+
 namespace UE::D3D12Descriptors
 {
-	FD3D12DescriptorHeap* CreateDescriptorHeap(FD3D12Device* Device, const TCHAR* DebugName, ERHIDescriptorHeapType HeapType, uint32 NumDescriptors, D3D12_DESCRIPTOR_HEAP_FLAGS Flags, bool bGlobal = false);
-	void CopyDescriptor(FD3D12Device* Device, FD3D12DescriptorHeap* TargetHeap, FRHIDescriptorHandle Handle, D3D12_CPU_DESCRIPTOR_HANDLE SourceCpuHandle);
-	void CopyDescriptors(FD3D12Device* Device, FD3D12DescriptorHeap* TargetHeap, FD3D12DescriptorHeap* SourceHeap, uint32 NumDescriptors);
+	FD3D12DescriptorHeap* CreateDescriptorHeap(FD3D12Device* Device, const TCHAR* DebugName, ERHIDescriptorHeapType HeapType, uint32 NumDescriptors, ED3D12DescriptorHeapFlags Flags, bool bGlobal = false);
+	void CopyDescriptor(FD3D12Device* Device, FD3D12DescriptorHeap* TargetHeap, FRHIDescriptorHandle DstHandle, D3D12_CPU_DESCRIPTOR_HANDLE SrcCpuHandle);
+	void CopyDescriptors(FD3D12Device* Device, FD3D12DescriptorHeap* TargetHeap, FD3D12DescriptorHeap* SourceHeap, uint32 FirstDescriptor, uint32 NumDescriptors);
+	void CopyDescriptors(FD3D12Device* Device, FD3D12DescriptorHeap* TargetHeap, TConstArrayView<FRHIDescriptorHandle> DstHandles, TConstArrayView<FD3D12OfflineDescriptor> SrcOfflineDescriptors);
+
+	FD3D12OfflineDescriptor CreateOfflineCopy(FD3D12Device* Device, D3D12_CPU_DESCRIPTOR_HANDLE InDescriptor, ERHIDescriptorHeapType InType);
+	FD3D12OfflineDescriptor CreateOfflineCopy(FD3D12Device* Device, FD3D12DescriptorHeap* InHeap, FRHIDescriptorHandle InHandle);
+	void FreeOfflineCopy(FD3D12Device* Device, FD3D12OfflineDescriptor& InDescriptor, ERHIDescriptorHeapType InType);
 }
 
 struct FD3D12DescriptorHeap : public FD3D12DeviceChild, public FThreadSafeRefCountedObject
@@ -37,7 +64,7 @@ public:
 	FD3D12DescriptorHeap() = delete;
 
 	// Heap created with its own D3D heap object.
-	FD3D12DescriptorHeap(FD3D12Device* InDevice, ID3D12DescriptorHeap* InHeap, uint32 InNumDescriptors, ERHIDescriptorHeapType InType, D3D12_DESCRIPTOR_HEAP_FLAGS InFlags, bool bInIsGlobal);
+	FD3D12DescriptorHeap(FD3D12Device* InDevice, TRefCountPtr<ID3D12DescriptorHeap>&& InHeap, uint32 InNumDescriptors, ERHIDescriptorHeapType InType, ED3D12DescriptorHeapFlags InFlags, bool bInIsGlobal);
 
 	// Heap created as a suballocation of another heap.
 	FD3D12DescriptorHeap(FD3D12DescriptorHeap* SubAllocateSourceHeap, uint32 InOffset, uint32 InNumDescriptors);
@@ -46,7 +73,7 @@ public:
 
 	inline ID3D12DescriptorHeap*       GetHeap()  const { return Heap; }
 	inline ERHIDescriptorHeapType      GetType()  const { return Type; }
-	inline D3D12_DESCRIPTOR_HEAP_FLAGS GetFlags() const { return Flags; }
+	inline ED3D12DescriptorHeapFlags   GetFlags() const { return Flags; }
 
 	inline uint32 GetOffset()         const { return Offset; }
 	inline uint32 GetNumDescriptors() const { return NumDescriptors; }
@@ -59,6 +86,8 @@ public:
 	inline D3D12_CPU_DESCRIPTOR_HANDLE GetCPUSlotHandle(uint32 Slot) const { return CD3DX12_CPU_DESCRIPTOR_HANDLE(CpuBase, Slot, DescriptorSize); }
 	inline D3D12_GPU_DESCRIPTOR_HANDLE GetGPUSlotHandle(uint32 Slot) const { return CD3DX12_GPU_DESCRIPTOR_HANDLE(GpuBase, Slot, DescriptorSize); }
 
+	inline bool CanBePooled() const { return EnumHasAnyFlags(GetFlags(), ED3D12DescriptorHeapFlags::Poolable); }
+
 private:
 	TRefCountPtr<ID3D12DescriptorHeap> Heap;
 
@@ -75,7 +104,7 @@ private:
 	const uint32 DescriptorSize;
 
 	const ERHIDescriptorHeapType Type;
-	const D3D12_DESCRIPTOR_HEAP_FLAGS Flags;
+	const ED3D12DescriptorHeapFlags Flags;
 
 	// Enabled if this heap is the "global" heap.
 	const bool bIsGlobal;
@@ -98,7 +127,7 @@ public:
 	inline       FD3D12DescriptorHeap* GetHeap()       { return Heap.GetReference(); }
 	inline const FD3D12DescriptorHeap* GetHeap() const { return Heap.GetReference(); }
 
-	inline bool HandlesAllocationWithFlags(ERHIDescriptorHeapType InHeapType, D3D12_DESCRIPTOR_HEAP_FLAGS InHeapFlags) const
+	inline bool HandlesAllocationWithFlags(ERHIDescriptorHeapType InHeapType, ED3D12DescriptorHeapFlags InHeapFlags) const
 	{
 		return HandlesAllocation(InHeapType) && Heap->GetFlags() == InHeapFlags;
 	}
@@ -136,6 +165,7 @@ public:
 
 	// Setup the actual heap
 	void Init(uint32 InTotalSize, uint32 InBlockSize, bool bBindlessResources);
+	void CleanupResources();
 
 	// Allocate an available sub heap block from the global heap
 	FD3D12OnlineDescriptorBlock* AllocateHeapBlock();
@@ -206,6 +236,8 @@ public:
 	FD3D12OfflineDescriptor AllocateHeapSlot();
 	void FreeHeapSlot(FD3D12OfflineDescriptor& Descriptor);
 
+	void CleanupResources();
+
 private:
 	void AllocateHeap();
 
@@ -230,14 +262,33 @@ public:
 	void Init(uint32 InNumGlobalResourceDescriptors, uint32 InNumGlobalSamplerDescriptors);
 	void Destroy();
 
-	FD3D12DescriptorHeap* AllocateIndependentHeap(const TCHAR* InDebugName, ERHIDescriptorHeapType InHeapType, uint32 InNumDescriptors, D3D12_DESCRIPTOR_HEAP_FLAGS InHeapFlags);
-	FD3D12DescriptorHeap* AllocateHeap(const TCHAR* InDebugName, ERHIDescriptorHeapType InHeapType, uint32 InNumDescriptors, D3D12_DESCRIPTOR_HEAP_FLAGS InHeapFlags);
+	FD3D12DescriptorHeap* AllocateIndependentHeap(const TCHAR* InDebugName, ERHIDescriptorHeapType InHeapType, uint32 InNumDescriptors, ED3D12DescriptorHeapFlags InHeapFlags);
+	FD3D12DescriptorHeap* AllocateHeap(const TCHAR* InDebugName, ERHIDescriptorHeapType InHeapType, uint32 InNumDescriptors, ED3D12DescriptorHeapFlags InHeapFlags);
 	void DeferredFreeHeap(FD3D12DescriptorHeap* InHeap);
 	void ImmediateFreeHeap(FD3D12DescriptorHeap* InHeap);
 
-	//inline FD3D12DescriptorHeap* GetGlobalHeap() const { return GlobalHeap; }
+	void AddHeapToPool(TRefCountPtr<ID3D12DescriptorHeap>&& InHeap, ERHIDescriptorHeapType InType, uint32 InNumDescriptors, ED3D12DescriptorHeapFlags InFlags);
 
 private:
+	TRefCountPtr<ID3D12DescriptorHeap> AcquirePooledHeap(ERHIDescriptorHeapType InType, uint32 InNumDescriptors, ED3D12DescriptorHeapFlags InFlags);
+
 	TArray<FD3D12DescriptorManager> GlobalHeaps;
-};
 
+	struct FPooledHeap
+	{
+		FPooledHeap(TRefCountPtr<ID3D12DescriptorHeap>&& InHeap, ERHIDescriptorHeapType InType, uint32 InNumDescriptors, ED3D12DescriptorHeapFlags InFlags)
+			: Heap(Forward<TRefCountPtr<ID3D12DescriptorHeap>>(InHeap))
+			, NumDescriptors(InNumDescriptors)
+			, Type(InType)
+			, Flags(InFlags)
+		{
+		}
+
+		TRefCountPtr<ID3D12DescriptorHeap> Heap;
+		uint32                             NumDescriptors;
+		ERHIDescriptorHeapType             Type;
+		ED3D12DescriptorHeapFlags          Flags;
+	};
+	TArray<FPooledHeap> PooledHeaps;
+	FCriticalSection PooledHeapsCS;
+};
diff --git a/Engine/Source/Runtime/D3D12RHI/Private/D3D12Device.cpp b/Engine/Source/Runtime/D3D12RHI/Private/D3D12Device.cpp
--- a/Engine/Source/Runtime/D3D12RHI/Private/D3D12Device.cpp
+++ b/Engine/Source/Runtime/D3D12RHI/Private/D3D12Device.cpp
@@ -418,6 +418,23 @@ void FD3D12Device::SetupAfterDeviceCreation()
 	ImmediateCommandContext = FD3D12DynamicRHI::GetD3DRHI()->CreateCommandContext(this, ED3D12QueueType::Direct, true);
 }
 
+void FD3D12Device::CleanupResources()
+{
+	for (FD3D12OfflineDescriptorManager& Manager : OfflineDescriptorManagers)
+	{
+		Manager.CleanupResources();
+	}
+	OnlineDescriptorManager.CleanupResources();
+
+#if PLATFORM_SUPPORTS_BINDLESS_RENDERING
+	BindlessDescriptorManager.CleanupResources();
+#endif
+
+#if D3D12_RHI_RAYTRACING
+	CleanupRayTracing();
+#endif
+}
+
 void FD3D12Device::CreateDefaultViews()
 {
 	{
diff --git a/Engine/Source/Runtime/D3D12RHI/Private/D3D12Device.h b/Engine/Source/Runtime/D3D12RHI/Private/D3D12Device.h
--- a/Engine/Source/Runtime/D3D12RHI/Private/D3D12Device.h
+++ b/Engine/Source/Runtime/D3D12RHI/Private/D3D12Device.h
@@ -341,6 +341,7 @@ public:
 	// shared code for different D3D12  devices (e.g. PC DirectX12 and XboxOne) called
 	// after device creation and GRHISupportsAsyncTextureCreation was set and before resource init
 	void SetupAfterDeviceCreation();
+	void CleanupResources();
 
 private:
 	// called by SetupAfterDeviceCreation() when the device gets initialized
diff --git a/Engine/Source/Runtime/D3D12RHI/Private/D3D12ExplicitDescriptorCache.cpp b/Engine/Source/Runtime/D3D12RHI/Private/D3D12ExplicitDescriptorCache.cpp
--- a/Engine/Source/Runtime/D3D12RHI/Private/D3D12ExplicitDescriptorCache.cpp
+++ b/Engine/Source/Runtime/D3D12RHI/Private/D3D12ExplicitDescriptorCache.cpp
@@ -337,10 +337,8 @@ void FD3D12ExplicitDescriptorCache::SetDescriptorHeaps(FD3D12CommandContext& Com
 	check(bBindlessViews || ViewHeap.GetParentDevice() == CommandContext.GetParentDevice());
 	check(bBindlessSamplers || SamplerHeap.GetParentDevice() == CommandContext.GetParentDevice());
 
-	FD3D12BindlessDescriptorManager& BindlessManager = GetParentDevice()->GetBindlessDescriptorManager();
-
-	ID3D12DescriptorHeap* ViewHeapToSet = bBindlessViews ? BindlessManager.GetResourceHeap(CommandContext.GetPipeline())->GetHeap() : ViewHeap.D3D12Heap;
-	ID3D12DescriptorHeap* SamplerHeapToSet = bBindlessSamplers ? BindlessManager.GetSamplerHeap()->GetHeap() : SamplerHeap.D3D12Heap;
+	ID3D12DescriptorHeap* ViewHeapToSet = bBindlessViews ? nullptr : ViewHeap.D3D12Heap;
+	ID3D12DescriptorHeap* SamplerHeapToSet = bBindlessSamplers ? nullptr : SamplerHeap.D3D12Heap;
 #else
 	check(ViewHeap.GetParentDevice() == CommandContext.GetParentDevice());
 	check(SamplerHeap.GetParentDevice() == CommandContext.GetParentDevice());
diff --git a/Engine/Source/Runtime/D3D12RHI/Private/D3D12RHIPrivate.h b/Engine/Source/Runtime/D3D12RHI/Private/D3D12RHIPrivate.h
--- a/Engine/Source/Runtime/D3D12RHI/Private/D3D12RHIPrivate.h
+++ b/Engine/Source/Runtime/D3D12RHI/Private/D3D12RHIPrivate.h
@@ -305,6 +305,7 @@ public:
 	virtual void* RHILockBufferMGPU(FRHICommandListBase& RHICmdList, FRHIBuffer* Buffer, uint32 GPUIndex, uint32 Offset, uint32 Size, EResourceLockMode LockMode) final override;
 	virtual void RHIUnlockBuffer(FRHICommandListBase& RHICmdList, FRHIBuffer* Buffer) final override;
 	virtual void RHIUnlockBufferMGPU(FRHICommandListBase& RHICmdList, FRHIBuffer* Buffer, uint32 GPUIndex) final override;
+	virtual FTextureReferenceRHIRef RHICreateTextureReference(FRHICommandListBase& RHICmdList, FRHITexture* InReferencedTexture) final override;
 	virtual void RHIUpdateTextureReference(FRHICommandListBase& RHICmdList, FRHITextureReference* TextureRef, FRHITexture* NewTexture) final override;
 	virtual FRHICalcTextureSizeResult RHICalcTexturePlatformSize(const FRHITextureDesc& Desc, uint32 FirstMipIndex) override;
 	virtual void RHIGetTextureMemoryStats(FTextureMemoryStats& OutStats) final override;
diff --git a/Engine/Source/Runtime/D3D12RHI/Private/D3D12Resources.h b/Engine/Source/Runtime/D3D12RHI/Private/D3D12Resources.h
--- a/Engine/Source/Runtime/D3D12RHI/Private/D3D12Resources.h
+++ b/Engine/Source/Runtime/D3D12RHI/Private/D3D12Resources.h
@@ -785,11 +785,8 @@ struct FD3D12LockedResource : public FD3D12DeviceChild
 };
 
 /** Resource which might needs to be notified about changes on dependent resources (Views, RTGeometryObject, Cached binding tables) */
-class FD3D12ShaderResourceRenameListener
+struct FD3D12ShaderResourceRenameListener
 {
-protected:
-
-	friend class FD3D12BaseShaderResource;
 	virtual void ResourceRenamed(FRHICommandListBase& RHICmdList, FD3D12BaseShaderResource* InRenamedResource, FD3D12ResourceLocation* InNewResourceLocation) = 0;
 };
 
diff --git a/Engine/Source/Runtime/D3D12RHI/Private/D3D12State.cpp b/Engine/Source/Runtime/D3D12RHI/Private/D3D12State.cpp
--- a/Engine/Source/Runtime/D3D12RHI/Private/D3D12State.cpp
+++ b/Engine/Source/Runtime/D3D12RHI/Private/D3D12State.cpp
@@ -669,13 +669,7 @@ FD3D12SamplerState::FD3D12SamplerState(FD3D12Device* InParent, const D3D12_SAMPL
 	GetParentDevice()->CreateSamplerInternal(Desc, OfflineDescriptor);
 
 #if PLATFORM_SUPPORTS_BINDLESS_RENDERING
-	FD3D12BindlessDescriptorManager& BindlessDescriptorManager = GetParentDevice()->GetBindlessDescriptorManager();
-	BindlessHandle = BindlessDescriptorManager.Allocate(ERHIDescriptorHeapType::Sampler);
-
-	if (BindlessHandle.IsValid())
-	{
-		GetParentDevice()->GetBindlessDescriptorManager().UpdateDescriptorImmediately(BindlessHandle, OfflineDescriptor);
-	}
+	BindlessHandle = GetParentDevice()->GetBindlessDescriptorManager().AllocateAndInitialize(this);
 #endif
 }
 
diff --git a/Engine/Source/Runtime/D3D12RHI/Private/D3D12StateCache.cpp b/Engine/Source/Runtime/D3D12RHI/Private/D3D12StateCache.cpp
--- a/Engine/Source/Runtime/D3D12RHI/Private/D3D12StateCache.cpp
+++ b/Engine/Source/Runtime/D3D12RHI/Private/D3D12StateCache.cpp
@@ -453,14 +453,14 @@ void FD3D12StateCache::ApplyState(ERHIPipeline HardwarePipe, ED3D12PipelineType
 	{
 		FD3D12BindlessDescriptorManager& BindlessManager = GetParentDevice()->GetBindlessDescriptorManager();
 
-		FD3D12DescriptorHeap* ResourceHeap = BindlessManager.GetResourceHeap(HardwarePipe, ERHIBindlessConfiguration::AllShaders);
-		FD3D12DescriptorHeap* SamplerHeap = BindlessManager.GetSamplerHeap(ERHIBindlessConfiguration::AllShaders);
+		bool bHaveResourceHeap = BindlessManager.AreResourcesBindless(ERHIBindlessConfiguration::AllShaders);
+		bool bHaveSamplerHeap = BindlessManager.AreSamplersBindless(ERHIBindlessConfiguration::AllShaders);
 
-		checkf(!bBindlessResources || ResourceHeap != nullptr, TEXT("Using dynamic samplers without the bindless sampler heap configured. Please check your configuration."));
-		checkf(!bBindlessSamplers || SamplerHeap != nullptr, TEXT("Using dynamic samplers without the bindless sampler heap configured. Please check your configuration."));
+		checkf(!bBindlessResources || bHaveResourceHeap, TEXT("Using dynamic samplers without the bindless sampler heap configured. Please check your configuration."));
+		checkf(!bBindlessSamplers  || bHaveSamplerHeap, TEXT("Using dynamic samplers without the bindless sampler heap configured. Please check your configuration."));
 
-		check(!(ResourceHeap != nullptr && bHasTableResources));
-		check(!(SamplerHeap  != nullptr && bHasSamplers      ));
+		check(!(bHaveResourceHeap && bHasTableResources));
+		check(!(bHaveSamplerHeap  && bHasSamplers      ));
 	}
 #endif
 
@@ -768,6 +768,7 @@ void FD3D12StateCache::ApplyResources(const FD3D12RootSignature* const pRootSign
 
 void FD3D12StateCache::ApplyBindlessResources(const FD3D12RootSignature* const pRootSignature, uint32 StartStage, uint32 EndStage)
 {
+#if PLATFORM_SUPPORTS_BINDLESS_RENDERING
 	for (uint32 Index = StartStage; Index < EndStage; Index++)
 	{
 		DescriptorCache.PrepareBindlessViews(
@@ -835,6 +836,8 @@ void FD3D12StateCache::ApplyBindlessResources(const FD3D12RootSignature* const p
 		}
 	}
 #endif // D3D12RHI_USE_CONSTANT_BUFFER_VIEWS
+
+#endif // PLATFORM_SUPPORTS_BINDLESS_RENDERING
 }
 
 void FD3D12StateCache::ApplyConstants(const FD3D12RootSignature* const pRootSignature, uint32 StartStage, uint32 EndStage)
diff --git a/Engine/Source/Runtime/D3D12RHI/Private/D3D12Texture.cpp b/Engine/Source/Runtime/D3D12RHI/Private/D3D12Texture.cpp
--- a/Engine/Source/Runtime/D3D12RHI/Private/D3D12Texture.cpp
+++ b/Engine/Source/Runtime/D3D12RHI/Private/D3D12Texture.cpp
@@ -958,27 +958,6 @@ FTextureRHIRef FD3D12DynamicRHI::RHICreateTexture(FRHICommandListBase& RHICmdLis
 	return CreateD3D12Texture(CreateDesc, &RHICmdList, ResourceAllocator);
 }
 
-void FD3D12DynamicRHI::RHIUpdateTextureReference(FRHICommandListBase& RHICmdList, FRHITextureReference* TextureRef, FRHITexture* InNewTexture)
-{
-	FRHITexture* NewTexture = InNewTexture ? InNewTexture : FRHITextureReference::GetDefaultTexture();
-
-#if PLATFORM_SUPPORTS_BINDLESS_RENDERING
-	if (FRHIShaderResourceView* DestinationSRVRHI = TextureRef ? TextureRef->GetBindlessView() : nullptr)
-	{
-		FD3D12ShaderResourceView* DestinationSRV = ResourceCast(DestinationSRVRHI);
-		if (DestinationSRV->IsBindless())
-		{
-			FD3D12ShaderResourceView* SourceSRV = ResourceCast(NewTexture)->GetShaderResourceView();
-
-			FD3D12BindlessDescriptorManager& BindlessDescriptorManager = DestinationSRV->GetParentDevice()->GetBindlessDescriptorManager();
-			BindlessDescriptorManager.UpdateResourceDescriptor(RHICmdList, DestinationSRV->GetBindlessHandle(), SourceSRV->GetOfflineCpuHandle());
-		}
-	}
-#endif // PLATFORM_SUPPORTS_BINDLESS_RENDERING
-
-	FDynamicRHI::RHIUpdateTextureReference(RHICmdList, TextureRef, NewTexture);
-}
-
 class FWaitInitialMipDataUploadTask
 {
 public:
@@ -2977,7 +2956,12 @@ void FD3D12CommandContext::RHICopyTexture(FRHITexture* SourceTextureRHI, FRHITex
 // FD3D12BackBufferReferenceTexture2D functions
 ///////////////////////////////////////////////////////////////////////////////////////////
 
-FRHITexture* FD3D12BackBufferReferenceTexture2D::GetBackBufferTexture()
+FRHITexture* FD3D12BackBufferReferenceTexture2D::GetBackBufferTexture() const
 {
 	return bIsSDR ? Viewport->GetSDRBackBuffer_RHIThread() : Viewport->GetBackBuffer_RHIThread();
 }
+
+FRHIDescriptorHandle FD3D12BackBufferReferenceTexture2D::GetDefaultBindlessHandle() const
+{
+	return GetBackBufferTexture()->GetDefaultBindlessHandle();
+}
diff --git a/Engine/Source/Runtime/D3D12RHI/Private/D3D12Texture.h b/Engine/Source/Runtime/D3D12RHI/Private/D3D12Texture.h
--- a/Engine/Source/Runtime/D3D12RHI/Private/D3D12Texture.h
+++ b/Engine/Source/Runtime/D3D12RHI/Private/D3D12Texture.h
@@ -54,7 +54,7 @@ public:
 	// FRHITexture overrides
 	virtual void* GetTextureBaseRHI() override final			{ return this; }
 	virtual void* GetNativeResource() const override final;
-	virtual FRHIDescriptorHandle GetDefaultBindlessHandle() const final;
+	virtual FRHIDescriptorHandle GetDefaultBindlessHandle() const override;
 			
 	// Accessors.
 	bool IsStreamable() const									{ return EnumHasAnyFlags(GetDesc().Flags, ETextureCreateFlags::Streamable); }
@@ -196,24 +196,20 @@ class FD3D12Viewport;
 class FD3D12BackBufferReferenceTexture2D : public FD3D12Texture
 {
 public:
-
-	FD3D12BackBufferReferenceTexture2D(
-		const FRHITextureCreateDesc& InDesc,
-		FD3D12Viewport* InViewPort,
-		bool bInIsSDR,
-		FD3D12Device* InDevice) :
-		FD3D12Texture(InDesc, InDevice),
-		Viewport(InViewPort), bIsSDR(bInIsSDR)
+	FD3D12BackBufferReferenceTexture2D(const FRHITextureCreateDesc& InDesc, FD3D12Viewport* InViewPort, bool bInIsSDR, FD3D12Device* InDevice)
+		: FD3D12Texture(InDesc, InDevice)
+		, Viewport(InViewPort)
+		, bIsSDR(bInIsSDR)
 	{
 	}
 
-	FD3D12Viewport* GetViewPort() { return Viewport; }
+	FD3D12Viewport* GetViewPort() const { return Viewport; }
 	bool IsSDR() const { return bIsSDR; }
 
-	FRHITexture* GetBackBufferTexture();
+	FRHITexture* GetBackBufferTexture() const;
+	virtual FRHIDescriptorHandle GetDefaultBindlessHandle() const override;
 
 private:
-
 	FD3D12Viewport* Viewport = nullptr;
 	bool bIsSDR = false;
 };
diff --git a/Engine/Source/Runtime/D3D12RHI/Private/D3D12TextureReference.cpp b/Engine/Source/Runtime/D3D12RHI/Private/D3D12TextureReference.cpp
new file mode 100644
--- /dev/null
+++ b/Engine/Source/Runtime/D3D12RHI/Private/D3D12TextureReference.cpp
@@ -0,0 +1,93 @@
+// Copyright Epic Games, Inc. All Rights Reserved.
+
+#include "D3D12TextureReference.h"
+
+FD3D12RHITextureReference::FD3D12RHITextureReference(FD3D12Device* InDevice, FD3D12Texture* InReferencedTexture)
+	: FD3D12DeviceChild(InDevice)
+#if PLATFORM_SUPPORTS_BINDLESS_RENDERING
+	, FRHITextureReference(InReferencedTexture, InDevice->GetBindlessDescriptorManager().AllocateResourceHandle())
+#else
+	, FRHITextureReference(InReferencedTexture)
+#endif
+{
+#if PLATFORM_SUPPORTS_BINDLESS_RENDERING
+	if (BindlessHandle.IsValid())
+	{
+		InReferencedTexture->AddRenameListener(this);
+
+		InDevice->GetBindlessDescriptorManager().UpdateDescriptorImmediately(BindlessHandle, InReferencedTexture->GetShaderResourceView());
+	}
+#endif // PLATFORM_SUPPORTS_BINDLESS_RENDERING
+}
+
+FD3D12RHITextureReference::~FD3D12RHITextureReference()
+{
+#if PLATFORM_SUPPORTS_BINDLESS_RENDERING
+	if (BindlessHandle.IsValid())
+	{
+		FD3D12DynamicRHI::ResourceCast(GetReferencedTexture())->RemoveRenameListener(this);
+
+		GetParentDevice()->GetBindlessDescriptorManager().DeferredFreeFromDestructor(BindlessHandle);
+	}
+#endif // PLATFORM_SUPPORTS_BINDLESS_RENDERING
+}
+
+void FD3D12RHITextureReference::SwitchToNewTexture(FRHICommandListBase& RHICmdList, FD3D12Texture* InNewTexture)
+{
+#if PLATFORM_SUPPORTS_BINDLESS_RENDERING
+	if (BindlessHandle.IsValid())
+	{
+		FD3D12Texture* NewTexture = InNewTexture ? InNewTexture : FD3D12DynamicRHI::ResourceCast(FRHITextureReference::GetDefaultTexture());
+
+		FD3D12Texture* CurrentTexture = FD3D12DynamicRHI::ResourceCast(GetReferencedTexture());
+
+		if (CurrentTexture != NewTexture)
+		{
+			if (CurrentTexture)
+			{
+				CurrentTexture->RemoveRenameListener(this);
+			}
+
+			NewTexture->AddRenameListener(this);
+
+			GetParentDevice()->GetBindlessDescriptorManager().UpdateDescriptor(RHICmdList, BindlessHandle, NewTexture->GetShaderResourceView());
+		}
+	}
+#endif // PLATFORM_SUPPORTS_BINDLESS_RENDERING
+
+	SetReferencedTexture(InNewTexture);
+}
+
+#if PLATFORM_SUPPORTS_BINDLESS_RENDERING
+void FD3D12RHITextureReference::ResourceRenamed(FRHICommandListBase& RHICmdList, FD3D12BaseShaderResource* InRenamedResource, FD3D12ResourceLocation* InNewResourceLocation)
+{
+	if (ensure(BindlessHandle.IsValid()))
+	{
+		FD3D12Texture* RenamedTexture = static_cast<FD3D12Texture*>(InRenamedResource);
+		checkSlow(RenamedTexture == ReferencedTexture);
+
+		GetParentDevice()->GetBindlessDescriptorManager().UpdateDescriptor(RHICmdList, BindlessHandle, RenamedTexture->GetShaderResourceView());
+	}
+}
+#endif // PLATFORM_SUPPORTS_BINDLESS_RENDERING
+
+FTextureReferenceRHIRef FD3D12DynamicRHI::RHICreateTextureReference(FRHICommandListBase& RHICmdList, FRHITexture* InReferencedTexture)
+{
+	FRHITexture* ReferencedTexture = InReferencedTexture ? InReferencedTexture : FRHITextureReference::GetDefaultTexture();
+
+	FD3D12Adapter* Adapter = &GetAdapter();
+	return Adapter->CreateLinkedObject<FD3D12RHITextureReference>(FRHIGPUMask::All(), [ReferencedTexture](FD3D12Device* Device)
+	{
+		return new FD3D12RHITextureReference(Device, ResourceCast(ReferencedTexture, Device->GetGPUIndex()));
+	});
+}
+
+void FD3D12DynamicRHI::RHIUpdateTextureReference(FRHICommandListBase& RHICmdList, FRHITextureReference* TextureRef, FRHITexture* InNewTexture)
+{
+	FRHITexture* NewTexture = InNewTexture ? InNewTexture : FRHITextureReference::GetDefaultTexture();
+
+	for (TD3D12DualLinkedObjectIterator<FD3D12RHITextureReference, FD3D12Texture> It(ResourceCast(TextureRef), ResourceCast(NewTexture)); It; ++It)
+	{
+		It.GetFirst()->SwitchToNewTexture(RHICmdList, It.GetSecond());
+	}
+}
diff --git a/Engine/Source/Runtime/D3D12RHI/Private/D3D12TextureReference.h b/Engine/Source/Runtime/D3D12RHI/Private/D3D12TextureReference.h
new file mode 100644
--- /dev/null
+++ b/Engine/Source/Runtime/D3D12RHI/Private/D3D12TextureReference.h
@@ -0,0 +1,30 @@
+// Copyright Epic Games, Inc. All Rights Reserved.
+
+#pragma once
+
+#include "D3D12RHIPrivate.h"
+#include "RHITextureReference.h"
+
+class FD3D12RHITextureReference : public FD3D12DeviceChild, public FRHITextureReference, public FD3D12LinkedAdapterObject<FD3D12RHITextureReference>
+#if PLATFORM_SUPPORTS_BINDLESS_RENDERING
+	, public FD3D12ShaderResourceRenameListener
+#endif
+{
+public:
+	FD3D12RHITextureReference() = delete;
+	FD3D12RHITextureReference(FD3D12Device* InDevice, FD3D12Texture* InReferencedTexture);
+	~FD3D12RHITextureReference();
+
+#if PLATFORM_SUPPORTS_BINDLESS_RENDERING
+	// FD3D12ShaderResourceRenameListener
+	virtual void ResourceRenamed(FRHICommandListBase& RHICmdList, FD3D12BaseShaderResource* InRenamedResource, FD3D12ResourceLocation* InNewResourceLocation) final override;
+#endif // PLATFORM_SUPPORTS_BINDLESS_RENDERING
+
+	void SwitchToNewTexture(FRHICommandListBase& RHICmdList, FD3D12Texture* InNewTexture);
+};
+
+template<>
+struct TD3D12ResourceTraits<FRHITextureReference>
+{
+	using TConcreteType = FD3D12RHITextureReference;
+};
diff --git a/Engine/Source/Runtime/D3D12RHI/Private/D3D12UAV.cpp b/Engine/Source/Runtime/D3D12RHI/Private/D3D12UAV.cpp
--- a/Engine/Source/Runtime/D3D12RHI/Private/D3D12UAV.cpp
+++ b/Engine/Source/Runtime/D3D12RHI/Private/D3D12UAV.cpp
@@ -362,8 +362,20 @@ void FD3D12CommandContext::ClearUAV(TRHICommandList_RecursiveHazardous<FD3D12Com
 				FD3D12UnorderedAccessView UAV(ParentDevice);
 				UAV.CreateView(UnorderedAccessView->GetResourceLocation(), R32UAVDesc, FD3D12UnorderedAccessView::EFlags::None);
 
+				FD3D12OfflineDescriptor OfflineHandle = UAV.GetOfflineCpuHandle();
 				D3D12_GPU_DESCRIPTOR_HANDLE GPUHandle{};
-				if (!Context.StateCache.GetDescriptorCache()->IsUsingBindlessResources())
+
+#if PLATFORM_SUPPORTS_BINDLESS_RENDERING
+				if (UAV.GetBindlessHandle().IsValid())
+				{
+					Context.FlushPendingDescriptorUpdates();
+
+					FD3D12DescriptorHeap* BindlessHeap = Context.StateCache.GetDescriptorCache()->GetBindlessResourcesHeap();
+					UE::D3D12Descriptors::CopyDescriptor(ParentDevice, BindlessHeap, UAV.GetBindlessHandle(), OfflineHandle);
+					GPUHandle = BindlessHeap->GetGPUSlotHandle(UAV.GetBindlessHandle().GetIndex());
+				}
+				else
+#endif
 				{
 					// Check if the view heap is full and needs to rollover.
 					if (!Context.StateCache.GetDescriptorCache()->GetCurrentViewHeap()->CanReserveSlots(1))
@@ -375,19 +387,13 @@ void FD3D12CommandContext::ClearUAV(TRHICommandList_RecursiveHazardous<FD3D12Com
 					D3D12_CPU_DESCRIPTOR_HANDLE DestSlot = Context.StateCache.GetDescriptorCache()->GetCurrentViewHeap()->GetCPUSlotHandle(ReservedSlot);
 					GPUHandle = Context.StateCache.GetDescriptorCache()->GetCurrentViewHeap()->GetGPUSlotHandle(ReservedSlot);
 
-					Device->CopyDescriptorsSimple(1, DestSlot, UAV.GetOfflineCpuHandle(), D3D12_DESCRIPTOR_HEAP_TYPE_CBV_SRV_UAV);
-				}
-#if PLATFORM_SUPPORTS_BINDLESS_RENDERING
-				else
-				{
-					GPUHandle = ParentDevice->GetBindlessDescriptorManager().GetResourceGpuHandle(Context.GetPipeline(), UAV.GetBindlessHandle());
+					Device->CopyDescriptorsSimple(1, DestSlot, OfflineHandle, D3D12_DESCRIPTOR_HEAP_TYPE_CBV_SRV_UAV);
 				}
-#endif
 
 				Context.TransitionResource(UnorderedAccessView, D3D12_RESOURCE_STATE_UNORDERED_ACCESS);
 
 				Context.FlushResourceBarriers();
-				Context.GraphicsCommandList()->ClearUnorderedAccessViewUint(GPUHandle, UAV.GetOfflineCpuHandle(), UAV.GetResource()->GetResource(), *reinterpret_cast<const UINT(*)[4]>(ClearValues), 0, nullptr);
+				Context.GraphicsCommandList()->ClearUnorderedAccessViewUint(GPUHandle, OfflineHandle, UAV.GetResource()->GetResource(), *reinterpret_cast<const UINT(*)[4]>(ClearValues), 0, nullptr);
 				Context.UpdateResidency(UnorderedAccessView->GetResidencyHandles());
 				Context.ConditionalSplitCommandList();
 
diff --git a/Engine/Source/Runtime/D3D12RHI/Private/D3D12View.cpp b/Engine/Source/Runtime/D3D12RHI/Private/D3D12View.cpp
--- a/Engine/Source/Runtime/D3D12RHI/Private/D3D12View.cpp
+++ b/Engine/Source/Runtime/D3D12RHI/Private/D3D12View.cpp
@@ -220,7 +220,7 @@ FD3D12View::FD3D12View(FD3D12Device* InDevice, ERHIDescriptorHeapType InHeapType
 #if PLATFORM_SUPPORTS_BINDLESS_RENDERING
 	if (HeapType == ERHIDescriptorHeapType::Standard)
 	{
-		BindlessHandle = InDevice->GetBindlessDescriptorManager().Allocate(ERHIDescriptorHeapType::Standard);
+		BindlessHandle = InDevice->GetBindlessDescriptorManager().AllocateResourceHandle();
 	}
 #endif
 }
@@ -229,10 +229,9 @@ FD3D12View::~FD3D12View()
 {
 	// Unregister this view from the underlying resource
 	if (ResourceInfo.BaseResource)
+	{
 		ResourceInfo.BaseResource->RemoveRenameListener(this);
-
-	// Free the descriptor heap slot and bindless handle
-	GetParentDevice()->GetOfflineDescriptorManager(HeapType).FreeHeapSlot(OfflineCpuHandle);
+	}
 
 #if PLATFORM_SUPPORTS_BINDLESS_RENDERING
 	if (BindlessHandle.IsValid())
@@ -241,6 +240,9 @@ FD3D12View::~FD3D12View()
 		BindlessHandle = {};
 	}
 #endif
+
+	// Free the descriptor heap slot and bindless handle
+	GetParentDevice()->GetOfflineDescriptorManager(HeapType).FreeHeapSlot(OfflineCpuHandle);
 }
 
 void FD3D12View::InitializeBindlessSlot()
@@ -248,8 +250,7 @@ void FD3D12View::InitializeBindlessSlot()
 #if PLATFORM_SUPPORTS_BINDLESS_RENDERING
 	if (BindlessHandle.IsValid())
 	{
-		FD3D12BindlessDescriptorManager& BindlessManager = GetParentDevice()->GetBindlessDescriptorManager();
-		BindlessManager.UpdateDescriptorImmediately(BindlessHandle, OfflineCpuHandle);
+		GetParentDevice()->GetBindlessDescriptorManager().UpdateDescriptorImmediately(BindlessHandle, this);
 	}
 #endif
 }
@@ -260,8 +261,7 @@ void FD3D12View::UpdateBindlessSlot(FRHICommandListBase& RHICmdList)
 	if (BindlessHandle.IsValid())
 	{
 		check(BindlessHandle.GetType() == ERHIDescriptorHeapType::Standard);
-		FD3D12BindlessDescriptorManager& BindlessManager = GetParentDevice()->GetBindlessDescriptorManager();
-		BindlessManager.UpdateResourceDescriptor(RHICmdList, BindlessHandle, OfflineCpuHandle);
+		GetParentDevice()->GetBindlessDescriptorManager().UpdateDescriptor(RHICmdList, BindlessHandle, this);
 	}
 #endif
 }
diff --git a/Engine/Source/Runtime/RHI/Private/DynamicRHI.cpp b/Engine/Source/Runtime/RHI/Private/DynamicRHI.cpp
--- a/Engine/Source/Runtime/RHI/Private/DynamicRHI.cpp
+++ b/Engine/Source/Runtime/RHI/Private/DynamicRHI.cpp
@@ -563,6 +563,12 @@ void FDynamicRHI::EnableIdealGPUCaptureOptions(bool bEnabled)
 	}	
 }
 
+FTextureReferenceRHIRef FDynamicRHI::RHICreateTextureReference(FRHICommandListBase& RHICmdList, FRHITexture* InReferencedTexture)
+{
+	FRHITexture* ReferencedTexture = InReferencedTexture ? InReferencedTexture : FRHITextureReference::GetDefaultTexture();
+	return new FRHITextureReference(ReferencedTexture);
+}
+
 void FDynamicRHI::RHIUpdateTextureReference(FRHICommandListBase& RHICmdList, FRHITextureReference* TextureRef, FRHITexture* InReferencedTexture)
 {
 	FRHITexture* ReferencedTexture = InReferencedTexture ? InReferencedTexture : FRHITextureReference::GetDefaultTexture();
@@ -584,24 +590,6 @@ uint64 FDynamicRHI::RHIGetMinimumAlignmentForBufferBackedSRV(EPixelFormat Format
 	return GPixelFormats[Format].BlockBytes;
 }
 
-FTextureReferenceRHIRef FDynamicRHI::RHICreateTextureReference(FRHICommandListBase& RHICmdList, FRHITexture* InReferencedTexture)
-{
-	FRHITexture* ReferencedTexture = InReferencedTexture ? InReferencedTexture : FRHITextureReference::GetDefaultTexture();
-
-	FShaderResourceViewRHIRef ShaderResourceView;
-
-#if PLATFORM_SUPPORTS_BINDLESS_RENDERING
-	// If the referenced texture is configured for bindless, make sure we also create an SRV to use for bindless.
-	if (ReferencedTexture && ReferencedTexture->GetDefaultBindlessHandle().IsValid())
-	{
-		ShaderResourceView = RHICmdList.CreateShaderResourceView(ReferencedTexture, 0u);
-	}
-#endif
-
-	return new FRHITextureReference(ReferencedTexture, ShaderResourceView);
-}
-
-
 uint64 FDynamicRHI::RHIComputeStatePrecachePSOHash(const FGraphicsPipelineStateInitializer& Initializer)
 {
 	struct FHashKey
diff --git a/Engine/Source/Runtime/RHI/Private/RHITextureReference.cpp b/Engine/Source/Runtime/RHI/Private/RHITextureReference.cpp
--- a/Engine/Source/Runtime/RHI/Private/RHITextureReference.cpp
+++ b/Engine/Source/Runtime/RHI/Private/RHITextureReference.cpp
@@ -5,21 +5,25 @@
 #include "RHICommandList.h"
 #include "RHIShaderPlatform.h"
 
-FRHITextureReference::FRHITextureReference()
+FRHITextureReference::FRHITextureReference(FRHITexture* InReferencedTexture)
 	: FRHITexture(RRT_TextureReference)
-	, ReferencedTexture(DefaultTexture.GetReference())
-	, BindlessView(nullptr)
+	, ReferencedTexture(InReferencedTexture ? InReferencedTexture : DefaultTexture.GetReference())
+#if PLATFORM_SUPPORTS_BINDLESS_RENDERING
+	, BindlessHandle(FRHIDescriptorHandle())
+#endif
 {
 	check(DefaultTexture);
 }
 
-FRHITextureReference::FRHITextureReference(FRHITexture* InReferencedTexture, FRHIShaderResourceView* InBindlessView)
+#if PLATFORM_SUPPORTS_BINDLESS_RENDERING
+FRHITextureReference::FRHITextureReference(FRHITexture* InReferencedTexture, FRHIDescriptorHandle InBindlessHandle)
 	: FRHITexture(RRT_TextureReference)
 	, ReferencedTexture(InReferencedTexture ? InReferencedTexture : DefaultTexture.GetReference())
-	, BindlessView(InBindlessView)
+	, BindlessHandle(InBindlessHandle)
 {
 	check(DefaultTexture);
 }
+#endif // PLATFORM_SUPPORTS_BINDLESS_RENDERING
 
 FRHITextureReference::~FRHITextureReference() = default;
 
@@ -30,14 +34,14 @@ FRHITextureReference* FRHITextureReference::GetTextureReference()
 
 FRHIDescriptorHandle FRHITextureReference::GetDefaultBindlessHandle() const
 { 
-	check(ReferencedTexture);
-
-	// If an SRV has been created, return its handle
-	if (BindlessView.IsValid())
+#if PLATFORM_SUPPORTS_BINDLESS_RENDERING
+	if (BindlessHandle.IsValid())
 	{
-		return BindlessView->GetBindlessHandle();
+		return BindlessHandle;
 	}
+#endif // PLATFORM_SUPPORTS_BINDLESS_RENDERING
 
+	check(ReferencedTexture);
 	return ReferencedTexture->GetDefaultBindlessHandle();
 }
 
diff --git a/Engine/Source/Runtime/RHI/Public/RHITextureReference.h b/Engine/Source/Runtime/RHI/Public/RHITextureReference.h
--- a/Engine/Source/Runtime/RHI/Public/RHITextureReference.h
+++ b/Engine/Source/Runtime/RHI/Public/RHITextureReference.h
@@ -4,12 +4,14 @@
 
 #include "RHIResources.h"
 
-class FRHITextureReference final : public FRHITexture
+class FRHITextureReference : public FRHITexture
 {
 public:
-	UE_DEPRECATED(5.3, "FRHITextureReference now requires a texture reference on construction")
-	RHI_API FRHITextureReference();
-	RHI_API FRHITextureReference(FRHITexture* InReferencedTexture, FRHIShaderResourceView* InBindlessView = nullptr);
+	FRHITextureReference() = delete;
+	RHI_API FRHITextureReference(FRHITexture* InReferencedTexture);
+#if PLATFORM_SUPPORTS_BINDLESS_RENDERING
+	RHI_API FRHITextureReference(FRHITexture* InReferencedTexture, FRHIDescriptorHandle InBindlessHandle);
+#endif
 	RHI_API ~FRHITextureReference();
 
 	RHI_API virtual class FRHITextureReference* GetTextureReference() override;
@@ -28,12 +30,15 @@ public:
 #endif
 
 	inline FRHITexture* GetReferencedTexture() const { return ReferencedTexture.GetReference(); }
-	inline FRHIShaderResourceView* GetBindlessView() const { return BindlessView.GetReference(); }
 
 	static inline FRHITexture* GetDefaultTexture() { return DefaultTexture; }
 
-private:
-	friend class FRHICommandListImmediate;
+#if PLATFORM_SUPPORTS_BINDLESS_RENDERING
+	FRHIDescriptorHandle GetBindlessHandle() const { return BindlessHandle; }
+	bool                 IsBindless()        const { return GetBindlessHandle().IsValid(); }
+#endif
+
+protected:
 	friend class FDynamicRHI;
 
 	// Called only from FDynamicRHI::RHIUpdateTextureReference
@@ -44,8 +49,9 @@ private:
 
 	TRefCountPtr<FRHITexture> ReferencedTexture;
 
-	// SRV to create a reusable slot for bindless handles for this texture reference
-	FShaderResourceViewRHIRef BindlessView;
+#if PLATFORM_SUPPORTS_BINDLESS_RENDERING
+	const FRHIDescriptorHandle BindlessHandle;
+#endif
 
 	// This pointer is set by the InitRHI() function on the FBlackTextureWithSRV global resource,
 	// to allow FRHITextureReference to use the global black texture when the reference is nullptr.
diff --git a/Engine/Source/Runtime/RHICore/Private/RHIDescriptorAllocator.cpp b/Engine/Source/Runtime/RHICore/Private/RHIDescriptorAllocator.cpp
--- a/Engine/Source/Runtime/RHICore/Private/RHIDescriptorAllocator.cpp
+++ b/Engine/Source/Runtime/RHICore/Private/RHIDescriptorAllocator.cpp
@@ -4,13 +4,6 @@
 #include "Misc/ScopeLock.h"
 #include "RHIDefinitions.h"
 
-struct FRHIDescriptorAllocatorRange
-{
-	FRHIDescriptorAllocatorRange(uint32 InFirst, uint32 InLast) : First(InFirst), Last(InLast) {}
-	uint32 First;
-	uint32 Last;
-};
-
 FRHIDescriptorAllocator::FRHIDescriptorAllocator()
 {
 }
@@ -198,6 +191,43 @@ void FRHIDescriptorAllocator::Free(uint32 Offset, uint32 NumDescriptors)
 	}
 }
 
+bool FRHIDescriptorAllocator::GetAllocatedRange(FRHIDescriptorAllocatorRange& OutRange)
+{
+	const uint32 VeryFirstIndex = 0;
+	const uint32 VeryLastIndex = GetCapacity() - 1;
+
+	OutRange.First = VeryFirstIndex;
+	OutRange.Last = VeryLastIndex;
+
+	FScopeLock Lock(&CriticalSection);
+	if (Ranges.Num() > 0)
+	{
+		const FRHIDescriptorAllocatorRange FirstRange = Ranges[0];
+
+		// If the free range matches the entire usable range, that means we have zero allocations.
+		if (FirstRange.First == VeryFirstIndex && FirstRange.Last == VeryLastIndex)
+		{
+			return false;
+		}
+
+		// If the first free range is at the start, then the first allocation is right after this range
+		if (FirstRange.First == VeryFirstIndex)
+		{
+			OutRange.First = FMath::Min(FirstRange.Last + 1, VeryLastIndex);
+		}
+
+		const FRHIDescriptorAllocatorRange LastRange = Ranges[Ranges.Num() - 1];
+
+		// If the last free range is at the end of the usable range, our last allocation is right before this range 
+		if (LastRange.Last == VeryLastIndex)
+		{
+			OutRange.Last = LastRange.First > 0 ? LastRange.First - 1 : 0;
+		}
+	}
+
+	return true;
+}
+
 ////////////////////////////////////////////////////////////////////////////////////////////////////////////////////////
 // FRHIHeapDescriptorAllocator
 
diff --git a/Engine/Source/Runtime/RHICore/Public/RHIDescriptorAllocator.h b/Engine/Source/Runtime/RHICore/Public/RHIDescriptorAllocator.h
--- a/Engine/Source/Runtime/RHICore/Public/RHIDescriptorAllocator.h
+++ b/Engine/Source/Runtime/RHICore/Public/RHIDescriptorAllocator.h
@@ -8,7 +8,12 @@
 enum class ERHIDescriptorHeapType : uint8;
 struct FRHIDescriptorHandle;
 
-struct FRHIDescriptorAllocatorRange;
+struct FRHIDescriptorAllocatorRange
+{
+	FRHIDescriptorAllocatorRange(uint32 InFirst, uint32 InLast) : First(InFirst), Last(InLast) {}
+	uint32 First;
+	uint32 Last;
+};
 
 class FRHIDescriptorAllocator
 {
@@ -26,7 +31,10 @@ public:
 	RHICORE_API bool Allocate(uint32 NumDescriptors, uint32& OutSlot);
 	RHICORE_API void Free(uint32 Slot, uint32 NumDescriptors);
 
-	inline uint32 GetCapacity() const { return Capacity; }
+	// Get the range of allocated descriptors. Useful for determining the smallest range to copy between heaps.
+	RHICORE_API bool GetAllocatedRange(FRHIDescriptorAllocatorRange& OutRange);
+
+	uint32 GetCapacity() const { return Capacity; }
 
 private:
 	void RecordAlloc(uint32 Count)
@@ -59,6 +67,7 @@ private:
 #endif
 };
 
+
 class FRHIHeapDescriptorAllocator : protected FRHIDescriptorAllocator
 {
 public:
@@ -71,6 +80,8 @@ public:
 	RHICORE_API bool Allocate(uint32 NumDescriptors, uint32& OutSlot);
 	RHICORE_API void Free(uint32 Slot, uint32 NumDescriptors);
 
+	using FRHIDescriptorAllocator::GetAllocatedRange;
+
 	using FRHIDescriptorAllocator::GetCapacity;
 	inline ERHIDescriptorHeapType GetType() const { return Type; }
 
diff --git a/Engine/Source/Runtime/VulkanRHI/Private/VulkanTexture.cpp b/Engine/Source/Runtime/VulkanRHI/Private/VulkanTexture.cpp
--- a/Engine/Source/Runtime/VulkanRHI/Private/VulkanTexture.cpp
+++ b/Engine/Source/Runtime/VulkanRHI/Private/VulkanTexture.cpp
@@ -907,15 +907,59 @@ uint32 FVulkanDynamicRHI::RHIComputeMemorySize(FRHITexture* TextureRHI)
 	return ResourceCast(TextureRHI)->GetMemorySize();
 }
 
+class FVulkanTextureReference : public FRHITextureReference
+{
+public:
+	FVulkanTextureReference(FRHITexture* InReferencedTexture)
+		: FRHITextureReference(InReferencedTexture)
+	{
+	}
+
+#if PLATFORM_SUPPORTS_BINDLESS_RENDERING
+	FVulkanTextureReference(FRHITexture* InReferencedTexture, FVulkanShaderResourceView* InBindlessView)
+		: FRHITextureReference(InReferencedTexture, InBindlessView->GetBindlessHandle())
+		, BindlessView(InBindlessView)
+	{
+	}
+
+	TRefCountPtr<FVulkanShaderResourceView> BindlessView;
+#endif
+};
+
+template<>
+struct TVulkanResourceTraits<FRHITextureReference>
+{
+	using TConcreteType = FVulkanTextureReference;
+};
+
+FTextureReferenceRHIRef FVulkanDynamicRHI::RHICreateTextureReference(FRHICommandListBase& RHICmdList, FRHITexture* InReferencedTexture)
+{
+	FRHITexture* ReferencedTexture = InReferencedTexture ? InReferencedTexture : FRHITextureReference::GetDefaultTexture();
+
+#if PLATFORM_SUPPORTS_BINDLESS_RENDERING
+	// If the referenced texture is configured for bindless, make sure we also create an SRV to use for bindless.
+	if (ReferencedTexture && ReferencedTexture->GetDefaultBindlessHandle().IsValid())
+	{
+		FShaderResourceViewRHIRef BindlessView = RHICmdList.CreateShaderResourceView(ReferencedTexture, 0u);
+		return new FVulkanTextureReference(ReferencedTexture, ResourceCast(BindlessView.GetReference()));
+	}
+#endif
+
+	return new FVulkanTextureReference(ReferencedTexture);
+}
+
 void FVulkanDynamicRHI::RHIUpdateTextureReference(FRHICommandListBase& RHICmdList, FRHITextureReference* TextureRef, FRHITexture* InNewTexture)
 {
+	FRHITexture* NewTexture = InNewTexture ? InNewTexture : FRHITextureReference::GetDefaultTexture();
+
+#if PLATFORM_SUPPORTS_BINDLESS_RENDERING
 	if (Device->SupportsBindless())
 	{
-		FRHITexture* NewTexture = InNewTexture ? InNewTexture : FRHITextureReference::GetDefaultTexture();
-
-		if (FRHIShaderResourceView* TextureRefSRV = TextureRef ? TextureRef->GetBindlessView() : nullptr)
+		if (TextureRef && TextureRef->IsBindless())
 		{
-			FVulkanShaderResourceView* VulkanTextureRefSRV = ResourceCast(TextureRefSRV);
+			FVulkanTextureReference* VulkanTextureReference = ResourceCast(TextureRef);
+
+			FVulkanShaderResourceView* VulkanTextureRefSRV = VulkanTextureReference->BindlessView;
 			FRHIDescriptorHandle DestHandle = VulkanTextureRefSRV->GetBindlessHandle();
 
 			if (DestHandle.IsValid())
@@ -939,10 +983,10 @@ void FVulkanDynamicRHI::RHIUpdateTextureReference(FRHICommandListBase& RHICmdLis
 					, !NewVulkanTexture->SupportsSampling());
 			}
 		}
-
 	}
+#endif // PLATFORM_SUPPORTS_BINDLESS_RENDERING
 
-	FDynamicRHI::RHIUpdateTextureReference(RHICmdList, TextureRef, InNewTexture);
+	FDynamicRHI::RHIUpdateTextureReference(RHICmdList, TextureRef, NewTexture);
 }
 
 /*-----------------------------------------------------------------------------
diff --git a/Engine/Source/Runtime/VulkanRHI/Public/VulkanDynamicRHI.h b/Engine/Source/Runtime/VulkanRHI/Public/VulkanDynamicRHI.h
--- a/Engine/Source/Runtime/VulkanRHI/Public/VulkanDynamicRHI.h
+++ b/Engine/Source/Runtime/VulkanRHI/Public/VulkanDynamicRHI.h
@@ -117,6 +117,7 @@ public:
 	virtual void* LockBuffer_BottomOfPipe(FRHICommandListBase& RHICmdList, FRHIBuffer* Buffer, uint32 Offset, uint32 Size, EResourceLockMode LockMode) final override;
 	virtual void UnlockBuffer_BottomOfPipe(FRHICommandListBase& RHICmdList, FRHIBuffer* Buffer) final override;
 	virtual void RHIUnlockBuffer(FRHICommandListBase& RHICmdList, FRHIBuffer* Buffer) final override;
+	virtual FTextureReferenceRHIRef RHICreateTextureReference(FRHICommandListBase& RHICmdList, FRHITexture* InReferencedTexture) final override;
 	virtual void RHIUpdateTextureReference(FRHICommandListBase& RHICmdList, FRHITextureReference* TextureRef, FRHITexture* NewTexture) final override;
 	virtual FRHICalcTextureSizeResult RHICalcTexturePlatformSize(FRHITextureDesc const& Desc, uint32 FirstMipIndex) final override;
 	virtual void RHIGetTextureMemoryStats(FTextureMemoryStats& OutStats) final override;
