Fix crash due to null residency handle pointer dereference when residency is not being tracked for a resource

diff --git a/Engine/Source/Runtime/D3D12RHI/Private/D3D12RayTracing.cpp b/Engine/Source/Runtime/D3D12RHI/Private/D3D12RayTracing.cpp
--- a/Engine/Source/Runtime/D3D12RHI/Private/D3D12RayTracing.cpp
+++ b/Engine/Source/Runtime/D3D12RHI/Private/D3D12RayTracing.cpp
@@ -3598,7 +3598,7 @@ void FD3D12RayTracingScene::BuildAccelerationStructure(FD3D12CommandContext& Com
 			#if ENABLE_RESIDENCY_MANAGEMENT
 				for (FD3D12ResidencyHandle* ResidencyHandle : Resource->GetResidencyHandles())
 				{
-					if (D3DX12Residency::IsInitialized(*ResidencyHandle))
+					if (D3DX12Residency::IsInitialized(ResidencyHandle))
 					{
 						bool bIsAlreadyInSet = false;
 						UniqueResidencyHandles.Add(ResidencyHandle, &bIsAlreadyInSet);
