Add bRequiresResidencyTracking flag to FD3D12Heap and FD3D12Resource and use it to validate residency tracking status

* Fix null residency handle access in some cases

diff --git a/Engine/Source/Runtime/D3D12RHI/Private/D3D12PoolAllocator.cpp b/Engine/Source/Runtime/D3D12RHI/Private/D3D12PoolAllocator.cpp
--- a/Engine/Source/Runtime/D3D12RHI/Private/D3D12PoolAllocator.cpp
+++ b/Engine/Source/Runtime/D3D12RHI/Private/D3D12PoolAllocator.cpp
@@ -425,8 +425,9 @@ void FD3D12PoolAllocator::AllocateResource(uint32 GPUIndex, D3D12_HEAP_TYPE InHe
 			VERIFYD3D12RESULT(Adapter->GetD3DDevice()->CreateHeap(&HeapDesc, IID_PPV_ARGS(&Heap)));
 			TRefCountPtr<FD3D12Heap> BackingHeap = new FD3D12Heap(GetParentDevice(), GetVisibilityMask(), TraceHeapId);
 			bool bTrack = false;
-			BackingHeap->SetHeap(Heap, InName, bTrack);		
-			
+			BackingHeap->SetHeap(Heap, InName, bTrack);
+			BackingHeap->BeginTrackingResidency(HeapDesc.SizeInBytes);
+
 			VERIFYD3D12RESULT(Adapter->CreatePlacedResource(Desc, BackingHeap, 0, InCreateState, InResourceStateMode, InCreateState, InClearValue, &NewResource, InName));
 		}
 		else
diff --git a/Engine/Source/Runtime/D3D12RHI/Private/D3D12Resources.cpp b/Engine/Source/Runtime/D3D12RHI/Private/D3D12Resources.cpp
--- a/Engine/Source/Runtime/D3D12RHI/Private/D3D12Resources.cpp
+++ b/Engine/Source/Runtime/D3D12RHI/Private/D3D12Resources.cpp
@@ -92,6 +92,7 @@ FD3D12Resource::FD3D12Resource(FD3D12Device* ParentDevice,
 	, HeapType(InHeapType)
 	, PlaneCount(UE::DXGIUtilities::GetPlaneCount(Desc.Format))
 	, bRequiresResourceStateTracking(true)
+	, bRequiresResidencyTracking(true)
 	, bDepthStencil(false)
 	, bDeferDelete(true)
 	, bBackBuffer(false)
@@ -100,6 +101,17 @@ FD3D12Resource::FD3D12Resource(FD3D12Device* ParentDevice,
 	FPlatformAtomics::InterlockedIncrement(&TotalResourceCount);
 #endif
 
+	D3D12_HEAP_DESC HeapDesc = {};
+	D3D12_HEAP_PROPERTIES* HeapProps = nullptr;
+	if (InHeap)
+	{
+		HeapDesc = InHeap->GetHeapDesc();
+		HeapProps = &HeapDesc.Properties;
+	}
+
+	// Residency tracking is not used for CPU-accessible resources or back buffers
+	bRequiresResidencyTracking = IsGPUOnly(InHeapType, HeapProps) && !bBackBuffer;
+
 	// On Windows it's sadly enough not possible to get the GPU virtual address from the resource directly
 	if (Resource
 #if PLATFORM_WINDOWS
@@ -360,6 +372,9 @@ ID3D12Pageable* FD3D12Resource::GetPageable()
 void FD3D12Resource::StartTrackingForResidency()
 {
 #if ENABLE_RESIDENCY_MANAGEMENT
+
+	checkf(bRequiresResidencyTracking, TEXT("Residency tracking is not expected for this resource"));
+
 	if (bBackBuffer)
 	{
 		// Back buffers may be referenced outside of command lists (during presents), however D3DX12Residency.h library 
@@ -441,6 +456,8 @@ void FD3D12Heap::SetHeap(ID3D12Heap* HeapIn, const TCHAR* const InName, bool bIn
 
 	SetName(HeapIn, InName);
 
+	bRequiresResidencyTracking = IsGPUOnly(HeapDesc.Properties.Type, &HeapDesc.Properties);
+
 	// Create a buffer placed resource on the heap to extract the gpu virtual address
 	// if we are tracking all allocations
 	FD3D12Adapter* Adapter = GetParentDevice()->GetParentAdapter();	
@@ -473,9 +490,16 @@ void FD3D12Heap::SetHeap(ID3D12Heap* HeapIn, const TCHAR* const InName, bool bIn
 	}
 }
 
+void FD3D12Heap::DisallowTrackingResidency()
+{
+	checkf(ResidencyHandle == nullptr, TEXT("Can't disallow residency tracking after it has started. Call this function instead of BeginTrackingResidency()."));
+	bRequiresResidencyTracking = false;
+}
+
 void FD3D12Heap::BeginTrackingResidency(uint64 Size)
 {
 #if ENABLE_RESIDENCY_MANAGEMENT
+	checkf(bRequiresResidencyTracking, TEXT("Residency tracking is not expected for this resource"));
 	checkf(!ResidencyHandle, TEXT("Residency tracking is already initialzied for this resource"));
 	ResidencyHandle = new FD3D12ResidencyHandle;
 	D3DX12Residency::Initialize(*ResidencyHandle, Heap.GetReference(), Size, this);
diff --git a/Engine/Source/Runtime/D3D12RHI/Private/D3D12Resources.h b/Engine/Source/Runtime/D3D12RHI/Private/D3D12Resources.h
--- a/Engine/Source/Runtime/D3D12RHI/Private/D3D12Resources.h
+++ b/Engine/Source/Runtime/D3D12RHI/Private/D3D12Resources.h
@@ -79,12 +79,24 @@ public:
 	void SetHeap(ID3D12Heap* HeapIn, const TCHAR* const InName, bool bTrack = true, bool bForceGetGPUAddress = false);
 
 	void BeginTrackingResidency(uint64 Size);
+	void DisallowTrackingResidency(); // Part of workaround for UE-174791
 
 	void DeferDelete();
 
 	inline FName GetName() const { return HeapName; }
 	inline D3D12_HEAP_DESC GetHeapDesc() const { return HeapDesc; }
-	inline TConstArrayView<FD3D12ResidencyHandle*> GetResidencyHandles() { return MakeArrayView(&ResidencyHandle, 1); }
+	inline TConstArrayView<FD3D12ResidencyHandle*> GetResidencyHandles()
+	{ 
+		if (bRequiresResidencyTracking)
+		{
+			checkf(ResidencyHandle, TEXT("Resource requires residency tracking, but BeginTrackingResidency() was not called."));
+			return MakeArrayView(&ResidencyHandle, 1); 
+		}
+		else
+		{
+			return {};
+		}		
+	}
 	inline D3D12_GPU_VIRTUAL_ADDRESS GetGPUVirtualAddress() const { return GPUVirtualAddress; }
 	inline void SetIsTransient(bool bInIsTransient) { bIsTransient = bInIsTransient; }
 	inline bool GetIsTransient() const { return bIsTransient; }
@@ -100,6 +112,7 @@ private:
 	HeapId TraceHeapId;
 	HeapId TraceParentHeapId;
 	bool bIsTransient = false; // Whether this is a transient heap
+	bool bRequiresResidencyTracking = true;
 };
 
 struct FD3D12ResourceDesc : public D3D12_RESOURCE_DESC
@@ -174,6 +187,7 @@ private:
 	uint16 SubresourceCount{};
 	uint8 PlaneCount;
 	bool bRequiresResourceStateTracking : 1;
+	bool bRequiresResidencyTracking : 1;
 	bool bDepthStencil : 1;
 	bool bDeferDelete : 1;
 	bool bBackBuffer : 1;
@@ -277,7 +291,16 @@ public:
 	bool RequiresResourceStateTracking() const { return bRequiresResourceStateTracking; }
 
 	inline bool IsBackBuffer() const { return bBackBuffer; }
-	inline void SetIsBackBuffer(bool bBackBufferIn) { bBackBuffer = bBackBufferIn; }
+	inline void SetIsBackBuffer(bool bBackBufferIn) 
+	{
+		bBackBuffer = bBackBufferIn;
+
+		if (bBackBuffer)
+		{
+			checkf(ResidencyHandle == nullptr, TEXT("Can't mark a resource as back buffer once residency tracking has started."));
+			bRequiresResidencyTracking = false;
+		}
+	}
 
 	void SetName(const TCHAR* Name)
 	{
@@ -313,9 +336,16 @@ public:
 	bool IsResident() const
 	{
 #if ENABLE_RESIDENCY_MANAGEMENT
+		TConstArrayView<FD3D12ResidencyHandle*> ResidencyHandles = GetResidencyHandles();
+		if (ResidencyHandles.IsEmpty())
+		{
+			// No residency tracking for this resource
+			return true;
+		}
+
 		// Treat resource as resident if at least one backing heap is resident.
 		// Technically we should return a partial residency status.
-		for (FD3D12ResidencyHandle* Handle : GetResidencyHandles())
+		for (FD3D12ResidencyHandle* Handle : ResidencyHandles)
 		{
 			if (Handle->ResidencyStatus == FD3D12ResidencyHandle::RESIDENCY_STATUS::RESIDENT)
 			{
@@ -331,7 +361,11 @@ public:
 	TConstArrayView<FD3D12ResidencyHandle*> GetResidencyHandles() const
 	{
 #if ENABLE_RESIDENCY_MANAGEMENT
-		if (IsPlacedResource())
+		if (!bRequiresResidencyTracking)
+		{
+			return {};
+		}
+		else if (IsPlacedResource())
 		{
 			return Heap->GetResidencyHandles();
 		}
@@ -341,6 +375,7 @@ public:
 		}
 		else
 		{
+			checkf(ResidencyHandle, TEXT("Resource requires residency tracking, but StartTrackingForResidency() was not called."));
 			return MakeArrayView(&ResidencyHandle, 1);
 		}
 #else // ENABLE_RESIDENCY_MANAGEMENT
diff --git a/Engine/Source/Runtime/D3D12RHI/Private/D3D12TransientResourceAllocator.cpp b/Engine/Source/Runtime/D3D12RHI/Private/D3D12TransientResourceAllocator.cpp
--- a/Engine/Source/Runtime/D3D12RHI/Private/D3D12TransientResourceAllocator.cpp
+++ b/Engine/Source/Runtime/D3D12RHI/Private/D3D12TransientResourceAllocator.cpp
@@ -103,6 +103,7 @@ FD3D12TransientHeap::FD3D12TransientHeap(const FInitializer& Initializer, FD3D12
 	// when they're no longer needed. Disabling residency tracking will not fix the underlying bug, but should make it less likely to occur,
 	// and might make the GPU crash data more useful when it does happen.
 	//Heap->BeginTrackingResidency(Desc.SizeInBytes);
+	Heap->DisallowTrackingResidency(); // Remove this when the above workaround is not needed
 
 	SetGpuVirtualAddress(Heap->GetGPUVirtualAddress());
 
diff --git a/Engine/Source/Runtime/D3D12RHI/Private/D3D12Viewport.cpp b/Engine/Source/Runtime/D3D12RHI/Private/D3D12Viewport.cpp
--- a/Engine/Source/Runtime/D3D12RHI/Private/D3D12Viewport.cpp
+++ b/Engine/Source/Runtime/D3D12RHI/Private/D3D12Viewport.cpp
@@ -218,7 +218,6 @@ FD3D12Texture* GetSwapChainSurface(FD3D12Device* Parent, EPixelFormat PixelForma
 			FD3D12Resource* NewResourceWrapper = new FD3D12Resource(Device, FRHIGPUMask::All(), BackBufferResource, InitialState, BackBufferDesc);
 			NewResourceWrapper->SetIsBackBuffer(true);
 			NewResourceWrapper->AddRef();
-			NewResourceWrapper->StartTrackingForResidency();
 			NewTexture->ResourceLocation.AsStandAlone(NewResourceWrapper);
 		}
 		else // If this is not the GPU which will hold the back buffer, create a compatible texture so that it can still render to the viewport.
