Disable residency tracking for D3D12 reserved resources

* Temporary workaround for lack of support for reserved resources in D3DX12 Residency library

diff --git a/Engine/Source/Runtime/D3D12RHI/Private/D3D12Resources.cpp b/Engine/Source/Runtime/D3D12RHI/Private/D3D12Resources.cpp
--- a/Engine/Source/Runtime/D3D12RHI/Private/D3D12Resources.cpp
+++ b/Engine/Source/Runtime/D3D12RHI/Private/D3D12Resources.cpp
@@ -344,7 +344,8 @@ void FD3D12Resource::CommitReservedResource(ID3D12CommandQueue* D3DCommandQueue,
 
 		TRefCountPtr<FD3D12Heap> NewHeap = new FD3D12Heap(GetParentDevice(), GetVisibilityMask());
 		NewHeap->SetHeap(D3DHeap, HeapNameChars, true /*bTrack*/, false /*bForceGetGPUAddress*/);
-		NewHeap->BeginTrackingResidency(ThisHeapSize);
+		//NewHeap->BeginTrackingResidency(ThisHeapSize);
+		NewHeap->DisallowTrackingResidency(); // Workaround for UE-202367: D3DX12Residency library does not track UpdateTileMappings that may be in flight
 
 		ReservedResourceData->ResidencyHandles.Append(NewHeap->GetResidencyHandles());
 		ReservedResourceData->BackingHeaps.Add(MoveTemp(NewHeap));
diff --git a/Engine/Source/Runtime/D3D12RHI/Private/D3D12Resources.h b/Engine/Source/Runtime/D3D12RHI/Private/D3D12Resources.h
--- a/Engine/Source/Runtime/D3D12RHI/Private/D3D12Resources.h
+++ b/Engine/Source/Runtime/D3D12RHI/Private/D3D12Resources.h
@@ -79,7 +79,7 @@ public:
 	void SetHeap(ID3D12Heap* HeapIn, const TCHAR* const InName, bool bTrack = true, bool bForceGetGPUAddress = false);
 
 	void BeginTrackingResidency(uint64 Size);
-	void DisallowTrackingResidency(); // Part of workaround for UE-174791
+	void DisallowTrackingResidency(); // Part of workaround for UE-174791 and UE-202367
 
 	void DeferDelete();
 
