Implement support for shrinking reserved resource memory allocations in D3D12 RHI

* Release backing heaps when they become unreferenced (contain no committed tiles)
* Use available slack space in the backing heaps before allocating new heaps
* Track committed / free memory in ReservedResourceData in terms of tiles instead of bytes
* Add unit test to cover the basic functionality and a few edge cases
* Perform AddRef() on the heap directly in FD3D12Heap::DeferDelete() instead of requiring caller beforehand
* Add RHI stat counter for physical memory used by reserved resources

diff --git a/Engine/Plugins/Tests/RHITests/Source/RHITests/Private/RHIReservedResourceTests.cpp b/Engine/Plugins/Tests/RHITests/Source/RHITests/Private/RHIReservedResourceTests.cpp
--- a/Engine/Plugins/Tests/RHITests/Source/RHITests/Private/RHIReservedResourceTests.cpp
+++ b/Engine/Plugins/Tests/RHITests/Source/RHITests/Private/RHIReservedResourceTests.cpp
@@ -139,6 +139,51 @@ static void CommitBuffer(FRHICommandListImmediate& RHICmdList, FRHIBuffer* Buffe
 	RHICmdList.EndTransition(Transition);
 }
 
+struct FBufferVerificationRange
+{
+	// Offsets in bytes, uint32-aligned
+	uint32 BeginByteOffset = 0;
+	uint32 EndByteOffset = 0;
+
+	uint32 ExpectedValue = 0;
+};
+
+static bool VerifyBufferRanges(const TCHAR* TestName, FRHICommandListImmediate& RHICmdList, FRHIBuffer* Buffer, TConstArrayView<FBufferVerificationRange> Ranges)
+{
+	return FRHIBufferTests::VerifyBufferContents(TestName, RHICmdList, MakeArrayView(&Buffer, 1),
+		[Ranges](int32 /*BufferIndex*/, void* Ptr, uint32 NumBytes)
+		{
+			const uint32 Stride = sizeof(FBufferVerificationRange::ExpectedValue);
+			const uint32 BufferNumElements = NumBytes / Stride;
+			const uint32* BufferData = reinterpret_cast<const uint32*>(Ptr);
+
+			for (const FBufferVerificationRange& Range : Ranges)
+			{
+				check(Range.BeginByteOffset % Stride == 0);
+				check(Range.EndByteOffset % Stride == 0);
+
+				const uint32 ElemBegin = Range.BeginByteOffset / Stride;
+				const uint32 ElemEnd = Range.EndByteOffset / Stride;
+
+				if (ElemBegin >= BufferNumElements || ElemEnd > BufferNumElements)
+				{
+					return false;
+				}
+
+				for (uint32 i = ElemBegin; i < ElemEnd; ++i)
+				{
+					if (BufferData[i] != Range.ExpectedValue)
+					{
+						return false;
+					}
+				}
+			}
+
+			return true;
+		});
+}
+
+
 bool FRHIReservedResourceTests::Test_ReservedResource_CommitBuffer(FRHICommandListImmediate& RHICmdList)
 {
 	if (!GRHIGlobals.ReservedResources.Supported)
@@ -168,39 +213,100 @@ bool FRHIReservedResourceTests::Test_ReservedResource_CommitBuffer(FRHICommandLi
 	const int32 CommitSizeInBytes = BufferSizeInBytes / 2;
 	CommitBuffer(RHICmdList, Buffer, CommitSizeInBytes, ERHIAccess::UAVCompute, ERHIAccess::UAVCompute);
 
-	RHICmdList.ClearUAVUint(BufferUAV, FUintVector4(~0u));
+	const uint32 ClearValue = ~0u;
+	RHICmdList.ClearUAVUint(BufferUAV, FUintVector4(ClearValue));
 
 	RHICmdList.Transition(FRHITransitionInfo(Buffer, ERHIAccess::UAVCompute, ERHIAccess::CopySrc));
-	
+
+	bool bSucceeded = VerifyBufferRanges(TEXT("Test_ReservedResource_CommitBuffer"), RHICmdList, Buffer,
+		{
+			{0,                 CommitSizeInBytes, ClearValue}, // First range is committed and expected to contain the value written by ClearUAVUint
+			{CommitSizeInBytes, BufferSizeInBytes, 0u} // We follow the D3D convention for unmapped page access: writes are no-op, reads return 0
+		});
+
+	return bSucceeded;
+}
+
+bool FRHIReservedResourceTests::Test_ReservedResource_DecommitBuffer(FRHICommandListImmediate& RHICmdList)
+{
+	if (!GRHIGlobals.ReservedResources.Supported)
+	{
+		return true;
+	}
+
+	const int32 TileSizeInBytes = GRHIGlobals.ReservedResources.TileSizeInBytes;
+	const int32 BufferSizeInBytes = TileSizeInBytes * 128;
+
+	FRHIResourceCreateInfo CreateInfo(TEXT("TestReservedBufferExplicitCommit"));
+
+	FBufferRHIRef Buffer = RHICmdList.CreateBuffer(BufferSizeInBytes,
+		BUF_ReservedResource | BUF_UnorderedAccess | BUF_ShaderResource | BUF_SourceCopy,
+		4, ERHIAccess::UAVCompute, CreateInfo);
+
+	FUnorderedAccessViewRHIRef BufferUAV = RHICmdList.CreateUnorderedAccessView(Buffer,
+		FRHIViewDesc::CreateBufferUAV()
+		.SetType(FRHIViewDesc::EBufferType::Typed)
+		.SetFormat(PF_R32_UINT));
+
+	// Commit the entire buffer using N commands (to force multiple backing physical memory allocations) and fill it with ~0u
+
+	CommitBuffer(RHICmdList, Buffer, BufferSizeInBytes / 4, ERHIAccess::UAVCompute, ERHIAccess::UAVCompute);
+	CommitBuffer(RHICmdList, Buffer, BufferSizeInBytes / 2, ERHIAccess::UAVCompute, ERHIAccess::UAVCompute);
+	CommitBuffer(RHICmdList, Buffer, BufferSizeInBytes, ERHIAccess::UAVCompute, ERHIAccess::UAVCompute);
+
+	const uint32 ClearValue = ~0u;
+	RHICmdList.ClearUAVUint(BufferUAV, FUintVector4(ClearValue));
+	RHICmdList.Transition(FRHITransitionInfo(Buffer, ERHIAccess::UAVCompute, ERHIAccess::CopySrc));
+
 	FRHIBuffer* Buffers[] = { Buffer.GetReference() };
-	bool bSucceeded = FRHIBufferTests::VerifyBufferContents(TEXT("Test_ReservedResource_CommitBuffer"), RHICmdList, Buffers, 
-		[BufferSizeInBytes, CommitSizeInBytes](int32 BufferIndex, void* Ptr, uint32 NumBytes)
+
+	// Check that the clear has fully initialzied the buffer contents
+
+	bool bSucceeded = false;
+
+	bSucceeded = VerifyBufferRanges(TEXT("Test_ReservedResource_DecommitBuffer_Init"), RHICmdList, Buffer,
 		{
-			uint64 ExpectedCommittedValue = ~0ull;
+			{0, BufferSizeInBytes, ClearValue}
+		});
 
-			uint32 CommittedSizeInElements = CommitSizeInBytes / sizeof(ExpectedCommittedValue);
-			uint32 TotalSizeInElements = BufferSizeInBytes / sizeof(ExpectedCommittedValue);
-			const uint64* BufferData = reinterpret_cast<const uint64*>(Ptr);
+	if (!bSucceeded)
+	{
+		return false;
+	}
 
-			for (uint32 i = 0; i < CommittedSizeInElements; ++i)
-			{
-				if (BufferData[i] != ExpectedCommittedValue)
-				{
-					return false;
-				}
-			}
+	// Partially decommit the buffer using multiple operations, leaving only one tile mapped. Reading from decommited regions is expected to return 0.
+	CommitBuffer(RHICmdList, Buffer, BufferSizeInBytes - TileSizeInBytes, ERHIAccess::CopySrc, ERHIAccess::UAVCompute); // shrink
+	CommitBuffer(RHICmdList, Buffer, TileSizeInBytes * 2, ERHIAccess::UAVCompute, ERHIAccess::UAVCompute); // shrink
+	CommitBuffer(RHICmdList, Buffer, TileSizeInBytes * 3, ERHIAccess::UAVCompute, ERHIAccess::UAVCompute); // grow
+	CommitBuffer(RHICmdList, Buffer, TileSizeInBytes, ERHIAccess::UAVCompute, ERHIAccess::UAVCompute); // shrink to final size
 
-			// We follow the D3D convention for unmapped page access: writes are no-op, reads return 0
-			const uint64 ExpectedTailValue = 0;
-			for (uint32 i = CommittedSizeInElements; i < TotalSizeInElements; ++i)
-			{
-				if (BufferData[i] != ExpectedTailValue)
-				{
-					return false;
-				}
-			}
+	RHICmdList.ClearUAVUint(BufferUAV, FUintVector4(ClearValue)); // Expected to have no effect on contents of the buffer except the first tile
+	RHICmdList.Transition(FRHITransitionInfo(Buffer, ERHIAccess::UAVCompute, ERHIAccess::CopySrc));
 
-			return true;
+	bSucceeded = VerifyBufferRanges(TEXT("Test_ReservedResource_DecommitBuffer_Partial"), RHICmdList, Buffer,
+		{
+			{0,               TileSizeInBytes, ClearValue}, // First tile expected to retain value
+			{TileSizeInBytes, BufferSizeInBytes, 0u} // The rest of the buffer is expected to read 0s
+		});
+
+	if (!bSucceeded)
+	{
+		return false;
+	}
+
+	// Re-commit the buffer and verify that it is fully cleared again
+
+	CommitBuffer(RHICmdList, Buffer, BufferSizeInBytes / 2, ERHIAccess::CopySrc, ERHIAccess::UAVCompute);
+	CommitBuffer(RHICmdList, Buffer, BufferSizeInBytes, ERHIAccess::UAVCompute, ERHIAccess::UAVCompute);
+
+	const uint32 ClearValue2 = 0x11223344;
+
+	RHICmdList.ClearUAVUint(BufferUAV, FUintVector4(ClearValue2));
+	RHICmdList.Transition(FRHITransitionInfo(Buffer, ERHIAccess::UAVCompute, ERHIAccess::CopySrc));
+
+	bSucceeded = VerifyBufferRanges(TEXT("Test_ReservedResource_DecommitBuffer_Init"), RHICmdList, Buffer,
+		{
+			{0, BufferSizeInBytes, ClearValue2}
 		});
 
 	return bSucceeded;
diff --git a/Engine/Plugins/Tests/RHITests/Source/RHITests/Private/RHITests.spec.cpp b/Engine/Plugins/Tests/RHITests/Source/RHITests/Private/RHITests.spec.cpp
--- a/Engine/Plugins/Tests/RHITests/Source/RHITests/Private/RHITests.spec.cpp
+++ b/Engine/Plugins/Tests/RHITests/Source/RHITests/Private/RHITests.spec.cpp
@@ -144,5 +144,11 @@ void FAutomationRHITest::Define()
 			bool bResult = RunOnRenderThreadSynchronous(FRHIReservedResourceTests::Test_ReservedResource_CommitBuffer);
 			TestEqual("Commit Reserved Buffer failed", bResult, 1);
 		});
+
+		It("Decommit Reserved Buffer", [this]()
+		{
+			bool bResult = RunOnRenderThreadSynchronous(FRHIReservedResourceTests::Test_ReservedResource_DecommitBuffer);
+			TestEqual("Decommit Reserved Buffer failed", bResult, 1);
+		});
 	});
 }
diff --git a/Engine/Plugins/Tests/RHITests/Source/RHITests/Private/RHITestsModule.cpp b/Engine/Plugins/Tests/RHITests/Source/RHITests/Private/RHITestsModule.cpp
--- a/Engine/Plugins/Tests/RHITests/Source/RHITests/Private/RHITestsModule.cpp
+++ b/Engine/Plugins/Tests/RHITests/Source/RHITests/Private/RHITestsModule.cpp
@@ -22,6 +22,7 @@ static bool RunTests_RenderThread(FRHICommandListImmediate& RHICmdList)
 	// ------------------------------------------------
 	{
 		RUN_TEST(FRHIReservedResourceTests::Test_ReservedResource_CommitBuffer(RHICmdList));
+		RUN_TEST(FRHIReservedResourceTests::Test_ReservedResource_DecommitBuffer(RHICmdList));
 		RUN_TEST(FRHIReservedResourceTests::Test_ReservedResource_CreateBuffer(RHICmdList));
 		RUN_TEST(FRHIReservedResourceTests::Test_ReservedResource_CreateTexture(RHICmdList));
 	}
diff --git a/Engine/Plugins/Tests/RHITests/Source/RHITests/Public/RHIReservedResourceTests.h b/Engine/Plugins/Tests/RHITests/Source/RHITests/Public/RHIReservedResourceTests.h
--- a/Engine/Plugins/Tests/RHITests/Source/RHITests/Public/RHIReservedResourceTests.h
+++ b/Engine/Plugins/Tests/RHITests/Source/RHITests/Public/RHIReservedResourceTests.h
@@ -11,5 +11,6 @@ public:
 	static bool Test_ReservedResource_CreateTexture(FRHICommandListImmediate& RHICmdList);
 	static bool Test_ReservedResource_CreateBuffer(FRHICommandListImmediate& RHICmdList);
 	static bool Test_ReservedResource_CommitBuffer(FRHICommandListImmediate& RHICmdList);
+	static bool Test_ReservedResource_DecommitBuffer(FRHICommandListImmediate& RHICmdList);
 };
 
diff --git a/Engine/Source/Runtime/D3D12RHI/Private/D3D12Resources.cpp b/Engine/Source/Runtime/D3D12RHI/Private/D3D12Resources.cpp
--- a/Engine/Source/Runtime/D3D12RHI/Private/D3D12Resources.cpp
+++ b/Engine/Source/Runtime/D3D12RHI/Private/D3D12Resources.cpp
@@ -186,6 +186,7 @@ void FD3D12Resource::CommitReservedResource(ID3D12CommandQueue* D3DCommandQueue,
 	D3D12_PACKED_MIP_INFO PackedMipDesc = {};
 	D3D12_TILE_SHAPE TileShape = {};
 	const uint32 FirstSubresource = 0;
+	const uint32 NumSubresources = SubresourceCount;
 
 	// We assume that all subresources in a 2D texture array are identical, so only query the tiling config for the first
 	uint32 NumSubresourceTilings = 1;
@@ -235,69 +236,33 @@ void FD3D12Resource::CommitReservedResource(ID3D12CommandQueue* D3DCommandQueue,
 	{
 		checkf(SubresourceTiling.DepthInTiles == 1, TEXT("3D reserved textures are not supported/implemented"));
 		NumStandardTilesPerSubresource = SubresourceTiling.WidthInTiles * SubresourceTiling.HeightInTiles;
-		NumTotalTiles = NumStandardTilesPerSubresource * SubresourceCount;
+		NumTotalTiles = NumStandardTilesPerSubresource * NumSubresources;
 	}
 	else // packed mip case
 	{
 		checkf(PackedMipDesc.NumPackedMips != 0, TEXT("It is expected that reserved resources have at least one standard or one packed tile mip level"));
-		NumTotalTiles = PackedMipDesc.NumTilesForPackedMips * SubresourceCount;
+		NumTotalTiles = PackedMipDesc.NumTilesForPackedMips * NumSubresources;
 	}
 
 	checkf(D3DResourceNumTiles == NumTotalTiles,
 		TEXT("D3D resource size in tiles: %d, computed size in tiles: %d"),
 		D3DResourceNumTiles, NumTotalTiles);
 
-	checkf(RequiredCommitSizeInBytes % TileSizeInBytes == 0,
-		TEXT("Reserved resources memory is expected to be committed at tile granularity"));
-
-	checkf(ReservedResourceData->CommittedSizeInBytes % TileSizeInBytes == 0,
-		TEXT("Reserved resources memory is expected to be committed at tile granularity"));
-
-	checkf(RequiredCommitSizeInBytes >= ReservedResourceData->CommittedSizeInBytes,
-		TEXT("Shrinking committed resources is not yet implemented")); // #yuriy_todo
-
 	const uint32 NumRequiredCommitTiles = RequiredCommitSizeInBytes / TileSizeInBytes;
-	checkf(NumRequiredCommitTiles <= NumTotalTiles,
-		TEXT("Shrinking committed resources is not yet implemented")); // #yuriy_todo
-
-	uint32 NumMappedTiles = ReservedResourceData->CommittedSizeInBytes / TileSizeInBytes;
 
-	while (NumMappedTiles < NumRequiredCommitTiles)
+	auto GetTiledResourceCoordinate = [SubresourceTiling, NumStandardTilesPerSubresource, NumSubresources, MaxTilesPerHeap]
+		(uint32 OffsetInTiles, uint32 NumTiles) -> D3D12_TILED_RESOURCE_COORDINATE 
 	{
-		const uint32 NumRemainingTiles = NumRequiredCommitTiles - NumMappedTiles;
-
-		D3D12_TILE_REGION_SIZE RegionSize = {};
-		RegionSize.UseBox = false;
-		RegionSize.NumTiles = FMath::Min(MaxTilesPerHeap, NumRemainingTiles);
-
-		const D3D12_HEAP_FLAGS HeapFlags = Desc.Dimension == D3D12_RESOURCE_DIMENSION_BUFFER
-			? D3D12_HEAP_FLAG_ALLOW_ONLY_BUFFERS
-			: D3D12_HEAP_FLAG_ALLOW_ONLY_NON_RT_DS_TEXTURES;
-
-		static_assert((D3D12_HEAP_FLAG_DENY_BUFFERS | D3D12_HEAP_FLAG_DENY_RT_DS_TEXTURES) == D3D12_HEAP_FLAG_ALLOW_ONLY_NON_RT_DS_TEXTURES);
-
-		const uint32 ThisHeapSize = RegionSize.NumTiles * TileSizeInBytes;
-		D3D12_HEAP_DESC NewHeapDesc = {};
-		NewHeapDesc.Alignment = D3D12_DEFAULT_RESOURCE_PLACEMENT_ALIGNMENT;
-		NewHeapDesc.Flags = HeapFlags;
-		NewHeapDesc.SizeInBytes = ThisHeapSize;
-		NewHeapDesc.Properties = BackingHeapProps;
-		ID3D12Heap* D3DHeap = nullptr;
-		VERIFYD3D12RESULT(D3DDevice->CreateHeap(&NewHeapDesc, IID_PPV_ARGS(&D3DHeap)));
-
-		if (bHighPriorityResource)
-		{
-			Adapter->SetResidencyPriority(D3DHeap, D3D12_RESIDENCY_PRIORITY_HIGH, GPUIndex);
-		}
-
 		D3D12_TILED_RESOURCE_COORDINATE ResourceCoordinate = {}; // Coordinates are in tiles, not pixels
 		if (NumStandardTilesPerSubresource)
 		{
-			const uint32 TileIndexInSubresource = NumMappedTiles % NumStandardTilesPerSubresource;
-			ResourceCoordinate.Subresource = NumMappedTiles / NumStandardTilesPerSubresource;
+			const uint32 TileIndexInSubresource = OffsetInTiles % NumStandardTilesPerSubresource;
 			ResourceCoordinate.X = TileIndexInSubresource % SubresourceTiling.WidthInTiles;
 			ResourceCoordinate.Y = (TileIndexInSubresource / SubresourceTiling.WidthInTiles) % SubresourceTiling.HeightInTiles;
 			ResourceCoordinate.Z = 0; // Only simple 2D / Array2D textures are implemented
+
+			ResourceCoordinate.Subresource = OffsetInTiles / NumStandardTilesPerSubresource;
+			check(ResourceCoordinate.Subresource <= NumSubresources);
 		}
 		else
 		{
@@ -305,14 +270,14 @@ void FD3D12Resource::CommitReservedResource(ID3D12CommandQueue* D3DCommandQueue,
 			// - Only simple 2D textures are expected (single subresource, no arrays)
 			// - Entire packed mip level must be covered in one map operation, so mapping origin is always 0
 
-			checkf(SubresourceCount == 1,
-				TEXT("Reserved textures with packed mips and multiple subresources are not supported. Current subresource count: %d"),
-				SubresourceCount);
+			checkf(NumSubresources == 1,
+			       TEXT("Reserved textures with packed mips and multiple subresources are not supported. Current subresource count: %d"),
+			       NumSubresources);
 
-			checkf(RegionSize.NumTiles <= MaxTilesPerHeap,
-				TEXT("Reserved texture packed mip level requires tiles: %d, maximum supported tiles: %d. ")
-				TEXT("Increase d3d12.ReservedResourceHeapSizeMB or avoid packed mips by using a larger texture dimensions."),
-				RegionSize.NumTiles, MaxTilesPerHeap);
+			checkf(NumTiles <= MaxTilesPerHeap,
+			       TEXT("Reserved texture packed mip level requires tiles: %d, maximum supported tiles: %d. ")
+			       TEXT("Increase d3d12.ReservedResourceHeapSizeMB or avoid packed mips by using a larger texture dimensions."),
+			       NumTiles, MaxTilesPerHeap);
 
 			ResourceCoordinate.Subresource = 0; // Packed mips are currently supported for single subresource textures
 			ResourceCoordinate.X = 0;
@@ -320,42 +285,158 @@ void FD3D12Resource::CommitReservedResource(ID3D12CommandQueue* D3DCommandQueue,
 			ResourceCoordinate.Z = 0;
 		}
 
-		const D3D12_TILE_RANGE_FLAGS RangeFlags = D3D12_TILE_RANGE_FLAG_NONE;
-		const uint32 HeapRangeStartOffsetInTiles = 0;
-		const uint32 RangeTileCount = RegionSize.NumTiles;
+		return ResourceCoordinate;
+	};
+
+	if (ReservedResourceData->NumCommittedTiles > NumRequiredCommitTiles)
+	{
+		check(!ReservedResourceData->BackingHeaps.IsEmpty());
+
+		// Iterate through heaps in reverse order, unmap ranges and release heaps if they are completely unused
+		while (ReservedResourceData->NumCommittedTiles > NumRequiredCommitTiles)
+		{
+			TRefCountPtr<FD3D12Heap>& LastHeap = ReservedResourceData->BackingHeaps.Last();
+			const uint32 NumTotalTilesInHeap = LastHeap->GetHeapDesc().SizeInBytes / TileSizeInBytes;
+
+			check(ReservedResourceData->NumSlackTiles <= NumTotalTilesInHeap);
+			const uint32 NumUsedTilesInHeap = NumTotalTilesInHeap - ReservedResourceData->NumSlackTiles;
 
-		D3DCommandQueue->UpdateTileMappings(GetResource(), 1,
-			&ResourceCoordinate, &RegionSize, D3DHeap,
-			1,
-			&RangeFlags,
-			&HeapRangeStartOffsetInTiles,
-			&RangeTileCount,
-			MappingFlags);
+			check(NumUsedTilesInHeap <= ReservedResourceData->NumCommittedTiles);
+			const uint32 HeapFirstTile = ReservedResourceData->NumCommittedTiles - NumUsedTilesInHeap;
 
-	#if NAME_OBJECTS
-		const int32 HeapIndex = ReservedResourceData->BackingHeaps.Num();
-		FString HeapName = FString::Printf(TEXT("%s.Heap[%d]"), DebugName.IsValid() ? *DebugName.ToString() : TEXT("UNKNOWN"), HeapIndex);
-		const TCHAR* HeapNameChars = *HeapName;
-	#else
-		const TCHAR* HeapNameChars = TEXT("ReservedResourceBackingHeap");
-	#endif // NAME_OBJECTS
+			const uint32 RegionEnd = ReservedResourceData->NumCommittedTiles;
+			const uint32 RegionBegin = FMath::Max(HeapFirstTile, NumRequiredCommitTiles);
 
-		NumMappedTiles += RegionSize.NumTiles;
+			D3D12_TILE_REGION_SIZE RegionSize = {};
+			RegionSize.UseBox = false;
+			RegionSize.NumTiles = RegionEnd - RegionBegin;
 
-		TRefCountPtr<FD3D12Heap> NewHeap = new FD3D12Heap(GetParentDevice(), GetVisibilityMask());
-		NewHeap->SetHeap(D3DHeap, HeapNameChars, true /*bTrack*/, false /*bForceGetGPUAddress*/);
-		//NewHeap->BeginTrackingResidency(ThisHeapSize);
-		NewHeap->DisallowTrackingResidency(); // Workaround for UE-202367: D3DX12Residency library does not track UpdateTileMappings that may be in flight
+			// Coordinates are in tiles, not pixels
+			D3D12_TILED_RESOURCE_COORDINATE ResourceCoordinate = GetTiledResourceCoordinate(RegionBegin, RegionSize.NumTiles);
+
+			const D3D12_TILE_RANGE_FLAGS RangeFlags = D3D12_TILE_RANGE_FLAG_NULL;
+
+			D3DCommandQueue->UpdateTileMappings(GetResource(), 1 /*NumRegions*/,
+				&ResourceCoordinate, &RegionSize, nullptr /*Heap*/,
+				1 /*NumRanges*/,
+				&RangeFlags,
+				nullptr /*HeapRangeStartOffsets*/,
+				&RegionSize.NumTiles,
+				MappingFlags);
+
+			if (HeapFirstTile == RegionBegin)
+			{
+				// All tiles from this heap were unmapped, so it can be dropped
+				DEC_MEMORY_STAT_BY(STAT_D3D12ReservedResourcePhysical, LastHeap->GetHeapDesc().SizeInBytes);
+				LastHeap->DeferDelete();
+				ReservedResourceData->BackingHeaps.Pop();
+				ReservedResourceData->ResidencyHandles.Pop();
+				ReservedResourceData->NumSlackTiles = 0;
+			}
+			else
+			{
+				// Heap remains referenced, but now contains some free tiles at the end (which we just unmapped)
+				ReservedResourceData->NumSlackTiles += RegionSize.NumTiles;
+				check(ReservedResourceData->NumSlackTiles <= NumTotalTilesInHeap);
+			}
 
-		ReservedResourceData->ResidencyHandles.Append(NewHeap->GetResidencyHandles());
-		ReservedResourceData->BackingHeaps.Add(MoveTemp(NewHeap));
+			check(ReservedResourceData->NumCommittedTiles >= RegionSize.NumTiles);
+			ReservedResourceData->NumCommittedTiles -= RegionSize.NumTiles;
+		}
 	}
+	else
+	{
+		while (ReservedResourceData->NumCommittedTiles < NumRequiredCommitTiles)
+		{
+			const uint32 NumRemainingTiles = NumRequiredCommitTiles - ReservedResourceData->NumCommittedTiles;
+
+			ID3D12Heap* D3DHeap = nullptr;
+
+			uint32 HeapRangeStartOffsetInTiles = 0;
+
+			D3D12_TILE_REGION_SIZE RegionSize = {};
+			RegionSize.UseBox = false;
+
+			if (ReservedResourceData->NumSlackTiles)
+			{
+				// Consume any heap slack space before allocating a new heap
+
+				const TRefCountPtr<FD3D12Heap>& LastHeap = ReservedResourceData->BackingHeaps.Last();
+				const uint32 NumTotalTilesInHeap = LastHeap->GetHeapDesc().SizeInBytes / TileSizeInBytes;
+
+				RegionSize.NumTiles = FMath::Min(ReservedResourceData->NumSlackTiles, NumRemainingTiles);
+				HeapRangeStartOffsetInTiles = NumTotalTilesInHeap - ReservedResourceData->NumSlackTiles;
 
-	ReservedResourceData->CommittedSizeInBytes = NumMappedTiles * TileSizeInBytes;
+				D3DHeap = LastHeap->GetHeap();
 
-	checkf(NumMappedTiles == NumRequiredCommitTiles,
+				check(RegionSize.NumTiles <= ReservedResourceData->NumSlackTiles);
+				ReservedResourceData->NumSlackTiles -= RegionSize.NumTiles;
+			}
+			else
+			{
+				// Create a new heap to service the commit request
+
+				RegionSize.NumTiles = FMath::Min(MaxTilesPerHeap, NumRemainingTiles);
+				HeapRangeStartOffsetInTiles = 0;
+
+#if NAME_OBJECTS
+				const int32 HeapIndex = ReservedResourceData->BackingHeaps.Num();
+				FString HeapName = FString::Printf(TEXT("%s.Heap[%d]"), DebugName.IsValid() ? *DebugName.ToString() : TEXT("UNKNOWN"), HeapIndex);
+				const TCHAR* HeapNameChars = *HeapName;
+#else
+				const TCHAR* HeapNameChars = TEXT("ReservedResourceBackingHeap");
+#endif // NAME_OBJECTS
+
+				const D3D12_HEAP_FLAGS HeapFlags = Desc.Dimension == D3D12_RESOURCE_DIMENSION_BUFFER
+					? D3D12_HEAP_FLAG_ALLOW_ONLY_BUFFERS
+					: D3D12_HEAP_FLAG_ALLOW_ONLY_NON_RT_DS_TEXTURES;
+
+				static_assert((D3D12_HEAP_FLAG_DENY_BUFFERS | D3D12_HEAP_FLAG_DENY_RT_DS_TEXTURES) == D3D12_HEAP_FLAG_ALLOW_ONLY_NON_RT_DS_TEXTURES);
+
+				const uint32 ThisHeapSize = RegionSize.NumTiles * TileSizeInBytes;
+				D3D12_HEAP_DESC NewHeapDesc = {};
+				NewHeapDesc.Alignment = D3D12_DEFAULT_RESOURCE_PLACEMENT_ALIGNMENT;
+				NewHeapDesc.Flags = HeapFlags;
+				NewHeapDesc.SizeInBytes = ThisHeapSize;
+				NewHeapDesc.Properties = BackingHeapProps;
+
+				VERIFYD3D12RESULT(D3DDevice->CreateHeap(&NewHeapDesc, IID_PPV_ARGS(&D3DHeap)));
+				INC_MEMORY_STAT_BY(STAT_D3D12ReservedResourcePhysical, NewHeapDesc.SizeInBytes);
+
+				if (bHighPriorityResource)
+				{
+					Adapter->SetResidencyPriority(D3DHeap, D3D12_RESIDENCY_PRIORITY_HIGH, GPUIndex);
+				}
+
+				TRefCountPtr<FD3D12Heap> NewHeap = new FD3D12Heap(GetParentDevice(), GetVisibilityMask());
+				NewHeap->SetHeap(D3DHeap, HeapNameChars, true /*bTrack*/, false /*bForceGetGPUAddress*/);
+				//NewHeap->BeginTrackingResidency(ThisHeapSize);
+				NewHeap->DisallowTrackingResidency(); // Workaround for UE-202367: D3DX12Residency library does not track UpdateTileMappings that may be in flight
+
+				ReservedResourceData->ResidencyHandles.Append(NewHeap->GetResidencyHandles());
+				ReservedResourceData->BackingHeaps.Add(MoveTemp(NewHeap));
+			}
+
+			// Coordinates are in tiles, not pixels
+			D3D12_TILED_RESOURCE_COORDINATE ResourceCoordinate = GetTiledResourceCoordinate(ReservedResourceData->NumCommittedTiles, RegionSize.NumTiles);
+
+			const D3D12_TILE_RANGE_FLAGS RangeFlags = D3D12_TILE_RANGE_FLAG_NONE;
+
+			D3DCommandQueue->UpdateTileMappings(GetResource(), 1 /*NumRegions*/,
+				&ResourceCoordinate, &RegionSize, D3DHeap,
+				1 /*NumRanges*/,
+				&RangeFlags,
+				&HeapRangeStartOffsetInTiles,
+				&RegionSize.NumTiles,
+				MappingFlags);
+
+			ReservedResourceData->NumCommittedTiles += RegionSize.NumTiles;
+		}
+	}
+
+	checkf(ReservedResourceData->NumCommittedTiles == NumRequiredCommitTiles,
 		TEXT("Reserved resource was not fully processed while committing physical memory. Expected to process tiles: %d, actually processed: %d"),
-		D3DResourceNumTiles, NumMappedTiles);
+		D3DResourceNumTiles, ReservedResourceData->NumCommittedTiles);
 }
 
 ID3D12Pageable* FD3D12Resource::GetPageable()
@@ -445,6 +526,9 @@ FD3D12Heap::~FD3D12Heap()
 
 void FD3D12Heap::DeferDelete()
 {
+	// ProcessDeferredDeletionQueue() performs final Release(), but deletion queue itself only holds a raw pointer, so explicit addref is required.
+	AddRef();
+
 	FD3D12DynamicRHI::GetD3DRHI()->DeferredDelete(this);
 }
 
diff --git a/Engine/Source/Runtime/D3D12RHI/Private/D3D12Resources.h b/Engine/Source/Runtime/D3D12RHI/Private/D3D12Resources.h
--- a/Engine/Source/Runtime/D3D12RHI/Private/D3D12Resources.h
+++ b/Engine/Source/Runtime/D3D12RHI/Private/D3D12Resources.h
@@ -200,8 +200,15 @@ private:
 	struct FD3D12ReservedResourceData
 	{
 		TArray<TRefCountPtr<FD3D12Heap>> BackingHeaps;
-		TArray<FD3D12ResidencyHandle*> ResidencyHandles; // Flattened array of residency handles owned by backing heaps, used to support batched GetResidencyHandles()
-		uint64 CommittedSizeInBytes = 0;
+
+		// Flattened array of residency handles owned by backing heaps, used to support batched GetResidencyHandles()
+		TArray<FD3D12ResidencyHandle*> ResidencyHandles;
+
+		// Tiles currently assigned to the resource
+		uint32 NumCommittedTiles = 0;
+
+		// Available tiles at the end of the last backing heap
+		uint32 NumSlackTiles = 0;
 	};
 	TUniquePtr<FD3D12ReservedResourceData> ReservedResourceData;
 
diff --git a/Engine/Source/Runtime/D3D12RHI/Private/D3D12Stats.h b/Engine/Source/Runtime/D3D12RHI/Private/D3D12Stats.h
--- a/Engine/Source/Runtime/D3D12RHI/Private/D3D12Stats.h
+++ b/Engine/Source/Runtime/D3D12RHI/Private/D3D12Stats.h
@@ -116,6 +116,8 @@ DECLARE_MEMORY_STAT_EXTERN(TEXT("UploadPool Memory Alignment Waste"), STAT_D3D12
 DECLARE_DWORD_ACCUMULATOR_STAT_EXTERN(TEXT("UploadPool Page Count"), STAT_D3D12UploadPoolPageCount, STATGROUP_D3D12BufferDetails, );
 DECLARE_DWORD_ACCUMULATOR_STAT_EXTERN(TEXT("UploadPool Full Pages"), STAT_D3D12UploadPoolFullPages, STATGROUP_D3D12BufferDetails, );
 
+DECLARE_MEMORY_STAT_EXTERN(TEXT("Reserved Resource Physical Memory"), STAT_D3D12ReservedResourcePhysical, STATGROUP_D3D12MemoryDetails, );
+
 /**
 * Detailed Descriptor heap stats
 */
diff --git a/Engine/Source/Runtime/D3D12RHI/Private/D3D12TransientResourceAllocator.cpp b/Engine/Source/Runtime/D3D12RHI/Private/D3D12TransientResourceAllocator.cpp
--- a/Engine/Source/Runtime/D3D12RHI/Private/D3D12TransientResourceAllocator.cpp
+++ b/Engine/Source/Runtime/D3D12RHI/Private/D3D12TransientResourceAllocator.cpp
@@ -123,8 +123,6 @@ FD3D12TransientHeap::~FD3D12TransientHeap()
 		MemoryTrace_Free((uint64)Heap->GetHeap(), EMemoryTraceRootHeap::VideoMemory);
 #endif
 
-		// Add a reference to the underlying to heap to be deferred released later.
-		Heap->AddRef();
 		Heap->DeferDelete();
 	}
 }
diff --git a/Engine/Source/Runtime/D3D12RHI/Private/D3D12Util.cpp b/Engine/Source/Runtime/D3D12RHI/Private/D3D12Util.cpp
--- a/Engine/Source/Runtime/D3D12RHI/Private/D3D12Util.cpp
+++ b/Engine/Source/Runtime/D3D12RHI/Private/D3D12Util.cpp
@@ -2061,6 +2061,8 @@ DEFINE_STAT(STAT_D3D12UploadPoolAlignmentWaste);
 DEFINE_STAT(STAT_D3D12UploadPoolPageCount);
 DEFINE_STAT(STAT_D3D12UploadPoolFullPages);
 
+DEFINE_STAT(STAT_D3D12ReservedResourcePhysical);
+
 DEFINE_STAT(STAT_UniqueSamplers);
 
 DEFINE_STAT(STAT_ViewHeapChanged);
