Exclude external textures created through FD3D12DynamicRHI::CreateTextureFromResource from residency tracking

* Add FD3D12ResourceDesc::bBackBuffer and bExternal members
* Remove FD3D12Resource::SetIsBackBuffer

diff --git a/Engine/Source/Runtime/D3D12RHI/Private/D3D12Resources.cpp b/Engine/Source/Runtime/D3D12RHI/Private/D3D12Resources.cpp
--- a/Engine/Source/Runtime/D3D12RHI/Private/D3D12Resources.cpp
+++ b/Engine/Source/Runtime/D3D12RHI/Private/D3D12Resources.cpp
@@ -95,7 +95,6 @@ FD3D12Resource::FD3D12Resource(FD3D12Device* ParentDevice,
 	, bRequiresResidencyTracking(true)
 	, bDepthStencil(false)
 	, bDeferDelete(true)
-	, bBackBuffer(false)
 {
 #if UE_BUILD_DEBUG
 	FPlatformAtomics::InterlockedIncrement(&TotalResourceCount);
@@ -109,8 +108,8 @@ FD3D12Resource::FD3D12Resource(FD3D12Device* ParentDevice,
 		HeapProps = &HeapDesc.Properties;
 	}
 
-	// Residency tracking is not used for CPU-accessible resources or back buffers
-	bRequiresResidencyTracking = IsGPUOnly(InHeapType, HeapProps) && !bBackBuffer;
+	// Residency tracking is only used for GPU-only resources owned by the Engine
+	bRequiresResidencyTracking = IsGPUOnly(InHeapType, HeapProps) && !Desc.bExternal && !Desc.bBackBuffer;
 
 	// On Windows it's sadly enough not possible to get the GPU virtual address from the resource directly
 	if (Resource
@@ -155,11 +154,10 @@ FD3D12Resource::~FD3D12Resource()
 	}
 #endif
 
-	if (bBackBuffer)
+	if (Desc.bBackBuffer)
 	{
 		// Don't make the windows association call and release back buffer at the same time (see notes on critical section)
 		FScopeLock Lock(&FD3D12Viewport::DXGIBackBufferLock);
-		bBackBuffer = false;
 		Resource.SafeRelease();
 	}
 }
@@ -457,7 +455,7 @@ void FD3D12Resource::StartTrackingForResidency()
 
 	checkf(bRequiresResidencyTracking, TEXT("Residency tracking is not expected for this resource"));
 
-	if (bBackBuffer)
+	if (Desc.bBackBuffer)
 	{
 		// Back buffers may be referenced outside of command lists (during presents), however D3DX12Residency.h library 
 		// uses fences tied to command lists to detect when it's safe to evict a resource, which is wrong for back buffers.
diff --git a/Engine/Source/Runtime/D3D12RHI/Private/D3D12Resources.h b/Engine/Source/Runtime/D3D12RHI/Private/D3D12Resources.h
--- a/Engine/Source/Runtime/D3D12RHI/Private/D3D12Resources.h
+++ b/Engine/Source/Runtime/D3D12RHI/Private/D3D12Resources.h
@@ -135,10 +135,15 @@ struct FD3D12ResourceDesc : public D3D12_RESOURCE_DESC
 	EPixelFormat UAVPixelFormat{ PF_Unknown };
 
 #if D3D12RHI_NEEDS_VENDOR_EXTENSIONS
-	bool bRequires64BitAtomicSupport{ false };
+	bool bRequires64BitAtomicSupport : 1 = false;
 #endif
 
-	bool bReservedResource { false };
+	bool bReservedResource : 1 = false;
+
+	bool bBackBuffer : 1 = false;
+
+	// External resources are owned by another application or middleware, not the Engine
+	bool bExternal : 1 = false;
 
 	// If we support the new format list casting, use the newer APIs; otherwise, fall back to our UAV Aliasing approach.
 #if D3D12RHI_SUPPORTS_UNCOMPRESSED_UAV
@@ -190,7 +195,6 @@ private:
 	bool bRequiresResidencyTracking : 1;
 	bool bDepthStencil : 1;
 	bool bDeferDelete : 1;
-	bool bBackBuffer : 1;
 
 #if UE_BUILD_DEBUG
 	static int64 TotalResourceCount;
@@ -297,17 +301,7 @@ public:
 #endif
 	bool RequiresResourceStateTracking() const { return bRequiresResourceStateTracking; }
 
-	inline bool IsBackBuffer() const { return bBackBuffer; }
-	inline void SetIsBackBuffer(bool bBackBufferIn) 
-	{
-		bBackBuffer = bBackBufferIn;
-
-		if (bBackBuffer)
-		{
-			checkf(ResidencyHandle == nullptr, TEXT("Can't mark a resource as back buffer once residency tracking has started."));
-			bRequiresResidencyTracking = false;
-		}
-	}
+	inline bool IsBackBuffer() const { return Desc.bBackBuffer; }
 
 	void SetName(const TCHAR* Name)
 	{
diff --git a/Engine/Source/Runtime/D3D12RHI/Private/D3D12Texture.cpp b/Engine/Source/Runtime/D3D12RHI/Private/D3D12Texture.cpp
--- a/Engine/Source/Runtime/D3D12RHI/Private/D3D12Texture.cpp
+++ b/Engine/Source/Runtime/D3D12RHI/Private/D3D12Texture.cpp
@@ -2694,7 +2694,8 @@ FD3D12Texture* FD3D12DynamicRHI::CreateTextureFromResource(bool bTextureArray, b
 	check(Resource);
 	FD3D12Adapter* Adapter = &GetAdapter();
 
-	D3D12_RESOURCE_DESC TextureDesc = Resource->GetDesc();
+	FD3D12ResourceDesc TextureDesc = Resource->GetDesc();
+	TextureDesc.bExternal = true;
 	TextureDesc.Alignment = 0;
 
 	uint32 SizeX = TextureDesc.Width;
@@ -2722,6 +2723,7 @@ FD3D12Texture* FD3D12DynamicRHI::CreateTextureFromResource(bool bTextureArray, b
 	FD3D12Device* Device = Adapter->GetDevice(0);
 	FD3D12Resource* TextureResource = new FD3D12Resource(Device, Device->GetGPUMask(), Resource, DestinationState, TextureDesc);
 	TextureResource->AddRef();
+	TextureResource->SetName(TEXT("TextureFromResource"));
 
 	FRHITextureCreateDesc CreateDesc =
 		((SizeZ > 1) ? FRHITextureCreateDesc::Create2DArray(TEXT("TextureFromResource"), FIntPoint(SizeX, SizeY), SizeZ, Format) :
diff --git a/Engine/Source/Runtime/D3D12RHI/Private/D3D12Viewport.cpp b/Engine/Source/Runtime/D3D12RHI/Private/D3D12Viewport.cpp
--- a/Engine/Source/Runtime/D3D12RHI/Private/D3D12Viewport.cpp
+++ b/Engine/Source/Runtime/D3D12RHI/Private/D3D12Viewport.cpp
@@ -178,7 +178,8 @@ FD3D12Texture* GetSwapChainSurface(FD3D12Device* Parent, EPixelFormat PixelForma
 		Parent->GetDevice()->CreateCommittedResource(&HeapProps, D3D12_HEAP_FLAG_NONE, &TextureDesc, D3D12_RESOURCE_STATE_PRESENT, nullptr, IID_PPV_ARGS(BackBufferResource.GetInitReference()));
 	}
 
-	D3D12_RESOURCE_DESC BackBufferDesc = BackBufferResource->GetDesc();
+	FD3D12ResourceDesc BackBufferDesc = BackBufferResource->GetDesc();
+	BackBufferDesc.bBackBuffer = true;
 
 	FString Name = FString::Printf(TEXT("BackBuffer%d"), BackBufferIndex);
 
@@ -216,7 +217,6 @@ FD3D12Texture* GetSwapChainSurface(FD3D12Device* Parent, EPixelFormat PixelForma
 		if (Device->GetGPUIndex() == Parent->GetGPUIndex())
 		{
 			FD3D12Resource* NewResourceWrapper = new FD3D12Resource(Device, FRHIGPUMask::All(), BackBufferResource, InitialState, BackBufferDesc);
-			NewResourceWrapper->SetIsBackBuffer(true);
 			NewResourceWrapper->AddRef();
 			NewTexture->ResourceLocation.AsStandAlone(NewResourceWrapper);
 		}
@@ -285,8 +285,6 @@ FD3D12Texture* GetSwapChainSurface(FD3D12Device* Parent, EPixelFormat PixelForma
 
 	SetName(SwapChainTexture->GetResource(), *Name);
 
-	SwapChainTexture->GetResource()->SetIsBackBuffer(true);
-
 	const D3D12_RESOURCE_ALLOCATION_INFO AllocationInfo = Parent->GetDevice()->GetResourceAllocationInfo(0, 1, &SwapChainTexture->GetResource()->GetDesc());
 	SwapChainTexture->ResourceLocation.SetSize(AllocationInfo.SizeInBytes);
 
