Fix resindency handle tracking for reserved resource heaps. A heap may have no residency handles at all, but code previously assumed that there is exactly one per heap.

diff --git a/Engine/Source/Runtime/D3D12RHI/Private/D3D12Resources.cpp b/Engine/Source/Runtime/D3D12RHI/Private/D3D12Resources.cpp
--- a/Engine/Source/Runtime/D3D12RHI/Private/D3D12Resources.cpp
+++ b/Engine/Source/Runtime/D3D12RHI/Private/D3D12Resources.cpp
@@ -328,7 +328,14 @@ void FD3D12Resource::CommitReservedResource(ID3D12CommandQueue* D3DCommandQueue,
 				DEC_MEMORY_STAT_BY(STAT_D3D12ReservedResourcePhysical, LastHeap->GetHeapDesc().SizeInBytes);
 				LastHeap->DeferDelete();
 				ReservedResourceData->BackingHeaps.Pop();
-				ReservedResourceData->ResidencyHandles.Pop();
+				int32 NumResidencyHandles = ReservedResourceData->NumResidencyHandlesPerHeap.Last();
+				ReservedResourceData->NumResidencyHandlesPerHeap.Pop();
+				while (NumResidencyHandles != 0)
+				{
+					ReservedResourceData->ResidencyHandles.Pop();
+					--NumResidencyHandles;
+				}
+
 				ReservedResourceData->NumSlackTiles = 0;
 			}
 			else
@@ -411,7 +418,9 @@ void FD3D12Resource::CommitReservedResource(ID3D12CommandQueue* D3DCommandQueue,
 				//NewHeap->BeginTrackingResidency(ThisHeapSize);
 				NewHeap->DisallowTrackingResidency(); // Workaround for UE-202367: D3DX12Residency library does not track UpdateTileMappings that may be in flight
 
-				ReservedResourceData->ResidencyHandles.Append(NewHeap->GetResidencyHandles());
+				TConstArrayView<FD3D12ResidencyHandle*> HeapResidencyHandles = NewHeap->GetResidencyHandles();
+				ReservedResourceData->ResidencyHandles.Append(HeapResidencyHandles);
+				ReservedResourceData->NumResidencyHandlesPerHeap.Add(HeapResidencyHandles.Num());
 				ReservedResourceData->BackingHeaps.Add(MoveTemp(NewHeap));
 			}
 
diff --git a/Engine/Source/Runtime/D3D12RHI/Private/D3D12Resources.h b/Engine/Source/Runtime/D3D12RHI/Private/D3D12Resources.h
--- a/Engine/Source/Runtime/D3D12RHI/Private/D3D12Resources.h
+++ b/Engine/Source/Runtime/D3D12RHI/Private/D3D12Resources.h
@@ -207,6 +207,7 @@ private:
 
 		// Flattened array of residency handles owned by backing heaps, used to support batched GetResidencyHandles()
 		TArray<FD3D12ResidencyHandle*> ResidencyHandles;
+		TArray<int32> NumResidencyHandlesPerHeap;
 
 		// Tiles currently assigned to the resource
 		uint32 NumCommittedTiles = 0;
