Fix compile errors when ENABLE_RESIDENCY_MANAGEMENT is 0

diff --git a/Engine/Source/Runtime/D3D12RHI/Private/D3D12Resources.cpp b/Engine/Source/Runtime/D3D12RHI/Private/D3D12Resources.cpp
--- a/Engine/Source/Runtime/D3D12RHI/Private/D3D12Resources.cpp
+++ b/Engine/Source/Runtime/D3D12RHI/Private/D3D12Resources.cpp
@@ -449,6 +449,7 @@ void FD3D12Resource::CommitReservedResource(ID3D12CommandQueue* D3DCommandQueue,
 			ReservedResourceData->NumCommittedTiles += RegionSize.NumTiles;
 		}
 
+#if ENABLE_RESIDENCY_MANAGEMENT
 		FD3D12ResidencyManager& ResidencyManager = GetParentDevice()->GetResidencyManager();
 
 		if (!UsedResidencyHandles.IsEmpty())
@@ -475,6 +476,7 @@ void FD3D12Resource::CommitReservedResource(ID3D12CommandQueue* D3DCommandQueue,
 			HR = ResidencyManager.MakeResident(D3DCommandQueue, MoveTemp(ResidencySet));
 			checkf(SUCCEEDED(HR), TEXT("Failed to process residency set. Error code: 0x%08x."), uint32(HR));
 		}
+#endif // ENABLE_RESIDENCY_MANAGEMENT
 
 		for (const FD3D12UpdateTileMappingsParams& Params : MappingParams)
 		{
@@ -487,12 +489,14 @@ void FD3D12Resource::CommitReservedResource(ID3D12CommandQueue* D3DCommandQueue,
 				D3D12_TILE_MAPPING_FLAG_NONE);
 		}
 
+#if ENABLE_RESIDENCY_MANAGEMENT
 		if (!UsedResidencyHandles.IsEmpty())
 		{
 			// Signal the fence for this queue after UpdateTileMappings complete.
 			// This is analogous to executing a command list that references a set of resources.
 			ResidencyManager.SignalFence(D3DCommandQueue);
 		}
+#endif // ENABLE_RESIDENCY_MANAGEMENT
 	}
 
 	checkf(ReservedResourceData->NumCommittedTiles == NumRequiredCommitTiles,
