Fix redundant/erroneous initialization of Params.RangeFlags in FD3D12Resource::CommitReservedResource

diff --git a/Engine/Source/Runtime/D3D12RHI/Private/D3D12Resources.cpp b/Engine/Source/Runtime/D3D12RHI/Private/D3D12Resources.cpp
--- a/Engine/Source/Runtime/D3D12RHI/Private/D3D12Resources.cpp
+++ b/Engine/Source/Runtime/D3D12RHI/Private/D3D12Resources.cpp
@@ -438,12 +438,11 @@ void FD3D12Resource::CommitReservedResource(ID3D12CommandQueue* D3DCommandQueue,
 			D3D12_TILED_RESOURCE_COORDINATE ResourceCoordinate = GetTiledResourceCoordinate(ReservedResourceData->NumCommittedTiles, RegionSize.NumTiles);
 
 			FD3D12UpdateTileMappingsParams Params = {};
-			Params.RangeFlags = D3D12_TILE_RANGE_FLAG_NULL;
+			Params.RangeFlags = D3D12_TILE_RANGE_FLAG_NONE;
 			Params.Coord = ResourceCoordinate;
 			Params.Size = RegionSize;
 			Params.Heap = D3DHeap;
 			Params.HeapOffsetInTiles = HeapRangeStartOffsetInTiles;
-			Params.RangeFlags = D3D12_TILE_RANGE_FLAG_NONE;
 			MappingParams.Add(Params);
 
 			ReservedResourceData->NumCommittedTiles += RegionSize.NumTiles;
