Fix missing UpdateTileMappings calls when decommitting tiles from reserved resources

diff --git a/Engine/Source/Runtime/D3D12RHI/Private/D3D12Resources.cpp b/Engine/Source/Runtime/D3D12RHI/Private/D3D12Resources.cpp
--- a/Engine/Source/Runtime/D3D12RHI/Private/D3D12Resources.cpp
+++ b/Engine/Source/Runtime/D3D12RHI/Private/D3D12Resources.cpp
@@ -294,6 +294,7 @@ void FD3D12Resource::CommitReservedResource(ID3D12CommandQueue* D3DCommandQueue,
 	};
 
 	TArray<FD3D12UpdateTileMappingsParams> MappingParams;
+	TArray<FD3D12ResidencyHandle*> UsedResidencyHandles;
 
 	if (ReservedResourceData->NumCommittedTiles > NumRequiredCommitTiles)
 	{
@@ -356,8 +357,6 @@ void FD3D12Resource::CommitReservedResource(ID3D12CommandQueue* D3DCommandQueue,
 	}
 	else
 	{
-		TArray<FD3D12ResidencyHandle*> UsedResidencyHandles;
-
 		while (ReservedResourceData->NumCommittedTiles < NumRequiredCommitTiles)
 		{
 			const uint32 NumRemainingTiles = NumRequiredCommitTiles - ReservedResourceData->NumCommittedTiles;
@@ -447,56 +446,56 @@ void FD3D12Resource::CommitReservedResource(ID3D12CommandQueue* D3DCommandQueue,
 
 			ReservedResourceData->NumCommittedTiles += RegionSize.NumTiles;
 		}
+	}
 
 #if ENABLE_RESIDENCY_MANAGEMENT
-		FD3D12ResidencyManager& ResidencyManager = GetParentDevice()->GetResidencyManager();
+	FD3D12ResidencyManager& ResidencyManager = GetParentDevice()->GetResidencyManager();
 
-		if (!UsedResidencyHandles.IsEmpty())
-		{
-			HRESULT HR = S_OK;
+	if (!UsedResidencyHandles.IsEmpty())
+	{
+		HRESULT HR = S_OK;
 
-			FD3D12ResidencySet* ResidencySet = ResidencyManager.CreateResidencySet();
-			HR = ResidencySet->Open();
-			checkf(SUCCEEDED(HR), TEXT("Failed to open residency set. Error code: 0x%08x."), uint32(HR));
+		FD3D12ResidencySet* ResidencySet = ResidencyManager.CreateResidencySet();
+		HR = ResidencySet->Open();
+		checkf(SUCCEEDED(HR), TEXT("Failed to open residency set. Error code: 0x%08x."), uint32(HR));
 
-			for (FD3D12ResidencyHandle* Handle : UsedResidencyHandles)
+		for (FD3D12ResidencyHandle* Handle : UsedResidencyHandles)
+		{
+			if (D3DX12Residency::IsInitialized(Handle))
 			{
-				if (D3DX12Residency::IsInitialized(Handle))
-				{
-					ResidencySet->Insert(Handle);
-				}
+				ResidencySet->Insert(Handle);
 			}
+		}
 
-			HR = ResidencySet->Close();
-			checkf(SUCCEEDED(HR), TEXT("Failed to close residency set. Error code: 0x%08x."), uint32(HR));
+		HR = ResidencySet->Close();
+		checkf(SUCCEEDED(HR), TEXT("Failed to close residency set. Error code: 0x%08x."), uint32(HR));
 
-			// NOTE: ResidencySet ownership is taken over by the residency manager.
-			// It is destroyed when paging work completes, which may happen async on another thread in some cases.
-			HR = ResidencyManager.MakeResident(D3DCommandQueue, MoveTemp(ResidencySet));
-			checkf(SUCCEEDED(HR), TEXT("Failed to process residency set. Error code: 0x%08x."), uint32(HR));
-		}
+		// NOTE: ResidencySet ownership is taken over by the residency manager.
+		// It is destroyed when paging work completes, which may happen async on another thread in some cases.
+		HR = ResidencyManager.MakeResident(D3DCommandQueue, MoveTemp(ResidencySet));
+		checkf(SUCCEEDED(HR), TEXT("Failed to process residency set. Error code: 0x%08x."), uint32(HR));
+	}
 #endif // ENABLE_RESIDENCY_MANAGEMENT
 
-		for (const FD3D12UpdateTileMappingsParams& Params : MappingParams)
-		{
-			D3DCommandQueue->UpdateTileMappings(GetResource(), 1 /*NumRegions*/,
-				&Params.Coord, &Params.Size, Params.Heap,
-				1 /*NumRanges*/,
-				&Params.RangeFlags,
-				&Params.HeapOffsetInTiles,
-				&Params.Size.NumTiles,
-				D3D12_TILE_MAPPING_FLAG_NONE);
-		}
+	for (const FD3D12UpdateTileMappingsParams& Params : MappingParams)
+	{
+		D3DCommandQueue->UpdateTileMappings(GetResource(), 1 /*NumRegions*/,
+			&Params.Coord, &Params.Size, Params.Heap,
+			1 /*NumRanges*/,
+			&Params.RangeFlags,
+			&Params.HeapOffsetInTiles,
+			&Params.Size.NumTiles,
+			D3D12_TILE_MAPPING_FLAG_NONE);
+	}
 
 #if ENABLE_RESIDENCY_MANAGEMENT
-		if (!UsedResidencyHandles.IsEmpty())
-		{
-			// Signal the fence for this queue after UpdateTileMappings complete.
-			// This is analogous to executing a command list that references a set of resources.
-			ResidencyManager.SignalFence(D3DCommandQueue);
-		}
-#endif // ENABLE_RESIDENCY_MANAGEMENT
+	if (!UsedResidencyHandles.IsEmpty())
+	{
+		// Signal the fence for this queue after UpdateTileMappings complete.
+		// This is analogous to executing a command list that references a set of resources.
+		ResidencyManager.SignalFence(D3DCommandQueue);
 	}
+#endif // ENABLE_RESIDENCY_MANAGEMENT
 
 	checkf(ReservedResourceData->NumCommittedTiles == NumRequiredCommitTiles,
 		TEXT("Reserved resource was not fully processed while committing physical memory. Expected to process tiles: %d, actually processed: %d"),
