Fix reserved resource commit/decommit when ENABLE_RESIDENCY_MANAGEMENT is 0

* Add compile guards around some of the residency tracking code that should not run when residency management is not used

diff --git a/Engine/Source/Runtime/D3D12RHI/Private/D3D12Resources.cpp b/Engine/Source/Runtime/D3D12RHI/Private/D3D12Resources.cpp
--- a/Engine/Source/Runtime/D3D12RHI/Private/D3D12Resources.cpp
+++ b/Engine/Source/Runtime/D3D12RHI/Private/D3D12Resources.cpp
@@ -92,7 +92,7 @@ FD3D12Resource::FD3D12Resource(FD3D12Device* ParentDevice,
 	, HeapType(InHeapType)
 	, PlaneCount(UE::DXGIUtilities::GetPlaneCount(Desc.Format))
 	, bRequiresResourceStateTracking(true)
-	, bRequiresResidencyTracking(true)
+	, bRequiresResidencyTracking(bool(ENABLE_RESIDENCY_MANAGEMENT))
 	, bDepthStencil(false)
 	, bDeferDelete(true)
 {
@@ -108,8 +108,10 @@ FD3D12Resource::FD3D12Resource(FD3D12Device* ParentDevice,
 		HeapProps = &HeapDesc.Properties;
 	}
 
+#if ENABLE_RESIDENCY_MANAGEMENT
 	// Residency tracking is only used for GPU-only resources owned by the Engine
 	bRequiresResidencyTracking = IsGPUOnly(InHeapType, HeapProps) && !Desc.bExternal && !Desc.bBackBuffer;
+#endif
 
 	// On Windows it's sadly enough not possible to get the GPU virtual address from the resource directly
 	if (Resource
@@ -604,7 +606,9 @@ void FD3D12Heap::SetHeap(ID3D12Heap* HeapIn, const TCHAR* const InName, bool bIn
 
 	SetName(HeapIn, InName);
 
+#if ENABLE_RESIDENCY_MANAGEMENT
 	bRequiresResidencyTracking = IsGPUOnly(HeapDesc.Properties.Type, &HeapDesc.Properties);
+#endif // ENABLE_RESIDENCY_MANAGEMENT
 
 	// Create a buffer placed resource on the heap to extract the gpu virtual address
 	// if we are tracking all allocations
@@ -640,8 +644,10 @@ void FD3D12Heap::SetHeap(ID3D12Heap* HeapIn, const TCHAR* const InName, bool bIn
 
 void FD3D12Heap::DisallowTrackingResidency()
 {
+#if ENABLE_RESIDENCY_MANAGEMENT
 	checkf(ResidencyHandle == nullptr, TEXT("Can't disallow residency tracking after it has started. Call this function instead of BeginTrackingResidency()."));
 	bRequiresResidencyTracking = false;
+#endif // ENABLE_RESIDENCY_MANAGEMENT
 }
 
 void FD3D12Heap::BeginTrackingResidency(uint64 Size)
@@ -652,7 +658,7 @@ void FD3D12Heap::BeginTrackingResidency(uint64 Size)
 	ResidencyHandle = new FD3D12ResidencyHandle;
 	D3DX12Residency::Initialize(*ResidencyHandle, Heap.GetReference(), Size, this);
 	D3DX12Residency::BeginTrackingObject(GetParentDevice()->GetResidencyManager(), *ResidencyHandle);
-#endif
+#endif // ENABLE_RESIDENCY_MANAGEMENT
 }
 
 /////////////////////////////////////////////////////////////////////
diff --git a/Engine/Source/Runtime/D3D12RHI/Private/D3D12Resources.h b/Engine/Source/Runtime/D3D12RHI/Private/D3D12Resources.h
--- a/Engine/Source/Runtime/D3D12RHI/Private/D3D12Resources.h
+++ b/Engine/Source/Runtime/D3D12RHI/Private/D3D12Resources.h
@@ -87,12 +87,14 @@ public:
 	inline D3D12_HEAP_DESC GetHeapDesc() const { return HeapDesc; }
 	inline TConstArrayView<FD3D12ResidencyHandle*> GetResidencyHandles()
 	{ 
+#if ENABLE_RESIDENCY_MANAGEMENT
 		if (bRequiresResidencyTracking)
 		{
 			checkf(ResidencyHandle, TEXT("Resource requires residency tracking, but BeginTrackingResidency() was not called."));
 			return MakeArrayView(&ResidencyHandle, 1); 
 		}
 		else
+#endif // ENABLE_RESIDENCY_MANAGEMENT
 		{
 			return {};
 		}		
@@ -112,7 +114,7 @@ private:
 	HeapId TraceHeapId;
 	HeapId TraceParentHeapId;
 	bool bIsTransient = false; // Whether this is a transient heap
-	bool bRequiresResidencyTracking = true;
+	bool bRequiresResidencyTracking = bool(ENABLE_RESIDENCY_MANAGEMENT);
 };
 
 struct FD3D12ResourceDesc : public D3D12_RESOURCE_DESC
