Fix missing D3D12_HEAP_FLAG_SHARED for shared resources when allocation tracking is enabled.

CreatePlacedResource correctly sets the flag, but when allocation tracking is enabled (due to -gpucrashdebugging or D3D12.TrackAllAllocations), the code is creating a heap and placing the resource explicitly (in order to be able to get the GPU VA), and that path didn't set the flag.

#rnx

diff --git a/Engine/Source/Runtime/D3D12RHI/Private/D3D12PoolAllocator.cpp b/Engine/Source/Runtime/D3D12RHI/Private/D3D12PoolAllocator.cpp
--- a/Engine/Source/Runtime/D3D12RHI/Private/D3D12PoolAllocator.cpp
+++ b/Engine/Source/Runtime/D3D12RHI/Private/D3D12PoolAllocator.cpp
@@ -421,6 +421,11 @@ void FD3D12PoolAllocator::AllocateResource(uint32 GPUIndex, D3D12_HEAP_TYPE InHe
 				HeapDesc.Flags |= FD3D12_HEAP_FLAG_CREATE_NOT_ZEROED;
 			}
 
+			if (Desc.Flags & D3D12_RESOURCE_FLAG_ALLOW_SIMULTANEOUS_ACCESS)
+			{
+				HeapDesc.Flags |= D3D12_HEAP_FLAG_SHARED;
+			}
+
 			ID3D12Heap* Heap = nullptr;
 			VERIFYD3D12RESULT(Adapter->GetD3DDevice()->CreateHeap(&HeapDesc, IID_PPV_ARGS(&Heap)));
 			TRefCountPtr<FD3D12Heap> BackingHeap = new FD3D12Heap(GetParentDevice(), GetVisibilityMask(), TraceHeapId);
