Perform tile mapping updates on the direct queue when the current queue does not support the operation

diff --git a/Engine/Source/Runtime/D3D12RHI/Private/D3D12Device.cpp b/Engine/Source/Runtime/D3D12RHI/Private/D3D12Device.cpp
--- a/Engine/Source/Runtime/D3D12RHI/Private/D3D12Device.cpp
+++ b/Engine/Source/Runtime/D3D12RHI/Private/D3D12Device.cpp
@@ -36,6 +36,7 @@ FD3D12Queue::FD3D12Queue(FD3D12Device* Device, ED3D12QueueType QueueType)
 	: Device(Device)
 	, QueueType(QueueType)
 	, BarrierTimestamps(Device, QueueType, D3D12_QUERY_TYPE_TIMESTAMP)
+	, bSupportsTileMapping(FD3D12DynamicRHI::GetD3DRHI()->QueueSupportsTileMapping(QueueType))
 {
 	FD3D12Adapter* Adapter = Device->GetParentAdapter();
 	const bool bFullGPUCrashDebugging = (Adapter->GetGPUCrashDebuggingModes() == ED3D12GPUCrashDebuggingModes::All);
@@ -140,6 +141,18 @@ FD3D12Device::FD3D12Device(FRHIGPUMask InGPUMask, FD3D12Adapter* InAdapter)
 	{
 		Queues.Emplace(this, (ED3D12QueueType)QueueType);
 	}
+
+	// Some hardware is not capable of running tile mapping operations on all queue types.
+	// Direct queue is used as a fallback if tile updates are requested on unsupported queue.
+	TileMappingQueue = Queues[size_t(ED3D12QueueType::Direct)].D3DCommandQueue;
+	VERIFYD3D12RESULT(GetDevice()->CreateFence(
+		0,
+		D3D12_FENCE_FLAG_NONE,
+		IID_PPV_ARGS(TileMappingFence.D3DFence.GetInitReference())
+	));
+#if NAME_OBJECTS
+	TileMappingFence.D3DFence->SetName(TEXT("TileMappingFence"));
+#endif
 }
 
 FD3D12Device::~FD3D12Device()
diff --git a/Engine/Source/Runtime/D3D12RHI/Private/D3D12Device.h b/Engine/Source/Runtime/D3D12RHI/Private/D3D12Device.h
--- a/Engine/Source/Runtime/D3D12RHI/Private/D3D12Device.h
+++ b/Engine/Source/Runtime/D3D12RHI/Private/D3D12Device.h
@@ -125,6 +125,9 @@ public:
 
 	bool bRequiresSignal = false;
 
+	// On some hardware, some auxiliary queue types may not support tile mapping and a separate queue must be used
+	bool bSupportsTileMapping = true;
+
 	// The underlying D3D queue object
 	TRefCountPtr<ID3D12CommandQueue> D3DCommandQueue;
 
@@ -343,6 +346,9 @@ public:
 	void SetupAfterDeviceCreation();
 	void CleanupResources();
 
+	TRefCountPtr<ID3D12CommandQueue> TileMappingQueue;
+	FD3D12Fence TileMappingFence;
+
 private:
 	// called by SetupAfterDeviceCreation() when the device gets initialized
 	void CreateDefaultViews();
diff --git a/Engine/Source/Runtime/D3D12RHI/Private/D3D12RHIPrivate.h b/Engine/Source/Runtime/D3D12RHI/Private/D3D12RHIPrivate.h
--- a/Engine/Source/Runtime/D3D12RHI/Private/D3D12RHIPrivate.h
+++ b/Engine/Source/Runtime/D3D12RHI/Private/D3D12RHIPrivate.h
@@ -266,6 +266,8 @@ public:
 		return Object ? static_cast<TReturnType*>(Object->GetLinkedObject(GPUIndex)) : nullptr;
 	}
 
+	virtual bool QueueSupportsTileMapping(ED3D12QueueType /*InQueueType*/) { return true; }
+
 	virtual FD3D12CommandContext* CreateCommandContext(FD3D12Device* InParent, ED3D12QueueType InQueueType, bool InIsDefaultContext);
 	virtual void CreateCommandQueue(FD3D12Device* Device, const D3D12_COMMAND_QUEUE_DESC& Desc, TRefCountPtr<ID3D12CommandQueue>& OutCommandQueue);
 
diff --git a/Engine/Source/Runtime/D3D12RHI/Private/D3D12Submission.cpp b/Engine/Source/Runtime/D3D12RHI/Private/D3D12Submission.cpp
--- a/Engine/Source/Runtime/D3D12RHI/Private/D3D12Submission.cpp
+++ b/Engine/Source/Runtime/D3D12RHI/Private/D3D12Submission.cpp
@@ -778,9 +778,38 @@ uint64 FD3D12Queue::ExecutePayload()
 
 	PayloadToSubmit->PreExecute();
 
-	for (const FD3D12CommitReservedResourceDesc& CommitDesc : PayloadToSubmit->ReservedResourcesToCommit)
+	if (!PayloadToSubmit->ReservedResourcesToCommit.IsEmpty())
 	{
-		CommitDesc.Resource->CommitReservedResource(D3DCommandQueue, CommitDesc.CommitSizeInBytes);
+		// On some devices, some queues cannot perform tile remapping operations.
+		// We can work around this limitation by running the remapping in lockstep on another queue:
+		// - tile mapping queue waits for commands on this queue to finish
+		// - tile mapping queue performs the commit/decommit operations
+		// - this queue waits for tile mapping queue to finish
+		// The extra sync is not required when the current queue is capable of the remapping operations.
+
+		ID3D12CommandQueue* TileMappingQueue = (bSupportsTileMapping ? D3DCommandQueue : Device->TileMappingQueue).GetReference();
+		FD3D12Fence& TileMappingFence = Device->TileMappingFence;
+
+		const bool bCrossQueueSyncRequired = TileMappingQueue != D3DCommandQueue.GetReference();
+
+		if (bCrossQueueSyncRequired)
+		{
+			// tile mapping queue waits for commands on this queue to finish
+			D3DCommandQueue->Signal(TileMappingFence.D3DFence, ++TileMappingFence.LastSignaledValue);
+			TileMappingQueue->Wait(TileMappingFence.D3DFence, TileMappingFence.LastSignaledValue);
+		}
+
+		for (const FD3D12CommitReservedResourceDesc& CommitDesc : PayloadToSubmit->ReservedResourcesToCommit)
+		{
+			CommitDesc.Resource->CommitReservedResource(TileMappingQueue, CommitDesc.CommitSizeInBytes);
+		}
+
+		if (bCrossQueueSyncRequired)
+		{
+			// this queue waits for tile mapping operations to finish
+			TileMappingQueue->Signal(TileMappingFence.D3DFence, ++TileMappingFence.LastSignaledValue);
+			D3DCommandQueue->Wait(TileMappingFence.D3DFence, TileMappingFence.LastSignaledValue);
+		}
 	}
 
 	if (const int32 NumCommandLists = PayloadToSubmit->CommandListsToExecute.Num())
