add ability to clear the UAV cache instead of only the SRV cache. This fixes an issue where we would do dispatch,copy,dispatch,copy and the last 2 events would keep the UAV in the copy_source state because we only do transition when we detect a UAV slot is dirty

add slot clear when we copy a texture

diff --git a/Engine/Source/Runtime/D3D12RHI/Private/D3D12Buffer.cpp b/Engine/Source/Runtime/D3D12RHI/Private/D3D12Buffer.cpp
--- a/Engine/Source/Runtime/D3D12RHI/Private/D3D12Buffer.cpp
+++ b/Engine/Source/Runtime/D3D12RHI/Private/D3D12Buffer.cpp
@@ -94,7 +94,7 @@ struct FRHICommandRenameUploadBuffer final : public FRHICommand<FRHICommandRenam
 			const uint32 GPUIndex = 0; // @todo mgpu - seems wrong we're only doing this for the 0th GPU
 
 			FD3D12CommandContext& Context = FD3D12CommandContext::Get(RHICmdList, GPUIndex);
-			Context.ConditionalClearShaderResource(ResourceLocation);
+			Context.ConditionalClearShaderResource(ResourceLocation, EShaderParameterTypeMask::SRVMask);
 		});
 
 #if UE_MEMORY_TRACE_ENABLED
@@ -667,7 +667,7 @@ void* FD3D12DynamicRHI::LockBuffer(FRHICommandListBase& RHICmdList, FD3D12Buffer
 				{
 					FD3D12CommandContext& DefaultContext = Device->GetDefaultCommandContext();
 
-					FScopedResourceBarrier ScopeResourceBarrierSource(DefaultContext, pResource, D3D12_RESOURCE_STATE_COPY_SOURCE, 0);
+					FScopedResourceBarrier ScopeResourceBarrierSource(DefaultContext, pResource, &Buffer->ResourceLocation, D3D12_RESOURCE_STATE_COPY_SOURCE, 0);
 					// Don't need to transition upload heaps
 
 					uint64 SubAllocOffset = Buffer->ResourceLocation.GetOffsetFromBaseOfResource();
@@ -955,8 +955,8 @@ void FD3D12CommandContext::RHICopyBufferRegion(FRHIBuffer* DestBufferRHI, uint64
 	check(DstOffset + NumBytes <= DestBufferDesc.Width);
 	check(SrcOffset + NumBytes <= SourceBufferDesc.Width);
 
-	FScopedResourceBarrier ScopeResourceBarrierSrc(*this, pSourceResource, D3D12_RESOURCE_STATE_COPY_SOURCE, 0);
-	FScopedResourceBarrier ScopeResourceBarrierDst(*this, pDestResource  , D3D12_RESOURCE_STATE_COPY_DEST  , 0);
+	FScopedResourceBarrier ScopeResourceBarrierSrc(*this, pSourceResource, &SourceBuffer->ResourceLocation, D3D12_RESOURCE_STATE_COPY_SOURCE, 0);
+	FScopedResourceBarrier ScopeResourceBarrierDst(*this, pDestResource  , &DestBuffer->ResourceLocation, D3D12_RESOURCE_STATE_COPY_DEST  , 0);
 	FlushResourceBarriers();
 
 	GraphicsCommandList()->CopyBufferRegion(pDestResource->GetResource(), DestBuffer->ResourceLocation.GetOffsetFromBaseOfResource() + DstOffset, pSourceResource->GetResource(), SourceBuffer->ResourceLocation.GetOffsetFromBaseOfResource() + SrcOffset, NumBytes);
diff --git a/Engine/Source/Runtime/D3D12RHI/Private/D3D12CommandContext.cpp b/Engine/Source/Runtime/D3D12RHI/Private/D3D12CommandContext.cpp
--- a/Engine/Source/Runtime/D3D12RHI/Private/D3D12CommandContext.cpp
+++ b/Engine/Source/Runtime/D3D12RHI/Private/D3D12CommandContext.cpp
@@ -713,29 +713,29 @@ void FD3D12CommandContext::ClearState(EClearStateMode Mode)
 	}
 }
 
-void FD3D12CommandContext::ConditionalClearShaderResource(FD3D12ResourceLocation* Resource)
+void FD3D12CommandContext::ConditionalClearShaderResource(FD3D12ResourceLocation* Resource, EShaderParameterTypeMask ShaderParameterTypeMask)
 {
 	check(Resource);
 
 	for (int32 Index = 0; Index < SF_NumStandardFrequencies; Index++)
 	{
-		StateCache.ClearShaderResourceViews(static_cast<EShaderFrequency>(Index), Resource);
+		StateCache.ClearResourceViewCaches(static_cast<EShaderFrequency>(Index), Resource, ShaderParameterTypeMask);
 	}
 }
 
-void FD3D12CommandContext::ClearShaderResources(FD3D12UnorderedAccessView* UAV)
+void FD3D12CommandContext::ClearShaderResources(FD3D12UnorderedAccessView* UAV, EShaderParameterTypeMask ShaderParameterTypeMask)
 {
 	if (UAV)
 	{
-		ConditionalClearShaderResource(UAV->GetResourceLocation());
+		ConditionalClearShaderResource(UAV->GetResourceLocation(), ShaderParameterTypeMask);
 	}
 }
 
-void FD3D12CommandContext::ClearShaderResources(FD3D12BaseShaderResource* Resource)
+void FD3D12CommandContext::ClearShaderResources(FD3D12BaseShaderResource* Resource, EShaderParameterTypeMask ShaderParameterTypeMask)
 {
 	if (Resource)
 	{
-		ConditionalClearShaderResource(&Resource->ResourceLocation);
+		ConditionalClearShaderResource(&Resource->ResourceLocation, ShaderParameterTypeMask);
 	}
 }
 
diff --git a/Engine/Source/Runtime/D3D12RHI/Private/D3D12CommandContext.h b/Engine/Source/Runtime/D3D12RHI/Private/D3D12CommandContext.h
--- a/Engine/Source/Runtime/D3D12RHI/Private/D3D12CommandContext.h
+++ b/Engine/Source/Runtime/D3D12RHI/Private/D3D12CommandContext.h
@@ -160,6 +160,8 @@ public:
 	};
 
 	virtual void ClearState(EClearStateMode ClearStateMode = EClearStateMode::All) {}
+	virtual void ConditionalClearShaderResource(FD3D12ResourceLocation* Resource, EShaderParameterTypeMask ShaderParameterTypeMask) {}
+
 
 	// Inserts a command to signal the specified sync point
 	void SignalSyncPoint(FD3D12SyncPoint* SyncPoint);
@@ -487,9 +489,9 @@ public:
 	}
 
 	virtual void ClearState(EClearStateMode ClearStateMode = EClearStateMode::All) override final;
-	void ConditionalClearShaderResource(FD3D12ResourceLocation* Resource);
-	void ClearShaderResources(FD3D12UnorderedAccessView* UAV);
-	void ClearShaderResources(FD3D12BaseShaderResource* Resource);
+	virtual void ConditionalClearShaderResource(FD3D12ResourceLocation* Resource, EShaderParameterTypeMask ShaderParameterTypeMask) override final;
+	void ClearShaderResources(FD3D12UnorderedAccessView* UAV, EShaderParameterTypeMask ShaderParameterTypeMask);
+	void ClearShaderResources(FD3D12BaseShaderResource* Resource, EShaderParameterTypeMask ShaderParameterTypeMask);
 	void ClearAllShaderResources();
 
 	FD3D12FastConstantAllocator ConstantsAllocator;
diff --git a/Engine/Source/Runtime/D3D12RHI/Private/D3D12Commands.cpp b/Engine/Source/Runtime/D3D12RHI/Private/D3D12Commands.cpp
--- a/Engine/Source/Runtime/D3D12RHI/Private/D3D12Commands.cpp
+++ b/Engine/Source/Runtime/D3D12RHI/Private/D3D12Commands.cpp
@@ -509,7 +509,7 @@ void FD3D12CommandContext::HandleTransientAliasing(const FD3D12TransitionData* T
 			TransitionResource(Resource, D3D12_RESOURCE_STATE_TBD, FinalState, D3D12_RESOURCE_BARRIER_ALL_SUBRESOURCES);
 			
 			// Remove from caches
-			ClearShaderResources(BaseShaderResource);
+			ClearShaderResources(BaseShaderResource, EShaderParameterTypeMask::SRVMask | EShaderParameterTypeMask::UAVMask);
 		}
 	}
 }
@@ -900,13 +900,13 @@ void FD3D12CommandContext::RHISetComputePipelineState(FRHIComputePipelineState*
 
 void FD3D12CommandContext::SetUAVParameter(EShaderFrequency Frequency, uint32 UAVIndex, FD3D12UnorderedAccessView* UAV)
 {
-	ClearShaderResources(UAV);
+	ClearShaderResources(UAV, EShaderParameterTypeMask::SRVMask);
 	StateCache.SetUAV(SF_Pixel, UAVIndex, UAV);
 }
 
 void FD3D12CommandContext::SetUAVParameter(EShaderFrequency Frequency, uint32 UAVIndex, FD3D12UnorderedAccessView* UAV, uint32 InitialCount)
 {
-	ClearShaderResources(UAV);
+	ClearShaderResources(UAV, EShaderParameterTypeMask::SRVMask);
 	StateCache.SetUAV(SF_Pixel, UAVIndex, UAV, InitialCount);
 }
 
@@ -954,7 +954,7 @@ struct FD3D12ResourceBinder
 		FD3D12UnorderedAccessView_RHI* D3D12UnorderedAccessView = FD3D12CommandContext::RetrieveObject<FD3D12UnorderedAccessView_RHI>(InUnorderedAccessView, GpuIndex);
 		if (bClearResources)
 		{
-			Context.ClearShaderResources(D3D12UnorderedAccessView);
+			Context.ClearShaderResources(D3D12UnorderedAccessView, EShaderParameterTypeMask::SRVMask);
 		}
 
 #if PLATFORM_SUPPORTS_BINDLESS_RENDERING
@@ -1288,7 +1288,7 @@ void FD3D12CommandContext::SetRenderTargets(
 		DepthStencilView = NewDepthStencilTarget->GetDepthStencilView(NewDepthStencilTargetRHI->GetDepthStencilAccess());
 
 		// Unbind any shader views of the depth stencil target that are bound.
-		ClearShaderResources(NewDepthStencilTarget);
+		ClearShaderResources(NewDepthStencilTarget, EShaderParameterTypeMask::SRVMask | EShaderParameterTypeMask::UAVMask);
 	}
 
 	// Gather the render target views for the new render targets.
@@ -1306,7 +1306,7 @@ void FD3D12CommandContext::SetRenderTargets(
 			ensureMsgf(RenderTargetView, TEXT("Texture being set as render target has no RTV"));
 
 			// Unbind any shader views of the render target that are bound.
-			ClearShaderResources(NewRenderTarget);
+			ClearShaderResources(NewRenderTarget, EShaderParameterTypeMask::SRVMask | EShaderParameterTypeMask::UAVMask);
 		}
 
 		NewRenderTargetViews[RenderTargetIndex] = RenderTargetView;
@@ -1994,9 +1994,9 @@ void FD3D12CommandContext::UpdateBuffer(FD3D12ResourceLocation* Dest, uint32 Des
 	uint32 DestFullOffset = Dest->GetOffsetFromBaseOfResource() + DestOffset;
 
 	// Clear the resource if still bound to make sure the SRVs are rebound again on next operation (and get correct resource transitions enqueued)
-	ConditionalClearShaderResource(Dest);
+	ConditionalClearShaderResource(Dest, EShaderParameterTypeMask::SRVMask);
 
-	FScopedResourceBarrier ScopeResourceBarrierDest(*this, DestResource, D3D12_RESOURCE_STATE_COPY_DEST, 0);
+	FScopedResourceBarrier ScopeResourceBarrierDest(*this, DestResource, Dest, D3D12_RESOURCE_STATE_COPY_DEST, 0);
 	// Don't need to transition upload heaps
 
 	FlushResourceBarriers();
diff --git a/Engine/Source/Runtime/D3D12RHI/Private/D3D12PoolAllocator.cpp b/Engine/Source/Runtime/D3D12RHI/Private/D3D12PoolAllocator.cpp
--- a/Engine/Source/Runtime/D3D12RHI/Private/D3D12PoolAllocator.cpp
+++ b/Engine/Source/Runtime/D3D12RHI/Private/D3D12PoolAllocator.cpp
@@ -841,8 +841,8 @@ void FD3D12PoolAllocator::FlushPendingCopyOps(FD3D12CommandContext& InCommandCon
 		else
 #endif // D3D12_RHI_RAYTRACING
 		{
-			FScopedResourceBarrier ScopedResourceBarrierSrc(InCommandContext, CopyOperation.SourceResource, D3D12_RESOURCE_STATE_COPY_SOURCE, D3D12_RESOURCE_BARRIER_ALL_SUBRESOURCES);
-			FScopedResourceBarrier ScopedResourceBarrierDst(InCommandContext, CopyOperation.DestResource  , D3D12_RESOURCE_STATE_COPY_DEST  , D3D12_RESOURCE_BARRIER_ALL_SUBRESOURCES);
+			FScopedResourceBarrier ScopedResourceBarrierSrc(InCommandContext, CopyOperation.SourceResource, nullptr, D3D12_RESOURCE_STATE_COPY_SOURCE, D3D12_RESOURCE_BARRIER_ALL_SUBRESOURCES);
+			FScopedResourceBarrier ScopedResourceBarrierDst(InCommandContext, CopyOperation.DestResource  , nullptr, D3D12_RESOURCE_STATE_COPY_DEST  , D3D12_RESOURCE_BARRIER_ALL_SUBRESOURCES);
 
 			InCommandContext.FlushResourceBarriers();
 
diff --git a/Engine/Source/Runtime/D3D12RHI/Private/D3D12RHIPrivate.h b/Engine/Source/Runtime/D3D12RHI/Private/D3D12RHIPrivate.h
--- a/Engine/Source/Runtime/D3D12RHI/Private/D3D12RHIPrivate.h
+++ b/Engine/Source/Runtime/D3D12RHI/Private/D3D12RHIPrivate.h
@@ -639,16 +639,18 @@ class FScopedResourceBarrier
 {
 private:
 	FD3D12ContextCommon&        Context;
-	FD3D12Resource*       const Resource;
+	FD3D12Resource* const Resource;
+	FD3D12ResourceLocation* const ResourceLocation;
 	D3D12_RESOURCE_STATES const DesiredState;
 	uint32                const Subresource;
 
 	bool bRestoreState = false;
 
 public:
-	FScopedResourceBarrier(FD3D12ContextCommon& Context, FD3D12Resource* Resource, D3D12_RESOURCE_STATES DesiredState, uint32 Subresource)
+	FScopedResourceBarrier(FD3D12ContextCommon& Context, FD3D12Resource* Resource, FD3D12ResourceLocation* InResourceLocation, D3D12_RESOURCE_STATES DesiredState, uint32 Subresource)
 		: Context     (Context)
-		, Resource    (Resource)
+		, Resource(Resource)
+		, ResourceLocation(InResourceLocation)
 		, DesiredState(DesiredState)
 		, Subresource (Subresource)
 	{
@@ -681,6 +683,10 @@ public:
 
 			Context.AddTransitionBarrier(Resource, DesiredState, Resource->GetDefaultResourceState(), Subresource);
 		}
+		else if (Resource->RequiresResourceStateTracking() && ResourceLocation)
+		{
+			Context.ConditionalClearShaderResource(ResourceLocation, EShaderParameterTypeMask::SRVMask | EShaderParameterTypeMask::UAVMask);
+		}
 	}
 };
 
diff --git a/Engine/Source/Runtime/D3D12RHI/Private/D3D12RenderTarget.cpp b/Engine/Source/Runtime/D3D12RHI/Private/D3D12RenderTarget.cpp
--- a/Engine/Source/Runtime/D3D12RHI/Private/D3D12RenderTarget.cpp
+++ b/Engine/Source/Runtime/D3D12RHI/Private/D3D12RenderTarget.cpp
@@ -168,7 +168,7 @@ void FD3D12CommandContext::ResolveTextureUsingShader(
 	GraphicsPSOInit.RasterizerState = TStaticRasterizerState<FM_Solid, CM_None>::GetRHI();
 
 	// Make sure the destination is not bound as a shader resource.
-	ClearShaderResources(DestTexture);
+	ClearShaderResources(DestTexture, EShaderParameterTypeMask::SRVMask | EShaderParameterTypeMask::UAVMask);
 
 	// Determine if the entire destination surface is being resolved to.
 	// If the entire surface is being resolved to, then it means we can clear it and signal the driver that it can discard
@@ -279,7 +279,7 @@ void FD3D12CommandContext::ResolveTextureUsingShader(
 		})(RHICmdList);
 	}
 
-	ClearShaderResources(SourceTexture);
+	ClearShaderResources(SourceTexture, EShaderParameterTypeMask::SRVMask | EShaderParameterTypeMask::UAVMask);
 
 	// Reset saved viewport
 	{
@@ -354,8 +354,8 @@ void FD3D12CommandContext::ResolveTexture(UE::RHICore::FResolveTextureInfo Info)
 			int32 DestSubresource   = CalcSubresource(Info.MipLevel, ArraySlice, DestDesc.NumMips);
 			int32 SourceSubresource = CalcSubresource(Info.MipLevel, ArraySlice, SourceDesc.NumMips);
 
-			FScopedResourceBarrier ConditionalScopeResourceBarrierDst(*this, DestTexture->GetResource()  , D3D12_RESOURCE_STATE_RESOLVE_DEST  , DestSubresource  );
-			FScopedResourceBarrier ConditionalScopeResourceBarrierSrc(*this, SourceTexture->GetResource(), D3D12_RESOURCE_STATE_RESOLVE_SOURCE, SourceSubresource);
+			FScopedResourceBarrier ConditionalScopeResourceBarrierDst(*this, DestTexture->GetResource(),   &DestTexture->ResourceLocation,   D3D12_RESOURCE_STATE_RESOLVE_DEST,   DestSubresource);
+			FScopedResourceBarrier ConditionalScopeResourceBarrierSrc(*this, SourceTexture->GetResource(), &SourceTexture->ResourceLocation, D3D12_RESOURCE_STATE_RESOLVE_SOURCE, SourceSubresource);
 
 			FlushResourceBarriers();
 			GraphicsCommandList()->ResolveSubresource(DestResource->GetResource(), DestSubresource, SourceTexture->GetResource()->GetResource(), SourceSubresource, DestFormatTypeless);
@@ -454,7 +454,7 @@ TRefCountPtr<FD3D12Resource> FD3D12DynamicRHI::GetStagingTexture(FRHITexture* Te
 	CD3DX12_TEXTURE_COPY_LOCATION DestCopyLocation(TempTexture2D->GetResource(), DestFootprint);
 	CD3DX12_TEXTURE_COPY_LOCATION SourceCopyLocation(Texture->GetResource()->GetResource(), Subresource);
 
-	FScopedResourceBarrier ScopeResourceBarrierSource(Context, Texture->GetResource(), D3D12_RESOURCE_STATE_COPY_SOURCE, SourceCopyLocation.SubresourceIndex);
+	FScopedResourceBarrier ScopeResourceBarrierSource(Context, Texture->GetResource(), &Texture->ResourceLocation, D3D12_RESOURCE_STATE_COPY_SOURCE, SourceCopyLocation.SubresourceIndex);
 	Context.FlushResourceBarriers();
 	// Upload heap doesn't need to transition
 
@@ -763,7 +763,7 @@ void FD3D12DynamicRHI::ReadSurfaceDataMSAARaw(FRHITexture* TextureRHI, FIntRect
 			SampleIndex
 			);
 
-		FScopedResourceBarrier ScopeResourceBarrierSource(DefaultContext, NonMSAATexture2D, D3D12_RESOURCE_STATE_COPY_SOURCE, SourceCopyLocation.SubresourceIndex);
+		FScopedResourceBarrier ScopeResourceBarrierSource(DefaultContext, NonMSAATexture2D, nullptr, D3D12_RESOURCE_STATE_COPY_SOURCE, SourceCopyLocation.SubresourceIndex);
 		// Upload heap doesn't need to transition
 
 		// Copy the resolved sample data to the staging texture.
@@ -935,7 +935,7 @@ void FD3D12DynamicRHI::RHIReadSurfaceFloatData(FRHITexture* TextureRHI, FIntRect
 	CD3DX12_TEXTURE_COPY_LOCATION SourceCopyLocation(Texture->GetResource()->GetResource(), Subresource);
 
 	{
-		FScopedResourceBarrier ConditionalScopeResourceBarrier(DefaultContext, Texture->GetResource(), D3D12_RESOURCE_STATE_COPY_SOURCE, SourceCopyLocation.SubresourceIndex);
+		FScopedResourceBarrier ConditionalScopeResourceBarrier(DefaultContext, Texture->GetResource(), &Texture->ResourceLocation, D3D12_RESOURCE_STATE_COPY_SOURCE, SourceCopyLocation.SubresourceIndex);
 		// Don't need to transition upload heaps
 
 		DefaultContext.FlushResourceBarriers();
@@ -1044,7 +1044,7 @@ void FD3D12DynamicRHI::RHIRead3DSurfaceFloatData(FRHITexture* TextureRHI, FIntRe
 	CD3DX12_TEXTURE_COPY_LOCATION SourceCopyLocation(Texture->GetResource()->GetResource(), Subresource);
 
 	{
-		FScopedResourceBarrier ConditionalScopeResourceBarrier(DefaultContext, Texture->GetResource(), D3D12_RESOURCE_STATE_COPY_SOURCE, SourceCopyLocation.SubresourceIndex);
+		FScopedResourceBarrier ConditionalScopeResourceBarrier(DefaultContext, Texture->GetResource(), &Texture->ResourceLocation, D3D12_RESOURCE_STATE_COPY_SOURCE, SourceCopyLocation.SubresourceIndex);
 		// Don't need to transition upload heaps
 
 		DefaultContext.FlushResourceBarriers();
diff --git a/Engine/Source/Runtime/D3D12RHI/Private/D3D12StateCache.cpp b/Engine/Source/Runtime/D3D12RHI/Private/D3D12StateCache.cpp
--- a/Engine/Source/Runtime/D3D12RHI/Private/D3D12StateCache.cpp
+++ b/Engine/Source/Runtime/D3D12RHI/Private/D3D12StateCache.cpp
@@ -150,23 +150,37 @@ void FD3D12StateCache::ClearSRVs()
 	bSRVSCleared = true;
 }
 
-void FD3D12StateCache::ClearShaderResourceViews(EShaderFrequency ShaderFrequency, FD3D12ResourceLocation*& ResourceLocation)
+void FD3D12StateCache::ClearResourceViewCaches(EShaderFrequency ShaderFrequency, FD3D12ResourceLocation*& ResourceLocation, EShaderParameterTypeMask ShaderParameterTypeMask)
 {
 	//SCOPE_CYCLE_COUNTER(STAT_D3D12ClearShaderResourceViewsTime);
 
-	if (PipelineState.Common.SRVCache.MaxBoundIndex[ShaderFrequency] < 0)
+	if (EnumHasAnyFlags(ShaderParameterTypeMask, EShaderParameterTypeMask::SRVMask))
 	{
-		return;
+		if (PipelineState.Common.SRVCache.MaxBoundIndex[ShaderFrequency] >= 0)
+		{
+			auto& CurrentShaderResourceViews = PipelineState.Common.SRVCache.Views[ShaderFrequency];
+			for (int32 i = 0; i <= PipelineState.Common.SRVCache.MaxBoundIndex[ShaderFrequency]; ++i)
+			{
+				if (CurrentShaderResourceViews[i] && CurrentShaderResourceViews[i]->GetResourceLocation() == ResourceLocation)
+				{
+					SetShaderResourceView(ShaderFrequency, nullptr, i);
+				}
+			}
+		}
 	}
 
-	auto& CurrentShaderResourceViews = PipelineState.Common.SRVCache.Views[ShaderFrequency];
-	for (int32 i = 0; i <= PipelineState.Common.SRVCache.MaxBoundIndex[ShaderFrequency]; ++i)
+	if (EnumHasAnyFlags(ShaderParameterTypeMask, EShaderParameterTypeMask::UAVMask))
 	{
-		if (CurrentShaderResourceViews[i] && CurrentShaderResourceViews[i]->GetResourceLocation() == ResourceLocation)
+		auto& CurrentShaderResourceViews = PipelineState.Common.UAVCache.Views[ShaderFrequency];
+		for (int32 i = 0; i <= MAX_UAVS; ++i)
 		{
-			SetShaderResourceView(ShaderFrequency, nullptr, i);
+			if (CurrentShaderResourceViews[i] && CurrentShaderResourceViews[i]->GetResourceLocation() == ResourceLocation)
+			{
+				SetUAV(ShaderFrequency, i, nullptr);
+			}
 		}
 	}
+
 }
 
 void FD3D12StateCache::FlushComputeShaderCache(bool bForce)
diff --git a/Engine/Source/Runtime/D3D12RHI/Private/D3D12StateCachePrivate.h b/Engine/Source/Runtime/D3D12RHI/Private/D3D12StateCachePrivate.h
--- a/Engine/Source/Runtime/D3D12RHI/Private/D3D12StateCachePrivate.h
+++ b/Engine/Source/Runtime/D3D12RHI/Private/D3D12StateCachePrivate.h
@@ -454,7 +454,7 @@ public:
 
 	void ClearSRVs();
 
-	void ClearShaderResourceViews(EShaderFrequency ShaderFrequency, FD3D12ResourceLocation*& ResourceLocation);
+	void ClearResourceViewCaches(EShaderFrequency ShaderFrequency, FD3D12ResourceLocation*& ResourceLocation, EShaderParameterTypeMask ShaderParameterTypeMask);
 
 	void SetShaderResourceView(EShaderFrequency ShaderFrequency, FD3D12ShaderResourceView* SRV, uint32 ResourceIndex);
 	
diff --git a/Engine/Source/Runtime/D3D12RHI/Private/D3D12Texture.cpp b/Engine/Source/Runtime/D3D12RHI/Private/D3D12Texture.cpp
--- a/Engine/Source/Runtime/D3D12RHI/Private/D3D12Texture.cpp
+++ b/Engine/Source/Runtime/D3D12RHI/Private/D3D12Texture.cpp
@@ -236,8 +236,8 @@ struct FRHICommandD3D12AsyncReallocateTexture2D final : public FRHICommand<FRHIC
 
 			FD3D12CommandContext& Context = Device->GetDefaultCommandContext();
 
-			FScopedResourceBarrier ScopeResourceBarrierDst(Context, NewTexture->GetResource(), D3D12_RESOURCE_STATE_COPY_DEST  , D3D12_RESOURCE_BARRIER_ALL_SUBRESOURCES);
-			FScopedResourceBarrier ScopeResourceBarrierSrc(Context, OldTexture->GetResource(), D3D12_RESOURCE_STATE_COPY_SOURCE, D3D12_RESOURCE_BARRIER_ALL_SUBRESOURCES);
+			FScopedResourceBarrier ScopeResourceBarrierDst(Context, NewTexture->GetResource(), &NewTexture->ResourceLocation, D3D12_RESOURCE_STATE_COPY_DEST  , D3D12_RESOURCE_BARRIER_ALL_SUBRESOURCES);
+			FScopedResourceBarrier ScopeResourceBarrierSrc(Context, OldTexture->GetResource(), &OldTexture->ResourceLocation, D3D12_RESOURCE_STATE_COPY_SOURCE, D3D12_RESOURCE_BARRIER_ALL_SUBRESOURCES);
 			Context.FlushResourceBarriers();	// Must flush so the desired state is actually set.
 
 			for (uint32 MipIndex = 0; MipIndex < NumSharedMips; ++MipIndex)
@@ -1801,7 +1801,7 @@ void* FD3D12Texture::Lock(class FRHICommandListImmediate* RHICmdList, uint32 Mip
 
 		const auto& pfnCopyTextureRegion = [&]()
 		{
-			FScopedResourceBarrier ScopeResourceBarrierSource(Context, GetResource(), D3D12_RESOURCE_STATE_COPY_SOURCE, SourceCopyLocation.SubresourceIndex);
+			FScopedResourceBarrier ScopeResourceBarrierSource(Context, GetResource(), &ResourceLocation, D3D12_RESOURCE_STATE_COPY_SOURCE, SourceCopyLocation.SubresourceIndex);
 
 			Context.FlushResourceBarriers();
 			Context.GraphicsCommandList()->CopyTextureRegion(
@@ -1844,7 +1844,7 @@ void FD3D12Texture::UpdateTexture(uint32 MipIndex, uint32 DestX, uint32 DestY, u
 	LLM_SCOPE_BYNAME(TEXT("D3D12CopyTextureRegion"));
 	FD3D12CommandContext& DefaultContext = GetParentDevice()->GetDefaultCommandContext();
 
-	FScopedResourceBarrier ScopeResourceBarrierDest(DefaultContext, GetResource(), D3D12_RESOURCE_STATE_COPY_DEST, MipIndex);
+	FScopedResourceBarrier ScopeResourceBarrierDest(DefaultContext, GetResource(), &ResourceLocation, D3D12_RESOURCE_STATE_COPY_DEST, MipIndex);
 	// Don't need to transition upload heaps
 
 	CD3DX12_TEXTURE_COPY_LOCATION DestCopyLocation(GetResource()->GetResource(), MipIndex);
@@ -1870,8 +1870,8 @@ void FD3D12Texture::CopyTextureRegion(uint32 DestX, uint32 DestY, uint32 DestZ,
 	CD3DX12_TEXTURE_COPY_LOCATION DestCopyLocation(GetResource()->GetResource(), 0);
 	CD3DX12_TEXTURE_COPY_LOCATION SourceCopyLocation(SourceTexture->GetResource()->GetResource(), 0);
 
-	FScopedResourceBarrier ConditionalScopeResourceBarrierDst(DefaultContext, GetResource(), D3D12_RESOURCE_STATE_COPY_DEST, DestCopyLocation.SubresourceIndex);
-	FScopedResourceBarrier ConditionalScopeResourceBarrierSrc(DefaultContext, SourceTexture->GetResource(), D3D12_RESOURCE_STATE_COPY_SOURCE, SourceCopyLocation.SubresourceIndex);
+	FScopedResourceBarrier ConditionalScopeResourceBarrierDst(DefaultContext, GetResource(), &ResourceLocation, D3D12_RESOURCE_STATE_COPY_DEST, DestCopyLocation.SubresourceIndex);
+	FScopedResourceBarrier ConditionalScopeResourceBarrierSrc(DefaultContext, SourceTexture->GetResource(), &SourceTexture->ResourceLocation, D3D12_RESOURCE_STATE_COPY_SOURCE, SourceCopyLocation.SubresourceIndex);
 
 	DefaultContext.FlushResourceBarriers();
 	DefaultContext.GraphicsCommandList()->CopyTextureRegion(
@@ -2298,7 +2298,7 @@ public:
 
 			CD3DX12_TEXTURE_COPY_LOCATION DestCopyLocation(TextureLink.GetResource()->GetResource(), MipIdx);
 
-			FScopedResourceBarrier ScopeResourceBarrierDest(Context, TextureLink.GetResource(), D3D12_RESOURCE_STATE_COPY_DEST, DestCopyLocation.SubresourceIndex);
+			FScopedResourceBarrier ScopeResourceBarrierDest(Context, TextureLink.GetResource(), &TextureLink.ResourceLocation, D3D12_RESOURCE_STATE_COPY_DEST, DestCopyLocation.SubresourceIndex);
 
 			Context.FlushResourceBarriers();
 
@@ -2544,7 +2544,7 @@ public:
 			CD3DX12_TEXTURE_COPY_LOCATION DestCopyLocation(TextureLink.GetResource()->GetResource(), MipIdx);
 			CD3DX12_TEXTURE_COPY_LOCATION SourceCopyLocation(UploadBuffer->GetResource(), PlacedSubresourceFootprint);
 
-			FScopedResourceBarrier ScopeResourceBarrierDest(Context, TextureLink.GetResource(), D3D12_RESOURCE_STATE_COPY_DEST, DestCopyLocation.SubresourceIndex);
+			FScopedResourceBarrier ScopeResourceBarrierDest(Context, TextureLink.GetResource(), &TextureLink.ResourceLocation, D3D12_RESOURCE_STATE_COPY_DEST, DestCopyLocation.SubresourceIndex);
 
 			Context.FlushResourceBarriers();
 			Context.GraphicsCommandList()->CopyTextureRegion(
@@ -2615,7 +2615,7 @@ void* FD3D12DynamicRHI::RHILockTextureCubeFace(FRHITextureCube* TextureCubeRHI,
 	FD3D12Texture*  TextureCube = FD3D12DynamicRHI::ResourceCast(TextureCubeRHI);
 	for (uint32 GPUIndex : TextureCube->GetLinkedObjectsGPUMask())
 	{
-		GetRHIDevice(GPUIndex)->GetDefaultCommandContext().ConditionalClearShaderResource(&TextureCube->ResourceLocation);
+		GetRHIDevice(GPUIndex)->GetDefaultCommandContext().ConditionalClearShaderResource(&TextureCube->ResourceLocation, EShaderParameterTypeMask::SRVMask | EShaderParameterTypeMask::UAVMask);
 	}
 	uint32 D3DFace = GetD3D12CubeFace((ECubeFace)FaceIndex);
 	return TextureCube->Lock(nullptr, MipIndex, D3DFace + ArrayIndex * 6, LockMode, DestStride);
@@ -2852,8 +2852,8 @@ void FD3D12CommandContext::RHICopyTexture(FRHITexture* SourceTextureRHI, FRHITex
 	FD3D12Texture* SourceTexture = RetrieveTexture(SourceTextureRHI);
 	FD3D12Texture* DestTexture = RetrieveTexture(DestTextureRHI);
 
-	FScopedResourceBarrier ConditionalScopeResourceBarrierSrc(*this, SourceTexture->GetResource(), D3D12_RESOURCE_STATE_COPY_SOURCE, D3D12_RESOURCE_BARRIER_ALL_SUBRESOURCES);
-	FScopedResourceBarrier ConditionalScopeResourceBarrierDst(*this, DestTexture->GetResource()  , D3D12_RESOURCE_STATE_COPY_DEST  , D3D12_RESOURCE_BARRIER_ALL_SUBRESOURCES);
+	FScopedResourceBarrier ConditionalScopeResourceBarrierSrc(*this, SourceTexture->GetResource(), &SourceTexture->ResourceLocation, D3D12_RESOURCE_STATE_COPY_SOURCE, D3D12_RESOURCE_BARRIER_ALL_SUBRESOURCES);
+	FScopedResourceBarrier ConditionalScopeResourceBarrierDst(*this, DestTexture->GetResource(), &DestTexture->ResourceLocation, D3D12_RESOURCE_STATE_COPY_DEST  , D3D12_RESOURCE_BARRIER_ALL_SUBRESOURCES);
 
 	FlushResourceBarriers();
 
diff --git a/Engine/Source/Runtime/RenderCore/Public/ShaderCore.h b/Engine/Source/Runtime/RenderCore/Public/ShaderCore.h
--- a/Engine/Source/Runtime/RenderCore/Public/ShaderCore.h
+++ b/Engine/Source/Runtime/RenderCore/Public/ShaderCore.h
@@ -226,6 +226,19 @@ enum class EShaderParameterType : uint8
 	Num
 };
 
+enum class EShaderParameterTypeMask : uint16
+{
+	LooseDataMask = 1 << uint16(EShaderParameterType::LooseData),
+	UniformBufferMask = 1 << uint16(EShaderParameterType::UniformBuffer),
+	SamplerMask = 1 << uint16(EShaderParameterType::Sampler),
+	SRVMask = 1 << uint16(EShaderParameterType::SRV),
+	UAVMask = 1 << uint16(EShaderParameterType::UAV),
+	BindlessSamplerMask = 1 << uint16(EShaderParameterType::BindlessSampler),
+	BindlessSRVMask = 1 << uint16(EShaderParameterType::BindlessSRV),
+	BindlessUAVMask = 1 << uint16(EShaderParameterType::BindlessUAV),
+};
+ENUM_CLASS_FLAGS(EShaderParameterTypeMask);
+
 inline bool IsParameterBindless(EShaderParameterType ParameterType)
 {
 	return ParameterType == EShaderParameterType::BindlessSampler
