fix wrong cast for texture with >= 256 subresources, for example, a 128x128 texture array with 8 mips and 32 slices

diff --git a/Engine/Source/Runtime/D3D12RHI/Private/D3D12CommandList.cpp b/Engine/Source/Runtime/D3D12RHI/Private/D3D12CommandList.cpp
--- a/Engine/Source/Runtime/D3D12RHI/Private/D3D12CommandList.cpp
+++ b/Engine/Source/Runtime/D3D12RHI/Private/D3D12CommandList.cpp
@@ -585,7 +585,7 @@ bool FD3D12ContextCommon::TransitionResource(FD3D12Resource* Resource, D3D12_RES
 	{
 		// Slow path. Want to transition the entire resource (with multiple subresources). But they aren't in the same state.
 
-		const uint8 SubresourceCount = Resource->GetSubresourceCount();
+		const uint32 SubresourceCount = Resource->GetSubresourceCount();
 		for (uint32 SubresourceIndex = 0; SubresourceIndex < SubresourceCount; SubresourceIndex++)
 		{
 			bool bForceInAfterState = true;
