Addressing issue where a UAV clear isn't given a context that was already open. The bindless code was assuming the context was already open and was handing around a null heap because of this.

diff --git a/Engine/Source/Runtime/D3D12RHI/Private/D3D12CommandContext.cpp b/Engine/Source/Runtime/D3D12RHI/Private/D3D12CommandContext.cpp
--- a/Engine/Source/Runtime/D3D12RHI/Private/D3D12CommandContext.cpp
+++ b/Engine/Source/Runtime/D3D12RHI/Private/D3D12CommandContext.cpp
@@ -405,6 +405,16 @@ void FD3D12CommandContext::Finalize(TArray<FD3D12Payload*>& OutPayloads)
 	FD3D12ContextCommon::Finalize(OutPayloads);
 }
 
+#if PLATFORM_SUPPORTS_BINDLESS_RENDERING
+FD3D12DescriptorHeap* FD3D12CommandContext::GetBindlessResourcesHeap()
+{
+	// We require the descriptor cache to be setup correctly before it can have a valid bindless heap.
+	OpenIfNotAlready();
+
+	return StateCache.GetDescriptorCache()->GetBindlessResourcesHeap();
+}
+#endif
+
 FD3D12QueryLocation FD3D12QueryAllocator::Allocate(ED3D12QueryType Type, void* Target)
 {
 	check(Type != ED3D12QueryType::None);
diff --git a/Engine/Source/Runtime/D3D12RHI/Private/D3D12CommandContext.h b/Engine/Source/Runtime/D3D12RHI/Private/D3D12CommandContext.h
--- a/Engine/Source/Runtime/D3D12RHI/Private/D3D12CommandContext.h
+++ b/Engine/Source/Runtime/D3D12RHI/Private/D3D12CommandContext.h
@@ -254,10 +254,7 @@ private:
 	// Returns the current command list (or creates a new one if the command list was not open).
 	FD3D12CommandList& GetCommandList()
 	{
-		if (!CommandList)
-		{
-			OpenCommandList();
-		}
+		OpenIfNotAlready();
 
 		return *CommandList;
 	}
@@ -293,6 +290,15 @@ protected:
 
 	uint32 ActiveQueries = 0;
 
+	// Open the command list if it's not already open.
+	void OpenIfNotAlready()
+	{
+		if (!CommandList)
+		{
+			OpenCommandList();
+		}
+	}
+
 public:
 	// Flushes any pending commands in this context to the GPU.
 	void FlushCommands(ED3D12FlushFlags FlushFlags = ED3D12FlushFlags::None);
@@ -749,6 +755,7 @@ public:
 	virtual void Finalize(TArray<FD3D12Payload*>& OutPayloads) override;
 
 #if PLATFORM_SUPPORTS_BINDLESS_RENDERING
+	FD3D12DescriptorHeap* GetBindlessResourcesHeap();
 	FD3D12ContextBindlessState& GetBindlessState() { return BindlessState; }
 #endif
 
diff --git a/Engine/Source/Runtime/D3D12RHI/Private/D3D12UAV.cpp b/Engine/Source/Runtime/D3D12RHI/Private/D3D12UAV.cpp
--- a/Engine/Source/Runtime/D3D12RHI/Private/D3D12UAV.cpp
+++ b/Engine/Source/Runtime/D3D12RHI/Private/D3D12UAV.cpp
@@ -370,7 +370,7 @@ void FD3D12CommandContext::ClearUAV(TRHICommandList_RecursiveHazardous<FD3D12Com
 				{
 					Context.FlushPendingDescriptorUpdates();
 
-					FD3D12DescriptorHeap* BindlessHeap = Context.StateCache.GetDescriptorCache()->GetBindlessResourcesHeap();
+					FD3D12DescriptorHeap* BindlessHeap = Context.GetBindlessResourcesHeap();
 					UE::D3D12Descriptors::CopyDescriptor(ParentDevice, BindlessHeap, UAV.GetBindlessHandle(), OfflineHandle);
 					GPUHandle = BindlessHeap->GetGPUSlotHandle(UAV.GetBindlessHandle().GetIndex());
 				}
