Extend DXGI workaround to enforce atomic swap chain creation

diff --git a/Engine/Source/Runtime/D3D12RHI/Private/Windows/WindowsD3D12Viewport.cpp b/Engine/Source/Runtime/D3D12RHI/Private/Windows/WindowsD3D12Viewport.cpp
--- a/Engine/Source/Runtime/D3D12RHI/Private/Windows/WindowsD3D12Viewport.cpp
+++ b/Engine/Source/Runtime/D3D12RHI/Private/Windows/WindowsD3D12Viewport.cpp
@@ -171,9 +171,15 @@ void FD3D12Viewport::Init()
 
 			TRefCountPtr<IDXGISwapChain> SwapChain;
 			
-			HRESULT hr = DXGISwapchainProvider ?
-				DXGISwapchainProvider->CreateSwapChain(Factory2, CommandQueue, &SwapChainDesc, SwapChain.GetInitReference()) :
-				Factory2->CreateSwapChain(CommandQueue, &SwapChainDesc, SwapChain.GetInitReference());
+			HRESULT hr;
+						
+			{
+				// Don't create swap chain and release back buffer at the same time (see notes on critical section)
+				FScopeLock Lock(&DXGIBackBufferLock); 
+				hr = DXGISwapchainProvider ?
+					DXGISwapchainProvider->CreateSwapChain(Factory2, CommandQueue, &SwapChainDesc, SwapChain.GetInitReference()) :
+					Factory2->CreateSwapChain(CommandQueue, &SwapChainDesc, SwapChain.GetInitReference());
+			}
 			
 			if (FAILED(hr))
 			{
