fix backbuffer resize: we were trying to release the swapchain buffers before their UAV was actually released.

Moving the UAV deletion code just above doesn't work though: FD3D12Viewport::Resize is called from the game thread. When destroyed, Backbuffer UAVs are pushed to the RHI's PendingDeletes queue before being actually released.

Because of that, we also need to ensure the queue is processed before trying to destroy the swapchain buffers. ImmediateFlush(EImmediateFlushType::FlushRHIThreadFlushResources) takes care of that, but FlushRenderingCommands() already does it + add the GameThread-RenderThread synchronization

diff --git a/Engine/Source/Runtime/D3D12RHI/Private/D3D12Viewport.cpp b/Engine/Source/Runtime/D3D12RHI/Private/D3D12Viewport.cpp
--- a/Engine/Source/Runtime/D3D12RHI/Private/D3D12Viewport.cpp
+++ b/Engine/Source/Runtime/D3D12RHI/Private/D3D12Viewport.cpp
@@ -410,44 +410,56 @@ void FD3D12Viewport::Resize(uint32 InSizeX, uint32 InSizeY, bool bInIsFullscreen
 		CustomPresent->OnBackBufferResize();
 	}
 
+	bool bWaitForBackBuffersUAVDelete = false;
 	// Release our backbuffer reference, as required by DXGI before calling ResizeBuffers.
 	for (uint32 i = 0; i < NumBackBuffers; ++i)
 	{
-		if (IsValidRef(BackBuffers[i]))
+		if (IsValidRef(BackBuffersUAV[i]))
 		{
+			bWaitForBackBuffersUAVDelete = true;
 			// Tell the back buffer to delete immediately so that we can call resize.
-			if (BackBuffers[i]->GetRefCount() != 1)
+			if (BackBuffersUAV[i]->GetRefCount() != 1)
 			{
-				UE_LOG(LogD3D12RHI, Log, TEXT("Backbuffer %d leaking with %d refs during Resize."), i, BackBuffers[i]->GetRefCount());
+				UE_LOG(LogD3D12RHI, Log, TEXT("Backbuffer %d leaking with %d refs during Resize."), i, BackBuffersUAV[i]->GetRefCount());
 			}
-			check(BackBuffers[i]->GetRefCount() == 1);
+			check(BackBuffersUAV[i]->GetRefCount() == 1);
 
-			for (FD3D12Texture& Tex : *BackBuffers[i])
+			for (FD3D12UnorderedAccessView& Uav : *BackBuffersUAV[i])
 			{
-				Tex.GetResource()->DoNotDeferDelete();
+				Uav.GetResource()->DoNotDeferDelete();
 			}
 		}
-		
-		BackBuffers[i].SafeRelease();
-		check(BackBuffers[i] == nullptr);
 
-		if (IsValidRef(BackBuffersUAV[i]))
+		BackBuffersUAV[i].SafeRelease();
+		check(BackBuffersUAV[i] == nullptr);
+	}
+
+	if (bWaitForBackBuffersUAVDelete)
+	{
+		// The D3D12 UAV releases don't happen immediately, but are pushed to a delete queue processed on the RHI Thread. We need to ensure these are processed before releasing the swapchain buffers
+        // Calling FlushRenderingCommands is enough because it calls ImmediateFlush(EImmediateFlushType::FlushRHIThreadFlushResources) / ImmediateFlush(EImmediateFlushType::FlushRHIThread) internally
+		FlushRenderingCommands();
+	}
+
+	for (uint32 i = 0; i < NumBackBuffers; ++i)
+	{
+		if (IsValidRef(BackBuffers[i]))
 		{
 			// Tell the back buffer to delete immediately so that we can call resize.
-			if (BackBuffersUAV[i]->GetRefCount() != 1)
+			if (BackBuffers[i]->GetRefCount() != 1)
 			{
-				UE_LOG(LogD3D12RHI, Log, TEXT("Backbuffer %d leaking with %d refs during Resize."), i, BackBuffersUAV[i]->GetRefCount());
+				UE_LOG(LogD3D12RHI, Log, TEXT("Backbuffer %d leaking with %d refs during Resize."), i, BackBuffers[i]->GetRefCount());
 			}
-			check(BackBuffersUAV[i]->GetRefCount() == 1);
+			check(BackBuffers[i]->GetRefCount() == 1);
 
-			for (FD3D12UnorderedAccessView& Uav : *BackBuffersUAV[i])
+			for (FD3D12Texture& Tex : *BackBuffers[i])
 			{
-				Uav.GetResource()->DoNotDeferDelete();
+				Tex.GetResource()->DoNotDeferDelete();
 			}
 		}
 
-		BackBuffersUAV[i].SafeRelease();
-		check(BackBuffersUAV[i] == nullptr);
+		BackBuffers[i].SafeRelease();
+		check(BackBuffers[i] == nullptr);
 
 		if (IsValidRef(SDRBackBuffers[i]))
 		{
