Update PSO driver cache clearing behavior for new NVIDIA cache location.
#tests tested locally

diff --git a/Engine/Source/Runtime/D3D12RHI/Private/Windows/WindowsD3D12Device.cpp b/Engine/Source/Runtime/D3D12RHI/Private/Windows/WindowsD3D12Device.cpp
--- a/Engine/Source/Runtime/D3D12RHI/Private/Windows/WindowsD3D12Device.cpp
+++ b/Engine/Source/Runtime/D3D12RHI/Private/Windows/WindowsD3D12Device.cpp
@@ -1307,9 +1307,14 @@ static void ClearPSODriverCache()
 
 		if (IsRHIDeviceNVIDIA())
 		{
-			FString PSOPath = FPaths::Combine(*LocalAppDataFolder, TEXT("NVIDIA"), TEXT("DXCache"));
-			ClearFolder(PSOPath, TEXT(".bin"));
-			ClearFolder(PSOPath, TEXT(".toc"));
+			// NVIDIA used to have a global cache, but now also has a per-driver cache in a different folder in LocalLow.
+			FString GlobalPSOPath = FPaths::Combine(*LocalAppDataFolder, TEXT("NVIDIA"), TEXT("DXCache"));
+			ClearFolder(GlobalPSOPath, TEXT(".bin"));
+			ClearFolder(GlobalPSOPath, TEXT(".toc"));
+
+			FString PerDriverPSOPath = FPaths::Combine(*LocalAppDataFolder, TEXT(".."), TEXT("LocalLow"), TEXT("NVIDIA"), TEXT("PerDriverVersion"), TEXT("DXCache"));
+			ClearFolder(PerDriverPSOPath, TEXT(".bin"));
+			ClearFolder(PerDriverPSOPath, TEXT(".toc"));
 		}
 		else if (IsRHIDeviceAMD())
 		{
