Fix logic mismatch between FD3D12Adapter::CreateCommittedResource() which calls StartTrackingForResidency() and initialization of bRequiresResidencyTracking in FD3D12Resource()

diff --git a/Engine/Source/Runtime/D3D12RHI/Private/D3D12Resources.cpp b/Engine/Source/Runtime/D3D12RHI/Private/D3D12Resources.cpp
--- a/Engine/Source/Runtime/D3D12RHI/Private/D3D12Resources.cpp
+++ b/Engine/Source/Runtime/D3D12RHI/Private/D3D12Resources.cpp
@@ -109,7 +109,10 @@ FD3D12Resource::FD3D12Resource(FD3D12Device* ParentDevice,
 	}
 
 #if ENABLE_RESIDENCY_MANAGEMENT
-	// Residency tracking is only used for GPU-only resources owned by the Engine
+	// Residency tracking is only used for GPU-only resources owned by the Engine.
+	// Back buffers may be referenced outside of command lists (during presents), however D3DX12Residency.h library
+	// uses fences tied to command lists to detect when it's safe to evict a resource, which is wrong for back buffers.
+	// External/shared resources may be referenced by command buffers in third-party code.
 	bRequiresResidencyTracking = IsGPUOnly(InHeapType, HeapProps) && !Desc.bExternal && !Desc.bBackBuffer;
 #endif
 
@@ -520,17 +523,15 @@ void FD3D12Resource::StartTrackingForResidency()
 {
 #if ENABLE_RESIDENCY_MANAGEMENT
 
-	checkf(bRequiresResidencyTracking, TEXT("Residency tracking is not expected for this resource"));
-
-	if (Desc.bBackBuffer)
+	if (!bRequiresResidencyTracking)
 	{
-		// Back buffers may be referenced outside of command lists (during presents), however D3DX12Residency.h library 
-		// uses fences tied to command lists to detect when it's safe to evict a resource, which is wrong for back buffers.
-		// Simply disable residency tracking for back buffers as a workaround (keep them always resident).
 		return;
 	}
 
-	check(IsGPUOnly(HeapType));	// This is checked at a higher level before calling this function.
+	checkf(IsGPUOnly(HeapType), TEXT("Residency tracking is not expected for CPU-accessible resources"));
+	checkf(!Desc.bBackBuffer, TEXT("Residency tracking is not expected for back buffers"));
+	checkf(!Desc.bExternal, TEXT("Residency tracking is not expected for externally-owned resources"));
+
 	if (!IsPlacedResource() && !IsReservedResource())
 	{
 		checkf(!ResidencyHandle, TEXT("Residency tracking is already initialzied for this resource"));
@@ -759,11 +760,7 @@ HRESULT FD3D12Adapter::CreateCommittedResource(const FD3D12ResourceDesc& InDesc,
 		// Set a default name (can override later).
 		SetName(*ppOutResource, Name);
 
-		// Only track resources that cannot be accessed on the CPU.
-		if (IsGPUOnly(HeapProps.Type, &HeapProps))
-		{
-			(*ppOutResource)->StartTrackingForResidency();
-		}
+		(*ppOutResource)->StartTrackingForResidency();
 
 		TraceMemoryAllocation(*ppOutResource);
 	}
