#UE Fix shadow variable

diff --git a/Engine/Source/Runtime/D3D12RHI/Private/D3D12Submission.cpp b/Engine/Source/Runtime/D3D12RHI/Private/D3D12Submission.cpp
--- a/Engine/Source/Runtime/D3D12RHI/Private/D3D12Submission.cpp
+++ b/Engine/Source/Runtime/D3D12RHI/Private/D3D12Submission.cpp
@@ -753,15 +753,15 @@ void FD3D12Queue::FlushBatchedPayloads()
 	auto Wait = [this](FD3D12Payload* Payload)
 	{
 		// Wait for queue fences
-		for (auto& [Fence, Value] : Payload->QueueFencesToWait)
+		for (auto& [LocalFence, Value] : Payload->QueueFencesToWait)
 		{
-			VERIFYD3D12RESULT(D3DCommandQueue->Wait(Fence.D3DFence, Value));
+			VERIFYD3D12RESULT(D3DCommandQueue->Wait(LocalFence.D3DFence, Value));
 		}
 
 		// Wait for manual fences
-		for (auto& [Fence, Value] : Payload->ManualFencesToWait)
+		for (auto& [LocalFence, Value] : Payload->ManualFencesToWait)
 		{
-			VERIFYD3D12RESULT(D3DCommandQueue->Wait(Fence, Value));
+			VERIFYD3D12RESULT(D3DCommandQueue->Wait(LocalFence, Value));
 		}
 	};
 
