Fix static analysis warnings

diff --git a/Engine/Source/Runtime/D3D12RHI/Private/D3D12Adapter.cpp b/Engine/Source/Runtime/D3D12RHI/Private/D3D12Adapter.cpp
--- a/Engine/Source/Runtime/D3D12RHI/Private/D3D12Adapter.cpp
+++ b/Engine/Source/Runtime/D3D12RHI/Private/D3D12Adapter.cpp
@@ -94,7 +94,7 @@ static TAutoConsoleVariable<int32> CVarD3D12EnableDRED(
 
 static TAutoConsoleVariable<int32> CVarD3D12EnableLightweightDRED(
 	TEXT("r.D3D12.LightweightDRED"),
-	(UE_BUILD_SHIPPING || UE_BUILD_TEST) ? 1 : 1,
+	1,
 	TEXT("Enable Lightweight DRED GPU Crash debugging mode to track the current GPU state and logs information what operations the GPU executed last.")
 	TEXT("Gives the basic information on the current GPU state when it crashes or hangs on all PC hardware.\n"),
 	ECVF_RenderThreadSafe | ECVF_ReadOnly);
diff --git a/Engine/Source/Runtime/RHI/Private/RHICommandList.cpp b/Engine/Source/Runtime/RHI/Private/RHICommandList.cpp
--- a/Engine/Source/Runtime/RHI/Private/RHICommandList.cpp
+++ b/Engine/Source/Runtime/RHI/Private/RHICommandList.cpp
@@ -663,7 +663,7 @@ void FRHICommandListExecutor::FSubmitState::Dispatch(FRHICommandListBase* CmdLis
 			) || (
 				MaxCommandsPerTranslate  < 0 ? true  :
 				MaxCommandsPerTranslate == 0 ? false :
-				(CurrentTranslateJob->NumCommands + CmdList->NumCommands) > MaxCommandsPerTranslate
+				(CurrentTranslateJob->NumCommands + CmdList->NumCommands) > uint32(MaxCommandsPerTranslate)
 			))
 		{
 			FinalizeCurrent();
diff --git a/Engine/Source/Runtime/RHI/Public/RHICommandList.h b/Engine/Source/Runtime/RHI/Public/RHICommandList.h
--- a/Engine/Source/Runtime/RHI/Public/RHICommandList.h
+++ b/Engine/Source/Runtime/RHI/Public/RHICommandList.h
@@ -4956,7 +4956,7 @@ private:
 		FGraphEventArray TranslateEvents;
 		FTranslateState* CurrentTranslateJob = nullptr;
 
-		uint32 MaxCommandsPerTranslate = 0;
+		int32 MaxCommandsPerTranslate = 0;
 		bool bAllowSingleParallelCombine = false;
 		bool bAllowParallelTranslate = true;
 
