Add a null-check and a log for root signatures at pipeline creation time.

diff --git a/Engine/Source/Runtime/D3D12RHI/Private/D3D12State.cpp b/Engine/Source/Runtime/D3D12RHI/Private/D3D12State.cpp
--- a/Engine/Source/Runtime/D3D12RHI/Private/D3D12State.cpp
+++ b/Engine/Source/Runtime/D3D12RHI/Private/D3D12State.cpp
@@ -604,6 +604,12 @@ FGraphicsPipelineStateRHIRef FD3D12DynamicRHI::RHICreateGraphicsPipelineState(co
 
 	const FD3D12RootSignature* RootSignature = GetAdapter().GetRootSignature(Initializer.BoundShaderState);
 
+	if (!RootSignature || !RootSignature->GetRootSignature())
+	{
+		UE_LOG(LogD3D12RHI, Error, TEXT("Unexpected null root signature at graphics pipeline creation time"));
+		return nullptr;
+	}
+
 	// Next try to find the PSO based on the hash of its desc.
 
 	FD3D12LowLevelGraphicsPipelineStateDesc LowLevelDesc;
@@ -645,6 +651,12 @@ TRefCountPtr<FRHIComputePipelineState> FD3D12DynamicRHI::RHICreateComputePipelin
 
 	const FD3D12RootSignature* RootSignature = ComputeShader->RootSignature;
 
+	if (!RootSignature || !RootSignature->GetRootSignature())
+	{
+		UE_LOG(LogD3D12RHI, Error, TEXT("Unexpected null root signature at compute pipeline creation time (shader hash %s)"), *ComputeShader->GetHash().ToString());
+		return nullptr;
+	}
+
 	// Next try to find the PSO based on the hash of its desc.
 	FD3D12ComputePipelineStateDesc LowLevelDesc;
 	Found = PSOCache.FindInLoadedCache(ComputeShader, RootSignature, LowLevelDesc);
