HWRT: Some implementations are able to update BLAS completely in-place without requiring any scratch memory. Because we don't really handle that and expect scratch buffer all the time make sure we always create it.

diff --git a/Engine/Source/Runtime/D3D12RHI/Private/D3D12RayTracing.cpp b/Engine/Source/Runtime/D3D12RHI/Private/D3D12RayTracing.cpp
--- a/Engine/Source/Runtime/D3D12RHI/Private/D3D12RayTracing.cpp
+++ b/Engine/Source/Runtime/D3D12RHI/Private/D3D12RayTracing.cpp
@@ -2671,7 +2671,7 @@ FRayTracingAccelerationStructureSize FD3D12DynamicRHI::RHICalcRayTracingGeometry
 
 	SizeInfo.ResultSize = Align(SizeInfo.ResultSize, GRHIRayTracingAccelerationStructureAlignment);
 	SizeInfo.BuildScratchSize = Align(SizeInfo.BuildScratchSize, GRHIRayTracingScratchBufferAlignment);
-	SizeInfo.UpdateScratchSize = Align(SizeInfo.UpdateScratchSize, GRHIRayTracingScratchBufferAlignment);
+	SizeInfo.UpdateScratchSize = Align(FMath::Max(1ULL, SizeInfo.UpdateScratchSize), GRHIRayTracingScratchBufferAlignment);
 
 	return SizeInfo;
 }
