Switching new code to use UE::RHICore::AllowVendorDevice()

diff --git a/Engine/Source/Runtime/D3D12RHI/Private/D3D12Adapter.cpp b/Engine/Source/Runtime/D3D12RHI/Private/D3D12Adapter.cpp
--- a/Engine/Source/Runtime/D3D12RHI/Private/D3D12Adapter.cpp
+++ b/Engine/Source/Runtime/D3D12RHI/Private/D3D12Adapter.cpp
@@ -573,7 +573,7 @@ void FD3D12Adapter::CreateRootDevice(bool bWithDebug)
 	{
 #if INTEL_EXTENSIONS
 		// Enable Intel App Discovery
-		if (IsRHIDeviceIntel() && bAllowVendorDevice)
+		if (IsRHIDeviceIntel() && UE::RHICore::AllowVendorDevice())
 		{
 			EnableIntelAppDiscovery(GRHIDeviceId);
 		}
