Don't ensure when an Aftermath API call fails.

#rnx

diff --git a/Engine/Source/Runtime/D3D12RHI/Private/D3D12NvidiaExtensions.cpp b/Engine/Source/Runtime/D3D12RHI/Private/D3D12NvidiaExtensions.cpp
--- a/Engine/Source/Runtime/D3D12RHI/Private/D3D12NvidiaExtensions.cpp
+++ b/Engine/Source/Runtime/D3D12RHI/Private/D3D12NvidiaExtensions.cpp
@@ -1,5 +1,6 @@
 // Copyright Epic Games, Inc. All Rights Reserved.
 
+#include "D3D12RHICommon.h"
 #include "D3D12NvidiaExtensions.h"
 
 #if NV_AFTERMATH
@@ -34,8 +35,9 @@ namespace UE::RHICore::Nvidia::Aftermath::D3D12
 		if (IsEnabled())
 		{
 			GFSDK_Aftermath_Result Result = GFSDK_Aftermath_DX12_CreateContextHandle(D3DCommandList, &Handle);
-			if (!ensureMsgf(Result == GFSDK_Aftermath_Result_Success, TEXT("GFSDK_Aftermath_DX12_CreateContextHandle failed: 0x%08x"), Result))
+			if (Result != GFSDK_Aftermath_Result_Success)
 			{
+				UE_LOG(LogD3D12RHI, VeryVerbose, TEXT("GFSDK_Aftermath_DX12_CreateContextHandle failed: 0x%08x"), Result);
 				Handle = {};
 			}
 		}
@@ -48,7 +50,7 @@ namespace UE::RHICore::Nvidia::Aftermath::D3D12
 		if (IsEnabled() && CommandList)
 		{
 			GFSDK_Aftermath_Result Result = GFSDK_Aftermath_ReleaseContextHandle(CommandList);
-			ensureMsgf(Result == GFSDK_Aftermath_Result_Success, TEXT("GFSDK_Aftermath_ReleaseContextHandle failed: 0x%08x"), Result);
+			UE_CLOG(Result != GFSDK_Aftermath_Result_Success, LogD3D12RHI, VeryVerbose, TEXT("GFSDK_Aftermath_ReleaseContextHandle failed: 0x%08x"), Result);
 		}
 	}
 
@@ -58,8 +60,9 @@ namespace UE::RHICore::Nvidia::Aftermath::D3D12
 		if (IsEnabled())
 		{
 			GFSDK_Aftermath_Result Result = GFSDK_Aftermath_DX12_RegisterResource(D3DResource, &Handle);
-			if (!ensureMsgf(Result == GFSDK_Aftermath_Result_Success, TEXT("GFSDK_Aftermath_DX12_RegisterResource failed: 0x%08x"), Result))
+			if (Result != GFSDK_Aftermath_Result_Success)
 			{
+				UE_LOG(LogD3D12RHI, VeryVerbose, TEXT("GFSDK_Aftermath_DX12_RegisterResource failed: 0x%08x"), Result);
 				Handle = {};
 			}
 		}
@@ -71,7 +74,7 @@ namespace UE::RHICore::Nvidia::Aftermath::D3D12
 		if (IsEnabled() && Resource)
 		{
 			GFSDK_Aftermath_Result Result = GFSDK_Aftermath_DX12_UnregisterResource(Resource);
-			ensureMsgf(Result == GFSDK_Aftermath_Result_Success, TEXT("GFSDK_Aftermath_DX12_UnregisterResource failed: 0x%08x"), Result);
+			UE_CLOG(Result != GFSDK_Aftermath_Result_Success, LogD3D12RHI, VeryVerbose, TEXT("GFSDK_Aftermath_DX12_UnregisterResource failed: 0x%08x"), Result);
 		}
 	}
 
@@ -82,7 +85,7 @@ namespace UE::RHICore::Nvidia::Aftermath::D3D12
 		if (Marker)
 		{
 			GFSDK_Aftermath_Result Result = GFSDK_Aftermath_SetEventMarker(CommandList, Marker.GetPtr(), Marker.GetSize());
-			ensureMsgf(Result == GFSDK_Aftermath_Result_Success, TEXT("GFSDK_Aftermath_SetEventMarker failed in BeginBreadcrumb: 0x%08x"), Result);
+			UE_CLOG(Result != GFSDK_Aftermath_Result_Success, LogD3D12RHI, VeryVerbose, TEXT("GFSDK_Aftermath_SetEventMarker failed in BeginBreadcrumb: 0x%08x"), Result);
 		}
 	}
 
@@ -92,7 +95,7 @@ namespace UE::RHICore::Nvidia::Aftermath::D3D12
 		if (Marker)
 		{
 			GFSDK_Aftermath_Result Result = GFSDK_Aftermath_SetEventMarker(CommandList, Marker.GetPtr(), Marker.GetSize());
-			ensureMsgf(Result == GFSDK_Aftermath_Result_Success, TEXT("GFSDK_Aftermath_SetEventMarker failed in EndBreadcrumb: 0x%08x"), Result);
+			UE_CLOG(Result != GFSDK_Aftermath_Result_Success, LogD3D12RHI, VeryVerbose, TEXT("GFSDK_Aftermath_SetEventMarker failed in EndBreadcrumb: 0x%08x"), Result);
 		}
 	}
 #endif
