fix check in D3D12RenderTarget when doing DumpGPU

diff --git a/Engine/Source/Runtime/D3D12RHI/Private/D3D12RenderTarget.cpp b/Engine/Source/Runtime/D3D12RHI/Private/D3D12RenderTarget.cpp
--- a/Engine/Source/Runtime/D3D12RHI/Private/D3D12RenderTarget.cpp
+++ b/Engine/Source/Runtime/D3D12RHI/Private/D3D12RenderTarget.cpp
@@ -612,7 +612,7 @@ void FD3D12DynamicRHI::RHIReadSurfaceData(FRHITexture* InRHITexture, FIntRect In
 	// Allocate the output buffer.
 	OutData.SetNumUninitialized(SizeX * SizeY);
 	
-	DXGI_FORMAT DXGIFormat = TextureDesc.Format;
+	// dest format :
 	EPixelFormat PixelFormat = DestTexture2D->GetFormat();
 
 	check( PixelFormat != PF_Unknown );
@@ -625,16 +625,19 @@ void FD3D12DynamicRHI::RHIReadSurfaceData(FRHITexture* InRHITexture, FIntRect In
 	
 	// switching on the EPixelFormat is risky if the mapping is not what you expect
 	//	verify against TextureDesc.Format
+	
+	DXGI_FORMAT DXGIFormat = TextureDesc.Format;
 
 	if ( DXGIFormat == DXGI_FORMAT_UNKNOWN )
 	{
-		// when called on actual textures, format is valid
-		// but this is also called on untyped buffers
+		// when called on actual textures, DXGIFormat is valid
+		// but this is also called on untyped buffers, in which case we only know the type from dest PixelFormat
 		DXGIFormat = (DXGI_FORMAT) FormatInfo.PlatformFormat;
 	}
 	else
 	{
-		check( FormatInfo.PlatformFormat == DXGIFormat );
+		// source and dest format must match, except for _TYPELESS vs _UNORM
+		check( ConvertTypelessToUnorm((DXGI_FORMAT)FormatInfo.PlatformFormat) == ConvertTypelessToUnorm(DXGIFormat) );
 		check( FormatInfo.BlockBytes == D3D12RT_ComputeBytesPerPixel(DXGIFormat) );
 	}
 
