Allow some platforms to signal a fence between ExecuteCommandLists

#rnx

diff --git a/Engine/Source/Runtime/D3D12RHI/Private/D3D12Device.cpp b/Engine/Source/Runtime/D3D12RHI/Private/D3D12Device.cpp
--- a/Engine/Source/Runtime/D3D12RHI/Private/D3D12Device.cpp
+++ b/Engine/Source/Runtime/D3D12RHI/Private/D3D12Device.cpp
@@ -57,6 +57,14 @@ FD3D12Queue::FD3D12Queue(FD3D12Device* Device, ED3D12QueueType QueueType)
 		IID_PPV_ARGS(Fence.D3DFence.GetInitReference())
 	));
 	Fence.D3DFence->SetName(*FString::Printf(TEXT("%s Queue Fence (GPU %d)"), GetD3DCommandQueueTypeName(QueueType), Device->GetGPUIndex()));
+
+
+	VERIFYD3D12RESULT(Device->GetDevice()->CreateFence(
+		0,
+		D3D12_FENCE_FLAG_NONE,
+		IID_PPV_ARGS(ExecuteCommandListsFence.D3DFence.GetInitReference())
+	));
+	ExecuteCommandListsFence.D3DFence->SetName(*FString::Printf(TEXT("%s ExecuteCommandListsFence (GPU %d)"), GetD3DCommandQueueTypeName(QueueType), Device->GetGPUIndex()));
 }
 
 FD3D12Queue::~FD3D12Queue()
diff --git a/Engine/Source/Runtime/D3D12RHI/Private/D3D12Device.h b/Engine/Source/Runtime/D3D12RHI/Private/D3D12Device.h
--- a/Engine/Source/Runtime/D3D12RHI/Private/D3D12Device.h
+++ b/Engine/Source/Runtime/D3D12RHI/Private/D3D12Device.h
@@ -183,6 +183,11 @@ public:
 
 	FD3D12Queue(FD3D12Device* Device, ED3D12QueueType QueueType);
 	~FD3D12Queue();
+
+private:
+	// Internal fence which may be used before calling ExecuteCommandLists
+	FD3D12Fence ExecuteCommandListsFence;
+
 };
 
 class FD3D12Device final : public FD3D12SingleNodeGPUObject, public FNoncopyable, public FD3D12AdapterChild
