Changed the nvidia driver update link to point to the page which also has downloads for non-gaming GPUs.

Also removed the suggestion to use -d3ddebug when a GPU crash occurs, since it's not useful in most cases.

#rnx

diff --git a/Engine/Source/Runtime/D3D12RHI/Private/D3D12Util.cpp b/Engine/Source/Runtime/D3D12RHI/Private/D3D12Util.cpp
--- a/Engine/Source/Runtime/D3D12RHI/Private/D3D12Util.cpp
+++ b/Engine/Source/Runtime/D3D12RHI/Private/D3D12Util.cpp
@@ -974,17 +974,13 @@ void FD3D12DynamicRHI::OutputGPUCrashReport(FTextBuilder& ErrorMessage)
 
 	// Build the error message
 	ErrorMessage.AppendLine(LOCTEXT("GPU Crashed", "GPU Crashed or D3D Device Removed.\n"));
-	if (!GetAdapter().IsDebugDevice())
-	{
-		ErrorMessage.AppendLine(LOCTEXT("D3D Debug Device", "Use -d3ddebug to enable the D3D debug device."));
-	}
 	if (UE::RHI::UseGPUCrashDebugging())
 	{
 		ErrorMessage.AppendLine(LOCTEXT("GPU Crash Debugging enabled", "Check log for GPU state information."));
 	}
 	else
 	{
-		ErrorMessage.AppendLine(LOCTEXT("GPU Crash Debugging disabled", "Use -gpucrashdebugging to track current GPU state."));
+		ErrorMessage.AppendLine(LOCTEXT("GPU Crash Debugging disabled", "Use -gpucrashdebugging to enable GPU state tracking."));
 	}
 #if NV_AFTERMATH
 	if (AftermathResult.DumpPath.IsSet())
diff --git a/Engine/Source/Runtime/RHI/Private/DynamicRHI.cpp b/Engine/Source/Runtime/RHI/Private/DynamicRHI.cpp
--- a/Engine/Source/Runtime/RHI/Private/DynamicRHI.cpp
+++ b/Engine/Source/Runtime/RHI/Private/DynamicRHI.cpp
@@ -208,7 +208,7 @@ void RHIDetectAndWarnOfBadDrivers(bool bHasEditorToken)
 			if (DriverInfo.IsNVIDIA())
 			{
 				VendorString = TEXT("NVIDIA");
-				HyperlinkText = NSLOCTEXT("MessageDialog", "DriverDownloadLinkNVIDIA", "https://www.nvidia.com/en-us/geforce/drivers/");
+				HyperlinkText = NSLOCTEXT("MessageDialog", "DriverDownloadLinkNVIDIA", "https://www.nvidia.com/download/index.aspx");
 			}
 			else if (DriverInfo.IsAMD())
 			{
