Enable Async Compute on Intel

diff --git a/Engine/Source/Runtime/D3D12RHI/Private/D3D12Adapter.cpp b/Engine/Source/Runtime/D3D12RHI/Private/D3D12Adapter.cpp
--- a/Engine/Source/Runtime/D3D12RHI/Private/D3D12Adapter.cpp
+++ b/Engine/Source/Runtime/D3D12RHI/Private/D3D12Adapter.cpp
@@ -1095,6 +1095,10 @@ void FD3D12Adapter::InitializeDevices()
 			{
 				GSupportsEfficientAsyncCompute = true;
 			}
+			else if (IsRHIDeviceIntel() && GMaxRHIFeatureLevel >= ERHIFeatureLevel::SM6)
+			{
+				GSupportsEfficientAsyncCompute = true;
+			}
 #if PLATFORM_WINDOWS
 			else
 			{
