- Added Platform LLM tags that track:
    - D3D12 unused memory from Buddy allocator and Pool allocator. On some platforms, this account for around 40% of untracked memory in default tracker
    - D3D12 wasted memory from Buddy allocator
- Fixed D3D12RHI_TRACK_DETAILED_STATS usage that was always true
- Prevent stack overflow by logging a Fatal error in BuddyAllocator if order is higher than MaxOrder

#rnx

diff --git a/Engine/Source/Runtime/D3D12RHI/Private/D3D12Allocation.cpp b/Engine/Source/Runtime/D3D12RHI/Private/D3D12Allocation.cpp
--- a/Engine/Source/Runtime/D3D12RHI/Private/D3D12Allocation.cpp
+++ b/Engine/Source/Runtime/D3D12RHI/Private/D3D12Allocation.cpp
@@ -9,6 +9,8 @@
 #include "D3D12RHIPrivate.h"
 #include "Misc/BufferedOutputDevice.h"
 #include "HAL/PlatformStackWalk.h"
+#include "HAL/LowLevelMemTracker.h"
+#include "HAL/LowLevelMemStats.h"
 #include "ProfilingDebugging/MemoryTrace.h"
 
 // Fix for random GPU crashes on draw indirects on multiple IHVs. Force all indirect arg buffers as committed resources (see UE-115982)
@@ -176,6 +178,11 @@ static FAutoConsoleVariableRef CVarD3D12FastAllocatorMinPagesToRetain(
 	TEXT("Minimum number of pages to retain. Pages below this limit will never be released. Pages above can be released after being unused for a certain number of frames."),
 	ECVF_Default);
 
+DECLARE_LLM_MEMORY_STAT(TEXT("D3D12AllocatorUnused"), STAT_D3D12AllocatorUnusedLLM, STATGROUP_LLMFULL);
+LLM_DEFINE_TAG(D3D12AllocatorUnused, NAME_None, NAME_None, GET_STATFNAME(STAT_D3D12AllocatorUnusedLLM), GET_STATFNAME(STAT_EngineSummaryLLM));
+DECLARE_LLM_MEMORY_STAT(TEXT("D3D12AllocatorWasted"), STAT_D3D12AllocatorWastedLLM, STATGROUP_LLMFULL);
+LLM_DEFINE_TAG(D3D12AllocatorWasted, NAME_None, NAME_None, GET_STATFNAME(STAT_D3D12AllocatorWastedLLM), GET_STATFNAME(STAT_EngineSummaryLLM));
+
 namespace ED3D12AllocatorID
 {
 	enum Type
@@ -201,7 +208,7 @@ FD3D12ResourceAllocator::FD3D12ResourceAllocator(FD3D12Device* ParentDevice,
 	, DebugName(Name)
 	, Initialized(false)
 	, MaximumAllocationSizeForPooling(MaxSizeForPooling)
-#if defined(D3D12RHI_TRACK_DETAILED_STATS)
+#if D3D12RHI_TRACK_DETAILED_STATS
 	, SpaceAlignedUsed(0)
 	, SpaceActualUsed(0)
 	, NumBlocksInDeferredDeletionQueue(0)
@@ -278,6 +285,7 @@ void FD3D12BuddyAllocator::Initialize()
 			// we are tracking allocations ourselves, so don't let XMemAlloc track these as well
 			LLM_SCOPED_PAUSE_TRACKING_FOR_TRACKER(ELLMTracker::Default, ELLMAllocType::System);
 			VERIFYD3D12RESULT(Adapter->GetD3DDevice()->CreateHeap(&Desc, IID_PPV_ARGS(&Heap)));
+			LLM_SCOPED_PAUSE_TRACKING_WITH_ENUM_AND_AMOUNT_BYTAG(D3D12AllocatorUnused, MaxBlockSize, ELLMTracker::Platform, ELLMAllocType::System);
 		}
 
 		BackingHeap = new FD3D12Heap(GetParentDevice(), GetVisibilityMask(), TraceHeapId);
@@ -295,6 +303,7 @@ void FD3D12BuddyAllocator::Initialize()
 			LLM_SCOPED_PAUSE_TRACKING_FOR_TRACKER(ELLMTracker::Default, ELLMAllocType::System);
 			const D3D12_HEAP_PROPERTIES HeapProps = CD3DX12_HEAP_PROPERTIES(InitConfig.HeapType, GetGPUMask().GetNative(), GetVisibilityMask().GetNative());
 			VERIFYD3D12RESULT(Adapter->CreateBuffer(HeapProps, GetGPUMask(), InitConfig.InitialResourceState, ED3D12ResourceStateMode::SingleState, InitConfig.InitialResourceState, MaxBlockSize, BackingResource.GetInitReference(), TEXT("Resource Allocator Underlying Buffer"), InitConfig.ResourceFlags));
+			LLM_SCOPED_PAUSE_TRACKING_WITH_ENUM_AND_AMOUNT_BYTAG(D3D12AllocatorUnused, MaxBlockSize, ELLMTracker::Platform, ELLMAllocType::System);
 #if UE_MEMORY_TRACE_ENABLED
 			MemoryTrace_MarkAllocAsHeap(BackingResource->GetGPUVirtualAddress(), TraceHeapId);
 #endif
@@ -318,7 +327,10 @@ uint32 FD3D12BuddyAllocator::AllocateBlock(uint32 order)
 
 	if (order > MaxOrder)
 	{
-		check(false); // Can't allocate a block that large  
+		// Can't allocate a block that large
+		check(false); 
+		// Crash to avoid infinite recursivity
+		UE_LOG(LogD3D12RHI, Fatal, TEXT("Buddy Allocator cant allocate a block that large (order %d)"), order);
 	}
 
 	if (FreeBlocks[order].Num() == 0)
@@ -408,9 +420,13 @@ void FD3D12BuddyAllocator::Allocate(uint32 SizeInBytes, uint32 Alignment, FD3D12
 	INCREASE_ALLOC_COUNTER(SpaceAlignedUsed, AllocSize);
 	INCREASE_ALLOC_COUNTER(SpaceActualUsed, SizeInBytes);
 	
+	// Decrease only texture size so wasted amount stays in D3D12AllocatorUnused
+	LLM_SCOPED_PAUSE_TRACKING_WITH_ENUM_AND_AMOUNT_BYTAG(D3D12AllocatorUnused, 0 - int64(SizeInBytes), ELLMTracker::Platform, ELLMAllocType::System);
+	LLM_SCOPED_PAUSE_TRACKING_WITH_ENUM_AND_AMOUNT_BYTAG(D3D12AllocatorWasted, int64(AllocSize - SizeInBytes), ELLMTracker::Platform, ELLMAllocType::System);
+
 	TotalSizeUsed += AllocSize;
 
-#if defined(D3D12RHI_TRACK_DETAILED_STATS)
+#if D3D12RHI_TRACK_DETAILED_STATS
 	if (SpaceActualUsed > PeakUsage)
 	{
 		PeakUsage = SpaceActualUsed;
@@ -490,14 +506,11 @@ void FD3D12BuddyAllocator::Deallocate(FD3D12ResourceLocation& ResourceLocation)
 	FD3D12BuddyAllocatorPrivateData& PrivateData = ResourceLocation.GetBuddyAllocatorPrivateData();
 	Block.Data.Order = PrivateData.Order;
 	Block.Data.Offset = PrivateData.Offset;
+	Block.AllocationSize = ResourceLocation.GetSize();
 
 	// update the last used framce fence used during garbage collection
 	LastUsedFrameFence = FMath::Max(LastUsedFrameFence, Block.FrameFence);
 
-#if defined(D3D12RHI_TRACK_DETAILED_STATS)
-	Block.AllocationSize = ResourceLocation.GetSize();
-#endif
-
 	if (ResourceLocation.GetResource()->IsPlacedResource())
 	{
 		Block.PlacedResource = ResourceLocation.GetResource();
@@ -527,6 +540,8 @@ void FD3D12BuddyAllocator::DeallocateInternal(RetiredBlock& Block)
 	const uint32 Size = uint32(OrderToUnitSize(Block.Data.Order) * MinBlockSize);
 	DECREASE_ALLOC_COUNTER(SpaceAlignedUsed, Size);
 	DECREASE_ALLOC_COUNTER(SpaceActualUsed, Block.AllocationSize);
+	LLM_SCOPED_PAUSE_TRACKING_WITH_ENUM_AND_AMOUNT_BYTAG(D3D12AllocatorUnused, int64(Block.AllocationSize), ELLMTracker::Platform, ELLMAllocType::System);
+	LLM_SCOPED_PAUSE_TRACKING_WITH_ENUM_AND_AMOUNT_BYTAG(D3D12AllocatorWasted, 0 - int64(Size - Block.AllocationSize), ELLMTracker::Platform, ELLMAllocType::System);
 
 	TotalSizeUsed -= Size;
 
@@ -573,6 +588,7 @@ void FD3D12BuddyAllocator::CleanUpAllocations()
 void FD3D12BuddyAllocator::ReleaseAllResources()
 {
 	LLM_SCOPED_PAUSE_TRACKING_FOR_TRACKER(ELLMTracker::Default, ELLMAllocType::System);
+	LLM_SCOPED_PAUSE_TRACKING_WITH_ENUM_AND_AMOUNT_BYTAG(D3D12AllocatorUnused, 0 - int64(MaxBlockSize), ELLMTracker::Platform, ELLMAllocType::System);
 
 #if UE_MEMORY_TRACE_ENABLED
 	if (AllocationStrategy != EResourceAllocationStrategy::kPlacedResource)
@@ -634,7 +650,7 @@ void FD3D12BuddyAllocator::DumpAllocatorStats(class FOutputDevice& Ar)
 
 void FD3D12BuddyAllocator::UpdateMemoryStats(uint32& IOMemoryAllocated, uint32& IOMemoryUsed, uint32& IOMemoryFree, uint32& IOAlignmentWaste, uint32& IOAllocatedPageCount, uint32& IOFullPageCount)
 {
-#if defined(D3D12RHI_TRACK_DETAILED_STATS)
+#if D3D12RHI_TRACK_DETAILED_STATS
 	IOMemoryAllocated += MaxBlockSize;
 	IOMemoryUsed += SpaceActualUsed;
 	IOMemoryFree += (MaxBlockSize - SpaceAlignedUsed);
@@ -796,7 +812,7 @@ void FD3D12MultiBuddyAllocator::DumpAllocatorStats(class FOutputDevice& Ar)
 
 void FD3D12MultiBuddyAllocator::UpdateMemoryStats(uint32& IOMemoryAllocated, uint32& IOMemoryUsed, uint32& IOMemoryFree, uint32& IOAlignmentWaste, uint32& IOAllocatedPageCount, uint32& IOFullPageCount)
 {
-#if defined(D3D12RHI_TRACK_DETAILED_STATS)
+#if D3D12RHI_TRACK_DETAILED_STATS
 	FScopeLock Lock(&CS);
 
 	for (FD3D12BuddyAllocator* Allocator : Allocators)
@@ -1151,7 +1167,7 @@ void FD3D12UploadHeapAllocator::UpdateMemoryStats()
 	uint32 AllocatedPageCount = 0;
 	uint32 FullPageCount = 0;
 
-#if defined(D3D12RHI_TRACK_DETAILED_STATS)
+#if D3D12RHI_TRACK_DETAILED_STATS
 	SmallBlockAllocator.UpdateMemoryStats(MemoryAllocated, MemoryUsed, FreeMemory, AlignmentWaste, AllocatedPageCount, FullPageCount);
 	{
 		FD3D12ScopeLock Lock(&BigBlockCS);
@@ -1617,7 +1633,7 @@ void FD3D12DefaultBufferAllocator::UpdateMemoryStats()
 	uint32 AllocatedPageCount = 0;
 	uint32 FullPageCount = 0;
 
-#if defined(D3D12RHI_TRACK_DETAILED_STATS)
+#if D3D12RHI_TRACK_DETAILED_STATS
 	for (FD3D12BufferPool* DefaultBufferPool : DefaultBufferPools)
 	{
 		if (DefaultBufferPool)
diff --git a/Engine/Source/Runtime/D3D12RHI/Private/D3D12PoolAllocator.cpp b/Engine/Source/Runtime/D3D12RHI/Private/D3D12PoolAllocator.cpp
--- a/Engine/Source/Runtime/D3D12RHI/Private/D3D12PoolAllocator.cpp
+++ b/Engine/Source/Runtime/D3D12RHI/Private/D3D12PoolAllocator.cpp
@@ -2,11 +2,14 @@
 
 #include "D3D12PoolAllocator.h"
 #include "D3D12RHIPrivate.h"
+#include "HAL/LowLevelMemTracker.h"
 
 #ifndef NEEDS_D3D12_INDIRECT_ARGUMENT_HEAP_WORKAROUND
 #define NEEDS_D3D12_INDIRECT_ARGUMENT_HEAP_WORKAROUND 0
 #endif
 
+LLM_DECLARE_TAG(D3D12AllocatorUnused);
+
 //-----------------------------------------------------------------------------
 //	FD3D12MemoryPool
 //-----------------------------------------------------------------------------
@@ -132,6 +135,7 @@ void FD3D12MemoryPool::Init()
 	}
 #endif // D3D12_RHI_RAYTRACING
 
+	LLM_SCOPED_PAUSE_TRACKING_WITH_ENUM_AND_AMOUNT_BYTAG(D3D12AllocatorUnused, int64(PoolSize), ELLMTracker::Platform, ELLMAllocType::System);
 	FRHIMemoryPool::Init();
 }
 
@@ -365,6 +369,7 @@ void FD3D12PoolAllocator::AllocateResource(uint32 GPUIndex, D3D12_HEAP_TYPE InHe
 		// Try to allocate in one of the pools
 		FRHIPoolAllocationData& AllocationData = ResourceLocation.GetPoolAllocatorPrivateData().PoolData;
 		verify(TryAllocateInternal(InSize, AllocationAlignment, AllocationResourceType, AllocationData));
+		LLM_SCOPED_PAUSE_TRACKING_WITH_ENUM_AND_AMOUNT_BYTAG(D3D12AllocatorUnused, 0 - int64(InSize), ELLMTracker::Platform, ELLMAllocType::System);
 
 		// Setup the resource location
 		ResourceLocation.SetType(FD3D12ResourceLocation::ResourceLocationType::eSubAllocation);
@@ -639,7 +644,7 @@ void FD3D12PoolAllocator::CleanUpAllocations(uint64 InFrameLag, bool bForceFree)
 				// Not pending anymore
 				check(PendingDeleteRequestSize >= Operation.AllocationData->GetSize());
 				PendingDeleteRequestSize -= Operation.AllocationData->GetSize();
-
+				LLM_SCOPED_PAUSE_TRACKING_WITH_ENUM_AND_AMOUNT_BYTAG(D3D12AllocatorUnused, int64(Operation.AllocationData->GetSize()), ELLMTracker::Platform, ELLMAllocType::System);
 				// Deallocate the locked block (actually free now)
 				DeallocateInternal(*Operation.AllocationData);
 				Operation.AllocationData->Reset();
@@ -692,6 +697,7 @@ void FD3D12PoolAllocator::CleanUpAllocations(uint64 InFrameLag, bool bForceFree)
 		FD3D12MemoryPool* MemoryPool = (FD3D12MemoryPool*) Pools[PoolIndex];
 		if (MemoryPool != nullptr && MemoryPool->IsEmpty() && (bForceFree || (MemoryPool->GetLastUsedFrameFence() + InFrameLag <= CompletedFence)))
 		{
+			LLM_SCOPED_PAUSE_TRACKING_WITH_ENUM_AND_AMOUNT_BYTAG(D3D12AllocatorUnused, 0 - int64(MemoryPool->GetPoolSize()), ELLMTracker::Platform, ELLMAllocType::System);
 			MemoryPool->Destroy();
 			delete(MemoryPool);
 			Pools[PoolIndex] = nullptr;
