Initialize all unused render target data on the blend state to make sure we are not giving uninitialized memory to the driver

diff --git a/Engine/Source/Runtime/D3D12RHI/Private/D3D12Pipelinestate.cpp b/Engine/Source/Runtime/D3D12RHI/Private/D3D12Pipelinestate.cpp
--- a/Engine/Source/Runtime/D3D12RHI/Private/D3D12Pipelinestate.cpp
+++ b/Engine/Source/Runtime/D3D12RHI/Private/D3D12Pipelinestate.cpp
@@ -115,6 +115,18 @@ static FD3D12LowLevelGraphicsPipelineStateDesc GetLowLevelGraphicsPipelineStateD
 	Desc.Desc.DepthStencilState.DepthBoundsTestEnable = GSupportsDepthBoundsTest && Initializer.bDepthBounds;
 #endif
 
+	// The blend state is liable to have non default values for RenderTarget indices that exceed
+	// NumRenderTargets, this will in turn cause them to have a different `D3D12_PIPELINE_STATE_STREAM_DESC`
+	// as far as the driver is concerned.
+	// Which will then cause a hash collision when using driver level pipeline caching, due to the hash
+	// only being derived from the active render target blend states.
+#if !D3D12_USE_DERIVED_PSO || D3D12_USE_DERIVED_PSO_SHADER_EXPORTS
+	for (UINT RenderTarget = Desc.Desc.RTFormatArray.NumRenderTargets; RenderTarget < 8; ++RenderTarget)
+	{
+		Desc.Desc.BlendState.RenderTarget[RenderTarget] = D3D12_RENDER_TARGET_BLEND_DESC{};
+	}
+#endif // !D3D12_USE_DERIVED_PSO || D3D12_USE_DERIVED_PSO_SHADER_EXPORTS
+
 	Desc.bFromPSOFileCache = Initializer.bFromPSOFileCache;
 
 	return Desc;
