Deprecate:

diff --git a/Engine/Plugins/Runtime/nDisplay/Source/DisplayClusterShaders/Private/Shaders/DisplayClusterShadersGenerateMips.h b/Engine/Plugins/Runtime/nDisplay/Source/DisplayClusterShaders/Private/Shaders/DisplayClusterShadersGenerateMips.h
--- a/Engine/Plugins/Runtime/nDisplay/Source/DisplayClusterShaders/Private/Shaders/DisplayClusterShadersGenerateMips.h
+++ b/Engine/Plugins/Runtime/nDisplay/Source/DisplayClusterShaders/Private/Shaders/DisplayClusterShadersGenerateMips.h
@@ -11,6 +11,6 @@
 class FDisplayClusterShadersGenerateMips
 {
 public:
-	static bool GenerateMips(FRHICommandListImmediate& RHICmdList, FRHITexture2D* InOutMipsTexture, const FDisplayClusterShaderParameters_GenerateMips& InSettings);
+	static bool GenerateMips(FRHICommandListImmediate& RHICmdList, FRHITexture* InOutMipsTexture, const FDisplayClusterShaderParameters_GenerateMips& InSettings);
 };
 
diff --git a/Engine/Plugins/Runtime/nDisplay/Source/DisplayClusterShaders/Private/Shaders/DisplayClusterShadersPostprocess_Blur.cpp b/Engine/Plugins/Runtime/nDisplay/Source/DisplayClusterShaders/Private/Shaders/DisplayClusterShadersPostprocess_Blur.cpp
--- a/Engine/Plugins/Runtime/nDisplay/Source/DisplayClusterShaders/Private/Shaders/DisplayClusterShadersPostprocess_Blur.cpp
+++ b/Engine/Plugins/Runtime/nDisplay/Source/DisplayClusterShaders/Private/Shaders/DisplayClusterShadersPostprocess_Blur.cpp
@@ -121,8 +121,8 @@ IMPLEMENT_SHADER_TYPE(template<>, FPicpBlurPostProcessDilatePS,  PostProcessShad
 template<uint32 ShaderType>
 static void PicpBlurPostProcess_RenderThread(
 	FRHICommandListImmediate& RHICmdList,
-	FRHITexture2D* InShaderTexture,
-	FRHITexture2D* OutRenderTargetableTexture,
+	FRHITexture* InShaderTexture,
+	FRHITexture* OutRenderTargetableTexture,
 	const FDisplayClusterShaderParameters_PostprocessBlur& InSettings)
 {
 	check(IsInRenderingThread());
@@ -198,7 +198,7 @@ static void PicpBlurPostProcess_RenderThread(
 DECLARE_GPU_STAT_NAMED(nDisplay_Picp_PostProcess_Compose, TEXT("nDisplay Picp_PostProcess::Compose"));
 DECLARE_GPU_STAT_NAMED(nDisplay_Picp_PostProcess_Blur, TEXT("nDisplay Picp_PostProcess::Blur"));
 
-bool FDisplayClusterShadersPostprocess_Blur::RenderPostprocess_Blur(FRHICommandListImmediate& RHICmdList, FRHITexture2D* InSourceTexture, FRHITexture2D* InRenderTargetableDestTexture, const FDisplayClusterShaderParameters_PostprocessBlur& InSettings)
+bool FDisplayClusterShadersPostprocess_Blur::RenderPostprocess_Blur(FRHICommandListImmediate& RHICmdList, FRHITexture* InSourceTexture, FRHITexture* InRenderTargetableDestTexture, const FDisplayClusterShaderParameters_PostprocessBlur& InSettings)
 {
 	check(IsInRenderingThread());
 
diff --git a/Engine/Plugins/Runtime/nDisplay/Source/DisplayClusterShaders/Private/Shaders/DisplayClusterShadersPostprocess_Blur.h b/Engine/Plugins/Runtime/nDisplay/Source/DisplayClusterShaders/Private/Shaders/DisplayClusterShadersPostprocess_Blur.h
--- a/Engine/Plugins/Runtime/nDisplay/Source/DisplayClusterShaders/Private/Shaders/DisplayClusterShadersPostprocess_Blur.h
+++ b/Engine/Plugins/Runtime/nDisplay/Source/DisplayClusterShaders/Private/Shaders/DisplayClusterShadersPostprocess_Blur.h
@@ -11,6 +11,6 @@
 class FDisplayClusterShadersPostprocess_Blur
 {
 public:
-	static bool RenderPostprocess_Blur(FRHICommandListImmediate& RHICmdList, FRHITexture2D* InSourceTexture, FRHITexture2D* InRenderTargetableDestTexture, const FDisplayClusterShaderParameters_PostprocessBlur& InSettings);
+	static bool RenderPostprocess_Blur(FRHICommandListImmediate& RHICmdList, FRHITexture* InSourceTexture, FRHITexture* InRenderTargetableDestTexture, const FDisplayClusterShaderParameters_PostprocessBlur& InSettings);
 };
 
diff --git a/Engine/Plugins/Runtime/nDisplay/Source/DisplayClusterShaders/Private/Shaders/DisplayClusterShadersPostprocess_OutputRemap.cpp b/Engine/Plugins/Runtime/nDisplay/Source/DisplayClusterShaders/Private/Shaders/DisplayClusterShadersPostprocess_OutputRemap.cpp
--- a/Engine/Plugins/Runtime/nDisplay/Source/DisplayClusterShaders/Private/Shaders/DisplayClusterShadersPostprocess_OutputRemap.cpp
+++ b/Engine/Plugins/Runtime/nDisplay/Source/DisplayClusterShaders/Private/Shaders/DisplayClusterShadersPostprocess_OutputRemap.cpp
@@ -100,7 +100,7 @@ IMPLEMENT_SHADER_TYPE(, FOutputRemapVS, OutputRemapShaderFileName, TEXT("OutputR
 
 DECLARE_GPU_STAT_NAMED(nDisplay_PostProcess_OutputRemap, TEXT("nDisplay PostProcess::OutputRemap"));
 
-bool FDisplayClusterShadersPostprocess_OutputRemap::RenderPostprocess_OutputRemap(FRHICommandListImmediate& RHICmdList, FRHITexture2D* InSourceTexture, FRHITexture2D* InRenderTargetableDestTexture, const IDisplayClusterRender_MeshComponentProxy& InMeshProxy)
+bool FDisplayClusterShadersPostprocess_OutputRemap::RenderPostprocess_OutputRemap(FRHICommandListImmediate& RHICmdList, FRHITexture* InSourceTexture, FRHITexture* InRenderTargetableDestTexture, const IDisplayClusterRender_MeshComponentProxy& InMeshProxy)
 {
 	check(IsInRenderingThread());
 
diff --git a/Engine/Plugins/Runtime/nDisplay/Source/DisplayClusterShaders/Private/Shaders/DisplayClusterShadersPostprocess_OutputRemap.h b/Engine/Plugins/Runtime/nDisplay/Source/DisplayClusterShaders/Private/Shaders/DisplayClusterShadersPostprocess_OutputRemap.h
--- a/Engine/Plugins/Runtime/nDisplay/Source/DisplayClusterShaders/Private/Shaders/DisplayClusterShadersPostprocess_OutputRemap.h
+++ b/Engine/Plugins/Runtime/nDisplay/Source/DisplayClusterShaders/Private/Shaders/DisplayClusterShadersPostprocess_OutputRemap.h
@@ -11,5 +11,5 @@ class IDisplayClusterRender_MeshComponentProxy;
 class FDisplayClusterShadersPostprocess_OutputRemap
 {
 public:
-	static bool RenderPostprocess_OutputRemap(FRHICommandListImmediate& RHICmdList, FRHITexture2D* InSourceTexture, FRHITexture2D* InRenderTargetableDestTexture, const IDisplayClusterRender_MeshComponentProxy& MeshProxy);
+	static bool RenderPostprocess_OutputRemap(FRHICommandListImmediate& RHICmdList, FRHITexture* InSourceTexture, FRHITexture* InRenderTargetableDestTexture, const IDisplayClusterRender_MeshComponentProxy& MeshProxy);
 };
diff --git a/Engine/Plugins/Runtime/nDisplay/Source/DisplayClusterShaders/Public/IDisplayClusterShaders.h b/Engine/Plugins/Runtime/nDisplay/Source/DisplayClusterShaders/Public/IDisplayClusterShaders.h
--- a/Engine/Plugins/Runtime/nDisplay/Source/DisplayClusterShaders/Public/IDisplayClusterShaders.h
+++ b/Engine/Plugins/Runtime/nDisplay/Source/DisplayClusterShaders/Public/IDisplayClusterShaders.h
@@ -108,7 +108,7 @@ public:
 	*
 	* @return - true if success
 	*/
-	virtual bool RenderPostprocess_OutputRemap(FRHICommandListImmediate& RHICmdList, FRHITexture2D* InSourceTexture, FRHITexture2D* InRenderTargetableDestTexture, const IDisplayClusterRender_MeshComponentProxy& MeshProxy) const = 0;
+	virtual bool RenderPostprocess_OutputRemap(FRHICommandListImmediate& RHICmdList, FRHITexture* InSourceTexture, FRHITexture* InRenderTargetableDestTexture, const IDisplayClusterRender_MeshComponentProxy& MeshProxy) const = 0;
 
 	/**
 	* Render postprocess Blur
@@ -120,7 +120,7 @@ public:
 	*
 	* @return - true if success
 	*/
-	virtual bool RenderPostprocess_Blur(FRHICommandListImmediate& RHICmdList, FRHITexture2D* InSourceTexture, FRHITexture2D* InRenderTargetableDestTexture, const FDisplayClusterShaderParameters_PostprocessBlur& InSettings) const = 0;
+	virtual bool RenderPostprocess_Blur(FRHICommandListImmediate& RHICmdList, FRHITexture* InSourceTexture, FRHITexture* InRenderTargetableDestTexture, const FDisplayClusterShaderParameters_PostprocessBlur& InSettings) const = 0;
 
 	/**
 	* Generate Mips texture
@@ -131,5 +131,5 @@ public:
 	*
 	* @return - true if success
 	*/
-	virtual bool GenerateMips(FRHICommandListImmediate& RHICmdList, FRHITexture2D* InOutMipsTexture, const FDisplayClusterShaderParameters_GenerateMips& InSettings) const = 0;
+	virtual bool GenerateMips(FRHICommandListImmediate& RHICmdList, FRHITexture* InOutMipsTexture, const FDisplayClusterShaderParameters_GenerateMips& InSettings) const = 0;
 };
diff --git a/Engine/Plugins/Runtime/nDisplay/Source/DisplayClusterShaders/Public/ShaderParameters/DisplayClusterShaderParameters_WarpBlend.h b/Engine/Plugins/Runtime/nDisplay/Source/DisplayClusterShaders/Public/ShaderParameters/DisplayClusterShaderParameters_WarpBlend.h
--- a/Engine/Plugins/Runtime/nDisplay/Source/DisplayClusterShaders/Public/ShaderParameters/DisplayClusterShaderParameters_WarpBlend.h
+++ b/Engine/Plugins/Runtime/nDisplay/Source/DisplayClusterShaders/Public/ShaderParameters/DisplayClusterShaderParameters_WarpBlend.h
@@ -13,10 +13,10 @@ struct FDisplayClusterShaderParameters_WarpBlend
 {
 	struct FResourceWithRect
 	{
-		FRHITexture2D* Texture;
+		FRHITexture* Texture;
 		FIntRect       Rect;
 
-		void Set(FRHITexture2D* InTexture, const FIntRect& InRect)
+		void Set(FRHITexture* InTexture, const FIntRect& InRect)
 		{
 			Texture = InTexture;
 			Rect = InRect;
diff --git a/Engine/Plugins/Runtime/nDisplay/Source/SharedMemoryMedia/Private/SharedMemoryMediaTextureSampleConverter.h b/Engine/Plugins/Runtime/nDisplay/Source/SharedMemoryMedia/Private/SharedMemoryMediaTextureSampleConverter.h
--- a/Engine/Plugins/Runtime/nDisplay/Source/SharedMemoryMedia/Private/SharedMemoryMediaTextureSampleConverter.h
+++ b/Engine/Plugins/Runtime/nDisplay/Source/SharedMemoryMedia/Private/SharedMemoryMediaTextureSampleConverter.h
@@ -22,7 +22,7 @@ public:
 		return ConverterInfoFlags_PreprocessOnly;
 	}
 
-	virtual bool Convert(FTexture2DRHIRef& InDstTexture, const FConversionHints& Hints) override
+	virtual bool Convert(FTextureRHIRef& InDstTexture, const FConversionHints& Hints) override
 	{
 		if (!Player)
 		{
diff --git a/Engine/Plugins/Tests/RHITests/Source/RHITests/Public/RHITextureTests.h b/Engine/Plugins/Tests/RHITests/Source/RHITests/Public/RHITextureTests.h
--- a/Engine/Plugins/Tests/RHITests/Source/RHITests/Public/RHITextureTests.h
+++ b/Engine/Plugins/Tests/RHITests/Source/RHITests/Public/RHITextureTests.h
@@ -185,7 +185,7 @@ public:
 				.SetNumMips(NumMips)
 				.SetFlags(ETextureCreateFlags::UAV | ETextureCreateFlags::ShaderResource);
 
-			FTexture3DRHIRef Texture = RHICreateTexture(Desc);
+			FTextureRHIRef Texture = RHICreateTexture(Desc);
 
 			for (uint32 Mip = 0; Mip < NumMips; ++Mip)
 			{
@@ -617,7 +617,7 @@ public:
 
 			Desc.SetFlags(ETextureCreateFlags::ShaderResource);
 
-			FTexture2DRHIRef Texture = RHICreateTexture(Desc);
+			FTextureRHIRef Texture = RHICreateTexture(Desc);
 			if (Texture == nullptr)
 			{
 				UE_LOG(LogRHIUnitTestCommandlet, Display, TEXT("Test failed (couldn't create texture). \"%s\""), *TestName);
@@ -1371,8 +1371,8 @@ public:
 		const EPixelFormat Format = PF_R8G8B8A8;
 		FRHITextureCreateDesc DescArray = FRHITextureCreateDesc::Create2DArray(TEXT("Multiple Texture Slice Lock"), SliceDim, ArrayLength, Format).DetermineInititialState();
 		
-		FTexture2DArrayRHIRef Texture_SingleLock = RHICreateTexture(DescArray);
-		FTexture2DArrayRHIRef Texture_MultipleLock = RHICreateTexture(DescArray);
+		FTextureRHIRef Texture_SingleLock = RHICreateTexture(DescArray);
+		FTextureRHIRef Texture_MultipleLock = RHICreateTexture(DescArray);
 		
 		// Create Reasonable Test Data that can be seen in the GPU debugger if required
 		const auto WriteSliceTestData = [=](uint8* pDestData, uint32 DestStride, uint32 SliceIdx)
diff --git a/Engine/Plugins/VirtualProduction/DMX/DMXPixelMapping/Source/DMXPixelMappingRenderer/Private/DMXPixelMappingRenderer.cpp b/Engine/Plugins/VirtualProduction/DMX/DMXPixelMapping/Source/DMXPixelMappingRenderer/Private/DMXPixelMappingRenderer.cpp
--- a/Engine/Plugins/VirtualProduction/DMX/DMXPixelMapping/Source/DMXPixelMappingRenderer/Private/DMXPixelMappingRenderer.cpp
+++ b/Engine/Plugins/VirtualProduction/DMX/DMXPixelMapping/Source/DMXPixelMappingRenderer/Private/DMXPixelMappingRenderer.cpp
@@ -367,14 +367,14 @@ void FDMXPixelMappingRenderer::RenderPreview(const FTextureResource* TextureReso
 }
 #endif // WITH_EDITOR
 
-void FDMXPixelMappingRenderer::RenderTextureToRectangle(const FTextureResource* InTextureResource, const FTexture2DRHIRef InRenderTargetTexture, FVector2D InSize, bool bSRGBSource) const
+void FDMXPixelMappingRenderer::RenderTextureToRectangle(const FTextureResource* InTextureResource, const FTextureRHIRef InRenderTargetTexture, FVector2D InSize, bool bSRGBSource) const
 {
 	check(IsInGameThread());
 
 	struct FRenderContext
 	{
 		const FTextureResource* TextureResource = nullptr;
-		const FTexture2DRHIRef Texture2DRHI = nullptr;
+		const FTextureRHIRef Texture2DRHI = nullptr;
 		FVector2D ViewportSize;
 		bool bSRGBSource;
 	};
diff --git a/Engine/Plugins/VirtualProduction/DMX/DMXPixelMapping/Source/DMXPixelMappingRenderer/Private/DMXPixelMappingRenderer.h b/Engine/Plugins/VirtualProduction/DMX/DMXPixelMapping/Source/DMXPixelMappingRenderer/Private/DMXPixelMappingRenderer.h
--- a/Engine/Plugins/VirtualProduction/DMX/DMXPixelMapping/Source/DMXPixelMappingRenderer/Private/DMXPixelMappingRenderer.h
+++ b/Engine/Plugins/VirtualProduction/DMX/DMXPixelMapping/Source/DMXPixelMappingRenderer/Private/DMXPixelMappingRenderer.h
@@ -33,7 +33,7 @@ public:
 	virtual void RenderMaterial(UTextureRenderTarget2D* InRenderTarget, UMaterialInterface* InMaterialInterface) const override;
 	virtual void RenderWidget(UTextureRenderTarget2D* InRenderTarget, UUserWidget* InUserWidget) const override;
 
-	virtual void RenderTextureToRectangle(const FTextureResource* InTextureResource, const FTexture2DRHIRef InRenderTargetTexture, FVector2D InSize, bool bSRGBSource) const override;
+	virtual void RenderTextureToRectangle(const FTextureResource* InTextureResource, const FTextureRHIRef InRenderTargetTexture, FVector2D InSize, bool bSRGBSource) const override;
 
 #if WITH_EDITOR
 	virtual void RenderPreview(const FTextureResource* TextureResource, const FTextureResource* DownsampleResource, TArray<FDMXPixelMappingDownsamplePixelPreviewParam>&& InPixelPreviewParamSet) const override;
diff --git a/Engine/Plugins/VirtualProduction/DMX/DMXPixelMapping/Source/DMXPixelMappingRenderer/Public/IDMXPixelMappingRenderer.h b/Engine/Plugins/VirtualProduction/DMX/DMXPixelMapping/Source/DMXPixelMappingRenderer/Public/IDMXPixelMappingRenderer.h
--- a/Engine/Plugins/VirtualProduction/DMX/DMXPixelMapping/Source/DMXPixelMappingRenderer/Public/IDMXPixelMappingRenderer.h
+++ b/Engine/Plugins/VirtualProduction/DMX/DMXPixelMapping/Source/DMXPixelMappingRenderer/Public/IDMXPixelMappingRenderer.h
@@ -161,7 +161,7 @@ public:
 	 * @param bSRGBSource					If the source texture is sRGB
 	 */
 	UE_DEPRECATED(5.3, "Removed without replacement. Please use core engine methods to draw texture to rectangle.")
-	virtual void RenderTextureToRectangle(const FTextureResource* InTextureResource, const FTexture2DRHIRef InRenderTargetTexture, FVector2D InSize, bool bSRGBSource) const = 0;
+	virtual void RenderTextureToRectangle(const FTextureResource* InTextureResource, const FTextureRHIRef InRenderTargetTexture, FVector2D InSize, bool bSRGBSource) const = 0;
 
 #if WITH_EDITOR
 	/**
diff --git a/Engine/Plugins/VirtualProduction/Rivermax/RivermaxMedia/Source/RivermaxMedia/Private/RivermaxMediaTextureSampleConverter.cpp b/Engine/Plugins/VirtualProduction/Rivermax/RivermaxMedia/Source/RivermaxMedia/Private/RivermaxMediaTextureSampleConverter.cpp
--- a/Engine/Plugins/VirtualProduction/Rivermax/RivermaxMedia/Source/RivermaxMedia/Private/RivermaxMediaTextureSampleConverter.cpp
+++ b/Engine/Plugins/VirtualProduction/Rivermax/RivermaxMedia/Source/RivermaxMedia/Private/RivermaxMediaTextureSampleConverter.cpp
@@ -19,7 +19,7 @@ void FRivermaxMediaTextureSampleConverter::Setup(const TSharedPtr<FRivermaxMedia
 	Sample = InSample;
 }
 
-bool FRivermaxMediaTextureSampleConverter::Convert(FTexture2DRHIRef& InDestinationTexture, const FConversionHints& Hints)
+bool FRivermaxMediaTextureSampleConverter::Convert(FTextureRHIRef& InDestinationTexture, const FConversionHints& Hints)
 {
 	TRACE_CPUPROFILER_EVENT_SCOPE(RivermaxSampleConverter::Convert);
 
diff --git a/Engine/Plugins/VirtualProduction/Rivermax/RivermaxMedia/Source/RivermaxMedia/Private/RivermaxMediaTextureSampleConverter.h b/Engine/Plugins/VirtualProduction/Rivermax/RivermaxMedia/Source/RivermaxMedia/Private/RivermaxMediaTextureSampleConverter.h
--- a/Engine/Plugins/VirtualProduction/Rivermax/RivermaxMedia/Source/RivermaxMedia/Private/RivermaxMediaTextureSampleConverter.h
+++ b/Engine/Plugins/VirtualProduction/Rivermax/RivermaxMedia/Source/RivermaxMedia/Private/RivermaxMediaTextureSampleConverter.h
@@ -56,7 +56,7 @@ public:
 	void Setup(const TSharedPtr<FRivermaxMediaTextureSample>& InSample);
 
 	//~ Begin IMediaTextureSampleConverter interface
-	virtual bool Convert(FTexture2DRHIRef& InDstTexture, const FConversionHints& Hints) override;
+	virtual bool Convert(FTextureRHIRef& InDstTexture, const FConversionHints& Hints) override;
 	virtual uint32 GetConverterInfoFlags() const override;
 	//~ End IMediaTextureSampleConverter interface
 
diff --git a/Engine/Plugins/VirtualProduction/TextureShare/Source/TextureShare/Private/Game/ViewExtension/TextureShareSceneViewExtension.cpp b/Engine/Plugins/VirtualProduction/TextureShare/Source/TextureShare/Private/Game/ViewExtension/TextureShareSceneViewExtension.cpp
--- a/Engine/Plugins/VirtualProduction/TextureShare/Source/TextureShare/Private/Game/ViewExtension/TextureShareSceneViewExtension.cpp
+++ b/Engine/Plugins/VirtualProduction/TextureShare/Source/TextureShare/Private/Game/ViewExtension/TextureShareSceneViewExtension.cpp
@@ -308,7 +308,7 @@ void FTextureShareSceneViewExtension::PostRenderViewFamily_RenderThread(FRHIComm
 		// Share only if the resource is requested from a remote process
 		if (const FTextureShareCoreResourceRequest* ExistResourceRequest = ObjectProxy->GetData_RenderThread().FindResourceRequest(FTextureShareCoreResourceDesc(UE::TextureShareStrings::SceneTextures::FinalColor, ViewIt.ViewInfo.ViewDesc, ETextureShareTextureOp::Undefined)))
 		{
-			FTexture2DRHIRef RenderTargetTexture = InViewFamily.RenderTarget->GetRenderTargetTexture();
+			FTextureRHIRef RenderTargetTexture = InViewFamily.RenderTarget->GetRenderTargetTexture();
 			if (RenderTargetTexture.IsValid())
 			{
 				// Send
@@ -346,7 +346,7 @@ void FTextureShareSceneViewExtension::OnResolvedSceneColor_RenderThread(FRDGBuil
 	}
 }
 
-void FTextureShareSceneViewExtension::OnBackBufferReadyToPresent_RenderThread(SWindow& SlateWindow, const FTexture2DRHIRef& InBackbuffer)
+void FTextureShareSceneViewExtension::OnBackBufferReadyToPresent_RenderThread(SWindow& SlateWindow, const FTextureRHIRef& InBackbuffer)
 {
 	FScopeLock Lock(&DataCS);
 
diff --git a/Engine/Plugins/VirtualProduction/TextureShare/Source/TextureShare/Private/Game/ViewExtension/TextureShareSceneViewExtension.h b/Engine/Plugins/VirtualProduction/TextureShare/Source/TextureShare/Private/Game/ViewExtension/TextureShareSceneViewExtension.h
--- a/Engine/Plugins/VirtualProduction/TextureShare/Source/TextureShare/Private/Game/ViewExtension/TextureShareSceneViewExtension.h
+++ b/Engine/Plugins/VirtualProduction/TextureShare/Source/TextureShare/Private/Game/ViewExtension/TextureShareSceneViewExtension.h
@@ -13,7 +13,7 @@ struct FTextureShareSceneView;
  * Convenience type definition of a function that gives an opinion of whether the scene view extension should be active in the given context for the current frame.
  */
 using TFunctionTextureShareViewExtension = TFunction<void(FRHICommandListImmediate& RHICmdList, class FTextureShareSceneViewExtension& InViewExtension)>;
-using TFunctionTextureShareOnBackBufferReadyToPresent = TFunction<void(FRHICommandListImmediate& RHICmdList, class FTextureShareSceneViewExtension& InViewExtension, const FTexture2DRHIRef& InBackbuffer)>;
+using TFunctionTextureShareOnBackBufferReadyToPresent = TFunction<void(FRHICommandListImmediate& RHICmdList, class FTextureShareSceneViewExtension& InViewExtension, const FTextureRHIRef& InBackbuffer)>;
 
 /**
  * A view extension to handle a multi-threaded renderer for a TextureShare object.
@@ -40,7 +40,7 @@ public:
 	virtual void PostRenderViewFamily_RenderThread(FRDGBuilder& GraphBuilder, FSceneViewFamily& InViewFamily) override;
 
 	virtual void OnResolvedSceneColor_RenderThread(FRDGBuilder& GraphBuilder, const FSceneTextures& SceneTextures);
-	virtual void OnBackBufferReadyToPresent_RenderThread(SWindow&, const FTexture2DRHIRef&);
+	virtual void OnBackBufferReadyToPresent_RenderThread(SWindow&, const FTextureRHIRef&);
 
 	virtual bool IsActiveThisFrame_Internal(const FSceneViewExtensionContext& Context) const override;
 	//~End ISceneViewExtension interface
diff --git a/Engine/Plugins/VirtualProduction/TextureShare/Source/TextureShare/Private/Game/WorldSubsystem/TextureShareWorldSubsystemObjectProxy.cpp b/Engine/Plugins/VirtualProduction/TextureShare/Source/TextureShare/Private/Game/WorldSubsystem/TextureShareWorldSubsystemObjectProxy.cpp
--- a/Engine/Plugins/VirtualProduction/TextureShare/Source/TextureShare/Private/Game/WorldSubsystem/TextureShareWorldSubsystemObjectProxy.cpp
+++ b/Engine/Plugins/VirtualProduction/TextureShare/Source/TextureShare/Private/Game/WorldSubsystem/TextureShareWorldSubsystemObjectProxy.cpp
@@ -131,7 +131,7 @@ void FTextureShareWorldSubsystemObjectProxy::Update_RenderThread(FRHICommandList
 				}
 			};
 
-			TFunctionTextureShareOnBackBufferReadyToPresent OnBackBufferReadyToPresentFunction = [ProxyData, bBackbufferShared](FRHICommandListImmediate& RHICmdList, FTextureShareSceneViewExtension& InViewExtension, const FTexture2DRHIRef& InBackbuffer)
+			TFunctionTextureShareOnBackBufferReadyToPresent OnBackBufferReadyToPresentFunction = [ProxyData, bBackbufferShared](FRHICommandListImmediate& RHICmdList, FTextureShareSceneViewExtension& InViewExtension, const FTextureRHIRef& InBackbuffer)
 			{
 				// Share backbuffer
 				if (bBackbufferShared)
@@ -215,7 +215,7 @@ bool FTextureShareWorldSubsystemObjectProxy::UpdateResources_RenderThread(FRHICo
 	return bResult;
 }
 
-bool FTextureShareWorldSubsystemObjectProxy::UpdateFrameProxyBackbuffer_RenderThread(FRHICommandListImmediate& RHICmdList, const FTexture2DRHIRef& InBackbuffer)
+bool FTextureShareWorldSubsystemObjectProxy::UpdateFrameProxyBackbuffer_RenderThread(FRHICommandListImmediate& RHICmdList, const FTextureRHIRef& InBackbuffer)
 {
 	if (ObjectProxy.IsValid()
 		&& InBackbuffer.IsValid()
diff --git a/Engine/Plugins/VirtualProduction/TextureShare/Source/TextureShare/Private/Game/WorldSubsystem/TextureShareWorldSubsystemObjectProxy.h b/Engine/Plugins/VirtualProduction/TextureShare/Source/TextureShare/Private/Game/WorldSubsystem/TextureShareWorldSubsystemObjectProxy.h
--- a/Engine/Plugins/VirtualProduction/TextureShare/Source/TextureShare/Private/Game/WorldSubsystem/TextureShareWorldSubsystemObjectProxy.h
+++ b/Engine/Plugins/VirtualProduction/TextureShare/Source/TextureShare/Private/Game/WorldSubsystem/TextureShareWorldSubsystemObjectProxy.h
@@ -28,7 +28,7 @@ public:
 private:
 	void Update_RenderThread(FRHICommandListImmediate& RHICmdList, TSharedPtr<FProxyResourcesData> ProxyData);
 	bool UpdateResources_RenderThread(FRHICommandListImmediate& RHICmdList, TSharedPtr<FProxyResourcesData> ProxyData);
-	bool UpdateFrameProxyBackbuffer_RenderThread(FRHICommandListImmediate& RHICmdList, const FTexture2DRHIRef& InBackbuffer);
+	bool UpdateFrameProxyBackbuffer_RenderThread(FRHICommandListImmediate& RHICmdList, const FTextureRHIRef& InBackbuffer);
 
 private:
 	TSharedPtr<ITextureShareObjectProxy, ESPMode::ThreadSafe> ObjectProxy;
diff --git a/Engine/Plugins/VirtualProduction/TextureShare/Source/TextureShare/Private/Module/TextureShareAPI.cpp b/Engine/Plugins/VirtualProduction/TextureShare/Source/TextureShare/Private/Module/TextureShareAPI.cpp
--- a/Engine/Plugins/VirtualProduction/TextureShare/Source/TextureShare/Private/Module/TextureShareAPI.cpp
+++ b/Engine/Plugins/VirtualProduction/TextureShare/Source/TextureShare/Private/Module/TextureShareAPI.cpp
@@ -233,7 +233,7 @@ void FTextureShareAPI::OnResolvedSceneColor_RenderThread(FRDGBuilder& GraphBuild
 	}
 }
 
-void FTextureShareAPI::OnBackBufferReadyToPresent_RenderThread(SWindow& InWindow, const FTexture2DRHIRef& InBackbuffer)
+void FTextureShareAPI::OnBackBufferReadyToPresent_RenderThread(SWindow& InWindow, const FTextureRHIRef& InBackbuffer)
 {
 	check(IsInRenderingThread());
 
diff --git a/Engine/Plugins/VirtualProduction/TextureShare/Source/TextureShare/Private/Module/TextureShareAPI.h b/Engine/Plugins/VirtualProduction/TextureShare/Source/TextureShare/Private/Module/TextureShareAPI.h
--- a/Engine/Plugins/VirtualProduction/TextureShare/Source/TextureShare/Private/Module/TextureShareAPI.h
+++ b/Engine/Plugins/VirtualProduction/TextureShare/Source/TextureShare/Private/Module/TextureShareAPI.h
@@ -58,7 +58,7 @@ private:
 	void OnResolvedSceneColor_RenderThread(FRDGBuilder& GraphBuilder, const FSceneTextures& SceneTextures);
 
 	/** Slate app callback before present to share app backbuffer */
-	void OnBackBufferReadyToPresent_RenderThread(SWindow&, const FTexture2DRHIRef&);
+	void OnBackBufferReadyToPresent_RenderThread(SWindow&, const FTextureRHIRef&);
 	
 	void RegisterCallbacks();
 	void UnregisterCallbacks();
diff --git a/Engine/Plugins/VirtualProduction/TextureShare/Source/TextureShare/Private/Resources/TextureShareResource.cpp b/Engine/Plugins/VirtualProduction/TextureShare/Source/TextureShare/Private/Resources/TextureShareResource.cpp
--- a/Engine/Plugins/VirtualProduction/TextureShare/Source/TextureShare/Private/Resources/TextureShareResource.cpp
+++ b/Engine/Plugins/VirtualProduction/TextureShare/Source/TextureShare/Private/Resources/TextureShareResource.cpp
@@ -39,18 +39,18 @@ const FString& FTextureShareResource::GetCoreObjectName() const
 struct FCachedSharedResource
 {
 	FCachedSharedResource() = default;
-	FCachedSharedResource(const void* InNativeResourcePtr, const uint32  InGPUIndex, const FTexture2DRHIRef& InTextureRHI)
+	FCachedSharedResource(const void* InNativeResourcePtr, const uint32  InGPUIndex, const FTextureRHIRef& InTextureRHI)
 		: NativeResourcePtr(InNativeResourcePtr), TextureRHI(InTextureRHI), GPUIndex(InGPUIndex)
 	{ }
 
 	const void* NativeResourcePtr = nullptr;
-	FTexture2DRHIRef TextureRHI;
+	FTextureRHIRef TextureRHI;
 	uint32 GPUIndex = 0;
 
 	bool bUnused = false;
 };
 
-bool FTextureShareResource::FindCachedSharedResource_RenderThread(void* InNativeResourcePtr, const uint32 InGPUIndex, FTexture2DRHIRef& OutRHIResource) const
+bool FTextureShareResource::FindCachedSharedResource_RenderThread(void* InNativeResourcePtr, const uint32 InGPUIndex, FTextureRHIRef& OutRHIResource) const
 {
 	if(const FCachedSharedResource* CachedResourcePtr = CachedSharedResources.FindByPredicate([InNativeResourcePtr, InGPUIndex](const FCachedSharedResource& CachedResourceIt)
 		{
@@ -66,7 +66,7 @@ bool FTextureShareResource::FindCachedSharedResource_RenderThread(void* InNative
 	return false;
 }
 
-void FTextureShareResource::AddCachedSharedResource_RenderThread(void* InNativeResourcePtr, const uint32 InGPUIndex, const FTexture2DRHIRef& InRHIResource)
+void FTextureShareResource::AddCachedSharedResource_RenderThread(void* InNativeResourcePtr, const uint32 InGPUIndex, const FTextureRHIRef& InRHIResource)
 {
 	if (InNativeResourcePtr != nullptr && InRHIResource.IsValid())
 	{
@@ -166,7 +166,7 @@ bool FTextureShareResource::ReleaseTextureShareHandle_RenderThread()
 //////////////////////////////////////////////////////////////////////////////////////////////
 void FTextureShareResource::InitRHI(FRHICommandListBase&)
 {
-	FTexture2DRHIRef NewTextureRHI;
+	FTextureRHIRef NewTextureRHI;
 	switch (CoreObject->GetObjectDesc_RenderThread().ProcessDesc.DeviceType)
 	{
 	case ETextureShareDeviceType::D3D12:
@@ -181,7 +181,7 @@ void FTextureShareResource::InitRHI(FRHICommandListBase&)
 	TextureRHI = (FTextureRHIRef&)NewTextureRHI;
 }
 
-void FTextureShareResource::InitDynamicRHI_Default(FTexture2DRHIRef& OutTextureRHI)
+void FTextureShareResource::InitDynamicRHI_Default(FTextureRHIRef& OutTextureRHI)
 {
 	FRHITextureCreateDesc Desc =
 		FRHITextureCreateDesc::Create2D(TEXT("TextureShareResource"), GetSizeX(), GetSizeY(), ResourceSettings.Format)
diff --git a/Engine/Plugins/VirtualProduction/TextureShare/Source/TextureShare/Private/Resources/TextureShareResource.h b/Engine/Plugins/VirtualProduction/TextureShare/Source/TextureShare/Private/Resources/TextureShareResource.h
--- a/Engine/Plugins/VirtualProduction/TextureShare/Source/TextureShare/Private/Resources/TextureShareResource.h
+++ b/Engine/Plugins/VirtualProduction/TextureShare/Source/TextureShare/Private/Resources/TextureShareResource.h
@@ -18,11 +18,11 @@ struct FTextureShareCoreResourceRequest;
  */
 struct FTextureShareExternalTextureRHI
 {
-	FTextureShareExternalTextureRHI(const FTexture2DRHIRef InTextureRHI, const FTextureShareCoreResourceHandle* InResourceHandle, const uint32 InGPUIndex = 0)
+	FTextureShareExternalTextureRHI(const FTextureRHIRef InTextureRHI, const FTextureShareCoreResourceHandle* InResourceHandle, const uint32 InGPUIndex = 0)
 		: TextureRHI(InTextureRHI), ResourceHandle(InResourceHandle), GPUIndex(InGPUIndex)
 	{ }
 
-	const FTexture2DRHIRef TextureRHI;
+	const FTextureRHIRef TextureRHI;
 	const FTextureShareCoreResourceHandle* ResourceHandle;
 	const uint32 GPUIndex;
 };
@@ -48,9 +48,9 @@ public:
 		return ResourceSettings;
 	}
 
-	const FTexture2DRHIRef& GetResourceTextureRHI() const
+	const FTextureRHIRef& GetResourceTextureRHI() const
 	{
-		return (const FTexture2DRHIRef&)TextureRHI;
+		return (const FTextureRHIRef&)TextureRHI;
 	}
 
 	const FString& GetCoreObjectName() const;
@@ -89,8 +89,8 @@ public:
 	}
 
 protected:
-	void InitDynamicRHI_Default(FTexture2DRHIRef& OutTextureRHI);
-	void InitDynamicRHI_D3D12(FTexture2DRHIRef& OutTextureRHI);
+	void InitDynamicRHI_Default(FTextureRHIRef& OutTextureRHI);
+	void InitDynamicRHI_D3D12(FTextureRHIRef& OutTextureRHI);
 
 	bool D3D11RegisterResourceHandle_RenderThread(FRHICommandListImmediate& RHICmdList, const FTextureShareCoreResourceRequest& InResourceRequest);
 	bool D3D11ReleaseTextureShareHandle_RenderThread();
@@ -104,8 +104,8 @@ protected:
 	bool VulkanReleaseTextureShareHandle_RenderThread();
 #endif
 
-	bool FindCachedSharedResource_RenderThread(void* InNativeResourcePtr, const uint32 InGPUIndex, FTexture2DRHIRef& OutRHIResource) const;
-	void AddCachedSharedResource_RenderThread(void* InNativeResourcePtr, const uint32 InGPUIndex, const FTexture2DRHIRef& InRHIResource);
+	bool FindCachedSharedResource_RenderThread(void* InNativeResourcePtr, const uint32 InGPUIndex, FTextureRHIRef& OutRHIResource) const;
+	void AddCachedSharedResource_RenderThread(void* InNativeResourcePtr, const uint32 InGPUIndex, const FTextureRHIRef& InRHIResource);
 	void CopyToDestResources_RenderThread(FRHICommandListImmediate& RHICmdList, const TArray<FTextureShareExternalTextureRHI>& InDestResources);
 
 private:
diff --git a/Engine/Plugins/VirtualProduction/TextureShare/Source/TextureShare/Private/Resources/TextureShareResource_D3D12.cpp b/Engine/Plugins/VirtualProduction/TextureShare/Source/TextureShare/Private/Resources/TextureShareResource_D3D12.cpp
--- a/Engine/Plugins/VirtualProduction/TextureShare/Source/TextureShare/Private/Resources/TextureShareResource_D3D12.cpp
+++ b/Engine/Plugins/VirtualProduction/TextureShare/Source/TextureShare/Private/Resources/TextureShareResource_D3D12.cpp
@@ -81,7 +81,7 @@ using namespace UE::TextureShareCore;
 //////////////////////////////////////////////////////////////////////////////////////////////
 // FTextureShareResource
 //////////////////////////////////////////////////////////////////////////////////////////////
-void FTextureShareResource::InitDynamicRHI_D3D12(FTexture2DRHIRef& OutTextureRHI)
+void FTextureShareResource::InitDynamicRHI_D3D12(FTextureRHIRef& OutTextureRHI)
 {
 	switch (ResourceDesc.ResourceType)
 	{
diff --git a/Engine/Plugins/VirtualProduction/TextureShare/Source/TextureShareDisplayCluster/Private/PostProcess/TextureSharePostprocess_FrameProxy.cpp b/Engine/Plugins/VirtualProduction/TextureShare/Source/TextureShareDisplayCluster/Private/PostProcess/TextureSharePostprocess_FrameProxy.cpp
--- a/Engine/Plugins/VirtualProduction/TextureShare/Source/TextureShareDisplayCluster/Private/PostProcess/TextureSharePostprocess_FrameProxy.cpp
+++ b/Engine/Plugins/VirtualProduction/TextureShare/Source/TextureShareDisplayCluster/Private/PostProcess/TextureSharePostprocess_FrameProxy.cpp
@@ -51,7 +51,7 @@ void FTextureSharePostprocess::ShareViewport_RenderThread(FRHICommandListImmedia
 					const TSharedPtr<IDisplayClusterProjectionPolicy, ESPMode::ThreadSafe>& PrjPolicy = ViewportProxyIt->GetProjectionPolicy_RenderThread();
 					if (PrjPolicy->ShouldUseAdditionalTargetableResource())
 					{
-						TArray<FRHITexture2D*> Resources;
+						TArray<FRHITexture*> Resources;
 						if (ViewportProxyIt->GetResources_RenderThread(EDisplayClusterViewportResourceType::AdditionalTargetableResource, Resources) && Resources.Num())
 						{
 							ResourceType = EDisplayClusterViewportResourceType::AdditionalTargetableResource;
@@ -61,7 +61,7 @@ void FTextureSharePostprocess::ShareViewport_RenderThread(FRHICommandListImmedia
 
 				const bool bMonoscopic = ViewportProxyIt->GetContexts_RenderThread().Num() == 1;
 
-				TArray<FRHITexture2D*> ViewportResources;
+				TArray<FRHITexture*> ViewportResources;
 				TArray<FIntRect> ViewportResourceRects;
 				if (ViewportProxyIt->GetResourcesWithRects_RenderThread(ResourceType, ViewportResources, ViewportResourceRects))
 				{
@@ -91,12 +91,12 @@ void FTextureSharePostprocess::ShareFrame_RenderThread(FRHICommandListImmediate&
 {
 	if (InViewportManagerProxy)
 	{
-		TArray<FRHITexture2D*> FrameResources;
-		TArray<FRHITexture2D*> AdditionalFrameResources;
+		TArray<FRHITexture*> FrameResources;
+		TArray<FRHITexture*> AdditionalFrameResources;
 		TArray<FIntPoint> TargetOffsets;
 		if (InViewportManagerProxy->GetFrameTargets_RenderThread(FrameResources, TargetOffsets, &AdditionalFrameResources))
 		{
-			TArray<FRHITexture2D*>& SharedResources = (InResourceType == EDisplayClusterViewportResourceType::AdditionalFrameTargetableResource)
+			TArray<FRHITexture*>& SharedResources = (InResourceType == EDisplayClusterViewportResourceType::AdditionalFrameTargetableResource)
 				? AdditionalFrameResources : FrameResources;
 
 			const bool bMonoscopic = SharedResources.Num() == 1;
diff --git a/Engine/Source/Developer/FunctionalTesting/Private/FunctionalUIScreenshotTest.cpp b/Engine/Source/Developer/FunctionalTesting/Private/FunctionalUIScreenshotTest.cpp
--- a/Engine/Source/Developer/FunctionalTesting/Private/FunctionalUIScreenshotTest.cpp
+++ b/Engine/Source/Developer/FunctionalTesting/Private/FunctionalUIScreenshotTest.cpp
@@ -47,7 +47,7 @@ void GetBackbufferInfo(const FViewport* InViewport, EPixelFormat* OutPixelFormat
 	{
 		FViewportRHIRef ViewportRHI = InViewport->GetViewportRHI();
 		check(ViewportRHI.IsValid());
-		FTexture2DRHIRef BackbufferTexture = RHIGetViewportBackBuffer(ViewportRHI);
+		FTextureRHIRef BackbufferTexture = RHIGetViewportBackBuffer(ViewportRHI);
 		check(BackbufferTexture.IsValid());
 		*OutPixelFormat = BackbufferTexture->GetFormat();
 		*OutIsSRGB = (BackbufferTexture->GetFlags() & TexCreate_SRGB) == TexCreate_SRGB;
@@ -180,7 +180,7 @@ void ReadBackbuffer(const FViewport* InViewport, TArray<FColor>* OutPixels)
 		[InViewport, OutPixels](FRHICommandListImmediate& RHICmdList)
 	{
 		FViewportRHIRef ViewportRHI = InViewport->GetViewportRHI();
-		FTexture2DRHIRef BackbufferTexture = RHIGetViewportBackBuffer(ViewportRHI);
+		FTextureRHIRef BackbufferTexture = RHIGetViewportBackBuffer(ViewportRHI);
 		RHICmdList.ReadSurfaceData(
 			BackbufferTexture,
 			FIntRect(0, 0, BackbufferTexture->GetSizeX(), BackbufferTexture->GetSizeY()),
diff --git a/Engine/Source/Developer/MaterialBaking/Private/MaterialBakingModule.cpp b/Engine/Source/Developer/MaterialBaking/Private/MaterialBakingModule.cpp
--- a/Engine/Source/Developer/MaterialBaking/Private/MaterialBakingModule.cpp
+++ b/Engine/Source/Developer/MaterialBaking/Private/MaterialBakingModule.cpp
@@ -168,19 +168,19 @@ namespace
 	class FStagingBufferPool
 	{
 	public:
-		FTexture2DRHIRef CreateStagingBuffer_RenderThread(FRHICommandListImmediate& RHICmdList, int32 Width, int32 Height, EPixelFormat Format, bool bIsSRGB)
+		FTextureRHIRef CreateStagingBuffer_RenderThread(FRHICommandListImmediate& RHICmdList, int32 Width, int32 Height, EPixelFormat Format, bool bIsSRGB)
 		{
 			TRACE_CPUPROFILER_EVENT_SCOPE(CreateStagingBuffer_RenderThread)
 
 			auto StagingBufferPredicate = 
-				[Width, Height, Format, bIsSRGB](const FTexture2DRHIRef& Texture2DRHIRef)
+				[Width, Height, Format, bIsSRGB](const FTextureRHIRef& Texture2DRHIRef)
 				{
 					return Texture2DRHIRef->GetSizeX() == Width && Texture2DRHIRef->GetSizeY() == Height && Texture2DRHIRef->GetFormat() == Format && bool(Texture2DRHIRef->GetFlags() & TexCreate_SRGB) == bIsSRGB;
 				};
 
 			// Process any staging buffers available for unmapping
 			{
-				TArray<FTexture2DRHIRef> ToUnmapLocal;
+				TArray<FTextureRHIRef> ToUnmapLocal;
 				{
 					FScopeLock Lock(&ToUnmapLock);
 					ToUnmapLocal = MoveTemp(ToUnmap);
@@ -198,7 +198,7 @@ namespace
 
 			if (Index != -1)
 			{
-				FTexture2DRHIRef StagingBuffer = MoveTemp(Pool[Index]);
+				FTextureRHIRef StagingBuffer = MoveTemp(Pool[Index]);
 				Pool.RemoveAtSwap(Index);
 				return StagingBuffer;
 			}
@@ -217,7 +217,7 @@ namespace
 			return RHICreateTexture(Desc);
 		}
 
-		void ReleaseStagingBufferForUnmap_AnyThread(FTexture2DRHIRef& Texture2DRHIRef)
+		void ReleaseStagingBufferForUnmap_AnyThread(FTextureRHIRef& Texture2DRHIRef)
 		{
 			TRACE_CPUPROFILER_EVENT_SCOPE(ReleaseStagingBufferForUnmap_AnyThread)
 			FScopeLock Lock(&ToUnmapLock);
@@ -227,7 +227,7 @@ namespace
 		void Clear_RenderThread(FRHICommandListImmediate& RHICmdList)
 		{
 			TRACE_CPUPROFILER_EVENT_SCOPE(Clear_RenderThread)
-			for (FTexture2DRHIRef& StagingSurface : ToUnmap)
+			for (FTextureRHIRef& StagingSurface : ToUnmap)
 			{
 				RHICmdList.UnmapStagingSurface(StagingSurface);
 			}
@@ -242,11 +242,11 @@ namespace
 		}
 
 	private:
-		TArray<FTexture2DRHIRef> Pool;
+		TArray<FTextureRHIRef> Pool;
 
 		// Not contented enough to warrant the use of lockless structures.
-		FCriticalSection         ToUnmapLock;
-		TArray<FTexture2DRHIRef> ToUnmap;
+		FCriticalSection       ToUnmapLock;
+		TArray<FTextureRHIRef> ToUnmap;
 	};
 
 	struct FRenderItemKey
@@ -953,7 +953,7 @@ private:
 			{
 				FTextureRenderTargetResource* RenderTargetResource = RenderTarget->GetRenderTargetResource();
 
-				FTexture2DRHIRef StagingBufferRef = StagingBufferPool.CreateStagingBuffer_RenderThread(RHICmdList, RenderTargetResource->GetSizeX(), RenderTargetResource->GetSizeY(), RenderTarget->GetFormat(), RenderTarget->IsSRGB());
+				FTextureRHIRef StagingBufferRef = StagingBufferPool.CreateStagingBuffer_RenderThread(RHICmdList, RenderTargetResource->GetSizeX(), RenderTargetResource->GetSizeY(), RenderTarget->GetFormat(), RenderTarget->IsSRGB());
 				FGPUFenceRHIRef GPUFence = RHICreateGPUFence(TEXT("MaterialBackingFence"));
 				TransitionAndCopyTexture(RHICmdList, RenderTargetResource->GetRenderTargetTexture(), StagingBufferRef, {});
 				RHICmdList.WriteGPUFence(GPUFence);
@@ -961,7 +961,7 @@ private:
 				// Prepare a lambda for final processing that will be executed asynchronously
 				NumTasks++;
 				auto FinalProcessing_AnyThread =
-					[this, CurrentMaterialSettings, &CurrentOutput, Property](FTexture2DRHIRef& StagingBuffer, void* Data, int32 DataWidth, int32 DataHeight)
+					[this, CurrentMaterialSettings, &CurrentOutput, Property](FTextureRHIRef& StagingBuffer, void* Data, int32 DataWidth, int32 DataHeight)
 					{
 						TRACE_CPUPROFILER_EVENT_SCOPE(FinalProcessing)
 
diff --git a/Engine/Source/Editor/Blutility/Private/AsyncImageExport.cpp b/Engine/Source/Editor/Blutility/Private/AsyncImageExport.cpp
--- a/Engine/Source/Editor/Blutility/Private/AsyncImageExport.cpp
+++ b/Engine/Source/Editor/Blutility/Private/AsyncImageExport.cpp
@@ -63,7 +63,7 @@ void UAsyncImageExport::Activate()
 			[this, WeakThis](FRHICommandListImmediate& RHICmdList)
 		{
 			FTextureResource* Resource = TextureToExport->GetResource();
-			FRHITexture2D* ResourceRHI = Resource->GetTexture2DRHI();
+			FRHITexture* ResourceRHI = Resource->GetTexture2DRHI();
 
 			TArray<FColor> OutPixels;
 			if (ensure(ResourceRHI))
diff --git a/Engine/Source/Editor/GraphEditor/Private/MaterialNodes/SGraphNodeMaterialBase.cpp b/Engine/Source/Editor/GraphEditor/Private/MaterialNodes/SGraphNodeMaterialBase.cpp
--- a/Engine/Source/Editor/GraphEditor/Private/MaterialNodes/SGraphNodeMaterialBase.cpp
+++ b/Engine/Source/Editor/GraphEditor/Private/MaterialNodes/SGraphNodeMaterialBase.cpp
@@ -87,7 +87,7 @@ public:
 	}
 
 	/** Sets the texture that this target renders to */
-	void SetRenderTargetTexture( FTexture2DRHIRef& InRHIRef )
+	void SetRenderTargetTexture( FTextureRHIRef& InRHIRef )
 	{
 		RenderTargetTextureRHI = InRHIRef;
 	}
@@ -268,7 +268,7 @@ void FPreviewElement::Draw_RenderThread(FRHICommandListImmediate& RHICmdList, co
 {
 	if(ExpressionPreview)
 	{
-		RenderTarget->SetRenderTargetTexture(*(FTexture2DRHIRef*)InWindowBackBuffer);
+		RenderTarget->SetRenderTargetTexture(*(FTextureRHIRef*)InWindowBackBuffer);
 		{
 			// Check realtime mode for whether to pass current time to canvas
 			double CurrentTime = bIsRealtime ? (FApp::GetCurrentTime() - GStartTime) : 0.0;
diff --git a/Engine/Source/Editor/VirtualTexturingEditor/Private/RuntimeVirtualTextureBuildStreamingMips.cpp b/Engine/Source/Editor/VirtualTexturingEditor/Private/RuntimeVirtualTextureBuildStreamingMips.cpp
--- a/Engine/Source/Editor/VirtualTexturingEditor/Private/RuntimeVirtualTextureBuildStreamingMips.cpp
+++ b/Engine/Source/Editor/VirtualTexturingEditor/Private/RuntimeVirtualTextureBuildStreamingMips.cpp
@@ -74,8 +74,8 @@ namespace
 		EPixelFormat GetLayerFormat(int32 Index) const { return LayerFormats[Index]; }
 		int64 GetLayerOffset(int32 Index) const { return LayerOffsets[Index]; }
 
-		FRHITexture2D* GetRenderTarget(int32 Index) const { return Index < NumLayers ? RenderTargets[Index] : nullptr; }
-		FRHITexture2D* GetStagingTexture(int32 Index) const { return Index < NumLayers ? StagingTextures[Index] : nullptr; }
+		FRHITexture* GetRenderTarget(int32 Index) const { return Index < NumLayers ? RenderTargets[Index] : nullptr; }
+		FRHITexture* GetStagingTexture(int32 Index) const { return Index < NumLayers ? StagingTextures[Index] : nullptr; }
 		FRHIGPUFence* GetFence() const { return Fence; }
 
 	private:
@@ -86,8 +86,8 @@ namespace
 		TArray<EPixelFormat> LayerFormats;
 		TArray<int64> LayerOffsets;
 
-		TArray<FTexture2DRHIRef> RenderTargets;
-		TArray<FTexture2DRHIRef> StagingTextures;
+		TArray<FTextureRHIRef> RenderTargets;
+		TArray<FTextureRHIRef> StagingTextures;
 		FGPUFenceRHIRef Fence;
 	};
 
diff --git a/Engine/Source/Editor/WorldBrowser/Private/Tiles/WorldTileThumbnails.cpp b/Engine/Source/Editor/WorldBrowser/Private/Tiles/WorldTileThumbnails.cpp
--- a/Engine/Source/Editor/WorldBrowser/Private/Tiles/WorldTileThumbnails.cpp
+++ b/Engine/Source/Editor/WorldBrowser/Private/Tiles/WorldTileThumbnails.cpp
@@ -220,7 +220,7 @@ void FTileAtlasPage::UpdateSlotImageData(int32 SlotIdx, FSlateTextureDataPtr Ima
 		ENQUEUE_RENDER_COMMAND(UpdateSlotImageData)(
 			[Context](FRHICommandList& RHICmdList)
 			{
-				FRHITexture2D* RHITexture2D = (FRHITexture2D*)Context.TextureRHI.GetReference();
+				FRHITexture* RHITexture2D = (FRHITexture*)Context.TextureRHI.GetReference();
 				RHIUpdateTexture2D(RHITexture2D, 0, Context.Region, Context.SourcePitch, Context.ImageData->GetRawBytesPtr());
 			});
 	}
diff --git a/Engine/Source/Runtime/Apple/MetalRHI/Private/MetalContext.cpp b/Engine/Source/Runtime/Apple/MetalRHI/Private/MetalContext.cpp
--- a/Engine/Source/Runtime/Apple/MetalRHI/Private/MetalContext.cpp
+++ b/Engine/Source/Runtime/Apple/MetalRHI/Private/MetalContext.cpp
@@ -1203,7 +1203,7 @@ bool FMetalContext::PrepareToDraw(uint32 PrimitiveType)
 		
 		FRHIRenderPassInfo Info = StateCache.GetRenderPassInfo();
 		
-		FTexture2DRHIRef FallbackDepthStencilSurface = StateCache.CreateFallbackDepthStencilSurface(FBSize.width, FBSize.height);
+		FTextureRHIRef FallbackDepthStencilSurface = StateCache.CreateFallbackDepthStencilSurface(FBSize.width, FBSize.height);
 		check(IsValidRef(FallbackDepthStencilSurface));
 		
 		if (bBindDepthStencilForWrite)
@@ -1273,7 +1273,7 @@ void FMetalContext::SetRenderPassInfo(const FRHIRenderPassInfo& RenderTargetsInf
 		{
 			FRHIRenderPassInfo Info = RenderTargetsInfo;
 			CGSize FBSize = CGSizeMake(StateCache.GetViewport(0).width, StateCache.GetViewport(0).height);
-			FTexture2DRHIRef FallbackDepthStencilSurface = StateCache.CreateFallbackDepthStencilSurface(FBSize.width, FBSize.height);
+			FTextureRHIRef FallbackDepthStencilSurface = StateCache.CreateFallbackDepthStencilSurface(FBSize.width, FBSize.height);
 			check(IsValidRef(FallbackDepthStencilSurface));
 
 			Info.DepthStencilRenderTarget.DepthStencilTarget = FallbackDepthStencilSurface;
diff --git a/Engine/Source/Runtime/Apple/MetalRHI/Private/MetalDynamicRHI.h b/Engine/Source/Runtime/Apple/MetalRHI/Private/MetalDynamicRHI.h
--- a/Engine/Source/Runtime/Apple/MetalRHI/Private/MetalDynamicRHI.h
+++ b/Engine/Source/Runtime/Apple/MetalRHI/Private/MetalDynamicRHI.h
@@ -74,20 +74,20 @@ public:
 	virtual void RHIGetTextureMemoryStats(FTextureMemoryStats& OutStats) final override;
 	virtual bool RHIGetTextureMemoryVisualizeData(FColor* TextureData,int32 SizeX,int32 SizeY,int32 Pitch,int32 PixelSize) final override;
 	virtual FTextureRHIRef RHICreateTexture(FRHICommandListBase& RHICmdList, const FRHITextureCreateDesc& CreateDesc) final override;
-	virtual FTexture2DRHIRef RHIAsyncCreateTexture2D(uint32 SizeX, uint32 SizeY, uint8 Format, uint32 NumMips, ETextureCreateFlags Flags, ERHIAccess InResourceState, void** InitialMipData, uint32 NumInitialMips, const TCHAR* DebugName, FGraphEventRef& OutCompletionEvent) final override;
+	virtual FTextureRHIRef RHIAsyncCreateTexture2D(uint32 SizeX, uint32 SizeY, uint8 Format, uint32 NumMips, ETextureCreateFlags Flags, ERHIAccess InResourceState, void** InitialMipData, uint32 NumInitialMips, const TCHAR* DebugName, FGraphEventRef& OutCompletionEvent) final override;
 	virtual void RHIGenerateMips(FRHITexture* Texture) final override;
 	virtual uint32 RHIComputeMemorySize(FRHITexture* TextureRHI) final override;
-	virtual FTexture2DRHIRef RHIAsyncReallocateTexture2D(FRHITexture2D* Texture2D, int32 NewMipCount, int32 NewSizeX, int32 NewSizeY, FThreadSafeCounter* RequestStatus) final override;
-	virtual ETextureReallocationStatus RHIFinalizeAsyncReallocateTexture2D(FRHITexture2D* Texture2D, bool bBlockUntilCompleted) final override;
-	virtual ETextureReallocationStatus RHICancelAsyncReallocateTexture2D(FRHITexture2D* Texture2D, bool bBlockUntilCompleted) final override;
-	virtual void* RHILockTexture2D(FRHITexture2D* Texture, uint32 MipIndex, EResourceLockMode LockMode, uint32& DestStride, bool bLockWithinMiptail, uint64* OutLockedByteCount = nullptr) final override;
-	virtual void RHIUnlockTexture2D(FRHITexture2D* Texture, uint32 MipIndex, bool bLockWithinMiptail) final override;
-	virtual void* RHILockTexture2DArray(FRHITexture2DArray* Texture, uint32 TextureIndex, uint32 MipIndex, EResourceLockMode LockMode, uint32& DestStride, bool bLockWithinMiptail) final override;
-	virtual void RHIUnlockTexture2DArray(FRHITexture2DArray* Texture, uint32 TextureIndex, uint32 MipIndex, bool bLockWithinMiptail) final override;
-	virtual void RHIUpdateTexture2D(FRHICommandListBase& RHICmdList, FRHITexture2D* Texture, uint32 MipIndex, const struct FUpdateTextureRegion2D& UpdateRegion, uint32 SourcePitch, const uint8* SourceData) final override;
-	virtual void RHIUpdateTexture3D(FRHICommandListBase& RHICmdList, FRHITexture3D* Texture, uint32 MipIndex, const struct FUpdateTextureRegion3D& UpdateRegion, uint32 SourceRowPitch, uint32 SourceDepthPitch, const uint8* SourceData) final override;
-	virtual void* RHILockTextureCubeFace(FRHITextureCube* Texture, uint32 FaceIndex, uint32 ArrayIndex, uint32 MipIndex, EResourceLockMode LockMode, uint32& DestStride, bool bLockWithinMiptail) final override;
-	virtual void RHIUnlockTextureCubeFace(FRHITextureCube* Texture, uint32 FaceIndex, uint32 ArrayIndex, uint32 MipIndex, bool bLockWithinMiptail) final override;
+	virtual FTextureRHIRef RHIAsyncReallocateTexture2D(FRHITexture* Texture2D, int32 NewMipCount, int32 NewSizeX, int32 NewSizeY, FThreadSafeCounter* RequestStatus) final override;
+	virtual ETextureReallocationStatus RHIFinalizeAsyncReallocateTexture2D(FRHITexture* Texture2D, bool bBlockUntilCompleted) final override;
+	virtual ETextureReallocationStatus RHICancelAsyncReallocateTexture2D(FRHITexture* Texture2D, bool bBlockUntilCompleted) final override;
+	virtual void* RHILockTexture2D(FRHITexture* Texture, uint32 MipIndex, EResourceLockMode LockMode, uint32& DestStride, bool bLockWithinMiptail, uint64* OutLockedByteCount = nullptr) final override;
+	virtual void RHIUnlockTexture2D(FRHITexture* Texture, uint32 MipIndex, bool bLockWithinMiptail) final override;
+	virtual void* RHILockTexture2DArray(FRHITexture* Texture, uint32 TextureIndex, uint32 MipIndex, EResourceLockMode LockMode, uint32& DestStride, bool bLockWithinMiptail) final override;
+	virtual void RHIUnlockTexture2DArray(FRHITexture* Texture, uint32 TextureIndex, uint32 MipIndex, bool bLockWithinMiptail) final override;
+	virtual void RHIUpdateTexture2D(FRHICommandListBase& RHICmdList, FRHITexture* Texture, uint32 MipIndex, const struct FUpdateTextureRegion2D& UpdateRegion, uint32 SourcePitch, const uint8* SourceData) final override;
+	virtual void RHIUpdateTexture3D(FRHICommandListBase& RHICmdList, FRHITexture* Texture, uint32 MipIndex, const struct FUpdateTextureRegion3D& UpdateRegion, uint32 SourceRowPitch, uint32 SourceDepthPitch, const uint8* SourceData) final override;
+	virtual void* RHILockTextureCubeFace(FRHITexture* Texture, uint32 FaceIndex, uint32 ArrayIndex, uint32 MipIndex, EResourceLockMode LockMode, uint32& DestStride, bool bLockWithinMiptail) final override;
+	virtual void RHIUnlockTextureCubeFace(FRHITexture* Texture, uint32 FaceIndex, uint32 ArrayIndex, uint32 MipIndex, bool bLockWithinMiptail) final override;
 	virtual void RHIBindDebugLabelName(FRHICommandListBase& RHICmdList, FRHITexture* Texture, const TCHAR* Name) final override;
 	virtual void RHIReadSurfaceData(FRHITexture* Texture,FIntRect Rect,TArray<FColor>& OutData,FReadSurfaceDataFlags InFlags) final override;
 	virtual void RHIReadSurfaceData(FRHITexture* TextureRHI, FIntRect InRect, TArray<FLinearColor>& OutData, FReadSurfaceDataFlags InFlags) final override;
@@ -97,7 +97,7 @@ public:
 	virtual void RHIRead3DSurfaceFloatData(FRHITexture* Texture,FIntRect Rect,FIntPoint ZMinMax,TArray<FFloat16Color>& OutData) final override;
 	virtual FRenderQueryRHIRef RHICreateRenderQuery(ERenderQueryType QueryType) final override;
 	virtual bool RHIGetRenderQueryResult(FRHIRenderQuery* RenderQuery, uint64& OutResult, bool bWait, uint32 GPUIndex = INDEX_NONE) final override;
-	virtual FTexture2DRHIRef RHIGetViewportBackBuffer(FRHIViewport* Viewport) final override;
+	virtual FTextureRHIRef RHIGetViewportBackBuffer(FRHIViewport* Viewport) final override;
 	virtual void RHIAdvanceFrameForGetViewportBackBuffer(FRHIViewport* Viewport) final override;
 	virtual void RHIFlushResources() final override;
 	virtual uint32 RHIGetGPUFrameCycles(uint32 GPUIndex = 0) final override;
@@ -118,14 +118,14 @@ public:
 	virtual IRHIPlatformCommandList* RHIFinalizeContext(FRHIFinalizeContextArgs&& Args) final override;
 	virtual void RHISubmitCommandLists(FRHISubmitCommandListsArgs&& Args) final override;
 
-	virtual FTexture2DRHIRef AsyncReallocateTexture2D_RenderThread(class FRHICommandListImmediate& RHICmdList, FRHITexture2D* Texture2D, int32 NewMipCount, int32 NewSizeX, int32 NewSizeY, FThreadSafeCounter* RequestStatus) final override;
-	virtual ETextureReallocationStatus FinalizeAsyncReallocateTexture2D_RenderThread(class FRHICommandListImmediate& RHICmdList, FRHITexture2D* Texture2D, bool bBlockUntilCompleted) final override;
-	virtual ETextureReallocationStatus CancelAsyncReallocateTexture2D_RenderThread(class FRHICommandListImmediate& RHICmdList, FRHITexture2D* Texture2D, bool bBlockUntilCompleted) final override;
+	virtual FTextureRHIRef AsyncReallocateTexture2D_RenderThread(class FRHICommandListImmediate& RHICmdList, FRHITexture* Texture2D, int32 NewMipCount, int32 NewSizeX, int32 NewSizeY, FThreadSafeCounter* RequestStatus) final override;
+	virtual ETextureReallocationStatus FinalizeAsyncReallocateTexture2D_RenderThread(class FRHICommandListImmediate& RHICmdList, FRHITexture* Texture2D, bool bBlockUntilCompleted) final override;
+	virtual ETextureReallocationStatus CancelAsyncReallocateTexture2D_RenderThread(class FRHICommandListImmediate& RHICmdList, FRHITexture* Texture2D, bool bBlockUntilCompleted) final override;
 
-	virtual void* LockTexture2D_RenderThread(class FRHICommandListImmediate& RHICmdList, FRHITexture2D* Texture, uint32 MipIndex, EResourceLockMode LockMode, uint32& DestStride, bool bLockWithinMiptail, bool bNeedsDefaultRHIFlush = true, uint64* OutLockedByteCount = nullptr) final override;
-	virtual void UnlockTexture2D_RenderThread(class FRHICommandListImmediate& RHICmdList, FRHITexture2D* Texture, uint32 MipIndex, bool bLockWithinMiptail, bool bNeedsDefaultRHIFlush = true) final override;
+	virtual void* LockTexture2D_RenderThread(class FRHICommandListImmediate& RHICmdList, FRHITexture* Texture, uint32 MipIndex, EResourceLockMode LockMode, uint32& DestStride, bool bLockWithinMiptail, bool bNeedsDefaultRHIFlush = true, uint64* OutLockedByteCount = nullptr) final override;
+	virtual void UnlockTexture2D_RenderThread(class FRHICommandListImmediate& RHICmdList, FRHITexture* Texture, uint32 MipIndex, bool bLockWithinMiptail, bool bNeedsDefaultRHIFlush = true) final override;
 
-	virtual FUpdateTexture3DData RHIBeginUpdateTexture3D(FRHICommandListBase& RHICmdList, FRHITexture3D* Texture, uint32 MipIndex, const struct FUpdateTextureRegion3D& UpdateRegion) final override;
+	virtual FUpdateTexture3DData RHIBeginUpdateTexture3D(FRHICommandListBase& RHICmdList, FRHITexture* Texture, uint32 MipIndex, const struct FUpdateTextureRegion3D& UpdateRegion) final override;
 	virtual void RHIEndUpdateTexture3D(FRHICommandListBase& RHICmdList, FUpdateTexture3DData& UpdateData) final override;
 
 	virtual void RHICreateTransition(FRHITransition* Transition, const FRHITransitionCreateInfo& CreateInfo) final override;
diff --git a/Engine/Source/Runtime/Apple/MetalRHI/Private/MetalStateCache.cpp b/Engine/Source/Runtime/Apple/MetalRHI/Private/MetalStateCache.cpp
--- a/Engine/Source/Runtime/Apple/MetalRHI/Private/MetalStateCache.cpp
+++ b/Engine/Source/Runtime/Apple/MetalRHI/Private/MetalStateCache.cpp
@@ -2865,7 +2865,7 @@ void FMetalStateCache::CommitResourceTable(EMetalShaderStages const Frequency, M
 	}
 }
 
-FTexture2DRHIRef FMetalStateCache::CreateFallbackDepthStencilSurface(uint32 Width, uint32 Height)
+FTextureRHIRef FMetalStateCache::CreateFallbackDepthStencilSurface(uint32 Width, uint32 Height)
 {
 #if PLATFORM_MAC
 	if (!IsValidRef(FallbackDepthStencilSurface) || FallbackDepthStencilSurface->GetSizeX() < Width || FallbackDepthStencilSurface->GetSizeY() < Height)
diff --git a/Engine/Source/Runtime/Apple/MetalRHI/Private/MetalStateCache.h b/Engine/Source/Runtime/Apple/MetalRHI/Private/MetalStateCache.h
--- a/Engine/Source/Runtime/Apple/MetalRHI/Private/MetalStateCache.h
+++ b/Engine/Source/Runtime/Apple/MetalRHI/Private/MetalStateCache.h
@@ -194,7 +194,7 @@ public:
 	MTL::VisibilityResultMode GetVisibilityResultMode() { return VisibilityMode; }
 	uint32 GetVisibilityResultOffset() { return VisibilityOffset; }
 	
-	FTexture2DRHIRef CreateFallbackDepthStencilSurface(uint32 Width, uint32 Height);
+	FTextureRHIRef CreateFallbackDepthStencilSurface(uint32 Width, uint32 Height);
 	bool GetFallbackDepthStencilBound(void) const { return bFallbackDepthStencilBound; }
 	
 	void SetRenderPipelineState(FMetalCommandEncoder& CommandEncoder);
@@ -357,7 +357,7 @@ private:
 	FTextureRHIRef DepthStencilSurface;
 	FTextureRHIRef DepthStencilResolve;
 	/** A fallback depth-stencil surface for draw calls that write to depth without a depth-stencil surface bound. */
-	FTexture2DRHIRef FallbackDepthStencilSurface;
+	FTextureRHIRef FallbackDepthStencilSurface;
     MTL::RenderPassDescriptor* RenderPassDesc;
 	uint32 RasterBits;
     uint8 PipelineBits;
diff --git a/Engine/Source/Runtime/Apple/MetalRHI/Private/MetalTexture.cpp b/Engine/Source/Runtime/Apple/MetalRHI/Private/MetalTexture.cpp
--- a/Engine/Source/Runtime/Apple/MetalRHI/Private/MetalTexture.cpp
+++ b/Engine/Source/Runtime/Apple/MetalRHI/Private/MetalTexture.cpp
@@ -1476,10 +1476,10 @@ FTextureRHIRef FMetalDynamicRHI::RHICreateTexture(FRHICommandListBase& RHICmdLis
     return new FMetalSurface(&RHICmdList, CreateDesc);
 }
 
-FTexture2DRHIRef FMetalDynamicRHI::RHIAsyncCreateTexture2D(uint32 SizeX, uint32 SizeY, uint8 Format, uint32 NumMips, ETextureCreateFlags Flags, ERHIAccess InResourceState, void** InitialMipData, uint32 NumInitialMips, const TCHAR* DebugName, FGraphEventRef& OutCompletionEvent)
+FTextureRHIRef FMetalDynamicRHI::RHIAsyncCreateTexture2D(uint32 SizeX, uint32 SizeY, uint8 Format, uint32 NumMips, ETextureCreateFlags Flags, ERHIAccess InResourceState, void** InitialMipData, uint32 NumInitialMips, const TCHAR* DebugName, FGraphEventRef& OutCompletionEvent)
 {
 	UE_LOG(LogMetal, Fatal, TEXT("RHIAsyncCreateTexture2D is not supported"));
-	return FTexture2DRHIRef();
+	return FTextureRHIRef();
 }
 
 void FMetalDynamicRHI::RHIGenerateMips(FRHITexture* SourceSurfaceRHI)
@@ -1492,24 +1492,24 @@ void FMetalDynamicRHI::RHIGenerateMips(FRHITexture* SourceSurfaceRHI)
     }
 }
 
-FTexture2DRHIRef FMetalDynamicRHI::AsyncReallocateTexture2D_RenderThread(class FRHICommandListImmediate& RHICmdList, FRHITexture2D* Texture2D, int32 NewMipCount, int32 NewSizeX, int32 NewSizeY, FThreadSafeCounter* RequestStatus)
+FTextureRHIRef FMetalDynamicRHI::AsyncReallocateTexture2D_RenderThread(class FRHICommandListImmediate& RHICmdList, FRHITexture* Texture2D, int32 NewMipCount, int32 NewSizeX, int32 NewSizeY, FThreadSafeCounter* RequestStatus)
 {
 	return this->RHIAsyncReallocateTexture2D(Texture2D, NewMipCount, NewSizeX, NewSizeY, RequestStatus);
 }
 
-ETextureReallocationStatus FMetalDynamicRHI::FinalizeAsyncReallocateTexture2D_RenderThread(class FRHICommandListImmediate& RHICmdList, FRHITexture2D* Texture2D, bool bBlockUntilCompleted)
+ETextureReallocationStatus FMetalDynamicRHI::FinalizeAsyncReallocateTexture2D_RenderThread(class FRHICommandListImmediate& RHICmdList, FRHITexture* Texture2D, bool bBlockUntilCompleted)
 {
 	// No need to flush - does nothing
 	return this->RHIFinalizeAsyncReallocateTexture2D(Texture2D, bBlockUntilCompleted);
 }
 
-ETextureReallocationStatus FMetalDynamicRHI::CancelAsyncReallocateTexture2D_RenderThread(class FRHICommandListImmediate& RHICmdList, FRHITexture2D* Texture2D, bool bBlockUntilCompleted)
+ETextureReallocationStatus FMetalDynamicRHI::CancelAsyncReallocateTexture2D_RenderThread(class FRHICommandListImmediate& RHICmdList, FRHITexture* Texture2D, bool bBlockUntilCompleted)
 {
 	// No need to flush - does nothing
 	return this->RHICancelAsyncReallocateTexture2D(Texture2D, bBlockUntilCompleted);
 }
 
-FTexture2DRHIRef FMetalDynamicRHI::RHIAsyncReallocateTexture2D(FRHITexture2D* OldTextureRHI, int32 NewMipCount, int32 NewSizeX, int32 NewSizeY, FThreadSafeCounter* RequestStatus)
+FTextureRHIRef FMetalDynamicRHI::RHIAsyncReallocateTexture2D(FRHITexture* OldTextureRHI, int32 NewMipCount, int32 NewSizeX, int32 NewSizeY, FThreadSafeCounter* RequestStatus)
 {
     MTL_SCOPED_AUTORELEASE_POOL;
 
@@ -1583,17 +1583,17 @@ FTexture2DRHIRef FMetalDynamicRHI::RHIAsyncReallocateTexture2D(FRHITexture2D* Ol
     return NewTexture;
 }
 
-ETextureReallocationStatus FMetalDynamicRHI::RHIFinalizeAsyncReallocateTexture2D(FRHITexture2D* Texture2D, bool bBlockUntilCompleted )
+ETextureReallocationStatus FMetalDynamicRHI::RHIFinalizeAsyncReallocateTexture2D(FRHITexture* Texture2D, bool bBlockUntilCompleted )
 {
 	return TexRealloc_Succeeded;
 }
 
-ETextureReallocationStatus FMetalDynamicRHI::RHICancelAsyncReallocateTexture2D(FRHITexture2D* Texture2D, bool bBlockUntilCompleted )
+ETextureReallocationStatus FMetalDynamicRHI::RHICancelAsyncReallocateTexture2D(FRHITexture* Texture2D, bool bBlockUntilCompleted )
 {
 	return TexRealloc_Failed;
 }
 
-void* FMetalDynamicRHI::LockTexture2D_RenderThread(class FRHICommandListImmediate& RHICmdList, FRHITexture2D* Texture, uint32 MipIndex, EResourceLockMode LockMode, uint32& DestStride, bool bLockWithinMiptail, bool bNeedsDefaultRHIFlush, uint64* OutLockedByteCount)
+void* FMetalDynamicRHI::LockTexture2D_RenderThread(class FRHICommandListImmediate& RHICmdList, FRHITexture* Texture, uint32 MipIndex, EResourceLockMode LockMode, uint32& DestStride, bool bLockWithinMiptail, bool bNeedsDefaultRHIFlush, uint64* OutLockedByteCount)
 {
     MTL_SCOPED_AUTORELEASE_POOL;
     
@@ -1604,7 +1604,7 @@ void* FMetalDynamicRHI::LockTexture2D_RenderThread(class FRHICommandListImmediat
     return BufferData;
 }
 
-void FMetalDynamicRHI::UnlockTexture2D_RenderThread(class FRHICommandListImmediate& RHICmdList, FRHITexture2D* Texture, uint32 MipIndex, bool bLockWithinMiptail, bool bNeedsDefaultRHIFlush)
+void FMetalDynamicRHI::UnlockTexture2D_RenderThread(class FRHICommandListImmediate& RHICmdList, FRHITexture* Texture, uint32 MipIndex, bool bLockWithinMiptail, bool bNeedsDefaultRHIFlush)
 {
     MTL_SCOPED_AUTORELEASE_POOL;
 
@@ -1615,28 +1615,28 @@ void FMetalDynamicRHI::UnlockTexture2D_RenderThread(class FRHICommandListImmedia
 }
 
 
-void* FMetalDynamicRHI::RHILockTexture2D(FRHITexture2D* TextureRHI,uint32 MipIndex,EResourceLockMode LockMode,uint32& DestStride,bool bLockWithinMiptail, uint64* OutLockedByteCount)
+void* FMetalDynamicRHI::RHILockTexture2D(FRHITexture* TextureRHI,uint32 MipIndex,EResourceLockMode LockMode,uint32& DestStride,bool bLockWithinMiptail, uint64* OutLockedByteCount)
 {
     MTL_SCOPED_AUTORELEASE_POOL;
     FMetalSurface* Texture = ResourceCast(TextureRHI);
 	return Texture->Lock(MipIndex, 0, LockMode, DestStride, OutLockedByteCount);
 }
 
-void FMetalDynamicRHI::RHIUnlockTexture2D(FRHITexture2D* TextureRHI,uint32 MipIndex,bool bLockWithinMiptail)
+void FMetalDynamicRHI::RHIUnlockTexture2D(FRHITexture* TextureRHI,uint32 MipIndex,bool bLockWithinMiptail)
 {
     MTL_SCOPED_AUTORELEASE_POOL;
     FMetalSurface* Texture = ResourceCast(TextureRHI);
 	Texture->Unlock(MipIndex, 0, false);
 }
 
-void* FMetalDynamicRHI::RHILockTexture2DArray(FRHITexture2DArray* TextureRHI,uint32 TextureIndex,uint32 MipIndex,EResourceLockMode LockMode,uint32& DestStride,bool bLockWithinMiptail)
+void* FMetalDynamicRHI::RHILockTexture2DArray(FRHITexture* TextureRHI,uint32 TextureIndex,uint32 MipIndex,EResourceLockMode LockMode,uint32& DestStride,bool bLockWithinMiptail)
 {
     MTL_SCOPED_AUTORELEASE_POOL;
     FMetalSurface* Texture = ResourceCast(TextureRHI);
 	return Texture->Lock(MipIndex, TextureIndex, LockMode, DestStride);
 }
 
-void FMetalDynamicRHI::RHIUnlockTexture2DArray(FRHITexture2DArray* TextureRHI,uint32 TextureIndex,uint32 MipIndex,bool bLockWithinMiptail)
+void FMetalDynamicRHI::RHIUnlockTexture2DArray(FRHITexture* TextureRHI,uint32 TextureIndex,uint32 MipIndex,bool bLockWithinMiptail)
 {
     MTL_SCOPED_AUTORELEASE_POOL;
     FMetalSurface* Texture = ResourceCast(TextureRHI);
@@ -1662,7 +1662,7 @@ static void InternalExpandR8ToStandardRGBA(uint32* pDest, const struct FUpdateTe
 }
 #endif
 
-static FMetalBufferPtr Internal_CreateBufferAndCopyTexture2DUpdateRegionData(FRHITexture2D* TextureRHI, const struct FUpdateTextureRegion2D& UpdateRegion, uint32& InOutSourcePitch, const uint8* SourceData)
+static FMetalBufferPtr Internal_CreateBufferAndCopyTexture2DUpdateRegionData(FRHITexture* TextureRHI, const struct FUpdateTextureRegion2D& UpdateRegion, uint32& InOutSourcePitch, const uint8* SourceData)
 {
 	const EPixelFormat PixelFormat = TextureRHI->GetFormat();
 	const FPixelFormatInfo& FormatInfo = GPixelFormats[PixelFormat];
@@ -1722,7 +1722,7 @@ static FMetalBufferPtr Internal_CreateBufferAndCopyTexture2DUpdateRegionData(FRH
 	return OutBuffer;
 }
 
-static void InternalUpdateTexture2D(FMetalContext& Context, FRHITexture2D* TextureRHI, uint32 MipIndex, FUpdateTextureRegion2D const& UpdateRegion, uint32 SourcePitch, FMetalBufferPtr Buffer)
+static void InternalUpdateTexture2D(FMetalContext& Context, FRHITexture* TextureRHI, uint32 MipIndex, FUpdateTextureRegion2D const& UpdateRegion, uint32 SourcePitch, FMetalBufferPtr Buffer)
 {
 	FMetalSurface* Texture = ResourceCast(TextureRHI);
 	MTLTexturePtr Tex = Texture->Texture;
@@ -1759,7 +1759,7 @@ static void InternalUpdateTexture2D(FMetalContext& Context, FRHITexture2D* Textu
 	FPlatformAtomics::InterlockedExchange(&Texture->Written, 1);
 }
 
-void FMetalDynamicRHI::RHIUpdateTexture2D(FRHICommandListBase& RHICmdList, FRHITexture2D* TextureRHI, uint32 MipIndex, const struct FUpdateTextureRegion2D& UpdateRegion, uint32 SourcePitch, const uint8* SourceData)
+void FMetalDynamicRHI::RHIUpdateTexture2D(FRHICommandListBase& RHICmdList, FRHITexture* TextureRHI, uint32 MipIndex, const struct FUpdateTextureRegion2D& UpdateRegion, uint32 SourcePitch, const uint8* SourceData)
 {
     MTL_SCOPED_AUTORELEASE_POOL;
 
@@ -1774,7 +1774,7 @@ void FMetalDynamicRHI::RHIUpdateTexture2D(FRHICommandListBase& RHICmdList, FRHIT
     INC_DWORD_STAT_BY(STAT_MetalTextureMemUpdate, UpdateRegion.Height*SourcePitch);
 }
 
-static FMetalBufferPtr Internal_CreateBufferAndCopyTexture3DUpdateRegionData(FRHITexture3D* TextureRHI, const struct FUpdateTextureRegion3D& UpdateRegion, uint32 SourceRowPitch, uint32 SourceDepthPitch, const uint8* SourceData)
+static FMetalBufferPtr Internal_CreateBufferAndCopyTexture3DUpdateRegionData(FRHITexture* TextureRHI, const struct FUpdateTextureRegion3D& UpdateRegion, uint32 SourceRowPitch, uint32 SourceDepthPitch, const uint8* SourceData)
 {
 	FMetalSurface* Texture = ResourceCast(TextureRHI);
 	
@@ -1808,7 +1808,7 @@ static FMetalBufferPtr Internal_CreateBufferAndCopyTexture3DUpdateRegionData(FRH
 }
 
 
-static void InternalUpdateTexture3D(FMetalContext& Context, FRHITexture3D* TextureRHI, uint32 MipIndex, const FUpdateTextureRegion3D& UpdateRegion, uint32 SourceRowPitch, uint32 SourceDepthPitch, FMetalBufferPtr Buffer)
+static void InternalUpdateTexture3D(FMetalContext& Context, FRHITexture* TextureRHI, uint32 MipIndex, const FUpdateTextureRegion3D& UpdateRegion, uint32 SourceRowPitch, uint32 SourceDepthPitch, FMetalBufferPtr Buffer)
 {
 	FMetalSurface* Texture = ResourceCast(TextureRHI);
 	MTLTexturePtr Tex = Texture->Texture;
@@ -1842,7 +1842,7 @@ static void InternalUpdateTexture3D(FMetalContext& Context, FRHITexture3D* Textu
 	FPlatformAtomics::InterlockedExchange(&Texture->Written, 1);
 }
 
-FUpdateTexture3DData FMetalDynamicRHI::RHIBeginUpdateTexture3D(FRHICommandListBase& RHICmdList, FRHITexture3D* Texture, uint32 MipIndex, const struct FUpdateTextureRegion3D& UpdateRegion)
+FUpdateTexture3DData FMetalDynamicRHI::RHIBeginUpdateTexture3D(FRHICommandListBase& RHICmdList, FRHITexture* Texture, uint32 MipIndex, const struct FUpdateTextureRegion3D& UpdateRegion)
 {
 	const int32 FormatSize = PixelFormatBlockBytes[Texture->GetFormat()];
 	const int32 RowPitch = UpdateRegion.Width * FormatSize;
@@ -1863,7 +1863,7 @@ void FMetalDynamicRHI::RHIEndUpdateTexture3D(FRHICommandListBase& RHICmdList, FU
 	UpdateData.Data = nullptr;
 }
 
-void FMetalDynamicRHI::RHIUpdateTexture3D(FRHICommandListBase& RHICmdList, FRHITexture3D* TextureRHI,uint32 MipIndex,const FUpdateTextureRegion3D& UpdateRegion,uint32 SourceRowPitch,uint32 SourceDepthPitch, const uint8* SourceData)
+void FMetalDynamicRHI::RHIUpdateTexture3D(FRHICommandListBase& RHICmdList, FRHITexture* TextureRHI,uint32 MipIndex,const FUpdateTextureRegion3D& UpdateRegion,uint32 SourceRowPitch,uint32 SourceDepthPitch, const uint8* SourceData)
 {
     MTL_SCOPED_AUTORELEASE_POOL;
     
@@ -1881,7 +1881,7 @@ void FMetalDynamicRHI::RHIUpdateTexture3D(FRHICommandListBase& RHICmdList, FRHIT
 /*-----------------------------------------------------------------------------
  Cubemap texture support.
  -----------------------------------------------------------------------------*/
-void* FMetalDynamicRHI::RHILockTextureCubeFace(FRHITextureCube* TextureCubeRHI,uint32 FaceIndex,uint32 ArrayIndex,uint32 MipIndex,EResourceLockMode LockMode,uint32& DestStride,bool bLockWithinMiptail)
+void* FMetalDynamicRHI::RHILockTextureCubeFace(FRHITexture* TextureCubeRHI,uint32 FaceIndex,uint32 ArrayIndex,uint32 MipIndex,EResourceLockMode LockMode,uint32& DestStride,bool bLockWithinMiptail)
 {
     MTL_SCOPED_AUTORELEASE_POOL;
     
@@ -1890,7 +1890,7 @@ void* FMetalDynamicRHI::RHILockTextureCubeFace(FRHITextureCube* TextureCubeRHI,u
     return TextureCube->Lock(MipIndex, MetalFace + (6 * ArrayIndex), LockMode, DestStride, true);
 }
 
-void FMetalDynamicRHI::RHIUnlockTextureCubeFace(FRHITextureCube* TextureCubeRHI,uint32 FaceIndex,uint32 ArrayIndex,uint32 MipIndex,bool bLockWithinMiptail)
+void FMetalDynamicRHI::RHIUnlockTextureCubeFace(FRHITexture* TextureCubeRHI,uint32 FaceIndex,uint32 ArrayIndex,uint32 MipIndex,bool bLockWithinMiptail)
 {
     MTL_SCOPED_AUTORELEASE_POOL;
     FMetalSurface* TextureCube = ResourceCast(TextureCubeRHI);
diff --git a/Engine/Source/Runtime/Apple/MetalRHI/Private/MetalViewport.cpp b/Engine/Source/Runtime/Apple/MetalRHI/Private/MetalViewport.cpp
--- a/Engine/Source/Runtime/Apple/MetalRHI/Private/MetalViewport.cpp
+++ b/Engine/Source/Runtime/Apple/MetalRHI/Private/MetalViewport.cpp
@@ -846,12 +846,12 @@ void FMetalRHIImmediateCommandContext::RHIEndDrawingViewport(FRHIViewport* Viewp
 	((FMetalDeviceContext*)Context)->EndDrawingViewport(Viewport, bPresent, bLockToVsync);
 }
 
-FTexture2DRHIRef FMetalDynamicRHI::RHIGetViewportBackBuffer(FRHIViewport* ViewportRHI)
+FTextureRHIRef FMetalDynamicRHI::RHIGetViewportBackBuffer(FRHIViewport* ViewportRHI)
 {
     MTL_SCOPED_AUTORELEASE_POOL;
     
 	FMetalViewport* Viewport = ResourceCast(ViewportRHI);
-	return FTexture2DRHIRef(Viewport->GetBackBuffer(EMetalViewportAccessRenderer).GetReference());
+	return FTextureRHIRef(Viewport->GetBackBuffer(EMetalViewportAccessRenderer).GetReference());
 }
 
 void FMetalDynamicRHI::RHIAdvanceFrameForGetViewportBackBuffer(FRHIViewport* ViewportRHI)
diff --git a/Engine/Source/Runtime/D3D12RHI/Private/D3D12RHIPrivate.h b/Engine/Source/Runtime/D3D12RHI/Private/D3D12RHIPrivate.h
--- a/Engine/Source/Runtime/D3D12RHI/Private/D3D12RHIPrivate.h
+++ b/Engine/Source/Runtime/D3D12RHI/Private/D3D12RHIPrivate.h
@@ -327,22 +327,22 @@ public:
 	virtual FTextureRHIRef RHIAsyncCreateTexture2D(uint32 SizeX, uint32 SizeY, uint8 Format, uint32 NumMips, ETextureCreateFlags Flags, ERHIAccess InResourceState, void** InitialMipData, uint32 NumInitialMips, const TCHAR* DebugName, FGraphEventRef& OutCompletionEvent) final override;
 	virtual FTextureRHIRef RHICreateTexture(FRHICommandListBase& RHICmdList, const FRHITextureCreateDesc& CreateDesc) override;
 	virtual uint32 RHIComputeMemorySize(FRHITexture* TextureRHI) final override;
-	virtual FTexture2DRHIRef RHIAsyncReallocateTexture2D(FRHITexture2D* Texture2D, int32 NewMipCount, int32 NewSizeX, int32 NewSizeY, FThreadSafeCounter* RequestStatus) override;
-	virtual ETextureReallocationStatus RHIFinalizeAsyncReallocateTexture2D(FRHITexture2D* Texture2D, bool bBlockUntilCompleted) final override;
-	virtual ETextureReallocationStatus RHICancelAsyncReallocateTexture2D(FRHITexture2D* Texture2D, bool bBlockUntilCompleted) final override;
-	virtual void* RHILockTexture2D(FRHITexture2D* Texture, uint32 MipIndex, EResourceLockMode LockMode, uint32& DestStride, bool bLockWithinMiptail, uint64* OutLockedByteCount = nullptr) final override;
-	virtual void RHIUnlockTexture2D(FRHITexture2D* Texture, uint32 MipIndex, bool bLockWithinMiptail) final override;
-	virtual void* RHILockTexture2DArray(FRHITexture2DArray* Texture, uint32 TextureIndex, uint32 MipIndex, EResourceLockMode LockMode, uint32& DestStride, bool bLockWithinMiptail) final override;
-	virtual void RHIUnlockTexture2DArray(FRHITexture2DArray* Texture, uint32 TextureIndex, uint32 MipIndex, bool bLockWithinMiptail) final override;
-	virtual void* RHILockTextureCubeFace_RenderThread(class FRHICommandListImmediate& RHICmdList, FRHITextureCube* TextureRHI, uint32 FaceIndex, uint32 ArrayIndex, uint32 MipIndex, EResourceLockMode LockMode, uint32& DestStride, bool bLockWithinMiptail) final override;
-	virtual void RHIUnlockTextureCubeFace_RenderThread(class FRHICommandListImmediate& RHICmdList, FRHITextureCube* TextureRHI, uint32 FaceIndex, uint32 ArrayIndex, uint32 MipIndex, bool bLockWithinMiptail) final override;
-	virtual void RHIUpdateTexture2D(FRHICommandListBase& RHICmdList, FRHITexture2D* Texture, uint32 MipIndex, const struct FUpdateTextureRegion2D& UpdateRegion, uint32 SourcePitch, const uint8* SourceData) final override;
-	virtual void RHIUpdateTexture3D(FRHICommandListBase& RHICmdList, FRHITexture3D* Texture, uint32 MipIndex, const struct FUpdateTextureRegion3D& UpdateRegion, uint32 SourceRowPitch, uint32 SourceDepthPitch, const uint8* SourceData) final override;
-	virtual FUpdateTexture3DData RHIBeginUpdateTexture3D(FRHICommandListBase& RHICmdList, FRHITexture3D* Texture, uint32 MipIndex, const struct FUpdateTextureRegion3D& UpdateRegion) final override;
+	virtual FTextureRHIRef RHIAsyncReallocateTexture2D(FRHITexture* Texture2D, int32 NewMipCount, int32 NewSizeX, int32 NewSizeY, FThreadSafeCounter* RequestStatus) override;
+	virtual ETextureReallocationStatus RHIFinalizeAsyncReallocateTexture2D(FRHITexture* Texture2D, bool bBlockUntilCompleted) final override;
+	virtual ETextureReallocationStatus RHICancelAsyncReallocateTexture2D(FRHITexture* Texture2D, bool bBlockUntilCompleted) final override;
+	virtual void* RHILockTexture2D(FRHITexture* Texture, uint32 MipIndex, EResourceLockMode LockMode, uint32& DestStride, bool bLockWithinMiptail, uint64* OutLockedByteCount = nullptr) final override;
+	virtual void RHIUnlockTexture2D(FRHITexture* Texture, uint32 MipIndex, bool bLockWithinMiptail) final override;
+	virtual void* RHILockTexture2DArray(FRHITexture* Texture, uint32 TextureIndex, uint32 MipIndex, EResourceLockMode LockMode, uint32& DestStride, bool bLockWithinMiptail) final override;
+	virtual void RHIUnlockTexture2DArray(FRHITexture* Texture, uint32 TextureIndex, uint32 MipIndex, bool bLockWithinMiptail) final override;
+	virtual void* RHILockTextureCubeFace_RenderThread(class FRHICommandListImmediate& RHICmdList, FRHITexture* TextureRHI, uint32 FaceIndex, uint32 ArrayIndex, uint32 MipIndex, EResourceLockMode LockMode, uint32& DestStride, bool bLockWithinMiptail) final override;
+	virtual void RHIUnlockTextureCubeFace_RenderThread(class FRHICommandListImmediate& RHICmdList, FRHITexture* TextureRHI, uint32 FaceIndex, uint32 ArrayIndex, uint32 MipIndex, bool bLockWithinMiptail) final override;
+	virtual void RHIUpdateTexture2D(FRHICommandListBase& RHICmdList, FRHITexture* Texture, uint32 MipIndex, const struct FUpdateTextureRegion2D& UpdateRegion, uint32 SourcePitch, const uint8* SourceData) final override;
+	virtual void RHIUpdateTexture3D(FRHICommandListBase& RHICmdList, FRHITexture* Texture, uint32 MipIndex, const struct FUpdateTextureRegion3D& UpdateRegion, uint32 SourceRowPitch, uint32 SourceDepthPitch, const uint8* SourceData) final override;
+	virtual FUpdateTexture3DData RHIBeginUpdateTexture3D(FRHICommandListBase& RHICmdList, FRHITexture* Texture, uint32 MipIndex, const struct FUpdateTextureRegion3D& UpdateRegion) final override;
 	virtual void RHIEndUpdateTexture3D(FRHICommandListBase& RHICmdList, FUpdateTexture3DData& UpdateData) final override;
 	virtual void RHIEndMultiUpdateTexture3D(FRHICommandListBase& RHICmdList, TArray<FUpdateTexture3DData>& UpdateDataArray) final override;
-	virtual void* RHILockTextureCubeFace(FRHITextureCube* Texture, uint32 FaceIndex, uint32 ArrayIndex, uint32 MipIndex, EResourceLockMode LockMode, uint32& DestStride, bool bLockWithinMiptail) final override;
-	virtual void RHIUnlockTextureCubeFace(FRHITextureCube* Texture, uint32 FaceIndex, uint32 ArrayIndex, uint32 MipIndex, bool bLockWithinMiptail) final override;
+	virtual void* RHILockTextureCubeFace(FRHITexture* Texture, uint32 FaceIndex, uint32 ArrayIndex, uint32 MipIndex, EResourceLockMode LockMode, uint32& DestStride, bool bLockWithinMiptail) final override;
+	virtual void RHIUnlockTextureCubeFace(FRHITexture* Texture, uint32 FaceIndex, uint32 ArrayIndex, uint32 MipIndex, bool bLockWithinMiptail) final override;
 	virtual void RHIBindDebugLabelName(FRHICommandListBase& RHICmdList, FRHITexture* Texture, const TCHAR* Name) final override;
 	virtual void RHIBindDebugLabelName(FRHICommandListBase& RHICmdList, FRHIBuffer* Buffer, const TCHAR* Name) final override;
 	virtual void RHIReadSurfaceData(FRHITexture* Texture, FIntRect Rect, TArray<FColor>& OutData, FReadSurfaceDataFlags InFlags) final override;
@@ -360,7 +360,7 @@ public:
 	virtual void RHIEndRenderQuery_TopOfPipe(FRHICommandListBase& RHICmdList, FRHIRenderQuery* RenderQuery) final override;
 	virtual bool RHIGetRenderQueryResult(FRHIRenderQuery* RenderQuery, uint64& OutResult, bool bWait, uint32 GPUIndex = INDEX_NONE) final override;
 	virtual uint32 RHIGetViewportNextPresentGPUIndex(FRHIViewport* Viewport) final override;
-	virtual FTexture2DRHIRef RHIGetViewportBackBuffer(FRHIViewport* Viewport) final override;
+	virtual FTextureRHIRef RHIGetViewportBackBuffer(FRHIViewport* Viewport) final override;
 	virtual FUnorderedAccessViewRHIRef RHIGetViewportBackBufferUAV(FRHIViewport* Viewport) final override;
 	virtual void RHIAliasTextureResources(FTextureRHIRef& DestTexture, FTextureRHIRef& SrcTexture) final override;
 	virtual FTextureRHIRef RHICreateAliasedTexture(FTextureRHIRef& SourceTexture) final override;
@@ -409,9 +409,9 @@ public:
 	virtual uint32 RHIGetDeviceNodeMask(uint32 InIndex) const final override;
 	virtual ID3D12GraphicsCommandList* RHIGetGraphicsCommandList(uint32 InDeviceIndex) const final override;
 	virtual DXGI_FORMAT RHIGetSwapChainFormat(EPixelFormat InFormat) const final override;
-	virtual FTexture2DRHIRef RHICreateTexture2DFromResource(EPixelFormat Format, ETextureCreateFlags TexCreateFlags, const FClearValueBinding& ClearValueBinding, ID3D12Resource* Resource) final override;
-	virtual FTexture2DArrayRHIRef RHICreateTexture2DArrayFromResource(EPixelFormat Format, ETextureCreateFlags TexCreateFlags, const FClearValueBinding& ClearValueBinding, ID3D12Resource* Resource) final override;
-	virtual FTextureCubeRHIRef RHICreateTextureCubeFromResource(EPixelFormat Format, ETextureCreateFlags TexCreateFlags, const FClearValueBinding& ClearValueBinding, ID3D12Resource* Resource) final override;
+	virtual FTextureRHIRef RHICreateTexture2DFromResource(EPixelFormat Format, ETextureCreateFlags TexCreateFlags, const FClearValueBinding& ClearValueBinding, ID3D12Resource* Resource) final override;
+	virtual FTextureRHIRef RHICreateTexture2DArrayFromResource(EPixelFormat Format, ETextureCreateFlags TexCreateFlags, const FClearValueBinding& ClearValueBinding, ID3D12Resource* Resource) final override;
+	virtual FTextureRHIRef RHICreateTextureCubeFromResource(EPixelFormat Format, ETextureCreateFlags TexCreateFlags, const FClearValueBinding& ClearValueBinding, ID3D12Resource* Resource) final override;
 	virtual ID3D12Resource* RHIGetResource(FRHIBuffer* InBuffer) const final override;
 	virtual uint32 RHIGetResourceDeviceIndex(FRHIBuffer* InBuffer) const final override;
 	virtual int64 RHIGetResourceMemorySize(FRHIBuffer* InBuffer) const final override;
@@ -432,18 +432,18 @@ public:
 	// These will be un-commented as they are implemented.
 	//
 
-	virtual void* LockTexture2D_RenderThread(class FRHICommandListImmediate& RHICmdList, FRHITexture2D* Texture, uint32 MipIndex, EResourceLockMode LockMode, uint32& DestStride, bool bLockWithinMiptail, bool bNeedsDefaultRHIFlush = true, uint64* OutLockedByteCount = nullptr) override final;
-	virtual void UnlockTexture2D_RenderThread(class FRHICommandListImmediate& RHICmdList, FRHITexture2D* Texture, uint32 MipIndex, bool bLockWithinMiptail, bool bNeedsDefaultRHIFlush = true) override final;
+	virtual void* LockTexture2D_RenderThread(class FRHICommandListImmediate& RHICmdList, FRHITexture* Texture, uint32 MipIndex, EResourceLockMode LockMode, uint32& DestStride, bool bLockWithinMiptail, bool bNeedsDefaultRHIFlush = true, uint64* OutLockedByteCount = nullptr) override final;
+	virtual void UnlockTexture2D_RenderThread(class FRHICommandListImmediate& RHICmdList, FRHITexture* Texture, uint32 MipIndex, bool bLockWithinMiptail, bool bNeedsDefaultRHIFlush = true) override final;
 
-	virtual void* LockTexture2DArray_RenderThread(class FRHICommandListImmediate& RHICmdList, FRHITexture2DArray* Texture, uint32 ArrayIndex, uint32 MipIndex, EResourceLockMode LockMode, uint32& DestStride, bool bLockWithinMiptail) override final;
-	virtual void UnlockTexture2DArray_RenderThread(class FRHICommandListImmediate& RHICmdList, FRHITexture2DArray* Texture, uint32 ArrayIndex, uint32 MipIndex, bool bLockWithinMiptail) override final;
+	virtual void* LockTexture2DArray_RenderThread(class FRHICommandListImmediate& RHICmdList, FRHITexture* Texture, uint32 ArrayIndex, uint32 MipIndex, EResourceLockMode LockMode, uint32& DestStride, bool bLockWithinMiptail) override final;
+	virtual void UnlockTexture2DArray_RenderThread(class FRHICommandListImmediate& RHICmdList, FRHITexture* Texture, uint32 ArrayIndex, uint32 MipIndex, bool bLockWithinMiptail) override final;
 
-	virtual FTexture2DRHIRef AsyncReallocateTexture2D_RenderThread(class FRHICommandListImmediate& RHICmdList, FRHITexture2D* Texture2D, int32 NewMipCount, int32 NewSizeX, int32 NewSizeY, FThreadSafeCounter* RequestStatus);
-	virtual ETextureReallocationStatus FinalizeAsyncReallocateTexture2D_RenderThread(class FRHICommandListImmediate& RHICmdList, FRHITexture2D* Texture2D, bool bBlockUntilCompleted)
+	virtual FTextureRHIRef AsyncReallocateTexture2D_RenderThread(class FRHICommandListImmediate& RHICmdList, FRHITexture* Texture2D, int32 NewMipCount, int32 NewSizeX, int32 NewSizeY, FThreadSafeCounter* RequestStatus);
+	virtual ETextureReallocationStatus FinalizeAsyncReallocateTexture2D_RenderThread(class FRHICommandListImmediate& RHICmdList, FRHITexture* Texture2D, bool bBlockUntilCompleted)
 	{
 		return RHIFinalizeAsyncReallocateTexture2D(Texture2D, bBlockUntilCompleted);
 	}
-	virtual ETextureReallocationStatus CancelAsyncReallocateTexture2D_RenderThread(class FRHICommandListImmediate& RHICmdList, FRHITexture2D* Texture2D, bool bBlockUntilCompleted)
+	virtual ETextureReallocationStatus CancelAsyncReallocateTexture2D_RenderThread(class FRHICommandListImmediate& RHICmdList, FRHITexture* Texture2D, bool bBlockUntilCompleted)
 	{
 		return RHICancelAsyncReallocateTexture2D(Texture2D, bBlockUntilCompleted);
 	}
@@ -484,7 +484,7 @@ public:
 		// Not supported on PC
 	}
 
-	FUpdateTexture3DData BeginUpdateTexture3D_Internal(FRHITexture3D* Texture, uint32 MipIndex, const struct FUpdateTextureRegion3D& UpdateRegion);
+	FUpdateTexture3DData BeginUpdateTexture3D_Internal(FRHITexture* Texture, uint32 MipIndex, const struct FUpdateTextureRegion3D& UpdateRegion);
 	void EndUpdateTexture3D_Internal(FRHICommandListBase& RHICmdList, FUpdateTexture3DData& UpdateData);
 
 public:
diff --git a/Engine/Source/Runtime/D3D12RHI/Private/D3D12Texture.cpp b/Engine/Source/Runtime/D3D12RHI/Private/D3D12Texture.cpp
--- a/Engine/Source/Runtime/D3D12RHI/Private/D3D12Texture.cpp
+++ b/Engine/Source/Runtime/D3D12RHI/Private/D3D12Texture.cpp
@@ -1326,7 +1326,7 @@ uint32 FD3D12DynamicRHI::RHIComputeMemorySize(FRHITexture* TextureRHI)
 }
 
 
-FTexture2DRHIRef FD3D12DynamicRHI::AsyncReallocateTexture2D_RenderThread(class FRHICommandListImmediate& RHICmdList, FRHITexture2D* Texture2DRHI, int32 NewMipCount, int32 NewSizeX, int32 NewSizeY, FThreadSafeCounter* RequestStatus)
+FTextureRHIRef FD3D12DynamicRHI::AsyncReallocateTexture2D_RenderThread(class FRHICommandListImmediate& RHICmdList, FRHITexture* Texture2DRHI, int32 NewMipCount, int32 NewSizeX, int32 NewSizeY, FThreadSafeCounter* RequestStatus)
 {
 	if (RHICmdList.Bypass())
 	{
@@ -1371,7 +1371,7 @@ FTexture2DRHIRef FD3D12DynamicRHI::AsyncReallocateTexture2D_RenderThread(class F
  * @param RequestStatus	- Will be decremented by 1 when the reallocation is complete (success or failure).
  * @return				- New reference to the texture, or an invalid reference upon failure
  */
-FTexture2DRHIRef FD3D12DynamicRHI::RHIAsyncReallocateTexture2D(FRHITexture2D* Texture2DRHI, int32 NewMipCount, int32 NewSizeX, int32 NewSizeY, FThreadSafeCounter* RequestStatus)
+FTextureRHIRef FD3D12DynamicRHI::RHIAsyncReallocateTexture2D(FRHITexture* Texture2DRHI, int32 NewMipCount, int32 NewSizeX, int32 NewSizeY, FThreadSafeCounter* RequestStatus)
 {
 	FD3D12Texture* Texture2D = FD3D12DynamicRHI::ResourceCast(Texture2DRHI);
 
@@ -1396,12 +1396,12 @@ FTexture2DRHIRef FD3D12DynamicRHI::RHIAsyncReallocateTexture2D(FRHITexture2D* Te
 	return NewTexture2D;
 }
 
-ETextureReallocationStatus FD3D12DynamicRHI::RHIFinalizeAsyncReallocateTexture2D(FRHITexture2D* Texture2D, bool bBlockUntilCompleted)
+ETextureReallocationStatus FD3D12DynamicRHI::RHIFinalizeAsyncReallocateTexture2D(FRHITexture* Texture2D, bool bBlockUntilCompleted)
 {
 	return TexRealloc_Succeeded;
 }
 
-ETextureReallocationStatus FD3D12DynamicRHI::RHICancelAsyncReallocateTexture2D(FRHITexture2D* Texture2D, bool bBlockUntilCompleted)
+ETextureReallocationStatus FD3D12DynamicRHI::RHICancelAsyncReallocateTexture2D(FRHITexture* Texture2D, bool bBlockUntilCompleted)
 {
 	return TexRealloc_Succeeded;
 }
@@ -2154,7 +2154,7 @@ void FD3D12Texture::GetReadBackHeapDesc(D3D12_PLACED_SUBRESOURCE_FOOTPRINT& OutF
 	}
 }
 
-void* FD3D12DynamicRHI::LockTexture2D_RenderThread(class FRHICommandListImmediate& RHICmdList, FRHITexture2D* TextureRHI, uint32 MipIndex, EResourceLockMode LockMode, uint32& DestStride, bool bLockWithinMiptail, bool bNeedsDefaultRHIFlush, uint64* OutLockedByteCount)
+void* FD3D12DynamicRHI::LockTexture2D_RenderThread(class FRHICommandListImmediate& RHICmdList, FRHITexture* TextureRHI, uint32 MipIndex, EResourceLockMode LockMode, uint32& DestStride, bool bLockWithinMiptail, bool bNeedsDefaultRHIFlush, uint64* OutLockedByteCount)
 {
 	if (CVarD3D12Texture2DRHIFlush.GetValueOnRenderThread() && bNeedsDefaultRHIFlush)
 	{
@@ -2168,14 +2168,14 @@ void* FD3D12DynamicRHI::LockTexture2D_RenderThread(class FRHICommandListImmediat
 	return Texture->Lock(&RHICmdList, MipIndex, 0, LockMode, DestStride, OutLockedByteCount);
 }
 
-void* FD3D12DynamicRHI::RHILockTexture2D(FRHITexture2D* TextureRHI, uint32 MipIndex, EResourceLockMode LockMode, uint32& DestStride, bool bLockWithinMiptail, uint64* OutLockedByteCount)
+void* FD3D12DynamicRHI::RHILockTexture2D(FRHITexture* TextureRHI, uint32 MipIndex, EResourceLockMode LockMode, uint32& DestStride, bool bLockWithinMiptail, uint64* OutLockedByteCount)
 {
 	check(TextureRHI);
 	FD3D12Texture*  Texture = FD3D12DynamicRHI::ResourceCast(TextureRHI);
 	return Texture->Lock(nullptr, MipIndex, 0, LockMode, DestStride, OutLockedByteCount);
 }
 
-void FD3D12DynamicRHI::UnlockTexture2D_RenderThread(class FRHICommandListImmediate& RHICmdList, FRHITexture2D* TextureRHI, uint32 MipIndex, bool bLockWithinMiptail, bool bNeedsDefaultRHIFlush)
+void FD3D12DynamicRHI::UnlockTexture2D_RenderThread(class FRHICommandListImmediate& RHICmdList, FRHITexture* TextureRHI, uint32 MipIndex, bool bLockWithinMiptail, bool bNeedsDefaultRHIFlush)
 {
 	if (CVarD3D12Texture2DRHIFlush.GetValueOnRenderThread() && bNeedsDefaultRHIFlush)
 	{
@@ -2190,42 +2190,42 @@ void FD3D12DynamicRHI::UnlockTexture2D_RenderThread(class FRHICommandListImmedia
 	Texture->Unlock(&RHICmdList, MipIndex, 0);
 }
 
-void FD3D12DynamicRHI::RHIUnlockTexture2D(FRHITexture2D* TextureRHI, uint32 MipIndex, bool bLockWithinMiptail)
+void FD3D12DynamicRHI::RHIUnlockTexture2D(FRHITexture* TextureRHI, uint32 MipIndex, bool bLockWithinMiptail)
 {
 	check(TextureRHI);
 	FD3D12Texture*  Texture = FD3D12DynamicRHI::ResourceCast(TextureRHI);
 	Texture->Unlock(nullptr, MipIndex, 0);
 }
 
-void* FD3D12DynamicRHI::LockTexture2DArray_RenderThread(class FRHICommandListImmediate& RHICmdList, FRHITexture2DArray* TextureRHI, uint32 TextureIndex, uint32 MipIndex, EResourceLockMode LockMode, uint32& DestStride, bool bLockWithinMiptail)
+void* FD3D12DynamicRHI::LockTexture2DArray_RenderThread(class FRHICommandListImmediate& RHICmdList, FRHITexture* TextureRHI, uint32 TextureIndex, uint32 MipIndex, EResourceLockMode LockMode, uint32& DestStride, bool bLockWithinMiptail)
 {
 	check(TextureRHI);
 	FD3D12Texture* Texture = FD3D12DynamicRHI::ResourceCast(TextureRHI);
 	return Texture->Lock(&RHICmdList, MipIndex, TextureIndex, LockMode, DestStride);
 }
 
-void* FD3D12DynamicRHI::RHILockTexture2DArray(FRHITexture2DArray* TextureRHI, uint32 TextureIndex, uint32 MipIndex, EResourceLockMode LockMode, uint32& DestStride, bool bLockWithinMiptail)
+void* FD3D12DynamicRHI::RHILockTexture2DArray(FRHITexture* TextureRHI, uint32 TextureIndex, uint32 MipIndex, EResourceLockMode LockMode, uint32& DestStride, bool bLockWithinMiptail)
 {
 	check(TextureRHI);
 	FD3D12Texture* Texture = FD3D12DynamicRHI::ResourceCast(TextureRHI);
 	return Texture->Lock(nullptr, MipIndex, TextureIndex, LockMode, DestStride);
 }
 
-void FD3D12DynamicRHI::UnlockTexture2DArray_RenderThread(class FRHICommandListImmediate& RHICmdList, FRHITexture2DArray* TextureRHI, uint32 TextureIndex, uint32 MipIndex, bool bLockWithinMiptail)
+void FD3D12DynamicRHI::UnlockTexture2DArray_RenderThread(class FRHICommandListImmediate& RHICmdList, FRHITexture* TextureRHI, uint32 TextureIndex, uint32 MipIndex, bool bLockWithinMiptail)
 {
 	check(TextureRHI);
 	FD3D12Texture* Texture = FD3D12DynamicRHI::ResourceCast(TextureRHI);
 	Texture->Unlock(&RHICmdList, MipIndex, TextureIndex);
 }
 
-void FD3D12DynamicRHI::RHIUnlockTexture2DArray(FRHITexture2DArray* TextureRHI, uint32 TextureIndex, uint32 MipIndex, bool bLockWithinMiptail)
+void FD3D12DynamicRHI::RHIUnlockTexture2DArray(FRHITexture* TextureRHI, uint32 TextureIndex, uint32 MipIndex, bool bLockWithinMiptail)
 {
 	check(TextureRHI);
 	FD3D12Texture* Texture = FD3D12DynamicRHI::ResourceCast(TextureRHI);
 	Texture->Unlock(nullptr, MipIndex, TextureIndex);
 }
 
-void* FD3D12DynamicRHI::RHILockTextureCubeFace_RenderThread(class FRHICommandListImmediate& RHICmdList, FRHITextureCube* TextureRHI, uint32 FaceIndex, uint32 ArrayIndex, uint32 MipIndex, EResourceLockMode LockMode, uint32& DestStride, bool bLockWithinMiptail)
+void* FD3D12DynamicRHI::RHILockTextureCubeFace_RenderThread(class FRHICommandListImmediate& RHICmdList, FRHITexture* TextureRHI, uint32 FaceIndex, uint32 ArrayIndex, uint32 MipIndex, EResourceLockMode LockMode, uint32& DestStride, bool bLockWithinMiptail)
 {
 	check(TextureRHI);
 	FD3D12Texture* TextureCube = FD3D12DynamicRHI::ResourceCast(TextureRHI);
@@ -2233,7 +2233,7 @@ void* FD3D12DynamicRHI::RHILockTextureCubeFace_RenderThread(class FRHICommandLis
 	return TextureCube->Lock(&RHICmdList, MipIndex, D3DFace + ArrayIndex * 6, LockMode, DestStride);
 }
 
-void FD3D12DynamicRHI::RHIUnlockTextureCubeFace_RenderThread(class FRHICommandListImmediate& RHICmdList, FRHITextureCube* TextureRHI, uint32 FaceIndex, uint32 ArrayIndex, uint32 MipIndex, bool bLockWithinMiptail)
+void FD3D12DynamicRHI::RHIUnlockTextureCubeFace_RenderThread(class FRHICommandListImmediate& RHICmdList, FRHITexture* TextureRHI, uint32 FaceIndex, uint32 ArrayIndex, uint32 MipIndex, bool bLockWithinMiptail)
 {
 	check(TextureRHI);
 	FD3D12Texture* TextureCube = FD3D12DynamicRHI::ResourceCast(TextureRHI);
@@ -2241,14 +2241,14 @@ void FD3D12DynamicRHI::RHIUnlockTextureCubeFace_RenderThread(class FRHICommandLi
 	TextureCube->Unlock(&RHICmdList, MipIndex, D3DFace + ArrayIndex * 6);
 }
 
-void FD3D12DynamicRHI::RHIUpdateTexture2D(FRHICommandListBase& RHICmdList, FRHITexture2D* TextureRHI, uint32 MipIndex, const FUpdateTextureRegion2D& UpdateRegion, uint32 SourcePitch, const uint8* SourceData)
+void FD3D12DynamicRHI::RHIUpdateTexture2D(FRHICommandListBase& RHICmdList, FRHITexture* TextureRHI, uint32 MipIndex, const FUpdateTextureRegion2D& UpdateRegion, uint32 SourcePitch, const uint8* SourceData)
 {
 	check(TextureRHI);
 	FD3D12Texture* Texture = FD3D12DynamicRHI::ResourceCast(TextureRHI);
 	Texture->UpdateTexture2D(RHICmdList, MipIndex, UpdateRegion, SourcePitch, SourceData);
 }
 
-FUpdateTexture3DData FD3D12DynamicRHI::RHIBeginUpdateTexture3D(FRHICommandListBase& RHICmdList, FRHITexture3D* Texture, uint32 MipIndex, const struct FUpdateTextureRegion3D& UpdateRegion)
+FUpdateTexture3DData FD3D12DynamicRHI::RHIBeginUpdateTexture3D(FRHICommandListBase& RHICmdList, FRHITexture* Texture, uint32 MipIndex, const struct FUpdateTextureRegion3D& UpdateRegion)
 {
 	return BeginUpdateTexture3D_Internal(Texture, MipIndex, UpdateRegion);
 }
@@ -2370,7 +2370,7 @@ private:
 	};
 
 	uint32 MipIdx;
-	FTexture3DRHIRef DstTexture;
+	FTextureRHIRef DstTexture;
 	TArray<FUpdateInfo> UpdateInfos;
 };
 
@@ -2421,7 +2421,7 @@ void FD3D12DynamicRHI::RHIEndMultiUpdateTexture3D(FRHICommandListBase& RHICmdLis
 	}
 }
 
-void FD3D12DynamicRHI::RHIUpdateTexture3D(FRHICommandListBase& RHICmdList, FRHITexture3D* TextureRHI, uint32 MipIndex, const FUpdateTextureRegion3D& InUpdateRegion, uint32 SourceRowPitch, uint32 SourceDepthPitch, const uint8* SourceData)
+void FD3D12DynamicRHI::RHIUpdateTexture3D(FRHICommandListBase& RHICmdList, FRHITexture* TextureRHI, uint32 MipIndex, const FUpdateTextureRegion3D& InUpdateRegion, uint32 SourceRowPitch, uint32 SourceDepthPitch, const uint8* SourceData)
 {
 	FD3D12Texture* Texture = FD3D12DynamicRHI::ResourceCast(TextureRHI);
 	const FPixelFormatInfo& FormatInfo = GPixelFormats[Texture->GetFormat()];
@@ -2464,7 +2464,7 @@ void FD3D12DynamicRHI::RHIUpdateTexture3D(FRHICommandListBase& RHICmdList, FRHIT
 }
 
 
-FUpdateTexture3DData FD3D12DynamicRHI::BeginUpdateTexture3D_Internal(FRHITexture3D* TextureRHI, uint32 MipIndex, const struct FUpdateTextureRegion3D& UpdateRegion)
+FUpdateTexture3DData FD3D12DynamicRHI::BeginUpdateTexture3D_Internal(FRHITexture* TextureRHI, uint32 MipIndex, const struct FUpdateTextureRegion3D& UpdateRegion)
 {
 	check(IsInParallelRenderingThread());
 	FUpdateTexture3DData UpdateData(TextureRHI, MipIndex, UpdateRegion, 0, 0, nullptr, 0, GFrameNumberRenderThread);
@@ -2600,7 +2600,7 @@ private:
 	uint32 DstStartX;
 	uint32 DstStartY;
 	uint32 DstStartZ;
-	FTexture3DRHIRef DstTexture;
+	FTextureRHIRef DstTexture;
 	FD3D12ResourceLocation* SrcResourceLocation;
 	D3D12_PLACED_SUBRESOURCE_FOOTPRINT PlacedSubresourceFootprint;
 };
@@ -2634,7 +2634,7 @@ void FD3D12DynamicRHI::EndUpdateTexture3D_Internal(FRHICommandListBase& RHICmdLi
 	Cubemap texture support.
 	-----------------------------------------------------------------------------*/
 
-void* FD3D12DynamicRHI::RHILockTextureCubeFace(FRHITextureCube* TextureCubeRHI, uint32 FaceIndex, uint32 ArrayIndex, uint32 MipIndex, EResourceLockMode LockMode, uint32& DestStride, bool bLockWithinMiptail)
+void* FD3D12DynamicRHI::RHILockTextureCubeFace(FRHITexture* TextureCubeRHI, uint32 FaceIndex, uint32 ArrayIndex, uint32 MipIndex, EResourceLockMode LockMode, uint32& DestStride, bool bLockWithinMiptail)
 {
 	FD3D12Texture*  TextureCube = FD3D12DynamicRHI::ResourceCast(TextureCubeRHI);
 	for (uint32 GPUIndex : TextureCube->GetLinkedObjectsGPUMask())
@@ -2644,7 +2644,7 @@ void* FD3D12DynamicRHI::RHILockTextureCubeFace(FRHITextureCube* TextureCubeRHI,
 	uint32 D3DFace = GetD3D12CubeFace((ECubeFace)FaceIndex);
 	return TextureCube->Lock(nullptr, MipIndex, D3DFace + ArrayIndex * 6, LockMode, DestStride);
 }
-void FD3D12DynamicRHI::RHIUnlockTextureCubeFace(FRHITextureCube* TextureCubeRHI, uint32 FaceIndex, uint32 ArrayIndex, uint32 MipIndex, bool bLockWithinMiptail)
+void FD3D12DynamicRHI::RHIUnlockTextureCubeFace(FRHITexture* TextureCubeRHI, uint32 FaceIndex, uint32 ArrayIndex, uint32 MipIndex, bool bLockWithinMiptail)
 {
 	FD3D12Texture* TextureCube = FD3D12DynamicRHI::ResourceCast(TextureCubeRHI);
 	uint32 D3DFace = GetD3D12CubeFace((ECubeFace)FaceIndex);
@@ -2775,17 +2775,17 @@ FD3D12Texture* FD3D12DynamicRHI::CreateTextureFromResource(bool bTextureArray, b
 	return Texture2D;
 }
 
-FTexture2DRHIRef FD3D12DynamicRHI::RHICreateTexture2DFromResource(EPixelFormat Format, ETextureCreateFlags TexCreateFlags, const FClearValueBinding& ClearValueBinding, ID3D12Resource* Resource)
+FTextureRHIRef FD3D12DynamicRHI::RHICreateTexture2DFromResource(EPixelFormat Format, ETextureCreateFlags TexCreateFlags, const FClearValueBinding& ClearValueBinding, ID3D12Resource* Resource)
 {
 	return CreateTextureFromResource(false, false, Format, TexCreateFlags, ClearValueBinding, Resource);
 }
 
-FTexture2DArrayRHIRef FD3D12DynamicRHI::RHICreateTexture2DArrayFromResource(EPixelFormat Format, ETextureCreateFlags TexCreateFlags, const FClearValueBinding& ClearValueBinding, ID3D12Resource* Resource)
+FTextureRHIRef FD3D12DynamicRHI::RHICreateTexture2DArrayFromResource(EPixelFormat Format, ETextureCreateFlags TexCreateFlags, const FClearValueBinding& ClearValueBinding, ID3D12Resource* Resource)
 {
 	return CreateTextureFromResource(true, false, Format, TexCreateFlags, ClearValueBinding, Resource);
 }
 
-FTextureCubeRHIRef FD3D12DynamicRHI::RHICreateTextureCubeFromResource(EPixelFormat Format, ETextureCreateFlags TexCreateFlags, const FClearValueBinding& ClearValueBinding, ID3D12Resource* Resource)
+FTextureRHIRef FD3D12DynamicRHI::RHICreateTextureCubeFromResource(EPixelFormat Format, ETextureCreateFlags TexCreateFlags, const FClearValueBinding& ClearValueBinding, ID3D12Resource* Resource)
 {
 	return CreateTextureFromResource(false, true, Format, TexCreateFlags, ClearValueBinding, Resource);
 }
diff --git a/Engine/Source/Runtime/D3D12RHI/Private/D3D12Viewport.cpp b/Engine/Source/Runtime/D3D12RHI/Private/D3D12Viewport.cpp
--- a/Engine/Source/Runtime/D3D12RHI/Private/D3D12Viewport.cpp
+++ b/Engine/Source/Runtime/D3D12RHI/Private/D3D12Viewport.cpp
@@ -714,8 +714,8 @@ bool FD3D12Viewport::Present(bool bLockToVsync)
 		FD3D12CommandContext& DefaultContext = Adapter->GetDevice(GPUIndex)->GetDefaultCommandContext();
 
 		// Those are not necessarily the swap chain back buffer in case of multi-gpu
-		FD3D12Texture* DeviceBackBuffer = DefaultContext.RetrieveObject<FD3D12Texture, FRHITexture2D>(GetBackBuffer_RHIThread());
-		FD3D12Texture* DeviceSDRBackBuffer = DefaultContext.RetrieveObject<FD3D12Texture, FRHITexture2D>(GetSDRBackBuffer_RHIThread());
+		FD3D12Texture* DeviceBackBuffer = DefaultContext.RetrieveObject<FD3D12Texture, FRHITexture>(GetBackBuffer_RHIThread());
+		FD3D12Texture* DeviceSDRBackBuffer = DefaultContext.RetrieveObject<FD3D12Texture, FRHITexture>(GetSDRBackBuffer_RHIThread());
 
 		DefaultContext.TransitionResource(
 			DeviceBackBuffer->GetShaderResourceView()->GetResource(),
@@ -1027,7 +1027,7 @@ uint32 FD3D12DynamicRHI::RHIGetViewportNextPresentGPUIndex(FRHIViewport* Viewpor
 	return 0;
 }
 
-FTexture2DRHIRef FD3D12DynamicRHI::RHIGetViewportBackBuffer(FRHIViewport* ViewportRHI)
+FTextureRHIRef FD3D12DynamicRHI::RHIGetViewportBackBuffer(FRHIViewport* ViewportRHI)
 {
 	check(IsInRenderingThread());
 
