HWRT: Fix race condition when transfering ray tracing geometry during streaming:
* Remove FD3D12RHICommandInitializeRayTracingGeometry
* Enqueue the BLAS deserialization on the command list that was passed in to the ctor.

Submitting on behalf of Luke Thatcher.

diff --git a/Engine/Source/Runtime/D3D12RHI/Private/D3D12RayTracing.cpp b/Engine/Source/Runtime/D3D12RHI/Private/D3D12RayTracing.cpp
--- a/Engine/Source/Runtime/D3D12RHI/Private/D3D12RayTracing.cpp
+++ b/Engine/Source/Runtime/D3D12RHI/Private/D3D12RayTracing.cpp
@@ -2343,59 +2343,6 @@ public:
 #endif // !NO_LOGGING
 };
 
-struct FD3D12RHICommandInitializeRayTracingGeometryString
-{
-	static const TCHAR* TStr() { return TEXT("FD3D12RHICommandInitializeRayTracingGeometryString"); }
-};
-struct FD3D12RHICommandInitializeRayTracingGeometry final : public FRHICommand<FD3D12RHICommandInitializeRayTracingGeometry, FD3D12RHICommandInitializeRayTracingGeometryString>
-{
-	TRefCountPtr<FD3D12RayTracingGeometry> Geometry;
-	FD3D12ResourceLocation SrcResourceLoc;
-	bool bForRendering;
-
-	FORCEINLINE_DEBUGGABLE FD3D12RHICommandInitializeRayTracingGeometry(TRefCountPtr<FD3D12RayTracingGeometry>&& InGeometry, FD3D12ResourceLocation& InSrcResourceLoc, bool bForRendering)
-		: Geometry(MoveTemp(InGeometry))
-		, SrcResourceLoc(InSrcResourceLoc.GetParentDevice())
-		, bForRendering(bForRendering)
-	{
-		FD3D12ResourceLocation::TransferOwnership(SrcResourceLoc, InSrcResourceLoc);
-	}
-
-	void Execute(FRHICommandListBase& /* unused */)
-	{
-		ExecuteNoCmdList();
-	}
-
-	void ExecuteNoCmdList()
-	{
-		for (uint32 GPUIndex = 0; GPUIndex < MAX_NUM_GPUS && GPUIndex < GNumExplicitGPUsForRendering; ++GPUIndex)
-		{
-			FD3D12Buffer* AccelerationStructure = Geometry->AccelerationStructureBuffers[GPUIndex];
-			FD3D12Resource* Destination = AccelerationStructure->ResourceLocation.GetResource();
-			FD3D12Device* Device = Destination->GetParentDevice();
-
-			FD3D12CommandContext& Context = Device->GetDefaultCommandContext();
-
-			Context.RayTracingCommandList()->CopyRaytracingAccelerationStructure(
-				AccelerationStructure->ResourceLocation.GetGPUVirtualAddress(),
-				SrcResourceLoc.GetGPUVirtualAddress(),
-				D3D12_RAYTRACING_ACCELERATION_STRUCTURE_COPY_MODE_DESERIALIZE
-			);
-
-			Context.UpdateResidency(SrcResourceLoc.GetResource());
-			Context.ConditionalSplitCommandList();
-			
-			if (bForRendering)
-			{
-				Geometry->RegisterAsRenameListener(GPUIndex);
-				Geometry->SetupHitGroupSystemParameters(GPUIndex);
-			}
-
-			Geometry->SetDirty(FRHIGPUMask::FromIndex(GPUIndex), false);
-		}
-	}
-};
-
 void FD3D12Device::InitRayTracing()
 {
 	LLM_SCOPE_BYNAME(TEXT("FD3D12RT"));
@@ -2901,54 +2848,36 @@ FD3D12RayTracingGeometry::FD3D12RayTracingGeometry(FRHICommandListBase& RHICmdLi
 		uint32 Size = Initializer.OfflineData->GetResourceDataSize() - sizeof(FOfflineBVHHeader);
 
 		FD3D12ResourceLocation SrcResourceLoc(Device);
+		uint8* DstDataBase = (uint8*)Adapter->GetUploadHeapAllocator(0).AllocUploadResource(Size, 256, SrcResourceLoc);
+		FMemory::Memcpy(DstDataBase, Data, Size);
 
-		const bool bOnAsyncThread = !IsInRHIThread() && !IsInRenderingThread();
-		uint8* DstDataBase;
-		if (bOnAsyncThread)
-		{
-			DstDataBase = (uint8*)Adapter->GetUploadHeapAllocator(0).AllocUploadResource(Size, 256, SrcResourceLoc);
-		}
-		else
+		RHICmdList.EnqueueLambda([this, SrcResourceLoc = MoveTemp(SrcResourceLoc), bForRendering](FRHICommandListBase& ExecutingCmdList)
 		{
-			DstDataBase = (uint8*)Device->GetDefaultFastAllocator().Allocate(Size, 256, &SrcResourceLoc);
-		}
+			for (uint32 GPUIndex = 0; GPUIndex < MAX_NUM_GPUS && GPUIndex < GNumExplicitGPUsForRendering; ++GPUIndex)
+			{
+				FD3D12CommandContext& Context = FD3D12CommandContext::Get(ExecutingCmdList, GPUIndex);
 
-		FMemory::Memcpy(DstDataBase, Data, Size);
+				FD3D12Buffer* AccelerationStructure = AccelerationStructureBuffers[GPUIndex];
 
-		if (bOnAsyncThread)
-		{
-			check(Initializer.Type == ERayTracingGeometryInitializerType::StreamingSource);
+				Context.RayTracingCommandList()->CopyRaytracingAccelerationStructure(
+					AccelerationStructure->ResourceLocation.GetGPUVirtualAddress(),
+					SrcResourceLoc.GetGPUVirtualAddress(),
+					D3D12_RAYTRACING_ACCELERATION_STRUCTURE_COPY_MODE_DESERIALIZE
+				);
 
-			FD3D12ResourceLocation* SrcResourceLoc_Heap = new FD3D12ResourceLocation(SrcResourceLoc.GetParentDevice());
-			FD3D12ResourceLocation::TransferOwnership(*SrcResourceLoc_Heap, SrcResourceLoc);
+				Context.UpdateResidency(SrcResourceLoc.GetResource());
+				Context.ConditionalSplitCommandList();
 
-			TRefCountPtr<FD3D12RayTracingGeometry> ThisRef = this;
-			ENQUEUE_RENDER_COMMAND(CmdD3D12InitializeRayTracingGeometry)(
-				[DestinationGeometry = MoveTemp(ThisRef), SrcResourceLoc_Heap](FRHICommandListImmediate& RHICmdList) mutable
-			{
-				const bool bForRendering = false;
-				if (RHICmdList.Bypass())
+				if (bForRendering)
 				{
-					FD3D12RHICommandInitializeRayTracingGeometry Command(MoveTemp(DestinationGeometry), *SrcResourceLoc_Heap, bForRendering);
-					Command.ExecuteNoCmdList();
+					RegisterAsRenameListener(GPUIndex);
+					SetupHitGroupSystemParameters(GPUIndex);
 				}
-				else
-				{
-					ALLOC_COMMAND_CL(RHICmdList, FD3D12RHICommandInitializeRayTracingGeometry)(MoveTemp(DestinationGeometry), *SrcResourceLoc_Heap, bForRendering);
-				}
-				delete SrcResourceLoc_Heap;
-			});
-		}
-		else if (!RHICmdList.Bypass() && IsRunningRHIInSeparateThread())
-		{
-			ALLOC_COMMAND_CL(RHICmdList, FD3D12RHICommandInitializeRayTracingGeometry)(this, SrcResourceLoc, bForRendering);
-		}
-		else
-		{
-			FD3D12RHICommandInitializeRayTracingGeometry Command(this, SrcResourceLoc, bForRendering);
-			Command.ExecuteNoCmdList();
-		}
-		
+			}
+
+			SetDirty(FRHIGPUMask::All(), false);
+		});
+
 		Initializer.OfflineData->Discard();
 	}
 	else
