Prevent false sharing when filling SBT on D3D12 RHI.
- all worker threads were updating bIsDirty.
- instead update bIsDirty once after bindings are written.

On project that has SBT with ~110K bindings this reduces cost of SetRayTracingBindings from 6ms to 0.7ms (running on PC).

diff --git a/Engine/Source/Runtime/D3D12RHI/Private/D3D12RayTracing.cpp b/Engine/Source/Runtime/D3D12RHI/Private/D3D12RayTracing.cpp
--- a/Engine/Source/Runtime/D3D12RHI/Private/D3D12RayTracing.cpp
+++ b/Engine/Source/Runtime/D3D12RHI/Private/D3D12RayTracing.cpp
@@ -1266,8 +1266,6 @@ private:
 #endif // DO_CHECK && DO_GUARD_SLOW
 
 		FMemory::Memcpy(Data.GetData() + WriteOffset, InData, InDataSize);
-
-		bIsDirty = true;
 	}
 
 	void WriteLocalShaderRecord(uint32 ShaderTableOffset, uint32 RecordIndex, uint32 OffsetWithinRecord, const void* InData, uint32 InDataSize)
@@ -1424,8 +1422,6 @@ public:
 			Data.GetData() + DestOffset,
 			Data.GetData() + SourceOffset,
 			CopySize);
-
-		bIsDirty = true;
 	}
 
 	void CopyHitGroupParameters(uint32 InDestRecordIndex, uint32 InSourceRecordIndex, uint32 InOffsetWithinRootSignature)
@@ -4956,6 +4952,8 @@ void FD3D12CommandContext::RHISetRayTracingBindings(
 	const int32 ItemsPerTask = 1024;
 
 	ParallelForWithExistingTaskContext(TEXT("SetRayTracingBindings"), MakeArrayView(TaskContexts), NumBindings, ItemsPerTask, BindingTask);
+
+	ShaderTable->bIsDirty = true;
 }
 
 #endif // D3D12_RHI_RAYTRACING
