Remove most uses of the default command context and ALLOC_COMMAND_CL in D3D12 RHI
 - Replaced with RHI command list EnqueueLambda
 - The correct D3D12 command context is retrieved from the RHICmdList itself.

Removed bottom-of-pipe implementations of the following functions. These are implemented by the _RenderThread versions:
 - RHIAsyncReallocateTexture2D
 - RHILockTexture2D
 - RHIUnlockTexture2D
 - RHILockTexture2DArray
 - RHIUnlockTexture2DArray
 - RHILockTextureCubeFace
 - RHIUnlockTextureCubeFace

Replace some manual PIX markers with RHI breadcrumbs

Also removed "D3D12.LockTexture2DRHIFlush" cvar.
 - Unused and no longer necessary.

diff --git a/Engine/Source/Runtime/D3D12RHI/Private/D3D12Buffer.cpp b/Engine/Source/Runtime/D3D12RHI/Private/D3D12Buffer.cpp
--- a/Engine/Source/Runtime/D3D12RHI/Private/D3D12Buffer.cpp
+++ b/Engine/Source/Runtime/D3D12RHI/Private/D3D12Buffer.cpp
@@ -31,86 +31,6 @@ FD3D12Buffer::~FD3D12Buffer()
 	}
 }
 
-struct FRHICommandUpdateBufferString
-{
-	static const TCHAR* TStr() { return TEXT("FRHICommandUpdateBuffer"); }
-};
-struct FRHICommandUpdateBuffer final : public FRHICommand<FRHICommandUpdateBuffer, FRHICommandUpdateBufferString>
-{
-	FD3D12ResourceLocation Source;
-	FD3D12ResourceLocation* Destination;
-	uint32 NumBytes;
-	uint32 DestinationOffset;
-
-	FORCEINLINE_DEBUGGABLE FRHICommandUpdateBuffer(FD3D12ResourceLocation* InDest, FD3D12ResourceLocation& InSource, uint32 InDestinationOffset, uint32 InNumBytes)
-		: Source(nullptr)
-		, Destination(InDest)
-		, NumBytes(InNumBytes)
-		, DestinationOffset(InDestinationOffset)
-	{
-		FD3D12ResourceLocation::TransferOwnership(Source, InSource);
-	}
-
-	void Execute(FRHICommandListBase& CmdList)
-	{
-		FD3D12CommandContextBase::Get(CmdList).UpdateBuffer(Destination, DestinationOffset, &Source, 0, NumBytes);
-	}
-};
-
-// This allows us to rename resources from the RenderThread i.e. all the 'hard' work of allocating a new resource
-// is done in parallel and this small function is called to switch the resource to point to the correct location
-// a the correct time.
-struct FRHICommandRenameUploadBufferString
-{
-	static const TCHAR* TStr() { return TEXT("FRHICommandRenameUploadBuffer"); }
-};
-struct FRHICommandRenameUploadBuffer final : public FRHICommand<FRHICommandRenameUploadBuffer, FRHICommandRenameUploadBufferString>
-{
-	FD3D12Buffer* Resource;
-	FD3D12ResourceLocation NewLocation;
-
-	FORCEINLINE_DEBUGGABLE FRHICommandRenameUploadBuffer(FD3D12Buffer* InResource, FD3D12Device* Device)
-		: Resource(InResource)
-		, NewLocation(Device)
-	{}
-
-	void Execute(FRHICommandListBase& CmdList)
-	{
-		const static FLazyName ExecuteName(TEXT("FRHICommandRenameUploadBuffer::Execute"));
-		UE_TRACE_METADATA_SCOPE_ASSET_FNAME(Resource->GetName(), ExecuteName, Resource->GetOwnerName());
-
-		// Make sure we have an active pipeline when running at the top of the pipe, so the lambda below can obtain a context.
-		TOptional<ERHIPipeline> PreviousPipeline;
-		if (CmdList.IsTopOfPipe() && CmdList.GetPipeline() == ERHIPipeline::None)
-		{
-			PreviousPipeline = CmdList.SwitchPipeline(ERHIPipeline::Graphics);
-		}
-
-		// Clear the resource if still bound to make sure the SRVs are rebound again on next operation. This needs to happen
-		// on the RHI timeline when this command runs at the top of the pipe (which can happen when locking buffers in
-		// RLM_WriteOnly_NoOverwrite mode).
-		CmdList.EnqueueLambda([ResourceLocation = &Resource->ResourceLocation](FRHICommandListBase& RHICmdList)
-		{
-			const uint32 GPUIndex = 0; // @todo mgpu - seems wrong we're only doing this for the 0th GPU
-
-			FD3D12CommandContext& Context = FD3D12CommandContext::Get(RHICmdList, GPUIndex);
-			Context.ConditionalClearShaderResource(ResourceLocation, EShaderParameterTypeMask::SRVMask);
-		});
-
-#if UE_MEMORY_TRACE_ENABLED
-		// This memory trace happens before RenameLDAChain so the old & new GPU addresses are correct
-		MemoryTrace_ReallocFree(Resource->ResourceLocation.GetGPUVirtualAddress(), EMemoryTraceRootHeap::VideoMemory);
-		MemoryTrace_ReallocAlloc(NewLocation.GetGPUVirtualAddress(), Resource->ResourceLocation.GetSize(), Resource->BufferAlignment, EMemoryTraceRootHeap::VideoMemory);
-#endif
-		Resource->RenameLDAChain(CmdList, NewLocation);
-
-		if (PreviousPipeline.IsSet())
-		{
-			CmdList.SwitchPipeline(PreviousPipeline.GetValue());
-		}
-	}
-};
-
 void FD3D12Buffer::UploadResourceData(FRHICommandListBase& RHICmdList, FResourceArrayInterface* InResourceArray, D3D12_RESOURCE_STATES InDestinationState, const TCHAR* AssetName, const FName& ClassName, const FName& PackageName)
 {
 	UE_TRACE_METADATA_SCOPE_ASSET_FNAME(AssetName, ClassName, PackageName);
@@ -551,6 +471,7 @@ void* FD3D12DynamicRHI::LockBuffer(FRHICommandListBase& RHICmdList, FD3D12Buffer
 
 	FD3D12LockedResource& LockedData = Buffer->LockedData;
 	check(LockedData.bLocked == false);
+
 	FD3D12Adapter& Adapter = GetAdapter();
 
 	void* Data = nullptr;
@@ -570,18 +491,47 @@ void* FD3D12DynamicRHI::LockBuffer(FRHICommandListBase& RHICmdList, FD3D12Buffer
 		{
 			FD3D12Device* Device = Buffer->GetParentDevice();
 
-			// If on the RenderThread, queue up a command on the RHIThread to rename this buffer at the correct time
+			FD3D12ResourceLocation NewLocation(Device);
+			Data = Adapter.GetUploadHeapAllocator(Device->GetGPUIndex()).AllocUploadResource(BufferSize, Buffer->BufferAlignment, NewLocation);
+
+			// Make sure we have an active pipeline when running at the top of the pipe, so the lambda below can obtain a context.
+			TOptional<ERHIPipeline> PreviousPipeline;
+			if (RHICmdList.IsTopOfPipe() && RHICmdList.GetPipeline() == ERHIPipeline::None)
+			{
+				PreviousPipeline = RHICmdList.SwitchPipeline(ERHIPipeline::Graphics);
+			}
+
+			RHICmdList.EnqueueLambda([
+				Resource = Buffer,
+				NewLocation = MoveTemp(NewLocation)
+			](FRHICommandListBase& ExecutingCmdList) mutable
+			{
+				const static FLazyName ExecuteName(TEXT("FRHICommandRenameUploadBuffer::Execute"));
+				UE_TRACE_METADATA_SCOPE_ASSET_FNAME(Resource->GetName(), ExecuteName, Resource->GetOwnerName());
+
+				// Clear the resource if still bound to make sure the SRVs are rebound again on next operation. This needs to happen
+				// on the RHI timeline when this command runs at the top of the pipe (which can happen when locking buffers in
+				// RLM_WriteOnly_NoOverwrite mode).
+				
+				FD3D12CommandContext& Context = FD3D12CommandContext::Get(ExecutingCmdList, Resource->GetParentDevice()->GetGPUIndex());
+				Context.ConditionalClearShaderResource(&Resource->ResourceLocation, EShaderParameterTypeMask::SRVMask);
+
+#if UE_MEMORY_TRACE_ENABLED
+				// This memory trace happens before RenameLDAChain so the old & new GPU addresses are correct
+				MemoryTrace_ReallocFree(Resource->ResourceLocation.GetGPUVirtualAddress(), EMemoryTraceRootHeap::VideoMemory);
+				MemoryTrace_ReallocAlloc(NewLocation.GetGPUVirtualAddress(), Resource->ResourceLocation.GetSize(), Resource->BufferAlignment, EMemoryTraceRootHeap::VideoMemory);
+#endif
+				Resource->RenameLDAChain(ExecutingCmdList, NewLocation);
+			});
+
 			if (RHICmdList.IsTopOfPipe())
 			{
-				FRHICommandRenameUploadBuffer* Command = ALLOC_COMMAND_CL(RHICmdList, FRHICommandRenameUploadBuffer)(Buffer, Device);
-				Data = Adapter.GetUploadHeapAllocator(Device->GetGPUIndex()).AllocUploadResource(BufferSize, Buffer->BufferAlignment, Command->NewLocation);
 				RHICmdList.RHIThreadFence(true);
 			}
-			else
+
+			if (PreviousPipeline.IsSet())
 			{
-				FRHICommandRenameUploadBuffer Command(Buffer, Device);
-				Data = Adapter.GetUploadHeapAllocator(Device->GetGPUIndex()).AllocUploadResource(BufferSize, Buffer->BufferAlignment, Command.NewLocation);
-				Command.Execute(RHICmdList);
+				RHICmdList.SwitchPipeline(PreviousPipeline.GetValue());
 			}
 		}
 	}
@@ -594,8 +544,6 @@ void* FD3D12DynamicRHI::LockBuffer(FRHICommandListBase& RHICmdList, FD3D12Buffer
 		// Locking for read must occur immediately so we can't queue up the operations later.
 		if (LockMode == RLM_ReadOnly)
 		{
-			FRHICommandListImmediate& RHICmdListImmediate = RHICmdList.GetAsImmediate();
-
 			LockedData.bLockedForReadOnly = true;
 			// If the static buffer is being locked for reading, create a staging buffer.
 			FD3D12Resource* StagingBuffer = nullptr;
@@ -604,43 +552,26 @@ void* FD3D12DynamicRHI::LockBuffer(FRHICommandListBase& RHICmdList, FD3D12Buffer
 			VERIFYD3D12RESULT(Adapter.CreateBuffer(D3D12_HEAP_TYPE_READBACK, Node, Node, Offset + Size, &StagingBuffer, nullptr));
 
 			// Copy the contents of the buffer to the staging buffer.
+			RHICmdList.EnqueueLambda([Node, StagingBuffer, pResource, Buffer, Offset, Size](FRHICommandListBase& ExecutingCmdList)
 			{
-				const auto& pfnCopyContents = [&]()
-				{
-					FD3D12CommandContext& DefaultContext = Device->GetDefaultCommandContext();
-
-					FScopedResourceBarrier ScopeResourceBarrierSource(DefaultContext, pResource, &Buffer->ResourceLocation, D3D12_RESOURCE_STATE_COPY_SOURCE, 0);
-					// Don't need to transition upload heaps
-
-					uint64 SubAllocOffset = Buffer->ResourceLocation.GetOffsetFromBaseOfResource();
-
-					DefaultContext.FlushResourceBarriers();	// Must flush so the desired state is actually set.
-					DefaultContext.GraphicsCommandList()->CopyBufferRegion(
-						StagingBuffer->GetResource(),
-						0,
-						pResource->GetResource(),
-						SubAllocOffset + Offset, Size);
-
-					DefaultContext.UpdateResidency(StagingBuffer);
-					DefaultContext.UpdateResidency(pResource);
-
-					DefaultContext.FlushCommands(ED3D12FlushFlags::WaitForCompletion);
-				};
-
-				if (RHICmdListImmediate.IsTopOfPipe())
-				{
-					// Sync when in the render thread implementation
-					check(IsInRHIThread() == false);
-
-					RHICmdListImmediate.ImmediateFlush(EImmediateFlushType::FlushRHIThread);
-					pfnCopyContents();
-				}
-				else
-				{
-					check(IsInRenderingThread() && !IsRHIThreadRunning());
-					pfnCopyContents();
-				}
-			}
+				FD3D12CommandContext& Context = FD3D12CommandContext::Get(ExecutingCmdList, Node.GetFirstIndex());
+				uint64 SubAllocOffset = Buffer->ResourceLocation.GetOffsetFromBaseOfResource();
+
+				FScopedResourceBarrier ScopeResourceBarrierSource(Context, pResource, &Buffer->ResourceLocation, D3D12_RESOURCE_STATE_COPY_SOURCE, 0);
+				// Don't need to transition upload heaps
+				Context.FlushResourceBarriers(); // Must flush so the desired state is actually set.
+
+				Context.UpdateResidency(StagingBuffer);
+				Context.UpdateResidency(pResource);
+
+				Context.GraphicsCommandList()->CopyBufferRegion(
+					StagingBuffer->GetResource(),
+					0,
+					pResource->GetResource(),
+					SubAllocOffset + Offset, Size);
+			});
+			
+			RHICmdList.GetAsImmediate().SubmitAndBlockUntilGPUIdle();
 
 			LockedData.ResourceLocation.AsStandAlone(StagingBuffer, Size);
 			Data = LockedData.ResourceLocation.GetMappedBaseAddress();
@@ -674,45 +605,53 @@ void FD3D12DynamicRHI::UnlockBuffer(FRHICommandListBase& RHICmdList, FD3D12Buffe
 	{
 		// If the Buffer is dynamic, its upload heap memory can always stay mapped. Don't do anything.
 	}
+	else if (LockedData.bLockedForReadOnly)
+	{
+		// Nothing to do, just release the locked data at the end of the function
+	}
 	else
 	{
-		if (LockedData.bLockedForReadOnly)
-		{
-			//Nothing to do, just release the locked data at the end of the function
-		}
-		else
-		{
-			// Update all of the resources in the LDA chain
-			check(Buffer->IsHeadLink());
-			FD3D12Buffer* LastBuffer = Buffer->GetLinkedObject(Buffer->GetLinkedObjectsGPUMask().GetLastIndex());
+		// Update all of the resources in the LDA chain
+		check(Buffer->IsHeadLink());
 
-			for (FD3D12Buffer::FLinkedObjectIterator CurrentBuffer(Buffer); CurrentBuffer; ++CurrentBuffer)
+		RHICmdList.EnqueueLambda([
+			RootBuffer = Buffer,
+			LockedData = MoveTemp(LockedData)
+		](FRHICommandListBase& ExecutingCmdList)
+		{
+			for (FD3D12Buffer& Buffer : *RootBuffer)
 			{
-				SCOPED_GPU_MASK((FRHIComputeCommandList&)RHICmdList, FRHIGPUMask::FromIndex(CurrentBuffer->GetParentGPUIndex()));
+				FD3D12CommandContext& Context = FD3D12CommandContext::Get(ExecutingCmdList, Buffer.GetParentDevice()->GetGPUIndex());
 
-				// If we are on the render thread, queue up the copy on the RHIThread so it happens at the correct time.
-				if (RHICmdList.IsTopOfPipe())
-				{
-					if (CurrentBuffer.Get() == LastBuffer)
-					{
-						// Command associated with last buffer (will be only buffer if single GPU) receives ownership of locked data
-						ALLOC_COMMAND_CL(RHICmdList, FRHICommandUpdateBuffer)(&CurrentBuffer->ResourceLocation, LockedData.ResourceLocation, LockedData.LockOffset, LockedData.LockSize);
-					}
-					else
-					{
-						// Other commands receive a reference copy of the locked data.  Commands get replayed in order, with the
-						// last command handling clean up the locked data after it has been propagated to all GPUs.
-						FD3D12ResourceLocation NodeResourceLocation(LockedData.ResourceLocation.GetParentDevice());
-						FD3D12ResourceLocation::ReferenceNode(NodeResourceLocation.GetParentDevice(), NodeResourceLocation, LockedData.ResourceLocation);
-						ALLOC_COMMAND_CL(RHICmdList, FRHICommandUpdateBuffer)(&CurrentBuffer->ResourceLocation, NodeResourceLocation, LockedData.LockOffset, LockedData.LockSize);
-					}
-				}
-				else
-				{
-					FD3D12CommandContextBase::Get(RHICmdList).UpdateBuffer(&CurrentBuffer->ResourceLocation, LockedData.LockOffset, &LockedData.ResourceLocation, 0, LockedData.LockSize);
-				}
+				FD3D12Resource* SourceResource = LockedData.ResourceLocation.GetResource();
+				uint32 SourceFullOffset = LockedData.ResourceLocation.GetOffsetFromBaseOfResource();
+
+				FD3D12Resource* DestResource = Buffer.ResourceLocation.GetResource();
+				uint32 DestFullOffset = Buffer.ResourceLocation.GetOffsetFromBaseOfResource() + LockedData.LockOffset;
+
+				// Clear the resource if still bound to make sure the SRVs are rebound again on next operation (and get correct resource transitions enqueued)
+				Context.ConditionalClearShaderResource(&Buffer.ResourceLocation, EShaderParameterTypeMask::SRVMask);
+
+				FScopedResourceBarrier ScopeResourceBarrierDest(Context, DestResource, &Buffer.ResourceLocation, D3D12_RESOURCE_STATE_COPY_DEST, 0);
+				// Don't need to transition upload heaps
+				Context.FlushResourceBarriers();
+
+				Context.UpdateResidency(DestResource);
+				Context.UpdateResidency(SourceResource);
+
+				Context.GraphicsCommandList()->CopyBufferRegion(
+					DestResource->GetResource(),
+					DestFullOffset,
+					SourceResource->GetResource(),
+					SourceFullOffset,
+					LockedData.LockSize
+				);
+				
+				Context.ConditionalSplitCommandList();
+
+				DEBUG_RHI_EXECUTE_COMMAND_LIST(this);
 			}
-		}
+		});
 	}
 
 	LockedData.Reset();
diff --git a/Engine/Source/Runtime/D3D12RHI/Private/D3D12CommandContext.h b/Engine/Source/Runtime/D3D12RHI/Private/D3D12CommandContext.h
--- a/Engine/Source/Runtime/D3D12RHI/Private/D3D12CommandContext.h
+++ b/Engine/Source/Runtime/D3D12RHI/Private/D3D12CommandContext.h
@@ -473,8 +473,6 @@ public:
 		return static_cast<FD3D12CommandContextBase&>(RHICmdList.GetComputeContext().GetLowestLevelContext());
 	}
 
-	virtual void UpdateBuffer(FD3D12ResourceLocation* Dest, uint32 DestOffset, FD3D12ResourceLocation* Source, uint32 SourceOffset, uint32 NumBytes) = 0;
-
 protected:
 	friend class FD3D12CommandContext;
 	virtual FD3D12CommandContext* GetContext(uint32 InGPUIndex) = 0;
@@ -742,8 +740,6 @@ public:
 		ensure(InGPUMask == GPUMask);
 	}
 
-	virtual void UpdateBuffer(FD3D12ResourceLocation* Dest, uint32 DestOffset, FD3D12ResourceLocation* Source, uint32 SourceOffset, uint32 NumBytes) final override;
-
 	inline const TArray<FRHIUniformBuffer*>& GetStaticUniformBuffers() const
 	{
 		return StaticUniformBuffers;
@@ -1013,11 +1009,6 @@ public:
 		ContextRedirect(RHIEndRenderPass());
 	}
 
-	virtual void UpdateBuffer(FD3D12ResourceLocation* Dest, uint32 DestOffset, FD3D12ResourceLocation* Source, uint32 SourceOffset, uint32 NumBytes) final override
-	{
-		ContextRedirect(UpdateBuffer(Dest, DestOffset, Source, SourceOffset, NumBytes));
-	}
-
 #if D3D12_RHI_RAYTRACING
 	virtual void RHIBuildAccelerationStructures(TConstArrayView<FRayTracingGeometryBuildParams> Params, const FRHIBufferRange& ScratchBufferRange) final override
 	{
diff --git a/Engine/Source/Runtime/D3D12RHI/Private/D3D12Commands.cpp b/Engine/Source/Runtime/D3D12RHI/Private/D3D12Commands.cpp
--- a/Engine/Source/Runtime/D3D12RHI/Private/D3D12Commands.cpp
+++ b/Engine/Source/Runtime/D3D12RHI/Private/D3D12Commands.cpp
@@ -1948,27 +1948,3 @@ void FD3D12CommandContext::SetShadingRateImage(FD3D12Resource* RateImageTexture)
 	}
  #endif
 }
-
-void FD3D12CommandContext::UpdateBuffer(FD3D12ResourceLocation* Dest, uint32 DestOffset, FD3D12ResourceLocation* Source, uint32 SourceOffset, uint32 NumBytes)
-{
-	FD3D12Resource* SourceResource = Source->GetResource();
-	uint32 SourceFullOffset = Source->GetOffsetFromBaseOfResource() + SourceOffset;
-
-	FD3D12Resource* DestResource = Dest->GetResource();
-	uint32 DestFullOffset = Dest->GetOffsetFromBaseOfResource() + DestOffset;
-
-	// Clear the resource if still bound to make sure the SRVs are rebound again on next operation (and get correct resource transitions enqueued)
-	ConditionalClearShaderResource(Dest, EShaderParameterTypeMask::SRVMask);
-
-	FScopedResourceBarrier ScopeResourceBarrierDest(*this, DestResource, Dest, D3D12_RESOURCE_STATE_COPY_DEST, 0);
-	// Don't need to transition upload heaps
-
-	FlushResourceBarriers();
-	GraphicsCommandList()->CopyBufferRegion(DestResource->GetResource(), DestFullOffset, SourceResource->GetResource(), SourceFullOffset, NumBytes);
-	UpdateResidency(DestResource);
-	UpdateResidency(SourceResource);
-
-	ConditionalSplitCommandList();
-
-	DEBUG_RHI_EXECUTE_COMMAND_LIST(this);
-}
diff --git a/Engine/Source/Runtime/D3D12RHI/Private/D3D12RHIPrivate.h b/Engine/Source/Runtime/D3D12RHI/Private/D3D12RHIPrivate.h
--- a/Engine/Source/Runtime/D3D12RHI/Private/D3D12RHIPrivate.h
+++ b/Engine/Source/Runtime/D3D12RHI/Private/D3D12RHIPrivate.h
@@ -494,7 +494,7 @@ public:
 	virtual void DestroyVirtualTexture(ETextureCreateFlags InFlags, void* RawTextureMemory, FPlatformMemory::FPlatformVirtualMemoryBlock& RawTextureBlock, uint64 CommittedTextureSize) = 0;
 #endif
 	virtual bool HandleSpecialLock(void*& MemoryOut, uint32 MipIndex, uint32 ArrayIndex, FD3D12Texture* InTexture, EResourceLockMode LockMode, uint32& DestStride, uint64* OutLockedByteCount = nullptr) { return false; }
-	virtual bool HandleSpecialUnlock(FRHICommandListBase* RHICmdList, uint32 MipIndex, FD3D12Texture* InTexture) { return false; }
+	virtual bool HandleSpecialUnlock(FRHICommandListBase& RHICmdList, uint32 MipIndex, FD3D12Texture* InTexture) { return false; }
 
 	FD3D12Adapter& GetAdapter(uint32_t Index = 0) { return *ChosenAdapters[Index]; }
 	const FD3D12Adapter& GetAdapter(uint32_t Index = 0) const { return *ChosenAdapters[Index]; }
diff --git a/Engine/Source/Runtime/D3D12RHI/Private/D3D12Texture.cpp b/Engine/Source/Runtime/D3D12RHI/Private/D3D12Texture.cpp
--- a/Engine/Source/Runtime/D3D12RHI/Private/D3D12Texture.cpp
+++ b/Engine/Source/Runtime/D3D12RHI/Private/D3D12Texture.cpp
@@ -24,14 +24,6 @@ static FAutoConsoleVariableRef CVarAdjustTexturePoolSizeBasedOnBudget(
 	TEXT("Indicates if the RHI should lower the texture pool size when the application is over the memory budget provided by the OS. This can result in lower quality textures (but hopefully improve performance).")
 	);
 
-static TAutoConsoleVariable<int32> CVarD3D12Texture2DRHIFlush(
-	TEXT("D3D12.LockTexture2DRHIFlush"),
-	0,
-	TEXT("If enabled, we do RHIThread flush on LockTexture2D. Likely not required on any platform, but keeping just for testing for now")
-	TEXT(" 0: off (default)\n")
-	TEXT(" 1: on"),
-	ECVF_RenderThreadSafe);
-
 static TAutoConsoleVariable<int32> CVarUseUpdateTexture3DComputeShader(
 	TEXT("D3D12.UseUpdateTexture3DComputeShader"),
 	0,
@@ -53,228 +45,6 @@ static FAutoConsoleVariableRef CVarTexturePoolOnlyAccountStreamableTexture(
 
 extern int32 GD3D12BindResourceLabels;
 
-///////////////////////////////////////////////////////////////////////////////////////////
-// Texture Commands
-///////////////////////////////////////////////////////////////////////////////////////////
-
-static bool ShouldDeferCmdListOperation(FRHICommandListBase* RHICmdList)
-{
-	if (RHICmdList == nullptr)
-	{
-		return false;
-	}
-
-	if (RHICmdList->Bypass() || !IsRunningRHIInSeparateThread())
-	{
-		return false;
-	}
-
-	return true;
-}
-
-struct FRHICommandUpdateTextureString
-{
-	static const TCHAR* TStr() { return TEXT("FRHICommandUpdateTexture"); }
-};
-struct FRHICommandUpdateTexture final : public FRHICommand<FRHICommandUpdateTexture, FRHICommandUpdateTextureString>
-{
-	FD3D12Texture* Texture;
-	uint32 MipIndex;
-	uint32 DestX;
-	uint32 DestY;
-	uint32 DestZ;
-	D3D12_TEXTURE_COPY_LOCATION SourceCopyLocation;
-	FD3D12ResourceLocation Source;
-
-	FORCEINLINE_DEBUGGABLE FRHICommandUpdateTexture(FD3D12Texture* InTexture,
-		uint32 InMipIndex, uint32 InDestX, uint32 InDestY, uint32 InDestZ,
-		const D3D12_TEXTURE_COPY_LOCATION& InSourceCopyLocation, FD3D12ResourceLocation* InSource)
-		: Texture(InTexture)
-		, MipIndex(InMipIndex)
-		, DestX(InDestX)
-		, DestY(InDestY)
-		, DestZ(InDestZ)
-		, SourceCopyLocation(InSourceCopyLocation)
-		, Source(nullptr)
-	{
-		if (InSource)
-		{
-			FD3D12ResourceLocation::TransferOwnership(Source, *InSource);
-		}
-	}
-
-	~FRHICommandUpdateTexture()
-	{
-	}
-
-	void Execute(FRHICommandListBase& CmdList)
-	{
-		Texture->UpdateTexture(MipIndex, DestX, DestY, DestZ, SourceCopyLocation);
-	}
-};
-
-struct FD3D12RHICommandInitializeTextureString
-{
-	static const TCHAR* TStr() { return TEXT("FD3D12RHICommandInitializeTexture"); }
-};
-struct FD3D12RHICommandInitializeTexture final : public FRHICommand<FD3D12RHICommandInitializeTexture, FD3D12RHICommandInitializeTextureString>
-{
-	FD3D12Texture* Texture;
-	FD3D12ResourceLocation SrcResourceLoc;
-	uint32 NumSubresources;
-	D3D12_RESOURCE_STATES DestinationState;
-
-	FORCEINLINE_DEBUGGABLE FD3D12RHICommandInitializeTexture(FD3D12Texture* InTexture, FD3D12ResourceLocation& InSrcResourceLoc, uint32 InNumSubresources, D3D12_RESOURCE_STATES InDestinationState)
-		: Texture(InTexture)
-		, SrcResourceLoc(InSrcResourceLoc.GetParentDevice())
-		, NumSubresources(InNumSubresources)
-		, DestinationState(InDestinationState)
-	{
-		FD3D12ResourceLocation::TransferOwnership(SrcResourceLoc, InSrcResourceLoc);
-	}
-
-	void Execute(FRHICommandListBase& RHICmdList)
-	{
-		size_t MemSize = NumSubresources * (sizeof(D3D12_PLACED_SUBRESOURCE_FOOTPRINT) + sizeof(UINT) + sizeof(UINT64));
-		const bool bAllocateOnStack = (MemSize < 4096);
-		void* Mem = bAllocateOnStack? FMemory_Alloca(MemSize) : FMemory::Malloc(MemSize);
-
-		D3D12_PLACED_SUBRESOURCE_FOOTPRINT* Footprints = (D3D12_PLACED_SUBRESOURCE_FOOTPRINT*) Mem;
-		check(Footprints);
-		UINT* Rows = (UINT*) (Footprints + NumSubresources);
-		UINT64* RowSizeInBytes = (UINT64*) (Rows + NumSubresources);
-
-		uint64 Size = 0;
-		const D3D12_RESOURCE_DESC& Desc = Texture->GetResource()->GetDesc();
-		Texture->GetParentDevice()->GetDevice()->GetCopyableFootprints(&Desc, 0, NumSubresources, SrcResourceLoc.GetOffsetFromBaseOfResource(), Footprints, Rows, RowSizeInBytes, &Size);
-
-		D3D12_TEXTURE_COPY_LOCATION Src;
-		Src.pResource = SrcResourceLoc.GetResource()->GetResource();
-		Src.Type = D3D12_TEXTURE_COPY_TYPE_PLACED_FOOTPRINT;
-
-		// Initialize all the textures in the chain
-		for (FD3D12Texture& CurrentTexture : *Texture)
-		{
-			FD3D12Device* Device = CurrentTexture.GetParentDevice();
-			FD3D12Resource* Resource = CurrentTexture.GetResource();
-			FD3D12CommandContext& Context = FD3D12CommandContext::Get(RHICmdList, Device->GetGPUIndex());
-
-			// resource should be in copy dest already, because it's created like that, so no transition required here
-
-			D3D12_TEXTURE_COPY_LOCATION Dst;
-			Dst.pResource = Resource->GetResource();
-			Dst.Type = D3D12_TEXTURE_COPY_TYPE_SUBRESOURCE_INDEX;
-
-			for (uint32 Subresource = 0; Subresource < NumSubresources; Subresource++)
-			{
-				Dst.SubresourceIndex = Subresource;
-				Src.PlacedFootprint = Footprints[Subresource];
-				Context.GraphicsCommandList()->CopyTextureRegion(&Dst, 0, 0, 0, &Src, nullptr);
-			}
-
-			// Update the resource state after the copy has been done (will take care of updating the residency as well)
-			Context.AddTransitionBarrier(Resource, D3D12_RESOURCE_STATE_COPY_DEST, DestinationState, D3D12_RESOURCE_BARRIER_ALL_SUBRESOURCES);
-			Context.ConditionalSplitCommandList();
-
-			// Texture is now written and ready, so unlock the block (locked after creation and can be defragmented if needed)
-			CurrentTexture.ResourceLocation.UnlockPoolData();
-
-			// If the resource is untracked, the destination state must match the default state of the resource.
-			check(Resource->RequiresResourceStateTracking() || (Resource->GetDefaultResourceState() == DestinationState));
-		}
-
-		if (!bAllocateOnStack)
-		{
-			FMemory::Free(Mem);
-		}
-	}
-};
-
-struct FD3D12RHICommandAsyncReallocateTexture2DString
-{
-	static const TCHAR* TStr() { return TEXT("FD3D12RHICommandAsyncReallocateTexture2D"); }
-};
-struct FRHICommandD3D12AsyncReallocateTexture2D final : public FRHICommand<FRHICommandD3D12AsyncReallocateTexture2D, FD3D12RHICommandAsyncReallocateTexture2DString>
-{
-	FD3D12Texture* OldTexture;
-	FD3D12Texture* NewTexture;
-	int32 NewMipCount;
-	int32 NewSizeX;
-	int32 NewSizeY;
-	FThreadSafeCounter* RequestStatus;
-
-	FORCEINLINE_DEBUGGABLE FRHICommandD3D12AsyncReallocateTexture2D(FD3D12Texture* InOldTexture, FD3D12Texture* InNewTexture, int32 InNewMipCount, int32 InNewSizeX, int32 InNewSizeY, FThreadSafeCounter* InRequestStatus)
-		: OldTexture(InOldTexture)
-		, NewTexture(InNewTexture)
-		, NewMipCount(InNewMipCount)
-		, NewSizeX(InNewSizeX)
-		, NewSizeY(InNewSizeY)
-		, RequestStatus(InRequestStatus)
-	{
-	}
-
-	void Execute(FRHICommandListBase& RHICmdList)
-	{
-		CopyMips();
-	}
-
-	void CopyMips()
-	{
-		// Use the GPU to asynchronously copy the old mip-maps into the new texture.
-		const uint32 NumSharedMips = FMath::Min(OldTexture->GetNumMips(), NewTexture->GetNumMips());
-		const uint32 SourceMipOffset = OldTexture->GetNumMips() - NumSharedMips;
-		const uint32 DestMipOffset = NewTexture->GetNumMips() - NumSharedMips;
-
-		uint32 destSubresource = 0;
-		uint32 srcSubresource = 0;
-
-		for (FD3D12Texture::FDualLinkedObjectIterator It(OldTexture, NewTexture); It; ++It)
-		{
-			OldTexture = static_cast<FD3D12Texture*>(It.GetFirst());
-			NewTexture = static_cast<FD3D12Texture*>(It.GetSecond());
-
-			FD3D12Device* Device = OldTexture->GetParentDevice();
-
-			FD3D12CommandContext& Context = Device->GetDefaultCommandContext();
-
-			FScopedResourceBarrier ScopeResourceBarrierDst(Context, NewTexture->GetResource(), &NewTexture->ResourceLocation, D3D12_RESOURCE_STATE_COPY_DEST  , D3D12_RESOURCE_BARRIER_ALL_SUBRESOURCES);
-			FScopedResourceBarrier ScopeResourceBarrierSrc(Context, OldTexture->GetResource(), &OldTexture->ResourceLocation, D3D12_RESOURCE_STATE_COPY_SOURCE, D3D12_RESOURCE_BARRIER_ALL_SUBRESOURCES);
-			Context.FlushResourceBarriers();	// Must flush so the desired state is actually set.
-
-			for (uint32 MipIndex = 0; MipIndex < NumSharedMips; ++MipIndex)
-			{
-				// Use the GPU to copy between mip-maps.
-				// This is serialized with other D3D commands, so it isn't necessary to increment Counter to signal a pending asynchronous copy.
-
-				srcSubresource = CalcSubresource(MipIndex + SourceMipOffset, 0, OldTexture->GetNumMips());
-				destSubresource = CalcSubresource(MipIndex + DestMipOffset, 0, NewTexture->GetNumMips());
-
-				CD3DX12_TEXTURE_COPY_LOCATION DestCopyLocation(NewTexture->GetResource()->GetResource(), destSubresource);
-				CD3DX12_TEXTURE_COPY_LOCATION SourceCopyLocation(OldTexture->GetResource()->GetResource(), srcSubresource);
-
-				Context.GraphicsCommandList()->CopyTextureRegion(
-					&DestCopyLocation,
-					0, 0, 0,
-					&SourceCopyLocation,
-					nullptr);
-
-				Context.UpdateResidency(NewTexture->GetResource());
-				Context.UpdateResidency(OldTexture->GetResource());
-
-				Context.ConditionalSplitCommandList();
-
-				DEBUG_EXECUTE_COMMAND_CONTEXT(Context);
-			}
-		}
-
-		// Decrement the thread-safe counter used to track the completion of the reallocation, since D3D handles sequencing the
-		// async mip copies with other D3D calls.
-		RequestStatus->Decrement();
-	}
-};
-
-
-
 ///////////////////////////////////////////////////////////////////////////////////////////
 // Texture Stats
 ///////////////////////////////////////////////////////////////////////////////////////////
@@ -1326,16 +1096,11 @@ uint32 FD3D12DynamicRHI::RHIComputeMemorySize(FRHITexture* TextureRHI)
 }
 
 
-FTextureRHIRef FD3D12DynamicRHI::AsyncReallocateTexture2D_RenderThread(class FRHICommandListImmediate& RHICmdList, FRHITexture* Texture2DRHI, int32 NewMipCount, int32 NewSizeX, int32 NewSizeY, FThreadSafeCounter* RequestStatus)
+FTextureRHIRef FD3D12DynamicRHI::AsyncReallocateTexture2D_RenderThread(FRHICommandListImmediate& RHICmdList, FRHITexture* Texture2DRHI, int32 NewMipCount, int32 NewSizeX, int32 NewSizeY, FThreadSafeCounter* RequestStatus)
 {
-	if (RHICmdList.Bypass())
-	{
-		return FDynamicRHI::AsyncReallocateTexture2D_RenderThread(RHICmdList, Texture2DRHI, NewMipCount, NewSizeX, NewSizeY, RequestStatus);
-	}
-
-	FD3D12Texture* Texture2D = FD3D12DynamicRHI::ResourceCast(Texture2DRHI);
+	FD3D12Texture* OldTexture = FD3D12DynamicRHI::ResourceCast(Texture2DRHI);
 	
-	FRHITextureDesc Desc = Texture2D->GetDesc();
+	FRHITextureDesc Desc = OldTexture->GetDesc();
 	Desc.Extent = FIntPoint(NewSizeX, NewSizeY);
 	Desc.NumMips = NewMipCount;
 
@@ -1346,13 +1111,64 @@ FTextureRHIRef FD3D12DynamicRHI::AsyncReallocateTexture2D_RenderThread(class FRH
 	);
 
 	// Allocate a new texture.
-	FRHICommandListImmediate* RHIImmediateCmdList = nullptr;
 	ID3D12ResourceAllocator* ResourceAllocator = nullptr;
-	FD3D12Texture* NewTexture2D = CreateD3D12Texture(CreateDesc, RHIImmediateCmdList, ResourceAllocator);
-		
-	ALLOC_COMMAND_CL(RHICmdList, FRHICommandD3D12AsyncReallocateTexture2D)(Texture2D, NewTexture2D, NewMipCount, NewSizeX, NewSizeY, RequestStatus);
+	FD3D12Texture* NewTexture = CreateD3D12Texture(CreateDesc, &RHICmdList, ResourceAllocator);
+
+	RHICmdList.EnqueueLambda([
+		RootOldTexture = OldTexture,
+		RootNewTexture = NewTexture,
+		RequestStatus
+	](FRHICommandListBase& ExecutingCmdList)
+	{
+		// Use the GPU to asynchronously copy the old mip-maps into the new texture.
+		const uint32 NumSharedMips   = FMath::Min(RootOldTexture->GetNumMips(), RootNewTexture->GetNumMips());
+		const uint32 SourceMipOffset = RootOldTexture->GetNumMips() - NumSharedMips;
+		const uint32 DestMipOffset   = RootNewTexture->GetNumMips() - NumSharedMips;
 
-	return NewTexture2D;
+		for (FD3D12Texture::FDualLinkedObjectIterator It(RootOldTexture, RootNewTexture); It; ++It)
+		{
+			FD3D12Texture* DeviceOldTexture = It.GetFirst();
+			FD3D12Texture* DeviceNewTexture = It.GetSecond();
+			check(DeviceOldTexture->GetParentDevice() == DeviceNewTexture->GetParentDevice());
+
+			FD3D12CommandContext& Context = FD3D12CommandContext::Get(ExecutingCmdList, DeviceOldTexture->GetParentDevice()->GetGPUIndex());
+
+			FScopedResourceBarrier ScopeResourceBarrierDst(Context, DeviceNewTexture->GetResource(), &DeviceNewTexture->ResourceLocation, D3D12_RESOURCE_STATE_COPY_DEST, D3D12_RESOURCE_BARRIER_ALL_SUBRESOURCES);
+			FScopedResourceBarrier ScopeResourceBarrierSrc(Context, DeviceOldTexture->GetResource(), &DeviceOldTexture->ResourceLocation, D3D12_RESOURCE_STATE_COPY_SOURCE, D3D12_RESOURCE_BARRIER_ALL_SUBRESOURCES);
+			Context.FlushResourceBarriers();	// Must flush so the desired state is actually set.
+
+			for (uint32 MipIndex = 0; MipIndex < NumSharedMips; ++MipIndex)
+			{
+				// Use the GPU to copy between mip-maps.
+				// This is serialized with other D3D commands, so it isn't necessary to increment Counter to signal a pending asynchronous copy.
+
+				uint32 SrcSubresource = CalcSubresource(MipIndex + SourceMipOffset, 0, DeviceOldTexture->GetNumMips());
+				uint32 DstSubresource = CalcSubresource(MipIndex + DestMipOffset, 0, DeviceNewTexture->GetNumMips());
+
+				CD3DX12_TEXTURE_COPY_LOCATION DestCopyLocation(DeviceNewTexture->GetResource()->GetResource(), DstSubresource);
+				CD3DX12_TEXTURE_COPY_LOCATION SourceCopyLocation(DeviceOldTexture->GetResource()->GetResource(), SrcSubresource);
+
+				Context.GraphicsCommandList()->CopyTextureRegion(
+					&DestCopyLocation,
+					0, 0, 0,
+					&SourceCopyLocation,
+					nullptr);
+
+				Context.UpdateResidency(DeviceNewTexture->GetResource());
+				Context.UpdateResidency(DeviceOldTexture->GetResource());
+
+				Context.ConditionalSplitCommandList();
+
+				DEBUG_EXECUTE_COMMAND_CONTEXT(Context);
+			}
+		}
+
+		// Decrement the thread-safe counter used to track the completion of the reallocation, since D3D handles sequencing the
+		// async mip copies with other D3D calls.
+		RequestStatus->Decrement();
+	});
+
+	return NewTexture;
 }
 
 
@@ -1373,27 +1189,8 @@ FTextureRHIRef FD3D12DynamicRHI::AsyncReallocateTexture2D_RenderThread(class FRH
  */
 FTextureRHIRef FD3D12DynamicRHI::RHIAsyncReallocateTexture2D(FRHITexture* Texture2DRHI, int32 NewMipCount, int32 NewSizeX, int32 NewSizeY, FThreadSafeCounter* RequestStatus)
 {
-	FD3D12Texture* Texture2D = FD3D12DynamicRHI::ResourceCast(Texture2DRHI);
-
-	FRHITextureDesc Desc = Texture2D->GetDesc();
-	Desc.Extent = FIntPoint(NewSizeX, NewSizeY);
-	Desc.NumMips = NewMipCount;
-
-	FRHITextureCreateDesc CreateDesc(
-		Desc,
-		RHIGetDefaultResourceState(Desc.Flags, false),
-		TEXT("RHIAsyncReallocateTexture2D")
-	);
-
-	// Allocate a new texture.
-	FRHICommandListImmediate* RHIImmediateCmdList = nullptr;
-	ID3D12ResourceAllocator* ResourceAllocator = nullptr;
-	FD3D12Texture* NewTexture2D = CreateD3D12Texture(CreateDesc, nullptr);
-	
-	FRHICommandD3D12AsyncReallocateTexture2D AsyncReallocateTexture2D(Texture2D, NewTexture2D, NewMipCount, NewSizeX, NewSizeY, RequestStatus);
-	AsyncReallocateTexture2D.CopyMips();
-
-	return NewTexture2D;
+	UE_LOG(LogD3D12RHI, Fatal, TEXT("RHIAsyncReallocateTexture2D should not be called. AsyncReallocateTexture2D_RenderThread is where this is implemented."));
+	return nullptr;
 }
 
 ETextureReallocationStatus FD3D12DynamicRHI::RHIFinalizeAsyncReallocateTexture2D(FRHITexture* Texture2D, bool bBlockUntilCompleted)
@@ -1739,7 +1536,7 @@ void FD3D12Texture::AliasResources(FD3D12Texture* Texture)
 	}
 }
 
-void* FD3D12Texture::Lock(class FRHICommandListImmediate* RHICmdList, uint32 MipIndex, uint32 ArrayIndex, EResourceLockMode LockMode, uint32& DestStride, uint64* OutLockedByteCount)
+void* FD3D12Texture::Lock(FRHICommandListImmediate& RHICmdList, uint32 MipIndex, uint32 ArrayIndex, EResourceLockMode LockMode, uint32& DestStride, uint64* OutLockedByteCount)
 {
 	SCOPE_CYCLE_COUNTER(STAT_D3D12LockTextureTime);
 
@@ -1753,7 +1550,7 @@ void* FD3D12Texture::Lock(class FRHICommandListImmediate* RHICmdList, uint32 Mip
 	const uint32 Subresource = CalcSubresource(MipIndex, ArrayIndex, this->GetNumMips());
 
 	check(LockedMap.Find(Subresource) == nullptr);
-	FD3D12LockedResource* LockedResource = new FD3D12LockedResource(Device);
+	TUniquePtr<FD3D12LockedResource> LockedResource = MakeUnique<FD3D12LockedResource>(Device);
 
 	const D3D12_RESOURCE_DESC& ResourceDesc = GetResource()->GetDesc();
 
@@ -1774,8 +1571,6 @@ void* FD3D12Texture::Lock(class FRHICommandListImmediate* RHICmdList, uint32 Mip
 		*OutLockedByteCount = SubresourceSize;
 	}
 
-	FD3D12CommandContext& Context = Device->GetDefaultCommandContext();
-
 #if !PLATFORM_SUPPORTS_VIRTUAL_TEXTURES
 	void* RawTextureMemory = (void*)ResourceLocation.GetGPUVirtualAddress();
 #endif
@@ -1787,8 +1582,7 @@ void* FD3D12Texture::Lock(class FRHICommandListImmediate* RHICmdList, uint32 Mip
 		// nothing left to do...
 		check(Data != nullptr);
 	}
-	else
-	if (LockMode == RLM_WriteOnly)
+	else if (LockMode == RLM_WriteOnly)
 	{
 		// If we're writing to the texture, allocate a system memory buffer to receive the new contents.
 		// Use an upload heap to copy data to a default resource.
@@ -1823,8 +1617,9 @@ void* FD3D12Texture::Lock(class FRHICommandListImmediate* RHICmdList, uint32 Mip
 		CD3DX12_TEXTURE_COPY_LOCATION DestCopyLocation(StagingTexture->GetResource(), PlacedFootprint);
 		CD3DX12_TEXTURE_COPY_LOCATION SourceCopyLocation(GetResource()->GetResource(), Subresource);
 
-		const auto& pfnCopyTextureRegion = [&]()
+		RHICmdList.EnqueueLambda([this, DestCopyLocation, SourceCopyLocation](FRHICommandListImmediate& ExecutingCmdList)
 		{
+			FD3D12CommandContext& Context = FD3D12CommandContext::Get(ExecutingCmdList, 0);
 			FScopedResourceBarrier ScopeResourceBarrierSource(Context, GetResource(), &ResourceLocation, D3D12_RESOURCE_STATE_COPY_SOURCE, SourceCopyLocation.SubresourceIndex);
 
 			Context.FlushResourceBarriers();
@@ -1835,77 +1630,60 @@ void* FD3D12Texture::Lock(class FRHICommandListImmediate* RHICmdList, uint32 Mip
 				nullptr);
 
 			Context.UpdateResidency(GetResource());
-		};
-
-		if (RHICmdList != nullptr)
-		{
-			check(IsInRHIThread() == false);
-
-			RHICmdList->ImmediateFlush(EImmediateFlushType::FlushRHIThread);
-			pfnCopyTextureRegion();
-		}
-		else
-		{
-			check(IsInRHIThread());
-
-			pfnCopyTextureRegion();
-		}
+		});
 
 		// We need to execute the command list so we can read the data from the map below
-		Context.FlushCommands(ED3D12FlushFlags::WaitForCompletion);
+		RHICmdList.SubmitAndBlockUntilGPUIdle();
 
 		Data = LockedResource->ResourceLocation.GetMappedBaseAddress();
 	}
 
-	LockedMap.Add(Subresource, LockedResource);
+	LockedMap.Emplace(Subresource, MoveTemp(LockedResource));
 
 	check(Data != nullptr);
 	return Data;
 }
 
-void FD3D12Texture::UpdateTexture(uint32 MipIndex, uint32 DestX, uint32 DestY, uint32 DestZ, const D3D12_TEXTURE_COPY_LOCATION& SourceCopyLocation)
+void FD3D12Texture::UpdateTexture(FD3D12CommandContext& Context, uint32 MipIndex, uint32 DestX, uint32 DestY, uint32 DestZ, const D3D12_TEXTURE_COPY_LOCATION& SourceCopyLocation)
 {
 	LLM_SCOPE_BYNAME(TEXT("D3D12CopyTextureRegion"));
-	FD3D12CommandContext& DefaultContext = GetParentDevice()->GetDefaultCommandContext();
 
-	FScopedResourceBarrier ScopeResourceBarrierDest(DefaultContext, GetResource(), &ResourceLocation, D3D12_RESOURCE_STATE_COPY_DEST, MipIndex);
+	FScopedResourceBarrier ScopeResourceBarrierDest(Context, GetResource(), &ResourceLocation, D3D12_RESOURCE_STATE_COPY_DEST, MipIndex);
 	// Don't need to transition upload heaps
 
 	CD3DX12_TEXTURE_COPY_LOCATION DestCopyLocation(GetResource()->GetResource(), MipIndex);
 
-	DefaultContext.FlushResourceBarriers();
-	DefaultContext.GraphicsCommandList()->CopyTextureRegion(
+	Context.FlushResourceBarriers();
+	Context.GraphicsCommandList()->CopyTextureRegion(
 		&DestCopyLocation,
 		DestX, DestY, DestZ,
 		&SourceCopyLocation,
 		nullptr);
 
-	DefaultContext.UpdateResidency(GetResource());
+	Context.UpdateResidency(GetResource());
 	
-	DefaultContext.ConditionalSplitCommandList();
+	Context.ConditionalSplitCommandList();
 
-	DEBUG_EXECUTE_COMMAND_CONTEXT(DefaultContext);
+	DEBUG_EXECUTE_COMMAND_CONTEXT(Context);
 }
 
-void FD3D12Texture::CopyTextureRegion(uint32 DestX, uint32 DestY, uint32 DestZ, FD3D12Texture* SourceTexture, const D3D12_BOX& SourceBox)
+void FD3D12Texture::CopyTextureRegion(FD3D12CommandContext& Context, uint32 DestX, uint32 DestY, uint32 DestZ, FD3D12Texture* SourceTexture, const D3D12_BOX& SourceBox)
 {
-	FD3D12CommandContext& DefaultContext = GetParentDevice()->GetDefaultCommandContext();
-
 	CD3DX12_TEXTURE_COPY_LOCATION DestCopyLocation(GetResource()->GetResource(), 0);
 	CD3DX12_TEXTURE_COPY_LOCATION SourceCopyLocation(SourceTexture->GetResource()->GetResource(), 0);
 
-	FScopedResourceBarrier ConditionalScopeResourceBarrierDst(DefaultContext, GetResource(), &ResourceLocation, D3D12_RESOURCE_STATE_COPY_DEST, DestCopyLocation.SubresourceIndex);
-	FScopedResourceBarrier ConditionalScopeResourceBarrierSrc(DefaultContext, SourceTexture->GetResource(), &SourceTexture->ResourceLocation, D3D12_RESOURCE_STATE_COPY_SOURCE, SourceCopyLocation.SubresourceIndex);
+	FScopedResourceBarrier ConditionalScopeResourceBarrierDst(Context, GetResource(), &ResourceLocation, D3D12_RESOURCE_STATE_COPY_DEST, DestCopyLocation.SubresourceIndex);
+	FScopedResourceBarrier ConditionalScopeResourceBarrierSrc(Context, SourceTexture->GetResource(), &SourceTexture->ResourceLocation, D3D12_RESOURCE_STATE_COPY_SOURCE, SourceCopyLocation.SubresourceIndex);
 
-	DefaultContext.FlushResourceBarriers();
-	DefaultContext.GraphicsCommandList()->CopyTextureRegion(
+	Context.FlushResourceBarriers();
+	Context.GraphicsCommandList()->CopyTextureRegion(
 		&DestCopyLocation,
 		DestX, DestY, DestZ,
 		&SourceCopyLocation,
 		&SourceBox);
 
-	DefaultContext.UpdateResidency(SourceTexture->GetResource());
-	DefaultContext.UpdateResidency(GetResource());
+	Context.UpdateResidency(SourceTexture->GetResource());
+	Context.UpdateResidency(GetResource());
 }
 
 void FD3D12Texture::InitializeTextureData(FRHICommandListBase& RHICmdList, const FRHITextureCreateDesc& CreateDesc, D3D12_RESOURCE_STATES DestinationState)
@@ -1922,7 +1700,14 @@ void FD3D12Texture::InitializeTextureData(FRHICommandListBase& RHICmdList, const
 
 	size_t MemSize = NumSubresources * (sizeof(D3D12_PLACED_SUBRESOURCE_FOOTPRINT) + sizeof(UINT) + sizeof(UINT64));
 	const bool bAllocateOnStack = (MemSize < 4096);
-	void* Mem = bAllocateOnStack? FMemory_Alloca(MemSize) : FMemory::Malloc(MemSize);
+	void* Mem = bAllocateOnStack ? FMemory_Alloca(MemSize) : FMemory::Malloc(MemSize);
+	ON_SCOPE_EXIT
+	{
+		if (!bAllocateOnStack)
+		{
+			FMemory::Free(Mem);
+		}
+	};
 
 	D3D12_PLACED_SUBRESOURCE_FOOTPRINT* Footprints = (D3D12_PLACED_SUBRESOURCE_FOOTPRINT*) Mem;
 	check(Footprints);
@@ -1967,84 +1752,103 @@ void FD3D12Texture::InitializeTextureData(FRHICommandListBase& RHICmdList, const
 
 	check(SrcData == (uint8*)CreateDesc.BulkData->GetResourceBulkData() + CreateDesc.BulkData->GetResourceBulkDataSize());
 
-	if (RHICmdList.IsTopOfPipe())
-	{
-		ALLOC_COMMAND_CL(RHICmdList, FD3D12RHICommandInitializeTexture)(this, SrcResourceLoc, NumSubresources, DestinationState);
-	}
-	else
+	RHICmdList.EnqueueLambda([Texture = this, SrcResourceLoc = MoveTemp(SrcResourceLoc), NumSubresources, DestinationState](FRHICommandListBase& ExecutingCmdList)
 	{
-		FD3D12RHICommandInitializeTexture Command(this, SrcResourceLoc, NumSubresources, DestinationState);
-		Command.Execute(RHICmdList);
-	}
+		size_t MemSize = NumSubresources * (sizeof(D3D12_PLACED_SUBRESOURCE_FOOTPRINT) + sizeof(UINT) + sizeof(UINT64));
+		const bool bAllocateOnStack = (MemSize < 4096);
+		void* Mem = bAllocateOnStack ? FMemory_Alloca(MemSize) : FMemory::Malloc(MemSize);
+		ON_SCOPE_EXIT
+		{
+			if (!bAllocateOnStack)
+			{
+				FMemory::Free(Mem);
+			}
+		};
 
-	if (!bAllocateOnStack)
-	{
-		FMemory::Free(Mem);
-	}
-}
+		D3D12_PLACED_SUBRESOURCE_FOOTPRINT* Footprints = (D3D12_PLACED_SUBRESOURCE_FOOTPRINT*)Mem;
+		check(Footprints);
+		UINT* Rows = (UINT*)(Footprints + NumSubresources);
+		UINT64* RowSizeInBytes = (UINT64*)(Rows + NumSubresources);
 
-void FD3D12Texture::Unlock(class FRHICommandListImmediate* RHICmdList, uint32 MipIndex, uint32 ArrayIndex)
-{
-	SCOPE_CYCLE_COUNTER(STAT_D3D12UnlockTextureTime);
+		uint64 Size = 0;
+		const D3D12_RESOURCE_DESC& Desc = Texture->GetResource()->GetDesc();
+		Texture->GetParentDevice()->GetDevice()->GetCopyableFootprints(&Desc, 0, NumSubresources, SrcResourceLoc.GetOffsetFromBaseOfResource(), Footprints, Rows, RowSizeInBytes, &Size);
+
+		D3D12_TEXTURE_COPY_LOCATION Src;
+		Src.pResource = SrcResourceLoc.GetResource()->GetResource();
+		Src.Type = D3D12_TEXTURE_COPY_TYPE_PLACED_FOOTPRINT;
+
+		// Initialize all the textures in the chain
+		for (FD3D12Texture& CurrentTexture : *Texture)
+		{
+			FD3D12Device* Device = CurrentTexture.GetParentDevice();
+			FD3D12Resource* Resource = CurrentTexture.GetResource();
+			FD3D12CommandContext& Context = FD3D12CommandContext::Get(ExecutingCmdList, Device->GetGPUIndex());
+
+			// resource should be in copy dest already, because it's created like that, so no transition required here
 
-	UnlockInternal(RHICmdList, ++FLinkedObjectIterator(this), MipIndex, ArrayIndex);
+			D3D12_TEXTURE_COPY_LOCATION Dst;
+			Dst.pResource = Resource->GetResource();
+			Dst.Type = D3D12_TEXTURE_COPY_TYPE_SUBRESOURCE_INDEX;
+
+			for (uint32 Subresource = 0; Subresource < NumSubresources; Subresource++)
+			{
+				Dst.SubresourceIndex = Subresource;
+				Src.PlacedFootprint = Footprints[Subresource];
+				Context.GraphicsCommandList()->CopyTextureRegion(&Dst, 0, 0, 0, &Src, nullptr);
+			}
+
+			// Update the resource state after the copy has been done (will take care of updating the residency as well)
+			Context.AddTransitionBarrier(Resource, D3D12_RESOURCE_STATE_COPY_DEST, DestinationState, D3D12_RESOURCE_BARRIER_ALL_SUBRESOURCES);
+			Context.ConditionalSplitCommandList();
+
+			// Texture is now written and ready, so unlock the block (locked after creation and can be defragmented if needed)
+			CurrentTexture.ResourceLocation.UnlockPoolData();
+
+			// If the resource is untracked, the destination state must match the default state of the resource.
+			check(Resource->RequiresResourceStateTracking() || (Resource->GetDefaultResourceState() == DestinationState));
+		}
+	});
 }
 
-void FD3D12Texture::UnlockInternal(class FRHICommandListImmediate* RHICmdList, FLinkedObjectIterator NextObject, uint32 MipIndex, uint32 ArrayIndex)
+void FD3D12Texture::Unlock(FRHICommandListBase& RHICmdList, uint32 MipIndex, uint32 ArrayIndex)
 {
-	// Calculate the subresource index corresponding to the specified mip-map.
-	const uint32 Subresource = CalcSubresource(MipIndex, ArrayIndex, this->GetNumMips());
+	SCOPE_CYCLE_COUNTER(STAT_D3D12UnlockTextureTime);
+	check(IsHeadLink());
 
-	auto* FirstObject = static_cast<FD3D12Texture*>(GetFirstLinkedObject());
-	TMap<uint32, FD3D12LockedResource*>& Map = FirstObject->LockedMap;
-	FD3D12LockedResource* LockedResource = Map[Subresource];
+	// Calculate the subresource index corresponding to the specified mip-map.
+	const uint32 Subresource = CalcSubresource(MipIndex, ArrayIndex, GetNumMips());
 
+	TUniquePtr<FD3D12LockedResource> LockedResource = LockedMap.FindAndRemoveChecked(Subresource);
 	check(LockedResource);
 
 	if (FD3D12DynamicRHI::GetD3DRHI()->HandleSpecialUnlock(RHICmdList, MipIndex, this))
 	{
 		// nothing left to do...
 	}
-	else
+	else if (!LockedResource->bLockedForReadOnly)
 	{
-		if (!LockedResource->bLockedForReadOnly)
+		RHICmdList.EnqueueLambda([
+			RootTexture = this,
+			Subresource,
+			LockedResource = MoveTemp(LockedResource)
+		](FRHICommandListBase& ExecutingCmdList)
 		{
-			FD3D12Resource* Resource = GetResource();
-			FD3D12ResourceLocation& UploadLocation = LockedResource->ResourceLocation;
-
-			// Copy the mip-map data from the real resource into the staging resource
-
-			D3D12_PLACED_SUBRESOURCE_FOOTPRINT PlacedFootprint;
-			PlacedFootprint.Offset = UploadLocation.GetOffsetFromBaseOfResource();
-			PlacedFootprint.Footprint = LockedResource->Footprint;
-			CD3DX12_TEXTURE_COPY_LOCATION SourceCopyLocation(UploadLocation.GetResource()->GetResource(), PlacedFootprint);
-
-			// If we are on the render thread, queue up the copy on the RHIThread so it happens at the correct time.
-			if (ShouldDeferCmdListOperation(RHICmdList))
-			{
-				// Same FD3D12ResourceLocation is used for all resources in the chain, therefore only the last command must be responsible for releasing it.
-				FD3D12ResourceLocation* Source = NextObject ? nullptr : &UploadLocation;
-				ALLOC_COMMAND_CL(*RHICmdList, FRHICommandUpdateTexture)(this, Subresource, 0, 0, 0, SourceCopyLocation, Source);
-			}
-			else
+			D3D12_PLACED_SUBRESOURCE_FOOTPRINT PlacedFootprint
 			{
-				UpdateTexture(Subresource, 0, 0, 0, SourceCopyLocation);
-			}
+				.Offset = LockedResource->ResourceLocation.GetOffsetFromBaseOfResource(),
+				.Footprint = LockedResource->Footprint
+			};
+
+			CD3DX12_TEXTURE_COPY_LOCATION SourceCopyLocation(LockedResource->ResourceLocation.GetResource()->GetResource(), PlacedFootprint);
 
-			// Recurse to update all of the resources in the LDA chain
-			if (NextObject)
+			// Copy the mip-map data from the real resource into the staging resource
+			for (FD3D12Texture& Texture : *RootTexture)
 			{
-				// We pass the first link in the chain as that's the one that got locked
-				((FD3D12Texture*)NextObject.Get())->UnlockInternal(RHICmdList, ++NextObject, MipIndex, ArrayIndex);
+				FD3D12CommandContext& Context = FD3D12CommandContext::Get(ExecutingCmdList, Texture.GetParentDevice()->GetGPUIndex());
+				Texture.UpdateTexture(Context, Subresource, 0, 0, 0, SourceCopyLocation);
 			}
-		}
-	}
-
-	if (FirstObject == this)
-	{
-		// Remove the lock from the outstanding lock list.
-		delete(LockedResource);
-		Map.Remove(Subresource);
+		});
 	}
 }
 
@@ -2068,47 +1872,50 @@ void FD3D12Texture::UpdateTexture2D(FRHICommandListBase& RHICmdList, uint32 MipI
 	const size_t StagingPitch = Align(static_cast<size_t>(WidthInBlocks) * FormatInfo.BlockBytes, FD3D12_TEXTURE_DATA_PITCH_ALIGNMENT);
 	const size_t StagingBufferSize = Align(StagingPitch * HeightInBlocks, D3D12_TEXTURE_DATA_PLACEMENT_ALIGNMENT);
 
-	for (FD3D12Texture& Texture : *this)
+	FD3D12Device* Device = GetParentDevice();
+	FD3D12ResourceLocation UploadHeapResourceLocation(Device);
+	void* const StagingMemory = Device->GetDefaultFastAllocator().Allocate(StagingBufferSize, D3D12_TEXTURE_DATA_PLACEMENT_ALIGNMENT, &UploadHeapResourceLocation);
+	check(StagingMemory);
+
+	const uint8* CopySrc = SourceData + FormatInfo.BlockBytes * SrcXInBlocks + SourcePitch * SrcYInBlocks * FormatInfo.BlockSizeY;
+	uint8* CopyDst = (uint8*)StagingMemory;
+	for (uint32 BlockRow = 0; BlockRow < HeightInBlocks; BlockRow++)
 	{
-		FD3D12Device* Device = Texture.GetParentDevice();
+		FMemory::Memcpy(CopyDst, CopySrc, WidthInBlocks * FormatInfo.BlockBytes);
+		CopySrc += SourcePitch;
+		CopyDst += StagingPitch;
+	}
 
-		FD3D12ResourceLocation UploadHeapResourceLocation(Device);
-		void* const StagingMemory = Device->GetDefaultFastAllocator().Allocate(StagingBufferSize, D3D12_TEXTURE_DATA_PLACEMENT_ALIGNMENT, &UploadHeapResourceLocation);
-		check(StagingMemory);
+	check(StagingPitch % FD3D12_TEXTURE_DATA_PITCH_ALIGNMENT == 0);
 
-		const uint8* CopySrc = SourceData + FormatInfo.BlockBytes * SrcXInBlocks + SourcePitch * SrcYInBlocks * FormatInfo.BlockSizeY;
-		uint8* CopyDst = (uint8*)StagingMemory;
-		for (uint32 BlockRow = 0; BlockRow < HeightInBlocks; BlockRow++)
+	D3D12_PLACED_SUBRESOURCE_FOOTPRINT PlacedTexture2D
+	{
+		.Offset = UploadHeapResourceLocation.GetOffsetFromBaseOfResource(),
+		.Footprint
 		{
-			FMemory::Memcpy(CopyDst, CopySrc, WidthInBlocks * FormatInfo.BlockBytes);
-			CopySrc += SourcePitch;
-			CopyDst += StagingPitch;
+			.Format = (DXGI_FORMAT)FormatInfo.PlatformFormat,
+			.Width = UpdateRegion.Width,
+			.Height = UpdateRegion.Height,
+			.Depth = 1,
+			.RowPitch = (uint32)StagingPitch
 		}
+	};
 
-		D3D12_SUBRESOURCE_FOOTPRINT SourceSubresource{};
-		SourceSubresource.Depth = 1;
-		SourceSubresource.Height = UpdateRegion.Height;
-		SourceSubresource.Width = UpdateRegion.Width;
-		SourceSubresource.Format = (DXGI_FORMAT)FormatInfo.PlatformFormat;
-		SourceSubresource.RowPitch = StagingPitch;
-		check(SourceSubresource.RowPitch % FD3D12_TEXTURE_DATA_PITCH_ALIGNMENT == 0);
-
-		D3D12_PLACED_SUBRESOURCE_FOOTPRINT PlacedTexture2D{};
-		PlacedTexture2D.Offset = UploadHeapResourceLocation.GetOffsetFromBaseOfResource();
-		PlacedTexture2D.Footprint = SourceSubresource;
-
-		CD3DX12_TEXTURE_COPY_LOCATION SourceCopyLocation(UploadHeapResourceLocation.GetResource()->GetResource(), PlacedTexture2D);
+	CD3DX12_TEXTURE_COPY_LOCATION SourceCopyLocation(UploadHeapResourceLocation.GetResource()->GetResource(), PlacedTexture2D);
 
-		// If we are on the render thread, queue up the copy on the RHIThread so it happens at the correct time.
-		if (RHICmdList.IsTopOfPipe())
-		{
-			ALLOC_COMMAND_CL(RHICmdList, FRHICommandUpdateTexture)(&Texture, MipIndex, UpdateRegion.DestX, UpdateRegion.DestY, 0, SourceCopyLocation, &UploadHeapResourceLocation);
-		}
-		else
+	RHICmdList.EnqueueLambda([
+		RootTexture = this,
+		SourceCopyLocation,
+		MipIndex,
+		UpdateRegion
+	](FRHICommandListBase& ExecutingCmdList)
+	{
+		for (FD3D12Texture& Texture : *RootTexture)
 		{
-			Texture.UpdateTexture(MipIndex, UpdateRegion.DestX, UpdateRegion.DestY, 0, SourceCopyLocation);
+			FD3D12CommandContext& Context = FD3D12CommandContext::Get(ExecutingCmdList, Texture.GetParentDevice()->GetGPUIndex());
+			Texture.UpdateTexture(Context, MipIndex, UpdateRegion.DestX, UpdateRegion.DestY, 0, SourceCopyLocation);
 		}
-	}
+	});
 }
 
 static void GetReadBackHeapDescImpl(D3D12_PLACED_SUBRESOURCE_FOOTPRINT& OutFootprint, ID3D12Device* InDevice, D3D12_RESOURCE_DESC const& InResourceDesc, uint32 InSubresource)
@@ -2154,93 +1961,87 @@ void FD3D12Texture::GetReadBackHeapDesc(D3D12_PLACED_SUBRESOURCE_FOOTPRINT& OutF
 	}
 }
 
-void* FD3D12DynamicRHI::LockTexture2D_RenderThread(class FRHICommandListImmediate& RHICmdList, FRHITexture* TextureRHI, uint32 MipIndex, EResourceLockMode LockMode, uint32& DestStride, bool bLockWithinMiptail, bool bNeedsDefaultRHIFlush, uint64* OutLockedByteCount)
+void* FD3D12DynamicRHI::LockTexture2D_RenderThread(FRHICommandListImmediate& RHICmdList, FRHITexture* TextureRHI, uint32 MipIndex, EResourceLockMode LockMode, uint32& DestStride, bool bLockWithinMiptail, bool bNeedsDefaultRHIFlush, uint64* OutLockedByteCount)
 {
-	if (CVarD3D12Texture2DRHIFlush.GetValueOnRenderThread() && bNeedsDefaultRHIFlush)
-	{
-		QUICK_SCOPE_CYCLE_COUNTER(STAT_RHIMETHOD_LockTexture2D_Flush);
-		RHICmdList.ImmediateFlush(EImmediateFlushType::FlushRHIThread);
-		return RHILockTexture2D(TextureRHI, MipIndex, LockMode, DestStride, bLockWithinMiptail, OutLockedByteCount);
-	}
-
 	check(TextureRHI);
 	FD3D12Texture* Texture = FD3D12DynamicRHI::ResourceCast(TextureRHI);
-	return Texture->Lock(&RHICmdList, MipIndex, 0, LockMode, DestStride, OutLockedByteCount);
+	return Texture->Lock(RHICmdList, MipIndex, 0, LockMode, DestStride, OutLockedByteCount);
 }
 
-void* FD3D12DynamicRHI::RHILockTexture2D(FRHITexture* TextureRHI, uint32 MipIndex, EResourceLockMode LockMode, uint32& DestStride, bool bLockWithinMiptail, uint64* OutLockedByteCount)
+void FD3D12DynamicRHI::UnlockTexture2D_RenderThread(FRHICommandListImmediate& RHICmdList, FRHITexture* TextureRHI, uint32 MipIndex, bool bLockWithinMiptail, bool bNeedsDefaultRHIFlush)
 {
 	check(TextureRHI);
-	FD3D12Texture*  Texture = FD3D12DynamicRHI::ResourceCast(TextureRHI);
-	return Texture->Lock(nullptr, MipIndex, 0, LockMode, DestStride, OutLockedByteCount);
+	FD3D12Texture* Texture = FD3D12DynamicRHI::ResourceCast(TextureRHI);
+	Texture->Unlock(RHICmdList, MipIndex, 0);
 }
 
-void FD3D12DynamicRHI::UnlockTexture2D_RenderThread(class FRHICommandListImmediate& RHICmdList, FRHITexture* TextureRHI, uint32 MipIndex, bool bLockWithinMiptail, bool bNeedsDefaultRHIFlush)
+void* FD3D12DynamicRHI::LockTexture2DArray_RenderThread(FRHICommandListImmediate& RHICmdList, FRHITexture* TextureRHI, uint32 TextureIndex, uint32 MipIndex, EResourceLockMode LockMode, uint32& DestStride, bool bLockWithinMiptail)
 {
-	if (CVarD3D12Texture2DRHIFlush.GetValueOnRenderThread() && bNeedsDefaultRHIFlush)
-	{
-		QUICK_SCOPE_CYCLE_COUNTER(STAT_RHIMETHOD_UnlockTexture2D_Flush);
-		RHICmdList.ImmediateFlush(EImmediateFlushType::FlushRHIThread);
-		RHIUnlockTexture2D(TextureRHI, MipIndex, bLockWithinMiptail);
-		return;
-	}
-
 	check(TextureRHI);
 	FD3D12Texture* Texture = FD3D12DynamicRHI::ResourceCast(TextureRHI);
-	Texture->Unlock(&RHICmdList, MipIndex, 0);
+	return Texture->Lock(RHICmdList, MipIndex, TextureIndex, LockMode, DestStride);
 }
 
-void FD3D12DynamicRHI::RHIUnlockTexture2D(FRHITexture* TextureRHI, uint32 MipIndex, bool bLockWithinMiptail)
+void FD3D12DynamicRHI::UnlockTexture2DArray_RenderThread(FRHICommandListImmediate& RHICmdList, FRHITexture* TextureRHI, uint32 TextureIndex, uint32 MipIndex, bool bLockWithinMiptail)
 {
 	check(TextureRHI);
-	FD3D12Texture*  Texture = FD3D12DynamicRHI::ResourceCast(TextureRHI);
-	Texture->Unlock(nullptr, MipIndex, 0);
+	FD3D12Texture* Texture = FD3D12DynamicRHI::ResourceCast(TextureRHI);
+	Texture->Unlock(RHICmdList, MipIndex, TextureIndex);
 }
 
-void* FD3D12DynamicRHI::LockTexture2DArray_RenderThread(class FRHICommandListImmediate& RHICmdList, FRHITexture* TextureRHI, uint32 TextureIndex, uint32 MipIndex, EResourceLockMode LockMode, uint32& DestStride, bool bLockWithinMiptail)
+void* FD3D12DynamicRHI::RHILockTextureCubeFace_RenderThread(FRHICommandListImmediate& RHICmdList, FRHITexture* TextureRHI, uint32 FaceIndex, uint32 ArrayIndex, uint32 MipIndex, EResourceLockMode LockMode, uint32& DestStride, bool bLockWithinMiptail)
 {
 	check(TextureRHI);
-	FD3D12Texture* Texture = FD3D12DynamicRHI::ResourceCast(TextureRHI);
-	return Texture->Lock(&RHICmdList, MipIndex, TextureIndex, LockMode, DestStride);
+	FD3D12Texture* TextureCube = FD3D12DynamicRHI::ResourceCast(TextureRHI);
+	uint32 D3DFace = GetD3D12CubeFace((ECubeFace)FaceIndex);
+	return TextureCube->Lock(RHICmdList, MipIndex, D3DFace + ArrayIndex * 6, LockMode, DestStride);
 }
 
-void* FD3D12DynamicRHI::RHILockTexture2DArray(FRHITexture* TextureRHI, uint32 TextureIndex, uint32 MipIndex, EResourceLockMode LockMode, uint32& DestStride, bool bLockWithinMiptail)
+void FD3D12DynamicRHI::RHIUnlockTextureCubeFace_RenderThread(FRHICommandListImmediate& RHICmdList, FRHITexture* TextureRHI, uint32 FaceIndex, uint32 ArrayIndex, uint32 MipIndex, bool bLockWithinMiptail)
 {
 	check(TextureRHI);
-	FD3D12Texture* Texture = FD3D12DynamicRHI::ResourceCast(TextureRHI);
-	return Texture->Lock(nullptr, MipIndex, TextureIndex, LockMode, DestStride);
+	FD3D12Texture* TextureCube = FD3D12DynamicRHI::ResourceCast(TextureRHI);
+	uint32 D3DFace = GetD3D12CubeFace((ECubeFace)FaceIndex);
+	TextureCube->Unlock(RHICmdList, MipIndex, D3DFace + ArrayIndex * 6);
 }
 
-void FD3D12DynamicRHI::UnlockTexture2DArray_RenderThread(class FRHICommandListImmediate& RHICmdList, FRHITexture* TextureRHI, uint32 TextureIndex, uint32 MipIndex, bool bLockWithinMiptail)
+
+
+void* FD3D12DynamicRHI::RHILockTexture2D(FRHITexture* TextureRHI, uint32 MipIndex, EResourceLockMode LockMode, uint32& DestStride, bool bLockWithinMiptail, uint64* OutLockedByteCount)
 {
-	check(TextureRHI);
-	FD3D12Texture* Texture = FD3D12DynamicRHI::ResourceCast(TextureRHI);
-	Texture->Unlock(&RHICmdList, MipIndex, TextureIndex);
+	UE_LOG(LogD3D12RHI, Fatal, TEXT("RHILockTexture2D is unimplemented and should never be called. Use LockTexture2D_RenderThread."));
+	return nullptr;
+}
+
+void FD3D12DynamicRHI::RHIUnlockTexture2D(FRHITexture* TextureRHI, uint32 MipIndex, bool bLockWithinMiptail)
+{
+	UE_LOG(LogD3D12RHI, Fatal, TEXT("RHIUnlockTexture2D is unimplemented and should never be called. Use UnlockTexture2D_RenderThread."));
+}
+
+void* FD3D12DynamicRHI::RHILockTexture2DArray(FRHITexture* TextureRHI, uint32 TextureIndex, uint32 MipIndex, EResourceLockMode LockMode, uint32& DestStride, bool bLockWithinMiptail)
+{
+	UE_LOG(LogD3D12RHI, Fatal, TEXT("RHILockTexture2DArray is unimplemented and should never be called. Use LockTexture2DArray_RenderThread."));
+	return nullptr;
 }
 
 void FD3D12DynamicRHI::RHIUnlockTexture2DArray(FRHITexture* TextureRHI, uint32 TextureIndex, uint32 MipIndex, bool bLockWithinMiptail)
 {
-	check(TextureRHI);
-	FD3D12Texture* Texture = FD3D12DynamicRHI::ResourceCast(TextureRHI);
-	Texture->Unlock(nullptr, MipIndex, TextureIndex);
+	UE_LOG(LogD3D12RHI, Fatal, TEXT("RHIUnlockTexture2DArray is unimplemented and should never be called. Use UnlockTexture2DArray_RenderThread."));
 }
 
-void* FD3D12DynamicRHI::RHILockTextureCubeFace_RenderThread(class FRHICommandListImmediate& RHICmdList, FRHITexture* TextureRHI, uint32 FaceIndex, uint32 ArrayIndex, uint32 MipIndex, EResourceLockMode LockMode, uint32& DestStride, bool bLockWithinMiptail)
+void* FD3D12DynamicRHI::RHILockTextureCubeFace(FRHITexture* TextureCubeRHI, uint32 FaceIndex, uint32 ArrayIndex, uint32 MipIndex, EResourceLockMode LockMode, uint32& DestStride, bool bLockWithinMiptail)
 {
-	check(TextureRHI);
-	FD3D12Texture* TextureCube = FD3D12DynamicRHI::ResourceCast(TextureRHI);
-	uint32 D3DFace = GetD3D12CubeFace((ECubeFace)FaceIndex);
-	return TextureCube->Lock(&RHICmdList, MipIndex, D3DFace + ArrayIndex * 6, LockMode, DestStride);
+	UE_LOG(LogD3D12RHI, Fatal, TEXT("RHIUnlockTexture2DArray is unimplemented and should never be called. Use UnlockTexture2DArray_RenderThread."));
+	return nullptr;
 }
 
-void FD3D12DynamicRHI::RHIUnlockTextureCubeFace_RenderThread(class FRHICommandListImmediate& RHICmdList, FRHITexture* TextureRHI, uint32 FaceIndex, uint32 ArrayIndex, uint32 MipIndex, bool bLockWithinMiptail)
+void FD3D12DynamicRHI::RHIUnlockTextureCubeFace(FRHITexture* TextureCubeRHI, uint32 FaceIndex, uint32 ArrayIndex, uint32 MipIndex, bool bLockWithinMiptail)
 {
-	check(TextureRHI);
-	FD3D12Texture* TextureCube = FD3D12DynamicRHI::ResourceCast(TextureRHI);
-	uint32 D3DFace = GetD3D12CubeFace((ECubeFace)FaceIndex);
-	TextureCube->Unlock(&RHICmdList, MipIndex, D3DFace + ArrayIndex * 6);
+	UE_LOG(LogD3D12RHI, Fatal, TEXT("RHIUnlockTextureCubeFace is unimplemented and should never be called. Use RHIUnlockTextureCubeFace_RenderThread."));
 }
 
+
+
 void FD3D12DynamicRHI::RHIUpdateTexture2D(FRHICommandListBase& RHICmdList, FRHITexture* TextureRHI, uint32 MipIndex, const FUpdateTextureRegion2D& UpdateRegion, uint32 SourcePitch, const uint8* SourceData)
 {
 	check(TextureRHI);
@@ -2311,14 +2112,14 @@ public:
 		UpdateInfos.Empty();
 	}
 
-	void Execute(FRHICommandListBase& RHICmdList)
+	void Execute(FRHICommandListBase& ExecutingCmdList)
 	{
 		FD3D12Texture* NativeTexture = FD3D12DynamicRHI::ResourceCast(DstTexture.GetReference());
 
 		for (FD3D12Texture& TextureLink : *NativeTexture)
 		{
 			FD3D12Device* Device = TextureLink.GetParentDevice();
-			FD3D12CommandContext& Context = Device->GetDefaultCommandContext();
+			FD3D12CommandContext& Context = FD3D12CommandContext::Get(ExecutingCmdList, Device->GetGPUIndex());
 
 			CD3DX12_TEXTURE_COPY_LOCATION DestCopyLocation(TextureLink.GetResource()->GetResource(), MipIdx);
 
@@ -2331,12 +2132,9 @@ public:
 				const FUpdateInfo& UpdateInfo = UpdateInfos[Idx];
 				FD3D12Resource* UploadBuffer = UpdateInfo.SrcResourceLocation->GetResource();
 				CD3DX12_TEXTURE_COPY_LOCATION SourceCopyLocation(UploadBuffer->GetResource(), UpdateInfo.PlacedSubresourceFootprint);
-#if USE_PIX
-				if (FD3D12DynamicRHI::GetD3DRHI()->IsPixEventEnabled())
-				{
-					PIXBeginEvent(Context.GraphicsCommandList().Get(), PIX_COLOR(255, 255, 255), TEXT("EndMultiUpdateTexture3D"));
-				}
-#endif
+
+				RHI_BREADCRUMB_EVENT(Context, EndMultiUpdateTexture3D);
+
 				Context.GraphicsCommandList()->CopyTextureRegion(
 					&DestCopyLocation,
 					UpdateInfo.DstStartX,
@@ -2347,12 +2145,6 @@ public:
 
 				Context.UpdateResidency(TextureLink.GetResource());
 				DEBUG_EXECUTE_COMMAND_CONTEXT(Context);
-#if USE_PIX
-				if (FD3D12DynamicRHI::GetD3DRHI()->IsPixEventEnabled())
-				{
-					PIXEndEvent(Context.GraphicsCommandList().Get());
-				}
-#endif
 			}
 
 			Context.ConditionalSplitCommandList();
@@ -2509,146 +2301,73 @@ FUpdateTexture3DData FD3D12DynamicRHI::BeginUpdateTexture3D_Internal(FRHITexture
 	return UpdateData;
 }
 
-struct FD3D12RHICmdEndUpdateTexture3DString
-{
-	static const TCHAR* TStr() { return TEXT("FD3D12RHICmdEndUpdateTexture3D"); }
-};
-class FD3D12RHICmdEndUpdateTexture3D : public FRHICommand<FD3D12RHICmdEndUpdateTexture3D, FD3D12RHICmdEndUpdateTexture3DString>
+void FD3D12DynamicRHI::EndUpdateTexture3D_Internal(FRHICommandListBase& RHICmdList, FUpdateTexture3DData& UpdateData)
 {
-public:
-	FD3D12RHICmdEndUpdateTexture3D(FUpdateTexture3DData& UpdateData) :
-		MipIdx(UpdateData.MipIndex),
-		DstStartX(UpdateData.UpdateRegion.DestX),
-		DstStartY(UpdateData.UpdateRegion.DestY),
-		DstStartZ(UpdateData.UpdateRegion.DestZ),
-		DstTexture(UpdateData.Texture)
-	{
-		FMemory::Memset(&PlacedSubresourceFootprint, 0, sizeof(PlacedSubresourceFootprint));
-
-		D3D12_SUBRESOURCE_FOOTPRINT& SubresourceFootprint = PlacedSubresourceFootprint.Footprint;
-		SubresourceFootprint.Depth = UpdateData.UpdateRegion.Depth;
-		SubresourceFootprint.Height = UpdateData.UpdateRegion.Height;
-		SubresourceFootprint.Width = UpdateData.UpdateRegion.Width;
-		SubresourceFootprint.Format = static_cast<DXGI_FORMAT>(GPixelFormats[DstTexture->GetFormat()].PlatformFormat);
-		SubresourceFootprint.RowPitch = UpdateData.RowPitch;
-		check(SubresourceFootprint.RowPitch % FD3D12_TEXTURE_DATA_PITCH_ALIGNMENT == 0);
-
-		FD3D12UpdateTexture3DData* UpdateDataD3D12 =
-			reinterpret_cast<FD3D12UpdateTexture3DData*>(&UpdateData.PlatformData[0]);
+	check(GFrameNumberRenderThread == UpdateData.FrameNumber);
 
-		SrcResourceLocation = UpdateDataD3D12->UploadHeapResourceLocation;
-		PlacedSubresourceFootprint.Offset = SrcResourceLocation->GetOffsetFromBaseOfResource();
-	}
+	FD3D12UpdateTexture3DData* UpdateDataD3D12 = reinterpret_cast<FD3D12UpdateTexture3DData*>(&UpdateData.PlatformData[0]);
+	check(UpdateDataD3D12->UploadHeapResourceLocation != nullptr);
 
-	virtual ~FD3D12RHICmdEndUpdateTexture3D()
+	if (UpdateDataD3D12->bComputeShaderCopy)
 	{
-		if (SrcResourceLocation)
-		{
-			delete SrcResourceLocation;
-			SrcResourceLocation = nullptr;
-		}
+		EndUpdateTexture3D_ComputeShader(static_cast<FRHIComputeCommandList&>(RHICmdList), UpdateData, UpdateDataD3D12);
 	}
-
-	void Execute(FRHICommandListBase& RHICmdList)
+	else
 	{
-		FD3D12Texture* NativeTexture = FD3D12DynamicRHI::ResourceCast(DstTexture.GetReference());
-		FD3D12Resource* UploadBuffer = SrcResourceLocation->GetResource();
+		check(UpdateData.RowPitch % FD3D12_TEXTURE_DATA_PITCH_ALIGNMENT == 0);
+		FD3D12ResourceLocation* SrcResourceLocation = UpdateDataD3D12->UploadHeapResourceLocation;
 
-		for (FD3D12Texture& TextureLink : *NativeTexture)
+		D3D12_PLACED_SUBRESOURCE_FOOTPRINT PlacedSubresourceFootprint
 		{
-			FD3D12Device* Device = TextureLink.GetParentDevice();
-			FD3D12CommandContext& Context = Device->GetDefaultCommandContext();
-
-#if USE_PIX
-			if (FD3D12DynamicRHI::GetD3DRHI()->IsPixEventEnabled())
+			.Offset = SrcResourceLocation->GetOffsetFromBaseOfResource(),
+			.Footprint
 			{
-				PIXBeginEvent(Context.GraphicsCommandList().Get(), PIX_COLOR(255, 255, 255), TEXT("EndUpdateTexture3D"));
+				.Format   = static_cast<DXGI_FORMAT>(GPixelFormats[UpdateData.Texture->GetFormat()].PlatformFormat),
+				.Width    = UpdateData.UpdateRegion.Width,
+				.Height   = UpdateData.UpdateRegion.Height,
+				.Depth    = UpdateData.UpdateRegion.Depth,
+				.RowPitch = UpdateData.RowPitch
 			}
-#endif
-			CD3DX12_TEXTURE_COPY_LOCATION DestCopyLocation(TextureLink.GetResource()->GetResource(), MipIdx);
-			CD3DX12_TEXTURE_COPY_LOCATION SourceCopyLocation(UploadBuffer->GetResource(), PlacedSubresourceFootprint);
-
-			FScopedResourceBarrier ScopeResourceBarrierDest(Context, TextureLink.GetResource(), &TextureLink.ResourceLocation, D3D12_RESOURCE_STATE_COPY_DEST, DestCopyLocation.SubresourceIndex);
-
-			Context.FlushResourceBarriers();
-			Context.GraphicsCommandList()->CopyTextureRegion(
-				&DestCopyLocation,
-				DstStartX,
-				DstStartY,
-				DstStartZ,
-				&SourceCopyLocation,
-				nullptr);
-
-			Context.UpdateResidency(TextureLink.GetResource());
+		};
 
-			Context.ConditionalSplitCommandList();
-			DEBUG_EXECUTE_COMMAND_CONTEXT(Context);
-#if USE_PIX
-			if (FD3D12DynamicRHI::GetD3DRHI()->IsPixEventEnabled())
+		RHICmdList.EnqueueLambda([
+			MipIdx      = UpdateData.MipIndex,
+			DstStartX   = UpdateData.UpdateRegion.DestX,
+			DstStartY   = UpdateData.UpdateRegion.DestY,
+			DstStartZ   = UpdateData.UpdateRegion.DestZ,
+			RootTexture = UpdateData.Texture,
+			PlacedSubresourceFootprint,
+			SrcResourceLocation
+		](FRHICommandListBase& ExecutingCmdList)
+		{
+			for (FD3D12Texture& Texture : *FD3D12DynamicRHI::ResourceCast(RootTexture))
 			{
-				PIXEndEvent(Context.GraphicsCommandList().Get());
-			}
-#endif
-		}
+				FD3D12CommandContext& Context = FD3D12CommandContext::Get(ExecutingCmdList, Texture.GetParentDevice()->GetGPUIndex());
+				RHI_BREADCRUMB_EVENT(Context, EndUpdateTexture3D);
 
-		delete SrcResourceLocation;
-		SrcResourceLocation = nullptr;
-	}
+				CD3DX12_TEXTURE_COPY_LOCATION DestCopyLocation(Texture.GetResource()->GetResource(), MipIdx);
+				CD3DX12_TEXTURE_COPY_LOCATION SourceCopyLocation(SrcResourceLocation->GetResource()->GetResource(), PlacedSubresourceFootprint);
 
-private:
-	uint32 MipIdx;
-	uint32 DstStartX;
-	uint32 DstStartY;
-	uint32 DstStartZ;
-	FTextureRHIRef DstTexture;
-	FD3D12ResourceLocation* SrcResourceLocation;
-	D3D12_PLACED_SUBRESOURCE_FOOTPRINT PlacedSubresourceFootprint;
-};
+				FScopedResourceBarrier ScopeResourceBarrierDest(Context, Texture.GetResource(), &Texture.ResourceLocation, D3D12_RESOURCE_STATE_COPY_DEST, DestCopyLocation.SubresourceIndex);
 
-void FD3D12DynamicRHI::EndUpdateTexture3D_Internal(FRHICommandListBase& RHICmdList, FUpdateTexture3DData& UpdateData)
-{
-	check(GFrameNumberRenderThread == UpdateData.FrameNumber);
+				Context.FlushResourceBarriers();
+				Context.UpdateResidency(Texture.GetResource());
 
-	FD3D12UpdateTexture3DData* UpdateDataD3D12 = reinterpret_cast<FD3D12UpdateTexture3DData*>(&UpdateData.PlatformData[0]);
-	check( UpdateDataD3D12->UploadHeapResourceLocation != nullptr );
-
-	if (UpdateDataD3D12->bComputeShaderCopy)
-	{
-		EndUpdateTexture3D_ComputeShader(static_cast<FRHIComputeCommandList&>(RHICmdList), UpdateData, UpdateDataD3D12);
-	}
-	else
-	{
-		if (RHICmdList.IsBottomOfPipe())
-		{
-			FD3D12RHICmdEndUpdateTexture3D RHICmd(UpdateData);
-			RHICmd.Execute(RHICmdList);
-		}
-		else
-		{
-			ALLOC_COMMAND_CL(RHICmdList, FD3D12RHICmdEndUpdateTexture3D)(UpdateData);
-		}
-	}
-}
+				Context.GraphicsCommandList()->CopyTextureRegion(
+					&DestCopyLocation,
+					DstStartX,
+					DstStartY,
+					DstStartZ,
+					&SourceCopyLocation,
+					nullptr);
 
-/*-----------------------------------------------------------------------------
-	Cubemap texture support.
-	-----------------------------------------------------------------------------*/
+				Context.ConditionalSplitCommandList();
+				DEBUG_EXECUTE_COMMAND_CONTEXT(Context);
+			}
 
-void* FD3D12DynamicRHI::RHILockTextureCubeFace(FRHITexture* TextureCubeRHI, uint32 FaceIndex, uint32 ArrayIndex, uint32 MipIndex, EResourceLockMode LockMode, uint32& DestStride, bool bLockWithinMiptail)
-{
-	FD3D12Texture*  TextureCube = FD3D12DynamicRHI::ResourceCast(TextureCubeRHI);
-	for (uint32 GPUIndex : TextureCube->GetLinkedObjectsGPUMask())
-	{
-		GetRHIDevice(GPUIndex)->GetDefaultCommandContext().ConditionalClearShaderResource(&TextureCube->ResourceLocation, EShaderParameterTypeMask::SRVMask | EShaderParameterTypeMask::UAVMask);
+			delete SrcResourceLocation;
+		});
 	}
-	uint32 D3DFace = GetD3D12CubeFace((ECubeFace)FaceIndex);
-	return TextureCube->Lock(nullptr, MipIndex, D3DFace + ArrayIndex * 6, LockMode, DestStride);
-}
-void FD3D12DynamicRHI::RHIUnlockTextureCubeFace(FRHITexture* TextureCubeRHI, uint32 FaceIndex, uint32 ArrayIndex, uint32 MipIndex, bool bLockWithinMiptail)
-{
-	FD3D12Texture* TextureCube = FD3D12DynamicRHI::ResourceCast(TextureCubeRHI);
-	uint32 D3DFace = GetD3D12CubeFace((ECubeFace)FaceIndex);
-	TextureCube->Unlock(nullptr, MipIndex, D3DFace + ArrayIndex * 6);
 }
 
 void FD3D12DynamicRHI::RHIBindDebugLabelName(FRHICommandListBase& RHICmdList, FRHITexture* TextureRHI, const TCHAR* Name)
diff --git a/Engine/Source/Runtime/D3D12RHI/Private/D3D12Texture.h b/Engine/Source/Runtime/D3D12RHI/Private/D3D12Texture.h
--- a/Engine/Source/Runtime/D3D12RHI/Private/D3D12Texture.h
+++ b/Engine/Source/Runtime/D3D12RHI/Private/D3D12Texture.h
@@ -119,21 +119,19 @@ public:
 	}
 
 	// Locking/update functions
-	void* Lock(class FRHICommandListImmediate* RHICmdList, uint32 MipIndex, uint32 ArrayIndex, EResourceLockMode LockMode, uint32& DestStride, uint64* OutLockedByteCount = nullptr);
-	void Unlock(class FRHICommandListImmediate* RHICmdList, uint32 MipIndex, uint32 ArrayIndex);
+	void* Lock(FRHICommandListImmediate& RHICmdList, uint32 MipIndex, uint32 ArrayIndex, EResourceLockMode LockMode, uint32& DestStride, uint64* OutLockedByteCount = nullptr);
+	void Unlock(FRHICommandListBase& RHICmdList, uint32 MipIndex, uint32 ArrayIndex);
+
 	void UpdateTexture2D(FRHICommandListBase& RHICmdList, uint32 MipIndex, const struct FUpdateTextureRegion2D& UpdateRegion, uint32 SourcePitch, const uint8* SourceData);
-	void UpdateTexture(uint32 MipIndex, uint32 DestX, uint32 DestY, uint32 DestZ, const D3D12_TEXTURE_COPY_LOCATION& SourceCopyLocation);
-	void CopyTextureRegion(uint32 DestX, uint32 DestY, uint32 DestZ, FD3D12Texture* SourceTexture, const D3D12_BOX& SourceBox);
+	void UpdateTexture(FD3D12CommandContext& Context, uint32 MipIndex, uint32 DestX, uint32 DestY, uint32 DestZ, const D3D12_TEXTURE_COPY_LOCATION& SourceCopyLocation);
+
+	void CopyTextureRegion(FD3D12CommandContext& Context, uint32 DestX, uint32 DestY, uint32 DestZ, FD3D12Texture* SourceTexture, const D3D12_BOX& SourceBox);
 
 	// Resource aliasing
 	void AliasResources(FD3D12Texture* Texture);
 	void SetAliasingSource(FTextureRHIRef& SourceTextureRHI)	{ AliasingSourceTexture = SourceTextureRHI; }
 
 protected:
-
-	// Lock helper functions
-	void UnlockInternal(class FRHICommandListImmediate* RHICmdList, FLinkedObjectIterator NextObject, uint32 MipIndex, uint32 ArrayIndex);
-		
 	// A shader resource view of the texture.
 	TSharedPtr<FD3D12ShaderResourceView> ShaderResourceView;
 
@@ -154,7 +152,7 @@ protected:
 	TSharedPtr<FD3D12DepthStencilView> DepthStencilViews[FExclusiveDepthStencil::MaxIndex];
 
 	// Data for each subresource while texture is locked
-	TMap<uint32, FD3D12LockedResource*> LockedMap;
+	TMap<uint32, TUniquePtr<FD3D12LockedResource>> LockedMap;
 
 	// Cached footprint size of first resource - optimization
 	mutable TUniquePtr<D3D12_PLACED_SUBRESOURCE_FOOTPRINT> FirstSubresourceFootprint;
