Fix Texture Cube SRV view range not being a multiple of 6
 - This was causing us to miss SRV transitions on the all but the first face of a texture cube resource, e.g. during mip downsample for reflection captures.

diff --git a/Engine/Source/Runtime/D3D12RHI/Private/D3D12View.cpp b/Engine/Source/Runtime/D3D12RHI/Private/D3D12View.cpp
--- a/Engine/Source/Runtime/D3D12RHI/Private/D3D12View.cpp
+++ b/Engine/Source/Runtime/D3D12RHI/Private/D3D12View.cpp
@@ -117,7 +117,7 @@ FD3D12ViewRange::FD3D12ViewRange(D3D12_SHADER_RESOURCE_VIEW_DESC const& ViewDesc
 
 	case D3D12_SRV_DIMENSION_TEXTURECUBE:
 		Mip   = { ViewDesc.TextureCube.MostDetailedMip, ViewDesc.TextureCube.MipLevels };
-		Array = { 0, 1 };
+		Array = { 0, 6 };
 		break;
 
 	case D3D12_SRV_DIMENSION_TEXTURECUBEARRAY:
