Removed frame counter from D3D12 adapter that is always zero, and fixed the transient allocator virtual heap logic to use the garbage collection counter instead.

diff --git a/Engine/Source/Runtime/D3D12RHI/Private/D3D12Adapter.h b/Engine/Source/Runtime/D3D12RHI/Private/D3D12Adapter.h
--- a/Engine/Source/Runtime/D3D12RHI/Private/D3D12Adapter.h
+++ b/Engine/Source/Runtime/D3D12RHI/Private/D3D12Adapter.h
@@ -398,8 +398,6 @@ public:
 	void UpdateMemoryInfo();
 	FORCEINLINE const FD3D12MemoryInfo& GetMemoryInfo() const { return MemoryInfo; }
 
-	FORCEINLINE uint32 GetFrameCount() const { return FrameCounter; }
-
 	bool IsTrackingAllAllocations() const { return bTrackAllAllocation; }
 	void TrackAllocationData(FD3D12ResourceLocation* InAllocation, uint64 InAllocationSize, bool bCollectCallstack);
 	void ReleaseTrackedAllocationData(FD3D12ResourceLocation* InAllocation, bool bDefragFree);
@@ -568,8 +566,6 @@ protected:
 
 	FD3D12CommandContextRedirector DefaultContextRedirector;
 
-	uint32 FrameCounter = 0;
-
 	bool bTrackAllAllocation = false;
 
 	/** Information about an allocated resource. */
