Build fix for platforms not supporting local root signature.

diff --git a/Engine/Source/Runtime/D3D12RHI/Private/D3D12RootSignature.cpp b/Engine/Source/Runtime/D3D12RHI/Private/D3D12RootSignature.cpp
--- a/Engine/Source/Runtime/D3D12RHI/Private/D3D12RootSignature.cpp
+++ b/Engine/Source/Runtime/D3D12RHI/Private/D3D12RootSignature.cpp
@@ -306,6 +306,7 @@ FD3D12RootSignatureDesc::FD3D12RootSignatureDesc(const FD3D12QuantizedBoundShade
 		Flags |= D3D12_ROOT_SIGNATURE_FLAG_SAMPLER_HEAP_DIRECTLY_INDEXED;
 	}
 
+#if D3D12_RHI_RAYTRACING
 	if (QBSS.RootSignatureType == RS_RayTracingLocal || QBSS.RootSignatureType == RS_WorkGraphLocal)
 	{
 		Flags = D3D12_ROOT_SIGNATURE_FLAG_LOCAL_ROOT_SIGNATURE;
@@ -315,6 +316,7 @@ FD3D12RootSignatureDesc::FD3D12RootSignatureDesc(const FD3D12QuantizedBoundShade
 		Flags |= FD3D12_ROOT_SIGNATURE_FLAG_GLOBAL_ROOT_SIGNATURE;
 	}
 	else if (QBSS.RootSignatureType == RS_Raster)
+#endif // D3D12_RHI_RAYTRACING
 	{
 		// Determine what shader stages need access in the root signature.
 
@@ -375,6 +377,7 @@ FD3D12RootSignatureDesc::FD3D12RootSignatureDesc(const FD3D12QuantizedBoundShade
 	}
 #endif
 
+#if D3D12_RHI_RAYTRACING
 	if (QBSS.RootSignatureType == RS_RayTracingLocal || QBSS.RootSignatureType == RS_WorkGraphLocal)
 	{
 		// Local root signatures don't need to provide static samplers as they are provided by global RS already.
@@ -382,6 +385,7 @@ FD3D12RootSignatureDesc::FD3D12RootSignatureDesc(const FD3D12QuantizedBoundShade
 		RootDesc.Init_1_1(RootParameterCount, TableSlots, 0, nullptr, Flags);
 	}
 	else
+#endif
 	{
 		// Only use static samplers for binding tier higher than 1 otherwise root signature only supports 16 samplers
 		// Only use by DXR shaders and validated that DXR has at least Tier 2 support
@@ -560,7 +564,11 @@ void FD3D12RootSignature::InternalAnalyzeSignature(const RootSignatureDescType&
 	const bool bDenyAS = (Desc.Flags & D3D12_ROOT_SIGNATURE_FLAG_DENY_AMPLIFICATION_SHADER_ROOT_ACCESS) != 0;
 #endif
 
+#if D3D12_RHI_RAYTRACING
 	const uint32 RootDescriptorTableCost = (Desc.Flags & D3D12_ROOT_SIGNATURE_FLAG_LOCAL_ROOT_SIGNATURE) ? RootDescriptorTableCostLocal : RootDescriptorTableCostGlobal;
+#else
+	const uint32 RootDescriptorTableCost = RootDescriptorTableCostGlobal;
+#endif
 
 	// Go through each root parameter.
 	for (uint32 i = 0; i < Desc.NumParameters; i++)
