Fixed D3D12AllocatorUnused stat by taking D3D12PoolAllocator defrag into account

#rnx

diff --git a/Engine/Source/Runtime/D3D12RHI/Private/D3D12PoolAllocator.cpp b/Engine/Source/Runtime/D3D12RHI/Private/D3D12PoolAllocator.cpp
--- a/Engine/Source/Runtime/D3D12RHI/Private/D3D12PoolAllocator.cpp
+++ b/Engine/Source/Runtime/D3D12RHI/Private/D3D12PoolAllocator.cpp
@@ -369,7 +369,6 @@ void FD3D12PoolAllocator::AllocateResource(uint32 GPUIndex, D3D12_HEAP_TYPE InHe
 		// Try to allocate in one of the pools
 		FRHIPoolAllocationData& AllocationData = ResourceLocation.GetPoolAllocatorPrivateData().PoolData;
 		verify(TryAllocateInternal(InSize, AllocationAlignment, AllocationResourceType, AllocationData));
-		LLM_SCOPED_PAUSE_TRACKING_WITH_ENUM_AND_AMOUNT_BYTAG(D3D12AllocatorUnused, 0 - int64(InSize), ELLMTracker::Platform, ELLMAllocType::System);
 
 		// Setup the resource location
 		ResourceLocation.SetType(FD3D12ResourceLocation::ResourceLocationType::eSubAllocation);
@@ -644,7 +643,6 @@ void FD3D12PoolAllocator::CleanUpAllocations(uint64 InFrameLag, bool bForceFree)
 				// Not pending anymore
 				check(PendingDeleteRequestSize >= Operation.AllocationData->GetSize());
 				PendingDeleteRequestSize -= Operation.AllocationData->GetSize();
-				LLM_SCOPED_PAUSE_TRACKING_WITH_ENUM_AND_AMOUNT_BYTAG(D3D12AllocatorUnused, int64(Operation.AllocationData->GetSize()), ELLMTracker::Platform, ELLMAllocType::System);
 				// Deallocate the locked block (actually free now)
 				DeallocateInternal(*Operation.AllocationData);
 				Operation.AllocationData->Reset();
@@ -814,6 +812,21 @@ void FD3D12PoolAllocator::UpdateAllocationTracking(FD3D12ResourceLocation& InAll
 		}
 	}
 #endif // TRACK_RESOURCE_ALLOCATIONS
+
+#if ENABLE_LOW_LEVEL_MEM_TRACKER
+	{
+		int64 AllocSize = int64(InAllocation.GetPoolAllocatorPrivateData().PoolData.GetSize());
+		
+		if (InAllocationType == EAllocationType::Allocate)
+		{
+			LLM_SCOPED_PAUSE_TRACKING_WITH_ENUM_AND_AMOUNT_BYTAG(D3D12AllocatorUnused, 0 - AllocSize, ELLMTracker::Platform, ELLMAllocType::System);
+		}
+		else
+		{
+			LLM_SCOPED_PAUSE_TRACKING_WITH_ENUM_AND_AMOUNT_BYTAG(D3D12AllocatorUnused, AllocSize, ELLMTracker::Platform, ELLMAllocType::System);
+		}
+	}
+#endif
 }
 
 
