Fix crash in FD3D12BindlessResourceManager::UpdateDescriptor
 - Context pointer may be null if async compute is disabled.

diff --git a/Engine/Source/Runtime/D3D12RHI/Private/D3D12BindlessDescriptors.cpp b/Engine/Source/Runtime/D3D12RHI/Private/D3D12BindlessDescriptors.cpp
--- a/Engine/Source/Runtime/D3D12RHI/Private/D3D12BindlessDescriptors.cpp
+++ b/Engine/Source/Runtime/D3D12RHI/Private/D3D12BindlessDescriptors.cpp
@@ -144,11 +144,14 @@ void FD3D12BindlessResourceManager::UpdateDescriptor(FD3D12ContextArray const& C
 		uint32 const GPUIndex = GetParentDevice()->GetGPUIndex();
 		for (FD3D12CommandContextBase* ContextBase : Contexts)
 		{
-			FD3D12CommandContext& Context = *ContextBase->GetSingleDeviceContext(GPUIndex);
-			if (Context.IsOpen())
+			if (ContextBase)
 			{
-				FD3D12OfflineDescriptor CopyOfPreviousDescriptorValue = UE::D3D12Descriptors::CreateOfflineCopy(GetParentDevice(), CpuHeap, DstHandle);
-				Context.GetBindlessState().PendingDescriptorRollbacks.Add(GetParentDevice(), DstHandle, CopyOfPreviousDescriptorValue);
+				FD3D12CommandContext& Context = *ContextBase->GetSingleDeviceContext(GPUIndex);
+				if (Context.IsOpen())
+				{
+					FD3D12OfflineDescriptor CopyOfPreviousDescriptorValue = UE::D3D12Descriptors::CreateOfflineCopy(GetParentDevice(), CpuHeap, DstHandle);
+					Context.GetBindlessState().PendingDescriptorRollbacks.Add(GetParentDevice(), DstHandle, CopyOfPreviousDescriptorValue);
+				}
 			}
 		}
 
