D3D12: Add mode to set stable power state on device creation instead of only during profiling

This can be useful for in-editor benchmarking on PC by reducing the influence of adaptive GPU clock rate on the frame time.

diff --git a/Engine/Source/Runtime/D3D12RHI/Private/D3D12Adapter.cpp b/Engine/Source/Runtime/D3D12RHI/Private/D3D12Adapter.cpp
--- a/Engine/Source/Runtime/D3D12RHI/Private/D3D12Adapter.cpp
+++ b/Engine/Source/Runtime/D3D12RHI/Private/D3D12Adapter.cpp
@@ -770,6 +770,17 @@ void FD3D12Adapter::CreateRootDevice(bool bWithDebug)
 		UE_LOG(LogD3D12RHI, Log, TEXT("Enabling multi-GPU with %d nodes"), Desc.NumDeviceNodes);
 	}
 #endif
+
+	if (RootDevice)
+	{
+		extern int32 GD3D12RHIStablePowerState;
+		if (GD3D12RHIStablePowerState == 2)
+		{
+			bool bWorked = SUCCEEDED(RootDevice->SetStablePowerState(true));
+			// This will fail if windows developper mode is not enabled. Windows will remove the adapter on failure, so we can't really gracefully exit here.
+			checkf(bWorked, TEXT("Enabling state power state requires Windows developer mode to be enabled."));
+		}
+	}
 }
 
 void FD3D12Adapter::SetDrawingViewport(FD3D12Viewport* InViewport)
diff --git a/Engine/Source/Runtime/D3D12RHI/Private/D3D12GPUProfiler.cpp b/Engine/Source/Runtime/D3D12RHI/Private/D3D12GPUProfiler.cpp
--- a/Engine/Source/Runtime/D3D12RHI/Private/D3D12GPUProfiler.cpp
+++ b/Engine/Source/Runtime/D3D12RHI/Private/D3D12GPUProfiler.cpp
@@ -10,7 +10,10 @@ int32 GD3D12RHIStablePowerState = 0;
 static FAutoConsoleVariableRef CVarD3D12RHIStablePowerState(
 	TEXT("D3D12.StablePowerState"),
 	GD3D12RHIStablePowerState,
-	TEXT("If true, enable stable power state. This increases GPU timing measurement accuracy but may decrease overall GPU clock rate."),
+	TEXT("Enable stable power state. This increases GPU timing measurement accuracy but may decrease overall GPU clock rate.\n")
+	TEXT("    0 (default): off\n")
+	TEXT("    1          : set during profiling\n")
+	TEXT("    2          : set on startup\n"),
 	ECVF_Default
 );
 
