WorkGraphs: Changed pipeline init to use  CD3DX12_COMMON_COMPUTE_NODE_OVERRIDES.
The new shared overide type is from the recently released agility SDK.
With this we longer need to "know" the node type to override the node name.

diff --git a/Engine/Source/Runtime/D3D12RHI/Private/D3D12WorkGraph.cpp b/Engine/Source/Runtime/D3D12RHI/Private/D3D12WorkGraph.cpp
--- a/Engine/Source/Runtime/D3D12RHI/Private/D3D12WorkGraph.cpp
+++ b/Engine/Source/Runtime/D3D12RHI/Private/D3D12WorkGraph.cpp
@@ -77,7 +77,7 @@ FD3D12WorkGraphPipelineState::FD3D12WorkGraphPipelineState(FD3D12Device* Device,
 			if (LocalNodeShaders[Index])
 			{
 				FString NodeName = FString::Printf(TEXT("%s_%d"), *Initializer.GetShaderBundleNodeName(), Index);
-				CD3DX12_BROADCASTING_LAUNCH_NODE_OVERRIDES* Override = WorkGraph->CreateBroadcastingLaunchNodeOverrides(*NodeName);
+				CD3DX12_COMMON_COMPUTE_NODE_OVERRIDES* Override = WorkGraph->CreateCommonComputeNodeOverrides(*NodeName);
 				Override->NewName(D3D12_NODE_ID{ *Initializer.GetShaderBundleNodeName(), (uint32)Index });
 			}
 		}
