Remove unused virtual functions in FRHIShaderBundle

diff --git a/Engine/Source/Runtime/D3D12RHI/Private/D3D12Resources.h b/Engine/Source/Runtime/D3D12RHI/Private/D3D12Resources.h
--- a/Engine/Source/Runtime/D3D12RHI/Private/D3D12Resources.h
+++ b/Engine/Source/Runtime/D3D12RHI/Private/D3D12Resources.h
@@ -1059,17 +1059,6 @@ public:
 		: FRHIShaderBundle(InNumRecords)
 	{
 	}
-
-	virtual FUint32Vector3 GetPlatformData() const override
-	{
-		return FUint32Vector3(0u, 0u, 0u);
-	}
-
-	virtual void CalcDispatchBufferSizes(uint32& OutDataBufferSize, uint32& OutExecutionBufferSize) const override
-	{
-		OutDataBufferSize = 4u;
-		OutExecutionBufferSize = 4u;
-	}
 };
 
 class FD3D12GPUFence final : public FRHIGPUFence
diff --git a/Engine/Source/Runtime/RHI/Public/RHIResources.h b/Engine/Source/Runtime/RHI/Public/RHIResources.h
--- a/Engine/Source/Runtime/RHI/Public/RHIResources.h
+++ b/Engine/Source/Runtime/RHI/Public/RHIResources.h
@@ -3258,9 +3258,6 @@ public:
 
 	const uint32 NumRecords = 0;
 
-	virtual FUint32Vector3 GetPlatformData() const = 0;
-	virtual void CalcDispatchBufferSizes(uint32& OutDataBufferSize, uint32& OutExecutionBufferSize) const = 0;
-
 public:
 	FRHIShaderBundle(uint32 InNumRecords)
 		: FRHIResource(RRT_ShaderBundle)
