Add FRHICommandListBase argument to FRHIRayTracingGeometry::SetInitializer to remove use of FRHICommandListImmediate::Get() in platform RHIs.

diff --git a/Engine/Source/Runtime/Apple/MetalRHI/Private/MetalRayTracing.cpp b/Engine/Source/Runtime/Apple/MetalRHI/Private/MetalRayTracing.cpp
--- a/Engine/Source/Runtime/Apple/MetalRHI/Private/MetalRayTracing.cpp
+++ b/Engine/Source/Runtime/Apple/MetalRHI/Private/MetalRayTracing.cpp
@@ -392,7 +392,7 @@ void FMetalRayTracingGeometry::ReleaseUnderlyingResource()
 	}
 }
 
-void FMetalRayTracingGeometry::SetInitializer(const FRayTracingGeometryInitializer& InInitializer)
+void FMetalRayTracingGeometry::SetInitializer(FRHICommandListBase& RHICmdList, const FRayTracingGeometryInitializer& InInitializer)
 {
 	checkf(InitializedType == ERayTracingGeometryInitializerType::StreamingDestination, TEXT("Only FMetalRayTracingGeometry that was created as StreamingDestination can update their initializer."));
 	Initializer = InInitializer;
diff --git a/Engine/Source/Runtime/Apple/MetalRHI/Private/MetalRayTracing.h b/Engine/Source/Runtime/Apple/MetalRHI/Private/MetalRayTracing.h
--- a/Engine/Source/Runtime/Apple/MetalRHI/Private/MetalRayTracing.h
+++ b/Engine/Source/Runtime/Apple/MetalRHI/Private/MetalRayTracing.h
@@ -29,7 +29,7 @@ public:
 	/** FRHIRayTracingGeometry Interface */
 	virtual FRayTracingAccelerationStructureAddress GetAccelerationStructureAddress(uint64 GPUIndex) const final override { return (FRayTracingAccelerationStructureAddress)SceneIndex;
 	}
-	virtual void SetInitializer(const FRayTracingGeometryInitializer& Initializer) final override;
+	virtual void SetInitializer(FRHICommandListBase& RHICmdList, const FRayTracingGeometryInitializer& Initializer) final override;
 	/** FRHIRayTracingGeometry Interface */
 
 	void Swap(FMetalRayTracingGeometry& Other);
diff --git a/Engine/Source/Runtime/D3D12RHI/Private/D3D12RayTracing.cpp b/Engine/Source/Runtime/D3D12RHI/Private/D3D12RayTracing.cpp
--- a/Engine/Source/Runtime/D3D12RHI/Private/D3D12RayTracing.cpp
+++ b/Engine/Source/Runtime/D3D12RHI/Private/D3D12RayTracing.cpp
@@ -3094,7 +3094,7 @@ void FD3D12RayTracingGeometry::UpdateResidency(FD3D12CommandContext& CommandCont
 	CommandContext.UpdateResidency(AccelerationStructureBuffers[GPUIndex]->GetResource());
 }
 
-void FD3D12RayTracingGeometry::SetInitializer(const FRayTracingGeometryInitializer& InInitializer)
+void FD3D12RayTracingGeometry::SetInitializer(FRHICommandListBase& RHICmdList, const FRayTracingGeometryInitializer& InInitializer)
 {
 	checkf(InitializedType == ERayTracingGeometryInitializerType::StreamingDestination, TEXT("Only FD3D12RayTracingGeometry that was created as StreamingDestination can update their initializer."));
 	Initializer = InInitializer;
@@ -3111,17 +3111,15 @@ void FD3D12RayTracingGeometry::SetInitializer(const FRayTracingGeometryInitializ
 	GeometryDescs.SetNumUninitialized(Initializer.Segments.Num());
 	TranslateRayTracingGeometryDescs(Initializer, GeometryDescs);
 	
-	FRHICommandListImmediate& RHICmdList = FRHICommandListExecutor::GetImmediateCommandList();
-	RHICmdList.EnqueueLambda(
-		[this](FRHICommandListImmediate& RHICmdList)
-		{
-			for (uint32 GPUIndex = 0; GPUIndex < MAX_NUM_GPUS && GPUIndex < GNumExplicitGPUsForRendering; ++GPUIndex)
-			{				
-				RegisterAsRenameListener(GPUIndex);
-				SetupHitGroupSystemParameters(GPUIndex);		
-			}
+	RHICmdList.EnqueueLambda([this](FRHICommandListBase&)
+	{
+		for (uint32 GPUIndex = 0; GPUIndex < MAX_NUM_GPUS && GPUIndex < GNumExplicitGPUsForRendering; ++GPUIndex)
+		{				
+			RegisterAsRenameListener(GPUIndex);
+			SetupHitGroupSystemParameters(GPUIndex);		
 		}
-	);
+	});
+	RHICmdList.RHIThreadFence(true);
 }
 
 void FD3D12RayTracingGeometry::SetupHitGroupSystemParameters(uint32 InGPUIndex)
diff --git a/Engine/Source/Runtime/D3D12RHI/Private/D3D12RayTracing.h b/Engine/Source/Runtime/D3D12RHI/Private/D3D12RayTracing.h
--- a/Engine/Source/Runtime/D3D12RHI/Private/D3D12RayTracing.h
+++ b/Engine/Source/Runtime/D3D12RHI/Private/D3D12RayTracing.h
@@ -33,7 +33,7 @@ public:
 			TEXT("Trying to get address of acceleration structure '%s' without allocated memory."), *DebugName.ToString());
 		return AccelerationStructureBuffers[GPUIndex]->ResourceLocation.GetGPUVirtualAddress();
 	}
-	virtual void SetInitializer(const FRayTracingGeometryInitializer& Initializer) final override;
+	virtual void SetInitializer(FRHICommandListBase& RHICmdList, const FRayTracingGeometryInitializer& Initializer) final override;
 
 	void SetupHitGroupSystemParameters(uint32 InGPUIndex);
 	void TransitionBuffers(FD3D12CommandContext& CommandContext);
diff --git a/Engine/Source/Runtime/Engine/Internal/Streaming/SkeletalMeshUpdate.cpp b/Engine/Source/Runtime/Engine/Internal/Streaming/SkeletalMeshUpdate.cpp
--- a/Engine/Source/Runtime/Engine/Internal/Streaming/SkeletalMeshUpdate.cpp
+++ b/Engine/Source/Runtime/Engine/Internal/Streaming/SkeletalMeshUpdate.cpp
@@ -226,7 +226,7 @@ void FSkeletalMeshStreamIn::DoFinishUpdate(const FContext& Context)
 				// Under very rare circumstances that we switch ray tracing on/off right in the middle of streaming RayTracingGeometryRHI might not be valid.
 				if (IsRayTracingEnabled() && ensure(LODResource.StaticRayTracingGeometry.IsValid()))
 				{
-					LODResource.StaticRayTracingGeometry.RequestBuildIfNeeded(ERTAccelerationStructureBuildPriority::Normal);
+					LODResource.StaticRayTracingGeometry.RequestBuildIfNeeded(FRHICommandListImmediate::Get(), ERTAccelerationStructureBuildPriority::Normal);
 				}
 			}
 		}
diff --git a/Engine/Source/Runtime/Engine/Private/Streaming/StaticMeshUpdate.cpp b/Engine/Source/Runtime/Engine/Private/Streaming/StaticMeshUpdate.cpp
--- a/Engine/Source/Runtime/Engine/Private/Streaming/StaticMeshUpdate.cpp
+++ b/Engine/Source/Runtime/Engine/Private/Streaming/StaticMeshUpdate.cpp
@@ -305,7 +305,7 @@ void FStaticMeshStreamIn::DoFinishUpdate(const FContext& Context)
 				// Under very rare circumstances that we switch ray tracing on/off right in the middle of streaming RayTracingGeometryRHI might not be valid.
 				if (IsRayTracingEnabled() && ensure(LODResource.RayTracingGeometry.IsValid()))
 				{
-					LODResource.RayTracingGeometry.RequestBuildIfNeeded(ERTAccelerationStructureBuildPriority::Normal);
+					LODResource.RayTracingGeometry.RequestBuildIfNeeded(FRHICommandListImmediate::Get(), ERTAccelerationStructureBuildPriority::Normal);
 				}
 			}
 		}
diff --git a/Engine/Source/Runtime/RHI/Public/RHIResources.h b/Engine/Source/Runtime/RHI/Public/RHIResources.h
--- a/Engine/Source/Runtime/RHI/Public/RHIResources.h
+++ b/Engine/Source/Runtime/RHI/Public/RHIResources.h
@@ -3216,7 +3216,7 @@ public:
 	{}
 
 	virtual FRayTracingAccelerationStructureAddress GetAccelerationStructureAddress(uint64 GPUIndex) const = 0;
-	virtual void SetInitializer(const FRayTracingGeometryInitializer& Initializer) = 0;
+	virtual void SetInitializer(FRHICommandListBase& RHICmdList, const FRayTracingGeometryInitializer& Initializer) = 0;
 	virtual bool IsCompressed() const { return false; }
 
 	const FRayTracingGeometryInitializer& GetInitializer() const
diff --git a/Engine/Source/Runtime/RenderCore/Private/RayTracingGeometry.cpp b/Engine/Source/Runtime/RenderCore/Private/RayTracingGeometry.cpp
--- a/Engine/Source/Runtime/RenderCore/Private/RayTracingGeometry.cpp
+++ b/Engine/Source/Runtime/RenderCore/Private/RayTracingGeometry.cpp
@@ -61,9 +61,9 @@ void FRayTracingGeometry::ReleaseRHIForStreaming(FRHIResourceReplaceBatcher& Bat
 	}
 }
 
-void FRayTracingGeometry::RequestBuildIfNeeded(ERTAccelerationStructureBuildPriority InBuildPriority)
+void FRayTracingGeometry::RequestBuildIfNeeded(FRHICommandListBase& RHICmdList, ERTAccelerationStructureBuildPriority InBuildPriority)
 {
-	RayTracingGeometryRHI->SetInitializer(Initializer);
+	RayTracingGeometryRHI->SetInitializer(RHICmdList, Initializer);
 
 	if (GetRequiresBuild())
 	{
@@ -106,7 +106,7 @@ void FRayTracingGeometry::MakeResident(FRHICommandList& RHICmdList)
 			// We need to do it before we build the current geometry (also on RHI thread).
 		}
 
-		RequestBuildIfNeeded(ERTAccelerationStructureBuildPriority::Normal);
+		RequestBuildIfNeeded(RHICmdList, ERTAccelerationStructureBuildPriority::Normal);
 	}
 	else
 	{
diff --git a/Engine/Source/Runtime/RenderCore/Public/RayTracingGeometry.h b/Engine/Source/Runtime/RenderCore/Public/RayTracingGeometry.h
--- a/Engine/Source/Runtime/RenderCore/Public/RayTracingGeometry.h
+++ b/Engine/Source/Runtime/RenderCore/Public/RayTracingGeometry.h
@@ -130,7 +130,7 @@ public:
 	RENDERCORE_API void InitRHIForStreaming(FRHIRayTracingGeometry* IntermediateGeometry, FRHIResourceReplaceBatcher& Batcher);
 	RENDERCORE_API void ReleaseRHIForStreaming(FRHIResourceReplaceBatcher& Batcher);
 
-	RENDERCORE_API void RequestBuildIfNeeded(ERTAccelerationStructureBuildPriority InBuildPriority);
+	RENDERCORE_API void RequestBuildIfNeeded(FRHICommandListBase& RHICmdList, ERTAccelerationStructureBuildPriority InBuildPriority);
 
 	RENDERCORE_API void CreateRayTracingGeometry(FRHICommandListBase& RHICmdList, ERTAccelerationStructureBuildPriority InBuildPriority);
 
diff --git a/Engine/Source/Runtime/VulkanRHI/Private/VulkanRayTracing.cpp b/Engine/Source/Runtime/VulkanRHI/Private/VulkanRayTracing.cpp
--- a/Engine/Source/Runtime/VulkanRHI/Private/VulkanRayTracing.cpp
+++ b/Engine/Source/Runtime/VulkanRHI/Private/VulkanRayTracing.cpp
@@ -420,7 +420,7 @@ FVulkanRayTracingGeometry::~FVulkanRayTracingGeometry()
 	RemoveCompactionRequest();
 }
 
-void FVulkanRayTracingGeometry::SetInitializer(const FRayTracingGeometryInitializer& InInitializer)
+void FVulkanRayTracingGeometry::SetInitializer(FRHICommandListBase& RHICmdList, const FRayTracingGeometryInitializer& InInitializer)
 {
 	checkf(InitializedType == ERayTracingGeometryInitializerType::StreamingDestination, TEXT("Only FVulkanRayTracingGeometry that was created as StreamingDestination can update their initializer."));
 	Initializer = InInitializer;
diff --git a/Engine/Source/Runtime/VulkanRHI/Private/VulkanRayTracing.h b/Engine/Source/Runtime/VulkanRHI/Private/VulkanRayTracing.h
--- a/Engine/Source/Runtime/VulkanRHI/Private/VulkanRayTracing.h
+++ b/Engine/Source/Runtime/VulkanRHI/Private/VulkanRayTracing.h
@@ -139,7 +139,7 @@ public:
 	~FVulkanRayTracingGeometry();
 
 	virtual FRayTracingAccelerationStructureAddress GetAccelerationStructureAddress(uint64 GPUIndex) const final override { return Address; }	
-	virtual void SetInitializer(const FRayTracingGeometryInitializer& Initializer) final override;
+	virtual void SetInitializer(FRHICommandListBase& RHICmdList, const FRayTracingGeometryInitializer& Initializer) final override;
 
 	void Swap(FVulkanRayTracingGeometry& Other);
 
