Always log VRAM info on D3D12 GPU OOM.

diff --git a/Engine/Source/Runtime/D3D12RHI/Private/D3D12Util.cpp b/Engine/Source/Runtime/D3D12RHI/Private/D3D12Util.cpp
--- a/Engine/Source/Runtime/D3D12RHI/Private/D3D12Util.cpp
+++ b/Engine/Source/Runtime/D3D12RHI/Private/D3D12Util.cpp
@@ -829,19 +829,12 @@ void FD3D12DynamicRHI::TerminateOnOutOfMemory(ID3D12Device* InDevice, HRESULT D3
 
 	static IConsoleVariable* GPUCrashOOM = IConsoleManager::Get().FindConsoleVariable(TEXT("r.GPUCrashOnOutOfMemory"));
 	const bool bGPUCrashOOM = GPUCrashOOM && GPUCrashOOM->GetInt();
-	if (bGPUCrashOOM)
+	// If no device provided then log the memory information for each device.
+	D3D12RHI->ForEachDevice(InDevice, [&](FD3D12Device* IterationDevice)
 	{
-		// If no device provided then try and log the DRED status of each device
-		ForEachDevice(InDevice, [&](FD3D12Device* IterationDevice)
-		{
-			if (InDevice == nullptr || InDevice == IterationDevice->GetDevice())
-			{
-				FD3D12Adapter* Adapter = IterationDevice->GetParentAdapter();
-				LogMemoryInfo(Adapter);
-			}
-		});
-		UE_LOG(LogD3D12RHI, Fatal, TEXT("Out of video memory trying to allocate a rendering resource"));
-	}
+		FD3D12Adapter* Adapter = IterationDevice->GetParentAdapter();
+		LogMemoryInfo(Adapter);
+	});
 	
 	UE_LOG(LogD3D12RHI, Fatal, TEXT("Out of video memory trying to allocate a rendering resource"));
 	if (!bGPUCrashOOM)
