Fix compile error after merge.

diff --git a/Engine/Source/Runtime/D3D12RHI/Private/D3D12Util.cpp b/Engine/Source/Runtime/D3D12RHI/Private/D3D12Util.cpp
--- a/Engine/Source/Runtime/D3D12RHI/Private/D3D12Util.cpp
+++ b/Engine/Source/Runtime/D3D12RHI/Private/D3D12Util.cpp
@@ -830,7 +830,7 @@ void FD3D12DynamicRHI::TerminateOnOutOfMemory(ID3D12Device* InDevice, HRESULT D3
 	static IConsoleVariable* GPUCrashOOM = IConsoleManager::Get().FindConsoleVariable(TEXT("r.GPUCrashOnOutOfMemory"));
 	const bool bGPUCrashOOM = GPUCrashOOM && GPUCrashOOM->GetInt();
 	// If no device provided then log the memory information for each device.
-	D3D12RHI->ForEachDevice(InDevice, [&](FD3D12Device* IterationDevice)
+	ForEachDevice(InDevice, [&](FD3D12Device* IterationDevice)
 	{
 		FD3D12Adapter* Adapter = IterationDevice->GetParentAdapter();
 		LogMemoryInfo(Adapter);
