Limit DX12 Insight memory trace to when tracking all allocations is enabled

diff --git a/Engine/Source/Runtime/D3D12RHI/Private/D3D12Resources.cpp b/Engine/Source/Runtime/D3D12RHI/Private/D3D12Resources.cpp
--- a/Engine/Source/Runtime/D3D12RHI/Private/D3D12Resources.cpp
+++ b/Engine/Source/Runtime/D3D12RHI/Private/D3D12Resources.cpp
@@ -998,7 +998,7 @@ void FD3D12Adapter::TraceMemoryAllocation(FD3D12Resource* Resource)
 		const D3D12_RESOURCE_ALLOCATION_INFO Info = Resource->GetParentDevice()->GetResourceAllocationInfo(Resource->GetDesc());
 		D3D12_GPU_VIRTUAL_ADDRESS GPUAddress = Resource->GetGPUVirtualAddress();
 		// Textures don't have valid GPUVirtualAddress when IsTrackingAllAllocations() is false, so don't do memory trace in this case.
-		if (IsTrackingAllAllocations() || GPUAddress != 0)
+		if (IsTrackingAllAllocations() && GPUAddress != 0)
 		{
 			MemoryTrace_Alloc(GPUAddress, Info.SizeInBytes, Info.Alignment, EMemoryTraceRootHeap::VideoMemory);
 		}
