D3D12 SwapChain
- D3D12RHI was only using CreateSwapChainForHwnd when stereo rendering was enabled, presumably as a very old compatibility workaround. We should just rely on that code path all the time now given our hardware requirements.
- With this, we can also set the DXGI scaling mode to None when possible to address ghosting issues when resizing Editor windows.
- Moving bNeedSwapChain into FD3D12Viewport.

diff --git a/Engine/Source/Runtime/D3D12RHI/Private/D3D12Viewport.cpp b/Engine/Source/Runtime/D3D12RHI/Private/D3D12Viewport.cpp
--- a/Engine/Source/Runtime/D3D12RHI/Private/D3D12Viewport.cpp
+++ b/Engine/Source/Runtime/D3D12RHI/Private/D3D12Viewport.cpp
@@ -129,12 +129,6 @@ void FD3D12FramePacing::PrePresentQueued(ID3D12CommandQueue* Queue)
 }
 #endif //WITH_MGPU
 
-// TODO: Move this bool into D3D12Viewport.h where it belongs. It's here because it was added as a hotfix for 4.23 and we don't want to touch public headers.
-// Whether to create swap chain and use swap chain's back buffer surface,
-// or don't create swap chain and create an off-screen back buffer surface.
-// Currently used for pixel streaming plugin "windowless" mode to run in the cloud without on screen display.
-bool bNeedSwapChain = true;
-
 /**
  * Creates a FD3D12Surface to represent a swap chain's back buffer.
  */
diff --git a/Engine/Source/Runtime/D3D12RHI/Private/D3D12Viewport.h b/Engine/Source/Runtime/D3D12RHI/Private/D3D12Viewport.h
--- a/Engine/Source/Runtime/D3D12RHI/Private/D3D12Viewport.h
+++ b/Engine/Source/Runtime/D3D12RHI/Private/D3D12Viewport.h
@@ -187,6 +187,7 @@ private:
 	EPixelFormat PixelFormat;
 	bool bIsValid;
 	bool bAllowTearing;
+	bool bNeedSwapChain = false;
 
 #if D3D12_VIEWPORT_EXPOSES_SWAP_CHAIN
 	TRefCountPtr<IDXGISwapChain1> SwapChain1;
diff --git a/Engine/Source/Runtime/D3D12RHI/Private/Windows/WindowsD3D12Viewport.cpp b/Engine/Source/Runtime/D3D12RHI/Private/Windows/WindowsD3D12Viewport.cpp
--- a/Engine/Source/Runtime/D3D12RHI/Private/Windows/WindowsD3D12Viewport.cpp
+++ b/Engine/Source/Runtime/D3D12RHI/Private/Windows/WindowsD3D12Viewport.cpp
@@ -33,6 +33,7 @@ FD3D12Viewport::FD3D12Viewport(class FD3D12Adapter* InParent, HWND InWindowHandl
 	, bFullscreenLost(false)
 	, PixelFormat(InPreferredPixelFormat)
 	, bIsValid(true)
+	, bNeedSwapChain(!FParse::Param(FCommandLine::Get(), TEXT("RenderOffScreen")))
 	, ColorSpace(DXGI_COLOR_SPACE_RGB_FULL_G22_NONE_P709)
 	, NumBackBuffers(WindowsDefaultNumBackBuffers)
 	, DummyBackBuffer_RenderThread(nullptr)
@@ -97,116 +98,81 @@ void FD3D12Viewport::Init()
 
 	CalculateSwapChainDepth(WindowsDefaultNumBackBuffers);
 
-	UINT SwapChainFlags = DXGI_SWAP_CHAIN_FLAG_ALLOW_MODE_SWITCH;
-
-	if (bAllowTearing)
+	bool bStereoMode = false;
+	if (FD3D12DynamicRHI::GetD3DRHI()->IsQuadBufferStereoEnabled())
 	{
-		SwapChainFlags |= DXGI_SWAP_CHAIN_FLAG_ALLOW_TEARING;
+		if (Factory2->IsWindowedStereoEnabled())
+		{
+			bStereoMode = true;
+		}
+		else
+		{
+			UE_LOG(LogD3D12RHI, Log, TEXT("FD3D12Viewport::FD3D12Viewport was not able to create stereo SwapChain; Please enable stereo in driver settings."));
+			FD3D12DynamicRHI::GetD3DRHI()->DisableQuadBufferStereo();
+		}
 	}
 
-	const DXGI_MODE_DESC BufferDesc = SetupDXGI_MODE_DESC();
-
-	// The command queue used here is irrelevant in regard to multi - GPU as it gets overriden in the Resize
-	ID3D12CommandQueue* CommandQueue = Adapter->GetDevice(0)->GetQueue(ED3D12QueueType::Direct).D3DCommandQueue;
-
 	// Create the swapchain.
-
-	extern bool bNeedSwapChain;
-
-	bNeedSwapChain = !FParse::Param(FCommandLine::Get(), TEXT("RenderOffScreen"));
 	if (bNeedSwapChain)
 	{
-		if (FD3D12DynamicRHI::GetD3DRHI()->IsQuadBufferStereoEnabled())
+		// The command queue used here is irrelevant in regard to multi - GPU as it gets overriden in the Resize
+		ID3D12CommandQueue* CommandQueue = Adapter->GetDevice(0)->GetQueue(ED3D12QueueType::Direct).D3DCommandQueue;
+
+		UINT SwapChainFlags = DXGI_SWAP_CHAIN_FLAG_ALLOW_MODE_SWITCH;
+		if (bAllowTearing)
 		{
-			if (Factory2->IsWindowedStereoEnabled())
-			{
-				DXGI_SWAP_CHAIN_DESC1 SwapChainDesc1{};
-
-				// Enable stereo 
-				SwapChainDesc1.Stereo = true;
-				// MSAA Sample count
-				SwapChainDesc1.SampleDesc.Count = 1;
-				SwapChainDesc1.SampleDesc.Quality = 0;
-
-				SwapChainDesc1.Format = GetRenderTargetFormat(PixelFormat);
-				SwapChainDesc1.BufferUsage = DXGI_USAGE_RENDER_TARGET_OUTPUT | DXGI_USAGE_SHADER_INPUT;
-				// Double buffering required to create stereo swap chain
-				SwapChainDesc1.BufferCount = NumBackBuffers;
-				SwapChainDesc1.Scaling = DXGI_SCALING_NONE;
-				SwapChainDesc1.SwapEffect = DXGI_SWAP_EFFECT_FLIP_DISCARD;
-				SwapChainDesc1.Flags = SwapChainFlags;
-				SwapChainDesc1.Width = SizeX;
-				SwapChainDesc1.Height = SizeY;
-
-				HRESULT hr = DXGISwapchainProvider ?
-					DXGISwapchainProvider->CreateSwapChainForHwnd(Factory2, CommandQueue, WindowHandle, &SwapChainDesc1, nullptr, nullptr, SwapChain1.GetInitReference()) :
-					Factory2->CreateSwapChainForHwnd(CommandQueue, WindowHandle, &SwapChainDesc1, nullptr, nullptr, SwapChain1.GetInitReference());
-
-				VERIFYD3D12RESULT(hr);
-			}
-			else
-			{
-				UE_LOG(LogD3D12RHI, Log, TEXT("FD3D12Viewport::FD3D12Viewport was not able to create stereo SwapChain; Please enable stereo in driver settings."));
-				FD3D12DynamicRHI::GetD3DRHI()->DisableQuadBufferStereo();
-			}
+			SwapChainFlags |= DXGI_SWAP_CHAIN_FLAG_ALLOW_TEARING;
 		}
 
-		// if stereo was not activated or not enabled in settings
-		if (SwapChain1 == nullptr)
+		DXGI_SWAP_CHAIN_DESC1 SwapChainDesc1{};
+
+		SwapChainDesc1.Width       = SizeX;
+		SwapChainDesc1.Height      = SizeY;
+		SwapChainDesc1.Format      = GetRenderTargetFormat(PixelFormat);
+		SwapChainDesc1.Stereo      = bStereoMode ? TRUE : FALSE;
+		SwapChainDesc1.SampleDesc.Count = 1;
+		SwapChainDesc1.SampleDesc.Quality = 0;
+		SwapChainDesc1.BufferUsage = DXGI_USAGE_RENDER_TARGET_OUTPUT | DXGI_USAGE_SHADER_INPUT;
+		SwapChainDesc1.BufferCount = NumBackBuffers;
+		SwapChainDesc1.Scaling     = DXGI_SCALING_NONE;
+		SwapChainDesc1.SwapEffect  = DXGI_SWAP_EFFECT_FLIP_DISCARD;
+		SwapChainDesc1.AlphaMode   = DXGI_ALPHA_MODE_UNSPECIFIED;
+		SwapChainDesc1.Flags       = SwapChainFlags;
+
+		DXGI_SWAP_CHAIN_FULLSCREEN_DESC FullscreenDesc{};
+		FullscreenDesc.RefreshRate.Numerator   = 0;
+		FullscreenDesc.RefreshRate.Denominator = 0;
+		FullscreenDesc.ScanlineOrdering        = DXGI_MODE_SCANLINE_ORDER_UNSPECIFIED;
+		FullscreenDesc.Scaling                 = DXGI_MODE_SCALING_UNSPECIFIED;
+		FullscreenDesc.Windowed                = !bIsFullscreen;
+
+		HRESULT hr = DXGISwapchainProvider ?
+			DXGISwapchainProvider->CreateSwapChainForHwnd(Factory2, CommandQueue, WindowHandle, &SwapChainDesc1, &FullscreenDesc, nullptr, SwapChain1.GetInitReference()) :
+			Factory2->CreateSwapChainForHwnd(CommandQueue, WindowHandle, &SwapChainDesc1, &FullscreenDesc, nullptr, SwapChain1.GetInitReference());
+
+		if (FAILED(hr))
 		{
-			// Create the swapchain.
-			DXGI_SWAP_CHAIN_DESC SwapChainDesc = {};
-			SwapChainDesc.BufferDesc = BufferDesc;
-			// MSAA Sample count
-			SwapChainDesc.SampleDesc.Count = 1;
-			SwapChainDesc.SampleDesc.Quality = 0;
-			SwapChainDesc.BufferUsage = DXGI_USAGE_RENDER_TARGET_OUTPUT | DXGI_USAGE_SHADER_INPUT;
-			// 1:single buffering, 2:double buffering, 3:triple buffering
-			SwapChainDesc.BufferCount = NumBackBuffers;
-			SwapChainDesc.OutputWindow = WindowHandle;
-			SwapChainDesc.Windowed = !bIsFullscreen;
-			// DXGI_SWAP_EFFECT_DISCARD / DXGI_SWAP_EFFECT_SEQUENTIAL
-			SwapChainDesc.SwapEffect = DXGI_SWAP_EFFECT_FLIP_DISCARD;
-			SwapChainDesc.Flags = SwapChainFlags;
-
-			TRefCountPtr<IDXGISwapChain> SwapChain;
-			
-			HRESULT hr;
-						
-			{
-				// Don't create swap chain and release back buffer at the same time (see notes on critical section)
-				FScopeLock Lock(&DXGIBackBufferLock); 
-				hr = DXGISwapchainProvider ?
-					DXGISwapchainProvider->CreateSwapChain(Factory2, CommandQueue, &SwapChainDesc, SwapChain.GetInitReference()) :
-					Factory2->CreateSwapChain(CommandQueue, &SwapChainDesc, SwapChain.GetInitReference());
-			}
-			
-			if (FAILED(hr))
-			{
-				UE_LOG(LogD3D12RHI, Warning, TEXT("Failed to create swapchain with the following parameters:"));
-				UE_LOG(LogD3D12RHI, Warning, TEXT("\tDXGI_MODE_DESC: width: %d height: %d DXGI format: %d"), SwapChainDesc.BufferDesc.Width, SwapChainDesc.BufferDesc.Height, SwapChainDesc.BufferDesc.Format);
-				UE_LOG(LogD3D12RHI, Warning, TEXT("\tBack buffer count: %d"), NumBackBuffers);
-				UE_LOG(LogD3D12RHI, Warning, TEXT("\tWindows handle: 0x%x (IsWindow: %s)"), WindowHandle, IsWindow(WindowHandle) ? TEXT("true") : TEXT("false"));
-				UE_LOG(LogD3D12RHI, Warning, TEXT("\tFullscreen: %s"), bIsFullscreen ? TEXT("true") : TEXT("false"));
-				UE_LOG(LogD3D12RHI, Warning, TEXT("\tSwapchain flags: %d"), SwapChainFlags);
-				UE_LOG(LogD3D12RHI, Warning, TEXT("\tCustom swapchain provider: %s"), DXGISwapchainProvider ? DXGISwapchainProvider->GetProviderName() : TEXT("none"));
-
-				VERIFYD3D12RESULT(hr);
-			}
-
-			VERIFYD3D12RESULT(SwapChain->QueryInterface(IID_PPV_ARGS(SwapChain1.GetInitReference())));
+			UE_LOG(LogD3D12RHI, Warning, TEXT("Failed to create swapchain with the following parameters:"));
+			UE_LOG(LogD3D12RHI, Warning, TEXT("\tWidth: %d Height: %d DXGI format: %d"), SwapChainDesc1.Width, SwapChainDesc1.Height, SwapChainDesc1.Format);
+			UE_LOG(LogD3D12RHI, Warning, TEXT("\tBack buffer count: %d"), SwapChainDesc1.BufferCount);
+			UE_LOG(LogD3D12RHI, Warning, TEXT("\tWindows handle: 0x%x (IsWindow: %s)"), WindowHandle, IsWindow(WindowHandle) ? TEXT("true") : TEXT("false"));
+			UE_LOG(LogD3D12RHI, Warning, TEXT("\tFullscreen: %s"), bIsFullscreen ? TEXT("true") : TEXT("false"));
+			UE_LOG(LogD3D12RHI, Warning, TEXT("\tSwapchain flags: 0x%08x"), SwapChainDesc1.Flags);
+			UE_LOG(LogD3D12RHI, Warning, TEXT("\tCustom swapchain provider: %s"), DXGISwapchainProvider ? DXGISwapchainProvider->GetProviderName() : TEXT("none"));
+
+			VERIFYD3D12RESULT(hr);
 		}
-	}
 
-	if (SwapChain1)
-	{
-		SwapChain1->QueryInterface(IID_PPV_ARGS(SwapChain2.GetInitReference()));
+		if (SwapChain1)
+		{
+			SwapChain1->QueryInterface(IID_PPV_ARGS(SwapChain2.GetInitReference()));
 #if DXGI_MAX_SWAPCHAIN_INTERFACE >= 3
-		SwapChain1->QueryInterface(IID_PPV_ARGS(SwapChain3.GetInitReference()));
+			SwapChain1->QueryInterface(IID_PPV_ARGS(SwapChain3.GetInitReference()));
 #endif
 #if DXGI_MAX_SWAPCHAIN_INTERFACE >= 4
-		SwapChain1->QueryInterface(IID_PPV_ARGS(SwapChain4.GetInitReference()));
+			SwapChain1->QueryInterface(IID_PPV_ARGS(SwapChain4.GetInitReference()));
 #endif
+		}
 	}
 
 	{
@@ -218,7 +184,7 @@ void FD3D12Viewport::Init()
 	}
 
 	// Resize to setup mGPU correctly.
-	Resize(BufferDesc.Width, BufferDesc.Height, bIsFullscreen, PixelFormat);
+	Resize(SizeX, SizeY, bIsFullscreen, PixelFormat);
 
 	// Tell the window to redraw when they can.
 	// @todo: For Slate viewports, it doesn't make sense to post WM_PAINT messages (we swallow those.)
