Implement support for reserved textures with mip maps in D3D12 RHI. Allow reserved render targets and depth buffers.

diff --git a/Engine/Plugins/Tests/RHITests/Source/RHITests/Private/RHIReservedResourceTests.cpp b/Engine/Plugins/Tests/RHITests/Source/RHITests/Private/RHIReservedResourceTests.cpp
--- a/Engine/Plugins/Tests/RHITests/Source/RHITests/Private/RHIReservedResourceTests.cpp
+++ b/Engine/Plugins/Tests/RHITests/Source/RHITests/Private/RHIReservedResourceTests.cpp
@@ -168,6 +168,43 @@ bool FRHIReservedResourceTests::Test_ReservedResource_CreateTexture(FRHICommandL
 	return true;
 }
 
+bool FRHIReservedResourceTests::Test_ReservedResource_CreateTextureWithMips(FRHICommandListImmediate& RHICmdList)
+{
+	if (!GRHIGlobals.ReservedResources.Supported)
+	{
+		return true;
+	}
+
+	// Try to create a 2D texture array with mips
+
+	{
+		FRHITextureCreateDesc Desc = FRHITextureCreateDesc::Create2DArray(TEXT("TestReservedTextureArray2DWithMips"));
+		Desc.SetFlags(TexCreate_ShaderResource | TexCreate_RenderTargetable | TexCreate_ReservedResource | TexCreate_ImmediateCommit)
+			.SetExtent(FIntPoint(1024, 1024))
+			.SetArraySize(2)
+			.SetNumMips(3)
+			.SetFormat(PF_R8)
+			.SetInitialState(ERHIAccess::SRVCompute);
+
+		FTextureRHIRef Texture = RHICreateTexture(Desc);
+	}
+
+	// Try to create a 2D texture with mips
+
+	{
+		FRHITextureCreateDesc Desc = FRHITextureCreateDesc::Create2D(TEXT("TestReservedTexture2DWithMips"));
+		Desc.SetFlags(TexCreate_ShaderResource | TexCreate_ReservedResource | TexCreate_ImmediateCommit)
+			.SetExtent(FIntPoint(128, 128))
+			.SetNumMips(8) // full mip chain down to 1x1
+			.SetFormat(PF_A32B32G32R32F)
+			.SetInitialState(ERHIAccess::SRVCompute);
+
+		FTextureRHIRef Texture = RHICreateTexture(Desc);
+	}
+
+	return true;
+}
+
 bool FRHIReservedResourceTests::Test_ReservedResource_CreateBuffer(FRHICommandListImmediate& RHICmdList)
 {
 	if (!GRHIGlobals.ReservedResources.Supported)
diff --git a/Engine/Plugins/Tests/RHITests/Source/RHITests/Private/RHITests.spec.cpp b/Engine/Plugins/Tests/RHITests/Source/RHITests/Private/RHITests.spec.cpp
--- a/Engine/Plugins/Tests/RHITests/Source/RHITests/Private/RHITests.spec.cpp
+++ b/Engine/Plugins/Tests/RHITests/Source/RHITests/Private/RHITests.spec.cpp
@@ -150,6 +150,12 @@ void FAutomationRHITest::Define()
 			TestEqual("Create Reserved Texture failed", bResult, 1);
 		});
 
+		It("Create Reserved Texture With Mips", [this]()
+		{
+			bool bResult = RunOnRenderThreadSynchronous(FRHIReservedResourceTests::Test_ReservedResource_CreateTextureWithMips);
+			TestEqual("Create Reserved Texture With Mips failed", bResult, 1);
+		});
+
 		It("Create Reserved Buffer", [this]()
 		{
 			bool bResult = RunOnRenderThreadSynchronous(FRHIReservedResourceTests::Test_ReservedResource_CreateBuffer);
diff --git a/Engine/Plugins/Tests/RHITests/Source/RHITests/Private/RHITestsModule.cpp b/Engine/Plugins/Tests/RHITests/Source/RHITests/Private/RHITestsModule.cpp
--- a/Engine/Plugins/Tests/RHITests/Source/RHITests/Private/RHITestsModule.cpp
+++ b/Engine/Plugins/Tests/RHITests/Source/RHITests/Private/RHITestsModule.cpp
@@ -27,6 +27,7 @@ static bool RunTests_RenderThread(FRHICommandListImmediate& RHICmdList)
 		RUN_TEST(FRHIReservedResourceTests::Test_ReservedResource_DecommitBuffer(RHICmdList));
 		RUN_TEST(FRHIReservedResourceTests::Test_ReservedResource_CreateBuffer(RHICmdList));
 		RUN_TEST(FRHIReservedResourceTests::Test_ReservedResource_CreateTexture(RHICmdList));
+		RUN_TEST(FRHIReservedResourceTests::Test_ReservedResource_CreateTextureWithMips(RHICmdList));
 		RUN_TEST(FRHIReservedResourceTests::Test_ReservedResource_CreateVolumeTexture(RHICmdList));
 	}
 
diff --git a/Engine/Plugins/Tests/RHITests/Source/RHITests/Public/RHIReservedResourceTests.h b/Engine/Plugins/Tests/RHITests/Source/RHITests/Public/RHIReservedResourceTests.h
--- a/Engine/Plugins/Tests/RHITests/Source/RHITests/Public/RHIReservedResourceTests.h
+++ b/Engine/Plugins/Tests/RHITests/Source/RHITests/Public/RHIReservedResourceTests.h
@@ -9,6 +9,7 @@ class FRHIReservedResourceTests
 public:
 
 	static bool Test_ReservedResource_CreateTexture(FRHICommandListImmediate& RHICmdList);
+	static bool Test_ReservedResource_CreateTextureWithMips(FRHICommandListImmediate& RHICmdList);
 	static bool Test_ReservedResource_CreateVolumeTexture(FRHICommandListImmediate& RHICmdList);
 	static bool Test_ReservedResource_CreateBuffer(FRHICommandListImmediate& RHICmdList);
 	static bool Test_ReservedResource_CommitBuffer(FRHICommandListImmediate& RHICmdList);
diff --git a/Engine/Source/Runtime/D3D12RHI/Private/D3D12Resources.cpp b/Engine/Source/Runtime/D3D12RHI/Private/D3D12Resources.cpp
--- a/Engine/Source/Runtime/D3D12RHI/Private/D3D12Resources.cpp
+++ b/Engine/Source/Runtime/D3D12RHI/Private/D3D12Resources.cpp
@@ -191,25 +191,39 @@ void FD3D12Resource::CommitReservedResource(ID3D12CommandQueue* D3DCommandQueue,
 			TEXT("Current RHI does not support reserved volume textures"));
 	}
 
-	if (Desc.Dimension != D3D12_RESOURCE_DIMENSION_BUFFER)
-	{
-		checkf(Desc.MipLevels == 1, TEXT("CommitReservedResource is currently only implemented for textures without mips"));
-	}
-
 	uint32 D3DResourceNumTiles = 0;
 	D3D12_PACKED_MIP_INFO PackedMipDesc = {};
 	D3D12_TILE_SHAPE TileShape = {};
 	const uint32 FirstSubresource = 0;
+
+	const bool bBuffer = Desc.Dimension == D3D12_RESOURCE_DIMENSION_BUFFER;
+
 	const uint32 NumSubresources = SubresourceCount;
+	const uint32 NumMipLevels = GetMipLevels();
+	const uint32 NumArraySlices = GetArraySize();
 
-	// We assume that all subresources in a 2D texture array are identical, so only query the tiling config for the first
-	uint32 NumSubresourceTilings = 1;
-	D3D12_SUBRESOURCE_TILING SubresourceTiling = {}; 
+	uint32 NumSubresourceTilings = NumMipLevels;
+	TArray<D3D12_SUBRESOURCE_TILING, TInlineAllocator<16>> MipTilingInfo;
+	
+	check(NumSubresourceTilings >= 1);
+	MipTilingInfo.SetNum(NumSubresourceTilings);
 
 	ID3D12Device* D3DDevice = GetParentDevice()->GetDevice();
 	FD3D12Adapter* Adapter = GetParentDevice()->GetParentAdapter();
 
-	D3DDevice->GetResourceTiling(GetResource(), &D3DResourceNumTiles, &PackedMipDesc, &TileShape, &NumSubresourceTilings, FirstSubresource, &SubresourceTiling);
+	D3DDevice->GetResourceTiling(GetResource(), &D3DResourceNumTiles, &PackedMipDesc, &TileShape, &NumSubresourceTilings, FirstSubresource, MipTilingInfo.GetData());
+
+	if (bBuffer)
+	{
+		// Buffers obviously don't have mips, but we can pretend they do to make the code below agnostic to resource type
+		PackedMipDesc.NumStandardMips = 1;
+	}
+
+	check(MipTilingInfo.Num() == PackedMipDesc.NumStandardMips + PackedMipDesc.NumPackedMips);
+
+	const uint32 NumPackedTilesPerArraySlice = PackedMipDesc.NumTilesForPackedMips;
+	const uint32 NumTotalPackedMipTiles = NumPackedTilesPerArraySlice * NumArraySlices;
+	const uint32 NumTotalStandardMipTiles = D3DResourceNumTiles - NumTotalPackedMipTiles;
 
 	const uint64 TotalSize = D3DResourceNumTiles * TileSizeInBytes;
 
@@ -225,7 +239,8 @@ void FD3D12Resource::CommitReservedResource(ID3D12CommandQueue* D3DCommandQueue,
 
 	// Set high residency priority based on the same heuristics as D3D12 committed resources,
 	// i.e. normal priority unless it's a UAV/RT/DS texture.
-	const bool bHighPriorityResource = EnumHasAnyFlags(Desc.Flags, D3D12_RESOURCE_FLAG_ALLOW_RENDER_TARGET | D3D12_RESOURCE_FLAG_ALLOW_DEPTH_STENCIL | D3D12_RESOURCE_FLAG_ALLOW_UNORDERED_ACCESS);
+	const bool bRenderOrDepthTarget = EnumHasAnyFlags(Desc.Flags, D3D12_RESOURCE_FLAG_ALLOW_RENDER_TARGET | D3D12_RESOURCE_FLAG_ALLOW_DEPTH_STENCIL);
+	const bool bHighPriorityResource = bRenderOrDepthTarget || EnumHasAnyFlags(Desc.Flags, D3D12_RESOURCE_FLAG_ALLOW_UNORDERED_ACCESS);
 	const uint32 GPUIndex = GetParentDevice()->GetGPUIndex();
 
 	D3D12_HEAP_PROPERTIES BackingHeapProps = {};
@@ -235,63 +250,86 @@ void FD3D12Resource::CommitReservedResource(ID3D12CommandQueue* D3DCommandQueue,
 	BackingHeapProps.CreationNodeMask = GetGPUMask().GetNative();
 	BackingHeapProps.VisibleNodeMask = GetVisibilityMask().GetNative();
 
-	uint32 NumStandardTilesPerSubresource = 0;
+	uint32 NumStandardTilesPerArraySlice = 0;
 	uint32 NumTotalTiles = 0;
 
-	if (Desc.Dimension == D3D12_RESOURCE_DIMENSION_BUFFER)
+	if (bBuffer)
 	{
-		NumTotalTiles = NumStandardTilesPerSubresource = D3DResourceNumTiles;
-		checkf(D3DResourceNumTiles == SubresourceTiling.WidthInTiles,
+		NumTotalTiles = NumStandardTilesPerArraySlice = D3DResourceNumTiles;
+		checkf(D3DResourceNumTiles == MipTilingInfo[0].WidthInTiles,
 			TEXT("Reserved buffers are expected to have trivial tiling configuration: single 1D subresource that contains all tiles."));
 	}
-	else if (PackedMipDesc.NumStandardMips != 0)
-	{
-		NumStandardTilesPerSubresource = SubresourceTiling.WidthInTiles * SubresourceTiling.HeightInTiles * SubresourceTiling.DepthInTiles;
-		NumTotalTiles = NumStandardTilesPerSubresource * NumSubresources;
-	}
-	else // packed mip case
+	else
 	{
-		checkf(PackedMipDesc.NumPackedMips != 0, TEXT("It is expected that reserved resources have at least one standard or one packed tile mip level"));
-		NumTotalTiles = PackedMipDesc.NumTilesForPackedMips * NumSubresources;
+		NumStandardTilesPerArraySlice = NumTotalStandardMipTiles / NumArraySlices;
+		NumTotalTiles = (NumStandardTilesPerArraySlice + NumPackedTilesPerArraySlice) * NumArraySlices;
 	}
 
+	const uint32 NumTotalTilesPerArraySlice = NumStandardTilesPerArraySlice + NumPackedTilesPerArraySlice;
+
 	checkf(D3DResourceNumTiles == NumTotalTiles,
 		TEXT("D3D resource size in tiles: %d, computed size in tiles: %d"),
 		D3DResourceNumTiles, NumTotalTiles);
 
 	const uint32 NumRequiredCommitTiles = RequiredCommitSizeInBytes / TileSizeInBytes;
-	const uint32 NumTilesPerSlice = SubresourceTiling.WidthInTiles * SubresourceTiling.HeightInTiles;
 
-	auto GetTiledResourceCoordinate = [SubresourceTiling, NumStandardTilesPerSubresource, NumSubresources, NumTilesPerSlice, MaxTilesPerHeap]
-		(uint32 OffsetInTiles, uint32 NumTiles) -> D3D12_TILED_RESOURCE_COORDINATE 
+	auto GetTiledResourceCoordinate = [&MipTilingInfo, &PackedMipDesc, D3DResourceNumTiles, NumTotalTilesPerArraySlice, NumSubresources, MaxTilesPerHeap]
+		(uint32 OffsetInTiles, uint32 NumTiles) -> D3D12_TILED_RESOURCE_COORDINATE
 	{
+		check(OffsetInTiles < D3DResourceNumTiles);
+
+		const uint32 ArraySliceIndex = OffsetInTiles / NumTotalTilesPerArraySlice;
+		const uint32 TileIndexInArraySlice = OffsetInTiles % NumTotalTilesPerArraySlice;
+		const uint32 NumTotalMips = MipTilingInfo.Num();
+
+		uint32 MipLevel = 0;
+
+		{
+			uint32 NextMipTileThreshold = 0;
+			while (MipLevel < PackedMipDesc.NumStandardMips)
+			{
+				const D3D12_SUBRESOURCE_TILING& CurrentMipTiling = MipTilingInfo[MipLevel];
+				NextMipTileThreshold += CurrentMipTiling.WidthInTiles * CurrentMipTiling.HeightInTiles * CurrentMipTiling.DepthInTiles;
+
+				if (TileIndexInArraySlice < NextMipTileThreshold)
+				{
+					break;
+				}
+
+				MipLevel += 1;
+			}
+		}
+
 		D3D12_TILED_RESOURCE_COORDINATE ResourceCoordinate = {}; // Coordinates are in tiles, not pixels
-		if (NumStandardTilesPerSubresource)
+
+		ResourceCoordinate.Subresource = MipLevel + ArraySliceIndex * NumTotalMips;
+
+		const D3D12_SUBRESOURCE_TILING& CurrentMipTiling = MipTilingInfo[MipLevel];
+
+		if (MipLevel < PackedMipDesc.NumStandardMips)
 		{
-			const uint32 TileIndexInSubresource = OffsetInTiles % NumStandardTilesPerSubresource;
-			ResourceCoordinate.X = TileIndexInSubresource % SubresourceTiling.WidthInTiles;
-			ResourceCoordinate.Y = (TileIndexInSubresource / SubresourceTiling.WidthInTiles) % SubresourceTiling.HeightInTiles;
-			ResourceCoordinate.Z = TileIndexInSubresource / NumTilesPerSlice;
+			// Standard mip level case
+
+			check(CurrentMipTiling.StartTileIndexInOverallResource != ~0u)
+
+			const uint32 NumTilesPerVolumeSlice = CurrentMipTiling.WidthInTiles * CurrentMipTiling.HeightInTiles;
 
-			ResourceCoordinate.Subresource = OffsetInTiles / NumStandardTilesPerSubresource;
-			check(ResourceCoordinate.Subresource <= NumSubresources);
+			const uint32 TileIndexInMipLevel = TileIndexInArraySlice - CurrentMipTiling.StartTileIndexInOverallResource;
+
+			ResourceCoordinate.X = TileIndexInMipLevel % CurrentMipTiling.WidthInTiles;
+			ResourceCoordinate.Y = (TileIndexInMipLevel / CurrentMipTiling.WidthInTiles) % CurrentMipTiling.HeightInTiles;
+			ResourceCoordinate.Z = TileIndexInMipLevel / NumTilesPerVolumeSlice;
 		}
 		else
 		{
-			// Packed mip level case:
-			// - Only simple textures are expected (single subresource, no arrays)
-			// - Entire packed mip level must be covered in one map operation, so mapping origin is always 0
-
-			checkf(NumSubresources == 1,
-			       TEXT("Reserved textures with packed mips and multiple subresources are not supported. Current subresource count: %d"),
-			       NumSubresources);
+			// Packed mip level case
 
 			checkf(NumTiles <= MaxTilesPerHeap,
 			       TEXT("Reserved texture packed mip level requires tiles: %d, maximum supported tiles: %d. ")
 			       TEXT("Increase d3d12.ReservedResourceHeapSizeMB or avoid packed mips by using a larger texture dimensions."),
 			       NumTiles, MaxTilesPerHeap);
 
-			ResourceCoordinate.Subresource = 0; // Packed mips are currently supported for single subresource textures
+			// Entire packed mip chain must be covered in one map operation, so mapping origin is always 0
 			ResourceCoordinate.X = 0;
 			ResourceCoordinate.Y = 0;
 			ResourceCoordinate.Z = 0;
@@ -303,7 +341,7 @@ void FD3D12Resource::CommitReservedResource(ID3D12CommandQueue* D3DCommandQueue,
 	TArray<FD3D12UpdateTileMappingsParams> MappingParams;
 	TArray<FD3D12ResidencyHandle*> UsedResidencyHandles;
 
-	if (ReservedResourceData->NumCommittedTiles > NumRequiredCommitTiles)
+	if (ReservedResourceData->NumCommittedTiles > NumRequiredCommitTiles) // Decommit / shrink case
 	{
 		check(!ReservedResourceData->BackingHeaps.IsEmpty());
 
@@ -362,7 +400,7 @@ void FD3D12Resource::CommitReservedResource(ID3D12CommandQueue* D3DCommandQueue,
 			ReservedResourceData->NumCommittedTiles -= RegionSize.NumTiles;
 		}
 	}
-	else
+	else // Commit / grow case
 	{
 		while (ReservedResourceData->NumCommittedTiles < NumRequiredCommitTiles)
 		{
@@ -407,11 +445,14 @@ void FD3D12Resource::CommitReservedResource(ID3D12CommandQueue* D3DCommandQueue,
 				const TCHAR* HeapNameChars = TEXT("ReservedResourceBackingHeap");
 #endif // NAME_OBJECTS
 
-				const D3D12_HEAP_FLAGS HeapFlags = Desc.Dimension == D3D12_RESOURCE_DIMENSION_BUFFER
-					? D3D12_HEAP_FLAG_ALLOW_ONLY_BUFFERS
+				const D3D12_HEAP_FLAGS TextureHeapFlags = bRenderOrDepthTarget
+					? D3D12_HEAP_FLAG_ALLOW_ONLY_RT_DS_TEXTURES 
 					: D3D12_HEAP_FLAG_ALLOW_ONLY_NON_RT_DS_TEXTURES;
 
+				const D3D12_HEAP_FLAGS HeapFlags = bBuffer ? D3D12_HEAP_FLAG_ALLOW_ONLY_BUFFERS : TextureHeapFlags;
+
 				static_assert((D3D12_HEAP_FLAG_DENY_BUFFERS | D3D12_HEAP_FLAG_DENY_RT_DS_TEXTURES) == D3D12_HEAP_FLAG_ALLOW_ONLY_NON_RT_DS_TEXTURES);
+				static_assert((D3D12_HEAP_FLAG_DENY_BUFFERS | D3D12_HEAP_FLAG_DENY_NON_RT_DS_TEXTURES) == D3D12_HEAP_FLAG_ALLOW_ONLY_RT_DS_TEXTURES);
 
 				const uint32 ThisHeapSize = RegionSize.NumTiles * TileSizeInBytes;
 				D3D12_HEAP_DESC NewHeapDesc = {};
diff --git a/Engine/Source/Runtime/D3D12RHI/Private/D3D12Texture.cpp b/Engine/Source/Runtime/D3D12RHI/Private/D3D12Texture.cpp
--- a/Engine/Source/Runtime/D3D12RHI/Private/D3D12Texture.cpp
+++ b/Engine/Source/Runtime/D3D12RHI/Private/D3D12Texture.cpp
@@ -351,11 +351,9 @@ FD3D12ResourceDesc FD3D12DynamicRHI::GetResourceDesc(const FRHITextureDesc& Text
 		ResourceDesc.Flags |= D3D12_RESOURCE_FLAG_ALLOW_UNORDERED_ACCESS;
 	}
 
-	// Only 2D textures without mips are implemented/supported, to avoid the complexity associated with packed mips
 	if (EnumHasAllFlags(TextureDesc.Flags, TexCreate_ReservedResource))
 	{
 		checkf(GRHIGlobals.ReservedResources.Supported, TEXT("Reserved resources resources are not supported on this machine"));
-		checkf(TextureDesc.NumMips == 1, TEXT("Reserved resources with mips are not supported"));
 		checkf(TextureDesc.IsTexture2D() || TextureDesc.IsTexture3D(), TEXT("Only 2D and 3D textures can be created as reserved resources"));
 		checkf(!TextureDesc.IsTexture3D() || GRHIGlobals.ReservedResources.SupportsVolumeTextures, TEXT("Current RHI does not support reserved volume textures"));
 
diff --git a/Engine/Source/Runtime/RHI/Private/RHI.cpp b/Engine/Source/Runtime/RHI/Private/RHI.cpp
--- a/Engine/Source/Runtime/RHI/Private/RHI.cpp
+++ b/Engine/Source/Runtime/RHI/Private/RHI.cpp
@@ -1735,19 +1735,18 @@ bool FRHITextureDesc::Validate(const FRHITextureCreateInfo& Desc, const TCHAR* N
 				Name, GetTextureDimensionString(Desc.Dimension));
 		}
 
-		ValidateResourceDesc(Desc.NumMips == 1,
-			TEXT("Reserved Texture %s's NumMips=%d is invalid. Expected only 1 mip level."),
-			Name, Desc.NumMips);
-
 		if (Desc.Dimension == ETextureDimension::Texture2DArray)
 		{
-			ValidateResourceDesc(Desc.Extent.X >= GRHIGlobals.ReservedResources.TextureArrayMinimumMipDimension,
-				TEXT("Reserved Texture array %s's Desc.Extent.X=%d is invalid. It is required to be be no less than %d."),
-				Name, Desc.Extent.X, GRHIGlobals.ReservedResources.TextureArrayMinimumMipDimension);
+			const uint32 MipShift = FMath::Max<uint32>(1u, Desc.NumMips) - 1;
+			const FIntPoint SmallestMipExtent = Desc.Extent / (1 << MipShift);
+
+			ValidateResourceDesc(SmallestMipExtent.X >= GRHIGlobals.ReservedResources.TextureArrayMinimumMipDimension,
+				TEXT("Reserved Texture array %s's SmallestMipExtent.X=%d is invalid. It is required to be be no less than %d."),
+				Name, SmallestMipExtent.X, GRHIGlobals.ReservedResources.TextureArrayMinimumMipDimension);
 
-			ValidateResourceDesc(Desc.Extent.Y >= GRHIGlobals.ReservedResources.TextureArrayMinimumMipDimension,
-				TEXT("Reserved Texture array %s's Desc.Extent.Y=%d is invalid. It is required to be be no less than %d."),
-				Name, Desc.Extent.Y, GRHIGlobals.ReservedResources.TextureArrayMinimumMipDimension);
+			ValidateResourceDesc(SmallestMipExtent.Y >= GRHIGlobals.ReservedResources.TextureArrayMinimumMipDimension,
+				TEXT("Reserved Texture array %s's SmallestMipExtent.Y=%d is invalid. It is required to be be no less than %d."),
+				Name, SmallestMipExtent.Y, GRHIGlobals.ReservedResources.TextureArrayMinimumMipDimension);
 		}
 	}
 
