Workaround performance issues with the GPU profiler by batching timestamp queries.
 - Extended RHI[Begin|End]OcclusionQueryBatch to support timestamp queries, and renamed to RHI[Begin|End]RenderQueryBatch
 - The GPU profiler starts a query batch in BeginFrame, which persists on the immediate command list until EndFrame.
 - RDG forwards the current batch to the parallel command lists it creates.
 - Platform RHIs use the active batch to group completion events / sync points to reduce overhead. If no batch is active, they fall back to how they worked before, creating individual sync points per query.

diff --git a/Engine/Source/Runtime/D3D12RHI/Private/D3D12Query.cpp b/Engine/Source/Runtime/D3D12RHI/Private/D3D12Query.cpp
--- a/Engine/Source/Runtime/D3D12RHI/Private/D3D12Query.cpp
+++ b/Engine/Source/Runtime/D3D12RHI/Private/D3D12Query.cpp
@@ -159,21 +159,22 @@ FRenderQueryRHIRef FD3D12DynamicRHI::RHICreateRenderQuery(ERenderQueryType Query
 	});
 }
 
-void FD3D12DynamicRHI::RHIBeginOcclusionQueryBatch_TopOfPipe(FRHICommandListBase& RHICmdList, uint32 NumQueriesInBatch)
+void FD3D12DynamicRHI::RHIBeginRenderQueryBatch_TopOfPipe(FRHICommandListBase& RHICmdList, ERenderQueryType QueryType)
 {
-	// Each occlusion query batch uses a single sync point to signal when the results are ready (one per active GPU).
+	// Each query batch uses a single sync point to signal when the results are ready (one per active GPU).
 	for (uint32 GPUIndex : RHICmdList.GetGPUMask())
 	{
-		checkf(RHICmdList.QueryBatchData[GPUIndex] == nullptr, TEXT("An occlusion query batch has already begun on this command list."));
+		auto& QueryBatchData = RHICmdList.GetQueryBatchData(QueryType);
+		checkf(QueryBatchData[GPUIndex] == nullptr, TEXT("A query batch for this type has already begun on this command list."));
 
 		FD3D12SyncPointRef SyncPoint = FD3D12SyncPoint::Create(ED3D12SyncPointType::GPUAndCPU);
 
 		// Keep a reference on the RHI command list, so we can retrieve it later in BeginQuery/EndQuery/EndBatch.
-		RHICmdList.QueryBatchData[GPUIndex] = SyncPoint.GetReference();
+		QueryBatchData[GPUIndex] = SyncPoint.GetReference();
 		SyncPoint->AddRef();
 	}
 
-	if (RHIConsoleVariables::GInsertOuterOcclusionQuery)
+	if (QueryType == RQT_Occlusion && RHIConsoleVariables::GInsertOuterOcclusionQuery)
 	{
 		// Insert an outer query that encloses the whole batch
 		RHICmdList.EnqueueLambda([](FRHICommandListBase& ExecutingCmdList)
@@ -192,22 +193,6 @@ void FD3D12DynamicRHI::RHIBeginOcclusionQueryBatch_TopOfPipe(FRHICommandListBase
 	}
 }
 
-void FD3D12DynamicRHI::RHIBeginRenderQuery_TopOfPipe(FRHICommandListBase& RHICmdList, FRHIRenderQuery* RenderQuery)
-{
-	// Set the current occlusion query batch sync point into each query object.
-	for (uint32 GPUIndex : RHICmdList.GetGPUMask())
-	{
-		FD3D12RenderQuery* Query = ResourceCast(RenderQuery, GPUIndex);
-		checkf(Query->Type == RQT_Occlusion, TEXT("Only occlusion queries support RHIBeginRenderQuery()."));
-
-		checkf(RHICmdList.QueryBatchData[GPUIndex], TEXT("Cannot use an occlusion query outside of an occlusion query batch."));
-		Query->SyncPoint = static_cast<FD3D12SyncPoint*>(RHICmdList.QueryBatchData[GPUIndex]);
-	}
-
-	// Enqueue the RHI command to record the BeginQuery() call on the context.
-	FDynamicRHI::RHIBeginRenderQuery_TopOfPipe(RHICmdList, RenderQuery);
-}
-
 void FD3D12CommandContext::RHIBeginRenderQuery(FRHIRenderQuery* QueryRHI)
 {
 	FD3D12RenderQuery* Query = RetrieveObject<FD3D12RenderQuery>(QueryRHI);
@@ -224,14 +209,16 @@ void FD3D12DynamicRHI::RHIEndRenderQuery_TopOfPipe(FRHICommandListBase& RHICmdLi
 	for (uint32 GPUIndex : RHICmdList.GetGPUMask())
 	{
 		FD3D12RenderQuery* Query = ResourceCast(RenderQuery, GPUIndex);
-		if (Query->Type == RQT_Occlusion)
+		auto& QueryBatchData = RHICmdList.GetQueryBatchData(Query->Type);
+
+		if (QueryBatchData[GPUIndex])
 		{
-			// Occlusion query sync points are allocated by BeginOcclusionQueryBatch().
-			checkf(RHICmdList.QueryBatchData[GPUIndex], TEXT("Cannot use an occlusion query outside of an occlusion query batch."));
+			// This query belongs to a batch. Use the sync point we created earlier
+			Query->SyncPoint = static_cast<FD3D12SyncPoint*>(QueryBatchData[0]);
 		}
 		else
 		{
-			// All other query types use one sync point per query.
+			// Queries issued outside of a batch use one sync point per query.
 			Query->SyncPoint = FD3D12SyncPoint::Create(ED3D12SyncPointType::GPUAndCPU);
 
 			RHICmdList.EnqueueLambda([SyncPoint = Query->SyncPoint, GPUIndex](FRHICommandListBase& ExecutingCmdList) mutable
@@ -269,29 +256,33 @@ void FD3D12CommandContext::RHIEndRenderQuery(FRHIRenderQuery* QueryRHI)
 	}
 }
 
-void FD3D12DynamicRHI::RHIEndOcclusionQueryBatch_TopOfPipe(FRHICommandListBase& RHICmdList)
+void FD3D12DynamicRHI::RHIEndRenderQueryBatch_TopOfPipe(FRHICommandListBase& RHICmdList, ERenderQueryType QueryType)
 {
 	for (uint32 GPUIndex : RHICmdList.GetGPUMask())
 	{
-		checkf(RHICmdList.QueryBatchData[GPUIndex], TEXT("An occlusion query batch is not open on this command list."));
-		FD3D12SyncPointRef SyncPoint = static_cast<FD3D12SyncPoint*>(RHICmdList.QueryBatchData[GPUIndex]);
+		auto& QueryBatchData = RHICmdList.GetQueryBatchData(QueryType);
+		checkf(QueryBatchData[GPUIndex], TEXT("A query batch for this type is not open on this command list."));
+
+		FD3D12SyncPointRef SyncPoint = static_cast<FD3D12SyncPoint*>(QueryBatchData[GPUIndex]);
 
 		// Clear the sync point reference on the RHI command list
 		SyncPoint->Release();
-		RHICmdList.QueryBatchData[GPUIndex] = nullptr;
+		QueryBatchData[GPUIndex] = nullptr;
 
-		RHICmdList.EnqueueLambda([GPUIndex, SyncPoint = MoveTemp(SyncPoint)](FRHICommandListBase& ExecutingCmdList)
+		RHICmdList.EnqueueLambda([GPUIndex, SyncPoint = MoveTemp(SyncPoint), QueryType](FRHICommandListBase& ExecutingCmdList)
 		{
 			FD3D12CommandContext& Context = FD3D12CommandContext::Get(ExecutingCmdList, GPUIndex);
+			Context.BatchedSyncPoints.ToSignal.Add(SyncPoint);
 
-			// End the outer query
-			if (Context.bOuterOcclusionQuerySubmitted)
+			if (QueryType == RQT_Occlusion)
 			{
-				Context.RHIEndRenderQuery(Context.OuterOcclusionQuery);
-				Context.bOuterOcclusionQuerySubmitted = false;
+				// End the outer query
+				if (Context.bOuterOcclusionQuerySubmitted)
+				{
+					Context.RHIEndRenderQuery(Context.OuterOcclusionQuery);
+					Context.bOuterOcclusionQuerySubmitted = false;
+				}
 			}
-
-			Context.BatchedSyncPoints.ToSignal.Add(SyncPoint);
 		});
 	}
 }
diff --git a/Engine/Source/Runtime/D3D12RHI/Private/D3D12RHIPrivate.h b/Engine/Source/Runtime/D3D12RHI/Private/D3D12RHIPrivate.h
--- a/Engine/Source/Runtime/D3D12RHI/Private/D3D12RHIPrivate.h
+++ b/Engine/Source/Runtime/D3D12RHI/Private/D3D12RHIPrivate.h
@@ -356,9 +356,8 @@ public:
 	virtual void RHIRead3DSurfaceFloatData(FRHITexture* Texture, FIntRect Rect, FIntPoint ZMinMax, TArray<FFloat16Color>& OutData) final override;
 	virtual void RHIRead3DSurfaceFloatData(FRHITexture* Texture, FIntRect Rect, FIntPoint ZMinMax, TArray<FFloat16Color>& OutData, FReadSurfaceDataFlags InFlags) final override;
 	virtual FRenderQueryRHIRef RHICreateRenderQuery(ERenderQueryType QueryType) final override;
-	virtual void RHIBeginOcclusionQueryBatch_TopOfPipe(FRHICommandListBase& RHICmdList, uint32 NumQueriesInBatch) final override;
-	virtual void RHIEndOcclusionQueryBatch_TopOfPipe(FRHICommandListBase& RHICmdList) final override;
-	virtual void RHIBeginRenderQuery_TopOfPipe(FRHICommandListBase& RHICmdList, FRHIRenderQuery* RenderQuery) final override;
+	virtual void RHIBeginRenderQueryBatch_TopOfPipe(FRHICommandListBase& RHICmdList, ERenderQueryType QueryType) final override;
+	virtual void RHIEndRenderQueryBatch_TopOfPipe(FRHICommandListBase& RHICmdList, ERenderQueryType QueryType) final override;
 	virtual void RHIEndRenderQuery_TopOfPipe(FRHICommandListBase& RHICmdList, FRHIRenderQuery* RenderQuery) final override;
 	virtual bool RHIGetRenderQueryResult(FRHIRenderQuery* RenderQuery, uint64& OutResult, bool bWait, uint32 GPUIndex = INDEX_NONE) final override;
 	virtual uint32 RHIGetViewportNextPresentGPUIndex(FRHIViewport* Viewport) final override;
diff --git a/Engine/Source/Runtime/RHI/Private/RHICommandList.cpp b/Engine/Source/Runtime/RHI/Private/RHICommandList.cpp
--- a/Engine/Source/Runtime/RHI/Private/RHICommandList.cpp
+++ b/Engine/Source/Runtime/RHI/Private/RHICommandList.cpp
@@ -144,7 +144,7 @@ FRHICommandListBase::~FRHICommandListBase()
 
 	checkf(!HasCommands() || IsExecuting(), TEXT("FRHICommandListBase has been deleted while it still contained commands. The command list was not submitted."));
 
-	for (void* Data : QueryBatchData)
+	for (void* Data : PersistentState.QueryBatchData_Occlusion)
 	{
 		check(Data == nullptr);
 	}
diff --git a/Engine/Source/Runtime/RHI/Public/DynamicRHI.h b/Engine/Source/Runtime/RHI/Public/DynamicRHI.h
--- a/Engine/Source/Runtime/RHI/Public/DynamicRHI.h
+++ b/Engine/Source/Runtime/RHI/Public/DynamicRHI.h
@@ -637,8 +637,8 @@ public:
 	// FlushType: Wait RHI Thread
 	virtual FRenderQueryRHIRef RHICreateRenderQuery(ERenderQueryType QueryType) = 0;
 
-	virtual void RHIBeginOcclusionQueryBatch_TopOfPipe(FRHICommandListBase& RHICmdList, uint32 NumQueriesInBatch) {}
-	virtual void RHIEndOcclusionQueryBatch_TopOfPipe  (FRHICommandListBase& RHICmdList) {}
+	virtual void RHIBeginRenderQueryBatch_TopOfPipe(FRHICommandListBase& RHICmdList, ERenderQueryType QueryType) {}
+	virtual void RHIEndRenderQueryBatch_TopOfPipe  (FRHICommandListBase& RHICmdList, ERenderQueryType QueryType) {}
 
 	RHI_API virtual void RHIBeginRenderQuery_TopOfPipe(FRHICommandListBase& RHICmdList, FRHIRenderQuery* RenderQuery);
 	RHI_API virtual void RHIEndRenderQuery_TopOfPipe  (FRHICommandListBase& RHICmdList, FRHIRenderQuery* RenderQuery);
diff --git a/Engine/Source/Runtime/RHI/Public/RHICommandList.h b/Engine/Source/Runtime/RHI/Public/RHICommandList.h
--- a/Engine/Source/Runtime/RHI/Public/RHICommandList.h
+++ b/Engine/Source/Runtime/RHI/Public/RHICommandList.h
@@ -1193,6 +1193,9 @@ protected:
 		TOptional<FRHIDrawStatsCategory const*> CurrentDrawStatsCategory {};
 #endif
 
+		TStaticArray<void*, MAX_NUM_GPUS> QueryBatchData_Timestamp { InPlace, nullptr };
+		TStaticArray<void*, MAX_NUM_GPUS> QueryBatchData_Occlusion { InPlace, nullptr };
+
 		FPersistentState(FRHIGPUMask InInitialGPUMask, bool bInImmediate = false)
 			: bHasFragmentDensityAttachment(0)
 			, bInsideRenderPass(0)
@@ -1231,7 +1234,15 @@ public:
 	}
 #endif
 
-	TStaticArray<void*, MAX_NUM_GPUS> QueryBatchData { InPlace, nullptr };
+	TStaticArray<void*, MAX_NUM_GPUS>& GetQueryBatchData(ERenderQueryType QueryType)
+	{
+		switch (QueryType)
+		{
+		default: checkNoEntry(); [[fallthrough]];
+		case RQT_AbsoluteTime: return PersistentState.QueryBatchData_Timestamp;
+		case RQT_Occlusion:    return PersistentState.QueryBatchData_Occlusion;
+		}
+	}
 
 private:
 	FRHICommandListBase(FPersistentState const& InPersistentState);
@@ -3575,7 +3586,7 @@ public:
 		if (InInfo.NumOcclusionQueries)
 		{
 			PersistentState.bInsideOcclusionQueryBatch = true;
-			GDynamicRHI->RHIBeginOcclusionQueryBatch_TopOfPipe(*this, InInfo.NumOcclusionQueries);
+			GDynamicRHI->RHIBeginRenderQueryBatch_TopOfPipe(*this, RQT_Occlusion);
 		}
 	}
 
@@ -3586,7 +3597,7 @@ public:
 
 		if (PersistentState.bInsideOcclusionQueryBatch)
 		{
-			GDynamicRHI->RHIEndOcclusionQueryBatch_TopOfPipe(*this);
+			GDynamicRHI->RHIEndRenderQueryBatch_TopOfPipe(*this, RQT_Occlusion);
 			PersistentState.bInsideOcclusionQueryBatch = false;
 		}
 
diff --git a/Engine/Source/Runtime/RHI/Public/RHIValidation.h b/Engine/Source/Runtime/RHI/Public/RHIValidation.h
--- a/Engine/Source/Runtime/RHI/Public/RHIValidation.h
+++ b/Engine/Source/Runtime/RHI/Public/RHIValidation.h
@@ -804,14 +804,14 @@ public:
 		return RHI->RHIGetRenderQueryResult(RenderQuery, OutResult, bWait, GPUIndex);
 	}
 
-	virtual void RHIBeginOcclusionQueryBatch_TopOfPipe(FRHICommandListBase& RHICmdList, uint32 NumQueriesInBatch) override final
+	virtual void RHIBeginRenderQueryBatch_TopOfPipe(FRHICommandListBase& RHICmdList, ERenderQueryType QueryType) override final
 	{
-		RHI->RHIBeginOcclusionQueryBatch_TopOfPipe(RHICmdList, NumQueriesInBatch);
+		RHI->RHIBeginRenderQueryBatch_TopOfPipe(RHICmdList, QueryType);
 	}
 
-	virtual void RHIEndOcclusionQueryBatch_TopOfPipe(FRHICommandListBase& RHICmdList) override final
+	virtual void RHIEndRenderQueryBatch_TopOfPipe(FRHICommandListBase& RHICmdList, ERenderQueryType QueryType) override final
 	{
-		RHI->RHIEndOcclusionQueryBatch_TopOfPipe(RHICmdList);
+		RHI->RHIEndRenderQueryBatch_TopOfPipe(RHICmdList, QueryType);
 	}
 
 	virtual void RHIBeginRenderQuery_TopOfPipe(FRHICommandListBase& RHICmdList, FRHIRenderQuery* RenderQuery) override final
diff --git a/Engine/Source/Runtime/RenderCore/Private/ProfilingDebugging/RealtimeGPUProfiler.cpp b/Engine/Source/Runtime/RenderCore/Private/ProfilingDebugging/RealtimeGPUProfiler.cpp
--- a/Engine/Source/Runtime/RenderCore/Private/ProfilingDebugging/RealtimeGPUProfiler.cpp
+++ b/Engine/Source/Runtime/RenderCore/Private/ProfilingDebugging/RealtimeGPUProfiler.cpp
@@ -932,6 +932,12 @@ void FRealtimeGPUProfiler::BeginFrame(FRHICommandListImmediate& RHICmdList)
 
 	Frames[WriteBufferIndex]->TimestampCalibrationQuery = new FRHITimestampCalibrationQuery();
 	RHICmdList.CalibrateTimers(Frames[WriteBufferIndex]->TimestampCalibrationQuery);
+
+	check(!bQueryBatchStarted);
+	bQueryBatchStarted = true;
+	// Using GDynamicRHI directly to avoid exposing this as a public API on RHICmdList.
+	// This is a temporary fix until we refactor the GPU profiler.
+	GDynamicRHI->RHIBeginRenderQueryBatch_TopOfPipe(RHICmdList, RQT_AbsoluteTime);
 }
 
 bool AreGPUStatsEnabled()
@@ -971,6 +977,15 @@ bool AreGPUStatsEnabled()
 
 void FRealtimeGPUProfiler::EndFrame(FRHICommandListImmediate& RHICmdList)
 {
+	// Do this before the AreGPUStatsEnabled check, in case the cvar changes mid-frame.
+	if (bQueryBatchStarted)
+	{
+	    // Using GDynamicRHI directly to avoid exposing this as a public API on RHICmdList.
+	    // This is a temporary fix until we refactor the GPU profiler.
+		GDynamicRHI->RHIEndRenderQueryBatch_TopOfPipe(RHICmdList, RQT_AbsoluteTime);
+		bQueryBatchStarted = false;
+	}
+
 	if (!AreGPUStatsEnabled())
 	{
 		return;
diff --git a/Engine/Source/Runtime/RenderCore/Private/RenderGraphBuilder.cpp b/Engine/Source/Runtime/RenderCore/Private/RenderGraphBuilder.cpp
--- a/Engine/Source/Runtime/RenderCore/Private/RenderGraphBuilder.cpp
+++ b/Engine/Source/Runtime/RenderCore/Private/RenderGraphBuilder.cpp
@@ -1773,7 +1773,10 @@ void FRDGBuilder::Execute()
 
 		if (ParallelExecute.bEnabled)
 		{
-			AddSetupTask([this] { SetupParallelExecute(); });
+			AddSetupTask([this, QueryBatchData = RHICmdList.GetQueryBatchData(RQT_AbsoluteTime)]
+			{
+				SetupParallelExecute(QueryBatchData);
+			});
 		}
 
 		UE::Tasks::FTask AllocatePooledBuffersTask;
@@ -2566,7 +2569,7 @@ void FRDGBuilder::SubmitBufferUploads(FRHICommandListBase& RHICmdListUpload, UE:
 
 ///////////////////////////////////////////////////////////////////////////////////////////////////
 
-void FRDGBuilder::SetupParallelExecute()
+void FRDGBuilder::SetupParallelExecute(TStaticArray<void*, MAX_NUM_GPUS> const& QueryBatchData)
 {
 	SCOPED_NAMED_EVENT(SetupParallelExecute, FColor::Emerald);
 	FRDGAllocatorScope AllocatorScope(Allocators.Task);
@@ -2686,6 +2689,11 @@ void FRDGBuilder::SetupParallelExecute()
 	for (FParallelPassSet& ParallelPassSet : ParallelExecute.ParallelPassSets)
 	{
 		FRHICommandList* RHICmdListPass = new FRHICommandList(FRHIGPUMask::All());
+
+		// Propagate the immediate command list's timestamp query batch.
+		// This is a workaround for poor fence batching on some platforms due to the realtime GPU profiler / timestamp query API design.
+		RHICmdListPass->GetQueryBatchData(RQT_AbsoluteTime) = QueryBatchData;
+
 		ParallelPassSet.CmdList = RHICmdListPass;
 
 		ParallelExecute.Tasks.Emplace(UE::Tasks::Launch(TEXT("FRDGBuilder::ParallelExecute"), [this, &ParallelPassSet, RHICmdListPass]
diff --git a/Engine/Source/Runtime/RenderCore/Public/ProfilingDebugging/RealtimeGPUProfiler.h b/Engine/Source/Runtime/RenderCore/Public/ProfilingDebugging/RealtimeGPUProfiler.h
--- a/Engine/Source/Runtime/RenderCore/Public/ProfilingDebugging/RealtimeGPUProfiler.h
+++ b/Engine/Source/Runtime/RenderCore/Public/ProfilingDebugging/RealtimeGPUProfiler.h
@@ -319,6 +319,7 @@ private:
 	bool bStatGatheringPaused;
 	bool bInBeginEndBlock;
 	bool bLocked = false;
+	bool bQueryBatchStarted = false;
 
 #if GPUPROFILERTRACE_ENABLED
 	FRealtimeGPUProfilerHistoryByDescription HistoryByDescription;
diff --git a/Engine/Source/Runtime/RenderCore/Public/RenderGraphBuilder.h b/Engine/Source/Runtime/RenderCore/Public/RenderGraphBuilder.h
--- a/Engine/Source/Runtime/RenderCore/Public/RenderGraphBuilder.h
+++ b/Engine/Source/Runtime/RenderCore/Public/RenderGraphBuilder.h
@@ -954,7 +954,7 @@ private:
 
 	} ParallelExecute;
 
-	void SetupParallelExecute();
+	void SetupParallelExecute(TStaticArray<void*, MAX_NUM_GPUS> const& QueryBatchData);
 
 	/////////////////////////////////////////////////////////////////////////////
 	// Buffer Uploads
