Add GPU memory stats to DX11, share DXGI code between DX11 and DX12, and standardize reporting across platforms.

* In case of multiple GPUs, use the memory information from the one with the largest memory budget.
* Do not crash if the driver does not support querying memory information.
* Only write CSV stats for Windows, where there's a high chance of this information being useful (with a discrete card having dedicated VRAM).

diff --git a/Engine/Source/Runtime/D3D12RHI/Private/D3D12Adapter.cpp b/Engine/Source/Runtime/D3D12RHI/Private/D3D12Adapter.cpp
--- a/Engine/Source/Runtime/D3D12RHI/Private/D3D12Adapter.cpp
+++ b/Engine/Source/Runtime/D3D12RHI/Private/D3D12Adapter.cpp
@@ -1533,73 +1533,41 @@ void FD3D12Adapter::ReleaseTransientUniformBufferAllocator(FTransientUniformBuff
 	verify(TransientUniformBufferAllocators.Remove(InAllocator) == 1);
 }
 
-void FD3D12Adapter::UpdateMemoryInfo()
+const FD3DMemoryStats& FD3D12Adapter::CollectMemoryStats()
 {
 #if PLATFORM_WINDOWS
 	const uint64 UpdateFrame = FrameFence != nullptr ? FrameFence->GetNextFenceToSignal() : 0;
 
-	// Avoid spurious query calls if we have already captured this frame.
-	if (MemoryInfo.UpdateFrameNumber == UpdateFrame)
+	// Avoid spurious query calls if we have already captured stats this frame.
+	if (MemoryStatsUpdateFrame == UpdateFrame)
 	{
-		return;
+		return MemoryStats;
 	}
+	MemoryStatsUpdateFrame = UpdateFrame;
 
-	// Update the frame number that the memory is captured from.
-	MemoryInfo.UpdateFrameNumber = UpdateFrame;
-
-	TRefCountPtr<IDXGIAdapter3> Adapter3;
-	VERIFYD3D12RESULT(GetAdapter()->QueryInterface(IID_PPV_ARGS(Adapter3.GetInitReference())));
+	if (FAILED(UE::DXGIUtilities::GetD3DMemoryStats(GetAdapter(), MemoryStats)))
+	{
+		return MemoryStats;
+	}
 
-	VERIFYD3D12RESULT(Adapter3->QueryVideoMemoryInfo(0, DXGI_MEMORY_SEGMENT_GROUP_LOCAL, &MemoryInfo.LocalMemoryInfo));
-	VERIFYD3D12RESULT(Adapter3->QueryVideoMemoryInfo(0, DXGI_MEMORY_SEGMENT_GROUP_NON_LOCAL, &MemoryInfo.NonLocalMemoryInfo));
+	// Update global RHI state (for warning output, etc.)
+	GDemotedLocalMemorySize = MemoryStats.DemotedLocal;
 
 #if ENABLE_RESIDENCY_MANAGEMENT
 	// D3D12 residency manager will only evict resources when the resident set is larger than DXGI reported budget.
 	// However, the DXGI reports high budget even when multiple applications use large amounts of VRAM.
-	// This casuses VidMm to automatically page out allocations out of VRAM based on its own heuristics, which can cause significant 
+	// This causes VidMm to automatically page out allocations out of VRAM based on its own heuristics, which can cause significant 
 	// performance degradation when paged-out resources are used for rendering before VidMm pages them back in (which can take a long time).
 	// By overriding the budget to 0 and stopping rendering at the high-level, we can immediately free VRAM and avoid VidMm paging.
 	const bool bEvictResidentResources = GEnableResidencyManagement && GD3D12EvictAllResidentResourcesInBackground && !FApp::HasFocus();
-	const uint64 LocalMemoryBudgetLimit = bEvictResidentResources ? 0 : MemoryInfo.LocalMemoryInfo.Budget;
+	const uint64 LocalMemoryBudgetLimit = bEvictResidentResources ? 0 : MemoryStats.BudgetLocal;
 	for (uint32 GPUIndex : FRHIGPUMask::All())
 	{
 		GetDevice(GPUIndex)->GetResidencyManager().SetLocalMemoryBudgetLimit(LocalMemoryBudgetLimit);
 	}
 #endif // ENABLE_RESIDENCY_MANAGEMENT
-
-	// Over budget?
-	if (MemoryInfo.LocalMemoryInfo.CurrentUsage > MemoryInfo.LocalMemoryInfo.Budget)
-	{
-		MemoryInfo.AvailableLocalMemory = 0;
-		MemoryInfo.DemotedLocalMemory = MemoryInfo.LocalMemoryInfo.CurrentUsage - MemoryInfo.LocalMemoryInfo.Budget;
-	}
-	else
-	{
-		MemoryInfo.AvailableLocalMemory = MemoryInfo.LocalMemoryInfo.Budget - MemoryInfo.LocalMemoryInfo.CurrentUsage;
-		MemoryInfo.DemotedLocalMemory = 0;
-	}
-
-	// Update global RHI state (for warning output, etc.)
-	GDemotedLocalMemorySize = MemoryInfo.DemotedLocalMemory;
-
-	if (!GVirtualMGPU)
-	{
-		for (uint32 Index = 1; Index < GNumExplicitGPUsForRendering; ++Index)
-		{
-			DXGI_QUERY_VIDEO_MEMORY_INFO TempVideoMemoryInfo;
-			VERIFYD3D12RESULT(Adapter3->QueryVideoMemoryInfo(Index, DXGI_MEMORY_SEGMENT_GROUP_LOCAL, &TempVideoMemoryInfo));
-
-			DXGI_QUERY_VIDEO_MEMORY_INFO TempSystemMemoryInfo;
-			VERIFYD3D12RESULT(Adapter3->QueryVideoMemoryInfo(Index, DXGI_MEMORY_SEGMENT_GROUP_NON_LOCAL, &TempSystemMemoryInfo));
-			
-			MemoryInfo.LocalMemoryInfo.Budget = FMath::Min(MemoryInfo.LocalMemoryInfo.Budget, TempVideoMemoryInfo.Budget);
-			MemoryInfo.LocalMemoryInfo.CurrentUsage = FMath::Min(MemoryInfo.LocalMemoryInfo.CurrentUsage, TempVideoMemoryInfo.CurrentUsage);
-
-			MemoryInfo.NonLocalMemoryInfo.Budget = FMath::Min(MemoryInfo.NonLocalMemoryInfo.Budget, TempSystemMemoryInfo.Budget);
-			MemoryInfo.NonLocalMemoryInfo.CurrentUsage = FMath::Min(MemoryInfo.NonLocalMemoryInfo.CurrentUsage, TempSystemMemoryInfo.CurrentUsage);
-		}
-	}
-#endif
+#endif // PLATFORM_WINDOWS
+	return MemoryStats;
 }
 
 void FD3D12Adapter::BlockUntilIdle()
diff --git a/Engine/Source/Runtime/D3D12RHI/Private/D3D12Adapter.h b/Engine/Source/Runtime/D3D12RHI/Private/D3D12Adapter.h
--- a/Engine/Source/Runtime/D3D12RHI/Private/D3D12Adapter.h
+++ b/Engine/Source/Runtime/D3D12RHI/Private/D3D12Adapter.h
@@ -396,8 +396,9 @@ public:
 
 	void BlockUntilIdle();
 
-	void UpdateMemoryInfo();
-	FORCEINLINE const FD3D12MemoryInfo& GetMemoryInfo() const { return MemoryInfo; }
+	const FD3DMemoryStats& CollectMemoryStats();
+	FORCEINLINE const FD3DMemoryStats& GetMemoryStats() const { return MemoryStats; }
+	FORCEINLINE uint64 GetMemoryStatsUpdateFrame() const { return MemoryStatsUpdateFrame; }
 
 	bool IsTrackingAllAllocations() const { return bTrackAllAllocation; }
 	void TrackAllocationData(FD3D12ResourceLocation* InAllocation, uint64 InAllocationSize, bool bCollectCallstack);
@@ -586,7 +587,8 @@ protected:
 	TArray<FReleasedAllocationData> ReleasedAllocationData;
 	FCriticalSection TrackedAllocationDataCS;
 
-	FD3D12MemoryInfo MemoryInfo;
+	FD3DMemoryStats MemoryStats;
+	uint64 MemoryStatsUpdateFrame;
 
 	TArray<FTransientUniformBufferAllocator*> TransientUniformBufferAllocators;
 	FCriticalSection TransientUniformBufferAllocatorsCS;
diff --git a/Engine/Source/Runtime/D3D12RHI/Private/D3D12CommandContext.cpp b/Engine/Source/Runtime/D3D12RHI/Private/D3D12CommandContext.cpp
--- a/Engine/Source/Runtime/D3D12RHI/Private/D3D12CommandContext.cpp
+++ b/Engine/Source/Runtime/D3D12RHI/Private/D3D12CommandContext.cpp
@@ -730,36 +730,16 @@ void FD3D12CommandContextBase::RHIEndFrame()
 	FD3D12DynamicRHI::GetD3DRHI()->ProcessInterruptQueueUntil(nullptr);
 }
 
-#if PLATFORM_WINDOWS && CSV_PROFILER && !UE_BUILD_SHIPPING
-	CSV_DEFINE_CATEGORY(GPUMem, true);
-#endif
-DEFINE_STAT(STAT_D3D12UpdateVideoMemoryStats);
-
 void FD3D12CommandContextBase::UpdateMemoryStats()
 {
-#if PLATFORM_WINDOWS && (STATS || CSV_PROFILER) && !UE_BUILD_SHIPPING
-	SCOPE_CYCLE_COUNTER(STAT_D3D12UpdateVideoMemoryStats);
-	// Refresh captured memory info.
-	ParentAdapter->UpdateMemoryInfo();
-
-	const FD3D12MemoryInfo& MemoryInfo = ParentAdapter->GetMemoryInfo();
-
-#if CSV_PROFILER
-	{
-		CSV_CUSTOM_STAT(GPUMem, TotalMB, float(MemoryInfo.LocalMemoryInfo.Budget / 1024.0 / 1024.0), ECsvCustomStatOp::Set);
-		CSV_CUSTOM_STAT(GPUMem, UsedMB, float(MemoryInfo.LocalMemoryInfo.CurrentUsage / 1024.0 / 1024.0), ECsvCustomStatOp::Set);
-		CSV_CUSTOM_STAT(GPUMem, AvailableMB, float(MemoryInfo.AvailableLocalMemory / 1024.0 / 1024.0), ECsvCustomStatOp::Set);
-		CSV_CUSTOM_STAT(GPUMem, DemotedMB, float(MemoryInfo.DemotedLocalMemory / 1024.0 / 1024.0), ECsvCustomStatOp::Set);
-	}
-#endif // CSV_PROFILER
+#if PLATFORM_WINDOWS && (STATS || CSV_PROFILER)
+	SCOPE_CYCLE_COUNTER(STAT_D3DUpdateVideoMemoryStats);
 
+	// Refresh captured memory stats.
+	const FD3DMemoryStats& MemoryStats = ParentAdapter->CollectMemoryStats();
+	UpdateD3DMemoryStatsAndCSV(MemoryStats, true);
+	
 #if STATS
-	SET_MEMORY_STAT(STAT_D3D12UsedVideoMemory, MemoryInfo.LocalMemoryInfo.CurrentUsage);
-	SET_MEMORY_STAT(STAT_D3D12UsedSystemMemory, MemoryInfo.NonLocalMemoryInfo.CurrentUsage);
-	SET_MEMORY_STAT(STAT_D3D12AvailableVideoMemory, MemoryInfo.AvailableLocalMemory);
-	SET_MEMORY_STAT(STAT_D3D12DemotedVideoMemory, MemoryInfo.DemotedLocalMemory);
-	SET_MEMORY_STAT(STAT_D3D12TotalVideoMemory, MemoryInfo.LocalMemoryInfo.Budget);
-
 	uint64 MaxTexAllocWastage = 0;
 	for (uint32 GPUIndex : GPUMask)
 	{
diff --git a/Engine/Source/Runtime/D3D12RHI/Private/D3D12Stats.h b/Engine/Source/Runtime/D3D12RHI/Private/D3D12Stats.h
--- a/Engine/Source/Runtime/D3D12RHI/Private/D3D12Stats.h
+++ b/Engine/Source/Runtime/D3D12RHI/Private/D3D12Stats.h
@@ -70,13 +70,6 @@ DECLARE_CYCLE_STAT_WITH_FLAGS_EXTERN(TEXT("Clear MRT time"), STAT_D3D12ClearMRT,
 DECLARE_CYCLE_STAT_EXTERN(TEXT("ExecuteCommandList time"), STAT_D3D12ExecuteCommandListTime, STATGROUP_D3D12RHI, );
 DECLARE_CYCLE_STAT_EXTERN(TEXT("WaitForFence time"), STAT_D3D12WaitForFenceTime, STATGROUP_D3D12RHI, );
 
-DECLARE_MEMORY_STAT_EXTERN(TEXT("Total Video Memory"), STAT_D3D12TotalVideoMemory, STATGROUP_D3D12Memory, );
-DECLARE_MEMORY_STAT_EXTERN(TEXT("Available Video Memory"), STAT_D3D12AvailableVideoMemory, STATGROUP_D3D12Memory, );
-DECLARE_MEMORY_STAT_EXTERN(TEXT("Demoted Video Memory"), STAT_D3D12DemotedVideoMemory, STATGROUP_D3D12Memory, );
-DECLARE_MEMORY_STAT_EXTERN(TEXT("Used Video Memory"), STAT_D3D12UsedVideoMemory, STATGROUP_D3D12Memory, );
-DECLARE_MEMORY_STAT_EXTERN(TEXT("Used System Memory"), STAT_D3D12UsedSystemMemory, STATGROUP_D3D12Memory, );
-DECLARE_CYCLE_STAT_EXTERN(TEXT("Video Memory stats update time"), STAT_D3D12UpdateVideoMemoryStats, STATGROUP_D3D12RHI, );
-
 DECLARE_CYCLE_STAT_EXTERN(TEXT("Global Constant buffer update time"), STAT_D3D12GlobalConstantBufferUpdateTime, STATGROUP_D3D12RHI, );
 
 DECLARE_MEMORY_STAT_EXTERN(TEXT("TOTAL"), STAT_D3D12MemoryCurrentTotal, STATGROUP_D3D12Resources, );
diff --git a/Engine/Source/Runtime/D3D12RHI/Private/D3D12Texture.cpp b/Engine/Source/Runtime/D3D12RHI/Private/D3D12Texture.cpp
--- a/Engine/Source/Runtime/D3D12RHI/Private/D3D12Texture.cpp
+++ b/Engine/Source/Runtime/D3D12RHI/Private/D3D12Texture.cpp
@@ -405,17 +405,16 @@ void FD3D12DynamicRHI::RHIGetTextureMemoryStats(FTextureMemoryStats& OutStats)
 #if PLATFORM_WINDOWS
 	if (GAdjustTexturePoolSizeBasedOnBudget)
 	{
-		GetAdapter().UpdateMemoryInfo();
-		const DXGI_QUERY_VIDEO_MEMORY_INFO& LocalVideoMemoryInfo = GetAdapter().GetMemoryInfo().LocalMemoryInfo;
+		const FD3DMemoryStats& MemoryStats = GetAdapter().CollectMemoryStats();
 
 		// Applications must explicitly manage their usage of physical memory and keep usage within the budget 
 		// assigned to the application process. Processes that cannot keep their usage within their assigned budgets 
 		// will likely experience stuttering, as they are intermittently frozen and paged out to allow other processes to run.
-		const int64 TargetBudget = LocalVideoMemoryInfo.Budget * 0.90f;	// Target using 90% of our budget to account for some fragmentation.
+		const int64 TargetBudget = MemoryStats.BudgetLocal * 0.90f;	// Target using 90% of our budget to account for some fragmentation.
 		OutStats.TotalGraphicsMemory = TargetBudget;
 
 		const int64 BudgetPadding = TargetBudget * 0.05f;
-		const int64 AvailableSpace = TargetBudget - int64(LocalVideoMemoryInfo.CurrentUsage);	// Note: AvailableSpace can be negative
+		const int64 AvailableSpace = TargetBudget - int64(MemoryStats.UsedLocal);	// Note: AvailableSpace can be negative
 		const int64 PreviousTexturePoolSize = RequestedTexturePoolSize;
 		const bool bOverbudget = AvailableSpace < 0;
 
@@ -424,12 +423,12 @@ void FD3D12DynamicRHI::RHIGetTextureMemoryStats(FTextureMemoryStats& OutStats)
 		if (bOverbudget)
 		{
 			// Attempt to lower the texture pool size to meet the budget.
-			const bool bOverActualBudget = LocalVideoMemoryInfo.CurrentUsage > LocalVideoMemoryInfo.Budget;
+			const bool bOverActualBudget = MemoryStats.UsedLocal > MemoryStats.BudgetLocal;
 			UE_CLOG(bOverActualBudget, LogD3D12RHI, Warning,
 				TEXT("Video memory usage is overbudget by %llu MB (using %lld MB/%lld MB budget). Usage breakdown: %lld MB (Streaming Textures), %lld MB (Non Streaming Textures). Last requested texture pool size is %lld MB. This can cause stuttering due to paging."),
-				(LocalVideoMemoryInfo.CurrentUsage - LocalVideoMemoryInfo.Budget) / 1024ll / 1024ll,
-				LocalVideoMemoryInfo.CurrentUsage / 1024ll / 1024ll,
-				LocalVideoMemoryInfo.Budget / 1024ll / 1024ll,
+				(MemoryStats.UsedLocal - MemoryStats.BudgetLocal) / 1024ll / 1024ll,
+				MemoryStats.UsedLocal / 1024ll / 1024ll,
+				MemoryStats.BudgetLocal / 1024ll / 1024ll,
 				GRHIGlobals.StreamingTextureMemorySizeInKB / 1024ll,
 				GRHIGlobals.NonStreamingTextureMemorySizeInKB / 1024ll,
 				PreviousTexturePoolSize / 1024ll / 1024ll);
@@ -439,7 +438,7 @@ void FD3D12DynamicRHI::RHIGetTextureMemoryStats(FTextureMemoryStats& OutStats)
 
 			UE_CLOG(bOverActualBudget && (OutStats.TexturePoolSize >= PreviousTexturePoolSize) && (OutStats.TexturePoolSize > MinTexturePoolSize), LogD3D12RHI, Fatal,
 				TEXT("Video memory usage is overbudget by %llu MB and the texture pool size didn't shrink."),
-				(LocalVideoMemoryInfo.CurrentUsage - LocalVideoMemoryInfo.Budget) / 1024ll / 1024ll);
+				(MemoryStats.UsedLocal - MemoryStats.BudgetLocal) / 1024ll / 1024ll);
 		}
 		else if (AvailableSpace > BudgetPadding)
 		{
diff --git a/Engine/Source/Runtime/D3D12RHI/Private/D3D12Util.cpp b/Engine/Source/Runtime/D3D12RHI/Private/D3D12Util.cpp
--- a/Engine/Source/Runtime/D3D12RHI/Private/D3D12Util.cpp
+++ b/Engine/Source/Runtime/D3D12RHI/Private/D3D12Util.cpp
@@ -772,13 +772,15 @@ void LogPageFaultData(FD3D12Adapter* InAdapter, FD3D12Device* InDevice, D3D12_GP
 
 } // namespace D3D12RHI
 
-void LogMemoryInfo(FD3D12Adapter* InAdapter)
+void LogMemoryStats(FD3D12Adapter* InAdapter)
 {	
-	const FD3D12MemoryInfo& MemoryInfo = InAdapter->GetMemoryInfo();
+	const FD3DMemoryStats& MemoryStats = InAdapter->GetMemoryStats();
 
-	UE_LOG(LogD3D12RHI, Error, TEXT("Memory Info from frame ID %d:"), MemoryInfo.UpdateFrameNumber);
-	UE_LOG(LogD3D12RHI, Error, TEXT("\tBudget:\t%7.2f MB"), MemoryInfo.LocalMemoryInfo.Budget / (1024.0f * 1024));
-	UE_LOG(LogD3D12RHI, Error, TEXT("\tUsed:\t%7.2f MB"), MemoryInfo.LocalMemoryInfo.CurrentUsage / (1024.0f * 1024));
+	UE_LOG(LogD3D12RHI, Error, TEXT("Video Memory Stats from frame ID %d:"), InAdapter->GetMemoryStatsUpdateFrame());
+	UE_LOG(LogD3D12RHI, Error, TEXT("\tLocal Budget:\t%7.2f MB"), MemoryStats.BudgetLocal / (1024.0f * 1024));
+	UE_LOG(LogD3D12RHI, Error, TEXT("\tLocal Used:\t%7.2f MB"), MemoryStats.UsedLocal / (1024.0f * 1024));
+	UE_LOG(LogD3D12RHI, Error, TEXT("\tSystem Budget:\t%7.2f MB"), MemoryStats.BudgetSystem / (1024.0f * 1024));
+	UE_LOG(LogD3D12RHI, Error, TEXT("\tSystem Used:\t%7.2f MB"), MemoryStats.UsedSystem / (1024.0f * 1024));
 }
 
 #endif  // PLATFORM_WINDOWS
@@ -807,8 +809,8 @@ void FD3D12DynamicRHI::TerminateOnOutOfMemory(ID3D12Device* InDevice, HRESULT D3
 
 	if (Adapter)
 	{
-		const auto& MemoryInfo = Adapter->GetMemoryInfo().LocalMemoryInfo;
-		FCoreDelegates::GetGPUOutOfMemoryDelegate().Broadcast(MemoryInfo.Budget, MemoryInfo.CurrentUsage);
+		const auto& MemoryStats = Adapter->GetMemoryStats();
+		FCoreDelegates::GetGPUOutOfMemoryDelegate().Broadcast(MemoryStats.BudgetLocal, MemoryStats.UsedLocal);
 	}
 
 	if (!FApp::IsUnattended())
@@ -833,7 +835,7 @@ void FD3D12DynamicRHI::TerminateOnOutOfMemory(ID3D12Device* InDevice, HRESULT D3
 	ForEachDevice(InDevice, [&](FD3D12Device* IterationDevice)
 	{
 		FD3D12Adapter* Adapter = IterationDevice->GetParentAdapter();
-		LogMemoryInfo(Adapter);
+		LogMemoryStats(Adapter);
 	});
 	
 	UE_LOG(LogD3D12RHI, Fatal, TEXT("Out of video memory trying to allocate a rendering resource"));
@@ -966,7 +968,7 @@ void FD3D12DynamicRHI::OutputGPUCrashReport(FTextBuilder& ErrorMessage)
 
 		FD3D12Adapter* Adapter = Device->GetParentAdapter();
 		LogPageFaultData(Adapter, Device, PageFaultAddress);
-		LogMemoryInfo(Adapter);
+		LogMemoryStats(Adapter);
 	});
 #endif  // PLATFORM_WINDOWS
 
@@ -2000,12 +2002,6 @@ DEFINE_STAT(STAT_D3D12ClearMRT);
 DEFINE_STAT(STAT_D3D12ExecuteCommandListTime);
 DEFINE_STAT(STAT_D3D12WaitForFenceTime);
 
-DEFINE_STAT(STAT_D3D12UsedVideoMemory);
-DEFINE_STAT(STAT_D3D12UsedSystemMemory);
-DEFINE_STAT(STAT_D3D12AvailableVideoMemory);
-DEFINE_STAT(STAT_D3D12DemotedVideoMemory);
-DEFINE_STAT(STAT_D3D12TotalVideoMemory);
-
 DEFINE_STAT(STAT_D3D12MemoryCurrentTotal);
 DEFINE_STAT(STAT_D3D12RenderTargets);
 DEFINE_STAT(STAT_D3D12UAVTextures);
diff --git a/Engine/Source/Runtime/RHI/Private/RHIStats.cpp b/Engine/Source/Runtime/RHI/Private/RHIStats.cpp
--- a/Engine/Source/Runtime/RHI/Private/RHIStats.cpp
+++ b/Engine/Source/Runtime/RHI/Private/RHIStats.cpp
@@ -4,6 +4,7 @@
 #include "HAL/Platform.h"
 #include "HAL/PlatformAtomics.h"
 #include "Templates/Atomic.h"
+#include "ProfilingDebugging/CsvProfiler.h"
 
 int32 GNumDrawCallsRHI[MAX_NUM_GPUS] = {};
 int32 GNumPrimitivesDrawnRHI[MAX_NUM_GPUS] = {};
@@ -39,3 +40,59 @@ DEFINE_STAT(STAT_BindlessSamplerHeapMemory);
 DEFINE_STAT(STAT_BindlessResourceHeapMemory);
 DEFINE_STAT(STAT_BindlessSamplerDescriptorsAllocated);
 DEFINE_STAT(STAT_BindlessResourceDescriptorsAllocated);
+
+#if PLATFORM_MICROSOFT
+
+// Define D3D memory stats.
+DEFINE_STAT(STAT_D3DUpdateVideoMemoryStats);
+DEFINE_STAT(STAT_D3DTotalVideoMemory);
+DEFINE_STAT(STAT_D3DTotalSystemMemory);
+DEFINE_STAT(STAT_D3DUsedVideoMemory);
+DEFINE_STAT(STAT_D3DUsedSystemMemory);
+DEFINE_STAT(STAT_D3DAvailableVideoMemory);
+DEFINE_STAT(STAT_D3DAvailableSystemMemory);
+DEFINE_STAT(STAT_D3DDemotedVideoMemory);
+DEFINE_STAT(STAT_D3DDemotedSystemMemory);
+
+#if CSV_PROFILER
+CSV_DEFINE_CATEGORY(GPUMem, true);
+#endif
+
+void UpdateD3DMemoryStatsAndCSV(const FD3DMemoryStats& MemoryStats, bool bUpdateCSV)
+{
+#if STATS || CSV_PROFILER
+	SCOPE_CYCLE_COUNTER(STAT_D3DUpdateVideoMemoryStats);
+
+#if STATS
+	SET_MEMORY_STAT(STAT_D3DTotalVideoMemory, MemoryStats.BudgetLocal);	
+	SET_MEMORY_STAT(STAT_D3DUsedVideoMemory, MemoryStats.UsedLocal);
+	SET_MEMORY_STAT(STAT_D3DAvailableVideoMemory, MemoryStats.AvailableLocal);
+	SET_MEMORY_STAT(STAT_D3DDemotedVideoMemory, MemoryStats.DemotedLocal);
+	
+	if (MemoryStats.BudgetSystem > 0)
+	{
+		SET_MEMORY_STAT(STAT_D3DTotalSystemMemory, MemoryStats.BudgetSystem);
+		SET_MEMORY_STAT(STAT_D3DUsedSystemMemory, MemoryStats.UsedSystem);
+		SET_MEMORY_STAT(STAT_D3DAvailableSystemMemory, MemoryStats.AvailableSystem);
+		SET_MEMORY_STAT(STAT_D3DDemotedSystemMemory, MemoryStats.DemotedSystem);
+	}
+#endif // STATS
+
+#if CSV_PROFILER
+	if (bUpdateCSV)
+	{
+		// Just output the two main stats (budget and used) to avoid bloating the CSV, since the rest can be inferred from them.
+		CSV_CUSTOM_STAT(GPUMem, LocalBudgetMB, float(MemoryStats.BudgetLocal / 1024.0 / 1024.0), ECsvCustomStatOp::Set);
+		CSV_CUSTOM_STAT(GPUMem, LocalUsedMB, float(MemoryStats.UsedLocal / 1024.0 / 1024.0), ECsvCustomStatOp::Set);
+
+		if (MemoryStats.BudgetSystem > 0)
+		{
+			CSV_CUSTOM_STAT(GPUMem, SystemBudgetMB, float(MemoryStats.BudgetSystem / 1024.0 / 1024.0), ECsvCustomStatOp::Set);
+			CSV_CUSTOM_STAT(GPUMem, SystemUsedMB, float(MemoryStats.UsedSystem / 1024.0 / 1024.0), ECsvCustomStatOp::Set);
+		}
+	}
+#endif // CSV_PROFILER
+#endif // STATS || CSV_PROFILER
+}
+
+#endif // PLATFORM_MICROSOFT
\ No newline at end of file
diff --git a/Engine/Source/Runtime/RHI/Public/RHIStats.h b/Engine/Source/Runtime/RHI/Public/RHIStats.h
--- a/Engine/Source/Runtime/RHI/Public/RHIStats.h
+++ b/Engine/Source/Runtime/RHI/Public/RHIStats.h
@@ -239,3 +239,47 @@ DECLARE_MEMORY_STAT_POOL_EXTERN(TEXT("Bindless Resource Heap"), STAT_BindlessRes
 
 DECLARE_DWORD_ACCUMULATOR_STAT_EXTERN(TEXT("Bindless Sampler Descriptors Allocated"), STAT_BindlessSamplerDescriptorsAllocated, STATGROUP_RHI, RHI_API);
 DECLARE_DWORD_ACCUMULATOR_STAT_EXTERN(TEXT("Bindless Resource Descriptors Allocated"), STAT_BindlessResourceDescriptorsAllocated, STATGROUP_RHI, RHI_API);
+
+#if PLATFORM_MICROSOFT
+
+// D3D memory stats.
+struct FD3DMemoryStats
+{
+	// Budget assigned by the OS. This can be considered the total memory
+	// the application should use, but an application can also go over-budget.
+	uint64 BudgetLocal = 0;
+	uint64 BudgetSystem = 0;
+
+	// Used memory.
+	uint64 UsedLocal = 0;
+	uint64 UsedSystem = 0;
+
+	// Over-budget memory. This is Budget - Used if Used > Budget.
+	uint64 DemotedLocal = 0;
+	uint64 DemotedSystem = 0;
+
+	// Available memory within budget. This is Budget - Used clamped to 0 if over-budget.
+	uint64 AvailableLocal = 0;
+	uint64 AvailableSystem = 0;
+
+	bool IsOverBudget() const
+	{
+		return DemotedLocal > 0 || DemotedSystem > 0;
+	}
+};
+
+DECLARE_STATS_GROUP(TEXT("D3D Video Memory"), STATGROUP_D3DMemory, STATCAT_Advanced);
+DECLARE_MEMORY_STAT_EXTERN(TEXT("Total Video Memory (Budget)"), STAT_D3DTotalVideoMemory, STATGROUP_D3DMemory, );
+DECLARE_MEMORY_STAT_EXTERN(TEXT("Total System Memory (Budget)"), STAT_D3DTotalSystemMemory, STATGROUP_D3DMemory, );
+DECLARE_MEMORY_STAT_EXTERN(TEXT("Available Video Memory"), STAT_D3DAvailableVideoMemory, STATGROUP_D3DMemory, );
+DECLARE_MEMORY_STAT_EXTERN(TEXT("Available System Memory"), STAT_D3DAvailableSystemMemory, STATGROUP_D3DMemory, );
+DECLARE_MEMORY_STAT_EXTERN(TEXT("Used Video Memory"), STAT_D3DUsedVideoMemory, STATGROUP_D3DMemory, );
+DECLARE_MEMORY_STAT_EXTERN(TEXT("Used System Memory"), STAT_D3DUsedSystemMemory, STATGROUP_D3DMemory, );
+DECLARE_MEMORY_STAT_EXTERN(TEXT("Demoted Video Memory"), STAT_D3DDemotedVideoMemory, STATGROUP_D3DMemory, );
+DECLARE_MEMORY_STAT_EXTERN(TEXT("Demoted System Memory"), STAT_D3DDemotedSystemMemory, STATGROUP_D3DMemory, );
+DECLARE_CYCLE_STAT_EXTERN(TEXT("Video Memory stats update time"), STAT_D3DUpdateVideoMemoryStats, STATGROUP_D3DMemory, RHI_API);
+
+// Update D3D memory stat counters and CSV profiler stats, if enabled.
+RHI_API void UpdateD3DMemoryStatsAndCSV(const FD3DMemoryStats& MemoryStats, bool bUpdateCSV);
+
+#endif // PLATFORM_MICROSOFT
\ No newline at end of file
diff --git a/Engine/Source/Runtime/RHICore/Internal/DXGIUtilities.h b/Engine/Source/Runtime/RHICore/Internal/DXGIUtilities.h
--- a/Engine/Source/Runtime/RHICore/Internal/DXGIUtilities.h
+++ b/Engine/Source/Runtime/RHICore/Internal/DXGIUtilities.h
@@ -15,13 +15,16 @@
 
 #include "Microsoft/AllowMicrosoftPlatformTypes.h"
 THIRD_PARTY_INCLUDES_START
-	#include RHICORE_PLATFORM_DXGI_H
+#include RHICORE_PLATFORM_DXGI_H
 THIRD_PARTY_INCLUDES_END
 #include "Microsoft/HideMicrosoftPlatformTypes.h"
 
+struct FD3DMemoryStats;
+
 namespace UE::DXGIUtilities
 {
 	RHICORE_API const TCHAR* GetFormatString(DXGI_FORMAT Format);
+	RHICORE_API HRESULT GetD3DMemoryStats(IDXGIAdapter* Adapter, FD3DMemoryStats& OutStats);
 
 	inline DXGI_FORMAT FindSharedResourceFormat(DXGI_FORMAT InFormat, bool bSRGB)
 	{
diff --git a/Engine/Source/Runtime/RHICore/Private/DXGIUtilities.cpp b/Engine/Source/Runtime/RHICore/Private/DXGIUtilities.cpp
--- a/Engine/Source/Runtime/RHICore/Private/DXGIUtilities.cpp
+++ b/Engine/Source/Runtime/RHICore/Private/DXGIUtilities.cpp
@@ -2,6 +2,9 @@
 
 #include "DXGIUtilities.h"
 
+#include "RHIStats.h"
+#include "CoreMinimal.h"
+
 #if PLATFORM_MICROSOFT
 
 const TCHAR* UE::DXGIUtilities::GetFormatString(DXGI_FORMAT Format)
@@ -49,4 +52,91 @@ const TCHAR* UE::DXGIUtilities::GetFormatString(DXGI_FORMAT Format)
 	return Result;
 }
 
+HRESULT UE::DXGIUtilities::GetD3DMemoryStats(IDXGIAdapter* Adapter, FD3DMemoryStats& OutStats)
+{
+#if PLATFORM_WINDOWS
+	SCOPE_CYCLE_COUNTER(STAT_D3DUpdateVideoMemoryStats);
+
+	TRefCountPtr<IDXGIAdapter3> Adapter3;
+	HRESULT Res = Adapter->QueryInterface(IID_PPV_ARGS(Adapter3.GetInitReference()));
+	if (FAILED(Res))
+	{
+		return Res;
+	}
+
+	DXGI_QUERY_VIDEO_MEMORY_INFO LocalMemoryInfo;
+	Res = Adapter3->QueryVideoMemoryInfo(0, DXGI_MEMORY_SEGMENT_GROUP_LOCAL, &LocalMemoryInfo);
+	if (FAILED(Res))
+	{
+		return Res;
+	}
+
+	DXGI_QUERY_VIDEO_MEMORY_INFO NonLocalMemoryInfo;
+	Res = Adapter3->QueryVideoMemoryInfo(0, DXGI_MEMORY_SEGMENT_GROUP_NON_LOCAL, &NonLocalMemoryInfo);
+	if (FAILED(Res))
+	{
+		return Res;
+	}
+
+	// In case of multiple GPUs, use the memory info from the one with the highest local budget.
+	if (!GVirtualMGPU)
+	{
+		for (uint32 Index = 1; Index < GNumExplicitGPUsForRendering; ++Index)
+		{
+			DXGI_QUERY_VIDEO_MEMORY_INFO TempLocalMemoryInfo;
+			Res = Adapter3->QueryVideoMemoryInfo(Index, DXGI_MEMORY_SEGMENT_GROUP_LOCAL, &TempLocalMemoryInfo);
+			if (FAILED(Res))
+			{
+				return Res;
+			}
+
+			DXGI_QUERY_VIDEO_MEMORY_INFO TempNonLocalMemoryInfo;
+			Res = Adapter3->QueryVideoMemoryInfo(Index, DXGI_MEMORY_SEGMENT_GROUP_NON_LOCAL, &TempNonLocalMemoryInfo);
+			if (FAILED(Res))
+			{
+				return Res;
+			}
+
+			if (TempLocalMemoryInfo.Budget > LocalMemoryInfo.Budget)
+			{
+				LocalMemoryInfo = TempLocalMemoryInfo;
+				NonLocalMemoryInfo = TempNonLocalMemoryInfo;
+			}
+		}
+	}
+
+	OutStats.BudgetLocal = LocalMemoryInfo.Budget;
+	OutStats.BudgetSystem = NonLocalMemoryInfo.Budget;
+	OutStats.UsedLocal = LocalMemoryInfo.CurrentUsage;
+	OutStats.UsedSystem = NonLocalMemoryInfo.CurrentUsage;
+
+	// Check if we're over budget.
+	if (OutStats.UsedLocal > OutStats.BudgetLocal)
+	{
+		OutStats.AvailableLocal = 0;
+		OutStats.DemotedLocal = OutStats.UsedLocal - OutStats.BudgetLocal;
+	}
+	else
+	{
+		OutStats.AvailableLocal = OutStats.BudgetLocal - OutStats.UsedLocal;
+		OutStats.DemotedLocal= 0;
+	}
+
+	if (OutStats.UsedSystem > OutStats.BudgetSystem)
+	{
+		OutStats.AvailableSystem = 0;
+		OutStats.DemotedSystem = OutStats.UsedSystem - OutStats.BudgetSystem;
+	}
+	else
+	{
+		OutStats.AvailableSystem = OutStats.BudgetSystem - OutStats.UsedSystem;
+		OutStats.DemotedSystem = 0;
+	}
+
+	return S_OK;
+#else
+	return E_NOINTERFACE;
+#endif
+}
+
 #endif // PLATFORM_MICROSOFT
diff --git a/Engine/Source/Runtime/RHICore/RHICore.Build.cs b/Engine/Source/Runtime/RHICore/RHICore.Build.cs
--- a/Engine/Source/Runtime/RHICore/RHICore.Build.cs
+++ b/Engine/Source/Runtime/RHICore/RHICore.Build.cs
@@ -12,7 +12,7 @@ public class RHICore : ModuleRules
 
 		if (Target.Platform.IsInGroup(UnrealPlatformGroup.Windows))
 		{
-			PublicDefinitions.Add("RHICORE_PLATFORM_DXGI_H=<dxgi.h>");
+			PublicDefinitions.Add("RHICORE_PLATFORM_DXGI_H=<dxgi1_4.h>");
 		}
 
 		AddEngineThirdPartyPrivateStaticDependencies(Target, "NVAftermath");
diff --git a/Engine/Source/Runtime/Windows/D3D11RHI/Private/D3D11RHI.cpp b/Engine/Source/Runtime/Windows/D3D11RHI/Private/D3D11RHI.cpp
--- a/Engine/Source/Runtime/Windows/D3D11RHI/Private/D3D11RHI.cpp
+++ b/Engine/Source/Runtime/Windows/D3D11RHI/Private/D3D11RHI.cpp
@@ -331,9 +331,22 @@ void FD3DGPUProfiler::BeginFrame(FD3D11DynamicRHI* InRHI)
 void FD3D11DynamicRHI::RHIEndFrame()
 {
 	GPUProfilingData.EndFrame();
+	UpdateMemoryStats();
 	CurrentComputeShader = nullptr;
 }
 
+void FD3D11DynamicRHI::UpdateMemoryStats()
+{
+#if PLATFORM_WINDOWS && (STATS || CSV_PROFILER)
+	// Some older drivers don't support querying memory stats, so don't do anything if this fails.
+	FD3DMemoryStats MemoryStats;
+	if (SUCCEEDED(UE::DXGIUtilities::GetD3DMemoryStats(GetAdapter().DXGIAdapter, MemoryStats)))
+	{
+		UpdateD3DMemoryStatsAndCSV(MemoryStats, true);
+	}
+#endif // PLATFORM_WINDOWS && (STATS || CSV_PROFILER)
+}
+
 void FD3DGPUProfiler::EndFrame()
 {
 	if (GetEmitDrawEvents())
diff --git a/Engine/Source/Runtime/Windows/D3D11RHI/Private/D3D11RHIPrivate.h b/Engine/Source/Runtime/Windows/D3D11RHI/Private/D3D11RHIPrivate.h
--- a/Engine/Source/Runtime/Windows/D3D11RHI/Private/D3D11RHIPrivate.h
+++ b/Engine/Source/Runtime/Windows/D3D11RHI/Private/D3D11RHIPrivate.h
@@ -942,6 +942,8 @@ protected:
 
 	bool SetupDisplayHDRMetaData();
 
+	void UpdateMemoryStats();
+
 	friend struct FD3DGPUProfiler;
 
 };
diff --git a/Engine/Source/Runtime/Windows/D3D11RHI/Private/D3D11Util.cpp b/Engine/Source/Runtime/Windows/D3D11RHI/Private/D3D11Util.cpp
--- a/Engine/Source/Runtime/Windows/D3D11RHI/Private/D3D11Util.cpp
+++ b/Engine/Source/Runtime/Windows/D3D11RHI/Private/D3D11Util.cpp
@@ -134,18 +134,17 @@ static void TerminateOnDeviceRemoved(HRESULT D3DResult, ID3D11Device* Direct3DDe
 	}
 }
 
-void GetAndLogMemoryInfo(const FD3D11Adapter& InAdapter, uint64& OutVRAMBudgetBytes, uint64& OutVRAMUsageBytes)
+void GetAndLogMemoryStats(const FD3D11Adapter& InAdapter, uint64& OutVRAMBudgetBytes, uint64& OutVRAMUsageBytes)
 {
-	TRefCountPtr<IDXGIAdapter3> Adapter3;
-	const HRESULT AdapterHR = InAdapter.DXGIAdapter->QueryInterface(IID_PPV_ARGS(Adapter3.GetInitReference()));
-	if (SUCCEEDED(AdapterHR))
+	FD3DMemoryStats MemoryStats;
+	if (SUCCEEDED(UE::DXGIUtilities::GetD3DMemoryStats(InAdapter.DXGIAdapter, MemoryStats)))
 	{
-		DXGI_QUERY_VIDEO_MEMORY_INFO LocalMemoryInfo{};
-		VERIFYD3D11RESULT(Adapter3->QueryVideoMemoryInfo(0, DXGI_MEMORY_SEGMENT_GROUP_LOCAL, &LocalMemoryInfo));
-		UE_LOG(LogD3D11RHI, Error, TEXT("\tBudget:\t%7.2f MB"), LocalMemoryInfo.Budget / (1024.0f * 1024));
-		UE_LOG(LogD3D11RHI, Error, TEXT("\tUsed:\t%7.2f MB"), LocalMemoryInfo.CurrentUsage / (1024.0f * 1024));
-		OutVRAMBudgetBytes = LocalMemoryInfo.Budget;
-		OutVRAMUsageBytes = LocalMemoryInfo.CurrentUsage;
+		UE_LOG(LogD3D11RHI, Error, TEXT("\tLocal Budget:\t%7.2f MB"), MemoryStats.BudgetLocal / (1024.0f * 1024));
+		UE_LOG(LogD3D11RHI, Error, TEXT("\tLocal Used:\t%7.2f MB"), MemoryStats.UsedLocal / (1024.0f * 1024));
+		UE_LOG(LogD3D11RHI, Error, TEXT("\tSystem Budget:\t%7.2f MB"), MemoryStats.BudgetSystem / (1024.0f * 1024));
+		UE_LOG(LogD3D11RHI, Error, TEXT("\tSystem Used:\t%7.2f MB"), MemoryStats.UsedSystem / (1024.0f * 1024));
+		OutVRAMBudgetBytes = MemoryStats.BudgetLocal;
+		OutVRAMUsageBytes = MemoryStats.UsedLocal;
 	}
 }
 
@@ -154,7 +153,7 @@ static void TerminateOnOutOfMemory(HRESULT D3DResult, bool bCreatingTextures)
 	if (D3DResult == E_OUTOFMEMORY)
 	{
 		uint64 VRAMBudgetBytes = 0, VRAMUsageBytes = 0;
-		GetAndLogMemoryInfo(GD3D11RHI->GetAdapter(), VRAMBudgetBytes, VRAMUsageBytes);
+		GetAndLogMemoryStats(GD3D11RHI->GetAdapter(), VRAMBudgetBytes, VRAMUsageBytes);
 		FCoreDelegates::GetGPUOutOfMemoryDelegate().Broadcast(VRAMBudgetBytes, VRAMUsageBytes);
 
 		if (bCreatingTextures)
