Don't try to write bindless indices to uniform buffers if bindless isn't enabled.

diff --git a/Engine/Source/Runtime/D3D12RHI/Private/D3D12UniformBuffer.cpp b/Engine/Source/Runtime/D3D12RHI/Private/D3D12UniformBuffer.cpp
--- a/Engine/Source/Runtime/D3D12RHI/Private/D3D12UniformBuffer.cpp
+++ b/Engine/Source/Runtime/D3D12RHI/Private/D3D12UniformBuffer.cpp
@@ -9,6 +9,18 @@
 #include "ShaderParameterStruct.h"
 #include "RHIUniformBufferDataShared.h"
 
+inline bool AreBindlessUniformConstantsEnabled(FD3D12Device* Device)
+{
+#if PLATFORM_SUPPORTS_BINDLESS_RENDERING
+	FD3D12BindlessDescriptorManager& Manager = Device->GetBindlessDescriptorManager();
+	if (Manager.AreResourcesBindless() || Manager.AreSamplersBindless())
+	{
+		return true;
+	}
+#endif
+	return false;
+}
+
 FUniformBufferRHIRef FD3D12DynamicRHI::RHICreateUniformBuffer(const void* Contents, const FRHIUniformBufferLayout* Layout, EUniformBufferUsage Usage, EUniformBufferValidation Validation)
 {
 	SCOPE_CYCLE_COUNTER(STAT_D3D12UpdateUniformBufferTime);
@@ -61,7 +73,7 @@ FUniformBufferRHIRef FD3D12DynamicRHI::RHICreateUniformBuffer(const void* Conten
 				// Copy the data to the upload heap
 				check(MappedData != nullptr);
 
-				UE::RHICore::UpdateUniformBufferConstants(MappedData, Contents, *Layout);
+				UE::RHICore::UpdateUniformBufferConstants(MappedData, Contents, *Layout, AreBindlessUniformConstantsEnabled(Device));
 
 #if D3D12RHI_USE_CONSTANT_BUFFER_VIEWS
 				NewUniformBuffer->View->CreateView(&NewUniformBuffer->ResourceLocation, 0, NumBytesActualData);
@@ -197,7 +209,7 @@ void FD3D12DynamicRHI::RHIUpdateUniformBuffer(FRHICommandListBase& RHICmdList, F
 			}
 
 			check(MappedData != nullptr);
-			UE::RHICore::UpdateUniformBufferConstants(MappedData, Contents, Layout);
+			UE::RHICore::UpdateUniformBufferConstants(MappedData, Contents, Layout, AreBindlessUniformConstantsEnabled(Device));
 		}
 
 		if (bBypass)
