Adding RHIResourceCollection RHI resource type which can store bindless indices in a buffer.

diff --git a/Engine/Shaders/Public/BindlessResources.ush b/Engine/Shaders/Public/BindlessResources.ush
--- a/Engine/Shaders/Public/BindlessResources.ush
+++ b/Engine/Shaders/Public/BindlessResources.ush
@@ -38,3 +38,48 @@
 #else
 	#define UB_RESOURCE_MEMBER_SAMPLER(Type, UBName, Name) Type UBName##_##Name
 #endif
+
+#if ENABLE_BINDLESS_RESOURCES && PLATFORM_SUPPORTS_BINDLESS
+
+typedef ByteAddressBuffer FResourceCollection;
+
+uint GetCountFromResourceCollection(FResourceCollection ResourceCollection)
+{
+	return ResourceCollection.Load(0);
+}
+
+uint GetResourceIndexFromResourceCollection(FResourceCollection ResourceCollection, uint CollectionIndex)
+{
+	int Offset = (CollectionIndex + 1) * 4;
+	return ResourceCollection.Load(Offset);
+}
+
+#if defined(BINDLESS_SRV_ARRAY_PREFIX)
+	#define RESOURCE_COLLECTION_HEAP_NAME_INNER1(Prefix, Name) Prefix##Name
+	#define RESOURCE_COLLECTION_HEAP_NAME_INNER(Prefix, Name) RESOURCE_COLLECTION_HEAP_NAME_INNER1(Prefix,Name)
+	#define RESOURCE_COLLECTION_HEAP_NAME(Name) RESOURCE_COLLECTION_HEAP_NAME_INNER(BINDLESS_SRV_ARRAY_PREFIX, Name)
+
+	#define TEXTURE_FROM_COLLECTION_T(TType) \
+		typedef TType SafeType##TType; \
+		SafeType##TType RESOURCE_COLLECTION_HEAP_NAME(TType)[]; \
+		SafeType##TType TextureFromCollection_##TType(FResourceCollection ResourceCollection, uint CollectionIndex) \
+		{ \
+			uint ResourceIndex = GetResourceIndexFromResourceCollection(ResourceCollection, CollectionIndex); \
+			return RESOURCE_COLLECTION_HEAP_NAME(TType)[ResourceIndex]; \
+		}
+#else
+	#define TEXTURE_FROM_COLLECTION_T(TType) \
+		TType TextureFromCollection_##TType(FResourceCollection ResourceCollection, uint CollectionIndex) \
+		{ \
+			uint ResourceIndex = GetResourceIndexFromResourceCollection(ResourceCollection, CollectionIndex); \
+			return GetResourceFromHeap(TType, ResourceIndex); \
+		}
+#endif
+
+TEXTURE_FROM_COLLECTION_T(Texture2D)
+TEXTURE_FROM_COLLECTION_T(TextureCube)
+TEXTURE_FROM_COLLECTION_T(Texture2DArray)
+TEXTURE_FROM_COLLECTION_T(TextureCubeArray)
+TEXTURE_FROM_COLLECTION_T(Texture3D)
+
+#endif
diff --git a/Engine/Source/Developer/ShaderCompilerCommon/Private/ShaderCompilerCommon.cpp b/Engine/Source/Developer/ShaderCompilerCommon/Private/ShaderCompilerCommon.cpp
--- a/Engine/Source/Developer/ShaderCompilerCommon/Private/ShaderCompilerCommon.cpp
+++ b/Engine/Source/Developer/ShaderCompilerCommon/Private/ShaderCompilerCommon.cpp
@@ -42,6 +42,7 @@ int16 GetNumUniformBuffersUsed(const FShaderCompilerResourceTable& InSRT)
 	Num = FMath::Max(Num, (int16)CountLambda(InSRT.ShaderResourceViewMap));
 	Num = FMath::Max(Num, (int16)CountLambda(InSRT.TextureMap));
 	Num = FMath::Max(Num, (int16)CountLambda(InSRT.UnorderedAccessViewMap));
+	Num = FMath::Max(Num, (int16)CountLambda(InSRT.ResourceCollectionMap));
 	return Num;
 }
 
@@ -95,6 +96,7 @@ void UE::ShaderCompilerCommon::BuildShaderResourceTable(const FShaderCompilerRes
 	BuildResourceTableTokenStream(GenericSRT.ShaderResourceViewMap,  GenericSRT.MaxBoundResourceTable, OutSRT.ShaderResourceViewMap,  bGenerateEmptyTokenStreamIfNoResources);
 	BuildResourceTableTokenStream(GenericSRT.SamplerMap,             GenericSRT.MaxBoundResourceTable, OutSRT.SamplerMap,             bGenerateEmptyTokenStreamIfNoResources);
 	BuildResourceTableTokenStream(GenericSRT.UnorderedAccessViewMap, GenericSRT.MaxBoundResourceTable, OutSRT.UnorderedAccessViewMap, bGenerateEmptyTokenStreamIfNoResources);
+	BuildResourceTableTokenStream(GenericSRT.ResourceCollectionMap,  GenericSRT.MaxBoundResourceTable, OutSRT.ResourceCollectionMap,  bGenerateEmptyTokenStreamIfNoResources);
 }
 
 
@@ -163,6 +165,9 @@ bool BuildResourceTableMapping(
 			case UBMT_RDG_BUFFER_SRV:
 				OutSRT.ShaderResourceViewMap.Add(ResourceMap);
 				break;
+			case UBMT_RESOURCE_COLLECTION:
+				OutSRT.ResourceCollectionMap.Add(ResourceMap);
+				break;
 			case UBMT_UAV:
 			case UBMT_RDG_TEXTURE_UAV:
 			case UBMT_RDG_BUFFER_UAV:
diff --git a/Engine/Source/Developer/VulkanShaderFormat/Private/VulkanShaderCompiler.cpp b/Engine/Source/Developer/VulkanShaderFormat/Private/VulkanShaderCompiler.cpp
--- a/Engine/Source/Developer/VulkanShaderFormat/Private/VulkanShaderCompiler.cpp
+++ b/Engine/Source/Developer/VulkanShaderFormat/Private/VulkanShaderCompiler.cpp
@@ -871,6 +871,9 @@ static FShaderResourceTable BuildSRTFromHeader(const FVulkanShaderHeader& NEWHea
 				case UBMT_RDG_BUFFER_SRV:
 					GenericSRT.ShaderResourceViewMap.Add(ResourceMap);
 					break;
+				case UBMT_RESOURCE_COLLECTION:
+					GenericSRT.ResourceCollectionMap.Add(ResourceMap);
+					break;
 				case UBMT_UAV:
 				case UBMT_RDG_TEXTURE_UAV:
 				case UBMT_RDG_BUFFER_UAV:
diff --git a/Engine/Source/Runtime/Apple/MetalRHI/Private/MetalCommands.cpp b/Engine/Source/Runtime/Apple/MetalRHI/Private/MetalCommands.cpp
--- a/Engine/Source/Runtime/Apple/MetalRHI/Private/MetalCommands.cpp
+++ b/Engine/Source/Runtime/Apple/MetalRHI/Private/MetalCommands.cpp
@@ -23,6 +23,7 @@
 #include "RHICoreShader.h"
 #include "RHIShaderParametersShared.h"
 #include "RHIUtilities.h"
+#include "MetalResourceCollection.h"
 
 static const bool GUsesInvertedZ = true;
 
@@ -401,6 +402,14 @@ struct FMetalShaderBinder
         FMetalSamplerState* Sampler = ResourceCast(InSampler);
         StateCache.SetShaderSamplerState(Stage, Sampler, Index);
     }
+
+#if PLATFORM_SUPPORTS_BINDLESS_RENDERING
+	void SetResourceCollection(FRHIResourceCollection* ResourceCollection, uint32 Index)
+	{
+		FMetalResourceCollection* MetalResourceCollection = ResourceCast(ResourceCollection);
+		SetSRV(MetalResourceCollection->GetShaderResourceView(), Index);
+	}
+#endif
 };
 
 static void SetShaderParametersOnContext(
@@ -466,6 +475,9 @@ static void SetShaderParametersOnContext(
         case FRHIShaderParameterResource::EType::UniformBuffer:
             BindUniformBuffer(Binder.StateCache, Shader, Binder.Stage, Parameter.Index, static_cast<FRHIUniformBuffer*>(Parameter.Resource));
             break;
+        case FRHIShaderParameterResource::EType::ResourceCollection:
+            Binder.SetResourceCollection(static_cast<FRHIResourceCollection*>(Parameter.Resource), Parameter.Index);
+            break;
         default:
             checkf(false, TEXT("Unhandled resource type?"));
             break;
diff --git a/Engine/Source/Runtime/Apple/MetalRHI/Private/MetalDynamicRHI.h b/Engine/Source/Runtime/Apple/MetalRHI/Private/MetalDynamicRHI.h
--- a/Engine/Source/Runtime/Apple/MetalRHI/Private/MetalDynamicRHI.h
+++ b/Engine/Source/Runtime/Apple/MetalRHI/Private/MetalDynamicRHI.h
@@ -69,6 +69,10 @@ public:
 	virtual FShaderResourceViewRHIRef  RHICreateShaderResourceView (class FRHICommandListBase& RHICmdList, FRHIViewableResource* Resource, FRHIViewDesc const& ViewDesc) final override;
 	virtual FUnorderedAccessViewRHIRef RHICreateUnorderedAccessView(class FRHICommandListBase& RHICmdList, FRHIViewableResource* Resource, FRHIViewDesc const& ViewDesc) final override;
 
+#if PLATFORM_SUPPORTS_BINDLESS_RENDERING
+	virtual FRHIResourceCollectionRef RHICreateResourceCollection(FRHICommandListBase& RHICmdList, TConstArrayView<FRHIResourceCollectionMember> InMembers) final override;
+#endif
+
 	virtual FRHICalcTextureSizeResult RHICalcTexturePlatformSize(FRHITextureDesc const& Desc, uint32 FirstMipIndex) final override;
 	virtual uint64 RHIGetMinimumAlignmentForBufferBackedSRV(EPixelFormat Format) final override;
 	virtual void RHIGetTextureMemoryStats(FTextureMemoryStats& OutStats) final override;
diff --git a/Engine/Source/Runtime/Apple/MetalRHI/Private/MetalRHI.cpp b/Engine/Source/Runtime/Apple/MetalRHI/Private/MetalRHI.cpp
--- a/Engine/Source/Runtime/Apple/MetalRHI/Private/MetalRHI.cpp
+++ b/Engine/Source/Runtime/Apple/MetalRHI/Private/MetalRHI.cpp
@@ -28,6 +28,7 @@
 #include "EngineGlobals.h"
 #include "MetalBindlessDescriptors.h"
 #include "DataDrivenShaderPlatformInfo.h"
+#include "MetalResourceCollection.h"
  
 DEFINE_LOG_CATEGORY(LogMetal)
 
@@ -1537,3 +1538,10 @@ void FMetalDynamicRHI::RHIReplaceResources(FRHICommandListBase& RHICmdList, TArr
 
 	RHICmdList.RHIThreadFence(true);
 }
+
+#if PLATFORM_SUPPORTS_BINDLESS_RENDERING
+FRHIResourceCollectionRef FMetalDynamicRHI::RHICreateResourceCollection(FRHICommandListBase& RHICmdList, TConstArrayView<FRHIResourceCollectionMember> InMembers)
+{
+	return new FMetalResourceCollection(RHICmdList, InMembers);
+}
+#endif
diff --git a/Engine/Source/Runtime/Apple/MetalRHI/Private/MetalResourceCollection.h b/Engine/Source/Runtime/Apple/MetalRHI/Private/MetalResourceCollection.h
new file mode 100644
--- /dev/null
+++ b/Engine/Source/Runtime/Apple/MetalRHI/Private/MetalResourceCollection.h
@@ -0,0 +1,14 @@
+// Copyright Epic Games, Inc. All Rights Reserved.
+
+#pragma once
+
+#include "RHICoreResourceCollection.h"
+#include "MetalRHIPrivate.h"
+
+using FMetalResourceCollection = UE::RHICore::FGenericResourceCollection;
+
+template<>
+struct TMetalResourceTraits<FRHIResourceCollection>
+{
+	using TConcreteType = FMetalResourceCollection;
+};
diff --git a/Engine/Source/Runtime/Apple/MetalRHI/Private/MetalStateCache.cpp b/Engine/Source/Runtime/Apple/MetalRHI/Private/MetalStateCache.cpp
--- a/Engine/Source/Runtime/Apple/MetalRHI/Private/MetalStateCache.cpp
+++ b/Engine/Source/Runtime/Apple/MetalRHI/Private/MetalStateCache.cpp
@@ -11,6 +11,7 @@
 #include "MetalVertexDeclaration.h"
 #include "MetalBindlessDescriptors.h"
 #include "RHIShaderParametersShared.h"
+#include "MetalResourceCollection.h"
 
 #if PLATFORM_MAC
 	#ifndef UINT128_MAX
@@ -1953,6 +1954,14 @@ void FMetalStateCache::SetResourcesFromTables(ShaderType Shader, CrossCompiler::
 					StateCache.SetShaderSamplerState(Frequency, static_cast<FMetalSamplerState*>(Sampler), Index);
 				}
 			}
+
+#if PLATFORM_SUPPORTS_BINDLESS_RENDERING
+			void SetResourceCollection(FRHIResourceCollection* ResourceCollection, uint32 Index)
+			{
+				FMetalResourceCollection* MetalResourceCollection = ResourceCast(ResourceCollection);
+				SetSRV(MetalResourceCollection->GetShaderResourceView(), Index);
+			}
+#endif
 		};
 
 		UE::RHICore::SetResourcesFromTables(
diff --git a/Engine/Source/Runtime/D3D12RHI/Private/D3D12Buffer.cpp b/Engine/Source/Runtime/D3D12RHI/Private/D3D12Buffer.cpp
--- a/Engine/Source/Runtime/D3D12RHI/Private/D3D12Buffer.cpp
+++ b/Engine/Source/Runtime/D3D12RHI/Private/D3D12Buffer.cpp
@@ -31,6 +31,79 @@ FD3D12Buffer::~FD3D12Buffer()
 	}
 }
 
+void FD3D12Buffer::UploadResourceData(FD3D12CommandContext& CommandContext, D3D12_RESOURCE_STATES InDestinationState, FD3D12ResourceLocation& DestinationResourceLocation, const FD3D12ResourceLocation& SourceResourceLocation, uint32 Size)
+{
+	FD3D12Resource* Destination = DestinationResourceLocation.GetResource();
+
+	// Copy from the temporary upload heap to the default resource
+
+	// if resource doesn't require state tracking then transition to copy dest here (could have been suballocated from shared resource) - not very optimal and should be batched
+	if (!Destination->RequiresResourceStateTracking())
+	{
+		CommandContext.AddTransitionBarrier(Destination, Destination->GetDefaultResourceState(), D3D12_RESOURCE_STATE_COPY_DEST, D3D12_RESOURCE_BARRIER_ALL_SUBRESOURCES);
+	}
+
+	CommandContext.FlushResourceBarriers();
+
+	CommandContext.GraphicsCommandList()->CopyBufferRegion(
+		Destination->GetResource(),
+		DestinationResourceLocation.GetOffsetFromBaseOfResource(),
+		SourceResourceLocation.GetResource()->GetResource(),
+		SourceResourceLocation.GetOffsetFromBaseOfResource(),
+		Size);
+
+	// Update the resource state after the copy has been done (will take care of updating the residency as well)
+	if (InDestinationState != D3D12_RESOURCE_STATE_COPY_DEST)
+	{
+		CommandContext.AddTransitionBarrier(Destination, D3D12_RESOURCE_STATE_COPY_DEST, InDestinationState, D3D12_RESOURCE_BARRIER_ALL_SUBRESOURCES);
+	}
+
+	CommandContext.UpdateResidency(SourceResourceLocation.GetResource());
+
+	CommandContext.ConditionalSplitCommandList();
+
+	// If the resource is untracked, the destination state must match the default state of the resource.
+	check(Destination->RequiresResourceStateTracking() || (Destination->GetDefaultResourceState() == InDestinationState));
+
+	// Buffer is now written and ready, so unlock the block (locked after creation and can be defragmented if needed)
+	DestinationResourceLocation.UnlockPoolData();
+}
+
+void FD3D12Buffer::UploadResourceData(FRHICommandListBase& InRHICmdList, FRHIGPUMask GPUMask, D3D12_RESOURCE_STATES InDestinationState, const void* SourceData, int32 SourceDataSize)
+{
+	FD3D12Device* ParentDevice = GetParentDevice();
+
+	const uint32 BufferSize = GetSize();
+	check(BufferSize == SourceDataSize);
+
+	const bool bOnAsyncThread = !IsInRHIThread() && !IsInRenderingThread();
+
+	// Get an upload heap and initialize data
+	FD3D12ResourceLocation SrcResourceLoc(ParentDevice);
+	void* pData;
+	if (bOnAsyncThread)
+	{
+		pData = ParentDevice->GetParentAdapter()->GetUploadHeapAllocator(ParentDevice->GetGPUIndex()).AllocUploadResource(BufferSize, 4u, SrcResourceLoc);
+	}
+	else
+	{
+		pData = ParentDevice->GetDefaultFastAllocator().Allocate(BufferSize, 4UL, &SrcResourceLoc);
+	}
+	check(pData);
+	FMemory::Memcpy(pData, SourceData, BufferSize);
+
+	InRHICmdList.EnqueueLambda(
+		[this, GPUMask, InDestinationState, SrcResourceLoc = MoveTemp(SrcResourceLoc)](FRHICommandListBase& ExecutingCmdList)
+		{
+			const FRHIGPUMask EffectiveMask = GPUMask & ExecutingCmdList.GetGPUMask();
+			for (uint32 GPUIndex : EffectiveMask)
+			{
+				FD3D12CommandContext& CommandContext = FD3D12CommandContext::Get(ExecutingCmdList, GPUIndex);
+				UploadResourceData(CommandContext, InDestinationState, this->ResourceLocation, SrcResourceLoc, GetSize());
+			}
+		});
+}
+
 void FD3D12Buffer::UploadResourceData(FRHICommandListBase& RHICmdList, FResourceArrayInterface* InResourceArray, D3D12_RESOURCE_STATES InDestinationState, const TCHAR* AssetName, const FName& ClassName, const FName& PackageName)
 {
 	UE_TRACE_METADATA_SCOPE_ASSET_FNAME(AssetName, ClassName, PackageName);
@@ -49,70 +122,7 @@ void FD3D12Buffer::UploadResourceData(FRHICommandListBase& RHICmdList, FResource
 	}
 	else
 	{
-		const bool bOnAsyncThread = !IsInRHIThread() && !IsInRenderingThread();
-
-		// Get an upload heap and initialize data
-		FD3D12ResourceLocation SrcResourceLoc(GetParentDevice());
-		void* pData;
-		if (bOnAsyncThread)
-		{
-			const uint32 GPUIdx = SrcResourceLoc.GetParentDevice()->GetGPUIndex();
-			pData = GetParentDevice()->GetParentAdapter()->GetUploadHeapAllocator(GPUIdx).AllocUploadResource(BufferSize, 4u, SrcResourceLoc);
-		}
-		else
-		{
-			pData = SrcResourceLoc.GetParentDevice()->GetDefaultFastAllocator().Allocate(BufferSize, 4UL, &SrcResourceLoc);
-		}
-		check(pData);
-		FMemory::Memcpy(pData, InResourceArray->GetResourceData(), BufferSize);
-
-		RHICmdList.EnqueueLambda(
-			[
-				  Buffer = static_cast<FRHIBuffer*>(this)
-				, SrcResourceLoc = MoveTemp(SrcResourceLoc)
-				, Size = BufferSize
-				, DestinationState = InDestinationState
-			](FRHICommandListBase& ExecutingCmdList)
-		{
-			for (uint32 GPUIndex : ExecutingCmdList.GetGPUMask())
-			{
-				FD3D12CommandContext& CommandContext = FD3D12CommandContext::Get(ExecutingCmdList, GPUIndex);
-				FD3D12Buffer* CurrentBuffer = FD3D12DynamicRHI::ResourceCast(Buffer, GPUIndex);
-				FD3D12Resource* Destination = CurrentBuffer->ResourceLocation.GetResource();
-
-				// Copy from the temporary upload heap to the default resource
-				
-				// if resource doesn't require state tracking then transition to copy dest here (could have been suballocated from shared resource) - not very optimal and should be batched
-				if (!Destination->RequiresResourceStateTracking())
-				{
-					CommandContext.AddTransitionBarrier(Destination, Destination->GetDefaultResourceState(), D3D12_RESOURCE_STATE_COPY_DEST, D3D12_RESOURCE_BARRIER_ALL_SUBRESOURCES);
-				}
-
-				CommandContext.FlushResourceBarriers();
-
-				CommandContext.GraphicsCommandList()->CopyBufferRegion(
-					Destination->GetResource(),
-					CurrentBuffer->ResourceLocation.GetOffsetFromBaseOfResource(),
-					SrcResourceLoc.GetResource()->GetResource(),
-					SrcResourceLoc.GetOffsetFromBaseOfResource(), Size);
-
-				// Update the resource state after the copy has been done (will take care of updating the residency as well)
-				if (DestinationState != D3D12_RESOURCE_STATE_COPY_DEST)
-				{
-					CommandContext.AddTransitionBarrier(Destination, D3D12_RESOURCE_STATE_COPY_DEST, DestinationState, D3D12_RESOURCE_BARRIER_ALL_SUBRESOURCES);
-				}
-
-				CommandContext.UpdateResidency(SrcResourceLoc.GetResource());
-
-				CommandContext.ConditionalSplitCommandList();
-
-				// If the resource is untracked, the destination state must match the default state of the resource.
-				check(Destination->RequiresResourceStateTracking() || (Destination->GetDefaultResourceState() == DestinationState));
-
-				// Buffer is now written and ready, so unlock the block (locked after creation and can be defragmented if needed)
-				CurrentBuffer->ResourceLocation.UnlockPoolData();
-			}
-		});
+		UploadResourceData(RHICmdList, FRHIGPUMask::All(), InDestinationState, InResourceArray->GetResourceData(), BufferSize);
 	}
 
 	// Discard the resource array's contents.
@@ -209,7 +219,7 @@ FD3D12Buffer* FD3D12Adapter::CreateRHIBuffer(
 	FRHIBufferDesc const& BufferDesc,
 	ED3D12ResourceStateMode InResourceStateMode,
 	D3D12_RESOURCE_STATES InCreateState,
-	bool bHasInitialData,
+	bool bKeepUnlocked,
 	const FRHIGPUMask& InGPUMask,
 	ID3D12ResourceAllocator* ResourceAllocator,
 	const TCHAR* InDebugName,
@@ -275,7 +285,7 @@ FD3D12Buffer* FD3D12Adapter::CreateRHIBuffer(
 			AllocateBuffer(Device, InDesc, BufferDesc.Size, BufferDesc.Usage, InResourceStateMode, InCreateState, Alignment, NewBuffer, NewBuffer->ResourceLocation, ResourceAllocator, InDebugName);
 			
 			// Unlock immediately if there is no initial data
-			if (!bHasInitialData)
+			if (!bKeepUnlocked)
 			{
 				NewBuffer->ResourceLocation.UnlockPoolData();
 			}
@@ -407,7 +417,7 @@ FBufferRHIRef FD3D12DynamicRHI::CreateBuffer(FRHICommandListBase& RHICmdList, FR
 	return CreateD3D12Buffer(&RHICmdList, BufferDesc, InResourceState, CreateInfo);
 }
 
-FD3D12Buffer* FD3D12DynamicRHI::CreateD3D12Buffer(class FRHICommandListBase* RHICmdList, FRHIBufferDesc const& BufferDesc, ERHIAccess InResourceState, FRHIResourceCreateInfo& CreateInfo, ID3D12ResourceAllocator* ResourceAllocator)
+FD3D12Buffer* FD3D12DynamicRHI::CreateD3D12Buffer(class FRHICommandListBase* RHICmdList, FRHIBufferDesc const& BufferDesc, ERHIAccess InResourceState, FRHIResourceCreateInfo& CreateInfo, ID3D12ResourceAllocator* ResourceAllocator, bool bForceKeepUnlocked)
 {
 	FName TraceClassName = GetRHIBufferClassName(CreateInfo.GetTraceClassName());
 
@@ -443,7 +453,7 @@ FD3D12Buffer* FD3D12DynamicRHI::CreateD3D12Buffer(class FRHICommandListBase* RHI
 	// Setup the state at which the resource needs to be created - copy dest only supported for placed resources
 	D3D12_RESOURCE_STATES CreateState = (CreateInfo.ResourceArray && bSupportResourceStateTracking) ? D3D12_RESOURCE_STATE_COPY_DEST : DesiredState;
 
-	FD3D12Buffer* Buffer = GetAdapter().CreateRHIBuffer(Desc, Alignment, BufferDesc, StateMode, CreateState, bHasInitialData, CreateInfo.GPUMask, ResourceAllocator, CreateInfo.DebugName, CreateInfo.OwnerName, TraceClassName);
+	FD3D12Buffer* Buffer = GetAdapter().CreateRHIBuffer(Desc, Alignment, BufferDesc, StateMode, CreateState, bHasInitialData || bForceKeepUnlocked, CreateInfo.GPUMask, ResourceAllocator, CreateInfo.DebugName, CreateInfo.OwnerName, TraceClassName);
 	check(Buffer->ResourceLocation.IsValid());
 
 	// Copy the resource data if available 
diff --git a/Engine/Source/Runtime/D3D12RHI/Private/D3D12Commands.cpp b/Engine/Source/Runtime/D3D12RHI/Private/D3D12Commands.cpp
--- a/Engine/Source/Runtime/D3D12RHI/Private/D3D12Commands.cpp
+++ b/Engine/Source/Runtime/D3D12RHI/Private/D3D12Commands.cpp
@@ -7,6 +7,7 @@ D3D12Commands.cpp: D3D RHI commands implementation.
 #include "D3D12RHIPrivate.h"
 #include "ProfilingDebugging/AssetMetadataTrace.h"
 #include "ProfilingDebugging/RealtimeGPUProfiler.h"
+#include "D3D12ResourceCollection.h"
 
 static int32 GD3D12AllowDiscardResources = 1;
 static FAutoConsoleVariableRef CVarD3D12AllowDiscardResources(
@@ -941,6 +942,23 @@ struct FD3D12ResourceBinder
 			Context.StateCache.SetSamplerState(Frequency, D3D12SamplerState, Index);
 		}
 	}
+
+	void SetResourceCollection(FRHIResourceCollection* ResourceCollection, uint32 Index)
+	{
+#if PLATFORM_SUPPORTS_BINDLESS_RENDERING
+		FD3D12ResourceCollection* D3D12ResourceCollection = FD3D12CommandContext::RetrieveObject<FD3D12ResourceCollection>(ResourceCollection, GpuIndex);
+		FD3D12ShaderResourceView* D3D12ShaderResourceView = D3D12ResourceCollection ? D3D12ResourceCollection->GetShaderResourceView() : nullptr;
+
+		if (bBindlessResources)
+		{
+			Context.StateCache.QueueBindlessSRV(Frequency, D3D12ShaderResourceView);
+		}
+		else
+		{
+			checkNoEntry();
+		}
+#endif // PLATFORM_SUPPORTS_BINDLESS_RENDERING
+	}
 };
 
 static void SetShaderParametersOnContext(
@@ -989,6 +1007,10 @@ static void SetShaderParametersOnContext(
 				Handle = static_cast<FRHISamplerState*>(Resource)->GetBindlessHandle();
 				Binder.SetSampler(static_cast<FRHISamplerState*>(Resource), Parameter.Index);
 				break;
+			case FRHIShaderParameterResource::EType::ResourceCollection:
+				Handle = static_cast<FRHIResourceCollection*>(Resource)->GetBindlessHandle();
+				Binder.SetResourceCollection(static_cast<FRHIResourceCollection*>(Resource), Parameter.Index);
+				break;
 			}
 
 			checkf(Handle.IsValid(), TEXT("D3D12 resource did not provide a valid descriptor handle. Please validate that all D3D12 types can provide this or that the resource is still valid."));
@@ -1030,6 +1052,9 @@ static void SetShaderParametersOnContext(
 		case FRHIShaderParameterResource::EType::UniformBuffer:
 			BindUniformBuffer(Context, Shader, ShaderFrequency, Parameter.Index, FD3D12CommandContext::RetrieveObject<FD3D12UniformBuffer>(Parameter.Resource, GpuIndex));
 			break;
+		case FRHIShaderParameterResource::EType::ResourceCollection:
+			Binder.SetResourceCollection(static_cast<FRHIResourceCollection*>(Parameter.Resource), Parameter.Index);
+			break;
 		default:
 			checkf(false, TEXT("Unhandled resource type?"));
 			break;
diff --git a/Engine/Source/Runtime/D3D12RHI/Private/D3D12RHICommon.h b/Engine/Source/Runtime/D3D12RHI/Private/D3D12RHICommon.h
--- a/Engine/Source/Runtime/D3D12RHI/Private/D3D12RHICommon.h
+++ b/Engine/Source/Runtime/D3D12RHI/Private/D3D12RHICommon.h
@@ -339,6 +339,7 @@ public:
 
 	ObjectType0* GetFirst() const { return Object0; }
 	ObjectType1* GetSecond() const { return Object1; }
+	uint32 GetGPUIndex() const { return *GPUIterator; }
 
 private:
 	FRHIGPUMask::FIterator GPUIterator;
diff --git a/Engine/Source/Runtime/D3D12RHI/Private/D3D12RHIPrivate.h b/Engine/Source/Runtime/D3D12RHI/Private/D3D12RHIPrivate.h
--- a/Engine/Source/Runtime/D3D12RHI/Private/D3D12RHIPrivate.h
+++ b/Engine/Source/Runtime/D3D12RHI/Private/D3D12RHIPrivate.h
@@ -401,6 +401,7 @@ public:
 	// SRV / UAV creation functions
 	virtual FShaderResourceViewRHIRef  RHICreateShaderResourceView (class FRHICommandListBase& RHICmdList, FRHIViewableResource* Resource, FRHIViewDesc const& ViewDesc) override;
 	virtual FUnorderedAccessViewRHIRef RHICreateUnorderedAccessView(class FRHICommandListBase& RHICmdList, FRHIViewableResource* Resource, FRHIViewDesc const& ViewDesc) override;
+	virtual FRHIResourceCollectionRef RHICreateResourceCollection(FRHICommandListBase& RHICmdList, TConstArrayView<FRHIResourceCollectionMember> InMembers) override;
 
 	// ID3D12DynamicRHI interface.
 	virtual TArray<FD3D12MinimalAdapterDesc> RHIGetAdapterDescs() const final override;
@@ -549,7 +550,7 @@ public:
 	virtual FD3D12ResourceDesc GetResourceDesc(const FRHITextureDesc& CreateInfo) const;
 
 	virtual FD3D12Texture* CreateD3D12Texture(const FRHITextureCreateDesc& CreateDesc, FRHICommandListBase* RHICmdList, ID3D12ResourceAllocator* ResourceAllocator = nullptr);
-	FD3D12Buffer* CreateD3D12Buffer(class FRHICommandListBase* RHICmdList, FRHIBufferDesc const& BufferDesc, ERHIAccess ResourceState, FRHIResourceCreateInfo& CreateInfo, ID3D12ResourceAllocator* ResourceAllocator = nullptr);
+	FD3D12Buffer* CreateD3D12Buffer(class FRHICommandListBase* RHICmdList, FRHIBufferDesc const& BufferDesc, ERHIAccess ResourceState, FRHIResourceCreateInfo& CreateInfo, ID3D12ResourceAllocator* ResourceAllocator = nullptr, bool bForceKeepUnlocked = false);
 	virtual FD3D12Texture* CreateNewD3D12Texture(const FRHITextureCreateDesc& CreateDesc, class FD3D12Device* Device);
 
 	FRHIBuffer* CreateBuffer(const FRHIBufferCreateInfo& CreateInfo, const TCHAR* DebugName, ERHIAccess InitialState, ID3D12ResourceAllocator* ResourceAllocator);
diff --git a/Engine/Source/Runtime/D3D12RHI/Private/D3D12RayTracing.cpp b/Engine/Source/Runtime/D3D12RHI/Private/D3D12RayTracing.cpp
--- a/Engine/Source/Runtime/D3D12RHI/Private/D3D12RayTracing.cpp
+++ b/Engine/Source/Runtime/D3D12RHI/Private/D3D12RayTracing.cpp
@@ -21,6 +21,7 @@
 #include "GlobalRenderResources.h"
 #include "D3D12RayTracingDebug.h"
 #include "D3D12ExplicitDescriptorCache.h"
+#include "D3D12ResourceCollection.h"
 
 extern int32 GD3D12ExplicitViewDescriptorHeapSize;
 extern int32 GD3D12ExplicitViewDescriptorHeapOverflowReported;
@@ -4141,6 +4142,12 @@ static bool SetRayTracingShaderResources(
 
 	struct FBindings
 	{
+		FBindings(ResourceBinderType& InBinder, uint32 InGPUIndex)
+			: Binder(InBinder)
+			, GPUIndex(InGPUIndex)
+		{
+		}
+
 		ResourceBinderType& Binder;
 		uint32 GPUIndex;
 
@@ -4216,6 +4223,23 @@ static bool SetRayTracingShaderResources(
 			Binder.AddResourceTransition(SRV);
 		}
 
+		void SetResourceCollection(FRHIResourceCollection* ResourceCollection, uint8 Index)
+		{
+			FD3D12ResourceCollection* D3D12ResourceCollection = FD3D12CommandContext::RetrieveObject<FD3D12ResourceCollection>(ResourceCollection, GPUIndex);
+			FD3D12ShaderResourceView* SRV = D3D12ResourceCollection ? D3D12ResourceCollection->GetShaderResourceView() : nullptr;
+
+			check(SRV != nullptr);
+
+			FD3D12OfflineDescriptor Descriptor = SRV->GetOfflineCpuHandle();
+			LocalSRVs[Index] = Descriptor;
+			SRVVersions[Index] = Descriptor.GetVersion();
+
+			BoundSRVMask |= 1ull << Index;
+
+			ReferencedResources.Add(SRV->GetResource());
+			Binder.AddResourceTransition(SRV);
+		}
+
 		void SetSampler(FRHISamplerState* RHISampler, uint8 Index)
 		{
 			FD3D12SamplerState* Sampler = FD3D12CommandContext::RetrieveObject<FD3D12SamplerState>(RHISampler, GPUIndex);
@@ -4227,8 +4251,8 @@ static bool SetRayTracingShaderResources(
 
 			BoundSamplerMask |= 1ull << Index;
 		}
-
-	} Bindings { Binder, Binder.GetDevice()->GetGPUIndex() };
+	};
+	FBindings Bindings(Binder, Binder.GetDevice()->GetGPUIndex());
 
 	for (uint32 TextureIndex = 0; TextureIndex < InNumTextures; ++TextureIndex)
 	{
diff --git a/Engine/Source/Runtime/D3D12RHI/Private/D3D12ResourceCollection.cpp b/Engine/Source/Runtime/D3D12RHI/Private/D3D12ResourceCollection.cpp
new file mode 100644
--- /dev/null
+++ b/Engine/Source/Runtime/D3D12RHI/Private/D3D12ResourceCollection.cpp
@@ -0,0 +1,99 @@
+// Copyright Epic Games, Inc. All Rights Reserved.
+
+#include "D3D12ResourceCollection.h"
+
+#if PLATFORM_SUPPORTS_BINDLESS_RENDERING
+
+#include "D3D12RHIPrivate.h"
+#include "D3D12CommandContext.h"
+#include "D3D12TextureReference.h"
+
+FD3D12ResourceCollection::FD3D12ResourceCollection(FD3D12Device* InParent, FD3D12Buffer* InBuffer, TConstArrayView<FRHIResourceCollectionMember> InMembers)
+	: FRHIResourceCollection(InMembers)
+	, FD3D12DeviceChild(InParent)
+	, Buffer(InBuffer)
+{
+	D3D12_SHADER_RESOURCE_VIEW_DESC SRVDesc{};
+	SRVDesc.Format = DXGI_FORMAT_R32_TYPELESS;
+	SRVDesc.ViewDimension = D3D12_SRV_DIMENSION_BUFFER;
+	SRVDesc.Shader4ComponentMapping = D3D12_DEFAULT_SHADER_4_COMPONENT_MAPPING;
+	SRVDesc.Buffer.FirstElement = InBuffer->ResourceLocation.GetOffsetFromBaseOfResource() / 4;
+	SRVDesc.Buffer.NumElements = UE::RHICore::CalculateResourceCollectionMemorySize(InMembers) / 4;
+	SRVDesc.Buffer.Flags = D3D12_BUFFER_SRV_FLAG_RAW;
+
+	ShaderResourceView = MakeShared<FD3D12ShaderResourceView>(InParent);
+	ShaderResourceView->CreateView(InBuffer, SRVDesc, FD3D12ShaderResourceView::EFlags::None);
+}
+
+FD3D12ResourceCollection::~FD3D12ResourceCollection() = default;
+
+FRHIDescriptorHandle FD3D12ResourceCollection::GetBindlessHandle() const
+{
+	return ShaderResourceView->GetBindlessHandle();
+}
+
+static FD3D12Buffer* CreateCollectionBuffer(FD3D12DynamicRHI& RHI, FRHICommandListBase& RHICmdList, TConstArrayView<FRHIResourceCollectionMember> InMembers)
+{
+	const size_t BufferSize = UE::RHICore::CalculateResourceCollectionMemorySize(InMembers);
+	FRHIBufferDesc BufferDesc(BufferSize, 4, EBufferUsageFlags::Static | EBufferUsageFlags::ByteAddressBuffer);
+
+	FRHIResourceCreateInfo CreateInfo(TEXT("ResourceCollection"));
+	return RHI.CreateD3D12Buffer(&RHICmdList, BufferDesc, ERHIAccess::SRVMask, CreateInfo, nullptr, true);
+}
+
+FRHIResourceCollectionRef FD3D12DynamicRHI::RHICreateResourceCollection(FRHICommandListBase& RHICmdList, TConstArrayView<FRHIResourceCollectionMember> InMembers)
+{
+	FD3D12Buffer* Buffer = CreateCollectionBuffer(*this, RHICmdList, InMembers);
+
+	FRHIViewDesc::FBufferSRV::FInitializer ViewDesc = FRHIViewDesc::CreateBufferSRV().SetType(FRHIViewDesc::EBufferType::Raw);
+	FShaderResourceViewRHIRef ShaderResourceView = RHICmdList.CreateShaderResourceView(Buffer, ViewDesc);
+
+	return GetAdapter().CreateLinkedObject<FD3D12ResourceCollection>(FRHIGPUMask::All(), [&RHICmdList, Buffer, InMembers](FD3D12Device* Device)
+	{
+		const uint32 GpuIndex = Device->GetGPUIndex();
+
+		TArray<FRHIDescriptorHandle> Handles;
+
+		for (const FRHIResourceCollectionMember& Member : InMembers)
+		{
+			switch (Member.Type)
+			{
+			case FRHIResourceCollectionMember::EType::Texture:
+			{
+				FD3D12Texture* Texture = FD3D12CommandContext::RetrieveTexture(static_cast<FRHITexture*>(Member.Resource), GpuIndex);
+				Handles.Emplace(Texture->GetDefaultBindlessHandle());
+			}
+			break;
+			case FRHIResourceCollectionMember::EType::TextureReference:
+			{
+				FD3D12RHITextureReference* TextureReference = FD3D12CommandContext::RetrieveObject<FD3D12RHITextureReference>(Member.Resource, GpuIndex);
+				Handles.Emplace(TextureReference->GetDefaultBindlessHandle());
+			}
+			break;
+			case FRHIResourceCollectionMember::EType::ShaderResourceView:
+			{
+				FD3D12ShaderResourceView_RHI* ShaderResourceView = FD3D12CommandContext::RetrieveObject<FD3D12ShaderResourceView_RHI>(Member.Resource, GpuIndex);
+				Handles.Emplace(ShaderResourceView->GetBindlessHandle());
+			}
+			break;
+			case FRHIResourceCollectionMember::EType::UnorderedAccessView:
+			{
+				FD3D12UnorderedAccessView_RHI* UnorderedAccessView = FD3D12CommandContext::RetrieveObject<FD3D12UnorderedAccessView_RHI>(Member.Resource, GpuIndex);
+				Handles.Emplace(UnorderedAccessView->GetBindlessHandle());
+			}
+			break;
+			}
+		}
+
+		FD3D12Buffer* GpuBuffer = Buffer->GetLinkedObject(GpuIndex);
+		FD3D12ResourceCollection* ResourceCollection = new FD3D12ResourceCollection(Device, GpuBuffer, InMembers);
+
+		constexpr D3D12_RESOURCE_STATES States = D3D12_RESOURCE_STATE_GENERIC_READ;
+		const TArray<uint32> CollectionMemory = UE::RHICore::CreateResourceCollectionArray<FRHIDescriptorHandle>(Handles);
+		GpuBuffer->UploadResourceData(RHICmdList, FRHIGPUMask::FromIndex(GpuIndex), States, CollectionMemory.GetData(), CollectionMemory.GetTypeSize() * CollectionMemory.Num());
+
+		return ResourceCollection;
+	});
+}
+
+#endif // PLATFORM_SUPPORTS_BINDLESS_RENDERING
diff --git a/Engine/Source/Runtime/D3D12RHI/Private/D3D12ResourceCollection.h b/Engine/Source/Runtime/D3D12RHI/Private/D3D12ResourceCollection.h
new file mode 100644
--- /dev/null
+++ b/Engine/Source/Runtime/D3D12RHI/Private/D3D12ResourceCollection.h
@@ -0,0 +1,36 @@
+// Copyright Epic Games, Inc. All Rights Reserved.
+
+#pragma once
+
+#include "D3D12Resources.h"
+#include "D3D12RHICommon.h"
+#include "RHIResources.h"
+#include "RHICoreResourceCollection.h"
+
+#if PLATFORM_SUPPORTS_BINDLESS_RENDERING
+
+class FD3D12Buffer;
+class FD3D12ResourceCollection;
+class FD3D12ShaderResourceView;
+
+class FD3D12ResourceCollection : public FRHIResourceCollection, public FD3D12DeviceChild, public FD3D12LinkedAdapterObject<FD3D12ResourceCollection>
+{
+public:
+	FD3D12ResourceCollection(FD3D12Device* InParent, FD3D12Buffer* InBuffer, TConstArrayView<FRHIResourceCollectionMember> InMembers);
+	~FD3D12ResourceCollection();
+
+	virtual FRHIDescriptorHandle GetBindlessHandle() const final;
+
+	FD3D12ShaderResourceView* GetShaderResourceView() const { return ShaderResourceView.Get(); }
+
+	TRefCountPtr<FD3D12Buffer> Buffer;
+	TSharedPtr<FD3D12ShaderResourceView> ShaderResourceView;
+};
+
+template<>
+struct TD3D12ResourceTraits<FRHIResourceCollection>
+{
+	using TConcreteType = FD3D12ResourceCollection;
+};
+
+#endif // PLATFORM_SUPPORTS_BINDLESS_RENDERING
\ No newline at end of file
diff --git a/Engine/Source/Runtime/D3D12RHI/Private/D3D12Resources.h b/Engine/Source/Runtime/D3D12RHI/Private/D3D12Resources.h
--- a/Engine/Source/Runtime/D3D12RHI/Private/D3D12Resources.h
+++ b/Engine/Source/Runtime/D3D12RHI/Private/D3D12Resources.h
@@ -939,7 +939,10 @@ public:
 
 	virtual uint32 GetParentGPUIndex() const override;
 
+	static void UploadResourceData(FD3D12CommandContext& CommandContext, D3D12_RESOURCE_STATES InDestinationState, FD3D12ResourceLocation& DestinationResourceLocation, const FD3D12ResourceLocation& SourceResourceLocation, uint32 Size);
+	void UploadResourceData(FRHICommandListBase& InRHICmdList, FRHIGPUMask GPUMask, D3D12_RESOURCE_STATES InDestinationState, const void* SourceData, int32 SourceDataSize);
 	void UploadResourceData(FRHICommandListBase& InRHICmdList, FResourceArrayInterface* InResourceArray, D3D12_RESOURCE_STATES InDestinationState, const TCHAR* AssetName = nullptr, const FName& ClassName = NAME_None, const FName& PackageName = NAME_None);
+
 	FD3D12SyncPointRef UploadResourceDataViaCopyQueue(FResourceArrayInterface* InResourceArray);
 
 	// FRHIResource overrides
diff --git a/Engine/Source/Runtime/D3D12RHI/Private/D3D12WorkGraph.cpp b/Engine/Source/Runtime/D3D12RHI/Private/D3D12WorkGraph.cpp
--- a/Engine/Source/Runtime/D3D12RHI/Private/D3D12WorkGraph.cpp
+++ b/Engine/Source/Runtime/D3D12RHI/Private/D3D12WorkGraph.cpp
@@ -6,6 +6,7 @@
 #include "Containers/DynamicRHIResourceArray.h"
 #include "D3D12ExplicitDescriptorCache.h"
 #include "D3D12RHIPrivate.h"
+#include "D3D12ResourceCollection.h"
 #include "D3D12Shader.h"
 #include "PipelineStateCache.h"
 #include "ShaderBundles.h"
@@ -333,6 +334,20 @@ struct FWorkGraphShaderBundleBinder
 			BoundSamplerMask |= 1ull << Index;
 		}
 	}
+
+	void SetResourceCollection(FRHIResourceCollection* ResourceCollection, uint32 Index)
+	{
+#if PLATFORM_SUPPORTS_BINDLESS_RENDERING
+		FD3D12ResourceCollection* D3D12ResourceCollection = FD3D12CommandContext::RetrieveObject<FD3D12ResourceCollection>(ResourceCollection, GPUIndex);
+		FD3D12ShaderResourceView* SRV = D3D12ResourceCollection ? D3D12ResourceCollection->GetShaderResourceView() : nullptr;
+
+		if (bBindlessResources)
+		{
+			checkNoEntry();
+			//Context.StateCache.QueueBindlessSRV(Frequency, SRV);
+		}
+#endif // PLATFORM_SUPPORTS_BINDLESS_RENDERING
+	}
 };
 
 // Record bindings from shader bundle parameters.
@@ -418,6 +433,9 @@ static void RecordBindings(
 			BundleUniformBuffers[Parameter.Index] = FD3D12CommandContext::RetrieveObject<FD3D12UniformBuffer>(Parameter.Resource, 0 /*GpuIndex*/);
 			UniformBufferMask |= (1 << Parameter.Index);
 			break;
+		case FRHIShaderParameterResource::EType::ResourceCollection:
+			BundleBinder.SetResourceCollection(static_cast<FRHIResourceCollection*>(Parameter.Resource), Parameter.Index);
+			break;
 		default:
 			checkf(false, TEXT("Unhandled resource type?"));
 			break;
diff --git a/Engine/Source/Runtime/OpenGLDrv/Private/OpenGLCommands.cpp b/Engine/Source/Runtime/OpenGLDrv/Private/OpenGLCommands.cpp
--- a/Engine/Source/Runtime/OpenGLDrv/Private/OpenGLCommands.cpp
+++ b/Engine/Source/Runtime/OpenGLDrv/Private/OpenGLCommands.cpp
@@ -1094,6 +1094,13 @@ struct FOpenGLResourceBinder
 	{
 		RHI.InternalSetSamplerStates(FirstTextureUnit + Index, FOpenGLDynamicRHI::ResourceCast(Sampler));
 	}
+
+#if PLATFORM_SUPPORTS_BINDLESS_RENDERING
+	void SetResourceCollection(FRHIResourceCollection* ResourceCollection, uint32 Index)
+	{
+		checkNoEntry();
+	}
+#endif
 };
 
 void FOpenGLDynamicRHI::SetShaderParametersCommon(EShaderFrequency ShaderFrequency, TConstArrayView<uint8> InParametersData, TConstArrayView<FRHIShaderParameter> InParameters, TConstArrayView<FRHIShaderParameterResource> InResourceParameters)
diff --git a/Engine/Source/Runtime/RHI/Public/DynamicRHI.h b/Engine/Source/Runtime/RHI/Public/DynamicRHI.h
--- a/Engine/Source/Runtime/RHI/Public/DynamicRHI.h
+++ b/Engine/Source/Runtime/RHI/Public/DynamicRHI.h
@@ -28,6 +28,7 @@ class IRHITransientResourceAllocator;
 struct FDepthStencilStateInitializerRHI;
 struct FDisplayInformation;
 struct FRasterizerStateInitializerRHI;
+struct FRHIResourceCollectionMember;
 struct FRHIResourceCreateInfo;
 struct FRHIResourceInfo;
 struct FRHIUniformBufferLayout;
@@ -443,6 +444,11 @@ public:
 	virtual FShaderResourceViewRHIRef  RHICreateShaderResourceView (class FRHICommandListBase& RHICmdList, FRHIViewableResource* Resource, FRHIViewDesc const& ViewDesc) = 0;
 	virtual FUnorderedAccessViewRHIRef RHICreateUnorderedAccessView(class FRHICommandListBase& RHICmdList, FRHIViewableResource* Resource, FRHIViewDesc const& ViewDesc) = 0;
 
+	virtual FRHIResourceCollectionRef RHICreateResourceCollection(FRHICommandListBase& RHICmdList, TConstArrayView<FRHIResourceCollectionMember> InMembers)
+	{
+		return FRHIResourceCollectionRef{};
+	}
+
 	/**
 	* Generates mip maps for a texture.
 	*/
diff --git a/Engine/Source/Runtime/RHI/Public/RHICommandList.h b/Engine/Source/Runtime/RHI/Public/RHICommandList.h
--- a/Engine/Source/Runtime/RHI/Public/RHICommandList.h
+++ b/Engine/Source/Runtime/RHI/Public/RHICommandList.h
@@ -965,6 +965,12 @@ public:
 		);
 	}
 
+	FORCEINLINE FRHIResourceCollectionRef CreateResourceCollection(TConstArrayView<FRHIResourceCollectionMember> InMembers)
+	{
+		LLM_SCOPE_BYNAME(TEXT("RHIMisc/CreateResourceCollection"));
+		return GDynamicRHI->RHICreateResourceCollection(*this, InMembers);
+	}
+
 #if RHI_RAYTRACING
 	FORCEINLINE FRayTracingGeometryRHIRef CreateRayTracingGeometry(const FRayTracingGeometryInitializer& Initializer)
 	{
diff --git a/Engine/Source/Runtime/RHI/Public/RHIDefinitions.h b/Engine/Source/Runtime/RHI/Public/RHIDefinitions.h
--- a/Engine/Source/Runtime/RHI/Public/RHIDefinitions.h
+++ b/Engine/Source/Runtime/RHI/Public/RHIDefinitions.h
@@ -501,6 +501,8 @@ enum class EShaderCodeResourceBindingType : uint8
 
 	RasterizerOrderedTexture2D,
 
+	ResourceCollection,
+
 	MAX
 };
 
@@ -518,6 +520,7 @@ inline bool IsResourceBindingTypeSRV(EShaderCodeResourceBindingType Type)
 	case EShaderCodeResourceBindingType::StructuredBuffer:
 	case EShaderCodeResourceBindingType::Buffer:
 	case EShaderCodeResourceBindingType::RaytracingAccelerationStructure:
+	case EShaderCodeResourceBindingType::ResourceCollection:
 		return true;
 	case EShaderCodeResourceBindingType::RWTexture2D:
 	case EShaderCodeResourceBindingType::RWTexture2DArray:
@@ -579,6 +582,8 @@ enum EUniformBufferBaseType : uint8
 	// Structure dedicated to setup render targets for a rasterizer pass.
 	UBMT_RENDER_TARGET_BINDING_SLOTS,
 
+	UBMT_RESOURCE_COLLECTION,
+
 	EUniformBufferBaseType_Num,
 	EUniformBufferBaseType_NumBits = 5,
 };
@@ -954,6 +959,7 @@ enum ERHIResourceType : uint8
 	RRT_WorkGraphShader,
 	RRT_WorkGraphPipelineState,
 	RRT_StreamSourceSlot,
+	RRT_ResourceCollection,
 
 	RRT_Num
 };
@@ -1340,15 +1346,22 @@ inline bool IsRDGResourceReferenceShaderParameterType(EUniformBufferBaseType Bas
 	return IsRDGTextureReferenceShaderParameterType(BaseType) || IsRDGBufferReferenceShaderParameterType(BaseType) || BaseType == UBMT_RDG_UNIFORM_BUFFER;
 }
 
-/** Returns whether the shader parameter type needs to be passdown to RHI through FRHIUniformBufferLayout when creating an uniform buffer. */
-inline bool IsShaderParameterTypeForUniformBufferLayout(EUniformBufferBaseType BaseType)
+inline bool IsShaderParameterTypeRHIResource(EUniformBufferBaseType BaseType)
 {
 	return
-		// RHI resource referenced in shader parameter structures.
 		BaseType == UBMT_TEXTURE ||
 		BaseType == UBMT_SRV ||
 		BaseType == UBMT_SAMPLER ||
 		BaseType == UBMT_UAV ||
+		BaseType == UBMT_RESOURCE_COLLECTION;
+}
+
+/** Returns whether the shader parameter type needs to be passdown to RHI through FRHIUniformBufferLayout when creating an uniform buffer. */
+inline bool IsShaderParameterTypeForUniformBufferLayout(EUniformBufferBaseType BaseType)
+{
+	return
+		// RHI resource referenced in shader parameter structures.
+		IsShaderParameterTypeRHIResource(BaseType) ||
 
 		// RHI is able to access RHI resources from RDG.
 		IsRDGResourceReferenceShaderParameterType(BaseType) ||
diff --git a/Engine/Source/Runtime/RHI/Public/RHIFwd.h b/Engine/Source/Runtime/RHI/Public/RHIFwd.h
--- a/Engine/Source/Runtime/RHI/Public/RHIFwd.h
+++ b/Engine/Source/Runtime/RHI/Public/RHIFwd.h
@@ -70,6 +70,7 @@ class FRHIRayTracingShader;
 class FRHIRenderQuery;
 class FRHIRenderQueryPool;
 class FRHIResource;
+class FRHIResourceCollection;
 class FRHISamplerState;
 class FRHIShader;
 class FRHIShaderData;
@@ -114,8 +115,10 @@ using FRayTracingSceneRHIRef           = TRefCountPtr<FRHIRayTracingScene>;
 using FRayTracingShaderRHIRef          = TRefCountPtr<FRHIRayTracingShader>;
 using FRenderQueryPoolRHIRef           = TRefCountPtr<FRHIRenderQueryPool>;
 using FRenderQueryRHIRef               = TRefCountPtr<FRHIRenderQuery>;
+using FRHIResourceCollectionRef        = TRefCountPtr<FRHIResourceCollection>;
 using FRHIPipelineBinaryLibraryRef     = TRefCountPtr<FRHIPipelineBinaryLibrary>;
 using FRHIShaderLibraryRef             = TRefCountPtr<FRHIShaderLibrary>;
+using FRHIShaderResourceViewRef        = TRefCountPtr<FRHIShaderResourceView>;
 using FSamplerStateRHIRef              = TRefCountPtr<FRHISamplerState>;
 using FShaderResourceViewRHIRef        = TRefCountPtr<FRHIShaderResourceView>;
 using FShaderBundleRHIRef              = TRefCountPtr<FRHIShaderBundle>;
diff --git a/Engine/Source/Runtime/RHI/Public/RHIResourceCollection.h b/Engine/Source/Runtime/RHI/Public/RHIResourceCollection.h
new file mode 100644
--- /dev/null
+++ b/Engine/Source/Runtime/RHI/Public/RHIResourceCollection.h
@@ -0,0 +1,69 @@
+// Copyright Epic Games, Inc. All Rights Reserved.
+
+#pragma once
+
+#include "RHIResources.h"
+#include "RHITextureReference.h"
+
+struct FRHIResourceCollectionMember
+{
+	enum class EType : uint8
+	{
+		Texture,
+		TextureReference,
+		ShaderResourceView,
+		UnorderedAccessView,
+		Sampler,
+	};
+
+	FRHIResourceCollectionMember() = default;
+	FRHIResourceCollectionMember(EType InType, FRHIResource* InResource)
+		: Resource(InResource)
+		, Type(InType)
+	{
+	}
+	FRHIResourceCollectionMember(FRHITexture* InTexture)
+		: FRHIResourceCollectionMember(FRHIResourceCollectionMember::EType::Texture, InTexture)
+	{
+	}
+	FRHIResourceCollectionMember(FRHITextureReference* InTextureReference)
+		: FRHIResourceCollectionMember(FRHIResourceCollectionMember::EType::TextureReference, InTextureReference)
+	{
+	}
+	FRHIResourceCollectionMember(FRHIShaderResourceView* InView)
+		: FRHIResourceCollectionMember(FRHIResourceCollectionMember::EType::ShaderResourceView, InView)
+	{
+	}
+	FRHIResourceCollectionMember(FRHIUnorderedAccessView* InUAV)
+		: FRHIResourceCollectionMember(FRHIResourceCollectionMember::EType::UnorderedAccessView, InUAV)
+	{
+	}
+	FRHIResourceCollectionMember(FRHISamplerState* InSamplerState)
+		: FRHIResourceCollectionMember(FRHIResourceCollectionMember::EType::Sampler, InSamplerState)
+	{
+	}
+
+	FRHIResource* Resource = nullptr;
+	EType         Type = EType::Texture;
+};
+
+class FRHIResourceCollection : public FRHIResource
+{
+public:
+	FRHIResourceCollection(TConstArrayView<FRHIResourceCollectionMember> InMembers)
+		: FRHIResource(RRT_ResourceCollection)
+	{
+		Resources.Reserve(InMembers.Num());
+		for (const FRHIResourceCollectionMember& Member : InMembers)
+		{
+			Resources.Emplace(Member.Resource);
+		}
+	}
+
+	virtual FRHIDescriptorHandle GetBindlessHandle() const
+	{
+		return FRHIDescriptorHandle();
+	}
+
+	TArray<TRefCountPtr<FRHIResource>> Resources;
+};
diff --git a/Engine/Source/Runtime/RHI/Public/RHIResources.h b/Engine/Source/Runtime/RHI/Public/RHIResources.h
--- a/Engine/Source/Runtime/RHI/Public/RHIResources.h
+++ b/Engine/Source/Runtime/RHI/Public/RHIResources.h
@@ -778,6 +778,9 @@ struct FShaderResourceTable
 	/** Mapping of bound Textures to their location in resource tables. */
 	TArray<uint32> TextureMap;
 
+	/** Mapping of bound Resource Collections to their location in resource tables. */
+	TArray<uint32> ResourceCollectionMap;
+
 	friend bool operator == (const FShaderResourceTable& A, const FShaderResourceTable& B)
 	{
 		bool bEqual = true;
@@ -787,6 +790,7 @@ struct FShaderResourceTable
 		bEqual &= (A.UnorderedAccessViewMap   .Num() == B.UnorderedAccessViewMap   .Num());
 		bEqual &= (A.ResourceTableLayoutHashes.Num() == B.ResourceTableLayoutHashes.Num());
 		bEqual &= (A.TextureMap               .Num() == B.TextureMap               .Num());
+		bEqual &= (A.ResourceCollectionMap    .Num() == B.ResourceCollectionMap    .Num());
 
 		if (!bEqual)
 		{
@@ -798,6 +802,7 @@ struct FShaderResourceTable
 		bEqual &= (FMemory::Memcmp(A.UnorderedAccessViewMap   .GetData(), B.UnorderedAccessViewMap   .GetData(), A.UnorderedAccessViewMap   .GetTypeSize() * A.UnorderedAccessViewMap   .Num()) == 0);
 		bEqual &= (FMemory::Memcmp(A.ResourceTableLayoutHashes.GetData(), B.ResourceTableLayoutHashes.GetData(), A.ResourceTableLayoutHashes.GetTypeSize() * A.ResourceTableLayoutHashes.Num()) == 0);
 		bEqual &= (FMemory::Memcmp(A.TextureMap               .GetData(), B.TextureMap               .GetData(), A.TextureMap               .GetTypeSize() * A.TextureMap               .Num()) == 0);
+		bEqual &= (FMemory::Memcmp(A.ResourceCollectionMap    .GetData(), B.ResourceCollectionMap    .GetData(), A.ResourceCollectionMap    .GetTypeSize() * A.ResourceCollectionMap    .Num()) == 0);
 		return bEqual;
 	}
 
@@ -809,6 +814,7 @@ struct FShaderResourceTable
 		Ar << SRT.UnorderedAccessViewMap;
 		Ar << SRT.ResourceTableLayoutHashes;
 		Ar << SRT.TextureMap;
+		Ar << SRT.ResourceCollectionMap;
 	
 		return Ar;
 	}
diff --git a/Engine/Source/Runtime/RHI/Public/RHIShaderParameters.h b/Engine/Source/Runtime/RHI/Public/RHIShaderParameters.h
--- a/Engine/Source/Runtime/RHI/Public/RHIShaderParameters.h
+++ b/Engine/Source/Runtime/RHI/Public/RHIShaderParameters.h
@@ -4,6 +4,7 @@
 
 #include "Containers/Array.h"
 #include "RHIResources.h"
+#include "RHIResourceCollection.h"
 
 class FRHICommandList;
 class FRHIComputeCommandList;
@@ -34,6 +35,7 @@ struct FRHIShaderParameterResource
 		UnorderedAccessView,
 		Sampler,
 		UniformBuffer,
+		ResourceCollection,
 	};
 
 	FRHIShaderParameterResource() = default;
@@ -63,6 +65,10 @@ struct FRHIShaderParameterResource
 		: FRHIShaderParameterResource(FRHIShaderParameterResource::EType::UniformBuffer, InUniformBuffer, InIndex)
 	{
 	}
+	FRHIShaderParameterResource(FRHIResourceCollection* InResourceCollection, uint16 InIndex)
+		: FRHIShaderParameterResource(FRHIShaderParameterResource::EType::ResourceCollection, InResourceCollection, InIndex)
+	{
+	}
 
 	FRHIResource* Resource = nullptr;
 	uint16        Index = 0;
@@ -118,6 +124,10 @@ struct FRHIBatchedShaderParameters
 	{
 		ResourceParameters.Emplace(UAV, (uint16)Index);
 	}
+	FORCEINLINE_DEBUGGABLE void SetResourceCollection(uint32 Index, FRHIResourceCollection* ResourceCollection)
+	{
+		ResourceParameters.Emplace(ResourceCollection, (uint16)Index);
+	}
 
 	FORCEINLINE_DEBUGGABLE void SetBindlessTexture(uint32 Index, FRHITexture* Texture)
 	{
@@ -135,6 +145,10 @@ struct FRHIBatchedShaderParameters
 	{
 		BindlessParameters.Emplace(UAV, (uint16)Index);
 	}
+	FORCEINLINE_DEBUGGABLE void SetBindlessResourceCollection(uint32 Index, FRHIResourceCollection* ResourceCollection)
+	{
+		BindlessParameters.Emplace(ResourceCollection, (uint16)Index);
+	}
 };
 
 /** Compact representation of a resource parameter unbind, limited to  SRVs and UAVs */
diff --git a/Engine/Source/Runtime/RHI/Public/RHIValidationContext.h b/Engine/Source/Runtime/RHI/Public/RHIValidationContext.h
--- a/Engine/Source/Runtime/RHI/Public/RHIValidationContext.h
+++ b/Engine/Source/Runtime/RHI/Public/RHIValidationContext.h
@@ -68,6 +68,9 @@ inline void ValidateShaderParameters(FRHIShader* RHIShader, RHIValidation::FTrac
 				StaticUniformBuffers.ValidateSetShaderUniformBuffer(UniformBuffer);
 			}
 			break;
+		case FRHIShaderParameterResource::EType::ResourceCollection:
+			// TODO: christopher.waters - ResourceCollection validation
+			break;
 		default:
 			checkf(false, TEXT("Unhandled resource type?"));
 			break;
diff --git a/Engine/Source/Runtime/RHICore/Internal/RHICoreResourceCollection.h b/Engine/Source/Runtime/RHICore/Internal/RHICoreResourceCollection.h
new file mode 100644
--- /dev/null
+++ b/Engine/Source/Runtime/RHICore/Internal/RHICoreResourceCollection.h
@@ -0,0 +1,146 @@
+// Copyright Epic Games, Inc. All Rights Reserved.
+
+#pragma once
+
+#include "RHIResources.h"
+#include "RHIResourceCollection.h"
+#include "RHICommandList.h"
+#include "DynamicRHI.h"
+#include "Containers/ResourceArray.h"
+
+#if PLATFORM_SUPPORTS_BINDLESS_RENDERING
+
+namespace UE::RHICore
+{
+	template<typename TType>
+	inline size_t CalculateResourceCollectionMemorySize(TConstArrayView<TType> InValues)
+	{
+		return (1 + InValues.Num()) * sizeof(uint32);
+	}
+
+	inline FRHIDescriptorHandle GetHandleForResourceCollectionValue(const FRHIResourceCollectionMember& Member)
+	{
+		switch (Member.Type)
+		{
+		case FRHIResourceCollectionMember::EType::Texture:
+			return static_cast<const FRHITexture*>(Member.Resource)->GetDefaultBindlessHandle();
+		case FRHIResourceCollectionMember::EType::TextureReference:
+			return static_cast<const FRHITextureReference*>(Member.Resource)->GetBindlessHandle();
+		case FRHIResourceCollectionMember::EType::ShaderResourceView:
+			return static_cast<const FRHIShaderResourceView*>(Member.Resource)->GetBindlessHandle();
+		case FRHIResourceCollectionMember::EType::UnorderedAccessView:
+			return static_cast<const FRHIUnorderedAccessView*>(Member.Resource)->GetBindlessHandle();
+		}
+
+		return FRHIDescriptorHandle();
+	}
+
+	inline FRHIDescriptorHandle GetHandleForResourceCollectionValue(const FRHIDescriptorHandle& Handle)
+	{
+		return Handle;
+	}
+
+	template<typename TType>
+	inline void FillResourceCollectionMemory(TArray<uint32>& Destination, TConstArrayView<TType> InValues)
+	{
+		Destination.SetNumUninitialized(1 + InValues.Num());
+
+		int32 WriteIndex = 0;
+
+		Destination[WriteIndex] = static_cast<uint32>(InValues.Num());
+		++WriteIndex;
+
+		for (const TType& Value : InValues)
+		{
+			const FRHIDescriptorHandle Handle = GetHandleForResourceCollectionValue(Value);
+			check(Handle.IsValid());
+
+			const uint32 BindlessIndex = Handle.IsValid() ? Handle.GetIndex() : 0;
+
+			Destination[WriteIndex] = BindlessIndex;
+			++WriteIndex;
+		}
+	}
+
+	template<typename TType>
+	inline TArray<uint32> CreateResourceCollectionArray(TConstArrayView<TType> InValues)
+	{
+		TArray<uint32> Result;
+		FillResourceCollectionMemory(Result, InValues);
+		return Result;
+	}
+
+	struct FResourceCollectionUpload : public FResourceArrayInterface
+	{
+		FResourceCollectionUpload(TConstArrayView<FRHIResourceCollectionMember> InMembers)
+			: Memory(CreateResourceCollectionArray(InMembers))
+		{
+		}
+		virtual const void* GetResourceData() const final
+		{
+			return Memory.GetData();
+		}
+		virtual uint32 GetResourceDataSize() const final
+		{
+			return Memory.Num() * Memory.GetTypeSize();
+		}
+		virtual void Discard() final
+		{
+			Memory.Reset();
+		}
+		virtual bool IsStatic() const final
+		{
+			return true;
+		}
+		virtual bool GetAllowCPUAccess() const final
+		{
+			return false;
+		}
+		virtual void SetAllowCPUAccess(bool bInNeedsCPUAccess) final
+		{
+		}
+
+		TArray<uint32> Memory;
+	};
+
+	class FGenericResourceCollection : public FRHIResourceCollection
+	{
+	public:
+		FGenericResourceCollection(FRHICommandListBase& RHICmdList, TConstArrayView<FRHIResourceCollectionMember> InMembers)
+			: FRHIResourceCollection(InMembers)
+		{
+			FResourceCollectionUpload UploadData(InMembers);
+
+			const size_t BufferSize = UE::RHICore::CalculateResourceCollectionMemorySize(InMembers);
+			FRHIResourceCreateInfo CreateInfo(TEXT("ResourceCollection"), &UploadData);
+			Buffer = RHICmdList.CreateBuffer(BufferSize, EBufferUsageFlags::Static | EBufferUsageFlags::ByteAddressBuffer, sizeof(uint32), ERHIAccess::SRVMask, CreateInfo);
+
+			FRHIViewDesc::FBufferSRV::FInitializer ViewDesc = FRHIViewDesc::CreateBufferSRV().SetType(FRHIViewDesc::EBufferType::Raw);
+			ShaderResourceView = RHICmdList.CreateShaderResourceView(Buffer, ViewDesc);
+		}
+
+		~FGenericResourceCollection() = default;
+
+		// FRHIResourceCollection
+		virtual FRHIDescriptorHandle GetBindlessHandle() const final
+		{
+			return ShaderResourceView->GetBindlessHandle();
+		}
+		//~FRHIResourceCollection
+
+		FRHIShaderResourceView* GetShaderResourceView() const
+		{
+			return ShaderResourceView;
+		}
+
+		TRefCountPtr<FRHIBuffer> Buffer;
+		TRefCountPtr<FRHIShaderResourceView> ShaderResourceView;
+	};
+
+	inline FRHIResourceCollectionRef CreateGenericResourceCollection(FRHICommandListBase& RHICmdList, TConstArrayView<FRHIResourceCollectionMember> InMembers)
+	{
+		return new FGenericResourceCollection(RHICmdList, InMembers);
+	}
+}
+
+#endif // PLATFORM_SUPPORTS_BINDLESS_RENDERING
diff --git a/Engine/Source/Runtime/RHICore/Internal/RHIShaderParametersShared.h b/Engine/Source/Runtime/RHICore/Internal/RHIShaderParametersShared.h
--- a/Engine/Source/Runtime/RHICore/Internal/RHIShaderParametersShared.h
+++ b/Engine/Source/Runtime/RHICore/Internal/RHIShaderParametersShared.h
@@ -16,6 +16,7 @@ namespace UE::RHICore
 			case FRHIShaderParameterResource::EType::ResourceView:        return static_cast<FRHIShaderResourceView*>(Resource)->GetBindlessHandle();
 			case FRHIShaderParameterResource::EType::UnorderedAccessView: return static_cast<FRHIUnorderedAccessView*>(Resource)->GetBindlessHandle();
 			case FRHIShaderParameterResource::EType::Sampler:             return static_cast<FRHISamplerState*>(Resource)->GetBindlessHandle();
+			case FRHIShaderParameterResource::EType::ResourceCollection:  return static_cast<FRHIResourceCollection*>(Resource)->GetBindlessHandle();
 			default:
 				checkf(false, TEXT("Unhandled resource type?"));
 				break;
diff --git a/Engine/Source/Runtime/RHICore/Internal/RHIUniformBufferDataShared.h b/Engine/Source/Runtime/RHICore/Internal/RHIUniformBufferDataShared.h
--- a/Engine/Source/Runtime/RHICore/Internal/RHIUniformBufferDataShared.h
+++ b/Engine/Source/Runtime/RHICore/Internal/RHIUniformBufferDataShared.h
@@ -42,6 +42,12 @@ namespace UE::RHICore
 			return SamplerState ? SamplerState->GetBindlessHandle() : FRHIDescriptorHandle();
 		}
 		break;
+		case UBMT_RESOURCE_COLLECTION:
+		{
+			FRHIResourceCollection* ResourceCollection = Reader.Read<FRHIResourceCollection*>(Resource);
+			return ResourceCollection ? ResourceCollection->GetBindlessHandle() : FRHIDescriptorHandle();
+		}
+		break;
 		case UBMT_RDG_TEXTURE:
 		{
 			FRDGTexture* RDGTexture = Reader.Read<FRDGTexture*>(Resource);
@@ -74,16 +80,9 @@ namespace UE::RHICore
 			return UnorderedAccessView ? UnorderedAccessView->GetBindlessHandle() : FRHIDescriptorHandle();
 		}
 		break;
-		case UBMT_NESTED_STRUCT:
-		case UBMT_INCLUDED_STRUCT:
-		case UBMT_REFERENCED_STRUCT:
-		{
-			// Do nothing?
-		}
-		break;
 
 		default:
-			//checkf(false, TEXT("Unhandled resource type?"));
+			// Do nothing.
 			break;
 		}
 		return FRHIDescriptorHandle();
diff --git a/Engine/Source/Runtime/RHICore/Public/RHICoreShader.h b/Engine/Source/Runtime/RHICore/Public/RHICoreShader.h
--- a/Engine/Source/Runtime/RHICore/Public/RHICoreShader.h
+++ b/Engine/Source/Runtime/RHICore/Public/RHICoreShader.h
@@ -97,6 +97,7 @@ template<> struct TResourceTypeStr<FRHISamplerState       > { static constexpr T
 template<> struct TResourceTypeStr<FRHITexture            > { static constexpr TCHAR String[] = TEXT("Texture"); };
 template<> struct TResourceTypeStr<FRHIShaderResourceView > { static constexpr TCHAR String[] = TEXT("Shader Resource View"); };
 template<> struct TResourceTypeStr<FRHIUnorderedAccessView> { static constexpr TCHAR String[] = TEXT("Unordered Access View"); };
+template<> struct TResourceTypeStr<FRHIResourceCollection>  { static constexpr TCHAR String[] = TEXT("Resource Collection"); };
 
 template <typename TResourceType, typename TCallback>
 inline void EnumerateUniformBufferResources(FRHIUniformBuffer* RESTRICT Buffer, int32 BufferIndex, const uint32* RESTRICT ResourceMap, TCallback&& Callback)
@@ -240,6 +241,17 @@ void SetResourcesFromTables(TBinder&& Binder, FRHIShader const& Shader, TBitMask
 				Binder.SetSRV(SRV, Index);
 			});
 
+#if PLATFORM_SUPPORTS_BINDLESS_RENDERING
+		EnumerateUniformBufferResources<FRHIResourceCollection>(Buffer, BufferIndex, SRT.ResourceCollectionMap.GetData(),
+			[&](FRHIResourceCollection* ResourceCollection, uint8 Index)
+			{
+#if ENABLE_RHI_VALIDATION
+				// todo: christopher.waters - ResourceCollection validation
+#endif
+				Binder.SetResourceCollection(ResourceCollection, Index);
+			});
+#endif // PLATFORM_SUPPORTS_BINDLESS_RENDERING
+
 		// Samplers
 		EnumerateUniformBufferResources<FRHISamplerState>(Buffer, BufferIndex, SRT.SamplerMap.GetData(),
 			[&](FRHISamplerState* Sampler, uint8 Index)
diff --git a/Engine/Source/Runtime/RenderCore/Private/ShaderParameterMetadata.cpp b/Engine/Source/Runtime/RenderCore/Private/ShaderParameterMetadata.cpp
--- a/Engine/Source/Runtime/RenderCore/Private/ShaderParameterMetadata.cpp
+++ b/Engine/Source/Runtime/RenderCore/Private/ShaderParameterMetadata.cpp
@@ -213,6 +213,8 @@ const TCHAR* const kShaderParameterMacroNames[] = {
 
 	// Structure dedicated to setup render targets for a rasterizer pass.
 	TEXT("RENDER_TARGET_BINDING_SLOTS"), // UBMT_RENDER_TARGET_BINDING_SLOTS,
+
+	TEXT("RESOURCE_COLLECTION"), // UBMT_RESOURCE_COLLECTION,
 };
 
 static_assert(UE_ARRAY_COUNT(kShaderParameterMacroNames) == int32(EUniformBufferBaseType_Num), "Shader parameter enum does not match name macro name array.");
@@ -481,10 +483,7 @@ void FShaderParametersMetadata::FMember::HashLayout(FMemoryHasherBlake3& Hasher)
 	Hasher.Serialize(const_cast<TCHAR*>(Name), FCString::Strlen(Name));
 	Hasher << NumElements;
 
-	const bool bIsRHIResource = (
-		BaseType == UBMT_TEXTURE ||
-		BaseType == UBMT_SRV ||
-		BaseType == UBMT_SAMPLER);
+	const bool bIsRHIResource = IsShaderParameterTypeRHIResource(BaseType);
 	const bool bIsRDGResource = IsRDGResourceReferenceShaderParameterType(BaseType);
 
 	if (BaseType == UBMT_INT32 ||
@@ -577,10 +576,7 @@ void FShaderParametersMetadata::InitializeLayout(FRHIUniformBufferLayoutInitiali
 		const TCHAR* ShaderType = CurrentMember.GetShaderType();
 
 		const bool bIsArray = ArraySize > 0;
-		const bool bIsRHIResource = (
-			BaseType == UBMT_TEXTURE ||
-			BaseType == UBMT_SRV ||
-			BaseType == UBMT_SAMPLER);
+		const bool bIsRHIResource = IsShaderParameterTypeRHIResource(BaseType);
 		const bool bIsRDGResource = IsRDGResourceReferenceShaderParameterType(BaseType);
 		const bool bIsVariableNativeType = CurrentMember.IsVariableNativeType();
 
@@ -683,7 +679,8 @@ void FShaderParametersMetadata::InitializeLayout(FRHIUniformBufferLayoutInitiali
 				{
 					bIsValidBindingType = (
 						BaseType == UBMT_SRV ||
-						BaseType == UBMT_RDG_BUFFER_SRV);
+						BaseType == UBMT_RDG_BUFFER_SRV ||
+						BaseType == UBMT_RESOURCE_COLLECTION);
 				}
 				else if (BindingType == EShaderCodeResourceBindingType::RaytracingAccelerationStructure)
 				{
@@ -826,10 +823,7 @@ void FShaderParametersMetadata::InitializeLayout(FRHIUniformBufferLayoutInitiali
 			MemberHash = HashCombine(MemberHash, FCrc::Strihash_DEPRECATED(CurrentMember.GetName()));
 			MemberHash = HashCombine(MemberHash, GetTypeHash(int32(CurrentMember.GetNumElements())));
 
-			const bool bIsRHIResource = (
-				BaseType == UBMT_TEXTURE ||
-				BaseType == UBMT_SRV ||
-				BaseType == UBMT_SAMPLER);
+			const bool bIsRHIResource = IsShaderParameterTypeRHIResource(BaseType);
 			const bool bIsRDGResource = IsRDGResourceReferenceShaderParameterType(BaseType);
 
 			if (BaseType == UBMT_INT32 ||
@@ -1089,9 +1083,7 @@ void FShaderParametersMetadata::FindMemberFromOffset(uint16 MemberOffset, const
 			}
 		}
 		else if (NumElements > 0 && (
-			BaseType == UBMT_TEXTURE ||
-			BaseType == UBMT_SRV ||
-			BaseType == UBMT_SAMPLER ||
+			IsShaderParameterTypeRHIResource(BaseType) ||
 			IsRDGResourceReferenceShaderParameterType(BaseType)))
 		{
 			uint16 ArrayStartOffset = Member.GetOffset();
diff --git a/Engine/Source/Runtime/RenderCore/Private/ShaderParameterParser.cpp b/Engine/Source/Runtime/RenderCore/Private/ShaderParameterParser.cpp
--- a/Engine/Source/Runtime/RenderCore/Private/ShaderParameterParser.cpp
+++ b/Engine/Source/Runtime/RenderCore/Private/ShaderParameterParser.cpp
@@ -75,14 +75,7 @@ static void IterateShaderParameterMembersInternal(
 		}
 		else
 		{
-			const bool bParametersAreExpanded =
-				NumElements > 0 &&
-				(BaseType == UBMT_TEXTURE ||
-					BaseType == UBMT_SRV ||
-					BaseType == UBMT_UAV ||
-					BaseType == UBMT_SAMPLER ||
-					IsRDGResourceReferenceShaderParameterType(BaseType));
-
+			const bool bParametersAreExpanded = NumElements > 0 && (IsShaderParameterTypeRHIResource(BaseType) || IsRDGResourceReferenceShaderParameterType(BaseType));
 			if (bParametersAreExpanded)
 			{
 				const uint16 ElementSize = SHADER_PARAMETER_POINTER_ALIGNMENT;
diff --git a/Engine/Source/Runtime/RenderCore/Public/ShaderCompilerCore.h b/Engine/Source/Runtime/RenderCore/Public/ShaderCompilerCore.h
--- a/Engine/Source/Runtime/RenderCore/Public/ShaderCompilerCore.h
+++ b/Engine/Source/Runtime/RenderCore/Public/ShaderCompilerCore.h
@@ -190,6 +190,9 @@ struct FShaderCompilerResourceTable
 	/** Mapping of bound UAVs to their location in resource tables. */
 	TArray<uint32> UnorderedAccessViewMap;
 
+	/** Mapping of bound respource collections to their location in resource tables. */
+	TArray<uint32> ResourceCollectionMap;
+
 	/** Hash of the layouts of resource tables at compile time, used for runtime validation. */
 	TArray<uint32> ResourceTableLayoutHashes;
 
diff --git a/Engine/Source/Runtime/VulkanRHI/Private/VulkanCommands.cpp b/Engine/Source/Runtime/VulkanRHI/Private/VulkanCommands.cpp
--- a/Engine/Source/Runtime/VulkanRHI/Private/VulkanCommands.cpp
+++ b/Engine/Source/Runtime/VulkanRHI/Private/VulkanCommands.cpp
@@ -13,6 +13,7 @@
 #include "GlobalRenderResources.h"
 #include "RHIShaderParametersShared.h"
 #include "RHIUtilities.h"
+#include "VulkanResourceCollection.h"
 #include "RHICoreShader.h"
 
 static TAutoConsoleVariable<int32> GCVarSubmitOnDispatch(
@@ -162,6 +163,14 @@ struct FVulkanResourceBinder
 	{
 		PendingState->SetSamplerStateForUBResource(GlobalRemappingInfo[Index].NewDescriptorSet, GlobalRemappingInfo[Index].NewBindingIndex, ResourceCast(Sampler));
 	}
+
+#if PLATFORM_SUPPORTS_BINDLESS_RENDERING
+	void SetResourceCollection(FRHIResourceCollection* ResourceCollection, uint32 Index)
+	{
+		FVulkanResourceCollection* VulkanResourceCollection = ResourceCast(ResourceCollection);
+		SetSRV(VulkanResourceCollection->GetShaderResourceView(), Index);
+	}
+#endif
 };
 
 template <class ShaderType> 
diff --git a/Engine/Source/Runtime/VulkanRHI/Private/VulkanRHI.cpp b/Engine/Source/Runtime/VulkanRHI/Private/VulkanRHI.cpp
--- a/Engine/Source/Runtime/VulkanRHI/Private/VulkanRHI.cpp
+++ b/Engine/Source/Runtime/VulkanRHI/Private/VulkanRHI.cpp
@@ -20,6 +20,7 @@
 #include "VulkanLLM.h"
 #include "Misc/EngineVersion.h"
 #include "GlobalShader.h"
+#include "VulkanResourceCollection.h"
 #include "RHIValidation.h"
 #include "RHIUtilities.h"
 #include "IHeadMountedDisplayModule.h"
@@ -2291,4 +2292,12 @@ void FVulkanDynamicRHI::RHIReplaceResources(FRHICommandListBase& RHICmdList, TAr
 	RHICmdList.RHIThreadFence(true);
 }
 
+#if PLATFORM_SUPPORTS_BINDLESS_RENDERING
+FRHIResourceCollectionRef FVulkanDynamicRHI::RHICreateResourceCollection(FRHICommandListBase& RHICmdList, TConstArrayView<FRHIResourceCollectionMember> InMembers)
+{
+	return new FVulkanResourceCollection(RHICmdList, InMembers);
+}
+#endif
+
+
 #undef LOCTEXT_NAMESPACE
diff --git a/Engine/Source/Runtime/VulkanRHI/Private/VulkanResourceCollection.h b/Engine/Source/Runtime/VulkanRHI/Private/VulkanResourceCollection.h
new file mode 100644
--- /dev/null
+++ b/Engine/Source/Runtime/VulkanRHI/Private/VulkanResourceCollection.h
@@ -0,0 +1,16 @@
+// Copyright Epic Games, Inc. All Rights Reserved.
+
+#pragma once
+
+#include "RHICoreResourceCollection.h"
+#include "VulkanResources.h"
+
+#if PLATFORM_SUPPORTS_BINDLESS_RENDERING
+using FVulkanResourceCollection = UE::RHICore::FGenericResourceCollection;
+
+template<>
+struct TVulkanResourceTraits<FRHIResourceCollection>
+{
+	using TConcreteType = FVulkanResourceCollection;
+};
+#endif
\ No newline at end of file
diff --git a/Engine/Source/Runtime/VulkanRHI/Public/VulkanDynamicRHI.h b/Engine/Source/Runtime/VulkanRHI/Public/VulkanDynamicRHI.h
--- a/Engine/Source/Runtime/VulkanRHI/Public/VulkanDynamicRHI.h
+++ b/Engine/Source/Runtime/VulkanRHI/Public/VulkanDynamicRHI.h
@@ -219,6 +219,10 @@ public:
 	virtual FShaderResourceViewRHIRef  RHICreateShaderResourceView (class FRHICommandListBase& RHICmdList, FRHIViewableResource* Resource, FRHIViewDesc const& ViewDesc) final override;
 	virtual FUnorderedAccessViewRHIRef RHICreateUnorderedAccessView(class FRHICommandListBase& RHICmdList, FRHIViewableResource* Resource, FRHIViewDesc const& ViewDesc) final override;
 
+#if PLATFORM_SUPPORTS_BINDLESS_RENDERING
+	virtual FRHIResourceCollectionRef RHICreateResourceCollection(FRHICommandListBase& RHICmdList, TConstArrayView<FRHIResourceCollectionMember> InMembers) final;
+#endif
+
 #if VULKAN_RHI_RAYTRACING
 	virtual FRayTracingAccelerationStructureSize RHICalcRayTracingSceneSize(uint32 MaxInstances, ERayTracingAccelerationStructureFlags Flags) final override;
 	virtual FRayTracingAccelerationStructureSize RHICalcRayTracingGeometrySize(const FRayTracingGeometryInitializer& Initializer) final override;
diff --git a/Engine/Source/Runtime/Windows/D3D11RHI/Private/D3D11Commands.cpp b/Engine/Source/Runtime/Windows/D3D11RHI/Private/D3D11Commands.cpp
--- a/Engine/Source/Runtime/Windows/D3D11RHI/Private/D3D11Commands.cpp
+++ b/Engine/Source/Runtime/Windows/D3D11RHI/Private/D3D11Commands.cpp
@@ -392,6 +392,13 @@ struct FD3D11ResourceBinder
 	{
 		RHI.GetStateCache().SetSamplerState<ShaderFrequency>(FD3D11DynamicRHI::ResourceCast(Sampler)->Resource, Index);
 	}
+
+#if PLATFORM_SUPPORTS_BINDLESS_RENDERING
+	void SetResourceCollection(FRHIResourceCollection* ResourceCollection, uint32 Index)
+	{
+		checkNoEntry();
+	}
+#endif
 };
 
 template<EShaderFrequency ShaderFrequency>
