Fixing build error caused by a function declaration not being wrapped in a platform check.

diff --git a/Engine/Source/Runtime/D3D12RHI/Private/D3D12RHIPrivate.h b/Engine/Source/Runtime/D3D12RHI/Private/D3D12RHIPrivate.h
--- a/Engine/Source/Runtime/D3D12RHI/Private/D3D12RHIPrivate.h
+++ b/Engine/Source/Runtime/D3D12RHI/Private/D3D12RHIPrivate.h
@@ -401,7 +401,9 @@ public:
 	// SRV / UAV creation functions
 	virtual FShaderResourceViewRHIRef  RHICreateShaderResourceView (class FRHICommandListBase& RHICmdList, FRHIViewableResource* Resource, FRHIViewDesc const& ViewDesc) override;
 	virtual FUnorderedAccessViewRHIRef RHICreateUnorderedAccessView(class FRHICommandListBase& RHICmdList, FRHIViewableResource* Resource, FRHIViewDesc const& ViewDesc) override;
+#if PLATFORM_SUPPORTS_BINDLESS_RENDERING
 	virtual FRHIResourceCollectionRef RHICreateResourceCollection(FRHICommandListBase& RHICmdList, TConstArrayView<FRHIResourceCollectionMember> InMembers) override;
+#endif
 
 	// ID3D12DynamicRHI interface.
 	virtual TArray<FD3D12MinimalAdapterDesc> RHIGetAdapterDescs() const final override;
