Big D3D12Viewport cleanup
- Moved all FD3D12Viewport members to be together.
- Merged the three FD3D12Viewport viewport texture arrays and the mGPU index array into a single array.
- CalculateSwapChainDepth always took in the same count that was given on viewport construction so it was renamed to InitializeBackBufferArrays and NumBackBuffers is now constant per platform.
- Merged D3D11 and D3D12 GetRenderTargetFormat implementations, which were the exact same, into UE::DXGIUtilities::GetSwapChainFormat.
- Renaming D3D12_USE_DUMMY_BACKBUFFER to D3D12RHI_USE_DUMMY_BACKBUFFER since it's not a d3d12.h define.
- FD3D12BackBufferReferenceTexture2D is now only used when D3D12RHI_USE_DUMMY_BACKBUFFER is defined.
- Adding D3D12RHI_USE_SDR_BACKBUFFER to prevent uses of SDR texture pointers when we don't need them.
- Adding D3D12RHI_SUPPORTS_UAV_BACKBUFFER to replace uses of RHISupportsSwapchainUAVs(GMaxRHIShaderPlatform)

diff --git a/Engine/Source/Runtime/D3D12RHI/Private/D3D12RHIPrivate.h b/Engine/Source/Runtime/D3D12RHI/Private/D3D12RHIPrivate.h
--- a/Engine/Source/Runtime/D3D12RHI/Private/D3D12RHIPrivate.h
+++ b/Engine/Source/Runtime/D3D12RHI/Private/D3D12RHIPrivate.h
@@ -362,7 +362,9 @@ public:
 	virtual bool RHIGetRenderQueryResult(FRHIRenderQuery* RenderQuery, uint64& OutResult, bool bWait, uint32 GPUIndex = INDEX_NONE) final override;
 	virtual uint32 RHIGetViewportNextPresentGPUIndex(FRHIViewport* Viewport) final override;
 	virtual FTextureRHIRef RHIGetViewportBackBuffer(FRHIViewport* Viewport) final override;
+#if D3D12RHI_SUPPORTS_UAV_BACKBUFFER
 	virtual FUnorderedAccessViewRHIRef RHIGetViewportBackBufferUAV(FRHIViewport* Viewport) final override;
+#endif
 	virtual void RHIAliasTextureResources(FTextureRHIRef& DestTexture, FTextureRHIRef& SrcTexture) final override;
 	virtual FTextureRHIRef RHICreateAliasedTexture(FTextureRHIRef& SourceTexture) final override;
 	virtual void RHIGetDisplaysInformation(FDisplayInformationArray& OutDisplayInformation) final override;
diff --git a/Engine/Source/Runtime/D3D12RHI/Private/D3D12RenderTarget.cpp b/Engine/Source/Runtime/D3D12RHI/Private/D3D12RenderTarget.cpp
--- a/Engine/Source/Runtime/D3D12RHI/Private/D3D12RenderTarget.cpp
+++ b/Engine/Source/Runtime/D3D12RHI/Private/D3D12RenderTarget.cpp
@@ -577,7 +577,7 @@ void FD3D12DynamicRHI::RHIReadSurfaceData(FRHITexture* InRHITexture, FIntRect In
 
 	// Could be back buffer reference texture, so get the correct D3D12 texture here
 	// We know already that it's a FD3D12Texture2D so cast is safe
-#if D3D12_USE_DUMMY_BACKBUFFER
+#if D3D12RHI_USE_DUMMY_BACKBUFFER
 	if (EnumHasAnyFlags(InRHITexture->GetFlags(), TexCreate_Presentable))
 	{
 		FD3D12BackBufferReferenceTexture2D* BufferBufferReferenceTexture = (FD3D12BackBufferReferenceTexture2D*)InRHITexture;
diff --git a/Engine/Source/Runtime/D3D12RHI/Private/D3D12Texture.cpp b/Engine/Source/Runtime/D3D12RHI/Private/D3D12Texture.cpp
--- a/Engine/Source/Runtime/D3D12RHI/Private/D3D12Texture.cpp
+++ b/Engine/Source/Runtime/D3D12RHI/Private/D3D12Texture.cpp
@@ -2693,6 +2693,7 @@ void FD3D12CommandContext::RHICopyTexture(FRHITexture* SourceTextureRHI, FRHITex
 	ConditionalSplitCommandList();
 }
 
+#if D3D12RHI_USE_DUMMY_BACKBUFFER
 
 ///////////////////////////////////////////////////////////////////////////////////////////
 // FD3D12BackBufferReferenceTexture2D functions
@@ -2700,10 +2701,12 @@ void FD3D12CommandContext::RHICopyTexture(FRHITexture* SourceTextureRHI, FRHITex
 
 FRHITexture* FD3D12BackBufferReferenceTexture2D::GetBackBufferTexture() const
 {
-	return bIsSDR ? Viewport->GetSDRBackBuffer_RHIThread() : Viewport->GetBackBuffer_RHIThread();
+	return Viewport->GetBackBuffer_RHIThread();
 }
 
 FRHIDescriptorHandle FD3D12BackBufferReferenceTexture2D::GetDefaultBindlessHandle() const
 {
 	return GetBackBufferTexture()->GetDefaultBindlessHandle();
 }
+
+#endif // D3D12RHI_USE_DUMMY_BACKBUFFER
diff --git a/Engine/Source/Runtime/D3D12RHI/Private/D3D12Texture.h b/Engine/Source/Runtime/D3D12RHI/Private/D3D12Texture.h
--- a/Engine/Source/Runtime/D3D12RHI/Private/D3D12Texture.h
+++ b/Engine/Source/Runtime/D3D12RHI/Private/D3D12Texture.h
@@ -192,26 +192,25 @@ struct TD3D12ResourceTraits<FRHITexture>
 
 class FD3D12Viewport;
 
+#if D3D12RHI_USE_DUMMY_BACKBUFFER
 class FD3D12BackBufferReferenceTexture2D : public FD3D12Texture
 {
 public:
-	FD3D12BackBufferReferenceTexture2D(const FRHITextureCreateDesc& InDesc, FD3D12Viewport* InViewPort, bool bInIsSDR, FD3D12Device* InDevice)
+	FD3D12BackBufferReferenceTexture2D(const FRHITextureCreateDesc& InDesc, FD3D12Viewport* InViewPort, FD3D12Device* InDevice)
 		: FD3D12Texture(InDesc, InDevice)
 		, Viewport(InViewPort)
-		, bIsSDR(bInIsSDR)
 	{
 	}
 
 	FD3D12Viewport* GetViewPort() const { return Viewport; }
-	bool IsSDR() const { return bIsSDR; }
 
 	FRHITexture* GetBackBufferTexture() const;
 	virtual FRHIDescriptorHandle GetDefaultBindlessHandle() const override;
 
 private:
-	FD3D12Viewport* Viewport = nullptr;
-	bool bIsSDR = false;
+	FD3D12Viewport* const Viewport;
 };
+#endif
 
 /** Given a pointer to a RHI texture that was created by the D3D12 RHI, returns a pointer to the FD3D12Texture it encapsulates. */
 FORCEINLINE FD3D12Texture* GetD3D12TextureFromRHITexture(FRHITexture* Texture)
@@ -223,7 +222,7 @@ FORCEINLINE FD3D12Texture* GetD3D12TextureFromRHITexture(FRHITexture* Texture)
 	
 	// If it's the dummy backbuffer then swap with actual current RHI backbuffer right now
 	FRHITexture* RHITexture = Texture;
-#if D3D12_USE_DUMMY_BACKBUFFER
+#if D3D12RHI_USE_DUMMY_BACKBUFFER
 	if (RHITexture && EnumHasAnyFlags(RHITexture->GetFlags(), TexCreate_Presentable))
 	{
 		FD3D12BackBufferReferenceTexture2D* BufferBufferReferenceTexture = (FD3D12BackBufferReferenceTexture2D*)RHITexture;
diff --git a/Engine/Source/Runtime/D3D12RHI/Private/D3D12Viewport.cpp b/Engine/Source/Runtime/D3D12RHI/Private/D3D12Viewport.cpp
--- a/Engine/Source/Runtime/D3D12RHI/Private/D3D12Viewport.cpp
+++ b/Engine/Source/Runtime/D3D12RHI/Private/D3D12Viewport.cpp
@@ -164,7 +164,7 @@ FD3D12Texture* GetSwapChainSurface(FD3D12Device* Parent, EPixelFormat PixelForma
 		TextureDesc.Height = SizeY;
 		TextureDesc.DepthOrArraySize = 1;
 		TextureDesc.MipLevels = 1;
-		TextureDesc.Format = FD3D12Viewport::GetRenderTargetFormat(PixelFormat);
+		TextureDesc.Format = UE::DXGIUtilities::GetSwapChainFormat(PixelFormat);
 		TextureDesc.SampleDesc.Count = 1;
 		TextureDesc.SampleDesc.Quality = 0;
 		TextureDesc.Layout = D3D12_TEXTURE_LAYOUT_UNKNOWN;
@@ -179,10 +179,9 @@ FD3D12Texture* GetSwapChainSurface(FD3D12Device* Parent, EPixelFormat PixelForma
 	FString Name = FString::Printf(TEXT("BackBuffer%d"), BackBufferIndex);
 
 	ETextureCreateFlags SwapchainTextureCreateFlags = ETextureCreateFlags::RenderTargetable;
-	if (RHISupportsSwapchainUAVs(GetFeatureLevelShaderPlatform(GMaxRHIFeatureLevel)))
-	{
-		SwapchainTextureCreateFlags |= ETextureCreateFlags::UAV;
-	}
+#if D3D12RHI_SUPPORTS_UAV_BACKBUFFER
+	SwapchainTextureCreateFlags |= ETextureCreateFlags::UAV;
+#endif
 
 	bool const bQuadBufferStereo = FD3D12DynamicRHI::GetD3DRHI()->IsQuadBufferStereoEnabled();
 
@@ -287,27 +286,6 @@ FD3D12Texture* GetSwapChainSurface(FD3D12Device* Parent, EPixelFormat PixelForma
 	return SwapChainTexture;
 }
 
-
-/**
-* Create the dummy back buffer textures - They don't have actual D3D resource but are used to always reference the current back buffer index on the RHI thread
-*/
-FD3D12Texture* FD3D12Viewport::CreateDummyBackBufferTextures(FD3D12Adapter* InAdapter, EPixelFormat InPixelFormat, uint32 InSizeX, uint32 InSizeY, bool bInIsSDR)
-{
-	FRHITextureCreateDesc CreateDesc =
-		FRHITextureCreateDesc::Create2D(TEXT("BackBufferReference"))
-		.SetExtent(FIntPoint(InSizeX, InSizeY))
-		.SetFormat(InPixelFormat)
-		.SetFlags(ETextureCreateFlags::RenderTargetable | ETextureCreateFlags::Presentable | ETextureCreateFlags::ResolveTargetable)
-		.SetInitialState(ERHIAccess::Present);
-
-	FD3D12Texture* Result = InAdapter->CreateLinkedObject<FD3D12Texture>(FRHIGPUMask::All(), [&](FD3D12Device* Device)
-	{
-		FD3D12Texture* NewTexture = new FD3D12BackBufferReferenceTexture2D(CreateDesc, this, bInIsSDR, Device);
-		return NewTexture;
-	});
-	return Result;
-}
-
 FD3D12Viewport::~FD3D12Viewport()
 {
 	check(IsInRHIThread() || IsInRenderingThread());
@@ -334,26 +312,19 @@ FD3D12Viewport::~FD3D12Viewport()
 	FinalDestroyInternal();
 }
 
-DXGI_MODE_DESC FD3D12Viewport::SetupDXGI_MODE_DESC() const
+#if D3D12_VIEWPORT_EXPOSES_SWAP_CHAIN
+inline DXGI_MODE_DESC SetupDXGI_MODE_DESC(uint32 SizeX, uint32 SizeY, EPixelFormat PixelFormat)
 {
-	DXGI_MODE_DESC Ret;
-
+	DXGI_MODE_DESC Ret{};
 	Ret.Width = SizeX;
 	Ret.Height = SizeY;
-	Ret.RefreshRate.Numerator = 0;	// illamas: use 0 to avoid a potential mismatch with hw
-	Ret.RefreshRate.Denominator = 0;	// illamas: ditto
-	Ret.Format = GetRenderTargetFormat(PixelFormat);
-	Ret.ScanlineOrdering = DXGI_MODE_SCANLINE_ORDER_UNSPECIFIED;
-	Ret.Scaling = DXGI_MODE_SCALING_UNSPECIFIED;
-
+	Ret.Format = UE::DXGIUtilities::GetSwapChainFormat(PixelFormat);
 	return Ret;
 }
+#endif
 
-void FD3D12Viewport::CalculateSwapChainDepth(int32 DefaultSwapChainDepth)
+void FD3D12Viewport::InitializeBackBufferArrays()
 {
-	FD3D12Adapter* Adapter = GetParentAdapter();
-
-	NumBackBuffers = DefaultSwapChainDepth;
 #if WITH_MGPU
 	// This is a temporary helper to visualize what each GPU is rendering. 
 	// Not specifying a value will cycle swap chain through all GPUs.
@@ -367,13 +338,35 @@ void FD3D12Viewport::CalculateSwapChainDepth(int32 DefaultSwapChainDepth)
 	}
 #endif // WITH_MGPU
 
-	BackBuffers.Empty();
-	BackBuffersUAV.Empty();
-	BackBuffers.AddZeroed(NumBackBuffers);
-	BackBuffersUAV.AddZeroed(NumBackBuffers);
+	for (FBackBufferData& Data : BackBuffers)
+	{
+		Data = FBackBufferData();
+	}
+}
+
+template<typename TResource, typename TSubresource = TResource>
+bool ReleaseBackBufferResource(TRefCountPtr<TResource>& Resource, const TCHAR* ErrorName, uint32 Index)
+{
+	const bool bValidReference = IsValidRef(Resource);
+	if (bValidReference)
+	{
+		// Tell the back buffer to delete immediately so that we can call resize.
+		if (Resource->GetRefCount() != 1)
+		{
+			UE_LOG(LogD3D12RHI, Log, TEXT("%s %d leaking with %d refs during Resize."), ErrorName, Index, Resource->GetRefCount());
+		}
+		check(Resource->GetRefCount() == 1);
 
-	SDRBackBuffers.Empty();
-	SDRBackBuffers.AddZeroed(NumBackBuffers);
+		for (TSubresource& SubResource : *Resource)
+		{
+			SubResource.GetResource()->DoNotDeferDelete();
+		}
+	}
+
+	Resource.SafeRelease();
+	check(Resource == nullptr);
+
+	return bValidReference;
 }
 
 void FD3D12Viewport::Resize(uint32 InSizeX, uint32 InSizeY, bool bInIsFullscreen, EPixelFormat PreferredPixelFormat)
@@ -404,28 +397,13 @@ void FD3D12Viewport::Resize(uint32 InSizeX, uint32 InSizeY, bool bInIsFullscreen
 		CustomPresent->OnBackBufferResize();
 	}
 
+#if D3D12RHI_SUPPORTS_UAV_BACKBUFFER
 	bool bWaitForBackBuffersUAVDelete = false;
+
 	// Release our backbuffer reference, as required by DXGI before calling ResizeBuffers.
-	for (uint32 i = 0; i < NumBackBuffers; ++i)
+	for (uint32 Index = 0; Index < NumBackBuffers; Index++)
 	{
-		if (IsValidRef(BackBuffersUAV[i]))
-		{
-			bWaitForBackBuffersUAVDelete = true;
-			// Tell the back buffer to delete immediately so that we can call resize.
-			if (BackBuffersUAV[i]->GetRefCount() != 1)
-			{
-				UE_LOG(LogD3D12RHI, Log, TEXT("Backbuffer %d leaking with %d refs during Resize."), i, BackBuffersUAV[i]->GetRefCount());
-			}
-			check(BackBuffersUAV[i]->GetRefCount() == 1);
-
-			for (FD3D12UnorderedAccessView& Uav : *BackBuffersUAV[i])
-			{
-				Uav.GetResource()->DoNotDeferDelete();
-			}
-		}
-
-		BackBuffersUAV[i].SafeRelease();
-		check(BackBuffersUAV[i] == nullptr);
+		bWaitForBackBuffersUAVDelete |= ReleaseBackBufferResource<FD3D12UnorderedAccessView_RHI, FD3D12UnorderedAccessView>(BackBuffers[Index].UAV, TEXT("BackBuffer UAV"), Index);
 	}
 
 	if (bWaitForBackBuffersUAVDelete)
@@ -434,43 +412,14 @@ void FD3D12Viewport::Resize(uint32 InSizeX, uint32 InSizeY, bool bInIsFullscreen
         // Calling FlushRenderingCommands is enough because it calls ImmediateFlush(EImmediateFlushType::FlushRHIThreadFlushResources) / ImmediateFlush(EImmediateFlushType::FlushRHIThread) internally
 		FlushRenderingCommands();
 	}
+#endif
 
-	for (uint32 i = 0; i < NumBackBuffers; ++i)
+	for (uint32 Index = 0; Index < NumBackBuffers; Index++)
 	{
-		if (IsValidRef(BackBuffers[i]))
-		{
-			// Tell the back buffer to delete immediately so that we can call resize.
-			if (BackBuffers[i]->GetRefCount() != 1)
-			{
-				UE_LOG(LogD3D12RHI, Log, TEXT("Backbuffer %d leaking with %d refs during Resize."), i, BackBuffers[i]->GetRefCount());
-			}
-			check(BackBuffers[i]->GetRefCount() == 1);
-
-			for (FD3D12Texture& Tex : *BackBuffers[i])
-			{
-				Tex.GetResource()->DoNotDeferDelete();
-			}
-		}
-
-		BackBuffers[i].SafeRelease();
-		check(BackBuffers[i] == nullptr);
-
-		if (IsValidRef(SDRBackBuffers[i]))
-		{
-			if (SDRBackBuffers[i]->GetRefCount() != 1)
-			{
-				UE_LOG(LogD3D12RHI, Log, TEXT("SDR Backbuffer %d leaking with %d refs during Resize."), i, SDRBackBuffers[i]->GetRefCount());
-			}
-			check(SDRBackBuffers[i]->GetRefCount() == 1);
-
-			for (FD3D12Texture& Tex : *SDRBackBuffers[i])
-			{
-				Tex.GetResource()->DoNotDeferDelete();
-			}
-		}
-
-		SDRBackBuffers[i].SafeRelease();
-		check(SDRBackBuffers[i] == nullptr);
+		ReleaseBackBufferResource(BackBuffers[Index].Texture,    TEXT("BackBuffer"),     Index);
+#if D3D12RHI_USE_SDR_BACKBUFFER
+		ReleaseBackBufferResource(BackBuffers[Index].TextureSDR, TEXT("SDR BackBuffer"), Index);
+#endif
 	}
 
 	ClearPresentQueue();
@@ -502,7 +451,7 @@ void FD3D12Viewport::Resize(uint32 InSizeX, uint32 InSizeY, bool bInIsFullscreen
 		{
 			if (bInIsFullscreen)
 			{
-				const DXGI_MODE_DESC BufferDesc = SetupDXGI_MODE_DESC();
+				const DXGI_MODE_DESC BufferDesc = SetupDXGI_MODE_DESC(SizeX, SizeY, PixelFormat);
 				if (FAILED(SwapChain1->ResizeTarget(&BufferDesc)))
 				{
 					ConditionalResetSwapChain(true);
@@ -532,7 +481,7 @@ void FD3D12Viewport::Resize(uint32 InSizeX, uint32 InSizeY, bool bInIsFullscreen
 				//	* InSizeX and InSizeY are given to us as the monitor resolution, so we resize the buffers to the correct resolution below (in ResizeInternal)
 				//	* however, the target still has the smaller size, because Slate doesn't know it has to resize the window too (as far as it's concerned, it's already the right size)
 				//	* therefore, we need to call ResizeTarget, which in windowed mode behaves like SetWindowPos.
-				const DXGI_MODE_DESC BufferDesc = SetupDXGI_MODE_DESC();
+				const DXGI_MODE_DESC BufferDesc = SetupDXGI_MODE_DESC(SizeX, SizeY, PixelFormat);
 				SwapChain1->ResizeTarget(&BufferDesc);
 			}
 #endif // D3D12_VIEWPORT_EXPOSES_SWAP_CHAIN
@@ -561,28 +510,6 @@ void FD3D12Viewport::Resize(uint32 InSizeX, uint32 InSizeY, bool bInIsFullscreen
 	}
 }
 
-FD3D12Texture* FD3D12Viewport::GetDummyBackBuffer_RenderThread(bool bInIsSDR) const
-{
-#if D3D12_USE_DUMMY_BACKBUFFER
-	return bInIsSDR ? SDRDummyBackBuffer_RenderThread : DummyBackBuffer_RenderThread;
-#else
-	checkNoEntry();
-	return nullptr;
-#endif
-}
-
-FD3D12UnorderedAccessView_RHI* FD3D12Viewport::GetBackBufferUAV_RenderThread() const
-{ 
-#if D3D12_USE_DUMMY_BACKBUFFER
-    // See FD3D12Viewport::PresentChecked: if we change fullscreen state (which is detected on RHI thread), we might end up with invalid backbuffer: the safe way is to rely on the dummybackbuffer instead, but 
-	// managing UAV is a bit tricky: disallow UAV on viewports right now
-    checkNoEntry();
-	return nullptr;
-#else
-	return BackBufferUAV_RenderThread;
-#endif
-}
-
 #if D3D12_VIEWPORT_EXPOSES_SWAP_CHAIN
 void* FD3D12Viewport::GetNativeSwapChain() const
 {
@@ -622,11 +549,8 @@ void FD3D12Viewport::AdvanceExpectedBackBufferIndex_RenderThread()
 		FScopeLock Lock(&ExpectedBackBufferIndexLock);
 #endif
 
-		ExpectedBackBufferIndex_RenderThread++;
-		ExpectedBackBufferIndex_RenderThread = ExpectedBackBufferIndex_RenderThread % NumBackBuffers;
+		SetBackBufferIndex_RenderThread(ExpectedBackBufferIndex_RenderThread + 1);
 
-		BackBuffer_RenderThread = BackBuffers[ExpectedBackBufferIndex_RenderThread];
-		BackBufferUAV_RenderThread = BackBuffersUAV[ExpectedBackBufferIndex_RenderThread];
 #if !UE_BUILD_SHIPPING
 		if (RHIConsoleVariables::LogViewportEvents)
 		{
@@ -662,7 +586,7 @@ bool FD3D12Viewport::PresentChecked(int32 SyncInterval)
 		// Present failed so current expected GPU index will not match anymore, so patch up expected back buffer index
 		// Warning: Present is skipped for this frame but could cause a black screen for the next frame as well
 		FScopeLock Lock(&ExpectedBackBufferIndexLock);
-		ExpectedBackBufferIndex_RenderThread = (ExpectedBackBufferIndex_RenderThread == 0) ? NumBackBuffers - 1 : ExpectedBackBufferIndex_RenderThread - 1;
+		SetBackBufferIndex_RenderThread((ExpectedBackBufferIndex_RenderThread == 0) ? NumBackBuffers - 1 : ExpectedBackBufferIndex_RenderThread - 1);
 #endif // WITH_MGPU
 		return false;
 	}
@@ -709,7 +633,6 @@ bool FD3D12Viewport::Present(bool bLockToVsync)
 
 		// Those are not necessarily the swap chain back buffer in case of multi-gpu
 		FD3D12Texture* DeviceBackBuffer = DefaultContext.RetrieveObject<FD3D12Texture, FRHITexture>(GetBackBuffer_RHIThread());
-		FD3D12Texture* DeviceSDRBackBuffer = DefaultContext.RetrieveObject<FD3D12Texture, FRHITexture>(GetSDRBackBuffer_RHIThread());
 
 		DefaultContext.TransitionResource(
 			DeviceBackBuffer->GetShaderResourceView()->GetResource(),
@@ -718,8 +641,11 @@ bool FD3D12Viewport::Present(bool bLockToVsync)
 			0
 		);
 
-		if (SDRBackBuffer_RHIThread != nullptr)
+#if D3D12RHI_USE_SDR_BACKBUFFER
+		if (CurrentBackBuffer_RHIThread->TextureSDR)
 		{
+			FD3D12Texture* DeviceSDRBackBuffer = DefaultContext.RetrieveObject<FD3D12Texture, FRHITexture>(GetSDRBackBuffer_RHIThread());
+
 			DefaultContext.TransitionResource(
 				DeviceSDRBackBuffer->GetShaderResourceView()->GetResource(),
 				D3D12_RESOURCE_STATE_TBD,
@@ -727,6 +653,7 @@ bool FD3D12Viewport::Present(bool bLockToVsync)
 				0
 			);
 		}
+#endif
 
 		DefaultContext.FlushResourceBarriers();
 	}
@@ -766,22 +693,19 @@ bool FD3D12Viewport::Present(bool bLockToVsync)
 #if DXGI_MAX_SWAPCHAIN_INTERFACE >= 3
 		if (bNativelyPresented && SwapChain3)
 		{
-			CurrentBackBufferIndex_RHIThread = SwapChain3->GetCurrentBackBufferIndex();
+			SetBackBufferIndex_RHIThread(SwapChain3->GetCurrentBackBufferIndex());
 		}
 		else
 #endif
 		{
-			CurrentBackBufferIndex_RHIThread++;
-			CurrentBackBufferIndex_RHIThread = CurrentBackBufferIndex_RHIThread % NumBackBuffers;
+			SetBackBufferIndex_RHIThread(CurrentBackBufferIndex_RHIThread + 1);
 		}
-		BackBuffer_RHIThread = BackBuffers[CurrentBackBufferIndex_RHIThread].GetReference();
-		SDRBackBuffer_RHIThread = SDRBackBuffers[CurrentBackBufferIndex_RHIThread].GetReference();
 
 #if !UE_BUILD_SHIPPING
 		if (RHIConsoleVariables::LogViewportEvents)
 		{
 			const FString& ThreadName = FThreadManager::GetThreadName(FPlatformTLS::GetCurrentThreadId());
-			UE_LOG(LogD3D12RHI, Log, TEXT("Thread %s: Incrementing RHIThread back buffer index of viewport: %#016llx to value: %u BackBuffer %#016llx"), ThreadName.GetCharArray().GetData(), this, CurrentBackBufferIndex_RHIThread, BackBuffer_RHIThread);
+			UE_LOG(LogD3D12RHI, Log, TEXT("Thread %s: Incrementing RHIThread back buffer index of viewport: %#016llx to value: %u BackBuffer %#016llx"), ThreadName.GetCharArray().GetData(), this, CurrentBackBufferIndex_RHIThread, CurrentBackBuffer_RHIThread->Texture.GetReference());
 		}
 #endif
 	}
@@ -1038,10 +962,3 @@ FTextureRHIRef FD3D12DynamicRHI::RHIGetViewportBackBuffer(FRHIViewport* Viewport
 
 	return SelectedBackBuffer;
 }
-
-FUnorderedAccessViewRHIRef FD3D12DynamicRHI::RHIGetViewportBackBufferUAV(FRHIViewport* ViewportRHI)
-{
-	check(IsInRenderingThread());
-	const FD3D12Viewport* const Viewport = FD3D12DynamicRHI::ResourceCast(ViewportRHI);
-	return Viewport->GetBackBufferUAV_RenderThread();
-}
diff --git a/Engine/Source/Runtime/D3D12RHI/Private/D3D12Viewport.h b/Engine/Source/Runtime/D3D12RHI/Private/D3D12Viewport.h
--- a/Engine/Source/Runtime/D3D12RHI/Private/D3D12Viewport.h
+++ b/Engine/Source/Runtime/D3D12RHI/Private/D3D12Viewport.h
@@ -11,6 +11,7 @@
 #include "MultiGPU.h"
 #include "RHIResources.h"
 #include "Templates/RefCounting.h"
+#include "DXGIUtilities.h"
 
 class FD3D12Texture;
 class FD3D12UnorderedAccessView_RHI;
@@ -64,7 +65,6 @@ private:
 };
 #endif //WITH_MGPU
 
-
 class FD3D12Viewport : public FRHIViewport, public FD3D12AdapterChild
 {
 public:
@@ -95,15 +95,57 @@ public:
 	bool Present(bool bLockToVsync);
 
 	// Accessors.
-	FIntPoint GetSizeXY() const { return FIntPoint(SizeX, SizeY); }
+	FIntPoint GetSizeXY() const
+	{
+		return FIntPoint(SizeX, SizeY);
+	}
+
+	FD3D12Texture* GetBackBuffer_RenderThread() const
+	{
+		check(IsInRenderingThread());
+#if D3D12RHI_USE_DUMMY_BACKBUFFER
+		return DummyBackBuffer_RenderThread;
+#else
+		checkSlow(CurrentBackBuffer_RenderThread);
+		return CurrentBackBuffer_RenderThread->Texture;
+#endif
+	}
+
+#if D3D12RHI_SUPPORTS_UAV_BACKBUFFER
+	FD3D12UnorderedAccessView_RHI* GetBackBufferUAV_RenderThread() const
+	{
+		checkSlow(CurrentBackBuffer_RenderThread);
+		return CurrentBackBuffer_RenderThread->UAV;
+	}
+#endif
 
-	FD3D12Texture* GetDummyBackBuffer_RenderThread(bool bInIsSDR) const;
-	FD3D12Texture* GetBackBuffer_RenderThread() const;
+	FD3D12Texture* GetBackBuffer_RHIThread() const
+	{
+		checkSlow(CurrentBackBuffer_RHIThread);
+		return CurrentBackBuffer_RHIThread->Texture;
+	}
+
+	FD3D12Texture* GetSDRBackBuffer_RHIThread() const
+	{
+		checkSlow(CurrentBackBuffer_RHIThread);
 
-	FD3D12Texture* GetBackBuffer_RHIThread() const { return BackBuffer_RHIThread; }
-	FD3D12Texture* GetSDRBackBuffer_RHIThread() const { return (PixelFormat == SDRPixelFormat) ? GetBackBuffer_RHIThread() : SDRBackBuffer_RHIThread; }
+#if D3D12RHI_USE_SDR_BACKBUFFER
+		if (PixelFormat != SDRPixelFormat)
+		{
+			return CurrentBackBuffer_RHIThread->TextureSDR;
+		}
+#endif
 
-	FD3D12UnorderedAccessView_RHI* GetBackBufferUAV_RenderThread() const;
+		return CurrentBackBuffer_RHIThread->Texture;
+	}
+
+#if WITH_MGPU
+	uint32 GetNextPresentGPUIndex() const
+	{
+		FScopeLock Lock(&ExpectedBackBufferIndexLock);
+		return BackBuffers[ExpectedBackBufferIndex_RenderThread].GPUIndex;
+	}
+#endif // WITH_MGPU
 
 	virtual void WaitForFrameEventCompletion() override;
 	virtual void IssueFrameEvent() override;
@@ -129,39 +171,16 @@ public:
 
 	/** Advance and get the next present GPU index */
 	void AdvanceExpectedBackBufferIndex_RenderThread();
-#if WITH_MGPU
-	uint32 GetNextPresentGPUIndex() const
-	{
-		FScopeLock Lock(&ExpectedBackBufferIndexLock);
-		return BackBufferGPUIndices.IsValidIndex(ExpectedBackBufferIndex_RenderThread) ? BackBufferGPUIndices[ExpectedBackBufferIndex_RenderThread] : 0;
-	}
-#endif // WITH_MGPU
 
 	void OnResumeRendering();
 	void OnSuspendRendering();
 
-	static DXGI_FORMAT GetRenderTargetFormat(EPixelFormat PixelFormat)
-	{
-		DXGI_FORMAT	DXFormat = (DXGI_FORMAT)GPixelFormats[PixelFormat].PlatformFormat;
-		switch (DXFormat)
-		{
-		case DXGI_FORMAT_B8G8R8A8_TYPELESS:		return DXGI_FORMAT_B8G8R8A8_UNORM;
-		case DXGI_FORMAT_BC1_TYPELESS:			return DXGI_FORMAT_BC1_UNORM;
-		case DXGI_FORMAT_BC2_TYPELESS:			return DXGI_FORMAT_BC2_UNORM;
-		case DXGI_FORMAT_BC3_TYPELESS:			return DXGI_FORMAT_BC3_UNORM;
-		case DXGI_FORMAT_R16_TYPELESS:			return DXGI_FORMAT_R16_UNORM;
-		case DXGI_FORMAT_R8G8B8A8_TYPELESS:		return DXGI_FORMAT_R8G8B8A8_UNORM;
-		default: 								return DXFormat;
-		}
-	}
-
 private:
 	bool IsPresentAllowed();
 
-	/**
-	 * Create the dummy back buffer textures
-	 */
-	FD3D12Texture* CreateDummyBackBufferTextures(FD3D12Adapter* InAdapter, EPixelFormat InPixelFormat, uint32 InSizeX, uint32 InSizeY, bool bInIsSDR);
+#if D3D12RHI_USE_DUMMY_BACKBUFFER
+	FD3D12Texture* CreateDummyBackBufferTextures(FD3D12Adapter* InAdapter, EPixelFormat InPixelFormat, uint32 InSizeX, uint32 InSizeY);
+#endif
 
 	/**
 	 * Presents the swap chain checking the return result.
@@ -179,14 +198,48 @@ private:
 	void FinalDestroyInternal();
 	void ClearPresentQueue();
 
-	HWND WindowHandle;
-	uint32 SizeX;
-	uint32 SizeY;
-	bool bIsFullscreen;
-	bool bFullscreenLost;
+	// Determine how deep the swapchain should be
+	void InitializeBackBufferArrays();
+
+	/** See if HDR can be enabled or not based on RHI support and current engine settings. */
+	bool CheckHDRSupport();
+
+	/** Enable HDR meta data transmission and set the necessary color space. */
+	void EnableHDR();
+
+	/** Disable HDR meta data transmission and set the necessary color space. */
+	void ShutdownHDR();
+
+#if D3D12RHI_USE_DXGI_COLOR_SPACE
+	/** Ensure the correct color space is set on the swap chain */
+	void EnsureColorSpace(EDisplayColorGamut DisplayGamut, EDisplayOutputFormat OutputDevice);
+#endif
+
+	void SetBackBufferIndex_RHIThread(uint32 Index)
+	{
+		CurrentBackBufferIndex_RHIThread = Index % NumBackBuffers;
+		CurrentBackBuffer_RHIThread = &BackBuffers[CurrentBackBufferIndex_RHIThread];
+	}
+
+	void SetBackBufferIndex_RenderThread(uint32 Index)
+	{
+		ExpectedBackBufferIndex_RenderThread = Index % NumBackBuffers;
+		CurrentBackBuffer_RenderThread = &BackBuffers[ExpectedBackBufferIndex_RenderThread];
+	}
+
+private:
+	const HWND WindowHandle;
+	uint32 SizeX = 0;
+	uint32 SizeY = 0;
 	EPixelFormat PixelFormat;
-	bool bIsValid;
-	bool bAllowTearing;
+#if D3D12RHI_USE_SDR_BACKBUFFER
+	static constexpr EPixelFormat SDRPixelFormat = PF_B8G8R8A8;
+#endif
+
+	bool bIsFullscreen = false;
+	bool bFullscreenLost = false;
+	bool bIsValid = true;
+	bool bAllowTearing = true;
 	bool bNeedSwapChain = false;
 
 #if D3D12_VIEWPORT_EXPOSES_SWAP_CHAIN
@@ -201,75 +254,57 @@ private:
 	TRefCountPtr<IDXGISwapChain4> SwapChain4;
 #endif
 
-#if PLATFORM_WINDOWS
-	DXGI_COLOR_SPACE_TYPE ColorSpace;
+#if D3D12RHI_USE_DXGI_COLOR_SPACE
+	DXGI_COLOR_SPACE_TYPE ColorSpace = DXGI_COLOR_SPACE_RGB_FULL_G22_NONE_P709;
 #endif
 #endif // D3D12_VIEWPORT_EXPOSES_SWAP_CHAIN
 
-	TArray<TRefCountPtr<FD3D12Texture>> BackBuffers;
-	TArray<TRefCountPtr<FD3D12UnorderedAccessView_RHI>> BackBuffersUAV;
-	uint32 NumBackBuffers;
+	struct FBackBufferData
+	{
+		TRefCountPtr<FD3D12Texture>                 Texture;
+#if D3D12RHI_USE_SDR_BACKBUFFER
+		// When HDR is enabled, SDR backbuffers may be required on some architectures for game DVR or broadcasting
+		TRefCountPtr<FD3D12Texture>                 TextureSDR;
+#endif
+#if D3D12RHI_SUPPORTS_UAV_BACKBUFFER
+		TRefCountPtr<FD3D12UnorderedAccessView_RHI> UAV;
+#endif
+#if WITH_MGPU
+		uint32                                      GPUIndex = 0;
+#endif
+	};
 
-	TRefCountPtr<FD3D12Texture> DummyBackBuffer_RenderThread; // Dummy back buffer texture which always references the current back buffer on the RHI thread
-	uint32 CurrentBackBufferIndex_RHIThread;
-	FD3D12Texture* BackBuffer_RHIThread;
-	FD3D12Texture* BackBuffer_RenderThread;
-	FD3D12UnorderedAccessView_RHI* BackBufferUAV_RenderThread;
+	static constexpr uint32 NumBackBuffers = GD3D12RHINumBackBuffers;
+	TStaticArray<FBackBufferData, NumBackBuffers> BackBuffers;
 
-#if WITH_MGPU
-	int32 BackbufferMultiGPUBinding; // where INDEX_NONE cycles through the GPU, otherwise the GPU index.
-	mutable FCriticalSection ExpectedBackBufferIndexLock; // Can very rarely be modified on the RHI thread as well if present is skipped
+#if D3D12RHI_USE_DUMMY_BACKBUFFER
+	// Dummy back buffer texture which always references the current back buffer on the RHI thread
+	TRefCountPtr<FD3D12Texture> DummyBackBuffer_RenderThread;
 #endif
-	uint32 ExpectedBackBufferIndex_RenderThread; // Expected back buffer GPU index - used and updated on RenderThread!
-#if WITH_MGPU
-	TArray<uint32> BackBufferGPUIndices;
-	FD3D12SyncPointRef LastFrameSyncPoint;
-#endif // WITH_MGPU
 
-	/** 
-	 * When HDR is enabled, SDR backbuffers may be required on some architectures for game DVR or broadcasting
-	 */
-	TArray<TRefCountPtr<FD3D12Texture>> SDRBackBuffers;
-	TRefCountPtr<FD3D12Texture> SDRDummyBackBuffer_RenderThread;
-	FD3D12Texture* SDRBackBuffer_RHIThread;
-	EPixelFormat SDRPixelFormat;
-	EDisplayColorGamut DisplayColorGamut;
-	EDisplayOutputFormat DisplayOutputFormat;
+	FBackBufferData* CurrentBackBuffer_RHIThread = nullptr;
+	FBackBufferData* CurrentBackBuffer_RenderThread = nullptr;
+
+	uint32 CurrentBackBufferIndex_RHIThread = 0;
+	uint32 ExpectedBackBufferIndex_RenderThread = 0;
+	EDisplayColorGamut DisplayColorGamut = EDisplayColorGamut::sRGB_D65;
+	EDisplayOutputFormat DisplayOutputFormat = EDisplayOutputFormat::SDR_sRGB;
 
 	/** A fence value used to track the GPU's progress. */
 	TArray<FD3D12SyncPointRef> FrameSyncPoints;
 
-	// Determine how deep the swapchain should be
-	void CalculateSwapChainDepth(int32 DefaultSwapChainDepth);
-
 	FCustomPresentRHIRef CustomPresent;
 
-	DXGI_MODE_DESC SetupDXGI_MODE_DESC() const;
-
 #if WITH_MGPU
-	FD3D12FramePacing* FramePacerRunnable;
-#endif //WITH_MGPU
+	// Where INDEX_NONE cycles through the GPU, otherwise the GPU index.
+	int32 BackbufferMultiGPUBinding = 0;
 
-	struct DisplayChromacities
-	{
-		float RedX, RedY;
-		float GreenX, GreenY;
-		float BlueX, BlueY;
-		float WpX, WpY;
-	};
+	// Can very rarely be modified on the RHI thread as well if present is skipped
+	mutable FCriticalSection ExpectedBackBufferIndexLock;
 
-	/** See if HDR can be enabled or not based on RHI support and current engine settings. */
-	bool CheckHDRSupport();
-
-	/** Enable HDR meta data transmission and set the necessary color space. */
-	void EnableHDR();
-
-	/** Disable HDR meta data transmission and set the necessary color space. */
-	void ShutdownHDR();
+	FD3D12SyncPointRef LastFrameSyncPoint;
 
-#if PLATFORM_WINDOWS
-	/** Ensure the correct color space is set on the swap chain */
-	void EnsureColorSpace(EDisplayColorGamut DisplayGamut, EDisplayOutputFormat OutputDevice);
+	FD3D12FramePacing* FramePacerRunnable = nullptr;
 #endif
 };
 
diff --git a/Engine/Source/Runtime/D3D12RHI/Private/Windows/WindowsD3D12RHIDefinitions.h b/Engine/Source/Runtime/D3D12RHI/Private/Windows/WindowsD3D12RHIDefinitions.h
--- a/Engine/Source/Runtime/D3D12RHI/Private/Windows/WindowsD3D12RHIDefinitions.h
+++ b/Engine/Source/Runtime/D3D12RHI/Private/Windows/WindowsD3D12RHIDefinitions.h
@@ -4,6 +4,8 @@
 
 #include "D3D12ThirdParty.h"
 
+constexpr uint32 GD3D12RHINumBackBuffers = 3;
+
 // Windows uses the graphics command list interface for the copy queue.
 typedef ID3D12GraphicsCommandList ID3D12CopyCommandList;
 
@@ -24,7 +26,10 @@ typedef ID3D12GraphicsCommandList ID3D12CopyCommandList;
 
 #define USE_STATIC_ROOT_SIGNATURE			0
 #define D3D12RHI_USE_CONSTANT_BUFFER_VIEWS	0
-#define D3D12_USE_DUMMY_BACKBUFFER			1
+#define D3D12RHI_USE_DUMMY_BACKBUFFER		1
+#define D3D12RHI_USE_SDR_BACKBUFFER			0
+#define D3D12RHI_SUPPORTS_UAV_BACKBUFFER	0
+#define D3D12RHI_USE_DXGI_COLOR_SPACE		1
 
 // Only enable pipeline statistics if we've got the stats system enabled to display them
 #define D3D12RHI_ENABLE_PIPELINE_STATISTICS (1 && STATS)
diff --git a/Engine/Source/Runtime/D3D12RHI/Private/Windows/WindowsD3D12Viewport.cpp b/Engine/Source/Runtime/D3D12RHI/Private/Windows/WindowsD3D12Viewport.cpp
--- a/Engine/Source/Runtime/D3D12RHI/Private/Windows/WindowsD3D12Viewport.cpp
+++ b/Engine/Source/Runtime/D3D12RHI/Private/Windows/WindowsD3D12Viewport.cpp
@@ -12,8 +12,6 @@
 
 #include "Windows/AllowWindowsPlatformTypes.h"
 
-static const uint32 WindowsDefaultNumBackBuffers = 3;
-
 extern FD3D12Texture* GetSwapChainSurface(FD3D12Device* Parent, EPixelFormat PixelFormat, uint32 SizeX, uint32 SizeY, IDXGISwapChain* SwapChain, uint32 BackBufferIndex, TRefCountPtr<ID3D12Resource> BackBufferResourceOverride);
 
 static int32 GD3D12UseAllowTearing = 1;
@@ -29,30 +27,9 @@ FD3D12Viewport::FD3D12Viewport(class FD3D12Adapter* InParent, HWND InWindowHandl
 	, WindowHandle(InWindowHandle)
 	, SizeX(InSizeX)
 	, SizeY(InSizeY)
-	, bIsFullscreen(bInIsFullscreen)
-	, bFullscreenLost(false)
 	, PixelFormat(InPreferredPixelFormat)
-	, bIsValid(true)
+	, bIsFullscreen(bInIsFullscreen)
 	, bNeedSwapChain(!FParse::Param(FCommandLine::Get(), TEXT("RenderOffScreen")))
-	, ColorSpace(DXGI_COLOR_SPACE_RGB_FULL_G22_NONE_P709)
-	, NumBackBuffers(WindowsDefaultNumBackBuffers)
-	, DummyBackBuffer_RenderThread(nullptr)
-	, CurrentBackBufferIndex_RHIThread(0)
-	, BackBuffer_RHIThread(nullptr)
-	, BackBuffer_RenderThread(nullptr)
-	, BackBufferUAV_RenderThread(nullptr)
-#if WITH_MGPU
-	, BackbufferMultiGPUBinding(0)
-#endif //WITH_MGPU
-	, ExpectedBackBufferIndex_RenderThread(0)
-	, SDRDummyBackBuffer_RenderThread(nullptr)
-	, SDRBackBuffer_RHIThread(nullptr)
-	, SDRPixelFormat(PF_B8G8R8A8)
-	, DisplayColorGamut(EDisplayColorGamut::sRGB_D65)
-	, DisplayOutputFormat(EDisplayOutputFormat::SDR_sRGB)
-#if WITH_MGPU
-	, FramePacerRunnable(nullptr)
-#endif //WITH_MGPU
 {
 	check(IsInGameThread());
 	GetParentAdapter()->GetViewports().Add(this);
@@ -96,8 +73,6 @@ void FD3D12Viewport::Init()
 		}
 	}
 
-	CalculateSwapChainDepth(WindowsDefaultNumBackBuffers);
-
 	bool bStereoMode = false;
 	if (FD3D12DynamicRHI::GetD3DRHI()->IsQuadBufferStereoEnabled())
 	{
@@ -112,6 +87,8 @@ void FD3D12Viewport::Init()
 		}
 	}
 
+	InitializeBackBufferArrays();
+
 	// Create the swapchain.
 	if (bNeedSwapChain)
 	{
@@ -128,7 +105,7 @@ void FD3D12Viewport::Init()
 
 		SwapChainDesc1.Width       = SizeX;
 		SwapChainDesc1.Height      = SizeY;
-		SwapChainDesc1.Format      = GetRenderTargetFormat(PixelFormat);
+		SwapChainDesc1.Format      = UE::DXGIUtilities::GetSwapChainFormat(PixelFormat);
 		SwapChainDesc1.Stereo      = bStereoMode ? TRUE : FALSE;
 		SwapChainDesc1.SampleDesc.Count = 1;
 		SwapChainDesc1.SampleDesc.Quality = 0;
@@ -254,11 +231,30 @@ void FD3D12Viewport::ConditionalResetSwapChain(bool bIgnoreFocus)
 	}
 }
 
+/**
+* Create the dummy back buffer textures - They don't have actual D3D resource but are used to always reference the current back buffer index on the RHI thread
+*/
+FD3D12Texture* FD3D12Viewport::CreateDummyBackBufferTextures(FD3D12Adapter* InAdapter, EPixelFormat InPixelFormat, uint32 InSizeX, uint32 InSizeY)
+{
+	FRHITextureCreateDesc CreateDesc =
+		FRHITextureCreateDesc::Create2D(TEXT("BackBufferReference"))
+		.SetExtent(FIntPoint(InSizeX, InSizeY))
+		.SetFormat(InPixelFormat)
+		.SetFlags(ETextureCreateFlags::RenderTargetable | ETextureCreateFlags::Presentable | ETextureCreateFlags::ResolveTargetable)
+		.SetInitialState(ERHIAccess::Present);
+
+	FD3D12Texture* Result = InAdapter->CreateLinkedObject<FD3D12Texture>(FRHIGPUMask::All(), [this, &CreateDesc](FD3D12Device* Device)
+		{
+			return new FD3D12BackBufferReferenceTexture2D(CreateDesc, this, Device);
+		});
+	return Result;
+}
+
 void FD3D12Viewport::ResizeInternal()
 {
 	FD3D12Adapter* Adapter = GetParentAdapter();
 
-	CalculateSwapChainDepth(WindowsDefaultNumBackBuffers);
+	InitializeBackBufferArrays();
 
 	UINT SwapChainFlags = DXGI_SWAP_CHAIN_FLAG_ALLOW_MODE_SWITCH;
 	if (bAllowTearing)
@@ -271,20 +267,17 @@ void FD3D12Viewport::ResizeInternal()
 	{
 		TArray<ID3D12CommandQueue*> CommandQueues;
 		TArray<uint32> NodeMasks;
-		BackBufferGPUIndices.Empty(NumBackBuffers);
 
-		for (uint32 i = 0; i < NumBackBuffers; ++i)
+		for (uint32 Index = 0; Index < NumBackBuffers; Index++)
 		{
 			// When BackbufferMultiGPUBinding == INDEX_NONE, cycle through each GPU.
-			const uint32 BackBufferGPUIndex = BackbufferMultiGPUBinding == INDEX_NONE ? i % GNumExplicitGPUsForRendering : BackbufferMultiGPUBinding;
-			BackBufferGPUIndices.Add(BackBufferGPUIndex);
+			BackBuffers[Index].GPUIndex = BackbufferMultiGPUBinding == INDEX_NONE ? (Index % GNumExplicitGPUsForRendering) : BackbufferMultiGPUBinding;
 		}
 
 		// Select the GPU for each element in the swapchain 
-		for (uint32 i = 0; i < NumBackBuffers; ++i)
+		for (uint32 Index = 0; Index < NumBackBuffers; Index++)
 		{
-			const uint32 GPUIndex = BackBufferGPUIndices[i];
-			FD3D12Device* Device = Adapter->GetDevice(GPUIndex);
+			FD3D12Device* Device = Adapter->GetDevice(BackBuffers[Index].GPUIndex);
 
 			CommandQueues.Add(Device->GetQueue(ED3D12QueueType::Direct).D3DCommandQueue);
 			NodeMasks.Add(Device->GetGPUMask().GetNative());
@@ -292,16 +285,15 @@ void FD3D12Viewport::ResizeInternal()
 
 		if (SwapChain3)
 		{
-			VERIFYD3D12RESULT_EX(SwapChain3->ResizeBuffers1(NumBackBuffers, SizeX, SizeY, GetRenderTargetFormat(PixelFormat), SwapChainFlags, NodeMasks.GetData(), (IUnknown**)CommandQueues.GetData()), Adapter->GetD3DDevice());
+			VERIFYD3D12RESULT_EX(SwapChain3->ResizeBuffers1(NumBackBuffers, SizeX, SizeY, UE::DXGIUtilities::GetSwapChainFormat(PixelFormat), SwapChainFlags, NodeMasks.GetData(), (IUnknown**)CommandQueues.GetData()), Adapter->GetD3DDevice());
 		}
 
-		for (uint32 i = 0; i < NumBackBuffers; ++i)
+		for (uint32 Index = 0; Index < NumBackBuffers; Index++)
 		{
-			const uint32 GPUIndex = BackBufferGPUIndices[i];
-			FD3D12Device* Device = Adapter->GetDevice(GPUIndex);
+			FD3D12Device* Device = Adapter->GetDevice(BackBuffers[Index].GPUIndex);
 
-			check(BackBuffers[i].GetReference() == nullptr);
-			BackBuffers[i] = GetSwapChainSurface(Device, PixelFormat, SizeX, SizeY, SwapChain1, i, nullptr);
+			check(BackBuffers[Index].Texture.GetReference() == nullptr);
+			BackBuffers[Index].Texture = GetSwapChainSurface(Device, PixelFormat, SizeX, SizeY, SwapChain1, Index, nullptr);
 		}
 	}
 	else
@@ -311,33 +303,29 @@ void FD3D12Viewport::ResizeInternal()
 		{
 			auto Lambda = [this, SwapChainFlags]() -> FString
 			{
-				return FString::Printf(TEXT("Num=%d, Size=(%d,%d), PF=%d, DXGIFormat=0x%x, Flags=0x%x"), NumBackBuffers, SizeX, SizeY, (int32)PixelFormat, (int32)GetRenderTargetFormat(PixelFormat), SwapChainFlags);
+				return FString::Printf(TEXT("Num=%d, Size=(%d,%d), PF=%d, DXGIFormat=0x%x, Flags=0x%x"), NumBackBuffers, SizeX, SizeY, (int32)PixelFormat, (int32)UE::DXGIUtilities::GetSwapChainFormat(PixelFormat), SwapChainFlags);
 			};
 
-			VERIFYD3D12RESULT_LAMBDA(SwapChain1->ResizeBuffers(NumBackBuffers, SizeX, SizeY, GetRenderTargetFormat(PixelFormat), SwapChainFlags), Adapter->GetD3DDevice(), Lambda);
+			VERIFYD3D12RESULT_LAMBDA(SwapChain1->ResizeBuffers(NumBackBuffers, SizeX, SizeY, UE::DXGIUtilities::GetSwapChainFormat(PixelFormat), SwapChainFlags), Adapter->GetD3DDevice(), Lambda);
 		}
 
 		FD3D12Device* Device = Adapter->GetDevice(0);
 		for (uint32 i = 0; i < NumBackBuffers; ++i)
 		{
-			check(BackBuffers[i].GetReference() == nullptr);
-			BackBuffers[i] = GetSwapChainSurface(Device, PixelFormat, SizeX, SizeY, SwapChain1, i, nullptr);
+			check(BackBuffers[i].Texture.GetReference() == nullptr);
+			BackBuffers[i].Texture = GetSwapChainSurface(Device, PixelFormat, SizeX, SizeY, SwapChain1, i, nullptr);
 		}
 	}
 
-	CurrentBackBufferIndex_RHIThread = 0;
-	BackBuffer_RHIThread = BackBuffers[CurrentBackBufferIndex_RHIThread].GetReference();
-	SDRBackBuffer_RHIThread = SDRBackBuffers[CurrentBackBufferIndex_RHIThread].GetReference();
-
-#if !D3D12_USE_DUMMY_BACKBUFFER
-    static_assert(false, "D3D12_USE_DUMMY_BACKBUFFER==0 has not been tested for viewports");
-	ExpectedBackBufferIndex_RenderThread = 0;
-	BackBuffer_RenderThread = BackBuffers[ExpectedBackBufferIndex_RenderThread].GetReference();
-#endif
+	SetBackBufferIndex_RHIThread(0);
 
+#if D3D12RHI_USE_DUMMY_BACKBUFFER
 	// Create dummy back buffer which always reference to the actual RHI thread back buffer - can't be bound directly to D3D12
-	DummyBackBuffer_RenderThread = CreateDummyBackBufferTextures(Adapter, PixelFormat, SizeX, SizeY, false);
-	SDRDummyBackBuffer_RenderThread = (SDRBackBuffer_RHIThread != nullptr) ? CreateDummyBackBufferTextures(Adapter, PixelFormat, SizeX, SizeY, true) : nullptr;
+	DummyBackBuffer_RenderThread = CreateDummyBackBufferTextures(Adapter, PixelFormat, SizeX, SizeY);
+#else
+	static_assert(false, "D3D12RHI_USE_DUMMY_BACKBUFFER==0 has not been tested for viewports");
+	SetBackBufferIndex_RenderThread(0);
+#endif
 }
 
 HRESULT FD3D12Viewport::PresentInternal(int32 SyncInterval)
@@ -497,14 +485,6 @@ void FD3D12Viewport::OnResumeRendering()
 void FD3D12Viewport::OnSuspendRendering()
 {}
 
-FD3D12Texture* FD3D12Viewport::GetBackBuffer_RenderThread() const
-{
-	check(IsInRenderingThread());
-	const bool sIsSDR = false;
-	FD3D12Texture* const DummyBackBuffer = GetDummyBackBuffer_RenderThread(sIsSDR);
-	return DummyBackBuffer;
-}
-
 bool FD3D12Viewport::IsPresentAllowed()
 {
 	return !FD3D12DynamicRHI::GetD3DRHI()->RHIIsRenderingSuspended();
diff --git a/Engine/Source/Runtime/RHICore/Internal/DXGIUtilities.h b/Engine/Source/Runtime/RHICore/Internal/DXGIUtilities.h
--- a/Engine/Source/Runtime/RHICore/Internal/DXGIUtilities.h
+++ b/Engine/Source/Runtime/RHICore/Internal/DXGIUtilities.h
@@ -10,6 +10,7 @@
 	#error "Platform needs to define RHICORE_PLATFORM_DXGI_H"
 #endif
 
+#include "PixelFormat.h"
 #include "RHIDefinitions.h"
 #include "Misc/AssertionMacros.h"
 
@@ -470,6 +471,21 @@ namespace UE::DXGIUtilities
 	{
 		return GetFormatSizeInBits(Format) / 8;
 	}
+
+	inline DXGI_FORMAT GetSwapChainFormat(EPixelFormat PixelFormat)
+	{
+		DXGI_FORMAT	DXFormat = static_cast<DXGI_FORMAT>(GPixelFormats[PixelFormat].PlatformFormat);
+		switch (DXFormat)
+		{
+		case DXGI_FORMAT_B8G8R8A8_TYPELESS:		return DXGI_FORMAT_B8G8R8A8_UNORM;
+		case DXGI_FORMAT_BC1_TYPELESS:			return DXGI_FORMAT_BC1_UNORM;
+		case DXGI_FORMAT_BC2_TYPELESS:			return DXGI_FORMAT_BC2_UNORM;
+		case DXGI_FORMAT_BC3_TYPELESS:			return DXGI_FORMAT_BC3_UNORM;
+		case DXGI_FORMAT_R16_TYPELESS:			return DXGI_FORMAT_R16_UNORM;
+		case DXGI_FORMAT_R8G8B8A8_TYPELESS:		return DXGI_FORMAT_R8G8B8A8_UNORM;
+		default: 								return DXFormat;
+		}
+	}
 } // UE::DXGIUtilities
 
 #endif // PLATFORM_MICROSOFT
diff --git a/Engine/Source/Runtime/Windows/D3D11RHI/Public/D3D11Viewport.h b/Engine/Source/Runtime/Windows/D3D11RHI/Public/D3D11Viewport.h
--- a/Engine/Source/Runtime/Windows/D3D11RHI/Public/D3D11Viewport.h
+++ b/Engine/Source/Runtime/Windows/D3D11RHI/Public/D3D11Viewport.h
@@ -11,6 +11,7 @@
 #include "RenderUtils.h"
 #include "RHIResources.h"
 #include "Windows/D3D11ThirdParty.h"
+#include "DXGIUtilities.h"
 
 /** A D3D event query resource. */
 class FD3D11EventQuery
@@ -97,17 +98,7 @@ public:
 
 	static DXGI_FORMAT GetRenderTargetFormat(EPixelFormat PixelFormat)
 	{
-		DXGI_FORMAT	DXFormat = (DXGI_FORMAT)GPixelFormats[PixelFormat].PlatformFormat;
-		switch(DXFormat)
-		{
-		case DXGI_FORMAT_B8G8R8A8_TYPELESS:		return DXGI_FORMAT_B8G8R8A8_UNORM;
-		case DXGI_FORMAT_BC1_TYPELESS:			return DXGI_FORMAT_BC1_UNORM;
-		case DXGI_FORMAT_BC2_TYPELESS:			return DXGI_FORMAT_BC2_UNORM;
-		case DXGI_FORMAT_BC3_TYPELESS:			return DXGI_FORMAT_BC3_UNORM;
-		case DXGI_FORMAT_R16_TYPELESS:			return DXGI_FORMAT_R16_UNORM;
-		case DXGI_FORMAT_R8G8B8A8_TYPELESS:		return DXGI_FORMAT_R8G8B8A8_UNORM;
-		default: 								return DXFormat;
-		}
+		return UE::DXGIUtilities::GetSwapChainFormat(PixelFormat);
 	}
 
 protected:
