Update PSO cache clearing for AMD.

Now it also uses LocalLow, and it uses multiple file extensions per folder. Clear the entire folder instead of looking at specific extensions.

diff --git a/Engine/Source/Runtime/D3D12RHI/Private/Windows/WindowsD3D12Device.cpp b/Engine/Source/Runtime/D3D12RHI/Private/Windows/WindowsD3D12Device.cpp
--- a/Engine/Source/Runtime/D3D12RHI/Private/Windows/WindowsD3D12Device.cpp
+++ b/Engine/Source/Runtime/D3D12RHI/Private/Windows/WindowsD3D12Device.cpp
@@ -1308,10 +1308,16 @@ static void ClearPSODriverCache()
 		else if (IsRHIDeviceAMD())
 		{
 			FString PSOPath = FPaths::Combine(*LocalAppDataFolder, TEXT("AMD"), TEXT("DxCache"));
-			ClearFolder(PSOPath, TEXT(".bin"));
+			ClearFolder(PSOPath, nullptr);
 
 			PSOPath = FPaths::Combine(*LocalAppDataFolder, TEXT("AMD"), TEXT("DxcCache"));
-			ClearFolder(PSOPath, TEXT(".parc"));
+			ClearFolder(PSOPath, nullptr);
+
+			PSOPath = FPaths::Combine(*LocalAppDataFolder, TEXT(".."), TEXT("LocalLow"), TEXT("AMD"), TEXT("DxCache"));
+			ClearFolder(PSOPath, nullptr);
+
+			PSOPath = FPaths::Combine(*LocalAppDataFolder, TEXT(".."), TEXT("LocalLow"), TEXT("AMD"), TEXT("DxcCache"));
+			ClearFolder(PSOPath, nullptr);
 		}
 		else if (IsRHIDeviceIntel())
 		{
