Use top-of-pipe timestamps for command list / idle begin markers on supported platforms.
 - CommandListBegin and IdleBegin timestamps should be written out before any future command list work is started.
 - Previously we were always doing bottom-of-pipe writes, which won't accurately reflect the time a command list spent executing on the GPU.

diff --git a/Engine/Source/Runtime/D3D12RHI/Private/D3D12CommandList.cpp b/Engine/Source/Runtime/D3D12RHI/Private/D3D12CommandList.cpp
--- a/Engine/Source/Runtime/D3D12RHI/Private/D3D12CommandList.cpp
+++ b/Engine/Source/Runtime/D3D12RHI/Private/D3D12CommandList.cpp
@@ -358,20 +358,43 @@ void FD3D12CommandList::EndQuery(FD3D12QueryLocation const& Location)
 		break;
 
 	case D3D12_QUERY_TYPE_TIMESTAMP:
-		WriteTimestamp(Location);
-
-		// Command list begin/end timestamps are handled separately by the 
-		// submission thread, so shouldn't be in the TimestampQueries array.
-		if (Location.Type != ED3D12QueryType::CommandListBegin && Location.Type != ED3D12QueryType::CommandListEnd)
 		{
-			State.TimestampQueries.Add(Location);
+			ED3D12QueryPosition Position;
+			switch (Location.Type)
+			{
+			default:
+				checkf(false, TEXT("Query location type is not a top or bottom of pipe timestamp."));
+				Position = ED3D12QueryPosition::BottomOfPipe;
+				break;
+
+			case ED3D12QueryType::CommandListBegin:
+			case ED3D12QueryType::IdleBegin:
+				Position = ED3D12QueryPosition::TopOfPipe;
+				break;
+
+			case ED3D12QueryType::CommandListEnd:
+			case ED3D12QueryType::IdleEnd:
+			case ED3D12QueryType::AdjustedMicroseconds:
+			case ED3D12QueryType::AdjustedRaw:
+				Position = ED3D12QueryPosition::BottomOfPipe;
+				break;
+			}
+
+			WriteTimestamp(Location, Position);
+
+			// Command list begin/end timestamps are handled separately by the 
+			// submission thread, so shouldn't be in the TimestampQueries array.
+			if (Location.Type != ED3D12QueryType::CommandListBegin && Location.Type != ED3D12QueryType::CommandListEnd)
+			{
+				State.TimestampQueries.Add(Location);
+			}
 		}
 		break;
 	}
 }
 
 #if D3D12RHI_PLATFORM_USES_TIMESTAMP_QUERIES
-void FD3D12CommandList::WriteTimestamp(FD3D12QueryLocation const& Location)
+void FD3D12CommandList::WriteTimestamp(FD3D12QueryLocation const& Location, ED3D12QueryPosition Position)
 {
 	GraphicsCommandList()->EndQuery(
 		Location.Heap->GetD3DQueryHeap(),
diff --git a/Engine/Source/Runtime/D3D12RHI/Private/D3D12CommandList.h b/Engine/Source/Runtime/D3D12RHI/Private/D3D12CommandList.h
--- a/Engine/Source/Runtime/D3D12RHI/Private/D3D12CommandList.h
+++ b/Engine/Source/Runtime/D3D12RHI/Private/D3D12CommandList.h
@@ -238,7 +238,7 @@ public:
 #endif
 
 private:
-	void WriteTimestamp(FD3D12QueryLocation const& Location);
+	void WriteTimestamp(FD3D12QueryLocation const& Location, ED3D12QueryPosition Position);
 
 	// Contents of the state struct are reset when the command list is recycled
 	struct FState
diff --git a/Engine/Source/Runtime/D3D12RHI/Private/D3D12Submission.h b/Engine/Source/Runtime/D3D12RHI/Private/D3D12Submission.h
--- a/Engine/Source/Runtime/D3D12RHI/Private/D3D12Submission.h
+++ b/Engine/Source/Runtime/D3D12RHI/Private/D3D12Submission.h
@@ -146,6 +146,15 @@ enum class ED3D12QueryType
 	Occlusion
 };
 
+enum class ED3D12QueryPosition
+{
+	// Query result should be written before any future command list work is started.
+	TopOfPipe,
+
+	// Query result should be written after all prior command list work has completed.
+	BottomOfPipe
+};
+
 // The location of a single (timestamp or occlusion) query result.
 struct FD3D12QueryLocation
 {
