Improve batching of timestamp and render query resolves
 - Grouping query ranges by heap, sorting them, and then merging contiguous regions allows us to greatly reduce the number of ResolveQueryData commands issued.

diff --git a/Engine/Source/Runtime/D3D12RHI/Private/D3D12Device.h b/Engine/Source/Runtime/D3D12RHI/Private/D3D12Device.h
--- a/Engine/Source/Runtime/D3D12RHI/Private/D3D12Device.h
+++ b/Engine/Source/Runtime/D3D12RHI/Private/D3D12Device.h
@@ -150,7 +150,7 @@ public:
 	uint32 NumCommandListsInBatch = 0;
 
 	// Query ranges/locations to be resolved when the submission thread receives a command list which is still open.
-	TArray<FD3D12QueryRange   > PendingQueryRanges;
+	TMap<FD3D12QueryHeap*, TArray<FD3D12QueryRange>> PendingQueryRanges;
 	TArray<FD3D12QueryLocation> PendingTimestampQueries;
 	TArray<FD3D12QueryLocation> PendingOcclusionQueries;
 	TArray<FD3D12QueryLocation> PendingPipelineStatsQueries;
diff --git a/Engine/Source/Runtime/D3D12RHI/Private/D3D12Submission.cpp b/Engine/Source/Runtime/D3D12RHI/Private/D3D12Submission.cpp
--- a/Engine/Source/Runtime/D3D12RHI/Private/D3D12Submission.cpp
+++ b/Engine/Source/Runtime/D3D12RHI/Private/D3D12Submission.cpp
@@ -443,8 +443,6 @@ FD3D12DynamicRHI::FProcessResult FD3D12DynamicRHI::ProcessSubmissionQueue()
 
 						AccumulateQueries(CurrentCommandList);
 					}
-
-					CurrentQueue.PendingQueryRanges.Append(MoveTemp(Payload->QueryRanges));
 				}
 				
 				// Prepare the command lists from each payload for submission
@@ -453,10 +451,17 @@ FD3D12DynamicRHI::FProcessResult FD3D12DynamicRHI::ProcessSubmissionQueue()
 					FD3D12Payload* Payload = Queue->PayloadToSubmit;
 					check(Payload->SyncPointsToWait.Num() == 0);
 
-					Queue->BarrierTimestamps.CloseAndReset(Queue->PendingQueryRanges);
+					Queue->BarrierTimestamps.CloseAndReset(Payload->QueryRanges);
 					Queue->NumCommandListsInBatch = 0;
 
-					check(Payload->QueryRanges.Num() == 0);
+					// Gather query ranges from this payload, grouping by heap pointer
+					for (FD3D12QueryRange& Range : Payload->QueryRanges)
+					{
+						check(Range.End > Range.Start);
+						CurrentQueue.PendingQueryRanges.FindOrAdd(Range.Heap).Emplace(MoveTemp(Range));
+					}
+					Payload->QueryRanges.Reset();
+
 					check(Payload->TimestampQueries.Num() == 0);
 					check(Payload->OcclusionQueries.Num() == 0);
 					check(Payload->PipelineStatsQueries.Num() == 0);
@@ -479,6 +484,7 @@ FD3D12DynamicRHI::FProcessResult FD3D12DynamicRHI::ProcessSubmissionQueue()
 						{
 							FD3D12CommandList* ResolveCommandList = nullptr;
 							{
+								// We've got queries to resolve. Allocate a command list.
 								auto GetResolveCommandList = [&]() -> FD3D12CommandList*
 								{
 									if (ResolveCommandList)
@@ -490,33 +496,65 @@ FD3D12DynamicRHI::FProcessResult FD3D12DynamicRHI::ProcessSubmissionQueue()
 									return ResolveCommandList = Queue->Device->ObtainCommandList(Queue->BarrierAllocator, nullptr, nullptr);
 								};
 
-								// We've got queries to resolve. Allocate a command list.
-								for (FD3D12QueryRange const& Range : Queue->PendingQueryRanges)
+								// Ranges are grouped by heap pointer.
+								for (auto& [Heap, Ranges] : Queue->PendingQueryRanges)
 								{
-									check(Range.End > Range.Start);
-
+									{
 								#if ENABLE_RESIDENCY_MANAGEMENT
-									TArray<FD3D12ResidencyHandle*, TInlineAllocator<2>> ResidencyHandles;
-									ResidencyHandles.Add(&Range.Heap->GetHeapResidencyHandle());
-									ResidencyHandles.Append(Range.Heap->GetResultBuffer()->GetResidencyHandles());
-									GetResolveCommandList()->UpdateResidency(ResidencyHandles);
+										TArray<FD3D12ResidencyHandle*, TInlineAllocator<2>> ResidencyHandles;
+										ResidencyHandles.Add(&Heap->GetHeapResidencyHandle());
+										ResidencyHandles.Append(Heap->GetResultBuffer()->GetResidencyHandles());
+										GetResolveCommandList()->UpdateResidency(ResidencyHandles);
 								#endif // ENABLE_RESIDENCY_MANAGEMENT
+									}
 
-									if (Range.Heap->GetD3DQueryHeap())
+									// Sort the ranges into ascending order so we can merge adjacent ones,
+									// to reduce the number of ResolveQueryData calls we need to make.
+									Ranges.Sort();
+
+									int32 Index = 0;
+									uint32 Start, End;
+
+									auto GetNextRange = [&]()
+									{
+										if (Index >= Ranges.Num())
+											return false;
+
+										Start = Ranges[Index].Start;
+										End = Ranges[Index].End;
+										Index++;
+
+										while (Index < Ranges.Num() && Ranges[Index].Start == End)
+										{
+											// Ranges are contiguous. Extend.
+											End = Ranges[Index].End;
+											Index++;
+										}
+
+										return true;
+									};
+
+									while (GetNextRange())
 									{
-										GetResolveCommandList()->GraphicsCommandList()->ResolveQueryData(
-											Range.Heap->GetD3DQueryHeap(),
-											Range.Heap->QueryType,
-											Range.Start,
-											Range.End - Range.Start,
-											Range.Heap->GetResultBuffer()->GetResource(),
-											Range.Start * Range.Heap->GetResultSize()
-										);
+										if (Heap->GetD3DQueryHeap())
+										{
+											GetResolveCommandList()->GraphicsCommandList()->ResolveQueryData(
+												Heap->GetD3DQueryHeap(),
+												Heap->QueryType,
+												Start,
+												End - Start,
+												Heap->GetResultBuffer()->GetResource(),
+												Start * Heap->GetResultSize()
+											);
+										}
+
+										Payload->QueryRanges.Emplace(Heap, Start, End);
 									}
 								}
+
+								Queue->PendingQueryRanges.Reset();
 							}
 
-							Payload->QueryRanges          = MoveTemp(Queue->PendingQueryRanges         );
 							Payload->TimestampQueries     = MoveTemp(Queue->PendingTimestampQueries    );
 							Payload->OcclusionQueries     = MoveTemp(Queue->PendingOcclusionQueries    );
 							Payload->PipelineStatsQueries = MoveTemp(Queue->PendingPipelineStatsQueries);
diff --git a/Engine/Source/Runtime/D3D12RHI/Private/D3D12Submission.h b/Engine/Source/Runtime/D3D12RHI/Private/D3D12Submission.h
--- a/Engine/Source/Runtime/D3D12RHI/Private/D3D12Submission.h
+++ b/Engine/Source/Runtime/D3D12RHI/Private/D3D12Submission.h
@@ -191,7 +191,26 @@ struct FD3D12QueryRange
 	TRefCountPtr<FD3D12QueryHeap> Heap;
 	uint32 Start = 0, End = 0;
 
+	FD3D12QueryRange() = default;
+	FD3D12QueryRange(FD3D12QueryHeap* Heap, uint32 Start, uint32 End)
+		: Heap(Heap)
+		, Start(Start)
+		, End(End)
+	{}
+
 	inline bool IsFull() const;
+
+	bool operator == (FD3D12QueryRange const& RHS) const
+	{
+		return Heap  == RHS.Heap
+			&& Start == RHS.Start
+			&& End   == RHS.End;
+	}
+
+	bool operator < (FD3D12QueryRange const& RHS) const
+	{
+		return Start < RHS.Start;
+	}
 };
 
 struct FD3D12CommitReservedResourceDesc
