Support EAccelerationStructureBuildMode::Update for FRHIRayTracingScene.

diff --git a/Engine/Source/Runtime/D3D12RHI/Private/D3D12RayTracing.cpp b/Engine/Source/Runtime/D3D12RHI/Private/D3D12RayTracing.cpp
--- a/Engine/Source/Runtime/D3D12RHI/Private/D3D12RayTracing.cpp
+++ b/Engine/Source/Runtime/D3D12RHI/Private/D3D12RayTracing.cpp
@@ -99,6 +99,7 @@ DECLARE_DWORD_ACCUMULATOR_STAT(TEXT("Triangles in all BL acceleration structures
 DECLARE_DWORD_COUNTER_STAT(TEXT("Built BL AS (per frame)"), STAT_D3D12RayTracingBuiltBLAS, STATGROUP_D3D12RayTracing);
 DECLARE_DWORD_COUNTER_STAT(TEXT("Updated BL AS (per frame)"), STAT_D3D12RayTracingUpdatedBLAS, STATGROUP_D3D12RayTracing);
 DECLARE_DWORD_COUNTER_STAT(TEXT("Built TL AS (per frame)"), STAT_D3D12RayTracingBuiltTLAS, STATGROUP_D3D12RayTracing);
+DECLARE_DWORD_COUNTER_STAT(TEXT("Updated TL AS (per frame)"), STAT_D3D12RayTracingUpdatedTLAS, STATGROUP_D3D12RayTracing);
 
 DECLARE_MEMORY_STAT(TEXT("Total BL AS Memory"), STAT_D3D12RayTracingBLASMemory, STATGROUP_D3D12RayTracing);
 DECLARE_MEMORY_STAT(TEXT("Static BL AS Memory"), STAT_D3D12RayTracingStaticBLASMemory, STATGROUP_D3D12RayTracing);
@@ -3253,7 +3254,7 @@ void FD3D12RayTracingGeometry::CreateAccelerationStructureBuildDesc(FD3D12Comman
 		}
 	}
 
-	D3D12_RAYTRACING_ACCELERATION_STRUCTURE_BUILD_FLAGS LocalBuildFlags = D3D12_RAYTRACING_ACCELERATION_STRUCTURE_BUILD_FLAGS(BuildFlags);
+	D3D12_RAYTRACING_ACCELERATION_STRUCTURE_BUILD_FLAGS LocalBuildFlags = TranslateRayTracingAccelerationStructureFlags(BuildFlags);
 
 	if (bIsUpdate)
 	{
@@ -3376,10 +3377,12 @@ FD3D12RayTracingScene::FD3D12RayTracingScene(FD3D12Adapter* Adapter, FRayTracing
 		Layer.SizeInfo = RHICalcRayTracingSceneSize(Layer.BuildInputs.NumDescs, Initializer.BuildFlags);
 
 		Layer.BufferOffset = Align(SizeInfo.ResultSize, GRHIRayTracingAccelerationStructureAlignment);
-		Layer.ScratchBufferOffset = Align(SizeInfo.BuildScratchSize, GRHIRayTracingScratchBufferAlignment);
+		Layer.BuildScratchOffset = Align(SizeInfo.BuildScratchSize, GRHIRayTracingScratchBufferAlignment);
+		Layer.UpdateScratchOffset = Align(SizeInfo.UpdateScratchSize, GRHIRayTracingScratchBufferAlignment);
 
 		SizeInfo.ResultSize = Layer.BufferOffset + Layer.SizeInfo.ResultSize;
-		SizeInfo.BuildScratchSize = Layer.ScratchBufferOffset + Layer.SizeInfo.BuildScratchSize;
+		SizeInfo.BuildScratchSize = Layer.BuildScratchOffset + Layer.SizeInfo.BuildScratchSize;
+		SizeInfo.UpdateScratchSize = Layer.UpdateScratchOffset + Layer.SizeInfo.UpdateScratchSize;
 	}
 };
 
@@ -3440,21 +3443,26 @@ void FD3D12RayTracingScene::BindBuffer(FRHIBuffer* InBuffer, uint32 InBufferOffs
 
 void FD3D12RayTracingScene::BuildAccelerationStructure(FD3D12CommandContext& CommandContext,
 	FD3D12Buffer* ScratchBuffer, uint32 ScratchBufferOffset,
-	FD3D12Buffer* InstanceBuffer, uint32 InstanceBufferOffset)
+	FD3D12Buffer* InstanceBuffer, uint32 InstanceBufferOffset,
+	EAccelerationStructureBuildMode BuildMode)
 {
 	check(InstanceBuffer != nullptr);
 
 	TRACE_CPUPROFILER_EVENT_SCOPE(BuildAccelerationStructure_TopLevel);
 	SCOPE_CYCLE_COUNTER(STAT_D3D12BuildTLAS);
 
+	const bool bIsUpdate = BuildMode == EAccelerationStructureBuildMode::Update;
+
 	const uint32 GPUIndex = CommandContext.GetGPUIndex();
 	FD3D12Adapter* Adapter = CommandContext.GetParentAdapter();
 
 	TRefCountPtr<FD3D12Buffer> AutoScratchBuffer;
 	if (ScratchBuffer == nullptr)
 	{
+		const uint64 ScratchBufferSize = bIsUpdate ? SizeInfo.UpdateScratchSize : SizeInfo.BuildScratchSize;
+
 		static const FName ScratchBufferName("AutoBuildScratchTLAS");
-		AutoScratchBuffer = CreateRayTracingBuffer(Adapter, GPUIndex, SizeInfo.BuildScratchSize, ERayTracingBufferType::Scratch, ScratchBufferName);
+		AutoScratchBuffer = CreateRayTracingBuffer(Adapter, GPUIndex, ScratchBufferSize, ERayTracingBufferType::Scratch, ScratchBufferName);
 		ScratchBuffer = AutoScratchBuffer.GetReference();
 		ScratchBufferOffset = 0;
 	}
@@ -3479,12 +3487,28 @@ void FD3D12RayTracingScene::BuildAccelerationStructure(FD3D12CommandContext& Com
 			TEXT("TLAS build scratch buffer now requires %lld bytes, but only %lld was calculated in the constructor."),
 			PrebuildInfo.ScratchDataSizeInBytes, Layer.SizeInfo.BuildScratchSize);
 
-		checkf(PrebuildInfo.ScratchDataSizeInBytes + ScratchBufferOffset <= ScratchBuffer->GetSize(),
-			TEXT("TLAS scratch buffer size is %d bytes with offset %d (%d bytes available), but the build requires %lld bytes. ")
-			TEXT("BuildInputs.NumDescs = %d, Instances.Num = %d, PerInstanceGeometries.Num = %d."),
-			ScratchBuffer->GetSize(), ScratchBufferOffset, ScratchBuffer->GetSize() - ScratchBufferOffset,
-			PrebuildInfo.ScratchDataSizeInBytes,
-			Layer.BuildInputs.NumDescs, Initializer.NumNativeInstancesPerLayer[LayerIndex], Initializer.PerInstanceGeometries.Num());
+		checkf(PrebuildInfo.UpdateScratchDataSizeInBytes <= Layer.SizeInfo.UpdateScratchSize,
+			TEXT("TLAS update scratch buffer now requires %lld bytes, but only %lld was calculated in the constructor."),
+			PrebuildInfo.UpdateScratchDataSizeInBytes, Layer.SizeInfo.UpdateScratchSize);
+
+		if (bIsUpdate)
+		{
+			checkf(ScratchBufferOffset + Layer.UpdateScratchOffset + PrebuildInfo.UpdateScratchDataSizeInBytes <= ScratchBuffer->GetSize(),
+				TEXT("TLAS scratch buffer size is %d bytes with offset %d (%d bytes available), but the update requires %lld bytes. ")
+				TEXT("BuildInputs.NumDescs = %d, Instances.Num = %d, PerInstanceGeometries.Num = %d."),
+				ScratchBuffer->GetSize(), ScratchBufferOffset + Layer.UpdateScratchOffset, ScratchBuffer->GetSize() - ScratchBufferOffset,
+				PrebuildInfo.UpdateScratchDataSizeInBytes,
+				Layer.BuildInputs.NumDescs, Initializer.NumNativeInstancesPerLayer[LayerIndex], Initializer.PerInstanceGeometries.Num());
+		}
+		else
+		{
+			checkf(ScratchBufferOffset + Layer.BuildScratchOffset + PrebuildInfo.ScratchDataSizeInBytes <= ScratchBuffer->GetSize(),
+				TEXT("TLAS scratch buffer size is %d bytes with offset %d (%d bytes available), but the build requires %lld bytes. ")
+				TEXT("BuildInputs.NumDescs = %d, Instances.Num = %d, PerInstanceGeometries.Num = %d."),
+				ScratchBuffer->GetSize(), ScratchBufferOffset + Layer.BuildScratchOffset, ScratchBuffer->GetSize() - ScratchBufferOffset,
+				PrebuildInfo.ScratchDataSizeInBytes,
+				Layer.BuildInputs.NumDescs, Initializer.NumNativeInstancesPerLayer[LayerIndex], Initializer.PerInstanceGeometries.Num());
+		}
 	}
 
 	{
@@ -3611,12 +3635,21 @@ void FD3D12RayTracingScene::BuildAccelerationStructure(FD3D12CommandContext& Com
 
 	uint32 InstanceBaseOffset = 0;
 
-	for(uint32 LayerIndex = 0; LayerIndex < NumLayers; ++LayerIndex)
+	for (uint32 LayerIndex = 0; LayerIndex < NumLayers; ++LayerIndex)
 	{
 		const FLayerData& Layer = Layers[LayerIndex];
 
 		const D3D12_GPU_VIRTUAL_ADDRESS BufferAddress = AccelerationStructureBuffer->ResourceLocation.GetGPUVirtualAddress() + BufferOffset + Layer.BufferOffset;
-		const D3D12_GPU_VIRTUAL_ADDRESS ScratchAddress = ScratchBuffer->ResourceLocation.GetGPUVirtualAddress() + ScratchBufferOffset + Layer.ScratchBufferOffset;
+		D3D12_GPU_VIRTUAL_ADDRESS ScratchAddress = ScratchBuffer->ResourceLocation.GetGPUVirtualAddress() + ScratchBufferOffset;
+
+		if (bIsUpdate)
+		{
+			ScratchAddress += Layer.UpdateScratchOffset;
+		}
+		else
+		{
+			ScratchAddress += Layer.BuildScratchOffset;
+		}
 
 		checkf(BufferAddress % GRHIRayTracingAccelerationStructureAlignment == 0,
 			TEXT("TLAS  buffer (plus offset) must be aligned to %lld bytes."),
@@ -3629,11 +3662,27 @@ void FD3D12RayTracingScene::BuildAccelerationStructure(FD3D12CommandContext& Com
 		D3D12_BUILD_RAYTRACING_ACCELERATION_STRUCTURE_DESC& BuildDesc = BuildDescs.AddDefaulted_GetRef();
 		BuildDesc.Inputs = Layer.BuildInputs;
 		BuildDesc.Inputs.InstanceDescs = InstanceBuffer->ResourceLocation.GetGPUVirtualAddress() + InstanceBufferOffset + InstanceBaseOffset * GRHIRayTracingInstanceDescriptorSize;
+
+		if (bIsUpdate)
+		{
+			checkf(EnumHasAllFlags(Initializer.BuildFlags, ERayTracingAccelerationStructureFlags::AllowUpdate),
+				TEXT("Acceleration structure must be created with FRayTracingGeometryInitializer::bAllowUpdate=true to perform refit / update."));
+
+			BuildDesc.Inputs.Flags |= D3D12_RAYTRACING_ACCELERATION_STRUCTURE_BUILD_FLAG_PERFORM_UPDATE;
+		}
+
 		BuildDesc.DestAccelerationStructureData = BufferAddress;
 		BuildDesc.ScratchAccelerationStructureData = ScratchAddress;
-		BuildDesc.SourceAccelerationStructureData = D3D12_GPU_VIRTUAL_ADDRESS(0); // Null source TLAS as this is a build command
+		BuildDesc.SourceAccelerationStructureData = bIsUpdate ? BufferAddress : D3D12_GPU_VIRTUAL_ADDRESS(0);
 
-		INC_DWORD_STAT(STAT_D3D12RayTracingBuiltTLAS);
+		if (bIsUpdate)
+		{
+			INC_DWORD_STAT(STAT_D3D12RayTracingUpdatedTLAS);
+		}
+		else
+		{
+			INC_DWORD_STAT(STAT_D3D12RayTracingBuiltTLAS);
+		}
 
 		InstanceBaseOffset += Initializer.NumNativeInstancesPerLayer[LayerIndex];
 	}
@@ -3918,7 +3967,8 @@ void FD3D12CommandContext::RHIBuildAccelerationStructure(const FRayTracingSceneB
 	Scene->BuildAccelerationStructure(
 		*this,
 		ScratchBuffer, SceneBuildParams.ScratchBufferOffset,
-		InstanceBuffer, SceneBuildParams.InstanceBufferOffset
+		InstanceBuffer, SceneBuildParams.InstanceBufferOffset,
+		SceneBuildParams.BuildMode
 	);
 }
 
diff --git a/Engine/Source/Runtime/D3D12RHI/Private/D3D12RayTracing.h b/Engine/Source/Runtime/D3D12RHI/Private/D3D12RayTracing.h
--- a/Engine/Source/Runtime/D3D12RHI/Private/D3D12RayTracing.h
+++ b/Engine/Source/Runtime/D3D12RHI/Private/D3D12RayTracing.h
@@ -112,7 +112,8 @@ public:
 
 	void BuildAccelerationStructure(FD3D12CommandContext& CommandContext,
 		FD3D12Buffer* ScratchBuffer, uint32 ScratchBufferOffset,
-		FD3D12Buffer* InstanceBuffer, uint32 InstanceBufferOffset
+		FD3D12Buffer* InstanceBuffer, uint32 InstanceBufferOffset,
+		EAccelerationStructureBuildMode BuildMode
 	);
 
 	struct FLayerData
@@ -120,7 +121,8 @@ public:
 		D3D12_BUILD_RAYTRACING_ACCELERATION_STRUCTURE_INPUTS BuildInputs = {};
 		FRayTracingAccelerationStructureSize SizeInfo = {};
 		uint32 BufferOffset;
-		uint32 ScratchBufferOffset;
+		uint32 BuildScratchOffset;
+		uint32 UpdateScratchOffset;
 	};
 
 	TArray<FLayerData> Layers;
diff --git a/Engine/Source/Runtime/RHI/Public/RHIContext.h b/Engine/Source/Runtime/RHI/Public/RHIContext.h
--- a/Engine/Source/Runtime/RHI/Public/RHIContext.h
+++ b/Engine/Source/Runtime/RHI/Public/RHIContext.h
@@ -565,6 +565,8 @@ struct FRayTracingSceneBuildParams
 	// Buffer of native ray tracing instance descriptors. Must be in SRV state.
 	FRHIBuffer* InstanceBuffer = nullptr;
 	uint32 InstanceBufferOffset = 0;
+
+	EAccelerationStructureBuildMode BuildMode = EAccelerationStructureBuildMode::Build;
 };
 
 struct FCopyBufferRegionParams
diff --git a/Engine/Source/Runtime/VulkanRHI/Private/VulkanRayTracing.cpp b/Engine/Source/Runtime/VulkanRHI/Private/VulkanRayTracing.cpp
--- a/Engine/Source/Runtime/VulkanRHI/Private/VulkanRayTracing.cpp
+++ b/Engine/Source/Runtime/VulkanRHI/Private/VulkanRayTracing.cpp
@@ -63,6 +63,7 @@ DECLARE_DWORD_ACCUMULATOR_STAT(TEXT("Triangles in all BL acceleration structures
 DECLARE_DWORD_COUNTER_STAT(TEXT("Built BL AS (per frame)"), STAT_VulkanRayTracingBuiltBLAS, STATGROUP_VulkanRayTracing);
 DECLARE_DWORD_COUNTER_STAT(TEXT("Updated BL AS (per frame)"), STAT_VulkanRayTracingUpdatedBLAS, STATGROUP_VulkanRayTracing);
 DECLARE_DWORD_COUNTER_STAT(TEXT("Built TL AS (per frame)"), STAT_VulkanRayTracingBuiltTLAS, STATGROUP_VulkanRayTracing);
+DECLARE_DWORD_COUNTER_STAT(TEXT("Updated TL AS (per frame)"), STAT_VulkanRayTracingUpdatedTLAS, STATGROUP_VulkanRayTracing);
 
 DECLARE_MEMORY_STAT(TEXT("Total BL AS Memory"), STAT_VulkanRayTracingBLASMemory, STATGROUP_VulkanRayTracing);
 DECLARE_MEMORY_STAT(TEXT("Static BL AS Memory"), STAT_VulkanRayTracingStaticBLASMemory, STATGROUP_VulkanRayTracing);
@@ -591,6 +592,7 @@ static void GetTLASBuildData(
 	const uint32 NumInstances,
 	const VkDeviceAddress InstanceBufferAddress,
 	ERayTracingAccelerationStructureFlags BuildFlags,
+	EAccelerationStructureBuildMode BuildMode,
 	FVkRtTLASBuildData& BuildData)
 {
 	VkDeviceOrHostAddressConstKHR InstanceBufferDeviceAddress = {};
@@ -602,7 +604,7 @@ static void GetTLASBuildData(
 	BuildData.Geometry.geometry.instances.data = InstanceBufferDeviceAddress;
 
 	BuildData.GeometryInfo.type = VK_ACCELERATION_STRUCTURE_TYPE_TOP_LEVEL_KHR;
-	BuildData.GeometryInfo.mode = VK_BUILD_ACCELERATION_STRUCTURE_MODE_BUILD_KHR;
+	BuildData.GeometryInfo.mode = BuildMode == EAccelerationStructureBuildMode::Build ? VK_BUILD_ACCELERATION_STRUCTURE_MODE_BUILD_KHR : VK_BUILD_ACCELERATION_STRUCTURE_MODE_UPDATE_KHR;
 	BuildData.GeometryInfo.flags = TranslateRayTracingAccelerationStructureFlags(BuildFlags);
 	BuildData.GeometryInfo.geometryCount = 1;
 	BuildData.GeometryInfo.pGeometries = &BuildData.Geometry;
@@ -683,10 +685,12 @@ FVulkanRayTracingScene::FVulkanRayTracingScene(FRayTracingSceneInitializer2 InIn
 
 		Layer.SizeInfo = RHICalcRayTracingSceneSize(Initializer.NumNativeInstancesPerLayer[LayerIndex], Initializer.BuildFlags);
 		Layer.BufferOffset = Align(SizeInfo.ResultSize, GRHIRayTracingAccelerationStructureAlignment);
-		Layer.ScratchBufferOffset = Align(SizeInfo.BuildScratchSize, GRHIRayTracingScratchBufferAlignment);
+		Layer.BuildScratchOffset = Align(SizeInfo.BuildScratchSize, GRHIRayTracingScratchBufferAlignment);
+		Layer.UpdateScratchOffset = Align(SizeInfo.UpdateScratchSize, GRHIRayTracingScratchBufferAlignment);
 
 		SizeInfo.ResultSize = Layer.BufferOffset + Layer.SizeInfo.ResultSize;
-		SizeInfo.BuildScratchSize = Layer.ScratchBufferOffset + Layer.SizeInfo.BuildScratchSize;
+		SizeInfo.BuildScratchSize = Layer.BuildScratchOffset + Layer.SizeInfo.BuildScratchSize;
+		SizeInfo.UpdateScratchSize = Layer.UpdateScratchOffset + Layer.SizeInfo.UpdateScratchSize;
 	}
 
 	const uint32 ParameterBufferSize = FMath::Max<uint32>(1, Initializer.NumTotalSegments) * sizeof(FVulkanRayTracingGeometryParameters);
@@ -753,8 +757,11 @@ void FVulkanRayTracingScene::BindBuffer(FRHIBuffer* InBuffer, uint32 InBufferOff
 void FVulkanRayTracingScene::BuildAccelerationStructure(
 	FVulkanCommandListContext& CommandContext,
 	FVulkanResourceMultiBuffer* InScratchBuffer, uint32 InScratchOffset,
-	FVulkanResourceMultiBuffer* InInstanceBuffer, uint32 InInstanceOffset)
+	FVulkanResourceMultiBuffer* InInstanceBuffer, uint32 InInstanceOffset,
+	EAccelerationStructureBuildMode BuildMode)
 {
+	const bool bIsUpdate = BuildMode == EAccelerationStructureBuildMode::Update;
+
 	// Build a metadata buffer	that contains VulkanRHI-specific per-geometry parameters that allow us to access
 	// vertex and index buffers from shaders that use inline ray tracing.
 	BuildPerInstanceGeometryParameterBuffer(CommandContext);
@@ -766,9 +773,11 @@ void FVulkanRayTracingScene::BuildAccelerationStructure(
 
 	if (InScratchBuffer == nullptr)
 	{
+		const uint64 ScratchBufferSize = bIsUpdate ? SizeInfo.UpdateScratchSize : SizeInfo.BuildScratchSize;
+
 		TRHICommandList_RecursiveHazardous<FVulkanCommandListContext> RHICmdList(&CommandContext);
 		FRHIResourceCreateInfo ScratchBufferCreateInfo(TEXT("BuildScratchTLAS"));
-		ScratchBuffer = RHICmdList.CreateBuffer(SizeInfo.BuildScratchSize, BUF_StructuredBuffer | BUF_RayTracingScratch, 0, ERHIAccess::UAVCompute, ScratchBufferCreateInfo);
+		ScratchBuffer = RHICmdList.CreateBuffer(ScratchBufferSize, BUF_StructuredBuffer | BUF_RayTracingScratch, 0, ERHIAccess::UAVCompute, ScratchBufferCreateInfo);
 		InScratchBuffer = ResourceCast(ScratchBuffer.GetReference());
 		InScratchOffset = 0;
 	}
@@ -795,11 +804,14 @@ void FVulkanRayTracingScene::BuildAccelerationStructure(
 		const FLayerData& Layer = Layers[LayerIndex];
 
 		FVkRtTLASBuildData& BuildData = BuildDatas[LayerIndex];
-		GetTLASBuildData(Device->GetInstanceHandle(), Initializer.NumNativeInstancesPerLayer[LayerIndex], InstanceBufferAddress, Initializer.BuildFlags, BuildData);
+		GetTLASBuildData(Device->GetInstanceHandle(), Initializer.NumNativeInstancesPerLayer[LayerIndex], InstanceBufferAddress, Initializer.BuildFlags, BuildMode, BuildData);
+
+		const uint64 LayerScratchOffset = bIsUpdate ? Layer.UpdateScratchOffset : Layer.BuildScratchOffset;
 
 		checkf(Layer.View.IsValid(), TEXT("A buffer must be bound to the ray tracing scene before it can be built."));
 		BuildData.GeometryInfo.dstAccelerationStructure = Layer.View->GetAccelerationStructureView().Handle;
-		BuildData.GeometryInfo.scratchData.deviceAddress = InScratchBuffer->GetDeviceAddress() + InScratchOffset + Layer.ScratchBufferOffset;
+		BuildData.GeometryInfo.srcAccelerationStructure = bIsUpdate ? Layer.View->GetAccelerationStructureView().Handle : nullptr;
+		BuildData.GeometryInfo.scratchData.deviceAddress = InScratchBuffer->GetDeviceAddress() + InScratchOffset + LayerScratchOffset;
 
 		GeometryInfos[LayerIndex] = BuildData.GeometryInfo;
 
@@ -811,7 +823,14 @@ void FVulkanRayTracingScene::BuildAccelerationStructure(
 
 		pBuildRanges[LayerIndex] = &TLASBuildRangeInfo;
 
-		INC_DWORD_STAT(STAT_VulkanRayTracingBuiltTLAS);
+		if (bIsUpdate)
+		{
+			INC_DWORD_STAT(STAT_VulkanRayTracingUpdatedTLAS);
+		}
+		else
+		{
+			INC_DWORD_STAT(STAT_VulkanRayTracingBuiltTLAS);
+		}
 
 		InstanceBaseOffset += Initializer.NumNativeInstancesPerLayer[LayerIndex];
 	}
@@ -1090,7 +1109,7 @@ FRayTracingAccelerationStructureSize FVulkanDynamicRHI::RHICalcRayTracingSceneSi
 {
 	FVkRtTLASBuildData BuildData;
 	const VkDeviceAddress InstanceBufferAddress = 0; // No device address available when only querying TLAS size
-	GetTLASBuildData(Device->GetInstanceHandle(), MaxInstances, InstanceBufferAddress, Flags, BuildData);
+	GetTLASBuildData(Device->GetInstanceHandle(), MaxInstances, InstanceBufferAddress, Flags, EAccelerationStructureBuildMode::Build, BuildData);
 
 	FRayTracingAccelerationStructureSize Result;
 	Result.ResultSize = BuildData.SizesInfo.accelerationStructureSize;
@@ -1280,7 +1299,8 @@ void FVulkanCommandListContext::RHIBuildAccelerationStructure(const FRayTracingS
 	Scene->BuildAccelerationStructure(
 		*this, 
 		ScratchBuffer, SceneBuildParams.ScratchBufferOffset, 
-		InstanceBuffer, SceneBuildParams.InstanceBufferOffset);
+		InstanceBuffer, SceneBuildParams.InstanceBufferOffset,
+		SceneBuildParams.BuildMode);
 }
 
 template<typename ShaderType>
diff --git a/Engine/Source/Runtime/VulkanRHI/Private/VulkanRayTracing.h b/Engine/Source/Runtime/VulkanRHI/Private/VulkanRayTracing.h
--- a/Engine/Source/Runtime/VulkanRHI/Private/VulkanRayTracing.h
+++ b/Engine/Source/Runtime/VulkanRHI/Private/VulkanRayTracing.h
@@ -186,7 +186,8 @@ public:
 	void BuildAccelerationStructure(
 		FVulkanCommandListContext& CommandContext, 
 		FVulkanResourceMultiBuffer* ScratchBuffer, uint32 ScratchOffset, 
-		FVulkanResourceMultiBuffer* InstanceBuffer, uint32 InstanceOffset);
+		FVulkanResourceMultiBuffer* InstanceBuffer, uint32 InstanceOffset,
+		EAccelerationStructureBuildMode BuildMode);
 
 	virtual FRHIShaderResourceView* GetOrCreateMetadataBufferSRV(FRHICommandListImmediate& RHICmdList) override final
 	{
@@ -226,7 +227,8 @@ private:
 		TUniquePtr<FVulkanView> View;
 		FRayTracingAccelerationStructureSize SizeInfo;
 		uint32 BufferOffset;
-		uint32 ScratchBufferOffset;
+		uint32 BuildScratchOffset;
+		uint32 UpdateScratchOffset;
 	};
 
 	TArray<FLayerData> Layers;
