Adding missing thread fence when locking dynamic d3d12 buffers for write (internal resouce location is changed with new upload memory data)

diff --git a/Engine/Source/Runtime/D3D12RHI/Private/D3D12Buffer.cpp b/Engine/Source/Runtime/D3D12RHI/Private/D3D12Buffer.cpp
--- a/Engine/Source/Runtime/D3D12RHI/Private/D3D12Buffer.cpp
+++ b/Engine/Source/Runtime/D3D12RHI/Private/D3D12Buffer.cpp
@@ -534,6 +534,8 @@ void* FD3D12DynamicRHI::LockBuffer(FRHICommandListBase& RHICmdList, FD3D12Buffer
 #endif
 				Resource->RenameLDAChain(Contexts, NewLocation);
 			});
+
+			RHICmdList.RHIThreadFence(true);
 		}
 	}
 	else
