Track reserved buffer and texture memory separately, both uncommitted and committed.

Also fix D3D12 physical memory tracking for heaps used as backing heaps for reserved resources. This memory wasn't accounted for when releasing a reserved resource.

#tests tested locally

diff --git a/Engine/Source/Runtime/D3D12RHI/Private/D3D12Resources.cpp b/Engine/Source/Runtime/D3D12RHI/Private/D3D12Resources.cpp
--- a/Engine/Source/Runtime/D3D12RHI/Private/D3D12Resources.cpp
+++ b/Engine/Source/Runtime/D3D12RHI/Private/D3D12Resources.cpp
@@ -12,6 +12,7 @@ D3D12Resources.cpp: D3D RHI utility implementation.
 #include "HAL/LowLevelMemTracker.h"
 #include "ProfilingDebugging/MemoryTrace.h"
 #include "ProfilingDebugging/AssetMetadataTrace.h"
+#include "RHICoreStats.h"
 
 static TAutoConsoleVariable<int32> CVarD3D12ReservedResourceHeapSizeMB(
 	TEXT("d3d12.ReservedResourceHeapSizeMB"),
@@ -161,6 +162,20 @@ FD3D12Resource::~FD3D12Resource()
 		FScopeLock Lock(&FD3D12Viewport::DXGIBackBufferLock);
 		Resource.SafeRelease();
 	}
+
+	// Update reserved resources' physical memory stats.
+	if (ReservedResourceData.IsValid() && ReservedResourceData->NumCommittedTiles > 0)
+	{
+		bool bBuffer = Desc.Dimension == D3D12_RESOURCE_DIMENSION_BUFFER;
+		const uint64 DecommitBytes = GRHIGlobals.ReservedResources.TileSizeInBytes * ReservedResourceData->NumCommittedTiles;
+		UE::RHICore::UpdateReservedResourceStatsOnCommit(DecommitBytes, bBuffer, false /* Decommit */);
+		
+		// The backing heaps are going to be released once this resource is destroyed.
+		for (const TRefCountPtr<FD3D12Heap>& BackingHeap : ReservedResourceData->BackingHeaps)
+		{
+			DEC_MEMORY_STAT_BY(STAT_D3D12ReservedResourcePhysical, BackingHeap->GetHeapDesc().SizeInBytes);
+		}
+	}
 }
 
 struct FD3D12UpdateTileMappingsParams
@@ -341,6 +356,7 @@ void FD3D12Resource::CommitReservedResource(ID3D12CommandQueue* D3DCommandQueue,
 	TArray<FD3D12UpdateTileMappingsParams> MappingParams;
 	TArray<FD3D12ResidencyHandle*> UsedResidencyHandles;
 
+	const uint32 NumPreviousCommittedTiles = ReservedResourceData->NumCommittedTiles;
 	if (ReservedResourceData->NumCommittedTiles > NumRequiredCommitTiles) // Decommit / shrink case
 	{
 		check(!ReservedResourceData->BackingHeaps.IsEmpty());
@@ -548,6 +564,12 @@ void FD3D12Resource::CommitReservedResource(ID3D12CommandQueue* D3DCommandQueue,
 	checkf(ReservedResourceData->NumCommittedTiles == NumRequiredCommitTiles,
 		TEXT("Reserved resource was not fully processed while committing physical memory. Expected to process tiles: %d, actually processed: %d"),
 		D3DResourceNumTiles, ReservedResourceData->NumCommittedTiles);
+
+	if (ReservedResourceData->NumCommittedTiles != NumPreviousCommittedTiles)
+	{
+		int64 CommitDeltaInBytes = TileSizeInBytes * FMath::Abs((int32)ReservedResourceData->NumCommittedTiles - (int32)NumPreviousCommittedTiles);
+		UE::RHICore::UpdateReservedResourceStatsOnCommit(CommitDeltaInBytes, bBuffer, ReservedResourceData->NumCommittedTiles > NumPreviousCommittedTiles);
+	}
 }
 
 ID3D12Pageable* FD3D12Resource::GetPageable()
diff --git a/Engine/Source/Runtime/RHI/Private/RHIStats.cpp b/Engine/Source/Runtime/RHI/Private/RHIStats.cpp
--- a/Engine/Source/Runtime/RHI/Private/RHIStats.cpp
+++ b/Engine/Source/Runtime/RHI/Private/RHIStats.cpp
@@ -33,6 +33,11 @@ DEFINE_STAT(STAT_ByteAddressBufferMemory);
 DEFINE_STAT(STAT_DrawIndirectBufferMemory);
 DEFINE_STAT(STAT_MiscBufferMemory);
 
+DEFINE_STAT(STAT_ReservedUncommittedBufferMemory);
+DEFINE_STAT(STAT_ReservedCommittedBufferMemory);
+DEFINE_STAT(STAT_ReservedUncommittedTextureMemory);
+DEFINE_STAT(STAT_ReservedCommittedTextureMemory);
+
 DEFINE_STAT(STAT_SamplerDescriptorsAllocated);
 DEFINE_STAT(STAT_ResourceDescriptorsAllocated);
 
diff --git a/Engine/Source/Runtime/RHI/Public/RHIStats.h b/Engine/Source/Runtime/RHI/Public/RHIStats.h
--- a/Engine/Source/Runtime/RHI/Public/RHIStats.h
+++ b/Engine/Source/Runtime/RHI/Public/RHIStats.h
@@ -231,6 +231,11 @@ DECLARE_MEMORY_STAT_POOL_EXTERN(TEXT("Byte Address Buffer Memory"), STAT_ByteAdd
 DECLARE_MEMORY_STAT_POOL_EXTERN(TEXT("Draw Indirect Buffer Memory"), STAT_DrawIndirectBufferMemory, STATGROUP_RHI, FPlatformMemory::MCR_GPU, RHI_API);
 DECLARE_MEMORY_STAT_POOL_EXTERN(TEXT("Misc Buffer Memory"), STAT_MiscBufferMemory, STATGROUP_RHI, FPlatformMemory::MCR_GPU, RHI_API);
 
+DECLARE_MEMORY_STAT_POOL_EXTERN(TEXT("Reserved Buffer Memory (Uncommitted)"), STAT_ReservedUncommittedBufferMemory, STATGROUP_RHI, FPlatformMemory::MCR_GPU, RHI_API);
+DECLARE_MEMORY_STAT_POOL_EXTERN(TEXT("Reserved Buffer Memory (Committed)"), STAT_ReservedCommittedBufferMemory, STATGROUP_RHI, FPlatformMemory::MCR_GPU, RHI_API);
+DECLARE_MEMORY_STAT_POOL_EXTERN(TEXT("Reserved Texture Memory (Uncommitted)"), STAT_ReservedUncommittedTextureMemory, STATGROUP_RHI, FPlatformMemory::MCR_GPU, RHI_API);
+DECLARE_MEMORY_STAT_POOL_EXTERN(TEXT("Reserved Texture Memory (Committed)"), STAT_ReservedCommittedTextureMemory, STATGROUP_RHI, FPlatformMemory::MCR_GPU, RHI_API);
+
 DECLARE_DWORD_ACCUMULATOR_STAT_EXTERN(TEXT("Sampler Descriptors Allocated"), STAT_SamplerDescriptorsAllocated, STATGROUP_RHI, RHI_API);
 DECLARE_DWORD_ACCUMULATOR_STAT_EXTERN(TEXT("Resource Descriptors Allocated"), STAT_ResourceDescriptorsAllocated, STATGROUP_RHI, RHI_API);
 
diff --git a/Engine/Source/Runtime/RHICore/Internal/RHICoreStats.h b/Engine/Source/Runtime/RHICore/Internal/RHICoreStats.h
--- a/Engine/Source/Runtime/RHICore/Internal/RHICoreStats.h
+++ b/Engine/Source/Runtime/RHICore/Internal/RHICoreStats.h
@@ -14,7 +14,11 @@ namespace UE::RHICore
 		constexpr ETextureCreateFlags AllRenderTargetFlags = ETextureCreateFlags::RenderTargetable | ETextureCreateFlags::ResolveTargetable | ETextureCreateFlags::DepthStencilTargetable;
 #if STATS
 		const int64 TextureSizeDeltaInBytes = bAllocating ? static_cast<int64>(TextureSizeInBytes) : -static_cast<int64>(TextureSizeInBytes);
-		if (EnumHasAnyFlags(TextureFlags, AllRenderTargetFlags))
+		if (EnumHasAnyFlags(TextureFlags, ETextureCreateFlags::ReservedResource))
+		{
+			INC_MEMORY_STAT_BY(STAT_ReservedUncommittedTextureMemory, TextureSizeDeltaInBytes);
+		}
+		else if (EnumHasAnyFlags(TextureFlags, AllRenderTargetFlags))
 		{
 			switch (Dimension)
 			{
@@ -99,7 +103,11 @@ namespace UE::RHICore
 	inline void UpdateGlobalBufferStats(const FRHIBufferDesc& BufferDesc, int64 BufferSizeDelta)
 	{
 #if STATS
-		if (EnumHasAnyFlags(BufferDesc.Usage, EBufferUsageFlags::VertexBuffer))
+		if (EnumHasAnyFlags(BufferDesc.Usage, EBufferUsageFlags::ReservedResource))
+		{
+			INC_MEMORY_STAT_BY(STAT_ReservedUncommittedBufferMemory, BufferSizeDelta);
+		}
+		else if (EnumHasAnyFlags(BufferDesc.Usage, EBufferUsageFlags::VertexBuffer))
 		{
 			INC_MEMORY_STAT_BY(STAT_VertexBufferMemory, BufferSizeDelta);
 		}
@@ -145,4 +153,22 @@ namespace UE::RHICore
 		INC_MEMORY_STAT_BY(STAT_UniformBufferMemory, BufferSizeDelta);
 		FPlatformAtomics::InterlockedAdd((volatile int64*)&GRHIGlobals.UniformBufferMemorySize, BufferSizeDelta);
 	}
+
+	inline void UpdateReservedResourceStatsOnCommit(int64 CommitDelta, bool bBuffer, bool bCommitting)
+	{
+#if STATS
+		const int64 CommittedDelta = bCommitting ? static_cast<int64>(CommitDelta) : -static_cast<int64>(CommitDelta);
+
+		if (bBuffer)
+		{
+			INC_MEMORY_STAT_BY(STAT_ReservedCommittedBufferMemory, CommittedDelta);
+			DEC_MEMORY_STAT_BY(STAT_ReservedUncommittedBufferMemory, CommittedDelta);
+		}
+		else
+		{
+			INC_MEMORY_STAT_BY(STAT_ReservedCommittedTextureMemory, CommittedDelta);
+			DEC_MEMORY_STAT_BY(STAT_ReservedUncommittedTextureMemory, CommittedDelta);
+		}
+#endif
+	}
 }
