Fix some LLM accounting issue in the d3d12 RHI on non-Windows platforms. FD3D12PoolAllocator is now tracking allocs and frees through LLM. Buddy and pool allocator slacks are now tracked as D3D12AllocatorUnused. Some small misc fixes

diff --git a/Engine/Source/Runtime/Core/Public/HAL/LowLevelMemTracker.h b/Engine/Source/Runtime/Core/Public/HAL/LowLevelMemTracker.h
--- a/Engine/Source/Runtime/Core/Public/HAL/LowLevelMemTracker.h
+++ b/Engine/Source/Runtime/Core/Public/HAL/LowLevelMemTracker.h
@@ -372,7 +372,7 @@ extern FName LLMGetTagStat(ELLMTag Tag);
 																	FLLMScope SCOPE_NAME(PREPROCESSOR_JOIN(LLMScope_Name,__LINE__), false /* bIsStatTag */, ELLMTagSet::Assets, ELLMTracker::Default, false /* bOverride */);
 #define LLM_PLATFORM_SCOPE(Tag)										FLLMScope SCOPE_NAME(Tag, false /* bIsStatTag */, ELLMTagSet::None, ELLMTracker::Platform);
 #define LLM_PLATFORM_SCOPE_BYNAME(Tag) 								static FName PREPROCESSOR_JOIN(LLMScope_Name,__LINE__)(Tag);\
-																	FLLMScope SCOPE_NAME(PREPROCESSOR_JOIN(LLMLLMScope_NameScope,__LINE__), false /* bIsStatTag */, ELLMTagSet::None, ELLMTracker::Platform);
+																	FLLMScope SCOPE_NAME(PREPROCESSOR_JOIN(LLMScope_Name,__LINE__), false /* bIsStatTag */, ELLMTagSet::None, ELLMTracker::Platform);
 #define LLM_PLATFORM_SCOPE_BYTAG(TagDeclName)						FLLMScope SCOPE_NAME(PREPROCESSOR_JOIN(LLMTagDeclaration_, TagDeclName).GetUniqueName(), false /* bIsStatTag */, ELLMTagSet::None, ELLMTracker::Platform);
 
  /**
diff --git a/Engine/Source/Runtime/D3D12RHI/Private/D3D12Allocation.cpp b/Engine/Source/Runtime/D3D12RHI/Private/D3D12Allocation.cpp
--- a/Engine/Source/Runtime/D3D12RHI/Private/D3D12Allocation.cpp
+++ b/Engine/Source/Runtime/D3D12RHI/Private/D3D12Allocation.cpp
@@ -282,10 +282,15 @@ void FD3D12BuddyAllocator::Initialize()
 		{
 			LLM_PLATFORM_SCOPE(ELLMTag::GraphicsPlatform);
 
-			// we are tracking allocations ourselves, so don't let XMemAlloc track these as well
+#if PLATFORM_WINDOWS
+			// we are tracking allocations ourselves
 			LLM_SCOPED_PAUSE_TRACKING_FOR_TRACKER(ELLMTracker::Default, ELLMAllocType::System);
 			VERIFYD3D12RESULT(Adapter->GetD3DDevice()->CreateHeap(&Desc, IID_PPV_ARGS(&Heap)));
 			LLM_SCOPED_PAUSE_TRACKING_WITH_ENUM_AND_AMOUNT_BYTAG(D3D12AllocatorUnused, MaxBlockSize, ELLMTracker::Platform, ELLMAllocType::System);
+#else
+			LLM_SCOPE_BYTAG(D3D12AllocatorUnused);
+			VERIFYD3D12RESULT(Adapter->GetD3DDevice()->CreateHeap(&Desc, IID_PPV_ARGS(&Heap)));
+#endif
 		}
 
 		BackingHeap = new FD3D12Heap(GetParentDevice(), GetVisibilityMask(), TraceHeapId);
@@ -300,10 +305,15 @@ void FD3D12BuddyAllocator::Initialize()
 	else
 	{
 		{
-			LLM_SCOPED_PAUSE_TRACKING_FOR_TRACKER(ELLMTracker::Default, ELLMAllocType::System);
 			const D3D12_HEAP_PROPERTIES HeapProps = CD3DX12_HEAP_PROPERTIES(InitConfig.HeapType, GetGPUMask().GetNative(), GetVisibilityMask().GetNative());
+#if PLATFORM_WINDOWS
+			LLM_SCOPED_PAUSE_TRACKING_FOR_TRACKER(ELLMTracker::Default, ELLMAllocType::System);
 			VERIFYD3D12RESULT(Adapter->CreateBuffer(HeapProps, GetGPUMask(), InitConfig.InitialResourceState, ED3D12ResourceStateMode::SingleState, InitConfig.InitialResourceState, MaxBlockSize, BackingResource.GetInitReference(), TEXT("Resource Allocator Underlying Buffer"), InitConfig.ResourceFlags));
 			LLM_SCOPED_PAUSE_TRACKING_WITH_ENUM_AND_AMOUNT_BYTAG(D3D12AllocatorUnused, MaxBlockSize, ELLMTracker::Platform, ELLMAllocType::System);
+#else
+			LLM_SCOPE_BYTAG(D3D12AllocatorUnused);
+			VERIFYD3D12RESULT(Adapter->CreateBuffer(HeapProps, GetGPUMask(), InitConfig.InitialResourceState, ED3D12ResourceStateMode::SingleState, InitConfig.InitialResourceState, MaxBlockSize, BackingResource.GetInitReference(), TEXT("Resource Allocator Underlying Buffer"), InitConfig.ResourceFlags));
+#endif
 #if UE_MEMORY_TRACE_ENABLED
 			MemoryTrace_MarkAllocAsHeap(BackingResource->GetGPUVirtualAddress(), TraceHeapId);
 #endif
@@ -420,9 +430,11 @@ void FD3D12BuddyAllocator::Allocate(uint32 SizeInBytes, uint32 Alignment, FD3D12
 	INCREASE_ALLOC_COUNTER(SpaceAlignedUsed, AllocSize);
 	INCREASE_ALLOC_COUNTER(SpaceActualUsed, SizeInBytes);
 	
+#if PLATFORM_WINDOWS
 	// Decrease only texture size so wasted amount stays in D3D12AllocatorUnused
 	LLM_SCOPED_PAUSE_TRACKING_WITH_ENUM_AND_AMOUNT_BYTAG(D3D12AllocatorUnused, 0 - int64(SizeInBytes), ELLMTracker::Platform, ELLMAllocType::System);
 	LLM_SCOPED_PAUSE_TRACKING_WITH_ENUM_AND_AMOUNT_BYTAG(D3D12AllocatorWasted, int64(AllocSize - SizeInBytes), ELLMTracker::Platform, ELLMAllocType::System);
+#endif
 
 	TotalSizeUsed += AllocSize;
 
@@ -462,7 +474,10 @@ void FD3D12BuddyAllocator::Allocate(uint32 SizeInBytes, uint32 Alignment, FD3D12
 
 	// track the allocation
 #if !PLATFORM_WINDOWS
-	LLM_IF_ENABLED(FLowLevelMemTracker::Get().OnLowLevelAlloc(ELLMTracker::Default, ResourceLocation.GetAddressForLLMTracking(), SizeInBytes));
+	{
+		LLM_SCOPED_PAUSE_TRACKING_WITH_ENUM_AND_AMOUNT_BYTAG(D3D12AllocatorUnused, 0 - int64(AllocSize), ELLMTracker::Default, ELLMAllocType::System);
+	}
+	LLM_IF_ENABLED(FLowLevelMemTracker::Get().OnLowLevelAlloc(ELLMTracker::Default, ResourceLocation.GetAddressForLLMTracking(), AllocSize));
 	// Note: Disabling this LLM hook for Windows is due to a work-around in the way that d3d12 buffers are tracked
 	// by LLM. LLM tracks buffer data in the UpdateBufferStats function because that is the easiest place to ensure that LLM
 	// can be updated whenever a buffer is created or released. Unfortunately, some buffers allocate from this allocator
@@ -527,6 +542,10 @@ void FD3D12BuddyAllocator::Deallocate(FD3D12ResourceLocation& ResourceLocation)
 	// This does mean that non-buffer memory that goes through this allocator won't be tracked, so this does need a better solution.
 	// see UpdateBufferStats for a more detailed explanation.
 	LLM_IF_ENABLED(FLowLevelMemTracker::Get().OnLowLevelFree(ELLMTracker::Default, ResourceLocation.GetAddressForLLMTracking()));
+	{
+		const uint32 AllocSize = uint32(OrderToUnitSize(Block.Data.Order) * MinBlockSize);
+		LLM_SCOPED_PAUSE_TRACKING_WITH_ENUM_AND_AMOUNT_BYTAG(D3D12AllocatorUnused, int64(AllocSize), ELLMTracker::Default, ELLMAllocType::System);
+	}
 #if UE_MEMORY_TRACE_ENABLED
 	MemoryTrace_Free(ResourceLocation.GetGPUVirtualAddress(), EMemoryTraceRootHeap::VideoMemory);
 #endif
@@ -540,9 +559,11 @@ void FD3D12BuddyAllocator::DeallocateInternal(RetiredBlock& Block)
 	const uint32 Size = uint32(OrderToUnitSize(Block.Data.Order) * MinBlockSize);
 	DECREASE_ALLOC_COUNTER(SpaceAlignedUsed, Size);
 	DECREASE_ALLOC_COUNTER(SpaceActualUsed, Block.AllocationSize);
+#if PLATFORM_WINDOWS
 	LLM_SCOPED_PAUSE_TRACKING_WITH_ENUM_AND_AMOUNT_BYTAG(D3D12AllocatorUnused, int64(Block.AllocationSize), ELLMTracker::Platform, ELLMAllocType::System);
 	LLM_SCOPED_PAUSE_TRACKING_WITH_ENUM_AND_AMOUNT_BYTAG(D3D12AllocatorWasted, 0 - int64(Size - Block.AllocationSize), ELLMTracker::Platform, ELLMAllocType::System);
-
+#endif
+	
 	TotalSizeUsed -= Size;
 
 	if (AllocationStrategy == EResourceAllocationStrategy::kPlacedResource)
@@ -587,8 +608,10 @@ void FD3D12BuddyAllocator::CleanUpAllocations()
 
 void FD3D12BuddyAllocator::ReleaseAllResources()
 {
+#if PLATFORM_WINDOWS
 	LLM_SCOPED_PAUSE_TRACKING_FOR_TRACKER(ELLMTracker::Default, ELLMAllocType::System);
 	LLM_SCOPED_PAUSE_TRACKING_WITH_ENUM_AND_AMOUNT_BYTAG(D3D12AllocatorUnused, 0 - int64(MaxBlockSize), ELLMTracker::Platform, ELLMAllocType::System);
+#endif
 
 #if UE_MEMORY_TRACE_ENABLED
 	if (AllocationStrategy != EResourceAllocationStrategy::kPlacedResource)
diff --git a/Engine/Source/Runtime/D3D12RHI/Private/D3D12PoolAllocator.cpp b/Engine/Source/Runtime/D3D12RHI/Private/D3D12PoolAllocator.cpp
--- a/Engine/Source/Runtime/D3D12RHI/Private/D3D12PoolAllocator.cpp
+++ b/Engine/Source/Runtime/D3D12RHI/Private/D3D12PoolAllocator.cpp
@@ -88,8 +88,12 @@ void FD3D12MemoryPool::Init()
 		{
 			LLM_PLATFORM_SCOPE(ELLMTag::GraphicsPlatform);
 
-			// we are tracking allocations ourselves, so don't let XMemAlloc track these as well
+#if PLATFORM_WINDOWS
+			// we are tracking allocations ourselves
 			LLM_SCOPED_PAUSE_TRACKING_FOR_TRACKER(ELLMTracker::Default, ELLMAllocType::System);
+#else
+			LLM_SCOPE_BYTAG(D3D12AllocatorUnused);
+#endif
 			VERIFYD3D12RESULT(Adapter->GetD3DDevice()->CreateHeap(&Desc, IID_PPV_ARGS(&Heap)));
 		}
 
@@ -105,7 +109,11 @@ void FD3D12MemoryPool::Init()
 	else
 	{
 		{
+#if PLATFORM_WINDOWS
 			LLM_SCOPED_PAUSE_TRACKING_FOR_TRACKER(ELLMTracker::Default, ELLMAllocType::System);
+#else
+			LLM_SCOPE_BYTAG(D3D12AllocatorUnused);
+#endif
 			const D3D12_HEAP_PROPERTIES HeapProps = CD3DX12_HEAP_PROPERTIES(InitConfig.HeapType, GetGPUMask().GetNative(), GetVisibilityMask().GetNative());
 			VERIFYD3D12RESULT(Adapter->CreateBuffer(HeapProps, GetGPUMask(), InitConfig.InitialResourceState, ED3D12ResourceStateMode::SingleState, InitConfig.InitialResourceState, PoolSize, BackingResource.GetInitReference(), TEXT("Resource Allocator Underlying Buffer"), InitConfig.ResourceFlags));
 #if UE_MEMORY_TRACE_ENABLED
@@ -135,14 +143,18 @@ void FD3D12MemoryPool::Init()
 	}
 #endif // D3D12_RHI_RAYTRACING
 
+#if PLATFORM_WINDOWS
 	LLM_SCOPED_PAUSE_TRACKING_WITH_ENUM_AND_AMOUNT_BYTAG(D3D12AllocatorUnused, int64(PoolSize), ELLMTracker::Platform, ELLMAllocType::System);
+#endif
 	FRHIMemoryPool::Init();
 }
 
 
 void FD3D12MemoryPool::Destroy()
 {
+#if PLATFORM_WINDOWS
 	LLM_SCOPED_PAUSE_TRACKING_FOR_TRACKER(ELLMTracker::Default, ELLMAllocType::System);
+#endif
 
 	FRHIMemoryPool::Destroy();
 
@@ -695,7 +707,9 @@ void FD3D12PoolAllocator::CleanUpAllocations(uint64 InFrameLag, bool bForceFree)
 		FD3D12MemoryPool* MemoryPool = (FD3D12MemoryPool*) Pools[PoolIndex];
 		if (MemoryPool != nullptr && MemoryPool->IsEmpty() && (bForceFree || (MemoryPool->GetLastUsedFrameFence() + InFrameLag <= CompletedFence)))
 		{
+#if PLATFORM_WINDOWS
 			LLM_SCOPED_PAUSE_TRACKING_WITH_ENUM_AND_AMOUNT_BYTAG(D3D12AllocatorUnused, 0 - int64(MemoryPool->GetPoolSize()), ELLMTracker::Platform, ELLMAllocType::System);
+#endif
 			MemoryPool->Destroy();
 			delete(MemoryPool);
 			Pools[PoolIndex] = nullptr;
@@ -819,11 +833,21 @@ void FD3D12PoolAllocator::UpdateAllocationTracking(FD3D12ResourceLocation& InAll
 		
 		if (InAllocationType == EAllocationType::Allocate)
 		{
+#if PLATFORM_WINDOWS
 			LLM_SCOPED_PAUSE_TRACKING_WITH_ENUM_AND_AMOUNT_BYTAG(D3D12AllocatorUnused, 0 - AllocSize, ELLMTracker::Platform, ELLMAllocType::System);
+#else
+			LLM_IF_ENABLED(FLowLevelMemTracker::Get().OnLowLevelAlloc(ELLMTracker::Default, InAllocation.GetAddressForLLMTracking(), AllocSize));
+			LLM_SCOPED_PAUSE_TRACKING_WITH_ENUM_AND_AMOUNT_BYTAG(D3D12AllocatorUnused, 0 - AllocSize, ELLMTracker::Default, ELLMAllocType::System);
+#endif
 		}
 		else
 		{
+#if PLATFORM_WINDOWS
 			LLM_SCOPED_PAUSE_TRACKING_WITH_ENUM_AND_AMOUNT_BYTAG(D3D12AllocatorUnused, AllocSize, ELLMTracker::Platform, ELLMAllocType::System);
+#else
+			LLM_IF_ENABLED(FLowLevelMemTracker::Get().OnLowLevelFree(ELLMTracker::Default, InAllocation.GetAddressForLLMTracking()));
+			LLM_SCOPED_PAUSE_TRACKING_WITH_ENUM_AND_AMOUNT_BYTAG(D3D12AllocatorUnused, AllocSize, ELLMTracker::Default, ELLMAllocType::System);
+#endif
 		}
 	}
 #endif
