More work on Resource Collections
- Removing UAVs for now.
- Adding RHI Validation support.
- Adding D3D12RHI transition and residency handling.
- Adding a few semi-related missing RHIValidation forwarding methods.

diff --git a/Engine/Source/Runtime/D3D12RHI/Private/D3D12Commands.cpp b/Engine/Source/Runtime/D3D12RHI/Private/D3D12Commands.cpp
--- a/Engine/Source/Runtime/D3D12RHI/Private/D3D12Commands.cpp
+++ b/Engine/Source/Runtime/D3D12RHI/Private/D3D12Commands.cpp
@@ -8,6 +8,7 @@ D3D12Commands.cpp: D3D RHI commands implementation.
 #include "ProfilingDebugging/AssetMetadataTrace.h"
 #include "ProfilingDebugging/RealtimeGPUProfiler.h"
 #include "D3D12ResourceCollection.h"
+#include "D3D12TextureReference.h"
 
 static int32 GD3D12AllowDiscardResources = 1;
 static FAutoConsoleVariableRef CVarD3D12AllowDiscardResources(
@@ -945,12 +946,23 @@ struct FD3D12ResourceBinder
 	void SetResourceCollection(FRHIResourceCollection* ResourceCollection, uint32 Index)
 	{
 #if PLATFORM_SUPPORTS_BINDLESS_RENDERING
-		FD3D12ResourceCollection* D3D12ResourceCollection = FD3D12CommandContext::RetrieveObject<FD3D12ResourceCollection>(ResourceCollection, GpuIndex);
-		FD3D12ShaderResourceView* D3D12ShaderResourceView = D3D12ResourceCollection ? D3D12ResourceCollection->GetShaderResourceView() : nullptr;
-
 		if (bBindlessResources)
 		{
-			Context.StateCache.QueueBindlessSRV(Frequency, D3D12ShaderResourceView);
+			if (FD3D12ResourceCollection* D3D12ResourceCollection = FD3D12CommandContext::RetrieveObject<FD3D12ResourceCollection>(ResourceCollection, GpuIndex))
+			{
+				FD3D12ShaderResourceView* D3D12ShaderResourceView = D3D12ResourceCollection->GetShaderResourceView();
+				Context.StateCache.QueueBindlessSRV(Frequency, D3D12ShaderResourceView);
+				Context.StateCache.QueueBindlessSRVs(Frequency, D3D12ResourceCollection->AllSrvs);
+
+				// We have to go through each TextureReference to get the most recent version.
+				for (FD3D12RHITextureReference* TextureReference : D3D12ResourceCollection->AllTextureReferences)
+				{
+					if (FD3D12Texture* Texture = GetD3D12TextureFromRHITexture(TextureReference))
+					{
+						Context.StateCache.QueueBindlessSRV(Frequency, Texture->GetShaderResourceView());
+					}
+				}
+			}
 		}
 		else
 		{
diff --git a/Engine/Source/Runtime/D3D12RHI/Private/D3D12ResourceCollection.cpp b/Engine/Source/Runtime/D3D12RHI/Private/D3D12ResourceCollection.cpp
--- a/Engine/Source/Runtime/D3D12RHI/Private/D3D12ResourceCollection.cpp
+++ b/Engine/Source/Runtime/D3D12RHI/Private/D3D12ResourceCollection.cpp
@@ -8,11 +8,57 @@
 #include "D3D12CommandContext.h"
 #include "D3D12TextureReference.h"
 
-FD3D12ResourceCollection::FD3D12ResourceCollection(FD3D12Device* InParent, FD3D12Buffer* InBuffer, TConstArrayView<FRHIResourceCollectionMember> InMembers)
+FD3D12ResourceCollection::FD3D12ResourceCollection(FD3D12Device* InParent, FRHICommandListBase& RHICmdList, FD3D12Buffer* InBuffer, TConstArrayView<FRHIResourceCollectionMember> InMembers)
 	: FRHIResourceCollection(InMembers)
 	, FD3D12DeviceChild(InParent)
-	, Buffer(InBuffer)
+	, Buffer(InBuffer->GetLinkedObject(InParent->GetGPUIndex()))
 {
+	const uint32 GpuIndex = InParent->GetGPUIndex();
+
+	TArray<FRHIDescriptorHandle> Handles;
+	Handles.Reserve(InMembers.Num());
+
+	for (const FRHIResourceCollectionMember& Member : InMembers)
+	{
+		switch (Member.Type)
+		{
+		case FRHIResourceCollectionMember::EType::Texture:
+		{
+			if (FRHITextureReference* TextureReferenceRHI = static_cast<FRHITexture*>(Member.Resource)->GetTextureReference())
+			{
+				FD3D12RHITextureReference* TextureReference = FD3D12CommandContext::RetrieveObject<FD3D12RHITextureReference>(TextureReferenceRHI, GpuIndex);
+				Handles.Emplace(TextureReference->GetDefaultBindlessHandle());
+				AllTextureReferences.Emplace(TextureReference);
+			}
+			else
+			{
+				FD3D12Texture* Texture = FD3D12CommandContext::RetrieveTexture(static_cast<FRHITexture*>(Member.Resource), GpuIndex);
+				Handles.Emplace(Texture->GetDefaultBindlessHandle());
+				AllSrvs.Emplace(Texture->GetShaderResourceView());
+			}
+		}
+		break;
+		case FRHIResourceCollectionMember::EType::TextureReference:
+		{
+			FD3D12RHITextureReference* TextureReference = FD3D12CommandContext::RetrieveObject<FD3D12RHITextureReference>(Member.Resource, GpuIndex);
+			Handles.Emplace(TextureReference->GetDefaultBindlessHandle());
+			AllTextureReferences.Emplace(TextureReference);
+		}
+		break;
+		case FRHIResourceCollectionMember::EType::ShaderResourceView:
+		{
+			FD3D12ShaderResourceView_RHI* ShaderResourceView = FD3D12CommandContext::RetrieveObject<FD3D12ShaderResourceView_RHI>(Member.Resource, GpuIndex);
+			Handles.Emplace(ShaderResourceView->GetBindlessHandle());
+			AllSrvs.Emplace(ShaderResourceView);
+		}
+		break;
+		}
+	}
+
+	constexpr D3D12_RESOURCE_STATES States = D3D12_RESOURCE_STATE_GENERIC_READ;
+	const TArray<uint32> CollectionMemory = UE::RHICore::CreateResourceCollectionArray<FRHIDescriptorHandle>(Handles);
+	InBuffer->UploadResourceData(RHICmdList, FRHIGPUMask::FromIndex(GpuIndex), States, CollectionMemory.GetData(), CollectionMemory.GetTypeSize() * CollectionMemory.Num());
+
 	D3D12_SHADER_RESOURCE_VIEW_DESC SRVDesc{};
 	SRVDesc.Format = DXGI_FORMAT_R32_TYPELESS;
 	SRVDesc.ViewDimension = D3D12_SRV_DIMENSION_BUFFER;
@@ -21,15 +67,15 @@ FD3D12ResourceCollection::FD3D12ResourceCollection(FD3D12Device* InParent, FD3D1
 	SRVDesc.Buffer.NumElements = UE::RHICore::CalculateResourceCollectionMemorySize(InMembers) / 4;
 	SRVDesc.Buffer.Flags = D3D12_BUFFER_SRV_FLAG_RAW;
 
-	ShaderResourceView = MakeShared<FD3D12ShaderResourceView>(InParent);
-	ShaderResourceView->CreateView(InBuffer, SRVDesc, FD3D12ShaderResourceView::EFlags::None);
+	BufferSRV = MakeShared<FD3D12ShaderResourceView>(InParent);
+	BufferSRV->CreateView(InBuffer, SRVDesc, FD3D12ShaderResourceView::EFlags::None);
 }
 
 FD3D12ResourceCollection::~FD3D12ResourceCollection() = default;
 
 FRHIDescriptorHandle FD3D12ResourceCollection::GetBindlessHandle() const
 {
-	return ShaderResourceView->GetBindlessHandle();
+	return BufferSRV->GetBindlessHandle();
 }
 
 static FD3D12Buffer* CreateCollectionBuffer(FD3D12DynamicRHI& RHI, FRHICommandListBase& RHICmdList, TConstArrayView<FRHIResourceCollectionMember> InMembers)
@@ -50,49 +96,7 @@ FRHIResourceCollectionRef FD3D12DynamicRHI::RHICreateResourceCollection(FRHIComm
 
 	return GetAdapter().CreateLinkedObject<FD3D12ResourceCollection>(FRHIGPUMask::All(), [&RHICmdList, Buffer, InMembers](FD3D12Device* Device)
 	{
-		const uint32 GpuIndex = Device->GetGPUIndex();
-
-		TArray<FRHIDescriptorHandle> Handles;
-
-		for (const FRHIResourceCollectionMember& Member : InMembers)
-		{
-			switch (Member.Type)
-			{
-			case FRHIResourceCollectionMember::EType::Texture:
-			{
-				FD3D12Texture* Texture = FD3D12CommandContext::RetrieveTexture(static_cast<FRHITexture*>(Member.Resource), GpuIndex);
-				Handles.Emplace(Texture->GetDefaultBindlessHandle());
-			}
-			break;
-			case FRHIResourceCollectionMember::EType::TextureReference:
-			{
-				FD3D12RHITextureReference* TextureReference = FD3D12CommandContext::RetrieveObject<FD3D12RHITextureReference>(Member.Resource, GpuIndex);
-				Handles.Emplace(TextureReference->GetDefaultBindlessHandle());
-			}
-			break;
-			case FRHIResourceCollectionMember::EType::ShaderResourceView:
-			{
-				FD3D12ShaderResourceView_RHI* ShaderResourceView = FD3D12CommandContext::RetrieveObject<FD3D12ShaderResourceView_RHI>(Member.Resource, GpuIndex);
-				Handles.Emplace(ShaderResourceView->GetBindlessHandle());
-			}
-			break;
-			case FRHIResourceCollectionMember::EType::UnorderedAccessView:
-			{
-				FD3D12UnorderedAccessView_RHI* UnorderedAccessView = FD3D12CommandContext::RetrieveObject<FD3D12UnorderedAccessView_RHI>(Member.Resource, GpuIndex);
-				Handles.Emplace(UnorderedAccessView->GetBindlessHandle());
-			}
-			break;
-			}
-		}
-
-		FD3D12Buffer* GpuBuffer = Buffer->GetLinkedObject(GpuIndex);
-		FD3D12ResourceCollection* ResourceCollection = new FD3D12ResourceCollection(Device, GpuBuffer, InMembers);
-
-		constexpr D3D12_RESOURCE_STATES States = D3D12_RESOURCE_STATE_GENERIC_READ;
-		const TArray<uint32> CollectionMemory = UE::RHICore::CreateResourceCollectionArray<FRHIDescriptorHandle>(Handles);
-		GpuBuffer->UploadResourceData(RHICmdList, FRHIGPUMask::FromIndex(GpuIndex), States, CollectionMemory.GetData(), CollectionMemory.GetTypeSize() * CollectionMemory.Num());
-
-		return ResourceCollection;
+		return new FD3D12ResourceCollection(Device, RHICmdList, Buffer, InMembers);
 	});
 }
 
diff --git a/Engine/Source/Runtime/D3D12RHI/Private/D3D12ResourceCollection.h b/Engine/Source/Runtime/D3D12RHI/Private/D3D12ResourceCollection.h
--- a/Engine/Source/Runtime/D3D12RHI/Private/D3D12ResourceCollection.h
+++ b/Engine/Source/Runtime/D3D12RHI/Private/D3D12ResourceCollection.h
@@ -11,20 +11,24 @@
 
 class FD3D12Buffer;
 class FD3D12ResourceCollection;
+class FD3D12RHITextureReference;
 class FD3D12ShaderResourceView;
 
 class FD3D12ResourceCollection : public FRHIResourceCollection, public FD3D12DeviceChild, public FD3D12LinkedAdapterObject<FD3D12ResourceCollection>
 {
 public:
-	FD3D12ResourceCollection(FD3D12Device* InParent, FD3D12Buffer* InBuffer, TConstArrayView<FRHIResourceCollectionMember> InMembers);
+	FD3D12ResourceCollection(FD3D12Device* InParent, FRHICommandListBase& RHICmdList, FD3D12Buffer* InBuffer, TConstArrayView<FRHIResourceCollectionMember> InMembers);
 	~FD3D12ResourceCollection();
 
 	virtual FRHIDescriptorHandle GetBindlessHandle() const final;
 
-	FD3D12ShaderResourceView* GetShaderResourceView() const { return ShaderResourceView.Get(); }
+	FD3D12ShaderResourceView* GetShaderResourceView() const { return BufferSRV.Get(); }
 
-	TRefCountPtr<FD3D12Buffer> Buffer;
-	TSharedPtr<FD3D12ShaderResourceView> ShaderResourceView;
+	TRefCountPtr<FD3D12Buffer>           Buffer;
+	TSharedPtr<FD3D12ShaderResourceView> BufferSRV;
+
+	TArray<FD3D12ShaderResourceView*>    AllSrvs;
+	TArray<FD3D12RHITextureReference*>   AllTextureReferences;
 };
 
 template<>
diff --git a/Engine/Source/Runtime/D3D12RHI/Private/D3D12StateCachePrivate.h b/Engine/Source/Runtime/D3D12RHI/Private/D3D12StateCachePrivate.h
--- a/Engine/Source/Runtime/D3D12RHI/Private/D3D12StateCachePrivate.h
+++ b/Engine/Source/Runtime/D3D12RHI/Private/D3D12StateCachePrivate.h
@@ -417,8 +417,10 @@ protected:
 			uint32 CurrentShaderCBCounts     [SF_NumStandardFrequencies] = {};
 			uint32 CurrentShaderUAVCounts    [SF_NumStandardFrequencies] = {};
 
+#if PLATFORM_SUPPORTS_BINDLESS_RENDERING
 			TArray<FD3D12ShaderResourceView*> QueuedBindlessSRVs[SF_NumStandardFrequencies];
 			TArray<FD3D12UnorderedAccessView*> QueuedBindlessUAVs[SF_NumStandardFrequencies];
+#endif
 		} Common = {};
 	} PipelineState = {};
 
@@ -754,12 +756,22 @@ public:
 
 	void ForceSetComputeRootSignature() { PipelineState.Compute.bNeedSetRootSignature = true; }
 
+#if PLATFORM_SUPPORTS_BINDLESS_RENDERING
 	void QueueBindlessSRV(EShaderFrequency ShaderFrequency, FD3D12ShaderResourceView* SRV)
 	{
 		PipelineState.Common.QueuedBindlessSRVs[ShaderFrequency].Emplace(SRV);
 	}
+	void QueueBindlessSRVs(EShaderFrequency ShaderFrequency, TConstArrayView<FD3D12ShaderResourceView*> SRVs)
+	{
+		PipelineState.Common.QueuedBindlessSRVs[ShaderFrequency].Append(SRVs);
+	}
 	void QueueBindlessUAV(EShaderFrequency ShaderFrequency, FD3D12UnorderedAccessView* UAV)
 	{
 		PipelineState.Common.QueuedBindlessUAVs[ShaderFrequency].Emplace(UAV);
 	}
+	void QueueBindlessUAVs(EShaderFrequency ShaderFrequency, TConstArrayView<FD3D12UnorderedAccessView*> UAVs)
+	{
+		PipelineState.Common.QueuedBindlessUAVs[ShaderFrequency].Append(UAVs);
+	}
+#endif
 };
diff --git a/Engine/Source/Runtime/RHI/Public/RHIResourceCollection.h b/Engine/Source/Runtime/RHI/Public/RHIResourceCollection.h
--- a/Engine/Source/Runtime/RHI/Public/RHIResourceCollection.h
+++ b/Engine/Source/Runtime/RHI/Public/RHIResourceCollection.h
@@ -12,7 +12,6 @@ struct FRHIResourceCollectionMember
 		Texture,
 		TextureReference,
 		ShaderResourceView,
-		UnorderedAccessView,
 		Sampler,
 	};
 
@@ -34,10 +33,6 @@ struct FRHIResourceCollectionMember
 		: FRHIResourceCollectionMember(FRHIResourceCollectionMember::EType::ShaderResourceView, InView)
 	{
 	}
-	FRHIResourceCollectionMember(FRHIUnorderedAccessView* InUAV)
-		: FRHIResourceCollectionMember(FRHIResourceCollectionMember::EType::UnorderedAccessView, InUAV)
-	{
-	}
 	FRHIResourceCollectionMember(FRHISamplerState* InSamplerState)
 		: FRHIResourceCollectionMember(FRHIResourceCollectionMember::EType::Sampler, InSamplerState)
 	{
@@ -52,11 +47,25 @@ class FRHIResourceCollection : public FRHIResource
 public:
 	FRHIResourceCollection(TConstArrayView<FRHIResourceCollectionMember> InMembers)
 		: FRHIResource(RRT_ResourceCollection)
+		, Members(InMembers)
+	{
+		for (const FRHIResourceCollectionMember& Member : Members)
+		{
+			if (Member.Resource)
+			{
+				Member.Resource->AddRef();
+			}
+		}
+	}
+	FRHIResourceCollection(const FRHIResourceCollection&) = delete;
+	~FRHIResourceCollection()
 	{
-		Resources.Reserve(InMembers.Num());
-		for (const FRHIResourceCollectionMember& Member : InMembers)
+		for (const FRHIResourceCollectionMember& Member : Members)
 		{
-			Resources.Emplace(Member.Resource);
+			if (Member.Resource)
+			{
+				Member.Resource->Release();
+			}
 		}
 	}
 
@@ -65,5 +74,5 @@ public:
 		return FRHIDescriptorHandle();
 	}
 
-	TArray<TRefCountPtr<FRHIResource>> Resources;
+	TArray<FRHIResourceCollectionMember> Members;
 };
diff --git a/Engine/Source/Runtime/RHI/Public/RHIValidation.h b/Engine/Source/Runtime/RHI/Public/RHIValidation.h
--- a/Engine/Source/Runtime/RHI/Public/RHIValidation.h
+++ b/Engine/Source/Runtime/RHI/Public/RHIValidation.h
@@ -455,6 +455,15 @@ public:
 		RHI->RHIUnlockBufferMGPU(RHICmdList, Buffer, GPUIndex);
 	}
 
+	virtual FTextureReferenceRHIRef RHICreateTextureReference(FRHICommandListBase& RHICmdList, FRHITexture* InReferencedTexture)
+	{
+		return RHI->RHICreateTextureReference(RHICmdList, InReferencedTexture);
+	}
+	virtual void RHIUpdateTextureReference(FRHICommandListBase& RHICmdList, FRHITextureReference* TextureRef, FRHITexture* NewTexture)
+	{
+		RHI->RHIUpdateTextureReference(RHICmdList, TextureRef, NewTexture);
+	}
+
 	virtual FShaderResourceViewRHIRef RHICreateShaderResourceView(class FRHICommandListBase& RHICmdList, FRHIViewableResource* Resource, FRHIViewDesc const& ViewDesc) override final
 	{
 		return RHI->RHICreateShaderResourceView(RHICmdList, Resource, ViewDesc);
@@ -464,6 +473,11 @@ public:
 		return RHI->RHICreateUnorderedAccessView(RHICmdList, Resource, ViewDesc);
 	}
 
+	virtual FRHIResourceCollectionRef RHICreateResourceCollection(FRHICommandListBase& RHICmdList, TConstArrayView<FRHIResourceCollectionMember> InMembers)
+	{
+		return RHI->RHICreateResourceCollection(RHICmdList, InMembers);
+	}
+
 	virtual FRHICalcTextureSizeResult RHICalcTexturePlatformSize(FRHITextureDesc const& Desc, uint32 FirstMipIndex) override final
 	{
 		ensure(Desc.IsValid());
diff --git a/Engine/Source/Runtime/RHI/Public/RHIValidationContext.h b/Engine/Source/Runtime/RHI/Public/RHIValidationContext.h
--- a/Engine/Source/Runtime/RHI/Public/RHIValidationContext.h
+++ b/Engine/Source/Runtime/RHI/Public/RHIValidationContext.h
@@ -69,7 +69,35 @@ inline void ValidateShaderParameters(FRHIShader* RHIShader, RHIValidation::FTrac
 			}
 			break;
 		case FRHIShaderParameterResource::EType::ResourceCollection:
-			// TODO: christopher.waters - ResourceCollection validation
+			if (const FRHIResourceCollection* ResourceCollection = static_cast<const FRHIResourceCollection*>(Parameter.Resource))
+			{
+				for (const FRHIResourceCollectionMember& Member : ResourceCollection->Members)
+				{
+					switch (Member.Type)
+					{
+					case FRHIResourceCollectionMember::EType::Texture:
+						if (FRHITexture* Texture = static_cast<FRHITexture*>(Member.Resource))
+						{
+							Tracker->Assert(Texture->GetWholeResourceIdentitySRV(), InRequiredAccess);
+						}
+						break;
+					case FRHIResourceCollectionMember::EType::TextureReference:
+						if (FRHITextureReference* Texture = static_cast<FRHITextureReference*>(Member.Resource))
+						{
+							Tracker->Assert(Texture->GetWholeResourceIdentitySRV(), InRequiredAccess);
+						}
+						break;
+					case FRHIResourceCollectionMember::EType::ShaderResourceView:
+						if (FRHIShaderResourceView* SRV = static_cast<FRHIShaderResourceView*>(Parameter.Resource))
+						{
+							Tracker->Assert(SRV->GetViewIdentity(), InRequiredAccess);
+						}
+						break;
+					default:
+						break;
+					}
+				}
+			}
 			break;
 		default:
 			checkf(false, TEXT("Unhandled resource type?"));
diff --git a/Engine/Source/Runtime/RHICore/Internal/RHICoreResourceCollection.h b/Engine/Source/Runtime/RHICore/Internal/RHICoreResourceCollection.h
--- a/Engine/Source/Runtime/RHICore/Internal/RHICoreResourceCollection.h
+++ b/Engine/Source/Runtime/RHICore/Internal/RHICoreResourceCollection.h
@@ -28,8 +28,6 @@ namespace UE::RHICore
 			return static_cast<const FRHITextureReference*>(Member.Resource)->GetBindlessHandle();
 		case FRHIResourceCollectionMember::EType::ShaderResourceView:
 			return static_cast<const FRHIShaderResourceView*>(Member.Resource)->GetBindlessHandle();
-		case FRHIResourceCollectionMember::EType::UnorderedAccessView:
-			return static_cast<const FRHIUnorderedAccessView*>(Member.Resource)->GetBindlessHandle();
 		}
 
 		return FRHIDescriptorHandle();
