Refactor d3d12 bindless resource manager to allocate less new GPU heaps and copy over less data (only needed when descriptor is updated, not for every command list)
Also track which handles are dirty per command list to reduce the overhead of copying a full heap each time

diff --git a/Engine/Source/Runtime/D3D12RHI/Private/D3D12BindlessDescriptors.cpp b/Engine/Source/Runtime/D3D12RHI/Private/D3D12BindlessDescriptors.cpp
--- a/Engine/Source/Runtime/D3D12RHI/Private/D3D12BindlessDescriptors.cpp
+++ b/Engine/Source/Runtime/D3D12RHI/Private/D3D12BindlessDescriptors.cpp
@@ -48,11 +48,6 @@ FD3D12DescriptorHeap* UE::D3D12BindlessDescriptors::CreateGpuHeap(FD3D12Device*
 	);
 }
 
-void UE::D3D12BindlessDescriptors::DeferredFreeHeap(FD3D12Device* InDevice, FD3D12DescriptorHeap* InHeap)
-{
-	InDevice->GetDescriptorHeapManager().DeferredFreeHeap(InHeap);
-}
-
 ////////////////////////////////////////////////////////////////////////////////////////////////////////////////////////
 // FD3D12BindlessSamplerManager
 
@@ -104,14 +99,28 @@ void FD3D12BindlessSamplerManager::CloseCommandList(FD3D12CommandContext& Contex
 
 FD3D12BindlessResourceManager::FD3D12BindlessResourceManager(FD3D12Device* InDevice, uint32 InNumDescriptors, TConstArrayView<TStatId> InStats)
 	: FD3D12DeviceChild(InDevice)
-	, CpuHeap(UE::D3D12BindlessDescriptors::CreateCpuHeap(InDevice, ERHIDescriptorHeapType::Standard, InNumDescriptors))
 	, Allocator(ERHIDescriptorHeapType::Standard, InNumDescriptors, InStats)
+	, CpuHeap(UE::D3D12BindlessDescriptors::CreateCpuHeap(InDevice, ERHIDescriptorHeapType::Standard, InNumDescriptors))	
 {
+	// Create the first active GPU heap
+	FGpuHeapData& GpuHeapData = GpuHeaps.AddDefaulted_GetRef();
+	GpuHeapData.GpuHeap = UE::D3D12BindlessDescriptors::CreateGpuHeap(GetParentDevice(), Allocator.GetType(), Allocator.GetCapacity());
+	GpuHeapData.bInUse = true; // mark as in use
+
+	ActiveGpuHeapIndex = 0;
+
+	INC_DWORD_STAT(STAT_D3D12BindlessResourceHeaps);
 }
 
 void FD3D12BindlessResourceManager::CleanupResources()
 {
 	CpuHeap.SafeRelease();
+
+	for (FGpuHeapData& GpuHeap : GpuHeaps)
+	{
+		GpuHeap.GpuHeap.SafeRelease();
+	}
+	GpuHeaps.Empty();
 }
 
 FRHIDescriptorHandle FD3D12BindlessResourceManager::Allocate()
@@ -129,11 +138,39 @@ void FD3D12BindlessResourceManager::Free(FRHIDescriptorHandle InHandle)
 	}
 }
 
-void FD3D12BindlessResourceManager::UpdateDescriptorImmediately(FRHIDescriptorHandle DstHandle, FD3D12View* View)
+void FD3D12BindlessResourceManager::Recycle(FD3D12DescriptorHeap* DescriptorHeap)
+{
+	FScopeLock ScopeLock(&GpuHeapsCS);
+
+	bool bFound = false;
+	for (FGpuHeapData& GpuHeap : GpuHeaps)
+	{
+		if (GpuHeap.GpuHeap == DescriptorHeap)
+		{
+			check(GpuHeap.bInUse);
+			GpuHeap.bInUse = false;
+			bFound = true;
+		}
+	}
+	check(bFound);
+}
+
+void FD3D12BindlessResourceManager::InitializeDescriptor(FRHIDescriptorHandle DstHandle, FD3D12View* View)
 {
 	if (DstHandle.IsValid())
 	{
-		UE::D3D12Descriptors::CopyDescriptor(GetParentDevice(), CpuHeap, DstHandle, View->GetOfflineCpuHandle());
+		TRACE_CPUPROFILER_EVENT_SCOPE("FD3D12BindlessResourceManager::InitializeDescriptor");
+
+		// Update both CPU and active GPU heap since it's initialization and we know the handle isn't currently in use by the GPU
+		FD3D12OfflineDescriptor OfflineCpuHandle = View->GetOfflineCpuHandle();
+		UE::D3D12Descriptors::CopyDescriptor(GetParentDevice(), CpuHeap, DstHandle, OfflineCpuHandle);
+
+		// Copy descriptor to active gpu heaps and to dirty list (needs lock because active gpu heap could be changed on RHI thread)
+		FScopeLock ScopeLock(&GpuHeapsCS);
+		UE::D3D12Descriptors::CopyDescriptor(GetParentDevice(), GpuHeaps[ActiveGpuHeapIndex].GpuHeap, DstHandle, OfflineCpuHandle);
+		GpuHeaps[ActiveGpuHeapIndex].UpdatedHandles.Add(DstHandle);
+
+		INC_DWORD_STAT(STAT_D3D12BindlessResourceDescriptorsInitialized);		
 	}
 }
 
@@ -141,21 +178,30 @@ void FD3D12BindlessResourceManager::UpdateDescriptor(FD3D12ContextArray const& C
 {
 	if (DstHandle.IsValid())
 	{
+		TRACE_CPUPROFILER_EVENT_SCOPE("FD3D12BindlessResourceManager::UpdateDescriptor");
+
+		check(IsInRHIThread() || GRHICommandList.Bypass());
+	
+		// Request allocation of new heap because current GPU heap is used by GPU and can't modify handles in use
 		uint32 const GPUIndex = GetParentDevice()->GetGPUIndex();
 		for (FD3D12CommandContextBase* ContextBase : Contexts)
 		{
 			if (ContextBase)
 			{
 				FD3D12CommandContext& Context = *ContextBase->GetSingleDeviceContext(GPUIndex);
-				if (Context.IsOpen())
-				{
-					FD3D12OfflineDescriptor CopyOfPreviousDescriptorValue = UE::D3D12Descriptors::CreateOfflineCopy(GetParentDevice(), CpuHeap, DstHandle);
-					Context.GetBindlessState().PendingDescriptorRollbacks.Add(GetParentDevice(), DstHandle, CopyOfPreviousDescriptorValue);
-				}
+				Context.GetBindlessState().bRequestNewGpuHeap = true;
+				check(Context.IsDefaultContext());
 			}
 		}
 
-		UpdateDescriptorImmediately(DstHandle, View);
+		// Update the shared CPU heap
+		UE::D3D12Descriptors::CopyDescriptor(GetParentDevice(), CpuHeap, DstHandle, View->GetOfflineCpuHandle());
+
+		// Add to update list so it's updated for the next heap
+		FScopeLock ScopeLock(&GpuHeapsCS);
+		GpuHeaps[ActiveGpuHeapIndex].UpdatedHandles.Add(DstHandle);
+
+		INC_DWORD_STAT(STAT_D3D12BindlessResourceDescriptorsUpdated);
 	}
 }
 
@@ -163,15 +209,14 @@ void FD3D12BindlessResourceManager::FlushPendingDescriptorUpdates(FD3D12CommandC
 {
 	FD3D12ContextBindlessState& State = Context.GetBindlessState();
 
-	if (State.PendingDescriptorRollbacks.Num() > 0)
+	// Create a new heap because there have been descriptor updates?
+	if (State.bRequestNewGpuHeap)
 	{
-		// If we have rollbacks to apply, we have to move to a new heap.
-
 		// First finalize the previous heap if it was set.
 		FinalizeHeapOnState(State);
 
-		// Then create a new heap to use
-		CreateHeapOnState(State);
+		// Then assign the current heap to the state
+		AssignHeapToState(State);
 
 		if (ensure(Context.IsOpen()))
 		{
@@ -186,9 +231,8 @@ void FD3D12BindlessResourceManager::OpenCommandList(FD3D12CommandContext& Contex
 {
 	FD3D12ContextBindlessState& State = Context.GetBindlessState();
 
-	// Always create a new descriptor heap to use.
-	// TODO: defer this to the first FlushPendingDescriptorUpdates
-	CreateHeapOnState(State);
+	// Assign the current active Gpu heap to the context
+	AssignHeapToState(State);
 
 	// Assign the heap to the descriptor cache
 	Context.StateCache.GetDescriptorCache()->SetBindlessResourcesHeapDirectly(State.CurrentGpuHeap);
@@ -214,14 +258,18 @@ void FD3D12BindlessResourceManager::FinalizeContext(FD3D12CommandContext& Contex
 
 	FD3D12ContextBindlessState& State = Context.GetBindlessState();
 
+	// If context wasn't opened but did have descriptor updates make sure the shared gpu heap is updated
+	// (can happen due to texture reference updates not adding any real GPU work)
+	FinalizeHeapOnState(State);
+
 	if (State.UsedHeaps.Num() > 0)
 	{
 		for (const FD3D12DescriptorHeapPtr& UsedHeap : State.UsedHeaps)
 		{
 			checkSlow(UsedHeap);
 
-			// Now queue it up for deletion
-			UE::D3D12BindlessDescriptors::DeferredFreeHeap(GetParentDevice(), UsedHeap);
+			// Now queue it up for recycle when GPU is done
+			FD3D12DynamicRHI::GetD3DRHI()->DeferredDelete(UsedHeap, FD3D12DeferredDeleteObject::EType::BindlessDescriptorHeap);
 		}
 
 		State.UsedHeaps.Empty();
@@ -238,42 +286,97 @@ void FD3D12BindlessResourceManager::CopyCpuHeap(FD3D12DescriptorHeap* Destinatio
 	{
 		const uint32 NumDescriptorsToCopy = AllocatedRange.Last - AllocatedRange.First + 1;
 		UE::D3D12Descriptors::CopyDescriptors(GetParentDevice(), DestinationHeap, CpuHeap, AllocatedRange.First, NumDescriptorsToCopy);
+
+		INC_DWORD_STAT_BY(STAT_D3D12BindlessResourceGPUDescriptorsCopied, NumDescriptorsToCopy);
 	}
 }
 
-void FD3D12BindlessResourceManager::CreateHeapOnState(FD3D12ContextBindlessState& State)
+void FD3D12BindlessResourceManager::AssignHeapToState(FD3D12ContextBindlessState& State)
 {
-	checkf(State.CurrentGpuHeap == nullptr, TEXT("FinalizeHeapOnState was not called before CreateHeapOnState"));
+	checkf(State.CurrentGpuHeap == nullptr, TEXT("FinalizeHeapOnState was not called before AssignHeapToState"));
+
+	FScopeLock ScopeLock(&GpuHeapsCS);
+	check(GpuHeaps[ActiveGpuHeapIndex].GpuHeap);
 
-	State.CurrentGpuHeap = UE::D3D12BindlessDescriptors::CreateGpuHeap(GetParentDevice(), Allocator.GetType(), Allocator.GetCapacity());
+	// By default use the active GPU heap (will be versioned when needed during update while GPU is using it)
+	State.CurrentGpuHeap = GpuHeaps[ActiveGpuHeapIndex].GpuHeap;
 }
 
 void FD3D12BindlessResourceManager::FinalizeHeapOnState(FD3D12ContextBindlessState& State)
 {
-	if (State.CurrentGpuHeap)
+	if (State.bRequestNewGpuHeap)
 	{
-		// Since we're about to stop using this heap, make sure it's in the correct state for the GPU:
-		//  1. Update wholesale from the CPU copy
-		//  2. Take the set of previous descriptor values and apply them to this heap since they were updated in the CPU copy
-		// This effectively "rolls back" the dynamic updates to their correct state while keeping all other immediate updates
+		TRACE_CPUPROFILER_EVENT_SCOPE("FD3D12BindlessResourceManager::FinalizeHeapOnState");
+
+		check(IsInRHIThread() || GRHICommandList.Bypass());
+
+		FScopeLock ScopeLock(&GpuHeapsCS);
 
-		CopyCpuHeap(State.CurrentGpuHeap);
+		check(GpuHeaps[ActiveGpuHeapIndex].GpuHeap == State.CurrentGpuHeap || State.CurrentGpuHeap == nullptr);
 
-		FD3D12PendingDescriptorUpdates& Updates = State.PendingDescriptorRollbacks;
+		// Move the current heap as used previous heaps so it can be garbage collected
+		// (Contains valid data upto this point in time)
+		State.UsedHeaps.Emplace(GpuHeaps[ActiveGpuHeapIndex].GpuHeap);
+
+		int32 NumGpuHeaps = GpuHeaps.Num();
+
+		// Copy over dirty handles to all other heaps so they are updated when reused as well
+		for (int32 GpuHeapIndex = 0; GpuHeapIndex < NumGpuHeaps; ++GpuHeapIndex)
+		{
+			if (GpuHeapIndex != ActiveGpuHeapIndex)
+			{
+				GpuHeaps[GpuHeapIndex].UpdatedHandles.Append(GpuHeaps[ActiveGpuHeapIndex].UpdatedHandles);
+			}
+		}
 
-		if (!Updates.IsEmpty())
+		// Try and reuse a pooled heap (incremented from last used to reduce the possible spike on reuse of lots of heap and dirty handle increase)
+		int32 NewActiveGpuHeapIndex = -1;
+		for (int32 NextIndex = 1; NextIndex < NumGpuHeaps; ++NextIndex)
 		{
-			UE::D3D12Descriptors::CopyDescriptors(GetParentDevice(), State.CurrentGpuHeap, Updates.Handles, Updates.OfflineDescriptors);
+			int32 GpuHeapIndex = (ActiveGpuHeapIndex + NextIndex) % NumGpuHeaps;
 
-			Updates.Empty(GetParentDevice());
+			// Not used by the GPU anymore and not the current one
+			if (GpuHeapIndex != ActiveGpuHeapIndex && !GpuHeaps[GpuHeapIndex].bInUse)
+			{
+				NewActiveGpuHeapIndex = GpuHeapIndex;
+				break;
+			}
+		}		
+
+		// Found a pooled heap, then copy over the dirty descriptor handles
+		if (NewActiveGpuHeapIndex >= 0)
+		{
+			// NOTE: copying over duplicate descriptor entries is faster then adding them to set for reduction
+			//		 CitySample there is about 2 to 4 times duplication but still faster to copy all then deduplication
+
+			FGpuHeapData& NewGpuHeapData = GpuHeaps[NewActiveGpuHeapIndex];
+			INC_DWORD_STAT_BY(STAT_D3D12BindlessResourceGPUDescriptorsCopied, NewGpuHeapData.UpdatedHandles.Num());
+
+			UE::D3D12Descriptors::CopyDescriptors(GetParentDevice(), NewGpuHeapData.GpuHeap, CpuHeap, NewGpuHeapData.UpdatedHandles);
+			NewGpuHeapData.UpdatedHandles.Reset(NewGpuHeapData.UpdatedHandles.Max());
 		}
+		else
+		{
+			NewActiveGpuHeapIndex = GpuHeaps.Num();
+
+			// Allocate a new GPU heap
+			FGpuHeapData& GpuHeapData = GpuHeaps.AddDefaulted_GetRef();
+			GpuHeapData.GpuHeap = UE::D3D12BindlessDescriptors::CreateGpuHeap(GetParentDevice(), Allocator.GetType(), Allocator.GetCapacity());
+			INC_DWORD_STAT(STAT_D3D12BindlessResourceHeaps);
+
+			// Copy over the current CPU state (which contains all updates and latest correct state)
+			CopyCpuHeap(GpuHeapData.GpuHeap);
+		}
+
+		ActiveGpuHeapIndex = NewActiveGpuHeapIndex;
+		GpuHeaps[ActiveGpuHeapIndex].bInUse = true;
 
-		// Move this heap to the used list for cleanup at context finalize.
-		State.UsedHeaps.Emplace(MoveTemp(State.CurrentGpuHeap));
+		INC_DWORD_STAT(STAT_D3D12BindlessResourceHeapsVersioned);
 	}
 
-	// Always clear the pending rollbacks. If we didn't previously bind a descriptor heap, we don't need to roll any descriptors back
-	State.PendingDescriptorRollbacks.Empty(GetParentDevice());
+	// Clear the state data
+	State.CurrentGpuHeap = nullptr;
+	State.bRequestNewGpuHeap = false;
 }
 
 #endif // D3D12RHI_CUSTOM_BINDLESS_RESOURCE_MANAGER
@@ -354,6 +457,14 @@ FRHIDescriptorHandle FD3D12BindlessDescriptorManager::AllocateResourceHandle()
 	return FRHIDescriptorHandle();
 }
 
+void FD3D12BindlessDescriptorManager::Recycle(FD3D12DescriptorHeap* DescriptorHeap)
+{
+	if (ResourceManager)
+	{
+		ResourceManager->Recycle(DescriptorHeap);
+	}
+}
+
 FRHIDescriptorHandle FD3D12BindlessDescriptorManager::AllocateAndInitialize(FD3D12SamplerState* SamplerState)
 {
 	if (SamplerManager)
@@ -390,11 +501,11 @@ void FD3D12BindlessDescriptorManager::DeferredFreeFromDestructor(FRHIDescriptorH
 	}
 }
 
-void FD3D12BindlessDescriptorManager::UpdateDescriptorImmediately(FRHIDescriptorHandle DstHandle, FD3D12View* View)
+void FD3D12BindlessDescriptorManager::InitializeDescriptor(FRHIDescriptorHandle DstHandle, FD3D12View* View)
 {
 	if (DstHandle.GetType() == ERHIDescriptorHeapType::Standard && ResourceManager)
 	{
-		ResourceManager->UpdateDescriptorImmediately(DstHandle, View);
+		ResourceManager->InitializeDescriptor(DstHandle, View);
 		return;
 	}
 
diff --git a/Engine/Source/Runtime/D3D12RHI/Private/D3D12BindlessDescriptors.h b/Engine/Source/Runtime/D3D12RHI/Private/D3D12BindlessDescriptors.h
--- a/Engine/Source/Runtime/D3D12RHI/Private/D3D12BindlessDescriptors.h
+++ b/Engine/Source/Runtime/D3D12RHI/Private/D3D12BindlessDescriptors.h
@@ -89,8 +89,8 @@ struct FD3D12PendingDescriptorUpdates
 // Helper container for all context related bindless state.
 struct FD3D12ContextBindlessState
 {
-	FD3D12PendingDescriptorUpdates   PendingDescriptorRollbacks;
 	FD3D12DescriptorHeapPtr          CurrentGpuHeap;
+	bool							 bRequestNewGpuHeap = false;
 
 	// All heaps used on the context. Used for lifetime management.
 	TArray<FD3D12DescriptorHeapPtr> UsedHeaps;
@@ -98,12 +98,12 @@ struct FD3D12ContextBindlessState
 	FD3D12ContextBindlessState() = default;
 	~FD3D12ContextBindlessState()
 	{
-		check(PendingDescriptorRollbacks.IsEmpty());
+		check(!bRequestNewGpuHeap);
 	}
 
 	bool HasAnyPending() const
 	{
-		return UsedHeaps.Num() > 0 || PendingDescriptorRollbacks.Num() > 0;
+		return UsedHeaps.Num() > 0 || bRequestNewGpuHeap;
 	}
 };
 
@@ -119,7 +119,9 @@ public:
 	FRHIDescriptorHandle Allocate();
 	void                 Free(FRHIDescriptorHandle InHandle);
 
-	void UpdateDescriptorImmediately(FRHIDescriptorHandle DstHandle, FD3D12View* View);
+	void Recycle(FD3D12DescriptorHeap* DescriptorHeap);
+
+	void InitializeDescriptor(FRHIDescriptorHandle DstHandle, FD3D12View* View);
 	void UpdateDescriptor(FD3D12ContextArray const& Contexts, FRHIDescriptorHandle DstHandle, FD3D12View* View);
 
 	void FlushPendingDescriptorUpdates(FD3D12CommandContext& Context);
@@ -136,11 +138,23 @@ public:
 
 private:
 	void CopyCpuHeap(FD3D12DescriptorHeap* DestinationHeap);
-	void CreateHeapOnState(FD3D12ContextBindlessState& State);
+	void AssignHeapToState(FD3D12ContextBindlessState& State);
 	void FinalizeHeapOnState(FD3D12ContextBindlessState& State);
 
-	FD3D12DescriptorHeapPtr      CpuHeap;
-	FRHIHeapDescriptorAllocator  Allocator;
+	FRHIHeapDescriptorAllocator		Allocator;
+
+	FD3D12DescriptorHeapPtr			CpuHeap;	
+
+	struct FGpuHeapData
+	{
+		FD3D12DescriptorHeapPtr		GpuHeap;
+		TArray<FRHIDescriptorHandle> UpdatedHandles;
+		bool						bInUse = true;
+	};
+
+	FCriticalSection				GpuHeapsCS;
+	int32							ActiveGpuHeapIndex;
+	TArray<FGpuHeapData>			GpuHeaps;
 };
 
 #endif
@@ -172,7 +186,9 @@ public:
 	void                 ImmediateFree(FRHIDescriptorHandle InHandle);
 	void                 DeferredFreeFromDestructor(FRHIDescriptorHandle InHandle);
 
-	void UpdateDescriptorImmediately(FRHIDescriptorHandle DstHandle, FD3D12View* View);
+	void Recycle(FD3D12DescriptorHeap* DescriptorHeap);
+
+	void InitializeDescriptor(FRHIDescriptorHandle DstHandle, FD3D12View* View);
 	void UpdateDescriptor(FD3D12ContextArray const& Contexts, FRHIDescriptorHandle DstHandle, FD3D12View* SourceView);
 
 	void FinalizeContext(FD3D12CommandContext& Context);
diff --git a/Engine/Source/Runtime/D3D12RHI/Private/D3D12CommandContext.h b/Engine/Source/Runtime/D3D12RHI/Private/D3D12CommandContext.h
--- a/Engine/Source/Runtime/D3D12RHI/Private/D3D12CommandContext.h
+++ b/Engine/Source/Runtime/D3D12RHI/Private/D3D12CommandContext.h
@@ -50,6 +50,7 @@ struct FD3D12DeferredDeleteObject
 		DescriptorHeap,
 #if PLATFORM_SUPPORTS_BINDLESS_RENDERING
 		BindlessDescriptor,
+		BindlessDescriptorHeap,
 #endif
 		CPUAllocation,
 		DescriptorBlock,
@@ -99,10 +100,16 @@ struct FD3D12DeferredDeleteObject
 		, Heap(InHeap)
 	{}
 
-	explicit FD3D12DeferredDeleteObject(FD3D12DescriptorHeap* InDescriptorHeap)
-		: Type(EType::DescriptorHeap)
+	explicit FD3D12DeferredDeleteObject(FD3D12DescriptorHeap* InDescriptorHeap, EType Type)
+		: Type(Type)
 		, DescriptorHeap(InDescriptorHeap)
-	{}
+	{
+#if PLATFORM_SUPPORTS_BINDLESS_RENDERING
+		check(Type == EType::BindlessDescriptorHeap || Type == EType::DescriptorHeap);
+#else
+		check(Type == EType::DescriptorHeap);
+#endif
+	}
 
 	explicit FD3D12DeferredDeleteObject(ID3D12Object* D3DObject)
 		: Type(EType::D3DObject)
diff --git a/Engine/Source/Runtime/D3D12RHI/Private/D3D12Descriptors.cpp b/Engine/Source/Runtime/D3D12RHI/Private/D3D12Descriptors.cpp
--- a/Engine/Source/Runtime/D3D12RHI/Private/D3D12Descriptors.cpp
+++ b/Engine/Source/Runtime/D3D12RHI/Private/D3D12Descriptors.cpp
@@ -32,6 +32,35 @@ void UE::D3D12Descriptors::CopyDescriptors(FD3D12Device* Device, FD3D12Descripto
 	);
 }
 
+void UE::D3D12Descriptors::CopyDescriptors(FD3D12Device* Device, FD3D12DescriptorHeap* TargetHeap, FD3D12DescriptorHeap* SourceHeap, TConstArrayView<FRHIDescriptorHandle> Handles)
+{
+	SCOPED_NAMED_EVENT_F(TEXT("CopyDescriptors HeapToHeap Scattered (%d)"), FColor::Turquoise, Handles.Num());
+
+	const int32 NumDescriptors = Handles.Num();
+	const D3D12_DESCRIPTOR_HEAP_TYPE D3DHeapType = Translate(TargetHeap->GetType());
+
+	TArray<D3D12_CPU_DESCRIPTOR_HANDLE> DstDescriptors;
+	TArray<D3D12_CPU_DESCRIPTOR_HANDLE> SrcDescriptors;
+	TArray<UINT> CopyCounts;
+
+	DstDescriptors.SetNumUninitialized(NumDescriptors);
+	SrcDescriptors.SetNumUninitialized(NumDescriptors);
+	CopyCounts.SetNumUninitialized(NumDescriptors);
+
+	for (int32 Index = 0; Index < NumDescriptors; Index++)
+	{
+		DstDescriptors[Index] = TargetHeap->GetCPUSlotHandle(Handles[Index].GetIndex());
+		SrcDescriptors[Index] = SourceHeap->GetCPUSlotHandle(Handles[Index].GetIndex());
+		CopyCounts[Index] = 1;
+	}
+
+	Device->GetDevice()->CopyDescriptors(
+		NumDescriptors, DstDescriptors.GetData(), CopyCounts.GetData(),
+		NumDescriptors, SrcDescriptors.GetData(), CopyCounts.GetData(),
+		D3DHeapType
+	);
+}
+
 void UE::D3D12Descriptors::CopyDescriptors(FD3D12Device* Device, FD3D12DescriptorHeap* TargetHeap, TConstArrayView<FRHIDescriptorHandle> DstHandles, TConstArrayView<FD3D12OfflineDescriptor> SrcOfflineDescriptors)
 {
 	SCOPED_NAMED_EVENT_F(TEXT("CopyDescriptors Scattered (%d)"), FColor::Turquoise, DstHandles.Num());
@@ -164,11 +193,6 @@ FD3D12DescriptorManager::FD3D12DescriptorManager(FD3D12Device* Device, FD3D12Des
 
 FD3D12DescriptorManager::~FD3D12DescriptorManager() = default;
 
-void FD3D12DescriptorManager::UpdateDescriptorImmediately(FRHIDescriptorHandle InHandle, D3D12_CPU_DESCRIPTOR_HANDLE InSourceCpuHandle)
-{
-	UE::D3D12Descriptors::CopyDescriptor(GetParentDevice(), GetHeap(), InHandle, InSourceCpuHandle);
-}
-
 ////////////////////////////////////////////////////////////////////////////////////////////////////////////////////////
 // FD3D12OnlineDescriptorManager
 
@@ -560,7 +584,7 @@ void FD3D12DescriptorHeapManager::DeferredFreeHeap(FD3D12DescriptorHeap* InHeap)
 			if (GlobalHeap.IsHeapAChild(InHeap))
 			{
 				InHeap->AddRef();
-				FD3D12DynamicRHI::GetD3DRHI()->DeferredDelete(InHeap);
+				FD3D12DynamicRHI::GetD3DRHI()->DeferredDelete(InHeap, FD3D12DeferredDeleteObject::EType::DescriptorHeap);
 				return;
 			}
 		}
@@ -568,7 +592,7 @@ void FD3D12DescriptorHeapManager::DeferredFreeHeap(FD3D12DescriptorHeap* InHeap)
 	else
 	{
 		InHeap->AddRef();
-		FD3D12DynamicRHI::GetD3DRHI()->DeferredDelete(InHeap);
+		FD3D12DynamicRHI::GetD3DRHI()->DeferredDelete(InHeap, FD3D12DeferredDeleteObject::EType::DescriptorHeap);
 	}
 }
 
diff --git a/Engine/Source/Runtime/D3D12RHI/Private/D3D12Descriptors.h b/Engine/Source/Runtime/D3D12RHI/Private/D3D12Descriptors.h
--- a/Engine/Source/Runtime/D3D12RHI/Private/D3D12Descriptors.h
+++ b/Engine/Source/Runtime/D3D12RHI/Private/D3D12Descriptors.h
@@ -51,6 +51,7 @@ namespace UE::D3D12Descriptors
 	FD3D12DescriptorHeap* CreateDescriptorHeap(FD3D12Device* Device, const TCHAR* DebugName, ERHIDescriptorHeapType HeapType, uint32 NumDescriptors, ED3D12DescriptorHeapFlags Flags, bool bGlobal = false);
 	void CopyDescriptor(FD3D12Device* Device, FD3D12DescriptorHeap* TargetHeap, FRHIDescriptorHandle DstHandle, D3D12_CPU_DESCRIPTOR_HANDLE SrcCpuHandle);
 	void CopyDescriptors(FD3D12Device* Device, FD3D12DescriptorHeap* TargetHeap, FD3D12DescriptorHeap* SourceHeap, uint32 FirstDescriptor, uint32 NumDescriptors);
+	void CopyDescriptors(FD3D12Device* Device, FD3D12DescriptorHeap* TargetHeap, FD3D12DescriptorHeap* SourceHeap, TConstArrayView<FRHIDescriptorHandle> Handles);
 	void CopyDescriptors(FD3D12Device* Device, FD3D12DescriptorHeap* TargetHeap, TConstArrayView<FRHIDescriptorHandle> DstHandles, TConstArrayView<FD3D12OfflineDescriptor> SrcOfflineDescriptors);
 
 	FD3D12OfflineDescriptor CreateOfflineCopy(FD3D12Device* Device, D3D12_CPU_DESCRIPTOR_HANDLE InDescriptor, ERHIDescriptorHeapType InType);
@@ -122,8 +123,6 @@ public:
 	FD3D12DescriptorManager(FD3D12Device* Device, FD3D12DescriptorHeap* InHeap, TConstArrayView<TStatId> InStats);
 	~FD3D12DescriptorManager();
 
-	void UpdateDescriptorImmediately(FRHIDescriptorHandle InHandle, D3D12_CPU_DESCRIPTOR_HANDLE InSourceCpuHandle);
-
 	inline       FD3D12DescriptorHeap* GetHeap()       { return Heap.GetReference(); }
 	inline const FD3D12DescriptorHeap* GetHeap() const { return Heap.GetReference(); }
 
diff --git a/Engine/Source/Runtime/D3D12RHI/Private/D3D12RHI.cpp b/Engine/Source/Runtime/D3D12RHI/Private/D3D12RHI.cpp
--- a/Engine/Source/Runtime/D3D12RHI/Private/D3D12RHI.cpp
+++ b/Engine/Source/Runtime/D3D12RHI/Private/D3D12RHI.cpp
@@ -501,6 +501,9 @@ void FD3D12DynamicRHI::RHIProcessDeleteQueue()
 				case FD3D12DeferredDeleteObject::EType::BindlessDescriptor:
 					ObjectToDelete.BindlessDescriptor.Device->GetBindlessDescriptorManager().ImmediateFree(ObjectToDelete.BindlessDescriptor.Handle);
 					break;
+				case FD3D12DeferredDeleteObject::EType::BindlessDescriptorHeap:
+					ObjectToDelete.DescriptorHeap->GetParentDevice()->GetBindlessDescriptorManager().Recycle(ObjectToDelete.DescriptorHeap);
+					break;
 #endif
 				case FD3D12DeferredDeleteObject::EType::CPUAllocation:
 					FMemory::Free(ObjectToDelete.CPUAllocation);
diff --git a/Engine/Source/Runtime/D3D12RHI/Private/D3D12Stats.h b/Engine/Source/Runtime/D3D12RHI/Private/D3D12Stats.h
--- a/Engine/Source/Runtime/D3D12RHI/Private/D3D12Stats.h
+++ b/Engine/Source/Runtime/D3D12RHI/Private/D3D12Stats.h
@@ -12,6 +12,7 @@ DECLARE_STATS_GROUP(TEXT("D3D12RHI"), STATGROUP_D3D12RHI, STATCAT_Advanced);
 DECLARE_STATS_GROUP(TEXT("D3D12RHI: Memory"), STATGROUP_D3D12Memory, STATCAT_Advanced);
 DECLARE_STATS_GROUP(TEXT("D3D12RHI: Memory Details"), STATGROUP_D3D12MemoryDetails, STATCAT_Advanced);
 DECLARE_STATS_GROUP(TEXT("D3D12RHI: Resources"), STATGROUP_D3D12Resources, STATCAT_Advanced);
+DECLARE_STATS_GROUP(TEXT("D3D12RHI: Bindless"), STATGROUP_D3D12Bindless, STATCAT_Advanced);
 DECLARE_STATS_GROUP(TEXT("D3D12RHI: Buffer Details"), STATGROUP_D3D12BufferDetails, STATCAT_Advanced);
 DECLARE_STATS_GROUP(TEXT("D3D12RHI: Pipeline State (PSO)"), STATGROUP_D3D12PipelineState, STATCAT_Advanced);
 DECLARE_STATS_GROUP(TEXT("D3D12RHI: Descriptor Heap (GPU Visible)"), STATGROUP_D3D12DescriptorHeap, STATCAT_Advanced);
@@ -38,6 +39,12 @@ DECLARE_CYCLE_STAT_EXTERN(TEXT("Decommit transient resource time"), STAT_D3D12De
 
 DECLARE_DWORD_COUNTER_STAT_EXTERN(TEXT("UAV Barriers"), STAT_D3D12UAVBarriers, STATGROUP_D3D12RHI, );
 
+DECLARE_DWORD_ACCUMULATOR_STAT_EXTERN(TEXT("Resource Heaps"), STAT_D3D12BindlessResourceHeaps, STATGROUP_D3D12Bindless, );
+DECLARE_DWORD_COUNTER_STAT_EXTERN(TEXT("Resource Heaps Versioned"), STAT_D3D12BindlessResourceHeapsVersioned, STATGROUP_D3D12Bindless, );
+DECLARE_DWORD_COUNTER_STAT_EXTERN(TEXT("Resource Descriptors Initialized"), STAT_D3D12BindlessResourceDescriptorsInitialized, STATGROUP_D3D12Bindless, );
+DECLARE_DWORD_COUNTER_STAT_EXTERN(TEXT("Resource Descriptors Updated"), STAT_D3D12BindlessResourceDescriptorsUpdated, STATGROUP_D3D12Bindless, );
+DECLARE_DWORD_COUNTER_STAT_EXTERN(TEXT("Resource GPU Descriptors Copied"), STAT_D3D12BindlessResourceGPUDescriptorsCopied, STATGROUP_D3D12Bindless, );
+
 DECLARE_CYCLE_STAT_EXTERN(TEXT("CreateBoundShaderState time"), STAT_D3D12CreateBoundShaderStateTime, STATGROUP_D3D12RHI, );
 DECLARE_CYCLE_STAT_EXTERN(TEXT("New bound shader state time"), STAT_D3D12NewBoundShaderStateTime, STATGROUP_D3D12RHI, );
 DECLARE_DWORD_ACCUMULATOR_STAT_EXTERN(TEXT("Num bound shader states"), STAT_D3D12NumBoundShaderState, STATGROUP_D3D12RHI, );
diff --git a/Engine/Source/Runtime/D3D12RHI/Private/D3D12TextureReference.cpp b/Engine/Source/Runtime/D3D12RHI/Private/D3D12TextureReference.cpp
--- a/Engine/Source/Runtime/D3D12RHI/Private/D3D12TextureReference.cpp
+++ b/Engine/Source/Runtime/D3D12RHI/Private/D3D12TextureReference.cpp
@@ -15,7 +15,7 @@ FD3D12RHITextureReference::FD3D12RHITextureReference(FD3D12Device* InDevice, FD3
 	{
 		InReferencedTexture->AddRenameListener(this);
 
-		InDevice->GetBindlessDescriptorManager().UpdateDescriptorImmediately(BindlessHandle, InReferencedTexture->GetShaderResourceView());
+		InDevice->GetBindlessDescriptorManager().InitializeDescriptor(BindlessHandle, InReferencedTexture->GetShaderResourceView());
 	}
 #endif // PLATFORM_SUPPORTS_BINDLESS_RENDERING
 }
diff --git a/Engine/Source/Runtime/D3D12RHI/Private/D3D12Util.cpp b/Engine/Source/Runtime/D3D12RHI/Private/D3D12Util.cpp
--- a/Engine/Source/Runtime/D3D12RHI/Private/D3D12Util.cpp
+++ b/Engine/Source/Runtime/D3D12RHI/Private/D3D12Util.cpp
@@ -1971,6 +1971,12 @@ DEFINE_STAT(STAT_D3D12DecommitTransientResourceTime);
 
 DEFINE_STAT(STAT_D3D12UAVBarriers);
 
+DEFINE_STAT(STAT_D3D12BindlessResourceHeaps);
+DEFINE_STAT(STAT_D3D12BindlessResourceHeapsVersioned);
+DEFINE_STAT(STAT_D3D12BindlessResourceDescriptorsInitialized);
+DEFINE_STAT(STAT_D3D12BindlessResourceDescriptorsUpdated);
+DEFINE_STAT(STAT_D3D12BindlessResourceGPUDescriptorsCopied);
+
 DEFINE_STAT(STAT_D3D12NewBoundShaderStateTime);
 DEFINE_STAT(STAT_D3D12CreateBoundShaderStateTime);
 DEFINE_STAT(STAT_D3D12NumBoundShaderState);
diff --git a/Engine/Source/Runtime/D3D12RHI/Private/D3D12View.cpp b/Engine/Source/Runtime/D3D12RHI/Private/D3D12View.cpp
--- a/Engine/Source/Runtime/D3D12RHI/Private/D3D12View.cpp
+++ b/Engine/Source/Runtime/D3D12RHI/Private/D3D12View.cpp
@@ -251,7 +251,7 @@ void FD3D12View::InitializeBindlessSlot()
 #if PLATFORM_SUPPORTS_BINDLESS_RENDERING
 	if (BindlessHandle.IsValid())
 	{
-		GetParentDevice()->GetBindlessDescriptorManager().UpdateDescriptorImmediately(BindlessHandle, this);
+		GetParentDevice()->GetBindlessDescriptorManager().InitializeDescriptor(BindlessHandle, this);
 	}
 #endif
 }
