Moving code around to fix header compiles.

diff --git a/Engine/Source/Runtime/D3D12RHI/Private/D3D12Query.h b/Engine/Source/Runtime/D3D12RHI/Private/D3D12Query.h
--- a/Engine/Source/Runtime/D3D12RHI/Private/D3D12Query.h
+++ b/Engine/Source/Runtime/D3D12RHI/Private/D3D12Query.h
@@ -1,42 +1,36 @@
 // Copyright Epic Games, Inc. All Rights Reserved.
 
-/*=============================================================================
-D3D12Query.h: Implementation of D3D12 Query
-=============================================================================*/
 #pragma once
 
 #include "D3D12RHICommon.h"
-#include "D3D12Submission.h"
 #include "D3D12Residency.h"
+#include "D3D12Resources.h"
 #include "RHIResources.h"
 
 class FD3D12SyncPoint;
 using FD3D12SyncPointRef = TRefCountPtr<FD3D12SyncPoint>;
+struct FD3D12QueryLocation;
 
-/** D3D12 Render query */
-class FD3D12RenderQuery : public FRHIRenderQuery, public FD3D12DeviceChild, public FD3D12LinkedAdapterObject<FD3D12RenderQuery>
+enum class ED3D12QueryType
 {
-public:
-	FD3D12RenderQuery(FD3D12Device* Parent, ERenderQueryType InQueryType);
-	~FD3D12RenderQuery();
-
-	ERenderQueryType const Type;
-
-	// Signaled when the result is available. Nullptr if the query has never been used.
-	FD3D12SyncPointRef SyncPoint;
-
-	// The query result, read from the GPU. Heap allocated since it is
-	// accessed by the interrupt thread, and needs to outlive the RHI object.
-	uint64* Result;
-	
-	// The current query location for occlusion queries.
-	FD3D12QueryLocation ActiveLocation;
+	None,
+	CommandListBegin,
+	CommandListEnd,
+	PipelineStats,
+	IdleBegin,
+	IdleEnd,
+	AdjustedMicroseconds,
+	AdjustedRaw,
+	Occlusion
 };
 
-template<>
-struct TD3D12ResourceTraits<FRHIRenderQuery>
+enum class ED3D12QueryPosition
 {
-	typedef FD3D12RenderQuery TConcreteType;
+	// Query result should be written before any future command list work is started.
+	TopOfPipe,
+
+	// Query result should be written after all prior command list work has completed.
+	BottomOfPipe
 };
 
 // Wraps an ID3D12QueryHeap and its readback buffer. Used by command contexts to create timestamp and occlusion queries.
@@ -96,6 +90,70 @@ private:
 	FThreadSafeCounter NumRefs;
 };
 
+struct FD3D12QueryRange
+{
+	TRefCountPtr<FD3D12QueryHeap> Heap;
+	uint32 Start = 0;
+	uint32 End = 0;
+
+	FD3D12QueryRange() = default;
+	FD3D12QueryRange(FD3D12QueryHeap* Heap, uint32 Start, uint32 End)
+	: Heap(Heap)
+	, Start(Start)
+	, End(End)
+	{
+	}
+
+	inline bool IsFull() const;
+
+	bool operator == (FD3D12QueryRange const& RHS) const
+	{
+		return Heap  == RHS.Heap
+		&& Start == RHS.Start
+		&& End   == RHS.End;
+	}
+
+	bool operator < (FD3D12QueryRange const& RHS) const
+	{
+		return Start < RHS.Start;
+	}
+};
+
+// The location of a single (timestamp or occlusion) query result.
+struct FD3D12QueryLocation
+{
+	// The heap in which the result is contained.
+	FD3D12QueryHeap* Heap = nullptr;
+
+	// The index of the query within the heap.
+	uint32 Index = 0;
+
+	ED3D12QueryType Type = ED3D12QueryType::None;
+
+	// The location into which the result is written by the interrupt thread.
+	void* Target = nullptr;
+
+	// Reads the query result from the heap
+	inline void CopyResultTo(void* Dst) const;
+
+	template <typename TValueType>
+	inline TValueType GetResult() const;
+
+	FD3D12QueryLocation() = default;
+	FD3D12QueryLocation(FD3D12QueryHeap* Heap, uint32 Index, ED3D12QueryType Type, void* Target)
+	: Heap(Heap)
+	, Index(Index)
+	, Type(Type)
+	, Target(Target)
+	{
+	}
+
+	operator bool() const
+	{
+		return Heap != nullptr;
+	}
+};
+
 inline void FD3D12QueryLocation::CopyResultTo(void* Dst) const
 {
 	check(Dst);
@@ -146,3 +204,29 @@ public:
 private:
 	TArray<FD3D12QueryRange> Ranges;
 };
+
+/** D3D12 Render query */
+class FD3D12RenderQuery : public FRHIRenderQuery, public FD3D12DeviceChild, public FD3D12LinkedAdapterObject<FD3D12RenderQuery>
+{
+public:
+	FD3D12RenderQuery(FD3D12Device* Parent, ERenderQueryType InQueryType);
+	~FD3D12RenderQuery();
+
+	ERenderQueryType const Type;
+
+	// Signaled when the result is available. Nullptr if the query has never been used.
+	FD3D12SyncPointRef SyncPoint;
+
+	// The query result, read from the GPU. Heap allocated since it is
+	// accessed by the interrupt thread, and needs to outlive the RHI object.
+	uint64* Result;
+	
+	// The current query location for occlusion queries.
+	FD3D12QueryLocation ActiveLocation;
+};
+
+template<>
+struct TD3D12ResourceTraits<FRHIRenderQuery>
+{
+	typedef FD3D12RenderQuery TConcreteType;
+};
diff --git a/Engine/Source/Runtime/D3D12RHI/Private/D3D12Submission.h b/Engine/Source/Runtime/D3D12RHI/Private/D3D12Submission.h
--- a/Engine/Source/Runtime/D3D12RHI/Private/D3D12Submission.h
+++ b/Engine/Source/Runtime/D3D12RHI/Private/D3D12Submission.h
@@ -3,11 +3,12 @@
 #pragma once
 
 #include "Async/TaskGraphInterfaces.h"
-#include "D3D12RHICommon.h"
+#include "D3D12BindlessDescriptors.h"
+#include "D3D12Query.h"
 #include "D3D12Queue.h"
+#include "D3D12RHICommon.h"
 #include "Templates/Function.h"
 #include "Templates/RefCounting.h"
-#include "D3D12BindlessDescriptors.h"
 #include "RHIBreadcrumbs.h"
 
 enum class ED3D12QueueType;
@@ -133,86 +134,6 @@ public:
 	}
 };
 
-enum class ED3D12QueryType
-{
-	None,
-	CommandListBegin,
-	CommandListEnd,
-	PipelineStats,
-	IdleBegin,
-	IdleEnd,
-	AdjustedMicroseconds,
-	AdjustedRaw,
-	Occlusion
-};
-
-enum class ED3D12QueryPosition
-{
-	// Query result should be written before any future command list work is started.
-	TopOfPipe,
-
-	// Query result should be written after all prior command list work has completed.
-	BottomOfPipe
-};
-
-// The location of a single (timestamp or occlusion) query result.
-struct FD3D12QueryLocation
-{
-	// The heap in which the result is contained.
-	FD3D12QueryHeap* Heap = nullptr;
-
-	// The index of the query within the heap.
-	uint32 Index = 0;
-
-	ED3D12QueryType Type = ED3D12QueryType::None;
-
-	// The location into which the result is written by the interrupt thread.
-	void* Target = nullptr;
-
-	// Reads the query result from the heap
-	inline void CopyResultTo(void* Dst) const;
-
-	template <typename TValueType>
-	inline TValueType GetResult() const;
-
-	FD3D12QueryLocation() = default;
-	FD3D12QueryLocation(FD3D12QueryHeap* Heap, uint32 Index, ED3D12QueryType Type, void* Target)
-		: Heap	(Heap  )
-		, Index	(Index )
-		, Type	(Type  )
-		, Target(Target)
-	{}
-
-	operator bool() const { return Heap != nullptr; }
-};
-
-struct FD3D12QueryRange
-{
-	TRefCountPtr<FD3D12QueryHeap> Heap;
-	uint32 Start = 0, End = 0;
-
-	FD3D12QueryRange() = default;
-	FD3D12QueryRange(FD3D12QueryHeap* Heap, uint32 Start, uint32 End)
-		: Heap(Heap)
-		, Start(Start)
-		, End(End)
-	{}
-
-	inline bool IsFull() const;
-
-	bool operator == (FD3D12QueryRange const& RHS) const
-	{
-		return Heap  == RHS.Heap
-			&& Start == RHS.Start
-			&& End   == RHS.End;
-	}
-
-	bool operator < (FD3D12QueryRange const& RHS) const
-	{
-		return Start < RHS.Start;
-	}
-};
-
 struct FD3D12CommitReservedResourceDesc
 {
 	FD3D12Resource* Resource = nullptr;
