Fixing D3D12 Bindless with RayTracingOnly
- Instead of trying to hook into DescriptorCache from ExplicitDescriptorCache, ExplicitDescriptorCache should be passed to DescriptorCache to make decisions on which heap to bind based on the explicit cache setup.
- With this new setup, the binding of the explicit cache can now force bind the bindless heaps when necessary.
- Changing Bindless managers to only bind their heaps when configured to do so. The heaps need to be maintained the same when RT only is enabled, but they shouldn't be bound all the time.

diff --git a/Engine/Source/Runtime/D3D12RHI/Private/D3D12BindlessDescriptors.cpp b/Engine/Source/Runtime/D3D12RHI/Private/D3D12BindlessDescriptors.cpp
--- a/Engine/Source/Runtime/D3D12RHI/Private/D3D12BindlessDescriptors.cpp
+++ b/Engine/Source/Runtime/D3D12RHI/Private/D3D12BindlessDescriptors.cpp
@@ -48,12 +48,18 @@ FD3D12DescriptorHeap* UE::D3D12BindlessDescriptors::CreateGpuHeap(FD3D12Device*
 	);
 }
 
+void UE::D3D12BindlessDescriptors::DeferredFreeHeap(FD3D12Device* InDevice, FD3D12DescriptorHeap* InHeap)
+{
+	FD3D12DynamicRHI::GetD3DRHI()->DeferredDelete(InHeap, FD3D12DeferredDeleteObject::EType::BindlessDescriptorHeap);
+}
+
 ////////////////////////////////////////////////////////////////////////////////////////////////////////////////////////
 // FD3D12BindlessSamplerManager
 
-FD3D12BindlessSamplerManager::FD3D12BindlessSamplerManager(FD3D12Device* InDevice, uint32 InNumDescriptors, TConstArrayView<TStatId> InStats)
+FD3D12BindlessSamplerManager::FD3D12BindlessSamplerManager(FD3D12Device* InDevice, ERHIBindlessConfiguration InConfiguration, uint32 InNumDescriptors, TConstArrayView<TStatId> InStats)
 	: FD3D12DeviceChild(InDevice)
 	, Allocator(ERHIDescriptorHeapType::Sampler, InNumDescriptors, InStats)
+	, Configuration(InConfiguration)
 {
 	GpuHeap = UE::D3D12BindlessDescriptors::CreateGpuHeap(InDevice, ERHIDescriptorHeapType::Sampler, InNumDescriptors);
 }
@@ -84,12 +90,23 @@ void FD3D12BindlessSamplerManager::Free(FRHIDescriptorHandle InHandle)
 
 void FD3D12BindlessSamplerManager::OpenCommandList(FD3D12CommandContext& Context)
 {
-	Context.StateCache.GetDescriptorCache()->SetBindlessSamplersHeapDirectly(GetHeap());
+	if (GetConfiguration() == ERHIBindlessConfiguration::AllShaders)
+	{
+		Context.StateCache.GetDescriptorCache()->SetBindlessSamplersHeapDirectly(GetHeap());
+	}
 }
 
 void FD3D12BindlessSamplerManager::CloseCommandList(FD3D12CommandContext& Context)
 {
-	Context.StateCache.GetDescriptorCache()->SetBindlessSamplersHeapDirectly(nullptr);
+	if (GetConfiguration() == ERHIBindlessConfiguration::AllShaders)
+	{
+		Context.StateCache.GetDescriptorCache()->SetBindlessSamplersHeapDirectly(nullptr);
+	}
+}
+
+FD3D12DescriptorHeap* FD3D12BindlessSamplerManager::GetHeapForContext(FD3D12CommandContext& Context) const
+{
+	return GetHeap();
 }
 
 ////////////////////////////////////////////////////////////////////////////////////////////////////////////////////////
@@ -97,10 +114,11 @@ void FD3D12BindlessSamplerManager::CloseCommandList(FD3D12CommandContext& Contex
 
 #if !D3D12RHI_CUSTOM_BINDLESS_RESOURCE_MANAGER
 
-FD3D12BindlessResourceManager::FD3D12BindlessResourceManager(FD3D12Device* InDevice, uint32 InNumDescriptors, TConstArrayView<TStatId> InStats)
+FD3D12BindlessResourceManager::FD3D12BindlessResourceManager(FD3D12Device* InDevice, ERHIBindlessConfiguration InConfiguration, uint32 InNumDescriptors, TConstArrayView<TStatId> InStats)
 	: FD3D12DeviceChild(InDevice)
+	, CpuHeap(UE::D3D12BindlessDescriptors::CreateCpuHeap(InDevice, ERHIDescriptorHeapType::Standard, InNumDescriptors))
 	, Allocator(ERHIDescriptorHeapType::Standard, InNumDescriptors, InStats)
-	, CpuHeap(UE::D3D12BindlessDescriptors::CreateCpuHeap(InDevice, ERHIDescriptorHeapType::Standard, InNumDescriptors))	
+	, Configuration(InConfiguration)
 {
 	// Create the first active GPU heap
 	FGpuHeapData& GpuHeapData = GpuHeaps.AddDefaulted_GetRef();
@@ -170,7 +188,7 @@ void FD3D12BindlessResourceManager::InitializeDescriptor(FRHIDescriptorHandle Ds
 		UE::D3D12Descriptors::CopyDescriptor(GetParentDevice(), GpuHeaps[ActiveGpuHeapIndex].GpuHeap, DstHandle, OfflineCpuHandle);
 		GpuHeaps[ActiveGpuHeapIndex].UpdatedHandles.Add(DstHandle);
 
-		INC_DWORD_STAT(STAT_D3D12BindlessResourceDescriptorsInitialized);		
+		INC_DWORD_STAT(STAT_D3D12BindlessResourceDescriptorsInitialized);
 	}
 }
 
@@ -218,7 +236,7 @@ void FD3D12BindlessResourceManager::FlushPendingDescriptorUpdates(FD3D12CommandC
 		// Then assign the current heap to the state
 		AssignHeapToState(State);
 
-		if (ensure(Context.IsOpen()))
+		if (GetConfiguration() == ERHIBindlessConfiguration::AllShaders && ensure(Context.IsOpen()))
 		{
 			// Finally tell the Context that we're using this heap,
 			// this call also makes sure the heap is set on the d3d command list.
@@ -234,8 +252,11 @@ void FD3D12BindlessResourceManager::OpenCommandList(FD3D12CommandContext& Contex
 	// Assign the current active Gpu heap to the context
 	AssignHeapToState(State);
 
-	// Assign the heap to the descriptor cache
-	Context.StateCache.GetDescriptorCache()->SetBindlessResourcesHeapDirectly(State.CurrentGpuHeap);
+	if (GetConfiguration() == ERHIBindlessConfiguration::AllShaders)
+	{
+		// Assign the heap to the descriptor cache
+		Context.StateCache.GetDescriptorCache()->SetBindlessResourcesHeapDirectly(State.CurrentGpuHeap);
+	}
 }
 
 void FD3D12BindlessResourceManager::CloseCommandList(FD3D12CommandContext& Context)
@@ -245,8 +266,11 @@ void FD3D12BindlessResourceManager::CloseCommandList(FD3D12CommandContext& Conte
 	// First finalize the current heap if any was set
 	FinalizeHeapOnState(State);
 
-	// Then clear the reference from the state cache
-	Context.StateCache.GetDescriptorCache()->SetBindlessResourcesHeapDirectly(nullptr);
+	if (GetConfiguration() == ERHIBindlessConfiguration::AllShaders)
+	{
+		// Then clear the reference from the state cache
+		Context.StateCache.GetDescriptorCache()->SetBindlessResourcesHeapDirectly(nullptr);
+	}
 }
 
 void FD3D12BindlessResourceManager::FinalizeContext(FD3D12CommandContext& Context)
@@ -269,7 +293,7 @@ void FD3D12BindlessResourceManager::FinalizeContext(FD3D12CommandContext& Contex
 			checkSlow(UsedHeap);
 
 			// Now queue it up for recycle when GPU is done
-			FD3D12DynamicRHI::GetD3DRHI()->DeferredDelete(UsedHeap, FD3D12DeferredDeleteObject::EType::BindlessDescriptorHeap);
+			UE::D3D12BindlessDescriptors::DeferredFreeHeap(GetParentDevice(), UsedHeap);
 		}
 
 		State.UsedHeaps.Empty();
@@ -278,6 +302,18 @@ void FD3D12BindlessResourceManager::FinalizeContext(FD3D12CommandContext& Contex
 	check(!Context.GetBindlessState().HasAnyPending());
 }
 
+FD3D12DescriptorHeap* FD3D12BindlessResourceManager::GetHeap(ERHIPipeline Pipeline) const
+{
+	checkNoEntry();
+	return nullptr;
+}
+
+FD3D12DescriptorHeap* FD3D12BindlessResourceManager::GetHeapForContext(FD3D12CommandContext& Context) const
+{
+	FD3D12ContextBindlessState& State = Context.GetBindlessState();
+	return State.CurrentGpuHeap;
+}
+
 void FD3D12BindlessResourceManager::CopyCpuHeap(FD3D12DescriptorHeap* DestinationHeap)
 {
 	// Copy the smallest possible set of descriptors from the CPU heap to the new GPU heap.
@@ -304,7 +340,7 @@ void FD3D12BindlessResourceManager::AssignHeapToState(FD3D12ContextBindlessState
 
 void FD3D12BindlessResourceManager::FinalizeHeapOnState(FD3D12ContextBindlessState& State)
 {
-	if (State.bRequestNewGpuHeap)
+	if (State.bRequestNewGpuHeap && State.CurrentGpuHeap)
 	{
 		TRACE_CPUPROFILER_EVENT_SCOPE("FD3D12BindlessResourceManager::FinalizeHeapOnState");
 
@@ -341,7 +377,7 @@ void FD3D12BindlessResourceManager::FinalizeHeapOnState(FD3D12ContextBindlessSta
 				NewActiveGpuHeapIndex = GpuHeapIndex;
 				break;
 			}
-		}		
+		}
 
 		// Found a pooled heap, then copy over the dirty descriptor handles
 		if (NewActiveGpuHeapIndex >= 0)
@@ -349,11 +385,11 @@ void FD3D12BindlessResourceManager::FinalizeHeapOnState(FD3D12ContextBindlessSta
 			// NOTE: copying over duplicate descriptor entries is faster then adding them to set for reduction
 			//		 CitySample there is about 2 to 4 times duplication but still faster to copy all then deduplication
 
-			FGpuHeapData& NewGpuHeapData = GpuHeaps[NewActiveGpuHeapIndex];
-			INC_DWORD_STAT_BY(STAT_D3D12BindlessResourceGPUDescriptorsCopied, NewGpuHeapData.UpdatedHandles.Num());
+			FGpuHeapData& GpuHeapData = GpuHeaps[NewActiveGpuHeapIndex];
+			INC_DWORD_STAT_BY(STAT_D3D12BindlessResourceGPUDescriptorsCopied, GpuHeapData.UpdatedHandles.Num());
 
-			UE::D3D12Descriptors::CopyDescriptors(GetParentDevice(), NewGpuHeapData.GpuHeap, CpuHeap, NewGpuHeapData.UpdatedHandles);
-			NewGpuHeapData.UpdatedHandles.Reset(NewGpuHeapData.UpdatedHandles.Max());
+			UE::D3D12Descriptors::CopyDescriptors(GetParentDevice(), GpuHeapData.GpuHeap, CpuHeap, GpuHeapData.UpdatedHandles);
+			GpuHeapData.UpdatedHandles.Reset();
 		}
 		else
 		{
@@ -409,7 +445,7 @@ void FD3D12BindlessDescriptorManager::Init()
 		NumResourceDescriptors += GBindlessOnlineDescriptorHeapBlockSize;
 #endif
 
-		ResourceManager = MakeUnique<FD3D12BindlessResourceManager>(GetParentDevice(), NumResourceDescriptors, Stats);
+		ResourceManager = MakeUnique<FD3D12BindlessResourceManager>(GetParentDevice(), ResourcesConfiguration, NumResourceDescriptors, Stats);
 	}
 
 	if (SamplersConfiguration != ERHIBindlessConfiguration::Disabled)
@@ -430,7 +466,7 @@ void FD3D12BindlessDescriptorManager::Init()
 			NumSamplerDescriptors = D3D12_MAX_SHADER_VISIBLE_SAMPLER_HEAP_SIZE;
 		}
 
-		SamplerManager = MakeUnique<FD3D12BindlessSamplerManager>(GetParentDevice(), NumSamplerDescriptors, Stats);
+		SamplerManager = MakeUnique<FD3D12BindlessSamplerManager>(GetParentDevice(), SamplersConfiguration, NumSamplerDescriptors, Stats);
 	}
 }
 
@@ -567,48 +603,35 @@ void FD3D12BindlessDescriptorManager::FlushPendingDescriptorUpdates(FD3D12Comman
 	}
 }
 
-#if D3D12RHI_USE_CONSTANT_BUFFER_VIEWS
-TRHIPipelineArray<FD3D12DescriptorHeapPtr> FD3D12BindlessDescriptorManager::AllocateResourceHeapsForAllPipelines(int32 InSize)
+FD3D12DescriptorHeapPair FD3D12BindlessDescriptorManager::GetHeapsForContext(FD3D12CommandContext& Context, ERHIBindlessConfiguration InConfiguration) const
 {
-	if (ResourceManager)
+	FD3D12DescriptorHeapPair Result{};
+
+	if (AreResourcesBindless(InConfiguration) && ensure(ResourceManager))
 	{
-		return ResourceManager->AllocateResourceHeapsForAllPipelines(InSize);
+		Result.ResourceHeap = ResourceManager->GetHeapForContext(Context);
 	}
 
-	// Bad configuration?
-	checkNoEntry();
-	return TRHIPipelineArray<FD3D12DescriptorHeapPtr>();
-}
-#endif // D3D12RHI_USE_CONSTANT_BUFFER_VIEWS
-
-FD3D12DescriptorHeap* FD3D12BindlessDescriptorManager::GetResourceHeap(ERHIPipeline Pipeline)
-{
-	return ResourceManager->GetHeap(Pipeline);
-}
-
-FD3D12DescriptorHeap* FD3D12BindlessDescriptorManager::GetSamplerHeap()
-{
-	return SamplerManager->GetHeap();
-}
-
-FD3D12DescriptorHeap* FD3D12BindlessDescriptorManager::GetResourceHeap(ERHIPipeline Pipeline, ERHIBindlessConfiguration InConfiguration)
-{
-	if (AreResourcesBindless(InConfiguration))
+	if (AreSamplersBindless(InConfiguration) && ensure(SamplerManager))
 	{
-		return GetResourceHeap(Pipeline);
+		Result.SamplerHeap = SamplerManager->GetHeapForContext(Context);
 	}
 
-	return nullptr;
+	return Result;
 }
 
-FD3D12DescriptorHeap* FD3D12BindlessDescriptorManager::GetSamplerHeap(ERHIBindlessConfiguration InConfiguration)
+#if D3D12RHI_USE_CONSTANT_BUFFER_VIEWS
+TRHIPipelineArray<FD3D12DescriptorHeapPtr> FD3D12BindlessDescriptorManager::AllocateResourceHeapsForAllPipelines(int32 InSize)
 {
-	if (AreSamplersBindless(InConfiguration))
+	if (ResourceManager)
 	{
-		return GetSamplerHeap();
+		return ResourceManager->AllocateResourceHeapsForAllPipelines(InSize);
 	}
 
-	return nullptr;
+	// Bad configuration?
+	checkNoEntry();
+	return TRHIPipelineArray<FD3D12DescriptorHeapPtr>();
 }
+#endif // D3D12RHI_USE_CONSTANT_BUFFER_VIEWS
 
 #endif // PLATFORM_SUPPORTS_BINDLESS_RENDERING
diff --git a/Engine/Source/Runtime/D3D12RHI/Private/D3D12BindlessDescriptors.h b/Engine/Source/Runtime/D3D12RHI/Private/D3D12BindlessDescriptors.h
--- a/Engine/Source/Runtime/D3D12RHI/Private/D3D12BindlessDescriptors.h
+++ b/Engine/Source/Runtime/D3D12RHI/Private/D3D12BindlessDescriptors.h
@@ -31,7 +31,7 @@ class FD3D12BindlessSamplerManager : public FD3D12DeviceChild
 {
 public:
 	FD3D12BindlessSamplerManager() = delete;
-	FD3D12BindlessSamplerManager(FD3D12Device* InDevice, uint32 InNumDescriptors, TConstArrayView<TStatId> InStats);
+	FD3D12BindlessSamplerManager(FD3D12Device* InDevice, ERHIBindlessConfiguration InConfiguration, uint32 InNumDescriptors, TConstArrayView<TStatId> InStats);
 
 	void CleanupResources();
 
@@ -41,56 +41,24 @@ public:
 	void OpenCommandList(FD3D12CommandContext& Context);
 	void CloseCommandList(FD3D12CommandContext& Context);
 
-	FD3D12DescriptorHeap* GetHeap() { return GpuHeap.GetReference(); }
+	FD3D12DescriptorHeap* GetHeapForContext(FD3D12CommandContext& Context) const;
+
+	FD3D12DescriptorHeap* GetHeap() const { return GpuHeap.GetReference(); }
+	ERHIBindlessConfiguration GetConfiguration() const { return Configuration; }
 
 private:
 	FD3D12DescriptorHeapPtr      GpuHeap;
 	FRHIHeapDescriptorAllocator  Allocator;
+	ERHIBindlessConfiguration    Configuration;
 };
 
 #if !D3D12RHI_CUSTOM_BINDLESS_RESOURCE_MANAGER
 
-struct FD3D12PendingDescriptorUpdates
-{
-	// List of handles that need updating
-	TArray<FRHIDescriptorHandle>    Handles;
-	// Copies of descriptors to update each handle with. These need to be freed in a specific way.
-	TArray<FD3D12OfflineDescriptor> OfflineDescriptors;
-
-	~FD3D12PendingDescriptorUpdates()
-	{
-		checkSlow(Handles.Num() == 0);
-	}
-
-	void Add(FD3D12Device* Device, FRHIDescriptorHandle DestinationHandle, const FD3D12OfflineDescriptor& OfflineDescriptor)
-	{
-		if (ensure(DestinationHandle.IsValid()))
-		{
-			Handles.Emplace(DestinationHandle);
-			OfflineDescriptors.Emplace(OfflineDescriptor);
-		}
-	}
-
-	void Empty(FD3D12Device* Device)
-	{
-		for (FD3D12OfflineDescriptor& OfflineDescriptor : OfflineDescriptors)
-		{
-			UE::D3D12Descriptors::FreeOfflineCopy(Device, OfflineDescriptor, ERHIDescriptorHeapType::Standard);
-		}
-
-		Handles.Empty();
-		OfflineDescriptors.Empty();
-	}
-
-	int32 Num()     const { return Handles.Num(); }
-	bool  IsEmpty() const { return Num() == 0; }
-};
-
 // Helper container for all context related bindless state.
 struct FD3D12ContextBindlessState
 {
-	FD3D12DescriptorHeapPtr          CurrentGpuHeap;
-	bool							 bRequestNewGpuHeap = false;
+	FD3D12DescriptorHeapPtr CurrentGpuHeap;
+	bool                    bRequestNewGpuHeap = false;
 
 	// All heaps used on the context. Used for lifetime management.
 	TArray<FD3D12DescriptorHeapPtr> UsedHeaps;
@@ -112,7 +80,7 @@ class FD3D12BindlessResourceManager : public FD3D12DeviceChild
 {
 public:
 	FD3D12BindlessResourceManager() = delete;
-	FD3D12BindlessResourceManager(FD3D12Device* InDevice, uint32 InNumDescriptors, TConstArrayView<TStatId> InStats);
+	FD3D12BindlessResourceManager(FD3D12Device* InDevice, ERHIBindlessConfiguration InConfiguration, uint32 InNumDescriptors, TConstArrayView<TStatId> InStats);
 
 	void CleanupResources();
 
@@ -130,35 +98,40 @@ public:
 	void CloseCommandList(FD3D12CommandContext& Context);
 	void FinalizeContext(FD3D12CommandContext& Context);
 
-	FD3D12DescriptorHeap* GetHeap(ERHIPipeline Pipeline)
-	{
-		checkNoEntry();
-		return nullptr;
-	}
+	FD3D12DescriptorHeap* GetHeap(ERHIPipeline Pipeline) const;
+	FD3D12DescriptorHeap* GetHeapForContext(FD3D12CommandContext& Context) const;
+
+	ERHIBindlessConfiguration GetConfiguration() const { return Configuration; }
 
 private:
 	void CopyCpuHeap(FD3D12DescriptorHeap* DestinationHeap);
 	void AssignHeapToState(FD3D12ContextBindlessState& State);
 	void FinalizeHeapOnState(FD3D12ContextBindlessState& State);
 
-	FRHIHeapDescriptorAllocator		Allocator;
-
-	FD3D12DescriptorHeapPtr			CpuHeap;	
+	FD3D12DescriptorHeapPtr         CpuHeap;
+	FRHIHeapDescriptorAllocator     Allocator;
+	const ERHIBindlessConfiguration Configuration;
 
 	struct FGpuHeapData
 	{
-		FD3D12DescriptorHeapPtr		GpuHeap;
+		FD3D12DescriptorHeapPtr      GpuHeap;
 		TArray<FRHIDescriptorHandle> UpdatedHandles;
-		bool						bInUse = true;
+		bool                         bInUse = true;
 	};
 
-	FCriticalSection				GpuHeapsCS;
-	int32							ActiveGpuHeapIndex;
-	TArray<FGpuHeapData>			GpuHeaps;
+	FCriticalSection     GpuHeapsCS;
+	int32                ActiveGpuHeapIndex;
+	TArray<FGpuHeapData> GpuHeaps;
 };
 
 #endif
 
+struct FD3D12DescriptorHeapPair
+{
+	FD3D12DescriptorHeap* SamplerHeap;
+	FD3D12DescriptorHeap* ResourceHeap;
+};
+
 /** Manager for descriptors used in bindless rendering. */
 class FD3D12BindlessDescriptorManager : public FD3D12DeviceChild
 {
@@ -197,17 +170,14 @@ public:
 	void CloseCommandList(FD3D12CommandContext& Context);
 
 	void FlushPendingDescriptorUpdates(FD3D12CommandContext& Context);
+	void SetHeapsForRayTracing(FD3D12CommandContext& Context);
+
+	FD3D12DescriptorHeapPair GetHeapsForContext(FD3D12CommandContext& Context, ERHIBindlessConfiguration InConfiguration) const;
 
 #if D3D12RHI_USE_CONSTANT_BUFFER_VIEWS
 	TRHIPipelineArray<FD3D12DescriptorHeapPtr> AllocateResourceHeapsForAllPipelines(int32 InSize);
 #endif
 
-	FD3D12DescriptorHeap* GetResourceHeap(ERHIPipeline Pipeline);
-	FD3D12DescriptorHeap* GetSamplerHeap();
-
-	FD3D12DescriptorHeap* GetResourceHeap(ERHIPipeline Pipeline, ERHIBindlessConfiguration InConfiguration);
-	FD3D12DescriptorHeap* GetSamplerHeap(ERHIBindlessConfiguration InConfiguration);
-
 private:
 	TUniquePtr<FD3D12BindlessResourceManager> ResourceManager;
 	TUniquePtr<FD3D12BindlessSamplerManager>  SamplerManager;
diff --git a/Engine/Source/Runtime/D3D12RHI/Private/D3D12CommandContext.h b/Engine/Source/Runtime/D3D12RHI/Private/D3D12CommandContext.h
--- a/Engine/Source/Runtime/D3D12RHI/Private/D3D12CommandContext.h
+++ b/Engine/Source/Runtime/D3D12RHI/Private/D3D12CommandContext.h
@@ -39,6 +39,7 @@ class FD3D12Heap;
 struct FD3D12DescriptorHeap;
 class FD3D12ResourceLocation;
 class FD3D12RootSignature;
+class FD3D12ExplicitDescriptorCache;
 
 struct FD3D12DeferredDeleteObject
 {
@@ -785,6 +786,9 @@ public:
 
 	void FlushPendingDescriptorUpdates();
 
+	void SetExplicitDescriptorCache(FD3D12ExplicitDescriptorCache& ExplicitDescriptorCache);
+	void UnsetExplicitDescriptorCache();
+
 	virtual void Finalize(TArray<FD3D12Payload*>& OutPayloads) override;
 
 #if PLATFORM_SUPPORTS_BINDLESS_RENDERING
diff --git a/Engine/Source/Runtime/D3D12RHI/Private/D3D12Commands.cpp b/Engine/Source/Runtime/D3D12RHI/Private/D3D12Commands.cpp
--- a/Engine/Source/Runtime/D3D12RHI/Private/D3D12Commands.cpp
+++ b/Engine/Source/Runtime/D3D12RHI/Private/D3D12Commands.cpp
@@ -79,6 +79,16 @@ void FD3D12CommandContext::FlushPendingDescriptorUpdates()
 #endif
 }
 
+void FD3D12CommandContext::SetExplicitDescriptorCache(FD3D12ExplicitDescriptorCache& ExplicitDescriptorCache)
+{
+	StateCache.GetDescriptorCache()->SetExplicitDescriptorCache(ExplicitDescriptorCache);
+}
+
+void FD3D12CommandContext::UnsetExplicitDescriptorCache()
+{
+	StateCache.GetDescriptorCache()->UnsetExplicitDescriptorCache();
+}
+
 // Vertex state.
 void FD3D12CommandContext::RHISetStreamSource(uint32 StreamIndex, FRHIBuffer* VertexBufferRHI, uint32 Offset)
 {
diff --git a/Engine/Source/Runtime/D3D12RHI/Private/D3D12DescriptorCache.cpp b/Engine/Source/Runtime/D3D12RHI/Private/D3D12DescriptorCache.cpp
--- a/Engine/Source/Runtime/D3D12RHI/Private/D3D12DescriptorCache.cpp
+++ b/Engine/Source/Runtime/D3D12RHI/Private/D3D12DescriptorCache.cpp
@@ -4,6 +4,7 @@
 //	Include Files
 //-----------------------------------------------------------------------------
 #include "D3D12DescriptorCache.h"
+#include "D3D12ExplicitDescriptorCache.h"
 #include "D3D12RHIPrivate.h"
 
 bool FD3D12DescriptorCache::HeapRolledOver(ERHIDescriptorHeapType InHeapType)
@@ -83,7 +84,7 @@ bool FD3D12DescriptorCache::SetDescriptorHeaps(bool bForceHeapChanged)
 
 	ID3D12DescriptorHeap* PendingViewHeap =
 #if PLATFORM_SUPPORTS_BINDLESS_RENDERING
-		IsUsingBindlessResources() ? BindlessResourcesHeap->GetHeap() :
+		BindlessResourcesHeap ? BindlessResourcesHeap->GetHeap() :
 #endif
 		CurrentViewHeap->GetHeap();
 	if (LastSetViewHeap != PendingViewHeap)
@@ -91,7 +92,9 @@ bool FD3D12DescriptorCache::SetDescriptorHeaps(bool bForceHeapChanged)
 		// The view heap changed, so dirty the descriptor tables.
 		bHeapChanged = true;
 
-		if (!IsUsingBindlessResources())
+#if PLATFORM_SUPPORTS_BINDLESS_RENDERING
+		if (!BindlessResourcesHeap)
+#endif
 		{
 			Context.StateCache.DirtyViewDescriptorTables();
 		}
@@ -143,10 +146,8 @@ void FD3D12DescriptorCache::OpenCommandList()
 	LastSetSamplerHeap = nullptr;
 
 #if PLATFORM_SUPPORTS_BINDLESS_RENDERING
-	if (IsUsingBindlessSamplers() || IsUsingBindlessResources())
-	{
-		GetParentDevice()->GetBindlessDescriptorManager().OpenCommandList(Context);
-	}
+	// Always call the Bindless Manager OpenCommandList, it will determine when it needs to do anything.
+	GetParentDevice()->GetBindlessDescriptorManager().OpenCommandList(Context);
 
 	if (!IsUsingBindlessSamplers())
 #endif
@@ -179,10 +180,8 @@ void FD3D12DescriptorCache::CloseCommandList()
 	}
 
 #if PLATFORM_SUPPORTS_BINDLESS_RENDERING
-	if (IsUsingBindlessSamplers() || IsUsingBindlessResources())
-	{
-		GetParentDevice()->GetBindlessDescriptorManager().CloseCommandList(Context);
-	}
+	// Always call the Bindless Manager CloseCommandList, it will determine when it needs to do anything.
+	GetParentDevice()->GetBindlessDescriptorManager().CloseCommandList(Context);
 
 	if (!IsUsingBindlessSamplers())
 #endif
@@ -765,10 +764,40 @@ void FD3D12DescriptorCache::SwitchToNewBindlessResourceHeap(FD3D12DescriptorHeap
 }
 #endif
 
-void FD3D12DescriptorCache::OverrideLastSetHeaps(ID3D12DescriptorHeap* ViewHeap, ID3D12DescriptorHeap* SamplerHeap)
+void FD3D12DescriptorCache::SetExplicitDescriptorCache(FD3D12ExplicitDescriptorCache& ExplicitDescriptorCache)
 {
-	ID3D12DescriptorHeap* ViewHeapToSet = ViewHeap ? ViewHeap : LastSetViewHeap;
-	ID3D12DescriptorHeap* SamplerHeapToSet = SamplerHeap ? SamplerHeap : LastSetSamplerHeap;
+	ID3D12DescriptorHeap* ViewHeapToSet = nullptr;
+	ID3D12DescriptorHeap* SamplerHeapToSet = nullptr;
+
+#if PLATFORM_SUPPORTS_BINDLESS_RENDERING
+	// We have three scenarios:
+	//   Bindless on globally: BindlessHeaps and LastSetXXX will match
+	//   Bindless RT Only: BindlessHeaps will override LastSetXXX
+	//   No Bindless:  BindlessHeaps will be null, ExplicitDescriptorCache heaps will override LastSetXXX
+
+	const FD3D12DescriptorHeapPair BindlessHeaps = GetParentDevice()->GetBindlessDescriptorManager().GetHeapsForContext(Context, ExplicitDescriptorCache.BindlessConfiguration);
+
+	ViewHeapToSet = BindlessHeaps.ResourceHeap ? BindlessHeaps.ResourceHeap->GetHeap() : nullptr;
+	SamplerHeapToSet = BindlessHeaps.SamplerHeap ? BindlessHeaps.SamplerHeap->GetHeap() : nullptr;
+#endif
+
+	const bool bViewHeapIsBindless = (ViewHeapToSet != nullptr);
+
+	if (!ViewHeapToSet)
+	{
+		check(ExplicitDescriptorCache.ViewHeap.GetParentDevice() == GetParentDevice());
+
+		ExplicitDescriptorCache.ViewHeap.UpdateSyncPoint();
+		ViewHeapToSet = ExplicitDescriptorCache.ViewHeap.D3D12Heap;
+	}
+
+	if (!SamplerHeapToSet)
+	{
+		check(ExplicitDescriptorCache.SamplerHeap.GetParentDevice() == GetParentDevice());
+
+		ExplicitDescriptorCache.SamplerHeap.UpdateSyncPoint();
+		SamplerHeapToSet = ExplicitDescriptorCache.SamplerHeap.D3D12Heap;
+	}
 
 	if (ViewHeapToSet != LastSetViewHeap || SamplerHeapToSet != LastSetSamplerHeap)
 	{
@@ -778,16 +807,18 @@ void FD3D12DescriptorCache::OverrideLastSetHeaps(ID3D12DescriptorHeap* ViewHeap,
 		ID3D12DescriptorHeap* ppHeaps[] = { ViewHeapToSet, SamplerHeapToSet };
 		Context.GraphicsCommandList()->SetDescriptorHeaps(UE_ARRAY_COUNT(ppHeaps), ppHeaps);
 
-		bHeapsOverridden = true;
+		bUsingExplicitCacheHeaps = true;
+		bExplicitViewHeapIsBindless = bViewHeapIsBindless;
 	}
 }
 
-void FD3D12DescriptorCache::RestoreAfterExternalHeapsSet()
+void FD3D12DescriptorCache::UnsetExplicitDescriptorCache()
 {
-	if (bHeapsOverridden)
+	if (bUsingExplicitCacheHeaps)
 	{
 		SetDescriptorHeaps();
-		bHeapsOverridden = false;
+		bUsingExplicitCacheHeaps = false;
+		bExplicitViewHeapIsBindless = false;
 	}
 }
 
diff --git a/Engine/Source/Runtime/D3D12RHI/Private/D3D12DescriptorCache.h b/Engine/Source/Runtime/D3D12RHI/Private/D3D12DescriptorCache.h
--- a/Engine/Source/Runtime/D3D12RHI/Private/D3D12DescriptorCache.h
+++ b/Engine/Source/Runtime/D3D12RHI/Private/D3D12DescriptorCache.h
@@ -12,6 +12,7 @@
 class FD3D12CommandContext;
 class FD3D12DynamicRHI;
 class FD3D12DepthStencilView;
+class FD3D12ExplicitDescriptorCache;
 class FD3D12RenderTargetView;
 class FD3D12ShaderResourceView;
 class FD3D12UnorderedAccessView;
@@ -359,8 +360,8 @@ public:
 	void SwitchToNewBindlessResourceHeap(FD3D12DescriptorHeap* InHeap);
 #endif
 
-	void OverrideLastSetHeaps(ID3D12DescriptorHeap* ViewHeap, ID3D12DescriptorHeap* SamplerHeap);
-	void RestoreAfterExternalHeapsSet();
+	void SetExplicitDescriptorCache(FD3D12ExplicitDescriptorCache& ExplicitDescriptorCache);
+	void UnsetExplicitDescriptorCache();
 
 	inline bool UsingGlobalSamplerHeap() const { return CurrentSamplerHeap != &LocalSamplerHeap; }
 	FD3D12SamplerSet& GetLocalSamplerSet() { return *LocalSamplerSet.Get(); }
@@ -389,6 +390,11 @@ public:
 	{
 		return BindlessSamplersHeap;
 	}
+
+	bool IsActiveViewHeapBindless() const
+	{
+		return (BindlessResourcesHeap != nullptr) || (bUsingExplicitCacheHeaps && bExplicitViewHeapIsBindless);
+	}
 #endif
 
 protected:
@@ -412,7 +418,8 @@ private:
 	TArray<FD3D12UniqueSamplerTable> UniqueTables;
 
 	TSharedPtr<FD3D12SamplerSet> LocalSamplerSet;
-	bool bHeapsOverridden = false;
+	bool bUsingExplicitCacheHeaps = false;
+	bool bExplicitViewHeapIsBindless = false;
 	bool bLocalSamplerHeapOpen = false;
 	bool bUsingViewHeap = true;
 
diff --git a/Engine/Source/Runtime/D3D12RHI/Private/D3D12ExplicitDescriptorCache.cpp b/Engine/Source/Runtime/D3D12RHI/Private/D3D12ExplicitDescriptorCache.cpp
--- a/Engine/Source/Runtime/D3D12RHI/Private/D3D12ExplicitDescriptorCache.cpp
+++ b/Engine/Source/Runtime/D3D12RHI/Private/D3D12ExplicitDescriptorCache.cpp
@@ -297,6 +297,7 @@ void FD3D12ExplicitDescriptorCache::Init(uint32 NumViewDescriptors, uint32 NumSa
 #if PLATFORM_SUPPORTS_BINDLESS_RENDERING
 	FD3D12BindlessDescriptorManager& BindlessManager = GetParentDevice()->GetBindlessDescriptorManager();
 
+	BindlessConfiguration = BindlessConfig;
 	bBindlessViews = BindlessManager.AreResourcesBindless(BindlessConfig);
 	bBindlessSamplers = BindlessManager.AreSamplersBindless(BindlessConfig);
 #else
@@ -315,45 +316,6 @@ void FD3D12ExplicitDescriptorCache::Init(uint32 NumViewDescriptors, uint32 NumSa
 	}
 }
 
-void FD3D12ExplicitDescriptorCache::UpdateSyncPoint()
-{
-#if !PLATFORM_SUPPORTS_BINDLESS_RENDERING
-	const bool bBindlessViews = false;
-	const bool bBindlessSamplers = false;
-#endif
-
-	if (!bBindlessViews)
-	{
-		ViewHeap.UpdateSyncPoint();
-	}
-
-	if (!bBindlessSamplers)
-	{
-		SamplerHeap.UpdateSyncPoint();
-	}
-}
-
-void FD3D12ExplicitDescriptorCache::SetDescriptorHeaps(FD3D12CommandContext& CommandContext)
-{
-	UpdateSyncPoint();
-
-#if PLATFORM_SUPPORTS_BINDLESS_RENDERING
-	check(bBindlessViews || ViewHeap.GetParentDevice() == CommandContext.GetParentDevice());
-	check(bBindlessSamplers || SamplerHeap.GetParentDevice() == CommandContext.GetParentDevice());
-
-	ID3D12DescriptorHeap* ViewHeapToSet = bBindlessViews ? nullptr : ViewHeap.D3D12Heap;
-	ID3D12DescriptorHeap* SamplerHeapToSet = bBindlessSamplers ? nullptr : SamplerHeap.D3D12Heap;
-#else
-	check(ViewHeap.GetParentDevice() == CommandContext.GetParentDevice());
-	check(SamplerHeap.GetParentDevice() == CommandContext.GetParentDevice());
-
-	ID3D12DescriptorHeap* ViewHeapToSet = ViewHeap.D3D12Heap;
-	ID3D12DescriptorHeap* SamplerHeapToSet = SamplerHeap.D3D12Heap;
-#endif
-
-	CommandContext.StateCache.GetDescriptorCache()->OverrideLastSetHeaps(ViewHeapToSet, SamplerHeapToSet);
-}
-
 // Returns descriptor heap base index for this descriptor table allocation (checking for duplicates and reusing existing tables) or -1 if allocation failed.
 int32 FD3D12ExplicitDescriptorCache::AllocateDeduplicated(const uint32* DescriptorVersions, const D3D12_CPU_DESCRIPTOR_HANDLE* Descriptors, uint32 NumDescriptors, D3D12_DESCRIPTOR_HEAP_TYPE Type, uint32 WorkerIndex)
 {
diff --git a/Engine/Source/Runtime/D3D12RHI/Private/D3D12ExplicitDescriptorCache.h b/Engine/Source/Runtime/D3D12RHI/Private/D3D12ExplicitDescriptorCache.h
--- a/Engine/Source/Runtime/D3D12RHI/Private/D3D12ExplicitDescriptorCache.h
+++ b/Engine/Source/Runtime/D3D12RHI/Private/D3D12ExplicitDescriptorCache.h
@@ -77,7 +77,6 @@ struct FD3D12ExplicitDescriptorHeap : public FD3D12DeviceChild
 	bool CompareDescriptors(int32 BaseIndex, const D3D12_CPU_DESCRIPTOR_HANDLE* InDescriptors, uint32 InNumDescriptors);
 
 	D3D12_CPU_DESCRIPTOR_HANDLE GetDescriptorCPU(uint32 Index) const;
-
 	D3D12_GPU_DESCRIPTOR_HANDLE GetDescriptorGPU(uint32 Index) const;
 
 	void UpdateSyncPoint();
@@ -119,10 +118,6 @@ public:
 
 	void Init(uint32 NumViewDescriptors, uint32 NumSamplerDescriptors, ERHIBindlessConfiguration BindlessConfig);
 
-	void UpdateSyncPoint();
-
-	void SetDescriptorHeaps(FD3D12CommandContext& CommandContext);
-
 	// Returns descriptor heap base index for this descriptor table allocation or -1 if allocation failed.
 	int32 Allocate(const D3D12_CPU_DESCRIPTOR_HANDLE* Descriptors, uint32 NumDescriptors, D3D12_DESCRIPTOR_HEAP_TYPE Type, uint32 WorkerIndex);
 
@@ -133,6 +128,7 @@ public:
 	FD3D12ExplicitDescriptorHeap SamplerHeap;
 
 #if PLATFORM_SUPPORTS_BINDLESS_RENDERING
+	ERHIBindlessConfiguration BindlessConfiguration{};
 	bool bBindlessViews = false;
 	bool bBindlessSamplers = false;
 #endif
diff --git a/Engine/Source/Runtime/D3D12RHI/Private/D3D12RayTracing.cpp b/Engine/Source/Runtime/D3D12RHI/Private/D3D12RayTracing.cpp
--- a/Engine/Source/Runtime/D3D12RHI/Private/D3D12RayTracing.cpp
+++ b/Engine/Source/Runtime/D3D12RHI/Private/D3D12RayTracing.cpp
@@ -4629,7 +4629,8 @@ static void DispatchRays(FD3D12CommandContext& CommandContext,
 		UE::TScopeLock Lock(OptShaderTable->DispatchMutex);
 		TRACE_CPUPROFILER_EVENT_SCOPE(SetRayTracingShaderResources);
 
-		DescriptorCache->SetDescriptorHeaps(CommandContext);
+		CommandContext.SetExplicitDescriptorCache(*DescriptorCache);
+
 		FD3D12RayTracingGlobalResourceBinder ResourceBinder(CommandContext, *DescriptorCache);
 		bResourcesBound = SetRayTracingShaderResources(RayGenShader, GlobalBindings, ResourceBinder);
 
@@ -4639,7 +4640,7 @@ static void DispatchRays(FD3D12CommandContext& CommandContext,
 	{
 		FD3D12ExplicitDescriptorCache TransientDescriptorCache(CommandContext.GetParentDevice(), FD3D12RayTracingScene::MaxBindingWorkers);
 		TransientDescriptorCache.Init(MAX_SRVS + MAX_UAVS, MAX_SAMPLERS, ERHIBindlessConfiguration::RayTracingShaders);
-		TransientDescriptorCache.SetDescriptorHeaps(CommandContext);
+		CommandContext.SetExplicitDescriptorCache(TransientDescriptorCache);
 		FD3D12RayTracingGlobalResourceBinder ResourceBinder(CommandContext, TransientDescriptorCache);
 		bResourcesBound = SetRayTracingShaderResources(RayGenShader, GlobalBindings, ResourceBinder);
 	}
@@ -4699,7 +4700,7 @@ static void DispatchRays(FD3D12CommandContext& CommandContext,
 	}
 
 	// Restore old global descriptor heaps
-	CommandContext.StateCache.GetDescriptorCache()->RestoreAfterExternalHeapsSet();
+	CommandContext.UnsetExplicitDescriptorCache();
 }
 
 void FD3D12CommandContext::RHIRayTraceDispatch(FRHIRayTracingPipelineState* InRayTracingPipelineState, FRHIRayTracingShader* RayGenShaderRHI,
diff --git a/Engine/Source/Runtime/D3D12RHI/Private/D3D12UAV.cpp b/Engine/Source/Runtime/D3D12RHI/Private/D3D12UAV.cpp
--- a/Engine/Source/Runtime/D3D12RHI/Private/D3D12UAV.cpp
+++ b/Engine/Source/Runtime/D3D12RHI/Private/D3D12UAV.cpp
@@ -330,7 +330,7 @@ void FD3D12CommandContext::ClearUAV(TRHICommandList_RecursiveHazardous<FD3D12Com
 		if (UAVDesc.Format == DXGI_FORMAT_UNKNOWN || (UAVDesc.Buffer.Flags & D3D12_BUFFER_UAV_FLAG_RAW) != 0)
 		{
 			// Structured buffer.
-			RHICmdList.RunOnContext([UnorderedAccessView, ClearValues, UAVDesc](auto& Context)
+			RHICmdList.RunOnContext([UnorderedAccessView, ClearValues, UAVDesc](FD3D12CommandContext& Context)
 			{
 				// Alias the structured buffer with an R32_UINT UAV to perform the clear.
 				// We construct a temporary UAV on the offline heap, copy it to the online heap, and then call ClearUnorderedAccessViewUint.
@@ -365,8 +365,9 @@ void FD3D12CommandContext::ClearUAV(TRHICommandList_RecursiveHazardous<FD3D12Com
 				D3D12_GPU_DESCRIPTOR_HANDLE GPUHandle{};
 
 #if PLATFORM_SUPPORTS_BINDLESS_RENDERING
-				if (UAV.GetBindlessHandle().IsValid())
+				if (Context.StateCache.GetDescriptorCache()->IsActiveViewHeapBindless())
 				{
+					check(UAV.GetBindlessHandle().IsValid());
 					Context.FlushPendingDescriptorUpdates();
 
 					FD3D12DescriptorHeap* BindlessHeap = Context.GetBindlessResourcesHeap();
diff --git a/Engine/Source/Runtime/D3D12RHI/Private/D3D12WorkGraph.cpp b/Engine/Source/Runtime/D3D12RHI/Private/D3D12WorkGraph.cpp
--- a/Engine/Source/Runtime/D3D12RHI/Private/D3D12WorkGraph.cpp
+++ b/Engine/Source/Runtime/D3D12RHI/Private/D3D12WorkGraph.cpp
@@ -603,7 +603,7 @@ void FD3D12CommandContext::DispatchWorkGraphShaderBundle(FRHIShaderBundle* Shade
 	}
 
 	FD3D12ExplicitDescriptorCache TransientDescriptorCache(GetParentDevice(), MaxTasks /* Worker Count */);
-	TransientDescriptorCache.Init(NumViewDescriptors, NumSamplerDescriptors, ERHIBindlessConfiguration::Disabled);
+	TransientDescriptorCache.Init(NumViewDescriptors, NumSamplerDescriptors, ERHIBindlessConfiguration::AllShaders);
 
 	TArray<FShaderBundleBinderOps, TInlineAllocator<MaxWorkers>> BinderOps;
 	BinderOps.SetNum(MaxTasks);
@@ -700,7 +700,7 @@ void FD3D12CommandContext::DispatchWorkGraphShaderBundle(FRHIShaderBundle* Shade
 	FlushResourceBarriers();
 
 	// Apply the transient descriptor heaps.
-	TransientDescriptorCache.SetDescriptorHeaps(*this);
+	SetExplicitDescriptorCache(TransientDescriptorCache);
 
 	// Gather root arguments for global work graph.
 	int32 DispatchSRVTable = INDEX_NONE;
@@ -742,7 +742,7 @@ void FD3D12CommandContext::DispatchWorkGraphShaderBundle(FRHIShaderBundle* Shade
 	Pipeline->bInitialized = true;
 
 	// Restore old global descriptor heaps
-	StateCache.GetDescriptorCache()->RestoreAfterExternalHeapsSet();
+	UnsetExplicitDescriptorCache();
 
 	// We did not write through the state cache, so we need to invalidate it so subsequent workloads correctly re-bind state
 	StateCache.DirtyState();
