Add RHI_NEW_GPU_PROFILER define to remove legacy GPU profiling code, in preparation for implementing the new GPU profiler

Refactored parts of D3D12 submission
 - Grouped query range by the heap pointer at the payload/queue level, and simplified the batching code.
 - Added ProfilerTOP and ProfilerBOP timestamp types, used by the new profiler.

diff --git a/Engine/Plugins/Runtime/nDisplay/Source/DisplayCluster/Private/Game/EngineClasses/Basics/DisplayClusterViewportClient.cpp b/Engine/Plugins/Runtime/nDisplay/Source/DisplayCluster/Private/Game/EngineClasses/Basics/DisplayClusterViewportClient.cpp
--- a/Engine/Plugins/Runtime/nDisplay/Source/DisplayCluster/Private/Game/EngineClasses/Basics/DisplayClusterViewportClient.cpp
+++ b/Engine/Plugins/Runtime/nDisplay/Source/DisplayCluster/Private/Game/EngineClasses/Basics/DisplayClusterViewportClient.cpp
@@ -187,7 +187,7 @@ public:
 
 	virtual int32 DrawStatsHUD(FCanvas* InCanvas, int32 InX, int32 InY) override
 	{
-#if GPUPROFILERTRACE_ENABLED
+#if GPUPROFILERTRACE_ENABLED && (RHI_NEW_GPU_PROFILER == 0)
 		if (GDisplayClusterShowStats)
 		{
 			// Get GPU perf results
diff --git a/Engine/Source/Runtime/D3D12RHI/Private/D3D12CommandContext.cpp b/Engine/Source/Runtime/D3D12RHI/Private/D3D12CommandContext.cpp
--- a/Engine/Source/Runtime/D3D12RHI/Private/D3D12CommandContext.cpp
+++ b/Engine/Source/Runtime/D3D12RHI/Private/D3D12CommandContext.cpp
@@ -198,6 +198,16 @@ void FD3D12ContextCommon::BindDiagnosticBuffer(FD3D12RootSignature const* RootSi
 		#endif
 		}
 
+	#if RHI_NEW_GPU_PROFILER
+		{
+			using namespace UE::RHI::GPUProfiler;
+
+			auto& Event = BreadcrumbEvents.Emplace_GetRef(MakeUnique<FBreadcrumbEvent>(FBreadcrumbEvent::FBeginBreadcrumb{ .Breadcrumb = Breadcrumb }));
+
+			FD3D12QueryLocation TimestampQuery = AllocateQuery(ED3D12QueryType::ProfilerTimestampTOP, &Event->Value.Get<FBreadcrumbEvent::FBeginBreadcrumb>().GPUTimestampTOP);
+			EndQuery(TimestampQuery);
+		}
+	#else
 		if (IsDefaultContext() && !IsAsyncComputeContext())
 		{
 			FD3D12GPUProfiler& GPUProfiler = GetParentDevice()->GetGPUProfiler();
@@ -206,10 +216,21 @@ void FD3D12ContextCommon::BindDiagnosticBuffer(FD3D12RootSignature const* RootSi
 				GPUProfiler.PushEvent(GetNameStr(), FColor::White);
 			}
 		}
+	#endif // (RHI_NEW_GPU_PROFILER == 0)
 	}
 
 	void FD3D12CommandContext::RHIEndBreadcrumbGPU(FRHIBreadcrumbNode* Breadcrumb)
 	{
+	#if RHI_NEW_GPU_PROFILER
+		{
+			using namespace UE::RHI::GPUProfiler;
+
+			auto& Event = BreadcrumbEvents.Emplace_GetRef(MakeUnique<FBreadcrumbEvent>(FBreadcrumbEvent::FEndBreadcrumb{ .Breadcrumb = Breadcrumb }));
+
+			FD3D12QueryLocation TimestampQuery = AllocateQuery(ED3D12QueryType::ProfilerTimestampBOP, &Event->Value.Get<FBreadcrumbEvent::FEndBreadcrumb>().GPUTimestampBOP);
+			EndQuery(TimestampQuery);
+		}
+	#else
 		if (IsDefaultContext() && !IsAsyncComputeContext())
 		{
 			FD3D12GPUProfiler& GPUProfiler = GetParentDevice()->GetGPUProfiler();
@@ -218,6 +239,7 @@ void FD3D12ContextCommon::BindDiagnosticBuffer(FD3D12RootSignature const* RootSi
 				GPUProfiler.PopEvent();
 			}
 		}
+	#endif // (RHI_NEW_GPU_PROFILER == 0)
 
 		// Only emit formatted strings to platform APIs when requested.
 		if (ShouldEmitBreadcrumbs())
@@ -313,6 +335,10 @@ FD3D12QueryLocation FD3D12ContextCommon::AllocateQuery(ED3D12QueryType Type, voi
 
 	case ED3D12QueryType::AdjustedRaw:
 	case ED3D12QueryType::AdjustedMicroseconds:
+#if RHI_NEW_GPU_PROFILER
+	case ED3D12QueryType::ProfilerTimestampTOP:
+	case ED3D12QueryType::ProfilerTimestampBOP:
+#endif
 		return TimestampQueries.Allocate(Type, Target);
 
 	case ED3D12QueryType::Occlusion:
@@ -398,9 +424,13 @@ void FD3D12ContextCommon::CloseCommandList()
 	CommandList->Close();
 	CommandList = nullptr;
 
-	TimestampQueries.CloseAndReset(Payload->QueryRanges);
-	OcclusionQueries.CloseAndReset(Payload->QueryRanges);
-	PipelineStatsQueries.CloseAndReset(Payload->QueryRanges);
+	TimestampQueries    .CloseAndReset(Payload->BatchedObjects.QueryRanges);
+	OcclusionQueries    .CloseAndReset(Payload->BatchedObjects.QueryRanges);
+	PipelineStatsQueries.CloseAndReset(Payload->BatchedObjects.QueryRanges);
+
+#if RHI_NEW_GPU_PROFILER
+	Payload->BatchedObjects.BreadcrumbEvents.Append(MoveTemp(BreadcrumbEvents));
+#endif
 }
 
 void FD3D12CommandContext::CloseCommandList()
@@ -482,7 +512,7 @@ FD3D12QueryLocation FD3D12QueryAllocator::Allocate(ED3D12QueryType Type, void* T
 	check(Type != ED3D12QueryType::None);
 
 	// Allocate a new heap if needed
-	if (Ranges.Num() == 0 || Ranges.Last().IsFull())
+	if (!CurrentRange || CurrentRange->IsFull(CurrentHeap))
 	{
 		TRefCountPtr<FD3D12QueryHeap> Heap = Device->ObtainQueryHeap(QueueType, QueryType);
 		if (!Heap)
@@ -491,38 +521,44 @@ FD3D12QueryLocation FD3D12QueryAllocator::Allocate(ED3D12QueryType Type, void* T
 			return {};
 		}
 
-		FD3D12QueryRange& Range = Ranges.Emplace_GetRef();
-		Range.Heap = MoveTemp(Heap);
+		CurrentHeap = Heap;
+		CurrentRange = &Heaps.FindOrAdd(MoveTemp(Heap));
 	}
 
-	FD3D12QueryRange& Range = Ranges.Last();
 	return FD3D12QueryLocation(
-		Range.Heap,
-		Range.End++,
+		CurrentHeap,
+		CurrentRange->End++,
 		Type,
 		Target
 	);
 }
 
-void FD3D12QueryAllocator::CloseAndReset(TArray<FD3D12QueryRange>& OutRanges)
+void FD3D12QueryAllocator::CloseAndReset(TMap<TRefCountPtr<FD3D12QueryHeap>, TArray<FD3D12QueryRange>>& OutRanges)
 {
 	if (HasQueries())
 	{
-		OutRanges.Append(Ranges);
+		for (auto const& Pair : Heaps)
+		{
+			OutRanges.FindOrAdd(Pair.Key).Emplace(Pair.Value);
+		}
 
-		if (Ranges.Last().IsFull())
+		if (CurrentRange->IsFull(CurrentHeap))
 		{
 			// No space in any heap. Reset the whole array.
-			Ranges.Reset();
+			Heaps.Reset();
+
+			CurrentRange = nullptr;
+			CurrentHeap = nullptr;
 		}
 		else
 		{
 			// The last heap still has space. Reuse it for the next batch of command lists.
-			FD3D12QueryRange LastRange = MoveTemp(Ranges.Last());
+			FD3D12QueryRange LastRange = *CurrentRange;
 			LastRange.Start = LastRange.End;
 
-			Ranges.Reset();
-			Ranges.Emplace(MoveTemp(LastRange));
+			Heaps.Reset();
+			CurrentRange = &Heaps.FindOrAdd(CurrentHeap);
+			*CurrentRange = LastRange;
 		}
 	}
 }
@@ -624,7 +660,9 @@ void FD3D12DynamicRHI::RHIBeginFrame(FRHICommandListImmediate& RHICmdList)
 			{
 				for (auto& Device : Adapter->GetDevices())
 				{
+#if (RHI_NEW_GPU_PROFILER == 0)
 					Device->GetGPUProfiler().BeginFrame();
+#endif
 
 					Device->GetDefaultBufferAllocator().BeginFrame(Contexts);
 					Device->GetTextureAllocator().BeginFrame(Contexts);
@@ -636,8 +674,6 @@ void FD3D12DynamicRHI::RHIBeginFrame(FRHICommandListImmediate& RHICmdList)
 
 void FD3D12CommandContext::RHIBeginFrame()
 {
-	bTrackingEvents = IsDefaultContext() && Device->GetGPUProfiler().bTrackingEvents;
-
 #if D3D12_RHI_RAYTRACING
 	Device->GetRayTracingCompactionRequestHandler()->Update(*this);
 #endif // D3D12_RHI_RAYTRACING
@@ -723,7 +759,9 @@ void FD3D12CommandContextBase::RHIEndFrame()
 	for (uint32 GPUIndex : GPUMask)
 	{
 		Device = ParentAdapter->GetDevice(GPUIndex);
+#if (RHI_NEW_GPU_PROFILER == 0)
 		Device->GetGPUProfiler().EndFrame();
+#endif
 	}
 
 	// Close the previous frame's timing and start a new one
diff --git a/Engine/Source/Runtime/D3D12RHI/Private/D3D12CommandContext.h b/Engine/Source/Runtime/D3D12RHI/Private/D3D12CommandContext.h
--- a/Engine/Source/Runtime/D3D12RHI/Private/D3D12CommandContext.h
+++ b/Engine/Source/Runtime/D3D12RHI/Private/D3D12CommandContext.h
@@ -20,6 +20,8 @@ D3D12CommandContext.h: D3D12 Command Context Interfaces
 #include "RHICoreShader.h"
 #include "RHICore.h"
 
+#include "GPUProfiler.h"
+
 #if USE_PIX
 	#include "Windows/AllowWindowsPlatformTypes.h"
 	THIRD_PARTY_INCLUDES_START
@@ -329,6 +331,10 @@ protected:
 		}
 	}
 
+#if RHI_NEW_GPU_PROFILER
+	TArray<TUniquePtr<UE::RHI::GPUProfiler::FBreadcrumbEvent>> BreadcrumbEvents; // @todo - new gpu profiler - optimize storage / allocation
+#endif
+
 public:
 	// Flushes any pending commands in this context to the GPU.
 	void FlushCommands(ED3D12FlushFlags FlushFlags = ED3D12FlushFlags::None);
@@ -556,8 +562,6 @@ public:
 	/** When a new compute PSO is set, we discard all old constants set for the previous shader. */
 	bool bDiscardSharedComputeConstants = false;
 
-	bool bTrackingEvents = false;
-
 	/** Used by variable rate shading to cache the current state of the combiners and the constant shading rate*/
 #if PLATFORM_SUPPORTS_VARIABLE_RATE_SHADING
 	static_assert(D3D12_RS_SET_SHADING_RATE_COMBINER_COUNT == ED3D12VRSCombinerStages::Num);
@@ -649,7 +653,9 @@ public:
 	virtual void RHICopyToStagingBuffer(FRHIBuffer* SourceBuffer, FRHIStagingBuffer* DestinationStagingBuffer, uint32 Offset, uint32 NumBytes) final override;
 	virtual void RHIBeginRenderQuery(FRHIRenderQuery* RenderQuery) final override;
 	virtual void RHIEndRenderQuery(FRHIRenderQuery* RenderQuery) final override;
+#if (RHI_NEW_GPU_PROFILER == 0)
 	virtual void RHICalibrateTimers(FRHITimestampCalibrationQuery* CalibrationQuery) final override;
+#endif
 	virtual void RHIBeginScene() final override;
 	virtual void RHIEndScene() final override;
 	virtual void RHISetStreamSource(uint32 StreamIndex, FRHIBuffer* VertexBuffer, uint32 Offset) final override;
@@ -937,10 +943,12 @@ public:
 	{
 		ContextRedirect(RHIEndRenderQuery(RenderQuery));
 	}
+#if (RHI_NEW_GPU_PROFILER == 0)
 	FORCEINLINE virtual void RHICalibrateTimers(FRHITimestampCalibrationQuery* CalibrationQuery) final override
 	{
 		ContextRedirect(RHICalibrateTimers(CalibrationQuery));
 	}
+#endif
 	FORCEINLINE virtual void RHIBeginScene() final override
 	{
 		ContextRedirect(RHIBeginScene());
diff --git a/Engine/Source/Runtime/D3D12RHI/Private/D3D12CommandList.cpp b/Engine/Source/Runtime/D3D12RHI/Private/D3D12CommandList.cpp
--- a/Engine/Source/Runtime/D3D12RHI/Private/D3D12CommandList.cpp
+++ b/Engine/Source/Runtime/D3D12RHI/Private/D3D12CommandList.cpp
@@ -369,6 +369,9 @@ void FD3D12CommandList::EndQuery(FD3D12QueryLocation const& Location)
 
 			case ED3D12QueryType::CommandListBegin:
 			case ED3D12QueryType::IdleBegin:
+#if RHI_NEW_GPU_PROFILER
+			case ED3D12QueryType::ProfilerTimestampTOP:
+#endif
 				Position = ED3D12QueryPosition::TopOfPipe;
 				break;
 
@@ -376,6 +379,9 @@ void FD3D12CommandList::EndQuery(FD3D12QueryLocation const& Location)
 			case ED3D12QueryType::IdleEnd:
 			case ED3D12QueryType::AdjustedMicroseconds:
 			case ED3D12QueryType::AdjustedRaw:
+#if RHI_NEW_GPU_PROFILER
+			case ED3D12QueryType::ProfilerTimestampBOP:
+#endif
 				Position = ED3D12QueryPosition::BottomOfPipe;
 				break;
 			}
diff --git a/Engine/Source/Runtime/D3D12RHI/Private/D3D12Commands.cpp b/Engine/Source/Runtime/D3D12RHI/Private/D3D12Commands.cpp
--- a/Engine/Source/Runtime/D3D12RHI/Private/D3D12Commands.cpp
+++ b/Engine/Source/Runtime/D3D12RHI/Private/D3D12Commands.cpp
@@ -1400,12 +1400,14 @@ void FD3D12CommandContext::SetRenderTargetsAndClear(const FRHISetRenderTargetsIn
 #endif
 }
 
+#if (RHI_NEW_GPU_PROFILER == 0)
 void FD3D12CommandContext::RHICalibrateTimers(FRHITimestampCalibrationQuery* CalibrationQuery)
 {
 	FGPUTimingCalibrationTimestamp Timestamp = GetParentDevice()->GetCalibrationTimestamp(QueueType);
 	CalibrationQuery->CPUMicroseconds[GetGPUIndex()] = Timestamp.CPUMicroseconds;
 	CalibrationQuery->GPUMicroseconds[GetGPUIndex()] = Timestamp.GPUMicroseconds;
 }
+#endif
 
 // Primitive drawing.
 
@@ -1540,10 +1542,12 @@ void FD3D12CommandContext::RHIDispatchShaderBundle(
 
 void FD3D12CommandContext::SetupDraw(FRHIBuffer* IndexBufferRHI, uint32 NumPrimitives /* = 0 */, uint32 NumVertices /* = 0 */)
 {
-	if (bTrackingEvents)
+#if (RHI_NEW_GPU_PROFILER == 0)
+	if (IsDefaultContext() && Device->GetGPUProfiler().bTrackingEvents)
 	{
 		GetParentDevice()->RegisterGPUWork(NumPrimitives, NumVertices);
 	}
+#endif
 
 	CommitGraphicsResourceTables();
 	CommitNonComputeShaderConstants();
@@ -1565,10 +1569,12 @@ void FD3D12CommandContext::SetupDraw(FRHIBuffer* IndexBufferRHI, uint32 NumPrimi
 
 void FD3D12CommandContext::SetupDispatchDraw(uint32 ThreadGroupCountX, uint32 ThreadGroupCountY, uint32 ThreadGroupCountZ)
 {
-	if (bTrackingEvents)
+#if (RHI_NEW_GPU_PROFILER == 0)
+	if (IsDefaultContext() && Device->GetGPUProfiler().bTrackingEvents)
 	{
 		GetParentDevice()->RegisterGPUDispatch(FIntVector(ThreadGroupCountX, ThreadGroupCountY, ThreadGroupCountZ));
 	}
+#endif
 
 	CommitGraphicsResourceTables();
 	CommitNonComputeShaderConstants();
diff --git a/Engine/Source/Runtime/D3D12RHI/Private/D3D12Device.cpp b/Engine/Source/Runtime/D3D12RHI/Private/D3D12Device.cpp
--- a/Engine/Source/Runtime/D3D12RHI/Private/D3D12Device.cpp
+++ b/Engine/Source/Runtime/D3D12RHI/Private/D3D12Device.cpp
@@ -77,7 +77,9 @@ FD3D12Queue::~FD3D12Queue()
 FD3D12Device::FD3D12Device(FRHIGPUMask InGPUMask, FD3D12Adapter* InAdapter)
 	: FD3D12SingleNodeGPUObject(InGPUMask)
 	, FD3D12AdapterChild       (InAdapter)
+#if (RHI_NEW_GPU_PROFILER == 0)
 	, GPUProfilingData         (this)
+#endif
 	, ResidencyManager         (*this)
 	, DescriptorHeapManager    (this)
 #if PLATFORM_SUPPORTS_BINDLESS_RENDERING
@@ -507,16 +509,6 @@ void FD3D12Device::UpdateMSAASettings()
 	AvailableMSAAQualities[8] = 0;
 }
 
-void FD3D12Device::RegisterGPUWork(uint32 NumPrimitives, uint32 NumVertices)
-{
-	GetGPUProfiler().RegisterGPUWork(NumPrimitives, NumVertices);
-}
-
-void FD3D12Device::RegisterGPUDispatch(FIntVector GroupCount)
-{
-	GetGPUProfiler().RegisterGPUDispatch(GroupCount);
-}
-
 void FD3D12Device::BlockUntilIdle()
 {
 	// Submit a new sync point to each queue
@@ -756,6 +748,7 @@ uint64 FD3D12Device::GetTimestampFrequency(ED3D12QueueType QueueType)
 	return Frequency;
 }
 
+#if (RHI_NEW_GPU_PROFILER == 0)
 FGPUTimingCalibrationTimestamp FD3D12Device::GetCalibrationTimestamp(ED3D12QueueType QueueType)
 {
 	TRACE_CPUPROFILER_EVENT_SCOPE(D3D12GetCalibrationTimestamp);
@@ -775,6 +768,7 @@ FGPUTimingCalibrationTimestamp FD3D12Device::GetCalibrationTimestamp(ED3D12Queue
 
 	return Result;
 }
+#endif // (RHI_NEW_GPU_PROFILER == 0)
 
 void FD3D12Device::InitExplicitDescriptorHeap()
 {
diff --git a/Engine/Source/Runtime/D3D12RHI/Private/D3D12Device.h b/Engine/Source/Runtime/D3D12RHI/Private/D3D12Device.h
--- a/Engine/Source/Runtime/D3D12RHI/Private/D3D12Device.h
+++ b/Engine/Source/Runtime/D3D12RHI/Private/D3D12Device.h
@@ -104,6 +104,18 @@ public:
 
 	D3D12_QUERY_DATA_PIPELINE_STATISTICS PipelineStats {};
 
+#if RHI_NEW_GPU_PROFILER
+	// Timer calibration data
+	uint64 GPUFrequency = 0, GPUTimestamp = 0;
+	uint64 CPUFrequency = 0, CPUTimestamp = 0;
+
+#if WITH_RHI_BREADCRUMBS
+	TArray<TSharedPtr<FRHIBreadcrumbAllocatorArray>> BreadcrumbAllocators {};
+	TArray<TUniquePtr<UE::RHI::GPUProfiler::FBreadcrumbEvent>> BreadcrumbEvents;
+#endif
+
+#endif
+
 	uint64 GetCurrentTimestamp()  const { return Timestamps[TimestampIndex]; }
 	uint64 GetPreviousTimestamp() const { return Timestamps[TimestampIndex - 1]; }
 
@@ -149,11 +161,7 @@ public:
 
 	uint32 NumCommandListsInBatch = 0;
 
-	// Query ranges/locations to be resolved when the submission thread receives a command list which is still open.
-	TMap<FD3D12QueryHeap*, TArray<FD3D12QueryRange>> PendingQueryRanges;
-	TArray<FD3D12QueryLocation> PendingTimestampQueries;
-	TArray<FD3D12QueryLocation> PendingOcclusionQueries;
-	TArray<FD3D12QueryLocation> PendingPipelineStatsQueries;
+	FD3D12BatchedPayloadObjects BatchedObjects;
 
 	// A pool of reusable command list/allocator/context objects
 	struct
@@ -175,7 +183,7 @@ public:
 	bool bSupportsTileMapping = true;
 
 	// Batches the current payload's command lists, returning the latest fence value signaled for this queue.
-	uint64 FinalizePayload(bool bRequiresSignal);
+	uint64 FinalizePayload(bool bRequiresSignal, TArray<FD3D12Payload*, TInlineAllocator<64>>& PayloadsToHandDown);
 
 	// Ensures all prior batched command lists have reached the driver ID3D12Queue object.
 	void FlushBatchedPayloads();
@@ -204,14 +212,22 @@ public:
 
 	ID3D12Device* GetDevice();
 
+#if RHI_NEW_GPU_PROFILER
+	void RegisterGPUWork(uint32 NumPrimitives = 0, uint32 NumVertices = 0)	{ /*checkNoEntry();*/ } // @todo - new gpu profiler
+	void RegisterGPUDispatch(FIntVector GroupCount)	                        { /*checkNoEntry();*/ } // @todo - new gpu profiler
+#else
+	void RegisterGPUWork(uint32 NumPrimitives = 0, uint32 NumVertices = 0)	{ GPUProfilingData.RegisterGPUWork(NumPrimitives, NumVertices); }
+	void RegisterGPUDispatch(FIntVector GroupCount)	                        { GPUProfilingData.RegisterGPUDispatch(GroupCount); }
+
 	// GPU Profiler
 	FORCEINLINE FD3D12GPUProfiler& GetGPUProfiler() { return GPUProfilingData; }
-
-	void RegisterGPUWork(uint32 NumPrimitives = 0, uint32 NumVertices = 0);
-	void RegisterGPUDispatch(FIntVector GroupCount);
+#endif
 
 	uint64 GetTimestampFrequency(ED3D12QueueType QueueType);
+
+#if (RHI_NEW_GPU_PROFILER == 0)
 	FGPUTimingCalibrationTimestamp GetCalibrationTimestamp(ED3D12QueueType QueueType);
+#endif
 
 	// Misc
 	void BlockUntilIdle();
@@ -322,7 +338,9 @@ private:
 	void UpdateMSAASettings();
 	void UpdateConstantBufferPageProperties();
 
+#if (RHI_NEW_GPU_PROFILER == 0)
 	FD3D12GPUProfiler GPUProfilingData;
+#endif
 
 	struct FResidencyManager : public FD3D12ResidencyManager
 	{
diff --git a/Engine/Source/Runtime/D3D12RHI/Private/D3D12GPUProfiler.cpp b/Engine/Source/Runtime/D3D12RHI/Private/D3D12GPUProfiler.cpp
--- a/Engine/Source/Runtime/D3D12RHI/Private/D3D12GPUProfiler.cpp
+++ b/Engine/Source/Runtime/D3D12RHI/Private/D3D12GPUProfiler.cpp
@@ -17,6 +17,8 @@ static FAutoConsoleVariableRef CVarD3D12RHIStablePowerState(
 	ECVF_Default
 );
 
+#if (RHI_NEW_GPU_PROFILER == 0)
+
 FD3D12BufferedGPUTiming::FD3D12BufferedGPUTiming(FD3D12Device* InParent)
 	: FD3D12DeviceChild(InParent)
 {
@@ -332,3 +334,4 @@ void FD3D12GPUProfiler::EndFrame()
 	CurrentEventNodeFrame = NULL;
 }
 
+#endif // (RHI_NEW_GPU_PROFILER == 0)
diff --git a/Engine/Source/Runtime/D3D12RHI/Private/D3D12GPUProfiler.h b/Engine/Source/Runtime/D3D12RHI/Private/D3D12GPUProfiler.h
--- a/Engine/Source/Runtime/D3D12RHI/Private/D3D12GPUProfiler.h
+++ b/Engine/Source/Runtime/D3D12RHI/Private/D3D12GPUProfiler.h
@@ -8,6 +8,8 @@
 class FD3D12SyncPoint;
 using FD3D12SyncPointRef = TRefCountPtr<FD3D12SyncPoint>;
 
+#if (RHI_NEW_GPU_PROFILER == 0)
+
 // This class has multiple inheritance but really FGPUTiming is a static class
 class FD3D12BufferedGPUTiming : public FGPUTiming, public FD3D12DeviceChild
 {
@@ -120,3 +122,5 @@ struct FD3D12GPUProfiler : public FGPUProfiler, public FD3D12DeviceChild
 	void BeginFrame();
 	void EndFrame();
 };
+
+#endif // (RHI_NEW_GPU_PROFILER == 0)
diff --git a/Engine/Source/Runtime/D3D12RHI/Private/D3D12Query.h b/Engine/Source/Runtime/D3D12RHI/Private/D3D12Query.h
--- a/Engine/Source/Runtime/D3D12RHI/Private/D3D12Query.h
+++ b/Engine/Source/Runtime/D3D12RHI/Private/D3D12Query.h
@@ -21,6 +21,10 @@ enum class ED3D12QueryType
 	IdleEnd,
 	AdjustedMicroseconds,
 	AdjustedRaw,
+#if RHI_NEW_GPU_PROFILER
+	ProfilerTimestampTOP,
+	ProfilerTimestampBOP,
+#endif
 	Occlusion
 };
 
@@ -92,25 +96,20 @@ private:
 
 struct FD3D12QueryRange
 {
-	TRefCountPtr<FD3D12QueryHeap> Heap;
-	uint32 Start = 0;
-	uint32 End = 0;
+	uint32 Start = 0, End = 0;
 
 	FD3D12QueryRange() = default;
-	FD3D12QueryRange(FD3D12QueryHeap* Heap, uint32 Start, uint32 End)
-	: Heap(Heap)
-	, Start(Start)
-	, End(End)
-	{
-	}
+	FD3D12QueryRange(uint32 Start, uint32 End)
+		: Start(Start)
+		, End(End)
+	{}
 
-	inline bool IsFull() const;
+	inline bool IsFull(FD3D12QueryHeap* Heap) const;
 
 	bool operator == (FD3D12QueryRange const& RHS) const
 	{
-		return Heap  == RHS.Heap
-		&& Start == RHS.Start
-		&& End   == RHS.End;
+		return Start == RHS.Start
+			&& End   == RHS.End;
 	}
 
 	bool operator < (FD3D12QueryRange const& RHS) const
@@ -141,17 +140,13 @@ struct FD3D12QueryLocation
 
 	FD3D12QueryLocation() = default;
 	FD3D12QueryLocation(FD3D12QueryHeap* Heap, uint32 Index, ED3D12QueryType Type, void* Target)
-	: Heap(Heap)
-	, Index(Index)
-	, Type(Type)
-	, Target(Target)
-	{
-	}
+		: Heap	(Heap  )
+		, Index	(Index )
+		, Type	(Type  )
+		, Target(Target)
+	{}
 
-	operator bool() const
-	{
-		return Heap != nullptr;
-	}
+	operator bool() const { return Heap != nullptr; }
 };
 
 inline void FD3D12QueryLocation::CopyResultTo(void* Dst) const
@@ -174,7 +169,7 @@ inline TValueType FD3D12QueryLocation::GetResult() const
 	return Value;
 }
 
-inline bool FD3D12QueryRange::IsFull() const
+inline bool FD3D12QueryRange::IsFull(FD3D12QueryHeap* Heap) const
 {
 	return End >= Heap->NumQueries;
 }
@@ -197,12 +192,18 @@ public:
 	FD3D12QueryLocation Allocate(ED3D12QueryType Type, void* Target);
 
 	// Resets the allocator and returns the used query ranges
-	void CloseAndReset(TArray<FD3D12QueryRange>& OutRanges);
+	void CloseAndReset(TMap<TRefCountPtr<FD3D12QueryHeap>, TArray<FD3D12QueryRange>>& OutRanges);
 
-	bool HasQueries() const	{ return !(Ranges.Num() == 0 || Ranges[0].Start == Ranges[0].End); }
+	bool HasQueries() const
+	{
+		return !(CurrentRange == nullptr || CurrentRange->Start == CurrentRange->End);
+	}
 
 private:
-	TArray<FD3D12QueryRange> Ranges;
+	TMap<TRefCountPtr<FD3D12QueryHeap>, FD3D12QueryRange> Heaps;
+	
+	FD3D12QueryHeap* CurrentHeap = nullptr;
+	FD3D12QueryRange* CurrentRange = nullptr;
 };
 
 /** D3D12 Render query */
diff --git a/Engine/Source/Runtime/D3D12RHI/Private/D3D12RHI.cpp b/Engine/Source/Runtime/D3D12RHI/Private/D3D12RHI.cpp
--- a/Engine/Source/Runtime/D3D12RHI/Private/D3D12RHI.cpp
+++ b/Engine/Source/Runtime/D3D12RHI/Private/D3D12RHI.cpp
@@ -435,7 +435,7 @@ void FD3D12DynamicRHI::EnqueueEndOfPipeTask(TUniqueFunction<void()> TaskFunc, TU
 
 void FD3D12DynamicRHI::FlushTiming(bool bCreateNew)
 {
-	auto Lambda = [this, OldTiming = MoveTemp(CurrentTiming)]()
+	auto Lambda = [this, OldTiming = MoveTemp(CurrentTimingPerQueue)]()
 	{
 		if (OldTiming)
 		{
@@ -445,8 +445,8 @@ void FD3D12DynamicRHI::FlushTiming(bool bCreateNew)
 
 	if (bCreateNew)
 	{
-		CurrentTiming = MakeUnique<TIndirectArray<FD3D12Timing>>();
-		CurrentTiming->Reserve(GD3D12MaxNumQueues);
+		CurrentTimingPerQueue = MakeUnique<TIndirectArray<FD3D12Timing>>();
+		CurrentTimingPerQueue->Reserve(GD3D12MaxNumQueues);
 	}
 
 	EnqueueEndOfPipeTask(MoveTemp(Lambda), [&](FD3D12Payload& Payload)
@@ -455,7 +455,7 @@ void FD3D12DynamicRHI::FlushTiming(bool bCreateNew)
 		{
 			FD3D12Timing* NewTiming = new FD3D12Timing(Payload.Queue);
 			Payload.Timing = NewTiming;
-			CurrentTiming->Add(NewTiming);
+			CurrentTimingPerQueue->Add(NewTiming);
 		}
 		else
 		{
diff --git a/Engine/Source/Runtime/D3D12RHI/Private/D3D12RHIPrivate.h b/Engine/Source/Runtime/D3D12RHI/Private/D3D12RHIPrivate.h
--- a/Engine/Source/Runtime/D3D12RHI/Private/D3D12RHIPrivate.h
+++ b/Engine/Source/Runtime/D3D12RHI/Private/D3D12RHIPrivate.h
@@ -239,9 +239,9 @@ public:
 	// Called when the GPU has crashed. This function will not return.
 	void ProcessInterruptQueueOnGPUCrash();
 
-	TUniquePtr<TIndirectArray<FD3D12Timing>> CurrentTiming;
+	TUniquePtr<TIndirectArray<FD3D12Timing>> CurrentTimingPerQueue;
 	void FlushTiming(bool bCreateNew);
-	void ProcessTimestamps(TIndirectArray<FD3D12Timing>& Timing);
+	void ProcessTimestamps(TIndirectArray<FD3D12Timing>& TimingPerQueue);
 
 	void InitializeSubmissionPipe();
 	void ShutdownSubmissionPipe();
diff --git a/Engine/Source/Runtime/D3D12RHI/Private/D3D12Submission.cpp b/Engine/Source/Runtime/D3D12RHI/Private/D3D12Submission.cpp
--- a/Engine/Source/Runtime/D3D12RHI/Private/D3D12Submission.cpp
+++ b/Engine/Source/Runtime/D3D12RHI/Private/D3D12Submission.cpp
@@ -220,32 +220,64 @@ void FD3D12DynamicRHI::SubmitCommands(TConstArrayView<FD3D12FinalizedCommands*>
 {
 	SCOPED_NAMED_EVENT_TEXT("CommandList_Submit", FColor::Magenta);
 
+#if RHI_NEW_GPU_PROFILER
+
+	TArray<FD3D12Payload*> AllPayloads;
+	for (FD3D12FinalizedCommands* Payloads : Commands)
+	{
+	#if WITH_RHI_BREADCRUMBS
+		TSharedPtr<FRHIBreadcrumbAllocatorArray> BreadcrumbAllocators {};
+		if (Payloads->BreadcrumbAllocators.Num())
+		{
+			BreadcrumbAllocators = MakeShared<FRHIBreadcrumbAllocatorArray>(MoveTemp(Payloads->BreadcrumbAllocators));
+		}
+
+		for (FD3D12Payload* Payload : *Payloads)
+		{
+			Payload->BreadcrumbRange = Payloads->BreadcrumbRange;
+			if (BreadcrumbAllocators.IsValid())
+			{
+				Payload->BatchedObjects.BreadcrumbAllocators.Add(BreadcrumbAllocators);
+			}
+		}
+	#endif
+
+		AllPayloads.Append(MoveTemp(*Payloads));
+		delete Payloads;
+	}
+
+	SubmitPayloads(MoveTemp(AllPayloads));
+
+#else
+
 	TArray<FD3D12Payload*> AllPayloads;
-#if WITH_RHI_BREADCRUMBS
+	#if WITH_RHI_BREADCRUMBS
 	TArray<TSharedPtr<FRHIBreadcrumbAllocator>> BreadcrumbAllocators;
-#endif
+	#endif
 
 	for (FD3D12FinalizedCommands* Payloads : Commands)
 	{
-#if WITH_RHI_BREADCRUMBS
+	#if WITH_RHI_BREADCRUMBS
 		for (FD3D12Payload* Payload : *Payloads)
 		{
 			Payload->BreadcrumbRange = Payloads->BreadcrumbRange;
 		}
-#endif
+	#endif
 
 		AllPayloads.Append(MoveTemp(static_cast<TArray<FD3D12Payload*>&>(*Payloads)));
-#if WITH_RHI_BREADCRUMBS
+	#if WITH_RHI_BREADCRUMBS
 		BreadcrumbAllocators.Append(MoveTemp(Payloads->BreadcrumbAllocators));
-#endif
+	#endif
 		delete Payloads;
 	}
 
 	SubmitPayloads(MoveTemp(AllPayloads));
 
-#if WITH_RHI_BREADCRUMBS
+	#if WITH_RHI_BREADCRUMBS
 	// Enqueue the breadcrumb allocator references for cleanup once all prior payloads have completed on the GPU.
 	DeferredDelete([Array = MoveTemp(BreadcrumbAllocators)]() {});
+	#endif
+
 #endif
 }
 
@@ -382,24 +414,24 @@ FD3D12DynamicRHI::FProcessResult FD3D12DynamicRHI::ProcessSubmissionQueue()
 						FD3D12Queue& TargetQueue = CommandList->Device->GetQueue(CommandList->QueueType);
 
 						// Occlusion + Pipeline Stats Queries
-						TargetQueue.PendingOcclusionQueries.Append(MoveTemp(CommandList->State.OcclusionQueries));
-						TargetQueue.PendingPipelineStatsQueries.Append(MoveTemp(CommandList->State.PipelineStatsQueries));
+						TargetQueue.BatchedObjects.OcclusionQueries.Append(MoveTemp(CommandList->State.OcclusionQueries));
+						TargetQueue.BatchedObjects.PipelineStatsQueries.Append(MoveTemp(CommandList->State.PipelineStatsQueries));
 
 						// Timestamp Queries
 						// Keep only the first Begin() in the batch
 						if (TargetQueue.NumCommandListsInBatch++ == 0)
 						{
-							TargetQueue.PendingTimestampQueries.Emplace(MoveTemp(CommandList->State.BeginTimestamp));
+							TargetQueue.BatchedObjects.TimestampQueries.Emplace(MoveTemp(CommandList->State.BeginTimestamp));
 						}
 						else
 						{
 							// Remove the previous End() timestamp, to join the range together.
-							check(TargetQueue.PendingTimestampQueries.Last().Type == ED3D12QueryType::CommandListEnd);
-							TargetQueue.PendingTimestampQueries.RemoveAt(TargetQueue.PendingTimestampQueries.Num() - 1);
+							check(TargetQueue.BatchedObjects.TimestampQueries.Last().Type == ED3D12QueryType::CommandListEnd);
+							TargetQueue.BatchedObjects.TimestampQueries.RemoveAt(TargetQueue.BatchedObjects.TimestampQueries.Num() - 1);
 						}
 
-						TargetQueue.PendingTimestampQueries.Append(MoveTemp(CommandList->State.TimestampQueries));
-						TargetQueue.PendingTimestampQueries.Emplace(MoveTemp(CommandList->State.EndTimestamp));
+						TargetQueue.BatchedObjects.TimestampQueries.Append(MoveTemp(CommandList->State.TimestampQueries));
+						TargetQueue.BatchedObjects.TimestampQueries.Emplace(MoveTemp(CommandList->State.EndTimestamp));
 
 						if (TargetQueue.NumCommandListsInBatch >= MaxBatchSize)
 						{
@@ -444,138 +476,6 @@ FD3D12DynamicRHI::FProcessResult FD3D12DynamicRHI::ProcessSubmissionQueue()
 						AccumulateQueries(CurrentCommandList);
 					}
 				}
-				
-				// Prepare the command lists from each payload for submission
-				for (FD3D12Queue* Queue : QueuesWithPayloads)
-				{
-					FD3D12Payload* Payload = Queue->PayloadToSubmit;
-					check(Payload->SyncPointsToWait.Num() == 0);
-
-					Queue->BarrierTimestamps.CloseAndReset(Payload->QueryRanges);
-					Queue->NumCommandListsInBatch = 0;
-
-					// Gather query ranges from this payload, grouping by heap pointer
-					for (FD3D12QueryRange& Range : Payload->QueryRanges)
-					{
-						check(Range.End > Range.Start);
-						CurrentQueue.PendingQueryRanges.FindOrAdd(Range.Heap).Emplace(MoveTemp(Range));
-					}
-					Payload->QueryRanges.Reset();
-
-					check(Payload->TimestampQueries.Num() == 0);
-					check(Payload->OcclusionQueries.Num() == 0);
-					check(Payload->PipelineStatsQueries.Num() == 0);
-
-					if (Queue->PendingQueryRanges.Num())
-					{
-						// If this payload will signal a CPU-visible sync point, we need to resolve queries.
-						// This makes sure that the query data has reached the CPU before the sync point the CPU is waiting on is signaled.
-						bool bResolveQueries = false;
-						for (FD3D12SyncPoint* SyncPoint : Payload->SyncPointsToSignal)
-						{
-							if (SyncPoint->GetType() == ED3D12SyncPointType::GPUAndCPU)
-							{
-								bResolveQueries = true;
-								break;
-							}
-						}
-
-						if (bResolveQueries)
-						{
-							FD3D12CommandList* ResolveCommandList = nullptr;
-							{
-								// We've got queries to resolve. Allocate a command list.
-								auto GetResolveCommandList = [&]() -> FD3D12CommandList*
-								{
-									if (ResolveCommandList)
-										return ResolveCommandList;
-
-									if (!Queue->BarrierAllocator)
-										Queue->BarrierAllocator = Queue->Device->ObtainCommandAllocator(Queue->QueueType); 
-
-									return ResolveCommandList = Queue->Device->ObtainCommandList(Queue->BarrierAllocator, nullptr, nullptr);
-								};
-
-								// Ranges are grouped by heap pointer.
-								for (auto& [Heap, Ranges] : Queue->PendingQueryRanges)
-								{
-									{
-								#if ENABLE_RESIDENCY_MANAGEMENT
-										TArray<FD3D12ResidencyHandle*, TInlineAllocator<2>> ResidencyHandles;
-										ResidencyHandles.Add(&Heap->GetHeapResidencyHandle());
-										ResidencyHandles.Append(Heap->GetResultBuffer()->GetResidencyHandles());
-										GetResolveCommandList()->UpdateResidency(ResidencyHandles);
-								#endif // ENABLE_RESIDENCY_MANAGEMENT
-									}
-
-									// Sort the ranges into ascending order so we can merge adjacent ones,
-									// to reduce the number of ResolveQueryData calls we need to make.
-									Ranges.Sort();
-
-									int32 Index = 0;
-									uint32 Start, End;
-
-									auto GetNextRange = [&]()
-									{
-										if (Index >= Ranges.Num())
-											return false;
-
-										Start = Ranges[Index].Start;
-										End = Ranges[Index].End;
-										Index++;
-
-										while (Index < Ranges.Num() && Ranges[Index].Start == End)
-										{
-											// Ranges are contiguous. Extend.
-											End = Ranges[Index].End;
-											Index++;
-										}
-
-										return true;
-									};
-
-									while (GetNextRange())
-									{
-										if (Heap->GetD3DQueryHeap())
-										{
-											GetResolveCommandList()->GraphicsCommandList()->ResolveQueryData(
-												Heap->GetD3DQueryHeap(),
-												Heap->QueryType,
-												Start,
-												End - Start,
-												Heap->GetResultBuffer()->GetResource(),
-												Start * Heap->GetResultSize()
-											);
-										}
-
-										Payload->QueryRanges.Emplace(Heap, Start, End);
-									}
-								}
-
-								Queue->PendingQueryRanges.Reset();
-							}
-
-							Payload->TimestampQueries     = MoveTemp(Queue->PendingTimestampQueries    );
-							Payload->OcclusionQueries     = MoveTemp(Queue->PendingOcclusionQueries    );
-							Payload->PipelineStatsQueries = MoveTemp(Queue->PendingPipelineStatsQueries);
-
-							if (ResolveCommandList)
-							{
-								ResolveCommandList->Close();
-								Payload->CommandListsToExecute.Add(ResolveCommandList);
-							}
-						}
-					}
-
-					if (Queue->BarrierAllocator)
-					{
-						Payload->AllocatorsToRelease.Add(Queue->BarrierAllocator);
-						Queue->BarrierAllocator = nullptr;
-					}
-
-					// Hand payload down to the interrupt thread.
-					PayloadsToHandDown.Add(Payload);
-				}
 
 				// Queues with work to submit other than the current one (CurrentQueue) are performing barrier operations.
 				// Submit this work first, followed by a fence signal + enqueued wait.
@@ -583,13 +483,13 @@ FD3D12DynamicRHI::FProcessResult FD3D12DynamicRHI::ProcessSubmissionQueue()
 				{
 					if (OtherQueue != &CurrentQueue)
 					{
-						uint64 ValueSignaled = OtherQueue->FinalizePayload(true);
+						uint64 ValueSignaled = OtherQueue->FinalizePayload(true, PayloadsToHandDown);
 						CurrentQueue.PayloadToSubmit->AddQueueFenceWait(OtherQueue->Fence, ValueSignaled);
 					}
 				}
 
 				// Now submit the original payload
-				CurrentQueue.FinalizePayload(false);
+				CurrentQueue.FinalizePayload(false, PayloadsToHandDown);
 			}
 		});
 	} while (bProgress);
@@ -774,11 +674,137 @@ void FD3D12DynamicRHI::GenerateBarrierCommandListAndUpdateState(FD3D12CommandLis
 	}
 }
 
-uint64 FD3D12Queue::FinalizePayload(bool bRequiresSignal)
+uint64 FD3D12Queue::FinalizePayload(bool bRequiresSignal, TArray<FD3D12Payload*, TInlineAllocator<64>>& PayloadsToHandDown)
 {
 	TRACE_CPUPROFILER_EVENT_SCOPE(ExecuteCommandList);
 	LLM_SCOPE_BYNAME(TEXT("RHIMisc/ExecuteCommandLists"));
+
 	check(PayloadToSubmit && this == &PayloadToSubmit->Queue);
+	check(PayloadToSubmit->SyncPointsToWait.Num() == 0);
+
+	NumCommandListsInBatch = 0;
+
+	BarrierTimestamps.CloseAndReset(PayloadToSubmit->BatchedObjects.QueryRanges);
+
+#if WITH_RHI_BREADCRUMBS && RHI_NEW_GPU_PROFILER
+	if (PayloadToSubmit->BatchedObjects.BreadcrumbAllocators.Num())
+	{
+		BatchedObjects.BreadcrumbAllocators.Append(MoveTemp(PayloadToSubmit->BatchedObjects.BreadcrumbAllocators));
+	}
+
+	BatchedObjects.BreadcrumbEvents.Append(MoveTemp(PayloadToSubmit->BatchedObjects.BreadcrumbEvents));
+#endif
+
+	// Gather query ranges from this payload, grouping by heap pointer
+	if (BatchedObjects.QueryRanges.Num())
+	{
+		for (auto& [Heap, Ranges] : PayloadToSubmit->BatchedObjects.QueryRanges)
+		{
+			BatchedObjects.QueryRanges.FindOrAdd(Heap).Append(MoveTemp(Ranges));
+		}
+		PayloadToSubmit->BatchedObjects.QueryRanges.Reset();
+	}
+	else
+	{
+		BatchedObjects.QueryRanges = MoveTemp(PayloadToSubmit->BatchedObjects.QueryRanges);
+	}
+
+	check(PayloadToSubmit->BatchedObjects.IsEmpty());
+
+	if (!BatchedObjects.IsEmpty())
+	{
+		// Always resolve queries if we're switching the Timing struct,
+		// since we need to gather the timestamps for that frame.
+		bool bResolveQueries = PayloadToSubmit->Timing.IsSet();
+
+		if (!bResolveQueries)
+		{
+			// If this payload will signal a CPU-visible sync point, we need to resolve queries.
+			// This makes sure that the query data has reached the CPU before the sync point the CPU is waiting on is signaled.
+			for (FD3D12SyncPoint* SyncPoint : PayloadToSubmit->SyncPointsToSignal)
+			{
+				if (SyncPoint->GetType() == ED3D12SyncPointType::GPUAndCPU)
+				{
+					bResolveQueries = true;
+					break;
+				}
+			}
+		}
+
+		if (bResolveQueries)
+		{
+			{
+				FD3D12CommandList* ResolveCommandList = nullptr;
+
+				// We've got queries to resolve. Allocate a command list.
+				auto GetResolveCommandList = [&]() -> FD3D12CommandList*
+				{
+					if (ResolveCommandList)
+						return ResolveCommandList;
+
+					if (!BarrierAllocator)
+						BarrierAllocator = Device->ObtainCommandAllocator(QueueType);
+
+					return ResolveCommandList = Device->ObtainCommandList(BarrierAllocator, nullptr, nullptr);
+				};
+
+				// Ranges are grouped by heap pointer.
+				for (auto& [Heap, Ranges] : BatchedObjects.QueryRanges)
+				{
+					{
+#if ENABLE_RESIDENCY_MANAGEMENT
+						TArray<FD3D12ResidencyHandle*, TInlineAllocator<2>> ResidencyHandles;
+						ResidencyHandles.Add(&Heap->GetHeapResidencyHandle());
+						ResidencyHandles.Append(Heap->GetResultBuffer()->GetResidencyHandles());
+						GetResolveCommandList()->UpdateResidency(ResidencyHandles);
+#endif // ENABLE_RESIDENCY_MANAGEMENT
+					}
+
+					if (Heap->GetD3DQueryHeap())
+					{
+						// Sort the ranges into ascending order so we can merge adjacent ones,
+						// to reduce the number of ResolveQueryData calls we need to make.
+						Ranges.Sort();
+
+						for (int32 Index = 0; Index < Ranges.Num(); )
+						{
+							FD3D12QueryRange Range = Ranges[Index++];
+
+							while (Index < Ranges.Num() && Ranges[Index].Start == Range.End)
+							{
+								// Ranges are contiguous. Extend.
+								Range.End = Ranges[Index++].End;
+							}
+
+							GetResolveCommandList()->GraphicsCommandList()->ResolveQueryData(
+								Heap->GetD3DQueryHeap(),
+								Heap->QueryType,
+								Range.Start,
+								Range.End - Range.Start,
+								Heap->GetResultBuffer()->GetResource(),
+								Range.Start * Heap->GetResultSize()
+							);
+						}
+					}
+				}
+
+				if (ResolveCommandList)
+				{
+					ResolveCommandList->Close();
+					PayloadToSubmit->CommandListsToExecute.Add(ResolveCommandList);
+				}
+			}
+
+			// Move all the batched objects in this queue into the payload, so they get passed down the pipe.
+			PayloadToSubmit->BatchedObjects = MoveTemp(BatchedObjects);
+		}
+	}
+
+	if (BarrierAllocator)
+	{
+		PayloadToSubmit->AllocatorsToRelease.Add(BarrierAllocator);
+		BarrierAllocator = nullptr;
+	}
 
 	// Keep the latest fence value in the submitted payload.
 	// The interrupt thread uses this to determine when work has completed.
@@ -793,6 +819,8 @@ uint64 FD3D12Queue::FinalizePayload(bool bRequiresSignal)
 	}
 
 	PayloadsToSubmit.Add(PayloadToSubmit);
+	PayloadsToHandDown.Add(PayloadToSubmit);
+
 	PayloadToSubmit = nullptr;
 
 	return Fence.NextCompletionValue;
@@ -948,6 +976,22 @@ void FD3D12Queue::FlushBatchedPayloads()
 
 	for (FD3D12Payload* Payload : PayloadsToSubmit)
 	{
+#if RHI_NEW_GPU_PROFILER
+		if (Payload->Timing.IsSet())
+		{
+			Flush();
+
+			{
+				SCOPED_NAMED_EVENT(CalibrateClocks, FColor::Red);
+				FD3D12Timing* LocalTiming = *Payload->Timing;
+				// Calibrate the GPU timestamp / clock
+				VERIFYD3D12RESULT(D3DCommandQueue->GetClockCalibration(&LocalTiming->GPUTimestamp, &LocalTiming->CPUTimestamp));
+				VERIFYD3D12RESULT(D3DCommandQueue->GetTimestampFrequency(&LocalTiming->GPUFrequency));
+				QueryPerformanceFrequency(reinterpret_cast<LARGE_INTEGER*>(&LocalTiming->CPUFrequency));
+			}
+		}
+#endif // RHI_NEW_GPU_PROFILER
+
 		if (Payload->HasWaitWork())
 		{
 			Flush();
@@ -1220,13 +1264,13 @@ FD3D12DynamicRHI::FProcessResult FD3D12DynamicRHI::ProcessInterruptQueue()
 					CurrentQueue.Timing = Payload->Timing.GetValue();
 				}
 
-				for (FD3D12QueryLocation& Query : Payload->OcclusionQueries)
+				for (FD3D12QueryLocation& Query : Payload->BatchedObjects.OcclusionQueries)
 				{
 					check(Query.Target);
 					Query.CopyResultTo(Query.Target);
 				}
 
-				for (FD3D12QueryLocation& Query : Payload->PipelineStatsQueries)
+				for (FD3D12QueryLocation& Query : Payload->BatchedObjects.PipelineStatsQueries)
 				{
 					if (Query.Target)
 					{
@@ -1239,7 +1283,7 @@ FD3D12DynamicRHI::FProcessResult FD3D12DynamicRHI::ProcessInterruptQueue()
 					}
 				}
 
-				if (Payload->TimestampQueries.Num())
+				if (Payload->BatchedObjects.TimestampQueries.Num())
 				{
 					FD3D12QueryLocation* IdleBegin = nullptr;
 					FD3D12QueryLocation* ListBegin = nullptr;
@@ -1247,7 +1291,7 @@ FD3D12DynamicRHI::FProcessResult FD3D12DynamicRHI::ProcessInterruptQueue()
 					// Some timestamp queries report in microseconds
 					const double MicrosecondsScale = 1000000.0 / CurrentQueue.Device->GetTimestampFrequency(CurrentQueue.QueueType);
 
-					for (FD3D12QueryLocation& Query : Payload->TimestampQueries)
+					for (FD3D12QueryLocation& Query : Payload->BatchedObjects.TimestampQueries)
 					{
 						if (Query.Target)
 						{
@@ -1305,6 +1349,21 @@ FD3D12DynamicRHI::FProcessResult FD3D12DynamicRHI::ProcessInterruptQueue()
 								*static_cast<uint64*>(Query.Target) = FPlatformMath::TruncToInt(double(*static_cast<uint64*>(Query.Target)) * MicrosecondsScale);
 							}
 							break;
+
+#if RHI_NEW_GPU_PROFILER
+						case ED3D12QueryType::ProfilerTimestampTOP:
+						case ED3D12QueryType::ProfilerTimestampBOP:
+							{
+								// Convert from GPU timestamp to CPU timestamp (relative to FPlatformTime::Cycles64())
+								uint64& Target = *static_cast<uint64*>(Query.Target);
+
+								uint64 GPUDelta = Target - CurrentQueue.Timing->GPUTimestamp;
+								uint64 CPUDelta = (GPUDelta * CurrentQueue.Timing->CPUFrequency) / CurrentQueue.Timing->GPUFrequency;
+
+								Target = CPUDelta + CurrentQueue.Timing->CPUTimestamp;
+							}
+							break;
+#endif // RHI_NEW_GPU_PROFILER
 						}
 					}
 
@@ -1321,6 +1380,20 @@ FD3D12DynamicRHI::FProcessResult FD3D12DynamicRHI::ProcessInterruptQueue()
 				}
 			}
 
+#if RHI_NEW_GPU_PROFILER && WITH_RHI_BREADCRUMBS
+			if (Payload->BatchedObjects.BreadcrumbAllocators.Num())
+			{
+				check(CurrentQueue.Timing);
+				CurrentQueue.Timing->BreadcrumbAllocators.Append(MoveTemp(Payload->BatchedObjects.BreadcrumbAllocators));
+			}
+
+			if (Payload->BatchedObjects.BreadcrumbEvents.Num())
+			{
+				check(CurrentQueue.Timing);
+				CurrentQueue.Timing->BreadcrumbEvents.Append(MoveTemp(Payload->BatchedObjects.BreadcrumbEvents));
+			}
+#endif
+
 			// We're done with this payload now.
 
 			// GPU resources the payload is holding a reference to will be cleaned up here.
@@ -1364,8 +1437,27 @@ static TAutoConsoleVariable<int32> CVarGPUTimeFromTimestamps(
 	TEXT("Prefer timestamps instead of GetHardwareGPUFrameTime to compute GPU frame time"),
 	ECVF_RenderThreadSafe);
 
-void FD3D12DynamicRHI::ProcessTimestamps(TIndirectArray<FD3D12Timing>& Timing)
+void FD3D12DynamicRHI::ProcessTimestamps(TIndirectArray<FD3D12Timing>& TimingPerQueue)
 {
+#if RHI_NEW_GPU_PROFILER
+	for (FD3D12Timing const& Timing : TimingPerQueue)
+	{
+		UE::RHI::GPUProfiler::FQueue QueueID {};
+		QueueID.GPU = Timing.Queue.Device->GetGPUIndex();
+		QueueID.Index = 0;
+
+		switch (Timing.Queue.QueueType)
+		{
+		default: checkNoEntry(); [[fallthrough]];
+		case ED3D12QueueType::Direct: QueueID.Type = UE::RHI::GPUProfiler::FQueue::EType::Graphics; break;
+		case ED3D12QueueType::Async : QueueID.Type = UE::RHI::GPUProfiler::FQueue::EType::Compute ; break;
+		case ED3D12QueueType::Copy  : QueueID.Type = UE::RHI::GPUProfiler::FQueue::EType::Copy    ; break;
+		}
+		
+		UE::RHI::GPUProfiler::PushEvents(QueueID, Timing.BreadcrumbEvents, {}, {});
+	}
+#endif // RHI_NEW_GPU_PROFILER
+
 	// The total number of cycles where at least one GPU pipe was busy during the frame.
 	uint64 UnionBusyCycles = 0;
 	int32 BusyPipes = 0;
@@ -1378,7 +1470,7 @@ void FD3D12DynamicRHI::ProcessTimestamps(TIndirectArray<FD3D12Timing>& Timing)
 	{
 		// Find the next minimum timestamp
 		FD3D12Timing* NextMin = nullptr;
-		for (FD3D12Timing& Current : Timing)
+		for (FD3D12Timing& Current : TimingPerQueue)
 		{
 			if (Current.HasMoreTimestamps() && (!NextMin || Current.GetCurrentTimestamp() < NextMin->GetCurrentTimestamp()))
 			{
@@ -1416,7 +1508,7 @@ void FD3D12DynamicRHI::ProcessTimestamps(TIndirectArray<FD3D12Timing>& Timing)
 	check(BusyPipes == 0);
 	
 	D3D12_QUERY_DATA_PIPELINE_STATISTICS PipelineStats{};
-	for (FD3D12Timing& Current : Timing)
+	for (FD3D12Timing& Current : TimingPerQueue)
 	{
 		PipelineStats += Current.PipelineStats;
 	}
@@ -1455,7 +1547,7 @@ void FD3D12DynamicRHI::ProcessTimestamps(TIndirectArray<FD3D12Timing>& Timing)
 		GGPUFrameTime = FPlatformMath::TruncToInt(double(UnionBusyCycles) * Scale32);
 	}
 
-	for (FD3D12Timing& Current : Timing)
+	for (FD3D12Timing& Current : TimingPerQueue)
 	{
 		switch (Current.Queue.QueueType)
 		{
diff --git a/Engine/Source/Runtime/D3D12RHI/Private/D3D12Submission.h b/Engine/Source/Runtime/D3D12RHI/Private/D3D12Submission.h
--- a/Engine/Source/Runtime/D3D12RHI/Private/D3D12Submission.h
+++ b/Engine/Source/Runtime/D3D12RHI/Private/D3D12Submission.h
@@ -10,6 +10,7 @@
 #include "Templates/Function.h"
 #include "Templates/RefCounting.h"
 #include "RHIBreadcrumbs.h"
+#include "GPUProfiler.h"
 
 enum class ED3D12QueueType;
 
@@ -140,6 +141,34 @@ struct FD3D12CommitReservedResourceDesc
 	uint64 CommitSizeInBytes = 0;
 };
 
+struct FD3D12BatchedPayloadObjects
+{
+	TArray<FD3D12QueryLocation> TimestampQueries;
+	TArray<FD3D12QueryLocation> OcclusionQueries;
+	TArray<FD3D12QueryLocation> PipelineStatsQueries;
+	TMap<TRefCountPtr<FD3D12QueryHeap>, TArray<FD3D12QueryRange>> QueryRanges;
+
+#if WITH_RHI_BREADCRUMBS && RHI_NEW_GPU_PROFILER
+	TArray<TSharedPtr<FRHIBreadcrumbAllocatorArray>, TInlineAllocator<1>> BreadcrumbAllocators {};
+	TArray<TUniquePtr<UE::RHI::GPUProfiler::FBreadcrumbEvent>> BreadcrumbEvents;
+#endif
+
+	bool IsEmpty() const
+	{
+		return
+			   TimestampQueries    .Num() == 0
+			&& OcclusionQueries    .Num() == 0
+			&& PipelineStatsQueries.Num() == 0
+			&& QueryRanges         .Num() == 0
+#if WITH_RHI_BREADCRUMBS && RHI_NEW_GPU_PROFILER
+			&& BreadcrumbAllocators.Num() == 0
+			&& BreadcrumbEvents    .Num() == 0
+#endif
+		;
+
+	}
+};
+
 // A single unit of work (specific to a single GPU node and queue type) to be processed by the submission thread.
 struct FD3D12PayloadBase
 {
@@ -206,10 +235,8 @@ struct FD3D12PayloadBase
 
 	// Cleanup
 	TArray<FD3D12CommandAllocator*> AllocatorsToRelease;
-	TArray<FD3D12QueryLocation> TimestampQueries;
-	TArray<FD3D12QueryLocation> OcclusionQueries;
-	TArray<FD3D12QueryLocation> PipelineStatsQueries;
-	TArray<FD3D12QueryRange> QueryRanges;
+
+	FD3D12BatchedPayloadObjects BatchedObjects;
 
 #if WITH_RHI_BREADCRUMBS
 	FRHIBreadcrumbRange BreadcrumbRange {};
diff --git a/Engine/Source/Runtime/D3D12RHI/Private/Windows/WindowsD3D12Device.cpp b/Engine/Source/Runtime/D3D12RHI/Private/Windows/WindowsD3D12Device.cpp
--- a/Engine/Source/Runtime/D3D12RHI/Private/Windows/WindowsD3D12Device.cpp
+++ b/Engine/Source/Runtime/D3D12RHI/Private/Windows/WindowsD3D12Device.cpp
@@ -1753,10 +1753,12 @@ void FD3D12DynamicRHI::Init()
 
 	InitializeSubmissionPipe();
 
+#if (RHI_NEW_GPU_PROFILER == 0)
 	for (TSharedPtr<FD3D12Adapter>& Adapter : ChosenAdapters)
 	{
 		FD3D12BufferedGPUTiming::Initialize(Adapter.Get());
 	}
+#endif
 
 	FRenderResource::InitPreRHIResources();
 	GIsRHIInitialized = true;
diff --git a/Engine/Source/Runtime/Launch/Private/LaunchEngineLoop.cpp b/Engine/Source/Runtime/Launch/Private/LaunchEngineLoop.cpp
--- a/Engine/Source/Runtime/Launch/Private/LaunchEngineLoop.cpp
+++ b/Engine/Source/Runtime/Launch/Private/LaunchEngineLoop.cpp
@@ -5124,7 +5124,7 @@ void FEngineLoop::Exit()
 
 	IStreamingManager::Shutdown();
 
-#if HAS_GPU_STATS
+#if HAS_GPU_STATS && (RHI_NEW_GPU_PROFILER == 0)
 	FRealtimeGPUProfiler::SafeRelease();
 #endif
 
diff --git a/Engine/Source/Runtime/OpenGLDrv/Private/OpenGLCommands.cpp b/Engine/Source/Runtime/OpenGLDrv/Private/OpenGLCommands.cpp
--- a/Engine/Source/Runtime/OpenGLDrv/Private/OpenGLCommands.cpp
+++ b/Engine/Source/Runtime/OpenGLDrv/Private/OpenGLCommands.cpp
@@ -2298,7 +2298,7 @@ void FOpenGLDynamicRHI::RHIDrawPrimitive(uint32 BaseVertexIndex,uint32 NumPrimit
 	GLsizei NumElements = 0;
 	FindPrimitiveType(PrimitiveType, NumPrimitives, DrawMode, NumElements);
 
-	GPUProfilingData.RegisterGPUWork(NumPrimitives * NumInstances, VertexCount * NumInstances);
+	RegisterGPUWork(NumPrimitives * NumInstances, VertexCount * NumInstances);
 	if (NumInstances == 1)
 	{
 		SCOPE_CYCLE_COUNTER_DETAILED(STAT_OpenGLDrawPrimitiveDriverTime);
@@ -2321,7 +2321,7 @@ void FOpenGLDynamicRHI::RHIDrawPrimitiveIndirect(FRHIBuffer* ArgumentBufferRHI,
 		VERIFY_GL_SCOPE();
 
 		check(ArgumentBufferRHI);
-		GPUProfilingData.RegisterGPUWork(0);
+		RegisterGPUWork(0);
 
 		FOpenGLContextState& ContextState = GetContextStateForCurrentContext();
 		BindPendingFramebuffer(ContextState);
@@ -2368,7 +2368,7 @@ void FOpenGLDynamicRHI::RHIDrawIndexedIndirect(FRHIBuffer* IndexBufferRHI, FRHIB
 		VERIFY_GL_SCOPE();
 
 		FOpenGLBuffer* IndexBuffer = ResourceCast(IndexBufferRHI);
-		GPUProfilingData.RegisterGPUWork(1);
+		RegisterGPUWork(1);
 
 		check(ArgumentsBufferRHI);
 
@@ -2485,7 +2485,7 @@ void FOpenGLDynamicRHI::RHIDrawIndexedPrimitive(FRHIBuffer* IndexBufferRHI, int3
 	GLenum IndexType = IndexBuffer->GetStride() == sizeof(uint32) ? GL_UNSIGNED_INT : GL_UNSIGNED_SHORT;
 	StartIndex *= IndexBuffer->GetStride() == sizeof(uint32) ? sizeof(uint32) : sizeof(uint16);
 
-	GPUProfilingData.RegisterGPUWork(NumPrimitives * NumInstances, NumElements * NumInstances);
+	RegisterGPUWork(NumPrimitives * NumInstances, NumElements * NumInstances);
 	if (NumInstances > 1)
 	{
 		SCOPE_CYCLE_COUNTER_DETAILED(STAT_OpenGLDrawPrimitiveDriverTime);
@@ -2516,7 +2516,7 @@ void FOpenGLDynamicRHI::RHIDrawIndexedPrimitiveIndirect(FRHIBuffer* IndexBufferR
 		VERIFY_GL_SCOPE();
 
 		FOpenGLBuffer* IndexBuffer = ResourceCast(IndexBufferRHI);
-		GPUProfilingData.RegisterGPUWork(1);
+		RegisterGPUWork(1);
 
 		check(ArgumentBufferRHI);
 
@@ -2639,7 +2639,7 @@ void FOpenGLDynamicRHI::RHIClearMRT(const bool* bClearColorArray,int32 NumClearC
 	bool bPrevScissorEnabled = PendingState.bScissorEnabled;
 
 	bool bScissorChanged = false;
-	GPUProfilingData.RegisterGPUWork(0);
+	RegisterGPUWork(0);
 	FOpenGLContextState& ContextState = GetContextStateForCurrentContext();
 	BindPendingFramebuffer(ContextState);
 
@@ -2820,7 +2820,7 @@ void FOpenGLDynamicRHI::RHIDispatchComputeShader(uint32 ThreadGroupCountX, uint3
 
 	FOpenGLContextState& ContextState = GetContextStateForCurrentContext();
 
-	GPUProfilingData.RegisterGPUDispatch(FIntVector(ThreadGroupCountX, ThreadGroupCountY, ThreadGroupCountZ));	
+	RegisterGPUDispatch(FIntVector(ThreadGroupCountX, ThreadGroupCountY, ThreadGroupCountZ));	
 
 	BindPendingComputeShaderState(ContextState, ComputeShader);
 	CommitComputeResourceTables(ComputeShader);
@@ -2845,7 +2845,7 @@ void FOpenGLDynamicRHI::RHIDispatchIndirectComputeShader(FRHIBuffer* ArgumentBuf
 
 	FOpenGLContextState& ContextState = GetContextStateForCurrentContext();
 
-	GPUProfilingData.RegisterGPUDispatch(FIntVector(1, 1, 1));	
+	RegisterGPUDispatch(FIntVector(1, 1, 1));	
 
 	BindPendingComputeShaderState(ContextState, ComputeShader);
 
diff --git a/Engine/Source/Runtime/OpenGLDrv/Private/OpenGLDevice.cpp b/Engine/Source/Runtime/OpenGLDrv/Private/OpenGLDevice.cpp
--- a/Engine/Source/Runtime/OpenGLDrv/Private/OpenGLDevice.cpp
+++ b/Engine/Source/Runtime/OpenGLDrv/Private/OpenGLDevice.cpp
@@ -211,7 +211,9 @@ FOpenGLContextState& FOpenGLDynamicRHI::GetContextStateForCurrentContext(bool bA
 
 void FOpenGLDynamicRHI::RHIBeginFrame()
 {
+#if (RHI_NEW_GPU_PROFILER == 0)
 	GPUProfilingData.BeginFrame(this);
+#endif
 
 #if PLATFORM_ANDROID //adding #if since not sure if this is required for any other platform.
 	PendingState.DepthStencil = 0 ;
@@ -224,7 +226,9 @@ extern void OpenGLCommands_OnEndFrame();
 
 void FOpenGLDynamicRHI::RHIEndFrame()
 {
+#if (RHI_NEW_GPU_PROFILER == 0)
 	GPUProfilingData.EndFrame();
+#endif
 
 	OpenGL_PollAllFences();
 
@@ -1294,7 +1298,9 @@ FOpenGLDynamicRHI::FOpenGLDynamicRHI()
 ,	bIsRenderingContextAcquired(false)
 ,   BeginSceneContextType(CONTEXT_Other)
 ,	PlatformDevice(NULL)
+#if (RHI_NEW_GPU_PROFILER == 0)
 ,	GPUProfilingData(this)
+#endif
 {
 	check(Singleton == nullptr);
 	Singleton = this;
@@ -1509,7 +1515,9 @@ void FOpenGLDynamicRHI::Cleanup()
 		// Reset the RHI initialized flag.
 		GIsRHIInitialized = false;
 
+#if (RHI_NEW_GPU_PROFILER == 0)
 		GPUProfilingData.Cleanup();
+#endif
 
 		// Ask all initialized FRenderResources to release their RHI resources.
 		FRenderResource::ReleaseRHIForAllResources();
diff --git a/Engine/Source/Runtime/OpenGLDrv/Private/OpenGLDrv.cpp b/Engine/Source/Runtime/OpenGLDrv/Private/OpenGLDrv.cpp
--- a/Engine/Source/Runtime/OpenGLDrv/Private/OpenGLDrv.cpp
+++ b/Engine/Source/Runtime/OpenGLDrv/Private/OpenGLDrv.cpp
@@ -58,7 +58,11 @@ bool FOpenGLDynamicRHI::RHIIsValidTexture(GLuint InTexture) const
 
 void FOpenGLDynamicRHI::RHISetExternalGPUTime(uint32 InExternalGPUTime)
 {
-	GetGPUProfilingData().ExternalGPUTime = InExternalGPUTime;
+#if RHI_NEW_GPU_PROFILER
+	checkNoEntry(); // @todo - new gpu profiler
+#else
+	GPUProfilingData.ExternalGPUTime = InExternalGPUTime;
+#endif
 }
 
 #if PLATFORM_ANDROID
@@ -138,18 +142,22 @@ void FOpenGLDynamicRHI::RHIEGLTerminateContext()
 	#endif
 		}
 
+	#if (RHI_NEW_GPU_PROFILER == 0)
 		if (GPUProfilingData.IsProfilingGPU())
 		{
 			GPUProfilingData.PushEvent(GetNameStr(), FColor::White);
 		}
+	#endif
 	}
 
 	void FOpenGLDynamicRHI::RHIEndBreadcrumbGPU(FRHIBreadcrumbNode* Breadcrumb)
 	{
+	#if (RHI_NEW_GPU_PROFILER == 0)
 		if (GPUProfilingData.IsProfilingGPU())
 		{
 			GPUProfilingData.PopEvent();
 		}
+	#endif
 
 		if (ShouldEmitBreadcrumbs())
 		{
@@ -160,17 +168,6 @@ void FOpenGLDynamicRHI::RHIEGLTerminateContext()
 	}
 #endif // WITH_RHI_BREADCRUMBS
 
-void FOpenGLGPUProfiler::PushEvent(const TCHAR* Name, FColor Color)
-{
-	FGPUProfiler::PushEvent(Name, Color);
-}
-
-void FOpenGLGPUProfiler::PopEvent()
-{
-	FGPUProfiler::PopEvent();
-
-}
-
 // only use shader hashes to determine GL PSO hash;
 uint64 FOpenGLDynamicRHI::RHIComputeStatePrecachePSOHash(const FGraphicsPipelineStateInitializer& Initializer)
 {
@@ -227,6 +224,8 @@ bool FOpenGLDynamicRHI::RHIMatchPrecachePSOInitializers(const FGraphicsPipelineS
 	return true;
 }
 
+#if (RHI_NEW_GPU_PROFILER == 0)
+
 void FOpenGLGPUProfiler::BeginFrame(FOpenGLDynamicRHI* InRHI)
 {
 	if (NestedFrameCount++>0)
@@ -534,6 +533,8 @@ float FOpenGLEventNode::GetTiming()
 	return Result;
 }
 
+#endif // (RHI_NEW_GPU_PROFILER == 0)
+
 void FOpenGLDynamicRHI::InitializeStateResources()
 {
 	VERIFY_GL_SCOPE();
diff --git a/Engine/Source/Runtime/OpenGLDrv/Private/OpenGLDrv.h b/Engine/Source/Runtime/OpenGLDrv/Private/OpenGLDrv.h
--- a/Engine/Source/Runtime/OpenGLDrv/Private/OpenGLDrv.h
+++ b/Engine/Source/Runtime/OpenGLDrv/Private/OpenGLDrv.h
@@ -37,6 +37,8 @@ struct Rect;
 
 template<class T> struct TOpenGLResourceTraits;
 
+#if (RHI_NEW_GPU_PROFILER == 0)
+
 // This class has multiple inheritance but really FGPUTiming is a static class
 class FOpenGLBufferedGPUTiming : public FGPUTiming
 {
@@ -259,13 +261,12 @@ struct FOpenGLGPUProfiler : public FGPUProfiler
 
 	void Cleanup();
 
-	virtual void PushEvent(const TCHAR* Name, FColor Color) override;
-	virtual void PopEvent() override;
-
 	void BeginFrame(class FOpenGLDynamicRHI* InRHI);
 	void EndFrame();
 };
 
+#endif // (RHI_NEW_GPU_PROFILER == 0)
+
 PRAGMA_DISABLE_DEPRECATION_WARNINGS
 
 /** The interface which is implemented by the dynamically bound RHI. */
@@ -779,10 +780,6 @@ public:
 
 	virtual void RHIPostExternalCommandsReset() final override;
 
-	FOpenGLGPUProfiler& GetGPUProfilingData() {
-		return GPUProfilingData;
-	}
-
 	GLuint GetOpenGLFramebuffer(uint32 NumSimultaneousRenderTargets, FOpenGLTexture** RenderTargets, const uint32* ArrayIndices, const uint32* MipmapLevels, FOpenGLTexture* DepthStencilTarget);
 	GLuint GetOpenGLFramebuffer(uint32 NumSimultaneousRenderTargets, FOpenGLTexture** RenderTargets, const uint32* ArrayIndices, const uint32* MipmapLevels, FOpenGLTexture* DepthStencilTarget, int32 NumRenderingSamples);
 	
@@ -852,9 +849,34 @@ private:
 	/** A critical section to protect modifications and iteration over Queries list */
 	FCriticalSection QueriesListCriticalSection;
 
+#if RHI_NEW_GPU_PROFILER
+
+	void RegisterGPUWork(uint32 NumPrimitives = 0, uint32 NumVertices = 0)
+	{
+		checkNoEntry(); // @todo - new gpu profiler
+	}
+
+	void RegisterGPUDispatch(FIntVector GroupCount)
+	{
+		checkNoEntry(); // @todo - new gpu profiler
+	}
+
+#else
+
 	FOpenGLGPUProfiler GPUProfilingData;
 	friend FOpenGLGPUProfiler;
 
+	void RegisterGPUWork(uint32 NumPrimitives = 0, uint32 NumVertices = 0)
+	{
+		GPUProfilingData.RegisterGPUWork(NumPrimitives, NumVertices);
+	}
+	void RegisterGPUDispatch(FIntVector GroupCount)
+	{
+		GPUProfilingData.RegisterGPUDispatch(GroupCount);
+	}
+
+#endif
+
 	FCriticalSection CustomPresentSection;
 	TRefCountPtr<class FRHICustomPresent> CustomPresent;
 
diff --git a/Engine/Source/Runtime/OpenGLDrv/Private/OpenGLQuery.cpp b/Engine/Source/Runtime/OpenGLDrv/Private/OpenGLQuery.cpp
--- a/Engine/Source/Runtime/OpenGLDrv/Private/OpenGLQuery.cpp
+++ b/Engine/Source/Runtime/OpenGLDrv/Private/OpenGLQuery.cpp
@@ -605,11 +605,15 @@ bool FOpenGLDynamicRHI::RHIGetRenderQueryResult(FRHIRenderQuery* QueryRHI, uint6
 	{
 		if (Query->QueryType == RQT_AbsoluteTime)
 		{
+#if RHI_NEW_GPU_PROFILER
+			checkNoEntry(); // @todo - new gpu profiler
+#else
 			// GetTimingFrequency is the number of ticks per second
 			uint64 Div = FMath::Max(1llu, FOpenGLBufferedGPUTiming::GetTimingFrequency() / (1000 * 1000));
 
 			// convert from GPU specific timestamp to micro sec (1 / 1 000 000 s) which seems a reasonable resolution
 			OutResult = Query->Result / Div;
+#endif
 		}
 		else
 		{
@@ -737,6 +741,8 @@ FOpenGLEventQuery::~FOpenGLEventQuery()
  * class FOpenGLBufferedGPUTiming
  *=============================================================================*/
 
+#if (RHI_NEW_GPU_PROFILER == 0)
+
 /**
  * Constructor.
  *
@@ -1114,6 +1120,8 @@ void FOpenGLDisjointTimeStampQuery::ReleaseResources()
 	}
 }
 
+#endif // (RHI_NEW_GPU_PROFILER == 0)
+
 // Fence implementation
 struct FOpenGLGPUFenceProxy
 {
diff --git a/Engine/Source/Runtime/OpenGLDrv/Private/OpenGLTexture.cpp b/Engine/Source/Runtime/OpenGLDrv/Private/OpenGLTexture.cpp
--- a/Engine/Source/Runtime/OpenGLDrv/Private/OpenGLTexture.cpp
+++ b/Engine/Source/Runtime/OpenGLDrv/Private/OpenGLTexture.cpp
@@ -1396,7 +1396,7 @@ FTextureRHIRef FOpenGLDynamicRHI::RHIAsyncCreateTexture2D(uint32 SizeX, uint32 S
 void FOpenGLDynamicRHI::RHIGenerateMips(FRHITexture* TextureRHI)
 {
 	VERIFY_GL_SCOPE();
-	GPUProfilingData.RegisterGPUWork(0);
+	RegisterGPUWork(0);
 
 	FOpenGLContextState& ContextState = GetContextStateForCurrentContext();
 	FOpenGLTexture* Texture = ResourceCast(TextureRHI);
diff --git a/Engine/Source/Runtime/RHI/Private/GPUProfiler.cpp b/Engine/Source/Runtime/RHI/Private/GPUProfiler.cpp
--- a/Engine/Source/Runtime/RHI/Private/GPUProfiler.cpp
+++ b/Engine/Source/Runtime/RHI/Private/GPUProfiler.cpp
@@ -103,15 +103,7 @@ static TAutoConsoleVariable<int32> GProfileGPUSort(
 	TEXT("3 : By number of verts\n"),
 	ECVF_Default);
 
-struct FNodeStatsCompare
-{
-	/** Sorts nodes by descending durations. */
-	FORCEINLINE bool operator()( const FGPUProfilerEventNodeStats& A, const FGPUProfilerEventNodeStats& B ) const
-	{
-		return B.TimingResult < A.TimingResult;
-	}
-};
-
+#if (RHI_NEW_GPU_PROFILER == 0)
 
 /** Recursively generates a histogram of nodes and stores their timing in TimingResult. */
 static void GatherStatsEventNode(FGPUProfilerEventNode* Node, int32 Depth, TMap<FString, FGPUProfilerEventNodeStats>& EventHistogram)
@@ -510,6 +502,15 @@ void FGPUProfilerEventNodeFrame::DumpEventTree()
 
 		if (RootWildcardString == TEXT("*") && bShowHistogram)
 		{
+			struct FNodeStatsCompare
+			{
+				/** Sorts nodes by descending durations. */
+				FORCEINLINE bool operator()(const FGPUProfilerEventNodeStats& A, const FGPUProfilerEventNodeStats& B) const
+				{
+					return B.TimingResult < A.TimingResult;
+				}
+			};
+
 			// Sort descending based on node duration
 			EventHistogram.ValueSort( FNodeStatsCompare() );
 
@@ -666,4 +667,59 @@ TStaticArray<FGPUTimingCalibrationTimestamp, MAX_NUM_GPUS> FGPUTiming::GCalibrat
 /** Whether the static variables have been initialized. */
 bool FGPUTiming::GAreGlobalsInitialized = false;
 
+#else
+
+namespace UE::RHI::GPUProfiler
+{
+	static TArray<IEventSink*>& GetSinks()
+	{
+		static TArray<IEventSink*> Sinks;
+		return Sinks;
+	}
+
+	void RegisterEventSink(IEventSink* Sink)
+	{
+		GetSinks().Add(Sink);
+	}
+
+	void PushEvents(
+		FQueue Queue,
+		TConstArrayView<TUniquePtr<FBreadcrumbEvent>> BreadcrumbEvents,
+		TConstArrayView<TUniquePtr<FWorkEvent      >> WorkEvents,
+		TConstArrayView<TUniquePtr<FMarkerEvent    >> MarkerEvents)
+	{
+		if (BreadcrumbEvents.IsEmpty()
+		   || WorkEvents.IsEmpty()
+		   || MarkerEvents.IsEmpty())
+		{
+			return;
+		}
+
+		for (IEventSink* Sink : GetSinks())
+		{
+			Sink->ProcessEvents(Queue, BreadcrumbEvents, WorkEvents, MarkerEvents);
+		}
+	}
+
+	struct FGPUProfilerSink_ProfileGPU final : public IEventSink
+	{
+		FGPUProfilerSink_ProfileGPU()
+		{
+			RegisterEventSink(this);
+		}
+
+		void ProcessEvents(
+			FQueue Queue,
+			TConstArrayView<TUniquePtr<FBreadcrumbEvent>> const& BreadcrumbEvents,
+			TConstArrayView<TUniquePtr<FWorkEvent      >> const& WorkEvents,
+			TConstArrayView<TUniquePtr<FMarkerEvent    >> const& MarkerEvents) override
+		{
+			// @todo - new gpu profiler
+		}
+
+	} GGPUProfilerSink_ProfileGPU;
+}
+
+#endif // RHI_NEW_GPU_PROFILER
+
 #undef LOCTEXT_NAMESPACE
diff --git a/Engine/Source/Runtime/RHI/Public/GPUProfiler.h b/Engine/Source/Runtime/RHI/Public/GPUProfiler.h
--- a/Engine/Source/Runtime/RHI/Public/GPUProfiler.h
+++ b/Engine/Source/Runtime/RHI/Public/GPUProfiler.h
@@ -10,6 +10,138 @@
 #include "Templates/RefCounting.h"
 #include "RHI.h"
 
+#if RHI_NEW_GPU_PROFILER
+
+namespace UE::RHI::GPUProfiler
+{
+	struct FQueue
+	{
+		enum class EType : uint8
+		{
+			Graphics,
+			Compute,
+			Copy,
+			SwapChain
+		};
+
+		EType Type;
+		uint8 GPU;
+		uint8 Index;
+	};
+
+	struct FBreadcrumbEvent
+	{
+		struct FFrameBoundary
+		{
+			uint32 FrameNumber;
+		};
+
+		struct FBeginBreadcrumb
+		{
+			FRHIBreadcrumbNode* Breadcrumb;
+			uint64 GPUTimestampTOP;
+		};
+
+		struct FEndBreadcrumb
+		{
+			FRHIBreadcrumbNode* Breadcrumb;
+			uint64 GPUTimestampBOP;
+		};
+		
+		TVariant<
+			  FFrameBoundary
+			, FBeginBreadcrumb
+			, FEndBreadcrumb
+		> Value;
+
+		template <typename T>
+		FBreadcrumbEvent(T const& Value)
+			: Value(TInPlaceType<T>(), Value)
+		{}
+	};
+
+	struct FWorkEvent
+	{
+		struct FBeginWork
+		{
+			uint64 CPUTimestamp;
+			uint64 GPUTimestampTOP;
+		};
+
+		struct FEndWork
+		{
+			uint64 GPUTimestampBOP;
+		};
+
+		struct FSignalFence
+		{
+			uint64 CPUTimestamp;
+			uint64 ID;
+		};
+
+		struct FWaitFence
+		{
+			uint64 CPUTimestamp;
+			uint64 ID;
+		};
+
+		TVariant<
+			  FBeginWork
+			, FEndWork
+			, FSignalFence
+			, FWaitFence
+		> Value;
+
+		template <typename T>
+		FWorkEvent(T const& Value)
+			: Value(TInPlaceType<T>(), Value)
+		{}
+	};
+
+	struct FMarkerEvent
+	{
+		struct FFlip
+		{
+			uint64 GPUTimestamp;
+		};
+
+		struct FVsync
+		{
+			uint64 GPUTimestamp;
+		};
+
+		TVariant<
+			  FFlip
+			, FVsync
+		> Value;
+
+		template <typename T>
+		FMarkerEvent(T const& Value)
+			: Value(TInPlaceType<T>(), Value)
+		{}
+	};
+
+	struct IEventSink
+	{
+		virtual void ProcessEvents(
+			FQueue Queue,
+			TConstArrayView<TUniquePtr<FBreadcrumbEvent>> const& BreadcrumbEvents,
+			TConstArrayView<TUniquePtr<FWorkEvent      >> const& WorkEvents,
+			TConstArrayView<TUniquePtr<FMarkerEvent    >> const& MarkerEvents) = 0;
+	};
+
+	RHI_API void RegisterEventSink(IEventSink* Sink);
+
+	RHI_API void PushEvents(
+		FQueue Queue,
+		TConstArrayView<TUniquePtr<FBreadcrumbEvent>> BreadcrumbEvents,
+		TConstArrayView<TUniquePtr<FWorkEvent      >> WorkEvents,
+		TConstArrayView<TUniquePtr<FMarkerEvent    >> MarkerEvents
+	);
+}
+
+#else
+
 /** Stats for a single perf event node. */
 class FGPUProfilerEventNodeStats : public FRefCountedObject
 {
@@ -328,4 +460,6 @@ struct FGPUProfiler
 	RHI_API virtual void PopEvent();
 
 	bool IsProfilingGPU() const { return bTrackingEvents; }
-};
\ No newline at end of file
+};
+
+#endif
diff --git a/Engine/Source/Runtime/RHI/Public/GpuProfilerTrace.h b/Engine/Source/Runtime/RHI/Public/GpuProfilerTrace.h
--- a/Engine/Source/Runtime/RHI/Public/GpuProfilerTrace.h
+++ b/Engine/Source/Runtime/RHI/Public/GpuProfilerTrace.h
@@ -6,7 +6,7 @@
 #include "Trace/Config.h"
 
 #if !defined(GPUPROFILERTRACE_ENABLED)
-#if UE_TRACE_ENABLED && !UE_BUILD_SHIPPING
+#if UE_TRACE_ENABLED && !UE_BUILD_SHIPPING && (RHI_NEW_GPU_PROFILER == 0)
 #define GPUPROFILERTRACE_ENABLED 1
 #else
 #define GPUPROFILERTRACE_ENABLED 0
diff --git a/Engine/Source/Runtime/RHI/RHI.Build.cs b/Engine/Source/Runtime/RHI/RHI.Build.cs
--- a/Engine/Source/Runtime/RHI/RHI.Build.cs
+++ b/Engine/Source/Runtime/RHI/RHI.Build.cs
@@ -11,6 +11,9 @@ public class RHI : ModuleRules
 		PrivateDependencyModuleNames.Add("TraceLog");
 		PrivateDependencyModuleNames.Add("ApplicationCore");
 
+		// @todo - new gpu profiler. This is experimental.
+		PublicDefinitions.Add("RHI_NEW_GPU_PROFILER=0");
+
 		if (Target.bCompileAgainstEngine)
 		{
 			DynamicallyLoadedModuleNames.Add("NullDrv");
diff --git a/Engine/Source/Runtime/RenderCore/Private/ProfilingDebugging/RealtimeGPUProfiler.cpp b/Engine/Source/Runtime/RenderCore/Private/ProfilingDebugging/RealtimeGPUProfiler.cpp
--- a/Engine/Source/Runtime/RenderCore/Private/ProfilingDebugging/RealtimeGPUProfiler.cpp
+++ b/Engine/Source/Runtime/RenderCore/Private/ProfilingDebugging/RealtimeGPUProfiler.cpp
@@ -8,6 +8,12 @@
 #include "GPUProfiler.h"
 #include "Misc/ScopeRWLock.h"
 
+#if RHI_NEW_GPU_PROFILER
+
+// @todo - new gpu profiler
+
+#else
+
 typedef TArray<TCHAR, TInlineAllocator<4096u>> FDescriptionStringBuffer;
 
 // Only exposed for debugging. Disabling this carries a severe performance penalty
@@ -1153,3 +1159,4 @@ void FRealtimeGPUProfiler::FetchPerfByDescription(TArray<FRealtimeGPUProfilerDes
 }
 #endif  // GPUPROFILERTRACE_ENABLED
 
+#endif // (RHI_NEW_GPU_PROFILER == 0)
diff --git a/Engine/Source/Runtime/RenderCore/Public/ProfilingDebugging/RealtimeGPUProfiler.h b/Engine/Source/Runtime/RenderCore/Public/ProfilingDebugging/RealtimeGPUProfiler.h
--- a/Engine/Source/Runtime/RenderCore/Public/ProfilingDebugging/RealtimeGPUProfiler.h
+++ b/Engine/Source/Runtime/RenderCore/Public/ProfilingDebugging/RealtimeGPUProfiler.h
@@ -30,15 +30,8 @@
 #include "UObject/NameTypes.h"
 #include <tuple>
 
-struct FColor;
-
 #define WANTS_DRAW_MESH_EVENTS (WITH_PROFILEGPU && WITH_RHI_BREADCRUMBS)
 
-class FRealtimeGPUProfiler;
-class FRealtimeGPUProfilerEvent;
-class FRealtimeGPUProfilerFrame;
-class FRenderQueryPool;
-
 #if WITH_RHI_BREADCRUMBS
 
 	struct FRHIBreadcrumbEvent_GameThread
@@ -152,6 +145,32 @@ class FRenderQueryPool;
 #define SCOPED_RHI_CONDITIONAL_DRAW_EVENT_COLOR(RHICmdContext, Color, Name, Condition)               UE_DEPRECATED_MACRO(5.5, "SCOPED_RHI_CONDITIONAL_DRAW_EVENT_COLOR has been deprecated. Use standard RHI breadcrumb events instead." )
 #define SCOPED_RHI_CONDITIONAL_DRAW_EVENTF_COLOR(RHICmdContext, Color, Name, Condition, Format, ...) UE_DEPRECATED_MACRO(5.5, "SCOPED_RHI_CONDITIONAL_DRAW_EVENTF_COLOR has been deprecated. Use standard RHI breadcrumb events instead.")
 
+#if RHI_NEW_GPU_PROFILER
+
+	// @todo
+	#define DECLARE_GPU_STAT(StatName)
+	#define DECLARE_GPU_STAT_NAMED(StatName, NameString)
+	#define DECLARE_GPU_DRAWCALL_STAT(StatName)
+	#define DECLARE_GPU_DRAWCALL_STAT_NAMED(StatName, NameString)
+	#define DECLARE_GPU_DRAWCALL_STAT_EXTERN(StatName)
+
+	#define DECLARE_GPU_STAT_NAMED_EXTERN(StatName, NameString)
+	#define DEFINE_GPU_STAT(StatName)
+	#define DEFINE_GPU_DRAWCALL_STAT(StatName)
+
+	#define SCOPED_GPU_STAT_VERBOSE(RHICmdList, StatName, Description)
+	#define SCOPED_GPU_STAT(RHICmdList, StatName) 
+
+	#define GPU_STATS_BEGINFRAME(RHICmdList) 
+	#define GPU_STATS_ENDFRAME(RHICmdList) 
+	#define GPU_STATS_SUSPENDFRAME()
+
+#else
+
+class FRealtimeGPUProfiler;
+class FRealtimeGPUProfilerEvent;
+class FRealtimeGPUProfilerFrame;
+
 #if HAS_GPU_STATS
 
 	CSV_DECLARE_CATEGORY_MODULE_EXTERN(RENDERCORE_API,GPU);
@@ -210,10 +229,6 @@ RENDERCORE_API bool AreGPUStatsEnabled();
 
 #if HAS_GPU_STATS
 
-class FRealtimeGPUProfilerEvent;
-class FRealtimeGPUProfilerFrame;
-class FRenderQueryPool;
-
 class FRealtimeGPUProfilerQuery
 {
 public:
@@ -349,3 +364,5 @@ public:
 };
 
 #endif // HAS_GPU_STATS
+
+#endif // (RHI_NEW_GPU_PROFILER == 0)
diff --git a/Engine/Source/Runtime/RenderCore/Public/RenderGraphEvent.h b/Engine/Source/Runtime/RenderCore/Public/RenderGraphEvent.h
--- a/Engine/Source/Runtime/RenderCore/Public/RenderGraphEvent.h
+++ b/Engine/Source/Runtime/RenderCore/Public/RenderGraphEvent.h
@@ -77,7 +77,7 @@ enum class ERDGScopeFlags : uint8
 };
 ENUM_CLASS_FLAGS(ERDGScopeFlags);
 
-#if HAS_GPU_STATS
+#if HAS_GPU_STATS && (RHI_NEW_GPU_PROFILER == 0)
 // Scope type for the legacy "realtime" GPU profiler and draw call counter stats
 struct FRDGScope_GPU
 {
@@ -103,7 +103,7 @@ struct FRDGScope_GPU
 	inline void BeginGPU(FRHIComputeCommandList& RHICmdList);
 	inline void EndGPU  (FRHIComputeCommandList& RHICmdList);
 };
-#endif // HAS_GPU_STATS
+#endif // HAS_GPU_STATS && (RHI_NEW_GPU_PROFILER == 0)
 
 #if CSV_PROFILER
 struct FRDGScope_CSVExclusive
@@ -282,7 +282,7 @@ struct FRDGScope
 #if RDG_EVENTS
 		, FRDGScope_RHI
 #endif
-#if HAS_GPU_STATS
+#if HAS_GPU_STATS && (RHI_NEW_GPU_PROFILER == 0)
 		, FRDGScope_GPU
 #endif
 #if CSV_PROFILER
@@ -385,7 +385,7 @@ private:
 
 #endif
 
-#if HAS_GPU_STATS
+#if HAS_GPU_STATS && (RHI_NEW_GPU_PROFILER == 0)
 	#define RDG_GPU_STAT_SCOPE(GraphBuilder, StatName)                      TRDGEventScopeGuard<FRDGScope_GPU> PREPROCESSOR_JOIN(__RDG_GPUStatEvent_##StatName,__LINE__) ((GraphBuilder), ERDGScopeFlags::AlwaysEnable, (GraphBuilder).RHICmdList.GetGPUMask(), CSV_STAT_FNAME(StatName), GET_STATID(Stat_GPU_##StatName), nullptr    , DrawcallCountCategory_##StatName);
 	#define RDG_GPU_STAT_SCOPE_VERBOSE(GraphBuilder, StatName, Description) TRDGEventScopeGuard<FRDGScope_GPU> PREPROCESSOR_JOIN(__RDG_GPUStatEvent_##StatName,__LINE__) ((GraphBuilder), ERDGScopeFlags::AlwaysEnable, (GraphBuilder).RHICmdList.GetGPUMask(), CSV_STAT_FNAME(StatName), GET_STATID(Stat_GPU_##StatName), Description, DrawcallCountCategory_##StatName);
 #else
@@ -496,7 +496,7 @@ public:
 	template <typename TScopeType>
 	friend class TRDGEventScopeGuard;
 	friend FRDGScope_Budget;
-#if HAS_GPU_STATS
+#if HAS_GPU_STATS && (RHI_NEW_GPU_PROFILER == 0)
 	friend FRDGScope_GPU;
 #endif
 #if RDG_EVENTS
diff --git a/Engine/Source/Runtime/RenderCore/Public/RenderGraphEvent.inl b/Engine/Source/Runtime/RenderCore/Public/RenderGraphEvent.inl
--- a/Engine/Source/Runtime/RenderCore/Public/RenderGraphEvent.inl
+++ b/Engine/Source/Runtime/RenderCore/Public/RenderGraphEvent.inl
@@ -2,6 +2,12 @@
 
 #pragma once
 
+#if RHI_NEW_GPU_PROFILER
+
+// @todo - new gpu profiler
+
+#else
+
 #if HAS_GPU_STATS
 inline FRDGScope_GPU::FRDGScope_GPU(FRDGScopeState& State, FRHIGPUMask GPUMask, const FName& CsvStatName, const TStatId& Stat, const TCHAR* Description, FRHIDrawStatsCategory const& Category)
 	: CurrentCategory(Category.ShouldCountDraws() ? &Category : nullptr),
@@ -90,6 +96,8 @@ inline void FRDGScope_GPU::EndGPU(FRHIComputeCommandList& RHICmdList)
 }
 #endif
 
+#endif // (RHI_NEW_GPU_PROFILER == 0)
+
 #if RDG_EVENTS
 
 	inline FRDGScope_RHI::FRDGScope_RHI(FRDGScopeState& State, FRDGEventName&& Name)
diff --git a/Engine/Source/Runtime/Renderer/Private/GPUBenchmark.cpp b/Engine/Source/Runtime/Renderer/Private/GPUBenchmark.cpp
--- a/Engine/Source/Runtime/Renderer/Private/GPUBenchmark.cpp
+++ b/Engine/Source/Runtime/Renderer/Private/GPUBenchmark.cpp
@@ -395,6 +395,7 @@ void RendererGPUBenchmark(FRHICommandListImmediate& RHICmdList, FSynthBenchmarkR
 {
 	check(IsInRenderingThread());
 
+#if (RHI_NEW_GPU_PROFILER == 0)
 	// Multi-GPU support : Benchmark needs to be made MGPU-aware.
 	bool bValidGPUTimer = (FGPUTiming::GetTimingFrequency() / (1000 * 1000)) != 0;
 
@@ -403,6 +404,7 @@ void RendererGPUBenchmark(FRHICommandListImmediate& RHICmdList, FSynthBenchmarkR
 		UE_LOG(LogSynthBenchmark, Warning, TEXT("RendererGPUBenchmark failed, look for \"GPU Timing Frequency\" in the log"));
 		return;
 	}
+#endif
 
 	RHICmdList.ImmediateFlush(EImmediateFlushType::FlushRHIThread);
 	MeasureLongGPUTaskExecutionTime(RHICmdList);
diff --git a/Engine/Source/Runtime/VulkanRHI/Private/VulkanCommandBuffer.cpp b/Engine/Source/Runtime/VulkanRHI/Private/VulkanCommandBuffer.cpp
--- a/Engine/Source/Runtime/VulkanRHI/Private/VulkanCommandBuffer.cpp
+++ b/Engine/Source/Runtime/VulkanRHI/Private/VulkanCommandBuffer.cpp
@@ -66,7 +66,9 @@ FVulkanCmdBuffer::FVulkanCmdBuffer(FVulkanDevice* InDevice, FVulkanCommandBuffer
 	, FenceSignaledCounter(0)
 	, SubmittedFenceCounter(0)
 	, CommandBufferPool(InCommandBufferPool)
+#if (RHI_NEW_GPU_PROFILER == 0)
 	, Timing(nullptr)
+#endif
 	, LastValidTiming(0)
 	, LayoutManager(InDevice->SupportsParallelRendering() && !GVulkanAutoCorrectUnknownLayouts,
 		&InCommandBufferPool->GetMgr().GetCommandListContext()->GetQueue()->GetLayoutManager())
@@ -124,12 +126,14 @@ FVulkanCmdBuffer::~FVulkanCmdBuffer()
 		FreeMemory();
 	}
 
+#if (RHI_NEW_GPU_PROFILER == 0)
 	if (Timing)
 	{
 		Timing->Release();
 		delete Timing;
 		Timing = nullptr;
 	}
+#endif
 }
 
 void FVulkanCmdBuffer::FreeMemory()
@@ -229,6 +233,7 @@ void FVulkanCmdBuffer::End()
 {
 	checkf(IsOutsideRenderPass(), TEXT("Can't End as we're inside a render pass! CmdBuffer 0x%p State=%d"), CommandBufferHandle, (int32)State);
 
+#if (RHI_NEW_GPU_PROFILER == 0)
 	if (GVulkanProfileCmdBuffers || GVulkanUseCmdBufferTimingForGPUTime)
 	{
 		if (Timing)
@@ -237,6 +242,7 @@ void FVulkanCmdBuffer::End()
 			LastValidTiming = FenceSignaledCounter;
 		}
 	}
+#endif
 
 	for (PendingQuery& Query : PendingTimestampQueries)
 	{
@@ -262,6 +268,7 @@ void FVulkanCmdBuffer::End()
 
 inline void FVulkanCmdBuffer::InitializeTimings(FVulkanCommandListContext* InContext)
 {
+#if (RHI_NEW_GPU_PROFILER == 0)
 	if ((GVulkanProfileCmdBuffers || GVulkanUseCmdBufferTimingForGPUTime) && !Timing)
 	{
 		if (InContext)
@@ -274,6 +281,7 @@ inline void FVulkanCmdBuffer::InitializeTimings(FVulkanCommandListContext* InCon
 			Timing->Initialize(PoolSize);
 		}
 	}
+#endif
 }
 
 void FVulkanCmdBuffer::AddWaitSemaphore(VkPipelineStageFlags InWaitFlags, TArrayView<VulkanRHI::FSemaphore*> InWaitSemaphores)
@@ -314,10 +322,12 @@ void FVulkanCmdBuffer::Begin()
 	if (GVulkanProfileCmdBuffers || GVulkanUseCmdBufferTimingForGPUTime)
 	{
 		InitializeTimings(CommandBufferPool->GetMgr().GetCommandListContext());
+#if (RHI_NEW_GPU_PROFILER == 0)
 		if (Timing)
 		{
 			Timing->StartTiming(this);
 		}
+#endif
 	}
 	check(!CurrentDescriptorPoolSetContainer);
 
@@ -761,6 +771,7 @@ void FVulkanCommandBufferManager::PrepareForNewActiveCommandBuffer()
 	ActiveCmdBuffer->Begin();
 }
 
+#if (RHI_NEW_GPU_PROFILER == 0)
 uint32 FVulkanCommandBufferManager::CalculateGPUTime()
 {
 	uint32 Time = 0;
@@ -774,6 +785,7 @@ uint32 FVulkanCommandBufferManager::CalculateGPUTime()
 	}
 	return Time;
 }
+#endif // (RHI_NEW_GPU_PROFILER == 0)
 
 FVulkanCmdBuffer* FVulkanCommandBufferManager::GetUploadCmdBuffer()
 {
diff --git a/Engine/Source/Runtime/VulkanRHI/Private/VulkanCommandBuffer.h b/Engine/Source/Runtime/VulkanRHI/Private/VulkanCommandBuffer.h
--- a/Engine/Source/Runtime/VulkanRHI/Private/VulkanCommandBuffer.h
+++ b/Engine/Source/Runtime/VulkanRHI/Private/VulkanCommandBuffer.h
@@ -149,10 +149,12 @@ public:
 		return LayoutManager;
 	}
 
+#if (RHI_NEW_GPU_PROFILER == 0)
 	inline bool HasValidTiming() const
 	{
 		return (Timing != nullptr) && (FMath::Abs((int64)FenceSignaledCounter - (int64)LastValidTiming) < 3);
 	}
+#endif
 
 	void VULKANRHI_API AddWaitSemaphore(VkPipelineStageFlags InWaitFlags, VulkanRHI::FSemaphore* InWaitSemaphore)
 	{
@@ -241,7 +243,9 @@ private:
 
 	FVulkanCommandBufferPool* CommandBufferPool;
 
+#if (RHI_NEW_GPU_PROFILER == 0)
 	FVulkanGPUTiming* Timing;
+#endif
 	uint64 LastValidTiming;
 
 	void AcquirePoolSetContainer();
@@ -374,7 +378,9 @@ public:
 		return Pool.GetHandle();
 	}
 
+#if (RHI_NEW_GPU_PROFILER == 0)
 	uint32 CalculateGPUTime();
+#endif
 
 	void FreeUnusedCmdBuffers(bool bTrimMemory);
 
diff --git a/Engine/Source/Runtime/VulkanRHI/Private/VulkanCommands.cpp b/Engine/Source/Runtime/VulkanRHI/Private/VulkanCommands.cpp
--- a/Engine/Source/Runtime/VulkanRHI/Private/VulkanCommands.cpp
+++ b/Engine/Source/Runtime/VulkanRHI/Private/VulkanCommands.cpp
@@ -286,7 +286,7 @@ void FVulkanCommandListContext::RHIDispatchComputeShader(uint32 ThreadGroupCount
 
 	if (FVulkanPlatform::RegisterGPUWork() && IsImmediate())
 	{
-		GpuProfiler.RegisterGPUDispatch(FIntVector(ThreadGroupCountX, ThreadGroupCountY, ThreadGroupCountZ));	
+		RegisterGPUDispatch(FIntVector(ThreadGroupCountX, ThreadGroupCountY, ThreadGroupCountZ));	
 	}
 
 	VulkanRHI::DebugHeavyWeightBarrier(CmdBuffer, 2);
@@ -313,7 +313,7 @@ void FVulkanCommandListContext::RHIDispatchIndirectComputeShader(FRHIBuffer* Arg
 
 	if (FVulkanPlatform::RegisterGPUWork()/* && IsImmediate()*/)
 	{
-		GpuProfiler.RegisterGPUDispatch(FIntVector(1, 1, 1));	
+		RegisterGPUDispatch(FIntVector(1, 1, 1));	
 	}
 
 	VulkanRHI::DebugHeavyWeightBarrier(CmdBuffer, 2);
@@ -635,7 +635,7 @@ void FVulkanCommandListContext::RHIDrawPrimitive(uint32 BaseVertexIndex, uint32
 
 	if (FVulkanPlatform::RegisterGPUWork() && IsImmediate())
 	{
-		GpuProfiler.RegisterGPUWork(NumPrimitives * NumInstances, NumVertices * NumInstances);
+		RegisterGPUWork(NumPrimitives * NumInstances, NumVertices * NumInstances);
 	}
 }
 
@@ -661,7 +661,7 @@ void FVulkanCommandListContext::RHIDrawPrimitiveIndirect(FRHIBuffer* ArgumentBuf
 
 	if (FVulkanPlatform::RegisterGPUWork() && IsImmediate())
 	{
-		GpuProfiler.RegisterGPUWork(1);
+		RegisterGPUWork(1);
 	}
 }
 
@@ -688,7 +688,7 @@ void FVulkanCommandListContext::RHIDrawIndexedPrimitive(FRHIBuffer* IndexBufferR
 
 	if (FVulkanPlatform::RegisterGPUWork() && IsImmediate())
 	{
-		GpuProfiler.RegisterGPUWork(NumPrimitives * NumInstances, NumVertices * NumInstances);
+		RegisterGPUWork(NumPrimitives * NumInstances, NumVertices * NumInstances);
 	}
 }
 
@@ -715,7 +715,7 @@ void FVulkanCommandListContext::RHIDrawIndexedIndirect(FRHIBuffer* IndexBufferRH
 
 	if (FVulkanPlatform::RegisterGPUWork() && IsImmediate())
 	{
-		GpuProfiler.RegisterGPUWork(1);
+		RegisterGPUWork(1);
 	}
 }
 
@@ -741,7 +741,7 @@ void FVulkanCommandListContext::RHIDrawIndexedPrimitiveIndirect(FRHIBuffer* Inde
 
 	if (FVulkanPlatform::RegisterGPUWork() && IsImmediate())
 	{
-		GpuProfiler.RegisterGPUWork(1); 
+		RegisterGPUWork(1); 
 	}
 }
 
diff --git a/Engine/Source/Runtime/VulkanRHI/Private/VulkanContext.h b/Engine/Source/Runtime/VulkanRHI/Private/VulkanContext.h
--- a/Engine/Source/Runtime/VulkanRHI/Private/VulkanContext.h
+++ b/Engine/Source/Runtime/VulkanRHI/Private/VulkanContext.h
@@ -93,7 +93,10 @@ public:
 	// Render time measurement
 	virtual void RHIBeginRenderQuery(FRHIRenderQuery* RenderQuery) final override;
 	virtual void RHIEndRenderQuery(FRHIRenderQuery* RenderQuery) final override;
+
+#if (RHI_NEW_GPU_PROFILER == 0)
 	virtual void RHICalibrateTimers(FRHITimestampCalibrationQuery* CalibrationQuery) final override;
+#endif
 
 	virtual void RHIBeginDrawingViewport(FRHIViewport* Viewport, FRHITexture* RenderTargetRHI) final override;
 	virtual void RHIEndDrawingViewport(FRHIViewport* Viewport, bool bPresent, bool bLockToVsync) final override;
@@ -191,12 +194,14 @@ public:
 	void WriteBeginTimestamp(FVulkanCmdBuffer* CmdBuffer);
 	void WriteEndTimestamp(FVulkanCmdBuffer* CmdBuffer);
 
+#if (RHI_NEW_GPU_PROFILER == 0)
 	void ReadAndCalculateGPUFrameTime();
 	
 	inline FVulkanGPUProfiler& GetGPUProfiler()
 	{
 		return GpuProfiler;
 	}
+#endif
 
 	inline FVulkanDevice* GetDevice() const
 	{
@@ -283,8 +288,16 @@ private:
 	// Number of times EndFrame() has been called on this context
 	uint64 FrameCounter;
 
+#if RHI_NEW_GPU_PROFILER
+	void RegisterGPUWork(uint32 NumPrimitives = 0, uint32 NumVertices = 0)	{ checkNoEntry(); } // @todo - new gpu profiler
+	void RegisterGPUDispatch(FIntVector GroupCount)	                        { checkNoEntry(); } // @todo - new gpu profiler
+#else
+	void RegisterGPUWork(uint32 NumPrimitives = 0, uint32 NumVertices = 0)	{ GpuProfiler.RegisterGPUWork(NumPrimitives, NumVertices); }
+	void RegisterGPUDispatch(FIntVector GroupCount)	                        { GpuProfiler.RegisterGPUDispatch(GroupCount); }
+
 	FVulkanGPUProfiler GpuProfiler;
 	FVulkanGPUTiming* FrameTiming;
+#endif
 
 	template <typename TRHIShader>
 	void ApplyStaticUniformBuffers(TRHIShader* Shader);
diff --git a/Engine/Source/Runtime/VulkanRHI/Private/VulkanDevice.cpp b/Engine/Source/Runtime/VulkanRHI/Private/VulkanDevice.cpp
--- a/Engine/Source/Runtime/VulkanRHI/Private/VulkanDevice.cpp
+++ b/Engine/Source/Runtime/VulkanRHI/Private/VulkanDevice.cpp
@@ -1529,6 +1529,7 @@ FVulkanTransientHeapCache& FVulkanDevice::GetOrCreateTransientHeapCache()
 	return *TransientHeapCache;
 }
 
+#if (RHI_NEW_GPU_PROFILER == 0)
 FGPUTimingCalibrationTimestamp FVulkanDevice::GetCalibrationTimestamp()
 {
 	auto ToMicroseconds = [](uint64_t Timestamp)
@@ -1555,3 +1556,4 @@ FGPUTimingCalibrationTimestamp FVulkanDevice::GetCalibrationTimestamp()
 	}
 	return CalibrationTimestamp;
 }
+#endif // (RHI_NEW_GPU_PROFILER == 0)
diff --git a/Engine/Source/Runtime/VulkanRHI/Private/VulkanDevice.h b/Engine/Source/Runtime/VulkanRHI/Private/VulkanDevice.h
--- a/Engine/Source/Runtime/VulkanRHI/Private/VulkanDevice.h
+++ b/Engine/Source/Runtime/VulkanRHI/Private/VulkanDevice.h
@@ -528,9 +528,11 @@ public:
 
 	const TArray<const ANSICHAR*>& GetDeviceExtensions() { return DeviceExtensions; }
 
+#if (RHI_NEW_GPU_PROFILER == 0)
 	// Performs a GPU and CPU timestamp at nearly the same time.
 	// This allows aligning GPU and CPU events on the same timeline in profile visualization.
 	FGPUTimingCalibrationTimestamp GetCalibrationTimestamp();
+#endif
 
 private:
 	void MapBufferFormatSupport(FPixelFormatInfo& PixelFormatInfo, EPixelFormat UEFormat, VkFormat VulkanFormat);
diff --git a/Engine/Source/Runtime/VulkanRHI/Private/VulkanExtensions.cpp b/Engine/Source/Runtime/VulkanRHI/Private/VulkanExtensions.cpp
--- a/Engine/Source/Runtime/VulkanRHI/Private/VulkanExtensions.cpp
+++ b/Engine/Source/Runtime/VulkanRHI/Private/VulkanExtensions.cpp
@@ -1034,7 +1034,7 @@ public:
 			DeviceDiagnosticsConfigCreateInfoNV.flags = VK_DEVICE_DIAGNOSTICS_CONFIG_ENABLE_SHADER_DEBUG_INFO_BIT_NV | VK_DEVICE_DIAGNOSTICS_CONFIG_ENABLE_RESOURCE_TRACKING_BIT_NV | VK_DEVICE_DIAGNOSTICS_CONFIG_ENABLE_AUTOMATIC_CHECKPOINTS_BIT_NV;
 			AddToPNext(DeviceCreateInfo, DeviceDiagnosticsConfigCreateInfoNV);
 
-#if NV_AFTERMATH
+#if NV_AFTERMATH && (RHI_NEW_GPU_PROFILER == 0)
 			// Vulkan's breadcrumb / markers implementation differs from the one provided by RHICore.
 			// @todo unify the implementation so this isn't necessary
 			extern void AftermathResolveMarkerCallback(const void* Marker, void** ResolvedMarkerData, uint32_t * MarkerSize);
diff --git a/Engine/Source/Runtime/VulkanRHI/Private/VulkanGPUProfiler.h b/Engine/Source/Runtime/VulkanRHI/Private/VulkanGPUProfiler.h
--- a/Engine/Source/Runtime/VulkanRHI/Private/VulkanGPUProfiler.h
+++ b/Engine/Source/Runtime/VulkanRHI/Private/VulkanGPUProfiler.h
@@ -15,6 +15,8 @@ class FVulkanCommandListContext;
 class FVulkanDevice;
 class FVulkanTimingQueryPool;
 
+#if (RHI_NEW_GPU_PROFILER == 0)
+
 class FVulkanGPUTiming : public FGPUTiming
 {
 public:
@@ -205,3 +207,5 @@ struct FVulkanGPUProfiler : public FGPUProfiler
 	TArray<uint64> CrashMarkers;
 	bool bBeginFrame;
 };
+
+#endif // (RHI_NEW_GPU_PROFILER == 0)
\ No newline at end of file
diff --git a/Engine/Source/Runtime/VulkanRHI/Private/VulkanQuery.cpp b/Engine/Source/Runtime/VulkanRHI/Private/VulkanQuery.cpp
--- a/Engine/Source/Runtime/VulkanRHI/Private/VulkanQuery.cpp
+++ b/Engine/Source/Runtime/VulkanRHI/Private/VulkanQuery.cpp
@@ -385,6 +385,7 @@ void FVulkanOcclusionQuery::ReleaseFromPool()
 	IndexInPool = UINT32_MAX;
 }
 
+#if (RHI_NEW_GPU_PROFILER == 0)
 void FVulkanCommandListContext::ReadAndCalculateGPUFrameTime()
 {
 	check(IsImmediate());
@@ -412,6 +413,7 @@ void FVulkanCommandListContext::ReadAndCalculateGPUFrameTime()
 		GGPUFrameTime = 0;
 	}
 }
+#endif // (RHI_NEW_GPU_PROFILER == 0)
 
 
 FRenderQueryRHIRef FVulkanDynamicRHI::RHICreateRenderQuery(ERenderQueryType QueryType)
@@ -437,11 +439,16 @@ FRenderQueryRHIRef FVulkanDynamicRHI::RHICreateRenderQuery(ERenderQueryType Quer
 
 bool FVulkanDynamicRHI::RHIGetRenderQueryResult(FRHIRenderQuery* QueryRHI, uint64& OutNumPixels, bool bWait, uint32 GPUIndex)
 {
-	auto ToMicroseconds = [](uint64 Timestamp)
+	auto ToMicroseconds = [](uint64 Timestamp) -> uint64
 	{
+#if RHI_NEW_GPU_PROFILER
+		checkNoEntry(); // @todo - new gpu profiler
+		return 0;
+#else
 		const double Frequency = double(FVulkanGPUTiming::GetTimingFrequency());
 		uint64 Microseconds = (uint64)((double(Timestamp) / Frequency) * 1000.0 * 1000.0);
 		return Microseconds;
+#endif
 	};
 
 	FRHICommandListImmediate& RHICmdList = FRHICommandListImmediate::Get();
@@ -628,6 +635,7 @@ void FVulkanCommandListContext::RHIEndRenderQuery(FRHIRenderQuery* QueryRHI)
 	}
 }
 
+#if (RHI_NEW_GPU_PROFILER == 0)
 void FVulkanCommandListContext::RHICalibrateTimers(FRHITimestampCalibrationQuery* CalibrationQuery)
 {
 	if (Device->GetOptionalExtensions().HasEXTCalibratedTimestamps)
@@ -637,15 +645,20 @@ void FVulkanCommandListContext::RHICalibrateTimers(FRHITimestampCalibrationQuery
 		CalibrationQuery->GPUMicroseconds[0] = CalibrationTimestamp.GPUMicroseconds;
 	}
 }
+#endif
 
 void FVulkanCommandListContext::WriteBeginTimestamp(FVulkanCmdBuffer* CmdBuffer)
 {
+#if (RHI_NEW_GPU_PROFILER == 0)
 	FrameTiming->StartTiming(CmdBuffer);
+#endif
 }
 
 void FVulkanCommandListContext::WriteEndTimestamp(FVulkanCmdBuffer* CmdBuffer)
 {
+#if (RHI_NEW_GPU_PROFILER == 0)
 	FrameTiming->EndTiming(CmdBuffer);
+#endif
 }
 
 FVulkanTimingQuery::FVulkanTimingQuery()
diff --git a/Engine/Source/Runtime/VulkanRHI/Private/VulkanRHI.cpp b/Engine/Source/Runtime/VulkanRHI/Private/VulkanRHI.cpp
--- a/Engine/Source/Runtime/VulkanRHI/Private/VulkanRHI.cpp
+++ b/Engine/Source/Runtime/VulkanRHI/Private/VulkanRHI.cpp
@@ -431,14 +431,22 @@ FVulkanCommandListContext::FVulkanCommandListContext(FVulkanDynamicRHI* InRHI, F
 	, PendingGfxState(nullptr)
 	, PendingComputeState(nullptr)
 	, FrameCounter(0)
+#if (RHI_NEW_GPU_PROFILER == 0)
 	, GpuProfiler(this, InDevice)
+#endif
 {
+#if (RHI_NEW_GPU_PROFILER == 0)
 	FrameTiming = new FVulkanGPUTiming(this, InDevice);
+#endif
 
 	// Create CommandBufferManager, contain all active buffers
 	CommandBufferManager = new FVulkanCommandBufferManager(InDevice, this);
 	CommandBufferManager->Init(this);
+
+#if (RHI_NEW_GPU_PROFILER == 0)
 	FrameTiming->Initialize();
+#endif
+
 	if (IsImmediate())
 	{
 		// Insert the Begin frame timestamp query. On EndDrawingViewport() we'll insert the End and immediately after a new Begin()
@@ -472,9 +480,11 @@ FVulkanCommandListContext::~FVulkanCommandListContext()
 {
 	if (FVulkanPlatform::SupportsTimestampRenderQueries())
 	{
+#if (RHI_NEW_GPU_PROFILER == 0)
 		FrameTiming->Release();
 		delete FrameTiming;
 		FrameTiming = nullptr;
+#endif
 	}
 
 	check(CommandBufferManager != nullptr);
@@ -1107,7 +1117,9 @@ void FVulkanCommandListContext::RHIBeginFrame()
 	extern uint32 GVulkanRHIDeletionFrameNumber;
 	++GVulkanRHIDeletionFrameNumber;
 
+#if (RHI_NEW_GPU_PROFILER == 0)
 	GpuProfiler.BeginFrame();
+#endif
 
 #if VULKAN_RHI_RAYTRACING
 	if (GRHISupportsRayTracing)
@@ -1176,9 +1188,10 @@ void FVulkanCommandListContext::RHIEndFrame()
 	check(IsImmediate());
 	//FRCLog::Printf(FString::Printf(TEXT("FVulkanCommandListContext::RHIEndFrame()")));
 	
+#if (RHI_NEW_GPU_PROFILER == 0)
 	ReadAndCalculateGPUFrameTime();
-
-	GetGPUProfiler().EndFrame();
+	GpuProfiler.EndFrame();
+#endif
 
 	bool bTrimMemory = false;
 	GetCommandBufferManager()->FreeUnusedCmdBuffers(bTrimMemory);
@@ -1243,6 +1256,7 @@ void FVulkanCommandListContext::RHIEndFrame()
 		#endif
 		}
 		
+	#if (RHI_NEW_GPU_PROFILER == 0)
 		if (IsImmediate())
 		{
 		#if VULKAN_SUPPORTS_GPU_CRASH_DUMPS
@@ -1256,10 +1270,12 @@ void FVulkanCommandListContext::RHIEndFrame()
 				GpuProfiler.PushEvent(GetNameStr(), Color);
 			}
 		}
+	#endif // (RHI_NEW_GPU_PROFILER == 0)
 	}
 
 	void FVulkanCommandListContext::RHIEndBreadcrumbGPU(FRHIBreadcrumbNode* Breadcrumb)
 	{
+	#if (RHI_NEW_GPU_PROFILER == 0)
 		//only valid on immediate context currently.  needs to be fixed for parallel rhi execute
 		if (IsImmediate())
 		{
@@ -1275,6 +1291,7 @@ void FVulkanCommandListContext::RHIEndFrame()
 			}
 		#endif
 		}
+	#endif // (RHI_NEW_GPU_PROFILER == 0)
 
 		if (ShouldEmitBreadcrumbs())
 		{
@@ -1524,7 +1541,7 @@ void FVulkanDynamicRHI::RHIRegisterWork(uint32 NumPrimitives)
 	FVulkanCommandListContextImmediate& ImmediateContext = GetDevice()->GetImmediateContext();
 	if (FVulkanPlatform::RegisterGPUWork() && ImmediateContext.IsImmediate())
 	{
-		ImmediateContext.GetGPUProfiler().RegisterGPUWork(1);
+		ImmediateContext.RegisterGPUWork(NumPrimitives);
 	}
 }
 
diff --git a/Engine/Source/Runtime/VulkanRHI/Private/VulkanUtil.cpp b/Engine/Source/Runtime/VulkanRHI/Private/VulkanUtil.cpp
--- a/Engine/Source/Runtime/VulkanRHI/Private/VulkanUtil.cpp
+++ b/Engine/Source/Runtime/VulkanRHI/Private/VulkanUtil.cpp
@@ -20,6 +20,7 @@
 
 extern CORE_API bool GIsGPUCrashed;
 
+#if (RHI_NEW_GPU_PROFILER == 0)
 
 static FString		EventDeepString(TEXT("EventTooDeep"));
 static const uint32	EventDeepCRC = FCrc::StrCrc32<TCHAR>(*EventDeepString);
@@ -62,79 +63,6 @@
 	}
 }
 
-FVulkanStagingBuffer::~FVulkanStagingBuffer()
-{
-	if (StagingBuffer)
-	{
-		check(Device);
-		Device->GetStagingManager().ReleaseBuffer(nullptr, StagingBuffer);
-	}
-}
-
-void* FVulkanStagingBuffer::Lock(uint32 Offset, uint32 NumBytes)
-{
-	check(!bIsLocked);
-	bIsLocked = true;
-	const uint32 EndOffset = Offset + NumBytes;
-	checkf(EndOffset <= QueuedNumBytes, TEXT("Lock at Offset (%u) and NumBytes (%u) reads beyond the allocated size of the staging buffer (%u)"), Offset, NumBytes, QueuedNumBytes);
-	// make sure cached memory is invalidated
-	StagingBuffer->InvalidateMappedMemory();
-	return (void*)((uint8*)StagingBuffer->GetMappedPointer() + Offset);
-}
-
-void FVulkanStagingBuffer::Unlock()
-{
-	check(bIsLocked);
-	bIsLocked = false;
-}
-
-FStagingBufferRHIRef FVulkanDynamicRHI::RHICreateStagingBuffer()
-{
-	return new FVulkanStagingBuffer();
-}
-
-void* FVulkanDynamicRHI::RHILockStagingBuffer(FRHIStagingBuffer* StagingBufferRHI, FRHIGPUFence* FenceRHI, uint32 Offset, uint32 NumBytes)
-{
-	FVulkanStagingBuffer* StagingBuffer = ResourceCast(StagingBufferRHI);
-
-	if (FenceRHI && !FenceRHI->Poll())
-	{
-		Device->SubmitCommandsAndFlushGPU();
-
-		// SubmitCommandsAndFlushGPU might update fence state if it was tied to a previously submitted command buffer.
-		// Its state will have been updated from Submitted to NeedReset, and would assert in WaitForCmdBuffer (which is not needed in such a case)
-		if (!FenceRHI->Poll())
-		{
-			FVulkanGPUFence* Fence = ResourceCast(FenceRHI);
-			Device->GetImmediateContext().GetCommandBufferManager()->WaitForCmdBuffer(Fence->GetCmdBuffer());
-		}
-	}
-
-	return StagingBuffer->Lock(Offset, NumBytes);
-}
-
-void FVulkanDynamicRHI::RHIUnlockStagingBuffer(FRHIStagingBuffer* StagingBufferRHI)
-{
-	FVulkanStagingBuffer* StagingBuffer = ResourceCast(StagingBufferRHI);
-	StagingBuffer->Unlock();
-}
-
-void FVulkanGPUFence::Clear()
-{
-	CmdBuffer = nullptr;
-	FenceSignaledCounter = MAX_uint64;
-}
-
-bool FVulkanGPUFence::Poll() const
-{
-	return (CmdBuffer && (FenceSignaledCounter < CmdBuffer->GetFenceSignaledCounter()));
-}
-
-FGPUFenceRHIRef FVulkanDynamicRHI::RHICreateGPUFence(const FName& Name)
-{
-	return new FVulkanGPUFence(Name);
-}
-
 FVulkanGPUTiming::~FVulkanGPUTiming()
 {
 	check(!Pool);
@@ -769,6 +697,81 @@
 }
 #endif
 
+#endif // (RHI_NEW_GPU_PROFILER == 0)
+
+FVulkanStagingBuffer::~FVulkanStagingBuffer()
+{
+	if (StagingBuffer)
+	{
+		check(Device);
+		Device->GetStagingManager().ReleaseBuffer(nullptr, StagingBuffer);
+	}
+}
+
+void* FVulkanStagingBuffer::Lock(uint32 Offset, uint32 NumBytes)
+{
+	check(!bIsLocked);
+	bIsLocked = true;
+	const uint32 EndOffset = Offset + NumBytes;
+	checkf(EndOffset <= QueuedNumBytes, TEXT("Lock at Offset (%u) and NumBytes (%u) reads beyond the allocated size of the staging buffer (%u)"), Offset, NumBytes, QueuedNumBytes);
+	// make sure cached memory is invalidated
+	StagingBuffer->InvalidateMappedMemory();
+	return (void*)((uint8*)StagingBuffer->GetMappedPointer() + Offset);
+}
+
+void FVulkanStagingBuffer::Unlock()
+{
+	check(bIsLocked);
+	bIsLocked = false;
+}
+
+FStagingBufferRHIRef FVulkanDynamicRHI::RHICreateStagingBuffer()
+{
+	return new FVulkanStagingBuffer();
+}
+
+void* FVulkanDynamicRHI::RHILockStagingBuffer(FRHIStagingBuffer* StagingBufferRHI, FRHIGPUFence* FenceRHI, uint32 Offset, uint32 NumBytes)
+{
+	FVulkanStagingBuffer* StagingBuffer = ResourceCast(StagingBufferRHI);
+
+	if (FenceRHI && !FenceRHI->Poll())
+	{
+		Device->SubmitCommandsAndFlushGPU();
+
+		// SubmitCommandsAndFlushGPU might update fence state if it was tied to a previously submitted command buffer.
+		// Its state will have been updated from Submitted to NeedReset, and would assert in WaitForCmdBuffer (which is not needed in such a case)
+		if (!FenceRHI->Poll())
+		{
+			FVulkanGPUFence* Fence = ResourceCast(FenceRHI);
+			Device->GetImmediateContext().GetCommandBufferManager()->WaitForCmdBuffer(Fence->GetCmdBuffer());
+		}
+	}
+
+	return StagingBuffer->Lock(Offset, NumBytes);
+}
+
+void FVulkanDynamicRHI::RHIUnlockStagingBuffer(FRHIStagingBuffer* StagingBufferRHI)
+{
+	FVulkanStagingBuffer* StagingBuffer = ResourceCast(StagingBufferRHI);
+	StagingBuffer->Unlock();
+}
+
+void FVulkanGPUFence::Clear()
+{
+	CmdBuffer = nullptr;
+	FenceSignaledCounter = MAX_uint64;
+}
+
+bool FVulkanGPUFence::Poll() const
+{
+	return (CmdBuffer && (FenceSignaledCounter < CmdBuffer->GetFenceSignaledCounter()));
+}
+
+FGPUFenceRHIRef FVulkanDynamicRHI::RHICreateGPUFence(const FName& Name)
+{
+	return new FVulkanGPUFence(Name);
+}
+
 namespace VulkanRHI
 {
 	VkBuffer CreateBuffer(FVulkanDevice* InDevice, VkDeviceSize Size, VkBufferUsageFlags BufferUsageFlags, VkMemoryRequirements& OutMemoryRequirements)
@@ -921,6 +924,10 @@
 		{
 			FVulkanDevice* Device = GVulkanRHI->GetDevice();
 
+#if RHI_NEW_GPU_PROFILER
+			checkNoEntry(); // @todo - new gpu profiler
+#else
+
 #if VULKAN_SUPPORTS_GPU_CRASH_DUMPS
 			if (UE::RHI::UseGPUCrashDebugging())
 			{
@@ -928,4 +935,6 @@
 			}
 #endif
 
+#endif
+
 			CheckDeviceFault(Device);
diff --git a/Engine/Source/Runtime/Windows/D3D11RHI/Private/D3D11Buffer.cpp b/Engine/Source/Runtime/Windows/D3D11RHI/Private/D3D11Buffer.cpp
--- a/Engine/Source/Runtime/Windows/D3D11RHI/Private/D3D11Buffer.cpp
+++ b/Engine/Source/Runtime/Windows/D3D11RHI/Private/D3D11Buffer.cpp
@@ -275,7 +275,7 @@ void FD3D11DynamicRHI::RHICopyBuffer(FRHIBuffer* SourceBufferRHI, FRHIBuffer* De
 
 	Direct3DDeviceIMContext->CopyResource(DestBuffer->Resource,SourceBuffer->Resource);
 
-	GPUProfilingData.RegisterGPUWork(1);
+	RegisterGPUWork(1);
 }
 
 void FD3D11DynamicRHI::RHIReplaceResources(FRHICommandListBase& RHICmdList, TArray<FRHIResourceReplaceInfo>&& ReplaceInfos)
diff --git a/Engine/Source/Runtime/Windows/D3D11RHI/Private/D3D11Commands.cpp b/Engine/Source/Runtime/Windows/D3D11RHI/Private/D3D11Commands.cpp
--- a/Engine/Source/Runtime/Windows/D3D11RHI/Private/D3D11Commands.cpp
+++ b/Engine/Source/Runtime/Windows/D3D11RHI/Private/D3D11Commands.cpp
@@ -165,7 +165,7 @@ void FD3D11DynamicRHI::RHIDispatchComputeShader(uint32 ThreadGroupCountX, uint32
 
 	StateCache.SetComputeShader(ComputeShader->Resource);
 
-	GPUProfilingData.RegisterGPUDispatch(FIntVector(ThreadGroupCountX, ThreadGroupCountY, ThreadGroupCountZ));	
+	RegisterGPUDispatch(FIntVector(ThreadGroupCountX, ThreadGroupCountY, ThreadGroupCountZ));	
 
 	if (ComputeShader->bShaderNeedsGlobalConstantBuffer)
 	{
@@ -184,7 +184,7 @@ void FD3D11DynamicRHI::RHIDispatchIndirectComputeShader(FRHIBuffer* ArgumentBuff
 	FD3D11ComputeShader* ComputeShader = ResourceCast(ComputeShaderRHI);
 	FD3D11Buffer* ArgumentBuffer = ResourceCast(ArgumentBufferRHI);
 
-	GPUProfilingData.RegisterGPUDispatch(FIntVector(1, 1, 1));
+	RegisterGPUDispatch(FIntVector(1, 1, 1));
 
 	StateCache.SetComputeShader(ComputeShader->Resource);
 	
@@ -1096,7 +1096,7 @@ void FD3D11DynamicRHI::RHIDrawPrimitive(uint32 BaseVertexIndex,uint32 NumPrimiti
 
 	uint32 VertexCount = GetVertexCountForPrimitiveCount(NumPrimitives,PrimitiveType);
 
-	GPUProfilingData.RegisterGPUWork(NumPrimitives * NumInstances, VertexCount * NumInstances);
+	RegisterGPUWork(NumPrimitives * NumInstances, VertexCount * NumInstances);
 	StateCache.SetPrimitiveTopology(GetD3D11PrimitiveType(PrimitiveType));
 	if(NumInstances > 1)
 	{
@@ -1116,7 +1116,7 @@ void FD3D11DynamicRHI::RHIDrawPrimitiveIndirect(FRHIBuffer* ArgumentBufferRHI, u
 
 	RHI_DRAW_CALL_INC();
 
-	GPUProfilingData.RegisterGPUWork(0);
+	RegisterGPUWork(0);
 
 	CommitGraphicsResourceTables();
 	CommitNonComputeShaderConstants();
@@ -1134,7 +1134,7 @@ void FD3D11DynamicRHI::RHIDrawIndexedIndirect(FRHIBuffer* IndexBufferRHI, FRHIBu
 
 	RHI_DRAW_CALL_INC();
 
-	GPUProfilingData.RegisterGPUWork(1);
+	RegisterGPUWork(1);
 
 	CommitGraphicsResourceTables();
 	CommitNonComputeShaderConstants();
@@ -1160,7 +1160,7 @@ void FD3D11DynamicRHI::RHIDrawIndexedPrimitive(FRHIBuffer* IndexBufferRHI, int32
 	// called should make sure the input is valid, this avoid hidden bugs
 	ensure(NumPrimitives > 0);
 
-	GPUProfilingData.RegisterGPUWork(NumPrimitives * NumInstances, NumVertices * NumInstances);
+	RegisterGPUWork(NumPrimitives * NumInstances, NumVertices * NumInstances);
 
 	CommitGraphicsResourceTables();
 	CommitNonComputeShaderConstants();
@@ -1200,7 +1200,7 @@ void FD3D11DynamicRHI::RHIDrawIndexedPrimitiveIndirect(FRHIBuffer* IndexBufferRH
 
 	RHI_DRAW_CALL_INC();
 
-	GPUProfilingData.RegisterGPUWork(0);
+	RegisterGPUWork(0);
 	
 	CommitGraphicsResourceTables();
 	CommitNonComputeShaderConstants();
@@ -1264,7 +1264,7 @@ void FD3D11DynamicRHI::RHIClearMRTImpl(const bool* bClearColorArray, int32 NumCl
 		Direct3DDeviceIMContext->ClearDepthStencilView(DepthStencilView,ClearFlags,Depth,Stencil);
 	}
 
-	GPUProfilingData.RegisterGPUWork(0);
+	RegisterGPUWork(0);
 }
 
 // Blocks the CPU until the GPU catches up and goes idle.
diff --git a/Engine/Source/Runtime/Windows/D3D11RHI/Private/D3D11Device.cpp b/Engine/Source/Runtime/Windows/D3D11RHI/Private/D3D11Device.cpp
--- a/Engine/Source/Runtime/Windows/D3D11RHI/Private/D3D11Device.cpp
+++ b/Engine/Source/Runtime/Windows/D3D11RHI/Private/D3D11Device.cpp
@@ -56,7 +56,9 @@ FD3D11DynamicRHI::FD3D11DynamicRHI(IDXGIFactory1* InDXGIFactory1, D3D_FEATURE_LE
 	, PresentCounter(0)
 	, CurrentDSVAccessType(FExclusiveDepthStencil::DepthWrite_StencilWrite)
 	, bDiscardSharedConstants(false)
+#if (RHI_NEW_GPU_PROFILER == 0)
 	, GPUProfilingData(this)
+#endif
 	, Adapter(InAdapter)
 {
 	// This should be called once at the start 
@@ -244,8 +246,10 @@ void FD3D11DynamicRHI::Shutdown()
 	// Cleanup the D3D device.
 	CleanupD3DDevice();
 
+#if (RHI_NEW_GPU_PROFILER == 0)
 	// Release buffered timestamp queries
 	GPUProfilingData.FrameTiming.ReleaseResource();
+#endif
 
 	// Release the buffer of zeroes.
 	FMemory::Free(ZeroBuffer);
@@ -266,7 +270,9 @@ void FD3D11DynamicRHI::Shutdown()
 			FRHIBreadcrumb::FBuffer Buffer;
 			TCHAR const* Name = Breadcrumb->Name.GetTCHAR(Buffer);
 
+#if (RHI_NEW_GPU_PROFILER == 0)
 			GPUProfilingData.PushEvent(Name, FColor::White);
+#endif
 		}
 	}
 
@@ -274,7 +280,9 @@ void FD3D11DynamicRHI::Shutdown()
 	{
 		if (ShouldEmitBreadcrumbs())
 		{
+#if (RHI_NEW_GPU_PROFILER == 0)
 			GPUProfilingData.PopEvent();
+#endif
 		}
 
 	#if NV_AFTERMATH
diff --git a/Engine/Source/Runtime/Windows/D3D11RHI/Private/D3D11Query.cpp b/Engine/Source/Runtime/Windows/D3D11RHI/Private/D3D11Query.cpp
--- a/Engine/Source/Runtime/Windows/D3D11RHI/Private/D3D11Query.cpp
+++ b/Engine/Source/Runtime/Windows/D3D11RHI/Private/D3D11Query.cpp
@@ -166,11 +166,15 @@ bool FD3D11RenderQuery::CacheResult(FD3D11DynamicRHI& RHI, bool bWait)
 	// Adjust timer queries to engine-clock ticks.
 	if (QueryType == RQT_AbsoluteTime)
 	{
+#if RHI_NEW_GPU_PROFILER
+		checkNoEntry(); // @todo - new gpu profiler
+#else
 		// GetTimingFrequency is the number of ticks per second
 		uint64 Div = FMath::Max(1llu, FGPUTiming::GetTimingFrequency() / (1000 * 1000));
 
 		// convert from GPU specific timestamp to micro sec (1 / 1 000 000 s) which seems a reasonable resolution
 		Temp = Temp / Div;
+#endif
 	}
 
 	Result = Temp;
@@ -340,6 +344,8 @@ FD3D11EventQuery::FD3D11EventQuery(class FD3D11DynamicRHI* InD3DRHI):
  * class FD3D11BufferedGPUTiming
  *=============================================================================*/
 
+#if (RHI_NEW_GPU_PROFILER == 0)
+
 /**
  * Constructor.
  *
@@ -754,3 +760,5 @@ void FD3D11DisjointTimeStampQuery::ReleaseRHI()
 {
 
 }
+
+#endif // (RHI_NEW_GPU_PROFILER == 0)
diff --git a/Engine/Source/Runtime/Windows/D3D11RHI/Private/D3D11RHI.cpp b/Engine/Source/Runtime/Windows/D3D11RHI/Private/D3D11RHI.cpp
--- a/Engine/Source/Runtime/Windows/D3D11RHI/Private/D3D11RHI.cpp
+++ b/Engine/Source/Runtime/Windows/D3D11RHI/Private/D3D11RHI.cpp
@@ -32,7 +32,28 @@ extern void UniformBufferBeginFrame();
 void FD3D11DynamicRHI::RHIBeginFrame()
 {
 	UniformBufferBeginFrame();
+#if (RHI_NEW_GPU_PROFILER == 0)
 	GPUProfilingData.BeginFrame(this);
+#endif
+}
+
+void FD3D11DynamicRHI::RHIBeginScene()
+{
+	// Nothing to do
+}
+
+void FD3D11DynamicRHI::RHIEndScene()
+{
+	// Nothing to do
+}
+
+void FD3D11DynamicRHI::RHIEndFrame()
+{
+#if (RHI_NEW_GPU_PROFILER == 0)
+	GPUProfilingData.EndFrame();
+#endif
+	UpdateMemoryStats();
+	CurrentComputeShader = nullptr;
 }
 
 template <int32 Frequency>
@@ -276,6 +297,8 @@ void FD3D11DynamicRHI::ClearAllShaderResources()
 	ClearAllShaderResourcesForFrequency<SF_Compute>();
 }
 
+#if (RHI_NEW_GPU_PROFILER == 0)
+
 void FD3DGPUProfiler::BeginFrame(FD3D11DynamicRHI* InRHI)
 {
 	CurrentEventNode = NULL;
@@ -328,25 +351,6 @@ void FD3DGPUProfiler::BeginFrame(FD3D11DynamicRHI* InRHI)
 	}
 }
 
-void FD3D11DynamicRHI::RHIEndFrame()
-{
-	GPUProfilingData.EndFrame();
-	UpdateMemoryStats();
-	CurrentComputeShader = nullptr;
-}
-
-void FD3D11DynamicRHI::UpdateMemoryStats()
-{
-#if PLATFORM_WINDOWS && (STATS || CSV_PROFILER)
-	// Some older drivers don't support querying memory stats, so don't do anything if this fails.
-	FD3DMemoryStats MemoryStats;
-	if (SUCCEEDED(UE::DXGIUtilities::GetD3DMemoryStats(GetAdapter().DXGIAdapter, MemoryStats)))
-	{
-		UpdateD3DMemoryStatsAndCSV(MemoryStats, true);
-	}
-#endif // PLATFORM_WINDOWS && (STATS || CSV_PROFILER)
-}
-
 void FD3DGPUProfiler::EndFrame()
 {
 	if (GetEmitDrawEvents())
@@ -472,16 +476,6 @@ float FD3D11EventNode::GetTiming()
 	return Result;
 }
 
-void FD3D11DynamicRHI::RHIBeginScene()
-{
-	// Nothing to do
-}
-
-void FD3D11DynamicRHI::RHIEndScene()
-{
-	// Nothing to do
-}
-
 FD3DGPUProfiler::FD3DGPUProfiler(class FD3D11DynamicRHI* InD3DRHI)
 	: FGPUProfiler()
 	, FrameTiming(InD3DRHI, 4)
@@ -546,6 +540,8 @@ void FD3D11EventNodeFrame::LogDisjointQuery()
 	}
 }
 
+#endif // (RHI_NEW_GPU_PROFILER == 0)
+
 static void D3D11UpdateBufferStatsCommon(ID3D11Buffer* Buffer, int64 BufferSize, bool bAllocating)
 {
 	// this is a work-around on Windows. Due to the fact that there is no way
@@ -583,6 +579,18 @@ void D3D11BufferStats::UpdateBufferStats(FD3D11Buffer& Buffer, bool bAllocating)
 	}
 }
 
+void FD3D11DynamicRHI::UpdateMemoryStats()
+{
+#if PLATFORM_WINDOWS && (STATS || CSV_PROFILER)
+	// Some older drivers don't support querying memory stats, so don't do anything if this fails.
+	FD3DMemoryStats MemoryStats;
+	if (SUCCEEDED(UE::DXGIUtilities::GetD3DMemoryStats(GetAdapter().DXGIAdapter, MemoryStats)))
+	{
+		UpdateD3DMemoryStatsAndCSV(MemoryStats, true);
+	}
+#endif // PLATFORM_WINDOWS && (STATS || CSV_PROFILER)
+}
+
 ID3D11Device* FD3D11DynamicRHI::RHIGetDevice() const
 {
 	return GetDevice();
diff --git a/Engine/Source/Runtime/Windows/D3D11RHI/Private/D3D11RHIPrivate.h b/Engine/Source/Runtime/Windows/D3D11RHI/Private/D3D11RHIPrivate.h
--- a/Engine/Source/Runtime/Windows/D3D11RHI/Private/D3D11RHIPrivate.h
+++ b/Engine/Source/Runtime/Windows/D3D11RHI/Private/D3D11RHIPrivate.h
@@ -102,6 +102,7 @@ struct FD3D11GlobalStats
 	static int64 GTotalGraphicsMemory;
 };
 
+#if (RHI_NEW_GPU_PROFILER == 0)
 
 // This class has multiple inheritance but really FGPUTiming is a static class
 class FD3D11BufferedGPUTiming : public FRenderResource, public FGPUTiming
@@ -302,6 +303,8 @@ struct FD3DGPUProfiler : public FGPUProfiler
 	void EndFrame();
 };
 
+#endif // (RHI_NEW_GPU_PROFILER == 0)
+
 struct FD3D11TransitionData
 {
 	bool bUAVBarrier;
@@ -834,19 +837,20 @@ protected:
 	bool bRenderDoc = false;
 
 public:
-	void RegisterGPUWork(uint32 NumPrimitives = 0, uint32 NumVertices = 0)
-	{
-		GPUProfilingData.RegisterGPUWork(NumPrimitives, NumVertices);
-	}
-	void RegisterGPUDispatch(FIntVector GroupCount)
-	{
-		GPUProfilingData.RegisterGPUDispatch(GroupCount);
-	}
+#if RHI_NEW_GPU_PROFILER
+	void RegisterGPUWork(uint32 NumPrimitives = 0, uint32 NumVertices = 0)	{ checkNoEntry(); } // @todo - new gpu profiler
+	void RegisterGPUDispatch(FIntVector GroupCount)	                        { checkNoEntry(); } // @todo - new gpu profiler
+#else
+	void RegisterGPUWork(uint32 NumPrimitives = 0, uint32 NumVertices = 0)	{ GPUProfilingData.RegisterGPUWork(NumPrimitives, NumVertices); }
+	void RegisterGPUDispatch(FIntVector GroupCount)	                        { GPUProfilingData.RegisterGPUDispatch(GroupCount); }
+#endif
 
 	inline const FD3D11Adapter& GetAdapter() const { return Adapter; }
 
 protected:
+#if (RHI_NEW_GPU_PROFILER == 0)
 	FD3DGPUProfiler GPUProfilingData;
+#endif
 
 	FD3D11Adapter Adapter;
 
diff --git a/Engine/Source/Runtime/Windows/D3D11RHI/Private/D3D11RenderTarget.cpp b/Engine/Source/Runtime/Windows/D3D11RHI/Private/D3D11RenderTarget.cpp
--- a/Engine/Source/Runtime/Windows/D3D11RHI/Private/D3D11RenderTarget.cpp
+++ b/Engine/Source/Runtime/Windows/D3D11RHI/Private/D3D11RenderTarget.cpp
@@ -102,7 +102,7 @@ void FD3D11DynamicRHI::ResolveTextureUsingShader(
 				// Clear the destination texture.
 				if (bClearDestTexture)
 				{
-					Context.GPUProfilingData.RegisterGPUWork(0);
+					Context.RegisterGPUWork(0);
 
 					Context.Direct3DDeviceIMContext->ClearDepthStencilView(DestTextureDSV,D3D11_CLEAR_DEPTH | D3D11_CLEAR_STENCIL,0,0);
 				}
@@ -124,7 +124,7 @@ void FD3D11DynamicRHI::ResolveTextureUsingShader(
 				// Clear the destination texture.
 				if (bClearDestTexture)
 				{
-					Context.GPUProfilingData.RegisterGPUWork(0);
+					Context.RegisterGPUWork(0);
 
 					FLinearColor ClearColor(0,0,0,0);
 					Context.Direct3DDeviceIMContext->ClearRenderTargetView(DestTextureRTV,(float*)&ClearColor);
@@ -874,7 +874,7 @@ void FD3D11DynamicRHI::RHIEndRenderPass()
 
 void FD3D11DynamicRHI::ResolveTexture(UE::RHICore::FResolveTextureInfo Info)
 {
-	GPUProfilingData.RegisterGPUWork();
+	RegisterGPUWork();
 
 	FD3D11Texture* SourceTexture      = ResourceCast(Info.SourceTexture);
 	const FRHITextureDesc& SourceDesc = SourceTexture->GetDesc();
diff --git a/Engine/Source/Runtime/Windows/D3D11RHI/Private/D3D11Texture.cpp b/Engine/Source/Runtime/Windows/D3D11RHI/Private/D3D11Texture.cpp
--- a/Engine/Source/Runtime/Windows/D3D11RHI/Private/D3D11Texture.cpp
+++ b/Engine/Source/Runtime/Windows/D3D11RHI/Private/D3D11Texture.cpp
@@ -1777,7 +1777,7 @@ void FD3D11DynamicRHI::RHICopyTexture(FRHITexture* SourceTextureRHI, FRHITexture
 
 	check(SourceTexture && DestTexture);
 
-	GPUProfilingData.RegisterGPUWork();
+	RegisterGPUWork();
 
 	const FRHITextureDesc& SourceDesc = SourceTextureRHI->GetDesc();
 	const FRHITextureDesc& DestDesc = DestTextureRHI->GetDesc();
@@ -1849,7 +1849,7 @@ void FD3D11DynamicRHI::RHICopyBufferRegion(FRHIBuffer* DstBuffer, uint64 DstOffs
 	check(DstBufferD3D11 && SrcBufferD3D11);
 	check(DstOffset + NumBytes <= DstBuffer->GetSize() && SrcOffset + NumBytes <= SrcBuffer->GetSize());
 
-	GPUProfilingData.RegisterGPUWork();
+	RegisterGPUWork();
 
 	D3D11_BOX SrcBox;
 	SrcBox.left = SrcOffset;
diff --git a/Engine/Source/Runtime/Windows/D3D11RHI/Private/D3D11UAV.cpp b/Engine/Source/Runtime/Windows/D3D11RHI/Private/D3D11UAV.cpp
--- a/Engine/Source/Runtime/Windows/D3D11RHI/Private/D3D11UAV.cpp
+++ b/Engine/Source/Runtime/Windows/D3D11RHI/Private/D3D11UAV.cpp
@@ -205,7 +205,7 @@ void FD3D11DynamicRHI::ClearUAV(TRHICommandList_RecursiveHazardous<FD3D11Dynamic
 			RHICmdList.RunOnContext([UnorderedAccessView, ClearValues](auto& Context)
 			{
 				Context.Direct3DDeviceIMContext->ClearUnorderedAccessViewUint(UnorderedAccessView->View, *reinterpret_cast<const UINT(*)[4]>(ClearValues));
-				Context.GPUProfilingData.RegisterGPUWork(1);
+				Context.RegisterGPUWork(1);
 			});
 		}
 		else
