Added FD3D12Adapter accessors for static root signatures (compute and graphics), also exposed ability to reset graphics root signature on the state cache

diff --git a/Engine/Source/Runtime/D3D12RHI/Private/D3D12Adapter.h b/Engine/Source/Runtime/D3D12RHI/Private/D3D12Adapter.h
--- a/Engine/Source/Runtime/D3D12RHI/Private/D3D12Adapter.h
+++ b/Engine/Source/Runtime/D3D12RHI/Private/D3D12Adapter.h
@@ -241,6 +241,18 @@ public:
 	const FD3D12RootSignature* GetRootSignature(const class FD3D12WorkGraphShader* Shader);
 	const FD3D12RootSignature* GetGlobalRayTracingRootSignature();
 
+#if USE_STATIC_ROOT_SIGNATURE
+	FORCEINLINE const FD3D12RootSignature* GetStaticGraphicsWithConstantsRootSignature() const
+	{
+		return &StaticGraphicsWithConstantsRootSignature;
+	}
+
+	FORCEINLINE const FD3D12RootSignature* GetStaticComputeWithConstantsRootSignature() const
+	{
+		return &StaticComputeWithConstantsRootSignature;
+	}
+#endif
+
 	FORCEINLINE FD3D12RootSignatureManager* GetRootSignatureManager()
 	{
 		return &RootSignatureManager;
diff --git a/Engine/Source/Runtime/D3D12RHI/Private/D3D12StateCachePrivate.h b/Engine/Source/Runtime/D3D12RHI/Private/D3D12StateCachePrivate.h
--- a/Engine/Source/Runtime/D3D12RHI/Private/D3D12StateCachePrivate.h
+++ b/Engine/Source/Runtime/D3D12RHI/Private/D3D12StateCachePrivate.h
@@ -755,6 +755,7 @@ public:
 	void ClearState();
 
 	void ForceSetComputeRootSignature() { PipelineState.Compute.bNeedSetRootSignature = true; }
+	void ForceSetGraphicsRootSignature() { PipelineState.Graphics.bNeedSetRootSignature = true; }
 
 #if PLATFORM_SUPPORTS_BINDLESS_RENDERING
 	void QueueBindlessSRV(EShaderFrequency ShaderFrequency, FD3D12ShaderResourceView* SRV)
