Fixing D3D12RHI crash with full bindless. ClearUAV with full bindless enabled might not always have a heap set.

diff --git a/Engine/Source/Runtime/D3D12RHI/Private/D3D12DescriptorCache.h b/Engine/Source/Runtime/D3D12RHI/Private/D3D12DescriptorCache.h
--- a/Engine/Source/Runtime/D3D12RHI/Private/D3D12DescriptorCache.h
+++ b/Engine/Source/Runtime/D3D12RHI/Private/D3D12DescriptorCache.h
@@ -393,7 +393,7 @@ public:
 
 	bool IsActiveViewHeapBindless() const
 	{
-		return (BindlessResourcesHeap != nullptr) || (bUsingExplicitCacheHeaps && bExplicitViewHeapIsBindless);
+		return IsUsingBindlessResources() || (BindlessResourcesHeap != nullptr) || (bUsingExplicitCacheHeaps && bExplicitViewHeapIsBindless);
 	}
 #endif
 
