Fix r.D3D12.SubmissionTimeout not updating properly in the case of a HF

diff --git a/Engine/Source/Runtime/D3D12RHI/Private/D3D12Submission.cpp b/Engine/Source/Runtime/D3D12RHI/Private/D3D12Submission.cpp
--- a/Engine/Source/Runtime/D3D12RHI/Private/D3D12Submission.cpp
+++ b/Engine/Source/Runtime/D3D12RHI/Private/D3D12Submission.cpp
@@ -1226,8 +1226,8 @@ FD3D12DynamicRHI::FProcessResult FD3D12DynamicRHI::ProcessInterruptQueue()
 
 				if (Payload->SubmissionTime != TNumericLimits<uint64>::Max())
 				{
-					static const uint64 TimeoutCycles = FMath::TruncToInt64(GD3D12SubmissionTimeout / FPlatformTime::GetSecondsPerCycle64());
 					static const double CyclesPerSecond = 1.0 / FPlatformTime::GetSecondsPerCycle64();
+					const uint64 TimeoutCycles = FMath::TruncToInt64(GD3D12SubmissionTimeout * CyclesPerSecond);
 
 					uint64 ElapsedCycles = Timer.Elapsed - Payload->SubmissionTime.GetValue();
 
