Move some of the generic parts of D3D12DiagnosticBuffer into a cross-platform base in RHICore

* The diagnostic buffer is now always created on startup, independently of GPU breadcrumbs

diff --git a/Engine/Source/Runtime/D3D12RHI/Private/D3D12Device.cpp b/Engine/Source/Runtime/D3D12RHI/Private/D3D12Device.cpp
--- a/Engine/Source/Runtime/D3D12RHI/Private/D3D12Device.cpp
+++ b/Engine/Source/Runtime/D3D12RHI/Private/D3D12Device.cpp
@@ -386,13 +386,11 @@ void FD3D12Device::SetupAfterDeviceCreation()
 	check(!ImmediateCommandContext);
 	ImmediateCommandContext = FD3D12DynamicRHI::GetD3DRHI()->CreateCommandContext(this, ED3D12QueueType::Direct, true);
 
-	// setup the bread crumb data to track GPU progress on this command queue when GPU crash debugging is enabled
-	if (UE::RHI::UseGPUCrashBreadcrumbs())
+	// Setup diagnostic buffer that contains GPU messages as well as breadcrumb data to to track GPU progress on this command queue (when GPU crash debugging is enabled).
+	// The buffer is always allocated and bound to shaders that require it, but breadcrumbs are controlled by UE::RHI::UseGPUCrashBreadcrumbs() and WITH_RHI_BREADCRUMBS.
+	for (FD3D12Queue& Queue : Queues)
 	{
-		for (FD3D12Queue& Queue : Queues)
-		{
-			Queue.DiagnosticBuffer = MakeUnique<FD3D12DiagnosticBuffer>(Queue);
-		}
+		Queue.DiagnosticBuffer = MakeUnique<FD3D12DiagnosticBuffer>(Queue);
 	}
 }
 
diff --git a/Engine/Source/Runtime/D3D12RHI/Private/D3D12Device.h b/Engine/Source/Runtime/D3D12RHI/Private/D3D12Device.h
--- a/Engine/Source/Runtime/D3D12RHI/Private/D3D12Device.h
+++ b/Engine/Source/Runtime/D3D12RHI/Private/D3D12Device.h
@@ -7,6 +7,7 @@ D3D12Device.h: D3D12 Device Interfaces
 #pragma once
 
 #include "RHIBreadcrumbs.h"
+#include "RHIDiagnosticBuffer.h"
 
 #include "D3D12BindlessDescriptors.h"
 #include "D3D12CommandContext.h"
@@ -30,38 +31,9 @@ struct FD3D12RayTracingPipelineInfo;
 // Diagnostic buffer, backed by a virtual heap. Stays accessible after a GPU crash to allow readback of diagnostic messages.
 // Also used to track the progress of the GPU via breadcrumb markers.
 //
-class FD3D12DiagnosticBuffer
+class FD3D12DiagnosticBuffer : public FRHIDiagnosticBuffer
 {
 private:
-	// Counterpart to UEDiagnosticBuffer in D3DCommon.ush
-	struct FLane
-	{
-		uint32 Counter;
-		uint32 MessageID;
-		union
-		{
-			int32  AsInt[4];
-			uint32 AsUint[4];
-			float  AsFloat[4];
-		} Payload;
-	};
-
-	static_assert(sizeof(FLane) == 6 * sizeof(uint32), "Remember to change UEDiagnosticBuffer layout in the shaders when changing FLane");
-
-	struct FQueue
-	{
-		// Counterpart to UEDiagnosticMaxLanes in D3DCommon.ush
-		static constexpr uint32 MaxLanes = 64;
-		FLane Lanes[MaxLanes];
-
-#if WITH_RHI_BREADCRUMBS
-		// GPU breadcrumb markers
-		uint32 MarkerIn;
-		uint32 MarkerOut;
-#endif
-	} *Data = nullptr;
-
-	static constexpr uint32 SizeInBytes = sizeof(FQueue);
 
 	TRefCountPtr<FD3D12Heap> Heap;
 	TRefCountPtr<FD3D12Resource> Resource;
@@ -87,8 +59,6 @@ public:
 #endif
 
 	bool IsValid() const { return Resource.IsValid(); }
-
-	FString LogShaderAsserts(uint32 DeviceIndex, uint32 QueueIndex);
 };
 
 // Encapsulates the state required for tracking GPU queue performance across a frame.
diff --git a/Engine/Source/Runtime/D3D12RHI/Private/D3D12Util.cpp b/Engine/Source/Runtime/D3D12RHI/Private/D3D12Util.cpp
--- a/Engine/Source/Runtime/D3D12RHI/Private/D3D12Util.cpp
+++ b/Engine/Source/Runtime/D3D12RHI/Private/D3D12Util.cpp
@@ -164,77 +164,6 @@ static FString GetD3D12TextureFlagString(uint32 TextureFlags)
 	return TextureFormatText;
 }
 
-/** Log the GPU progress of the given queue to the Error log if breadcrumb data is available */
-FString FD3D12DiagnosticBuffer::LogShaderAsserts(uint32 DeviceIndex, uint32 QueueIndex)
-{
-	FString ShaderDiagnostics = FString::Printf(TEXT("\r\n\r\n\tDevice: %d, Queue %s:"), DeviceIndex, GetD3DCommandQueueTypeName(ED3D12QueueType(QueueIndex)));
-
-	bool bFound = false;
-	FString LanesString;
-	for (int32 LaneIndex = 0; LaneIndex < UE_ARRAY_COUNT(Data->Lanes); ++LaneIndex)
-	{
-		FLane& Lane = Data->Lanes[LaneIndex];
-
-		LanesString += FString::Printf(TEXT("\r\n\t\tLane %02d: "), LaneIndex);
-
-		if (Lane.Counter)
-		{
-			bFound = true;
-
-			const uint32 Line = Lane.Payload.AsUint[0];
-			const FString* File = UE::RHICore::GetDiagnosticMessage(Lane.Payload.AsUint[1]);
-			const FString* Message = UE::RHICore::GetDiagnosticMessage(Lane.Payload.AsUint[2]);
-
-			if (File && Message)
-			{
-				LanesString += FString::Printf(TEXT("Shader assertion failed - %s:%d - %s"), **File, Line, **Message);
-			}
-			else
-			{
-				LanesString += FString::Printf(TEXT("Shader assertion failed - ID: 0x%08X (%d)"), Lane.MessageID, Lane.MessageID);
-			}
-
-			{
-				const int32* Payload = Lane.Payload.AsInt;
-				if (Payload[0] < 0 || Payload[1] < 0 || Payload[2] < 0 || Payload[3] < 0)
-				{
-					LanesString += FString::Printf(TEXT("\r\n\t\t\tPayload [ int32]: %d %d %d %d"), Payload[0], Payload[1], Payload[2], Payload[3]);
-				}
-			}
-
-			{
-				const uint32* Payload = Lane.Payload.AsUint;
-				LanesString += FString::Printf(TEXT("\r\n\t\t\tPayload [uint32]: %u %u %u %u"), Payload[0], Payload[1], Payload[2], Payload[3]);
-			}
-
-			{
-				const uint32* Payload = Lane.Payload.AsUint;
-				LanesString += FString::Printf(TEXT("\r\n\t\t\tPayload [   hex]: 0x%08X 0x%08X 0x%08X 0x%08X"), Payload[0], Payload[1], Payload[2], Payload[3]);
-			}
-
-			{
-				const float* Payload = Lane.Payload.AsFloat;
-				LanesString += FString::Printf(TEXT("\r\n\t\t\tPayload [ float]: %f %f %f %f"), Payload[0], Payload[1], Payload[2], Payload[3]);
-			}
-		}
-		else
-		{
-			LanesString += TEXT("OK");
-		}
-	}
-
-	if (bFound)
-	{
-		ShaderDiagnostics += LanesString;
-	}
-	else
-	{
-		ShaderDiagnostics += TEXT("\r\n\t\tNo shader diagnostics found for this queue.");
-	}
-
-	return ShaderDiagnostics;
-}
-
 #if PLATFORM_WINDOWS
 
 static TArrayView<D3D12_DRED_BREADCRUMB_CONTEXT> GetBreadcrumbContexts(const D3D12_AUTO_BREADCRUMB_NODE* Node)
@@ -898,7 +827,9 @@ void FD3D12DynamicRHI::OutputGPUCrashReport(FTextBuilder& ErrorMessage)
 				if (!Queue.DiagnosticBuffer)
 					continue;
 
-				ShaderDiagnostics += Queue.DiagnosticBuffer->LogShaderAsserts(DeviceIndex, QueueIndex);
+				const TCHAR* QueueName = GetD3DCommandQueueTypeName(ED3D12QueueType(QueueIndex));
+
+				ShaderDiagnostics += Queue.DiagnosticBuffer->GetShaderDiagnosticMessages(DeviceIndex, QueueIndex, QueueName);
 
 #if WITH_RHI_BREADCRUMBS
 				ERHIPipeline Pipeline;
diff --git a/Engine/Source/Runtime/RHICore/Private/RHIDiagnosticBuffer.cpp b/Engine/Source/Runtime/RHICore/Private/RHIDiagnosticBuffer.cpp
new file mode 100644
--- /dev/null
+++ b/Engine/Source/Runtime/RHICore/Private/RHIDiagnosticBuffer.cpp
@@ -0,0 +1,74 @@
+// Copyright Epic Games, Inc. All Rights Reserved.
+
+#include "RHIDiagnosticBuffer.h"
+#include "RHICoreShader.h"
+
+FString FRHIDiagnosticBuffer::GetShaderDiagnosticMessages(uint32 DeviceIndex, uint32 QueueIndex, const TCHAR* QueueName)
+{
+	FString ShaderDiagnostics = FString::Printf(TEXT("\r\n\r\n\tDevice: %d, Queue %s:"), DeviceIndex, QueueName);
+
+	bool bFound = false;
+	FString LanesString;
+	for (int32 LaneIndex = 0; LaneIndex < UE_ARRAY_COUNT(Data->Lanes); ++LaneIndex)
+	{
+		FLane& Lane = Data->Lanes[LaneIndex];
+
+		LanesString += FString::Printf(TEXT("\r\n\t\tLane %02d: "), LaneIndex);
+
+		if (Lane.Counter)
+		{
+			bFound = true;
+
+			const uint32 Line = Lane.Payload.AsUint[0];
+			const FString* File = UE::RHICore::GetDiagnosticMessage(Lane.Payload.AsUint[1]);
+			const FString* Message = UE::RHICore::GetDiagnosticMessage(Lane.Payload.AsUint[2]);
+
+			if (File && Message)
+			{
+				LanesString += FString::Printf(TEXT("Shader assertion failed - %s:%d - %s"), **File, Line, **Message);
+			}
+			else
+			{
+				LanesString += FString::Printf(TEXT("Shader assertion failed - ID: 0x%08X (%d)"), Lane.MessageID, Lane.MessageID);
+			}
+
+			{
+				const int32* Payload = Lane.Payload.AsInt;
+				if (Payload[0] < 0 || Payload[1] < 0 || Payload[2] < 0 || Payload[3] < 0)
+				{
+					LanesString += FString::Printf(TEXT("\r\n\t\t\tPayload [ int32]: %d %d %d %d"), Payload[0], Payload[1], Payload[2], Payload[3]);
+				}
+			}
+
+			{
+				const uint32* Payload = Lane.Payload.AsUint;
+				LanesString += FString::Printf(TEXT("\r\n\t\t\tPayload [uint32]: %u %u %u %u"), Payload[0], Payload[1], Payload[2], Payload[3]);
+			}
+
+			{
+				const uint32* Payload = Lane.Payload.AsUint;
+				LanesString += FString::Printf(TEXT("\r\n\t\t\tPayload [   hex]: 0x%08X 0x%08X 0x%08X 0x%08X"), Payload[0], Payload[1], Payload[2], Payload[3]);
+			}
+
+			{
+				const float* Payload = Lane.Payload.AsFloat;
+				LanesString += FString::Printf(TEXT("\r\n\t\t\tPayload [ float]: %f %f %f %f"), Payload[0], Payload[1], Payload[2], Payload[3]);
+			}
+		}
+		else
+		{
+			LanesString += TEXT("OK");
+		}
+	}
+
+	if (bFound)
+	{
+		ShaderDiagnostics += LanesString;
+	}
+	else
+	{
+		ShaderDiagnostics += TEXT("\r\n\t\tNo shader diagnostics found for this queue.");
+	}
+
+	return ShaderDiagnostics;
+}
diff --git a/Engine/Source/Runtime/RHICore/Public/RHIDiagnosticBuffer.h b/Engine/Source/Runtime/RHICore/Public/RHIDiagnosticBuffer.h
new file mode 100644
--- /dev/null
+++ b/Engine/Source/Runtime/RHICore/Public/RHIDiagnosticBuffer.h
@@ -0,0 +1,51 @@
+// Copyright Epic Games, Inc. All Rights Reserved.
+
+#pragma once
+
+#include "RHIBreadcrumbs.h"
+
+/// Common base for platform-specific implementations of GPU diagnostic buffer.
+/// The diagnostic buffer contains GPU messages (debug logs, shader asserts, etc.) and GPU progress breadcrumbs.
+/// This buffer is persistently mapped and can be accessed on CPU at any point, including after a GPU crash has been detected.
+/// Platform-specific code is responsible for allocating the actual underlying resource and binding it to shaders that need it.
+/// Diagnostic buffer functionality may be used independently of GPU breadcrumbs.
+class FRHIDiagnosticBuffer
+{
+public:
+
+	// Counterpart to UEDiagnosticBuffer in shader code
+	struct FLane
+	{
+		uint32 Counter;
+		uint32 MessageID;
+		union
+		{
+			int32  AsInt[4];
+			uint32 AsUint[4];
+			float  AsFloat[4];
+		} Payload;
+	};
+
+	static_assert(sizeof(FLane) == 6 * sizeof(uint32), "Remember to change UEDiagnosticBuffer layout in the shaders when changing FLane");
+
+	struct FQueue
+	{
+		// Counterpart to UEDiagnosticMaxLanes in shader code
+		static constexpr uint32 MaxLanes = 64;
+		FLane Lanes[MaxLanes];
+
+		#if WITH_RHI_BREADCRUMBS
+		// GPU breadcrumb markers
+		uint32 MarkerIn;
+		uint32 MarkerOut;
+		#endif
+	}; 
+	
+	static constexpr uint32 SizeInBytes = sizeof(FQueue);
+
+	// Persistently mapped diagnostic buffer data, initialized by platform-specific code
+	FQueue* Data = nullptr;
+
+	// Log the GPU progress of the given queue to the Error log if breadcrumb data is available
+	FString RHICORE_API GetShaderDiagnosticMessages(uint32 DeviceIndex, uint32 QueueIndex, const TCHAR* QueueName);
+};
