- Add RHIShaderBindingLayout which can be used to declare a fixed shader resource table for a certain shader type (code only contains RHI side changes and is not used yet)
- Add Data driven shader platform support for StaticShaderBindingLayout support: Unsupported, RaytracingOnly, AllShaderTypes
- Rename and move ED3D12RootSignatureFlags to EShaderBindingLayoutFlags because it's part of the FRHIShaderBindingLayout now

diff --git a/Engine/Source/Runtime/D3D12RHI/Private/D3D12Adapter.cpp b/Engine/Source/Runtime/D3D12RHI/Private/D3D12Adapter.cpp
--- a/Engine/Source/Runtime/D3D12RHI/Private/D3D12Adapter.cpp
+++ b/Engine/Source/Runtime/D3D12RHI/Private/D3D12Adapter.cpp
@@ -1237,33 +1237,33 @@ void FD3D12Adapter::InitializeDevices()
 #endif
 
 #if USE_STATIC_ROOT_SIGNATURE
-		ED3D12RootSignatureFlags GraphicsFlags{};
+		EShaderBindingLayoutFlags GraphicsFlags{};
 #if PLATFORM_SUPPORTS_MESH_SHADERS
-		EnumAddFlags(GraphicsFlags, ED3D12RootSignatureFlags::AllowMeshShaders);
+		EnumAddFlags(GraphicsFlags, EShaderBindingLayoutFlags::AllowMeshShaders);
 #endif
 		if (BindlessResourcesConfig == ERHIBindlessConfiguration::AllShaders)
 		{
-			EnumAddFlags(GraphicsFlags, ED3D12RootSignatureFlags::BindlessResources);
+			EnumAddFlags(GraphicsFlags, EShaderBindingLayoutFlags::BindlessResources);
 		}
 		if (BindlessSamplersConfig == ERHIBindlessConfiguration::AllShaders)
 		{
-			EnumAddFlags(GraphicsFlags, ED3D12RootSignatureFlags::BindlessSamplers);
+			EnumAddFlags(GraphicsFlags, EShaderBindingLayoutFlags::BindlessSamplers);
 		}
 
 		StaticGraphicsRootSignature.InitStaticGraphicsRootSignature(GraphicsFlags);
-		StaticGraphicsWithConstantsRootSignature.InitStaticGraphicsRootSignature(GraphicsFlags | ED3D12RootSignatureFlags::RootConstants);
+		StaticGraphicsWithConstantsRootSignature.InitStaticGraphicsRootSignature(GraphicsFlags | EShaderBindingLayoutFlags::RootConstants);
 		StaticComputeRootSignature.InitStaticComputeRootSignatureDesc(GraphicsFlags);
-		StaticComputeWithConstantsRootSignature.InitStaticComputeRootSignatureDesc(GraphicsFlags | ED3D12RootSignatureFlags::RootConstants);
+		StaticComputeWithConstantsRootSignature.InitStaticComputeRootSignatureDesc(GraphicsFlags | EShaderBindingLayoutFlags::RootConstants);
 
 #if D3D12_RHI_RAYTRACING
-		ED3D12RootSignatureFlags RayTracingFlags{};
+		EShaderBindingLayoutFlags RayTracingFlags{};
 		if (BindlessResourcesConfig != ERHIBindlessConfiguration::Disabled)
 		{
-			EnumAddFlags(RayTracingFlags, ED3D12RootSignatureFlags::BindlessResources);
+			EnumAddFlags(RayTracingFlags, EShaderBindingLayoutFlags::BindlessResources);
 		}
 		if (BindlessSamplersConfig != ERHIBindlessConfiguration::Disabled)
 		{
-			EnumAddFlags(RayTracingFlags, ED3D12RootSignatureFlags::BindlessSamplers);
+			EnumAddFlags(RayTracingFlags, EShaderBindingLayoutFlags::BindlessSamplers);
 		}
 
 		StaticRayTracingGlobalRootSignature.InitStaticRayTracingGlobalRootSignatureDesc(RayTracingFlags);
diff --git a/Engine/Source/Runtime/D3D12RHI/Private/D3D12RootSignature.cpp b/Engine/Source/Runtime/D3D12RHI/Private/D3D12RootSignature.cpp
--- a/Engine/Source/Runtime/D3D12RHI/Private/D3D12RootSignature.cpp
+++ b/Engine/Source/Runtime/D3D12RHI/Private/D3D12RootSignature.cpp
@@ -401,13 +401,13 @@ FD3D12RootSignatureDesc::FD3D12RootSignatureDesc(const FD3D12QuantizedBoundShade
 
 }
 
-void FD3D12RootSignature::InitStaticGraphicsRootSignature(ED3D12RootSignatureFlags InFlags)
+void FD3D12RootSignature::InitStaticGraphicsRootSignature(EShaderBindingLayoutFlags InFlags)
 {
 	D3D12ShaderUtils::FBinaryRootSignatureCreator Creator;
 	D3D12ShaderUtils::CreateGfxRootSignature(Creator, InFlags);
 	Init(Creator.Finalize());
 
-	if (EnumHasAnyFlags(InFlags, ED3D12RootSignatureFlags::RootConstants))
+	if (EnumHasAnyFlags(InFlags, EShaderBindingLayoutFlags::RootConstants))
 	{
 		for (int32 ParameterSlot = 0; ParameterSlot < Creator.Parameters.Num(); ++ParameterSlot)
 		{
@@ -424,13 +424,13 @@ void FD3D12RootSignature::InitStaticGraphicsRootSignature(ED3D12RootSignatureFla
 	}
 }
 
-void FD3D12RootSignature::InitStaticComputeRootSignatureDesc(ED3D12RootSignatureFlags InFlags)
+void FD3D12RootSignature::InitStaticComputeRootSignatureDesc(EShaderBindingLayoutFlags InFlags)
 {
 	D3D12ShaderUtils::FBinaryRootSignatureCreator Creator;
 	D3D12ShaderUtils::CreateComputeRootSignature(Creator, InFlags);
 	Init(Creator.Finalize());
 
-	if (EnumHasAnyFlags(InFlags, ED3D12RootSignatureFlags::RootConstants))
+	if (EnumHasAnyFlags(InFlags, EShaderBindingLayoutFlags::RootConstants))
 	{
 		for (int32 ParameterSlot = 0; ParameterSlot < Creator.Parameters.Num(); ++ParameterSlot)
 		{
@@ -448,14 +448,14 @@ void FD3D12RootSignature::InitStaticComputeRootSignatureDesc(ED3D12RootSignature
 }
 
 #if D3D12_RHI_RAYTRACING
-void FD3D12RootSignature::InitStaticRayTracingGlobalRootSignatureDesc(ED3D12RootSignatureFlags InFlags)
+void FD3D12RootSignature::InitStaticRayTracingGlobalRootSignatureDesc(EShaderBindingLayoutFlags InFlags)
 {
 	D3D12ShaderUtils::FBinaryRootSignatureCreator Creator;
 	D3D12ShaderUtils::CreateRayTracingSignature(Creator, false, FD3D12_ROOT_SIGNATURE_FLAG_GLOBAL_ROOT_SIGNATURE, InFlags);
 	Init(Creator.Finalize(), UE_HLSL_SPACE_RAY_TRACING_GLOBAL);
 }
 
-void FD3D12RootSignature::InitStaticRayTracingLocalRootSignatureDesc(ED3D12RootSignatureFlags InFlags)
+void FD3D12RootSignature::InitStaticRayTracingLocalRootSignatureDesc(EShaderBindingLayoutFlags InFlags)
 {
 	D3D12ShaderUtils::FBinaryRootSignatureCreator Creator;
 	D3D12ShaderUtils::CreateRayTracingSignature(Creator, true, D3D12_ROOT_SIGNATURE_FLAG_LOCAL_ROOT_SIGNATURE, InFlags);
diff --git a/Engine/Source/Runtime/D3D12RHI/Private/D3D12RootSignature.h b/Engine/Source/Runtime/D3D12RHI/Private/D3D12RootSignature.h
--- a/Engine/Source/Runtime/D3D12RHI/Private/D3D12RootSignature.h
+++ b/Engine/Source/Runtime/D3D12RHI/Private/D3D12RootSignature.h
@@ -93,11 +93,11 @@ public:
 	void Init(const FD3D12QuantizedBoundShaderState& InQBSS);
 	void Init(const D3D12_VERSIONED_ROOT_SIGNATURE_DESC& InDesc, uint32 BindingSpace = 0);
 
-	void InitStaticGraphicsRootSignature(ED3D12RootSignatureFlags InFlags);
-	void InitStaticComputeRootSignatureDesc(ED3D12RootSignatureFlags InFlags);
+	void InitStaticGraphicsRootSignature(EShaderBindingLayoutFlags InFlags);
+	void InitStaticComputeRootSignatureDesc(EShaderBindingLayoutFlags InFlags);
 #if D3D12_RHI_RAYTRACING
-	void InitStaticRayTracingGlobalRootSignatureDesc(ED3D12RootSignatureFlags InFlags);
-	void InitStaticRayTracingLocalRootSignatureDesc(ED3D12RootSignatureFlags InFlags);
+	void InitStaticRayTracingGlobalRootSignatureDesc(EShaderBindingLayoutFlags InFlags);
+	void InitStaticRayTracingLocalRootSignatureDesc(EShaderBindingLayoutFlags InFlags);
 #endif
 
 	ID3D12RootSignature* GetRootSignature() const { return RootSignature.GetReference(); }
diff --git a/Engine/Source/Runtime/D3D12RHI/Public/D3D12RootSignatureDefinitions.h b/Engine/Source/Runtime/D3D12RHI/Public/D3D12RootSignatureDefinitions.h
--- a/Engine/Source/Runtime/D3D12RHI/Public/D3D12RootSignatureDefinitions.h
+++ b/Engine/Source/Runtime/D3D12RHI/Public/D3D12RootSignatureDefinitions.h
@@ -9,22 +9,12 @@
 #include "Misc/StringBuilder.h"
 #include "RayTracingBuiltInResources.h"
 #include "RHIDefinitions.h"
+#include "RHIShaderBindingLayout.h"
 
 #if !defined(D3D12_MAJOR_VERSION)
 	#include "D3D12ThirdParty.h"
 #endif
 
-enum class ED3D12RootSignatureFlags
-{
-	None = 0,
-	AllowMeshShaders = 1 << 0,
-	InputAssembler = 1 << 1,
-	BindlessResources = 1 << 2,
-	BindlessSamplers = 1 << 3,
-	RootConstants = 1 << 4,
-};
-ENUM_CLASS_FLAGS(ED3D12RootSignatureFlags)
-
 namespace D3D12ShaderUtils
 {
 	namespace StaticRootSignatureConstants
@@ -183,7 +173,7 @@ namespace D3D12ShaderUtils
 	// Simple base class to help write out a root signature (subclass to generate either to a binary struct or a #define)
 	struct FRootSignatureCreator
 	{
-		ED3D12RootSignatureFlags Flags = ED3D12RootSignatureFlags::None;
+		EShaderBindingLayoutFlags Flags = EShaderBindingLayoutFlags::None;
 		uint32 RegisterSpace = 0;
 
 		virtual ~FRootSignatureCreator() = default;
@@ -194,27 +184,27 @@ namespace D3D12ShaderUtils
 		virtual void AddTable(ERootSignatureVisibility Visibility, ERootSignatureRangeType Type, int32 NumDescriptors, D3D12_DESCRIPTOR_RANGE_FLAGS FlagsOverride = D3D12_DESCRIPTOR_RANGE_FLAG_NONE) = 0;
 		virtual void AddConstantsParameter(uint32 Num32BitValues, uint32 Register, uint32 Space) = 0;
 
-		void SetFlags(ED3D12RootSignatureFlags InFlags)
+		void SetFlags(EShaderBindingLayoutFlags InFlags)
 		{
 			Flags = InFlags;
 
-			if (EnumHasAnyFlags(InFlags, ED3D12RootSignatureFlags::InputAssembler))
+			if (EnumHasAnyFlags(InFlags, EShaderBindingLayoutFlags::InputAssembler))
 			{
 				AddRootFlag(D3D12_ROOT_SIGNATURE_FLAG_ALLOW_INPUT_ASSEMBLER_INPUT_LAYOUT);
 			}
 
-			if (EnumHasAnyFlags(InFlags, ED3D12RootSignatureFlags::BindlessResources))
+			if (EnumHasAnyFlags(InFlags, EShaderBindingLayoutFlags::BindlessResources))
 			{
 				AddRootFlag(D3D12_ROOT_SIGNATURE_FLAG_CBV_SRV_UAV_HEAP_DIRECTLY_INDEXED);
 			}
 
-			if (EnumHasAnyFlags(InFlags, ED3D12RootSignatureFlags::BindlessSamplers))
+			if (EnumHasAnyFlags(InFlags, EShaderBindingLayoutFlags::BindlessSamplers))
 			{
 				AddRootFlag(D3D12_ROOT_SIGNATURE_FLAG_SAMPLER_HEAP_DIRECTLY_INDEXED);
 			}
 		}
 
-		bool HasFlags(ED3D12RootSignatureFlags InFlags) const
+		bool HasFlags(EShaderBindingLayoutFlags InFlags) const
 		{
 			return EnumHasAllFlags(Flags, InFlags);
 		}
@@ -233,12 +223,12 @@ namespace D3D12ShaderUtils
 		{
 			if (Type == ERootSignatureRangeType::SRV || Type == ERootSignatureRangeType::UAV)
 			{
-				return EnumHasAnyFlags(Flags, ED3D12RootSignatureFlags::BindlessResources);
+				return EnumHasAnyFlags(Flags, EShaderBindingLayoutFlags::BindlessResources);
 			}
 
 			if (Type == ERootSignatureRangeType::Sampler)
 			{
-				return EnumHasAnyFlags(Flags, ED3D12RootSignatureFlags::BindlessSamplers);
+				return EnumHasAnyFlags(Flags, EShaderBindingLayoutFlags::BindlessSamplers);
 			}
 
 			return false;
@@ -397,17 +387,17 @@ namespace D3D12ShaderUtils
 	}
 
 	// Fat/Static Gfx Root Signature
-	inline void CreateGfxRootSignature(FRootSignatureCreator& Creator, ED3D12RootSignatureFlags InFlags)
+	inline void CreateGfxRootSignature(FRootSignatureCreator& Creator, EShaderBindingLayoutFlags InFlags)
 	{
 		// Ensure the creator starts in a clean state (in cases of creator reuse, etc.).
-		Creator.SetFlags(InFlags | ED3D12RootSignatureFlags::InputAssembler);
+		Creator.SetFlags(InFlags | EShaderBindingLayoutFlags::InputAssembler);
 
 		AddAllStandardTablesForVisibility(Creator, ERootSignatureVisibility::Pixel);
 		AddAllStandardTablesForVisibility(Creator, ERootSignatureVisibility::Vertex);
 		AddAllStandardTablesForVisibility(Creator, ERootSignatureVisibility::Geometry);
 
 #if !defined(D3D12RHI_TOOLS_MESH_SHADERS_UNSUPPORTED)
-		if (EnumHasAnyFlags(InFlags, ED3D12RootSignatureFlags::AllowMeshShaders))
+		if (EnumHasAnyFlags(InFlags, EShaderBindingLayoutFlags::AllowMeshShaders))
 		{
 			AddAllStandardTablesForVisibility(Creator, ERootSignatureVisibility::Mesh);
 			AddAllStandardTablesForVisibility(Creator, ERootSignatureVisibility::Amplification);
@@ -417,7 +407,7 @@ namespace D3D12ShaderUtils
 		Creator.AddTable(ERootSignatureVisibility::Pixel, ERootSignatureRangeType::UAV, MAX_UAVS);
 		Creator.AddTable(ERootSignatureVisibility::Vertex, ERootSignatureRangeType::UAV, MAX_UAVS);
 
-		if (EnumHasAnyFlags(InFlags, ED3D12RootSignatureFlags::RootConstants))
+		if (EnumHasAnyFlags(InFlags, EShaderBindingLayoutFlags::RootConstants))
 		{
 			const uint32 NumConstants = 4u;
 			Creator.AddConstantsParameter(NumConstants, 0u, UE_HLSL_SPACE_SHADER_ROOT_CONSTANTS); // UERootConstants
@@ -425,14 +415,14 @@ namespace D3D12ShaderUtils
 	}
 
 	// Fat/Static Compute Root Signature
-	inline void CreateComputeRootSignature(FRootSignatureCreator& Creator, ED3D12RootSignatureFlags InFlags)
+	inline void CreateComputeRootSignature(FRootSignatureCreator& Creator, EShaderBindingLayoutFlags InFlags)
 	{
 		// Ensure the creator starts in a clean state (in cases of creator reuse, etc.).
 		Creator.SetFlags(InFlags);
 		AddAllStandardTablesForVisibility(Creator, ERootSignatureVisibility::All);
 		Creator.AddTable(ERootSignatureVisibility::All, ERootSignatureRangeType::UAV, MAX_UAVS);
 
-		if (EnumHasAnyFlags(InFlags, ED3D12RootSignatureFlags::RootConstants))
+		if (EnumHasAnyFlags(InFlags, EShaderBindingLayoutFlags::RootConstants))
 		{
 			const uint32 NumConstants = 4u;
 			Creator.AddConstantsParameter(NumConstants, 0u, UE_HLSL_SPACE_SHADER_ROOT_CONSTANTS); // UERootConstants
@@ -440,7 +430,7 @@ namespace D3D12ShaderUtils
 	}
 
 #if !defined(D3D12RHI_TOOLS_RAYTRACING_SHADERS_UNSUPPORTED)
-	inline void CreateRayTracingSignature(FRootSignatureCreator& Creator, bool bLocalRootSignature, D3D12_ROOT_SIGNATURE_FLAGS BaseRootFlags, ED3D12RootSignatureFlags InFlags)
+	inline void CreateRayTracingSignature(FRootSignatureCreator& Creator, bool bLocalRootSignature, D3D12_ROOT_SIGNATURE_FLAGS BaseRootFlags, EShaderBindingLayoutFlags InFlags)
 	{
 		Creator.SetFlags(InFlags);
 		Creator.AddRootFlag(BaseRootFlags);
@@ -455,7 +445,7 @@ namespace D3D12ShaderUtils
 			Creator.AddConstantsParameter(NumConstants, RAY_TRACING_SYSTEM_ROOTCONSTANT_REGISTER, UE_HLSL_SPACE_RAY_TRACING_SYSTEM);
 		}
 
-		if (Creator.HasFlags(ED3D12RootSignatureFlags::BindlessResources))
+		if (Creator.HasFlags(EShaderBindingLayoutFlags::BindlessResources))
 		{
 			for (uint32 Index = 0; Index < MAX_CBS; Index++)
 			{
@@ -470,7 +460,7 @@ namespace D3D12ShaderUtils
 		Creator.AddTable(ERootSignatureVisibility::All, ERootSignatureRangeType::UAV, MAX_UAVS, D3D12_DESCRIPTOR_RANGE_FLAG_DATA_STATIC_WHILE_SET_AT_EXECUTE | D3D12_DESCRIPTOR_RANGE_FLAG_DESCRIPTORS_VOLATILE);
 	}
 
-	inline FString GenerateRayTracingSignatureString(bool bLocalRootSignature, D3D12_ROOT_SIGNATURE_FLAGS BaseRootFlags, ED3D12RootSignatureFlags InFlags)
+	inline FString GenerateRayTracingSignatureString(bool bLocalRootSignature, D3D12_ROOT_SIGNATURE_FLAGS BaseRootFlags, EShaderBindingLayoutFlags InFlags)
 	{
 		FTextRootSignatureCreator Creator;
 		CreateRayTracingSignature(Creator, bLocalRootSignature, BaseRootFlags, InFlags);
@@ -478,7 +468,7 @@ namespace D3D12ShaderUtils
 	}
 #endif //!defined(D3D12RHI_TOOLS_RAYTRACING_SHADERS_UNSUPPORTED)
 
-	inline FString GenerateRootSignatureString(EShaderFrequency InFrequency, ED3D12RootSignatureFlags InFlags)
+	inline FString GenerateRootSignatureString(EShaderFrequency InFrequency, EShaderBindingLayoutFlags InFlags)
 	{
 		FTextRootSignatureCreator Creator;
 
diff --git a/Engine/Source/Runtime/RHI/Private/DataDrivenShaderPlatformInfo.cpp b/Engine/Source/Runtime/RHI/Private/DataDrivenShaderPlatformInfo.cpp
--- a/Engine/Source/Runtime/RHI/Private/DataDrivenShaderPlatformInfo.cpp
+++ b/Engine/Source/Runtime/RHI/Private/DataDrivenShaderPlatformInfo.cpp
@@ -135,6 +135,33 @@ static inline uint32 GetSectionBindlessSupport(const FConfigSection& Section, FN
 	return OriginalValue;
 }
 
+static inline uint32 GetSectionStaticShaderBindingLayoutSupport(const FConfigSection& Section, FName Key, uint32 OriginalValue)
+{
+	const FConfigValue* ConfigValue = Section.Find(Key);
+	if (ConfigValue != nullptr)
+	{
+		FString Value = ConfigValue->GetValue();
+		if (Value == TEXT("Unsupported"))
+		{
+			return uint32(ERHIStaticShaderBindingLayoutSupport::Unsupported);
+		}
+		if (Value == TEXT("RayTracingOnly"))
+		{
+			return uint32(ERHIStaticShaderBindingLayoutSupport::RayTracingOnly);
+		}
+		else if (Value == TEXT("AllShaderTypes"))
+		{
+			return uint32(ERHIStaticShaderBindingLayoutSupport::AllShaderTypes);
+		}
+		else
+		{
+			checkf(false, TEXT("Unknown ERHIStaticShaderBindingLayoutSupport value \"%s\" for %s"), *Value, *Key.ToString());
+		}
+	}
+
+	return OriginalValue;
+}
+
 void FGenericDataDrivenShaderPlatformInfo::SetDefaultValues()
 {
 	MaxFeatureLevel = ERHIFeatureLevel::Num;
@@ -192,6 +219,10 @@ void FGenericDataDrivenShaderPlatformInfo::ParseDataDrivenShaderInfo(const FConf
 #define GET_SECTION_BINDLESS_SUPPORT_HELPER(SettingName) \
 	Info.SettingName = GetSectionBindlessSupport(Section, #SettingName, Info.SettingName); \
 	ADD_TO_PROPERTIES_STRING(SettingName, Info.SettingName)
+		
+#define GET_SECTION_STATIC_SHADER_BINDING_LAYOUT_SUPPORT_HELPER(SettingName) \
+	Info.SettingName = GetSectionStaticShaderBindingLayoutSupport(Section, #SettingName, Info.SettingName); \
+	ADD_TO_PROPERTIES_STRING(SettingName, Info.SettingName)
 
 	// These properties will be exposed to the MaterialEditor that use the ShaderPlatformInfo Node
 	// If you remove/rename a property be sure to address this in UMaterialExpressionDataDrivenShaderPlatformInfoSwitch serialization
@@ -285,6 +316,7 @@ void FGenericDataDrivenShaderPlatformInfo::ParseDataDrivenShaderInfo(const FConf
 	GET_SECTION_BOOL_HELPER(bSupportsFFTBloom);
 	GET_SECTION_BOOL_HELPER(bSupportsVertexShaderLayer);
 	GET_SECTION_BINDLESS_SUPPORT_HELPER(BindlessSupport);
+	GET_SECTION_STATIC_SHADER_BINDING_LAYOUT_SUPPORT_HELPER(StaticShaderBindingLayoutSupport);
 	GET_SECTION_BOOL_HELPER(bSupportsVolumeTextureAtomics);
 	GET_SECTION_BOOL_HELPER(bSupportsROV);
 	GET_SECTION_BOOL_HELPER(bSupportsOIT);
diff --git a/Engine/Source/Runtime/RHI/Private/RHIContext.cpp b/Engine/Source/Runtime/RHI/Private/RHIContext.cpp
--- a/Engine/Source/Runtime/RHI/Private/RHIContext.cpp
+++ b/Engine/Source/Runtime/RHI/Private/RHIContext.cpp
@@ -4,6 +4,7 @@
 #include "DynamicRHI.h"
 #include "RHI.h"
 #include "RHIStats.h"
+#include "RHIShaderBindingLayout.h"
 
 void RHIGenerateCrossGPUPreTransferFences(const TArrayView<const FTransferResourceParams> Params, TArray<FCrossGPUTransferFence*>& OutPreTransfer)
 {
@@ -47,6 +48,41 @@ void RHIGenerateCrossGPUPreTransferFences(const TArrayView<const FTransferResour
 	}
 }
 
+FUniformBufferStaticBindings::FUniformBufferStaticBindings(const FRHIShaderBindingLayout* InShaderBindingLayout) : ShaderBindingLayout(InShaderBindingLayout) 
+{
+	UniformBuffers.SetNum(ShaderBindingLayout ? ShaderBindingLayout->GetNumUniformBufferEntries() : 0);
+}
+
+void FUniformBufferStaticBindings::AddUniformBuffer(FRHIUniformBuffer* UniformBuffer)
+{
+	checkf(UniformBuffer, TEXT("Attemped to assign a null uniform buffer to the global uniform buffer bindings."));		
+	const FRHIUniformBufferLayout& Layout = UniformBuffer->GetLayout();
+
+	// Only care about the static slots if no shader desc is used, otherwise the desc is used to validate that it contains the uniform buffer
+	if (ShaderBindingLayout)
+	{
+		const FRHIUniformBufferShaderBindingLayout* UniformBufferEntry = ShaderBindingLayout->FindEntry(Layout.GetHash());
+		checkf(UniformBufferEntry, TEXT("Attempted to set a static uniform buffer %s which is not defined in the ShaderBindingLayout provided."), *Layout.GetDebugName());
+		UniformBuffers[UniformBufferEntry->CBVResourceIndex] = UniformBuffer;
+	}
+	else
+	{
+		const FUniformBufferStaticSlot Slot = Layout.StaticSlot;
+		checkf(IsUniformBufferStaticSlotValid(Slot), TEXT("Attempted to set a global uniform buffer %s with an invalid slot."), *Layout.GetDebugName());
+
+		#if VALIDATE_UNIFORM_BUFFER_STATIC_BINDINGS
+		if (int32 SlotIndex = Slots.Find(Slot); SlotIndex != INDEX_NONE)
+		{
+			checkf(UniformBuffers[SlotIndex] == UniformBuffer, TEXT("Uniform Buffer %s was added multiple times to the binding array but with different values."), *Layout.GetDebugName());
+		}
+		#endif
+
+		Slots.Add(Slot);
+		UniformBuffers.Add(UniformBuffer);
+		SlotCount = FMath::Max(SlotCount, Slot + 1);
+	}
+}
+
 void IRHICommandContextPSOFallback::SetGraphicsPipelineStateFromInitializer(const FGraphicsPipelineStateInitializer& PsoInit, uint32 StencilRef, bool bApplyAdditionalState)
 {
 	RHISetBoundShaderState(
diff --git a/Engine/Source/Runtime/RHI/Public/DataDrivenShaderPlatformInfo.h b/Engine/Source/Runtime/RHI/Public/DataDrivenShaderPlatformInfo.h
--- a/Engine/Source/Runtime/RHI/Public/DataDrivenShaderPlatformInfo.h
+++ b/Engine/Source/Runtime/RHI/Public/DataDrivenShaderPlatformInfo.h
@@ -108,6 +108,7 @@ class FGenericDataDrivenShaderPlatformInfo
 	uint32 bSupportsRayTracingShaders : 1;
 	uint32 bSupportsVertexShaderLayer : 1;
 	uint32 BindlessSupport : int32(ERHIBindlessSupport::NumBits);
+	uint32 StaticShaderBindingLayoutSupport : int32(ERHIStaticShaderBindingLayoutSupport::NumBits);
 	uint32 bSupportsVolumeTextureAtomics : 1;
 	uint32 bSupportsROV : 1;
 	uint32 bSupportsOIT : 1;
@@ -720,6 +721,12 @@ public:
 		return static_cast<ERHIBindlessSupport>(Infos[Platform].BindlessSupport);
 	}
 
+	static FORCEINLINE_DEBUGGABLE const ERHIStaticShaderBindingLayoutSupport GetStaticShaderBindingLayoutSupport(const FStaticShaderPlatform Platform)
+	{
+		check(IsValid(Platform));
+		return static_cast<ERHIStaticShaderBindingLayoutSupport>(Infos[Platform].StaticShaderBindingLayoutSupport);
+	}
+
 	static FORCEINLINE_DEBUGGABLE const bool GetSupportsVolumeTextureAtomics(const FStaticShaderPlatform Platform)
 	{
 		return Infos[Platform].bSupportsVolumeTextureAtomics;
@@ -1195,6 +1202,12 @@ inline ERHIBindlessSupport RHIGetBindlessSupport(const FStaticShaderPlatform Pla
 	return FDataDrivenShaderPlatformInfo::GetBindlessSupport(Platform);
 }
 
+/** True if the given shader platform supports static shader resource tables. */
+inline ERHIStaticShaderBindingLayoutSupport RHIGetStaticShaderBindingLayoutSupport(const FStaticShaderPlatform Platform)
+{
+	return FDataDrivenShaderPlatformInfo::GetStaticShaderBindingLayoutSupport(Platform);
+}
+
 inline bool RHISupportsVolumeTextureAtomics(EShaderPlatform Platform)
 {
 	return FDataDrivenShaderPlatformInfo::GetSupportsVolumeTextureAtomics(Platform);
diff --git a/Engine/Source/Runtime/RHI/Public/RHICommandList.h b/Engine/Source/Runtime/RHI/Public/RHICommandList.h
--- a/Engine/Source/Runtime/RHI/Public/RHICommandList.h
+++ b/Engine/Source/Runtime/RHI/Public/RHICommandList.h
@@ -244,6 +244,8 @@ struct FRayTracingShaderBindings
 	FRHIUniformBuffer* UniformBuffers[16] = {};
 	FRHISamplerState* Samplers[32] = {};
 	FRHIUnorderedAccessView* UAVs[16] = {};
+
+	FUniformBufferStaticBindings StaticUniformBuffers;
 };
 
 struct FRayTracingLocalShaderBindings
diff --git a/Engine/Source/Runtime/RHI/Public/RHIContext.h b/Engine/Source/Runtime/RHI/Public/RHIContext.h
--- a/Engine/Source/Runtime/RHI/Public/RHIContext.h
+++ b/Engine/Source/Runtime/RHI/Public/RHIContext.h
@@ -23,6 +23,7 @@
 class FRHIDepthRenderTargetView;
 class FRHIRenderTargetView;
 class FRHISetRenderTargetsInfo;
+class FRHIShaderBindingLayout;
 struct FViewportBounds;
 struct FRayTracingGeometryInstance;
 struct FRayTracingShaderBindings;
@@ -58,24 +59,9 @@ public:
 		}
 	}
 
-	inline void AddUniformBuffer(FRHIUniformBuffer* UniformBuffer)
-	{
-		checkf(UniformBuffer, TEXT("Attemped to assign a null uniform buffer to the global uniform buffer bindings."));
-		const FRHIUniformBufferLayout& Layout = UniformBuffer->GetLayout();
-		const FUniformBufferStaticSlot Slot = Layout.StaticSlot;
-		checkf(IsUniformBufferStaticSlotValid(Slot), TEXT("Attempted to set a global uniform buffer %s with an invalid slot."), *Layout.GetDebugName());
-
-#if VALIDATE_UNIFORM_BUFFER_STATIC_BINDINGS
-		if (int32 SlotIndex = Slots.Find(Slot); SlotIndex != INDEX_NONE)
-		{
-			checkf(UniformBuffers[SlotIndex] == UniformBuffer, TEXT("Uniform Buffer %s was added multiple times to the binding array but with different values."), *Layout.GetDebugName());
-		}
-#endif
+	RHI_API FUniformBufferStaticBindings(const FRHIShaderBindingLayout* InSRTDesc);
 
-		Slots.Add(Slot);
-		UniformBuffers.Add(UniformBuffer);
-		SlotCount = FMath::Max(SlotCount, Slot + 1);
-	}
+	RHI_API void AddUniformBuffer(FRHIUniformBuffer* UniformBuffer);
 
 	inline void TryAddUniformBuffer(FRHIUniformBuffer* UniformBuffer)
 	{
@@ -105,6 +91,11 @@ public:
 		return SlotCount;
 	}
 
+	const FRHIShaderBindingLayout* GetShaderBindingLayout() const
+	{
+		return ShaderBindingLayout;
+	}
+
 	void Bind(TArray<FRHIUniformBuffer*>& Bindings) const
 	{
 		Bindings.Reset();
@@ -121,6 +112,9 @@ private:
 	TArray<FUniformBufferStaticSlot, TInlineAllocator<InlineUniformBufferCount>> Slots;
 	TArray<FRHIUniformBuffer*, TInlineAllocator<InlineUniformBufferCount>> UniformBuffers;
 	int32 SlotCount = 0;
+
+	// Shader binding layout used during shader generation to validate runtime bindings and know where uniform buffers need to be bound
+	const FRHIShaderBindingLayout* ShaderBindingLayout = nullptr;
 };
 
 struct FTransferResourceFenceData
diff --git a/Engine/Source/Runtime/RHI/Public/RHIDefinitions.h b/Engine/Source/Runtime/RHI/Public/RHIDefinitions.h
--- a/Engine/Source/Runtime/RHI/Public/RHIDefinitions.h
+++ b/Engine/Source/Runtime/RHI/Public/RHIDefinitions.h
@@ -116,6 +116,15 @@ enum class ERHIBindlessSupport : uint8
 	NumBits = 2
 };
 
+enum class ERHIStaticShaderBindingLayoutSupport : uint8
+{
+	Unsupported,
+	RayTracingOnly,
+	AllShaderTypes,
+
+	NumBits = 2
+};
+
 enum EShaderFrequency : uint8
 {
 	SF_Vertex				= 0,
diff --git a/Engine/Source/Runtime/RHI/Public/RHIResources.h b/Engine/Source/Runtime/RHI/Public/RHIResources.h
--- a/Engine/Source/Runtime/RHI/Public/RHIResources.h
+++ b/Engine/Source/Runtime/RHI/Public/RHIResources.h
@@ -30,6 +30,8 @@ class FHazardPointerCollection;
 class FRHIComputeCommandList;
 class FRHICommandListImmediate;
 class FRHITextureReference;
+class FRHIShaderBindingLayout;
+
 struct FClearValueBinding;
 struct FRHIResourceInfo;
 struct FGenerateMipsStruct;
@@ -4414,6 +4416,10 @@ public:
 	// Depends on GRHISupportsRayTracingPSOAdditions support at runtime (base pipeline is simply ignored if it is unsupported).
 	FRayTracingPipelineStateRHIRef BasePipeline;
 
+	// Shader binding table layout used during shader compilation which needs to be the same for all shaders in the RTPSO and defines
+	// how uniform buffers needs to be bound at runtime (global(RayGen) vs local(miss/hit/callable) data)
+	const FRHIShaderBindingLayout* ShaderBindingLayout = nullptr;
+
 	const TArrayView<FRHIRayTracingShader*>& GetRayGenTable()   const { return RayGenTable; }
 	const TArrayView<FRHIRayTracingShader*>& GetMissTable()     const { return MissTable; }
 	const TArrayView<FRHIRayTracingShader*>& GetHitGroupTable() const { return HitGroupTable; }
diff --git a/Engine/Source/Runtime/RHI/Public/RHIShaderBindingLayout.h b/Engine/Source/Runtime/RHI/Public/RHIShaderBindingLayout.h
new file mode 100644
--- /dev/null
+++ b/Engine/Source/Runtime/RHI/Public/RHIShaderBindingLayout.h
@@ -0,0 +1,146 @@
+// Copyright Epic Games, Inc. All Rights Reserved.
+
+#pragma once
+
+#include "Containers/ArrayView.h"
+#include "Containers/StaticArray.h"
+#include "Misc/Optional.h"
+#include "RHIDefinitions.h"
+
+/*
+* FRHIUniformBufferShaderBindingLayout contains data on how a uniform buffer is used in a FRHIShaderBindingLayout
+*/
+struct FRHIUniformBufferShaderBindingLayout
+{
+	uint32 LayoutHash;
+	union
+	{
+		struct
+		{
+			uint32 RegisterSpace : 4;
+			uint32 CBVResourceIndex : 6;
+			uint32 BaseSRVResourceIndex : 8;
+			uint32 BaseUAVResourceIndex : 8;
+			uint32 BaseSamplerResourceIndex : 6;
+		};
+		uint32 Flags;
+	};
+
+	FRHIUniformBufferShaderBindingLayout() : LayoutHash(0), Flags(0) {}
+			
+	friend uint32 GetTypeHash(const FRHIUniformBufferShaderBindingLayout& Entry)
+	{
+		uint32 Hash = GetTypeHash(Entry.LayoutHash);
+		Hash = HashCombineFast(Hash, GetTypeHash(Entry.Flags));
+		return Hash;
+	}
+	
+	bool operator==(const FRHIUniformBufferShaderBindingLayout& Other) const
+	{
+		return LayoutHash == Other.LayoutHash &&
+			Flags == Other.Flags;			
+	}
+
+	friend inline FArchive& operator << (FArchive& Ar, FRHIUniformBufferShaderBindingLayout& F)
+	{
+		Ar << F.LayoutHash;
+		Ar << F.Flags;
+		return Ar;
+	}
+};
+
+enum class EShaderBindingLayoutFlags : uint8
+{
+	None = 0,
+	AllowMeshShaders = 1 << 0,
+	InputAssembler = 1 << 1,
+	BindlessResources = 1 << 2,
+	BindlessSamplers = 1 << 3,
+	RootConstants = 1 << 4,
+};
+ENUM_CLASS_FLAGS(EShaderBindingLayoutFlags)
+
+
+/*
+* FRHIShaderBindingLayout contains data which is used during shader generation to build the shareable shader resource tables between multiple shaders.
+* All shaders using the same FRHIShaderBindingLayout only have to bind the uniform buffers declared one at runtime. The shaders can have different PSOs
+* but they will define the resources of the uniform buffers at a specific resource index or SRT offset.
+* FRHIShaderBindingLayout is also used at runtime to know how/where the resources of the uniform buffers need to be bound.
+*/
+class FRHIShaderBindingLayout
+{
+public:
+
+	FRHIShaderBindingLayout() = default;
+	
+	FRHIShaderBindingLayout(EShaderBindingLayoutFlags InFlags, TConstArrayView<FRHIUniformBufferShaderBindingLayout> InUniformBufferEntries) : Flags(InFlags)
+	{
+		NumUniformBufferEntries = InUniformBufferEntries.Num();
+		check(NumUniformBufferEntries < MaxUniformBufferEntries);
+		for (uint32 Index = 0; Index < NumUniformBufferEntries; ++Index)
+		{
+			UniformBufferEntries[Index] = InUniformBufferEntries[Index];
+		}
+
+		ComputeHash();
+	}
+
+	uint32 GetHash() const { return Hash; }
+	EShaderBindingLayoutFlags GetFlags() const { return Flags; }
+	uint32 GetNumUniformBufferEntries() const { return NumUniformBufferEntries; }
+	const FRHIUniformBufferShaderBindingLayout& GetUniformBufferEntry(uint32 Index) const { check(Index < NumUniformBufferEntries); return UniformBufferEntries[Index]; }
+	   	
+	const FRHIUniformBufferShaderBindingLayout* FindEntry(uint32 LayoutHash) const
+	{
+		for (const FRHIUniformBufferShaderBindingLayout& Entry : UniformBufferEntries)
+		{
+			if (Entry.LayoutHash == LayoutHash)
+			{
+				return &Entry;
+			}
+		}
+
+		return nullptr;
+	}
+
+	friend uint32 GetTypeHash(const FRHIShaderBindingLayout& Desc)
+	{
+		return Desc.Hash;
+	}
+
+	bool operator==(const FRHIShaderBindingLayout& Other) const
+	{
+		return Hash == Other.Hash &&
+			Flags == Other.Flags &&
+			NumUniformBufferEntries == Other.NumUniformBufferEntries &&
+			UniformBufferEntries == Other.UniformBufferEntries;			
+	}
+
+	friend inline FArchive& operator << (FArchive& Ar, FRHIShaderBindingLayout& F)
+	{
+		Ar << F.Hash;
+		Ar << F.Flags;
+		Ar << F.NumUniformBufferEntries;
+		Ar << F.UniformBufferEntries;
+		return Ar;
+	}
+
+private:
+
+	void ComputeHash()
+	{
+		Hash = HashCombineFast(Hash, GetTypeHash(Flags));
+		Hash = HashCombineFast(Hash, GetTypeHash(NumUniformBufferEntries));
+		Hash = HashCombineFast(Hash, GetTypeHash(UniformBufferEntries));
+	}
+
+	enum
+	{
+		MaxUniformBufferEntries = 8,
+	};
+	
+	uint32 Hash = 0;
+	EShaderBindingLayoutFlags Flags = EShaderBindingLayoutFlags::None;
+	uint32 NumUniformBufferEntries = 0;
+	TStaticArray<FRHIUniformBufferShaderBindingLayout, MaxUniformBufferEntries> UniformBufferEntries = {};
+};
\ No newline at end of file
