- Add FShaderBindingLayout as RenderCore wrapper for FRHIShaderBindingLayout which can also cache static uniform buffer declerations used during shader code gen
- Add option for ShaderType to declare which FShaderBindingLayout needs to be used for shader gen and runtime binding
- Add helper function to generate FShaderBindingLayoutContainer from array of uniform buffers and EShaderResourceTableFlags (only supports bindless mode right now)
- Serialize the FRHIShaderBindingLayout in FShaderCompilerEnvironment so it can be used by the shader compiler backends
- FShaderBindingLayout can add required symbols to the FShaderCompilerInput which shouldn't be dead code eliminated (required for Shader resource table building)
- Serialize the hash of RHIShaderResourceTable in the platform specific shader binary data so it can be used during runtime RTPSO validation (all shaders in the RTPSO need the same RHIShaderResourceTable)
- Renamed pRootSignaure member to RootSignature (style guide)

diff --git a/Engine/Plugins/Compositing/OpenColorIO/Source/OpenColorIO/Public/OpenColorIOShaderType.h b/Engine/Plugins/Compositing/OpenColorIO/Source/OpenColorIO/Public/OpenColorIOShaderType.h
--- a/Engine/Plugins/Compositing/OpenColorIO/Source/OpenColorIO/Public/OpenColorIOShaderType.h
+++ b/Engine/Plugins/Compositing/OpenColorIO/Source/OpenColorIO/Public/OpenColorIOShaderType.h
@@ -74,6 +74,7 @@ public:
 		ShouldCompilePermutationType InShouldCompilePermutationRef,
 		ShouldPrecachePermutationType InShouldPrecachePermutationRef,
 		GetRayTracingPayloadTypeType InGetRayTracingPayloadTypeRef,
+		GetShaderBindingLayoutType InGetShaderBindingLayoutTypeRef,
 #if WITH_EDITOR
 		ModifyCompilationEnvironmentType InModifyCompilationEnvironmentRef,
 		ValidateCompiledResultType InValidateCompiledResultRef,
@@ -87,6 +88,7 @@ public:
 			InShouldCompilePermutationRef,
 			InShouldPrecachePermutationRef,
 			InGetRayTracingPayloadTypeRef,
+			InGetShaderBindingLayoutTypeRef,
 #if WITH_EDITOR
 			InModifyCompilationEnvironmentRef,
 			InValidateCompiledResultRef,
diff --git a/Engine/Plugins/FX/Niagara/Source/NiagaraShader/Public/NiagaraShaderType.h b/Engine/Plugins/FX/Niagara/Source/NiagaraShader/Public/NiagaraShaderType.h
--- a/Engine/Plugins/FX/Niagara/Source/NiagaraShader/Public/NiagaraShaderType.h
+++ b/Engine/Plugins/FX/Niagara/Source/NiagaraShader/Public/NiagaraShaderType.h
@@ -91,6 +91,7 @@ public:
 		ShouldCompilePermutationType InShouldCompilePermutationRef,
 		ShouldPrecachePermutationType InShouldPrecachePermutationRef,
 		GetRayTracingPayloadTypeType InGetRayTracingPayloadTypeRef,
+		GetShaderBindingLayoutType InGetShaderBindingLayoutTypeRef,
 #if WITH_EDITOR
 		ModifyCompilationEnvironmentType InModifyCompilationEnvironmentRef,
 		ValidateCompiledResultType InValidateCompiledResultRef,
@@ -104,6 +105,7 @@ public:
 			InShouldCompilePermutationRef,
 			InShouldPrecachePermutationRef,
 			InGetRayTracingPayloadTypeRef,
+			InGetShaderBindingLayoutTypeRef,
 #if WITH_EDITOR
 			InModifyCompilationEnvironmentRef,
 			InValidateCompiledResultRef,
diff --git a/Engine/Plugins/Runtime/ComputeFramework/Source/ComputeFramework/Private/ComputeFramework/ComputeKernelShaderType.h b/Engine/Plugins/Runtime/ComputeFramework/Source/ComputeFramework/Private/ComputeFramework/ComputeKernelShaderType.h
--- a/Engine/Plugins/Runtime/ComputeFramework/Source/ComputeFramework/Private/ComputeFramework/ComputeKernelShaderType.h
+++ b/Engine/Plugins/Runtime/ComputeFramework/Source/ComputeFramework/Private/ComputeFramework/ComputeKernelShaderType.h
@@ -76,6 +76,7 @@ public:
 		ShouldCompilePermutationType InShouldCompilePermutationRef,
 		ShouldPrecachePermutationType InShouldPrecachePermutationRef,
 		GetRayTracingPayloadTypeType InGetRayTracingPayloadTypeRef,
+		GetShaderBindingLayoutType InGetShaderBindingLayoutTypeRef,
 #if WITH_EDITOR
 		ModifyCompilationEnvironmentType InModifyCompilationEnvironmentRef,
 		ValidateCompiledResultType InValidateCompiledResultRef,
@@ -96,6 +97,7 @@ public:
 			InShouldCompilePermutationRef,
 			InShouldPrecachePermutationRef,
 			InGetRayTracingPayloadTypeRef,
+			InGetShaderBindingLayoutTypeRef,
 #if WITH_EDITOR
 			InModifyCompilationEnvironmentRef,
 			InValidateCompiledResultRef,
diff --git a/Engine/Shaders/Public/Platform.ush b/Engine/Shaders/Public/Platform.ush
--- a/Engine/Shaders/Public/Platform.ush
+++ b/Engine/Shaders/Public/Platform.ush
@@ -499,6 +499,7 @@
 	#define UB_DECL_RESOURCE(UBName, StructName, GlobalName)  UBName##.##StructName = UBName##_##GlobalName
 #else
 	#define UB_CB_DEFINITION_START(      UBName                    ) cbuffer UBName {
+	#define UB_STATIC_CB_DEFINITION_START(UBName, ResourceIndex, Space) cbuffer UBName : register(b##ResourceIndex, space##Space) {
 	#define UB_CB_MEMBER_NAME(           UBName,         MemberName) UBName##_##MemberName
 	#define UB_CB_PREFIXED_MEMBER_NAME(  UBName, Prefix, MemberName) Prefix##UBName##_##MemberName
 	#define UB_CB_DEFINITION_END(        UBName                    ) }
diff --git a/Engine/Shaders/Shared/HLSLReservedSpaces.h b/Engine/Shaders/Shared/HLSLReservedSpaces.h
--- a/Engine/Shaders/Shared/HLSLReservedSpaces.h
+++ b/Engine/Shaders/Shared/HLSLReservedSpaces.h
@@ -34,6 +34,9 @@
 // Register space for shader root constants (4xu32)
 #define UE_HLSL_SPACE_SHADER_ROOT_CONSTANTS 3
 
+// Register space for static shader bindings
+#define UE_HLSL_SPACE_STATIC_SHADER_BINDINGS 4
+
 // Register space for UE diagnotic debug buffer
 #define UE_HLSL_SPACE_DIAGNOSTIC 999
 
diff --git a/Engine/Source/Developer/Windows/ShaderFormatD3D/Private/D3DShaderCompiler.inl b/Engine/Source/Developer/Windows/ShaderFormatD3D/Private/D3DShaderCompiler.inl
--- a/Engine/Source/Developer/Windows/ShaderFormatD3D/Private/D3DShaderCompiler.inl
+++ b/Engine/Source/Developer/Windows/ShaderFormatD3D/Private/D3DShaderCompiler.inl
@@ -3,6 +3,7 @@
 #pragma once
 
 #include "D3D12RHI.h"
+#include "RHIShaderBindingLayout.h"
 
 struct FD3DShaderCompileData
 {
@@ -461,6 +462,19 @@ inline void GenerateFinalOutput(
 		Output.ShaderCode.AddOptionalData(FShaderCodeUniformBuffers::Key, UniformBufferNameBytes.GetData(), UniformBufferNameBytes.Num());
 		AddOptionalDataCallback(Output.ShaderCode);
 	}
+	
+	// Append the shader binding layout hash used for validation
+	{
+		uint32 ShaderBindingLayoutHash = Input.Environment.RHIShaderBindingLayout.GetHash();
+
+		TArray<uint8> WriterBytes;
+		FMemoryWriter Writer(WriterBytes);
+		Writer << ShaderBindingLayoutHash;
+		if (WriterBytes.Num() > 0)
+		{
+			Output.ShaderCode.AddOptionalData(FShaderCodeShaderResourceTableDataDesc::Key, WriterBytes.GetData(), WriterBytes.Num());
+		}
+	}
 
 	// Append information about optional hardware vendor extensions
 	if (CompileData.VendorExtensions.Num() > 0)
diff --git a/Engine/Source/Runtime/D3D12RHI/Private/D3D12RayTracing.cpp b/Engine/Source/Runtime/D3D12RHI/Private/D3D12RayTracing.cpp
--- a/Engine/Source/Runtime/D3D12RHI/Private/D3D12RayTracing.cpp
+++ b/Engine/Source/Runtime/D3D12RHI/Private/D3D12RayTracing.cpp
@@ -22,6 +22,7 @@
 #include "D3D12RayTracingDebug.h"
 #include "D3D12ExplicitDescriptorCache.h"
 #include "D3D12ResourceCollection.h"
+#include "RHIShaderBindingLayout.h"
 
 extern int32 GD3D12ExplicitViewDescriptorHeapSize;
 extern int32 GD3D12ExplicitViewDescriptorHeapOverflowReported;
@@ -1156,7 +1157,7 @@ public:
 		else
 		{
 			// All other shaders (hit groups, miss, callable) use custom root signatures.
-			LocalRootSignature = Shader->pRootSignature->GetRootSignature();
+			LocalRootSignature = Shader->RootSignature->GetRootSignature();
 		}
 
 		FKey CacheKey;
@@ -1970,6 +1971,8 @@ public:
 		TArrayView<FRHIRayTracingShader*> InitializerMissShaders = Initializer.GetMissTable();
 		TArrayView<FRHIRayTracingShader*> InitializerRayGenShaders = Initializer.GetRayGenTable();
 		TArrayView<FRHIRayTracingShader*> InitializerCallableShaders = Initializer.GetCallableTable();
+				
+		uint32 ShaderBindingLayoutHash = Initializer.ShaderBindingLayout ? Initializer.ShaderBindingLayout->GetHash() : 0;
 
 		const uint32 MaxTotalShaders = InitializerRayGenShaders.Num() + InitializerMissShaders.Num() + InitializerHitGroups.Num() + InitializerCallableShaders.Num();
 		checkf(MaxTotalShaders >= 1, TEXT("Ray tracing pipelines are expected to contain at least one shader"));
@@ -1998,11 +2001,14 @@ public:
 
 		// Helper function to acquire a D3D12_EXISTING_COLLECTION_DESC for a compiled shader via cache
 
-		auto AddShaderCollection = [Device, GlobalRootSignature = this->GlobalRootSignature, PipelineCache,
+		auto AddShaderCollection = [Device, ShaderBindingLayoutHash, GlobalRootSignature = this->GlobalRootSignature, PipelineCache,
 										&UniqueShaderHashes = this->PipelineShaderHashes, &UniqueShaderCollections, &Initializer, &NumCacheHits, &CompileTime,
 										&CompileCompletionList]
 			(FD3D12RayTracingShader* Shader, FD3D12RayTracingPipelineCache::ECollectionType CollectionType)
 		{
+			// verify that that the same shader binding layout is used for all shaders in the RTPSO
+			checkf(Shader->ShaderBindingLayoutHash == ShaderBindingLayoutHash, TEXT("Raytracing shader with with entry point %s doesn't match the RTPSO ShaderBindingLayout"), *Shader->EntryPoint);
+
 			bool bIsAlreadyInSet = false;
 			const uint64 ShaderHash = GetShaderHash64(Shader);
 			UniqueShaderHashes.Add(ShaderHash, &bIsAlreadyInSet);
@@ -2047,7 +2053,7 @@ public:
 		for (FRHIRayTracingShader* ShaderRHI : InitializerRayGenShaders)
 		{
 			FD3D12RayTracingShader* Shader = FD3D12DynamicRHI::ResourceCast(ShaderRHI);
-			checkf(Shader->pRootSignature->GetRootSignature() == GlobalRootSignature, TEXT("All raygen and miss shaders must share the same root signature"));
+			checkf(Shader->RootSignature->GetRootSignature() == GlobalRootSignature, TEXT("All raygen and miss shaders must share the same root signature"));
 			checkf(!Shader->UsesGlobalUniformBuffer(), TEXT("Global uniform buffers are not implemented for ray generation shaders"));
 
 			FD3D12RayTracingPipelineCache::FEntry* ShaderCacheEntry = AddShaderCollection(Shader, FD3D12RayTracingPipelineCache::ECollectionType::RayGen);
@@ -2076,7 +2082,7 @@ public:
 
 			const uint32 ShaderViewDescriptors = Shader->ResourceCounts.NumSRVs + Shader->ResourceCounts.NumUAVs;
 			MaxHitGroupViewDescriptors = FMath::Max(MaxHitGroupViewDescriptors, ShaderViewDescriptors);
-			MaxLocalRootSignatureSize = FMath::Max(MaxLocalRootSignatureSize, Shader->pRootSignature->GetTotalRootSignatureSizeInBytes());
+			MaxLocalRootSignatureSize = FMath::Max(MaxLocalRootSignatureSize, Shader->RootSignature->GetTotalRootSignatureSizeInBytes());
 
 			FD3D12RayTracingPipelineCache::FEntry* ShaderCacheEntry = AddShaderCollection(Shader, FD3D12RayTracingPipelineCache::ECollectionType::Miss);
 
@@ -2098,7 +2104,7 @@ public:
 
 			const uint32 ShaderViewDescriptors = Shader->ResourceCounts.NumSRVs + Shader->ResourceCounts.NumUAVs;
 			MaxHitGroupViewDescriptors = FMath::Max(MaxHitGroupViewDescriptors, ShaderViewDescriptors);
-			MaxLocalRootSignatureSize = FMath::Max(MaxLocalRootSignatureSize, Shader->pRootSignature->GetTotalRootSignatureSizeInBytes());
+			MaxLocalRootSignatureSize = FMath::Max(MaxLocalRootSignatureSize, Shader->RootSignature->GetTotalRootSignatureSizeInBytes());
 
 			FD3D12RayTracingPipelineCache::FEntry* ShaderCacheEntry = AddShaderCollection(Shader, FD3D12RayTracingPipelineCache::ECollectionType::HitGroup);
 
@@ -2121,7 +2127,7 @@ public:
 
 			const uint32 ShaderViewDescriptors = Shader->ResourceCounts.NumSRVs + Shader->ResourceCounts.NumUAVs;
 			MaxHitGroupViewDescriptors = FMath::Max(MaxHitGroupViewDescriptors, ShaderViewDescriptors);
-			MaxLocalRootSignatureSize = FMath::Max(MaxLocalRootSignatureSize, Shader->pRootSignature->GetTotalRootSignatureSizeInBytes());
+			MaxLocalRootSignatureSize = FMath::Max(MaxLocalRootSignatureSize, Shader->RootSignature->GetTotalRootSignatureSizeInBytes());
 
 			FD3D12RayTracingPipelineCache::FEntry* ShaderCacheEntry = AddShaderCollection(Shader, FD3D12RayTracingPipelineCache::ECollectionType::Callable);
 
@@ -4279,7 +4285,7 @@ static bool SetRayTracingShaderResources(
 	uint32 InLooseParameterDataSize, const void* InLooseParameterData,
 	ResourceBinderType& Binder)
 {
-	const FD3D12RootSignature* RootSignature = Shader->pRootSignature;
+	const FD3D12RootSignature* RootSignature = Shader->RootSignature;
 
 	struct FBindings
 	{
@@ -4905,7 +4911,7 @@ static void SetRayTracingHitGroup(
 		}
 	}
 
-	FD3D12RayTracingLocalResourceBinder ResourceBinder(*Device, *ShaderTable, *(Shader->pRootSignature), RecordIndex, WorkerIndex, ERayTracingBindingType::HitGroup);
+	FD3D12RayTracingLocalResourceBinder ResourceBinder(*Device, *ShaderTable, *(Shader->RootSignature), RecordIndex, WorkerIndex, ERayTracingBindingType::HitGroup);
 	const bool bResourcesBound = SetRayTracingShaderResources(Shader,
 		0, nullptr, // Textures
 		0, nullptr, // SRVs
@@ -4950,7 +4956,7 @@ static void SetRayTracingCallableShader(
 	{
 		const FD3D12RayTracingShader* Shader = Pipeline->CallableShaders.Shaders[ShaderIndexInPipeline];
 
-		FD3D12RayTracingLocalResourceBinder ResourceBinder(*Device, *ShaderTable, *(Shader->pRootSignature), RecordIndex, WorkerIndex, ERayTracingBindingType::CallableShader);
+		FD3D12RayTracingLocalResourceBinder ResourceBinder(*Device, *ShaderTable, *(Shader->RootSignature), RecordIndex, WorkerIndex, ERayTracingBindingType::CallableShader);
 		const bool bResourcesBound = SetRayTracingShaderResources(Shader,
 			0, nullptr, // Textures
 			0, nullptr, // SRVs
@@ -4989,7 +4995,7 @@ static void SetRayTracingMissShader(
 
 	const FD3D12RayTracingShader* Shader = Pipeline->MissShaders.Shaders[ShaderIndexInPipeline];
 
-	FD3D12RayTracingLocalResourceBinder ResourceBinder(*Device, *ShaderTable, *(Shader->pRootSignature), RecordIndex, WorkerIndex, ERayTracingBindingType::MissShader);
+	FD3D12RayTracingLocalResourceBinder ResourceBinder(*Device, *ShaderTable, *(Shader->RootSignature), RecordIndex, WorkerIndex, ERayTracingBindingType::MissShader);
 	const bool bResourcesBound = SetRayTracingShaderResources(Shader,
 		0, nullptr, // Textures
 		0, nullptr, // SRVs
diff --git a/Engine/Source/Runtime/D3D12RHI/Private/D3D12Shader.h b/Engine/Source/Runtime/D3D12RHI/Private/D3D12Shader.h
--- a/Engine/Source/Runtime/D3D12RHI/Private/D3D12Shader.h
+++ b/Engine/Source/Runtime/D3D12RHI/Private/D3D12Shader.h
@@ -71,6 +71,9 @@ struct FD3D12ShaderData
 	TArray<uint8> Code;
 
 	FShaderCodePackedResourceCounts ResourceCounts{};
+		
+	// Hash of the static shader binding layout used during shader gen
+	uint32 ShaderBindingLayoutHash = 0;
 
 #if D3D12RHI_NEEDS_VENDOR_EXTENSIONS
 	TArray<FShaderCodeVendorExtension> VendorExtensions;
@@ -172,7 +175,7 @@ class FD3D12RayTracingShader : public FRHIRayTracingShader, public FD3D12ShaderD
 public:
 	explicit FD3D12RayTracingShader(EShaderFrequency InFrequency) : FRHIRayTracingShader(InFrequency) {}
 
-	const FD3D12RootSignature* pRootSignature = nullptr;
+	const FD3D12RootSignature* RootSignature = nullptr;
 
 	/** The shader's DXIL entrypoint & base export name for DXR (required for RTPSO creation) */
 	FString EntryPoint; // Primary entry point for all ray tracing shaders. Assumed to be closest hit shader for SF_RayHitGroup.
diff --git a/Engine/Source/Runtime/D3D12RHI/Private/D3D12Shaders.cpp b/Engine/Source/Runtime/D3D12RHI/Private/D3D12Shaders.cpp
--- a/Engine/Source/Runtime/D3D12RHI/Private/D3D12Shaders.cpp
+++ b/Engine/Source/Runtime/D3D12RHI/Private/D3D12Shaders.cpp
@@ -52,6 +52,15 @@ static inline bool ReadShaderOptionalData(FShaderCodeReader& InShaderCode, TShad
 	}
 #endif
 
+	int32 ShaderBindingLayoutSize = 0;
+	auto* ShaderBindingLayoutData = InShaderCode.FindOptionalDataAndSize(FShaderCodeShaderResourceTableDataDesc::Key, ShaderBindingLayoutSize);
+	if (ShaderBindingLayoutData && ShaderBindingLayoutSize > 0)
+	{
+		check(ShaderBindingLayoutSize == sizeof(OutShader.ShaderBindingLayoutHash));
+		FBufferReader Ar((void*)ShaderBindingLayoutData, ShaderBindingLayoutSize, false);
+		Ar << OutShader.ShaderBindingLayoutHash;
+	}
+
 	UE::RHICore::SetupShaderCodeValidationData(&OutShader, InShaderCode);
 	UE::RHICore::SetupShaderDiagnosticData(&OutShader, InShaderCode);
 
@@ -240,7 +249,7 @@ FRayTracingShaderRHIRef FD3D12DynamicRHI::RHICreateRayTracingShader(TArrayView<c
 	FD3D12RayTracingShader* Shader = InitStandardShaderWithCustomSerialization(new FD3D12RayTracingShader(ShaderFrequency), Code, CustomSerialization);
 	if (Shader)
 	{
-		Shader->pRootSignature = GetAdapter().GetRootSignature(Shader);
+		Shader->RootSignature = GetAdapter().GetRootSignature(Shader);
 	}
 
 	return Shader;
diff --git a/Engine/Source/Runtime/Engine/Private/ShaderCompiler/ShaderCompiler.cpp b/Engine/Source/Runtime/Engine/Private/ShaderCompiler/ShaderCompiler.cpp
--- a/Engine/Source/Runtime/Engine/Private/ShaderCompiler/ShaderCompiler.cpp
+++ b/Engine/Source/Runtime/Engine/Private/ShaderCompiler/ShaderCompiler.cpp
@@ -8557,6 +8557,12 @@ void GlobalBeginCompileShader(
 	SET_SHADER_DEFINE(Input.Environment, UE_LWC_RENDER_TILE_SIZE_FMOD_PI, (float)FMath::Fmod(TileSize, UE_DOUBLE_PI));
 	SET_SHADER_DEFINE(Input.Environment, UE_LWC_RENDER_TILE_SIZE_FMOD_2PI, (float)FMath::Fmod(TileSize, 2.0 * UE_DOUBLE_PI));
 
+	// Add required symbols from the shader binding layout if set
+	if (Input.Environment.ShaderBindingLayout)
+	{
+		Input.Environment.ShaderBindingLayout->AddRequiredSymbols(Input.RequiredSymbols);
+	}
+
 	// Allow the target shader format to modify the shader input before we add it as a job
 	const IShaderFormat* Format = GetTargetPlatformManagerRef().FindShaderFormat(ShaderFormatName);
 	checkf(Format, TEXT("Shader format %s cannot be found"), *ShaderFormatName.ToString());
diff --git a/Engine/Source/Runtime/Engine/Public/MaterialShaderType.h b/Engine/Source/Runtime/Engine/Public/MaterialShaderType.h
--- a/Engine/Source/Runtime/Engine/Public/MaterialShaderType.h
+++ b/Engine/Source/Runtime/Engine/Public/MaterialShaderType.h
@@ -142,6 +142,7 @@ public:
 		ShouldCompilePermutationType InShouldCompilePermutationRef,
 		ShouldPrecachePermutationType InShouldPrecachePermutationRef,
 		GetRayTracingPayloadTypeType InGetRayTracingPayloadTypeRef,
+		GetShaderBindingLayoutType InGetShaderBindingLayoutTypeRef,
 #if WITH_EDITOR
 		ModifyCompilationEnvironmentType InModifyCompilationEnvironmentRef,
 		ValidateCompiledResultType InValidateCompiledResultRef,
@@ -155,6 +156,7 @@ public:
 			InShouldCompilePermutationRef,
 			InShouldPrecachePermutationRef,
 			InGetRayTracingPayloadTypeRef,
+			InGetShaderBindingLayoutTypeRef,
 #if WITH_EDITOR
 			InModifyCompilationEnvironmentRef,
 			InValidateCompiledResultRef,
diff --git a/Engine/Source/Runtime/Engine/Public/MeshMaterialShaderType.h b/Engine/Source/Runtime/Engine/Public/MeshMaterialShaderType.h
--- a/Engine/Source/Runtime/Engine/Public/MeshMaterialShaderType.h
+++ b/Engine/Source/Runtime/Engine/Public/MeshMaterialShaderType.h
@@ -55,6 +55,7 @@ public:
 		ShouldCompilePermutationType InShouldCompilePermutationRef,
 		ShouldPrecachePermutationType InShouldPrecachePermutationRef,
 		GetRayTracingPayloadTypeType InGetRayTracingPayloadTypeRef,
+		GetShaderBindingLayoutType InGetShaderBindingLayoutTypeRef,
 #if WITH_EDITOR
 		ModifyCompilationEnvironmentType InModifyCompilationEnvironmentRef,
 		ValidateCompiledResultType InValidateCompiledResultRef,
@@ -68,6 +69,7 @@ public:
 			InShouldCompilePermutationRef,
 			InShouldPrecachePermutationRef,
 			InGetRayTracingPayloadTypeRef,
+			InGetShaderBindingLayoutTypeRef,
 #if WITH_EDITOR
 			InModifyCompilationEnvironmentRef,
 			InValidateCompiledResultRef,
diff --git a/Engine/Source/Runtime/RenderCore/Private/Shader.cpp b/Engine/Source/Runtime/RenderCore/Private/Shader.cpp
--- a/Engine/Source/Runtime/RenderCore/Private/Shader.cpp
+++ b/Engine/Source/Runtime/RenderCore/Private/Shader.cpp
@@ -247,6 +247,7 @@ FShaderType::FShaderType(
 	ShouldCompilePermutationType InShouldCompilePermutationRef,
 	ShouldPrecachePermutationType InShouldPrecachePermutationRef,
 	GetRayTracingPayloadTypeType InGetRayTracingPayloadTypeRef,
+	GetShaderBindingLayoutType InGetShaderBindingLayoutTypeRef,
 #if WITH_EDITOR
 	ModifyCompilationEnvironmentType InModifyCompilationEnvironmentRef,
 	ValidateCompiledResultType InValidateCompiledResultRef,
@@ -270,6 +271,7 @@ FShaderType::FShaderType(
 	ShouldCompilePermutationRef(InShouldCompilePermutationRef),
 	ShouldPrecachePermutationRef(InShouldPrecachePermutationRef),
 	GetRayTracingPayloadTypeRef(InGetRayTracingPayloadTypeRef),
+	GetShaderBindingLayoutTypeRef(InGetShaderBindingLayoutTypeRef),
 #if WITH_EDITOR
 	ModifyCompilationEnvironmentRef(InModifyCompilationEnvironmentRef),
 	ValidateCompiledResultRef(InValidateCompiledResultRef),
@@ -435,11 +437,24 @@ EShaderPermutationPrecacheRequest FShaderType::ShouldPrecachePermutation(const F
 	return ShouldCompileShaderFrequency((EShaderFrequency)Frequency, Parameters.Platform) ? (*ShouldPrecachePermutationRef)(Parameters) : EShaderPermutationPrecacheRequest::NotUsed;
 }
 
+const FShaderBindingLayout* FShaderType::GetShaderBindingLayout(const FShaderPermutationParameters& Parameters) const
+{
+	return (*GetShaderBindingLayoutTypeRef)(Parameters);
+}
+
+
 #if WITH_EDITOR
 
 void FShaderType::ModifyCompilationEnvironment(const FShaderPermutationParameters& Parameters, FShaderCompilerEnvironment& OutEnvironment) const
 {
 	(*ModifyCompilationEnvironmentRef)(Parameters, OutEnvironment);
+	
+	OutEnvironment.ShaderBindingLayout = GetShaderBindingLayout(Parameters);
+	if (OutEnvironment.ShaderBindingLayout)
+	{
+		// Store copy of RHI version of the shader binding layout in the environment so it can be serialized for the shader compiler workers
+		OutEnvironment.RHIShaderBindingLayout = OutEnvironment.ShaderBindingLayout->RHILayout;
+	}
 
 	if (Frequency == SF_RayHitGroup)
 	{
@@ -2611,6 +2626,14 @@ void ShaderMapAppendKey(EShaderPlatform Platform, FShaderKeyGenerator& KeyGen)
 		}
 	}
 
+	const ERHIStaticShaderBindingLayoutSupport StaticShaderBindingLayoutSupport = FDataDrivenShaderPlatformInfo::GetStaticShaderBindingLayoutSupport(Platform);
+	if (StaticShaderBindingLayoutSupport != ERHIStaticShaderBindingLayoutSupport::Unsupported)
+	{
+		KeyGen.AppendSeparator();
+		KeyGen.Append(StaticShaderBindingLayoutSupport == ERHIStaticShaderBindingLayoutSupport::RayTracingOnly ?
+		              TEXT("SSBL-RT") : TEXT("SSBL"));
+	}
+
 	if (ShouldCompileRayTracingShadersForProject(Platform))
 	{
 		static FShaderPlatformCachedIniValue<int32> CVarCompileRayTracingMaterialCHS(TEXT("r.RayTracing.CompileMaterialCHS"));
diff --git a/Engine/Source/Runtime/RenderCore/Private/ShaderCore.cpp b/Engine/Source/Runtime/RenderCore/Private/ShaderCore.cpp
--- a/Engine/Source/Runtime/RenderCore/Private/ShaderCore.cpp
+++ b/Engine/Source/Runtime/RenderCore/Private/ShaderCore.cpp
@@ -772,6 +772,40 @@ uint32 FShaderParameterMap::CountParametersOfType(EShaderParameterType InType) c
 	return Result;
 }
 
+#if WITH_EDITOR
+
+FThreadSafeSharedAnsiStringPtr FShaderBindingLayout::GetUniformBufferDeclarationAnsiPtr(const FShaderParametersMetadata* ShaderParametersMetadata) const
+{
+	// Does the binding layout contain a uniform buffer declaration for this shader paramaters
+	uint32 LayoutHash = ShaderParametersMetadata->GetLayout().GetHash();
+	const FRHIUniformBufferShaderBindingLayout* Entry = RHILayout.FindEntry(LayoutHash);
+	if (Entry)
+	{
+		uint8 UniformBufferIndex = Entry->CBVResourceIndex;
+		return UniformBufferDeclarationsAnsi[UniformBufferIndex];
+	}
+
+	// use the shared declaration from the metadata object itself (no fixed register or space defined)
+	return ShaderParametersMetadata->GetUniformBufferDeclarationAnsiPtr();
+}
+
+void FShaderBindingLayout::SetUniformBufferDeclarationAnsiPtr(const FShaderParametersMetadata* ShaderParametersMetadata, FThreadSafeSharedAnsiStringPtr UniformBufferDeclarationAnsi)
+{
+	uint32 LayoutHash = ShaderParametersMetadata->GetLayout().GetHash();
+	const FRHIUniformBufferShaderBindingLayout* Entry = RHILayout.FindEntry(LayoutHash);
+	check(Entry);
+	UniformBufferDeclarationsAnsi.SetNum(FMath::Max(UniformBufferDeclarationsAnsi.Num(), int32(Entry->CBVResourceIndex + 1)));
+	UniformBufferDeclarationsAnsi[Entry->CBVResourceIndex] = UniformBufferDeclarationAnsi;
+}
+
+void FShaderBindingLayout::AddRequiredSymbols(TArray<FString>& RequiredSymbols) const
+{
+	// assume only bindless for now so only need to add the CBuffer declares as required symbols
+	check(EnumHasAllFlags(RHILayout.GetFlags(), EShaderBindingLayoutFlags::BindlessResources | EShaderBindingLayoutFlags::BindlessSamplers));
+}
+
+#endif // WITH_EDITOR
+
 void FShaderResourceTableMap::Append(const FShaderResourceTableMap& Other)
 {
 	// Get the set of uniform buffers used by the target resource table map
@@ -1177,6 +1211,7 @@ void FShaderCompilerEnvironment::SerializeEverythingButFiles(FArchive& Ar)
 	Ar << RenderTargetOutputFormatsMap;
 	Ar << ResourceTableMap.Resources;
 	Ar << UniformBufferMap;
+	Ar << RHIShaderBindingLayout;
 	Ar << FullPrecisionInPS;
 	if (Ar.IsLoading())
 	{
@@ -1192,6 +1227,7 @@ void FShaderCompilerEnvironment::SerializeCompilationDependencies(FArchive& Ar)
 	Ar << RenderTargetOutputFormatsMap;
 	Ar << ResourceTableMap.Resources;
 	Ar << UniformBufferMap;
+	Ar << RHIShaderBindingLayout;
 	Ar << FullPrecisionInPS;
 	if (Ar.IsLoading())
 	{
diff --git a/Engine/Source/Runtime/RenderCore/Private/ShaderParameterMetadata.cpp b/Engine/Source/Runtime/RenderCore/Private/ShaderParameterMetadata.cpp
--- a/Engine/Source/Runtime/RenderCore/Private/ShaderParameterMetadata.cpp
+++ b/Engine/Source/Runtime/RenderCore/Private/ShaderParameterMetadata.cpp
@@ -11,12 +11,14 @@
 #include "Misc/CommandLine.h"
 #include "RenderCore.h"
 #include "RHIUniformBufferLayoutInitializer.h"
+#include "RHIShaderBindingLayout.h"
 #include "Serialization/MemoryHasher.h"
 #include "Serialization/ShaderKeyGenerator.h"
 #include "ShaderCore.h"
 #include "ShaderCompilerCore.h"
 #include "ShaderParameters.h"
 #include "ShaderParameterMacros.h"
+#include "HLSLReservedSpaces.h"
 
 bool SupportShaderPrecisionModifier(EShaderPlatform Platform)
 {
@@ -147,6 +149,83 @@ FShaderParametersMetadata* FindUniformBufferStructByLayoutHash(uint32 Hash)
 	return GetLayoutHashStructMap().FindRef(Hash);
 }
 
+void BuildShaderBindingLayout(TConstArrayView<FShaderParametersMetadata*> UniformBuffers, EShaderBindingLayoutFlags BaseShaderBindingLayoutFlags, FShaderBindingLayoutContainer& OutShaderBindingLayoutContainer)
+{
+	// only build the bindless version for now
+	FShaderBindingLayoutContainer::EBindingType BindingType = FShaderBindingLayoutContainer::EBindingType::Bindless;
+	
+	FShaderBindingLayout ShaderBindingLayout;
+
+	EShaderBindingLayoutFlags ShaderBindingLayoutFlags = BaseShaderBindingLayoutFlags;
+	if (BindingType == FShaderBindingLayoutContainer::EBindingType::Bindless)
+	{
+		ShaderBindingLayoutFlags |= EShaderBindingLayoutFlags::BindlessResources | EShaderBindingLayoutFlags::BindlessSamplers;
+	}
+
+	bool bBindlessResources = EnumHasAllFlags(ShaderBindingLayoutFlags, EShaderBindingLayoutFlags::BindlessResources);
+	bool bBindlessSamplers = EnumHasAllFlags(ShaderBindingLayoutFlags, EShaderBindingLayoutFlags::BindlessSamplers);
+
+	TArray<FRHIUniformBufferShaderBindingLayout> UniformBufferEntries;
+	UniformBufferEntries.Reserve(UniformBuffers.Num());
+
+	int32 CurrentCBVResourceIndex = 0;
+	int32 CurrentBaseSRVResourceIndex = bBindlessResources ? -1 : 0;
+	int32 CurrentBaseUAVResourceIndex = bBindlessResources ? -1 : 0;
+	int32 CurrentBaseSamplerResourceIndex = bBindlessSamplers ? -1 : 0;
+	for (FShaderParametersMetadata* ShaderParametersMetaData : UniformBuffers)
+	{
+		const FRHIUniformBufferLayout& UniformBufferLayout = ShaderParametersMetaData->GetLayout();
+		check(UniformBufferLayout.ConstantBufferSize > 0);
+
+		FRHIUniformBufferShaderBindingLayout& NewEntry = UniformBufferEntries.AddDefaulted_GetRef();
+		NewEntry.LayoutHash = UniformBufferLayout.GetHash();
+		NewEntry.RegisterSpace = UE_HLSL_SPACE_STATIC_SHADER_BINDINGS;
+		NewEntry.CBVResourceIndex = CurrentCBVResourceIndex;
+		NewEntry.BaseSRVResourceIndex = CurrentBaseSRVResourceIndex;
+		NewEntry.BaseUAVResourceIndex = CurrentBaseUAVResourceIndex;
+		NewEntry.BaseSamplerResourceIndex = CurrentBaseSamplerResourceIndex;
+
+		// Increment counters		
+		CurrentCBVResourceIndex++;
+
+		// Don't need to increment resources if using bindless because all data is in the CBV
+		if (!bBindlessSamplers)
+		{
+			// TODO use layout to find out all the samplers
+			check(false);
+		}
+		if (!bBindlessResources)
+		{
+			// TODO use layout to find out all the resources and uavs
+			check(false);
+		}
+	}
+
+	ShaderBindingLayout.RHILayout = FRHIShaderBindingLayout(ShaderBindingLayoutFlags, UniformBufferEntries);
+
+#if WITH_EDITOR
+	// Also create the uniform buffer decleration shared ansi strings used during shader gen
+	for (uint8 UniformBufferIndex = 0; UniformBufferIndex < UniformBuffers.Num(); ++UniformBufferIndex)	
+	{
+		FShaderParametersMetadata* ShaderParametersMetaData = UniformBuffers[UniformBufferIndex];
+		const FRHIUniformBufferShaderBindingLayout* UniformBufferSBLayout = ShaderBindingLayout.RHILayout.FindEntry(ShaderParametersMetaData->GetLayout().GetHash());
+		check(UniformBufferSBLayout);
+
+		FString NewDeclaration = UE::ShaderParameters::CreateUniformBufferShaderDeclaration(ShaderParametersMetaData->GetShaderVariableName(), *ShaderParametersMetaData, UniformBufferSBLayout);
+		check(!NewDeclaration.IsEmpty());
+			
+		// Cache preprocessor friendly copy of uniform buffer declaration
+		TArray<ANSICHAR>* NewDeclarationAnsi = new TArray<ANSICHAR>;
+		ShaderConvertAndStripComments(NewDeclaration, *NewDeclarationAnsi);
+		FThreadSafeSharedAnsiStringPtr UniformBufferDeclarationAnsi = MakeShareable(NewDeclarationAnsi);
+
+		ShaderBindingLayout.SetUniformBufferDeclarationAnsiPtr(ShaderParametersMetaData, UniformBufferDeclarationAnsi);
+	}
+#endif //WITH_EDITOR
+
+	OutShaderBindingLayoutContainer.SetLayout(BindingType, ShaderBindingLayout);
+}
+
 static TArray<const FShaderParametersMetadataRegistration*>* GShaderParametersMetadataRegistrationInstances = nullptr;
 TArray<const FShaderParametersMetadataRegistration*>& FShaderParametersMetadataRegistration::GetInstances()
 {
@@ -900,7 +979,7 @@ void FShaderParametersMetadata::InitializeUniformBufferDeclaration()
 {
 	if (UseCase != EUseCase::ShaderParameterStruct)
 	{
-		FString* NewDeclaration = new FString(UE::ShaderParameters::CreateUniformBufferShaderDeclaration(ShaderVariableName, *this));
+		FString* NewDeclaration = new FString(UE::ShaderParameters::CreateUniformBufferShaderDeclaration(ShaderVariableName, *this, nullptr));
 		check(!NewDeclaration->IsEmpty());
 
 		UniformBufferDeclaration = MakeShareable(NewDeclaration);
diff --git a/Engine/Source/Runtime/RenderCore/Private/ShaderParameters.cpp b/Engine/Source/Runtime/RenderCore/Private/ShaderParameters.cpp
--- a/Engine/Source/Runtime/RenderCore/Private/ShaderParameters.cpp
+++ b/Engine/Source/Runtime/RenderCore/Private/ShaderParameters.cpp
@@ -81,7 +81,7 @@ void FShaderUniformBufferParameter::ModifyCompilationEnvironment(const TCHAR* Pa
 	// if the name matches the struct's name, use the struct's cached version; otherwise, generate it now with the correct variable name.
 	if (FCString::Strcmp(ParameterName, Struct.GetShaderVariableName()) != 0)
 	{
-		const FString Declaration = UE::ShaderParameters::CreateUniformBufferShaderDeclaration(ParameterName, Struct);
+		const FString Declaration = UE::ShaderParameters::CreateUniformBufferShaderDeclaration(ParameterName, Struct, nullptr);
 		OutEnvironment.IncludeVirtualPathToContentsMap.Add(IncludeName, Declaration);
 	}
 	else
@@ -363,7 +363,7 @@ static void CreateHLSLUniformBufferStructMembersDeclaration(
 }
 
 /** Creates a HLSL declaration of a uniform buffer with the given structure. */
-static FString CreateHLSLUniformBufferDeclaration(const TCHAR* UniformBufferName, const FShaderParametersMetadata& UniformBufferStruct)
+static FString CreateHLSLUniformBufferDeclaration(const TCHAR* UniformBufferName, const FShaderParametersMetadata& UniformBufferStruct, const FRHIUniformBufferShaderBindingLayout* UniformBufferSBLayout)
 {
 	// If the uniform buffer has no members, we don't want to write out anything.  Shader compilers throw errors when faced with empty cbuffers and structs.
 	if (UniformBufferStruct.GetMembers().Num() > 0)
@@ -372,6 +372,8 @@ static FString CreateHLSLUniformBufferDeclaration(const TCHAR* UniformBufferName
 		uint32 HLSLBaseOffset = 0;
 		CreateHLSLUniformBufferStructMembersDeclaration(UniformBufferStruct, UniformBufferName, TEXT(""), TEXT(""), 0, Decl, HLSLBaseOffset);
 
+		// TODO: use UniformBufferSBLayout to generate CB_DEFINITION with fixed register & space
+
 		return FString::Printf(
 			TEXT("#pragma once\n")
 			TEXT("UB_CB_DEFINITION_START(%s)\n")
@@ -394,9 +396,9 @@ static FString CreateHLSLUniformBufferDeclaration(const TCHAR* UniformBufferName
 	return FString(TEXT("\n"));
 }
 
-FString UE::ShaderParameters::CreateUniformBufferShaderDeclaration(const TCHAR* UniformBufferName, const FShaderParametersMetadata& UniformBufferStruct)
+FString UE::ShaderParameters::CreateUniformBufferShaderDeclaration(const TCHAR* UniformBufferName, const FShaderParametersMetadata& UniformBufferStruct, const FRHIUniformBufferShaderBindingLayout* UniformBufferSBLayout)
 {
-	return CreateHLSLUniformBufferDeclaration(UniformBufferName, UniformBufferStruct);
+	return CreateHLSLUniformBufferDeclaration(UniformBufferName, UniformBufferStruct, UniformBufferSBLayout);
 }
 
 void UE::ShaderParameters::AddUniformBufferIncludesToEnvironment(FShaderCompilerEnvironment& OutEnvironment, const TSet<const FShaderParametersMetadata*>& InUniformBuffers)
@@ -405,22 +407,32 @@ void UE::ShaderParameters::AddUniformBufferIncludesToEnvironment(FShaderCompiler
 
 	FString UniformBufferIncludes;
 
-	for (const FShaderParametersMetadata* Metadata : InUniformBuffers)
+	for (const FShaderParametersMetadata* Metadata : InUniformBuffers)	
 	{
 		FStringView UniformBufferNameView(Metadata->GetShaderVariableName());
 		uint32 UniformBufferNameHash = GetTypeHash(UniformBufferNameView);
 		if (!OutEnvironment.UniformBufferMap.FindByHash(UniformBufferNameHash, UniformBufferNameView))
 		{
-				const FThreadSafeSharedAnsiStringPtr UniformBufferDeclaration = Metadata->GetUniformBufferDeclarationAnsiPtr();
-
-				check(UniformBufferDeclaration.Get() != NULL);
-				check(!UniformBufferDeclaration.Get()->IsEmpty());
-
-				UniformBufferIncludes += Metadata->GetUniformBufferInclude();
+			FThreadSafeSharedAnsiStringPtr UniformBufferDeclaration;
 
-				OutEnvironment.IncludeVirtualPathToSharedContentsMap.AddByHash(Metadata->GetUniformBufferPathHash(), Metadata->GetUniformBufferPath(), UniformBufferDeclaration);
-
-				Metadata->AddResourceTableEntries(OutEnvironment.ResourceTableMap, OutEnvironment.UniformBufferMap);
+			// Use the shader binding layout to retrieve the uniform buffer declaration if available
+			if (OutEnvironment.ShaderBindingLayout)
+			{
+				UniformBufferDeclaration = OutEnvironment.ShaderBindingLayout->GetUniformBufferDeclarationAnsiPtr(Metadata);
+			}
+			else
+			{
+				UniformBufferDeclaration = Metadata->GetUniformBufferDeclarationAnsiPtr();
+			}
+			
+			check(UniformBufferDeclaration.Get() != NULL);
+			check(!UniformBufferDeclaration.Get()->IsEmpty());
+			
+			UniformBufferIncludes += Metadata->GetUniformBufferInclude();
+			
+			OutEnvironment.IncludeVirtualPathToSharedContentsMap.AddByHash(Metadata->GetUniformBufferPathHash(), Metadata->GetUniformBufferPath(), UniformBufferDeclaration);
+			
+			Metadata->AddResourceTableEntries(OutEnvironment.ResourceTableMap, OutEnvironment.UniformBufferMap);
 		}
 	}
 
diff --git a/Engine/Source/Runtime/RenderCore/Public/GlobalShader.h b/Engine/Source/Runtime/RenderCore/Public/GlobalShader.h
--- a/Engine/Source/Runtime/RenderCore/Public/GlobalShader.h
+++ b/Engine/Source/Runtime/RenderCore/Public/GlobalShader.h
@@ -99,6 +99,7 @@ public:
 		ShouldCompilePermutationType InShouldCompilePermutationRef,
 		ShouldPrecachePermutationType InShouldPrecachePermutationRef,
 		GetRayTracingPayloadTypeType InGetRayTracingPayloadTypeRef,
+		GetShaderBindingLayoutType InGetShaderBindingLayoutTypeRef,
 #if WITH_EDITOR
 		ModifyCompilationEnvironmentType InModifyCompilationEnvironmentRef,
 		ValidateCompiledResultType InValidateCompiledResultRef,
@@ -112,6 +113,7 @@ public:
 			InShouldCompilePermutationRef,
 			InShouldPrecachePermutationRef,
 			InGetRayTracingPayloadTypeRef,
+			InGetShaderBindingLayoutTypeRef,
 #if WITH_EDITOR
 			InModifyCompilationEnvironmentRef,
 			InValidateCompiledResultRef,
diff --git a/Engine/Source/Runtime/RenderCore/Public/Shader.h b/Engine/Source/Runtime/RenderCore/Public/Shader.h
--- a/Engine/Source/Runtime/RenderCore/Public/Shader.h
+++ b/Engine/Source/Runtime/RenderCore/Public/Shader.h
@@ -81,6 +81,7 @@ class FShaderParametersMetadata;
 class FShaderPipelineType;
 class FShaderType;
 class FVertexFactoryType;
+class FRHIShaderBindingLayout;
 struct FShaderCompiledShaderInitializerType;
 struct FShaderCompilerOutput;
 using FShaderMapAssetPaths = TSet<FName>; // Copied from ShaderCodeLibrary.h
@@ -877,6 +878,9 @@ public:
 	/** Can be overridden by FShader subclasses to modify their compile environment just before compilation occurs. */
 	static void ModifyCompilationEnvironment(const FShaderPermutationParameters&, FShaderCompilerEnvironment&) {}
 
+	/**	Get shader binding layout used by the shader */
+	static const FShaderBindingLayout* GetShaderBindingLayout(const FShaderPermutationParameters&) { return nullptr; }
+
 	/** Can be overridden by FShader subclasses to determine whether a specific permutation should be compiled. */
 	static bool ShouldCompilePermutation(const FShaderPermutationParameters&) { return true; }
 
@@ -1255,6 +1259,7 @@ public:
 	typedef bool (*ShouldCompilePermutationType)(const FShaderPermutationParameters&);
 	typedef EShaderPermutationPrecacheRequest (*ShouldPrecachePermutationType)(const FShaderPermutationParameters&);
 	typedef ERayTracingPayloadType(*GetRayTracingPayloadTypeType)(const int32 PermutationId);
+	typedef const FShaderBindingLayout*(*GetShaderBindingLayoutType)(const FShaderPermutationParameters&);
 #if WITH_EDITOR
 	typedef void (*ModifyCompilationEnvironmentType)(const FShaderPermutationParameters&, FShaderCompilerEnvironment&);
 	typedef bool (*ValidateCompiledResultType)(EShaderPlatform, const FShaderParameterMap&, TArray<FString>&);
@@ -1291,6 +1296,7 @@ public:
 		ShouldCompilePermutationType InShouldCompilePermutationRef,
 		ShouldPrecachePermutationType InShouldPrecachePermutationRef,
 		GetRayTracingPayloadTypeType InGetRayTracingPayloadTypeRef,
+		GetShaderBindingLayoutType InGetShaderBindingLayoutTypeRef,
 #if WITH_EDITOR
 		ModifyCompilationEnvironmentType InModifyCompilationEnvironmentRef,
 		ValidateCompiledResultType InValidateCompiledResultRef,
@@ -1307,6 +1313,7 @@ public:
 
 	RENDERCORE_API bool ShouldCompilePermutation(const FShaderPermutationParameters& Parameters) const;
 	RENDERCORE_API EShaderPermutationPrecacheRequest ShouldPrecachePermutation(const FShaderPermutationParameters& Parameters) const;
+	RENDERCORE_API const FShaderBindingLayout* GetShaderBindingLayout(const FShaderPermutationParameters& Parameters) const; 
 
 #if WITH_EDITOR
 	RENDERCORE_API void ModifyCompilationEnvironment(const FShaderPermutationParameters& Parameters, FShaderCompilerEnvironment& OutEnvironment) const;
@@ -1515,6 +1522,7 @@ private:
 	ShouldCompilePermutationType ShouldCompilePermutationRef;
 	ShouldPrecachePermutationType ShouldPrecachePermutationRef;
 	GetRayTracingPayloadTypeType GetRayTracingPayloadTypeRef;
+	GetShaderBindingLayoutType GetShaderBindingLayoutTypeRef;
 #if WITH_EDITOR
 	ModifyCompilationEnvironmentType ModifyCompilationEnvironmentRef;
 	ValidateCompiledResultType ValidateCompiledResultRef;
@@ -1659,7 +1667,8 @@ struct FShaderCompiledShaderInitializerType
 	ShaderClass::ConstructCompiledInstance, \
 	ShaderClass::ShouldCompilePermutationImpl, \
 	ShaderClass::ShouldPrecachePermutationImpl, \
-	ShaderClass::GetRayTracingPayloadType \
+	ShaderClass::GetRayTracingPayloadType, \
+	ShaderClass::GetShaderBindingLayout \
 	SHADER_TYPE_EDITOR_VTABLE(ShaderClass)
 
 #if !UE_BUILD_DOCS
diff --git a/Engine/Source/Runtime/RenderCore/Public/ShaderCore.h b/Engine/Source/Runtime/RenderCore/Public/ShaderCore.h
--- a/Engine/Source/Runtime/RenderCore/Public/ShaderCore.h
+++ b/Engine/Source/Runtime/RenderCore/Public/ShaderCore.h
@@ -27,6 +27,7 @@
 #include "Misc/TVariant.h"
 #include "PixelFormat.h"
 #include "RHIDefinitions.h"
+#include "RHIShaderBindingLayout.h"
 #include "Serialization/Archive.h"
 #include "Serialization/MemoryLayout.h"
 #include "ShaderParameterMetadata.h"
@@ -423,6 +424,54 @@ private:
 	uint64 Data;
 };
 
+/**
+RenderCore wrapper around FRHIShaderBindingLayout which can also cache the uniform buffer declarations used during shader code generation
+*/
+class FShaderBindingLayout
+{
+public:
+
+	FRHIShaderBindingLayout RHILayout;
+
+#if WITH_EDITOR
+	FThreadSafeSharedAnsiStringPtr GetUniformBufferDeclarationAnsiPtr(const FShaderParametersMetadata* ShaderParametersMetadata) const;
+	void SetUniformBufferDeclarationAnsiPtr(const FShaderParametersMetadata* ShaderParametersMetadata, FThreadSafeSharedAnsiStringPtr UniformBufferDeclarationAnsi);
+
+	RENDERCORE_API void AddRequiredSymbols(TArray<FString>& RequiredSymbols) const;
+
+protected:
+	
+	TArray<FThreadSafeSharedAnsiStringPtr> UniformBufferDeclarationsAnsi;
+#endif
+};
+
+/**
+Static shader binding layout object managing all possible binding type versions of the FShaderBindingLayout
+*/
+class FShaderBindingLayoutContainer
+{
+public:
+	enum class EBindingType : uint8
+	{
+		Bindless,
+		NotBindless,
+		Num
+	};
+
+	const FShaderBindingLayout& GetLayout(EBindingType BindingType) const
+	{
+		return Layouts[(uint8)BindingType];
+	}
+	void SetLayout(EBindingType BindingType, const FShaderBindingLayout& InLayout)
+	{
+		Layouts[(uint8)BindingType] = InLayout;
+	}
+
+protected:
+
+	FShaderBindingLayout Layouts[(uint8)EBindingType::Num];
+};
+
 struct FShaderResourceTableMap
 {
 	TArray<FUniformResourceEntry> Resources;
@@ -447,6 +496,12 @@ struct FShaderCompilerEnvironment
 	TMap<uint32,uint8> RenderTargetOutputFormatsMap;
 	FShaderResourceTableMap ResourceTableMap;
 	TMap<FString, FUniformBufferEntry> UniformBufferMap;
+	
+	// Optional shader binding layout which can be used build the Uniform buffer map
+	const FShaderBindingLayout* ShaderBindingLayout = nullptr;
+
+	// Serialized version of the shader binding layout which can be used during platform specific shader code generation and serialization
+	FRHIShaderBindingLayout RHIShaderBindingLayout;
 
 	const ITargetPlatform* TargetPlatform = nullptr;
 
@@ -669,6 +724,7 @@ enum class EShaderOptionalDataKey : uint8
 	UniformBuffers       = uint8('u'),
 	Validation           = uint8('V'),
 	VendorExtension      = uint8('v'),
+	ShaderBindingLayout  = uint8('s'),
 };
 
 enum class EShaderResourceUsageFlags : uint8
@@ -741,6 +797,12 @@ struct FShaderCodeUniformBuffers
 	// We store an array of FString objects
 };
 
+struct FShaderCodeShaderResourceTableDataDesc
+{
+	static const EShaderOptionalDataKey Key = EShaderOptionalDataKey::ShaderBindingLayout;
+	// We store FRHIShaderBindingLayout
+};
+
 // if this changes you need to make sure all shaders get invalidated
 struct FShaderCodeVendorExtension
 {
diff --git a/Engine/Source/Runtime/RenderCore/Public/ShaderParameterMacros.h b/Engine/Source/Runtime/RenderCore/Public/ShaderParameterMacros.h
--- a/Engine/Source/Runtime/RenderCore/Public/ShaderParameterMacros.h
+++ b/Engine/Source/Runtime/RenderCore/Public/ShaderParameterMacros.h
@@ -26,6 +26,8 @@ class FRDGBuffer;
 class FRDGBufferSRV;
 class FRDGBufferUAV;
 class FRDGUniformBuffer;
+enum class EShaderBindingLayoutFlags : uint8;
+class FShaderBindingLayoutContainer;
 template <typename TUniformStruct> class TRDGUniformBuffer;
 
 /** Alignements tools because alignas() does not work on type in clang. */
@@ -1463,6 +1465,9 @@ extern RENDERCORE_API FShaderParametersMetadata* FindUniformBufferStructByLayout
 
 extern RENDERCORE_API FShaderParametersMetadata* FindUniformBufferStructByShaderVariableName(const FHashedName& Name);
 
+/** Build the shader binding layout desc from array of used Uniform buffers which will be bound as static uniform buffers to the rhi */
+extern RENDERCORE_API void BuildShaderBindingLayout(TConstArrayView<FShaderParametersMetadata*> UniformBuffers, EShaderBindingLayoutFlags BaseShaderBindingLayoutFlags, FShaderBindingLayoutContainer& OutShaderBindingLayoutContainer);
+
 /** Begins & ends a shader parameter structure.
  *
  * Example:
diff --git a/Engine/Source/Runtime/RenderCore/Public/ShaderParameters.h b/Engine/Source/Runtime/RenderCore/Public/ShaderParameters.h
--- a/Engine/Source/Runtime/RenderCore/Public/ShaderParameters.h
+++ b/Engine/Source/Runtime/RenderCore/Public/ShaderParameters.h
@@ -26,6 +26,7 @@ class FShaderParametersMetadata;
 struct FRWBuffer;
 struct FRWBufferStructured;
 struct FShaderCompilerEnvironment;
+struct FRHIUniformBufferShaderBindingLayout;
 
 enum class EShaderParameterType : uint8;
 DECLARE_INTRINSIC_TYPE_LAYOUT(EShaderParameterType);
@@ -37,7 +38,7 @@ DECLARE_INTRINSIC_TYPE_LAYOUT(EShaderCodeResourceBindingType);
 namespace UE::ShaderParameters
 {
 	/** Creates a shader code declaration of this struct for the given shader platform. */
-	RENDERCORE_API FString CreateUniformBufferShaderDeclaration(const TCHAR* Name, const FShaderParametersMetadata& UniformBufferStruct);
+	RENDERCORE_API FString CreateUniformBufferShaderDeclaration(const TCHAR* Name, const FShaderParametersMetadata& UniformBufferStruct, const FRHIUniformBufferShaderBindingLayout* UniformBufferSBLayout);
 
 	UE_DEPRECATED(5.5, "AddUniformBufferIncludesToEnviroment now takes a set of FShaderParametersMetadata pointers")
 	inline void AddUniformBufferIncludesToEnvironment(FShaderCompilerEnvironment& OutEnvironment, const TSet<const TCHAR*, TStringPointerSetKeyFuncs_DEPRECATED<const TCHAR*>>& InUniformBufferNames) {}
