Support GPU diagnostic buffer binding when using static root signature

diff --git a/Engine/Source/Developer/Windows/ShaderFormatD3D/Private/D3DShaderCompiler.inl b/Engine/Source/Developer/Windows/ShaderFormatD3D/Private/D3DShaderCompiler.inl
--- a/Engine/Source/Developer/Windows/ShaderFormatD3D/Private/D3DShaderCompiler.inl
+++ b/Engine/Source/Developer/Windows/ShaderFormatD3D/Private/D3DShaderCompiler.inl
@@ -375,6 +375,11 @@ inline FShaderCodePackedResourceCounts InitPackedResourceCounts(const FD3DShader
 		PackedResourceCounts.UsageFlags |= EShaderResourceUsageFlags::BindlessSamplers;
 	}
 
+	if (CompileData.bDiagnosticBufferUsed)
+	{
+		PackedResourceCounts.UsageFlags |= EShaderResourceUsageFlags::DiagnosticBuffer;
+	}
+
 	PackedResourceCounts.NumSamplers = static_cast<uint8>(CompileData.NumSamplers);
 	PackedResourceCounts.NumSRVs = static_cast<uint8>(CompileData.NumSRVs);
 	PackedResourceCounts.NumCBs = static_cast<uint8>(CompileData.NumCBs);
diff --git a/Engine/Source/Developer/Windows/ShaderFormatD3D/Private/D3DShaderCompilerDXC.cpp b/Engine/Source/Developer/Windows/ShaderFormatD3D/Private/D3DShaderCompilerDXC.cpp
--- a/Engine/Source/Developer/Windows/ShaderFormatD3D/Private/D3DShaderCompilerDXC.cpp
+++ b/Engine/Source/Developer/Windows/ShaderFormatD3D/Private/D3DShaderCompilerDXC.cpp
@@ -836,7 +836,6 @@ inline bool IsCompatibleBinding(const D3D12_SHADER_INPUT_BIND_DESC& BindDesc, ui
 	}
 	if (!bIsCompatibleBinding)
 	{
-		// #todo: there is currently no common header where a binding space number or buffer name could be defined. See D3DCommon.ush and D3D12RootSignature.cpp.
 		const bool bIsUEDebugBuffer = (FCStringAnsi::Strcmp(BindDesc.Name, "UEDiagnosticBuffer") == 0);
 		bIsCompatibleBinding = bIsUEDebugBuffer && (BindDesc.Space == UE_HLSL_SPACE_DIAGNOSTIC);
 	}
@@ -1264,11 +1263,6 @@ bool CompileAndProcessD3DShaderDXC(
 					EnumAddFlags(CodeFeatures.CodeFeatures, EShaderCodeFeatures::Atomic64);
 				}
 
-				if (CompileData.bDiagnosticBufferUsed)
-				{
-					EnumAddFlags(CodeFeatures.CodeFeatures, EShaderCodeFeatures::DiagnosticBuffer);
-				}
-
 				if ((ShaderRequiresFlags & D3D_SHADER_REQUIRES_RESOURCE_DESCRIPTOR_HEAP_INDEXING) != 0)
 				{
 					EnumAddFlags(CodeFeatures.CodeFeatures, EShaderCodeFeatures::BindlessResources);
diff --git a/Engine/Source/Runtime/D3D12RHI/Private/D3D12RootSignature.cpp b/Engine/Source/Runtime/D3D12RHI/Private/D3D12RootSignature.cpp
--- a/Engine/Source/Runtime/D3D12RHI/Private/D3D12RootSignature.cpp
+++ b/Engine/Source/Runtime/D3D12RHI/Private/D3D12RootSignature.cpp
@@ -412,13 +412,19 @@ void FD3D12RootSignature::InitStaticGraphicsRootSignature(EShaderBindingLayoutFl
 		for (int32 ParameterSlot = 0; ParameterSlot < Creator.Parameters.Num(); ++ParameterSlot)
 		{
 			const CD3DX12_ROOT_PARAMETER1 RootParameter = Creator.Parameters[ParameterSlot];
-			if (RootParameter.ParameterType == D3D12_ROOT_PARAMETER_TYPE_32BIT_CONSTANTS)
+
+			if (RootParameter.ParameterType == D3D12_ROOT_PARAMETER_TYPE_32BIT_CONSTANTS
+				&& RootParameter.Constants.RegisterSpace == UE_HLSL_SPACE_SHADER_ROOT_CONSTANTS
+				&& RootParameter.Constants.ShaderRegister == 0)
 			{
-				if (RootParameter.Constants.RegisterSpace == UE_HLSL_SPACE_SHADER_ROOT_CONSTANTS && RootParameter.Constants.ShaderRegister == 0)
-				{
-					RootConstantsSlot = int8(ParameterSlot);
-					break;
-				}
+				RootConstantsSlot = int8(ParameterSlot);
+			}
+
+			if (RootParameter.ParameterType == D3D12_ROOT_PARAMETER_TYPE_UAV
+				&& RootParameter.Constants.RegisterSpace == UE_HLSL_SPACE_DIAGNOSTIC
+				&& RootParameter.Constants.ShaderRegister == 0)
+			{
+				DiagnosticBufferSlot = int8(ParameterSlot);
 			}
 		}
 	}
@@ -435,13 +441,19 @@ void FD3D12RootSignature::InitStaticComputeRootSignatureDesc(EShaderBindingLayou
 		for (int32 ParameterSlot = 0; ParameterSlot < Creator.Parameters.Num(); ++ParameterSlot)
 		{
 			const CD3DX12_ROOT_PARAMETER1 RootParameter = Creator.Parameters[ParameterSlot];
-			if (RootParameter.ParameterType == D3D12_ROOT_PARAMETER_TYPE_32BIT_CONSTANTS)
+
+			if (RootParameter.ParameterType == D3D12_ROOT_PARAMETER_TYPE_32BIT_CONSTANTS
+				&& RootParameter.Constants.RegisterSpace == UE_HLSL_SPACE_SHADER_ROOT_CONSTANTS
+				&& RootParameter.Constants.ShaderRegister == 0)
 			{
-				if (RootParameter.Constants.RegisterSpace == UE_HLSL_SPACE_SHADER_ROOT_CONSTANTS && RootParameter.Constants.ShaderRegister == 0)
-				{
-					RootConstantsSlot = int8(ParameterSlot);
-					break;
-				}
+				RootConstantsSlot = int8(ParameterSlot);
+			}
+
+			if (RootParameter.ParameterType == D3D12_ROOT_PARAMETER_TYPE_UAV
+				&& RootParameter.Constants.RegisterSpace == UE_HLSL_SPACE_DIAGNOSTIC
+				&& RootParameter.Constants.ShaderRegister == 0)
+			{
+				DiagnosticBufferSlot = int8(ParameterSlot);
 			}
 		}
 	}
diff --git a/Engine/Source/Runtime/D3D12RHI/Private/D3D12Shader.h b/Engine/Source/Runtime/D3D12RHI/Private/D3D12Shader.h
--- a/Engine/Source/Runtime/D3D12RHI/Private/D3D12Shader.h
+++ b/Engine/Source/Runtime/D3D12RHI/Private/D3D12Shader.h
@@ -110,7 +110,7 @@ struct FD3D12ShaderData
 	FORCEINLINE EShaderCodeFeatures GetFeatures() const { return EShaderCodeFeatures::None; }
 #endif
 
-	FORCEINLINE bool UsesDiagnosticBuffer() const { return EnumHasAnyFlags(GetFeatures(), EShaderCodeFeatures::DiagnosticBuffer); }
+	FORCEINLINE bool UsesDiagnosticBuffer() const { return EnumHasAnyFlags(ResourceCounts.UsageFlags, EShaderResourceUsageFlags::DiagnosticBuffer); }
 	FORCEINLINE bool UsesGlobalUniformBuffer() const { return EnumHasAnyFlags(ResourceCounts.UsageFlags, EShaderResourceUsageFlags::GlobalUniformBuffer); }
 	FORCEINLINE bool UsesBindlessResources() const { return EnumHasAnyFlags(ResourceCounts.UsageFlags, EShaderResourceUsageFlags::BindlessResources); }
 	FORCEINLINE bool UsesBindlessSamplers() const { return EnumHasAnyFlags(ResourceCounts.UsageFlags, EShaderResourceUsageFlags::BindlessSamplers); }
diff --git a/Engine/Source/Runtime/D3D12RHI/Public/D3D12RootSignatureDefinitions.h b/Engine/Source/Runtime/D3D12RHI/Public/D3D12RootSignatureDefinitions.h
--- a/Engine/Source/Runtime/D3D12RHI/Public/D3D12RootSignatureDefinitions.h
+++ b/Engine/Source/Runtime/D3D12RHI/Public/D3D12RootSignatureDefinitions.h
@@ -181,6 +181,7 @@ namespace D3D12ShaderUtils
 		virtual void AddRootFlag(D3D12_ROOT_SIGNATURE_FLAGS Flag) = 0;
 		virtual void AddConstantBufferViewParameter(uint32 Register, uint32 Space, ERootSignatureVisibility Visibility = ERootSignatureVisibility::All) = 0;
 		virtual void AddShaderResourceViewParameter(uint32 Register, uint32 Space, ERootSignatureVisibility Visibility = ERootSignatureVisibility::All) = 0;
+		virtual void AddUnorderedAccessViewParameter(uint32 Register, uint32 Space, ERootSignatureVisibility Visibility = ERootSignatureVisibility::All) = 0;
 		virtual void AddTable(ERootSignatureVisibility Visibility, ERootSignatureRangeType Type, int32 NumDescriptors, D3D12_DESCRIPTOR_RANGE_FLAGS FlagsOverride = D3D12_DESCRIPTOR_RANGE_FLAG_NONE) = 0;
 		virtual void AddConstantsParameter(uint32 Num32BitValues, uint32 Register, uint32 Space) = 0;
 
@@ -268,6 +269,12 @@ namespace D3D12ShaderUtils
 			Parameter.InitAsShaderResourceView(Register, Space, D3D12_ROOT_DESCRIPTOR_FLAG_NONE, GetD3D12ShaderVisibility(Visibility));
 		}
 
+		void AddUnorderedAccessViewParameter(uint32 Register, uint32 Space, ERootSignatureVisibility Visibility = ERootSignatureVisibility::All) override
+		{
+			CD3DX12_ROOT_PARAMETER1& Parameter = Parameters.AddZeroed_GetRef();
+			Parameter.InitAsUnorderedAccessView(Register, Space, D3D12_ROOT_DESCRIPTOR_FLAG_NONE, GetD3D12ShaderVisibility(Visibility));
+		}
+
 		void AddConstantsParameter(uint32 Num32BitValues, uint32 Register, uint32 Space) override
 		{
 			CD3DX12_ROOT_PARAMETER1& Parameter = Parameters.AddZeroed_GetRef();
@@ -345,9 +352,24 @@ namespace D3D12ShaderUtils
 			Table += Line;
 		}
 
+		void AddUnorderedAccessViewParameter(uint32 Register, uint32 Space, ERootSignatureVisibility Visibility = ERootSignatureVisibility::All) override
+		{
+			FString Line = FString::Printf(TEXT("UAV(u%d, space=%d, visibility=%s)"), Register, Space, GetVisibilityFlag(Visibility));
+			if (Table.Len() > 0)
+			{
+				Table += ",";
+			}
+			Table += Line;
+		}
+
 		void AddConstantsParameter(uint32 Num32BitValues, uint32 Register, uint32 Space) override
 		{
-			Constants.Appendf(TEXT("RootConstants(num32BitConstants=%d, b%d, space=%d),"), Num32BitValues, Register, Space);
+			FString Line = FString::Printf(TEXT("RootConstants(num32BitConstants=%d, b%d, space=%d)"), Num32BitValues, Register, Space);
+			if (Table.Len() > 0)
+			{
+				Table += ",";
+			}
+			Table += Line;
 		}
 
 		void AddTable(ERootSignatureVisibility Visibility, ERootSignatureRangeType Type, int32 NumDescriptors, D3D12_DESCRIPTOR_RANGE_FLAGS FlagsOverride = D3D12_DESCRIPTOR_RANGE_FLAG_NONE) override
@@ -367,16 +389,14 @@ namespace D3D12ShaderUtils
 
 		FString GenerateString() const
 		{
-			FString String = FString::Printf(TEXT("\"RootFlags(%s),%s%s\""),
+			FString String = FString::Printf(TEXT("\"RootFlags(%s),%s\""),
 				RootFlags.Len() == 0 ? TEXT("0") : *RootFlags,
-				*Constants,
 				*Table);
 			return String;
 		}
 
 		FString RootFlags;
 		FString Table;
-		FStringBuilderBase Constants;
 	};
 
 	inline void AddAllStandardTablesForVisibility(FRootSignatureCreator& Creator, ERootSignatureVisibility Visibility)
@@ -392,6 +412,16 @@ namespace D3D12ShaderUtils
 		// Ensure the creator starts in a clean state (in cases of creator reuse, etc.).
 		Creator.SetFlags(InFlags | EShaderBindingLayoutFlags::InputAssembler);
 
+		Creator.AddUnorderedAccessViewParameter(0, UE_HLSL_SPACE_DIAGNOSTIC);
+
+		if (EnumHasAnyFlags(InFlags, EShaderBindingLayoutFlags::RootConstants))
+		{
+			const uint32 NumConstants = 4u;
+			Creator.AddConstantsParameter(NumConstants, 0u, UE_HLSL_SPACE_SHADER_ROOT_CONSTANTS); // UERootConstants
+		}
+
+		// Tables must be added only after other parameters
+
 		AddAllStandardTablesForVisibility(Creator, ERootSignatureVisibility::Pixel);
 		AddAllStandardTablesForVisibility(Creator, ERootSignatureVisibility::Vertex);
 		AddAllStandardTablesForVisibility(Creator, ERootSignatureVisibility::Geometry);
@@ -406,12 +436,6 @@ namespace D3D12ShaderUtils
 
 		Creator.AddTable(ERootSignatureVisibility::Pixel, ERootSignatureRangeType::UAV, MAX_UAVS);
 		Creator.AddTable(ERootSignatureVisibility::Vertex, ERootSignatureRangeType::UAV, MAX_UAVS);
-
-		if (EnumHasAnyFlags(InFlags, EShaderBindingLayoutFlags::RootConstants))
-		{
-			const uint32 NumConstants = 4u;
-			Creator.AddConstantsParameter(NumConstants, 0u, UE_HLSL_SPACE_SHADER_ROOT_CONSTANTS); // UERootConstants
-		}
 	}
 
 	// Fat/Static Compute Root Signature
@@ -419,14 +443,19 @@ namespace D3D12ShaderUtils
 	{
 		// Ensure the creator starts in a clean state (in cases of creator reuse, etc.).
 		Creator.SetFlags(InFlags);
-		AddAllStandardTablesForVisibility(Creator, ERootSignatureVisibility::All);
-		Creator.AddTable(ERootSignatureVisibility::All, ERootSignatureRangeType::UAV, MAX_UAVS);
+
+		Creator.AddUnorderedAccessViewParameter(0, UE_HLSL_SPACE_DIAGNOSTIC);
 
 		if (EnumHasAnyFlags(InFlags, EShaderBindingLayoutFlags::RootConstants))
 		{
 			const uint32 NumConstants = 4u;
 			Creator.AddConstantsParameter(NumConstants, 0u, UE_HLSL_SPACE_SHADER_ROOT_CONSTANTS); // UERootConstants
 		}
+
+		// Tables must be added only after other parameters
+
+		AddAllStandardTablesForVisibility(Creator, ERootSignatureVisibility::All);
+		Creator.AddTable(ERootSignatureVisibility::All, ERootSignatureRangeType::UAV, MAX_UAVS);
 	}
 
 #if !defined(D3D12RHI_TOOLS_RAYTRACING_SHADERS_UNSUPPORTED)
diff --git a/Engine/Source/Runtime/RenderCore/Public/ShaderCore.h b/Engine/Source/Runtime/RenderCore/Public/ShaderCore.h
--- a/Engine/Source/Runtime/RenderCore/Public/ShaderCore.h
+++ b/Engine/Source/Runtime/RenderCore/Public/ShaderCore.h
@@ -735,6 +735,7 @@ enum class EShaderResourceUsageFlags : uint8
 	RootConstants         = 1 << 3,
 	NoDerivativeOps       = 1 << 4,
 	ShaderBundle          = 1 << 5,
+	DiagnosticBuffer      = 1 << 6,
 };
 ENUM_CLASS_FLAGS(EShaderResourceUsageFlags)
 
@@ -767,7 +768,7 @@ enum class EShaderCodeFeatures : uint16
 	SixteenBitTypes         = 1 << 1,
 	TypedUAVLoadsExtended   = 1 << 2,
 	Atomic64                = 1 << 3,
-	DiagnosticBuffer        = 1 << 4,
+	DiagnosticBuffer UE_DEPRECATED(5.5, "EShaderCodeFeatures::DiagnosticBuffer is superseded by EShaderResourceUsageFlags::DiagnosticBuffer") = 1 << 4,
 	BindlessResources       = 1 << 5,
 	BindlessSamplers        = 1 << 6,
 	StencilRef              = 1 << 7,
