Fix diagnostic buffer root parameter slot initialization for static root signatures

diff --git a/Engine/Source/Runtime/D3D12RHI/Private/D3D12RootSignature.cpp b/Engine/Source/Runtime/D3D12RHI/Private/D3D12RootSignature.cpp
--- a/Engine/Source/Runtime/D3D12RHI/Private/D3D12RootSignature.cpp
+++ b/Engine/Source/Runtime/D3D12RHI/Private/D3D12RootSignature.cpp
@@ -407,25 +407,22 @@ void FD3D12RootSignature::InitStaticGraphicsRootSignature(EShaderBindingLayoutFl
 	D3D12ShaderUtils::CreateGfxRootSignature(Creator, InFlags);
 	Init(Creator.Finalize());
 
-	if (EnumHasAnyFlags(InFlags, EShaderBindingLayoutFlags::RootConstants))
+	for (int32 ParameterSlot = 0; ParameterSlot < Creator.Parameters.Num(); ++ParameterSlot)
 	{
-		for (int32 ParameterSlot = 0; ParameterSlot < Creator.Parameters.Num(); ++ParameterSlot)
-		{
-			const CD3DX12_ROOT_PARAMETER1 RootParameter = Creator.Parameters[ParameterSlot];
+		const CD3DX12_ROOT_PARAMETER1 RootParameter = Creator.Parameters[ParameterSlot];
 
-			if (RootParameter.ParameterType == D3D12_ROOT_PARAMETER_TYPE_32BIT_CONSTANTS
-				&& RootParameter.Constants.RegisterSpace == UE_HLSL_SPACE_SHADER_ROOT_CONSTANTS
-				&& RootParameter.Constants.ShaderRegister == 0)
-			{
-				RootConstantsSlot = int8(ParameterSlot);
-			}
+		if (RootParameter.ParameterType == D3D12_ROOT_PARAMETER_TYPE_32BIT_CONSTANTS
+			&& RootParameter.Constants.RegisterSpace == UE_HLSL_SPACE_SHADER_ROOT_CONSTANTS
+			&& RootParameter.Constants.ShaderRegister == 0)
+		{
+			RootConstantsSlot = int8(ParameterSlot);
+		}
 
-			if (RootParameter.ParameterType == D3D12_ROOT_PARAMETER_TYPE_UAV
-				&& RootParameter.Constants.RegisterSpace == UE_HLSL_SPACE_DIAGNOSTIC
-				&& RootParameter.Constants.ShaderRegister == 0)
-			{
-				DiagnosticBufferSlot = int8(ParameterSlot);
-			}
+		if (RootParameter.ParameterType == D3D12_ROOT_PARAMETER_TYPE_UAV
+			&& RootParameter.Descriptor.RegisterSpace == UE_HLSL_SPACE_DIAGNOSTIC
+			&& RootParameter.Descriptor.ShaderRegister == 0)
+		{
+			DiagnosticBufferSlot = int8(ParameterSlot);
 		}
 	}
 }
@@ -436,25 +433,22 @@ void FD3D12RootSignature::InitStaticComputeRootSignatureDesc(EShaderBindingLayou
 	D3D12ShaderUtils::CreateComputeRootSignature(Creator, InFlags);
 	Init(Creator.Finalize());
 
-	if (EnumHasAnyFlags(InFlags, EShaderBindingLayoutFlags::RootConstants))
+	for (int32 ParameterSlot = 0; ParameterSlot < Creator.Parameters.Num(); ++ParameterSlot)
 	{
-		for (int32 ParameterSlot = 0; ParameterSlot < Creator.Parameters.Num(); ++ParameterSlot)
-		{
-			const CD3DX12_ROOT_PARAMETER1 RootParameter = Creator.Parameters[ParameterSlot];
+		const CD3DX12_ROOT_PARAMETER1 RootParameter = Creator.Parameters[ParameterSlot];
 
-			if (RootParameter.ParameterType == D3D12_ROOT_PARAMETER_TYPE_32BIT_CONSTANTS
-				&& RootParameter.Constants.RegisterSpace == UE_HLSL_SPACE_SHADER_ROOT_CONSTANTS
-				&& RootParameter.Constants.ShaderRegister == 0)
-			{
-				RootConstantsSlot = int8(ParameterSlot);
-			}
+		if (RootParameter.ParameterType == D3D12_ROOT_PARAMETER_TYPE_32BIT_CONSTANTS
+			&& RootParameter.Constants.RegisterSpace == UE_HLSL_SPACE_SHADER_ROOT_CONSTANTS
+			&& RootParameter.Constants.ShaderRegister == 0)
+		{
+			RootConstantsSlot = int8(ParameterSlot);
+		}
 
-			if (RootParameter.ParameterType == D3D12_ROOT_PARAMETER_TYPE_UAV
-				&& RootParameter.Constants.RegisterSpace == UE_HLSL_SPACE_DIAGNOSTIC
-				&& RootParameter.Constants.ShaderRegister == 0)
-			{
-				DiagnosticBufferSlot = int8(ParameterSlot);
-			}
+		if (RootParameter.ParameterType == D3D12_ROOT_PARAMETER_TYPE_UAV
+			&& RootParameter.Descriptor.RegisterSpace == UE_HLSL_SPACE_DIAGNOSTIC
+			&& RootParameter.Descriptor.ShaderRegister == 0)
+		{
+			DiagnosticBufferSlot = int8(ParameterSlot);
 		}
 	}
 }
