Fix submission timeout detection by starting it when a queue fence is signaled, not as soon a payload is submitted

#rnx

diff --git a/Engine/Source/Runtime/D3D12RHI/Private/D3D12Submission.cpp b/Engine/Source/Runtime/D3D12RHI/Private/D3D12Submission.cpp
--- a/Engine/Source/Runtime/D3D12RHI/Private/D3D12Submission.cpp
+++ b/Engine/Source/Runtime/D3D12RHI/Private/D3D12Submission.cpp
@@ -964,7 +964,7 @@ void FD3D12Queue::FlushBatchedPayloads()
 			check(Fence.LastSignaledValue < Payload->CompletionFenceValue);
 
 			VERIFYD3D12RESULT(D3DCommandQueue->Signal(Fence.D3DFence, Payload->CompletionFenceValue));
-			Fence.LastSignaledValue = Payload->CompletionFenceValue;
+			Fence.LastSignaledValue.store(Payload->CompletionFenceValue, std::memory_order_release);
 		}
 
 		// Submission of this payload is completed. Signal the submission event if one was provided.
@@ -1192,6 +1192,7 @@ FD3D12DynamicRHI::FProcessResult FD3D12DynamicRHI::ProcessInterruptQueue()
 
 			// Check for GPU completion
 			uint64 CompletedFenceValue = CurrentQueue.Fence.D3DFence->GetCompletedValue();
+			uint64 LastSignaledFenceValue = CurrentQueue.Fence.LastSignaledValue.load(std::memory_order_acquire);
 
 			// If the GPU crashes or hangs, the driver will signal all fences to UINT64_MAX.
 			if (CompletedFenceValue == UINT64_MAX)
@@ -1214,7 +1215,7 @@ FD3D12DynamicRHI::FProcessResult FD3D12DynamicRHI::ProcessInterruptQueue()
 				Result.Status |= EQueueStatus::Pending;
 
 				// Detect a hung GPU
-				if (!Payload->SubmissionTime.IsSet())
+				if (!Payload->SubmissionTime.IsSet() && LastSignaledFenceValue >= Payload->CompletionFenceValue)
 				{
 					//
 					// Keep track of the first time we've checked for completion on the interrupt thread.
@@ -1224,7 +1225,7 @@ FD3D12DynamicRHI::FProcessResult FD3D12DynamicRHI::ProcessInterruptQueue()
 					Payload->SubmissionTime = Timer.Elapsed;
 				}
 
-				if (Payload->SubmissionTime != TNumericLimits<uint64>::Max())
+				if (Payload->SubmissionTime.IsSet() && Payload->SubmissionTime != TNumericLimits<uint64>::Max())
 				{
 					static const double CyclesPerSecond = 1.0 / FPlatformTime::GetSecondsPerCycle64();
 					const uint64 TimeoutCycles = FMath::TruncToInt64(GD3D12SubmissionTimeout * CyclesPerSecond);
diff --git a/Engine/Source/Runtime/D3D12RHI/Private/D3D12Submission.h b/Engine/Source/Runtime/D3D12RHI/Private/D3D12Submission.h
--- a/Engine/Source/Runtime/D3D12RHI/Private/D3D12Submission.h
+++ b/Engine/Source/Runtime/D3D12RHI/Private/D3D12Submission.h
@@ -42,7 +42,7 @@ struct FD3D12Fence
 {
 	TRefCountPtr<ID3D12Fence> D3DFence;
 	uint64 NextCompletionValue = 0;
-	uint64 LastSignaledValue = 0;
+	std::atomic<uint64> LastSignaledValue = 0;
 	bool bInterruptAwaited = false;
 };
 
