Precalculate bundle descriptor count so that descriptor heap size is known up front.
Explicit descriptor heap cache aligns descriptor count to power of 2 so that there is greater cache reuse.

diff --git a/Engine/Source/Runtime/D3D12RHI/Private/D3D12ExplicitDescriptorCache.cpp b/Engine/Source/Runtime/D3D12RHI/Private/D3D12ExplicitDescriptorCache.cpp
--- a/Engine/Source/Runtime/D3D12RHI/Private/D3D12ExplicitDescriptorCache.cpp
+++ b/Engine/Source/Runtime/D3D12RHI/Private/D3D12ExplicitDescriptorCache.cpp
@@ -68,6 +68,10 @@ FD3D12ExplicitDescriptorHeapCache::Entry FD3D12ExplicitDescriptorHeapCache::Allo
 {
 	FScopeLock Lock(&CriticalSection);
 
+	// Align request to enable greater reusue in the cache.
+	const uint32 MaxDescriptors = (Type == D3D12_DESCRIPTOR_HEAP_TYPE_SAMPLER) ? D3D12_MAX_SHADER_VISIBLE_SAMPLER_HEAP_SIZE : D3D12_MAX_SHADER_VISIBLE_DESCRIPTOR_HEAP_SIZE_TIER_1;
+	NumDescriptors = FMath::Clamp(FMath::RoundUpToPowerOfTwo(NumDescriptors), 0, MaxDescriptors);
+
 	++AllocatedEntries;
 
 	Entry Result = {};
diff --git a/Engine/Source/Runtime/D3D12RHI/Private/D3D12WorkGraph.cpp b/Engine/Source/Runtime/D3D12RHI/Private/D3D12WorkGraph.cpp
--- a/Engine/Source/Runtime/D3D12RHI/Private/D3D12WorkGraph.cpp
+++ b/Engine/Source/Runtime/D3D12RHI/Private/D3D12WorkGraph.cpp
@@ -555,6 +555,9 @@ void FD3D12CommandContext::DispatchWorkGraphShaderBundle(FRHIShaderBundle* Shade
 	TShaderRef<FDispatchShaderBundleWorkGraph> WorkGraphGlobalShader = GetGlobalShaderMap(GMaxRHIFeatureLevel)->GetShader<FDispatchShaderBundleWorkGraph>();
 	FD3D12WorkGraphShader* WorkGraphGlobalShaderRHI = static_cast<FD3D12WorkGraphShader*>(WorkGraphGlobalShader.GetWorkGraphShader());
 
+	uint32 ViewDescriptorCount = WorkGraphGlobalShaderRHI->ResourceCounts.NumSRVs + WorkGraphGlobalShaderRHI->ResourceCounts.NumCBs + WorkGraphGlobalShaderRHI->ResourceCounts.NumUAVs;
+	uint32 SamplerDescriptorCount = WorkGraphGlobalShaderRHI->ResourceCounts.NumSamplers;
+
 	const int32 NumRecords = Dispatches.Num();
 	checkf(NumRecords <= FDispatchShaderBundleWorkGraph::GetMaxShaderBundleSize(), TEXT("Too many entries in a shader bundle (%d). Try increasing 'r.ShaderBundle.MaxSize'"), NumRecords);
 
@@ -565,11 +568,18 @@ void FD3D12CommandContext::DispatchWorkGraphShaderBundle(FRHIShaderBundle* Shade
 
 	for (int32 DispatchIndex = 0; DispatchIndex < NumRecords; ++DispatchIndex)
 	{
-		FRHIWorkGraphShader* Shader = Dispatches[DispatchIndex].IsValid() ? Dispatches[DispatchIndex].WorkGraphShader : nullptr;
+		const FRHIShaderBundleComputeDispatch& Dispatch = Dispatches[DispatchIndex];
+		FRHIWorkGraphShader* Shader = Dispatch.IsValid() ? Dispatch.WorkGraphShader : nullptr;
 		LocalNodeShaders.Add(Shader);
 		if (Shader != nullptr)
 		{
 			ValidRecords.Emplace(uint32(DispatchIndex));
+
+			if (FD3D12WorkGraphShader* D3D12Shader = FD3D12DynamicRHI::ResourceCast(Shader))
+			{
+				ViewDescriptorCount += D3D12Shader->ResourceCounts.NumSRVs + D3D12Shader->ResourceCounts.NumCBs + D3D12Shader->ResourceCounts.NumUAVs;
+				SamplerDescriptorCount += D3D12Shader->ResourceCounts.NumSamplers;
+			}
 		}
 	}
 	const int32 NumValidRecords = ValidRecords.Num();
@@ -582,10 +592,8 @@ void FD3D12CommandContext::DispatchWorkGraphShaderBundle(FRHIShaderBundle* Shade
 	FWorkGraphPipelineState* WorkGraphPipelineState = PipelineStateCache::GetAndOrCreateWorkGraphPipelineState(RHICmdList, Initializer);
 	FD3D12WorkGraphPipelineState* Pipeline = static_cast<FD3D12WorkGraphPipelineState*>(GetRHIWorkGraphPipelineState(WorkGraphPipelineState));
 
-	// D3D12 is guaranteed to support 1M (D3D12_MAX_SHADER_VISIBLE_DESCRIPTOR_HEAP_SIZE_TIER_1) descriptors in a CBV/SRV/UAV heap, so clamp the size to this.
-	// https://docs.microsoft.com/en-us/windows/desktop/direct3d12/hardware-support
-	const uint32 NumViewDescriptors = D3D12_MAX_SHADER_VISIBLE_DESCRIPTOR_HEAP_SIZE_TIER_1;
-	const uint32 NumSamplerDescriptors = D3D12_MAX_SHADER_VISIBLE_SAMPLER_HEAP_SIZE;
+	const uint32 NumViewDescriptors = ViewDescriptorCount;
+	const uint32 NumSamplerDescriptors = SamplerDescriptorCount;
 
 	const uint32 MaxWorkers = 4u;
 	const uint32 NumWorkerThreads = FTaskGraphInterface::Get().GetNumWorkerThreads();
