Use normal D3D12 RHI deferred deletion queue to manage explicit descriptor heap reuse, instead of relying on frame fence

* Prevents memory bloat when BeginFrame/EndFrame is not called periodically

diff --git a/Engine/Source/Runtime/D3D12RHI/Private/D3D12DescriptorCache.cpp b/Engine/Source/Runtime/D3D12RHI/Private/D3D12DescriptorCache.cpp
--- a/Engine/Source/Runtime/D3D12RHI/Private/D3D12DescriptorCache.cpp
+++ b/Engine/Source/Runtime/D3D12RHI/Private/D3D12DescriptorCache.cpp
@@ -786,16 +786,12 @@ void FD3D12DescriptorCache::SetExplicitDescriptorCache(FD3D12ExplicitDescriptorC
 	if (!ViewHeapToSet)
 	{
 		check(ExplicitDescriptorCache.ViewHeap.GetParentDevice() == GetParentDevice());
-
-		ExplicitDescriptorCache.ViewHeap.UpdateSyncPoint();
 		ViewHeapToSet = ExplicitDescriptorCache.ViewHeap.D3D12Heap;
 	}
 
 	if (!SamplerHeapToSet)
 	{
 		check(ExplicitDescriptorCache.SamplerHeap.GetParentDevice() == GetParentDevice());
-
-		ExplicitDescriptorCache.SamplerHeap.UpdateSyncPoint();
 		SamplerHeapToSet = ExplicitDescriptorCache.SamplerHeap.D3D12Heap;
 	}
 
diff --git a/Engine/Source/Runtime/D3D12RHI/Private/D3D12ExplicitDescriptorCache.cpp b/Engine/Source/Runtime/D3D12RHI/Private/D3D12ExplicitDescriptorCache.cpp
--- a/Engine/Source/Runtime/D3D12RHI/Private/D3D12ExplicitDescriptorCache.cpp
+++ b/Engine/Source/Runtime/D3D12RHI/Private/D3D12ExplicitDescriptorCache.cpp
@@ -7,6 +7,7 @@
 #include "D3D12RHIPrivate.h"
 #include "D3D12Stats.h"
 #include "HAL/ConsoleManager.h"
+#include "HAL/PlatformTime.h"
 #include "Hash/xxhash.h"
 
 // Whether to compare the full descriptor table on cache lookup or only use FXxHash64 digest.
@@ -33,10 +34,10 @@ static FAutoConsoleVariableRef CVarD3D12ExplicitViewDescriptorHeapSize(
 
 FD3D12ExplicitDescriptorHeapCache::~FD3D12ExplicitDescriptorHeapCache()
 {
-	check(AllocatedEntries == 0);
+	check(NumAllocatedEntries == 0);
 
 	FScopeLock Lock(&CriticalSection);
-	for (const Entry& It : Entries)
+	for (const FEntry& It : FreeList)
 	{
 		if (It.Type == D3D12_DESCRIPTOR_HEAP_TYPE_CBV_SRV_UAV)
 		{
@@ -51,20 +52,33 @@ FD3D12ExplicitDescriptorHeapCache::~FD3D12ExplicitDescriptorHeapCache()
 
 		It.Heap->Release();
 	}
-	Entries.Empty();
+	FreeList.Empty();
 }
 
-void FD3D12ExplicitDescriptorHeapCache::ReleaseHeap(FD3D12ExplicitDescriptorHeapCache::Entry& Entry)
+void FD3D12ExplicitDescriptorHeapCache::DeferredReleaseHeap(FD3D12ExplicitDescriptorHeapCache::FEntry&& Entry)
+{
+	FD3D12DynamicRHI::GetD3DRHI()->DeferredDelete(
+		[Cache = this, Entry = MoveTemp(Entry)]() mutable
+		{
+			Cache->ReleaseHeap(MoveTemp(Entry));
+		});
+}
+
+void FD3D12ExplicitDescriptorHeapCache::ReleaseHeap(FD3D12ExplicitDescriptorHeapCache::FEntry&& Entry)
 {
 	FScopeLock Lock(&CriticalSection);
 
-	Entries.Add(Entry);
+	check(NumAllocatedEntries != 0);
 
-	check(AllocatedEntries != 0);
-	--AllocatedEntries;
+	Entry.LastUsedFrame = GetParentDevice()->GetParentAdapter()->GetFrameFence().GetNextFenceToSignal();
+	Entry.LastUsedTime = FPlatformTime::Seconds();
+
+	FreeList.Add(MoveTemp(Entry));
+
+	--NumAllocatedEntries;
 }
 
-FD3D12ExplicitDescriptorHeapCache::Entry FD3D12ExplicitDescriptorHeapCache::AllocateHeap(D3D12_DESCRIPTOR_HEAP_TYPE Type, uint32 NumDescriptors)
+FD3D12ExplicitDescriptorHeapCache::FEntry FD3D12ExplicitDescriptorHeapCache::AllocateHeap(D3D12_DESCRIPTOR_HEAP_TYPE Type, uint32 NumDescriptors)
 {
 	FScopeLock Lock(&CriticalSection);
 
@@ -72,22 +86,17 @@ FD3D12ExplicitDescriptorHeapCache::Entry FD3D12ExplicitDescriptorHeapCache::Allo
 	const uint32 MaxDescriptors = (Type == D3D12_DESCRIPTOR_HEAP_TYPE_SAMPLER) ? D3D12_MAX_SHADER_VISIBLE_SAMPLER_HEAP_SIZE : D3D12_MAX_SHADER_VISIBLE_DESCRIPTOR_HEAP_SIZE_TIER_1;
 	NumDescriptors = FMath::Clamp(FMath::RoundUpToPowerOfTwo(NumDescriptors), 0, MaxDescriptors);
 
-	++AllocatedEntries;
-
-	Entry Result = {};
+	++NumAllocatedEntries;
 
-	FD3D12ManualFence& Fence = GetParentDevice()->GetParentAdapter()->GetFrameFence();
-	const uint64 CompletedFenceValue = Fence.GetCompletedFenceValue(/* bUpdateCachedFenceValue */ false);
+	FEntry Result = {};
 
-	for (int32 EntryIndex = 0; EntryIndex < Entries.Num(); ++EntryIndex)
+	for (int32 EntryIndex = 0; EntryIndex < FreeList.Num(); ++EntryIndex)
 	{
-		const Entry& It = Entries[EntryIndex];
-		if (It.Type == Type && It.NumDescriptors >= NumDescriptors && It.FenceValue <= CompletedFenceValue)
+		const FEntry& It = FreeList[EntryIndex];
+		if (It.Type == Type && It.NumDescriptors >= NumDescriptors)
 		{
 			Result = It;
-
-			Entries[EntryIndex] = Entries.Last();
-			Entries.Pop();
+			FreeList.RemoveAtSwap(EntryIndex, EAllowShrinking::No);
 
 			return Result;
 		}
@@ -95,7 +104,8 @@ FD3D12ExplicitDescriptorHeapCache::Entry FD3D12ExplicitDescriptorHeapCache::Allo
 
 	// Compatible heap was not found in cache, so create a new one.
 
-	ReleaseStaleEntries(100, CompletedFenceValue); // Release heaps that were not used for 100 frames before allocating new.
+	// Release heaps that were not used for a while before allocating new.
+	ReleaseStaleEntries(100 /*MaxAgeInFrames*/, 5.0f /*MaxAgeInSeconds*/);
 
 	D3D12_DESCRIPTOR_HEAP_DESC Desc = {};
 
@@ -130,13 +140,18 @@ FD3D12ExplicitDescriptorHeapCache::Entry FD3D12ExplicitDescriptorHeapCache::Allo
 	return Result;
 }
 
-void FD3D12ExplicitDescriptorHeapCache::ReleaseStaleEntries(uint32 MaxAge, uint64 CompletedFenceValue)
+void FD3D12ExplicitDescriptorHeapCache::ReleaseStaleEntries(uint32 MaxAgeInFrames, float MaxAgeInSeconds)
 {
+	const uint64 CurrentFrame = GetParentDevice()->GetParentAdapter()->GetFrameFence().GetNextFenceToSignal();
+	const double CurrentTime = FPlatformTime::Seconds();
+
 	int32 EntryIndex = 0;
-	while (EntryIndex < Entries.Num())
+	while (EntryIndex < FreeList.Num())
 	{
-		Entry& It = Entries[EntryIndex];
-		if ((It.FenceValue + MaxAge) <= CompletedFenceValue)
+		FEntry& It = FreeList[EntryIndex];
+
+		if ((It.LastUsedFrame + MaxAgeInFrames) <= CurrentFrame
+			|| (It.LastUsedTime + MaxAgeInSeconds) <= CurrentTime)
 		{
 			if (It.Type == D3D12_DESCRIPTOR_HEAP_TYPE_CBV_SRV_UAV)
 			{
@@ -149,10 +164,9 @@ void FD3D12ExplicitDescriptorHeapCache::ReleaseStaleEntries(uint32 MaxAge, uint6
 				DEC_DWORD_STAT_BY(STAT_ExplicitSamplerDescriptors, It.NumDescriptors);
 			}
 
-			It.Heap->Release();
+			FD3D12DynamicRHI::GetD3DRHI()->DeferredDelete(It.Heap);
 
-			Entries[EntryIndex] = Entries.Last();
-			Entries.Pop(EAllowShrinking::No);
+			FreeList.RemoveAtSwap(EntryIndex, EAllowShrinking::No);
 		}
 		else
 		{
@@ -161,15 +175,15 @@ void FD3D12ExplicitDescriptorHeapCache::ReleaseStaleEntries(uint32 MaxAge, uint6
 	}
 }
 
-void FD3D12ExplicitDescriptorHeapCache::Flush()
+void FD3D12ExplicitDescriptorHeapCache::FlushFreeList()
 {
 	FScopeLock Lock(&CriticalSection);
 
-	for (const Entry& It : Entries)
+	for (const FEntry& It : FreeList)
 	{
 		FD3D12DynamicRHI::GetD3DRHI()->DeferredDelete(It.Heap);
 	}
-	Entries.Empty();
+	FreeList.Empty();
 }
 
 ///
@@ -178,7 +192,7 @@ FD3D12ExplicitDescriptorHeap::~FD3D12ExplicitDescriptorHeap()
 {
 	if (D3D12Heap)
 	{
-		GetParentDevice()->GetExplicitDescriptorHeapCache()->ReleaseHeap(HeapCacheEntry);
+		GetParentDevice()->GetExplicitDescriptorHeapCache()->DeferredReleaseHeap(MoveTemp(HeapCacheEntry));
 	}
 }
 
@@ -288,12 +302,6 @@ D3D12_GPU_DESCRIPTOR_HANDLE FD3D12ExplicitDescriptorHeap::GetDescriptorGPU(uint3
 	return Result;
 }
 
-void FD3D12ExplicitDescriptorHeap::UpdateSyncPoint()
-{
-	const FD3D12ManualFence& Fence = GetParentDevice()->GetParentAdapter()->GetFrameFence();
-	HeapCacheEntry.FenceValue = FMath::Max(HeapCacheEntry.FenceValue, Fence.GetNextFenceToSignal());
-}
-
 ///
 
 void FD3D12ExplicitDescriptorCache::Init(uint32 NumViewDescriptors, uint32 NumSamplerDescriptors, ERHIBindlessConfiguration BindlessConfig)
diff --git a/Engine/Source/Runtime/D3D12RHI/Private/D3D12ExplicitDescriptorCache.h b/Engine/Source/Runtime/D3D12RHI/Private/D3D12ExplicitDescriptorCache.h
--- a/Engine/Source/Runtime/D3D12RHI/Private/D3D12ExplicitDescriptorCache.h
+++ b/Engine/Source/Runtime/D3D12RHI/Private/D3D12ExplicitDescriptorCache.h
@@ -27,12 +27,15 @@ public:
 
 	UE_NONCOPYABLE(FD3D12ExplicitDescriptorHeapCache)
 
-	struct Entry
+	struct FEntry
 	{
 		ID3D12DescriptorHeap* Heap = nullptr;
-		uint64 FenceValue = 0;
 		uint32 NumDescriptors = 0;
 		D3D12_DESCRIPTOR_HEAP_TYPE Type = D3D12_DESCRIPTOR_HEAP_TYPE_NUM_TYPES;
+
+		// Information for stale entry release, updated upon adding the entry to the free list
+		uint64 LastUsedFrame = 0;
+		double LastUsedTime = 0.0;
 	};
 
 	FD3D12ExplicitDescriptorHeapCache(FD3D12Device* Device)
@@ -42,17 +45,21 @@ public:
 
 	~FD3D12ExplicitDescriptorHeapCache();
 
-	void ReleaseHeap(Entry& Entry);
+	FEntry AllocateHeap(D3D12_DESCRIPTOR_HEAP_TYPE Type, uint32 NumDescriptors);
+	void DeferredReleaseHeap(FEntry&& Entry);
+
+	void FlushFreeList();
 
-	Entry AllocateHeap(D3D12_DESCRIPTOR_HEAP_TYPE Type, uint32 NumDescriptors);
+private:
 
-	void ReleaseStaleEntries(uint32 MaxAge, uint64 CompletedFenceValue);
+	void ReleaseHeap(FEntry&& Entry);
 
-	void Flush();
+	// Assumes CriticalSection is already locked
+	void ReleaseStaleEntries(uint32 MaxAgeInFrames, float MaxAgeInSeconds);
 
 	FCriticalSection CriticalSection;
-	TArray<Entry> Entries;
-	uint32 AllocatedEntries = 0;
+	TArray<FEntry> FreeList;
+	uint32 NumAllocatedEntries = 0;
 };
 
 struct FD3D12ExplicitDescriptorHeap : public FD3D12DeviceChild
@@ -79,8 +86,6 @@ struct FD3D12ExplicitDescriptorHeap : public FD3D12DeviceChild
 	D3D12_CPU_DESCRIPTOR_HANDLE GetDescriptorCPU(uint32 Index) const;
 	D3D12_GPU_DESCRIPTOR_HANDLE GetDescriptorGPU(uint32 Index) const;
 
-	void UpdateSyncPoint();
-
 	D3D12_DESCRIPTOR_HEAP_TYPE Type = D3D12_DESCRIPTOR_HEAP_TYPE_NUM_TYPES;
 	ID3D12DescriptorHeap* D3D12Heap = nullptr;
 	uint32 MaxNumDescriptors = 0;
@@ -94,7 +99,7 @@ struct FD3D12ExplicitDescriptorHeap : public FD3D12DeviceChild
 	D3D12_CPU_DESCRIPTOR_HANDLE CPUBase = {};
 	D3D12_GPU_DESCRIPTOR_HANDLE GPUBase = {};
 
-	FD3D12ExplicitDescriptorHeapCache::Entry HeapCacheEntry;
+	FD3D12ExplicitDescriptorHeapCache::FEntry HeapCacheEntry;
 
 	TArray<D3D12_CPU_DESCRIPTOR_HANDLE> Descriptors;
 
