- Use shader binding layout for raytracing global root signature creation and runtime binding
- Add uniform buffers declared in shader binding layout as required symboles for shader gen

diff --git a/Engine/Shaders/Public/ShaderVersion.ush b/Engine/Shaders/Public/ShaderVersion.ush
--- a/Engine/Shaders/Public/ShaderVersion.ush
+++ b/Engine/Shaders/Public/ShaderVersion.ush
@@ -3,4 +3,4 @@
 // Each time the console command is executed it generates a new GUID. As this file is included
 // in Platform.ush (which should be included in any shader) it allows to invalidate the shader DDC.
 // 
-#pragma message("UESHADERMETADATA_VERSION 4BA01FFC-127B-AE42-AAA5-259B8F406FCB")
+#pragma message("UESHADERMETADATA_VERSION EE2C3AC2-33D9-412B-BCEC-245641AB9666")
diff --git a/Engine/Source/Runtime/D3D12RHI/Private/D3D12Adapter.h b/Engine/Source/Runtime/D3D12RHI/Private/D3D12Adapter.h
--- a/Engine/Source/Runtime/D3D12RHI/Private/D3D12Adapter.h
+++ b/Engine/Source/Runtime/D3D12RHI/Private/D3D12Adapter.h
@@ -236,10 +236,11 @@ public:
 	FORCEINLINE FD3D12PipelineStateCache& GetPSOCache() { return PipelineStateCache; }
 
 	const FD3D12RootSignature* GetRootSignature(const FBoundShaderStateInput& BoundShaderState);
-	const FD3D12RootSignature* GetRootSignature(const class FD3D12RayTracingShader* Shader);
 	const FD3D12RootSignature* GetRootSignature(const class FD3D12ComputeShader* Shader);
 	const FD3D12RootSignature* GetRootSignature(const class FD3D12WorkGraphShader* Shader);
-	const FD3D12RootSignature* GetGlobalRayTracingRootSignature();
+
+	const FD3D12RootSignature* GetLocalRootSignature(const class FD3D12RayTracingShader* Shader);
+	const FD3D12RootSignature* GetGlobalRayTracingRootSignature(const FRHIShaderBindingLayout& ShaderBindingLayout);
 
 #if USE_STATIC_ROOT_SIGNATURE
 	FORCEINLINE const FD3D12RootSignature* GetStaticGraphicsWithConstantsRootSignature() const
diff --git a/Engine/Source/Runtime/D3D12RHI/Private/D3D12DescriptorCache.cpp b/Engine/Source/Runtime/D3D12RHI/Private/D3D12DescriptorCache.cpp
--- a/Engine/Source/Runtime/D3D12RHI/Private/D3D12DescriptorCache.cpp
+++ b/Engine/Source/Runtime/D3D12RHI/Private/D3D12DescriptorCache.cpp
@@ -1254,11 +1254,6 @@ uint32 GetTypeHash(const FD3D12SamplerArrayDesc& Key)
 	return FD3D12PipelineStateCache::HashData((void*)Key.SamplerID, Key.Count * sizeof(Key.SamplerID[0]));
 }
 
-uint32 GetTypeHash(const FD3D12QuantizedBoundShaderState& Key)
-{
-	return FD3D12PipelineStateCache::HashData((void*)&Key, sizeof(Key));
-}
-
 uint32 GetTypeHash(const FD3D12UniqueSamplerTable& Table)
 {
 	return FD3D12PipelineStateCache::HashData((void*)Table.Key.SamplerID, Table.Key.Count * sizeof(Table.Key.SamplerID[0]));
diff --git a/Engine/Source/Runtime/D3D12RHI/Private/D3D12RayTracing.cpp b/Engine/Source/Runtime/D3D12RHI/Private/D3D12RayTracing.cpp
--- a/Engine/Source/Runtime/D3D12RHI/Private/D3D12RayTracing.cpp
+++ b/Engine/Source/Runtime/D3D12RHI/Private/D3D12RayTracing.cpp
@@ -1157,7 +1157,7 @@ public:
 		else
 		{
 			// All other shaders (hit groups, miss, callable) use custom root signatures.
-			LocalRootSignature = Shader->RootSignature->GetRootSignature();
+			LocalRootSignature = Shader->LocalRootSignature->GetRootSignature();
 		}
 
 		FKey CacheKey;
@@ -1246,10 +1246,10 @@ public:
 		Cache.Reset();
 	}
 
-	ID3D12RootSignature* GetGlobalRootSignature()
+	ID3D12RootSignature* GetGlobalRootSignature(const FRHIShaderBindingLayout& ShaderBindingLayout)
 	{
 		FD3D12Adapter* Adapter = GetParentAdapter();
-		const FD3D12RootSignature* RootSignature = Adapter->GetGlobalRayTracingRootSignature();
+		const FD3D12RootSignature* RootSignature = Adapter->GetGlobalRayTracingRootSignature(ShaderBindingLayout);
 		return RootSignature->GetRootSignature();
 	}
 
@@ -1971,7 +1971,7 @@ public:
 		TArrayView<FRHIRayTracingShader*> InitializerRayGenShaders = Initializer.GetRayGenTable();
 		TArrayView<FRHIRayTracingShader*> InitializerCallableShaders = Initializer.GetCallableTable();
 				
-		uint32 ShaderBindingLayoutHash = Initializer.ShaderBindingLayout ? Initializer.ShaderBindingLayout->GetHash() : 0;
+		FRHIShaderBindingLayout ShaderBindingLayout = Initializer.ShaderBindingLayout ? *Initializer.ShaderBindingLayout : FRHIShaderBindingLayout();
 
 		const uint32 MaxTotalShaders = InitializerRayGenShaders.Num() + InitializerMissShaders.Num() + InitializerHitGroups.Num() + InitializerCallableShaders.Num();
 		checkf(MaxTotalShaders >= 1, TEXT("Ray tracing pipelines are expected to contain at least one shader"));
@@ -1980,7 +1980,7 @@ public:
 
 		// All raygen shaders must share the same global root signature (this is validated below)
 
-		GlobalRootSignature = PipelineCache->GetGlobalRootSignature();
+		GlobalRootSignature = PipelineCache->GetGlobalRootSignature(ShaderBindingLayout);
 
 		const FD3D12RayTracingPipelineState* BasePipeline = GRHISupportsRayTracingPSOAdditions 
 			? FD3D12DynamicRHI::ResourceCast(Initializer.BasePipeline.GetReference())
@@ -2000,7 +2000,7 @@ public:
 
 		// Helper function to acquire a D3D12_EXISTING_COLLECTION_DESC for a compiled shader via cache
 
-		auto AddShaderCollection = [Device, ShaderBindingLayoutHash, GlobalRootSignature = this->GlobalRootSignature, PipelineCache,
+		auto AddShaderCollection = [Device, ShaderBindingLayoutHash = ShaderBindingLayout.GetHash(), GlobalRootSignature = this->GlobalRootSignature, PipelineCache,
 										&UniqueShaderHashes = this->PipelineShaderHashes, &UniqueShaderCollections, &Initializer, &NumCacheHits, &CompileTime,
 										&CompileCompletionList]
 			(FD3D12RayTracingShader* Shader, FD3D12RayTracingPipelineCache::ECollectionType CollectionType)
@@ -2052,7 +2052,6 @@ public:
 		for (FRHIRayTracingShader* ShaderRHI : InitializerRayGenShaders)
 		{
 			FD3D12RayTracingShader* Shader = FD3D12DynamicRHI::ResourceCast(ShaderRHI);
-			checkf(Shader->RootSignature->GetRootSignature() == GlobalRootSignature, TEXT("All raygen and miss shaders must share the same root signature"));
 			checkf(!Shader->UsesGlobalUniformBuffer(), TEXT("Global uniform buffers are not implemented for ray generation shaders"));
 
 			FD3D12RayTracingPipelineCache::FEntry* ShaderCacheEntry = AddShaderCollection(Shader, FD3D12RayTracingPipelineCache::ECollectionType::RayGen);
@@ -2081,7 +2080,7 @@ public:
 
 			const uint32 ShaderViewDescriptors = Shader->ResourceCounts.NumSRVs + Shader->ResourceCounts.NumUAVs;
 			MaxHitGroupViewDescriptors = FMath::Max(MaxHitGroupViewDescriptors, ShaderViewDescriptors);
-			MaxLocalRootSignatureSize = FMath::Max(MaxLocalRootSignatureSize, Shader->RootSignature->GetTotalRootSignatureSizeInBytes());
+			MaxLocalRootSignatureSize = FMath::Max(MaxLocalRootSignatureSize, Shader->LocalRootSignature->GetTotalRootSignatureSizeInBytes());
 
 			FD3D12RayTracingPipelineCache::FEntry* ShaderCacheEntry = AddShaderCollection(Shader, FD3D12RayTracingPipelineCache::ECollectionType::Miss);
 
@@ -2103,7 +2102,7 @@ public:
 
 			const uint32 ShaderViewDescriptors = Shader->ResourceCounts.NumSRVs + Shader->ResourceCounts.NumUAVs;
 			MaxHitGroupViewDescriptors = FMath::Max(MaxHitGroupViewDescriptors, ShaderViewDescriptors);
-			MaxLocalRootSignatureSize = FMath::Max(MaxLocalRootSignatureSize, Shader->RootSignature->GetTotalRootSignatureSizeInBytes());
+			MaxLocalRootSignatureSize = FMath::Max(MaxLocalRootSignatureSize, Shader->LocalRootSignature->GetTotalRootSignatureSizeInBytes());
 
 			FD3D12RayTracingPipelineCache::FEntry* ShaderCacheEntry = AddShaderCollection(Shader, FD3D12RayTracingPipelineCache::ECollectionType::HitGroup);
 
@@ -2126,7 +2125,7 @@ public:
 
 			const uint32 ShaderViewDescriptors = Shader->ResourceCounts.NumSRVs + Shader->ResourceCounts.NumUAVs;
 			MaxHitGroupViewDescriptors = FMath::Max(MaxHitGroupViewDescriptors, ShaderViewDescriptors);
-			MaxLocalRootSignatureSize = FMath::Max(MaxLocalRootSignatureSize, Shader->RootSignature->GetTotalRootSignatureSizeInBytes());
+			MaxLocalRootSignatureSize = FMath::Max(MaxLocalRootSignatureSize, Shader->LocalRootSignature->GetTotalRootSignatureSizeInBytes());
 
 			FD3D12RayTracingPipelineCache::FEntry* ShaderCacheEntry = AddShaderCollection(Shader, FD3D12RayTracingPipelineCache::ECollectionType::Callable);
 
@@ -4276,6 +4275,7 @@ struct FD3D12RayTracingLocalResourceBinder
 template <typename ResourceBinderType>
 static bool SetRayTracingShaderResources(
 	const FD3D12RayTracingShader* Shader,
+	const FD3D12RootSignature* RootSignature,
 	uint32 InNumTextures, FRHITexture* const* Textures,
 	uint32 InNumSRVs, FRHIShaderResourceView* const* SRVs,
 	uint32 InNumUniformBuffers, FRHIUniformBuffer* const* UniformBuffers,
@@ -4284,7 +4284,7 @@ static bool SetRayTracingShaderResources(
 	uint32 InLooseParameterDataSize, const void* InLooseParameterData,
 	ResourceBinderType& Binder)
 {
-	const FD3D12RootSignature* RootSignature = Shader->RootSignature;
+	check(Shader && RootSignature);
 
 	struct FBindings
 	{
@@ -4598,6 +4598,7 @@ static bool SetRayTracingShaderResources(
 template <typename ResourceBinderType>
 static bool SetRayTracingShaderResources(
 	const FD3D12RayTracingShader* Shader,
+	const FD3D12RootSignature* RootSignature,
 	const FRayTracingShaderBindings& ResourceBindings,
 	ResourceBinderType& Binder)
 {
@@ -4616,6 +4617,7 @@ static bool SetRayTracingShaderResources(
 
 	return SetRayTracingShaderResources(
 		Shader,
+		RootSignature,
 		UE_ARRAY_COUNT(ResourceBindings.Textures), ResourceBindings.Textures,
 		UE_ARRAY_COUNT(ResourceBindings.SRVs), ResourceBindings.SRVs,
 		UE_ARRAY_COUNT(ResourceBindings.UniformBuffers), ResourceBindings.UniformBuffers,
@@ -4700,11 +4702,15 @@ static void DispatchRays(FD3D12CommandContext& CommandContext,
 
 	CommandContext.GraphicsCommandList()->SetComputeRootSignature(Pipeline->GlobalRootSignature);
 
-	// Bind diagnostic buffer to allow asserts in ray generation shaders
-	CommandContext.BindDiagnosticBuffer(Adapter->GetGlobalRayTracingRootSignature(), ED3D12PipelineType::Compute);
-
 	FD3D12RayTracingShader* RayGenShader = Pipeline->RayGenShaders.Shaders[RayGenShaderIndex];
 
+	FRHIShaderBindingLayout ShaderBindingLayout = GlobalBindings.StaticUniformBuffers.GetShaderBindingLayout() ? *GlobalBindings.StaticUniformBuffers.GetShaderBindingLayout() : FRHIShaderBindingLayout();
+	check(RayGenShader->ShaderBindingLayoutHash == ShaderBindingLayout.GetHash());
+
+	// Bind diagnostic buffer to allow asserts in ray generation shaders
+	const FD3D12RootSignature* GlobalRTRootSignature = Adapter->GetGlobalRayTracingRootSignature(ShaderBindingLayout);
+	CommandContext.BindDiagnosticBuffer(GlobalRTRootSignature, ED3D12PipelineType::Compute);
+
 	bool bResourcesBound = false;
 	if (OptShaderTable && OptShaderTable->DescriptorCache)
 	{
@@ -4717,7 +4723,7 @@ static void DispatchRays(FD3D12CommandContext& CommandContext,
 		CommandContext.SetExplicitDescriptorCache(*DescriptorCache);
 
 		FD3D12RayTracingGlobalResourceBinder ResourceBinder(CommandContext, *DescriptorCache);
-		bResourcesBound = SetRayTracingShaderResources(RayGenShader, GlobalBindings, ResourceBinder);
+		bResourcesBound = SetRayTracingShaderResources(RayGenShader, GlobalRTRootSignature, GlobalBindings, ResourceBinder);
 
 		OptShaderTable->UpdateResidency(CommandContext);
 	}
@@ -4727,7 +4733,28 @@ static void DispatchRays(FD3D12CommandContext& CommandContext,
 		TransientDescriptorCache.Init(MAX_SRVS + MAX_UAVS, MAX_SAMPLERS, ERHIBindlessConfiguration::RayTracingShaders);
 		CommandContext.SetExplicitDescriptorCache(TransientDescriptorCache);
 		FD3D12RayTracingGlobalResourceBinder ResourceBinder(CommandContext, TransientDescriptorCache);
-		bResourcesBound = SetRayTracingShaderResources(RayGenShader, GlobalBindings, ResourceBinder);
+		bResourcesBound = SetRayTracingShaderResources(RayGenShader, GlobalRTRootSignature, GlobalBindings, ResourceBinder);
+	}
+
+	int8 StaticShaderBindingSlot = GlobalRTRootSignature->GetStaticShaderBindingSlot();
+	if (StaticShaderBindingSlot >= 0)
+	{
+		check(GlobalBindings.StaticUniformBuffers.GetShaderBindingLayout());
+		check(GlobalBindings.StaticUniformBuffers.GetShaderBindingLayout()->GetNumUniformBufferEntries() == GlobalBindings.StaticUniformBuffers.GetUniformBufferCount());
+		for (int32 Index = 0; Index < GlobalBindings.StaticUniformBuffers.GetUniformBufferCount(); ++Index)
+		{
+			int32 SlotIndex = Index;
+			FRHIUniformBuffer* UniformBuffer = GlobalBindings.StaticUniformBuffers.GetUniformBuffer(Index);
+			if (UniformBuffer)
+			{
+				FD3D12UniformBuffer* D3D12UniformBuffer = FD3D12CommandContext::RetrieveObject<FD3D12UniformBuffer>(UniformBuffer, 0);//GpuIndex);
+				if (D3D12UniformBuffer->ResourceLocation.GetGPUVirtualAddress())
+				{
+					const FD3D12ResourceLocation& ResourceLocation = D3D12UniformBuffer->ResourceLocation;
+					CommandContext.GraphicsCommandList()->SetComputeRootConstantBufferView(StaticShaderBindingSlot + SlotIndex, ResourceLocation.GetGPUVirtualAddress());
+				}
+			}
+		}
 	}
 
 	if (bResourcesBound)
@@ -4897,8 +4924,8 @@ static void SetRayTracingHitGroup(
 		}
 	}
 
-	FD3D12RayTracingLocalResourceBinder ResourceBinder(*Device, *ShaderTable, *(Shader->RootSignature), RecordIndex, WorkerIndex, ERayTracingBindingType::HitGroup);
-	const bool bResourcesBound = SetRayTracingShaderResources(Shader,
+	FD3D12RayTracingLocalResourceBinder ResourceBinder(*Device, *ShaderTable, *(Shader->LocalRootSignature), RecordIndex, WorkerIndex, ERayTracingBindingType::HitGroup);
+	const bool bResourcesBound = SetRayTracingShaderResources(Shader, Shader->LocalRootSignature,		
 		0, nullptr, // Textures
 		0, nullptr, // SRVs
 		NumUniformBuffers, UniformBuffers,
@@ -4938,8 +4965,8 @@ static void SetRayTracingCallableShader(
 	{
 		const FD3D12RayTracingShader* Shader = Pipeline->CallableShaders.Shaders[ShaderIndexInPipeline];
 
-		FD3D12RayTracingLocalResourceBinder ResourceBinder(*Device, *ShaderTable, *(Shader->RootSignature), RecordIndex, WorkerIndex, ERayTracingBindingType::CallableShader);
-		const bool bResourcesBound = SetRayTracingShaderResources(Shader,
+		FD3D12RayTracingLocalResourceBinder ResourceBinder(*Device, *ShaderTable, *(Shader->LocalRootSignature), RecordIndex, WorkerIndex, ERayTracingBindingType::CallableShader);
+		const bool bResourcesBound = SetRayTracingShaderResources(Shader, Shader->LocalRootSignature,
 			0, nullptr, // Textures
 			0, nullptr, // SRVs
 			NumUniformBuffers, UniformBuffers,
@@ -4973,8 +5000,8 @@ static void SetRayTracingMissShader(
 
 	const FD3D12RayTracingShader* Shader = Pipeline->MissShaders.Shaders[ShaderIndexInPipeline];
 
-	FD3D12RayTracingLocalResourceBinder ResourceBinder(*Device, *ShaderTable, *(Shader->RootSignature), RecordIndex, WorkerIndex, ERayTracingBindingType::MissShader);
-	const bool bResourcesBound = SetRayTracingShaderResources(Shader,
+	FD3D12RayTracingLocalResourceBinder ResourceBinder(*Device, *ShaderTable, *(Shader->LocalRootSignature), RecordIndex, WorkerIndex, ERayTracingBindingType::MissShader);
+	const bool bResourcesBound = SetRayTracingShaderResources(Shader, Shader->LocalRootSignature,
 		0, nullptr, // Textures
 		0, nullptr, // SRVs
 		NumUniformBuffers, UniformBuffers,
diff --git a/Engine/Source/Runtime/D3D12RHI/Private/D3D12RootSignature.cpp b/Engine/Source/Runtime/D3D12RHI/Private/D3D12RootSignature.cpp
--- a/Engine/Source/Runtime/D3D12RHI/Private/D3D12RootSignature.cpp
+++ b/Engine/Source/Runtime/D3D12RHI/Private/D3D12RootSignature.cpp
@@ -325,17 +325,21 @@ FD3D12RootSignatureDesc::FD3D12RootSignatureDesc(const FD3D12QuantizedBoundShade
 			Flags |= D3D12_ROOT_SIGNATURE_FLAG_ALLOW_INPUT_ASSEMBLER_INPUT_LAYOUT;
 		}
 
-		for (uint32 ShaderVisibilityIndex = 0; ShaderVisibilityIndex < UE_ARRAY_COUNT(ShaderVisibilityPriorityOrder); ShaderVisibilityIndex++)
+		// don't set deny flags when static shader resource table is used
+		if (QBSS.ShaderBindingLayout.GetNumUniformBufferEntries() == 0)
 		{
-			const EShaderVisibility Visibility = ShaderVisibilityPriorityOrder[ShaderVisibilityIndex];
-			const FShaderRegisterCounts& Shader = QBSS.RegisterCounts[Visibility];
-			if ((Shader.ShaderResourceCount == 0) &&
-				(Shader.ConstantBufferCount == 0) &&
-				(Shader.UnorderedAccessCount == 0) &&
-				(Shader.SamplerCount == 0))
+			for (uint32 ShaderVisibilityIndex = 0; ShaderVisibilityIndex < UE_ARRAY_COUNT(ShaderVisibilityPriorityOrder); ShaderVisibilityIndex++)
 			{
-				// This shader stage doesn't use any descriptors, deny access to the shader stage in the root signature.
-				Flags = (Flags | GetD3D12RootSignatureDenyFlag(Visibility));
+				const EShaderVisibility Visibility = ShaderVisibilityPriorityOrder[ShaderVisibilityIndex];
+				const FShaderRegisterCounts& Shader = QBSS.RegisterCounts[Visibility];
+				if ((Shader.ShaderResourceCount == 0) &&
+					(Shader.ConstantBufferCount == 0) &&
+					(Shader.UnorderedAccessCount == 0) &&
+					(Shader.SamplerCount == 0))
+				{
+					// This shader stage doesn't use any descriptors, deny access to the shader stage in the root signature.
+					Flags = (Flags | GetD3D12RootSignatureDenyFlag(Visibility));
+				}
 			}
 		}
 	}
@@ -349,6 +353,25 @@ FD3D12RootSignatureDesc::FD3D12RootSignatureDesc(const FD3D12QuantizedBoundShade
 		RootParametersSize += RootDescriptorCost;
 	}
 #endif
+	
+	// Add all the static defined uniform buffers
+	if (QBSS.ShaderBindingLayout.GetNumUniformBufferEntries() > 0)
+	{
+		StaticShaderBindingSlot = int8(RootParameterCount);
+		StaticShaderBindingCount = 0;
+		for (uint32 Index = 0; Index < QBSS.ShaderBindingLayout.GetNumUniformBufferEntries(); ++Index)
+		{
+			const FRHIUniformBufferShaderBindingLayout& UniformBufferSBLayout = QBSS.ShaderBindingLayout.GetUniformBufferEntry(Index);
+			
+			check(RootParameterCount < MaxRootParameters);
+			check(UniformBufferSBLayout.RegisterSpace > 0);
+			TableSlots[RootParameterCount].InitAsConstantBufferView(UniformBufferSBLayout.CBVResourceIndex, UniformBufferSBLayout.RegisterSpace, CBVRootDescriptorFlags, D3D12_SHADER_VISIBILITY_ALL);
+			StaticShaderBindingCount = FMath::Max(StaticShaderBindingCount, int8(UniformBufferSBLayout.CBVResourceIndex + 1));
+
+			RootParameterCount++;
+			RootParametersSize += RootDescriptorCost;
+		}
+	}
 
 	if (QBSS.bUseRootConstants)
 	{
@@ -477,6 +500,8 @@ void FD3D12RootSignature::Init(const FD3D12QuantizedBoundShaderState& InQBSS)
 	FD3D12RootSignatureDesc Desc(InQBSS, ResourceBindingTier);
 
 	RootConstantsSlot = Desc.GetRootConstantsSlot();
+	StaticShaderBindingSlot = Desc.GetStaticShaderBindingSlot();
+	StaticShaderBindingCount = Desc.GetStaticShaderBindingCount();
 	DiagnosticBufferSlot = Desc.GetDiagnosticBufferSlot();
 
 	uint32 BindingSpace = 0; // Default binding space for D3D 11 & 12 shaders
diff --git a/Engine/Source/Runtime/D3D12RHI/Private/D3D12RootSignature.h b/Engine/Source/Runtime/D3D12RHI/Private/D3D12RootSignature.h
--- a/Engine/Source/Runtime/D3D12RHI/Private/D3D12RootSignature.h
+++ b/Engine/Source/Runtime/D3D12RHI/Private/D3D12RootSignature.h
@@ -54,12 +54,16 @@ public:
 
 	inline int8 GetRootConstantsSlot() const { return RootConstantsSlot; }
 	inline int8 GetDiagnosticBufferSlot() const { return DiagnosticBufferSlot; }
+	inline int8 GetStaticShaderBindingSlot() const { return StaticShaderBindingSlot; }
+	inline int8 GetStaticShaderBindingCount() const { return StaticShaderBindingCount; }
 
 private:
 
 	uint32 RootParametersSize;	// The size of all root parameters in the root signature. Size in DWORDs, the limit is 64.
 	int8 RootConstantsSlot = -1;
 	int8 DiagnosticBufferSlot = -1;
+	int8 StaticShaderBindingSlot = -1;
+	int8 StaticShaderBindingCount = -1;
 	CD3DX12_ROOT_PARAMETER1 TableSlots[MaxRootParameters];
 	CD3DX12_DESCRIPTOR_RANGE1 DescriptorRanges[MaxRootParameters];
 	CD3DX12_VERSIONED_ROOT_SIGNATURE_DESC RootDesc;
@@ -213,6 +217,8 @@ public:
 	uint32 GetTotalRootSignatureSizeInBytes() const { return 4 * TotalRootSignatureSizeInDWORDs; }
 
 	inline int8 GetRootConstantsSlot() const { return RootConstantsSlot; }
+	inline int8 GetStaticShaderBindingSlot() const { return StaticShaderBindingSlot; }
+	inline int8 GetStaticShaderBindingCount() const { return StaticShaderBindingCount; }
 
 	// Returns root parameter slot for the internal shader diagnostic buffer (used for asserts, etc.) or -1 if not available.
 	inline int8 GetDiagnosticBufferSlot() const { return DiagnosticBufferSlot; }
@@ -461,6 +467,8 @@ private:
 	uint8 TotalRootSignatureSizeInDWORDs = 0;
 	int8 RootConstantsSlot = -1;
 	int8 DiagnosticBufferSlot = -1;
+	int8 StaticShaderBindingSlot = -1;
+	int8 StaticShaderBindingCount = -1;
 
 	uint8 bHasUAVs : 1;
 	uint8 bHasSRVs : 1;
diff --git a/Engine/Source/Runtime/D3D12RHI/Private/D3D12Shader.h b/Engine/Source/Runtime/D3D12RHI/Private/D3D12Shader.h
--- a/Engine/Source/Runtime/D3D12RHI/Private/D3D12Shader.h
+++ b/Engine/Source/Runtime/D3D12RHI/Private/D3D12Shader.h
@@ -175,7 +175,7 @@ class FD3D12RayTracingShader : public FRHIRayTracingShader, public FD3D12ShaderD
 public:
 	explicit FD3D12RayTracingShader(EShaderFrequency InFrequency) : FRHIRayTracingShader(InFrequency) {}
 
-	const FD3D12RootSignature* RootSignature = nullptr;
+	const FD3D12RootSignature* LocalRootSignature = nullptr;
 
 	/** The shader's DXIL entrypoint & base export name for DXR (required for RTPSO creation) */
 	FString EntryPoint; // Primary entry point for all ray tracing shaders. Assumed to be closest hit shader for SF_RayHitGroup.
diff --git a/Engine/Source/Runtime/D3D12RHI/Private/D3D12Shaders.cpp b/Engine/Source/Runtime/D3D12RHI/Private/D3D12Shaders.cpp
--- a/Engine/Source/Runtime/D3D12RHI/Private/D3D12Shaders.cpp
+++ b/Engine/Source/Runtime/D3D12RHI/Private/D3D12Shaders.cpp
@@ -247,9 +247,11 @@ FRayTracingShaderRHIRef FD3D12DynamicRHI::RHICreateRayTracingShader(TArrayView<c
 	};
 
 	FD3D12RayTracingShader* Shader = InitStandardShaderWithCustomSerialization(new FD3D12RayTracingShader(ShaderFrequency), Code, CustomSerialization);
-	if (Shader)
+
+	// Setup local root signature (RayGen only has global root signature)
+	if (Shader && ShaderFrequency != SF_RayGen)
 	{
-		Shader->RootSignature = GetAdapter().GetRootSignature(Shader);
+		Shader->LocalRootSignature = GetAdapter().GetLocalRootSignature(Shader);
 	}
 
 	return Shader;
diff --git a/Engine/Source/Runtime/D3D12RHI/Private/D3D12Util.cpp b/Engine/Source/Runtime/D3D12RHI/Private/D3D12Util.cpp
--- a/Engine/Source/Runtime/D3D12RHI/Private/D3D12Util.cpp
+++ b/Engine/Source/Runtime/D3D12RHI/Private/D3D12Util.cpp
@@ -1347,7 +1347,7 @@ const FD3D12RootSignature* FD3D12Adapter::GetRootSignature(const FD3D12WorkGraph
 
 #if D3D12_RHI_RAYTRACING
 
-const FD3D12RootSignature* FD3D12Adapter::GetGlobalRayTracingRootSignature()
+const FD3D12RootSignature* FD3D12Adapter::GetGlobalRayTracingRootSignature(const FRHIShaderBindingLayout& ShaderBindingLayout)
 {
 #if USE_STATIC_ROOT_SIGNATURE
 
@@ -1358,6 +1358,7 @@ const FD3D12RootSignature* FD3D12Adapter::GetGlobalRayTracingRootSignature()
 	FD3D12QuantizedBoundShaderState QBSS{};
 	FShaderRegisterCounts& QBSSRegisterCounts = QBSS.RegisterCounts[SV_All];
 
+	QBSS.ShaderBindingLayout = ShaderBindingLayout;
 	QBSS.RootSignatureType = RS_RayTracingGlobal;
 	QBSS.bUseDiagnosticBuffer = true;
 
@@ -1376,7 +1377,7 @@ const FD3D12RootSignature* FD3D12Adapter::GetGlobalRayTracingRootSignature()
 #endif //! USE_STATIC_ROOT_SIGNATURE
 }
 
-const FD3D12RootSignature* FD3D12Adapter::GetRootSignature(const FD3D12RayTracingShader* RayTracingShader)
+const FD3D12RootSignature* FD3D12Adapter::GetLocalRootSignature(const FD3D12RayTracingShader* RayTracingShader)
 {
 #if USE_STATIC_ROOT_SIGNATURE
 
@@ -1404,9 +1405,10 @@ const FD3D12RootSignature* FD3D12Adapter::GetRootSignature(const FD3D12RayTracin
 	switch (RayTracingShader->GetFrequency())
 	{
 	case SF_RayGen:
-		// Shared conservative root signature layout is used for all raygen and miss shaders.
-		return GetGlobalRayTracingRootSignature();
-
+		// Ray gen only uses global root signature and needs the RHIShaderBindingLayout which is provided through the RTPSO initializer 
+		// and verified against hash stored in the RHIShader data
+		checkNoEntry();
+		break;
 	case SF_RayHitGroup:
 	case SF_RayCallable:
 	case SF_RayMiss:
diff --git a/Engine/Source/Runtime/D3D12RHI/Private/D3D12Util.h b/Engine/Source/Runtime/D3D12RHI/Private/D3D12Util.h
--- a/Engine/Source/Runtime/D3D12RHI/Private/D3D12Util.h
+++ b/Engine/Source/Runtime/D3D12RHI/Private/D3D12Util.h
@@ -88,26 +88,64 @@ struct FShaderRegisterCounts
 	uint8 ConstantBufferCount;
 	uint8 ShaderResourceCount;
 	uint8 UnorderedAccessCount;
+	
+	inline bool operator==(const FShaderRegisterCounts& RHS) const
+	{
+		return SamplerCount == RHS.SamplerCount &&
+			ConstantBufferCount == RHS.ConstantBufferCount &&
+			ShaderResourceCount == RHS.ShaderResourceCount &&
+			UnorderedAccessCount == RHS.UnorderedAccessCount;
+	}
+	
+	friend uint32 GetTypeHash(const FShaderRegisterCounts& Counts)
+	{
+		uint32 Hash = GetTypeHash(Counts.SamplerCount);
+		Hash = HashCombineFast(Hash, GetTypeHash(Counts.ConstantBufferCount));
+		Hash = HashCombineFast(Hash, GetTypeHash(Counts.ShaderResourceCount));
+		Hash = HashCombineFast(Hash, GetTypeHash(Counts.UnorderedAccessCount));
+		return Hash;
+	}
 };
 
 struct FD3D12QuantizedBoundShaderState
 {
-	FShaderRegisterCounts RegisterCounts[SV_ShaderVisibilityCount];
+	FRHIShaderBindingLayout ShaderBindingLayout;
+	TStaticArray<FShaderRegisterCounts, SV_ShaderVisibilityCount> RegisterCounts;
 	ERTRootSignatureType RootSignatureType = RS_Raster;
-	uint8 bAllowIAInputLayout : 1;
-	uint8 bNeedsAgsIntrinsicsSpace : 1;
-	uint8 bUseDiagnosticBuffer : 1;
-	uint8 bUseDirectlyIndexedResourceHeap : 1;
-	uint8 bUseDirectlyIndexedSamplerHeap : 1;
-	uint8 bUseRootConstants : 1;
-	uint8 Padding : 2;
+	union
+	{
+		struct
+		{
+			uint8 bAllowIAInputLayout : 1;
+			uint8 bNeedsAgsIntrinsicsSpace : 1;
+			uint8 bUseDiagnosticBuffer : 1;
+			uint8 bUseDirectlyIndexedResourceHeap : 1;
+			uint8 bUseDirectlyIndexedSamplerHeap : 1;
+			uint8 bUseRootConstants : 1;
+			uint8 Padding : 2;
+		};
+		uint8 Flags;
+	};
 
 	inline bool operator==(const FD3D12QuantizedBoundShaderState& RHS) const
 	{
-		return 0 == FMemory::Memcmp(this, &RHS, sizeof(RHS));
+		return ShaderBindingLayout == RHS.ShaderBindingLayout &&
+			RegisterCounts == RHS.RegisterCounts &&
+			RootSignatureType == RHS.RootSignatureType &&
+			Flags == RHS.Flags;
+	}
+	
+	friend uint32 GetTypeHash(const FD3D12QuantizedBoundShaderState& State)
+	{
+		uint32 Hash = GetTypeHash(State.ShaderBindingLayout);
+		for (int32 Index = 0; Index < SV_ShaderVisibilityCount; ++Index)
+		{
+			Hash = HashCombineFast(Hash, GetTypeHash(State.RegisterCounts[Index]));
+		}
+		Hash = HashCombineFast(Hash, GetTypeHash(State.RootSignatureType));
+		Hash = HashCombineFast(Hash, GetTypeHash(State.Flags));
+		return Hash;
 	}
-
-	friend uint32 GetTypeHash(const FD3D12QuantizedBoundShaderState& Key);
 
 	static void InitShaderRegisterCounts(D3D12_RESOURCE_BINDING_TIER ResourceBindingTier, const FShaderCodePackedResourceCounts& Counts, FShaderRegisterCounts& Shader, bool bAllowUAVs = false);
 };
diff --git a/Engine/Source/Runtime/RenderCore/Private/ShaderCore.cpp b/Engine/Source/Runtime/RenderCore/Private/ShaderCore.cpp
--- a/Engine/Source/Runtime/RenderCore/Private/ShaderCore.cpp
+++ b/Engine/Source/Runtime/RenderCore/Private/ShaderCore.cpp
@@ -784,26 +784,33 @@ FThreadSafeSharedAnsiStringPtr FShaderBindingLayout::GetUniformBufferDeclaration
 	if (Entry)
 	{
 		uint8 UniformBufferIndex = Entry->CBVResourceIndex;
-		return UniformBufferDeclarationsAnsi[UniformBufferIndex];
+		return UniformBufferData[UniformBufferIndex].UniformBufferDeclarationAnsi;
 	}
 
 	// use the shared declaration from the metadata object itself (no fixed register or space defined)
 	return ShaderParametersMetadata->GetUniformBufferDeclarationAnsiPtr();
 }
 
-void FShaderBindingLayout::SetUniformBufferDeclarationAnsiPtr(const FShaderParametersMetadata* ShaderParametersMetadata, FThreadSafeSharedAnsiStringPtr UniformBufferDeclarationAnsi)
+void FShaderBindingLayout::SetUniformBufferDeclarationAnsiPtr(const FShaderParametersMetadata* ShaderParametersMetadata, const FString& UniformBufferName, FThreadSafeSharedAnsiStringPtr UniformBufferDeclarationAnsi)
 {
 	uint32 LayoutHash = ShaderParametersMetadata->GetLayout().GetHash();
 	const FRHIUniformBufferShaderBindingLayout* Entry = RHILayout.FindEntry(LayoutHash);
 	check(Entry);
-	UniformBufferDeclarationsAnsi.SetNum(FMath::Max(UniformBufferDeclarationsAnsi.Num(), int32(Entry->CBVResourceIndex + 1)));
-	UniformBufferDeclarationsAnsi[Entry->CBVResourceIndex] = UniformBufferDeclarationAnsi;
+	UniformBufferData.SetNum(FMath::Max(UniformBufferData.Num(), int32(Entry->CBVResourceIndex + 1)));
+	UniformBufferData[Entry->CBVResourceIndex].UniformBufferName = UniformBufferName;
+	UniformBufferData[Entry->CBVResourceIndex].UniformBufferDeclarationAnsi = UniformBufferDeclarationAnsi;
 }
 
 void FShaderBindingLayout::AddRequiredSymbols(TArray<FString>& RequiredSymbols) const
 {
 	// assume only bindless for now so only need to add the CBuffer declares as required symbols
 	check(EnumHasAllFlags(RHILayout.GetFlags(), EShaderBindingLayoutFlags::BindlessResources | EShaderBindingLayoutFlags::BindlessSamplers));
+
+	// Don't remove unused uniform buffers defined in the fixed shader binding layout because they are required to be declared for certain platforms
+	for (const FUniformBufferData& Data : UniformBufferData)
+	{
+		RequiredSymbols.Add(Data.UniformBufferName);
+	}
 }
 
 #endif // WITH_EDITOR
diff --git a/Engine/Source/Runtime/RenderCore/Private/ShaderParameterMetadata.cpp b/Engine/Source/Runtime/RenderCore/Private/ShaderParameterMetadata.cpp
--- a/Engine/Source/Runtime/RenderCore/Private/ShaderParameterMetadata.cpp
+++ b/Engine/Source/Runtime/RenderCore/Private/ShaderParameterMetadata.cpp
@@ -211,7 +211,8 @@ void BuildShaderBindingLayout(TConstArrayView<FShaderParametersMetadata*> Unifor
 		const FRHIUniformBufferShaderBindingLayout* UniformBufferSBLayout = ShaderBindingLayout.RHILayout.FindEntry(ShaderParametersMetaData->GetLayout().GetHash());
 		check(UniformBufferSBLayout);
 
-		FString NewDeclaration = UE::ShaderParameters::CreateUniformBufferShaderDeclaration(ShaderParametersMetaData->GetShaderVariableName(), *ShaderParametersMetaData, UniformBufferSBLayout);
+		FString UniformBufferName(ShaderParametersMetaData->GetShaderVariableName());
+		FString NewDeclaration = UE::ShaderParameters::CreateUniformBufferShaderDeclaration(*UniformBufferName, *ShaderParametersMetaData, UniformBufferSBLayout);
 		check(!NewDeclaration.IsEmpty());
 			
 		// Cache preprocessor friendly copy of uniform buffer declaration
@@ -219,7 +220,7 @@ void BuildShaderBindingLayout(TConstArrayView<FShaderParametersMetadata*> Unifor
 		ShaderConvertAndStripComments(NewDeclaration, *NewDeclarationAnsi);
 		FThreadSafeSharedAnsiStringPtr UniformBufferDeclarationAnsi = MakeShareable(NewDeclarationAnsi);
 
-		ShaderBindingLayout.SetUniformBufferDeclarationAnsiPtr(ShaderParametersMetaData, UniformBufferDeclarationAnsi);
+		ShaderBindingLayout.SetUniformBufferDeclarationAnsiPtr(ShaderParametersMetaData, UniformBufferName, UniformBufferDeclarationAnsi);
 	}
 #endif //WITH_EDITOR
 
diff --git a/Engine/Source/Runtime/RenderCore/Private/ShaderParameters.cpp b/Engine/Source/Runtime/RenderCore/Private/ShaderParameters.cpp
--- a/Engine/Source/Runtime/RenderCore/Private/ShaderParameters.cpp
+++ b/Engine/Source/Runtime/RenderCore/Private/ShaderParameters.cpp
@@ -372,25 +372,48 @@ static FString CreateHLSLUniformBufferDeclaration(const TCHAR* UniformBufferName
 		uint32 HLSLBaseOffset = 0;
 		CreateHLSLUniformBufferStructMembersDeclaration(UniformBufferStruct, UniformBufferName, TEXT(""), TEXT(""), 0, Decl, HLSLBaseOffset);
 
-		// TODO: use UniformBufferSBLayout to generate CB_DEFINITION with fixed register & space
-
-		return FString::Printf(
-			TEXT("#pragma once\n")
-			TEXT("UB_CB_DEFINITION_START(%s)\n")
-			TEXT("%s")
-			TEXT("UB_CB_DEFINITION_END(%s)\n")
-			TEXT("%s")
-			TEXT("UniformBuffer %s\n")
-			TEXT("{\n")
-			TEXT("%s")
-			TEXT("};\n"),
-			UniformBufferName,
-			*Decl.ConstantBufferMembers,
-			UniformBufferName,
-			*Decl.ResourceMembers,
-			UniformBufferName,
-			*Decl.StructMembers
-		);
+		if (UniformBufferSBLayout && UniformBufferSBLayout->RegisterSpace > 0)
+		{	
+			return FString::Printf(
+				TEXT("#pragma once\n")
+				TEXT("UB_STATIC_CB_DEFINITION_START(%s,%d,%d)\n")
+				TEXT("%s")
+				TEXT("UB_CB_DEFINITION_END(%s)\n")
+				TEXT("%s")
+				TEXT("UniformBuffer %s\n")
+				TEXT("{\n")
+				TEXT("%s")
+				TEXT("};\n"),
+				UniformBufferName,
+				UniformBufferSBLayout->CBVResourceIndex,
+				UniformBufferSBLayout->RegisterSpace,
+				*Decl.ConstantBufferMembers,
+				UniformBufferName,
+				*Decl.ResourceMembers,
+				UniformBufferName,
+				*Decl.StructMembers
+			);
+		}
+		else
+		{
+			return FString::Printf(
+				TEXT("#pragma once\n")
+				TEXT("UB_CB_DEFINITION_START(%s)\n")
+				TEXT("%s")
+				TEXT("UB_CB_DEFINITION_END(%s)\n")
+				TEXT("%s")
+				TEXT("UniformBuffer %s\n")
+				TEXT("{\n")
+				TEXT("%s")
+				TEXT("};\n"),
+				UniformBufferName,
+				*Decl.ConstantBufferMembers,
+				UniformBufferName,
+				*Decl.ResourceMembers,
+				UniformBufferName,
+				*Decl.StructMembers
+			);
+		}
 	}
 
 	return FString(TEXT("\n"));
diff --git a/Engine/Source/Runtime/RenderCore/Public/ShaderCore.h b/Engine/Source/Runtime/RenderCore/Public/ShaderCore.h
--- a/Engine/Source/Runtime/RenderCore/Public/ShaderCore.h
+++ b/Engine/Source/Runtime/RenderCore/Public/ShaderCore.h
@@ -436,13 +436,18 @@ public:
 
 #if WITH_EDITOR
 	FThreadSafeSharedAnsiStringPtr GetUniformBufferDeclarationAnsiPtr(const FShaderParametersMetadata* ShaderParametersMetadata) const;
-	void SetUniformBufferDeclarationAnsiPtr(const FShaderParametersMetadata* ShaderParametersMetadata, FThreadSafeSharedAnsiStringPtr UniformBufferDeclarationAnsi);
+	void SetUniformBufferDeclarationAnsiPtr(const FShaderParametersMetadata* ShaderParametersMetadata, const FString& UniformBufferName, FThreadSafeSharedAnsiStringPtr UniformBufferDeclarationAnsi);
 
 	RENDERCORE_API void AddRequiredSymbols(TArray<FString>& RequiredSymbols) const;
 
 protected:
 	
-	TArray<FThreadSafeSharedAnsiStringPtr> UniformBufferDeclarationsAnsi;
+	struct FUniformBufferData
+	{
+		FString	UniformBufferName;
+		FThreadSafeSharedAnsiStringPtr UniformBufferDeclarationAnsi;
+	};
+	TArray<FUniformBufferData> UniformBufferData;
 #endif
 };
 
