Fixing full Bindless on D3D12. The index of the first allocated heap wasn't being saved.

diff --git a/Engine/Source/Runtime/D3D12RHI/Private/D3D12BindlessDescriptors.cpp b/Engine/Source/Runtime/D3D12RHI/Private/D3D12BindlessDescriptors.cpp
--- a/Engine/Source/Runtime/D3D12RHI/Private/D3D12BindlessDescriptors.cpp
+++ b/Engine/Source/Runtime/D3D12RHI/Private/D3D12BindlessDescriptors.cpp
@@ -130,7 +130,7 @@ FD3D12BindlessResourceManager::FD3D12BindlessResourceManager(FD3D12Device* InDev
 	if (GetConfiguration() == ERHIBindlessConfiguration::AllShaders)
 	{
 		// Always allocate a heap when full bindless
-		AddActiveGPUHeap();
+		ActiveGpuHeapIndex = AddActiveGPUHeap();
 	}
 }
 
