D3D12 don't use fast allocator for buffer locks.
 - The fast allocator does not hold a strong reference on the allocated memory, which can lead to use-after-free in buffer locks.

diff --git a/Engine/Source/Runtime/D3D12RHI/Private/D3D12Buffer.cpp b/Engine/Source/Runtime/D3D12RHI/Private/D3D12Buffer.cpp
--- a/Engine/Source/Runtime/D3D12RHI/Private/D3D12Buffer.cpp
+++ b/Engine/Source/Runtime/D3D12RHI/Private/D3D12Buffer.cpp
@@ -582,7 +582,7 @@ void* FD3D12DynamicRHI::LockBuffer(FRHICommandListBase& RHICmdList, FD3D12Buffer
 		else
 		{
 			// If the static buffer is being locked for writing, allocate memory for the contents to be written to.
-			Data = Device->GetDefaultFastAllocator().Allocate(Size, D3D12_TEXTURE_DATA_PLACEMENT_ALIGNMENT, &LockedData.ResourceLocation);
+			Data = Adapter.GetUploadHeapAllocator(Device->GetGPUIndex()).AllocUploadResource(Size, D3D12_TEXTURE_DATA_PLACEMENT_ALIGNMENT, LockedData.ResourceLocation);
 		}
 	}
 
