Improvements to BLAS residency management when TLAS is used in shader.

Context:
- FD3D12RayTracingScene::UpdateResidency() was only called in RHIRayTraceDispatch(...) and only for one FD3D12RayTracingScene object provided as argument.
- This meant that a RGS couldn't use multiple FD3D12RayTracingScene and also not compatible with inline ray tracing.

Changes:
- Optionally provide FRHIRayTracingScene* when creating Buffer SRV.
- Call FD3D12RayTracingScene::UpdateResidency() based on shader parameters instead of explicit RHIRayTraceDispatch(...) argument.

diff --git a/Engine/Source/Runtime/D3D12RHI/Private/D3D12DescriptorCache.cpp b/Engine/Source/Runtime/D3D12RHI/Private/D3D12DescriptorCache.cpp
--- a/Engine/Source/Runtime/D3D12RHI/Private/D3D12DescriptorCache.cpp
+++ b/Engine/Source/Runtime/D3D12RHI/Private/D3D12DescriptorCache.cpp
@@ -6,6 +6,7 @@
 #include "D3D12DescriptorCache.h"
 #include "D3D12ExplicitDescriptorCache.h"
 #include "D3D12RHIPrivate.h"
+#include "D3D12RayTracing.h"
 
 bool FD3D12DescriptorCache::HeapRolledOver(ERHIDescriptorHeapType InHeapType)
 {
@@ -498,6 +499,12 @@ D3D12_GPU_DESCRIPTOR_HANDLE FD3D12DescriptorCache::BuildSRVTable(EShaderFrequenc
 
 			Context.TransitionResource(SRV, State & ValidResourceStates);
 			Context.UpdateResidency(Cache.Resources[ShaderStage][SlotIndex]);
+
+			FD3D12RayTracingScene* RayTracingScene = SRV->GetRayTracingScene();
+			if (RayTracingScene)
+			{
+				RayTracingScene->UpdateResidency(Context);
+			}
 		}
 		else
 		{
diff --git a/Engine/Source/Runtime/D3D12RHI/Private/D3D12RayTracing.cpp b/Engine/Source/Runtime/D3D12RHI/Private/D3D12RayTracing.cpp
--- a/Engine/Source/Runtime/D3D12RHI/Private/D3D12RayTracing.cpp
+++ b/Engine/Source/Runtime/D3D12RHI/Private/D3D12RayTracing.cpp
@@ -4133,6 +4133,11 @@ struct FD3D12RayTracingGlobalResourceBinder
 		CommandContext.TransitionResource(UAV, D3D12_RESOURCE_STATE_UNORDERED_ACCESS);
 	}
 
+	void AddRayTracingSceneReference(FD3D12RayTracingScene* D3D12RayTracingScene)
+	{
+		D3D12RayTracingScene->UpdateResidency(CommandContext);
+	}
+
 	FD3D12Device* GetDevice()
 	{
 		return CommandContext.GetParentDevice();
@@ -4249,6 +4254,11 @@ struct FD3D12RayTracingLocalResourceBinder
 		}
 	}
 
+	void AddRayTracingSceneReference(FD3D12RayTracingScene* D3D12RayTracingScene)
+	{
+		checkf(false, TEXT("Unexpected RayTracingScene reference in local shader bindings"));
+	}
+
 	FD3D12Device* GetDevice()
 	{
 		return &Device;
@@ -4314,6 +4324,7 @@ static bool SetRayTracingShaderResources(
 		uint32 SamplerVersions[MAX_SRVS];
 
 		TArray<FD3D12Resource*, TInlineAllocator<MAX_CBS + MAX_SRVS + MAX_UAVS>> ReferencedResources;
+		TArray<FD3D12RayTracingScene*, TInlineAllocator<1>> ReferencedRayTracingScenes;
 
 		uint64 BoundSRVMask = 0;
 		uint64 BoundCBVMask = 0;
@@ -4348,6 +4359,12 @@ static bool SetRayTracingShaderResources(
 
 			ReferencedResources.Add(SRV->GetResource());
 			Binder.AddResourceTransition(SRV);
+
+			FD3D12RayTracingScene* ReferencedRayTracingScene = SRV->GetRayTracingScene();
+			if (ReferencedRayTracingScene)
+			{
+				ReferencedRayTracingScenes.Add(ReferencedRayTracingScene);
+			}
 		}
 
 		void SetTexture(FRHITexture* RHITexture, uint8 Index)
@@ -4592,6 +4609,11 @@ static bool SetRayTracingShaderResources(
 		Binder.AddResourceReference(Resource);
 	}
 
+	for (FD3D12RayTracingScene* RayTracingScene : Bindings.ReferencedRayTracingScenes)
+	{
+		Binder.AddRayTracingSceneReference(RayTracingScene);
+	}
+
 	return true;
 }
 
@@ -4824,9 +4846,6 @@ void FD3D12CommandContext::RHIRayTraceDispatch(FRHIRayTracingPipelineState* InRa
 
 	const FD3D12RayTracingPipelineState* Pipeline = FD3D12DynamicRHI::ResourceCast(InRayTracingPipelineState);
 
-	FD3D12RayTracingScene* Scene = FD3D12DynamicRHI::ResourceCast(InScene);
-	Scene->UpdateResidency(*this);
-
 	FD3D12RayTracingShaderBindingTable* SBT = FD3D12DynamicRHI::ResourceCast(InSBT);
 	FD3D12RayTracingShaderBindingTableInternal* ShaderTableForDevice = SBT->GetTableForDevice(GetParentDevice());
 	checkf(!ShaderTableForDevice->bIsDirty, TEXT("The shader table contains pending modifications. CommitRayTracingBindings must be called after SetRayTracingBindings"));
@@ -4857,9 +4876,6 @@ void FD3D12CommandContext::RHIRayTraceDispatchIndirect(FRHIRayTracingPipelineSta
 
 	const FD3D12RayTracingPipelineState* Pipeline = FD3D12DynamicRHI::ResourceCast(InRayTracingPipelineState);
 
-	FD3D12RayTracingScene* Scene = FD3D12DynamicRHI::ResourceCast(InScene);
-	Scene->UpdateResidency(*this);
-
 	FD3D12RayTracingShaderBindingTable* SBT = FD3D12DynamicRHI::ResourceCast(InSBT);
 	FD3D12RayTracingShaderBindingTableInternal* ShaderTableForDevice = SBT->GetTableForDevice(GetParentDevice());
 	checkf(!ShaderTableForDevice->bIsDirty, TEXT("The shader table contains pending modifications. CommitRayTracingBindings must be called after SetRayTracingBindings"));
diff --git a/Engine/Source/Runtime/D3D12RHI/Private/D3D12SRV.cpp b/Engine/Source/Runtime/D3D12RHI/Private/D3D12SRV.cpp
--- a/Engine/Source/Runtime/D3D12RHI/Private/D3D12SRV.cpp
+++ b/Engine/Source/Runtime/D3D12RHI/Private/D3D12SRV.cpp
@@ -1,6 +1,7 @@
 // Copyright Epic Games, Inc. All Rights Reserved.
 
 #include "D3D12RHIPrivate.h"
+#include "D3D12RayTracing.h"
 
 // -----------------------------------------------------------------------------------------------------
 //
@@ -12,6 +13,11 @@ FD3D12ShaderResourceView::FD3D12ShaderResourceView(FD3D12Device* InDevice)
 	: TD3D12View(InDevice, ERHIDescriptorHeapType::Standard)
 {}
 
+FD3D12ShaderResourceView::FD3D12ShaderResourceView(FD3D12Device* InDevice, FD3D12RayTracingScene* InRayTracingScene)
+	: TD3D12View(InDevice, ERHIDescriptorHeapType::Standard)
+	, RayTracingScene(InRayTracingScene)
+{}
+
 void FD3D12ShaderResourceView::UpdateResourceInfo(const FResourceInfo& InResource, const D3D12_SHADER_RESOURCE_VIEW_DESC& InD3DViewDesc, EFlags InFlags)
 {
 	OffsetInBytes = 0;
@@ -289,7 +295,7 @@ void FD3D12ShaderResourceView_RHI::UpdateView(FD3D12ContextArray const& Contexts
 
 FD3D12ShaderResourceView_RHI::FD3D12ShaderResourceView_RHI(FD3D12Device* InDevice, FRHIViewableResource* InResource, FRHIViewDesc const& InViewDesc)
 	: FRHIShaderResourceView(InResource, InViewDesc)
-	, FD3D12ShaderResourceView(InDevice)
+	, FD3D12ShaderResourceView(InDevice, InViewDesc.Buffer.SRV.BufferType == FRHIViewDesc::EBufferType::AccelerationStructure ? FD3D12DynamicRHI::ResourceCast(InViewDesc.Buffer.SRV.RayTracingScene) : nullptr)
 {}
 
 
diff --git a/Engine/Source/Runtime/D3D12RHI/Private/D3D12View.h b/Engine/Source/Runtime/D3D12RHI/Private/D3D12View.h
--- a/Engine/Source/Runtime/D3D12RHI/Private/D3D12View.h
+++ b/Engine/Source/Runtime/D3D12RHI/Private/D3D12View.h
@@ -300,12 +300,15 @@ public:
 	FRIEND_ENUM_CLASS_FLAGS(EFlags)
 
 	FD3D12ShaderResourceView(FD3D12Device* InDevice);
+	FD3D12ShaderResourceView(FD3D12Device* InDevice, FD3D12RayTracingScene* InRayTracingScene);
 	void CreateView(FResourceInfo const& InResource, D3D12_SHADER_RESOURCE_VIEW_DESC const& InD3DViewDesc, EFlags InFlags);
 	void UpdateView(FD3D12ContextArray const& Contexts, const FResourceInfo& InResource, const D3D12_SHADER_RESOURCE_VIEW_DESC& InD3DViewDesc, EFlags InFlags);
 
 	bool GetSkipFastClearFinalize() const { return EnumHasAnyFlags(Flags, EFlags::SkipFastClearFinalize); }
 	void UpdateMinLODClamp(FD3D12ContextArray const& Contexts, float MinLODClamp);
 
+	FD3D12RayTracingScene* GetRayTracingScene() const { return RayTracingScene; }
+
 protected:
 	void UpdateResourceInfo(const FResourceInfo& InResource, const D3D12_SHADER_RESOURCE_VIEW_DESC& InD3DViewDesc, EFlags InFlags);
 	virtual void ResourceRenamed(FD3D12ContextArray const& Contexts, FD3D12BaseShaderResource* InRenamedResource, FD3D12ResourceLocation* InNewResourceLocation) override;
@@ -316,6 +319,8 @@ protected:
 	uint32 StrideInBytes = 0;
 
 	EFlags Flags = EFlags::None;
+
+	FD3D12RayTracingScene* RayTracingScene = nullptr;
 };
 
 ENUM_CLASS_FLAGS(FD3D12ShaderResourceView::EFlags)
diff --git a/Engine/Source/Runtime/RHI/Private/DynamicRHI.cpp b/Engine/Source/Runtime/RHI/Private/DynamicRHI.cpp
--- a/Engine/Source/Runtime/RHI/Private/DynamicRHI.cpp
+++ b/Engine/Source/Runtime/RHI/Private/DynamicRHI.cpp
@@ -887,9 +887,19 @@ RHI_API FShaderResourceViewInitializer::FShaderResourceViewInitializer(FRHIBuffe
 	: FRHIViewDesc::FBufferSRV::FInitializer()
 	, Buffer(InBuffer)
 {
+	SetTypeFromBuffer(Buffer);
 	SetOffsetInBytes(InStartOffsetBytes);
 	SetNumElements(InNumElements);
+}
+
+// Provided for back-compat.
+RHI_API FShaderResourceViewInitializer::FShaderResourceViewInitializer(FRHIBuffer* InBuffer, FRHIRayTracingScene* InRayTracingScene, uint32 InStartOffsetBytes)
+	: FRHIViewDesc::FBufferSRV::FInitializer()
+	, Buffer(InBuffer)
+{
 	SetTypeFromBuffer(Buffer);
+	SetRayTracingScene(InRayTracingScene);
+	SetOffsetInBytes(InStartOffsetBytes);
 }
 
 // Provided for back-compat.
diff --git a/Engine/Source/Runtime/RHI/Private/RHI.cpp b/Engine/Source/Runtime/RHI/Private/RHI.cpp
--- a/Engine/Source/Runtime/RHI/Private/RHI.cpp
+++ b/Engine/Source/Runtime/RHI/Private/RHI.cpp
@@ -2028,12 +2028,6 @@ FRHIShaderResourceView* FRHIBufferViewCache::GetOrCreateSRV(FRHICommandListBase&
 	}
 
 	auto CreateDesc = FRHIViewDesc::CreateBufferSRV();
-	CreateDesc.SetOffsetInBytes(SRVCreateInfo.StartOffsetBytes);
-
-	if (SRVCreateInfo.NumElements != UINT32_MAX)
-	{
-		CreateDesc.SetNumElements(SRVCreateInfo.NumElements);
-	}
 
 	if (EnumHasAnyFlags(Buffer->GetUsage(), BUF_ByteAddressBuffer))
 	{
@@ -2046,6 +2040,7 @@ FRHIShaderResourceView* FRHIBufferViewCache::GetOrCreateSRV(FRHICommandListBase&
 	else if (EnumHasAnyFlags(Buffer->GetUsage(), BUF_AccelerationStructure))
 	{
 		CreateDesc.SetType(FRHIViewDesc::EBufferType::AccelerationStructure);
+		CreateDesc.SetRayTracingScene(SRVCreateInfo.RayTracingScene);
 	}
 	else
 	{
@@ -2053,6 +2048,13 @@ FRHIShaderResourceView* FRHIBufferViewCache::GetOrCreateSRV(FRHICommandListBase&
 		CreateDesc.SetFormat(SRVCreateInfo.Format);
 	}
 
+	CreateDesc.SetOffsetInBytes(SRVCreateInfo.StartOffsetBytes);
+
+	if (SRVCreateInfo.NumElements != UINT32_MAX)
+	{
+		CreateDesc.SetNumElements(SRVCreateInfo.NumElements);
+	}
+
 	FShaderResourceViewRHIRef RHIShaderResourceView = RHICmdList.CreateShaderResourceView(Buffer, CreateDesc);
 
 	FRHIShaderResourceView* View = RHIShaderResourceView.GetReference();
diff --git a/Engine/Source/Runtime/RHI/Private/RHIResources.cpp b/Engine/Source/Runtime/RHI/Private/RHIResources.cpp
--- a/Engine/Source/Runtime/RHI/Private/RHIResources.cpp
+++ b/Engine/Source/Runtime/RHI/Private/RHIResources.cpp
@@ -213,7 +213,7 @@ FRHIViewDesc::FBuffer::FViewInfo FRHIViewDesc::FBuffer::GetViewInfo(FRHIBuffer*
 	case EBufferType::AccelerationStructure:
 		checkf(EnumHasAnyFlags(Desc.Usage, BUF_AccelerationStructure), TEXT("The buffer descriptor does not a ray tracing acceleration structure, so is incompatible with this view type."));
 		checkf(Format == PF_Unknown, TEXT("Acceleration structure views should not specify a format."));
-		checkf(Stride == 0, TEXT("Do not specify a stride for acceleration structure views."));
+		checkf(RayTracingScene != nullptr, TEXT("RayTracingScene must be specified when creating view of ray tracing acceleration structures."));
 
 		// Treat acceleration structures as a byte array.
 		Info.StrideInBytes = 1;
@@ -238,8 +238,8 @@ FRHIViewDesc::FBuffer::FViewInfo FRHIViewDesc::FBuffer::GetViewInfo(FRHIBuffer*
 	// OffsetInBytes == 0 && NumElements == 0 is a special case to mean "whole resource". If offset is non-zero, we need the caller to pass the required number of elements, except for acceleration structures.
 	checkf(Info.BufferType == EBufferType::AccelerationStructure || (OffsetInBytes == 0 || NumElements > 0), TEXT("NumElements field must be non-zero if a byte offset is used."));
 		
-	// When NumElements is zero, use "whole buffer".
-	Info.NumElements = NumElements == 0 ? (Desc.Size - OffsetInBytes) / Info.StrideInBytes : NumElements;
+	// If BufferType is AccelerationStructure or NumElements is zero, use "whole buffer".
+	Info.NumElements = (Info.BufferType == EBufferType::AccelerationStructure || NumElements == 0) ? (Desc.Size - OffsetInBytes) / Info.StrideInBytes : NumElements;
 	Info.SizeInBytes = Info.NumElements * Info.StrideInBytes;
 
 	checkf(Info.OffsetInBytes + Info.SizeInBytes <= Desc.Size,
diff --git a/Engine/Source/Runtime/RHI/Public/DynamicRHI.h b/Engine/Source/Runtime/RHI/Public/DynamicRHI.h
--- a/Engine/Source/Runtime/RHI/Public/DynamicRHI.h
+++ b/Engine/Source/Runtime/RHI/Public/DynamicRHI.h
@@ -68,6 +68,7 @@ struct FShaderResourceViewInitializer : public FRHIViewDesc::FBufferSRV::FInitia
 	RHI_API FShaderResourceViewInitializer(FRHIBuffer* InBuffer, EPixelFormat InFormat, uint32 InStartOffsetBytes, uint32 InNumElements);
 	RHI_API FShaderResourceViewInitializer(FRHIBuffer* InBuffer, EPixelFormat InFormat);
 	RHI_API FShaderResourceViewInitializer(FRHIBuffer* InBuffer, uint32 InStartOffsetBytes, uint32 InNumElements);
+	RHI_API FShaderResourceViewInitializer(FRHIBuffer* InBuffer, FRHIRayTracingScene* InRayTracingScene, uint32 InStartOffsetBytes);
 	RHI_API FShaderResourceViewInitializer(FRHIBuffer* InBuffer);
 };
 
diff --git a/Engine/Source/Runtime/RHI/Public/RHIResources.h b/Engine/Source/Runtime/RHI/Public/RHIResources.h
--- a/Engine/Source/Runtime/RHI/Public/RHIResources.h
+++ b/Engine/Source/Runtime/RHI/Public/RHIResources.h
@@ -2381,8 +2381,15 @@ struct FRHIViewDesc
 		uint8       bAppendBuffer  : 1; // UAV only
 		uint8       /* padding */  : 6;
 		uint32      OffsetInBytes;
-		uint32      NumElements;
-		uint32      Stride;
+		union
+		{
+			struct
+			{
+				uint32 NumElements;
+				uint32 Stride;
+			};
+			FRHIRayTracingScene* RayTracingScene; // only if BufferType == AccelerationStructure
+		};
 
 		struct FViewInfo;
 	protected:
@@ -2541,15 +2548,24 @@ public:
 
 	FInitializer& SetStride(uint32 InStride)
 	{
+		check(Buffer.SRV.BufferType != EBufferType::Unknown && Buffer.SRV.BufferType != EBufferType::AccelerationStructure);
 		Buffer.SRV.Stride = InStride;
 		return *this;
 	}
 
 	FInitializer& SetNumElements(uint32 InNumElements)
 	{
+		check(Buffer.SRV.BufferType != EBufferType::Unknown && Buffer.SRV.BufferType != EBufferType::AccelerationStructure);
 		Buffer.SRV.NumElements = InNumElements;
 		return *this;
 	}
+
+	FInitializer& SetRayTracingScene(FRHIRayTracingScene* InRayTracingScene)
+	{
+		check(Buffer.SRV.BufferType != EBufferType::Unknown && Buffer.SRV.BufferType == EBufferType::AccelerationStructure);
+		Buffer.SRV.RayTracingScene = InRayTracingScene;
+		return *this;
+	}
 };
 
 struct FRHIViewDesc::FBufferUAV::FInitializer : private FRHIViewDesc
@@ -2838,7 +2854,7 @@ struct FRHIViewDesc::FBuffer::FViewInfo
 	// The format of the data exposed by this view. PF_Unknown for all buffer types except typed buffer views.
 	EPixelFormat Format;
 
-	// When true, the view is refering to a BUF_NullResource, so a null descriptor should be created.
+	// When true, the view is referring to a BUF_NullResource, so a null descriptor should be created.
 	bool bNullView;
 };
 
@@ -5085,11 +5101,17 @@ struct FRHIBufferSRVCreateInfo
 		, NumElements(InNumElements)
 	{}
 
+	FRHIBufferSRVCreateInfo(FRHIRayTracingScene* InRayTracingScene, uint32 InStartOffsetBytes)
+		: StartOffsetBytes(InStartOffsetBytes)
+		, RayTracingScene(InRayTracingScene)
+	{}
+
 	FORCEINLINE bool operator==(const FRHIBufferSRVCreateInfo& Other)const
 	{
 		return Format == Other.Format
 			&& StartOffsetBytes == Other.StartOffsetBytes
-			&& NumElements == Other.NumElements;
+			&& NumElements == Other.NumElements
+			&& RayTracingScene == Other.RayTracingScene;
 	}
 
 	FORCEINLINE bool operator!=(const FRHIBufferSRVCreateInfo& Other)const
@@ -5100,8 +5122,10 @@ struct FRHIBufferSRVCreateInfo
 	friend uint32 GetTypeHash(const FRHIBufferSRVCreateInfo& Desc)
 	{
 		return HashCombine(
-			HashCombine(GetTypeHash(Desc.Format), GetTypeHash(Desc.StartOffsetBytes)),
-			GetTypeHash(Desc.NumElements)
+			HashCombine(
+				HashCombine(GetTypeHash(Desc.Format), GetTypeHash(Desc.StartOffsetBytes)),
+				GetTypeHash(Desc.NumElements)),
+			GetTypeHash(Desc.RayTracingScene)
 		);
 	}
 
@@ -5113,6 +5137,9 @@ struct FRHIBufferSRVCreateInfo
 
 	/** Number of elements (whole buffer by default) */
 	uint32 NumElements = UINT32_MAX;
+
+	/** Ray tracing scene associated with the SRV (if BUF_AccelerationStructure) */
+	FRHIRayTracingScene* RayTracingScene = nullptr;
 };
 
 struct FRHIBufferUAVCreateInfo
diff --git a/Engine/Source/Runtime/RenderCore/Public/RenderGraphResources.h b/Engine/Source/Runtime/RenderCore/Public/RenderGraphResources.h
--- a/Engine/Source/Runtime/RenderCore/Public/RenderGraphResources.h
+++ b/Engine/Source/Runtime/RenderCore/Public/RenderGraphResources.h
@@ -1105,6 +1105,11 @@ struct FRDGBufferSRVDesc final
 		, Buffer(InBuffer)
 	{}
 
+	FRDGBufferSRVDesc(FRDGBufferRef InBuffer, FRHIRayTracingScene* InRayTracingScene, uint32 InStartOffsetBytes)
+		: FRHIBufferSRVCreateInfo(InRayTracingScene, InStartOffsetBytes)
+		, Buffer(InBuffer)
+	{}
+
 	bool operator == (const FRDGBufferSRVDesc& Other) const
 	{
 		return Buffer == Other.Buffer && FRHIBufferSRVCreateInfo::operator==(Other);
diff --git a/Engine/Source/Runtime/Renderer/Private/RayTracing/RayTracingScene.cpp b/Engine/Source/Runtime/Renderer/Private/RayTracing/RayTracingScene.cpp
--- a/Engine/Source/Runtime/Renderer/Private/RayTracing/RayTracingScene.cpp
+++ b/Engine/Source/Runtime/Renderer/Private/RayTracing/RayTracingScene.cpp
@@ -105,7 +105,7 @@ void FRayTracingScene::CreateWithInitializationData(FRDGBuilder& GraphBuilder, c
 	LayerSRVs.SetNum(NumLayers);
 	for (uint32 LayerIndex = 0; LayerIndex < NumLayers; ++LayerIndex)
 	{
-		LayerSRVs[LayerIndex] = GraphBuilder.CreateSRV(FRDGBufferSRVDesc(RayTracingSceneBufferRDG, RayTracingSceneRHI->GetLayerBufferOffset(LayerIndex), 0));
+		LayerSRVs[LayerIndex] = GraphBuilder.CreateSRV(FRDGBufferSRVDesc(RayTracingSceneBufferRDG, RayTracingSceneRHI, RayTracingSceneRHI->GetLayerBufferOffset(LayerIndex)));
 	}
 
 	{
@@ -465,7 +465,7 @@ FShaderResourceViewRHIRef FRayTracingScene::CreateLayerViewRHI(FRHICommandListBa
 {
 	const uint8 LayerIndex = uint8(Layer);
 	checkf(RayTracingScenePooledBuffer, TEXT("Ray tracing scene was not created.Perhaps Create() was not called."));
-	return RHICmdList.CreateShaderResourceView(FShaderResourceViewInitializer(RayTracingScenePooledBuffer->GetRHI(), RayTracingSceneRHI->GetLayerBufferOffset(LayerIndex), 0));
+	return RHICmdList.CreateShaderResourceView(FShaderResourceViewInitializer(RayTracingScenePooledBuffer->GetRHI(), RayTracingSceneRHI, RayTracingSceneRHI->GetLayerBufferOffset(LayerIndex)));
 }
 
 FRDGBufferSRVRef FRayTracingScene::GetLayerView(ERayTracingSceneLayer Layer) const
