Fix errors when compiling with RHI_RAYTRACING disabled.

diff --git a/Engine/Source/Runtime/D3D12RHI/Private/D3D12DescriptorCache.cpp b/Engine/Source/Runtime/D3D12RHI/Private/D3D12DescriptorCache.cpp
--- a/Engine/Source/Runtime/D3D12RHI/Private/D3D12DescriptorCache.cpp
+++ b/Engine/Source/Runtime/D3D12RHI/Private/D3D12DescriptorCache.cpp
@@ -500,11 +500,13 @@ D3D12_GPU_DESCRIPTOR_HANDLE FD3D12DescriptorCache::BuildSRVTable(EShaderFrequenc
 			Context.TransitionResource(SRV, State & ValidResourceStates);
 			Context.UpdateResidency(Cache.Resources[ShaderStage][SlotIndex]);
 
+#if RHI_RAYTRACING
 			FD3D12RayTracingScene* RayTracingScene = SRV->GetRayTracingScene();
 			if (RayTracingScene)
 			{
 				RayTracingScene->UpdateResidency(Context);
 			}
+#endif
 		}
 		else
 		{
diff --git a/Engine/Source/Runtime/D3D12RHI/Private/D3D12SRV.cpp b/Engine/Source/Runtime/D3D12RHI/Private/D3D12SRV.cpp
--- a/Engine/Source/Runtime/D3D12RHI/Private/D3D12SRV.cpp
+++ b/Engine/Source/Runtime/D3D12RHI/Private/D3D12SRV.cpp
@@ -295,7 +295,11 @@ void FD3D12ShaderResourceView_RHI::UpdateView(FD3D12ContextArray const& Contexts
 
 FD3D12ShaderResourceView_RHI::FD3D12ShaderResourceView_RHI(FD3D12Device* InDevice, FRHIViewableResource* InResource, FRHIViewDesc const& InViewDesc)
 	: FRHIShaderResourceView(InResource, InViewDesc)
+#if RHI_RAYTRACING
 	, FD3D12ShaderResourceView(InDevice, InViewDesc.Buffer.SRV.BufferType == FRHIViewDesc::EBufferType::AccelerationStructure ? FD3D12DynamicRHI::ResourceCast(InViewDesc.Buffer.SRV.RayTracingScene) : nullptr)
+#else
+	, FD3D12ShaderResourceView(InDevice)
+#endif
 {}
 
 
diff --git a/Engine/Source/Runtime/D3D12RHI/Private/D3D12View.h b/Engine/Source/Runtime/D3D12RHI/Private/D3D12View.h
--- a/Engine/Source/Runtime/D3D12RHI/Private/D3D12View.h
+++ b/Engine/Source/Runtime/D3D12RHI/Private/D3D12View.h
@@ -10,6 +10,7 @@
 
 class FD3D12Buffer;
 class FD3D12Texture;
+class FD3D12RayTracingScene;
 struct FD3D12ResidencyHandle;
 
 struct FD3D12DefaultViews
diff --git a/Engine/Source/Runtime/RHI/Public/RHIResources.h b/Engine/Source/Runtime/RHI/Public/RHIResources.h
--- a/Engine/Source/Runtime/RHI/Public/RHIResources.h
+++ b/Engine/Source/Runtime/RHI/Public/RHIResources.h
@@ -2562,7 +2562,7 @@ public:
 
 	FInitializer& SetRayTracingScene(FRHIRayTracingScene* InRayTracingScene)
 	{
-		check(Buffer.SRV.BufferType != EBufferType::Unknown && Buffer.SRV.BufferType == EBufferType::AccelerationStructure);
+		check(Buffer.SRV.BufferType == EBufferType::AccelerationStructure);
 		Buffer.SRV.RayTracingScene = InRayTracingScene;
 		return *this;
 	}
