Add bindless parameters to FRayTracingShaderBindings to make sure they are tracked for residency

diff --git a/Engine/Source/Runtime/D3D12RHI/Private/D3D12RayTracing.cpp b/Engine/Source/Runtime/D3D12RHI/Private/D3D12RayTracing.cpp
--- a/Engine/Source/Runtime/D3D12RHI/Private/D3D12RayTracing.cpp
+++ b/Engine/Source/Runtime/D3D12RHI/Private/D3D12RayTracing.cpp
@@ -4286,6 +4286,7 @@ template <typename ResourceBinderType>
 static bool SetRayTracingShaderResources(
 	const FD3D12RayTracingShader* Shader,
 	const FD3D12RootSignature* RootSignature,
+	uint32 InNumBindlessParameters, FRHIShaderParameterResource const* BindlessParameters,
 	uint32 InNumTextures, FRHITexture* const* Textures,
 	uint32 InNumSRVs, FRHIShaderResourceView* const* SRVs,
 	uint32 InNumUniformBuffers, FRHIUniformBuffer* const* UniformBuffers,
@@ -4298,14 +4299,22 @@ static bool SetRayTracingShaderResources(
 
 	struct FBindings
 	{
-		FBindings(ResourceBinderType& InBinder, uint32 InGPUIndex)
+		FBindings(ResourceBinderType& InBinder, uint32 InGPUIndex, const FD3D12ShaderData* ShaderData)
 			: Binder(InBinder)
 			, GPUIndex(InGPUIndex)
+#if PLATFORM_SUPPORTS_BINDLESS_RENDERING
+			, bBindlessResources(EnumHasAnyFlags(ShaderData->ResourceCounts.UsageFlags, EShaderResourceUsageFlags::BindlessResources))
+			, bBindlessSamplers(EnumHasAnyFlags(ShaderData->ResourceCounts.UsageFlags, EShaderResourceUsageFlags::BindlessSamplers))
+#endif
 		{
 		}
 
 		ResourceBinderType& Binder;
 		uint32 GPUIndex;
+#if PLATFORM_SUPPORTS_BINDLESS_RENDERING
+		const bool bBindlessResources;
+		const bool bBindlessSamplers;
+#endif
 
 #if D3D12RHI_USE_CONSTANT_BUFFER_VIEWS
 		D3D12_CPU_DESCRIPTOR_HANDLE LocalCBVs[MAX_CBS];
@@ -4336,11 +4345,15 @@ static bool SetRayTracingShaderResources(
 			FD3D12UnorderedAccessView* UAV = FD3D12CommandContext::RetrieveObject<FD3D12UnorderedAccessView_RHI>(RHIUAV, GPUIndex);
 			check(UAV != nullptr);
 
-			FD3D12OfflineDescriptor Descriptor = UAV->GetOfflineCpuHandle();
-			LocalUAVs[Index] = Descriptor;
-			UAVVersions[Index] = Descriptor.GetVersion();
-
-			BoundUAVMask |= 1ull << Index;
+#if PLATFORM_SUPPORTS_BINDLESS_RENDERING
+			if (!bBindlessResources)
+#endif
+			{
+				FD3D12OfflineDescriptor Descriptor = UAV->GetOfflineCpuHandle();
+				LocalUAVs[Index] = Descriptor;
+				UAVVersions[Index] = Descriptor.GetVersion();
+				BoundUAVMask |= 1ull << Index;
+			}
 
 			ReferencedResources.Add(UAV->GetResource());
 			Binder.AddResourceTransition(UAV);
@@ -4351,11 +4364,15 @@ static bool SetRayTracingShaderResources(
 			FD3D12ShaderResourceView* SRV = FD3D12CommandContext::RetrieveObject<FD3D12ShaderResourceView_RHI>(RHISRV, GPUIndex);
 			check(SRV != nullptr);
 
-			FD3D12OfflineDescriptor Descriptor = SRV->GetOfflineCpuHandle();
-			LocalSRVs[Index] = Descriptor;
-			SRVVersions[Index] = Descriptor.GetVersion();
-
-			BoundSRVMask |= 1ull << Index;
+#if PLATFORM_SUPPORTS_BINDLESS_RENDERING
+			if (!bBindlessResources)
+#endif
+			{
+				FD3D12OfflineDescriptor Descriptor = SRV->GetOfflineCpuHandle();
+				LocalSRVs[Index] = Descriptor;
+				SRVVersions[Index] = Descriptor.GetVersion();
+				BoundSRVMask |= 1ull << Index;
+			}
 
 			ReferencedResources.Add(SRV->GetResource());
 			Binder.AddResourceTransition(SRV);
@@ -4376,11 +4393,15 @@ static bool SetRayTracingShaderResources(
 			}
 			check(SRV != nullptr);
 
-			FD3D12OfflineDescriptor Descriptor = SRV->GetOfflineCpuHandle();
-			LocalSRVs[Index] = Descriptor;
-			SRVVersions[Index] = Descriptor.GetVersion();
-
-			BoundSRVMask |= 1ull << Index;
+#if PLATFORM_SUPPORTS_BINDLESS_RENDERING
+			if (!bBindlessResources)
+#endif
+			{
+				FD3D12OfflineDescriptor Descriptor = SRV->GetOfflineCpuHandle();
+				LocalSRVs[Index] = Descriptor;
+				SRVVersions[Index] = Descriptor.GetVersion();
+				BoundSRVMask |= 1ull << Index;
+			}
 
 			ReferencedResources.Add(SRV->GetResource());
 			Binder.AddResourceTransition(SRV);
@@ -4393,9 +4414,14 @@ static bool SetRayTracingShaderResources(
 
 			check(SRV != nullptr);
 
-			FD3D12OfflineDescriptor Descriptor = SRV->GetOfflineCpuHandle();
-			LocalSRVs[Index] = Descriptor;
-			SRVVersions[Index] = Descriptor.GetVersion();
+#if PLATFORM_SUPPORTS_BINDLESS_RENDERING
+			if (!bBindlessResources)
+#endif
+			{
+				FD3D12OfflineDescriptor Descriptor = SRV->GetOfflineCpuHandle();
+				LocalSRVs[Index] = Descriptor;
+				SRVVersions[Index] = Descriptor.GetVersion();
+			}
 
 			BoundSRVMask |= 1ull << Index;
 
@@ -4408,14 +4434,43 @@ static bool SetRayTracingShaderResources(
 			FD3D12SamplerState* Sampler = FD3D12CommandContext::RetrieveObject<FD3D12SamplerState>(RHISampler, GPUIndex);
 			check(Sampler != nullptr);
 
-			FD3D12OfflineDescriptor Descriptor = Sampler->OfflineDescriptor;
-			LocalSamplers[Index] = Descriptor;
-			SamplerVersions[Index] = Descriptor.GetVersion();
-
-			BoundSamplerMask |= 1ull << Index;
+#if PLATFORM_SUPPORTS_BINDLESS_RENDERING
+			if (!bBindlessSamplers)
+#endif
+			{
+				FD3D12OfflineDescriptor Descriptor = Sampler->OfflineDescriptor;
+				LocalSamplers[Index] = Descriptor;
+				SamplerVersions[Index] = Descriptor.GetVersion();
+				BoundSamplerMask |= 1ull << Index;
+			}
 		}
 	};
-	FBindings Bindings(Binder, Binder.GetDevice()->GetGPUIndex());
+	FBindings Bindings(Binder, Binder.GetDevice()->GetGPUIndex(), Shader);
+
+#if PLATFORM_SUPPORTS_BINDLESS_RENDERING
+	for (uint32 BindlessParameterIndex = 0; BindlessParameterIndex < InNumBindlessParameters; ++BindlessParameterIndex)
+	{
+		const FRHIShaderParameterResource& ShaderParameterResource = BindlessParameters[BindlessParameterIndex];
+		if (FRHIResource* Resource = ShaderParameterResource.Resource)
+		{
+			switch (ShaderParameterResource.Type)
+			{
+				case FRHIShaderParameterResource::EType::Texture:
+					Bindings.SetTexture(static_cast<FRHITexture*>(Resource), BindlessParameterIndex);
+					break;
+				case FRHIShaderParameterResource::EType::ResourceView:
+					Bindings.SetSRV(static_cast<FRHIShaderResourceView*>(Resource), BindlessParameterIndex);
+					break;
+				case FRHIShaderParameterResource::EType::UnorderedAccessView:
+					Bindings.SetUAV(static_cast<FRHIUnorderedAccessView*>(Resource), BindlessParameterIndex);
+					break;
+				case FRHIShaderParameterResource::EType::Sampler:
+					Bindings.SetSampler(static_cast<FRHISamplerState*>(Resource), BindlessParameterIndex);
+					break;
+			}
+		}
+	}
+#endif
 
 	for (uint32 TextureIndex = 0; TextureIndex < InNumTextures; ++TextureIndex)
 	{
@@ -4640,6 +4695,7 @@ static bool SetRayTracingShaderResources(
 	return SetRayTracingShaderResources(
 		Shader,
 		RootSignature,
+		ResourceBindings.BindlessParameters.Num(), ResourceBindings.BindlessParameters.GetData(),
 		UE_ARRAY_COUNT(ResourceBindings.Textures), ResourceBindings.Textures,
 		UE_ARRAY_COUNT(ResourceBindings.SRVs), ResourceBindings.SRVs,
 		UE_ARRAY_COUNT(ResourceBindings.UniformBuffers), ResourceBindings.UniformBuffers,
@@ -4940,6 +4996,7 @@ static void SetRayTracingHitGroup(
 
 	FD3D12RayTracingLocalResourceBinder ResourceBinder(*Device, *ShaderTable, *(Shader->LocalRootSignature), RecordIndex, WorkerIndex, ERayTracingBindingType::HitGroup);
 	const bool bResourcesBound = SetRayTracingShaderResources(Shader, Shader->LocalRootSignature,		
+		0, nullptr,	// BindlessParameters
 		0, nullptr, // Textures
 		0, nullptr, // SRVs
 		NumUniformBuffers, UniformBuffers,
@@ -4981,6 +5038,7 @@ static void SetRayTracingCallableShader(
 
 		FD3D12RayTracingLocalResourceBinder ResourceBinder(*Device, *ShaderTable, *(Shader->LocalRootSignature), RecordIndex, WorkerIndex, ERayTracingBindingType::CallableShader);
 		const bool bResourcesBound = SetRayTracingShaderResources(Shader, Shader->LocalRootSignature,
+			0, nullptr,	// BindlessParameters
 			0, nullptr, // Textures
 			0, nullptr, // SRVs
 			NumUniformBuffers, UniformBuffers,
@@ -5016,6 +5074,7 @@ static void SetRayTracingMissShader(
 
 	FD3D12RayTracingLocalResourceBinder ResourceBinder(*Device, *ShaderTable, *(Shader->LocalRootSignature), RecordIndex, WorkerIndex, ERayTracingBindingType::MissShader);
 	const bool bResourcesBound = SetRayTracingShaderResources(Shader, Shader->LocalRootSignature,
+		0, nullptr,	// BindlessParameters
 		0, nullptr, // Textures
 		0, nullptr, // SRVs
 		NumUniformBuffers, UniformBuffers,
diff --git a/Engine/Source/Runtime/RHI/Public/RHICommandList.h b/Engine/Source/Runtime/RHI/Public/RHICommandList.h
--- a/Engine/Source/Runtime/RHI/Public/RHICommandList.h
+++ b/Engine/Source/Runtime/RHI/Public/RHICommandList.h
@@ -245,6 +245,7 @@ struct FRayTracingShaderBindings
 	FRHISamplerState* Samplers[32] = {};
 	FRHIUnorderedAccessView* UAVs[16] = {};
 
+	TArray<FRHIShaderParameterResource> BindlessParameters;
 	FUniformBufferStaticBindings StaticUniformBuffers;
 };
 
diff --git a/Engine/Source/Runtime/RenderCore/Private/ShaderParameterStruct.cpp b/Engine/Source/Runtime/RenderCore/Private/ShaderParameterStruct.cpp
--- a/Engine/Source/Runtime/RenderCore/Private/ShaderParameterStruct.cpp
+++ b/Engine/Source/Runtime/RenderCore/Private/ShaderParameterStruct.cpp
@@ -854,6 +854,17 @@ void SetShaderParameters(
 {
 	const FShaderParameterReader Reader(ParametersData, ParametersMetadata->GetSize());
 
+#if PLATFORM_SUPPORTS_BINDLESS_RENDERING
+	if (int32 NumBindings = Bindings.BindlessResourceParameters.Num())
+	{
+		for (const FShaderParameterBindings::FBindlessResourceParameter& Parameter : Bindings.BindlessResourceParameters)
+		{
+			const FRHIShaderParameterResource ShaderParameterResource = ExtractShaderParameterResource(Reader, Parameter);
+			RTBindingsWriter.AddBindlessParameter(ShaderParameterResource);
+		}
+	}
+#endif // PLATFORM_SUPPORTS_BINDLESS_RENDERING
+
 	for (const FShaderParameterBindings::FResourceParameter& Parameter : Bindings.ResourceParameters)
 	{
 		const EUniformBufferBaseType BaseType = static_cast<EUniformBufferBaseType>(Parameter.BaseType);
diff --git a/Engine/Source/Runtime/RenderCore/Public/ShaderParameters.h b/Engine/Source/Runtime/RenderCore/Public/ShaderParameters.h
--- a/Engine/Source/Runtime/RenderCore/Public/ShaderParameters.h
+++ b/Engine/Source/Runtime/RenderCore/Public/ShaderParameters.h
@@ -170,6 +170,11 @@ struct FRayTracingShaderBindingsWriter : FRayTracingShaderBindings
 {
 	FUniformBufferRHIRef RootUniformBuffer;
 
+	void AddBindlessParameter(const FRHIShaderParameterResource& Parameter)
+	{
+		BindlessParameters.Add(Parameter);
+	}
+
 	void Set(const FShaderResourceParameter& Param, FRHITexture* Value)
 	{
 		if (Param.IsBound())
