Fix deadlock in D3D12 residency management
 - D3DX12Residency includes device-wide CPU waits on fences. This happens whenever resources need to be evicted, blocking the submission thread until the GPU is idle.
 - Payloads were getting batched for submission into an array on their respective target queues, however this means that submission to the GPU no longer follows the "only submit work to the GPU that is guaranteed to complete" rule, as following that rule means we need to process payloads in the order they were batched, not grouped by queue.
 - Ordinarily this wouldn't be a problem, but the residency management inserts and awaits its own fences. Submitting payloads out-of-order can lead to a 3 way deadlock between the 2 GPU queues and the submission thread. The CPU is blocked in residency waiting for the GPU to go idle, but one of the GPU queues is blocked waiting for some dependent work to be submitted on the other queue, which will never happen because the submission thread is blocked.
 - Solution is to batch the payloads into a single ordered array on the FD3D12DynamicRHI, in the same order they are processed by the submission thread when resolving sync points. This guarantees we never submit work to the GPU that cannot complete without further CPU intervention, and therefore works around the side effects residency management has with its own fences.

diff --git a/Engine/Source/Runtime/D3D12RHI/Private/D3D12Device.h b/Engine/Source/Runtime/D3D12RHI/Private/D3D12Device.h
--- a/Engine/Source/Runtime/D3D12RHI/Private/D3D12Device.h
+++ b/Engine/Source/Runtime/D3D12RHI/Private/D3D12Device.h
@@ -124,7 +124,6 @@ public:
 		}		
 	} PendingSubmission, PendingInterrupt;
 
-	TArray<FD3D12Payload*>  PayloadsToSubmit;
 	FD3D12Payload*          PayloadToSubmit  = nullptr;
 	FD3D12CommandAllocator* BarrierAllocator = nullptr;
 	FD3D12QueryAllocator    BarrierTimestamps;
@@ -155,9 +154,6 @@ public:
 	// Batches the current payload's command lists, returning the latest fence value signaled for this queue.
 	uint64 FinalizePayload(bool bRequiresSignal, TArray<FD3D12Payload*, TInlineAllocator<64>>& PayloadsToHandDown);
 
-	// Ensures all prior batched command lists have reached the driver ID3D12Queue object.
-	void FlushBatchedPayloads();
-
 	// Call the underlying ID3D12Queue::ExecuteCommandLists function
 	void ExecuteCommandLists(TArrayView<ID3D12CommandList*> D3DCommandLists
 #if ENABLE_RESIDENCY_MANAGEMENT
diff --git a/Engine/Source/Runtime/D3D12RHI/Private/D3D12RHIPrivate.h b/Engine/Source/Runtime/D3D12RHI/Private/D3D12RHIPrivate.h
--- a/Engine/Source/Runtime/D3D12RHI/Private/D3D12RHIPrivate.h
+++ b/Engine/Source/Runtime/D3D12RHI/Private/D3D12RHIPrivate.h
@@ -210,6 +210,9 @@ private:
 
 	TQueue<TArray<FD3D12Payload*>*, EQueueMode::Mpsc> PendingPayloadsForSubmission;
 
+	// Ensures all prior batched command lists have reached the driver ID3D12Queue object.
+	void FlushBatchedPayloads(TArray<FD3D12Payload*, TInlineAllocator<64>>& PayloadsToHandDown);
+
 	FCriticalSection SubmissionCS;
 	FCriticalSection InterruptCS;
 
diff --git a/Engine/Source/Runtime/D3D12RHI/Private/D3D12Submission.cpp b/Engine/Source/Runtime/D3D12RHI/Private/D3D12Submission.cpp
--- a/Engine/Source/Runtime/D3D12RHI/Private/D3D12Submission.cpp
+++ b/Engine/Source/Runtime/D3D12RHI/Private/D3D12Submission.cpp
@@ -497,17 +497,9 @@ FD3D12DynamicRHI::FProcessResult FD3D12DynamicRHI::ProcessSubmissionQueue()
 	if (PayloadsToHandDown.Num())
 	{
 		Result.Status |= EQueueStatus::Processed;
-		for (FD3D12Payload* Payload : PayloadsToHandDown)
-		{
-			Payload->Queue.PendingInterrupt.Enqueue(Payload);
-		}
+		FlushBatchedPayloads(PayloadsToHandDown);
 	}
 
-	ForEachQueue([&](FD3D12Queue& CurrentQueue)
-	{
-		CurrentQueue.FlushBatchedPayloads();
-	});
-
 	if (InterruptThread && EnumHasAnyFlags(Result.Status, EQueueStatus::Processed))
 	{
 		InterruptThread->Kick();
@@ -818,35 +810,37 @@ uint64 FD3D12Queue::FinalizePayload(bool bRequiresSignal, TArray<FD3D12Payload*,
 		SyncPoint->ResolvedFence.Emplace(Fence, PayloadToSubmit->CompletionFenceValue);
 	}
 
-	PayloadsToSubmit.Add(PayloadToSubmit);
 	PayloadsToHandDown.Add(PayloadToSubmit);
-
 	PayloadToSubmit = nullptr;
 
 	return Fence.NextCompletionValue;
 }
 
-void FD3D12Queue::FlushBatchedPayloads()
+void FD3D12DynamicRHI::FlushBatchedPayloads(TArray<FD3D12Payload*, TInlineAllocator<64>>& PayloadsToSubmit)
 {
 	uint32 FirstPayload = 0, LastPayload = 0;
 
 	auto Wait = [this](FD3D12Payload* Payload)
 	{
+		FD3D12Queue& Queue = Payload->Queue;
+
 		// Wait for queue fences
 		for (auto& [LocalFence, Value] : Payload->QueueFencesToWait)
 		{
-			VERIFYD3D12RESULT(D3DCommandQueue->Wait(LocalFence.D3DFence, Value));
+			VERIFYD3D12RESULT(Queue.D3DCommandQueue->Wait(LocalFence.D3DFence, Value));
 		}
 
 		// Wait for manual fences
 		for (auto& [LocalFence, Value] : Payload->ManualFencesToWait)
 		{
-			VERIFYD3D12RESULT(D3DCommandQueue->Wait(LocalFence, Value));
+			VERIFYD3D12RESULT(Queue.D3DCommandQueue->Wait(LocalFence, Value));
 		}
 	};
 
 	auto UpdateReservedResources = [this](FD3D12Payload* Payload)
 	{
+		FD3D12Queue& Queue = Payload->Queue;
+
 		// On some devices, some queues cannot perform tile remapping operations.
 		// We can work around this limitation by running the remapping in lockstep on another queue:
 		// - tile mapping queue waits for commands on this queue to finish
@@ -854,15 +848,15 @@ void FD3D12Queue::FlushBatchedPayloads()
 		// - this queue waits for tile mapping queue to finish
 		// The extra sync is not required when the current queue is capable of the remapping operations.
 
-		ID3D12CommandQueue* TileMappingQueue = (bSupportsTileMapping ? D3DCommandQueue : Device->TileMappingQueue).GetReference();
-		FD3D12Fence& TileMappingFence = Device->TileMappingFence;
+		ID3D12CommandQueue* TileMappingQueue = (Queue.bSupportsTileMapping ? Queue.D3DCommandQueue : Queue.Device->TileMappingQueue).GetReference();
+		FD3D12Fence& TileMappingFence = Queue.Device->TileMappingFence;
 
-		const bool bCrossQueueSyncRequired = TileMappingQueue != D3DCommandQueue.GetReference();
+		const bool bCrossQueueSyncRequired = TileMappingQueue != Queue.D3DCommandQueue.GetReference();
 
 		if (bCrossQueueSyncRequired)
 		{
 			// tile mapping queue waits for commands on this queue to finish
-			D3DCommandQueue->Signal(TileMappingFence.D3DFence, ++TileMappingFence.LastSignaledValue);
+			Queue.D3DCommandQueue->Signal(TileMappingFence.D3DFence, ++TileMappingFence.LastSignaledValue);
 			TileMappingQueue->Wait(TileMappingFence.D3DFence, TileMappingFence.LastSignaledValue);
 		}
 
@@ -875,7 +869,7 @@ void FD3D12Queue::FlushBatchedPayloads()
 		{
 			// this queue waits for tile mapping operations to finish
 			TileMappingQueue->Signal(TileMappingFence.D3DFence, ++TileMappingFence.LastSignaledValue);
-			D3DCommandQueue->Wait(TileMappingFence.D3DFence, TileMappingFence.LastSignaledValue);
+			Queue.D3DCommandQueue->Wait(TileMappingFence.D3DFence, TileMappingFence.LastSignaledValue);
 		}
 	};
 
@@ -884,6 +878,8 @@ void FD3D12Queue::FlushBatchedPayloads()
 		if (FirstPayload == LastPayload)
 			return;
 
+		FD3D12Queue& Queue = PayloadsToSubmit[FirstPayload]->Queue;
+
 		// Build SOA layout needed to call ExecuteCommandLists().
 		TArray<FD3D12CommandList*, TInlineAllocator<128>> CommandLists;
 		TArray<ID3D12CommandList*, TInlineAllocator<128>> D3DCommandLists;
@@ -895,6 +891,7 @@ void FD3D12Queue::FlushBatchedPayloads()
 		for (uint32 Index = FirstPayload; Index < LastPayload; ++Index)
 		{
 			FD3D12Payload* Payload = PayloadsToSubmit[Index];
+			check(&Payload->Queue == &Queue);
 
 			for (FD3D12CommandList* CommandList : Payload->CommandListsToExecute)
 			{
@@ -929,7 +926,7 @@ void FD3D12Queue::FlushBatchedPayloads()
 				DispatchNum = Index;
 			}
 
-			ExecuteCommandLists(
+			Queue.ExecuteCommandLists(
 				MakeArrayView<ID3D12CommandList*>(&D3DCommandLists[Offset], DispatchNum)
 #if ENABLE_RESIDENCY_MANAGEMENT
 				, MakeArrayView<FD3D12ResidencySet*>(&ResidencySets[Offset], DispatchNum)
@@ -952,19 +949,21 @@ void FD3D12Queue::FlushBatchedPayloads()
 
 	auto Signal = [this](FD3D12Payload* Payload)
 	{
+		FD3D12Queue& Queue = Payload->Queue;
+
 		// Signal any manual fences
 		for (auto& [ManualFence, Value] : Payload->ManualFencesToSignal)
 		{
-			VERIFYD3D12RESULT(D3DCommandQueue->Signal(ManualFence, Value));
+			VERIFYD3D12RESULT(Queue.D3DCommandQueue->Signal(ManualFence, Value));
 		}
 
 		// Signal the queue fence
 		if (Payload->RequiresQueueFenceSignal())
 		{
-			check(Fence.LastSignaledValue < Payload->CompletionFenceValue);
+			check(Queue.Fence.LastSignaledValue < Payload->CompletionFenceValue);
 
-			VERIFYD3D12RESULT(D3DCommandQueue->Signal(Fence.D3DFence, Payload->CompletionFenceValue));
-			Fence.LastSignaledValue.store(Payload->CompletionFenceValue, std::memory_order_release);
+			VERIFYD3D12RESULT(Queue.D3DCommandQueue->Signal(Queue.Fence.D3DFence, Payload->CompletionFenceValue));
+			Queue.Fence.LastSignaledValue.store(Payload->CompletionFenceValue, std::memory_order_release);
 		}
 
 		// Submission of this payload is completed. Signal the submission event if one was provided.
@@ -974,8 +973,17 @@ void FD3D12Queue::FlushBatchedPayloads()
 		}
 	};
 
+	FD3D12Queue* PrevQueue = nullptr;
 	for (FD3D12Payload* Payload : PayloadsToSubmit)
 	{
+		if (PrevQueue != &Payload->Queue)
+		{
+			Flush();
+			PrevQueue = &Payload->Queue;
+		}
+
+		Payload->Queue.PendingInterrupt.Enqueue(Payload);
+
 #if RHI_NEW_GPU_PROFILER
 		if (Payload->Timing.IsSet())
 		{
