Remove old driver work around which are below current minimum required driver version

diff --git a/Engine/Source/Runtime/D3D12RHI/Private/D3D12PoolAllocator.cpp b/Engine/Source/Runtime/D3D12RHI/Private/D3D12PoolAllocator.cpp
--- a/Engine/Source/Runtime/D3D12RHI/Private/D3D12PoolAllocator.cpp
+++ b/Engine/Source/Runtime/D3D12RHI/Private/D3D12PoolAllocator.cpp
@@ -336,17 +336,8 @@ void FD3D12PoolAllocator::AllocateResource(uint32 GPUIndex, D3D12_HEAP_TYPE InHe
 			bPlacedResource = false;
 		}
 	}
-
-	// Disable pooling for VRAM allocated textures and force use the committed resource path
-	bool bForceCommittedResourcePath = false;
-#if PLATFORM_WINDOWS
-	if (bPoolResource && GD3D12WorkaroundFlags.bForceCommittedResourceTextureAllocation && InHeapType == D3D12_HEAP_TYPE_DEFAULT && InDesc.Dimension != D3D12_RESOURCE_DIMENSION_BUFFER)
-	{
-		bForceCommittedResourcePath = true;
-	}
-#endif // PLATFORM_WINDOWS
-
-	if (bPoolResource && !bForceCommittedResourcePath)
+	   
+	if (bPoolResource)
 	{
 		uint32 AllocationAlignment = InAllocationAlignment;
 
@@ -425,7 +416,7 @@ void FD3D12PoolAllocator::AllocateResource(uint32 GPUIndex, D3D12_HEAP_TYPE InHe
 
 		// If we are tracking all allocation data and allocating a standalone texture, then first create a heap so we can retrieve the GPU virtual address as well
 		// UAV Aliasing needs a Heap to create the aliased resource in.
-		if (InDesc.NeedsUAVAliasWorkarounds() || (Adapter->IsTrackingAllAllocations() && InDesc.Dimension != D3D12_RESOURCE_DIMENSION_BUFFER && !bForceCommittedResourcePath))
+		if (InDesc.NeedsUAVAliasWorkarounds() || (Adapter->IsTrackingAllAllocations() && InDesc.Dimension != D3D12_RESOURCE_DIMENSION_BUFFER))
 		{
 			D3D12_HEAP_DESC HeapDesc = {};
 			HeapDesc.SizeInBytes = FMath::Max((uint64)D3D12_DEFAULT_RESOURCE_PLACEMENT_ALIGNMENT, InSize);
diff --git a/Engine/Source/Runtime/D3D12RHI/Private/D3D12RHIPrivate.h b/Engine/Source/Runtime/D3D12RHI/Private/D3D12RHIPrivate.h
--- a/Engine/Source/Runtime/D3D12RHI/Private/D3D12RHIPrivate.h
+++ b/Engine/Source/Runtime/D3D12RHI/Private/D3D12RHIPrivate.h
@@ -140,20 +140,11 @@ struct FD3D12RayTracingPipelineInfo
 	uint32 ScratchSize = 0;
 };
 
+/**
+* Struct used for driver work around on certain driver versions
+*/
 struct FD3D12WorkaroundFlags
 {
-	/** 
-	* Certain drivers crash when GetShaderIdentifier() is called on a ray tracing pipeline collection.
-	* If we detect such driver, we have to fall back to the path that queries identifiers on full linked RTPSO.
-	* This is less efficient and can also trigger another known issue with D3D12 Agility version <= 4.
-	*/
-	bool bAllowGetShaderIdentifierOnCollectionSubObject = true;
-
-	/**
-	* Certain drivers can cause texture corruption when sub allocating textures from a shared heap. Committed resource
-	* allocations are used then
-	*/
-	bool bForceCommittedResourceTextureAllocation = false;
 };
 
 extern FD3D12WorkaroundFlags GD3D12WorkaroundFlags;
diff --git a/Engine/Source/Runtime/D3D12RHI/Private/D3D12RayTracing.cpp b/Engine/Source/Runtime/D3D12RHI/Private/D3D12RayTracing.cpp
--- a/Engine/Source/Runtime/D3D12RHI/Private/D3D12RayTracing.cpp
+++ b/Engine/Source/Runtime/D3D12RHI/Private/D3D12RayTracing.cpp
@@ -699,26 +699,6 @@ static FD3D12ShaderIdentifier GetShaderIdentifier(ID3D12StateObject* StateObject
 	return GetShaderIdentifier(PipelineProperties, ExportName);
 }
 
-static bool TryGetShaderIdentifier(ID3D12StateObject* StateObject, const TCHAR* ExportName, FD3D12ShaderIdentifier& OutShaderIdentifier)
-{
-	TRefCountPtr<ID3D12StateObjectProperties> PipelineProperties;
-	HRESULT QueryInterfaceResult = StateObject->QueryInterface(IID_PPV_ARGS(PipelineProperties.GetInitReference()));
-	if (!SUCCEEDED(QueryInterfaceResult))
-	{
-		return false;
-	}
-	const void* ShaderIdData = PipelineProperties->GetShaderIdentifier(ExportName);
-	if (ShaderIdData)
-	{
-		OutShaderIdentifier.SetData(ShaderIdData);
-		return true;
-	}
-	else
-	{
-		return false;
-	}
-}
-
 FD3D12RayTracingCompactionRequestHandler::FD3D12RayTracingCompactionRequestHandler(FD3D12Device* Device)
 	: FD3D12DeviceChild(Device)
 {
@@ -1096,15 +1076,9 @@ public:
 				Device->GetRayTracingPipelineInfo(Entry.StateObject, &Entry.PipelineInfo);
 			}
 
-			// Shader identifier can be queried immediately here per PSO collection, however this does not work on old NVIDIA drivers (430.00).
-			// Therefore shader identifiers need to be queried from the final linked pipeline (JIRA DH-2182) if a known bad driver is detected.
-			if (GD3D12WorkaroundFlags.bAllowGetShaderIdentifierOnCollectionSubObject)
-			{
-				// Some DXR implementations (such as VKD3D) may not support shader ID queries from collections.
-				// If shader ID query fails, we will attempt to query it from the final linked RTPSO later.
-				TryGetShaderIdentifier(Entry.StateObject, Entry.GetPrimaryExportNameChars(), Entry.Identifier);
-			}
-
+			// Retrieve the identifier from the library
+			Entry.Identifier = GetShaderIdentifier(Entry.StateObject, Entry.GetPrimaryExportNameChars());
+			
 			CompileTimeCycles += FPlatformTime::Cycles64();
 
 			Entry.CompileTimeMS = float(FPlatformTime::ToMilliseconds64(CompileTimeCycles));
diff --git a/Engine/Source/Runtime/D3D12RHI/Private/Windows/WindowsD3D12Device.cpp b/Engine/Source/Runtime/D3D12RHI/Private/Windows/WindowsD3D12Device.cpp
--- a/Engine/Source/Runtime/D3D12RHI/Private/Windows/WindowsD3D12Device.cpp
+++ b/Engine/Source/Runtime/D3D12RHI/Private/Windows/WindowsD3D12Device.cpp
@@ -57,18 +57,6 @@ static FAutoConsoleVariableRef CVarMinDriverVersionForRayTracingNVIDIA(
 	ECVF_ReadOnly | ECVF_RenderThreadSafe
 );
 
-#define DXR_ALLOW_EMULATED_RAYTRACING 0
-
-#if DXR_ALLOW_EMULATED_RAYTRACING
-int32 GAllowEmulatedRayTracing = 0;
-static FAutoConsoleVariableRef CVarAllowEmulatedRayTracing(
-	TEXT("r.D3D12.DXR.AllowEmulatedRayTracing"),
-	GAllowEmulatedRayTracing,
-	TEXT("Allows ray tracing emulation support on NVIDIA cards with the Pascal architecture (default=0)."),
-	ECVF_ReadOnly | ECVF_RenderThreadSafe
-);
-#endif
-
 // Use AGS_MAKE_VERSION() macro to define the version.
 // i.e. AGS_MAKE_VERSION(major, minor, patch) ((major << 22) | (minor << 12) | patch)
 int32 GMinimumDriverVersionForRayTracingAMD = 0;
@@ -1490,28 +1478,8 @@ void FD3D12DynamicRHI::Init()
 		if (GRHISupportsRayTracing
 			&& IsRayTracingEmulated(AdapterDesc.DeviceId))
 		{
-#if DXR_ALLOW_EMULATED_RAYTRACING
-			if (!GAllowEmulatedRayTracing)
-			{
-				DisableRayTracingSupport();
-				UE_LOG(LogD3D12RHI, Warning, TEXT("Ray tracing is disabled for NVIDIA cards with the Pascal architecture. This can be overridden with the following CVar: r.D3D12.DXR.AllowEmulatedRayTracing=1"));
-			}
-#else
 			DisableRayTracingSupport();
 			UE_LOG(LogD3D12RHI, Warning, TEXT("Ray tracing is disabled for NVIDIA cards with the Pascal architecture."));
-#endif // DXR_ALLOW_EMULATED_RAYTRACING
-		}
-
-		if (GRHISupportsRayTracing && DriverVersion < 45700u)
-		{
-			GD3D12WorkaroundFlags.bAllowGetShaderIdentifierOnCollectionSubObject = false;
-			UE_LOG(LogD3D12RHI, Warning, TEXT("GD3D12WorkaroundFlags.bAllowGetShaderIdentifierOnCollectionSubObject is disabled due to a known issue with current driver version."));
-		}
-
-		if (DriverVersion < 53141u)
-		{
-			GD3D12WorkaroundFlags.bForceCommittedResourceTextureAllocation = true;
-			UE_LOG(LogD3D12RHI, Log, TEXT("GD3D12WorkaroundFlags.bForceCommittedResourceTextureAllocation is set due to a known issue with current driver version."));
 		}
 
 	} // if NVIDIA
