Don't terminate the process when creating a PSO returns E_OUTOFMEMORY.

This error might be recoverable if the creation request is coming from a precompile task. Let the caller decide what the appropriate action is.

#rnx

diff --git a/Engine/Source/Runtime/D3D12RHI/Private/Windows/WindowsD3D12PipelineState.cpp b/Engine/Source/Runtime/D3D12RHI/Private/Windows/WindowsD3D12PipelineState.cpp
--- a/Engine/Source/Runtime/D3D12RHI/Private/Windows/WindowsD3D12PipelineState.cpp
+++ b/Engine/Source/Runtime/D3D12RHI/Private/Windows/WindowsD3D12PipelineState.cpp
@@ -799,8 +799,8 @@ static void CreatePipelineStateWrapper(ID3D12PipelineState** PSO, FD3D12Adapter*
 			// Always dump the graphics PSO state - internal driver compiler error can still return DXGI_ERROR_DEVICE_REMOVED
 			DumpGraphicsPSO(CreationArgs->Desc.Desc, Name);
 
-			// First check if D3D device removed, hung or out of memory and handle that separately 
-			if (hr == DXGI_ERROR_DEVICE_REMOVED || hr == DXGI_ERROR_DEVICE_HUNG || hr == E_OUTOFMEMORY)
+			// Catch GPU crashes explicitly.
+			if (hr == DXGI_ERROR_DEVICE_REMOVED || hr == DXGI_ERROR_DEVICE_HUNG)
 			{
 				VERIFYD3D12RESULT_EX(hr, pDevice2);
 			}
@@ -817,8 +817,8 @@ static void CreatePipelineStateWrapper(ID3D12PipelineState** PSO, FD3D12Adapter*
 			// Always dump the graphics PSO state - internal driver compiler error can still return DXGI_ERROR_DEVICE_REMOVED
 			DumpGraphicsPSO(CreationArgs->Desc.Desc, Name);
 
-			// First check if D3D device removed, hung or out of memory and handle that separately 
-			if (hr == DXGI_ERROR_DEVICE_REMOVED || hr == DXGI_ERROR_DEVICE_HUNG || hr == E_OUTOFMEMORY)
+			// Catch GPU crashes explicitly.
+			if (hr == DXGI_ERROR_DEVICE_REMOVED || hr == DXGI_ERROR_DEVICE_HUNG)
 			{
 				VERIFYD3D12RESULT_EX(hr, pDevice2);
 			}
@@ -847,7 +847,7 @@ static void CreatePipelineStateWrapper(ID3D12PipelineState** PSO, FD3D12Adapter*
 	if (FAILED(hr))
 	{
 		// First check if D3D device removed, hung or out of memory and handle that separately 
-		if (hr == DXGI_ERROR_DEVICE_REMOVED || hr == DXGI_ERROR_DEVICE_HUNG || hr == E_OUTOFMEMORY)
+		if (hr == DXGI_ERROR_DEVICE_REMOVED || hr == DXGI_ERROR_DEVICE_HUNG)
 		{
 			VERIFYD3D12RESULT_EX(hr, pDevice2);
 		}
