Deprecate ray tracing scene layers in RHI.

diff --git a/Engine/Plugins/Experimental/GPULightmass/Source/GPULightmass/Private/LightmapRenderer.cpp b/Engine/Plugins/Experimental/GPULightmass/Source/GPULightmass/Private/LightmapRenderer.cpp
--- a/Engine/Plugins/Experimental/GPULightmass/Source/GPULightmass/Private/LightmapRenderer.cpp
+++ b/Engine/Plugins/Experimental/GPULightmass/Source/GPULightmass/Private/LightmapRenderer.cpp
@@ -1137,7 +1137,7 @@ bool FSceneRenderState::SetupRayTracingScene(FRDGBuilder& GraphBuilder, FSceneUn
 
 			const FRayTracingSceneInitializer2& SceneInitializer = RayTracingScene->GetInitializer();
 
-			FRayTracingAccelerationStructureSize SizeInfo = RHICalcRayTracingSceneSize(SceneInitializer.NumNativeInstancesPerLayer[0], ERayTracingAccelerationStructureFlags::FastTrace);
+			FRayTracingAccelerationStructureSize SizeInfo = RHICalcRayTracingSceneSize(SceneInitializer.NumNativeInstances, ERayTracingAccelerationStructureFlags::FastTrace);
 			FRHIResourceCreateInfo BufferCreateInfo(TEXT("LightmassRayTracingSceneBuffer"));
 			RayTracingSceneBuffer = RHICmdList.CreateBuffer(uint32(SizeInfo.ResultSize), BUF_AccelerationStructure, 0, ERHIAccess::BVHWrite, BufferCreateInfo);
 			RayTracingSceneSRV = RHICmdList.CreateShaderResourceView(RayTracingSceneBuffer);
@@ -1149,7 +1149,7 @@ bool FSceneRenderState::SetupRayTracingScene(FRDGBuilder& GraphBuilder, FSceneUn
 				ScratchBufferCreateInfo);
 
 			FRWBufferStructured InstanceBuffer;
-			InstanceBuffer.Initialize(RHICmdList, TEXT("LightmassRayTracingInstanceBuffer"), GRHIRayTracingInstanceDescriptorSize, SceneInitializer.NumNativeInstancesPerLayer[0]);
+			InstanceBuffer.Initialize(RHICmdList, TEXT("LightmassRayTracingInstanceBuffer"), GRHIRayTracingInstanceDescriptorSize, SceneInitializer.NumNativeInstances);
 
 			// Need to pass "BUF_MultiGPUAllocate", as acceleration structure virtual addresses are different per GPU
 			FByteAddressBuffer AccelerationStructureAddressesBuffer;
@@ -1159,7 +1159,7 @@ bool FSceneRenderState::SetupRayTracingScene(FRDGBuilder& GraphBuilder, FSceneUn
 				SceneInitializer.ReferencedGeometries.Num() * sizeof(FRayTracingAccelerationStructureAddress),
 				BUF_Volatile | BUF_MultiGPUAllocate);
 
-			const uint32 InstanceUploadBufferSize = SceneInitializer.NumNativeInstancesPerLayer[0] * sizeof(FRayTracingInstanceDescriptorInput);
+			const uint32 InstanceUploadBufferSize = SceneInitializer.NumNativeInstances * sizeof(FRayTracingInstanceDescriptorInput);
 			FBufferRHIRef InstanceUploadBuffer;
 			FShaderResourceViewRHIRef InstanceUploadSRV;
 			{
@@ -1189,7 +1189,7 @@ bool FSceneRenderState::SetupRayTracingScene(FRDGBuilder& GraphBuilder, FSceneUn
 					SceneWithGeometryInstances.BaseInstancePrefixSum,
 					SceneWithGeometryInstances.NumNativeGPUSceneInstances,
 					SceneWithGeometryInstances.NumNativeCPUInstances,
-					MakeArrayView(InstanceUploadData, SceneInitializer.NumNativeInstancesPerLayer[0]),
+					MakeArrayView(InstanceUploadData, SceneInitializer.NumNativeInstances),
 					MakeArrayView(TransformUploadData, SceneWithGeometryInstances.NumNativeCPUInstances * 3));
 				RHICmdList.UnlockBuffer(TransformUploadBuffer);
 				RHICmdList.UnlockBuffer(InstanceUploadBuffer);
diff --git a/Engine/Source/Runtime/Apple/MetalRHI/Private/MetalRayTracing.cpp b/Engine/Source/Runtime/Apple/MetalRHI/Private/MetalRayTracing.cpp
--- a/Engine/Source/Runtime/Apple/MetalRHI/Private/MetalRayTracing.cpp
+++ b/Engine/Source/Runtime/Apple/MetalRHI/Private/MetalRayTracing.cpp
@@ -435,7 +435,9 @@ FMetalRayTracingScene::FMetalRayTracingScene(FRayTracingSceneInitializer2 InInit
 {
 	SizeInfo = {};
 
+	PRAGMA_DISABLE_DEPRECATION_WARNINGS
 	const uint32 NumLayers = Initializer.NumNativeInstancesPerLayer.Num();
+	PRAGMA_ENABLE_DEPRECATION_WARNINGS
 	check(NumLayers > 0);
 
 	Layers.SetNum(NumLayers);
@@ -445,7 +447,9 @@ FMetalRayTracingScene::FMetalRayTracingScene(FRayTracingSceneInitializer2 InInit
 		FLayerData& Layer = Layers[LayerIndex];
 
 		MTL::InstanceAccelerationStructureDescriptor* InstanceDescriptor;
+		PRAGMA_DISABLE_DEPRECATION_WARNINGS
 		InstanceDescriptor.SetInstanceCount(Initializer.NumNativeInstancesPerLayer[LayerIndex]);
+		PRAGMA_ENABLE_DEPRECATION_WARNINGS
 
 		Layer.SizeInfo = CalcRayTracingGeometrySize(InstanceDescriptor);
 
@@ -620,7 +624,9 @@ void FMetalRayTracingScene::BuildAccelerationStructure(
 	MTL::AccelerationStructureCommandEncoder* CommandEncoder = Encoder.GetAccelerationStructureCommandEncoder();
 	check(CommandEncoder.GetPtr());
 
+	PRAGMA_DISABLE_DEPRECATION_WARNINGS
 	const uint32 NumLayers = Initializer.NumNativeInstancesPerLayer.Num();
+	PRAGMA_ENABLE_DEPRECATION_WARNINGS
 
 	FMetalBuffer CurScratchBuffer = InScratchBuffer->GetCurrentBuffer();
 	check(CurScratchBuffer);
@@ -631,7 +637,9 @@ void FMetalRayTracingScene::BuildAccelerationStructure(
 		FLayerData& Layer = Layers[LayerIndex];
 
 		MTL::InstanceAccelerationStructureDescriptor* InstanceDescriptor = MTL::InstanceAccelerationStructureDescriptor();
+		PRAGMA_DISABLE_DEPRECATION_WARNINGS
 		InstanceDescriptor.SetInstanceCount(Initializer.NumNativeInstancesPerLayer[LayerIndex]);
+		PRAGMA_ENABLE_DEPRECATION_WARNINGS
 		InstanceDescriptor.SetInstanceDescriptorBuffer(CurInstanceBuffer);
 		InstanceDescriptor.SetInstanceDescriptorBufferOffset(InstanceBufferOffset);
 		InstanceDescriptor.SetInstancedAccelerationStructures((__bridge NSArray*)MutableAccelerationStructures);
diff --git a/Engine/Source/Runtime/D3D12RHI/Private/D3D12RayTracing.cpp b/Engine/Source/Runtime/D3D12RHI/Private/D3D12RayTracing.cpp
--- a/Engine/Source/Runtime/D3D12RHI/Private/D3D12RayTracing.cpp
+++ b/Engine/Source/Runtime/D3D12RHI/Private/D3D12RayTracing.cpp
@@ -3439,7 +3439,9 @@ FD3D12RayTracingScene::FD3D12RayTracingScene(FD3D12Adapter* Adapter, FRayTracing
 
 	SizeInfo = {};
 
+	PRAGMA_DISABLE_DEPRECATION_WARNINGS
 	const uint32 NumLayers = Initializer.NumNativeInstancesPerLayer.Num();
+	PRAGMA_ENABLE_DEPRECATION_WARNINGS
 	check(NumLayers > 0);
 
 	Layers.SetNum(NumLayers);
@@ -3451,7 +3453,9 @@ FD3D12RayTracingScene::FD3D12RayTracingScene(FD3D12Adapter* Adapter, FRayTracing
 		Layer.BuildInputs = {};
 		Layer.BuildInputs.Type = D3D12_RAYTRACING_ACCELERATION_STRUCTURE_TYPE_TOP_LEVEL;
 		Layer.BuildInputs.DescsLayout = D3D12_ELEMENTS_LAYOUT_ARRAY;
+		PRAGMA_DISABLE_DEPRECATION_WARNINGS
 		Layer.BuildInputs.NumDescs = Initializer.NumNativeInstancesPerLayer[LayerIndex];
+		PRAGMA_ENABLE_DEPRECATION_WARNINGS
 		Layer.BuildInputs.Flags = TranslateRayTracingAccelerationStructureFlags(Initializer.BuildFlags);
 
 		// Get maximum buffer sizes for all GPUs in the system
@@ -3542,7 +3546,9 @@ void FD3D12RayTracingScene::BuildAccelerationStructure(FD3D12CommandContext& Com
 
 	checkf(ScratchBuffer, TEXT("TLAS build requires scratch buffer of at least %lld bytes."), SizeInfo.BuildScratchSize);
 
+	PRAGMA_DISABLE_DEPRECATION_WARNINGS
 	const uint32 NumLayers = Initializer.NumNativeInstancesPerLayer.Num();
+	PRAGMA_ENABLE_DEPRECATION_WARNINGS
 
 	for (uint32 LayerIndex = 0; LayerIndex < NumLayers; ++LayerIndex)
 	{
@@ -3566,21 +3572,25 @@ void FD3D12RayTracingScene::BuildAccelerationStructure(FD3D12CommandContext& Com
 
 		if (bIsUpdate)
 		{
+			PRAGMA_DISABLE_DEPRECATION_WARNINGS
 			checkf(ScratchBufferOffset + Layer.UpdateScratchOffset + PrebuildInfo.UpdateScratchDataSizeInBytes <= ScratchBuffer->GetSize(),
 				TEXT("TLAS scratch buffer size is %d bytes with offset %d (%d bytes available), but the update requires %lld bytes. ")
 				TEXT("BuildInputs.NumDescs = %d, Instances.Num = %d, PerInstanceGeometries.Num = %d."),
 				ScratchBuffer->GetSize(), ScratchBufferOffset + Layer.UpdateScratchOffset, ScratchBuffer->GetSize() - ScratchBufferOffset,
 				PrebuildInfo.UpdateScratchDataSizeInBytes,
 				Layer.BuildInputs.NumDescs, Initializer.NumNativeInstancesPerLayer[LayerIndex], Initializer.PerInstanceGeometries.Num());
+			PRAGMA_ENABLE_DEPRECATION_WARNINGS
 		}
 		else
 		{
+			PRAGMA_DISABLE_DEPRECATION_WARNINGS
 			checkf(ScratchBufferOffset + Layer.BuildScratchOffset + PrebuildInfo.ScratchDataSizeInBytes <= ScratchBuffer->GetSize(),
 				TEXT("TLAS scratch buffer size is %d bytes with offset %d (%d bytes available), but the build requires %lld bytes. ")
 				TEXT("BuildInputs.NumDescs = %d, Instances.Num = %d, PerInstanceGeometries.Num = %d."),
 				ScratchBuffer->GetSize(), ScratchBufferOffset + Layer.BuildScratchOffset, ScratchBuffer->GetSize() - ScratchBufferOffset,
 				PrebuildInfo.ScratchDataSizeInBytes,
 				Layer.BuildInputs.NumDescs, Initializer.NumNativeInstancesPerLayer[LayerIndex], Initializer.PerInstanceGeometries.Num());
+			PRAGMA_ENABLE_DEPRECATION_WARNINGS
 		}
 	}
 
@@ -3743,7 +3753,9 @@ void FD3D12RayTracingScene::BuildAccelerationStructure(FD3D12CommandContext& Com
 			INC_DWORD_STAT(STAT_D3D12RayTracingBuiltTLAS);
 		}
 
+		PRAGMA_DISABLE_DEPRECATION_WARNINGS
 		InstanceBaseOffset += Initializer.NumNativeInstancesPerLayer[LayerIndex];
+		PRAGMA_ENABLE_DEPRECATION_WARNINGS
 	}
 
 	CommandContext.BuildAccelerationStructuresInternal(BuildDescs);
diff --git a/Engine/Source/Runtime/RHI/Public/RHIResources.h b/Engine/Source/Runtime/RHI/Public/RHIResources.h
--- a/Engine/Source/Runtime/RHI/Public/RHIResources.h
+++ b/Engine/Source/Runtime/RHI/Public/RHIResources.h
@@ -3066,6 +3066,7 @@ struct FRayTracingGeometryInstance
 	// If binary AND of instance mask with ray mask is zero, then the instance is considered not intersected / invisible.
 	uint8 Mask = 0xFF;
 
+	UE_DEPRECATED(5.5, "Specify layer when adding instances to FRayTracingScene instead.")
 	uint8 LayerIndex = 0;
 
 	// Flags to control triangle back face culling, whether to allow any-hit shaders, etc.
@@ -3285,9 +3286,9 @@ struct FRayTracingSceneInitializer2
 	TArray<uint32> SegmentPrefixSum;
 
 	// Total flattened number of ray tracing geometry instances (a single FRayTracingGeometryInstance may represent many) per layer.
+	UE_DEPRECATED(5.5, "FRHIRayTracingScene layers are deprecated. Use NumNativeInstances and create one FRHIRayTracingScene per layer instead.")
 	TArray<uint32> NumNativeInstancesPerLayer;
 
-	UE_DEPRECATED(5.1, "Use NumNativeInstancesPerLayer instead.")
 	uint32 NumNativeInstances = 0;
 
 	uint32 NumTotalSegments = 0;
@@ -3385,6 +3386,7 @@ public:
 		return nullptr;
 	}
 
+	UE_DEPRECATED(5.5, "FRHIRayTracingScene layers are deprecated. Create one FRHIRayTracingScene per layer instead.")
 	virtual uint32 GetLayerBufferOffset(uint32 LayerIndex) const = 0;
 
 	UE_DEPRECATED(5.5, "Create standalone FRHIShaderBindingTable instead.")
diff --git a/Engine/Source/Runtime/Renderer/Private/PathTracing.cpp b/Engine/Source/Runtime/Renderer/Private/PathTracing.cpp
--- a/Engine/Source/Runtime/Renderer/Private/PathTracing.cpp
+++ b/Engine/Source/Runtime/Renderer/Private/PathTracing.cpp
@@ -793,7 +793,7 @@ static void PreparePathTracingData(const FScene* Scene, const FViewInfo& View, F
 
 	PathTracingData.DecalRoughnessCutoff = PathTracing::UsesDecals(*View.Family) && View.bHasRayTracingDecals ? CVarPathTracingDecalRoughnessCutoff.GetValueOnRenderThread() : -1.0f;
 
-	PathTracingData.MeshDecalRoughnessCutoff = PathTracing::UsesDecals(*View.Family) && Scene->RayTracingScene.GetRHIRayTracingScene(ERayTracingSceneLayer::Decals)->GetInitializer().NumNativeInstancesPerLayer[0] > 0 ? CVarPathTracingMeshDecalRoughnessCutoff.GetValueOnRenderThread() : -1.0f;
+	PathTracingData.MeshDecalRoughnessCutoff = PathTracing::UsesDecals(*View.Family) && Scene->RayTracingScene.GetRHIRayTracingScene(ERayTracingSceneLayer::Decals)->GetInitializer().NumNativeInstances > 0 ? CVarPathTracingMeshDecalRoughnessCutoff.GetValueOnRenderThread() : -1.0f;
 	PathTracingData.MeshDecalBias = CVarPathTracingMeshDecalBias.GetValueOnRenderThread();
 
 	PathTracingData.MaxRaymarchSteps = CVarPathTracingMaxRaymarchSteps.GetValueOnRenderThread();
diff --git a/Engine/Source/Runtime/Renderer/Private/RayTracing/RayTracingDebug.cpp b/Engine/Source/Runtime/Renderer/Private/RayTracing/RayTracingDebug.cpp
--- a/Engine/Source/Runtime/Renderer/Private/RayTracing/RayTracingDebug.cpp
+++ b/Engine/Source/Runtime/Renderer/Private/RayTracing/RayTracingDebug.cpp
@@ -1076,7 +1076,7 @@ static FRDGBufferRef RayTracingPerformHitStatsPerPrimitive(FRDGBuilder& GraphBui
 
 	FShaderBindingTableRHIRef HitStatsSBT = RHICreateShaderBindingTable(SBTInitializer);
 
-	const uint32 NumInstancesInTLAS = FMath::Max(RayTracingScene.GetRHIRayTracingSceneChecked(ERayTracingSceneLayer::Base)->GetInitializer().NumNativeInstancesPerLayer[0], (uint32)CVarRayTracingDebugHitCountTopKHits.GetValueOnRenderThread());
+	const uint32 NumInstancesInTLAS = FMath::Max(RayTracingScene.GetRHIRayTracingSceneChecked(ERayTracingSceneLayer::Base)->GetInitializer().NumNativeInstances, (uint32)CVarRayTracingDebugHitCountTopKHits.GetValueOnRenderThread());
 	FRDGBufferDesc HitStatsPerPrimitiveBufferDesc = FRDGBufferDesc::CreateStructuredDesc(sizeof(FRayTracingHitStatsEntry), NumInstancesInTLAS);
 	HitStatsPerPrimitiveBufferDesc.Usage = EBufferUsageFlags(HitStatsPerPrimitiveBufferDesc.Usage | BUF_SourceCopy);
 	FRDGBufferRef HitStatsBuffer = GraphBuilder.CreateBuffer(HitStatsPerPrimitiveBufferDesc, TEXT("RayTracingDebug.HitStatsBuffer"));
diff --git a/Engine/Source/Runtime/Renderer/Private/RayTracing/RayTracingInstanceBufferUtil.cpp b/Engine/Source/Runtime/Renderer/Private/RayTracing/RayTracingInstanceBufferUtil.cpp
--- a/Engine/Source/Runtime/Renderer/Private/RayTracing/RayTracingInstanceBufferUtil.cpp
+++ b/Engine/Source/Runtime/Renderer/Private/RayTracing/RayTracingInstanceBufferUtil.cpp
@@ -109,22 +109,28 @@ FRayTracingSceneWithGeometryInstances CreateRayTracingSceneWithGeometryInstances
 			checkNoEntry();
 		}
 
+		PRAGMA_DISABLE_DEPRECATION_WARNINGS
 		checkf(InstanceDesc.LayerIndex < NumLayers, 
 			TEXT("FRayTracingGeometryInstance is assigned to layer %d but raytracing scene being created only has %d layers."),
 			InstanceDesc.LayerIndex, NumLayers);
 
-		PRAGMA_DISABLE_DEPRECATION_WARNINGS
 		Initializer.BaseInstancePrefixSum[InstanceIndex] = Initializer.NumNativeInstancesPerLayer[InstanceDesc.LayerIndex];
-		PRAGMA_ENABLE_DEPRECATION_WARNINGS
 
 		// Can't support same instance in multiple layers because BaseInstancePrefixSum would be different per layer
 		Output.BaseInstancePrefixSum[InstanceIndex] = Initializer.NumNativeInstancesPerLayer[InstanceDesc.LayerIndex];
 
 		Initializer.NumNativeInstancesPerLayer[InstanceDesc.LayerIndex] += InstanceDesc.NumTransforms;
+		PRAGMA_ENABLE_DEPRECATION_WARNINGS
 	}
 
 	PRAGMA_DISABLE_DEPRECATION_WARNINGS
 	Initializer.NumTotalSegments = Output.TotalNumSegments;
+
+	if(NumLayers == 1)
+	{
+		Initializer.NumNativeInstances = Initializer.NumNativeInstancesPerLayer[0];
+		Initializer.NumNativeInstancesPerLayer.Empty();
+	}
 	PRAGMA_ENABLE_DEPRECATION_WARNINGS
 
 	Output.Scene = RHICreateRayTracingScene(MoveTemp(Initializer));
@@ -156,6 +162,8 @@ void FillRayTracingInstanceUploadBuffer(
 	TRACE_CPUPROFILER_EVENT_SCOPE(FillRayTracingInstanceUploadBuffer);
 
 	const FRayTracingSceneInitializer2& SceneInitializer = RayTracingSceneRHI->GetInitializer();
+
+	PRAGMA_DISABLE_DEPRECATION_WARNINGS
 	const uint32 NumLayers = SceneInitializer.NumNativeInstancesPerLayer.Num();
 
 	TArray<uint32> LayerBaseIndices;
@@ -166,6 +174,7 @@ void FillRayTracingInstanceUploadBuffer(
 	{
 		LayerBaseIndices[LayerIndex] = LayerBaseIndices[LayerIndex - 1] + SceneInitializer.NumNativeInstancesPerLayer[LayerIndex - 1];
 	}
+	PRAGMA_ENABLE_DEPRECATION_WARNINGS
 
 	const int32 NumSceneInstances = Instances.Num();
 	const int32 MinBatchSize = 128;
@@ -202,7 +211,9 @@ void FillRayTracingInstanceUploadBuffer(
 			checkf(bGpuSceneInstance + bCpuInstance == 1, TEXT("Instance can only get transforms from one of GPUScene, or Transforms array."));
 
 			const uint32 AccelerationStructureIndex = InstanceGeometryIndices[SceneInstanceIndex];
+			PRAGMA_DISABLE_DEPRECATION_WARNINGS
 			const uint32 LayerBaseIndex = LayerBaseIndices[SceneInstance.LayerIndex];
+			PRAGMA_ENABLE_DEPRECATION_WARNINGS
 			const uint32 BaseInstanceIndex = BaseInstancePrefixSum[SceneInstanceIndex];
 			const uint32 BaseTransformIndex = bCpuInstance ? BaseUploadBufferOffsets[SceneInstanceIndex] : 0;
 
diff --git a/Engine/Source/Runtime/Renderer/Private/Tests/RayTracingTestbed.cpp b/Engine/Source/Runtime/Renderer/Private/Tests/RayTracingTestbed.cpp
--- a/Engine/Source/Runtime/Renderer/Private/Tests/RayTracingTestbed.cpp
+++ b/Engine/Source/Runtime/Renderer/Private/Tests/RayTracingTestbed.cpp
@@ -201,12 +201,12 @@ bool RunRayTracingTestbed_RenderThread(const FString& Parameters)
 		ScratchBufferCreateInfo);
 
 	FRWBufferStructured InstanceBuffer;
-	InstanceBuffer.Initialize(RHICmdList, TEXT("RayTracingTestBedInstanceBuffer"), GRHIRayTracingInstanceDescriptorSize, SceneInitializer.NumNativeInstancesPerLayer[0]);
+	InstanceBuffer.Initialize(RHICmdList, TEXT("RayTracingTestBedInstanceBuffer"), GRHIRayTracingInstanceDescriptorSize, SceneInitializer.NumNativeInstances);
 
 	FByteAddressBuffer AccelerationStructureAddressesBuffer;
 	AccelerationStructureAddressesBuffer.Initialize(RHICmdList, TEXT("RayTracingTestBedAccelerationStructureAddressesBuffer"), sizeof(FRayTracingAccelerationStructureAddress), BUF_Volatile | BUF_MultiGPUAllocate);
 
-	const uint32 InstanceUploadBufferSize = SceneInitializer.NumNativeInstancesPerLayer[0] * sizeof(FRayTracingInstanceDescriptorInput);
+	const uint32 InstanceUploadBufferSize = SceneInitializer.NumNativeInstances * sizeof(FRayTracingInstanceDescriptorInput);
 	FBufferRHIRef InstanceUploadBuffer;
 	FShaderResourceViewRHIRef InstanceUploadSRV;
 	{
@@ -236,7 +236,7 @@ bool RunRayTracingTestbed_RenderThread(const FString& Parameters)
 			RayTracingScene.BaseInstancePrefixSum,
 			RayTracingScene.NumNativeGPUSceneInstances,
 			RayTracingScene.NumNativeCPUInstances,
-			MakeArrayView(InstanceUploadData, SceneInitializer.NumNativeInstancesPerLayer[0]),
+			MakeArrayView(InstanceUploadData, SceneInitializer.NumNativeInstances),
 			MakeArrayView(TransformUploadData, RayTracingScene.NumNativeCPUInstances * 3));
 		RHICmdList.UnlockBuffer(TransformUploadBuffer);
 		RHICmdList.UnlockBuffer(InstanceUploadBuffer);
diff --git a/Engine/Source/Runtime/VulkanRHI/Private/VulkanRayTracing.cpp b/Engine/Source/Runtime/VulkanRHI/Private/VulkanRayTracing.cpp
--- a/Engine/Source/Runtime/VulkanRHI/Private/VulkanRayTracing.cpp
+++ b/Engine/Source/Runtime/VulkanRHI/Private/VulkanRayTracing.cpp
@@ -671,7 +671,9 @@ FVulkanRayTracingScene::FVulkanRayTracingScene(FRayTracingSceneInitializer2 InIn
 
 	SizeInfo = {};
 
+	PRAGMA_DISABLE_DEPRECATION_WARNINGS
 	const uint32 NumLayers = Initializer.NumNativeInstancesPerLayer.Num();
+	PRAGMA_ENABLE_DEPRECATION_WARNINGS
 	check(NumLayers > 0);
 
 	Layers.SetNum(NumLayers);
@@ -680,7 +682,9 @@ FVulkanRayTracingScene::FVulkanRayTracingScene(FRayTracingSceneInitializer2 InIn
 	{
 		FLayerData& Layer = Layers[LayerIndex];
 
+		PRAGMA_DISABLE_DEPRECATION_WARNINGS
 		Layer.SizeInfo = RHICalcRayTracingSceneSize(Initializer.NumNativeInstancesPerLayer[LayerIndex], Initializer.BuildFlags);
+		PRAGMA_ENABLE_DEPRECATION_WARNINGS
 		Layer.BufferOffset = Align(SizeInfo.ResultSize, GRHIRayTracingAccelerationStructureAlignment);
 		Layer.BuildScratchOffset = Align(SizeInfo.BuildScratchSize, GRHIRayTracingScratchBufferAlignment);
 		Layer.UpdateScratchOffset = Align(SizeInfo.UpdateScratchSize, GRHIRayTracingScratchBufferAlignment);
@@ -778,7 +782,9 @@ void FVulkanRayTracingScene::BuildAccelerationStructure(
 		InScratchOffset = 0;
 	}
 
+	PRAGMA_DISABLE_DEPRECATION_WARNINGS
 	const uint32 NumLayers = Initializer.NumNativeInstancesPerLayer.Num();
+	PRAGMA_ENABLE_DEPRECATION_WARNINGS
 
 	TArray<FVkRtTLASBuildData> BuildDatas;
 	BuildDatas.SetNum(NumLayers);
@@ -800,7 +806,9 @@ void FVulkanRayTracingScene::BuildAccelerationStructure(
 		const FLayerData& Layer = Layers[LayerIndex];
 
 		FVkRtTLASBuildData& BuildData = BuildDatas[LayerIndex];
+		PRAGMA_DISABLE_DEPRECATION_WARNINGS
 		GetTLASBuildData(Device->GetInstanceHandle(), Initializer.NumNativeInstancesPerLayer[LayerIndex], InstanceBufferAddress, Initializer.BuildFlags, BuildMode, BuildData);
+		PRAGMA_ENABLE_DEPRECATION_WARNINGS
 
 		const uint64 LayerScratchOffset = bIsUpdate ? Layer.UpdateScratchOffset : Layer.BuildScratchOffset;
 
@@ -812,7 +820,9 @@ void FVulkanRayTracingScene::BuildAccelerationStructure(
 		GeometryInfos[LayerIndex] = BuildData.GeometryInfo;
 
 		VkAccelerationStructureBuildRangeInfoKHR& TLASBuildRangeInfo = BuildRanges[LayerIndex];
+		PRAGMA_DISABLE_DEPRECATION_WARNINGS
 		TLASBuildRangeInfo.primitiveCount = Initializer.NumNativeInstancesPerLayer[LayerIndex];
+		PRAGMA_ENABLE_DEPRECATION_WARNINGS
 		TLASBuildRangeInfo.primitiveOffset = InstanceBaseOffset;
 		TLASBuildRangeInfo.transformOffset = 0;
 		TLASBuildRangeInfo.firstVertex = 0;
@@ -828,7 +838,9 @@ void FVulkanRayTracingScene::BuildAccelerationStructure(
 			INC_DWORD_STAT(STAT_VulkanRayTracingBuiltTLAS);
 		}
 
+		PRAGMA_DISABLE_DEPRECATION_WARNINGS
 		InstanceBaseOffset += Initializer.NumNativeInstancesPerLayer[LayerIndex];
+		PRAGMA_ENABLE_DEPRECATION_WARNINGS
 	}
 
 	FVulkanCommandBufferManager& CommandBufferManager = *CommandContext.GetCommandBufferManager();
