Always defer view creation on D3D12 to avoid race condition with Defrag.
 - Defrag can run on the RHI thread at the same moment a view is created on the Render Thread, which can cause the view to reference stale buffer state.

diff --git a/Engine/Source/Runtime/D3D12RHI/Private/D3D12SRV.cpp b/Engine/Source/Runtime/D3D12RHI/Private/D3D12SRV.cpp
--- a/Engine/Source/Runtime/D3D12RHI/Private/D3D12SRV.cpp
+++ b/Engine/Source/Runtime/D3D12RHI/Private/D3D12SRV.cpp
@@ -325,8 +325,6 @@ FShaderResourceViewRHIRef FD3D12DynamicRHI::RHICreateShaderResourceView(class FR
 		return new FD3D12ShaderResourceView_RHI(Device, TargetResource, ViewDesc);
 	});
 
-	bool bDynamic = View->IsBuffer() && EnumHasAnyFlags(View->GetBuffer()->GetUsage(), EBufferUsageFlags::AnyDynamic);
-	View->CreateViews(RHICmdList, bDynamic);
-
+	View->CreateViews(RHICmdList);
 	return View;
 }
diff --git a/Engine/Source/Runtime/D3D12RHI/Private/D3D12UAV.cpp b/Engine/Source/Runtime/D3D12RHI/Private/D3D12UAV.cpp
--- a/Engine/Source/Runtime/D3D12RHI/Private/D3D12UAV.cpp
+++ b/Engine/Source/Runtime/D3D12RHI/Private/D3D12UAV.cpp
@@ -266,9 +266,7 @@ FUnorderedAccessViewRHIRef FD3D12DynamicRHI::RHICreateUnorderedAccessView(FRHICo
 		return new FD3D12UnorderedAccessView_RHI(Device, TargetResource, ViewDesc);
 	});
 
-	bool bDynamic = View->IsBuffer() && EnumHasAnyFlags(View->GetBuffer()->GetUsage(), EBufferUsageFlags::AnyDynamic);
-	View->CreateViews(RHICmdList, bDynamic);
-
+	View->CreateViews(RHICmdList);
 	return View;
 }
 
diff --git a/Engine/Source/Runtime/D3D12RHI/Private/D3D12View.h b/Engine/Source/Runtime/D3D12RHI/Private/D3D12View.h
--- a/Engine/Source/Runtime/D3D12RHI/Private/D3D12View.h
+++ b/Engine/Source/Runtime/D3D12RHI/Private/D3D12View.h
@@ -394,7 +394,7 @@ private:
 template <typename TParent>
 struct FD3D12DeferredInitView : public FD3D12LinkedAdapterObject<TParent>
 {
-	void CreateViews(FRHICommandListBase& RHICmdList, bool bDynamic)
+	void CreateViews(FRHICommandListBase& RHICmdList)
 	{
 		auto InitLambda = [this](FRHICommandListBase&)
 		{
@@ -404,7 +404,7 @@ struct FD3D12DeferredInitView : public FD3D12LinkedAdapterObject<TParent>
 			}
 		};
 
-		if (RHICmdList.IsTopOfPipe() && bDynamic)
+		if (RHICmdList.IsTopOfPipe())
 		{
 			// We have to defer the view initialization to the RHI thread if the resource is dynamic (and RHI threading is enabled), since dynamic resources can be renamed.
 			// Also insert an RHI thread fence to prevent parallel translate tasks running until this command has completed.
