Fix bindless resource heap versioning with async compute enabled - single active heap for all queues

diff --git a/Engine/Source/Runtime/D3D12RHI/Private/D3D12BindlessDescriptors.cpp b/Engine/Source/Runtime/D3D12RHI/Private/D3D12BindlessDescriptors.cpp
--- a/Engine/Source/Runtime/D3D12RHI/Private/D3D12BindlessDescriptors.cpp
+++ b/Engine/Source/Runtime/D3D12RHI/Private/D3D12BindlessDescriptors.cpp
@@ -378,10 +378,11 @@ void FD3D12BindlessResourceManager::UpdateDescriptor(FD3D12ContextArray const& C
 				{
 					FD3D12CommandContext& Context = *ContextBase->GetSingleDeviceContext(GPUIndex);
 					Context.GetBindlessState().bRequestNewGpuHeap = true;
-					check(Context.IsDefaultContext() || !IsRunningRHIInSeparateThread());
+					check(!Context.GetExecutingCommandList().AllowParallelTranslate());
 				}
 			}
 
+			bRequestNewActiveGpuHeap = true;
 			ActiveGpuHeaps[ActiveGpuHeapIndex].UpdatedHandles.Add(DstHandle);
 		}
 
@@ -452,19 +453,6 @@ void FD3D12BindlessResourceManager::FinalizeContext(FD3D12CommandContext& Contex
 	// (can happen due to texture reference updates not adding any real GPU work)
 	FinalizeHeapOnState(State);
 
-	if (State.UsedHeaps.Num() > 0)
-	{
-		for (const FD3D12DescriptorHeapPtr& UsedHeap : State.UsedHeaps)
-		{
-			checkSlow(UsedHeap);
-
-			// Now queue it up for recycle when GPU is done
-			UE::D3D12BindlessDescriptors::DeferredFreeHeap(GetParentDevice(), UsedHeap);
-		}
-
-		State.UsedHeaps.Empty();
-	}
-
 	check(!Context.GetBindlessState().HasAnyPending());
 }
 
@@ -525,77 +513,88 @@ void FD3D12BindlessResourceManager::AssignHeapToState(FD3D12ContextBindlessState
 
 void FD3D12BindlessResourceManager::FinalizeHeapOnState(FD3D12ContextBindlessState& State)
 {
-	if (State.bRequestNewGpuHeap)
+	// Possibly version the GPU heap if it not requested by another queue yet
+	CheckRequestNewActiveGPUHeap();
+
+	// Clear the state data
+	State.CurrentGpuHeap = nullptr;
+	State.bRequestNewGpuHeap = false;
+}
+
+void FD3D12BindlessResourceManager::CheckRequestNewActiveGPUHeap()
+{
+	if (!bRequestNewActiveGpuHeap)
 	{
-		TRACE_CPUPROFILER_EVENT_SCOPE("FD3D12BindlessResourceManager::FinalizeHeapOnState");
+		return;
+	}
 
-		check(IsInRHIThread() || GRHICommandList.Bypass());
+	TRACE_CPUPROFILER_EVENT_SCOPE("FD3D12BindlessResourceManager::RequestNewActiveGPUHeap");
 
-		FScopeLock ScopeLock(&GpuHeapsCS);
+	FScopeLock ScopeLock(&GpuHeapsCS);
 
-		check(ActiveGpuHeaps[ActiveGpuHeapIndex].GpuHeap == State.CurrentGpuHeap || State.CurrentGpuHeap == nullptr);
+	if (!bRequestNewActiveGpuHeap)
+	{
+		return;
+	}
 
-		// Update the the last used garbage collect cycle before moving over to a new heap
-		ActiveGpuHeaps[ActiveGpuHeapIndex].LastUsedGarbageCollectCycle = GarbageCollectCycle;
+	// Update the the last used garbage collect cycle before moving over to a new heap
+	ActiveGpuHeaps[ActiveGpuHeapIndex].LastUsedGarbageCollectCycle = GarbageCollectCycle;
 
-		// Move the current heap as used previous heaps so it can be garbage collected
-		// (Contains valid data upto this point in time)
-		State.UsedHeaps.Emplace(ActiveGpuHeaps[ActiveGpuHeapIndex].GpuHeap);
+	// Queue the heap for recycle when the GPU is done using it
+	UE::D3D12BindlessDescriptors::DeferredFreeHeap(GetParentDevice(), ActiveGpuHeaps[ActiveGpuHeapIndex].GpuHeap);
 
-		int32 NumGpuHeaps = ActiveGpuHeaps.Num();
+	int32 NumGpuHeaps = ActiveGpuHeaps.Num();
 
-		// Copy over dirty handles to all other heaps so they are updated when reused as well
-		for (int32 GpuHeapIndex = 0; GpuHeapIndex < NumGpuHeaps; ++GpuHeapIndex)
+	// Copy over dirty handles to all other heaps so they are updated when reused as well
+	for (int32 GpuHeapIndex = 0; GpuHeapIndex < NumGpuHeaps; ++GpuHeapIndex)
+	{
+		if (GpuHeapIndex != ActiveGpuHeapIndex)
 		{
-			if (GpuHeapIndex != ActiveGpuHeapIndex)
-			{
-				ActiveGpuHeaps[GpuHeapIndex].UpdatedHandles.Append(ActiveGpuHeaps[ActiveGpuHeapIndex].UpdatedHandles);
-			}
+			ActiveGpuHeaps[GpuHeapIndex].UpdatedHandles.Append(ActiveGpuHeaps[ActiveGpuHeapIndex].UpdatedHandles);
 		}
+	}
 
-		// Try and reuse a pooled heap (incremented from last used to reduce the possible spike on reuse of lots of heap and dirty handle increase)
-		int32 NewActiveGpuHeapIndex = -1;
-		for (int32 NextIndex = 1; NextIndex < NumGpuHeaps; ++NextIndex)
-		{
-			int32 GpuHeapIndex = (ActiveGpuHeapIndex + NextIndex) % NumGpuHeaps;
-
-			// Not used by the GPU anymore and not the current one
-			if (GpuHeapIndex != ActiveGpuHeapIndex && !ActiveGpuHeaps[GpuHeapIndex].bInUse)
-			{
-				NewActiveGpuHeapIndex = GpuHeapIndex;
-				break;
-			}
-		}
+	// Try and reuse a pooled heap (incremented from last used to reduce the possible spike on reuse of lots of heap and dirty handle increase)
+	int32 NewActiveGpuHeapIndex = -1;
+	for (int32 NextIndex = 1; NextIndex < NumGpuHeaps; ++NextIndex)
+	{
+		int32 GpuHeapIndex = (ActiveGpuHeapIndex + NextIndex) % NumGpuHeaps;
 
-		// Found a pooled heap, then copy over the dirty descriptor handles
-		if (NewActiveGpuHeapIndex >= 0)
+		// Not used by the GPU anymore and not the current one
+		if (GpuHeapIndex != ActiveGpuHeapIndex && !ActiveGpuHeaps[GpuHeapIndex].bInUse)
 		{
-			// NOTE: copying over duplicate descriptor entries is faster then adding them to set for reduction
-			//		 CitySample there is about 2 to 4 times duplication but still faster to copy all then deduplication
+			NewActiveGpuHeapIndex = GpuHeapIndex;
+			break;
+		}
+	}
 
-			FGpuHeapData& GpuHeapData = ActiveGpuHeaps[NewActiveGpuHeapIndex];
-			INC_DWORD_STAT_BY(STAT_D3D12BindlessResourceGPUDescriptorsCopied, GpuHeapData.UpdatedHandles.Num());
+	// Found a pooled heap, then copy over the dirty descriptor handles
+	if (NewActiveGpuHeapIndex >= 0)
+	{
+		// NOTE: copying over duplicate descriptor entries is faster then adding them to set for reduction
+		//		 CitySample there is about 2 to 4 times duplication but still faster to copy all then deduplication
 
-			UE::D3D12Descriptors::CopyDescriptors(GetParentDevice(), GpuHeapData.GpuHeap, CpuHeap, GpuHeapData.UpdatedHandles);
-			GpuHeapData.UpdatedHandles.Reset();
+		FGpuHeapData& GpuHeapData = ActiveGpuHeaps[NewActiveGpuHeapIndex];
+		INC_DWORD_STAT_BY(STAT_D3D12BindlessResourceGPUDescriptorsCopied, GpuHeapData.UpdatedHandles.Num());
 
-			// Mark in use by GPU again
-			GpuHeapData.bInUse = true;
-			UpdateInUseGPUHeaps(true);
-		}
-		else
-		{
-			NewActiveGpuHeapIndex = AddActiveGPUHeap();
-		}
+		UE::D3D12Descriptors::CopyDescriptors(GetParentDevice(), GpuHeapData.GpuHeap, CpuHeap, GpuHeapData.UpdatedHandles);
+		GpuHeapData.UpdatedHandles.Reset();
 
-		// Update the active gpu index
-		ActiveGpuHeapIndex = NewActiveGpuHeapIndex;
-		INC_DWORD_STAT(STAT_D3D12BindlessResourceHeapsVersioned);
+		// Mark in use by GPU again
+		GpuHeapData.bInUse = true;
+		UpdateInUseGPUHeaps(true);
+	}
+	else
+	{
+		NewActiveGpuHeapIndex = AddActiveGPUHeap();
 	}
 
-	// Clear the state data
-	State.CurrentGpuHeap = nullptr;
-	State.bRequestNewGpuHeap = false;
+	// clear the request
+	bRequestNewActiveGpuHeap = false;
+
+	// Update the active gpu index
+	ActiveGpuHeapIndex = NewActiveGpuHeapIndex;
+	INC_DWORD_STAT(STAT_D3D12BindlessResourceHeapsVersioned);
 }
 
 #endif // D3D12RHI_CUSTOM_BINDLESS_RESOURCE_MANAGER
diff --git a/Engine/Source/Runtime/D3D12RHI/Private/D3D12BindlessDescriptors.h b/Engine/Source/Runtime/D3D12RHI/Private/D3D12BindlessDescriptors.h
--- a/Engine/Source/Runtime/D3D12RHI/Private/D3D12BindlessDescriptors.h
+++ b/Engine/Source/Runtime/D3D12RHI/Private/D3D12BindlessDescriptors.h
@@ -60,9 +60,6 @@ struct FD3D12ContextBindlessState
 	FD3D12DescriptorHeapPtr CurrentGpuHeap;
 	bool                    bRequestNewGpuHeap = false;
 
-	// All heaps used on the context. Used for lifetime management.
-	TArray<FD3D12DescriptorHeapPtr> UsedHeaps;
-
 	FD3D12ContextBindlessState() = default;
 	~FD3D12ContextBindlessState()
 	{
@@ -71,7 +68,7 @@ struct FD3D12ContextBindlessState
 
 	bool HasAnyPending() const
 	{
-		return UsedHeaps.Num() > 0 || bRequestNewGpuHeap;
+		return bRequestNewGpuHeap;
 	}
 };
 
@@ -153,6 +150,7 @@ private:
 	void AssignHeapToState(FD3D12ContextBindlessState& State);
 	void FinalizeHeapOnState(FD3D12ContextBindlessState& State);
 
+	void CheckRequestNewActiveGPUHeap();
 	int AddActiveGPUHeap();
 	void ReleaseGPUHeaps();
 	void UpdateInUseGPUHeaps(bool bInUse);
@@ -164,6 +162,8 @@ private:
 	uint64 							GarbageCollectCycle = 0;
 	uint64							LastUsedExplicitHeapCycle = 0;
 
+	bool							bRequestNewActiveGpuHeap = false;
+
 	uint32							InUseGPUHeaps = 0;
 	uint32							MaxInUseGPUHeaps = 0;
 	FMovingWindowMax<uint32, 100>	MovingWindowMaxInUseGPUHeaps;
