Make sure bindless heaps are bound before setting Root Signatures.
We also have to make sure the RS is set again after a bindless heap change.

diff --git a/Engine/Source/Runtime/D3D12RHI/Private/D3D12BindlessDescriptors.cpp b/Engine/Source/Runtime/D3D12RHI/Private/D3D12BindlessDescriptors.cpp
--- a/Engine/Source/Runtime/D3D12RHI/Private/D3D12BindlessDescriptors.cpp
+++ b/Engine/Source/Runtime/D3D12RHI/Private/D3D12BindlessDescriptors.cpp
@@ -390,8 +390,10 @@ void FD3D12BindlessResourceManager::UpdateDescriptor(FD3D12ContextArray const& C
 	}
 }
 
-void FD3D12BindlessResourceManager::FlushPendingDescriptorUpdates(FD3D12CommandContext& Context)
+bool FD3D12BindlessResourceManager::FlushPendingDescriptorUpdates(FD3D12CommandContext& Context)
 {
+	bool bSetNewHeaps = false;
+
 	FD3D12ContextBindlessState& State = Context.GetBindlessState();
 
 	// Create a new heap because there have been descriptor updates?
@@ -407,9 +409,11 @@ void FD3D12BindlessResourceManager::FlushPendingDescriptorUpdates(FD3D12CommandC
 		{
 			// Finally tell the Context that we're using this heap,
 			// this call also makes sure the heap is set on the d3d command list.
-			Context.StateCache.GetDescriptorCache()->SwitchToNewBindlessResourceHeap(State.CurrentGpuHeap);
+			bSetNewHeaps = Context.StateCache.GetDescriptorCache()->SwitchToNewBindlessResourceHeap(State.CurrentGpuHeap);
 		}
 	}
+
+	return bSetNewHeaps;
 }
 
 void FD3D12BindlessResourceManager::OpenCommandList(FD3D12CommandContext& Context)
@@ -785,12 +789,16 @@ void FD3D12BindlessDescriptorManager::CloseCommandList(FD3D12CommandContext& Con
 	}
 }
 
-void FD3D12BindlessDescriptorManager::FlushPendingDescriptorUpdates(FD3D12CommandContext& Context)
+bool FD3D12BindlessDescriptorManager::FlushPendingDescriptorUpdates(FD3D12CommandContext& Context)
 {
+	bool bSetNewHeaps = false;
+
 	if (ResourceManager)
 	{
-		ResourceManager->FlushPendingDescriptorUpdates(Context);
+		bSetNewHeaps = ResourceManager->FlushPendingDescriptorUpdates(Context);
 	}
+
+	return bSetNewHeaps;
 }
 
 FD3D12DescriptorHeapPair FD3D12BindlessDescriptorManager::GetExplicitHeapsForContext(FD3D12CommandContext& Context, ERHIBindlessConfiguration InConfiguration)
diff --git a/Engine/Source/Runtime/D3D12RHI/Private/D3D12BindlessDescriptors.h b/Engine/Source/Runtime/D3D12RHI/Private/D3D12BindlessDescriptors.h
--- a/Engine/Source/Runtime/D3D12RHI/Private/D3D12BindlessDescriptors.h
+++ b/Engine/Source/Runtime/D3D12RHI/Private/D3D12BindlessDescriptors.h
@@ -134,7 +134,7 @@ public:
 	void InitializeDescriptor(FRHIDescriptorHandle DstHandle, FD3D12View* View);
 	void UpdateDescriptor(FD3D12ContextArray const& Contexts, FRHIDescriptorHandle DstHandle, FD3D12View* View);
 
-	void FlushPendingDescriptorUpdates(FD3D12CommandContext& Context);
+	bool FlushPendingDescriptorUpdates(FD3D12CommandContext& Context);
 
 	void OpenCommandList(FD3D12CommandContext& Context);
 	void CloseCommandList(FD3D12CommandContext& Context);
@@ -228,7 +228,7 @@ public:
 	void OpenCommandList(FD3D12CommandContext& Context);
 	void CloseCommandList(FD3D12CommandContext& Context);
 
-	void FlushPendingDescriptorUpdates(FD3D12CommandContext& Context);
+	bool FlushPendingDescriptorUpdates(FD3D12CommandContext& Context);
 	void SetHeapsForRayTracing(FD3D12CommandContext& Context);
 
 	FD3D12DescriptorHeapPair GetExplicitHeapsForContext(FD3D12CommandContext& Context, ERHIBindlessConfiguration InConfiguration);
diff --git a/Engine/Source/Runtime/D3D12RHI/Private/D3D12CommandContext.h b/Engine/Source/Runtime/D3D12RHI/Private/D3D12CommandContext.h
--- a/Engine/Source/Runtime/D3D12RHI/Private/D3D12CommandContext.h
+++ b/Engine/Source/Runtime/D3D12RHI/Private/D3D12CommandContext.h
@@ -796,7 +796,7 @@ public:
 		return StaticUniformBuffers;
 	}
 
-	void FlushPendingDescriptorUpdates();
+	bool FlushPendingDescriptorUpdates();
 
 	void SetExplicitDescriptorCache(FD3D12ExplicitDescriptorCache& ExplicitDescriptorCache);
 	void UnsetExplicitDescriptorCache();
diff --git a/Engine/Source/Runtime/D3D12RHI/Private/D3D12Commands.cpp b/Engine/Source/Runtime/D3D12RHI/Private/D3D12Commands.cpp
--- a/Engine/Source/Runtime/D3D12RHI/Private/D3D12Commands.cpp
+++ b/Engine/Source/Runtime/D3D12RHI/Private/D3D12Commands.cpp
@@ -72,11 +72,15 @@ static void BindUniformBuffer(FD3D12CommandContext& Context, FRHIShader* Shader,
 	Context.DirtyUniformBuffers[ShaderFrequency] |= (1 << BufferIndex);
 }
 
-void FD3D12CommandContext::FlushPendingDescriptorUpdates()
+bool FD3D12CommandContext::FlushPendingDescriptorUpdates()
 {
+	bool bNewBindlessHeapsSet = false;
+
 #if PLATFORM_SUPPORTS_BINDLESS_RENDERING
-	GetParentDevice()->GetBindlessDescriptorManager().FlushPendingDescriptorUpdates(*this);
+	bNewBindlessHeapsSet = GetParentDevice()->GetBindlessDescriptorManager().FlushPendingDescriptorUpdates(*this);
 #endif
+
+	return bNewBindlessHeapsSet;
 }
 
 void FD3D12CommandContext::SetExplicitDescriptorCache(FD3D12ExplicitDescriptorCache& ExplicitDescriptorCache)
@@ -104,12 +108,12 @@ void FD3D12CommandContext::SetupDispatch(uint32 ThreadGroupCountX, uint32 Thread
 		GetParentDevice()->RegisterGPUDispatch(FIntVector(ThreadGroupCountX, ThreadGroupCountY, ThreadGroupCountZ));
 	}
 
+	const bool bNewBindlessHeapsSet = FlushPendingDescriptorUpdates();
+
 	CommitComputeShaderConstants();
 	CommitComputeResourceTables();
 
-	StateCache.ApplyState(GetPipeline(), ED3D12PipelineType::Compute);
-
-	FlushPendingDescriptorUpdates();
+	StateCache.ApplyState(GetPipeline(), ED3D12PipelineType::Compute, bNewBindlessHeapsSet);
 }
 
 FD3D12ResourceLocation& FD3D12CommandContext::SetupIndirectArgument(FRHIBuffer* ArgumentBufferRHI, D3D12_RESOURCE_STATES ExtraStates)
@@ -1568,6 +1572,8 @@ void FD3D12CommandContext::SetupDraw(FRHIBuffer* IndexBufferRHI, uint32 NumPrimi
 	}
 #endif
 
+	const bool bNewBindlessHeapsSet = FlushPendingDescriptorUpdates();
+
 	CommitGraphicsResourceTables();
 	CommitNonComputeShaderConstants();
 
@@ -1581,9 +1587,7 @@ void FD3D12CommandContext::SetupDraw(FRHIBuffer* IndexBufferRHI, uint32 NumPrimi
 		StateCache.SetIndexBuffer(IndexBuffer->ResourceLocation, Format, 0);
 	}
 
-	StateCache.ApplyState(GetPipeline(), ED3D12PipelineType::Graphics);
-
-	FlushPendingDescriptorUpdates();
+	StateCache.ApplyState(GetPipeline(), ED3D12PipelineType::Graphics, bNewBindlessHeapsSet);
 }
 
 void FD3D12CommandContext::SetupDispatchDraw(uint32 ThreadGroupCountX, uint32 ThreadGroupCountY, uint32 ThreadGroupCountZ)
@@ -1595,12 +1599,12 @@ void FD3D12CommandContext::SetupDispatchDraw(uint32 ThreadGroupCountX, uint32 Th
 	}
 #endif
 
+	const bool bNewBindlessHeapsSet = FlushPendingDescriptorUpdates();
+
 	CommitGraphicsResourceTables();
 	CommitNonComputeShaderConstants();
 
-	StateCache.ApplyState(GetPipeline(), ED3D12PipelineType::Graphics);
-
-	FlushPendingDescriptorUpdates();
+	StateCache.ApplyState(GetPipeline(), ED3D12PipelineType::Graphics, bNewBindlessHeapsSet);
 }
 
 void FD3D12CommandContext::RHIDrawPrimitive(uint32 BaseVertexIndex, uint32 NumPrimitives, uint32 NumInstances)
diff --git a/Engine/Source/Runtime/D3D12RHI/Private/D3D12DescriptorCache.cpp b/Engine/Source/Runtime/D3D12RHI/Private/D3D12DescriptorCache.cpp
--- a/Engine/Source/Runtime/D3D12RHI/Private/D3D12DescriptorCache.cpp
+++ b/Engine/Source/Runtime/D3D12RHI/Private/D3D12DescriptorCache.cpp
@@ -755,8 +755,10 @@ void FD3D12DescriptorCache::SwitchToGlobalSamplerHeap()
 }
 
 #if PLATFORM_SUPPORTS_BINDLESS_RENDERING
-void FD3D12DescriptorCache::SwitchToNewBindlessResourceHeap(FD3D12DescriptorHeap* InHeap)
+bool FD3D12DescriptorCache::SwitchToNewBindlessResourceHeap(FD3D12DescriptorHeap* InHeap)
 {
+	bool bSetNewHeaps = false;
+
 	if (ensure(IsUsingBindlessResources()))
 	{
 		BindlessResourcesHeap = InHeap;
@@ -767,9 +769,14 @@ void FD3D12DescriptorCache::SwitchToNewBindlessResourceHeap(FD3D12DescriptorHeap
 			check(BindlessSamplersHeap != nullptr);
 		}
 
+		// Switch to the new heaps
+		bSetNewHeaps = SetDescriptorHeaps();
+
 		// If we didn't change heaps, then the caller sent us the wrong heap.
-		ensure(SetDescriptorHeaps());
+		ensure(bSetNewHeaps);
 	}
+
+	return bSetNewHeaps;
 }
 #endif
 
diff --git a/Engine/Source/Runtime/D3D12RHI/Private/D3D12DescriptorCache.h b/Engine/Source/Runtime/D3D12RHI/Private/D3D12DescriptorCache.h
--- a/Engine/Source/Runtime/D3D12RHI/Private/D3D12DescriptorCache.h
+++ b/Engine/Source/Runtime/D3D12RHI/Private/D3D12DescriptorCache.h
@@ -357,7 +357,7 @@ public:
 	bool SwitchToContextLocalSamplerHeap();
 	void SwitchToGlobalSamplerHeap();
 #if PLATFORM_SUPPORTS_BINDLESS_RENDERING
-	void SwitchToNewBindlessResourceHeap(FD3D12DescriptorHeap* InHeap);
+	bool SwitchToNewBindlessResourceHeap(FD3D12DescriptorHeap* InHeap);
 #endif
 
 	void SetExplicitDescriptorCache(FD3D12ExplicitDescriptorCache& ExplicitDescriptorCache);
diff --git a/Engine/Source/Runtime/D3D12RHI/Private/D3D12RayTracing.cpp b/Engine/Source/Runtime/D3D12RHI/Private/D3D12RayTracing.cpp
--- a/Engine/Source/Runtime/D3D12RHI/Private/D3D12RayTracing.cpp
+++ b/Engine/Source/Runtime/D3D12RHI/Private/D3D12RayTracing.cpp
@@ -4723,16 +4723,14 @@ static void DispatchRays(FD3D12CommandContext& CommandContext,
 	// Invalidate state cache to ensure all root parameters for regular shaders are reset when non-RT work is dispatched later.
 	CommandContext.StateCache.TransitionComputeState(ED3D12PipelineType::RayTracing);
 
-	CommandContext.GraphicsCommandList()->SetComputeRootSignature(Pipeline->GlobalRootSignature);
+	CommandContext.GraphicsCommandList();
 
 	FD3D12RayTracingShader* RayGenShader = Pipeline->RayGenShaders.Shaders[RayGenShaderIndex];
 
 	FRHIShaderBindingLayout ShaderBindingLayout = GlobalBindings.StaticUniformBuffers.GetShaderBindingLayout() ? *GlobalBindings.StaticUniformBuffers.GetShaderBindingLayout() : FRHIShaderBindingLayout();
 	check(RayGenShader->ShaderBindingLayoutHash == ShaderBindingLayout.GetHash());
 
-	// Bind diagnostic buffer to allow asserts in ray generation shaders
 	const FD3D12RootSignature* GlobalRTRootSignature = Adapter->GetGlobalRayTracingRootSignature(ShaderBindingLayout);
-	CommandContext.BindDiagnosticBuffer(GlobalRTRootSignature, ED3D12PipelineType::Compute);
 
 	bool bResourcesBound = false;
 	if (OptShaderTable && OptShaderTable->DescriptorCache)
@@ -4744,6 +4742,7 @@ static void DispatchRays(FD3D12CommandContext& CommandContext,
 		TRACE_CPUPROFILER_EVENT_SCOPE(SetRayTracingShaderResources);
 
 		CommandContext.SetExplicitDescriptorCache(*DescriptorCache);
+		CommandContext.GraphicsCommandList()->SetComputeRootSignature(Pipeline->GlobalRootSignature);
 
 		FD3D12RayTracingGlobalResourceBinder ResourceBinder(CommandContext, *DescriptorCache);
 		bResourcesBound = SetRayTracingShaderResources(RayGenShader, GlobalRTRootSignature, GlobalBindings, ResourceBinder);
@@ -4754,11 +4753,17 @@ static void DispatchRays(FD3D12CommandContext& CommandContext,
 	{
 		FD3D12ExplicitDescriptorCache TransientDescriptorCache(CommandContext.GetParentDevice(), FD3D12RayTracingShaderBindingTableInternal::MaxBindingWorkers);
 		TransientDescriptorCache.Init(MAX_SRVS + MAX_UAVS, MAX_SAMPLERS, ERHIBindlessConfiguration::RayTracingShaders);
+
 		CommandContext.SetExplicitDescriptorCache(TransientDescriptorCache);
+		CommandContext.GraphicsCommandList()->SetComputeRootSignature(Pipeline->GlobalRootSignature);
+
 		FD3D12RayTracingGlobalResourceBinder ResourceBinder(CommandContext, TransientDescriptorCache);
 		bResourcesBound = SetRayTracingShaderResources(RayGenShader, GlobalRTRootSignature, GlobalBindings, ResourceBinder);
 	}
 
+	// Bind diagnostic buffer to allow asserts in ray generation shaders
+	CommandContext.BindDiagnosticBuffer(GlobalRTRootSignature, ED3D12PipelineType::Compute);
+
 	int8 StaticShaderBindingSlot = GlobalRTRootSignature->GetStaticShaderBindingSlot();
 	if (StaticShaderBindingSlot >= 0)
 	{
diff --git a/Engine/Source/Runtime/D3D12RHI/Private/D3D12StateCache.cpp b/Engine/Source/Runtime/D3D12RHI/Private/D3D12StateCache.cpp
--- a/Engine/Source/Runtime/D3D12RHI/Private/D3D12StateCache.cpp
+++ b/Engine/Source/Runtime/D3D12RHI/Private/D3D12StateCache.cpp
@@ -417,7 +417,7 @@ void FD3D12StateCache::InternalSetPipelineState(FD3D12PipelineState* InPipelineS
 	}
 }
 
-void FD3D12StateCache::ApplyState(ERHIPipeline HardwarePipe, ED3D12PipelineType PipelineType)
+void FD3D12StateCache::ApplyState(ERHIPipeline HardwarePipe, ED3D12PipelineType PipelineType, bool bBindlessHeapsWereJustSet)
 {
 	//SCOPE_CYCLE_COUNTER(STAT_D3D12ApplyStateTime);
 	const bool bForceState = false;
@@ -426,6 +426,14 @@ void FD3D12StateCache::ApplyState(ERHIPipeline HardwarePipe, ED3D12PipelineType
 		// Mark all state as dirty.
 		DirtyState();
 	}
+#if PLATFORM_SUPPORTS_BINDLESS_RENDERING
+	// If we just switched to a new bindless heap, we have to make sure to set the RootSignatures again.
+	else if (bBindlessHeapsWereJustSet)
+	{
+		PipelineState.Compute.bNeedSetRootSignature = true;
+		PipelineState.Graphics.bNeedSetRootSignature = true;
+	}
+#endif
 
 #if PLATFORM_SUPPORTS_VIRTUAL_TEXTURES
 	CmdContext.FlushTextureCacheIfNeeded();
diff --git a/Engine/Source/Runtime/D3D12RHI/Private/D3D12StateCachePrivate.h b/Engine/Source/Runtime/D3D12RHI/Private/D3D12StateCachePrivate.h
--- a/Engine/Source/Runtime/D3D12RHI/Private/D3D12StateCachePrivate.h
+++ b/Engine/Source/Runtime/D3D12RHI/Private/D3D12StateCachePrivate.h
@@ -675,7 +675,7 @@ public:
 	ED3D12PipelineType LastComputePipelineType = ED3D12PipelineType::Compute;
 #endif // D3D12_RHI_RAYTRACING
 
-	void ApplyState(ERHIPipeline HardwarePipe, ED3D12PipelineType PipelineType);
+	void ApplyState(ERHIPipeline HardwarePipe, ED3D12PipelineType PipelineType, bool bBindlessHeapsWereJustSet);
 	void ApplySamplers(const FD3D12RootSignature* const pRootSignature, uint32 StartStage, uint32 EndStage);
 	void ApplyResources(const FD3D12RootSignature* const pRootSignature, uint32 StartStage, uint32 EndStage);
 	void ApplyBindlessResources(const FD3D12RootSignature* const pRootSignature, uint32 StartStage, uint32 EndStage);
