Remove FUniformBufferStaticBindings from FRayTracingShaderBindings and cache ShaderBindingLayout on the RHI context and is set via RHISetStaticUniformBuffers also for RayTracing

diff --git a/Engine/Source/Runtime/D3D12RHI/Private/D3D12CommandContext.h b/Engine/Source/Runtime/D3D12RHI/Private/D3D12CommandContext.h
--- a/Engine/Source/Runtime/D3D12RHI/Private/D3D12CommandContext.h
+++ b/Engine/Source/Runtime/D3D12RHI/Private/D3D12CommandContext.h
@@ -19,6 +19,7 @@ D3D12CommandContext.h: D3D12 Command Context Interfaces
 
 #include "RHICoreShader.h"
 #include "RHICore.h"
+#include "RHIShaderBindingLayout.h"
 
 #include "GPUProfiler.h"
 
@@ -808,6 +809,12 @@ public:
 	FD3D12ContextBindlessState& GetBindlessState() { return BindlessState; }
 #endif
 
+	const FRHIShaderBindingLayout& GetShaderBindingLayout() const
+	{
+		static const FRHIShaderBindingLayout Default;
+		return ShaderBindinglayout ? *ShaderBindinglayout : Default;
+	}
+
 protected:
 
 	FD3D12CommandContext* GetSingleDeviceContext(uint32 InGPUIndex) final override
@@ -833,6 +840,7 @@ private:
 	void HandleReservedResourceCommits   (const struct FD3D12TransitionData* TransitionData);
 
 	TArray<FRHIUniformBuffer*> StaticUniformBuffers;
+	const FRHIShaderBindingLayout* ShaderBindinglayout = nullptr;
 
 #if PLATFORM_SUPPORTS_BINDLESS_RENDERING
 	FD3D12ContextBindlessState BindlessState;
diff --git a/Engine/Source/Runtime/D3D12RHI/Private/D3D12Commands.cpp b/Engine/Source/Runtime/D3D12RHI/Private/D3D12Commands.cpp
--- a/Engine/Source/Runtime/D3D12RHI/Private/D3D12Commands.cpp
+++ b/Engine/Source/Runtime/D3D12RHI/Private/D3D12Commands.cpp
@@ -637,10 +637,29 @@ void FD3D12CommandContext::RHISetStaticUniformBuffers(const FUniformBufferStatic
 {
 	FMemory::Memzero(StaticUniformBuffers.GetData(), StaticUniformBuffers.Num() * sizeof(FRHIUniformBuffer*));
 
-	for (int32 Index = 0; Index < InUniformBuffers.GetUniformBufferCount(); ++Index)
+	if (const FRHIShaderBindingLayout* Layout = InUniformBuffers.GetShaderBindingLayout())
 	{
-		StaticUniformBuffers[InUniformBuffers.GetSlot(Index)] = InUniformBuffers.GetUniformBuffer(Index);
+		check(InUniformBuffers.GetUniformBufferCount() == Layout->GetNumUniformBufferEntries());
+
+		for (int32 Index = 0; Index < InUniformBuffers.GetUniformBufferCount(); ++Index)
+		{
+			StaticUniformBuffers[Index] = InUniformBuffers.GetUniformBuffer(Index);
+			checkf(StaticUniformBuffers[Index], TEXT("Static uniform buffer at index %d is referenced in the shader binding layout but is not provided"), Index);
+		}
+
+		ShaderBindinglayout = Layout;
 	}
+	else
+	{
+		for (int32 Index = 0; Index < InUniformBuffers.GetUniformBufferCount(); ++Index)
+		{
+			FUniformBufferStaticSlot Slot = InUniformBuffers.GetSlot(Index);
+			StaticUniformBuffers[Slot] = InUniformBuffers.GetUniformBuffer(Index);
+		}
+
+		ShaderBindinglayout = nullptr;
+	}
+
 }
 
 void FD3D12CommandContext::RHICopyToStagingBuffer(FRHIBuffer* SourceBufferRHI, FRHIStagingBuffer* StagingBufferRHI, uint32 Offset, uint32 NumBytes)
diff --git a/Engine/Source/Runtime/D3D12RHI/Private/D3D12RayTracing.cpp b/Engine/Source/Runtime/D3D12RHI/Private/D3D12RayTracing.cpp
--- a/Engine/Source/Runtime/D3D12RHI/Private/D3D12RayTracing.cpp
+++ b/Engine/Source/Runtime/D3D12RHI/Private/D3D12RayTracing.cpp
@@ -4727,9 +4727,11 @@ static void DispatchRays(FD3D12CommandContext& CommandContext,
 
 	FD3D12RayTracingShader* RayGenShader = Pipeline->RayGenShaders.Shaders[RayGenShaderIndex];
 
-	FRHIShaderBindingLayout ShaderBindingLayout = GlobalBindings.StaticUniformBuffers.GetShaderBindingLayout() ? *GlobalBindings.StaticUniformBuffers.GetShaderBindingLayout() : FRHIShaderBindingLayout();
+	const FRHIShaderBindingLayout& ShaderBindingLayout = CommandContext.GetShaderBindingLayout();
 	check(RayGenShader->ShaderBindingLayoutHash == ShaderBindingLayout.GetHash());
 
+	const TArray<FRHIUniformBuffer*>& StaticUniformBuffers = CommandContext.GetStaticUniformBuffers();
+
 	const FD3D12RootSignature* GlobalRTRootSignature = Adapter->GetGlobalRayTracingRootSignature(ShaderBindingLayout);
 
 	bool bResourcesBound = false;
@@ -4767,20 +4769,19 @@ static void DispatchRays(FD3D12CommandContext& CommandContext,
 	int8 StaticShaderBindingSlot = GlobalRTRootSignature->GetStaticShaderBindingSlot();
 	if (StaticShaderBindingSlot >= 0)
 	{
-		check(GlobalBindings.StaticUniformBuffers.GetShaderBindingLayout());
-		check(GlobalBindings.StaticUniformBuffers.GetShaderBindingLayout()->GetNumUniformBufferEntries() == GlobalBindings.StaticUniformBuffers.GetUniformBufferCount());
-		for (int32 Index = 0; Index < GlobalBindings.StaticUniformBuffers.GetUniformBufferCount(); ++Index)
-		{
-			int32 SlotIndex = Index;
-			FRHIUniformBuffer* UniformBuffer = GlobalBindings.StaticUniformBuffers.GetUniformBuffer(Index);
-			if (UniformBuffer)
+		for (uint32 Index = 0; Index < ShaderBindingLayout.GetNumUniformBufferEntries(); ++Index)
 			{
+			const FRHIUniformBufferShaderBindingLayout& LayoutEntry = ShaderBindingLayout.GetUniformBufferEntry(Index);
+			const uint32 RootParameterSlotIndex = uint32(StaticShaderBindingSlot) + LayoutEntry.CBVResourceIndex;
+
+			FRHIUniformBuffer* UniformBuffer = StaticUniformBuffers[Index];
+			checkf(UniformBuffer, TEXT("Static uniform buffer at index %d is referenced in the shader binding layout but not provided in the last RHISetStaticUniformBuffers() command"), Index);
+
 				FD3D12UniformBuffer* D3D12UniformBuffer = FD3D12CommandContext::RetrieveObject<FD3D12UniformBuffer>(UniformBuffer, 0);//GpuIndex);
 				if (D3D12UniformBuffer->ResourceLocation.GetGPUVirtualAddress())
 				{
 					const FD3D12ResourceLocation& ResourceLocation = D3D12UniformBuffer->ResourceLocation;
-					CommandContext.GraphicsCommandList()->SetComputeRootConstantBufferView(StaticShaderBindingSlot + SlotIndex, ResourceLocation.GetGPUVirtualAddress());
-				}
+				CommandContext.GraphicsCommandList()->SetComputeRootConstantBufferView(RootParameterSlotIndex, ResourceLocation.GetGPUVirtualAddress());
 			}
 		}
 	}
diff --git a/Engine/Source/Runtime/RHI/Private/RHICommandList.cpp b/Engine/Source/Runtime/RHI/Private/RHICommandList.cpp
--- a/Engine/Source/Runtime/RHI/Private/RHICommandList.cpp
+++ b/Engine/Source/Runtime/RHI/Private/RHICommandList.cpp
@@ -29,7 +29,16 @@ DECLARE_DWORD_COUNTER_STAT(TEXT("Immed. Command count"), STAT_ImmedCmdListCount,
 UE_TRACE_CHANNEL_DEFINE(RHICommandsChannel);
 
 #if VALIDATE_UNIFORM_BUFFER_STATIC_BINDINGS
-bool FScopedUniformBufferStaticBindings::bRecursionGuard = false;
+static thread_local bool bScopedUniformBufferStaticBindingsRecursionGuard = false;
+void FScopedUniformBufferStaticBindings::OnScopeEnter()
+{
+	checkf(!bScopedUniformBufferStaticBindingsRecursionGuard, TEXT("Uniform buffer global binding scope has been called recursively!"));
+	bScopedUniformBufferStaticBindingsRecursionGuard = true;
+}
+void FScopedUniformBufferStaticBindings::OnScopeExit()
+{
+	bScopedUniformBufferStaticBindingsRecursionGuard = false;
+}
 #endif
 
 #if !PLATFORM_USES_FIXED_RHI_CLASS
diff --git a/Engine/Source/Runtime/RHI/Public/RHICommandList.h b/Engine/Source/Runtime/RHI/Public/RHICommandList.h
--- a/Engine/Source/Runtime/RHI/Public/RHICommandList.h
+++ b/Engine/Source/Runtime/RHI/Public/RHICommandList.h
@@ -246,7 +246,6 @@ struct FRayTracingShaderBindings
 	FRHIUnorderedAccessView* UAVs[16] = {};
 
 	TArray<FRHIShaderParameterResource> BindlessParameters;
-	FUniformBufferStaticBindings StaticUniformBuffers;
 };
 
 struct FRayTracingLocalShaderBindings
@@ -4961,8 +4960,7 @@ struct FScopedUniformBufferStaticBindings
 		: RHICmdList(InRHICmdList)
 	{
 #if VALIDATE_UNIFORM_BUFFER_STATIC_BINDINGS
-		checkf(!bRecursionGuard, TEXT("Uniform buffer global binding scope has been called recursively!"));
-		bRecursionGuard = true;
+		OnScopeEnter();
 #endif
 
 		RHICmdList.SetStaticUniformBuffers(UniformBuffers);
@@ -4978,14 +4976,16 @@ struct FScopedUniformBufferStaticBindings
 		RHICmdList.SetStaticUniformBuffers(FUniformBufferStaticBindings());
 
 #if VALIDATE_UNIFORM_BUFFER_STATIC_BINDINGS
-		bRecursionGuard = false;
+		OnScopeExit();
 #endif
 	}
 
 	FRHIComputeCommandList& RHICmdList;
 
+private:
 #if VALIDATE_UNIFORM_BUFFER_STATIC_BINDINGS
-	RHI_API static bool bRecursionGuard;
+	RHI_API static void OnScopeEnter();
+	RHI_API static void OnScopeExit();
 #endif
 };
 
