Fix editor crash when running with "-GPUCrashDebugging"
 - Do not track residency on resources that are CPU accessible.

diff --git a/Engine/Source/Runtime/D3D12RHI/Private/D3D12PoolAllocator.cpp b/Engine/Source/Runtime/D3D12RHI/Private/D3D12PoolAllocator.cpp
--- a/Engine/Source/Runtime/D3D12RHI/Private/D3D12PoolAllocator.cpp
+++ b/Engine/Source/Runtime/D3D12RHI/Private/D3D12PoolAllocator.cpp
@@ -438,7 +438,12 @@ void FD3D12PoolAllocator::AllocateResource(uint32 GPUIndex, D3D12_HEAP_TYPE InHe
 			TRefCountPtr<FD3D12Heap> BackingHeap = new FD3D12Heap(GetParentDevice(), GetVisibilityMask(), TraceHeapId);
 			bool bTrack = false;
 			BackingHeap->SetHeap(Heap, InName, bTrack);
-			BackingHeap->BeginTrackingResidency(HeapDesc.SizeInBytes);
+
+			// Only track resources that cannot be accessed on the CPU.
+			if (IsGPUOnly(InHeapType, &HeapProps))
+			{
+				BackingHeap->BeginTrackingResidency(HeapDesc.SizeInBytes);
+			}
 
 			VERIFYD3D12RESULT(Adapter->CreatePlacedResource(Desc, BackingHeap, 0, InCreateState, InResourceStateMode, InCreateState, InClearValue, &NewResource, InName));
 		}
