Remove unrequired check because it can happen that the heap has already been freed when bindless heaps are not used for certain amount of frames and then all heaps are freed already

diff --git a/Engine/Source/Runtime/D3D12RHI/Private/D3D12BindlessDescriptors.cpp b/Engine/Source/Runtime/D3D12RHI/Private/D3D12BindlessDescriptors.cpp
--- a/Engine/Source/Runtime/D3D12RHI/Private/D3D12BindlessDescriptors.cpp
+++ b/Engine/Source/Runtime/D3D12RHI/Private/D3D12BindlessDescriptors.cpp
@@ -327,7 +327,6 @@ void FD3D12BindlessResourceManager::Recycle(FD3D12DescriptorHeap* DescriptorHeap
 			break;
 		}
 	}
-	check(bFound);
 }
 
 void FD3D12BindlessResourceManager::InitializeDescriptor(FRHIDescriptorHandle DstHandle, FD3D12View* View)
