Allow missing/different shader binding layout on shader if it's not sampling any resources (eg BuildInRayTracingShaders)

diff --git a/Engine/Source/Runtime/D3D12RHI/Private/D3D12RayTracing.cpp b/Engine/Source/Runtime/D3D12RHI/Private/D3D12RayTracing.cpp
--- a/Engine/Source/Runtime/D3D12RHI/Private/D3D12RayTracing.cpp
+++ b/Engine/Source/Runtime/D3D12RHI/Private/D3D12RayTracing.cpp
@@ -1955,8 +1955,9 @@ public:
 										&CompileCompletionList]
 			(FD3D12RayTracingShader* Shader, FD3D12RayTracingPipelineCache::ECollectionType CollectionType)
 		{
-			// verify that that the same shader binding layout is used for all shaders in the RTPSO
-			checkf(Shader->ShaderBindingLayoutHash == ShaderBindingLayoutHash, TEXT("Raytracing shader with with entry point %s doesn't match the RTPSO ShaderBindingLayout"), *Shader->EntryPoint);
+			// verify that that the same shader binding layout is used for all shaders in the RTPSO or not sampling any resources
+			uint32 TotalResourceCount = Shader->ResourceCounts.NumCBs + Shader->ResourceCounts.NumSRVs + Shader->ResourceCounts.NumUAVs + Shader->ResourceCounts.NumSamplers;
+			checkf(TotalResourceCount == 0 || Shader->ShaderBindingLayoutHash == ShaderBindingLayoutHash, TEXT("Raytracing shader with with entry point %s doesn't match the RTPSO ShaderBindingLayout"), *Shader->EntryPoint);
 
 			bool bIsAlreadyInSet = false;
 			const uint64 ShaderHash = GetShaderHash64(Shader);
