Move RHIThreadFence inside EnqueueLambdaMultiPipe instead of relying on caller to call it explicity (enabled/disabled via bool parameter)

diff --git a/Engine/Source/Runtime/D3D12RHI/Private/D3D12Buffer.cpp b/Engine/Source/Runtime/D3D12RHI/Private/D3D12Buffer.cpp
--- a/Engine/Source/Runtime/D3D12RHI/Private/D3D12Buffer.cpp
+++ b/Engine/Source/Runtime/D3D12RHI/Private/D3D12Buffer.cpp
@@ -504,7 +504,7 @@ void* FD3D12DynamicRHI::LockBuffer(FRHICommandListBase& RHICmdList, FD3D12Buffer
 			FD3D12ResourceLocation NewLocation(Device);
 			Data = Adapter.GetUploadHeapAllocator(Device->GetGPUIndex()).AllocUploadResource(BufferSize, Buffer->BufferAlignment, NewLocation);
 
-			RHICmdList.EnqueueLambdaMultiPipe(GetEnabledRHIPipelines(), TEXT("FD3D12DynamicRHI::LockBuffer"),
+			RHICmdList.EnqueueLambdaMultiPipe(GetEnabledRHIPipelines(), FRHICommandListBase::EThreadFence::Enabled, TEXT("FD3D12DynamicRHI::LockBuffer"),
 			[
 				Resource = Buffer,
 				NewLocation = MoveTemp(NewLocation)
@@ -534,8 +534,6 @@ void* FD3D12DynamicRHI::LockBuffer(FRHICommandListBase& RHICmdList, FD3D12Buffer
 #endif
 				Resource->RenameLDAChain(Contexts, NewLocation);
 			});
-
-			RHICmdList.RHIThreadFence(true);
 		}
 	}
 	else
diff --git a/Engine/Source/Runtime/D3D12RHI/Private/D3D12CommandContext.cpp b/Engine/Source/Runtime/D3D12RHI/Private/D3D12CommandContext.cpp
--- a/Engine/Source/Runtime/D3D12RHI/Private/D3D12CommandContext.cpp
+++ b/Engine/Source/Runtime/D3D12RHI/Private/D3D12CommandContext.cpp
@@ -653,7 +653,7 @@ void FD3D12ContextCommon::ConditionalSplitCommandList()
 
 void FD3D12DynamicRHI::RHIBeginFrame(FRHICommandListImmediate& RHICmdList)
 {
-	RHICmdList.EnqueueLambdaMultiPipe(GetEnabledRHIPipelines(), TEXT("FD3D12DynamicRHI::RHIBeginFrame"),
+	RHICmdList.EnqueueLambdaMultiPipe(GetEnabledRHIPipelines(), FRHICommandListBase::EThreadFence::Enabled, TEXT("FD3D12DynamicRHI::RHIBeginFrame"),
 		[this](FD3D12ContextArray const& Contexts)
 		{
 			for (auto& Adapter : ChosenAdapters)
diff --git a/Engine/Source/Runtime/D3D12RHI/Private/D3D12Resources.cpp b/Engine/Source/Runtime/D3D12RHI/Private/D3D12Resources.cpp
--- a/Engine/Source/Runtime/D3D12RHI/Private/D3D12Resources.cpp
+++ b/Engine/Source/Runtime/D3D12RHI/Private/D3D12Resources.cpp
@@ -1709,7 +1709,7 @@ uint32 FD3D12Buffer::GetParentGPUIndex() const
 
 void FD3D12DynamicRHI::RHIReplaceResources(FRHICommandListBase& RHICmdList, TArray<FRHIResourceReplaceInfo>&& ReplaceInfos)
 {
-	RHICmdList.EnqueueLambdaMultiPipe(GetEnabledRHIPipelines(), TEXT("FD3D12DynamicRHI::RHIReplaceResources"),
+	RHICmdList.EnqueueLambdaMultiPipe(GetEnabledRHIPipelines(), FRHICommandListBase::EThreadFence::Enabled, TEXT("FD3D12DynamicRHI::RHIReplaceResources"),
 		[ReplaceInfos = MoveTemp(ReplaceInfos)](FD3D12ContextArray const& Contexts)
 		{
 			for (FRHIResourceReplaceInfo const& Info : ReplaceInfos)
@@ -1761,6 +1761,4 @@ void FD3D12DynamicRHI::RHIReplaceResources(FRHICommandListBase& RHICmdList, TArr
 			}
 		}
 	);
-
-	RHICmdList.RHIThreadFence(true);
 }
diff --git a/Engine/Source/Runtime/D3D12RHI/Private/D3D12TextureReference.cpp b/Engine/Source/Runtime/D3D12RHI/Private/D3D12TextureReference.cpp
--- a/Engine/Source/Runtime/D3D12RHI/Private/D3D12TextureReference.cpp
+++ b/Engine/Source/Runtime/D3D12RHI/Private/D3D12TextureReference.cpp
@@ -89,7 +89,7 @@ void FD3D12DynamicRHI::RHIUpdateTextureReference(FRHICommandListBase& RHICmdList
 	// @todo dev-pr - This should be refactored out when we eventually remove FRHITextureReference.
 	TRefCountPtr<FRHITextureReference> Ref = TextureRef;
 
-	RHICmdList.EnqueueLambdaMultiPipe(GetEnabledRHIPipelines(), TEXT("FD3D12DynamicRHI::RHIUpdateTextureReference"),
+	RHICmdList.EnqueueLambdaMultiPipe(GetEnabledRHIPipelines(), FRHICommandListBase::EThreadFence::Enabled, TEXT("FD3D12DynamicRHI::RHIUpdateTextureReference"),
 	[
 		TextureRef = MoveTemp(Ref),
 		NewTexture = InNewTexture ? InNewTexture : FRHITextureReference::GetDefaultTexture()
@@ -100,6 +100,4 @@ void FD3D12DynamicRHI::RHIUpdateTextureReference(FRHICommandListBase& RHICmdList
 			It.GetFirst()->SwitchToNewTexture(Contexts, It.GetSecond());
 		}
 	});
-
-	RHICmdList.RHIThreadFence(true);
 }
diff --git a/Engine/Source/Runtime/RHI/Public/RHICommandList.h b/Engine/Source/Runtime/RHI/Public/RHICommandList.h
--- a/Engine/Source/Runtime/RHI/Public/RHICommandList.h
+++ b/Engine/Source/Runtime/RHI/Public/RHICommandList.h
@@ -593,9 +593,15 @@ public:
 	{
 		FRHICommandListBase::EnqueueLambda(TEXT("TRHILambdaCommand"), Forward<LAMBDA>(Lambda));
 	}
+	
+	enum class EThreadFence
+	{
+		Enabled,
+		Disabled
+	};
 
 	template <typename LAMBDA>
-	void EnqueueLambdaMultiPipe(ERHIPipeline Pipelines, const TCHAR* LambdaName, LAMBDA&& Lambda)
+	void EnqueueLambdaMultiPipe(ERHIPipeline Pipelines, EThreadFence ThreadFence, const TCHAR* LambdaName, LAMBDA&& Lambda)
 	{
 		checkf(IsTopOfPipe() || Bypass(), TEXT("Cannot enqueue a multi-pipe lambda from the bottom of pipe."));
 
@@ -620,6 +626,11 @@ public:
 		}
 
 		ActivatePipelines(OldPipeline);
+		
+		if (ThreadFence == EThreadFence::Enabled)
+		{
+			RHIThreadFence(true);
+		}
 	}
 
 	FORCEINLINE bool HasCommands() const
