Remove RHICopyBuffer implementations and deprecate CopyBuffer() on the immediate command list
 - The few platforms that implemented this had dangerous use of the the default context.
 - RHICopyBufferRegion is the correct alternative and is properly supported on all platforms.

diff --git a/Engine/Source/Runtime/Apple/MetalRHI/Private/MetalDynamicRHI.h b/Engine/Source/Runtime/Apple/MetalRHI/Private/MetalDynamicRHI.h
--- a/Engine/Source/Runtime/Apple/MetalRHI/Private/MetalDynamicRHI.h
+++ b/Engine/Source/Runtime/Apple/MetalRHI/Private/MetalDynamicRHI.h
@@ -61,7 +61,6 @@ public:
 	virtual FUniformBufferRHIRef RHICreateUniformBuffer(const void* Contents, const FRHIUniformBufferLayout* Layout, EUniformBufferUsage Usage, EUniformBufferValidation Validation) final override;
 	virtual FBufferRHIRef RHICreateBuffer(FRHICommandListBase& RHICmdList, FRHIBufferDesc const& Desc, ERHIAccess ResourceState, FRHIResourceCreateInfo& CreateInfo) final override;
 	virtual void RHIReplaceResources(FRHICommandListBase& RHICmdList, TArray<FRHIResourceReplaceInfo>&& ReplaceInfos) final override;
-	virtual void RHICopyBuffer(FRHIBuffer* SourceBuffer, FRHIBuffer* DestBuffer) final override;
 	virtual void* LockBuffer_BottomOfPipe(FRHICommandListBase& RHICmdList, FRHIBuffer* Buffer, uint32 Offset, uint32 SizeRHI, EResourceLockMode LockMode) final override;
 	virtual void UnlockBuffer_BottomOfPipe(FRHICommandListBase& RHICmdList, FRHIBuffer* Buffer) final override;
 
diff --git a/Engine/Source/Runtime/Apple/MetalRHI/Private/MetalVertexBuffer.cpp b/Engine/Source/Runtime/Apple/MetalRHI/Private/MetalVertexBuffer.cpp
--- a/Engine/Source/Runtime/Apple/MetalRHI/Private/MetalVertexBuffer.cpp
+++ b/Engine/Source/Runtime/Apple/MetalRHI/Private/MetalVertexBuffer.cpp
@@ -627,35 +627,3 @@ void FMetalDynamicRHI::UnlockBuffer_BottomOfPipe(FRHICommandListBase& RHICmdList
     FMetalRHIBuffer* Buffer = ResourceCast(BufferRHI);
 	Buffer->Unlock();
 }
-
-void FMetalDynamicRHI::RHICopyBuffer(FRHIBuffer* SourceBufferRHI, FRHIBuffer* DestBufferRHI)
-{
-    MTL_SCOPED_AUTORELEASE_POOL;
-    
-    FMetalRHIBuffer* SrcBuffer = ResourceCast(SourceBufferRHI);
-    FMetalRHIBuffer* DstBuffer = ResourceCast(DestBufferRHI);
-    
-    FMetalBufferPtr TheSrcBuffer = SrcBuffer->GetCurrentBuffer();
-    FMetalBufferPtr TheDstBuffer = DstBuffer->GetCurrentBuffer();
-
-    if (TheSrcBuffer && TheDstBuffer)
-    {
-        GetMetalDeviceContext().CopyFromBufferToBuffer(TheSrcBuffer, 0, TheDstBuffer, 0, FMath::Min(SrcBuffer->GetSize(), DstBuffer->GetSize()));
-    }
-    else if (TheDstBuffer)
-    {
-        FMetalPooledBufferArgs ArgsCPU(GetMetalDeviceContext().GetDevice(), SrcBuffer->GetSize(), BUF_Dynamic, MTL::StorageModeShared);
-        FMetalBufferPtr TempBuffer = GetMetalDeviceContext().CreatePooledBuffer(ArgsCPU);
-        FMemory::Memcpy(TempBuffer->Contents(), SrcBuffer->Data->Data, SrcBuffer->GetSize());
-        GetMetalDeviceContext().CopyFromBufferToBuffer(TempBuffer, 0, TheDstBuffer, 0, FMath::Min(SrcBuffer->GetSize(), DstBuffer->GetSize()));
-        SafeReleaseMetalBuffer(TempBuffer);
-    }
-    else
-    {
-        void const* SrcData = SrcBuffer->Lock(true, RLM_ReadOnly, 0);
-        void* DstData = DstBuffer->Lock(true, RLM_WriteOnly, 0);
-        FMemory::Memcpy(DstData, SrcData, FMath::Min(SrcBuffer->GetSize(), DstBuffer->GetSize()));
-        SrcBuffer->Unlock();
-        DstBuffer->Unlock();
-    }
-}
diff --git a/Engine/Source/Runtime/D3D12RHI/Private/D3D12Buffer.cpp b/Engine/Source/Runtime/D3D12RHI/Private/D3D12Buffer.cpp
--- a/Engine/Source/Runtime/D3D12RHI/Private/D3D12Buffer.cpp
+++ b/Engine/Source/Runtime/D3D12RHI/Private/D3D12Buffer.cpp
@@ -696,52 +696,6 @@ void FD3D12DynamicRHI::RHIUnlockBufferMGPU(FRHICommandListBase& RHICmdList, FRHI
 	UnlockBuffer(RHICmdList, Buffer, Buffer->GetUsage());
 }
 
-void FD3D12DynamicRHI::RHICopyBuffer(FRHIBuffer* SourceBufferRHI, FRHIBuffer* DestBufferRHI)
-{
-	FD3D12Buffer* SrcBuffer = FD3D12DynamicRHI::ResourceCast(SourceBufferRHI);
-	FD3D12Buffer* DstBuffer = FD3D12DynamicRHI::ResourceCast(DestBufferRHI);
-	check(SrcBuffer->GetSize() == DstBuffer->GetSize());
-
-	FD3D12Buffer* SourceBuffer = SrcBuffer;
-	FD3D12Buffer* DestBuffer = DstBuffer;
-
-	for (FD3D12Buffer::FDualLinkedObjectIterator It(SourceBuffer, DestBuffer); It; ++It)
-	{
-		SourceBuffer = It.GetFirst();
-		DestBuffer = It.GetSecond();
-
-		FD3D12Device* Device = SourceBuffer->GetParentDevice();
-		check(Device == DestBuffer->GetParentDevice());
-
-		FD3D12Resource* pSourceResource = SourceBuffer->ResourceLocation.GetResource();
-		D3D12_RESOURCE_DESC const& SourceBufferDesc = pSourceResource->GetDesc();
-
-		FD3D12Resource* pDestResource = DestBuffer->ResourceLocation.GetResource();
-		D3D12_RESOURCE_DESC const& DestBufferDesc = pDestResource->GetDesc();
-
-		check(SourceBuffer->GetSize() == DestBuffer->GetSize());
-
-		FD3D12CommandContext& Context = Device->GetDefaultCommandContext();
-
-		// The underlying D3D12 buffer can be larger than the RHI buffer due to pooling.
-		Context.GraphicsCommandList()->CopyBufferRegion(
-			pDestResource->GetResource(), 
-			DestBuffer->ResourceLocation.GetOffsetFromBaseOfResource(), 
-			pSourceResource->GetResource(), 
-			SourceBuffer->ResourceLocation.GetOffsetFromBaseOfResource(), 
-			SourceBufferRHI->GetSize());
-
-		Context.UpdateResidency(pDestResource);
-		Context.UpdateResidency(pSourceResource);
-
-		Context.ConditionalSplitCommandList();
-
-		DEBUG_EXECUTE_COMMAND_CONTEXT(Context);
-
-		Device->RegisterGPUWork(1);
-	}
-}
-
 void FD3D12DynamicRHI::RHIBindDebugLabelName(FRHICommandListBase& RHICmdList, FRHIBuffer* BufferRHI, const TCHAR* Name)
 {
 	if (BufferRHI == nullptr || !GD3D12BindResourceLabels)
diff --git a/Engine/Source/Runtime/D3D12RHI/Private/D3D12RHIPrivate.h b/Engine/Source/Runtime/D3D12RHI/Private/D3D12RHIPrivate.h
--- a/Engine/Source/Runtime/D3D12RHI/Private/D3D12RHIPrivate.h
+++ b/Engine/Source/Runtime/D3D12RHI/Private/D3D12RHIPrivate.h
@@ -316,7 +316,6 @@ public:
 	virtual void RHIUpdateUniformBuffer(FRHICommandListBase& RHICmdList, FRHIUniformBuffer* UniformBufferRHI, const void* Contents) final override;
 	virtual FBufferRHIRef RHICreateBuffer(FRHICommandListBase& RHICmdList, FRHIBufferDesc const& Desc, ERHIAccess ResourceState, FRHIResourceCreateInfo& CreateInfo) final override;
 	virtual void RHIReplaceResources(FRHICommandListBase& RHICmdList, TArray<FRHIResourceReplaceInfo>&& ReplaceInfos) final override;
-	virtual void RHICopyBuffer(FRHIBuffer* SourceBuffer, FRHIBuffer* DestBuffer) final override;
 	virtual void* RHILockBuffer(FRHICommandListBase& RHICmdList, FRHIBuffer* Buffer, uint32 Offset, uint32 Size, EResourceLockMode LockMode) final override;
 	virtual void* RHILockBufferMGPU(FRHICommandListBase& RHICmdList, FRHIBuffer* Buffer, uint32 GPUIndex, uint32 Offset, uint32 Size, EResourceLockMode LockMode) final override;
 	virtual void RHIUnlockBuffer(FRHICommandListBase& RHICmdList, FRHIBuffer* Buffer) final override;
diff --git a/Engine/Source/Runtime/NullDrv/Public/NullRHI.h b/Engine/Source/Runtime/NullDrv/Public/NullRHI.h
--- a/Engine/Source/Runtime/NullDrv/Public/NullRHI.h
+++ b/Engine/Source/Runtime/NullDrv/Public/NullRHI.h
@@ -125,11 +125,6 @@ public:
 		return new FRHIBuffer(Desc);
 	}
 
-	virtual void RHICopyBuffer(FRHIBuffer* SourceBuffer, FRHIBuffer* DestBuffer) final override
-	{
-
-	}
-
 	virtual void* LockBuffer_BottomOfPipe(FRHICommandListBase& RHICmdList, FRHIBuffer* Buffer, uint32 Offset, uint32 SizeRHI, EResourceLockMode LockMode) final override
 	{
 		return GetStaticBuffer(Buffer->GetSize());
diff --git a/Engine/Source/Runtime/OpenGLDrv/Private/OpenGLBuffer.cpp b/Engine/Source/Runtime/OpenGLDrv/Private/OpenGLBuffer.cpp
--- a/Engine/Source/Runtime/OpenGLDrv/Private/OpenGLBuffer.cpp
+++ b/Engine/Source/Runtime/OpenGLDrv/Private/OpenGLBuffer.cpp
@@ -245,12 +245,6 @@ void FOpenGLDynamicRHI::UnlockBuffer_BottomOfPipe(FRHICommandListBase& RHICmdLis
 	RHITHREAD_GLCOMMAND_EPILOGUE();
 }
 
-void FOpenGLDynamicRHI::RHICopyBuffer(FRHIBuffer* SourceBufferRHI, FRHIBuffer* DestBufferRHI)
-{
-	check(SourceBufferRHI && DestBufferRHI && SourceBufferRHI->GetSize() == DestBufferRHI->GetSize());
-	RHICopyBufferRegion(DestBufferRHI, 0, SourceBufferRHI, 0, SourceBufferRHI->GetSize());
-}
-
 void FOpenGLDynamicRHI::RHICopyBufferRegion(FRHIBuffer* DestBufferRHI, uint64 DstOffset, FRHIBuffer* SourceBufferRHI, uint64 SrcOffset, uint64 NumBytes)
 {
 	VERIFY_GL_SCOPE();
diff --git a/Engine/Source/Runtime/OpenGLDrv/Private/OpenGLDrv.h b/Engine/Source/Runtime/OpenGLDrv/Private/OpenGLDrv.h
--- a/Engine/Source/Runtime/OpenGLDrv/Private/OpenGLDrv.h
+++ b/Engine/Source/Runtime/OpenGLDrv/Private/OpenGLDrv.h
@@ -355,7 +355,6 @@ public:
 	virtual FUniformBufferRHIRef RHICreateUniformBuffer(const void* Contents, const FRHIUniformBufferLayout* Layout, EUniformBufferUsage Usage, EUniformBufferValidation Validation) final override;
 	virtual void RHIUpdateUniformBuffer(FRHICommandListBase& RHICmdList, FRHIUniformBuffer* UniformBufferRHI, const void* Contents) final override;
 	virtual FBufferRHIRef RHICreateBuffer(FRHICommandListBase& RHICmdList, FRHIBufferDesc const& Desc, ERHIAccess ResourceState, FRHIResourceCreateInfo& CreateInfo) final override;
-	virtual void RHICopyBuffer(FRHIBuffer* SourceBuffer, FRHIBuffer* DestBuffer) final override;
 	virtual void RHIReplaceResources(FRHICommandListBase& RHICmdList, TArray<FRHIResourceReplaceInfo>&& ReplaceInfos) final override;
 	virtual void* LockBuffer_BottomOfPipe(FRHICommandListBase& RHICmdList, FRHIBuffer* Buffer, uint32 Offset, uint32 Size, EResourceLockMode LockMode) final override;
 	virtual void UnlockBuffer_BottomOfPipe(FRHICommandListBase& RHICmdList, FRHIBuffer* Buffer) final override;
diff --git a/Engine/Source/Runtime/RHI/Public/DynamicRHI.h b/Engine/Source/Runtime/RHI/Public/DynamicRHI.h
--- a/Engine/Source/Runtime/RHI/Public/DynamicRHI.h
+++ b/Engine/Source/Runtime/RHI/Public/DynamicRHI.h
@@ -337,10 +337,6 @@ public:
 
 	virtual void RHIUpdateUniformBuffer(FRHICommandListBase& RHICmdList, FRHIUniformBuffer* UniformBufferRHI, const void* Contents) = 0;
 
-	/** Copies the contents of one buffer to another buffer. They must have identical sizes. */
-	// FlushType: Flush Immediate (seems dangerous)
-	virtual void RHICopyBuffer(FRHIBuffer* SourceBuffer, FRHIBuffer* DestBuffer) = 0;
-
 	/**
 	 * Transfer metadata and underlying resource from src to dest and release any resource owned by dest.
 	 * @param DestBuffer - the buffer to update
diff --git a/Engine/Source/Runtime/RHI/Public/RHICommandList.h b/Engine/Source/Runtime/RHI/Public/RHICommandList.h
--- a/Engine/Source/Runtime/RHI/Public/RHICommandList.h
+++ b/Engine/Source/Runtime/RHI/Public/RHICommandList.h
@@ -4640,11 +4640,11 @@ public:
 		GDynamicRHI->UnlockStagingBuffer_RenderThread(*this, StagingBuffer);
 	}
 
+	UE_DEPRECATED(5.5, "CopyBuffer is deprecated. Use CopyBufferRegion.")
 	FORCEINLINE void CopyBuffer(FRHIBuffer* SourceBuffer, FRHIBuffer* DestBuffer)
 	{
-		QUICK_SCOPE_CYCLE_COUNTER(STAT_RHIMETHOD_CopyBuffer_Flush);
-		ImmediateFlush(EImmediateFlushType::FlushRHIThread);  
-		GDynamicRHI->RHICopyBuffer(SourceBuffer, DestBuffer);
+		uint64 Size = FMath::Min(SourceBuffer->GetSize(), DestBuffer->GetSize());
+		CopyBufferRegion(DestBuffer, 0, SourceBuffer, 0, Size);
 	}
 
 	FORCEINLINE bool GetTextureMemoryVisualizeData(FColor* TextureData,int32 SizeX,int32 SizeY,int32 Pitch,int32 PixelSize)
diff --git a/Engine/Source/Runtime/RHI/Public/RHIValidation.h b/Engine/Source/Runtime/RHI/Public/RHIValidation.h
--- a/Engine/Source/Runtime/RHI/Public/RHIValidation.h
+++ b/Engine/Source/Runtime/RHI/Public/RHIValidation.h
@@ -429,13 +429,6 @@ public:
 		return Buffer;
 	}
 
-	/** Copies the contents of one vertex buffer to another vertex buffer.  They must have identical sizes. */
-	// FlushType: Flush Immediate (seems dangerous)
-	virtual void RHICopyBuffer(FRHIBuffer* SourceBuffer, FRHIBuffer* DestBuffer) override final
-	{
-		RHI->RHICopyBuffer(SourceBuffer, DestBuffer);
-	}
-
 	// FlushType: Flush RHI Thread
 	virtual void* RHILockBuffer(class FRHICommandListBase& RHICmdList, FRHIBuffer* Buffer, uint32 Offset, uint32 SizeRHI, EResourceLockMode LockMode) override final;
 	virtual void* RHILockBufferMGPU(class FRHICommandListBase& RHICmdList, FRHIBuffer* Buffer, uint32 GPUIndex, uint32 Offset, uint32 SizeRHI, EResourceLockMode LockMode) override final;
diff --git a/Engine/Source/Runtime/VulkanRHI/Private/VulkanBuffer.cpp b/Engine/Source/Runtime/VulkanRHI/Private/VulkanBuffer.cpp
--- a/Engine/Source/Runtime/VulkanRHI/Private/VulkanBuffer.cpp
+++ b/Engine/Source/Runtime/VulkanRHI/Private/VulkanBuffer.cpp
@@ -458,11 +458,6 @@ void FVulkanDynamicRHI::UnlockBuffer_BottomOfPipe(FRHICommandListBase& RHICmdLis
 	Buffer->Unlock(RHICmdList);
 }
 
-void FVulkanDynamicRHI::RHICopyBuffer(FRHIBuffer* SourceBufferRHI, FRHIBuffer* DestBufferRHI)
-{
-	VULKAN_SIGNAL_UNIMPLEMENTED();
-}
-
 void* FVulkanDynamicRHI::RHILockBuffer(FRHICommandListBase& RHICmdList, FRHIBuffer* BufferRHI, uint32 Offset, uint32 Size, EResourceLockMode LockMode)
 {
 	QUICK_SCOPE_CYCLE_COUNTER(STAT_LockBuffer_RenderThread);
diff --git a/Engine/Source/Runtime/VulkanRHI/Public/VulkanDynamicRHI.h b/Engine/Source/Runtime/VulkanRHI/Public/VulkanDynamicRHI.h
--- a/Engine/Source/Runtime/VulkanRHI/Public/VulkanDynamicRHI.h
+++ b/Engine/Source/Runtime/VulkanRHI/Public/VulkanDynamicRHI.h
@@ -116,7 +116,6 @@ public:
 	virtual FUniformBufferRHIRef RHICreateUniformBuffer(const void* Contents, const FRHIUniformBufferLayout* Layout, EUniformBufferUsage Usage, EUniformBufferValidation Validation) final override;
 	virtual void RHIUpdateUniformBuffer(FRHICommandListBase& RHICmdList, FRHIUniformBuffer* UniformBufferRHI, const void* Contents) final override;
 	virtual FBufferRHIRef RHICreateBuffer(FRHICommandListBase& RHICmdList, FRHIBufferDesc const& Desc, ERHIAccess ResourceState, FRHIResourceCreateInfo& CreateInfo) final override;
-	virtual void RHICopyBuffer(FRHIBuffer* SourceBuffer, FRHIBuffer* DestBuffer) final override;
 	virtual void RHIReplaceResources(FRHICommandListBase& RHICmdList, TArray<FRHIResourceReplaceInfo>&& ReplaceInfos) final override;
 	virtual void* LockBuffer_BottomOfPipe(FRHICommandListBase& RHICmdList, FRHIBuffer* Buffer, uint32 Offset, uint32 Size, EResourceLockMode LockMode) final override;
 	virtual void UnlockBuffer_BottomOfPipe(FRHICommandListBase& RHICmdList, FRHIBuffer* Buffer) final override;
diff --git a/Engine/Source/Runtime/Windows/D3D11RHI/Private/D3D11Buffer.cpp b/Engine/Source/Runtime/Windows/D3D11RHI/Private/D3D11Buffer.cpp
--- a/Engine/Source/Runtime/Windows/D3D11RHI/Private/D3D11Buffer.cpp
+++ b/Engine/Source/Runtime/Windows/D3D11RHI/Private/D3D11Buffer.cpp
@@ -260,24 +260,6 @@ void FD3D11DynamicRHI::UnlockBuffer_BottomOfPipe(FRHICommandListBase& RHICmdList
 	}
 }
 
-void FD3D11DynamicRHI::RHICopyBuffer(FRHIBuffer* SourceBufferRHI, FRHIBuffer* DestBufferRHI)
-{
-	FD3D11Buffer* SourceBuffer = ResourceCast(SourceBufferRHI);
-	FD3D11Buffer* DestBuffer = ResourceCast(DestBufferRHI);
-
-	D3D11_BUFFER_DESC SourceBufferDesc;
-	SourceBuffer->Resource->GetDesc(&SourceBufferDesc);
-	
-	D3D11_BUFFER_DESC DestBufferDesc;
-	DestBuffer->Resource->GetDesc(&DestBufferDesc);
-
-	check(SourceBufferDesc.ByteWidth == DestBufferDesc.ByteWidth);
-
-	Direct3DDeviceIMContext->CopyResource(DestBuffer->Resource,SourceBuffer->Resource);
-
-	RegisterGPUWork(1);
-}
-
 void FD3D11DynamicRHI::RHIReplaceResources(FRHICommandListBase& RHICmdList, TArray<FRHIResourceReplaceInfo>&& ReplaceInfos)
 {
 	RHICmdList.EnqueueLambda(TEXT("FD3D11DynamicRHI::RHIReplaceResources"),
diff --git a/Engine/Source/Runtime/Windows/D3D11RHI/Private/D3D11RHIPrivate.h b/Engine/Source/Runtime/Windows/D3D11RHI/Private/D3D11RHIPrivate.h
--- a/Engine/Source/Runtime/Windows/D3D11RHI/Private/D3D11RHIPrivate.h
+++ b/Engine/Source/Runtime/Windows/D3D11RHI/Private/D3D11RHIPrivate.h
@@ -444,7 +444,6 @@ public:
 	virtual FUniformBufferRHIRef RHICreateUniformBuffer(const void* Contents, const FRHIUniformBufferLayout* Layout, EUniformBufferUsage Usage, EUniformBufferValidation Validation) final override;
 	virtual void RHIUpdateUniformBuffer(FRHICommandListBase& RHICmdList, FRHIUniformBuffer* UniformBufferRHI, const void* Contents) final override;
 	virtual FBufferRHIRef RHICreateBuffer(FRHICommandListBase& RHICmdList, FRHIBufferDesc const& Desc, ERHIAccess ResourceState, FRHIResourceCreateInfo& CreateInfo) final override;
-	virtual void RHICopyBuffer(FRHIBuffer* SourceBuffer, FRHIBuffer* DestBuffer) final override;
 	virtual void RHIReplaceResources(FRHICommandListBase& RHICmdList, TArray<FRHIResourceReplaceInfo>&& ReplaceInfos) final override;
 	virtual void* LockBuffer_BottomOfPipe(FRHICommandListBase& RHICmdList, FRHIBuffer* Buffer, uint32 Offset, uint32 Size, EResourceLockMode LockMode) final override;
 	virtual void UnlockBuffer_BottomOfPipe(FRHICommandListBase& RHICmdList, FRHIBuffer* Buffer) final override;
