Defer residency update for reserved resources as their residency handles are only known on the RHI Submission Thread

* Fixes crash due to race condition between parallel translation and tile mapping updates

diff --git a/Engine/Source/Runtime/D3D12RHI/Private/D3D12CommandContext.h b/Engine/Source/Runtime/D3D12RHI/Private/D3D12CommandContext.h
--- a/Engine/Source/Runtime/D3D12RHI/Private/D3D12CommandContext.h
+++ b/Engine/Source/Runtime/D3D12RHI/Private/D3D12CommandContext.h
@@ -392,15 +392,9 @@ public:
 	void EndQuery  (FD3D12QueryLocation const& Location) { GetCommandList().EndQuery  (Location); }
 
 #if ENABLE_RESIDENCY_MANAGEMENT
-	void UpdateResidency(TConstArrayView<FD3D12ResidencyHandle*> Handles) { GetCommandList().UpdateResidency(Handles); }
-	void UpdateResidency(FD3D12ResidencyHandle& Handle                  ) { GetCommandList().UpdateResidency({ &Handle }); }
-	void UpdateResidency(FD3D12ResidencyHandle* Handle                  ) { check(Handle  ); GetCommandList().UpdateResidency({ Handle }); }
-	void UpdateResidency(FD3D12Resource* Resource                       ) { check(Resource); GetCommandList().UpdateResidency(Resource->GetResidencyHandles()); }
+	void UpdateResidency(const FD3D12Resource* Resource) { check(Resource); GetCommandList().UpdateResidency(Resource); }
 #else
-	void UpdateResidency(TConstArrayView<FD3D12ResidencyHandle*> Handles) { }
-	void UpdateResidency(FD3D12ResidencyHandle& Handle                  ) { }
-	void UpdateResidency(FD3D12ResidencyHandle* Handle                  ) { }
-	void UpdateResidency(FD3D12Resource* Resource                       ) { }
+	void UpdateResidency(const FD3D12Resource* Resource) { }
 #endif
 
 	// Pending resource barriers are resolved by a dedicated barrier command list, generated during command list submission.
diff --git a/Engine/Source/Runtime/D3D12RHI/Private/D3D12CommandList.cpp b/Engine/Source/Runtime/D3D12RHI/Private/D3D12CommandList.cpp
--- a/Engine/Source/Runtime/D3D12RHI/Private/D3D12CommandList.cpp
+++ b/Engine/Source/Runtime/D3D12RHI/Private/D3D12CommandList.cpp
@@ -19,10 +19,35 @@ static FAutoConsoleVariableRef CVarD3D12ExtraDepthTransitions(
 
 int64 FD3D12CommandList::FState::NextCommandListID = 0;
 
-void FD3D12CommandList::UpdateResidency(TConstArrayView<FD3D12ResidencyHandle*> Handles)
+void FD3D12CommandList::UpdateResidency(const FD3D12Resource* Resource)
 {
 #if ENABLE_RESIDENCY_MANAGEMENT
-	for (FD3D12ResidencyHandle* Handle : Handles)
+	if (Resource->NeedsDeferredResidencyUpdate())
+	{
+		State.DeferredResidencyUpdateSet.Add(Resource);
+	}
+	else
+	{
+		AddToResidencySet(Resource->GetResidencyHandles());
+	}
+#endif // ENABLE_RESIDENCY_MANAGEMENT
+}
+
+#if ENABLE_RESIDENCY_MANAGEMENT
+FD3D12ResidencySet* FD3D12CommandList::CloseResidencySet()
+{
+	for (const FD3D12Resource* Resource : State.DeferredResidencyUpdateSet)
+	{
+		AddToResidencySet(Resource->GetResidencyHandles());
+	}
+
+	D3DX12Residency::Close(ResidencySet);
+	return ResidencySet;
+}
+
+void FD3D12CommandList::AddToResidencySet(TConstArrayView<FD3D12ResidencyHandle*> ResidencyHandles)
+{
+	for (FD3D12ResidencyHandle* Handle : ResidencyHandles)
 	{
 		if (D3DX12Residency::IsInitialized(Handle))
 		{
@@ -30,8 +55,8 @@ void FD3D12CommandList::UpdateResidency(TConstArrayView<FD3D12ResidencyHandle*>
 			D3DX12Residency::Insert(*ResidencySet, *Handle);
 		}
 	}
-#endif
 }
+#endif // ENABLE_RESIDENCY_MANAGEMENT
 
 void FD3D12ContextCommon::AddPendingResourceBarrier(FD3D12Resource* Resource, D3D12_RESOURCE_STATES After, uint32 SubResource, CResourceState& ResourceState_OnCommandList)
 {
@@ -278,7 +303,7 @@ void FD3D12CommandList::Close()
 	{
 		VERIFYD3D12RESULT(Interfaces.GraphicsCommandList->Close());
 	}
-	D3DX12Residency::Close(ResidencySet);
+
 	State.IsClosed = true;
 }
 
diff --git a/Engine/Source/Runtime/D3D12RHI/Private/D3D12CommandList.h b/Engine/Source/Runtime/D3D12RHI/Private/D3D12CommandList.h
--- a/Engine/Source/Runtime/D3D12RHI/Private/D3D12CommandList.h
+++ b/Engine/Source/Runtime/D3D12RHI/Private/D3D12CommandList.h
@@ -93,13 +93,27 @@ public:
 
 	FD3D12Device*       const Device;
 	ED3D12QueueType     const QueueType;
+
+private:
 	FD3D12ResidencySet* const ResidencySet;
 
+public:
 	// Get the state of a resource on this command lists.
 	// This is only used for resources that require state tracking.
 	CResourceState& GetResourceState_OnCommandList(FD3D12Resource* pResource);
 
-	void UpdateResidency(TConstArrayView<FD3D12ResidencyHandle*> Handles);
+	// Indicate that a resource must be made resident before execution on GPU.
+	// Either immediately adds residency handle for this resource to the residency set
+	// or defers it until submission time (residency handles may not be known until then).
+	void UpdateResidency(const FD3D12Resource* Resource);
+
+#if ENABLE_RESIDENCY_MANAGEMENT
+	// Immediately add residency handles to the residency set for this command list.
+	void AddToResidencySet(TConstArrayView<FD3D12ResidencyHandle*> Handles);
+
+	// Closes and returns the residency set. Used only during submission.
+	FD3D12ResidencySet* CloseResidencySet();
+#endif // ENABLE_RESIDENCY_MANAGEMENT
 
 private:
 	struct FInterfaces
@@ -275,6 +289,10 @@ private:
 		TArray<FD3D12QueryLocation> OcclusionQueries;
 		TArray<FD3D12QueryLocation> PipelineStatsQueries;
 
+		// Resources whose residency must be updated on the submission thread, as their residency handles are not known during translation.
+		// This includes reserved resources that may refer to different heaps at different points on the submission timeline.
+		TSet<const FD3D12Resource*> DeferredResidencyUpdateSet;
+
 		uint32 NumCommands = 0;
 
 		bool IsClosed = false;
diff --git a/Engine/Source/Runtime/D3D12RHI/Private/D3D12RayTracing.cpp b/Engine/Source/Runtime/D3D12RHI/Private/D3D12RayTracing.cpp
--- a/Engine/Source/Runtime/D3D12RHI/Private/D3D12RayTracing.cpp
+++ b/Engine/Source/Runtime/D3D12RHI/Private/D3D12RayTracing.cpp
@@ -3557,27 +3557,45 @@ void FD3D12RayTracingScene::BuildAccelerationStructure(FD3D12CommandContext& Com
 		CommandContext.UpdateResidency(InstanceBuffer->GetResource());
 
 		{
-			TArray<FD3D12ResidencyHandle*>& GeometryResidencyHandlesForThisGPU = GeometryResidencyHandles[GPUIndex];
+			TArray<const FD3D12Resource*>& ResourcesToMakeResidentForThisGPU = ResourcesToMakeResident[GPUIndex];
 
-			GeometryResidencyHandlesForThisGPU.Reset(0);
+			ResourcesToMakeResidentForThisGPU.Reset(0);
 
 			Experimental::TSherwoodSet<FD3D12ResidencyHandle*> UniqueResidencyHandles;
 
-			auto AddResidencyHandleForResource = [&UniqueResidencyHandles, &GeometryResidencyHandlesForThisGPU] (FD3D12Resource* Resource)
+			auto AddResidencyHandleForResource = [&UniqueResidencyHandles, &ResourcesToMakeResidentForThisGPU] (FD3D12Resource* Resource)
 			{
 			#if ENABLE_RESIDENCY_MANAGEMENT
-				for (FD3D12ResidencyHandle* ResidencyHandle : Resource->GetResidencyHandles())
+
+				bool bShouldTrackResidency = false;
+
+				if (Resource->NeedsDeferredResidencyUpdate())
 				{
-					if (D3DX12Residency::IsInitialized(ResidencyHandle))
+					// Resources whose residency handles might change dynamically must always be tracked
+					bShouldTrackResidency = true;
+				}
+				else
+				{
+					// Resources that share *all* residency handles with what's already tracked don't need to be added to be tracked separately
+					for (FD3D12ResidencyHandle* ResidencyHandle : Resource->GetResidencyHandles())
 					{
-						bool bIsAlreadyInSet = false;
-						UniqueResidencyHandles.Add(ResidencyHandle, &bIsAlreadyInSet);
-						if (!bIsAlreadyInSet)
+						if (D3DX12Residency::IsInitialized(ResidencyHandle))
 						{
-							GeometryResidencyHandlesForThisGPU.Add(ResidencyHandle);
+							bool bIsAlreadyInSet = false;
+							UniqueResidencyHandles.Add(ResidencyHandle, &bIsAlreadyInSet);
+							if (!bIsAlreadyInSet)
+							{
+								bShouldTrackResidency = true;
+							}
 						}
 					}
 				}
+
+				if (bShouldTrackResidency)
+				{
+					ResourcesToMakeResidentForThisGPU.Add(Resource);
+				}
+
 			#endif // ENABLE_RESIDENCY_MANAGEMENT
 			};
 
@@ -3733,7 +3751,10 @@ void FD3D12RayTracingScene::UpdateResidency(FD3D12CommandContext& CommandContext
 #if ENABLE_RESIDENCY_MANAGEMENT
 	const uint32 GPUIndex = CommandContext.GetGPUIndex();
 	CommandContext.UpdateResidency(AccelerationStructureBuffers[GPUIndex]->GetResource());
-	CommandContext.UpdateResidency(GeometryResidencyHandles[GPUIndex]);
+	for (const FD3D12Resource* Resource : ResourcesToMakeResident[GPUIndex])
+	{
+		CommandContext.UpdateResidency(Resource);
+	}
 #endif // ENABLE_RESIDENCY_MANAGEMENT
 }
 
diff --git a/Engine/Source/Runtime/D3D12RHI/Private/D3D12RayTracing.h b/Engine/Source/Runtime/D3D12RHI/Private/D3D12RayTracing.h
--- a/Engine/Source/Runtime/D3D12RHI/Private/D3D12RayTracing.h
+++ b/Engine/Source/Runtime/D3D12RHI/Private/D3D12RayTracing.h
@@ -128,7 +128,8 @@ public:
 
 	// Scene keeps track of child acceleration structure buffers to ensure
 	// they are resident when any ray tracing work is dispatched.
-	TArray<FD3D12ResidencyHandle*> GeometryResidencyHandles[MAX_NUM_GPUS];
+	// Resources that share residency handles are deduplicated.
+	TArray<const FD3D12Resource*> ResourcesToMakeResident[MAX_NUM_GPUS];
 
 	void UpdateResidency(FD3D12CommandContext& CommandContext) const;
 
diff --git a/Engine/Source/Runtime/D3D12RHI/Private/D3D12Resources.h b/Engine/Source/Runtime/D3D12RHI/Private/D3D12Resources.h
--- a/Engine/Source/Runtime/D3D12RHI/Private/D3D12Resources.h
+++ b/Engine/Source/Runtime/D3D12RHI/Private/D3D12Resources.h
@@ -342,11 +342,21 @@ public:
 	inline FD3D12Heap* GetHeap() const { return Heap; };
 	inline bool IsDepthStencilResource() const { return bDepthStencil; }
 
+	inline bool NeedsDeferredResidencyUpdate() const { return IsReservedResource(); }
+
 	void StartTrackingForResidency();
 
 	bool IsResident() const
 	{
 #if ENABLE_RESIDENCY_MANAGEMENT
+
+		if (NeedsDeferredResidencyUpdate())
+		{
+			// We don't know the state because the set of residency handles is only known on the 
+			// RHI Submission Thread and may change throughout the frame.
+			return true;
+		}
+
 		TConstArrayView<FD3D12ResidencyHandle*> ResidencyHandles = GetResidencyHandles();
 		if (ResidencyHandles.IsEmpty())
 		{
diff --git a/Engine/Source/Runtime/D3D12RHI/Private/D3D12Submission.cpp b/Engine/Source/Runtime/D3D12RHI/Private/D3D12Submission.cpp
--- a/Engine/Source/Runtime/D3D12RHI/Private/D3D12Submission.cpp
+++ b/Engine/Source/Runtime/D3D12RHI/Private/D3D12Submission.cpp
@@ -516,9 +516,9 @@ void FD3D12DynamicRHI::GenerateBarrierCommandListAndUpdateState(FD3D12CommandLis
 	FD3D12ResourceBarrierBatcher BatcherForGraphicsToAsync;
 
 #if ENABLE_RESIDENCY_MANAGEMENT
-	TArray<FD3D12ResidencyHandle*> ResidencyHandles;
-	ResidencyHandles.Reserve(SourceCommandList->State.PendingResourceBarriers.Num());
-#endif
+	TArray<FD3D12Resource*> ResourcesToMakeResident;
+	ResourcesToMakeResident.Reserve(SourceCommandList->State.PendingResourceBarriers.Num());
+#endif // ENABLE_RESIDENCY_MANAGEMENT
 
 	bool bHasGraphicStates = false;
 	for (const FD3D12PendingResourceBarrier& PRB : SourceCommandList->State.PendingResourceBarriers)
@@ -590,12 +590,9 @@ void FD3D12DynamicRHI::GenerateBarrierCommandListAndUpdateState(FD3D12CommandLis
 			}
 		}
 
-#if ENABLE_RESIDENCY_MANAGEMENT
-		for (FD3D12ResidencyHandle* Handle : PRB.Resource->GetResidencyHandles())
-		{
-			ResidencyHandles.Add(Handle);
-		}
-#endif // ENABLE_RESIDENCY_MANAGEMENT
+	#if ENABLE_RESIDENCY_MANAGEMENT
+		ResourcesToMakeResident.Add(PRB.Resource);
+	#endif // ENABLE_RESIDENCY_MANAGEMENT
 	}
 
 	// Update the tracked resource states with the final states from the command list
@@ -645,8 +642,11 @@ void FD3D12DynamicRHI::GenerateBarrierCommandListAndUpdateState(FD3D12CommandLis
 		FD3D12CommandList* BarrierCommandList = Queue.Device->ObtainCommandList(Queue.BarrierAllocator, &Queue.BarrierTimestamps, nullptr);
 
 	#if ENABLE_RESIDENCY_MANAGEMENT
-		BarrierCommandList->UpdateResidency(ResidencyHandles);
-	#endif
+		for (FD3D12Resource* Resource : ResourcesToMakeResident)
+		{
+			BarrierCommandList->UpdateResidency(Resource);
+		}
+	#endif // ENABLE_RESIDENCY_MANAGEMENT
 
 		Batcher.FlushIntoCommandList(*BarrierCommandList, Queue.BarrierTimestamps);
 		BarrierCommandList->Close();
@@ -748,7 +748,7 @@ uint64 FD3D12Queue::FinalizePayload(bool bRequiresSignal, TArray<FD3D12Payload*,
 						TArray<FD3D12ResidencyHandle*, TInlineAllocator<2>> ResidencyHandles;
 						ResidencyHandles.Add(&Heap->GetHeapResidencyHandle());
 						ResidencyHandles.Append(Heap->GetResultBuffer()->GetResidencyHandles());
-						GetResolveCommandList()->UpdateResidency(ResidencyHandles);
+						GetResolveCommandList()->AddToResidencySet(ResidencyHandles);
 #endif // ENABLE_RESIDENCY_MANAGEMENT
 					}
 
@@ -899,7 +899,7 @@ void FD3D12DynamicRHI::FlushBatchedPayloads(TArray<FD3D12Payload*, TInlineAlloca
 				D3DCommandLists.Add(CommandList->Interfaces.CommandList);
 
 #if ENABLE_RESIDENCY_MANAGEMENT
-				ResidencySets.Add(CommandList->ResidencySet);
+				ResidencySets.Add(CommandList->CloseResidencySet());
 #endif
 			}
 			CommandLists.Append(MoveTemp(Payload->CommandListsToExecute));
diff --git a/Engine/Source/Runtime/D3D12RHI/Private/D3D12UAV.cpp b/Engine/Source/Runtime/D3D12RHI/Private/D3D12UAV.cpp
--- a/Engine/Source/Runtime/D3D12RHI/Private/D3D12UAV.cpp
+++ b/Engine/Source/Runtime/D3D12RHI/Private/D3D12UAV.cpp
@@ -392,7 +392,7 @@ void FD3D12CommandContext::ClearUAV(TRHICommandList_RecursiveHazardous<FD3D12Com
 
 				Context.FlushResourceBarriers();
 				Context.GraphicsCommandList()->ClearUnorderedAccessViewUint(GPUHandle, OfflineHandle, UAV.GetResource()->GetResource(), *reinterpret_cast<const UINT(*)[4]>(ClearValues), 0, nullptr);
-				Context.UpdateResidency(UnorderedAccessView->GetResidencyHandles());
+				Context.UpdateResidency(UAV.GetResource());
 				Context.ConditionalSplitCommandList();
 
 				if (Context.IsDefaultContext())
