Fix MGPU D3D12 buffer upload using the wrong resource location.

diff --git a/Engine/Source/Runtime/D3D12RHI/Private/D3D12Buffer.cpp b/Engine/Source/Runtime/D3D12RHI/Private/D3D12Buffer.cpp
--- a/Engine/Source/Runtime/D3D12RHI/Private/D3D12Buffer.cpp
+++ b/Engine/Source/Runtime/D3D12RHI/Private/D3D12Buffer.cpp
@@ -99,7 +99,9 @@ void FD3D12Buffer::UploadResourceData(FRHICommandListBase& InRHICmdList, FRHIGPU
 			for (uint32 GPUIndex : EffectiveMask)
 			{
 				FD3D12CommandContext& CommandContext = FD3D12CommandContext::Get(ExecutingCmdList, GPUIndex);
-				UploadResourceData(CommandContext, InDestinationState, this->ResourceLocation, SrcResourceLoc, GetSize());
+				FD3D12Buffer* DeviceBuffer = CommandContext.RetrieveObject<FD3D12Buffer>(this, GPUIndex);
+
+				UploadResourceData(CommandContext, InDestinationState, DeviceBuffer->ResourceLocation, SrcResourceLoc, GetSize());
 			}
 		});
 }
