Remove use of GetDefaultCommandContext in D3D12RenderTarget.cpp
 - The various RHI surface readback functions should retrieve the context to use from the immediate command list.
 - Also, use FRHICommandListImmediate::SubmitAndBlockUntilGPUIdle() to submit work, rather than flushing the context directly inside the D3D12 RHI.

diff --git a/Engine/Source/Runtime/D3D12RHI/Private/D3D12RenderTarget.cpp b/Engine/Source/Runtime/D3D12RHI/Private/D3D12RenderTarget.cpp
--- a/Engine/Source/Runtime/D3D12RHI/Private/D3D12RenderTarget.cpp
+++ b/Engine/Source/Runtime/D3D12RHI/Private/D3D12RenderTarget.cpp
@@ -376,8 +376,6 @@ TRefCountPtr<FD3D12Resource> FD3D12DynamicRHI::GetStagingTexture(FRHITexture* Te
 	FD3D12Adapter* Adapter = Device->GetParentAdapter();
 	const FRHIGPUMask Node = Device->GetGPUMask();
 
-	FD3D12CommandContext& Context = Device->GetDefaultCommandContext();
-
 	FD3D12Texture* Texture = GetD3D12TextureFromRHITexture(TextureRHI, GPUIndex);
 	D3D12_RESOURCE_DESC const& SourceDesc = Texture->GetResource()->GetDesc();
 
@@ -454,23 +452,29 @@ TRefCountPtr<FD3D12Resource> FD3D12DynamicRHI::GetStagingTexture(FRHITexture* Te
 	CD3DX12_TEXTURE_COPY_LOCATION DestCopyLocation(TempTexture2D->GetResource(), DestFootprint);
 	CD3DX12_TEXTURE_COPY_LOCATION SourceCopyLocation(Texture->GetResource()->GetResource(), Subresource);
 
-	FScopedResourceBarrier ScopeResourceBarrierSource(Context, Texture->GetResource(), &Texture->ResourceLocation, D3D12_RESOURCE_STATE_COPY_SOURCE, SourceCopyLocation.SubresourceIndex);
-	Context.FlushResourceBarriers();
-	// Upload heap doesn't need to transition
+	// Remember the width, height, pitch, etc...
+	ReadbackHeapDesc = DestFootprint;
 
-	Context.GraphicsCommandList()->CopyTextureRegion(
-		&DestCopyLocation,
-		0, 0, 0,
-		&SourceCopyLocation,
-		RectPtr);
+	FRHICommandListImmediate& RHICmdList = FRHICommandListImmediate::Get();
+	RHICmdList.EnqueueLambda([&](FRHICommandList& ExecutingCmdList)
+	{
+		FD3D12CommandContext& Context = FD3D12CommandContext::Get(ExecutingCmdList, GPUIndex);
 
-	Context.UpdateResidency(Texture->GetResource());
+		FScopedResourceBarrier ScopeResourceBarrierSource(Context, Texture->GetResource(), &Texture->ResourceLocation, D3D12_RESOURCE_STATE_COPY_SOURCE, SourceCopyLocation.SubresourceIndex);
+		Context.FlushResourceBarriers();
+		// Upload heap doesn't need to transition
 
-	// Remember the width, height, pitch, etc...
-	ReadbackHeapDesc = DestFootprint;
+		Context.GraphicsCommandList()->CopyTextureRegion(
+			&DestCopyLocation,
+			0, 0, 0,
+			&SourceCopyLocation,
+			RectPtr);
+
+		Context.UpdateResidency(Texture->GetResource());
+	});
 
 	// We need to execute the command list so we can read the data from readback heap
-	Context.FlushCommands(ED3D12FlushFlags::WaitForCompletion);
+	RHICmdList.SubmitAndBlockUntilGPUIdle();
 
 	return TempTexture2D;
 }
@@ -478,7 +482,6 @@ TRefCountPtr<FD3D12Resource> FD3D12DynamicRHI::GetStagingTexture(FRHITexture* Te
 void FD3D12DynamicRHI::ReadSurfaceDataNoMSAARaw(FRHITexture* TextureRHI, FIntRect InRect, TArray<uint8>& OutData, FReadSurfaceDataFlags InFlags)
 {
 	const uint32 GPUIndex = InFlags.GetGPUIndex();
-	FD3D12Texture* Texture = GetD3D12TextureFromRHITexture(TextureRHI, GPUIndex);
 
 	const uint32 SizeX = InRect.Width();
 	const uint32 SizeY = InRect.Height();
@@ -586,11 +589,7 @@ void FD3D12DynamicRHI::RHIReadSurfaceData(FRHITexture* InRHITexture, FIntRect In
 #endif
 
 	// Retrieve the base texture
-	FD3D12Device* Device = GetRHIDevice(GPUIndex);
-	Device->BlockUntilIdle();
-
-	FD3D12CommandContext& CommandContext = Device->GetDefaultCommandContext();
-	FD3D12Texture* DestTexture2D = CommandContext.RetrieveTexture(InRHITexture);
+	FD3D12Texture* DestTexture2D = FD3D12CommandContext::RetrieveTexture(InRHITexture, GPUIndex);
 
 	// Check the format of the surface
 	FD3D12ResourceDesc const& TextureDesc = DestTexture2D->GetResource()->GetDesc();
@@ -657,8 +656,6 @@ void FD3D12DynamicRHI::ReadSurfaceDataMSAARaw(FRHITexture* TextureRHI, FIntRect
 	FD3D12Adapter* Adapter = Device->GetParentAdapter();
 	const FRHIGPUMask NodeMask = Device->GetGPUMask();
 
-	FD3D12CommandContext& DefaultContext = Device->GetDefaultCommandContext();
-
 	FD3D12Texture* Texture = GetD3D12TextureFromRHITexture(TextureRHI, GPUIndex);
 
 	const uint32 SizeX = InRect.Width();
@@ -751,36 +748,43 @@ void FD3D12DynamicRHI::ReadSurfaceDataMSAARaw(FRHITexture* TextureRHI, FIntRect
 	// Allocate the output buffer.
 	OutData.SetNumUninitialized(SizeX * SizeY * NumSamples * BytesPerPixel);
 
+	FRHICommandListImmediate& RHICmdList = FRHICommandListImmediate::Get();
+
 	// Can be optimized by doing all subsamples into a large enough rendertarget in one pass (multiple draw calls)
 	for (uint32 SampleIndex = 0; SampleIndex < NumSamples; ++SampleIndex)
 	{
-		// Resolve the sample to the non-MSAA render target.
-		DefaultContext.ResolveTextureUsingShader<FResolveSingleSamplePS>(
-			ResourceCast(TextureRHI->GetTexture2D()),
-			NULL,
-			&NonMSAARTV,
-			NULL,
-			NonMSAADesc,
-			FResolveRect(InRect.Min.X, InRect.Min.Y, InRect.Max.X, InRect.Max.Y),
-			FResolveRect(0, 0, SizeX, SizeY),
-			SampleIndex
-			);
-
-		FScopedResourceBarrier ScopeResourceBarrierSource(DefaultContext, NonMSAATexture2D, nullptr, D3D12_RESOURCE_STATE_COPY_SOURCE, SourceCopyLocation.SubresourceIndex);
-		// Upload heap doesn't need to transition
-
-		// Copy the resolved sample data to the staging texture.
-		DefaultContext.GraphicsCommandList()->CopyTextureRegion(
-			&DestCopyLocation,
-			0, 0, 0,
-			&SourceCopyLocation,
-			&Rect);
-
-		DefaultContext.UpdateResidency(StagingTexture2D);
-		DefaultContext.UpdateResidency(NonMSAATexture2D);
+		RHICmdList.EnqueueLambda([&](FRHICommandListBase& ExecutingCmdList)
+		{
+			FD3D12CommandContext& Context = FD3D12CommandContext::Get(ExecutingCmdList, GPUIndex);
+
+			// Resolve the sample to the non-MSAA render target.
+			Context.ResolveTextureUsingShader<FResolveSingleSamplePS>(
+				ResourceCast(TextureRHI->GetTexture2D()),
+				NULL,
+				&NonMSAARTV,
+				NULL,
+				NonMSAADesc,
+				FResolveRect(InRect.Min.X, InRect.Min.Y, InRect.Max.X, InRect.Max.Y),
+				FResolveRect(0, 0, SizeX, SizeY),
+				SampleIndex
+				);
+
+			FScopedResourceBarrier ScopeResourceBarrierSource(Context, NonMSAATexture2D, nullptr, D3D12_RESOURCE_STATE_COPY_SOURCE, SourceCopyLocation.SubresourceIndex);
+			// Upload heap doesn't need to transition
+
+			// Copy the resolved sample data to the staging texture.
+			Context.GraphicsCommandList()->CopyTextureRegion(
+				&DestCopyLocation,
+				0, 0, 0,
+				&SourceCopyLocation,
+				&Rect);
+
+			Context.UpdateResidency(StagingTexture2D);
+			Context.UpdateResidency(NonMSAATexture2D);
+		});
 
 		// We need to execute the command list so we can read the data in the map below
-		DefaultContext.FlushCommands(ED3D12FlushFlags::WaitForCompletion);
+		RHICmdList.SubmitAndBlockUntilGPUIdle();
 
 		// Lock the staging texture.
 		void* pData;
@@ -871,8 +875,6 @@ void FD3D12DynamicRHI::RHIReadSurfaceFloatData(FRHITexture* TextureRHI, FIntRect
 	FD3D12Adapter* Adapter = Device->GetParentAdapter();
 	const FRHIGPUMask Node = Device->GetGPUMask();
 
-	FD3D12CommandContext& DefaultContext = Device->GetDefaultCommandContext();
-
 	FD3D12Texture* Texture = GetD3D12TextureFromRHITexture(TextureRHI, GPUIndex);
 
 	const uint32 SizeX = InRect.Width();
@@ -937,22 +939,26 @@ void FD3D12DynamicRHI::RHIReadSurfaceFloatData(FRHITexture* TextureRHI, FIntRect
 	CD3DX12_TEXTURE_COPY_LOCATION DestCopyLocation(TempTexture2D->GetResource(), DestFootprint);
 	CD3DX12_TEXTURE_COPY_LOCATION SourceCopyLocation(Texture->GetResource()->GetResource(), Subresource);
 
+	FRHICommandListImmediate& RHICmdList = FRHICommandListImmediate::Get();
+	RHICmdList.EnqueueLambda([&](FRHICommandListBase& ExecutingCmdList)
 	{
-		FScopedResourceBarrier ConditionalScopeResourceBarrier(DefaultContext, Texture->GetResource(), &Texture->ResourceLocation, D3D12_RESOURCE_STATE_COPY_SOURCE, SourceCopyLocation.SubresourceIndex);
+		FD3D12CommandContext& Context = FD3D12CommandContext::Get(ExecutingCmdList, GPUIndex);
+
+		FScopedResourceBarrier ConditionalScopeResourceBarrier(Context, Texture->GetResource(), &Texture->ResourceLocation, D3D12_RESOURCE_STATE_COPY_SOURCE, SourceCopyLocation.SubresourceIndex);
 		// Don't need to transition upload heaps
 
-		DefaultContext.FlushResourceBarriers();
-		DefaultContext.GraphicsCommandList()->CopyTextureRegion(
+		Context.FlushResourceBarriers();
+		Context.GraphicsCommandList()->CopyTextureRegion(
 			&DestCopyLocation,
 			0, 0, 0,
 			&SourceCopyLocation,
 			&Rect);
 
-		DefaultContext.UpdateResidency(Texture->GetResource());
-	}
+		Context.UpdateResidency(Texture->GetResource());
+	});
 
 	// We need to execute the command list so we can read the data from the map below
-	DefaultContext.FlushCommands(ED3D12FlushFlags::WaitForCompletion);
+	RHICmdList.SubmitAndBlockUntilGPUIdle();
 
 	// Lock the staging resource.
 	void* pData;
@@ -989,8 +995,6 @@ void FD3D12DynamicRHI::RHIRead3DSurfaceFloatData(FRHITexture* TextureRHI, FIntRe
 	FD3D12Adapter* Adapter = Device->GetParentAdapter();
 	const FRHIGPUMask Node = Device->GetGPUMask();
 
-	FD3D12CommandContext& DefaultContext = Device->GetDefaultCommandContext();
-
 	FD3D12Texture* Texture = GetD3D12TextureFromRHITexture(TextureRHI, GPUIndex);
 
 	const uint32 SizeX = InRect.Width();
@@ -1046,22 +1050,26 @@ void FD3D12DynamicRHI::RHIRead3DSurfaceFloatData(FRHITexture* TextureRHI, FIntRe
 	CD3DX12_TEXTURE_COPY_LOCATION DestCopyLocation(TempTexture3D->GetResource(), PlacedTexture3D);
 	CD3DX12_TEXTURE_COPY_LOCATION SourceCopyLocation(Texture->GetResource()->GetResource(), Subresource);
 
+	FRHICommandListImmediate& RHICmdList = FRHICommandListImmediate::Get();
+	RHICmdList.EnqueueLambda([&](FRHICommandListBase& ExecutingCmdList)
 	{
-		FScopedResourceBarrier ConditionalScopeResourceBarrier(DefaultContext, Texture->GetResource(), &Texture->ResourceLocation, D3D12_RESOURCE_STATE_COPY_SOURCE, SourceCopyLocation.SubresourceIndex);
+		FD3D12CommandContext& Context = FD3D12CommandContext::Get(ExecutingCmdList, GPUIndex);
+
+		FScopedResourceBarrier ConditionalScopeResourceBarrier(Context, Texture->GetResource(), &Texture->ResourceLocation, D3D12_RESOURCE_STATE_COPY_SOURCE, SourceCopyLocation.SubresourceIndex);
 		// Don't need to transition upload heaps
 
-		DefaultContext.FlushResourceBarriers();
-		DefaultContext.GraphicsCommandList()->CopyTextureRegion(
+		Context.FlushResourceBarriers();
+		Context.GraphicsCommandList()->CopyTextureRegion(
 			&DestCopyLocation,
 			0, 0, 0,
 			&SourceCopyLocation,
 			&Rect);
 
-		DefaultContext.UpdateResidency(Texture->GetResource());
-	}
+		Context.UpdateResidency(Texture->GetResource());
+	});
 
 	// We need to execute the command list so we can read the data from the map below
-	DefaultContext.FlushCommands(ED3D12FlushFlags::WaitForCompletion);
+	RHICmdList.SubmitAndBlockUntilGPUIdle();
 
 	// Lock the staging resource.
 	void* pData;
