Remove use of GetDefaultCommandContext in D3D12 viewport
 - FD3D12Viewport::Present is already passed a D3D12 context, which can be used directly.

Removed FD3D12FramePacing
 - Dead code from mGPU AFR support which no longer exists.

diff --git a/Engine/Source/Runtime/D3D12RHI/Private/D3D12Viewport.cpp b/Engine/Source/Runtime/D3D12RHI/Private/D3D12Viewport.cpp
--- a/Engine/Source/Runtime/D3D12RHI/Private/D3D12Viewport.cpp
+++ b/Engine/Source/Runtime/D3D12RHI/Private/D3D12Viewport.cpp
@@ -38,97 +38,6 @@ using namespace D3D12RHI;
 
 FCriticalSection FD3D12Viewport::DXGIBackBufferLock;
 
-#if WITH_MGPU
-FD3D12FramePacing::FD3D12FramePacing(FD3D12Adapter* Parent)
-	: FD3D12AdapterChild(Parent)
-	, bKeepRunning(true)
-	, AvgFrameTimeMs(0.0f)
-	, LastFrameTimeMs(0)
-	, Thread(nullptr)
-{
-	VERIFYD3D12RESULT(Parent->GetD3DDevice()->CreateFence(0, D3D12_FENCE_FLAG_NONE, IID_PPV_ARGS(Fence.GetInitReference())));
-	FMemory::Memset(SleepTimes, 0);
-
-	Thread = FRunnableThread::Create(this, TEXT("FramePacer"), 0, TPri_AboveNormal);
-}
-
-FD3D12FramePacing::~FD3D12FramePacing()
-{
-	delete Thread;
-	Thread = nullptr;
-}
-
-bool FD3D12FramePacing::Init()
-{
-	Semaphore = CreateSemaphore(nullptr, 0, MaxFrames, nullptr);
-	return Semaphore != INVALID_HANDLE_VALUE;
-}
-
-void FD3D12FramePacing::Stop()
-{
-	bKeepRunning = false;
-	FMemory::Memset(SleepTimes, 0);
-
-	ReleaseSemaphore(Semaphore, 1, nullptr);
-	VERIFYD3D12RESULT(Fence->Signal(UINT64_MAX));
-}
-
-void FD3D12FramePacing::Exit()
-{
-	CloseHandle(Semaphore);
-}
-
-uint32 FD3D12FramePacing::Run()
-{
-	while (bKeepRunning)
-	{
-		// Wait for the present to be submitted so we know which GPU to wait on
-		WaitForSingleObjectEx(Semaphore, INFINITE, false);
-		check(CurIndex <= NextIndex || !bKeepRunning);
-
-		// Wait for the present to be completed so we can start timing to the next one
-		const uint32 ReadIndex = CurIndex % MaxFrames;
-
-		// Wait for the right amount of time to pass
-		const uint32 SleepTime = SleepTimes[ReadIndex];
-		Sleep(SleepTime);
-
-		VERIFYD3D12RESULT(Fence->Signal(++CurIndex));
-	}
-	return 0;
-}
-
-void FD3D12FramePacing::PrePresentQueued(ID3D12CommandQueue* Queue)
-{
-	const uint64 CurrTimeMs = GetTickCount64();
-	check(CurrTimeMs >= LastFrameTimeMs);
-
-	const float Delta = float(CurrTimeMs - LastFrameTimeMs);
-	const float Alpha = FMath::Clamp(Delta / 1000.0f / FramePacingAvgTimePeriod, 0.0f, 1.0f);
-
-	/** Number of milliseconds the GPU was busy last frame. */
-	/**
-	 * TODO:  Proper Multi-GPU support for measuring GPU frame cycles would involve doing something more complicated.
-	 * This could involve either using the frame cycles for whichever GPU was idle the least, or identifying time spans
-	 * where all GPUs are idle.  The latter seems more appropriate for the purpose of frame pacing, but would require
-	 * adding detailed idle time span tracking across all GPUs, which is some work.  The call below just queries cycles
-	 * for GPU 0 only.
-	 */
-	const uint32 GPUCycles = RHIGetGPUFrameCycles();
-	const float GPUMsForFrame = FPlatformTime::ToMilliseconds(GPUCycles);
-
-	AvgFrameTimeMs = (Alpha * GPUMsForFrame) + ((1.0f - Alpha) * AvgFrameTimeMs);
-	LastFrameTimeMs = CurrTimeMs;
-
-	const float TargetFrameTime = AvgFrameTimeMs * FramePacingPercentage;
-
-	const uint32 WriteIndex = NextIndex % MaxFrames;
-	SleepTimes[WriteIndex] = (uint32)TargetFrameTime;
-	VERIFYD3D12RESULT(Queue->Wait(Fence, ++NextIndex));
-	ReleaseSemaphore(Semaphore, 1, nullptr);
-}
-#endif //WITH_MGPU
-
 /**
  * Creates a FD3D12Surface to represent a swap chain's back buffer.
  */
@@ -301,14 +210,6 @@ FD3D12Viewport::~FD3D12Viewport()
 
 	GetParentAdapter()->GetViewports().Remove(this);
 
-#if WITH_MGPU
-	if (FramePacerRunnable)
-	{
-		delete FramePacerRunnable;
-		FramePacerRunnable = nullptr;
-	}
-#endif //WITH_MGPU
-
 	FinalDestroyInternal();
 }
 
@@ -618,23 +519,23 @@ bool FD3D12Viewport::PresentChecked(IRHICommandContext& RHICmdContext, int32 Syn
 	return bNeedNativePresent;
 }
 
-bool FD3D12Viewport::Present(IRHICommandContext& RHICmdContext, bool bLockToVsync)
+bool FD3D12Viewport::Present(FD3D12CommandContextBase& ContextBase, bool bLockToVsync)
 {
 	if (!IsPresentAllowed())
 	{
 		return false;
 	}
 
-	FD3D12Adapter* Adapter = GetParentAdapter();
+	check(ContextBase.GetParentAdapter() == GetParentAdapter());
 	
 	for (uint32 GPUIndex : FRHIGPUMask::All())
 	{
-		FD3D12CommandContext& DefaultContext = Adapter->GetDevice(GPUIndex)->GetDefaultCommandContext();
+		FD3D12CommandContext& Context = *ContextBase.GetSingleDeviceContext(GPUIndex);
 
 		// Those are not necessarily the swap chain back buffer in case of multi-gpu
-		FD3D12Texture* DeviceBackBuffer = DefaultContext.RetrieveObject<FD3D12Texture, FRHITexture>(GetBackBuffer_RHIThread());
+		FD3D12Texture* DeviceBackBuffer = Context.RetrieveObject<FD3D12Texture, FRHITexture>(GetBackBuffer_RHIThread());
 
-		DefaultContext.TransitionResource(
+		Context.TransitionResource(
 			DeviceBackBuffer->GetShaderResourceView()->GetResource(),
 			D3D12_RESOURCE_STATE_TBD,
 			D3D12_RESOURCE_STATE_PRESENT,
@@ -644,9 +545,9 @@ bool FD3D12Viewport::Present(IRHICommandContext& RHICmdContext, bool bLockToVsyn
 #if D3D12RHI_USE_SDR_BACKBUFFER
 		if (CurrentBackBuffer_RHIThread->TextureSDR)
 		{
-			FD3D12Texture* DeviceSDRBackBuffer = DefaultContext.RetrieveObject<FD3D12Texture, FRHITexture>(GetSDRBackBuffer_RHIThread());
+			FD3D12Texture* DeviceSDRBackBuffer = Context.RetrieveObject<FD3D12Texture, FRHITexture>(GetSDRBackBuffer_RHIThread());
 
-			DefaultContext.TransitionResource(
+			Context.TransitionResource(
 				DeviceSDRBackBuffer->GetShaderResourceView()->GetResource(),
 				D3D12_RESOURCE_STATE_TBD,
 				D3D12_RESOURCE_STATE_PRESENT,
@@ -655,37 +556,15 @@ bool FD3D12Viewport::Present(IRHICommandContext& RHICmdContext, bool bLockToVsyn
 		}
 #endif
 
-		DefaultContext.FlushResourceBarriers();
-	}
+		Context.FlushResourceBarriers();
 
-	FD3D12CommandContext* PresentContext = nullptr;
-#if WITH_MGPU
-	if (FramePacerRunnable)
-	{
-		delete FramePacerRunnable;
-		FramePacerRunnable = nullptr;
-	}
-#endif //WITH_MGPU
-
-	for (uint32 GPUIndex : FRHIGPUMask::All())
-	{
 		// Currently, the swap chain Present() is called directly by the RHI thread.
 		// We need to submit the above commands and wait for the submission thread to process everything before we can continue.
-		FD3D12CommandContext& DefaultContext = Adapter->GetDevice(GPUIndex)->GetDefaultCommandContext();
-		DefaultContext.FlushCommands(ED3D12FlushFlags::WaitForSubmission);
+		Context.FlushCommands(ED3D12FlushFlags::WaitForSubmission);
 	}
 
 	const int32 SyncInterval = bLockToVsync ? RHIGetSyncInterval() : 0;
-	const bool bNativelyPresented = PresentChecked(RHICmdContext, SyncInterval);
-
-#if WITH_MGPU
-	if (PresentContext)
-	{
-		LastFrameSyncPoint = FD3D12SyncPoint::Create(ED3D12SyncPointType::GPUOnly);
-		PresentContext->SignalSyncPoint(LastFrameSyncPoint);
-		PresentContext->FlushCommands();
-	}
-#endif
+	const bool bNativelyPresented = PresentChecked(ContextBase, SyncInterval);
 
 	if (bNativelyPresented || (CustomPresent && CustomPresent->NeedsAdvanceBackbuffer()))
 	{
diff --git a/Engine/Source/Runtime/D3D12RHI/Private/D3D12Viewport.h b/Engine/Source/Runtime/D3D12RHI/Private/D3D12Viewport.h
--- a/Engine/Source/Runtime/D3D12RHI/Private/D3D12Viewport.h
+++ b/Engine/Source/Runtime/D3D12RHI/Private/D3D12Viewport.h
@@ -19,52 +19,6 @@ class FD3D12UnorderedAccessView_RHI;
 class FD3D12SyncPoint;
 using FD3D12SyncPointRef = TRefCountPtr<FD3D12SyncPoint>;
 
-#define ALTERNATE_TIMESTAMP_METRIC 0
-
-#if WITH_MGPU
-class FD3D12FramePacing : public FRunnable, public FD3D12AdapterChild
-{
-public:
-	explicit FD3D12FramePacing(FD3D12Adapter* Parent);
-	~FD3D12FramePacing();
-	bool Init() override;
-	void Stop() override;
-	void Exit() override;
-	uint32 Run() override;
-
-	void PrePresentQueued(ID3D12CommandQueue* Queue);
-
-private:
-	static const uint32 MaxFrames = MAX_NUM_GPUS + 1;
-
-	TRefCountPtr<ID3D12Fence> Fence;
-	uint64 NextIndex = 0;
-	uint64 CurIndex = 0;
-	uint32 SleepTimes[MaxFrames];
-	HANDLE Semaphore;
-	bool bKeepRunning;
-
-	float AvgFrameTimeMs;
-	uint64 LastFrameTimeMs;
-
-	// ======== Some knobs for tweaking the algorithm ========
-	// How long to average the GPU time over, in seconds.
-	// - Higher = Smoother when framerate is steady, less smooth when frametime drops.
-	// - Lower = Quicker to smooth out after frametime drops, less smooth from incremental changes.
-	const float FramePacingAvgTimePeriod = 0.25f;
-	// What percentage of average GPU time to wait for on the pacing thread.
-	// - Higher = More consistent pacing, potential to starve the GPU in order to maintain pacing.
-	// - Lower = More allowable deviation between frame times, depending on GPU workload.
-#if ALTERNATE_TIMESTAMP_METRIC
-	const float FramePacingPercentage = 1.15f;
-#else
-	const float FramePacingPercentage = 1.05f;
-#endif
-
-	FRunnableThread* Thread;
-};
-#endif //WITH_MGPU
-
 class FD3D12Viewport : public FRHIViewport, public FD3D12AdapterChild
 {
 public:
@@ -92,7 +46,7 @@ public:
 	/** Presents the swap chain.
 	 * Returns true if Present was done by Engine.
 	 */
-	bool Present(IRHICommandContext& RHICmdContext, bool bLockToVsync);
+	bool Present(class FD3D12CommandContextBase& Context, bool bLockToVsync);
 
 	// Accessors.
 	FIntPoint GetSizeXY() const
@@ -301,10 +255,6 @@ private:
 
 	// Can very rarely be modified on the RHI thread as well if present is skipped
 	mutable FCriticalSection ExpectedBackBufferIndexLock;
-
-	FD3D12SyncPointRef LastFrameSyncPoint;
-
-	FD3D12FramePacing* FramePacerRunnable = nullptr;
 #endif
 };
 
