Add missing MGPU RHI raytracing function redirects
 - Were causing checkNoEntry() asserts when running with MGPU enabled.

diff --git a/Engine/Source/Runtime/D3D12RHI/Private/D3D12CommandContext.h b/Engine/Source/Runtime/D3D12RHI/Private/D3D12CommandContext.h
--- a/Engine/Source/Runtime/D3D12RHI/Private/D3D12CommandContext.h
+++ b/Engine/Source/Runtime/D3D12RHI/Private/D3D12CommandContext.h
@@ -1112,6 +1112,16 @@ public:
 		ContextRedirect(RHIBindAccelerationStructureMemory(Scene, Buffer, BufferOffset));
 	}
 
+	virtual void RHICommitShaderBindingTable(FRHIShaderBindingTable* SBT) final override
+	{
+		ContextRedirect(RHICommitShaderBindingTable(SBT));
+	}
+
+	virtual void RHICommitRayTracingBindings(FRHIRayTracingScene* Scene) final override
+	{
+		ContextRedirect(RHICommitRayTracingBindings(Scene));
+	}
+
 #endif // D3D12_RHI_RAYTRACING
 
 	virtual void RHISetGPUMask(FRHIGPUMask InGPUMask) final override
