Remove use of GetDefaultCommandContext in RHIGetGraphicsCommandList and RHIFinishExternalComputeWork
 - All uses of these functions had an FRHICommandList which the context could be extracted from, rather than grabbing the default context.

diff --git a/Engine/Source/Runtime/D3D12RHI/Private/D3D12RHI.cpp b/Engine/Source/Runtime/D3D12RHI/Private/D3D12RHI.cpp
--- a/Engine/Source/Runtime/D3D12RHI/Private/D3D12RHI.cpp
+++ b/Engine/Source/Runtime/D3D12RHI/Private/D3D12RHI.cpp
@@ -562,9 +562,10 @@ uint32 FD3D12DynamicRHI::RHIGetDeviceNodeMask(uint32 InIndex) const
 	return GetAdapter().GetDevice(InIndex)->GetGPUMask().GetNative();
 }
 
-ID3D12GraphicsCommandList* FD3D12DynamicRHI::RHIGetGraphicsCommandList(uint32 InDeviceIndex) const
+ID3D12GraphicsCommandList* FD3D12DynamicRHI::RHIGetGraphicsCommandList(FRHICommandListBase& ExecutingCmdList, uint32 InDeviceIndex) const
 {
-	return GetRHIDevice(InDeviceIndex)->GetDefaultCommandContext().GraphicsCommandList().Get();
+	FD3D12CommandContext& Context = FD3D12CommandContext::Get(ExecutingCmdList, InDeviceIndex);
+	return Context.GraphicsCommandList().Get();
 }
 
 DXGI_FORMAT FD3D12DynamicRHI::RHIGetSwapChainFormat(EPixelFormat InFormat) const
@@ -627,14 +628,13 @@ D3D12_CPU_DESCRIPTOR_HANDLE FD3D12DynamicRHI::RHIGetRenderTargetView(FRHITexture
 	return RTV ? RTV->GetOfflineCpuHandle() : D3D12_CPU_DESCRIPTOR_HANDLE{};
 }
 
-void FD3D12DynamicRHI::RHIFinishExternalComputeWork(uint32 InDeviceIndex, ID3D12GraphicsCommandList* InCommandList)
+void FD3D12DynamicRHI::RHIFinishExternalComputeWork(FRHICommandListBase& ExecutingCmdList, uint32 InDeviceIndex, ID3D12GraphicsCommandList* InCommandList)
 {
-	FD3D12Device* Device = GetRHIDevice(InDeviceIndex);
+	FD3D12CommandContext& Context = FD3D12CommandContext::Get(ExecutingCmdList, InDeviceIndex);
+	check(InCommandList == Context.GraphicsCommandList().GetNoRefCount());
 
-	check(InCommandList == Device->GetDefaultCommandContext().GraphicsCommandList().GetNoRefCount());
-
-	Device->GetDefaultCommandContext().StateCache.ForceSetComputeRootSignature();
-	Device->GetDefaultCommandContext().StateCache.GetDescriptorCache()->SetDescriptorHeaps(true);
+	Context.StateCache.ForceSetComputeRootSignature();
+	Context.StateCache.GetDescriptorCache()->SetDescriptorHeaps(true);
 }
 
 void FD3D12DynamicRHI::RHITransitionResource(FRHICommandList& RHICmdList, FRHITexture* InTexture, D3D12_RESOURCE_STATES InState, uint32 InSubResource)
diff --git a/Engine/Source/Runtime/D3D12RHI/Private/D3D12RHIPrivate.h b/Engine/Source/Runtime/D3D12RHI/Private/D3D12RHIPrivate.h
--- a/Engine/Source/Runtime/D3D12RHI/Private/D3D12RHIPrivate.h
+++ b/Engine/Source/Runtime/D3D12RHI/Private/D3D12RHIPrivate.h
@@ -412,7 +412,7 @@ public:
 	virtual ID3D12CommandQueue* RHIGetCommandQueue() const final override;
 	virtual ID3D12Device* RHIGetDevice(uint32 InIndex) const final override;
 	virtual uint32 RHIGetDeviceNodeMask(uint32 InIndex) const final override;
-	virtual ID3D12GraphicsCommandList* RHIGetGraphicsCommandList(uint32 InDeviceIndex) const final override;
+	virtual ID3D12GraphicsCommandList* RHIGetGraphicsCommandList(FRHICommandListBase& ExecutingCmdList, uint32 InDeviceIndex) const final override;
 	virtual DXGI_FORMAT RHIGetSwapChainFormat(EPixelFormat InFormat) const final override;
 	virtual FTextureRHIRef RHICreateTexture2DFromResource(EPixelFormat Format, ETextureCreateFlags TexCreateFlags, const FClearValueBinding& ClearValueBinding, ID3D12Resource* Resource) final override;
 	virtual FTextureRHIRef RHICreateTexture2DArrayFromResource(EPixelFormat Format, ETextureCreateFlags TexCreateFlags, const FClearValueBinding& ClearValueBinding, ID3D12Resource* Resource) final override;
@@ -426,7 +426,7 @@ public:
 	virtual int64 RHIGetResourceMemorySize(FRHITexture* InTexture) const final override;
 	virtual bool RHIIsResourcePlaced(FRHITexture* InTexture) const final override;
 	virtual D3D12_CPU_DESCRIPTOR_HANDLE RHIGetRenderTargetView(FRHITexture* InTexture, int32 InMipIndex = 0, int32 InArraySliceIndex = 0) const final override;
-	virtual void RHIFinishExternalComputeWork(uint32 InDeviceIndex, ID3D12GraphicsCommandList* InCommandList) final override;
+	virtual void RHIFinishExternalComputeWork(FRHICommandListBase& ExecutingCmdList, uint32 InDeviceIndex, ID3D12GraphicsCommandList* InCommandList) final override;
 	virtual void RHITransitionResource(FRHICommandList& RHICmdList, FRHITexture* InTexture, D3D12_RESOURCE_STATES InState, uint32 InSubResource) final override;
 	virtual void RHISignalManualFence(FRHICommandList& RHICmdList, ID3D12Fence* Fence, uint64 Value) final override;
 	virtual void RHIWaitManualFence(FRHICommandList& RHICmdList, ID3D12Fence* Fence, uint64 Value) final override;
diff --git a/Engine/Source/Runtime/D3D12RHI/Public/ID3D12DynamicRHI.h b/Engine/Source/Runtime/D3D12RHI/Public/ID3D12DynamicRHI.h
--- a/Engine/Source/Runtime/D3D12RHI/Public/ID3D12DynamicRHI.h
+++ b/Engine/Source/Runtime/D3D12RHI/Public/ID3D12DynamicRHI.h
@@ -28,7 +28,7 @@ struct ID3D12DynamicRHI : public FDynamicRHI
 	virtual ID3D12CommandQueue*        RHIGetCommandQueue() const = 0;
 	virtual ID3D12Device*              RHIGetDevice(uint32 InIndex) const = 0;
 	virtual uint32                     RHIGetDeviceNodeMask(uint32 InIndex) const = 0;
-	virtual ID3D12GraphicsCommandList* RHIGetGraphicsCommandList(uint32 InDeviceIndex) const = 0;
+	virtual ID3D12GraphicsCommandList* RHIGetGraphicsCommandList(FRHICommandListBase& ExecutingCmdList, uint32 InDeviceIndex) const = 0;
 	virtual DXGI_FORMAT                RHIGetSwapChainFormat(EPixelFormat InFormat) const = 0;
 
 	virtual FTextureRHIRef        RHICreateTexture2DFromResource(EPixelFormat Format, ETextureCreateFlags TexCreateFlags, const FClearValueBinding& ClearValueBinding, ID3D12Resource* Resource) = 0;
@@ -46,7 +46,7 @@ struct ID3D12DynamicRHI : public FDynamicRHI
 
 	virtual D3D12_CPU_DESCRIPTOR_HANDLE RHIGetRenderTargetView(FRHITexture* InTexture, int32 InMipIndex = 0, int32 InArraySliceIndex = 0) const = 0;
 
-	virtual void                  RHIFinishExternalComputeWork(uint32 InDeviceIndex, ID3D12GraphicsCommandList* InCommandList) = 0;
+	virtual void                  RHIFinishExternalComputeWork(FRHICommandListBase& ExecutingCmdList, uint32 InDeviceIndex, ID3D12GraphicsCommandList* InCommandList) = 0;
 	virtual void                  RHITransitionResource(FRHICommandList& RHICmdList, FRHITexture* InTexture, D3D12_RESOURCE_STATES InState, uint32 InSubResource) = 0;
 
 	virtual void                  RHISignalManualFence(FRHICommandList& RHICmdList, ID3D12Fence* Fence, uint64 Value) = 0;
