- Always create the LumenHardwareRayTracingUniformBuffer when raytracing is enabled because it's bound via the shader binding layout and is available to all shaders in the RTPSO
- Fix RHIValidation for Static uniform buffers with shader binding layout

diff --git a/Engine/Source/Runtime/D3D12RHI/Private/D3D12RayTracing.cpp b/Engine/Source/Runtime/D3D12RHI/Private/D3D12RayTracing.cpp
--- a/Engine/Source/Runtime/D3D12RHI/Private/D3D12RayTracing.cpp
+++ b/Engine/Source/Runtime/D3D12RHI/Private/D3D12RayTracing.cpp
@@ -4792,17 +4792,17 @@ static void DispatchRays(FD3D12CommandContext& CommandContext,
 	if (StaticShaderBindingSlot >= 0)
 	{
 		for (uint32 Index = 0; Index < ShaderBindingLayout.GetNumUniformBufferEntries(); ++Index)
-			{
+		{
 			const FRHIUniformBufferShaderBindingLayout& LayoutEntry = ShaderBindingLayout.GetUniformBufferEntry(Index);
 			const uint32 RootParameterSlotIndex = uint32(StaticShaderBindingSlot) + LayoutEntry.CBVResourceIndex;
 
 			FRHIUniformBuffer* UniformBuffer = StaticUniformBuffers[Index];
 			checkf(UniformBuffer, TEXT("Static uniform buffer at index %d is referenced in the shader binding layout but not provided in the last RHISetStaticUniformBuffers() command"), Index);
 
-				FD3D12UniformBuffer* D3D12UniformBuffer = FD3D12CommandContext::RetrieveObject<FD3D12UniformBuffer>(UniformBuffer, 0);//GpuIndex);
-				if (D3D12UniformBuffer->ResourceLocation.GetGPUVirtualAddress())
-				{
-					const FD3D12ResourceLocation& ResourceLocation = D3D12UniformBuffer->ResourceLocation;
+			FD3D12UniformBuffer* D3D12UniformBuffer = FD3D12CommandContext::RetrieveObject<FD3D12UniformBuffer>(UniformBuffer, Device->GetGPUIndex());
+			if (D3D12UniformBuffer->ResourceLocation.GetGPUVirtualAddress())
+			{
+				const FD3D12ResourceLocation& ResourceLocation = D3D12UniformBuffer->ResourceLocation;
 				CommandContext.GraphicsCommandList()->SetComputeRootConstantBufferView(RootParameterSlotIndex, ResourceLocation.GetGPUVirtualAddress());
 			}
 		}
diff --git a/Engine/Source/Runtime/RHI/Private/RHIContext.cpp b/Engine/Source/Runtime/RHI/Private/RHIContext.cpp
--- a/Engine/Source/Runtime/RHI/Private/RHIContext.cpp
+++ b/Engine/Source/Runtime/RHI/Private/RHIContext.cpp
@@ -83,6 +83,30 @@ void FUniformBufferStaticBindings::AddUniformBuffer(FRHIUniformBuffer* UniformBu
 	}
 }
 
+void FUniformBufferStaticBindings::Bind(TArray<FRHIUniformBuffer*>& Bindings) const
+{
+	Bindings.Reset();
+
+	if (ShaderBindingLayout)
+	{
+		Bindings.SetNumZeroed(UniformBuffers.Num());
+
+		for (int32 Index = 0; Index < UniformBuffers.Num(); ++Index)
+		{
+			Bindings[Index] = UniformBuffers[Index];
+		}
+	}
+	else
+	{
+		Bindings.SetNumZeroed(SlotCount);
+
+		for (int32 Index = 0; Index < UniformBuffers.Num(); ++Index)
+		{
+			Bindings[Slots[Index]] = UniformBuffers[Index];
+		}
+	}
+}
+
 void IRHICommandContextPSOFallback::SetGraphicsPipelineStateFromInitializer(const FGraphicsPipelineStateInitializer& PsoInit, uint32 StencilRef, bool bApplyAdditionalState)
 {
 	RHISetBoundShaderState(
diff --git a/Engine/Source/Runtime/RHI/Public/RHIContext.h b/Engine/Source/Runtime/RHI/Public/RHIContext.h
--- a/Engine/Source/Runtime/RHI/Public/RHIContext.h
+++ b/Engine/Source/Runtime/RHI/Public/RHIContext.h
@@ -96,16 +96,7 @@ public:
 		return ShaderBindingLayout;
 	}
 
-	void Bind(TArray<FRHIUniformBuffer*>& Bindings) const
-	{
-		Bindings.Reset();
-		Bindings.SetNumZeroed(SlotCount);
-
-		for (int32 Index = 0; Index < UniformBuffers.Num(); ++Index)
-		{
-			Bindings[Slots[Index]] = UniformBuffers[Index];
-		}
-	}
+	RHI_API void Bind(TArray<FRHIUniformBuffer*>& Bindings) const;
 
 private:
 	static const uint32 InlineUniformBufferCount = 8;
@@ -806,7 +797,7 @@ public:
 			ArgumentBuffer, ArgumentOffset);
 		PRAGMA_ENABLE_DEPRECATION_WARNINGS
 	}
-
+	
 	virtual void RHISetBindingsOnShaderBindingTable(FRHIShaderBindingTable* SBT, FRHIRayTracingPipelineState* Pipeline, uint32 NumBindings, const FRayTracingLocalShaderBindings* Bindings, ERayTracingBindingType BindingType)
 	{
 		checkNoEntry();
diff --git a/Engine/Source/Runtime/Renderer/Private/DeferredShadingRenderer.cpp b/Engine/Source/Runtime/Renderer/Private/DeferredShadingRenderer.cpp
--- a/Engine/Source/Runtime/Renderer/Private/DeferredShadingRenderer.cpp
+++ b/Engine/Source/Runtime/Renderer/Private/DeferredShadingRenderer.cpp
@@ -908,11 +908,7 @@ void FDeferredShadingSceneRenderer::WaitForRayTracingScene(FRDGBuilder& GraphBui
 		}
 	}
 
-	if (Lumen::UseHardwareRayTracing(ViewFamily)
-		|| ManyLights::UseHardwareRayTracing(ViewFamily))
-	{
-		SetupLumenHardwareRayTracingUniformBuffer(GraphBuilder, ReferenceView);
-	}
+	SetupLumenHardwareRayTracingUniformBuffer(GraphBuilder, ReferenceView);
 
 	SetupRayTracingPipelineStates(GraphBuilder);
 
