Eagerly flush PayloadsToHandDown in FD3D12DynamicRHI::ProcessSubmissionQueue when the list grows beyond a certain threshold

diff --git a/Engine/Source/Runtime/D3D12RHI/Private/D3D12Device.h b/Engine/Source/Runtime/D3D12RHI/Private/D3D12Device.h
--- a/Engine/Source/Runtime/D3D12RHI/Private/D3D12Device.h
+++ b/Engine/Source/Runtime/D3D12RHI/Private/D3D12Device.h
@@ -151,8 +151,11 @@ public:
 	// On some hardware, some auxiliary queue types may not support tile mapping and a separate queue must be used
 	bool bSupportsTileMapping = true;
 
+	static constexpr uint32 MaxBatchedPayloads = 128;
+	using FPayloadArray = TArray<FD3D12Payload*, TInlineAllocator<MaxBatchedPayloads>>;
+
 	// Batches the current payload's command lists, returning the latest fence value signaled for this queue.
-	uint64 FinalizePayload(bool bRequiresSignal, TArray<FD3D12Payload*, TInlineAllocator<64>>& PayloadsToHandDown);
+	uint64 FinalizePayload(bool bRequiresSignal, FPayloadArray& PayloadsToHandDown);
 
 	// Call the underlying ID3D12Queue::ExecuteCommandLists function
 	void ExecuteCommandLists(TArrayView<ID3D12CommandList*> D3DCommandLists
diff --git a/Engine/Source/Runtime/D3D12RHI/Private/D3D12RHIPrivate.h b/Engine/Source/Runtime/D3D12RHI/Private/D3D12RHIPrivate.h
--- a/Engine/Source/Runtime/D3D12RHI/Private/D3D12RHIPrivate.h
+++ b/Engine/Source/Runtime/D3D12RHI/Private/D3D12RHIPrivate.h
@@ -202,7 +202,7 @@ private:
 	TQueue<TArray<FD3D12Payload*>*, EQueueMode::Mpsc> PendingPayloadsForSubmission;
 
 	// Ensures all prior batched command lists have reached the driver ID3D12Queue object.
-	void FlushBatchedPayloads(TArray<FD3D12Payload*, TInlineAllocator<64>>& PayloadsToHandDown);
+	void FlushBatchedPayloads(FD3D12Queue::FPayloadArray& PayloadsToHandDown);
 
 	FCriticalSection SubmissionCS;
 	FCriticalSection InterruptCS;
diff --git a/Engine/Source/Runtime/D3D12RHI/Private/D3D12Submission.cpp b/Engine/Source/Runtime/D3D12RHI/Private/D3D12Submission.cpp
--- a/Engine/Source/Runtime/D3D12RHI/Private/D3D12Submission.cpp
+++ b/Engine/Source/Runtime/D3D12RHI/Private/D3D12Submission.cpp
@@ -323,9 +323,19 @@ FD3D12DynamicRHI::FProcessResult FD3D12DynamicRHI::ProcessSubmissionQueue()
 	SCOPE_CYCLE_COUNTER(STAT_D3D12Submit);
 	LLM_SCOPE_BYNAME(TEXT("RHIMisc/ProcessSubmissionQueue"));
 
-	TArray<FD3D12Payload*, TInlineAllocator<64>> PayloadsToHandDown;
+	FD3D12Queue::FPayloadArray PayloadsToHandDown;
 
 	FProcessResult Result;
+
+	auto FlushPayloads = [&PayloadsToHandDown, &Result, DynamicRHI = this](int32 MinPayloadsToFlush = 1)
+	{
+		if (PayloadsToHandDown.Num() >= MinPayloadsToFlush)
+		{
+			Result.Status |= EQueueStatus::Processed;
+			DynamicRHI->FlushBatchedPayloads(PayloadsToHandDown);
+		}
+	};
+
 	bool bProgress;
 
 	do
@@ -486,19 +496,17 @@ FD3D12DynamicRHI::FProcessResult FD3D12DynamicRHI::ProcessSubmissionQueue()
 						uint64 ValueSignaled = OtherQueue->FinalizePayload(true, PayloadsToHandDown);
 						CurrentQueue.PayloadToSubmit->AddQueueFenceWait(OtherQueue->Fence, ValueSignaled);
 					}
+					FlushPayloads(FD3D12Queue::MaxBatchedPayloads);
 				}
 
 				// Now submit the original payload
 				CurrentQueue.FinalizePayload(false, PayloadsToHandDown);
+				FlushPayloads(FD3D12Queue::MaxBatchedPayloads);
 			}
 		});
 	} while (bProgress);
 
-	if (PayloadsToHandDown.Num())
-	{
-		Result.Status |= EQueueStatus::Processed;
-		FlushBatchedPayloads(PayloadsToHandDown);
-	}
+	FlushPayloads();
 
 	if (InterruptThread && EnumHasAnyFlags(Result.Status, EQueueStatus::Processed))
 	{
@@ -666,7 +674,7 @@ void FD3D12DynamicRHI::GenerateBarrierCommandListAndUpdateState(FD3D12CommandLis
 	}
 }
 
-uint64 FD3D12Queue::FinalizePayload(bool bRequiresSignal, TArray<FD3D12Payload*, TInlineAllocator<64>>& PayloadsToHandDown)
+uint64 FD3D12Queue::FinalizePayload(bool bRequiresSignal, FPayloadArray& PayloadsToHandDown)
 {
 	TRACE_CPUPROFILER_EVENT_SCOPE(ExecuteCommandList);
 	LLM_SCOPE_BYNAME(TEXT("RHIMisc/ExecuteCommandLists"));
@@ -816,7 +824,7 @@ uint64 FD3D12Queue::FinalizePayload(bool bRequiresSignal, TArray<FD3D12Payload*,
 	return Fence.NextCompletionValue;
 }
 
-void FD3D12DynamicRHI::FlushBatchedPayloads(TArray<FD3D12Payload*, TInlineAllocator<64>>& PayloadsToSubmit)
+void FD3D12DynamicRHI::FlushBatchedPayloads(FD3D12Queue::FPayloadArray& PayloadsToSubmit)
 {
 	uint32 FirstPayload = 0, LastPayload = 0;
 
