Move merging of worker context data to commit when there is a thread fence and not dispatch which can be translated in parallel (fixes random memory corruption)

diff --git a/Engine/Source/Runtime/D3D12RHI/Private/D3D12RayTracing.cpp b/Engine/Source/Runtime/D3D12RHI/Private/D3D12RayTracing.cpp
--- a/Engine/Source/Runtime/D3D12RHI/Private/D3D12RayTracing.cpp
+++ b/Engine/Source/Runtime/D3D12RHI/Private/D3D12RayTracing.cpp
@@ -1474,7 +1474,20 @@ public:
 				AddResourceReference(Resource, 0);
 			}
 
+			for (FD3D12ShaderResourceView* SRV : WorkerData[WorkerIndex].TransitionSRVs)
+			{
+				AddResourceTransition(SRV, 0);
+			}
+
+			for (FD3D12UnorderedAccessView* UAV : WorkerData[WorkerIndex].TransitionUAVs)
+			{
+				AddResourceTransition(UAV, 0);
+			}
+
 			WorkerData[WorkerIndex].ReferencedD3D12Resources.Empty();
+			WorkerData[WorkerIndex].TransitionSRVs.Empty();
+			WorkerData[WorkerIndex].TransitionUAVs.Empty();
+			WorkerData[WorkerIndex].TransitionViewSet.Empty();
 		}
 
 		FD3D12Device* Device = Context.GetParentDevice();
@@ -1630,7 +1643,7 @@ public:
 		}
 	}
 
-	void UpdateResidency(FD3D12CommandContext& CommandContext)
+	void UpdateResidency(FD3D12CommandContext& CommandContext) const
 	{
 		// Skip redundant resource residency updates when a shader table is repeatedly used on the same command list
 		bool bWasAlreadyInSet = false;
@@ -1671,29 +1684,8 @@ public:
 		}
 	}
 
-	void TransitionResources(FD3D12CommandContext& CommandContext)
-	{
-		// Merge all data from worker threads into the main set
-
-		for (uint32 WorkerIndex = 1; WorkerIndex < MaxBindingWorkers; ++WorkerIndex)
-		{
-			for (FD3D12ShaderResourceView* SRV : WorkerData[WorkerIndex].TransitionSRVs)
-			{
-				AddResourceTransition(SRV, 0);
-			}
-
-			for (FD3D12UnorderedAccessView* UAV : WorkerData[WorkerIndex].TransitionUAVs)
-			{
-				AddResourceTransition(UAV, 0);
-			}
-
-			WorkerData[WorkerIndex].TransitionSRVs.Empty();
-			WorkerData[WorkerIndex].TransitionUAVs.Empty();
-			WorkerData[WorkerIndex].TransitionViewSet.Empty();
-		}
-
-		// Use the main (merged) set data to perform resource transitions
-
+	void TransitionResources(FD3D12CommandContext& CommandContext) const
+	{	
 		for (FD3D12ShaderResourceView* SRV : WorkerData[0].TransitionSRVs)
 		{
 			CommandContext.TransitionResource(SRV, D3D12_RESOURCE_STATE_NON_PIXEL_SHADER_RESOURCE);
