Only perform discard workaround on render target textures that require it.

diff --git a/Engine/Source/Runtime/D3D12RHI/Private/D3D12Commands.cpp b/Engine/Source/Runtime/D3D12RHI/Private/D3D12Commands.cpp
--- a/Engine/Source/Runtime/D3D12RHI/Private/D3D12Commands.cpp
+++ b/Engine/Source/Runtime/D3D12RHI/Private/D3D12Commands.cpp
@@ -398,7 +398,7 @@ void FD3D12CommandContext::HandleDiscardResources(
 		else if (EnumHasAnyFlags(DiscardResource.Flags, EResourceTransitionFlags::Discard))
 		{
 #if PLATFORM_REQUIRES_TYPELESS_RESOURCE_DISCARD_WORKAROUND
-			if (DiscardResource.Texture && DiscardResource.RTV)
+			if (DiscardResource.Texture && DiscardResource.RTV && DiscardResource.Texture->GetRequiresTypelessResourceDiscardWorkaround())
 			{
 				FLinearColor ClearColor = DiscardResource.Texture->GetClearColor();
 				GraphicsCommandList()->ClearRenderTargetView(DiscardResource.RTV->GetOfflineCpuHandle(), reinterpret_cast<float*>(&ClearColor), 0, nullptr);
