Fix D3D12 crash on machines that do not support copy queue timestamps
 - The array of timestamp queries for each payload always included the Begin/End timestamps, even when they were null/empty on copy queue command lists. This always resulted in a call to GetTimestampFrequency, which fails on machines with GPUs/drivers that don't support copy queue timestamps.
 - Also added checks in GetCalibrationTimestamp and GetTimestampFrequency to ensure it is not called when not supported.

diff --git a/Engine/Source/Runtime/D3D12RHI/Private/D3D12Device.cpp b/Engine/Source/Runtime/D3D12RHI/Private/D3D12Device.cpp
--- a/Engine/Source/Runtime/D3D12RHI/Private/D3D12Device.cpp
+++ b/Engine/Source/Runtime/D3D12RHI/Private/D3D12Device.cpp
@@ -741,6 +741,8 @@ void FD3D12Device::ReleaseQueryHeap(FD3D12QueryHeap* QueryHeap)
 
 uint64 FD3D12Device::GetTimestampFrequency(ED3D12QueueType QueueType)
 {
+	check(QueueType != ED3D12QueueType::Copy || GetParentAdapter()->AreCopyQueueTimestampQueriesSupported());
+
 	uint64 Frequency;
 	VERIFYD3D12RESULT(Queues[(uint32)QueueType].D3DCommandQueue->GetTimestampFrequency(&Frequency));
 	return Frequency;
@@ -749,6 +751,8 @@ uint64 FD3D12Device::GetTimestampFrequency(ED3D12QueueType QueueType)
 #if (RHI_NEW_GPU_PROFILER == 0)
 FGPUTimingCalibrationTimestamp FD3D12Device::GetCalibrationTimestamp(ED3D12QueueType QueueType)
 {
+	check(QueueType != ED3D12QueueType::Copy || GetParentAdapter()->AreCopyQueueTimestampQueriesSupported());
+
 	TRACE_CPUPROFILER_EVENT_SCOPE(D3D12GetCalibrationTimestamp);
 
 	uint64 GPUTimestampFrequency = GetTimestampFrequency(QueueType);
diff --git a/Engine/Source/Runtime/D3D12RHI/Private/D3D12Submission.cpp b/Engine/Source/Runtime/D3D12RHI/Private/D3D12Submission.cpp
--- a/Engine/Source/Runtime/D3D12RHI/Private/D3D12Submission.cpp
+++ b/Engine/Source/Runtime/D3D12RHI/Private/D3D12Submission.cpp
@@ -428,25 +428,34 @@ FD3D12DynamicRHI::FProcessResult FD3D12DynamicRHI::ProcessSubmissionQueue()
 						TargetQueue.BatchedObjects.PipelineStatsQueries.Append(MoveTemp(CommandList->State.PipelineStatsQueries));
 
 						// Timestamp Queries
-						// Keep only the first Begin() in the batch
-						if (TargetQueue.NumCommandListsInBatch++ == 0)
+						if (CommandList->State.BeginTimestamp)
 						{
-							TargetQueue.BatchedObjects.TimestampQueries.Emplace(MoveTemp(CommandList->State.BeginTimestamp));
-						}
-						else
-						{
-							// Remove the previous End() timestamp, to join the range together.
-							check(TargetQueue.BatchedObjects.TimestampQueries.Last().Type == ED3D12QueryType::CommandListEnd);
-							TargetQueue.BatchedObjects.TimestampQueries.RemoveAt(TargetQueue.BatchedObjects.TimestampQueries.Num() - 1);
-						}
+							// Keep only the first Begin() in the batch
+							if (TargetQueue.NumCommandListsInBatch++ == 0)
+							{
+								TargetQueue.BatchedObjects.TimestampQueries.Emplace(MoveTemp(CommandList->State.BeginTimestamp));
+							}
+							else
+							{
+								// Remove the previous End() timestamp, to join the range together.
+								check(TargetQueue.BatchedObjects.TimestampQueries.Last().Type == ED3D12QueryType::CommandListEnd);
+								TargetQueue.BatchedObjects.TimestampQueries.RemoveAt(TargetQueue.BatchedObjects.TimestampQueries.Num() - 1);
+							}
 
-						TargetQueue.BatchedObjects.TimestampQueries.Append(MoveTemp(CommandList->State.TimestampQueries));
-						TargetQueue.BatchedObjects.TimestampQueries.Emplace(MoveTemp(CommandList->State.EndTimestamp));
+							TargetQueue.BatchedObjects.TimestampQueries.Append(MoveTemp(CommandList->State.TimestampQueries));
+							TargetQueue.BatchedObjects.TimestampQueries.Emplace(MoveTemp(CommandList->State.EndTimestamp));
 
-						if (TargetQueue.NumCommandListsInBatch >= MaxBatchSize)
+							if (TargetQueue.NumCommandListsInBatch >= MaxBatchSize)
+							{
+								// Start a new batch
+								TargetQueue.NumCommandListsInBatch = 0;
+							}
+						}
+						else
 						{
-							// Start a new batch
-							TargetQueue.NumCommandListsInBatch = 0;
+							// No begin timestamp means timestamps aren't supported on this queue
+							check(CommandList->State.TimestampQueries.IsEmpty());
+							check(!CommandList->State.EndTimestamp);
 						}
 					};
 
