- fix wrong assumption code in pool allocator
- remove unused NEEDS_D3D12_INDIRECT_ARGUMENT_HEAP_WORKAROUND

diff --git a/Engine/Source/Runtime/D3D12RHI/Private/D3D12PoolAllocator.cpp b/Engine/Source/Runtime/D3D12RHI/Private/D3D12PoolAllocator.cpp
--- a/Engine/Source/Runtime/D3D12RHI/Private/D3D12PoolAllocator.cpp
+++ b/Engine/Source/Runtime/D3D12RHI/Private/D3D12PoolAllocator.cpp
@@ -4,10 +4,6 @@
 #include "D3D12RHIPrivate.h"
 #include "HAL/LowLevelMemTracker.h"
 
-#ifndef NEEDS_D3D12_INDIRECT_ARGUMENT_HEAP_WORKAROUND
-#define NEEDS_D3D12_INDIRECT_ARGUMENT_HEAP_WORKAROUND 0
-#endif
-
 LLM_DECLARE_TAG(D3D12AllocatorUnused);
 
 //-----------------------------------------------------------------------------
@@ -208,10 +204,7 @@ FD3D12ResourceInitConfig FD3D12PoolAllocator::GetResourceAllocatorInitConfig(D3D
 	InitConfig.HeapFlags = D3D12_HEAP_FLAG_ALLOW_ONLY_BUFFERS;
 	if (EnumHasAnyFlags(InBufferUsage, BUF_DrawIndirect))
 	{
-		check(InResourceFlags & D3D12_RESOURCE_FLAG_ALLOW_UNORDERED_ACCESS);
-#if !NEEDS_D3D12_INDIRECT_ARGUMENT_HEAP_WORKAROUND
 		InitConfig.HeapFlags |= D3D12RHI_HEAP_FLAG_ALLOW_INDIRECT_BUFFERS;
-#endif
 	}
 
 	return InitConfig;
