Fixing D3D12 Bindless Texture Reference streaming
- RHIUpdateTextureReference was only updating the TextureReference descriptor when the referenced texture changed.
- It turns out that partially resident texture updates modify the already set referenced texture before passing it back to RHIUpdateTextureReference, meaning we need to detect this case to make sure the TextureReference descriptor is updated.
- We can lean on the "Version" field of the offline CPU descriptor handles to detect if the texture reference actually needs to update.
- Changing FD3D12ShaderResourceView::UpdateMinLODClamp to only update the descriptor if the MinLodClamp value actually changes.

diff --git a/Engine/Source/Runtime/D3D12RHI/Private/D3D12SRV.cpp b/Engine/Source/Runtime/D3D12RHI/Private/D3D12SRV.cpp
--- a/Engine/Source/Runtime/D3D12RHI/Private/D3D12SRV.cpp
+++ b/Engine/Source/Runtime/D3D12RHI/Private/D3D12SRV.cpp
@@ -89,18 +89,25 @@ void FD3D12ShaderResourceView::UpdateMinLODClamp(FD3D12ContextArray const& Conte
 {
 	check(IsInitialized());
 
+	FLOAT* pResourceMinLODClamp = nullptr;
+
 	switch (D3DViewDesc.ViewDimension)
 	{
 	default: checkNoEntry(); return; // not supported
-	case D3D12_SRV_DIMENSION_TEXTURE2D       : D3DViewDesc.Texture2D       .ResourceMinLODClamp = MinLODClamp; break;
-	case D3D12_SRV_DIMENSION_TEXTURE2DARRAY  : D3DViewDesc.Texture2DArray  .ResourceMinLODClamp = MinLODClamp; break;
-	case D3D12_SRV_DIMENSION_TEXTURE3D       : D3DViewDesc.Texture3D       .ResourceMinLODClamp = MinLODClamp; break;
-	case D3D12_SRV_DIMENSION_TEXTURECUBE     : D3DViewDesc.TextureCube     .ResourceMinLODClamp = MinLODClamp; break;
-	case D3D12_SRV_DIMENSION_TEXTURECUBEARRAY: D3DViewDesc.TextureCubeArray.ResourceMinLODClamp = MinLODClamp; break;
+	case D3D12_SRV_DIMENSION_TEXTURE2D       : pResourceMinLODClamp = &D3DViewDesc.Texture2D       .ResourceMinLODClamp; break;
+	case D3D12_SRV_DIMENSION_TEXTURE2DARRAY  : pResourceMinLODClamp = &D3DViewDesc.Texture2DArray  .ResourceMinLODClamp; break;
+	case D3D12_SRV_DIMENSION_TEXTURE3D       : pResourceMinLODClamp = &D3DViewDesc.Texture3D       .ResourceMinLODClamp; break;
+	case D3D12_SRV_DIMENSION_TEXTURECUBE     : pResourceMinLODClamp = &D3DViewDesc.TextureCube     .ResourceMinLODClamp; break;
+	case D3D12_SRV_DIMENSION_TEXTURECUBEARRAY: pResourceMinLODClamp = &D3DViewDesc.TextureCubeArray.ResourceMinLODClamp; break;
 	}
 
-	UpdateDescriptor();
-	UpdateBindlessSlot(Contexts);
+	if (pResourceMinLODClamp && *pResourceMinLODClamp != MinLODClamp)
+	{
+		*pResourceMinLODClamp = MinLODClamp;
+
+		UpdateDescriptor();
+		UpdateBindlessSlot(Contexts);
+	}
 }
 
 void FD3D12ShaderResourceView::UpdateDescriptor()
diff --git a/Engine/Source/Runtime/D3D12RHI/Private/D3D12TextureReference.cpp b/Engine/Source/Runtime/D3D12RHI/Private/D3D12TextureReference.cpp
--- a/Engine/Source/Runtime/D3D12RHI/Private/D3D12TextureReference.cpp
+++ b/Engine/Source/Runtime/D3D12RHI/Private/D3D12TextureReference.cpp
@@ -15,7 +15,10 @@ FD3D12RHITextureReference::FD3D12RHITextureReference(FD3D12Device* InDevice, FD3
 	{
 		InReferencedTexture->AddRenameListener(this);
 
-		InDevice->GetBindlessDescriptorManager().InitializeDescriptor(BindlessHandle, InReferencedTexture->GetShaderResourceView());
+		FD3D12ShaderResourceView* View = InReferencedTexture->GetShaderResourceView();
+		ReferencedDescriptorVersion = View->GetOfflineCpuHandle().GetVersion();
+
+		InDevice->GetBindlessDescriptorManager().InitializeDescriptor(BindlessHandle, View);
 	}
 #endif // PLATFORM_SUPPORTS_BINDLESS_RENDERING
 }
@@ -32,14 +35,16 @@ FD3D12RHITextureReference::~FD3D12RHITextureReference()
 #endif // PLATFORM_SUPPORTS_BINDLESS_RENDERING
 }
 
-void FD3D12RHITextureReference::SwitchToNewTexture(FD3D12ContextArray const& Contexts, FD3D12Texture* InNewTexture)
+void FD3D12RHITextureReference::SwitchToNewTexture(const FD3D12ContextArray& Contexts, FD3D12Texture* InNewTexture)
 {
 #if PLATFORM_SUPPORTS_BINDLESS_RENDERING
 	if (BindlessHandle.IsValid())
 	{
+		FD3D12Texture* CurrentTexture = FD3D12DynamicRHI::ResourceCast(GetReferencedTexture());
 		FD3D12Texture* NewTexture = InNewTexture ? InNewTexture : FD3D12DynamicRHI::ResourceCast(FRHITextureReference::GetDefaultTexture());
 
-		FD3D12Texture* CurrentTexture = FD3D12DynamicRHI::ResourceCast(GetReferencedTexture());
+		FD3D12ShaderResourceView* NewView = NewTexture->GetShaderResourceView();
+		const FD3D12OfflineDescriptor NewViewDescriptor = NewView->GetOfflineCpuHandle();
 
 		if (CurrentTexture != NewTexture)
 		{
@@ -50,8 +55,14 @@ void FD3D12RHITextureReference::SwitchToNewTexture(FD3D12ContextArray const& Con
 
 			NewTexture->AddRenameListener(this);
 
-			GetParentDevice()->GetBindlessDescriptorManager().UpdateDescriptor(Contexts, BindlessHandle, NewTexture->GetShaderResourceView());
+			GetParentDevice()->GetBindlessDescriptorManager().UpdateDescriptor(Contexts, BindlessHandle, NewView);
 		}
+		else if (NewViewDescriptor.GetVersion() != ReferencedDescriptorVersion)
+		{
+			GetParentDevice()->GetBindlessDescriptorManager().UpdateDescriptor(Contexts, BindlessHandle, NewView);
+		}
+
+		ReferencedDescriptorVersion = NewViewDescriptor.GetVersion();
 	}
 #endif // PLATFORM_SUPPORTS_BINDLESS_RENDERING
 
@@ -59,14 +70,17 @@ void FD3D12RHITextureReference::SwitchToNewTexture(FD3D12ContextArray const& Con
 }
 
 #if PLATFORM_SUPPORTS_BINDLESS_RENDERING
-void FD3D12RHITextureReference::ResourceRenamed(FD3D12ContextArray const& Contexts, FD3D12BaseShaderResource* InRenamedResource, FD3D12ResourceLocation* InNewResourceLocation)
+void FD3D12RHITextureReference::ResourceRenamed(const FD3D12ContextArray& Contexts, FD3D12BaseShaderResource* InRenamedResource, FD3D12ResourceLocation* InNewResourceLocation)
 {
 	if (ensure(BindlessHandle.IsValid()))
 	{
 		FD3D12Texture* RenamedTexture = static_cast<FD3D12Texture*>(InRenamedResource);
 		checkSlow(RenamedTexture == ReferencedTexture);
 
-		GetParentDevice()->GetBindlessDescriptorManager().UpdateDescriptor(Contexts, BindlessHandle, RenamedTexture->GetShaderResourceView());
+		FD3D12ShaderResourceView* RenamedTextureView = RenamedTexture->GetShaderResourceView();
+		ReferencedDescriptorVersion = RenamedTextureView->GetOfflineCpuHandle().GetVersion();
+
+		GetParentDevice()->GetBindlessDescriptorManager().UpdateDescriptor(Contexts, BindlessHandle, RenamedTextureView);
 	}
 }
 #endif // PLATFORM_SUPPORTS_BINDLESS_RENDERING
@@ -82,20 +96,17 @@ FTextureReferenceRHIRef FD3D12DynamicRHI::RHICreateTextureReference(FRHICommandL
 	});
 }
 
-void FD3D12DynamicRHI::RHIUpdateTextureReference(FRHICommandListBase& RHICmdList, FRHITextureReference* TextureRef, FRHITexture* InNewTexture)
+void FD3D12DynamicRHI::RHIUpdateTextureReference(FRHICommandListBase& RHICmdList, FRHITextureReference* InTextureReference, FRHITexture* InNewTexture)
 {
 	// Workaround for a crash bug where FRHITextureReferences are deleted before this command is executed on the RHI thread.
 	// Take a reference on the FRHITextureReference object to keep it alive.
-	// @todo dev-pr - This should be refactored out when we eventually remove FRHITextureReference.
-	TRefCountPtr<FRHITextureReference> Ref = TextureRef;
+	TRefCountPtr<FD3D12RHITextureReference> TextureReferenceRef = ResourceCast(InTextureReference);
+	FD3D12Texture* NewTexture = ResourceCast(InNewTexture ? InNewTexture : FRHITextureReference::GetDefaultTexture());
 
 	RHICmdList.EnqueueLambdaMultiPipe(GetEnabledRHIPipelines(), FRHICommandListBase::EThreadFence::Enabled, TEXT("FD3D12DynamicRHI::RHIUpdateTextureReference"),
-	[
-		TextureRef = MoveTemp(Ref),
-		NewTexture = InNewTexture ? InNewTexture : FRHITextureReference::GetDefaultTexture()
-	](FD3D12ContextArray const& Contexts)
+	[TextureReference = MoveTemp(TextureReferenceRef), NewTexture](const FD3D12ContextArray& Contexts)
 	{
-		for (TD3D12DualLinkedObjectIterator<FD3D12RHITextureReference, FD3D12Texture> It(ResourceCast(TextureRef.GetReference()), ResourceCast(NewTexture)); It; ++It)
+		for (TD3D12DualLinkedObjectIterator<FD3D12RHITextureReference, FD3D12Texture> It(TextureReference.GetReference(), NewTexture); It; ++It)
 		{
 			It.GetFirst()->SwitchToNewTexture(Contexts, It.GetSecond());
 		}
diff --git a/Engine/Source/Runtime/D3D12RHI/Private/D3D12TextureReference.h b/Engine/Source/Runtime/D3D12RHI/Private/D3D12TextureReference.h
--- a/Engine/Source/Runtime/D3D12RHI/Private/D3D12TextureReference.h
+++ b/Engine/Source/Runtime/D3D12RHI/Private/D3D12TextureReference.h
@@ -21,6 +21,10 @@ public:
 #endif // PLATFORM_SUPPORTS_BINDLESS_RENDERING
 
 	void SwitchToNewTexture(FD3D12ContextArray const& Contexts, FD3D12Texture* InNewTexture);
+
+#if PLATFORM_SUPPORTS_BINDLESS_RENDERING
+	uint32 ReferencedDescriptorVersion{};
+#endif
 };
 
 template<>
