Do not use fast allocator when allocating upload memory on async copy queue tasks because lifetime of allocation is unknown

diff --git a/Engine/Source/Runtime/D3D12RHI/Private/D3D12Texture.cpp b/Engine/Source/Runtime/D3D12RHI/Private/D3D12Texture.cpp
--- a/Engine/Source/Runtime/D3D12RHI/Private/D3D12Texture.cpp
+++ b/Engine/Source/Runtime/D3D12RHI/Private/D3D12Texture.cpp
@@ -775,16 +775,6 @@ public:
 			// Initial data upload is done
 			CurrentTexture.ResourceLocation.UnlockPoolData();
 		}
-
-		// These are clear to be recycled now because GPU is done with it at this point because this task use the copy command list sync points
-		// as prerequisites. No defer delete required but can be reused immediately
-		TempResourceLocation.GetResource()->DoNotDeferDelete();
-		TempResourceLocation.GetResource()->Release();
-		if (TempResourceLocationLowMips.IsValid())
-		{
-			TempResourceLocationLowMips.GetResource()->DoNotDeferDelete();
-			TempResourceLocationLowMips.GetResource()->Release();
-		}
 	}
 
 	static ESubsequentsMode::Type GetSubsequentsMode()
@@ -912,12 +902,13 @@ FTextureRHIRef FD3D12DynamicRHI::RHIAsyncCreateTexture2D(uint32 SizeX, uint32 Si
 
 		check((TextureDesc.Flags & D3D12_RESOURCE_FLAG_DENY_SHADER_RESOURCE) == 0);
 
-		FD3D12FastAllocator& FastAllocator = TextureOut->GetParentDevice()->GetDefaultFastAllocator();
+		FD3D12Device* Device = TextureOut->GetParentDevice();
+		FD3D12UploadHeapAllocator& UploadHeapAllocator = Adapter->GetUploadHeapAllocator(Device->GetGPUIndex());
 		uint64 Size = GetRequiredIntermediateSize(TextureOut->GetResource()->GetResource(), 0, NumMips);
 		uint64 SizeLowMips = 0;
 
-		FD3D12ResourceLocation TempResourceLocation(FastAllocator.GetParentDevice());
-		FD3D12ResourceLocation TempResourceLocationLowMips(FastAllocator.GetParentDevice());
+		FD3D12ResourceLocation TempResourceLocation(Device);
+		FD3D12ResourceLocation TempResourceLocationLowMips(Device);
 
 		// The allocator work in pages of 4MB. Increasing page size is undesirable from a hitching point of view because there's a performance cliff above 4MB
 		// where creation time of new pages can increase by an order of magnitude. Most allocations are smaller than 4MB, but a common exception is
@@ -940,12 +931,12 @@ FTextureRHIRef FD3D12DynamicRHI::RHIAsyncCreateTexture2D(uint32 SizeX, uint32 Si
 			
 			// Mip 0
 			SizeMip0 = Layouts[1].Offset;
-			FastAllocator.Allocate(SizeMip0, D3D12_TEXTURE_DATA_PLACEMENT_ALIGNMENT, &TempResourceLocation);
+			UploadHeapAllocator.AllocUploadResource(SizeMip0, D3D12_TEXTURE_DATA_PLACEMENT_ALIGNMENT, TempResourceLocation);
 			Layouts[0].Offset = TempResourceLocation.GetOffsetFromBaseOfResource();
 
 			// Remaining mip chain
 			SizeLowMips = TotalBytes - SizeMip0;
-			FastAllocator.Allocate(SizeLowMips, D3D12_TEXTURE_DATA_PLACEMENT_ALIGNMENT, &TempResourceLocationLowMips);
+			UploadHeapAllocator.AllocUploadResource(SizeLowMips, D3D12_TEXTURE_DATA_PLACEMENT_ALIGNMENT, TempResourceLocationLowMips);
 
 			const uint64 LowMipsTotalBufferSize = TempResourceLocationLowMips.GetResource()->GetDesc().Width;
 			
@@ -981,21 +972,14 @@ FTextureRHIRef FD3D12DynamicRHI::RHIAsyncCreateTexture2D(uint32 SizeX, uint32 Si
 					TEXT("Mip tail upload buffer total size is too small for mip %llu. Layouts[MipIndex].Offset=%llu, MipCopySize=%llu, AbsoluteMipCopyEndOffset=%llu, LowMipsTotalBufferSize=%llu."),
 					MipIndex, Layouts[MipIndex].Offset, MipCopySize, AbsoluteMipCopyEndOffset, LowMipsTotalBufferSize);
 			}
-			
-			TempResourceLocationLowMips.GetResource()->AddRef();
 		}
 		else
 		{
-			FastAllocator.Allocate(Size, D3D12_TEXTURE_DATA_PLACEMENT_ALIGNMENT, &TempResourceLocation);
+			UploadHeapAllocator.AllocUploadResource(Size, D3D12_TEXTURE_DATA_PLACEMENT_ALIGNMENT, TempResourceLocation);
 		}
-		// We AddRef() the resource here to make sure it doesn't get recycled prematurely. We are likely to be done with it during the frame,
-		// but lifetime of the allocation is not strictly tied to the frame because we're using the copy queue here. Because we're waiting
-		// on the GPU before returning here, this protection is safe, even if we end up straddling frame boundaries.
-		TempResourceLocation.GetResource()->AddRef();
 
 		for (FD3D12Texture& CurrentTexture : *TextureOut)
 		{
-			FD3D12Device* Device = CurrentTexture.GetParentDevice();
 			FD3D12Resource* Resource = CurrentTexture.GetResource();
 
 			FD3D12SyncPointRef SyncPoint;
