Enable d3d12.AllowPoolAllocateIndirectArgBuffers by default

diff --git a/Engine/Source/Runtime/D3D12RHI/Private/D3D12Allocation.cpp b/Engine/Source/Runtime/D3D12RHI/Private/D3D12Allocation.cpp
--- a/Engine/Source/Runtime/D3D12RHI/Private/D3D12Allocation.cpp
+++ b/Engine/Source/Runtime/D3D12RHI/Private/D3D12Allocation.cpp
@@ -14,7 +14,7 @@
 #include "ProfilingDebugging/MemoryTrace.h"
 
 // Fix for random GPU crashes on draw indirects on multiple IHVs. Force all indirect arg buffers as committed resources (see UE-115982)
-static int32 GD3D12AllowPoolAllocateIndirectArgBuffers = 0;
+static int32 GD3D12AllowPoolAllocateIndirectArgBuffers = 1;
 static FAutoConsoleVariableRef CVarD3D12AllowPoolAllocateIndirectArgBuffers(
 	TEXT("d3d12.AllowPoolAllocateIndirectArgBuffers"),
 	GD3D12AllowPoolAllocateIndirectArgBuffers,
